Remove Gbuffer access from SSS passes when Strata is used.

This CL hooks the FStrataSubsurfaceData to the SSS passes, so that no Gbuffer data are no longer needed.

Tested on preflight & EngineTests

diff --git a/Engine/Shaders/Private/BurleyNormalizedSSSCommon.ush b/Engine/Shaders/Private/BurleyNormalizedSSSCommon.ush
--- a/Engine/Shaders/Private/BurleyNormalizedSSSCommon.ush
+++ b/Engine/Shaders/Private/BurleyNormalizedSSSCommon.ush
@@ -45,42 +45,18 @@ bool GetSubsurfaceTransmittanceProfileUseBurley(uint SubsurfaceProfileInt)
 
 bool GetSubsurfaceProfileUseBurley(uint SubsurfaceProfileInt)
 {
-#if STRATA_ENABLED
-	// STRATA_TODO: For now force Burley to be used by default
-	return true;
-#else
 	// 0..255, which SubSurface profile to pick
 	float Type = ActualSSProfilesTexture.Load(int3(SSSS_BOUNDARY_COLOR_BLEED_OFFSET, SubsurfaceProfileInt, 0)).a;
 
 	return abs(Type - SSS_TYPE_BURLEY) < 0.01f;
-#endif
-}
-
-bool GetSubsurfaceProfileUseBurley(float2 BufferUV)
-{
-	FScreenSpaceData ScreenSpaceData = GetScreenSpaceData(BufferUV);
-	uint SubsurfaceProfileInt = ExtractSubsurfaceProfileInt(ScreenSpaceData.GBuffer);
-	return GetSubsurfaceProfileUseBurley(SubsurfaceProfileInt);
 }
 
 bool GetSubsurfaceProfileUseSeparable(uint SubsurfaceProfileInt)
 {
-#if STRATA_ENABLED
-	// STRATA_TODO: For now force Burley to be used by default
-	return false;
-#else
 	// 0..255, which SubSurface profile to pick
 	float Type = ActualSSProfilesTexture.Load(int3(SSSS_BOUNDARY_COLOR_BLEED_OFFSET, SubsurfaceProfileInt, 0)).a;
 
 	return abs(Type - SSS_TYPE_SSSS) < 0.01f;
-#endif
-}
-
-bool GetSubsurfaceProfileUseSeparable(float2 BufferUV)
-{
-	FScreenSpaceData ScreenSpaceData = GetScreenSpaceData(BufferUV);
-	uint SubsurfaceProfileInt = ExtractSubsurfaceProfileInt(ScreenSpaceData.GBuffer);
-	return GetSubsurfaceProfileUseSeparable(SubsurfaceProfileInt);
 }
 
 inline float3 Burley_Profile(float Radius, float3 Albedo, float3 S3D, float L)
@@ -172,6 +148,9 @@ float3 GetSearchLightDiffuseScalingFactor3D(float3 SurfaceAlbedo)
 
 //@TODO Revisit here after offline comparison.
 //call this function in TransmissionCommon.ush for Debug. No transmission tint is applied.
+
+// This code is not used
+#if 0
 float4 GetBurleyTransmissionProfile(FGBufferData GBufferData, float Thickness)
 {
 	// 0..255, which SubSurface profile to pick
@@ -186,4 +165,5 @@ float4 GetBurleyTransmissionProfile(FGBufferData GBufferData, float Thickness)
 	float3 r = Thickness / MeanFreePath.xyz;
 	Output.xyz= 0.25*SurfaceAlbedo*(exp(-ScalingFactor * r) + 3 * exp(-ScalingFactor * r / 3));
 	return Output;
-}
\ No newline at end of file
+}
+#endif
\ No newline at end of file
diff --git a/Engine/Shaders/Private/PostProcessSubsurface.usf b/Engine/Shaders/Private/PostProcessSubsurface.usf
--- a/Engine/Shaders/Private/PostProcessSubsurface.usf
+++ b/Engine/Shaders/Private/PostProcessSubsurface.usf
@@ -10,6 +10,10 @@
 #define STRATA_ENABLED PROJECT_STRATA
 #endif
 
+#ifdef STRATA_ENABLED
+#include "Strata/Strata.ush"
+#endif
+
 #include "PostProcessCommon.ush"
 #include "DeferredShadingCommon.ush"
 #include "SubsurfaceProfileCommon.ush"
@@ -103,12 +107,18 @@ float GetMaskFromDepthInAlpha(float Alpha)
 }
 
 // can be optimized
-float GetSubsurfaceStrength(float2 UV)
+float GetSubsurfaceStrength(uint2 PixelCoord, float2 UV)
 {
+#if STRATA_ENABLED
+	const FStrataSubsurfaceData SSSData = StrataUnpackSSSData(Strata.SSSTexture.Load(uint3(PixelCoord, 0)));
+	const float Mask = SSSData.bIsValid ? 1.0f : 0.0f;
+	const float RadiusScale = SSSData.ProfileRadiusScale;
+#else
 	FScreenSpaceData ScreenSpaceData = GetScreenSpaceData(UV);
-
-	float Mask = UseSubsurfaceProfile(ScreenSpaceData.GBuffer.ShadingModelID) ? 1.0f : 0.0f;
-	return Mask * ScreenSpaceData.GBuffer.CustomData.a;
+	const float Mask = UseSubsurfaceProfile(ScreenSpaceData.GBuffer.ShadingModelID) ? 1.0f : 0.0f;
+	const float RadiusScale = ScreenSpaceData.GBuffer.CustomData.a;
+#endif
+	return Mask * RadiusScale;
 }
 
 // @return .rgb is the weight for color channel, .a is the sample location
@@ -122,26 +132,6 @@ float4 GetSubsurfaceProfileKernel(uint SampleIndex, uint SubsurfaceProfileInt)
 	return ActualSSProfilesTexture.Load(int3(SampleIndex, SubsurfaceProfileInt, 0)) * TableMax;
 }
 
-float GetSubsurfaceProfileRadiusScale(FGBufferData GBufferData)
-{
-	// 0..255, which SubSurface profile to pick
-	uint SubsurfaceProfileInt = ExtractSubsurfaceProfileInt(GBufferData);
-#if STRATA_ENABLED
-	if (SubsurfaceProfileInt == SSS_PROFILE_ID_PERPIXEL) return 1.f;
-#endif
-	return GetSubsurfaceProfileKernel(SSSS_N_KERNELWEIGHTOFFSET + SSSS_N_KERNELWEIGHTCOUNT - 1, SubsurfaceProfileInt).a;
-}
-
-float3 GetSubsurfaceProfileColor(FGBufferData GBufferData)
-{
-	// 0..255, which SubSurface profile to pick
-	uint SubsurfaceProfileInt = ExtractSubsurfaceProfileInt(GBufferData);
-#if STRATA_ENABLED
-	if (SubsurfaceProfileInt == SSS_PROFILE_ID_PERPIXEL) return GBufferData.BaseColor;
-#endif
-	return GetSubsurfaceProfileKernel(SSSS_SUBSURFACE_COLOR_OFFSET, SubsurfaceProfileInt).rgb;
-}
-
 float3 GetSubsurfaceProfileNormal(float2 BufferUV)
 {
 
@@ -151,19 +141,21 @@ float3 GetSubsurfaceProfileNormal(float2 BufferUV)
 	return LocalGBufferData.WorldNormal;
 }
 
-
-uint GetSubsurfaceProfileId(float2 BufferUV)
+uint GetSubsurfaceProfileId(float2 BufferUV, float2 Extent)
 {
 	BufferUV = clamp(BufferUV, SubsurfaceInput0_UVViewportBilinearMin, SubsurfaceInput0_UVViewportBilinearMax);
-
+#if STRATA_ENABLED
+	const uint2 PixelCoord = BufferUV * Extent;
+	const FStrataSubsurfaceData SSSData = StrataUnpackSSSData(Strata.SSSTexture.Load(uint3(PixelCoord, 0)));
+	return SSSData.ProfileId;
+#else
 	FGBufferData LocalGBufferData = GetGBufferData(BufferUV);
 	return ExtractSubsurfaceProfileInt(LocalGBufferData);
+#endif
 }
 
-float3 GetSubsurfaceProfileBoundaryColorBleed(FGBufferData GBufferData)
+float3 GetSubsurfaceProfileBoundaryColorBleed(uint SubsurfaceProfileInt)
 {
-	// 0..255, which SubSurface profile to pick
-	uint SubsurfaceProfileInt = ExtractSubsurfaceProfileInt(GBufferData);
 #if STRATA_ENABLED
 	if (SubsurfaceProfileInt == SSS_PROFILE_ID_PERPIXEL) return float3(1.f, 1.f, 1.f);
 #endif
@@ -190,7 +182,7 @@ bool InUnitBox(float2 UV)
 }
 
 // @return 0=don't blend in, 1:fully blend in
-float ComputeFullResLerp(FScreenSpaceData ScreenSpaceData, float2 UVSceneColor, float2 FullResInputSizeInverse)
+float ComputeFullResLerp(float ProfileRadiusScale, float PixelRadiusScale, float2 UVSceneColor, float2 FullResInputSizeInverse)
 {
 	float SSSScaleX = SubsurfaceParams.x;
 
@@ -201,7 +193,7 @@ float ComputeFullResLerp(FScreenSpaceData ScreenSpaceData, float2 UVSceneColor,
 	// Calculate the final step to fetch the surrounding pixels:
 	float finalStep = scale * HorizontalScaler;
 
-	finalStep *= GetSubsurfaceProfileRadiusScale(ScreenSpaceData.GBuffer);
+	finalStep *= ProfileRadiusScale;
 
 	float PixelSizeRadius = finalStep / (FullResInputSizeInverse.x * 0.5f);
 
@@ -211,7 +203,7 @@ float ComputeFullResLerp(FScreenSpaceData ScreenSpaceData, float2 UVSceneColor,
 	float Ret = saturate(PixelSizeRadius - PixelSize);
 
 	// opacity allows to scale the radius - at some point we should fade in the full resolution, we don't have a masking other than that.
-	Ret *= saturate(ScreenSpaceData.GBuffer.CustomData.a * 10);
+	Ret *= saturate(PixelRadiusScale * 10);
 
 	// todo: Subsurface has some non scatter contribution - all that should come from the Full res
 	return Ret;
@@ -234,13 +226,22 @@ void VisualizePS(in noperspective float4 UVAndScreenPos : TEXCOORD0, float4 SvPo
 	{
 		float2 UV = SubsurfaceInput0_ViewportMin * SubsurfaceInput0_ExtentInverse + IDAreaLocalUV * (SubsurfaceInput0_ViewportSize * SubsurfaceInput0_ExtentInverse);
 
-		FScreenSpaceData ScreenSpaceData = GetScreenSpaceData(UV);
-
-		int SubsurfaceProfileInt = ExtractSubsurfaceProfileInt(ScreenSpaceData.GBuffer);
+	#if STRATA_ENABLED
+		const FStrataSubsurfaceData SSSData = StrataUnpackSSSData(Strata.SSSTexture.Load(uint3(PixelPos, 0)));
+		const uint SubsurfaceProfileInt = SSSData.ProfileId;
+		const bool bIsProfile = SSSData.bIsProfile;
+		const float PixelRadiusScale = SSSData.ProfileRadiusScale;
+	#else
+		const FScreenSpaceData ScreenSpaceData = GetScreenSpaceData(UV);
+		const uint SubsurfaceProfileInt = ExtractSubsurfaceProfileInt(ScreenSpaceData.GBuffer);
+		const bool bIsProfile = UseSubsurfaceProfile(ScreenSpaceData.GBuffer.ShadingModelID);
+		const float PixelRadiusScale = ScreenSpaceData.GBuffer.CustomData.a;
+	#endif
+		const float ProfileRadiusScale = GetSubsurfaceProfileKernel(SSSS_N_KERNELWEIGHTOFFSET + SSSS_N_KERNELWEIGHTCOUNT - 1, SubsurfaceProfileInt).a;
 
 		OutColor = float4(0.5f, 0.5f, 0.5f, 0);
 
-		BRANCH if (UseSubsurfaceProfile(ScreenSpaceData.GBuffer.ShadingModelID))
+		BRANCH if (bIsProfile)
 		{
 			if (SubsurfaceProfileInt == 0)
 			{
@@ -283,10 +284,9 @@ void VisualizePS(in noperspective float4 UVAndScreenPos : TEXCOORD0, float4 SvPo
 			int2 LeftTop = (PixelPos / 8) * 8;
 			PrintCharacter(PixelPos, OutColor.rgb, float3(1, 1, 1), LeftTop, SubsurfaceProfileInt);
 
-			OutColor.rgb *= ComputeFullResLerp(ScreenSpaceData, UV, SubsurfaceInput0_ExtentInverse);
+			OutColor.rgb *= ComputeFullResLerp(ProfileRadiusScale, PixelRadiusScale, UV, SubsurfaceInput0_ExtentInverse);
 		}
 	}
-
 }
 
 
@@ -388,13 +388,19 @@ SDiffuseAndSpecular ReconstructLighting(float2 UVSceneColor, uint ReconstructMet
 
 // @param UVSceneColor for the full res rendertarget (BufferSize) e.g. SceneColor or GBuffers
 // @return .RGB Color that should be scattared, .A:1 for subsurface scattering material, 0 for not
-float4 SetupSubsurfaceForOnePixel(float2 UVSceneColor)
+float4 SetupSubsurfaceForOnePixel(float2 UVSceneColor, float2 Extent)
 {
-	float4 Ret = 0;
-
+#if STRATA_ENABLED
+	const uint2 PixelCoord = UVSceneColor * Extent;
+	const FStrataSubsurfaceData SSSData = StrataUnpackSSSData(Strata.SSSTexture.Load(uint3(PixelCoord, 0)));
+	const bool bHasSSSProfile = SSSData.bIsValid;
+#else
 	FScreenSpaceData ScreenSpaceData = GetScreenSpaceData(UVSceneColor);
+	const bool bHasSSSProfile = UseSubsurfaceProfile(ScreenSpaceData.GBuffer.ShadingModelID);
+#endif
 
-	BRANCH if (UseSubsurfaceProfile(ScreenSpaceData.GBuffer.ShadingModelID))
+	float4 Ret = 0;
+	BRANCH if (bHasSSSProfile)
 	{
 		// '1' is lower quality but that is acceptable here
 		SDiffuseAndSpecular DiffuseAndSpecular = ReconstructLighting(UVSceneColor, 1);
@@ -404,7 +410,6 @@ float4 SetupSubsurfaceForOnePixel(float2 UVSceneColor)
 		// it's a valid sample
 		Ret.a = 1;
 	}
-
 	return Ret;
 }
 
@@ -528,10 +533,10 @@ void SetupIndirectCS(uint3 DT_ID : SV_DispatchThreadID,
 	// order aligned with Gather() hardware implementation
 	// RGB: color*A, A:weight 0 if no subsurface scattering
 
-	float4 A = SetupSubsurfaceForOnePixel(min(BufferUV + float2(-0.5, 0.5f) * SubsurfaceInput0_ExtentInverse, SubsurfaceInput0_UVViewportBilinearMax));
-	float4 B = SetupSubsurfaceForOnePixel(min(BufferUV + float2(0.5, 0.5f) * SubsurfaceInput0_ExtentInverse, SubsurfaceInput0_UVViewportBilinearMax));
-	float4 C = SetupSubsurfaceForOnePixel(min(BufferUV + float2(0.5, -0.5f) * SubsurfaceInput0_ExtentInverse, SubsurfaceInput0_UVViewportBilinearMax));
-	float4 D = SetupSubsurfaceForOnePixel(min(BufferUV + float2(-0.5, -0.5f) * SubsurfaceInput0_ExtentInverse, SubsurfaceInput0_UVViewportBilinearMax));
+	float4 A = SetupSubsurfaceForOnePixel(min(BufferUV + float2(-0.5, 0.5f) * SubsurfaceInput0_ExtentInverse, SubsurfaceInput0_UVViewportBilinearMax), SubsurfaceInput0_Extent);
+	float4 B = SetupSubsurfaceForOnePixel(min(BufferUV + float2(0.5, 0.5f) * SubsurfaceInput0_ExtentInverse, SubsurfaceInput0_UVViewportBilinearMax), SubsurfaceInput0_Extent);
+	float4 C = SetupSubsurfaceForOnePixel(min(BufferUV + float2(0.5, -0.5f) * SubsurfaceInput0_ExtentInverse, SubsurfaceInput0_UVViewportBilinearMax), SubsurfaceInput0_Extent);
+	float4 D = SetupSubsurfaceForOnePixel(min(BufferUV + float2(-0.5, -0.5f) * SubsurfaceInput0_ExtentInverse, SubsurfaceInput0_UVViewportBilinearMax), SubsurfaceInput0_Extent);
 
 	float4 Sum = (A + B) + (C + D);
 
@@ -559,7 +564,7 @@ void SetupIndirectCS(uint3 DT_ID : SV_DispatchThreadID,
 #else // SUBSURFACE_HALF_RES
 	
 	bool bHasSubsurface = false;
-	OutColor = SetupSubsurfaceForOnePixel(BufferUV);
+	OutColor = SetupSubsurfaceForOnePixel(BufferUV, SubsurfaceInput0_Extent);
 
 	if (OutColor.a > 0)
 	{
@@ -578,7 +583,7 @@ void SetupIndirectCS(uint3 DT_ID : SV_DispatchThreadID,
 #endif // SUBSURFACE_HALF_RES
 
 #if !SUBSURFACE_FORCE_SEPARABLE && ENABLE_PROFILE_ID_CACHE
-	uint ProfileId = GetSubsurfaceProfileId(BufferUV);
+	uint ProfileId = GetSubsurfaceProfileId(BufferUV, SubsurfaceInput0_Extent);
 #endif
 
 	// setup the Subsurface type
@@ -590,13 +595,13 @@ void SetupIndirectCS(uint3 DT_ID : SV_DispatchThreadID,
 	#if 0
 		// the code here is designed for half resolution burley.
 		float2 UVA = min(BufferUV + float2(-0.5, 0.5f) * SubsurfaceInput0_ExtentInverse, SubsurfaceInput0_UVViewportBilinearMax);
-		uint PidA = GetSubsurfaceProfileId(UVA);
+		uint PidA = GetSubsurfaceProfileId(UVA, SubsurfaceInput0_Extent);
 		float2 UVB = min(BufferUV + float2(0.5, 0.5f) * SubsurfaceInput0_ExtentInverse, SubsurfaceInput0_UVViewportBilinearMax);
-		uint PidB = GetSubsurfaceProfileId(UVB);
+		uint PidB = GetSubsurfaceProfileId(UVB, SubsurfaceInput0_Extent);
 		float2 UVC = min(BufferUV + float2(0.5, -0.5f) * SubsurfaceInput0_ExtentInverse, SubsurfaceInput0_UVViewportBilinearMax);
-		uint PidC = GetSubsurfaceProfileId(UVC);
+		uint PidC = GetSubsurfaceProfileId(UVC, SubsurfaceInput0_Extent);
 		float2 UVD = min(BufferUV + float2(-0.5, -0.5f) * SubsurfaceInput0_ExtentInverse, SubsurfaceInput0_UVViewportBilinearMax);
-		uint PidD = GetSubsurfaceProfileId(UVD);
+		uint PidD = GetSubsurfaceProfileId(UVD, SubsurfaceInput0_Extent);
 
 		// Gatter available and use the last subsurface type for the current pixel.
 		SelectedProfile = lerp(SelectedProfile, PidA, A.a);
@@ -610,14 +615,14 @@ void SetupIndirectCS(uint3 DT_ID : SV_DispatchThreadID,
 		BufferUV = lerp(BufferUV, UVC, C.a);
 		BufferUV = lerp(BufferUV, UVD, D.a);
 		
-		SelectedProfile = GetSubsurfaceProfileId(BufferUV);
+		SelectedProfile = GetSubsurfaceProfileId(BufferUV, SubsurfaceInput0_Extent);
 		Type = GetSubsurfaceType(SelectedProfile);
 	#else
 		Type = TYPE_SEPARABLE; // use separable for half resolution
 	#endif
 		
 #else
-		SelectedProfile = GetSubsurfaceProfileId(BufferUV);
+		SelectedProfile = GetSubsurfaceProfileId(BufferUV, SubsurfaceInput0_Extent);
 		Type = GetSubsurfaceType(SelectedProfile);
 #endif
 	}
@@ -654,10 +659,20 @@ void MainIndirectDispatchCS(
 
 	uint2 Pos = DT_ID.xy*SUBSURFACE_GROUP_SIZE / THREAD_SIZE_1D + Output_ViewportMin;
 	float2 BufferUV = ConvertGridPos2UV(Pos);
+	uint2  BufferPos = Pos;
+
+#if STRATA_ENABLED
+	const FStrataSubsurfaceData SSSData = StrataUnpackSSSData(Strata.SSSTexture.Load(uint3(BufferPos, 0)));
+	const uint bUseBurley = SSSData.bIsProfile ? GetSubsurfaceProfileUseBurley(SSSData.ProfileId) : SSSData.bIsValid;
+#else
+	const FScreenSpaceData ScreenSpaceData = GetScreenSpaceData(BufferUV);
+	const uint SubsurfaceProfileInt = ExtractSubsurfaceProfileInt(ScreenSpaceData.GBuffer);
+	const bool bUseBurley = GetSubsurfaceProfileUseBurley(SubsurfaceProfileInt);
+#endif
 
 	// We need to check Burley here because the previous pass might write burley samplings here unless we have SUBSURFACE_FORCE_SEPARABLE to overwrite the setting.
 #if !SUBSURFACE_FORCE_SEPARABLE
-	BRANCH if (GetSubsurfaceProfileUseBurley(BufferUV))
+	BRANCH if (bUseBurley)
 	{
 		return;
 	}
@@ -674,7 +689,7 @@ void MainIndirectDispatchCS(
 
 	ViewportDirectionUV *= (SubsurfaceInput0_ViewportSize.x * SubsurfaceInput0_ExtentInverse.x);
 
-	float4 OutColor = SSSSBlurPS(BufferUV, ViewportDirectionUV, false);
+	float4 OutColor = SSSSBlurPS(BufferPos, BufferUV, ViewportDirectionUV, false, SubsurfaceInput0_Extent);
 
 #if SUBSURFACE_PASS == SUBSURFACE_PASS_TWO
 	// second pass prepares the setup from the recombine pass which doesn't need depth but wants to reconstruct the color
@@ -693,16 +708,53 @@ void MainIndirectDispatchCS(
 // with the SceneColor.
 void SubsurfaceRecombinePS(noperspective float4 UVAndScreenPos : TEXCOORD0, float4 SvPosition : SV_POSITION, out float4 OutColor : SV_Target0)
 {
-	float2 BufferUV = UVAndScreenPos.xy;
+	const float2 BufferUV  = UVAndScreenPos.xy;
+	const uint2  BufferPos = SvPosition.xy;
 
+#if STRATA_ENABLED
+	const FStrataSubsurfaceData SSSData = StrataUnpackSSSData(Strata.SSSTexture.Load(uint3(BufferPos, 0)));
+	const bool bUseSSSProfile = SSSData.bIsValid;
+#else
 	FScreenSpaceData ScreenSpaceData = GetScreenSpaceData(BufferUV);
-	
-	if (!UseSubsurfaceProfile(ScreenSpaceData.GBuffer.ShadingModelID))
+	const bool bUseSSSProfile = UseSubsurfaceProfile(ScreenSpaceData.GBuffer.ShadingModelID);
+#endif
+
+	if (!bUseSSSProfile)
 	{
 		OutColor = Texture2DSample(SubsurfaceInput0_Texture, SharedSubsurfaceSampler0, BufferUV);
 		return;
 	} 
 
+#if STRATA_ENABLED
+	const uint ProfileId = SSSData.ProfileId;
+	// we multiply the base color later in to get more crips human skin textures (scanned data always has Subsurface included)
+	const float3 StoredBaseColor	= SSSData.BaseColor;
+	const float StoredSpecular		= SSSData.Specular;
+	const float PixelRadiusScale	= SSSData.ProfileRadiusScale;
+
+	// asset specific color
+	float3 SubsurfaceColor = 0.0f;
+	float ProfileRadiusScale = 1.0f;
+	if (SSSData.bIsProfile)
+	{
+		SubsurfaceColor = GetSubsurfaceProfileKernel(SSSS_SUBSURFACE_COLOR_OFFSET, ProfileId).rgb;
+		ProfileRadiusScale = GetSubsurfaceProfileKernel(SSSS_N_KERNELWEIGHTOFFSET + SSSS_N_KERNELWEIGHTCOUNT - 1, ProfileId).a;
+	}
+	else if (ProfileId == SSS_PROFILE_ID_PERPIXEL)
+	{
+		SubsurfaceColor = SSSData.DMFP;
+	}
+#else
+	const uint ProfileId = ExtractSubsurfaceProfileInt(ScreenSpaceData.GBuffer);
+	// we multiply the base color later in to get more crips human skin textures (scanned data always has Subsurface included)
+	const float3 StoredBaseColor	= ScreenSpaceData.GBuffer.StoredBaseColor;
+	const float StoredSpecular		= ScreenSpaceData.GBuffer.StoredSpecular;
+	const float PixelRadiusScale	= ScreenSpaceData.GBuffer.CustomData.a;
+	// asset specific color
+	const float3 SubsurfaceColor	= GetSubsurfaceProfileKernel(SSSS_SUBSURFACE_COLOR_OFFSET, ProfileId).rgb;
+	const float ProfileRadiusScale	= GetSubsurfaceProfileKernel(SSSS_N_KERNELWEIGHTOFFSET + SSSS_N_KERNELWEIGHTCOUNT - 1, ProfileId).a;
+#endif
+
 	float3 SSSColor = float3(0, 0, 0);
 	float LerpFactor = 1;
 
@@ -710,7 +762,7 @@ void SubsurfaceRecombinePS(noperspective float4 UVAndScreenPos : TEXCOORD0, floa
 #if SUBSURFACE_HALF_RES
 	// fade out subsurface scattering if radius is too small to be more crips (not blend with half resolution)
 	// minor quality improvement (faces are more detailed in distance)
-	LerpFactor = ComputeFullResLerp(ScreenSpaceData, BufferUV, SubsurfaceInput1_ExtentInverse);
+	LerpFactor = ComputeFullResLerp(ProfileRadiusScale, PixelRadiusScale, BufferUV, SubsurfaceInput1_ExtentInverse);
 #endif // SUBSURFACE_HALF_RES
 
 	{
@@ -726,12 +778,14 @@ void SubsurfaceRecombinePS(noperspective float4 UVAndScreenPos : TEXCOORD0, floa
 
 	float3 DirectColor = Texture2DSample(SubsurfaceInput0_Texture, SharedSubsurfaceSampler0, BufferUV).rgb;
 
-	uint ProfileId = ExtractSubsurfaceProfileInt(ScreenSpaceData.GBuffer);
-
 	bool bIsHalfResUpsample = false;
 	if (bIsHalfResUpsample && GetSubsurfaceProfileUseBurley(ProfileId))
 	{
-		FBurleyParameter BurleyParameter = GetBurleyParameters(ScreenSpaceData.GBuffer);
+	#if STRATA_ENABLED
+		FBurleyParameter BurleyParameter = GetBurleyParameters(SSSData);
+	#else
+		FBurleyParameter BurleyParameter = GetBurleyParameters(ProfileId, ScreenSpaceData.GBuffer);
+	#endif
 		
 		float Depth = CalcSceneDepth(BufferUV);
 
@@ -740,34 +794,18 @@ void SubsurfaceRecombinePS(noperspective float4 UVAndScreenPos : TEXCOORD0, floa
 		SSSColor = lerp(SSSColor,DirectColor,CenterSampleWeight);
 	}
 
-	// we multiply the base color later in to get more crips human skin textures (scanned data always has Subsurface included)
-	float3 StoredBaseColor = ScreenSpaceData.GBuffer.StoredBaseColor;
-	float StoredSpecular = ScreenSpaceData.GBuffer.StoredSpecular;
-
 	uint ReconstructMethod = (SUBSURFACE_RECOMBINE_QUALITY == SUBSURFACE_RECOMBINE_QUALITY_HIGH) ? 3 : 1;
 
 	SDiffuseAndSpecular DiffuseAndSpecular = ReconstructLighting(BufferUV, ReconstructMethod);
 	
 	float3 ExtractedNonSubsurface = DiffuseAndSpecular.Specular;
 
-	// asset specific color
-	float3 SubsurfaceColor = GetSubsurfaceProfileColor(ScreenSpaceData.GBuffer);
-
 	float3 FadedSubsurfaceColor = SubsurfaceColor * LerpFactor;
 
 	// combine potentially half res with full res
 	float3 SubsurfaceLighting = lerp(DiffuseAndSpecular.Diffuse, SSSColor, FadedSubsurfaceColor);
 
-#if STRATA_ENABLED
-	if (ProfileId == SSS_PROFILE_ID_PERPIXEL)
-	{
-		OutColor = float4(SubsurfaceLighting + ExtractedNonSubsurface, 0);
-	}
-	else
-#endif
-	{
-		OutColor = float4(SubsurfaceLighting * StoredBaseColor + ExtractedNonSubsurface, 0);
-	}
+	OutColor = float4(SubsurfaceLighting * StoredBaseColor + ExtractedNonSubsurface, 0);
 }
 
 void SubsurfaceViewportCopyPS(noperspective float4 InUV : TEXCOORD0, out float4 OutColor : SV_Target0)
diff --git a/Engine/Shaders/Private/SeparableSSS.ush b/Engine/Shaders/Private/SeparableSSS.ush
--- a/Engine/Shaders/Private/SeparableSSS.ush
+++ b/Engine/Shaders/Private/SeparableSSS.ush
@@ -114,7 +114,7 @@
 #define SSSSSampleShadowmap(coord)	float4(0,0,0,0)	// todo: not used yet
 #define SSSSSampleSceneColor(coord) GetSceneColor(coord)
 #define SSSSSampleSceneColorPoint(coord) GetSceneColor(coord) // todo: could be made point
-#define SSSSSampleProfileId(coord) GetSubsurfaceProfileId(coord)
+#define SSSSSampleProfileId(coord, ext) GetSubsurfaceProfileId(coord, ext)
 #define SSSSLerp(a, b, t) lerp(a, b, t)
 #define SSSSMad(a, b, c) mad(a, b, c)
 #define SSSSMul(v, m) mul(v, m)
@@ -182,7 +182,7 @@ float3 SSSSTransmittance(float translucency, float sssWidth, float3 worldPositio
 // @param texcoord The usual quad texture coordinates.
 // @param dir Direction of the blur: First pass:   float2(1.0, 0.0), Second pass:  float2(0.0, 1.0)
 // @param initStencil indicates whether the stencil buffer should be initialized. Should be set to 'true' for the first pass if not previously initialized, to enable optimization of the second pass.
-float4 SSSSBlurPS(float2 BufferUV, float2 dir, bool initStencil)
+float4 SSSSBlurPS(uint2 BufferPos, float2 BufferUV, float2 dir, bool initStencil, float2 Extent)
 {
     // Fetch color of current pixel:
     float4 colorM = SSSSSampleSceneColorPoint(BufferUV);
@@ -201,7 +201,7 @@ float4 SSSSBlurPS(float2 BufferUV, float2 dir, bool initStencil)
 	}
 
 	// 0..1
-	float SSSStrength = GetSubsurfaceStrength(BufferUV);
+	float SSSStrength = GetSubsurfaceStrength(BufferPos, BufferUV);
 
 #if !SSSS_COMPUTESHADER
 	// Initialize the stencil buffer in case it was not already available:
@@ -218,11 +218,17 @@ float4 SSSSBlurPS(float2 BufferUV, float2 dir, bool initStencil)
 	// ideally this comes from a half res buffer as well - there are some minor artifacts
 	finalStep *= SSSStrength; // Modulate it using the opacity (0..1 range)
     
-	FGBufferData GBufferData = GetGBufferData(BufferUV);
-
+#if STRATA_ENABLED
+	const FStrataSubsurfaceData SSSData = StrataUnpackSSSData(Strata.SSSTexture.Load(uint3(BufferPos, 0)));
+	// For SSS_PROFILE_ID_PERPIXEL is managed throught the burley passes
+	const uint SubsurfaceProfileInt = SSSData.bIsProfile ? SSSData.ProfileId : SSS_PROFILE_ID_INVALID;
+#else
+	const FGBufferData GBufferData = GetGBufferData(BufferUV);
 	// 0..255, which SubSurface profile to pick
 	// ideally this comes from a half res buffer as well - there are some minor artifacts
-	uint SubsurfaceProfileInt = ExtractSubsurfaceProfileInt(GBufferData);
+	const uint SubsurfaceProfileInt = ExtractSubsurfaceProfileInt(GBufferData);
+#endif
+
 
     // Accumulate the center sample:
     float3 colorAccum = 0;
@@ -233,7 +239,7 @@ float4 SSSSBlurPS(float2 BufferUV, float2 dir, bool initStencil)
 	colorInvDiv += GetSubsurfaceProfileKernel(SSSS_N_KERNELWEIGHTOFFSET, SubsurfaceProfileInt).rgb;
 	colorAccum = colorM.rgb * GetSubsurfaceProfileKernel(SSSS_N_KERNELWEIGHTOFFSET, SubsurfaceProfileInt).rgb;
 
-	float3 BoundaryColorBleed = GetSubsurfaceProfileBoundaryColorBleed(GBufferData);
+	float3 BoundaryColorBleed = GetSubsurfaceProfileBoundaryColorBleed(SubsurfaceProfileInt);
 
 	// Accumulate the other samples:
     SSSS_UNROLL
@@ -254,7 +260,7 @@ float4 SSSSBlurPS(float2 BufferUV, float2 dir, bool initStencil)
 			// Fetch color and depth for current sample:
 			float2 LocalUV = BufferUV + UVOffset * Side;
 			float4 color = SSSSSampleSceneColor(LocalUV);
-			uint LocalSubsurfaceProfileInt = SSSSSampleProfileId(LocalUV);
+			uint LocalSubsurfaceProfileInt = SSSSSampleProfileId(LocalUV, Extent);
 			float3 ColorTint = LocalSubsurfaceProfileInt == SubsurfaceProfileInt ? 1.0f : BoundaryColorBleed;
 
 
diff --git a/Engine/Shaders/Private/ShadowProjectionPixelShader.usf b/Engine/Shaders/Private/ShadowProjectionPixelShader.usf
--- a/Engine/Shaders/Private/ShadowProjectionPixelShader.usf
+++ b/Engine/Shaders/Private/ShadowProjectionPixelShader.usf
@@ -333,7 +333,7 @@ void Main(
 		#if STRATA_ENABLED
 					const FStrataTopLayerData TopLayerData = StrataUnpackTopLayerData(Strata.TopLayerNormalTexture.Load(uint3(SVPos.xy, 0)));
 					const float3 WorldNormal = TopLayerData.WorldNormal;
-					FTransmissionProfileParams TransmissionParams = GetTransmissionProfileParams(ExtractSubsurfaceProfileInt(SSSData.ProfilePackedId));
+					FTransmissionProfileParams TransmissionParams = GetTransmissionProfileParams(SSSData.ProfileId);
 		#else // STRATA_ENABLED
 					float3 WorldNormal = GBufferData.WorldNormal;
 					FTransmissionProfileParams TransmissionParams = GetTransmissionProfileParams(GBufferData);
@@ -483,7 +483,7 @@ void MainOnePassPointLightPS(
 
 	const bool IsSubsurfaceProfile = StrataData.bHasSubsurfaceProfile;
 	const float3 WorldNormal = TopLayerData.WorldNormal;
-	FTransmissionProfileParams TransmissionParams = GetTransmissionProfileParams(ExtractSubsurfaceProfileInt(SSSData.ProfilePackedId));
+	FTransmissionProfileParams TransmissionParams = GetTransmissionProfileParams(SSSData.ProfileId);
 
 #else // STRATA_ENABLED
 
diff --git a/Engine/Shaders/Private/Strata/Strata.ush b/Engine/Shaders/Private/Strata/Strata.ush
--- a/Engine/Shaders/Private/Strata/Strata.ush
+++ b/Engine/Shaders/Private/Strata/Strata.ush
@@ -290,9 +290,11 @@ FStrataTopLayerData StrataUnpackTopLayerData(uint In)
 struct FStrataSubsurfaceData
 {
 	float3 BaseColor;
+	float  Specular;
 
+	bool   bIsValid;
 	bool   bIsProfile;
-	float  ProfilePackedId;		// Integer ID packed in 8bits as float
+	uint   ProfileId;
 	float  ProfileRadiusScale;
 
 	float3 DMFP;				// When the profile is not used, a per pixel DMFP is used
@@ -301,28 +303,38 @@ struct FStrataSubsurfaceData
 uint2 StrataPackSSSData(FStrataSubsurfaceData In)
 {
 	uint2 Out;
-	Out.x = PackColorLinearToGamma2AlphaLinear(float4(In.BaseColor, In.bIsProfile ? 1.0f : 0.0f));
-	Out.y = In.bIsProfile ? PackRGBA8(float4(In.ProfilePackedId, In.ProfileRadiusScale, 0.0f, 0.0f)) : PackR11G11B10F(In.DMFP);
+	Out.x = PackColorLinearToGamma2AlphaLinear(float4(In.BaseColor, In.bIsValid ? (In.bIsProfile ? 1.0f : 0.5f) : 0.0f));
+	Out.y = In.bIsProfile ? ((In.ProfileId & 0xFF) | PackRGBA8(float4(0.f, In.ProfileRadiusScale, In.Specular, 0.0f))) : PackR11G11B10F(In.DMFP);
 	return Out;
 }
 
+uint StrataSubsurfaceProfileIdTo8bits(float In)
+{
+	// Similar encoding than ExtractSubsurfaceProfileInt. Valid profile ID start at 1.
+	return In * 255.0f + 0.5f;
+}
+
 FStrataSubsurfaceData StrataUnpackSSSData(uint2 In)
 {
 	FStrataSubsurfaceData Out = (FStrataSubsurfaceData)0;
 
 	float4 Data0 = UnpackColorGamma2ToLinearAlphaLinear(In.x);
 	Out.BaseColor = Data0.rgb;
-	Out.bIsProfile= Data0.a > 0.0f;
+	Out.bIsValid  = Data0.a > 0.0f;
+	Out.bIsProfile= Data0.a > 0.75f;
+	Out.ProfileId = SSS_PROFILE_ID_INVALID;
 
 	if (Out.bIsProfile)
 	{
 		float4 Data1 = UnpackRGBA8(In.y);
-		Out.ProfilePackedId = Data1.x;
+		Out.ProfileId = In.y & 0xFF;
 		Out.ProfileRadiusScale = Data1.y;
+		Out.Specular = Data1.z;
 	}
-	else
+	else if (Out.bIsValid)
 	{
-		Out.DMFP = PackR11G11B10F(In.y);
+		Out.ProfileId = SSS_PROFILE_ID_PERPIXEL;
+		Out.DMFP = UnpackR11G11B10F(In.y);
 	}
 	return Out;
 }
diff --git a/Engine/Shaders/Private/Strata/StrataEvaluation.ush b/Engine/Shaders/Private/Strata/StrataEvaluation.ush
--- a/Engine/Shaders/Private/Strata/StrataEvaluation.ush
+++ b/Engine/Shaders/Private/Strata/StrataEvaluation.ush
@@ -219,15 +219,16 @@ FStrataEvaluateResult StrataEvaluateBSDF(FStrataBSDFContext BSDFContext, FShadow
 	{
 		case STRATA_BSDF_TYPE_SLAB:
 		{
+			const bool bHasSubsurface = BSDF_GETHASSCATTERING(BSDFContext.BSDF) > 0;
 			const float3 BaseColor = SLAB_BASECOLOR(BSDFContext.BSDF);
 			const float  Metallic = SLAB_METALLIC(BSDFContext.BSDF);
-			const float3 DiffuseColor = BaseColor - BaseColor * Metallic;
+			const float3 DiffuseColor = bHasSubsurface ? float(1.f).xxx : BaseColor - BaseColor * Metallic;
 			const float SafeRoughness = MakeRoughnessSafe(SLAB_ROUGHNESSX(BSDFContext.BSDF));
 			float a2 = Pow4(SafeRoughness);
 
 			Sample.DiffusePathValue	= Diffuse_Chan(DiffuseColor, a2, BSDFContext.SatNoV, BSDFContext.SatNoL, BSDFContext.Context.VoH, BSDFContext.Context.NoH);
 			Sample.DiffusePDF		= BSDFContext.SatNoV / PI;
-			Sample.bSubsurface		= BSDF_GETHASSCATTERING(BSDFContext.BSDF) > 0;
+			Sample.bSubsurface		= bHasSubsurface;
 
 			const float3 F0  = ComputeF0(SLAB_SPECULAR(BSDFContext.BSDF), BaseColor, Metallic);
 			const float3 F90 = ComputeF90(F0, SLAB_EDGECOLOR(BSDFContext.BSDF), Metallic);
diff --git a/Engine/Shaders/Private/Strata/StrataMaterialClassification.usf b/Engine/Shaders/Private/Strata/StrataMaterialClassification.usf
--- a/Engine/Shaders/Private/Strata/StrataMaterialClassification.usf
+++ b/Engine/Shaders/Private/Strata/StrataMaterialClassification.usf
@@ -69,12 +69,13 @@ void ComputePostProcessBSDFData(
 			Classification.bHasSubsurface = true;
 
 			SSS.BaseColor = SLAB_BASECOLOR(BSDF);
-
+			SSS.Specular  = SLAB_SPECULAR(BSDF);
 			// STRATA_TODO when layered translucency is figured out once again, only bottom layer should run SSS
 
 			// Profile ID: we only keep the latest ID, and profile will overtake any non-profile based SSS
+			SSS.bIsValid			= true;
 			SSS.bIsProfile			= SSS.bIsProfile || BSDF_GETHASSSSPROFILE(BSDF);
-			SSS.ProfilePackedId		= SLAB_SSSPROFILEID(BSDF);
+			SSS.ProfileId			= StrataSubsurfaceProfileIdTo8bits(SLAB_SSSPROFILEID(BSDF));
 			SSS.ProfileRadiusScale	= SLAB_SSSPROFILERADIUSSCALE(BSDF);
 
 			// STRATA_TODO This should be accumulated according to the BSDF weight when we run this in the base pass.
diff --git a/Engine/Shaders/Private/Strata/StrataVisualize.usf b/Engine/Shaders/Private/Strata/StrataVisualize.usf
--- a/Engine/Shaders/Private/Strata/StrataVisualize.usf
+++ b/Engine/Shaders/Private/Strata/StrataVisualize.usf
@@ -199,7 +199,7 @@ void VisualizeMaterialPS(
 #define PRINTCLASSIFROUGH()		PRINTBSDFCHAR(_C_); PRINTBSDFCHAR(_L_); PRINTBSDFCHAR(_A_); PRINTBSDFCHAR(_S_); PRINTBSDFCHAR(_DOT_); PRINTBSDFCHAR(_R_) PRINTBSDFCHAR(_O_) PRINTBSDFCHAR(_U_) PRINTBSDFCHAR(_G_) PRINTBSDFCHAR(_H_) PRINTBSDFCHAR(_SPC_);	PRINTBSDFCHAR(_SPC_); PRINTFLOAT1(StrataData.TopLayerRoughness);
 #define PRINTCLASSIFBASCOL()	PRINTBSDFCHAR(_C_); PRINTBSDFCHAR(_L_); PRINTBSDFCHAR(_A_); PRINTBSDFCHAR(_S_); PRINTBSDFCHAR(_DOT_); PRINTBSDFCHAR(_B_) PRINTBSDFCHAR(_A_) PRINTBSDFCHAR(_S_) PRINTBSDFCHAR(_C_) PRINTBSDFCHAR(_O_) PRINTBSDFCHAR(_L_);	PRINTBSDFCHAR(_SPC_); PRINTFLOAT3(SSSData.BaseColor);
 #define PRINTCLASSIFISPROF()	PRINTBSDFCHAR(_C_); PRINTBSDFCHAR(_L_); PRINTBSDFCHAR(_A_); PRINTBSDFCHAR(_S_); PRINTBSDFCHAR(_DOT_); PRINTBSDFCHAR(_I_) PRINTBSDFCHAR(_S_) PRINTBSDFCHAR(_P_) PRINTBSDFCHAR(_R_) PRINTBSDFCHAR(_O_) PRINTBSDFCHAR(_F_);	PRINTBSDFCHAR(_SPC_); PRINTBSDFCHAR(SSSData.bIsProfile?_1_:_0_);
-#define PRINTCLASSIFPROID()		PRINTBSDFCHAR(_C_); PRINTBSDFCHAR(_L_); PRINTBSDFCHAR(_A_); PRINTBSDFCHAR(_S_); PRINTBSDFCHAR(_DOT_); PRINTBSDFCHAR(_P_) PRINTBSDFCHAR(_R_) PRINTBSDFCHAR(_O_) PRINTBSDFCHAR(_I_) PRINTBSDFCHAR(_D_) PRINTBSDFCHAR(_SPC_);	PRINTBSDFCHAR(_SPC_); PRINTFLOATNOFRAC(ExtractSubsurfaceProfileInt(SSSData.ProfilePackedId), 3);
+#define PRINTCLASSIFPROID()		PRINTBSDFCHAR(_C_); PRINTBSDFCHAR(_L_); PRINTBSDFCHAR(_A_); PRINTBSDFCHAR(_S_); PRINTBSDFCHAR(_DOT_); PRINTBSDFCHAR(_P_) PRINTBSDFCHAR(_R_) PRINTBSDFCHAR(_O_) PRINTBSDFCHAR(_I_) PRINTBSDFCHAR(_D_) PRINTBSDFCHAR(_SPC_);	PRINTBSDFCHAR(_SPC_); PRINTFLOATNOFRAC(SSSData.ProfileId, 3);
 #define PRINTCLASSIFPROSCL()	PRINTBSDFCHAR(_C_); PRINTBSDFCHAR(_L_); PRINTBSDFCHAR(_A_); PRINTBSDFCHAR(_S_); PRINTBSDFCHAR(_DOT_); PRINTBSDFCHAR(_P_) PRINTBSDFCHAR(_R_) PRINTBSDFCHAR(_O_) PRINTBSDFCHAR(_S_) PRINTBSDFCHAR(_C_) PRINTBSDFCHAR(_L_);	PRINTBSDFCHAR(_SPC_); PRINTFLOAT1(SSSData.ProfileRadiusScale);
 #define PRINTCLASSIFSSDMFP()	PRINTBSDFCHAR(_C_); PRINTBSDFCHAR(_L_); PRINTBSDFCHAR(_A_); PRINTBSDFCHAR(_S_); PRINTBSDFCHAR(_DOT_); PRINTBSDFCHAR(_S_) PRINTBSDFCHAR(_S_) PRINTBSDFCHAR(_D_) PRINTBSDFCHAR(_M_) PRINTBSDFCHAR(_F_) PRINTBSDFCHAR(_P_);	PRINTBSDFCHAR(_SPC_); PRINTFLOAT3(SSSData.DMFP);
 
diff --git a/Engine/Shaders/Private/SubsurfaceBurleyNormalized.ush b/Engine/Shaders/Private/SubsurfaceBurleyNormalized.ush
--- a/Engine/Shaders/Private/SubsurfaceBurleyNormalized.ush
+++ b/Engine/Shaders/Private/SubsurfaceBurleyNormalized.ush
@@ -184,7 +184,6 @@ float GetSubsurfaceProfileWorldUnitScale(uint SubsurfaceProfileInt)
 
 float SampleDepthTexture(float2 ScreenUV)
 {
-
 #if SUBSURFACE_HALF_RES && !(SUBSURFACE_SINGLE_PASS)
 	ConvertToDeviceZ(Texture2DSample(SubsurfaceInput1_Texture, SubsurfaceSampler1, ScreenUV).g);
 #else
@@ -197,61 +196,55 @@ float SampleDepthTexture(float2 ScreenUV)
 // fetch surface albedo and diffuse mean free path
 // The diffuse mean free path is modulated by world scale and a custom term passed in
 // which is the opacity texture.
-FBurleyParameter GetBurleyParameters(FGBufferData GBuffer)
+#if STRATA_ENABLED
+FBurleyParameter GetBurleyParameters(FStrataSubsurfaceData SSSData)
 {
 	FBurleyParameter Out = (FBurleyParameter)0;
-#if STRATA_ENABLED
-	if (GBuffer.ShadingModelID == SHADINGMODELID_STRATA)
+	if (SSSData.bIsProfile)
 	{
-		const bool bUseProfileData = GBuffer.CustomData.y == 0 && GBuffer.CustomData.z == 0;
-		if (bUseProfileData)
-		{
-			uint SubsurfaceProfileInt = ExtractSubsurfaceProfileInt(GBuffer);
-
-			Out.SurfaceAlbedo		= GetSubsurfaceProfileSurfaceAlbedo(SubsurfaceProfileInt);
-			Out.DiffuseMeanFreePath = DecodeDiffuseMeanFreePath(GetSubsurfaceProfileDiffuseMeanFreePath(SubsurfaceProfileInt));
-			Out.WorldUnitScale		= DecodeWorldUnitScale(GetSubsurfaceProfileWorldUnitScale(SubsurfaceProfileInt)) * BURLEY_CM_2_MM;
+		Out.SurfaceAlbedo		= GetSubsurfaceProfileSurfaceAlbedo(SSSData.ProfileId);
+		Out.DiffuseMeanFreePath = DecodeDiffuseMeanFreePath(GetSubsurfaceProfileDiffuseMeanFreePath(SSSData.ProfileId));
+		Out.WorldUnitScale		= DecodeWorldUnitScale(GetSubsurfaceProfileWorldUnitScale(SSSData.ProfileId)) * BURLEY_CM_2_MM;
 
-			Out.SurfaceOpacity		= GBuffer.CustomData.a;
-			Out.DiffuseMeanFreePath*= Out.SurfaceOpacity;
-		}
-		else
-		{
-			// Decode the DMFP stored as Float111110 in a uint over a RGBA8 unorm
-			uint PackedDMFP11_11_10 = PackRGBA8(GBuffer.CustomData);
-			float3 DMFP = UnpackR11G11B10F(PackedDMFP11_11_10);
+		Out.SurfaceOpacity		= SSSData.ProfileRadiusScale;
+		Out.DiffuseMeanFreePath*= Out.SurfaceOpacity;
+	}
+	else if(SSSData.bIsValid)
+	{
+		// Decode the DMFP stored as Float111110 in a uint over a RGBA8 unorm
+		float3 DMFP = SSSData.DMFP;
 
-			// Conversion to match other SSS encoding 
-			DMFP *= BURLEY_CM_2_MM * 5.0f; // STRATA_TODO why a scale of 5 is needed here to match between regular and Strata?
+		// Conversion to match other SSS encoding 
+		DMFP *= BURLEY_CM_2_MM * 5.0f; // STRATA_TODO why a scale of 5 is needed here to match between regular and Strata?
 
-			const float MaxDMFP				= max3(DMFP.x, DMFP.y, DMFP.z);
-			const uint IndexOfMaxDmfp		= (DMFP.r == MaxDMFP) ? 0 : ((DMFP.g == MaxDMFP) ? 1 : 2);
+		const float MaxDMFP				= max3(DMFP.x, DMFP.y, DMFP.z);
+		const uint IndexOfMaxDmfp		= (DMFP.r == MaxDMFP) ? 0 : ((DMFP.g == MaxDMFP) ? 1 : 2);
 
-			Out.SurfaceAlbedo				= float4(GBuffer.BaseColor, GBuffer.BaseColor[IndexOfMaxDmfp]);
-			Out.DiffuseMeanFreePath			= float4(DMFP, MaxDMFP);
-			Out.WorldUnitScale				= DecodeWorldUnitScale(1.0f * ENC_WORLDUNITSCALE_IN_CM_TO_UNIT) * BURLEY_CM_2_MM;
+		Out.SurfaceAlbedo				= float4(SSSData.BaseColor, SSSData.BaseColor[IndexOfMaxDmfp]);
+		Out.DiffuseMeanFreePath			= float4(DMFP, MaxDMFP);
+		Out.WorldUnitScale				= DecodeWorldUnitScale(1.0f * ENC_WORLDUNITSCALE_IN_CM_TO_UNIT) * BURLEY_CM_2_MM;
 
-			// STRATA_TODO smooth transistion based on pixel footprint w.r.t. MFPRadius and pixel footprint
-			Out.SurfaceOpacity				= 1.0f;
-			Out.DiffuseMeanFreePath			*= Out.SurfaceOpacity;
-		}
+		// STRATA_TODO smooth transistion based on pixel footprint w.r.t. MFPRadius and pixel footprint
+		Out.SurfaceOpacity				= 1.0f;
+		Out.DiffuseMeanFreePath			*= Out.SurfaceOpacity;
 	}
-	else
-#endif
-	{
-		uint SubsurfaceProfileInt = ExtractSubsurfaceProfileInt(GBuffer);
-
-		Out.SurfaceAlbedo		= GetSubsurfaceProfileSurfaceAlbedo(SubsurfaceProfileInt);
-		Out.DiffuseMeanFreePath	= DecodeDiffuseMeanFreePath(GetSubsurfaceProfileDiffuseMeanFreePath(SubsurfaceProfileInt));
-		Out.WorldUnitScale		= DecodeWorldUnitScale(GetSubsurfaceProfileWorldUnitScale(SubsurfaceProfileInt)) * BURLEY_CM_2_MM;
+	return Out;
+}
+#else
+FBurleyParameter GetBurleyParameters(uint SubsurfaceProfileInt, FGBufferData GBuffer)
+{
+	FBurleyParameter Out = (FBurleyParameter)0;
+	Out.SurfaceAlbedo		= GetSubsurfaceProfileSurfaceAlbedo(SubsurfaceProfileInt);
+	Out.DiffuseMeanFreePath = DecodeDiffuseMeanFreePath(GetSubsurfaceProfileDiffuseMeanFreePath(SubsurfaceProfileInt));
+	Out.WorldUnitScale		= DecodeWorldUnitScale(GetSubsurfaceProfileWorldUnitScale(SubsurfaceProfileInt)) * BURLEY_CM_2_MM;
 
-		// Mask should always be true, we can probably ignore it?
-		Out.SurfaceOpacity		= UseSubsurfaceProfile(GBuffer.ShadingModelID) ? GBuffer.CustomData.a : 0.0f;
-		Out.DiffuseMeanFreePath*= Out.SurfaceOpacity;
-	}
+	// Mask should always be true, we can probably ignore it?
+	Out.SurfaceOpacity = UseSubsurfaceProfile(GBuffer.ShadingModelID) ? GBuffer.CustomData.a : 0.0f;
+	Out.DiffuseMeanFreePath *= Out.SurfaceOpacity;
 
 	return Out;
 }
+#endif
 
 float2 Generate2DRandomNumber(int3 Seed)
 {
@@ -408,11 +401,19 @@ FBurleySampleInfo GenerateSampleInfo(float2 Rand0T1, float DiffuseMeanFreePathFo
 	return BurleySampleInfo;
 }
 
-float GetProfileMask(float2 BufferUV)
+float GetProfileMask(float2 BufferUV, float2 Extent)
 {
-	FScreenSpaceData ScreenSpaceData = GetScreenSpaceData(BufferUV);
+#if STRATA_ENABLED
+	const uint2 PixelCoord = BufferUV * Extent;
+	const FStrataSubsurfaceData SSSData = StrataUnpackSSSData(Strata.SSSTexture.Load(uint3(PixelCoord, 0)));
+	const bool bIsProfile = SSSData.bIsValid;
+#else
+	const FScreenSpaceData ScreenSpaceData = GetScreenSpaceData(BufferUV);
+	const bool bIsProfile = UseSubsurfaceProfile(ScreenSpaceData.GBuffer.ShadingModelID);	
+#endif
+
 	float Ret = 0;
-	BRANCH if (UseSubsurfaceProfile(ScreenSpaceData.GBuffer.ShadingModelID))
+	BRANCH if (bIsProfile)
 	{
 		Ret = 1.0f;
 	}
@@ -428,15 +429,15 @@ float GetProfileEdgeMask(float2 BufferUV)
 	#define PIXELOFFSET_UVDELTA 1.0f
 #endif
 
-	float P11 = GetProfileMask((BufferUV + float2(0.0f, 0.0f) * SubsurfaceInput0_ExtentInverse));
-	float P00 = GetProfileMask((BufferUV + float2(-PIXELOFFSET_UVDELTA, -PIXELOFFSET_UVDELTA) * SubsurfaceInput0_ExtentInverse));
-	float P01 = GetProfileMask((BufferUV + float2(0.0f, -PIXELOFFSET_UVDELTA) * SubsurfaceInput0_ExtentInverse));
-	float P02 = GetProfileMask((BufferUV + float2(PIXELOFFSET_UVDELTA, -PIXELOFFSET_UVDELTA) * SubsurfaceInput0_ExtentInverse));
-	float P10 = GetProfileMask((BufferUV + float2(-PIXELOFFSET_UVDELTA, 0.0f) * SubsurfaceInput0_ExtentInverse));
-	float P12 = GetProfileMask((BufferUV + float2(PIXELOFFSET_UVDELTA, 0.0f) * SubsurfaceInput0_ExtentInverse));
-	float P20 = GetProfileMask((BufferUV + float2(-PIXELOFFSET_UVDELTA, PIXELOFFSET_UVDELTA) * SubsurfaceInput0_ExtentInverse));
-	float P21 = GetProfileMask((BufferUV + float2(0.0f, PIXELOFFSET_UVDELTA) * SubsurfaceInput0_ExtentInverse));
-	float P22 = GetProfileMask((BufferUV + float2(PIXELOFFSET_UVDELTA, PIXELOFFSET_UVDELTA) * SubsurfaceInput0_ExtentInverse));
+	float P11 = GetProfileMask((BufferUV + float2(0.0f, 0.0f) * SubsurfaceInput0_ExtentInverse), SubsurfaceInput0_Extent);
+	float P00 = GetProfileMask((BufferUV + float2(-PIXELOFFSET_UVDELTA, -PIXELOFFSET_UVDELTA) * SubsurfaceInput0_ExtentInverse), SubsurfaceInput0_Extent);
+	float P01 = GetProfileMask((BufferUV + float2(0.0f, -PIXELOFFSET_UVDELTA) * SubsurfaceInput0_ExtentInverse), SubsurfaceInput0_Extent);
+	float P02 = GetProfileMask((BufferUV + float2(PIXELOFFSET_UVDELTA, -PIXELOFFSET_UVDELTA) * SubsurfaceInput0_ExtentInverse), SubsurfaceInput0_Extent);
+	float P10 = GetProfileMask((BufferUV + float2(-PIXELOFFSET_UVDELTA, 0.0f) * SubsurfaceInput0_ExtentInverse), SubsurfaceInput0_Extent);
+	float P12 = GetProfileMask((BufferUV + float2(PIXELOFFSET_UVDELTA, 0.0f) * SubsurfaceInput0_ExtentInverse), SubsurfaceInput0_Extent);
+	float P20 = GetProfileMask((BufferUV + float2(-PIXELOFFSET_UVDELTA, PIXELOFFSET_UVDELTA) * SubsurfaceInput0_ExtentInverse), SubsurfaceInput0_Extent);
+	float P21 = GetProfileMask((BufferUV + float2(0.0f, PIXELOFFSET_UVDELTA) * SubsurfaceInput0_ExtentInverse), SubsurfaceInput0_Extent);
+	float P22 = GetProfileMask((BufferUV + float2(PIXELOFFSET_UVDELTA, PIXELOFFSET_UVDELTA) * SubsurfaceInput0_ExtentInverse), SubsurfaceInput0_Extent);
 
 	return (P00 + P01 + P02 + P10 + P11 + P12 + P20 + P21 + P22) / 9.0f;
 }
@@ -483,6 +484,12 @@ int GetNumOfSamplesBasedOnQuality(float2 UV)
 
 // If the shading model is valid, return the profile. Otherwise, return SSS_PROFILE_ID_INVALID.
 // That way if we have a valid profile, we can check of the other profile is valid with a single boolean expression.
+#if STRATA_ENABLED
+uint ExtractSubsurfaceProfileIntWithInvalid(FStrataSubsurfaceData SSSData)
+{
+	return SSSData.ProfileId;
+}
+#else
 uint ExtractSubsurfaceProfileIntWithInvalid(FGBufferData BufferData)
 {
 	uint ProfileID = SSS_PROFILE_ID_INVALID;
@@ -492,6 +499,7 @@ uint ExtractSubsurfaceProfileIntWithInvalid(FGBufferData BufferData)
 	}
 	return ProfileID;
 }
+#endif
 
 // BufferUV back to grid posistion of the input
 uint2 ConvertUVPos2Input0GridPos(float2 BufferUV)
@@ -691,11 +699,17 @@ BurleySampleDiffuseNormal SampleSharedDiffuseNormal(float2 SampleGroupThreadId)
 	#define SharedBurleyPointSampler SubsurfaceSampler0
 #endif
 
-float3 GetBilateralNormal(float2 ClampedUV)
+float3 GetBilateralNormal(float2 ClampedUV, float2 Extent)
 {
 	//BRANCH to avoid extra evaluation of the GBuffer.
 #if BILATERAL_FILTER_KERNEL_FUNCTION_TYPE == BILATERAL_FILTER_KERNEL_FUNCTION_DEPTH_AND_NORMAL
+	#if STRATA_ENABLED
+	const uint2 BufferPos = ClampedUV * Extent;
+	const FStrataTopLayerData TopLayerData = StrataUnpackTopLayerData(Strata.TopLayerNormalTexture.Load(uint3(BufferPos, 0)));
+	return TopLayerData.WorldNormal;
+	#else
 	return GetScreenSpaceData(ClampedUV).GBuffer.WorldNormal;
+	#endif
 #elif BILATERAL_FILTER_KERNEL_FUNCTION_TYPE == BILATERAL_FILTER_KERNEL_FUNCTION_DEPTH
 	return BILATERAL_FILTER_DEFAULT_NORMAL;
 #endif
@@ -705,6 +719,15 @@ BurleySampleDiffuseNormal SampleSSSColorConsideringLocalShared(float2 CenterUV,
 {
 	float2 ClampedUV = clamp(CenterUV+ UVOffset, SubsurfaceInput0_UVViewportBilinearMin, SubsurfaceInput0_UVViewportBilinearMax);
 	
+#if !(ENABLE_PROFILE_ID_CACHE && (SUBSURFACE_FORCE_SEPARABLE !=1))
+	#if STRATA_ENABLED
+	const uint2 CenterPos = ClampedUV * SubsurfaceInput0_Extent;
+	const FStrataSubsurfaceData SSSData = StrataUnpackSSSData(Strata.SSSTexture.Load(uint3(CenterPos, 0)));
+	#else
+	const FGBufferData SSSData = GetScreenSpaceData(ClampedUV).GBuffer;
+	#endif
+#endif
+
 	// Burley works only with point sampler when we support world unit scale. Bilinear and trilinear will create artifacts.
 	// So we use SubsurfaceSampler1 here instead of using SubsurfaceSampler0, which is a point sampler.
 	// The reason that we can have different sampler for SubsurfaceSampler0 is that we have Separable running in the same pass.
@@ -721,13 +744,13 @@ BurleySampleDiffuseNormal SampleSSSColorConsideringLocalShared(float2 CenterUV,
 	Sample.DiffuseLighting = Texture2DSampleLevel(SubsurfaceInput0_Texture, SharedBurleyPointSampler, ClampedUV, MipLevel);
 
 #if BILATERAL_FILTER_KERNEL_FUNCTION_TYPE == BILATERAL_FILTER_KERNEL_FUNCTION_DEPTH_AND_NORMAL
-	Sample.WorldNormal.xyz = GetBilateralNormal(ClampedUV);
+	Sample.WorldNormal.xyz = GetBilateralNormal(ClampedUV, SubsurfaceInput0_Extent);
 #endif
 
-#if ENABLE_PROFILE_ID_CACHE && (SUBSURFACE_FALCE_SEPARABLE !=1)
+#if ENABLE_PROFILE_ID_CACHE && (SUBSURFACE_FORCE_SEPARABLE !=1)
 	Sample.ProfileID = ExtractSubsurfaceProfileIntWithInvalid(ClampedUV,Sample.DiffuseLighting.w);
 #else
-	Sample.ProfileID = ExtractSubsurfaceProfileIntWithInvalid(GetScreenSpaceData(ClampedUV).GBuffer);
+	Sample.ProfileID = ExtractSubsurfaceProfileIntWithInvalid(SSSData);
 #endif
 
 	return Sample;
@@ -745,13 +768,13 @@ BurleySampleDiffuseNormal SampleSSSColorConsideringLocalShared(float2 CenterUV,
 		{
 			Sample.DiffuseLighting = Texture2DSampleLevel(SubsurfaceInput0_Texture, SharedBurleyPointSampler, ClampedUV, MipLevel);
 #if BILATERAL_FILTER_KERNEL_FUNCTION_TYPE == BILATERAL_FILTER_KERNEL_FUNCTION_DEPTH_AND_NORMAL
-			Sample.WorldNormal.xyz = GetBilateralNormal(ClampedUV);
+			Sample.WorldNormal.xyz = GetBilateralNormal(ClampedUV, SubsurfaceInput0_Extent);
 #endif
 
 #if ENABLE_PROFILE_ID_CACHE && (SUBSURFACE_FORCE_SEPARABLE != 1)
 			Sample.ProfileID = ExtractSubsurfaceProfileIntWithInvalid(ClampedUV, Sample.DiffuseLighting.w);
 #else
-			Sample.ProfileID = ExtractSubsurfaceProfileIntWithInvalid(GetScreenSpaceData(ClampedUV).GBuffer);
+			Sample.ProfileID = ExtractSubsurfaceProfileIntWithInvalid(SSSData);
 #endif
 
 			return Sample;
@@ -762,13 +785,13 @@ BurleySampleDiffuseNormal SampleSSSColorConsideringLocalShared(float2 CenterUV,
 	Sample.DiffuseLighting = Texture2DSampleLevel(SubsurfaceInput0_Texture, SharedBurleyPointSampler, ClampedUV, MipLevel);
 
 #if BILATERAL_FILTER_KERNEL_FUNCTION_TYPE == BILATERAL_FILTER_KERNEL_FUNCTION_DEPTH_AND_NORMAL
-	Sample.WorldNormal.xyz = GetBilateralNormal(ClampedUV);
+	Sample.WorldNormal.xyz = GetBilateralNormal(ClampedUV, SubsurfaceInput0_Extent);
 #endif
 	
 #if ENABLE_PROFILE_ID_CACHE && (SUBSURFACE_FORCE_SEPARABLE != 1)
 	Sample.ProfileID = ExtractSubsurfaceProfileIntWithInvalid(ClampedUV, Sample.DiffuseLighting.w);
 #else
-	Sample.ProfileID = ExtractSubsurfaceProfileIntWithInvalid(GetScreenSpaceData(ClampedUV).GBuffer);
+	Sample.ProfileID = ExtractSubsurfaceProfileIntWithInvalid(SSSData);
 #endif
 
 	return Sample;
@@ -882,9 +905,6 @@ void UpdateSeed(int3 Seed3D, inout int StartSeed)
 
 float4 BurleyNormalizedSS(float2 BufferUV, uint2 GroupThreadID)
 {
-	FScreenSpaceData ScreenSpaceData = GetScreenSpaceData(BufferUV);
-	float3 WorldNormal = ScreenSpaceData.GBuffer.WorldNormal;
-
 	BurleySampleDiffuseNormal CenterSample = SampleSSSColorConsideringLocalShared(BufferUV, 0, GroupThreadID, 0);
 	float DepthAtDiscCenter = CenterSample.DiffuseLighting.w;
 
@@ -897,11 +917,23 @@ float4 BurleyNormalizedSS(float2 BufferUV, uint2 GroupThreadID)
 		return OutColor;
 	}
 
-	FBurleyParameter BurleyParameter = GetBurleyParameters(ScreenSpaceData.GBuffer);
+#if STRATA_ENABLED
+	const uint2 BufferPos = BufferUV * SubsurfaceInput0_Extent;
+	const FStrataSubsurfaceData SSSData = StrataUnpackSSSData(Strata.SSSTexture.Load(uint3(BufferPos, 0)));
+	const FStrataTopLayerData TopLayerData = StrataUnpackTopLayerData(Strata.TopLayerNormalTexture.Load(uint3(BufferPos, 0)));
+	const float3 WorldNormal = TopLayerData.WorldNormal;
+	const uint SubsurfaceProfileInt = SSSData.ProfileId;
+	const FBurleyParameter BurleyParameter = GetBurleyParameters(SSSData);
+#else
+	const FScreenSpaceData ScreenSpaceData = GetScreenSpaceData(BufferUV);
+	const float3 WorldNormal = ScreenSpaceData.GBuffer.WorldNormal;
+	const uint SubsurfaceProfileInt = ExtractSubsurfaceProfileInt(ScreenSpaceData.GBuffer);
+	const FBurleyParameter BurleyParameter = GetBurleyParameters(SubsurfaceProfileInt, ScreenSpaceData.GBuffer);
+#endif
+
 	float DiffuseMeanFreePathForSampling = GetDiffuseMeanFreePathForSampling(BurleyParameter.DiffuseMeanFreePath);
 	float A = GetComponentForScalingFactorEstimation(BurleyParameter.SurfaceAlbedo);
-	uint SubsurfaceProfileInt = ExtractSubsurfaceProfileInt(ScreenSpaceData.GBuffer);
-	float3 BoundaryColorBleed = GetSubsurfaceProfileBoundaryColorBleed(ScreenSpaceData.GBuffer).xyz;
+	float3 BoundaryColorBleed = GetSubsurfaceProfileBoundaryColorBleed(SubsurfaceProfileInt).xyz;
 
 	float S = GetScalingFactor(A);
 	float3 S3D = GetScalingFactor3D(BurleyParameter.SurfaceAlbedo.xyz);
@@ -991,9 +1023,16 @@ float4 BurleyNormalizedSS(float2 BufferUV, uint2 GroupThreadID)
 			SDiffuseAndSpecular SampledDiffuseAndSpecular = ReconstructLighting(SampledDiscUV, ReconstructMethod);
 			float4 SampledRadianceAndDepth = float4(SampledDiffuseAndSpecular.Diffuse, CalcSceneDepth(SampledDiscUV));
 
-			FScreenSpaceData SampleScreenSpaceData = GetScreenSpaceData(SampledDiscUV);
-			float3 SampleWorldNormal = SampleScreenSpaceData.GBuffer.WorldNormal;
-			uint LocalProfile = ExtractSubsurfaceProfileIntWithInvalid(GetSubsurfaceProfileId(SampledDiscUV));
+		#if STRATA_ENABLED
+			const uint2 BufferPos = SampledDiscUV * SubsurfaceInput0_Extent;
+			const FStrataTopLayerData TopLayerData = StrataUnpackTopLayerData(Strata.TopLayerNormalTexture.Load(uint3(BufferPos, 0));
+			const float3 SampleWorldNormal = TopLayerData.WorldNormal;
+		#else
+			const FScreenSpaceData SampleScreenSpaceData = GetScreenSpaceData(SampledDiscUV);
+			const float3 SampleWorldNormal = SampleScreenSpaceData.GBuffer.WorldNormal;
+		#endif
+
+			uint LocalProfile = ExtractSubsurfaceProfileIntWithInvalid(GetSubsurfaceProfileId(SampledDiscUV, SubsurfaceInput0_Extent));
 #else
 
 			// Determine the miplevel with the expected number of samples at the pixel.
@@ -1223,11 +1262,20 @@ void BurleyComputeMain(uint2 DT_ID, uint2 G_ID, uint GI)
 			SharedSubsurfaceDiffuseLighting[LocalSharedIndex] = SubsurfaceColorWithSSSIndicator;
 
 			// Also fetch the normal and profile ID
-			FScreenSpaceData SampleScreenSpaceData = GetScreenSpaceData(SampleUV);
+		#if STRATA_ENABLED
+			const uint2 BufferPos = SampleUV * Output_Extent;
+			const FStrataTopLayerData TopLayerData = StrataUnpackTopLayerData(Strata.TopLayerNormalTexture.Load(uint3(BufferPos, 0)));
+			const float3 SampleWorldNormal = TopLayerData.WorldNormal;
+			const FStrataSubsurfaceData SampleSSSData = StrataUnpackSSSData(Strata.SSSTexture.Load(uint3(BufferPos, 0)));
+		#else
+			const FGBufferData SampleSSSData = GetScreenSpaceData(SampleUV).GBuffer;
+			const float3 SampleWorldNormal = SampleSSSData.WorldNormal;
+		#endif
+
 #if BILATERAL_FILTER_KERNEL_FUNCTION_TYPE == BILATERAL_FILTER_KERNEL_FUNCTION_DEPTH_AND_NORMAL
-			SharedSubsurfaceWorldNormal[LocalSharedIndex] = SampleScreenSpaceData.GBuffer.WorldNormal;
+			SharedSubsurfaceWorldNormal[LocalSharedIndex] = SampleWorldNormal;
 #endif
-			SharedSubsurfaceProfileID[LocalSharedIndex] = ExtractSubsurfaceProfileIntWithInvalid(SampleScreenSpaceData.GBuffer);
+			SharedSubsurfaceProfileID[LocalSharedIndex] = ExtractSubsurfaceProfileIntWithInvalid(SampleSSSData);
 		}
 	}
 
@@ -1274,7 +1322,7 @@ void BurleyComputeMain(uint2 DT_ID, uint2 G_ID, uint GI)
 			float4 SurfaceColor = Texture2DSample(SubsurfaceInput0_Texture, SubsurfaceSampler0, LBufferUV);
 
 			// We have to check per pixel if this is a burley so that we don't overwrite separable pixels
-			uint SelectedProfile = GetSubsurfaceProfileId(LBufferUV);
+			uint SelectedProfile = GetSubsurfaceProfileId(LBufferUV, SubsurfaceInput0_Extent);
 			bool UseBurley = GetSubsurfaceProfileUseBurley(SelectedProfile);
 			if (UseBurley)
 			{
diff --git a/Engine/Source/Runtime/Renderer/Private/PostProcess/PostProcessSubsurface.cpp b/Engine/Source/Runtime/Renderer/Private/PostProcess/PostProcessSubsurface.cpp
--- a/Engine/Source/Runtime/Renderer/Private/PostProcess/PostProcessSubsurface.cpp
+++ b/Engine/Source/Runtime/Renderer/Private/PostProcess/PostProcessSubsurface.cpp
@@ -371,6 +371,7 @@ public:
 	SHADER_USE_PARAMETER_STRUCT(FSubsurfaceVisualizePS, FSubsurfaceShader);
 
 	BEGIN_SHADER_PARAMETER_STRUCT(FParameters, )
+		SHADER_PARAMETER_STRUCT_REF(FStrataGlobalUniformParameters, Strata)
 		SHADER_PARAMETER_STRUCT_INCLUDE(FSubsurfaceParameters, Subsurface)
 		SHADER_PARAMETER_STRUCT(FSubsurfaceInput, SubsurfaceInput0)
 		SHADER_PARAMETER_TEXTURE(Texture2D, MiniFontTexture)
@@ -489,6 +490,7 @@ public:
 	}
 
 	BEGIN_SHADER_PARAMETER_STRUCT(FParameters, )
+		SHADER_PARAMETER_STRUCT_REF(FStrataGlobalUniformParameters, Strata)
 		SHADER_PARAMETER_STRUCT_INCLUDE(FSubsurfaceParameters, Subsurface)
 		SHADER_PARAMETER_STRUCT(FScreenPassTextureViewportParameters, Output)
 		SHADER_PARAMETER_STRUCT(FSubsurfaceInput, SubsurfaceInput0)
@@ -525,6 +527,7 @@ public:
 	}
 
 	BEGIN_SHADER_PARAMETER_STRUCT(FParameters, )
+		SHADER_PARAMETER_STRUCT_REF(FStrataGlobalUniformParameters, Strata)
 		SHADER_PARAMETER_STRUCT_INCLUDE(FSubsurfaceParameters, Subsurface)
 		SHADER_PARAMETER_STRUCT(FScreenPassTextureViewportParameters, Output)
 		SHADER_PARAMETER_RDG_TEXTURE_UAV(RWTexture2D, SSSColorUAV)
@@ -749,13 +752,14 @@ class FSubsurfaceRecombinePS : public FSubsurfaceShader
 	SHADER_USE_PARAMETER_STRUCT(FSubsurfaceRecombinePS, FSubsurfaceShader);
 
 	BEGIN_SHADER_PARAMETER_STRUCT(FParameters, )
+		SHADER_PARAMETER_STRUCT_REF(FStrataGlobalUniformParameters, Strata)
 		SHADER_PARAMETER_STRUCT_INCLUDE(FSubsurfaceParameters, Subsurface)
 		SHADER_PARAMETER_STRUCT(FSubsurfaceInput, SubsurfaceInput0)
 		SHADER_PARAMETER_STRUCT(FSubsurfaceInput, SubsurfaceInput1)
 		SHADER_PARAMETER_SAMPLER(SamplerState, SubsurfaceSampler0)
 		SHADER_PARAMETER_SAMPLER(SamplerState, SubsurfaceSampler1)
 		RENDER_TARGET_BINDING_SLOTS()
-		END_SHADER_PARAMETER_STRUCT();
+	END_SHADER_PARAMETER_STRUCT();
 
 	// Controls the quality of lighting reconstruction.
 	enum class EQuality : uint32
@@ -988,6 +992,7 @@ void AddSubsurfaceViewPass(
 			PassParameters->RWBurleyGroupBuffer = GraphBuilder.CreateUAV(BurleyGroupBuffer, EPixelFormat::PF_R32_UINT);
 			PassParameters->RWSeparableGroupBuffer = GraphBuilder.CreateUAV(SeparableGroupBuffer, EPixelFormat::PF_R32_UINT);
 			PassParameters->SubsurfaceUniformParameters = UniformBuffer;
+			PassParameters->Strata = Strata::BindStrataGlobalUniformParameters(View);
 
 			SHADER::FPermutationDomain ComputeShaderPermutationVector;
 			ComputeShaderPermutationVector.Set<SHADER::FDimensionHalfRes>(bHalfRes);
@@ -1091,6 +1096,7 @@ void AddSubsurfaceViewPass(
 				PassParameters->SubsurfaceSampler0 = SubsurfaceSamplerState;
 				PassParameters->GroupBuffer = GraphBuilder.CreateSRV(SubsurfaceBufferUsage[SubsurfaceTypeIndex], EPixelFormat::PF_R32_UINT);
 				PassParameters->IndirectDispatchArgsBuffer = SubsurfaceBufferArgs[SubsurfaceTypeIndex];
+				PassParameters->Strata = Strata::BindStrataGlobalUniformParameters(View);
 
 				if (SubsurfacePassFunction == SHADER::ESubsurfacePass::PassOne && SubsurfaceType == SHADER::ESubsurfaceType::BURLEY)
 				{
@@ -1150,6 +1156,7 @@ void AddSubsurfaceViewPass(
 		PassParameters->RenderTargets[0] = FRenderTargetBinding(SceneColorTextureOutput, SceneColorTextureLoadAction);
 		PassParameters->SubsurfaceInput0 = GetSubsurfaceInput(SceneColorTexture, SceneViewportParameters);
 		PassParameters->SubsurfaceSampler0 = BilinearBorderSampler;
+		PassParameters->Strata = Strata::BindStrataGlobalUniformParameters(View);
 
 		// Scattering output target is only used when scattering is enabled.
 		if (SubsurfaceMode != ESubsurfaceMode::Bypass)
@@ -1285,6 +1292,7 @@ FScreenPassTexture AddVisualizeSubsurfacePass(FRDGBuilder& GraphBuilder, const F
 	PassParameters->SubsurfaceInput0.Viewport = GetScreenPassTextureViewportParameters(InputViewport);
 	PassParameters->SubsurfaceSampler0 = TStaticSamplerState<SF_Point, AM_Clamp, AM_Clamp, AM_Clamp>::GetRHI();
 	PassParameters->MiniFontTexture = GetMiniFontTexture();
+	PassParameters->Strata = Strata::BindStrataGlobalUniformParameters(View);
 
 	TShaderMapRef<FSubsurfaceVisualizePS> PixelShader(View.ShaderMap);
 
