- Add support to mark coarse pages in mips/clipmap levels
- Fix optimal plane bias with clipmap LOD fallback
- Clean up and optimize some of the clipmap level offset math
- Rename some fields to make the resuse between clipmaps/mips clear

diff --git a/Engine/Shaders/Private/VirtualShadowMaps/Debug.usf b/Engine/Shaders/Private/VirtualShadowMaps/Debug.usf
--- a/Engine/Shaders/Private/VirtualShadowMaps/Debug.usf
+++ b/Engine/Shaders/Private/VirtualShadowMaps/Debug.usf
@@ -66,7 +66,7 @@ void DebugVisualizeVirtualSmCS(uint2 OutputPixel : SV_DispatchThreadID)
 			uint2 vPage = vAddress >> VSM_LOG2_PAGE_SIZE;
 			FShadowPhysicalPage pPage = ShadowGetPhysicalPage( CalcPageOffset( ShadowMapID, Level, vPage ) );
 
-			if( pPage.bThisMipValid )
+			if( pPage.bThisLODValid )
 			{
 				uint2 pAddress = pPage.PageIndex * VSM_PAGE_SIZE + ( vAddress & VSM_PAGE_SIZE_MASK );
 				Depth = asfloat( PhysicalPagePool[ pAddress ] );
@@ -82,7 +82,7 @@ void DebugVisualizeVirtualSmCS(uint2 OutputPixel : SV_DispatchThreadID)
 			//uint2 pPage = HZBPageTable[ CalcPageOffset( ShadowMapID, Level, vPage ) ];
 			FShadowPhysicalPage pPage = ShadowGetPhysicalPage( CalcPageOffset( ShadowMapID, Level, vPage ) );
 
-			if( pPage.bThisMipValid )
+			if( pPage.bThisLODValid )
 			{
 				vAddress >>= 1;
 
@@ -176,7 +176,7 @@ void DebugVisualizeVirtualSmCS(uint2 OutputPixel : SV_DispatchThreadID)
 			{
 				uint Offset = CalcPageOffset(VirtualShadowMapId, MipLevel, PageAddress);
 				uint PageFlag = PageFlags[Offset];
-				bool PageValid = ShadowGetPhysicalPage(Offset).bThisMipValid;
+				bool PageValid = ShadowGetPhysicalPage(Offset).bThisLODValid;
 				float3 Color = float3(0, !PageValid && PageFlag, 0);
 				DebugOutput[OutputPixel] = float4(Color, 1.0f);
 			}
diff --git a/Engine/Shaders/Private/VirtualShadowMaps/PageAccessCommon.ush b/Engine/Shaders/Private/VirtualShadowMaps/PageAccessCommon.ush
--- a/Engine/Shaders/Private/VirtualShadowMaps/PageAccessCommon.ush
+++ b/Engine/Shaders/Private/VirtualShadowMaps/PageAccessCommon.ush
@@ -137,15 +137,15 @@ uint PhysPageAddressToIndex(uint2 PhysicalPageAddress)
 // NOTE: Some redundancy in flags and encoding, but we have spare bits for now
 //   [0:9] PageIndex.x
 //   [10:19] PageIndex.y
-//   [20:23] MipOffset
-//   [24:30] (currently unused)
-//   [31] bAnyMipValid
+//   [20:25] LODOffset
+//   [26:30] (currently unused)
+//   [31] bAnyLODValid
 struct FShadowPhysicalPage
 {
 	uint2 PageIndex;		// Physical page index X, Y
-	uint MipOffset;			// 0 if page is mapped at this mip; 1 if mapped at next courser level, etc. [0..16)
-	bool bAnyMipValid;		// Valid physical page mapped at some mip level
-	bool bThisMipValid;		// Valid page mapped at this specific mip level (equivalent to bAnyMipValid && MipOffset == 0)
+	uint LODOffset;			// 0 if page is mapped at this mip/clipmap level; 1 if mapped at next courser level, etc. [0..64)
+	bool bAnyLODValid;		// Valid physical page mapped at some LOD level
+	bool bThisLODValid;		// Valid page mapped at this specific level (equivalent to bAnyMipValid && LODOffset == 0)
 };
 
 #define VSM_PHYSICAL_PAGE_ANY_MIP_VALID_FLAG 0x8000000
@@ -158,18 +158,18 @@ uint ShadowEncodePageTable(uint2 PageIndex)
 {
 	return VSM_PHYSICAL_PAGE_ANY_MIP_VALID_FLAG | (PageIndex.y << 10) | (PageIndex.x);
 }
-uint ShadowEncodePageTable(uint2 PageIndex, uint MipOffset)
+uint ShadowEncodePageTable(uint2 PageIndex, uint LODOffset)
 {
-	return VSM_PHYSICAL_PAGE_ANY_MIP_VALID_FLAG | (MipOffset << 20) | (PageIndex.y << 10) | (PageIndex.x);
+	return VSM_PHYSICAL_PAGE_ANY_MIP_VALID_FLAG | (LODOffset << 20) | (PageIndex.y << 10) | (PageIndex.x);
 }
 
 FShadowPhysicalPage ShadowDecodePageTable(uint Value)
 {
 	FShadowPhysicalPage Result;
 	Result.PageIndex = uint2(Value & 0x3FF, (Value >> 10) & 0x3FF);
-	Result.MipOffset = (Value >> 20) & 0xF;
-	Result.bAnyMipValid = (Value & VSM_PHYSICAL_PAGE_ANY_MIP_VALID_FLAG) != 0;
-	Result.bThisMipValid = Result.bAnyMipValid && Result.MipOffset == 0;
+	Result.LODOffset = (Value >> 20) & 0x3F;
+	Result.bAnyLODValid = (Value & VSM_PHYSICAL_PAGE_ANY_MIP_VALID_FLAG) != 0;
+	Result.bThisLODValid = Result.bAnyLODValid && Result.LODOffset == 0;
 	return Result;
 }
 
@@ -185,13 +185,13 @@ bool VirtualToPhysicalTexel(uint ShadowMapID, uint Level, uint2 VirtualTexelAddr
 
 	FShadowPhysicalPage PhysicalPageEntry = ShadowGetPhysicalPage(CalcPageOffset(ShadowMapID, Level, uint2(VPageX, VPageY)));
 	PhysicalTexelAddress = PhysicalPageEntry.PageIndex * VSM_PAGE_SIZE + (VirtualTexelAddress & VSM_PAGE_SIZE_MASK);
-	return (PhysicalPageEntry.bThisMipValid);
+	return (PhysicalPageEntry.bThisLODValid);
 }
 
 struct FShadowPageTranslationResult
 {
 	bool bValid;
-	uint MipLevel;
+	uint LODOffset;
 	uint2 VirtualTexelAddress;
 	float2 VirtualTexelAddressFloat;
 	uint2 PhysicalTexelAddress;
@@ -207,10 +207,10 @@ FShadowPageTranslationResult ShadowVirtualToPhysicalUV(uint VirtualShadowMapID,
 	FShadowPhysicalPage PhysicalPageEntry = ShadowGetPhysicalPage(CalcPageOffset(VirtualShadowMapID, 0, vPageLevel0));
 
 	FShadowPageTranslationResult Result;
-	Result.bValid = PhysicalPageEntry.bAnyMipValid;
-	Result.MipLevel = PhysicalPageEntry.MipOffset;	
+	Result.bValid = PhysicalPageEntry.bAnyLODValid;
+	Result.LODOffset = PhysicalPageEntry.LODOffset;
 	// TODO: Can optimize this slightly based on relative offset
-	Result.VirtualTexelAddressFloat = ShadowMapUV * float(CalcLevelDimsTexels(Result.MipLevel));
+	Result.VirtualTexelAddressFloat = ShadowMapUV * float(CalcLevelDimsTexels(Result.LODOffset));
 	Result.VirtualTexelAddress = uint2(Result.VirtualTexelAddressFloat);
 	Result.PhysicalTexelAddress = PhysicalPageEntry.PageIndex * VSM_PAGE_SIZE + (Result.VirtualTexelAddress & VSM_PAGE_SIZE_MASK);
 
@@ -323,7 +323,7 @@ bool IsVisibleHZB( uint ShadowMapID, uint MipLevel, FScreenRect Rect )
 			uint PageOffset = PageTableLevelOffset + CalcPageOffsetInLevel( MipLevel, vPage );
 			FShadowPhysicalPage pPage = ShadowDecodePageTable( ShadowHZBPageTable[ PageOffset ] );
 
-			if( !pPage.bThisMipValid )
+			if( !pPage.bThisLODValid )
 			{
 				// If this is a new page we can't consider everything occluded the first time.
 				//if( PageFlags[ PageOffset ] != 0 )
diff --git a/Engine/Shaders/Private/VirtualShadowMaps/PageManagement.usf b/Engine/Shaders/Private/VirtualShadowMaps/PageManagement.usf
--- a/Engine/Shaders/Private/VirtualShadowMaps/PageManagement.usf
+++ b/Engine/Shaders/Private/VirtualShadowMaps/PageManagement.usf
@@ -258,6 +258,52 @@ void GeneratePageFlagsFromPixels(uint3 GroupId : SV_GroupID, uint GroupIndex : S
 	}
 }
 
+int ClipmapFirstCoarseLevel;
+
+[numthreads(VSM_DEFAULT_CS_GROUP_X, 1, 1)]
+void MarkCoarsePages(uint3 Index : SV_DispatchThreadID)
+{
+	int VirtualShadowMapId = int(Index.x);
+	if (VirtualShadowMapId >= int(VirtualShadowMap.NumShadowMaps))
+	{
+		return;
+	}
+
+	FVirtualShadowMapProjectionShaderData ProjectionData = GetVirtualShadowMapProjectionData(VirtualShadowMapId);
+
+	if (ProjectionData.LightType == LIGHT_TYPE_DIRECTIONAL)
+	{
+		// Idea here is the clipmaps already cover supersets of lower levels
+		// Thus to get coarser pages we can just mark the center page(s) offset by a level/LOD bias
+		// The limit on how far dense data goes out from the camera then becomes the world space size of the marked page(s) on the coarses clipmap
+		// We could of course mark a broader set of pages in the coarses clipmap level, but the effective radius
+		// even from just marking a single one is usually already large enough for the systems that need this
+		// data (volumetric fog, translucent light volume).
+		
+		if (ProjectionData.ClipmapLevel > ClipmapFirstCoarseLevel)
+		{
+			// TODO: Optimize this... can be boiled down to be just in terms of the snap offsets
+			float4 ShadowUVz = mul(float4(ProjectionData.ClipmapWorldOrigin + ProjectionData.ShadowPreViewTranslation, 1.0f), ProjectionData.TranslatedWorldToShadowUVMatrix);
+			float2 VirtualTexelAddressFloat = ShadowUVz.xy * float(CalcLevelDimsTexels(0));
+			float2 PageAddressFloat = VirtualTexelAddressFloat * float(1.0f / VSM_PAGE_SIZE);
+			// NOTE: Page addresses round down/truncate normally, so grab the surrounding 4
+			int4 PageAddressLowHigh = int4(floor(PageAddressFloat - 0.5f), ceil(PageAddressFloat - 0.5f));
+
+			OutPageRequestFlags[CalcPageOffset(VirtualShadowMapId, 0, PageAddressLowHigh.xy)] = 1;
+			OutPageRequestFlags[CalcPageOffset(VirtualShadowMapId, 0, PageAddressLowHigh.xw)] = 1;
+			OutPageRequestFlags[CalcPageOffset(VirtualShadowMapId, 0, PageAddressLowHigh.zy)] = 1;
+			OutPageRequestFlags[CalcPageOffset(VirtualShadowMapId, 0, PageAddressLowHigh.zw)] = 1;
+		}
+	}
+	else
+	{
+		// Mark last mip
+		uint PageOffset = CalcPageOffset(VirtualShadowMapId, VSM_MAX_MIP_LEVELS - 1, uint2(0, 0));
+		OutPageRequestFlags[PageOffset] = 1;
+	}
+}
+
+
 
 RWStructuredBuffer<uint> OutHPageFlags;
 RWStructuredBuffer<uint4> PageRectBoundsOut;
@@ -616,17 +662,17 @@ void PropagateMappedMips(uint2 ThreadId : SV_DispatchThreadID, uint2 GroupId : S
 		// Directional lights propagate pages to their coarser/larger clipmap levels (and only use mip0 pages)
 		// Each clipmap level is a separate VSM, so we gather any mapped coarser pages as necessary and write only our own page output
 		// There's also technically a race similar to below with other threads writing the PT data we are reading,
-		// but it's still deterministic as long as we only look at pages with "bThisMipValid".
+		// but it's still deterministic as long as we only look at pages with "bThisLODValid".
 		// There's some redundant work of course, but this shader is pretty cheap overall
 
 		uint Page0Offset = CalcPageOffset(VirtualShadowMapID, 0, Level0Page);
 		FShadowPhysicalPage pPage0 = ShadowDecodePageTable(OutPageTable[Page0Offset]);
 
 		BRANCH
-		if (!pPage0.bThisMipValid)
+		if (!pPage0.bThisLODValid)
 		{
 			const int OffsetScale = (VSM_LEVEL0_DIM_PAGES_XY >> 2);
-			int2 BaseOffset = OffsetScale * (ProjectionData.ClipmapCenterSnapUnits * int2(-1, 1) + 2);
+			int2 BaseOffset = OffsetScale * ProjectionData.ClipmapCornerOffset;
 			int2 BasePage   = int2(Level0Page) - BaseOffset;
 
 			// Search for first mapped page past this one
@@ -636,7 +682,7 @@ void PropagateMappedMips(uint2 ThreadId : SV_DispatchThreadID, uint2 GroupId : S
 				const int ClipmapLevelId = VirtualShadowMapID + int(ClipmapOffset);
 
 				FVirtualShadowMapProjectionShaderData LevelProjectionData = GetVirtualShadowMapProjectionData(ClipmapLevelId);								
-				int2 LevelOffset = OffsetScale * (LevelProjectionData.ClipmapCenterSnapUnits * int2(-1, 1) + 2);
+				int2 LevelOffset = OffsetScale * LevelProjectionData.ClipmapCornerOffset;
 
 				const int LevelScale = 1 << ClipmapOffset;
 				int2 LevelPage = (BasePage + LevelScale * LevelOffset) / LevelScale;
@@ -645,7 +691,7 @@ void PropagateMappedMips(uint2 ThreadId : SV_DispatchThreadID, uint2 GroupId : S
 				{
 					uint LevelPageOffset = CalcPageOffset(ClipmapLevelId, 0, uint2(LevelPage));
 					FShadowPhysicalPage pPage = ShadowDecodePageTable(OutPageTable[LevelPageOffset]);
-					if (pPage.bThisMipValid)
+					if (pPage.bThisLODValid)
 					{
 						OutPageTable[Page0Offset] = ShadowEncodePageTable(pPage.PageIndex, ClipmapOffset);
 						break;
@@ -653,8 +699,7 @@ void PropagateMappedMips(uint2 ThreadId : SV_DispatchThreadID, uint2 GroupId : S
 				}
 				else
 				{
-					// TODO: We're off the edge... can we ever get back on?
-					// TODO: Can this ever even happen in practice given the construction?
+					// TODO: We're off the edge... can this ever even happen in practice given the construction?
 				}
 			}
 		}
@@ -672,9 +717,9 @@ void PropagateMappedMips(uint2 ThreadId : SV_DispatchThreadID, uint2 GroupId : S
 			FShadowPhysicalPage pPage = ShadowDecodePageTable(OutPageTable[PageOffset]);
 
 			BRANCH
-			if (pPage.bThisMipValid)
+			if (pPage.bThisLODValid)
 			{
-				// This page is mapped, so leave it alone and propogate downwards
+				// This page is mapped, so leave it alone and propagate downwards
 				MappedPageLevel = Level;
 				MappedPageIndex = pPage.PageIndex;
 			}
@@ -682,7 +727,7 @@ void PropagateMappedMips(uint2 ThreadId : SV_DispatchThreadID, uint2 GroupId : S
 			{
 				// This page is not mapped; replace it with our suitably offset parent mapped page
 				// Ensure only one thread writes each value to avoid races, but we read on all threads as the broadcast
-				// Note that this can race with the other threads reading this value, but since bThisMipValid will
+				// Note that this can race with the other threads reading this value, but since bThisLODValid will
 				// always be false on these updated pages the values will be ignored. As long as the writes to the page
 				// table are atomic (currently a single DWORD), this is safe.
 				if (all((vPage << Level) == Level0Page))
@@ -896,7 +941,7 @@ void MarkRenderedPhysicalPages(uint3 PageTableEntryIndex : SV_DispatchThreadID)
 	}
 
 	FShadowPhysicalPage pPage = ShadowGetPhysicalPage(PageTableEntryIndex.x);
-	if (pPage.bThisMipValid)
+	if (pPage.bThisLODValid)
 	{
 		uint PageIndex = PhysPageAddressToIndex(pPage.PageIndex);
 		FPhysicalPageMetaData MetaData = InOutPhysicalPageMetaData[PageIndex];
diff --git a/Engine/Shaders/Private/VirtualShadowMaps/ProjectionCommon.ush b/Engine/Shaders/Private/VirtualShadowMaps/ProjectionCommon.ush
--- a/Engine/Shaders/Private/VirtualShadowMaps/ProjectionCommon.ush
+++ b/Engine/Shaders/Private/VirtualShadowMaps/ProjectionCommon.ush
@@ -25,7 +25,7 @@ struct FVirtualShadowMapProjectionShaderData
 	float3 ClipmapWorldOrigin;
 	int VirtualShadowMapId;
 	
-	int2 ClipmapCenterSnapUnits;
+	int2 ClipmapCornerOffset;
 	int ClipmapIndex;		// 0..ClipmapLevelCount-1
 	int ClipmapLevel;
 	int ClipmapLevelCount;	
@@ -66,7 +66,7 @@ int CalcClipmapLevel(FVirtualShadowMapProjectionShaderData BaseProjectionData, f
 	return int(floor(CalcClipmapLevelFloat(BaseProjectionData, WorldPosition)));
 }
 
-struct FVirtualSmSample
+struct FVirtualShadowMapSample
 {
 	float Depth;
 	uint MipLevel;
@@ -77,9 +77,9 @@ struct FVirtualSmSample
 	uint2 PhysicalTexelAddress;
 };
 
-FVirtualSmSample InitVirtualSmSample()
+FVirtualShadowMapSample InitVirtualShadowMapSample()
 {
-	FVirtualSmSample Result;
+	FVirtualShadowMapSample Result;
 	Result.Depth = 0.0f;
 	Result.MipLevel = 0;
 	Result.VirtualShadowMapId = -1; 
@@ -89,7 +89,7 @@ FVirtualSmSample InitVirtualSmSample()
 	return Result;
 }
 
-float SampleVirtualSmPhysicalDepth(uint2 PhysicalTexelAddress)
+float SampleVirtualShadowMapPhysicalDepth(uint2 PhysicalTexelAddress)
 {
 #if ENABLE_NON_NANITE_VSM
 	return max(
@@ -101,18 +101,18 @@ float SampleVirtualSmPhysicalDepth(uint2 PhysicalTexelAddress)
 #endif // ENABLE_NON_NANITE_VSM
 }
 
-FVirtualSmSample SampleVirtualSmLevel(int VirtualShadowMapId, float2 ShadowMapUV, uint VSMLevel)
+FVirtualShadowMapSample SampleVirtualShadowMapLevel(int VirtualShadowMapId, float2 ShadowMapUV, uint MipLevel)
 {
-	FVirtualSmSample Result = InitVirtualSmSample();
+	FVirtualShadowMapSample Result = InitVirtualShadowMapSample();
 
-	Result.VirtualTexelAddressFloat = ShadowMapUV * float(CalcLevelDimsTexels(VSMLevel));
+	Result.VirtualTexelAddressFloat = ShadowMapUV * float(CalcLevelDimsTexels(MipLevel));
 	Result.VirtualTexelAddress = uint2(Result.VirtualTexelAddressFloat);
 
-	bool PagePresent = VirtualToPhysicalTexel(VirtualShadowMapId, VSMLevel, Result.VirtualTexelAddress, Result.PhysicalTexelAddress);
+	bool PagePresent = VirtualToPhysicalTexel(VirtualShadowMapId, MipLevel, Result.VirtualTexelAddress, Result.PhysicalTexelAddress);
 	if (PagePresent)
 	{
-		Result.Depth = SampleVirtualSmPhysicalDepth(Result.PhysicalTexelAddress);
-		Result.MipLevel = VSMLevel;
+		Result.Depth = SampleVirtualShadowMapPhysicalDepth(Result.PhysicalTexelAddress);
+		Result.MipLevel = MipLevel;
 		Result.VirtualShadowMapId = VirtualShadowMapId;
 		Result.bValid = true;
 	}
@@ -127,20 +127,20 @@ struct FVirtualShadowMapClipmapRelativeTransform
 	float3 Bias;
 };
 
-// NOTE: Two virtual shadow IDs must reference different levels from the same clipmap
-// Resulting transform takes coordinates and depths from A to B
-FVirtualShadowMapClipmapRelativeTransform CalcClipmapRelativeTransform(int VirtualShadowIdA, int VirtualShadowIdB)
+// Resulting transform takes UV's/depths from clipmap level ClipmapId to ClipmapId + LevelOffset 
+// ClipmapId + LevelOffset should be a valid level in the same clipmap
+FVirtualShadowMapClipmapRelativeTransform CalcClipmapRelativeTransform(int ClipmapId, int LevelOffset)
 {
-	const FVirtualShadowMapProjectionShaderData ProjectionDataA = GetVirtualShadowMapProjectionData(VirtualShadowIdA);
-	const FVirtualShadowMapProjectionShaderData ProjectionDataB = GetVirtualShadowMapProjectionData(VirtualShadowIdB);
+	const FVirtualShadowMapProjectionShaderData ProjectionDataA = GetVirtualShadowMapProjectionData(ClipmapId);
+	const FVirtualShadowMapProjectionShaderData ProjectionDataB = GetVirtualShadowMapProjectionData(ClipmapId + LevelOffset);
 
 	// TODO: Simplify/optimize
-	float2 UVA = float2(ProjectionDataA.ClipmapCenterSnapUnits) * float2(-0.25f, 0.25f) + float2(0.5f, 0.5f);
-	float2 UVB = float2(ProjectionDataB.ClipmapCenterSnapUnits) * float2(-0.25f, 0.25f) + float2(0.5f, 0.5f);
+	float2 OffsetA = float2(ProjectionDataA.ClipmapCornerOffset);
+	float2 OffsetB = float2(ProjectionDataB.ClipmapCornerOffset);
 
 	FVirtualShadowMapClipmapRelativeTransform Result;
-	Result.Scale = pow(2.0f, float(VirtualShadowIdA - VirtualShadowIdB));
-	Result.Bias.xy = (UVB - Result.Scale * UVA);
+	Result.Scale = exp2(-float(LevelOffset));
+	Result.Bias.xy = 0.25f * (OffsetB - Result.Scale * OffsetA);
 	// NOTE: Currently no snapping in Z
 	Result.Bias.z = 0.5f - Result.Scale * 0.5f;
 	return Result;
@@ -151,36 +151,32 @@ FVirtualShadowMapClipmapRelativeTransform CalcClipmapRelativeTransform(int Virtu
 
 // Will sample from the given clipmap level Id, onwards to coarser levels if no valid data is present
 // ShadowMapUVs are in terms of the given clipmap level/virtual shadow map
-FVirtualSmSample SampleVirtualSmClipmap(int VirtualShadowMapId, float2 ShadowMapUV)
+FVirtualShadowMapSample SampleVirtualShadowMapClipmap(int VirtualShadowMapId, float2 ShadowMapUV)
 {
-	FVirtualSmSample Result = InitVirtualSmSample();
-		
-#if 1
+	FVirtualShadowMapSample Result = InitVirtualShadowMapSample();
+
+#define DEBUG_VSM_CLIPMAP_LEVEL_SEARCH 0
+#if !DEBUG_VSM_CLIPMAP_LEVEL_SEARCH
 	FShadowPageTranslationResult Page = ShadowVirtualToPhysicalUV(VirtualShadowMapId, ShadowMapUV);
 	if (Page.bValid)
 	{
-		// TODO: Clean up - clipmap offset in mip level
-		int ClipmapId = VirtualShadowMapId + Page.MipLevel;
-
-		FVirtualShadowMapClipmapRelativeTransform Transform = CalcClipmapRelativeTransform(VirtualShadowMapId, ClipmapId);
-
+		FVirtualShadowMapClipmapRelativeTransform Transform = CalcClipmapRelativeTransform(VirtualShadowMapId, Page.LODOffset);
 		float2 ClipmapUV = ShadowMapUV * Transform.Scale + Transform.Bias.xy;
-		Result = SampleVirtualSmLevel(ClipmapId, ClipmapUV, 0);
+		Result = SampleVirtualShadowMapLevel(VirtualShadowMapId + Page.LODOffset, ClipmapUV, 0);
 		if (Result.bValid)
 		{
 			// Convert depth into back into the original reference clipmap level range
 			Result.Depth = (Result.Depth - Transform.Bias.z) / Transform.Scale;
 		}
 	}
-#else
+#else // DEBUG_VSM_CLIPMAP_LEVEL_SEARCH
 	FVirtualShadowMapProjectionShaderData BaseProjectionData = GetVirtualShadowMapProjectionData(VirtualShadowMapId);
-	const int LastId = VirtualShadowMapId + BaseProjectionData.ClipmapLevelCount;
-	for (int ClipmapId = VirtualShadowMapId; ClipmapId < VirtualShadowMapId + 2; ++ClipmapId)
+	const int MaxOffset = BaseProjectionData.ClipmapLevelCount - BaseProjectionData.ClipmapIndex;
+	for (int Offset = 0; Offset < MaxOffset; ++Offset)
 	{
-		FVirtualShadowMapClipmapRelativeTransform Transform = CalcClipmapRelativeTransform(VirtualShadowMapId, ClipmapId);
-		
+		FVirtualShadowMapClipmapRelativeTransform Transform = CalcClipmapRelativeTransform(VirtualShadowMapId, Offset);		
 		float2 ClipmapUV = ShadowMapUV * Transform.Scale + Transform.Bias.xy;
-		Result = SampleVirtualSmLevel(ClipmapId, ClipmapUV, 0);
+		Result = SampleVirtualShadowMapLevel(VirtualShadowMapId + Offset, ClipmapUV, 0);
 		if (Result.bValid)
 		{
 			// Convert depth into back into the original reference clipmap level range
@@ -188,32 +184,37 @@ FVirtualSmSample SampleVirtualSmClipmap(int VirtualShadowMapId, float2 ShadowMap
 			break;
 		}
 	}
-#endif
+#endif // DEBUG_VSM_CLIPMAP_LEVEL_SEARCH
 
 	return Result;
 }
 
 
-FVirtualSmSample SampleVirtualSm(int VirtualShadowMapId, float2 ShadowMapUV)
+FVirtualShadowMapSample SampleVirtualShadowMap(int VirtualShadowMapId, float2 ShadowMapUV)
 {
 	FShadowPageTranslationResult Page = ShadowVirtualToPhysicalUV(VirtualShadowMapId, ShadowMapUV);
 	if (Page.bValid)
 	{
-		FVirtualSmSample Result;
+		FVirtualShadowMapSample Result;
 		Result.bValid = true;
-		Result.MipLevel = Page.MipLevel;
+		Result.MipLevel = Page.LODOffset;
 		Result.VirtualShadowMapId = VirtualShadowMapId;
 		Result.VirtualTexelAddress = Page.VirtualTexelAddress;
 		Result.VirtualTexelAddressFloat = Page.VirtualTexelAddressFloat;
 		Result.PhysicalTexelAddress = Page.PhysicalTexelAddress;
-		Result.Depth = SampleVirtualSmPhysicalDepth(Result.PhysicalTexelAddress);
+		Result.Depth = SampleVirtualShadowMapPhysicalDepth(Result.PhysicalTexelAddress);
 		return Result;
 	}
 
-	return InitVirtualSmSample();
+	return InitVirtualShadowMapSample();
 }
 
-float ComputeVirtualShadowMapOptimalSlopeBias(FVirtualSmSample SmSample, float3 WorldPosition, float3 EstimatedGeoWorldNormal)
+float ComputeVirtualShadowMapOptimalSlopeBias(
+	// Used to compare to the sampled SmSample.VirtualShadowMapId to adjust depth scale for clipmaps
+	int RequestedVirtualShadowMapId,
+	FVirtualShadowMapSample SmSample,
+	float3 WorldPosition,
+	float3 EstimatedGeoWorldNormal)
 {
 	FVirtualShadowMapProjectionShaderData ProjectionData = GetVirtualShadowMapProjectionData(SmSample.VirtualShadowMapId);
 
@@ -223,14 +224,20 @@ float ComputeVirtualShadowMapOptimalSlopeBias(FVirtualSmSample SmSample, float3
 
 	float2 DepthSlopeUV = -NormalPlaneUV.xy / NormalPlaneUV.z;
 	float MipLevelDim = float(CalcLevelDimsTexels(SmSample.MipLevel));
-	float2 TexelCenter = float2(SmSample.VirtualTexelAddress) + 0.5;
+	float2 TexelCenter = float2(SmSample.VirtualTexelAddress) + 0.5f;
 	float2 TexelCenterOffset = TexelCenter - SmSample.VirtualTexelAddressFloat;
 	float2 TexelCenterOffsetUV = TexelCenterOffset / MipLevelDim;
 	// 2x factor due to lack of precision (probably)
 	float OptimalSlopeBias = 2.0f * max(0.0f, dot(DepthSlopeUV, TexelCenterOffsetUV));
+	
 	// Clamp to avoid excessive degenerate slope biases causing flickering lit pixels
 	OptimalSlopeBias = min(OptimalSlopeBias, abs(10.0f * ProjectionData.ShadowViewToClipMatrix._33));
 
+	// Adjust depth scale if we sampled a different clipmap level
+	// NOTE: Sampled clipmap should always be > the requested one (coarser)
+	// Do this after clamping to be consistent in world space
+	OptimalSlopeBias *= float(1 << (SmSample.VirtualShadowMapId - RequestedVirtualShadowMapId));
+
 	return OptimalSlopeBias;
 }
 
@@ -287,7 +294,7 @@ uint GetCubeFace( float3 Dir )
 *   - This is used to compute a receiver-plane-based bias for the sample point
 *   - The shading normal can be used if the geometric normal is not available, but divergence from the geometric normal can cause biasing issues.
 * RayStartDistance is an optional offset to move the lookup along the shadow ray towards the light
-*   - Should zero or positive
+*   - Should be zero or positive
 *   - This offset is useful in that it does *not* affect the selection of clipmap level, unlike offsetting the WorldPosition itself
 *   - OccluderDistance will still be relative to the original sample position
 */
@@ -306,16 +313,15 @@ FVirtualShadowMapSampleResult SampleVirtualShadowMap(int VirtualShadowMapId, flo
 		// Check if sample is within the clipmap range (from camera)
 		if (ClipmapIndex < BaseProjectionData.ClipmapLevelCount)
 		{
-			// NOTE: Clipmap levels are contiguous in the virtual shadow map list
 			int ClipmapLevelVirtualShadowMapId = VirtualShadowMapId + ClipmapIndex;
 			FVirtualShadowMapProjectionShaderData ProjectionData = GetVirtualShadowMapProjectionData(ClipmapLevelVirtualShadowMapId);
 
 			// No perspective divided needed for ortho projection
 			float4 ShadowUVz = mul(float4(WorldPosition + ProjectionData.ShadowPreViewTranslation, 1.0f), ProjectionData.TranslatedWorldToShadowUVMatrix);
 			
-			FVirtualSmSample SmSample;
-			SmSample = SampleVirtualSmLevel(ClipmapLevelVirtualShadowMapId, ShadowUVz.xy, 0);
-			//SmSample = SampleVirtualSmClipmap(ClipmapLevelVirtualShadowMapId, ShadowUVz.xy);
+			FVirtualShadowMapSample SmSample;
+			//SmSample = SampleVirtualShadowMapLevel(ClipmapLevelVirtualShadowMapId, ShadowUVz.xy, 0);
+			SmSample = SampleVirtualShadowMapClipmap(ClipmapLevelVirtualShadowMapId, ShadowUVz.xy);
 			
 			if (SmSample.bValid)
 			{
@@ -327,7 +333,7 @@ FVirtualShadowMapSampleResult SampleVirtualShadowMap(int VirtualShadowMapId, flo
 				Result.OccluderDistance = -1.0f;
 				Result.ClipmapIndexOrMipLevel = SampledClipmapIndex;
 
-				float OptimalSlopeBias = ComputeVirtualShadowMapOptimalSlopeBias(SmSample, WorldPosition, EstimatedGeoWorldNormal);
+				float OptimalSlopeBias = ComputeVirtualShadowMapOptimalSlopeBias(ClipmapLevelVirtualShadowMapId, SmSample, WorldPosition, EstimatedGeoWorldNormal);
 				float RayStartBias = -RayStartDistance * ProjectionData.ShadowViewToClipMatrix._33;
 				float BiasedDepth = SmSample.Depth - OptimalSlopeBias - RayStartBias;
 
@@ -358,7 +364,7 @@ FVirtualShadowMapSampleResult SampleVirtualShadowMap(int VirtualShadowMapId, flo
 		float4 ShadowUVz = mul(float4(TranslatedWorldPosition, 1.0f), BaseProjectionData.TranslatedWorldToShadowUVMatrix);
 		ShadowUVz.xyz /= ShadowUVz.w;
 
-		FVirtualSmSample SmSample = SampleVirtualSm(VirtualShadowMapId, ShadowUVz.xy);
+		FVirtualShadowMapSample SmSample = SampleVirtualShadowMap(VirtualShadowMapId, ShadowUVz.xy);
 		if (SmSample.bValid)
 		{
 			FVirtualShadowMapSampleResult Result;
@@ -367,7 +373,7 @@ FVirtualShadowMapSampleResult SampleVirtualShadowMap(int VirtualShadowMapId, flo
 			Result.OccluderDistance = -1.0f;
 			Result.ClipmapIndexOrMipLevel = SmSample.MipLevel;
 
-			float OptimalSlopeBias = ComputeVirtualShadowMapOptimalSlopeBias(SmSample, WorldPosition, EstimatedGeoWorldNormal);
+			float OptimalSlopeBias = ComputeVirtualShadowMapOptimalSlopeBias(VirtualShadowMapId, SmSample, WorldPosition, EstimatedGeoWorldNormal);
 			float RayStartBias = -RayStartDistance * BaseProjectionData.ShadowViewToClipMatrix._33 / ShadowUVz.w;
 			float BiasedDepth = SmSample.Depth - OptimalSlopeBias - RayStartBias;
 
diff --git a/Engine/Shaders/Private/VirtualShadowMaps/ProjectionDirectional.ush b/Engine/Shaders/Private/VirtualShadowMaps/ProjectionDirectional.ush
--- a/Engine/Shaders/Private/VirtualShadowMaps/ProjectionDirectional.ush
+++ b/Engine/Shaders/Private/VirtualShadowMaps/ProjectionDirectional.ush
@@ -95,7 +95,7 @@ FSMRTSample SMRTFindSample(inout FSMRTClipmapRayState RayState, float SampleTime
 
 	if (bValid)
 	{
-		float SampleDepth = ToBaseDepthScaleBias.x * SampleVirtualSmPhysicalDepth(PhysicalTexelAddress) + ToBaseDepthScaleBias.y;
+		float SampleDepth = ToBaseDepthScaleBias.x * SampleVirtualShadowMapPhysicalDepth(PhysicalTexelAddress) + ToBaseDepthScaleBias.y;
 		Sample.bValid = true;
 		Sample.SampleDepth = SampleDepth;
 		Sample.ReferenceDepth = SampleTexelsZ.z;
diff --git a/Engine/Shaders/Private/VirtualShadowMaps/ProjectionSpot.ush b/Engine/Shaders/Private/VirtualShadowMaps/ProjectionSpot.ush
--- a/Engine/Shaders/Private/VirtualShadowMaps/ProjectionSpot.ush
+++ b/Engine/Shaders/Private/VirtualShadowMaps/ProjectionSpot.ush
@@ -59,7 +59,7 @@ FSMRTSample SMRTFindSample(inout FSMRTSingleRayState RayState, float SampleTime)
 
 	if (all(UVz.xy == saturate(UVz.xy)))
 	{
-		FVirtualSmSample SmSample = SampleVirtualSm(RayState.VirtualShadowMapId, UVz.xy);
+		FVirtualShadowMapSample SmSample = SampleVirtualShadowMap(RayState.VirtualShadowMapId, UVz.xy);
 		if (SmSample.bValid)
 		{
 			Sample.bValid = true;
diff --git a/Engine/Source/Runtime/Renderer/Private/VirtualShadowMaps/VirtualShadowMapArray.cpp b/Engine/Source/Runtime/Renderer/Private/VirtualShadowMaps/VirtualShadowMapArray.cpp
--- a/Engine/Source/Runtime/Renderer/Private/VirtualShadowMaps/VirtualShadowMapArray.cpp
+++ b/Engine/Source/Runtime/Renderer/Private/VirtualShadowMaps/VirtualShadowMapArray.cpp
@@ -89,16 +89,6 @@ static TAutoConsoleVariable<float> CVarPageDilationBorderSize(
 	ECVF_RenderThreadSafe
 );
 
-FMatrix CalcTranslatedWorldToShadowUVMatrix(
-	const FMatrix& TranslatedWorldToShadowView,
-	const FMatrix& ViewToClip)
-{
-	FMatrix TranslatedWorldToShadowClip = TranslatedWorldToShadowView * ViewToClip;
-	FMatrix ScaleAndBiasToSmUV = FScaleMatrix(FVector(0.5f, -0.5f, 1.0f)) * FTranslationMatrix(FVector(0.5f, 0.5f, 0.0f));
-	FMatrix TranslatedWorldToShadowUv = TranslatedWorldToShadowClip * ScaleAndBiasToSmUV;
-	return TranslatedWorldToShadowUv;
-}
-
 TAutoConsoleVariable<int32> CVarAllocatePagesUsingRects(
 	TEXT("r.Shadow.Virtual.AllocatePagesUsingRects"),
 	0,
@@ -106,6 +96,27 @@ TAutoConsoleVariable<int32> CVarAllocatePagesUsingRects(
 	ECVF_RenderThreadSafe
 );
 
+TAutoConsoleVariable<int32> CVarMarkPixelPages(
+	TEXT("r.Shadow.Virtual.MarkPixelPages"),
+	1,
+	TEXT("Marks pages in virtual shadow maps based on depth buffer pixels. Ability to disable is primarily for profiling and debugging."),
+	ECVF_RenderThreadSafe
+);
+
+TAutoConsoleVariable<int32> CVarMarkCoarsePages(
+	TEXT("r.Shadow.Virtual.MarkCoarsePages"),
+	0,
+	TEXT("Marks coarse pages in virtual shadow maps so that low resolution data is available everywhere. Ability to disable is primarily for profiling and debugging."),
+	ECVF_RenderThreadSafe
+);
+
+TAutoConsoleVariable<int32> CVarVirtualShadowMapClipmapFirstCoarseLevel(
+	TEXT("r.Shadow.Virtual.Clipmap.FirstCoarseLevel"),
+	14,
+	TEXT("First level of the clipmap to mark coarse pages for. Lower values allow higher resolution coarse pages near the camera but increase total page counts."),
+	ECVF_RenderThreadSafe
+);
+
 #if ENABLE_NON_NANITE_VSM
 
 static TAutoConsoleVariable<int32> CVarShowClipmapStats(
@@ -117,6 +128,16 @@ static TAutoConsoleVariable<int32> CVarShowClipmapStats(
 
 #endif // ENABLE_NON_NANITE_VSM
 
+FMatrix CalcTranslatedWorldToShadowUVMatrix(
+	const FMatrix& TranslatedWorldToShadowView,
+	const FMatrix& ViewToClip)
+{
+	FMatrix TranslatedWorldToShadowClip = TranslatedWorldToShadowView * ViewToClip;
+	FMatrix ScaleAndBiasToSmUV = FScaleMatrix(FVector(0.5f, -0.5f, 1.0f)) * FTranslationMatrix(FVector(0.5f, 0.5f, 0.0f));
+	FMatrix TranslatedWorldToShadowUv = TranslatedWorldToShadowClip * ScaleAndBiasToSmUV;
+	return TranslatedWorldToShadowUv;
+}
+
 FMatrix CalcTranslatedWorldToShadowUVNormalMatrix(
 	const FMatrix& TranslatedWorldToShadowView,
 	const FMatrix& ViewToClip)
@@ -297,6 +318,20 @@ class FGeneratePageFlagsFromPixelsCS : public FVirtualPageManagementShader
 };
 IMPLEMENT_GLOBAL_SHADER(FGeneratePageFlagsFromPixelsCS, "/Engine/Private/VirtualShadowMaps/PageManagement.usf", "GeneratePageFlagsFromPixels", SF_Compute);
 
+class FMarkCoarsePagesCS : public FVirtualPageManagementShader
+{
+	DECLARE_GLOBAL_SHADER(FMarkCoarsePagesCS);
+	SHADER_USE_PARAMETER_STRUCT(FMarkCoarsePagesCS, FVirtualPageManagementShader)
+
+	BEGIN_SHADER_PARAMETER_STRUCT(FParameters, )
+		SHADER_PARAMETER_RDG_UNIFORM_BUFFER(FVirtualShadowMapCommonParameters, VirtualSmCommon)
+		SHADER_PARAMETER_RDG_BUFFER_UAV(RWStructuredBuffer<uint>, OutPageRequestFlags)
+		SHADER_PARAMETER_RDG_BUFFER_SRV(StructuredBuffer< FVirtualShadowMapProjectionShaderData >, VirtualShadowMapProjectionData)
+		SHADER_PARAMETER(int32, ClipmapFirstCoarseLevel)
+	END_SHADER_PARAMETER_STRUCT()
+};
+IMPLEMENT_GLOBAL_SHADER(FMarkCoarsePagesCS, "/Engine/Private/VirtualShadowMaps/PageManagement.usf", "MarkCoarsePages", SF_Compute);
+
 
 class FGenerateHierarchicalPageFlagsCS : public FVirtualPageManagementShader
 {
@@ -801,39 +836,65 @@ void FVirtualShadowMapArray::BuildPageAllocations(
 
 			FRDGBufferRef ScreenSpaceGridBoundsRDG = nullptr;
 			
-			// Project Pixels onto SMs
 			{
-				FGeneratePageFlagsFromPixelsCS::FPermutationDomain PermutationVector;
-				PermutationVector.Set<FGeneratePageFlagsFromPixelsCS::FNaniteDepthBufferDim>(NaniteVisBuffer64 != nullptr && !bPostBasePass);
+				// It's safe to overlap these passes that all write to page request flags
+				FRDGBufferUAVRef PageRequestFlagsUAV = GraphBuilder.CreateUAV(PageRequestFlagsRDG, ERDGUnorderedAccessViewFlags::SkipBarrier);
 
-				FGeneratePageFlagsFromPixelsCS::FParameters* PassParameters = GraphBuilder.AllocParameters< FGeneratePageFlagsFromPixelsCS::FParameters >();
-				PassParameters->VirtualSmCommon = GetCommonUniformBuffer(GraphBuilder);
+				// Mark pages based on projected depth buffer pixels
+				if (CVarMarkPixelPages.GetValueOnRenderThread() != 0)
+				{
+					FGeneratePageFlagsFromPixelsCS::FPermutationDomain PermutationVector;
+					PermutationVector.Set<FGeneratePageFlagsFromPixelsCS::FNaniteDepthBufferDim>(NaniteVisBuffer64 != nullptr && !bPostBasePass);
+
+					FGeneratePageFlagsFromPixelsCS::FParameters* PassParameters = GraphBuilder.AllocParameters< FGeneratePageFlagsFromPixelsCS::FParameters >();
+					PassParameters->VirtualSmCommon = GetCommonUniformBuffer(GraphBuilder);
 
-				PassParameters->SceneTexturesStruct = SceneTextures.UniformBuffer;
-				PassParameters->bPostBasePass = bPostBasePass;
+					PassParameters->SceneTexturesStruct = SceneTextures.UniformBuffer;
+					PassParameters->bPostBasePass = bPostBasePass;
 				
-				PassParameters->VisBuffer64 = VisBuffer64;
-				PassParameters->View = View.ViewUniformBuffer;
-				PassParameters->OutPageRequestFlags = GraphBuilder.CreateUAV(PageRequestFlagsRDG);
-				PassParameters->ForwardLightData = View.ForwardLightingResources->ForwardLightDataUniformBuffer;
-				PassParameters->VirtualShadowMapIdRemap = GraphBuilder.CreateSRV(VirtualShadowMapIdRemapRDG[ViewIndex]);
-				PassParameters->VirtualShadowMapProjectionData = GraphBuilder.CreateSRV(ShadowMapProjectionDataRDG);
-				PassParameters->NumDirectionalLightSmInds = DirectionalLightSmInds.Num();
-				PassParameters->LodFootprintScale = LodFootprintScale;
-				PassParameters->PageDilationBorderSize = PageDilationBorderSize;
+					PassParameters->VisBuffer64 = VisBuffer64;
+					PassParameters->View = View.ViewUniformBuffer;
+					PassParameters->OutPageRequestFlags = PageRequestFlagsUAV;
+					PassParameters->ForwardLightData = View.ForwardLightingResources->ForwardLightDataUniformBuffer;
+					PassParameters->VirtualShadowMapIdRemap = GraphBuilder.CreateSRV(VirtualShadowMapIdRemapRDG[ViewIndex]);
+					PassParameters->VirtualShadowMapProjectionData = GraphBuilder.CreateSRV(ShadowMapProjectionDataRDG);
+					PassParameters->NumDirectionalLightSmInds = DirectionalLightSmInds.Num();
+					PassParameters->LodFootprintScale = LodFootprintScale;
+					PassParameters->PageDilationBorderSize = PageDilationBorderSize;
 				
-				auto ComputeShader = View.ShaderMap->GetShader<FGeneratePageFlagsFromPixelsCS>(PermutationVector);
+					auto ComputeShader = View.ShaderMap->GetShader<FGeneratePageFlagsFromPixelsCS>(PermutationVector);
 
-				static_assert((FVirtualPageManagementShader::DefaultCSGroupXY % 2) == 0, "GeneratePageFlagsFromPixels requires even-sized CS groups for quad swizzling.");
-				const FIntPoint GridSize = FIntPoint::DivideAndRoundUp(View.ViewRect.Size(), FVirtualPageManagementShader::DefaultCSGroupXY);
+					static_assert((FVirtualPageManagementShader::DefaultCSGroupXY % 2) == 0, "GeneratePageFlagsFromPixels requires even-sized CS groups for quad swizzling.");
+					const FIntPoint GridSize = FIntPoint::DivideAndRoundUp(View.ViewRect.Size(), FVirtualPageManagementShader::DefaultCSGroupXY);
 
-				FComputeShaderUtils::AddPass(
-					GraphBuilder,
-					RDG_EVENT_NAME("GeneratePageFlagsFromPixels"),
-					ComputeShader,
-					PassParameters,
-					FIntVector(GridSize.X, GridSize.Y, 1)
-				);
+					FComputeShaderUtils::AddPass(
+						GraphBuilder,
+						RDG_EVENT_NAME("GeneratePageFlagsFromPixels"),
+						ComputeShader,
+						PassParameters,
+						FIntVector(GridSize.X, GridSize.Y, 1)
+					);
+				}
+
+				// Mark coarse pages
+				if (CVarMarkCoarsePages.GetValueOnRenderThread() != 0)
+				{
+					FMarkCoarsePagesCS::FParameters* PassParameters = GraphBuilder.AllocParameters< FMarkCoarsePagesCS::FParameters >();
+					PassParameters->VirtualSmCommon = GetCommonUniformBuffer(GraphBuilder);
+					PassParameters->OutPageRequestFlags = PageRequestFlagsUAV;
+					PassParameters->VirtualShadowMapProjectionData = GraphBuilder.CreateSRV(ShadowMapProjectionDataRDG);
+					PassParameters->ClipmapFirstCoarseLevel = CVarVirtualShadowMapClipmapFirstCoarseLevel.GetValueOnRenderThread();
+
+					auto ComputeShader = View.ShaderMap->GetShader<FMarkCoarsePagesCS>();
+
+					FComputeShaderUtils::AddPass(
+						GraphBuilder,
+						RDG_EVENT_NAME("MarkCoarsePages"),
+						ComputeShader,
+						PassParameters,
+						FIntVector(FMath::DivideAndRoundUp(uint32(ShadowMaps.Num()), FMarkCoarsePagesCS::DefaultCSGroupX), 1, 1)
+					);
+				}
 			}
 		}
 
diff --git a/Engine/Source/Runtime/Renderer/Private/VirtualShadowMaps/VirtualShadowMapArray.h b/Engine/Source/Runtime/Renderer/Private/VirtualShadowMaps/VirtualShadowMapArray.h
--- a/Engine/Source/Runtime/Renderer/Private/VirtualShadowMaps/VirtualShadowMapArray.h
+++ b/Engine/Source/Runtime/Renderer/Private/VirtualShadowMaps/VirtualShadowMapArray.h
@@ -97,7 +97,7 @@ struct FVirtualShadowMapProjectionShaderData
 	FVector ClipmapWorldOrigin;
 	int32 VirtualShadowMapId = INDEX_NONE;
 	
-	FIntPoint ClipmapCenterSnapUnits;
+	FIntPoint ClipmapCornerOffset;
 	int32 ClipmapIndex = 0;					// 0 .. ClipmapLevelCount-1
 	int32 ClipmapLevel = 0;					// "Absolute" level, can be negative
 	int32 ClipmapLevelCount = 0;
diff --git a/Engine/Source/Runtime/Renderer/Private/VirtualShadowMaps/VirtualShadowMapClipmap.cpp b/Engine/Source/Runtime/Renderer/Private/VirtualShadowMaps/VirtualShadowMapClipmap.cpp
--- a/Engine/Source/Runtime/Renderer/Private/VirtualShadowMaps/VirtualShadowMapClipmap.cpp
+++ b/Engine/Source/Runtime/Renderer/Private/VirtualShadowMaps/VirtualShadowMapClipmap.cpp
@@ -26,11 +26,13 @@ static TAutoConsoleVariable<int32> CVarVirtualShadowMapClipmapFirstLevel(
 static TAutoConsoleVariable<float> CVarVirtualShadowMapClipmapMaxRadius(
 	TEXT( "r.Shadow.Virtual.Clipmap.MaxRadius" ),
 	1000000.0f,
-	TEXT( "Maximum distance the clipmap covers. Determines the number of clipmap levels." ),
+	TEXT( "Maximum distance the clipmap is guaranteed to cover. Determines the number of clipmap levels." ),
 	ECVF_RenderThreadSafe
 );
 
 // "Virtual" clipmap level to clipmap radius
+// NOTE: This is the radius of around the clipmap origin that this level must cover
+// The actual clipmap dimensions will be larger due to snapping and other accomodations
 static float GetLevelRadius(int32 Level)
 {
 	// NOTE: Virtual clipmap indices can be negative (although not commonly)
@@ -97,7 +99,7 @@ FVirtualShadowMapClipmap::FVirtualShadowMapClipmap(
 
 		const float RawLevelRadius = GetLevelRadius(LevelIndex);
 
-		float SnappedLevelRadius = 2.0f * RawLevelRadius;
+		float HalfLevelDim = 2.0f * RawLevelRadius;
 		float SnapSize = RawLevelRadius;
 
 		FVector ViewCenter = WorldToViewRotationMatrix.TransformPosition(WorldOrigin);
@@ -107,14 +109,18 @@ FVirtualShadowMapClipmap::FVirtualShadowMapClipmap(
 		ViewCenter.X = CenterSnapUnits.X * SnapSize;
 		ViewCenter.Y = CenterSnapUnits.Y * SnapSize;
 
+		FIntPoint CornerOffset;
+		CornerOffset.X = -CenterSnapUnits.X + 2;
+		CornerOffset.Y =  CenterSnapUnits.Y + 2;
+
 		const FVector SnappedWorldCenter = ViewToWorldRotationMatrix.TransformPosition(ViewCenter);
 
 		const float ZScale = 0.5f / RawLevelRadius;
 		const float ZOffset = RawLevelRadius;
 
 		Level.WorldCenter = SnappedWorldCenter;
-		Level.ViewToClip = FReversedZOrthoMatrix(SnappedLevelRadius, SnappedLevelRadius, ZScale, ZOffset);
-		Level.CenterSnapUnits = CenterSnapUnits;
+		Level.ViewToClip = FReversedZOrthoMatrix(HalfLevelDim, HalfLevelDim, ZScale, ZOffset);
+		Level.CornerOffset = CornerOffset;
 
 		if (VirtualShadowMapArrayCacheManager)
 		{
@@ -124,8 +130,7 @@ FVirtualShadowMapClipmap::FVirtualShadowMapClipmap(
 			{
 				// We snap to half the size of the VSM at each level
 				check((FVirtualShadowMap::Level0DimPagesXY & 1) == 0);
-				FIntPoint PageOffset(CenterSnapUnits * (FVirtualShadowMap::Level0DimPagesXY >> 1));
-				PageOffset.Y = -PageOffset.Y;		// Viewport
+				FIntPoint PageOffset(CornerOffset * (FVirtualShadowMap::Level0DimPagesXY >> 2));
 				float DepthOffset = -ViewCenter.Z * ZScale;
 
 				CacheEntry->UpdateClipmap(Level.VirtualShadowMap->ID, WorldToLightRotationMatrix, PageOffset, DepthOffset);
@@ -174,7 +179,7 @@ FVirtualShadowMapProjectionShaderData FVirtualShadowMapClipmap::GetProjectionSha
 	Data.ClipmapLevel = FirstLevel + ClipmapIndex;
 	Data.ClipmapLevelCount = LevelData.Num();
 	Data.ClipmapResolutionLodBias = ResolutionLodBias;
-	Data.ClipmapCenterSnapUnits = Level.CenterSnapUnits;
+	Data.ClipmapCornerOffset = Level.CornerOffset;
 
 	return Data;
 }
diff --git a/Engine/Source/Runtime/Renderer/Private/VirtualShadowMaps/VirtualShadowMapClipmap.h b/Engine/Source/Runtime/Renderer/Private/VirtualShadowMaps/VirtualShadowMapClipmap.h
--- a/Engine/Source/Runtime/Renderer/Private/VirtualShadowMaps/VirtualShadowMapClipmap.h
+++ b/Engine/Source/Runtime/Renderer/Private/VirtualShadowMaps/VirtualShadowMapClipmap.h
@@ -88,7 +88,7 @@ private:
 		FVirtualShadowMap* VirtualShadowMap;
 		FMatrix ViewToClip;
 		FVector WorldCenter;
-		FIntPoint CenterSnapUnits;	// Center of clipmap in world units / level radius
+		FIntPoint CornerOffset;
 	};
 
 	TArray< FLevelData, TInlineAllocator<16> > LevelData;
diff --git a/Engine/Source/Runtime/Renderer/Private/VirtualShadowMaps/VirtualShadowMapProjection.cpp b/Engine/Source/Runtime/Renderer/Private/VirtualShadowMaps/VirtualShadowMapProjection.cpp
--- a/Engine/Source/Runtime/Renderer/Private/VirtualShadowMaps/VirtualShadowMapProjection.cpp
+++ b/Engine/Source/Runtime/Renderer/Private/VirtualShadowMaps/VirtualShadowMapProjection.cpp
@@ -40,9 +40,9 @@ static TAutoConsoleVariable<float> CVarNormalOffsetWorld(
 );
 
 static TAutoConsoleVariable<int32> CVarVirtualShadowMapDebugProjection(
-	TEXT( "r.Shadow.Virtual.DebugVisualizeProjection" ),
+	TEXT( "r.Shadow.Virtual.DebugProjection" ),
 	0,
-	TEXT( "Projection pass debug output visualization for use with 'vis VirtSmDebugProj'." ),
+	TEXT( "Projection pass debug output visualization for use with 'vis Shadow.Virtual.DebugProjection'." ),
 	ECVF_RenderThreadSafe
 );
 
@@ -287,7 +287,7 @@ static FRDGTextureRef CreateDebugOutput(FRDGBuilder& GraphBuilder, FIntPoint Ext
 		TexCreate_ShaderResource | TexCreate_RenderTargetable
 	);
 
-	return GraphBuilder.CreateTexture(DebugOutputDesc, TEXT("Shadow.Virtual.DebugProj"));
+	return GraphBuilder.CreateTexture(DebugOutputDesc, TEXT("Shadow.Virtual.DebugProjection"));
 }
 
 // Returns true if temporal denoising should be used
