Misc Non-Nanite VSM and GPU-Scene instance culling fixes.
- Fix ordering issue in virtual shadow map setup (need to commit and upload dynamic primitive data before building rendering commands).
- Workaround for incorrect culling of late (shadow view) dynamic primitives.
- Mark local lights as rendered.

diff --git a/Engine/Shaders/Private/InstanceCulling/CullInstances.usf b/Engine/Shaders/Private/InstanceCulling/CullInstances.usf
--- a/Engine/Shaders/Private/InstanceCulling/CullInstances.usf
+++ b/Engine/Shaders/Private/InstanceCulling/CullInstances.usf
@@ -28,24 +28,28 @@ void CullInstancesCs(uint InstanceId : SV_DispatchThreadID)
 	uint WordMask = 1U << (InstanceId % 32U);
 	uint InstanceWordOffset = InstanceId / 32U;
 
-	bool bIsVisible = InstanceData.PrimitiveId != 0xFFFFFFFFu;
-
-	if (!bIsVisible)
-	{
-		return;
-	}
+	bool bIsValid = InstanceData.PrimitiveId != 0xFFFFFFFFu;
 
 	for (uint ViewId = 0; ViewId < NumViews; ++ViewId)
 	{
-		FNaniteView NaniteView = GetNaniteView(ViewId);
+		// NOTE: Default to flagging as visible, which means that invalid instances are marked as visible. This is a workaround that prevents instances that 
+		// are uploaded after this kernel runs from being culled. When the culling tests are moved into the later kernel we should switch this back to discarding.
+		uint Flag = WordMask;
+		if (bIsValid)
+		{
+			FNaniteView NaniteView = GetNaniteView(ViewId);
 
-		float4x4 LocalToTranslatedWorld = InstanceData.LocalToWorld;
-		LocalToTranslatedWorld[3].xyz += NaniteView.PreViewTranslation.xyz;
-		float4x4 LocalToClip = mul(LocalToTranslatedWorld, NaniteView.TranslatedWorldToClip);
+			float4x4 LocalToTranslatedWorld = InstanceData.LocalToWorld;
+			LocalToTranslatedWorld[3].xyz += NaniteView.PreViewTranslation.xyz;
+			float4x4 LocalToClip = mul(LocalToTranslatedWorld, NaniteView.TranslatedWorldToClip);
 
-		FFrustumCullData Cull = BoxCullFrustum(InstanceData.LocalBoundsCenter, InstanceData.LocalBoundsExtent, LocalToClip, bNearClip, false);
+			FFrustumCullData Cull = BoxCullFrustum(InstanceData.LocalBoundsCenter, InstanceData.LocalBoundsExtent, LocalToClip, bNearClip, false);
 
-		uint Flag = (Cull.bIsVisible ? WordMask : 0U);
+			if (!Cull.bIsVisible)
+			{
+				Flag = 0U;
+			}
+		}
 
 		// TODO: we can use ballot etc to do this without atomics (and without clearing) - but it would need to be portable and not sure if 
 		// we'll keep this code anyway.
diff --git a/Engine/Source/Runtime/Renderer/Private/VirtualShadowMaps/VirtualShadowMapArray.cpp b/Engine/Source/Runtime/Renderer/Private/VirtualShadowMaps/VirtualShadowMapArray.cpp
--- a/Engine/Source/Runtime/Renderer/Private/VirtualShadowMaps/VirtualShadowMapArray.cpp
+++ b/Engine/Source/Runtime/Renderer/Private/VirtualShadowMaps/VirtualShadowMapArray.cpp
@@ -1527,6 +1527,8 @@ void FVirtualShadowMapArray::RenderVirtualShadowMapsHw(FRDGBuilder& GraphBuilder
 		const FViewInfo& View = *ViewUsedToCreateShadow;
 		TArray<Nanite::FPackedView, SceneRenderingAllocator> VirtualShadowViews;
 
+		ProjectedShadowInfo->BeginRenderView(GraphBuilder, &Scene);
+
 		if( Clipmap )
 		{
 			Nanite::FPackedViewParams BaseParams;
@@ -1571,7 +1573,7 @@ void FVirtualShadowMapArray::RenderVirtualShadowMapsHw(FRDGBuilder& GraphBuilder
 				Params.ViewMatrices = ProjectedShadowInfo->GetShadowDepthRenderingViewMatrices(i, true);
 
 				VirtualShadowViews.Add(Nanite::CreatePackedView(Params));
-				//VirtualShadowMapFlags[ProjectedShadowInfo->VirtualShadowMaps[i]->ID] = 1;
+				VirtualShadowMapFlags[ProjectedShadowInfo->VirtualShadowMaps[i]->ID] = 1;
 			}
 		}
 
@@ -1709,8 +1711,6 @@ void FVirtualShadowMapArray::RenderVirtualShadowMapsHw(FRDGBuilder& GraphBuilder
 				);
 			}
 
-			ProjectedShadowInfo->BeginRenderView(GraphBuilder, &Scene);
-
 			FVirtualShadowDepthPassParameters* PassParameters = GraphBuilder.AllocParameters<FVirtualShadowDepthPassParameters>();
 			PassParameters->View = ShadowDepthView->ViewUniformBuffer;
 			PassParameters->RenderTargets.DepthStencil = FDepthStencilBinding(PhysicalPagePoolHw, bWasCleared ? ERenderTargetLoadAction::ELoad : ERenderTargetLoadAction::EClear, FExclusiveDepthStencil::DepthWrite_StencilNop);
