Vertical layering of Strata slab is now working again.
Volumetric evalaution in place of diffuse is only happening in forward or top layers of opaque mateiral (but not at the bottom level for now).
DMFP is considered as MFP however because DMFP is not inversible to medium porperties directly.
Fresnel throughput between layer is still approximate.
Transition from diffuse to volumetric is simple and will be updated.

Shooter game compiled on win64 and console without strata enabled.

diff --git a/Engine/Shaders/Private/BasePassPixelShader.usf b/Engine/Shaders/Private/BasePassPixelShader.usf
--- a/Engine/Shaders/Private/BasePassPixelShader.usf
+++ b/Engine/Shaders/Private/BasePassPixelShader.usf
@@ -83,8 +83,6 @@
 #endif
 
 #define STRATA_ENABLED					(PROJECT_STRATA && MATERIAL_IS_STRATA)
-#define STRATA_OPAQUE_ENABLED			(STRATA_ENABLED && MATERIALBLENDING_ANY_TRANSLUCENT<=0)
-#define STRATA_TRANSLUCENT_ENABLED		(STRATA_ENABLED && MATERIALBLENDING_ANY_TRANSLUCENT>0)
 
 #if STRATA_TRANSLUCENT_ENABLED
 #define STRATA_USE_ENV_BRDF_APPROX (MATERIAL_USE_PREINTEGRATED_GF == 0)
@@ -1563,8 +1561,7 @@ void FPixelShaderInOut_MainPS(
 	PackStrataOut(
 		OpaqueBasePass.Strata.MaterialLobesBufferUAV, 
 		GetStrataPixelDataByteOffset(PixelPos, uint2(ResolvedView.BufferSizeAndInvSize.xy), OpaqueBasePass.Strata.MaxBytesPerPixel),
-		StrataPixelHeader, Strata,
-		MaterialParameters.CameraVector, bStrataSubsurfaceEnable, EmissiveLuminance);
+		StrataPixelHeader, Strata, MaterialParameters.CameraVector, bStrataSubsurfaceEnable, EmissiveLuminance);
 
 	Color += EmissiveLuminance;
 #endif // STRATA_OPAQUE_ENABLED
diff --git a/Engine/Shaders/Private/MaterialTemplate.ush b/Engine/Shaders/Private/MaterialTemplate.ush
--- a/Engine/Shaders/Private/MaterialTemplate.ush
+++ b/Engine/Shaders/Private/MaterialTemplate.ush
@@ -19,10 +19,6 @@
 #include "/Engine/Private/HairStrands/HairCardsAttributeCommon.ush"
 #include "/Engine/Private/HairStrands/HairStrandsAttributeCommon.ush"
 
-#define TEMPLATE_USES_STRATA (PROJECT_STRATA && MATERIAL_IS_STRATA)
-
-#include "/Engine/Private/Strata/Strata.ush"
-
 #if USES_SPEEDTREE
 	#include "/Engine/Private/SpeedTreeCommon.ush"
 #endif
@@ -134,6 +130,11 @@
 
 #define IS_HAIR_FACTORY (HAIR_STRAND_MESH_FACTORY || HAIR_CARD_MESH_FACTORY)
 
+#define TEMPLATE_USES_STRATA			(PROJECT_STRATA && MATERIAL_IS_STRATA)
+#define STRATA_OPAQUE_ENABLED			(TEMPLATE_USES_STRATA && MATERIALBLENDING_ANY_TRANSLUCENT<=0)
+#define STRATA_TRANSLUCENT_ENABLED		(TEMPLATE_USES_STRATA && MATERIALBLENDING_ANY_TRANSLUCENT>0)
+#include "/Engine/Private/Strata/Strata.ush"
+
 /**
  * Parameters used by vertex and pixel shaders to access particle properties.
  */
diff --git a/Engine/Shaders/Private/ParticipatingMediaCommon.ush b/Engine/Shaders/Private/ParticipatingMediaCommon.ush
--- a/Engine/Shaders/Private/ParticipatingMediaCommon.ush
+++ b/Engine/Shaders/Private/ParticipatingMediaCommon.ush
@@ -6,7 +6,26 @@
 
 #include "Common.ush"
 
+struct FParticipatingMedia
+{
+	float3 ScatteringCoef;	// sigma_s (1/meter)
+	float3 AbsorptionCoef;	// sigma_a (1/meter)
+	float3 ExtinctionCoef;	// sigma_t (1/meter)
+	float3 MeanFreePath;	// (meter)
+	float3 Albedo;			// (unitless)
+};
 
+// The mean free path must be in meters
+FParticipatingMedia CreateMediumFromAlbedoMFP(float3 Albedo, float3 MeanFreePathMeters)
+{
+	FParticipatingMedia PM = (FParticipatingMedia)0;
+	PM.Albedo = Albedo;
+	PM.MeanFreePath = MeanFreePathMeters;
+	PM.ExtinctionCoef = 1.0f / max(0.000001f, PM.MeanFreePath);
+	PM.ScatteringCoef = Albedo * PM.ExtinctionCoef;
+	PM.AbsorptionCoef = max(0.0f, PM.ExtinctionCoef - PM.ScatteringCoef);
+	return PM;
+}
 
 float IsotropicPhase()
 {
diff --git a/Engine/Shaders/Private/Strata/Strata.ush b/Engine/Shaders/Private/Strata/Strata.ush
--- a/Engine/Shaders/Private/Strata/Strata.ush
+++ b/Engine/Shaders/Private/Strata/Strata.ush
@@ -5,6 +5,8 @@
 #include "../Common.ush"
 #include "../BRDF.ush"
 #include "../DeferredShadingCommon.ush"
+#include "../ParticipatingMediaCommon.ush"
+#include "../ThinFilmBSDF.ush"
 
 #include "/Engine/Shared/StrataDefinitions.h"
 
@@ -25,6 +27,12 @@
 
 #define STRATA_INLINE_SINGLELAYERWATER (PROJECT_STRATA && MATERIAL_IS_STRATA && STRATA_INLINE_SHADING && MATERIAL_SHADINGMODEL_SINGLELAYERWATER)
 
+// We do all the simple volumetric math considering a slab of 1 centimeter = 0.01 meter.
+// If the user specifies a different thickness to simulate varying appearance, we rescale the mean free path in order to not have to store thickness as part of the BSDF description.
+#define STRATA_SIMPLEVOLUME_THICKNESS_CM	1.0f
+#define STRATA_CENTIMETER_TO_METER			0.01f
+#define STRATA_SIMPLEVOLUME_THICKNESS_M		(STRATA_SIMPLEVOLUME_THICKNESS_CM * STRATA_CENTIMETER_TO_METER)
+
 ///////////////////////////////////////////////////////////////////////////////
 // BSDFs representation
 struct FStrataBSDF
@@ -34,6 +42,8 @@ struct FStrataBSDF
 
 #if STRATA_INLINE_SHADING
 	float3				Emissive;	// All bsdf can potientially have emissive. This is never written out to the strata gbuffer but instead accumulated in the emissive buffer during the base pass.
+	float				ThicknessCm;// This must be kept in order to be able to normalize the mean free path to comply with our volumetric material math all done in a normalised slab of medium.
+	float3				TmpDMFP;
 #endif
 
 	// Water has too many parameters to be stores in the extra VGPRs below. So we extend this structure for the case where it is evaluated inline only.
@@ -102,7 +112,8 @@ struct FStrataBSDF
 #define STATE_BIT_COUNT_HAZINESS				1
 #define STATE_BIT_COUNT_EDGECOLOR				1
 #define STATE_BIT_COUNT_THINFILM				1
-#define STATE_BIT_COUNT___UNUSED___				9
+#define STATE_BIT_COUNT_SIMPLEVOLUME			1
+#define STATE_BIT_COUNT___UNUSED___				8
 #define STATE_BIT_COUNT_WEIGHT10F				10
 
 // Bit offset of each component store in BSDF state 
@@ -116,6 +127,7 @@ struct FStrataBSDF
 #define STATE_BIT_OFFSET_HAZINESS		(STATE_BIT_OFFSET_GREYWEIGHT		+ STATE_BIT_COUNT_GREYWEIGHT)
 #define STATE_BIT_OFFSET_EDGECOLOR		(STATE_BIT_OFFSET_HAZINESS			+ STATE_BIT_COUNT_HAZINESS)
 #define STATE_BIT_OFFSET_THINFILM		(STATE_BIT_OFFSET_EDGECOLOR			+ STATE_BIT_COUNT_EDGECOLOR)
+#define STATE_BIT_OFFSET_SIMPLEVOLUME	(STATE_BIT_OFFSET_THINFILM			+ STATE_BIT_COUNT_THINFILM)
 #define STATE_BIT_OFFSET_WEIGHT10F		22
 
 // Accessors for the BSDF type
@@ -164,6 +176,10 @@ struct FStrataBSDF
 #define BSDF_GETWEIGHT10F(X)			READ_BITS(X.State,  STATE_BIT_COUNT_WEIGHT10F, STATE_BIT_OFFSET_WEIGHT10F)
 #define BSDF_SETWEIGHT10F(X, Weight)	WRITE_BITS(X.State, STATE_BIT_COUNT_WEIGHT10F, STATE_BIT_OFFSET_WEIGHT10F, Weight)
 
+// Indicates if the BSDF is part of the top layer
+#define BSDF_GETISSIMPLEVOLUME(X)		READ_BITS(X.State,  STATE_BIT_COUNT_SIMPLEVOLUME, STATE_BIT_OFFSET_SIMPLEVOLUME)
+#define BSDF_SETISSIMPLEVOLUME(X, Value)WRITE_BITS(X.State, STATE_BIT_COUNT_SIMPLEVOLUME, STATE_BIT_OFFSET_SIMPLEVOLUME, Value)
+
 #if STRATA_INLINE_SHADING
 #define BSDF_SETEMISSIVE(X, V)			X.Emissive = V
 #define BSDF_GETEMISSIVE(X)				X.Emissive
@@ -172,6 +188,14 @@ struct FStrataBSDF
 #define BSDF_GETEMISSIVE(X)				0.0f
 #endif
 
+#if STRATA_INLINE_SHADING
+#define BSDF_SETTHICKNESSCM(X, V)		X.ThicknessCm = V
+#define BSDF_GETTHICKNESSCM(X)			X.ThicknessCm
+#else
+#define BSDF_SETTHICKNESSCM(X, V)		X
+#define BSDF_GETTHICKNESSCM(X)			STRATA_SIMPLEVOLUME_THICKNESS_CM
+#endif
+
 #define SLAB_BASECOLOR(X)				X.VGPRs[0].xyz
 #define SLAB_SPECULAR(X)				X.VGPRs[0].w
 #define SLAB_METALLIC(X)				X.VGPRs[1].x
@@ -182,7 +206,7 @@ struct FStrataBSDF
 // Either with SSS profile
 #define SLAB_SSSPROFILEID(X)			X.VGPRs[2].x
 #define SLAB_SSSPROFILERADIUSSCALE(X)	X.VGPRs[2].y
-// Or explicit DMFP
+// Or explicit DMFP in centimeter
 #define SLAB_SSSDMFP(X)					X.VGPRs[2].xyz
 // Thin film. IOR is harcoded to oil IOR in order to save VGPR
 #define SLAB_THINFILM_THICKNESS(X)		X.VGPRs[2].w
@@ -401,7 +425,7 @@ FStrataData GetInitialisedStrataData()
 FStrataData GetStrataSlabBSDF(
 	float3 BaseColor, float3 EdgeColor, float Specular, float Metallic, float2 Roughness,
 	float SSSProfileID, float3 SSSDMFP, float SSSDMFPScale,
-	float3 Emissive, float Haziness, float ThinFilmThickness, 
+	float3 Emissive, float Haziness, float ThinFilmThickness,  float Thickness, 
 	uint SharedNormalIndex, inout uint SharedNormalTypes)
 {
 	FStrataData StrataData = GetInitialisedStrataData();
@@ -431,6 +455,7 @@ FStrataData GetStrataSlabBSDF(
 	BSDF_SETHASHAZINESS		(StrataData.Layers[0].BSDFs[0],		bHasHaziness);
 	BSDF_SETHASEDGECOLOR	(StrataData.Layers[0].BSDFs[0],		bHasEdgeColor);
 	BSDF_SETHASTHINFILM		(StrataData.Layers[0].BSDFs[0],		bHasThinFilm);
+	BSDF_SETTHICKNESSCM		(StrataData.Layers[0].BSDFs[0],		Thickness);
 
 	SLAB_BASECOLOR			(StrataData.Layers[0].BSDFs[0])		= BaseColor;
 	SLAB_EDGECOLOR			(StrataData.Layers[0].BSDFs[0])		= EdgeColor;
@@ -450,6 +475,11 @@ FStrataData GetStrataSlabBSDF(
 		SLAB_SSSDMFP	(StrataData.Layers[0].BSDFs[0])			= SSSDMFP * saturate(SSSDMFPScale);
 	}
 
+#if STRATA_INLINE_SHADING
+	// At this stage we do not know if this strata BSDF is going to be using simple volume lighting so always store the DMFP on the side to not override SSSPROFILE data.
+	StrataData.Layers[0].BSDFs[0].TmpDMFP = SSSDMFP * saturate(SSSDMFPScale);
+#endif
+
 	StrataData.Layers[0].BSDFs[0].Weight						= 1.0f;
 	return StrataData;
 }
@@ -733,7 +763,7 @@ float3x3 StrataGetBSDFSharedBasis(in FStrataPixelHeader StrataPixelHeader, in FS
 
 
 ///////////////////////////////////////////////////////////////////////////////
-// Operators node
+// Operator nodes
 
 FStrataData StrataAdd(FStrataData A, FStrataData B)
 {
@@ -823,6 +853,9 @@ void SanitizeStrataSlab(inout FStrataBSDF BSDF)
 	SLAB_THINFILM_THICKNESS(BSDF) = saturate(SLAB_THINFILM_THICKNESS(BSDF));
 	SLAB_SSSPROFILERADIUSSCALE(BSDF) = saturate(SLAB_SSSPROFILERADIUSSCALE(BSDF));
 	SLAB_SSSDMFP(BSDF) = clamp(SLAB_SSSDMFP(BSDF), 0.0f, Max111110BitsFloat3);
+#if STRATA_INLINE_SHADING
+	BSDF.TmpDMFP = clamp(BSDF.TmpDMFP, 0.0f, Max111110BitsFloat3);
+#endif
 }
 
 void SanitizeStrataSheen(inout FStrataBSDF BSDF)
@@ -889,6 +922,45 @@ void StrataSanitizeBSDF(inout FStrataBSDF BSDF)
 
 
 
+///////////////////////////////////////////////////////////////////////////////
+// Other tools
+
+FParticipatingMedia StrataSlabCreateParticipatingMedia(float3 DiffuseColor, float3 MeanFreePathCentimeters)
+{
+	const float3 MeanFreePathMeters = MeanFreePathCentimeters * STRATA_CENTIMETER_TO_METER;
+	return CreateMediumFromAlbedoMFP(DiffuseColor, MeanFreePathMeters);
+}
+
+float StrataSlabDiffuseToVolumeBlend(FParticipatingMedia PM)
+{
+	// We soft lerp from opaque diffuse to volue representation when transmittance to background (or layers below) reaches 0.329.
+	// mean free path = 0.9 * 0.01 = 0.009, extinction = 111.1, transmittance = exp(extinction * 0.01) = 0.329
+	// mean free path = 1.0 * 0.01 = 0.01, extinction = 100.0, transmittance = exp(extinction * 0.01) = 0.367
+	const float MaxMeanFreePathMeters = max(PM.MeanFreePath.x, max(PM.MeanFreePath.y, PM.MeanFreePath.z));
+	return saturate((MaxMeanFreePathMeters / STRATA_SIMPLEVOLUME_THICKNESS_M) * 10.0f - 9.0f);
+}
+
+struct FStrataSlabSingleScattering
+{
+	float3 Transmittance;
+	float3 SingleScattering;
+};
+FStrataSlabSingleScattering StrataSlabSingleScattering(FParticipatingMedia PM, float NoV)
+{
+	const float3 SafeExtinctionThreshold = 0.000001f;
+	const float3 SafeExtinctionCoefficients = max(SafeExtinctionThreshold, PM.ExtinctionCoef);
+
+	const float PathLength = STRATA_SIMPLEVOLUME_THICKNESS_M / max(0.0001f, abs(NoV));
+	const float3 SafePathSegmentTransmittance = exp(-SafeExtinctionCoefficients * PathLength);
+
+	// Single scattering analytical integration from "Physically Based Unified Volumetric Rendering in Frostbite", slide 27/28.
+	FStrataSlabSingleScattering Result;
+	Result.SingleScattering = (PM.ScatteringCoef - PM.ScatteringCoef * SafePathSegmentTransmittance) / SafeExtinctionCoefficients;
+	Result.Transmittance = SafePathSegmentTransmittance;
+	return Result;
+}
+
+
 ///////////////////////////////////////////////////////////////////////////////
 // BSDFs packing/unpacking
 // 11G11B10F and 10F are unsigned float format
@@ -1187,6 +1259,8 @@ void PackStrataOut(
 		float3 LayerCoverage = 0.0f;
 
 		const bool bTopLayer = l == 0;
+		const bool bBottomLayer = l == (Strata.LayerCount - 1);
+		const bool bIsSimpleVolume = !bBottomLayer;
 
 		STRATA_UNROLL
 		for (int i = 0; i < Strata.Layers[l].BSDFCount; ++i)
@@ -1212,15 +1286,60 @@ void PackStrataOut(
 			{
 			case STRATA_BSDF_TYPE_SLAB:
 			{
-				const float3 F0 = ComputeF0(SLAB_SPECULAR(BSDF), SLAB_BASECOLOR(BSDF), SLAB_METALLIC(BSDF));
-				const float3 F90 = ComputeF90(F0, SLAB_EDGECOLOR(BSDF), SLAB_METALLIC(BSDF));
-
 				EmissiveColor += BSDF_GETEMISSIVE(BSDF) * Throughput;
 
-				bStrataSubsurfaceEnable = bStrataSubsurfaceEnable || (BSDF_GETHASSCATTERING(BSDF) && BSDFCoverageAvg > 0);
+				bStrataSubsurfaceEnable = (bStrataSubsurfaceEnable || (BSDF_GETHASSCATTERING(BSDF) && BSDFCoverageAvg > 0)) && !bIsSimpleVolume;
+
+				float3 F = 0;
+				// STRATA_TODO Fresnel throughput, Use a LUT, this is currently in progress and will be coming soon.
+				// In the meantime, we use a simple approach to feed Fresnel fnuctions and compute an approximation of the general transmition/througput.
+				const float VoH = ContextNoLight.NoV;
+				const float NoL = ContextNoLight.NoV;
+				const float3 BaseColor = SLAB_BASECOLOR(BSDF);
+				const float  Metallic = SLAB_METALLIC(BSDF);
+				const float3 F0 = ComputeF0(SLAB_SPECULAR(BSDF), BaseColor, Metallic);
+				const float3 F90 = ComputeF90(F0, SLAB_EDGECOLOR(BSDF), Metallic);
+				const bool bHasThinFilm = BSDF_GETHASTHINFILM(BSDF);
+				BRANCH
+				if (bHasThinFilm)
+				{
+					float ThinFilmIOR = SLAB_THINFILM_IOR(BSDF);
+					float ThinFilmThickness = SLAB_THINFILM_THICKNESS(BSDF);
+					F = F_ThinFilm(ContextNoLight.NoV, NoL, VoH, F0, F90, ThinFilmIOR, ThinFilmThickness);
+				}
+				else
+				{
+					F = F_Schlick(F0, F90, VoH);
+				}
+				float3 BSDFThroughput = 1.0f - F;
+
+				if (bIsSimpleVolume)
+				{
+				#if STRATA_INLINE_SHADING
+					// Now we know that this is 
+					SLAB_SSSDMFP(BSDF) = BSDF.TmpDMFP;
+				#endif
+
+					// Rescale the BSDF mean free path so that we can drop the thickness parameter and run volumetric material math on a unit slab of medium
+					// For instance, if thickness is larger => we need to reduce the MFP for the medium to appear visually thicker.
+					SLAB_SSSDMFP(BSDF) = bIsSimpleVolume ? SLAB_SSSDMFP(BSDF) * (STRATA_SIMPLEVOLUME_THICKNESS_CM / BSDF_GETTHICKNESSCM(BSDF)) : SLAB_SSSDMFP(BSDF);
+
+					const float3 UnusedDiffuseColor = 0; // Unused when only trasmittance is needed
+					FParticipatingMedia PM = StrataSlabCreateParticipatingMedia(UnusedDiffuseColor, SLAB_SSSDMFP(BSDF));
+					const float DiffuseToVolumeBlend = StrataSlabDiffuseToVolumeBlend(PM);
+					const FStrataSlabSingleScattering Result = StrataSlabSingleScattering(PM, dot(N,V));
+
+					// Update the see-through throughput accoding to the medium transmittance
+					BSDFThroughput = lerp(OpaqueBSDFThroughput, (1-F) * Result.Transmittance, DiffuseToVolumeBlend);
+
+					// Mark the BSDF has not having any SSS through post process.
+					BSDF_SETHASSCATTERING(Strata.Layers[l].BSDFs[i], 0);
+					BSDF_SETHASSSSPROFILE(Strata.Layers[l].BSDFs[i], 0);
+					bStrataSubsurfaceEnable = false;
+				}
 
 				LayerCoverage	+= BSDFCoverage;
-				LayerThroughput += BSDFCoverage * OpaqueBSDFThroughput; // STRATA_TODO: take into accound SSS transmittance
+				LayerThroughput += BSDFCoverage * BSDFThroughput;
 				break;
 			}
 
@@ -1259,6 +1378,10 @@ void PackStrataOut(
 			
 			// Notify that the BSDF is at the top for SSR to only affect reflection there and not on the lower layers
 			BSDF_SETISTOPLAYER(Strata.Layers[l].BSDFs[i], bTopLayer ? 1 : 0);
+
+			// Notify that the BSDF can use SSS post-process only if at the bottom (layered considered optically thick, e.g. cannot see-through)
+			// When forward rendering is used or if a layer is not a bottom layer, SSS post cannot be used and simple scattering is used.
+			BSDF_SETISSIMPLEVOLUME(Strata.Layers[l].BSDFs[i], bIsSimpleVolume ? 1 : 0);
 		}
 
 		// Take into account the non covered part for this layer and update throughput accordingly
@@ -1326,9 +1449,10 @@ void PackStrataOut(
 					STRATA_STORE_UINT1(PackColorLinearToGamma2AlphaLinear(float4(SLAB_EDGECOLOR(BSDF), SLAB_THINFILM_THICKNESS(BSDF))));
 				}
 
-				if (BSDF_GETHASSCATTERING(BSDF))
+				const bool bIsSimpleVolume = BSDF_GETISSIMPLEVOLUME(BSDF);
+				if (BSDF_GETHASSCATTERING(BSDF) || bIsSimpleVolume)
 				{
-					if (BSDF_GETHASSSSPROFILE(BSDF))
+					if (BSDF_GETHASSSSPROFILE(BSDF) && !bIsSimpleVolume)
 					{
 						STRATA_STORE_UINT1(PackRGBA8(float4(SLAB_SSSPROFILEID(BSDF), SLAB_SSSPROFILERADIUSSCALE(BSDF), 0, 0)));
 					}
@@ -1511,9 +1635,10 @@ FStrataBSDF UnpackStrataBSDFIn(ByteAddressBuffer StrataBuffer, inout FStrataAddr
 			StrataLoad_ColorGamma2ToLinearAlphaLinear(StrataBuffer, StrataAddressing, SLAB_EDGECOLOR(OutBSDF), SLAB_THINFILM_THICKNESS(OutBSDF));
 		}
 
-		if (BSDF_GETHASSCATTERING(OutBSDF))
+		const bool bIsSimpleVolume = BSDF_GETISSIMPLEVOLUME(OutBSDF);
+		if (BSDF_GETHASSCATTERING(OutBSDF) || bIsSimpleVolume)
 		{
-			if (BSDF_GETHASSSSPROFILE(OutBSDF))
+			if (BSDF_GETHASSSSPROFILE(OutBSDF) && !bIsSimpleVolume)
 			{
 				float4 Data = UnpackRGBA8(StrataLoadUint1(StrataBuffer, StrataAddressing));
 				SLAB_SSSPROFILEID(OutBSDF) = Data.r;
diff --git a/Engine/Shaders/Private/Strata/StrataEvaluation.ush b/Engine/Shaders/Private/Strata/StrataEvaluation.ush
--- a/Engine/Shaders/Private/Strata/StrataEvaluation.ush
+++ b/Engine/Shaders/Private/Strata/StrataEvaluation.ush
@@ -339,9 +339,29 @@ FStrataEvaluateResult StrataEvaluateBSDF(FStrataBSDFContext BSDFContext, FShadow
 			Sample.SpecularHazePathValue	= HazeD * HazeVis * F * HazeMSScale;
 			Sample.SpecularHazePDF			= HazePDF * Sample.SpecularHazePathProbability;
 
-			Sample.EmissivePathValue= BSDF_GETEMISSIVE(BSDFContext.BSDF);
+			Sample.EmissivePathValue = BSDF_GETEMISSIVE(BSDFContext.BSDF);
 
 			Sample.Throughput = OpaqueBSDFThroughput;
+
+			#if STRATA_FASTPATH==0
+			BRANCH
+			if (BSDF_GETISSIMPLEVOLUME(BSDFContext.BSDF))
+			{
+				FParticipatingMedia PM = StrataSlabCreateParticipatingMedia(DiffuseColor, SLAB_SSSDMFP(BSDFContext.BSDF));
+				const float DiffuseToVolumeBlend = StrataSlabDiffuseToVolumeBlend(PM);
+				const FStrataSlabSingleScattering Result = StrataSlabSingleScattering(PM, BSDFContext.Context.NoV);
+
+				const float3 SpecularThroughput = 1.0f - F; // This is already taken into account by the Chan model above. But for volumetric we need to apply it.
+				const float Phase = IsotropicPhase();
+
+				Sample.DiffusePathValue = lerp(Sample.DiffusePathValue,		Result.SingleScattering * Phase * SpecularThroughput,	DiffuseToVolumeBlend);
+				Sample.DiffusePDF		= lerp(Sample.DiffusePDF,			Phase,													DiffuseToVolumeBlend);
+				Sample.Throughput		= lerp(Sample.Throughput,			Result.Transmittance,									DiffuseToVolumeBlend);
+
+				Sample.bSubsurface = false;	// It should already be the case because Enforce in this case because BSDF_GETHASSCATTERING should be false when BSDF_GETISSIMPLEVOLUME is true
+			}
+			#endif
+
 			break;
 		}
 
@@ -437,7 +457,8 @@ FStrataImportanceSampleResult StrataImportanceSampleBSDF(FStrataPixelHeader Stra
 			const float3 BaseColor = SLAB_BASECOLOR(BSDFContext.BSDF);
 			const float  Metallic = SLAB_METALLIC(BSDFContext.BSDF);
 
-			// STRATA_TODO take into account roughness
+			// We assume a cosine hemisphere sampling is enough for any roughness.
+			// We also assume that this will work just as well for volumetric (considering it is an isotropic phase function and we are interesting in sampling mostly the hemisphere and not the layers below)
 			float4 ImportanceSample = CosineSampleHemisphere(E);
 			Sample.DiffuseDirection = mul(ImportanceSample.xyz, BSDFContext.TangentBasis);
 			PDF = ImportanceSample.w;
@@ -536,14 +557,14 @@ FStrataEnvLightResult StrataEvaluateForEnvLight(FStrataBSDFContext BSDFContext,
 			StrataEnvLightResult.DiffuseNormal = BSDFContext.N;
 			StrataEnvLightResult.bSubsurface = BSDF_GETHASSCATTERING(BSDFContext.BSDF) > 0;
 
+			float3 FresnelWeight = 1;
 			BRANCH
 			if (bEnableSpecular)
 			{
 				float3 F0 = ComputeF0(SLAB_SPECULAR(BSDFContext.BSDF), BaseColor, Metallic);
 				const float3 F90 = ComputeF90(F0, SLAB_EDGECOLOR(BSDFContext.BSDF), Metallic);
 
-				float3 FresnelWeight = 1;
-#if STRATA_FASTPATH==0
+				#if STRATA_FASTPATH==0
 				const bool bHasThinFilm = BSDF_GETHASTHINFILM(BSDFContext.BSDF);
 				if (bHasThinFilm)
 				{
@@ -552,11 +573,11 @@ FStrataEnvLightResult StrataEvaluateForEnvLight(FStrataBSDFContext BSDFContext,
 					FresnelWeight = F_ThinFilm(BSDFContext.Context.NoV, BSDFContext.Context.NoL, BSDFContext.Context.VoH, F0, F90, ThinFilmIOR, ThinFilmThickness);
 					F0 = 1.f;
 				}
-#endif
+				#endif
 
 				float SafeRoughnessX = MakeRoughnessSafe(SLAB_ROUGHNESSX(BSDFContext.BSDF));
 
-#if STRATA_FASTPATH==0
+				#if STRATA_FASTPATH==0
 				const bool bHasAnisotropy = BSDF_GETANISOTROPY(BSDFContext.BSDF);
 				if (bHasAnisotropy)
 				{
@@ -572,7 +593,7 @@ FStrataEnvLightResult StrataEvaluateForEnvLight(FStrataBSDFContext BSDFContext,
 					BSDFContext.SatNoL = saturate(BSDFContext.Context.NoL);
 					BSDFContext.R = 2 * dot(BSDFContext.V, BSDFContext.N) * BSDFContext.N - BSDFContext.V;
 				}
-#endif
+				#endif
 
 				float3 EvalEnvBRDF = FresnelWeight * StrataEnvBRDF(F0, F90, SafeRoughnessX, BSDFContext.SatNoV);
 
@@ -580,7 +601,7 @@ FStrataEnvLightResult StrataEvaluateForEnvLight(FStrataBSDFContext BSDFContext,
 				StrataEnvLightResult.SpecularWeight = EvalEnvBRDF;
 				StrataEnvLightResult.SpecularSafeRoughness = SafeRoughnessX;
 
-#if STRATA_FASTPATH==0
+				#if STRATA_FASTPATH==0
 				const bool bHasHaziness = BSDF_GETHASHAZINESS(BSDFContext.BSDF);
 				if (bHasHaziness)
 				{
@@ -597,8 +618,24 @@ FStrataEnvLightResult StrataEvaluateForEnvLight(FStrataBSDFContext BSDFContext,
 					StrataEnvLightResult.SpecularHazeWeight = lerp(0.0f, SPECULAR_HAZINESS_WEIGHT, HazinessFadeIn) * StrataEnvBRDF(F0, F90, SafeRoughnessHaze, BSDFContext.SatNoV);
 					StrataEnvLightResult.SpecularHazeSafeRoughness = SafeRoughnessHaze;
 				}
-#endif
+				#endif
 			}
+
+			#if STRATA_FASTPATH==0
+			BRANCH
+			if (BSDF_GETISSIMPLEVOLUME(BSDFContext.BSDF))
+			{
+				FParticipatingMedia PM = StrataSlabCreateParticipatingMedia(StrataEnvLightResult.DiffuseColor, SLAB_SSSDMFP(BSDFContext.BSDF));
+				const float DiffuseToVolumeBlend = StrataSlabDiffuseToVolumeBlend(PM);
+				const FStrataSlabSingleScattering Result = StrataSlabSingleScattering(PM, BSDFContext.Context.NoV);
+
+				// STRATA_TODO: big approximation for volumetric + env lighting: We ignore the phase for now due to the sampled convolved env lighting, 
+				// no self shadow and should not be single scattering using StrataSlabSingleScattering. We will propose a better approach. 
+				StrataEnvLightResult.DiffuseColor = lerp(StrataEnvLightResult.DiffuseColor, Result.SingleScattering * FresnelWeight,	DiffuseToVolumeBlend);
+				// We keep the same normal
+				StrataEnvLightResult.bSubsurface = false;
+			}
+			#endif
 			break;
 		}
 
diff --git a/Engine/Shaders/Private/Strata/StrataVisualize.usf b/Engine/Shaders/Private/Strata/StrataVisualize.usf
--- a/Engine/Shaders/Private/Strata/StrataVisualize.usf
+++ b/Engine/Shaders/Private/Strata/StrataVisualize.usf
@@ -185,6 +185,7 @@ void VisualizeMaterialPS(
 #define PRINTBSDFANISO()	PRINTBSDFCHAR(_SPC_); PRINTBSDFCHAR(_MINUS_); PRINTBSDFCHAR(_A_); PRINTBSDFCHAR(_N_); PRINTBSDFCHAR(_I_); PRINTBSDFCHAR(_S_); PRINTBSDFCHAR(_O_); PRINTBSDFCHAR(_SPC_); PrintFloatNoFraction(PixelPos, OutColor.xyz, FontColor, TopLeft, float(BSDF_GETANISOTROPY(BSDF)), 1); TopLeft.x+=8*1;
 #define PRINTBSDFTOPLAY()	PRINTBSDFCHAR(_SPC_); PRINTBSDFCHAR(_MINUS_); PRINTBSDFCHAR(_T_); PRINTBSDFCHAR(_O_); PRINTBSDFCHAR(_P_); PRINTBSDFCHAR(_L_); PRINTBSDFCHAR(_A_); PRINTBSDFCHAR(_Y_); PRINTBSDFCHAR(_SPC_); PrintFloatNoFraction(PixelPos, OutColor.xyz, FontColor, TopLeft, float(BSDF_GETISTOPLAYER(BSDF)), 1); TopLeft.x+=8*1;
 #define PRINTBSDFHASSCAT()	PRINTBSDFCHAR(_SPC_); PRINTBSDFCHAR(_MINUS_); PRINTBSDFCHAR(_H_); PRINTBSDFCHAR(_A_); PRINTBSDFCHAR(_S_); PRINTBSDFCHAR(_S_); PRINTBSDFCHAR(_C_); PRINTBSDFCHAR(_A_); PRINTBSDFCHAR(_T_); PRINTBSDFCHAR(_SPC_); PrintFloatNoFraction(PixelPos, OutColor.xyz, FontColor, TopLeft, float(BSDF_GETHASSCATTERING(BSDF)), 1); TopLeft.x+=8*1;
+#define PRINTBSDFSIMPSCAT()	PRINTBSDFCHAR(_SPC_); PRINTBSDFCHAR(_MINUS_); PRINTBSDFCHAR(_S_); PRINTBSDFCHAR(_I_); PRINTBSDFCHAR(_M_); PRINTBSDFCHAR(_P_); PRINTBSDFCHAR(_S_); PRINTBSDFCHAR(_C_); PRINTBSDFCHAR(_A_); PRINTBSDFCHAR(_T_); PRINTBSDFCHAR(_SPC_); PrintFloatNoFraction(PixelPos, OutColor.xyz, FontColor, TopLeft, float(BSDF_GETISSIMPLEVOLUME(BSDF)), 1); TopLeft.x+=8*1;
 #define PRINTFLOAT3(X) PrintFloat(PixelPos, OutColor.xyz, PARAMFontColor, TopLeft, X.x); TopLeft.x+=8*7; PRINTBSDFCHAR(_SPC_); PrintFloat(PixelPos, OutColor.xyz, PARAMFontColor, TopLeft, X.y); TopLeft.x+=8*7; PRINTBSDFCHAR(_SPC_); PrintFloat(PixelPos, OutColor.xyz, PARAMFontColor, TopLeft, X.z);
 #define PRINTFLOAT1(Value) PrintFloat(PixelPos, OutColor.xyz, PARAMFontColor, TopLeft, Value); TopLeft.x+=8*7;
 #define PRINTFLOATNOFRAC(Value, DigitCount) PrintFloatNoFraction(PixelPos, OutColor.xyz, PARAMFontColor, TopLeft, Value, DigitCount); TopLeft.x+=8*7;
@@ -309,6 +310,7 @@ void VisualizeMaterialPS(
 				PRINTBSDFANISO();
 				PRINTBSDFTOPLAY();
 				PRINTBSDFHASSCAT();
+				PRINTBSDFSIMPSCAT();
 				NEXTLINE;
 
 				HORIZONTALTAB; PRINTPARAMCHAR(_A_); PRINTPARAMCHAR(_L_); PRINTPARAMCHAR(_B_); PRINTPARAMCHAR(_E_); PRINTPARAMCHAR(_D_); PRINTPARAMCHAR(_O_); 
@@ -344,9 +346,9 @@ void VisualizeMaterialPS(
 				HORIZONTALTAB; PRINTFLOAT1(SLAB_HAZINESS(BSDF));
 				NEXTLINE;
 
-				if (BSDF_GETHASSCATTERING(BSDF))
+				if (BSDF_GETHASSCATTERING(BSDF) || BSDF_GETISSIMPLEVOLUME(BSDF))
 				{
-					if (BSDF_GETHASSSSPROFILE(BSDF))
+					if (BSDF_GETHASSSSPROFILE(BSDF) && BSDF_GETISSIMPLEVOLUME(BSDF)==0)
 					{
 						HORIZONTALTAB; PRINTPARAMCHAR(_S_); PRINTPARAMCHAR(_S_); PRINTPARAMCHAR(_S_); PRINTPARAMCHAR(_SPC_); PRINTPARAMCHAR(_I_); PRINTPARAMCHAR(_D_); PRINTPARAMCHAR(_SPC_);
 						HORIZONTALTAB; PRINTFLOATNOFRAC(floor(SLAB_SSSPROFILEID(BSDF)*255.0f+0.5f),3);
diff --git a/Engine/Shaders/Private/ThinFilmBSDF.ush b/Engine/Shaders/Private/ThinFilmBSDF.ush
--- a/Engine/Shaders/Private/ThinFilmBSDF.ush
+++ b/Engine/Shaders/Private/ThinFilmBSDF.ush
@@ -1,5 +1,7 @@
 // Copyright Epic Games, Inc. All Rights Reserved.
 
+#pragma once
+
 /////////////////////////////////////////////////////////////////////////////////////////////////
 // Iridescence
 
diff --git a/Engine/Source/Runtime/Engine/Classes/Materials/MaterialExpressionStrata.h b/Engine/Source/Runtime/Engine/Classes/Materials/MaterialExpressionStrata.h
--- a/Engine/Source/Runtime/Engine/Classes/Materials/MaterialExpressionStrata.h
+++ b/Engine/Source/Runtime/Engine/Classes/Materials/MaterialExpressionStrata.h
@@ -85,7 +85,7 @@ class UMaterialExpressionStrataSlabBSDF : public UMaterialExpressionStrataBSDF
 	FExpressionInput Tangent;
 
 	/**
-	 * Chromatic mean free path . Only used when there is not any sub-surface profile provided. (type = float3, unit = unitless)
+	 * Chromatic mean free path . Only used when there is not any sub-surface profile provided. (type = float3, unit = centimeters, default = 0)
 	 */
 	UPROPERTY()
 	FExpressionInput SSSDMFP;
@@ -109,11 +109,17 @@ class UMaterialExpressionStrataSlabBSDF : public UMaterialExpressionStrataBSDF
 	FExpressionInput Haziness;
 
 	/**
-	 * Thin film contosl the thin film layer coating the current slab. 0 means disabled and 1 means a coating layer of 10 micrometer. (type = float, unitless, default = 0)
+	 * Thin film controls the thin film layer coating the current slab. 0 means disabled and 1 means a coating layer of 10 micrometer. (type = float, unitless, default = 0)
 	 */
 	UPROPERTY()
 	FExpressionInput ThinFilmThickness;
 
+	/**
+	 * The slab thickness. (type = float, centimeters, default = 0.01 centimeter = 0.1 millimeter)
+	 */
+	UPROPERTY()
+	FExpressionInput Thickness;
+
 	/** SubsurfaceProfile, for Screen Space Subsurface Scattering. The profile needs to be set up on both the Strata diffuse node, and the material node at the moment. */
 	UPROPERTY(EditAnywhere, BlueprintReadOnly, Category = Material, meta = (DisplayName = "Subsurface Profile"))
 	TObjectPtr<class USubsurfaceProfile> SubsurfaceProfile;
diff --git a/Engine/Source/Runtime/Engine/Private/Materials/HLSLMaterialTranslator.cpp b/Engine/Source/Runtime/Engine/Private/Materials/HLSLMaterialTranslator.cpp
--- a/Engine/Source/Runtime/Engine/Private/Materials/HLSLMaterialTranslator.cpp
+++ b/Engine/Source/Runtime/Engine/Private/Materials/HLSLMaterialTranslator.cpp
@@ -9025,10 +9025,10 @@ int32 FHLSLMaterialTranslator::StrataCreateAndRegisterNullMaterial()
 int32 FHLSLMaterialTranslator::StrataSlabBSDF(
 	int32 BaseColor, int32 EdgeColor, int32 Specular, int32 Metallic, int32 RoughnessX, int32 RoughnessY,
 	int32 SSSProfileId, int32 SSSDMFP, int32 SSSDMFPScale, int32 EmissiveColor,
-	int32 Haziness, int32 ThinFilmThickness, int32 Normal, int32 Tangent, uint8 SharedNormalIndex)
+	int32 Haziness, int32 ThinFilmThickness, int32 Thickness, int32 Normal, int32 Tangent, uint8 SharedNormalIndex)
 {
 	return AddCodeChunk(
-		MCT_Strata, TEXT("GetStrataSlabBSDF(%s, %s, %s, %s, float2(%s, %s), %s, %s, %s, %s, %s, %s, %u, Parameters.SharedNormals.NormalTypes) /* Normal = %s ; Tangent = %s */"),
+		MCT_Strata, TEXT("GetStrataSlabBSDF(%s, %s, %s, %s, float2(%s, %s), %s, %s, %s, %s, %s, %s, %s, %u, Parameters.SharedNormals.NormalTypes) /* Normal = %s ; Tangent = %s */"),
 		*GetParameterCode(BaseColor),
 		*GetParameterCode(EdgeColor),
 		*GetParameterCode(Specular),
@@ -9041,6 +9041,7 @@ int32 FHLSLMaterialTranslator::StrataSlabBSDF(
 		*GetParameterCode(EmissiveColor),
 		*GetParameterCode(Haziness),
 		*GetParameterCode(ThinFilmThickness),
+		*GetParameterCode(Thickness),
 		SharedNormalIndex,
 		*GetParameterCode(Normal),
 		*GetParameterCode(Tangent)
diff --git a/Engine/Source/Runtime/Engine/Private/Materials/HLSLMaterialTranslator.h b/Engine/Source/Runtime/Engine/Private/Materials/HLSLMaterialTranslator.h
--- a/Engine/Source/Runtime/Engine/Private/Materials/HLSLMaterialTranslator.h
+++ b/Engine/Source/Runtime/Engine/Private/Materials/HLSLMaterialTranslator.h
@@ -991,7 +991,8 @@ protected:
 	virtual int32 StrataSlabBSDF(
 		int32 BaseColor, int32 EdgeColor, int32 Specular, int32 Metallic, int32 RoughnessX, int32 RoughnessY,
 		int32 SSSProfileId, int32 SSSDMFP, int32 SSSDMFPScale,
-		int32 EmissiveColor, int32 Haziness, int32 ThinFilmThickness, int32 Normal, int32 Tangent, uint8 SharedNormalIndex) override;
+		int32 EmissiveColor, int32 Haziness, int32 ThinFilmThickness, int32 Thickness,
+		int32 Normal, int32 Tangent, uint8 SharedNormalIndex) override;
 	virtual int32 StrataSheenBSDF(int32 BaseColor, int32 Roughness, int32 Normal, uint8 SharedNormalIndex) override;
 	virtual int32 StrataVolumetricFogCloudBSDF(int32 Albedo, int32 Extinction, int32 EmissiveColor, int32 AmbientOcclusion) override;
 	virtual int32 StrataUnlitBSDF(int32 EmissiveColor, int32 TransmittanceColor) override;
diff --git a/Engine/Source/Runtime/Engine/Private/Materials/MaterialExpressions.cpp b/Engine/Source/Runtime/Engine/Private/Materials/MaterialExpressions.cpp
--- a/Engine/Source/Runtime/Engine/Private/Materials/MaterialExpressions.cpp
+++ b/Engine/Source/Runtime/Engine/Private/Materials/MaterialExpressions.cpp
@@ -19775,6 +19775,7 @@ int32 UMaterialExpressionStrataSlabBSDF::Compile(class FMaterialCompiler* Compil
 		CompileWithDefaultFloat3(Compiler, EmissiveColor, 0.0f, 0.0f, 0.0f),
 		CompileWithDefaultFloat1(Compiler, Haziness, 0.0f),
 		CompileWithDefaultFloat1(Compiler, ThinFilmThickness, 0.0f),
+		CompileWithDefaultFloat1(Compiler, Thickness, 0.01f),
 		NormalCodeChunk,
 		TangentCodeChunk,
 		SharedNormalIndex);
@@ -19836,6 +19837,9 @@ uint32 UMaterialExpressionStrataSlabBSDF::GetInputType(int32 InputIndex)
 	case 12:
 		return MCT_Float1; // ThinFilm Thickness
 		break;
+	case 13:
+		return MCT_Float1; // Thickness
+		break;
 	}
 
 	check(false);
@@ -19896,6 +19900,10 @@ FName UMaterialExpressionStrataSlabBSDF::GetInputName(int32 InputIndex) const
 	{
 		return TEXT("ThinFilm Thickness");
 	}
+	else if (InputIndex == 13)
+	{
+		return TEXT("Thickness");
+	}
 	return TEXT("Unknown");
 }
 
diff --git a/Engine/Source/Runtime/Engine/Public/MaterialCompiler.h b/Engine/Source/Runtime/Engine/Public/MaterialCompiler.h
--- a/Engine/Source/Runtime/Engine/Public/MaterialCompiler.h
+++ b/Engine/Source/Runtime/Engine/Public/MaterialCompiler.h
@@ -412,7 +412,8 @@ public:
 	virtual int32 StrataSlabBSDF(
 		int32 BaseColor, int32 EdgeColor, int32 Specular, int32 Metallic, int32 RoughnessX, int32 RoughnessY, 
 		int32 SSSProfileId, int32 SSSDMFP, int32 SSSDMFPScale,
-		int32 EmissiveColor, int32 Haziness, int32 ThinFilmThickness, int32 Normal, int32 Tangent, uint8 SharedNormalIndex) = 0;
+		int32 EmissiveColor, int32 Haziness, int32 ThinFilmThickness, int32 Thickness,
+		int32 Normal, int32 Tangent, uint8 SharedNormalIndex) = 0;
 	virtual int32 StrataSheenBSDF(int32 BaseColor, int32 Roughness, int32 Normal, uint8 SharedNormalIndex) = 0;
 	virtual int32 StrataVolumetricFogCloudBSDF(int32 Albedo, int32 Extinction, int32 EmissiveColor, int32 AmbientOcclusion) = 0;
 	virtual int32 StrataUnlitBSDF(int32 EmissiveColor, int32 TransmittanceColor) = 0;
@@ -894,12 +895,13 @@ public:
 	virtual int32 StrataSlabBSDF(
 		int32 BaseColor, int32 EdgeColor, int32 Specular, int32 Metallic, int32 RoughnessX, int32 RoughnessY,
 		int32 SSSProfileId, int32 SSSDMFP, int32 SSSDMFPScale,
-		int32 EmissiveColor, int32 Haziness, int32 ThinFilmThickness, int32 Normal, int32 Tangent, uint8 SharedNormalIndex) override
+		int32 EmissiveColor, int32 Haziness, int32 ThinFilmThickness, int32 Thickness,
+		int32 Normal, int32 Tangent, uint8 SharedNormalIndex) override
 	{
 		return Compiler->StrataSlabBSDF(
 			BaseColor, EdgeColor, Specular, Metallic, RoughnessX, RoughnessY,
 			SSSProfileId, SSSDMFP, SSSDMFPScale,
-			EmissiveColor, Haziness, ThinFilmThickness, Normal, Tangent, SharedNormalIndex);
+			EmissiveColor, Haziness, ThinFilmThickness, Thickness, Normal, Tangent, SharedNormalIndex);
 	}
 
 	virtual int32 StrataSheenBSDF(int32 BaseColor, int32 Roughness, int32 Normal, uint8 SharedNormalIndex) override
