Fix incorrect coverage computation for hair strands texture computation

diff --git a/Engine/Shaders/Private/HairStrands/HairStrandsTexturesGeneration.usf b/Engine/Shaders/Private/HairStrands/HairStrandsTexturesGeneration.usf
--- a/Engine/Shaders/Private/HairStrands/HairStrandsTexturesGeneration.usf
+++ b/Engine/Shaders/Private/HairStrands/HairStrandsTexturesGeneration.usf
@@ -413,7 +413,7 @@ void MainPS(
 			dot(Tangent_ObjectSpace, B), 
 			dot(Tangent_ObjectSpace, N));
 
-		OutCoverage /= bValidCount;
+		OutCoverage /= SampleCount;
 		OutTangent.xyz = (Tangent_FrameSpace + 1) * 0.5f;
 		OutTangent.w = 1;
 	}
