Do not set SimpleVolume on Strata slab node if DMFP is not plugin to same a uint in this case.
Make the slab default diffuse completely opaque when simpe volumetric is not used.

diff --git a/Engine/Shaders/Private/Strata/Strata.ush b/Engine/Shaders/Private/Strata/Strata.ush
--- a/Engine/Shaders/Private/Strata/Strata.ush
+++ b/Engine/Shaders/Private/Strata/Strata.ush
@@ -1260,7 +1260,6 @@ void PackStrataOut(
 
 		const bool bTopLayer = l == 0;
 		const bool bBottomLayer = l == (Strata.LayerCount - 1);
-		const bool bIsSimpleVolume = !bBottomLayer;
 
 		STRATA_UNROLL
 		for (int i = 0; i < Strata.Layers[l].BSDFCount; ++i)
@@ -1288,6 +1287,9 @@ void PackStrataOut(
 			{
 				EmissiveColor += BSDF_GETEMISSIVE(BSDF) * Throughput;
 
+				// We want simple volumetric only if the DMFP input is pluged in, otherwise we avoid storing the DFMP
+				const bool bIsSimpleVolume = !bBottomLayer && BSDF_GETHASSCATTERING(Strata.Layers[l].BSDFs[i]); 
+
 				bStrataSubsurfaceEnable = (bStrataSubsurfaceEnable || (BSDF_GETHASSCATTERING(BSDF) && BSDFCoverageAvg > 0)) && !bIsSimpleVolume;
 
 				float3 F = 0;
@@ -1311,7 +1313,9 @@ void PackStrataOut(
 				{
 					F = F_Schlick(F0, F90, VoH);
 				}
-				float3 BSDFThroughput = 1.0f - F;
+
+				// By default, without coverage, the diffuse BSDF is completely opaque
+				float3 BSDFThroughput = OpaqueBSDFThroughput;
 
 				if (bIsSimpleVolume)
 				{
@@ -1336,6 +1340,10 @@ void PackStrataOut(
 					BSDF_SETHASSCATTERING(Strata.Layers[l].BSDFs[i], 0);
 					BSDF_SETHASSSSPROFILE(Strata.Layers[l].BSDFs[i], 0);
 					bStrataSubsurfaceEnable = false;
+
+					// Notify that the Slab can use SSS post-process only if at the bottom (layered considered optically thick, e.g. cannot see-through)
+					// When forward rendering is used or if a layer is not a bottom layer, SSS post cannot be used and simple scattering is used.
+					BSDF_SETISSIMPLEVOLUME(Strata.Layers[l].BSDFs[i], bIsSimpleVolume ? 1 : 0);
 				}
 
 				LayerCoverage	+= BSDFCoverage;
@@ -1378,10 +1386,6 @@ void PackStrataOut(
 			
 			// Notify that the BSDF is at the top for SSR to only affect reflection there and not on the lower layers
 			BSDF_SETISTOPLAYER(Strata.Layers[l].BSDFs[i], bTopLayer ? 1 : 0);
-
-			// Notify that the BSDF can use SSS post-process only if at the bottom (layered considered optically thick, e.g. cannot see-through)
-			// When forward rendering is used or if a layer is not a bottom layer, SSS post cannot be used and simple scattering is used.
-			BSDF_SETISSIMPLEVOLUME(Strata.Layers[l].BSDFs[i], bIsSimpleVolume ? 1 : 0);
 		}
 
 		// Take into account the non covered part for this layer and update throughput accordingly
