* Add checkerboard SSS support for Strata
 * Fix missing alpha channel output from Strata env lighting evaluation, causing the diffuse lighting ratio to be incorrectly computed with SSS.

diff --git a/Engine/Shaders/Private/BasePassPixelShader.usf b/Engine/Shaders/Private/BasePassPixelShader.usf
--- a/Engine/Shaders/Private/BasePassPixelShader.usf
+++ b/Engine/Shaders/Private/BasePassPixelShader.usf
@@ -940,11 +940,12 @@ void FPixelShaderInOut_MainPS(
 	Metallic *= Opacity;
 #endif
 
-	// 0..1, SubsurfaceProfileId = int(x * 255)
-	float SubsurfaceProfile = 0;
-
 	// If we don't use this shading model the color should be black (don't generate shader code for unused data, don't do indirectlighting cache lighting with this color).
 	float3 SubsurfaceColor = 0;
+
+#if !STRATA_ENABLED
+	// 0..1, SubsurfaceProfileId = int(x * 255)
+	float SubsurfaceProfile = 0;
 #if MATERIAL_SHADINGMODEL_SUBSURFACE || MATERIAL_SHADINGMODEL_PREINTEGRATED_SKIN || MATERIAL_SHADINGMODEL_SUBSURFACE_PROFILE || MATERIAL_SHADINGMODEL_TWOSIDED_FOLIAGE || MATERIAL_SHADINGMODEL_CLOTH || MATERIAL_SHADINGMODEL_EYE
 	if (ShadingModel == SHADINGMODELID_SUBSURFACE || ShadingModel == SHADINGMODELID_PREINTEGRATED_SKIN || ShadingModel == SHADINGMODELID_SUBSURFACE_PROFILE || ShadingModel == SHADINGMODELID_TWOSIDED_FOLIAGE || ShadingModel == SHADINGMODELID_CLOTH || ShadingModel == SHADINGMODELID_EYE)
 	{
@@ -969,6 +970,7 @@ void FPixelShaderInOut_MainPS(
 		SubsurfaceProfile = SubsurfaceData.a;
 	}
 #endif
+#endif // !STRATA_ENABLED
 
 #if STRATA_ENABLED && !MATERIAL_SHADINGMODEL_UNLIT
 	FStrataData Strata = PixelMaterialInputs.FrontMaterial;
@@ -1111,7 +1113,7 @@ void FPixelShaderInOut_MainPS(
 	}
 #endif
 
-	bool bChecker = CheckerFromPixelPos(MaterialParameters.SvPosition.xy);
+	const bool bChecker = CheckerFromPixelPos(MaterialParameters.SvPosition.xy);
 #if !STRATA_ENABLED || STRATA_INLINE_SINGLELAYERWATER
 
 	// So that the following code can still use DiffuseColor and SpecularColor.
diff --git a/Engine/Shaders/Private/DeferredShadingCommon.ush b/Engine/Shaders/Private/DeferredShadingCommon.ush
--- a/Engine/Shaders/Private/DeferredShadingCommon.ush
+++ b/Engine/Shaders/Private/DeferredShadingCommon.ush
@@ -578,7 +578,7 @@ void EncodeGBuffer(
 // SubsurfaceProfile does deferred lighting with a checker board pixel pattern
 // we separate the view from the non view dependent lighting and later recombine the two color constributions in a postprocess
 // We have the option to apply the BaseColor/Specular in the base pass or do it later in the postprocess (has implications to texture detail, fresnel and performance)
-void AdjustBaseColorAndSpecularColorForSubsurfaceProfileLighting(inout float3 BaseColor, inout float3 SpecularColor, inout float Specular, bool bChecker)
+bool AdjustBaseColorAndSpecularColorForSubsurfaceProfileLighting(inout float3 BaseColor, inout float Specular, bool bChecker)
 {
 #if SUBSURFACE_CHANNEL_MODE == 0
 	// If SUBSURFACE_CHANNEL_MODE is 0, we can't support full-resolution lighting, so we 
@@ -591,12 +591,19 @@ void AdjustBaseColorAndSpecularColorForSubsurfaceProfileLighting(inout float3 Ba
 	if (bCheckerboardRequired)
 	{
 		// because we adjust the BaseColor here, we need StoredBaseColor
-
 		// we apply the base color later in SubsurfaceRecombinePS()
 		BaseColor = bChecker;
+		Specular *= !bChecker;
+	}
+	return bCheckerboardRequired;
+}
+void AdjustBaseColorAndSpecularColorForSubsurfaceProfileLighting(inout float3 BaseColor, inout float3 SpecularColor, inout float Specular, bool bChecker)
+{
+	const bool bCheckerboardRequired = AdjustBaseColorAndSpecularColorForSubsurfaceProfileLighting(BaseColor, Specular, bChecker);
+	if (bCheckerboardRequired)
+	{
 		// in SubsurfaceRecombinePS() does not multiply with Specular so we do it here
 		SpecularColor *= !bChecker;
-		Specular *= !bChecker;
 	}
 }
 
diff --git a/Engine/Shaders/Private/ReflectionEnvironmentPixelShader.usf b/Engine/Shaders/Private/ReflectionEnvironmentPixelShader.usf
--- a/Engine/Shaders/Private/ReflectionEnvironmentPixelShader.usf
+++ b/Engine/Shaders/Private/ReflectionEnvironmentPixelShader.usf
@@ -389,6 +389,7 @@ void ReflectionEnvironmentSkyLighting(
 		ApplyCloudVolumetricShadow(WorldPosition, SkyLighting);
 
 		OutColor.xyz = SkyLighting;
+		OutColor.a = StrataEnvLighting.a;
 	}
 
 #endif // ENABLE_SKY_LIGHT
diff --git a/Engine/Shaders/Private/ShadowProjectionPixelShader.usf b/Engine/Shaders/Private/ShadowProjectionPixelShader.usf
--- a/Engine/Shaders/Private/ShadowProjectionPixelShader.usf
+++ b/Engine/Shaders/Private/ShadowProjectionPixelShader.usf
@@ -262,13 +262,11 @@ void Main(
 		bool bIsEyeOrHair 			= false;
 	#if STRATA_ENABLED
 		const FStrataClassification StrataData = StrataUnpackClassificationData(Strata.ClassificationTexture.Load(PixelCoord));
+		if (StrataData.bIsStrataMaterial)
 		{
-			if (StrataData.bIsStrataMaterial)
-			{
-				bIsEyeOrHair			= StrataHasShadingModel(StrataData.ShadingModels, STRATA_BSDF_TYPE_HAIR);
-				bHasSubsurface			= StrataData.bHasSubsurface || bIsEyeOrHair;
-				bHasSubsurfaceProfile	= StrataData.bHasSubsurfaceProfile;
-			}
+			bIsEyeOrHair			= StrataHasShadingModel(StrataData.ShadingModels, STRATA_BSDF_TYPE_HAIR);
+			bHasSubsurface			= StrataData.bHasSubsurface || bIsEyeOrHair;
+			bHasSubsurfaceProfile	= StrataData.bHasSubsurfaceProfile;
 		}
 	#else // STRATA_ENABLED
 		FGBufferData GBufferData = GetGBufferData(ScreenUV);
@@ -330,14 +328,14 @@ void Main(
 			  #if USE_TRANSMISSION
 				if (bHasSubsurfaceProfile)
 				{
-		#if STRATA_ENABLED
+				#if STRATA_ENABLED
 					const FStrataTopLayerData TopLayerData = StrataUnpackTopLayerData(Strata.TopLayerNormalTexture.Load(uint3(SVPos.xy, 0)));
 					const float3 WorldNormal = TopLayerData.WorldNormal;
-					FTransmissionProfileParams TransmissionParams = GetTransmissionProfileParams(SSSData.ProfileId);
-		#else // STRATA_ENABLED
+					const FTransmissionProfileParams TransmissionParams = GetTransmissionProfileParams(SSSData.ProfileId);
+				#else // STRATA_ENABLED
 					float3 WorldNormal = GBufferData.WorldNormal;
 					FTransmissionProfileParams TransmissionParams = GetTransmissionProfileParams(GBufferData);
-		#endif // STRATA_ENABLED
+				#endif // STRATA_ENABLED
 					SSSTransmission = CalcTransmissionThickness(ScreenPosition.xyz, LightPositionAndInvRadius.xyz, WorldNormal, TransmissionParams, Settings);
 				}
 				else
@@ -483,7 +481,7 @@ void MainOnePassPointLightPS(
 
 	const bool IsSubsurfaceProfile = StrataData.bHasSubsurfaceProfile;
 	const float3 WorldNormal = TopLayerData.WorldNormal;
-	FTransmissionProfileParams TransmissionParams = GetTransmissionProfileParams(SSSData.ProfileId);
+	const FTransmissionProfileParams TransmissionParams = GetTransmissionProfileParams(SSSData.ProfileId);
 
 #else // STRATA_ENABLED
 
diff --git a/Engine/Shaders/Private/Strata/Strata.ush b/Engine/Shaders/Private/Strata/Strata.ush
--- a/Engine/Shaders/Private/Strata/Strata.ush
+++ b/Engine/Shaders/Private/Strata/Strata.ush
@@ -17,6 +17,14 @@
 #define STRATA_INLINE_SHADING 1
 #endif
 
+// During the unpacking of the Strata material data, if a slab has some sub-surface scattering, its BaseColor/Specular value will be 
+// patched/overriden with white albedo (and optionally no specular) for computing the incoming irradiance without any albedo information. 
+// This is required by the post-process/screen-space sub-surface methods. This overriden can be opt-out. This is used during the classification 
+// pass for writing out the correct base color value into the FStrataSubsurfaceData
+#ifndef STRATA_SSS_MATERIAL_OVERRIDE
+#define STRATA_SSS_MATERIAL_OVERRIDE 1
+#endif
+
 #ifndef STRATA_CLAMPED_LAYER_COUNT
 #define STRATA_CLAMPED_LAYER_COUNT 1
 #endif
@@ -589,9 +597,10 @@ uint DivideAndRoundUp(uint Dividend, uint Divisor, uint DivisorAsBitShift)
 
 struct FStrataAddressing
 {
-	uint CurrentAddr;
-	uint ByteOffsetToNextUint;
-	uint ReadBytes;
+	uint  CurrentAddr;
+	uint  ByteOffsetToNextUint;
+	uint  ReadBytes;
+	uint2 PixelCoords;
 };
 FStrataAddressing GetStrataPixelDataByteOffset(uint2 PixelPos, uint2 ViewBufferSize, uint StrataMaxBytesPerPixel)
 {
@@ -621,6 +630,7 @@ FStrataAddressing GetStrataPixelDataByteOffset(uint2 PixelPos, uint2 ViewBufferS
 
 #endif
 
+	StrataAddressing.PixelCoords = PixelPos;
 	StrataAddressing.ReadBytes = 0;
 	return StrataAddressing;
 }
@@ -1652,6 +1662,16 @@ FStrataBSDF UnpackStrataBSDFIn(ByteAddressBuffer StrataBuffer, inout FStrataAddr
 			{
 				StrataLoad_R11G11B10F(StrataBuffer, StrataAddressing, SLAB_SSSDMFP(OutBSDF));
 			}
+
+			// When loading Strata BSDF which have SSS enable, BaseColor and (optionally) Specular value can be overriden 
+			// based on the SSS method used (screen space SSS / checherboarded screen space SSS / ...)
+			#if STRATA_SSS_MATERIAL_OVERRIDE
+			if (!bIsSimpleVolume)
+			{
+				const bool bChecker = CheckerFromPixelPos(StrataAddressing.PixelCoords);
+				AdjustBaseColorAndSpecularColorForSubsurfaceProfileLighting(SLAB_BASECOLOR(OutBSDF), SLAB_SPECULAR(OutBSDF), bChecker);
+			}
+			#endif
 		}
 	}
 	break;
diff --git a/Engine/Shaders/Private/Strata/StrataEvaluation.ush b/Engine/Shaders/Private/Strata/StrataEvaluation.ush
--- a/Engine/Shaders/Private/Strata/StrataEvaluation.ush
+++ b/Engine/Shaders/Private/Strata/StrataEvaluation.ush
@@ -222,7 +222,7 @@ FStrataEvaluateResult StrataEvaluateBSDF(FStrataBSDFContext BSDFContext, FShadow
 			const bool bHasSubsurface = BSDF_GETHASSCATTERING(BSDFContext.BSDF) > 0;
 			const float3 BaseColor = SLAB_BASECOLOR(BSDFContext.BSDF);
 			const float  Metallic = SLAB_METALLIC(BSDFContext.BSDF);
-			const float3 DiffuseColor = bHasSubsurface ? float(1.f).xxx : BaseColor - BaseColor * Metallic;
+			const float3 DiffuseColor = BaseColor - BaseColor * Metallic;
 			const float SafeRoughness = MakeRoughnessSafe(SLAB_ROUGHNESSX(BSDFContext.BSDF));
 			float a2 = Pow4(SafeRoughness);
 
diff --git a/Engine/Shaders/Private/Strata/StrataMaterialClassification.usf b/Engine/Shaders/Private/Strata/StrataMaterialClassification.usf
--- a/Engine/Shaders/Private/Strata/StrataMaterialClassification.usf
+++ b/Engine/Shaders/Private/Strata/StrataMaterialClassification.usf
@@ -3,6 +3,7 @@
 #include "/Engine/Private/Common.ush"
 
 #define STRATA_INLINE_SHADING 0
+#define STRATA_SSS_MATERIAL_OVERRIDE 0
 #include "/Engine/Private/Strata/Strata.ush"
 
 #define TILE_SIZE 8
diff --git a/Engine/Shaders/Private/Strata/StrataVisualize.usf b/Engine/Shaders/Private/Strata/StrataVisualize.usf
--- a/Engine/Shaders/Private/Strata/StrataVisualize.usf
+++ b/Engine/Shaders/Private/Strata/StrataVisualize.usf
@@ -4,6 +4,7 @@
 #include "/Engine/Private/DeferredShadingCommon.ush"
 
 #define STRATA_INLINE_SHADING 0
+#define STRATA_SSS_MATERIAL_OVERRIDE 0
 #include "/Engine/Private/Strata/StrataEvaluation.ush"
 
 #include "/Engine/Private/MiniFontCommon.ush"
