Strata slab can now be rendred translucent with colored transmittance/translucency.
Strata simple volume is enabled when DMFP is plugged in, There is SSS needed as it should be separated.
Strata naming clarification (because HasScattering !=HasSSS)
Compiler side BSDF size is now finer grained thanks to that (when DMPF is not plug in for top bsdf).

diff --git a/Engine/Shaders/Private/Strata/Strata.ush b/Engine/Shaders/Private/Strata/Strata.ush
--- a/Engine/Shaders/Private/Strata/Strata.ush
+++ b/Engine/Shaders/Private/Strata/Strata.ush
@@ -114,14 +114,15 @@ struct FStrataBSDF
 #define STATE_BIT_COUNT_SHAREDNORMALID			2
 #define STATE_BIT_COUNT_ANISOTROPY				1
 #define STATE_BIT_COUNT_TOPLAYER				1
-#define STATE_BIT_COUNT_HASSCATTERING			1
+#define STATE_BIT_COUNT_HASSSS					1
 #define STATE_BIT_COUNT_HASSSSPROFILE			1
 #define STATE_BIT_COUNT_GREYWEIGHT				1
 #define STATE_BIT_COUNT_HAZINESS				1
 #define STATE_BIT_COUNT_EDGECOLOR				1
 #define STATE_BIT_COUNT_THINFILM				1
 #define STATE_BIT_COUNT_SIMPLEVOLUME			1
-#define STATE_BIT_COUNT___UNUSED___				8
+#define STATE_BIT_COUNT_DMFPPLUGGED				1
+#define STATE_BIT_COUNT___UNUSED___				7
 #define STATE_BIT_COUNT_WEIGHT10F				10
 
 // Bit offset of each component store in BSDF state 
@@ -129,13 +130,14 @@ struct FStrataBSDF
 #define STATE_BIT_OFFSET_SHAREDNORMALID	(STATE_BIT_OFFSET_BSDF				+ STATE_BIT_COUNT_BSDF)
 #define STATE_BIT_OFFSET_ANISOTROPY		(STATE_BIT_OFFSET_SHAREDNORMALID	+ STATE_BIT_COUNT_SHAREDNORMALID)
 #define STATE_BIT_OFFSET_TOPLAYER		(STATE_BIT_OFFSET_ANISOTROPY		+ STATE_BIT_COUNT_ANISOTROPY)
-#define STATE_BIT_OFFSET_HASSCATTERING	(STATE_BIT_OFFSET_TOPLAYER			+ STATE_BIT_COUNT_TOPLAYER)
-#define STATE_BIT_OFFSET_HASSSSPROFILE	(STATE_BIT_OFFSET_HASSCATTERING		+ STATE_BIT_COUNT_HASSCATTERING)
+#define STATE_BIT_OFFSET_HASSSS			(STATE_BIT_OFFSET_TOPLAYER			+ STATE_BIT_COUNT_TOPLAYER)
+#define STATE_BIT_OFFSET_HASSSSPROFILE	(STATE_BIT_OFFSET_HASSSS			+ STATE_BIT_COUNT_HASSSS)
 #define STATE_BIT_OFFSET_GREYWEIGHT		(STATE_BIT_OFFSET_HASSSSPROFILE		+ STATE_BIT_COUNT_HASSSSPROFILE)
 #define STATE_BIT_OFFSET_HAZINESS		(STATE_BIT_OFFSET_GREYWEIGHT		+ STATE_BIT_COUNT_GREYWEIGHT)
 #define STATE_BIT_OFFSET_EDGECOLOR		(STATE_BIT_OFFSET_HAZINESS			+ STATE_BIT_COUNT_HAZINESS)
 #define STATE_BIT_OFFSET_THINFILM		(STATE_BIT_OFFSET_EDGECOLOR			+ STATE_BIT_COUNT_EDGECOLOR)
 #define STATE_BIT_OFFSET_SIMPLEVOLUME	(STATE_BIT_OFFSET_THINFILM			+ STATE_BIT_COUNT_THINFILM)
+#define STATE_BIT_OFFSET_DMFPPLUGGED	(STATE_BIT_OFFSET_SIMPLEVOLUME		+ STATE_BIT_COUNT_SIMPLEVOLUME)
 #define STATE_BIT_OFFSET_WEIGHT10F		22
 
 // Accessors for the BSDF type
@@ -157,8 +159,8 @@ struct FStrataBSDF
 #define BSDF_SETISTOPLAYER(X, IsTop)	WRITE_BITS(X.State, STATE_BIT_COUNT_TOPLAYER, STATE_BIT_OFFSET_TOPLAYER, IsTop)
 
 // Indicates if the BSDF has scattering component
-#define BSDF_GETHASSCATTERING(X)		READ_BITS(X.State,  STATE_BIT_COUNT_HASSCATTERING, STATE_BIT_OFFSET_HASSCATTERING)
-#define BSDF_SETHASSCATTERING(X, Scatt)	WRITE_BITS(X.State, STATE_BIT_COUNT_HASSCATTERING, STATE_BIT_OFFSET_HASSCATTERING, Scatt)
+#define BSDF_GETHASSSS(X)			READ_BITS(X.State,  STATE_BIT_COUNT_HASSSS, STATE_BIT_OFFSET_HASSSS)
+#define BSDF_SETHASSSS(X, Scatt)	WRITE_BITS(X.State, STATE_BIT_COUNT_HASSSS, STATE_BIT_OFFSET_HASSSS, Scatt)
 
 // Indicates if the BSDF has a SSS profile
 #define BSDF_GETHASSSSPROFILE(X)		READ_BITS(X.State,  STATE_BIT_COUNT_HASSSSPROFILE, STATE_BIT_OFFSET_HASSSSPROFILE)
@@ -188,6 +190,10 @@ struct FStrataBSDF
 #define BSDF_GETISSIMPLEVOLUME(X)		READ_BITS(X.State,  STATE_BIT_COUNT_SIMPLEVOLUME, STATE_BIT_OFFSET_SIMPLEVOLUME)
 #define BSDF_SETISSIMPLEVOLUME(X, Value)WRITE_BITS(X.State, STATE_BIT_COUNT_SIMPLEVOLUME, STATE_BIT_OFFSET_SIMPLEVOLUME, Value)
 
+// Indicates if the BSDF is part of the top layer
+#define BSDF_GETHASDMFP(X)				READ_BITS(X.State,  STATE_BIT_COUNT_DMFPPLUGGED, STATE_BIT_OFFSET_DMFPPLUGGED)
+#define BSDF_SETHASDMFP(X, Value)		WRITE_BITS(X.State, STATE_BIT_COUNT_DMFPPLUGGED, STATE_BIT_OFFSET_DMFPPLUGGED, Value)
+
 #if STRATA_INLINE_SHADING
 #define BSDF_SETEMISSIVE(X, V)			X.Emissive = V
 #define BSDF_GETEMISSIVE(X)				X.Emissive
@@ -393,11 +399,6 @@ bool StrataHasAreaLightIntegrator(in FStrataBSDF BSDF)
 	return BSDFType == STRATA_BSDF_TYPE_SLAB;
 }
 
-bool StrataHasScattering(in float3 Albedo)
-{
-	return any(Albedo > 0.0f);
-}
-
 bool StrataHasGreyScaleWeight(in float3 Weight)
 {
 	return Weight.x == Weight.y && Weight.y == Weight.z;
@@ -449,8 +450,9 @@ FStrataData GetStrataSlabBSDF(
 	}
 
 	// STRATA_TODO apply SSS as a function of pixel footprint w.r.t. MFP
+	const bool bDMFPPlugged = any(SSSDMFP > 0);
 	const bool bHasSSSProfile = SSSProfileID  > 0.0f;	// When no profile is provide, we use Burley. STRATA_TODO we should clarify as burley is also an option of the profile.
-	const bool bHasSSS = any(SSSDMFP > 0) || bHasSSSProfile;
+	const bool bHasSSS = bDMFPPlugged || bHasSSSProfile;
 
 	StrataData.Layers[0].BSDFCount++;
 	BSDF_SETTYPE			(StrataData.Layers[0].BSDFs[0],		STRATA_BSDF_TYPE_SLAB);
@@ -458,8 +460,9 @@ FStrataData GetStrataSlabBSDF(
 	BSDF_SETEMISSIVE		(StrataData.Layers[0].BSDFs[0],		Emissive);
 	BSDF_SETANISOTROPY		(StrataData.Layers[0].BSDFs[0],		bHasAnisotropy);
 	BSDF_SETISTOPLAYER		(StrataData.Layers[0].BSDFs[0],		0);
-	BSDF_SETHASSCATTERING	(StrataData.Layers[0].BSDFs[0],		bHasSSS ? 1 : 0);
+	BSDF_SETHASSSS			(StrataData.Layers[0].BSDFs[0],		bHasSSS ? 1 : 0);
 	BSDF_SETHASSSSPROFILE	(StrataData.Layers[0].BSDFs[0],		bHasSSSProfile ? 1 : 0);
+	BSDF_SETHASDMFP			(StrataData.Layers[0].BSDFs[0],		bDMFPPlugged ? 1 : 0);
 	BSDF_SETHASHAZINESS		(StrataData.Layers[0].BSDFs[0],		bHasHaziness);
 	BSDF_SETHASEDGECOLOR	(StrataData.Layers[0].BSDFs[0],		bHasEdgeColor);
 	BSDF_SETHASTHINFILM		(StrataData.Layers[0].BSDFs[0],		bHasThinFilm);
@@ -935,6 +938,8 @@ void StrataSanitizeBSDF(inout FStrataBSDF BSDF)
 ///////////////////////////////////////////////////////////////////////////////
 // Other tools
 
+// Those participating media funtion are use by the Slab node. In this case, the DMFP is considered as a simple MFP because DMFP is not directly inversible to medium properties.
+
 FParticipatingMedia StrataSlabCreateParticipatingMedia(float3 DiffuseColor, float3 MeanFreePathCentimeters)
 {
 	const float3 MeanFreePathMeters = MeanFreePathCentimeters * STRATA_CENTIMETER_TO_METER;
@@ -970,6 +975,26 @@ FStrataSlabSingleScattering StrataSlabSingleScattering(FParticipatingMedia PM, f
 	return Result;
 }
 
+void EnableSlabBSDFSimpleVolumetric(inout FStrataBSDF BSDF)
+{
+#if STRATA_INLINE_SHADING
+	// Now we know that this is a simple volume so we can freely set the mean free path
+	SLAB_SSSDMFP(BSDF) = BSDF.TmpDMFP;
+#endif
+
+	// Rescale the BSDF mean free path so that we can drop the thickness parameter and run volumetric material math on a unit slab of medium
+	// For instance, if thickness is larger => we need to reduce the MFP for the medium to appear visually thicker.
+	SLAB_SSSDMFP(BSDF) = SLAB_SSSDMFP(BSDF) * (STRATA_SIMPLEVOLUME_THICKNESS_CM / BSDF_GETTHICKNESSCM(BSDF));
+
+	// Mark the BSDF has not having any SSS through post process.
+	BSDF_SETHASSSS(BSDF, 0);
+	BSDF_SETHASSSSPROFILE(BSDF, 0);
+
+	// Notify that the Slab can use SSS post-process only if at the bottom (layered considered optically thick, e.g. cannot see-through)
+	// When forward rendering is used or if a layer is not a bottom layer, SSS post cannot be used and simple scattering is used.
+	BSDF_SETISSIMPLEVOLUME(BSDF, 1);
+}
+
 
 ///////////////////////////////////////////////////////////////////////////////
 // BSDFs packing/unpacking
@@ -1298,9 +1323,9 @@ void PackStrataOut(
 				EmissiveColor += BSDF_GETEMISSIVE(BSDF) * Throughput;
 
 				// We want simple volumetric only if the DMFP input is pluged in, otherwise we avoid storing the DFMP
-				const bool bIsSimpleVolume = !bBottomLayer && BSDF_GETHASSCATTERING(Strata.Layers[l].BSDFs[i]); 
+				const bool bIsSimpleVolume = !bBottomLayer && BSDF_GETHASDMFP(Strata.Layers[l].BSDFs[i]);
 
-				bStrataSubsurfaceEnable = (bStrataSubsurfaceEnable || (BSDF_GETHASSCATTERING(BSDF) && BSDFCoverageAvg > 0)) && !bIsSimpleVolume;
+				bStrataSubsurfaceEnable = bStrataSubsurfaceEnable || (BSDF_GETHASSSS(BSDF) && BSDFCoverageAvg > 0 && !bIsSimpleVolume);
 
 				float3 F = 0;
 				// STRATA_TODO Fresnel throughput, Use a LUT, this is currently in progress and will be coming soon.
@@ -1329,17 +1354,10 @@ void PackStrataOut(
 
 				if (bIsSimpleVolume)
 				{
-				#if STRATA_INLINE_SHADING
-					// Now we know that this is 
-					SLAB_SSSDMFP(BSDF) = BSDF.TmpDMFP;
-				#endif
-
-					// Rescale the BSDF mean free path so that we can drop the thickness parameter and run volumetric material math on a unit slab of medium
-					// For instance, if thickness is larger => we need to reduce the MFP for the medium to appear visually thicker.
-					SLAB_SSSDMFP(BSDF) = bIsSimpleVolume ? SLAB_SSSDMFP(BSDF) * (STRATA_SIMPLEVOLUME_THICKNESS_CM / BSDF_GETTHICKNESSCM(BSDF)) : SLAB_SSSDMFP(BSDF);
+					EnableSlabBSDFSimpleVolumetric(Strata.Layers[l].BSDFs[i]);
 
 					const float3 UnusedDiffuseColor = 0; // Unused when only trasmittance is needed
-					FParticipatingMedia PM = StrataSlabCreateParticipatingMedia(UnusedDiffuseColor, SLAB_SSSDMFP(BSDF));
+					FParticipatingMedia PM = StrataSlabCreateParticipatingMedia(UnusedDiffuseColor, SLAB_SSSDMFP(Strata.Layers[l].BSDFs[i]));
 					const float DiffuseToVolumeBlend = StrataSlabDiffuseToVolumeBlend(PM);
 					const FStrataSlabSingleScattering Result = StrataSlabSingleScattering(PM, dot(N,V));
 
@@ -1347,13 +1365,7 @@ void PackStrataOut(
 					BSDFThroughput = lerp(OpaqueBSDFThroughput, (1-F) * Result.Transmittance, DiffuseToVolumeBlend);
 
 					// Mark the BSDF has not having any SSS through post process.
-					BSDF_SETHASSCATTERING(Strata.Layers[l].BSDFs[i], 0);
-					BSDF_SETHASSSSPROFILE(Strata.Layers[l].BSDFs[i], 0);
 					bStrataSubsurfaceEnable = false;
-
-					// Notify that the Slab can use SSS post-process only if at the bottom (layered considered optically thick, e.g. cannot see-through)
-					// When forward rendering is used or if a layer is not a bottom layer, SSS post cannot be used and simple scattering is used.
-					BSDF_SETISSIMPLEVOLUME(Strata.Layers[l].BSDFs[i], bIsSimpleVolume ? 1 : 0);
 				}
 
 				LayerCoverage	+= BSDFCoverage;
@@ -1464,7 +1476,7 @@ void PackStrataOut(
 				}
 
 				const bool bIsSimpleVolume = BSDF_GETISSIMPLEVOLUME(BSDF);
-				if (BSDF_GETHASSCATTERING(BSDF) || bIsSimpleVolume)
+				if (BSDF_GETHASSSS(BSDF) || bIsSimpleVolume)
 				{
 					if (BSDF_GETHASSSSPROFILE(BSDF) && !bIsSimpleVolume)
 					{
@@ -1563,11 +1575,12 @@ bool IsStrataFastPathCompatible(in FStrataBSDF BSDF)
 {
 	return
 		!BSDF_GETANISOTROPY(BSDF) &&
-		!BSDF_GETHASSCATTERING(BSDF) &&
+		!BSDF_GETHASSSS(BSDF) &&
 		!BSDF_GETHASSSSPROFILE(BSDF) &&
 		!BSDF_GETHASHAZINESS(BSDF) &&
 		!BSDF_GETHASTHINFILM(BSDF) &&
-		!BSDF_GETHASEDGECOLOR(BSDF);
+		!BSDF_GETHASEDGECOLOR(BSDF) &&
+		!BSDF_GETISSIMPLEVOLUME(BSDF);
 }
 
 void UnpackFastPathSlabBSDF(inout FStrataBSDF OutBSDF, uint2 PackedData01)
@@ -1604,10 +1617,11 @@ FStrataBSDF UnpackFastPathStrataBSDFIn(ByteAddressBuffer StrataBuffer, inout FSt
 	BSDF_SETHASEDGECOLOR(OutBSDF, 0);
 	BSDF_SETANISOTROPY(OutBSDF, 0);
 	BSDF_SETISTOPLAYER(OutBSDF, 1);
-	BSDF_SETHASSCATTERING(OutBSDF, 0);
+	BSDF_SETHASSSS(OutBSDF, 0);
 	BSDF_SETHASSSSPROFILE(OutBSDF, 0);
 	BSDF_SETHASHAZINESS(OutBSDF, 0);
 	BSDF_SETHASTHINFILM(OutBSDF, 0);
+	BSDF_SETISSIMPLEVOLUME(OutBSDF, 0);
 
 	return OutBSDF;
 }
@@ -1650,7 +1664,7 @@ FStrataBSDF UnpackStrataBSDFIn(ByteAddressBuffer StrataBuffer, inout FStrataAddr
 		}
 
 		const bool bIsSimpleVolume = BSDF_GETISSIMPLEVOLUME(OutBSDF);
-		if (BSDF_GETHASSCATTERING(OutBSDF) || bIsSimpleVolume)
+		if (BSDF_GETHASSSS(OutBSDF) || bIsSimpleVolume)
 		{
 			if (BSDF_GETHASSSSPROFILE(OutBSDF) && !bIsSimpleVolume)
 			{
diff --git a/Engine/Shaders/Private/Strata/StrataEvaluation.ush b/Engine/Shaders/Private/Strata/StrataEvaluation.ush
--- a/Engine/Shaders/Private/Strata/StrataEvaluation.ush
+++ b/Engine/Shaders/Private/Strata/StrataEvaluation.ush
@@ -219,7 +219,7 @@ FStrataEvaluateResult StrataEvaluateBSDF(FStrataBSDFContext BSDFContext, FShadow
 	{
 		case STRATA_BSDF_TYPE_SLAB:
 		{
-			const bool bHasSubsurface = BSDF_GETHASSCATTERING(BSDFContext.BSDF) > 0;
+			const bool bHasSubsurface = BSDF_GETHASSSS(BSDFContext.BSDF) > 0;
 			const float3 BaseColor = SLAB_BASECOLOR(BSDFContext.BSDF);
 			const float  Metallic = SLAB_METALLIC(BSDFContext.BSDF);
 			const float3 DiffuseColor = BaseColor - BaseColor * Metallic;
@@ -356,9 +356,9 @@ FStrataEvaluateResult StrataEvaluateBSDF(FStrataBSDFContext BSDFContext, FShadow
 
 				Sample.DiffusePathValue = lerp(Sample.DiffusePathValue,		Result.SingleScattering * Phase * SpecularThroughput,	DiffuseToVolumeBlend);
 				Sample.DiffusePDF		= lerp(Sample.DiffusePDF,			Phase,													DiffuseToVolumeBlend);
-				Sample.Throughput		= lerp(Sample.Throughput,			Result.Transmittance,									DiffuseToVolumeBlend);
+				Sample.Throughput		= lerp(Sample.Throughput,			Result.Transmittance * SpecularThroughput,				DiffuseToVolumeBlend);
 
-				Sample.bSubsurface = false;	// It should already be the case because Enforce in this case because BSDF_GETHASSCATTERING should be false when BSDF_GETISSIMPLEVOLUME is true
+				Sample.bSubsurface = false;	// It should already be the case because Enforce in this case because BSDF_GETHASSSS should be false when BSDF_GETISSIMPLEVOLUME is true
 			}
 			#endif
 
@@ -555,7 +555,7 @@ FStrataEnvLightResult StrataEvaluateForEnvLight(FStrataBSDFContext BSDFContext,
 
 			StrataEnvLightResult.DiffuseColor = BaseColor - BaseColor * Metallic;
 			StrataEnvLightResult.DiffuseNormal = BSDFContext.N;
-			StrataEnvLightResult.bSubsurface = BSDF_GETHASSCATTERING(BSDFContext.BSDF) > 0;
+			StrataEnvLightResult.bSubsurface = BSDF_GETHASSSS(BSDFContext.BSDF) > 0;
 
 			float3 FresnelWeight = 1;
 			BRANCH
diff --git a/Engine/Shaders/Private/Strata/StrataForwardLighting.ush b/Engine/Shaders/Private/Strata/StrataForwardLighting.ush
--- a/Engine/Shaders/Private/Strata/StrataForwardLighting.ush
+++ b/Engine/Shaders/Private/Strata/StrataForwardLighting.ush
@@ -31,11 +31,25 @@ float3 StrataForwardLighting(
 		float3 LayerThroughput = 0.0f;
 		float3 LayerCoverage = 0.0f;
 
+		const bool bTopLayer = l == 0;
+		const bool bBottomLayer = l == (Strata.LayerCount - 1);
+
 		UNROLL
 		for (int i = 0; i < Strata.Layers[l].BSDFCount; ++i)
 		{
 			FStrataBSDF BSDF = Strata.Layers[l].BSDFs[i];
 
+			BRANCH
+			if(BSDF_GETTYPE(BSDF) == STRATA_BSDF_TYPE_SLAB)
+			{
+				// We want simple volumetric only if the DMFP input is pluged in, otherwise we avoid dealing with simple volumetric
+				const bool bIsSimpleVolume = BSDF_GETHASDMFP(BSDF);
+				if (bIsSimpleVolume)
+				{
+					EnableSlabBSDFSimpleVolumetric(BSDF);
+				}
+			}
+
 			// Sanitize BSDF before it is used for forward shading
 			StrataSanitizeBSDF(BSDF);
 
@@ -54,11 +68,7 @@ float3 StrataForwardLighting(
 			FStrataEnvLightResult StrataEnvLight = StrataEvaluateForEnvLight(StrataBSDFContext, bEnableSpecular);
 			if (any(StrataEnvLight.DiffuseColor > 0.0f))
 			{
-				// SkySHDiffuse is pre integrated according to a Lambert diffuse lobe. Thus, we need to rescale to account for the phase function operating over 4PI.
-				// We apply a 2PI factor to recover average luminance over half hemisphere, and then times the phase function to account for the scattering events.
-				const float VolumeOnlyFactor = 1.0f; // STRATA_TODO BSDF_GETTYPE(BSDF) == STRATA_BSDF_TYPE_VOLUME ? 2.0f * PI * IsotropicPhase() : 1.0f;
-
-				Color += GetEffectiveSkySHDiffuse(StrataEnvLight.DiffuseNormal) * StrataEnvLight.DiffuseColor * SkyLightCommonFactor * VolumeOnlyFactor;
+				Color += GetEffectiveSkySHDiffuse(StrataEnvLight.DiffuseNormal) * StrataEnvLight.DiffuseColor * SkyLightCommonFactor;
 			}
 			if (any(StrataEnvLight.SpecularWeight > 0.0f))
 			{
diff --git a/Engine/Shaders/Private/Strata/StrataFurnaceTest.usf b/Engine/Shaders/Private/Strata/StrataFurnaceTest.usf
--- a/Engine/Shaders/Private/Strata/StrataFurnaceTest.usf
+++ b/Engine/Shaders/Private/Strata/StrataFurnaceTest.usf
@@ -108,7 +108,7 @@ void MainPS(
 			BSDF_SETANISOTROPY(BSDF, 0);
 			BSDF_SETISTOPLAYER(BSDF, 1);
 			BSDF_SETHASEDGECOLOR(BSDF, 0);
-			BSDF_SETHASSCATTERING(BSDF, 0);
+			BSDF_SETHASSSS(BSDF, 0);
 			BSDF_SETHASSSSPROFILE(BSDF, 0);
 			BSDF_SETHASHAZINESS(BSDF, 0);
 			BSDF_SETHASTHINFILM(BSDF, 0);
diff --git a/Engine/Shaders/Private/Strata/StrataMaterialClassification.usf b/Engine/Shaders/Private/Strata/StrataMaterialClassification.usf
--- a/Engine/Shaders/Private/Strata/StrataMaterialClassification.usf
+++ b/Engine/Shaders/Private/Strata/StrataMaterialClassification.usf
@@ -65,7 +65,7 @@ void ComputePostProcessBSDFData(
 		Classification.bIsSimple = Classification.bIsSimple && IsStrataFastPathCompatible(BSDF);
 
 		// Select the latest encoutered SSS BSDF
-		if (BSDF_GETHASSCATTERING(BSDF) && BSDFWeightAvg > 0)
+		if (BSDF_GETHASSSS(BSDF) && BSDFWeightAvg > 0)
 		{
 			Classification.bHasSubsurface = true;
 
diff --git a/Engine/Shaders/Private/Strata/StrataVisualize.usf b/Engine/Shaders/Private/Strata/StrataVisualize.usf
--- a/Engine/Shaders/Private/Strata/StrataVisualize.usf
+++ b/Engine/Shaders/Private/Strata/StrataVisualize.usf
@@ -185,7 +185,7 @@ void VisualizeMaterialPS(
 #define PRINTBSDFNORMALID()	PRINTBSDFCHAR(_SPC_); PRINTBSDFCHAR(_MINUS_); PRINTBSDFCHAR(_N_); PRINTBSDFCHAR(_R_); PRINTBSDFCHAR(_M_); PRINTBSDFCHAR(_I_); PRINTBSDFCHAR(_D_); PRINTBSDFCHAR(_SPC_); PrintFloatNoFraction(PixelPos, OutColor.xyz, FontColor, TopLeft, float(BSDF_GETNORMALID(BSDF)), 1); TopLeft.x+=8*1;
 #define PRINTBSDFANISO()	PRINTBSDFCHAR(_SPC_); PRINTBSDFCHAR(_MINUS_); PRINTBSDFCHAR(_A_); PRINTBSDFCHAR(_N_); PRINTBSDFCHAR(_I_); PRINTBSDFCHAR(_S_); PRINTBSDFCHAR(_O_); PRINTBSDFCHAR(_SPC_); PrintFloatNoFraction(PixelPos, OutColor.xyz, FontColor, TopLeft, float(BSDF_GETANISOTROPY(BSDF)), 1); TopLeft.x+=8*1;
 #define PRINTBSDFTOPLAY()	PRINTBSDFCHAR(_SPC_); PRINTBSDFCHAR(_MINUS_); PRINTBSDFCHAR(_T_); PRINTBSDFCHAR(_O_); PRINTBSDFCHAR(_P_); PRINTBSDFCHAR(_L_); PRINTBSDFCHAR(_A_); PRINTBSDFCHAR(_Y_); PRINTBSDFCHAR(_SPC_); PrintFloatNoFraction(PixelPos, OutColor.xyz, FontColor, TopLeft, float(BSDF_GETISTOPLAYER(BSDF)), 1); TopLeft.x+=8*1;
-#define PRINTBSDFHASSCAT()	PRINTBSDFCHAR(_SPC_); PRINTBSDFCHAR(_MINUS_); PRINTBSDFCHAR(_H_); PRINTBSDFCHAR(_A_); PRINTBSDFCHAR(_S_); PRINTBSDFCHAR(_S_); PRINTBSDFCHAR(_C_); PRINTBSDFCHAR(_A_); PRINTBSDFCHAR(_T_); PRINTBSDFCHAR(_SPC_); PrintFloatNoFraction(PixelPos, OutColor.xyz, FontColor, TopLeft, float(BSDF_GETHASSCATTERING(BSDF)), 1); TopLeft.x+=8*1;
+#define PRINTBSDFHASSCAT()	PRINTBSDFCHAR(_SPC_); PRINTBSDFCHAR(_MINUS_); PRINTBSDFCHAR(_H_); PRINTBSDFCHAR(_A_); PRINTBSDFCHAR(_S_); PRINTBSDFCHAR(_S_); PRINTBSDFCHAR(_C_); PRINTBSDFCHAR(_A_); PRINTBSDFCHAR(_T_); PRINTBSDFCHAR(_SPC_); PrintFloatNoFraction(PixelPos, OutColor.xyz, FontColor, TopLeft, float(BSDF_GETHASSSS(BSDF)), 1); TopLeft.x+=8*1;
 #define PRINTBSDFSIMPSCAT()	PRINTBSDFCHAR(_SPC_); PRINTBSDFCHAR(_MINUS_); PRINTBSDFCHAR(_S_); PRINTBSDFCHAR(_I_); PRINTBSDFCHAR(_M_); PRINTBSDFCHAR(_P_); PRINTBSDFCHAR(_S_); PRINTBSDFCHAR(_C_); PRINTBSDFCHAR(_A_); PRINTBSDFCHAR(_T_); PRINTBSDFCHAR(_SPC_); PrintFloatNoFraction(PixelPos, OutColor.xyz, FontColor, TopLeft, float(BSDF_GETISSIMPLEVOLUME(BSDF)), 1); TopLeft.x+=8*1;
 #define PRINTFLOAT3(X) PrintFloat(PixelPos, OutColor.xyz, PARAMFontColor, TopLeft, X.x); TopLeft.x+=8*7; PRINTBSDFCHAR(_SPC_); PrintFloat(PixelPos, OutColor.xyz, PARAMFontColor, TopLeft, X.y); TopLeft.x+=8*7; PRINTBSDFCHAR(_SPC_); PrintFloat(PixelPos, OutColor.xyz, PARAMFontColor, TopLeft, X.z);
 #define PRINTFLOAT1(Value) PrintFloat(PixelPos, OutColor.xyz, PARAMFontColor, TopLeft, Value); TopLeft.x+=8*7;
@@ -347,7 +347,7 @@ void VisualizeMaterialPS(
 				HORIZONTALTAB; PRINTFLOAT1(SLAB_HAZINESS(BSDF));
 				NEXTLINE;
 
-				if (BSDF_GETHASSCATTERING(BSDF) || BSDF_GETISSIMPLEVOLUME(BSDF))
+				if (BSDF_GETHASSSS(BSDF) || BSDF_GETISSIMPLEVOLUME(BSDF))
 				{
 					if (BSDF_GETHASSSSPROFILE(BSDF) && BSDF_GETISSIMPLEVOLUME(BSDF)==0)
 					{
@@ -386,7 +386,7 @@ void VisualizeMaterialPS(
 					NEXTLINE;
 				}
 
-				if (BSDF_GETHASSCATTERING(BSDF))
+				if (BSDF_GETHASSSS(BSDF))
 				{
 					if (BSDF_GETHASSSSPROFILE(BSDF))
 					{
diff --git a/Engine/Source/Runtime/Engine/Classes/Materials/MaterialExpressionStrata.h b/Engine/Source/Runtime/Engine/Classes/Materials/MaterialExpressionStrata.h
--- a/Engine/Source/Runtime/Engine/Classes/Materials/MaterialExpressionStrata.h
+++ b/Engine/Source/Runtime/Engine/Classes/Materials/MaterialExpressionStrata.h
@@ -135,8 +135,10 @@ class UMaterialExpressionStrataSlabBSDF : public UMaterialExpressionStrataBSDF
 	virtual FName GetInputName(int32 InputIndex) const override;
 
 	bool HasEdgeColor() const;
-	bool HasScattering() const;
 	bool HasThinFilm() const;
+	bool HasSSS() const;
+	bool HasSSSProfile() const;
+	bool HasDMFPPluggedIn() const;
 #endif
 	//~ End UMaterialExpression Interface
 };
diff --git a/Engine/Source/Runtime/Engine/Private/Materials/Material.cpp b/Engine/Source/Runtime/Engine/Private/Materials/Material.cpp
--- a/Engine/Source/Runtime/Engine/Private/Materials/Material.cpp
+++ b/Engine/Source/Runtime/Engine/Private/Materials/Material.cpp
@@ -4727,15 +4727,30 @@ void UMaterial::RebuildShadingModelField()
 		{
 			if (StrataMaterialInfo.CountShadingModels() == 2 && StrataMaterialInfo.HasShadingModel(EStrataShadingModel::SSM_DefaultLit) && StrataMaterialInfo.HasShadingModel(EStrataShadingModel::SSM_SubsurfaceLit))
 			{
-				StrataMaterialInfo = FStrataMaterialInfo();
-				StrataMaterialInfo.AddShadingModel(SSM_SubsurfaceLit);
+				if (BlendMode == EBlendMode::BLEND_Opaque || BlendMode == EBlendMode::BLEND_Masked) 
+				{
+					// We only consider SSS subsurface post process for opaque materials)
+					StrataMaterialInfo = FStrataMaterialInfo();
+					StrataMaterialInfo.AddShadingModel(SSM_SubsurfaceLit);
+				}
+				else
+				{
+					// For transparent, we will fall back to use DefaultLit with simple volumetric
+					bSanitizeMaterial = true;
+				}
 			}
 			else
 			{
+				// Clear the material to default Lit
 				bSanitizeMaterial = true;
 				UE_LOG(LogMaterial, Error, TEXT("%s: Material has more than a single material represented."), *GetName());
 			}
 		}
+		if (StrataMaterialInfo.HasOnlyShadingModel(SSM_SubsurfaceLit) && (BlendMode != EBlendMode::BLEND_Opaque && BlendMode != EBlendMode::BLEND_Masked))
+		{
+			// For transparent, we will fall back to use DefaultLit with simple volumetric
+			bSanitizeMaterial = true;
+		}
 		if (bSanitizeMaterial)
 		{
 			StrataMaterialInfo = FStrataMaterialInfo();
diff --git a/Engine/Source/Runtime/Engine/Private/Materials/MaterialExpressions.cpp b/Engine/Source/Runtime/Engine/Private/Materials/MaterialExpressions.cpp
--- a/Engine/Source/Runtime/Engine/Private/Materials/MaterialExpressions.cpp
+++ b/Engine/Source/Runtime/Engine/Private/Materials/MaterialExpressions.cpp
@@ -19753,10 +19753,11 @@ int32 UMaterialExpressionStrataSlabBSDF::Compile(class FMaterialCompiler* Compil
 
 	const bool bHasEdgeColor = HasEdgeColor();
 	const bool bHasThinFilm = HasThinFilm();
+	const bool bHasDMFPPluggedIn = HasDMFPPluggedIn();
 
 	int32 SSSProfileCodeChunk = INDEX_NONE;
-	const bool bHasScattering = HasScattering();
-	if (bHasScattering)
+	const bool bHasSSS = HasSSS();
+	if (bHasSSS)
 	{
 		FName NameSubsurfaceProfile(FString(TEXT("__SubsurfaceProfile")));
 		SSSProfileCodeChunk = Compiler->ForceCast(Compiler->ScalarParameter(NameSubsurfaceProfile, 1.0f), MCT_Float1);
@@ -19779,7 +19780,7 @@ int32 UMaterialExpressionStrataSlabBSDF::Compile(class FMaterialCompiler* Compil
 		NormalCodeChunk,
 		TangentCodeChunk,
 		SharedNormalIndex);
-	StrataCompilationInfoCreateSingleBSDFMaterial(Compiler, OutputCodeChunk, SharedNormalIndex, STRATA_BSDF_TYPE_SLAB, bHasEdgeColor, bHasScattering, bHasThinFilm);
+	StrataCompilationInfoCreateSingleBSDFMaterial(Compiler, OutputCodeChunk, SharedNormalIndex, STRATA_BSDF_TYPE_SLAB, bHasSSS, bHasDMFPPluggedIn, bHasEdgeColor, bHasThinFilm);
 
 	return OutputCodeChunk;
 }
@@ -19914,8 +19915,9 @@ bool UMaterialExpressionStrataSlabBSDF::IsResultStrataMaterial(int32 OutputIndex
 
 void UMaterialExpressionStrataSlabBSDF::GatherStrataMaterialInfo(FStrataMaterialInfo& StrataMaterialInfo, int32 OutputIndex)
 {
-	if (HasScattering())
+	if (HasSSS())
 	{
+		// We still do not know if this is going to be a real SSS node because it is only possible for BSDF at the bottom of the stack. Nevertheless, we take the worst case into account.
 		StrataMaterialInfo.AddShadingModel(SSM_SubsurfaceLit);
 		if (SubsurfaceProfile)
 		{
@@ -19928,11 +19930,21 @@ void UMaterialExpressionStrataSlabBSDF::GatherStrataMaterialInfo(FStrataMaterial
 	}
 }
 
-bool UMaterialExpressionStrataSlabBSDF::HasScattering() const
+bool UMaterialExpressionStrataSlabBSDF::HasSSS() const
 {
 	return SubsurfaceProfile != nullptr || SSSDMFP.IsConnected();
 }
 
+bool UMaterialExpressionStrataSlabBSDF::HasSSSProfile() const
+{
+	return SubsurfaceProfile != nullptr;
+}
+
+bool UMaterialExpressionStrataSlabBSDF::HasDMFPPluggedIn() const
+{
+	return SSSDMFP.IsConnected();
+}
+
 bool UMaterialExpressionStrataSlabBSDF::HasEdgeColor() const
 {
 	return EdgeColor.IsConnected();
diff --git a/Engine/Source/Runtime/Engine/Private/Materials/StrataMaterial.cpp b/Engine/Source/Runtime/Engine/Private/Materials/StrataMaterial.cpp
--- a/Engine/Source/Runtime/Engine/Private/Materials/StrataMaterial.cpp
+++ b/Engine/Source/Runtime/Engine/Private/Materials/StrataMaterial.cpp
@@ -57,14 +57,15 @@ uint8 StrataCompilationInfoCreateSharedNormal(FMaterialCompiler* Compiler, int32
 	return Compiler->StrataCompilationInfoRegisterSharedNormalIndex(NormalCodeChunk, TangentCodeChunk);
 }
 
-void StrataCompilationInfoCreateSingleBSDFMaterial(FMaterialCompiler* Compiler, int32 CodeChunk, uint8 SharedNormalIndex, uint8 BSDFType, bool bHasEdgeColor, bool bHasScattering, bool bHasThinFilm)
+void StrataCompilationInfoCreateSingleBSDFMaterial(FMaterialCompiler* Compiler, int32 CodeChunk, uint8 SharedNormalIndex, uint8 BSDFType, bool bHasSSS, bool bHasDMFPPluggedIn, bool bHasEdgeColor, bool bHasThinFilm)
 {
 	FStrataMaterialCompilationInfo StrataInfo;
 	StrataInfo.LayerCount = 1;
 	StrataInfo.Layers[0].BSDFCount = 1;
 	StrataInfo.Layers[0].BSDFs[0].Type = BSDFType;
 	StrataInfo.Layers[0].BSDFs[0].SharedNormalIndex = SharedNormalIndex;
-	StrataInfo.Layers[0].BSDFs[0].bHasScattering = bHasScattering;
+	StrataInfo.Layers[0].BSDFs[0].bHasSSS = bHasSSS;
+	StrataInfo.Layers[0].BSDFs[0].bHasDMFPPluggedIn = bHasDMFPPluggedIn;
 	StrataInfo.Layers[0].BSDFs[0].bHasEdgeColor = bHasEdgeColor;
 	StrataInfo.Layers[0].BSDFs[0].bHasThinFilm = bHasThinFilm;
 	UpdateTotalBSDFCount(Compiler, StrataInfo);
@@ -259,7 +260,7 @@ FStrataMaterialAnalysisResult StrataCompilationInfoMaterialAnalysis(FMaterialCom
 	Result.RequestedByteCount += UintByteSize;
 	// Shared normals between BSDFs
 	Result.RequestedByteCount += Compiler->StrataCompilationInfoGetSharedNormalCount() * STRATA_PACKED_NORMAL_STRIDE_BYTES;
-
+	
 	// 2. The list of BSDFs
 
 	// We process layers from top to bottom to cull the bottom ones in case we run out of pixel bytes
@@ -267,6 +268,9 @@ FStrataMaterialAnalysisResult StrataCompilationInfoMaterialAnalysis(FMaterialCom
 	{
 		const FStrataMaterialCompilationInfo::FLayer& Layer = Material.Layers[LayerIt];
 
+		const bool bTopLayer = LayerIt == 0;
+		const bool bBottomLayer = LayerIt == (Material.LayerCount - 1);
+
 		for (uint32 BSDFIt = 0; BSDFIt < Layer.BSDFCount; BSDFIt++)
 		{
 			const FStrataMaterialCompilationInfo::FBSDF& BSDF = Layer.BSDFs[BSDFIt];
@@ -281,6 +285,10 @@ FStrataMaterialAnalysisResult StrataCompilationInfoMaterialAnalysis(FMaterialCom
 				Result.RequestedByteCount += UintByteSize;
 			}
 
+			// Compute values closer to the reality for HasSSS and IsSimpleVolume, now that we know that we know the topology of the material.
+			const bool bIsSimpleVolume = !bBottomLayer && BSDF.bHasDMFPPluggedIn;
+			const bool bHasSSS = bBottomLayer && BSDF.bHasSSS && !bIsSimpleVolume;
+
 			switch (BSDF.Type)
 			{
 			case STRATA_BSDF_TYPE_SLAB:
@@ -291,7 +299,7 @@ FStrataMaterialAnalysisResult StrataCompilationInfoMaterialAnalysis(FMaterialCom
 				{
 					Result.RequestedByteCount += UintByteSize;
 				}
-				if (BSDF.bHasScattering)
+				if (bHasSSS || bIsSimpleVolume)
 				{
 					Result.RequestedByteCount += UintByteSize;
 				}
