Updated diffuse to volumetric lighting transition.
+ updated asset

diff --git a/Engine/Shaders/Private/Strata/Strata.ush b/Engine/Shaders/Private/Strata/Strata.ush
--- a/Engine/Shaders/Private/Strata/Strata.ush
+++ b/Engine/Shaders/Private/Strata/Strata.ush
@@ -948,11 +948,10 @@ FParticipatingMedia StrataSlabCreateParticipatingMedia(float3 DiffuseColor, floa
 
 float StrataSlabDiffuseToVolumeBlend(FParticipatingMedia PM)
 {
-	// We soft lerp from opaque diffuse to volue representation when transmittance to background (or layers below) reaches 0.329.
-	// mean free path = 0.9 * 0.01 = 0.009, extinction = 111.1, transmittance = exp(extinction * 0.01) = 0.329
-	// mean free path = 1.0 * 0.01 = 0.01, extinction = 100.0, transmittance = exp(extinction * 0.01) = 0.367
+	// End   lerp when transmittance = 0.05, when extinction over 1cm is = -ln(0.05)/0.01 = 299.57 => MFP = 0.003338
+	// Start lerp when transmittance = 0.01, when extinction over 1cm is = -ln(0.01)/0.01 = 230.25 => MFP = 0.004343
 	const float MaxMeanFreePathMeters = max(PM.MeanFreePath.x, max(PM.MeanFreePath.y, PM.MeanFreePath.z));
-	return saturate((MaxMeanFreePathMeters / STRATA_SIMPLEVOLUME_THICKNESS_M) * 10.0f - 9.0f);
+	return saturate((MaxMeanFreePathMeters - 0.003338) / (0.004343 - 0.003338));
 }
 
 struct FStrataSlabSingleScattering
