Fix missing propagation of PrimitiveId from VS to PS for GPUSkinVF & LanscapeVF (when using instance culling).

diff --git a/Engine/Shaders/Private/GpuSkinVertexFactory.ush b/Engine/Shaders/Private/GpuSkinVertexFactory.ush
--- a/Engine/Shaders/Private/GpuSkinVertexFactory.ush
+++ b/Engine/Shaders/Private/GpuSkinVertexFactory.ush
@@ -189,6 +189,10 @@ struct FVertexFactoryInterpolantsVSToPS
 #if INSTANCED_STEREO
 	nointerpolation uint EyeIndex : PACKED_EYE_INDEX;
 #endif
+
+#if VF_USE_PRIMITIVE_SCENE_DATA
+	nointerpolation uint PrimitiveId : PRIMITIVE_ID;
+#endif
 };
 
 struct FVertexFactoryInterpolantsVSToDS
@@ -244,6 +248,11 @@ FMaterialPixelParameters GetMaterialPixelParameters(FVertexFactoryInterpolantsVS
 	Result.VertexColor = 0;
 #endif
 	Result.TwoSidedSign = 1;
+
+#if VF_USE_PRIMITIVE_SCENE_DATA
+	Result.PrimitiveId = Interpolants.PrimitiveId;
+#endif
+
 	return Result;
 }
 
@@ -878,6 +887,10 @@ FVertexFactoryInterpolantsVSToPS VertexFactoryGetInterpolantsVSToPS(FVertexFacto
 	Interpolants.EyeIndex = 0;
 #endif
 
+#if VF_USE_PRIMITIVE_SCENE_DATA
+	Interpolants.PrimitiveId = Intermediates.SceneData.PrimitiveId;
+#endif
+
 	return Interpolants;
 }
 
@@ -1004,7 +1017,11 @@ float4 VertexFactoryGetTranslatedPrimitiveVolumeBounds(FVertexFactoryInterpolant
 
 uint VertexFactoryGetPrimitiveId(FVertexFactoryInterpolantsVSToPS Interpolants)
 {
+#if VF_USE_PRIMITIVE_SCENE_DATA
+	return Interpolants.PrimitiveId;
+#else // VF_USE_PRIMITIVE_SCENE_DATA
 	return 0;
+#endif // VF_USE_PRIMITIVE_SCENE_DATA
 }
 
 #include "VertexFactoryDefaultInterface.ush"
diff --git a/Engine/Shaders/Private/LandscapeVertexFactory.ush b/Engine/Shaders/Private/LandscapeVertexFactory.ush
--- a/Engine/Shaders/Private/LandscapeVertexFactory.ush
+++ b/Engine/Shaders/Private/LandscapeVertexFactory.ush
@@ -255,6 +255,10 @@ struct FVertexFactoryInterpolantsVSToPS
 #if INSTANCED_STEREO
 	nointerpolation uint EyeIndex : PACKED_EYE_INDEX;
 #endif
+
+#if VF_USE_PRIMITIVE_SCENE_DATA
+	nointerpolation uint PrimitiveId : PRIMITIVE_ID;
+#endif
 };
 
 struct FLandscapeTexCoords
@@ -564,6 +568,11 @@ FMaterialPixelParameters GetMaterialPixelParameters(FVertexFactoryInterpolantsVS
 #endif
 
 	Result.TwoSidedSign = 1;
+
+#if VF_USE_PRIMITIVE_SCENE_DATA
+	Result.PrimitiveId = Interpolants.PrimitiveId;
+#endif
+
 	return Result;
 }
 
@@ -851,6 +860,10 @@ FVertexFactoryInterpolantsVSToPS VertexFactoryGetInterpolantsVSToPS(FVertexFacto
 	Interpolants.EyeIndex = 0;
 #endif
 	
+#if VF_USE_PRIMITIVE_SCENE_DATA
+	Interpolants.PrimitiveId = Intermediates.SceneData.PrimitiveId;
+#endif
+
 	return Interpolants;
 }
 
@@ -937,7 +950,11 @@ float4 VertexFactoryGetTranslatedPrimitiveVolumeBounds(FVertexFactoryInterpolant
 
 uint VertexFactoryGetPrimitiveId(FVertexFactoryInterpolantsVSToPS Interpolants)
 {
+#if VF_USE_PRIMITIVE_SCENE_DATA
+	return Interpolants.PrimitiveId;
+#else // VF_USE_PRIMITIVE_SCENE_DATA
 	return 0;
+#endif // VF_USE_PRIMITIVE_SCENE_DATA
 }
 
 #include "VertexFactoryDefaultInterface.ush"
