Workaround for primitives with zero bounds being culled.

diff --git a/Engine/Shaders/Private/InstanceCulling/CullInstances.usf b/Engine/Shaders/Private/InstanceCulling/CullInstances.usf
--- a/Engine/Shaders/Private/InstanceCulling/CullInstances.usf
+++ b/Engine/Shaders/Private/InstanceCulling/CullInstances.usf
@@ -28,7 +28,10 @@ void CullInstancesCs(uint InstanceId : SV_DispatchThreadID)
 	uint WordMask = 1U << (InstanceId % 32U);
 	uint InstanceWordOffset = InstanceId / 32U;
 
-	bool bIsValid = InstanceData.PrimitiveId != 0xFFFFFFFFu;
+	// TODO: The test for dot(InstanceData.LocalBoundsExtent, InstanceData.LocalBoundsExtent) > 0.0f is just a workaround since the FDynamicMeshBuilder::GetMesh
+	//       seems to just set empty bounds (and FLineBatcherSceneProxy pretends everything is at the origin). In the future these should compute reasonable bounds and 
+	//       this should be removed.
+	bool bIsValid = InstanceData.PrimitiveId != 0xFFFFFFFFu && dot(InstanceData.LocalBoundsExtent, InstanceData.LocalBoundsExtent) > 0.0f;
 
 	for (uint ViewId = 0; ViewId < NumViews; ++ViewId)
 	{
