Fix Strata microfacet multiple scattering for dieletric

diff --git a/Engine/Shaders/Private/Strata/StrataMultipleScattering.ush b/Engine/Shaders/Private/Strata/StrataMultipleScattering.ush
--- a/Engine/Shaders/Private/Strata/StrataMultipleScattering.ush
+++ b/Engine/Shaders/Private/Strata/StrataMultipleScattering.ush
@@ -45,7 +45,7 @@ FStrataLUTResult StrataSampleGGXLUT(Texture2D<float4> LUT2D, Texture3D<float2> L
 		const float3 UVW = float3(saturate(NoV), saturate(Roughness), saturate(F0.y)) * LUTScaleBias.x + LUTScaleBias.yyy;
 		const float2 S = LUT3D.SampleLevel(LUTSampler, UVW, 0);
 		Out.E_R = S.x;
-		Out.E_R = S.y;
+		Out.E_T = S.y;
 	}
 
 	return Out;
