Fixed StrataVisualize due to multiple recent changes not taking multiple passes path into account.

diff --git a/Engine/Shaders/Private/Strata/StrataVisualize.usf b/Engine/Shaders/Private/Strata/StrataVisualize.usf
--- a/Engine/Shaders/Private/Strata/StrataVisualize.usf
+++ b/Engine/Shaders/Private/Strata/StrataVisualize.usf
@@ -389,17 +389,27 @@ void VisualizeMaterialPS(
 #else
 				NEXTLINE;
 				NEXTLINE;
-				NEXTLINE;
+				if (BSDF_GETHASEDGECOLOR(BSDF))
+				{
+					NEXTLINE;
+				}
 				NEXTLINE;
 				NEXTLINE;
 				if (BSDF_GETANISOTROPY(BSDF))
+				{
+					NEXTLINE;
+					NEXTLINE;
+				}
+				else
 				{
 					NEXTLINE;
 				}
 
-				if (BSDF_GETHASSSS(BSDF))
+				NEXTLINE;
+
+				if (BSDF_GETHASSSS(BSDF) || BSDF_GETISSIMPLEVOLUME(BSDF))
 				{
-					if (BSDF_GETHASSSSPROFILE(BSDF))
+					if (BSDF_GETHASSSSPROFILE(BSDF) && BSDF_GETISSIMPLEVOLUME(BSDF) == 0)
 					{
 						NEXTLINE;
 						NEXTLINE;
@@ -409,6 +419,12 @@ void VisualizeMaterialPS(
 						NEXTLINE;
 					}
 				}
+
+				if (BSDF_GETHASTHINFILM(BSDF))
+				{
+					NEXTLINE;
+					NEXTLINE;
+				}
 #endif
 				break;
 
