Strata: fixed artefact in forward where Frenel would cause dark spot when the light is on the otherside ofthe surface.

diff --git a/Engine/Shaders/Private/Strata/Strata.ush b/Engine/Shaders/Private/Strata/Strata.ush
--- a/Engine/Shaders/Private/Strata/Strata.ush
+++ b/Engine/Shaders/Private/Strata/Strata.ush
@@ -1328,33 +1328,33 @@ void PackStrataOut(
 
 				bStrataSubsurfaceEnable = bStrataSubsurfaceEnable || (BSDF_GETHASSSS(BSDF) && BSDFCoverageAvg > 0 && !bIsSimpleVolume);
 
-				float3 F = 0;
-				// STRATA_TODO Fresnel throughput, Use a LUT, this is currently in progress and will be coming soon.
-				// In the meantime, we use a simple approach to feed Fresnel fnuctions and compute an approximation of the general transmition/througput.
-				const float VoH = ContextNoLight.NoV;
-				const float NoL = ContextNoLight.NoV;
-				const float3 BaseColor = SLAB_BASECOLOR(BSDF);
-				const float  Metallic = SLAB_METALLIC(BSDF);
-				const float3 F0 = ComputeF0(SLAB_SPECULAR(BSDF), BaseColor, Metallic);
-				const float3 F90 = ComputeF90(F0, SLAB_EDGECOLOR(BSDF), Metallic);
-				const bool bHasThinFilm = BSDF_GETHASTHINFILM(BSDF);
-				BRANCH
-				if (bHasThinFilm)
-				{
-					float ThinFilmIOR = SLAB_THINFILM_IOR(BSDF);
-					float ThinFilmThickness = SLAB_THINFILM_THICKNESS(BSDF);
-					F = F_ThinFilm(ContextNoLight.NoV, NoL, VoH, F0, F90, ThinFilmIOR, ThinFilmThickness);
-				}
-				else
-				{
-					F = F_Schlick(F0, F90, VoH);
-				}
-
 				// By default, without coverage, the diffuse BSDF is completely opaque
 				float3 BSDFThroughput = OpaqueBSDFThroughput;
 
 				if (bIsSimpleVolume)
 				{
+					float3 ApproxFresnelForTransmittance = 0;
+					// STRATA_TODO Fresnel throughput, Use a LUT, this is currently in progress and will be coming soon.
+					// In the meantime, we use a simple approach to feed Fresnel function and compute an approximation of the general transmition/througput.
+					const float VoH = ContextNoLight.NoV;
+					const float NoL = ContextNoLight.NoV;
+					const float3 BaseColor = SLAB_BASECOLOR(BSDF);
+					const float  Metallic = SLAB_METALLIC(BSDF);
+					const float3 F0 = ComputeF0(SLAB_SPECULAR(BSDF), BaseColor, Metallic);
+					const float3 F90 = ComputeF90(F0, SLAB_EDGECOLOR(BSDF), Metallic);
+					const bool bHasThinFilm = BSDF_GETHASTHINFILM(BSDF);
+					BRANCH
+					if (bHasThinFilm)
+					{
+						float ThinFilmIOR = SLAB_THINFILM_IOR(BSDF);
+						float ThinFilmThickness = SLAB_THINFILM_THICKNESS(BSDF);
+						ApproxFresnelForTransmittance = F_ThinFilm(ContextNoLight.NoV, NoL, VoH, F0, F90, ThinFilmIOR, ThinFilmThickness);
+					}
+					else
+					{
+						ApproxFresnelForTransmittance = F_Schlick(F0, F90, VoH);
+					}
+
 					EnableSlabBSDFSimpleVolumetric(Strata.Layers[l].BSDFs[i]);
 
 					const float3 UnusedDiffuseColor = 0; // Unused when only trasmittance is needed
@@ -1363,7 +1363,7 @@ void PackStrataOut(
 					const FStrataSlabSingleScattering Result = StrataSlabSingleScattering(PM, dot(N,V));
 
 					// Update the see-through throughput accoding to the medium transmittance
-					BSDFThroughput = lerp(OpaqueBSDFThroughput, (1-F) * Result.Transmittance, DiffuseToVolumeBlend);
+					BSDFThroughput = lerp(OpaqueBSDFThroughput, (1- ApproxFresnelForTransmittance) * Result.Transmittance, DiffuseToVolumeBlend);
 
 					// Mark the BSDF has not having any SSS through post process.
 					bStrataSubsurfaceEnable = false;
@@ -1376,11 +1376,14 @@ void PackStrataOut(
 
 			case STRATA_BSDF_TYPE_SHEEN:
 			{
+				// STRATA_TODO Fresnel throughput, Use a LUT, this is currently in progress and will be coming soon.
+				// In the meantime, we use a simple approach to feed Fresnel function and compute an approximation of the general transmition/througput.
+				const float NoL = ContextNoLight.NoV;
 				const float3 F0  = SHEEN_BASECOLOR(BSDF);
-				const float3 F = F_Schlick(F0, ContextNoLight.VoH);
+				const float3 ApproxFresnelForTransmittance = F_Schlick(F0, NoL);
 
 				LayerCoverage	+= BSDFCoverage;
-				LayerThroughput += BSDFCoverage * (1.0f - F); // STRATA_TODO use 1-FGD instead of simple 1-F, event if we ignore ray scattering/spreading
+				LayerThroughput += BSDFCoverage * (1.0f - ApproxFresnelForTransmittance);
 				break;
 			}
 
diff --git a/Engine/Shaders/Private/Strata/StrataEvaluation.ush b/Engine/Shaders/Private/Strata/StrataEvaluation.ush
--- a/Engine/Shaders/Private/Strata/StrataEvaluation.ush
+++ b/Engine/Shaders/Private/Strata/StrataEvaluation.ush
@@ -351,7 +351,12 @@ FStrataEvaluateResult StrataEvaluateBSDF(FStrataBSDFContext BSDFContext, FShadow
 				const float DiffuseToVolumeBlend = StrataSlabDiffuseToVolumeBlend(PM);
 				const FStrataSlabSingleScattering Result = StrataSlabSingleScattering(PM, BSDFContext.Context.NoV);
 
-				const float3 SpecularThroughput = 1.0f - F; // This is already taken into account by the Chan model above. But for volumetric we need to apply it.
+				// STRATA_TODO Fresnel throughput, Use a LUT, this is currently in progress and will be coming soon.
+				// In the meantime, we use a simple approach to feed Fresnel function and compute an approximation of the general transmition/througput.
+				// We also ignored ThinFilm interaction as another simplification for perf.
+				const float ApproxFresnelForTransmittance = F_Schlick(F0, F90, BSDFContext.Context.NoV);
+
+				const float3 SpecularThroughput = 1.0f - ApproxFresnelForTransmittance; // This is already taken into account by the Chan model above. But for volumetric we need to apply it.
 				const float Phase = IsotropicPhase();
 
 				Sample.DiffusePathValue = lerp(Sample.DiffusePathValue,		Result.SingleScattering * Phase * SpecularThroughput,	DiffuseToVolumeBlend);
@@ -379,9 +384,13 @@ FStrataEvaluateResult StrataEvaluateBSDF(FStrataBSDFContext BSDFContext, FShadow
 			float Vis = Vis_Cloth(BSDFContext.Context.NoV, BSDFContext.SatNoL);
 			float3 F  = F_Schlick(F0, BSDFContext.Context.VoH);
 
+			// STRATA_TODO Fresnel throughput, Use a LUT, this is currently in progress and will be coming soon.
+			// In the meantime, we use a simple approach to feed Fresnel function and compute an approximation of the general transmition/througput.
+			const float ApproxFresnelForTransmittance = F_Schlick(F0, BSDFContext.Context.NoV);
+
 			Sample.DiffusePathValue	= D * Vis * F;
 			Sample.DiffusePDF		= BSDFContext.SatNoV / PI; // Per "Production Friendly Microfacet Sheen BRDF", hemispherical sampling give good result as the roughness is usually high.
-			Sample.Throughput		= (1.0f - F);
+			Sample.Throughput		= (1.0f - ApproxFresnelForTransmittance);
 			break;
 		}
 
