Strata: made transmittance to mfp workflow safe and consistent.
- Fixed default thickenss value
- Made extinction and transmittance range safe for new participating media functions

diff --git a/Engine/Shaders/Private/ParticipatingMediaCommon.ush b/Engine/Shaders/Private/ParticipatingMediaCommon.ush
--- a/Engine/Shaders/Private/ParticipatingMediaCommon.ush
+++ b/Engine/Shaders/Private/ParticipatingMediaCommon.ush
@@ -15,7 +15,9 @@ struct FParticipatingMedia
 	float3 Albedo;			// (unitless)
 };
 
-#define PARTICIPATING_MEDIA_MIN_MFP		0.000001f
+#define PARTICIPATING_MEDIA_MIN_MFP_METER		0.000001f
+#define PARTICIPATING_MEDIA_MIN_EXTINCTION		0.000001f
+#define PARTICIPATING_MEDIA_MIN_TRANSMITTANCE	0.000000000001f
 
 // The mean free path must be in meters
 FParticipatingMedia CreateMediumFromAlbedoMFP(float3 Albedo, float3 MeanFreePathMeters)
@@ -23,7 +25,7 @@ FParticipatingMedia CreateMediumFromAlbedoMFP(float3 Albedo, float3 MeanFreePath
 	FParticipatingMedia PM = (FParticipatingMedia)0;
 	PM.Albedo = Albedo;
 	PM.MeanFreePath = MeanFreePathMeters;
-	PM.ExtinctionCoef = 1.0f / max(PARTICIPATING_MEDIA_MIN_MFP, PM.MeanFreePath);
+	PM.ExtinctionCoef = 1.0f / max(PARTICIPATING_MEDIA_MIN_MFP_METER, PM.MeanFreePath);
 	PM.ScatteringCoef = Albedo * PM.ExtinctionCoef;
 	PM.AbsorptionCoef = max(0.0f, PM.ExtinctionCoef - PM.ScatteringCoef);
 	return PM;
@@ -88,15 +90,15 @@ float4 ImportanceSampleHenyeyGreensteinPhase(float2 E, float G)
 	return float4(H, HenyeyGreensteinPhase(G, CosTheta));
 }
 
-float3 TransmittanceToExtinction(in float3 TransmittanceColor, in float Thickness)
+float3 TransmittanceToExtinction(in float3 TransmittanceColor, in float ThicknessMeters)
 {
 	// TransmittanceColor	= exp(-Extinction * Thickness)
 	// Extinction			= -log(TransmittanceColor) / Thickness
-	return -log(saturate(TransmittanceColor)) / max(PARTICIPATING_MEDIA_MIN_MFP, Thickness);
+	return -log(clamp(TransmittanceColor, PARTICIPATING_MEDIA_MIN_TRANSMITTANCE, 1.0f)) / max(PARTICIPATING_MEDIA_MIN_MFP_METER, ThicknessMeters);
 }
 
-float3 TransmittanceToMeanFreePath(in float3 TransmittanceColor, in float Thickness)
+float3 TransmittanceToMeanFreePath(in float3 TransmittanceColor, in float ThicknessMeters)
 {
-	return 1.0f / TransmittanceToExtinction(TransmittanceColor, Thickness);
+	return 1.0f / max(PARTICIPATING_MEDIA_MIN_EXTINCTION, TransmittanceToExtinction(TransmittanceColor, ThicknessMeters));
 }
 
diff --git a/Engine/Shaders/Private/Strata/Strata.ush b/Engine/Shaders/Private/Strata/Strata.ush
--- a/Engine/Shaders/Private/Strata/Strata.ush
+++ b/Engine/Shaders/Private/Strata/Strata.ush
@@ -39,6 +39,7 @@
 // If the user specifies a different thickness to simulate varying appearance, we rescale the mean free path in order to not have to store thickness as part of the BSDF description.
 #define STRATA_SIMPLEVOLUME_THICKNESS_CM	1.0f
 #define STRATA_CENTIMETER_TO_METER			0.01f
+#define STRATA_METER_TO_CENTIMETER			100.0f
 #define STRATA_SIMPLEVOLUME_THICKNESS_M		(STRATA_SIMPLEVOLUME_THICKNESS_CM * STRATA_CENTIMETER_TO_METER)
 
 ///////////////////////////////////////////////////////////////////////////////
diff --git a/Engine/Source/Runtime/Engine/Private/Materials/HLSLMaterialTranslator.cpp b/Engine/Source/Runtime/Engine/Private/Materials/HLSLMaterialTranslator.cpp
--- a/Engine/Source/Runtime/Engine/Private/Materials/HLSLMaterialTranslator.cpp
+++ b/Engine/Source/Runtime/Engine/Private/Materials/HLSLMaterialTranslator.cpp
@@ -9181,10 +9181,14 @@ int32 FHLSLMaterialTranslator::StrataTransmittanceToMFP(int32 TransmittanceColor
 	switch (OutputIndex)
 	{
 	case 0:
-		return AddCodeChunk(MCT_Float3, TEXT("TransmittanceToMeanFreePath(%s, STRATA_SIMPLEVOLUME_THICKNESS_CM)"), *GetParameterCode(TransmittanceColor));
+		return AddCodeChunk(MCT_Float3, 
+			// For the math to be valid, input to TransmittanceToMeanFreePath must be in meter.
+			// Then the output needs to be is converted to centimeters.
+			TEXT("(TransmittanceToMeanFreePath(%s, STRATA_SIMPLEVOLUME_THICKNESS_M) * STRATA_METER_TO_CENTIMETER)"), 
+			*GetParameterCode(TransmittanceColor));
 		break;
 	case 1:
-		return AddInlinedCodeChunk(MCT_Float1, TEXT("STRATA_SIMPLEVOLUME_THICKNESS_CM")); // Thickness
+		return AddInlinedCodeChunk(MCT_Float1, TEXT("STRATA_SIMPLEVOLUME_THICKNESS_CM")); // Thickness to be plugged into other nodes thickness input (in centimeters to match other nodes)
 		break;
 	}
 	return INDEX_NONE;
diff --git a/Engine/Source/Runtime/Engine/Private/Materials/MaterialExpressions.cpp b/Engine/Source/Runtime/Engine/Private/Materials/MaterialExpressions.cpp
--- a/Engine/Source/Runtime/Engine/Private/Materials/MaterialExpressions.cpp
+++ b/Engine/Source/Runtime/Engine/Private/Materials/MaterialExpressions.cpp
@@ -19778,7 +19778,7 @@ int32 UMaterialExpressionStrataSlabBSDF::Compile(class FMaterialCompiler* Compil
 		CompileWithDefaultFloat3(Compiler, EmissiveColor, 0.0f, 0.0f, 0.0f),
 		CompileWithDefaultFloat1(Compiler, Haziness, 0.0f),
 		CompileWithDefaultFloat1(Compiler, ThinFilmThickness, 0.0f),
-		CompileWithDefaultFloat1(Compiler, Thickness, 0.01f),
+		CompileWithDefaultFloat1(Compiler, Thickness, 1.0f),	// thickness is in centimeters, it defaults to 1 centimeter
 		NormalCodeChunk,
 		TangentCodeChunk,
 		SharedNormalIndex);
@@ -20730,17 +20730,17 @@ void UMaterialExpressionStrataTransmittanceToMFP::GetConnectorToolTip(int32 Inpu
 	switch (OutputIndex)
 	{
 	case 0:
-		ConvertToMultilineToolTip(TEXT("The Mean Free Path in centimeters defining the participating media constituting the slab of material."), 80, OutToolTip);
+		ConvertToMultilineToolTip(TEXT("The Mean Free Path defining the participating media constituting the slab of material (unit = centimeters)."), 80, OutToolTip);
 		break;
 	case 1:
-		ConvertToMultilineToolTip(TEXT("The thickness of the slab of material."), 80, OutToolTip);
+		ConvertToMultilineToolTip(TEXT("The thickness of the slab of material (unit = centimeters)."), 80, OutToolTip);
 		break;
 	}
 }
 
 void UMaterialExpressionStrataTransmittanceToMFP::GetExpressionToolTip(TArray<FString>& OutToolTip)
 {
-	ConvertToMultilineToolTip(TEXT("Convert a transmittance color corresponding to a loab of participating media viewed perpendicularly to its surface.\nThis node directly maps to the Slab BSDF input."), 80, OutToolTip);
+	ConvertToMultilineToolTip(TEXT("Convert a transmittance color corresponding to a slab of participating media viewed perpendicularly to its surface.\nThis node directly maps to the Slab BSDF input."), 80, OutToolTip);
 
 }
 #endif // WITH_EDITOR
