Strata in forward, we must normalize all normal and tangents to avoid non normalised vector due to per vertex interpolation of texture filtering.
Fixed vector type warning.

diff --git a/Engine/Shaders/Private/Strata/StrataEvaluation.ush b/Engine/Shaders/Private/Strata/StrataEvaluation.ush
--- a/Engine/Shaders/Private/Strata/StrataEvaluation.ush
+++ b/Engine/Shaders/Private/Strata/StrataEvaluation.ush
@@ -354,7 +354,7 @@ FStrataEvaluateResult StrataEvaluateBSDF(FStrataBSDFContext BSDFContext, FShadow
 				// STRATA_TODO Fresnel throughput, Use a LUT, this is currently in progress and will be coming soon.
 				// In the meantime, we use a simple approach to feed Fresnel function and compute an approximation of the general transmition/througput.
 				// We also ignored ThinFilm interaction as another simplification for perf.
-				const float ApproxFresnelForTransmittance = F_Schlick(F0, F90, BSDFContext.Context.NoV);
+				const float3 ApproxFresnelForTransmittance = F_Schlick(F0, F90, BSDFContext.Context.NoV);
 
 				const float3 SpecularThroughput = 1.0f - ApproxFresnelForTransmittance; // This is already taken into account by the Chan model above. But for volumetric we need to apply it.
 				const float Phase = IsotropicPhase();
@@ -386,7 +386,7 @@ FStrataEvaluateResult StrataEvaluateBSDF(FStrataBSDFContext BSDFContext, FShadow
 
 			// STRATA_TODO Fresnel throughput, Use a LUT, this is currently in progress and will be coming soon.
 			// In the meantime, we use a simple approach to feed Fresnel function and compute an approximation of the general transmition/througput.
-			const float ApproxFresnelForTransmittance = F_Schlick(F0, BSDFContext.Context.NoV);
+			const float3 ApproxFresnelForTransmittance = F_Schlick(F0, BSDFContext.Context.NoV);
 
 			Sample.DiffusePathValue	= D * Vis * F;
 			Sample.DiffusePDF		= BSDFContext.SatNoV / PI; // Per "Production Friendly Microfacet Sheen BRDF", hemispherical sampling give good result as the roughness is usually high.
diff --git a/Engine/Shaders/Private/Strata/StrataForwardLighting.ush b/Engine/Shaders/Private/Strata/StrataForwardLighting.ush
--- a/Engine/Shaders/Private/Strata/StrataForwardLighting.ush
+++ b/Engine/Shaders/Private/Strata/StrataForwardLighting.ush
@@ -25,6 +25,19 @@ float3 StrataForwardLighting(
 	float3 Color = 0;
 	OutThroughput = 1.0f;
 
+#if STRATA_INLINE_SHADING
+	// In forward, we must normalize each normal and tangent to avoid non normalised vectors due to per vertex interpolation or texture filtering.
+	UNROLL
+	for (uint i = 0; i < Header.SharedNormals.NormalCount; ++i)
+	{
+		Header.SharedNormals.Normals[i] = normalize(Header.SharedNormals.Normals[i]);
+		if (StrataGetSharedNormalType(Header.SharedNormals.NormalTypes, i) == STRATA_BASIS_TYPE_TANGENT)
+		{
+			Header.SharedNormals.Tangents[i] = normalize(Header.SharedNormals.Tangents[i]);
+		}
+	}
+#endif
+
 	UNROLL
 	for (int l = 0; l < Strata.LayerCount; ++l)
 	{
