Strata: in fact pulling strata normal/tangent normalisation code in common path because in deferred our packing relies on normalized normals

diff --git a/Engine/Shaders/Private/BasePassPixelShader.usf b/Engine/Shaders/Private/BasePassPixelShader.usf
--- a/Engine/Shaders/Private/BasePassPixelShader.usf
+++ b/Engine/Shaders/Private/BasePassPixelShader.usf
@@ -1541,6 +1541,20 @@ void FPixelShaderInOut_MainPS(
 	bool bStrataSubsurfaceEnable = false;
 #if STRATA_ENABLED && !MATERIAL_SHADINGMODEL_UNLIT
 
+#if STRATA_INLINE_SHADING
+	// We must normalize each normal and tangent to avoid non normalised vectors due to per vertex interpolation or texture filtering,
+	// for the deferred (our packing relies on normalized normal) and forward (normals are going to be used as-is from registers) paths.
+	UNROLL
+	for (uint i = 0; i < StrataPixelHeader.SharedNormals.NormalCount; ++i)
+	{
+		StrataPixelHeader.SharedNormals.Normals[i] = normalize(StrataPixelHeader.SharedNormals.Normals[i]);
+		if (StrataGetSharedNormalType(StrataPixelHeader.SharedNormals.NormalTypes, i) == STRATA_BASIS_TYPE_TANGENT)
+		{
+			StrataPixelHeader.SharedNormals.Tangents[i] = normalize(StrataPixelHeader.SharedNormals.Tangents[i]);
+		}
+	}
+#endif
+
 #if STRATA_OPAQUE_ENABLED
 
 #if STRATA_INLINE_SINGLELAYERWATER==0
diff --git a/Engine/Shaders/Private/Strata/StrataForwardLighting.ush b/Engine/Shaders/Private/Strata/StrataForwardLighting.ush
--- a/Engine/Shaders/Private/Strata/StrataForwardLighting.ush
+++ b/Engine/Shaders/Private/Strata/StrataForwardLighting.ush
@@ -25,19 +25,6 @@ float3 StrataForwardLighting(
 	float3 Color = 0;
 	OutThroughput = 1.0f;
 
-#if STRATA_INLINE_SHADING
-	// In forward, we must normalize each normal and tangent to avoid non normalised vectors due to per vertex interpolation or texture filtering.
-	UNROLL
-	for (uint i = 0; i < Header.SharedNormals.NormalCount; ++i)
-	{
-		Header.SharedNormals.Normals[i] = normalize(Header.SharedNormals.Normals[i]);
-		if (StrataGetSharedNormalType(Header.SharedNormals.NormalTypes, i) == STRATA_BASIS_TYPE_TANGENT)
-		{
-			Header.SharedNormals.Tangents[i] = normalize(Header.SharedNormals.Tangents[i]);
-		}
-	}
-#endif
-
 	UNROLL
 	for (int l = 0; l < Strata.LayerCount; ++l)
 	{
