Change hair strands raster compute to output samples based on coverage.

diff --git a/Engine/Shaders/Private/HairStrands/HairStrandsVisibilityRasterCompute.usf b/Engine/Shaders/Private/HairStrands/HairStrandsVisibilityRasterCompute.usf
--- a/Engine/Shaders/Private/HairStrands/HairStrandsVisibilityRasterCompute.usf
+++ b/Engine/Shaders/Private/HairStrands/HairStrandsVisibilityRasterCompute.usf
@@ -179,23 +179,21 @@ void Rasterize(int2 Coords, float Ratio, float Depth, uint PrimitiveId, float Co
 	#endif
 
 	#if PERMUTATION_SPP == 2
-		Out = WriteVisibilityData(OutVisibilityTexture0, Coords, In);
-		In  = NextValue(In, Out);
-			
-		WriteVisibilityData(OutVisibilityTexture1, Coords, In);
+		if (Ratio < 0.5f)
+			WriteVisibilityData(OutVisibilityTexture0, Coords, In);
+		else
+			WriteVisibilityData(OutVisibilityTexture1, Coords, In);
 	#endif
 
 	#if PERMUTATION_SPP == 4
-		Out = WriteVisibilityData(OutVisibilityTexture0, Coords, In);
-		In  = NextValue(In, Out);
-			
-		Out = WriteVisibilityData(OutVisibilityTexture1, Coords, In);
-		In  = NextValue(In, Out);
-			
-		Out = WriteVisibilityData(OutVisibilityTexture2, Coords, In);
-		In  = NextValue(In, Out);
-			
-		WriteVisibilityData(OutVisibilityTexture3, Coords, In);
+		if (Ratio < 0.25f)
+			WriteVisibilityData(OutVisibilityTexture0, Coords, In);
+		else if (Ratio < 0.5f)
+			WriteVisibilityData(OutVisibilityTexture1, Coords, In);
+		else if (Ratio < 0.75f)
+			WriteVisibilityData(OutVisibilityTexture2, Coords, In);
+		else
+			WriteVisibilityData(OutVisibilityTexture3, Coords, In);
 	#endif
 	}
 }
