Add missing tesselation propagation for PrimitiveId & add missing include of SceneData.ush

diff --git a/Engine/Shaders/Private/GpuSkinVertexFactory.ush b/Engine/Shaders/Private/GpuSkinVertexFactory.ush
--- a/Engine/Shaders/Private/GpuSkinVertexFactory.ush
+++ b/Engine/Shaders/Private/GpuSkinVertexFactory.ush
@@ -324,6 +324,11 @@ FMaterialVertexParameters GetMaterialVertexParameters(FVertexFactoryInput Input,
 		Result.TexCoords[CoordinateIndex] = Input.TexCoords[CoordinateIndex];
 	}
 #endif
+
+#if VF_USE_PRIMITIVE_SCENE_DATA
+	Result.PrimitiveId = Intermediates.SceneData.PrimitiveId;
+#endif
+
 	return Result;
 }
 
@@ -992,6 +997,10 @@ FVertexFactoryInterpolantsVSToDS VertexFactoryInterpolate(FVertexFactoryInterpol
 	}
 #endif
 
+#if VF_USE_PRIMITIVE_SCENE_DATA
+	O.InterpolantsVSToPS.PrimitiveId = a.InterpolantsVSToPS.PrimitiveId;
+#endif
+
 	return O;
 }
 
diff --git a/Engine/Shaders/Private/LandscapeVertexFactory.ush b/Engine/Shaders/Private/LandscapeVertexFactory.ush
--- a/Engine/Shaders/Private/LandscapeVertexFactory.ush
+++ b/Engine/Shaders/Private/LandscapeVertexFactory.ush
@@ -632,6 +632,10 @@ FMaterialVertexParameters GetMaterialVertexParameters(FVertexFactoryInput Input,
 #endif // 1
 #endif // 0
 
+#if VF_USE_PRIMITIVE_SCENE_DATA
+	Result.PrimitiveId = Intermediates.SceneData.PrimitiveId;
+#endif
+
 	return Result;
 }
 
@@ -912,6 +916,10 @@ FVertexFactoryInterpolantsVSToDS VertexFactoryInterpolate(FVertexFactoryInterpol
 	TESSELLATION_INTERPOLATE_MEMBER(CalculatedLOD);
 #endif
 
+#if VF_USE_PRIMITIVE_SCENE_DATA
+	O.InterpolantsVSToPS.PrimitiveId = a.InterpolantsVSToPS.PrimitiveId;
+#endif
+
 	return O;
 }
 
diff --git a/Engine/Shaders/Private/LocalVertexFactory.ush b/Engine/Shaders/Private/LocalVertexFactory.ush
--- a/Engine/Shaders/Private/LocalVertexFactory.ush
+++ b/Engine/Shaders/Private/LocalVertexFactory.ush
@@ -4,7 +4,6 @@
 	LocalVertexFactory.usf: Local vertex factory shader code.
 =============================================================================*/
 
-#include "SceneData.ush"
 #include "VertexFactoryCommon.ush"
 #include "LocalVertexFactoryCommon.ush"
 #include "LightmapData.ush"
diff --git a/Engine/Shaders/Private/VertexFactoryCommon.ush b/Engine/Shaders/Private/VertexFactoryCommon.ush
--- a/Engine/Shaders/Private/VertexFactoryCommon.ush
+++ b/Engine/Shaders/Private/VertexFactoryCommon.ush
@@ -1,6 +1,7 @@
 // Copyright Epic Games, Inc. All Rights Reserved.
 
 #include "Common.ush"
+#include "SceneData.ush"
 
 #if defined(GPUCULL_TODO)
 float3 TransformLocalToWorld(float3 LocalPosition, float4x4 LocalToWorld)
