Fixed analytic volumetric not matching between opaque and forward translucent.

diff --git a/Engine/Shaders/Private/Strata/Strata.ush b/Engine/Shaders/Private/Strata/Strata.ush
--- a/Engine/Shaders/Private/Strata/Strata.ush
+++ b/Engine/Shaders/Private/Strata/Strata.ush
@@ -867,11 +867,22 @@ void SanitizeStrataSlab(inout FStrataBSDF BSDF)
 	SLAB_ROUGHNESSY(BSDF) = saturate(SLAB_ROUGHNESSY(BSDF));
 	SLAB_HAZINESS(BSDF) = saturate(SLAB_HAZINESS(BSDF));
 	SLAB_THINFILM_THICKNESS(BSDF) = saturate(SLAB_THINFILM_THICKNESS(BSDF));
-	SLAB_SSSPROFILERADIUSSCALE(BSDF) = saturate(SLAB_SSSPROFILERADIUSSCALE(BSDF));
-	SLAB_SSSDMFP(BSDF) = clamp(SLAB_SSSDMFP(BSDF), 0.0f, Max111110BitsFloat3);
 #if STRATA_INLINE_SHADING
 	BSDF.TmpDMFP = clamp(BSDF.TmpDMFP, 0.0f, Max111110BitsFloat3);
 #endif
+
+	// BSDF feature tests must be executed for profile path to not interact with DMFP, especially for forward rendered simple volumetrics.
+	if (BSDF_GETHASSSS(BSDF))
+	{
+		if (BSDF_GETHASSSSPROFILE(BSDF))
+		{
+			SLAB_SSSPROFILERADIUSSCALE(BSDF) = saturate(SLAB_SSSPROFILERADIUSSCALE(BSDF));
+		}
+		else
+		{
+			SLAB_SSSDMFP(BSDF) = clamp(SLAB_SSSDMFP(BSDF), 0.0f, Max111110BitsFloat3);
+		}
+	}
 }
 
 void SanitizeStrataSheen(inout FStrataBSDF BSDF)
@@ -986,7 +997,7 @@ void EnableSlabBSDFSimpleVolumetric(inout FStrataBSDF BSDF)
 
 	// Rescale the BSDF mean free path so that we can drop the thickness parameter and run volumetric material math on a unit slab of medium
 	// For instance, if thickness is larger => we need to reduce the MFP for the medium to appear visually thicker.
-	SLAB_SSSDMFP(BSDF) = SLAB_SSSDMFP(BSDF) * (STRATA_SIMPLEVOLUME_THICKNESS_CM / BSDF_GETTHICKNESSCM(BSDF));
+	SLAB_SSSDMFP(BSDF) = clamp(SLAB_SSSDMFP(BSDF) * (STRATA_SIMPLEVOLUME_THICKNESS_CM / BSDF_GETTHICKNESSCM(BSDF)), 0.0f, Max111110BitsFloat3);
 
 	// Mark the BSDF has not having any SSS through post process.
 	BSDF_SETHASSSS(BSDF, 0);
