Add nmore coverage to furnace test
* Add sampling integrator
* Add sky integrator

* Fix importance sampling evaluation

diff --git a/Engine/Shaders/Private/Strata/StrataEvaluation.ush b/Engine/Shaders/Private/Strata/StrataEvaluation.ush
--- a/Engine/Shaders/Private/Strata/StrataEvaluation.ush
+++ b/Engine/Shaders/Private/Strata/StrataEvaluation.ush
@@ -128,8 +128,8 @@ FStrataBSDFContext StrataCreateBSDFContext(float3x3 TangentBasis, FStrataBSDF BS
 	}
 	BSDFContext.TangentBasis = float3x3(BSDFContext.X, BSDFContext.Y, BSDFContext.N);
 
-	BSDFContext.TangentV = mul(BSDFContext.V, BSDFContext.TangentBasis);
-	BSDFContext.TangentH = mul(BSDFContext.H, BSDFContext.TangentBasis);
+	BSDFContext.TangentV = mul(BSDFContext.TangentBasis, BSDFContext.V);
+	BSDFContext.TangentH = mul(BSDFContext.TangentBasis, BSDFContext.H);
 
 	BSDFContext.SatNoL = saturate(BSDFContext.Context.NoL);
 	BSDFContext.SatNoV = saturate(BSDFContext.Context.NoV);
@@ -434,12 +434,10 @@ FStrataEvaluateResult StrataEvaluateBSDF(FStrataBSDFContext BSDFContext)
 struct FStrataImportanceSampleResult
 {
 	float3 SpecularDirection;
-	float  SpecularConeAngle;		// bad approximate of Cone angle as 1/PDF as done by Lumen STRATA_TODO update when Lumen is pluged in...
+	float  SpecularPDF;
 
 	float3 DiffuseDirection;
-	float  DiffuseConeAngle;		// bad approximate of Cone angle as 1/PDF as done by Lumen STRATA_TODO update when Lumen is pluged in...
-
-//	FStrataEvaluateResult Eval;
+	float  DiffusePDF;
 };
 
 /**
@@ -448,15 +446,10 @@ struct FStrataImportanceSampleResult
  * E: two random numbers
  * CameraVector: vector from the camera to the considered direction (not V)
  */
-FStrataImportanceSampleResult StrataImportanceSampleBSDF(FStrataPixelHeader StrataPixelHeader, FStrataBSDF BSDF, const FStrataAddressing StrataAddressing, float2 E, float3 CameraVector)
+FStrataImportanceSampleResult StrataImportanceSampleBSDF(FStrataBSDFContext BSDFContext, FStrataBSDF BSDF, float2 E)
 {
 	FStrataImportanceSampleResult Sample = (FStrataImportanceSampleResult)0;
 
-	// Create a context from the information we have
-	const float3 UnusedL = float3(0.0f, 0.0f, 1.0f);
-	const float3 V = -CameraVector;
-	FStrataBSDFContext BSDFContext = StrataCreateBSDFContext(StrataPixelHeader, BSDF, StrataAddressing, V, UnusedL);
-
 	float PDF = 0.0f;
 	const uint BSDFType = BSDF_GETTYPE(BSDF);
 	switch (BSDFType)
@@ -470,31 +463,34 @@ FStrataImportanceSampleResult StrataImportanceSampleBSDF(FStrataPixelHeader Stra
 			// We also assume that this will work just as well for volumetric (considering it is an isotropic phase function and we are interesting in sampling mostly the hemisphere and not the layers below)
 			float4 ImportanceSample = CosineSampleHemisphere(E);
 			Sample.DiffuseDirection = mul(ImportanceSample.xyz, BSDFContext.TangentBasis);
-			PDF = ImportanceSample.w;
-			Sample.DiffuseConeAngle = 1.0f / PDF;
+			Sample.DiffusePDF = ImportanceSample.w;
+			PDF = Sample.DiffusePDF;
 
 			const float3 F0 = ComputeF0(SLAB_SPECULAR(BSDF), BaseColor, Metallic);
-			// STRATA_TODO F82
+			// STRATA_TODO F82pb 
 
 			const float SafeRoughnessX = MakeRoughnessSafe(SLAB_ROUGHNESSX(BSDFContext.BSDF));
 			const bool bHasAnisotropy = BSDF_GETANISOTROPY(BSDFContext.BSDF);
-			float4 GGXSample = 0;
+			float4 TangentH = 0;
 			if (bHasAnisotropy)
 			{
 				const float SafeRoughnessY = MakeRoughnessSafe(SLAB_ROUGHNESSY(BSDFContext.BSDF));
-				GGXSample = ImportanceSampleVisibleGGX_aniso(UniformSampleDisk(E), float2(Pow2(SafeRoughnessX), Pow2(SafeRoughnessY)), BSDFContext.TangentV);
+				TangentH = ImportanceSampleVisibleGGX_aniso(UniformSampleDisk(E), float2(Pow2(SafeRoughnessX), Pow2(SafeRoughnessY)), BSDFContext.TangentV);
 			}
 			else
 			{
-				GGXSample = ImportanceSampleVisibleGGX(UniformSampleDisk(E), Pow4(SafeRoughnessX), BSDFContext.TangentV);
+				TangentH = ImportanceSampleVisibleGGX(UniformSampleDisk(E), Pow4(SafeRoughnessX), BSDFContext.TangentV);
 			}
 
 			// STRATA_TODO: take into account Haziness
 
-			float3 WorldH = mul(GGXSample.xyz, BSDFContext.TangentBasis);
-			Sample.SpecularDirection = reflect(CameraVector, WorldH);
-			PDF = GGXSample.w;
-			Sample.SpecularConeAngle = 1.0f / PDF;
+			const float HPDF = TangentH.w;
+			const float3 H = mul(TangentH.xyz, BSDFContext.TangentBasis);
+			const float VoH = saturate(dot(BSDFContext.V, H));
+
+			Sample.SpecularDirection = 2 * dot(BSDFContext.V, H) * H - BSDFContext.V;
+			Sample.SpecularPDF = RayPDFToReflectionRayPDF(VoH, HPDF);
+			PDF = Sample.SpecularPDF;
 			break;
 		}
 
@@ -505,7 +501,6 @@ FStrataImportanceSampleResult StrataImportanceSampleBSDF(FStrataPixelHeader Stra
 			float4 ImportanceSample = CosineSampleHemisphere(E);
 			Sample.DiffuseDirection = mul(ImportanceSample.xyz, BSDFContext.TangentBasis);
 			PDF = ImportanceSample.w;
-			Sample.DiffuseConeAngle = 1.0f / PDF;
 			break;
 		}
 
@@ -515,7 +510,6 @@ FStrataImportanceSampleResult StrataImportanceSampleBSDF(FStrataPixelHeader Stra
 			float4 ImportanceSample = CosineSampleHemisphere(E);
 			Sample.DiffuseDirection = mul(ImportanceSample.xyz, BSDFContext.TangentBasis);
 			PDF = ImportanceSample.w;
-			Sample.DiffuseConeAngle = 1.0f / PDF;
 			break;
 		}
 
@@ -533,6 +527,17 @@ FStrataImportanceSampleResult StrataImportanceSampleBSDF(FStrataPixelHeader Stra
 	return Sample;
 }
 
+FStrataImportanceSampleResult StrataImportanceSampleBSDF(FStrataPixelHeader StrataPixelHeader, FStrataBSDF BSDF, const FStrataAddressing StrataAddressing, float2 E, float3 CameraVector)
+{
+	FStrataImportanceSampleResult Sample = (FStrataImportanceSampleResult)0;
+
+	// Create a context from the information we have
+	const float3 UnusedL = float3(0.0f, 0.0f, 1.0f);
+	const float3 V = -CameraVector;
+	FStrataBSDFContext BSDFContext = StrataCreateBSDFContext(StrataPixelHeader, BSDF, StrataAddressing, V, UnusedL);
+	return StrataImportanceSampleBSDF(BSDFContext, BSDF, E);
+}
+
 struct FStrataEnvLightResult
 {
 	float3 DiffuseColor;
diff --git a/Engine/Shaders/Private/Strata/StrataFurnaceTest.usf b/Engine/Shaders/Private/Strata/StrataFurnaceTest.usf
--- a/Engine/Shaders/Private/Strata/StrataFurnaceTest.usf
+++ b/Engine/Shaders/Private/Strata/StrataFurnaceTest.usf
@@ -8,9 +8,14 @@
 #include "StrataMultipleScattering.ush"
 #include "StrataEvaluation.ush"
 
+#define STRATA_INTEGRATION_EVALUATE 0
+#define STRATA_INTEGRATION_SAMPLING 1
+#define STRATA_INTEGRATION_ENV 2
+
 #if SHADER_FURNACE_ANALYTIC
 
 uint SceneType;
+uint IntegratorType;
 uint NumSamples;
 
 void MainPS(
@@ -30,13 +35,13 @@ void MainPS(
 	const uint2 TileCoord = floor(UV * TileCount);
 	const uint2 TilePixelCoord = TileUV * TileResolution;
 
-	float3 V = float3(0, 0, -1);
+	float3 V = float3(0, 0, 1);
 	float3 N = 0;
 	bool bEvaluate = false;
 	{
-		const float3 RayO = float3(TileUV, 0);
-		const float3 RayD = float3(0, 0, 1);
-		const float4 Sphere = float4(0.5f, 0.5f, 1.f, 0.4f);
+		const float3 RayO = float3(TileUV, 1);
+		const float3 RayD = -V;
+		const float4 Sphere = float4(0.5f, 0.5f, 0.f, 0.4f);
 		const float2 Hit = RayIntersectSphere(RayO, RayD, Sphere);
 		if (Hit.x >= 0)
 		{
@@ -126,26 +131,79 @@ void MainPS(
 		// Sample uniformly the upper hemisphere of each pixel
 		float3 Acc = 0;
 		const FShadowTerms ShadowTerms = { 1.0f, 1.0, 1.0, InitHairTransmittanceData() };
-		LOOP
-		for (uint SampleIt = 0; SampleIt < NumSamples; ++SampleIt)
+		const float3x3 TangentBasis = GetTangentBasis(N);
+
+		if (IntegratorType == STRATA_INTEGRATION_EVALUATE)
 		{
-			const float3x3 TangentBasis = GetTangentBasis(N);
-			const float2 U = Hammersley(SampleIt, NumSamples, 0);
+			LOOP
+			for (uint SampleIt = 0; SampleIt < NumSamples; ++SampleIt)
+			{
+				const float2 U = Hammersley(SampleIt, NumSamples, 0);
+
+				const float4 S = UniformSampleHemisphere(U);
+				const float3 L = mul(S.xyz, TangentBasis);
+
+				const float SatNoL = saturate(dot(L, N));
+
+				FStrataBSDFContext BSDFContext = StrataCreateBSDFContext(TangentBasis, BSDF, V, L);
+				FStrataEvaluateResult BSDFEvaluate = StrataEvaluateBSDF(BSDFContext, ShadowTerms);
+
+				const float3 DiffusePathWeight  = BSDF.Weight * BSDFEvaluate.DiffusePathValue;
+				const float3 SpecularPathWeight = BSDF.Weight * BSDFEvaluate.SpecularPathValue;
+
+				Acc += SatNoL * InRadiance * (DiffusePathWeight + SpecularPathWeight) / (S.w != 0 ?  S.w : 1);
+			}
+			OutColor = Acc / max(1, NumSamples);
+		}
+		else if (IntegratorType == STRATA_INTEGRATION_SAMPLING)
+		{
+			LOOP
+			for (uint SampleIt = 0; SampleIt < NumSamples; ++SampleIt)
+			{
+				const float2 U = Hammersley(SampleIt, NumSamples, 0);
+
+				const float3 UnusedL = float3(0.0f, 0.0f, 1.0f);
+				FStrataBSDFContext BSDFContextIS = StrataCreateBSDFContext(TangentBasis, BSDF, V, UnusedL);
+				FStrataImportanceSampleResult S = StrataImportanceSampleBSDF(BSDFContextIS, BSDF, U);
+
+				// Evaluate specular part for now. 
+				// Importance sample (STRATA_TODO: change API to support diffuse/specular filter and MIS)
+				const float3 L = S.SpecularDirection;
+				const float PDF = S.SpecularPDF;
+				//const float3 L = S.DiffuseDirection;
+				//const float PDF = S.DiffusePDF;
+
+				const float SatNoL = saturate(dot(L, N));
+				FStrataBSDFContext BSDFContext = StrataCreateBSDFContext(TangentBasis, BSDF, V, L);
+				FStrataEvaluateResult BSDFEvaluate = StrataEvaluateBSDF(BSDFContext, ShadowTerms);
+
+				const float3 DiffusePathWeight  = BSDF.Weight * BSDFEvaluate.DiffusePathValue;
+				const float3 SpecularPathWeight = BSDF.Weight * BSDFEvaluate.SpecularPathValue;
 
-			const float4 S = UniformSampleHemisphere(U);
-			const float3 L = mul(S.xyz, TangentBasis);
+				Acc += SatNoL * InRadiance * (DiffusePathWeight + SpecularPathWeight) / (PDF != 0 ? PDF : 1);
 
-			const float SatNoL = saturate(dot(L, N));
+			}
+			OutColor = Acc / max(1, NumSamples);
+		}
+		else if (IntegratorType == STRATA_INTEGRATION_ENV)
+		{
+			const float3 BSDFThroughput = BSDF.Weight;
+
+			const float3 UnusedL = float3(0.0f, 0.0f, 1.0f);
+			const float3x3 TangentBasis = GetTangentBasis(N);
+			FStrataBSDFContext BSDFContext = StrataCreateBSDFContext(TangentBasis, BSDF, V, UnusedL);
 
-			FStrataBSDFContext BSDFContext = StrataCreateBSDFContext(TangentBasis, BSDF, V, L);
-			FStrataEvaluateResult BSDFEvaluate = StrataEvaluateBSDF(BSDFContext, ShadowTerms);
+			const bool bEnableSpecular = true;
+			FStrataEnvLightResult StrataEnvLight = StrataEvaluateForEnvLight(BSDFContext, bEnableSpecular);
 
-			const float3 DiffusePathWeight  = BSDF.Weight * BSDFEvaluate.DiffusePathValue;
-			const float3 SpecularPathWeight = BSDF.Weight * BSDFEvaluate.SpecularPathValue;
+			const float3 InIrradiance = InRadiance; //It should have a Pi factor though
+			const float3 SpecularSkyLighting = BSDF.Weight * StrataEnvLight.SpecularWeight * InRadiance;
+			const float3 DiffuseSkyLighting  = BSDF.Weight * StrataEnvLight.DiffuseColor * InIrradiance;
 
-			Acc += SatNoL * InRadiance * (DiffusePathWeight + SpecularPathWeight) / (S.w != 0 ?  S.w : 1);
+			OutColor = 
+				SpecularSkyLighting + 
+				DiffuseSkyLighting;;
 		}
-		OutColor = Acc / max(1, SampleIt); 
 	}
 }
 
diff --git a/Engine/Source/Runtime/Renderer/Private/Strata/Strata.cpp b/Engine/Source/Runtime/Renderer/Private/Strata/Strata.cpp
--- a/Engine/Source/Runtime/Renderer/Private/Strata/Strata.cpp
+++ b/Engine/Source/Runtime/Renderer/Private/Strata/Strata.cpp
@@ -62,6 +62,12 @@ static TAutoConsoleVariable<int32> CVarStrataFurnaceTest(
 	TEXT("Enable Strata furnace test (for debug purpose) 1: roughness/metallic 2: a selection of conductors."),
 	ECVF_RenderThreadSafe);
 
+static TAutoConsoleVariable<int32> CVarStrataFurnaceTestIntegratorType(
+	TEXT("r.Strata.FurnaceTest.IntegratorType"),
+	0,
+	TEXT("Change Strata furnace test integrator (for debug purpose) 0: evaluate integrator 1: importance sampling integrator 2: env. integrator."),
+	ECVF_RenderThreadSafe);
+
 static TAutoConsoleVariable<int32> CVarStrataFurnaceTestSampleCount(
 	TEXT("r.Strata.FurnaceTest.SampleCount"),
 	1024,
@@ -805,6 +811,9 @@ class FStrataFurnaceTestPassPS : public FGlobalShader
 		SHADER_PARAMETER_STRUCT_REF(FStrataGlobalUniformParameters, Strata)
 		SHADER_PARAMETER(uint32, NumSamples)
 		SHADER_PARAMETER(uint32, SceneType)
+		SHADER_PARAMETER(uint32, IntegratorType)
+		SHADER_PARAMETER_TEXTURE(Texture2D, PreIntegratedGF)
+		SHADER_PARAMETER_SAMPLER(SamplerState, PreIntegratedGFSampler)
 		SHADER_PARAMETER_RDG_TEXTURE(Texture3D, OutLUT3D)
 		SHADER_PARAMETER_RDG_TEXTURE(Texture2D, OutLUT2D)
 		RENDER_TARGET_BINDING_SLOTS()
@@ -836,6 +845,9 @@ static void AddStrataFurnacePass(FRDGBuilder& GraphBuilder, const FViewInfo& Vie
 	Parameters->ViewUniformBuffer = View.ViewUniformBuffer;
 	Parameters->Strata = Strata::BindStrataGlobalUniformParameters(View);
 	Parameters->SceneType = FMath::Clamp(CVarStrataFurnaceTest.GetValueOnAnyThread(), 1, 2);
+	Parameters->IntegratorType = FMath::Clamp(CVarStrataFurnaceTestIntegratorType.GetValueOnAnyThread(), 0, 2);
+	Parameters->PreIntegratedGF = GSystemTextures.PreintegratedGF->GetRenderTargetItem().ShaderResourceTexture;
+	Parameters->PreIntegratedGFSampler = TStaticSamplerState<SF_Bilinear, AM_Clamp, AM_Clamp, AM_Clamp>::GetRHI();
 	Parameters->NumSamples = FMath::Clamp(CVarStrataFurnaceTestSampleCount.GetValueOnAnyThread(), 16, 2048);
 	Parameters->RenderTargets[0] = FRenderTargetBinding(OutTexture, ERenderTargetLoadAction::ELoad);
 
