Add anisotropy and haziness test case to furnace test.

diff --git a/Engine/Shaders/Private/Strata/StrataFurnaceTest.usf b/Engine/Shaders/Private/Strata/StrataFurnaceTest.usf
--- a/Engine/Shaders/Private/Strata/StrataFurnaceTest.usf
+++ b/Engine/Shaders/Private/Strata/StrataFurnaceTest.usf
@@ -37,6 +37,7 @@ void MainPS(
 
 	float3 V = float3(0, 0, 1);
 	float3 N = 0;
+	float3 T = 0;
 	bool bEvaluate = false;
 	{
 		const float3 RayO = float3(TileUV, 1);
@@ -47,6 +48,7 @@ void MainPS(
 		{
 			const float3 P = RayO + RayD * Hit.x;
 			N = normalize(P - Sphere.xyz);
+			T = normalize(cross(float3(0,0,1), N));
 			bEvaluate = true;
 		}
 	}
@@ -64,7 +66,11 @@ void MainPS(
 			float Roughness  = 0.5f;
 			float Metallic   = 0.0f;
 			float Specular   = 0.5f;
+			float Anisoptropy = 0.0f;
+			float Haziness    = 0.0f;
 
+			const float X = saturate((TileCoord.x) / float(TileCount.x - 1));
+			const float Y = saturate((TileCoord.y) / float(TileCount.y - 1));
 			if (SceneType == 1)
 			{
 				//  ------> Roughness
@@ -72,10 +78,32 @@ void MainPS(
 				// |
 				// v
 				// Metallic
-				Roughness = max(MinRoughness, (TileCoord.x) / float(TileCount.x - 1));
-				Metallic = (TileCoord.y) / float(TileCount.y - 1);
+				Roughness = max(MinRoughness, X);
+				Metallic = Y;
 			}
 			else if (SceneType == 2)
+			{
+				//  ------> Roughness
+				// |
+				// |
+				// v
+				// Anisotropy
+				Roughness   = max(MinRoughness, X);
+				Anisoptropy = Y * 2 - 1;
+				Metallic    = 1.f;
+			}
+			else if (SceneType == 3)
+			{
+				//  ------> Roughness
+				// |
+				// |
+				// v
+				// Hazizness
+				Roughness = max(MinRoughness, X);
+				Haziness  = Y;
+				Metallic  = 1.f;
+			}
+			else if (SceneType == 4)
 			{
 				// All metalic with different roughness & metal type (F0)
 				//  ------> Roughness
@@ -93,7 +121,7 @@ void MainPS(
 				const float3 Cobalt      = float3(0.662124,    0.654864,    0.633732);
 				const float3 Platinum    = float3(0.672411,    0.637331,    0.585456);
 
-				Roughness = max(MinRoughness, (TileCoord.x) / float(TileCount.x - 1));
+				Roughness = max(MinRoughness, X);
 				Metallic = 1;
 				switch (TileCoord.y)
 				{
@@ -110,28 +138,34 @@ void MainPS(
 
 			BSDF_SETTYPE(BSDF, STRATA_BSDF_TYPE_SLAB);
 			BSDF_SETNORMALID(BSDF, 0);
-			BSDF_SETANISOTROPY(BSDF, 0);
+			BSDF_SETANISOTROPY(BSDF, Anisoptropy != 0 ? 1 : 0);
 			BSDF_SETISTOPLAYER(BSDF, 1);
 			BSDF_SETHASEDGECOLOR(BSDF, 0);
 			BSDF_SETHASSSS(BSDF, 0);
 			BSDF_SETHASSSSPROFILE(BSDF, 0);
-			BSDF_SETHASHAZINESS(BSDF, 0);
+			BSDF_SETHASHAZINESS(BSDF, Haziness != 0 ? 1 : 0);
 			BSDF_SETHASTHINFILM(BSDF, 0);
 			BSDF.Weight = 1.0f;
 
+			const float2 RoughnessXY = GetAnisotropicRoughness(Roughness, Anisoptropy);
+
 			SLAB_BASECOLOR(BSDF)	= BaseColor;
 			SLAB_EDGECOLOR(BSDF)	= EdgeColor;
-			SLAB_ROUGHNESSX(BSDF)	= Roughness;
-			SLAB_ROUGHNESSY(BSDF)	= Roughness;
+			SLAB_ROUGHNESSX(BSDF)	= RoughnessXY.x;
+			SLAB_ROUGHNESSY(BSDF)	= RoughnessXY.y;
 			SLAB_SPECULAR(BSDF)		= Specular;
 			SLAB_METALLIC(BSDF)		= Metallic;
-			SLAB_HAZINESS(BSDF)		= 0;
+			SLAB_HAZINESS(BSDF)		= Haziness;
 		}
 
 		// Sample uniformly the upper hemisphere of each pixel
 		float3 Acc = 0;
 		const FShadowTerms ShadowTerms = { 1.0f, 1.0, 1.0, InitHairTransmittanceData() };
-		const float3x3 TangentBasis = GetTangentBasis(N);
+		float3x3 TangentBasis = GetTangentBasis(N);
+		if (BSDF_GETANISOTROPY(BSDF))
+		{
+			TangentBasis = float3x3(T, normalize(cross(N,T)), N);
+		}
 
 		if (IntegratorType == STRATA_INTEGRATION_EVALUATE)
 		{
diff --git a/Engine/Source/Runtime/Renderer/Private/Strata/Strata.cpp b/Engine/Source/Runtime/Renderer/Private/Strata/Strata.cpp
--- a/Engine/Source/Runtime/Renderer/Private/Strata/Strata.cpp
+++ b/Engine/Source/Runtime/Renderer/Private/Strata/Strata.cpp
@@ -59,7 +59,7 @@ static TAutoConsoleVariable<int32> CVarStrataLUTContinousUpdate(
 static TAutoConsoleVariable<int32> CVarStrataFurnaceTest(
 	TEXT("r.Strata.FurnaceTest"),
 	0,
-	TEXT("Enable Strata furnace test (for debug purpose) 1: roughness/metallic 2: a selection of conductors."),
+	TEXT("Enable Strata furnace test (for debug purpose) 1:roughness/metallic, 2:roughness/aniso, 3:roughness/haze, 4:a selection of conductors."),
 	ECVF_RenderThreadSafe);
 
 static TAutoConsoleVariable<int32> CVarStrataFurnaceTestIntegratorType(
@@ -844,7 +844,7 @@ static void AddStrataFurnacePass(FRDGBuilder& GraphBuilder, const FViewInfo& Vie
 	FStrataFurnaceTestPassPS::FParameters* Parameters = GraphBuilder.AllocParameters<FStrataFurnaceTestPassPS::FParameters>();
 	Parameters->ViewUniformBuffer = View.ViewUniformBuffer;
 	Parameters->Strata = Strata::BindStrataGlobalUniformParameters(View);
-	Parameters->SceneType = FMath::Clamp(CVarStrataFurnaceTest.GetValueOnAnyThread(), 1, 2);
+	Parameters->SceneType = FMath::Clamp(CVarStrataFurnaceTest.GetValueOnAnyThread(), 1, 4);
 	Parameters->IntegratorType = FMath::Clamp(CVarStrataFurnaceTestIntegratorType.GetValueOnAnyThread(), 0, 2);
 	Parameters->PreIntegratedGF = GSystemTextures.PreintegratedGF->GetRenderTargetItem().ShaderResourceTexture;
 	Parameters->PreIntegratedGFSampler = TStaticSamplerState<SF_Bilinear, AM_Clamp, AM_Clamp, AM_Clamp>::GetRHI();
