Clean up MS strata LUT.
* Remove A_T/B_T terms are they are incorrect
* Clean the MS lookup method
* Added directional albedo computation for diffuse term
* Rename FStrataLUTSample into FStrataGGXLUTSample

diff --git a/Engine/Shaders/Private/Strata/StrataEvaluation.ush b/Engine/Shaders/Private/Strata/StrataEvaluation.ush
--- a/Engine/Shaders/Private/Strata/StrataEvaluation.ush
+++ b/Engine/Shaders/Private/Strata/StrataEvaluation.ush
@@ -197,9 +197,9 @@ float3 StrataEvaluateMicrofacetMultipleScattering(float3 F0, float NoV, float Me
 	// "Revisiting Physically Based Shading at Imageworks" Kulla 2017
 	const float SafeRoughness = SafeRoughnessX;
 
-	const FStrataLUTResult StrataLUTSample = StrataSampleGGXLUT(Strata.GGXEnergyLUT2DTexture, Strata.GGXEnergyLUT3DTexture, Strata.GGXEnergyLUTSampler, Strata.GGXEnergyLUTScaleBias, NoV, SafeRoughness, F0);
-	const float3 Fs_0 = StrataEvaluateMultipleScatteringDielectric3D(F0, StrataLUTSample);
-	const float3 Fs_1 = StrataEvaluateMultipleScatteringConductor(F0, StrataLUTSample);
+	const FStrataGGXLUTResult LUTSample = StrataSampleGGXLUT(Strata.GGXEnergyLUT2DTexture, Strata.GGXEnergyLUT3DTexture, Strata.GGXEnergyLUTSampler, Strata.GGXEnergyLUTScaleBias, NoV, SafeRoughness, F0);
+	const float3 Fs_0 = StrataEvaluateMultipleScatteringDielectric(F0, LUTSample);
+	const float3 Fs_1 = StrataEvaluateMultipleScatteringConductor(F0, LUTSample);
 	const float3 Fs   = lerp(Fs_0, Fs_1, Metallic);
 
 	return Fs;
diff --git a/Engine/Shaders/Private/Strata/StrataLUT.usf b/Engine/Shaders/Private/Strata/StrataLUT.usf
--- a/Engine/Shaders/Private/Strata/StrataLUT.usf
+++ b/Engine/Shaders/Private/Strata/StrataLUT.usf
@@ -80,11 +80,9 @@ void MainPS(
 	float A_R = 0.0f;
 	float B_R = 0.0f;
 
-	float A_T = 0.0f;
-	float B_T = 0.0f;
-
-	float E_R = 0.0f;
-	float E_T = 0.0f;
+	float E_R = 0.0f; // Directional albedo for 'specular' reflection 
+	float E_T = 0.0f; // Directional albedo for 'specular' transmission
+	float E_D = 0.0f; // Directional albedo for 'diffuse'  term
 
 	// Roughness
 	const float Roughness = float(Coord.y + 0.5f) / LUTResolution.y;
@@ -100,6 +98,24 @@ void MainPS(
 	{
 		const float2 E = Hammersley(i, NumSamples, uint2(0, 0));
 
+		// Diffuse term integration
+		{
+			const float4 S = UniformSampleHemisphere(E);
+			const float3 L = S.xyz;
+
+			const float DiffuseColor = 1.0f;
+
+			const float  SatNoV = saturate(NoV);
+			const float  SatNoL = saturate(L.z);
+			const float3 H = normalize(V + L);
+			const float  VoH = saturate(dot(V, H));
+			const float  NoH = saturate(H.z);
+			if (S.w > 0)
+			{
+				E_D += Diffuse_Chan(DiffuseColor.xxx, m2, SatNoV, SatNoL, VoH, NoH).x * SatNoL / S.w;
+			}
+		}
+
 		// Decomposition based on Schlick's fresnel formula
 		// F_Schlick = F0 + (1-F0) (1-cos)^5 = F0 + (1-F0) Fc
 		// 
@@ -116,26 +132,6 @@ void MainPS(
 		//
 		// Reconstruction of the directional Albedo for upper hemisphere
 		// E_R = F0.A_R + (1-F0).B_R
-		//
-		//
-		// LOWER \int{Omera_{-}}
-		// =====================
-		// Decomposition in A B terms of the directional albedo for microfacter for *lower* hemisphere
-		//
-		// E_T
-		// = \int (1-F).D.G
-		// = \int (1- (F0 + (1-F0) Fc)).D.G
-		// = \int (1 - F0 - (1-F0) Fc).D.G
-		// = \int D.G - F0.D.G - (1-F0) Fc.D.G
-		// = \int D.G - \int F0.D.G - \int (1-F0) Fc.D.G
-		// = \int D.G - F0 \int D.G - (1-F0) \int Fc.D.G
-		//   ========      ========          ===========
-		//     A_T            A_T                B_T
-		//
-		//
-		// Reconstruction of the directional Albedo for lower hemisphere
-		// E_T = A_T + F0.A_T + (1-F0).B_T
-		//
 		
 		// "Based on "Practical multiple scattering compensation for microfacet models". E. Turquin
 		//
@@ -182,27 +178,23 @@ void MainPS(
 			A_R += R_WeightPdf;
 			B_R += R_WeightPdf * Fc;
 
-			A_T += T_WeightPdf;
-			B_T += T_WeightPdf * Fc;
+			E_D /= NumSamples;
 		}
 	}
 
 	A_R /= NumSamples;
 	B_R /= NumSamples;
-			
-	A_T /= NumSamples;
-	B_T /= NumSamples;
 
 	E_R /= NumSamples;
 	E_T /= NumSamples;
 
 	const float2 Out3D = float2(E_R, E_T);
-	const float4 Out2D = float4(A_R, B_R, A_T, B_T);
+	const float4 Out2D = float4(A_R, B_R, E_D, 0.0f);
 
 	// At the moment ouput 2 LUTs:
 	// * A 2D LUT which factor out IOR/F0 data
-	// * A 3D LUT which compute the directional albedo for all F0 value (for debug purpose)
-	if (Coord.z == 2) // ~ IOR 1.5 for 64x64 LUT (this is still experimental)
+	// * A 3D LUT which compute the directional albedo for all F0 value
+	if (Coord.z == 2) // ~ IOR 1.5 for 64x64 LUT (the actual IOR does not matter as the factors A_R/B_R are independent of F0)
 	{
 		OutLUT2D[Coord.xy]	= Out2D;
 	}
diff --git a/Engine/Shaders/Private/Strata/StrataMultipleScattering.ush b/Engine/Shaders/Private/Strata/StrataMultipleScattering.ush
--- a/Engine/Shaders/Private/Strata/StrataMultipleScattering.ush
+++ b/Engine/Shaders/Private/Strata/StrataMultipleScattering.ush
@@ -10,24 +10,21 @@
 // 
 // These two variants exist only for testing purpose at the moment
 
-struct FStrataLUTResult
+struct FStrataGGXLUTResult
 {
 	// Fresnel A/B term for reconstructing directional albedo for reflection
 	float A_R;
 	float B_R;
 
-	// Fresnel A/B term for reconstructing directional albedo for refraction
-	float A_T;
-	float B_T;
-
-	// Complete directional albedo for reflection and refraction
+	// Complete directional albedo for reflection, refraction, and diffuse
 	float E_R;
 	float E_T;
+	float E_D;
 };
 
-FStrataLUTResult StrataSampleGGXLUT(Texture2D<float4> LUT2D, Texture3D<float2> LUT3D, SamplerState LUTSampler, float2 LUTScaleBias, float NoV, float Roughness, float3 F0)
+FStrataGGXLUTResult StrataSampleGGXLUT(Texture2D<float4> LUT2D, Texture3D<float2> LUT3D, SamplerState LUTSampler, float2 LUTScaleBias, float NoV, float Roughness, float3 F0)
 {
-	FStrataLUTResult Out = (FStrataLUTResult)0;
+	FStrataGGXLUTResult Out = (FStrataGGXLUTResult)0;
 
 	// 2D LUT sampling
 	{
@@ -35,9 +32,7 @@ FStrataLUTResult StrataSampleGGXLUT(Texture2D<float4> LUT2D, Texture3D<float2> L
 		const float4 S = LUT2D.SampleLevel(LUTSampler, UV, 0);
 		Out.A_R = S.x;
 		Out.B_R = S.y;
-
-		Out.A_T = S.z;
-		Out.B_T = S.w;
+		Out.E_D = S.z;
 	}
 
 	// 3D LUT sampling
@@ -52,7 +47,7 @@ FStrataLUTResult StrataSampleGGXLUT(Texture2D<float4> LUT2D, Texture3D<float2> L
 }
 
 // From "Practical multiple scattering compensation for microfacet models." E. Turquin
-float3 StrataEvaluateMultipleScatteringConductor(float3 F0, FStrataLUTResult LUT)
+float3 StrataEvaluateMultipleScatteringConductor(float3 F0, FStrataGGXLUTResult LUT)
 {
 	// Reconstruction of the directional Albedo for upper hemisphere
 	// E_R = F0.A_R + (1-F0).B_R
@@ -61,35 +56,15 @@ float3 StrataEvaluateMultipleScatteringConductor(float3 F0, FStrataLUTResult LUT
 	//  f_ms(wo,wi) = F_ss . (1 +  F0 (1-E_R)/ E_R)
 	
 	const float3 E_R = F0 * LUT.A_R + (1 - F0) * LUT.B_R;
-
 	return (1 + F0 * (1 - E_R) / E_R);
 }
 
-float3 StrataEvaluateMultipleScatteringDielectric(float3 F0, FStrataLUTResult LUT)
+float3 StrataEvaluateMultipleScatteringDielectric(float3 F0, FStrataGGXLUTResult LUT)
 {
-	// Reconstruction of the directional Albedo for upper hemisphere
-	// E_R = F0.A_R + (1-F0).B_R
-	//
-	// Reconstruction of the directional Albedo for lower hemisphere
-	// E_T = A_T + F0.A_T + (1-F0).B_T
-	//
 	// Microfacet multiple scattering term for *dielietric*
 	//  f_ms(wo,wi) = F_ss / (E_T + E_R)
 
-	const float3 E_R = F0 * LUT.A_R + (1 - F0) * LUT.B_R;
-	const float3 E_T = LUT.A_T + F0 * LUT.A_T + (1 - F0) * LUT.B_T;
-
-	return 1.f / (E_R + E_T);
-}
-
-float3 StrataEvaluateMultipleScatteringConductor3D(float3 F0, FStrataLUTResult LUT)
-{
-	return (1 + F0 * (1 - LUT.E_R) / LUT.E_R);
-}
-
-float3 StrataEvaluateMultipleScatteringDielectric3D(float3 F0, FStrataLUTResult LUT)
-{
 	return 1.f / (LUT.E_R + LUT.E_T);
 }
 
-////////////////////////////////////////////////////////////////////////////////////////////////////////////
\ No newline at end of file
+////////////////////////////////////////////////////////////////////////////////////////////////////////////
