* Restructure FStrataGGXLUTResult and add FStrataGGXDirectionalAlbedo for storing directional albedo for specular reflection/transmission (with MS) and diffuse
* Apply GGX specular MS on sky integrator
* Apply GGX diffuse directional albedo on sky integrator as a first order approximation.

diff --git a/Engine/Shaders/Private/Strata/StrataEnvironmentLighting.ush b/Engine/Shaders/Private/Strata/StrataEnvironmentLighting.ush
--- a/Engine/Shaders/Private/Strata/StrataEnvironmentLighting.ush
+++ b/Engine/Shaders/Private/Strata/StrataEnvironmentLighting.ush
@@ -38,11 +38,11 @@ float4 StrataEnvironmentLighting(
 	for (uint BSDFIndex = 0; BSDFIndex < StrataPixelHeader.BSDFCount; ++BSDFIndex)
 #endif
 	{
-#if STRATA_FASTPATH
+	#if STRATA_FASTPATH
 		FStrataBSDF BSDF = UnpackFastPathStrataBSDFIn(StrataDataBuffer, StrataAddressing);
-#else
+	#else
 		FStrataBSDF BSDF = UnpackStrataBSDFIn(StrataDataBuffer, StrataAddressing);
-#endif
+	#endif
 		const float3 BSDFThroughput = BSDF.Weight;
 
 		// Create the BSDF context
@@ -68,9 +68,10 @@ float4 StrataEnvironmentLighting(
 			FSkyLightVisibilityData SkyVisData = GetSkyLightVisibilityData(StrataBSDFContext.N, StrataBSDFContext.N, StrataPixelHeader.MaterialAO, AmbientOcclusion, StrataEnvLight.DiffuseNormal);
 
 			// Finally sample the sky diffuse contribution (spherical harmonic, Lambert BRDF)
-			float3 DiffuseLookup = GetSkySHDiffuse(StrataEnvLight.DiffuseNormal) * View.SkyLightColor.rgb;			// STRATA_TODO compute SH coefficients for ON or Chan
+			float3 DiffuseLookup = GetSkySHDiffuse(StrataEnvLight.DiffuseNormal) * View.SkyLightColor.rgb;		
 			// And accumulate
-			DiffuseSkyLighting = BSDFThroughput * (SkyVisData.SkyDiffuseLookUpMul * DiffuseLookup + SkyVisData.SkyDiffuseLookUpAdd) * StrataEnvLight.DiffuseColor;
+			// Note: Use diffuse directional albedo (i.e., DiffuseWeight) as first order approximation for env. integration (STRATA_TODO instead compute SH coefficients for Chan)
+			DiffuseSkyLighting = BSDFThroughput * (SkyVisData.SkyDiffuseLookUpMul * DiffuseLookup + SkyVisData.SkyDiffuseLookUpAdd) * StrataEnvLight.DiffuseWeight;
 		}
 
 		// Specular component
@@ -85,7 +86,7 @@ float4 StrataEnvironmentLighting(
 				(bIsTopLayer ? TopLayerSpecularContributionFactor : 1.0f);
 		}
 
-#if STRATA_FASTPATH==0
+	#if STRATA_FASTPATH==0
 		if (BSDF_GETHASHAZINESS(BSDF))
 		{
 			float SkyAverageBrightness = 1.0f;
@@ -97,7 +98,7 @@ float4 StrataEnvironmentLighting(
 			// And we also try to make the transition smooth using Haziness
 			SSRReductionFactor -= bIsTopLayer ? dot(BSDFThroughput, float3(1.0f / 3.0f, 1.0f / 3.0f, 1.0f / 3.0f)) * StrataEnvLight.SSRReduction : 0.0f;
 		}
-#endif 
+	#endif 
 
 		if (bIsTopLayer)
 		{
diff --git a/Engine/Shaders/Private/Strata/StrataEvaluation.ush b/Engine/Shaders/Private/Strata/StrataEvaluation.ush
--- a/Engine/Shaders/Private/Strata/StrataEvaluation.ush
+++ b/Engine/Shaders/Private/Strata/StrataEvaluation.ush
@@ -181,6 +181,15 @@ float StrataBSDFProjectedSolidAngleFactor(in FStrataBSDFContext BSDFContext)
 	return BSDFType == STRATA_BSDF_TYPE_HAIR ? 1.0f : BSDFContext.SatNoL;
 }
 
+#if STRATA_MICROFACET_MULTIPLESCATTERING
+
+FStrataGGXLUTResult StrataEvaluateGGXLUT(float3 DiffuseAlbedo, float3 F0, float NoV, float Metallic, float SafeRoughnessX, float SafeRoughnessY)
+{
+	return StrataSampleGGXLUT(Strata.GGXEnergyLUT2DTexture, Strata.GGXEnergyLUT3DTexture, Strata.GGXEnergyLUTSampler, Strata.GGXEnergyLUTScaleBias, NoV, SafeRoughnessX, SafeRoughnessY, F0, Metallic, DiffuseAlbedo);
+}
+
+#else // STRATA_MICROFACET_MULTIPLESCATTERING
+
 float3 StrataEnvBRDF(float3 F0, float3 F90, float SafeRoughnessX, float SatNoV)
 {
 #if STRATA_USE_ENV_BRDF_APPROX
@@ -190,23 +199,9 @@ float3 StrataEnvBRDF(float3 F0, float3 F90, float SafeRoughnessX, float SatNoV)
 #endif
 }
 
-float3 StrataEvaluateMicrofacetMultipleScattering(float3 F0, float NoV, float Metallic, float SafeRoughnessX, float SafeRoughnessY)
-{
-#if STRATA_MICROFACET_MULTIPLESCATTERING
-	// Use the main Roughness term to drive the energy compensation.
-	// "Revisiting Physically Based Shading at Imageworks" Kulla 2017
-	const float SafeRoughness = SafeRoughnessX;
+#endif // STRATA_MICROFACET_MULTIPLESCATTERING
 
-	const FStrataGGXLUTResult LUTSample = StrataSampleGGXLUT(Strata.GGXEnergyLUT2DTexture, Strata.GGXEnergyLUT3DTexture, Strata.GGXEnergyLUTSampler, Strata.GGXEnergyLUTScaleBias, NoV, SafeRoughness, F0);
-	const float3 Fs_0 = StrataEvaluateMultipleScatteringDielectric(F0, LUTSample);
-	const float3 Fs_1 = StrataEvaluateMultipleScatteringConductor(F0, LUTSample);
-	const float3 Fs   = lerp(Fs_0, Fs_1, Metallic);
 
-	return Fs;
-#else
-	return 1.0f;
-#endif
-}
 
 FStrataEvaluateResult StrataEvaluateBSDF(FStrataBSDFContext BSDFContext, FShadowTerms ShadowTerms)
 {
@@ -304,11 +299,11 @@ FStrataEvaluateResult StrataEvaluateBSDF(FStrataBSDFContext BSDFContext, FShadow
 			float3 HazeMSScale = 1;
 			#if STRATA_MICROFACET_MULTIPLESCATTERING
 			{
-				MSScale	= StrataEvaluateMicrofacetMultipleScattering(F0, BSDFContext.Context.NoV, Metallic, SafeRoughnessX, SafeRoughnessY);
+				MSScale	= StrataEvaluateGGXLUT(DiffuseColor, F0, BSDFContext.Context.NoV, Metallic, SafeRoughnessX, SafeRoughnessY).SpecularMultipleScatteringScale;
 				#if STRATA_FASTPATH==0
 				if (bHaziness)
 				{
-					HazeMSScale = StrataEvaluateMicrofacetMultipleScattering(F0, BSDFContext.Context.NoV, Metallic, HazeSafeRoughnessX, HazeSafeRoughnessY);
+					HazeMSScale = StrataEvaluateGGXLUT(DiffuseColor, F0, BSDFContext.Context.NoV, Metallic, HazeSafeRoughnessX, HazeSafeRoughnessY).SpecularMultipleScatteringScale;
 				}
 				#endif
 			}
@@ -542,6 +537,7 @@ struct FStrataEnvLightResult
 {
 	float3 DiffuseColor;
 	float3 DiffuseNormal;
+	float3 DiffuseWeight;
 
 	float3 SpecularDirection;
 	float3 SpecularWeight;
@@ -570,6 +566,9 @@ FStrataEnvLightResult StrataEvaluateForEnvLight(FStrataBSDFContext BSDFContext,
 			StrataEnvLightResult.DiffuseColor = BaseColor - BaseColor * Metallic;
 			StrataEnvLightResult.DiffuseNormal = BSDFContext.N;
 			StrataEnvLightResult.bSubsurface = BSDF_GETHASSSS(BSDFContext.BSDF) > 0;
+			StrataEnvLightResult.DiffuseWeight = 1;
+
+			float SafeRoughnessX = MakeRoughnessSafe(SLAB_ROUGHNESSX(BSDFContext.BSDF));
 
 			float3 FresnelWeight = 1;
 			BRANCH
@@ -589,7 +588,6 @@ FStrataEnvLightResult StrataEvaluateForEnvLight(FStrataBSDFContext BSDFContext,
 				}
 				#endif
 
-				float SafeRoughnessX = MakeRoughnessSafe(SLAB_ROUGHNESSX(BSDFContext.BSDF));
 
 				#if STRATA_FASTPATH==0
 				const bool bHasAnisotropy = BSDF_GETANISOTROPY(BSDFContext.BSDF);
@@ -609,7 +607,17 @@ FStrataEnvLightResult StrataEvaluateForEnvLight(FStrataBSDFContext BSDFContext,
 				}
 				#endif
 
-				float3 EvalEnvBRDF = FresnelWeight * StrataEnvBRDF(F0, F90, SafeRoughnessX, BSDFContext.SatNoV);
+				// Compute the directional albedo for specular and diffuse
+			#if STRATA_MICROFACET_MULTIPLESCATTERING
+				float3 EvalEnvBRDF = 0;
+				{
+					FStrataGGXLUTResult GGXLUT	= StrataEvaluateGGXLUT(StrataEnvLightResult.DiffuseColor, F0, BSDFContext.Context.NoV, Metallic, SafeRoughnessX, SafeRoughnessX);
+					EvalEnvBRDF					= FresnelWeight * GGXLUT.DirectionalAlbedo.SpecularR;
+					StrataEnvLightResult.DiffuseWeight = GGXLUT.DirectionalAlbedo.Diffuse;
+				}
+			#else
+				const float3 EvalEnvBRDF = FresnelWeight * StrataEnvBRDF(F0, F90, SafeRoughnessX, BSDFContext.SatNoV);
+			#endif
 
 				StrataEnvLightResult.SpecularDirection = StrataGetOffSpecularPeakReflectionDir(BSDFContext.N, BSDFContext.R, SafeRoughnessX);
 				StrataEnvLightResult.SpecularWeight = EvalEnvBRDF;
@@ -629,11 +637,28 @@ FStrataEnvLightResult StrataEvaluateForEnvLight(FStrataBSDFContext BSDFContext,
 
 					// Compute the second specular weight and roughness
 					float SafeRoughnessHaze = ComputeHazyLobeRoughness(SafeRoughnessX, Haziness);
-					StrataEnvLightResult.SpecularHazeWeight = lerp(0.0f, SPECULAR_HAZINESS_WEIGHT, HazinessFadeIn) * StrataEnvBRDF(F0, F90, SafeRoughnessHaze, BSDFContext.SatNoV);
+
+					// Compute the directional albedo for hazy specular
+				#if STRATA_MICROFACET_MULTIPLESCATTERING
+					FStrataGGXLUTResult GGXLUTHaze = StrataEvaluateGGXLUT(StrataEnvLightResult.DiffuseColor, F0, BSDFContext.Context.NoV, Metallic, SafeRoughnessHaze, SafeRoughnessHaze);
+					const float3 EvalEnvBRDFHaze = FresnelWeight * GGXLUTHaze.DirectionalAlbedo.SpecularR;
+				#else
+					const float3 EvalEnvBRDFHaze = FresnelWeight * StrataEnvBRDF(F0, F90, SafeRoughnessX, BSDFContext.SatNoV);
+				#endif
+
+					StrataEnvLightResult.SpecularHazeWeight = lerp(0.0f, SPECULAR_HAZINESS_WEIGHT, HazinessFadeIn) * EvalEnvBRDFHaze;
 					StrataEnvLightResult.SpecularHazeSafeRoughness = SafeRoughnessHaze;
 				}
 				#endif
 			}
+			#if STRATA_MICROFACET_MULTIPLESCATTERING
+			else
+			{
+				// Compute the directional albedo for diffuse
+				FStrataGGXLUTResult GGXLUT = StrataEvaluateGGXLUT(StrataEnvLightResult.DiffuseColor, 1.f, BSDFContext.Context.NoV, 0.f, SafeRoughnessX, SafeRoughnessX);
+				StrataEnvLightResult.DiffuseWeight = GGXLUT.DirectionalAlbedo.Diffuse;
+			}
+			#endif
 
 			#if STRATA_FASTPATH==0
 			BRANCH
diff --git a/Engine/Shaders/Private/Strata/StrataMultipleScattering.ush b/Engine/Shaders/Private/Strata/StrataMultipleScattering.ush
--- a/Engine/Shaders/Private/Strata/StrataMultipleScattering.ush
+++ b/Engine/Shaders/Private/Strata/StrataMultipleScattering.ush
@@ -7,23 +7,71 @@
 // Multiple scattering computation use a precomputed LUT. There are two variant: 
 //  1) Using 2D LUT using Schlick Fresnel decomposition (view angle, roughness)
 //  2) Using 3D LUT building a full 3D LUT (view angle, roughness, F0)
-// 
-// These two variants exist only for testing purpose at the moment
+
+struct FStrataGGXDirectionalAlbedo
+{
+	float3 SpecularR;	// Specular Directional albedo with multiple scattering (reflection)
+	float3 SpecularT;	// Specular Directional albedo with multiple scattering (transmission)
+	float3 Diffuse;		// Diffuse  Directional albedo with multiple scattering (embedded into Chan's model itself)
+};
 
 struct FStrataGGXLUTResult
 {
-	// Fresnel A/B term for reconstructing directional albedo for reflection
+	// Fresnel A/B term for reconstructing directional albedo for reflection (used to evaluate a chromatic F0). These values are dependent of F0.
 	float A_R;
 	float B_R;
 
-	// Complete directional albedo for reflection, refraction, and diffuse
+	// Complete directional albedo for reflection, refraction, and diffuse (used to evluate an achromatic F0). These values are independent of F0
 	float E_R;
 	float E_T;
 	float E_D;
+
+	// Derived terms
+	float3 SpecularMultipleScatteringScale;			// Multiple scattering scale factor for specular reflection
+	FStrataGGXDirectionalAlbedo DirectionalAlbedo;	// Directional albedo (with multiple scattering for reflection/transmission/diffuse
 };
 
-FStrataGGXLUTResult StrataSampleGGXLUT(Texture2D<float4> LUT2D, Texture3D<float2> LUT3D, SamplerState LUTSampler, float2 LUTScaleBias, float NoV, float Roughness, float3 F0)
+// Return a scaling factor to multiple a specular GGX interface with, to add multiple scattering contribution (for conductor)
+// From "Practical multiple scattering compensation for microfacet models." E. Turquin
+float3 StrataEvaluateMultipleScatteringConductor(float3 F0, FStrataGGXLUTResult LUT)
 {
+	// Reconstruction of the directional Albedo for upper hemisphere
+	// E_R = F0.A_R + (1-F0).B_R
+	// 
+	// Microfacet multiple scattering term for *conductor*
+	//  f_ms(wo,wi) = F_ss . (1 +  F0 (1-E_R)/ E_R)
+	
+	const float3 E_R = F0 * LUT.A_R + (1 - F0) * LUT.B_R;
+	return (1 + F0 * (1 - E_R) / E_R);
+}
+
+// Return a scaling factor to multiple a specular GGX interface with, to add multiple scattering contribution (for dieletric)
+float3 StrataEvaluateMultipleScatteringDielectric(float3 F0, FStrataGGXLUTResult LUT)
+{
+	// Microfacet multiple scattering term for *dielietric*
+	//  f_ms(wo,wi) = F_ss / (E_T + E_R)
+
+	return 1.f / (LUT.E_R + LUT.E_T);
+}
+
+// Return the directional albedo of a specular GGX interface for reflection term without multiple scattering
+float3 StrataEvaluateSpecularGGXDirectionalAlbedo(float3 F0, FStrataGGXLUTResult LUT)
+{
+	return F0 * LUT.A_R + (1 - F0) * LUT.B_R;
+}
+
+// Return the directional alebedo of a diffuse GGX interface without multiple scattering
+float3 StrataEvaluateDiffuseGGXDirectionalAlbedo(float3 Albedo, FStrataGGXLUTResult LUT)
+{
+	return Albedo * LUT.E_D;
+}
+
+FStrataGGXLUTResult StrataSampleGGXLUT(Texture2D<float4> LUT2D, Texture3D<float2> LUT3D, SamplerState LUTSampler, float2 LUTScaleBias, float NoV, float RoughnessX, float RoughnessY, float3 F0, float Metallic, float3 DiffuseAlbedo)
+{
+	// Use the main Roughness term to drive the energy compensation.
+	// "Revisiting Physically Based Shading at Imageworks" Kulla 2017
+	const float Roughness = RoughnessX;
+
 	FStrataGGXLUTResult Out = (FStrataGGXLUTResult)0;
 
 	// 2D LUT sampling
@@ -43,28 +91,18 @@ FStrataGGXLUTResult StrataSampleGGXLUT(Texture2D<float4> LUT2D, Texture3D<float2
 		Out.E_T = S.y;
 	}
 
-	return Out;
-}
-
-// From "Practical multiple scattering compensation for microfacet models." E. Turquin
-float3 StrataEvaluateMultipleScatteringConductor(float3 F0, FStrataGGXLUTResult LUT)
-{
-	// Reconstruction of the directional Albedo for upper hemisphere
-	// E_R = F0.A_R + (1-F0).B_R
-	// 
-	// Microfacet multiple scattering term for *conductor*
-	//  f_ms(wo,wi) = F_ss . (1 +  F0 (1-E_R)/ E_R)
-	
-	const float3 E_R = F0 * LUT.A_R + (1 - F0) * LUT.B_R;
-	return (1 + F0 * (1 - E_R) / E_R);
-}
-
-float3 StrataEvaluateMultipleScatteringDielectric(float3 F0, FStrataGGXLUTResult LUT)
-{
-	// Microfacet multiple scattering term for *dielietric*
-	//  f_ms(wo,wi) = F_ss / (E_T + E_R)
+	// Derived MS scales
+	{
+		const float3 DA_ConductorSS = StrataEvaluateSpecularGGXDirectionalAlbedo(F0, Out);
 
-	return 1.f / (LUT.E_R + LUT.E_T);
-}
+		const float3 MSScale_Dieletric = StrataEvaluateMultipleScatteringDielectric(F0, Out);
+		const float3 MSScale_Conductor = StrataEvaluateMultipleScatteringConductor(F0, Out);
+		const float3 MSScale		   = lerp(MSScale_Dieletric, MSScale_Conductor, Metallic);
 
-////////////////////////////////////////////////////////////////////////////////////////////////////////////
+		Out.SpecularMultipleScatteringScale = MSScale;
+		Out.DirectionalAlbedo.SpecularR		= MSScale * lerp(Out.E_R.xxx, DA_ConductorSS, Metallic);
+		Out.DirectionalAlbedo.SpecularT		= MSScale * lerp(Out.E_T, 0.f, Metallic);
+		Out.DirectionalAlbedo.Diffuse		= StrataEvaluateDiffuseGGXDirectionalAlbedo(DiffuseAlbedo, Out);
+	}
+	return Out;
+}
\ No newline at end of file
