Fix strata specular PDF within the Evaluate function.

diff --git a/Engine/Shaders/Private/Strata/StrataEvaluation.ush b/Engine/Shaders/Private/Strata/StrataEvaluation.ush
--- a/Engine/Shaders/Private/Strata/StrataEvaluation.ush
+++ b/Engine/Shaders/Private/Strata/StrataEvaluation.ush
@@ -253,7 +253,9 @@ FStrataEvaluateResult StrataEvaluateBSDF(FStrataBSDFContext BSDFContext, FShadow
 					const float Alpha1 = Pow2(SafeRoughnessY);
 					D   = D_GGXaniso(Alpha0, Alpha1, BSDFContext.Context.NoH, BSDFContext.Context.XoH, BSDFContext.Context.YoH);
 					Vis = Vis_SmithJointAniso(Alpha0, Alpha1, BSDFContext.Context.NoV, BSDFContext.Context.NoL, BSDFContext.Context.XoV, BSDFContext.Context.XoL, BSDFContext.Context.YoV, BSDFContext.Context.YoL);
-					PDF = VisibleGGXPDF_aniso(BSDFContext.TangentV, BSDFContext.TangentH, float2(Alpha0, Alpha1));
+					const float H_PDF = VisibleGGXPDF_aniso(BSDFContext.TangentV, BSDFContext.TangentH, float2(Alpha0, Alpha1));
+					PDF = RayPDFToReflectionRayPDF(BSDFContext.Context.VoH, H_PDF);
+
 				}
 
 				if (bHaziness)
@@ -265,7 +267,8 @@ FStrataEvaluateResult StrataEvaluateBSDF(FStrataBSDFContext BSDFContext, FShadow
 					const float HazeAlpha1 = Pow2(HazeSafeRoughnessY);
 					HazeD   = D_GGXaniso(HazeAlpha0, HazeAlpha1, BSDFContext.Context.NoH, BSDFContext.Context.XoH, BSDFContext.Context.YoH);
 					HazeVis = Vis_SmithJointAniso(HazeAlpha0, HazeAlpha1, BSDFContext.Context.NoV, BSDFContext.Context.NoL, BSDFContext.Context.XoV, BSDFContext.Context.XoL, BSDFContext.Context.YoV, BSDFContext.Context.YoL);
-					HazePDF = VisibleGGXPDF_aniso(BSDFContext.TangentV, BSDFContext.TangentH, float2(HazeAlpha0, HazeAlpha1));
+					const float H_PDF = VisibleGGXPDF_aniso(BSDFContext.TangentV, BSDFContext.TangentH, float2(HazeAlpha0, HazeAlpha1));
+					HazePDF = RayPDFToReflectionRayPDF(BSDFContext.Context.VoH, H_PDF);
 				}
 			}
 			else
@@ -278,7 +281,8 @@ FStrataEvaluateResult StrataEvaluateBSDF(FStrataBSDFContext BSDFContext, FShadow
 				{
 					D   = D_GGX(a2, BSDFContext.Context.NoH);
 					Vis = Vis_SmithJoint(a2, BSDFContext.Context.NoV, BSDFContext.SatNoL);
-					PDF = VisibleGGXPDF(BSDFContext.TangentV, BSDFContext.TangentH, a2);
+					const float H_PDF = VisibleGGXPDF(BSDFContext.TangentV, BSDFContext.TangentH, a2);
+					PDF = RayPDFToReflectionRayPDF(BSDFContext.Context.VoH, H_PDF);
 				}
 
 				#if STRATA_FASTPATH==0
@@ -289,7 +293,8 @@ FStrataEvaluateResult StrataEvaluateBSDF(FStrataBSDFContext BSDFContext, FShadow
 					const float a2Haze = BSDFContext.a2OverrideHazy > 0 ? BSDFContext.a2OverrideHazy : Pow4(HazeSafeRoughnessX);
 					HazeD   = D_GGX(a2Haze, BSDFContext.Context.NoH);
 					HazeVis = Vis_SmithJoint(a2Haze, BSDFContext.Context.NoV, BSDFContext.SatNoL);
-					HazePDF = VisibleGGXPDF(BSDFContext.TangentV, BSDFContext.TangentH, a2Haze);
+					const float H_PDF = VisibleGGXPDF(BSDFContext.TangentV, BSDFContext.TangentH, a2Haze);
+					HazePDF = RayPDFToReflectionRayPDF(BSDFContext.Context.VoH, H_PDF);
 
 				}
 				#endif
