Fix skylighting missing albedo on Strata surfaces.

diff --git a/Engine/Shaders/Private/Strata/StrataEvaluation.ush b/Engine/Shaders/Private/Strata/StrataEvaluation.ush
--- a/Engine/Shaders/Private/Strata/StrataEvaluation.ush
+++ b/Engine/Shaders/Private/Strata/StrataEvaluation.ush
@@ -571,7 +571,7 @@ FStrataEnvLightResult StrataEvaluateForEnvLight(FStrataBSDFContext BSDFContext,
 			StrataEnvLightResult.DiffuseColor = BaseColor - BaseColor * Metallic;
 			StrataEnvLightResult.DiffuseNormal = BSDFContext.N;
 			StrataEnvLightResult.bSubsurface = BSDF_GETHASSSS(BSDFContext.BSDF) > 0;
-			StrataEnvLightResult.DiffuseWeight = 1;
+			StrataEnvLightResult.DiffuseWeight = StrataEnvLightResult.DiffuseColor;
 
 			float SafeRoughnessX = MakeRoughnessSafe(SLAB_ROUGHNESSX(BSDFContext.BSDF));
 
