Use int4 array in constant buffer for PhysicsFieldTargets to be compatbile with std140 layout for Vulkan RHI.

#rnx

diff --git a/Engine/Shaders/Private/PhysicsFieldSampler.ush b/Engine/Shaders/Private/PhysicsFieldSampler.ush
--- a/Engine/Shaders/Private/PhysicsFieldSampler.ush
+++ b/Engine/Shaders/Private/PhysicsFieldSampler.ush
@@ -19,9 +19,7 @@
 	#define PhysicsFieldClipmapExponent			View.PhysicsFieldClipmapExponent
 	#define PhysicsFieldClipmapCount			View.PhysicsFieldClipmapCount
 	#define PhysicsFieldTargetCount				View.PhysicsFieldTargetCount
-	#define PhysicsFieldVectorTargets			View.PhysicsFieldVectorTargets
-	#define PhysicsFieldScalarTargets			View.PhysicsFieldScalarTargets
-	#define PhysicsFieldIntegerTargets			View.PhysicsFieldIntegerTargets
+	#define PhysicsFieldTargets					View.PhysicsFieldTargets
 #else
 	Buffer<float> PhysicsFieldClipmapBuffer;
 	float3		PhysicsFieldClipmapCenter;
@@ -30,26 +28,24 @@
 	int			PhysicsFieldClipmapExponent;
 	int			PhysicsFieldClipmapCount;
 	int			PhysicsFieldTargetCount;
-	int			PhysicsFieldVectorTargets[MAX_PHYSICS_FIELD_TARGETS];
-	int			PhysicsFieldScalarTargets[MAX_PHYSICS_FIELD_TARGETS];
-	int			PhysicsFieldIntegerTargets[MAX_PHYSICS_FIELD_TARGETS];
+	int4		PhysicsFieldTargets[MAX_PHYSICS_FIELD_TARGETS]; // x (Vector), y (Scalar), z (Integer), w (Padding)
 #endif
 
 float3 MatPhysicsField_SamplePhysicsVectorField(float3 WorldPosition, int VectorTarget)
 {
-	return PhysicsField_SamplePhysicsVectorField(WorldPosition, VectorTarget, PhysicsFieldVectorTargets, PhysicsFieldTargetCount, PhysicsFieldClipmapCenter, PhysicsFieldClipmapDistance,
+	return PhysicsField_SamplePhysicsVectorField(WorldPosition, VectorTarget, PhysicsFieldTargets, PhysicsFieldTargetCount, PhysicsFieldClipmapCenter, PhysicsFieldClipmapDistance,
 				PhysicsFieldClipmapExponent, PhysicsFieldClipmapCount, PhysicsFieldClipmapResolution, PhysicsFieldClipmapBuffer);
 }
 
 float MatPhysicsField_SamplePhysicsScalarField(float3 WorldPosition, int ScalarTarget)
 {
-	return PhysicsField_SamplePhysicsScalarField(WorldPosition, ScalarTarget, PhysicsFieldScalarTargets, PhysicsFieldTargetCount, PhysicsFieldClipmapCenter, PhysicsFieldClipmapDistance,
+	return PhysicsField_SamplePhysicsScalarField(WorldPosition, ScalarTarget, PhysicsFieldTargets, PhysicsFieldTargetCount, PhysicsFieldClipmapCenter, PhysicsFieldClipmapDistance,
 				PhysicsFieldClipmapExponent, PhysicsFieldClipmapCount, PhysicsFieldClipmapResolution, PhysicsFieldClipmapBuffer);
 }
 
 int MatPhysicsField_SamplePhysicsIntegerField(float3 WorldPosition, int IntegerTarget)
 {
-	return PhysicsField_SamplePhysicsIntegerField(WorldPosition, IntegerTarget, PhysicsFieldIntegerTargets, PhysicsFieldTargetCount, PhysicsFieldClipmapCenter, PhysicsFieldClipmapDistance,
+	return PhysicsField_SamplePhysicsIntegerField(WorldPosition, IntegerTarget, PhysicsFieldTargets, PhysicsFieldTargetCount, PhysicsFieldClipmapCenter, PhysicsFieldClipmapDistance,
 				PhysicsFieldClipmapExponent, PhysicsFieldClipmapCount, PhysicsFieldClipmapResolution, PhysicsFieldClipmapBuffer);
 	
 }
diff --git a/Engine/Shaders/Private/PhysicsFieldShared.ush b/Engine/Shaders/Private/PhysicsFieldShared.ush
--- a/Engine/Shaders/Private/PhysicsFieldShared.ush
+++ b/Engine/Shaders/Private/PhysicsFieldShared.ush
@@ -124,7 +124,7 @@ void PhysicsField_GetSamplingInfos(in float3 LocalPosition, in int TargetIndex,
 	SampleFraction = SampleVoxel - SampleIndex;
 }
 
-float3 PhysicsField_SamplePhysicsVectorField(in float3 WorldPosition, in int VectorTarget, in int VectorTargets[MAX_PHYSICS_FIELD_TARGETS], in int TargetCount, in float3 ClipmapCenter,
+float3 PhysicsField_SamplePhysicsVectorField(in float3 WorldPosition, in int VectorTarget, in int4 VectorTargets[MAX_PHYSICS_FIELD_TARGETS], in int TargetCount, in float3 ClipmapCenter,
 			in float ClipmapDistance, in int ClipmapExponent, in int ClipmapCount, in int ClipmapResolution, in Buffer<float> ClipmapBuffer)
 {
 	const float3 LocalPosition = WorldPosition - ClipmapCenter;
@@ -132,7 +132,7 @@ float3 PhysicsField_SamplePhysicsVectorField(in float3 WorldPosition, in int Vec
 	float3 FieldValue = float3(0, 0, 0);
 	if (VectorTarget != -1 && VectorTarget < MAX_PHYSICS_FIELD_TARGETS && all(abs(LocalPosition) < float3(ClipmapDistance, ClipmapDistance, ClipmapDistance)))
 	{
-		const int TargetIndex = VectorTargets[VectorTarget];
+		const int TargetIndex = VectorTargets[VectorTarget].x;
 		if (TargetIndex != -1)
 		{
 			int3 SampleIndex = int3(0, 0, 0);
@@ -145,7 +145,7 @@ float3 PhysicsField_SamplePhysicsVectorField(in float3 WorldPosition, in int Vec
 	return FieldValue;
 }
 
-float PhysicsField_SamplePhysicsScalarField(in float3 WorldPosition, in int ScalarTarget, in int ScalarTargets[MAX_PHYSICS_FIELD_TARGETS], in int TargetCount, in float3 ClipmapCenter,
+float PhysicsField_SamplePhysicsScalarField(in float3 WorldPosition, in int ScalarTarget, in int4 ScalarTargets[MAX_PHYSICS_FIELD_TARGETS], in int TargetCount, in float3 ClipmapCenter,
 			in float ClipmapDistance, in int ClipmapExponent, in int ClipmapCount, in int ClipmapResolution, in Buffer<float> ClipmapBuffer)
 {
 	const float3 LocalPosition = WorldPosition - ClipmapCenter;
@@ -153,7 +153,7 @@ float PhysicsField_SamplePhysicsScalarField(in float3 WorldPosition, in int Scal
 	float FieldValue = 0.0;
 	if (ScalarTarget != -1 && ScalarTarget < MAX_PHYSICS_FIELD_TARGETS && all(abs(LocalPosition) < float3(ClipmapDistance, ClipmapDistance, ClipmapDistance)))
 	{
-		const int TargetIndex = ScalarTargets[ScalarTarget];
+		const int TargetIndex = ScalarTargets[ScalarTarget].y;
 		if (TargetIndex != -1)
 		{
 			int3 SampleIndex = int3(0, 0, 0);
@@ -166,7 +166,7 @@ float PhysicsField_SamplePhysicsScalarField(in float3 WorldPosition, in int Scal
 	return FieldValue;
 }
 
-int PhysicsField_SamplePhysicsIntegerField(in float3 WorldPosition, in int IntegerTarget, in int IntegerTargets[MAX_PHYSICS_FIELD_TARGETS], in int TargetCount, in float3 ClipmapCenter,
+int PhysicsField_SamplePhysicsIntegerField(in float3 WorldPosition, in int IntegerTarget, in int4 IntegerTargets[MAX_PHYSICS_FIELD_TARGETS], in int TargetCount, in float3 ClipmapCenter,
 			in float ClipmapDistance, in int ClipmapExponent, in int ClipmapCount, in int ClipmapResolution, in Buffer<float> ClipmapBuffer)
 {
 	const float3 LocalPosition = WorldPosition - ClipmapCenter;
@@ -174,7 +174,7 @@ int PhysicsField_SamplePhysicsIntegerField(in float3 WorldPosition, in int Integ
 	int FieldValue = 0.0;
 	if (IntegerTarget != -1 && IntegerTarget < MAX_PHYSICS_FIELD_TARGETS && all(abs(LocalPosition) < float3(ClipmapDistance, ClipmapDistance, ClipmapDistance)))
 	{
-		const int TargetIndex = IntegerTargets[IntegerTarget];
+		const int TargetIndex = IntegerTargets[IntegerTarget].z;
 		if (TargetIndex != -1)
 		{
 			int3 SampleIndex = int3(0, 0, 0);
diff --git a/Engine/Source/Runtime/Engine/Public/SceneView.h b/Engine/Source/Runtime/Engine/Public/SceneView.h
--- a/Engine/Source/Runtime/Engine/Public/SceneView.h
+++ b/Engine/Source/Runtime/Engine/Public/SceneView.h
@@ -782,9 +782,7 @@ enum ETranslucencyVolumeCascade
 	VIEW_UNIFORM_BUFFER_MEMBER(int, PhysicsFieldClipmapExponent) \
 	VIEW_UNIFORM_BUFFER_MEMBER(int, PhysicsFieldClipmapCount) \
 	VIEW_UNIFORM_BUFFER_MEMBER(int, PhysicsFieldTargetCount) \
-	VIEW_UNIFORM_BUFFER_MEMBER_ARRAY(int, PhysicsFieldVectorTargets, [MAX_PHYSICS_FIELD_TARGETS]) \
-	VIEW_UNIFORM_BUFFER_MEMBER_ARRAY(int, PhysicsFieldScalarTargets, [MAX_PHYSICS_FIELD_TARGETS]) \
-	VIEW_UNIFORM_BUFFER_MEMBER_ARRAY(int, PhysicsFieldIntegerTargets, [MAX_PHYSICS_FIELD_TARGETS]) \
+	VIEW_UNIFORM_BUFFER_MEMBER_ARRAY(FIntVector4, PhysicsFieldTargets, [MAX_PHYSICS_FIELD_TARGETS]) \
 	VIEW_UNIFORM_BUFFER_MEMBER(uint32, InstanceDataSOAStride) \
 	VIEW_UNIFORM_BUFFER_MEMBER(uint32, GPUSceneViewId) \
 
diff --git a/Engine/Source/Runtime/Renderer/Private/SceneRendering.cpp b/Engine/Source/Runtime/Renderer/Private/SceneRendering.cpp
--- a/Engine/Source/Runtime/Renderer/Private/SceneRendering.cpp
+++ b/Engine/Source/Runtime/Renderer/Private/SceneRendering.cpp
@@ -1132,13 +1132,17 @@ void SetupPhysicsFieldUniformBufferParameters(const FScene* Scene, FEngineShowFl
 		ViewUniformShaderParameters.PhysicsFieldClipmapExponent = FieldResource->FieldInfos.ClipmapExponent;
 		ViewUniformShaderParameters.PhysicsFieldClipmapCount = FieldResource->FieldInfos.ClipmapCount;
 		ViewUniformShaderParameters.PhysicsFieldTargetCount = FieldResource->FieldInfos.TargetCount;
-		ViewUniformShaderParameters.PhysicsFieldVectorTargets = FieldResource->FieldInfos.VectorTargets;
-		ViewUniformShaderParameters.PhysicsFieldScalarTargets = FieldResource->FieldInfos.ScalarTargets;
-		ViewUniformShaderParameters.PhysicsFieldIntegerTargets = FieldResource->FieldInfos.IntegerTargets;
+		for (int32 Index = 0; Index < MAX_PHYSICS_FIELD_TARGETS; ++Index)
+		{
+			ViewUniformShaderParameters.PhysicsFieldTargets[Index].X = FieldResource->FieldInfos.VectorTargets[Index];
+			ViewUniformShaderParameters.PhysicsFieldTargets[Index].Y = FieldResource->FieldInfos.ScalarTargets[Index];
+			ViewUniformShaderParameters.PhysicsFieldTargets[Index].Z = FieldResource->FieldInfos.IntegerTargets[Index];
+			ViewUniformShaderParameters.PhysicsFieldTargets[Index].W = 0; // Padding
+		}
 	}
 	else
 	{
-		TStaticArray<int32, MAX_PHYSICS_FIELD_TARGETS, 16> EmptyTargets;
+		TStaticArray<FIntVector4, MAX_PHYSICS_FIELD_TARGETS, 16> EmptyTargets;
 		ViewUniformShaderParameters.PhysicsFieldClipmapBuffer = GWhiteVertexBufferWithSRV->ShaderResourceViewRHI;
 		ViewUniformShaderParameters.PhysicsFieldClipmapCenter = FVector::ZeroVector;
 		ViewUniformShaderParameters.PhysicsFieldClipmapDistance = 1.0;
@@ -1146,9 +1150,7 @@ void SetupPhysicsFieldUniformBufferParameters(const FScene* Scene, FEngineShowFl
 		ViewUniformShaderParameters.PhysicsFieldClipmapExponent = 1;
 		ViewUniformShaderParameters.PhysicsFieldClipmapCount = 1;
 		ViewUniformShaderParameters.PhysicsFieldTargetCount = 0;
-		ViewUniformShaderParameters.PhysicsFieldVectorTargets = EmptyTargets;
-		ViewUniformShaderParameters.PhysicsFieldScalarTargets = EmptyTargets;
-		ViewUniformShaderParameters.PhysicsFieldIntegerTargets = EmptyTargets;
+		ViewUniformShaderParameters.PhysicsFieldTargets = EmptyTargets;
 	}
 }
 
