Use binary search instead of linear search for light CDF sampling in the path tracer

This is always slightly faster in the scenes I have tried. On the infiltrator scene it was up to 12% faster overall

diff --git a/Engine/Shaders/Private/PathTracing/Light/PathTracingLightSampling.ush b/Engine/Shaders/Private/PathTracing/Light/PathTracingLightSampling.ush
--- a/Engine/Shaders/Private/PathTracing/Light/PathTracingLightSampling.ush
+++ b/Engine/Shaders/Private/PathTracing/Light/PathTracingLightSampling.ush
@@ -130,8 +130,8 @@ void SelectLight(float RandSample, inout float LightPickingCdf[RAY_TRACING_LIGHT
 {
 	// This code assumes SceneLightCount > 0
 
-	// Draw random variable
-	// #dxr_todo: binary search
+#if 0
+	// linear search
 	float PreviousCdfValue = 0;
 	for (LightId = 0; LightId < SceneLightCount - 1; ++LightId)
 	{
@@ -143,6 +143,27 @@ void SelectLight(float RandSample, inout float LightPickingCdf[RAY_TRACING_LIGHT
 	}
 
 	LightPickPdf = LightPickingCdf[LightId] - PreviousCdfValue;
+#else
+	// binary search
+	LightId = 0;
+	for (int Count = SceneLightCount; Count > 0;)
+	{
+		int Step = Count / 2;
+		int Iterator = LightId + Step;
+		if (RandSample < LightPickingCdf[Iterator])
+		{
+			Count = Step;
+		}
+		else
+		{
+			LightId = Iterator + 1;
+			Count -= Step + 1;
+		}
+	}
+
+	LightPickPdf = LightPickingCdf[LightId] - ((LightId > 0) ? LightPickingCdf[LightId - 1] : 0.0);
+#endif
+
 }
 
 bool SampleLight(
