Fix sky integrator with Strata with other material than simple slab

diff --git a/Engine/Shaders/Private/Strata/StrataEnvironmentLighting.ush b/Engine/Shaders/Private/Strata/StrataEnvironmentLighting.ush
--- a/Engine/Shaders/Private/Strata/StrataEnvironmentLighting.ush
+++ b/Engine/Shaders/Private/Strata/StrataEnvironmentLighting.ush
@@ -55,7 +55,7 @@ float4 StrataEnvironmentLighting(
 
 		// Diffuse component
 		float3 DiffuseSkyLighting = 0;
-		if (any(StrataEnvLight.DiffuseColor > 0.0f))
+		if (any(StrataEnvLight.DiffuseWeight > 0.0f))
 		{
 			// Apply bent normal if needed
 			#if APPLY_SKY_SHADOWING
diff --git a/Engine/Shaders/Private/Strata/StrataEvaluation.ush b/Engine/Shaders/Private/Strata/StrataEvaluation.ush
--- a/Engine/Shaders/Private/Strata/StrataEvaluation.ush
+++ b/Engine/Shaders/Private/Strata/StrataEvaluation.ush
@@ -540,7 +540,6 @@ FStrataImportanceSampleResult StrataImportanceSampleBSDF(FStrataPixelHeader Stra
 
 struct FStrataEnvLightResult
 {
-	float3 DiffuseColor;
 	float3 DiffuseNormal;
 	float3 DiffuseWeight;
 
@@ -568,10 +567,10 @@ FStrataEnvLightResult StrataEvaluateForEnvLight(FStrataBSDFContext BSDFContext,
 			const float3 BaseColor = SLAB_BASECOLOR(BSDFContext.BSDF);
 			const float  Metallic = SLAB_METALLIC(BSDFContext.BSDF);
 
-			StrataEnvLightResult.DiffuseColor = BaseColor - BaseColor * Metallic;
+			const float3 DiffuseColor = BaseColor - BaseColor * Metallic;
 			StrataEnvLightResult.DiffuseNormal = BSDFContext.N;
 			StrataEnvLightResult.bSubsurface = BSDF_GETHASSSS(BSDFContext.BSDF) > 0;
-			StrataEnvLightResult.DiffuseWeight = StrataEnvLightResult.DiffuseColor;
+			StrataEnvLightResult.DiffuseWeight = DiffuseColor;
 
 			float SafeRoughnessX = MakeRoughnessSafe(SLAB_ROUGHNESSX(BSDFContext.BSDF));
 
@@ -616,7 +615,7 @@ FStrataEnvLightResult StrataEvaluateForEnvLight(FStrataBSDFContext BSDFContext,
 			#if STRATA_MICROFACET_MULTIPLESCATTERING
 				float3 EvalEnvBRDF = 0;
 				{
-					FStrataGGXLUTResult GGXLUT	= StrataEvaluateGGXLUT(StrataEnvLightResult.DiffuseColor, F0, BSDFContext.Context.NoV, Metallic, SafeRoughnessX, SafeRoughnessX);
+					FStrataGGXLUTResult GGXLUT	= StrataEvaluateGGXLUT(DiffuseColor, F0, BSDFContext.Context.NoV, Metallic, SafeRoughnessX, SafeRoughnessX);
 					EvalEnvBRDF					= FresnelWeight * GGXLUT.DirectionalAlbedo.SpecularR;
 					StrataEnvLightResult.DiffuseWeight = GGXLUT.DirectionalAlbedo.Diffuse;
 				}
@@ -645,7 +644,7 @@ FStrataEnvLightResult StrataEvaluateForEnvLight(FStrataBSDFContext BSDFContext,
 
 					// Compute the directional albedo for hazy specular
 				#if STRATA_MICROFACET_MULTIPLESCATTERING
-					FStrataGGXLUTResult GGXLUTHaze = StrataEvaluateGGXLUT(StrataEnvLightResult.DiffuseColor, F0, BSDFContext.Context.NoV, Metallic, SafeRoughnessHaze, SafeRoughnessHaze);
+					FStrataGGXLUTResult GGXLUTHaze = StrataEvaluateGGXLUT(DiffuseColor, F0, BSDFContext.Context.NoV, Metallic, SafeRoughnessHaze, SafeRoughnessHaze);
 					const float3 EvalEnvBRDFHaze = FresnelWeight * GGXLUTHaze.DirectionalAlbedo.SpecularR;
 				#else
 					const float3 EvalEnvBRDFHaze = FresnelWeight * StrataEnvBRDF(F0, F90, SafeRoughnessX, BSDFContext.SatNoV);
@@ -660,7 +659,7 @@ FStrataEnvLightResult StrataEvaluateForEnvLight(FStrataBSDFContext BSDFContext,
 			else
 			{
 				// Compute the directional albedo for diffuse
-				FStrataGGXLUTResult GGXLUT = StrataEvaluateGGXLUT(StrataEnvLightResult.DiffuseColor, 1.f, BSDFContext.Context.NoV, 0.f, SafeRoughnessX, SafeRoughnessX);
+				FStrataGGXLUTResult GGXLUT = StrataEvaluateGGXLUT(DiffuseColor, 1.f, BSDFContext.Context.NoV, 0.f, SafeRoughnessX, SafeRoughnessX);
 				StrataEnvLightResult.DiffuseWeight = GGXLUT.DirectionalAlbedo.Diffuse;
 			}
 			#endif
@@ -669,13 +668,13 @@ FStrataEnvLightResult StrataEvaluateForEnvLight(FStrataBSDFContext BSDFContext,
 			BRANCH
 			if (BSDF_GETISSIMPLEVOLUME(BSDFContext.BSDF))
 			{
-				FParticipatingMedia PM = StrataSlabCreateParticipatingMedia(StrataEnvLightResult.DiffuseColor, SLAB_SSSDMFP(BSDFContext.BSDF));
+				FParticipatingMedia PM = StrataSlabCreateParticipatingMedia(DiffuseColor, SLAB_SSSDMFP(BSDFContext.BSDF));
 				const float DiffuseToVolumeBlend = StrataSlabDiffuseToVolumeBlend(PM);
 				const FStrataSlabSingleScattering Result = StrataSlabSingleScattering(PM, BSDFContext.Context.NoV);
 
 				// STRATA_TODO: big approximation for volumetric + env lighting: We ignore the phase for now due to the sampled convolved env lighting, 
 				// no self shadow and should not be single scattering using StrataSlabSingleScattering. We will propose a better approach. 
-				StrataEnvLightResult.DiffuseColor = lerp(StrataEnvLightResult.DiffuseColor, Result.SingleScattering * FresnelWeight,	DiffuseToVolumeBlend);
+				StrataEnvLightResult.DiffuseWeight = lerp(StrataEnvLightResult.DiffuseWeight, Result.SingleScattering * FresnelWeight,	DiffuseToVolumeBlend);
 				// We keep the same normal
 				StrataEnvLightResult.bSubsurface = false;
 			}
@@ -686,7 +685,7 @@ FStrataEnvLightResult StrataEvaluateForEnvLight(FStrataBSDFContext BSDFContext,
 #if STRATA_FASTPATH==0
 		case STRATA_BSDF_TYPE_SHEEN:
 		{
-			StrataEnvLightResult.DiffuseColor = SHEEN_BASECOLOR(BSDFContext.BSDF);;
+			StrataEnvLightResult.DiffuseWeight = SHEEN_BASECOLOR(BSDFContext.BSDF);
 			StrataEnvLightResult.DiffuseNormal = BSDFContext.N;
 			break;
 		}
@@ -713,7 +712,7 @@ FStrataEnvLightResult StrataEvaluateForEnvLight(FStrataBSDFContext BSDFContext,
 			float3 L = normalize(V - N * dot(V, N));
 			StrataEnvLightResult.DiffuseNormal = L;
 
-			StrataEnvLightResult.DiffuseColor = PI * HairShading(GBuffer, L, V, N, TransmissionShadow, TransmittanceData, BacklitEnabled, Area, Random);
+			StrataEnvLightResult.DiffuseWeight = PI * HairShading(GBuffer, L, V, N, TransmissionShadow, TransmittanceData, BacklitEnabled, Area, Random);
 			// No specular environment contribution as of today if not using the special HairStrand render path
 			break;
 		}
diff --git a/Engine/Shaders/Private/Strata/StrataFurnaceTest.usf b/Engine/Shaders/Private/Strata/StrataFurnaceTest.usf
--- a/Engine/Shaders/Private/Strata/StrataFurnaceTest.usf
+++ b/Engine/Shaders/Private/Strata/StrataFurnaceTest.usf
@@ -232,7 +232,7 @@ void MainPS(
 
 			const float3 InIrradiance = InRadiance; //It should have a Pi factor though
 			const float3 SpecularSkyLighting = BSDF.Weight * StrataEnvLight.SpecularWeight * InRadiance;
-			const float3 DiffuseSkyLighting  = BSDF.Weight * StrataEnvLight.DiffuseColor * InIrradiance;
+			const float3 DiffuseSkyLighting  = BSDF.Weight * StrataEnvLight.DiffuseWeight * InIrradiance;
 
 			OutColor = 
 				SpecularSkyLighting + 
