Fixed Strata forward conpilation error

diff --git a/Engine/Shaders/Private/Strata/StrataForwardLighting.ush b/Engine/Shaders/Private/Strata/StrataForwardLighting.ush
--- a/Engine/Shaders/Private/Strata/StrataForwardLighting.ush
+++ b/Engine/Shaders/Private/Strata/StrataForwardLighting.ush
@@ -66,9 +66,9 @@ float3 StrataForwardLighting(
 			// Evaluate environment lighting
 			const bool bEnableSpecular = ReflectionStruct.SkyLightParameters.y > 0.0f;
 			FStrataEnvLightResult StrataEnvLight = StrataEvaluateForEnvLight(StrataBSDFContext, bEnableSpecular);
-			if (any(StrataEnvLight.DiffuseColor > 0.0f))
+			if (any(StrataEnvLight.DiffuseWeight > 0.0f))
 			{
-				Color += GetEffectiveSkySHDiffuse(StrataEnvLight.DiffuseNormal) * StrataEnvLight.DiffuseColor * SkyLightCommonFactor;
+				Color += GetEffectiveSkySHDiffuse(StrataEnvLight.DiffuseNormal) * StrataEnvLight.DiffuseWeight * SkyLightCommonFactor;
 			}
 			if (any(StrataEnvLight.SpecularWeight > 0.0f))
 			{
