Sevral fixes Strata environemnt reflections
- GetSpecularOcclusion is now taken into account
- DFAO IndirectSpecularOcclusion and IndirectDiffuseOcclusion is also taken into account
- TODO ExtraIndirectSpecular
- haziness is no longer applied if specular weight is set to 0

diff --git a/Engine/Shaders/Private/Strata/StrataEnvironmentLighting.ush b/Engine/Shaders/Private/Strata/StrataEnvironmentLighting.ush
--- a/Engine/Shaders/Private/Strata/StrataEnvironmentLighting.ush
+++ b/Engine/Shaders/Private/Strata/StrataEnvironmentLighting.ush
@@ -34,6 +34,8 @@ float4 StrataEnvironmentLighting(
 	// Sample DFAO only once
 	FUpsampleDFAOOutput UpsampleDFAOOutput = UpsampleDFAO(BufferUV, SceneDepth);
 
+	float CombinedScreenAndMaterialAO = StrataPixelHeader.MaterialAO * AmbientOcclusion;
+
 #if STRATA_FASTPATH == 0
 	for (uint BSDFIndex = 0; BSDFIndex < StrataPixelHeader.BSDFCount; ++BSDFIndex)
 #endif
@@ -53,22 +55,32 @@ float4 StrataEnvironmentLighting(
 		const bool bEnableSpecular = ReflectionStruct.SkyLightParameters.y > 0.0f;
 		FStrataEnvLightResult StrataEnvLight = StrataEvaluateForEnvLight(StrataBSDFContext, bEnableSpecular);
 
+	#if APPLY_SKY_SHADOWING
+		// Set DiffuseNormal as the bent normal for all diffuse computations.
+		StrataEnvLight.DiffuseNormal = ApplyDFAO(UpsampleDFAOOutput, StrataEnvLight.DiffuseNormal);
+	#endif
+
+		// Compute some extra occlusion information from DFAO and sky light data
+		float IndirectSpecularOcclusion = 1.0f;
+		float IndirectDiffuseOcclusion = 1.0f;
+		float3 ExtraIndirectSpecular = 0.0f;
+	#if SUPPORT_DFAO_INDIRECT_OCCLUSION
+		const bool bTwoSideFoliage = false;
+		GetDistanceFieldAOSpecularOcclusion(
+			StrataEnvLight.DiffuseNormal, StrataEnvLight.SpecularDirection, StrataEnvLight.SpecularSafeRoughness, bTwoSideFoliage,
+			IndirectSpecularOcclusion, IndirectDiffuseOcclusion, ExtraIndirectSpecular);
+		// STRATA_TODO: handle ExtraIndirectSpecular
+	#endif
+
 		// Diffuse component
 		float3 DiffuseSkyLighting = 0;
 		if (any(StrataEnvLight.DiffuseWeight > 0.0f))
 		{
-			// Apply bent normal if needed
-			#if APPLY_SKY_SHADOWING
-			{
-				StrataEnvLight.DiffuseNormal = ApplyDFAO(UpsampleDFAOOutput, StrataEnvLight.DiffuseNormal); // Bent normal
-			}
-			#endif
-
 			// Compute the common sky visibility factors
 			FSkyLightVisibilityData SkyVisData = GetSkyLightVisibilityData(StrataBSDFContext.N, StrataBSDFContext.N, StrataPixelHeader.MaterialAO, AmbientOcclusion, StrataEnvLight.DiffuseNormal);
 
 			// Finally sample the sky diffuse contribution (spherical harmonic, Lambert BRDF)
-			float3 DiffuseLookup = GetSkySHDiffuse(StrataEnvLight.DiffuseNormal) * View.SkyLightColor.rgb;		
+			float3 DiffuseLookup = GetSkySHDiffuse(StrataEnvLight.DiffuseNormal) * View.SkyLightColor.rgb * IndirectDiffuseOcclusion;
 			// And accumulate
 			// Note: Use diffuse directional albedo (i.e., DiffuseWeight) as first order approximation for env. integration (STRATA_TODO instead compute SH coefficients for Chan)
 			DiffuseSkyLighting = BSDFThroughput * (SkyVisData.SkyDiffuseLookUpMul * DiffuseLookup + SkyVisData.SkyDiffuseLookUpAdd) * StrataEnvLight.DiffuseWeight;
@@ -79,26 +91,29 @@ float4 StrataEnvironmentLighting(
 		float3 SpecularSkyLighting = 0;
 		if (any(StrataEnvLight.SpecularWeight > 0.0f))
 		{
+			float RoughnessSquared = StrataEnvLight.SpecularSafeRoughness * StrataEnvLight.SpecularSafeRoughness;
+			float SpecularOcclusion = IndirectSpecularOcclusion * GetSpecularOcclusion(StrataBSDFContext.SatNoV, RoughnessSquared, CombinedScreenAndMaterialAO);
+
 			float SkyAverageBrightness = 1.0f;
 			const float3 SpecularLuminance = GetSkyLightReflection(StrataEnvLight.SpecularDirection, StrataEnvLight.SpecularSafeRoughness, SkyAverageBrightness);
 
-			SpecularSkyLighting += BSDFThroughput * SpecularLuminance * StrataEnvLight.SpecularWeight * View.SkyLightColor.rgb * AmbientOcclusion * 
+			SpecularSkyLighting += BSDFThroughput * SpecularLuminance * StrataEnvLight.SpecularWeight * View.SkyLightColor.rgb * SpecularOcclusion *
 				(bIsTopLayer ? TopLayerSpecularContributionFactor : 1.0f);
-		}
 
-	#if STRATA_FASTPATH==0
-		if (BSDF_GETHASHAZINESS(BSDF))
-		{
-			float SkyAverageBrightness = 1.0f;
-			const float3 SpecularLuminance = GetSkyLightReflection(StrataEnvLight.SpecularDirection, StrataEnvLight.SpecularHazeSafeRoughness, SkyAverageBrightness);
+		#if STRATA_FASTPATH==0
+			if (BSDF_GETHASHAZINESS(BSDF))
+			{
+				float SkyAverageBrightness = 1.0f;
+				const float3 SpecularLuminance = GetSkyLightReflection(StrataEnvLight.SpecularDirection, StrataEnvLight.SpecularHazeSafeRoughness, SkyAverageBrightness);
 
-			SpecularSkyLighting += BSDFThroughput * SpecularLuminance * StrataEnvLight.SpecularHazeWeight * View.SkyLightColor.rgb * AmbientOcclusion;
+				SpecularSkyLighting += BSDFThroughput * SpecularLuminance * StrataEnvLight.SpecularHazeWeight * View.SkyLightColor.rgb * SpecularOcclusion;
 
-			// SSR is traced for the sharpest lob. The smoothest one does not rely on SSR so we need to lower energy coming from SSR according to the lobe blend weight.
-			// And we also try to make the transition smooth using Haziness
-			SSRReductionFactor -= bIsTopLayer ? dot(BSDFThroughput, float3(1.0f / 3.0f, 1.0f / 3.0f, 1.0f / 3.0f)) * StrataEnvLight.SSRReduction : 0.0f;
+				// SSR is traced for the sharpest lob. The smoothest one does not rely on SSR so we need to lower energy coming from SSR according to the lobe blend weight.
+				// And we also try to make the transition smooth using Haziness
+				SSRReductionFactor -= bIsTopLayer ? dot(BSDFThroughput, float3(1.0f / 3.0f, 1.0f / 3.0f, 1.0f / 3.0f)) * StrataEnvLight.SSRReduction : 0.0f;
+			}
+		#endif 
 		}
-	#endif 
 
 		if (bIsTopLayer)
 		{
