Fix Strata SSS non-profile composition and scaling.

diff --git a/Engine/Shaders/Private/PostProcessSubsurface.usf b/Engine/Shaders/Private/PostProcessSubsurface.usf
--- a/Engine/Shaders/Private/PostProcessSubsurface.usf
+++ b/Engine/Shaders/Private/PostProcessSubsurface.usf
@@ -742,7 +742,9 @@ void SubsurfaceRecombinePS(noperspective float4 UVAndScreenPos : TEXCOORD0, floa
 	}
 	else if (ProfileId == SSS_PROFILE_ID_PERPIXEL)
 	{
-		SubsurfaceColor = SSSData.DMFP;
+		// This color is used to lerping 'incident light' into 'diffused light', i.e, controlling the sharpness of the lighting.
+		// In case of direct DMPF control, this is not needed.
+		SubsurfaceColor = float3(1, 1, 1);
 	}
 #else
 	const uint ProfileId = ExtractSubsurfaceProfileInt(ScreenSpaceData.GBuffer);
diff --git a/Engine/Shaders/Private/SubsurfaceBurleyNormalized.ush b/Engine/Shaders/Private/SubsurfaceBurleyNormalized.ush
--- a/Engine/Shaders/Private/SubsurfaceBurleyNormalized.ush
+++ b/Engine/Shaders/Private/SubsurfaceBurleyNormalized.ush
@@ -212,17 +212,14 @@ FBurleyParameter GetBurleyParameters(FStrataSubsurfaceData SSSData)
 	else if(SSSData.bIsValid)
 	{
 		// Decode the DMFP stored as Float111110 in a uint over a RGBA8 unorm
-		float3 DMFP = SSSData.DMFP;
-
-		// Conversion to match other SSS encoding 
-		DMFP *= BURLEY_CM_2_MM * 5.0f; // STRATA_TODO why a scale of 5 is needed here to match between regular and Strata?
+		const float3 DMFP				= SSSData.DMFP * BURLEY_CM_2_MM; // Stored DMFP is in cm, converted into mm
 
 		const float MaxDMFP				= max3(DMFP.x, DMFP.y, DMFP.z);
 		const uint IndexOfMaxDmfp		= (DMFP.r == MaxDMFP) ? 0 : ((DMFP.g == MaxDMFP) ? 1 : 2);
 
 		Out.SurfaceAlbedo				= float4(SSSData.BaseColor, SSSData.BaseColor[IndexOfMaxDmfp]);
 		Out.DiffuseMeanFreePath			= float4(DMFP, MaxDMFP);
-		Out.WorldUnitScale				= DecodeWorldUnitScale(1.0f * ENC_WORLDUNITSCALE_IN_CM_TO_UNIT) * BURLEY_CM_2_MM;
+		Out.WorldUnitScale				= BURLEY_CM_2_MM;
 
 		// STRATA_TODO smooth transistion based on pixel footprint w.r.t. MFPRadius and pixel footprint
 		Out.SurfaceOpacity				= 1.0f;
