Fix top layer normal accumulating non-top layer data.

diff --git a/Engine/Shaders/Private/Strata/StrataMaterialClassification.usf b/Engine/Shaders/Private/Strata/StrataMaterialClassification.usf
--- a/Engine/Shaders/Private/Strata/StrataMaterialClassification.usf
+++ b/Engine/Shaders/Private/Strata/StrataMaterialClassification.usf
@@ -37,13 +37,16 @@ void ComputePostProcessBSDFData(
 	float3 N = TangentBasis[2];
 	//N = normalize(N);
 
-	TopLayer.WorldNormal += N;
 
 	BxDFContext ContextNoLight = (BxDFContext)0;
 	const float3 FakeL = float3(0, 0, 1);
 	Init(ContextNoLight, N, V, FakeL);
 
 	const bool bTopLayer = BSDF_GETISTOPLAYER(BSDF);
+	if (bTopLayer)
+	{
+		TopLayer.WorldNormal += N;
+	}
 
 	Classification.ShadingModels |= StrataShadingModelBit(BSDF_GETTYPE(BSDF));
 	Classification.bIsSimple = Classification.bIsSimple && BSDF_GETTYPE(BSDF) == STRATA_BSDF_TYPE_SLAB;
