* Fix ThinFilm input parameters
* Fix the update BSDF context  function which wasn't patching certain variables.
* Change BSDF context to construct context with and without a valid L. In the latter case, the reflection vector is used.

diff --git a/Engine/Shaders/Private/Strata/StrataEnvironmentLighting.ush b/Engine/Shaders/Private/Strata/StrataEnvironmentLighting.ush
--- a/Engine/Shaders/Private/Strata/StrataEnvironmentLighting.ush
+++ b/Engine/Shaders/Private/Strata/StrataEnvironmentLighting.ush
@@ -48,8 +48,7 @@ float4 StrataEnvironmentLighting(
 		const float3 BSDFThroughput = BSDF.Weight;
 
 		// Create the BSDF context
-		const float3 UnusedL = float3(0.0f, 0.0f, 1.0f);
-		FStrataBSDFContext StrataBSDFContext = StrataCreateBSDFContext(StrataPixelHeader, BSDF, StrataAddressing, V, UnusedL);
+		FStrataBSDFContext StrataBSDFContext = StrataCreateBSDFContext(StrataPixelHeader, BSDF, StrataAddressing, V);
 
 		// Evaluate environment lighting
 		const bool bEnableSpecular = ReflectionStruct.SkyLightParameters.y > 0.0f;
diff --git a/Engine/Shaders/Private/Strata/StrataEvaluation.ush b/Engine/Shaders/Private/Strata/StrataEvaluation.ush
--- a/Engine/Shaders/Private/Strata/StrataEvaluation.ush
+++ b/Engine/Shaders/Private/Strata/StrataEvaluation.ush
@@ -95,26 +95,18 @@ void StrataGetBSDFRoughness(in FStrataBSDF BSDF, inout float R)
 	}
 }
 
-FStrataBSDFContext StrataCreateBSDFContext(float3x3 TangentBasis, FStrataBSDF BSDF, float3 V, float3 L, bool bOverrideWithLocalFrame = false)
+FStrataBSDFContext StrataCreateBSDFContext(float3x3 TangentBasis, FStrataBSDF BSDF, float3 V, float3 L, bool bHasValidL=true)
 {
 	FStrataBSDFContext BSDFContext = (FStrataBSDFContext)0;
 
 	BSDFContext.BSDF = BSDF;
 
-	BSDFContext.V = V;
-	BSDFContext.L = L;
-
 	BSDFContext.X = TangentBasis[0]; 
 	BSDFContext.Y = TangentBasis[1];
 	BSDFContext.N = TangentBasis[2]; 
-	if (bOverrideWithLocalFrame)
-	{
-		BSDFContext.X = float3(1,0,0);
-		BSDFContext.Y = float3(0,1,0);
-		BSDFContext.N = float3(0,0,1);
-	}
-
+	BSDFContext.V = V;
 	BSDFContext.R = 2 * dot(BSDFContext.V, BSDFContext.N) * BSDFContext.N - BSDFContext.V;
+	BSDFContext.L = bHasValidL ? L : BSDFContext.R;
 	BSDFContext.H = normalize(BSDFContext.V + BSDFContext.L);
 
 	BSDFContext.Context = (BxDFContext)0;
@@ -137,18 +129,29 @@ FStrataBSDFContext StrataCreateBSDFContext(float3x3 TangentBasis, FStrataBSDF BS
 	return BSDFContext;
 }
 
-FStrataBSDFContext StrataCreateBSDFContext(FStrataPixelHeader StrataPixelHeader, FStrataBSDF BSDF, const FStrataAddressing StrataAddressing, float3 V, float3 L, bool bOverrideWithLocalFrame = false)
+FStrataBSDFContext StrataCreateBSDFContext(FStrataPixelHeader StrataPixelHeader, FStrataBSDF BSDF, const FStrataAddressing StrataAddressing, float3 V)
 {
+	float3 UnusedL = float3(0, 0, 1);
 	float3x3 TangentBasis = StrataGetBSDFSharedBasis(StrataPixelHeader, BSDF, StrataAddressing);
-	return StrataCreateBSDFContext(TangentBasis, BSDF, V, L, bOverrideWithLocalFrame);
+	return StrataCreateBSDFContext(TangentBasis, BSDF, V, UnusedL, false);
+}
+
+FStrataBSDFContext StrataCreateBSDFContext(FStrataPixelHeader StrataPixelHeader, FStrataBSDF BSDF, const FStrataAddressing StrataAddressing, float3 V, float3 L)
+{
+	float3x3 TangentBasis = StrataGetBSDFSharedBasis(StrataPixelHeader, BSDF, StrataAddressing);
+	return StrataCreateBSDFContext(TangentBasis, BSDF, V, L, true);
 }
 
 void StrataUpdateBSDFContext(inout FStrataBSDFContext BSDFContext, float3 NewL)
 {
 	// Update all the data related to L
-	Init(BSDFContext.Context, BSDFContext.N, BSDFContext.V, NewL);
+	BSDFContext.L = NewL;
+	Init(BSDFContext.Context, BSDFContext.N, BSDFContext.V, BSDFContext.L);
 	BSDFContext.H = normalize(BSDFContext.V + BSDFContext.L);
+	BSDFContext.R = 2 * dot(BSDFContext.V, BSDFContext.N) * BSDFContext.N - BSDFContext.V;
 	BSDFContext.TangentH = mul(BSDFContext.H, BSDFContext.TangentBasis);
+	BSDFContext.SatNoL = saturate(BSDFContext.Context.NoL);
+	BSDFContext.SatNoV = saturate(BSDFContext.Context.NoV);
 }
 
 struct FStrataEvaluateResult
@@ -532,9 +535,8 @@ FStrataImportanceSampleResult StrataImportanceSampleBSDF(FStrataPixelHeader Stra
 	FStrataImportanceSampleResult Sample = (FStrataImportanceSampleResult)0;
 
 	// Create a context from the information we have
-	const float3 UnusedL = float3(0.0f, 0.0f, 1.0f);
 	const float3 V = -CameraVector;
-	FStrataBSDFContext BSDFContext = StrataCreateBSDFContext(StrataPixelHeader, BSDF, StrataAddressing, V, UnusedL);
+	FStrataBSDFContext BSDFContext = StrataCreateBSDFContext(StrataPixelHeader, BSDF, StrataAddressing, V);
 	return StrataImportanceSampleBSDF(BSDFContext, BSDF, E);
 }
 
diff --git a/Engine/Shaders/Private/Strata/StrataFurnaceTest.usf b/Engine/Shaders/Private/Strata/StrataFurnaceTest.usf
--- a/Engine/Shaders/Private/Strata/StrataFurnaceTest.usf
+++ b/Engine/Shaders/Private/Strata/StrataFurnaceTest.usf
@@ -179,7 +179,7 @@ void MainPS(
 
 				const float SatNoL = saturate(dot(L, N));
 
-				FStrataBSDFContext BSDFContext = StrataCreateBSDFContext(TangentBasis, BSDF, V, L);
+				FStrataBSDFContext BSDFContext = StrataCreateBSDFContext(TangentBasis, BSDF, V, L, true);
 				FStrataEvaluateResult BSDFEvaluate = StrataEvaluateBSDF(BSDFContext, ShadowTerms);
 
 				const float3 DiffusePathWeight  = BSDF.Weight * BSDFEvaluate.DiffusePathValue;
@@ -196,8 +196,8 @@ void MainPS(
 			{
 				const float2 U = Hammersley(SampleIt, NumSamples, 0);
 
-				const float3 UnusedL = float3(0.0f, 0.0f, 1.0f);
-				FStrataBSDFContext BSDFContextIS = StrataCreateBSDFContext(TangentBasis, BSDF, V, UnusedL);
+				const float3 UnusedL = float3(0, 0, 1);
+				FStrataBSDFContext BSDFContextIS = StrataCreateBSDFContext(TangentBasis, BSDF, V, UnusedL, false);
 				FStrataImportanceSampleResult S = StrataImportanceSampleBSDF(BSDFContextIS, BSDF, U);
 
 				// Evaluate specular part for now. 
@@ -208,7 +208,7 @@ void MainPS(
 				//const float PDF = S.DiffusePDF;
 
 				const float SatNoL = saturate(dot(L, N));
-				FStrataBSDFContext BSDFContext = StrataCreateBSDFContext(TangentBasis, BSDF, V, L);
+				FStrataBSDFContext BSDFContext = StrataCreateBSDFContext(TangentBasis, BSDF, V, L, true);
 				FStrataEvaluateResult BSDFEvaluate = StrataEvaluateBSDF(BSDFContext, ShadowTerms);
 
 				const float3 DiffusePathWeight  = BSDF.Weight * BSDFEvaluate.DiffusePathValue;
@@ -223,9 +223,9 @@ void MainPS(
 		{
 			const float3 BSDFThroughput = BSDF.Weight;
 
-			const float3 UnusedL = float3(0.0f, 0.0f, 1.0f);
+			const float3 UnusedL = float3(0, 0, 1);
 			const float3x3 TangentBasis = GetTangentBasis(N);
-			FStrataBSDFContext BSDFContext = StrataCreateBSDFContext(TangentBasis, BSDF, V, UnusedL);
+			FStrataBSDFContext BSDFContext = StrataCreateBSDFContext(TangentBasis, BSDF, V, UnusedL, false);
 
 			const bool bEnableSpecular = true;
 			FStrataEnvLightResult StrataEnvLight = StrataEvaluateForEnvLight(BSDFContext, bEnableSpecular);
diff --git a/Engine/Shaders/Private/Strata/StrataVisualize.usf b/Engine/Shaders/Private/Strata/StrataVisualize.usf
--- a/Engine/Shaders/Private/Strata/StrataVisualize.usf
+++ b/Engine/Shaders/Private/Strata/StrataVisualize.usf
@@ -58,7 +58,12 @@ float3 DrawPlot(
 			LocalV = normalize(LocalV);
 		}
 
-		FStrataBSDFContext BSDFContext = StrataCreateBSDFContext(Header, BSDF, StrataAddressing, LocalV, L, true);
+		float3x3 TangentBasis;
+		TangentBasis[0] = float3(1, 0, 0);
+		TangentBasis[1] = float3(0, 1, 0);
+		TangentBasis[2] = float3(0, 0, 1);
+
+		FStrataBSDFContext BSDFContext = StrataCreateBSDFContext(TangentBasis, BSDF, LocalV, L, true);
 		FStrataEvaluateResult BSDFEvaluate = StrataEvaluateBSDF(BSDFContext);
 		Out = (BSDFEvaluate.DiffusePathValue + BSDFEvaluate.SpecularPathValue + BSDFEvaluate.TransmissionPathValue) * BSDFContext.SatNoL; // * BSDF.Weight
 	}
diff --git a/Engine/Shaders/Private/ThinFilmBSDF.ush b/Engine/Shaders/Private/ThinFilmBSDF.ush
--- a/Engine/Shaders/Private/ThinFilmBSDF.ush
+++ b/Engine/Shaders/Private/ThinFilmBSDF.ush
@@ -99,7 +99,7 @@ void ComputeComplexIORFromF0AndEdgeTint(in float3 Reflectivity, in float3 EdgeTi
 
 // Fresnel term for iridescent microfacet BRDF model 
 // This is a reference implementation adapted from the original paper's source code
-float3 F_ThinFilmRef(float NoL, float NoV, float VoH, float3 F0, float3 F90, float ThinFilmIOR, float ThinFilmTickness)
+float3 F_ThinFilmRef(float NoV, float NoL, float VoH, float3 F0, float3 F90, float ThinFilmIOR, float ThinFilmTickness)
 {
 	// Translate F0/F90 to eta/kappa form
 	float3 n3 = 1;
@@ -132,7 +132,7 @@ float3 F_ThinFilmRef(float NoL, float NoV, float VoH, float3 F0, float3 F90, flo
 	const float eta_2 = lerp(1.0, eta2, smoothstep(0.0, 0.03, Dinc));
 
 	// Compute dot products
-	const bool bIsValid = !(NoL < 0 || NoV < 0);
+	const bool bIsValid = NoL > 0 && NoV > 0;
 
 	float3 I = 0;
 	if (bIsValid)
@@ -186,7 +186,7 @@ float3 F_ThinFilmRef(float NoL, float NoV, float VoH, float3 F0, float3 F90, flo
 // Fresnel term for iridescent microfacet BRDF model 
 // Simplified version which relies on Schlick's Fresnel and de facto does not take into 
 // account Fresnel phase shift & polarization.
-float3 F_ThinFilm(float NoL, float NoV, float VoH, float3 F0, float3 F90, float ThinFilmIOR, float ThinFilmTickness)
+float3 F_ThinFilm(float NoV, float NoL, float VoH, float3 F0, float3 F90, float ThinFilmIOR, float ThinFilmTickness)
 {
 	// Thin film is limited (for now) to the top most layer (i.e., the upper layer need to be air)
 	//					|
@@ -203,7 +203,7 @@ float3 F_ThinFilm(float NoL, float NoV, float VoH, float3 F0, float3 F90, float
 	const float eta2 = lerp(1.0, ThinFilmIOR, smoothstep(0.0, 0.03, Dinc));
 
 	// Compute dot products
-	const bool bIsValid = !(NoL < 0 || NoV < 0);
+	const bool bIsValid = NoL > 0 && NoV > 0;
 
 	float3 I = 0;
 	if (bIsValid)
