Reverts previous bucket weight adjustment in the shadow denoiser's DRB

Causes more problems than its solves

diff --git a/Engine/Shaders/Private/ScreenSpaceDenoise/SSDSignalAccumulator.ush b/Engine/Shaders/Private/ScreenSpaceDenoise/SSDSignalAccumulator.ush
--- a/Engine/Shaders/Private/ScreenSpaceDenoise/SSDSignalAccumulator.ush
+++ b/Engine/Shaders/Private/ScreenSpaceDenoise/SSDSignalAccumulator.ush
@@ -316,7 +316,7 @@ void DijestSampleClusterInDRB(
 	float PreviousBucketFactor = (Accumulator.Previous.SampleCount == 0.0) ? 0.0 : (1.0 - CurrentClusterOpacity * bOccludingCluster);
 
 	// Normalize the previous bucket to respect the desired frequency bucket opacity.
-	PreviousBucketFactor *= Accumulator.Current.SampleCount * SafeRcp(Accumulator.Previous.SampleCount);
+	// PreviousBucketFactor *= Accumulator.Current.SampleCount * SafeRcp(Accumulator.Previous.SampleCount);
 
 	Accumulator.Previous = AddSignal(MulSignal(Accumulator.Previous, PreviousBucketFactor), Accumulator.Current);
 	Accumulator.PreviousInvFrequency = Accumulator.PreviousInvFrequency * PreviousBucketFactor + Accumulator.CurrentInvFrequency;
