[Switch] Fix for shader compile error "loop condition must be scalar boolean"

diff --git a/Engine/Shaders/Private/VirtualShadowMaps/BuildPerPageDrawCommands.usf b/Engine/Shaders/Private/VirtualShadowMaps/BuildPerPageDrawCommands.usf
--- a/Engine/Shaders/Private/VirtualShadowMaps/BuildPerPageDrawCommands.usf
+++ b/Engine/Shaders/Private/VirtualShadowMaps/BuildPerPageDrawCommands.usf
@@ -278,7 +278,7 @@ void CullPerPageDrawCommandsCs(uint InstanceIdIndex : SV_DispatchThreadID, uint
 		VisibleInstanceOutputOffset += GroupVisibleInstanceOutputOffset;
 
 		// Process mip levels with visible pages and output nodes
-		while (VisibleMipsMask)
+		while (VisibleMipsMask != 0)
 		{
 			uint MipLevel = firstbitlow(VisibleMipsMask);	// Jump to next set bit.
 			VisibleMipsMask &= VisibleMipsMask - 1u;		// Remove bottom bit.
