Fix warning on bitshift, by making explicit unsigned literal

diff --git a/Engine/Shaders/Private/RayTracing/RayTracingLightCullingCommon.ush b/Engine/Shaders/Private/RayTracing/RayTracingLightCullingCommon.ush
--- a/Engine/Shaders/Private/RayTracing/RayTracingLightCullingCommon.ush
+++ b/Engine/Shaders/Private/RayTracing/RayTracingLightCullingCommon.ush
@@ -20,7 +20,7 @@ struct FCulledLightList
 
 	static const uint PackedLightBits = 10;
 	static const uint PackedLightsPerComponent = 3;
-	static const uint PackedLightMask = (1 << PackedLightBits) - 1;
+	static const uint PackedLightMask = (1u << PackedLightBits) - 1;
 	static const uint PackedLightsMax = 11;
 
 	static FCulledLightList Create(float3 WorldPos)
