Fix to ifdef of non-nanite VSM in pixel shader

diff --git a/Engine/Shaders/Private/ShadowDepthPixelShader.usf b/Engine/Shaders/Private/ShadowDepthPixelShader.usf
--- a/Engine/Shaders/Private/ShadowDepthPixelShader.usf
+++ b/Engine/Shaders/Private/ShadowDepthPixelShader.usf
@@ -80,7 +80,7 @@ void Main(
 	OutDepth = SvPosition.z;
 #endif
 
-#if VIRTUAL_TEXTURE_TARGET
+#if ENABLE_NON_NANITE_VSM && VIRTUAL_TEXTURE_TARGET
 	uint2 vAddress = (uint2)SvPosition.xy;
 	float DeviceZ = SvPosition.z;
 
