Made Strata Slab computation happend for a layer of 1 meter. Required for the mapping comming soon.
MAde BaseColor part of the pariticpating media definition.

diff --git a/Engine/Shaders/Private/ParticipatingMediaCommon.ush b/Engine/Shaders/Private/ParticipatingMediaCommon.ush
--- a/Engine/Shaders/Private/ParticipatingMediaCommon.ush
+++ b/Engine/Shaders/Private/ParticipatingMediaCommon.ush
@@ -6,31 +6,75 @@
 
 #include "Common.ush"
 
+#define PARTICIPATING_MEDIA_MIN_MFP_METER		0.000001f
+#define PARTICIPATING_MEDIA_MIN_EXTINCTION		0.000001f
+#define PARTICIPATING_MEDIA_MIN_TRANSMITTANCE	0.000000000001f
+
+
+
+//---------------------------------------------
+// Participating media representation
+//---------------------------------------------
+
 struct FParticipatingMedia
 {
 	float3 ScatteringCoef;	// sigma_s (1/meter)
 	float3 AbsorptionCoef;	// sigma_a (1/meter)
 	float3 ExtinctionCoef;	// sigma_t (1/meter)
 	float3 MeanFreePath;	// (meter)
-	float3 Albedo;			// (unitless)
+	float3 Albedo;			// Represent sigma_s / sigma_t (unitless)
+	float3 BaseColor;		// Represent the reflectance albedo resulting from single+multiple scattering assuming an isotropic phase function (unitless)
 };
 
-#define PARTICIPATING_MEDIA_MIN_MFP_METER		0.000001f
-#define PARTICIPATING_MEDIA_MIN_EXTINCTION		0.000001f
-#define PARTICIPATING_MEDIA_MIN_TRANSMITTANCE	0.000000000001f
+// GetBaseColorFromAlbedo: returns the color of the participating media resulting from the single+multiple subsurface scattering. 
+// GetAlbedoFromBaseColor is the inverse transform.
+// [Kulla and Conty 2017, "Revisiting Physically Based Shading at Imageworks"] https://blog.selfshadow.com/publications/s2017-shading-course/imageworks/s2017_pbs_imageworks_slides_v2.pdf, slide 62 
+// [d'Eon, "A Hitchhiker's guide to multiple scattering"] or [REDACTED_PROFILE_URL]
+float3 GetBaseColorFromAlbedo(const float3 Albedo, const float g = 0.0f)
+{
+	const float3 s = sqrt((1 - Albedo) / (1.0f - Albedo * g));
+	const float3 BaseColor = ((1.0f - s) * (1 - 0.139 * s)) / (1.0f + 1.17 * s);
+	return BaseColor;
+}
+float3 GetAlbedoFromBaseColor(const float3 BaseColor, const float g = 0.0f)
+{
+	const float3 s = 4.09712 + 4.20863 * BaseColor - sqrt(9.59217 + 41.6808 * BaseColor + 17.7126 * BaseColor * BaseColor);
+	const float3 Albedo = (1.0f - s * s) / (1.0f - g * s * s);
+	return Albedo;
+}
 
 // The mean free path must be in meters
 FParticipatingMedia CreateMediumFromAlbedoMFP(float3 Albedo, float3 MeanFreePathMeters)
 {
 	FParticipatingMedia PM = (FParticipatingMedia)0;
 	PM.Albedo = Albedo;
+	PM.BaseColor = GetBaseColorFromAlbedo(Albedo);
 	PM.MeanFreePath = MeanFreePathMeters;
 	PM.ExtinctionCoef = 1.0f / max(PARTICIPATING_MEDIA_MIN_MFP_METER, PM.MeanFreePath);
-	PM.ScatteringCoef = Albedo * PM.ExtinctionCoef;
+	PM.ScatteringCoef = PM.Albedo * PM.ExtinctionCoef;
 	PM.AbsorptionCoef = max(0.0f, PM.ExtinctionCoef - PM.ScatteringCoef);
 	return PM;
 }
 
+// The mean free path must be in meters
+FParticipatingMedia CreateMediumFromBaseColorMFP(float3 BaseColor, float3 MeanFreePathMeters)
+{
+	FParticipatingMedia PM = (FParticipatingMedia)0;
+	PM.Albedo = GetAlbedoFromBaseColor(BaseColor);
+	PM.BaseColor = BaseColor;
+	PM.MeanFreePath = MeanFreePathMeters;
+	PM.ExtinctionCoef = 1.0f / max(PARTICIPATING_MEDIA_MIN_MFP_METER, PM.MeanFreePath);
+	PM.ScatteringCoef = PM.Albedo * PM.ExtinctionCoef;
+	PM.AbsorptionCoef = max(0.0f, PM.ExtinctionCoef - PM.ScatteringCoef);
+	return PM;
+}
+
+
+
+//---------------------------------------------
+// Phase functions
+//---------------------------------------------
+
 float IsotropicPhase()
 {
 	return 1.0f / (4.0f * PI);
@@ -90,6 +134,12 @@ float4 ImportanceSampleHenyeyGreensteinPhase(float2 E, float G)
 	return float4(H, HenyeyGreensteinPhase(G, CosTheta));
 }
 
+
+
+//---------------------------------------------
+// Utilities
+//---------------------------------------------
+
 float3 TransmittanceToExtinction(in float3 TransmittanceColor, in float ThicknessMeters)
 {
 	// TransmittanceColor	= exp(-Extinction * Thickness)
diff --git a/Engine/Shaders/Private/Strata/Strata.ush b/Engine/Shaders/Private/Strata/Strata.ush
--- a/Engine/Shaders/Private/Strata/Strata.ush
+++ b/Engine/Shaders/Private/Strata/Strata.ush
@@ -35,9 +35,9 @@
 
 #define STRATA_INLINE_SINGLELAYERWATER (PROJECT_STRATA && MATERIAL_IS_STRATA && STRATA_INLINE_SHADING && MATERIAL_SHADINGMODEL_SINGLELAYERWATER)
 
-// We do all the simple volumetric math considering a slab of 1 centimeter = 0.01 meter.
+// Behind the scene, we do all the simple volumetric lighting & transmitance computations considering a slab of 100 centimeters = 1 meter.
 // If the user specifies a different thickness to simulate varying appearance, we rescale the mean free path in order to not have to store thickness as part of the BSDF description.
-#define STRATA_SIMPLEVOLUME_THICKNESS_CM	1.0f
+#define STRATA_SIMPLEVOLUME_THICKNESS_CM	100.0f
 #define STRATA_CENTIMETER_TO_METER			0.01f
 #define STRATA_METER_TO_CENTIMETER			100.0f
 #define STRATA_SIMPLEVOLUME_THICKNESS_M		(STRATA_SIMPLEVOLUME_THICKNESS_CM * STRATA_CENTIMETER_TO_METER)
@@ -957,7 +957,7 @@ void StrataSanitizeBSDF(inout FStrataBSDF BSDF)
 FParticipatingMedia StrataSlabCreateParticipatingMedia(float3 DiffuseColor, float3 MeanFreePathCentimeters)
 {
 	const float3 MeanFreePathMeters = MeanFreePathCentimeters * STRATA_CENTIMETER_TO_METER;
-	return CreateMediumFromAlbedoMFP(DiffuseColor, MeanFreePathMeters);
+	return CreateMediumFromBaseColorMFP(DiffuseColor, MeanFreePathMeters);
 }
 
 float StrataSlabDiffuseToVolumeBlend(FParticipatingMedia PM)
