Fixed approximation to transmittance due to frensel effect to not go out of physical boundaries.

diff --git a/Engine/Shaders/Private/Strata/Strata.ush b/Engine/Shaders/Private/Strata/Strata.ush
--- a/Engine/Shaders/Private/Strata/Strata.ush
+++ b/Engine/Shaders/Private/Strata/Strata.ush
@@ -1348,8 +1348,9 @@ void PackStrataOut(
 					float3 ApproxFresnelForTransmittance = 0;
 					// STRATA_TODO Fresnel throughput, Use a LUT, this is currently in progress and will be coming soon.
 					// In the meantime, we use a simple approach to feed Fresnel function and compute an approximation of the general transmition/througput.
-					const float VoH = ContextNoLight.NoV;
-					const float NoL = ContextNoLight.NoV;
+					const float SatNoV = saturate(ContextNoLight.NoV);
+					const float VoH = SatNoV;
+					const float NoL = SatNoV;
 					const float3 BaseColor = SLAB_BASECOLOR(BSDF);
 					const float  Metallic = SLAB_METALLIC(BSDF);
 					const float3 F0 = ComputeF0(SLAB_SPECULAR(BSDF), BaseColor, Metallic);
@@ -1360,7 +1361,7 @@ void PackStrataOut(
 					{
 						float ThinFilmIOR = SLAB_THINFILM_IOR(BSDF);
 						float ThinFilmThickness = SLAB_THINFILM_THICKNESS(BSDF);
-						ApproxFresnelForTransmittance = F_ThinFilm(ContextNoLight.NoV, NoL, VoH, F0, F90, ThinFilmIOR, ThinFilmThickness);
+						ApproxFresnelForTransmittance = F_ThinFilm(SatNoV, NoL, VoH, F0, F90, ThinFilmIOR, ThinFilmThickness);
 					}
 					else
 					{
diff --git a/Engine/Shaders/Private/Strata/StrataEvaluation.ush b/Engine/Shaders/Private/Strata/StrataEvaluation.ush
--- a/Engine/Shaders/Private/Strata/StrataEvaluation.ush
+++ b/Engine/Shaders/Private/Strata/StrataEvaluation.ush
@@ -357,7 +357,7 @@ FStrataEvaluateResult StrataEvaluateBSDF(FStrataBSDFContext BSDFContext, FShadow
 				// STRATA_TODO Fresnel throughput, Use a LUT, this is currently in progress and will be coming soon.
 				// In the meantime, we use a simple approach to feed Fresnel function and compute an approximation of the general transmition/througput.
 				// We also ignored ThinFilm interaction as another simplification for perf.
-				const float3 ApproxFresnelForTransmittance = F_Schlick(F0, F90, BSDFContext.Context.NoV);
+				const float3 ApproxFresnelForTransmittance = F_Schlick(F0, F90, BSDFContext.SatNoV);
 
 				const float3 SpecularThroughput = 1.0f - ApproxFresnelForTransmittance; // This is already taken into account by the Chan model above. But for volumetric we need to apply it.
 				const float Phase = IsotropicPhase();
@@ -389,7 +389,7 @@ FStrataEvaluateResult StrataEvaluateBSDF(FStrataBSDFContext BSDFContext, FShadow
 
 			// STRATA_TODO Fresnel throughput, Use a LUT, this is currently in progress and will be coming soon.
 			// In the meantime, we use a simple approach to feed Fresnel function and compute an approximation of the general transmition/througput.
-			const float3 ApproxFresnelForTransmittance = F_Schlick(F0, BSDFContext.Context.NoV);
+			const float3 ApproxFresnelForTransmittance = F_Schlick(F0, BSDFContext.SatNoV);
 
 			Sample.DiffusePathValue	= D * Vis * F;
 			Sample.DiffusePDF		= BSDFContext.SatNoV / PI; // Per "Production Friendly Microfacet Sheen BRDF", hemispherical sampling give good result as the roughness is usually high.
