Missing unsaved changed to go with 15567665.

diff --git a/Engine/Shaders/Private/Strata/Strata.ush b/Engine/Shaders/Private/Strata/Strata.ush
--- a/Engine/Shaders/Private/Strata/Strata.ush
+++ b/Engine/Shaders/Private/Strata/Strata.ush
@@ -1391,7 +1391,7 @@ void PackStrataOut(
 			{
 				// STRATA_TODO Fresnel throughput, Use a LUT, this is currently in progress and will be coming soon.
 				// In the meantime, we use a simple approach to feed Fresnel function and compute an approximation of the general transmition/througput.
-				const float NoL = ContextNoLight.NoV;
+				const float NoL = saturate(ContextNoLight.NoV);
 				const float3 F0  = SHEEN_BASECOLOR(BSDF);
 				const float3 ApproxFresnelForTransmittance = F_Schlick(F0, NoL);
 
