Fix compilation issue with DXC due to overlapping semantic of input parameters in hair strands material shaders.

diff --git a/Engine/Shaders/Private/HairStrands/HairStrandsMaterialGBufferVS.usf b/Engine/Shaders/Private/HairStrands/HairStrandsMaterialGBufferVS.usf
--- a/Engine/Shaders/Private/HairStrands/HairStrandsMaterialGBufferVS.usf
+++ b/Engine/Shaders/Private/HairStrands/HairStrandsMaterialGBufferVS.usf
@@ -9,13 +9,12 @@ float3 DecodeVelocityFromTexture(float4 In);
 #if VERTEXSHADER
  
 void Main(
-	uint VertexId : SV_VertexID, 
 	FVertexFactoryInput Input,
 	out float2 OutUV : TEXCOORD,
 	out float4 OutPosition : SV_POSITION,
 	nointerpolation out uint  OutPrimitiveId : PRIMITIVE_ID)
 {
-	const float2 UV = float2(VertexId & 1, VertexId >> 1) * 2;
+	const float2 UV = float2(Input.VertexId & 1, Input.VertexId >> 1) * 2;
 	const float2 Pos = float2(UV.x, 1 - UV.y) * 2 - float2(1, 1);// * 4 - float2(2,2);
 	OutPosition = float4(Pos, 0.5f, 1);
 	OutUV = UV;
diff --git a/Engine/Shaders/Private/HairStrands/HairStrandsMaterialVS.usf b/Engine/Shaders/Private/HairStrands/HairStrandsMaterialVS.usf
--- a/Engine/Shaders/Private/HairStrands/HairStrandsMaterialVS.usf
+++ b/Engine/Shaders/Private/HairStrands/HairStrandsMaterialVS.usf
@@ -9,7 +9,6 @@ float3 DecodeVelocityFromTexture(float4 In);
 #if VERTEXSHADER
  
 void Main(
-	uint VertexId : SV_VertexID, 
 	FVertexFactoryInput Input,
 	out float4 Position : SV_POSITION,
 	nointerpolation out uint2 OutResolution : DISPATCH_RESOLUTION,
@@ -22,7 +21,7 @@ void Main(
 
 	const float2 ClipCoord = ((OutResolution +0.5f) / float2(MaterialPassParameters.MaxResolution)) * 2;
 
-	const float2 UV = float2(VertexId & 1, VertexId >> 1);
+	const float2 UV = float2(Input.VertexId & 1, Input.VertexId >> 1);
 	const float2 Pos = float2(-UV.x, UV.y) * 4 + float2(-1,+1) + float2(ClipCoord.x, -ClipCoord.y);
 	Position = float4(Pos, 0.5f, 1);
 
