Moved Geometry Collection from Local Vertex Factory to a dedicated VF.
- Simplifies the ifdef logic in both and cleans up naming significantly.
- also added support for GPU-scene instance culling.

diff --git a/Engine/Shaders/Private/GeometryCollectionVertexFactory.ush b/Engine/Shaders/Private/GeometryCollectionVertexFactory.ush
new file mode 100644
--- /dev/null
+++ b/Engine/Shaders/Private/GeometryCollectionVertexFactory.ush
@@ -0,0 +1,572 @@
+// Copyright Epic Games, Inc. All Rights Reserved.
+
+/*=============================================================================
+	GeometryCollectionVertexFactory.ush: Shader code for FGeometryCollectionVertexFactory.
+=============================================================================*/
+
+#include "SceneData.ush"
+#include "VertexFactoryCommon.ush"
+#include "LocalVertexFactoryCommon.ush"
+#include "LightmapData.ush"
+
+#include "/Engine/Generated/UniformBuffers/PrecomputedLightingBuffer.ush"
+
+#ifndef MANUAL_VERTEX_FETCH
+	#define MANUAL_VERTEX_FETCH 0
+#endif
+
+#if MANUAL_VERTEX_FETCH
+	#define VF_ColorIndexMask_Index 0
+	#define VF_NumTexcoords_Index 1
+	#define FV_LightMapIndex_Index 2
+	#define VF_VertexOffset 3
+
+	Buffer<float4> VertexFetch_BoneOriginBuffer;
+	Buffer<float4> VertexFetch_BoneTransformBuffer;
+	Buffer<float4> VertexFetch_BoneLightmapBuffer;
+	Buffer<float4> VertexFetch_BonePrevTransformBuffer;
+	Buffer<uint>   VertexFetch_BoneMapBuffer;
+#endif
+
+// Since we don't have template functions it will have to be this way.
+#if MANUAL_VERTEX_FETCH
+	#define GCVF_GET_INPUT_VERTEX_ID(InputStruct) ((InputStruct).VertexId)
+#else // !MANUAL_VERTEX_FETCH
+	#define GCVF_GET_INPUT_VERTEX_ID(InputStruct) (0U)
+#endif
+
+
+/**
+ * Per-vertex inputs from bound vertex buffers
+ */
+struct FVertexFactoryInput
+{
+	float4	Position	: ATTRIBUTE0;
+
+#if !MANUAL_VERTEX_FETCH
+	#if METAL_PROFILE
+		float3	TangentX	: ATTRIBUTE1;
+		// TangentZ.w contains sign of tangent basis determinant
+		float4	TangentZ	: ATTRIBUTE2;
+
+		float4	Color		: ATTRIBUTE3;
+	#else
+		half3	TangentX	: ATTRIBUTE1;
+		// TangentZ.w contains sign of tangent basis determinant
+		half4	TangentZ	: ATTRIBUTE2;
+
+		half4	Color		: ATTRIBUTE3;
+	#endif
+#endif
+
+#if NUM_MATERIAL_TEXCOORDS_VERTEX
+	#if !MANUAL_VERTEX_FETCH
+		#if NUM_MATERIAL_TEXCOORDS_VERTEX > 1
+			float4	PackedTexCoords4[NUM_MATERIAL_TEXCOORDS_VERTEX/2] : ATTRIBUTE4;
+		#endif
+		#if NUM_MATERIAL_TEXCOORDS_VERTEX == 1
+			float2	PackedTexCoords2 : ATTRIBUTE4;
+		#elif NUM_MATERIAL_TEXCOORDS_VERTEX == 3
+			float2	PackedTexCoords2 : ATTRIBUTE5;
+		#elif NUM_MATERIAL_TEXCOORDS_VERTEX == 5
+			float2	PackedTexCoords2 : ATTRIBUTE6;
+		#elif NUM_MATERIAL_TEXCOORDS_VERTEX == 7
+			float2	PackedTexCoords2 : ATTRIBUTE7;
+		#endif
+	#endif
+#elif USE_PARTICLE_SUBUVS && !MANUAL_VERTEX_FETCH
+	float2	TexCoords[1] : ATTRIBUTE4;
+#endif
+
+#if VF_USE_PRIMITIVE_SCENE_DATA
+	uint InstanceIdOffset : ATTRIBUTE13;
+	uint DrawInstanceId : SV_InstanceID;
+#endif // VF_USE_PRIMITIVE_SCENE_DATA
+
+#if NEEDS_LIGHTMAP_COORDINATE && !MANUAL_VERTEX_FETCH
+	float2	LightMapCoordinate : ATTRIBUTE15;
+#endif
+
+#if MANUAL_VERTEX_FETCH
+	uint VertexId : SV_VertexID;
+#endif
+};
+
+#if RAYHITGROUPSHADER
+
+uint GetNumRayTracingDynamicMeshVerticesIndirect()
+{
+	return 0;
+}
+
+FVertexFactoryInput LoadVertexFactoryInputForHGS(uint TriangleIndex, int VertexIndex)
+{
+	FVertexFactoryInput Input = (FVertexFactoryInput)0;
+
+	FTriangleBaseAttributes Tri = LoadTriangleBaseAttributes(TriangleIndex);
+
+	Input.VertexId = Tri.Indices[VertexIndex];
+	Input.Position = float4(Tri.LocalPositions[VertexIndex], 1.0f);
+#if VF_USE_PRIMITIVE_SCENE_DATA
+	// Note: GetInstanceUserData() stores the GPU-Scene primitive ID
+	Input.InstanceIdOffset = GetInstanceUserData() | VF_TREAT_INSTANCE_ID_OFFSET_AS_PRIMITIVE_ID_FLAG;
+	Input.DrawInstanceId = 0;
+#endif // VF_USE_PRIMITIVE_SCENE_DATA
+
+	return Input;
+}
+#endif
+
+/** 
+ * Per-vertex inputs from bound vertex buffers.  Used by passes with a trimmed down position-only shader.
+ */
+struct FPositionOnlyVertexFactoryInput
+{
+	float4	Position	: ATTRIBUTE0;
+
+#if VF_USE_PRIMITIVE_SCENE_DATA
+	uint InstanceIdOffset : ATTRIBUTE1;
+	uint DrawInstanceId : SV_InstanceID;
+#endif // VF_USE_PRIMITIVE_SCENE_DATA
+
+#if MANUAL_VERTEX_FETCH
+	uint VertexId : SV_VertexID;
+#endif
+};
+
+/** 
+ * Per-vertex inputs from bound vertex buffers.  Used by passes with a trimmed down position-and-normal-only shader.
+ */
+struct FPositionAndNormalOnlyVertexFactoryInput
+{
+	float4	Position	: ATTRIBUTE0;
+	float4	Normal		: ATTRIBUTE2;
+
+#if VF_USE_PRIMITIVE_SCENE_DATA
+	uint InstanceIdOffset : ATTRIBUTE1;
+	uint DrawInstanceId : SV_InstanceID;
+#endif // VF_USE_PRIMITIVE_SCENE_DATA
+
+#if MANUAL_VERTEX_FETCH
+	uint VertexId : SV_VertexID;
+#endif
+};
+
+/** 
+ * Caches intermediates that would otherwise have to be computed multiple times.  Avoids relying on the compiler to optimize out redundant operations.
+ */
+struct FVertexFactoryIntermediates
+{
+#if VF_USE_PRIMITIVE_SCENE_DATA
+	FSceneDataIntermediates SceneData;
+#endif
+	float4x4 LocalToWorld;
+	float4x4 PrevLocalToWorld;
+
+	half3x3 TangentToLocal;
+	half3x3 TangentToWorld;
+	half TangentToWorldSign;
+
+	half4 Color;
+
+	uint PrimitiveId;
+	uint InstanceId;
+};
+
+/** Converts from vertex factory specific interpolants to a FMaterialPixelParameters, which is used by material inputs. */
+FMaterialPixelParameters GetMaterialPixelParameters(FVertexFactoryInterpolantsVSToPS Interpolants, float4 SvPosition)
+{
+	// GetMaterialPixelParameters is responsible for fully initializing the result
+	FMaterialPixelParameters Result = MakeInitializedMaterialPixelParameters();
+
+#if NUM_TEX_COORD_INTERPOLATORS
+	UNROLL
+	for (int CoordinateIndex = 0; CoordinateIndex < NUM_TEX_COORD_INTERPOLATORS; CoordinateIndex++)
+	{
+		Result.TexCoords[CoordinateIndex] = GetUV(Interpolants, CoordinateIndex);
+	}
+#endif
+
+#if USE_PARTICLE_SUBUVS
+	// Output TexCoord0 for when previewing materials that use ParticleSubUV.
+	Result.Particle.SubUVCoords[0] = GetUV(Interpolants, 0);
+	Result.Particle.SubUVCoords[1] = GetUV(Interpolants, 0);
+#endif
+
+	half3 TangentToWorld0 = GetTangentToWorld0(Interpolants).xyz;
+	half4 TangentToWorld2 = GetTangentToWorld2(Interpolants);
+	Result.UnMirrored = TangentToWorld2.w;
+
+	Result.VertexColor = GetColor(Interpolants);
+
+	// Required for previewing materials that use ParticleColor
+	Result.Particle.Color = half4(1,1,1,1);
+	Result.TangentToWorld = AssembleTangentToWorld( TangentToWorld0, TangentToWorld2 );
+#if USE_WORLDVERTEXNORMAL_CENTER_INTERPOLATION
+	Result.WorldVertexNormal_Center = Interpolants.TangentToWorld2_Center.xyz;
+#endif
+
+#if LIGHTMAP_UV_ACCESS && NEEDS_LIGHTMAP_COORDINATE
+	#if (ES3_1_PROFILE)
+		// Not supported in pixel shader
+		Result.LightmapUVs = float2(0, 0);
+	#else
+		Result.LightmapUVs = Interpolants.LightMapCoordinate.xy;
+	#endif
+#endif
+
+	Result.TwoSidedSign = 1;
+	Result.PrimitiveId = GetPrimitiveId(Interpolants);
+
+#if NEEDS_PARTICLE_LOCAL_TO_WORLD
+	Result.Particle.ParticleToWorld = GetPrimitiveData(Result.PrimitiveId).LocalToWorld;
+#endif
+
+#if NEEDS_PARTICLE_WORLD_TO_LOCAL
+	Result.Particle.WorldToParticle = GetPrimitiveData(Result.PrimitiveId).WorldToLocal;
+#endif
+
+	return Result;
+}
+
+/** Converts from vertex factory specific input to a FMaterialVertexParameters, which is used by vertex shader material inputs. */
+FMaterialVertexParameters GetMaterialVertexParameters(FVertexFactoryInput Input, FVertexFactoryIntermediates Intermediates, float3 WorldPosition, half3x3 TangentToLocal)
+{
+	FMaterialVertexParameters Result = (FMaterialVertexParameters)0;
+	Result.WorldPosition = WorldPosition;
+	Result.VertexColor = Intermediates.Color;
+	Result.TangentToWorld = Intermediates.TangentToWorld;
+
+	Result.PrevFrameLocalToWorld = Intermediates.PrevLocalToWorld;
+
+	Result.PreSkinnedPosition = Input.Position.xyz;
+	Result.PreSkinnedNormal = TangentToLocal[2];
+
+#if MANUAL_VERTEX_FETCH && NUM_MATERIAL_TEXCOORDS_VERTEX
+	const uint NumFetchTexCoords = GeometryCollectionVF.VertexFetch_Parameters[VF_NumTexcoords_Index];
+	UNROLL
+	for (uint CoordinateIndex = 0; CoordinateIndex < NUM_MATERIAL_TEXCOORDS_VERTEX; CoordinateIndex++)
+	{
+		// Clamp coordinates to mesh's maximum as materials can request more than are available
+		uint ClampedCoordinateIndex = min(CoordinateIndex, NumFetchTexCoords - 1);
+		Result.TexCoords[CoordinateIndex] = GeometryCollectionVF.VertexFetch_TexCoordBuffer[NumFetchTexCoords * (GeometryCollectionVF.VertexFetch_Parameters[VF_VertexOffset] + Input.VertexId) + ClampedCoordinateIndex];
+	}
+#elif NUM_MATERIAL_TEXCOORDS_VERTEX
+	#if NUM_MATERIAL_TEXCOORDS_VERTEX > 1
+		UNROLL
+		for (int CoordinateIndex = 0; CoordinateIndex < NUM_MATERIAL_TEXCOORDS_VERTEX - 1; CoordinateIndex += 2)
+		{
+			Result.TexCoords[CoordinateIndex] = Input.PackedTexCoords4[CoordinateIndex / 2].xy;
+			if (CoordinateIndex + 1 < NUM_MATERIAL_TEXCOORDS_VERTEX)
+			{
+				Result.TexCoords[CoordinateIndex+1] = Input.PackedTexCoords4[CoordinateIndex / 2].zw;
+			}
+		}
+	#endif
+	#if NUM_MATERIAL_TEXCOORDS_VERTEX % 2 == 1
+		Result.TexCoords[NUM_MATERIAL_TEXCOORDS_VERTEX-1] = Input.PackedTexCoords2;
+	#endif
+#endif  //MANUAL_VERTEX_FETCH && NUM_MATERIAL_TEXCOORDS_VERTEX
+
+	Result.PrimitiveId = Intermediates.PrimitiveId;
+
+#if NEEDS_PARTICLE_LOCAL_TO_WORLD
+	Result.Particle.ParticleToWorld = GetPrimitiveData(Result.PrimitiveId).LocalToWorld;
+#endif
+
+#if NEEDS_PARTICLE_WORLD_TO_LOCAL
+	Result.Particle.WorldToParticle = GetPrimitiveData(Result.PrimitiveId).WorldToLocal;
+#endif
+
+	Result.MaterialVertexAttributes = EvaluateVertexMaterialAttributes(Result);
+
+	return Result;
+}
+
+float4 CalcWorldPosition(float4 Position, uint PrimitiveId)
+{
+	return TransformLocalToTranslatedWorld(Position.xyz, PrimitiveId);
+}
+
+half3x3 LoadTangentToLocal(FVertexFactoryInput Input, inout float TangentSignOut)
+{
+#if MANUAL_VERTEX_FETCH
+	half3 TangentInputX = GeometryCollectionVF.VertexFetch_PackedTangentsBuffer[2 * (GeometryCollectionVF.VertexFetch_Parameters[VF_VertexOffset] + Input.VertexId) + 0].xyz;
+	half4 TangentInputZ = GeometryCollectionVF.VertexFetch_PackedTangentsBuffer[2 * (GeometryCollectionVF.VertexFetch_Parameters[VF_VertexOffset] + Input.VertexId) + 1].xyzw;
+#else
+	half3 TangentInputX = Input.TangentX;
+	half4 TangentInputZ = Input.TangentZ;
+#endif
+
+	half3 TangentX = TangentBias(TangentInputX);
+	half4 TangentZ = TangentBias(TangentInputZ);
+
+	TangentSignOut = TangentZ.w;
+
+	// derive the binormal by getting the cross product of the normal and tangent
+	half3 TangentY = cross(TangentZ.xyz, TangentX) * TangentZ.w;
+	
+	// Recalculate TangentX off of the other two vectors
+	// This corrects quantization error since TangentX was passed in as a quantized vertex input
+	// The error shows up most in specular off of a mesh with a smoothed UV seam (normal is smooth, but tangents vary across the seam)
+	half3x3 Result;
+	Result[0] = cross(TangentY, TangentZ.xyz) * TangentZ.w;
+	Result[1] = TangentY;
+	Result[2] = TangentZ.xyz;
+
+	return Result;
+}
+
+half3x3 CalcTangentToWorld(half3 InvScale, half3x3 TangentToLocal, float4x4 LocalToWorld)
+{
+	half3x3 LocalToWorld3x3 = (half3x3)LocalToWorld;
+	LocalToWorld3x3[0] *= InvScale.x;
+	LocalToWorld3x3[1] *= InvScale.y;
+	LocalToWorld3x3[2] *= InvScale.z;
+	return mul(TangentToLocal, LocalToWorld3x3);
+}
+
+/**
+ * Temporary info to make it possible to load transform for position only & position/normal paths.
+ */
+struct FTransformData
+{
+#if VF_USE_PRIMITIVE_SCENE_DATA
+	FSceneDataIntermediates SceneData;
+#endif
+	float4x4 LocalToWorld;
+	float4x4 PrevLocalToWorld;
+	float3 InvScale;
+};
+
+FTransformData LoadTransformData(uint InstanceIdOffset, uint DrawInstanceId, uint VertexId)
+{
+	FTransformData TransformData;
+#if VF_USE_PRIMITIVE_SCENE_DATA
+	TransformData.SceneData = GetSceneDataIntermediates(InstanceIdOffset, DrawInstanceId);
+	TransformData.InvScale = TransformData.SceneData.InstanceData.InvNonUniformScaleAndDeterminantSign.xyz;
+
+	float4x4 InstanceLocalToWorld = TransformData.SceneData.InstanceData.LocalToWorld;
+	float4x4 PrevInstanceLocalToWorld = TransformData.SceneData.InstanceData.PrevLocalToWorld;
+#else
+	TransformData.InvScale = Primitive.InvNonUniformScaleAndDeterminantSign.xyz;
+
+	// Note: if there are no instances, this is the same as the InstanceLocalToWorld
+	float4x4 InstanceLocalToWorld = Primitive.LocalToWorld;
+	float4x4 PrevInstanceLocalToWorld = Primitive.PreviousLocalToWorld;
+#endif
+
+#if MANUAL_VERTEX_FETCH
+	uint BoneIndex = VertexFetch_BoneMapBuffer[GeometryCollectionVF.VertexFetch_Parameters[VF_VertexOffset] + VertexId];
+	float4x4 LocalToInstance = float4x4(
+		float4(VertexFetch_BoneTransformBuffer[4 * BoneIndex + 0].xyz, 0.0f),
+		float4(VertexFetch_BoneTransformBuffer[4 * BoneIndex + 1].xyz, 0.0f),
+		float4(VertexFetch_BoneTransformBuffer[4 * BoneIndex + 2].xyz, 0.0f),
+		float4(VertexFetch_BoneTransformBuffer[4 * BoneIndex + 3].xyz, 1.0f));
+
+	TransformData.LocalToWorld = mul(LocalToInstance, InstanceLocalToWorld);
+
+	float4x4 PrevLocalToInstance = float4x4(
+		float4(VertexFetch_BonePrevTransformBuffer[4 * BoneIndex + 0].xyz, 0.0f),
+		float4(VertexFetch_BonePrevTransformBuffer[4 * BoneIndex + 1].xyz, 0.0f),
+		float4(VertexFetch_BonePrevTransformBuffer[4 * BoneIndex + 2].xyz, 0.0f),
+		float4(VertexFetch_BonePrevTransformBuffer[4 * BoneIndex + 3].xyz, 1.0f));
+
+	TransformData.PrevLocalToWorld = mul(PrevLocalToInstance, PrevInstanceLocalToWorld);
+#else
+	// In this case the LTW because it already is skinned on CPU
+	// Note: this changes the meaning of tangent to local etc, I assume that is fine because that is how it used to be
+	TransformData.LocalToWorld = InstanceLocalToWorld;
+	TransformData.PrevLocalToWorld = InstanceLocalToWorld;
+#endif
+	return TransformData;
+}
+
+FVertexFactoryIntermediates GetVertexFactoryIntermediates(FVertexFactoryInput Input)
+{
+	FVertexFactoryIntermediates Intermediates = (FVertexFactoryIntermediates)0;
+#if VF_USE_PRIMITIVE_SCENE_DATA
+	FTransformData TransformData = LoadTransformData(Input.InstanceIdOffset, Input.DrawInstanceId, GCVF_GET_INPUT_VERTEX_ID(Input));
+	Intermediates.SceneData = TransformData.SceneData;
+	// Just aliases, remove later probably...
+	Intermediates.PrimitiveId = Intermediates.SceneData.PrimitiveId;
+	Intermediates.InstanceId = Intermediates.SceneData.InstanceId;
+	float DeterminantSign = Intermediates.SceneData.InstanceData.InvNonUniformScaleAndDeterminantSign.w;
+#else 
+	FTransformData TransformData = LoadTransformData(0U, 0U, GCVF_GET_INPUT_VERTEX_ID(Input));
+	float DeterminantSign = Primitive.InvNonUniformScaleAndDeterminantSign.w;
+#endif
+
+	Intermediates.LocalToWorld = TransformData.LocalToWorld;
+	Intermediates.PrevLocalToWorld = TransformData.PrevLocalToWorld;
+
+#if MANUAL_VERTEX_FETCH
+	Intermediates.Color = GeometryCollectionVF.VertexFetch_ColorComponentsBuffer[(GeometryCollectionVF.VertexFetch_Parameters[VF_VertexOffset] + Input.VertexId) & GeometryCollectionVF.VertexFetch_Parameters[VF_ColorIndexMask_Index]] FMANUALFETCH_COLOR_COMPONENT_SWIZZLE; // Swizzle vertex color.
+#else
+	Intermediates.Color = Input.Color FCOLOR_COMPONENT_SWIZZLE; // Swizzle vertex color.
+#endif
+
+	float TangentSign = 1.0;
+	Intermediates.TangentToLocal = LoadTangentToLocal(Input, TangentSign);
+	Intermediates.TangentToWorld = CalcTangentToWorld(TransformData.InvScale, Intermediates.TangentToLocal, Intermediates.LocalToWorld);
+
+	Intermediates.TangentToWorldSign = TangentSign * DeterminantSign;
+
+	return Intermediates;
+}
+
+/**
+* Get the 3x3 tangent basis vectors for this vertex factory
+* this vertex factory will calculate the binormal on-the-fly
+*
+* @param Input - vertex input stream structure
+* @return 3x3 matrix
+*/
+half3x3 VertexFactoryGetTangentToLocal( FVertexFactoryInput Input, FVertexFactoryIntermediates Intermediates )
+{
+	return Intermediates.TangentToLocal;
+}
+
+// @return translated world position
+float4 VertexFactoryGetWorldPosition(FVertexFactoryInput Input, FVertexFactoryIntermediates Intermediates)
+{
+	return TransformLocalToTranslatedWorld(Input.Position.xyz, Intermediates.LocalToWorld);
+}
+
+float4 VertexFactoryGetRasterizedWorldPosition(FVertexFactoryInput Input, FVertexFactoryIntermediates Intermediates, float4 InWorldPosition)
+{
+	return InWorldPosition;
+}
+
+float3 VertexFactoryGetPositionForVertexLighting(FVertexFactoryInput Input, FVertexFactoryIntermediates Intermediates, float3 TranslatedWorldPosition)
+{
+	return TranslatedWorldPosition;
+}
+
+FVertexFactoryInterpolantsVSToPS VertexFactoryGetInterpolantsVSToPS(FVertexFactoryInput Input, FVertexFactoryIntermediates Intermediates, FMaterialVertexParameters VertexParameters)
+{
+	FVertexFactoryInterpolantsVSToPS Interpolants;
+
+	// Initialize the whole struct to 0
+	// Really only the last two components of the packed UVs have the opportunity to be uninitialized
+	Interpolants = (FVertexFactoryInterpolantsVSToPS)0;
+
+#if NUM_TEX_COORD_INTERPOLATORS
+	float2 CustomizedUVs[NUM_TEX_COORD_INTERPOLATORS];
+	GetMaterialCustomizedUVs(VertexParameters, CustomizedUVs);
+	GetCustomInterpolators(VertexParameters, CustomizedUVs);
+	
+	UNROLL
+	for (int CoordinateIndex = 0; CoordinateIndex < NUM_TEX_COORD_INTERPOLATORS; CoordinateIndex++)
+	{
+		SetUV(Interpolants, CoordinateIndex, CustomizedUVs[CoordinateIndex]);
+	}
+
+#elif NUM_MATERIAL_TEXCOORDS_VERTEX == 0 && USE_PARTICLE_SUBUVS
+	#if MANUAL_VERTEX_FETCH
+		SetUV(Interpolants, 0, GeometryCollectionVF.VertexFetch_TexCoordBuffer[GeometryCollectionVF.VertexFetch_Parameters[VF_NumTexcoords_Index] * (GeometryCollectionVF.VertexFetch_Parameters[VF_VertexOffset] + Input.VertexId)]);
+	#else
+		SetUV(Interpolants, 0, Input.TexCoords[0]);
+	#endif
+#endif
+
+#if NEEDS_LIGHTMAP_COORDINATE
+	float2 LightMapCoordinate = 0;
+	float2 ShadowMapCoordinate = 0;
+	#if MANUAL_VERTEX_FETCH
+		float2 LightMapCoordinateInput = GeometryCollectionVF.VertexFetch_TexCoordBuffer[GeometryCollectionVF.VertexFetch_Parameters[VF_NumTexcoords_Index] * (GeometryCollectionVF.VertexFetch_Parameters[VF_VertexOffset] + Input.VertexId) + GeometryCollectionVF.VertexFetch_Parameters[FV_LightMapIndex_Index]];
+	#else
+		float2 LightMapCoordinateInput = Input.LightMapCoordinate;
+	#endif
+
+	uint LightmapDataIndex = 0;
+	uint LightmapUVIndex   = 0;
+
+#if VF_USE_PRIMITIVE_SCENE_DATA
+	LightmapDataIndex = Intermediates.SceneData.Primitive.LightmapDataIndex + GeometryCollectionVF.LODLightmapDataIndex;
+	LightmapUVIndex = Intermediates.SceneData.Primitive.LightmapUVIndex;
+#endif
+
+	float4 LightMapCoordinateScaleBias = GetLightmapData(LightmapDataIndex).LightMapCoordinateScaleBias;
+	LightMapCoordinate = LightMapCoordinateInput * LightMapCoordinateScaleBias.xy + LightMapCoordinateScaleBias.zw;
+
+#if STATICLIGHTING_TEXTUREMASK
+	float4 ShadowMapCoordinateScaleBias = GetLightmapData(LightmapDataIndex).ShadowMapCoordinateScaleBias;
+	ShadowMapCoordinate = LightMapCoordinateInput * ShadowMapCoordinateScaleBias.xy + ShadowMapCoordinateScaleBias.zw;
+#endif
+
+	SetLightMapCoordinate(Interpolants, LightMapCoordinate, ShadowMapCoordinate);
+	SetLightmapDataIndex(Interpolants, LightmapDataIndex);
+
+#endif	// NEEDS_LIGHTMAP_COORDINATE
+
+	SetTangents(Interpolants, Intermediates.TangentToWorld[0], Intermediates.TangentToWorld[2], Intermediates.TangentToWorldSign);
+	SetColor(Interpolants, Intermediates.Color);
+
+#if INSTANCED_STEREO
+	Interpolants.EyeIndex = 0;
+#endif
+
+	SetPrimitiveId(Interpolants, Intermediates.PrimitiveId);
+
+	return Interpolants;
+}
+
+/** for depth-only pass */
+float4 VertexFactoryGetWorldPosition(FPositionOnlyVertexFactoryInput Input)
+{
+#if VF_USE_PRIMITIVE_SCENE_DATA
+	FTransformData TransformData = LoadTransformData(Input.InstanceIdOffset, Input.DrawInstanceId, GCVF_GET_INPUT_VERTEX_ID(Input));
+#else 
+	FTransformData TransformData = LoadTransformData(0U, 0U, GCVF_GET_INPUT_VERTEX_ID(Input));
+#endif
+
+	return TransformLocalToTranslatedWorld(Input.Position.xyz, TransformData.LocalToWorld);
+}
+
+/** for depth-only pass (slope depth bias) */
+float4 VertexFactoryGetWorldPosition(FPositionAndNormalOnlyVertexFactoryInput Input)
+{
+#if VF_USE_PRIMITIVE_SCENE_DATA
+	FTransformData TransformData = LoadTransformData(Input.InstanceIdOffset, Input.DrawInstanceId, GCVF_GET_INPUT_VERTEX_ID(Input));
+#else 
+	FTransformData TransformData = LoadTransformData(0U, 0U, GCVF_GET_INPUT_VERTEX_ID(Input));
+#endif
+	return TransformLocalToTranslatedWorld(Input.Position.xyz, TransformData.LocalToWorld);
+}
+
+float3 VertexFactoryGetWorldNormal(FPositionAndNormalOnlyVertexFactoryInput Input)
+{
+#if VF_USE_PRIMITIVE_SCENE_DATA
+	FTransformData TransformData = LoadTransformData(Input.InstanceIdOffset, Input.DrawInstanceId, GCVF_GET_INPUT_VERTEX_ID(Input));
+#else 
+	FTransformData TransformData = LoadTransformData(0U, 0U, GCVF_GET_INPUT_VERTEX_ID(Input));
+#endif
+	return RotateLocalToWorld(Input.Position.xyz, TransformData.LocalToWorld, TransformData.InvScale).xyz;
+}
+
+float3 VertexFactoryGetWorldNormal(FVertexFactoryInput Input, FVertexFactoryIntermediates Intermediates)
+{
+	return Intermediates.TangentToWorld[2];
+}
+
+// @return previous translated world position
+float4 VertexFactoryGetPreviousWorldPosition(FVertexFactoryInput Input, FVertexFactoryIntermediates Intermediates)
+{
+	return TransformLocalToTranslatedWorld(Input.Position.xyz, Intermediates.PrevLocalToWorld);
+}
+
+float4 VertexFactoryGetTranslatedPrimitiveVolumeBounds(FVertexFactoryInterpolantsVSToPS Interpolants)
+{
+	float4 ObjectWorldPositionAndRadius = GetPrimitiveData(GetPrimitiveId(Interpolants)).ObjectWorldPositionAndRadius;
+	return float4(ObjectWorldPositionAndRadius.xyz + ResolvedView.PreViewTranslation.xyz, ObjectWorldPositionAndRadius.w);
+}
+
+uint VertexFactoryGetPrimitiveId(FVertexFactoryInterpolantsVSToPS Interpolants)
+{
+	return GetPrimitiveId(Interpolants);
+}
+
+#undef GCVF_GET_INPUT_VERTEX_ID
+
+#include "VertexFactoryDefaultInterface.ush"
diff --git a/Engine/Shaders/Private/LocalVertexFactory.ush b/Engine/Shaders/Private/LocalVertexFactory.ush
--- a/Engine/Shaders/Private/LocalVertexFactory.ush
+++ b/Engine/Shaders/Private/LocalVertexFactory.ush
@@ -1,7 +1,7 @@
 // Copyright Epic Games, Inc. All Rights Reserved.
 
 /*=============================================================================
-	LocalVertexFactory.usf: Local vertex factory shader code.
+	LocalVertexFactory.ush: Local vertex factory shader code.
 =============================================================================*/
 
 #include "VertexFactoryCommon.ush"
@@ -73,17 +73,9 @@ float4 SplineParams[10];
 	#define FV_LightMapIndex_Index 2
 	#define VF_VertexOffset 3
 
-
 	Buffer<float4> VertexFetch_InstanceOriginBuffer;
 	Buffer<float4> VertexFetch_InstanceTransformBuffer;
 	Buffer<float4> VertexFetch_InstanceLightmapBuffer;
-
-	#if USE_INSTANCING && USE_INSTANCING_BONEMAP
-		Buffer<float4> VertexFetch_InstancePrevTransformBuffer;
-		Buffer<uint> VertexFetch_InstanceBoneMapBuffer;
-	#endif
-
-
 #endif //! MANUAL_VERTEX_FETCH
 
 /**
@@ -364,13 +356,6 @@ struct FVertexFactoryIntermediates
 	float4 InstanceTransform2;
 	float4 InstanceTransform3;
 
-	#if USE_INSTANCING_BONEMAP
-		float4 InstancePrevOrigin;
-		float4 InstancePrevTransform1;
-		float4 InstancePrevTransform2;
-		float4 InstancePrevTransform3;
-	#endif
-
 	float4 InstanceLightmapAndShadowMapUVBias;
 
 	// x = per-instance random, y = per-instance fade out amount, z = hide/show flag, w dither fade cutoff
@@ -398,67 +383,44 @@ float4x4 GetInstanceTransform(FVertexFactoryIntermediates Intermediates)
 
 float4x4 GetInstancePrevTransform(FVertexFactoryIntermediates Intermediates)
 {
-#if USE_INSTANCING_BONEMAP
-	// when using geometry collections, we can correctly compute velocities because we pass along previous transforms.
-	return float4x4(
-		float4(Intermediates.InstancePrevTransform1.xyz, 0.0f),
-		float4(Intermediates.InstancePrevTransform2.xyz, 0.0f),
-		float4(Intermediates.InstancePrevTransform3.xyz, 0.0f),
-		float4(Intermediates.InstancePrevOrigin.xyz, 1.0f));
-#else
 	// Assumes instance transform never change, which means per-instance motion will cause TAA and motion blur artifacts.
 	return GetInstanceTransform(Intermediates);
-#endif
 }
 
 float4x4 GetInstanceTransform(FPositionOnlyVertexFactoryInput Input)
 {
-	#if MANUAL_VERTEX_FETCH && !USE_INSTANCING_BONEMAP
+#if MANUAL_VERTEX_FETCH
 	uint InstanceId = GetInstanceId(Input.InstanceId);
 	return float4x4(
 			float4(InstanceVF.VertexFetch_InstanceTransformBuffer[3 * (InstanceId + InstanceOffset) + 0].xyz, 0.0f),
 			float4(InstanceVF.VertexFetch_InstanceTransformBuffer[3 * (InstanceId + InstanceOffset) + 1].xyz, 0.0f),
 			float4(InstanceVF.VertexFetch_InstanceTransformBuffer[3 * (InstanceId + InstanceOffset) + 2].xyz, 0.0f),
 			float4(InstanceVF.VertexFetch_InstanceOriginBuffer[(InstanceId + InstanceOffset)].xyz, 1.0f));
-	#elif MANUAL_VERTEX_FETCH && USE_INSTANCING_BONEMAP
-	uint InstanceIndex = VertexFetch_InstanceBoneMapBuffer[LocalVF.VertexFetch_Parameters[VF_VertexOffset] + Input.VertexId];
-	return float4x4(
-		float4(VertexFetch_InstanceTransformBuffer[4 * InstanceIndex + 0].xyz, 0.0f),
-		float4(VertexFetch_InstanceTransformBuffer[4 * InstanceIndex + 1].xyz, 0.0f),
-		float4(VertexFetch_InstanceTransformBuffer[4 * InstanceIndex + 2].xyz, 0.0f),
-		float4(VertexFetch_InstanceTransformBuffer[4 * InstanceIndex + 3].xyz, 1.0f));
-	#else 
+#else 
 	return float4x4(
 		float4(Input.InstanceTransform1.xyz, 0.0f),
 		float4(Input.InstanceTransform2.xyz, 0.0f),
 		float4(Input.InstanceTransform3.xyz, 0.0f),
 		float4(Input.InstanceOrigin.xyz, 1.0f));
-	#endif
+#endif
 }
 
 float4x4 GetInstanceTransform(FPositionAndNormalOnlyVertexFactoryInput Input)
 {
-	#if MANUAL_VERTEX_FETCH && !USE_INSTANCING_BONEMAP
+#if MANUAL_VERTEX_FETCH
 	uint InstanceId = GetInstanceId(Input.InstanceId);
 	return float4x4(
 			float4(InstanceVF.VertexFetch_InstanceTransformBuffer[3 * (InstanceId + InstanceOffset) + 0].xyz, 0.0f),
 			float4(InstanceVF.VertexFetch_InstanceTransformBuffer[3 * (InstanceId + InstanceOffset) + 1].xyz, 0.0f),
 			float4(InstanceVF.VertexFetch_InstanceTransformBuffer[3 * (InstanceId + InstanceOffset) + 2].xyz, 0.0f),
 			float4(InstanceVF.VertexFetch_InstanceOriginBuffer[(InstanceId + InstanceOffset)].xyz, 1.0f));
-	#elif MANUAL_VERTEX_FETCH && USE_INSTANCING_BONEMAP
-		uint InstanceIndex = VertexFetch_InstanceBoneMapBuffer[LocalVF.VertexFetch_Parameters[VF_VertexOffset] + Input.VertexId];
-		return float4x4(
-			float4(VertexFetch_InstanceTransformBuffer[4 * InstanceIndex + 0].xyz, 0.0f),
-			float4(VertexFetch_InstanceTransformBuffer[4 * InstanceIndex + 1].xyz, 0.0f),
-			float4(VertexFetch_InstanceTransformBuffer[4 * InstanceIndex + 2].xyz, 0.0f),
-			float4(VertexFetch_InstanceTransformBuffer[4 * InstanceIndex + 3].xyz, 1.0f));
-	#else 
+#else 
 	return float4x4(
 		float4(Input.InstanceTransform1.xyz, 0.0f),
 		float4(Input.InstanceTransform2.xyz, 0.0f),
 		float4(Input.InstanceTransform3.xyz, 0.0f),
 		float4(Input.InstanceOrigin.xyz, 1.0f));
-	#endif
+#endif
 }
 
 half3x3 GetInstanceToLocal3x3(FVertexFactoryIntermediates Intermediates)
@@ -913,26 +875,13 @@ FVertexFactoryIntermediates GetVertexFactoryIntermediates(FVertexFactoryInput In
 	Intermediates.Color = Input.Color FCOLOR_COMPONENT_SWIZZLE; // Swizzle vertex color.
 #endif
 
-#if USE_INSTANCING && MANUAL_VERTEX_FETCH && !USE_INSTANCING_BONEMAP
+#if USE_INSTANCING && MANUAL_VERTEX_FETCH
 	uint InstanceId = GetInstanceId(Input.InstanceId);
 	Intermediates.InstanceTransform1 = InstanceVF.VertexFetch_InstanceTransformBuffer[3 * (InstanceId + InstanceOffset) + 0];
 	Intermediates.InstanceTransform2 = InstanceVF.VertexFetch_InstanceTransformBuffer[3 * (InstanceId + InstanceOffset) + 1];
 	Intermediates.InstanceTransform3 = InstanceVF.VertexFetch_InstanceTransformBuffer[3 * (InstanceId + InstanceOffset) + 2];
 	Intermediates.InstanceOrigin = InstanceVF.VertexFetch_InstanceOriginBuffer[(InstanceId + InstanceOffset)];
 	Intermediates.InstanceLightmapAndShadowMapUVBias = InstanceVF.VertexFetch_InstanceLightmapBuffer[(InstanceId + InstanceOffset)];
-#elif MANUAL_VERTEX_FETCH && USE_INSTANCING_BONEMAP
-	uint InstanceIndex = VertexFetch_InstanceBoneMapBuffer[LocalVF.VertexFetch_Parameters[VF_VertexOffset] + Input.VertexId];
-	Intermediates.InstanceTransform1 = VertexFetch_InstanceTransformBuffer[4 * InstanceIndex + 0];
-	Intermediates.InstanceTransform2 = VertexFetch_InstanceTransformBuffer[4 * InstanceIndex + 1];
-	Intermediates.InstanceTransform3 = VertexFetch_InstanceTransformBuffer[4 * InstanceIndex + 2];
-	Intermediates.InstanceOrigin = VertexFetch_InstanceTransformBuffer[4 * InstanceIndex + 3];	
-
-	Intermediates.InstancePrevTransform1 = VertexFetch_InstancePrevTransformBuffer[4 * InstanceIndex + 0];
-	Intermediates.InstancePrevTransform2 = VertexFetch_InstancePrevTransformBuffer[4 * InstanceIndex + 1];
-	Intermediates.InstancePrevTransform3 = VertexFetch_InstancePrevTransformBuffer[4 * InstanceIndex + 2];
-	Intermediates.InstancePrevOrigin = VertexFetch_InstancePrevTransformBuffer[4 * InstanceIndex + 3];	
-
-	Intermediates.InstanceLightmapAndShadowMapUVBias = float4(0,0,0,0);
 #elif USE_INSTANCING
 	Intermediates.InstanceTransform1 = Input.InstanceTransform1;
 	Intermediates.InstanceTransform2 = Input.InstanceTransform2;
@@ -950,7 +899,7 @@ FVertexFactoryIntermediates GetVertexFactoryIntermediates(FVertexFactoryInput In
 	Intermediates.TangentToWorldSign = TangentSign * GetPrimitiveData(Intermediates.PrimitiveId).InvNonUniformScaleAndDeterminantSign.w;
 #endif
 
-#if USE_INSTANCING && !USE_INSTANCING_BONEMAP
+#if USE_INSTANCING
 	// x = per-instance random, y = per-instance fade out factor, z = zero or one depending of if it is shown at all, w is dither cutoff 
 
 	// PerInstanceParams.z stores a hide/show flag for this instance
@@ -973,11 +922,6 @@ FVertexFactoryIntermediates GetVertexFactoryIntermediates(FVertexFactoryInput In
 	#else
 		Intermediates.PerInstanceParams.w = 0;
 	#endif
-#elif USE_INSTANCING && USE_INSTANCING_BONEMAP	
-	Intermediates.PerInstanceParams.x = 0;	
-	Intermediates.PerInstanceParams.y = 1;
-	Intermediates.PerInstanceParams.z = 1;
-	Intermediates.PerInstanceParams.w = 0;
 #endif	// USE_INSTANCING
 
 #if GPUSKIN_PASS_THROUGH
diff --git a/Engine/Shaders/Private/VertexFactoryCommon.ush b/Engine/Shaders/Private/VertexFactoryCommon.ush
--- a/Engine/Shaders/Private/VertexFactoryCommon.ush
+++ b/Engine/Shaders/Private/VertexFactoryCommon.ush
@@ -27,11 +27,24 @@ float4 TransformLocalToWorld(float3 LocalPosition)
 	return float4(RotatedPosition + Primitive.LocalToWorld[3].xyz ,1);
 }
 
-float4 TransformLocalToTranslatedWorld(float3 LocalPosition, uint PrimitiveId)
+float4 TransformLocalToTranslatedWorld(float3 LocalPosition, float4x4 LocalToWorld)
 {
-	float4x4 LocalToWorld = GetPrimitiveData(PrimitiveId).LocalToWorld;
 	float3 RotatedPosition = LocalToWorld[0].xyz * LocalPosition.xxx + LocalToWorld[1].xyz * LocalPosition.yyy + LocalToWorld[2].xyz * LocalPosition.zzz;
-	return float4(RotatedPosition + (LocalToWorld[3].xyz + ResolvedView.PreViewTranslation.xyz),1);
+	return float4(RotatedPosition + (LocalToWorld[3].xyz + ResolvedView.PreViewTranslation.xyz), 1.0f);
+}
+
+float4 TransformLocalToTranslatedWorld(float3 LocalPosition, uint PrimitiveId)
+{
+	return TransformLocalToTranslatedWorld(LocalPosition, GetPrimitiveData(PrimitiveId).LocalToWorld);
+}
+
+
+float3 RotateLocalToWorld(float3 LocalDirection, float4x4 LocalToWorld, float3 InvScale)
+{
+	return
+		InvScale.x * LocalToWorld[0].xyz * LocalDirection.xxx +
+		InvScale.y * LocalToWorld[1].xyz * LocalDirection.yyy +
+		InvScale.z * LocalToWorld[2].xyz * LocalDirection.zzz;
 }
 
 #if GPUCULL_TODO
@@ -46,17 +59,13 @@ float4 TransformLocalToTranslatedWorld(float3 LocalPosition, FInstanceSceneData
 
 float3 RotateLocalToWorld(float3 LocalDirection, FInstanceSceneData InstanceSceneData)
 {
-	const float4x4 LocalToWorld = InstanceSceneData.LocalToWorld;
-	const float3 InvScale = InstanceSceneData.InvNonUniformScaleAndDeterminantSign.xyz;
-	return
-		InvScale.x * LocalToWorld[0].xyz * LocalDirection.xxx +
-		InvScale.y * LocalToWorld[1].xyz * LocalDirection.yyy +
-		InvScale.z * LocalToWorld[2].xyz * LocalDirection.zzz;
+	return RotateLocalToWorld(LocalDirection, InstanceSceneData.LocalToWorld, InstanceSceneData.InvNonUniformScaleAndDeterminantSign.xyz);
 }
 
 #endif // INSTANCE_SCENE_DATA_SUPPORTED
 #endif // GPUCULL_TODO
 
+
 #if (!GPUCULL_TODO || !VF_USE_PRIMITIVE_SCENE_DATA)
 
 float4 TransformLocalToTranslatedWorld(float3 LocalPosition)
@@ -67,21 +76,12 @@ float4 TransformLocalToTranslatedWorld(float3 LocalPosition)
 
 float3 RotateLocalToWorld(float3 LocalDirection, uint PrimitiveId)
 {
-	const float4x4 LocalToWorld = GetPrimitiveData(PrimitiveId).LocalToWorld;
-	const float3 InvScale = GetPrimitiveData(PrimitiveId).InvNonUniformScaleAndDeterminantSign.xyz;
-	return 
-		InvScale.x * LocalToWorld[0].xyz * LocalDirection.xxx + 
-		InvScale.y * LocalToWorld[1].xyz * LocalDirection.yyy + 
-		InvScale.z * LocalToWorld[2].xyz * LocalDirection.zzz;	
+	return RotateLocalToWorld(LocalDirection, GetPrimitiveData(PrimitiveId).LocalToWorld, GetPrimitiveData(PrimitiveId).InvNonUniformScaleAndDeterminantSign.xyz);
 }
 
 float3 RotateLocalToWorld(float3 LocalDirection)
 {
-	const float3 InvScale = Primitive.InvNonUniformScaleAndDeterminantSign.xyz;
-	return 
-		InvScale.x * Primitive.LocalToWorld[0].xyz * LocalDirection.xxx + 
-		InvScale.y * Primitive.LocalToWorld[1].xyz * LocalDirection.yyy + 
-		InvScale.z * Primitive.LocalToWorld[2].xyz * LocalDirection.zzz;
+	return RotateLocalToWorld(LocalDirection, Primitive.LocalToWorld, Primitive.InvNonUniformScaleAndDeterminantSign.xyz);
 }
 
 float3 RotateWorldToLocal(float3 WorldDirection)
diff --git a/Engine/Source/Runtime/Engine/Private/GeometryCollectionRendering.cpp b/Engine/Source/Runtime/Engine/Private/GeometryCollectionRendering.cpp
--- a/Engine/Source/Runtime/Engine/Private/GeometryCollectionRendering.cpp
+++ b/Engine/Source/Runtime/Engine/Private/GeometryCollectionRendering.cpp
@@ -1,17 +1,39 @@
 // Copyright Epic Games, Inc. All Rights Reserved.
 
 #include "GeometryCollectionRendering.h"
+#include "MeshBatch.h"
 #include "MeshDrawShaderBindings.h"
 #include "MeshMaterialShader.h"
+#include "ShaderParameterUtils.h"
+#include "Rendering/ColorVertexBuffer.h"
+#include "ProfilingDebugging/LoadTimeTracker.h"
 
 IMPLEMENT_TYPE_LAYOUT(FGeometryCollectionVertexFactoryShaderParameters);
+
+IMPLEMENT_GLOBAL_SHADER_PARAMETER_STRUCT(FGeometryCollectionVertexFactoryUniformShaderParameters, "GeometryCollectionVF");
+
 IMPLEMENT_VERTEX_FACTORY_PARAMETER_TYPE(FGeometryCollectionVertexFactory, SF_Vertex, FGeometryCollectionVertexFactoryShaderParameters);
-IMPLEMENT_VERTEX_FACTORY_TYPE(FGeometryCollectionVertexFactory, "/Engine/Private/LocalVertexFactory.ush", true, true, true, true, true);
+
+#if GPUCULL_TODO
+IMPLEMENT_VERTEX_FACTORY_TYPE_EX(FGeometryCollectionVertexFactory, "/Engine/Private/GeometryCollectionVertexFactory.ush", true, true, true, true, true, false, true, false);
+#else // GPUCULL_TODO
+IMPLEMENT_VERTEX_FACTORY_TYPE(FGeometryCollectionVertexFactory, "/Engine/Private/GeometryCollectionVertexFactory.ush", true, true, true, true, true);
+#endif // GPUCULL_TODO
 
 bool FGeometryCollectionVertexFactory::ShouldCompilePermutation(const FVertexFactoryShaderPermutationParameters& Parameters)
 {
-	return (Parameters.MaterialParameters.bIsUsedWithGeometryCollections || Parameters.MaterialParameters.bIsSpecialEngineMaterial)
-		&& FLocalVertexFactory::ShouldCompilePermutation(Parameters);
+	if (!Parameters.MaterialParameters.bIsUsedWithGeometryCollections && !Parameters.MaterialParameters.bIsSpecialEngineMaterial)
+	{
+		return false;
+	}
+
+	// Only compile this permutation inside the editor - it's not applicable in games, but occasionally the editor needs it.
+	if (Parameters.MaterialParameters.MaterialDomain == MD_UI)
+	{
+		return !!WITH_EDITOR;
+	}
+
+	return true;
 }
 
 //
@@ -20,24 +42,189 @@ bool FGeometryCollectionVertexFactory::ShouldCompilePermutation(const FVertexFac
 //
 void FGeometryCollectionVertexFactory::ModifyCompilationEnvironment(const FVertexFactoryShaderPermutationParameters& Parameters, FShaderCompilerEnvironment& OutEnvironment)
 {
-	const bool ContainsManualVertexFetch = OutEnvironment.GetDefinitions().Contains("MANUAL_VERTEX_FETCH");
+	const FStaticFeatureLevel MaxSupportedFeatureLevel = GetMaxSupportedFeatureLevel(Parameters.Platform);
+	const bool bUseGPUScene = UseGPUScene(Parameters.Platform, MaxSupportedFeatureLevel);
+	const bool bSupportsPrimitiveIdStream = Parameters.VertexFactoryType->SupportsPrimitiveIdStream();
+
+	OutEnvironment.SetDefine(TEXT("VF_SUPPORTS_PRIMITIVE_SCENE_DATA"), bSupportsPrimitiveIdStream && bUseGPUScene);
+
+	bool ContainsManualVertexFetch = OutEnvironment.GetDefinitions().Contains("MANUAL_VERTEX_FETCH");
 	if (!ContainsManualVertexFetch && RHISupportsManualVertexFetch(Parameters.Platform))
 	{
 		OutEnvironment.SetDefine(TEXT("MANUAL_VERTEX_FETCH"), TEXT("1"));
+		ContainsManualVertexFetch = true;
 	}
 
-	if (RHISupportsManualVertexFetch(Parameters.Platform))
+	// Geometry collections use a custom hit proxy per bone
+	OutEnvironment.SetDefine(TEXT("USE_PER_VERTEX_HITPROXY_ID"), 1);
+}
+
+void FGeometryCollectionVertexFactory::ValidateCompiledResult(const FVertexFactoryType* Type, EShaderPlatform Platform, const FShaderParameterMap& ParameterMap, TArray<FString>& OutErrors)
+{
+	if (Type->SupportsPrimitiveIdStream()
+		&& UseGPUScene(Platform, GetMaxSupportedFeatureLevel(Platform))
+		&& ParameterMap.ContainsParameterAllocation(FPrimitiveUniformShaderParameters::StaticStructMetadata.GetShaderVariableName()))
 	{
-		OutEnvironment.SetDefine(TEXT("USE_INSTANCING"), TEXT("1"));
+		OutErrors.AddUnique(*FString::Printf(TEXT("Shader attempted to bind the Primitive uniform buffer even though Vertex Factory %s computes a PrimitiveId per-instance.  This will break auto-instancing.  Shaders should use GetPrimitiveData(Parameters.PrimitiveId).Member instead of Primitive.Member."), Type->GetName()));
+	}
+}
 
-		OutEnvironment.SetDefine(TEXT("USE_INSTANCING_BONEMAP"), TEXT("1"));
-		OutEnvironment.SetDefine(TEXT("USE_DITHERED_LOD_TRANSITION_FOR_INSTANCED"), TEXT("0"));
+void FGeometryCollectionVertexFactory::InitRHI()
+{
+	SCOPED_LOADTIMER(FGeometryCollectionVertexFactory_InitRHI);
+
+	// We create different streams based on feature level
+	check(HasValidFeatureLevel());
+
+	// VertexFactory needs to be able to support max possible shader platform and feature level
+	// in case if we switch feature level at runtime.
+	const bool bCanUseGPUScene = UseGPUScene(GMaxRHIShaderPlatform, GMaxRHIFeatureLevel);
+
+	// If the vertex buffer containing position is not the same vertex buffer containing the rest of the data,
+	// then initialize PositionStream and PositionDeclaration.
+	if (Data.PositionComponent.VertexBuffer != Data.TangentBasisComponents[0].VertexBuffer)
+	{
+		auto AddDeclaration = [this, bCanUseGPUScene](EVertexInputStreamType InputStreamType, bool bAddNormal)
+		{
+			FVertexDeclarationElementList StreamElements;
+			StreamElements.Add(AccessStreamComponent(Data.PositionComponent, 0, InputStreamType));
+
+			bAddNormal = bAddNormal && Data.TangentBasisComponents[1].VertexBuffer != NULL;
+			if (bAddNormal)
+			{
+				StreamElements.Add(AccessStreamComponent(Data.TangentBasisComponents[1], 2, InputStreamType));
+			}
+
+			AddPrimitiveIdStreamElement(InputStreamType, 1, StreamElements);
+
+			InitDeclaration(StreamElements, InputStreamType);
+		};
+
+		AddDeclaration(EVertexInputStreamType::PositionOnly, false);
+		AddDeclaration(EVertexInputStreamType::PositionAndNormalOnly, true);
 	}
 
-	// Geometry collections use a custom hit proxy per bone
-	OutEnvironment.SetDefine(TEXT("USE_PER_VERTEX_HITPROXY_ID"), 1);
+	FVertexDeclarationElementList Elements;
+	if (Data.PositionComponent.VertexBuffer != nullptr)
+	{
+		Elements.Add(AccessStreamComponent(Data.PositionComponent, 0));
+	}
+
+	AddPrimitiveIdStreamElement(EVertexInputStreamType::Default, 13, Elements);
+
+	// Only the tangent and normal are used by the stream; the bitangent is derived in the shader.
+	uint8 TangentBasisAttributes[2] = { 1, 2 };
+	for (int32 AxisIndex = 0; AxisIndex < 2; AxisIndex++)
+	{
+		if (Data.TangentBasisComponents[AxisIndex].VertexBuffer != nullptr)
+		{
+			Elements.Add(AccessStreamComponent(Data.TangentBasisComponents[AxisIndex], TangentBasisAttributes[AxisIndex]));
+		}
+	}
+
+	if (Data.ColorComponentsSRV == nullptr)
+	{
+		Data.ColorComponentsSRV = GNullColorVertexBuffer.VertexBufferSRV;
+		Data.ColorIndexMask = 0;
+	}
+
+	ColorStreamIndex = INDEX_NONE;
+	if (Data.ColorComponent.VertexBuffer)
+	{
+		Elements.Add(AccessStreamComponent(Data.ColorComponent, 3));
+		ColorStreamIndex = Elements.Last().StreamIndex;
+	}
+	else
+	{
+		// If the mesh has no color component, set the null color buffer on a new stream with a stride of 0.
+		// This wastes 4 bytes per vertex, but prevents having to compile out twice the number of vertex factories.
+		FVertexStreamComponent NullColorComponent(&GNullColorVertexBuffer, 0, 0, VET_Color, EVertexStreamUsage::ManualFetch);
+		Elements.Add(AccessStreamComponent(NullColorComponent, 3));
+		ColorStreamIndex = Elements.Last().StreamIndex;
+	}
+
+	if (Data.TextureCoordinates.Num())
+	{
+		const int32 BaseTexCoordAttribute = 4;
+		for (int32 CoordinateIndex = 0; CoordinateIndex < Data.TextureCoordinates.Num(); ++CoordinateIndex)
+		{
+			Elements.Add(AccessStreamComponent(
+				Data.TextureCoordinates[CoordinateIndex],
+				BaseTexCoordAttribute + CoordinateIndex
+			));
+		}
+
+		for (int32 CoordinateIndex = Data.TextureCoordinates.Num(); CoordinateIndex < MAX_STATIC_TEXCOORDS / 2; ++CoordinateIndex)
+		{
+			Elements.Add(AccessStreamComponent(
+				Data.TextureCoordinates[Data.TextureCoordinates.Num() - 1],
+				BaseTexCoordAttribute + CoordinateIndex
+			));
+		}
+	}
 
-	FLocalVertexFactory::ModifyCompilationEnvironment(Parameters, OutEnvironment);
+	if (Data.LightMapCoordinateComponent.VertexBuffer)
+	{
+		Elements.Add(AccessStreamComponent(Data.LightMapCoordinateComponent, 15));
+	}
+	else if (Data.TextureCoordinates.Num())
+	{
+		Elements.Add(AccessStreamComponent(Data.TextureCoordinates[0], 15));
+	}
+
+	check(Streams.Num() > 0);
+
+	InitDeclaration(Elements);
+	check(IsValidRef(GetDeclaration()));
+
+	const int32 DefaultBaseVertexIndex = 0;
+	const int32 DefaultPreSkinBaseVertexIndex = 0;
+
+	if (RHISupportsManualVertexFetch(GMaxRHIShaderPlatform) || bCanUseGPUScene)
+	{
+		SCOPED_LOADTIMER(FGeometryCollectionVertexFactory_InitRHI_CreateLocalVFUniformBuffer);
+
+		FGeometryCollectionVertexFactoryUniformShaderParameters UniformParameters;
+
+		UniformParameters.LODLightmapDataIndex = Data.LODLightmapDataIndex;
+		int32 ColorIndexMask = 0;
+
+		if (RHISupportsManualVertexFetch(GMaxRHIShaderPlatform))
+		{
+			UniformParameters.VertexFetch_PositionBuffer = GetPositionsSRV();
+			UniformParameters.VertexFetch_PackedTangentsBuffer = GetTangentsSRV();
+			UniformParameters.VertexFetch_TexCoordBuffer = GetTextureCoordinatesSRV();
+
+			UniformParameters.VertexFetch_ColorComponentsBuffer = GetColorComponentsSRV();
+			ColorIndexMask = (int32)GetColorIndexMask();
+		}
+		else
+		{
+			UniformParameters.VertexFetch_PositionBuffer = GNullColorVertexBuffer.VertexBufferSRV;
+			UniformParameters.VertexFetch_PackedTangentsBuffer = GNullColorVertexBuffer.VertexBufferSRV;
+			UniformParameters.VertexFetch_TexCoordBuffer = GNullColorVertexBuffer.VertexBufferSRV;
+		}
+
+		if (!UniformParameters.VertexFetch_ColorComponentsBuffer)
+		{
+			UniformParameters.VertexFetch_ColorComponentsBuffer = GNullColorVertexBuffer.VertexBufferSRV;
+		}
+
+		const int32 NumTexCoords = GetNumTexcoords();
+		const int32 LightMapCoordinateIndex = GetLightMapCoordinateIndex();
+
+		UniformParameters.VertexFetch_Parameters = { ColorIndexMask, NumTexCoords, LightMapCoordinateIndex, 0 };
+
+		UniformBuffer = TUniformBufferRef<FGeometryCollectionVertexFactoryUniformShaderParameters>::CreateUniformBufferImmediate(UniformParameters, UniformBuffer_MultiFrame);
+	}
+
+	check(IsValidRef(GetDeclaration()));
+}
+
+void FGeometryCollectionVertexFactory::ReleaseRHI()
+{
+	UniformBuffer.SafeRelease();
+	FVertexFactory::ReleaseRHI();
 }
 
 void FGeometryCollectionVertexFactoryShaderParameters::GetElementShaderBindings(
@@ -52,21 +239,23 @@ void FGeometryCollectionVertexFactoryShaderParameters::GetElementShaderBindings(
 	FVertexInputStreamArray& VertexStreams) const
 {
 	check(VertexFactory->GetType() == &FGeometryCollectionVertexFactory::StaticType);
+	const auto* TypedVertexFactory = static_cast<const FGeometryCollectionVertexFactory*>(VertexFactory);
 
+	const bool bSupportsManualFetch = TypedVertexFactory->SupportsManualVertexFetch(FeatureLevel);
 
-	const auto* LocalVertexFactory = static_cast<const FGeometryCollectionVertexFactory*>(VertexFactory);
-	FRHIUniformBuffer* VertexFactoryUniformBuffer = nullptr;
-	VertexFactoryUniformBuffer = LocalVertexFactory->GetUniformBuffer();
-
-
-	FLocalVertexFactoryShaderParametersBase::GetElementShaderBindingsBase(Scene, View, Shader, InputStreamType, FeatureLevel, VertexFactory, BatchElement, VertexFactoryUniformBuffer, ShaderBindings, VertexStreams);
+	FRHIUniformBuffer* VertexFactoryUniformBuffer = TypedVertexFactory->GetUniformBuffer();
 
+	if (bSupportsManualFetch || UseGPUScene(GMaxRHIShaderPlatform, FeatureLevel))
+	{
+		check(VertexFactoryUniformBuffer != nullptr);
+		ShaderBindings.Add(Shader->GetUniformBufferParameter<FGeometryCollectionVertexFactoryUniformShaderParameters>(), VertexFactoryUniformBuffer);
+	}
 
 	// We only want to set the SRV parameters if we support manual vertex fetch.
-	if (LocalVertexFactory->SupportsManualVertexFetch(View->GetFeatureLevel()))
+	if (bSupportsManualFetch)
 	{
-		ShaderBindings.Add(VertexFetch_InstanceTransformBufferParameter, LocalVertexFactory->GetInstanceTransformSRV());
-		ShaderBindings.Add(VertexFetch_InstancePrevTransformBufferParameter, LocalVertexFactory->GetInstancePrevTransformSRV());
-		ShaderBindings.Add(VertexFetch_InstanceBoneMapBufferParameter, LocalVertexFactory->GetInstanceBoneMapSRV());
+		ShaderBindings.Add(VertexFetch_BoneTransformBufferParameter, TypedVertexFactory->GetBoneTransformSRV());
+		ShaderBindings.Add(VertexFetch_BonePrevTransformBufferParameter, TypedVertexFactory->GetBonePrevTransformSRV());
+		ShaderBindings.Add(VertexFetch_BoneMapBufferParameter, TypedVertexFactory->GetBoneMapSRV());
 	}
-}
\ No newline at end of file
+}
diff --git a/Engine/Source/Runtime/Engine/Private/LocalVertexFactory.cpp b/Engine/Source/Runtime/Engine/Private/LocalVertexFactory.cpp
--- a/Engine/Source/Runtime/Engine/Private/LocalVertexFactory.cpp
+++ b/Engine/Source/Runtime/Engine/Private/LocalVertexFactory.cpp
@@ -301,27 +301,23 @@ void FLocalVertexFactory::InitRHI()
 	}
 
 	FVertexDeclarationElementList Elements;
-	if(Data.PositionComponent.VertexBuffer != NULL)
+	if (Data.PositionComponent.VertexBuffer != nullptr)
 	{
 		Elements.Add(AccessStreamComponent(Data.PositionComponent,0));
 	}
 
+	if (GetType()->SupportsPrimitiveIdStream() && bCanUseGPUScene)
 	{
-		const uint8 Index = static_cast<uint8>(EVertexInputStreamType::Default);
-		PrimitiveIdStreamIndex[Index] = -1;
-		if (GetType()->SupportsPrimitiveIdStream() && bCanUseGPUScene)
-		{
-			// When the VF is used for rendering in normal mesh passes, this vertex buffer and offset will be overridden
-			Elements.Add(AccessStreamComponent(FVertexStreamComponent(&GPrimitiveIdDummy, 0, 0, PrimitiveIdStreamStride, VET_UInt, EVertexStreamUsage::Instancing), 13));
-			PrimitiveIdStreamIndex[Index] = Elements.Last().StreamIndex;
-		}
+		// When the VF is used for rendering in normal mesh passes, this vertex buffer and offset will be overridden
+		Elements.Add(AccessStreamComponent(FVertexStreamComponent(&GPrimitiveIdDummy, 0, 0, PrimitiveIdStreamStride, VET_UInt, EVertexStreamUsage::Instancing), 13));
+		SetPrimitiveIdStreamIndex(EVertexInputStreamType::Default, Elements.Last().StreamIndex);
 	}
 
-	// only tangent,normal are used by the stream. the binormal is derived in the shader
+	// Only the tangent and normal are used by the stream; the bitangent is derived in the shader.
 	uint8 TangentBasisAttributes[2] = { 1, 2 };
-	for(int32 AxisIndex = 0;AxisIndex < 2;AxisIndex++)
+	for (int32 AxisIndex = 0;AxisIndex < 2;AxisIndex++)
 	{
-		if(Data.TangentBasisComponents[AxisIndex].VertexBuffer != NULL)
+		if (Data.TangentBasisComponents[AxisIndex].VertexBuffer != nullptr)
 		{
 			Elements.Add(AccessStreamComponent(Data.TangentBasisComponents[AxisIndex],TangentBasisAttributes[AxisIndex]));
 		}
@@ -334,24 +330,24 @@ void FLocalVertexFactory::InitRHI()
 	}
 
 	ColorStreamIndex = -1;
-	if(Data.ColorComponent.VertexBuffer)
+	if (Data.ColorComponent.VertexBuffer)
 	{
 		Elements.Add(AccessStreamComponent(Data.ColorComponent,3));
 		ColorStreamIndex = Elements.Last().StreamIndex;
 	}
 	else
 	{
-		//If the mesh has no color component, set the null color buffer on a new stream with a stride of 0.
-		//This wastes 4 bytes of bandwidth per vertex, but prevents having to compile out twice the number of vertex factories.
+		// If the mesh has no color component, set the null color buffer on a new stream with a stride of 0.
+		// This wastes 4 bytes per vertex, but prevents having to compile out twice the number of vertex factories.
 		FVertexStreamComponent NullColorComponent(&GNullColorVertexBuffer, 0, 0, VET_Color, EVertexStreamUsage::ManualFetch);
 		Elements.Add(AccessStreamComponent(NullColorComponent, 3));
 		ColorStreamIndex = Elements.Last().StreamIndex;
 	}
 
-	if(Data.TextureCoordinates.Num())
+	if (Data.TextureCoordinates.Num())
 	{
 		const int32 BaseTexCoordAttribute = 4;
-		for(int32 CoordinateIndex = 0;CoordinateIndex < Data.TextureCoordinates.Num();CoordinateIndex++)
+		for (int32 CoordinateIndex = 0; CoordinateIndex < Data.TextureCoordinates.Num(); ++CoordinateIndex)
 		{
 			Elements.Add(AccessStreamComponent(
 				Data.TextureCoordinates[CoordinateIndex],
@@ -359,7 +355,7 @@ void FLocalVertexFactory::InitRHI()
 				));
 		}
 
-		for (int32 CoordinateIndex = Data.TextureCoordinates.Num(); CoordinateIndex < MAX_STATIC_TEXCOORDS / 2; CoordinateIndex++)
+		for (int32 CoordinateIndex = Data.TextureCoordinates.Num(); CoordinateIndex < MAX_STATIC_TEXCOORDS / 2; ++CoordinateIndex)
 		{
 			Elements.Add(AccessStreamComponent(
 				Data.TextureCoordinates[Data.TextureCoordinates.Num() - 1],
@@ -368,11 +364,11 @@ void FLocalVertexFactory::InitRHI()
 		}
 	}
 
-	if(Data.LightMapCoordinateComponent.VertexBuffer)
+	if (Data.LightMapCoordinateComponent.VertexBuffer)
 	{
 		Elements.Add(AccessStreamComponent(Data.LightMapCoordinateComponent,15));
 	}
-	else if(Data.TextureCoordinates.Num())
+	else if (Data.TextureCoordinates.Num())
 	{
 		Elements.Add(AccessStreamComponent(Data.TextureCoordinates[0],15));
 	}
@@ -384,6 +380,7 @@ void FLocalVertexFactory::InitRHI()
 
 	const int32 DefaultBaseVertexIndex = 0;
 	const int32 DefaultPreSkinBaseVertexIndex = 0;
+
 	if (RHISupportsManualVertexFetch(GMaxRHIShaderPlatform) || bCanUseGPUScene)
 	{
 		SCOPED_LOADTIMER(FLocalVertexFactory_InitRHI_CreateLocalVFUniformBuffer);
diff --git a/Engine/Source/Runtime/Engine/Public/GeometryCollectionRendering.h b/Engine/Source/Runtime/Engine/Public/GeometryCollectionRendering.h
--- a/Engine/Source/Runtime/Engine/Public/GeometryCollectionRendering.h
+++ b/Engine/Source/Runtime/Engine/Public/GeometryCollectionRendering.h
@@ -10,19 +10,24 @@
 #include "Rendering/SkinWeightVertexBuffer.h"
 #include "ShaderParameterUtils.h"
 
-class FGeometryCollectionVertexFactoryShaderParameters : public FLocalVertexFactoryShaderParameters
+BEGIN_GLOBAL_SHADER_PARAMETER_STRUCT(FGeometryCollectionVertexFactoryUniformShaderParameters, ENGINE_API)
+	SHADER_PARAMETER(FIntVector4,VertexFetch_Parameters)
+	SHADER_PARAMETER(uint32,LODLightmapDataIndex)
+	SHADER_PARAMETER_SRV(Buffer<float2>, VertexFetch_TexCoordBuffer)
+	SHADER_PARAMETER_SRV(Buffer<float>, VertexFetch_PositionBuffer)
+	SHADER_PARAMETER_SRV(Buffer<float4>, VertexFetch_PackedTangentsBuffer)
+	SHADER_PARAMETER_SRV(Buffer<float4>, VertexFetch_ColorComponentsBuffer)
+END_GLOBAL_SHADER_PARAMETER_STRUCT()
+
+class FGeometryCollectionVertexFactoryShaderParameters : public FVertexFactoryShaderParameters
 {
 	DECLARE_TYPE_LAYOUT(FGeometryCollectionVertexFactoryShaderParameters, NonVirtual);
 public:
-
-	// #note: We are reusing VertexFetch_InstanceTransformBuffer that was created for instanced static mesh 
 	void Bind(const FShaderParameterMap& ParameterMap)
 	{
-		FLocalVertexFactoryShaderParameters::Bind(ParameterMap);
-
-		VertexFetch_InstanceTransformBufferParameter.Bind(ParameterMap, TEXT("VertexFetch_InstanceTransformBuffer"));
-		VertexFetch_InstancePrevTransformBufferParameter.Bind(ParameterMap, TEXT("VertexFetch_InstancePrevTransformBuffer"));
-		VertexFetch_InstanceBoneMapBufferParameter.Bind(ParameterMap, TEXT("VertexFetch_InstanceBoneMapBuffer"));
+		VertexFetch_BoneTransformBufferParameter.Bind(ParameterMap, TEXT("VertexFetch_BoneTransformBuffer"));
+		VertexFetch_BonePrevTransformBufferParameter.Bind(ParameterMap, TEXT("VertexFetch_BonePrevTransformBuffer"));
+		VertexFetch_BoneMapBufferParameter.Bind(ParameterMap, TEXT("VertexFetch_BoneMapBuffer"));
 	}	
 
 	void GetElementShaderBindings(
@@ -37,31 +42,31 @@ public:
 		FVertexInputStreamArray& VertexStreams) const;
 
 private:
-	
-		LAYOUT_FIELD(FShaderResourceParameter, VertexFetch_InstanceTransformBufferParameter)
-		LAYOUT_FIELD(FShaderResourceParameter, VertexFetch_InstancePrevTransformBufferParameter)
-		LAYOUT_FIELD(FShaderResourceParameter, VertexFetch_InstanceBoneMapBufferParameter)
-	
+	LAYOUT_FIELD(FShaderResourceParameter, VertexFetch_BoneTransformBufferParameter);
+	LAYOUT_FIELD(FShaderResourceParameter, VertexFetch_BonePrevTransformBufferParameter);
+	LAYOUT_FIELD(FShaderResourceParameter, VertexFetch_BoneMapBufferParameter);
 };
 
 /**
  * A vertex factory for Geometry Collections
  */
-struct ENGINE_API FGeometryCollectionVertexFactory : public FLocalVertexFactory
+struct ENGINE_API FGeometryCollectionVertexFactory : public FVertexFactory
 {
-    DECLARE_VERTEX_FACTORY_TYPE(FGeometryCollectionVertexFactory);
+	DECLARE_VERTEX_FACTORY_TYPE(FGeometryCollectionVertexFactory);
+
 public:
 	FGeometryCollectionVertexFactory(ERHIFeatureLevel::Type InFeatureLevel)
-		: FLocalVertexFactory(InFeatureLevel, "FGeometryCollectionVertexFactory")
+	: FVertexFactory(InFeatureLevel)
 	{
+		bSupportsManualVertexFetch = true;
 	}
 
 	// Data includes what we need for transform and everything in local vertex factory too
-	struct FDataType : public FLocalVertexFactory::FDataType
+	struct FDataType : public FStaticMeshDataType
 	{
-		FRHIShaderResourceView* InstanceTransformSRV = nullptr;
-		FRHIShaderResourceView* InstancePrevTransformSRV = nullptr;
-		FRHIShaderResourceView* InstanceBoneMapSRV = nullptr;
+		FRHIShaderResourceView* BoneTransformSRV = nullptr;
+		FRHIShaderResourceView* BonePrevTransformSRV = nullptr;
+		FRHIShaderResourceView* BoneMapSRV = nullptr;
 	};
 
 	//
@@ -75,12 +80,14 @@ public:
 	//
 	static void ModifyCompilationEnvironment(const FVertexFactoryShaderPermutationParameters& Parameters, FShaderCompilerEnvironment& OutEnvironment);
 
+	static void ValidateCompiledResult(const FVertexFactoryType* Type, EShaderPlatform Platform, const FShaderParameterMap& ParameterMap, TArray<FString>& OutErrors);
+
 	//
 	// Set the data on the vertex factory
 	//
 	void SetData(const FDataType& InData)
 	{
-		FLocalVertexFactory::Data = InData;
+		check(IsInRenderingThread());
 		Data = InData;
 		UpdateRHI();
 	}
@@ -101,37 +108,79 @@ public:
 		BeginUpdateResourceRHI(this);
 	}
 
-	// FRenderResource interface.	
-	virtual void InitRHI() override
+	// FRenderResource interface.
+	virtual void InitRHI() override;
+	virtual void ReleaseRHI() override;
+
+	static bool SupportsTessellationShaders() { return true; }
+
+	inline FRHIShaderResourceView* GetPositionsSRV() const
+	{
+		return Data.PositionComponentSRV;
+	}
+
+	inline FRHIShaderResourceView* GetTangentsSRV() const
+	{
+		return Data.TangentsSRV;
+	}
+
+	inline FRHIShaderResourceView* GetTextureCoordinatesSRV() const
+	{
+		return Data.TextureCoordinatesSRV;
+	}
+
+	inline FRHIShaderResourceView* GetColorComponentsSRV() const
+	{
+		return Data.ColorComponentsSRV;
+	}
+
+	inline const uint32 GetColorIndexMask() const
+	{
+		return Data.ColorIndexMask;
+	}
+
+	inline const int GetLightMapCoordinateIndex() const
+	{
+		return Data.LightMapCoordinateIndex;
+	}
+
+	inline const int GetNumTexcoords() const
+	{
+		return Data.NumTexCoords;
+	}
+
+	FRHIUniformBuffer* GetUniformBuffer() const
 	{
-		FLocalVertexFactory::InitRHI();
+		return UniformBuffer.GetReference();
 	}
 
-	inline void SetInstanceTransformSRV(FRHIShaderResourceView* InstanceTransformSRV)
+	inline void SetBoneTransformSRV(FRHIShaderResourceView* BoneTransformSRV)
 	{
-		Data.InstanceTransformSRV = InstanceTransformSRV;
+		Data.BoneTransformSRV = BoneTransformSRV;
 	}
 	
-	inline FRHIShaderResourceView* GetInstanceTransformSRV() const
+	inline FRHIShaderResourceView* GetBoneTransformSRV() const
 	{
-		return Data.InstanceTransformSRV;
+		return Data.BoneTransformSRV;
 	}
 
-	inline void SetInstancePrevTransformSRV(FRHIShaderResourceView* InstancePrevTransformSRV)
+	inline void SetBonePrevTransformSRV(FRHIShaderResourceView* BonePrevTransformSRV)
 	{
-		Data.InstancePrevTransformSRV = InstancePrevTransformSRV;
+		Data.BonePrevTransformSRV = BonePrevTransformSRV;
 	}
 
-	inline FRHIShaderResourceView* GetInstancePrevTransformSRV() const
+	inline FRHIShaderResourceView* GetBonePrevTransformSRV() const
 	{
-		return Data.InstancePrevTransformSRV;
+		return Data.BonePrevTransformSRV;
 	}
 
-	inline FRHIShaderResourceView* GetInstanceBoneMapSRV() const
+	inline FRHIShaderResourceView* GetBoneMapSRV() const
 	{
-		return Data.InstanceBoneMapSRV;
+		return Data.BoneMapSRV;
 	}
 	
 private:
 	FDataType Data;
+	TUniformBufferRef<FGeometryCollectionVertexFactoryUniformShaderParameters> UniformBuffer;
+	int32 ColorStreamIndex = INDEX_NONE;
 };
diff --git a/Engine/Source/Runtime/Experimental/GeometryCollectionEngine/Private/GeometryCollection/GeometryCollectionSceneProxy.cpp b/Engine/Source/Runtime/Experimental/GeometryCollectionEngine/Private/GeometryCollection/GeometryCollectionSceneProxy.cpp
--- a/Engine/Source/Runtime/Experimental/GeometryCollectionEngine/Private/GeometryCollection/GeometryCollectionSceneProxy.cpp
+++ b/Engine/Source/Runtime/Experimental/GeometryCollectionEngine/Private/GeometryCollection/GeometryCollectionSceneProxy.cpp
@@ -110,6 +110,10 @@ FGeometryCollectionSceneProxy::FGeometryCollectionSceneProxy(UGeometryCollection
 			Sections.Add(Section);
 		}
 	}
+#if GPUCULL_TODO
+	const auto FeatureLevel = GetScene().GetFeatureLevel();
+	bVFRequiresPrimitiveUniformBuffer = !UseGPUScene(GMaxRHIShaderPlatform, FeatureLevel);
+#endif // GPUCULL_TODO
 
 #if GEOMETRYCOLLECTION_EDITOR_SELECTION
 	// Init HitProxy array with the maximum number of subsections
@@ -238,9 +242,9 @@ void FGeometryCollectionSceneProxy::InitResources()
 
 		BoneMapBuffer.InitResource();
 
-		Data.InstanceBoneMapSRV = BoneMapBuffer.VertexBufferSRV;			
-		Data.InstanceTransformSRV = TransformBuffers[0].VertexBufferSRV;
-		Data.InstancePrevTransformSRV = PrevTransformBuffers[0].VertexBufferSRV;
+		Data.BoneMapSRV = BoneMapBuffer.VertexBufferSRV;
+		Data.BoneTransformSRV = TransformBuffers[0].VertexBufferSRV;
+		Data.BonePrevTransformSRV = PrevTransformBuffers[0].VertexBufferSRV;
 	}
 
 	// 
@@ -504,8 +508,8 @@ void FGeometryCollectionSceneProxy::SetDynamicData_RenderThread(FGeometryCollect
 				CycleTransformBuffers(bLocalGeometryCollectionTripleBufferUploads);
 				FGeometryCollectionTransformBuffer& TransformBuffer = GetCurrentTransformBuffer();
 				FGeometryCollectionTransformBuffer& PrevTransformBuffer = GetCurrentPrevTransformBuffer();
-				VertexFactory.SetInstanceTransformSRV(TransformBuffer.VertexBufferSRV);
-				VertexFactory.SetInstancePrevTransformSRV(PrevTransformBuffer.VertexBufferSRV);
+				VertexFactory.SetBoneTransformSRV(TransformBuffer.VertexBufferSRV);
+				VertexFactory.SetBonePrevTransformSRV(PrevTransformBuffer.VertexBufferSRV);
 
 				check(TransformBuffer.NumTransforms == DynamicData->Transforms.Num());
 				check(PrevTransformBuffer.NumTransforms == DynamicData->PrevTransforms.Num());
@@ -525,8 +529,8 @@ void FGeometryCollectionSceneProxy::SetDynamicData_RenderThread(FGeometryCollect
 				CycleTransformBuffers(bLocalGeometryCollectionTripleBufferUploads);
 				FGeometryCollectionTransformBuffer& TransformBuffer = GetCurrentTransformBuffer();
 				FGeometryCollectionTransformBuffer& PrevTransformBuffer = GetCurrentPrevTransformBuffer();
-				VertexFactory.SetInstanceTransformSRV(TransformBuffer.VertexBufferSRV);
-				VertexFactory.SetInstancePrevTransformSRV(PrevTransformBuffer.VertexBufferSRV);
+				VertexFactory.SetBoneTransformSRV(TransformBuffer.VertexBufferSRV);
+				VertexFactory.SetBonePrevTransformSRV(PrevTransformBuffer.VertexBufferSRV);
 
 				// if we are rendering the base mesh geometry, then use rest transforms rather than the simulated one for both current and previous transforms
 				void* VertexBufferData = RHILockBuffer(TransformBuffer.VertexBufferRHI, 0, ConstantData->RestTransforms.Num() * sizeof(FMatrix), LockMode);
@@ -576,6 +580,7 @@ void FGeometryCollectionSceneProxy::SetDynamicData_RenderThread(FGeometryCollect
 
 				if (ThisBatchSize > 0)
 				{
+					const FMatrix* RESTRICT BoneTransformsPtr = DynamicData->IsDynamic ? DynamicData->Transforms.GetData() : ConstantData->RestTransforms.GetData();
 	#if INTEL_ISPC
 					uint8* VertexBufferOffset = (uint8*)VertexBufferData + (IndexOffset * VertexBuffer.GetStride());
 					ispc::SetDynamicData_RenderThread(
@@ -583,12 +588,12 @@ void FGeometryCollectionSceneProxy::SetDynamicData_RenderThread(FGeometryCollect
 						ThisBatchSize, 
 						VertexBuffer.GetStride(), 
 						&ConstantData->BoneMap[IndexOffset], 
-						(ispc::FMatrix*)&DynamicData->Transforms[0], 
+						(ispc::FMatrix*)BoneTransformsPtr,
 						(ispc::FVector*)&ConstantData->Vertices[IndexOffset]);
 	#else
 					for (int32 i = IndexOffset; i < IndexOffset + ThisBatchSize; i++)
 					{
-						FVector Transformed = DynamicData->Transforms[ConstantData->BoneMap[i]].TransformPosition(ConstantData->Vertices[i]);
+						FVector Transformed = BoneTransformsPtr[ConstantData->BoneMap[i]].TransformPosition(ConstantData->Vertices[i]);
 						FMemory::Memcpy((uint8*)VertexBufferData + (i * VertexBuffer.GetStride()), &Transformed, sizeof(FVector));
 					}
 	#endif
@@ -842,8 +847,6 @@ HHitProxy* FGeometryCollectionSceneProxy::CreateHitProxies(UPrimitiveComponent*
 	{
 #if GEOMETRYCOLLECTION_EDITOR_SELECTION
 		const int32 NumTransforms = (Sections.Num() > 0) ? SubSectionHitProxies.Num() / Sections.Num(): 0;
-#endif  // #if GEOMETRYCOLLECTION_EDITOR_SELECTION
-
 		for (int32 SectionIndex = 0; SectionIndex < Sections.Num(); ++SectionIndex)
 		{
 			// Create HitProxy for regular material based sections, and update existing section
@@ -855,7 +858,6 @@ HHitProxy* FGeometryCollectionSceneProxy::CreateHitProxies(UPrimitiveComponent*
 			OutHitProxies.Add(HitProxy);
 			Section.HitProxy = HitProxy;
 
-#if GEOMETRYCOLLECTION_EDITOR_SELECTION
 			// Create HitProxy per transform index using the same material Id than the current sections
 			// All combinations of material id/transform index are populated,
 			// since it can't be assumed that any of them won't be needed.
@@ -871,16 +873,19 @@ HHitProxy* FGeometryCollectionSceneProxy::CreateHitProxies(UPrimitiveComponent*
 				SubSectionHitProxies[SectionOffset + TransformIndex] = SubSectionHitProxy;
 			}
 		}
-
-		// Update existing subsections and any HitProxy's section index that is currently being used
-		for (int32 SubSectionIndex = 0; SubSectionIndex < SubSections.Num(); ++SubSectionIndex)
+	#else
+		for (int32 SectionIndex = 0; SectionIndex < Sections.Num(); ++SectionIndex)
 		{
-			HGeometryCollection* const SubSectionHitProxy = SubSectionHitProxies[SubSectionHitProxyIndexMap[SubSectionIndex]];
+			// Create HitProxy for regular material based sections, and update existing section
+			FGeometryCollectionSection& Section = Sections[SectionIndex];
 
-			SubSections[SubSectionIndex].HitProxy = SubSectionHitProxy;
-			SubSectionHitProxy->SectionIndex = SubSectionIndex;
-#endif  // #if GEOMETRYCOLLECTION_EDITOR_SELECTION
+			const int32 MaterialID = Section.MaterialID;
+			HActor* const HitProxy = new HActor(Component->GetOwner(), Component, SectionIndex, MaterialID);
+
+			OutHitProxies.Add(HitProxy);
+			Section.HitProxy = HitProxy;
 		}
+	#endif
 	}
 
 	return DefaultHitProxy;
@@ -1030,9 +1035,8 @@ void FGeometryCollectionSceneProxy::GetPreSkinnedLocalBounds(FBoxSphereBounds& O
 }
 
 FNaniteGeometryCollectionSceneProxy::FNaniteGeometryCollectionSceneProxy(UGeometryCollectionComponent* Component)
-	: Nanite::FSceneProxyBase(Component)
-	//, RenderData(Component->GetStaticMesh()->RenderData.Get())
-	, GeometryCollection(Component->GetRestCollection())
+: Nanite::FSceneProxyBase(Component)
+, GeometryCollection(Component->GetRestCollection())
 {
 	LLM_SCOPE_BYTAG(Nanite);
 
@@ -1254,7 +1258,6 @@ void FNaniteGeometryCollectionSceneProxy::DrawStaticElements(FStaticPrimitiveDra
 
 #if WITH_EDITOR
 		HHitProxy* HitProxy = Section.HitProxy;
-		//check(HitProxy); // TODO: Is this valid? SME seems to have null proxies, but normal editor doesn't
 		PDI->SetHitProxy(HitProxy);
 #endif
 		PDI->DrawMesh(MeshBatch, FLT_MAX);
