Fix VirtualShadowMapping shader warnings

diff --git a/Engine/Shaders/Private/VirtualShadowMaps/PageAccessCommon.ush b/Engine/Shaders/Private/VirtualShadowMaps/PageAccessCommon.ush
--- a/Engine/Shaders/Private/VirtualShadowMaps/PageAccessCommon.ush
+++ b/Engine/Shaders/Private/VirtualShadowMaps/PageAccessCommon.ush
@@ -66,12 +66,12 @@ uint CalcLog2LevelDimsPages(uint Level)
 
 uint CalcLevelDimsPages(uint Level)
 {
-	return 1 << CalcLog2LevelDimsPages( Level );
+	return 1u << CalcLog2LevelDimsPages( Level );
 }
 
 uint CalcLevelDimsTexels(uint Level)
 {
-	return VSM_VIRTUAL_MAX_RESOLUTION_XY >> Level;
+	return uint(VSM_VIRTUAL_MAX_RESOLUTION_XY) >> Level;
 }
 
 uint CalcLevelOffsets(uint Level)
@@ -268,7 +268,7 @@ bool OverlapsAnyValidPage( uint ShadowMapID, uint MipLevel, int2 ViewPortRectMin
 		RectPages >>= HMipLevel;
 
 		// Root level size in pages
-		uint HLevelSizePages = VSM_LEVEL0_DIM_PAGES_XY >> (MipLevel + HMipLevel);
+		uint HLevelSizePages = uint(VSM_LEVEL0_DIM_PAGES_XY) >> (MipLevel + HMipLevel);
 		
 		int HPageLevelOffset = ShadowMapID * VirtualShadowMap.HPageTableSize + CalcHPageFlagLevelOffsets(MipLevel)
 			+ CalcLevelOffsets( HMipLevel + MipLevel)
diff --git a/Engine/Shaders/Private/VirtualShadowMaps/ProjectionCommon.ush b/Engine/Shaders/Private/VirtualShadowMaps/ProjectionCommon.ush
--- a/Engine/Shaders/Private/VirtualShadowMaps/ProjectionCommon.ush
+++ b/Engine/Shaders/Private/VirtualShadowMaps/ProjectionCommon.ush
@@ -240,7 +240,7 @@ float ComputeVirtualShadowMapOptimalSlopeBias(
 	// Adjust depth scale if we sampled a different clipmap level
 	// NOTE: Sampled clipmap should always be > the requested one (coarser)
 	// Do this after clamping to be consistent in world space
-	OptimalSlopeBias *= float(1 << (SmSample.VirtualShadowMapId - RequestedVirtualShadowMapId));
+	OptimalSlopeBias *= float(1u << (SmSample.VirtualShadowMapId - RequestedVirtualShadowMapId));
 
 	return OptimalSlopeBias;
 }
@@ -331,7 +331,7 @@ FVirtualShadowMapSampleResult SampleVirtualShadowMap(int VirtualShadowMapId, flo
 			{
 				int SampledClipmapIndex = SmSample.VirtualShadowMapId - VirtualShadowMapId;
 
-				FVirtualShadowMapSampleResult Result;
+				FVirtualShadowMapSampleResult Result = (FVirtualShadowMapSampleResult)0;
 				Result.bValid = true;
 				Result.bOccluded = false;
 				Result.OccluderDistance = -1.0f;
@@ -371,7 +371,7 @@ FVirtualShadowMapSampleResult SampleVirtualShadowMap(int VirtualShadowMapId, flo
 		FVirtualShadowMapSample SmSample = SampleVirtualShadowMap(VirtualShadowMapId, ShadowUVz.xy);
 		if (SmSample.bValid)
 		{
-			FVirtualShadowMapSampleResult Result;
+			FVirtualShadowMapSampleResult Result = (FVirtualShadowMapSampleResult)0;
 			Result.bValid = true;
 			Result.bOccluded = false;
 			Result.OccluderDistance = -1.0f;
@@ -398,7 +398,7 @@ FVirtualShadowMapSampleResult SampleVirtualShadowMap(int VirtualShadowMapId, flo
 	}
 
 	// Invalid
-	FVirtualShadowMapSampleResult Result;
+	FVirtualShadowMapSampleResult Result = (FVirtualShadowMapSampleResult)0;
 	Result.bValid = false;
 	Result.bOccluded = false;
 	Result.OccluderDistance = -1.0f;
diff --git a/Engine/Shaders/Private/VirtualShadowMaps/ProjectionDirectional.ush b/Engine/Shaders/Private/VirtualShadowMaps/ProjectionDirectional.ush
--- a/Engine/Shaders/Private/VirtualShadowMaps/ProjectionDirectional.ush
+++ b/Engine/Shaders/Private/VirtualShadowMaps/ProjectionDirectional.ush
@@ -76,7 +76,7 @@ FSMRTClipmapRayState SMRTClipmapRayInitialize(
 
 FSMRTSample SMRTFindSample(inout FSMRTClipmapRayState RayState, float SampleTime)
 {
-	FSMRTSample Sample;
+	FSMRTSample Sample = (FSMRTSample)0;
 	Sample.bValid = false;
 	
 	const float3 SampleTexelsZ = RayState.RayStartShadowTexelsZ + RayState.RayStepShadowTexelsZ * SampleTime;
diff --git a/Engine/Shaders/Private/VirtualShadowMaps/ProjectionSpot.ush b/Engine/Shaders/Private/VirtualShadowMaps/ProjectionSpot.ush
--- a/Engine/Shaders/Private/VirtualShadowMaps/ProjectionSpot.ush
+++ b/Engine/Shaders/Private/VirtualShadowMaps/ProjectionSpot.ush
@@ -53,7 +53,7 @@ FSMRTSingleRayState SMRTSingleRayInitialize(
 
 FSMRTSample SMRTFindSample(inout FSMRTSingleRayState RayState, float SampleTime)
 {
-	FSMRTSample Sample;
+	FSMRTSample Sample = (FSMRTSample)0;
 	Sample.bValid = false;
 	float3 UVz = RayState.RayStartUVz.xyz + RayState.RayStepUVz.xyz * SampleTime;
 
