Add support for virtual shadowing for hair strands.

diff --git a/Engine/Shaders/Private/VirtualShadowMaps/PageManagement.usf b/Engine/Shaders/Private/VirtualShadowMaps/PageManagement.usf
--- a/Engine/Shaders/Private/VirtualShadowMaps/PageManagement.usf
+++ b/Engine/Shaders/Private/VirtualShadowMaps/PageManagement.usf
@@ -10,6 +10,7 @@
 #include "../DeferredShadingCommon.ush"
 #include "../MortonCode.ush"
 #include "../Nanite/NaniteDataDecode.ush"
+#include "../HairStrands/HairStrandsVisibilityCommon.ush"
 #include "PageAccessCommon.ush"
 #include "PageCacheCommon.ush"
 #include "ProjectionCommon.ush"
@@ -32,6 +33,10 @@ uint bPostBasePass;
 float LodFootprintScale;
 float PageDilationBorderSize;
 
+// HAIR_TODO: move all data into a HairStrands UB
+Texture2D<uint4> HairCategorizationTexture;
+uint bUseHairData;
+
 // Convenience to group up a bunch of data computed for a screen pixel
 struct FScreenPixelData
 {
@@ -92,21 +97,10 @@ FScreenPixelData GetScreenPixelData(uint2 PixelPos, float DeviceZ)
 	float3 TranslatedWorldPosition = TranslatedWorldH.xyz / TranslatedWorldH.w;
 	// Subtract the view pre-translation to get back to world space
 	Data.WorldPosition = TranslatedWorldPosition - View.PreViewTranslation;
-
-	if (bPostBasePass != 0)
-	{
-		Data.GBufferData = GetGBufferData(Data.ScreenUV);
-	}
-
 	return Data;
 }
 
-FScreenPixelData GetScreenPixelData(uint2 PixelPos)
-{
-	return GetScreenPixelData(PixelPos, GetDeviceZFromPixel(PixelPos));
-}
-
-uint GetMipLevel(int VirtualShadowMapId, FScreenPixelData Pixel)
+uint GetMipLevel(int VirtualShadowMapId, float3 WorldPosition, float SceneDepth)
 {
 	float Footprint = float(VSM_VIRTUAL_MAX_RESOLUTION_XY);
 	FVirtualShadowMapProjectionShaderData ProjectionData = GetVirtualShadowMapProjectionData(VirtualShadowMapId);
@@ -118,11 +112,11 @@ uint GetMipLevel(int VirtualShadowMapId, FScreenPixelData Pixel)
 		// TODO: Roll into a uniform
 		float2 RadiusXY = 1.0f / (View.ViewSizeAndInvSize.xy * View.ViewToClip._m00_m11);
 		float RadiusScreen = min(RadiusXY.x, RadiusXY.y);
-		float DepthScale = Pixel.SceneDepth * View.ViewToClip[2][3] + View.ViewToClip[3][3];
+		float DepthScale = SceneDepth * View.ViewToClip[2][3] + View.ViewToClip[3][3];
 
 		float RadiusWorld = DepthScale * RadiusScreen;
 
-		float3 TranslatedWorldPosition = Pixel.WorldPosition + ProjectionData.ShadowPreViewTranslation;
+		float3 TranslatedWorldPosition = WorldPosition + ProjectionData.ShadowPreViewTranslation;
 		float4 ShadowUVz = mul(float4(TranslatedWorldPosition, 1.0f), ProjectionData.TranslatedWorldToShadowUVMatrix);
 
 		float4 RadiusClipH = mul(float4(RadiusWorld, 0.0f, ShadowUVz.w, 1.0f), ProjectionData.ShadowViewToClipMatrix);
@@ -133,10 +127,10 @@ uint GetMipLevel(int VirtualShadowMapId, FScreenPixelData Pixel)
 	return CalcMipLevelFromFootprint(LodFootprintScale * Footprint, VSM_MAX_MIP_LEVELS);
 }
 
-void MarkPage(uint VirtualShadowMapId, uint MipLevel, FScreenPixelData Pixel, float2 PageDilationOffset)
+void MarkPage(uint VirtualShadowMapId, uint MipLevel, float3 WorldPosition, float2 PageDilationOffset)
 {
 	FVirtualShadowMapProjectionShaderData ProjectionData = GetVirtualShadowMapProjectionData(VirtualShadowMapId);
-	float4 ShadowUVz = mul(float4(Pixel.WorldPosition + ProjectionData.ShadowPreViewTranslation, 1.0f), ProjectionData.TranslatedWorldToShadowUVMatrix);
+	float4 ShadowUVz = mul(float4(WorldPosition + ProjectionData.ShadowPreViewTranslation, 1.0f), ProjectionData.TranslatedWorldToShadowUVMatrix);
 	ShadowUVz.xyz /= ShadowUVz.w;
 
 	// Check overlap vs the shadow map space
@@ -183,7 +177,28 @@ void GeneratePageFlagsFromPixels(uint3 GroupId : SV_GroupID, uint GroupIndex : S
 		return;
 	}
 	
-	FScreenPixelData Pixel = GetScreenPixelData(PixelPos);
+	float3 WorldPosition = 0;
+	float  SceneDepth    = 0;
+	bool   bIsValid		 = false;
+	if (bUseHairData)
+	{
+		FCategorizationData CatData = DecodeCategorizationData(HairCategorizationTexture.Load(uint3(PixelPos, 0)));
+		FScreenPixelData HairPixel = GetScreenPixelData(PixelPos, CatData.ClosestDepth);
+		WorldPosition = HairPixel.WorldPosition;
+		SceneDepth = HairPixel.SceneDepth;
+		bIsValid = CatData.PixelCoverage > 0;
+	}
+	else
+	{
+		FScreenPixelData Pixel = GetScreenPixelData(PixelPos, GetDeviceZFromPixel(PixelPos));
+		WorldPosition = Pixel.WorldPosition;
+		SceneDepth = Pixel.SceneDepth;
+		bIsValid = true;
+		if (bPostBasePass != 0)
+		{
+			bIsValid = GetGBufferData(Pixel.ScreenUV).ShadingModelID != SHADINGMODELID_UNLIT;
+		}
+	}
 
 	// Dither pattern for page dilation
 	// We don't need to to check all 8 adjacent pages; as long as there's at least a single pixel near the edge
@@ -194,12 +209,9 @@ void GeneratePageFlagsFromPixels(uint3 GroupId : SV_GroupID, uint GroupIndex : S
 			   GroupIndex & 2 ? 1.0f : -1.0f);
 		
 	// excluding unlit to avoid including processing sky dome
-	if (bPostBasePass != 0)
+	if (!bIsValid)
 	{
-		if (Pixel.GBufferData.ShadingModelID == SHADINGMODELID_UNLIT)
-		{
-			return;
-		}
+		return;
 	}
 	
 	// Directional lights
@@ -212,19 +224,19 @@ void GeneratePageFlagsFromPixels(uint3 GroupId : SV_GroupID, uint GroupIndex : S
 
 		// Possibly use a single shadow projection data for the whole clipmap with just scales/biases for
 		// each level (to handle snapping). For now we maintain one per level.
-		//int ClipmapLevel = CalcClipmapLevel(length(Pixel.WorldPosition - ProjectionData.ClipmapWorldOrigin), ClipmapResolutionLodBias);
-		const int ClipmapLevel = CalcClipmapLevel(BaseProjectionData, Pixel.WorldPosition);
+		//int ClipmapLevel = CalcClipmapLevel(length(WorldPosition - ProjectionData.ClipmapWorldOrigin), ClipmapResolutionLodBias);
+		const int ClipmapLevel = CalcClipmapLevel(BaseProjectionData, WorldPosition);
 		int ClipmapIndex = max(0, ClipmapLevel - BaseProjectionData.ClipmapLevel);
 		if (ClipmapIndex < BaseProjectionData.ClipmapLevelCount)
 		{
-			MarkPage(ClipmapStartId + ClipmapIndex, 0, Pixel, PageDilationOffset);
+			MarkPage(ClipmapStartId + ClipmapIndex, 0, WorldPosition, PageDilationOffset);
 		}
 	}
 
 	// Local lights
 	{
 		uint EyeIndex = 0; // ??
-		uint3 GridCoordinate = ComputeLightGridCellCoordinate(PixelPos, Pixel.SceneDepth, EyeIndex);
+		uint3 GridCoordinate = ComputeLightGridCellCoordinate(PixelPos, SceneDepth, EyeIndex);
 		uint GridLinearIndex = ComputeLightGridCellIndex(GridCoordinate, EyeIndex);
 		const FCulledLightsGridData CulledLightGridData = GetCulledLightsGrid(GridLinearIndex, EyeIndex);
 
@@ -234,7 +246,7 @@ void GeneratePageFlagsFromPixels(uint3 GroupId : SV_GroupID, uint GroupIndex : S
 			const FLocalLightData LightData = GetLocalLightData(CulledLightGridData.DataStartIndex + Index, EyeIndex);
 			const uint LightGridLightIndex = ForwardLightData.CulledLightDataGrid[CulledLightGridData.DataStartIndex + Index];
 
-			float3 ToLight = normalize(LightData.LightPositionAndInvRadius.xyz - Pixel.WorldPosition);
+			float3 ToLight = normalize(LightData.LightPositionAndInvRadius.xyz - WorldPosition);
 			// Also do precise cone test, since froxels are pretty coarse at times.
 			if (dot(ToLight, LightData.LightDirectionAndShadowMask.xyz) < LightData.SpotAnglesAndSourceRadiusPacked.x)
 			{
@@ -251,8 +263,8 @@ void GeneratePageFlagsFromPixels(uint3 GroupId : SV_GroupID, uint GroupIndex : S
 					VirtualShadowMapId += GetCubeFace( -ToLight );
 				}
 
-				uint MipLevel = GetMipLevel(VirtualShadowMapId, Pixel);
-				MarkPage(VirtualShadowMapId, MipLevel, Pixel, PageDilationOffset);
+				uint MipLevel = GetMipLevel(VirtualShadowMapId, WorldPosition, SceneDepth);
+				MarkPage(VirtualShadowMapId, MipLevel, WorldPosition, PageDilationOffset);
 			}
 		}
 	}
diff --git a/Engine/Shaders/Private/VirtualShadowMaps/VirtualShadowMapProjection.usf b/Engine/Shaders/Private/VirtualShadowMaps/VirtualShadowMapProjection.usf
--- a/Engine/Shaders/Private/VirtualShadowMaps/VirtualShadowMapProjection.usf
+++ b/Engine/Shaders/Private/VirtualShadowMaps/VirtualShadowMapProjection.usf
@@ -15,6 +15,7 @@
 #include "../ScreenSpaceDenoise/SSDPublicHarmonics.ush"
 #include "../LightShaderParameters.ush"
 #include "../LightGridCommon.ush"
+#include "../HairStrands/HairStrandsVisibilityCommon.ush"
 #include "PageAccessCommon.ush"
 #include "ProjectionCommon.ush"
 #include "ProjectionDirectional.ush"
@@ -252,6 +253,9 @@ void VirtualShadowMapProjectionSpotPS(
 	OutputTarget = EncodeOutput(TranslatedWorldPosition, Light, LightType, RayMissFactor, HitDistance, SampleCount, DebugOutput);
 }
 
+// HAIR_TODO: move all data into a HairStrands UB
+uint bUseHairData;
+Texture2D<uint4> HairCategorizationTexture;
 
 void VirtualShadowMapProjectionDirectionalPS(
 	in float4 SvPosition : SV_Position,
@@ -260,13 +264,17 @@ void VirtualShadowMapProjectionDirectionalPS(
 {
 	// TODO: Unify this with the logic in PageManagement.usf - we need to get the same values in both passes
 	const float2 ScreenUV = SvPositionToBufferUV(SvPosition);
-	const float DeviceZ = LookupDeviceZ(ScreenUV);
+	float DeviceZ = LookupDeviceZ(ScreenUV);
+	if (bUseHairData > 0)
+	{
+		FCategorizationData CatData = DecodeCategorizationData(HairCategorizationTexture.Load(uint3(SvPosition.xy, 0)));
+		DeviceZ = CatData.ClosestDepth;
+	}
 	const float SceneDepth = ConvertFromDeviceZ(DeviceZ);
 	float3 TranslatedWorldPosition = SvPositionToTranslatedWorld(float4(SvPosition.xy, DeviceZ, 1.0f));
 
 	FLightShaderParameters Light = GetRootLightShaderParameters();
 		
-	FGBufferData GBufferData = GetGBufferData(ScreenUV);
 	float3 L = Light.Direction;
 	
 	float3 DebugOutput = float3(0, 0, 0);
@@ -282,6 +290,10 @@ void VirtualShadowMapProjectionDirectionalPS(
 		const float ContactShadowLengthWorld = ContactShadowLength * View.ClipToView[1][1] * SceneDepth;
 
 		float3 EstimatedGeoWorldNormal = GetEstimatedGeoWorldNormal(TranslatedWorldPosition, ScreenUV);
+		if (bUseHairData > 0)
+		{			
+			EstimatedGeoWorldNormal = L;
+		}
 
 		if (SMRTRayCount > 0)
 		{
diff --git a/Engine/Source/Runtime/Renderer/Private/DeferredShadingRenderer.cpp b/Engine/Source/Runtime/Renderer/Private/DeferredShadingRenderer.cpp
--- a/Engine/Source/Runtime/Renderer/Private/DeferredShadingRenderer.cpp
+++ b/Engine/Source/Runtime/Renderer/Private/DeferredShadingRenderer.cpp
@@ -2041,13 +2041,14 @@ void FDeferredShadingSceneRenderer::Render(FRDGBuilder& GraphBuilder)
 	CSV_CUSTOM_STAT(LightCount, ShadowOn, float(SortedLightSet.SortedLights.Num()) - float(SortedLightSet.AttenuationLightStart), ECsvCustomStatOp::Set);
 
 	// Local helper function to perform virtual shadow map allocation, which can occur early, or late.
+	FHairStrandsRenderingData* HairDatas = nullptr;
 	const auto AllocateVirtualShadowMaps = [&](bool bPostBasePass)
 	{
 		if (VirtualShadowMapArray.IsEnabled())
 		{
 			ensureMsgf(AreLightsInLightGrid(), TEXT("Virtual shadow map setup requires local lights to be injected into the light grid (this may be caused by 'r.LightCulling.Quality=0')."));
 			// ensure(ShadowMapSetupDone)
-			VirtualShadowMapArray.BuildPageAllocations(GraphBuilder, SceneTextures, Views, SortedLightSet, VisibleLightInfos, NaniteRasterResults, bPostBasePass, Scene->VirtualShadowMapArrayCacheManager);
+			VirtualShadowMapArray.BuildPageAllocations(GraphBuilder, SceneTextures, Views, SortedLightSet, VisibleLightInfos, NaniteRasterResults, bPostBasePass, Scene->VirtualShadowMapArrayCacheManager, HairDatas);
 		}
 	};
 
@@ -2064,7 +2065,7 @@ void FDeferredShadingSceneRenderer::Render(FRDGBuilder& GraphBuilder)
 	};
 
 	// Early occlusion queries
-	const bool bOcclusionBeforeBasePass = !bNaniteEnabled && !bAnyLumenEnabled && ((DepthPass.EarlyZPassMode == EDepthDrawingMode::DDM_AllOccluders) || bIsEarlyDepthComplete);
+	const bool bOcclusionBeforeBasePass = !bNaniteEnabled && !bAnyLumenEnabled && !bHairEnable && ((DepthPass.EarlyZPassMode == EDepthDrawingMode::DDM_AllOccluders) || bIsEarlyDepthComplete);
 
 	if (bOcclusionBeforeBasePass)
 	{
@@ -2158,8 +2159,7 @@ void FDeferredShadingSceneRenderer::Render(FRDGBuilder& GraphBuilder)
 		bool bSkipPerPixelTracing = true;
 		bAsyncComputeVolumetricCloud = RenderVolumetricCloud(GraphBuilder, SceneTextures, bSkipVolumetricRenderTarget, bSkipPerPixelTracing, HalfResolutionDepthCheckerboardMinMaxTexture, true, InstanceCullingManager);
 	}
-
-	FHairStrandsRenderingData* HairDatas = nullptr;
+	
 	FHairStrandsRenderingData* HairDatasStorage = GraphBuilder.AllocObject<FHairStrandsRenderingData>();
 
 	FRDGTextureRef ForwardScreenSpaceShadowMaskTexture = nullptr;
@@ -2292,6 +2292,14 @@ void FDeferredShadingSceneRenderer::Render(FRDGBuilder& GraphBuilder)
 		AddResolveSceneColorPass(GraphBuilder, Views, SceneTextures.Color);
 	}
 
+	// Render hair
+	if (bHairEnable && !IsForwardShadingEnabled(ShaderPlatform))
+	{
+		RenderHairPrePass(GraphBuilder, Scene, Views, InstanceCullingManager, *HairDatasStorage);
+		RenderHairBasePass(GraphBuilder, Scene, SceneTextures, Views, InstanceCullingManager, *HairDatasStorage);
+		HairDatas = HairDatasStorage;
+	}
+
 #if RHI_RAYTRACING
 	const bool bRayTracingEnabled = IsRayTracingEnabled();
 	if (bRayTracingEnabled)
@@ -2375,13 +2383,6 @@ void FDeferredShadingSceneRenderer::Render(FRDGBuilder& GraphBuilder)
 		// TODO: Populate velocity buffer from Nanite visibility buffer.
 	}
 
-	// Hair base pass for deferred shading
-	if (bHairEnable && !IsForwardShadingEnabled(ShaderPlatform))
-	{
-		RenderHairPrePass(GraphBuilder, Scene, Views, InstanceCullingManager, *HairDatasStorage);
-		HairDatas = HairDatasStorage;
-	}
-
 	// Copy lighting channels out of stencil before deferred decals which overwrite those values
 	FRDGTextureRef LightingChannelsTexture = CopyStencilToLightingChannelTexture(GraphBuilder, SceneTextures.Stencil);
 
@@ -2413,13 +2414,6 @@ void FDeferredShadingSceneRenderer::Render(FRDGBuilder& GraphBuilder)
 		}
 	}
 
-	// Hair base pass for deferred shading
-	if (bHairEnable && !IsForwardShadingEnabled(ShaderPlatform))
-	{
-		check(HairDatas);
-		RenderHairBasePass(GraphBuilder, Scene,  SceneTextures, Views, InstanceCullingManager, *HairDatasStorage);
-	}
-
 	// Rebuild scene textures to include velocity, custom depth, and SSAO.
 	SceneTextures.SetupMode |= ESceneTextureSetupMode::All;
 	SceneTextures.UniformBuffer = CreateSceneTextureUniformBuffer(GraphBuilder, FeatureLevel, SceneTextures.SetupMode);
diff --git a/Engine/Source/Runtime/Renderer/Private/ShadowRendering.cpp b/Engine/Source/Runtime/Renderer/Private/ShadowRendering.cpp
--- a/Engine/Source/Runtime/Renderer/Private/ShadowRendering.cpp
+++ b/Engine/Source/Runtime/Renderer/Private/ShadowRendering.cpp
@@ -1876,6 +1876,7 @@ void FDeferredShadingSceneRenderer::RenderShadowProjections(
 		}
 	}
 
+	FRDGTextureRef HairCatgorizationTexture = HairDatas && HairDatas->HairVisibilityViews.HairDatas.Num() > 0 ? HairDatas->HairVisibilityViews.HairDatas[0].CategorizationTexture : nullptr;
 	if (NormalShadows.Num() > 0)
 	{
 		const auto RenderNormalShadows = [&](FRDGTextureRef OutputTexture, FExclusiveDepthStencil ExclusiveDepthStencil, bool bSubPixel)
@@ -1886,7 +1887,7 @@ void FDeferredShadingSceneRenderer::RenderShadowProjections(
 
 			FShadowProjectionPassParameters CommonPassParameters;
 			CommonPassParameters.SceneTextures = GetSceneTextureShaderParameters(SceneTextures.UniformBuffer);
-			CommonPassParameters.HairCategorizationTexture = bSubPixel && HairDatas && HairDatas->HairVisibilityViews.HairDatas.Num() > 0 ? HairDatas->HairVisibilityViews.HairDatas[0].CategorizationTexture : nullptr;
+			CommonPassParameters.HairCategorizationTexture = bSubPixel && HairCatgorizationTexture ? HairCatgorizationTexture : nullptr;
 			CommonPassParameters.RenderTargets[0] = FRenderTargetBinding(OutputTexture, ERenderTargetLoadAction::ELoad);
 			CommonPassParameters.RenderTargets.DepthStencil = 
 				bSubPixel && HairDatas && HairDatas->HairVisibilityViews.HairDatas.Num() > 0 ?
@@ -1909,7 +1910,7 @@ void FDeferredShadingSceneRenderer::RenderShadowProjections(
 			RenderNormalShadows(ScreenShadowMaskTexture, FExclusiveDepthStencil::DepthRead_StencilWrite, false);
 		}
 
-		if (ScreenShadowMaskSubPixelTexture && HairDatas && HairDatas->HairVisibilityViews.HairDatas.Num() > 0 && HairDatas->HairVisibilityViews.HairDatas[0].CategorizationTexture)
+		if (ScreenShadowMaskSubPixelTexture && HairCatgorizationTexture)
 		{
 			RDG_EVENT_SCOPE(GraphBuilder, "SubPixelShadows");
 
@@ -2020,6 +2021,7 @@ void FDeferredShadingSceneRenderer::RenderShadowProjections(
 								VirtualShadowMapArray,
 								ScissorRect,
 								ScreenShadowMaskTexture,
+								nullptr,
 								bProjectingForForwardShading);
 						}
 						else
@@ -2032,6 +2034,37 @@ void FDeferredShadingSceneRenderer::RenderShadowProjections(
 								VirtualShadowMapArray,
 								ScissorRect,
 								ScreenShadowMaskTexture,
+								nullptr,
+								bProjectingForForwardShading);
+						}
+					}
+
+					// Sub-pixel shadow (no filtering for hair)
+					if (ScreenShadowMaskSubPixelTexture && HairCatgorizationTexture)
+					{
+						if (VisibleLightInfo.VirtualShadowMapClipmaps.Num() > 0)
+						{
+							RenderVirtualShadowMapProjection(
+								VisibleLightInfo.FindShadowClipmapForView(&View),
+								GraphBuilder,
+								View,
+								VirtualShadowMapArray,
+								ScissorRect,
+								ScreenShadowMaskSubPixelTexture,
+								HairCatgorizationTexture,
+								bProjectingForForwardShading);
+						}
+						else
+						{
+							check(VirtualShadowMaps.Num() == 1);
+							RenderVirtualShadowMapProjection(
+								VirtualShadowMaps[0],
+								GraphBuilder,
+								View,
+								VirtualShadowMapArray,
+								ScissorRect,
+								ScreenShadowMaskSubPixelTexture,
+								HairCatgorizationTexture,
 								bProjectingForForwardShading);
 						}
 					}
diff --git a/Engine/Source/Runtime/Renderer/Private/VirtualShadowMaps/VirtualShadowMapArray.cpp b/Engine/Source/Runtime/Renderer/Private/VirtualShadowMaps/VirtualShadowMapArray.cpp
--- a/Engine/Source/Runtime/Renderer/Private/VirtualShadowMaps/VirtualShadowMapArray.cpp
+++ b/Engine/Source/Runtime/Renderer/Private/VirtualShadowMaps/VirtualShadowMapArray.cpp
@@ -11,6 +11,7 @@
 #include "VirtualShadowMapCacheManager.h"
 #include "VirtualShadowMapClipmap.h"
 #include "ComponentRecreateRenderStateContext.h"
+#include "HairStrands/HairStrandsRendering.h"
 
 IMPLEMENT_STATIC_UNIFORM_BUFFER_SLOT(VirtualShadowMapUbSlot);
 
@@ -316,9 +317,11 @@ class FGeneratePageFlagsFromPixelsCS : public FVirtualPageManagementShader
 		SHADER_PARAMETER_STRUCT_REF(FViewUniformShaderParameters, View)
 		SHADER_PARAMETER_STRUCT_REF(FForwardLightData, ForwardLightData)
 		SHADER_PARAMETER_RDG_TEXTURE(Texture2D<uint2>, VisBuffer64)
+		SHADER_PARAMETER_RDG_TEXTURE(Texture2D<uint4>, HairCategorizationTexture)
 		SHADER_PARAMETER_RDG_BUFFER_UAV(RWStructuredBuffer<uint>, OutPageRequestFlags)
 		SHADER_PARAMETER_RDG_BUFFER_SRV(StructuredBuffer< FVirtualShadowMapProjectionShaderData >, VirtualShadowMapProjectionData)
 		SHADER_PARAMETER_RDG_BUFFER_SRV(StructuredBuffer< uint >, VirtualShadowMapIdRemap)
+		SHADER_PARAMETER(uint32, bUseHairData)
 		SHADER_PARAMETER(uint32, NumDirectionalLightSmInds)
 		SHADER_PARAMETER(uint32, bPostBasePass)
 		SHADER_PARAMETER(float, LodFootprintScale)
@@ -683,7 +686,8 @@ void FVirtualShadowMapArray::BuildPageAllocations(
 	const TArray<FVisibleLightInfo, SceneRenderingAllocator>& VisibleLightInfos,
 	const TArray<Nanite::FRasterResults, TInlineAllocator<2>>& NaniteRasterResults,
 	bool bPostBasePass,
-	FVirtualShadowMapArrayCacheManager* VirtualShadowMapArrayCacheManager)
+	FVirtualShadowMapArrayCacheManager* VirtualShadowMapArrayCacheManager,
+	FHairStrandsRenderingData* HairDatas)
 {
 	check(IsEnabled());
 	RDG_EVENT_SCOPE(GraphBuilder, "FVirtualShadowMapArray::GeneratePageFlagsFromLightGrid");
@@ -846,39 +850,50 @@ void FVirtualShadowMapArray::BuildPageAllocations(
 				// Mark pages based on projected depth buffer pixels
 				if (CVarMarkPixelPages.GetValueOnRenderThread() != 0)
 				{
-					FGeneratePageFlagsFromPixelsCS::FPermutationDomain PermutationVector;
-					PermutationVector.Set<FGeneratePageFlagsFromPixelsCS::FNaniteDepthBufferDim>(NaniteVisBuffer64 != nullptr && !bPostBasePass);
-
-					FGeneratePageFlagsFromPixelsCS::FParameters* PassParameters = GraphBuilder.AllocParameters< FGeneratePageFlagsFromPixelsCS::FParameters >();
-					PassParameters->VirtualSmCommon = GetCommonUniformBuffer(GraphBuilder);
-
-					PassParameters->SceneTexturesStruct = SceneTextures.UniformBuffer;
-					PassParameters->bPostBasePass = bPostBasePass;
-				
-					PassParameters->VisBuffer64 = VisBuffer64;
-					PassParameters->View = View.ViewUniformBuffer;
-					PassParameters->OutPageRequestFlags = PageRequestFlagsUAV;
-					PassParameters->ForwardLightData = View.ForwardLightingResources->ForwardLightDataUniformBuffer;
-					PassParameters->VirtualShadowMapIdRemap = GraphBuilder.CreateSRV(VirtualShadowMapIdRemapRDG[ViewIndex]);
-					PassParameters->VirtualShadowMapProjectionData = GraphBuilder.CreateSRV(ShadowMapProjectionDataRDG);
-					PassParameters->NumDirectionalLightSmInds = DirectionalLightSmInds.Num();
-					PassParameters->LodFootprintScale = LodFootprintScale;
-					PassParameters->PageDilationBorderSize = PageDilationBorderSize;
-				
-					auto ComputeShader = View.ShaderMap->GetShader<FGeneratePageFlagsFromPixelsCS>(PermutationVector);
-
-					static_assert((FVirtualPageManagementShader::DefaultCSGroupXY % 2) == 0, "GeneratePageFlagsFromPixels requires even-sized CS groups for quad swizzling.");
-					const FIntPoint GridSize = FIntPoint::DivideAndRoundUp(View.ViewRect.Size(), FVirtualPageManagementShader::DefaultCSGroupXY);
-
-					FComputeShaderUtils::AddPass(
-						GraphBuilder,
-						RDG_EVENT_NAME("GeneratePageFlagsFromPixels"),
-						ComputeShader,
-						PassParameters,
-						FIntVector(GridSize.X, GridSize.Y, 1)
-					);
+					auto GeneratePageFlags = [&](FRDGTextureRef HairCategorizationTexture)
+					{
+						FGeneratePageFlagsFromPixelsCS::FPermutationDomain PermutationVector;
+						PermutationVector.Set<FGeneratePageFlagsFromPixelsCS::FNaniteDepthBufferDim>(NaniteVisBuffer64 != nullptr && !bPostBasePass);
+
+						FGeneratePageFlagsFromPixelsCS::FParameters* PassParameters = GraphBuilder.AllocParameters< FGeneratePageFlagsFromPixelsCS::FParameters >();
+						PassParameters->VirtualSmCommon = GetCommonUniformBuffer(GraphBuilder);
+
+						PassParameters->SceneTexturesStruct = SceneTextures.UniformBuffer;
+						PassParameters->bPostBasePass = bPostBasePass;
+
+						PassParameters->VisBuffer64 = VisBuffer64;
+						PassParameters->bUseHairData = HairCategorizationTexture != nullptr ? 1 : 0;
+						PassParameters->HairCategorizationTexture = HairCategorizationTexture ? HairCategorizationTexture : GSystemTextures.GetBlackAlphaOneDummy(GraphBuilder);
+						PassParameters->View = View.ViewUniformBuffer;
+						PassParameters->OutPageRequestFlags = PageRequestFlagsUAV;
+						PassParameters->ForwardLightData = View.ForwardLightingResources->ForwardLightDataUniformBuffer;
+						PassParameters->VirtualShadowMapIdRemap = GraphBuilder.CreateSRV(VirtualShadowMapIdRemapRDG[ViewIndex]);
+						PassParameters->VirtualShadowMapProjectionData = GraphBuilder.CreateSRV(ShadowMapProjectionDataRDG);
+						PassParameters->NumDirectionalLightSmInds = DirectionalLightSmInds.Num();
+						PassParameters->LodFootprintScale = LodFootprintScale;
+						PassParameters->PageDilationBorderSize = PageDilationBorderSize;
+
+						auto ComputeShader = View.ShaderMap->GetShader<FGeneratePageFlagsFromPixelsCS>(PermutationVector);
+
+						static_assert((FVirtualPageManagementShader::DefaultCSGroupXY % 2) == 0, "GeneratePageFlagsFromPixels requires even-sized CS groups for quad swizzling.");
+						const FIntPoint GridSize = FIntPoint::DivideAndRoundUp(View.ViewRect.Size(), FVirtualPageManagementShader::DefaultCSGroupXY);
+
+						FComputeShaderUtils::AddPass(
+							GraphBuilder,
+							RDG_EVENT_NAME("GeneratePageFlagsFromPixels"),
+							ComputeShader,
+							PassParameters,
+							FIntVector(GridSize.X, GridSize.Y, 1)
+						);
+					};
+
+					GeneratePageFlags(nullptr);
+					FRDGTextureRef HairCategorizationTexture = HairDatas && ViewIndex < HairDatas->HairVisibilityViews.HairDatas.Num() ? HairDatas->HairVisibilityViews.HairDatas[ViewIndex].CategorizationTexture : nullptr;
+					if (HairCategorizationTexture)
+					{
+						GeneratePageFlags(HairCategorizationTexture);
+					}
 				}
-
 				// Mark coarse pages
 				if (CVarMarkCoarsePages.GetValueOnRenderThread() != 0)
 				{
diff --git a/Engine/Source/Runtime/Renderer/Private/VirtualShadowMaps/VirtualShadowMapArray.h b/Engine/Source/Runtime/Renderer/Private/VirtualShadowMaps/VirtualShadowMapArray.h
--- a/Engine/Source/Runtime/Renderer/Private/VirtualShadowMaps/VirtualShadowMapArray.h
+++ b/Engine/Source/Runtime/Renderer/Private/VirtualShadowMaps/VirtualShadowMapArray.h
@@ -192,7 +192,8 @@ public:
 		const TArray<FVisibleLightInfo, SceneRenderingAllocator> &VisibleLightInfos, 
 		const TArray<Nanite::FRasterResults, TInlineAllocator<2>> &NaniteRasterResults, 
 		bool bPostBasePass, 
-		FVirtualShadowMapArrayCacheManager *VirtualShadowMapArrayCacheManager);
+		FVirtualShadowMapArrayCacheManager *VirtualShadowMapArrayCacheManager,
+		struct FHairStrandsRenderingData* HairDatas);
 
 	bool IsAllocated() const
 	{
diff --git a/Engine/Source/Runtime/Renderer/Private/VirtualShadowMaps/VirtualShadowMapProjection.cpp b/Engine/Source/Runtime/Renderer/Private/VirtualShadowMaps/VirtualShadowMapProjection.cpp
--- a/Engine/Source/Runtime/Renderer/Private/VirtualShadowMaps/VirtualShadowMapProjection.cpp
+++ b/Engine/Source/Runtime/Renderer/Private/VirtualShadowMaps/VirtualShadowMapProjection.cpp
@@ -111,6 +111,7 @@ BEGIN_SHADER_PARAMETER_STRUCT(FProjectionParameters, )
 	SHADER_PARAMETER_STRUCT(FLightShaderParameters, Light)
 	SHADER_PARAMETER_RDG_UNIFORM_BUFFER(FSceneTextureUniformParameters, SceneTexturesStruct)
 	SHADER_PARAMETER_STRUCT_REF(FViewUniformShaderParameters, View)
+	SHADER_PARAMETER_RDG_TEXTURE(Texture2D<uint4>, HairCategorizationTexture)
 	SHADER_PARAMETER(int32, VirtualShadowMapId)
 	SHADER_PARAMETER(float, ContactShadowLength)
 	SHADER_PARAMETER(float, NormalOffsetWorld)
@@ -119,6 +120,7 @@ BEGIN_SHADER_PARAMETER_STRUCT(FProjectionParameters, )
 	SHADER_PARAMETER(int32, SMRTSamplesPerRay)
 	SHADER_PARAMETER(float, SMRTRayLengthScale)
 	SHADER_PARAMETER(float, SMRTCotMaxRayAngleFromLight)
+	SHADER_PARAMETER(uint32, bUseHairData)
 	RENDER_TARGET_BINDING_SLOTS()
 END_SHADER_PARAMETER_STRUCT()
 
@@ -203,6 +205,7 @@ static bool AddPass_RenderVirtualShadowMapProjection(
 	FVirtualShadowMapArray& VirtualShadowMapArray,
 	int VirtualShadowMapId,
 	const FIntRect ScissorRect,
+	FRDGTextureRef HairCategorization,
 	EVirtualShadowMapProjectionOutputType OutputType,
 	const FRenderTargetBinding& RenderTargetBinding,
 	FRHIBlendState* BlendState)
@@ -223,6 +226,8 @@ static bool AddPass_RenderVirtualShadowMapProjection(
 	PassParameters->DebugOutputType = CVarVirtualShadowMapDebugProjection.GetValueOnRenderThread();
 	PassParameters->ContactShadowLength = CVarContactShadowLength.GetValueOnRenderThread();
 	PassParameters->NormalOffsetWorld = CVarNormalOffsetWorld.GetValueOnRenderThread();
+	PassParameters->bUseHairData = HairCategorization != nullptr ? 1 : 0;
+	PassParameters->HairCategorizationTexture = HairCategorization ? HairCategorization : GSystemTextures.GetBlackDummy(GraphBuilder);
 
 	PassParameters->RenderTargets[0] = RenderTargetBinding;
 
@@ -298,6 +303,7 @@ static bool RenderVirtualShadowMapProjectionCommon(
 	FVirtualShadowMapArray& VirtualShadowMapArray,
 	int32 VirtualShadowMapId,
 	const FIntRect ScissorRect,
+	FRDGTextureRef HairCategorization,
 	EVirtualShadowMapProjectionOutputType OutputType,
 	const FRenderTargetBinding& RenderTargetBinding,
 	FRHIBlendState* BlendState)
@@ -310,6 +316,7 @@ static bool RenderVirtualShadowMapProjectionCommon(
 		VirtualShadowMapArray,
 		VirtualShadowMapId,
 		ScissorRect,
+		HairCategorization,
 		OutputType,
 		RenderTargetBinding,
 		BlendState);
@@ -328,6 +335,7 @@ static bool RenderVirtualShadowMapProjectionCommon(
 			VirtualShadowMapArray,
 			VirtualShadowMapId,
 			ScissorRect,
+			HairCategorization,
 			EVirtualShadowMapProjectionOutputType::Debug,
 			DebugRenderTargetBinding,
 			nullptr);
@@ -355,6 +363,7 @@ void RenderVirtualShadowMapProjectionForDenoising(
 		VirtualShadowMapArray,
 		ShadowInfo->VirtualShadowMaps[0]->ID,
 		ScissorRect,
+		nullptr,
 		EVirtualShadowMapProjectionOutputType::Denoiser,
 		RenderTargetBinding,
 		nullptr);
@@ -377,6 +386,7 @@ void RenderVirtualShadowMapProjectionForDenoising(
 		VirtualShadowMapArray,
 		Clipmap->GetVirtualShadowMap()->ID,
 		ScissorRect,
+		nullptr,
 		EVirtualShadowMapProjectionOutputType::Denoiser,
 		RenderTargetBinding,
 		nullptr);
@@ -389,6 +399,7 @@ void RenderVirtualShadowMapProjection(
 	FVirtualShadowMapArray& VirtualShadowMapArray,
 	const FIntRect ScissorRect,
 	FRDGTextureRef ScreenShadowMaskTexture,
+	FRDGTextureRef HairCategorization,
 	bool bProjectingForForwardShading)
 {
 	FRenderTargetBinding RenderTargetBinding(ScreenShadowMaskTexture, ERenderTargetLoadAction::ELoad);
@@ -400,6 +411,7 @@ void RenderVirtualShadowMapProjection(
 		VirtualShadowMapArray,
 		ShadowInfo->VirtualShadowMaps[0]->ID,
 		ScissorRect,
+		HairCategorization,
 		EVirtualShadowMapProjectionOutputType::ScreenShadowMask,
 		RenderTargetBinding,
 		BlendState);
@@ -412,6 +424,7 @@ void RenderVirtualShadowMapProjection(
 	FVirtualShadowMapArray& VirtualShadowMapArray,
 	const FIntRect ScissorRect,
 	FRDGTextureRef ScreenShadowMaskTexture,
+	FRDGTextureRef HairCategorization,
 	bool bProjectingForForwardShading)
 {
 	FRenderTargetBinding RenderTargetBinding(ScreenShadowMaskTexture, ERenderTargetLoadAction::ELoad);	
@@ -424,6 +437,7 @@ void RenderVirtualShadowMapProjection(
 		VirtualShadowMapArray,
 		Clipmap->GetVirtualShadowMap()->ID,
 		ScissorRect,
+		HairCategorization,
 		EVirtualShadowMapProjectionOutputType::ScreenShadowMask,
 		RenderTargetBinding,
 		BlendState);
diff --git a/Engine/Source/Runtime/Renderer/Private/VirtualShadowMaps/VirtualShadowMapProjection.h b/Engine/Source/Runtime/Renderer/Private/VirtualShadowMaps/VirtualShadowMapProjection.h
--- a/Engine/Source/Runtime/Renderer/Private/VirtualShadowMaps/VirtualShadowMapProjection.h
+++ b/Engine/Source/Runtime/Renderer/Private/VirtualShadowMaps/VirtualShadowMapProjection.h
@@ -38,6 +38,7 @@ void RenderVirtualShadowMapProjection(
 	FVirtualShadowMapArray& VirtualShadowMapArray,
 	const FIntRect ScissorRect,
 	FRDGTextureRef ScreenShadowMaskTexture,
+	FRDGTextureRef HairCategorization,
 	bool bProjectingForForwardShading);
 
 void RenderVirtualShadowMapProjection(
@@ -47,6 +48,7 @@ void RenderVirtualShadowMapProjection(
 	FVirtualShadowMapArray& VirtualShadowMapArray,
 	const FIntRect ScissorRect,
 	FRDGTextureRef ScreenShadowMaskTexture,
+	FRDGTextureRef HairCategorization,
 	bool bProjectingForForwardShading);
 
 void RenderVirtualShadowMapProjection(
