Enable specular multiple scattering for analytical lighting.

diff --git a/Engine/Shaders/Private/DeferredLightPixelShaders.usf b/Engine/Shaders/Private/DeferredLightPixelShaders.usf
--- a/Engine/Shaders/Private/DeferredLightPixelShaders.usf
+++ b/Engine/Shaders/Private/DeferredLightPixelShaders.usf
@@ -33,6 +33,7 @@ float CloudShadowmapStrength;
 
 #if STRATA_ENABLED
 #define STRATA_INLINE_SHADING 0
+#define STRATA_MICROFACET_MULTIPLESCATTERING 1
 #include "/Engine/Private/Strata/StrataEvaluation.ush"
 #include "/Engine/Private/Strata/StrataDeferredLighting.ush"
 #endif
diff --git a/Engine/Shaders/Private/Strata/StrataEvaluation.ush b/Engine/Shaders/Private/Strata/StrataEvaluation.ush
--- a/Engine/Shaders/Private/Strata/StrataEvaluation.ush
+++ b/Engine/Shaders/Private/Strata/StrataEvaluation.ush
@@ -3,6 +3,7 @@
 #pragma once
 
 #include "Strata.ush"
+#include "StrataMultipleScattering.ush"
 
 #include "../ParticipatingMediaCommon.ush"
 #include "../MonteCarlo.ush"
diff --git a/Engine/Shaders/Private/Strata/StrataFurnaceTest.usf b/Engine/Shaders/Private/Strata/StrataFurnaceTest.usf
--- a/Engine/Shaders/Private/Strata/StrataFurnaceTest.usf
+++ b/Engine/Shaders/Private/Strata/StrataFurnaceTest.usf
@@ -5,7 +5,6 @@
 #define STRATA_INLINE_SHADING 0
 #define STRATA_MICROFACET_MULTIPLESCATTERING 1
 
-#include "StrataMultipleScattering.ush"
 #include "StrataEvaluation.ush"
 
 #define STRATA_INTEGRATION_EVALUATE 0
