BSDF model for a slab of Isotropic homogeneous slab of particiapting media. (punctual and environmental light)
the user set up the color as wanted post scattering and for a visually thick medium.
The light response is then computed for such as setup avec the slab thickness is taken into account.

diff --git a/Engine/Shaders/Private/ParticipatingMediaCommon.ush b/Engine/Shaders/Private/ParticipatingMediaCommon.ush
--- a/Engine/Shaders/Private/ParticipatingMediaCommon.ush
+++ b/Engine/Shaders/Private/ParticipatingMediaCommon.ush
@@ -6,8 +6,8 @@
 
 #include "Common.ush"
 
-#define PARTICIPATING_MEDIA_MIN_MFP_METER		0.000001f
-#define PARTICIPATING_MEDIA_MIN_EXTINCTION		0.000001f
+#define PARTICIPATING_MEDIA_MIN_MFP_METER		0.000000000001f
+#define PARTICIPATING_MEDIA_MIN_EXTINCTION		0.000000000001f
 #define PARTICIPATING_MEDIA_MIN_TRANSMITTANCE	0.000000000001f
 
 
@@ -152,3 +152,59 @@ float3 TransmittanceToMeanFreePath(in float3 TransmittanceColor, in float Thickn
 	return 1.0f / max(PARTICIPATING_MEDIA_MIN_EXTINCTION, TransmittanceToExtinction(TransmittanceColor, ThicknessMeters));
 }
 
+
+
+//---------------------------------------------
+// Lighting evaluation function for an Isotropic participating media a slab of 1 meter
+//---------------------------------------------
+
+float3 IsotropicMediumSlabDirectionalAlbedoFade(float3 BaseColor, float3 MFP)
+{
+	float3 Fade;
+	// Progressively fade out to 0 for base color starting from 0.1 down to 0.
+	Fade = saturate(BaseColor * 10.0f);
+	// And also fade out from MFP = 20 (the last measured MFP for which we have fit the curve) to 0 for a large MFP=1000 meters
+	Fade*= saturate(1.0f - (MFP - 20.0f) / 980.0f);
+
+	return Fade;
+}
+
+float3 IsotropicMediumSlabPunctualDirectionalAlbedo(FParticipatingMedia PM)
+{
+	// Our measurement are only valid for a MFP of 0.01 meter so we clamp input MFP to that.
+	// This is relatively ok since all computation are done for a slab of 1 meter.
+	const float3 MFP = max(0.01f, PM.MeanFreePath);
+
+	const float3 EvaluateForBaseColor1	= 0.0855674 / (0.237742 + (MFP + ((0.0310849 - MFP) / (1.95492 * MFP + 2.07238))));
+	const float3 EvaluateForBaseColor01 = 0.0167964 / (0.541037 * (pow(1.17902, (-4.33046) / MFP) * (-0.294969 + MFP)) + 0.797592);
+
+	// Lerp between the two curves extreme curves
+	float3 FinalEvaluate = lerp(EvaluateForBaseColor01, EvaluateForBaseColor1, (PM.BaseColor - 0.1f) / (1.0f - 0.1f));
+	return FinalEvaluate * IsotropicMediumSlabDirectionalAlbedoFade(PM.BaseColor, MFP);
+}
+
+float3 IsotropicMediumSlabEnvDirectionalAlbedo(FParticipatingMedia PM)
+{
+	// Our measurement are only valid for a MFP of 0.01 meter so we clamp input MFP to that.
+	// This is relatively ok since all computation are done for a slab of 1 meter.
+	const float3 MFP = max(0.01f, PM.MeanFreePath);
+
+	const float3 EvaluateForBaseColor1	= 0.00231881 + (0.51379 / (pow(MFP, 1.03577) + 0.510465)); 
+	const float3 EvaluateForBaseColor01 = 0.189167 / (1.55597 + (MFP + pow(0.182843, 0.0666775 + MFP)));
+
+	// Lerp between the two curves extreme curves
+	float3 FinalEvaluate = lerp(EvaluateForBaseColor01, EvaluateForBaseColor1, (PM.BaseColor - 0.1f) / (1.0f - 0.1f));
+	return FinalEvaluate * IsotropicMediumSlabDirectionalAlbedoFade(PM.BaseColor, MFP);
+}
+
+float3 IsotropicMediumSlabTransmittance(FParticipatingMedia PM, float SlabThickness, float NoV)
+{
+	const float3 SafeExtinctionThreshold = 0.000001f;
+	const float3 SafeExtinctionCoefficients = max(SafeExtinctionThreshold, PM.ExtinctionCoef);
+
+	const float PathLength = SlabThickness / max(0.0001f, abs(NoV));
+	const float3 SafePathSegmentTransmittance = exp(-SafeExtinctionCoefficients * PathLength);
+
+	return SafePathSegmentTransmittance;
+}
+
diff --git a/Engine/Shaders/Private/Strata/Strata.ush b/Engine/Shaders/Private/Strata/Strata.ush
--- a/Engine/Shaders/Private/Strata/Strata.ush
+++ b/Engine/Shaders/Private/Strata/Strata.ush
@@ -962,10 +962,17 @@ FParticipatingMedia StrataSlabCreateParticipatingMedia(float3 DiffuseColor, floa
 
 float StrataSlabDiffuseToVolumeBlend(FParticipatingMedia PM)
 {
-	// End   lerp when transmittance = 0.05, when extinction over 1cm is = -ln(0.05)/0.01 = 299.57 => MFP = 0.003338
-	// Start lerp when transmittance = 0.01, when extinction over 1cm is = -ln(0.01)/0.01 = 230.25 => MFP = 0.004343
+	// Remember, computation are made for a slab of 1 metter, according to STRATA_SIMPLEVOLUME_THICKNESS_M.
+	// We start blending from diffuse toward a slab of participating media when the mean free path reaches 4 centimeters.
 	const float MaxMeanFreePathMeters = max(PM.MeanFreePath.x, max(PM.MeanFreePath.y, PM.MeanFreePath.z));
-	return saturate((MaxMeanFreePathMeters - 0.003338) / (0.004343 - 0.003338));
+	// Start blending volumetric after mean free path of 4 centimeters (also needed because threshold when recovering extinction from transmittance cause a small minimum value)
+	const float StartVolumeBlendMFP = 0.04f;
+	// And when the mean free path reaches 1/3.0f = 0.33 meter. That is extinction=1/0.33=3, transmittance at 1m (bottom of the slab) = exp(-3) = 0.05) we fully use the volumetric model.
+	const float Blend = saturate(max(0.0, MaxMeanFreePathMeters - StartVolumeBlendMFP) * 3.0f);
+	//const float Blend = max(0.0, MaxMeanFreePathMeters - StartVolumeBlendMFP) > 0.0 ? 1.0 : 0.0;
+
+	// Using a non linear blend helps smoothing out transitions visually.
+	return Blend * Blend;
 }
 
 struct FStrataSlabSingleScattering
diff --git a/Engine/Shaders/Private/Strata/StrataEvaluation.ush b/Engine/Shaders/Private/Strata/StrataEvaluation.ush
--- a/Engine/Shaders/Private/Strata/StrataEvaluation.ush
+++ b/Engine/Shaders/Private/Strata/StrataEvaluation.ush
@@ -353,7 +353,8 @@ FStrataEvaluateResult StrataEvaluateBSDF(FStrataBSDFContext BSDFContext, FShadow
 			{
 				FParticipatingMedia PM = StrataSlabCreateParticipatingMedia(DiffuseColor, SLAB_SSSDMFP(BSDFContext.BSDF));
 				const float DiffuseToVolumeBlend = StrataSlabDiffuseToVolumeBlend(PM);
-				const FStrataSlabSingleScattering Result = StrataSlabSingleScattering(PM, BSDFContext.Context.NoV);
+				const float3 SlabDirectionalAlbedo = IsotropicMediumSlabPunctualDirectionalAlbedo(PM);
+				const float3 SlabTransmittance = IsotropicMediumSlabTransmittance(PM, STRATA_SIMPLEVOLUME_THICKNESS_M, BSDFContext.Context.NoV);
 
 				// STRATA_TODO Fresnel throughput, Use a LUT, this is currently in progress and will be coming soon.
 				// In the meantime, we use a simple approach to feed Fresnel function and compute an approximation of the general transmition/througput.
@@ -363,9 +364,9 @@ FStrataEvaluateResult StrataEvaluateBSDF(FStrataBSDFContext BSDFContext, FShadow
 				const float3 SpecularThroughput = 1.0f - ApproxFresnelForTransmittance; // This is already taken into account by the Chan model above. But for volumetric we need to apply it.
 				const float Phase = IsotropicPhase();
 
-				Sample.DiffusePathValue = lerp(Sample.DiffusePathValue,		Result.SingleScattering * Phase * SpecularThroughput,	DiffuseToVolumeBlend);
+				Sample.DiffusePathValue = lerp(Sample.DiffusePathValue,		SlabDirectionalAlbedo * SpecularThroughput,				DiffuseToVolumeBlend);
 				Sample.DiffusePDF		= lerp(Sample.DiffusePDF,			Phase,													DiffuseToVolumeBlend);
-				Sample.Throughput		= lerp(Sample.Throughput,			Result.Transmittance * SpecularThroughput,				DiffuseToVolumeBlend);
+				Sample.Throughput		= lerp(Sample.Throughput,			SlabTransmittance * SpecularThroughput,					DiffuseToVolumeBlend);
 
 				Sample.bSubsurface = false;	// It should already be the case because Enforce in this case because BSDF_GETHASSSS should be false when BSDF_GETISSIMPLEVOLUME is true
 			}
@@ -673,11 +674,31 @@ FStrataEnvLightResult StrataEvaluateForEnvLight(FStrataBSDFContext BSDFContext,
 			{
 				FParticipatingMedia PM = StrataSlabCreateParticipatingMedia(DiffuseColor, SLAB_SSSDMFP(BSDFContext.BSDF));
 				const float DiffuseToVolumeBlend = StrataSlabDiffuseToVolumeBlend(PM);
-				const FStrataSlabSingleScattering Result = StrataSlabSingleScattering(PM, BSDFContext.Context.NoV);
+
+				// The environment response of has been measured for an hemisphere having a uniform radiance of 1. 
+				// This DirectionalAlbedo also contains the integral over the hemisphere according to the phase function.
+				// The spherical harmonic used for environment lighting contains the diffuse integral over the hemisphere as DiffLambert = int_{\omega} {1 * NoL d\Omega} = PI.
+				// We do not want to be affected by the NoL term here since The DirectionalAlbedo also contains the integral over the hemisphere according to the phase function.
+				// Integral of a value=1 over the Hemisphere is int_{\omega} {1  d\Omega} = 2PI.
+				// So, as an approximation, we recover the uniform hemisphere luminance as UniformEnvL = DiffLambert * (1 / PI) * (2 * PI).
+				// We consider that the diffuse integration of the environment SH over the hemisphere multiplied with the light response will be a good approximation similar to the "split sum integral".
+				float3 SlabDirectionalAlbedo = IsotropicMediumSlabEnvDirectionalAlbedo(PM) * ((1 / PI)* (2 * PI));
+
+			#if STRATA_MICROFACET_MULTIPLESCATTERING == 0
+				// Because the default EnvBRDF does not account for multiple scattering, we rescale the participating media response according to EnvBRDF.
+				// This is to ensure a monotone light response behavior when transitioning from the Diffuse to the Volumetric model (without that, multiple scattering can result in higher luminance while trnasitionning from diffuse to volume light model)
+				const float MaxDiffuseWeight = max3(StrataEnvLightResult.DiffuseWeight.x, StrataEnvLightResult.DiffuseWeight.y, StrataEnvLightResult.DiffuseWeight.z);
+				const float MaxSlabDirectionalAlbedo = max3(SlabDirectionalAlbedo.x, SlabDirectionalAlbedo.y, SlabDirectionalAlbedo.z);
+				if (MaxDiffuseWeight < MaxSlabDirectionalAlbedo)
+				{
+					// We only want to do that when the slab light response is higher than the diffse response, so that when the mfp is large (and scattering low), the light response still work correctly according to thickness and MFP.
+					SlabDirectionalAlbedo *= MaxDiffuseWeight / MaxSlabDirectionalAlbedo;
+				}
+			#endif
 
 				// STRATA_TODO: big approximation for volumetric + env lighting: We ignore the phase for now due to the sampled convolved env lighting, 
 				// no self shadow and should not be single scattering using StrataSlabSingleScattering. We will propose a better approach. 
-				StrataEnvLightResult.DiffuseWeight = lerp(StrataEnvLightResult.DiffuseWeight, Result.SingleScattering * FresnelWeight,	DiffuseToVolumeBlend);
+				StrataEnvLightResult.DiffuseWeight = lerp(StrataEnvLightResult.DiffuseWeight, SlabDirectionalAlbedo * FresnelWeight,	DiffuseToVolumeBlend);
 				// We keep the same normal
 				StrataEnvLightResult.bSubsurface = false;
 			}
diff --git a/Engine/Shaders/Private/Strata/StrataForwardLighting.ush b/Engine/Shaders/Private/Strata/StrataForwardLighting.ush
--- a/Engine/Shaders/Private/Strata/StrataForwardLighting.ush
+++ b/Engine/Shaders/Private/Strata/StrataForwardLighting.ush
@@ -76,7 +76,7 @@ float3 StrataForwardLighting(
 				Color += GetSkyLightReflection(StrataEnvLight.SpecularDirection, StrataEnvLight.SpecularSafeRoughness, SkyAverageBrightness) * StrataEnvLight.SpecularWeight * SkyLightCommonFactor;
 			}
 
-			// Evalaute the lighting coming from L (it will lso contain the BSDF throughput to the next layer)
+			// Evaluate the lighting coming from L (it will also contain the BSDF throughput to the next layer)
 			FStrataEvaluateResult StrataEvaluate = StrataEvaluateBSDF(StrataBSDFContext);
 			Color += (StrataEvaluate.DiffusePathValue + StrataEvaluate.SpecularPathValue) * DirLightGeometryFactor 
 					+ StrataEvaluate.EmissivePathValue * OutThroughput;
