Workaround for hlslcc crash where it can't figure out that array index is a constant
For repro QAGame -game -featureleveles31 -opengl

diff --git a/Engine/Shaders/Private/PostProcessCompositeEditorPrimitives.usf b/Engine/Shaders/Private/PostProcessCompositeEditorPrimitives.usf
--- a/Engine/Shaders/Private/PostProcessCompositeEditorPrimitives.usf
+++ b/Engine/Shaders/Private/PostProcessCompositeEditorPrimitives.usf
@@ -200,7 +200,8 @@ void MainTemporalUpsampleEditorDepthPS(
 			
 			#if CONFIG_SAMPLES == 9
 			{
-				PixelOffset = kOffsets3x3[kSquareIndexes3x3[SampleId]];
+				const uint SampleIdx = kSquareIndexes3x3[SampleId];
+				PixelOffset = kOffsets3x3[SampleIdx];
 			}
 			#elif CONFIG_SAMPLES == 5 || CONFIG_SAMPLES == 6
 			{
@@ -210,7 +211,8 @@ void MainTemporalUpsampleEditorDepthPS(
 				}
 				else
 				{
-					PixelOffset = kOffsets3x3[kPlusIndexes3x3[SampleId]];
+					const uint SampleIdx = kPlusIndexes3x3[SampleId];
+					PixelOffset = kOffsets3x3[SampleIdx];
 				}
 			}
 			#else
