Workaround another hlslcc issue with array initialization

diff --git a/Engine/Shaders/Private/PostProcessCompositeEditorPrimitives.usf b/Engine/Shaders/Private/PostProcessCompositeEditorPrimitives.usf
--- a/Engine/Shaders/Private/PostProcessCompositeEditorPrimitives.usf
+++ b/Engine/Shaders/Private/PostProcessCompositeEditorPrimitives.usf
@@ -63,10 +63,10 @@ static const int2 Offsets2x2[4] =
 };
 
 // Indexes of the 3x3 square.
-static const uint kSquareIndexes3x3[9] = { 4, 0, 1, 2, 3, 5, 6, 7, 8 };
+static const uint kSquareIndexes3x3[9] = { 4u, 0u, 1u, 2u, 3u, 5u, 6u, 7u, 8u };
 
 // Indexes of the offsets to have plus + shape.
-static const uint kPlusIndexes3x3[5] = { 4, 1, 3, 5, 7 };
+static const uint kPlusIndexes3x3[5] = { 4u, 1u, 3u, 5u, 7u };
 
 #if MSAA_SAMPLE_COUNT > 1
 Texture2DMS<float4, MSAA_SAMPLE_COUNT> EditorPrimitivesColor;
