Skip sky materials for indirect rays in path tracer

These contributions can be double counted with the captured skylight, so should be skipped to avoid extra noise.

Refactored raytracing payload flags into input vs output flags to make it clear we still have room to expand a bit more.

diff --git a/Engine/Shaders/Private/PathTracing/PathTracing.usf b/Engine/Shaders/Private/PathTracing/PathTracing.usf
--- a/Engine/Shaders/Private/PathTracing/PathTracing.usf
+++ b/Engine/Shaders/Private/PathTracing/PathTracing.usf
@@ -54,6 +54,13 @@ FMaterialClosestHitPayload TraceTransparentRay(RayDesc Ray, FRayCone RayCone, bo
 		// Set payload inputs
 		PackedPayload.SetRayCone(RayCone);
 		PackedPayload.SetPixelCoord(LaunchIndex);
+		if (!IsCameraRay)
+		{
+			// avoid double counting what was captured by the skylight
+			// also avoid noise from hot spots (they can be properly
+			// importance sampled if a capturing skylight is added)
+			PackedPayload.SetIgnoreSkyMaterials();
+		}
 		// Trace the ray
 		TraceRay(
 			TLAS,
diff --git a/Engine/Shaders/Private/RayTracing/RayTracingCommon.ush b/Engine/Shaders/Private/RayTracing/RayTracingCommon.ush
--- a/Engine/Shaders/Private/RayTracing/RayTracingCommon.ush
+++ b/Engine/Shaders/Private/RayTracing/RayTracingCommon.ush
@@ -254,12 +254,18 @@ uint2 GetPixelCoord(uint2 DispatchThreadId, uint UpscaleFactor)
 #define RAY_TRACING_BLEND_MODE_ALPHA_HOLDOUT		5
 
 // Check how Flags are packed in the FPackedMaterialClosestHitPayload before adding to ensure bits are free
-#define RAY_TRACING_PAYLOAD_FLAG_FRONT_FACE                    (1 << 0) // Indicates that ray has hit front face of a primitive. Set by closest hit shader.
-#define RAY_TRACING_PAYLOAD_FLAG_MINIMAL_PAYLOAD               (1 << 1) // Indicates that closest hit shader should only fill FMinimalPayload (HitT), skipping full material evaluation. Set by RayGen shader, before TraceRay().
-#define RAY_TRACING_PAYLOAD_FLAG_ENABLE_SKY_LIGHT_CONTRIBUTION (1 << 2) // Indicates that hit shaders should add in the Sky Light contribution in their output. Set by RayGen shader, before TraceRay().
-#define RAY_TRACING_PAYLOAD_FLAG_LUMEN_PAYLOAD				   (1 << 3) // Indicates that closest hit shader should evaluate geometry normal and instanceid.
-#define RAY_TRACING_PAYLOAD_FLAG_IGNORE_TRANSLUCENT            (1 << 4) // Indicates that translucent materials should be ignored by IgnoreHit() in anyhit shaders.
-#define RAY_TRACING_PAYLOAD_FLAG_PATHTRACER_VISIBILITY         (1 << 5) // Indicates that this is a path tracer visibility ray (which supports material evaluation for transparent shadows and fake caustics)
+// There are two sets of flags:
+//   - "inputs" are set before calling TraceRay
+//   - "outputs" are set during execution of the CHS and are valid after TraceRay returns
+// These two sets of masks can overlap because they never co-exist
+#define RAY_TRACING_PAYLOAD_OUTPUT_FLAG_FRONT_FACE                   (1 << 0) // Indicates that ray has hit front face of a primitive. Set by closest hit shader.
+
+#define RAY_TRACING_PAYLOAD_INPUT_FLAG_IGNORE_SKY_MATERIALS          (1 << 0) // Indicates that the AHS/CHS should skip processing if the material is tagged as sky
+#define RAY_TRACING_PAYLOAD_INPUT_FLAG_MINIMAL_PAYLOAD               (1 << 1) // Indicates that closest hit shader should only fill FMinimalPayload (HitT), skipping full material evaluation. Set by RayGen shader, before TraceRay().
+#define RAY_TRACING_PAYLOAD_INPUT_FLAG_ENABLE_SKY_LIGHT_CONTRIBUTION (1 << 2) // Indicates that hit shaders should add in the Sky Light contribution in their output. Set by RayGen shader, before TraceRay().
+#define RAY_TRACING_PAYLOAD_INPUT_FLAG_LUMEN_PAYLOAD		         (1 << 3) // Indicates that closest hit shader should evaluate geometry normal and instanceid.
+#define RAY_TRACING_PAYLOAD_INPUT_FLAG_IGNORE_TRANSLUCENT            (1 << 4) // Indicates that translucent materials should be ignored by IgnoreHit() in anyhit shaders.
+#define RAY_TRACING_PAYLOAD_INPUT_FLAG_PATHTRACER_VISIBILITY         (1 << 5) // Indicates that this is a path tracer visibility ray (which supports material evaluation for transparent shadows and fake caustics)
 
 struct FMaterialClosestHitPayload : FMinimalPayload
 {
@@ -293,20 +299,8 @@ struct FMaterialClosestHitPayload : FMinimalPayload
 	float Anisotropy;					// 164       2       16bits (packed with WorldTangent)
 										// 166 total
 
-	void SetMinimalPayloadMode() { Flags |= RAY_TRACING_PAYLOAD_FLAG_MINIMAL_PAYLOAD; }
-	bool IsMinimalPayloadMode() { return (Flags & RAY_TRACING_PAYLOAD_FLAG_MINIMAL_PAYLOAD) != 0; }
-
-	void SetFrontFace() { Flags |= RAY_TRACING_PAYLOAD_FLAG_FRONT_FACE; }
-	bool IsFrontFace() { return (Flags & RAY_TRACING_PAYLOAD_FLAG_FRONT_FACE) != 0; }
-
-	void SetEnableSkyLightContribution() { Flags |= RAY_TRACING_PAYLOAD_FLAG_ENABLE_SKY_LIGHT_CONTRIBUTION; }
-	bool IsEnableSkyLightContribution() { return (Flags & RAY_TRACING_PAYLOAD_FLAG_ENABLE_SKY_LIGHT_CONTRIBUTION) != 0; }
-
-	void SetLumenPayload() { Flags |= RAY_TRACING_PAYLOAD_FLAG_LUMEN_PAYLOAD; }
-	bool IsLumenPayload() { return (Flags & RAY_TRACING_PAYLOAD_FLAG_LUMEN_PAYLOAD) != 0; }
-
-	void SetIgnoreTranslucentMaterials() { Flags |= RAY_TRACING_PAYLOAD_FLAG_IGNORE_TRANSLUCENT; }
-	bool IsIgnoreTranslucentMaterials() { return (Flags & RAY_TRACING_PAYLOAD_FLAG_IGNORE_TRANSLUCENT) != 0; }
+	void SetFrontFace() { Flags |= RAY_TRACING_PAYLOAD_OUTPUT_FLAG_FRONT_FACE; }
+	bool IsFrontFace() { return (Flags & RAY_TRACING_PAYLOAD_OUTPUT_FLAG_FRONT_FACE) != 0; }
 
 	FRayCone GetRayCone() { return RayCone; }
 	void SetRayCone(FRayCone NewRayCone) { RayCone = NewRayCone; }
@@ -387,20 +381,23 @@ struct FPackedMaterialClosestHitPayload : FMinimalPayload
 	void SetPixelCoord(uint2 PixelCoord) { PackedPixelCoord = (PixelCoord.x & 0xFFFF) | (PixelCoord.y << 16); }
 	uint2 GetPixelCoord() { return uint2(PackedPixelCoord & 0xFFFF, PackedPixelCoord >> 16); }
 
-	void SetMinimalPayloadMode() { IorAndShadingModelIDAndBlendingModeAndFlagsAndPrimitiveLightingChannelMask |= RAY_TRACING_PAYLOAD_FLAG_MINIMAL_PAYLOAD << 23; }
-	bool IsMinimalPayloadMode() { return ((IorAndShadingModelIDAndBlendingModeAndFlagsAndPrimitiveLightingChannelMask >> 23) & RAY_TRACING_PAYLOAD_FLAG_MINIMAL_PAYLOAD) != 0; }
+	void SetIgnoreSkyMaterials() { IorAndShadingModelIDAndBlendingModeAndFlagsAndPrimitiveLightingChannelMask |= RAY_TRACING_PAYLOAD_INPUT_FLAG_IGNORE_SKY_MATERIALS << 23; }
+	bool IsIgnoreSkyMaterials() { return ((IorAndShadingModelIDAndBlendingModeAndFlagsAndPrimitiveLightingChannelMask >> 23) & RAY_TRACING_PAYLOAD_INPUT_FLAG_IGNORE_SKY_MATERIALS) != 0; }
+
+	void SetMinimalPayloadMode() { IorAndShadingModelIDAndBlendingModeAndFlagsAndPrimitiveLightingChannelMask |= RAY_TRACING_PAYLOAD_INPUT_FLAG_MINIMAL_PAYLOAD << 23; }
+	bool IsMinimalPayloadMode() { return ((IorAndShadingModelIDAndBlendingModeAndFlagsAndPrimitiveLightingChannelMask >> 23) & RAY_TRACING_PAYLOAD_INPUT_FLAG_MINIMAL_PAYLOAD) != 0; }
 
-	void SetEnableSkyLightContribution() { IorAndShadingModelIDAndBlendingModeAndFlagsAndPrimitiveLightingChannelMask |= RAY_TRACING_PAYLOAD_FLAG_ENABLE_SKY_LIGHT_CONTRIBUTION << 23; }
-	bool IsEnableSkyLightContribution() { return ((IorAndShadingModelIDAndBlendingModeAndFlagsAndPrimitiveLightingChannelMask >> 23) & RAY_TRACING_PAYLOAD_FLAG_ENABLE_SKY_LIGHT_CONTRIBUTION) != 0; }
+	void SetEnableSkyLightContribution() { IorAndShadingModelIDAndBlendingModeAndFlagsAndPrimitiveLightingChannelMask |= RAY_TRACING_PAYLOAD_INPUT_FLAG_ENABLE_SKY_LIGHT_CONTRIBUTION << 23; }
+	bool IsEnableSkyLightContribution() { return ((IorAndShadingModelIDAndBlendingModeAndFlagsAndPrimitiveLightingChannelMask >> 23) & RAY_TRACING_PAYLOAD_INPUT_FLAG_ENABLE_SKY_LIGHT_CONTRIBUTION) != 0; }
 
-	void SetLumenPayload() { IorAndShadingModelIDAndBlendingModeAndFlagsAndPrimitiveLightingChannelMask |= RAY_TRACING_PAYLOAD_FLAG_LUMEN_PAYLOAD << 23; }
-	bool IsLumenPayload() { return  ((IorAndShadingModelIDAndBlendingModeAndFlagsAndPrimitiveLightingChannelMask >> 23) & RAY_TRACING_PAYLOAD_FLAG_LUMEN_PAYLOAD) != 0;}
+	void SetLumenPayload() { IorAndShadingModelIDAndBlendingModeAndFlagsAndPrimitiveLightingChannelMask |= RAY_TRACING_PAYLOAD_INPUT_FLAG_LUMEN_PAYLOAD << 23; }
+	bool IsLumenPayload() { return  ((IorAndShadingModelIDAndBlendingModeAndFlagsAndPrimitiveLightingChannelMask >> 23) & RAY_TRACING_PAYLOAD_INPUT_FLAG_LUMEN_PAYLOAD) != 0;}
 
-	void SetIgnoreTranslucentMaterials() { IorAndShadingModelIDAndBlendingModeAndFlagsAndPrimitiveLightingChannelMask |= RAY_TRACING_PAYLOAD_FLAG_IGNORE_TRANSLUCENT << 23; }
-	bool IsIgnoreTranslucentMaterials() { return ((IorAndShadingModelIDAndBlendingModeAndFlagsAndPrimitiveLightingChannelMask >> 23) & RAY_TRACING_PAYLOAD_FLAG_IGNORE_TRANSLUCENT) != 0; }
+	void SetIgnoreTranslucentMaterials() { IorAndShadingModelIDAndBlendingModeAndFlagsAndPrimitiveLightingChannelMask |= RAY_TRACING_PAYLOAD_INPUT_FLAG_IGNORE_TRANSLUCENT << 23; }
+	bool IsIgnoreTranslucentMaterials() { return ((IorAndShadingModelIDAndBlendingModeAndFlagsAndPrimitiveLightingChannelMask >> 23) & RAY_TRACING_PAYLOAD_INPUT_FLAG_IGNORE_TRANSLUCENT) != 0; }
 
-	void SetPathTracingVisibilityRay() { IorAndShadingModelIDAndBlendingModeAndFlagsAndPrimitiveLightingChannelMask |= RAY_TRACING_PAYLOAD_FLAG_PATHTRACER_VISIBILITY << 23; }
-	bool IsPathTracingVisibilityRay() { return ((IorAndShadingModelIDAndBlendingModeAndFlagsAndPrimitiveLightingChannelMask >> 23) & RAY_TRACING_PAYLOAD_FLAG_PATHTRACER_VISIBILITY) != 0; }
+	void SetPathTracingVisibilityRay() { IorAndShadingModelIDAndBlendingModeAndFlagsAndPrimitiveLightingChannelMask |= RAY_TRACING_PAYLOAD_INPUT_FLAG_PATHTRACER_VISIBILITY << 23; }
+	bool IsPathTracingVisibilityRay() { return ((IorAndShadingModelIDAndBlendingModeAndFlagsAndPrimitiveLightingChannelMask >> 23) & RAY_TRACING_PAYLOAD_INPUT_FLAG_PATHTRACER_VISIBILITY) != 0; }
 
 	FRayCone GetRayCone() { return UnpackRayCone(PackedRayCone); }
 	void SetRayCone(FRayCone NewRayCone) { PackedRayCone = PackRayCone(NewRayCone); }
diff --git a/Engine/Shaders/Private/RayTracing/RayTracingMaterialHitShaders.usf b/Engine/Shaders/Private/RayTracing/RayTracingMaterialHitShaders.usf
--- a/Engine/Shaders/Private/RayTracing/RayTracingMaterialHitShaders.usf
+++ b/Engine/Shaders/Private/RayTracing/RayTracingMaterialHitShaders.usf
@@ -374,6 +374,14 @@ RAY_TRACING_ENTRY_CLOSEST_HIT(MaterialCHS,
 		Payload.SetFrontFace();
 	}
 
+#if MATERIAL_IS_SKY
+	if (PackedPayload.IsIgnoreSkyMaterials())
+	{
+		PackedPayload = PackRayTracingPayload(Payload, PropagatedCone);
+		return;
+	}
+#endif
+
 	// Store the results in local variables and reuse instead of calling the functions multiple times.
 	half3 BaseColor = GetMaterialBaseColor(PixelMaterialInputs);
 	half2 Refraction = GetMaterialRefraction(PixelMaterialInputs);
