* Fix warning in PageManagment.usf
* Remove GBuffer data from the FScreenPixelData struct as it is no longer used (and initialized

diff --git a/Engine/Shaders/Private/VirtualShadowMaps/PageManagement.usf b/Engine/Shaders/Private/VirtualShadowMaps/PageManagement.usf
--- a/Engine/Shaders/Private/VirtualShadowMaps/PageManagement.usf
+++ b/Engine/Shaders/Private/VirtualShadowMaps/PageManagement.usf
@@ -50,9 +50,6 @@ struct FScreenPixelData
 	
 	// Note: Actual world position (not pre-translated)
 	float3 WorldPosition;
-
-	// NOTE: This data is only valid when executing after the base pass (bPostBasePass is true)
-	FGBufferData GBufferData;
 };
 
 float2 CalcScreenUV(uint2 PixelPos)
@@ -1133,7 +1130,7 @@ void AllocatePagesUsingRects(uint GroupThreadId : SV_GroupThreadID, uint2 Dispat
 			bool bFoundFree = false;
 			
 			// Until we run off the edge 
-			while (ProbeRect.z < VirtualShadowMap.PhysicalPoolSizePages.x)
+			while (ProbeRect.z < uint(VirtualShadowMap.PhysicalPoolSizePages.x))
 			{
 				NextFreeOffsetX = 0;
 				GroupMemoryBarrierWithGroupSync();
