Fix missing propagation of PerInstanceParameters & lightmap coordinates for GPU-Scene Instance culling in LocalVF.

diff --git a/Engine/Shaders/Private/LocalVertexFactory.ush b/Engine/Shaders/Private/LocalVertexFactory.ush
--- a/Engine/Shaders/Private/LocalVertexFactory.ush
+++ b/Engine/Shaders/Private/LocalVertexFactory.ush
@@ -349,8 +349,6 @@ struct FVertexFactoryIntermediates
 #endif
 
 #if USE_INSTANCING
-	float IsInstanceSelected;
-
 	float4 InstanceOrigin;
 	float4 InstanceTransform1;
 	float4 InstanceTransform2;
@@ -455,6 +453,24 @@ float3 GetInstanceOrigin(FVertexFactoryIntermediates Intermediates)
 
 #endif	// USE_INSTANCING
 
+#if USE_INSTANCE_CULLING
+
+float2 GetInstanceShadowMapBias(FVertexFactoryIntermediates Intermediates)
+{
+	// GPUCULL_TODO: This data is not yet pushed to GPU-Scene for non-nanite instances.
+	//return Intermediates.SceneData.InstanceData.LightMapAndShadowMapUVBias.zw;
+	return Intermediates.InstanceLightmapAndShadowMapUVBias.zw;
+}
+
+float2 GetInstanceLightMapBias(FVertexFactoryIntermediates Intermediates)
+{
+	// GPUCULL_TODO: This data is not yet pushed to GPU-Scene for non-nanite instances.
+	//return Intermediates.SceneData.InstanceData.LightMapAndShadowMapUVBias.xy;
+	return Intermediates.InstanceLightmapAndShadowMapUVBias.xy;
+}
+
+#endif	// USE_INSTANCE_CULLING
+
 /** Converts from vertex factory specific interpolants to a FMaterialPixelParameters, which is used by material inputs. */
 FMaterialPixelParameters GetMaterialPixelParameters(FVertexFactoryInterpolantsVSToPS Interpolants, float4 SvPosition)
 {
@@ -1028,7 +1044,7 @@ FVertexFactoryInterpolantsVSToPS VertexFactoryGetInterpolantsVSToPS(FVertexFacto
 
 	float4 LightMapCoordinateScaleBias = GetLightmapData(LightmapDataIndex).LightMapCoordinateScaleBias;
 
-	#if USE_INSTANCING
+	#if USE_INSTANCING || USE_INSTANCE_CULLING
 		LightMapCoordinate = LightMapCoordinateInput * LightMapCoordinateScaleBias.xy + GetInstanceLightMapBias(Intermediates);
 	#else
 		LightMapCoordinate = LightMapCoordinateInput * LightMapCoordinateScaleBias.xy + LightMapCoordinateScaleBias.zw;
@@ -1036,7 +1052,7 @@ FVertexFactoryInterpolantsVSToPS VertexFactoryGetInterpolantsVSToPS(FVertexFacto
 	#if STATICLIGHTING_TEXTUREMASK
 		float4 ShadowMapCoordinateScaleBias = GetLightmapData(LightmapDataIndex).ShadowMapCoordinateScaleBias;
 
-		#if USE_INSTANCING
+		#if USE_INSTANCING || USE_INSTANCE_CULLING
 			ShadowMapCoordinate = LightMapCoordinateInput * ShadowMapCoordinateScaleBias.xy + GetInstanceShadowMapBias(Intermediates);
 		#else
 			ShadowMapCoordinate = LightMapCoordinateInput * ShadowMapCoordinateScaleBias.xy + ShadowMapCoordinateScaleBias.zw;
@@ -1049,7 +1065,8 @@ FVertexFactoryInterpolantsVSToPS VertexFactoryGetInterpolantsVSToPS(FVertexFacto
 
 	SetTangents(Interpolants, Intermediates.TangentToWorld[0], Intermediates.TangentToWorld[2], Intermediates.TangentToWorldSign);
 	SetColor(Interpolants, Intermediates.Color);
-#if USE_INSTANCING
+
+#if USE_INSTANCING || USE_INSTANCE_CULLING
 	Interpolants.PerInstanceParams = Intermediates.PerInstanceParams;
 #endif
 
@@ -1300,7 +1317,7 @@ float4 VertexFactoryGetPreviousWorldPosition(FVertexFactoryInput Input, FVertexF
 	#if INTERPOLATE_VERTEX_COLOR
 		TESSELLATION_INTERPOLATE_MEMBER(InterpolantsVSToPS.Color);
 	#endif
-	#if USE_INSTANCING
+	#if USE_INSTANCING || USE_INSTANCE_CULLING
 		TESSELLATION_INTERPOLATE_MEMBER(InterpolantsVSToPS.PerInstanceParams);
 	#endif
 
