Fix page fault with hair cards rendering on dx12.

This happens with manual vertex fetches, where vertex data are fetched beyond the valid boundary. This seems to happen during LOD transitions. The current CL fixes the issue by preventing the out-of-bound issue, but not the actual root cause of the issue.

diff --git a/Engine/Plugins/Runtime/HairStrands/Source/HairStrandsCore/Private/HairCardsVertexFactory.cpp b/Engine/Plugins/Runtime/HairStrands/Source/HairStrandsCore/Private/HairCardsVertexFactory.cpp
--- a/Engine/Plugins/Runtime/HairStrands/Source/HairStrandsCore/Private/HairCardsVertexFactory.cpp
+++ b/Engine/Plugins/Runtime/HairStrands/Source/HairStrandsCore/Private/HairCardsVertexFactory.cpp
@@ -39,6 +39,7 @@ FHairCardsUniformBuffer CreateHairCardsVFUniformBuffer(
 
 		// Cards atlas UV are inverted so fetching needs to be inverted on the y-axis
 		UniformParameters.bInvertUV = LOD.RestResource->bInvertUV;
+		UniformParameters.MaxVertexCount = LOD.RestResource->VertexCount;
 
 		// When the geometry is not-dynamic (no binding to skeletal mesh, no simulation), only a single vertex buffer is allocated. 
 		// In this case we force the buffer index to 0
@@ -88,6 +89,7 @@ FHairCardsUniformBuffer CreateHairCardsVFUniformBuffer(
 
 		// Meshes UV are not inverted so no need to invert the y-axis
 		UniformParameters.bInvertUV = 0;
+		UniformParameters.MaxVertexCount = LOD.RestResource->VertexCount;
 		UniformParameters.NormalsBuffer = LOD.RestResource->NormalsBuffer.ShaderResourceViewRHI.GetReference();
 		UniformParameters.UVsBuffer = LOD.RestResource->UVsBuffer.ShaderResourceViewRHI.GetReference();
 
diff --git a/Engine/Plugins/Runtime/HairStrands/Source/HairStrandsCore/Public/GroomInstance.h b/Engine/Plugins/Runtime/HairStrands/Source/HairStrandsCore/Public/GroomInstance.h
--- a/Engine/Plugins/Runtime/HairStrands/Source/HairStrandsCore/Public/GroomInstance.h
+++ b/Engine/Plugins/Runtime/HairStrands/Source/HairStrandsCore/Public/GroomInstance.h
@@ -14,6 +14,7 @@
 //  * 12 bits for cards count -> 4000 cards for a hair group
 BEGIN_GLOBAL_SHADER_PARAMETER_STRUCT(FHairCardsVertexFactoryUniformShaderParameters, HAIRSTRANDSCORE_API)
 	SHADER_PARAMETER(uint32, bInvertUV)
+	SHADER_PARAMETER(uint32, MaxVertexCount)
 	SHADER_PARAMETER_SRV(Buffer<float4>, PositionBuffer)
 	SHADER_PARAMETER_SRV(Buffer<float4>, PreviousPositionBuffer)
 	SHADER_PARAMETER_SRV(Buffer<float4>, NormalsBuffer)
diff --git a/Engine/Shaders/Private/HairStrands/HairCardsVertexFactory.ush b/Engine/Shaders/Private/HairStrands/HairCardsVertexFactory.ush
--- a/Engine/Shaders/Private/HairStrands/HairCardsVertexFactory.ush
+++ b/Engine/Shaders/Private/HairStrands/HairCardsVertexFactory.ush
@@ -101,10 +101,15 @@ uint GetPrimitiveId(FVertexFactoryInput Input)
 #endif
 }
 
+uint GetHairVertexFetchId(FVertexFactoryInput Input)
+{
+	return min(Input.VertexId, HairCardsVF.MaxVertexCount - 1u);
+}
+
 uint GetHairPrimitiveId(FVertexFactoryInput Input)
 {
 #if MANUAL_VERTEX_FETCH
-	return Input.VertexId;
+	return GetHairVertexFetchId(Input);
 #else
 	return 0;
 #endif
@@ -120,8 +125,9 @@ struct FHairCardsPoint
 FHairCardsPoint GetVertexPosition(FVertexFactoryInput Input)
 {
 #if MANUAL_VERTEX_FETCH
-	const float4 Data = HairCardsVF.PositionBuffer[Input.VertexId];
-	const float4 UVs  = HairCardsVF.UVsBuffer[Input.VertexId];
+	const uint VertexFetchId = GetHairVertexFetchId(Input);
+	const float4 Data = HairCardsVF.PositionBuffer[VertexFetchId];
+	const float4 UVs  = HairCardsVF.UVsBuffer[VertexFetchId];
 #else
 	const float4 Data = Input.Position;
 	const float4 UVs  = Input.UVs;
@@ -137,7 +143,8 @@ FHairCardsPoint GetVertexPosition(FVertexFactoryInput Input)
 float3 GetVertexPreviousPosition(FVertexFactoryInput Input)
 {
 #if MANUAL_VERTEX_FETCH
-	const float4 Data = HairCardsVF.PreviousPositionBuffer[Input.VertexId];
+	const uint VertexFetchId = GetHairVertexFetchId(Input);
+	const float4 Data = HairCardsVF.PreviousPositionBuffer[VertexFetchId];
 #else
 	const float4 Data = Input.Position;
 #endif
@@ -221,8 +228,9 @@ float4 CalcWorldPosition(float3 Position, uint PrimitiveId)
 half3x3 CalcTangentToLocal(FVertexFactoryInput Input, inout float TangentSign)
 {
 #if MANUAL_VERTEX_FETCH
-	const half3 TangentInputX = HairCardsVF.NormalsBuffer[Input.VertexId * 2 + 0].xyz;
-	const half4 TangentInputZ = HairCardsVF.NormalsBuffer[Input.VertexId * 2 + 1];
+	const uint VertexFetchId = GetHairVertexFetchId(Input);
+	const half3 TangentInputX = HairCardsVF.NormalsBuffer[VertexFetchId * 2 + 0].xyz;
+	const half4 TangentInputZ = HairCardsVF.NormalsBuffer[VertexFetchId * 2 + 1];
 #else
 	const half3 TangentInputX = Input.TangentX.xyz;
 	const half4 TangentInputZ = Input.TangentZ;
@@ -331,7 +339,8 @@ FVertexFactoryInterpolantsVSToPS VertexFactoryGetInterpolantsVSToPS(FVertexFacto
 #endif
 
 #if MANUAL_VERTEX_FETCH
-	Interpolants.HairPrimitiveId = Input.VertexId;
+	const uint VertexFetchId = GetHairVertexFetchId(Input);
+	Interpolants.HairPrimitiveId = VertexFetchId;
 	Interpolants.HairPrimitiveUV = GetVertexPosition(Input).CardsUVs;
 #else
 	Interpolants.HairPrimitiveId = 0;
