Fix hair shader compilation for ES31

diff --git a/Engine/Shaders/Private/HairStrands/HairCardsVertexFactory.ush b/Engine/Shaders/Private/HairStrands/HairCardsVertexFactory.ush
--- a/Engine/Shaders/Private/HairStrands/HairCardsVertexFactory.ush
+++ b/Engine/Shaders/Private/HairStrands/HairCardsVertexFactory.ush
@@ -103,7 +103,11 @@ uint GetPrimitiveId(FVertexFactoryInput Input)
 
 uint GetHairVertexFetchId(FVertexFactoryInput Input)
 {
+#if MANUAL_VERTEX_FETCH
 	return min(Input.VertexId, HairCardsVF.MaxVertexCount - 1u);
+#else
+	return 0u;
+#endif
 }
 
 uint GetHairPrimitiveId(FVertexFactoryInput Input)
