Rejects Gen5TAA history's high frequency details if theirs low frequencies drift from history's low frequency

diff --git a/Engine/Shaders/Private/TemporalAA/TAAUpdateHistory.usf b/Engine/Shaders/Private/TemporalAA/TAAUpdateHistory.usf
--- a/Engine/Shaders/Private/TemporalAA/TAAUpdateHistory.usf
+++ b/Engine/Shaders/Private/TemporalAA/TAAUpdateHistory.usf
@@ -18,6 +18,8 @@
 
 #define CONFIG_MANUAL_LDS_SPILL (CONFIG_COMPILE_FP16)
 
+#define CONFIG_LOW_FREQUENCY_DRIFT_REJECTION 1
+
 
 //------------------------------------------------------- CONSTANTS
 
@@ -250,6 +252,9 @@ void MainCS(
 	// Reproject history
 	taa_half3 PrevHistoryMoment1;
 	taa_half PrevHistoryValidity;
+	
+	taa_half3 PrevHistoryMommentMin;
+	taa_half3 PrevHistoryMommentMax;
 
 	taa_half3 PrevFallbackColor;
 	taa_half PrevFallbackWeight;
@@ -263,6 +268,9 @@ void MainCS(
 		taa_half3 RawHistory1 = taa_half(0);
 		taa_half2 RawHistory2 = taa_half(0);
 
+		taa_half3 RawHistory1Min = INFINITE_FLOAT;
+		taa_half3 RawHistory1Max = -INFINITE_FLOAT;
+
 		// Sample raw history
 		{
 			float2 PrevHistoryBufferUV = (PrevHistoryInfo_ScreenPosToViewportScale * PrevScreenPos + PrevHistoryInfo_ScreenPosToViewportBias) * PrevHistoryInfo_ExtentInverse;
@@ -276,9 +284,16 @@ void MainCS(
 				{
 					float2 SampleUV = clamp(Samples.UV[i], PrevHistoryInfo_UVViewportBilinearMin, PrevHistoryInfo_UVViewportBilinearMax);
 
-					RawHistory0 += PrevHistory_Textures_0.SampleLevel(GlobalBilinearClampedSampler, SampleUV, 0) * taa_half(Samples.Weight[i]);
-					RawHistory1 += PrevHistory_Textures_1.SampleLevel(GlobalBilinearClampedSampler, SampleUV, 0) * taa_half(Samples.Weight[i]);
-					RawHistory2 += PrevHistory_Textures_2.SampleLevel(GlobalBilinearClampedSampler, SampleUV, 0) * taa_half(Samples.Weight[i]);
+					taa_half3 Sample0 = PrevHistory_Textures_0.SampleLevel(GlobalBilinearClampedSampler, SampleUV, 0);
+					taa_half3 Sample1 = PrevHistory_Textures_1.SampleLevel(GlobalBilinearClampedSampler, SampleUV, 0);
+					taa_half2 Sample2 = PrevHistory_Textures_2.SampleLevel(GlobalBilinearClampedSampler, SampleUV, 0);
+
+					RawHistory1Min = min(RawHistory1Min, TransformColorForClampingBox(Sample1 * SafeRcp(Sample2.g)));
+					RawHistory1Max = max(RawHistory1Max, TransformColorForClampingBox(Sample1 * SafeRcp(Sample2.g)));
+
+					RawHistory0 += Sample0 * taa_half(Samples.Weight[i]);
+					RawHistory1 += Sample1 * taa_half(Samples.Weight[i]);
+					RawHistory2 += Sample2 * taa_half(Samples.Weight[i]);
 				}
 				RawHistory0 *= taa_half(Samples.FinalMultiplier);
 				RawHistory1 *= taa_half(Samples.FinalMultiplier);
@@ -312,6 +327,9 @@ void MainCS(
 		{
 			PrevFallbackColor = RawHistory0;
 			PrevFallbackWeight = RawHistory2.r;
+			
+			PrevHistoryMommentMin = RawHistory1Min;
+			PrevHistoryMommentMax = RawHistory1Max;
 
 			PrevHistoryMoment1 = RawHistory1;
 			PrevHistoryValidity = RawHistory2.g;
@@ -319,6 +337,8 @@ void MainCS(
 
 		// Correct history
 		{
+			PrevHistoryMommentMin *= taa_half(HistoryPreExposureCorrection);
+			PrevHistoryMommentMax *= taa_half(HistoryPreExposureCorrection);
 			PrevHistoryMoment1 *= taa_half(HistoryPreExposureCorrection);
 			PrevFallbackColor *= taa_half(HistoryPreExposureCorrection);
 		}
@@ -343,6 +363,22 @@ void MainCS(
 	}
 	#endif
 
+	// Rejects high frequency details if their low frequencies drift from the low frequency history
+	#if CONFIG_LOW_FREQUENCY_DRIFT_REJECTION
+	{
+		taa_half3 PrevHighFrequencyYCoCg = TransformColorForClampingBox(PrevHistoryMoment1 * SafeRcp(PrevHistoryValidity));
+		taa_half3 PrevYCoCg = TransformColorForClampingBox(PrevFallbackColor);
+		taa_half3 ClampedPrevYCoCg = clamp(PrevYCoCg, PrevHistoryMommentMin, PrevHistoryMommentMax);
+
+		taa_half HighFrequencyRejection = MeasureRejectionFactor(
+			PrevYCoCg, ClampedPrevYCoCg,
+			PrevHighFrequencyYCoCg, InputMinColor, InputMaxColor);
+		
+		PrevHistoryMoment1 *= HighFrequencyRejection;
+		PrevHistoryValidity *= HighFrequencyRejection;
+	}
+	#endif
+
 	// Contribute current frame input into the predictor for next frame.
 	const taa_half Histeresis = rcp(taa_half(MAX_SAMPLE_COUNT));
 	const taa_half PredictionOnlyValidity = Histeresis * taa_half(2.0);
