Optimize the 'y' component evaluation of the Sobol' pattern

Since this dimension has a very simple Sobol matrix, we can generate it on the fly instead of loading it from the constant array.

Also add the option (enabled by default) to unroll the core Sobol generation loop. This was suggested by [REDACTED_PERSON] and seems to be faster for cases where execution if coherent. On the other hand, the path tracer seems faster overall with the "early" exit version (presumably because there is more divergence in that case).

diff --git a/Engine/Shaders/Private/PathTracing/Utilities/PathTracingRandomSequence.ush b/Engine/Shaders/Private/PathTracing/Utilities/PathTracingRandomSequence.ush
--- a/Engine/Shaders/Private/PathTracing/Utilities/PathTracingRandomSequence.ush
+++ b/Engine/Shaders/Private/PathTracing/Utilities/PathTracingRandomSequence.ush
@@ -25,6 +25,10 @@
 #define RANDSEQ_ERROR_DIFFUSION 1 // Enabled by default, only works when using RANDSEQ_OWENSOBOL
 #endif
 
+#ifndef RANDSEQ_UNROLL_SOBOL
+#define RANDSEQ_UNROLL_SOBOL 1 // Use unrolled Sobol loop such that table can be inline constants.
+#endif
+
 // This hash mixes bits at the theoretical bias limit
 uint PerfectIntegerHash(uint x)
 {
@@ -206,41 +210,41 @@ uint FastOwenScrambling(uint Index, uint Seed) {
 // S. Joe and F. Y. Kuo, Constructing Sobol sequences with better two-dimensional projections, SIAM J. Sci. Comput. 30, 2635-2654 (2008)
 //    https://web.maths.unsw.edu.au/~fkuo/sobol/
 // NOTE: we don't bother storing dimension 0 since it is just a bit reversal
-// NOTE2: dimension 1 also follows a pattern that can be done procedurally but probably not worth the extra code?
-// NOTE3: the matrix eleements are reversed to save one reverse in the owen scrambling
-static const uint3 SobolMatrices[] = {
-	uint3(0x00000001, 0x00000001, 0x00000001),
-	uint3(0x00000003, 0x00000003, 0x00000003),
-	uint3(0x00000005, 0x00000006, 0x00000004),
-	uint3(0x0000000f, 0x00000009, 0x0000000a),
-	uint3(0x00000011, 0x00000017, 0x0000001f),
-	uint3(0x00000033, 0x0000003a, 0x0000002e),
-	uint3(0x00000055, 0x00000071, 0x00000045),
-	uint3(0x000000ff, 0x000000a3, 0x000000c9),
-	uint3(0x00000101, 0x00000116, 0x0000011b),
-	uint3(0x00000303, 0x00000339, 0x000002a4),
-	uint3(0x00000505, 0x00000677, 0x0000079a),
-	uint3(0x00000f0f, 0x000009aa, 0x00000b67),
-	uint3(0x00001111, 0x00001601, 0x0000101e),
-	uint3(0x00003333, 0x00003903, 0x0000302d),
-	uint3(0x00005555, 0x00007706, 0x00004041),
-	uint3(0x0000ffff, 0x0000aa09, 0x0000a0c3),
-	uint3(0x00010001, 0x00010117, 0x0001f104),
-	uint3(0x00030003, 0x0003033a, 0x0002e28a),
-	uint3(0x00050005, 0x00060671, 0x000457df),
-	uint3(0x000f000f, 0x000909a3, 0x000c9bae),
-	uint3(0x00110011, 0x00171616, 0x0011a105),
-	uint3(0x00330033, 0x003a3939, 0x002a7289),
-	uint3(0x00550055, 0x00717777, 0x0079e7db),
-	uint3(0x00ff00ff, 0x00a3aaaa, 0x00b6dba4),
-	uint3(0x01010101, 0x01170001, 0x0100011a),
-	uint3(0x03030303, 0x033a0003, 0x030002a7),
-	uint3(0x05050505, 0x06710006, 0x0400079e),
-	uint3(0x0f0f0f0f, 0x09a30009, 0x0a000b6d),
-	uint3(0x11111111, 0x16160017, 0x1f001001),
-	uint3(0x33333333, 0x3939003a, 0x2e003003),
-	uint3(0x55555555, 0x77770071, 0x45004004),
-	uint3(0xffffffff, 0xaaaa00a3, 0xc900a00a)
+// NOTE2: we don't bother storing dimension 1 either since it has a very simple pattern
+// NOTE3: the matrix elements are reversed to save one reverse in the owen scrambling
+static const uint2 SobolMatrices[] = {
+	uint2(0x00000001, 0x00000001),
+	uint2(0x00000003, 0x00000003),
+	uint2(0x00000006, 0x00000004),
+	uint2(0x00000009, 0x0000000a),
+	uint2(0x00000017, 0x0000001f),
+	uint2(0x0000003a, 0x0000002e),
+	uint2(0x00000071, 0x00000045),
+	uint2(0x000000a3, 0x000000c9),
+	uint2(0x00000116, 0x0000011b),
+	uint2(0x00000339, 0x000002a4),
+	uint2(0x00000677, 0x0000079a),
+	uint2(0x000009aa, 0x00000b67),
+	uint2(0x00001601, 0x0000101e),
+	uint2(0x00003903, 0x0000302d),
+	uint2(0x00007706, 0x00004041),
+	uint2(0x0000aa09, 0x0000a0c3),
+	uint2(0x00010117, 0x0001f104),
+	uint2(0x0003033a, 0x0002e28a),
+	uint2(0x00060671, 0x000457df),
+	uint2(0x000909a3, 0x000c9bae),
+	uint2(0x00171616, 0x0011a105),
+	uint2(0x003a3939, 0x002a7289),
+	uint2(0x00717777, 0x0079e7db),
+	uint2(0x00a3aaaa, 0x00b6dba4),
+	uint2(0x01170001, 0x0100011a),
+	uint2(0x033a0003, 0x030002a7),
+	uint2(0x06710006, 0x0400079e),
+	uint2(0x09a30009, 0x0a000b6d),
+	uint2(0x16160017, 0x1f001001),
+	uint2(0x3939003a, 0x2e003003),
+	uint2(0x77770071, 0x45004004),
+	uint2(0xaaaa00a3, 0xc900a00a)
 };
 
 uint EvolveSobolSeed(inout uint Seed)
@@ -270,9 +274,19 @@ float4 SobolSampler(uint SampleIndex, inout uint Seed)
 	uint SobolIndex = FastOwenScrambling(SampleIndex, EvolveSobolSeed(Seed));
 	// now get Sobol' point from this index
 	uint4 Result = uint4(SobolIndex, 0, 0, 0);
-	for (int i = 0; SobolIndex; SobolIndex >>= 1, i++)
-		Result.yzw ^= (SobolIndex & 1) * SobolMatrices[i];
-
+#if RANDSEQ_UNROLL_SOBOL // unrolled version seems faster in most cases
+	UNROLL for (uint b = 0, v = 1; b < 32; b++)
+	{
+		uint IndexBit = (SobolIndex >> b) & 1;
+#else // however a loop with early exit seems faster in the path tracing context (possibly due to low occupancy?)
+	for (uint b = 0, v = 1; SobolIndex; SobolIndex >>= 1, b++)
+	{
+		uint IndexBit = SobolIndex & 1;
+#endif
+		Result.y ^= IndexBit * v;
+		Result.zw ^= IndexBit * SobolMatrices[b];
+		v ^= v << 1;
+	}
 	// finally scramble the points to avoid structured artifacts
 	Result.x = FastOwenScrambling(Result.x, EvolveSobolSeed(Seed));
 	Result.y = FastOwenScrambling(Result.y, EvolveSobolSeed(Seed));
