Subsurface scattering support in path tracing

This adds a random walk stage to the path tracer, for use with the Subsurface (and Preintegrated Skin) shading models. For now this only does a simple analog walk as these modes only have a monochromatic radius control.

diff --git a/Engine/Shaders/Private/PathTracing/Material/PathTracingMaterialCommon.ush b/Engine/Shaders/Private/PathTracing/Material/PathTracingMaterialCommon.ush
--- a/Engine/Shaders/Private/PathTracing/Material/PathTracingMaterialCommon.ush
+++ b/Engine/Shaders/Private/PathTracing/Material/PathTracingMaterialCommon.ush
@@ -116,6 +116,11 @@ bool IsMaterialTransmissive(FMaterialClosestHitPayload HitInfo)
 		   (HitInfo.BlendingMode == RAY_TRACING_BLEND_MODE_TRANSLUCENT && HitInfo.Roughness != 0 && HitInfo.Opacity < 1);
 }
 
+bool IsSubsurfaceMaterial(FMaterialClosestHitPayload HitInfo)
+{
+	return HitInfo.ShadingModelID == SHADINGMODELID_SUBSURFACE || HitInfo.ShadingModelID == SHADINGMODELID_PREINTEGRATED_SKIN;
+}
+
 // Given two lobes that will roughly contribute colors A and B to the total (estimated for example using directional albedo)
 // return the probability of choosing lobe A
 float LobeSelectionProb(float3 A, float3 B)
diff --git a/Engine/Shaders/Private/PathTracing/Material/PathTracingMaterialSampling.ush b/Engine/Shaders/Private/PathTracing/Material/PathTracingMaterialSampling.ush
--- a/Engine/Shaders/Private/PathTracing/Material/PathTracingMaterialSampling.ush
+++ b/Engine/Shaders/Private/PathTracing/Material/PathTracingMaterialSampling.ush
@@ -47,6 +47,8 @@ bool IsNeeValidMaterial(FMaterialClosestHitPayload Payload)
 		case SHADINGMODELID_TWOSIDED_FOLIAGE:	
 		case SHADINGMODELID_THIN_TRANSLUCENT:
 		case SHADINGMODELID_CLOTH:
+		case SHADINGMODELID_SUBSURFACE:
+		case SHADINGMODELID_PREINTEGRATED_SKIN:
 		default:
 			return true;
 	}
@@ -55,17 +57,26 @@ bool IsNeeValidMaterial(FMaterialClosestHitPayload Payload)
 
 float3 GetMaterialTransparency(FMaterialClosestHitPayload Payload, float3 RayDirection)
 {
-	if (Payload.ShadingModelID == SHADINGMODELID_THIN_TRANSLUCENT)
+	switch (Payload.ShadingModelID)
 	{
-		if (Payload.Ior == 1.0)
+		case SHADINGMODELID_THIN_TRANSLUCENT:
 		{
-			float VoN = dot(RayDirection, Payload.WorldNormal);
-			return (1 - Payload.Opacity) * pow(Payload.CustomData.xyz, 1.0 / abs(VoN));
+			if (Payload.Ior == 1.0)
+			{
+				float VoN = dot(RayDirection, Payload.WorldNormal);
+				return (1 - Payload.Opacity) * pow(Payload.CustomData.xyz, 1.0 / abs(VoN));
+			}
+			else
+			{
+				return 0;
+			}
 		}
-		else
-		{
+		case SHADINGMODELID_SUBSURFACE:
+		case SHADINGMODELID_PREINTEGRATED_SKIN:
 			return 0;
-		}
+		default:
+			// other shading models support blend modes
+			break;
 	}
 	switch (Payload.BlendingMode)
 	{
@@ -138,6 +149,10 @@ float3 EstimateMaterialAlbedo(FMaterialClosestHitPayload Payload)
 		case SHADINGMODELID_CLOTH:
 			Albedo = BrdfWeight * (Payload.DiffuseColor + lerp(Payload.SpecularColor, Payload.CustomData.xyz, Payload.CustomData.w));
 			break;
+		case SHADINGMODELID_SUBSURFACE:
+		case SHADINGMODELID_PREINTEGRATED_SKIN:
+			Albedo = BrdfWeight * (Payload.DiffuseColor + Payload.SpecularColor + Payload.CustomData.xyz);
+			break;
 		default:
 			Albedo = BrdfWeight * Payload.BaseColor;
 			break;
@@ -169,6 +184,8 @@ void AdjustPathAndMaterialRoughness(inout FMaterialClosestHitPayload Payload, in
 		case SHADINGMODELID_DEFAULT_LIT:
 		case SHADINGMODELID_THIN_TRANSLUCENT:
 		case SHADINGMODELID_CLOTH:
+		case SHADINGMODELID_SUBSURFACE:
+		case SHADINGMODELID_PREINTEGRATED_SKIN:
 			Payload.Roughness = max(Payload.Roughness, PathRoughness);
 			break;
 		default:
@@ -271,6 +288,8 @@ FMaterialSample SampleMaterial(
 			// Nothing to do, material does not scatter light
 			return Result;
 		case SHADINGMODELID_DEFAULT_LIT:
+		case SHADINGMODELID_SUBSURFACE:
+		case SHADINGMODELID_PREINTEGRATED_SKIN:
 			Result = DefaultLit_SampleMaterial(RayDirection, Payload, RandSample);
 			break;
 		case SHADINGMODELID_CLOTH:
@@ -358,6 +377,8 @@ FMaterialEval EvalMaterial(
 		case SHADINGMODELID_UNLIT:
 			break;
 		case SHADINGMODELID_DEFAULT_LIT:
+		case SHADINGMODELID_SUBSURFACE:
+		case SHADINGMODELID_PREINTEGRATED_SKIN:
 			Result = DefaultLit_EvalMaterial(IncomingDirection, OutgoingDirection, Payload);
 			break;
 		case SHADINGMODELID_CLOTH:
diff --git a/Engine/Shaders/Private/PathTracing/PathTracing.usf b/Engine/Shaders/Private/PathTracing/PathTracing.usf
--- a/Engine/Shaders/Private/PathTracing/PathTracing.usf
+++ b/Engine/Shaders/Private/PathTracing/PathTracing.usf
@@ -6,6 +6,8 @@ PathTracingRayGenShader.usf: Reference path tracing
 
 #define PATH_TRACING  
 
+#define RANDSEQ_UNROLL_SOBOL 0 // unrolling sobol sampler doesn't seem to pay off in the path tracer
+
 #include "../Common.ush"
 #include "../PostProcessCommon.ush"
 #include "../RectLight.ush"
@@ -91,18 +93,6 @@ FMaterialClosestHitPayload TraceTransparentRay(RayDesc Ray, FRayCone RayCone, bo
 		// Unpack the payload
 		FMaterialClosestHitPayload HitPayload = UnpackRayTracingPayload(PackedPayload, Ray);
 
-
-		// material tweaks to workaround temporary limitations
-		if (HitPayload.ShadingModelID == SHADINGMODELID_SUBSURFACE ||
-			HitPayload.ShadingModelID == SHADINGMODELID_PREINTEGRATED_SKIN)
-		{
-			// for now just approximate SSS by moving it into the diffuse color
-			HitPayload.DiffuseColor += HitPayload.CustomData.xyz;
-			HitPayload.DiffuseColor = saturate(HitPayload.DiffuseColor);
-			HitPayload.Opacity = 1; // overwrite radius
-			HitPayload.ShadingModelID = SHADINGMODELID_DEFAULT_LIT;
-		}
-
 		if (HitPayload.BlendingMode == RAY_TRACING_BLEND_MODE_ADDITIVE ||
 			HitPayload.BlendingMode == RAY_TRACING_BLEND_MODE_TRANSLUCENT)
 		{
@@ -216,6 +206,154 @@ float3 TraceTransparentVisibilityRay(RayDesc Ray, uint2 PixelCoord, float PathRo
 	return 0;
 }
 
+struct FProbeResult
+{
+	float HitT;
+	float3 WorldNormal;
+
+	bool IsMiss() { return HitT <= 0; }
+};
+
+FProbeResult TraceProbeRay(RayDesc Ray, uint2 PixelCoord)
+{
+	// Trace a short ray to see if we escaped the surface
+	// TODO: figure out a way to only hit the object we started on (or handle blocking geometry)
+	FPackedMaterialClosestHitPayload PackedPayload = (FPackedMaterialClosestHitPayload)0;
+	// Set payload inputs
+	PackedPayload.SetPixelCoord(PixelCoord);
+	PackedPayload.SetIgnoreSkyMaterials();
+	// TODO: payload flag to skip shading (we only care about the normal)
+	const uint RayFlags = 0;
+	const uint MissShaderIndex = 0;
+	TraceRay(
+		TLAS,
+		RayFlags,
+		RAY_TRACING_MASK_ALL,
+		RAY_TRACING_SHADER_SLOT_MATERIAL,
+		RAY_TRACING_NUM_SHADER_SLOTS,
+		MissShaderIndex,
+		Ray,
+		PackedPayload);
+
+	if (PackedPayload.IsMiss())
+	{
+		return (FProbeResult)0;
+	}
+
+	FProbeResult Result;
+	Result.HitT = PackedPayload.HitT;
+	Result.WorldNormal = PackedPayload.GetWorldNormal();
+	return Result;
+}
+
+bool ProcessSubsurfaceRandomWalk(inout FMaterialClosestHitPayload Payload, inout float3 PathThroughput, inout RandomSequence RandSequence, uint2 PixelCoord, bool SimplifySSS)
+{
+	if (!IsSubsurfaceMaterial(Payload) || PathTracingData.MaxBounces == 0)
+	{
+		// material doesn't have SSS or we don't have any bounces configured
+		// just return so we can process the rest of the payload's shading
+		return true;
+	}
+
+	// TODO: support RGB extinction as well (SHADINGMODELID_SUBSURFACE_PROFILE)
+	// TODO: figure out the right multiplier between the 0-1 Opacity control and world space radius
+	float3 SSSColor = Payload.CustomData.xyz;
+	float SSSRadius = (1 - Payload.Opacity) * 10.0;
+
+	if (SimplifySSS || SSSRadius == 0)
+	{
+		// if we are running in a context that won't see the SSS directly -- just turn it into diffuse and skip the random walk
+		Payload.DiffuseColor += SSSColor;
+		Payload.DiffuseColor = saturate(Payload.DiffuseColor);
+		return true;
+	}
+	// decide randomly to evaluate the SSS part of the material, or keep shading the current part
+	float SSSProb = LobeSelectionProb(SSSColor, Payload.DiffuseColor + Payload.SpecularColor);
+	float3 RandSample = RandomSequence_GenerateSample3D(RandSequence);
+	if (RandSample.x < SSSProb)
+	{
+		// do random walk
+		PathThroughput *= 1 / SSSProb;
+		RandSample.x /= SSSProb;
+	}
+	else
+	{
+		// skip random walk and shade current point instead
+		PathThroughput *= 1 / (1 - SSSProb);
+		return true;
+	}
+
+	// create ray to enter the surface (using a diffuse scattering event)
+	RayDesc Ray;
+	Ray.Origin = Payload.WorldPos;
+	Ray.Direction = TangentToWorld(-CosineSampleHemisphere(RandSample.xy).xyz, Payload.WorldNormal);
+	Ray.TMin = 0;
+	ApplyPositionBias(Ray, -Payload.WorldNormal, PathTracingData.MaxNormalBias);
+
+	// Use the technique detailed in the Hyperion paper (Section 4.4.2)
+	// https://media.disneyanimation.com/uploads/production/publication_asset/177/asset/a.pdf
+	// This avoids overly bright thin regions by applying a bit of internal reflection when the boundary is hit
+#define SSS_USE_TIR 1
+
+#if SSS_USE_TIR
+	const float3 Albedo = 1 - exp(SSSColor * (-11.43 + SSSColor * (15.38 - 13.91 * SSSColor)));
+	SSSRadius *= 0.5; // roughly match parameterization below
+#else
+	// Van de-Hulst inverse mapping
+	// https://blog.selfshadow.com/publications/s2017-shading-course/imageworks/s2017_pbs_imageworks_slides_v2.pdf (Slide 44)
+	// [REDACTED_PROFILE_URL] (Section 7.5.3 of v0.1.3)
+	const float3 Albedo = 1 - Pow2(4.09712 + 4.20863 * SSSColor - sqrt(9.59217 + SSSColor * (41.6808 + 17.7126 * SSSColor)));
+#endif
+	const int MAX_SSS_BOUNCES = PathTracingData.MaxSSSBounces;
+	for (int i = 0; i < MAX_SSS_BOUNCES; i++)
+	{
+		// Using an analog random walk - transmittance and scattering probabilities cancel out so only the albedo
+		// at each bounce needs to be accounted for in the throughput. This will get more complicated once we add
+		// proper support for RGB transmittance and fancier escape strategies like Dwivedi sampling.
+
+		// pick a maximum distance to scatter
+		Ray.TMax = -log(1 - RandSample.z) * SSSRadius;
+
+		FProbeResult ProbeResult = TraceProbeRay(Ray, PixelCoord);
+		if (ProbeResult.IsMiss())
+		{
+			// we didn't hit anything, so scatter randomly in the volume and keep tracing
+			RandSample = RandomSequence_GenerateSample3D(RandSequence);
+			Ray.Origin += Ray.TMax * Ray.Direction;
+			Ray.Direction = UniformSampleSphere(RandSample.xy).xyz;
+			PathThroughput *= Albedo;
+		}
+		else
+		{
+			// our short ray hit the geomtry - decide if we should exit or not
+#if SSS_USE_TIR
+			RandSample = RandomSequence_GenerateSample2D(RandSequence).xxy;
+			float3 WorldNormal = ProbeResult.WorldNormal;
+			float CosTheta = abs(dot(Ray.Direction, WorldNormal));
+			float Fresnel = FresnelReflectance(CosTheta, 1.0 / 1.4);
+			if (RandSample.x < Fresnel)
+			{
+				// internal reflection occured -- reflect and keep tracing
+				// NOTE: weight and probability cancel out, so no throughput adjustment is needed
+				Ray.Origin += ProbeResult.HitT * Ray.Direction;
+				Ray.Direction = reflect(Ray.Direction, WorldNormal);
+				ApplyPositionBias(Ray, -WorldNormal, PathTracingData.MaxNormalBias);
+				continue;
+			}
+#endif
+			// we hit the boundary! overwrite the current payload and exit the walk with a diffuse scattering event
+			Payload.WorldPos = Ray.Origin + ProbeResult.HitT * Ray.Direction;
+			Payload.WorldNormal = ProbeResult.WorldNormal;
+			Payload.ShadingModelID = SHADINGMODELID_NUM; // invalid value so that we get diffuse shading
+			Payload.BlendingMode = RAY_TRACING_BLEND_MODE_OPAQUE;
+			Payload.BaseColor = 1;
+			return true;
+		}
+	}
+	// we scattered a bunch of times and never hit anything -- give up
+	return false;
+}
+
 
 RAY_TRACING_ENTRY_RAYGEN(PathTracingMainRG)
 {
@@ -278,6 +416,9 @@ RAY_TRACING_ENTRY_RAYGEN(PathTracingMainRG)
 			break;
 		}
 
+		// As soon as the path is blurry enough, we can get away with diffuse sampling only
+		const bool bSimplifySSS = PathRoughness >= 0.15;
+
 		if (PathTracingData.ApproximateCaustics)
 		{
 			// modify the path roughness as well as the payload roughness to minimize difficult caustics
@@ -288,16 +429,19 @@ RAY_TRACING_ENTRY_RAYGEN(PathTracingMainRG)
 		}
 		
 
-		// Update origin
-		Ray.Origin = Payload.WorldPos;
-
-
 		// Make sure the surface normal points toward the viewer
 		if (dot(Ray.Direction, Payload.WorldNormal) > 0.0)
 		{
 			Payload.WorldNormal = -Payload.WorldNormal;
 		}
 
+		// Processing the random walk will (stochastically) move the shading point elsewhere on the surface of the object
+		if (!ProcessSubsurfaceRandomWalk(Payload, PathThroughput, RandSequence, LaunchIndex, bSimplifySSS))
+		{
+			// random walk did not terminate at a valid point
+			break;
+		}
+
 		// Choose a random number for both Light sampling and BxDF sampling
 		float4 RandSample = RandomSequence_GenerateSample4D(RandSequence);
 
@@ -413,6 +557,7 @@ RAY_TRACING_ENTRY_RAYGEN(PathTracingMainRG)
 		}
 
 		// Update ray according to material sample
+		Ray.Origin = Payload.WorldPos;
 		Ray.Direction = MaterialSample.Direction;
 		Ray.TMin = 0;
 		Ray.TMax = POSITIVE_INFINITY;
diff --git a/Engine/Source/Runtime/Renderer/Private/PathTracing.cpp b/Engine/Source/Runtime/Renderer/Private/PathTracing.cpp
--- a/Engine/Source/Runtime/Renderer/Private/PathTracing.cpp
+++ b/Engine/Source/Runtime/Renderer/Private/PathTracing.cpp
@@ -76,6 +76,13 @@ TAutoConsoleVariable<int32> CVarPathTracingVisibleLights(
 	ECVF_RenderThreadSafe
 );
 
+TAutoConsoleVariable<int32> CVarPathTracingMaxSSSBounces(
+	TEXT("r.PathTracing.MaxSSSBounces"),
+	256,
+	TEXT("Sets the maximum number of bounces inside subsurface materials. Lowering this value can make subsurface scattering render too dim, while setting it too high can cause long render times.  (default = 256)"),
+	ECVF_RenderThreadSafe
+);
+
 TAutoConsoleVariable<float> CVarPathTracingMaxPathIntensity(
 	TEXT("r.PathTracing.MaxPathIntensity"),
 	-1,
@@ -164,6 +171,7 @@ static bool PrepareShaderArgs(const FViewInfo& View, FPathTracingData& PathTraci
 	}
 
 	PathTracingData.MaxBounces = MaxBounces;
+	PathTracingData.MaxSSSBounces = CVarPathTracingMaxSSSBounces.GetValueOnRenderThread();
 	PathTracingData.MaxNormalBias = GetRaytracingMaxNormalBias();
 	PathTracingData.MISMode = CVarPathTracingMISMode.GetValueOnRenderThread();
 	uint32 VisibleLights = CVarPathTracingVisibleLights.GetValueOnRenderThread();
@@ -190,6 +198,14 @@ static bool PrepareShaderArgs(const FViewInfo& View, FPathTracingData& PathTraci
 		PrevMaxBounces = PathTracingData.MaxBounces;
 	}
 
+	// Changing the number of SSS bounces requires starting over
+	static uint32 PrevMaxSSSBounces = PathTracingData.MaxSSSBounces;
+	if (PathTracingData.MaxSSSBounces != PrevMaxSSSBounces)
+	{
+		NeedInvalidation = true;
+		PrevMaxSSSBounces = PathTracingData.MaxSSSBounces;
+	}
+
 	// Changing MIS mode requires starting over
 	static uint32 PreviousMISMode = PathTracingData.MISMode;
 	if (PreviousMISMode != PathTracingData.MISMode)
diff --git a/Engine/Source/Runtime/Renderer/Private/PathTracingUniformBuffers.h b/Engine/Source/Runtime/Renderer/Private/PathTracingUniformBuffers.h
--- a/Engine/Source/Runtime/Renderer/Private/PathTracingUniformBuffers.h
+++ b/Engine/Source/Runtime/Renderer/Private/PathTracingUniformBuffers.h
@@ -15,6 +15,7 @@ BEGIN_GLOBAL_SHADER_PARAMETER_STRUCT(FPathTracingData, )
 	SHADER_PARAMETER(uint32, MaxSamples)
 	SHADER_PARAMETER(uint32, UseErrorDiffusion)
 	SHADER_PARAMETER(uint32, MaxBounces)
+	SHADER_PARAMETER(uint32, MaxSSSBounces)
 	SHADER_PARAMETER(uint32, MISMode)
 	SHADER_PARAMETER(uint32, ApproximateCaustics)
 	SHADER_PARAMETER(uint32, EnableCameraBackfaceCulling)
