Fix incorrect lightmap data index for landscape VF when using GPU-Scene

diff --git a/Engine/Shaders/Private/LandscapeVertexFactory.ush b/Engine/Shaders/Private/LandscapeVertexFactory.ush
--- a/Engine/Shaders/Private/LandscapeVertexFactory.ush
+++ b/Engine/Shaders/Private/LandscapeVertexFactory.ush
@@ -642,7 +642,12 @@ void GetLightMapCoordinates(FVertexFactoryInterpolantsVSToPS Interpolants, out f
 	LightmapUV0 = Interpolants.LightMapCoordinate.xy * float2( 1, 0.5 );
 	LightmapUV1 = LightmapUV0 + float2( 0, 0.5 );
 #endif
+
+#if VF_USE_PRIMITIVE_SCENE_DATA
+	LightmapDataIndex = GetPrimitiveData(Interpolants.PrimitiveId).LightmapDataIndex;
+#else
 	LightmapDataIndex = 0;
+#endif
 }
 
 void GetShadowMapCoordinate(FVertexFactoryInterpolantsVSToPS Interpolants, out float2 ShadowMapCoordinate, out uint LightmapDataIndex)
@@ -652,7 +657,13 @@ void GetShadowMapCoordinate(FVertexFactoryInterpolantsVSToPS Interpolants, out f
 #else
 	ShadowMapCoordinate = Interpolants.LightMapCoordinate.zw;
 #endif
+	
+#if VF_USE_PRIMITIVE_SCENE_DATA
+	LightmapDataIndex = GetPrimitiveData(Interpolants.PrimitiveId).LightmapDataIndex;
+#else
 	LightmapDataIndex = 0;
+#endif
+
 }
 #endif
 
