Fixes regression in TAA's super res history offsets

diff --git a/Engine/Shaders/Private/TemporalAA/TAACommon.ush b/Engine/Shaders/Private/TemporalAA/TAACommon.ush
--- a/Engine/Shaders/Private/TemporalAA/TAACommon.ush
+++ b/Engine/Shaders/Private/TemporalAA/TAACommon.ush
@@ -538,8 +538,8 @@ uint2 PackSuperResHistory(FSuperResHistory SuperRes)
 	Packed[0] |= ((SuperRes.Age) & SUPERRES_AGE_BITMASK) << 0;
 	Packed[0] |= ((CompressedInterferenceSeed) & SUPERRES_INTERFERENCE_BITMASK) << 5;
 
-	Packed[1] |= (uint(SuperRes.TexelToSample.x + SUPERRES_OFFSET_OFFSET) & SUPERRES_OFFSET_BITS) << 0 * SUPERRES_OFFSET_BITS;
-	Packed[1] |= (uint(SuperRes.TexelToSample.y + SUPERRES_OFFSET_OFFSET) & SUPERRES_OFFSET_BITS) << 1 * SUPERRES_OFFSET_BITS;
+	Packed[1] |= (uint(SuperRes.TexelToSample.x + SUPERRES_OFFSET_OFFSET) & SUPERRES_OFFSET_BITMASK) << 0 * SUPERRES_OFFSET_BITS;
+	Packed[1] |= (uint(SuperRes.TexelToSample.y + SUPERRES_OFFSET_OFFSET) & SUPERRES_OFFSET_BITMASK) << 1 * SUPERRES_OFFSET_BITS;
 
 	return Packed;
 }
diff --git a/Engine/Shaders/Private/TemporalAA/TAAUpdateSuperResHistory.usf b/Engine/Shaders/Private/TemporalAA/TAAUpdateSuperResHistory.usf
--- a/Engine/Shaders/Private/TemporalAA/TAAUpdateSuperResHistory.usf
+++ b/Engine/Shaders/Private/TemporalAA/TAAUpdateSuperResHistory.usf
@@ -136,7 +136,7 @@ void MainCS(
 
 			const float2 PixelOffset = float2(BilinearSamplingOffsets2x2[SampleId]);
 
-			int2 KernelToTexel = int2(round((PixelOffset - SampleInfos.BilinearInterp) * 0x20));
+			int2 KernelToTexel = int2(round((PixelOffset - SampleInfos.BilinearInterp) * 2 * SUPERRES_OFFSET_MAX));
 
 			int2 KernelToSample = KernelToTexel + UnpackedSubpixel.TexelToSample;
 
@@ -178,7 +178,7 @@ void MainCS(
 		SubpixelHistory.InterferenceSeed = saturate(ClosestInterferenceSeed + (E.x - 0.5) * rcp(taa_half(SUPERRES_INTERFERENCE_BITMASK)));
 	}
 
-	#if 0
+	#if 1
 	{
 		Debug.r = SubpixelHistory.InterferenceSeed;
 		Debug.g = SubpixelHistory.Age * rcp(SUPERRES_AGE_BITMASK);
