Optimized Sobol sampler by unrolling loop
Removed ReverseBits32. All platforms now support reversebits() HLSL function.

Robomerge will probably conflict going to main but I'll fix it.

diff --git a/Engine/Shaders/Private/Common.ush b/Engine/Shaders/Private/Common.ush
--- a/Engine/Shaders/Private/Common.ush
+++ b/Engine/Shaders/Private/Common.ush
@@ -1738,25 +1738,10 @@ float DecodePackedHeight(float2 PackedHeight)
 	return DecodeHeightValue(DecodePackedTwoChannelValue(PackedHeight));
 }
 
-/** Reverses all the 32 bits. */
-uint ReverseBits32( uint bits )
-{
-#if SM6_PROFILE || SM5_PROFILE || COMPILER_METAL || COMPILER_DXC
-	return reversebits( bits );
-#else
-	bits = ( bits << 16) | ( bits >> 16);
-	bits = ( (bits & 0x00ff00ff) << 8 ) | ( (bits & 0xff00ff00) >> 8 );
-	bits = ( (bits & 0x0f0f0f0f) << 4 ) | ( (bits & 0xf0f0f0f0) >> 4 );
-	bits = ( (bits & 0x33333333) << 2 ) | ( (bits & 0xcccccccc) >> 2 );
-	bits = ( (bits & 0x55555555) << 1 ) | ( (bits & 0xaaaaaaaa) >> 1 );
-	return bits;
-#endif
-}
-
 /** Reverses all the <BitCount> lowest significant bits. */
 uint ReverseBitsN(uint Bitfield, const uint BitCount)
 {
-	return ReverseBits32(Bitfield) >> (32 - BitCount);
+	return reversebits(Bitfield) >> (32 - BitCount);
 }
 
 
diff --git a/Engine/Shaders/Private/GPUFastFourierTransform.usf b/Engine/Shaders/Private/GPUFastFourierTransform.usf
--- a/Engine/Shaders/Private/GPUFastFourierTransform.usf
+++ b/Engine/Shaders/Private/GPUFastFourierTransform.usf
@@ -121,7 +121,7 @@ uint2 Coord(in uint2 P, in bool bHorizontal)
 
 uint BitReverse(in uint Idx, in uint BitRange)
 {
-	uint ReversedValue = ReverseBits32(Idx);
+	uint ReversedValue = reversebits(Idx);
 		
 	// Shift into correct range.
 	// 32 = sizeof(uint) 
diff --git a/Engine/Shaders/Private/MonteCarlo.ush b/Engine/Shaders/Private/MonteCarlo.ush
--- a/Engine/Shaders/Private/MonteCarlo.ush
+++ b/Engine/Shaders/Private/MonteCarlo.ush
@@ -58,14 +58,14 @@ float3 WorldToTangent(float3 Vec, float3 TangentZ)
 float2 Hammersley( uint Index, uint NumSamples, uint2 Random )
 {
 	float E1 = frac( (float)Index / NumSamples + float( Random.x & 0xffff ) / (1<<16) );
-	float E2 = float( ReverseBits32(Index) ^ Random.y ) * 2.3283064365386963e-10;
+	float E2 = float( reversebits(Index) ^ Random.y ) * 2.3283064365386963e-10;
 	return float2( E1, E2 );
 }
 
 float2 Hammersley16( uint Index, uint NumSamples, uint2 Random )
 {
 	float E1 = frac( (float)Index / NumSamples + float( Random.x ) * (1.0 / 65536.0) );
-	float E2 = float( ( ReverseBits32(Index) >> 16 ) ^ Random.y ) * (1.0 / 65536.0);
+	float E2 = float( ( reversebits(Index) >> 16 ) ^ Random.y ) * (1.0 / 65536.0);
 	return float2( E1, E2 );
 }
 
diff --git a/Engine/Shaders/Private/PathTracing/Utilities/PathTracingRandomSequence.ush b/Engine/Shaders/Private/PathTracing/Utilities/PathTracingRandomSequence.ush
--- a/Engine/Shaders/Private/PathTracing/Utilities/PathTracingRandomSequence.ush
+++ b/Engine/Shaders/Private/PathTracing/Utilities/PathTracingRandomSequence.ush
@@ -165,7 +165,7 @@ uint FastOwenScrambling(uint Index, uint Seed) {
 	// reference implementation - scramble bits one at a time
 
 	// loop below does not require pre-inverted bits, so undo the inversion assumed in the optimized case below
-	Index = ReverseBits32(Index);
+	Index = reversebits(Index);
 	for (uint Mask = 1u << 31, Input = Index; Mask; Mask >>= 1) {
 		Seed = PerfectIntegerHash(Seed); // randomize state
 		Index ^= Seed & Mask;  // flip output (depending on state)
@@ -177,7 +177,7 @@ uint FastOwenScrambling(uint Index, uint Seed) {
 
 	// The operations below will mix bits toward the left, so temporarily reverse the order
 	// NOTE: This operation has been performed outside this call
-	// Index = ReverseBits32(Index);
+	// Index = reversebits(Index);
 	
 	// This follows the basic construction from the paper above. The error-diffusion sampler which
 	// tiles a single point set across the screen, makes it much easier to visually "see" the impact of the scrambling.
@@ -202,7 +202,7 @@ uint FastOwenScrambling(uint Index, uint Seed) {
 
 	// Undo the reverse so that we get left-to-right scrambling
 	// thereby emulating owen-scrambling
-	return ReverseBits32(Index);
+	return reversebits(Index);
 #endif
 }
 
@@ -316,7 +316,7 @@ void RandomSequence_Initialize(inout RandomSequence RandSequence, uint PositionS
 	RandSequence.SampleSeed = 0; // this sampler just counts dimensions
 #elif RANDSEQ == RANDSEQ_OWENSOBOL
 	// pre-compute bit reversal needed for FastOwenScrambling since this index doesn't change
-	RandSequence.SampleIndex = ReverseBits32(TimeSeed);
+	RandSequence.SampleIndex = reversebits(TimeSeed);
 	// change seed to get a unique sequence per pixel
 	RandSequence.SampleSeed  = StrongIntegerHash(PositionSeed);
 #else
@@ -372,12 +372,12 @@ void RandomSequence_Initialize(inout RandomSequence RandSequence, uint2 PixelCoo
 #endif
 
 	// Combine pixel-level and sample-level bits into the sample index (visible structure will be hidden by owen scrambling of the index)
-	RandSequence.SampleIndex = ReverseBits32((65536 * TimeSeed0 + TileID) * MaxSamples + SampleIndex);
+	RandSequence.SampleIndex = reversebits((65536 * TimeSeed0 + TileID) * MaxSamples + SampleIndex);
 
 	// progressive encodings (these kind of work, but quality is much worse)
-	//RandSequence.SampleIndex = ReverseBits32(MortonEncode(uint2(TimeSeed0, TileID)) * MaxSamples + SampleIndex);
-	//RandSequence.SampleIndex = ReverseBits32((MortonCode3(MaxSamples * TimeSeed0 + SampleIndex) * 4) | (MortonCode3(Y) * 2) | MortonCode3(X));
-	//RandSequence.SampleIndex = ReverseBits32(MortonEncode(uint2(MortonEncode(uint2(X, Y)), SampleIndex)));
+	//RandSequence.SampleIndex = reversebits(MortonEncode(uint2(TimeSeed0, TileID)) * MaxSamples + SampleIndex);
+	//RandSequence.SampleIndex = reversebits((MortonCode3(MaxSamples * TimeSeed0 + SampleIndex) * 4) | (MortonCode3(Y) * 2) | MortonCode3(X));
+	//RandSequence.SampleIndex = reversebits(MortonEncode(uint2(MortonEncode(uint2(X, Y)), SampleIndex)));
 
 	RandSequence.SampleSeed = 0; // always use the same sequence
 #else
