Invert depth used for transmittion thickness calculation sampling to match to match reversed shadow projection matrices

diff --git a/Engine/Shaders/Private/TransmissionThickness.ush b/Engine/Shaders/Private/TransmissionThickness.ush
--- a/Engine/Shaders/Private/TransmissionThickness.ush
+++ b/Engine/Shaders/Private/TransmissionThickness.ush
@@ -141,7 +141,7 @@ float CalcTransmissionThickness(float3 ViewPos, float3 LightPos, float3 WorldNor
 	Settings.DensityMulConstant = DepthScale * ProjectionDepthBiasParameters.w;
 
 
-	Settings.SceneDepth = ShrinkedScreenPos.z + ProjectionDepthBiasParameters.x;
+	Settings.SceneDepth = (1.0f - ShrinkedScreenPos.z) + ProjectionDepthBiasParameters.x;
 	//Alpha is the thickness.
 
 	float ThicknessSum = 0.0;
