Optimize VSM projection and SMRT

diff --git a/Engine/Shaders/Private/PathTracing/Utilities/PathTracingRandomSequence.ush b/Engine/Shaders/Private/PathTracing/Utilities/PathTracingRandomSequence.ush
--- a/Engine/Shaders/Private/PathTracing/Utilities/PathTracingRandomSequence.ush
+++ b/Engine/Shaders/Private/PathTracing/Utilities/PathTracingRandomSequence.ush
@@ -277,7 +277,7 @@ float4 SobolSampler(uint SampleIndex, inout uint Seed)
 #if RANDSEQ_UNROLL_SOBOL // unrolled version seems faster in most cases
 	UNROLL for (uint b = 0, v = 1; b < 32; b++)
 	{
-		uint IndexBit = (SobolIndex >> b) & 1;
+		uint IndexBit = (SobolIndex >> b) & 1;		// bitfield extract
 #else // however a loop with early exit seems faster in the path tracing context (possibly due to low occupancy?)
 	for (uint b = 0, v = 1; SobolIndex; SobolIndex >>= 1, b++)
 	{
diff --git a/Engine/Shaders/Private/VirtualShadowMaps/PageAccessCommon.ush b/Engine/Shaders/Private/VirtualShadowMaps/PageAccessCommon.ush
--- a/Engine/Shaders/Private/VirtualShadowMaps/PageAccessCommon.ush
+++ b/Engine/Shaders/Private/VirtualShadowMaps/PageAccessCommon.ush
@@ -191,11 +191,9 @@ struct FShadowPageTranslationResult
 // Finds the best-resolution mapped page at the given UV
 FShadowPageTranslationResult ShadowVirtualToPhysicalUV(uint VirtualShadowMapID, float2 ShadowMapUV)
 {
-	float2 VirtualTexelAddressLevell0Float = ShadowMapUV * float(CalcLevelDimsTexels(0));
-	uint2 VirtualTexelAddressLevel0 = uint2(VirtualTexelAddressLevell0Float);
-	uint2 vPageLevel0 = uint2(VirtualTexelAddressLevel0 >> VSM_LOG2_PAGE_SIZE);
+	uint2 vPage = uint2( ShadowMapUV * VSM_LEVEL0_DIM_PAGES_XY );
 
-	FShadowPhysicalPage PhysicalPageEntry = ShadowGetPhysicalPage(CalcPageOffset(VirtualShadowMapID, 0, vPageLevel0));
+	FShadowPhysicalPage PhysicalPageEntry = ShadowGetPhysicalPage( VirtualShadowMapID * VSM_PAGE_TABLE_SIZE + vPage.x + vPage.y * VSM_LEVEL0_DIM_PAGES_XY );
 
 	FShadowPageTranslationResult Result;
 	Result.bValid = PhysicalPageEntry.bAnyLODValid;
diff --git a/Engine/Shaders/Private/VirtualShadowMaps/ProjectionSpot.ush b/Engine/Shaders/Private/VirtualShadowMaps/ProjectionSpot.ush
--- a/Engine/Shaders/Private/VirtualShadowMaps/ProjectionSpot.ush
+++ b/Engine/Shaders/Private/VirtualShadowMaps/ProjectionSpot.ush
@@ -144,7 +144,7 @@ FSMRTResult ShadowRayCastSpotLight(
 	float3 RayStart,
 	float3 RayEnd,
 	int NumSteps,
-	float StepOffset)
+	half StepOffset)
 {
 	FSMRTSingleRayState RayState = SMRTSingleRayInitialize(VirtualShadowMapId, RayStart, RayEnd);
 	return SMRTRayCast(RayState, NumSteps, StepOffset);
@@ -155,7 +155,7 @@ FSMRTResult ShadowRayCastPointLight(
 	float3 RayStart,
 	float3 RayEnd,
 	int NumSteps,
-	float StepOffset)
+	half StepOffset)
 {
 	FVirtualShadowMapProjectionShaderData BaseProjectionData = GetVirtualShadowMapProjectionData(VirtualShadowMapId);
 
@@ -181,28 +181,30 @@ FSMRTResult ShadowRayCastPointLight(
 	return Result;
 }
 
-float3 GetRandomPointLightRayVector(FLightShaderParameters LightParameters, float3 WorldPosition, float2 E)
+// MaxBits version reduces unrolled iterations
+float4 SobolSampler(uint SampleIndex, inout uint Seed, uint MaxBits)
 {
-	float3 LightDirection = LightParameters.Position - WorldPosition;
-	float RayLength2 = dot(LightDirection, LightDirection);
-	float Radius2 = LightParameters.SourceRadius * LightParameters.SourceRadius;
-
-	// Clamping this instead of handling sampling point inside light radius properly is good enough for now
-	float SinThetaMax2 = saturate(Radius2 / RayLength2);
-
-	float4 DirAndPdf = UniformSampleConeConcentricRobust(E, SinThetaMax2);
-	float CosTheta = DirAndPdf.z;
-	float SinTheta2 = 1.0 - CosTheta * CosTheta;
-
-	// Project ray direction to world-space, such that z-axis aligns with LightDirection
-	float RayLength = sqrt(RayLength2);
-	LightDirection *= rcp(RayLength + 1e-4);
-	float3 RayDirection = TangentToWorld(DirAndPdf.xyz, LightDirection);
+	// first scramble the index to decorelate from other 4-tuples
+	uint SobolIndex = FastOwenScrambling(SampleIndex, EvolveSobolSeed(Seed)) & ( (1 << MaxBits) - 1 );
+	// now get Sobol' point from this index
+	uint4 Result = uint4(SobolIndex, 0, 0, 0);
+	UNROLL for (uint b = 0, v = 1; b < MaxBits; b++)
+	{
+		uint IndexBit = (SobolIndex >> b) & 1;
+		Result.y ^= IndexBit * v;
+		Result.zw ^= IndexBit * SobolMatrices[b];
+		v ^= v << 1;
+	}
 
-	// Clip length to closest intersection with the sphere
-	//RayLength = RayLength * (CosTheta - sqrt(max(SinThetaMax2 - SinTheta2, 0.0)));
+	// finally scramble the points to avoid structured artifacts
+	Result.x = FastOwenScrambling(Result.x, EvolveSobolSeed(Seed));
+	Result.y = FastOwenScrambling(Result.y, EvolveSobolSeed(Seed));
+	Result.z = FastOwenScrambling(Result.z, EvolveSobolSeed(Seed));
+	Result.w = FastOwenScrambling(Result.w, EvolveSobolSeed(Seed));
 
-	return RayLength * RayDirection;
+	// output as float in [0,1) taking care not to skew the distribution
+	// due to the non-uniform spacing of floats in this range
+	return (Result >> 8) * 5.96046447754e-08; // * 2^-24
 }
 
 
@@ -222,74 +224,71 @@ float TraceLocalLight(
 	inout float3 DebugOutput,
 	inout float InOutHitDistance)
 {
-	FGBufferData GBufferData = GetGBufferDataUint( uint2( SvPosition.xy ), true );
+	uint2 PixelPos = uint2( SvPosition.xy );
+
+	const half3 N = GetGBufferDataUint( PixelPos, true ).WorldNormal;
 	const float3 WorldPosition = TranslatedWorldPosition - View.PreViewTranslation;
 
 	float3 ToLight = Light.Position - WorldPosition;
 	float DistSqr = dot(ToLight, ToLight);
 	float Falloff = rcp(DistSqr + 1);
 
-	float DistToLight = sqrt(DistSqr);
-	float3 L = ToLight / DistToLight;	
+	float InvDist = rsqrt( DistSqr );
+	float DistToLight = DistSqr * InvDist;
 
-	float SinAlphaSqr = saturate(Pow2(Light.SourceRadius) * Falloff);
-	float SinAlpha = sqrt(SinAlphaSqr);
-	bool bBackFace = dot(GBufferData.WorldNormal, normalize(ToLight)) < -SinAlpha;
+	half3 ConeAxis = ToLight * InvDist;
 
-	float RayMissFactor = 0;
+	half SinAlphaSqr = saturate(Pow2(Light.SourceRadius) * Falloff);
+	half SinAlpha = sqrt(SinAlphaSqr);
+	bool bBackFace = dot( N, ConeAxis ) < -SinAlpha;
+
+	half RayMissFactor = 0;
 	uint SampleCount = 0;
 	if (!bBackFace)
 	{
-		uint2 SobolBase = SobolPixel(SvPosition.xy);
-		uint2 SobolFrame = SobolIndex(SobolBase, View.StateFrameIndexMod8, 3);
-
 		float StepOffset = InterleavedGradientNoise(SvPosition.xy, View.StateFrameIndexMod8);
 
+		half ConeSin = Light.SourceRadius * InvDist;
+		half ConeCos = sqrt( 1 - Pow2( ConeSin ) );
+
+		uint TileID = MortonEncode( PixelPos );
+
 		uint i = 0;
 		for( ; i < uint(RayCount); i++ )
 		{
 			RandomSequence RandSequence;
-			RandomSequence_Initialize( RandSequence, uint2( SvPosition.xy ), i, View.StateFrameIndex, RayCount );
+#if 0
+			RandomSequence_Initialize( RandSequence, PixelPos, i, View.StateFrameIndex, RayCount );
 
-#if 1
-			float2 E = RandomSequence_GenerateSample2D( RandSequence );
+			float4 RandSample = RandomSequence_GenerateSample4D( RandSequence );
 #else
-			float2 E = float2( SobolIndex( SobolFrame, i << 3 ) ) / 0x10000;
+			// Combine pixel-level and sample-level bits into the sample index (visible structure will be hidden by owen scrambling of the index)
+			RandSequence.SampleIndex = reversebits((65536 * View.StateFrameIndex + TileID) * RayCount + i);
+			RandSequence.SampleSeed = 0; // always use the same sequence
+
+			float4 RandSample = SobolSampler( RandSequence.SampleIndex, RandSequence.SampleSeed, 8 );
 #endif
 
-			float3 RayEnd = float3(0, 0, 0);
-			if (LightType == LIGHT_TYPE_SPOT)
-			{
-				float2 LightUV = UniformSampleDiskConcentric( E ) * Light.SourceRadius;
+			float2 E = RandSample.xy;
 
-				RayEnd =
-					Light.Position +
-					LightUV.x * Light.Tangent +
-					LightUV.y * cross( Light.Direction, Light.Tangent );
-			}
-			else
-			{
-				// Point light
-				RayEnd = WorldPosition + GetRandomPointLightRayVector( Light, WorldPosition, E );
-			}
+			half2 LightUV = UniformSampleDiskConcentricApprox( E ).xy;
+			LightUV *= ConeSin;
 
-			float3 Dir = normalize( RayEnd - WorldPosition );
-			float3 RayStart = WorldPosition + Dir * RayStartOffset;
+			half SinTheta2 = dot( LightUV, LightUV );
+			half SinTheta = sqrt( SinTheta2 );
+			half CosTheta = sqrt( 1 - SinTheta2 );
+
+			half3 Dir = TangentToWorld( half3( LightUV, CosTheta ), ConeAxis );
 
 			// The further afield our rays go the poorer our approximation is as the "bend" due to our testing
 			// against a shadow map instead of along the ray increases. Thus we avoid going all the way to the light
-			// where the projeciton becomes extreme.
-			RayEnd = lerp( RayStart, RayEnd, 0.75f );
-
-			{
-				// Clip ray to max angle to limit penumbra sizes
-				float CosLambda = dot( L, Dir );
-				float SinLambda = sqrt( 1.0f - CosLambda * CosLambda );
-				float MaxRayLength = DistToLight / ( CosLambda + CotMaxRayAngleFromLight * SinLambda );
-				float ClampedLength = saturate( MaxRayLength / length( RayEnd - RayStart ) );
+			// where the projection becomes extreme.
+			
+			// Clip ray to max angle to limit penumbra sizes
+			float ClampedLength = DistToLight * 0.75 * saturate( 1.5 / ( CosTheta + CotMaxRayAngleFromLight * SinTheta ) );
 
-				RayEnd = lerp( RayStart, RayEnd, ClampedLength );
-			}
+			float3 RayStart = WorldPosition + Dir * RayStartOffset;
+			float3 RayEnd = WorldPosition + Dir * ClampedLength;
 
 			FSMRTResult Result;
 			if ( LightType == LIGHT_TYPE_SPOT )
@@ -334,9 +333,15 @@ float TraceLocalLight(
 				bool bHit = Result.bValidHit;
 				if( bHit )
 				{
-					float WorldBluringRadius = ComputeLightSampleWorldBluringRadius( TranslatedWorldPosition, LIGHT_TYPE_SPOT, Light, Result.OccluderDistance );
-					float PixelBluringRadius = ProjectWorldDistanceToPixelDistance( SceneDepth, WorldBluringRadius );
+				#if 0
+					//float WorldBluringRadius = ComputeLightSampleWorldBluringRadius( TranslatedWorldPosition, LIGHT_TYPE_SPOT, Light, Result.OccluderDistance );
+					//float PixelBluringRadius = ProjectWorldDistanceToPixelDistance( SceneDepth, WorldBluringRadius );
+					float WorldBluringRadius = Light.SourceRadius * Result.OccluderDistance / max( DistToLight - Result.OccluderDistance, 0.001 );
+					float PixelBluringRadius = WorldBluringRadius * View.ViewSizeAndInvSize.x * 0.5 / ( SceneDepth * View.ClipToView[0][0] );
 					bHit = PixelBluringRadius > 2;
+				#else
+					bHit = Result.OccluderDistance * Light.SourceRadius * ( 0.25 * View.ViewSizeAndInvSize.y / ( SceneDepth * View.ClipToView[1][1] ) ) > max( DistToLight - Result.OccluderDistance, 0.001 );
+				#endif
 				}
 
 				// All lanes miss or hit something close enough to not be noisy.
diff --git a/Engine/Shaders/Private/VirtualShadowMaps/SMRTTemplate.ush b/Engine/Shaders/Private/VirtualShadowMaps/SMRTTemplate.ush
--- a/Engine/Shaders/Private/VirtualShadowMaps/SMRTTemplate.ush
+++ b/Engine/Shaders/Private/VirtualShadowMaps/SMRTTemplate.ush
@@ -27,13 +27,18 @@ FSMRTResult SMRTRayCast(
 	int SampleMissCount = 0;
 	float OccluderDistance = 100000;
 
-	const float CompareScale = 2.0 / Pow2(NumSteps);
-	const float CompareBias = -CompareScale * (NumSteps - StepOffset + 1.0f);
+	const float TimeScale = -1.0f / NumSteps;
+	const float TimeBias = 1.0f + ( 1.0 - StepOffset ) * TimeScale;
+
+	// Backward delta. Time[i] - Time[i-1], Half because CompareTolerance is doubled.
+	const float CompareScale = Pow2( TimeScale );
+	const float CompareBias = ( TimeBias - 0.5 * TimeScale ) * TimeScale;
 
 	bool bValidHit = false;
 	for (int i = 0; i <= NumSteps; i++)
 	{
-		float SampleTime = Pow2(saturate((-1.0f / NumSteps) * i + 1.0f - (StepOffset - 0.5f) / NumSteps));
+		const float SampleTime		= i == NumSteps ? 0 : Pow2( TimeScale * i + TimeBias );
+		const float SampleCompare	= i == NumSteps ? 0.5 * Pow2( TimeScale * StepOffset ) : CompareScale * i + CompareBias;
 
 		FSMRTSample Sample = SMRTFindSample(RayState, SampleTime);
 		if (Sample.bValid)
@@ -53,7 +58,7 @@ FSMRTResult SMRTRayCast(
 			}
 			else
 			{
-				const float CompareTolerance = abs(Sample.CompareToleranceScale * (CompareScale * i + CompareBias));
+				const float CompareTolerance = abs(Sample.CompareToleranceScale * SampleCompare);
 
 				bool bBehind = (SampleDepth - ReferenceDepth) > (2.0f * CompareTolerance);
 				float DepthSample = bBehind ? DepthHistory : SampleDepth;
diff --git a/Engine/Shaders/Private/VirtualShadowMaps/VirtualShadowMapProjection.usf b/Engine/Shaders/Private/VirtualShadowMaps/VirtualShadowMapProjection.usf
--- a/Engine/Shaders/Private/VirtualShadowMaps/VirtualShadowMapProjection.usf
+++ b/Engine/Shaders/Private/VirtualShadowMaps/VirtualShadowMapProjection.usf
@@ -356,6 +356,9 @@ void VirtualShadowMapCompositePS(
 	OutShadowMask = EncodeLightAttenuationFromMask(Shadow);
 }
 
+SCHEDULER_MIN_PRESSURE
+MAX_OCCUPENCY
+
 FLightShaderParameters ConvertFromLocal( const FLocalLightData LightData )
 {
 	FLightShaderParameters Light = (FLightShaderParameters)0;
@@ -373,13 +376,6 @@ FLightShaderParameters ConvertFromLocal( const FLocalLightData LightData )
 	return Light;
 }
 
-uint EncodeBlurRadius( float Radius )
-{
-	// Does WORLD_RADIUS_MISS really need to be that big?
-	const float Scale = 254 / log2( WORLD_RADIUS_MISS + 0.5 );
-	return max( Scale * log2( Radius + 0.5 ), 254 );
-}
-
 StructuredBuffer< int > VirtualShadowMapIdRemap;
 uint NumDirectionalLightSmInds;
 
@@ -404,19 +400,18 @@ void VirtualShadowMapProjection(
 	const float SceneDepth = ConvertFromDeviceZ( DeviceZ );
 
 	const float4 SvPosition = float4( float2( PixelPos ) + 0.5, DeviceZ, 1.0f );
-	const float2 ScreenUV = SvPositionToBufferUV( SvPosition );
+	//const float2 ScreenUV = SvPositionToBufferUV( SvPosition );
 
 	const float3 TranslatedWorldPosition = SvPositionToTranslatedWorld( SvPosition );
 	const float3 WorldPosition = TranslatedWorldPosition - View.PreViewTranslation;
 
 	const float ContactShadowLengthWorld = ContactShadowLength * View.ClipToView[1][1] * SceneDepth;
-	const float Noise = InterleavedGradientNoise( SvPosition.xy, View.StateFrameIndexMod8 );
+	const half Noise = InterleavedGradientNoise( SvPosition.xy, View.StateFrameIndexMod8 );
 
-	const float3 EstimatedGeoWorldNormal = GetEstimatedGeoWorldNormal(TranslatedWorldPosition, ScreenUV);
+	//const half3 EstimatedGeoWorldNormal = GetEstimatedGeoWorldNormal(TranslatedWorldPosition, ScreenUV);
+	const half3 EstimatedGeoWorldNormal = GetGBufferDataUint( PixelPos, true ).WorldNormal;
 
 	uint ShadowMask = 0xffffffff;
-	uint ShadowIndexes = 0xffffffff;
-	uint ShadowValues = 0;
 	
 	uint EyeIndex = 0; // ??
 	uint GridLinearIndex = ComputeLightGridCellIndex( LocalPixelPos, SceneDepth, EyeIndex );
@@ -425,28 +420,27 @@ void VirtualShadowMapProjection(
 	LOOP
 	for (uint Index = 0; Index < CulledLightGridData.NumLocalLights; ++Index)
 	{
+		const uint LocalLightIndex = ForwardLightData.CulledLightDataGrid[CulledLightGridData.DataStartIndex + Index];
 		const FLocalLightData LightData = GetLocalLightData(CulledLightGridData.DataStartIndex + Index, EyeIndex);
-		const uint LightGridLightIndex = ForwardLightData.CulledLightDataGrid[CulledLightGridData.DataStartIndex + Index];
 
 		// The Virtual Shadow Remap stores directional lights first
-		int VirtualShadowMapId = VirtualShadowMapIdRemap[NumDirectionalLightSmInds + LightGridLightIndex];
+		int VirtualShadowMapId = VirtualShadowMapIdRemap[NumDirectionalLightSmInds + LocalLightIndex];
 		if (VirtualShadowMapId != INDEX_NONE)
 		{
-			float3 DebugOutput = float3(0, 0, 0);
-			int SampleCount = 1;
-			float HitDistance = WORLD_RADIUS_MISS;
-			float RayMissFactor = 0.0f;
-
 			float3 ToLight = LightData.LightPositionAndInvRadius.xyz - WorldPosition;
 			float d2 = dot( ToLight, ToLight );
 			float InvDist = rsqrt( d2 );
-			float3 L = ToLight * InvDist;
+			half3 L = ToLight * InvDist;
 			
 			bool bInLightRadius = InvDist >= LightData.LightPositionAndInvRadius.w;
 			bool bInLightCone = dot( L, LightData.LightDirectionAndShadowMask.xyz ) >= LightData.SpotAnglesAndSourceRadiusPacked.x;
 
 			if (bInLightRadius && bInLightCone)
 			{
+				half RayMissFactor = 1.0f;
+				half HitDistance = WORLD_RADIUS_MISS;
+				float3 DebugOutput = float3(0, 0, 0);
+
 				FLightShaderParameters Light = ConvertFromLocal( LightData );
 
 				uint LightType = Light.SpotAngles.x > -2.0f ? LIGHT_TYPE_SPOT : LIGHT_TYPE_POINT;
@@ -482,7 +476,7 @@ void VirtualShadowMapProjection(
 
 				if (ContactShadowLengthWorld > 0.0f)
 				{
-					float ContactShadowHitDistance = VirtualShadowMapScreenRayCast(TranslatedWorldPosition, L, ContactShadowLengthWorld, CONTACT_SHADOW_SAMPLES, Noise);		
+					half ContactShadowHitDistance = VirtualShadowMapScreenRayCast(TranslatedWorldPosition, L, ContactShadowLengthWorld, CONTACT_SHADOW_SAMPLES, Noise);		
 					if (ContactShadowHitDistance >= 0.0f)
 					{
 						RayMissFactor = 0.0f;
@@ -490,20 +484,8 @@ void VirtualShadowMapProjection(
 					}
 				}
 
-				//ShadowMask ^= uint( RayMissFactor < 1.0 ) << Index;
 				ShadowMask ^= ( ~uint( round( RayMissFactor * 7.0 ) ) & 7 ) << (Index*3);
-
-				if( Index < 4 && LightGridLightIndex < 255 )
-				{
-					float WorldBluringRadius = ComputeLightSampleWorldBluringRadius(TranslatedWorldPosition, LightType, Light, HitDistance);
-					uint EncodedValue = RayMissFactor < 1.0 ? EncodeBlurRadius( WorldBluringRadius ) : 255;
-
-					ShadowIndexes = ( ShadowIndexes << 8 ) | LightGridLightIndex;
-					ShadowValues = ( ShadowValues << 8 ) | EncodedValue;
-				}
 			}
-
-			//OutputTarget = EncodeOutput(TranslatedWorldPosition, Light, LIGHT_TYPE_SPOT, RayMissFactor, HitDistance, SampleCount, DebugOutput);
 		}
 	}
 
diff --git a/Engine/Shaders/Public/Platform.ush b/Engine/Shaders/Public/Platform.ush
--- a/Engine/Shaders/Public/Platform.ush
+++ b/Engine/Shaders/Public/Platform.ush
@@ -566,6 +566,14 @@ void ShaderYield()
 #define HOIST_DESCRIPTORS
 #endif
 
+#ifndef SCHEDULER_MIN_PRESSURE
+#define SCHEDULER_MIN_PRESSURE
+#endif
+
+#ifndef MAX_OCCUPENCY
+#define MAX_OCCUPENCY
+#endif
+
 // ---------------------------------------------------- Interpolator attribute fallbacks
 
 #ifndef COMPRESSED_16_FLOAT
diff --git a/Engine/Source/Runtime/Renderer/Private/VirtualShadowMaps/VirtualShadowMapProjection.cpp b/Engine/Source/Runtime/Renderer/Private/VirtualShadowMaps/VirtualShadowMapProjection.cpp
--- a/Engine/Source/Runtime/Renderer/Private/VirtualShadowMaps/VirtualShadowMapProjection.cpp
+++ b/Engine/Source/Runtime/Renderer/Private/VirtualShadowMaps/VirtualShadowMapProjection.cpp
@@ -514,6 +514,9 @@ class FVirtualShadowMapProjectionCS : public FGlobalShader
 		{
 			OutEnvironment.CompilerFlags.Add(CFLAG_WaveOperations);
 		}
+
+		OutEnvironment.CompilerFlags.Add( CFLAG_Wave32 );
+		OutEnvironment.CompilerFlags.Add( CFLAG_AllowRealTypes );
 	}
 
 	static bool ShouldCompilePermutation(const FGlobalShaderPermutationParameters& Parameters)
