* Fix self-shadowing issue with virtual map and hair rendering.
* Fix typo in early out condition based hair coverage.

diff --git a/Engine/Shaders/Private/VirtualShadowMaps/VirtualShadowMapProjection.usf b/Engine/Shaders/Private/VirtualShadowMaps/VirtualShadowMapProjection.usf
--- a/Engine/Shaders/Private/VirtualShadowMaps/VirtualShadowMapProjection.usf
+++ b/Engine/Shaders/Private/VirtualShadowMaps/VirtualShadowMapProjection.usf
@@ -277,7 +277,9 @@ FProjectLightResult ProjectLight(
 				HitDistance);
 		}
 
-		if (ContactShadowLengthWorld > 0.0f)
+		// Do not run contact shadow when computing the hair shadow mask (i.e. shadow mask applied on hair, has the scene 
+		// depth buffer contains fully opaque depth, which create false positive intersection resulting in wrong self shadowing)
+		if (ContactShadowLengthWorld > 0.0f && !bIsHair)
 		{
 			half ContactShadowHitDistance = VirtualShadowMapScreenRayCast(TranslatedWorldPosition, L, ContactShadowLengthWorld, CONTACT_SHADOW_SAMPLES, Noise);		
 			if (ContactShadowHitDistance >= 0.0f)
@@ -314,7 +316,7 @@ void VirtualShadowMapProjection(
 	{
 		FCategorizationData CatData = DecodeCategorizationData(HairStrands.HairCategorizationTexture.Load( int3(PixelPos, 0 ) ) );
 		DeviceZ = CatData.ClosestDepth;
-		if (CatData.PixelCoverage < 0)
+		if (CatData.PixelCoverage == 0)
 		{
 			return;
 		}
