Add support for Subsurface Profile shading model in path tracer

The random walk logic now can account for a colored mean free path. A new bsdf was added to account for the dual specular lobes.

Fix some weighting issues in transparency and avoid NaNs when the refraction mode is not based on IOR.

diff --git a/Engine/Shaders/Private/PathTracing/Material/PathTracingDefaultLit.ush b/Engine/Shaders/Private/PathTracing/Material/PathTracingDefaultLit.ush
--- a/Engine/Shaders/Private/PathTracing/Material/PathTracingDefaultLit.ush
+++ b/Engine/Shaders/Private/PathTracing/Material/PathTracingDefaultLit.ush
@@ -83,11 +83,8 @@ FMaterialSample DefaultLit_SampleMaterial(
 		RandSample.x /= LobeSelectionPdf;
 
 		// Lambert
-#if 1
+		// TODO: evaluate CosineSampleHemisphereConcentric
 		float4 Result = CosineSampleHemisphere(RandSample.xy);
-#else
-		float4 Result = CosineSampleHemisphereConcentric(RandSample.xy);
-#endif
 
 		L = Result.xyz;
 		H = normalize(L + V);
diff --git a/Engine/Shaders/Private/PathTracing/Material/PathTracingMaterialCommon.ush b/Engine/Shaders/Private/PathTracing/Material/PathTracingMaterialCommon.ush
--- a/Engine/Shaders/Private/PathTracing/Material/PathTracingMaterialCommon.ush
+++ b/Engine/Shaders/Private/PathTracing/Material/PathTracingMaterialCommon.ush
@@ -75,6 +75,36 @@ float3 GetSubsurfaceColor(FMaterialClosestHitPayload HitInfo)
 	return HitInfo.CustomData.xyz;
 }
 
+void SetSubsurfaceColor(inout FMaterialClosestHitPayload HitInfo, float3 SSSColor)
+{
+	HitInfo.CustomData.xyz = SSSColor;
+}
+
+// Steal an unused field for storing this information
+float3 GetSubsurfaceRadius(FMaterialClosestHitPayload HitInfo)
+{
+	return HitInfo.IndirectIrradiance;
+}
+
+void SetSubsurfaceRadius(inout FMaterialClosestHitPayload HitInfo, float3 SSSRadius)
+{
+	HitInfo.IndirectIrradiance = SSSRadius;
+}
+
+// Steal an unused float3 (the only material with dual specular does not support anisotropy)
+// Data is encoded as (x: roughness 0, y: roughness 1, z: roughness mix)
+float3 GetDualRoughnessSpecular(FMaterialClosestHitPayload HitInfo)
+{
+	return float3(HitInfo.WorldTangent);
+}
+
+void SetDualRoughnessSpecular(inout FMaterialClosestHitPayload HitInfo, float3 RoughnessData)
+{
+	// store data in unused field and compute average roughness
+	HitInfo.WorldTangent = RoughnessData;
+	HitInfo.Roughness = lerp(RoughnessData.x, RoughnessData.y, RoughnessData.z);
+}
+
 float3 GetWorldNormal(FMaterialClosestHitPayload HitInfo)
 {
 	return HitInfo.WorldNormal;
@@ -118,7 +148,9 @@ bool IsMaterialTransmissive(FMaterialClosestHitPayload HitInfo)
 
 bool IsSubsurfaceMaterial(FMaterialClosestHitPayload HitInfo)
 {
-	return HitInfo.ShadingModelID == SHADINGMODELID_SUBSURFACE || HitInfo.ShadingModelID == SHADINGMODELID_PREINTEGRATED_SKIN;
+	return HitInfo.ShadingModelID == SHADINGMODELID_SUBSURFACE ||
+		   HitInfo.ShadingModelID == SHADINGMODELID_PREINTEGRATED_SKIN ||
+		   HitInfo.ShadingModelID == SHADINGMODELID_SUBSURFACE_PROFILE;
 }
 
 // Given two lobes that will roughly contribute colors A and B to the total (estimated for example using directional albedo)
diff --git a/Engine/Shaders/Private/PathTracing/Material/PathTracingMaterialSampling.ush b/Engine/Shaders/Private/PathTracing/Material/PathTracingMaterialSampling.ush
--- a/Engine/Shaders/Private/PathTracing/Material/PathTracingMaterialSampling.ush
+++ b/Engine/Shaders/Private/PathTracing/Material/PathTracingMaterialSampling.ush
@@ -19,6 +19,8 @@ PathTracingMaterialSampling.usf: Material sampling functions
 #include "PathTracingClearCoat.ush"
 #include "PathTracingTwoSidedFoliage.ush"
 #include "PathTracingCloth.ush"
+#include "PathTracingSubsurfaceProfile.ush"
+#include "../../BurleyNormalizedSSSCommon.ush"
 
 #ifndef FORCE_ALL_MATERIALS_DIFFUSE
 	#define FORCE_ALL_MATERIALS_DIFFUSE 0
@@ -71,9 +73,6 @@ float3 GetMaterialTransparency(FMaterialClosestHitPayload Payload, float3 RayDir
 				return 0;
 			}
 		}
-		case SHADINGMODELID_SUBSURFACE:
-		case SHADINGMODELID_PREINTEGRATED_SKIN:
-			return 0;
 		default:
 			// other shading models support blend modes
 			break;
@@ -133,6 +132,7 @@ float3 EstimateMaterialAlbedo(FMaterialClosestHitPayload Payload)
 			break;
 	}
 
+	// FIXME: could make this more accurate by estimating the directional albedo
 	float3 Albedo = 0;
 	switch (Payload.ShadingModelID)
 	{
@@ -140,19 +140,19 @@ float3 EstimateMaterialAlbedo(FMaterialClosestHitPayload Payload)
 			// this object does not scatter light
 			return 0;
 		case SHADINGMODELID_TWOSIDED_FOLIAGE:
+		case SHADINGMODELID_SUBSURFACE:
+		case SHADINGMODELID_PREINTEGRATED_SKIN:
+		case SHADINGMODELID_SUBSURFACE_PROFILE:
 			Albedo = BrdfWeight * (Payload.DiffuseColor + Payload.SpecularColor + GetSubsurfaceColor(Payload));
 			break;
 		case SHADINGMODELID_DEFAULT_LIT:
 		case SHADINGMODELID_CLEAR_COAT:
-			// FIXME: could make this more accurate by estimating the directional albedo
+			// TODO: include clearcoat albedo as well
 			Albedo = BrdfWeight * (Payload.DiffuseColor + Payload.SpecularColor);
+			break;
 		case SHADINGMODELID_CLOTH:
 			Albedo = BrdfWeight * (Payload.DiffuseColor + lerp(Payload.SpecularColor, Payload.CustomData.xyz, Payload.CustomData.w));
 			break;
-		case SHADINGMODELID_SUBSURFACE:
-		case SHADINGMODELID_PREINTEGRATED_SKIN:
-			Albedo = BrdfWeight * (Payload.DiffuseColor + Payload.SpecularColor + Payload.CustomData.xyz);
-			break;
 		default:
 			Albedo = BrdfWeight * Payload.BaseColor;
 			break;
@@ -188,6 +188,14 @@ void AdjustPathAndMaterialRoughness(inout FMaterialClosestHitPayload Payload, in
 		case SHADINGMODELID_PREINTEGRATED_SKIN:
 			Payload.Roughness = max(Payload.Roughness, PathRoughness);
 			break;
+		case SHADINGMODELID_SUBSURFACE_PROFILE:
+		{
+			// Clamp both specular lobes
+			float3 DualRoughnessData = GetDualRoughnessSpecular(Payload);
+			DualRoughnessData = max(DualRoughnessData, float3(PathRoughness, PathRoughness, 0.0));
+			SetDualRoughnessSpecular(Payload, DualRoughnessData);
+			break;
+		}
 		default:
 			// fallback is diffuse
 			Payload.Roughness = 1.0;
@@ -199,6 +207,83 @@ void AdjustPathAndMaterialRoughness(inout FMaterialClosestHitPayload Payload, in
 	PathRoughness = lerp(PathRoughness, Payload.Roughness, GlassPercent);
 }
 
+FMaterialClosestHitPayload UnpackPathTracingPayload(FPackedMaterialClosestHitPayload PackedInput, RayDesc Ray)
+{
+	// the path tracer has a few more requirements compared to the ray-tracing path.
+	// Long term, we should reoranize the payload for the path tracer, but for now, just
+	// do a bit of post-processing and re-use fields to encode extracted information
+
+	// This function can also do a bit of per-material post-processing to simplify the rest of the code
+	FMaterialClosestHitPayload HitPayload = UnpackRayTracingPayload(PackedInput, Ray);
+
+	switch (HitPayload.ShadingModelID)
+	{
+		case SHADINGMODELID_SUBSURFACE:
+		case SHADINGMODELID_PREINTEGRATED_SKIN:
+		{
+			// TODO: is this accurate enough? hard to measure exactly since the raster path uses a very different approach
+			float SSSRadius = (1 - HitPayload.Opacity) * 10.0; // simple formula, up to 10cm radius
+			SetSubsurfaceRadius(HitPayload, SSSRadius);
+
+			// Reset opacity to avoid special cases elsewhere
+			HitPayload.Opacity = 1;
+			break;
+		}
+		case SHADINGMODELID_SUBSURFACE_PROFILE:
+		{
+			uint ProfileId = uint(HitPayload.CustomData.a * 255.0 + 0.5);
+			if (GetSubsurfaceProfileUseBurley(ProfileId))
+			{
+				// NOTE: would be better to just precompute all this on the CPU and just have a single value to lookup
+
+				// Burley parameterization
+				float3 SurfaceAlbedo = ActualSSProfilesTexture.Load(int3(BSSS_SURFACEALBEDO_OFFSET, ProfileId, 0)).xyz;
+				float3 DiffuseMeanFreePath = DecodeDiffuseMeanFreePath(ActualSSProfilesTexture.Load(int3(BSSS_DMFP_OFFSET, ProfileId, 0))).xyz;
+				float WorldUnitScale = DecodeWorldUnitScale(ActualSSProfilesTexture.Load(int3(SSSS_SUBSURFACE_COLOR_OFFSET, ProfileId, 0)).a) * BURLEY_MM_2_CM;
+
+				// Opacity acts as a per-pixel radius multiplier
+				// NOTE: this seems backwards? Opacity=0 acts like default-lit while Opacity=1 acts like SSS?
+				float3 SSSRadius = HitPayload.Opacity * WorldUnitScale * DiffuseMeanFreePath * GetDiffuseSurfaceScalingFactor3D(SurfaceAlbedo);
+
+				SetSubsurfaceRadius(HitPayload, SSSRadius);
+
+				// burley mode doesn't have any diffuse component
+				SetSubsurfaceColor(HitPayload, HitPayload.DiffuseColor);
+				HitPayload.DiffuseColor = 0;
+			}
+			else
+			{
+				// How to handle this case? a diffusion kernel is tabulated into the texture but I don't see a way to extract the original
+				// radius/color controls
+				SetSubsurfaceRadius(HitPayload, 0);
+				SetSubsurfaceColor(HitPayload, 0);
+				// TODO: non-burley profile allows a blend of diffuse and SSS which should be accounted for here
+			}
+			// Decode dual roughness info
+
+			// See logic in ShadingModels.ush
+			float4 Data = ActualSSProfilesTexture.Load(int3(SSSS_DUAL_SPECULAR_OFFSET, ProfileId, 0));
+			float AverageToRoughness0 = Data.x * SSSS_MAX_DUAL_SPECULAR_ROUGHNESS;
+			float AverageToRoughness1 = Data.y * SSSS_MAX_DUAL_SPECULAR_ROUGHNESS;
+			float LobeMix = Data.z;
+			float AverageRoughness = HitPayload.Roughness;
+			float Lobe0Roughness = max(saturate(AverageRoughness * AverageToRoughness0), 0.02f);
+			float Lobe1Roughness = saturate(AverageRoughness * AverageToRoughness1);
+
+			Lobe0Roughness = lerp(1.0f, Lobe0Roughness, saturate(HitPayload.Opacity * 10.0f));
+			Lobe1Roughness = lerp(1.0f, Lobe1Roughness, saturate(HitPayload.Opacity * 10.0f));
+
+			SetDualRoughnessSpecular(HitPayload, float3(Lobe0Roughness, Lobe1Roughness, LobeMix));
+			
+			// Reset opacity to avoid special cases elsewhere
+			HitPayload.Opacity = 1;
+			break;
+		}
+	}
+
+	return HitPayload;
+}
+
 FMaterialSample SampleMaterial(
 	float3 RayDirection,
 	FMaterialClosestHitPayload Payload,
@@ -292,6 +377,9 @@ FMaterialSample SampleMaterial(
 		case SHADINGMODELID_PREINTEGRATED_SKIN:
 			Result = DefaultLit_SampleMaterial(RayDirection, Payload, RandSample);
 			break;
+		case SHADINGMODELID_SUBSURFACE_PROFILE:
+			Result = SubsurfaceProfile_SampleMaterial(RayDirection, Payload, RandSample);
+			break;
 		case SHADINGMODELID_CLOTH:
 			Result = Cloth_SampleMaterial(RayDirection, Payload, RandSample);
 			break;
@@ -381,6 +469,9 @@ FMaterialEval EvalMaterial(
 		case SHADINGMODELID_PREINTEGRATED_SKIN:
 			Result = DefaultLit_EvalMaterial(IncomingDirection, OutgoingDirection, Payload);
 			break;
+		case SHADINGMODELID_SUBSURFACE_PROFILE:
+			Result = SubsurfaceProfile_EvalMaterial(IncomingDirection, OutgoingDirection, Payload);
+			break;
 		case SHADINGMODELID_CLOTH:
 			Result = Cloth_EvalMaterial(IncomingDirection, OutgoingDirection, Payload);
 			break;
diff --git a/Engine/Shaders/Private/PathTracing/Material/PathTracingSubsurfaceProfile.ush b/Engine/Shaders/Private/PathTracing/Material/PathTracingSubsurfaceProfile.ush
new file mode 100644
--- /dev/null
+++ b/Engine/Shaders/Private/PathTracing/Material/PathTracingSubsurfaceProfile.ush
@@ -0,0 +1,157 @@
+// Copyright Epic Games, Inc. All Rights Reserved.
+
+/*=============================================================================================
+	PathTracingSubsurfaceProfile.usf: Path tracing BRDF model for subsurface profile materials
+
+	The main difference between this and default-lit is the dual specular lobes and lack of anisotropy
+	The diffuse lobe is normally black for primary hits, but will get the subsurface color after
+	SSS simplification.
+===============================================================================================*/
+
+#pragma once
+
+#include "PathTracingMaterialCommon.ush"
+#include "PathTracingFresnel.ush"
+
+float2 SubsurfaceProfile_LobeCdf(float3 V, FMaterialClosestHitPayload Payload)
+{
+	// TODO: This could be improved if we had a directional albedo tabulation, but this formulation will at least prevent
+	// under-sampling when one lobe is black (like with metals)
+
+	float3 Diff = GetDiffuseColor(Payload);
+	float3 Spec = GetSpecularColor(Payload);
+
+	float DiffuseProb = LobeSelectionProb(Diff, Spec);
+	float SpecMix = GetDualRoughnessSpecular(Payload).z;
+	return float2(DiffuseProb, DiffuseProb + (1 - DiffuseProb) * (1 - SpecMix));
+}
+
+FMaterialEval SubsurfaceProfile_EvalMaterial(
+	float3 IncomingDirection,
+	float3 OutgoingDirection,
+	FMaterialClosestHitPayload Payload
+)
+{
+	FMaterialEval Result = NullMaterialEval();
+
+	const float3 V_World = -IncomingDirection;
+	const float3 L_World = OutgoingDirection;
+	const float3 N_World = GetWorldNormal(Payload);
+
+	float3 DualRoughnessData = GetDualRoughnessSpecular(Payload);
+	float2 Alpha0 = 0;
+	float2 Alpha1 = 0;
+	const float3x3 Basis = GetGGXBasis(DualRoughnessData.x, Payload.WorldNormal, Alpha0);
+	GetGGXBasis(DualRoughnessData.y, Payload.WorldNormal, Alpha1);
+
+	// move vectors into right shading frame
+	const float3 V = mul(Basis, V_World);
+	const float3 L = mul(Basis, L_World);
+	const float3 H = normalize(V + L);
+
+	const float NoL = saturate(L.z);
+	const float VoH = saturate(dot(V, H));
+
+	const float2 LobeCdf = SubsurfaceProfile_LobeCdf(V_World, Payload);
+	const float3 LobeProb = float3(LobeCdf.x, LobeCdf.y - LobeCdf.x, 1 - LobeCdf.y);
+
+	// Diffuse Lobe
+	AddLobeWithMIS(Result.Weight, Result.Pdf, GetDiffuseColor(Payload), NoL / PI, LobeProb.x);
+
+	// Specular lobes
+	const float2 GGXResult0 = GGXEvalReflection(L, V, H, Alpha0);
+	const float2 GGXResult1 = GGXEvalReflection(L, V, H, Alpha1);
+	const float3 F = F_Schlick(GetSpecularColor(Payload), VoH);
+	AddLobeWithMIS(Result.Weight, Result.Pdf, (1 - DualRoughnessData.z) * F * GGXResult0.x, GGXResult0.y, LobeProb.y);
+	AddLobeWithMIS(Result.Weight, Result.Pdf, (    DualRoughnessData.z) * F * GGXResult1.x, GGXResult1.y, LobeProb.z);
+
+	return Result;
+}
+
+FMaterialSample SubsurfaceProfile_SampleMaterial(
+	float3 RayDirection,
+	FMaterialClosestHitPayload Payload,
+	float4 RandSample
+)
+{
+	float3 N = GetWorldNormal(Payload);
+	float3 V = -RayDirection;
+
+	const float2 LobeCdf = SubsurfaceProfile_LobeCdf(V, Payload);
+	const float3 LobeProb = float3(LobeCdf.x, LobeCdf.y - LobeCdf.x, 1 - LobeCdf.y);
+
+	float3 DualRoughnessData = GetDualRoughnessSpecular(Payload);
+	float2 Alpha0 = 0;
+	float2 Alpha1 = 0;
+	const float3x3 Basis = GetGGXBasis(DualRoughnessData.x, Payload.WorldNormal, Alpha0);
+	GetGGXBasis(DualRoughnessData.y, Payload.WorldNormal, Alpha1);
+
+	V = mul(Basis, V);
+
+	// Randomly choose to sample diffuse or specular
+	float3 L = 0, H = 0;
+	if (RandSample.x < LobeCdf.x)
+	{
+		RandSample.x /= LobeCdf.x;
+
+		// Lambert
+		// TODO: evaluate CosineSampleHemisphereConcentric
+		float4 Result = CosineSampleHemisphere(RandSample.xy);
+
+		L = Result.xyz;
+		H = normalize(L + V);
+	}
+	else if (RandSample.x < LobeCdf.y)
+	{
+		RandSample.x -= LobeCdf.x;
+		RandSample.x /= LobeCdf.y - LobeCdf.x;
+
+
+		// #dxr_todo: Evaluate UniformSampleDiskConcentric
+		H = ImportanceSampleVisibleGGX_aniso(UniformSampleDisk(RandSample.xy), Alpha0, V).xyz;
+
+		L = reflect(-V, H);
+		if (L.z <= 0)
+		{
+			// invalid output direction, exit early
+			return NullMaterialSample();
+		}
+	}
+	else
+	{
+		RandSample.x -= LobeCdf.y;
+		RandSample.x /= 1 - LobeCdf.y;
+
+
+		// #dxr_todo: Evaluate UniformSampleDiskConcentric
+		H = ImportanceSampleVisibleGGX_aniso(UniformSampleDisk(RandSample.xy), Alpha1, V).xyz;
+
+		L = reflect(-V, H);
+		if (L.z <= 0)
+		{
+			// invalid output direction, exit early
+			return NullMaterialSample();
+		}
+	}
+
+	// With a valid direction in hand -- now evaluate the BxDF (taking advantage of already computed terms)
+
+	const float NoL = saturate(L.z);
+	const float VoH = saturate(dot(V, H));
+
+	float3 OutWeight = 0;
+	float OutPdf = 0;
+
+	// Diffuse Lobe
+	AddLobeWithMIS(OutWeight, OutPdf, GetDiffuseColor(Payload), NoL / PI, LobeProb.x);
+
+	// Specular lobe
+	const float2 GGXResult0 = GGXEvalReflection(L, V, H, Alpha0);
+	const float2 GGXResult1 = GGXEvalReflection(L, V, H, Alpha1);
+	const float3 F = F_Schlick(GetSpecularColor(Payload), VoH);
+	AddLobeWithMIS(OutWeight, OutPdf, (1 - DualRoughnessData.z) * F * GGXResult0.x, GGXResult0.y, LobeProb.y);
+	AddLobeWithMIS(OutWeight, OutPdf, (    DualRoughnessData.z) * F * GGXResult1.x, GGXResult1.y, LobeProb.z);
+
+	// transform to world space
+	return CreateMaterialSample(normalize(mul(L, Basis)), OutWeight, OutPdf, 1.0);
+}
diff --git a/Engine/Shaders/Private/PathTracing/PathTracing.usf b/Engine/Shaders/Private/PathTracing/PathTracing.usf
--- a/Engine/Shaders/Private/PathTracing/PathTracing.usf
+++ b/Engine/Shaders/Private/PathTracing/PathTracing.usf
@@ -91,7 +91,7 @@ FMaterialClosestHitPayload TraceTransparentRay(RayDesc Ray, FRayCone RayCone, bo
 		}
 
 		// Unpack the payload
-		FMaterialClosestHitPayload HitPayload = UnpackRayTracingPayload(PackedPayload, Ray);
+		FMaterialClosestHitPayload HitPayload = UnpackPathTracingPayload(PackedPayload, Ray);
 
 		if (HitPayload.BlendingMode == RAY_TRACING_BLEND_MODE_ADDITIVE ||
 			HitPayload.BlendingMode == RAY_TRACING_BLEND_MODE_TRANSLUCENT)
@@ -115,7 +115,6 @@ FMaterialClosestHitPayload TraceTransparentRay(RayDesc Ray, FRayCone RayCone, bo
 			float SelectionWeight = max3(Contrib.x, Contrib.y, Contrib.z);
 			SelectionWeightSum += SelectionWeight;
 			bool AcceptHit = false;
-			float PickWeight = 1;
 
 			// weighted reservoir sampling
 			if (SelectionWeight > 0)
@@ -126,7 +125,6 @@ FMaterialClosestHitPayload TraceTransparentRay(RayDesc Ray, FRayCone RayCone, bo
 					// generate a random number to see if we should accept this hit
 					float RandValue = RandomSequence_GenerateSample1D(RandSequence);
 					AcceptHit = RandValue * SelectionWeightSum < SelectionWeight;
-					PickWeight = SelectionWeightSum / SelectionWeight;
 				}
 				else
 				{
@@ -139,7 +137,7 @@ FMaterialClosestHitPayload TraceTransparentRay(RayDesc Ray, FRayCone RayCone, bo
 			{
 				// stash this hit for next time
 				Payload = HitPayload;
-				PayloadThroughput = PathThroughput * PickWeight;
+				PayloadThroughput = PathThroughput / SelectionWeight;
 			}
 		}
 
@@ -157,7 +155,7 @@ FMaterialClosestHitPayload TraceTransparentRay(RayDesc Ray, FRayCone RayCone, bo
 	if (SelectionWeightSum > 0)
 	{
 		// if we stored a valid hit in the payload, reset the path throughput to this point
-		PathThroughput = PayloadThroughput;
+		PathThroughput = PayloadThroughput * SelectionWeightSum;
 	}
 	else
 	{
@@ -248,25 +246,28 @@ FProbeResult TraceProbeRay(RayDesc Ray, uint2 PixelCoord)
 
 bool ProcessSubsurfaceRandomWalk(inout FMaterialClosestHitPayload Payload, inout float3 PathThroughput, inout RandomSequence RandSequence, uint2 PixelCoord, bool SimplifySSS)
 {
-	if (!IsSubsurfaceMaterial(Payload) || PathTracingData.MaxBounces == 0)
+	if (!IsSubsurfaceMaterial(Payload))
 	{
-		// material doesn't have SSS or we don't have any bounces configured
+		// material doesn't have SSS
 		// just return so we can process the rest of the payload's shading
 		return true;
 	}
 
-	// TODO: support RGB extinction as well (SHADINGMODELID_SUBSURFACE_PROFILE)
-	// TODO: figure out the right multiplier between the 0-1 Opacity control and world space radius
-	float3 SSSColor = Payload.CustomData.xyz;
-	float SSSRadius = (1 - Payload.Opacity) * 10.0;
+	float3 SSSColor = GetSubsurfaceColor(Payload);
+	float3 SSSRadius = GetSubsurfaceRadius(Payload);
 
-	if (SimplifySSS || SSSRadius == 0)
+	if (SimplifySSS || all(SSSRadius == 0) || all(SSSColor == 0) || PathTracingData.MaxSSSBounces == 0)
 	{
 		// if we are running in a context that won't see the SSS directly -- just turn it into diffuse and skip the random walk
+		// we can also skip random walk if the radius or color is 0
+		// or if the user decided not to perform any SSS bounces
 		Payload.DiffuseColor += SSSColor;
 		Payload.DiffuseColor = saturate(Payload.DiffuseColor);
+		SetSubsurfaceColor(Payload, 0.0);
 		return true;
 	}
+	SSSRadius = max(SSSRadius, 0.1); // 1mm minimum to make sure random walk can move forward and to keep SigmaT finite
+
 	// decide randomly to evaluate the SSS part of the material, or keep shading the current part
 	float SSSProb = LobeSelectionProb(SSSColor, Payload.DiffuseColor + Payload.SpecularColor);
 	float3 RandSample = RandomSequence_GenerateSample3D(RandSequence);
@@ -297,22 +298,49 @@ bool ProcessSubsurfaceRandomWalk(inout FMaterialClosestHitPayload Payload, inout
 
 #if SSS_USE_TIR
 	const float3 Albedo = 1 - exp(SSSColor * (-11.43 + SSSColor * (15.38 - 13.91 * SSSColor)));
-	SSSRadius *= 0.5; // roughly match parameterization below
 #else
 	// Van de-Hulst inverse mapping
 	// https://blog.selfshadow.com/publications/s2017-shading-course/imageworks/s2017_pbs_imageworks_slides_v2.pdf (Slide 44)
 	// [REDACTED_PROFILE_URL] (Section 7.5.3 of v0.1.3)
 	const float3 Albedo = 1 - Pow2(4.09712 + 4.20863 * SSSColor - sqrt(9.59217 + SSSColor * (41.6808 + 17.7126 * SSSColor)));
+	SSSRadius *= 2.0; // roughly match parameterization above
 #endif
 	const int MAX_SSS_BOUNCES = PathTracingData.MaxSSSBounces;
+	const float3 SigmaT = 1.0 / SSSRadius;
+	const float3 SigmaS = Albedo * SigmaT;
 	for (int i = 0; i < MAX_SSS_BOUNCES; i++)
 	{
-		// Using an analog random walk - transmittance and scattering probabilities cancel out so only the albedo
-		// at each bounce needs to be accounted for in the throughput. This will get more complicated once we add
-		// proper support for RGB transmittance and fancier escape strategies like Dwivedi sampling.
-
+		// TODO: add Dwivedi sampling for faster escape in low albedo media ...
+
+		// Pick a color channel for transmittance sampling:
+		float3 ColorChannelWeight = PathThroughput * Albedo;
+		float3 ColorChannelCdf = float3(ColorChannelWeight.x,
+										ColorChannelWeight.x + ColorChannelWeight.y,
+										ColorChannelWeight.x + ColorChannelWeight.y + ColorChannelWeight.z);
+		if (ColorChannelCdf.z == 0)
+		{
+			// no more energy left in the path
+			break;
+		}
+		ColorChannelCdf.xy /= ColorChannelCdf.z;
+		float3 ColorChannelPdf = float3(ColorChannelCdf.x, ColorChannelCdf.y - ColorChannelCdf.x, 1 - ColorChannelCdf.y);
 		// pick a maximum distance to scatter
-		Ray.TMax = -log(1 - RandSample.z) * SSSRadius;
+		const float OneMinusEpsilon = 0.99999994; // 32-bit float just before 1.0
+		if (RandSample.z < ColorChannelCdf.x)
+		{
+			float RescaleRand = min(RandSample.z / ColorChannelPdf.x, OneMinusEpsilon);
+			Ray.TMax = -log(1 - RescaleRand) * SSSRadius.x;
+		}
+		else if (RandSample.z < ColorChannelCdf.y)
+		{
+			float RescaleRand = min((RandSample.z - ColorChannelCdf.x) / ColorChannelPdf.y, OneMinusEpsilon);
+			Ray.TMax = -log(1 - RescaleRand) * SSSRadius.y;
+		}
+		else
+		{
+			float RescaleRand = min((RandSample.z - ColorChannelCdf.y) / ColorChannelPdf.z, OneMinusEpsilon);
+			Ray.TMax = -log(1 - RescaleRand) * SSSRadius.z;
+		}
 
 		FProbeResult ProbeResult = TraceProbeRay(Ray, PixelCoord);
 		if (ProbeResult.IsMiss())
@@ -321,10 +349,24 @@ bool ProcessSubsurfaceRandomWalk(inout FMaterialClosestHitPayload Payload, inout
 			RandSample = RandomSequence_GenerateSample3D(RandSequence);
 			Ray.Origin += Ray.TMax * Ray.Direction;
 			Ray.Direction = UniformSampleSphere(RandSample.xy).xyz;
-			PathThroughput *= Albedo;
+			// account for transmittance and scattering up to this point
+			float3 Transmittance = exp(-Ray.TMax * SigmaT);
+			PathThroughput *= SigmaS * Transmittance;
+			// account for transmittance pdf (which is the balance heuristic combination of transmittance sampling in each color channel)
+			float3 TransmittancePdf = SigmaT * Transmittance;
+			PathThroughput /= dot(TransmittancePdf, ColorChannelPdf);
+
+			// NOTE: phase function cancels out since it is being perfectly importance sampled
 		}
 		else
 		{
+			// account for transmittance and scattering up to this point
+			float3 Transmittance = exp(-ProbeResult.HitT * SigmaT);
+			PathThroughput *= Transmittance;
+			// account for transmittance pdf (which is the balance heuristic combination of transmittance sampling in each color channel)
+			float3 TransmittancePdf = Transmittance; // integral of pdf from HitT to infinity
+			PathThroughput /= dot(Transmittance, ColorChannelPdf);
+
 			// our short ray hit the geomtry - decide if we should exit or not
 #if SSS_USE_TIR
 			RandSample = RandomSequence_GenerateSample2D(RandSequence).xxy;
diff --git a/Engine/Shaders/Private/RayTracing/RayTracingMaterialHitShaders.usf b/Engine/Shaders/Private/RayTracing/RayTracingMaterialHitShaders.usf
--- a/Engine/Shaders/Private/RayTracing/RayTracingMaterialHitShaders.usf
+++ b/Engine/Shaders/Private/RayTracing/RayTracingMaterialHitShaders.usf
@@ -384,7 +384,11 @@ RAY_TRACING_ENTRY_CLOSEST_HIT(MaterialCHS,
 
 	// Store the results in local variables and reuse instead of calling the functions multiple times.
 	half3 BaseColor = GetMaterialBaseColor(PixelMaterialInputs);
+#if REFRACTION_USE_INDEX_OF_REFRACTION
 	half2 Refraction = GetMaterialRefraction(PixelMaterialInputs);
+#else
+	half2 Refraction = 1;
+#endif
 	half  Metallic = GetMaterialMetallic(PixelMaterialInputs);
 	half  Specular = GetMaterialSpecular(PixelMaterialInputs);
 	half Roughness = GetMaterialRoughness(PixelMaterialInputs);
@@ -413,6 +417,7 @@ RAY_TRACING_ENTRY_CLOSEST_HIT(MaterialCHS,
 	Payload.Ior = Refraction.x;
 	Payload.CustomData = float4(GetMaterialCustomData0(MaterialParameters), GetMaterialCustomData1(MaterialParameters), 0, 0);
 	
+#if MATERIAL_SHADINGMODEL_CLEAR_COAT
 	#if CLEAR_COAT_BOTTOM_NORMAL
 	if (Payload.ShadingModelID == SHADINGMODELID_CLEAR_COAT)
 	{
@@ -435,19 +440,26 @@ RAY_TRACING_ENTRY_CLOSEST_HIT(MaterialCHS,
 		#endif
 	}
 	#endif
+#endif
 
 	// Override custom data if sub-surface material
+#if MATERIAL_SHADINGMODEL_SUBSURFACE || MATERIAL_SHADINGMODEL_TWOSIDED_FOLIAGE || MATERIAL_SHADINGMODEL_CLOTH || MATERIAL_SHADINGMODEL_CLOTH || MATERIAL_SHADINGMODEL_SUBSURFACE_PROFILE || MATERIAL_SHADINGMODEL_EYE
 	if (Payload.ShadingModelID == SHADINGMODELID_SUBSURFACE ||
 		Payload.ShadingModelID == SHADINGMODELID_TWOSIDED_FOLIAGE ||
-		Payload.ShadingModelID == SHADINGMODELID_CLOTH)
+		Payload.ShadingModelID == SHADINGMODELID_CLOTH ||
+		Payload.ShadingModelID == SHADINGMODELID_SUBSURFACE_PROFILE |
+		Payload.ShadingModelID == SHADINGMODELID_EYE)
 	{
 		Payload.CustomData = GetMaterialSubsurfaceData(PixelMaterialInputs);
 	}
+#endif
 
+#if MATERIAL_SHADINGMODEL_CLOTH
 	if (Payload.ShadingModelID == SHADINGMODELID_CLOTH)
 	{
 		Payload.CustomData.a = saturate(GetMaterialCustomData0(MaterialParameters));
 	}
+#endif
 
 #if MATERIAL_SHADINGMODEL_THIN_TRANSLUCENT
 	if (Payload.ShadingModelID == SHADINGMODELID_THIN_TRANSLUCENT)
@@ -639,7 +651,11 @@ RAY_TRACING_ENTRY_ANY_HIT(MaterialAHS,
 	// Assume Payload.ShadingModelID == SHADINGMODELID_THIN_TRANSLUCENT
 	float3 Transmission = GetThinTranslucentMaterialOutput0(MaterialParameters);
 	float Opacity = GetMaterialOpacity(PixelMaterialInputs);
+#if REFRACTION_USE_INDEX_OF_REFRACTION
 	float Ior = GetMaterialRefraction(PixelMaterialInputs).x;
+#else
+	float Ior = 1.0;
+#endif
 	float3 Transparency = 1 - Opacity;
 	float3 V = WorldRayDirection();
 	float3 N = normalize(MaterialParameters.WorldNormal);
diff --git a/Engine/Source/Runtime/Renderer/Private/PathTracing.cpp b/Engine/Source/Runtime/Renderer/Private/PathTracing.cpp
--- a/Engine/Source/Runtime/Renderer/Private/PathTracing.cpp
+++ b/Engine/Source/Runtime/Renderer/Private/PathTracing.cpp
@@ -367,9 +367,11 @@ class FPathTracingRG : public FGlobalShader
 		SHADER_PARAMETER(int32, SkylightMipCount)
 		// IES Profiles
 		SHADER_PARAMETER_RDG_TEXTURE(Texture2DArray, IESTexture)
-		SHADER_PARAMETER_SAMPLER(SamplerState, IESTextureSampler)
-		// Shared sampler for all IES profiles
-		SHADER_PARAMETER(FIntVector, TileOffset)	// Used by multi-GPU rendering
+		SHADER_PARAMETER_SAMPLER(SamplerState, IESTextureSampler) // Shared sampler for all IES profiles
+		// Subsurface data
+		SHADER_PARAMETER_TEXTURE(Texture2D, SSProfilesTexture)
+		// Used by multi-GPU rendering
+		SHADER_PARAMETER(FIntVector, TileOffset)	
 	END_SHADER_PARAMETER_STRUCT()
 };
 IMPLEMENT_GLOBAL_SHADER(FPathTracingRG, "/Engine/Private/PathTracing/PathTracing.usf", "PathTracingMainRG", SF_RayGen);
@@ -845,6 +847,8 @@ void FDeferredShadingSceneRenderer::RenderPathTracing(
 		PassParameters->IESTextureSampler = TStaticSamplerState<SF_Bilinear, AM_Clamp, AM_Clamp, AM_Clamp>::GetRHI();
 		PassParameters->RadianceTexture = GraphBuilder.CreateUAV(RadianceTexture);
 
+		PassParameters->SSProfilesTexture = GetSubsufaceProfileTexture_RT(GraphBuilder.RHICmdList)->GetShaderResourceRHI();
+
 		// TODO: in multi-gpu case, split image into tiles
 		PassParameters->TileOffset.X = 0;
 		PassParameters->TileOffset.Y = 0;
