Fix shading artifacts at grazing angles in path tracer

Instead of doing the "faceforward" calculation based on the shading normal alone, do it based on the hit kind during material evaluation. This avoids inconsistencies between single and two-sided materials and simplifies the logic in the path tracer.

Some details of random walk subsurface exit handling had to be tweaked to compensate for this change.

diff --git a/Engine/Shaders/Private/MaterialTemplate.ush b/Engine/Shaders/Private/MaterialTemplate.ush
--- a/Engine/Shaders/Private/MaterialTemplate.ush
+++ b/Engine/Shaders/Private/MaterialTemplate.ush
@@ -2860,9 +2860,12 @@ void CalcMaterialParametersEx(
 	    Parameters.TwoSidedSign *= ResolvedView.CullingSign * GetPrimitiveData(Parameters.PrimitiveId).InvNonUniformScaleAndDeterminantSign.w;
     #endif
 
-	#if !MATERIAL_TWOSIDED_SEPARATE_PASS
-		Parameters.TwoSidedSign *= GetFloatFacingSign(bIsFrontFace);
-	#endif
+#endif
+
+#if (MATERIAL_TWOSIDED && !MATERIAL_TWOSIDED_SEPARATE_PASS) || RAYHITGROUPSHADER
+	// Either we have a two-sided material that needs a sign flip, or we have a ray tracing material
+	// that needs to consider rays arriving from either side
+	Parameters.TwoSidedSign *= GetFloatFacingSign(bIsFrontFace);
 #endif
 
 #if NUM_VIRTUALTEXTURE_SAMPLES || LIGHTMAP_VT_ENABLED
diff --git a/Engine/Shaders/Private/PathTracing/PathTracing.usf b/Engine/Shaders/Private/PathTracing/PathTracing.usf
--- a/Engine/Shaders/Private/PathTracing/PathTracing.usf
+++ b/Engine/Shaders/Private/PathTracing/PathTracing.usf
@@ -208,6 +208,7 @@ struct FProbeResult
 {
 	float HitT;
 	float3 WorldNormal;
+	int FrontFace;
 
 	bool IsMiss() { return HitT <= 0; }
 };
@@ -241,6 +242,7 @@ FProbeResult TraceProbeRay(RayDesc Ray, uint2 PixelCoord)
 	FProbeResult Result;
 	Result.HitT = PackedPayload.HitT;
 	Result.WorldNormal = PackedPayload.GetWorldNormal();
+	Result.FrontFace = (PackedPayload.GetFlags() & RAY_TRACING_PAYLOAD_OUTPUT_FLAG_FRONT_FACE) != 0;
 	return Result;
 }
 
@@ -379,13 +381,18 @@ bool ProcessSubsurfaceRandomWalk(inout FMaterialClosestHitPayload Payload, inout
 				// NOTE: weight and probability cancel out, so no throughput adjustment is needed
 				Ray.Origin += ProbeResult.HitT * Ray.Direction;
 				Ray.Direction = reflect(Ray.Direction, WorldNormal);
-				ApplyPositionBias(Ray, -WorldNormal, PathTracingData.MaxNormalBias);
+				ApplyPositionBias(Ray, WorldNormal, PathTracingData.MaxNormalBias);
 				continue;
 			}
 #endif
 			// we hit the boundary! overwrite the current payload and exit the walk with a diffuse scattering event
 			Payload.WorldPos = Ray.Origin + ProbeResult.HitT * Ray.Direction;
 			Payload.WorldNormal = ProbeResult.WorldNormal;
+			// make sure normal is pointed outward from the object so we capture illumination from the exterior
+			if (ProbeResult.FrontFace != Payload.IsFrontFace())
+			{
+				Payload.WorldNormal *= -1.0;
+			}
 			Payload.ShadingModelID = SHADINGMODELID_NUM; // invalid value so that we get diffuse shading
 			Payload.BlendingMode = RAY_TRACING_BLEND_MODE_OPAQUE;
 			Payload.BaseColor = 1;
@@ -469,13 +476,6 @@ RAY_TRACING_ENTRY_RAYGEN(PathTracingMainRG)
 			//   https://www.arnoldrenderer.com/research/Arnold_TOG2018.pdf (section 4.2)
 			AdjustPathAndMaterialRoughness(Payload, PathRoughness);
 		}
-		
-
-		// Make sure the surface normal points toward the viewer
-		if (dot(Ray.Direction, Payload.WorldNormal) > 0.0)
-		{
-			Payload.WorldNormal = -Payload.WorldNormal;
-		}
 
 		// Processing the random walk will (stochastically) move the shading point elsewhere on the surface of the object
 		if (!ProcessSubsurfaceRandomWalk(Payload, PathThroughput, RandSequence, LaunchIndex, bSimplifySSS))
