Texture support for Rect lights in the path tracer

This patch just adds the basic support for evaluating the texture, it is not taken into account in the importance sampling. Also, currently limited to 8 slots like the RT path until we get proper bindless support.

diff --git a/Engine/Shaders/Private/PathTracing/Light/PathTracingRectLight.ush b/Engine/Shaders/Private/PathTracing/Light/PathTracingRectLight.ush
--- a/Engine/Shaders/Private/PathTracing/Light/PathTracingRectLight.ush
+++ b/Engine/Shaders/Private/PathTracing/Light/PathTracingRectLight.ush
@@ -9,6 +9,40 @@
 #include "../../RectLight.ush"
 #include "PathTracingLightCommon.ush"
 
+#if USE_RECT_LIGHT_TEXTURES
+
+// The number of entries must match PATHTRACER_MAX_RECT_TEXTURES
+Texture2D RectLightTexture_0;
+Texture2D RectLightTexture_1;
+Texture2D RectLightTexture_2;
+Texture2D RectLightTexture_3;
+Texture2D RectLightTexture_4;
+Texture2D RectLightTexture_5;
+Texture2D RectLightTexture_6;
+Texture2D RectLightTexture_7;
+SamplerState RectLightSampler;
+
+float3 EvaluateTexture(int LightId, float2 UV)
+{
+	// TODO: could use path roughness / ray cone to lower the mip level here
+	UV.y = 1 - UV.y; // match orientation with raster path
+	switch (SceneLights[LightId].RectLightTextureIndex)
+	{
+		case 0:  return RectLightTexture_0.SampleLevel(RectLightSampler, UV, 0);
+		case 1:  return RectLightTexture_1.SampleLevel(RectLightSampler, UV, 0);
+		case 2:  return RectLightTexture_2.SampleLevel(RectLightSampler, UV, 0);
+		case 3:  return RectLightTexture_3.SampleLevel(RectLightSampler, UV, 0);
+		case 4:  return RectLightTexture_4.SampleLevel(RectLightSampler, UV, 0);
+		case 5:  return RectLightTexture_5.SampleLevel(RectLightSampler, UV, 0);
+		case 6:  return RectLightTexture_6.SampleLevel(RectLightSampler, UV, 0);
+		case 7:  return RectLightTexture_7.SampleLevel(RectLightSampler, UV, 0);
+		default: return 1;
+	}
+}
+
+#endif
+
+
 float3 BilinearQuadWarp(float2 uv, float W00, float W01, float W10, float W11) {
 	// "Practical Product Sampling by Fitting and Composing Warps" - EGSR 2020
 	// https://casual-effects.com/research/Hart2020Sampling/index.html
@@ -55,7 +89,6 @@ float4 SampleApproxProjectionSphericalRect(float2 Rand, FSphericalRect Spherical
 	return float4(Direction, OutPdf * WarpedRand.z);
 }
 
-
 FLightHit RectLight_TraceLight(RayDesc Ray, int LightId)
 {
 	float3 LightPosition = GetPosition(LightId);
@@ -86,13 +119,16 @@ FLightHit RectLight_TraceLight(RayDesc Ray, int LightId)
 				true);
 			P = t * Ray.Direction - Rect.Origin;
 			// test again with the clipped extents
-			UV = float2(dot(P, LightdPdu), dot(P, LightdPdv));
-			if (all(abs(UV) <= Rect.Extent))
+			float2 ClippedUV = float2(dot(P, LightdPdu), dot(P, LightdPdv));
+			if (all(abs(ClippedUV) <= Rect.Extent))
 			{
 				// stored color is radiance
 				float3 Radiance = GetColor(LightId) * ComputeIESAttenuation(LightId, Ray.Origin);
 				Radiance *= ComputeAttenuationFalloff(dot(LightDirection, LightDirection), LightId);
 				float SolidAngle = BuildSphericalRect(Rect).SolidAngle;
+#if USE_RECT_LIGHT_TEXTURES
+				Radiance *= EvaluateTexture(LightId, (UV + LightExtent) / (LightExtent * 2));
+#endif
 				return CreateLightHit(Radiance, 1.0 / SolidAngle, t);
 			}
 		}
@@ -126,10 +162,20 @@ FLightSample RectLight_SampleLight(
 		true /* bComputeVisibleRect */);
 	FSphericalRect SphericalRect = BuildSphericalRect(Rect);
 
+#if USE_RECT_LIGHT_TEXTURES
+	FSphericalRectSample Sample = UniformSampleSphericalRectWithUV(RandSample.yz, SphericalRect);
+	float2 ClippedUV = Sample.UV;
+	float3 Direction = Sample.Direction;
+#else
 	float3 Direction = UniformSampleSphericalRect(RandSample.yz, SphericalRect);
+#endif
 
 	float3 RadianceOverPdf = GetColor(LightId) * SphericalRect.SolidAngle * ComputeIESAttenuation(LightId, WorldPos);
 	RadianceOverPdf *= ComputeAttenuationFalloff(dot(LightDirection, LightDirection), LightId);
+#if USE_RECT_LIGHT_TEXTURES
+	float2 UV = 0.5 * ((2 * ClippedUV - 1) * Rect.Extent + Rect.Offset) / Rect.FullExtent + 0.5;
+	RadianceOverPdf *= EvaluateTexture(LightId, UV);
+#endif
 
 	return CreateLightSample(RadianceOverPdf, 1.0 / SphericalRect.SolidAngle, normalize(Direction), length(Direction));
 }
diff --git a/Engine/Shaders/Private/RectLight.ush b/Engine/Shaders/Private/RectLight.ush
--- a/Engine/Shaders/Private/RectLight.ush
+++ b/Engine/Shaders/Private/RectLight.ush
@@ -584,23 +584,37 @@ FSphericalRect BuildSphericalRect( FRect Rect )
 	return SphericalRect;
 }
 
-float3 UniformSampleSphericalRect( float2 E, FSphericalRect Rect )
+struct FSphericalRectSample {
+	float3 Direction;
+	float2 UV;
+};
+
+FSphericalRectSample UniformSampleSphericalRectWithUV(float2 E, FSphericalRect Rect)
 {
 	float au = E.x * Rect.SolidAngle + Rect.k;
-	float fu = ( cos(au) * Rect.b0 - Rect.b1 ) / sin(au);
-	float cu = rsqrt( fu*fu + Rect.b0 * Rect.b0 ) * ( fu > 0 ? 1 : -1 );
+	float fu = (cos(au) * Rect.b0 - Rect.b1) / sin(au);
+	float cu = rsqrt(fu * fu + Rect.b0 * Rect.b0) * (fu > 0 ? 1 : -1);
 	cu = clamp(cu, -1, 1);				// avoid NaNs
 
-	float xu = -(cu * Rect.z0) * rsqrt(1 - cu*cu);
-	xu = clamp( xu, Rect.x0, Rect.x1 );	// avoid Infs
+	float xu = -(cu * Rect.z0) * rsqrt(1 - cu * cu);
+	xu = clamp(xu, Rect.x0, Rect.x1);	// avoid Infs
 
-	float d = sqrt( xu*xu + Rect.z0 * Rect.z0 );
-	float h0 = Rect.y0 * rsqrt( d*d + Rect.y0 * Rect.y0 );
-	float h1 = Rect.y1 * rsqrt( d*d + Rect.y1 * Rect.y1 );
+	float d = sqrt(xu * xu + Rect.z0 * Rect.z0);
+	float h0 = Rect.y0 * rsqrt(d * d + Rect.y0 * Rect.y0);
+	float h1 = Rect.y1 * rsqrt(d * d + Rect.y1 * Rect.y1);
 	float hv = h0 + E.y * (h1 - h0);
-	float yv = (hv*hv < 1 - 1e-6) ? (hv * d) * rsqrt( 1 - hv*hv ) : Rect.y1;
+	float yv = (hv * hv < 1 - 1e-6) ? (hv * d) * rsqrt(1 - hv * hv) : Rect.y1;
+
+	FSphericalRectSample Result;
+	Result.Direction = mul(float3(xu, yv, Rect.z0), Rect.Axis);
+	Result.UV = float2(xu - Rect.x0, yv - Rect.y0) / float2(Rect.x1 - Rect.x0, Rect.y1 - Rect.y0);
 
-	return mul( float3( xu, yv, Rect.z0 ), Rect.Axis );
+	return Result;
+}
+
+float3 UniformSampleSphericalRect( float2 E, FSphericalRect Rect )
+{
+	return UniformSampleSphericalRectWithUV(E, Rect).Direction;
 }
 
 float3 RectIrradianceRef( float3 N, FRect Rect, out float BaseIrradiance, out float NoL )
diff --git a/Engine/Shaders/Shared/RayTracingDefinitions.h b/Engine/Shaders/Shared/RayTracingDefinitions.h
--- a/Engine/Shaders/Shared/RayTracingDefinitions.h
+++ b/Engine/Shaders/Shared/RayTracingDefinitions.h
@@ -35,6 +35,8 @@
 #define RAY_TRACING_MAX_ALLOWED_ATTRIBUTE_SIZE  8   // Sizeof 2 floats (barycentrics)
 #define RAY_TRACING_MAX_ALLOWED_PAYLOAD_SIZE    64  // Our maximum allowed payload size (sizeof FPackedMaterialClosestHitPayload)
 
+#define PATHTRACER_MAX_RECT_TEXTURES		8
+
 // Constants for the 'Flags' field of FPathTracingLight
 #define PATHTRACER_FLAG_LIGHTING_CHANNEL_MASK				(7 << 0)	// Which lighting channel is this light assigned to?
 #define PATHTRACER_FLAG_TRANSMISSION_MASK					(1 << 3)	// Does the light affect the transmission side?
diff --git a/Engine/Shaders/Shared/RayTracingTypes.h b/Engine/Shaders/Shared/RayTracingTypes.h
--- a/Engine/Shaders/Shared/RayTracingTypes.h
+++ b/Engine/Shaders/Shared/RayTracingTypes.h
@@ -23,8 +23,9 @@ struct FPathTracingLight {
 	float   FalloffExponent; // for non-inverse square decay lights only
 	int32   IESTextureSlice;
 	uint32  Flags; // see defines above
+	int32   RectLightTextureIndex;
 
-	float Padding[8]; // keep structure aligned - take new fields from here as needed
+	float Padding[7]; // keep structure aligned - take new fields from here as needed
 };
 
 static_assert(sizeof(FPathTracingLight) == 128, "Path tracing light structure should be aligned to 128 bytes for optimal access on the GPU");
@@ -45,8 +46,9 @@ struct FPathTracingLight {
 	float   FalloffExponent;
 	int     IESTextureSlice;
 	uint    Flags;
+	int	    RectLightTextureIndex;
 
-	float Padding[8];
+	float Padding[7];
 };
 
 #endif
diff --git a/Engine/Source/Runtime/Renderer/Private/PathTracing.cpp b/Engine/Source/Runtime/Renderer/Private/PathTracing.cpp
--- a/Engine/Source/Runtime/Renderer/Private/PathTracing.cpp
+++ b/Engine/Source/Runtime/Renderer/Private/PathTracing.cpp
@@ -347,6 +347,7 @@ class FPathTracingRG : public FGlobalShader
 	{
 		//OutEnvironment.CompilerFlags.Add(CFLAG_WarningsAsErrors);
 		OutEnvironment.SetDefine(TEXT("USE_NEW_SKYDOME"), 1);
+		OutEnvironment.SetDefine(TEXT("USE_RECT_LIGHT_TEXTURES"), 1);
 	}
 
 
@@ -368,6 +369,9 @@ class FPathTracingRG : public FGlobalShader
 		// IES Profiles
 		SHADER_PARAMETER_RDG_TEXTURE(Texture2DArray, IESTexture)
 		SHADER_PARAMETER_SAMPLER(SamplerState, IESTextureSampler) // Shared sampler for all IES profiles
+		// Rect lights
+		SHADER_PARAMETER_TEXTURE_ARRAY(Texture2D, RectLightTexture, [PATHTRACER_MAX_RECT_TEXTURES])
+		SHADER_PARAMETER_SAMPLER(SamplerState, RectLightSampler) // Shared sampler for all rectlights
 		// Subsurface data
 		SHADER_PARAMETER_TEXTURE(Texture2D, SSProfilesTexture)
 		// Used by multi-GPU rendering
@@ -538,6 +542,8 @@ void SetLightParameters(FRDGBuilder& GraphBuilder, FPathTracingRG::FParameters*
 		LightCount++;
 	}
 
+	int32 NextRectTextureIndex = 0;
+
 	TMap<FTexture*, int> IESLightProfilesMap;
 	for (auto Light : Scene->Lights)
 	{
@@ -565,6 +571,7 @@ void SetLightParameters(FRDGBuilder& GraphBuilder, FPathTracingRG::FParameters*
 		DestLight.Flags |= LightingChannelMask & PATHTRACER_FLAG_LIGHTING_CHANNEL_MASK;
 		DestLight.Flags |= Light.LightSceneInfo->Proxy->CastsDynamicShadow() ? PATHTRACER_FLAG_CAST_SHADOW_MASK : 0;
 		DestLight.IESTextureSlice = -1;
+		DestLight.RectLightTextureIndex = -1;
 
 		if (View.Family->EngineShowFlags.TexturedLightProfiles)
 		{
@@ -601,6 +608,34 @@ void SetLightParameters(FRDGBuilder& GraphBuilder, FPathTracingRG::FParameters*
 				DestLight.FalloffExponent = LightParameters.FalloffExponent;
 				DestLight.Flags |= Light.LightSceneInfo->Proxy->IsInverseSquared() ? 0 : PATHTRACER_FLAG_NON_INVERSE_SQUARE_FALLOFF_MASK;
 				DestLight.Flags |= PATHTRACING_LIGHT_RECT;
+				if (Light.LightSceneInfo->Proxy->HasSourceTexture())
+				{
+					// there is an actual texture associated with this light, go look for it
+					FLightShaderParameters ShaderParameters;
+					Light.LightSceneInfo->Proxy->GetLightShaderParameters(ShaderParameters);
+					FRHITexture* TextureRHI = ShaderParameters.SourceTexture;
+					if (TextureRHI != nullptr)
+					{
+						// have we already given this texture an index?
+						// NOTE: linear search is ok since max texture is small
+						for (int Index = 0; Index < NextRectTextureIndex; Index++)
+						{
+							if (PassParameters->RectLightTexture[Index] == TextureRHI)
+							{
+								DestLight.RectLightTextureIndex = Index;
+								break;
+							}
+						}
+						if (DestLight.RectLightTextureIndex == -1 && NextRectTextureIndex < PATHTRACER_MAX_RECT_TEXTURES)
+						{
+							// first time we see this texture and we still have free slots available
+							// assign texture to next slot and store it in the light
+							DestLight.RectLightTextureIndex = NextRectTextureIndex;
+							PassParameters->RectLightTexture[NextRectTextureIndex] = TextureRHI;
+							NextRectTextureIndex++;
+						}
+					}
+				}
 				break;
 			}
 			case LightType_Spot:
@@ -631,6 +666,15 @@ void SetLightParameters(FRDGBuilder& GraphBuilder, FPathTracingRG::FParameters*
 		LightCount++;
 	}
 
+	{
+		// assign dummy textures to the remaining unused slots
+		for (int32 Index = NextRectTextureIndex; Index < PATHTRACER_MAX_RECT_TEXTURES; Index++)
+		{
+			PassParameters->RectLightTexture[Index] = GWhiteTexture->TextureRHI;
+		}
+		PassParameters->RectLightSampler = TStaticSamplerState<SF_Bilinear, AM_Clamp, AM_Clamp, AM_Clamp>::GetRHI();
+	}
+
 	{
 		// Upload the buffer of lights to the GPU
 		size_t DataSize = sizeof(FPathTracingLight) * FMath::Max(LightCount, 1u);
