Light CDF optimizations for path tracer

Reduce the number of accesses to the light CDF array. Avoid needing to pre-fill the array, and remove the normalization pass instead just doing a division after sample selection.

Provides a 10% speedup on scenes with many lights

diff --git a/Engine/Plugins/Experimental/GPULightmass/Shaders/Private/LightmapPathTracing.usf b/Engine/Plugins/Experimental/GPULightmass/Shaders/Private/LightmapPathTracing.usf
--- a/Engine/Plugins/Experimental/GPULightmass/Shaders/Private/LightmapPathTracing.usf
+++ b/Engine/Plugins/Experimental/GPULightmass/Shaders/Private/LightmapPathTracing.usf
@@ -181,7 +181,7 @@ void PathTracingKernel(
 	// Seed the array with a uniform CDF at first so that we always have a valid CDF
 	float LightPickingCdf[RAY_TRACING_LIGHT_COUNT_MAXIMUM];
 
-	InitLightPickingCdf(LightPickingCdf);
+	float LightPickingCdfSum = 0.0;
 	
 	FRayCone RayCone = (FRayCone)0;
 	
@@ -269,7 +269,8 @@ void PathTracingKernel(
 				// GPULightmass records direct contribution for only static lights
 				if (Bounce > 1 || (Bounce == 1 && !IsStationary(LightId)))
 				{
-					AccumulateRadiance(Radiance, PathThroughput * LightResult.Radiance * MISWeightRobust(PrevMaterialPdf, LightResult.Pdf * LightPickPdf));
+					float MISWeight = LightPickingCdfSum > 0.0 ? MISWeightRobust(PrevMaterialPdf, LightResult.Pdf * LightPickPdf / LightPickingCdfSum) : 1.0;
+					AccumulateRadiance(Radiance, PathThroughput * LightResult.Radiance * MISWeight);
 				}
 			}
 		}
@@ -382,16 +383,20 @@ void PathTracingKernel(
 
 			bool IsTransmissiveMaterial = false; // ENABLE_TRANSMISSION && IsMaterialTransmissive(Payload);
 			
-			if (InitLightPickingCdf(WorldPos, WorldNormal, PrimitiveLightingChannelMask, IsTransmissiveMaterial, LightPickingCdf))
+			LightPickingCdfSum = InitLightPickingCdf(WorldPos, WorldNormal, PrimitiveLightingChannelMask, IsTransmissiveMaterial, LightPickingCdf);
+
+			if (LightPickingCdfSum > 0.0)
 			{
 				// init worked
 				int LightId;
 				float LightPickPdf = 0;
 
-				SelectLight(RandSample.x, LightPickingCdf, LightId, LightPickPdf);
+				SelectLight(RandSample.x * LightPickingCdfSum, LightPickingCdf, LightId, LightPickPdf);
 
 				FLightSample LightSample = SampleLight(LightId, RandSample, WorldPos, WorldNormal);
 
+				LightPickPdf /= LightPickingCdfSum;
+
 				LightSample.RadianceOverPdf /= LightPickPdf;
 				LightSample.Pdf *= LightPickPdf;
 
diff --git a/Engine/Shaders/Private/PathTracing/Light/PathTracingLightSampling.ush b/Engine/Shaders/Private/PathTracing/Light/PathTracingLightSampling.ush
--- a/Engine/Shaders/Private/PathTracing/Light/PathTracingLightSampling.ush
+++ b/Engine/Shaders/Private/PathTracing/Light/PathTracingLightSampling.ush
@@ -77,16 +77,7 @@ float EstimateLight(
 }
 
 // #dxr_todo: passing LightPickingCdf as anything but inout causes performance to drop significantly -- why ?
-void InitLightPickingCdf(inout float LightPickingCdf[RAY_TRACING_LIGHT_COUNT_MAXIMUM])
-{
-	for (uint LightIndex = 0; LightIndex < SceneLightCount; ++LightIndex)
-	{
-		LightPickingCdf[LightIndex] = float(LightIndex + 1) / float(SceneLightCount);
-	}
-}
-
-// #dxr_todo: passing LightPickingCdf as anything but inout causes performance to drop significantly -- why ?
-bool InitLightPickingCdf(
+float InitLightPickingCdf(
 	float3 WorldPos,
 	float3 WorldNormal,
 	uint PrimitiveLightingChannelMask,
@@ -95,41 +86,29 @@ bool InitLightPickingCdf(
 {
 	// This code assumes SceneLightCount > 0
 #ifdef UNIFORM_LIGHT_SELECTION
-	InitLightPickingCdf(LightPickingCdf);
+	return SceneLightCount;
 #else
 	// Build irradiance estimate prefix sum
-	float LastCdf = 0;
+	float CdfSum = 0;
 	{
 		for (uint LightIndex = 0; LightIndex < SceneLightCount; ++LightIndex)
 		{
-			LightPickingCdf[LightIndex] = EstimateLight(LightIndex, WorldPos, WorldNormal, PrimitiveLightingChannelMask, IsTransmissiveMaterial);
-			LightPickingCdf[LightIndex] += LastCdf;
-			LastCdf = LightPickingCdf[LightIndex];
+			CdfSum += EstimateLight(LightIndex, WorldPos, WorldNormal, PrimitiveLightingChannelMask, IsTransmissiveMaterial);
+			LightPickingCdf[LightIndex] = CdfSum;
 		}
 	}
-
-	if (LastCdf == 0)
-	{
-		// all the lights are invisible
-		return false;
-	}
-	{
-		// CDF normalization
-		for (uint LightIndex = 0; LightIndex < SceneLightCount; ++LightIndex)
-		{
-			LightPickingCdf[LightIndex] /= LastCdf;
-		}
-		LightPickingCdf[SceneLightCount - 1] = 1.0;
-	}
+	return CdfSum;
 #endif
-	return true;
 }
 
 // #dxr_todo: passing LightPickingCdf as anything but inout causes performance to drop significantly -- why ?
 void SelectLight(float RandSample, inout float LightPickingCdf[RAY_TRACING_LIGHT_COUNT_MAXIMUM], out uint LightId, out float LightPickPdf)
 {
 	// This code assumes SceneLightCount > 0
-
+#ifdef UNIFORM_LIGHT_SELECTION
+	LightId = min(int(floor(RandSample * SceneLightCount)), SceneLightCount - 1);
+	LightPickPdf = 1.0;
+#else
 #if 0
 	// linear search
 	float PreviousCdfValue = 0;
@@ -163,7 +142,7 @@ void SelectLight(float RandSample, inout float LightPickingCdf[RAY_TRACING_LIGHT
 
 	LightPickPdf = LightPickingCdf[LightId] - ((LightId > 0) ? LightPickingCdf[LightId - 1] : 0.0);
 #endif
-
+#endif
 }
 
 bool SampleLight(
@@ -193,7 +172,9 @@ bool SampleLight(
 
 	float LightPickingCdf[RAY_TRACING_LIGHT_COUNT_MAXIMUM];
 
-	if (!InitLightPickingCdf(WorldPos, WorldNormal, PrimitiveLightingChannelMask, IsTransmissiveMaterial, LightPickingCdf))
+	float LightPickingCdfSum = InitLightPickingCdf(WorldPos, WorldNormal, PrimitiveLightingChannelMask, IsTransmissiveMaterial, LightPickingCdf);
+
+	if (LightPickingCdfSum > 0.0)
 	{
 		// light picking cdf is all 0
 		return false;
@@ -201,7 +182,9 @@ bool SampleLight(
 
 	float LightPickPdf = 0;
 
-	SelectLight(RandSample.x, LightPickingCdf, OutLightId, LightPickPdf);
+	SelectLight(RandSample.x * LightPickingCdfSum, LightPickingCdf, OutLightId, LightPickPdf);
+
+	LightPickPdf /= LightPickingCdfSum;
 
 	FLightSample Result = SampleLight(OutLightId, RandSample, WorldPos, WorldNormal);
 	OutDirection = Result.Direction;
diff --git a/Engine/Shaders/Private/PathTracing/PathTracing.usf b/Engine/Shaders/Private/PathTracing/PathTracing.usf
--- a/Engine/Shaders/Private/PathTracing/PathTracing.usf
+++ b/Engine/Shaders/Private/PathTracing/PathTracing.usf
@@ -438,8 +438,6 @@ RAY_TRACING_ENTRY_RAYGEN(PathTracingMainRG)
 	// Seed the array with a uniform CDF at first so that we always have a valid CDF
 	float LightPickingCdf[RAY_TRACING_LIGHT_COUNT_MAXIMUM];
 
-	InitLightPickingCdf(LightPickingCdf);
-
 	Ray.Direction = normalize(Ray.Direction);
 
 	FRayCone RayCone = (FRayCone)0;
@@ -494,6 +492,8 @@ RAY_TRACING_ENTRY_RAYGEN(PathTracingMainRG)
 
 		const bool bDoLightLoop = PathTracingData.SkipDirectLighting == 0 || Bounce > 0;
 
+		float LightPickingCdfSum = 0;
+
 		// If we are using Light sampling and the material can use it ...
 		if (MISMode != 0 && bIsNeeValid && SceneLightCount > 0 && bDoLightLoop)
 		{
@@ -504,13 +504,16 @@ RAY_TRACING_ENTRY_RAYGEN(PathTracingMainRG)
 
 			bool IsTransmissiveMaterial = ENABLE_TRANSMISSION && IsMaterialTransmissive(Payload);
 
-			if (InitLightPickingCdf(WorldPos, WorldNormal, PrimitiveLightingChannelMask, IsTransmissiveMaterial, LightPickingCdf))
+			LightPickingCdfSum = InitLightPickingCdf(WorldPos, WorldNormal, PrimitiveLightingChannelMask, IsTransmissiveMaterial, LightPickingCdf);
+
+			if (LightPickingCdfSum > 0)
 			{
 				// init worked
 				int LightId;
 				float LightPickPdf = 0;
 
-				SelectLight(RandSample.x, LightPickingCdf, LightId, LightPickPdf);
+				SelectLight(RandSample.x * LightPickingCdfSum, LightPickingCdf, LightId, LightPickPdf);
+				LightPickPdf /= LightPickingCdfSum;
 
 				FLightSample LightSample = SampleLight(LightId, RandSample, WorldPos, WorldNormal);
 
@@ -614,6 +617,7 @@ RAY_TRACING_ENTRY_RAYGEN(PathTracingMainRG)
 			// While it would be possible to just loop around and use the indirect ray to do this, it would prevent the application
 			// of shadow ray specific logic for transparent shadows or various per light tricks like shadow casting
 			float PreviousCdfValue = 0.0;
+			const bool bUseMIS = MISMode == 2 && LightPickingCdfSum > 0;
 			for (uint LightId = 0; LightId < SceneLightCount; ++LightId)
 			{
 				if ((Payload.PrimitiveLightingChannelMask & GetLightingChannelMask(LightId)) == 0)
@@ -624,10 +628,15 @@ RAY_TRACING_ENTRY_RAYGEN(PathTracingMainRG)
 
 				FLightHit LightResult = TraceLight(Ray, LightId);
 
+				if (LightResult.IsMiss())
+				{
+					continue;
+				}
+
 				float3 LightContrib = PathThroughput * LightResult.Radiance;
-				if (MISMode == 2)
+				if (bUseMIS)
 				{
-					float LightPickPdf =  LightPickingCdf[LightId] - PreviousCdfValue;
+					float LightPickPdf =  (LightPickingCdf[LightId] - PreviousCdfValue) / LightPickingCdfSum;
 					PreviousCdfValue = LightPickingCdf[LightId];
 
 					LightContrib *= MISWeightRobust(MaterialSample.Pdf, LightResult.Pdf * LightPickPdf);
