Fixed shader compiler warning

diff --git a/Engine/Shaders/Private/MaterialTemplate.ush b/Engine/Shaders/Private/MaterialTemplate.ush
--- a/Engine/Shaders/Private/MaterialTemplate.ush
+++ b/Engine/Shaders/Private/MaterialTemplate.ush
@@ -1325,7 +1325,7 @@ float3 MaterialExpressionSkyAtmosphereViewLuminance(FMaterialPixelParameters Par
 	float2 SkyViewLutUv;
 	SkyViewLutParamsToUv(IntersectGround, ViewZenithCosAngle, WorldDirLocal, ViewHeight, ResolvedView.SkyAtmosphereBottomRadiusKm, ResolvedView.SkyViewLutSizeAndInvSize, SkyViewLutUv);
 	float3 SkyAtmosphereViewLuminance = Texture2DSampleLevel(View.SkyViewLutTexture, View.SkyViewLutTextureSampler, SkyViewLutUv, 0.0f).rgb;
-	SkyAtmosphereViewLuminance *= ResolvedView.SkyAtmosphereSkyLuminanceFactor;
+	SkyAtmosphereViewLuminance *= ResolvedView.SkyAtmosphereSkyLuminanceFactor.rgb;
 #if USE_PREEXPOSURE
 	SkyAtmosphereViewLuminance *= ResolvedView.OneOverPreExposure;
 #endif
