Normal offset contact shadows as well.

I previously killed the offset specifically for contact shadows to fix detail loss. At the time the amount of bias was larger than now and we still had the SMRT first sample leaking bug. Now that that is fixed and with the current 0.1 offset default which matches RT the difference of offset on and off on self shadow detail in Reverb is extremely small and not clearly better. It does though fix almost all of these triangle artifacts, at least on Echo.

diff --git a/Engine/Shaders/Private/VirtualShadowMaps/VirtualShadowMapProjection.usf b/Engine/Shaders/Private/VirtualShadowMaps/VirtualShadowMapProjection.usf
--- a/Engine/Shaders/Private/VirtualShadowMaps/VirtualShadowMapProjection.usf
+++ b/Engine/Shaders/Private/VirtualShadowMaps/VirtualShadowMapProjection.usf
@@ -231,7 +231,7 @@ FProjectLightResult ProjectLight(
 		const bool bIsHair = InputType == INPUT_TYPE_HAIRSTRANDS;
 		const half3 EstimatedGeoWorldNormal = bIsHair ? L : GBufferNormal;
 
-		float3 OffsetTranslatedWorldPosition = TranslatedWorldPosition + EstimatedGeoWorldNormal * NormalOffsetWorld;
+		TranslatedWorldPosition += EstimatedGeoWorldNormal * lerp(NormalOffsetWorld, 0.01, saturate(dot(EstimatedGeoWorldNormal, L)));
 
 		if (SMRTRayCount > 0)
 		{
@@ -242,7 +242,7 @@ FProjectLightResult ProjectLight(
 				Light,
 				PixelPos,
 				SceneDepth,
-				OffsetTranslatedWorldPosition,
+				TranslatedWorldPosition,
 				ContactShadowLengthWorld,
 				SMRTRayCount,
 				SMRTSamplesPerRay,
@@ -256,7 +256,7 @@ FProjectLightResult ProjectLight(
 				LightType,
 				PixelPos,
 				SceneDepth,
-				OffsetTranslatedWorldPosition,
+				TranslatedWorldPosition,
 				ContactShadowLengthWorld,
 				SMRTRayCount,
 				SMRTSamplesPerRay,
@@ -271,7 +271,7 @@ FProjectLightResult ProjectLight(
 			ShadowFactor = ProjectSingleSample(
 				VirtualShadowMapId,
 				L,
-				OffsetTranslatedWorldPosition,
+				TranslatedWorldPosition,
 				EstimatedGeoWorldNormal,
 				ContactShadowLengthWorld,
 				HitDistance);
