Added default directional light settings for VisualizeMeshDistanceFields when there is no directional light

diff --git a/Engine/Shaders/Private/DistanceFieldVisualization.usf b/Engine/Shaders/Private/DistanceFieldVisualization.usf
--- a/Engine/Shaders/Private/DistanceFieldVisualization.usf
+++ b/Engine/Shaders/Private/DistanceFieldVisualization.usf
@@ -233,15 +233,19 @@ void VisualizeMeshDistanceFieldCS(
 
 	float3 Lighting = 0;
 
+	float3 DirectionalLightColor = 10;
+	float3 DirectionalLightDirection = float3(.707f, 0, .707f);
+
+	if (ForwardLightData.HasDirectionalLight)
+	{
+		DirectionalLightColor = ForwardLightData.DirectionalLightColor;
+		DirectionalLightDirection = ForwardLightData.DirectionalLightDirection;
+	}
+
 	if (HitCulledObjectIndex != 0xFFFFFFFF)
 	{
 		// Compute simple lighting to help visualize errors in either position or normal in a way that is intuitive to artists
-		if (ForwardLightData.HasDirectionalLight)
-		{
-			float3 DirectionalLightColor = ForwardLightData.DirectionalLightColor;
-
-			Lighting += DirectionalLightColor * saturate(dot(ForwardLightData.DirectionalLightDirection, WorldHitNormal));
-		}
+		Lighting += DirectionalLightColor * saturate(dot(DirectionalLightDirection, WorldHitNormal));
 
 		Lighting += GetSkySHDiffuse(WorldHitNormal) * View.SkyLightColor.rgb;
 		float3 Albedo = .05f;
