Renamed old entries in HTileVisualize.usf to match new naming

diff --git a/Engine/Shaders/Private/HTileVisualize.usf b/Engine/Shaders/Private/HTileVisualize.usf
--- a/Engine/Shaders/Private/HTileVisualize.usf
+++ b/Engine/Shaders/Private/HTileVisualize.usf
@@ -56,22 +56,22 @@ void VisualizeHTile(uint3 PixelPos : SV_DispatchThreadID)
 	const float2 TileMinMax = TileMinMaxFP / (float2)((1 << 14) - 1);
 
 	float Result = 0.0f;
-	if (VisualizeMode == VISUALIZE_SCN_HTILE_MINZ || VisualizeMode == VISUALIZE_MAT_HTILE_MINZ)
+	if (VisualizeMode == VISUALIZE_SCENE_Z_MIN || VisualizeMode == VISUALIZE_MATERIAL_Z_MIN)
 	{
 		Result = pow(TileMinMax.x, 0.11f);
 	}
-	else if (VisualizeMode == VISUALIZE_SCN_HTILE_MAXZ || VisualizeMode == VISUALIZE_MAT_HTILE_MAXZ)
+	else if (VisualizeMode == VISUALIZE_SCENE_Z_MAX || VisualizeMode == VISUALIZE_MATERIAL_Z_MAX)
 	{
 		Result = pow(TileMinMax.y, 0.11f);
 	}
-	else if (VisualizeMode == VISUALIZE_SCN_HTILE_DELTAZ || VisualizeMode == VISUALIZE_MAT_HTILE_DELTAZ)
+	else if (VisualizeMode == VISUALIZE_SCENE_Z_DELTA || VisualizeMode == VISUALIZE_MATERIAL_Z_DELTA)
 	{
 		Result = pow(TileMinMax.y - TileMinMax.x, 0.11f);
 	}
-	else if (VisualizeMode == VISUALIZE_SCN_HTILE_ZMASK || VisualizeMode == VISUALIZE_MAT_HTILE_ZMASK)
+	/*else if (VisualizeMode == VISUALIZE_SCENE_Z_MASK || VisualizeMode == VISUALIZE_MATERIAL_Z_MASK)
 	{
 		Result = ExtractFloat(HTileEnc, 0, 4) / 15.0f;
-	}
+	}*/
 
 	HTileDisplay[PixelPos.xy] = float4(Result, 0, 0, 0);
 #else
