Removes branches from TAADilateVelocity.usf and implements COMPRESS_PREV_USE_COUNT

Saves ~0.1ms at 1080p input

diff --git a/Engine/Shaders/Private/TemporalAA/TAACommon.ush b/Engine/Shaders/Private/TemporalAA/TAACommon.ush
--- a/Engine/Shaders/Private/TemporalAA/TAACommon.ush
+++ b/Engine/Shaders/Private/TemporalAA/TAACommon.ush
@@ -33,6 +33,14 @@
 
 #define MAX_FALLBACK_SAMPLE_COUNT 20
 
+#define COMPRESS_PREV_USE_COUNT 1
+#if COMPRESS_PREV_USE_COUNT
+	// COMPRESS_PREV_USE_COUNT compress the use count in 8bit, so need to sum of a 2x2 to not overflow these 8bits
+	#define PREV_USE_COUNT_QUANTIZATION 63
+#else
+	#define PREV_USE_COUNT_QUANTIZATION 256
+#endif
+
 #define PARALLAX_REJECTION_MASK_THRESHOLD 0.5
 
 #define MAX_PARALLAX_FACTOR 8.0
@@ -556,3 +564,17 @@ FSuperResHistory UnpackSuperResHistory(uint2 Packed)
 
 	return SuperRes;
 }
+
+#if COMPRESS_PREV_USE_COUNT
+
+// encode 4 use count that are 8bits shifted in the PF_R32_UINT
+void ComputeCompressedUseCountPixelCoordinates(uint2 PixelPos, out uint2 CompressedUseCountPixelPos, out uint Shift)
+{
+	// Overlap 2x2 consecutive blocks of 8x8 pixel to reduce atomic contention in the scattering pass
+	const uint kTileSize = 8;
+
+	CompressedUseCountPixelPos = (PixelPos % kTileSize) | ((PixelPos >> 1) & ~(kTileSize - 1));
+	Shift = ((PixelPos.x / kTileSize) % 2) | (((PixelPos.y / kTileSize) % 2) << 1);
+}
+	
+#endif
diff --git a/Engine/Shaders/Private/TemporalAA/TAADecimateHistory.usf b/Engine/Shaders/Private/TemporalAA/TAADecimateHistory.usf
--- a/Engine/Shaders/Private/TemporalAA/TAADecimateHistory.usf
+++ b/Engine/Shaders/Private/TemporalAA/TAADecimateHistory.usf
@@ -103,7 +103,16 @@ void MainCS(
 
 			bool bValidPixelPos = all(PixelPos >= InputInfo_ViewportMin && PixelPos < InputInfo_ViewportMax);
 
-			uint HistoryUseCount = PrevUseCountTexture[PixelPos];
+			#if COMPRESS_PREV_USE_COUNT
+				uint2 CompressedUseCountPixelPos;
+				uint Shift;
+				ComputeCompressedUseCountPixelCoordinates(PixelPos, /* out */ CompressedUseCountPixelPos, /* out */ Shift);
+
+				uint HistoryUseCount = (PrevUseCountTexture[CompressedUseCountPixelPos] >> (8 * Shift)) & 0xFF;
+			#else
+				uint HistoryUseCount = PrevUseCountTexture[PixelPos];
+			#endif
+
 			float HistoryClosestDeviceZ = f16tof32(PrevClosestDepthTexture[PixelPos]);
 			float HistoryClosestWorldDepth = ConvertFromDeviceZ(HistoryClosestDeviceZ);
 
@@ -119,7 +128,7 @@ void MainCS(
 				DepthRejection = saturate(taa_half(2.0) - taa_half(DeltaDepth / WorldDepthEpsilon));
 			}
 
-			taa_half HistoryGhostingRejection = saturate(taa_half(1.0) - taa_half(3.0) * abs(taa_half(HistoryUseCount) * rcp(taa_half(256)) - taa_half(1.0)));
+			taa_half HistoryGhostingRejection = saturate(taa_half(1.0) - taa_half(3.0) * abs(taa_half(HistoryUseCount) * rcp(taa_half(PREV_USE_COUNT_QUANTIZATION)) - taa_half(1.0)));
 
 			FLATTEN
 			if (bValidPixelPos)
diff --git a/Engine/Shaders/Private/TemporalAA/TAADilateVelocity.usf b/Engine/Shaders/Private/TemporalAA/TAADilateVelocity.usf
--- a/Engine/Shaders/Private/TemporalAA/TAADilateVelocity.usf
+++ b/Engine/Shaders/Private/TemporalAA/TAADilateVelocity.usf
@@ -12,6 +12,9 @@
 
 //------------------------------------------------------- PARAMETERS
 
+float2 PrevOutputBufferUVMin;
+float2 PrevOutputBufferUVMax;
+
 Texture2D SceneDepthTexture;
 Texture2D SceneVelocityTexture;
 
@@ -124,39 +127,46 @@ void MainCS(
 	#endif
 
 	// Scatter parallax rejection
-	#if 1
+	float2 PrevScreenPos = ScreenPos - FinalScreenVelocity;
+	float2 PrevInputBufferUV = (InputInfo_ScreenPosToViewportScale * PrevScreenPos + InputInfo_ScreenPosToViewportBias) * InputInfo_ExtentInverse;
+
+
+	Debug.rg = VelocityPixelOffset * 0.5 + 0.5;
+
 	if (all(InputPixelPos < InputInfo_ViewportMax))
 	{
-		float2 PrevScreenPos = ScreenPos - FinalScreenVelocity;
-		float2 PrevInputBufferUV = (InputInfo_ScreenPosToViewportScale * PrevScreenPos + InputInfo_ScreenPosToViewportBias) * InputInfo_ExtentInverse;
+		#if 1
+		{
+			bool bValidHistoryCoord = all(PrevInputBufferUV > PrevOutputBufferUVMin && PrevInputBufferUV < PrevOutputBufferUVMax);
 
-		FBilinearSampleInfos BilinearInter = GetBilinearSampleLevelInfos(PrevInputBufferUV, InputInfo_Extent, InputInfo_ExtentInverse);
+			FBilinearSampleInfos BilinearInter = GetBilinearSampleLevelInfos(PrevInputBufferUV, InputInfo_Extent, InputInfo_ExtentInverse);
 
-		UNROLL_N(4)
-		for (uint i = 0; i < 4; i++)
-		{
-			float Weight = GetSampleWeight(BilinearInter, i);
-			uint2 PixelPos = GetSamplePixelCoord(BilinearInter, i);
+			UNROLL_N(4)
+			for (uint i = 0; i < 4; i++)
+			{
+				float Weight = GetSampleWeight(BilinearInter, i);
+				uint2 PixelPos = GetSamplePixelCoord(BilinearInter, i);
 
-			uint FinalAdd = uint(round(Weight * 256));
+				uint FinalAdd = uint(round(Weight * PREV_USE_COUNT_QUANTIZATION));
 			
-			if (all(PixelPos >= InputInfo_ViewportMin && PixelPos < InputInfo_ViewportMax))
-			{
-				InterlockedAdd(PrevUseCountOutput[PixelPos], FinalAdd);
+				PixelPos.x = ((FinalAdd > 0) && bValidHistoryCoord) ? PixelPos.x : (~0);
 
-				if (FinalAdd > 0)
-				{
-					InterlockedMax(PrevClosestDepthOutput[PixelPos], f32tof16(ClosestDeviceZ));
-				}
-			}
-		}
-	}
-	#endif
+				#if COMPRESS_PREV_USE_COUNT
+					uint2 CompressedUseCountPixelPos;
+					uint Shift;
+					ComputeCompressedUseCountPixelCoordinates(PixelPos, /* out */ CompressedUseCountPixelPos, /* out */ Shift);
 
-	Debug.rg = VelocityPixelOffset * 0.5 + 0.5;
+					FinalAdd = FinalAdd << (8 * Shift);
+				#else
+					uint2 CompressedUseCountPixelPos = PixelPos;
+				#endif
 
-	if (all(InputPixelPos < InputInfo_ViewportMax))
-	{
+				InterlockedAdd(PrevUseCountOutput[CompressedUseCountPixelPos], FinalAdd);
+				InterlockedMax(PrevClosestDepthOutput[PixelPos], f32tof16(ClosestDeviceZ));
+			}
+		}
+		#endif
+		
 		DilatedVelocityOutput[InputPixelPos] = EncodeVelocityToTexture(float3(FinalScreenVelocity, 0.0)).xy;
 		ClosestDepthOutput[InputPixelPos] = ClosestDeviceZ;
 		ParallaxFactorOutput[InputPixelPos] = CompressParallaxFactor(ParallaxFactorForHistory) | (CompressParallaxFactor(ParallaxFactorForHistoryRejection) << 4);
diff --git a/Engine/Source/Runtime/Renderer/Private/PostProcess/TemporalAA.cpp b/Engine/Source/Runtime/Renderer/Private/PostProcess/TemporalAA.cpp
--- a/Engine/Source/Runtime/Renderer/Private/PostProcess/TemporalAA.cpp
+++ b/Engine/Source/Runtime/Renderer/Private/PostProcess/TemporalAA.cpp
@@ -360,6 +360,9 @@ class FTAADilateVelocityCS : public FTAAGen5Shader
 	BEGIN_SHADER_PARAMETER_STRUCT(FParameters, )
 		SHADER_PARAMETER_STRUCT_INCLUDE(FTAACommonParameters, CommonParameters)
 
+		SHADER_PARAMETER(FVector2D, PrevOutputBufferUVMin)
+		SHADER_PARAMETER(FVector2D, PrevOutputBufferUVMax)
+
 		SHADER_PARAMETER_RDG_TEXTURE(Texture2D, SceneDepthTexture)
 		SHADER_PARAMETER_RDG_TEXTURE(Texture2D, SceneVelocityTexture)
 
@@ -1068,7 +1071,7 @@ FTAAOutputs AddTemporalAAPass(
 
 		FComputeShaderUtils::AddPass(
 			GraphBuilder,
-			RDG_EVENT_NAME("TAA %s%s %dx%d -> %dx%d",
+			RDG_EVENT_NAME("TAA Gen4 %s%s %dx%d -> %dx%d",
 				PassName, Inputs.bUseFast ? TEXT(" Fast") : TEXT(""),
 				PracticableSrcRect.Width(), PracticableSrcRect.Height(),
 				PracticableDestRect.Width(), PracticableDestRect.Height()),
@@ -1264,6 +1267,8 @@ static void AddGen5MainTemporalAAPasses(
 
 		FTAADilateVelocityCS::FParameters* PassParameters = GraphBuilder.AllocParameters<FTAADilateVelocityCS::FParameters>();
 		PassParameters->CommonParameters = CommonParameters;
+		PassParameters->PrevOutputBufferUVMin = CommonParameters.InputInfo.UVViewportBilinearMin - CommonParameters.InputInfo.ExtentInverse;
+		PassParameters->PrevOutputBufferUVMax = CommonParameters.InputInfo.UVViewportBilinearMax + CommonParameters.InputInfo.ExtentInverse;
 		PassParameters->SceneDepthTexture = PassInputs.SceneDepthTexture;
 		PassParameters->SceneVelocityTexture = PassInputs.SceneVelocityTexture;
 		PassParameters->DilatedVelocityOutput = GraphBuilder.CreateUAV(DilatedVelocityTexture);
