Optimises TAA's filter frequency pass

diff --git a/Engine/Shaders/Private/TemporalAA/TAACommon.ush b/Engine/Shaders/Private/TemporalAA/TAACommon.ush
--- a/Engine/Shaders/Private/TemporalAA/TAACommon.ush
+++ b/Engine/Shaders/Private/TemporalAA/TAACommon.ush
@@ -167,6 +167,8 @@ float2 InputInfo_UVViewportSizeInverse;
 float2 InputInfo_UVViewportBilinearMin;
 float2 InputInfo_UVViewportBilinearMax;
 float2 InputJitter;
+int2   InputPixelPosMin;
+int2   InputPixelPosMax;
 
 float2 LowFrequencyInfo_Extent;
 float2 LowFrequencyInfo_ExtentInverse;
@@ -238,6 +240,46 @@ uint bEnableInterferenceHeuristic;
 
 //------------------------------------------------------- FUNCTIONS
 
+/** Compute the group wave index into SGRP to then recompue the GroupThreadIndex latter */
+uint GetGroupWaveIndex(uint GroupThreadIndex, uint GroupSize)
+#if COMPILER_SUPPORTS_WAVE_ONCE
+{
+	uint LaneCountPerWave = WaveGetLaneCount();
+
+	if (LaneCountPerWave >= GroupSize)
+	{
+		return 0;
+	}
+
+	return WaveReadLaneFirst(GroupThreadIndex / LaneCountPerWave);
+}
+#else
+{
+	return 0;
+}
+#endif
+
+/** Force compute the group GroupThreadIndex through lane index and wave index if possible to reduce VGPR pressure. */
+uint GetGroupThreadIndex(uint GroupThreadIndex, uint GroupWaveIndex)
+#if COMPILER_SUPPORTS_WAVE_ONCE
+{
+	// shares GroupWaveOffset to save SALU
+	uint GroupWaveOffset = WaveGetLaneCount() * GroupWaveIndex;
+
+	// Do not share
+	ISOLATE
+	{
+		return GroupWaveOffset + WaveGetLaneIndex();
+	}
+}
+#else
+{
+	return GroupThreadIndex;
+}
+#endif
+
+
+
 float SafeRcp(float x)
 {
 	return x > 0.0 ? rcp(x) : 0.0;
diff --git a/Engine/Shaders/Private/TemporalAA/TAAFilterFrequencies.usf b/Engine/Shaders/Private/TemporalAA/TAAFilterFrequencies.usf
--- a/Engine/Shaders/Private/TemporalAA/TAAFilterFrequencies.usf
+++ b/Engine/Shaders/Private/TemporalAA/TAAFilterFrequencies.usf
@@ -8,14 +8,11 @@
 /** Kernel enum */
 	#define KERNEL_1X1 0
 	#define KERNEL_3X3 1
-	#define KERNEL_STACKOWIAK 2
 
 
 //------------------------------------------------------- CONFIG
 
-#define TILE_SIZE 8
-
-#define CONFIG_ACCUMULATE_KARIS_COLOR 0
+#define TILE_SIZE 16
 
 #define CONFIG_ACCUMULATE_KARIS_WEIGHTING 1
 
@@ -30,69 +27,6 @@
 #define CONFIG_USE_DYNAMIC_SPREAD_FACTOR 1
 
 
-//------------------------------------------------------- CONSTANTS
-
-static const float2 kStackowiakSampleSet0[56 * 4] =
-{
-	float2(-0.5, -0.5), float2(+0.5, -0.5), float2(-0.5, +0.5), float2(+0.5, +0.5),
-	float2(-1.5, +0.5), float2(-1.5, -0.5), float2(-0.5, +1.5), float2(+1.5, -0.5),
-	float2(+0.5, -1.5), float2(+2.5, -0.5), float2(+1.5, +0.5), float2(-0.5, -1.5),
-	float2(-1.5, -2.5), float2(-0.5, -2.5), float2(-1.5, -1.5), float2(-0.5, +2.5),
-	float2(-1.5, +1.5), float2(+1.5, -2.5), float2(-1.5, +2.5), float2(+1.5, +2.5),
-	float2(+0.5, -2.5), float2(-2.5, -0.5), float2(-2.5, -1.5), float2(-2.5, +0.5),
-	float2(+0.5, +1.5), float2(+0.5, +2.5), float2(-3.5, +0.5), float2(+0.5, +3.5),
-	float2(+1.5, -1.5), float2(+3.5, -0.5), float2(+2.5, +1.5), float2(+3.5, +0.5),
-	float2(+1.5, +1.5), float2(-2.5, +1.5), float2(-3.5, +2.5), float2(+3.5, +1.5),
-	float2(-3.5, -0.5), float2(-1.5, -3.5), float2(-2.5, -2.5), float2(-2.5, +2.5),
-	float2(+2.5, +0.5), float2(+2.5, +2.5), float2(+1.5, +3.5), float2(+3.5, -1.5),
-	float2(-3.5, -2.5), float2(+3.5, -2.5), float2(+2.5, -1.5), float2(+0.5, -3.5),
-	float2(-0.5, +3.5), float2(-0.5, -4.5), float2(-4.5, +0.5), float2(+4.5, +0.5),
-	float2(-4.5, -1.5), float2(-3.5, +1.5), float2(-0.5, -3.5), float2(+1.5, -3.5),
-	float2(+0.5, -4.5), float2(-1.5, +3.5), float2(+0.5, +4.5), float2(-3.5, -1.5),
-	float2(-4.5, +1.5), float2(+2.5, -4.5), float2(+2.5, -2.5), float2(-1.5, +4.5),
-	float2(-2.5, -4.5), float2(+4.5, -2.5), float2(+2.5, +3.5), float2(-3.5, +3.5),
-	float2(-2.5, +3.5), float2(+0.5, -5.5), float2(-4.5, +3.5), float2(-2.5, -3.5),
-	float2(-4.5, +2.5), float2(+3.5, +3.5), float2(+2.5, -3.5), float2(+4.5, +3.5),
-	float2(+3.5, -3.5), float2(+4.5, +2.5), float2(-5.5, +1.5), float2(-4.5, -0.5),
-	float2(+3.5, +2.5), float2(-0.5, +4.5), float2(-1.5, +5.5), float2(+1.5, +5.5),
-	float2(+4.5, -0.5), float2(+5.5, +0.5), float2(+4.5, +1.5), float2(-1.5, -4.5),
-	float2(-1.5, -5.5), float2(-4.5, -2.5), float2(-2.5, +5.5), float2(+2.5, +5.5),
-	float2(+1.5, +4.5), float2(+5.5, +1.5), float2(+1.5, -4.5), float2(-3.5, -3.5),
-	float2(+3.5, -4.5), float2(-3.5, -4.5), float2(+4.5, -1.5), float2(+4.5, -3.5),
-	float2(-3.5, -5.5), float2(-2.5, -5.5), float2(-4.5, -3.5), float2(+4.5, +4.5),
-	float2(-3.5, +4.5), float2(-2.5, +4.5), float2(-5.5, -2.5), float2(-5.5, +0.5),
-	float2(+2.5, -5.5), float2(+3.5, +4.5), float2(-0.5, -5.5), float2(-0.5, +6.5),
-	float2(+2.5, +4.5), float2(-5.5, -0.5), float2(-6.5, -1.5), float2(+1.5, -5.5),
-	float2(-6.5, -0.5), float2(+0.5, +5.5), float2(+1.5, +6.5), float2(+6.5, +1.5),
-	float2(-0.5, +5.5), float2(+6.5, -0.5), float2(-4.5, -4.5), float2(-5.5, +2.5),
-	float2(+5.5, -0.5), float2(-5.5, -1.5), float2(-6.5, +3.5), float2(-1.5, +6.5),
-	float2(-6.5, +0.5), float2(+4.5, -5.5), float2(-3.5, +6.5), float2(+6.5, -1.5),
-	float2(+0.5, -6.5), float2(-5.5, -3.5), float2(+5.5, -2.5), float2(+4.5, -4.5),
-	float2(+5.5, -1.5), float2(+3.5, -6.5), float2(+5.5, +3.5), float2(+3.5, -5.5),
-	float2(-5.5, -4.5), float2(+6.5, -3.5), float2(-0.5, -6.5), float2(+3.5, +6.5),
-	float2(-5.5, +3.5), float2(+0.5, +6.5), float2(+6.5, +0.5), float2(+6.5, -2.5),
-	float2(-6.5, -3.5), float2(-4.5, +4.5), float2(-7.5, -0.5), float2(+7.5, +0.5),
-	float2(+5.5, +2.5), float2(-0.5, -7.5), float2(+0.5, +7.5), float2(-4.5, +5.5),
-	float2(+3.5, +5.5), float2(-3.5, +5.5), float2(-4.5, -5.5), float2(+4.5, +6.5),
-	float2(+5.5, -4.5), float2(+4.5, +5.5), float2(-4.5, +6.5), float2(+6.5, +4.5),
-	float2(-7.5, +1.5), float2(-6.5, +1.5), float2(+5.5, -3.5), float2(-6.5, +2.5),
-	float2(-2.5, +6.5), float2(-1.5, -7.5), float2(+5.5, +4.5), float2(-1.5, -6.5),
-	float2(-3.5, -7.5), float2(+2.5, -7.5), float2(-7.5, +2.5), float2(-6.5, -2.5),
-	float2(-5.5, +5.5), float2(+2.5, +6.5), float2(-2.5, -6.5), float2(-7.5, +0.5),
-	float2(-0.5, +7.5), float2(+7.5, -2.5), float2(-2.5, +7.5), float2(+0.5, -7.5),
-	float2(-4.5, -7.5), float2(+7.5, +1.5), float2(+1.5, -6.5), float2(-6.5, +4.5),
-	float2(-1.5, +7.5), float2(-5.5, -5.5), float2(+6.5, +2.5), float2(-3.5, -6.5),
-	float2(+3.5, -7.5), float2(-5.5, +4.5), float2(+2.5, -6.5), float2(+1.5, -7.5),
-	float2(+6.5, +3.5), float2(+5.5, -6.5), float2(-6.5, +5.5), float2(+7.5, +4.5),
-	float2(+7.5, -1.5), float2(-7.5, -1.5), float2(+3.5, +7.5), float2(-5.5, +6.5),
-	float2(+1.5, +7.5), float2(+7.5, +3.5), float2(+7.5, -0.5), float2(-7.5, -2.5),
-	float2(+5.5, +5.5), float2(+6.5, +5.5), float2(+5.5, -5.5), float2(-2.5, -7.5),
-	float2(+2.5, +7.5), float2(-7.5, -3.5), float2(-7.5, -4.5), float2(-6.5, -4.5),
-	float2(+7.5, -3.5), float2(+5.5, +6.5), float2(-5.5, -6.5), float2(-4.5, -6.5),
-	float2(+7.5, +2.5), float2(-7.5, +3.5), float2(+4.5, -6.5), float2(+7.5, -4.5),
-};
-
-
 //------------------------------------------------------- PARAMETERS
 
 float3 OutputQuantizationError;
@@ -131,15 +65,14 @@ void MainCS(
 	uint2 GroupId : SV_GroupID,
 	uint GroupThreadIndex : SV_GroupIndex)
 {
-	float4 Debug = 0.0;
-
-	uint2 DispatchThreadId = (
-		ZOrder2D(GroupThreadIndex, log2(TILE_SIZE)) +
-		GroupId * uint2(TILE_SIZE, TILE_SIZE));
+	uint GroupWaveIndex = GetGroupWaveIndex(GroupThreadIndex, /* GroupSize = */ TILE_SIZE * TILE_SIZE);
 
-	uint2 InputPixelPos = LowFrequencyInfo_ViewportMin + DispatchThreadId;
+	float4 Debug = 0.0;
 
-	float2 InputBufferUV = (InputPixelPos + 0.5) * LowFrequencyInfo_ExtentInverse;
+	uint2 InputPixelPos = (
+		LowFrequencyInfo_ViewportMin +
+		GroupId * uint2(TILE_SIZE, TILE_SIZE) +
+		ZOrder2D(GroupThreadIndex, log2(TILE_SIZE)));
 
 	// Fetch center of the kernel
 	taa_half3 CenterInput;
@@ -151,18 +84,13 @@ void MainCS(
 	{
 		taa_half SampleSpatialWeight = 1.0;
 
-		taa_half3 InputSample = InputTexture.SampleLevel(GlobalPointClampedSampler, InputBufferUV, 0);
-		taa_half3 HistorySample = PredictionSceneColorTexture.SampleLevel(GlobalPointClampedSampler, InputBufferUV, 0);
+		taa_half3 InputSample = InputTexture[InputPixelPos];
+		taa_half3 HistorySample = PredictionSceneColorTexture[InputPixelPos];
 		taa_half HistorySampleValidity = ParallaxRejectionMaskToHistoryValidity(
-			ParallaxRejectionMaskTexture.SampleLevel(GlobalPointClampedSampler, InputBufferUV, 0));
+			ParallaxRejectionMaskTexture[InputPixelPos]);
 		
 		taa_half InputWeight = ComputeColorWeight(InputSample);
 		taa_half HistoryWeight = HistorySampleValidity * ComputeColorWeight(HistorySample);
-		
-		#if CONFIG_ACCUMULATE_KARIS_COLOR
-			InputSample *= 10.0 * HdrWeight4(InputSample);
-			HistorySample *= 10.0 * HdrWeight4(HistorySample);
-		#endif
 
 		CenterInput = InputSample;
 
@@ -176,35 +104,25 @@ void MainCS(
 	uint StartSampleId;
 	uint SampleCount;
 	uint SampleTrackId;
-	float2 KernelCenterUV;
 	
 	#if CONFIG_KERNEL == KERNEL_1X1
 	{
 		StartSampleId = 0;
 		SampleCount = 0;
-		KernelCenterUV = InputBufferUV;
 	}
 	#elif CONFIG_KERNEL == KERNEL_3X3
 	{
 		StartSampleId = 0;
 		SampleCount = 9;
-		KernelCenterUV = InputBufferUV;
-	}
-	#elif CONFIG_KERNEL == KERNEL_STACKOWIAK
-	{
-		StartSampleId = 1;
-		SampleCount = 16;
-		KernelCenterUV = float2(InputPixelPos | 1) * LowFrequencyInfo_ExtentInverse;
-		SampleTrackId = ((InputPixelPos.x & 1) | ((InputPixelPos.y & 1) << 1));
 	}
 	#else
 		#error Unknown kernel
 	#endif
 
 	#if CONFIG_USE_DYNAMIC_SPREAD_FACTOR
-		float2 DynamicSpreadFactor = CONFIG_KERNEL_SPREAD_FACTOR + float2((InputPixelPos % 2) ^ (InputPixelPos.yx % 2));
+		int2 DynamicSpreadFactor = CONFIG_KERNEL_SPREAD_FACTOR + int2((InputPixelPos % 2) ^ (InputPixelPos.yx % 2));
 	#else
-		float2 DynamicSpreadFactor = CONFIG_KERNEL_SPREAD_FACTOR;
+		int2 DynamicSpreadFactor = CONFIG_KERNEL_SPREAD_FACTOR;
 	#endif
 
 	// Fetch neighborhood of the kernel
@@ -215,13 +133,13 @@ void MainCS(
 			continue;
 		#endif
 
-		float2 Offset;
-		float SampleSpatialWeight;
+		int2 Offset;
+		taa_half SampleSpatialWeight;
 		
 		#if CONFIG_KERNEL == KERNEL_1X1
 		{
-			Offset = 0.0;
-			SampleSpatialWeight = 1.0;
+			Offset = 0;
+			SampleSpatialWeight = taa_half(1.0);
 		}
 		#elif CONFIG_KERNEL == KERNEL_3X3
 		{
@@ -237,11 +155,11 @@ void MainCS(
 			//}
 			
 			#if 0
-				SampleSpatialWeight = 1.0;
+				SampleSpatialWeight = taa_half(1.0);
 			#elif 1
-				SampleSpatialWeight = rcp(abs(Offset.x) + 1.0) * rcp(abs(Offset.y) + 1.0);
+				SampleSpatialWeight = taa_half(rcp(abs(Offset.x) + 1.0) * rcp(abs(Offset.y) + 1.0));
 			#else
-				SampleSpatialWeight = rcp(abs(Offset.x) + abs(Offset.y));
+				SampleSpatialWeight = taa_half(rcp(abs(Offset.x) + abs(Offset.y)));
 			#endif
 		}
 		#elif CONFIG_KERNEL == KERNEL_STACKOWIAK
@@ -249,27 +167,28 @@ void MainCS(
 			Offset = kStackowiakSampleSet0[4 * SampleId + SampleTrackId];
 			
 			#if 1
-				SampleSpatialWeight = 1.0;
+				SampleSpatialWeight = taa_half(1.0);
 			#elif 1
-				SampleSpatialWeight = rcp(abs(Offset.x) + 1.0) * rcp(abs(Offset.y) + 1.0);
+				SampleSpatialWeight = taa_half(rcp(abs(Offset.x) + 1.0) * rcp(abs(Offset.y) + 1.0));
 			#else
-				SampleSpatialWeight = rcp(abs(Offset.x) + abs(Offset.y));
+				SampleSpatialWeight = taa_half(rcp(abs(Offset.x) + abs(Offset.y)));
 			#endif
 		}
 		#else
 			#error Unknown kernel
 		#endif
 
-		float2 SampleInputBufferUV = InputBufferUV + (DynamicSpreadFactor * Offset) * LowFrequencyInfo_ExtentInverse;
-		SampleInputBufferUV = clamp(SampleInputBufferUV, LowFrequencyInfo_UVViewportBilinearMin, LowFrequencyInfo_UVViewportBilinearMax);
+		int2 SampleInputPixelPos = int2(InputPixelPos) + DynamicSpreadFactor * Offset;
+		SampleInputPixelPos = clamp(SampleInputPixelPos, InputPixelPosMin, InputPixelPosMax);
 
-		taa_half3 InputSample = InputTexture.SampleLevel(GlobalPointClampedSampler, SampleInputBufferUV, 0);
-		taa_half3 HistorySample = PredictionSceneColorTexture.SampleLevel(GlobalPointClampedSampler, SampleInputBufferUV, 0);
+		taa_half3 InputSample = InputTexture[SampleInputPixelPos];
+		taa_half3 HistorySample = PredictionSceneColorTexture[SampleInputPixelPos];
 		taa_half HistorySampleValidity = ParallaxRejectionMaskToHistoryValidity(
-			ParallaxRejectionMaskTexture.SampleLevel(GlobalPointClampedSampler, SampleInputBufferUV, 0));
+			ParallaxRejectionMaskTexture[SampleInputPixelPos]);
 		
-		taa_half SampleInterferenceWeight = InterferenceWeightTexture.SampleLevel(GlobalPointClampedSampler, SampleInputBufferUV, 0);
+		taa_half SampleInterferenceWeight = InterferenceWeightTexture[SampleInputPixelPos];
 		
+		FLATTEN
 		if (!bEnableInterferenceHeuristic)
 		{
 			SampleInterferenceWeight = 1.0;
@@ -277,11 +196,6 @@ void MainCS(
 
 		taa_half InputWeight = (HistorySampleValidity * SampleInterferenceWeight) * ComputeColorWeight(InputSample);
 		taa_half HistoryWeight = (HistorySampleValidity * SampleInterferenceWeight) * ComputeColorWeight(HistorySample);
-		
-		#if CONFIG_ACCUMULATE_KARIS_COLOR
-			InputSample *= 10.0 * HdrWeight4(InputSample);
-			HistorySample *= 10.0 * HdrWeight4(HistorySample);
-		#endif
 
 		FilteredInput       += (SampleSpatialWeight * InputWeight) * InputSample;
 		FilteredInputWeight += (SampleSpatialWeight * InputWeight);
@@ -296,7 +210,10 @@ void MainCS(
 	}
 	
 	#if 1
-		uint2 OutputPixelPos = uint2(InputBufferUV * LowFrequencyInfo_Extent);
+		uint2 OutputPixelPos = (
+			LowFrequencyInfo_ViewportMin + 
+			GroupId * uint2(TILE_SIZE, TILE_SIZE) +
+			ZOrder2D(GetGroupThreadIndex(GroupThreadIndex, GroupWaveIndex), log2(TILE_SIZE)));
 	#else
 		uint2 OutputPixelPos = InputPixelPos;
 	#endif
@@ -311,16 +228,14 @@ void MainCS(
 	}
 	#endif
 
+	OutputPixelPos.x = all(OutputPixelPos < LowFrequencyInfo_ViewportMax) ? OutputPixelPos.x : ~0;
 
-	if (all(OutputPixelPos < LowFrequencyInfo_ViewportMax))
-	{
-		FilteredInputOutput[OutputPixelPos] = FilteredInput; // * (HdrWeight4(FilteredInput) * 10);
-		FilteredPredictionSceneColorOutput[OutputPixelPos] = FilteredHistory; // * (HdrWeight4(FilteredHistory) * 10);
+	FilteredInputOutput[OutputPixelPos] = FilteredInput;
+	FilteredPredictionSceneColorOutput[OutputPixelPos] = FilteredHistory;
 		
-		#if DEBUG_OUTPUT
-		{
-			DebugOutput[OutputPixelPos - LowFrequencyInfo_ViewportMin] = Debug;
-		}
-		#endif
+	#if DEBUG_OUTPUT
+	{
+		DebugOutput[OutputPixelPos - LowFrequencyInfo_ViewportMin] = Debug;
 	}
+	#endif
 }
diff --git a/Engine/Source/Runtime/Renderer/Private/PostProcess/TemporalAA.cpp b/Engine/Source/Runtime/Renderer/Private/PostProcess/TemporalAA.cpp
--- a/Engine/Source/Runtime/Renderer/Private/PostProcess/TemporalAA.cpp
+++ b/Engine/Source/Runtime/Renderer/Private/PostProcess/TemporalAA.cpp
@@ -113,6 +113,9 @@ BEGIN_SHADER_PARAMETER_STRUCT(FTAACommonParameters, )
 	SHADER_PARAMETER_STRUCT(FScreenPassTextureViewportParameters, RejectionInfo)
 	SHADER_PARAMETER_STRUCT(FScreenPassTextureViewportParameters, HistoryInfo)
 
+	SHADER_PARAMETER(FIntPoint, InputPixelPosMin)
+	SHADER_PARAMETER(FIntPoint, InputPixelPosMax)
+
 	SHADER_PARAMETER(FVector2D, InputJitter)
 	SHADER_PARAMETER(int32, bCameraCut)
 	SHADER_PARAMETER(int32, bEnableInterferenceHeuristic)
@@ -1186,6 +1189,8 @@ static void AddGen5MainTemporalAAPasses(
 	{
 		CommonParameters.InputInfo = GetScreenPassTextureViewportParameters(FScreenPassTextureViewport(
 			InputExtent, InputRect));
+		CommonParameters.InputPixelPosMin = CommonParameters.InputInfo.ViewportMin;
+		CommonParameters.InputPixelPosMax = CommonParameters.InputInfo.ViewportMax - 1;
 
 		CommonParameters.LowFrequencyInfo = GetScreenPassTextureViewportParameters(FScreenPassTextureViewport(
 			LowFrequencyExtent, LowFrequencyRect));
@@ -1621,7 +1626,7 @@ static void AddGen5MainTemporalAAPasses(
 				RDG_EVENT_NAME("TAA FilterFrequencies %dx%d", LowFrequencyRect.Width(), LowFrequencyRect.Height()),
 				ComputeShader,
 				PassParameters,
-				FComputeShaderUtils::GetGroupCount(LowFrequencyRect.Size(), 8));
+				FComputeShaderUtils::GetGroupCount(LowFrequencyRect.Size(), 16));
 		}
 
 		// Compare the low frequencies
