Optimises TAA's CompareHistory pass

diff --git a/Engine/Shaders/Private/TemporalAA/TAACommon.ush b/Engine/Shaders/Private/TemporalAA/TAACommon.ush
--- a/Engine/Shaders/Private/TemporalAA/TAACommon.ush
+++ b/Engine/Shaders/Private/TemporalAA/TAACommon.ush
@@ -134,15 +134,37 @@ static const uint kPlusIndexes3x3[5] = { 4, 1, 3, 5, 7 };
 
 
 #if CONFIG_COMPILE_FP16
-	#define taa_half half
+	#define taa_half  half
 	#define taa_half2 half2
 	#define taa_half3 half3
 	#define taa_half4 half4
+
+	#define taa_short  int16_t
+	#define taa_short2 int16_t2
+	#define taa_short3 int16_t3
+	#define taa_short4 int16_t4
+
+	#define taa_ushort  uint16_t
+	#define taa_ushort2 uint16_t2
+	#define taa_ushort3 uint16_t3
+	#define taa_ushort4 uint16_t4
+
 #else
-	#define taa_half float
+	#define taa_half  float
 	#define taa_half2 float2
 	#define taa_half3 float3
 	#define taa_half4 float4
+
+	#define taa_short  int
+	#define taa_short2 int2
+	#define taa_short3 int3
+	#define taa_short4 int4
+
+	#define taa_ushort  uint
+	#define taa_ushort2 uint2
+	#define taa_ushort3 uint3
+	#define taa_ushort4 uint4
+
 #endif
 
 #define taa_subpixel_details uint
@@ -278,7 +300,54 @@ uint GetGroupThreadIndex(uint GroupThreadIndex, uint GroupWaveIndex)
 }
 #endif
 
+// Clamp the offset to be shared across multiple samples
+taa_short2 ClampPixelOffset(
+	taa_short2 KernelCenterPixelPos,
+	taa_short2 Offset,
+	const taa_short2 OffsetDirection,
+	int2 MinPixelPos, int2 MaxPixelPos)
+{
+	taa_short2 Min = taa_short2(MinPixelPos) - KernelCenterPixelPos;
+	taa_short2 Max = taa_short2(MaxPixelPos) - KernelCenterPixelPos;
+
+	// Only do clamp based on the compile time known direction of the offset.
+	// This turns only 1 v_pk_max_i16 and v_pk_min_i16 for a 3x3 kernel.
+	taa_short2 ClampedOffset = 0;
 
+	if (OffsetDirection.x > 0)
+		ClampedOffset.x = min(Offset.x, Max.x);
+	else if (OffsetDirection.x < 0)
+		ClampedOffset.x = max(Offset.x, Min.x);
+		
+	if (OffsetDirection.y > 0)
+		ClampedOffset.y = min(Offset.y, Max.y);
+	else if (OffsetDirection.y < 0)
+		ClampedOffset.y = max(Offset.y, Min.y);
+
+	return ClampedOffset;
+}
+
+taa_ushort2 Map8x8Tile2x2Lane(uint GroupThreadIndex)
+{
+	taa_ushort T = taa_ushort(GroupThreadIndex);
+
+	taa_ushort2 GroupId;
+	GroupId.x = ((T >> taa_ushort(0)) & taa_ushort(0x01)) | ((T >> taa_ushort(2 - 1)) & taa_ushort(0x03 << 1));
+	GroupId.y = ((T >> taa_ushort(1)) & taa_ushort(0x01)) | ((T >> taa_ushort(4 - 1)) & taa_ushort(0x03 << 1));
+
+	return GroupId;
+}
+
+taa_ushort2 Map16x16Tile2x2Lane(uint GroupThreadIndex)
+{
+	taa_ushort T = taa_ushort(GroupThreadIndex);
+
+	taa_ushort2 GroupId;
+	GroupId.x = ((T >> taa_ushort(0)) & taa_ushort(0x01)) | ((T >> taa_ushort(2 - 1)) & taa_ushort(0x07 << 1));
+	GroupId.y = ((T >> taa_ushort(1)) & taa_ushort(0x01)) | ((T >> taa_ushort(5 - 1)) & taa_ushort(0x07 << 1));
+	
+	return GroupId;
+}
 
 float SafeRcp(float x)
 {
diff --git a/Engine/Shaders/Private/TemporalAA/TAACompareHistory.usf b/Engine/Shaders/Private/TemporalAA/TAACompareHistory.usf
--- a/Engine/Shaders/Private/TemporalAA/TAACompareHistory.usf
+++ b/Engine/Shaders/Private/TemporalAA/TAACompareHistory.usf
@@ -7,7 +7,7 @@
 
 #define CONFIG_OUTPUT_HALF_RES 1
 
-#define TILE_SIZE 8
+#define TILE_SIZE 16
 
 
 //------------------------------------------------------- PARAMETERS
@@ -42,14 +42,10 @@ void MainCS(
 {
 	float4 Debug = 0.0;
 
-	uint2 DispatchThreadId = (
-		ZOrder2D(GroupThreadIndex, log2(TILE_SIZE)) +
-		GroupId * uint2(TILE_SIZE, TILE_SIZE));
-
-	uint2 InputPixelPos = LowFrequencyInfo_ViewportMin + DispatchThreadId;
-
-	float2 InputBufferUV = (InputPixelPos + 0.5) * LowFrequencyInfo_ExtentInverse;
-	InputBufferUV = clamp(InputBufferUV, LowFrequencyInfo_UVViewportBilinearMin, LowFrequencyInfo_UVViewportBilinearMax);
+	taa_short2 InputPixelPos = (
+		taa_short2(InputPixelPosMin) +
+		taa_short2(GroupId) * taa_short2(TILE_SIZE, TILE_SIZE) +
+		Map16x16Tile2x2Lane(GroupThreadIndex));
 
 	taa_half3 InputMinColor;
 	taa_half3 InputMaxColor;
@@ -57,7 +53,7 @@ void MainCS(
 	taa_half ParallaxRejectionMask;
 	taa_half InterferenceWeight;
 	{
-		InputCenterColor = TransformColorForClampingBox(FilteredInputTexture.SampleLevel(GlobalPointClampedSampler, InputBufferUV, 0));
+		InputCenterColor = TransformColorForClampingBox(FilteredInputTexture[InputPixelPos]);
 		ParallaxRejectionMask = ParallaxRejectionMaskTexture[InputPixelPos];
 
 		InputMinColor = InputCenterColor;
@@ -73,11 +69,13 @@ void MainCS(
 		{
 			continue;
 		}
+		
+		taa_short2 SampleInputPixelPos = InputPixelPos + ClampPixelOffset(
+			InputPixelPos,
+			kOffsets3x3[i], kOffsets3x3[i],
+			InputPixelPosMin, InputPixelPosMax);
 
-		float2 SampleInputBufferUV = InputBufferUV + float2(kOffsets3x3[i]) * LowFrequencyInfo_ExtentInverse;
-		SampleInputBufferUV = clamp(SampleInputBufferUV, LowFrequencyInfo_UVViewportBilinearMin, LowFrequencyInfo_UVViewportBilinearMax);
-
-		taa_half3 Color = TransformColorForClampingBox(FilteredInputTexture.SampleLevel(GlobalPointClampedSampler, SampleInputBufferUV, 0));
+		taa_half3 Color = TransformColorForClampingBox(FilteredInputTexture[SampleInputPixelPos]);
 
 		InputMinColor = min(InputMinColor, Color);
 		InputMaxColor = max(InputMaxColor, Color);
@@ -98,6 +96,7 @@ void MainCS(
 	}
 	#endif
 	
+	FLATTEN
 	if (bEnableInterferenceHeuristic)
 	{
 		Rejection = saturate(Rejection + 1.0 - InterferenceWeight);
@@ -122,25 +121,30 @@ void MainCS(
 			Rejection = min(Rejection, SharedRejection[GroupThreadIndex ^ ButterflySwap]);
 		}
 
-		
-		if (all(InputPixelPos < LowFrequencyInfo_ViewportMax) && all(((InputPixelPos - LowFrequencyInfo_ViewportMin) % 2) == 0))
-		{
-			HistoryRejectionOutput[(InputPixelPos - LowFrequencyInfo_ViewportMin) / 2] = Rejection;
-		}
+		bool bIsValid = all(InputPixelPos <= taa_short2(InputPixelPosMax)) && all((taa_ushort2(InputPixelPos - taa_short2(InputPixelPosMin)) % taa_ushort(2)) == taa_ushort(0));
+
+		taa_ushort2 OutputPixelPos = taa_ushort2(InputPixelPos - taa_short2(InputPixelPosMin)) / taa_ushort(2);
+		OutputPixelPos.x = bIsValid ? OutputPixelPos.x : ~taa_ushort(0);
+
+		HistoryRejectionOutput[OutputPixelPos] = Rejection;
 	}
 	#endif
 
-	if (all(InputPixelPos < LowFrequencyInfo_ViewportMax))
 	{
+		bool bIsValid = all(InputPixelPos <= taa_short2(InputPixelPosMax));
+
+		taa_ushort2 OutputPixelPos = taa_ushort2(InputPixelPos - taa_short2(InputPixelPosMin));
+		OutputPixelPos.x = bIsValid ? OutputPixelPos.x : ~taa_ushort(0);
+
 		#if !CONFIG_OUTPUT_HALF_RES
 		{
-			HistoryRejectionOutput[(InputPixelPos - LowFrequencyInfo_ViewportMin)] = Rejection;
+			HistoryRejectionOutput[OutputPixelPos] = Rejection;
 		}
 		#endif
 
 		#if DEBUG_OUTPUT
 		{
-			DebugOutput[InputPixelPos - LowFrequencyInfo_ViewportMin] = Debug;
+			DebugOutput[OutputPixelPos] = Debug;
 		}
 		#endif
 	}
diff --git a/Engine/Source/Runtime/Renderer/Private/PostProcess/TemporalAA.cpp b/Engine/Source/Runtime/Renderer/Private/PostProcess/TemporalAA.cpp
--- a/Engine/Source/Runtime/Renderer/Private/PostProcess/TemporalAA.cpp
+++ b/Engine/Source/Runtime/Renderer/Private/PostProcess/TemporalAA.cpp
@@ -1657,7 +1657,7 @@ static void AddGen5MainTemporalAAPasses(
 				RDG_EVENT_NAME("TAA CompareHistory %dx%d", LowFrequencyRect.Width(), LowFrequencyRect.Height()),
 				ComputeShader,
 				PassParameters,
-				FComputeShaderUtils::GetGroupCount(LowFrequencyRect.Size(), 8));
+				FComputeShaderUtils::GetGroupCount(LowFrequencyRect.Size(), 16));
 		}
 	}
 
