Uses ClampPixelOffset() in TAA's FilterFrequencies pass

diff --git a/Engine/Shaders/Private/TemporalAA/TAAFilterFrequencies.usf b/Engine/Shaders/Private/TemporalAA/TAAFilterFrequencies.usf
--- a/Engine/Shaders/Private/TemporalAA/TAAFilterFrequencies.usf
+++ b/Engine/Shaders/Private/TemporalAA/TAAFilterFrequencies.usf
@@ -68,11 +68,11 @@ void MainCS(
 	uint GroupWaveIndex = GetGroupWaveIndex(GroupThreadIndex, /* GroupSize = */ TILE_SIZE * TILE_SIZE);
 
 	float4 Debug = 0.0;
-
-	uint2 InputPixelPos = (
-		LowFrequencyInfo_ViewportMin +
-		GroupId * uint2(TILE_SIZE, TILE_SIZE) +
-		ZOrder2D(GroupThreadIndex, log2(TILE_SIZE)));
+	
+	taa_short2 InputPixelPos = (
+		taa_short2(InputPixelPosMin) +
+		taa_short2(GroupId) * taa_short2(TILE_SIZE, TILE_SIZE) +
+		Map16x16Tile2x2Lane(GroupThreadIndex));
 
 	// Fetch center of the kernel
 	taa_half3 CenterInput;
@@ -120,19 +120,17 @@ void MainCS(
 	#endif
 
 	#if CONFIG_USE_DYNAMIC_SPREAD_FACTOR
-		int2 DynamicSpreadFactor = CONFIG_KERNEL_SPREAD_FACTOR + int2((InputPixelPos % 2) ^ (InputPixelPos.yx % 2));
+		taa_short2 DynamicSpreadFactor = (
+			taa_short(CONFIG_KERNEL_SPREAD_FACTOR) +
+			taa_short2((InputPixelPos % taa_short(2)) ^ (InputPixelPos.yx % taa_short(2))));
 	#else
-		int2 DynamicSpreadFactor = CONFIG_KERNEL_SPREAD_FACTOR;
+		taa_short2 DynamicSpreadFactor = taa_short(CONFIG_KERNEL_SPREAD_FACTOR);
 	#endif
 
 	// Fetch neighborhood of the kernel
 	UNROLL
 	for (uint SampleId = StartSampleId; SampleId < SampleCount; SampleId++)
 	{
-		#if CONFIG_KERNEL_1x1
-			continue;
-		#endif
-
 		int2 Offset;
 		taa_half SampleSpatialWeight;
 		
@@ -140,6 +138,7 @@ void MainCS(
 		{
 			Offset = 0;
 			SampleSpatialWeight = taa_half(1.0);
+			continue;
 		}
 		#elif CONFIG_KERNEL == KERNEL_3X3
 		{
@@ -162,24 +161,14 @@ void MainCS(
 				SampleSpatialWeight = taa_half(rcp(abs(Offset.x) + abs(Offset.y)));
 			#endif
 		}
-		#elif CONFIG_KERNEL == KERNEL_STACKOWIAK
-		{
-			Offset = kStackowiakSampleSet0[4 * SampleId + SampleTrackId];
-			
-			#if 1
-				SampleSpatialWeight = taa_half(1.0);
-			#elif 1
-				SampleSpatialWeight = taa_half(rcp(abs(Offset.x) + 1.0) * rcp(abs(Offset.y) + 1.0));
-			#else
-				SampleSpatialWeight = taa_half(rcp(abs(Offset.x) + abs(Offset.y)));
-			#endif
-		}
 		#else
 			#error Unknown kernel
 		#endif
-
-		int2 SampleInputPixelPos = int2(InputPixelPos) + DynamicSpreadFactor * Offset;
-		SampleInputPixelPos = clamp(SampleInputPixelPos, InputPixelPosMin, InputPixelPosMax);
+		
+		taa_short2 SampleInputPixelPos = InputPixelPos + ClampPixelOffset(
+			InputPixelPos,
+			DynamicSpreadFactor * Offset, Offset,
+			InputPixelPosMin, InputPixelPosMax);
 
 		taa_half3 InputSample = InputTexture[SampleInputPixelPos];
 		taa_half3 HistorySample = PredictionSceneColorTexture[SampleInputPixelPos];
@@ -209,13 +198,17 @@ void MainCS(
 		FilteredHistory *= SafeRcp(FilteredHistoryWeight);
 	}
 	
+	taa_short2 OutputPixelPos;
 	#if 1
-		uint2 OutputPixelPos = (
-			LowFrequencyInfo_ViewportMin + 
-			GroupId * uint2(TILE_SIZE, TILE_SIZE) +
-			ZOrder2D(GetGroupThreadIndex(GroupThreadIndex, GroupWaveIndex), log2(TILE_SIZE)));
+		ISOLATE
+		{
+			OutputPixelPos = (
+				taa_short2(InputPixelPosMin) +
+				taa_short2(GroupId) * taa_short2(TILE_SIZE, TILE_SIZE) +
+				Map16x16Tile2x2Lane(GetGroupThreadIndex(GroupThreadIndex, GroupWaveIndex)));
+		}
 	#else
-		uint2 OutputPixelPos = InputPixelPos;
+		OutputPixelPos = InputPixelPos;
 	#endif
 
 	#if CONFIG_ENABLE_STOCASTIC_QUANTIZATION
@@ -228,7 +221,7 @@ void MainCS(
 	}
 	#endif
 
-	OutputPixelPos.x = all(OutputPixelPos < LowFrequencyInfo_ViewportMax) ? OutputPixelPos.x : ~0;
+	OutputPixelPos.x = all(OutputPixelPos < LowFrequencyInfo_ViewportMax) ? OutputPixelPos.x : ~taa_short(0);
 
 	FilteredInputOutput[OutputPixelPos] = FilteredInput;
 	FilteredPredictionSceneColorOutput[OutputPixelPos] = FilteredHistory;
