Uses ClampPixelOffset() in TAA's postfilter and dilate rejection passes

diff --git a/Engine/Shaders/Private/TemporalAA/TAADilateRejection.usf b/Engine/Shaders/Private/TemporalAA/TAADilateRejection.usf
--- a/Engine/Shaders/Private/TemporalAA/TAADilateRejection.usf
+++ b/Engine/Shaders/Private/TemporalAA/TAADilateRejection.usf
@@ -22,20 +22,13 @@ void MainCS(
 	uint GroupThreadIndex : SV_GroupIndex)
 {
 	float4 Debug = 0.0;
-
-	uint2 DispatchThreadId = (
-		ZOrder2D(GroupThreadIndex, log2(TILE_SIZE)) +
-		GroupId * uint2(TILE_SIZE, TILE_SIZE));
-
-	uint2 RejectionPixelPos = DispatchThreadId;
-
-	float2 RejectionBufferUV = (RejectionPixelPos + 0.5) * RejectionInfo_ExtentInverse;
+	
+	taa_short2 RejectionPixelPos = (
+		taa_short2(GroupId) * taa_short2(TILE_SIZE, TILE_SIZE) +
+		Map8x8Tile2x2Lane(GroupThreadIndex));
 
 	// Fetch center of the kernel
-	taa_half HistoryRejectionInfo;
-	{
-		HistoryRejectionInfo = HistoryRejectionTexture.SampleLevel(GlobalPointClampedSampler, RejectionBufferUV, 0);
-	}
+	taa_half HistoryRejectionInfo = HistoryRejectionTexture[RejectionPixelPos];
 
 	const int KernelSize = 2;
 
@@ -45,34 +38,33 @@ void MainCS(
 	{
 		UNROLL
 		for (int y = -KernelSize; y <= KernelSize; y++)
-		//UNROLL
-		//for (uint i = 0; i < 9; i++)
 		{
-			//float2 Offset = float2(kOffsets3x3[i]);
-			float2 Offset = float2(x, y);
+			taa_short2 Offset = taa_short2(x, y);
 		
-			if (all(Offset == 0.0))
+			if (all(Offset == 0))
 			{
 				continue;
 			}
+			
+			taa_short2 SampleInputPixelPos = RejectionPixelPos + ClampPixelOffset(
+				RejectionPixelPos,
+				Offset, Offset,
+				RejectionInfo_ViewportMin, RejectionInfo_ViewportMax - 1);
 
-			float2 SampleRejectionBufferUV = RejectionBufferUV + Offset * RejectionInfo_ExtentInverse;
-			SampleRejectionBufferUV = clamp(SampleRejectionBufferUV, RejectionInfo_UVViewportBilinearMin, RejectionInfo_UVViewportBilinearMax);
-
-			taa_half SampleRejectionInfo = HistoryRejectionTexture.SampleLevel(GlobalPointClampedSampler, SampleRejectionBufferUV, 0);
+			taa_half SampleRejectionInfo = HistoryRejectionTexture[SampleInputPixelPos];
 
 			HistoryRejectionInfo = min(HistoryRejectionInfo, SampleRejectionInfo);
 		}
 	}
+	
+	bool bIsValid = all(RejectionPixelPos < taa_short2(RejectionInfo_ViewportMax));
+	RejectionPixelPos.x = bIsValid ? RejectionPixelPos.x : ~taa_ushort(0);
 
-	if (all(RejectionPixelPos < RejectionInfo_ViewportSize))
-	{
-		DilatedHistoryRejectionOutput[RejectionPixelPos] = HistoryRejectionInfo;
+	DilatedHistoryRejectionOutput[RejectionPixelPos] = HistoryRejectionInfo;
 
-		#if DEBUG_OUTPUT
-		{
-			DebugOutput[RejectionPixelPos] = Debug;
-		}
-		#endif
+	#if DEBUG_OUTPUT
+	{
+		DebugOutput[RejectionPixelPos] = Debug;
 	}
+	#endif
 }
diff --git a/Engine/Shaders/Private/TemporalAA/TAAPostfilterRejection.usf b/Engine/Shaders/Private/TemporalAA/TAAPostfilterRejection.usf
--- a/Engine/Shaders/Private/TemporalAA/TAAPostfilterRejection.usf
+++ b/Engine/Shaders/Private/TemporalAA/TAAPostfilterRejection.usf
@@ -72,37 +72,33 @@ void MainCS(
 {
 	float4 Debug = 0.0;
 
-	uint2 DispatchThreadId = (
-		ZOrder2D(GroupThreadIndex, log2(TILE_SIZE)) +
-		GroupId * uint2(TILE_SIZE, TILE_SIZE));
-
-	uint2 RejectionPixelPos = DispatchThreadId;
-
-	float2 RejectionBufferUV = (RejectionPixelPos + 0.5) * RejectionInfo_ExtentInverse;
+	taa_short2 RejectionPixelPos = (
+		taa_short2(GroupId) * taa_short2(TILE_SIZE, TILE_SIZE) +
+		Map8x8Tile2x2Lane(GroupThreadIndex));
 
 	// Fetch neighborhood of the kernel
 	taa_half InSamples[9];
 	UNROLL
 	for (uint i = 0; i < 9; i++)
 	{
-		float2 Offset = float2(kOffsets3x3[i]);
-		
-		float2 SampleRejectionBufferUV = RejectionBufferUV + Offset * RejectionInfo_ExtentInverse;
-		SampleRejectionBufferUV = clamp(SampleRejectionBufferUV, RejectionInfo_UVViewportBilinearMin, RejectionInfo_UVViewportBilinearMax);
+		taa_short2 SampleInputPixelPos = RejectionPixelPos + ClampPixelOffset(
+			RejectionPixelPos,
+			kOffsets3x3[i], kOffsets3x3[i],
+			RejectionInfo_ViewportMin, RejectionInfo_ViewportMax - 1);
 
-		InSamples[i] = HistoryRejectionTexture.SampleLevel(GlobalPointClampedSampler, SampleRejectionBufferUV, 0);
+		InSamples[i] = HistoryRejectionTexture[SampleInputPixelPos];
 	}
 
 	taa_half Rejection = Median9(InSamples);
 
-	if (all(RejectionPixelPos < RejectionInfo_ViewportSize))
-	{
-		HistoryRejectionOutput[RejectionPixelPos] = Rejection;
+	bool bIsValid = all(RejectionPixelPos < taa_short2(RejectionInfo_ViewportMax));
+	RejectionPixelPos.x = bIsValid ? RejectionPixelPos.x : ~taa_ushort(0);
 
-		#if DEBUG_OUTPUT
-		{
-			DebugOutput[RejectionPixelPos] = Debug;
-		}
-		#endif
+	HistoryRejectionOutput[RejectionPixelPos] = Rejection;
+
+	#if DEBUG_OUTPUT
+	{
+		DebugOutput[RejectionPixelPos] = Debug;
 	}
+	#endif
 }
