Reduce hair transmittance data footprint.

This saves memory and improves slightly performance ~5%.

diff --git a/Engine/Shaders/Private/DeferredLightPixelShaders.usf b/Engine/Shaders/Private/DeferredLightPixelShaders.usf
--- a/Engine/Shaders/Private/DeferredLightPixelShaders.usf
+++ b/Engine/Shaders/Private/DeferredLightPixelShaders.usf
@@ -306,7 +306,7 @@ void DeferredLightPixelMain(
 Texture2D<float4> ScreenShadowMaskSubPixelTexture;
 
 uint HairTransmittanceBufferMaxCount;
-StructuredBuffer<FPackedHairTransmittanceMask> HairTransmittanceBuffer;
+Buffer<uint> HairTransmittanceBuffer;
 uint HairShadowMaskValid;
 
 void DeferredLightPixelMain(
diff --git a/Engine/Shaders/Private/HairStrands/HairScatter.usf b/Engine/Shaders/Private/HairStrands/HairScatter.usf
--- a/Engine/Shaders/Private/HairStrands/HairScatter.usf
+++ b/Engine/Shaders/Private/HairStrands/HairScatter.usf
@@ -15,6 +15,7 @@
 #include "../ShadingModels.ush"
 #include "HairStrandsVisibilityCommon.ush"
 #include "HairStrandsDeepTransmittanceCommon.ush"
+#include "HairStrandsDeepTransmittanceDualScattering.ush" 
 #include "HairStrandsVoxelPageCommon.ush"
 
 #define DEBUG_ENABLE 0
@@ -62,15 +63,10 @@ FHairTransmittanceMask ComputeTransmittanceVoxel(
 	float3 WorldPosition,
 	float3 SamplePosition,
 	uint MacroGroupId,
-	FGBufferData GBuffer,
 	FTransmittanceSettings Settings,
-	Texture3D<float4>	HairLUTTexture,
-	SamplerState		HairLUTSampler,
 	bool bDebugEnabled)
 {
 	const float3 L = normalize(SamplePosition - WorldPosition);;
-	const float3 T = GBuffer.WorldNormal;
-	const float SinLightAngle = dot(L, T);
 
 	// Depth bias
 	// Origin is shifted voxels away towards the light + a constant bias of the size of the voxel
@@ -107,7 +103,10 @@ FHairTransmittanceMask ComputeTransmittanceVoxel(
 		VirtualVoxel.PageTexture,
 		TraversalSettings);
 
-	return ComputeTransmittanceFromHairCount(Result.HairCount, Result.Visibility, GBuffer, SinLightAngle, HairLUTTexture, HairLUTSampler);
+	FHairTransmittanceMask Out;
+	Out.HairCount = Result.HairCount;
+	Out.Visibility = Result.Visibility;
+	return Out;
 }
 
 FGBufferData GetGBufferData(uint2 PixelCoord)
@@ -254,10 +253,7 @@ void MainPS(
 					WorldPosition,
 					SampleWorldPos,
 					MacroGroupId,
-					GBuffer,
 					Settings,
-					HairLUTTexture,
-					LinearSampler,
 					bDebugEnabled && DEBUG_ENABLE);
 			}
 			#endif
@@ -325,8 +321,10 @@ void MainPS(
 			}
 			#endif
 
+			const float3 Transmittance = GetHairTransmittanceOnly(InitHairTransmittanceMask(TransmittanceMask.HairCount, TransmittanceMask.Visibility), GBuffer, SampleSinLightAngle, HairLUTTexture, LinearSampler);
+
 			const float Area = 0;
-			const float3 TransmittedLight = TransmittanceMask.Transmittance * TransmittanceMask.Visibility * SampleColor.xyz;
+			const float3 TransmittedLight = Transmittance * TransmittanceMask.Visibility * SampleColor.xyz;
 			float3 Luminance = HairShading(GBuffer, L, V, T, 1, TransmittanceData, 1, Area, uint2(0,0)) * TransmittedLight;
 			
 			AccIndirectLuminance.rgb += Luminance.xyz;
diff --git a/Engine/Shaders/Private/HairStrands/HairStrandsBsdfPlot.usf b/Engine/Shaders/Private/HairStrands/HairStrandsBsdfPlot.usf
--- a/Engine/Shaders/Private/HairStrands/HairStrandsBsdfPlot.usf
+++ b/Engine/Shaders/Private/HairStrands/HairStrandsBsdfPlot.usf
@@ -260,13 +260,8 @@ void MainPS(
 				// Compute the transmittance based on precompute Hair transmittance LUT
 				FHairAverageScattering AverageScattering = SampleHairLUT(HairScatteringLUTTexture, HairLUTSampler, P.BaseColor, P.Roughness, SinLightAngle);
 
-				FHairTransmittanceMask TransmittanceMask;
-				TransmittanceMask.HairCount = 0;
-				TransmittanceMask.Visibility = 1;
-				TransmittanceMask.Transmittance = 1;
-
 				TransmittanceData = ComputeDualScatteringTerms(
-					TransmittanceMask,
+					InitHairTransmittanceMask(),
 					AverageScattering,
 					P.Roughness,
 					SampleBacklit,
diff --git a/Engine/Shaders/Private/HairStrands/HairStrandsDeepTransmittanceCommon.ush b/Engine/Shaders/Private/HairStrands/HairStrandsDeepTransmittanceCommon.ush
--- a/Engine/Shaders/Private/HairStrands/HairStrandsDeepTransmittanceCommon.ush
+++ b/Engine/Shaders/Private/HairStrands/HairStrandsDeepTransmittanceCommon.ush
@@ -6,48 +6,41 @@
 
 struct FHairTransmittanceMask
 {
-	float3 Transmittance;
 	float HairCount;
 	float Visibility;
 };
 
-struct FPackedHairTransmittanceMask
-{
-	float3 Transmittance;
-	uint PackedHairCountAndVisibility;
-};
-
-FPackedHairTransmittanceMask InitNullPackedHairTransmittanceMask()
-{
-	FPackedHairTransmittanceMask Out;
-	Out.Transmittance = 0;
-	Out.PackedHairCountAndVisibility = 0;
-	return Out;
-}
-
 FHairTransmittanceMask InitHairTransmittanceMask()
 {
 	FHairTransmittanceMask Out;
-	Out.Transmittance = 1;
 	Out.HairCount = 0;
 	Out.Visibility = 1;
 	return Out;
 }
 
-FPackedHairTransmittanceMask PackTransmittanceMask(FHairTransmittanceMask In)
+FHairTransmittanceMask InitHairTransmittanceMask(float HairCount, float Visibility)
 {
-	FPackedHairTransmittanceMask Out;
-	Out.Transmittance = In.Transmittance;
-	Out.PackedHairCountAndVisibility = min(uint(In.HairCount * 1000), uint(0x00FFFFFF)) | (min(uint(In.Visibility * 0xFF), uint(0xFF)) << 24);
+	FHairTransmittanceMask Out;
+	Out.HairCount = HairCount;
+	Out.Visibility = Visibility;
 	return Out;
 }
 
-FHairTransmittanceMask UnpackTransmittanceMask(FPackedHairTransmittanceMask  In)
+uint PackTransmittanceMask(FHairTransmittanceMask In)
+{
+	return min(uint(In.HairCount * 1000), uint(0x00FFFFFF)) | (min(uint(In.Visibility * 0xFF), uint(0xFF)) << 24);
+}
+
+uint InitNullPackedHairTransmittanceMask()
+{
+	return 0;
+}
+
+FHairTransmittanceMask UnpackTransmittanceMask(uint In)
 {
 	FHairTransmittanceMask Out;
-	Out.Transmittance = In.Transmittance;
-	Out.HairCount = float(In.PackedHairCountAndVisibility & 0x00FFFFFF) / 1000.f;
-	Out.Visibility = float((In.PackedHairCountAndVisibility & 0xFF000000) >> 24) / 255.f;
+	Out.HairCount = float(In & 0x00FFFFFF) / 1000.f;
+	Out.Visibility = float((In & 0xFF000000) >> 24) / 255.f;
 	return Out;
 }
 
@@ -68,18 +61,4 @@ struct FTransmittanceSettings
 	uint	 DebugMode;
 
 	uint TransmittanceKernelType;
-};
-
-FHairTransmittanceMask ComputeTransmittanceFromHairCount(float HairCount, float HairVisibility, FGBufferData GBuffer, float SinLightAngle, Texture3D<float4> InHairScatteringLUTTexture, SamplerState InHairLUTSampler)
-{
-	const FHairAverageScattering AbsorptionData = SampleHairLUT(InHairScatteringLUTTexture, InHairLUTSampler, GBuffer.BaseColor, GBuffer.Roughness, SinLightAngle);
-
-	// Always shift the hair count by one to remove self-occlusion/shadow aliasing and have smootheer transition
-	// This insure the the pow function always starts at 0 for front facing hair
-	FHairTransmittanceMask Out;
-	Out.HairCount = max(0, HairCount - 1);
-	Out.Transmittance = pow(AbsorptionData.A_front, Out.HairCount);
-	Out.Visibility = HairVisibility;
-
-	return Out;
-}
\ No newline at end of file
+};
\ No newline at end of file
diff --git a/Engine/Shaders/Private/HairStrands/HairStrandsDeepTransmittanceDualScattering.ush b/Engine/Shaders/Private/HairStrands/HairStrandsDeepTransmittanceDualScattering.ush
--- a/Engine/Shaders/Private/HairStrands/HairStrandsDeepTransmittanceDualScattering.ush
+++ b/Engine/Shaders/Private/HairStrands/HairStrandsDeepTransmittanceDualScattering.ush
@@ -59,12 +59,16 @@ FHairTransmittanceData ComputeDualScatteringTerms(
 	const float df = 0.7f;
 	const float db = 0.7f;
 
+	// Always shift the hair count by one to remove self-occlusion/shadow aliasing and have smoother transition
+	// This insure the the pow function always starts at 0 for front facing hair
+	const uint HairCount = max(0, TransmittanceMask.HairCount - 1);
+
 	// This is a coarse approximation of eq. 13. Normally, Beta_f should be weighted by the 'normalized' 
 	// R, TT, and TRT terms 
 	const float3 af_weights = af / (af.r + af.g + af.b);
 	const float3 Beta_f  = dot(float3(Beta_R, Beta_TT, Beta_TRT), af_weights);
 	const float3 Beta_f2 = Beta_f*Beta_f;
-	const float3 sigma_f2 = Beta_f2 * max(1, TransmittanceMask.HairCount);
+	const float3 sigma_f2 = Beta_f2 * max(1, HairCount);
 
 	const float Theta_d = asin(SinThetaL) + asin(SinThetaV);
 	const float Theta_h = Theta_d * 0.5f;
@@ -73,7 +77,7 @@ FHairTransmittanceData ComputeDualScatteringTerms(
 	float3 Sf = float3(	Hair_g2(sigma_f2.r, Theta_h), 
 						Hair_g2(sigma_f2.g, Theta_h), 
 						Hair_g2(sigma_f2.b, Theta_h)) / PI;
-	const float3 Tf = TransmittanceMask.Transmittance;
+	const float3 Tf = pow(AverageScattering.A_front, HairCount);
 
 	// Overall shift due to the various local scatteing event (all shift & roughnesss should vary with color)
 	const float3 shift_f = dot(float3(Shift_R, Shift_TT, Shift_TRT), af_weights);
@@ -95,7 +99,7 @@ FHairTransmittanceData ComputeDualScatteringTerms(
 
 
 	// Different variant for managing sefl-occlusion issue for global scattering
-	const float3 GlobalScattering = lerp(1, Tf * Sf * df, saturate(TransmittanceMask.HairCount));
+	const float3 GlobalScattering = lerp(1, Tf * Sf * df, saturate(HairCount));
 	const float3 LocalScattering  = 2 * Ab * Sb * db;
 	
 	// Decode if the Local/Global component should be enabled or not (for debug purpose)
@@ -147,7 +151,7 @@ FHairTransmittanceData GetHairTransmittance(
 	FGBufferData GBuffer,
 	const uint SampleOffset,
 	const uint HairTransmittanceBufferMaxCount,
-	StructuredBuffer<FPackedHairTransmittanceMask> HairTransmittanceBuffer,
+	Buffer<uint> HairTransmittanceBuffer,
 	Texture3D<float4> HairLUTTexture,
 	SamplerState HairLUTSampler,
 	const uint HairComponents)
@@ -155,8 +159,7 @@ FHairTransmittanceData GetHairTransmittance(
 	FHairTransmittanceMask TransmittanceMask = InitHairTransmittanceMask();
 	if (SampleOffset < HairTransmittanceBufferMaxCount)
 	{
-		const FPackedHairTransmittanceMask Packed = HairTransmittanceBuffer[SampleOffset];
-		TransmittanceMask = UnpackTransmittanceMask(Packed);
+		TransmittanceMask = UnpackTransmittanceMask(HairTransmittanceBuffer[SampleOffset]);
 	}
 	return GetHairTransmittance(
 		V,
@@ -168,6 +171,16 @@ FHairTransmittanceData GetHairTransmittance(
 		HairComponents);
 }
 
+float3 GetHairTransmittanceOnly(FHairTransmittanceMask TransmittanceMask, FGBufferData GBuffer, float SinLightAngle, Texture3D<float4> InHairScatteringLUTTexture, SamplerState InHairLUTSampler)
+{
+	// Always shift the hair count by one to remove self-occlusion/shadow aliasing and have smoother transition
+	// This insure the the pow function always starts at 0 for front facing hair
+	const uint HairCount = max(0, TransmittanceMask.HairCount - 1);
+
+	const FHairAverageScattering AbsorptionData = SampleHairLUT(InHairScatteringLUTTexture, InHairLUTSampler, GBuffer.BaseColor, GBuffer.Roughness, SinLightAngle);
+	return pow(AbsorptionData.A_front, HairCount) * TransmittanceMask.Visibility;
+}
+
 // Evaluate local scattering scattering
 FHairTransmittanceData EvaluateDualScattering(inout FGBufferData GBuffer, float3 V, float3 L)
 {
diff --git a/Engine/Shaders/Private/HairStrands/HairStrandsDeepTransmittanceMask.usf b/Engine/Shaders/Private/HairStrands/HairStrandsDeepTransmittanceMask.usf
--- a/Engine/Shaders/Private/HairStrands/HairStrandsDeepTransmittanceMask.usf
+++ b/Engine/Shaders/Private/HairStrands/HairStrandsDeepTransmittanceMask.usf
@@ -76,7 +76,7 @@ Buffer<uint4>				IndirectArgsBuffer;
 
 Texture3D<float4>			HairLUTTexture;
 SamplerState				LinearSampler;
-RWStructuredBuffer<FPackedHairTransmittanceMask> OutputColor;
+RWBuffer<uint>				OutputColor;
 
 #if PERMUTATION_GROUP_SIZE == 0
 	#define TRANSMITTANCE_GROUP_SIZE 64
@@ -88,6 +88,7 @@ RWStructuredBuffer<FPackedHairTransmittanceMask> OutputColor;
 
 #if PERMUTATION_TRANSMITTANCE_TYPE == PERMUTATION_TRANSMITTANCE_DEEP_SHADOW
 
+#if 0
 // Sample deep shadow with PCF kernel and doing transmittance evaluation for each sample
 FHairTransmittanceMask SampleDOM_PCF_Accurate(
 	float3 LightSpacePosition,
@@ -164,16 +165,14 @@ FHairTransmittanceMask SampleDOM_PCF_Accurate(
 
 	return Out;
 }
+#endif 
 
 FHairTransmittanceMask ComputeTransmittance(
 	float3 WorldPosition,
-	FGBufferData GBuffer, 
 	FTransmittanceSettings Settings,
 	float4				InLayerDepths,
 	Texture2D<float>	FrontDepthTexture,
-	Texture2D<float4>	DomTexture,
-	Texture3D<float4>	HairLUTTexture,
-	SamplerState		LinearSampler)
+	Texture2D<float4>	DomTexture)
 {
 	FHairTransmittanceMask Out = InitHairTransmittanceMask();
 	{
@@ -182,8 +181,6 @@ FHairTransmittanceMask ComputeTransmittance(
 		// LightSpacePosition is 'WorldPosition' in LightSpace (TexelCoord.XY / NDC Depth.Z)
 		const float3 LightSpacePosition = ToLightPosition(WorldPosition, Settings.WorldToLightTransform) * float3(Settings.DeepShadowResolution, 1) + float3(Settings.DeepShadowOffset, 0); 
 		const float3 L = bIsDirectionalLight ? Settings.LightDirection : normalize(Settings.LightPosition.xyz - WorldPosition);
-		const float3 T = GBuffer.WorldNormal;
-		const float SinLightAngle = dot(L, T);
 
 		float DepthBias = 0;
 		if (Settings.DeepShadowDepthBiasScale > 0)
@@ -200,7 +197,9 @@ FHairTransmittanceMask ComputeTransmittance(
 		// Compute the number of hair count between light & shading point
 		if (Settings.TransmittanceKernelType == 4)
 		{
-			Out = SampleDOM_PCF_Accurate(LightSpacePosition.xyz, DepthBias, InLayerDepths, FrontDepthTexture, DomTexture, GBuffer, SinLightAngle, HairLUTTexture, LinearSampler);
+			//Out = SampleDOM_PCF_Accurate(LightSpacePosition.xyz, DepthBias, InLayerDepths, FrontDepthTexture, DomTexture, GBuffer, SinLightAngle, HairLUTTexture, LinearSampler);
+			Out.HairCount = 0;
+			Out.Visibility = 1;
 		}
 		else
 		{
@@ -221,12 +220,12 @@ FHairTransmittanceMask ComputeTransmittance(
 			{
 				HairCount = 0;
 			}
-			Out = ComputeTransmittanceFromHairCount(HairCount * Settings.DeepShadowDensityScale, 1, GBuffer, SinLightAngle, HairLUTTexture, LinearSampler);
+			Out.HairCount = HairCount * Settings.DeepShadowDensityScale;
+			Out.Visibility = 1;
 		}
 
 		if (Settings.DebugMode != 0)
 		{
-			Out.Transmittance = ComputeDomDebugColor(WorldPosition, Settings.WorldToLightTransform, InLayerDepths, FrontDepthTexture, LinearSampler);
 			Out.HairCount = 1;
 		}
 	}
@@ -260,12 +259,6 @@ void MainCS(uint2 DispatchThreadId : SV_DispatchThreadID)
 			return;
 		}
 
-		FGBufferData GBuffer;
-		GBuffer.ShadingModelID = SHADINGMODELID_HAIR;
-		GBuffer.WorldNormal = Sample.Tangent;
-		GBuffer.BaseColor = Sample.BaseColor;
-		GBuffer.Roughness = Sample.Roughness;
-
 		const uint MacroGroupIndex = clamp(Sample.MacroGroupId, 0, MAX_HAIR_MACROGROUP_COUNT - 1);
 		const int3 SlotOffset_SlotIndex = DeepShadow_AtlasSlotOffsets_AtlasSlotIndex[MacroGroupIndex].xyz;
 		float4x4 WorldToLightTransform = 0;
@@ -293,10 +286,9 @@ void MainCS(uint2 DispatchThreadId : SV_DispatchThreadID)
 
 		const float SceneDepth = ConvertFromDeviceZ(Sample.Depth);
 		const float3 WorldPosition = ToWorldPosition(UV, SceneDepth);
-		const FHairTransmittanceMask TransmittanceMask = ComputeTransmittance(WorldPosition, GBuffer, Settings, DeepShadow_LayerDepths, DeepShadow_FrontDepthTexture, DeepShadow_DomTexture, HairLUTTexture, LinearSampler);
+		const FHairTransmittanceMask TransmittanceMask = ComputeTransmittance(WorldPosition, Settings, DeepShadow_LayerDepths, DeepShadow_FrontDepthTexture, DeepShadow_DomTexture);
 
-		const FPackedHairTransmittanceMask PackedTransmittanceMask = PackTransmittanceMask(TransmittanceMask);
-		OutputColor[LocalOffset] = PackedTransmittanceMask;
+		OutputColor[LocalOffset] = PackTransmittanceMask(TransmittanceMask);
 	}
 }
 #endif
@@ -352,10 +344,7 @@ FHairTransmittanceMask ComputeTransmittanceVoxel(
 	const float PixelRadius,
 	float3 WorldPosition,
 	uint MacroGroupId,
-	FGBufferData GBuffer,
 	FTransmittanceSettings Settings,
-	Texture3D<float4>	HairLUTTexture,
-	SamplerState		HairLUTSampler,
 	bool bDebugEnabled)
 {
 
@@ -364,8 +353,6 @@ FHairTransmittanceMask ComputeTransmittanceVoxel(
 	const float3 LightPosition  = bIsDirectionalLight ? -Settings.LightDirection*VoxelShadowMaxDistance + WorldPosition : Settings.LightPosition.xyz;
 	const float3 LightDirection = bIsDirectionalLight ? -Settings.LightDirection : normalize(Settings.LightPosition.xyz - WorldPosition);
 	const float TanLightAngle = bIsDirectionalLight ? 0 : LightRadius / length(Settings.LightPosition.xyz - WorldPosition);
-	const float3 T = GBuffer.WorldNormal;
-	const float SinLightAngle = dot(LightDirection, T);
 
 	// Compute the number of hair count between light & shading point
 	const float DistanceThreshold = 100000;
@@ -427,7 +414,10 @@ FHairTransmittanceMask ComputeTransmittanceVoxel(
 	AvgHairVisility /= max(1u, SampleCount);
 	#endif
 
-	return ComputeTransmittanceFromHairCount(AvgHairCount, AvgHairVisility, GBuffer, SinLightAngle, HairLUTTexture, HairLUTSampler);
+	FHairTransmittanceMask Out;
+	Out.HairCount = AvgHairCount;
+	Out.Visibility = AvgHairVisility;
+	return Out;
 }
 
 [numthreads(TRANSMITTANCE_GROUP_SIZE, 1, 1)]
@@ -446,9 +436,8 @@ void MainCS(uint2 DispatchThreadId : SV_DispatchThreadID)
 	if (RayMarchMaskTexture.Load(uint3(PixelCoord, 0)).r == 0.0f)
 	{
 		// This pixel is already fully shadowed opaque. (RayMarchMaskTexture is computed using the closest sample as of today)
-		FHairTransmittanceMask DefaultMask = InitHairTransmittanceMask();
-		const FPackedHairTransmittanceMask PackedTransmittanceMask = PackTransmittanceMask(DefaultMask);
-		OutputColor[LocalOffset] = PackedTransmittanceMask;
+		const FHairTransmittanceMask DefaultMask = InitHairTransmittanceMask();
+		OutputColor[LocalOffset] = PackTransmittanceMask(DefaultMask);
 		return;
 	}
 
@@ -489,17 +478,16 @@ void MainCS(uint2 DispatchThreadId : SV_DispatchThreadID)
 
 		const float PixelRadius = SceneDepth * VirtualVoxel.HairCoveragePixelRadiusAtDepth1;
 	
-		const FHairTransmittanceMask TransmittanceMask = ComputeTransmittanceVoxel(PixelCoord, PixelRadius, WorldPosition, MacroGroupIndex, GBuffer, Settings, HairLUTTexture, LinearSampler, bDebugEnabled);
+		const FHairTransmittanceMask Out = ComputeTransmittanceVoxel(PixelCoord, PixelRadius, WorldPosition, MacroGroupIndex, Settings, bDebugEnabled);
 
 		#if VOXEL_TRAVERSAL_DEBUG
 		if (bDebugEnabled)
 		{
-			PrintInfo(LocalOffset, PixelCoord, TransmittanceMask);
+			PrintInfo(LocalOffset, PixelCoord, Out);
 		}
 		#endif	
 
-		const FPackedHairTransmittanceMask PackedTransmittanceMask = PackTransmittanceMask(TransmittanceMask);
-		OutputColor[LocalOffset] = PackedTransmittanceMask;
+		OutputColor[LocalOffset] = PackTransmittanceMask(Out);
 	}
 }
 #endif
@@ -511,7 +499,7 @@ void MainCS(uint2 DispatchThreadId : SV_DispatchThreadID)
 #if SHADER_CLEAR
 
 uint ElementCount;
-RWStructuredBuffer<FPackedHairTransmittanceMask> OutputMask;
+RWBuffer<uint> OutputMask;
 
 [numthreads(64, 1, 1)]
 void MainCS(uint2 DispatchThreadId : SV_DispatchThreadID)
diff --git a/Engine/Shaders/Private/HairStrands/HairStrandsEnvironmentLighting.usf b/Engine/Shaders/Private/HairStrands/HairStrandsEnvironmentLighting.usf
--- a/Engine/Shaders/Private/HairStrands/HairStrandsEnvironmentLighting.usf
+++ b/Engine/Shaders/Private/HairStrands/HairStrandsEnvironmentLighting.usf
@@ -212,6 +212,11 @@ FHairTraversalResult InternalComputeHairCountVoxel(
 #define ADHOC_SAMPLING_TRT 1
 #define ADHOC_SAMPLING_MS 1
 
+float MaxComponent3(float3 In)
+{
+	return max(In.x, max(In.y, In.z));
+}
+
 float3 ReflectionEnvironment(FGBufferData GBuffer, float2 ScreenPosition, float2 SvPosition, uint PrimitiveId, uint MacroGroupId)
 {
 	const float PreExposure = GetPreExposure();
@@ -278,18 +283,17 @@ float3 ReflectionEnvironment(FGBufferData GBuffer, float2 ScreenPosition, float2
 			// Compute the number of hair count between light & shading point
 			const float SinLightAngle = dot(SampleL, T);
 			const FHairTraversalResult Result = InternalComputeHairCountVoxel(MacroGroupId, SampleWorldPosition, SampleLightPosition, 1, 0, MaxSearchDistance, PixelRadius, true, bDebugEnabled);
-			const FHairTransmittanceMask TransmittanceMask = ComputeTransmittanceFromHairCount(Result.HairCount, Result.Visibility, GBuffer, SinLightAngle, HairScatteringLUTTexture, HairLUTSampler);
-
-			const float ScalarTransmittance = max(TransmittanceMask.Transmittance.x, max(TransmittanceMask.Transmittance.y, TransmittanceMask.Transmittance.z));
-			AvgR += SampleL * ScalarTransmittance * Result.Visibility;
 
 			const uint ScatteringComponent = (HairComponents & HAIR_COMPONENT_R) | (HairComponents & HAIR_COMPONENT_TRT) | HAIR_COMPONENT_LS | HAIR_COMPONENT_GS;// | HAIR_MULTISCATTER;
-			const FHairTransmittanceData TransmittanceData = GetHairTransmittance(V, SampleL, GBuffer, TransmittanceMask, HairScatteringLUTTexture, HairLUTSampler, ScatteringComponent);
+			const FHairTransmittanceData TransmittanceData = GetHairTransmittance(V, SampleL, GBuffer, InitHairTransmittanceMask(Result.HairCount, Result.Visibility), HairScatteringLUTTexture, HairLUTSampler, ScatteringComponent);
 
 			const float Area = 0.2;
 			const float BackLit = 1;
 			//AvgDualScattering += HairShading(GBuffer, SampleL, V, T, 1, TransmittanceData, BackLit, Area, uint2(0, 0));
 			AvgDualScattering += TransmittanceData.LocalScattering * TransmittanceData.GlobalScattering * Result.Visibility;
+
+			const float TransmittedVisibility = MaxComponent3(TransmittanceData.GlobalScattering * Result.Visibility);
+			AvgR += SampleL * TransmittedVisibility;
 		}
 		const float SphereInt = 4 * PI;
 		AvgDualScattering *= SphereInt;
@@ -396,7 +400,8 @@ float3 ReflectionEnvironment(FGBufferData GBuffer, float2 ScreenPosition, float2
 		const float SinViewAngle = dot(V, T);
 
 		// Compute the transmittance
-		FHairTransmittanceMask TransmittanceMask = InitHairTransmittanceMask();
+		float3 Transmittance = 1;
+		float Visibility = 1;
 		{
 			const float3 SampleBias = ComputePositionBias(L, VoxelOffset, VirtualVoxel.VoxelWorldSize, VirtualVoxel.DepthBiasScale_Environment);
 			const float3 SampleWorldPosition = WorldPosition + SampleBias;
@@ -404,10 +409,10 @@ float3 ReflectionEnvironment(FGBufferData GBuffer, float2 ScreenPosition, float2
 
 			// Compute the number of hair count between light & shading point
 			const FHairTraversalResult Result = InternalComputeHairCountVoxel(MacroGroupId, SampleWorldPosition, SampleLightPosition, 1, 0, MaxSearchDistance, PixelRadius, true, bDebugEnabled);
-			TransmittanceMask = ComputeTransmittanceFromHairCount(Result.HairCount, Result.Visibility, GBuffer, SinLightAngle, HairScatteringLUTTexture, HairLUTSampler);
+			Transmittance = GetHairTransmittanceOnly(InitHairTransmittanceMask(Result.HairCount, Result.Visibility), GBuffer, SinLightAngle, HairScatteringLUTTexture, HairLUTSampler);
+			Visibility = Result.Visibility;
 		}
 
-		FHairTransmittanceData TransmittanceData = InitHairStrandsTransmittanceData();
 
 		// Directional albedo
 		#if 0
@@ -419,13 +424,15 @@ float3 ReflectionEnvironment(FGBufferData GBuffer, float2 ScreenPosition, float2
 
 		// BSDF
 		#if 1
-		{			
-			TransmittanceData.GlobalScattering = TransmittanceMask.Transmittance;
+		{
+			FHairTransmittanceData TransmittanceData = InitHairStrandsTransmittanceData();
+			TransmittanceData.OpaqueVisibility = Visibility;
+			TransmittanceData.GlobalScattering = Transmittance;
 			TransmittanceData.ScatteringComponent = (HairComponents & HAIR_COMPONENT_R) | (HairComponents & HAIR_COMPONENT_TRT) | HAIR_COMPONENT_MULTISCATTER;
 
 			const float Area = 0.2;
 			const float BackLit = 0;
-			ColorR.rgb *= HairShading(GBuffer, L, V, T, 1, TransmittanceData, BackLit, Area, uint2(0, 0)) * TransmittanceMask.Visibility;
+			ColorR.rgb *= HairShading(GBuffer, L, V, T, 1, TransmittanceData, BackLit, Area, uint2(0, 0)) * Visibility;
 		}
 		#endif
 	}
@@ -449,7 +456,8 @@ float3 ReflectionEnvironment(FGBufferData GBuffer, float2 ScreenPosition, float2
 		//ColorTT.rgb += PreExposure * GetSkySHDiffuse(L) * View.SkyLightColor.rgb;
 
 		// Compute the transmittance
-		FHairTransmittanceMask TransmittanceMask = InitHairTransmittanceMask();
+		float3 Transmittance = 1;
+		float Visibility = 1;
 		{
 			const float3 SampleBias = ComputePositionBias(L, VoxelOffset, VirtualVoxel.VoxelWorldSize, VirtualVoxel.DepthBiasScale_Environment);
 			const float3 SampleWorldPosition = WorldPosition + SampleBias;
@@ -468,7 +476,8 @@ float3 ReflectionEnvironment(FGBufferData GBuffer, float2 ScreenPosition, float2
 			{
 				Result = InternalComputeHairCountVoxel(MacroGroupId, SampleWorldPosition, SampleLightPosition, TransmissionDensityScaleFactor, 0, DistanceThresholdTT, PixelRadius, false, bDebugEnabled);
 			}
-			TransmittanceMask = ComputeTransmittanceFromHairCount(Result.HairCount, Result.Visibility, GBuffer, SinLightAngle, HairScatteringLUTTexture, HairLUTSampler);
+			Transmittance = GetHairTransmittanceOnly(Result.HairCount, Result.Visibility, GBuffer, SinLightAngle, HairScatteringLUTTexture, HairLUTSampler);
+			Visibility = Result.Visibility;
 		}
 
 		// BSDF with dual-scattering
@@ -487,7 +496,7 @@ float3 ReflectionEnvironment(FGBufferData GBuffer, float2 ScreenPosition, float2
 		{
 			const float Backlit = GBuffer.CustomData.z;
 			FHairAverageEnergy Energy = SampleHairEnergyLUT(HairEnergyLUTTexture, HairLUTSampler, GBuffer.BaseColor, GBuffer.Roughness, SinViewAngle);
-			ColorTT.rgb *= Energy.A_TT * TransmittanceMask.Transmittance * TransmittanceMask.Visibility * Backlit;
+			ColorTT.rgb *= Energy.A_TT * Transmittance * Visibility * Backlit;
 		}
 		#endif	
 
@@ -581,19 +590,11 @@ float3 ReflectionEnvironmentUniformSampling(
 			true, 
 			false/*bDebugEnabled*/);
 
-		const FHairTransmittanceMask TransmittanceMask = ComputeTransmittanceFromHairCount(
-			Result.HairCount, 
-			Result.Visibility, 
-			GBuffer, 
-			SinLightAngle, 
-			HairScatteringLUTTexture, 
-			HairLUTSampler);
-
 		FHairTransmittanceData TransmittanceData = GetHairTransmittance(
 			V,
 			SampleL,
 			GBuffer,
-			TransmittanceMask,
+			InitHairTransmittanceMask(Result.HairCount, Result.Visibility),
 			HairScatteringLUTTexture,
 			HairLUTSampler,
 			HairComponents);	
@@ -643,15 +644,11 @@ float3 ReflectionEnvironmentSH(
 	const float3 N = T;
 	const float3 L = normalize(V - N * dot(V, N));
 
-	FHairTransmittanceMask TransmittanceMask;
-	TransmittanceMask.Transmittance = 1;
-	TransmittanceMask.HairCount = 0;
-	TransmittanceMask.Visibility = 1;
 	FHairTransmittanceData TransmittanceData = GetHairTransmittance(
 		V,
 		L,
 		GBuffer,
-		TransmittanceMask,
+		InitHairTransmittanceMask(),
 		HairScatteringLUTTexture,
 		HairLUTSampler,
 		HairComponents);
diff --git a/Engine/Source/Runtime/Renderer/Private/HairStrands/HairStrandsTransmittance.cpp b/Engine/Source/Runtime/Renderer/Private/HairStrands/HairStrandsTransmittance.cpp
--- a/Engine/Source/Runtime/Renderer/Private/HairStrands/HairStrandsTransmittance.cpp
+++ b/Engine/Source/Runtime/Renderer/Private/HairStrands/HairStrandsTransmittance.cpp
@@ -93,7 +93,7 @@ class FHairStrandsClearTransmittanceMaskCS : public FGlobalShader
 
 	BEGIN_SHADER_PARAMETER_STRUCT(FParameters, )
 		SHADER_PARAMETER(uint32, ElementCount)
-		SHADER_PARAMETER_RDG_BUFFER_UAV(RWStructuredBuffer, OutputMask)
+		SHADER_PARAMETER_RDG_BUFFER_UAV(RWBuffer, OutputMask)
 	END_SHADER_PARAMETER_STRUCT()
 
 public:
@@ -118,7 +118,7 @@ static void AddHairStrandsClearTransmittanceMaskPass(
 {
 	FHairStrandsClearTransmittanceMaskCS::FParameters* Parameters = GraphBuilder.AllocParameters<FHairStrandsClearTransmittanceMaskCS::FParameters>();
 	Parameters->ElementCount = OutTransmittanceMask->Desc.NumElements;
-	Parameters->OutputMask = GraphBuilder.CreateUAV(OutTransmittanceMask);
+	Parameters->OutputMask = GraphBuilder.CreateUAV(OutTransmittanceMask, FHairStrandsTransmittanceMaskData::Format);
 
 	FHairStrandsClearTransmittanceMaskCS::FPermutationDomain PermutationVector;
 	TShaderMapRef<FHairStrandsClearTransmittanceMaskCS> ComputeShader(ShaderMap, PermutationVector);
@@ -135,8 +135,8 @@ static void AddHairStrandsClearTransmittanceMaskPass(
 
 static FRDGBufferRef CreateHairStrandsTransmittanceMaskBuffer(FRDGBuilder& GraphBuilder, uint32 NumElements)
 {
-	return GraphBuilder.CreateBuffer(FRDGBufferDesc::CreateStructuredDesc(
-		4 * sizeof(float),
+	return GraphBuilder.CreateBuffer(FRDGBufferDesc::CreateBufferDesc(
+		sizeof(uint32),
 		NumElements),
 		TEXT("Hair.TransmittanceNodeData"));
 }
@@ -202,7 +202,7 @@ class FDeepTransmittanceMaskCS : public FGlobalShader
 		SHADER_PARAMETER_RDG_BUFFER_SRV(StructuredBuffer, HairVisibilityNodeCoord)
 		RDG_BUFFER_ACCESS(IndirectArgsBuffer, ERHIAccess::IndirectArgs)
 
-		SHADER_PARAMETER_RDG_BUFFER_UAV(RWStructuredBuffer, OutputColor)
+		SHADER_PARAMETER_RDG_BUFFER_UAV(RWBuffer, OutputColor)
 		SHADER_PARAMETER_SAMPLER(SamplerState, LinearSampler)
 		SHADER_PARAMETER_SAMPLER(SamplerState, ShadowSampler)
 
@@ -288,7 +288,7 @@ static FRDGBufferRef AddDeepShadowTransmittanceMaskPass(
 	Parameters->DeepShadow_FrontDepthTexture = Params.DeepShadow_FrontDepthTexture;
 	Parameters->DeepShadow_DomTexture = Params.DeepShadow_DomTexture;
 	Parameters->LinearSampler = TStaticSamplerState<SF_Bilinear, AM_Clamp, AM_Clamp, AM_Clamp>::GetRHI();
-	Parameters->OutputColor = GraphBuilder.CreateUAV(OutBuffer);
+	Parameters->OutputColor = GraphBuilder.CreateUAV(OutBuffer, FHairStrandsTransmittanceMaskData::Format);
 	Parameters->LightChannelMask = Params.LightChannelMask;
 	Parameters->DeepShadow_Resolution = Params.DeepShadow_Resolution;
 	Parameters->LightDirection = Params.LightDirection;
diff --git a/Engine/Source/Runtime/Renderer/Private/HairStrands/HairStrandsTransmittance.h b/Engine/Source/Runtime/Renderer/Private/HairStrands/HairStrandsTransmittance.h
--- a/Engine/Source/Runtime/Renderer/Private/HairStrands/HairStrandsTransmittance.h
+++ b/Engine/Source/Runtime/Renderer/Private/HairStrands/HairStrandsTransmittance.h
@@ -15,6 +15,7 @@ class FViewInfo;
 
 struct FHairStrandsTransmittanceMaskData
 {
+	static const EPixelFormat Format = PF_R32_UINT;
 	FRDGBufferRef TransmittanceMask = nullptr;
 };
 
diff --git a/Engine/Source/Runtime/Renderer/Private/LightRendering.cpp b/Engine/Source/Runtime/Renderer/Private/LightRendering.cpp
--- a/Engine/Source/Runtime/Renderer/Private/LightRendering.cpp
+++ b/Engine/Source/Runtime/Renderer/Private/LightRendering.cpp
@@ -2441,7 +2441,7 @@ void FDeferredShadingSceneRenderer::RenderLightForHair(
 		check(InTransmittanceMaskData.TransmittanceMask);
 
 		PassParameters->HairStrands = HairStrands::BindHairStrandsViewUniformParameters(View);
-		PassParameters->HairTransmittanceMaskSRV = GraphBuilder.CreateSRV(InTransmittanceMaskData.TransmittanceMask);
+		PassParameters->HairTransmittanceMaskSRV = GraphBuilder.CreateSRV(InTransmittanceMaskData.TransmittanceMask, FHairStrandsTransmittanceMaskData::Format);
 
 		const bool bIsShadowMaskValid = !!PassParameters->Light.ShadowMaskTexture;
 		const uint32 MaxTransmittanceElementCount = InTransmittanceMaskData.TransmittanceMask ? InTransmittanceMaskData.TransmittanceMask->Desc.NumElements : 0;
