[Switch] Shader compile fixes for "max(uint, int)" errors.

diff --git a/Engine/Shaders/Private/DistanceFieldScreenGridLighting.usf b/Engine/Shaders/Private/DistanceFieldScreenGridLighting.usf
--- a/Engine/Shaders/Private/DistanceFieldScreenGridLighting.usf
+++ b/Engine/Shaders/Private/DistanceFieldScreenGridLighting.usf
@@ -92,7 +92,7 @@ void HemisphereConeTraceAgainstTileCulledObjects(uint ObjectIndex, uint OutputBa
 
 			uint NumMips = LoadDFAssetData(DFObjectData.AssetIndex, 0).NumMips;
 			// Always trace against medium resolution mip
-			FDFAssetData DFAssetData = LoadDFAssetData(DFObjectData.AssetIndex, min(1, NumMips - 1));
+			FDFAssetData DFAssetData = LoadDFAssetData(DFObjectData.AssetIndex, min(1u, NumMips - 1));
 
 			LOOP
 			for (uint ConeIndex = 0; ConeIndex < NUM_CONE_DIRECTIONS; ConeIndex++)
diff --git a/Engine/Shaders/Private/DistanceFieldVisualization.usf b/Engine/Shaders/Private/DistanceFieldVisualization.usf
--- a/Engine/Shaders/Private/DistanceFieldVisualization.usf
+++ b/Engine/Shaders/Private/DistanceFieldVisualization.usf
@@ -117,7 +117,7 @@ void RayTraceThroughTileCulledDistanceFields(float3 WorldRayStart, float3 WorldR
 				}
 			}
 
-			TotalStepsTaken += max(StepIndex, 1);
+			TotalStepsTaken += max(StepIndex, 1u);
 		}
 	}
 }
