Fixup ShadowDepthPixelShder so it more closely matches ShadowDepthVertexShader which addresses a mismatch between the shaders during PSO creation

diff --git a/Engine/Shaders/Private/ShadowDepthPixelShader.usf b/Engine/Shaders/Private/ShadowDepthPixelShader.usf
--- a/Engine/Shaders/Private/ShadowDepthPixelShader.usf
+++ b/Engine/Shaders/Private/ShadowDepthPixelShader.usf
@@ -8,6 +8,17 @@
 #define SHADOW_DEPTH_SHADER 1
 #define USE_STENCIL_LOD_DITHER	0
 
+#ifndef ENABLE_NON_NANITE_VSM
+#error "ENABLE_NON_NANITE_VSM should be defined to either 0 or 1!"
+#endif
+
+
+#if ENABLE_NON_NANITE_VSM
+	#define VIRTUAL_SM_ENABLED (!(USING_TESSELLATION || ONEPASS_POINTLIGHT_SHADOW || PERSPECTIVE_CORRECT_DEPTH))
+#else
+	#define VIRTUAL_SM_ENABLED 0
+#endif
+
 #include "Nanite/PackedNaniteView.ush"
 #include "Common.ush"
 
@@ -24,7 +35,8 @@
 #include "/Engine/Generated/VertexFactory.ush"
 #include "ShadowDepthCommon.ush"
 
-#if ENABLE_NON_NANITE_VSM
+
+#if VIRTUAL_SM_ENABLED
 #include "Nanite/NaniteDataDecode.ush"
 #include "SceneData.ush"
 #include "VirtualShadowMaps/PageAccessCommon.ush"
@@ -34,7 +46,7 @@
 
 void Main( 
 	FShadowDepthVSToPS Inputs,
-#if VIRTUAL_TEXTURE_TARGET
+#if VIRTUAL_SM_ENABLED
 	nointerpolation uint ViewId : TEXCOORD8,
 #endif
 	in float4 SvPosition : SV_Position		// after all interpolators
@@ -56,7 +68,6 @@ void Main(
 			float3 TranslatedWorld = Inputs.PixelPosition.xyz;
 
 			CalcMaterialParametersEx(MaterialParameters, PixelMaterialInputs, SvPosition, ScreenPosition, 1, TranslatedWorld, TranslatedWorld);
-
 		}
 	#else
 		CalcMaterialParameters(MaterialParameters, PixelMaterialInputs, SvPosition, 1);
