Stop non-shadow casters from invalidating virtual shadow map pages.

diff --git a/Engine/Shaders/Private/VirtualShadowMaps/CacheManagement.usf b/Engine/Shaders/Private/VirtualShadowMaps/CacheManagement.usf
--- a/Engine/Shaders/Private/VirtualShadowMaps/CacheManagement.usf
+++ b/Engine/Shaders/Private/VirtualShadowMaps/CacheManagement.usf
@@ -47,65 +47,72 @@ void VirtualSmInvalidateInstancePagesCS(
 	{
 		FInstanceSceneData InstanceData = GetInstanceData(InstanceId, InstanceDataSOAStride);
 
-		for (uint ShadowMapDataIndex = 0; ShadowMapDataIndex < VirtualShadowMap.NumShadowMaps; ++ShadowMapDataIndex)
+		// TODO: Clean up hardcoded flag field test.
+		bool bCastShadows = (GetPrimitiveData(InstanceData.PrimitiveId).Flags & 1U) != 0U;
+		
+		// TODO: test the flag on the instance instead once it is updated correctly InstanceData.CastShadows
+		if (bCastShadows)
 		{
-			// 1. Load cached projection data
-			FVirtualShadowMapProjectionShaderData ProjectionData = GetVirtualShadowMapProjectionData(ShadowMapDataIndex);
-			const bool bDirectionalLight = (ProjectionData.LightType == LIGHT_TYPE_DIRECTIONAL);
+			for (uint ShadowMapDataIndex = 0; ShadowMapDataIndex < VirtualShadowMap.NumShadowMaps; ++ShadowMapDataIndex)
+			{
+				// 1. Load cached projection data
+				FVirtualShadowMapProjectionShaderData ProjectionData = GetVirtualShadowMapProjectionData(ShadowMapDataIndex);
+				const bool bDirectionalLight = (ProjectionData.LightType == LIGHT_TYPE_DIRECTIONAL);
 
-			float4x4 LocalToTranslatedWorld = InstanceData.LocalToWorld;
-			LocalToTranslatedWorld[3].xyz += ProjectionData.ShadowPreViewTranslation;
-			
-			// Go back to clip space
-			float4x4 UVToClip;
-			UVToClip[0] = float4( 2,  0, 0, 0);
-			UVToClip[1] = float4( 0, -2, 0, 0);
-			UVToClip[2] = float4( 0,  0, 1, 0);
-			UVToClip[3] = float4(-1,  1, 0, 1);
+				float4x4 LocalToTranslatedWorld = InstanceData.LocalToWorld;
+				LocalToTranslatedWorld[3].xyz += ProjectionData.ShadowPreViewTranslation;
 
-			float4x4 LocalToClip = mul(LocalToTranslatedWorld, mul(ProjectionData.TranslatedWorldToShadowUVMatrix, UVToClip));
+				// Go back to clip space
+				float4x4 UVToClip;
+				UVToClip[0] = float4(2, 0, 0, 0);
+				UVToClip[1] = float4(0, -2, 0, 0);
+				UVToClip[2] = float4(0, 0, 1, 0);
+				UVToClip[3] = float4(-1, 1, 0, 1);
 
-			FFrustumCullData Cull = BoxCullFrustum(InstanceData.LocalBoundsCenter, InstanceData.LocalBoundsExtent, LocalToClip, !bDirectionalLight, false);
+				float4x4 LocalToClip = mul(LocalToTranslatedWorld, mul(ProjectionData.TranslatedWorldToShadowUVMatrix, UVToClip));
 
-			if (Cull.bIsVisible)
-			{
-				// 2. figure out overlap and all that
-				// case #1 mip-map VSM - loop all mip levels, case #2 clipmap, just one 'mip level'
-				int NumMipLevels = (ProjectionData.ClipmapLevelCount <= 0) ? VSM_MAX_MIP_LEVELS : 1;
+				FFrustumCullData Cull = BoxCullFrustum(InstanceData.LocalBoundsCenter, InstanceData.LocalBoundsExtent, LocalToClip, !bDirectionalLight, false);
+
+				if (Cull.bIsVisible)
 				{
-					for (int MipLevel = 0; MipLevel < NumMipLevels; ++MipLevel)
+					// 2. figure out overlap and all that
+					// case #1 mip-map VSM - loop all mip levels, case #2 clipmap, just one 'mip level'
+					int NumMipLevels = (ProjectionData.ClipmapLevelCount <= 0) ? VSM_MAX_MIP_LEVELS : 1;
 					{
-						int ViewDim = VSM_VIRTUAL_MAX_RESOLUTION_XY >> MipLevel;
-						FScreenRect Rect = GetScreenRect(int4(0, 0, ViewDim, ViewDim), Cull, 2);
-						// Use Hierarchical mip test to speed up (allows skipping invalidating areas that don't have any flags anyway)
-						if (OverlapsAnyValidPage(ProjectionData.VirtualShadowMapId, MipLevel, float2(0.0f, 0.0f), Rect, true /* bHasMoved - checks against allocated flags*/ ))
+						for (int MipLevel = 0; MipLevel < NumMipLevels; ++MipLevel)
 						{
-							uint4 RectPages = uint4(Rect.Pixels) >> VSM_LOG2_PAGE_SIZE;
-
-							// Clip to actually allocated pages
-							// TODO: move this to be done as part of or before the overlap test?
-							uint4 AllocatedBounds = PageRectBounds[ProjectionData.VirtualShadowMapId * VSM_MAX_MIP_LEVELS + MipLevel];
-							RectPages.xy = max(RectPages.xy, AllocatedBounds.xy);
-							RectPages.zw = min(RectPages.zw, AllocatedBounds.zw);
-							if (all(RectPages.zw >= RectPages.xy))
+							int ViewDim = VSM_VIRTUAL_MAX_RESOLUTION_XY >> MipLevel;
+							FScreenRect Rect = GetScreenRect(int4(0, 0, ViewDim, ViewDim), Cull, 2);
+							// Use Hierarchical mip test to speed up (allows skipping invalidating areas that don't have any flags anyway)
+							if (OverlapsAnyValidPage(ProjectionData.VirtualShadowMapId, MipLevel, float2(0.0f, 0.0f), Rect, true /* bHasMoved - checks against allocated flags*/))
 							{
-								// 3. do invalidation
-								uint PageTableLevelOffset = CalcPageTableLevelOffset(ProjectionData.VirtualShadowMapId, MipLevel);
-								for (uint y = RectPages.y; y <= RectPages.w; y++)
+								uint4 RectPages = uint4(Rect.Pixels) >> VSM_LOG2_PAGE_SIZE;
+
+								// Clip to actually allocated pages
+								// TODO: move this to be done as part of or before the overlap test?
+								uint4 AllocatedBounds = PageRectBounds[ProjectionData.VirtualShadowMapId * VSM_MAX_MIP_LEVELS + MipLevel];
+								RectPages.xy = max(RectPages.xy, AllocatedBounds.xy);
+								RectPages.zw = min(RectPages.zw, AllocatedBounds.zw);
+								if (all(RectPages.zw >= RectPages.xy))
 								{
-									for (uint x = RectPages.x; x <= RectPages.z; x++)
+									// 3. do invalidation
+									uint PageTableLevelOffset = CalcPageTableLevelOffset(ProjectionData.VirtualShadowMapId, MipLevel);
+									for (uint y = RectPages.y; y <= RectPages.w; y++)
 									{
-										uint PageFlagOffset = PageTableLevelOffset + CalcPageOffsetInLevel(MipLevel, uint2(x, y));
-										uint PageFlag = PageFlags[PageFlagOffset];
-
-										if (PageFlag != 0)
+										for (uint x = RectPages.x; x <= RectPages.z; x++)
 										{
-											// Somewhere
-											OutDynamicCasterFlags[PageFlagOffset] = 1;
+											uint PageFlagOffset = PageTableLevelOffset + CalcPageOffsetInLevel(MipLevel, uint2(x, y));
+											uint PageFlag = PageFlags[PageFlagOffset];
+
+											if (PageFlag != 0)
+											{
+												// Somewhere
+												OutDynamicCasterFlags[PageFlagOffset] = 1;
+											}
 										}
 									}
-								}
 
+								}
 							}
 						}
 					}
