Implement support for dithered masking mode in ray traced materials

Simply promote the blend mode to real transparency in this case since dithering would be difficult to do in the raytraced context.

diff --git a/Engine/Shaders/Private/RayTracing/RayTracingMaterialHitShaders.usf b/Engine/Shaders/Private/RayTracing/RayTracingMaterialHitShaders.usf
--- a/Engine/Shaders/Private/RayTracing/RayTracingMaterialHitShaders.usf
+++ b/Engine/Shaders/Private/RayTracing/RayTracingMaterialHitShaders.usf
@@ -263,6 +263,18 @@ void CalcInterpolants(in FRayCone RayCone, in FDefaultAttributes Attributes, out
 #endif
 }
 
+float DitheredOpacityMaskToOpacity(float Mask)
+{
+	// This represents the expected value of the function GetMaterialCoverageAndClipping
+	// which randomly dithers the fragment on or off to produce the effect of opacity
+
+	// The expected value of this dithering can be computed as:
+	//    Simplify[Integrate[If[Mask + Dither - 1/2 < 0, 0, 1], {Dither, 0, 1}]]
+	// which is just:
+
+	return saturate(Mask + 0.5);
+}
+
 float3 GetLumenGeometryNormal()
 {
 	float3 WorldPositions[3];
@@ -384,7 +396,7 @@ RAY_TRACING_ENTRY_CLOSEST_HIT(MaterialCHS,
 
 	// Store the results in local variables and reuse instead of calling the functions multiple times.
 	half3 BaseColor = GetMaterialBaseColor(PixelMaterialInputs);
-#if REFRACTION_USE_INDEX_OF_REFRACTION
+#if MATERIALBLENDING_TRANSLUCENT && REFRACTION_USE_INDEX_OF_REFRACTION
 	half2 Refraction = GetMaterialRefraction(PixelMaterialInputs);
 #else
 	half2 Refraction = 1;
@@ -408,10 +420,17 @@ RAY_TRACING_ENTRY_CLOSEST_HIT(MaterialCHS,
 	Payload.Opacity = GetMaterialOpacity(PixelMaterialInputs);
 
 #if MATERIALBLENDING_MASKED
+#if MATERIAL_DITHER_OPACITY_MASK
+	// dithering emulates real transparency, so switch to translucent
+	Payload.Opacity = DitheredOpacityMaskToOpacity(GetMaterialMask(PixelMaterialInputs));
+	Payload.BlendingMode = RAY_TRACING_BLEND_MODE_TRANSLUCENT;
+#else
+	// regular masked mode - binary decision
 	if (GetMaterialMask(PixelMaterialInputs) < 0)
 	{
 		Payload.Opacity = 0;
 	}
+#endif
 #endif // MATERIALBLENDING_MASKED
 
 	Payload.Ior = Refraction.x;
@@ -571,7 +590,7 @@ RAY_TRACING_ENTRY_ANY_HIT(MaterialAHS,
 	// All other blending modes _might_ need material evaluation - so we need
 	// to generate some code
 
-#if MATERIALBLENDING_MASKED
+#if MATERIALBLENDING_MASKED && !MATERIAL_DITHER_OPACITY_MASK
 	// NOTE: Masked mode execute always - regardless of ray flags
 	// This ensures that both material rays and shadow rays "see" the same thing
 #else
@@ -579,7 +598,7 @@ RAY_TRACING_ENTRY_ANY_HIT(MaterialAHS,
 	if (PackedPayload.IsIgnoreTranslucentMaterials())
 	{
 		// special mode used by regular RT shadows - translucent materials are skipped, everything else is treated as opaque
-#if MATERIALBLENDING_TRANSLUCENT
+#if MATERIALBLENDING_TRANSLUCENT || (MATERIALBLENDING_MASKED && MATERIAL_DITHER_OPACITY_MASK)
 		IgnoreHit();
 #endif
 		return;
@@ -631,7 +650,7 @@ RAY_TRACING_ENTRY_ANY_HIT(MaterialAHS,
 
 	// Now extract the relevant info from PixelMaterialInputs according to the blending mode
 
-#if MATERIALBLENDING_MASKED
+#if MATERIALBLENDING_MASKED && !MATERIAL_DITHER_OPACITY_MASK
 	// Regardless of payload flags -- we always apply this 
 	if (GetMaterialMask(PixelMaterialInputs) < 0)
 	{
@@ -681,10 +700,17 @@ RAY_TRACING_ENTRY_ANY_HIT(MaterialAHS,
 		float3 A = pow(Transmission, 1.0 / VoN);
 		Transparency *= A;
 	}
+#elif MATERIALBLENDING_MASKED && MATERIAL_DITHER_OPACITY_MASK
+	float Opacity = DitheredOpacityMaskToOpacity(GetMaterialMask(PixelMaterialInputs));
+	float Transparency = 1 - Opacity;
 #elif MATERIALBLENDING_TRANSLUCENT
 	float Opacity = GetMaterialOpacity(PixelMaterialInputs);
 	float Transparency = 1 - Opacity;
+#if REFRACTION_USE_INDEX_OF_REFRACTION
 	float Ior = GetMaterialRefraction(PixelMaterialInputs).x;
+#else
+	float Ior = 1.0;
+#endif
 
 	// Is refraction enabled?
 	if (Transparency > 0 && Ior != 1.0)
