Add option for Coherent path tracing

When enabled, shuffle pixel coordinate and share random number streams to improve the coherence of execution. This can provide large speedups (2x or more) on some scenes, and even modest speedups on complex scenes (30% or so). However it comes at the price of somewhat distracting artifacts at low sample counts that may be objectionable.

Leaving this disabled by default for now to gather feedback.

diff --git a/Engine/Shaders/Private/PathTracing/PathTracing.usf b/Engine/Shaders/Private/PathTracing/PathTracing.usf
--- a/Engine/Shaders/Private/PathTracing/PathTracing.usf
+++ b/Engine/Shaders/Private/PathTracing/PathTracing.usf
@@ -123,7 +123,7 @@ FMaterialClosestHitPayload TraceTransparentRay(RayDesc Ray, FRayCone RayCone, bo
 			{
 				if (SelectionWeight < SelectionWeightSum)
 				{
-					// the acceptanace probability is not 1.0
+					// the acceptance probability is not 1.0
 					// generate a random number to see if we should accept this hit
 					float RandValue = RandomSequence_GenerateSample1D(RandSequence);
 					AcceptHit = RandValue * SelectionWeightSum < SelectionWeight;
@@ -405,27 +405,81 @@ bool ProcessSubsurfaceRandomWalk(inout FMaterialClosestHitPayload Payload, inout
 	return false;
 }
 
+uint2 Shuffle(uint2 Value, uint2 Seed)
+{
+	const uint M = PATHTRACER_COHERENT_TILE_SIZE;
+#if 1
+	// Scramble the MxM tile of pixels with an error-diffusion pattern
+	// visualizing the seeds shows the distribution of seeds is slightly better
+	const uint X = Value.x;
+	const uint Y = Value.y;
+	uint TileID = 0, HilbertState = 0;
+	for (int i = PATHTRACER_COHERENT_TILE_SHIFT - 1; i >= 0; i--) {
+		uint xi = (X >> i) & 1;
+		uint yi = (Y >> i) & 1;
+		uint Row = 8 * HilbertState + 4 * xi + 2 * yi;
+		TileID = TileID * 4 + ((0x361E9CB4u >> Row) & 3);
+		HilbertState = (0x8FE65831u >> Row) & 3;
+	}
+	uint SeedTile = StrongIntegerHash(Seed.x + StrongIntegerHash(Seed.y));
+	uint SobolIndex = FastOwenScrambling(reversebits(TileID), SeedTile);
+	uint SobolX = SobolIndex;
+	uint SobolY = 0;
+	for (uint v = 1; SobolIndex; SobolIndex >>= 1, v ^= v << 1)
+	{
+		SobolY ^= (SobolIndex & 1) * v;
+	}
+	return reversebits(uint2(SobolX, SobolY)) >> (32 - PATHTRACER_COHERENT_TILE_SHIFT);
+#elif 1
+	// randomly shuffle coordinate (x,y) within an RxR tile of pixels according to the given seed
+	// this has a slight directional preference when visualizing seeds
+	uint SeedTile = StrongIntegerHash(Seed.x + StrongIntegerHash(Seed.y));
+	uint Index = Value.x + Value.y * M;
+	uint ShuffledIndex = FastOwenScrambling(reversebits(Index), SeedTile) % (M * M);
+	return uint2(ShuffledIndex % M, ShuffledIndex / M);
+#else
+	// no shuffling (just to show repeating artifacts)
+	return Value;
+#endif
+}
+
+uint2 Reorder(uint2 I, uint2 N)
+{
+	const uint M = PATHTRACER_COHERENT_TILE_SIZE;
+	uint2 GroupId = I * M / N;				// [0, M)  - which group we landed in?
+	uint2 GroupCoord = ((I * M) % N) / M;	// [0, N/M) - which pixel within the group?
+
+	// reconstruct pixel coordinate, shuffling the assignment of pixels within a group to hide repeating structure
+	return GroupCoord * M + Shuffle(GroupId, GroupCoord);
+}
 
 RAY_TRACING_ENTRY_RAYGEN(PathTracingMainRG)
 {
-	uint2 LaunchIndex = DispatchRaysIndex().xy + View.ViewRectMin.xy + TileOffset.xy;
-	
+	uint2 DispatchIdx = DispatchRaysIndex().xy;
+	uint2 DispatchDim = DispatchRaysDimensions().xy;
+
+	// shuffle the pixels around so the ones that share a seed land in the same warp
+	uint2 LaunchIndex = (PathTracingData.CoherentSampling ? Reorder(DispatchIdx, DispatchDim) : DispatchIdx) + View.ViewRectMin.xy + TileOffset.xy;
+
 	// Initialize random sequence 
-	float2 BufferSize = View.BufferSizeAndInvSize.xy; 
-	uint LinearIndex = LaunchIndex.x + (LaunchIndex.y * BufferSize.x); 
 	RandomSequence RandSequence;
 
+	uint2 RandPixel = PathTracingData.CoherentSampling
+		? (DispatchIdx / PATHTRACER_COHERENT_TILE_SIZE)	// unique seed within small squares of the shuffled sub-images
+		//? (DispatchIdx * PATHTRACER_COHERENT_TILE_SIZE) / DispatchDim // unique seed per sub-image
+		: LaunchIndex; // unique seed per pixel
 	if (PathTracingData.UseErrorDiffusion)
 	{
 		// z-sampler init
-		RandomSequence_Initialize(RandSequence, LaunchIndex, PathTracingData.Iteration, PathTracingData.TemporalSeed - PathTracingData.Iteration, PathTracingData.MaxSamples);
+		RandomSequence_Initialize(RandSequence, RandPixel, PathTracingData.Iteration, PathTracingData.TemporalSeed - PathTracingData.Iteration, PathTracingData.MaxSamples);
 	}
 	else
 	{
 		// random sobol init
-		RandomSequence_Initialize(RandSequence, LinearIndex, PathTracingData.TemporalSeed);
+		RandomSequence_Initialize(RandSequence, RandPixel.x + RandPixel.y * 65536, PathTracingData.TemporalSeed);
 	}
 
+
 	float3 Radiance = 0;
 
 	// Initialize ray and payload
@@ -457,6 +511,11 @@ RAY_TRACING_ENTRY_RAYGEN(PathTracingMainRG)
 
 	for (int Bounce = 0; Bounce <= PathTracingData.MaxBounces; Bounce++)
 	{
+#if 0
+		// visualize pixel seeds
+		Radiance = 0.5 + 0.3 * cos(2 * PI * ((StrongIntegerHash(RandPixel.x + StrongIntegerHash(RandPixel.y)) & 0xFFFFFF) * 5.96046447754e-08 + float3(0.0, 0.333, 0.6666)));
+		break;
+#endif
 		FMaterialClosestHitPayload Payload = TraceTransparentRay(Ray, RayCone, Bounce == 0, Bounce == PathTracingData.MaxBounces, PathTracingData.SkipDirectLighting == 0 || Bounce > 1, LaunchIndex, NumVisibleLights, RandSequence, PathThroughput, Radiance);
 
 		if (Payload.IsMiss())
diff --git a/Engine/Shaders/Shared/RayTracingDefinitions.h b/Engine/Shaders/Shared/RayTracingDefinitions.h
--- a/Engine/Shaders/Shared/RayTracingDefinitions.h
+++ b/Engine/Shaders/Shared/RayTracingDefinitions.h
@@ -36,6 +36,9 @@
 #define RAY_TRACING_MAX_ALLOWED_PAYLOAD_SIZE    64  // Our maximum allowed payload size (sizeof FPackedMaterialClosestHitPayload)
 
 #define PATHTRACER_MAX_RECT_TEXTURES		8
+#define PATHTRACER_COHERENT_TILE_SHIFT		3
+#define PATHTRACER_COHERENT_TILE_SIZE		(1 << PATHTRACER_COHERENT_TILE_SHIFT)
+
 
 // Constants for the 'Flags' field of FPathTracingLight
 #define PATHTRACER_FLAG_LIGHTING_CHANNEL_MASK				(7 << 0)	// Which lighting channel is this light assigned to?
diff --git a/Engine/Source/Runtime/Renderer/Private/PathTracing.cpp b/Engine/Source/Runtime/Renderer/Private/PathTracing.cpp
--- a/Engine/Source/Runtime/Renderer/Private/PathTracing.cpp
+++ b/Engine/Source/Runtime/Renderer/Private/PathTracing.cpp
@@ -104,7 +104,6 @@ TAutoConsoleVariable<int32> CVarPathTracingEnableCameraBackfaceCulling(
 	ECVF_RenderThreadSafe
 );
 
-
 TAutoConsoleVariable<int32> CVarPathTracingFrameIndependentTemporalSeed(
 	TEXT("r.PathTracing.FrameIndependentTemporalSeed"),
 	1,
@@ -114,6 +113,15 @@ TAutoConsoleVariable<int32> CVarPathTracingFrameIndependentTemporalSeed(
 	ECVF_RenderThreadSafe
 );
 
+TAutoConsoleVariable<int32> CVarPathTracingCoherentSampling(
+	TEXT("r.PathTracing.CoherentSampling"),
+	0,
+	TEXT("When non-zero, share pixel seeds to improve coherence of execution on the GPU. This trades some correlation across the image in exchange for better performance.\n")
+	TEXT("0: off (default)\n")
+	TEXT("1: on\n"),
+	ECVF_RenderThreadSafe
+);
+
 // r.PathTracing.GPUCount is read only because ComputeViewGPUMasks results cannot change after UE has been launched
 TAutoConsoleVariable<int32> CVarPathTracingGPUCount(
 	TEXT("r.PathTracing.GPUCount"),
@@ -179,6 +187,7 @@ static bool PrepareShaderArgs(const FViewInfo& View, FPathTracingData& PathTraci
 	PathTracingData.UseErrorDiffusion = CVarPathTracingUseErrorDiffusion.GetValueOnRenderThread();
 	PathTracingData.ApproximateCaustics = CVarPathTracingApproximateCaustics.GetValueOnRenderThread();
 	PathTracingData.EnableCameraBackfaceCulling = CVarPathTracingEnableCameraBackfaceCulling.GetValueOnRenderThread();
+	PathTracingData.CoherentSampling = CVarPathTracingCoherentSampling.GetValueOnRenderThread();
 	float FilterWidth = CVarPathTracingFilterWidth.GetValueOnRenderThread();
 	if (FilterWidth < 0)
 	{
@@ -270,6 +279,14 @@ static bool PrepareShaderArgs(const FViewInfo& View, FPathTracingData& PathTraci
 		PreviousSkipDirectLighting = PathTracingData.SkipDirectLighting;
 	}
 
+	// Changing coherent sampling requires starting over
+	static uint32 PreviousCoherentSampling = PathTracingData.CoherentSampling;
+	if (PreviousCoherentSampling != PathTracingData.CoherentSampling)
+	{
+		NeedInvalidation = true;
+		PreviousCoherentSampling = PathTracingData.CoherentSampling;
+	}
+
 	// the rest of PathTracingData and AdaptiveSamplingData is filled in by SetParameters below
 	return NeedInvalidation;
 }
@@ -914,8 +931,11 @@ void FDeferredShadingSceneRenderer::RenderPathTracing(
 		{
 			FRHIRayTracingScene* RayTracingSceneRHI = View.GetRayTracingSceneChecked();
 			
-			int32 DispatchSizeX = View.ViewRect.Size().X;
-			int32 DispatchSizeY = View.ViewRect.Size().Y;
+			// Round up to coherent path tracing tile size to simplify pixel shuffling
+			// TODO: be careful not to write extra pixels past the boundary when using multi-gpu
+			const int32 TS = PATHTRACER_COHERENT_TILE_SIZE;
+			int32 DispatchSizeX = FMath::DivideAndRoundUp(View.ViewRect.Size().X, TS) * TS;
+			int32 DispatchSizeY = FMath::DivideAndRoundUp(View.ViewRect.Size().Y, TS) * TS;
 
 			FRayTracingShaderBindingsWriter GlobalResources;
 			SetShaderParameters(GlobalResources, RayGenShader, *PassParameters);
diff --git a/Engine/Source/Runtime/Renderer/Private/PathTracingUniformBuffers.h b/Engine/Source/Runtime/Renderer/Private/PathTracingUniformBuffers.h
--- a/Engine/Source/Runtime/Renderer/Private/PathTracingUniformBuffers.h
+++ b/Engine/Source/Runtime/Renderer/Private/PathTracingUniformBuffers.h
@@ -20,6 +20,7 @@ BEGIN_GLOBAL_SHADER_PARAMETER_STRUCT(FPathTracingData, )
 	SHADER_PARAMETER(uint32, ApproximateCaustics)
 	SHADER_PARAMETER(uint32, EnableCameraBackfaceCulling)
 	SHADER_PARAMETER(uint32, SkipDirectLighting)
+	SHADER_PARAMETER(uint32, CoherentSampling)
 	SHADER_PARAMETER(float, MaxPathIntensity)
 	SHADER_PARAMETER(float, MaxNormalBias)
 	SHADER_PARAMETER(float, FilterWidth)
