Fix incorrect setup of instancing paramaters for ray tracing VF interface in Landscape

diff --git a/Engine/Shaders/Private/LandscapeVertexFactory.ush b/Engine/Shaders/Private/LandscapeVertexFactory.ush
--- a/Engine/Shaders/Private/LandscapeVertexFactory.ush
+++ b/Engine/Shaders/Private/LandscapeVertexFactory.ush
@@ -185,6 +185,11 @@ FVertexFactoryInput LoadVertexFactoryInputForHGS(uint TriangleIndex, int VertexI
 		uint Indices[6] = {i00, i11, i10, i00, i01, i10};
 		uint VertexId = Indices[TriangleIndex % 2 * 3 + VertexIndex];
 	#endif
+#if VF_USE_PRIMITIVE_SCENE_DATA
+	// Note: GetInstanceUserData() stores the GPU-Scene primitive ID
+	Input.InstanceIdOffset = GetInstanceUserData() | VF_TREAT_INSTANCE_ID_OFFSET_AS_PRIMITIVE_ID_FLAG;
+	Input.DrawInstanceId = 0;
+#endif // VF_USE_PRIMITIVE_SCENE_DATA
 
 
 	Input.Position.x = VertexId % SubsectionSizeVertsCurrentLOD;
@@ -212,6 +217,12 @@ FVertexFactoryInput LoadVertexFactoryInputForDynamicUpdate(uint TriangleIndex, i
 	Input.Position.z = LandscapeMVF.SubXY.x;
 	Input.Position.w = LandscapeMVF.SubXY.y;
 
+#if VF_USE_PRIMITIVE_SCENE_DATA
+	// GPUCULL_TODO: Makes use of DynamicDraw workaround, should instead refactor VF interfaces and provide the instance ID data directly.
+	Input.InstanceIdOffset = PrimitiveId | VF_TREAT_INSTANCE_ID_OFFSET_AS_PRIMITIVE_ID_FLAG;
+	Input.DrawInstanceId = 0;
+#endif // VF_USE_PRIMITIVE_SCENE_DATA
+
 	return Input;
 }
 
