Fix shader compile error "error: no matching function for call to 'max(int, uint)"

diff --git a/Engine/Shaders/Private/HairStrands/HairStrandsDeepTransmittanceDualScattering.ush b/Engine/Shaders/Private/HairStrands/HairStrandsDeepTransmittanceDualScattering.ush
--- a/Engine/Shaders/Private/HairStrands/HairStrandsDeepTransmittanceDualScattering.ush
+++ b/Engine/Shaders/Private/HairStrands/HairStrandsDeepTransmittanceDualScattering.ush
@@ -68,7 +68,7 @@ FHairTransmittanceData ComputeDualScatteringTerms(
 	const float3 af_weights = af / (af.r + af.g + af.b);
 	const float3 Beta_f  = dot(float3(Beta_R, Beta_TT, Beta_TRT), af_weights);
 	const float3 Beta_f2 = Beta_f*Beta_f;
-	const float3 sigma_f2 = Beta_f2 * max(1, HairCount);
+	const float3 sigma_f2 = Beta_f2 * max(1u, HairCount);
 
 	const float Theta_d = asin(SinThetaL) + asin(SinThetaV);
 	const float Theta_h = Theta_d * 0.5f;
