Change hair strands sample coord buffer format for avoiding ALU packing/unpacking

diff --git a/Engine/Shaders/Private/DeferredLightPixelShaders.usf b/Engine/Shaders/Private/DeferredLightPixelShaders.usf
--- a/Engine/Shaders/Private/DeferredLightPixelShaders.usf
+++ b/Engine/Shaders/Private/DeferredLightPixelShaders.usf
@@ -324,8 +324,7 @@ void DeferredLightPixelMain(
 		return;
 	}
 
-	const uint PackedCoord = HairStrands.HairSampleCoords[SampleIndex];
-	const uint2 PixelCoord = uint2((PackedCoord & 0xFFFF), ((PackedCoord >> 16) & 0xFFFF));
+	const uint2 PixelCoord = HairStrands.HairSampleCoords[SampleIndex];
 	const float2 UV = (PixelCoord + float2(0.5f, 0.5f)) / float2(View.BufferSizeAndInvSize.xy);
 	const float2 ScreenPosition = (UV - View.ScreenPositionScaleBias.wz) / View.ScreenPositionScaleBias.xy;
 	
diff --git a/Engine/Shaders/Private/HairStrands/HairStrandsDeepTransmittanceMask.usf b/Engine/Shaders/Private/HairStrands/HairStrandsDeepTransmittanceMask.usf
--- a/Engine/Shaders/Private/HairStrands/HairStrandsDeepTransmittanceMask.usf
+++ b/Engine/Shaders/Private/HairStrands/HairStrandsDeepTransmittanceMask.usf
@@ -1,8 +1,8 @@
 // Copyright Epic Games, Inc. All Rights Reserved.
 
+#include "HairStrandsVisibilityCommonStruct.ush"
 #include "HairStrandsVoxelPageCommonStruct.ush"
 #include "HairStrandsDeepShadowCommonStruct.ush"
-
 #include "../Common.ush"
 #include "../CommonViewUniformBuffer.ush"
 #include "../SceneTextureParameters.ush"
@@ -55,7 +55,6 @@ uint2				DeepShadow_Resolution;
 float4x4			DeepShadow_CPUWorldToLightTransforms[MAX_HAIR_MACROGROUP_COUNT];
 float3				LightDirection;
 uint				LightChannelMask;
-uint				MaxVisibilityNodeCount;
 float4				LightPosition;
 float				DeepShadow_DepthBiasScale;
 float				DeepShadow_DensityScale;
@@ -70,8 +69,6 @@ Texture2D<float4>	DeepShadow_DomTexture;
 
 Texture2D<float4> 	RayMarchMaskTexture;
 
-StructuredBuffer<FPackedHairSample> HairVisibilityNodeData;
-StructuredBuffer<uint>		HairVisibilityNodeCoord;
 Buffer<uint4>				IndirectArgsBuffer;
 
 Texture3D<float4>			HairLUTTexture;
@@ -238,20 +235,19 @@ void MainCS(uint2 DispatchThreadId : SV_DispatchThreadID)
 	// Note: Use a fixed group count width (HAIR_VISIBILITY_GROUP_COUNT_WIDTH) for avoiding loading the indirect args (to avoid dep. memory fetch)
 	const uint LocalOffset = DispatchThreadId.x + DispatchThreadId.y * TRANSMITTANCE_GROUP_SIZE * HAIR_VISIBILITY_GROUP_COUNT_WIDTH;
 
+	const uint MaxVisibilityNodeCount = HairStrands.HairSampleCount.Load(uint3(0, 0, 0));
 	if (LocalOffset >= MaxVisibilityNodeCount)
 	{
 		return;
 	}
 
-	const uint PackedCoord	 = HairVisibilityNodeCoord[LocalOffset];
-	const uint2 PixelCoord	 = uint2((PackedCoord & 0xFFFF), ((PackedCoord>>16) & 0xFFFF));
+	const uint2 PixelCoord = HairStrands.HairSampleCoords[LocalOffset];
 
 	const float2 UV = (PixelCoord + float2(0.5f, 0.5f)) / float2(View.BufferSizeAndInvSize.xy);
 	const float3 TransmittanceRandom = float3(InterleavedGradientNoise(PixelCoord, 1), InterleavedGradientNoise(PixelCoord, 2), InterleavedGradientNoise(PixelCoord, 3));
 
 	{
-		const FPackedHairSample NodeData = HairVisibilityNodeData[LocalOffset];
-		const FHairSample Sample = UnpackHairSample(NodeData);
+		const FHairSample Sample = UnpackHairSample(HairStrands.HairSampleData[LocalOffset]);
 
 		if ((LightChannelMask & Sample.LightChannelMask) == 0)
 		{
@@ -425,14 +421,13 @@ void MainCS(uint2 DispatchThreadId : SV_DispatchThreadID)
 {
 	const uint LocalOffset = DispatchThreadId.x + DispatchThreadId.y * TRANSMITTANCE_GROUP_SIZE * HAIR_VISIBILITY_GROUP_COUNT_WIDTH;
 
+	const uint MaxVisibilityNodeCount = HairStrands.HairSampleCount.Load(uint3(0, 0, 0));
 	if (LocalOffset >= MaxVisibilityNodeCount)
 	{
 		return;
 	}
 
-	const uint PackedCoord = HairVisibilityNodeCoord[LocalOffset];
-	const uint2 PixelCoord = uint2((PackedCoord & 0xFFFF), ((PackedCoord >> 16) & 0xFFFF));
-
+	const uint2 PixelCoord = HairStrands.HairSampleCoords[LocalOffset];
 	if (RayMarchMaskTexture.Load(uint3(PixelCoord, 0)).r == 0.0f)
 	{
 		// This pixel is already fully shadowed opaque. (RayMarchMaskTexture is computed using the closest sample as of today)
@@ -448,8 +443,7 @@ void MainCS(uint2 DispatchThreadId : SV_DispatchThreadID)
 	#endif	
 	const float2 UV = (PixelCoord.xy + float2(0.5f, 0.5f)) / float2(View.BufferSizeAndInvSize.xy);
 	{
-		const FPackedHairSample NodeData = HairVisibilityNodeData[LocalOffset];
-		const FHairSample Sample = UnpackHairSample(NodeData);
+		const FHairSample Sample = UnpackHairSample(HairStrands.HairSampleData[LocalOffset]);
 
 		if ((LightChannelMask & Sample.LightChannelMask) == 0)
 		{
diff --git a/Engine/Shaders/Private/HairStrands/HairStrandsEnvironmentLighting.usf b/Engine/Shaders/Private/HairStrands/HairStrandsEnvironmentLighting.usf
--- a/Engine/Shaders/Private/HairStrands/HairStrandsEnvironmentLighting.usf
+++ b/Engine/Shaders/Private/HairStrands/HairStrandsEnvironmentLighting.usf
@@ -4,6 +4,7 @@
 	ReflectionEnvironmentComputeShaders - functionality to apply local cubemaps.
 =============================================================================*/
 
+#include "HairStrandsVisibilityCommonStruct.ush"
 #include "HairStrandsVoxelPageCommonStruct.ush" 
 
 #include "../Common.ush"
@@ -93,8 +94,6 @@ float3 GatherRadiance(float3 RayDirection, float Roughness, uint2 PixelPosition,
 		EyeIndex);
 }
 
-StructuredBuffer<FPackedHairSample> HairVisibilityNodeData;
-
 Texture3D<float4>			HairScatteringLUTTexture;
 Texture3D<float4>			HairEnergyLUTTexture;
 SamplerState 				HairLUTSampler; 
@@ -156,8 +155,6 @@ uint bHairUseViewHairCount;
 Texture2D<float> HairCountTexture;
 
 uint  MultipleScatterSampleCount;
-uint  HairComponents;
-float HairDualScatteringRoughnessOverride;
 float TransmissionDensityScaleFactor;
 
 float Voxel_TanConeAngle;
@@ -242,7 +239,7 @@ float3 ReflectionEnvironment(FGBufferData GBuffer, float2 ScreenPosition, float2
 	float3 AvgDualScattering = 0;
 	const float MaxSearchDistance = HairDistanceThreshold; // in cm
 	#if ADHOC_SAMPLING_MS || ADHOC_SAMPLING_R || ADHOC_SAMPLING_TT || ADHOC_SAMPLING_TRT
-	if ((HairComponents & HAIR_COMPONENT_R) || (HairComponents & HAIR_COMPONENT_TRT) || (HairComponents & HAIR_COMPONENT_TT))
+	if ((View.HairComponents & HAIR_COMPONENT_R) || (View.HairComponents & HAIR_COMPONENT_TRT) || (View.HairComponents & HAIR_COMPONENT_TT))
 	{
 		#if PERMUTATION_SAMPLESET == 0
 		// Poisson disk position http://developer.download.nvidia.com/whitepapers/2008/PCSS_Integration.pdf
@@ -284,7 +281,7 @@ float3 ReflectionEnvironment(FGBufferData GBuffer, float2 ScreenPosition, float2
 			const float SinLightAngle = dot(SampleL, T);
 			const FHairTraversalResult Result = InternalComputeHairCountVoxel(MacroGroupId, SampleWorldPosition, SampleLightPosition, 1, 0, MaxSearchDistance, PixelRadius, true, bDebugEnabled);
 
-			const uint ScatteringComponent = (HairComponents & HAIR_COMPONENT_R) | (HairComponents & HAIR_COMPONENT_TRT) | HAIR_COMPONENT_LS | HAIR_COMPONENT_GS;// | HAIR_MULTISCATTER;
+			const uint ScatteringComponent = (View.HairComponents & HAIR_COMPONENT_R) | (View.HairComponents & HAIR_COMPONENT_TRT) | HAIR_COMPONENT_LS | HAIR_COMPONENT_GS;// | HAIR_MULTISCATTER;
 			const FHairTransmittanceData TransmittanceData = GetHairTransmittance(V, SampleL, GBuffer, InitHairTransmittanceMask(Result.HairCount, Result.Visibility), HairScatteringLUTTexture, HairLUTSampler, ScatteringComponent);
 
 			const float Area = 0.2;
@@ -327,7 +324,7 @@ float3 ReflectionEnvironment(FGBufferData GBuffer, float2 ScreenPosition, float2
 	// MS ==============================================================================
 	float3 ColorMS = 0;
 	#if ADHOC_SAMPLING_MS
-	if ((HairComponents & HAIR_COMPONENT_R) || (HairComponents & HAIR_COMPONENT_TRT) || (HairComponents & HAIR_COMPONENT_TT))
+	if ((View.HairComponents & HAIR_COMPONENT_R) || (View.HairComponents & HAIR_COMPONENT_TRT) || (View.HairComponents & HAIR_COMPONENT_TT))
 	{
 		float3 N = UnoccludedN;
 		float3 R = 2 * dot(V, N) * N - V;
@@ -368,7 +365,7 @@ float3 ReflectionEnvironment(FGBufferData GBuffer, float2 ScreenPosition, float2
 	// R ==============================================================================
 	float3 ColorR = 0;
 	#if ADHOC_SAMPLING_R || ADHOC_SAMPLING_TRT
-	if ((HairComponents & HAIR_COMPONENT_R) || (HairComponents & HAIR_COMPONENT_TRT))
+	if ((View.HairComponents & HAIR_COMPONENT_R) || (View.HairComponents & HAIR_COMPONENT_TRT))
 	{
 		float3 R = 0;
 		{
@@ -428,7 +425,7 @@ float3 ReflectionEnvironment(FGBufferData GBuffer, float2 ScreenPosition, float2
 			FHairTransmittanceData TransmittanceData = InitHairStrandsTransmittanceData();
 			TransmittanceData.OpaqueVisibility = Visibility;
 			TransmittanceData.GlobalScattering = Transmittance;
-			TransmittanceData.ScatteringComponent = (HairComponents & HAIR_COMPONENT_R) | (HairComponents & HAIR_COMPONENT_TRT) | HAIR_COMPONENT_MULTISCATTER;
+			TransmittanceData.ScatteringComponent = (View.HairComponents & HAIR_COMPONENT_R) | (View.HairComponents & HAIR_COMPONENT_TRT) | HAIR_COMPONENT_MULTISCATTER;
 
 			const float Area = 0.2;
 			const float BackLit = 0;
@@ -445,7 +442,7 @@ float3 ReflectionEnvironment(FGBufferData GBuffer, float2 ScreenPosition, float2
 	// TT ==============================================================================
 	float3 ColorTT = 0;
 	#if ADHOC_SAMPLING_TT
-	if (HairComponents & HAIR_COMPONENT_TT)
+	if (View.HairComponents & HAIR_COMPONENT_TT)
 	{
 		const float3 L = -V;
 		const float SinLightAngle = dot(L, T);
@@ -597,7 +594,7 @@ float3 ReflectionEnvironmentUniformSampling(
 			InitHairTransmittanceMask(Result.HairCount, Result.Visibility),
 			HairScatteringLUTTexture,
 			HairLUTSampler,
-			HairComponents);	
+			View.HairComponents);	
 
 		const float Area = 0;
 		const float Backlit = 1;
@@ -651,7 +648,7 @@ float3 ReflectionEnvironmentSH(
 		InitHairTransmittanceMask(),
 		HairScatteringLUTTexture,
 		HairLUTSampler,
-		HairComponents);
+		View.HairComponents);
 
 	bool bEvalMultiScatter = true;
 	float3 f_s = PI * HairShading(GBuffer, L, V, N, 1, TransmittanceData, 1, 0.2, uint2(0, 0));
@@ -702,7 +699,6 @@ float3 ReflectionEnvironmentSH(
 
 #if PERMUTATION_INTEGRATION_TYPE == INTEGRATION_SCENECOLOR
 Texture2D<float4> SceneColorTexture;
-Texture2D<uint4> HairCategorizationTexture;
 
 float3 ComputeThinTransmission(FGBufferData GBuffer, float2 ScreenPosition, float2 SvPosition, uint PrimitiveId)
 {
@@ -713,7 +709,7 @@ float3 ComputeThinTransmission(FGBufferData GBuffer, float2 ScreenPosition, floa
 	const float3 V = -CameraToPixel;
 	const float3 T = GBuffer.WorldNormal;
 
-	const FCategorizationData CatData = DecodeCategorizationData(HairCategorizationTexture.Load(uint3(PixelPosition, 0)));
+	const FCategorizationData CatData = DecodeCategorizationData(HairStrands.HairCategorizationTexture.Load(uint3(PixelPosition, 0)));
 	const float PixelCoverage = min(CatData.PixelCoverage, 1);
 	
 	// The scene color is attenuated propertionnaly to the hair coverage. In order to retrieve the actual lighting buffer 
@@ -735,20 +731,17 @@ float3 ComputeThinTransmission(FGBufferData GBuffer, float2 ScreenPosition, floa
 
 #if LIGHTING_VS
 
-int2 MaxViewportResolution;
-Texture2D<uint> HairVisibilityNodeCount;
-
 void MainVS(
 	in uint InVertexId : SV_VertexID,
 	out float4 OutPosition : SV_POSITION,
 	nointerpolation out uint OutNodeCount : DISPATCH_NODECOUNT,
 	nointerpolation out uint2 OutResolution : DISPATCH_RESOLUTION)
 {
-	OutNodeCount = HairVisibilityNodeCount.Load(uint3(0, 0, 0));
+	OutNodeCount = HairStrands.HairSampleCount.Load(uint3(0, 0, 0));
 	OutResolution.x = ceil(sqrt(OutNodeCount));
 	OutResolution.y = OutResolution.x;
 
-	const float2 ClipCoord = ((OutResolution + 0.5f) / float2(MaxViewportResolution)) * 2;
+	const float2 ClipCoord = ((OutResolution + 0.5f) / float2(HairStrands.HairSampleViewportResolution)) * 2;
 
 	const float2 UV = float2(InVertexId & 1, InVertexId >> 1);
 	const float2 Pos = float2(-UV.x, UV.y) * 4 + float2(-1, +1) + float2(ClipCoord.x, -ClipCoord.y);
@@ -759,9 +752,7 @@ void MainVS(
 
 #if LIGHTING_PS
 
-uint MaxVisibilityNodeCount;
-StructuredBuffer<uint>	HairVisibilityNodeCoord;
-Buffer<uint4>			IndirectArgsBuffer;
+Buffer<uint4> IndirectArgsBuffer;
 
 void MainPS(
 	float4 SVPos : SV_POSITION,
@@ -773,23 +764,21 @@ void MainPS(
 
 	const uint2 InCoord = uint2(SVPos.xy);
 	const uint LocalOffset = InCoord.x + InCoord.y * Resolution.x;
+	const uint MaxVisibilityNodeCount = HairStrands.HairSampleCount.Load(uint3(0, 0, 0));
 	if (LocalOffset >= NodeCount)
 	{
 		return;
 	}
 
-	const uint PackedCoord = HairVisibilityNodeCoord[LocalOffset];
-	const uint2 PixelCoord = uint2((PackedCoord & 0xFFFF), ((PackedCoord >> 16) & 0xFFFF));
+	const uint2 PixelCoord = HairStrands.HairSampleCoords[LocalOffset];
 	const float2 SvPosition = PixelCoord + float2(0.5f, 0.5f);
 
 	const float2 UV = (PixelCoord + float2(0.5f, 0.5f)) / float2(View.BufferSizeAndInvSize.xy);
 	const float2 ScreenPosition = (UV - View.ScreenPositionScaleBias.wz) / View.ScreenPositionScaleBias.xy;
 
-	if (LocalOffset > MaxVisibilityNodeCount)
-		return;
 
-	const FPackedHairSample PackedSample = HairVisibilityNodeData[LocalOffset];
-	const FHairSample Sample = UnpackHairSample(PackedSample);
+
+	const FHairSample Sample = UnpackHairSample(HairStrands.HairSampleData[LocalOffset]);
 
 	FGBufferData GBuffer = (FGBufferData)0;
 	GBuffer.ShadingModelID = SHADINGMODELID_HAIR;
@@ -800,7 +789,7 @@ void MainPS(
 	GBuffer.Metallic = 0; // Scattering;
 	GBuffer.Depth = ConvertFromDeviceZ(Sample.Depth);
 	GBuffer.GBufferAO = 1;
-	GBuffer.CustomData = float4(HairDualScatteringRoughnessOverride, 0, Sample.Backlit, 0);
+	GBuffer.CustomData = float4(HairStrands.HairDualScatteringRoughnessOverride, 0, Sample.Backlit, 0);
 	GBuffer.IndirectIrradiance = 1000000;
 
 #if PERMUTATION_DEBUG
diff --git a/Engine/Shaders/Private/HairStrands/HairStrandsVisibilityCompaction.usf b/Engine/Shaders/Private/HairStrands/HairStrandsVisibilityCompaction.usf
--- a/Engine/Shaders/Private/HairStrands/HairStrandsVisibilityCompaction.usf
+++ b/Engine/Shaders/Private/HairStrands/HairStrandsVisibilityCompaction.usf
@@ -29,7 +29,7 @@ RWTexture2D<uint> OutCompactNodeIndex;
 #else
 	RWStructuredBuffer<FPackedHairSample> OutCompactNodeData;
 #endif
-RWStructuredBuffer<uint> OutCompactNodeCoord;
+RWBuffer<uint2> OutCompactNodeCoord;
 RWTexture2D<uint4> OutCategorizationTexture;
 
 #ifndef PERMUTATION_VELOCITY
@@ -210,7 +210,7 @@ void MainCS(uint GroupIndex : SV_GroupIndex, uint3 DispatchThreadId : SV_Dispatc
 
 			// TODOAGREGATE if last, Create FPackedHairSample from aggregated  data
 			OutCompactNodeData[OffsetInStorage + j] = OutNodeData;
-			OutCompactNodeCoord[OffsetInStorage + j] = (PixelCoord.x & 0xFFFF) | ((PixelCoord.y & 0xFFFF) << 16);
+			OutCompactNodeCoord[OffsetInStorage + j] = PixelCoord;
 
 
 		#if PERMUTATION_VELOCITY != VELOCITY_TYPE_NONE
@@ -490,7 +490,7 @@ void MainCS(uint GroupIndex : SV_GroupIndex, uint3 DispatchThreadId : SV_Dispatc
 
 				const uint StoreIndex = NodeDesc.Offset + OutIndex;
 				OutCompactNodeData[StoreIndex] = OutNodeVis;
-				OutCompactNodeCoord[StoreIndex] = (PixelCoord.x & 0xFFFF) | ((PixelCoord.y & 0xFFFF) << 16);
+				OutCompactNodeCoord[StoreIndex] = PixelCoord;
 			}
 
 			NodeDesc.Count = NodeDesc.Count;
@@ -780,7 +780,7 @@ void MainCS(uint GroupIndex : SV_GroupIndex, uint3 DispatchThreadId : SV_Dispatc
 
 				const uint StoreIndex = NodeDesc.Offset + OutIndex;
 				OutCompactNodeData[StoreIndex] = OutNodeData;
-				OutCompactNodeCoord[StoreIndex] = (PixelCoord.x & 0xFFFF) | ((PixelCoord.y & 0xFFFF) << 16);
+				OutCompactNodeCoord[StoreIndex] = PixelCoord;
 			}
 
 			#if PERMUTATION_VIEWTRANSMITTANCE == 1
diff --git a/Engine/Shaders/Private/HairStrands/HairStrandsVisibilityCompactionComputeRaster.usf b/Engine/Shaders/Private/HairStrands/HairStrandsVisibilityCompactionComputeRaster.usf
--- a/Engine/Shaders/Private/HairStrands/HairStrandsVisibilityCompactionComputeRaster.usf
+++ b/Engine/Shaders/Private/HairStrands/HairStrandsVisibilityCompactionComputeRaster.usf
@@ -45,7 +45,7 @@ float CoverageThreshold;
 RWTexture2D<uint> 				OutCompactNodeCounter;
 RWTexture2D<uint> 				OutCompactNodeIndex;
 RWStructuredBuffer<FNodeVis> 	OutCompactNodeData;
-RWStructuredBuffer<uint> 		OutCompactNodeCoord;
+RWBuffer<uint2>					OutCompactNodeCoord;
 RWTexture2D<uint4> 				OutCategorizationTexture;
 
 groupshared uint AllocationNodeCount;
@@ -205,7 +205,7 @@ void MainCS(uint GroupIndex : SV_GroupIndex, uint3 DispatchThreadId : SV_Dispatc
 
 				const uint StoreIndex = NodeDesc.Offset + OutIndex;
 				OutCompactNodeData[StoreIndex] = OutNodeVis;
-				OutCompactNodeCoord[StoreIndex] = (PixelCoord.x & 0xFFFF) | ((PixelCoord.y & 0xFFFF) << 16);
+				OutCompactNodeCoord[StoreIndex] = PixelCoord;
 			}
 
 			NodeDesc.Count = NodeDesc.Count;
diff --git a/Engine/Source/Runtime/Renderer/Private/HairStrands/HairStrandsData.h b/Engine/Source/Runtime/Renderer/Private/HairStrands/HairStrandsData.h
--- a/Engine/Source/Runtime/Renderer/Private/HairStrands/HairStrandsData.h
+++ b/Engine/Source/Runtime/Renderer/Private/HairStrands/HairStrandsData.h
@@ -25,7 +25,7 @@ BEGIN_GLOBAL_SHADER_PARAMETER_STRUCT(FHairStrandsViewUniformParameters, )
 	SHADER_PARAMETER_RDG_TEXTURE(Texture2D<uint>, HairSampleOffset)						// Offset & count, for accessing pixel's samples, based on screen pixel position
 	SHADER_PARAMETER_RDG_TEXTURE(Texture2D<uint>, HairSampleCount)						// Total count of hair sample, in sample space
 	SHADER_PARAMETER_RDG_BUFFER_SRV(StructuredBuffer<FPackedHairSample>, HairSampleData)// Sample data (coverage, tangent, base color, ...), in sample space // HAIRSTRANDS_TODO: change this to be a uint4 so that we don't have to include the type for generated contant buffer
-	SHADER_PARAMETER_RDG_BUFFER_SRV(StructuredBuffer<uint>, HairSampleCoords)			// Screen pixel coordinate of each sample, in sample space
+	SHADER_PARAMETER_RDG_BUFFER_SRV(Buffer<uint2>, HairSampleCoords)					// Screen pixel coordinate of each sample, in sample space
 END_GLOBAL_SHADER_PARAMETER_STRUCT()
 
 ////////////////////////////////////////////////////////////////////////////////////
@@ -58,6 +58,8 @@ struct FHairStrandsVisibilityData
 	FRDGBufferRef	NodeIndirectArg = nullptr;
 	uint32			NodeGroupSize = 0;
 
+	const static EPixelFormat NodeCoordFormat = PF_R16G16_UINT;
+
 	// Hair lighting is accumulated within this buffer
 	// Allocated conservatively
 	// User indirect dispatch for accumulating contribution
diff --git a/Engine/Source/Runtime/Renderer/Private/HairStrands/HairStrandsEnvironment.cpp b/Engine/Source/Runtime/Renderer/Private/HairStrands/HairStrandsEnvironment.cpp
--- a/Engine/Source/Runtime/Renderer/Private/HairStrands/HairStrandsEnvironment.cpp
+++ b/Engine/Source/Runtime/Renderer/Private/HairStrands/HairStrandsEnvironment.cpp
@@ -218,10 +218,8 @@ public:
 	BEGIN_SHADER_PARAMETER_STRUCT(FParameters, )
 		SHADER_PARAMETER(float, Voxel_TanConeAngle)
 
-		SHADER_PARAMETER(uint32, MaxVisibilityNodeCount)
 		SHADER_PARAMETER(uint32, MultipleScatterSampleCount)
 
-		SHADER_PARAMETER(uint32, HairComponents)
 		SHADER_PARAMETER(float,  HairDualScatteringRoughnessOverride)
 		SHADER_PARAMETER(float, TransmissionDensityScaleFactor)
 		SHADER_PARAMETER(float, HairDistanceThreshold)
@@ -231,22 +229,16 @@ public:
 		SHADER_PARAMETER(uint32, SkyLight_OcclusionCombineMode)
 		SHADER_PARAMETER(float, SkyLight_OcclusionExponent)
 		SHADER_PARAMETER(uint32, bHairUseViewHairCount)
-		SHADER_PARAMETER(FIntPoint, MaxViewportResolution)
 
 		SHADER_PARAMETER_TEXTURE(Texture2D, PreIntegratedGF)
 		SHADER_PARAMETER_SAMPLER(SamplerState, PreIntegratedGFSampler)
 
 		SHADER_PARAMETER_STRUCT_INCLUDE(FSceneTextureParameters, SceneTextures)
 
+		SHADER_PARAMETER_RDG_UNIFORM_BUFFER(FHairStrandsViewUniformParameters, HairStrands)
 		SHADER_PARAMETER_RDG_UNIFORM_BUFFER(FVirtualVoxelParameters, VirtualVoxel)
 
-		SHADER_PARAMETER_RDG_TEXTURE(Texture2D, HairVisibilityNodeCount)
-		SHADER_PARAMETER_RDG_BUFFER_SRV(StructuredBuffer, HairVisibilityNodeData)
-		SHADER_PARAMETER_RDG_BUFFER_SRV(StructuredBuffer, HairVisibilityNodeCoord)
-
 		SHADER_PARAMETER_RDG_TEXTURE(Texture2D, SceneColorTexture)
-		SHADER_PARAMETER_RDG_TEXTURE(Texture2D, HairCountTexture)
-		SHADER_PARAMETER_RDG_TEXTURE(Texture2D, HairCategorizationTexture)
 		SHADER_PARAMETER_RDG_TEXTURE(Texture2D, HairEnergyLUTTexture)
 		SHADER_PARAMETER_RDG_TEXTURE(Texture2D, HairScatteringLUTTexture)
 		SHADER_PARAMETER_SAMPLER(SamplerState, HairLUTSampler)
@@ -340,7 +332,6 @@ static void AddHairStrandsEnvironmentLightingPassPS(
 	{
 		IntegrationType = EHairLightingIntegrationType::SceneColor;
 		PassParameters->SceneColorTexture = SceneColorTexture;
-		PassParameters->HairCategorizationTexture = VisibilityData.CategorizationTexture;
 	}
 	else
 	{
@@ -367,21 +358,17 @@ static void AddHairStrandsEnvironmentLightingPassPS(
 	PassParameters->SkyLight_OcclusionCombineMode = SkyLightOcclusionCombineMode == OCM_Minimum ? 0.0f : 1.0f;
 	PassParameters->SkyLight_OcclusionExponent = SkyLightOcclusionExponent;
 	PassParameters->SkyLight_OcclusionTintAndMinOcclusion = SkyLightOcclusionTintAndMinOcclusion;
-	PassParameters->MaxViewportResolution = VisibilityData.SampleLightingViewportResolution;
-	PassParameters->HairVisibilityNodeCount = VisibilityData.NodeCount;
 	PassParameters->Voxel_TanConeAngle = FMath::Tan(FMath::DegreesToRadians(GetHairStrandsSkyLightingConeAngle()));
 	PassParameters->HairDistanceThreshold = FMath::Max(GHairStrandsSkyLightingDistanceThreshold, 1.f);
 	PassParameters->bHairUseViewHairCount = VisibilityData.ViewHairCountTexture && GHairStrandsSkyLightingUseHairCountTexture ? 1 : 0;
-	PassParameters->MaxVisibilityNodeCount = VisibilityData.NodeData->Desc.NumElements;
 	PassParameters->MultipleScatterSampleCount = FMath::Max(uint32(GHairStrandsSkyLightingSampleCount), 1u);
-	PassParameters->HairComponents = ToBitfield(GetHairComponents());
 	PassParameters->HairDualScatteringRoughnessOverride = GetHairDualScatteringRoughnessOverride();
 	PassParameters->TransmissionDensityScaleFactor = FMath::Max(0.f, GHairStrandsTransmissionDensityScaleFactor);
 	PassParameters->PreIntegratedGF = GSystemTextures.PreintegratedGF->GetRenderTargetItem().ShaderResourceTexture;
 	PassParameters->PreIntegratedGFSampler = TStaticSamplerState<SF_Bilinear, AM_Clamp, AM_Clamp, AM_Clamp>::GetRHI();
-	PassParameters->HairCountTexture = VisibilityData.ViewHairCountTexture ? VisibilityData.ViewHairCountTexture : GSystemTextures.GetBlackDummy(GraphBuilder);
 	PassParameters->SceneTextures = SceneTextures;
 	PassParameters->VirtualVoxel = VirtualVoxelResources.UniformBuffer;
+	PassParameters->HairStrands = HairStrands::BindHairStrandsViewUniformParameters(View);
 	PassParameters->ViewUniformBuffer = View.ViewUniformBuffer;
 	PassParameters->ReflectionCaptureData = View.ReflectionCaptureUniformBuffer;
 	{
@@ -402,10 +389,6 @@ static void AddHairStrandsEnvironmentLightingPassPS(
 		FHairStrandsDebugData::SetParameters(GraphBuilder, *DebugData, ParametersPS->DebugData);
 	}
 
-	// Bind hair data
-	PassParameters->HairVisibilityNodeData = GraphBuilder.CreateSRV(VisibilityData.NodeData);
-	PassParameters->HairVisibilityNodeCoord = GraphBuilder.CreateSRV(VisibilityData.NodeCoord);
-
 	FHairEnvironmentLightingPS::FPermutationDomain PermutationVector;
 	PermutationVector.Set<FHairEnvironmentLighting::FIntegrationType>(uint32(IntegrationType));
 	PermutationVector.Set<FHairEnvironmentLighting::FDebug>(DebugData ? 1 : 0);
@@ -468,8 +451,7 @@ void RenderHairStrandsSceneColorScattering(
 
 		const FHairStrandsVisibilityData& VisibilityData = View.HairStrandsViewData.VisibilityData;
 		const FHairStrandsVoxelResources& VoxelResources = View.HairStrandsViewData.VirtualVoxelResources;
-		const bool bRenderHairLighting = VisibilityData.NodeIndex && VisibilityData.NodeData && VoxelResources.IsValid();
-		check(bRenderHairLighting);
+		check(VoxelResources.IsValid());
 
 		bool bNeedScatterSceneLighting = false;
 		for (const FHairStrandsMacroGroupData& MacroGroupData : View.HairStrandsViewData.MacroGroupDatas)
@@ -509,7 +491,7 @@ void RenderHairStrandsEnvironmentLighting(
 
 	const FHairStrandsVisibilityData& VisibilityData = View.HairStrandsViewData.VisibilityData;
 	const FHairStrandsVoxelResources& VoxelResources = View.HairStrandsViewData.VirtualVoxelResources;
-	const bool bRenderHairLighting = VisibilityData.NodeIndex && VisibilityData.NodeData && VoxelResources.IsValid() && VisibilityData.CategorizationTexture;
+	const bool bRenderHairLighting = VoxelResources.IsValid() && VisibilityData.CategorizationTexture;
 	if (!bRenderHairLighting)
 	{
 		return;
@@ -534,7 +516,7 @@ void RenderHairStrandsAmbientOcclusion(
 
 	const FHairStrandsVisibilityData& VisibilityData = View.HairStrandsViewData.VisibilityData;
 	const FHairStrandsVoxelResources& VoxelResources = View.HairStrandsViewData.VirtualVoxelResources;
-	const bool bRenderHairLighting = VisibilityData.NodeIndex && VisibilityData.NodeData && VoxelResources.IsValid() && VisibilityData.CategorizationTexture;
+	const bool bRenderHairLighting = VoxelResources.IsValid() && VisibilityData.CategorizationTexture;
 	check(bRenderHairLighting);
 
 	for (const FHairStrandsMacroGroupData& MacroGroupData : View.HairStrandsViewData.MacroGroupDatas)
diff --git a/Engine/Source/Runtime/Renderer/Private/HairStrands/HairStrandsRendering.cpp b/Engine/Source/Runtime/Renderer/Private/HairStrands/HairStrandsRendering.cpp
--- a/Engine/Source/Runtime/Renderer/Private/HairStrands/HairStrandsRendering.cpp
+++ b/Engine/Source/Runtime/Renderer/Private/HairStrands/HairStrandsRendering.cpp
@@ -17,7 +17,7 @@ static TRDGUniformBufferRef<FHairStrandsViewUniformParameters> InternalCreateHai
 		Parameters->HairOnlyDepthTexture = In->HairOnlyDepthTexture;
 		Parameters->HairSampleOffset = In->NodeIndex;
 		Parameters->HairSampleData = GraphBuilder.CreateSRV(In->NodeData);
-		Parameters->HairSampleCoords = GraphBuilder.CreateSRV(In->NodeCoord);
+		Parameters->HairSampleCoords = GraphBuilder.CreateSRV(In->NodeCoord, FHairStrandsVisibilityData::NodeCoordFormat);
 		Parameters->HairSampleCount = In->NodeCount;
 		Parameters->HairSampleViewportResolution = In->SampleLightingViewportResolution;
 	}
diff --git a/Engine/Source/Runtime/Renderer/Private/HairStrands/HairStrandsTransmittance.cpp b/Engine/Source/Runtime/Renderer/Private/HairStrands/HairStrandsTransmittance.cpp
--- a/Engine/Source/Runtime/Renderer/Private/HairStrands/HairStrandsTransmittance.cpp
+++ b/Engine/Source/Runtime/Renderer/Private/HairStrands/HairStrandsTransmittance.cpp
@@ -173,13 +173,13 @@ class FDeepTransmittanceMaskCS : public FGlobalShader
 	BEGIN_SHADER_PARAMETER_STRUCT(FParameters, )
 		SHADER_PARAMETER_STRUCT_INCLUDE(FSceneTextureParameters, SceneTextures)
 		SHADER_PARAMETER_STRUCT_INCLUDE(ShaderDrawDebug::FShaderDrawDebugParameters, ShaderDrawParameters)
+		SHADER_PARAMETER_RDG_UNIFORM_BUFFER(FHairStrandsViewUniformParameters, HairStrands)
 
 		SHADER_PARAMETER_ARRAY(FIntVector4, DeepShadow_AtlasSlotOffsets_AtlasSlotIndex, [FHairStrandsDeepShadowData::MaxMacroGroupCount])
 		SHADER_PARAMETER_ARRAY(FMatrix, DeepShadow_CPUWorldToLightTransforms, [FHairStrandsDeepShadowData::MaxMacroGroupCount])
 		SHADER_PARAMETER(FIntPoint, DeepShadow_Resolution)
 		SHADER_PARAMETER(float, LightRadius)
 		SHADER_PARAMETER(FVector, LightDirection)
-		SHADER_PARAMETER(uint32, MaxVisibilityNodeCount)
 		SHADER_PARAMETER(FVector4, LightPosition)
 		SHADER_PARAMETER(uint32, LightChannelMask)
 		SHADER_PARAMETER(FVector4, DeepShadow_LayerDepths)
@@ -198,8 +198,6 @@ class FDeepTransmittanceMaskCS : public FGlobalShader
 		SHADER_PARAMETER_RDG_BUFFER_SRV(StructuredBuffer, DeepShadow_WorldToLightTransformBuffer)
 
 		SHADER_PARAMETER_RDG_TEXTURE(Texture3D, HairLUTTexture)
-		SHADER_PARAMETER_RDG_BUFFER_SRV(StructuredBuffer, HairVisibilityNodeData)
-		SHADER_PARAMETER_RDG_BUFFER_SRV(StructuredBuffer, HairVisibilityNodeCoord)
 		RDG_BUFFER_ACCESS(IndirectArgsBuffer, ERHIAccess::IndirectArgs)
 
 		SHADER_PARAMETER_RDG_BUFFER_UAV(RWBuffer, OutputColor)
@@ -261,9 +259,6 @@ struct FDeepShadowTransmittanceParams
 	FRDGTextureRef DeepShadow_FrontDepthTexture = nullptr;
 	FRDGTextureRef DeepShadow_DomTexture = nullptr;
 
-	FRDGBufferRef HairVisibilityNodeData = nullptr;
-	FRDGBufferRef HairVisibilityNodeCoord = nullptr;
-
 	const FHairStrandsVoxelResources* VirtualVoxelResources = nullptr;
 };
 
@@ -279,7 +274,7 @@ static FRDGBufferRef AddDeepShadowTransmittanceMaskPass(
 	FRDGBufferRef IndirectArgsBuffer,
 	FRDGTextureRef ScreenShadowMaskSubPixelTexture)
 {
-	FRDGBufferRef OutBuffer = CreateHairStrandsTransmittanceMaskBuffer(GraphBuilder, Params.HairVisibilityNodeData->Desc.NumElements);
+	FRDGBufferRef OutBuffer = CreateHairStrandsTransmittanceMaskBuffer(GraphBuilder, View.HairStrandsViewData.VisibilityData.MaxNodeCount);
 
 	FDeepTransmittanceMaskCS::FParameters* Parameters = GraphBuilder.AllocParameters<FDeepTransmittanceMaskCS::FParameters>();
 	Parameters->ViewUniformBuffer = View.ViewUniformBuffer;
@@ -302,7 +297,7 @@ static FRDGBufferRef AddDeepShadowTransmittanceMaskPass(
 	Parameters->DeepShadow_LayerDepths = Params.DeepShadow_LayerDepths;
 	Parameters->DeepShadow_ShadowToWorld = Params.DeepShadow_ShadowToWorld;
 	Parameters->IndirectArgsBuffer = IndirectArgsBuffer;
-	Parameters->MaxVisibilityNodeCount = Params.HairVisibilityNodeData->Desc.NumElements;
+	Parameters->HairStrands = HairStrands::BindHairStrandsViewUniformParameters(View);
 	Parameters->DeepShadow_bIsGPUDriven = Params.DeepShadow_bIsGPUDriven ? 1 : 0;;
 	Parameters->DeepShadow_WorldToLightTransformBuffer = Params.DeepShadow_WorldToLightTransformBuffer;
 
@@ -324,9 +319,6 @@ static FRDGBufferRef AddDeepShadowTransmittanceMaskPass(
 		bIsSuperSampled = GHairStrandsTransmittanceSuperSampling > 0;
 	}
 
-	Parameters->HairVisibilityNodeData = GraphBuilder.CreateSRV(Params.HairVisibilityNodeData);
-	Parameters->HairVisibilityNodeCoord = GraphBuilder.CreateSRV(Params.HairVisibilityNodeCoord);
-
 	const bool bIsMipTraversal = GHairStrandsTransmittanceMaskUseMipTraversal > 0;
 	check(NodeGroupSize == 64 || NodeGroupSize == 32);
 	FDeepTransmittanceMaskCS::FPermutationDomain PermutationVector;
@@ -572,8 +564,6 @@ static FHairStrandsTransmittanceMaskData RenderHairStrandsTransmittanceMask(
 	FRDGTextureRef HairLUTTexture = InHairLUT.Textures[HairLUTType_DualScattering];
 
 	FDeepShadowTransmittanceParams Params;
-	Params.HairVisibilityNodeData = VisibilityData.NodeData;
-	Params.HairVisibilityNodeCoord = VisibilityData.NodeCoord;
 	Params.DeepShadow_DensityScale = GetDeepShadowDensityScale();
 	Params.DeepShadow_DepthBiasScale = GetDeepShadowDepthBiasScale();
 	memset(Params.DeepShadow_AtlasSlotOffsets_AtlasSlotIndex, 0, sizeof(Params.DeepShadow_AtlasSlotOffsets_AtlasSlotIndex));
diff --git a/Engine/Source/Runtime/Renderer/Private/HairStrands/HairStrandsVisibility.cpp b/Engine/Source/Runtime/Renderer/Private/HairStrands/HairStrandsVisibility.cpp
--- a/Engine/Source/Runtime/Renderer/Private/HairStrands/HairStrandsVisibility.cpp
+++ b/Engine/Source/Runtime/Renderer/Private/HairStrands/HairStrandsVisibility.cpp
@@ -996,7 +996,7 @@ BEGIN_GLOBAL_SHADER_PARAMETER_STRUCT(FVisibilityMaterialPassUniformParameters, )
 	SHADER_PARAMETER(uint32, bOutputEmissive)
 	SHADER_PARAMETER_RDG_TEXTURE(Texture2D<uint>, NodeIndex)
 	SHADER_PARAMETER_RDG_TEXTURE(Texture2D<uint>, TotalNodeCounter)
-	SHADER_PARAMETER_RDG_BUFFER_SRV(StructuredBuffer<uint>, NodeCoord)
+	SHADER_PARAMETER_RDG_BUFFER_SRV(Buffer<uint2>, NodeCoord)
 	SHADER_PARAMETER_RDG_BUFFER_SRV(StructuredBuffer<FNodeVis>, NodeVis)
 	SHADER_PARAMETER_RDG_BUFFER_SRV(Buffer<uint>, IndirectArgs)
 	SHADER_PARAMETER_RDG_BUFFER_UAV(RWStructuredBuffer<FPackedHairSample>, OutNodeData)
@@ -1114,7 +1114,7 @@ static FMaterialPassOutput AddHairMaterialPass(
 		UniformParameters->TotalNodeCounter = CompactNodeCounter;
 		UniformParameters->NodeIndex = CompactNodeIndex;
 		UniformParameters->NodeVis = GraphBuilder.CreateSRV(CompactNodeVis);
-		UniformParameters->NodeCoord = GraphBuilder.CreateSRV(CompactNodeCoord);
+		UniformParameters->NodeCoord = GraphBuilder.CreateSRV(CompactNodeCoord, FHairStrandsVisibilityData::NodeCoordFormat);
 		UniformParameters->IndirectArgs = GraphBuilder.CreateSRV(IndirectArgBuffer);
 		UniformParameters->OutNodeData = GraphBuilder.CreateUAV(FRDGBufferUAVDesc(Output.NodeData));
 		UniformParameters->OutNodeVelocity = GraphBuilder.CreateUAV(FRDGBufferUAVDesc(Output.NodeVelocity, FMaterialPassOutput::VelocityFormat));
@@ -1850,7 +1850,7 @@ class FHairVisibilityPrimitiveIdCompactionCS : public FGlobalShader
 		SHADER_PARAMETER_RDG_TEXTURE_UAV(RWTexture2D, OutCompactNodeIndex)
 		SHADER_PARAMETER_RDG_TEXTURE_UAV(RWTexture2D, OutCategorizationTexture)
 		SHADER_PARAMETER_RDG_BUFFER_UAV(RWStructuredBuffer, OutCompactNodeData)
-		SHADER_PARAMETER_RDG_BUFFER_UAV(RWStructuredBuffer, OutCompactNodeCoord)
+		SHADER_PARAMETER_RDG_BUFFER_UAV(RWBuffer, OutCompactNodeCoord)
 		SHADER_PARAMETER_RDG_TEXTURE_UAV(RWTexture2D, OutVelocityTexture)
 
 		SHADER_PARAMETER_STRUCT_REF(FViewUniformShaderParameters, ViewUniformBuffer)
@@ -1969,7 +1969,7 @@ static void AddHairVisibilityPrimitiveIdCompactionPass(
 
 	{
 		// Pixel coord of the node. Stored as 2*uint16, packed into a single uint32
-		OutCompactNodeCoord = GraphBuilder.CreateBuffer(FRDGBufferDesc::CreateStructuredDesc(sizeof(uint32), MaxRenderNodeCount), TEXT("Hair.VisibilityPrimitiveIdCompactNodeCoord"));
+		OutCompactNodeCoord = GraphBuilder.CreateBuffer(FRDGBufferDesc::CreateBufferDesc(sizeof(uint32), MaxRenderNodeCount), TEXT("Hair.VisibilityPrimitiveIdCompactNodeCoord"));
 	}
 
 	const bool bWriteOutVelocity = OutVelocityTexture != nullptr;
@@ -1996,7 +1996,7 @@ static void AddHairVisibilityPrimitiveIdCompactionPass(
 	PassParameters->OutCompactNodeCounter = GraphBuilder.CreateUAV(OutCompactCounter);
 	PassParameters->OutCompactNodeIndex = GraphBuilder.CreateUAV(OutCompactNodeIndex);
 	PassParameters->OutCompactNodeData = GraphBuilder.CreateUAV(OutCompactNodeData);
-	PassParameters->OutCompactNodeCoord = GraphBuilder.CreateUAV(OutCompactNodeCoord);
+	PassParameters->OutCompactNodeCoord = GraphBuilder.CreateUAV(OutCompactNodeCoord, FHairStrandsVisibilityData::NodeCoordFormat);
 	PassParameters->OutCategorizationTexture = GraphBuilder.CreateUAV(OutCategorizationTexture);
 
 	if (bWriteOutVelocity)
@@ -2053,7 +2053,7 @@ class FHairVisibilityCompactionComputeRasterCS : public FGlobalShader
 		SHADER_PARAMETER_RDG_TEXTURE_UAV(RWTexture2D, OutCompactNodeIndex)
 		SHADER_PARAMETER_RDG_TEXTURE_UAV(RWTexture2D, OutCategorizationTexture)
 		SHADER_PARAMETER_RDG_BUFFER_UAV(RWStructuredBuffer, OutCompactNodeData)
-		SHADER_PARAMETER_RDG_BUFFER_UAV(RWStructuredBuffer, OutCompactNodeCoord)
+		SHADER_PARAMETER_RDG_BUFFER_UAV(RWBuffer, OutCompactNodeCoord)
 
 		SHADER_PARAMETER_RDG_UNIFORM_BUFFER(FSceneTextureUniformParameters, SceneTexturesStruct)
 		SHADER_PARAMETER_STRUCT_REF(FViewUniformShaderParameters, ViewUniformBuffer)
@@ -2117,7 +2117,7 @@ static void AddHairVisibilityCompactionComputeRasterPass(
 	const uint32 PPLLMaxRenderNodePerPixel = GetMaxSamplePerPixel();
 	const uint32 MaxRenderNodeCount = GetTotalSampleCountForAllocation(Resolution);
 	OutCompactNodeData = GraphBuilder.CreateBuffer(FRDGBufferDesc::CreateStructuredDesc(sizeof(HairStrandsVisibilityInternal::NodeVis), MaxRenderNodeCount), TEXT("Hair.VisibilityPrimitiveIdCompactNodeData"));
-	OutCompactNodeCoord = GraphBuilder.CreateBuffer(FRDGBufferDesc::CreateStructuredDesc(sizeof(uint32), MaxRenderNodeCount), TEXT("Hair.VisibilityPrimitiveIdCompactNodeCoord"));
+	OutCompactNodeCoord = GraphBuilder.CreateBuffer(FRDGBufferDesc::CreateBufferDesc(sizeof(uint32), MaxRenderNodeCount), TEXT("Hair.VisibilityPrimitiveIdCompactNodeCoord"));
 
 	FRDGTextureRef DefaultTexture = GSystemTextures.GetBlackDummy(GraphBuilder);
 	FHairVisibilityCompactionComputeRasterCS::FParameters* PassParameters = GraphBuilder.AllocParameters<FHairVisibilityCompactionComputeRasterCS::FParameters>();
@@ -2135,7 +2135,7 @@ static void AddHairVisibilityCompactionComputeRasterPass(
 	PassParameters->OutCompactNodeCounter	= GraphBuilder.CreateUAV(OutCompactCounter);
 	PassParameters->OutCompactNodeIndex		= GraphBuilder.CreateUAV(OutCompactNodeIndex);
 	PassParameters->OutCompactNodeData		= GraphBuilder.CreateUAV(OutCompactNodeData);
-	PassParameters->OutCompactNodeCoord		= GraphBuilder.CreateUAV(OutCompactNodeCoord);
+	PassParameters->OutCompactNodeCoord		= GraphBuilder.CreateUAV(OutCompactNodeCoord, FHairStrandsVisibilityData::NodeCoordFormat);
 	PassParameters->OutCategorizationTexture= GraphBuilder.CreateUAV(OutCategorizationTexture);
 
 	const FIntPoint GroupSize = GetVendorOptimalGroupSize2D();
