Add one pass hair strands transmittance support.

This will be used when cluster shading is enabled.

diff --git a/Engine/Shaders/Private/HairStrands/HairStrandsDeepTransmittanceMask.usf b/Engine/Shaders/Private/HairStrands/HairStrandsDeepTransmittanceMask.usf
--- a/Engine/Shaders/Private/HairStrands/HairStrandsDeepTransmittanceMask.usf
+++ b/Engine/Shaders/Private/HairStrands/HairStrandsDeepTransmittanceMask.usf
@@ -15,6 +15,12 @@
 #include "HairStrandsVisibilityCommon.ush"
 #include "HairStrandsVoxelPageCommon.ush"
 
+#define DEBUG_ENABLE 0
+#if DEBUG_ENABLE
+	#include "../ShaderDrawDebug.ush"
+	#include "../ShaderPrintCommon.ush"
+#endif //DEBUG_ENABLE
+
 // Need to match FDeepShadowData::MaxMacroGroupCount in HairRendering.h
 #if SHADER_TRANSMITTANCE_DEEPSHADOW || SHADER_TRANSMITTANCE_VOXEL
 	#define MAX_HAIR_MACROGROUP_COUNT 16 
@@ -211,25 +217,25 @@ FHairTransmittanceMask ComputeTransmittance(
 void MainCS(uint2 DispatchThreadId : SV_DispatchThreadID)
 {
 	// Note: Use a fixed group count width (HAIR_VISIBILITY_GROUP_COUNT_WIDTH) for avoiding loading the indirect args (to avoid dep. memory fetch)
-	const uint LocalOffset = DispatchThreadId.x + DispatchThreadId.y * TRANSMITTANCE_GROUP_SIZE * HAIR_VISIBILITY_GROUP_COUNT_WIDTH;
+	const uint SampleIndex = DispatchThreadId.x + DispatchThreadId.y * TRANSMITTANCE_GROUP_SIZE * HAIR_VISIBILITY_GROUP_COUNT_WIDTH;
 
 	const uint MaxVisibilityNodeCount = HairStrands.HairSampleCount.Load(uint3(0, 0, 0));
-	if (LocalOffset >= MaxVisibilityNodeCount)
+	if (SampleIndex >= MaxVisibilityNodeCount)
 	{
 		return;
 	}
 
-	const uint2 PixelCoord = HairStrands.HairSampleCoords[LocalOffset];
+	const uint2 PixelCoord = HairStrands.HairSampleCoords[SampleIndex];
 
 	const float2 UV = (PixelCoord + float2(0.5f, 0.5f)) / float2(View.BufferSizeAndInvSize.xy);
 	const float3 TransmittanceRandom = float3(InterleavedGradientNoise(PixelCoord, 1), InterleavedGradientNoise(PixelCoord, 2), InterleavedGradientNoise(PixelCoord, 3));
 
 	{
-		const FHairSample Sample = UnpackHairSample(HairStrands.HairSampleData[LocalOffset]);
+		const FHairSample Sample = UnpackHairSample(HairStrands.HairSampleData[SampleIndex]);
 
 		if ((LightChannelMask & Sample.LightChannelMask) == 0)
 		{
-			OutTransmittanceMask[LocalOffset] = InitNullPackedHairTransmittanceMask();
+			OutTransmittanceMask[SampleIndex] = InitNullPackedHairTransmittanceMask();
 			return;
 		}
 
@@ -262,7 +268,7 @@ void MainCS(uint2 DispatchThreadId : SV_DispatchThreadID)
 		const float3 WorldPosition = ToWorldPosition(UV, SceneDepth);
 		const FHairTransmittanceMask TransmittanceMask = ComputeTransmittance(WorldPosition, Settings, DeepShadow_LayerDepths, DeepShadow_FrontDepthTexture, DeepShadow_DomTexture);
 
-		OutTransmittanceMask[LocalOffset] = PackTransmittanceMask(TransmittanceMask);
+		OutTransmittanceMask[SampleIndex] = PackTransmittanceMask(TransmittanceMask);
 	}
 }
 #endif // #if SHADER_TRANSMITTANCE_DEEPSHADOW
@@ -280,6 +286,7 @@ void MainCS(uint2 DispatchThreadId : SV_DispatchThreadID)
 
 #define VOXEL_TRAVERSAL_DEBUG 0
 #include "HairStrandsVoxelPageTraversal.ush"
+#include "../LightGridCommon.ush"
 
 #if VOXEL_TRAVERSAL_DEBUG
 #include "../ShaderPrintCommon.ush"
@@ -330,7 +337,11 @@ float3				LightDirection;
 uint				LightChannelMask;
 float4				LightPosition;
 float				LightRadius;
+#if PERMUTATION_ONE_PASS
+Texture2D<uint> 	ShadowMaskBitsTexture;
+#else
 Texture2D<float4> 	RayMarchMaskTexture;
+#endif
 Buffer<uint4>		IndirectArgsBuffer;
 RWBuffer<uint>		OutTransmittanceMask;
 
@@ -417,21 +428,103 @@ FHairTransmittanceMask ComputeTransmittanceVoxel(
 
 [numthreads(TRANSMITTANCE_GROUP_SIZE, 1, 1)]
 void MainCS(uint2 DispatchThreadId : SV_DispatchThreadID)
+#if PERMUTATION_ONE_PASS
 {
-	const uint LocalOffset = DispatchThreadId.x + DispatchThreadId.y * TRANSMITTANCE_GROUP_SIZE * HAIR_VISIBILITY_GROUP_COUNT_WIDTH;
+	const uint SampleIndex = DispatchThreadId.x + DispatchThreadId.y * TRANSMITTANCE_GROUP_SIZE * HAIR_VISIBILITY_GROUP_COUNT_WIDTH;
 
 	const uint MaxVisibilityNodeCount = HairStrands.HairSampleCount.Load(uint3(0, 0, 0));
-	if (LocalOffset >= MaxVisibilityNodeCount)
+	if (SampleIndex >= MaxVisibilityNodeCount)
 	{
 		return;
 	}
 
-	const uint2 PixelCoord = HairStrands.HairSampleCoords[LocalOffset];
+	const FHairSample Sample	= UnpackHairSample(HairStrands.HairSampleData[SampleIndex]);
+	const uint MacroGroupIndex	= clamp(Sample.MacroGroupId, 0, MAX_HAIR_MACROGROUP_COUNT - 1);
+	const float SceneDepth		= ConvertFromDeviceZ(Sample.Depth);
+	const uint2 PixelCoord		= HairStrands.HairSampleCoords[SampleIndex];
+	const float2 UV				= (PixelCoord.xy + float2(0.5f, 0.5f)) / float2(View.BufferSizeAndInvSize.xy);
+	const float3 WorldPosition	= ToWorldPosition(UV, SceneDepth);
+	const float PixelRadius		= SceneDepth * VirtualVoxel.HairCoveragePixelRadiusAtDepth1;
+
+	const uint PackedShadowMask = ~ShadowMaskBitsTexture[PixelCoord];
+	const float3 Random = GetHairVoxelJitter(PixelCoord.xy, View.StateFrameIndexMod8);
+
+	const uint2 LocalPixelPos = PixelCoord;
+
+	const uint EyeIndex = 0; // TODO: Instanced stereo
+	const uint GridLinearIndex = ComputeLightGridCellIndex(LocalPixelPos, SceneDepth, EyeIndex);
+	const FCulledLightsGridData CulledLightGridData = GetCulledLightsGrid(GridLinearIndex, EyeIndex);
+
+	const uint LightCount = min(GetShadowMaskMaxLightCount(), CulledLightGridData.NumLocalLights);
+
+	LOOP
+	for (uint LightIndex = 0; LightIndex < LightCount; ++LightIndex)
+	{
+		const uint OutputIndex = SampleIndex + MaxVisibilityNodeCount * LightIndex;
+
+		// TODO: Add early out when the shadow mask bits are zero
+		#if 0
+		if (UnpackShadowMask(PackedShadowMask, LightIndex) == 0)
+		{
+			OutTransmittanceMask[OutputIndex] = PackTransmittanceMask(InitHairTransmittanceMask());
+			continue;
+		}
+		#endif
+
+		// TODO: add support for light channel
+		#if 0
+		if ((LightChannelMask & Sample.LightChannelMask) == 0)
+		{
+			OutTransmittanceMask[OutputIndex] = InitNullPackedHairTransmittanceMask();
+			continue;
+		}
+		#endif	
+
+		// The Virtual Shadow Remap stores directional lights first
+		const FLocalLightData LightData = GetLocalLightData(CulledLightGridData.DataStartIndex + LightIndex, EyeIndex);
+		const float3 LightPosition		= LightData.LightPositionAndInvRadius.xyz;
+		const float3 LightDirection		= LightData.LightDirectionAndShadowMask.xyz; // Direction light only
+		const bool bIsLocalLight		= true;
+
+		FTransmittanceSettings Settings;
+		Settings.LightDirection				= normalize(LightPosition - WorldPosition);
+		Settings.LightPosition				= float4(LightPosition, bIsLocalLight ? 1 : 0);
+		Settings.DeepShadowDensityScale		= 0;	// Use virtual voxel struct value
+		Settings.DeepShadowDepthBiasScale	= 0;	// Use virtual voxel struct value
+		Settings.Random						= Random;
+
+	#if DEBUG_ENABLE
+		const bool bDebugEnabled = PixelCoord.x == GetCursorPos().x && PixelCoord.y == GetCursorPos().y;
+	#else
+		const bool bDebugEnabled = false;
+	#endif
+		const FHairTransmittanceMask Out = ComputeTransmittanceVoxel(
+			PixelCoord, 
+			PixelRadius, 
+			WorldPosition, 
+			MacroGroupIndex, 
+			Settings, 
+			bDebugEnabled);
+
+		OutTransmittanceMask[OutputIndex] = PackTransmittanceMask(Out);
+	}
+}
+#else // PERMUTATION_ONE_PASS
+{
+	const uint SampleIndex = DispatchThreadId.x + DispatchThreadId.y * TRANSMITTANCE_GROUP_SIZE * HAIR_VISIBILITY_GROUP_COUNT_WIDTH;
+
+	const uint MaxVisibilityNodeCount = HairStrands.HairSampleCount.Load(uint3(0, 0, 0));
+	if (SampleIndex >= MaxVisibilityNodeCount)
+	{
+		return;
+	}
+
+	const uint2 PixelCoord = HairStrands.HairSampleCoords[SampleIndex];
 	if (RayMarchMaskTexture.Load(uint3(PixelCoord, 0)).r == 0.0f)
 	{
 		// This pixel is already fully shadowed opaque. (RayMarchMaskTexture is computed using the closest sample as of today)
 		const FHairTransmittanceMask DefaultMask = InitHairTransmittanceMask();
-		OutTransmittanceMask[LocalOffset] = PackTransmittanceMask(DefaultMask);
+		OutTransmittanceMask[SampleIndex] = PackTransmittanceMask(DefaultMask);
 		return;
 	}
 
@@ -442,11 +535,10 @@ void MainCS(uint2 DispatchThreadId : SV_DispatchThreadID)
 	#endif	
 	const float2 UV = (PixelCoord.xy + float2(0.5f, 0.5f)) / float2(View.BufferSizeAndInvSize.xy);
 	{
-		const FHairSample Sample = UnpackHairSample(HairStrands.HairSampleData[LocalOffset]);
-
+		const FHairSample Sample = UnpackHairSample(HairStrands.HairSampleData[SampleIndex]);
 		if ((LightChannelMask & Sample.LightChannelMask) == 0)
 		{
-			OutTransmittanceMask[LocalOffset] = InitNullPackedHairTransmittanceMask();
+			OutTransmittanceMask[SampleIndex] = InitNullPackedHairTransmittanceMask();
 			return;
 		}
 
@@ -469,21 +561,20 @@ void MainCS(uint2 DispatchThreadId : SV_DispatchThreadID)
 		const float SceneDepth = ConvertFromDeviceZ(Sample.Depth);
 		const float3 WorldPosition = ToWorldPosition(UV, SceneDepth);
 
-		const float PixelRadius = SceneDepth * VirtualVoxel.HairCoveragePixelRadiusAtDepth1;
-	
+		const float PixelRadius = SceneDepth * VirtualVoxel.HairCoveragePixelRadiusAtDepth1;	
 		const FHairTransmittanceMask Out = ComputeTransmittanceVoxel(PixelCoord, PixelRadius, WorldPosition, MacroGroupIndex, Settings, bDebugEnabled);
 
 		#if VOXEL_TRAVERSAL_DEBUG
 		if (bDebugEnabled)
 		{
-			PrintInfo(LocalOffset, PixelCoord, Out);
+			PrintInfo(SampleIndex, PixelCoord, Out);
 		}
 		#endif	
 
-		OutTransmittanceMask[LocalOffset] = PackTransmittanceMask(Out);
+		OutTransmittanceMask[SampleIndex] = PackTransmittanceMask(Out);
 	}
 }
-
+#endif // PERMUTATION_ONE_PASS
 #endif // SHADER_TRANSMITTANCE_VOXEL
 
 ////////////////////////////////////////////////////////////////////////////////////////////////
diff --git a/Engine/Shaders/Private/LightGridCommon.ush b/Engine/Shaders/Private/LightGridCommon.ush
--- a/Engine/Shaders/Private/LightGridCommon.ush
+++ b/Engine/Shaders/Private/LightGridCommon.ush
@@ -224,3 +224,22 @@ float4 UnpackShadowMapChannelMask(uint ShadowMapChannelMaskPacked)
 {
 	return float4((ShadowMapChannelMaskPacked & 1) ? 1.0f : 0.0f, (ShadowMapChannelMaskPacked & 2) ? 1.0f : 0.0f, (ShadowMapChannelMaskPacked & 4) ? 1.0f : 0.0f, (ShadowMapChannelMaskPacked & 8) ? 1.0f : 0.0f);
 }
+
+/**
+ * Helper to pack/unpack the shadow mask for cluster shading and virtual shadow map
+ * This packing encode up to 10 lights into a 32bits uint
+ */
+uint GetShadowMaskMaxLightCount()
+{
+	return 10;
+}
+
+uint PackShadowMask(float InShadowFactor, uint InLightIndex)
+{
+	return (~uint(round(InShadowFactor * 7.0)) & 7u) << (InLightIndex * 3u);
+}
+
+float UnpackShadowMask(uint InShadowMask, uint InLightIndex)
+{
+	return ((InShadowMask >> (InLightIndex * 3u)) & 7u) / 7.0;
+}
\ No newline at end of file
diff --git a/Engine/Source/Runtime/Renderer/Private/HairStrands/HairStrandsTransmittance.cpp b/Engine/Source/Runtime/Renderer/Private/HairStrands/HairStrandsTransmittance.cpp
--- a/Engine/Source/Runtime/Renderer/Private/HairStrands/HairStrandsTransmittance.cpp
+++ b/Engine/Source/Runtime/Renderer/Private/HairStrands/HairStrandsTransmittance.cpp
@@ -54,6 +54,12 @@ static float GetDeepShadowDensityScale() { return FMath::Max(0.0f, GDeepShadowDe
 static float GetDeepShadowDepthBiasScale() { return FMath::Max(0.0f, GDeepShadowDepthBiasScale); }
 ///////////////////////////////////////////////////////////////////////////////////////////////////
 
+enum class EHairTransmittancePassType : uint8
+{
+	PerLight,
+	OnePass
+};
+
 static bool HasDeepShadowData(const FLightSceneInfo* LightSceneInfo, const FHairStrandsMacroGroupDatas& InDatas)
 {
 	for (const FHairStrandsMacroGroupData& MacroGroupData : InDatas)
@@ -168,11 +174,13 @@ class FHairStrandsVoxelTransmittanceMaskCS : public FGlobalShader
 	class FTransmittanceGroupSize : SHADER_PERMUTATION_SPARSE_INT("PERMUTATION_GROUP_SIZE", 32, 64);
 	class FSuperSampling : SHADER_PERMUTATION_INT("PERMUTATION_SUPERSAMPLING", 2);
 	class FTraversal : SHADER_PERMUTATION_INT("PERMUTATION_TRAVERSAL", 2);
-	using FPermutationDomain = TShaderPermutationDomain<FTransmittanceGroupSize, FSuperSampling, FTraversal>;
+	class FOnePass : SHADER_PERMUTATION_BOOL("PERMUTATION_ONE_PASS");
+	using FPermutationDomain = TShaderPermutationDomain<FTransmittanceGroupSize, FSuperSampling, FTraversal, FOnePass>;
 
 	BEGIN_SHADER_PARAMETER_STRUCT(FParameters, )
 		SHADER_PARAMETER_STRUCT_INCLUDE(FSceneTextureParameters, SceneTextures)
 		SHADER_PARAMETER_STRUCT_INCLUDE(ShaderDrawDebug::FShaderDrawDebugParameters, ShaderDrawParameters)
+		SHADER_PARAMETER_STRUCT_INCLUDE(ShaderPrint::FShaderParameters, ShaderPrintUniformBuffer)
 		SHADER_PARAMETER_RDG_UNIFORM_BUFFER(FHairStrandsViewUniformParameters, HairStrands)
 		SHADER_PARAMETER_RDG_UNIFORM_BUFFER(FVirtualVoxelParameters, VirtualVoxel)
 
@@ -181,9 +189,11 @@ class FHairStrandsVoxelTransmittanceMaskCS : public FGlobalShader
 		SHADER_PARAMETER(FVector4, LightPosition)
 		SHADER_PARAMETER(uint32, LightChannelMask)
 
+		SHADER_PARAMETER_RDG_TEXTURE(Texture2D<uint>, ShadowMaskBitsTexture)
 		SHADER_PARAMETER_RDG_TEXTURE(Texture2D, RayMarchMaskTexture)
 		SHADER_PARAMETER_RDG_BUFFER_UAV(RWBuffer, OutTransmittanceMask)
 		SHADER_PARAMETER_STRUCT_REF(FViewUniformShaderParameters, ViewUniformBuffer)
+		SHADER_PARAMETER_STRUCT_REF(FForwardLightData, ForwardLightData)
 
 		RDG_BUFFER_ACCESS(IndirectArgsBuffer, ERHIAccess::IndirectArgs)
 
@@ -198,6 +208,7 @@ public:
 	static void ModifyCompilationEnvironment(const FGlobalShaderPermutationParameters& Parameters, FShaderCompilerEnvironment& OutEnvironment)
 	{
 		FGlobalShader::ModifyCompilationEnvironment(Parameters, OutEnvironment);
+		FForwardLightingParameters::ModifyCompilationEnvironment(Parameters.Platform, OutEnvironment);
 		OutEnvironment.SetDefine(TEXT("SHADER_TRANSMITTANCE_VOXEL"), 1);
 	}
 };
@@ -209,27 +220,41 @@ static FRDGBufferRef AddHairStrandsVoxelTransmittanceMaskPass(
 	FRDGBuilder& GraphBuilder,
 	const FSceneTextureParameters& SceneTextures,
 	const FViewInfo& View,
+	const EHairTransmittancePassType PassType,
 	const FHairStrandsTransmittanceLightParams& Params,
 	const uint32 NodeGroupSize,
 	FRDGBufferRef IndirectArgsBuffer,
-	FRDGTextureRef ScreenShadowMaskSubPixelTexture)
+	FRDGTextureRef ShadowMaskTexture)
 {
 	check(HairStrands::HasViewHairStrandsVoxelData(View));
 	check(NodeGroupSize == 64 || NodeGroupSize == 32);
 
-	FRDGBufferRef OutBuffer = CreateHairStrandsTransmittanceMaskBuffer(GraphBuilder, View.HairStrandsViewData.VisibilityData.MaxNodeCount);
+	const uint32 MaxLightPerPass = 10u; // HAIR_TODO: Need to match the virtual shadow mask bits encoding
+	const uint32 AverageLightPerPixel = PassType == EHairTransmittancePassType::OnePass ? MaxLightPerPass : 1u;
+	FRDGBufferRef OutBuffer = CreateHairStrandsTransmittanceMaskBuffer(GraphBuilder, View.HairStrandsViewData.VisibilityData.MaxNodeCount * AverageLightPerPixel);
 
 	FHairStrandsVoxelTransmittanceMaskCS::FParameters* Parameters = GraphBuilder.AllocParameters<FHairStrandsVoxelTransmittanceMaskCS::FParameters>();
 	Parameters->ViewUniformBuffer = View.ViewUniformBuffer;
 	Parameters->SceneTextures = SceneTextures;
 	Parameters->OutTransmittanceMask = GraphBuilder.CreateUAV(OutBuffer, FHairStrandsTransmittanceMaskData::Format);
+	if (PassType == EHairTransmittancePassType::OnePass)
+	{
+		Parameters->ForwardLightData = View.ForwardLightingResources->ForwardLightDataUniformBuffer;
+		Parameters->RayMarchMaskTexture = nullptr;
+		Parameters->ShadowMaskBitsTexture = ShadowMaskTexture;
+	}
+	else
+	{
+		Parameters->LightDirection = Params.LightDirection;
+		Parameters->LightPosition = Params.LightPosition;
+		Parameters->LightRadius = Params.LightRadius;
+		Parameters->RayMarchMaskTexture = ShadowMaskTexture ? ShadowMaskTexture : GSystemTextures.GetWhiteDummy(GraphBuilder);
+		Parameters->ShadowMaskBitsTexture = nullptr;
+	}
+
 	Parameters->LightChannelMask = Params.LightChannelMask;
-	Parameters->LightDirection = Params.LightDirection;
-	Parameters->LightPosition = Params.LightPosition;
-	Parameters->LightRadius = Params.LightRadius;
 	Parameters->IndirectArgsBuffer = IndirectArgsBuffer;
 	Parameters->HairStrands = HairStrands::BindHairStrandsViewUniformParameters(View);
-	Parameters->RayMarchMaskTexture = ScreenShadowMaskSubPixelTexture ? ScreenShadowMaskSubPixelTexture : GSystemTextures.GetWhiteDummy(GraphBuilder);
 	Parameters->VirtualVoxel = HairStrands::BindHairStrandsVoxelUniformParameters(View);
 
 	if (ShaderDrawDebug::IsShaderDrawDebugEnabled(View))
@@ -237,6 +262,11 @@ static FRDGBufferRef AddHairStrandsVoxelTransmittanceMaskPass(
 		ShaderDrawDebug::SetParameters(GraphBuilder, View.ShaderDrawData, Parameters->ShaderDrawParameters);
 	}
 
+	if (ShaderPrint::IsSupported(View) && ShaderPrint::IsEnabled())
+	{
+		ShaderPrint::SetParameters(GraphBuilder, View, Parameters->ShaderPrintUniformBuffer);
+	}
+
 	const bool bIsSuperSampled = GHairStrandsTransmittanceSuperSampling > 0;
 	const bool bIsMipTraversal = GHairStrandsTransmittanceMaskUseMipTraversal > 0;
 
@@ -244,11 +274,11 @@ static FRDGBufferRef AddHairStrandsVoxelTransmittanceMaskPass(
 	PermutationVector.Set<FHairStrandsVoxelTransmittanceMaskCS::FTransmittanceGroupSize>(NodeGroupSize);
 	PermutationVector.Set<FHairStrandsVoxelTransmittanceMaskCS::FSuperSampling>(bIsSuperSampled ? 1 : 0);
 	PermutationVector.Set<FHairStrandsVoxelTransmittanceMaskCS::FTraversal>(bIsMipTraversal ? 1 : 0);
-
+	PermutationVector.Set<FHairStrandsVoxelTransmittanceMaskCS::FOnePass>(PassType == EHairTransmittancePassType::OnePass);
 	TShaderMapRef<FHairStrandsVoxelTransmittanceMaskCS> ComputeShader(View.ShaderMap, PermutationVector);
 	FComputeShaderUtils::AddPass(
 		GraphBuilder,
-		RDG_EVENT_NAME("HairStrandsTransmittanceMask(Voxel)"),
+		RDG_EVENT_NAME("HairStrandsTransmittanceMask(Voxel,%s)", PassType == EHairTransmittancePassType::OnePass ? TEXT("OnePass") : TEXT("PerLight")),
 		ComputeShader,
 		Parameters,
 		IndirectArgsBuffer,
@@ -701,13 +731,9 @@ static FHairStrandsTransmittanceMaskData InternalRenderHairStrandsTransmittanceM
 	RDG_EVENT_SCOPE(GraphBuilder, "HairStrandsTransmittanceMask");
 	RDG_GPU_STAT_SCOPE(GraphBuilder, HairStrandsTransmittanceMask);
 
-	const FHairLUT InHairLUT = GetHairLUT(GraphBuilder, View);
-
 	// Note: GbufferB.a store the shading model on the 4 lower bits (MATERIAL_SHADINGMODEL_HAIR)
 	FSceneTextureParameters SceneTextures = GetSceneTextureParameters(GraphBuilder);
 
-	FRDGTextureRef HairLUTTexture = InHairLUT.Textures[HairLUTType_DualScattering];
-
 	bool bHasFoundLight = false;
 	if (!IsHairStrandsForVoxelTransmittanceAndShadowEnable())
 	{
@@ -776,6 +802,7 @@ static FHairStrandsTransmittanceMaskData InternalRenderHairStrandsTransmittanceM
 			GraphBuilder,
 			SceneTextures,
 			View,
+			EHairTransmittancePassType::PerLight,
 			Params,
 			VisibilityData.NodeGroupSize,
 			VisibilityData.NodeIndirectArg,
@@ -784,6 +811,38 @@ static FHairStrandsTransmittanceMaskData InternalRenderHairStrandsTransmittanceM
 
 	return Out;
 }
+	
+FHairStrandsTransmittanceMaskData RenderHairStrandsOnePassTransmittanceMask(
+	FRDGBuilder& GraphBuilder,
+	const FViewInfo& View, 
+	FRDGTextureRef ShadowMaskBits)
+{
+	FHairStrandsTransmittanceMaskData Out;
+	if (HairStrands::HasViewHairStrandsData(View) && View.HairStrandsViewData.MacroGroupDatas.Num() > 0)
+	{
+		DECLARE_GPU_STAT(HairStrandsOnePassTransmittanceMask);
+		RDG_EVENT_SCOPE(GraphBuilder, "HairStrandsTransmittanceMask(OnePass)");
+		RDG_GPU_STAT_SCOPE(GraphBuilder, HairStrandsOnePassTransmittanceMask);
+
+		if (HairStrands::HasViewHairStrandsVoxelData(View))
+		{
+			// Note: GbufferB.a store the shading model on the 4 lower bits (MATERIAL_SHADINGMODEL_HAIR)
+			FSceneTextureParameters SceneTextures = GetSceneTextureParameters(GraphBuilder);
+			FHairStrandsTransmittanceLightParams DummyParams;
+
+			Out.TransmittanceMask = AddHairStrandsVoxelTransmittanceMaskPass(
+				GraphBuilder,
+				SceneTextures,
+				View,
+				EHairTransmittancePassType::OnePass,
+				DummyParams,
+				View.HairStrandsViewData.VisibilityData.NodeGroupSize,
+				View.HairStrandsViewData.VisibilityData.NodeIndirectArg,
+				ShadowMaskBits);
+		}
+	}
+	return Out;
+}
 
 FHairStrandsTransmittanceMaskData RenderHairStrandsTransmittanceMask(
 	FRDGBuilder& GraphBuilder,
diff --git a/Engine/Source/Runtime/Renderer/Private/HairStrands/HairStrandsTransmittance.h b/Engine/Source/Runtime/Renderer/Private/HairStrands/HairStrandsTransmittance.h
--- a/Engine/Source/Runtime/Renderer/Private/HairStrands/HairStrandsTransmittance.h
+++ b/Engine/Source/Runtime/Renderer/Private/HairStrands/HairStrandsTransmittance.h
@@ -26,9 +26,15 @@ void RenderHairStrandsShadowMask(
 	const FLightSceneInfo* LightSceneInfo,
 	FRDGTextureRef ScreenShadowMaskTexture); 
 
-/// Write hair transmittance onto screen shadow mask
+/// Output hair transmittance per hair sample for a given light
 FHairStrandsTransmittanceMaskData RenderHairStrandsTransmittanceMask(
 	FRDGBuilder& GraphBuilder,
 	const FViewInfo& View,
 	const class FLightSceneInfo* LightSceneInfo,
-	FRDGTextureRef ScreenShadowMaskSubPixelTexture);
\ No newline at end of file
+	FRDGTextureRef ScreenShadowMaskSubPixelTexture);
+
+/// Output hair transmittance per hair sample for all lights using the forward cluster lights
+FHairStrandsTransmittanceMaskData RenderHairStrandsOnePassTransmittanceMask(
+	FRDGBuilder& GraphBuilder,
+	const FViewInfo& View,
+	FRDGTextureRef ShadowMaskBits);
\ No newline at end of file
