Add cluster shadring permutation for applying hair shadow on top of opaque geometries.

diff --git a/Engine/Shaders/Private/LightGridCommon.ush b/Engine/Shaders/Private/LightGridCommon.ush
--- a/Engine/Shaders/Private/LightGridCommon.ush
+++ b/Engine/Shaders/Private/LightGridCommon.ush
@@ -234,9 +234,9 @@ uint GetShadowMaskMaxLightCount()
 	return 10;
 }
 
-uint PackShadowMask(float InShadowFactor, uint InLightIndex)
+void PackShadowMask(inout uint OutShadowMask, float InShadowFactor, uint InLightIndex)
 {
-	return (~uint(round(InShadowFactor * 7.0)) & 7u) << (InLightIndex * 3u);
+	OutShadowMask ^= (~uint(round(InShadowFactor * 7.0)) & 7u) << (InLightIndex * 3u);
 }
 
 float UnpackShadowMask(uint InShadowMask, uint InLightIndex)
diff --git a/Engine/Shaders/Private/VirtualShadowMaps/VirtualShadowMapProjection.usf b/Engine/Shaders/Private/VirtualShadowMaps/VirtualShadowMapProjection.usf
--- a/Engine/Shaders/Private/VirtualShadowMaps/VirtualShadowMapProjection.usf
+++ b/Engine/Shaders/Private/VirtualShadowMaps/VirtualShadowMapProjection.usf
@@ -6,6 +6,7 @@
 
 #define SUPPORT_CONTACT_SHADOWS 1
 #include "../HairStrands/HairStrandsVisibilityCommonStruct.ush"
+#include "../HairStrands/HairStrandsVoxelPageCommonStruct.ush"
 #include "../Common.ush"
 #include "../SceneTexturesCommon.ush"
 #include "../DeferredShadingCommon.ush"
@@ -21,6 +22,13 @@
 #include "ProjectionDirectional.ush"
 #include "ProjectionSpot.ush"
 
+#if HAS_HAIR_STRANDS
+#define VOXEL_TRAVERSAL_TYPE VOXEL_TRAVERSAL_LINEAR_MIPMAP
+#define VOXEL_TRAVERSAL_DEBUG 0
+#include "../HairStrands/HairStrandsVoxelPageCommon.ush"
+#include "../HairStrands/HairStrandsVoxelPageTraversal.ush"
+#endif
+
 #ifndef DIRECTIONAL_LIGHT
 #define DIRECTIONAL_LIGHT 0
 #endif
@@ -159,6 +167,15 @@ FLightShaderParameters ConvertFromLocal( const FLocalLightData LightData )
 	return Light;
 }
 
+float4 ComputeRandom4(uint2 PixelPosition)
+{
+	const uint InSeed = View.StateFrameIndexMod8;
+	const uint2 Seed0 = Rand3DPCG16(int3(PixelPosition, InSeed)).xy;
+	const uint2 Seed1 = Rand3DPCG16(int3(PixelPosition + 17, InSeed)).xy;
+	return float4(
+		Hammersley16(InSeed, 8, Seed0),
+		Hammersley16(InSeed, 8, Seed1));
+}
 
 int4 ProjectionRect;
 float ContactShadowLength;
@@ -288,6 +305,67 @@ FProjectLightResult ProjectLight(
 				HitDistance = ContactShadowHitDistance;
 			}
 		}
+
+		// Hair strands voxel traversal to apply hair shadow on top of opaque geometry
+		#if HAS_HAIR_STRANDS
+		if (ShadowFactor > 0)
+		{		
+			#if DIRECTIONAL_LIGHT
+			const float VoxelShadowMaxDistance = 1000.0; // 100m shadow distance
+			const float3 LightPosition = Light.Direction * VoxelShadowMaxDistance + WorldPosition;
+			#else
+			const float3 LightPosition = Light.Position;
+			#endif
+
+			// Depth bias
+			// Origin is shifted 2 voxels away towards the light + a constant bias of the size of the voxel
+			float3 OriginPosition = WorldPosition;
+			const float4 Random = ComputeRandom4(PixelPos);
+			{
+				const float PositionBiasScale = 0.5f;
+				const float3 DepthBias = VirtualVoxel.VoxelWorldSize * (VirtualVoxel.DepthBiasScale_Shadow * L + PositionBiasScale * (Random.xyz * 2 - 1));
+				OriginPosition += DepthBias;
+			}
+
+			const float DistanceThreshold = 100000.0f;
+			const float CoverageThreshold = 0.995f; // When Coverage is high, we do not trace shadow on opaque since hair/fur is covering the background.
+
+			FVirtualVoxelCommonDesc CommonDesc;
+			CommonDesc.PageCountResolution		= VirtualVoxel.PageCountResolution;
+			CommonDesc.VoxelWorldSize			= VirtualVoxel.VoxelWorldSize;
+			CommonDesc.PageTextureResolution	= VirtualVoxel.PageTextureResolution;
+			CommonDesc.PageResolution			= VirtualVoxel.PageResolution;
+
+			FHairTraversalSettings TraversalSettings = InitHairTraversalSettings();
+			TraversalSettings.DensityScale		= VirtualVoxel.DensityScale_Shadow;
+			TraversalSettings.CountThreshold	= 1;
+			TraversalSettings.DistanceThreshold = DistanceThreshold;
+			TraversalSettings.bDebugEnabled		= false;
+			TraversalSettings.SteppingScale		= VirtualVoxel.SteppingScale_Shadow;
+			TraversalSettings.Random			= Random.xyz;
+			TraversalSettings.PixelRadius		= SceneDepth * VirtualVoxel.HairCoveragePixelRadiusAtDepth1;
+
+			const uint VoxelDescCount = VirtualVoxel.NodeDescCount;
+			for (uint VoxelDescIt=0; VoxelDescIt<VoxelDescCount; ++VoxelDescIt)
+			{
+				const FPackedVirtualVoxelNodeDesc PackedNode = VirtualVoxel.NodeDescBuffer[VoxelDescIt];
+				const FVirtualVoxelNodeDesc NodeDesc = UnpackVoxelNode(PackedNode, VirtualVoxel.PageResolution);
+
+				FHairTraversalResult Result = InitHairTraversalResult();
+				Result = ComputeHairCountVirtualVoxel(
+					OriginPosition,
+					LightPosition,
+					CommonDesc,
+					NodeDesc,
+					VirtualVoxel.PageIndexBuffer,
+					VirtualVoxel.PageIndexOccupancyBuffer,
+					VirtualVoxel.PageTexture,
+					TraversalSettings);
+
+				ShadowFactor = min(ShadowFactor, saturate(1 - Result.HairCount));
+			}
+		}
+		#endif
 	}
 
 	FProjectLightResult Result;
@@ -337,7 +415,7 @@ void VirtualShadowMapProjection(
 	const FCulledLightsGridData CulledLightGridData = GetCulledLightsGrid( GridLinearIndex, EyeIndex );
 
 	// We can only handle 10 lights in our output encoding right now, so no purpose in computing more
-	const uint LightCount = min( 10, CulledLightGridData.NumLocalLights );
+	const uint LightCount = min(GetShadowMaskMaxLightCount(), CulledLightGridData.NumLocalLights );
 
 	uint ShadowMask = 0xffffffff;
 
@@ -363,7 +441,7 @@ void VirtualShadowMapProjection(
 				GBufferNormal,
 				Noise);
 
-			ShadowMask ^= ( ~uint( round( Result.ShadowFactor * 7.0 ) ) & 7 ) << (Index*3);
+			PackShadowMask(ShadowMask, Result.ShadowFactor, Index);
 		}
 	}
 
diff --git a/Engine/Source/Runtime/Renderer/Private/VirtualShadowMaps/VirtualShadowMapArray.cpp b/Engine/Source/Runtime/Renderer/Private/VirtualShadowMaps/VirtualShadowMapArray.cpp
--- a/Engine/Source/Runtime/Renderer/Private/VirtualShadowMaps/VirtualShadowMapArray.cpp
+++ b/Engine/Source/Runtime/Renderer/Private/VirtualShadowMaps/VirtualShadowMapArray.cpp
@@ -941,7 +941,7 @@ void FVirtualShadowMapArray::BuildPageAllocations(
 
 						FComputeShaderUtils::AddPass(
 							GraphBuilder,
-							RDG_EVENT_NAME("GeneratePageFlagsFromPixels"),
+							RDG_EVENT_NAME("GeneratePageFlagsFromPixels(%s)", bHairPass ? TEXT("HairStrands") : TEXT("GBuffer")),
 							ComputeShader,
 							PassParameters,
 							FIntVector(GridSize.X, GridSize.Y, 1)
diff --git a/Engine/Source/Runtime/Renderer/Private/VirtualShadowMaps/VirtualShadowMapProjection.cpp b/Engine/Source/Runtime/Renderer/Private/VirtualShadowMaps/VirtualShadowMapProjection.cpp
--- a/Engine/Source/Runtime/Renderer/Private/VirtualShadowMaps/VirtualShadowMapProjection.cpp
+++ b/Engine/Source/Runtime/Renderer/Private/VirtualShadowMaps/VirtualShadowMapProjection.cpp
@@ -171,18 +171,21 @@ class FVirtualShadowMapProjectionCS : public FGlobalShader
 	class FSMRTAdaptiveRayCountDim	: SHADER_PERMUTATION_BOOL("SMRT_ADAPTIVE_RAY_COUNT");
 	class FTwoPhysicalTexturesDim	: SHADER_PERMUTATION_BOOL("TWO_PHYSICAL_TEXTURES");
 	class FOnePassProjectionDim		: SHADER_PERMUTATION_BOOL("ONE_PASS_PROJECTION");
+	class FHairStrandsDim			: SHADER_PERMUTATION_BOOL("HAS_HAIR_STRANDS");
 
 	using FPermutationDomain = TShaderPermutationDomain<
 		FDirectionalLightDim,
 		FOnePassProjectionDim,
 		FSMRTAdaptiveRayCountDim,
-		FTwoPhysicalTexturesDim	>;
+		FTwoPhysicalTexturesDim,
+		FHairStrandsDim>;
 
 	BEGIN_SHADER_PARAMETER_STRUCT(FParameters, )
 		SHADER_PARAMETER_STRUCT_INCLUDE(FVirtualShadowMapSamplingParameters, SamplingParameters)
 		SHADER_PARAMETER_RDG_UNIFORM_BUFFER(FSceneTextureUniformParameters, SceneTexturesStruct)
 		SHADER_PARAMETER_RDG_UNIFORM_BUFFER(FHairStrandsViewUniformParameters, HairStrands)
-		SHADER_PARAMETER_STRUCT_REF(FViewUniformShaderParameters, View)		
+		SHADER_PARAMETER_RDG_UNIFORM_BUFFER(FVirtualVoxelParameters, HairStrandsVoxel)
+		SHADER_PARAMETER_STRUCT_REF(FViewUniformShaderParameters, View)
 		SHADER_PARAMETER(FIntVector4, ProjectionRect)
 		SHADER_PARAMETER(float, ContactShadowLength)
 		SHADER_PARAMETER(float, NormalOffsetWorld)
@@ -249,6 +252,9 @@ static void RenderVirtualShadowMapProjectionCommon(
 	const FLightSceneProxy* LightProxy = nullptr,
 	int32 VirtualShadowMapId = INDEX_NONE)
 {
+	// Use hair strands data (i.e., hair voxel tracing) only for Gbuffer input for casting hair shadow onto opaque geometry.
+	const bool bHasHairStrandsData = InputType == EVirtualShadowMapProjectionInputType::GBuffer && HairStrands::HasViewHairStrandsData(View);
+
 	FVirtualShadowMapProjectionCS::FParameters* PassParameters = GraphBuilder.AllocParameters< FVirtualShadowMapProjectionCS::FParameters >();
 	PassParameters->SamplingParameters = VirtualShadowMapArray.GetSamplingParameters(GraphBuilder);
 	PassParameters->SceneTexturesStruct = SceneTextures.UniformBuffer;
@@ -259,6 +265,10 @@ static void RenderVirtualShadowMapProjectionCommon(
 	PassParameters->NormalOffsetWorld = CVarNormalOffsetWorld.GetValueOnRenderThread();
 	PassParameters->InputType = uint32(InputType);
 	PassParameters->HairStrands = HairStrands::BindHairStrandsViewUniformParameters(View);
+	if (bHasHairStrandsData)
+	{
+		PassParameters->HairStrandsVoxel = HairStrands::BindHairStrandsVoxelUniformParameters(View);
+	}
 
 	bool bDirectionalLight = false;
 	bool bOnePassProjection = LightProxy == nullptr;
@@ -302,7 +312,7 @@ static void RenderVirtualShadowMapProjectionCommon(
 	PermutationVector.Set< FVirtualShadowMapProjectionCS::FOnePassProjectionDim >( bOnePassProjection );
 	PermutationVector.Set< FVirtualShadowMapProjectionCS::FSMRTAdaptiveRayCountDim >( bAdaptiveRayCount );
 	PermutationVector.Set< FVirtualShadowMapProjectionCS::FTwoPhysicalTexturesDim >( GVirtualShadowMapAtomicWrites == 0 );
-
+	PermutationVector.Set< FVirtualShadowMapProjectionCS::FHairStrandsDim >(bHasHairStrandsData ? 1 : 0);
 	auto ComputeShader = View.ShaderMap->GetShader< FVirtualShadowMapProjectionCS >( PermutationVector );
 	ClearUnusedGraphResources( ComputeShader, PassParameters );
 	ValidateShaderParameters( ComputeShader, *PassParameters );
@@ -310,7 +320,9 @@ static void RenderVirtualShadowMapProjectionCommon(
 	const FIntPoint GroupCount = FIntPoint::DivideAndRoundUp( ProjectionRect.Size(), 8 );
 	FComputeShaderUtils::AddPass(
 		GraphBuilder,
-		bAdaptiveRayCount ? RDG_EVENT_NAME("VirtualShadowMapProjection") : RDG_EVENT_NAME("VirtualShadowMapProjection_StaticRayCount"),
+		RDG_EVENT_NAME("VirtualShadowMapProjection(RayCount:%s,Input:%s)", 
+			bAdaptiveRayCount ? TEXT("Adaptive") : TEXT("Static"), 
+			(InputType == EVirtualShadowMapProjectionInputType::HairStrands) ? TEXT("HairStrands") : TEXT("GBuffer")),
 		ComputeShader,
 		PassParameters,
 		FIntVector( GroupCount.X, GroupCount.Y, 1 )
