Improve approximate caustics in path tracer

Use per lobe roughness instead of assigning a global roughness to all rays. This preserves reflected highlights better off of multi-lobed materials like dielectrics that have both a diffuse and specular lobe.

Add cvar to skip emissive shaders

This can be useful to debug noise coming from small emissive surfaces in the scene, or to deal with scenes that have both lights and modelled emitters.

Tweak transmission color to keep overall albedo under 1

Some assets with the twosided foliage or subsurface shading models have albedo over 1. Try to compensate for this by darkening both colors.

Fix potential bug with MIS weight calculation

diff --git a/Engine/Shaders/Private/PathTracing/Material/PathTracingClearCoat.ush b/Engine/Shaders/Private/PathTracing/Material/PathTracingClearCoat.ush
--- a/Engine/Shaders/Private/PathTracing/Material/PathTracingClearCoat.ush
+++ b/Engine/Shaders/Private/PathTracing/Material/PathTracingClearCoat.ush
@@ -140,6 +140,7 @@ FMaterialSample ClearCoat_SampleMaterial(
 
 	// choose between coat and bottom layer
 	float3 OutDirection = 0;
+	float OutRoughness = 0;
 	if (RandSample.x < CoatSelectionPdf)
 	{
 		RandSample.x /= CoatSelectionPdf;
@@ -162,6 +163,7 @@ FMaterialSample ClearCoat_SampleMaterial(
 		}
 		// reflect and transform
 		OutDirection = mul(L, Basis);
+		OutRoughness = Roughness;
 	}
 	else
 	{
@@ -172,16 +174,20 @@ FMaterialSample ClearCoat_SampleMaterial(
 		float3 Weight = 0;
 		float Pdf = 0;
 		Payload.WorldNormal = GetClearCoatBottomNormal(Payload);
-		OutDirection = DefaultLit_SampleMaterial(RayDirection, Payload, RandSample).Direction;
+		FMaterialSample InnerSample = DefaultLit_SampleMaterial(RayDirection, Payload, RandSample);
 		Payload.WorldNormal = N_World;
+
+		OutDirection = InnerSample.Direction;
+		OutRoughness = InnerSample.Roughness;
 #else
 		float3 Nb = GetClearCoatBottomNormal(Payload);
 		OutDirection = TangentToWorld(CosineSampleHemisphere(RandSample.xy).xyz, Nb);
+		OutRoughness = 1.0;
 #endif
 	}
 
 	// defer to eval for now to get proper MIS mix of both events
 	// TODO: could use temporary results above to inline Eval logic here
 	FMaterialEval Result = ClearCoat_EvalMaterial(RayDirection, OutDirection, Payload);
-	return CreateMaterialSample(OutDirection, Result.Weight, Result.Pdf, 1.0);
+	return CreateMaterialSample(OutDirection, Result.Weight, Result.Pdf, 1.0, OutRoughness);
 }
diff --git a/Engine/Shaders/Private/PathTracing/Material/PathTracingCloth.ush b/Engine/Shaders/Private/PathTracing/Material/PathTracingCloth.ush
--- a/Engine/Shaders/Private/PathTracing/Material/PathTracingCloth.ush
+++ b/Engine/Shaders/Private/PathTracing/Material/PathTracingCloth.ush
@@ -88,6 +88,7 @@ FMaterialSample Cloth_SampleMaterial(
 
 	// Randomly choose to sample diffuse, specular or cloth lobe
 	float3 L = 0, H = 0;
+	float OutRoughness = 0;
 	if (RandSample.x < LobeSelectionPdf)
 	{
 		RandSample.x /= LobeSelectionPdf;
@@ -101,6 +102,7 @@ FMaterialSample Cloth_SampleMaterial(
 
 		L = Result.xyz;
 		H = normalize(L + V);
+		OutRoughness = 1.0;
 	}
 	else
 	{
@@ -117,6 +119,7 @@ FMaterialSample Cloth_SampleMaterial(
 			// invalid output direction, exit early
 			return NullMaterialSample();
 		}
+		OutRoughness = Payload.Roughness;
 	}
 
 	// With a valid direction in hand -- now evaluate the BxDF (taking advantage of already computed terms)
@@ -146,5 +149,5 @@ FMaterialSample Cloth_SampleMaterial(
 	AddLobeWithMIS(OutWeight, OutPdf, SpecWeight, SpecPdf, 1 - LobeSelectionPdf);
 
 	// transform to world space
-	return CreateMaterialSample(normalize(mul(L, Basis)), OutWeight, OutPdf, 1.0);
+	return CreateMaterialSample(normalize(mul(L, Basis)), OutWeight, OutPdf, 1.0, OutRoughness);
 }
diff --git a/Engine/Shaders/Private/PathTracing/Material/PathTracingDefaultLit.ush b/Engine/Shaders/Private/PathTracing/Material/PathTracingDefaultLit.ush
--- a/Engine/Shaders/Private/PathTracing/Material/PathTracingDefaultLit.ush
+++ b/Engine/Shaders/Private/PathTracing/Material/PathTracingDefaultLit.ush
@@ -78,6 +78,7 @@ FMaterialSample DefaultLit_SampleMaterial(
 	
 	// Randomly choose to sample diffuse or specular
 	float3 L = 0, H = 0;
+	float OutRoughness = 0;
 	if (RandSample.x < LobeSelectionPdf)
 	{
 		RandSample.x /= LobeSelectionPdf;
@@ -88,13 +89,13 @@ FMaterialSample DefaultLit_SampleMaterial(
 
 		L = Result.xyz;
 		H = normalize(L + V);
+		OutRoughness = 1.0;
 	}
 	else
 	{
 		RandSample.x -= LobeSelectionPdf;
 		RandSample.x /= (1.0 - LobeSelectionPdf);
 
-
 		// #dxr_todo: Evaluate UniformSampleDiskConcentric
 		H = ImportanceSampleVisibleGGX_aniso(UniformSampleDisk(RandSample.xy), Alpha, V).xyz;
 
@@ -104,6 +105,7 @@ FMaterialSample DefaultLit_SampleMaterial(
 			// invalid output direction, exit early
 			return NullMaterialSample();
 		}
+		OutRoughness = Payload.Roughness;
 	}
 
 	// With a valid direction in hand -- now evaluate the BxDF (taking advantage of already computed terms)
@@ -125,5 +127,5 @@ FMaterialSample DefaultLit_SampleMaterial(
 	AddLobeWithMIS(OutWeight, OutPdf, SpecWeight, SpecPdf, 1.0 - LobeSelectionPdf);
 
 	// transform to world space
-	return CreateMaterialSample(normalize(mul(L, Basis)), OutWeight, OutPdf, 1.0);
+	return CreateMaterialSample(normalize(mul(L, Basis)), OutWeight, OutPdf, 1.0, OutRoughness);
 }
diff --git a/Engine/Shaders/Private/PathTracing/Material/PathTracingLambert.ush b/Engine/Shaders/Private/PathTracing/Material/PathTracingLambert.ush
--- a/Engine/Shaders/Private/PathTracing/Material/PathTracingLambert.ush
+++ b/Engine/Shaders/Private/PathTracing/Material/PathTracingLambert.ush
@@ -13,7 +13,7 @@ FMaterialSample Lambert_SampleMaterial(
 
 	float4 SampledValue = CosineSampleHemisphere(RandSample.xy);
 
-	return CreateMaterialSample(TangentToWorld(SampledValue.xyz, N_World), GetBaseColor(Payload), SampledValue.w, 1.0);
+	return CreateMaterialSample(TangentToWorld(SampledValue.xyz, N_World), GetBaseColor(Payload), SampledValue.w, 1.0, 1.0);
 }
 
 float Lambert_Pdf(float3 OutgoingDirection, FMaterialClosestHitPayload Payload)
diff --git a/Engine/Shaders/Private/PathTracing/Material/PathTracingMaterialCommon.ush b/Engine/Shaders/Private/PathTracing/Material/PathTracingMaterialCommon.ush
--- a/Engine/Shaders/Private/PathTracing/Material/PathTracingMaterialCommon.ush
+++ b/Engine/Shaders/Private/PathTracing/Material/PathTracingMaterialCommon.ush
@@ -49,6 +49,12 @@ float GetClearCoatRoughness(FMaterialClosestHitPayload HitInfo)
 	return max(HitInfo.CustomData.y, 0.02);
 }
 
+float SetClearCoatRoughness(inout FMaterialClosestHitPayload HitInfo, float Roughness)
+{
+	return HitInfo.CustomData.y = Roughness;
+}
+
+
 float3 GetClearCoatBottomNormal(FMaterialClosestHitPayload HitInfo)
 {
 #ifdef PATH_TRACING
@@ -197,6 +203,7 @@ struct FMaterialSample
 	float3 Weight;
 	float Pdf;
 	float PositionBiasSign;
+	float Roughness;
 };
 
 FMaterialSample NullMaterialSample()
@@ -205,13 +212,14 @@ FMaterialSample NullMaterialSample()
 	return (FMaterialSample)0;
 }
 
-FMaterialSample CreateMaterialSample(float3 Direction, float3 Weight, float Pdf, float PositionBiasSign)
+FMaterialSample CreateMaterialSample(float3 Direction, float3 Weight, float Pdf, float PositionBiasSign, float Roughness)
 {
 	FMaterialSample Result;
 	Result.Direction = Direction;
 	Result.Weight = Weight;
 	Result.Pdf = Pdf;
 	Result.PositionBiasSign = PositionBiasSign;
+	Result.Roughness = Roughness;
 	return Result;
 }
 
diff --git a/Engine/Shaders/Private/PathTracing/Material/PathTracingMaterialSampling.ush b/Engine/Shaders/Private/PathTracing/Material/PathTracingMaterialSampling.ush
--- a/Engine/Shaders/Private/PathTracing/Material/PathTracingMaterialSampling.ush
+++ b/Engine/Shaders/Private/PathTracing/Material/PathTracingMaterialSampling.ush
@@ -167,10 +167,9 @@ float3 EstimateMaterialAlbedo(FMaterialClosestHitPayload Payload)
 	return Albedo;
 }
 
-void AdjustPathAndMaterialRoughness(inout FMaterialClosestHitPayload Payload, inout float PathRoughness)
+void AdjustMaterialRoughness(inout FMaterialClosestHitPayload Payload, float PathRoughness)
 {
 	// boost material roughness according to what we've seen so far
-	float CoatRoughness = Payload.Roughness;
 	switch (Payload.ShadingModelID)
 	{
 		case SHADINGMODELID_UNLIT:
@@ -178,7 +177,7 @@ void AdjustPathAndMaterialRoughness(inout FMaterialClosestHitPayload Payload, in
 			return;
 		case SHADINGMODELID_CLEAR_COAT:
 			// clamp coat roughness
-			CoatRoughness = Payload.CustomData.y = max(GetClearCoatRoughness(Payload), PathRoughness);
+			SetClearCoatRoughness(Payload, max(GetClearCoatRoughness(Payload), PathRoughness));
 			// fallthrough
 		case SHADINGMODELID_TWOSIDED_FOLIAGE:
 		case SHADINGMODELID_DEFAULT_LIT:
@@ -201,12 +200,31 @@ void AdjustPathAndMaterialRoughness(inout FMaterialClosestHitPayload Payload, in
 			Payload.Roughness = 1.0;
 			break;
 	}
-	// now set the path roughness going forward
+}
+
+// returns an average roughness for the whole material (for the path roughness of shadow rays from this material)
+float GetAverageRoughness(FMaterialClosestHitPayload Payload)
+{
 	const float GlassPercent = (Payload.BlendingMode == RAY_TRACING_BLEND_MODE_TRANSLUCENT && Payload.Ior != 1.0) ? 1 - Payload.Opacity : 0.0;
-	PathRoughness = lerp(1.0 /* diffuse lobe */, max(CoatRoughness, Payload.Roughness), Payload.Metallic);
-	PathRoughness = lerp(PathRoughness, Payload.Roughness, GlassPercent);
+	float Roughness = Payload.Roughness;
+	switch (Payload.ShadingModelID)
+	{
+		case SHADINGMODELID_CLEAR_COAT:
+			Roughness = max(Roughness, GetClearCoatRoughness(Payload)); 
+			break;
+		case SHADINGMODELID_SUBSURFACE_PROFILE:
+		{
+			float3 DualRoughnessData = GetDualRoughnessSpecular(Payload);
+			Roughness = max(DualRoughnessData.x, DualRoughnessData.y);
+			break;
+		}
+	}
+	float AvgRoughness = lerp(1.0 /* diffuse lobe */, Roughness, Payload.Metallic);
+	AvgRoughness = lerp(AvgRoughness, Payload.Roughness, GlassPercent);
+	return AvgRoughness;
 }
 
+
 FMaterialClosestHitPayload UnpackPathTracingPayload(FPackedMaterialClosestHitPayload PackedInput, RayDesc Ray)
 {
 	// the path tracer has a few more requirements compared to the ray-tracing path.
@@ -227,6 +245,22 @@ FMaterialClosestHitPayload UnpackPathTracingPayload(FPackedMaterialClosestHitPay
 
 			// Reset opacity to avoid special cases elsewhere
 			HitPayload.Opacity = 1;
+			
+			// fallthrough to case below to avoid over-bright SSS
+		}
+		case SHADINGMODELID_TWOSIDED_FOLIAGE:
+		{
+			float3 Front = HitPayload.DiffuseColor;
+			float3 Back = GetSubsurfaceColor(HitPayload);
+			float3 Sum = Front + Back;
+			float S = max(Sum.x, max(Sum.y, Sum.z));
+			if (S > 1)
+			{
+				// nudge the material back to something physically plausible
+				// NOTE: we ignore spec here since it should be accounted for by the brdf model itself
+				HitPayload.DiffuseColor = Front / S;
+				SetSubsurfaceColor(HitPayload, Back / S);
+			}
 			break;
 		}
 		case SHADINGMODELID_SUBSURFACE_PROFILE:
diff --git a/Engine/Shaders/Private/PathTracing/Material/PathTracingRadianceProbe.ush b/Engine/Shaders/Private/PathTracing/Material/PathTracingRadianceProbe.ush
--- a/Engine/Shaders/Private/PathTracing/Material/PathTracingRadianceProbe.ush
+++ b/Engine/Shaders/Private/PathTracing/Material/PathTracingRadianceProbe.ush
@@ -123,7 +123,7 @@ FMaterialSample RadianceProbe_SampleMaterial(
 	// As the name suggests, radiance probe has a combined throughput (Pdf * Weight) == 1
 	float OutWeight = (OutPdf > 0.0f) ? (1.0f / OutPdf) : 0;
 
-	return CreateMaterialSample(OutDirection, OutWeight, OutPdf, 1.0);
+	return CreateMaterialSample(OutDirection, OutWeight, OutPdf, 1.0, 1.0);
 }
 
 float RadianceProbe_Pdf(float3 OutgoingDirection, FMaterialClosestHitPayload Payload)
@@ -277,7 +277,7 @@ FMaterialSample RadianceProbe_SampleMaterial(
 
 	float4 SampledValue = CosineSampleHemisphere(RandSample.xy);
 
-	return CreateMaterialSample(TangentToWorld(SampledValue.xyz, N_World), 1.0, SampledValue.w, 1.0);
+	return CreateMaterialSample(TangentToWorld(SampledValue.xyz, N_World), 1.0, SampledValue.w, 1.0, 1.0);
 }
 
 float RadianceProbe_Pdf(float3 OutgoingDirection, FMaterialClosestHitPayload Payload)
diff --git a/Engine/Shaders/Private/PathTracing/Material/PathTracingSolidGlass.ush b/Engine/Shaders/Private/PathTracing/Material/PathTracingSolidGlass.ush
--- a/Engine/Shaders/Private/PathTracing/Material/PathTracingSolidGlass.ush
+++ b/Engine/Shaders/Private/PathTracing/Material/PathTracingSolidGlass.ush
@@ -48,14 +48,14 @@ FMaterialSample RoughGlass_SampleMaterial(
 		const float2 GGXResult = GGXEvalRefraction(L, V, H, Alpha, Eta);
 
 		// transform to world space
-		return CreateMaterialSample(normalize(mul(L, Basis)), GGXResult.x, F * GGXResult.y, -1.0);
+		return CreateMaterialSample(normalize(mul(L, Basis)), GGXResult.x, F * GGXResult.y, -1.0, Payload.Roughness);
 	}
 	else
 	{
 		const float2 GGXResult = GGXEvalReflection(L, V, H, Alpha);
 
 		// transform to world space
-		return CreateMaterialSample(normalize(mul(L, Basis)), GGXResult.x, F * GGXResult.y, 1.0);
+		return CreateMaterialSample(normalize(mul(L, Basis)), GGXResult.x, F * GGXResult.y, 1.0, Payload.Roughness);
 	}
 }
 
diff --git a/Engine/Shaders/Private/PathTracing/Material/PathTracingSpecularTransmission.ush b/Engine/Shaders/Private/PathTracing/Material/PathTracingSpecularTransmission.ush
--- a/Engine/Shaders/Private/PathTracing/Material/PathTracingSpecularTransmission.ush
+++ b/Engine/Shaders/Private/PathTracing/Material/PathTracingSpecularTransmission.ush
@@ -31,7 +31,7 @@ FMaterialSample SpecularRefraction_SampleMaterial(
 	}
 
 	// Throughput is 1 since we pick reflect/refract according to the fresnel term
-	return CreateMaterialSample(OutDirection, 1.0, POSITIVE_INFINITY, BiasSign);
+	return CreateMaterialSample(OutDirection, 1.0, POSITIVE_INFINITY, BiasSign, 0.0);
 } 
 
 FMaterialEval SpecularTransmission_EvalMaterial()
diff --git a/Engine/Shaders/Private/PathTracing/Material/PathTracingSubsurfaceProfile.ush b/Engine/Shaders/Private/PathTracing/Material/PathTracingSubsurfaceProfile.ush
--- a/Engine/Shaders/Private/PathTracing/Material/PathTracingSubsurfaceProfile.ush
+++ b/Engine/Shaders/Private/PathTracing/Material/PathTracingSubsurfaceProfile.ush
@@ -90,6 +90,7 @@ FMaterialSample SubsurfaceProfile_SampleMaterial(
 
 	// Randomly choose to sample diffuse or specular
 	float3 L = 0, H = 0;
+	float OutRoughness = 0;
 	if (RandSample.x < LobeCdf.x)
 	{
 		RandSample.x /= LobeCdf.x;
@@ -100,6 +101,7 @@ FMaterialSample SubsurfaceProfile_SampleMaterial(
 
 		L = Result.xyz;
 		H = normalize(L + V);
+		OutRoughness = 1.0;
 	}
 	else if (RandSample.x < LobeCdf.y)
 	{
@@ -116,6 +118,7 @@ FMaterialSample SubsurfaceProfile_SampleMaterial(
 			// invalid output direction, exit early
 			return NullMaterialSample();
 		}
+		OutRoughness = DualRoughnessData.x;
 	}
 	else
 	{
@@ -132,6 +135,7 @@ FMaterialSample SubsurfaceProfile_SampleMaterial(
 			// invalid output direction, exit early
 			return NullMaterialSample();
 		}
+		OutRoughness = DualRoughnessData.y;
 	}
 
 	// With a valid direction in hand -- now evaluate the BxDF (taking advantage of already computed terms)
@@ -153,5 +157,5 @@ FMaterialSample SubsurfaceProfile_SampleMaterial(
 	AddLobeWithMIS(OutWeight, OutPdf, (    DualRoughnessData.z) * F * GGXResult1.x, GGXResult1.y, LobeProb.z);
 
 	// transform to world space
-	return CreateMaterialSample(normalize(mul(L, Basis)), OutWeight, OutPdf, 1.0);
+	return CreateMaterialSample(normalize(mul(L, Basis)), OutWeight, OutPdf, 1.0, OutRoughness);
 }
diff --git a/Engine/Shaders/Private/PathTracing/Material/PathTracingThinGlass.ush b/Engine/Shaders/Private/PathTracing/Material/PathTracingThinGlass.ush
--- a/Engine/Shaders/Private/PathTracing/Material/PathTracingThinGlass.ush
+++ b/Engine/Shaders/Private/PathTracing/Material/PathTracingThinGlass.ush
@@ -149,7 +149,7 @@ FMaterialSample RoughThinGlass_SampleMaterial(
 		OutWeight *= Result.Reflected / ProbR;
 		OutPdf *= ProbR;
 		float3 OutDirection = normalize(mul(L, Basis));
-		return CreateMaterialSample(OutDirection, OutWeight, OutPdf, 1.0);
+		return CreateMaterialSample(OutDirection, OutWeight, OutPdf, 1.0, Payload.Roughness);
 	}
 	else
 	{
@@ -157,6 +157,6 @@ FMaterialSample RoughThinGlass_SampleMaterial(
 		OutWeight *= Result.Transmitted / (1.0 - ProbR);
 		OutPdf *= 1.0 - ProbR;
 		float3 OutDirection = normalize(mul(float3(L.xy, -L.z), Basis)); // flip reflection to other side
-		return CreateMaterialSample(OutDirection, OutWeight, OutPdf, -1.0);
+		return CreateMaterialSample(OutDirection, OutWeight, OutPdf, -1.0, Payload.Roughness);
 	}
 }
diff --git a/Engine/Shaders/Private/PathTracing/Material/PathTracingTwoSidedFoliage.ush b/Engine/Shaders/Private/PathTracing/Material/PathTracingTwoSidedFoliage.ush
--- a/Engine/Shaders/Private/PathTracing/Material/PathTracingTwoSidedFoliage.ush
+++ b/Engine/Shaders/Private/PathTracing/Material/PathTracingTwoSidedFoliage.ush
@@ -36,7 +36,7 @@ FMaterialSample TwoSidedFoliage_SampleMaterial(
 
 		float4 SampledValue = CosineSampleHemisphere(RandSample.yz);
 
-		return CreateMaterialSample(TangentToWorld(SampledValue.xyz, -N), SubsurfaceColor / TransmissionSelectionPdf, TransmissionSelectionPdf * SampledValue.w, -1.0);
+		return CreateMaterialSample(TangentToWorld(SampledValue.xyz, -N), SubsurfaceColor / TransmissionSelectionPdf, TransmissionSelectionPdf * SampledValue.w, -1.0, 1.0);
 	}
 	else
 	{
diff --git a/Engine/Shaders/Private/PathTracing/PathTracing.usf b/Engine/Shaders/Private/PathTracing/PathTracing.usf
--- a/Engine/Shaders/Private/PathTracing/PathTracing.usf
+++ b/Engine/Shaders/Private/PathTracing/PathTracing.usf
@@ -516,7 +516,11 @@ RAY_TRACING_ENTRY_RAYGEN(PathTracingMainRG)
 		Radiance = 0.5 + 0.3 * cos(2 * PI * ((StrongIntegerHash(RandPixel.x + StrongIntegerHash(RandPixel.y)) & 0xFFFFFF) * 5.96046447754e-08 + float3(0.0, 0.333, 0.6666)));
 		break;
 #endif
-		FMaterialClosestHitPayload Payload = TraceTransparentRay(Ray, RayCone, Bounce == 0, Bounce == PathTracingData.MaxBounces, PathTracingData.SkipDirectLighting == 0 || Bounce > 1, LaunchIndex, NumVisibleLights, RandSequence, PathThroughput, Radiance);
+		const bool bIsCameraRay = Bounce == 0;
+		const bool bIsLastBounce = Bounce == PathTracingData.MaxBounces;
+		const bool bIncludeEmissive = (PathTracingData.SkipDirectLighting == 0 || Bounce > 1) &&
+									  (PathTracingData.SkipEmissive       == 0 || bIsCameraRay);
+		FMaterialClosestHitPayload Payload = TraceTransparentRay(Ray, RayCone, bIsCameraRay, bIsLastBounce, bIncludeEmissive, LaunchIndex, NumVisibleLights, RandSequence, PathThroughput, Radiance);
 
 		if (Payload.IsMiss())
 		{
@@ -529,11 +533,11 @@ RAY_TRACING_ENTRY_RAYGEN(PathTracingMainRG)
 
 		if (PathTracingData.ApproximateCaustics)
 		{
-			// modify the path roughness as well as the payload roughness to minimize difficult caustics
+			// modify the payload roughness to minimize difficult caustics
 			// This is inspired by a trick used in the Arnold renderer:
 			//   https://cgg.mff.cuni.cz/~jaroslav/gicourse2010/giai2010-02-marcos_fajardo-slides.pdf (slide 39)
 			//   https://www.arnoldrenderer.com/research/Arnold_TOG2018.pdf (section 4.2)
-			AdjustPathAndMaterialRoughness(Payload, PathRoughness);
+			AdjustMaterialRoughness(Payload, PathRoughness);
 		}
 
 		// Processing the random walk will (stochastically) move the shading point elsewhere on the surface of the object
@@ -592,7 +596,9 @@ RAY_TRACING_ENTRY_RAYGEN(PathTracingMainRG)
 						LightRay.TMax = LightSample.Distance;
 						ApplyPositionBias(LightRay, SignedPositionBias * Payload.WorldNormal, PathTracingData.MaxNormalBias);
 
-						LightSample.RadianceOverPdf *= TraceTransparentVisibilityRay(LightRay, LaunchIndex, PathRoughness);
+						float AvgRoughness = PathTracingData.ApproximateCaustics ? GetAverageRoughness(Payload) : 0.0;
+
+						LightSample.RadianceOverPdf *= TraceTransparentVisibilityRay(LightRay, LaunchIndex, AvgRoughness);
 					}
 
 					// #dxr_todo: Is it cheaper to fire the ray first? Or eval the material first?
@@ -669,13 +675,18 @@ RAY_TRACING_ENTRY_RAYGEN(PathTracingMainRG)
 		Ray.TMax = POSITIVE_INFINITY;
 		ApplyPositionBias(Ray, MaterialSample.PositionBiasSign * Payload.WorldNormal, PathTracingData.MaxNormalBias);
 
+		if (PathTracingData.ApproximateCaustics)
+		{
+			// enlarge roughness based on the chosen lobe roughness
+			PathRoughness = max(PathRoughness, MaterialSample.Roughness);
+		}
+
 		// If we are using Material sampling for lights
 		if (MISMode != 1 && bDoLightLoop)
 		{
 			// Check which lights can be seen by the material ray and trace a dedicated shadow ray
 			// While it would be possible to just loop around and use the indirect ray to do this, it would prevent the application
 			// of shadow ray specific logic for transparent shadows or various per light tricks like shadow casting
-			float PreviousCdfValue = 0.0;
 			const bool bUseMIS = MISMode == 2 && LightPickingCdfSum > 0;
 			for (uint LightId = 0; LightId < SceneLightCount; ++LightId)
 			{
@@ -695,8 +706,8 @@ RAY_TRACING_ENTRY_RAYGEN(PathTracingMainRG)
 				float3 LightContrib = PathThroughput * LightResult.Radiance;
 				if (bUseMIS)
 				{
+					float PreviousCdfValue = LightId > 0 ? LightPickingCdf[LightId - 1] : 0.0;
 					float LightPickPdf =  (LightPickingCdf[LightId] - PreviousCdfValue) / LightPickingCdfSum;
-					PreviousCdfValue = LightPickingCdf[LightId];
 
 					LightContrib *= MISWeightRobust(MaterialSample.Pdf, LightResult.Pdf * LightPickPdf);
 				}
diff --git a/Engine/Source/Runtime/Renderer/Private/PathTracing.cpp b/Engine/Source/Runtime/Renderer/Private/PathTracing.cpp
--- a/Engine/Source/Runtime/Renderer/Private/PathTracing.cpp
+++ b/Engine/Source/Runtime/Renderer/Private/PathTracing.cpp
@@ -97,6 +97,13 @@ TAutoConsoleVariable<int32> CVarPathTracingApproximateCaustics(
 	ECVF_RenderThreadSafe
 );
 
+TAutoConsoleVariable<int32> CVarPathTracingSkipEmissive(
+	TEXT("r.PathTracing.SkipEmissive"),
+	0,
+	TEXT("When non-zero, the path tracer will skip emissive results after the first bounce. (default = 0 (off))"),
+	ECVF_RenderThreadSafe
+);
+
 TAutoConsoleVariable<int32> CVarPathTracingEnableCameraBackfaceCulling(
 	TEXT("r.PathTracing.EnableCameraBackfaceCulling"),
 	1,
@@ -166,10 +173,10 @@ static bool PrepareShaderArgs(const FViewInfo& View, FPathTracingData& PathTraci
 	}
 	else
 	{
+		PathTracingData.SkipDirectLighting = true;
 		if (View.Family->EngineShowFlags.GlobalIllumination)
 		{
 			// skip direct lighting, but still do the full bounces
-			PathTracingData.SkipDirectLighting = true;
 		}
 		else
 		{
@@ -188,6 +195,7 @@ static bool PrepareShaderArgs(const FViewInfo& View, FPathTracingData& PathTraci
 	PathTracingData.ApproximateCaustics = CVarPathTracingApproximateCaustics.GetValueOnRenderThread();
 	PathTracingData.EnableCameraBackfaceCulling = CVarPathTracingEnableCameraBackfaceCulling.GetValueOnRenderThread();
 	PathTracingData.CoherentSampling = CVarPathTracingCoherentSampling.GetValueOnRenderThread();
+	PathTracingData.SkipEmissive = CVarPathTracingSkipEmissive.GetValueOnRenderThread();
 	float FilterWidth = CVarPathTracingFilterWidth.GetValueOnRenderThread();
 	if (FilterWidth < 0)
 	{
@@ -279,6 +287,14 @@ static bool PrepareShaderArgs(const FViewInfo& View, FPathTracingData& PathTraci
 		PreviousSkipDirectLighting = PathTracingData.SkipDirectLighting;
 	}
 
+	// Changing skip emissive requires starting over
+	static uint32 PreviousSkipEmissive = PathTracingData.SkipEmissive;
+	if (PreviousSkipEmissive != PathTracingData.SkipEmissive)
+	{
+		NeedInvalidation = true;
+		PreviousSkipEmissive = PathTracingData.SkipEmissive;
+	}
+
 	// Changing coherent sampling requires starting over
 	static uint32 PreviousCoherentSampling = PathTracingData.CoherentSampling;
 	if (PreviousCoherentSampling != PathTracingData.CoherentSampling)
@@ -912,6 +928,11 @@ void FDeferredShadingSceneRenderer::RenderPathTracing(
 		PassParameters->PathTracingData = CreateUniformBufferImmediate(PathTracingData, EUniformBufferUsage::UniformBuffer_SingleFrame);
 		// upload sky/lights data
 		SetLightParameters(GraphBuilder, PassParameters, Scene, View, UseMISCompensation);
+		if (PathTracingData.SkipDirectLighting)
+		{
+			PassParameters->SceneVisibleLightCount = 0;
+		}
+
 		PassParameters->IESTextureSampler = TStaticSamplerState<SF_Bilinear, AM_Clamp, AM_Clamp, AM_Clamp>::GetRHI();
 		PassParameters->RadianceTexture = GraphBuilder.CreateUAV(RadianceTexture);
 
diff --git a/Engine/Source/Runtime/Renderer/Private/PathTracingUniformBuffers.h b/Engine/Source/Runtime/Renderer/Private/PathTracingUniformBuffers.h
--- a/Engine/Source/Runtime/Renderer/Private/PathTracingUniformBuffers.h
+++ b/Engine/Source/Runtime/Renderer/Private/PathTracingUniformBuffers.h
@@ -20,6 +20,7 @@ BEGIN_GLOBAL_SHADER_PARAMETER_STRUCT(FPathTracingData, )
 	SHADER_PARAMETER(uint32, ApproximateCaustics)
 	SHADER_PARAMETER(uint32, EnableCameraBackfaceCulling)
 	SHADER_PARAMETER(uint32, SkipDirectLighting)
+	SHADER_PARAMETER(uint32, SkipEmissive)
 	SHADER_PARAMETER(uint32, CoherentSampling)
 	SHADER_PARAMETER(float, MaxPathIntensity)
 	SHADER_PARAMETER(float, MaxNormalBias)
