Combine scene depth and visbuffer depth for HZB to fix !Nanite killing occlusion culling.

diff --git a/Engine/Shaders/Private/HZB.usf b/Engine/Shaders/Private/HZB.usf
--- a/Engine/Shaders/Private/HZB.usf
+++ b/Engine/Shaders/Private/HZB.usf
@@ -163,9 +163,7 @@ void HZBBuildCS(
 
 	uint2 SrcPixelCoord = (DispatchThreadId << 1) + PixelViewPortMinMax.xy;
 #if VIS_BUFFER_FORMAT == 1
-	// TODO: Combine with SceneDepth if complete prepass.
-	//DeviceZ = max(DeviceZ, Gather4VisZ(BufferUV, SrcPixelCoord));
-	DeviceZ = Gather4VisZ(BufferUV, SrcPixelCoord);
+	DeviceZ = max(DeviceZ, Gather4VisZ(BufferUV, SrcPixelCoord));
 #elif VIS_BUFFER_FORMAT == 2
 	DeviceZ = Gather4VisZ(BufferUV, SrcPixelCoord);
 #endif
