Fix shader warnings.

diff --git a/Engine/Shaders/Private/HairStrands/HairStrandsInterpolation.usf b/Engine/Shaders/Private/HairStrands/HairStrandsInterpolation.usf
--- a/Engine/Shaders/Private/HairStrands/HairStrandsInterpolation.usf
+++ b/Engine/Shaders/Private/HairStrands/HairStrandsInterpolation.usf
@@ -639,7 +639,7 @@ void ClusterAABBEvaluationCS(uint3 DispatchThreadId : SV_DispatchThreadID, uint3
 		if (ThreadReadOffset < IndexCount)
 		{
 			const float3 OutHairPositionOffset = RenderDeformedOffsetBuffer[0].xyz;
-			const uint VertexIndex = clamp(ClusterVertexIdBuffer[IndexOffset + ThreadReadOffset], 0, VertexCountAfterCulling-1);
+			const uint VertexIndex = clamp(ClusterVertexIdBuffer[IndexOffset + ThreadReadOffset], 0u, VertexCountAfterCulling-1);
 			const FHairControlPoint CP = UnpackHairControlPoint(RenderDeformedPositionBuffer[VertexIndex], OutHairPositionOffset, 1, 1);
 
 			if (all(IsFinite(CP.Position)))
diff --git a/Engine/Shaders/Private/HairStrands/HairStrandsTangent.usf b/Engine/Shaders/Private/HairStrands/HairStrandsTangent.usf
--- a/Engine/Shaders/Private/HairStrands/HairStrandsTangent.usf
+++ b/Engine/Shaders/Private/HairStrands/HairStrandsTangent.usf
@@ -44,7 +44,7 @@ void MainCS(uint2 DispatchThreadId : SV_DispatchThreadID)
 	if (HairStrandsVF_bIsCullingEnable)
 	{
 		const int FetchIndex = DispatchThreadId.x + DispatchThreadId.y * GROUP_SIZE * GROUP_COUNT_X;
-		const int VertexCountAfterCulling = clamp(HairStrandsVF_CullingIndirectBuffer[3], 0, MaxVertexCount);
+		const int VertexCountAfterCulling = clamp((int)HairStrandsVF_CullingIndirectBuffer[3], 0, MaxVertexCount);
 		bValid = false;
 		if (FetchIndex < VertexCountAfterCulling)
 		{
