Improve GGX math in path tracer

Re-use the D function from BRDF.ush as it is both simpler and more robust. Optimize the calculation of shadowing and masking terms in the weight and pdf for both the reflected and refracted cases.

diff --git a/Engine/Shaders/Private/MonteCarlo.ush b/Engine/Shaders/Private/MonteCarlo.ush
--- a/Engine/Shaders/Private/MonteCarlo.ush
+++ b/Engine/Shaders/Private/MonteCarlo.ush
@@ -379,9 +379,7 @@ float VisibleGGXPDF(float3 V, float3 H, float a2)
 	float d = (NoH * a2 - NoH) * NoH + 1;
 	float D = a2 / (PI*d*d);
 
-	float G_SmithV = 2 * NoV / (NoV + sqrt(NoV * (NoV - NoV * a2) + a2));
-
-	float PDF = G_SmithV * VoH * D / NoV;
+	float PDF = 2 * VoH * D / (NoV + sqrt(NoV * (NoV - NoV * a2) + a2));
 	return PDF;
 }
 
@@ -390,12 +388,12 @@ float VisibleGGXPDF_aniso(float3 V, float3 H, float2 Alpha)
 	float NoV = V.z;
 	float NoH = H.z;
 	float VoH = dot(V, H);
-	float CosThetaM4 = Pow2(NoH);
-	float TanThetaM2 = (Pow2(H.x / Alpha.x) + Pow2(H.y / Alpha.y)) / Pow2(NoH);
-	float D = 1 / (PI * Pow2(1 + TanThetaM2) * Alpha.x * Alpha.y * CosThetaM4);
-	float Lambda_i = 0.5 * (-1.0 + sqrt(1.0 + (Pow2(V.x * Alpha.x) + Pow2(V.y * Alpha.y)) / Pow2(V.z)));
-	float G1 = 1 / (Lambda_i + 1);
-	float Pdf = (G1 * D * VoH) / NoV;
+	float a2 = Alpha.x * Alpha.y;
+	float3 Hs = float3(Alpha.y * H.x, Alpha.x * H.y, a2 * NoH);
+	float S = dot(Hs, Hs);
+	float D = (1.0f / PI) * a2 * Square(a2 / S);
+	float LenV = length(float3(V.x * Alpha.x, V.y * Alpha.y, NoV));
+	float Pdf = (2 * D * VoH) / (NoV + LenV);
 	return Pdf;
 }
 
@@ -411,15 +409,11 @@ float4 ImportanceSampleVisibleGGX( float2 DiskE, float a2, float3 V )
 	// stretch
 	float3 Vh = normalize( float3( a * V.xy, V.z ) );
 
-	// Orthonormal basis
-	// Tangent0 is orthogonal to N.
-	#if 1 // Stable tangent basis based on V.
-		float3 Tangent0 = (V.z < 0.9999) ? normalize( cross( float3(0, 0, 1), V ) ) : float3(1, 0, 0);
-		float3 Tangent1 = normalize(cross( Vh, Tangent0 ));
-	#else
-		float3 Tangent0 = (Vh.z < 0.9999) ? normalize( cross( float3(0, 0, 1), Vh ) ) : float3(1, 0, 0);
-		float3 Tangent1 = cross( Vh, Tangent0 );
-	#endif
+	// Stable tangent basis based on V
+	// Tangent0 is orthogonal to N
+	float LenSq = Vh.x * Vh.x + Vh.y * Vh.y;
+	float3 Tangent0 = LenSq > 0 ? float3(-Vh.y, Vh.x, 0) * rsqrt(LenSq) : float3(1, 0, 0);
+	float3 Tangent1 = cross(Vh, Tangent0);
 
 	float2 p = DiskE;
 	float s = 0.5 + 0.5 * Vh.z;
@@ -446,8 +440,9 @@ float4 ImportanceSampleVisibleGGX_aniso(float2 DiskE, float2 Alpha, float3 V)
 	float3 Vh = normalize(float3(Alpha * V.xy, V.z));
 
 	// Stable tangent basis based on V
-	float3 Tx = (V.z < 0.9999) ? normalize(float3(-V.y, V.x, 0)) : float3(1, 0, 0);
-	float3 Ty = normalize(cross(Vh, Tx));
+	float LenSq = Vh.x * Vh.x + Vh.y * Vh.y;
+	float3 Tx = LenSq > 0 ? float3(-Vh.y, Vh.x, 0) * rsqrt(LenSq) : float3(1, 0, 0);
+	float3 Ty = cross(Vh, Tx);
 
 	float2 p = DiskE;
 	float s = 0.5 + 0.5 * Vh.z;
diff --git a/Engine/Shaders/Private/PathTracing/Material/PathTracingGlossy.ush b/Engine/Shaders/Private/PathTracing/Material/PathTracingGlossy.ush
--- a/Engine/Shaders/Private/PathTracing/Material/PathTracingGlossy.ush
+++ b/Engine/Shaders/Private/PathTracing/Material/PathTracingGlossy.ush
@@ -39,21 +39,6 @@ float3x3 GetGGXBasis(float Roughness, float3 WorldNormal, out float2 Alpha)
 	return GetTangentBasis(WorldNormal);
 }
 
-float3 GGXHelper(float3 L, float3 V, float3 H, float2 Alpha)
-{
-	// TODO: Can this be further simplified / unified with other GGX routines?
-	const float CosThetaM4 = Pow2(H.z);
-	const float TanThetaM2 = (Pow2(H.x / Alpha.x) + Pow2(H.y / Alpha.y)) / Pow2(H.z);
-	const float D = 1 / (PI * Pow2(1 + TanThetaM2) * Alpha.x * Alpha.y * CosThetaM4);
-	const float Lambda_o = 0.5 * (-1.0 + sqrt(1.0 + (Pow2(L.x * Alpha.x) + Pow2(L.y * Alpha.y)) / Pow2(L.z)));
-	const float Lambda_i = 0.5 * (-1.0 + sqrt(1.0 + (Pow2(V.x * Alpha.x) + Pow2(V.y * Alpha.y)) / Pow2(V.z)));
-	const float G2 = 1 / (Lambda_i + Lambda_o + 1);
-	const float G1 = 1 / (Lambda_i + 1);
-
-	return float3(D, G2, G1);
-}
-
-
 // Compute GGX lobe Weight and Pdf (without Fresnel term) given a set of vectors in local space (Z up)
 float2 GGXEvalReflection(float3 L, float3 V, float3 H, float2 Alpha)
 {
@@ -62,13 +47,13 @@ float2 GGXEvalReflection(float3 L, float3 V, float3 H, float2 Alpha)
 
 	if (NoL > 0 && NoV > 0)
 	{
-		const float3 Result = GGXHelper(L, V, H, Alpha);
-		const float D = Result.x;
-		const float G2 = Result.y;
-		const float G1 = Result.z;
-
-		const float Weight = G2 / G1;
-		const float Pdf = (G1 * D * 0.25f) / NoV;
+		const float D = D_GGXaniso(Alpha.x, Alpha.y, H.z, H.x, H.y);
+		// See implementation in Vis_SmithJointAniso for G2/(4*NoV*NoL)
+		// We can simplify a bit further since we need both the weight G2/G1 and the pdf
+		const float LenL = length(float3(L.x * Alpha.x, L.y * Alpha.y, NoL));
+		const float LenV = length(float3(V.x * Alpha.x, V.y * Alpha.y, NoV));
+		const float Weight = NoL * (LenV + NoV) / (NoV * LenL + NoL * LenV);
+		const float Pdf = 0.5 * D * rcp(LenV + NoV);
 
 		return float2(Weight, Pdf);
 	}
@@ -85,17 +70,17 @@ float2 GGXEvalRefraction(float3 L, float3 V, float3 H, float2 Alpha, float Eta)
 	// This checks that V and L point in opposite sides 
 	if (NoL > 0 && NoV > 0)
 	{
-		const float3 Result = GGXHelper(L, V, H, Alpha);
-		const float D = Result.x;
-		const float G2 = Result.y;
-		const float G1 = Result.z;
+		const float D = D_GGXaniso(Alpha.x, Alpha.y, H.z, H.x, H.y);
 
 		const float VoH = dot(V, H);
 		const float LoH = dot(L, H);
 		const float3 Ht = (Eta * L + V);
 
-		const float Weight = G2 / G1;
-		const float Pdf = (abs(LoH * VoH) * (Eta * Eta) * (G1 * D)) / (dot(Ht, Ht) * NoV);
+		// Same as logic above, only the pdf math differs slightly
+		const float LenL = length(float3(L.x * Alpha.x, L.y * Alpha.y, NoL));
+		const float LenV = length(float3(V.x * Alpha.x, V.y * Alpha.y, NoV));
+		const float Weight = NoL * (LenV + NoV) / (NoV * LenL + NoL * LenV);
+		const float Pdf = (abs(LoH * VoH) * (Eta * Eta) * (2.0 * D)) / (dot(Ht, Ht) * (LenV + NoV));
 
 		return float2(Weight, Pdf);
 	}
