Fixed: Overexposed lighting on iOS in QA-Waterworks

diff --git a/Engine/Shaders/Private/SingleLayerWaterShading.ush b/Engine/Shaders/Private/SingleLayerWaterShading.ush
--- a/Engine/Shaders/Private/SingleLayerWaterShading.ush
+++ b/Engine/Shaders/Private/SingleLayerWaterShading.ush
@@ -58,7 +58,10 @@ WaterVolumeLightingOutput EvaluateWaterVolumeLighting(
 		float PixelSceneDeviceZ = LookupDeviceZ(ViewportUV);
 		float PixelSceneDepth = ConvertFromDeviceZ(PixelSceneDeviceZ);
 #endif
-
+		#if HAS_INVERTED_Z_BUFFER
+		// SvPositionToWorld will result in NaN if PixelSceneDeviceZ=0
+		PixelSceneDeviceZ = max(0.000000000001, PixelSceneDeviceZ);
+		#endif
 		float SceneDeviceZ = PixelSceneDeviceZ;
 		float SceneDepth = PixelSceneDepth;
 
