Stops decimating the interference history

diff --git a/Engine/Shaders/Private/TemporalAA/TAADecimateHistory.usf b/Engine/Shaders/Private/TemporalAA/TAADecimateHistory.usf
--- a/Engine/Shaders/Private/TemporalAA/TAADecimateHistory.usf
+++ b/Engine/Shaders/Private/TemporalAA/TAADecimateHistory.usf
@@ -11,6 +11,8 @@
 
 #define CONFIG_OUTPUT_HALF_RES (DIM_OUTPUT_HALF_RES)
 
+#define CONFIG_OUTPUT_INTERFERENCE_SEEDS 0
+
 
 //------------------------------------------------------- PARAMETERS
 
@@ -170,6 +172,7 @@ void MainCS(
 		
 		float2 PrevInputBufferUV = (InputInfo_ScreenPosToViewportScale * PrevOutputScreenPos + InputInfo_ScreenPosToViewportBias) * InputInfo_ExtentInverse;
 
+		#if CONFIG_OUTPUT_INTERFERENCE_SEEDS
 		{
 			const uint KernelSize = 4;
 
@@ -239,11 +242,12 @@ void MainCS(
 			// average of min and max that turns out interference free.
 			IntegratedInterferenceSeed = (MinInterferenceSeed + MaxInterferenceSeed) * 0.5;
 		}
+		#endif
 
 		taa_half PrevUninterferenceSeed = PrevHistory_LowResTextures_0.SampleLevel(GlobalBilinearClampedSampler, PrevInputBufferUV, 0);
 		taa_half PrevUninterferenceLuma = PrevUninterferenceSeed * HdrWeightInvY(PrevUninterferenceSeed);
 
-		#if 0
+		#if !CONFIG_OUTPUT_INTERFERENCE_SEEDS
 		{
 			DecimatedColor = PrevHistory_Textures_0.SampleLevel(GlobalBilinearClampedSampler, KernelHistoryBufferUV, 0);
 			
@@ -402,7 +406,12 @@ void MainCS(
 		#endif
 		
 		ParallaxRejectionMaskOutput[InputPixelPos] = ParallaxRejectionMask;
-		InterferenceSeedOutput[InputPixelPos] = taa_half4(ClosestInterferenceSeed, IntegratedInterferenceSeed, MinInterferenceSeed, MaxInterferenceSeed);
+
+		#if CONFIG_OUTPUT_INTERFERENCE_SEEDS
+		{
+			InterferenceSeedOutput[InputPixelPos] = taa_half4(ClosestInterferenceSeed, IntegratedInterferenceSeed, MinInterferenceSeed, MaxInterferenceSeed);
+		}
+		#endif
 
 		#if DEBUG_OUTPUT
 		{
