Removed duplicated sun edge smoothing effect when rendering a skydome.
Updated related screenshot tests.

diff --git a/Engine/Shaders/Private/SkyAtmosphere.usf b/Engine/Shaders/Private/SkyAtmosphere.usf
--- a/Engine/Shaders/Private/SkyAtmosphere.usf
+++ b/Engine/Shaders/Private/SkyAtmosphere.usf
@@ -404,6 +404,7 @@ float3 GetLightDiskLuminance(float3 WorldPos, float3 WorldDir, uint LightIndex)
 	if (t < 0.0f												// No intersection with the planet
 		&& View.RenderingReflectionCaptureMask==0.0f)	// Do not render light disk when in reflection capture in order to avoid double specular. The sun contribution is already computed analyticaly.
 	{
+		// GetLightDiskLuminance contains a tiny soft edge effect
 		float3 LightDiskLuminance = GetLightDiskLuminance(
 			WorldPos, WorldDir, Atmosphere.BottomRadiusKm, Atmosphere.TopRadiusKm,
 			TransmittanceLutTexture, TransmittanceLutTextureSampler,
@@ -414,16 +415,6 @@ float3 GetLightDiskLuminance(float3 WorldPos, float3 WorldDir, uint LightIndex)
 		float3 ExposedLightLuminance = LightDiskLuminance * OutputPreExposure;
 		ExposedLightLuminance = min(ExposedLightLuminance, MaxLightLuminance);
 
-#if 1
-		const float ViewDotLight = dot(WorldDir, View.AtmosphereLightDirection[LightIndex].xyz);
-		const float CosHalfApex = View.AtmosphereLightDiscCosHalfApexAngle[LightIndex].x;
-		const float HalfCosHalfApex = CosHalfApex + (1.0f - CosHalfApex) * 0.25; // Start fading when at 75% distance from light disk center (in cosine space)
-
-		// Apply smooth fading at edge. This is currently an eye balled fade out that works well in many cases.
-		const float Weight = 1.0-saturate((HalfCosHalfApex - ViewDotLight) / (HalfCosHalfApex - CosHalfApex));
-		ExposedLightLuminance = ExposedLightLuminance * Weight;
-#endif 
-
 		return ExposedLightLuminance;
 	}
 	return 0.0f;
