Always use anisotropic formulat to compute VT mip level, even if VT anisotropic filtering is disabled
- If filtering is disabled, we'll still avoid using SampleGrad (which should still avoid most of the shader overhead)
- This way the computed VT mip level will match the non-VT mip level much more closely
Also tweak the way noise is mixed in to VT mip level, better match HW formular (also slightly cheaper)

diff --git a/Engine/Shaders/Private/VirtualTextureCommon.ush b/Engine/Shaders/Private/VirtualTextureCommon.ush
--- a/Engine/Shaders/Private/VirtualTextureCommon.ush
+++ b/Engine/Shaders/Private/VirtualTextureCommon.ush
@@ -204,14 +204,12 @@ int TextureComputeVirtualMipLevel(
 	OutResult.dUVdx = dUVdx * PageTableUniform.SizeInPages;
 	OutResult.dUVdy = dUVdy * PageTableUniform.SizeInPages;
 
-#if VIRTUAL_TEXTURE_ANISOTROPIC_FILTERING
+	// Always compute mip level using MipLevelAniso2D, even if VIRTUAL_TEXTURE_ANISOTROPIC_FILTERING is disabled
+	// This way the VT mip level selection will come much closer to HW mip selection, even if we're not sampling the texture using anisotropic filtering
 	const float ComputedLevel = MipLevelAniso2D(OutResult.dUVdx, OutResult.dUVdy, PageTableUniform.MaxAnisoLog2);
-#else
-	const float ComputedLevel = MipLevel2D(OutResult.dUVdx, OutResult.dUVdy);
-#endif
 
 	const float Noise = InterleavedGradientNoise(SvPositionXY, View.StateFrameIndexMod8);
-	return (int)floor(ComputedLevel + MipBias + Noise * 0.5f - 0.25f) + int(PageTableUniform.vPageTableMipBias);
+	return (int)floor(ComputedLevel + MipBias + Noise) + int(PageTableUniform.vPageTableMipBias);
 }
 
 /** Samples page table indirection and any apply changes to UV, vLevel and PageTable information for adaptive page tables. */
diff --git a/Engine/Source/Runtime/Engine/Private/Texture2D.cpp b/Engine/Source/Runtime/Engine/Private/Texture2D.cpp
--- a/Engine/Source/Runtime/Engine/Private/Texture2D.cpp
+++ b/Engine/Source/Runtime/Engine/Private/Texture2D.cpp
@@ -1228,7 +1228,12 @@ void FVirtualTexture2DResource::InitRHI()
 
 	check(TextureOwner);
 
-	const uint32 MaxAnisotropy = FMath::Min<int32>(VirtualTextureScalability::GetMaxAnisotropy(), VTData->TileBorderSize);
+	uint32 MaxAnisotropy = 0;
+	if (VirtualTextureScalability::IsAnisotropicFilteringEnabled())
+	{
+		// Limit HW MaxAnisotropy to avoid sampling outside VT borders
+		MaxAnisotropy = FMath::Min<int32>(VirtualTextureScalability::GetMaxAnisotropy(), VTData->TileBorderSize);
+	}
 
 	// We always create a sampler state if we're attached to a texture. This is used to sample the cache texture during actual rendering and the miptails editor resource.
 	// If we're not attached to a texture it likely means we're light maps which have sampling handled differently.
diff --git a/Engine/Source/Runtime/Engine/Private/VT/VirtualTextureScalability.cpp b/Engine/Source/Runtime/Engine/Private/VT/VirtualTextureScalability.cpp
--- a/Engine/Source/Runtime/Engine/Private/VT/VirtualTextureScalability.cpp
+++ b/Engine/Source/Runtime/Engine/Private/VT/VirtualTextureScalability.cpp
@@ -225,13 +225,13 @@ namespace VirtualTextureScalability
 #endif
 	}
 
-	int32 GetMaxAnisotropy()
+	bool IsAnisotropicFilteringEnabled()
 	{
-		if (CVarVTEnableAnisotropy.GetValueOnAnyThread() == 0)
-		{
-			return 0;
-		}
+		return CVarVTEnableAnisotropy.GetValueOnAnyThread() != 0;
+	}
 
+	int32 GetMaxAnisotropy()
+	{
 		return CVarVTMaxAnisotropy.GetValueOnAnyThread();
 	}
 
diff --git a/Engine/Source/Runtime/Engine/Public/VT/VirtualTextureScalability.h b/Engine/Source/Runtime/Engine/Public/VT/VirtualTextureScalability.h
--- a/Engine/Source/Runtime/Engine/Public/VT/VirtualTextureScalability.h
+++ b/Engine/Source/Runtime/Engine/Public/VT/VirtualTextureScalability.h
@@ -16,6 +16,8 @@ namespace VirtualTextureScalability
 	ENGINE_API float GetPoolSizeScale();
 	/** Get resolution bias for runtime virtual textures. */
 	ENGINE_API int32 GetRuntimeVirtualTextureSizeBias();
+	/** Is HW Anisotropic filtering enabled for VT */
+	ENGINE_API bool IsAnisotropicFilteringEnabled();
 
 	/**
 	 * Get maximum anisotropy when virtual texture sampling. 
diff --git a/Engine/Source/Runtime/Renderer/Private/VT/AllocatedVirtualTexture.cpp b/Engine/Source/Runtime/Renderer/Private/VT/AllocatedVirtualTexture.cpp
--- a/Engine/Source/Runtime/Renderer/Private/VT/AllocatedVirtualTexture.cpp
+++ b/Engine/Source/Runtime/Renderer/Private/VT/AllocatedVirtualTexture.cpp
@@ -371,7 +371,9 @@ void FAllocatedVirtualTexture::GetPackedPageTableUniform(FUintVector4* OutUnifor
 	const uint32 HeightInPages = GetHeightInTiles();
 	const uint32 vPageTableMipBias = FMath::FloorLog2(vPageSize);
 
-	const uint32 MaxAnisotropy = FMath::Min<int32>(VirtualTextureScalability::GetMaxAnisotropy(), PageBorderSize);
+	// Here MaxAnisotropy only controls the VT mip level selection
+	// We don't need to limit this value based on border size, and we can add this factor in even if HW anisotropic filtering is disabled
+	const uint32 MaxAnisotropy = VirtualTextureScalability::GetMaxAnisotropy();
 	const uint32 MaxAnisotropyLog2 = (MaxAnisotropy > 0u) ? FMath::FloorLog2(MaxAnisotropy) : 0u;
 
 	const uint32 AdaptiveLevelBias = 0; // todo[vt]: Required for handling SampleLevel correctly on adaptive page tables.
