Fixes for VT texture viewer
- Account for VT requests now encoding vLevel+1
- Fix sampling in shader when explicit mip level is not given

diff --git a/Engine/Shaders/Private/SimpleElementTexture2DPreviewPixelShader.usf b/Engine/Shaders/Private/SimpleElementTexture2DPreviewPixelShader.usf
--- a/Engine/Shaders/Private/SimpleElementTexture2DPreviewPixelShader.usf
+++ b/Engine/Shaders/Private/SimpleElementTexture2DPreviewPixelShader.usf
@@ -96,7 +96,7 @@ void Main(
 	else
 	{
 #if SAMPLE_VIRTUAL_TEXTURE
-		VTPageTableResult PageTableResult = TextureLoadVirtualPageTableLevel(InPageTableTexture0, InPageTableTexture1, VTPageTableUniform_Unpack(VTPackedPageTableUniform[0], VTPackedPageTableUniform[1]), TextureCoordinate, VTADDRESSMODE_WRAP, VTADDRESSMODE_WRAP, 0.0f);
+		VTPageTableResult PageTableResult = TextureLoadVirtualPageTable(InPageTableTexture0, InPageTableTexture1, VTPageTableUniform_Unpack(VTPackedPageTableUniform[0], VTPackedPageTableUniform[1]), TextureCoordinate, VTADDRESSMODE_WRAP, VTADDRESSMODE_WRAP, 0.0f, InPosition.xy);
 		Sample = TextureVirtualSample( InTexture, InTextureSampler, PageTableResult, LayerIndex, VTUniform_Unpack(VTPackedUniform));
 #else
 		Sample = Texture2DSample(InTexture, InTextureSampler,TextureCoordinate);
diff --git a/Engine/Source/Runtime/Renderer/Private/VT/VirtualTextureSystem.cpp b/Engine/Source/Runtime/Renderer/Private/VT/VirtualTextureSystem.cpp
--- a/Engine/Source/Runtime/Renderer/Private/VT/VirtualTextureSystem.cpp
+++ b/Engine/Source/Runtime/Renderer/Private/VT/VirtualTextureSystem.cpp
@@ -115,10 +115,12 @@ static TAutoConsoleVariable<int32> CVarVTProduceLockedTilesOnFlush(
 
 static FORCEINLINE uint32 EncodePage(uint32 ID, uint32 vLevel, uint32 vTileX, uint32 vTileY)
 {
+	const uint32 vLevelPlus1 = vLevel + 1u;
+
 	uint32 Page;
 	Page = vTileX << 0;
 	Page |= vTileY << 12;
-	Page |= vLevel << 24;
+	Page |= vLevelPlus1 << 24;
 	Page |= ID << 28;
 	return Page;
 }
