Fixed missing SceneTextures for a mobile deferred shading

diff --git a/Engine/Shaders/Private/MobileDeferredShading.usf b/Engine/Shaders/Private/MobileDeferredShading.usf
--- a/Engine/Shaders/Private/MobileDeferredShading.usf
+++ b/Engine/Shaders/Private/MobileDeferredShading.usf
@@ -10,6 +10,9 @@
 #include "LightFunctionCommon.ush"
 #include "PlanarReflectionShared.ush"
 
+// Reroute MobileSceneTextures uniform buffer references to the base pass uniform buffer
+#define MobileSceneTextures MobileBasePass.SceneTextures
+
 #if APPLY_SKY_REFLECTION
 #define ENABLE_SKY_LIGHT 1
 #endif
