Strata comment and naming forgotten from a previous CL review.

diff --git a/Engine/Shaders/Private/ParticipatingMediaCommon.ush b/Engine/Shaders/Private/ParticipatingMediaCommon.ush
--- a/Engine/Shaders/Private/ParticipatingMediaCommon.ush
+++ b/Engine/Shaders/Private/ParticipatingMediaCommon.ush
@@ -161,10 +161,13 @@ float3 TransmittanceToMeanFreePath(in float3 TransmittanceColor, in float Thickn
 float3 IsotropicMediumSlabDirectionalAlbedoFade(float3 BaseColor, float3 MFP)
 {
 	float3 Fade;
-	// Progressively fade out to 0 for base color starting from 0.1 down to 0.
-	Fade = saturate(BaseColor * 10.0f);
-	// And also fade out from MFP = 20 (the last measured MFP for which we have fit the curve) to 0 for a large MFP=1000 meters
-	Fade*= saturate(1.0f - (MFP - 20.0f) / 980.0f);
+	// Ensure the scattering fade out to 0 as the BaseColor approach 0. This Starts when BaseColor is below 10%.
+	const float BaseColorFadesOutBelowPercentage = 10.0f;
+	Fade = saturate(BaseColor * BaseColorFadesOutBelowPercentage);
+	// And also fade out BaseColor to zero from MFP = 20meters (the last measured MFP for which we have fit the curve) to a large MFP=1000 meters
+	const float FitLastMeasuredSampleMFP = 20.0f;
+	const float AlbedoIsZeroForMFP = 1000.0f;
+	Fade*= saturate(1.0f - (MFP - FitLastMeasuredSampleMFP) / (AlbedoIsZeroForMFP - FitLastMeasuredSampleMFP));
 
 	return Fade;
 }
@@ -178,7 +181,7 @@ float3 IsotropicMediumSlabPunctualDirectionalAlbedo(FParticipatingMedia PM)
 	const float3 EvaluateForBaseColor1	= 0.0855674 / (0.237742 + (MFP + ((0.0310849 - MFP) / (1.95492 * MFP + 2.07238))));
 	const float3 EvaluateForBaseColor01 = 0.0167964 / (0.541037 * (pow(1.17902, (-4.33046) / MFP) * (-0.294969 + MFP)) + 0.797592);
 
-	// Lerp between the two curves extreme curves
+	// Lerp between the two extreme BaseColor curves. Measurement range was [0.1f, 1.0f].
 	float3 FinalEvaluate = lerp(EvaluateForBaseColor01, EvaluateForBaseColor1, (PM.BaseColor - 0.1f) / (1.0f - 0.1f));
 	return FinalEvaluate * IsotropicMediumSlabDirectionalAlbedoFade(PM.BaseColor, MFP);
 }
@@ -192,7 +195,7 @@ float3 IsotropicMediumSlabEnvDirectionalAlbedo(FParticipatingMedia PM)
 	const float3 EvaluateForBaseColor1	= 0.00231881 + (0.51379 / (pow(MFP, 1.03577) + 0.510465)); 
 	const float3 EvaluateForBaseColor01 = 0.189167 / (1.55597 + (MFP + pow(0.182843, 0.0666775 + MFP)));
 
-	// Lerp between the two curves extreme curves
+	// Lerp between the two extreme BaseColor curves. Measurement range was [0.1f, 1.0f].
 	float3 FinalEvaluate = lerp(EvaluateForBaseColor01, EvaluateForBaseColor1, (PM.BaseColor - 0.1f) / (1.0f - 0.1f));
 	return FinalEvaluate * IsotropicMediumSlabDirectionalAlbedoFade(PM.BaseColor, MFP);
 }
