Fixes G5TAA with -d3d12 -sm6

diff --git a/Engine/Config/Windows/DataDrivenPlatformInfo.ini b/Engine/Config/Windows/DataDrivenPlatformInfo.ini
--- a/Engine/Config/Windows/DataDrivenPlatformInfo.ini
+++ b/Engine/Config/Windows/DataDrivenPlatformInfo.ini
@@ -66,12 +66,10 @@ bSupportsSSDIndirect=true
 bRequiresVendorExtensionsForAtomics=true
 bSupportsAnisotropicMaterials=true
 bSupportsGPUSkinCache=true
+bSupportsGen5TemporalAA=true
 bSupportsMeshShaders=true
 FriendlyName=LOCTEXT("FriendlyShaderPlatformName_PCD3D_SM6", "Shader Model 6")
 
-; TODO(UE-111389)
-bSupportsGen5TemporalAA=false
-
 [PreviewPlatform PCD3D_SM6]
 PlatformName=PCD3D_SM6
 ShaderFormat=PCD3D_SM6
@@ -156,9 +154,6 @@ bSupportsWaveOperations = false
 ; //#todo-vulkanddpi
 bRequiresExplicit128bitRT = false
 
-; //#todo-vulkanddpi
-bSupportsGen5TemporalAA = false
-
 bTargetsTiledGPU = false
 bNeedsOfflineCompiler = false
 
diff --git a/Engine/Shaders/Private/TemporalAA/TAAUpdateHistory.usf b/Engine/Shaders/Private/TemporalAA/TAAUpdateHistory.usf
--- a/Engine/Shaders/Private/TemporalAA/TAAUpdateHistory.usf
+++ b/Engine/Shaders/Private/TemporalAA/TAAUpdateHistory.usf
@@ -71,20 +71,6 @@ taa_half3 TransformColorForClampingBox(taa_half3 Color)
 	return RGBToYCoCg(Color);
 }
 
-uint GetGroupThreadIndex(uint GroupThreadIndex)
-#if COMPILER_SUPPORTS_WAVE_ONCE
-{
-	ISOLATE
-	{
-		return WaveGetLaneIndex();
-	}
-}
-#else
-{
-	return GroupThreadIndex;
-}
-#endif
-
 
 //------------------------------------------------------- ENTRY POINT
 
@@ -93,6 +79,8 @@ void MainCS(
 	uint2 GroupId : SV_GroupID,
 	uint GroupThreadIndex : SV_GroupIndex)
 {
+	uint GroupWaveIndex = GetGroupWaveIndex(GroupThreadIndex, /* GroupSize = */ TILE_SIZE * TILE_SIZE);
+
 	float4 Debug = 0.0;
 
 	uint2 DispatchThreadId = (
@@ -241,7 +229,7 @@ void MainCS(
 	#if CONFIG_MANUAL_LDS_SPILL
 	ISOLATE
 	{
-		uint LocalGroupThreadIndex = GetGroupThreadIndex(GroupThreadIndex);
+		uint LocalGroupThreadIndex = GetGroupThreadIndex(GroupThreadIndex, GroupWaveIndex);
 
 		SharedArray0[LocalGroupThreadIndex] = taa_half4(FilteredInputColor, LowFrequencyRejection);
 		SharedArray1[LocalGroupThreadIndex] = taa_half4(InputMinColor, InputPixelAlignement);
@@ -347,7 +335,7 @@ void MainCS(
 	#if CONFIG_MANUAL_LDS_SPILL
 	ISOLATE
 	{
-		uint LocalGroupThreadIndex = GetGroupThreadIndex(GroupThreadIndex);
+		uint LocalGroupThreadIndex = GetGroupThreadIndex(GroupThreadIndex, GroupWaveIndex);
 
 		taa_half4 RawLDS0 = SharedArray0[LocalGroupThreadIndex];
 		taa_half4 RawLDS1 = SharedArray1[LocalGroupThreadIndex];
@@ -716,7 +704,7 @@ void MainCS(
 
 	ISOLATE
 	{
-		uint LocalGroupThreadIndex = GetGroupThreadIndex(GroupThreadIndex);
+		uint LocalGroupThreadIndex = GetGroupThreadIndex(GroupThreadIndex, GroupWaveIndex);
 
 		#if 1
 			uint2 LocalHistoryPixelPos = (
diff --git a/Engine/Shaders/Private/TemporalAA/TAAUpdateSuperResHistory.usf b/Engine/Shaders/Private/TemporalAA/TAAUpdateSuperResHistory.usf
--- a/Engine/Shaders/Private/TemporalAA/TAAUpdateSuperResHistory.usf
+++ b/Engine/Shaders/Private/TemporalAA/TAAUpdateSuperResHistory.usf
@@ -23,23 +23,6 @@ RWTexture2D<uint> HistoryOutput_SuperResTextures_0;
 RWTexture2D<uint> HistoryOutput_SuperResTextures_1;
 
 
-//------------------------------------------------------- FUNCTIONS
-
-uint GetGroupThreadIndex(uint GroupThreadIndex)
-#if COMPILER_SUPPORTS_WAVE_ONCE
-{
-	ISOLATE
-	{
-		return WaveGetLaneIndex();
-	}
-}
-#else
-{
-	return GroupThreadIndex;
-}
-#endif
-
-
 //------------------------------------------------------- ENTRY POINT
 
 [numthreads(TILE_SIZE, TILE_SIZE, 1)]
@@ -47,6 +30,8 @@ void MainCS(
 	uint2 GroupId : SV_GroupID,
 	uint GroupThreadIndex : SV_GroupIndex)
 {
+	uint GroupWaveIndex = GetGroupWaveIndex(GroupThreadIndex, /* GroupSize = */ TILE_SIZE * TILE_SIZE);
+
 	float4 Debug = 0.0;
 
 	uint2 DispatchThreadId = (
@@ -195,7 +180,7 @@ void MainCS(
 	ISOLATE
 	{
 		#if 1
-			uint LocalGroupThreadIndex = GetGroupThreadIndex(GroupThreadIndex);
+			uint LocalGroupThreadIndex = GetGroupThreadIndex(GroupThreadIndex, GroupWaveIndex);
 			uint2 LocalHistoryPixelPos = (
 				ZOrder2D(LocalGroupThreadIndex, log2(TILE_SIZE)) +
 				GroupId * uint2(TILE_SIZE, TILE_SIZE));
