Fix incorrect shadow penumbra cast by hair strands onto geometry.

diff --git a/Engine/Shaders/Private/VirtualShadowMaps/ProjectionDirectional.ush b/Engine/Shaders/Private/VirtualShadowMaps/ProjectionDirectional.ush
--- a/Engine/Shaders/Private/VirtualShadowMaps/ProjectionDirectional.ush
+++ b/Engine/Shaders/Private/VirtualShadowMaps/ProjectionDirectional.ush
@@ -238,3 +238,35 @@ float TraceDirectional(
 	float RayMissFactor = float(RayMissCount) / float(SampleCount);
 	return RayMissFactor;
 }
+
+// Generate a ray based on directional light source geometry (e.g, source radius)
+bool GenerateRayDirectional(
+	FLightShaderParameters Light,
+	uint2 PixelPos,
+	float3 TranslatedWorldPosition,
+	float RayLengthScale,
+	uint RayIndex,
+	uint RayCount,
+	inout float3 OutRayStart,
+	inout float3 OutRayEnd)
+{
+	float3 WorldPosition = TranslatedWorldPosition - View.PreViewTranslation;
+	float3 ViewPosition = mul(float4(TranslatedWorldPosition, 1.0f), View.TranslatedWorldToView).xyz;
+	float DistanceFromViewOrigin = length(ViewPosition);
+
+	// Compute max ray length based on view depth
+	// This affects both how far we are willing to trace across the shadow map (for clipmaps this is related to view depth)
+	// and the maximum size a penumbra can be.
+	// Too high values will cause shadows to detach from their contact points (unless more samples are used).
+	// Too low values will greatly restrict how large penumbras can be in screen space.
+	float RayLength = RayLengthScale * DistanceFromViewOrigin;
+
+	RandomSequence RandSequence;
+	RandomSequence_Initialize(RandSequence, PixelPos, RayIndex, View.StateFrameIndexMod8, RayCount);
+	float2 E = RandomSequence_GenerateSample2D(RandSequence);
+
+	float3 RayDir = GetRandomDirectionalLightRayDir(Light, E);
+	OutRayStart = WorldPosition;
+	OutRayEnd = OutRayStart + RayDir * RayLength;
+	return true;
+}
diff --git a/Engine/Shaders/Private/VirtualShadowMaps/ProjectionSpot.ush b/Engine/Shaders/Private/VirtualShadowMaps/ProjectionSpot.ush
--- a/Engine/Shaders/Private/VirtualShadowMaps/ProjectionSpot.ush
+++ b/Engine/Shaders/Private/VirtualShadowMaps/ProjectionSpot.ush
@@ -357,3 +357,64 @@ float TraceLocalLight(
 
 	return RayMissFactor;
 }
+
+// Generate ray based on light source geometry
+bool GenerateRayLocalLight(
+	FLightShaderParameters Light,
+	uint2 PixelPos,
+	float3 TranslatedWorldPosition,
+	float CotMaxRayAngleFromLight,
+	uint RayIndex,
+	uint RayCount,
+	inout float3 OutRayStart,
+	inout float3 OutRayEnd)
+{
+	const half3 N = GetGBufferDataUint(PixelPos, true).WorldNormal;
+	const float3 WorldPosition = TranslatedWorldPosition - View.PreViewTranslation;
+
+	float3 ToLight = Light.Position - WorldPosition;
+	float DistSqr = dot(ToLight, ToLight);
+	float Falloff = rcp(DistSqr + 1);
+
+	float InvDist = rsqrt(DistSqr);
+	float DistToLight = DistSqr * InvDist;
+
+	half3 ConeAxis = ToLight * InvDist;
+
+	half SinAlphaSqr = saturate(Pow2(Light.SourceRadius) * Falloff);
+	half SinAlpha = sqrt(SinAlphaSqr);
+	bool bBackFace = dot(N, ConeAxis) < -SinAlpha;
+
+	if (!bBackFace)
+	{
+		half ConeSin = Light.SourceRadius * InvDist;
+		half ConeCos = sqrt(1 - Pow2(ConeSin));
+
+		uint TileID = MortonEncode(PixelPos);
+		RandomSequence RandSequence;
+		// Combine pixel-level and sample-level bits into the sample index (visible structure will be hidden by owen scrambling of the index)
+		RandSequence.SampleIndex = reversebits((65536 * View.StateFrameIndex + TileID) * RayCount + RayIndex);
+		RandSequence.SampleSeed = 0; // always use the same sequence
+
+		float2 E = SobolSampler(RandSequence.SampleIndex, RandSequence.SampleSeed, 8).xy;
+		half2 LightUV = UniformSampleDiskConcentricApprox(E).xy;
+		LightUV *= ConeSin;
+
+		half SinTheta2 = dot(LightUV, LightUV);
+		half SinTheta = sqrt(SinTheta2);
+		half CosTheta = sqrt(1 - SinTheta2);
+
+		half3 Dir = TangentToWorld(half3(LightUV, CosTheta), ConeAxis);
+
+		// The further afield our rays go the poorer our approximation is as the "bend" due to our testing
+		// against a shadow map instead of along the ray increases. Thus we avoid going all the way to the light
+		// where the projection becomes extreme.
+
+		// Clip ray to max angle to limit penumbra sizes
+		float ClampedLength = DistToLight * 0.75 * saturate(1.5 / (CosTheta + CotMaxRayAngleFromLight * SinTheta));
+
+		OutRayStart = WorldPosition;
+		OutRayEnd = WorldPosition + Dir * ClampedLength;
+	}
+	return !bBackFace;
+}
diff --git a/Engine/Shaders/Private/VirtualShadowMaps/VirtualShadowMapProjection.usf b/Engine/Shaders/Private/VirtualShadowMaps/VirtualShadowMapProjection.usf
--- a/Engine/Shaders/Private/VirtualShadowMaps/VirtualShadowMapProjection.usf
+++ b/Engine/Shaders/Private/VirtualShadowMaps/VirtualShadowMapProjection.usf
@@ -308,23 +308,45 @@ FProjectLightResult ProjectLight(
 
 		// Hair strands voxel traversal to apply hair shadow on top of opaque geometry
 		#if HAS_HAIR_STRANDS
+		float3 RayStart = 0;
+		float3 RayEnd = 0;
+		bool bCastHairRay = false;
+		float4 Random = 0;
 		if (ShadowFactor > 0)
-		{		
-			#if DIRECTIONAL_LIGHT
-			const float VoxelShadowMaxDistance = 1000.0; // 100m shadow distance
-			const float3 LightPosition = Light.Direction * VoxelShadowMaxDistance + WorldPosition;
-			#else
-			const float3 LightPosition = Light.Position;
-			#endif
-
-			// Depth bias
-			// Origin is shifted 2 voxels away towards the light + a constant bias of the size of the voxel
-			float3 OriginPosition = WorldPosition;
-			const float4 Random = ComputeRandom4(PixelPos);
+		{
+			Random = ComputeRandom4(PixelPos);
+			uint RayIndex = min(Random.w * SMRTRayCount, SMRTRayCount - 1);
+
+		#if DIRECTIONAL_LIGHT
+			bCastHairRay = GenerateRayDirectional(
+				Light,
+				PixelPos,
+				TranslatedWorldPosition,
+				ContactShadowLengthWorld,
+				RayIndex,
+				SMRTRayCount,
+				RayStart,
+				RayEnd);
+		#else
+			bCastHairRay = GenerateRayLocalLight(
+				Light,
+				PixelPos,
+				TranslatedWorldPosition,
+				SMRTCotMaxRayAngleFromLight,
+				RayIndex,
+				SMRTRayCount,
+				RayStart,
+				RayEnd);
+		#endif
+		}
+
+		if (bCastHairRay)
+		{
+			// Jitter start position to mask/compensate coarse voxel resolution
 			{
 				const float PositionBiasScale = 0.5f;
 				const float3 DepthBias = VirtualVoxel.VoxelWorldSize * (VirtualVoxel.DepthBiasScale_Shadow * L + PositionBiasScale * (Random.xyz * 2 - 1));
-				OriginPosition += DepthBias;
+				RayStart += DepthBias;
 			}
 
 			const float DistanceThreshold = 100000.0f;
@@ -344,6 +366,7 @@ FProjectLightResult ProjectLight(
 			TraversalSettings.SteppingScale		= VirtualVoxel.SteppingScale_Shadow;
 			TraversalSettings.Random			= Random.xyz;
 			TraversalSettings.PixelRadius		= SceneDepth * VirtualVoxel.HairCoveragePixelRadiusAtDepth1;
+			TraversalSettings.bUseOpaqueVisibility = false;
 
 			const uint VoxelDescCount = VirtualVoxel.NodeDescCount;
 			for (uint VoxelDescIt=0; VoxelDescIt<VoxelDescCount; ++VoxelDescIt)
@@ -353,8 +376,8 @@ FProjectLightResult ProjectLight(
 
 				FHairTraversalResult Result = InitHairTraversalResult();
 				Result = ComputeHairCountVirtualVoxel(
-					OriginPosition,
-					LightPosition,
+					RayStart,
+					RayEnd,
 					CommonDesc,
 					NodeDesc,
 					VirtualVoxel.PageIndexBuffer,
@@ -362,6 +385,7 @@ FProjectLightResult ProjectLight(
 					VirtualVoxel.PageTexture,
 					TraversalSettings);
 
+				HitDistance = Result.HitT;
 				ShadowFactor = min(ShadowFactor, saturate(1 - Result.HairCount));
 			}
 		}
