Improve capsule light sampling in the path tracer

Use the logic from CapsuleLightIntegrate to sample proportional to solid angle with a spherical rectangle aligned to the visible portion of the capsule. We also build a spherical cone around the rectangle to handle grazing angle cases as well as very short capsules.

Fixed corner cases in the bounding rectangle calculation to account for required sign flip.

Improved robustness of rejection test for small radii

Refactored common code with CapsuleLightIntegrate, which now uses the same sampling method.

diff --git a/Engine/Shaders/Private/CapsuleLightIntegrate.ush b/Engine/Shaders/Private/CapsuleLightIntegrate.ush
--- a/Engine/Shaders/Private/CapsuleLightIntegrate.ush
+++ b/Engine/Shaders/Private/CapsuleLightIntegrate.ush
@@ -7,6 +7,7 @@
 #include "AreaLightCommon.ush"
 #include "ShadingModels.ush"
 #include "CapsuleLight.ush"
+#include "CapsuleLightSampling.ush"
 #include "SobolRandom.ush"
 
 Texture2D DummyRectLightTextureForCapsuleCompilerWarning;
@@ -133,41 +134,9 @@ FDirectLighting IntegrateBxDF( FGBufferData GBuffer, half3 N, half3 V, FCapsuleL
 
 	float DistanceSqr = dot( ToLight, ToLight );
 	float3 ConeAxis = ToLight * rsqrt( DistanceSqr );
-	float ConeCos = sqrt( 1 - Square( Capsule.Radius ) / DistanceSqr );
-
-
-	float3 ClosestPointOnAxis = ToLight - CapsuleAxis * dot( CapsuleAxis, ToLight );
-	
-	FRect Rect;
-	Rect.Origin = ToLight;
-	Rect.Axis[1] = CapsuleAxis;
-	Rect.Axis[2] = normalize( -ClosestPointOnAxis );
-	Rect.Axis[0] = cross( Rect.Axis[1], Rect.Axis[2] );
-
-	float SinCylinderAngle = Capsule.Radius * rsqrt( dot( ClosestPointOnAxis, ClosestPointOnAxis ) );
-	// width * cos(angle) = 2r 
-	float RectRadius = Capsule.Radius * rsqrt( 1 - Pow2( SinCylinderAngle ) );
-
-	float Extension[2];
-	for( uint i = 0; i < 2; i++ )
-	{
-		float3 PointPos = ToLight + CapsuleAxis * Capsule.Length * ( i > 0 ? 0.5 : -0.5 );
-
-		float SinSphereAngle = Capsule.Radius * rsqrt( dot( PointPos, PointPos ) );
-		float CosAxisAngle = abs( dot( CapsuleAxis, normalize( PointPos ) ) );
-		//float SinExtension = sin( AxisAngle -  SphereAngle );
-		float SinExtension = sqrt( ( 1 - Pow2( CosAxisAngle ) ) * ( 1 - Pow2( SinSphereAngle ) ) ) - CosAxisAngle * SinSphereAngle;
-		Extension[i] = Capsule.Radius / SinExtension;
-	}
-
-	float Translate = 0.5 * ( Extension[1] - Extension[0] );
-	float Extend = 0.5 * ( Extension[0] + Extension[1] );
-	
-	Rect.Origin += Translate * Rect.Axis[1];
-	Rect.Extent = float2( RectRadius, 0.5 * Capsule.Length + Extend );
-
-	FSphericalRect SphericalRect = BuildSphericalRect( Rect );
+	float SineConeSqr = saturate(Pow2(Capsule.Radius) / DistanceSqr);
 
+	FCapsuleSphericalBounds CapsuleBounds = CapsuleGetSphericalBounds(ToLight, CapsuleAxis, Capsule.Radius, Capsule.Length);
 	
 	const uint NumSets = 3;
 	const uint NumSamples[ NumSets ] =
@@ -255,13 +224,13 @@ FDirectLighting IntegrateBxDF( FGBufferData GBuffer, half3 N, half3 V, FCapsuleL
 
 				if( Capsule.Length > 0 )
 				{
-					float3 ToArea = UniformSampleSphericalRect( E, SphericalRect );
+					float3 ToArea = CapsuleBounds.Sample(E).xyz;
 					L = normalize( ToArea );
 					H = normalize( V + L );
 				}
 				else
 				{
-					L = TangentToWorld( UniformSampleCone( E, ConeCos ).xyz, ConeAxis );
+					L = TangentToWorld( UniformSampleConeRobust( E, SineConeSqr).xyz, ConeAxis );
 					H = normalize(V + L);
 				}
 			}
@@ -273,13 +242,7 @@ FDirectLighting IntegrateBxDF( FGBufferData GBuffer, half3 N, half3 V, FCapsuleL
 			if( NoL > 0 && VoH > 0 )
 			{
 				{
-					float3 ToSphere = ClosestPointLineToRay( Capsule.LightPos[0], Capsule.LightPos[1], Capsule.Length, L );
-
-					float DistSqr = dot( ToSphere, ToSphere );
-					float3 ConeAxis = ToSphere * rsqrt( DistSqr );
-					float ConeCos = sqrt( 1 - Square( Capsule.Radius ) / DistSqr );
-
-					if( dot( L, ConeAxis ) < ConeCos )
+					if (CapsuleTest(L, ToLight, CapsuleAxis, Pow2(Capsule.Radius), Capsule.Length) > 0)
 					{
 						// Ray misses sphere
 						continue;
@@ -290,12 +253,12 @@ FDirectLighting IntegrateBxDF( FGBufferData GBuffer, half3 N, half3 V, FCapsuleL
 				{
 					NoL / PI,
 					D_GGX( Pow4(GBuffer.Roughness), NoH ) * NoH / (4 * VoH),
-					1.0 / SphericalRect.SolidAngle,
+					1.0 / CapsuleBounds.GetSolidAngle(),
 				};
 
 				if( Capsule.Length == 0 )
 				{
-					PDF[2] = 1.0 / ( 2 * PI * (1 - ConeCos) );
+					PDF[2] = 1.0 / UniformConeSolidAngle(SineConeSqr);
 				}
 
 				// MIS power heuristic
diff --git a/Engine/Shaders/Private/CapsuleLightSampling.ush b/Engine/Shaders/Private/CapsuleLightSampling.ush
new file mode 100644
--- /dev/null
+++ b/Engine/Shaders/Private/CapsuleLightSampling.ush
@@ -0,0 +1,164 @@
+// Copyright Epic Games, Inc. All Rights Reserved.
+
+#pragma once
+
+#include "RectLight.ush"
+
+float SqrtOneMinusX(float x)
+{
+	return x < 0.01 ? 1 - x * (0.5 + x * 0.125) : sqrt(1 - x);
+}
+
+// The footprint of a capsule can be bounded by either a rectangle or a cone
+// This class holds both and will do sampling from whichever has the smallest solidangle
+struct FCapsuleSphericalBounds
+{
+	FSphericalRect SphericalRect;
+	float3 ConeAxis;
+	float ConeSinThetaMax2;
+	float ConeSolidAngle;
+
+	float GetSolidAngle()
+	{
+		// See explanation in UniformSampleConeRobust
+		return min(SphericalRect.SolidAngle, ConeSolidAngle);
+	}
+
+	// return world space direction within the cone
+	float4 Sample(float2 E)
+	{
+		if (ConeSolidAngle < SphericalRect.SolidAngle)
+		{
+			return float4(TangentToWorld(UniformSampleConeRobust(E, ConeSinThetaMax2).xyz, ConeAxis), ConeSolidAngle);
+		}
+		else
+		{
+			return float4(normalize(UniformSampleSphericalRect(E, SphericalRect)), SphericalRect.SolidAngle);
+		}
+	}
+};
+
+// Given a bounding rectangle of a capsule, compute a bounding cone
+FCapsuleSphericalBounds CapsuleGetSphericalBounds(float3 Origin, float3 Axis, float Radius, float Length)
+{
+	float h = dot(Axis, Origin);
+	float3 ClosestPointOnAxis = Origin - Axis * h;
+
+	float DistanceToAxisSqr = dot(ClosestPointOnAxis, ClosestPointOnAxis);
+	float RadiusSqr = Pow2(Radius);
+
+	if (DistanceToAxisSqr <= RadiusSqr)
+	{
+		// we are inside the infinite cylinder - only one of the caps can be visible
+		float3 CapCenter = Origin - Axis * Length * 0.5 * sign(h);
+		float LightDistanceSquared = dot(CapCenter, CapCenter);
+		float SinThetaMax2 = saturate(RadiusSqr / LightDistanceSquared);
+
+		FCapsuleSphericalBounds Result;
+		Result.SphericalRect = (FSphericalRect)0;
+		Result.SphericalRect.SolidAngle = POSITIVE_INFINITY; // make sure we don't pick this
+		Result.ConeAxis = normalize(CapCenter);
+		Result.ConeSinThetaMax2 = SinThetaMax2;
+		Result.ConeSolidAngle = UniformConeSolidAngle(SinThetaMax2);
+		return Result;
+	}
+
+	// we are outside the infinite cylinder, build a bounding rectangle
+	FRect Rect;
+	Rect.Origin = Origin;
+	Rect.Axis[1] = Axis;
+	Rect.Axis[2] = normalize(-ClosestPointOnAxis);
+	Rect.Axis[0] = cross(Rect.Axis[1], Rect.Axis[2]);
+
+	float SinCylinderAngle = Radius * rsqrt(DistanceToAxisSqr);
+	// width * cos(angle) = 2r 
+	float RectRadius = Radius * rsqrt(1 - Pow2(SinCylinderAngle));
+
+	float Extension[2];
+	for (int i = 0; i < 2; i++)
+	{
+		float hi = Length * (i > 0 ? 0.5 : -0.5);
+		float3 PointPos = Origin + Axis * hi;
+
+		float InverseDist = rsqrt(dot(PointPos, PointPos));
+
+		float SinSphereAngle = saturate(Radius * InverseDist);
+		float CosSphereAngle = SqrtOneMinusX(Pow2(SinSphereAngle));
+
+		float CosAxisAngle = -dot(Axis, PointPos) * InverseDist;
+		CosAxisAngle = sign(CosAxisAngle * hi) * saturate(abs(CosAxisAngle));
+		float SinAxisAngle = SqrtOneMinusX(Pow2(CosAxisAngle));
+
+		float CosExtension = SinAxisAngle * CosSphereAngle + CosAxisAngle * SinSphereAngle;
+		Extension[i] = Radius / CosExtension;
+	}
+
+	float Translate = 0.5 * (Extension[1] - Extension[0]);
+	float Extend = 0.5 * (Extension[0] + Extension[1]);
+
+	Rect.Origin += Translate * Rect.Axis[1];
+	Rect.Extent = float2(RectRadius, 0.5 * Length + Extend);
+
+	// get a bounding cone from the bounding rectangle
+	float3 R0 = Rect.Origin - Rect.Axis[1] * Rect.Extent.y;
+	float3 R1 = Rect.Origin + Rect.Axis[1] * Rect.Extent.y;
+	float InvDistR0 = rsqrt(dot(R0, R0));
+	float InvDistR1 = rsqrt(dot(R1, R1));
+
+	FCapsuleSphericalBounds Result;
+	Result.SphericalRect = BuildSphericalRect(Rect);
+	Result.ConeAxis = normalize(lerp(R0, R1, saturate(InvDistR1 / (InvDistR0 + InvDistR1))));
+	Result.ConeSinThetaMax2 = 0.5 - 0.5 * saturate(dot(R0, R1) * InvDistR0 * InvDistR1); // sin(x/2)^2 = (1-cos(x))/2
+	Result.ConeSolidAngle = UniformConeSolidAngle(Result.ConeSinThetaMax2);
+	return Result;
+}
+
+// Adapted from:
+//    https://www.iquilezles.org/www/articles/intersectors/intersectors.htm
+// This code assume the ray direction is normalized and only returns the front hit
+float CapsuleIntersect(float3 Rd, float3 Center, float3 Axis, float Radius2, float Length)
+{
+	float3 ba = Axis;
+	float3 oa = -Center;
+	float bard = dot(ba, Rd);
+	float baoa = dot(ba, oa);
+	float rdoa = dot(Rd, oa);
+	float oaoa = dot(oa, oa);
+	float a = 1 - bard * bard;
+	float b = rdoa - baoa * bard;
+	float c = oaoa - baoa * baoa - Radius2;
+	float h = b * b - a * c;
+	if (h >= 0.0)
+	{
+		float t = (-b - sqrt(h)) / a;
+		float y = baoa + t * bard;
+		// body
+		if (abs(y) < 0.5 * Length)
+		{
+			return t;
+		}
+
+		// caps
+		float3 oc = oa - (sign(y) * 0.5 * Length) * ba;
+		b = dot(Rd, oc);
+		c = dot(oc, oc) - Radius2;
+		h = b * b - c;
+		if (h > 0.0)
+		{
+			return -b - sqrt(h);
+		}
+	}
+	return -1.0;
+}
+
+// simpler intersection test for rays that are likely to intersect
+// returned intersection distance is approximate
+float CapsuleTest(float3 Rd, float3 Center, float3 Axis, float Radius2, float Length)
+{
+	// Shortest distance
+	float B = dot(Rd, Axis);
+	float t = clamp(dot(Center, B * Rd - Axis) / (1 - B * B), -0.5 * Length, 0.5 * Length);
+	float3 ToSphere = Center + t * Axis;
+	float3 C = cross(Rd, ToSphere);
+	return dot(C, C) <= Radius2 ? length(ToSphere) : -1.0;
+}
diff --git a/Engine/Shaders/Private/MonteCarlo.ush b/Engine/Shaders/Private/MonteCarlo.ush
--- a/Engine/Shaders/Private/MonteCarlo.ush
+++ b/Engine/Shaders/Private/MonteCarlo.ush
@@ -315,6 +315,12 @@ float4 UniformSampleConeRobust(float2 E, float SinThetaMax2)
 	return float4(L, PDF);
 }
 
+float UniformConeSolidAngle(float SinThetaMax2)
+{
+	float OneMinusCosThetaMax = SinThetaMax2 < 0.01 ? SinThetaMax2 * (0.5 + 0.125 * SinThetaMax2) : 1 - sqrt(1 - SinThetaMax2);
+	return 2 * PI * OneMinusCosThetaMax;
+}
+
 // Same as the function above, but uses a concentric mapping
 float4 UniformSampleConeConcentricRobust(float2 E, float SinThetaMax2)
 {
diff --git a/Engine/Shaders/Private/PathTracing/Light/PathTracingCapsuleLight.ush b/Engine/Shaders/Private/PathTracing/Light/PathTracingCapsuleLight.ush
--- a/Engine/Shaders/Private/PathTracing/Light/PathTracingCapsuleLight.ush
+++ b/Engine/Shaders/Private/PathTracing/Light/PathTracingCapsuleLight.ush
@@ -7,46 +7,9 @@ PathTracingCapsuleLight.ush: Light sampling functions for capsule case of point
 #pragma once
 
 #include "PathTracingLightCommon.ush"
+#include "../../CapsuleLightSampling.ush"
 
-
-// Adapted from:
-//    https://www.iquilezles.org/www/articles/intersectors/intersectors.htm
-// This code assume the ray direction is normalized and only returns the front hit
-float CapsuleIntersect(float3 Ro, float3 Rd, float3 Center, float3 Axis, float Radius)
-{
-	float3  ba = 0.5 * Axis;
-	float3  oa = Ro - Center;
-	float baba = dot(ba, ba);
-	float bard = dot(ba, Rd);
-	float baoa = dot(ba, oa);
-	float rdoa = dot(Rd, oa);
-	float oaoa = dot(oa, oa);
-	float a = baba - bard * bard;
-	float b = baba * rdoa - baoa * bard;
-	float c = baba * oaoa - baoa * baoa - Radius * Radius * baba;
-	float h = b * b - a * c;
-	if (h >= 0.0)
-	{
-		float t = (-b - sqrt(h)) / a;
-		float y = baoa + t * bard;
-		// body
-		if (abs(y) < baba)
-		{
-			return t;
-		}
-
-		// caps
-		float3 oc = oa - sign(y) * ba;
-		b = dot(Rd, oc);
-		c = dot(oc, oc) - Radius * Radius;
-		h = b * b - c;
-		if (h > 0.0)
-		{
-			return -b - sqrt(h);
-		}
-	}
-	return -1.0;
-}
+#define CAPSULE_SOLIDANGLE_SAMPLING		1
 
 float3 CapsuleNormal(float3 Pos, float3 Center, float3 Axis)
 {
@@ -55,34 +18,112 @@ float3 CapsuleNormal(float3 Pos, float3 Center, float3 Axis)
 	return normalize(pa - h * Axis);
 }
 
+float GetCapsuleRadius(int LightId)
+{
+	// don't allow 0 area capsules since we need to divide the light power by area to get radiance
+	return max(GetRadius(LightId), 0.01);
+}
+
 FLightHit CapsuleLight_TraceLight(RayDesc Ray, int LightId)
 {
 	float3 LightPosition = GetPosition(LightId);
-	float LightRadius = GetRadius(LightId);
+	float LightRadius = GetCapsuleRadius(LightId);
+	float LightRadius2 = Pow2(LightRadius);
 	float SourceLength = GetSourceLength(LightId);
-	float3 Axis = GetdPdv(LightId) * SourceLength;
-	float t = CapsuleIntersect(Ray.Origin, Ray.Direction, LightPosition, Axis, LightRadius);
+
+#if 0 // DEBUG code: visualize the bounding rectangle against the capsule intersection test
+	FCapsuleSphericalBounds CapsuleBounds = CapsuleGetSphericalBounds(LightPosition - Ray.Origin, GetdPdv(LightId), LightRadius, SourceLength);
+
+	float ConeSolidAngle = CapsuleBounds.ConeSolidAngle;
+	float RectSolidAngle = CapsuleBounds.SphericalRect.SolidAngle;
+	bool HitBounds = false;
+	float t = -1;
+	if (ConeSolidAngle < RectSolidAngle)
+	{
+		// Is ray pointing inside the cone of directions?
+		float CosTheta = saturate(dot(normalize(Ray.Direction), CapsuleBounds.Axis));
+		t = CosTheta >= SqrtOneMinusX(CapsuleBounds.SinThetaMax2) ? length(LightPosition) : -1.0;
+	}
+	else
+	{
+		float DoN = dot(Ray.Direction, CapsuleBounds.SphericalRect.Axis[2]);
+		t = CapsuleBounds.SphericalRect.z0 / DoN;
+		// ray points toward the plane and intersect it?
+		if (t > Ray.TMin && t < Ray.TMax)
+		{
+			float2 UV = t * float2(
+				dot(Ray.Direction, CapsuleBounds.SphericalRect.Axis[0]),
+				dot(Ray.Direction, CapsuleBounds.SphericalRect.Axis[1])) - float2(
+					0.5 * (CapsuleBounds.SphericalRect.x0 + CapsuleBounds.SphericalRect.x1),
+					0.5 * (CapsuleBounds.SphericalRect.y0 + CapsuleBounds.SphericalRect.y1)
+				);
+			float2 Extent = 0.5 * float2(
+				CapsuleBounds.SphericalRect.x1 - CapsuleBounds.SphericalRect.x0,
+				CapsuleBounds.SphericalRect.y1 - CapsuleBounds.SphericalRect.y0);
+			// reject hit if outside of quad extents
+			t = all(abs(UV) < Extent) ? t : -1.0;
+		}
+	}
+
+	// ray points toward the plane and intersect it?
 	if (t > Ray.TMin && t < Ray.TMax)
 	{
-		float CapsArea = Pow2(LightRadius);
+		float CapsArea = LightRadius2;
 		float BodyArea = 0.5 * LightRadius * SourceLength;
 		float3 LightPower = GetColor(LightId) * ComputeIESAttenuation(LightId, Ray.Origin);
 		float3 LightRadiance = LightPower / (PI * (CapsArea + BodyArea));
-		float3 LightDirection = GetPosition(LightId) - Ray.Origin;
+		float t = CapsuleIntersect(Ray.Direction, LightPosition - Ray.Origin, GetdPdv(LightId), LightRadius2, SourceLength);
+		if (t > Ray.TMin && t < Ray.TMax)
+		{
+			LightRadiance *= float3(0, 1, 0);
+		}
+		return CreateLightHit(LightRadiance, 0.0, t);
+	}
+
+	// missed the bounding rect, test capsule anyway to see if any part of the capsule was missing
+	{
+		float CapsArea = LightRadius2;
+		float BodyArea = 0.5 * LightRadius * SourceLength;
+		float3 LightPower = GetColor(LightId) * ComputeIESAttenuation(LightId, Ray.Origin);
+		float3 LightRadiance = LightPower / (PI * (CapsArea + BodyArea));
+		{
+			// exact
+			float t = CapsuleIntersect(Ray.Direction, LightPosition - Ray.Origin, GetdPdv(LightId), LightRadius2, SourceLength);
+			if (t > Ray.TMin && t < Ray.TMax)
+			{
+				LightRadiance *= float3(1, 0, 0);
+				return CreateLightHit(LightRadiance, 0.0, t);
+			}
+		}
+	}
+#else
+	float3 Axis = GetdPdv(LightId);
+	float t = CapsuleIntersect(Ray.Direction, LightPosition - Ray.Origin, Axis, LightRadius2, SourceLength);
+	if (t > Ray.TMin && t < Ray.TMax)
+	{
+		float CapsArea = LightRadius2;
+		float BodyArea = 0.5 * LightRadius * SourceLength;
+		float3 LightPower = GetColor(LightId) * ComputeIESAttenuation(LightId, Ray.Origin);
+		float3 LightRadiance = LightPower / (PI * (CapsArea + BodyArea));
+		float3 LightDirection = LightPosition - Ray.Origin;
 		float LightDistanceSquared = dot(LightDirection, LightDirection);
 		LightRadiance *= ComputeAttenuationFalloff(LightDistanceSquared, LightId);
-
-		float3 Normal = CapsuleNormal(Ray.Origin + t * Ray.Direction, LightPosition, Axis);
+#if CAPSULE_SOLIDANGLE_SAMPLING
+		float Pdf = 1.0 / CapsuleGetSphericalBounds(LightDirection, Axis, LightRadius, SourceLength).GetSolidAngle();
+#else
+		float3 Normal = CapsuleNormal(Ray.Origin + t * Ray.Direction, LightPosition, Axis * SourceLength);
 
 		float CosTheta = saturate(-dot(Normal, Ray.Direction));
 
 		float Pdf = CosTheta > 0 ? t * t / (4 * PI * CosTheta * (CapsArea + BodyArea)) : 0.0;
-
+#endif
 		return CreateLightHit(LightRadiance, Pdf, t);
 	}
+#endif
 	return NullLightHit();
 }
 
+
 FLightSample CapsuleLight_SampleLight(
 	int LightId,
 	float4 RandSample,
@@ -93,19 +134,47 @@ FLightSample CapsuleLight_SampleLight(
 	// Capsule case
 	// #dxr_todo: only sample the visible portion of the capsule and account for the 1/d^2 falloff down the axis
 
-	// area sampling doesn't work if the shape has 0 area, so just clamp radius a bit
-	float Radius = max(GetRadius(LightId), 0.01);
+	float Radius = GetCapsuleRadius(LightId);
+	float Radius2 = Radius * Radius;
 	float SourceLength = GetSourceLength(LightId);
 
 	// the caps are two halves of a full sphere
 	// the body is a cylinder
 	// the common factor of 4*PI is accounted for at the end
-	float CapsArea = Radius * Radius;
+	float CapsArea = Radius2;
 	float BodyArea = 0.5 * Radius * SourceLength;
 
-	float Prob = CapsArea / (CapsArea + BodyArea);
 	float3 Axis = GetdPdv(LightId);
-	float3x3 CylinderBasis = GetTangentBasis(Axis);
+
+#if CAPSULE_SOLIDANGLE_SAMPLING
+	float3 LightPower = GetColor(LightId);
+	float3 LightRadiance = LightPower / (PI * (CapsArea + BodyArea));
+
+	float3 LightDirection = GetPosition(LightId) - WorldPos;
+	FCapsuleSphericalBounds CapsuleBounds = CapsuleGetSphericalBounds(LightDirection, Axis, Radius, SourceLength);
+
+	float4 Result = CapsuleBounds.Sample(RandSample.yz);
+
+	float3 L = Result.xyz;
+	float SolidAngle = Result.w;
+		
+	// check direction to account for rays that hit the bounding shape but not the capsule
+#if 1
+	// optimized check (also more robust for tiny radii)
+	float Distance = CapsuleTest(L, LightDirection, Axis, Radius2, SourceLength);
+#else
+	// exact check (produces artifacts for tiny radii)
+	float Distance = CapsuleIntersect(L, LightDirection, Axis, Radius2, SourceLength);
+#endif
+	if (Distance > 0)
+	{
+		return CreateLightSample(LightRadiance * SolidAngle, 1.0 / SolidAngle, L, Distance);
+	}
+	// didn't pass the acceptance test -- reject this sample
+	return NullLightSample();
+#else
+	// plain area sampling
+	float Prob = CapsArea / (CapsArea + BodyArea);
 
 	float3 LightPoint;
 	float3 LightNormal;
@@ -131,6 +200,8 @@ FLightSample CapsuleLight_SampleLight(
 		LightNormal = float3(cos(Phi), sin(Phi), 0.0);
 		LightPoint = float3(Radius * LightNormal.xy, (RandSample.z - 0.5) * SourceLength);
 	}
+	float3x3 CylinderBasis = GetTangentBasis(Axis);
+
 	float3 LightDirection = GetPosition(LightId) - WorldPos;
 	float3 LocalLightDirection = mul(CylinderBasis, LightDirection); // World To Local
 
@@ -145,6 +216,6 @@ FLightSample CapsuleLight_SampleLight(
 	float3 LightPower = GetColor(LightId);
 	float3 RadianceOverPdf = LightPower * CosTheta * 4 / DistanceSquared;
 	float Pdf = CosTheta > 0 ? DistanceSquared / (4 * PI * CosTheta * (CapsArea + BodyArea)) : 0.0;
-
 	return CreateLightSample(RadianceOverPdf, Pdf, Direction, Distance);
+#endif
 }
