Removes VALU stall on SALU in TAA's SafeRcp(half)

diff --git a/Engine/Shaders/Private/TemporalAA/TAACommon.ush b/Engine/Shaders/Private/TemporalAA/TAACommon.ush
--- a/Engine/Shaders/Private/TemporalAA/TAACommon.ush
+++ b/Engine/Shaders/Private/TemporalAA/TAACommon.ush
@@ -358,7 +358,8 @@ float SafeRcp(float x)
 
 half SafeRcp(half x)
 {
-	return x > half(0.0) ? min(rcp(x), taa_half(MaxHalfFloat)) : half(0.0);
+	// If x=0.0, then MaxHalfFloat * 0.0 = 0.0
+	return min(rcp(x), half(MaxHalfFloat)) * saturate(x * asfloat16(uint16_t(0x7C00)));
 }
 
 #endif
