Fix minor path tracer capsule light bugs introduced during refactoring

diff --git a/Engine/Shaders/Private/CapsuleLightSampling.ush b/Engine/Shaders/Private/CapsuleLightSampling.ush
--- a/Engine/Shaders/Private/CapsuleLightSampling.ush
+++ b/Engine/Shaders/Private/CapsuleLightSampling.ush
@@ -108,7 +108,7 @@ FCapsuleSphericalBounds CapsuleGetSphericalBounds(float3 Origin, float3 Axis, fl
 	FCapsuleSphericalBounds Result;
 	Result.SphericalRect = BuildSphericalRect(Rect);
 	Result.ConeAxis = normalize(lerp(R0, R1, saturate(InvDistR1 / (InvDistR0 + InvDistR1))));
-	Result.ConeSinThetaMax2 = 0.5 - 0.5 * saturate(dot(R0, R1) * InvDistR0 * InvDistR1); // sin(x/2)^2 = (1-cos(x))/2
+	Result.ConeSinThetaMax2 = saturate(0.5 - 0.5 * dot(R0, R1) * InvDistR0 * InvDistR1); // sin(x/2)^2 = (1-cos(x))/2
 	Result.ConeSolidAngle = UniformConeSolidAngle(Result.ConeSinThetaMax2);
 	return Result;
 }
diff --git a/Engine/Shaders/Private/PathTracing/Light/PathTracingCapsuleLight.ush b/Engine/Shaders/Private/PathTracing/Light/PathTracingCapsuleLight.ush
--- a/Engine/Shaders/Private/PathTracing/Light/PathTracingCapsuleLight.ush
+++ b/Engine/Shaders/Private/PathTracing/Light/PathTracingCapsuleLight.ush
@@ -41,8 +41,8 @@ FLightHit CapsuleLight_TraceLight(RayDesc Ray, int LightId)
 	if (ConeSolidAngle < RectSolidAngle)
 	{
 		// Is ray pointing inside the cone of directions?
-		float CosTheta = saturate(dot(normalize(Ray.Direction), CapsuleBounds.Axis));
-		t = CosTheta >= SqrtOneMinusX(CapsuleBounds.SinThetaMax2) ? length(LightPosition) : -1.0;
+		float CosTheta = dot(normalize(Ray.Direction), CapsuleBounds.ConeAxis);
+		t = CosTheta >= SqrtOneMinusX(CapsuleBounds.ConeSinThetaMax2) ? length(LightPosition) : -1.0;
 	}
 	else
 	{
