Many fixes for anisotropic lighting to play well with clear coat.
- Fixed GBufferPostDecode to correctly decode custom data according to shading model.
- Fixed Reflection Env PS to apply the correct normal and roughness on the top clear coat layer lighting.
- Fixed ClearCoatBxDF to not apply anisotropy onto the top layer, for its roughness to not affect the bottom layer and added safe ranges to avoid NaNs.
- Fixed SSRReflection to trace for the top layer for clear coat. It will slowly lerp towards tracing for the bottom layer according to the clear coat parameter.

diff --git a/Engine/Shaders/Private/GBufferHelpers.ush b/Engine/Shaders/Private/GBufferHelpers.ush
--- a/Engine/Shaders/Private/GBufferHelpers.ush
+++ b/Engine/Shaders/Private/GBufferHelpers.ush
@@ -5,6 +5,7 @@
 
 #include "SceneTexturesCommon.ush"
 #include "SceneTextureParameters.ush"
+#include "DeferredShadingCommon.ush"
 
 float SquareInline(float X)
 {
@@ -373,7 +374,8 @@ void AdjustBaseColorAndSpecularColorForSubsurfaceProfileLightingCopyHack(inout f
 
 void GBufferPostDecode(inout FGBufferData Ret, bool bChecker, bool bGetNormalizedNormal)
 {
-	Ret.CustomData = !(Ret.SelectiveOutputMask & 0x1) ? Ret.CustomData : 0.0f;
+	Ret.CustomData = HasCustomGBufferData(Ret.ShadingModelID) ? Ret.CustomData : 0.0f;
+
 	Ret.PrecomputedShadowFactors = !(Ret.SelectiveOutputMask & 0x2) ? Ret.PrecomputedShadowFactors : ((Ret.SelectiveOutputMask & 0x4) ? 0.0f : 1.0f);
 	Ret.Velocity = !(Ret.SelectiveOutputMask & 0x8) ? Ret.Velocity : 0.0f;
 	bool bHasAnisotropy = (Ret.SelectiveOutputMask & 0x1);
diff --git a/Engine/Shaders/Private/ReflectionEnvironmentPixelShader.usf b/Engine/Shaders/Private/ReflectionEnvironmentPixelShader.usf
--- a/Engine/Shaders/Private/ReflectionEnvironmentPixelShader.usf
+++ b/Engine/Shaders/Private/ReflectionEnvironmentPixelShader.usf
@@ -215,6 +215,8 @@ float3 ReflectionEnvironment(FGBufferData GBuffer, float AmbientOcclusion, float
 
 	float3 N = GBuffer.WorldNormal;
 
+	const float3 SavedTopLayerNormal = N;
+
 #if SUPPORTS_ANISOTROPIC_MATERIALS
 	ModifyGGXAnisotropicNormalRoughness(GBuffer.WorldTangent, GBuffer.Anisotropy, GBuffer.Roughness, N, V);
 #endif
@@ -279,8 +281,16 @@ float3 ReflectionEnvironment(FGBufferData GBuffer, float AmbientOcclusion, float
 		const float ClearCoat			= GBuffer.CustomData.x;
 		const float ClearCoatRoughness	= GBuffer.CustomData.y;
 
+		// Restore saved values needed for the top layer.
+		GBuffer.WorldNormal = SavedTopLayerNormal;
+		// Recompute some values unaffected by anistropy for the top layer
+		N = GBuffer.WorldNormal;
+		R = 2 * dot(V, N) * N - V;
+		NoV = saturate(dot(N, V));
+		R = GetOffSpecularPeakReflectionDir(N, R, ClearCoatRoughness);
+
 		// TODO EnvBRDF should have a mask param
-		float2 AB = PreIntegratedGF.SampleLevel( PreIntegratedGFSampler, float2( NoV, GBuffer.Roughness ), 0 ).rg;
+		float2 AB = PreIntegratedGF.SampleLevel( PreIntegratedGFSampler, float2( NoV, ClearCoatRoughness), 0 ).rg;
 		Color.rgb *= SpecularColor * AB.x + AB.y * saturate( 50 * SpecularColor.g ) * (1 - ClearCoat);
 		
 		// F_Schlick
diff --git a/Engine/Shaders/Private/SSRT/SSRTReflections.usf b/Engine/Shaders/Private/SSRT/SSRTReflections.usf
--- a/Engine/Shaders/Private/SSRT/SSRTReflections.usf
+++ b/Engine/Shaders/Private/SSRT/SSRTReflections.usf
@@ -183,7 +183,13 @@ void ScreenSpaceReflections(
 	const float3 V = normalize(View.TranslatedWorldCameraOrigin - PositionTranslatedWorld);
 
 #if SUPPORTS_ANISOTROPIC_MATERIALS
-	ModifyGGXAnisotropicNormalRoughness(GBuffer.WorldTangent, GBuffer.Anisotropy, GBuffer.Roughness, N, V);
+	// We do not need to execute ModifyGGXAnisotropicNormalRoughness in order to handle anisotropy on clear coat material.
+	// This is because SSR is only tracing the top layer which we consider to not have any anisotropy. This is when ClearCoat is 1.
+	// When clear coat decreases, we smoothly fades in the anisotropy influence from the bottom layer.
+	const float ClearCoat = GBuffer.CustomData.x;
+	const float AnisotropyBlendValue = GBuffer.ShadingModelID == SHADINGMODELID_CLEAR_COAT ? ClearCoat : 0.0f;
+	const float SSRAnisotropy = lerp(GBuffer.Anisotropy, 0.0f, AnisotropyBlendValue);
+	ModifyGGXAnisotropicNormalRoughness(GBuffer.WorldTangent, SSRAnisotropy, GBuffer.Roughness, N, V);
 #endif
 	
 	float Roughness = GetRoughness(GBuffer);
diff --git a/Engine/Shaders/Private/ShadingModels.ush b/Engine/Shaders/Private/ShadingModels.ush
--- a/Engine/Shaders/Private/ShadingModels.ush
+++ b/Engine/Shaders/Private/ShadingModels.ush
@@ -289,8 +289,8 @@ BxDFContext RefractClearCoatContext(BxDFContext Context)
 	float Eta = 1.0 / 1.5;
 	float RefractionBlendFactor = RefractBlendClearCoatApprox(Context.VoH);
 	float RefractionProjectionTerm = RefractionBlendFactor * Context.NoH;
-	RefractedContext.NoV = saturate(Eta * Context.NoV - RefractionProjectionTerm);
-	RefractedContext.NoL = saturate(Eta * Context.NoL - RefractionProjectionTerm);
+	RefractedContext.NoV = clamp(Eta * Context.NoV - RefractionProjectionTerm, 0.001, 1.0); // Due to CalcThinTransmission and Vis_SmithJointAniso, we need to make sure
+	RefractedContext.NoL = clamp(Eta * Context.NoL - RefractionProjectionTerm, 0.001, 1.0); // those values are not 0s to avoid NaNs.
 	RefractedContext.VoH = saturate(Eta * Context.VoH - RefractionBlendFactor);
 	RefractedContext.VoL = 2.0 * RefractedContext.VoH * RefractedContext.VoH - 1.0;
 	RefractedContext.NoH = Context.NoH;
@@ -327,17 +327,12 @@ FDirectLighting ClearCoatBxDF( FGBufferData GBuffer, half3 N, half3 V, half3 L,
 	half3 X = 0;
 	half3 Y = 0;
 
-	BRANCH
-	if (bHasAnisotropy)
-	{
-		X = GBuffer.WorldTangent;
-		Y = normalize(cross(N, X));
-		Init(Context, Nspec, X, Y, V, L);
-	}
-	else
-	{
-		Init(Context, Nspec, V, L);
-	}
+	//////////////////////////////
+	/// Top Layer
+	//////////////////////////////
+
+	// No anisotropy for the top layer
+	Init(Context, Nspec, V, L);
 	
 	// Modify SphereSinAlpha, knowing that it was previously manipulated by roughness of the under coat
 	// Note: the operation is not invertible for GBuffer.Roughness = 1.0, so roughness is clamped to 254.0/255.0
@@ -371,13 +366,21 @@ FDirectLighting ClearCoatBxDF( FGBufferData GBuffer, half3 N, half3 V, half3 L,
 		float Fr1 = D * Vis * F;
 		Lighting.Specular = ClearCoat * AreaLight.FalloffColor * (Falloff * NoL * Fr1);
 	}
+
+	// Restore previously changed SphereSinAlpha for the top layer. 
+	// Alpha needs to also be restored to the bottom layer roughness.
 	AreaLight.SphereSinAlpha = SphereSinAlpha;
+	Alpha = Pow2(GBuffer.Roughness);
 
 	// Incoming and exiting Fresnel terms are identical to incoming Fresnel term (VoH == HoL)
 	// float FresnelCoeff = (1.0 - F1) * (1.0 - F2);
 	float FresnelCoeff = 1.0 - F;
 	FresnelCoeff *= FresnelCoeff;
 
+	//////////////////////////////
+	/// Bottom Layer
+	//////////////////////////////
+
 	if (CLEAR_COAT_BOTTOM_NORMAL)
 	{
 		BxDFContext TempContext;
@@ -425,6 +428,13 @@ FDirectLighting ClearCoatBxDF( FGBufferData GBuffer, half3 N, half3 V, half3 L,
 	//   V2 = -V2;
 	//   BxDFContext BottomContext;
 	//   Init(BottomContext, N, X, Y, V2, L2);
+	if (bHasAnisotropy)
+	{
+		// Prepare the anisotropic Context to refract.
+		X = GBuffer.WorldTangent;
+		Y = normalize(cross(N, X));
+		Init(Context, Nspec, X, Y, V, L);
+	}
 	BxDFContext BottomContext = RefractClearCoatContext(Context);
 	BottomContext.VoH = AreaLight.bIsRect ? BottomContext.NoV : BottomContext.VoH;
 
