DF shadow fix for mobile es31 without independent samplers
Moved most of PCD3D_ES3_1 settings from code to DataDrivenPlatformInfo

diff --git a/Engine/Config/Windows/DataDrivenPlatformInfo.ini b/Engine/Config/Windows/DataDrivenPlatformInfo.ini
--- a/Engine/Config/Windows/DataDrivenPlatformInfo.ini
+++ b/Engine/Config/Windows/DataDrivenPlatformInfo.ini
@@ -180,3 +180,11 @@ bSupportsPerPixelDBufferMask = false
 
 bIsHlslcc = false
 NumberOfComputeThreads = 64
+
+[ShaderPlatform PCD3D_ES3_1]
+Language=D3D
+MaxFeatureLevel=ES3_1
+bIsPC=true
+bIsMobile=true
+bSupportsDistanceFields = true
+bSupportsByteBufferComputeShaders = true
\ No newline at end of file
diff --git a/Engine/Shaders/Private/DistanceFieldLightingShared.ush b/Engine/Shaders/Private/DistanceFieldLightingShared.ush
--- a/Engine/Shaders/Private/DistanceFieldLightingShared.ush
+++ b/Engine/Shaders/Private/DistanceFieldLightingShared.ush
@@ -141,6 +141,7 @@ struct FDFObjectBounds
 
 uint NumSceneObjects;
 StructuredBuffer<float4> SceneObjectBounds;
+SamplerState DistanceFieldSampler;
 
 FDFObjectBounds LoadDFObjectBounds(uint ObjectIndex)
 {
@@ -403,6 +404,15 @@ float GetVPLOcclusion(float3 BentNormalAO, float3 NormalizedVectorToVPL, float V
 	return Visibility;
 }
 
+float SampleDistanceFieldBrickTexture(float3 UV)
+{
+#if SUPPORTS_INDEPENDENT_SAMPLERS
+	return Texture3DSampleLevel(DistanceFieldBrickTexture, GlobalBilinearWrappedSampler, UV, 0).x;
+#else
+	return Texture3DSampleLevel(DistanceFieldBrickTexture, DistanceFieldSampler, UV, 0).x;
+#endif
+}
+
 // Returns a Volume space signed distance to the mesh's surface
 // This distance is clamped to the band of the current mip
 float SampleSparseMeshSignedDistanceField(float3 SampleVolumePosition, FDFAssetData DFAssetData)
@@ -431,7 +441,8 @@ float SampleSparseMeshSignedDistanceField(float3 SampleVolumePosition, FDFAssetD
 			BrickIndex >> (DistanceFieldBrickAtlasSizeLog2.x + DistanceFieldBrickAtlasSizeLog2.y));
 
 		float3 AtlasUV = (BrickOffset * DistanceFieldBrickSize + BrickLocalUV) * DistanceFieldBrickAtlasTexelSize;
-		float EncodedDistanceField = Texture3DSampleLevel(DistanceFieldBrickTexture, GlobalBilinearWrappedSampler, AtlasUV, 0).x;
+		float EncodedDistanceField = SampleDistanceFieldBrickTexture(AtlasUV);
+
 		DistanceField = EncodedDistanceField * DFAssetData.DistanceFieldToVolumeScaleBias.x + DFAssetData.DistanceFieldToVolumeScaleBias.y;
 	}
 
@@ -499,13 +510,21 @@ Texture2D HFVisibilityTexture;
 
 float SampleHeightFieldAtlas(float2 UV)
 {
+#if SUPPORTS_INDEPENDENT_SAMPLERS
 	float4 SampleValue = Texture2DSampleLevel(HeightFieldTexture, GlobalBilinearWrappedSampler, UV, 0);
+#else
+	float4 SampleValue = Texture2DSampleLevel(HeightFieldTexture, DistanceFieldSampler, UV, 0);
+#endif
 	return DecodePackedHeight(SampleValue.xy);
 }
 
 float SampleHFVisibilityTexture(float2 UV)
 {
+#if SUPPORTS_INDEPENDENT_SAMPLERS
 	return Texture2DSampleLevel(HFVisibilityTexture, GlobalBilinearWrappedSampler, UV, 0).r;
+#else
+	return Texture2DSampleLevel(HFVisibilityTexture, DistanceFieldSampler, UV, 0).r;
+#endif
 }
 
 float2 HeightFieldAtlasTexelSize;
\ No newline at end of file
diff --git a/Engine/Source/Runtime/RHI/Public/RHIDefinitions.h b/Engine/Source/Runtime/RHI/Public/RHIDefinitions.h
--- a/Engine/Source/Runtime/RHI/Public/RHIDefinitions.h
+++ b/Engine/Source/Runtime/RHI/Public/RHIDefinitions.h
@@ -1612,7 +1612,7 @@ enum class EAsyncComputeBudget
 
 inline bool IsPCPlatform(const FStaticShaderPlatform Platform)
 {
-	return Platform == SP_PCD3D_SM5 || Platform == SP_PCD3D_ES3_1 ||
+	return Platform == SP_PCD3D_SM5 ||
 		Platform == SP_OPENGL_PCES3_1 ||
 		Platform == SP_METAL_SM5 ||
 		Platform == SP_VULKAN_PCES3_1 || Platform == SP_VULKAN_SM5 || Platform == SP_METAL_MACES3_1 || Platform == SP_METAL_MRT_MAC 
@@ -1624,7 +1624,6 @@ inline bool IsMobilePlatform(const EShaderPlatform Platform)
 {
 	return 
 		Platform == SP_METAL || Platform == SP_METAL_MACES3_1 || Platform == SP_METAL_TVOS
-		|| Platform == SP_PCD3D_ES3_1
 		|| Platform == SP_OPENGL_PCES3_1 || Platform == SP_OPENGL_ES3_1_ANDROID
 		|| Platform == SP_VULKAN_ES3_1_ANDROID || Platform == SP_VULKAN_PCES3_1 || Platform == SP_VULKAN_ES3_1_LUMIN
 		|| Platform == SP_SWITCH_FORWARD
@@ -1709,7 +1708,6 @@ inline bool IsD3DPlatform(const FStaticShaderPlatform Platform)
 	switch (Platform)
 	{
 	case SP_PCD3D_SM5:
-	case SP_PCD3D_ES3_1:
 PRAGMA_DISABLE_DEPRECATION_WARNINGS
 	case SP_XBOXONE_D3D12_REMOVED:
 PRAGMA_ENABLE_DEPRECATION_WARNINGS
@@ -1744,7 +1742,6 @@ inline FStaticFeatureLevel GetMaxSupportedFeatureLevel(const FStaticShaderPlatfo
 	case SP_METAL:
 	case SP_METAL_TVOS:
 	case SP_METAL_MACES3_1:
-	case SP_PCD3D_ES3_1:
 	case SP_OPENGL_PCES3_1:
 	case SP_VULKAN_PCES3_1:
 	case SP_VULKAN_ES3_1_ANDROID:
diff --git a/Engine/Source/Runtime/RHI/Public/RHIUtilities.h b/Engine/Source/Runtime/RHI/Public/RHIUtilities.h
--- a/Engine/Source/Runtime/RHI/Public/RHIUtilities.h
+++ b/Engine/Source/Runtime/RHI/Public/RHIUtilities.h
@@ -437,7 +437,6 @@ struct FByteAddressBuffer
 	void Initialize(const TCHAR* InDebugName, uint32 InNumBytes, uint32 AdditionalUsage = 0)
 	{
 		NumBytes = InNumBytes;
-		check(GMaxRHIFeatureLevel >= ERHIFeatureLevel::SM5);
 		check( NumBytes % 4 == 0 );
 		FRHIResourceCreateInfo CreateInfo(InDebugName);
 		Buffer = RHICreateStructuredBuffer(4, NumBytes, BUF_ShaderResource | BUF_ByteAddressBuffer | AdditionalUsage, ERHIAccess::SRVMask, CreateInfo);
diff --git a/Engine/Source/Runtime/RenderCore/Private/RenderUtils.cpp b/Engine/Source/Runtime/RenderCore/Private/RenderUtils.cpp
--- a/Engine/Source/Runtime/RenderCore/Private/RenderUtils.cpp
+++ b/Engine/Source/Runtime/RenderCore/Private/RenderUtils.cpp
@@ -1160,7 +1160,7 @@ RENDERCORE_API bool UseMobileAmbientOcclusion(const FStaticShaderPlatform Platfo
 
 RENDERCORE_API bool IsMobileDistanceFieldEnabled(const FStaticShaderPlatform Platform)
 {
-	return IsMobilePlatform(Platform) && IsSwitchPlatform(Platform) && IsUsingDistanceFields(Platform);
+	return IsMobilePlatform(Platform) && (IsSwitchPlatform(Platform) || IsD3DPlatform(Platform)) && IsUsingDistanceFields(Platform);
 }
 
 RENDERCORE_API bool IsMobileDistanceFieldShadowingEnabled(const FStaticShaderPlatform Platform)
