Fix invalid page occupancy initialization.

diff --git a/Engine/Shaders/Private/HairStrands/HairStrandsEnvironmentLighting.usf b/Engine/Shaders/Private/HairStrands/HairStrandsEnvironmentLighting.usf
--- a/Engine/Shaders/Private/HairStrands/HairStrandsEnvironmentLighting.usf
+++ b/Engine/Shaders/Private/HairStrands/HairStrandsEnvironmentLighting.usf
@@ -184,7 +184,7 @@ FHairTraversalResult InternalComputeHairCountVoxel(
 	TraversalSettings.bDebugEnabled = bDebugEnabled;
 	TraversalSettings.SteppingScale = VirtualVoxel.SteppingScale_Environment;
 	TraversalSettings.PixelRadius = PixelRadius;
-	TraversalSettings.bUseOccupancy = View.GeneralPurposeTweak > 0;
+	TraversalSettings.bUseOccupancy = false;
 
 	return ComputeHairCountVirtualVoxel(
 		WorldPosition,
