Add first pass of stochastic compute rasterization.

diff --git a/Engine/Shaders/Private/HairStrands/HairStrandsVisibilityRasterCompute.usf b/Engine/Shaders/Private/HairStrands/HairStrandsVisibilityRasterCompute.usf
--- a/Engine/Shaders/Private/HairStrands/HairStrandsVisibilityRasterCompute.usf
+++ b/Engine/Shaders/Private/HairStrands/HairStrandsVisibilityRasterCompute.usf
@@ -1,6 +1,7 @@
 // Copyright Epic Games, Inc. All Rights Reserved.
 
 #include "../Common.ush"
+#include "../MonteCarlo.ush"
 #include "HairStrandsVertexFactoryCommon.ush"
 #include "HairStrandsVisibilityCommon.ush"
 
@@ -85,7 +86,7 @@ float4x4		HairStrandsVF_LocalToWorldPrimitiveTransform;
 Buffer<uint4>	HairStrandsVF_PositionBuffer;
 Buffer<float4>	HairStrandsVF_PositionOffsetBuffer;
 
-#if PERMUTATION_CULLING == 1
+#if PERMUTATION_CULLING
 Buffer<uint>	HairStrandsVF_CullingIndirectBuffer;
 Buffer<uint>	HairStrandsVF_CullingIndexBuffer;
 Buffer<float>	HairStrandsVF_CullingRadiusScaleBuffer;
@@ -155,7 +156,7 @@ bool IsInScreen(int2 C)
 	return C.x >= 0 && C.x < OutputResolution.x && C.y >= 0 && C.y < OutputResolution.y;
 }
 
-void Rasterize(int2 Coords, float Ratio, float Depth, uint PrimitiveId, float Coverage)
+void InternalRasterize(int2 Coords, float Ratio, float Depth, uint PrimitiveId, float Coverage)
 {
 	if (IsInScreen(Coords))
 	{
@@ -225,6 +226,50 @@ FRasterInfo UnpackRasterInfo(uint4 In)
 	return (FRasterInfo)0;
 }
 
+float2 StochasticRasterRandom(uint2 PixelPosition)
+{
+	const uint InSeed = View.StateFrameIndexMod8;
+	const uint2 Seed0 = Rand3DPCG16(int3(PixelPosition, InSeed)).xy;
+	return Hammersley16(InSeed, 8, Seed0);
+}
+
+// Rasterize a given pixel with anti-aliasing. 
+// This can be done by writing the two adjacents pixel covered by the segments, or stochastically, by picking one of the two pixels randomly based on coverage
+uint2 Rasterize(uint2 PixelCoord, bool bIsSteep, float FracY, float RatioScale, float Coverage, float DeviceZ, uint ProvokingVertexId, bool bStochastic=true, uint2 SeedCoord=uint2(0,0))
+{
+#if PERMUTATION_STOCHASTIC
+	// Select randomly Bottom Left/Right
+	{
+		// No select based on the exact frac amount, as this will oscillate too much. 
+		// The stochastic raster only acts on the visibility check. The final covarage is computed analytically
+		float U = StochasticRasterRandom(SeedCoord).x;
+		const bool bLower = bStochastic ? (U < 0.5 ? FracY < 0.5 : FracY > 0.5) : FracY < 0.5f;
+
+		uint2 Coord = PixelCoord + (bLower ? uint2(0, 0) : uint2(0, 1));
+		Coord = bIsSteep ? Coord.yx : Coord.xy;
+		const float Ratio = (bLower ? 1 - FracY : FracY) * RatioScale;;
+		InternalRasterize(Coord, Ratio * RatioScale, DeviceZ, ProvokingVertexId, Coverage);
+		SeedCoord = Coord;
+	}
+#else
+	const bool bLower = FracY < 0.5f;
+	const float Ratio = (bLower ? 1 - FracY : FracY) * RatioScale;
+
+	// Bottom/Left Pixel (or inverse)
+	{
+		uint2 Coord = PixelCoord + (bLower ? uint2(0, 0) : uint2(0, 1));
+		Coord = bIsSteep ? Coord.yx : Coord.xy;
+		InternalRasterize(Coord, Ratio, DeviceZ, ProvokingVertexId, Coverage);
+	}
+	// Top/Right Pixel (or inverse)
+	{
+		uint2 Coord = PixelCoord + (bLower ? uint2(0, 0) : uint2(0, 1));
+		Coord = bIsSteep ? Coord.yx : Coord.xy;
+		InternalRasterize(Coord, 1.f - Ratio, DeviceZ, ProvokingVertexId, Coverage);
+	}
+#endif
+	return SeedCoord;
+}
 //////////////////////////////////////////////////////////////////////////////////////////
 
 // These number need to match HairStrandsClusterCulling.usf
@@ -234,7 +279,7 @@ FRasterInfo UnpackRasterInfo(uint4 In)
 [numthreads(GROUP_SIZE, 1, 1)]
 void MainCS(uint3 DispatchThreadID : SV_DispatchThreadID)
 {
-#if PERMUTATION_CULLING == 1
+#if PERMUTATION_CULLING
 	uint VertexIndex0 = DispatchThreadID.x + DispatchThreadID.y * GROUP_SIZE * GROUP_COUNT_X;
 #else
 	uint VertexIndex0 = DispatchThreadID.x + DispatchThreadID.y * GROUP_SIZE * DispatchCountX;
@@ -248,7 +293,7 @@ void MainCS(uint3 DispatchThreadID : SV_DispatchThreadID)
 
 	float LodRadiusScale0 = 1;
 	float LodRadiusScale1 = 1;
-#if PERMUTATION_CULLING == 1
+#if PERMUTATION_CULLING
 	if (HairStrandsVF_bIsCullingEnable)
 	{
 		const uint VertexCountAfterCulling = HairStrandsVF_CullingIndirectBuffer[3];
@@ -364,11 +409,7 @@ void MainCS(uint3 DispatchThreadID : SV_DispatchThreadID)
 		PX0.y = floor(End.y);
 
 		const float FracY = frac(End.y);
-		const bool bLower = FracY < 0.5f;
-		int2 Coord = PX0 + (bLower ? int2(0, 0) : int2(0, 1));
-		Coord = bIsSteep ? Coord.yx : Coord.xy;
-		const float Ratio = bLower ? 1 - FracY : FracY;
-		Rasterize(Coord, Ratio * XGap, SP0.z, VertexIndex0, Coverage0);
+		Rasterize(PX0, bIsSteep, FracY, XGap, Coverage0, SP0.z, VertexIndex0, false);
 
 		DeltaY = End.y + GradientY; // first y-intersection for the main loop
 		DeltaZ = SP1.z + GradientZ; // This is not perspective correct. Interpolate linear depth instead
@@ -387,35 +428,28 @@ void MainCS(uint3 DispatchThreadID : SV_DispatchThreadID)
 
 		// In order to avoid doubled rasterizated point at segments connection we exclude 
 		// the last pixel apart for pixel placed at the end of a curve
-		if (CP1.Type == HAIR_CONTROLPOINT_END)
+		//if (CP1.Type == HAIR_CONTROLPOINT_END)
 		{
 			const float FracY = frac(End.y);
-			const bool bLower = FracY < 0.5f;
-			int2 Coord = PX1 + (bLower ? int2(0, 0) : int2(0, 1));
-			Coord = bIsSteep ? Coord.yx : Coord.xy;
-			const float Ratio = bLower ? 1 - FracY : FracY;
-			Rasterize(Coord, Ratio * XGap, SP1.z, VertexIndex0, max(Coverage1,0.01f));
+			Rasterize(PX1, bIsSteep, FracY, XGap, max(Coverage1, 0.01f), SP1.z, VertexIndex0, false);
 		}
 	}
 
 	// Inner points
 	const uint MaxIteration = MaxRasterCount;
 	const float PixelCount = abs(PX1.x - (PX0.x + 1));
+	uint2 SeedCoord = PX0;
 	for (int i = PX0.x + 1, j = 0; i < PX1.x && j < MaxIteration; ++i, ++j)
 	{
 		const float S = float(j) / PixelCount;
 		const float FracY = frac(DeltaY);
-		const bool bLower = FracY < 0.5;
-		const float Ratio = bLower ? 1 - FracY : FracY;
-		const float Coverage = lerp(Coverage0, Coverage1, S);
-
 		const float DC_Z = lerp(NDC0.z, NDC1.z, S);
 		const float DC_W = lerp(NDC0.w, NDC1.w, S);
 		const float DeviceZ = DC_Z / DC_W;
+		
+		const uint2 PixelCoord = uint2(int2(i, floor(DeltaY)));
 
-		int2 Coord = int2(i, floor(DeltaY)) + (bLower ? int2(0, 0) : int2(0, 1));
-		Coord = bIsSteep ? Coord.yx : Coord.xy;
-		Rasterize(Coord, Ratio, DeviceZ /*DeltaZ*/, VertexIndex0, Coverage);
+		SeedCoord = Rasterize(PixelCoord, bIsSteep, FracY, 1.0f, lerp(Coverage0, Coverage1, S), DeviceZ, VertexIndex0, true, SeedCoord);
 
 		DeltaY += GradientY;
 		DeltaZ += GradientZ;
diff --git a/Engine/Source/Runtime/Renderer/Private/HairStrands/HairStrandsVisibility.cpp b/Engine/Source/Runtime/Renderer/Private/HairStrands/HairStrandsVisibility.cpp
--- a/Engine/Source/Runtime/Renderer/Private/HairStrands/HairStrandsVisibility.cpp
+++ b/Engine/Source/Runtime/Renderer/Private/HairStrands/HairStrandsVisibility.cpp
@@ -61,9 +61,11 @@ static int32 GHairVisibilityComputeRaster = 0;
 static int32 GHairVisibilityComputeRaster_MaxSamplePerPixel = 1;
 static float GHairVisibilityComputeRaster_MeanSamplePerPixel = 1;
 static int32 GHairVisibilityComputeRaster_MaxPixelCount = 64;
+static int32 GHairVisibilityComputeRaster_Stochastic = 0;
 static FAutoConsoleVariableRef CVarHairStrandsVisibilityComputeRaster(TEXT("r.HairStrands.Visibility.ComputeRaster"), GHairVisibilityComputeRaster, TEXT("Hair Visiblity uses raster compute."), ECVF_Scalability | ECVF_RenderThreadSafe);
 static FAutoConsoleVariableRef CVarHairStrandsVisibilityComputeRaster_MaxSamplePerPixel(TEXT("r.HairStrands.Visibility.ComputeRaster.SamplePerPixel"), GHairVisibilityComputeRaster_MaxSamplePerPixel, TEXT("Define the number of sampler per pixel using raster compute."));
 static FAutoConsoleVariableRef CVarHairStrandsVisibilityComputeRaster_MaxPixelCount(TEXT("r.HairStrands.Visibility.ComputeRaster.MaxPixelCount"), GHairVisibilityComputeRaster_MaxPixelCount, TEXT("Define the maximal length rasterize in compute."));
+static FAutoConsoleVariableRef CVarHairVisibilityComputeRaster_Stochastic(TEXT("r.HairStrands.Visibility.ComputeRaster.Stochastic"), GHairVisibilityComputeRaster_Stochastic, TEXT("Enable stochastic compute rasterization (faster, but more prone to aliasting). Experimental."));
 
 static float GHairStrandsFullCoverageThreshold = 0.98f;
 static FAutoConsoleVariableRef CVarHairStrandsFullCoverageThreshold(TEXT("r.HairStrands.Visibility.FullCoverageThreshold"), GHairStrandsFullCoverageThreshold, TEXT("Define the coverage threshold at which a pixel is considered fully covered."));
@@ -2876,8 +2878,9 @@ class FVisiblityRasterComputeCS : public FGlobalShader
 
 	class FRasterAtomic : SHADER_PERMUTATION_INT("PERMUTATION_RASTER_ATOMIC", 4);
 	class FSPP : SHADER_PERMUTATION_SPARSE_INT("PERMUTATION_SPP", 1, 2, 4); 
-	class FCulling : SHADER_PERMUTATION_INT("PERMUTATION_CULLING", 2);
-	using FPermutationDomain = TShaderPermutationDomain<FRasterAtomic, FSPP, FCulling>;
+	class FCulling : SHADER_PERMUTATION_BOOL("PERMUTATION_CULLING");
+	class FStochastic : SHADER_PERMUTATION_BOOL("PERMUTATION_STOCHASTIC");
+	using FPermutationDomain = TShaderPermutationDomain<FRasterAtomic, FSPP, FCulling, FStochastic>;
 
 	BEGIN_SHADER_PARAMETER_STRUCT(FParameters, )
 		SHADER_PARAMETER(uint32, MacroGroupId)
@@ -3014,6 +3017,7 @@ static FRasterComputeOutput AddVisibilityComputeRasterPass(
 	}
 
 	// Create and set the uniform buffer
+	const bool bStochasticRaster = GHairVisibilityComputeRaster_Stochastic > 0;
 	const bool bEnableMSAA = false;
 	TUniformBufferRef<FViewUniformShaderParameters> ViewUniformShaderParameters;
 	SetUpViewHairRenderInfo(ViewInfo, bEnableMSAA, ViewInfo.CachedViewUniformShaderParameters->HairRenderInfo, ViewInfo.CachedViewUniformShaderParameters->HairRenderInfoBits, ViewInfo.CachedViewUniformShaderParameters->HairComponents);
@@ -3040,11 +3044,12 @@ static FRasterComputeOutput AddVisibilityComputeRasterPass(
 		PermutationVector0.Set<FVisiblityRasterComputeCS::FRasterAtomic>(0);
 	}
 #endif
+	PermutationVector0.Set<FVisiblityRasterComputeCS::FStochastic>(bStochasticRaster);
 	PermutationVector0.Set<FVisiblityRasterComputeCS::FSPP>(SamplePerPixelCount);
 	PermutationVector1 = PermutationVector0; 
 
-	PermutationVector0.Set<FVisiblityRasterComputeCS::FCulling>(0);
-	PermutationVector1.Set<FVisiblityRasterComputeCS::FCulling>(1);
+	PermutationVector0.Set<FVisiblityRasterComputeCS::FCulling>(false);
+	PermutationVector1.Set<FVisiblityRasterComputeCS::FCulling>(true);
 	TShaderMapRef<FVisiblityRasterComputeCS> ComputeShader_CullingOff(ViewInfo.ShaderMap, PermutationVector0);
 	TShaderMapRef<FVisiblityRasterComputeCS> ComputeShader_CullingOn (ViewInfo.ShaderMap, PermutationVector1);
 
