Optimises VALU in TAA's UpdateHistory pass

diff --git a/Engine/Shaders/Private/TemporalAA/TAACommon.ush b/Engine/Shaders/Private/TemporalAA/TAACommon.ush
--- a/Engine/Shaders/Private/TemporalAA/TAACommon.ush
+++ b/Engine/Shaders/Private/TemporalAA/TAACommon.ush
@@ -327,6 +327,64 @@ taa_short2 ClampPixelOffset(
 	return ClampedOffset;
 }
 
+taa_short2 ClampPixelOffset(
+	taa_short2 SamplePixelPos,
+	int2 MinPixelPos, int2 MaxPixelPos)
+{
+	SamplePixelPos.x = min(SamplePixelPos.x, taa_short(MaxPixelPos.x));
+	SamplePixelPos.x = max(SamplePixelPos.x, taa_short(MinPixelPos.x));
+		
+	SamplePixelPos.y = min(SamplePixelPos.y, taa_short(MaxPixelPos.y));
+	SamplePixelPos.y = max(SamplePixelPos.y, taa_short(MinPixelPos.y));
+
+	return SamplePixelPos;
+}
+
+taa_short2 AddAndClampPixelOffset(
+	taa_short2 KernelCenterPixelPos,
+	taa_short2 Offset,
+	const taa_short2 OffsetDirection,
+	int2 MinPixelPos, int2 MaxPixelPos)
+{
+	taa_short2 SamplePixelPos = KernelCenterPixelPos + Offset;
+
+	// Only do clamp based on the compile time known direction of the offset.
+	// This turns only 1 v_pk_max_i16 and v_pk_min_i16 for a 3x3 kernel.
+	if (OffsetDirection.x > 0)
+		SamplePixelPos.x = min(SamplePixelPos.x, taa_short(MaxPixelPos.x));
+	else if (OffsetDirection.x < 0)
+		SamplePixelPos.x = max(SamplePixelPos.x, taa_short(MinPixelPos.x));
+		
+	if (OffsetDirection.y > 0)
+		SamplePixelPos.y = min(SamplePixelPos.y, taa_short(MaxPixelPos.y));
+	else if (OffsetDirection.y < 0)
+		SamplePixelPos.y = max(SamplePixelPos.y, taa_short(MinPixelPos.y));
+
+	return SamplePixelPos;
+}
+
+taa_short2 InvalidateOutputPixelPos(taa_short2 PixelPos, uint2 ViewportMax)
+#if 1
+{
+	taa_short2 Subtract = taa_short2(ViewportMax - 1) - PixelPos;
+	taa_ushort Override = taa_ushort(Subtract.x | Subtract.y);
+
+	#if CONFIG_COMPILE_FP16
+		PixelPos.x |= -taa_short((Override & uint16_t(0x8000)) >> 15);
+	#else
+		PixelPos.x |= -taa_short((Override & uint(0x80000000)) >> 31);
+	#endif
+
+	return PixelPos;
+}
+#else
+{
+	bool bIsValidPixel = all(PixelPos < ViewportMax);
+	PixelPos.x = bIsValidPixel ? PixelPos.x : ~taa_short(0);
+	return PixelPos;
+}
+#endif
+
 taa_ushort2 Map8x8Tile2x2Lane(uint GroupThreadIndex)
 {
 	taa_ushort T = taa_ushort(GroupThreadIndex);
@@ -443,7 +501,8 @@ taa_half ComputeSampleWeigth(taa_half UpscaleFactor, taa_half2 PixelDelta, const
 
 	// 1 - 1.9 * x^2 + 0.9 * x^4
 	taa_half x2 = saturate(u2 * dot(PixelDelta, PixelDelta));
-	return taa_half(((float(0.9) + MinimalContribution) * x2 - float(1.9)) * x2 + float(1.0));
+	//return taa_half(((float(0.9) + MinimalContribution) * x2 - float(1.9)) * x2 + float(1.0));
+	return saturate((taa_half(0.9) * x2 - taa_half(1.9)) * x2 + taa_half(1.0 + MinimalContribution));
 }
 
 taa_half2 WeightedLerpFactors(taa_half WeightA, taa_half WeightB, taa_half Blend)
diff --git a/Engine/Shaders/Private/TemporalAA/TAAUpdateHistory.usf b/Engine/Shaders/Private/TemporalAA/TAAUpdateHistory.usf
--- a/Engine/Shaders/Private/TemporalAA/TAAUpdateHistory.usf
+++ b/Engine/Shaders/Private/TemporalAA/TAAUpdateHistory.usf
@@ -83,11 +83,10 @@ void MainCS(
 
 	float4 Debug = 0.0;
 
-	uint2 DispatchThreadId = (
-		ZOrder2D(GroupThreadIndex, log2(TILE_SIZE)) +
-		GroupId * uint2(TILE_SIZE, TILE_SIZE));
+	taa_short2 HistoryPixelPos = (
+		taa_short2(GroupId) * taa_short2(TILE_SIZE, TILE_SIZE) +
+		Map8x8Tile2x2Lane(GroupThreadIndex));
 
-	uint2 HistoryPixelPos = DispatchThreadId;
 	float2 ViewportUV = (float2(HistoryPixelPos) + 0.5f) * HistoryInfo_ViewportSizeInverse;
 	float2 ScreenPos = ViewportUVToScreenPos(ViewportUV);
 	
@@ -97,24 +96,27 @@ void MainCS(
 	// Pixel coordinate of the center of the nearest input pixel K.
 	float2 PPCk = floor(PPCo) + 0.5;
 	
-	uint2 InputPixelPos = InputInfo_ViewportMin + uint2(PPCk);
+	taa_short2 InputPixelPos = ClampPixelOffset(
+		taa_short2(InputPixelPosMin) + taa_short2(PPCo),
+		InputPixelPosMin, InputPixelPosMax);
 
 	// Fetch reprojections related informations.
 	float2 PrevScreenPos = ScreenPos;
 	taa_half ParallaxRejectionMask = taa_half(1.0);
-	taa_half LowFrequencyRejection;
+	taa_half LowFrequencyRejection = taa_half(1.0);
+	taa_half OutputPixelVelocity = taa_half(0.0);
 	#if 1
 	{
-		float2 RejectionBilinearBufferUV = clamp(
-			PPCo * InputInfo_ExtentInverse,
-			RejectionInfo_UVViewportBilinearMin,
-			RejectionInfo_UVViewportBilinearMax);
-		
 		float2 EncodedVelocity = DilatedVelocityTexture[InputPixelPos];
 		ParallaxRejectionMask = ParallaxRejectionMaskTexture[InputPixelPos];
-		PrevScreenPos = ScreenPos - DecodeVelocityFromTexture(float4(EncodedVelocity, 0.0, 0.0)).xy;
 
-		LowFrequencyRejection = HistoryRejectionTexture.SampleLevel(GlobalBilinearClampedSampler, RejectionBilinearBufferUV, 0);
+		float2 ScreenVelocity = DecodeVelocityFromTexture(float4(EncodedVelocity, 0.0, 0.0)).xy;
+
+		PrevScreenPos = ScreenPos - ScreenVelocity;
+		OutputPixelVelocity = taa_half(length(ScreenVelocity * HistoryInfo_ViewportSize));
+
+		taa_ushort2 RejectionPixelPos = (taa_ushort2(InputPixelPos) - taa_short2(InputPixelPosMin)) / 2;
+		LowFrequencyRejection = HistoryRejectionTexture[RejectionPixelPos];
 		
 		#if !CONFIG_CLAMP
 		{
@@ -125,9 +127,6 @@ void MainCS(
 	}
 	#endif
 
-	float2 ScreenVelocity = ScreenPos - PrevScreenPos;
-	taa_half OutputPixelVelocity = taa_half(length(ScreenVelocity * HistoryInfo_ViewportSize));
-	
 	// Fetch wehther the pixel is responsive AA or not.
 	bool bIsResponsiveAAPixel = false;
 	#if CONFIG_RESPONSIVE_STENCIL
@@ -177,31 +176,50 @@ void MainCS(
 		UNROLL_N(CONFIG_SAMPLES)
 		for (uint SampleId = 0; SampleId < CONFIG_SAMPLES; SampleId++)
 		{
+			taa_short2 SampleInputPixelPos;
 			taa_half2 PixelOffset;
 			
 			#if CONFIG_SAMPLES == 9
 			{
-				PixelOffset = taa_half2(kOffsets3x3[kSquareIndexes3x3[SampleId]]);
+				taa_short2 iPixelOffset = taa_short2(kOffsets3x3[kSquareIndexes3x3[SampleId]]);
+				PixelOffset = taa_half2(iPixelOffset);
+				
+				SampleInputPixelPos = AddAndClampPixelOffset(
+					InputPixelPos,
+					iPixelOffset, iPixelOffset,
+					InputPixelPosMin, InputPixelPosMax);
 			}
 			#elif CONFIG_SAMPLES == 5 || CONFIG_SAMPLES == 6
 			{
 				if (SampleId == 5)
 				{
-					PixelOffset = taa_half2(SignFastInt(float2(dKO))); // TODO: taa_half
+					taa_short2 iPixelOffset;
+					#if CONFIG_COMPILE_FP16
+						iPixelOffset = int16_t2(1, 1) - int16_t((asuint16(dKO) & uint16_t(0x8000)) >> uint16_t(14));
+						PixelOffset = asfloat16(asuint16(half(1.0)).xx | (asuint16(dKO) & uint16_t(0x8000)));
+					#else
+						iPixelOffset = SignFastInt(dKO);
+						PixelOffset = asfloat(asuint(1.0).xx | (asuint(dKO) & uint(0x80000000)));
+					#endif
+						
+					SampleInputPixelPos = ClampPixelOffset(InputPixelPos, InputPixelPosMin, InputPixelPosMax);
 				}
 				else
 				{
-					PixelOffset = taa_half2(kOffsets3x3[kPlusIndexes3x3[SampleId]]);
+					taa_short2 iPixelOffset = taa_short2(kOffsets3x3[kPlusIndexes3x3[SampleId]]);
+					PixelOffset = taa_half2(iPixelOffset);
+					
+					SampleInputPixelPos = AddAndClampPixelOffset(
+						InputPixelPos,
+						iPixelOffset, iPixelOffset,
+						InputPixelPosMin, InputPixelPosMax);
 				}
 			}
 			#else
 				#error Unknown sample count
 			#endif
 
-			float2 SampleInputBufferUV = (InputInfo_ViewportMin + PPCk + PixelOffset) * InputInfo_ExtentInverse;
-			SampleInputBufferUV = clamp(SampleInputBufferUV, InputInfo_UVViewportBilinearMin, InputInfo_UVViewportBilinearMax);
-
-			taa_half3 InputColor = InputSceneColorTexture.SampleLevel(GlobalPointClampedSampler, SampleInputBufferUV, 0);
+			taa_half3 InputColor = InputSceneColorTexture[SampleInputPixelPos];
 
 			taa_half2 dPP = PixelOffset - dKO;
 			taa_half SampleSpatialWeight = ComputeSampleWeigth(FinalInputToHistoryFactor, dPP, /* MinimalContribution = */ float(0.005));
@@ -211,16 +229,20 @@ void MainCS(
 			FilteredInputColor       += (SampleSpatialWeight * ToneWeight) * InputColor;
 			FilteredInputColorWeight += (SampleSpatialWeight * ToneWeight);
 
-			InputMinColor = min(InputMinColor, TransformColorForClampingBox(InputColor));
-			InputMaxColor = max(InputMaxColor, TransformColorForClampingBox(InputColor));
-
 			if (SampleId == 0)
 			{
 				ClosestInputLuma4 = Luma4(InputColor);
+				InputMinColor = TransformColorForClampingBox(InputColor);
+				InputMaxColor = TransformColorForClampingBox(InputColor);
+			}
+			else
+			{
+				InputMinColor = min(InputMinColor, TransformColorForClampingBox(InputColor));
+				InputMaxColor = max(InputMaxColor, TransformColorForClampingBox(InputColor));
 			}
 		}
 		
-		FilteredInputColor *= SafeRcp(FilteredInputColorWeight);
+		FilteredInputColor *= rcp(FilteredInputColorWeight);
 
 		InputPixelAlignement = ComputeSampleWeigth(InputToHistoryFactor, dKO, /* MinimalContribution = */ float(0.0));
 	}
@@ -268,6 +290,7 @@ void MainCS(
 			{
 				FCatmullRomSamples Samples = GetBicubic2DCatmullRomSamples(PrevHistoryBufferUV, PrevHistoryInfo_Extent, PrevHistoryInfo_ExtentInverse);
 
+				UNROLL
 				for (uint i = 0; i < Samples.Count; i++)
 				{
 					float2 SampleUV = clamp(Samples.UV[i], PrevHistoryInfo_UVViewportBilinearMin, PrevHistoryInfo_UVViewportBilinearMax);
@@ -276,8 +299,8 @@ void MainCS(
 					taa_half3 Sample1 = PrevHistory_Textures_1.SampleLevel(GlobalBilinearClampedSampler, SampleUV, 0);
 					taa_half2 Sample2 = PrevHistory_Textures_2.SampleLevel(GlobalBilinearClampedSampler, SampleUV, 0);
 
-					RawHistory1Min = min(RawHistory1Min, TransformColorForClampingBox(Sample1 * SafeRcp(Sample2.g)));
-					RawHistory1Max = max(RawHistory1Max, TransformColorForClampingBox(Sample1 * SafeRcp(Sample2.g)));
+					RawHistory1Min = min(RawHistory1Min, Sample1 * SafeRcp(Sample2.g));
+					RawHistory1Max = max(RawHistory1Max, Sample1 * SafeRcp(Sample2.g));
 
 					RawHistory0 += Sample0 * taa_half(Samples.Weight[i]);
 					RawHistory1 += Sample1 * taa_half(Samples.Weight[i]);
@@ -356,7 +379,7 @@ void MainCS(
 	{
 		taa_half3 PrevHighFrequencyYCoCg = TransformColorForClampingBox(PrevHistoryMoment1 * SafeRcp(PrevHistoryValidity));
 		taa_half3 PrevYCoCg = TransformColorForClampingBox(PrevFallbackColor);
-		taa_half3 ClampedPrevYCoCg = clamp(PrevYCoCg, PrevHistoryMommentMin, PrevHistoryMommentMax);
+		taa_half3 ClampedPrevYCoCg = TransformColorForClampingBox(clamp(PrevFallbackColor, PrevHistoryMommentMin, PrevHistoryMommentMax));
 
 		taa_half HighFrequencyRejection = MeasureRejectionFactor(
 			PrevYCoCg, ClampedPrevYCoCg,
@@ -707,16 +730,16 @@ void MainCS(
 		uint LocalGroupThreadIndex = GetGroupThreadIndex(GroupThreadIndex, GroupWaveIndex);
 
 		#if 1
-			uint2 LocalHistoryPixelPos = (
-				ZOrder2D(LocalGroupThreadIndex, log2(TILE_SIZE)) +
-				GroupId * uint2(TILE_SIZE, TILE_SIZE));
+			taa_short2 LocalHistoryPixelPos = (
+				taa_short2(GroupId) * taa_short2(TILE_SIZE, TILE_SIZE) +
+				Map8x8Tile2x2Lane(LocalGroupThreadIndex));
 		#else
-			uint2 LocalHistoryPixelPos = HistoryPixelPos;
+			taa_short2 LocalHistoryPixelPos = HistoryPixelPos;
 		#endif
 			
+		LocalHistoryPixelPos = InvalidateOutputPixelPos(LocalHistoryPixelPos, HistoryInfo_ViewportMax);
+
 		// Output final history.
-		bool bIsValidhistoryPixel = all(LocalHistoryPixelPos < HistoryInfo_ViewportMax);
-		if (bIsValidhistoryPixel)
 		{
 			// Stocastically round up or down using the hardware RWTexture2D truncation unit to take into precision
 			// loss due to pixel format encoding.
@@ -766,10 +789,7 @@ void MainCS(
 			OutputColor = -min(-OutputColor, taa_half(0.0));
 			OutputColor = min(OutputColor, taa_half(Max10BitsFloat));
 
-			if (bIsValidhistoryPixel)
-			{
-				SceneColorOutput[LocalHistoryPixelPos] = OutputColor;
-			}
+			SceneColorOutput[LocalHistoryPixelPos] = OutputColor;
 		}
 	}
 }
