Change raster compute transmittance.

Accumlate transmittance in log space to emulate transmittance product rather than accumulation.

diff --git a/Engine/Shaders/Private/HairStrands/HairStrandsCoverage.usf b/Engine/Shaders/Private/HairStrands/HairStrandsCoverage.usf
--- a/Engine/Shaders/Private/HairStrands/HairStrandsCoverage.usf
+++ b/Engine/Shaders/Private/HairStrands/HairStrandsCoverage.usf
@@ -48,8 +48,10 @@ void MainCS(uint3 DispatchThreadId : SV_DispatchThreadID)
 		return;
 	}
 
-	const uint PackedCoverage = HairCountTexture.Load(uint3(PixelCoord, 0));
-	const float AccCoverage = UnpackCoverageAndHairCount(PackedCoverage);
-	OutputTexture[PixelCoord] = saturate(1 - saturate(AccCoverage));
+	const float PrecisionScale = 1000.f;
+	const uint PackedHairViewTransmittance = HairCountTexture.Load(uint3(PixelCoord, 0));
+	const float LogViewTransmittance = -float(PackedHairViewTransmittance) / PrecisionScale;
+	const float ViewTransmittance = exp2(LogViewTransmittance);
+	OutputTexture[PixelCoord] = ViewTransmittance;
 }
 #endif
\ No newline at end of file
diff --git a/Engine/Shaders/Private/HairStrands/HairStrandsVisibilityRasterCompute.usf b/Engine/Shaders/Private/HairStrands/HairStrandsVisibilityRasterCompute.usf
--- a/Engine/Shaders/Private/HairStrands/HairStrandsVisibilityRasterCompute.usf
+++ b/Engine/Shaders/Private/HairStrands/HairStrandsVisibilityRasterCompute.usf
@@ -167,12 +167,16 @@ void InternalRasterize(int2 Coords, float Ratio, float Depth, uint PrimitiveId,
 		if (!bIsVisible)
 			return;
 		#endif
-		const uint PackedCoverage = PackCoverageAndHairCount(Ratio * Coverage);
-		InterlockedAdd(OutHairCountTexture[uint2(Coords)], PackedCoverage);
+
+		const float PrecisionScale = 1000.f;
+		const float ViewTransmittance = 1 - saturate(Ratio * Coverage);
+		const float LogViewTransmittance = -log2(ViewTransmittance);
+		const uint PackedViewTransmittance = uint(LogViewTransmittance * PrecisionScale);
+		InterlockedAdd(OutHairCountTexture[uint2(Coords)], PackedViewTransmittance);
 
 		uint2 PackedVisibility;
 		PackedVisibility.x = PackHairVisSample(PrimitiveId, HairMaterialId);
-		PackedVisibility.y = PackHairVisSampleDepth(Depth, Coverage);
+		PackedVisibility.y = PackHairVisSampleDepth(Depth, Ratio * Coverage);
 
 		PackedType In = ToPackedType(PackedVisibility);
 		PackedType Out = 0;
