Fix capsule light build failure on some console platforms

Some older compilers do not support struct methods

diff --git a/Engine/Shaders/Private/CapsuleLightIntegrate.ush b/Engine/Shaders/Private/CapsuleLightIntegrate.ush
--- a/Engine/Shaders/Private/CapsuleLightIntegrate.ush
+++ b/Engine/Shaders/Private/CapsuleLightIntegrate.ush
@@ -224,7 +224,7 @@ FDirectLighting IntegrateBxDF( FGBufferData GBuffer, half3 N, half3 V, FCapsuleL
 
 				if( Capsule.Length > 0 )
 				{
-					float3 ToArea = CapsuleBounds.Sample(E).xyz;
+					float3 ToArea = SampleCapsuleBounds(CapsuleBounds, E).xyz;
 					L = normalize( ToArea );
 					H = normalize( V + L );
 				}
@@ -252,8 +252,8 @@ FDirectLighting IntegrateBxDF( FGBufferData GBuffer, half3 N, half3 V, FCapsuleL
 				float PDF[] =
 				{
 					NoL / PI,
-					D_GGX( Pow4(GBuffer.Roughness), NoH ) * NoH / (4 * VoH),
-					1.0 / CapsuleBounds.GetSolidAngle(),
+					D_GGX(Pow4(GBuffer.Roughness), NoH) * NoH / (4 * VoH),
+					1.0 / GetCapsuleBoundsSolidAngle(CapsuleBounds)
 				};
 
 				if( Capsule.Length == 0 )
diff --git a/Engine/Shaders/Private/CapsuleLightSampling.ush b/Engine/Shaders/Private/CapsuleLightSampling.ush
--- a/Engine/Shaders/Private/CapsuleLightSampling.ush
+++ b/Engine/Shaders/Private/CapsuleLightSampling.ush
@@ -17,26 +17,26 @@ struct FCapsuleSphericalBounds
 	float3 ConeAxis;
 	float ConeSinThetaMax2;
 	float ConeSolidAngle;
+};
+
+// return the smallest of the two possible bounding shapes
+float GetCapsuleBoundsSolidAngle(FCapsuleSphericalBounds Bounds)
+{
+	return min(Bounds.SphericalRect.SolidAngle, Bounds.ConeSolidAngle);
+}
 
-	float GetSolidAngle()
+// return world space direction within the cone and solidangle
+float4 SampleCapsuleBounds(FCapsuleSphericalBounds Bounds, float2 E)
+{
+	if (Bounds.ConeSolidAngle < Bounds.SphericalRect.SolidAngle)
 	{
-		// See explanation in UniformSampleConeRobust
-		return min(SphericalRect.SolidAngle, ConeSolidAngle);
+		return float4(TangentToWorld(UniformSampleConeRobust(E, Bounds.ConeSinThetaMax2).xyz, Bounds.ConeAxis), Bounds.ConeSolidAngle);
 	}
-
-	// return world space direction within the cone
-	float4 Sample(float2 E)
+	else
 	{
-		if (ConeSolidAngle < SphericalRect.SolidAngle)
-		{
-			return float4(TangentToWorld(UniformSampleConeRobust(E, ConeSinThetaMax2).xyz, ConeAxis), ConeSolidAngle);
-		}
-		else
-		{
-			return float4(normalize(UniformSampleSphericalRect(E, SphericalRect)), SphericalRect.SolidAngle);
-		}
+		return float4(normalize(UniformSampleSphericalRect(E, Bounds.SphericalRect)), Bounds.SphericalRect.SolidAngle);
 	}
-};
+}
 
 // Given a bounding rectangle of a capsule, compute a bounding cone
 FCapsuleSphericalBounds CapsuleGetSphericalBounds(float3 Origin, float3 Axis, float Radius, float Length)
diff --git a/Engine/Shaders/Private/PathTracing/Light/PathTracingCapsuleLight.ush b/Engine/Shaders/Private/PathTracing/Light/PathTracingCapsuleLight.ush
--- a/Engine/Shaders/Private/PathTracing/Light/PathTracingCapsuleLight.ush
+++ b/Engine/Shaders/Private/PathTracing/Light/PathTracingCapsuleLight.ush
@@ -109,7 +109,7 @@ FLightHit CapsuleLight_TraceLight(RayDesc Ray, int LightId)
 		float LightDistanceSquared = dot(LightDirection, LightDirection);
 		LightRadiance *= ComputeAttenuationFalloff(LightDistanceSquared, LightId);
 #if CAPSULE_SOLIDANGLE_SAMPLING
-		float Pdf = 1.0 / CapsuleGetSphericalBounds(LightDirection, Axis, LightRadius, SourceLength).GetSolidAngle();
+		float Pdf = 1.0 / GetCapsuleBoundsSolidAngle(CapsuleGetSphericalBounds(LightDirection, Axis, LightRadius, SourceLength));
 #else
 		float3 Normal = CapsuleNormal(Ray.Origin + t * Ray.Direction, LightPosition, Axis * SourceLength);
 
@@ -153,7 +153,7 @@ FLightSample CapsuleLight_SampleLight(
 	float3 LightDirection = GetPosition(LightId) - WorldPos;
 	FCapsuleSphericalBounds CapsuleBounds = CapsuleGetSphericalBounds(LightDirection, Axis, Radius, SourceLength);
 
-	float4 Result = CapsuleBounds.Sample(RandSample.yz);
+	float4 Result = SampleCapsuleBounds(CapsuleBounds, RandSample.yz);
 
 	float3 L = Result.xyz;
 	float SolidAngle = Result.w;
