Tweaked the noisy cloud samples acceptance factor.

diff --git a/Engine/Shaders/Private/VolumetricRenderTarget.usf b/Engine/Shaders/Private/VolumetricRenderTarget.usf
--- a/Engine/Shaders/Private/VolumetricRenderTarget.usf
+++ b/Engine/Shaders/Private/VolumetricRenderTarget.usf
@@ -718,24 +718,24 @@ void ComposeVolumetricRTOverScenePS(
 		}
 
 		// Now continue to accumulate in DataAcc
-		const float Extinction = 5.0;
+		const float Extinction = 30.0f;
 		{
-			const float weight = exp(-Extinction * abs(VolumeRGBT0.a - OutputRt0.a));
+			const float weight = exp2(-Extinction * abs(VolumeRGBT0.a - OutputRt0.a));
 			ValidSampleCount += weight;
 			DataAcc += VolumeRGBT0 * weight;
 		}
 		{
-			const float weight = exp(-Extinction * abs(VolumeRGBT1.a - OutputRt0.a));
+			const float weight = exp2(-Extinction * abs(VolumeRGBT1.a - OutputRt0.a));
 			ValidSampleCount += weight;
 			DataAcc += VolumeRGBT1 * weight;
 		}
 		{
-			const float weight = exp(-Extinction * abs(VolumeRGBT2.a - OutputRt0.a));
+			const float weight = exp2(-Extinction * abs(VolumeRGBT2.a - OutputRt0.a));
 			ValidSampleCount += weight;
 			DataAcc += VolumeRGBT2 * weight;
 		}
 		{
-			const float weight = exp(-Extinction * abs(VolumeRGBT3.a - OutputRt0.a));
+			const float weight = exp2(-Extinction * abs(VolumeRGBT3.a - OutputRt0.a));
 			ValidSampleCount += weight;
 			DataAcc += VolumeRGBT3 * weight;
 		}
