Remove the cloud/water interaction special path as this causes more visual  issues than it fixes. A proper fix is coming from the UE4 branch.

diff --git a/Engine/Shaders/Private/VolumetricRenderTarget.usf b/Engine/Shaders/Private/VolumetricRenderTarget.usf
--- a/Engine/Shaders/Private/VolumetricRenderTarget.usf
+++ b/Engine/Shaders/Private/VolumetricRenderTarget.usf
@@ -642,6 +642,8 @@ void ComposeVolumetricRTOverScenePS(
 		float4 VolumeRGBT2 = SafeLoadVolumetricTexture(VolumeCoord2);
 		float4 VolumeRGBT3 = SafeLoadVolumetricTexture(VolumeCoord3);
 
+		/*	// This cloud/water hack is commented out because it interacts badly with clouds. Proper fix is coming from UE4 main branch.
+		 
 		const float ConsideredFarWaterDistanceKm = 1.0f;				// We only apply the fix for pixel that are further than 1km
 		const float NotRenderedPixelDistanceKm = 500.0f;				// We assume pixel have a depth greater than that such distant depth pixel
 		if (any(Depth4 > NotRenderedPixelDistanceKm)					// Some pixels have non rendered pixel...
@@ -676,6 +678,8 @@ void ComposeVolumetricRTOverScenePS(
 			}
 			//DataAcc = float4(1.0, 1.0, 0, 0); 
 		}
+		
+		*/
 
 		//if (ValidSampleCount == 0 && bAllCloudSamplesInFrontOfWater)
 		if(ValidSampleCount == 0) // The above is removed because it does not always work. We need to test if pixel is water. A change list for that is coming from UE4/Main soon.
