Fixed a merge issue preventing volumetric cloud to render.

diff --git a/Engine/Shaders/Private/VolumetricCloud.usf b/Engine/Shaders/Private/VolumetricCloud.usf
--- a/Engine/Shaders/Private/VolumetricCloud.usf
+++ b/Engine/Shaders/Private/VolumetricCloud.usf
@@ -1099,8 +1099,6 @@ void MainCommon(in FMaterialPixelParameters MaterialParameters, in float4 SvPosi
 			OutColor1.x = bMediumHit ? OutColor1.x : MaxHalfFloat;
 		}
 #endif
-
-		return;
 	}
 	////////////////////////////////////////////////////////////////////////////////////////////////////
 	// End of cloud ray marching
