Fix incorrect backface culling when computing shadow mask for hair with virtual shadow mask.

diff --git a/Engine/Shaders/Private/VirtualShadowMaps/ProjectionSpot.ush b/Engine/Shaders/Private/VirtualShadowMaps/ProjectionSpot.ush
--- a/Engine/Shaders/Private/VirtualShadowMaps/ProjectionSpot.ush
+++ b/Engine/Shaders/Private/VirtualShadowMaps/ProjectionSpot.ush
@@ -221,6 +221,7 @@ float TraceLocalLight(
 	uint SamplesPerRay,
 	float CotMaxRayAngleFromLight,
 	const half Noise,
+	const bool bCullBackface, 
 	inout float InOutHitDistance)
 {
 	const half3 N = GetGBufferDataUint( PixelPos, true ).WorldNormal;
@@ -237,7 +238,7 @@ float TraceLocalLight(
 
 	half SinAlphaSqr = saturate(Pow2(Light.SourceRadius) * Falloff);
 	half SinAlpha = sqrt(SinAlphaSqr);
-	bool bBackFace = dot( N, ConeAxis ) < -SinAlpha;
+	bool bBackFace = bCullBackface ? dot( N, ConeAxis ) < -SinAlpha : false;
 
 	half RayMissFactor = 0;
 	if (!bBackFace)
diff --git a/Engine/Shaders/Private/VirtualShadowMaps/VirtualShadowMapProjection.usf b/Engine/Shaders/Private/VirtualShadowMaps/VirtualShadowMapProjection.usf
--- a/Engine/Shaders/Private/VirtualShadowMaps/VirtualShadowMapProjection.usf
+++ b/Engine/Shaders/Private/VirtualShadowMaps/VirtualShadowMapProjection.usf
@@ -279,6 +279,7 @@ FProjectLightResult ProjectLight(
 				SMRTSamplesPerRay,
 				SMRTCotMaxRayAngleFromLight,
 				Noise,
+				/*bCullBackface*/ !bIsHair,
 				/* out */ HitDistance);
 		#endif // DIRECTIONAL_LIGHT
 		}
