Made strata mulitple scattering LUT available for forward rendering.

diff --git a/Engine/Shaders/Private/BasePassPixelShader.usf b/Engine/Shaders/Private/BasePassPixelShader.usf
--- a/Engine/Shaders/Private/BasePassPixelShader.usf
+++ b/Engine/Shaders/Private/BasePassPixelShader.usf
@@ -975,14 +975,14 @@ void FPixelShaderInOut_MainPS(
 #endif // !STRATA_ENABLED
 
 #if STRATA_ENABLED && !MATERIAL_SHADINGMODEL_UNLIT
-	FStrataData Strata = PixelMaterialInputs.FrontMaterial;
+	FStrataData StrataData = PixelMaterialInputs.FrontMaterial;
 
 	// Initialise a Strata header with normal in registers
 	FStrataPixelHeader StrataPixelHeader = InitialiseStrataPixelHeader();
 	UNROLL
-	for (int l = 0; l < Strata.LayerCount; ++l)
+	for (int l = 0; l < StrataData.LayerCount; ++l)
 	{
-		StrataPixelHeader.BSDFCount += Strata.Layers[l].BSDFCount;
+		StrataPixelHeader.BSDFCount += StrataData.Layers[l].BSDFCount;
 	}
 	StrataPixelHeader.SharedNormals = MaterialParameters.SharedNormals;
 
@@ -1579,14 +1579,14 @@ void FPixelShaderInOut_MainPS(
 	PackStrataOut(
 		OpaqueBasePass.Strata.MaterialLobesBufferUAV, 
 		GetStrataPixelDataByteOffset(PixelPos, uint2(ResolvedView.BufferSizeAndInvSize.xy), OpaqueBasePass.Strata.MaxBytesPerPixel),
-		StrataPixelHeader, Strata, MaterialParameters.CameraVector, bStrataSubsurfaceEnable, EmissiveLuminance);
+		StrataPixelHeader, StrataData, MaterialParameters.CameraVector, bStrataSubsurfaceEnable, EmissiveLuminance);
 
 	Color += EmissiveLuminance;
 #endif // STRATA_OPAQUE_ENABLED
 
 #if STRATA_TRANSLUCENT_ENABLED
 	//FORWARD_SHADING
-	if (Strata.LayerCount > 0)
+	if (StrataData.LayerCount > 0)
 	{
 		float3 L = ResolvedView.DirectionalLightDirection;
 		float3 V = MaterialParameters.CameraVector;
@@ -1594,7 +1594,6 @@ void FPixelShaderInOut_MainPS(
 
 		float3 Throughput = 1.0f;
 		Color = StrataForwardLighting(
-			GBuffer,		// STRATA_TODO remove that!
 			MaterialParameters.WorldNormal,
 			V,
 			L,
@@ -1602,7 +1601,7 @@ void FPixelShaderInOut_MainPS(
 			DirectionalLightShadow,
 			IndirectOcclusion,
 			StrataPixelHeader,
-			Strata,
+			StrataData,
 			Throughput);
 
 		DualBlendColorCov = GetMaterialOpacity(PixelMaterialInputs); // Opacity represent coverage STRATA_TODO after validation, rename everything
diff --git a/Engine/Shaders/Private/Strata/StrataEvaluation.ush b/Engine/Shaders/Private/Strata/StrataEvaluation.ush
--- a/Engine/Shaders/Private/Strata/StrataEvaluation.ush
+++ b/Engine/Shaders/Private/Strata/StrataEvaluation.ush
@@ -21,6 +21,10 @@
 #define STRATA_USE_ENV_BRDF_APPROX 0
 #endif
 
+#ifndef STRATA_TRANSLUCENT_ENABLED
+#define STRATA_TRANSLUCENT_ENABLED 0
+#endif
+
 #ifndef STRATA_MICROFACET_MULTIPLESCATTERING
 #define STRATA_MICROFACET_MULTIPLESCATTERING 0
 #endif
@@ -193,7 +197,15 @@ float StrataBSDFProjectedSolidAngleFactor(in FStrataBSDFContext BSDFContext)
 
 FStrataGGXLUTResult StrataEvaluateGGXLUT(float3 DiffuseAlbedo, float3 F0, float NoV, float Metallic, float SafeRoughnessX, float SafeRoughnessY)
 {
-	return StrataSampleGGXLUT(Strata.GGXEnergyLUT2DTexture, Strata.GGXEnergyLUT3DTexture, Strata.GGXEnergyLUTSampler, Strata.GGXEnergyLUTScaleBias, NoV, SafeRoughnessX, SafeRoughnessY, F0, Metallic, DiffuseAlbedo);
+#if STRATA_TRANSLUCENT_ENABLED
+	return StrataSampleGGXLUT(
+		TranslucentBasePass.Strata.GGXEnergyLUT2DTexture, TranslucentBasePass.Strata.GGXEnergyLUT3DTexture, TranslucentBasePass.Strata.GGXEnergyLUTSampler,
+		TranslucentBasePass.Strata.GGXEnergyLUTScaleBias, NoV, SafeRoughnessX, SafeRoughnessY, F0, Metallic, DiffuseAlbedo);
+#else
+	return StrataSampleGGXLUT(
+		Strata.GGXEnergyLUT2DTexture, Strata.GGXEnergyLUT3DTexture, Strata.GGXEnergyLUTSampler,
+		Strata.GGXEnergyLUTScaleBias, NoV, SafeRoughnessX, SafeRoughnessY, F0, Metallic, DiffuseAlbedo);
+#endif
 }
 
 #else // STRATA_MICROFACET_MULTIPLESCATTERING
@@ -617,12 +629,12 @@ FStrataEnvLightResult StrataEvaluateForEnvLight(FStrataBSDFContext BSDFContext,
 			StrataEnvLightResult.DiffuseWeight = DiffuseColor;
 
 			float SafeRoughnessX = MakeRoughnessSafe(SLAB_ROUGHNESSX(BSDFContext.BSDF));
+			float3 F0 = ComputeF0(SLAB_SPECULAR(BSDFContext.BSDF), BaseColor, Metallic);
+			float3 ThinFilmFresnelWeight = 1;
 
-			float3 FresnelWeight = 1;
 			BRANCH
 			if (bEnableSpecular)
 			{
-				float3 F0 = ComputeF0(SLAB_SPECULAR(BSDFContext.BSDF), BaseColor, Metallic);
 				const float3 F90 = ComputeF90(F0, SLAB_EDGECOLOR(BSDFContext.BSDF), Metallic);
 
 				#if STRATA_FASTPATH==0
@@ -631,7 +643,7 @@ FStrataEnvLightResult StrataEvaluateForEnvLight(FStrataBSDFContext BSDFContext,
 				{
 					float ThinFilmIOR = SLAB_THINFILM_IOR(BSDFContext.BSDF);
 					float ThinFilmThickness = SLAB_THINFILM_THICKNESS(BSDFContext.BSDF);
-					FresnelWeight = F_ThinFilm(BSDFContext.Context.NoV, BSDFContext.Context.NoL, BSDFContext.Context.VoH, F0, F90, ThinFilmIOR, ThinFilmThickness);
+					ThinFilmFresnelWeight = F_ThinFilm(BSDFContext.Context.NoV, BSDFContext.Context.NoL, BSDFContext.Context.VoH, F0, F90, ThinFilmIOR, ThinFilmThickness);
 					F0 = 1.f;
 				}
 				#endif
@@ -660,11 +672,11 @@ FStrataEnvLightResult StrataEvaluateForEnvLight(FStrataBSDFContext BSDFContext,
 				float3 EvalEnvBRDF = 0;
 				{
 					FStrataGGXLUTResult GGXLUT	= StrataEvaluateGGXLUT(DiffuseColor, F0, BSDFContext.Context.NoV, Metallic, SafeRoughnessX, SafeRoughnessX);
-					EvalEnvBRDF					= FresnelWeight * GGXLUT.DirectionalAlbedo.SpecularR;
+					EvalEnvBRDF					= ThinFilmFresnelWeight * GGXLUT.DirectionalAlbedo.SpecularR;
 					StrataEnvLightResult.DiffuseWeight = GGXLUT.DirectionalAlbedo.Diffuse;
 				}
 			#else
-				const float3 EvalEnvBRDF = FresnelWeight * StrataEnvBRDF(F0, F90, SafeRoughnessX, BSDFContext.SatNoV);
+				const float3 EvalEnvBRDF = ThinFilmFresnelWeight * StrataEnvBRDF(F0, F90, SafeRoughnessX, BSDFContext.SatNoV);
 			#endif
 
 				StrataEnvLightResult.SpecularDirection = StrataGetOffSpecularPeakReflectionDir(BSDFContext.N, BSDFContext.R, SafeRoughnessX);
@@ -689,9 +701,9 @@ FStrataEnvLightResult StrataEvaluateForEnvLight(FStrataBSDFContext BSDFContext,
 					// Compute the directional albedo for hazy specular
 				#if STRATA_MICROFACET_MULTIPLESCATTERING
 					FStrataGGXLUTResult GGXLUTHaze = StrataEvaluateGGXLUT(DiffuseColor, F0, BSDFContext.Context.NoV, Metallic, SafeRoughnessHaze, SafeRoughnessHaze);
-					const float3 EvalEnvBRDFHaze = FresnelWeight * GGXLUTHaze.DirectionalAlbedo.SpecularR;
+					const float3 EvalEnvBRDFHaze = ThinFilmFresnelWeight * GGXLUTHaze.DirectionalAlbedo.SpecularR;
 				#else
-					const float3 EvalEnvBRDFHaze = FresnelWeight * StrataEnvBRDF(F0, F90, SafeRoughnessX, BSDFContext.SatNoV);
+					const float3 EvalEnvBRDFHaze = ThinFilmFresnelWeight * StrataEnvBRDF(F0, F90, SafeRoughnessX, BSDFContext.SatNoV);
 				#endif
 
 					StrataEnvLightResult.SpecularHazeWeight = lerp(0.0f, SPECULAR_HAZINESS_WEIGHT, HazinessFadeIn) * EvalEnvBRDFHaze;
@@ -734,11 +746,14 @@ FStrataEnvLightResult StrataEvaluateForEnvLight(FStrataBSDFContext BSDFContext,
 					// We only want to do that when the slab light response is higher than the diffse response, so that when the mfp is large (and scattering low), the light response still work correctly according to thickness and MFP.
 					SlabDirectionalAlbedo *= MaxDiffuseWeight / MaxSlabDirectionalAlbedo;
 				}
+			#else
+				// STRATA_TODO: Only considering the non hazy lob for now.
+				// STRATA_TODO: This is not enough to count the difference of energy when lerping from Chan to the directional albedo. This should be investigated.
+				FStrataGGXLUTResult GGXLUT = StrataEvaluateGGXLUT(DiffuseColor, F0, BSDFContext.Context.NoV, Metallic, SafeRoughnessX, SafeRoughnessX);
+				SlabDirectionalAlbedo *= GGXLUT.DirectionalAlbedo.SpecularT;
 			#endif
 
-				// STRATA_TODO: big approximation for volumetric + env lighting: We ignore the phase for now due to the sampled convolved env lighting, 
-				// no self shadow and should not be single scattering using StrataSlabSingleScattering. We will propose a better approach. 
-				StrataEnvLightResult.DiffuseWeight = lerp(StrataEnvLightResult.DiffuseWeight, SlabDirectionalAlbedo * FresnelWeight,	DiffuseToVolumeBlend);
+				StrataEnvLightResult.DiffuseWeight = lerp(StrataEnvLightResult.DiffuseWeight, SlabDirectionalAlbedo * ThinFilmFresnelWeight,	DiffuseToVolumeBlend);
 				// We keep the same normal
 				StrataEnvLightResult.bSubsurface = false;
 			}
diff --git a/Engine/Shaders/Private/Strata/StrataForwardLighting.ush b/Engine/Shaders/Private/Strata/StrataForwardLighting.ush
--- a/Engine/Shaders/Private/Strata/StrataForwardLighting.ush
+++ b/Engine/Shaders/Private/Strata/StrataForwardLighting.ush
@@ -5,7 +5,6 @@
 
 
 float3 StrataForwardLighting(
-	FGBufferData GBuffer,
 	float3 MaterialParametersWorldNormal,
 	float3 V,
 	float3 L,
@@ -13,7 +12,7 @@ float3 StrataForwardLighting(
 	float3 LShadow,
 	float3 IndirectOcclusion,
 	FStrataPixelHeader Header,
-	FStrataData Strata,
+	FStrataData StrataData,
 	inout float3 OutThroughput
 	)
 {
@@ -26,18 +25,18 @@ float3 StrataForwardLighting(
 	OutThroughput = 1.0f;
 
 	UNROLL
-	for (int l = 0; l < Strata.LayerCount; ++l)
+	for (int l = 0; l < StrataData.LayerCount; ++l)
 	{
 		float3 LayerThroughput = 0.0f;
 		float3 LayerCoverage = 0.0f;
 
 		const bool bTopLayer = l == 0;
-		const bool bBottomLayer = l == (Strata.LayerCount - 1);
+		const bool bBottomLayer = l == (StrataData.LayerCount - 1);
 
 		UNROLL
-		for (int i = 0; i < Strata.Layers[l].BSDFCount; ++i)
+		for (int i = 0; i < StrataData.Layers[l].BSDFCount; ++i)
 		{
-			FStrataBSDF BSDF = Strata.Layers[l].BSDFs[i];
+			FStrataBSDF BSDF = StrataData.Layers[l].BSDFs[i];
 
 			// Sanitize BSDF before it is used for forward shading
 			StrataSanitizeBSDF(BSDF);
diff --git a/Engine/Source/Runtime/Renderer/Private/BasePassRendering.h b/Engine/Source/Runtime/Renderer/Private/BasePassRendering.h
--- a/Engine/Source/Runtime/Renderer/Private/BasePassRendering.h
+++ b/Engine/Source/Runtime/Renderer/Private/BasePassRendering.h
@@ -73,7 +73,7 @@ END_GLOBAL_SHADER_PARAMETER_STRUCT()
 
 BEGIN_GLOBAL_SHADER_PARAMETER_STRUCT(FOpaqueBasePassUniformParameters,)
 	SHADER_PARAMETER_STRUCT(FSharedBasePassUniformParameters, Shared)
-	SHADER_PARAMETER_STRUCT(FStrataOpaquePassUniformParameters, Strata)
+	SHADER_PARAMETER_STRUCT(FStrataBasePassUniformParameters, Strata)
 	// Forward shading 
 	SHADER_PARAMETER(int32, UseForwardScreenSpaceShadowMask)
 	SHADER_PARAMETER_RDG_TEXTURE(Texture2D, ForwardScreenSpaceShadowMaskTexture)
@@ -97,6 +97,7 @@ END_GLOBAL_SHADER_PARAMETER_STRUCT()
 BEGIN_GLOBAL_SHADER_PARAMETER_STRUCT(FTranslucentBasePassUniformParameters,)
 	SHADER_PARAMETER_STRUCT(FSharedBasePassUniformParameters, Shared)
 	SHADER_PARAMETER_STRUCT(FSceneTextureUniformParameters, SceneTextures)
+	SHADER_PARAMETER_STRUCT(FStrataBasePassUniformParameters, Strata)
 	// Material SSR
 	SHADER_PARAMETER(FVector4, HZBUvFactorAndInvFactor)
 	SHADER_PARAMETER(FVector4, PrevScreenPositionScaleBias)
diff --git a/Engine/Source/Runtime/Renderer/Private/Strata/Strata.cpp b/Engine/Source/Runtime/Renderer/Private/Strata/Strata.cpp
--- a/Engine/Source/Runtime/Renderer/Private/Strata/Strata.cpp
+++ b/Engine/Source/Runtime/Renderer/Private/Strata/Strata.cpp
@@ -284,17 +284,24 @@ void InitialiseStrataFrameSceneData(FSceneRenderer& SceneRenderer, FRDGBuilder&
 	}
 }
 
-void BindStrataBasePassUniformParameters(FRDGBuilder& GraphBuilder, FStrataSceneData* StrataSceneData, FStrataOpaquePassUniformParameters& OutStrataUniformParameters)
+void BindStrataBasePassUniformParameters(FRDGBuilder& GraphBuilder, FStrataSceneData* StrataSceneData, FStrataBasePassUniformParameters& OutStrataUniformParameters)
 {
-	if (StrataSceneData)
+	OutStrataUniformParameters.GGXEnergyLUTScaleBias = GetStrataGGXEnergyLUTScaleBias();
+	OutStrataUniformParameters.GGXEnergyLUTSampler = TStaticSamplerState<SF_Bilinear, AM_Clamp, AM_Clamp, AM_Clamp>::GetRHI();
+
+	if (IsStrataEnabled() && StrataSceneData)
 	{
 		OutStrataUniformParameters.MaxBytesPerPixel = StrataSceneData->MaxBytesPerPixel;
 		OutStrataUniformParameters.MaterialLobesBufferUAV = StrataSceneData->MaterialLobesBufferUAV;
+		OutStrataUniformParameters.GGXEnergyLUT3DTexture = StrataSceneData->GGXEnergyLUT3DTexture->GetRenderTargetItem().ShaderResourceTexture;
+		OutStrataUniformParameters.GGXEnergyLUT2DTexture = StrataSceneData->GGXEnergyLUT2DTexture->GetRenderTargetItem().ShaderResourceTexture;
 	}
 	else
 	{
 		OutStrataUniformParameters.MaxBytesPerPixel = 0;
 		OutStrataUniformParameters.MaterialLobesBufferUAV = GraphBuilder.CreateUAV(GraphBuilder.RegisterExternalBuffer(GWhiteVertexBufferWithRDG->Buffer), PF_R32_UINT);
+		OutStrataUniformParameters.GGXEnergyLUT3DTexture = GSystemTextures.VolumetricBlackDummy->GetRenderTargetItem().ShaderResourceTexture;
+		OutStrataUniformParameters.GGXEnergyLUT2DTexture = GSystemTextures.BlackDummy->GetRenderTargetItem().ShaderResourceTexture;
 	}
 }
 
diff --git a/Engine/Source/Runtime/Renderer/Private/Strata/Strata.h b/Engine/Source/Runtime/Renderer/Private/Strata/Strata.h
--- a/Engine/Source/Runtime/Renderer/Private/Strata/Strata.h
+++ b/Engine/Source/Runtime/Renderer/Private/Strata/Strata.h
@@ -14,8 +14,12 @@ class FScene;
 class FRDGBuilder;
 struct FMinimalSceneTextures;
 
-BEGIN_SHADER_PARAMETER_STRUCT(FStrataOpaquePassUniformParameters, )
+BEGIN_SHADER_PARAMETER_STRUCT(FStrataBasePassUniformParameters, )
 	SHADER_PARAMETER(uint32, MaxBytesPerPixel)
+	SHADER_PARAMETER(FVector2D, GGXEnergyLUTScaleBias)
+	SHADER_PARAMETER_TEXTURE(Texture3D<float2>, GGXEnergyLUT3DTexture)
+	SHADER_PARAMETER_TEXTURE(Texture2D<float4>, GGXEnergyLUT2DTexture)
+	SHADER_PARAMETER_SAMPLER(SamplerState, GGXEnergyLUTSampler)
 	SHADER_PARAMETER_RDG_BUFFER_UAV(RWByteAddressBuffer, MaterialLobesBufferUAV)
 END_SHADER_PARAMETER_STRUCT()
 
@@ -85,7 +89,7 @@ bool ShouldPassesReadingStrataBeTiled(ERHIFeatureLevel::Type FeatureLevel);
 
 void InitialiseStrataFrameSceneData(FSceneRenderer& SceneRenderer, FRDGBuilder& GraphBuilder);
 
-void BindStrataBasePassUniformParameters(FRDGBuilder& GraphBuilder, FStrataSceneData* StrataSceneData, FStrataOpaquePassUniformParameters& OutStrataUniformParameters);
+void BindStrataBasePassUniformParameters(FRDGBuilder& GraphBuilder, FStrataSceneData* StrataSceneData, FStrataBasePassUniformParameters& OutStrataUniformParameters);
 
 TRDGUniformBufferRef<FStrataGlobalUniformParameters> BindStrataGlobalUniformParameters(FStrataSceneData* StrataSceneData);
 
diff --git a/Engine/Source/Runtime/Renderer/Private/TranslucentRendering.cpp b/Engine/Source/Runtime/Renderer/Private/TranslucentRendering.cpp
--- a/Engine/Source/Runtime/Renderer/Private/TranslucentRendering.cpp
+++ b/Engine/Source/Runtime/Renderer/Private/TranslucentRendering.cpp
@@ -13,6 +13,7 @@
 #include "VolumetricRenderTarget.h"
 #include "Lumen/LumenTranslucencyVolumeLighting.h"
 #include "VirtualShadowMaps/VirtualShadowMapArray.h"
+#include "Strata/Strata.h"
 
 DECLARE_CYCLE_STAT(TEXT("TranslucencyTimestampQueryFence Wait"), STAT_TranslucencyTimestampQueryFence_Wait, STATGROUP_SceneRendering);
 DECLARE_CYCLE_STAT(TEXT("TranslucencyTimestampQuery Wait"), STAT_TranslucencyTimestampQuery_Wait, STATGROUP_SceneRendering);
@@ -662,6 +663,7 @@ TRDGUniformBufferRef<FTranslucentBasePassUniformParameters> CreateTranslucentBas
 
 	SetupSharedBasePassParameters(GraphBuilder, View, BasePassParameters.Shared);
 	SetupSceneTextureUniformParameters(GraphBuilder, View.FeatureLevel, SceneTextureSetupMode, BasePassParameters.SceneTextures);
+	Strata::BindStrataBasePassUniformParameters(GraphBuilder, View.StrataSceneData, BasePassParameters.Strata);
 
 	const FRDGSystemTextures& SystemTextures = FRDGSystemTextures::Get(GraphBuilder);
 
