Tessellation Removal: Removing tessellation code from shaders. Leaving code in Vertex Factories because Raytracing depends on those functions.

diff --git a/Engine/Shaders/Private/AnisotropyPassShader.usf b/Engine/Shaders/Private/AnisotropyPassShader.usf
--- a/Engine/Shaders/Private/AnisotropyPassShader.usf
+++ b/Engine/Shaders/Private/AnisotropyPassShader.usf
@@ -19,58 +19,9 @@ struct FAnisotropyPassVSToPS
 	#endif
 };
 
-#if USING_TESSELLATION
-	struct FAnisotropyPassVSToDS
-	{
-		FVertexFactoryInterpolantsVSToDS FactoryInterpolants;
-		float4 Position : VS_To_DS_Position;
-		#if USE_WORLD_POSITION_EXCLUDING_SHADER_OFFSETS
-			float3 PixelPositionExcludingWPO : TEXCOORD7;
-		#endif
-		OPTIONAL_VertexID_VS_To_DS
-	};
-	
-	#define FVertexOutput FAnisotropyPassVSToDS
-	#define VertexFactoryGetInterpolants VertexFactoryGetInterpolantsVSToDS
-#else
-	#define FVertexOutput FAnisotropyPassVSToPS
-	#define VertexFactoryGetInterpolants VertexFactoryGetInterpolantsVSToPS
-#endif
-
-#if USING_TESSELLATION
-	#define FPassSpecificVSToDS FAnisotropyPassVSToDS
-	#define FPassSpecificVSToPS FAnisotropyPassVSToPS
-		
-	FAnisotropyPassVSToDS PassInterpolate(FAnisotropyPassVSToDS a, float aInterp, FAnisotropyPassVSToDS b, float bInterp)
-	{
-		FAnisotropyPassVSToDS O;
-
-		O.FactoryInterpolants = VertexFactoryInterpolate(a.FactoryInterpolants, aInterp, b.FactoryInterpolants, bInterp);
-		
-		#if USE_WORLD_POSITION_EXCLUDING_SHADER_OFFSETS
-			TESSELLATION_INTERPOLATE_MEMBER(PixelPositionExcludingWPO);
-		#endif
-
-		return O;
-	}
-	
-	FAnisotropyPassVSToPS PassFinalizeTessellationOutput(FAnisotropyPassVSToDS Interpolants, float4 WorldPosition, FMaterialTessellationParameters MaterialParameters)
-	{
-		FAnisotropyPassVSToPS O;
-		
-		O.Interps = VertexFactoryAssignInterpolants(Interpolants.FactoryInterpolants);
-		O.Position = mul(WorldPosition, ResolvedView.TranslatedWorldToClip);
+#define FVertexOutput FAnisotropyPassVSToPS
+#define VertexFactoryGetInterpolants VertexFactoryGetInterpolantsVSToPS
 
-		#if USE_WORLD_POSITION_EXCLUDING_SHADER_OFFSETS
-			O.PixelPositionExcludingWPO = Interpolants.PixelPositionExcludingWPO;
-		#endif
-
-		return O;
-	}
-	
-	#include "Tessellation.ush"
-#endif
-	
 /*=============================================================================
  * Vertex Shader
  *============================================================================*/
@@ -79,7 +30,7 @@ void MainVertexShader(
 	FVertexFactoryInput Input,
 	OPTIONAL_VertexID
 	out FVertexOutput Output 
-#if USE_GLOBAL_CLIP_PLANE && !USING_TESSELLATION
+#if USE_GLOBAL_CLIP_PLANE
 	, out float OutGlobalClipPlaneDistance : SV_ClipDistance
 #endif
 #if INSTANCED_STEREO
@@ -119,10 +70,6 @@ void MainVertexShader(
 		WorldPos.xyz += GetMaterialWorldPositionOffset(VertexParameters);
 	}
 
-#if USING_TESSELLATION
-	// Transformation is done in Domain shader when tessellating
-	Output.Position = WorldPos; 
-#else
 	{
 		float4 RasterizedWorldPosition = VertexFactoryGetRasterizedWorldPosition(Input, VFIntermediates, WorldPos);
 	#if ODS_CAPTURE
@@ -156,20 +103,11 @@ void MainVertexShader(
 	#if USE_GLOBAL_CLIP_PLANE
 		OutGlobalClipPlaneDistance = dot(ResolvedView.GlobalClippingPlane, float4(WorldPos.xyz - ResolvedView.PreViewTranslation.xyz, 1));
 	#endif
-#endif
 
-#if USING_TESSELLATION
-	Output.FactoryInterpolants = VertexFactoryGetInterpolants( Input, VFIntermediates, VertexParameters );
-#else
 	Output.Interps = VertexFactoryGetInterpolants(Input, VFIntermediates, VertexParameters);
-#endif // #if USING_TESSELLATION
 
 #if INSTANCED_STEREO
-	#if USING_TESSELLATION	
-		Output.Interps.InterpolantsVSToPS.EyeIndex = EyeIndex;
-	#else
-		Output.Interps.EyeIndex = EyeIndex;
-	#endif
+	Output.Interps.EyeIndex = EyeIndex;
 #endif
 
 #if USE_WORLD_POSITION_EXCLUDING_SHADER_OFFSETS
diff --git a/Engine/Shaders/Private/BasePassTessellationShaders.usf b/Engine/Shaders/Private/BasePassTessellationShaders.usf
deleted file mode 100644
diff --git a/Engine/Shaders/Private/BasePassVertexShader.usf b/Engine/Shaders/Private/BasePassVertexShader.usf
--- a/Engine/Shaders/Private/BasePassVertexShader.usf
+++ b/Engine/Shaders/Private/BasePassVertexShader.usf
@@ -8,11 +8,6 @@
 #include "SHCommon.ush"
 #include "VolumetricLightmapShared.ush"
 
-// Metal needs to combine the vertex & hull shaders for tesellation support.
-#if COMPILER_METAL && USING_TESSELLATION
-#include "BasePassTessellationShaders.usf"
-#endif
-
 #if TRANSLUCENCY_PERVERTEX_FORWARD_SHADING
 #include "ReflectionEnvironmentShared.ush"
 #include "ForwardLightingCommon.ush"
@@ -32,7 +27,7 @@ void Main(
 	FVertexFactoryInput Input,
 	OPTIONAL_VertexID
 	out FBasePassVSOutput Output
-#if USE_GLOBAL_CLIP_PLANE && !USING_TESSELLATION
+#if USE_GLOBAL_CLIP_PLANE
 	, out float OutGlobalClipPlaneDistance : SV_ClipDistance
 #endif
 #if INSTANCED_STEREO
@@ -73,14 +68,6 @@ void Main(
 		WorldPosition.xyz += GetMaterialWorldPositionOffset(VertexParameters);
 	}
 
-#if USING_TESSELLATION
-	// We let the Domain Shader convert to post projection when tessellating
-	Output.Position = WorldPosition;	
-
-	#if USE_WORLD_POSITION_EXCLUDING_SHADER_OFFSETS
-		Output.BasePassInterpolants.WorldPositionExcludingWPO = WorldPositionExcludingWPO.xyz;
-	#endif
-#else
 	{
 		float4 RasterizedWorldPosition = VertexFactoryGetRasterizedWorldPosition(Input, VFIntermediates, WorldPosition);
 	#if ODS_CAPTURE
@@ -113,16 +100,11 @@ void Main(
 	#if USE_WORLD_POSITION_EXCLUDING_SHADER_OFFSETS
 		Output.BasePassInterpolants.PixelPositionExcludingWPO = WorldPositionExcludingWPO.xyz;
 	#endif
-#endif	// USING_TESSELLATION
 
 	Output.FactoryInterpolants = VertexFactoryGetInterpolants(Input, VFIntermediates, VertexParameters);
 
 #if INSTANCED_STEREO
-	#if USING_TESSELLATION	
-		Output.FactoryInterpolants.InterpolantsVSToPS.EyeIndex = EyeIndex;
-	#else
-		Output.FactoryInterpolants.EyeIndex = EyeIndex;
-	#endif
+	Output.FactoryInterpolants.EyeIndex = EyeIndex;
 #endif
 
 // Calculate the fog needed for translucency
@@ -234,24 +216,14 @@ void Main(
 			VertexParameters = GetMaterialVertexParameters(Input, VFIntermediates, PrevTranslatedWorldPosition.xyz, TangentToLocal);
 			PrevTranslatedWorldPosition.xyz += GetMaterialPreviousWorldPositionOffset(VertexParameters);
 	
-			#if !USING_TESSELLATION
-				PrevTranslatedWorldPosition = mul(float4(PrevTranslatedWorldPosition.xyz, 1), ResolvedView.PrevTranslatedWorldToClip);
-			#endif
+			PrevTranslatedWorldPosition = mul(float4(PrevTranslatedWorldPosition.xyz, 1), ResolvedView.PrevTranslatedWorldToClip);
 		}
 
-		#if USING_TESSELLATION
-			// We let the Domain Shader convert to post projection when tessellating
-			Output.BasePassInterpolants.VelocityPrevScreenPosition = PrevTranslatedWorldPosition;
-			#if WRITES_VELOCITY_TO_GBUFFER_USE_POS_INTERPOLATOR
-			Output.BasePassInterpolants.VelocityScreenPosition = WorldPosition;
-			#endif
-		#else
-			// compute the old screen pos with the old world position and the old camera matrix
-			Output.BasePassInterpolants.VelocityPrevScreenPosition = PrevTranslatedWorldPosition; 
-			#if WRITES_VELOCITY_TO_GBUFFER_USE_POS_INTERPOLATOR
-			Output.BasePassInterpolants.VelocityScreenPosition = ClipSpacePosition;
-			#endif
-		#endif	// USING_TESSELLATION
+		// compute the old screen pos with the old world position and the old camera matrix
+		Output.BasePassInterpolants.VelocityPrevScreenPosition = PrevTranslatedWorldPosition; 
+		#if WRITES_VELOCITY_TO_GBUFFER_USE_POS_INTERPOLATOR
+		Output.BasePassInterpolants.VelocityScreenPosition = ClipSpacePosition;
+		#endif
 	}
 #endif	// WRITES_VELOCITY_TO_GBUFFER
 
diff --git a/Engine/Shaders/Private/DebugViewModeVertexShader.usf b/Engine/Shaders/Private/DebugViewModeVertexShader.usf
--- a/Engine/Shaders/Private/DebugViewModeVertexShader.usf
+++ b/Engine/Shaders/Private/DebugViewModeVertexShader.usf
@@ -31,69 +31,7 @@ MaterialFloat StoreTexSample(in out FTexCoordScalesParams Params, float4 C, int
 #include "DebugViewModeCommon.ush"
 // struct FDebugVSToPS
 
-#if USING_TESSELLATION	
-	struct FDebugVSToDS
-	{
-		FVertexFactoryInterpolantsVSToDS FactoryInterpolants;
-		float4 Position : VS_To_DS_Position;
-		OPTIONAL_VertexID_VS_To_DS 
-	};
-	
-	#define FDebugVSOutput FDebugVSToDS
-#else
-	#define FDebugVSOutput FDebugVSToPS
-#endif
-
-#if USING_TESSELLATION	
-	#define FPassSpecificVSToDS FDebugVSToDS
-	#define FPassSpecificVSToPS FDebugVSToPS
-	
-	FDebugVSToDS PassInterpolate(FDebugVSToDS a, float aInterp, FDebugVSToDS b, float bInterp)
-	{
-		FDebugVSToDS Output;
-		
-		Output.FactoryInterpolants = VertexFactoryInterpolate(a.FactoryInterpolants, aInterp, b.FactoryInterpolants, bInterp);
-
-		return Output;
-	}
-
-	FDebugVSToPS PassFinalizeTessellationOutput(FDebugVSToDS Interpolants, float4 WorldPosition, FMaterialTessellationParameters MaterialParameters)
-	{
-		FDebugVSToPS Output;
-		
-		// Finally, transform position to clip-space
-		{
-
-			Output.Position = mul(WorldPosition, ResolvedView.TranslatedWorldToClip);
-		}
-
-		Output.VertexColor = MaterialParameters.VertexColor;
-
-		float3x3 TangentToWorld = MaterialParameters.TangentToWorld;
-		Output.TangentToWorld0 = TangentToWorld[0];
-		Output.TangentToWorld1 = TangentToWorld[1];
-		Output.TangentToWorld2 = TangentToWorld[2];
-
-#if NUM_MATERIAL_TEXCOORDS_VERTEX > 0
-		Output.TexCoord23 = Output.TexCoord01 = MaterialParameters.TexCoords[0].xyxy;
-#else
-		Output.TexCoord23 = Output.TexCoord01 = 0;
-#endif
-#if NUM_MATERIAL_TEXCOORDS_VERTEX > 1
-		Output.TexCoord01.zw = MaterialParameters.TexCoords[1];
-#endif
-#if NUM_MATERIAL_TEXCOORDS_VERTEX > 2
-		Output.TexCoord23.xy = MaterialParameters.TexCoords[2];
-#endif
-#if NUM_MATERIAL_TEXCOORDS_VERTEX > 3
-		Output.TexCoord23.zw = MaterialParameters.TexCoords[3];
-#endif
-		
-		return Output;
-	}
-	
-	#include "Tessellation.ush"
-#endif
+#define FDebugVSOutput FDebugVSToPS
 
 #if VERTEXSHADER
 
@@ -101,7 +39,7 @@ void Main(
 	FVertexFactoryInput Input,
 	OPTIONAL_VertexID
 	out FDebugVSOutput Output
-#if USE_GLOBAL_CLIP_PLANE && !USING_TESSELLATION
+#if USE_GLOBAL_CLIP_PLANE
 	, out float OutGlobalClipPlaneDistance : SV_ClipDistance
 #endif
 #if INSTANCED_STEREO
@@ -139,11 +77,6 @@ void Main(
 		WorldPos.xyz += GetMaterialWorldPositionOffset(VertexParameters);
 	}
 
-#if USING_TESSELLATION
-	// Transformation is done in Domain shader when tessellating
-	Output.Position = WorldPos;
-	Output.FactoryInterpolants = VertexFactoryGetInterpolantsVSToDS(Input, VFIntermediates, VertexParameters);
-#else
 	{
 		float4 RasterizedWorldPosition = VertexFactoryGetRasterizedWorldPosition(Input, VFIntermediates, WorldPos);
 		Output.Position = mul(RasterizedWorldPosition, ResolvedView.TranslatedWorldToClip);
@@ -189,15 +122,9 @@ void Main(
 		Output.TexCoord23.zw = VertexParameters.TexCoords[3];
 #endif
 
-#endif // USING_TESSELLATION
-
-	#if INSTANCED_STEREO
-		#if USING_TESSELLATION	
-			Output.FactoryInterpolants.InterpolantsVSToPS.EyeIndex = EyeIndex;
-		#else
-			Output.EyeIndex = EyeIndex;
-		#endif
-	#endif
+#if INSTANCED_STEREO
+	Output.EyeIndex = EyeIndex;
+#endif
 
 	OutputVertexID( Output );
 }
diff --git a/Engine/Shaders/Private/DepthOnlyVertexShader.usf b/Engine/Shaders/Private/DepthOnlyVertexShader.usf
--- a/Engine/Shaders/Private/DepthOnlyVertexShader.usf
+++ b/Engine/Shaders/Private/DepthOnlyVertexShader.usf
@@ -24,67 +24,8 @@ struct FDepthOnlyVSToPS
 	#endif
 };
 
-#if USING_TESSELLATION	
-	struct FDepthOnlyVSToDS
-	{
-		FVertexFactoryInterpolantsVSToDS FactoryInterpolants;
-		float4 Position : VS_To_DS_Position;
-		#if USE_RAW_WORLD_POSITION
-			float3 PixelPositionExcludingWPO : TEXCOORD7;
-		#endif
-		OPTIONAL_VertexID_VS_To_DS
-	};
-	
-	#define FDepthOnlyVSOutput FDepthOnlyVSToDS
-	#define VertexFactoryGetInterpolants VertexFactoryGetInterpolantsVSToDS
-#else
-	#define FDepthOnlyVSOutput FDepthOnlyVSToPS
-	#define VertexFactoryGetInterpolants VertexFactoryGetInterpolantsVSToPS
-#endif
-
-#if USING_TESSELLATION	
-	#define FPassSpecificVSToDS FDepthOnlyVSToDS
-	#define FPassSpecificVSToPS FDepthOnlyVSToPS
-	
-	FDepthOnlyVSToDS PassInterpolate(FDepthOnlyVSToDS a, float aInterp, FDepthOnlyVSToDS b, float bInterp)
-	{
-		FDepthOnlyVSToDS O;
-		
-		O.FactoryInterpolants = VertexFactoryInterpolate(a.FactoryInterpolants, aInterp, b.FactoryInterpolants, bInterp);
-
-		#if USE_RAW_WORLD_POSITION
-			TESSELLATION_INTERPOLATE_MEMBER(PixelPositionExcludingWPO);
-		#endif
-
-		return O;
-	}
-
-	FDepthOnlyVSToPS PassFinalizeTessellationOutput(FDepthOnlyVSToDS Interpolants, float4 WorldPosition, FMaterialTessellationParameters MaterialParameters)
-	{
-		FDepthOnlyVSToPS O;
-		
-		#if !MATERIALBLENDING_SOLID || OUTPUT_PIXEL_DEPTH_OFFSET
-			O.FactoryInterpolants = VertexFactoryAssignInterpolants(Interpolants.FactoryInterpolants);
-		#endif
-
-		// Finally, transform position to clip-space
-		{
-			O.Position = mul(WorldPosition, ResolvedView.TranslatedWorldToClip);
-		}
-			
-		#if !MATERIALBLENDING_SOLID || OUTPUT_PIXEL_DEPTH_OFFSET
-			O.PixelPosition = WorldPosition;
-		#endif
-
-		#if USE_RAW_WORLD_POSITION
-			O.PixelPositionExcludingWPO = Interpolants.PixelPositionExcludingWPO;
-		#endif
-		
-		return O;
-	}
-	
-	#include "Tessellation.ush"
-#endif
+#define FDepthOnlyVSOutput FDepthOnlyVSToPS
+#define VertexFactoryGetInterpolants VertexFactoryGetInterpolantsVSToPS
 
 #if VERTEXSHADER
 
@@ -92,7 +33,7 @@ void Main(
 	FVertexFactoryInput Input,
 	OPTIONAL_VertexID
 	out FDepthOnlyVSOutput Output
-#if USE_GLOBAL_CLIP_PLANE && !USING_TESSELLATION
+#if USE_GLOBAL_CLIP_PLANE
 	, out float OutGlobalClipPlaneDistance : SV_ClipDistance
 #endif
 #if INSTANCED_STEREO
@@ -131,10 +72,6 @@ void Main(
 		WorldPos.xyz += GetMaterialWorldPositionOffset(VertexParameters);
 	}
 
-#if USING_TESSELLATION
-	// Transformation is done in Domain shader when tessellating
-	Output.Position = WorldPos; 
-#else
 	{
 		float4 RasterizedWorldPosition = VertexFactoryGetRasterizedWorldPosition(Input, VFIntermediates, WorldPos);
 	#if ODS_CAPTURE
@@ -168,27 +105,20 @@ void Main(
 #if USE_GLOBAL_CLIP_PLANE
 	OutGlobalClipPlaneDistance = dot(ResolvedView.GlobalClippingPlane, float4(WorldPos.xyz - ResolvedView.PreViewTranslation.xyz, 1));
 #endif
-#endif
 
-#if !MATERIALBLENDING_SOLID || USING_TESSELLATION || OUTPUT_PIXEL_DEPTH_OFFSET
+#if !MATERIALBLENDING_SOLID || OUTPUT_PIXEL_DEPTH_OFFSET
 	// Masked and transparent materials need texture coords to clip, and tessellated
 	// materials need texture coords to displace
 	Output.FactoryInterpolants = VertexFactoryGetInterpolants(Input, VFIntermediates, VertexParameters);
 
 #if INSTANCED_STEREO
-	#if USING_TESSELLATION	
-		Output.FactoryInterpolants.InterpolantsVSToPS.EyeIndex = EyeIndex;
-	#else
-		Output.FactoryInterpolants.EyeIndex = EyeIndex;
-	#endif
+	Output.FactoryInterpolants.EyeIndex = EyeIndex;
 #endif
 
 #endif
 
 #if !MATERIALBLENDING_SOLID || OUTPUT_PIXEL_DEPTH_OFFSET
-	#if !USING_TESSELLATION	
-		Output.PixelPosition = WorldPos;
-	#endif
+	Output.PixelPosition = WorldPos;
 #endif
 
 #if USE_RAW_WORLD_POSITION
diff --git a/Engine/Shaders/Private/DistortAccumulateVS.usf b/Engine/Shaders/Private/DistortAccumulateVS.usf
--- a/Engine/Shaders/Private/DistortAccumulateVS.usf
+++ b/Engine/Shaders/Private/DistortAccumulateVS.usf
@@ -25,51 +25,7 @@ struct FDistortAccumulateVSToPS
 	float4	Position		: SV_POSITION;
 };
 
-#if USING_TESSELLATION	
-	struct FDistortAccumulateVSToDS
-	{
-		FVertexFactoryInterpolantsVSToDS FactoryInterpolants;
-		float4 Position : VS_To_DS_Position;
-		OPTIONAL_VertexID_VS_To_DS
-	};
-	
-	#define VS_OUTPUT_TYPE FDistortAccumulateVSToDS
-#else
-	#define VS_OUTPUT_TYPE FDistortAccumulateVSToPS
-#endif
-
-#if USING_TESSELLATION	
-	#define FPassSpecificVSToDS FDistortAccumulateVSToDS
-	#define FPassSpecificVSToPS FDistortAccumulateVSToPS
-	#define INTERPOLATE_MEMBER(member) O.member = a.member * aInterp + b.member * bInterp
-
-	FDistortAccumulateVSToDS PassInterpolate(FDistortAccumulateVSToDS a, float aInterp, FDistortAccumulateVSToDS b, float bInterp)
-	{
-		FDistortAccumulateVSToDS O;
-		
-		O.FactoryInterpolants = VertexFactoryInterpolate(a.FactoryInterpolants, aInterp, b.FactoryInterpolants, bInterp);
-
-		return O;
-	}
-
-	#undef INTERPOLATE_MEMBER
-	
-	FDistortAccumulateVSToPS PassFinalizeTessellationOutput(FDistortAccumulateVSToDS Interpolants, float4 WorldPosition, FMaterialTessellationParameters MaterialParameters)
-	{
-		FDistortAccumulateVSToPS O;
-		
-		O.FactoryInterpolants = VertexFactoryAssignInterpolants(Interpolants.FactoryInterpolants);
-		
-		// Finally, transform position to clip-space
-		O.Position = mul(WorldPosition, ResolvedView.TranslatedWorldToClip);
-			
-		O.PixelPosition = WorldPosition;
-		
-		return O;
-	}
-	
-	#include "Tessellation.ush"
-#endif
+#define VS_OUTPUT_TYPE FDistortAccumulateVSToPS
 
 #if VERTEXSHADER
 
@@ -78,7 +34,7 @@ void Main(
 	FVertexFactoryInput Input,
 	OPTIONAL_VertexID
 	out VS_OUTPUT_TYPE Output
-#if USE_GLOBAL_CLIP_PLANE && !USING_TESSELLATION
+#if USE_GLOBAL_CLIP_PLANE
 	, out float OutGlobalClipPlaneDistance : SV_ClipDistance
 #endif
 #if INSTANCED_STEREO
@@ -116,12 +72,6 @@ void Main(
 		WorldPosition.xyz += GetMaterialWorldPositionOffset(VertexParameters);
 	}
 	
-#if USING_TESSELLATION
-	Output.Position = WorldPosition;
-	Output.FactoryInterpolants = VertexFactoryGetInterpolantsVSToDS(Input, VFIntermediates, VertexParameters);
-	
-#else // !USING_TESSELLATION
-
 	Output.Position = mul(WorldPosition, ResolvedView.TranslatedWorldToClip);
 	Output.FactoryInterpolants = VertexFactoryGetInterpolantsVSToPS(Input, VFIntermediates, VertexParameters);
 
@@ -146,14 +96,8 @@ void Main(
 
 	Output.PixelPosition = WorldPosition;
 
-#endif // USING_TESSELLATION
-
 #if INSTANCED_STEREO
-	#if USING_TESSELLATION	
-		Output.FactoryInterpolants.InterpolantsVSToPS.EyeIndex = EyeIndex;
-	#else
-		Output.FactoryInterpolants.EyeIndex = EyeIndex;
-	#endif
+	Output.FactoryInterpolants.EyeIndex = EyeIndex;
 #endif
 
 	OutputVertexID( Output );
diff --git a/Engine/Shaders/Private/FlatTessellation.ush b/Engine/Shaders/Private/FlatTessellation.ush
deleted file mode 100644
diff --git a/Engine/Shaders/Private/HitProxyVertexShader.usf b/Engine/Shaders/Private/HitProxyVertexShader.usf
--- a/Engine/Shaders/Private/HitProxyVertexShader.usf
+++ b/Engine/Shaders/Private/HitProxyVertexShader.usf
@@ -31,72 +31,7 @@ struct FHitProxyVSToPS
 	float4 Position : SV_POSITION;
 };
 
-#if USING_TESSELLATION	
-	struct FHitProxyVSToDS
-	{
-		FVertexFactoryInterpolantsVSToDS FactoryInterpolants;
-	    #if PER_INSTANCE_HITPROXY_ID
-		    float4 InstanceHitProxyId : HIT_PROXY_ID;
-	    #endif
-		float4 Position : VS_To_DS_Position;
-		#if USE_WORLD_POSITION_EXCLUDING_SHADER_OFFSETS
-			float3 PixelPositionExcludingWPO : TEXCOORD7;
-		#endif
-		OPTIONAL_VertexID_VS_To_DS
-	};
-	
-	#define FHitProxyVSOutput FHitProxyVSToDS
-#else
-	#define FHitProxyVSOutput FHitProxyVSToPS
-#endif
-
-#if USING_TESSELLATION	
-	#define FPassSpecificVSToDS FHitProxyVSToDS
-	#define FPassSpecificVSToPS FHitProxyVSToPS
-	
-	FHitProxyVSToDS PassInterpolate(FHitProxyVSToDS a, float aInterp, FHitProxyVSToDS b, float bInterp)
-	{
-		FHitProxyVSToDS O;
-
-		O.FactoryInterpolants = VertexFactoryInterpolate(a.FactoryInterpolants, aInterp, b.FactoryInterpolants, bInterp);
-
-		#if PER_INSTANCE_HITPROXY_ID
-			TESSELLATION_INTERPOLATE_MEMBER(InstanceHitProxyId);
-		#endif
-
-		#if USE_WORLD_POSITION_EXCLUDING_SHADER_OFFSETS
-			TESSELLATION_INTERPOLATE_MEMBER(PixelPositionExcludingWPO);
-		#endif
-
-		return O;
-	}
-
-	FHitProxyVSToPS PassFinalizeTessellationOutput(FHitProxyVSToDS Interpolants, float4 WorldPosition, FMaterialTessellationParameters MaterialParameters)
-	{
-		FHitProxyVSToPS O;
-		
-		O.FactoryInterpolants = VertexFactoryAssignInterpolants(Interpolants.FactoryInterpolants);
-		
-		#if PER_INSTANCE_HITPROXY_ID
-			O.InstanceHitProxyId = Interpolants.InstanceHitProxyId;
-		#endif
-
-		// Finally, transform position to clip-space
-		{
-			O.Position = mul(WorldPosition, View.TranslatedWorldToClip);
-		}
-			
-		O.PixelPosition = WorldPosition;
-
-		#if USE_WORLD_POSITION_EXCLUDING_SHADER_OFFSETS
-			O.PixelPositionExcludingWPO = Interpolants.PixelPositionExcludingWPO.xyz;
-		#endif
-		
-		return O;
-	}
-	
-	#include "Tessellation.ush"
-#endif
+#define FHitProxyVSOutput FHitProxyVSToPS
 
 #if VERTEXSHADER
 
@@ -104,7 +39,7 @@ void Main(
 	FVertexFactoryInput Input,
 	OPTIONAL_VertexID
 	out FHitProxyVSOutput Output
-#if USE_GLOBAL_CLIP_PLANE && !USING_TESSELLATION
+#if USE_GLOBAL_CLIP_PLANE
 	, out float OutGlobalClipPlaneDistance : SV_ClipDistance
 #endif
 	)
@@ -129,11 +64,6 @@ void Main(
 	Output.InstanceHitProxyId = VertexFactoryGetInstanceHitProxyId(Input, VFIntermediates);
 #endif
 
-#if USING_TESSELLATION
-	Output.Position = WorldPosition;
-	Output.FactoryInterpolants = VertexFactoryGetInterpolantsVSToDS(Input, VFIntermediates, VertexParameters);
-#else
-
 	{
 		float4 RasterizedWorldPosition = VertexFactoryGetRasterizedWorldPosition(Input, VFIntermediates, WorldPosition);
 		Output.Position = mul(RasterizedWorldPosition, View.TranslatedWorldToClip);
@@ -145,7 +75,6 @@ void Main(
 
 	Output.FactoryInterpolants = VertexFactoryGetInterpolantsVSToPS(Input, VFIntermediates, VertexParameters);
 	Output.PixelPosition = WorldPosition;
-#endif
 
 #if USE_WORLD_POSITION_EXCLUDING_SHADER_OFFSETS
 	Output.PixelPositionExcludingWPO = WorldPositionExcludingWPO.xyz;
diff --git a/Engine/Shaders/Private/LightMapDensityShader.usf b/Engine/Shaders/Private/LightMapDensityShader.usf
--- a/Engine/Shaders/Private/LightMapDensityShader.usf
+++ b/Engine/Shaders/Private/LightMapDensityShader.usf
@@ -20,61 +20,7 @@ struct FLightMapDensityVSToPS
 	float4 Position			: SV_POSITION;
 };
 
-#if USING_TESSELLATION	
-	struct FLightMapDensityVSToDS
-	{
-		FVertexFactoryInterpolantsVSToDS FactoryInterpolants;
-		float4 Position	: VS_To_DS_Position;
-		#if USE_WORLD_POSITION_EXCLUDING_SHADER_OFFSETS
-			float3 PixelPositionExcludingWPO : TEXCOORD7;
-		#endif
-		OPTIONAL_VertexID_VS_To_DS
-	};
-	
-	#define FLightMapDensityVSOutput FLightMapDensityVSToDS
-#else
-	#define FLightMapDensityVSOutput FLightMapDensityVSToPS
-#endif
-
-#if USING_TESSELLATION	
-	#define FPassSpecificVSToDS FLightMapDensityVSToDS
-	#define FPassSpecificVSToPS FLightMapDensityVSToPS
-	
-	FLightMapDensityVSToDS PassInterpolate(FLightMapDensityVSToDS a, float aInterp, FLightMapDensityVSToDS b, float bInterp)
-	{
-		FLightMapDensityVSToDS O;
-		
-		O.FactoryInterpolants = VertexFactoryInterpolate(a.FactoryInterpolants, aInterp, b.FactoryInterpolants, bInterp);
-
-		#if USE_WORLD_POSITION_EXCLUDING_SHADER_OFFSETS
-			TESSELLATION_INTERPOLATE_MEMBER(PixelPositionExcludingWPO);
-		#endif
-
-		return O;
-	}
-
-	FLightMapDensityVSToPS PassFinalizeTessellationOutput(FLightMapDensityVSToDS Interpolants, float4 WorldPosition, FMaterialTessellationParameters MaterialParameters)
-	{
-		FLightMapDensityVSToPS O;
-		
-		O.FactoryInterpolants = VertexFactoryAssignInterpolants(Interpolants.FactoryInterpolants);
-		
-		// Finally, transform position to clip-space
-		{
-			O.Position = mul(WorldPosition, ResolvedView.TranslatedWorldToClip);
-		}
-
-		O.WorldPosition = WorldPosition;
-
-		#if USE_WORLD_POSITION_EXCLUDING_SHADER_OFFSETS
-			O.PixelPositionExcludingWPO = Interpolants.PixelPositionExcludingWPO.xyz;
-		#endif
-		
-		return O;
-	}
-	
-	#include "Tessellation.ush"
-#endif
+#define FLightMapDensityVSOutput FLightMapDensityVSToPS
 
 /*=============================================================================
 	Vertex Shader
@@ -84,7 +30,7 @@ void MainVertexShader(
 	FVertexFactoryInput Input,
 	OPTIONAL_VertexID
 	out FLightMapDensityVSOutput Output
-#if USE_GLOBAL_CLIP_PLANE && !USING_TESSELLATION
+#if USE_GLOBAL_CLIP_PLANE
 	, out float OutGlobalClipPlaneDistance : SV_ClipDistance
 #endif
 	)
@@ -99,12 +45,6 @@ void MainVertexShader(
 	FMaterialVertexParameters VertexParameters = GetMaterialVertexParameters(Input, VFIntermediates, WorldPosition.xyz, TangentToLocal);
 	WorldPosition.xyz += GetMaterialWorldPositionOffset(VertexParameters);
 
-#if USING_TESSELLATION
-	Output.Position = WorldPosition;
-	Output.FactoryInterpolants = VertexFactoryGetInterpolantsVSToDS(Input, VFIntermediates, VertexParameters);
-	
-#else // !USING_TESSELLATION
-
 	Output.WorldPosition = WorldPosition;
 	{
 		float4 RasterizedWorldPosition = VertexFactoryGetRasterizedWorldPosition(Input, VFIntermediates, Output.WorldPosition);
@@ -120,8 +60,6 @@ void MainVertexShader(
 #endif
 	Output.FactoryInterpolants = VertexFactoryGetInterpolantsVSToPS(Input, VFIntermediates, VertexParameters);
 
-#endif // USING_TESSELLATION
-
 	OutputVertexID( Output );
 }
 #endif
diff --git a/Engine/Shaders/Private/MeshDecals.usf b/Engine/Shaders/Private/MeshDecals.usf
--- a/Engine/Shaders/Private/MeshDecals.usf
+++ b/Engine/Shaders/Private/MeshDecals.usf
@@ -41,59 +41,7 @@ struct FMeshDecalVSToPS
 	float4	Position		: SV_POSITION;
 };
 
-#if USING_TESSELLATION	
-	struct FMeshDecalVSToDS
-	{
-		FVertexFactoryInterpolantsVSToDS FactoryInterpolants;
-	#if NEED_MESH_DECAL_INTERPOLATORS
-		FMeshDecalInterpolants MeshDecalInterpolants;
-	#endif
-		float4 Position : VS_To_DS_Position;
-		OPTIONAL_VertexID_VS_To_DS
-	};
-	
-	#define VS_OUTPUT_TYPE FMeshDecalVSToDS
-#else
-	#define VS_OUTPUT_TYPE FMeshDecalVSToPS
-#endif
-
-#if USING_TESSELLATION	
-	#define FPassSpecificVSToDS FMeshDecalVSToDS
-	#define FPassSpecificVSToPS FMeshDecalVSToPS
-	#define INTERPOLATE_MEMBER(member) O.member = a.member * aInterp + b.member * bInterp
-
-	FMeshDecalVSToDS PassInterpolate(FMeshDecalVSToDS a, float aInterp, FMeshDecalVSToDS b, float bInterp)
-	{
-		FMeshDecalVSToDS O;
-		
-		O.FactoryInterpolants = VertexFactoryInterpolate(a.FactoryInterpolants, aInterp, b.FactoryInterpolants, bInterp);
-
-	#if USE_WORLD_POSITION_EXCLUDING_SHADER_OFFSETS
-		INTERPOLATE_MEMBER(MeshDecalInterpolants.PixelPositionExcludingWPO);
-	#endif
-
-		return O;
-	}
-
-	#undef INTERPOLATE_MEMBER
-	
-	FMeshDecalVSToPS PassFinalizeTessellationOutput(FMeshDecalVSToDS Interpolants, float4 WorldPosition, FMaterialTessellationParameters MaterialParameters)
-	{
-		FMeshDecalVSToPS O;
-		
-		O.FactoryInterpolants = VertexFactoryAssignInterpolants(Interpolants.FactoryInterpolants);
-
-	#if USE_WORLD_POSITION_EXCLUDING_SHADER_OFFSETS
-		O.MeshDecalInterpolants.PixelPositionExcludingWPO = Interpolants.MeshDecalInterpolants.PixelPositionExcludingWPO;
-	#endif
-		
-		// Finally, transform position to clip-space
-		O.Position = mul(WorldPosition, View.TranslatedWorldToClip);
-		return O;
-	}
-	
-	#include "Tessellation.ush"
-#endif
+#define VS_OUTPUT_TYPE FMeshDecalVSToPS
 
 #if VERTEXSHADER
 
@@ -102,7 +50,7 @@ void MainVS(
 	FVertexFactoryInput Input,
 	OPTIONAL_VertexID
 	out VS_OUTPUT_TYPE Output
-#if USE_GLOBAL_CLIP_PLANE && !USING_TESSELLATION
+#if USE_GLOBAL_CLIP_PLANE
 	, out float OutGlobalClipPlaneDistance : SV_ClipDistance
 #endif
 	)
@@ -125,12 +73,6 @@ void MainVS(
 		WorldPosition.xyz += GetMaterialWorldPositionOffset(VertexParameters);
 	}
 	
-#if USING_TESSELLATION
-	Output.Position = WorldPosition;
-	Output.FactoryInterpolants = VertexFactoryGetInterpolantsVSToDS(Input, VFIntermediates, VertexParameters);
-	
-#else // !USING_TESSELLATION
-
 	Output.Position = mul(WorldPosition, View.TranslatedWorldToClip);
 	Output.FactoryInterpolants = VertexFactoryGetInterpolantsVSToPS(Input, VFIntermediates, VertexParameters);
 
@@ -138,9 +80,7 @@ void MainVS(
 	OutGlobalClipPlaneDistance = dot(ResolvedView.GlobalClippingPlane, float4(WorldPosition.xyz - ResolvedView.PreViewTranslation.xyz, 1));
 #endif
 
-#endif // USING_TESSELLATION
-
-#if !USING_TESSELLATION && HAS_INVERTED_Z_BUFFER
+#if HAS_INVERTED_Z_BUFFER
 	// Move all geometry a small amount towards the camera to avoid z-fighting.
 	// Assuming a reverse z buffer this pushes the near plane a fixed small distance.
 	Output.Position.z += ResolvedView.DecalDepthBias;
diff --git a/Engine/Shaders/Private/PNTriangles.ush b/Engine/Shaders/Private/PNTriangles.ush
deleted file mode 100644
diff --git a/Engine/Shaders/Private/ShadowDepthPixelShader.usf b/Engine/Shaders/Private/ShadowDepthPixelShader.usf
--- a/Engine/Shaders/Private/ShadowDepthPixelShader.usf
+++ b/Engine/Shaders/Private/ShadowDepthPixelShader.usf
@@ -14,7 +14,7 @@
 
 
 #if ENABLE_NON_NANITE_VSM
-	#define VIRTUAL_SM_ENABLED (!(USING_TESSELLATION || ONEPASS_POINTLIGHT_SHADOW || PERSPECTIVE_CORRECT_DEPTH))
+	#define VIRTUAL_SM_ENABLED (!(ONEPASS_POINTLIGHT_SHADOW || PERSPECTIVE_CORRECT_DEPTH))
 #else
 	#define VIRTUAL_SM_ENABLED 0
 #endif
diff --git a/Engine/Shaders/Private/ShadowDepthVertexShader.usf b/Engine/Shaders/Private/ShadowDepthVertexShader.usf
--- a/Engine/Shaders/Private/ShadowDepthVertexShader.usf
+++ b/Engine/Shaders/Private/ShadowDepthVertexShader.usf
@@ -14,7 +14,7 @@
 
 
 #if ENABLE_NON_NANITE_VSM
-	#define VIRTUAL_SM_ENABLED (!(USING_TESSELLATION || ONEPASS_POINTLIGHT_SHADOW || PERSPECTIVE_CORRECT_DEPTH))
+	#define VIRTUAL_SM_ENABLED (!(ONEPASS_POINTLIGHT_SHADOW || PERSPECTIVE_CORRECT_DEPTH))
 #else
 	#define VIRTUAL_SM_ENABLED 0
 #endif
@@ -297,105 +297,9 @@ void MainOnePassPointLightGS(triangle FShadowDepthVSToGS Input[3], inout Triangl
 
 #endif
 
-#if USING_TESSELLATION
-
-struct FShadowDepthVSToDS
-{
-	FVertexFactoryInterpolantsVSToDS FactoryInterpolants;
-	float4 Position : VS_to_DS_Position;
-	OPTIONAL_VertexID_VS_To_DS
-#if ONEPASS_POINTLIGHT_SHADOW && USING_VERTEX_SHADER_LAYER
-	uint LayerIndex : VS_to_DS_RenderTargetArrayIndex;
-#endif
-};
-
-struct FShadowDepthDSToPS : FShadowDepthVSToPS
-{
-#if ONEPASS_POINTLIGHT_SHADOW
-#if USING_VERTEX_SHADER_LAYER
-	float4 Position : SV_POSITION;
-	float4 InterpPosition : TEXCOORD6;
-	uint LayerIndex : SV_RenderTargetArrayIndex;
-#else
-	float4 Position : TEXCOORD6;
-#endif
-	float3 WorldNormal : TEXCOORD8;
-#else
-	float4 Position : SV_POSITION;
-#endif
-};
-
-#define FPassSpecificVSToDS FShadowDepthVSToDS
-#define FPassSpecificVSToPS FShadowDepthDSToPS
-
-FShadowDepthVSToDS PassInterpolate(FShadowDepthVSToDS a, float aInterp, FShadowDepthVSToDS b, float bInterp)
-{
-	FShadowDepthVSToDS O;
-	O.FactoryInterpolants = VertexFactoryInterpolate(a.FactoryInterpolants, aInterp, b.FactoryInterpolants, bInterp);
-#if ONEPASS_POINTLIGHT_SHADOW && USING_VERTEX_SHADER_LAYER
-	O.LayerIndex = a.LayerIndex;
-#endif
-	return O;
-}
-
-FShadowDepthDSToPS PassFinalizeTessellationOutput(FShadowDepthVSToDS Interpolants, float4 WorldPosition, FMaterialTessellationParameters MaterialParameters)
-{
-	FShadowDepthDSToPS O;
-	
-#if INTERPOLATE_VF_ATTRIBUTES
-	O.FactoryInterpolants = VertexFactoryAssignInterpolants(Interpolants.FactoryInterpolants);
-#endif
-
-	const float3 WorldNormal = MaterialParameters.TangentToWorld[2];
-
-#if ONEPASS_POINTLIGHT_SHADOW
-#if USING_VERTEX_SHADER_LAYER
-	O.LayerIndex = Interpolants.LayerIndex;
-	O.Position = WorldPosition = mul(WorldPosition, PassStruct.ShadowViewProjectionMatrices[Interpolants.LayerIndex]);
-	O.InterpPosition = WorldPosition;
-#else
-	O.Position = WorldPosition;
-#endif // USING_VERTEX_SHADER_LAYER
-	O.WorldNormal = WorldNormal;
-#else
-
-	float ShadowDepth;
-	SetShadowDepthOutputs(
-		PassStruct.ProjectionMatrix,
-		PassStruct.ViewMatrix,
-		WorldPosition, 
-		WorldNormal,
-		O.Position, 
-#if PERSPECTIVE_CORRECT_DEPTH
-		O.ShadowDepth,
-		O.DepthBias
-#else
-		ShadowDepth
-#endif
-		);
-#endif
-
-// Conditions below are derived. See struct FShadowDepthVSToPS in ShadowDepthCommon.ush
-#if !PERSPECTIVE_CORRECT_DEPTH && !COMPILER_SUPPORTS_EMPTY_STRUCTS && !ONEPASS_POINTLIGHT_SHADOW
-	O.Dummy = 0;
-#endif
-
-#if INTERPOLATE_POSITION
-	O.PixelPosition = WorldPosition.xyz;
-#endif
-
-	return O;
-}
-	
-#include "Tessellation.ush"
-
-#endif // #if USING_TESSELLATION
-
 void Main(
 	FVertexFactoryInput Input,
-#if USING_TESSELLATION
-	out FShadowDepthVSToDS OutParameters,
-#elif ONEPASS_POINTLIGHT_SHADOW
+#if ONEPASS_POINTLIGHT_SHADOW
 	out FShadowDepthVSToGS OutParameters,
 #if FEATURE_LEVEL >= FEATURE_LEVEL_SM4 && USING_VERTEX_SHADER_LAYER
 	out uint LayerIndex : SV_RenderTargetArrayIndex,
@@ -425,22 +329,7 @@ void Main(
 
 	WorldPos.xyz += GetMaterialWorldPositionOffset(VertexParameters);
 
-#if USING_TESSELLATION
-	OutPosition = OutParameters.Position = WorldPos;
-	
-	// Tessellated materials need vertex coords for displacement
-	OutParameters.FactoryInterpolants = VertexFactoryGetInterpolantsVSToDS(Input, VFIntermediates, VertexParameters);
-
-	#if ONEPASS_POINTLIGHT_SHADOW && USING_VERTEX_SHADER_LAYER && VERTEXSHADER
-#if GPUCULL_TODO
-		OutParameters.LayerIndex = bUseGpuSceneInstancing ? VertexFactoryGetViewIndex(VFIntermediates) : LayerId;
-#else //!GPUCULL_TODO
-		OutParameters.LayerIndex = LayerId;
-#endif // GPUCULL_TODO
-
-	#endif
-
-#elif ONEPASS_POINTLIGHT_SHADOW
+#if ONEPASS_POINTLIGHT_SHADOW
 
 	OutPosition = OutParameters.GSPosition = WorldPos;
 	OutParameters.WorldNormal = WorldNormal;
@@ -527,9 +416,7 @@ void Main(
 void MainForGS(
 	OPTIONAL_VertexID
 	FVertexFactoryInput Input,
-#if USING_TESSELLATION
-	out FShadowDepthVSToDS OutParameters
-#elif ONEPASS_POINTLIGHT_SHADOW
+#if ONEPASS_POINTLIGHT_SHADOW
 	out FShadowDepthVSToGS OutParameters
 	#if FEATURE_LEVEL >= FEATURE_LEVEL_SM4 && USING_VERTEX_SHADER_LAYER
 		, out uint LayerIndex : SV_RenderTargetArrayIndex
@@ -551,7 +438,7 @@ void MainForGS(
 	Main(Input,
 		OutParameters,
 		OPTIONAL_VertexID_PARAM
-#if !USING_TESSELLATION && (FEATURE_LEVEL >= FEATURE_LEVEL_SM4 && ONEPASS_POINTLIGHT_SHADOW && USING_VERTEX_SHADER_LAYER)
+#if FEATURE_LEVEL >= FEATURE_LEVEL_SM4 && ONEPASS_POINTLIGHT_SHADOW && USING_VERTEX_SHADER_LAYER
 		LayerIndex,
 #endif
 		UnusedPosition);
diff --git a/Engine/Shaders/Private/Tessellation.ush b/Engine/Shaders/Private/Tessellation.ush
deleted file mode 100644
diff --git a/Engine/Shaders/Private/VelocityShader.usf b/Engine/Shaders/Private/VelocityShader.usf
--- a/Engine/Shaders/Private/VelocityShader.usf
+++ b/Engine/Shaders/Private/VelocityShader.usf
@@ -33,73 +33,9 @@ struct FVelocityVSToPS
 	FVelocityInterpsVSToPS Interps;
 };
 
-#if USING_TESSELLATION
+#define FVertexOutput FVelocityVSToPS
+#define VertexFactoryGetInterpolants VertexFactoryGetInterpolantsVSToPS
 
-	struct FVelocityInterpsVSToDS
-	{
-		// float4(ScreenPos.xyzw)
-		float4 PackedVelocityA : TEXCOORD6;
-		// float4(PrevScreenPos.xyzw)
-		float4 PackedVelocityC : TEXCOORD8;
-	};
-
-	struct FVelocityVSToDS
-	{
-		FVelocityInterpsVSToDS Interps;
-		FVertexFactoryInterpolantsVSToDS FactoryInterpolants;
-		float4 Position : VS_To_DS_Position;
-		OPTIONAL_VertexID_VS_To_DS
-	};
-	
-	#define FVertexOutput FVelocityVSToDS
-	#define VertexFactoryGetInterpolants VertexFactoryGetInterpolantsVSToDS
-#else
-	#define FVertexOutput FVelocityVSToPS
-	#define VertexFactoryGetInterpolants VertexFactoryGetInterpolantsVSToPS
-#endif
-
-#if USING_TESSELLATION
-	#define FPassSpecificVSToDS FVelocityVSToDS
-	#define FPassSpecificVSToPS FVelocityVSToPS
-	
-	#define VELOCITY_INTERPOLATE_MEMBER(member) O.member = a.member * aInterp + b.member * bInterp
-	
-	FVelocityVSToDS PassInterpolate(FVelocityVSToDS a, float aInterp, FVelocityVSToDS b, float bInterp)
-	{
-		FVelocityVSToDS O;
-
-		O.FactoryInterpolants = VertexFactoryInterpolate(a.FactoryInterpolants, aInterp, b.FactoryInterpolants, bInterp);
-		
-		VELOCITY_INTERPOLATE_MEMBER(Interps.PackedVelocityA);
-		VELOCITY_INTERPOLATE_MEMBER(Interps.PackedVelocityC);
-
-		return O;
-	}
-	
-	#undef VELOCITY_INTERPOLATE_MEMBER
-	
-	FVelocityVSToPS PassFinalizeTessellationOutput(FVelocityVSToDS Interpolants, float4 WorldPosition, FMaterialTessellationParameters MaterialParameters)
-	{
-		FVelocityVSToPS O;
-		
-		O.Interps.FactoryInterpolants = VertexFactoryAssignInterpolants(Interpolants.FactoryInterpolants);
-
-		{
-			O.Position = mul(WorldPosition, ResolvedView.TranslatedWorldToClip);
-		}
-
-		// Move all geometry a little bit towards the camera (to not get z fighting with existing zbuffer of the same meshes rendered with slightly different float computations).
-		O.Position.z += GTessellationDepthBias * O.Position.w;
-		
-		O.Interps.PackedVelocityA = Interpolants.Interps.PackedVelocityA;
-		O.Interps.PackedVelocityC = Interpolants.Interps.PackedVelocityC;
-		
-		return O;
-	}
-	
-	#include "Tessellation.ush"
-#endif
-	
 /*=============================================================================
  * Vertex Shader
  *============================================================================*/
@@ -108,7 +44,7 @@ void MainVertexShader(
 	FVertexFactoryInput Input,
 	OPTIONAL_VertexID
 	out FVertexOutput Output 
-#if USE_GLOBAL_CLIP_PLANE && !USING_TESSELLATION
+#if USE_GLOBAL_CLIP_PLANE
 	, out float OutGlobalClipPlaneDistance : SV_ClipDistance
 #endif
 #if INSTANCED_STEREO
@@ -145,19 +81,12 @@ void MainVertexShader(
 		TranslatedWorldPosition.xyz += GetMaterialWorldPositionOffset(VertexParameters);
 	}
 
-#if USING_TESSELLATION
-	Output.FactoryInterpolants = VertexFactoryGetInterpolants( Input, VFIntermediates, VertexParameters );
-#else
 	Output.Interps.FactoryInterpolants = VertexFactoryGetInterpolants( Input, VFIntermediates, VertexParameters );
-#endif // #if USING_TESSELLATION
+
 	float4 PrevTranslatedWorldPosition = VertexFactoryGetPreviousWorldPosition( Input, VFIntermediates );	
 	VertexParameters = GetMaterialVertexParameters(Input, VFIntermediates, PrevTranslatedWorldPosition.xyz, TangentToLocal);
 	PrevTranslatedWorldPosition.xyz += GetMaterialPreviousWorldPositionOffset(VertexParameters);
 	
-#if USING_TESSELLATION
-	Output.Position = TranslatedWorldPosition;
-#endif
-
 	{
 		float4 ScreenPos;
 		// compute the old screen pos with the old world position and the old camera matrix
@@ -169,9 +98,7 @@ void MainVertexShader(
 		float4 PrevScreenPosObj = mul(float4(PrevTranslatedWorldPosition.xyz, 1), ResolvedView.PrevTranslatedWorldToClip);
 	#endif
 
-#if !USING_TESSELLATION
 		{
-#endif
 
 		#if ODS_CAPTURE
 			float3 ODS = OffsetODS(TranslatedWorldPosition.xyz, ResolvedView.TranslatedWorldCameraOrigin.xyz, ResolvedView.StereoIPD);
@@ -180,10 +107,8 @@ void MainVertexShader(
 			ScreenPos = mul(float4(TranslatedWorldPosition.xyz, 1), ResolvedView.TranslatedWorldToClip);
 		#endif
 			
-#if !USING_TESSELLATION
 			Output.Position = ScreenPos;
 		}
-#endif
 
 	#if INSTANCED_STEREO && !MULTI_VIEW
 	BRANCH
@@ -200,7 +125,7 @@ void MainVertexShader(
 	}
 	#endif
 
-#if USE_GLOBAL_CLIP_PLANE && !USING_TESSELLATION
+#if USE_GLOBAL_CLIP_PLANE
 		OutGlobalClipPlaneDistance = dot(ResolvedView.GlobalClippingPlane, float4(TranslatedWorldPosition.xyz - ResolvedView.PreViewTranslation.xyz, 1));
 #endif
 
@@ -209,20 +134,10 @@ void MainVertexShader(
 	}
 
 // Move all geometry a little bit towards the camera (to not get z fighting with existing zbuffer of the same meshes rendered with slightly different float computations).
-#if !USING_TESSELLATION
     Output.Position.z += View.NumSceneColorMSAASamples > 1 ? GNonTessellationDepthBias * Output.Position.w : 0.0;
-#endif
-
-#if USING_TESSELLATION
-	OutputVertexID( Output );
-#endif
 
 #if INSTANCED_STEREO
-	#if USING_TESSELLATION	
-		Output.FactoryInterpolants.InterpolantsVSToPS.EyeIndex = EyeIndex;
-	#else
-		Output.Interps.FactoryInterpolants.EyeIndex = EyeIndex;
-	#endif
+	Output.Interps.FactoryInterpolants.EyeIndex = EyeIndex;
 #endif
 }
 #endif // VERTEXSHADER
diff --git a/Engine/Shaders/Public/Platform.ush b/Engine/Shaders/Public/Platform.ush
--- a/Engine/Shaders/Public/Platform.ush
+++ b/Engine/Shaders/Public/Platform.ush
@@ -355,11 +355,6 @@
 
 // ---------------------------------------------------- Compiler specific defaults and fallbacks
 
-/** Defined only for Metal's combined Vertex + Hull shader */
-#ifndef TESSELLATIONSHADER
-#define TESSELLATIONSHADER 0
-#endif
-
 // Hlslcc platforms ignore the uniform keyword as it can't properly optimize flow
 #if COMPILER_HLSLCC
 #define uniform
diff --git a/Engine/Shaders/Public/Platform/Metal/MetalCommon.ush b/Engine/Shaders/Public/Platform/Metal/MetalCommon.ush
--- a/Engine/Shaders/Public/Platform/Metal/MetalCommon.ush
+++ b/Engine/Shaders/Public/Platform/Metal/MetalCommon.ush
@@ -9,12 +9,6 @@
 
 #pragma once
 
-#if VERTEXSHADER && HULLSHADER
-#ifndef TESSELLATIONSHADER
-#define TESSELLATIONSHADER 1
-#endif
-#endif
-
 #define STRONG_TYPE
 #define StrongTypedBuffer Buffer
 
