Made starta weight operation clamp the weight to safe range.

diff --git a/Engine/Shaders/Private/Strata/Strata.ush b/Engine/Shaders/Private/Strata/Strata.ush
--- a/Engine/Shaders/Private/Strata/Strata.ush
+++ b/Engine/Shaders/Private/Strata/Strata.ush
@@ -807,6 +807,7 @@ FStrataData StrataAdd(FStrataData A, FStrataData B)
 
 FStrataData StrataMultiply(FStrataData A, float Weight)
 {
+	const float SafeWeight = saturate(Weight);
 	FStrataData StrataData = A;
 
 	// Apply the weight to each BSDF
@@ -816,7 +817,7 @@ FStrataData StrataMultiply(FStrataData A, float Weight)
 		UNROLL
 		for (int BSDF = 0; BSDF < StrataData.Layers[Layer].BSDFCount; ++BSDF)
 		{
-			StrataData.Layers[Layer].BSDFs[BSDF].Weight *= Weight;
+			StrataData.Layers[Layer].BSDFs[BSDF].Weight *= SafeWeight;
 		}
 	}
 
