Fix page invalidation for VSM atomic raster path

diff --git a/Engine/Shaders/Private/VirtualShadowMaps/BuildPerPageDrawCommands.usf b/Engine/Shaders/Private/VirtualShadowMaps/BuildPerPageDrawCommands.usf
--- a/Engine/Shaders/Private/VirtualShadowMaps/BuildPerPageDrawCommands.usf
+++ b/Engine/Shaders/Private/VirtualShadowMaps/BuildPerPageDrawCommands.usf
@@ -354,6 +354,25 @@ void CullPerPageDrawCommandsCs(uint InstanceIdIndex : SV_DispatchThreadID, uint
 
 						VisibleInstancesOut[VisibleInstanceOutputOffset] = VisibleInstanceCmd;
 						VisibleInstanceOutputOffset += 1U;
+
+						// Need to perform loop to take care of cache invalidations
+						if (bHasMoved)
+						{
+							uint PageTableLevelOffset = CalcPageTableLevelOffset(VirtualShadowMapId, MipLevel);
+							for (uint y = RectPages.y; y <= RectPages.w; y += 1U)
+							{
+								for (uint x = RectPages.x; x <= RectPages.z; x += 1U)
+								{
+									uint PageFlagOffset = PageTableLevelOffset + CalcPageOffsetInLevel(MipLevel, uint2(x, y));
+									uint PageFlag = PageFlags[PageFlagOffset];
+
+									if ((PageFlag & PageFlagMask) != 0)
+									{
+										OutDynamicCasterFlags[PageFlagOffset] = 1;
+									}
+								}
+							}
+						}
 					}
 				}
 			}
