Fixed: Shader requested a global uniform buffer of type 'MobileSceneTextures' at static slot

diff --git a/Engine/Shaders/Private/DecalCommon.ush b/Engine/Shaders/Private/DecalCommon.ush
--- a/Engine/Shaders/Private/DecalCommon.ush
+++ b/Engine/Shaders/Private/DecalCommon.ush
@@ -5,6 +5,12 @@
 =============================================================================*/
 
 #include "Common.ush"
+
+#if SHADING_PATH_MOBILE
+	// Reroute MobileSceneTextures uniform buffer references to the base pass uniform buffer 
+	#define MobileSceneTextures MobileBasePass.SceneTextures
+#endif
+
 #include "DeferredShadingCommon.ush"		// GBufferData
 #include "DBufferDecalShared.ush"
 
