Fix for error X3003: 'UseCSM' already defined as a variable

- Rename the function so it doesn't collide with variable name.

diff --git a/Engine/Shaders/Private/MobileBasePassPixelShader.usf b/Engine/Shaders/Private/MobileBasePassPixelShader.usf
--- a/Engine/Shaders/Private/MobileBasePassPixelShader.usf
+++ b/Engine/Shaders/Private/MobileBasePassPixelShader.usf
@@ -139,7 +139,7 @@ int NumDynamicPointLights;
 	uint UseCSM;
 #endif
 
-bool UseCSM()
+bool IsCSMEnabled()
 {
 #if MOBILE_USE_CSM_BRANCH
 	return UseCSM != 0u;
@@ -752,7 +752,7 @@ void Main(
 	float  ShadowPositionZ = 0;
 #if DIRECTIONAL_LIGHT_CSM && !MATERIAL_SHADINGMODEL_SINGLELAYERWATER
 	// Cascaded Shadow Map
-	if (UseCSM())
+	if (IsCSMEnabled())
 	{
 		half ShadowMap = MobileDirectionalLightCSM(MaterialParameters.ScreenPosition.xy, MaterialParameters.ScreenPosition.w, ShadowPositionZ);
 	#if ALLOW_STATIC_LIGHTING
