Fix Refraction == 0 producing NaNs in path tracer

This material parameter slot can become 0 in layered materials even when the material is tagged as having IOR refraction.

diff --git a/Engine/Shaders/Private/RayTracing/RayTracingMaterialHitShaders.usf b/Engine/Shaders/Private/RayTracing/RayTracingMaterialHitShaders.usf
--- a/Engine/Shaders/Private/RayTracing/RayTracingMaterialHitShaders.usf
+++ b/Engine/Shaders/Private/RayTracing/RayTracingMaterialHitShaders.usf
@@ -434,7 +434,7 @@ RAY_TRACING_ENTRY_CLOSEST_HIT(MaterialCHS,
 #endif
 #endif // MATERIALBLENDING_MASKED
 
-	Payload.Ior = Refraction.x;
+	Payload.Ior = Refraction.x == 0 ? 1.0 : Refraction.x;
 	Payload.CustomData = float4(GetMaterialCustomData0(MaterialParameters), GetMaterialCustomData1(MaterialParameters), 0, 0);
 	
 #if MATERIAL_SHADINGMODEL_CLEAR_COAT
@@ -673,6 +673,7 @@ RAY_TRACING_ENTRY_ANY_HIT(MaterialAHS,
 	float Opacity = GetMaterialOpacity(PixelMaterialInputs);
 #if REFRACTION_USE_INDEX_OF_REFRACTION
 	float Ior = GetMaterialRefraction(PixelMaterialInputs).x;
+	Ior = Ior == 0 ? 1.0 : Ior;
 #else
 	float Ior = 1.0;
 #endif
