Add option to skip velocity vector writting for low hair coverage.

diff --git a/Engine/Shaders/Private/HairStrands/HairStrandsVelocity.usf b/Engine/Shaders/Private/HairStrands/HairStrandsVelocity.usf
--- a/Engine/Shaders/Private/HairStrands/HairStrandsVelocity.usf
+++ b/Engine/Shaders/Private/HairStrands/HairStrandsVelocity.usf
@@ -15,10 +15,12 @@
 
 uint2 ResolutionOffset;
 float VelocityThreshold;
+float CoverageThreshold;
 
 Texture2D<uint> NodeIndex;
 Buffer<float2> NodeVelocity;
 StructuredBuffer<FNodeVis> NodeVis; 
+Texture2D<uint4> CategorizationTexture;
 
 #define VELOCITY_TYPE_NONE 0
 #define VELOCITY_TYPE_AVG 1  
@@ -63,51 +65,55 @@ void MainCS(uint GroupIndex : SV_GroupIndex, uint3 DispatchThreadId : SV_Dispatc
 		float MaxVelocityMagnitude2 = 0;
 		#endif		
 
-		for (uint NodeIt = 0; NodeIt < NodeDesc.Count; ++NodeIt)
+		const FCategorizationData CatData = DecodeCategorizationData(CategorizationTexture.Load(uint3(PixelCoord,0)));
+		if (CatData.PixelCoverage >= CoverageThreshold)
 		{
-			const uint SampleIndex = NodeDesc.Offset + NodeIt;
-			const float2 EncodedVelocity = NodeVelocity[SampleIndex];
+			for (uint NodeIt = 0; NodeIt < NodeDesc.Count; ++NodeIt)
+			{
+				const uint SampleIndex = NodeDesc.Offset + NodeIt;
+				const float2 EncodedVelocity = NodeVelocity[SampleIndex];
+
+				#if PERMUTATION_VELOCITY == VELOCITY_TYPE_AVG
+				AverageVelocity += DecodeVelocityFromTexture(float4(EncodedVelocity, 0.0, 0.0)).xy;
+				#endif			
+
+				#if PERMUTATION_VELOCITY == VELOCITY_TYPE_CLOSEST
+				const float NodeDepth = NodeVis[SampleIndex].Depth;
+				if (NodeDepth > ClosestDepth) // Inverse-Z
+				{
+					ClosestEncodedVelocity = EncodedVelocity;
+					ClosestDepth = NodeDepth;
+				}
+				#endif
+
+				#if PERMUTATION_VELOCITY == VELOCITY_TYPE_MAX
+				const float2 CurrentVelocity = DecodeVelocityFromTexture(float4(EncodedVelocity, 0.0, 0.0)).xy;
+				const float CurrentVelocityMagnitude2 = dot(CurrentVelocity, CurrentVelocity);
+				if (CurrentVelocityMagnitude2 > MaxVelocityMagnitude2)
+				{
+				MaxEncodedVelocity = EncodedVelocity;
+					MaxVelocityMagnitude2 = CurrentVelocityMagnitude2;
+				}
+				#endif			
+
+				// If the velocity is above a certain threshold, the pixel will be resolve with a fast resolved. 
+				// This will result into a sharper, but more noisy output. However it sill avoid getting smearing 
+				// from TAA.
+				bNeedFastResolve = bNeedFastResolve || NeedFastResolve(EncodedVelocity, VelocityThreshold);
+			}
 
 			#if PERMUTATION_VELOCITY == VELOCITY_TYPE_AVG
-			AverageVelocity += DecodeVelocityFromTexture(float4(EncodedVelocity, 0.0, 0.0)).xy;
-			#endif			
+			OutVelocityTexture[PixelCoord] = ToOutput(EncodeVelocityToTexture(float3(AverageVelocity / max(NodeDesc.Count, 1u), 0.0)).xy);
+			#endif		
 
 			#if PERMUTATION_VELOCITY == VELOCITY_TYPE_CLOSEST
-			const float NodeDepth = NodeVis[SampleIndex].Depth;
-			if (NodeDepth > ClosestDepth) // Inverse-Z
-			{
-				ClosestEncodedVelocity = EncodedVelocity;
-				ClosestDepth = NodeDepth;
-			}
-			#endif
+			OutVelocityTexture[PixelCoord] = ToOutput(ClosestEncodedVelocity);
+			#endif		
 
 			#if PERMUTATION_VELOCITY == VELOCITY_TYPE_MAX
-			const float2 CurrentVelocity = DecodeVelocityFromTexture(float4(EncodedVelocity, 0.0, 0.0)).xy;
-			const float CurrentVelocityMagnitude2 = dot(CurrentVelocity, CurrentVelocity);
-			if (CurrentVelocityMagnitude2 > MaxVelocityMagnitude2)
-			{
-				MaxEncodedVelocity = EncodedVelocity;
-				MaxVelocityMagnitude2 = CurrentVelocityMagnitude2;
-			}
-			#endif			
-
-			// If the velocity is above a certain threshold, the pixel will be resolve with a fast resolved. 
-			// This will result into a sharper, but more noisy output. However it sill avoid getting smearing 
-			// from TAA.
-			bNeedFastResolve = bNeedFastResolve || NeedFastResolve(EncodedVelocity, VelocityThreshold);
+			OutVelocityTexture[PixelCoord] = ToOutput(MaxEncodedVelocity);
+			#endif
 		}
-
-		#if PERMUTATION_VELOCITY == VELOCITY_TYPE_AVG
-		OutVelocityTexture[PixelCoord] = ToOutput(EncodeVelocityToTexture(float3(AverageVelocity / max(NodeDesc.Count, 1u), 0.0)).xy);
-		#endif		
-
-		#if PERMUTATION_VELOCITY == VELOCITY_TYPE_CLOSEST
-		OutVelocityTexture[PixelCoord] = ToOutput(ClosestEncodedVelocity);
-		#endif		
-
-		#if PERMUTATION_VELOCITY == VELOCITY_TYPE_MAX
-		OutVelocityTexture[PixelCoord] = ToOutput(MaxEncodedVelocity);
-		#endif
 	}
 
 	OutResolveMaskTexture[PixelCoord] = bNeedFastResolve ? 1 : 0;
diff --git a/Engine/Source/Runtime/Renderer/Private/HairStrands/HairStrandsVisibility.cpp b/Engine/Source/Runtime/Renderer/Private/HairStrands/HairStrandsVisibility.cpp
--- a/Engine/Source/Runtime/Renderer/Private/HairStrands/HairStrandsVisibility.cpp
+++ b/Engine/Source/Runtime/Renderer/Private/HairStrands/HairStrandsVisibility.cpp
@@ -70,6 +70,9 @@ static FAutoConsoleVariableRef CVarHairVisibilityComputeRaster_Stochastic(TEXT("
 static float GHairStrandsFullCoverageThreshold = 0.98f;
 static FAutoConsoleVariableRef CVarHairStrandsFullCoverageThreshold(TEXT("r.HairStrands.Visibility.FullCoverageThreshold"), GHairStrandsFullCoverageThreshold, TEXT("Define the coverage threshold at which a pixel is considered fully covered."));
 
+static float GHairStrandsWriteVelocityCoverageThreshold = 0.f;
+static FAutoConsoleVariableRef CVarHairStrandsWriteVelocityCoverageThreshold(TEXT("r.HairStrands.Visibility.WriteVelocityCoverageThreshold"), GHairStrandsWriteVelocityCoverageThreshold, TEXT("Define the coverage threshold at which a pixel write its hair velocity (default: 0, i.e., write for all pixel)"));
+
 static int32 GHairStrandsSortHairSampleByDepth = 0;
 static FAutoConsoleVariableRef CVarHairStrandsSortHairSampleByDepth(TEXT("r.HairStrands.Visibility.SortByDepth"), GHairStrandsSortHairSampleByDepth, TEXT("Sort hair fragment by depth and update their coverage based on ordered transmittance."));
 
@@ -290,6 +293,11 @@ bool IsHairStrandsComplexLightingEnabled()
 	return GHairStrandsVisibility_UseFastPath == 0 || GetMeanSamplePerPixel() > 1 || GetHairVisibilityRenderMode() == HairVisibilityRenderMode_PPLL;
 }
 
+float GetHairWriteVelocityCoverageThreshold()
+{
+	return FMath::Clamp(GHairStrandsWriteVelocityCoverageThreshold, 0.f, 1.f);
+}
+
 ///////////////////////////////////////////////////////////////////////////////////////////////////
 
 class FHairLightSampleClearVS : public FGlobalShader
@@ -1216,6 +1224,8 @@ class FHairVelocityCS: public FGlobalShader
 	BEGIN_SHADER_PARAMETER_STRUCT(FParameters, )
 		SHADER_PARAMETER(FIntPoint, ResolutionOffset)
 		SHADER_PARAMETER(float, VelocityThreshold)
+		SHADER_PARAMETER(float, CoverageThreshold)
+		SHADER_PARAMETER_RDG_TEXTURE(Texture2D, CategorizationTexture)
 		SHADER_PARAMETER_RDG_TEXTURE(Texture2D, NodeIndex)
 		SHADER_PARAMETER_RDG_BUFFER_SRV(Buffer, NodeVelocity)
 		SHADER_PARAMETER_RDG_BUFFER_SRV(StructuredBuffer<FNodeVis>, NodeVis)
@@ -1241,6 +1251,7 @@ static void AddHairVelocityPass(
 	FRDGBuilder& GraphBuilder,
 	const FViewInfo& View,
 	const FHairStrandsMacroGroupDatas& MacroGroupDatas,
+	FRDGTextureRef& CategorizationTexture,
 	FRDGTextureRef& NodeIndex,
 	FRDGBufferRef& NodeVis,
 	FRDGBufferRef& NodeVelocity,
@@ -1269,9 +1280,11 @@ static void AddHairVelocityPass(
 	PassParameters->SceneTexturesStruct = CreateSceneTextureUniformBuffer(GraphBuilder, View.FeatureLevel);
 	PassParameters->ViewUniformBuffer = View.ViewUniformBuffer;
 	PassParameters->VelocityThreshold = GetHairFastResolveVelocityThreshold(Resolution);
+	PassParameters->CoverageThreshold = GetHairWriteVelocityCoverageThreshold();
 	PassParameters->NodeIndex = NodeIndex;
 	PassParameters->NodeVis = GraphBuilder.CreateSRV(NodeVis);
 	PassParameters->NodeVelocity = GraphBuilder.CreateSRV(NodeVelocity, FMaterialPassOutput::VelocityFormat);
+	PassParameters->CategorizationTexture = CategorizationTexture;
 	PassParameters->OutVelocityTexture = GraphBuilder.CreateUAV(OutVelocityTexture);
 	PassParameters->OutResolveMaskTexture = GraphBuilder.CreateUAV(OutResolveMaskTexture);
 
@@ -3254,6 +3267,7 @@ void RenderHairStrandsVisibilityBuffer(
 							GraphBuilder,
 							View,
 							MacroGroupDatas,
+							CategorizationTexture,
 							CompactNodeIndex,
 							CompactNodeData,
 							PassOutput.NodeVelocity,
@@ -3444,6 +3458,7 @@ void RenderHairStrandsVisibilityBuffer(
 							GraphBuilder,
 							View,
 							MacroGroupDatas,
+							CategorizationTexture,
 							CompactNodeIndex,
 							CompactNodeData,
 							PassOutput.NodeVelocity,
