Rename USE_LIGHT_GRID to USE_PATH_TRACING_LIGHT_GRID to be more descriptive

diff --git a/Engine/Shaders/Private/PathTracing/Light/PathTracingLightGrid.ush b/Engine/Shaders/Private/PathTracing/Light/PathTracingLightGrid.ush
--- a/Engine/Shaders/Private/PathTracing/Light/PathTracingLightGrid.ush
+++ b/Engine/Shaders/Private/PathTracing/Light/PathTracingLightGrid.ush
@@ -1,13 +1,13 @@
 // Copyright Epic Games, Inc. All Rights Reserved.
 
-#ifndef USE_LIGHT_GRID
-#define USE_LIGHT_GRID 0
+#ifndef USE_PATH_TRACING_LIGHT_GRID
+#define USE_PATH_TRACING_LIGHT_GRID 0
 #endif
 
 #include "/Engine/Shared/PathTracingDefinitions.h"
 
 
-#if USE_LIGHT_GRID
+#if USE_PATH_TRACING_LIGHT_GRID
 // extra shader parameters that must be defined when using the light grid
 uint SceneInfiniteLightCount;
 float3 SceneLightsBoundMin;
@@ -31,12 +31,12 @@ float3 HashToColor(uint H)
 struct FLightLoopCount {
 	uint NumLights;
 	uint NumMISLights;
-#if USE_LIGHT_GRID
+#if USE_PATH_TRACING_LIGHT_GRID
 	uint LightGridOffset;
 #endif
 };
 
-#if USE_LIGHT_GRID
+#if USE_PATH_TRACING_LIGHT_GRID
 
 FLightLoopCount LightGridLookup(float3 WorldPos)
 {
diff --git a/Engine/Shaders/Private/PathTracing/PathTracing.usf b/Engine/Shaders/Private/PathTracing/PathTracing.usf
--- a/Engine/Shaders/Private/PathTracing/PathTracing.usf
+++ b/Engine/Shaders/Private/PathTracing/PathTracing.usf
@@ -10,7 +10,7 @@ PathTracingRayGenShader.usf: Reference path tracing
 
 #define RANDSEQ_UNROLL_SOBOL 0 // unrolling sobol sampler doesn't seem to pay off in the path tracer
 
-#define USE_LIGHT_GRID	1
+#define USE_PATH_TRACING_LIGHT_GRID	1
 
 #include "../Common.ush"
 #include "../PostProcessCommon.ush"
