Use back face culling for material gather rays

diff --git a/Engine/Shaders/Private/RayTracing/RayTracingDeferredReflections.usf b/Engine/Shaders/Private/RayTracing/RayTracingDeferredReflections.usf
--- a/Engine/Shaders/Private/RayTracing/RayTracingDeferredReflections.usf
+++ b/Engine/Shaders/Private/RayTracing/RayTracingDeferredReflections.usf
@@ -120,7 +120,7 @@ RAY_TRACING_ENTRY_RAYGEN(RayTracingDeferredReflectionsRGS)
 		if (bIsValidPixel)
 		{
 			TraceDeferredMaterialGatherRay(TLAS,
-				RAY_FLAG_FORCE_OPAQUE,
+				RAY_FLAG_FORCE_OPAQUE | RAY_FLAG_CULL_BACK_FACING_TRIANGLES,
 				RAY_TRACING_MASK_OPAQUE,
 				Ray,
 				Bookmark,
