Fix RTGI light selection

The logic for the light cdf got accidentally flipped in a previous refactor.

diff --git a/Engine/Shaders/Private/PathTracing/Light/PathTracingLightSampling.ush b/Engine/Shaders/Private/PathTracing/Light/PathTracingLightSampling.ush
--- a/Engine/Shaders/Private/PathTracing/Light/PathTracingLightSampling.ush
+++ b/Engine/Shaders/Private/PathTracing/Light/PathTracingLightSampling.ush
@@ -174,7 +174,7 @@ bool SampleLight(
 
 	float LightPickingCdfSum = InitLightPickingCdf(WorldPos, WorldNormal, PrimitiveLightingChannelMask, IsTransmissiveMaterial, LightPickingCdf);
 
-	if (LightPickingCdfSum > 0.0)
+	if (LightPickingCdfSum == 0.0)
 	{
 		// light picking cdf is all 0
 		return false;
