Moved fuzz/cloth onto the Slab node.

diff --git a/Engine/Shaders/Private/Strata/Strata.ush b/Engine/Shaders/Private/Strata/Strata.ush
--- a/Engine/Shaders/Private/Strata/Strata.ush
+++ b/Engine/Shaders/Private/Strata/Strata.ush
@@ -61,14 +61,15 @@ struct FStrataBSDF
 	float4 InlineVGPRs[3];
 #endif
 
-	float4 VGPRs[4];
+	float4 VGPRs[5];
 };
 
 
 
 #define TO_BITMASK(BitCount)							((1u<<BitCount)-1)
+#define TO_OFFSET_BITMASK(BitCount, BitOffset)			(TO_BITMASK(BitCount)<<BitOffset)
 #define READ_BITS(X, BitCount, BitOffset)				uint((X>>BitOffset) & TO_BITMASK(BitCount))
-#define WRITE_BITS(X, BitCount, BitOffset, NewValue)	X = (X & (~(TO_BITMASK(BitCount)<<BitOffset))) | ((NewValue & TO_BITMASK(BitCount)) << BitOffset)
+#define WRITE_BITS(X, BitCount, BitOffset, NewValue)	X = (X & (~TO_OFFSET_BITMASK(BitCount, BitOffset))) | ((NewValue & TO_BITMASK(BitCount)) << BitOffset)
 
 // Bit count for each component stored in a Strata header
 // For HEADER_BIT_BSDF_COUNT and  HEADER_BIT_SHAREDNORMAL_X, please refer to values in StrataDefinitions.h
@@ -123,7 +124,8 @@ struct FStrataBSDF
 #define STATE_BIT_COUNT_THINFILM				1
 #define STATE_BIT_COUNT_SIMPLEVOLUME			1
 #define STATE_BIT_COUNT_DMFPPLUGGED				1
-#define STATE_BIT_COUNT___UNUSED___				7
+#define STATE_BIT_COUNT_HASFUZZ					1
+#define STATE_BIT_COUNT___UNUSED___				6
 #define STATE_BIT_COUNT_WEIGHT10F				10
 
 // Bit offset of each component store in BSDF state 
@@ -139,6 +141,7 @@ struct FStrataBSDF
 #define STATE_BIT_OFFSET_THINFILM		(STATE_BIT_OFFSET_EDGECOLOR			+ STATE_BIT_COUNT_EDGECOLOR)
 #define STATE_BIT_OFFSET_SIMPLEVOLUME	(STATE_BIT_OFFSET_THINFILM			+ STATE_BIT_COUNT_THINFILM)
 #define STATE_BIT_OFFSET_DMFPPLUGGED	(STATE_BIT_OFFSET_SIMPLEVOLUME		+ STATE_BIT_COUNT_SIMPLEVOLUME)
+#define STATE_BIT_OFFSET_HASFUZZ		(STATE_BIT_OFFSET_DMFPPLUGGED		+ STATE_BIT_COUNT_DMFPPLUGGED)
 #define STATE_BIT_OFFSET_WEIGHT10F		22
 
 // Accessors for the BSDF type
@@ -195,6 +198,10 @@ struct FStrataBSDF
 #define BSDF_GETHASDMFP(X)				READ_BITS(X.State,  STATE_BIT_COUNT_DMFPPLUGGED, STATE_BIT_OFFSET_DMFPPLUGGED)
 #define BSDF_SETHASDMFP(X, Value)		WRITE_BITS(X.State, STATE_BIT_COUNT_DMFPPLUGGED, STATE_BIT_OFFSET_DMFPPLUGGED, Value)
 
+// Indicates if the BSDF is part of the top layer
+#define BSDF_GETHASFUZZ(X)				READ_BITS(X.State,  STATE_BIT_COUNT_HASFUZZ, STATE_BIT_OFFSET_HASFUZZ)
+#define BSDF_SETHASFUZZ(X, Value)		WRITE_BITS(X.State, STATE_BIT_COUNT_HASFUZZ, STATE_BIT_OFFSET_HASFUZZ, Value)
+
 #if STRATA_INLINE_SHADING
 #define BSDF_SETEMISSIVE(X, V)			X.Emissive = V
 #define BSDF_GETEMISSIVE(X)				X.Emissive
@@ -211,6 +218,12 @@ struct FStrataBSDF
 #define BSDF_GETTHICKNESSCM(X)			STRATA_SIMPLEVOLUME_THICKNESS_CM
 #endif
 
+#define SLAB_COMPLEX_PATH_MASK (TO_OFFSET_BITMASK(STATE_BIT_COUNT_ANISOTROPY, STATE_BIT_OFFSET_ANISOTROPY)		|  TO_OFFSET_BITMASK(STATE_BIT_COUNT_HASSSS, STATE_BIT_OFFSET_HASSSS) | \
+								TO_OFFSET_BITMASK(STATE_BIT_COUNT_HASSSSPROFILE, STATE_BIT_OFFSET_HASSSSPROFILE)|  TO_OFFSET_BITMASK(STATE_BIT_COUNT_HAZINESS, STATE_BIT_OFFSET_HAZINESS) | \
+								TO_OFFSET_BITMASK(STATE_BIT_COUNT_THINFILM, STATE_BIT_OFFSET_THINFILM)			|  TO_OFFSET_BITMASK(STATE_BIT_COUNT_EDGECOLOR, STATE_BIT_OFFSET_EDGECOLOR) | \
+								TO_OFFSET_BITMASK(STATE_BIT_COUNT_SIMPLEVOLUME, STATE_BIT_OFFSET_SIMPLEVOLUME)	|  TO_OFFSET_BITMASK(STATE_BIT_COUNT_HASFUZZ, STATE_BIT_OFFSET_HASFUZZ) )
+
+
 #define SLAB_BASECOLOR(X)				X.VGPRs[0].xyz
 #define SLAB_SPECULAR(X)				X.VGPRs[0].w
 #define SLAB_METALLIC(X)				X.VGPRs[1].x
@@ -226,9 +239,10 @@ struct FStrataBSDF
 // Thin film. IOR is harcoded to oil IOR in order to save VGPR
 #define SLAB_THINFILM_THICKNESS(X)		X.VGPRs[2].w
 #define SLAB_THINFILM_IOR(X)			1.44f
-
-#define SHEEN_BASECOLOR(X)				X.VGPRs[0].xyz
-#define SHEEN_ROUGHNESS(X)				X.VGPRs[0].w
+// Fuzz to simulate cloth
+#define SLAB_FUZZ_COLOR(X)				X.VGPRs[3].xyz
+#define SLAB_FUZZ_AMOUNT(X)				X.VGPRs[3].w
+// NOTE For SLAB, X.VGPRs[3].w is free for another use case.
 
 #define VOLUMETRICFOGCLOUD_ALBEDO(X)	X.VGPRs[0].xyz
 #define VOLUMETRICFOGCLOUD_EXTINCTION(X)X.VGPRs[1].xyz
@@ -433,9 +447,14 @@ FStrataData GetInitialisedStrataData()
 }
 
 FStrataData GetStrataSlabBSDF(
-	float3 BaseColor, float3 EdgeColor, float Specular, float Metallic, float Roughness, float Anisotropy,
+	float3 BaseColor, float3 EdgeColor, float Specular, float Metallic, float Roughness, 
+	float Anisotropy,
 	float SSSProfileID, float3 SSSDMFP, float SSSDMFPScale,
-	float3 Emissive, float Haziness, float ThinFilmThickness,  float Thickness, 
+	float3 Emissive, 
+	float Haziness, 
+	float ThinFilmThickness,
+	float FuzzAmount, float3 FuzzColor,
+	float Thickness, 
 	uint SharedNormalIndex, inout uint SharedNormalTypes)
 {
 	FStrataData StrataData = GetInitialisedStrataData();
@@ -445,6 +464,7 @@ FStrataData GetStrataSlabBSDF(
 	const uint bHasHaziness = Haziness > 0.0f ? 1 : 0;
 	const uint bHasThinFilm = ThinFilmThickness > 0.0f ? 1 : 0;
 	const uint bHasAnisotropy = Anisotropy != 0;
+	const uint bHasFuzz = FuzzAmount > 0.0f ? 1 : 0;
 	if (bHasAnisotropy)
 	{
 		StrataRequestSharedNormalTangentBasis(SharedNormalTypes, SharedNormalIndex);
@@ -470,6 +490,7 @@ FStrataData GetStrataSlabBSDF(
 	BSDF_SETHASEDGECOLOR	(StrataData.Layers[0].BSDFs[0],		bHasEdgeColor);
 	BSDF_SETHASTHINFILM		(StrataData.Layers[0].BSDFs[0],		bHasThinFilm);
 	BSDF_SETTHICKNESSCM		(StrataData.Layers[0].BSDFs[0],		Thickness);
+	BSDF_SETHASFUZZ			(StrataData.Layers[0].BSDFs[0],		bHasFuzz);
 
 	SLAB_BASECOLOR			(StrataData.Layers[0].BSDFs[0])		= BaseColor;
 	SLAB_EDGECOLOR			(StrataData.Layers[0].BSDFs[0])		= EdgeColor;
@@ -479,6 +500,7 @@ FStrataData GetStrataSlabBSDF(
 	SLAB_ROUGHNESSY			(StrataData.Layers[0].BSDFs[0])		= RoughnessXY.y;
 	SLAB_HAZINESS			(StrataData.Layers[0].BSDFs[0])		= Haziness;
 	SLAB_THINFILM_THICKNESS	(StrataData.Layers[0].BSDFs[0])		= ThinFilmThickness;
+
 	if (bHasSSSProfile)
 	{
 		SLAB_SSSPROFILEID	(StrataData.Layers[0].BSDFs[0])		= SSSProfileID;
@@ -489,6 +511,9 @@ FStrataData GetStrataSlabBSDF(
 		SLAB_SSSDMFP	(StrataData.Layers[0].BSDFs[0])			= SSSDMFP * saturate(SSSDMFPScale);
 	}
 
+	SLAB_FUZZ_AMOUNT		(StrataData.Layers[0].BSDFs[0])		= FuzzAmount;
+	SLAB_FUZZ_COLOR			(StrataData.Layers[0].BSDFs[0])		= FuzzColor;
+
 #if STRATA_INLINE_SHADING
 	// At this stage we do not know if this strata BSDF is going to be using simple volume lighting so always store the DMFP on the side to not override SSSPROFILE data.
 	StrataData.Layers[0].BSDFs[0].TmpDMFP = SSSDMFP * saturate(SSSDMFPScale);
@@ -498,23 +523,6 @@ FStrataData GetStrataSlabBSDF(
 	return StrataData;
 }
 
-FStrataData GetStrataSheenBSDF(float3 BaseColor, float Roughness, uint SharedNormalIndex)
-{
-	FStrataData StrataData = GetInitialisedStrataData();
-	StrataData.LayerCount++;
-
-	StrataData.Layers[0].BSDFCount++;
-	BSDF_SETTYPE			(StrataData.Layers[0].BSDFs[0],		STRATA_BSDF_TYPE_SHEEN);
-	BSDF_SETNORMALID		(StrataData.Layers[0].BSDFs[0],		SharedNormalIndex);
-	BSDF_SETISTOPLAYER		(StrataData.Layers[0].BSDFs[0],		0);
-	SHEEN_BASECOLOR			(StrataData.Layers[0].BSDFs[0])		= BaseColor;
-	SHEEN_ROUGHNESS			(StrataData.Layers[0].BSDFs[0])		= Roughness;
-
-	StrataData.Layers[0].BSDFs[0].Weight						= 1.0f;
-
-	return StrataData;
-}
-
 FStrataData GetStrataVolumeFogCloudBSDF(float3 Albedo, float3 Extinction, float3 Emissive, float AmbientOcclusion)
 {
 	FStrataData StrataData = GetInitialisedStrataData();
@@ -884,12 +892,9 @@ void SanitizeStrataSlab(inout FStrataBSDF BSDF)
 			SLAB_SSSDMFP(BSDF) = clamp(SLAB_SSSDMFP(BSDF), 0.0f, Max111110BitsFloat3);
 		}
 	}
-}
 
-void SanitizeStrataSheen(inout FStrataBSDF BSDF)
-{
-	SHEEN_BASECOLOR(BSDF) = saturate(SHEEN_BASECOLOR(BSDF));
-	SHEEN_ROUGHNESS(BSDF) = saturate(SHEEN_ROUGHNESS(BSDF));
+	SLAB_FUZZ_AMOUNT(BSDF) = saturate(SLAB_FUZZ_AMOUNT(BSDF));
+	SLAB_FUZZ_COLOR(BSDF) = saturate(SLAB_FUZZ_COLOR(BSDF));
 }
 
 void SanitizeStrataHair(inout FStrataBSDF BSDF)
@@ -929,11 +934,6 @@ void StrataSanitizeBSDF(inout FStrataBSDF BSDF)
 			SanitizeStrataSlab(BSDF);
 			break; 
 		}
-		case STRATA_BSDF_TYPE_SHEEN:
-		{
-			SanitizeStrataSheen(BSDF);
-			break;
-		}
 		case STRATA_BSDF_TYPE_HAIR:
 		{
 			SanitizeStrataHair(BSDF);
@@ -1390,21 +1390,11 @@ void PackStrataOut(
 					bStrataSubsurfaceEnable = false;
 				}
 
-				LayerCoverage	+= BSDFCoverage;
-				LayerThroughput += BSDFCoverage * BSDFThroughput;
-				break;
-			}
-
-			case STRATA_BSDF_TYPE_SHEEN:
-			{
-				// STRATA_TODO Fresnel throughput, Use a LUT, this is currently in progress and will be coming soon.
-				// In the meantime, we use a simple approach to feed Fresnel function and compute an approximation of the general transmition/througput.
-				const float NoL = saturate(ContextNoLight.NoV);
-				const float3 F0  = SHEEN_BASECOLOR(BSDF);
-				const float3 ApproxFresnelForTransmittance = F_Schlick(F0, NoL);
+				// STRATA_TODO	Should we have the Fuzz influence the next layter visibility accoding to fuzz amount? Should it also occlude diffuse?
+				//				If this is done, we should also update ComputePostProcessBSDFData.
 
 				LayerCoverage	+= BSDFCoverage;
-				LayerThroughput += BSDFCoverage * (1.0f - ApproxFresnelForTransmittance);
+				LayerThroughput += BSDFCoverage * BSDFThroughput;
 				break;
 			}
 
@@ -1513,13 +1503,12 @@ void PackStrataOut(
 					}
 				}
 
-				// 8-16 bytes
-			}
-			break;
-			case STRATA_BSDF_TYPE_SHEEN:
-			{
-				STRATA_STORE_UINT1(PackColorLinearToGamma2AlphaLinear(float4(SHEEN_BASECOLOR(BSDF), SHEEN_ROUGHNESS(BSDF))));
-				// 4 bytes
+				if (BSDF_GETHASFUZZ(BSDF))
+				{
+					STRATA_STORE_UINT1(PackColorLinearToGamma2AlphaLinear(float4(SLAB_FUZZ_COLOR(BSDF), SLAB_FUZZ_AMOUNT(BSDF))));
+				}
+
+				// 8-20 bytes
 			}
 			break;
 			case STRATA_BSDF_TYPE_HAIR:
@@ -1596,16 +1585,9 @@ FStrataPixelHeader UnpackStrataHeaderIn(ByteAddressBuffer StrataBuffer, inout FS
 	return Out;
 }
 
-bool IsStrataFastPathCompatible(in FStrataBSDF BSDF)
+bool IsStrataSlabFastPathCompatible(in FStrataBSDF BSDF)
 {
-	return
-		!BSDF_GETANISOTROPY(BSDF) &&
-		!BSDF_GETHASSSS(BSDF) &&
-		!BSDF_GETHASSSSPROFILE(BSDF) &&
-		!BSDF_GETHASHAZINESS(BSDF) &&
-		!BSDF_GETHASTHINFILM(BSDF) &&
-		!BSDF_GETHASEDGECOLOR(BSDF) &&
-		!BSDF_GETISSIMPLEVOLUME(BSDF);
+	return (BSDF.State & SLAB_COMPLEX_PATH_MASK) == 0;
 }
 
 void UnpackFastPathSlabBSDF(inout FStrataBSDF OutBSDF, uint2 PackedData01)
@@ -1647,6 +1629,7 @@ FStrataBSDF UnpackFastPathStrataBSDFIn(ByteAddressBuffer StrataBuffer, inout FSt
 	BSDF_SETHASHAZINESS(OutBSDF, 0);
 	BSDF_SETHASTHINFILM(OutBSDF, 0);
 	BSDF_SETISSIMPLEVOLUME(OutBSDF, 0);
+	BSDF_SETHASFUZZ(OutBSDF, 0);
 
 	return OutBSDF;
 }
@@ -1712,11 +1695,11 @@ FStrataBSDF UnpackStrataBSDFIn(ByteAddressBuffer StrataBuffer, inout FStrataAddr
 			}
 			#endif
 		}
-	}
-	break;
-	case STRATA_BSDF_TYPE_SHEEN:
-	{
-		StrataLoad_ColorGamma2ToLinearAlphaLinear(StrataBuffer, StrataAddressing, SHEEN_BASECOLOR(OutBSDF), SHEEN_ROUGHNESS(OutBSDF));
+
+		if (BSDF_GETHASFUZZ(OutBSDF))
+		{
+			StrataLoad_ColorGamma2ToLinearAlphaLinear(StrataBuffer, StrataAddressing, SLAB_FUZZ_COLOR(OutBSDF), SLAB_FUZZ_AMOUNT(OutBSDF));
+		}
 	}
 	break;
 	case STRATA_BSDF_TYPE_HAIR:
diff --git a/Engine/Shaders/Private/Strata/StrataEvaluation.ush b/Engine/Shaders/Private/Strata/StrataEvaluation.ush
--- a/Engine/Shaders/Private/Strata/StrataEvaluation.ush
+++ b/Engine/Shaders/Private/Strata/StrataEvaluation.ush
@@ -86,11 +86,6 @@ void StrataGetBSDFRoughness(in FStrataBSDF BSDF, inout float R)
 		R = SLAB_ROUGHNESSX(BSDF);
 		break;
 	}
-	case STRATA_BSDF_TYPE_SHEEN:
-	{
-		R = SHEEN_ROUGHNESS(BSDF);
-		break;
-	}
 	case STRATA_BSDF_TYPE_HAIR:
 	{
 		R = HAIR_ROUGHNESS(BSDF);
@@ -363,6 +358,38 @@ FStrataEvaluateResult StrataEvaluateBSDF(FStrataBSDFContext BSDFContext, FShadow
 
 			Sample.Throughput = OpaqueBSDFThroughput;
 
+			#if STRATA_FASTPATH==0
+			BRANCH
+			if (BSDF_GETHASFUZZ(BSDFContext.BSDF))
+			{
+				const float Fuzz = SLAB_FUZZ_AMOUNT(BSDFContext.BSDF);
+				const float3 FuzzF0  = SLAB_FUZZ_COLOR(BSDFContext.BSDF);
+
+				// Special override for roughness for supporting area light integrator with Sphere/Tube/Disk light, which modifies/increase roughness.
+				const float a2 = BSDFContext.a2Override > 0 ? BSDFContext.a2Override : Pow4(SafeRoughness);
+
+				const float SafeClothNoV = saturate(abs(BSDFContext.Context.NoV) + 1e-5);
+				
+				// Generalized inverse microfacet specular
+				float ClothD   = D_InvGGX(a2, BSDFContext.Context.NoH);
+				float ClothVis = Vis_Cloth(SafeClothNoV, BSDFContext.SatNoL);
+				float3 ClothF  = F_Schlick(FuzzF0, BSDFContext.Context.VoH);
+				float3 ClothDVF = ClothD * ClothVis * ClothF;
+				
+				// What to do with specular and diffuse here?
+				Sample.SpecularPathValue = lerp(Sample.SpecularPathValue,	ClothDVF,					Fuzz);
+				Sample.SpecularPDF		 = lerp(Sample.SpecularPDF,			BSDFContext.SatNoV / PI,	Fuzz); // Per "Production Friendly Microfacet Sheen BRDF", hemispherical sampling give good result as the roughness is usually high.
+				// Fuzz does not affect throughput
+
+				if (bHaziness)
+				{
+					Sample.SpecularHazePathValue= lerp(Sample.SpecularHazePathValue,	ClothDVF,					Fuzz);
+					Sample.SpecularHazePDF		= lerp(Sample.SpecularHazePDF,			BSDFContext.SatNoV / PI,	Fuzz);
+				}
+				break;
+			}
+			#endif
+
 			#if STRATA_FASTPATH==0 && STRATA_SSS_TRANSMISSION
 			const bool bHasProfile = BSDF_GETHASSSSPROFILE(BSDFContext.BSDF);
 			if (BSDF_GETHASSSS(BSDFContext.BSDF))
@@ -428,29 +455,6 @@ FStrataEvaluateResult StrataEvaluateBSDF(FStrataBSDFContext BSDFContext, FShadow
 		}
 
 		#if STRATA_FASTPATH==0
-		case STRATA_BSDF_TYPE_SHEEN:
-		{
-			const float3 F0  = SHEEN_BASECOLOR(BSDFContext.BSDF);
-			const float SafeRoughness = MakeRoughnessSafe(SHEEN_ROUGHNESS(BSDFContext.BSDF));
-
-			// Special override for roughness for supporting area light integrator with Sphere/Tube/Disk light, which modifies/increase roughness.
-			const float a2 = BSDFContext.a2Override > 0 ? BSDFContext.a2Override : Pow4(SafeRoughness);
-
-			// Generalized inverse microfacet specular
-			float D   = D_InvGGX(a2, BSDFContext.Context.NoH);
-			float Vis = Vis_Cloth(BSDFContext.Context.NoV, BSDFContext.SatNoL);
-			float3 F  = F_Schlick(F0, BSDFContext.Context.VoH);
-
-			// STRATA_TODO Fresnel throughput, Use a LUT, this is currently in progress and will be coming soon.
-			// In the meantime, we use a simple approach to feed Fresnel function and compute an approximation of the general transmition/througput.
-			const float3 ApproxFresnelForTransmittance = F_Schlick(F0, BSDFContext.SatNoV);
-
-			Sample.DiffusePathValue	= D * Vis * F;
-			Sample.DiffusePDF		= BSDFContext.SatNoV / PI; // Per "Production Friendly Microfacet Sheen BRDF", hemispherical sampling give good result as the roughness is usually high.
-			Sample.Throughput		= (1.0f - ApproxFresnelForTransmittance);
-			break;
-		}
-
 		case STRATA_BSDF_TYPE_HAIR:
 		{
 			FGBufferData GBuffer	= (FGBufferData)0;
@@ -524,7 +528,7 @@ FStrataImportanceSampleResult StrataImportanceSampleBSDF(FStrataBSDFContext BSDF
 			PDF = Sample.DiffusePDF;
 
 			const float3 F0 = ComputeF0(SLAB_SPECULAR(BSDF), BaseColor, Metallic);
-			// STRATA_TODO F82pb 
+			// STRATA_TODO take into account F90 
 
 			const float SafeRoughnessX = MakeRoughnessSafe(SLAB_ROUGHNESSX(BSDFContext.BSDF));
 			const bool bHasAnisotropy = BSDF_GETANISOTROPY(BSDFContext.BSDF);
@@ -539,7 +543,8 @@ FStrataImportanceSampleResult StrataImportanceSampleBSDF(FStrataBSDFContext BSDF
 				TangentH = ImportanceSampleVisibleGGX(UniformSampleDisk(E), Pow4(SafeRoughnessX), BSDFContext.TangentV);
 			}
 
-			// STRATA_TODO: take into account Haziness
+			// STRATA_TODO take into account Haziness
+			// STRATA_TODO take into account Fuzz: Per "Production Friendly Microfacet Sheen BRDF", uniform hemispherical sampling gives good result as the roughness is usually high and avoid issue at grazing angle.
 
 			const float HPDF = TangentH.w;
 			const float3 H = mul(TangentH.xyz, BSDFContext.TangentBasis);
@@ -551,16 +556,6 @@ FStrataImportanceSampleResult StrataImportanceSampleBSDF(FStrataBSDFContext BSDF
 			break;
 		}
 
-		case STRATA_BSDF_TYPE_SHEEN:
-		{
-			// Per "Production Friendly Microfacet Sheen BRDF", uniform hemispherical sampling gives good result as the roughness is 
-			// usually high and avoid issue at grazing angle
-			float4 ImportanceSample = CosineSampleHemisphere(E);
-			Sample.DiffuseDirection = mul(ImportanceSample.xyz, BSDFContext.TangentBasis);
-			PDF = ImportanceSample.w;
-			break;
-		}
-
 		case STRATA_BSDF_TYPE_HAIR:
 		{
 			// STRATA_TODO do something better when we get there with Lumen, and evaluate the different researched solution (e.g. Importance Sampling for Physically-Based Hair Fiber Models)
@@ -758,16 +753,27 @@ FStrataEnvLightResult StrataEvaluateForEnvLight(FStrataBSDFContext BSDFContext,
 				StrataEnvLightResult.bSubsurface = false;
 			}
 			#endif
+
+			#if STRATA_FASTPATH==0
+			BRANCH
+			if (BSDF_GETHASFUZZ(BSDFContext.BSDF))
+			{
+				const float Fuzz = SLAB_FUZZ_AMOUNT(BSDFContext.BSDF);
+				const float3 FuzzColor  = SLAB_FUZZ_COLOR(BSDFContext.BSDF);
+
+				// This should be a lerp, but SkyLightDiffuse from ReflectionEnvironmentPixelShader.usf add the contribution...
+				// STRATA_TODO fix the Fuzz model
+				//StrataEnvLightResult.DiffuseWeight = lerp(StrataEnvLightResult.DiffuseWeight, FuzzColor, Fuzz);
+				StrataEnvLightResult.DiffuseWeight += FuzzColor * Fuzz;
+				// Fuzz does not affect throughput.
+				break;
+			}
+			#endif
+
 			break;
 		}
 
 #if STRATA_FASTPATH==0
-		case STRATA_BSDF_TYPE_SHEEN:
-		{
-			StrataEnvLightResult.DiffuseWeight = SHEEN_BASECOLOR(BSDFContext.BSDF);
-			StrataEnvLightResult.DiffuseNormal = BSDFContext.N;
-			break;
-		}
 
 		case STRATA_BSDF_TYPE_HAIR:
 		{
diff --git a/Engine/Shaders/Private/Strata/StrataMaterialClassification.usf b/Engine/Shaders/Private/Strata/StrataMaterialClassification.usf
--- a/Engine/Shaders/Private/Strata/StrataMaterialClassification.usf
+++ b/Engine/Shaders/Private/Strata/StrataMaterialClassification.usf
@@ -52,14 +52,12 @@ void ComputePostProcessBSDFData(
 	}
 
 	Classification.ShadingModels |= StrataShadingModelBit(BSDF_GETTYPE(BSDF));
-	Classification.bIsSimple = Classification.bIsSimple && BSDF_GETTYPE(BSDF) == STRATA_BSDF_TYPE_SLAB;
 
 	switch (BSDF_GETTYPE(BSDF))
 	{
 	case STRATA_BSDF_TYPE_SLAB:
 	{
-		const float3 F0  = ComputeF0(SLAB_SPECULAR(BSDF), SLAB_BASECOLOR(BSDF), SLAB_METALLIC(BSDF));
-		const float3 F90 = ComputeF90(F0, SLAB_BASECOLOR(BSDF), SLAB_METALLIC(BSDF));
+		Classification.bIsSimple = Classification.bIsSimple && IsStrataSlabFastPathCompatible(BSDF);
 
 		if (bTopLayer)
 		{
@@ -68,8 +66,6 @@ void ComputePostProcessBSDFData(
 			TopLayerRoughnessAverageCount += BSDFWeightAvg;
 		}
 
-		Classification.bIsSimple = Classification.bIsSimple && IsStrataFastPathCompatible(BSDF);
-
 		// Select the latest encoutered SSS BSDF
 		if (BSDF_GETHASSSS(BSDF) && BSDFWeightAvg > 0)
 		{
@@ -92,21 +88,10 @@ void ComputePostProcessBSDFData(
 		break;
 	}
 
-	case STRATA_BSDF_TYPE_SHEEN:
-	{
-		const float3 F0  = SHEEN_BASECOLOR(BSDF);
-		const float3 F   = F_Schlick(F0, ContextNoLight.VoH);
-
-		if (bTopLayer)
-		{
-			TopLayerRoughness += SHEEN_ROUGHNESS(BSDF);
-			TopLayerRoughnessAverageCount += 1.0f;
-		}
-		break;
-	}
-
 	case STRATA_BSDF_TYPE_HAIR:
 	{
+		Classification.bIsSimple = false;
+
 		// No SSR
 		if (bTopLayer)
 		{
@@ -118,6 +103,8 @@ void ComputePostProcessBSDFData(
 
 	case STRATA_BSDF_TYPE_SINGLELAYERWATER:
 	{
+		Classification.bIsSimple = false;
+
 		if (bTopLayer)
 		{
 			const float Roughness = SLW_ROUGHNESS(BSDF);
diff --git a/Engine/Shaders/Private/Strata/StrataVisualize.usf b/Engine/Shaders/Private/Strata/StrataVisualize.usf
--- a/Engine/Shaders/Private/Strata/StrataVisualize.usf
+++ b/Engine/Shaders/Private/Strata/StrataVisualize.usf
@@ -392,6 +392,16 @@ void VisualizeMaterialPS(
 					HORIZONTALTAB; PRINTFLOAT1(SLAB_THINFILM_THICKNESS(BSDF));
 					NEXTLINE;
 				}
+
+				if (BSDF_GETHASFUZZ(BSDF))
+				{
+					HORIZONTALTAB; PRINTPARAMCHAR(_F_); PRINTPARAMCHAR(_U_); PRINTPARAMCHAR(_Z_); PRINTPARAMCHAR(_Z_); PRINTPARAMCHAR(_SPC_); PRINTPARAMCHAR(_SPC_); PRINTPARAMCHAR(_SPC_);
+					HORIZONTALTAB; PRINTFLOAT1(SLAB_FUZZ_AMOUNT(BSDF));
+					NEXTLINE;
+					HORIZONTALTAB; PRINTPARAMCHAR(_F_); PRINTPARAMCHAR(_U_); PRINTPARAMCHAR(_Z_); PRINTPARAMCHAR(_C_); PRINTPARAMCHAR(_O_); PRINTPARAMCHAR(_L_); PRINTPARAMCHAR(_SPC_);
+					HORIZONTALTAB; PRINTFLOAT3(SLAB_FUZZ_COLOR(BSDF));
+					NEXTLINE;
+				}
 #else
 				NEXTLINE;
 				NEXTLINE;
@@ -431,41 +441,17 @@ void VisualizeMaterialPS(
 					NEXTLINE;
 					NEXTLINE;
 				}
-#endif
-				break;
 
-			case STRATA_BSDF_TYPE_SHEEN:
-#if PERMUTATION_BSDF_PASS==1
-				PRINTBSDFCHAR(_S_); PRINTBSDFCHAR(_H_); PRINTBSDFCHAR(_E_); PRINTBSDFCHAR(_E_); PRINTBSDFCHAR(_N_); 
-				PRINTBSDFCHAR(_SPC_); PRINTBSDFCHAR(_SPC_); PRINTBSDFCHAR(_SPC_); PRINTBSDFCHAR(_SPC_); PRINTBSDFCHAR(_SPC_);
-				if (BSDF_GETGREYWEIGHT(BSDF))
-				{
-					PRINTBSDFGREYWEIGHT();
-				}
-				else
+				if (BSDF_GETHASFUZZ(BSDF))
 				{
-					PRINTBSDFWEIGHT();
+					NEXTLINE;
+					NEXTLINE;
 				}
-				PRINTBSDFNORMALID();
-				PRINTBSDFANISO();
-				PRINTBSDFTOPLAY();
-				NEXTLINE;
-
-				HORIZONTALTAB; PRINTPARAMCHAR(_A_); PRINTPARAMCHAR(_L_); PRINTPARAMCHAR(_B_); PRINTPARAMCHAR(_E_); PRINTPARAMCHAR(_D_); PRINTPARAMCHAR(_O_); 
-				HORIZONTALTAB; PRINTFLOAT3(SHEEN_BASECOLOR(BSDF));
-				NEXTLINE;
-				HORIZONTALTAB; PRINTPARAMCHAR(_R_); PRINTPARAMCHAR(_O_); PRINTPARAMCHAR(_U_); PRINTPARAMCHAR(_G_); PRINTPARAMCHAR(_H_); PRINTPARAMCHAR(_SPC_);
-				HORIZONTALTAB; PRINTFLOAT1(SHEEN_ROUGHNESS(BSDF));
-				NEXTLINE;
-#else
-				NEXTLINE;
-				NEXTLINE;
-				NEXTLINE;
 #endif
 				break;
 
 			case STRATA_BSDF_TYPE_HAIR:
-#if PERMUTATION_BSDF_PASS==2
+#if PERMUTATION_BSDF_PASS==1
 				PRINTBSDFCHAR(_H_); PRINTBSDFCHAR(_A_); PRINTBSDFCHAR(_I_); PRINTBSDFCHAR(_R_); PRINTBSDFCHAR(_SPC_);
 				PRINTBSDFCHAR(_SPC_); PRINTBSDFCHAR(_SPC_); PRINTBSDFCHAR(_SPC_); PRINTBSDFCHAR(_SPC_); PRINTBSDFCHAR(_SPC_);
 				if (BSDF_GETGREYWEIGHT(BSDF))
@@ -507,7 +493,7 @@ void VisualizeMaterialPS(
 				break;
 
 			case STRATA_BSDF_TYPE_SINGLELAYERWATER:
-#if PERMUTATION_BSDF_PASS==3
+#if PERMUTATION_BSDF_PASS==2
 				PRINTBSDFCHAR(_W_); PRINTBSDFCHAR(_A_); PRINTBSDFCHAR(_T_); PRINTBSDFCHAR(_E_); PRINTBSDFCHAR(_R_);
 				PRINTBSDFCHAR(_SPC_); PRINTBSDFCHAR(_SPC_); PRINTBSDFCHAR(_SPC_); PRINTBSDFCHAR(_SPC_); PRINTBSDFCHAR(_SPC_);
 				PRINTBSDFCHAR(_T_); PRINTBSDFCHAR(_O_); PRINTBSDFCHAR(_D_); PRINTBSDFCHAR(_O_); PRINTBSDFCHAR(_SPC_);
diff --git a/Engine/Shaders/Shared/StrataDefinitions.h b/Engine/Shaders/Shared/StrataDefinitions.h
--- a/Engine/Shaders/Shared/StrataDefinitions.h
+++ b/Engine/Shaders/Shared/StrataDefinitions.h
@@ -23,11 +23,10 @@
 
 
 #define STRATA_BSDF_TYPE_SLAB				0
-#define STRATA_BSDF_TYPE_SHEEN				1
-#define STRATA_BSDF_TYPE_VOLUMETRICFOGCLOUD	2
-#define STRATA_BSDF_TYPE_UNLIT				3
-#define STRATA_BSDF_TYPE_HAIR				4
-#define STRATA_BSDF_TYPE_SINGLELAYERWATER	5
+#define STRATA_BSDF_TYPE_VOLUMETRICFOGCLOUD	1
+#define STRATA_BSDF_TYPE_UNLIT				2
+#define STRATA_BSDF_TYPE_HAIR				3
+#define STRATA_BSDF_TYPE_SINGLELAYERWATER	4
 // When more than 8 BSDF must exists, please update STATE_BITS_BSDF and FStrataClassification.ShadingModels packing in Strata.ush
 
 #define STRATA_DATA_TILE_SIZE				8
diff --git a/Engine/Source/Runtime/Engine/Classes/Materials/MaterialExpressionStrata.h b/Engine/Source/Runtime/Engine/Classes/Materials/MaterialExpressionStrata.h
--- a/Engine/Source/Runtime/Engine/Classes/Materials/MaterialExpressionStrata.h
+++ b/Engine/Source/Runtime/Engine/Classes/Materials/MaterialExpressionStrata.h
@@ -33,7 +33,7 @@ class UMaterialExpressionStrataBSDF : public UMaterialExpression
 #endif
 };
 
-UCLASS(MinimalAPI, collapsecategories, hidecategories = Object, DisplayName = "Strata Slab BSDF")
+UCLASS(MinimalAPI, collapsecategories, hidecategories = Object, DisplayName = "Strata Slab")
 class UMaterialExpressionStrataSlabBSDF : public UMaterialExpressionStrataBSDF
 {
 	GENERATED_UCLASS_BODY()
@@ -124,6 +124,18 @@ class UMaterialExpressionStrataSlabBSDF : public UMaterialExpressionStrataBSDF
 	UPROPERTY()
 	FExpressionInput Thickness;
 
+	/**
+	 * The amount of fuzz on top of the surface used to simulate cloth-like appearance.
+	 */
+	UPROPERTY()
+	FExpressionInput FuzzAmount;
+
+	/**
+	 * The base colod of the fuzz.
+	 */
+	UPROPERTY()
+	FExpressionInput FuzzColor;
+
 	/** SubsurfaceProfile, for Screen Space Subsurface Scattering. The profile needs to be set up on both the Strata diffuse node, and the material node at the moment. */
 	UPROPERTY(EditAnywhere, BlueprintReadOnly, Category = Material, meta = (DisplayName = "Subsurface Profile"))
 	TObjectPtr<class USubsurfaceProfile> SubsurfaceProfile;
@@ -140,6 +152,7 @@ class UMaterialExpressionStrataSlabBSDF : public UMaterialExpressionStrataBSDF
 
 	bool HasEdgeColor() const;
 	bool HasThinFilm() const;
+	bool HasFuzz() const;
 	bool HasSSS() const;
 	bool HasSSSProfile() const;
 	bool HasDMFPPluggedIn() const;
@@ -148,41 +161,6 @@ class UMaterialExpressionStrataSlabBSDF : public UMaterialExpressionStrataBSDF
 	//~ End UMaterialExpression Interface
 };
 
-UCLASS(MinimalAPI, collapsecategories, hidecategories = Object, DisplayName = "Strata Sheen BSDF")
-class UMaterialExpressionStrataSheenBSDF : public UMaterialExpressionStrataBSDF
-{
-	GENERATED_UCLASS_BODY()
-
-	/**
-	* Defines the overall color of the Material. (type = float3, unit = unitless)
-	*/
-	UPROPERTY()
-	FExpressionInput BaseColor;
-
-	/**
-	 * Roughness (type = float, unit = unitless)
-	 */
-	UPROPERTY()
-	FExpressionInput Roughness;
-
-	/**
-	 * Take the surface normal as input. The normal is considered tangent or world space according to the space properties on the main material node. (type = float3, unit = unitless)
-	 */
-	UPROPERTY()
-	FExpressionInput Normal;
-
-	//~ Begin UMaterialExpression Interface
-#if WITH_EDITOR
-	virtual int32 Compile(class FMaterialCompiler* Compiler, int32 OutputIndex) override;
-	virtual void GetCaption(TArray<FString>& OutCaptions) const override;
-	virtual uint32 GetOutputType(int32 OutputIndex) override;
-	virtual uint32 GetInputType(int32 InputIndex) override;
-	virtual bool IsResultStrataMaterial(int32 OutputIndex) override;
-	virtual void GatherStrataMaterialInfo(FStrataMaterialInfo& StrataMaterialInfo, int32 OutputIndex) override;
-#endif
-	//~ End UMaterialExpression Interface
-};
-
 UCLASS(MinimalAPI, collapsecategories, hidecategories = Object, DisplayName = "Strata Volumetric-Fog-Cloud BSDF")
 class UMaterialExpressionStrataVolumetricFogCloudBSDF : public UMaterialExpressionStrataBSDF
 {
diff --git a/Engine/Source/Runtime/Engine/Private/Materials/HLSLMaterialTranslator.cpp b/Engine/Source/Runtime/Engine/Private/Materials/HLSLMaterialTranslator.cpp
--- a/Engine/Source/Runtime/Engine/Private/Materials/HLSLMaterialTranslator.cpp
+++ b/Engine/Source/Runtime/Engine/Private/Materials/HLSLMaterialTranslator.cpp
@@ -8840,12 +8840,18 @@ int32 FHLSLMaterialTranslator::StrataCreateAndRegisterNullMaterial()
 }
 
 int32 FHLSLMaterialTranslator::StrataSlabBSDF(
-	int32 BaseColor, int32 EdgeColor, int32 Specular, int32 Metallic, int32 Roughness, int32 Anisotropy,
-	int32 SSSProfileId, int32 SSSDMFP, int32 SSSDMFPScale, int32 EmissiveColor,
-	int32 Haziness, int32 ThinFilmThickness, int32 Thickness, int32 Normal, int32 Tangent, uint8 SharedNormalIndex)
+	int32 BaseColor, int32 EdgeColor, int32 Specular, int32 Metallic, int32 Roughness, 
+	int32 Anisotropy,
+	int32 SSSProfileId, int32 SSSDMFP, int32 SSSDMFPScale, 
+	int32 EmissiveColor,	
+	int32 Haziness, 
+	int32 ThinFilmThickness, 
+	int32 FuzzAmount, int32 FuzzColor,
+	int32 Thickness, 
+	int32 Normal, int32 Tangent, uint8 SharedNormalIndex)
 {
 	return AddCodeChunk(
-		MCT_Strata, TEXT("GetStrataSlabBSDF(%s, %s, %s, %s, %s, %s, %s, %s, %s, %s, %s, %s, %s, %u, Parameters.SharedNormals.NormalTypes) /* Normal = %s ; Tangent = %s */"),
+		MCT_Strata, TEXT("GetStrataSlabBSDF(%s, %s, %s, %s, %s, %s, %s, %s, %s, %s, %s, %s, %s, %s, %s, %u, Parameters.SharedNormals.NormalTypes) /* Normal = %s ; Tangent = %s */"),
 		*GetParameterCode(BaseColor),
 		*GetParameterCode(EdgeColor),
 		*GetParameterCode(Specular),
@@ -8858,6 +8864,8 @@ int32 FHLSLMaterialTranslator::StrataSlabBSDF(
 		*GetParameterCode(EmissiveColor),
 		*GetParameterCode(Haziness),
 		*GetParameterCode(ThinFilmThickness),
+		*GetParameterCode(FuzzAmount),
+		*GetParameterCode(FuzzColor),
 		*GetParameterCode(Thickness),
 		SharedNormalIndex,
 		*GetParameterCode(Normal),
diff --git a/Engine/Source/Runtime/Engine/Private/Materials/HLSLMaterialTranslator.h b/Engine/Source/Runtime/Engine/Private/Materials/HLSLMaterialTranslator.h
--- a/Engine/Source/Runtime/Engine/Private/Materials/HLSLMaterialTranslator.h
+++ b/Engine/Source/Runtime/Engine/Private/Materials/HLSLMaterialTranslator.h
@@ -970,7 +970,11 @@ protected:
 	virtual int32 StrataSlabBSDF(
 		int32 BaseColor, int32 EdgeColor, int32 Specular, int32 Metallic, int32 Roughness, int32 Anisotropy,
 		int32 SSSProfileId, int32 SSSDMFP, int32 SSSDMFPScale,
-		int32 EmissiveColor, int32 Haziness, int32 ThinFilmThickness, int32 Thickness,
+		int32 EmissiveColor, 
+		int32 Haziness, 
+		int32 ThinFilmThickness, 
+		int32 FuzzAmount, int32 FuzzColor, 
+		int32 Thickness,
 		int32 Normal, int32 Tangent, uint8 SharedNormalIndex) override;
 	virtual int32 StrataSheenBSDF(int32 BaseColor, int32 Roughness, int32 Normal, uint8 SharedNormalIndex) override;
 	virtual int32 StrataVolumetricFogCloudBSDF(int32 Albedo, int32 Extinction, int32 EmissiveColor, int32 AmbientOcclusion) override;
diff --git a/Engine/Source/Runtime/Engine/Private/Materials/Material.cpp b/Engine/Source/Runtime/Engine/Private/Materials/Material.cpp
--- a/Engine/Source/Runtime/Engine/Private/Materials/Material.cpp
+++ b/Engine/Source/Runtime/Engine/Private/Materials/Material.cpp
@@ -3711,55 +3711,22 @@ void UMaterial::ConvertMaterialToStrataMaterial()
 		}
 		else if (MaterialDomain == MD_Surface && ShadingModel == MSM_Cloth)
 		{
-			// Helpers
-			UMaterialExpressionConstant3Vector* ConstantZero3 = NewObject<UMaterialExpressionConstant3Vector>(this);
-			UMaterialExpressionConstant* ConstantZero1 = NewObject<UMaterialExpressionConstant>(this);
-			UMaterialExpressionPower* Power = NewObject<UMaterialExpressionPower>(this);
-			ConstantZero3->Constant = FLinearColor::Black;
-			ConstantZero1->R = 0.f;
-			Power->ConstExponent = 2.f;
-
-			// Sheen model affect only the specular component of a material, and lerp the slab specular with the Sheen specular. This is how the legacy cloth shading model works. In order to reproduce that, since we don't have fine grain control over the BSDF we:
-			// 1. Create a sheen BSDF
-			// 2. Create a slab with only specular value
-			// 3. Lerp the two slab based on the Fuzz amount
-			// 4. Create a slab with only diffuse value
-			// 5. Add the Diffuse slab and the previous other slab
-			UMaterialExpressionStrataSlabBSDF* DiffuseSlabBSDF = NewObject<UMaterialExpressionStrataSlabBSDF>(this);
-			UMaterialExpressionStrataSlabBSDF* SpecularSlabBSDF = NewObject<UMaterialExpressionStrataSlabBSDF>(this);
-			UMaterialExpressionStrataSheenBSDF* SheenBSDF = NewObject<UMaterialExpressionStrataSheenBSDF>(this);
-			
-			// Diffuse slab
-			MoveConnectionTo(BaseColor, DiffuseSlabBSDF, 0);		// BaseColor
-			DiffuseSlabBSDF->Metallic.Connect(0, ConstantZero1);	// Metallic
-			DiffuseSlabBSDF->Specular.Connect(0, ConstantZero1);	// Specular
-			CopyConnectionTo(Roughness, DiffuseSlabBSDF, 4);		// Roughness
-			CopyConnectionTo(Normal,	DiffuseSlabBSDF, 6);		// Normal
-
-			// Specular slab
-			SpecularSlabBSDF->BaseColor.Connect(0, ConstantZero3);	// BaseColor
-			SpecularSlabBSDF->Metallic.Connect(0, ConstantZero1);	// Metallic
-			CopyConnectionTo(Specular,	SpecularSlabBSDF, 3);		// Specular
-			CopyConnectionTo(Roughness, SpecularSlabBSDF, 4);		// Roughness
-			CopyConnectionTo(Normal,	SpecularSlabBSDF, 6);		// Normal
-
-			// Sheen 
-			MoveConnectionTo(SubsurfaceColor, Power, 0);			// Square the base color to match the legacy code
-			SheenBSDF->BaseColor.Connect(0, Power);					// BaseColor
-			CopyConnectionTo(Roughness, SheenBSDF, 1);				// Roughness
-			CopyConnectionTo(Normal, SheenBSDF, 2);					// Normal
-			
-			UMaterialExpressionStrataHorizontalMixing* Mix = NewObject<UMaterialExpressionStrataHorizontalMixing>(this);
-			UMaterialExpressionStrataAdd* Add = NewObject<UMaterialExpressionStrataAdd>(this);
-
-			Mix->Background.Connect(0, SpecularSlabBSDF);			// Background
-			Mix->Foreground.Connect(0, SheenBSDF);					// Foreground
-			MoveConnectionTo(ClearCoat, Mix, 2);					// Mix
+			UMaterialExpressionStrataSlabBSDF* SlabBSDF = NewObject<UMaterialExpressionStrataSlabBSDF>(this);
 
-			Add->A.Connect(0, Mix);									// A
-			Add->B.Connect(0, DiffuseSlabBSDF);						// B
+			MoveConnectionTo(BaseColor, SlabBSDF, 0);		// BaseColor
+			MoveConnectionTo(Metallic, SlabBSDF, 2);		// Metallic
+			MoveConnectionTo(Specular, SlabBSDF, 3);		// Specular
+			MoveConnectionTo(Roughness, SlabBSDF, 4);		// Roughness
+			MoveConnectionTo(Anisotropy, SlabBSDF, 5);		// Anisotropy
+			MoveConnectionTo(Normal, SlabBSDF, 6);			// Normal
+			MoveConnectionTo(Tangent, SlabBSDF, 7);			// Tangent
+			MoveConnectionTo(EmissiveColor, SlabBSDF, 10);	// Emissive
+			MoveConnectionTo(Tangent, SlabBSDF, 7);			// Tangent
+			MoveConnectionTo(EmissiveColor, SlabBSDF, 10);	// Emissive
+			MoveConnectionTo(ClearCoat, SlabBSDF, 14);		// FuzzAMount
+			MoveConnectionTo(SubsurfaceColor, SlabBSDF, 15);// FuzzColor
 
-			FrontMaterial.Connect(0, Add);
+			FrontMaterial.Connect(0, SlabBSDF);
 		}
 		else if (MaterialDomain == MD_Surface && ShadingModel == MSM_Hair)
 		{
diff --git a/Engine/Source/Runtime/Engine/Private/Materials/MaterialExpressions.cpp b/Engine/Source/Runtime/Engine/Private/Materials/MaterialExpressions.cpp
--- a/Engine/Source/Runtime/Engine/Private/Materials/MaterialExpressions.cpp
+++ b/Engine/Source/Runtime/Engine/Private/Materials/MaterialExpressions.cpp
@@ -19873,6 +19873,7 @@ int32 UMaterialExpressionStrataSlabBSDF::Compile(class FMaterialCompiler* Compil
 
 	const bool bHasEdgeColor = HasEdgeColor();
 	const bool bHasThinFilm = HasThinFilm();
+	const bool bHasFuzz = HasFuzz();
 	const bool bHasDMFPPluggedIn = HasDMFPPluggedIn();
 
 	int32 SSSProfileCodeChunk = INDEX_NONE;
@@ -19896,11 +19897,13 @@ int32 UMaterialExpressionStrataSlabBSDF::Compile(class FMaterialCompiler* Compil
 		CompileWithDefaultFloat3(Compiler, EmissiveColor, 0.0f, 0.0f, 0.0f),
 		CompileWithDefaultFloat1(Compiler, Haziness, 0.0f),
 		CompileWithDefaultFloat1(Compiler, ThinFilmThickness, 0.0f),
+		CompileWithDefaultFloat1(Compiler, FuzzAmount, 0.0f),
+		CompileWithDefaultFloat3(Compiler, FuzzColor, 0.0f, 0.0f, 0.0f),
 		CompileWithDefaultFloat1(Compiler, Thickness, STRATA_LAYER_DEFAULT_THICKNESS_CM),
 		NormalCodeChunk,
 		TangentCodeChunk,
 		SharedNormalIndex);
-	StrataCompilationInfoCreateSingleBSDFMaterial(Compiler, OutputCodeChunk, SharedNormalIndex, STRATA_BSDF_TYPE_SLAB, bHasSSS, bHasDMFPPluggedIn, bHasEdgeColor, bHasThinFilm);
+	StrataCompilationInfoCreateSingleBSDFMaterial(Compiler, OutputCodeChunk, SharedNormalIndex, STRATA_BSDF_TYPE_SLAB, bHasSSS, bHasDMFPPluggedIn, bHasEdgeColor, bHasThinFilm, bHasFuzz);
 
 	return OutputCodeChunk;
 }
@@ -19961,6 +19964,12 @@ uint32 UMaterialExpressionStrataSlabBSDF::GetInputType(int32 InputIndex)
 	case 13:
 		return MCT_Float1; // Thickness
 		break;
+	case 14:
+		return MCT_Float1; // FuzzAmount
+		break;
+	case 15:
+		return MCT_Float3; // FuzzColor
+		break;
 	}
 
 	check(false);
@@ -20025,6 +20034,14 @@ FName UMaterialExpressionStrataSlabBSDF::GetInputName(int32 InputIndex) const
 	{
 		return TEXT("Thickness");
 	}
+	else if (InputIndex == 14)
+	{
+		return TEXT("FuzzAmount");
+	}
+	else if (InputIndex == 15)
+	{
+		return TEXT("FuzzColor");
+	}
 	return TEXT("Unknown");
 }
 
@@ -20075,80 +20092,14 @@ bool UMaterialExpressionStrataSlabBSDF::HasThinFilm() const
 	return ThinFilmThickness.IsConnected();
 }
 
-bool UMaterialExpressionStrataSlabBSDF::HasAnisotropy() const
+bool UMaterialExpressionStrataSlabBSDF::HasFuzz() const
 {
-	return Anisotropy.IsConnected();
+	return FuzzAmount.IsConnected();
 }
-#endif // WITH_EDITOR
 
-
-UMaterialExpressionStrataSheenBSDF::UMaterialExpressionStrataSheenBSDF(const FObjectInitializer& ObjectInitializer)
-	: Super(ObjectInitializer)
-{
-	struct FConstructorStatics
-	{
-		FText NAME_Strata;
-		FConstructorStatics() : NAME_Strata(LOCTEXT("Strata BSDFs", "Strata BSDFs")) { }
-	};
-	static FConstructorStatics ConstructorStatics;
-#if WITH_EDITORONLY_DATA
-	MenuCategories.Add(ConstructorStatics.NAME_Strata);
-#endif
-}
-
-#if WITH_EDITOR
-int32 UMaterialExpressionStrataSheenBSDF::Compile(class FMaterialCompiler* Compiler, int32 OutputIndex)
-{
-	int32 NormalCodeChunk = CompileWithDefaultNormalWS(Compiler, Normal);
-	uint8 SharedNormalIndex = StrataCompilationInfoCreateSharedNormal(Compiler, NormalCodeChunk);
-
-	int32 OutputCodeChunk = Compiler->StrataSheenBSDF(
-		CompileWithDefaultFloat3(Compiler, BaseColor,	0.0f, 0.0f, 0.0f),
-		CompileWithDefaultFloat1(Compiler, Roughness,	0.0f),
-		NormalCodeChunk,
-		SharedNormalIndex);
-	StrataCompilationInfoCreateSingleBSDFMaterial(Compiler, OutputCodeChunk, SharedNormalIndex, STRATA_BSDF_TYPE_SHEEN);
-
-	return OutputCodeChunk;
-}
-
-void UMaterialExpressionStrataSheenBSDF::GetCaption(TArray<FString>& OutCaptions) const
-{
-	OutCaptions.Add(TEXT("Strata Sheen BSDF"));
-}
-
-uint32 UMaterialExpressionStrataSheenBSDF::GetOutputType(int32 OutputIndex)
-{
-	return MCT_Strata;
-}
-
-uint32 UMaterialExpressionStrataSheenBSDF::GetInputType(int32 InputIndex)
-{
-	switch (InputIndex)
-	{
-	case 0:
-		return MCT_Float3;
-		break;
-	case 1:
-		return MCT_Float1;
-		break;	
-	case 2:
-		return MCT_Float3;
-		break;
-	}
-
-	check(false);
-	return MCT_Float1;
-}
-
-bool UMaterialExpressionStrataSheenBSDF::IsResultStrataMaterial(int32 OutputIndex)
-{
-	return true;
-}
-
-void UMaterialExpressionStrataSheenBSDF::GatherStrataMaterialInfo(FStrataMaterialInfo& StrataMaterialInfo, int32 OutputIndex)
+bool UMaterialExpressionStrataSlabBSDF::HasAnisotropy() const
 {
-	StrataMaterialInfo.AddShadingModel(SSM_DefaultLit);
+	return Anisotropy.IsConnected();
 }
 #endif // WITH_EDITOR
 
diff --git a/Engine/Source/Runtime/Engine/Private/Materials/StrataMaterial.cpp b/Engine/Source/Runtime/Engine/Private/Materials/StrataMaterial.cpp
--- a/Engine/Source/Runtime/Engine/Private/Materials/StrataMaterial.cpp
+++ b/Engine/Source/Runtime/Engine/Private/Materials/StrataMaterial.cpp
@@ -12,9 +12,6 @@ FString GetStrataBSDFName(uint8 BSDFType)
 	case STRATA_BSDF_TYPE_SLAB:
 		return TEXT("SLAB");
 		break;
-	case STRATA_BSDF_TYPE_SHEEN:
-		return TEXT("SHEEN");
-		break;
 	case STRATA_BSDF_TYPE_VOLUMETRICFOGCLOUD:
 		return TEXT("VOLUMETRICFOGCLOUD");
 		break;
@@ -57,7 +54,7 @@ uint8 StrataCompilationInfoCreateSharedNormal(FMaterialCompiler* Compiler, int32
 	return Compiler->StrataCompilationInfoRegisterSharedNormalIndex(NormalCodeChunk, TangentCodeChunk);
 }
 
-void StrataCompilationInfoCreateSingleBSDFMaterial(FMaterialCompiler* Compiler, int32 CodeChunk, uint8 SharedNormalIndex, uint8 BSDFType, bool bHasSSS, bool bHasDMFPPluggedIn, bool bHasEdgeColor, bool bHasThinFilm)
+void StrataCompilationInfoCreateSingleBSDFMaterial(FMaterialCompiler* Compiler, int32 CodeChunk, uint8 SharedNormalIndex, uint8 BSDFType, bool bHasSSS, bool bHasDMFPPluggedIn, bool bHasEdgeColor, bool bHasThinFilm, bool bHasFuzz)
 {
 	FStrataMaterialCompilationInfo StrataInfo;
 	StrataInfo.LayerCount = 1;
@@ -68,6 +65,7 @@ void StrataCompilationInfoCreateSingleBSDFMaterial(FMaterialCompiler* Compiler,
 	StrataInfo.Layers[0].BSDFs[0].bHasDMFPPluggedIn = bHasDMFPPluggedIn;
 	StrataInfo.Layers[0].BSDFs[0].bHasEdgeColor = bHasEdgeColor;
 	StrataInfo.Layers[0].BSDFs[0].bHasThinFilm = bHasThinFilm;
+	StrataInfo.Layers[0].BSDFs[0].bHasFuzz = bHasFuzz;
 	UpdateTotalBSDFCount(Compiler, StrataInfo);
 	Compiler->StrataCompilationInfoRegisterCodeChunk(CodeChunk, StrataInfo);
 }
@@ -303,11 +301,10 @@ FStrataMaterialAnalysisResult StrataCompilationInfoMaterialAnalysis(FMaterialCom
 				{
 					Result.RequestedByteCount += UintByteSize;
 				}
-				break;
-			}
-			case STRATA_BSDF_TYPE_SHEEN:
-			{
-				Result.RequestedByteCount += UintByteSize;
+				if (BSDF.bHasFuzz)
+				{
+					Result.RequestedByteCount += UintByteSize;
+				}
 				break;
 			}
 			case STRATA_BSDF_TYPE_HAIR:
diff --git a/Engine/Source/Runtime/Engine/Private/Materials/StrataMaterial.h b/Engine/Source/Runtime/Engine/Private/Materials/StrataMaterial.h
--- a/Engine/Source/Runtime/Engine/Private/Materials/StrataMaterial.h
+++ b/Engine/Source/Runtime/Engine/Private/Materials/StrataMaterial.h
@@ -22,6 +22,7 @@ struct FStrataMaterialCompilationInfo
 		bool  bHasDMFPPluggedIn;
 		bool  bHasEdgeColor;
 		bool  bHasThinFilm;
+		bool  bHasFuzz;
 	};
 
 	struct FLayer
@@ -48,7 +49,7 @@ FString GetStrataBSDFName(uint8 BSDFType);
 uint8 StrataCompilationInfoCreateSharedNormal(FMaterialCompiler* Compiler, int32 NormalCodeChunk);
 uint8 StrataCompilationInfoCreateSharedNormal(FMaterialCompiler* Compiler, int32 NormalCodeChunk, int32 TangentCodeChunk);
 void  StrataCompilationInfoCreateNullBSDF(FMaterialCompiler* Compiler, int32 CodeChunk);
-void  StrataCompilationInfoCreateSingleBSDFMaterial(FMaterialCompiler* Compiler, int32 CodeChunk, uint8 SharedNormalIndex, uint8 BSDFType, bool bHasSSS = false, bool bHasDMFPPluggedIn = false, bool bHasEdgeColor = false, bool bHasThinFilm = false);
+void  StrataCompilationInfoCreateSingleBSDFMaterial(FMaterialCompiler* Compiler, int32 CodeChunk, uint8 SharedNormalIndex, uint8 BSDFType, bool bHasSSS = false, bool bHasDMFPPluggedIn = false, bool bHasEdgeColor = false, bool bHasThinFilm = false, bool bHasFuzz = false);
 
 FStrataMaterialCompilationInfo StrataCompilationInfoAdd(FMaterialCompiler* Compiler, const FStrataMaterialCompilationInfo& A, const FStrataMaterialCompilationInfo& B);
 FStrataMaterialCompilationInfo StrataCompilationInfoMultiply(FMaterialCompiler* Compiler, const FStrataMaterialCompilationInfo& A);
diff --git a/Engine/Source/Runtime/Engine/Public/MaterialCompiler.h b/Engine/Source/Runtime/Engine/Public/MaterialCompiler.h
--- a/Engine/Source/Runtime/Engine/Public/MaterialCompiler.h
+++ b/Engine/Source/Runtime/Engine/Public/MaterialCompiler.h
@@ -413,7 +413,11 @@ public:
 	virtual int32 StrataSlabBSDF(
 		int32 BaseColor, int32 EdgeColor, int32 Specular, int32 Metallic, int32 Roughness, int32 Anisotropy, 
 		int32 SSSProfileId, int32 SSSDMFP, int32 SSSDMFPScale,
-		int32 EmissiveColor, int32 Haziness, int32 ThinFilmThickness, int32 Thickness,
+		int32 EmissiveColor, 
+		int32 Haziness, 
+		int32 ThinFilmThickness, 
+		int32 FuzzAmount, int32 FuzzColor,
+		int32 Thickness,
 		int32 Normal, int32 Tangent, uint8 SharedNormalIndex) = 0;
 	virtual int32 StrataSheenBSDF(int32 BaseColor, int32 Roughness, int32 Normal, uint8 SharedNormalIndex) = 0;
 	virtual int32 StrataVolumetricFogCloudBSDF(int32 Albedo, int32 Extinction, int32 EmissiveColor, int32 AmbientOcclusion) = 0;
@@ -898,13 +902,22 @@ public:
 	virtual int32 StrataSlabBSDF(
 		int32 BaseColor, int32 EdgeColor, int32 Specular, int32 Metallic, int32 Roughness, int32 Anisotropy,
 		int32 SSSProfileId, int32 SSSDMFP, int32 SSSDMFPScale,
-		int32 EmissiveColor, int32 Haziness, int32 ThinFilmThickness, int32 Thickness,
+		int32 EmissiveColor, 
+		int32 Haziness, 
+		int32 ThinFilmThickness, 
+		int32 FuzzAmount, int32 FuzzColor,
+		int32 Thickness,
 		int32 Normal, int32 Tangent, uint8 SharedNormalIndex) override
 	{
 		return Compiler->StrataSlabBSDF(
 			BaseColor, EdgeColor, Specular, Metallic, Roughness, Anisotropy,
 			SSSProfileId, SSSDMFP, SSSDMFPScale,
-			EmissiveColor, Haziness, ThinFilmThickness, Thickness, Normal, Tangent, SharedNormalIndex);
+			EmissiveColor, 
+			Haziness, 
+			ThinFilmThickness,
+			FuzzAmount, FuzzColor,
+			Thickness,
+			Normal, Tangent, SharedNormalIndex);
 	}
 
 	virtual int32 StrataSheenBSDF(int32 BaseColor, int32 Roughness, int32 Normal, uint8 SharedNormalIndex) override
diff --git a/Engine/Source/Runtime/Renderer/Private/Strata/Strata.cpp b/Engine/Source/Runtime/Renderer/Private/Strata/Strata.cpp
--- a/Engine/Source/Runtime/Renderer/Private/Strata/Strata.cpp
+++ b/Engine/Source/Runtime/Renderer/Private/Strata/Strata.cpp
@@ -315,12 +315,14 @@ TRDGUniformBufferRef<FStrataGlobalUniformParameters> BindStrataGlobalUniformPara
 ////////////////////////////////////////////////////////////////////////// 
 // Debug
 
+#define VISUALIZE_MATERIAL_PASS_COUNT 3
+
 class FVisualizeMaterialPS : public FGlobalShader
 {
 	DECLARE_GLOBAL_SHADER(FVisualizeMaterialPS);
 	SHADER_USE_PARAMETER_STRUCT(FVisualizeMaterialPS, FGlobalShader);
 
-	class FBSDFPass : SHADER_PERMUTATION_INT("PERMUTATION_BSDF_PASS", 4);
+	class FBSDFPass : SHADER_PERMUTATION_INT("PERMUTATION_BSDF_PASS", VISUALIZE_MATERIAL_PASS_COUNT);
 	using FPermutationDomain = TShaderPermutationDomain<FBSDFPass>;
 
 	BEGIN_SHADER_PARAMETER_STRUCT(FParameters, )
@@ -376,7 +378,7 @@ static void AddVisualizeMaterialPasses(FRDGBuilder& GraphBuilder, const FViewInf
 			ShaderDrawDebug::SetParameters(GraphBuilder, View.ShaderDrawData, PassParameters->ShaderDrawParameters);
 		}
 
-		for (uint32 j = 0; j < 4; ++j)
+		for (uint32 j = 0; j < VISUALIZE_MATERIAL_PASS_COUNT; ++j)
 		{
 			FVisualizeMaterialPS::FPermutationDomain PermutationVector;
 			PermutationVector.Set<typename FVisualizeMaterialPS::FBSDFPass>(j);
