Path Tracer light sampling functions only need 2 random numbers

Tweak the API so the light sampling functions only get 2 random numbers. This ensures we are consistent in which dimensions we use across all lights and make it easier to play with which dimensions are used for what. I played a bit with swizzling around the light selection vs light sampling dimensions but couldn't get consistent improvements, so I'm leaving it alone for now.

diff --git a/Engine/Plugins/Experimental/GPULightmass/Shaders/Private/LightmapPathTracing.usf b/Engine/Plugins/Experimental/GPULightmass/Shaders/Private/LightmapPathTracing.usf
--- a/Engine/Plugins/Experimental/GPULightmass/Shaders/Private/LightmapPathTracing.usf
+++ b/Engine/Plugins/Experimental/GPULightmass/Shaders/Private/LightmapPathTracing.usf
@@ -393,7 +393,7 @@ void PathTracingKernel(
 
 				SelectLight(RandSample.x * LightPickingCdfSum, SceneLightCount, LightPickingCdf, LightId, LightPickPdf);
 
-				FLightSample LightSample = SampleLight(LightId, RandSample, WorldPos, WorldNormal);
+				FLightSample LightSample = SampleLight(LightId, RandSample.yz, WorldPos, WorldNormal);
 
 				LightPickPdf /= LightPickingCdfSum;
 
diff --git a/Engine/Shaders/Private/PathTracing/Light/PathTracingCapsuleLight.ush b/Engine/Shaders/Private/PathTracing/Light/PathTracingCapsuleLight.ush
--- a/Engine/Shaders/Private/PathTracing/Light/PathTracingCapsuleLight.ush
+++ b/Engine/Shaders/Private/PathTracing/Light/PathTracingCapsuleLight.ush
@@ -126,7 +126,7 @@ FLightHit CapsuleLight_TraceLight(RayDesc Ray, int LightId)
 
 FLightSample CapsuleLight_SampleLight(
 	int LightId,
-	float4 RandSample,
+	float2 RandSample,
 	float3 WorldPos,
 	float3 WorldNormal
 )
@@ -153,7 +153,7 @@ FLightSample CapsuleLight_SampleLight(
 	float3 LightDirection = GetPosition(LightId) - WorldPos;
 	FCapsuleSphericalBounds CapsuleBounds = CapsuleGetSphericalBounds(LightDirection, Axis, Radius, SourceLength);
 
-	float4 Result = SampleCapsuleBounds(CapsuleBounds, RandSample.yz);
+	float4 Result = SampleCapsuleBounds(CapsuleBounds, RandSample);
 
 	float3 L = Result.xyz;
 	float SolidAngle = Result.w;
@@ -178,11 +178,11 @@ FLightSample CapsuleLight_SampleLight(
 
 	float3 LightPoint;
 	float3 LightNormal;
-	if (RandSample.y < Prob)
+	if (RandSample.x < Prob)
 	{
-		RandSample.y /= Prob;
+		RandSample.x /= Prob;
 		// sample the caps
-		float4 Result = UniformSampleSphere(float2(RandSample.yz));
+		float4 Result = UniformSampleSphere(float2(RandSample));
 
 		LightPoint = Result.xyz * Radius;
 		LightPoint.z += 0.5 * SourceLength * sign(Result.z);
@@ -191,14 +191,14 @@ FLightSample CapsuleLight_SampleLight(
 	}
 	else
 	{
-		RandSample.y -= Prob;
-		RandSample.y /= 1 - Prob;
+		RandSample.x -= Prob;
+		RandSample.x /= 1 - Prob;
 		// sample the cylinder body
 
-		float Phi = 2 * PI * RandSample.y;
+		float Phi = 2 * PI * RandSample.x;
 
 		LightNormal = float3(cos(Phi), sin(Phi), 0.0);
-		LightPoint = float3(Radius * LightNormal.xy, (RandSample.z - 0.5) * SourceLength);
+		LightPoint = float3(Radius * LightNormal.xy, (RandSample.y - 0.5) * SourceLength);
 	}
 	float3x3 CylinderBasis = GetTangentBasis(Axis);
 
diff --git a/Engine/Shaders/Private/PathTracing/Light/PathTracingDirectionalLight.ush b/Engine/Shaders/Private/PathTracing/Light/PathTracingDirectionalLight.ush
--- a/Engine/Shaders/Private/PathTracing/Light/PathTracingDirectionalLight.ush
+++ b/Engine/Shaders/Private/PathTracing/Light/PathTracingDirectionalLight.ush
@@ -35,13 +35,13 @@ FLightHit DirectionalLight_TraceLight(RayDesc Ray, int LightId)
 
 FLightSample DirectionalLight_SampleLight(
 	int LightId,
-	float4 RandSample,
+	float2 RandSample,
 	float3 WorldPos,
 	float3 WorldNormal
 )
 {
 	float SinThetaMax = GetRadius(LightId); // currently sin(Angle / 2)
-	float4 DirAndPdf = UniformSampleConeRobust(RandSample.yz, SinThetaMax * SinThetaMax);
+	float4 DirAndPdf = UniformSampleConeRobust(RandSample, SinThetaMax * SinThetaMax);
 	float3 Direction = TangentToWorld(DirAndPdf.xyz, GetNormal(LightId));
 
 	// Because the light is normalized by the solid angle, the radiance/pdf ratio is just the color
diff --git a/Engine/Shaders/Private/PathTracing/Light/PathTracingLightSampling.ush b/Engine/Shaders/Private/PathTracing/Light/PathTracingLightSampling.ush
--- a/Engine/Shaders/Private/PathTracing/Light/PathTracingLightSampling.ush
+++ b/Engine/Shaders/Private/PathTracing/Light/PathTracingLightSampling.ush
@@ -36,7 +36,7 @@ FLightHit TraceLight(RayDesc Ray, int LightId)
 // Choose a point on a light as seen from the given shading point
 FLightSample SampleLight(
 	int LightId,
-	float4 RandSample,
+	float2 RandSample,
 	float3 WorldPos,
 	float3 WorldNormal
 )
@@ -146,7 +146,7 @@ void SelectLight(float RandSample, int NumLights, inout float LightPickingCdf[RA
 }
 
 bool SampleLight(
-	float4 RandSample,
+	float3 RandSample,
 	float3 WorldPos,
 	float3 WorldNormal,
 	uint PrimitiveLightingChannelMask,
@@ -186,7 +186,7 @@ bool SampleLight(
 
 	LightPickPdf /= LightPickingCdfSum;
 
-	FLightSample Result = SampleLight(OutLightId, RandSample, WorldPos, WorldNormal);
+	FLightSample Result = SampleLight(OutLightId, RandSample.yz, WorldPos, WorldNormal);
 	OutDirection = Result.Direction;
 	OutDistance = Result.Distance;
 	OutRadianceOverPdf = Result.RadianceOverPdf / LightPickPdf;
diff --git a/Engine/Shaders/Private/PathTracing/Light/PathTracingPointLight.ush b/Engine/Shaders/Private/PathTracing/Light/PathTracingPointLight.ush
--- a/Engine/Shaders/Private/PathTracing/Light/PathTracingPointLight.ush
+++ b/Engine/Shaders/Private/PathTracing/Light/PathTracingPointLight.ush
@@ -53,7 +53,7 @@ FLightHit PointLight_TraceLight(RayDesc Ray, int LightId)
 
 FLightSample PointLight_SampleLight(
 	int LightId,
-	float4 RandSample,
+	float2 RandSample,
 	float3 WorldPos,
 	float3 WorldNormal
 )
@@ -78,7 +78,7 @@ FLightSample PointLight_SampleLight(
 		float SinThetaMax2 = saturate(Radius2 / LightDistanceSquared);
 
 		// #dxr_todo: find a better way of handling the region inside the light than just clamping to 1.0 here
-		float4 DirAndPdf = UniformSampleConeRobust(RandSample.yz, SinThetaMax2);
+		float4 DirAndPdf = UniformSampleConeRobust(RandSample, SinThetaMax2);
 
 		float CosTheta = DirAndPdf.z;
 		float SinTheta2 = 1.0 - CosTheta * CosTheta;
diff --git a/Engine/Shaders/Private/PathTracing/Light/PathTracingRectLight.ush b/Engine/Shaders/Private/PathTracing/Light/PathTracingRectLight.ush
--- a/Engine/Shaders/Private/PathTracing/Light/PathTracingRectLight.ush
+++ b/Engine/Shaders/Private/PathTracing/Light/PathTracingRectLight.ush
@@ -138,7 +138,7 @@ FLightHit RectLight_TraceLight(RayDesc Ray, int LightId)
 
 FLightSample RectLight_SampleLight(
 	int LightId,
-	float4 RandSample,
+	float2 RandSample,
 	float3 WorldPos,
 	float3 WorldNormal
 )
@@ -167,11 +167,11 @@ FLightSample RectLight_SampleLight(
 	FSphericalRect SphericalRect = BuildSphericalRect(Rect);
 
 #if USE_RECT_LIGHT_TEXTURES
-	FSphericalRectSample Sample = UniformSampleSphericalRectWithUV(RandSample.yz, SphericalRect);
+	FSphericalRectSample Sample = UniformSampleSphericalRectWithUV(RandSample, SphericalRect);
 	float2 ClippedUV = Sample.UV;
 	float3 Direction = Sample.Direction;
 #else
-	float3 Direction = UniformSampleSphericalRect(RandSample.yz, SphericalRect);
+	float3 Direction = UniformSampleSphericalRect(RandSample, SphericalRect);
 #endif
 
 	float3 RadianceOverPdf = GetColor(LightId) * SphericalRect.SolidAngle * ComputeIESAttenuation(LightId, WorldPos);
diff --git a/Engine/Shaders/Private/PathTracing/Light/PathTracingSkyLight.ush b/Engine/Shaders/Private/PathTracing/Light/PathTracingSkyLight.ush
--- a/Engine/Shaders/Private/PathTracing/Light/PathTracingSkyLight.ush
+++ b/Engine/Shaders/Private/PathTracing/Light/PathTracingSkyLight.ush
@@ -50,7 +50,7 @@ FLightHit SkyLight_TraceLight(RayDesc Ray, int LightId) {
 
 FLightSample SkyLight_SampleLight(
 	int LightId,
-	float4 RandSample,
+	float2 RandSample,
 	float3 WorldPos,
 	float3 WorldNormal
 )
@@ -61,14 +61,14 @@ FLightSample SkyLight_SampleLight(
 	Ray.Origin = WorldPos;
 
 #if USE_HIERARCHICAL_IMPORTANCE_SAMPLING == 0
-	float4 Direction_Tangent = UniformSampleHemisphere(RandSample.yz);
+	float4 Direction_Tangent = UniformSampleHemisphere(RandSample);
 	float3 OutDirection = TangentToWorld(Direction_Tangent.xyz, WorldNormal);
 	float OutPdf = Direction_Tangent.w;
 	float2 UV = InverseEquiAreaSphericalMapping(OutDirection.yzx);
 	float3 Radiance = SkylightTexture.SampleLevel(SkylightTextureSampler, UV, 0).xyz;
 	return CreateLightSample(Radiance / OutPdf, OutPdf, OutDirection, POSITIVE_INFINITY);
 #else
-	float2 UV = RandSample.yz;
+	float2 UV = RandSample;
 
 	int3 Pixel = int3(0, 0, SkylightMipCount - 2);
 	for (; Pixel.z >= 0; Pixel.z--)
diff --git a/Engine/Shaders/Private/PathTracing/Light/PathTracingSpotLight.ush b/Engine/Shaders/Private/PathTracing/Light/PathTracingSpotLight.ush
--- a/Engine/Shaders/Private/PathTracing/Light/PathTracingSpotLight.ush
+++ b/Engine/Shaders/Private/PathTracing/Light/PathTracingSpotLight.ush
@@ -42,7 +42,7 @@ FLightHit SpotLight_TraceLight(RayDesc Ray, int LightId) {
 
 FLightSample SpotLight_SampleLight(
 	int LightId,
-	float4 RandSample,
+	float2 RandSample,
 	float3 WorldPos,
 	float3 WorldNormal
 )
diff --git a/Engine/Shaders/Private/PathTracing/PathTracing.usf b/Engine/Shaders/Private/PathTracing/PathTracing.usf
--- a/Engine/Shaders/Private/PathTracing/PathTracing.usf
+++ b/Engine/Shaders/Private/PathTracing/PathTracing.usf
@@ -620,7 +620,7 @@ RAY_TRACING_ENTRY_RAYGEN(PathTracingMainRG)
 				LightPickPdf /= LightPickingCdfSum;
 
 				LightId = GetLightId(LightId, LightLoopCount);
-				FLightSample LightSample = SampleLight(LightId, RandSample, WorldPos, WorldNormal);
+				FLightSample LightSample = SampleLight(LightId, RandSample.yz, WorldPos, WorldNormal);
 
 				LightSample.RadianceOverPdf /= LightPickPdf;
 				LightSample.Pdf *= LightPickPdf;
diff --git a/Engine/Shaders/Private/RayTracing/RayTracingCreateGatherPointsRGS.usf b/Engine/Shaders/Private/RayTracing/RayTracingCreateGatherPointsRGS.usf
--- a/Engine/Shaders/Private/RayTracing/RayTracingCreateGatherPointsRGS.usf
+++ b/Engine/Shaders/Private/RayTracing/RayTracingCreateGatherPointsRGS.usf
@@ -384,10 +384,10 @@ RAY_TRACING_ENTRY_RAYGEN(RayTracingCreateGatherPointsRGS)
 				float3 LightRadianceOverPdf = 0;
 				float NeePdf = 0.0;
 
-				float4 RandSample4 = RandomSequence_GenerateSample4D(RandSequence);
+				float3 RandSample3 = RandomSequence_GenerateSample3D(RandSequence);
 
 				RayDesc LightRay;
-				if (SampleLight(RandSample4,
+				if (SampleLight(RandSample3,
 								Position,
 								Normal,
 								Payload.PrimitiveLightingChannelMask,
diff --git a/Engine/Shaders/Private/RayTracing/RayTracingGlobalIlluminationRGS.usf b/Engine/Shaders/Private/RayTracing/RayTracingGlobalIlluminationRGS.usf
--- a/Engine/Shaders/Private/RayTracing/RayTracingGlobalIlluminationRGS.usf
+++ b/Engine/Shaders/Private/RayTracing/RayTracingGlobalIlluminationRGS.usf
@@ -229,10 +229,10 @@ RAY_TRACING_ENTRY_RAYGEN(GlobalIlluminationRGS)
 				float3 LightRadianceOverPdf = 0;
 				float NeePdf = 0.0;
 
-				float4 RandSample4 = RandomSequence_GenerateSample4D(RandSequence);
+				float3 RandSample3 = RandomSequence_GenerateSample3D(RandSequence);
 
 				RayDesc LightRay;
-				if (SampleLight(RandSample4,
+				if (SampleLight(RandSample3,
 								Ray.Origin,
 								Payload.WorldNormal,
 								Payload.PrimitiveLightingChannelMask,
