Fixed NaN in Strata for low specular/roughness values due to multiple scattering look up table.

diff --git a/Engine/Shaders/Private/Strata/StrataMultipleScattering.ush b/Engine/Shaders/Private/Strata/StrataMultipleScattering.ush
--- a/Engine/Shaders/Private/Strata/StrataMultipleScattering.ush
+++ b/Engine/Shaders/Private/Strata/StrataMultipleScattering.ush
@@ -42,7 +42,8 @@ float3 StrataEvaluateMultipleScatteringConductor(float3 F0, FStrataGGXLUTResult
 	//  f_ms(wo,wi) = F_ss . (1 +  F0 (1-E_R)/ E_R)
 	
 	const float3 E_R = F0 * LUT.A_R + (1 - F0) * LUT.B_R;
-	return (1 + F0 * (1 - E_R) / E_R);
+	const float3 Safe_E_R = max(1e-8f, 0.0f);	// the threshold making E_R safe is lower than any of the non-zero values in the current LUT
+	return (1 + F0 * (1 - E_R) / Safe_E_R);
 }
 
 // Return a scaling factor to multiple a specular GGX interface with, to add multiple scattering contribution (for dieletric)
