Fix minor warnings in RT shaders

diff --git a/Engine/Shaders/Private/RayTracing/RayTracingFinalGatherRGS.usf b/Engine/Shaders/Private/RayTracing/RayTracingFinalGatherRGS.usf
--- a/Engine/Shaders/Private/RayTracing/RayTracingFinalGatherRGS.usf
+++ b/Engine/Shaders/Private/RayTracing/RayTracingFinalGatherRGS.usf
@@ -66,7 +66,7 @@ bool ShouldBilateralReject(float2 PixelCoord, float Z, float LocalZ, float3 Worl
 
 RAY_TRACING_ENTRY_RAYGEN(RayTracingFinalGatherRGS)
 {
-	uint2 DispatchThreadId = DispatchRaysIndex().xy + View.ViewRectMin + uint2(RenderTileOffsetX, RenderTileOffsetY);
+	uint2 DispatchThreadId = DispatchRaysIndex().xy + uint2(View.ViewRectMin.xy) + uint2(RenderTileOffsetX, RenderTileOffsetY);
 	RWGlobalIlluminationUAV[DispatchThreadId] = 0.0;
 	RWGlobalIlluminationRayDistanceUAV[DispatchThreadId] = float2(-1.0, 0.0);
 
@@ -166,7 +166,7 @@ RAY_TRACING_ENTRY_RAYGEN(RayTracingFinalGatherRGS)
 		}
 
 		float2 ScreenUVTime = ScreenPositionTime.xy * float2(0.5, -0.5) + 0.5;
-		GatherPointsIndex = DispatchThreadIdToLinearIndex(int2(ScreenUVTime * GatherPointsResolution + View.ViewRectMin));
+		GatherPointsIndex = DispatchThreadIdToLinearIndex(int2(ScreenUVTime) * GatherPointsResolution + int2(View.ViewRectMin.xy));
 #endif
 
 		FGatherSample GatherSample = ReadGatherSample(GatherPointsBuffer, GatherPointsIndex, LocalSampleIndex, GatherPointsResolution.x * GatherPointsResolution.y);
@@ -280,7 +280,7 @@ RAY_TRACING_ENTRY_RAYGEN(RayTracingFinalGatherRGS)
 			float4 ScreenPositionTime = mul(float4(WorldPosition, 1.0), GatherPointData.ViewMatrices[ParentSampleIndex]);
 			ScreenPositionTime.xyz /= ScreenPositionTime.w;
 			float2 ScreenUVTime = ScreenPositionTime.xy * float2(0.5, -0.5) + 0.5;
-			int2 GatherPointCentroid = int2(ScreenUVTime * GatherPointsResolution + View.ViewRectMin);
+			int2 GatherPointCentroid = int2(ScreenUVTime) * GatherPointsResolution + int2(View.ViewRectMin.xy);
 #endif
 
 			for (uint LocalSampleIndex0 = 0; LocalSampleIndex0 < GatherFilterWidth2; ++LocalSampleIndex0)
diff --git a/Engine/Shaders/Private/RayTracing/RayTracingGatherPoints.ush b/Engine/Shaders/Private/RayTracing/RayTracingGatherPoints.ush
--- a/Engine/Shaders/Private/RayTracing/RayTracingGatherPoints.ush
+++ b/Engine/Shaders/Private/RayTracing/RayTracingGatherPoints.ush
@@ -85,15 +85,18 @@ void CreateRejectionMask(
 	out bool bShouldReject[MAXIMUM_GATHER_POINTS_PER_PIXEL]
 )
 {
-	for (int Index = 0; Index < SampleCount; ++Index)
 	{
-		FGatherSample GatherSample = ReadGatherSample(GatherBuffer, BaseIndex, SampleIndex + Index, ElementCount);
-		float3 GPDistance2 = RejectionCriteria.CreationPoint - GatherSample.CreationPoint;
-		bShouldReject[Index] = dot(GPDistance2, GPDistance2) > RejectionCriteria.DistanceThreshold2;
+		for (int Index = 0; Index < SampleCount; ++Index)
+		{
+			FGatherSample GatherSample = ReadGatherSample(GatherBuffer, BaseIndex, SampleIndex + Index, ElementCount);
+			float3 GPDistance2 = RejectionCriteria.CreationPoint - GatherSample.CreationPoint;
+			bShouldReject[Index] = dot(GPDistance2, GPDistance2) > RejectionCriteria.DistanceThreshold2;
+		}
 	}
-
-	for (int Index = SampleCount; Index < MAXIMUM_GATHER_POINTS_PER_PIXEL; ++Index)
 	{
-		bShouldReject[Index] = true;
+		for (int Index = SampleCount; Index < MAXIMUM_GATHER_POINTS_PER_PIXEL; ++Index)
+		{
+			bShouldReject[Index] = true;
+		}
 	}
 }
