Ensure to clear all resource even when a hair group has invalid voxel allocation.

diff --git a/Engine/Shaders/Private/HairStrands/HairStrandsVoxelPageAllocation.usf b/Engine/Shaders/Private/HairStrands/HairStrandsVoxelPageAllocation.usf
--- a/Engine/Shaders/Private/HairStrands/HairStrandsVoxelPageAllocation.usf
+++ b/Engine/Shaders/Private/HairStrands/HairStrandsVoxelPageAllocation.usf
@@ -25,6 +25,7 @@ RWBuffer<uint4>								OutPageIndexResolutionAndOffsetBuffer;
 RWBuffer<uint>								OutPageIndexAllocationIndirectBufferArgs;
 RWStructuredBuffer<FVoxelizationViewInfo>	OutVoxelizationViewInfoBuffer;
 
+#define INVALID_OFFSET 0xFFFFFFFF
 #define MAX_INSTANCE_GROUP_COUNT 32
 groupshared uint PageIndexOffsets[MAX_INSTANCE_GROUP_COUNT];
 
@@ -68,22 +69,42 @@ void AllocatePageIndex(uint2 DispatchThreadId : SV_DispatchThreadID)
 		// Postfix sum to instance group are always ordered by index
 		if (DispatchThreadId.x == 0)
 		{
+			bool bValidAllocation = true;
 			uint PageIndexOffset = 0;
 			for (uint LocalMacroGroupId = 0; LocalMacroGroupId < MacroGroupCount; ++LocalMacroGroupId)
 			{
 				const uint PageCount = PageIndexOffsets[LocalMacroGroupId];
-				PageIndexOffsets[LocalMacroGroupId] = PageIndexOffset;
+				bValidAllocation = bValidAllocation && (PageIndexOffset + PageCount <= TotalPageIndexCount);
+				PageIndexOffsets[LocalMacroGroupId] = bValidAllocation ? PageIndexOffset : INVALID_OFFSET;
 				PageIndexOffset += PageCount;
 			}
 		}
 		GroupMemoryBarrierWithGroupSync();
 
+		const uint PageIndexOffset = PageIndexOffsets[MacroGroupId];
+		bIsValid = bIsValid && (PageIndexOffset != INVALID_OFFSET);
 		if (bIsValid)
 		{
-			const uint PageIndexOffset = PageIndexOffsets[MacroGroupId];
 			OutPageIndexResolutionAndOffsetBuffer[MacroGroupId] = uint4(PageIndexResolution, PageIndexOffset);
 			WriteHairAABB(MacroGroupId, Bound, MacroGroupAABBBuffer);
 		}
+		else
+		{
+			// Clear all output if the allocation is not valid
+			const uint ArgsOffset = MacroGroupId * 3;
+			OutPageIndexResolutionAndOffsetBuffer[MacroGroupId]    = uint4(0, 0, 0, 0);
+			OutPageIndexAllocationIndirectBufferArgs[ArgsOffset]   = 0;
+			OutPageIndexAllocationIndirectBufferArgs[ArgsOffset+1] = 1;
+			OutPageIndexAllocationIndirectBufferArgs[ArgsOffset+2] = 1;
+
+			FVoxelizationViewInfo ViewInfo;
+			ViewInfo.WorldToClip = 0;
+			ViewInfo.ViewForward = 0;
+			ViewInfo.RasterResolution = 0;
+			ViewInfo.Pad0 = 0;
+			ViewInfo.Pad1 = 0;
+			OutVoxelizationViewInfoBuffer[MacroGroupId] = ViewInfo;
+		}
 	}
 
 	if (!bIsValid)
