Added support for thin lighting model to Strata (a.k.a. TwoSided ligthing model).
First version, will be improved.

DMFP input is used to control the back face lighting transmittance using MFP representation.
It is not possible to combine the Thin lighting and SSS mode so SSS post process is forced disabled when Thin lighting model is encoutered

diff --git a/Engine/Shaders/Private/Strata/Strata.ush b/Engine/Shaders/Private/Strata/Strata.ush
--- a/Engine/Shaders/Private/Strata/Strata.ush
+++ b/Engine/Shaders/Private/Strata/Strata.ush
@@ -38,9 +38,7 @@
 // Behind the scene, we do all the simple volumetric lighting & transmitance computations considering a slab of 100 centimeters = 1 meter.
 // If the user specifies a different thickness to simulate varying appearance, we rescale the mean free path in order to not have to store thickness as part of the BSDF description.
 #define STRATA_SIMPLEVOLUME_THICKNESS_CM	100.0f
-#define STRATA_CENTIMETER_TO_METER			0.01f
-#define STRATA_METER_TO_CENTIMETER			100.0f
-#define STRATA_SIMPLEVOLUME_THICKNESS_M		(STRATA_SIMPLEVOLUME_THICKNESS_CM * STRATA_CENTIMETER_TO_METER)
+#define STRATA_SIMPLEVOLUME_THICKNESS_M		(STRATA_SIMPLEVOLUME_THICKNESS_CM * CENTIMETER_TO_METER)
 
 ///////////////////////////////////////////////////////////////////////////////
 // BSDFs representation
@@ -125,7 +123,8 @@ struct FStrataBSDF
 #define STATE_BIT_COUNT_SIMPLEVOLUME			1
 #define STATE_BIT_COUNT_DMFPPLUGGED				1
 #define STATE_BIT_COUNT_HASFUZZ					1
-#define STATE_BIT_COUNT___UNUSED___				6
+#define STATE_BIT_COUNT_ISTHIN					1
+#define STATE_BIT_COUNT___UNUSED___				5
 #define STATE_BIT_COUNT_WEIGHT10F				10
 
 // Bit offset of each component store in BSDF state 
@@ -142,6 +141,7 @@ struct FStrataBSDF
 #define STATE_BIT_OFFSET_SIMPLEVOLUME	(STATE_BIT_OFFSET_THINFILM			+ STATE_BIT_COUNT_THINFILM)
 #define STATE_BIT_OFFSET_DMFPPLUGGED	(STATE_BIT_OFFSET_SIMPLEVOLUME		+ STATE_BIT_COUNT_SIMPLEVOLUME)
 #define STATE_BIT_OFFSET_HASFUZZ		(STATE_BIT_OFFSET_DMFPPLUGGED		+ STATE_BIT_COUNT_DMFPPLUGGED)
+#define STATE_BIT_OFFSET_ISTHIN			(STATE_BIT_OFFSET_HASFUZZ			+ STATE_BIT_COUNT_HASFUZZ)
 #define STATE_BIT_OFFSET_WEIGHT10F		22
 
 // Accessors for the BSDF type
@@ -163,8 +163,8 @@ struct FStrataBSDF
 #define BSDF_SETISTOPLAYER(X, IsTop)	WRITE_BITS(X.State, STATE_BIT_COUNT_TOPLAYER, STATE_BIT_OFFSET_TOPLAYER, IsTop)
 
 // Indicates if the BSDF has scattering component
-#define BSDF_GETHASSSS(X)			READ_BITS(X.State,  STATE_BIT_COUNT_HASSSS, STATE_BIT_OFFSET_HASSSS)
-#define BSDF_SETHASSSS(X, Scatt)	WRITE_BITS(X.State, STATE_BIT_COUNT_HASSSS, STATE_BIT_OFFSET_HASSSS, Scatt)
+#define BSDF_GETHASSSS(X)				READ_BITS(X.State,  STATE_BIT_COUNT_HASSSS, STATE_BIT_OFFSET_HASSSS)
+#define BSDF_SETHASSSS(X, Scatt)		WRITE_BITS(X.State, STATE_BIT_COUNT_HASSSS, STATE_BIT_OFFSET_HASSSS, Scatt)
 
 // Indicates if the BSDF has a SSS profile
 #define BSDF_GETHASSSSPROFILE(X)		READ_BITS(X.State,  STATE_BIT_COUNT_HASSSSPROFILE, STATE_BIT_OFFSET_HASSSSPROFILE)
@@ -202,6 +202,10 @@ struct FStrataBSDF
 #define BSDF_GETHASFUZZ(X)				READ_BITS(X.State,  STATE_BIT_COUNT_HASFUZZ, STATE_BIT_OFFSET_HASFUZZ)
 #define BSDF_SETHASFUZZ(X, Value)		WRITE_BITS(X.State, STATE_BIT_COUNT_HASFUZZ, STATE_BIT_OFFSET_HASFUZZ, Value)
 
+// Indicates if the BSDF is part of the top layer
+#define BSDF_GETISTHIN(X)				READ_BITS(X.State,  STATE_BIT_COUNT_ISTHIN, STATE_BIT_OFFSET_ISTHIN)
+#define BSDF_SETISTHIN(X, Value)		WRITE_BITS(X.State, STATE_BIT_COUNT_ISTHIN, STATE_BIT_OFFSET_ISTHIN, Value)
+
 #if STRATA_INLINE_SHADING
 #define BSDF_SETEMISSIVE(X, V)			X.Emissive = V
 #define BSDF_GETEMISSIVE(X)				X.Emissive
@@ -221,7 +225,8 @@ struct FStrataBSDF
 #define SLAB_COMPLEX_PATH_MASK (TO_OFFSET_BITMASK(STATE_BIT_COUNT_ANISOTROPY, STATE_BIT_OFFSET_ANISOTROPY)		|  TO_OFFSET_BITMASK(STATE_BIT_COUNT_HASSSS, STATE_BIT_OFFSET_HASSSS) | \
 								TO_OFFSET_BITMASK(STATE_BIT_COUNT_HASSSSPROFILE, STATE_BIT_OFFSET_HASSSSPROFILE)|  TO_OFFSET_BITMASK(STATE_BIT_COUNT_HAZINESS, STATE_BIT_OFFSET_HAZINESS) | \
 								TO_OFFSET_BITMASK(STATE_BIT_COUNT_THINFILM, STATE_BIT_OFFSET_THINFILM)			|  TO_OFFSET_BITMASK(STATE_BIT_COUNT_EDGECOLOR, STATE_BIT_OFFSET_EDGECOLOR) | \
-								TO_OFFSET_BITMASK(STATE_BIT_COUNT_SIMPLEVOLUME, STATE_BIT_OFFSET_SIMPLEVOLUME)	|  TO_OFFSET_BITMASK(STATE_BIT_COUNT_HASFUZZ, STATE_BIT_OFFSET_HASFUZZ) )
+								TO_OFFSET_BITMASK(STATE_BIT_COUNT_SIMPLEVOLUME, STATE_BIT_OFFSET_SIMPLEVOLUME)	|  TO_OFFSET_BITMASK(STATE_BIT_COUNT_HASFUZZ, STATE_BIT_OFFSET_HASFUZZ) | \
+								TO_OFFSET_BITMASK(STATE_BIT_COUNT_ISTHIN, STATE_BIT_OFFSET_ISTHIN) )
 
 
 #define SLAB_BASECOLOR(X)				X.VGPRs[0].xyz
@@ -474,9 +479,18 @@ FStrataData GetStrataSlabBSDF(
 
 	// STRATA_TODO apply SSS as a function of pixel footprint w.r.t. MFP
 	const bool bDMFPPlugged = any(SSSDMFP > 0);
-	const bool bHasSSSProfile = SSSProfileID  > 0.0f;	// When no profile is provide, we use Burley. STRATA_TODO we should clarify as burley is also an option of the profile.
+	const bool bIsThin = bDMFPPlugged && Thickness < STRATA_LAYER_ISTHIN_THICKNESS_THRESHOLD_CM;
+	const bool bHasSSSProfile = !bIsThin && SSSProfileID  > 0.0f;	// When no profile is provide, we use Burley. STRATA_TODO we should clarify as burley is also an option of the profile.
 	const bool bHasSSS = bDMFPPlugged || bHasSSSProfile;
 
+	// About the thin, a.k.a. two sided lighting, model:
+	// 	- Thin lighting model is used when thickness becomes less than STRATA_LAYER_ISTHIN_THICKNESS_THRESHOLD_CM
+	//  - It cannot work with a SSS Profile because the profile id is stored in the register used for the MFP (e.g. see SLAB_SSSPROFILEID).
+	// 	- Thin lighting model force disabled the SSSProfile.
+	//  - It can work with the Strata SSS because the MFP is shared in this case between both.
+	//  - We also only enabled thin lighting is the MFP input pin is plugged in.
+	//  - It will only be ran for the bottom layer only. Otherwise, the simple volume should be used. As such, Two-sided lighting can also be disabled later if part of the bottom layer.
+
 	StrataData.Layers[0].BSDFCount++;
 	BSDF_SETTYPE			(StrataData.Layers[0].BSDFs[0],		STRATA_BSDF_TYPE_SLAB);
 	BSDF_SETNORMALID		(StrataData.Layers[0].BSDFs[0],		SharedNormalIndex);
@@ -485,6 +499,7 @@ FStrataData GetStrataSlabBSDF(
 	BSDF_SETISTOPLAYER		(StrataData.Layers[0].BSDFs[0],		0);
 	BSDF_SETHASSSS			(StrataData.Layers[0].BSDFs[0],		bHasSSS ? 1 : 0);
 	BSDF_SETHASSSSPROFILE	(StrataData.Layers[0].BSDFs[0],		bHasSSSProfile ? 1 : 0);
+	BSDF_SETISTHIN			(StrataData.Layers[0].BSDFs[0],		bDMFPPlugged && Thickness < STRATA_LAYER_ISTHIN_THICKNESS_THRESHOLD_CM);
 	BSDF_SETHASDMFP			(StrataData.Layers[0].BSDFs[0],		bDMFPPlugged ? 1 : 0);
 	BSDF_SETHASHAZINESS		(StrataData.Layers[0].BSDFs[0],		bHasHaziness);
 	BSDF_SETHASEDGECOLOR	(StrataData.Layers[0].BSDFs[0],		bHasEdgeColor);
@@ -957,7 +972,7 @@ void StrataSanitizeBSDF(inout FStrataBSDF BSDF)
 
 FParticipatingMedia StrataSlabCreateParticipatingMedia(float3 DiffuseColor, float3 MeanFreePathCentimeters)
 {
-	const float3 MeanFreePathMeters = MeanFreePathCentimeters * STRATA_CENTIMETER_TO_METER;
+	const float3 MeanFreePathMeters = MeanFreePathCentimeters * CENTIMETER_TO_METER;
 	return CreateMediumFromBaseColorMFP(DiffuseColor, MeanFreePathMeters);
 }
 
@@ -996,7 +1011,7 @@ FStrataSlabSingleScattering StrataSlabSingleScattering(FParticipatingMedia PM, f
 	return Result;
 }
 
-void EnableSlabBSDFSimpleVolumetric(inout FStrataBSDF BSDF)
+void EnableSlabMFPComputationSpace(inout FStrataBSDF BSDF)
 {
 #if STRATA_INLINE_SHADING
 	// Now we know that this is a simple volume so we can freely set the mean free path
@@ -1006,6 +1021,11 @@ void EnableSlabBSDFSimpleVolumetric(inout FStrataBSDF BSDF)
 	// Rescale the BSDF mean free path so that we can drop the thickness parameter and run volumetric material math on a unit slab of medium
 	// For instance, if thickness is larger => we need to reduce the MFP for the medium to appear visually thicker.
 	SLAB_SSSDMFP(BSDF) = clamp(SLAB_SSSDMFP(BSDF) * (STRATA_SIMPLEVOLUME_THICKNESS_CM / BSDF_GETTHICKNESSCM(BSDF)), 0.0f, Max111110BitsFloat3);
+}
+
+void EnableSlabBSDFSimpleVolumetric(inout FStrataBSDF BSDF)
+{
+	EnableSlabMFPComputationSpace(BSDF);
 
 	// Mark the BSDF has not having any SSS through post process.
 	BSDF_SETHASSSS(BSDF, 0);
@@ -1016,6 +1036,11 @@ void EnableSlabBSDFSimpleVolumetric(inout FStrataBSDF BSDF)
 	BSDF_SETISSIMPLEVOLUME(BSDF, 1);
 }
 
+void EnableSlabBSDFIsThin(inout FStrataBSDF BSDF)
+{
+	EnableSlabMFPComputationSpace(BSDF);
+}
+
 
 ///////////////////////////////////////////////////////////////////////////////
 // BSDFs packing/unpacking
@@ -1343,6 +1368,9 @@ void PackStrataOut(
 			{
 				EmissiveColor += BSDF_GETEMISSIVE(BSDF) * Throughput;
 
+				// Two sided lighting can only be evalauted on the bottom layer, together with option SSS. Otherwise, the simple volume will be used.
+				const bool bIsThin = bBottomLayer && BSDF_GETISTHIN(Strata.Layers[l].BSDFs[i]);
+
 				// We want simple volumetric only if the DMFP input is pluged in, otherwise we avoid storing the DFMP
 				const bool bIsSimpleVolume = !bBottomLayer && BSDF_GETHASDMFP(Strata.Layers[l].BSDFs[i]);
 
@@ -1390,6 +1418,16 @@ void PackStrataOut(
 					bStrataSubsurfaceEnable = false;
 				}
 
+				if (bIsThin)
+				{
+					EnableSlabBSDFIsThin(Strata.Layers[l].BSDFs[i]);
+				}
+				else
+				{
+					// If not selected or not possible to have two sided lighting enabled (not bottom layer or SSSProfile is enabled), let's make sure we disable the feature.
+					BSDF_SETISTHIN(Strata.Layers[l].BSDFs[i], 0);
+				}
+
 				// STRATA_TODO	Should we have the Fuzz influence the next layter visibility accoding to fuzz amount? Should it also occlude diffuse?
 				//				If this is done, we should also update ComputePostProcessBSDFData.
 
@@ -1491,14 +1529,17 @@ void PackStrataOut(
 				}
 
 				const bool bIsSimpleVolume = BSDF_GETISSIMPLEVOLUME(BSDF);
-				if (BSDF_GETHASSSS(BSDF) || bIsSimpleVolume)
+				const bool bIsThin = BSDF_GETISTHIN(BSDF);
+				if (BSDF_GETHASSSS(BSDF) || bIsSimpleVolume || bIsThin)
 				{
 					if (BSDF_GETHASSSSPROFILE(BSDF) && !bIsSimpleVolume)
 					{
+						// Simple volume overrides SSS profile because it takes over for when not at the bottom of the BSDF layer, or during forward rendering.
 						STRATA_STORE_UINT1(PackRGBA8(float4(SLAB_SSSPROFILEID(BSDF), SLAB_SSSPROFILERADIUSSCALE(BSDF), 0, 0)));
 					}
 					else
 					{
+						// Path used for bottom most layer SSS, simple volume and two sided lighting.
 						STRATA_STORE_UINT1(PackR11G11B10F(SLAB_SSSDMFP(BSDF)));
 					}
 				}
@@ -1672,9 +1713,10 @@ FStrataBSDF UnpackStrataBSDFIn(ByteAddressBuffer StrataBuffer, inout FStrataAddr
 		}
 
 		const bool bIsSimpleVolume = BSDF_GETISSIMPLEVOLUME(OutBSDF);
-		if (BSDF_GETHASSSS(OutBSDF) || bIsSimpleVolume)
+		const bool bIsThin = BSDF_GETISTHIN(OutBSDF);
+		if (BSDF_GETHASSSS(OutBSDF) || bIsSimpleVolume || bIsThin)
 		{
-			if (BSDF_GETHASSSSPROFILE(OutBSDF) && !bIsSimpleVolume)
+			if (BSDF_GETHASSSSPROFILE(OutBSDF) && !bIsSimpleVolume && !bIsThin)
 			{
 				float4 Data = UnpackRGBA8(StrataLoadUint1(StrataBuffer, StrataAddressing));
 				SLAB_SSSPROFILEID(OutBSDF) = Data.r;
@@ -1688,7 +1730,7 @@ FStrataBSDF UnpackStrataBSDFIn(ByteAddressBuffer StrataBuffer, inout FStrataAddr
 			// When loading Strata BSDF which have SSS enable, BaseColor and (optionally) Specular value can be overriden 
 			// based on the SSS method used (screen space SSS / checherboarded screen space SSS / ...)
 			#if STRATA_SSS_MATERIAL_OVERRIDE
-			if (!bIsSimpleVolume)
+			if (!bIsSimpleVolume && !bIsThin)
 			{
 				const bool bChecker = CheckerFromPixelPos(StrataAddressing.PixelCoords);
 				AdjustBaseColorAndSpecularColorForSubsurfaceProfileLighting(SLAB_BASECOLOR(OutBSDF), SLAB_SPECULAR(OutBSDF), bChecker);
diff --git a/Engine/Shaders/Private/Strata/StrataDeferredLighting.ush b/Engine/Shaders/Private/Strata/StrataDeferredLighting.ush
--- a/Engine/Shaders/Private/Strata/StrataDeferredLighting.ush
+++ b/Engine/Shaders/Private/Strata/StrataDeferredLighting.ush
@@ -205,16 +205,18 @@ float4 StrataDeferredLighting(
 		{
 			FStrataEvaluateResult BSDFEvaluate = StrataEvaluateBSDF(BSDFContext, ShadowTerms);
 
-			const float3 DiffusePathWeight		= BSDF.Weight * BSDFEvaluate.DiffusePathValue;
-			const float3 SpecularPathWeight		= BSDF.Weight * BSDFEvaluate.SpecularPathValue     * BSDFEvaluate.SpecularPathProbability;
-			const float3 SpecularHazePathWeight = BSDF.Weight * BSDFEvaluate.SpecularHazePathValue * BSDFEvaluate.SpecularHazePathProbability;
-			const float3 TransmissionPathWeight	= BSDF.Weight * BSDFEvaluate.TransmissionPathValue;
+			const float3 DiffusePathWeight			= BSDF.Weight * BSDFEvaluate.DiffusePathValue;
+			const float3 SpecularPathWeight			= BSDF.Weight * BSDFEvaluate.SpecularPathValue     * BSDFEvaluate.SpecularPathProbability;
+			const float3 SpecularHazePathWeight		= BSDF.Weight * BSDFEvaluate.SpecularHazePathValue * BSDFEvaluate.SpecularHazePathProbability;
+			const float3 TransmissionPathWeight		= BSDF.Weight * BSDFEvaluate.TransmissionPathValue;
+			const float3 ThinTransmissionPathWeight	= BSDF.Weight * BSDFEvaluate.ThinTransmissionPathValueWithProjectedSolidAngle;
 
 			const bool bNeedsSeparateSubsurfaceLightAccumulation = BSDFEvaluate.bSubsurface;
 
-			LightAccumulator_AddSplit(Out, DiffusePathWeight,		SpecularPathWeight,		DiffusePathWeight,		IntegrationContext.IncidentIrradianceNoShadow * ShadowTerms.SurfaceShadow,					bNeedsSeparateSubsurfaceLightAccumulation);
-			LightAccumulator_AddSplit(Out, 0.0,						SpecularHazePathWeight, 0.0,					IntegrationContext.IncidentIrradianceNoShadowHaze * ShadowTerms.SurfaceShadow,				bNeedsSeparateSubsurfaceLightAccumulation);
-			LightAccumulator_AddSplit(Out, TransmissionPathWeight,	0.0f,					TransmissionPathWeight, IntegrationContext.IncidentTransmittedRadianceNoShadow * ShadowTerms.TransmissionShadow,	bNeedsSeparateSubsurfaceLightAccumulation);
+			LightAccumulator_AddSplit(Out, DiffusePathWeight,			SpecularPathWeight,					DiffusePathWeight,			IntegrationContext.IncidentIrradianceNoShadow * ShadowTerms.SurfaceShadow,					bNeedsSeparateSubsurfaceLightAccumulation);
+			LightAccumulator_AddSplit(Out, 0.0,							SpecularHazePathWeight,				0.0,						IntegrationContext.IncidentIrradianceNoShadowHaze * ShadowTerms.SurfaceShadow,				bNeedsSeparateSubsurfaceLightAccumulation);
+			LightAccumulator_AddSplit(Out, TransmissionPathWeight,		0.0f,								TransmissionPathWeight,		IntegrationContext.IncidentTransmittedRadianceNoShadow * ShadowTerms.TransmissionShadow,	bNeedsSeparateSubsurfaceLightAccumulation);
+			LightAccumulator_AddSplit(Out, ThinTransmissionPathWeight,	0.0f,								ThinTransmissionPathWeight,	IntegrationContext.IncidentRadianceNoShadow * ShadowTerms.SurfaceShadow,					bNeedsSeparateSubsurfaceLightAccumulation);
 		}
 		else
 		{
@@ -223,7 +225,7 @@ float4 StrataDeferredLighting(
 			const float3 SpecularRadiance = 1;
 			const bool bNeedsSeparateSubsurfaceLightAccumulation = false;
 
-			LightAccumulator_AddSplit(Out, DiffuseRadiance,			SpecularRadiance,	DiffuseRadiance,		IntegrationContext.OutgoingRadianceNoShadow * ShadowTerms.SurfaceShadow,		bNeedsSeparateSubsurfaceLightAccumulation);
+			LightAccumulator_AddSplit(Out, DiffuseRadiance,				SpecularRadiance,					DiffuseRadiance,			IntegrationContext.OutgoingRadianceNoShadow * ShadowTerms.SurfaceShadow,					bNeedsSeparateSubsurfaceLightAccumulation);
 		}
 	}
 
diff --git a/Engine/Shaders/Private/Strata/StrataEnvironmentLighting.ush b/Engine/Shaders/Private/Strata/StrataEnvironmentLighting.ush
--- a/Engine/Shaders/Private/Strata/StrataEnvironmentLighting.ush
+++ b/Engine/Shaders/Private/Strata/StrataEnvironmentLighting.ush
@@ -71,16 +71,32 @@ float4 StrataEnvironmentLighting(
 
 		// Diffuse component
 		float3 DiffuseSkyLighting = 0;
-		if (any(StrataEnvLight.DiffuseWeight > 0.0f))
+		const bool bProcessFrontFaceDiffuse = any(StrataEnvLight.DiffuseWeight > 0.0f);
+		const bool bProcessBackFaceDiffuse = any(StrataEnvLight.DiffuseBackFaceWeight > 0.0f);
+		if (bProcessFrontFaceDiffuse || bProcessBackFaceDiffuse)
 		{
 			// Compute the common sky visibility factors
 			FSkyLightVisibilityData SkyVisData = GetSkyLightVisibilityData(StrataBSDFContext.N, StrataBSDFContext.N, StrataPixelHeader.MaterialAO, AmbientOcclusion, StrataEnvLight.DiffuseNormal);
 
-			// Finally sample the sky diffuse contribution (spherical harmonic, Lambert BRDF)
-			float3 DiffuseLookup = GetSkySHDiffuse(StrataEnvLight.DiffuseNormal) * View.SkyLightColor.rgb * IndirectDiffuseOcclusion;
-			// And accumulate
-			// Note: Use diffuse directional albedo (i.e., DiffuseWeight) as first order approximation for env. integration (STRATA_TODO instead compute SH coefficients for Chan)
-			DiffuseSkyLighting = BSDFThroughput * (SkyVisData.SkyDiffuseLookUpMul * DiffuseLookup + SkyVisData.SkyDiffuseLookUpAdd) * StrataEnvLight.DiffuseWeight;
+			if (bProcessFrontFaceDiffuse)
+			{
+				// Finally sample the sky diffuse contribution (spherical harmonic, Lambert BRDF)
+				float3 DiffuseLookup = GetSkySHDiffuse(StrataEnvLight.DiffuseNormal) * View.SkyLightColor.rgb * IndirectDiffuseOcclusion;
+				// And accumulate
+				// Note: Use diffuse directional albedo (i.e., DiffuseWeight) as first order approximation for env. integration (STRATA_TODO instead compute SH coefficients for Chan)
+				DiffuseSkyLighting = BSDFThroughput * (SkyVisData.SkyDiffuseLookUpMul * DiffuseLookup + SkyVisData.SkyDiffuseLookUpAdd) * StrataEnvLight.DiffuseWeight;
+			}
+			if (bProcessBackFaceDiffuse)
+			{
+				// We do not evaluate back face sky light visibility data because all the data we have is for the front face only.  This could be evaluated at some cost.
+				// However, we do apply SkyVisData.SkyDiffuseLookUpMul for scaling consistency.
+
+				// Finally sample the sky diffuse contribution (spherical harmonic, Lambert BRDF) along the opposite normal direction
+				float3 DiffuseLookup = GetSkySHDiffuse(-StrataEnvLight.DiffuseNormal) * View.SkyLightColor.rgb * IndirectDiffuseOcclusion;
+				// And accumulate
+				// Note: Use diffuse directional albedo (i.e., DiffuseWeight) as first order approximation for env. integration (STRATA_TODO instead compute SH coefficients for Chan)
+				DiffuseSkyLighting += BSDFThroughput * (SkyVisData.SkyDiffuseLookUpMul * DiffuseLookup) * StrataEnvLight.DiffuseBackFaceWeight;
+			}
 		}
 
 		// Specular component
diff --git a/Engine/Shaders/Private/Strata/StrataEvaluation.ush b/Engine/Shaders/Private/Strata/StrataEvaluation.ush
--- a/Engine/Shaders/Private/Strata/StrataEvaluation.ush
+++ b/Engine/Shaders/Private/Strata/StrataEvaluation.ush
@@ -33,6 +33,9 @@
 #define STRATA_SSS_TRANSMISSION 0
 #endif
 
+// STRATA_TODO Enable the new model and get rid of the old one.
+#define THINLIGHTING_NEWMODEL 0
+
 // STRATA_TODO put in a common file
 // Point lobe in off-specular peak direction
 float3 StrataGetOffSpecularPeakReflectionDir(float3 Normal, float3 ReflectionVector, float Roughness)
@@ -167,6 +170,8 @@ struct FStrataEvaluateResult
 	float3 SpecularHazePathValue;
 	float3 TransmissionPathValue;
 
+	float3 ThinTransmissionPathValueWithProjectedSolidAngle;
+
 	// The following probabilities are part of the pdf to work correctly with a monter carlo based integrator and as expected are not applied on the xxxPathValue.
 	// So PathProbability should be applied on PathValue in the rasteriser to recover the correct balance.
 	float SpecularPathProbability;
@@ -176,6 +181,7 @@ struct FStrataEvaluateResult
 	float  SpecularPDF;
 	float  SpecularHazePDF;
 	float  TransmissionPDF;
+	float  ThinTransmissionPathPDF;
 
 	float3 Throughput;					// Throughput to the next layer
 	bool   bSubsurface;					// True if we need to separate the subsurface light contribution for the screen space diffusion process.
@@ -237,7 +243,7 @@ FStrataEvaluateResult StrataEvaluateBSDF(FStrataBSDFContext BSDFContext, FShadow
 			float a2 = Pow4(SafeRoughness);
 
 			Sample.DiffusePathValue	= Diffuse_Chan(DiffuseColor, a2, BSDFContext.SatNoV, BSDFContext.SatNoL, BSDFContext.Context.VoH, BSDFContext.Context.NoH);
-			Sample.DiffusePDF		= BSDFContext.SatNoV / PI;
+			Sample.DiffusePDF		= BSDFContext.SatNoL / PI;
 			Sample.bSubsurface		= bHasSubsurface;
 
 			const float3 F0  = ComputeF0(SLAB_SPECULAR(BSDFContext.BSDF), BaseColor, Metallic);
@@ -270,7 +276,6 @@ FStrataEvaluateResult StrataEvaluateBSDF(FStrataBSDFContext BSDFContext, FShadow
 					Vis = Vis_SmithJointAniso(Alpha0, Alpha1, BSDFContext.Context.NoV, BSDFContext.Context.NoL, BSDFContext.Context.XoV, BSDFContext.Context.XoL, BSDFContext.Context.YoV, BSDFContext.Context.YoL);
 					const float H_PDF = VisibleGGXPDF_aniso(BSDFContext.TangentV, BSDFContext.TangentH, float2(Alpha0, Alpha1));
 					PDF = RayPDFToReflectionRayPDF(BSDFContext.Context.VoH, H_PDF);
-
 				}
 
 				if (bHaziness)
@@ -310,7 +315,6 @@ FStrataEvaluateResult StrataEvaluateBSDF(FStrataBSDFContext BSDFContext, FShadow
 					HazeVis = Vis_SmithJoint(a2Haze, BSDFContext.Context.NoV, BSDFContext.SatNoL);
 					const float H_PDF = VisibleGGXPDF(BSDFContext.TangentV, BSDFContext.TangentH, a2Haze);
 					HazePDF = RayPDFToReflectionRayPDF(BSDFContext.Context.VoH, H_PDF);
-
 				}
 				#endif
 			}
@@ -426,6 +430,48 @@ FStrataEvaluateResult StrataEvaluateBSDF(FStrataBSDFContext BSDFContext, FShadow
 			}
 			#endif
 
+			#if STRATA_FASTPATH==0
+			BRANCH
+			if (BSDF_GETISTHIN(BSDFContext.BSDF))
+			{
+				FParticipatingMedia PM = StrataSlabCreateParticipatingMedia(DiffuseColor, SLAB_SSSDMFP(BSDFContext.BSDF));
+
+				// we consider single scattering and perpendicular to the surface light transmittance, no diffusion for now.
+				const float TransmittanceNoV = 1.0;	
+				const float3 SlabTransmittance = IsotropicMediumSlabTransmittance(PM, STRATA_SIMPLEVOLUME_THICKNESS_M, TransmittanceNoV);
+
+			#if THINLIGHTING_NEWMODEL
+				//--- New version
+
+				// STRATA_TODO Fresnel throughput, Use a LUT, this is currently in progress and will be coming soon.
+				// In the meantime, we use a simple approach to feed Fresnel function and compute an approximation of the general transmition/througput.
+				// We also ignored ThinFilm interaction as another simplification for perf.
+				const float3 ApproxFresnelForTransmittance = F_Schlick(F0, F90, BSDFContext.SatNoV);
+				const float3 SpecularThroughput = 1.0f - ApproxFresnelForTransmittance; // This is already taken into account by the Chan model. But for two sided lighting we need to apply it.
+
+				// We recover the transmitted luminance from as the back face Lambert affected by the transmittance.
+				// The MFP is specified on the slab node and then transmittance recoved from it.
+				// We assume the back face diffuse albedo is white=1 and no diffusion happens at this stage.
+				const float3 BackFaceWhiteDiffuseColor = 1.0f;
+				const float3 BackFaceDiffuse = saturate(-BSDFContext.Context.NoL) * Diffuse_Lambert(BackFaceWhiteDiffuseColor);
+			#else
+				//--- Legacy version
+
+				const float3 SpecularThroughput = 1.0f;
+
+				// Scatter distribution following GGX mixed with a wrapped diffuse Lambert.
+				float BackFaceDiffuse = D_GGX(0.6 * 0.6, saturate(-BSDFContext.Context.VoL));
+				const float Wrap = 0.5;
+				BackFaceDiffuse *= saturate((-BSDFContext.Context.NoL + Wrap) / Square(1 + Wrap));
+			#endif
+
+				Sample.ThinTransmissionPathValueWithProjectedSolidAngle = SlabTransmittance *  SpecularThroughput * BackFaceDiffuse;
+				Sample.ThinTransmissionPathPDF = saturate(-BSDFContext.Context.NoL) / PI;
+
+				// At this stage, Sample.bSubsurface should have been setup correctly already if SSS is enabled with two-sided lighting.
+			}
+			#endif
+
 			#if STRATA_FASTPATH==0
 			BRANCH
 			if (BSDF_GETISSIMPLEVOLUME(BSDFContext.BSDF))
@@ -545,6 +591,7 @@ FStrataImportanceSampleResult StrataImportanceSampleBSDF(FStrataBSDFContext BSDF
 
 			// STRATA_TODO take into account Haziness
 			// STRATA_TODO take into account Fuzz: Per "Production Friendly Microfacet Sheen BRDF", uniform hemispherical sampling gives good result as the roughness is usually high and avoid issue at grazing angle.
+			// STRATA_TODO take into account Two sided material
 
 			const float HPDF = TangentH.w;
 			const float3 H = mul(TangentH.xyz, BSDFContext.TangentBasis);
@@ -593,6 +640,7 @@ struct FStrataEnvLightResult
 {
 	float3 DiffuseNormal;
 	float3 DiffuseWeight;
+	float3 DiffuseBackFaceWeight;
 
 	float3 SpecularDirection;
 	float3 SpecularWeight;
@@ -625,12 +673,12 @@ FStrataEnvLightResult StrataEvaluateForEnvLight(FStrataBSDFContext BSDFContext,
 
 			float SafeRoughnessX = MakeRoughnessSafe(SLAB_ROUGHNESSX(BSDFContext.BSDF));
 			float3 F0 = ComputeF0(SLAB_SPECULAR(BSDFContext.BSDF), BaseColor, Metallic);
+			const float3 F90 = ComputeF90(F0, SLAB_EDGECOLOR(BSDFContext.BSDF), Metallic);
 			float3 ThinFilmFresnelWeight = 1;
 
 			BRANCH
 			if (bEnableSpecular)
 			{
-				const float3 F90 = ComputeF90(F0, SLAB_EDGECOLOR(BSDFContext.BSDF), Metallic);
 
 				#if STRATA_FASTPATH==0
 				const bool bHasThinFilm = BSDF_GETHASTHINFILM(BSDFContext.BSDF);
@@ -753,6 +801,34 @@ FStrataEnvLightResult StrataEvaluateForEnvLight(FStrataBSDFContext BSDFContext,
 				StrataEnvLightResult.bSubsurface = false;
 			}
 			#endif
+			
+
+			#if STRATA_FASTPATH==0
+			BRANCH
+			if (BSDF_GETISTHIN(BSDFContext.BSDF))
+			{
+				FParticipatingMedia PM = StrataSlabCreateParticipatingMedia(DiffuseColor, SLAB_SSSDMFP(BSDFContext.BSDF));
+
+				// we consider single scattering and perpendicular to the surface light transmittance, no diffusion for now.
+				const float TransmittanceNoV = 1.0;	
+				const float3 SlabTransmittance = IsotropicMediumSlabTransmittance(PM, STRATA_SIMPLEVOLUME_THICKNESS_M, TransmittanceNoV);
+
+				// STRATA_TODO Fresnel throughput, Use a LUT, this is currently in progress and will be coming soon.
+				// In the meantime, we use a simple approach to feed Fresnel function and compute an approximation of the general transmition/througput.
+				// We also ignored ThinFilm interaction as another simplification for perf.
+				const float3 ApproxFresnelForTransmittance = F_Schlick(F0, F90, BSDFContext.SatNoV);
+				const float3 SpecularThroughput = 1.0f - ApproxFresnelForTransmittance; // This is already taken into account by the Chan model. But for two sided lighting we need to apply it.
+
+				// We recover the transmitted luminance from as the back face Lambert affected by the transmittance.
+				// The MFP is specified on the slab node and then transmittance recoved from it.
+				// We assume the back face diffuse albedo is white=1 and no diffusion happens at this stage.
+				const float3 BackFaceWhiteDiffuseColor = float3(1.0f, 1.0f, 1.0f);
+
+				StrataEnvLightResult.DiffuseBackFaceWeight = SlabTransmittance * BackFaceWhiteDiffuseColor * SpecularThroughput * ThinFilmFresnelWeight;
+
+				// At this stage, StrataEnvLightResult.bSubsurface should have been setup correctly already if SSS is enabled with two-sided lighting.
+			}
+			#endif
 
 			#if STRATA_FASTPATH==0
 			BRANCH
diff --git a/Engine/Shaders/Private/Strata/StrataMaterialClassification.usf b/Engine/Shaders/Private/Strata/StrataMaterialClassification.usf
--- a/Engine/Shaders/Private/Strata/StrataMaterialClassification.usf
+++ b/Engine/Shaders/Private/Strata/StrataMaterialClassification.usf
@@ -67,7 +67,8 @@ void ComputePostProcessBSDFData(
 		}
 
 		// Select the latest encoutered SSS BSDF
-		if (BSDF_GETHASSSS(BSDF) && BSDFWeightAvg > 0)
+		if (!BSDF_GETISTHIN(BSDF) &&					// STRATA_TODO A material cannot have subsurface if using Thin mode. This is because DMFP does not match MFP at all (right now, SSS would just smears thin lighting model details out)
+			BSDF_GETHASSSS(BSDF) && BSDFWeightAvg > 0)
 		{
 			Classification.bHasSubsurface = true;
 
diff --git a/Engine/Shaders/Private/Strata/StrataVisualize.usf b/Engine/Shaders/Private/Strata/StrataVisualize.usf
--- a/Engine/Shaders/Private/Strata/StrataVisualize.usf
+++ b/Engine/Shaders/Private/Strata/StrataVisualize.usf
@@ -193,6 +193,7 @@ void VisualizeMaterialPS(
 #define PRINTBSDFTOPLAY()	PRINTBSDFCHAR(_SPC_); PRINTBSDFCHAR(_MINUS_); PRINTBSDFCHAR(_T_); PRINTBSDFCHAR(_O_); PRINTBSDFCHAR(_P_); PRINTBSDFCHAR(_L_); PRINTBSDFCHAR(_A_); PRINTBSDFCHAR(_Y_); PRINTBSDFCHAR(_SPC_); PrintFloatNoFraction(PixelPos, OutColor.xyz, FontColor, TopLeft, float(BSDF_GETISTOPLAYER(BSDF)), 1); TopLeft.x+=8*1;
 #define PRINTBSDFHASSCAT()	PRINTBSDFCHAR(_SPC_); PRINTBSDFCHAR(_MINUS_); PRINTBSDFCHAR(_H_); PRINTBSDFCHAR(_A_); PRINTBSDFCHAR(_S_); PRINTBSDFCHAR(_S_); PRINTBSDFCHAR(_C_); PRINTBSDFCHAR(_A_); PRINTBSDFCHAR(_T_); PRINTBSDFCHAR(_SPC_); PrintFloatNoFraction(PixelPos, OutColor.xyz, FontColor, TopLeft, float(BSDF_GETHASSSS(BSDF)), 1); TopLeft.x+=8*1;
 #define PRINTBSDFSIMPSCAT()	PRINTBSDFCHAR(_SPC_); PRINTBSDFCHAR(_MINUS_); PRINTBSDFCHAR(_S_); PRINTBSDFCHAR(_I_); PRINTBSDFCHAR(_M_); PRINTBSDFCHAR(_P_); PRINTBSDFCHAR(_S_); PRINTBSDFCHAR(_C_); PRINTBSDFCHAR(_A_); PRINTBSDFCHAR(_T_); PRINTBSDFCHAR(_SPC_); PrintFloatNoFraction(PixelPos, OutColor.xyz, FontColor, TopLeft, float(BSDF_GETISSIMPLEVOLUME(BSDF)), 1); TopLeft.x+=8*1;
+#define PRINTBSDFISTHIN()	PRINTBSDFCHAR(_SPC_); PRINTBSDFCHAR(_MINUS_); PRINTBSDFCHAR(_I_); PRINTBSDFCHAR(_S_); PRINTBSDFCHAR(_T_); PRINTBSDFCHAR(_H_); PRINTBSDFCHAR(_I_); PRINTBSDFCHAR(_N_); PRINTBSDFCHAR(_SPC_); PrintFloatNoFraction(PixelPos, OutColor.xyz, FontColor, TopLeft, float(BSDF_GETISTHIN(BSDF)), 1); TopLeft.x+=8*1;
 #define PRINTFLOAT3(X) PrintFloat(PixelPos, OutColor.xyz, PARAMFontColor, TopLeft, X.x); TopLeft.x+=8*7; PRINTBSDFCHAR(_SPC_); PrintFloat(PixelPos, OutColor.xyz, PARAMFontColor, TopLeft, X.y); TopLeft.x+=8*7; PRINTBSDFCHAR(_SPC_); PrintFloat(PixelPos, OutColor.xyz, PARAMFontColor, TopLeft, X.z);
 #define PRINTFLOAT1(Value) PrintFloat(PixelPos, OutColor.xyz, PARAMFontColor, TopLeft, Value); TopLeft.x+=8*7;
 #define PRINTFLOATNOFRAC(Value, DigitCount) PrintFloatNoFraction(PixelPos, OutColor.xyz, PARAMFontColor, TopLeft, Value, DigitCount); TopLeft.x+=8*7;
@@ -318,6 +319,7 @@ void VisualizeMaterialPS(
 				PRINTBSDFTOPLAY();
 				PRINTBSDFHASSCAT();
 				PRINTBSDFSIMPSCAT();
+				PRINTBSDFISTHIN();
 				NEXTLINE;
 
 				HORIZONTALTAB; PRINTPARAMCHAR(_A_); PRINTPARAMCHAR(_L_); PRINTPARAMCHAR(_B_); PRINTPARAMCHAR(_E_); PRINTPARAMCHAR(_D_); PRINTPARAMCHAR(_O_); 
@@ -364,7 +366,7 @@ void VisualizeMaterialPS(
 				HORIZONTALTAB; PRINTFLOAT1(SLAB_HAZINESS(BSDF));
 				NEXTLINE;
 
-				if (BSDF_GETHASSSS(BSDF) || BSDF_GETISSIMPLEVOLUME(BSDF))
+				if (BSDF_GETHASSSS(BSDF) || BSDF_GETISSIMPLEVOLUME(BSDF) || BSDF_GETISTHIN(BSDF))
 				{
 					if (BSDF_GETHASSSSPROFILE(BSDF) && BSDF_GETISSIMPLEVOLUME(BSDF)==0)
 					{
@@ -423,7 +425,7 @@ void VisualizeMaterialPS(
 
 				NEXTLINE;
 
-				if (BSDF_GETHASSSS(BSDF) || BSDF_GETISSIMPLEVOLUME(BSDF))
+				if (BSDF_GETHASSSS(BSDF) || BSDF_GETISSIMPLEVOLUME(BSDF) || BSDF_GETISTHIN(BSDF))
 				{
 					if (BSDF_GETHASSSSPROFILE(BSDF) && BSDF_GETISSIMPLEVOLUME(BSDF) == 0)
 					{
diff --git a/Engine/Shaders/Shared/StrataDefinitions.h b/Engine/Shaders/Shared/StrataDefinitions.h
--- a/Engine/Shaders/Shared/StrataDefinitions.h
+++ b/Engine/Shaders/Shared/StrataDefinitions.h
@@ -33,4 +33,9 @@
 #define STRATA_DATA_TILE_SIZE_DIV_AS_SHIFT	3
 
 // The default thickness of a layer is considered to be 0.01 centimeter = 0.1 millimeter
-#define STRATA_LAYER_DEFAULT_THICKNESS_CM	0.01f
\ No newline at end of file
+#define STRATA_LAYER_DEFAULT_THICKNESS_CM	0.01f
+
+// The thin lighting model is used when the thickness becomes strictly lower than 0.1 millimeter
+#define STRATA_LAYER_ISTHIN_THICKNESS_THRESHOLD_CM 0.01f
+
+
diff --git a/Engine/Source/Runtime/Engine/Classes/Materials/MaterialExpressionStrata.h b/Engine/Source/Runtime/Engine/Classes/Materials/MaterialExpressionStrata.h
--- a/Engine/Source/Runtime/Engine/Classes/Materials/MaterialExpressionStrata.h
+++ b/Engine/Source/Runtime/Engine/Classes/Materials/MaterialExpressionStrata.h
@@ -516,6 +516,13 @@ class UMaterialExpressionStrataTransmittanceToMFP : public UMaterialExpressionSt
 	UPROPERTY()
 	FExpressionInput TransmittanceColor;
 
+	/**
+	* The desired thickness in centimers. This can be set lower than 0.1mm (= 0.01cm) to enable the Thin lighting model on the slab node for instance.
+	* Another use case exemple: the thickenss ouput can be modulated to have simple scattering/transmittance variation of the same material before it is plugged in a slab node.
+	*/
+	UPROPERTY()
+	FExpressionInput Thickness;
+
 	//~ Begin UMaterialExpression Interface
 #if WITH_EDITOR
 	virtual int32 Compile(class FMaterialCompiler* Compiler, int32 OutputIndex) override;
diff --git a/Engine/Source/Runtime/Engine/Private/Materials/HLSLMaterialTranslator.cpp b/Engine/Source/Runtime/Engine/Private/Materials/HLSLMaterialTranslator.cpp
--- a/Engine/Source/Runtime/Engine/Private/Materials/HLSLMaterialTranslator.cpp
+++ b/Engine/Source/Runtime/Engine/Private/Materials/HLSLMaterialTranslator.cpp
@@ -8993,27 +8993,28 @@ int32 FHLSLMaterialTranslator::StrataMultiply(int32 A, int32 Weight)
 	);
 }
 
-int32 FHLSLMaterialTranslator::StrataTransmittanceToMFP(int32 TransmittanceColor, int32 OutputIndex)
+int32 FHLSLMaterialTranslator::StrataTransmittanceToMFP(int32 TransmittanceColor, int32 DesiredThickness, int32 OutputIndex)
 {
 	if (OutputIndex == INDEX_NONE)
 	{
 		return INDEX_NONE;
 	}
 
+	const int32 DefaultThicknessCodechunk = AddInlinedCodeChunk(MCT_Float1, TEXT("%f"), STRATA_LAYER_DEFAULT_THICKNESS_CM);
 	switch (OutputIndex)
 	{
 	case 0:
 		return AddCodeChunk(MCT_Float3, 
 			// For the math to be valid, input to TransmittanceToMeanFreePath must be in meter.
 			// Then the output needs to be is converted to centimeters.
-			TEXT("(TransmittanceToMeanFreePath(%s, %f * STRATA_CENTIMETER_TO_METER) * STRATA_METER_TO_CENTIMETER)"), 
+			TEXT("(TransmittanceToMeanFreePath(%s, %s * CENTIMETER_TO_METER) * METER_TO_CENTIMETER)"), 
 			*GetParameterCode(TransmittanceColor),
-			STRATA_LAYER_DEFAULT_THICKNESS_CM);
+			*GetParameterCode(DesiredThickness == INDEX_NONE ? DefaultThicknessCodechunk : DesiredThickness));
 		break;
 	case 1:
 		// Thickness to be plugged into other nodes thickness input.
 		// This matches the Slab node default using STRATA_LAYER_DEFAULT_THICKNESS_CM
-		return AddInlinedCodeChunk(MCT_Float1, TEXT("%f"), STRATA_LAYER_DEFAULT_THICKNESS_CM);
+		return DesiredThickness == INDEX_NONE ? DefaultThicknessCodechunk : DesiredThickness;
 		break;
 	}
 	return INDEX_NONE;
diff --git a/Engine/Source/Runtime/Engine/Private/Materials/HLSLMaterialTranslator.h b/Engine/Source/Runtime/Engine/Private/Materials/HLSLMaterialTranslator.h
--- a/Engine/Source/Runtime/Engine/Private/Materials/HLSLMaterialTranslator.h
+++ b/Engine/Source/Runtime/Engine/Private/Materials/HLSLMaterialTranslator.h
@@ -988,7 +988,7 @@ protected:
 	virtual int32 StrataVerticalLayering(int32 Top, int32 Base) override;
 	virtual int32 StrataAdd(int32 A, int32 B) override;
 	virtual int32 StrataMultiply(int32 A, int32 Weight) override;
-	virtual int32 StrataTransmittanceToMFP(int32 TransmittanceColor, int32 OutputIndex) override;
+	virtual int32 StrataTransmittanceToMFP(int32 TransmittanceColor, int32 DesiredThickness, int32 OutputIndex) override;
 
 	virtual void StrataCompilationInfoRegisterCodeChunk(int32 CodeChunk, FStrataMaterialCompilationInfo& StrataMaterialCompilationInfo) override;
 	virtual bool StrataCompilationInfoContainsCodeChunk(int32 CodeChunk) override;
diff --git a/Engine/Source/Runtime/Engine/Private/Materials/Material.cpp b/Engine/Source/Runtime/Engine/Private/Materials/Material.cpp
--- a/Engine/Source/Runtime/Engine/Private/Materials/Material.cpp
+++ b/Engine/Source/Runtime/Engine/Private/Materials/Material.cpp
@@ -3747,7 +3747,7 @@ void UMaterial::ConvertMaterialToStrataMaterial()
 			}
 		};
 
-		// STRATA_TODO for mateiral conversion
+		// STRATA_TODO for material conversion
 		//  - WorldPositionOffset can remain on the end point node
 		//  - ShadingModel
 		//  - Refraction
@@ -3755,8 +3755,6 @@ void UMaterial::ConvertMaterialToStrataMaterial()
 		//	- MSM_Subsurface
 		//	- MSM_PreintegratedSkin
 		//	- MSM_SubsurfaceProfile
-		//	- MSM_TwoSidedFoliage
-		//	- MSM_Cloth
 		//	- MSM_Eye
 
 		if (MaterialDomain == MD_Surface && ShadingModel == MSM_Unlit)
@@ -3853,13 +3851,37 @@ void UMaterial::ConvertMaterialToStrataMaterial()
 			MoveConnectionTo(Normal, SlabBSDF, 6);			// Normal
 			MoveConnectionTo(Tangent, SlabBSDF, 7);			// Tangent
 			MoveConnectionTo(EmissiveColor, SlabBSDF, 10);	// Emissive
-			MoveConnectionTo(Tangent, SlabBSDF, 7);			// Tangent
-			MoveConnectionTo(EmissiveColor, SlabBSDF, 10);	// Emissive
 			MoveConnectionTo(ClearCoat, SlabBSDF, 14);		// FuzzAMount
 			MoveConnectionTo(SubsurfaceColor, SlabBSDF, 15);// FuzzColor
 
 			FrontMaterial.Connect(0, SlabBSDF);
 		}
+		else if (MaterialDomain == MD_Surface && ShadingModel == MSM_TwoSidedFoliage)
+		{
+			UMaterialExpressionStrataSlabBSDF* SlabBSDF = NewObject<UMaterialExpressionStrataSlabBSDF>(this);
+
+			MoveConnectionTo(BaseColor, SlabBSDF, 0);		// BaseColor
+			MoveConnectionTo(Metallic, SlabBSDF, 2);		// Metallic
+			MoveConnectionTo(Specular, SlabBSDF, 3);		// Specular
+			MoveConnectionTo(Roughness, SlabBSDF, 4);		// Roughness
+			MoveConnectionTo(Anisotropy, SlabBSDF, 5);		// Anisotropy
+			MoveConnectionTo(Normal, SlabBSDF, 6);			// Normal
+			MoveConnectionTo(Tangent, SlabBSDF, 7);			// Tangent
+			MoveConnectionTo(EmissiveColor, SlabBSDF, 10);	// Emissive
+			MoveConnectionTo(Tangent, SlabBSDF, 7);			// Tangent
+
+			UMaterialExpressionStrataTransmittanceToMFP* TransmittanceToMFP = NewObject<UMaterialExpressionStrataTransmittanceToMFP>(this);
+			MoveConnectionTo(SubsurfaceColor, TransmittanceToMFP, 0);	// SubsurfaceColor -> TransmittanceColor
+			SlabBSDF->GetInput(8)->Connect(0, TransmittanceToMFP);		// MFP -> MFP
+			SlabBSDF->GetInput(13)->Connect(1, TransmittanceToMFP);		// Thickness -> Thickness
+
+			// Set a thickness that will enabled the thin lighting model (corresponding to the legacy two-sided lighting model)
+			UMaterialExpressionConstant* ThicknessConstant = NewObject<UMaterialExpressionConstant>(this);
+			ThicknessConstant->R = STRATA_LAYER_ISTHIN_THICKNESS_THRESHOLD_CM - 1e-5f;
+			TransmittanceToMFP->Thickness.Connect(0, ThicknessConstant);
+
+			FrontMaterial.Connect(0, SlabBSDF);
+		}
 		else if (MaterialDomain == MD_Surface && ShadingModel == MSM_Hair)
 		{
 			UMaterialExpressionStrataHairBSDF* HairBSDF = NewObject<UMaterialExpressionStrataHairBSDF>(this);
diff --git a/Engine/Source/Runtime/Engine/Private/Materials/MaterialExpressions.cpp b/Engine/Source/Runtime/Engine/Private/Materials/MaterialExpressions.cpp
--- a/Engine/Source/Runtime/Engine/Private/Materials/MaterialExpressions.cpp
+++ b/Engine/Source/Runtime/Engine/Private/Materials/MaterialExpressions.cpp
@@ -20765,7 +20765,10 @@ UMaterialExpressionStrataTransmittanceToMFP::UMaterialExpressionStrataTransmitta
 #if WITH_EDITOR
 int32 UMaterialExpressionStrataTransmittanceToMFP::Compile(class FMaterialCompiler* Compiler, int32 OutputIndex)
 {
-	return Compiler->StrataTransmittanceToMFP(TransmittanceColor.GetTracedInput().Expression ? TransmittanceColor.Compile(Compiler) : Compiler->Constant(0.5f), OutputIndex);
+	return Compiler->StrataTransmittanceToMFP(
+		TransmittanceColor.GetTracedInput().Expression ? TransmittanceColor.Compile(Compiler) : Compiler->Constant(0.5f), 
+		Thickness.GetTracedInput().Expression ? Thickness.Compile(Compiler) : INDEX_NONE,
+		OutputIndex);
 }
 
 void UMaterialExpressionStrataTransmittanceToMFP::GetCaption(TArray<FString>& OutCaptions) const
diff --git a/Engine/Source/Runtime/Engine/Private/Materials/StrataMaterial.h b/Engine/Source/Runtime/Engine/Private/Materials/StrataMaterial.h
--- a/Engine/Source/Runtime/Engine/Private/Materials/StrataMaterial.h
+++ b/Engine/Source/Runtime/Engine/Private/Materials/StrataMaterial.h
@@ -18,6 +18,8 @@ struct FStrataMaterialCompilationInfo
 	{
 		uint8 Type;
 		uint8 SharedNormalIndex;
+
+		// Those properties are centered around the SLAB node. Maybe this can be abstracted at some point.
 		bool  bHasSSS;
 		bool  bHasDMFPPluggedIn;
 		bool  bHasEdgeColor;
diff --git a/Engine/Source/Runtime/Engine/Public/MaterialCompiler.h b/Engine/Source/Runtime/Engine/Public/MaterialCompiler.h
--- a/Engine/Source/Runtime/Engine/Public/MaterialCompiler.h
+++ b/Engine/Source/Runtime/Engine/Public/MaterialCompiler.h
@@ -430,7 +430,7 @@ public:
 	virtual int32 StrataVerticalLayering(int32 Top, int32 Base) = 0;
 	virtual int32 StrataAdd(int32 A, int32 B) = 0;
 	virtual int32 StrataMultiply(int32 A, int32 Weight) = 0;
-	virtual int32 StrataTransmittanceToMFP(int32 TransmittanceColor, int32 OutputIndex) = 0;
+	virtual int32 StrataTransmittanceToMFP(int32 TransmittanceColor, int32 DesiredThickness, int32 OutputIndex) = 0;
 
 	virtual void StrataCompilationInfoRegisterCodeChunk(int32 CodeChunk, FStrataMaterialCompilationInfo& StrataMaterialCompilationInfo) = 0;
 	virtual bool StrataCompilationInfoContainsCodeChunk(int32 CodeChunk) = 0;
@@ -969,9 +969,9 @@ public:
 		return Compiler->StrataMultiply(A, Weight);
 	}
 
-	virtual int32 StrataTransmittanceToMFP(int32 TransmittanceColor, int32 OutputIndex) override
+	virtual int32 StrataTransmittanceToMFP(int32 TransmittanceColor, int32 DesiredThickness, int32 OutputIndex) override
 	{
-		return Compiler->StrataTransmittanceToMFP(TransmittanceColor, OutputIndex);
+		return Compiler->StrataTransmittanceToMFP(TransmittanceColor, DesiredThickness, OutputIndex);
 	}
 
 	virtual void StrataCompilationInfoRegisterCodeChunk(int32 CodeChunk, FStrataMaterialCompilationInfo& StrataMaterialCompilationInfo) override
