Removed some no longer used strata code for single-scattering/transmittance. This is also to be consistent with starta node evaluations.

diff --git a/Engine/Shaders/Private/Strata/Strata.ush b/Engine/Shaders/Private/Strata/Strata.ush
--- a/Engine/Shaders/Private/Strata/Strata.ush
+++ b/Engine/Shaders/Private/Strata/Strata.ush
@@ -991,26 +991,6 @@ float StrataSlabDiffuseToVolumeBlend(FParticipatingMedia PM)
 	return Blend * Blend;
 }
 
-struct FStrataSlabSingleScattering
-{
-	float3 Transmittance;
-	float3 SingleScattering;
-};
-FStrataSlabSingleScattering StrataSlabSingleScattering(FParticipatingMedia PM, float NoV)
-{
-	const float3 SafeExtinctionThreshold = 0.000001f;
-	const float3 SafeExtinctionCoefficients = max(SafeExtinctionThreshold, PM.ExtinctionCoef);
-
-	const float PathLength = STRATA_SIMPLEVOLUME_THICKNESS_M / max(0.0001f, abs(NoV));
-	const float3 SafePathSegmentTransmittance = exp(-SafeExtinctionCoefficients * PathLength);
-
-	// Single scattering analytical integration from "Physically Based Unified Volumetric Rendering in Frostbite", slide 27/28.
-	FStrataSlabSingleScattering Result;
-	Result.SingleScattering = (PM.ScatteringCoef - PM.ScatteringCoef * SafePathSegmentTransmittance) / SafeExtinctionCoefficients;
-	Result.Transmittance = SafePathSegmentTransmittance;
-	return Result;
-}
-
 void EnableSlabMFPComputationSpace(inout FStrataBSDF BSDF)
 {
 #if STRATA_INLINE_SHADING
@@ -1409,10 +1389,10 @@ void PackStrataOut(
 					const float3 UnusedDiffuseColor = 0; // Unused when only trasmittance is needed
 					FParticipatingMedia PM = StrataSlabCreateParticipatingMedia(UnusedDiffuseColor, SLAB_SSSDMFP(Strata.Layers[l].BSDFs[i]));
 					const float DiffuseToVolumeBlend = StrataSlabDiffuseToVolumeBlend(PM);
-					const FStrataSlabSingleScattering Result = StrataSlabSingleScattering(PM, dot(N,V));
+					const float3 SlabTransmittance = IsotropicMediumSlabTransmittance(PM, STRATA_SIMPLEVOLUME_THICKNESS_M, ContextNoLight.NoV);
 
 					// Update the see-through throughput accoding to the medium transmittance
-					BSDFThroughput = lerp(OpaqueBSDFThroughput, (1- ApproxFresnelForTransmittance) * Result.Transmittance, DiffuseToVolumeBlend);
+					BSDFThroughput = lerp(OpaqueBSDFThroughput, (1- ApproxFresnelForTransmittance) * SlabTransmittance, DiffuseToVolumeBlend);
 
 					// Mark the BSDF has not having any SSS through post process.
 					bStrataSubsurfaceEnable = false;
