Strata code refactoring / simplifications.

diff --git a/Engine/Shaders/Private/Strata/StrataDeferredLighting.ush b/Engine/Shaders/Private/Strata/StrataDeferredLighting.ush
--- a/Engine/Shaders/Private/Strata/StrataDeferredLighting.ush
+++ b/Engine/Shaders/Private/Strata/StrataDeferredLighting.ush
@@ -205,18 +205,25 @@ float4 StrataDeferredLighting(
 		{
 			FStrataEvaluateResult BSDFEvaluate = StrataEvaluateBSDF(BSDFContext, ShadowTerms);
 
-			const float3 DiffusePathWeight			= BSDF.Weight * BSDFEvaluate.DiffusePathValue;
-			const float3 SpecularPathWeight			= BSDF.Weight * BSDFEvaluate.SpecularPathValue     * BSDFEvaluate.SpecularPathProbability;
-			const float3 SpecularHazePathWeight		= BSDF.Weight * BSDFEvaluate.SpecularHazePathValue * BSDFEvaluate.SpecularHazePathProbability;
-			const float3 TransmissionPathWeight		= BSDF.Weight * BSDFEvaluate.TransmissionPathValue;
-			const float3 ThinTransmissionPathWeight	= BSDF.Weight * BSDFEvaluate.ThinTransmissionPathValueWithProjectedSolidAngle;
+			const float3 DiffusePathValue			= BSDFEvaluate.DiffusePathValue;
+			const float3 SpecularPathValue			= BSDFEvaluate.SpecularPathValue     * BSDFEvaluate.SpecularPathProbability;
+			const float3 SpecularHazePathValue		= BSDFEvaluate.SpecularHazePathValue * BSDFEvaluate.SpecularHazePathProbability;
+			const float3 TransmissionPathValue		= BSDFEvaluate.TransmissionPathValue;
+			const float3 ThinTransmissionPathValue	= BSDFEvaluate.ThinTransmissionPathValueWithProjectedSolidAngle; // Separated because it is back face shading with -N, while StrataBSDFProjectedSolidAngleFactor account for front face.
 
 			const bool bNeedsSeparateSubsurfaceLightAccumulation = BSDFEvaluate.bSubsurface;
 
-			LightAccumulator_AddSplit(Out, DiffusePathWeight,			SpecularPathWeight,					DiffusePathWeight,			IntegrationContext.IncidentIrradianceNoShadow * ShadowTerms.SurfaceShadow,					bNeedsSeparateSubsurfaceLightAccumulation);
-			LightAccumulator_AddSplit(Out, 0.0,							SpecularHazePathWeight,				0.0,						IntegrationContext.IncidentIrradianceNoShadowHaze * ShadowTerms.SurfaceShadow,				bNeedsSeparateSubsurfaceLightAccumulation);
-			LightAccumulator_AddSplit(Out, TransmissionPathWeight,		0.0f,								TransmissionPathWeight,		IntegrationContext.IncidentTransmittedRadianceNoShadow * ShadowTerms.TransmissionShadow,	bNeedsSeparateSubsurfaceLightAccumulation);
-			LightAccumulator_AddSplit(Out, ThinTransmissionPathWeight,	0.0f,								ThinTransmissionPathWeight,	IntegrationContext.IncidentRadianceNoShadow * ShadowTerms.SurfaceShadow,					bNeedsSeparateSubsurfaceLightAccumulation);
+			float3 DiffuseTotalLight	= ShadowTerms.SurfaceShadow			* (DiffusePathValue			* IntegrationContext.IncidentIrradianceNoShadow					+	ThinTransmissionPathValue	* IntegrationContext.IncidentRadianceNoShadow)
+										+ ShadowTerms.TransmissionShadow	* (TransmissionPathValue	* IntegrationContext.IncidentTransmittedRadianceNoShadow);
+
+			float3 SpecularTotalLight	= ShadowTerms.SurfaceShadow			* (SpecularPathValue		* IntegrationContext.IncidentIrradianceNoShadow					+	SpecularHazePathValue		* IntegrationContext.IncidentIrradianceNoShadowHaze);
+
+			// Apply the BSDF weight containing coverage, transmittance to view and visibility.
+			DiffuseTotalLight  *= BSDF.Weight;
+			SpecularTotalLight *= BSDF.Weight;
+
+			const float IdentityCommonMultiplier = 1.0f;
+			LightAccumulator_AddSplit(Out, DiffuseTotalLight, SpecularTotalLight, DiffuseTotalLight, IdentityCommonMultiplier, bNeedsSeparateSubsurfaceLightAccumulation);
 		}
 		else
 		{
