Real-time texture compression on the GPU changes:
* Added BC6H and BC7 mode with least squares endpoint optimization
* Extracted common ReadBlockRGB functionality into BlockCompressionCommon.ush
* Fixed some type casting shader warnings

diff --git a/Engine/Shaders/Private/BCCompressionCommon.ush b/Engine/Shaders/Private/BCCompressionCommon.ush
--- a/Engine/Shaders/Private/BCCompressionCommon.ush
+++ b/Engine/Shaders/Private/BCCompressionCommon.ush
@@ -14,6 +14,14 @@
 
 #include "GammaCorrectionCommon.ush"
 
+// A bit slower, but higher quality compression
+#define LEAST_SQUARES_ENDPOINT_OPTIMIZATION 1
+
+uint RoundToUInt(float X)
+{
+	return (uint)round(X);
+}
+
 // Simple convert float3 color to 565 uint using 'round' arithmetic
 uint Float3ToUint565( in float3 Color )
 {
@@ -105,7 +113,7 @@ uint GetPackedColorBlockIndices( in float3 Block[16], in float3 MinColor, in flo
 	for (int i=15; i>=0; --i)
 	{
 		// Compute the distance index for this element
-		uint Index = round(dot(Block[i], ScaledRange) + Bias);
+		uint Index = RoundToUInt(dot(Block[i], ScaledRange) + Bias);
 		// Convert distance index into the BC index
 		Index += (Index > 0) - (3 * (Index == 3));
 		// OR into the final PackedIndices
@@ -131,7 +139,7 @@ uint2 GetPackedAlphaBlockIndices( in float Block[16], in float MinAlpha, in floa
 	for (i=0; i<5; ++i)
 	{
 		// Compute the distance index for this element
-		uint Index = round(Block[i] * Scale + Bias);
+		uint Index = RoundToUInt(Block[i] * Scale + Bias);
 		// Convert distance index into the BC index
 		Index += (Index > 0) - (7 * (Index == 7));
 		// OR into the final PackedIndices
@@ -141,7 +149,7 @@ uint2 GetPackedAlphaBlockIndices( in float Block[16], in float MinAlpha, in floa
 	// The 6th element is split across the x and y components
 	{
 		i = 5;
-		uint Index = round(Block[i] * Scale + Bias);
+		uint Index = RoundToUInt(Block[i] * Scale + Bias);
 		Index += (Index > 0) - (7 * (Index == 7));
 		PackedIndices.x |= (Index << 31);
 		PackedIndices.y |= (Index >> 1);
@@ -150,7 +158,7 @@ uint2 GetPackedAlphaBlockIndices( in float Block[16], in float MinAlpha, in floa
 	// The rest of the elements fill the y component
 	for (i=6; i<16; ++i)
 	{
-		uint Index = round(Block[i] * Scale + Bias);
+		uint Index = RoundToUInt(Block[i] * Scale + Bias);
 		Index += (Index > 0) - (7 * (Index == 7 ? 1 : 0));
 		PackedIndices.y |= (Index << (i*3 - 16));
 	}
@@ -212,8 +220,8 @@ uint4 CompressBC3Block( in float3 BlockRGB[16], in float BlockA[16] )
 	float MinAlpha, MaxAlpha;
 	GetMinMax( BlockA, MinAlpha, MaxAlpha );
 
-	uint MinAlphaUint = (uint)round(MinAlpha * 255.f);
-	uint MaxAlphaUint = (uint)round(MaxAlpha * 255.f);
+	uint MinAlphaUint = RoundToUInt(MinAlpha * 255.f);
+	uint MaxAlphaUint = RoundToUInt(MaxAlpha * 255.f);
 
 	uint ColorIndices = 0;
 	if (MinColor565 < MaxColor565)
@@ -248,8 +256,8 @@ uint4 CompressBC3Block_SRGB( in float3 BlockRGB[16], in float BlockA[16] )
 	float MinAlpha, MaxAlpha;
 	GetMinMax( BlockA, MinAlpha, MaxAlpha );
 
-	uint MinAlphaUint = (uint)round(MinAlpha * 255.f);
-	uint MaxAlphaUint = (uint)round(MaxAlpha * 255.f);
+	uint MinAlphaUint = RoundToUInt(MinAlpha * 255.f);
+	uint MaxAlphaUint = RoundToUInt(MaxAlpha * 255.f);
 
 	uint ColorIndices = 0;
 	if (MinColor565 < MaxColor565)
@@ -272,8 +280,8 @@ uint2 CompressBC4Block( in float Block[16] )
 	float MinAlpha, MaxAlpha;
 	GetMinMax( Block, MinAlpha, MaxAlpha );
 
-	uint MinAlphaUint = (uint)round(MinAlpha * 255.f);
-	uint MaxAlphaUint = (uint)round(MaxAlpha * 255.f);
+	uint MinAlphaUint = RoundToUInt(MinAlpha * 255.f);
+	uint MaxAlphaUint = RoundToUInt(MaxAlpha * 255.f);
 
 	uint2 Indices = 0;
 	if (MinAlphaUint < MaxAlphaUint)
@@ -290,10 +298,10 @@ uint4 CompressBC5Block( in float BlockU[16], in float BlockV[16] )
 	float MinU, MaxU, MinV, MaxV;
 	GetMinMax( BlockU, BlockV, MinU, MaxU, MinV, MaxV );
 
-	uint MinUUint = (uint)round(MinU * 255.f);
-	uint MaxUUint = (uint)round(MaxU * 255.f);
-	uint MinVUint = (uint)round(MinV * 255.f);
-	uint MaxVUint = (uint)round(MaxV * 255.f);
+	uint MinUUint = RoundToUInt(MinU * 255.f);
+	uint MaxUUint = RoundToUInt(MaxU * 255.f);
+	uint MinVUint = RoundToUInt(MinV * 255.f);
+	uint MaxVUint = RoundToUInt(MaxV * 255.f);
 
 	uint2 IndicesU = 0;
 	if (MinUUint < MaxUUint)
@@ -396,7 +404,7 @@ uint GetPackedCoCgBlockIndices(in float3 Block[16], in float2 MinCoCg, in float2
 	for (int i = 15; i >= 0; --i)
 	{
 		// Compute the distance index for this element
-		uint Index = round(dot(Block[i].yz, ScaledRange) + Bias);
+		uint Index = RoundToUInt(dot(Block[i].yz, ScaledRange) + Bias);
 		// Convert distance index into the BC index
 		Index += (Index > 0) - (3 * (Index == 3));
 		// OR into the final PackedIndices
@@ -421,7 +429,7 @@ uint2 GetPackedLumaBlockIndices(in float3 Block[16], in float MinAlpha, in float
 	for (i = 0; i < 5; ++i)
 	{
 		// Compute the distance index for this element
-		uint Index = round(Block[i].x * Scale + Bias);
+		uint Index = RoundToUInt(Block[i].x * Scale + Bias);
 		// Convert distance index into the BC index
 		Index += (Index > 0) - (7 * (Index == 7));
 		// OR into the final PackedIndices
@@ -431,7 +439,7 @@ uint2 GetPackedLumaBlockIndices(in float3 Block[16], in float MinAlpha, in float
 	// The 6th element is split across the x and y components
 	{
 		i = 5;
-		uint Index = round(Block[i].x * Scale + Bias);
+		uint Index = RoundToUInt(Block[i].x * Scale + Bias);
 		Index += (Index > 0) - (7 * (Index == 7));
 		PackedIndices.x |= (Index << 31);
 		PackedIndices.y |= (Index >> 1);
@@ -440,7 +448,7 @@ uint2 GetPackedLumaBlockIndices(in float3 Block[16], in float MinAlpha, in float
 	// The rest of the elements fill the y component
 	for (i = 6; i < 16; ++i)
 	{
-		uint Index = round(Block[i].x * Scale + Bias);
+		uint Index = RoundToUInt(Block[i].x * Scale + Bias);
 		Index += (Index > 0) - (7 * (Index == 7 ? 1 : 0));
 		PackedIndices.y |= (Index << (i * 3 - 16));
 	}
@@ -476,10 +484,242 @@ uint4 CompressBC3BlockYCoCg(in float3 Block[16])
 
 	InsetLumaEndPoints(MinColor.x, MaxColor.x);
 
-	uint MinLumaUint = round(MinColor.x * 255.f);
-	uint MaxLumaUint = round(MaxColor.x * 255.f);
+	uint MinLumaUint = RoundToUInt(MinColor.x * 255.f);
+	uint MaxLumaUint = RoundToUInt(MaxColor.x * 255.f);
 
 	uint2 Indices = GetPackedLumaBlockIndices(Block, MinColor.x, MaxColor.x);
 
 	return uint4((MinLumaUint << 8) | MaxLumaUint | Indices.x, Indices.y, (MinColor565 << 16) | MaxColor565, CoCgIndices);
 }
+
+void Swap(inout float3 A, inout float3 B)
+{
+	float3 Temp = A;
+	A = B;
+	B = Temp;
+}
+
+void Swap(inout float A, inout float B)
+{
+	float Temp = A;
+	A = B;
+	B = Temp;
+}
+
+float3 Quantize10(float3 X)
+{
+	return (f32tof16(X) * 1024.0f) / (0x7bff + 1.0f);
+}
+
+uint ComputeIndexBC6HIndex(float3 Color, float3 BlockVector, float EndPoint0Pos, float EndPoint1Pos)
+{
+	float Pos = (float)f32tof16(dot(Color, BlockVector));
+	float NormalizedPos = saturate((Pos - EndPoint0Pos) / (EndPoint1Pos - EndPoint0Pos));
+	return (uint)clamp(NormalizedPos * 14.93333f + 0.03333f + 0.5f, 0.0f, 15.0f);
+}
+
+// Compress a BC6H block. Evaluates only mode 11 for performance
+uint4 CompressBC6HBlock(in float3 BlockRGB[16])
+{
+	// Compute initial endpoints
+	float3 BlockMin = BlockRGB[0];
+	float3 BlockMax = BlockRGB[0];
+	for (uint TexelIndex = 1; TexelIndex < 16; ++TexelIndex)
+	{
+		BlockMin = min(BlockMin, BlockRGB[TexelIndex]);
+		BlockMax = max(BlockMax, BlockRGB[TexelIndex]);
+	}
+
+	float3 BlockVector = BlockMax - BlockMin;
+	BlockVector = BlockVector / (BlockVector.x + BlockVector.y + BlockVector.z);
+
+	float3 Endpoint0 = Quantize10(BlockMin);
+	float3 Endpoint1 = Quantize10(BlockMax);
+	float EndPoint0Pos = (float)f32tof16(dot(BlockMin, BlockVector));
+	float EndPoint1Pos = (float)f32tof16(dot(BlockMax, BlockVector));
+
+	// Check if endpoint swap is required
+	uint FixupIndex = ComputeIndexBC6HIndex(BlockRGB[0], BlockVector, EndPoint0Pos, EndPoint1Pos);
+	if (FixupIndex > 7)
+	{
+		Swap(EndPoint0Pos, EndPoint1Pos);
+		Swap(Endpoint0, Endpoint1);
+	}
+
+	// Compute indices
+	uint Indices[16] = { 0, 0, 0, 0, 0, 0, 0, 0, 0, 0, 0, 0, 0, 0, 0, 0 };
+	for (uint TexelIndex = 0; TexelIndex < 16; ++TexelIndex)
+	{
+		Indices[TexelIndex] = ComputeIndexBC6HIndex(BlockRGB[TexelIndex], BlockVector, EndPoint0Pos, EndPoint1Pos);
+	}
+
+	uint4 Block = 0;
+	{
+		// Encode mode 11 block
+		Block.x = 0x03;
+
+		// Encode endpoints
+		Block.x |= (uint)Endpoint0.x << 5;
+		Block.x |= (uint)Endpoint0.y << 15;
+		Block.x |= (uint)Endpoint0.z << 25;
+		Block.y |= (uint)Endpoint0.z >> 7;
+		Block.y |= (uint)Endpoint1.x << 3;
+		Block.y |= (uint)Endpoint1.y << 13;
+		Block.y |= (uint)Endpoint1.z << 23;
+		Block.z |= (uint)Endpoint1.z >> 9;
+
+		// Encode indices
+		Block.z |= Indices[0] << 1;
+		Block.z |= Indices[1] << 4;
+		Block.z |= Indices[2] << 8;
+		Block.z |= Indices[3] << 12;
+		Block.z |= Indices[4] << 16;
+		Block.z |= Indices[5] << 20;
+		Block.z |= Indices[6] << 24;
+		Block.z |= Indices[7] << 28;
+		Block.w |= Indices[8] << 0;
+		Block.w |= Indices[9] << 4;
+		Block.w |= Indices[10] << 8;
+		Block.w |= Indices[11] << 12;
+		Block.w |= Indices[12] << 16;
+		Block.w |= Indices[13] << 20;
+		Block.w |= Indices[14] << 24;
+		Block.w |= Indices[15] << 28;
+	}
+
+	return Block;
+}
+
+uint3 Quantize7(float3 X)
+{
+	return (uint3(X * 0xFF)) >> 1;
+}
+
+uint ComputeBC7Index(float3 Color, float3 BlockVector, float EndPoint0Pos, float EndPoint1Pos)
+{
+	float Pos = dot(Color, BlockVector);
+	float NormalizedPos = saturate((Pos - EndPoint0Pos) / (EndPoint1Pos - EndPoint0Pos));
+	return (uint)clamp(NormalizedPos * 14.93333f + 0.03333f + 0.5f, 0.0f, 15.0f);
+}
+
+// Least squares optimization to find best endpoints for the selected block indices in a mode 6 BC7 block
+void OptimizeEndpointsBC7(float3 Texels[16], inout float3 BlockMin, inout float3 BlockMax)
+{
+	float3 BlockVector = BlockMax - BlockMin;
+
+	float EndPoint0Pos = dot(BlockMin, BlockVector);
+	float EndPoint1Pos = dot(BlockMax, BlockVector);
+
+	float3 AlphaTexelSum = 0.0f;
+	float3 BetaTexelSum = 0.0f;
+	float AlphaBetaSum = 0.0f;
+	float AlphaSqSum = 0.0f;
+	float BetaSqSum = 0.0f;
+
+	for (uint TexelIndex = 0; TexelIndex < 16; ++TexelIndex)
+	{
+		uint Index = ComputeBC7Index(Texels[TexelIndex], BlockVector, EndPoint0Pos, EndPoint1Pos);
+
+		float Beta = saturate(Index / 15.0f);
+		float Alpha = 1.0f - Beta;
+
+		AlphaTexelSum += Alpha * Texels[TexelIndex];
+		BetaTexelSum += Beta * Texels[TexelIndex];
+
+		AlphaBetaSum += Alpha * Beta;
+
+		AlphaSqSum += Alpha * Alpha;
+		BetaSqSum += Beta * Beta;
+	}
+
+	float Det = AlphaSqSum * BetaSqSum - AlphaBetaSum * AlphaBetaSum;
+
+	if (abs(Det) > 0.1f)
+	{
+		float RcpDet = rcp(Det);
+		BlockMin = saturate(RcpDet * (AlphaTexelSum * BetaSqSum - BetaTexelSum * AlphaBetaSum));
+		BlockMax = saturate(RcpDet * (BetaTexelSum * AlphaSqSum - AlphaTexelSum * AlphaBetaSum));
+	}
+}
+
+// Compress a BC7 color only block. Evaluates only mode 6 for performance
+uint4 CompressBC7Block(in float3 BlockRGB[16])
+{
+	// Compute initial endpoints
+	float3 BlockMin = BlockRGB[0];
+	float3 BlockMax = BlockRGB[0];
+	for (uint TexelIndex = 1; TexelIndex < 16; ++TexelIndex)
+	{
+		BlockMin = min(BlockMin, BlockRGB[TexelIndex]);
+		BlockMax = max(BlockMax, BlockRGB[TexelIndex]);
+	}
+
+	#if LEAST_SQUARES_ENDPOINT_OPTIMIZATION
+	{
+		OptimizeEndpointsBC7(BlockRGB, BlockMin, BlockMax);
+	}
+	#endif
+
+	float3 BlockVector = BlockMax - BlockMin;
+
+	uint3 Endpoint0 = Quantize7(BlockMin);
+	uint3 Endpoint1 = Quantize7(BlockMax);
+	float EndPoint0Pos = dot(BlockMin, BlockVector);
+	float EndPoint1Pos = dot(BlockMax, BlockVector);
+
+	// Check if endpoint swap is required
+	uint FixupIndex = ComputeBC7Index(BlockRGB[0], BlockVector, EndPoint0Pos, EndPoint1Pos);
+	if (FixupIndex > 7)
+	{
+		Swap(EndPoint0Pos, EndPoint1Pos);
+		Swap(Endpoint0, Endpoint1);
+	}
+
+	// Compute indices
+	uint Indices[16] = { 0, 0, 0, 0, 0, 0, 0, 0, 0, 0, 0, 0, 0, 0, 0, 0 };
+	for (uint TexelIndex = 0; TexelIndex < 16; ++TexelIndex)
+	{
+		Indices[TexelIndex] = ComputeBC7Index(BlockRGB[TexelIndex], BlockVector, EndPoint0Pos, EndPoint1Pos);
+	}
+
+	uint4 Block = 0;
+	Block.x = 0;
+	Block.y = 0;
+	Block.z = 0;
+	Block.w = 0;
+	{
+		// Encode mode 6 block
+		Block.x = 0x40;
+
+		// Encode endpoints
+		Block.x |= Endpoint0.x << 7;
+		Block.x |= Endpoint1.x << 14;
+		Block.x |= Endpoint0.y << 21;
+		Block.x |= Endpoint1.y << 28;
+		Block.y |= Endpoint1.y >> 4;
+		Block.y |= Endpoint0.z << 3;
+		Block.y |= Endpoint1.z << 10;
+
+		// Encode endpoint p-bit
+
+		// Encode indices
+		Block.z |= Indices[0] << 1;
+		Block.z |= Indices[1] << 4;
+		Block.z |= Indices[2] << 8;
+		Block.z |= Indices[3] << 12;
+		Block.z |= Indices[4] << 16;
+		Block.z |= Indices[5] << 20;
+		Block.z |= Indices[6] << 24;
+		Block.z |= Indices[7] << 28;
+		Block.w |= Indices[8] << 0;
+		Block.w |= Indices[9] << 4;
+		Block.w |= Indices[10] << 8;
+		Block.w |= Indices[11] << 12;
+		Block.w |= Indices[12] << 16;
+		Block.w |= Indices[13] << 20;
+		Block.w |= Indices[14] << 24;
+		Block.w |= Indices[15] << 28;
+	}
+
+	return Block;
+}
\ No newline at end of file
diff --git a/Engine/Shaders/Private/BlockCompressionCommon.ush b/Engine/Shaders/Private/BlockCompressionCommon.ush
new file mode 100644
--- /dev/null
+++ b/Engine/Shaders/Private/BlockCompressionCommon.ush
@@ -0,0 +1,208 @@
+// Copyright Epic Games, Inc. All Rights Reserved.
+
+/*=============================================================================
+	BlockCompressionCommon.ush:
+	Helpers for compute shader block texture compression
+=============================================================================*/
+
+#pragma once
+
+// Read a 4x4 color block ready for BC compression
+void ReadBlockRGB(Texture2D<float4> SourceTexture, SamplerState TextureSampler, float2 UV, float2 TexelUVSize, out float3 Block[16])
+{
+	{
+		float4 Red = SourceTexture.GatherRed(TextureSampler, UV);
+		float4 Green = SourceTexture.GatherGreen(TextureSampler, UV);
+		float4 Blue = SourceTexture.GatherBlue(TextureSampler, UV);
+		Block[0] = float3(Red[3], Green[3], Blue[3]);
+		Block[1] = float3(Red[2], Green[2], Blue[2]);
+		Block[4] = float3(Red[0], Green[0], Blue[0]);
+		Block[5] = float3(Red[1], Green[1], Blue[1]);
+	}
+	{
+		float2 UVOffset = UV + float2(2.f * TexelUVSize.x, 0);
+		float4 Red = SourceTexture.GatherRed(TextureSampler, UVOffset);
+		float4 Green = SourceTexture.GatherGreen(TextureSampler, UVOffset);
+		float4 Blue = SourceTexture.GatherBlue(TextureSampler, UVOffset);
+		Block[2] = float3(Red[3], Green[3], Blue[3]);
+		Block[3] = float3(Red[2], Green[2], Blue[2]);
+		Block[6] = float3(Red[0], Green[0], Blue[0]);
+		Block[7] = float3(Red[1], Green[1], Blue[1]);
+	}
+	{
+		float2 UVOffset = UV + float2(0, 2.f * TexelUVSize.y);
+		float4 Red = SourceTexture.GatherRed(TextureSampler, UVOffset);
+		float4 Green = SourceTexture.GatherGreen(TextureSampler, UVOffset);
+		float4 Blue = SourceTexture.GatherBlue(TextureSampler, UVOffset);
+		Block[8] = float3(Red[3], Green[3], Blue[3]);
+		Block[9] = float3(Red[2], Green[2], Blue[2]);
+		Block[12] = float3(Red[0], Green[0], Blue[0]);
+		Block[13] = float3(Red[1], Green[1], Blue[1]);
+	}
+	{
+		float2 UVOffset = UV + 2.f * TexelUVSize;
+		float4 Red = SourceTexture.GatherRed(TextureSampler, UVOffset);
+		float4 Green = SourceTexture.GatherGreen(TextureSampler, UVOffset);
+		float4 Blue = SourceTexture.GatherBlue(TextureSampler, UVOffset);
+		Block[10] = float3(Red[3], Green[3], Blue[3]);
+		Block[11] = float3(Red[2], Green[2], Blue[2]);
+		Block[14] = float3(Red[0], Green[0], Blue[0]);
+		Block[15] = float3(Red[1], Green[1], Blue[1]);
+	}
+}
+
+// Read a 4x4 alpha block ready for BC compression
+void ReadBlockAlpha(Texture2D<float4> SourceTexture, SamplerState TextureSampler, float2 UV, float2 TexelUVSize, out float Block[16])
+{
+	{
+		float4 Alpha = SourceTexture.GatherAlpha(TextureSampler, UV);
+		Block[0] = Alpha[3];
+		Block[1] = Alpha[2];
+		Block[4] = Alpha[0];
+		Block[5] = Alpha[1];
+	}
+	{
+		float2 UVOffset = UV + float2(2.f * TexelUVSize.x, 0);
+		float4 Alpha = SourceTexture.GatherAlpha(TextureSampler, UVOffset);
+		Block[2] = Alpha[3];
+		Block[3] = Alpha[2];
+		Block[6] = Alpha[0];
+		Block[7] = Alpha[1];
+	}
+	{
+		float2 UVOffset = UV + float2(0, 2.f * TexelUVSize.y);
+		float4 Alpha = SourceTexture.GatherAlpha(TextureSampler, UVOffset);
+		Block[8] = Alpha[3];
+		Block[9] = Alpha[2];
+		Block[12] = Alpha[0];
+		Block[13] = Alpha[1];
+	}
+	{
+		float2 UVOffset = UV + 2.f * TexelUVSize;
+		float4 Alpha = SourceTexture.GatherAlpha(TextureSampler, UVOffset);
+		Block[10] = Alpha[3];
+		Block[11] = Alpha[2];
+		Block[14] = Alpha[0];
+		Block[15] = Alpha[1];
+	}
+}
+
+// Read a 4x4 color of channel X block ready for BC compression
+void ReadBlockX(Texture2D<float4> SourceTexture, SamplerState TextureSampler, float2 UV, float2 TexelUVSize, out float Block[16])
+{
+	{
+		float4 Red = SourceTexture.GatherRed(TextureSampler, UV);
+		Block[0] = Red[3];
+		Block[1] = Red[2];
+		Block[4] = Red[0];
+		Block[5] = Red[1];
+	}
+	{
+		float2 UVOffset = UV + float2(2.f * TexelUVSize.x, 0);
+		float4 Red = SourceTexture.GatherRed(TextureSampler, UVOffset);
+		Block[2] = Red[3];
+		Block[3] = Red[2];
+		Block[6] = Red[0];
+		Block[7] = Red[1];
+	}
+	{
+		float2 UVOffset = UV + float2(0, 2.f * TexelUVSize.y);
+		float4 Red = SourceTexture.GatherRed(TextureSampler, UVOffset);
+		Block[8] = Red[3];
+		Block[9] = Red[2];
+		Block[12] = Red[0];
+		Block[13] = Red[1];
+	}
+	{
+		float2 UVOffset = UV + 2.f * TexelUVSize;
+		float4 Red = SourceTexture.GatherRed(TextureSampler, UVOffset);
+		Block[10] = Red[3];
+		Block[11] = Red[2];
+		Block[14] = Red[0];
+		Block[15] = Red[1];
+	}
+}
+
+// Read a 4x4 block of XY channels from a normal texture ready for BC5 compression.
+void ReadBlockXY(Texture2D<float4> SourceTexture, SamplerState TextureSampler, float2 UV, float2 TexelUVSize, out float BlockX[16], out float BlockY[16])
+{
+	{
+		float4 Red = SourceTexture.GatherRed(TextureSampler, UV);
+		float4 Green = SourceTexture.GatherGreen(TextureSampler, UV);
+		BlockX[0] = Red[3]; BlockY[0] = Green[3];
+		BlockX[1] = Red[2]; BlockY[1] = Green[2];
+		BlockX[4] = Red[0]; BlockY[4] = Green[0];
+		BlockX[5] = Red[1]; BlockY[5] = Green[1];
+	}
+	{
+		float2 UVOffset = UV + float2(2.f * TexelUVSize.x, 0);
+		float4 Red = SourceTexture.GatherRed(TextureSampler, UVOffset);
+		float4 Green = SourceTexture.GatherGreen(TextureSampler, UVOffset);
+		BlockX[2] = Red[3]; BlockY[2] = Green[3];
+		BlockX[3] = Red[2]; BlockY[3] = Green[2];
+		BlockX[6] = Red[0]; BlockY[6] = Green[0];
+		BlockX[7] = Red[1]; BlockY[7] = Green[1];
+	}
+	{
+		float2 UVOffset = UV + float2(0, 2.f * TexelUVSize.y);
+		float4 Red = SourceTexture.GatherRed(TextureSampler, UVOffset);
+		float4 Green = SourceTexture.GatherGreen(TextureSampler, UVOffset);
+		BlockX[8] = Red[3]; BlockY[8] = Green[3];
+		BlockX[9] = Red[2]; BlockY[9] = Green[2];
+		BlockX[12] = Red[0]; BlockY[12] = Green[0];
+		BlockX[13] = Red[1]; BlockY[13] = Green[1];
+	}
+	{
+		float2 UVOffset = UV + 2.f * TexelUVSize;
+		float4 Red = SourceTexture.GatherRed(TextureSampler, UVOffset);
+		float4 Green = SourceTexture.GatherGreen(TextureSampler, UVOffset);
+		BlockX[10] = Red[3]; BlockY[10] = Green[3];
+		BlockX[11] = Red[2]; BlockY[11] = Green[2];
+		BlockX[14] = Red[0]; BlockY[14] = Green[0];
+		BlockX[15] = Red[1]; BlockY[15] = Green[1];
+	}
+}
+
+// Read a 4x4 block of XYA channels ready for BC4/BC5 compression.
+void ReadBlockXYA(Texture2D<float4> SourceTexture, SamplerState TextureSampler, float2 UV, float2 TexelUVSize, out float BlockX[16], out float BlockY[16], out float BlockA[16])
+{
+	{
+		float4 Red = SourceTexture.GatherRed(TextureSampler, UV);
+		float4 Green = SourceTexture.GatherGreen(TextureSampler, UV);
+		float4 Blue = SourceTexture.GatherBlue(TextureSampler, UV);
+		BlockX[0] = Red[3]; BlockY[0] = Green[3]; BlockA[0] = Blue[3];
+		BlockX[1] = Red[2]; BlockY[1] = Green[2]; BlockA[1] = Blue[2];
+		BlockX[4] = Red[0]; BlockY[4] = Green[0]; BlockA[4] = Blue[0];
+		BlockX[5] = Red[1]; BlockY[5] = Green[1]; BlockA[5] = Blue[1];
+	}
+	{
+		float2 UVOffset = UV + float2(2.f * TexelUVSize.x, 0);
+		float4 Red = SourceTexture.GatherRed(TextureSampler, UVOffset);
+		float4 Green = SourceTexture.GatherGreen(TextureSampler, UVOffset);
+		float4 Blue = SourceTexture.GatherBlue(TextureSampler, UV);
+		BlockX[2] = Red[3]; BlockY[2] = Green[3]; BlockA[2] = Blue[3];
+		BlockX[3] = Red[2]; BlockY[3] = Green[2]; BlockA[3] = Blue[2];
+		BlockX[6] = Red[0]; BlockY[6] = Green[0]; BlockA[6] = Blue[0];
+		BlockX[7] = Red[1]; BlockY[7] = Green[1]; BlockA[7] = Blue[1];
+	}
+	{
+		float2 UVOffset = UV + float2(0, 2.f * TexelUVSize.y);
+		float4 Red = SourceTexture.GatherRed(TextureSampler, UVOffset);
+		float4 Green = SourceTexture.GatherGreen(TextureSampler, UVOffset);
+		float4 Blue = SourceTexture.GatherBlue(TextureSampler, UV);
+		BlockX[8] = Red[3]; BlockY[8] = Green[3]; BlockA[8] = Blue[3];
+		BlockX[9] = Red[2]; BlockY[9] = Green[2]; BlockA[9] = Blue[2];
+		BlockX[12] = Red[0]; BlockY[12] = Green[0]; BlockA[12] = Blue[0];
+		BlockX[13] = Red[1]; BlockY[13] = Green[1]; BlockA[13] = Blue[1];
+	}
+	{
+		float2 UVOffset = UV + 2.f * TexelUVSize;
+		float4 Red = SourceTexture.GatherRed(TextureSampler, UVOffset);
+		float4 Green = SourceTexture.GatherGreen(TextureSampler, UVOffset);
+		float4 Blue = SourceTexture.GatherBlue(TextureSampler, UV);
+		BlockX[10] = Red[3]; BlockY[10] = Green[3]; BlockA[10] = Blue[3];
+		BlockX[11] = Red[2]; BlockY[11] = Green[2]; BlockA[11] = Blue[2];
+		BlockX[14] = Red[0]; BlockY[14] = Green[0]; BlockA[14] = Blue[0];
+		BlockX[15] = Red[1]; BlockY[15] = Green[1]; BlockA[15] = Blue[1];
+	}
+}
\ No newline at end of file
diff --git a/Engine/Shaders/Private/VirtualTextureCompress.usf b/Engine/Shaders/Private/VirtualTextureCompress.usf
--- a/Engine/Shaders/Private/VirtualTextureCompress.usf
+++ b/Engine/Shaders/Private/VirtualTextureCompress.usf
@@ -7,6 +7,7 @@
 =============================================================================*/
 
 #include "Common.ush"
+#include "BlockCompressionCommon.ush"
 #include "BCCompressionCommon.ush"
 #include "ETCCompressionCommon.ush"
 
@@ -64,134 +65,6 @@ void CopyVS(in uint VertexId : SV_VertexID, out float4 OutPosition : SV_POSITION
 	OutPosition.zw = float2(0, 1);
 }
 
-/** Read a 4x4 color block ready for BC1 compression. */
-void ReadBlockRGB( Texture2D<float4> SourceTexture, SamplerState TextureSampler, float2 UV, float2 TexelUVSize, out float3 Block[16] )
-{
-	{
-		float4 Red = SourceTexture.GatherRed(TextureSampler, UV);
-		float4 Green = SourceTexture.GatherGreen(TextureSampler, UV);
-		float4 Blue = SourceTexture.GatherBlue(TextureSampler, UV);
-		Block[0] = float3(Red[3], Green[3], Blue[3]);
-		Block[1] = float3(Red[2], Green[2], Blue[2]);
-		Block[4] = float3(Red[0], Green[0], Blue[0]);
-		Block[5] = float3(Red[1], Green[1], Blue[1]);
-	}
-	{
-		float2 UVOffset = UV + float2(2.f * TexelUVSize.x, 0);
-		float4 Red = SourceTexture.GatherRed(TextureSampler, UVOffset);
-		float4 Green = SourceTexture.GatherGreen(TextureSampler, UVOffset);
-		float4 Blue = SourceTexture.GatherBlue(TextureSampler, UVOffset);
-		Block[2] = float3(Red[3], Green[3], Blue[3]);
-		Block[3] = float3(Red[2], Green[2], Blue[2]);
-		Block[6] = float3(Red[0], Green[0], Blue[0]);
-		Block[7] = float3(Red[1], Green[1], Blue[1]);
-	}
-	{
-		float2 UVOffset = UV + float2(0, 2.f * TexelUVSize.y);
-		float4 Red = SourceTexture.GatherRed(TextureSampler, UVOffset);
-		float4 Green = SourceTexture.GatherGreen(TextureSampler, UVOffset);
-		float4 Blue = SourceTexture.GatherBlue(TextureSampler, UVOffset);
-		Block[8] = float3(Red[3], Green[3], Blue[3]);
-		Block[9] = float3(Red[2], Green[2], Blue[2]);
-		Block[12] = float3(Red[0], Green[0], Blue[0]);
-		Block[13] = float3(Red[1], Green[1], Blue[1]);
-	}
-	{
-		float2 UVOffset = UV + 2.f * TexelUVSize;
-		float4 Red = SourceTexture.GatherRed(TextureSampler, UVOffset);
-		float4 Green = SourceTexture.GatherGreen(TextureSampler, UVOffset);
-		float4 Blue = SourceTexture.GatherBlue(TextureSampler, UVOffset);
-		Block[10] = float3(Red[3], Green[3], Blue[3]);
-		Block[11] = float3(Red[2], Green[2], Blue[2]);
-		Block[14] = float3(Red[0], Green[0], Blue[0]);
-		Block[15] = float3(Red[1], Green[1], Blue[1]);
-	}
-}
-
-/** Read a 4x4 block of XY channels from a normal texture ready for BC5 compression. */
-void ReadBlockXY( Texture2D<float4> SourceTexture, SamplerState TextureSampler, float2 UV, float2 TexelUVSize, out float BlockX[16], out float BlockY[16] )
-{
-	{
-		float4 Red = SourceTexture.GatherRed(TextureSampler, UV);
-		float4 Green = SourceTexture.GatherGreen(TextureSampler, UV);
-		BlockX[0] = Red[3]; BlockY[0] = Green[3];
-		BlockX[1] = Red[2]; BlockY[1] = Green[2];
-		BlockX[4] = Red[0]; BlockY[4] = Green[0];
-		BlockX[5] = Red[1]; BlockY[5] = Green[1];
-	}
-	{
-		float2 UVOffset = UV + float2(2.f * TexelUVSize.x, 0);
-		float4 Red = SourceTexture.GatherRed(TextureSampler, UVOffset);
-		float4 Green = SourceTexture.GatherGreen(TextureSampler, UVOffset);
-		BlockX[2] = Red[3]; BlockY[2] = Green[3];
-		BlockX[3] = Red[2]; BlockY[3] = Green[2];
-		BlockX[6] = Red[0]; BlockY[6] = Green[0];
-		BlockX[7] = Red[1]; BlockY[7] = Green[1];
-	}
-	{
-		float2 UVOffset = UV + float2(0, 2.f * TexelUVSize.y);
-		float4 Red = SourceTexture.GatherRed(TextureSampler, UVOffset);
-		float4 Green = SourceTexture.GatherGreen(TextureSampler, UVOffset);
-		BlockX[8]  = Red[3]; BlockY[8]  = Green[3];
-		BlockX[9]  = Red[2]; BlockY[9]  = Green[2];
-		BlockX[12] = Red[0]; BlockY[12] = Green[0];
-		BlockX[13] = Red[1]; BlockY[13] = Green[1];
-	}
-	{
-		float2 UVOffset = UV + 2.f * TexelUVSize;
-		float4 Red = SourceTexture.GatherRed(TextureSampler, UVOffset);
-		float4 Green = SourceTexture.GatherGreen(TextureSampler, UVOffset);
-		BlockX[10] = Red[3]; BlockY[10] = Green[3];
-		BlockX[11] = Red[2]; BlockY[11] = Green[2];
-		BlockX[14] = Red[0]; BlockY[14] = Green[0];
-		BlockX[15] = Red[1]; BlockY[15] = Green[1];
-	}
-}
-
-/** Read a 4x4 block of XYA channels ready for BC4/BC5 compression. */
-void ReadBlockXYA( Texture2D<float4> SourceTexture, SamplerState TextureSampler, float2 UV, float2 TexelUVSize, out float BlockX[16], out float BlockY[16], out float BlockA[16] )
-{
-	{
-		float4 Red = SourceTexture.GatherRed(TextureSampler, UV);
-		float4 Green = SourceTexture.GatherGreen(TextureSampler, UV);
-		float4 Blue = SourceTexture.GatherBlue(TextureSampler, UV);
-		BlockX[0] = Red[3]; BlockY[0] = Green[3]; BlockA[0] = Blue[3];
-		BlockX[1] = Red[2]; BlockY[1] = Green[2]; BlockA[1] = Blue[2];
-		BlockX[4] = Red[0]; BlockY[4] = Green[0]; BlockA[4] = Blue[0];
-		BlockX[5] = Red[1]; BlockY[5] = Green[1]; BlockA[5] = Blue[1];
-	}
-	{
-		float2 UVOffset = UV + float2(2.f * TexelUVSize.x, 0);
-		float4 Red = SourceTexture.GatherRed(TextureSampler, UVOffset);
-		float4 Green = SourceTexture.GatherGreen(TextureSampler, UVOffset);
-		float4 Blue = SourceTexture.GatherBlue(TextureSampler, UV);
-		BlockX[2] = Red[3]; BlockY[2] = Green[3]; BlockA[2] = Blue[3];
-		BlockX[3] = Red[2]; BlockY[3] = Green[2]; BlockA[3] = Blue[2];
-		BlockX[6] = Red[0]; BlockY[6] = Green[0]; BlockA[6] = Blue[0];
-		BlockX[7] = Red[1]; BlockY[7] = Green[1]; BlockA[7] = Blue[1];
-	}
-	{
-		float2 UVOffset = UV + float2(0, 2.f * TexelUVSize.y);
-		float4 Red = SourceTexture.GatherRed(TextureSampler, UVOffset);
-		float4 Green = SourceTexture.GatherGreen(TextureSampler, UVOffset);
-		float4 Blue = SourceTexture.GatherBlue(TextureSampler, UV);
-		BlockX[8]  = Red[3]; BlockY[8]  = Green[3]; BlockA[8]  = Blue[3];
-		BlockX[9]  = Red[2]; BlockY[9]  = Green[2]; BlockA[9]  = Blue[2];
-		BlockX[12] = Red[0]; BlockY[12] = Green[0]; BlockA[12] = Blue[0];
-		BlockX[13] = Red[1]; BlockY[13] = Green[1]; BlockA[13] = Blue[1];
-	}
-	{
-		float2 UVOffset = UV + 2.f * TexelUVSize;
-		float4 Red = SourceTexture.GatherRed(TextureSampler, UVOffset);
-		float4 Green = SourceTexture.GatherGreen(TextureSampler, UVOffset);
-		float4 Blue = SourceTexture.GatherBlue(TextureSampler, UV);
-		BlockX[10] = Red[3]; BlockY[10] = Green[3]; BlockA[10] = Blue[3];
-		BlockX[11] = Red[2]; BlockY[11] = Green[2]; BlockA[11] = Blue[2];
-		BlockX[14] = Red[0]; BlockY[14] = Green[0]; BlockA[14] = Blue[0];
-		BlockX[15] = Red[1]; BlockY[15] = Green[1]; BlockA[15] = Blue[1];
-	}
-}
-
 /** Compress base color to a single BC1 target. */
 [numthreads(8,8,1)]
 void CompressBaseColorCS(
