Implements experimental CONFIG_VFX_SURVIVAL

diff --git a/Engine/Shaders/Private/TemporalAA/TAAUpdateHistory.usf b/Engine/Shaders/Private/TemporalAA/TAAUpdateHistory.usf
--- a/Engine/Shaders/Private/TemporalAA/TAAUpdateHistory.usf
+++ b/Engine/Shaders/Private/TemporalAA/TAAUpdateHistory.usf
@@ -22,6 +22,8 @@
 
 #define CONFIG_TRANSLUCENCY_REJECTION 1
 
+#define CONFIG_VFX_SURVIVAL 1
+
 
 //------------------------------------------------------- CONSTANTS
 
@@ -269,6 +271,13 @@ void MainCS(
 
 			taa_half ToneWeight = HdrWeight4(InputColor);
 
+			#if CONFIG_VFX_SURVIVAL
+			if (SampleId == 0)
+			{
+				ToneWeight = max(ToneWeight, ComputeSampleWeigth(1.0, dPP, /* MinimalContribution = */ 0.0));
+			}
+			#endif
+
 			FilteredInputColor       += (SampleSpatialWeight * ToneWeight) * InputColor;
 			FilteredInputColorWeight += (SampleSpatialWeight * ToneWeight);
 
