Fixed Android shader compilation error

diff --git a/Engine/Shaders/Private/BCCompressionCommon.ush b/Engine/Shaders/Private/BCCompressionCommon.ush
--- a/Engine/Shaders/Private/BCCompressionCommon.ush
+++ b/Engine/Shaders/Private/BCCompressionCommon.ush
@@ -492,6 +492,13 @@ uint4 CompressBC3BlockYCoCg(in float3 Block[16])
 	return uint4((MinLumaUint << 8) | MaxLumaUint | Indices.x, Indices.y, (MinColor565 << 16) | MaxColor565, CoCgIndices);
 }
 
+void Swap(inout uint3 A, inout uint3 B)
+{
+	uint3 Temp = A;
+	A = B;
+	B = Temp;
+}
+
 void Swap(inout float3 A, inout float3 B)
 {
 	float3 Temp = A;
