Strata
- Using the new simple model after discussion with Jordan: our games does not use that model or can be adapted easily.
- Fixed usage of specular transmittance for Simple and Thin lighting models.

diff --git a/Engine/Shaders/Private/Strata/StrataDeferredLighting.ush b/Engine/Shaders/Private/Strata/StrataDeferredLighting.ush
--- a/Engine/Shaders/Private/Strata/StrataDeferredLighting.ush
+++ b/Engine/Shaders/Private/Strata/StrataDeferredLighting.ush
@@ -209,7 +209,7 @@ float4 StrataDeferredLighting(
 			const float3 SpecularPathValue			= BSDFEvaluate.SpecularPathValue     * BSDFEvaluate.SpecularPathProbability;
 			const float3 SpecularHazePathValue		= BSDFEvaluate.SpecularHazePathValue * BSDFEvaluate.SpecularHazePathProbability;
 			const float3 TransmissionPathValue		= BSDFEvaluate.TransmissionPathValue;
-			const float3 ThinTransmissionPathValue	= BSDFEvaluate.ThinTransmissionPathValueWithProjectedSolidAngle; // Separated because it is back face shading with -N, while StrataBSDFProjectedSolidAngleFactor account for front face.
+			const float3 ThinTransmissionPathValue	= BSDFEvaluate.ThinTransmissionPathValue * StrataBSDFBackFaceProjectedSolidAngleFactor(BSDFContext); // Separated because it is back face shading with -N, while StrataBSDFProjectedSolidAngleFactor account for front face.
 
 			const bool bNeedsSeparateSubsurfaceLightAccumulation = BSDFEvaluate.bSubsurface;
 
diff --git a/Engine/Shaders/Private/Strata/StrataEvaluation.ush b/Engine/Shaders/Private/Strata/StrataEvaluation.ush
--- a/Engine/Shaders/Private/Strata/StrataEvaluation.ush
+++ b/Engine/Shaders/Private/Strata/StrataEvaluation.ush
@@ -33,9 +33,6 @@
 #define STRATA_SSS_TRANSMISSION 0
 #endif
 
-// STRATA_TODO Enable the new model and get rid of the old one.
-#define THINLIGHTING_NEWMODEL 0
-
 // STRATA_TODO put in a common file
 // Point lobe in off-specular peak direction
 float3 StrataGetOffSpecularPeakReflectionDir(float3 Normal, float3 ReflectionVector, float Roughness)
@@ -169,8 +166,7 @@ struct FStrataEvaluateResult
 	float3 SpecularPathValue;
 	float3 SpecularHazePathValue;
 	float3 TransmissionPathValue;
-
-	float3 ThinTransmissionPathValueWithProjectedSolidAngle;
+	float3 ThinTransmissionPathValue;
 
 	// The following probabilities are part of the pdf to work correctly with a monter carlo based integrator and as expected are not applied on the xxxPathValue.
 	// So PathProbability should be applied on PathValue in the rasteriser to recover the correct balance.
@@ -194,6 +190,12 @@ float StrataBSDFProjectedSolidAngleFactor(in FStrataBSDFContext BSDFContext)
 	return BSDFType == STRATA_BSDF_TYPE_HAIR ? 1.0f : BSDFContext.SatNoL;
 }
 
+float StrataBSDFBackFaceProjectedSolidAngleFactor(in FStrataBSDFContext BSDFContext)
+{
+	const uint BSDFType = BSDF_GETTYPE(BSDFContext.BSDF);
+	return BSDFType == STRATA_BSDF_TYPE_HAIR ? 1.0f : saturate(-BSDFContext.Context.NoL);
+}
+
 #if STRATA_MICROFACET_MULTIPLESCATTERING
 
 FStrataGGXLUTResult StrataEvaluateGGXLUT(float3 DiffuseAlbedo, float3 F0, float NoV, float Metallic, float SafeRoughnessX, float SafeRoughnessY)
@@ -322,18 +324,21 @@ FStrataEvaluateResult StrataEvaluateBSDF(FStrataBSDFContext BSDFContext, FShadow
 			float3 MSScale = 1;
 			float3 HazeMSScale = 1;
 			#if STRATA_MICROFACET_MULTIPLESCATTERING
-			{
-				MSScale	= StrataEvaluateGGXLUT(DiffuseColor, F0, BSDFContext.Context.NoV, Metallic, SafeRoughnessX, SafeRoughnessY).SpecularMultipleScatteringScale;
+				FStrataGGXLUTResult StrataGGXLUTResult = StrataEvaluateGGXLUT(DiffuseColor, F0, BSDFContext.Context.NoV, Metallic, SafeRoughnessX, SafeRoughnessY);
+				MSScale	= StrataGGXLUTResult.SpecularMultipleScatteringScale;
+				const float3 DirectionalAlbedo_SpecularTransmission = StrataGGXLUTResult.DirectionalAlbedo.SpecularT; // We ignore Haziness influence
 				#if STRATA_FASTPATH==0
 				if (bHaziness)
 				{
 					HazeMSScale = StrataEvaluateGGXLUT(DiffuseColor, F0, BSDFContext.Context.NoV, Metallic, HazeSafeRoughnessX, HazeSafeRoughnessY).SpecularMultipleScatteringScale;
 				}
 				#endif
-			}
+			#else
+				const float3 ApproxFresnelForSpecularTransmittance = F_Schlick(F0, F90, BSDFContext.SatNoV);
+				const float3 DirectionalAlbedo_SpecularTransmission = 1.0f - ApproxFresnelForSpecularTransmittance;
 			#endif
 
-			float3 F = 0;
+			float3 FresnelTerm = 0;
 			#if STRATA_FASTPATH==0
 			const bool bHasThinFilm = BSDF_GETHASTHINFILM(BSDFContext.BSDF);
 			BRANCH
@@ -341,21 +346,21 @@ FStrataEvaluateResult StrataEvaluateBSDF(FStrataBSDFContext BSDFContext, FShadow
 			{
 				float ThinFilmIOR = SLAB_THINFILM_IOR(BSDFContext.BSDF);
 				float ThinFilmThickness = SLAB_THINFILM_THICKNESS(BSDFContext.BSDF);
-				F = F_ThinFilm(BSDFContext.Context.NoV, BSDFContext.Context.NoL, BSDFContext.Context.VoH, F0, F90, ThinFilmIOR, ThinFilmThickness);
+				FresnelTerm = F_ThinFilm(BSDFContext.Context.NoV, BSDFContext.Context.NoL, BSDFContext.Context.VoH, F0, F90, ThinFilmIOR, ThinFilmThickness);
 			}
 			else
 			#endif
 			{
-				F = F_Schlick(F0, F90, BSDFContext.Context.VoH);
+				FresnelTerm = F_Schlick(F0, F90, BSDFContext.Context.VoH);
 			}
 
 			Sample.SpecularPathProbability		= bHaziness	? (1.0f - SPECULAR_HAZINESS_WEIGHT) : 1.0f;
 			Sample.SpecularHazePathProbability	= bHaziness	? SPECULAR_HAZINESS_WEIGHT			: 0.0f;
 
-			Sample.SpecularPathValue		= D * Vis * F * MSScale;
+			Sample.SpecularPathValue		= D * Vis * FresnelTerm * MSScale;
 			Sample.SpecularPDF				= PDF * Sample.SpecularPathProbability;
 
-			Sample.SpecularHazePathValue	= HazeD * HazeVis * F * HazeMSScale;
+			Sample.SpecularHazePathValue	= HazeD * HazeVis * FresnelTerm * HazeMSScale;
 			Sample.SpecularHazePDF			= HazePDF * Sample.SpecularHazePathProbability;
 
 			Sample.EmissivePathValue = BSDF_GETEMISSIVE(BSDFContext.BSDF);
@@ -435,38 +440,22 @@ FStrataEvaluateResult StrataEvaluateBSDF(FStrataBSDFContext BSDFContext, FShadow
 			if (BSDF_GETISTHIN(BSDFContext.BSDF))
 			{
 				FParticipatingMedia PM = StrataSlabCreateParticipatingMedia(DiffuseColor, SLAB_SSSDMFP(BSDFContext.BSDF));
+				const float BackFaceNoL = -BSDFContext.Context.NoL;
+				const float BackFaceNoLSat = saturate(BackFaceNoL);
 
-				// we consider single scattering and perpendicular to the surface light transmittance, no diffusion for now.
-				const float TransmittanceNoV = 1.0;	
-				const float3 SlabTransmittance = IsotropicMediumSlabTransmittance(PM, STRATA_SIMPLEVOLUME_THICKNESS_M, TransmittanceNoV);
-
-			#if THINLIGHTING_NEWMODEL
-				//--- New version
-
-				// STRATA_TODO Fresnel throughput, Use a LUT, this is currently in progress and will be coming soon.
-				// In the meantime, we use a simple approach to feed Fresnel function and compute an approximation of the general transmition/througput.
-				// We also ignored ThinFilm interaction as another simplification for perf.
-				const float3 ApproxFresnelForTransmittance = F_Schlick(F0, F90, BSDFContext.SatNoV);
-				const float3 SpecularThroughput = 1.0f - ApproxFresnelForTransmittance; // This is already taken into account by the Chan model. But for two sided lighting we need to apply it.
+				// We consider single scattering and perpendicular to the surface light transmittance.
+				// Using BackFaceNoL instead of 1 for TransmittanceNoL would be more correct but it gives some high frequency details looking weird without multiple scattering.
+				const float TransmittanceNoL = 1.0f;	
+				const float3 SlabTransmittance = IsotropicMediumSlabTransmittance(PM, STRATA_SIMPLEVOLUME_THICKNESS_M, TransmittanceNoL);
 
 				// We recover the transmitted luminance from as the back face Lambert affected by the transmittance.
 				// The MFP is specified on the slab node and then transmittance recoved from it.
 				// We assume the back face diffuse albedo is white=1 and no diffusion happens at this stage.
-				const float3 BackFaceWhiteDiffuseColor = 1.0f;
-				const float3 BackFaceDiffuse = saturate(-BSDFContext.Context.NoL) * Diffuse_Lambert(BackFaceWhiteDiffuseColor);
-			#else
-				//--- Legacy version
-
-				const float3 SpecularThroughput = 1.0f;
-
-				// Scatter distribution following GGX mixed with a wrapped diffuse Lambert.
-				float BackFaceDiffuse = D_GGX(0.6 * 0.6, saturate(-BSDFContext.Context.VoL));
-				const float Wrap = 0.5;
-				BackFaceDiffuse *= saturate((-BSDFContext.Context.NoL + Wrap) / Square(1 + Wrap));
-			#endif
+				const float3 BackFaceWhiteDiffuseColor = float3(1.0f, 1.0f, 1.0f);
 
-				Sample.ThinTransmissionPathValueWithProjectedSolidAngle = SlabTransmittance *  SpecularThroughput * BackFaceDiffuse;
-				Sample.ThinTransmissionPathPDF = saturate(-BSDFContext.Context.NoL) / PI;
+				// The diffuse is attenuated by the slab transmittance and the SpecularTransmission through the GGX interface.
+				Sample.ThinTransmissionPathValue = Diffuse_Lambert(BackFaceWhiteDiffuseColor) * SlabTransmittance * DirectionalAlbedo_SpecularTransmission;
+				Sample.ThinTransmissionPathPDF = BackFaceNoLSat / PI;
 
 				// At this stage, Sample.bSubsurface should have been setup correctly already if SSS is enabled with two-sided lighting.
 			}
@@ -481,17 +470,11 @@ FStrataEvaluateResult StrataEvaluateBSDF(FStrataBSDFContext BSDFContext, FShadow
 				const float3 SlabDirectionalAlbedo = IsotropicMediumSlabPunctualDirectionalAlbedo(PM);
 				const float3 SlabTransmittance = IsotropicMediumSlabTransmittance(PM, STRATA_SIMPLEVOLUME_THICKNESS_M, BSDFContext.Context.NoV);
 
-				// STRATA_TODO Fresnel throughput, Use a LUT, this is currently in progress and will be coming soon.
-				// In the meantime, we use a simple approach to feed Fresnel function and compute an approximation of the general transmition/througput.
-				// We also ignored ThinFilm interaction as another simplification for perf.
-				const float3 ApproxFresnelForTransmittance = F_Schlick(F0, F90, BSDFContext.SatNoV);
-
-				const float3 SpecularThroughput = 1.0f - ApproxFresnelForTransmittance; // This is already taken into account by the Chan model above. But for volumetric we need to apply it.
 				const float Phase = IsotropicPhase();
 
-				Sample.DiffusePathValue = lerp(Sample.DiffusePathValue,		SlabDirectionalAlbedo * SpecularThroughput,				DiffuseToVolumeBlend);
-				Sample.DiffusePDF		= lerp(Sample.DiffusePDF,			Phase,													DiffuseToVolumeBlend);
-				Sample.Throughput		= lerp(Sample.Throughput,			SlabTransmittance * SpecularThroughput,					DiffuseToVolumeBlend);
+				Sample.DiffusePathValue = lerp(Sample.DiffusePathValue,		SlabDirectionalAlbedo * DirectionalAlbedo_SpecularTransmission,		DiffuseToVolumeBlend);
+				Sample.DiffusePDF		= lerp(Sample.DiffusePDF,			Phase,																DiffuseToVolumeBlend);
+				Sample.Throughput		= lerp(Sample.Throughput,			SlabTransmittance     * DirectionalAlbedo_SpecularTransmission,		DiffuseToVolumeBlend);
 
 				Sample.bSubsurface = false;	// It should already be the case because Enforce in this case because BSDF_GETHASSSS should be false when BSDF_GETISSIMPLEVOLUME is true
 			}
@@ -676,6 +659,9 @@ FStrataEnvLightResult StrataEvaluateForEnvLight(FStrataBSDFContext BSDFContext,
 			const float3 F90 = ComputeF90(F0, SLAB_EDGECOLOR(BSDFContext.BSDF), Metallic);
 			float3 ThinFilmFresnelWeight = 1;
 
+			const float3 ApproxFresnelForSpecularTransmittance = F_Schlick(F0, F90, BSDFContext.SatNoV);
+			float3 DirectionalAlbedo_SpecularTransmission = 1.0f - ApproxFresnelForSpecularTransmittance;
+
 			BRANCH
 			if (bEnableSpecular)
 			{
@@ -717,6 +703,7 @@ FStrataEnvLightResult StrataEvaluateForEnvLight(FStrataBSDFContext BSDFContext,
 					FStrataGGXLUTResult GGXLUT	= StrataEvaluateGGXLUT(DiffuseColor, F0, BSDFContext.Context.NoV, Metallic, SafeRoughnessX, SafeRoughnessX);
 					EvalEnvBRDF					= ThinFilmFresnelWeight * GGXLUT.DirectionalAlbedo.SpecularR;
 					StrataEnvLightResult.DiffuseWeight = GGXLUT.DirectionalAlbedo.Diffuse;
+					DirectionalAlbedo_SpecularTransmission = GGXLUT.DirectionalAlbedo.SpecularT; // We ignore Haziness influence
 				}
 			#else
 				const float3 EvalEnvBRDF = ThinFilmFresnelWeight * StrataEnvBRDF(F0, F90, SafeRoughnessX, BSDFContext.SatNoV);
@@ -796,7 +783,7 @@ FStrataEnvLightResult StrataEvaluateForEnvLight(FStrataBSDFContext BSDFContext,
 				SlabDirectionalAlbedo *= GGXLUT.DirectionalAlbedo.SpecularT;
 			#endif
 
-				StrataEnvLightResult.DiffuseWeight = lerp(StrataEnvLightResult.DiffuseWeight, SlabDirectionalAlbedo * ThinFilmFresnelWeight,	DiffuseToVolumeBlend);
+				StrataEnvLightResult.DiffuseWeight = lerp(StrataEnvLightResult.DiffuseWeight, SlabDirectionalAlbedo * DirectionalAlbedo_SpecularTransmission,	DiffuseToVolumeBlend);
 				// We keep the same normal
 				StrataEnvLightResult.bSubsurface = false;
 			}
@@ -809,24 +796,20 @@ FStrataEnvLightResult StrataEvaluateForEnvLight(FStrataBSDFContext BSDFContext,
 			{
 				FParticipatingMedia PM = StrataSlabCreateParticipatingMedia(DiffuseColor, SLAB_SSSDMFP(BSDFContext.BSDF));
 
-				// we consider single scattering and perpendicular to the surface light transmittance, no diffusion for now.
-				const float TransmittanceNoV = 1.0;	
-				const float3 SlabTransmittance = IsotropicMediumSlabTransmittance(PM, STRATA_SIMPLEVOLUME_THICKNESS_M, TransmittanceNoV);
-
-				// STRATA_TODO Fresnel throughput, Use a LUT, this is currently in progress and will be coming soon.
-				// In the meantime, we use a simple approach to feed Fresnel function and compute an approximation of the general transmition/througput.
-				// We also ignored ThinFilm interaction as another simplification for perf.
-				const float3 ApproxFresnelForTransmittance = F_Schlick(F0, F90, BSDFContext.SatNoV);
-				const float3 SpecularThroughput = 1.0f - ApproxFresnelForTransmittance; // This is already taken into account by the Chan model. But for two sided lighting we need to apply it.
+				// We consider single scattering and perpendicular to the surface light transmittance.
+				// Using BackFaceNoL instead of 1 for TransmittanceNoL would be more correct but it gives some high frequency details looking weird without multiple scattering.
+				const float TransmittanceNoL = 1.0f;
+				const float3 SlabTransmittance = IsotropicMediumSlabTransmittance(PM, STRATA_SIMPLEVOLUME_THICKNESS_M, TransmittanceNoL);
 
 				// We recover the transmitted luminance from as the back face Lambert affected by the transmittance.
 				// The MFP is specified on the slab node and then transmittance recoved from it.
 				// We assume the back face diffuse albedo is white=1 and no diffusion happens at this stage.
 				const float3 BackFaceWhiteDiffuseColor = float3(1.0f, 1.0f, 1.0f);
 
-				StrataEnvLightResult.DiffuseBackFaceWeight = SlabTransmittance * BackFaceWhiteDiffuseColor * SpecularThroughput * ThinFilmFresnelWeight;
+				// The diffuse is attenuated by the slab transmittance and the SpecularTransmission through the GGX interface.
+				StrataEnvLightResult.DiffuseBackFaceWeight = Diffuse_Lambert(BackFaceWhiteDiffuseColor) * SlabTransmittance * DirectionalAlbedo_SpecularTransmission;
 
-				// At this stage, StrataEnvLightResult.bSubsurface should have been setup correctly already if SSS is enabled with two-sided lighting.
+				// At this stage, StrataEnvLightResult.bSubsurface should have been setup correctly already if SSS is enabled to be applied on a thin surface.
 			}
 			#endif
 
