Fixed double specular transmision on Strata simple volume.

diff --git a/Engine/Shaders/Private/Strata/StrataEvaluation.ush b/Engine/Shaders/Private/Strata/StrataEvaluation.ush
--- a/Engine/Shaders/Private/Strata/StrataEvaluation.ush
+++ b/Engine/Shaders/Private/Strata/StrataEvaluation.ush
@@ -453,7 +453,7 @@ FStrataEvaluateResult StrataEvaluateBSDF(FStrataBSDFContext BSDFContext, FShadow
 				// We assume the back face diffuse albedo is white=1 and no diffusion happens at this stage.
 				const float3 BackFaceWhiteDiffuseColor = float3(1.0f, 1.0f, 1.0f);
 
-				// The diffuse is attenuated by the slab transmittance and the SpecularTransmission through the GGX interface.
+				// The diffuse is attenuated by the slab transmittance and the GGX interface SpecularTransmission.
 				Sample.ThinTransmissionPathValue = Diffuse_Lambert(BackFaceWhiteDiffuseColor) * SlabTransmittance * DirectionalAlbedo_SpecularTransmission;
 				Sample.ThinTransmissionPathPDF = BackFaceNoLSat / PI;
 
@@ -472,6 +472,7 @@ FStrataEvaluateResult StrataEvaluateBSDF(FStrataBSDFContext BSDFContext, FShadow
 
 				const float Phase = IsotropicPhase();
 
+				// The diffuse and throughput account for the GGX interface SpecularTransmission.
 				Sample.DiffusePathValue = lerp(Sample.DiffusePathValue,		SlabDirectionalAlbedo * DirectionalAlbedo_SpecularTransmission,		DiffuseToVolumeBlend);
 				Sample.DiffusePDF		= lerp(Sample.DiffusePDF,			Phase,																DiffuseToVolumeBlend);
 				Sample.Throughput		= lerp(Sample.Throughput,			SlabTransmittance     * DirectionalAlbedo_SpecularTransmission,		DiffuseToVolumeBlend);
@@ -766,9 +767,10 @@ FStrataEnvLightResult StrataEvaluateForEnvLight(FStrataBSDFContext BSDFContext,
 				// We consider that the diffuse integration of the environment SH over the hemisphere multiplied with the light response will be a good approximation similar to the "split sum integral".
 				float3 SlabDirectionalAlbedo = IsotropicMediumSlabEnvDirectionalAlbedo(PM) * ((1 / PI)* (2 * PI));
 
-			#if STRATA_MICROFACET_MULTIPLESCATTERING == 0
+			#if 1
 				// Because the default EnvBRDF does not account for multiple scattering, we rescale the participating media response according to EnvBRDF.
 				// This is to ensure a monotone light response behavior when transitioning from the Diffuse to the Volumetric model (without that, multiple scattering can result in higher luminance while trnasitionning from diffuse to volume light model)
+				// STRATA_TODO: We should NOT have to do that. This should be investigated futher...
 				const float MaxDiffuseWeight = max3(StrataEnvLightResult.DiffuseWeight.x, StrataEnvLightResult.DiffuseWeight.y, StrataEnvLightResult.DiffuseWeight.z);
 				const float MaxSlabDirectionalAlbedo = max3(SlabDirectionalAlbedo.x, SlabDirectionalAlbedo.y, SlabDirectionalAlbedo.z);
 				if (MaxDiffuseWeight < MaxSlabDirectionalAlbedo)
@@ -776,14 +778,12 @@ FStrataEnvLightResult StrataEvaluateForEnvLight(FStrataBSDFContext BSDFContext,
 					// We only want to do that when the slab light response is higher than the diffse response, so that when the mfp is large (and scattering low), the light response still work correctly according to thickness and MFP.
 					SlabDirectionalAlbedo *= MaxDiffuseWeight / MaxSlabDirectionalAlbedo;
 				}
-			#else
-				// STRATA_TODO: Only considering the non hazy lob for now.
-				// STRATA_TODO: This is not enough to count the difference of energy when lerping from Chan to the directional albedo. This should be investigated.
-				FStrataGGXLUTResult GGXLUT = StrataEvaluateGGXLUT(DiffuseColor, F0, BSDFContext.Context.NoV, Metallic, SafeRoughnessX, SafeRoughnessX);
-				SlabDirectionalAlbedo *= GGXLUT.DirectionalAlbedo.SpecularT;
 			#endif
 
-				StrataEnvLightResult.DiffuseWeight = lerp(StrataEnvLightResult.DiffuseWeight, SlabDirectionalAlbedo * DirectionalAlbedo_SpecularTransmission,	DiffuseToVolumeBlend);
+				// The diffuse is attenuated by the GGX interface SpecularTransmission.
+				SlabDirectionalAlbedo *= DirectionalAlbedo_SpecularTransmission;
+
+				StrataEnvLightResult.DiffuseWeight = lerp(StrataEnvLightResult.DiffuseWeight, SlabDirectionalAlbedo, DiffuseToVolumeBlend);
 				// We keep the same normal
 				StrataEnvLightResult.bSubsurface = false;
 			}
@@ -806,7 +806,7 @@ FStrataEnvLightResult StrataEvaluateForEnvLight(FStrataBSDFContext BSDFContext,
 				// We assume the back face diffuse albedo is white=1 and no diffusion happens at this stage.
 				const float3 BackFaceWhiteDiffuseColor = float3(1.0f, 1.0f, 1.0f);
 
-				// The diffuse is attenuated by the slab transmittance and the SpecularTransmission through the GGX interface.
+				// The diffuse is attenuated by the slab transmittance and the GGX interface SpecularTransmission.
 				StrataEnvLightResult.DiffuseBackFaceWeight = Diffuse_Lambert(BackFaceWhiteDiffuseColor) * SlabTransmittance * DirectionalAlbedo_SpecularTransmission;
 
 				// At this stage, StrataEnvLightResult.bSubsurface should have been setup correctly already if SSS is enabled to be applied on a thin surface.
