Move voxel traversal seed generation to the traversal code.

diff --git a/Engine/Shaders/Private/HairStrands/HairStrandsDeepShadowMask.usf b/Engine/Shaders/Private/HairStrands/HairStrandsDeepShadowMask.usf
--- a/Engine/Shaders/Private/HairStrands/HairStrandsDeepShadowMask.usf
+++ b/Engine/Shaders/Private/HairStrands/HairStrandsDeepShadowMask.usf
@@ -422,15 +422,6 @@ float4 ComputeRandom4_0(uint2 PixelPosition, uint InSeed)
 	return float4(U0.x, U0.y, U1.x, U1.y);
 }
 
-float4 ComputeRandom4_1(uint2 PixelPosition, uint InSeed)
-{
-	const uint2 Seed0 = Rand3DPCG16(int3(PixelPosition, InSeed)).xy;
-	const uint2 Seed1 = Rand3DPCG16(int3(PixelPosition + 17, InSeed)).xy;
-	return float4(
-		Hammersley16(InSeed, 8, Seed0),
-		Hammersley16(InSeed, 8, Seed1));
-}
-
 float3 GetShadowMaskRandom(uint2 PixelPosition)
 {
 	float3 Random = 0.5f;
@@ -444,7 +435,7 @@ float3 GetShadowMaskRandom(uint2 PixelPosition)
 	}
 	else if (Voxel_RandomType == 2)
 	{
-		Random = ComputeRandom4_1(PixelPosition, View.StateFrameIndexMod8).xyz;
+		Random = GetHairVoxelJitter2(PixelPosition, View.StateFrameIndexMod8).xyz;
 	}
 	return Random;
 }
diff --git a/Engine/Shaders/Private/HairStrands/HairStrandsVoxelPageTraversal.ush b/Engine/Shaders/Private/HairStrands/HairStrandsVoxelPageTraversal.ush
--- a/Engine/Shaders/Private/HairStrands/HairStrandsVoxelPageTraversal.ush
+++ b/Engine/Shaders/Private/HairStrands/HairStrandsVoxelPageTraversal.ush
@@ -88,6 +88,13 @@ float3 GetHairVoxelJitter(uint2 PixelCoord, uint Seed)
 		InterleavedGradientNoise(PixelCoord.xy, Seed * 7901));
 }
 
+float4 GetHairVoxelJitter2(uint2 PixelCoord, uint Seed)
+{
+	const uint2 Seed0 = Rand3DPCG16(int3(PixelCoord, Seed)).xy;
+	const uint2 Seed1 = Rand3DPCG16(int3(PixelCoord + 17, Seed)).xy;
+	return float4(Hammersley16(Seed, 8, Seed0), Hammersley16(Seed, 8, Seed1));
+}
+
 // This assume pages have a resolution of 32x32x32
 bool DoesVoxelContainData(uint2 PageOccupancy, uint3 CoordInPage)
 {
