Clean unused return type in hair compute raster.

diff --git a/Engine/Shaders/Private/HairStrands/HairStrandsVisibilityRasterCompute.usf b/Engine/Shaders/Private/HairStrands/HairStrandsVisibilityRasterCompute.usf
--- a/Engine/Shaders/Private/HairStrands/HairStrandsVisibilityRasterCompute.usf
+++ b/Engine/Shaders/Private/HairStrands/HairStrandsVisibilityRasterCompute.usf
@@ -29,39 +29,24 @@ PackedType ToPackedType(uint2 Data)
 {
 	return PackUlongType(Data);
 }
-PackedType NextValue(PackedType In, PackedType Out)
-{
-	return In < Out ? In : Out;
-}
 #else
 #define PackedType uint2
 PackedType ToPackedType(uint2 Data)
 {
 	return Data;
 }
-PackedType NextValue(PackedType In, PackedType Out)
-{
-	return In.y < Out.y ? In : Out;
-}
 #endif
 
-PackedType WriteVisibilityData(RWTexture2D<PackedType> OutTexture, uint2 Coords, PackedType Data)
+void WriteVisibilityData(RWTexture2D<PackedType> OutTexture, uint2 Coords, PackedType Data)
 {
 #if PERMUTATION_RASTER_ATOMIC == RASTER_NVIDIA
-	PackedType PrevData = NvInterlockedMaxUint64(OutTexture, Coords, Data);
-	return PrevData;
+	NvInterlockedMaxUint64(OutTexture, Coords, Data);
 #elif PERMUTATION_RASTER_ATOMIC == RASTER_AMD_DX11
-	PackedType PrevData = AmdDxExtShaderIntrinsics_AtomicMaxU64(OutTexture, Coords, Data);
-	return PrevData;
+	AmdDxExtShaderIntrinsics_AtomicMaxU64(OutTexture, Coords, Data);
 #elif PERMUTATION_RASTER_ATOMIC == RASTER_AMD_DX12
-	PackedType PrevData = AmdExtD3DShaderIntrinsics_AtomicMaxU64(OutTexture, Coords, Data);
-	return PrevData;
+	AmdExtD3DShaderIntrinsics_AtomicMaxU64(OutTexture, Coords, Data);
 #elif PERMUTATION_RASTER_ATOMIC == RASTER_PLATFORM && COMPILER_SUPPORTS_UINT64_IMAGE_ATOMICS && COMPILER_SUPPORTS_ULONG_TYPES
 	ImageInterlockedMaxUInt64(OutTexture, Coords, Data);
-	// Need to expose the returned value version of ImageInterlockedMaxUInt64
-	return Data;
-#else
-	return (PackedType)0;
 #endif
 }
 
