Fix incorrect sampling coordinates.
Bug flagged up by [REDACTED_PERSON] in pull request #7995.
Pull request will go to 4.27 branch, but won't automatically integrate because code has moved to a new file.

diff --git a/Engine/Shaders/Private/BlockCompressionCommon.ush b/Engine/Shaders/Private/BlockCompressionCommon.ush
--- a/Engine/Shaders/Private/BlockCompressionCommon.ush
+++ b/Engine/Shaders/Private/BlockCompressionCommon.ush
@@ -179,7 +179,7 @@ void ReadBlockXYA(Texture2D<float4> SourceTexture, SamplerState TextureSampler,
 		float2 UVOffset = UV + float2(2.f * TexelUVSize.x, 0);
 		float4 Red = SourceTexture.GatherRed(TextureSampler, UVOffset);
 		float4 Green = SourceTexture.GatherGreen(TextureSampler, UVOffset);
-		float4 Blue = SourceTexture.GatherBlue(TextureSampler, UV);
+		float4 Blue = SourceTexture.GatherBlue(TextureSampler, UVOffset);
 		BlockX[2] = Red[3]; BlockY[2] = Green[3]; BlockA[2] = Blue[3];
 		BlockX[3] = Red[2]; BlockY[3] = Green[2]; BlockA[3] = Blue[2];
 		BlockX[6] = Red[0]; BlockY[6] = Green[0]; BlockA[6] = Blue[0];
@@ -189,7 +189,7 @@ void ReadBlockXYA(Texture2D<float4> SourceTexture, SamplerState TextureSampler,
 		float2 UVOffset = UV + float2(0, 2.f * TexelUVSize.y);
 		float4 Red = SourceTexture.GatherRed(TextureSampler, UVOffset);
 		float4 Green = SourceTexture.GatherGreen(TextureSampler, UVOffset);
-		float4 Blue = SourceTexture.GatherBlue(TextureSampler, UV);
+		float4 Blue = SourceTexture.GatherBlue(TextureSampler, UVOffset);
 		BlockX[8] = Red[3]; BlockY[8] = Green[3]; BlockA[8] = Blue[3];
 		BlockX[9] = Red[2]; BlockY[9] = Green[2]; BlockA[9] = Blue[2];
 		BlockX[12] = Red[0]; BlockY[12] = Green[0]; BlockA[12] = Blue[0];
@@ -199,7 +199,7 @@ void ReadBlockXYA(Texture2D<float4> SourceTexture, SamplerState TextureSampler,
 		float2 UVOffset = UV + 2.f * TexelUVSize;
 		float4 Red = SourceTexture.GatherRed(TextureSampler, UVOffset);
 		float4 Green = SourceTexture.GatherGreen(TextureSampler, UVOffset);
-		float4 Blue = SourceTexture.GatherBlue(TextureSampler, UV);
+		float4 Blue = SourceTexture.GatherBlue(TextureSampler, UVOffset);
 		BlockX[10] = Red[3]; BlockY[10] = Green[3]; BlockA[10] = Blue[3];
 		BlockX[11] = Red[2]; BlockY[11] = Green[2]; BlockA[11] = Blue[2];
 		BlockX[14] = Red[0]; BlockY[14] = Green[0]; BlockA[14] = Blue[0];
