Disable VT LOD randomization in 2D texture preview pixel shader. Works around a crash due to View uniform buffer not being set.

diff --git a/Engine/Shaders/Private/SimpleElementTexture2DPreviewPixelShader.usf b/Engine/Shaders/Private/SimpleElementTexture2DPreviewPixelShader.usf
--- a/Engine/Shaders/Private/SimpleElementTexture2DPreviewPixelShader.usf
+++ b/Engine/Shaders/Private/SimpleElementTexture2DPreviewPixelShader.usf
@@ -18,6 +18,9 @@
 
 #if SAMPLE_VIRTUAL_TEXTURE
 	#define NUM_VIRTUALTEXTURE_SAMPLES 1
+	// Workaround for unbound View uniform buffer which is used for LOD randomization.
+	// View.StateFrameIndexMod8 is used as random seed.
+	#define VT_RANDOMIZE_LOD 0
 	#include "VirtualTextureCommon.ush"
 #endif
 
diff --git a/Engine/Shaders/Private/VirtualTextureCommon.ush b/Engine/Shaders/Private/VirtualTextureCommon.ush
--- a/Engine/Shaders/Private/VirtualTextureCommon.ush
+++ b/Engine/Shaders/Private/VirtualTextureCommon.ush
@@ -5,6 +5,10 @@
 #include "Common.ush"
 #include "Random.ush"
 
+#ifndef VT_RANDOMIZE_LOD
+#define VT_RANDOMIZE_LOD 1
+#endif
+
 #ifndef NUM_VIRTUALTEXTURE_SAMPLES
 #error NUM_VIRTUALTEXTURE_SAMPLES should be set to the number of VT samples that will be done before including this file (exluding any lightmap samples)
 #endif
@@ -208,7 +212,12 @@ int TextureComputeVirtualMipLevel(
 	// This way the VT mip level selection will come much closer to HW mip selection, even if we're not sampling the texture using anisotropic filtering
 	const float ComputedLevel = MipLevelAniso2D(OutResult.dUVdx, OutResult.dUVdy, PageTableUniform.MaxAnisoLog2);
 
+#if VT_RANDOMIZE_LOD
 	const float Noise = InterleavedGradientNoise(SvPositionXY, View.StateFrameIndexMod8);
+#else
+	const float Noise = 0;
+#endif
+
 	return (int)floor(ComputedLevel + MipBias + Noise) + int(PageTableUniform.vPageTableMipBias);
 }
 
