Fix for FROST-1976 Object normals are wrong in ray tracing effects

diff --git a/Engine/Shaders/Private/LocalVertexFactory.ush b/Engine/Shaders/Private/LocalVertexFactory.ush
--- a/Engine/Shaders/Private/LocalVertexFactory.ush
+++ b/Engine/Shaders/Private/LocalVertexFactory.ush
@@ -189,7 +189,7 @@ FVertexFactoryInput LoadVertexFactoryInputForHGS(uint TriangleIndex, int VertexI
 #if GPUCULL_TODO
 	// Note: GetInstanceUserData() stores the GPU-Scene primitive ID
 	Input.InstanceIdOffset = GetInstanceUserData() | VF_TREAT_INSTANCE_ID_OFFSET_AS_PRIMITIVE_ID_FLAG;
-	Input.DrawInstanceId = 0;
+	Input.DrawInstanceId = InstanceIndex() - GetBaseInstanceIndex();
 #else // !GPUCULL_TODO
 	// GPUCULL_TODO: Need to work out what this is meant to do
 	Input.PrimitiveId = GetInstanceUserData();
@@ -200,10 +200,6 @@ FVertexFactoryInput LoadVertexFactoryInputForHGS(uint TriangleIndex, int VertexI
 	Input.InstanceId = InstanceIndex() - GetBaseInstanceIndex();
 #endif
 
-#if USE_INSTANCE_CULLING
-	Input.DrawInstanceId = InstanceIndex() - GetBaseInstanceIndex();
-#endif
-
 	return Input;
 }
 #endif
