Fix shader compilation warnings.

diff --git a/Engine/Shaders/Private/HairStrands/HairStrandsVoxelPageTraversal.ush b/Engine/Shaders/Private/HairStrands/HairStrandsVoxelPageTraversal.ush
--- a/Engine/Shaders/Private/HairStrands/HairStrandsVoxelPageTraversal.ush
+++ b/Engine/Shaders/Private/HairStrands/HairStrandsVoxelPageTraversal.ush
@@ -83,9 +83,9 @@ bool ShouldStopTraversal(const FHairTraversalResult InResult, float InHairCountT
 float3 GetHairVoxelJitter(uint2 PixelCoord, uint Seed)
 {
 	return float3(
-		InterleavedGradientNoise(PixelCoord.xy, Seed),
-		InterleavedGradientNoise(PixelCoord.xy, Seed * 117),
-		InterleavedGradientNoise(PixelCoord.xy, Seed * 7901));
+		InterleavedGradientNoise(float2(PixelCoord.xy), float(Seed)),
+		InterleavedGradientNoise(float2(PixelCoord.xy), float(Seed * 117.f)),
+		InterleavedGradientNoise(float2(PixelCoord.xy), float(Seed * 7901.f)));
 }
 
 float4 GetHairVoxelJitter2(uint2 PixelCoord, uint Seed)
@@ -194,7 +194,7 @@ FHairTraversalResult ComputeHairCountVirtualVoxel(
 
 		// Step according to voxel size
 		const float3 D = normalize(E - O) * InCommonDesc.VoxelWorldSize;
-		const float MaxStep = float(min(ceil(OELength / InCommonDesc.VoxelWorldSize), 1024u));
+		const float MaxStep = float(min(ceil(OELength / InCommonDesc.VoxelWorldSize), 1024.f));
 		const float DeltaWorld = OELength / float(MaxStep);
 		float StepScale = 1.0f;
 		float3 PreviousP = O;
@@ -302,7 +302,7 @@ FHairTraversalResult ComputeHairCountVirtualVoxel(
 			}
 			else
 			{
-				StepScale = min(InCommonDesc.PageResolution, StepScale * InSettings.SteppingScale);
+				StepScale = min(float(InCommonDesc.PageResolution), StepScale * InSettings.SteppingScale);
 			}
 			#endif
 			PreviousP = HitP;
diff --git a/Engine/Shaders/Private/ShaderDrawDebug.ush b/Engine/Shaders/Private/ShaderDrawDebug.ush
--- a/Engine/Shaders/Private/ShaderDrawDebug.ush
+++ b/Engine/Shaders/Private/ShaderDrawDebug.ush
@@ -25,8 +25,8 @@ struct FPackedShaderDrawElement
 
 FPackedShaderDrawElement PackShaderElement(FShaderDrawElement E)
 {
-	uint4 PackedC0 = 255.0f * saturate(E.Color0);
-	uint4 PackedC1 = 255.0f * saturate(E.Color1);
+	uint4 PackedC0 = uint4(255.0f * saturate(E.Color0));
+	uint4 PackedC1 = uint4(255.0f * saturate(E.Color1));
 
 	FPackedShaderDrawElement PackedE;
 	PackedE.Pos0_ColorX = float4(E.Pos0, asfloat((PackedC0.x << 24) | (PackedC0.y << 16) | (PackedC0.z << 8) | (PackedC0.w)));
