Fix a bug of spotlight shadow because we use reverseZ in ShadowDepth.

diff --git a/Engine/Shaders/Private/MobileBasePassPixelShader.usf b/Engine/Shaders/Private/MobileBasePassPixelShader.usf
--- a/Engine/Shaders/Private/MobileBasePassPixelShader.usf
+++ b/Engine/Shaders/Private/MobileBasePassPixelShader.usf
@@ -427,7 +427,9 @@ inout half3 Color)
 			{
 				// Clamp pixel depth in light space for shadowing opaque, because areas of the shadow depth buffer that weren't rendered to will have been cleared to 1
 				// We want to force the shadow comparison to result in 'unshadowed' in that case, regardless of whether the pixel being shaded is in front or behind that plane
-				float LightSpacePixelDepthForOpaque = min(HomogeneousShadowPosition.z, 0.99999f);
+				// Invert ShadowZ as the shadow space has been changed (but not yet the filtering code)
+				float ShadowZ = 1.0f - HomogeneousShadowPosition.z;
+				float LightSpacePixelDepthForOpaque = min(ShadowZ, 0.99999f);
 				Settings.SceneDepth = LightSpacePixelDepthForOpaque;
 				Settings.TransitionScale = SpotLightAnglesAndSoftTransitionScaleAndLightShadowType.z;
 
