Fix primary ray length in RT Translucency

Refactoring the primary ray generation for the path tracer changed the exact origin of the ray which cause the existing heuristic to go a bit too far and hit gbuffer geometry.

diff --git a/Engine/Shaders/Private/RayTracing/RayTracingPrimaryRays.usf b/Engine/Shaders/Private/RayTracing/RayTracingPrimaryRays.usf
--- a/Engine/Shaders/Private/RayTracing/RayTracingPrimaryRays.usf
+++ b/Engine/Shaders/Private/RayTracing/RayTracingPrimaryRays.usf
@@ -105,16 +105,15 @@ RAY_TRACING_ENTRY_RAYGEN(RayTracingPrimaryRaysRGS)
 
 	float Depth = GBufferData.Depth;
 	float3 WorldPosition = ReconstructWorldPositionFromDepth(UV, Depth);
-	float WorldSpaceDistance = length(WorldPosition - View.WorldViewOrigin);
 
-	// Trace rays from camera origin to (Gbuffer - epsilon) to only intersect translucent objects
+	// Trace rays from camera origin to Gbuffer (only translucent objects are considered on the first ray, so a bias is not needed)
 	RayDesc Ray = CreatePrimaryRay(UV);
 	FRayCone RayCone = (FRayCone)0;
 	RayCone.SpreadAngle = View.EyeToPixelSpreadAngle;
 
 	if((ERayTracingPrimaryRaysFlag_UseGBufferForMaxDistance & PrimaryRayFlags) != 0) 
 	{
-		Ray.TMax = WorldSpaceDistance - 0.1;
+		Ray.TMax = dot(WorldPosition - Ray.Origin, Ray.Direction); // assumes Ray.Direction is normalized
 	}
 
 	bool bAllowSkySampling;
@@ -142,7 +141,7 @@ RAY_TRACING_ENTRY_RAYGEN(RayTracingPrimaryRaysRGS)
 	for (uint RefractionRayIndex = 0; RefractionRayIndex < MaxRefractionRays; ++RefractionRayIndex)
 	{
 		const uint RefractionRayFlags = RefractionRayIndex == 0? RAY_FLAG_CULL_BACK_FACING_TRIANGLES : 0;
-		const uint RefractionInstanceInclusionMask = RAY_TRACING_MASK_ALL;
+		const uint RefractionInstanceInclusionMask = RefractionRayIndex == 0 ? RAY_TRACING_MASK_TRANSLUCENT : RAY_TRACING_MASK_ALL;
 		const bool bRefractionRayTraceSkyLightContribution = false;
 		const bool bRefractionDecoupleSampleGeneration = true;
 		const bool bRefractionEnableSkyLightContribution = true;
