Add in FeedbackPos to the VT index that generates feedback.  Helps ensure that different pixels from the same material output feedback for different VTs

diff --git a/Engine/Shaders/Private/VirtualTextureCommon.ush b/Engine/Shaders/Private/VirtualTextureCommon.ush
--- a/Engine/Shaders/Private/VirtualTextureCommon.ush
+++ b/Engine/Shaders/Private/VirtualTextureCommon.ush
@@ -34,8 +34,10 @@ struct FVirtualTextureFeedbackParams
 void InitializeVirtualTextureFeedback(in out FVirtualTextureFeedbackParams Params, uint2 SvPosition, uint FrameNumber)
 {
 #if (NUM_VIRTUALTEXTURE_SAMPLES + LIGHTMAP_VT_ENABLED) > 1
-	uint NumVTSamplesInShader = NUM_VIRTUALTEXTURE_SAMPLES + LIGHTMAP_VT_ENABLED;
-	Params.RequestId = View.VirtualTextureFeedbackSampleOffset % NumVTSamplesInShader;
+	const uint NumVTSamplesInShader = NUM_VIRTUALTEXTURE_SAMPLES + LIGHTMAP_VT_ENABLED;
+	const uint2 PixelPos = SvPosition / VIRTUAL_TEXTURE_FEEDBACK_FACTOR;
+	const uint FeedbackPos = PixelPos.y * View.VirtualTextureFeedbackStride + PixelPos.x;
+	Params.RequestId = (View.VirtualTextureFeedbackSampleOffset + FeedbackPos) % NumVTSamplesInShader;
 	Params.Request = 0xFFFFFFFF;
 #else
 	Params.Request = 0xFFFFFFFF;	
