Fix flickering in RT translucent pass when refraction is disabled

Rays that are not bent by refraction should only look for translucent surfaces

diff --git a/Engine/Shaders/Private/RayTracing/RayTracingPrimaryRays.usf b/Engine/Shaders/Private/RayTracing/RayTracingPrimaryRays.usf
--- a/Engine/Shaders/Private/RayTracing/RayTracingPrimaryRays.usf
+++ b/Engine/Shaders/Private/RayTracing/RayTracingPrimaryRays.usf
@@ -141,7 +141,7 @@ RAY_TRACING_ENTRY_RAYGEN(RayTracingPrimaryRaysRGS)
 	for (uint RefractionRayIndex = 0; RefractionRayIndex < MaxRefractionRays; ++RefractionRayIndex)
 	{
 		const uint RefractionRayFlags = RefractionRayIndex == 0? RAY_FLAG_CULL_BACK_FACING_TRIANGLES : 0;
-		const uint RefractionInstanceInclusionMask = RefractionRayIndex == 0 ? RAY_TRACING_MASK_TRANSLUCENT : RAY_TRACING_MASK_ALL;
+		const uint RefractionInstanceInclusionMask = RefractionRayIndex == 0 || !TranslucencyRefraction ? RAY_TRACING_MASK_TRANSLUCENT : RAY_TRACING_MASK_ALL;
 		const bool bRefractionRayTraceSkyLightContribution = false;
 		const bool bRefractionDecoupleSampleGeneration = true;
 		const bool bRefractionEnableSkyLightContribution = true;
