Fixed edge color being broken by fuzz on slab node CL.

diff --git a/Engine/Shaders/Private/Strata/Strata.ush b/Engine/Shaders/Private/Strata/Strata.ush
--- a/Engine/Shaders/Private/Strata/Strata.ush
+++ b/Engine/Shaders/Private/Strata/Strata.ush
@@ -245,8 +245,8 @@ struct FStrataBSDF
 #define SLAB_THINFILM_THICKNESS(X)		X.VGPRs[2].w
 #define SLAB_THINFILM_IOR(X)			1.44f
 // Fuzz to simulate cloth
-#define SLAB_FUZZ_COLOR(X)				X.VGPRs[3].xyz
-#define SLAB_FUZZ_AMOUNT(X)				X.VGPRs[3].w
+#define SLAB_FUZZ_COLOR(X)				X.VGPRs[4].xyz
+#define SLAB_FUZZ_AMOUNT(X)				X.VGPRs[4].w
 // NOTE For SLAB, X.VGPRs[3].w is free for another use case.
 
 #define VOLUMETRICFOGCLOUD_ALBEDO(X)	X.VGPRs[0].xyz
