Fixed top layer normal resulting from horizontal mixing (weight was missing)

diff --git a/Engine/Shaders/Private/Strata/StrataMaterialClassification.usf b/Engine/Shaders/Private/Strata/StrataMaterialClassification.usf
--- a/Engine/Shaders/Private/Strata/StrataMaterialClassification.usf
+++ b/Engine/Shaders/Private/Strata/StrataMaterialClassification.usf
@@ -48,7 +48,9 @@ void ComputePostProcessBSDFData(
 	const bool bTopLayer = BSDF_GETISTOPLAYER(BSDF);
 	if (bTopLayer)
 	{
-		TopLayer.WorldNormal += N;
+		// We must take into account the weight in case horizontal blending is used,
+		// together with a threshold to avoid 0-length normal after normalization in case influence is fadeout using the Weight node.
+		TopLayer.WorldNormal += N * max(1e-5, BSDFWeightAvg); 
 	}
 
 	Classification.ShadingModels |= StrataShadingModelBit(BSDF_GETTYPE(BSDF));
