Optimises stencil and parallax rejection factor fetches

diff --git a/Engine/Shaders/Private/TemporalAA/TAAUpdateHistory.usf b/Engine/Shaders/Private/TemporalAA/TAAUpdateHistory.usf
--- a/Engine/Shaders/Private/TemporalAA/TAAUpdateHistory.usf
+++ b/Engine/Shaders/Private/TemporalAA/TAAUpdateHistory.usf
@@ -160,7 +160,9 @@ void MainCS(
 	taa_half LowFrequencyRejection = taa_half(1.0);
 	taa_half TranslucencyRejection = taa_half(1.0);
 	taa_half OutputPixelVelocity = taa_half(0.0);
-	#if 1
+	taa_subpixel_payload ParallaxFactorBits = 0x0;
+	bool bIsResponsiveAAPixel = false;
+
 	ISOLATE
 	{
 		// Vector in pixel between pixel K -> O.
@@ -168,6 +170,7 @@ void MainCS(
 
 		float2 EncodedVelocity = DilatedVelocityTexture[InputPixelPos];
 		ParallaxRejectionMask = ParallaxRejectionMaskTexture[InputPixelPos];
+		ParallaxFactorBits = ParallaxFactorTexture[InputPixelPos] & SUB_PIXEL_PARALLAX_FACTOR_BIT_MASK;
 
 		float2 ScreenVelocity = DecodeVelocityFromTexture(float4(EncodedVelocity, 0.0, 0.0)).xy;
 
@@ -192,6 +195,17 @@ void MainCS(
 		}
 
 		TranslucencyRejection = max(TranslucencyRejection, taa_half(MinTranslucencyRejection));
+		
+		// Fetch wehther the pixel is responsive AA or not.
+		#if CONFIG_RESPONSIVE_STENCIL
+		{
+			const uint kResponsiveStencilMask = 1 << 3;
+			
+			uint SceneStencilRef = InputSceneStencilTexture.Load(int3(InputPixelPos, 0)) STENCIL_COMPONENT_SWIZZLE;
+
+			bIsResponsiveAAPixel = (SceneStencilRef & kResponsiveStencilMask) != 0;
+		}
+		#endif
 
 		#if !CONFIG_CLAMP
 		{
@@ -201,24 +215,11 @@ void MainCS(
 		}
 		#endif
 	}
-	#endif
 
 	// Reject the translucency to respond quicker on firefliers.
 	#if CONFIG_TRANSLUCENCY_REJECTION
 		LowFrequencyRejection = min(LowFrequencyRejection, TranslucencyRejection);
 	#endif
-
-	// Fetch wehther the pixel is responsive AA or not.
-	bool bIsResponsiveAAPixel = false;
-	#if CONFIG_RESPONSIVE_STENCIL
-	{
-		const uint kResponsiveStencilMask = 1 << 3;
-			
-		uint SceneStencilRef = InputSceneStencilTexture.Load(int3(InputPixelPos, 0)) STENCIL_COMPONENT_SWIZZLE;
-
-		bIsResponsiveAAPixel = (SceneStencilRef & kResponsiveStencilMask) != 0;
-	}
-	#endif
 	
 	// Detect if HistoryBufferUV would be outside of the viewport.
 	bool bOffScreen = IsOffScreen(bCameraCut, PrevScreenPos, ParallaxRejectionMask);
@@ -676,8 +677,6 @@ void MainCS(
 
 		FinalSubpixelDetails = PrevSubpixelDetails;
 
-		// TODO Save VGPR here
-		taa_subpixel_payload ParallaxFactorBits = ParallaxFactorTexture[InputPixelPos] & SUB_PIXEL_PARALLAX_FACTOR_BIT_MASK;
 
 		{
 			const uint ParallaxFactorMask = (
