Fix DrawInstanceId indexing with ray tracing static mesh auto-instancing

diff --git a/Engine/Shaders/Private/LocalVertexFactory.ush b/Engine/Shaders/Private/LocalVertexFactory.ush
--- a/Engine/Shaders/Private/LocalVertexFactory.ush
+++ b/Engine/Shaders/Private/LocalVertexFactory.ush
@@ -188,8 +188,19 @@ FVertexFactoryInput LoadVertexFactoryInputForHGS(uint TriangleIndex, int VertexI
 #if VF_USE_PRIMITIVE_SCENE_DATA
 #if GPUCULL_TODO
 	// Note: GetInstanceUserData() stores the GPU-Scene primitive ID
-	Input.InstanceIdOffset = GetInstanceUserData() | VF_TREAT_INSTANCE_ID_OFFSET_AS_PRIMITIVE_ID_FLAG;
-	Input.DrawInstanceId = InstanceIndex() - GetBaseInstanceIndex();
+	int PrimitiveId = GetInstanceUserData();
+	Input.InstanceIdOffset = PrimitiveId | VF_TREAT_INSTANCE_ID_OFFSET_AS_PRIMITIVE_ID_FLAG;
+	// We have a unified path for both static meshes and instanced static meshes
+	// However due to static meshes being auto-instanced GetBaseInstanceIndex() may not match InstanceIndex() and a non-zero DrawInstanceId is generated
+	// Force DrawInstanceId to 0 when it is a static mesh by checking NumInstanceDataEntries
+	if (GetPrimitiveData(PrimitiveId).NumInstanceDataEntries == 1)
+	{
+		Input.DrawInstanceId = 0;
+	}
+	else
+	{
+		Input.DrawInstanceId = InstanceIndex() - GetBaseInstanceIndex();
+	}
 #else // !GPUCULL_TODO
 	// GPUCULL_TODO: Need to work out what this is meant to do
 	Input.PrimitiveId = GetInstanceUserData();
