Change VSM projection hair permutation to be used as a guard when hair strands are disabled.

diff --git a/Engine/Shaders/Private/VirtualShadowMaps/VirtualShadowMapProjection.usf b/Engine/Shaders/Private/VirtualShadowMaps/VirtualShadowMapProjection.usf
--- a/Engine/Shaders/Private/VirtualShadowMaps/VirtualShadowMapProjection.usf
+++ b/Engine/Shaders/Private/VirtualShadowMaps/VirtualShadowMapProjection.usf
@@ -268,84 +268,87 @@ FProjectLightResult ProjectLight(
 
 		// Hair strands voxel traversal to apply hair shadow on top of opaque geometry
 		#if HAS_HAIR_STRANDS
-		float3 RayStart = 0;
-		float3 RayEnd = 0;
-		bool bCastHairRay = false;
-		float4 Random = 0;
-		if (ShadowFactor > 0)
+		if (!bIsHair)
 		{
-			Random = ComputeRandom4(PixelPos);
-			uint RayIndex = min(Random.w * SMRTRayCount, SMRTRayCount - 1);
-
-		#if DIRECTIONAL_LIGHT
-			bCastHairRay = GenerateRayDirectional(
-				Light,
-				PixelPos,
-				TranslatedWorldPosition,
-				ContactShadowLengthWorld,
-				RayIndex,
-				SMRTRayCount,
-				RayStart,
-				RayEnd);
-		#else
-			bCastHairRay = GenerateRayLocalLight(
-				Light,
-				PixelPos,
-				TranslatedWorldPosition,
-				SMRTCotMaxRayAngleFromLight,
-				RayIndex,
-				SMRTRayCount,
-				RayStart,
-				RayEnd);
-		#endif
-		}
-
-		if (bCastHairRay)
-		{
-			// Jitter start position to mask/compensate coarse voxel resolution
+			float3 RayStart = 0;
+			float3 RayEnd = 0;
+			bool bCastHairRay = false;
+			float4 Random = 0;
+			if (ShadowFactor > 0)
 			{
-				const float PositionBiasScale = 0.5f;
-				const float3 DepthBias = VirtualVoxel.VoxelWorldSize * (VirtualVoxel.DepthBiasScale_Shadow * L + PositionBiasScale * (Random.xyz * 2 - 1));
-				RayStart += DepthBias;
+				Random = ComputeRandom4(PixelPos);
+				uint RayIndex = min(Random.w * SMRTRayCount, SMRTRayCount - 1);
+
+			#if DIRECTIONAL_LIGHT
+				bCastHairRay = GenerateRayDirectional(
+					Light,
+					PixelPos,
+					TranslatedWorldPosition,
+					ContactShadowLengthWorld,
+					RayIndex,
+					SMRTRayCount,
+					RayStart,
+					RayEnd);
+			#else
+				bCastHairRay = GenerateRayLocalLight(
+					Light,
+					PixelPos,
+					TranslatedWorldPosition,
+					SMRTCotMaxRayAngleFromLight,
+					RayIndex,
+					SMRTRayCount,
+					RayStart,
+					RayEnd);
+			#endif
 			}
 
-			const float DistanceThreshold = 100000.0f;
-			const float CoverageThreshold = 0.995f; // When Coverage is high, we do not trace shadow on opaque since hair/fur is covering the background.
-
-			FVirtualVoxelCommonDesc CommonDesc;
-			CommonDesc.PageCountResolution		= VirtualVoxel.PageCountResolution;
-			CommonDesc.VoxelWorldSize			= VirtualVoxel.VoxelWorldSize;
-			CommonDesc.PageTextureResolution	= VirtualVoxel.PageTextureResolution;
-			CommonDesc.PageResolution			= VirtualVoxel.PageResolution;
-
-			FHairTraversalSettings TraversalSettings = InitHairTraversalSettings();
-			TraversalSettings.DensityScale		= VirtualVoxel.DensityScale_Shadow;
-			TraversalSettings.CountThreshold	= 1;
-			TraversalSettings.DistanceThreshold = DistanceThreshold;
-			TraversalSettings.bDebugEnabled		= false;
-			TraversalSettings.SteppingScale		= VirtualVoxel.SteppingScale_Shadow;
-			TraversalSettings.Random			= Random.xyz;
-			TraversalSettings.PixelRadius		= SceneDepth * VirtualVoxel.HairCoveragePixelRadiusAtDepth1;
-			TraversalSettings.bUseOpaqueVisibility = false;
-
-			const uint VoxelDescCount = VirtualVoxel.NodeDescCount;
-			for (uint VoxelDescIt=0; VoxelDescIt<VoxelDescCount; ++VoxelDescIt)
+			if (bCastHairRay)
 			{
-				const FPackedVirtualVoxelNodeDesc PackedNode = VirtualVoxel.NodeDescBuffer[VoxelDescIt];
-				const FVirtualVoxelNodeDesc NodeDesc = UnpackVoxelNode(PackedNode, VirtualVoxel.PageResolution);
-
-				FHairTraversalResult Result = InitHairTraversalResult();
-				Result = ComputeHairCountVirtualVoxel(
-					RayStart,
-					RayEnd,
-					CommonDesc,
-					NodeDesc,
-					VirtualVoxel.PageIndexBuffer,
-					VirtualVoxel.PageIndexOccupancyBuffer,
-					VirtualVoxel.PageTexture,
-					TraversalSettings);
-
-				ShadowFactor = min(ShadowFactor, saturate(1 - Result.HairCount));
+				// Jitter start position to mask/compensate coarse voxel resolution
+				{
+					const float PositionBiasScale = 0.5f;
+					const float3 DepthBias = VirtualVoxel.VoxelWorldSize * (VirtualVoxel.DepthBiasScale_Shadow * L + PositionBiasScale * (Random.xyz * 2 - 1));
+					RayStart += DepthBias;
+				}
+
+				const float DistanceThreshold = 100000.0f;
+				const float CoverageThreshold = 0.995f; // When Coverage is high, we do not trace shadow on opaque since hair/fur is covering the background.
+
+				FVirtualVoxelCommonDesc CommonDesc;
+				CommonDesc.PageCountResolution		= VirtualVoxel.PageCountResolution;
+				CommonDesc.VoxelWorldSize			= VirtualVoxel.VoxelWorldSize;
+				CommonDesc.PageTextureResolution	= VirtualVoxel.PageTextureResolution;
+				CommonDesc.PageResolution			= VirtualVoxel.PageResolution;
+
+				FHairTraversalSettings TraversalSettings = InitHairTraversalSettings();
+				TraversalSettings.DensityScale		= VirtualVoxel.DensityScale_Shadow;
+				TraversalSettings.CountThreshold	= 1;
+				TraversalSettings.DistanceThreshold = DistanceThreshold;
+				TraversalSettings.bDebugEnabled		= false;
+				TraversalSettings.SteppingScale		= VirtualVoxel.SteppingScale_Shadow;
+				TraversalSettings.Random			= Random.xyz;
+				TraversalSettings.PixelRadius		= SceneDepth * VirtualVoxel.HairCoveragePixelRadiusAtDepth1;
+				TraversalSettings.bUseOpaqueVisibility = false;
+
+				const uint VoxelDescCount = VirtualVoxel.NodeDescCount;
+				for (uint VoxelDescIt=0; VoxelDescIt<VoxelDescCount; ++VoxelDescIt)
+				{
+					const FPackedVirtualVoxelNodeDesc PackedNode = VirtualVoxel.NodeDescBuffer[VoxelDescIt];
+					const FVirtualVoxelNodeDesc NodeDesc = UnpackVoxelNode(PackedNode, VirtualVoxel.PageResolution);
+
+					FHairTraversalResult Result = InitHairTraversalResult();
+					Result = ComputeHairCountVirtualVoxel(
+						RayStart,
+						RayEnd,
+						CommonDesc,
+						NodeDesc,
+						VirtualVoxel.PageIndexBuffer,
+						VirtualVoxel.PageIndexOccupancyBuffer,
+						VirtualVoxel.PageTexture,
+						TraversalSettings);
+
+					ShadowFactor = min(ShadowFactor, saturate(1 - Result.HairCount));
+				}
 			}
 		}
 		#endif
@@ -372,6 +375,7 @@ void VirtualShadowMapProjection(
 	}
 	
 	float DeviceZ = SceneTexturesStruct.SceneDepthTexture.Load( int3( PixelPos, 0 ) ).r;
+	#if HAS_HAIR_STRANDS
 	if (InputType == INPUT_TYPE_HAIRSTRANDS)
 	{
 		FCategorizationData CatData = DecodeCategorizationData(HairStrands.HairCategorizationTexture.Load( int3(PixelPos, 0 ) ) );
@@ -381,6 +385,7 @@ void VirtualShadowMapProjection(
 			return;
 		}
 	}
+	#endif
 	const float SceneDepth = ConvertFromDeviceZ( DeviceZ );
 
 	const float4 SvPosition = float4( float2( PixelPos ) + 0.5, DeviceZ, 1.0f );
diff --git a/Engine/Source/Runtime/Renderer/Private/VirtualShadowMaps/VirtualShadowMapProjection.cpp b/Engine/Source/Runtime/Renderer/Private/VirtualShadowMaps/VirtualShadowMapProjection.cpp
--- a/Engine/Source/Runtime/Renderer/Private/VirtualShadowMaps/VirtualShadowMapProjection.cpp
+++ b/Engine/Source/Runtime/Renderer/Private/VirtualShadowMaps/VirtualShadowMapProjection.cpp
@@ -258,7 +258,7 @@ static void RenderVirtualShadowMapProjectionCommon(
 	int32 VirtualShadowMapId = INDEX_NONE)
 {
 	// Use hair strands data (i.e., hair voxel tracing) only for Gbuffer input for casting hair shadow onto opaque geometry.
-	const bool bHasHairStrandsData = InputType == EVirtualShadowMapProjectionInputType::GBuffer && HairStrands::HasViewHairStrandsData(View);
+	const bool bHasHairStrandsData = HairStrands::HasViewHairStrandsData(View);
 
 	FVirtualShadowMapProjectionCS::FParameters* PassParameters = GraphBuilder.AllocParameters< FVirtualShadowMapProjectionCS::FParameters >();
 	PassParameters->SamplingParameters = VirtualShadowMapArray.GetSamplingParameters(GraphBuilder);
@@ -269,9 +269,9 @@ static void RenderVirtualShadowMapProjectionCommon(
 	PassParameters->ContactShadowLength = CVarContactShadowLength.GetValueOnRenderThread();
 	PassParameters->NormalBias = GetNormalBiasForShader();
 	PassParameters->InputType = uint32(InputType);
-	PassParameters->HairStrands = HairStrands::BindHairStrandsViewUniformParameters(View);
 	if (bHasHairStrandsData)
 	{
+		PassParameters->HairStrands = HairStrands::BindHairStrandsViewUniformParameters(View);
 		PassParameters->HairStrandsVoxel = HairStrands::BindHairStrandsVoxelUniformParameters(View);
 	}
 
