Made color<->HSV space conversion common.

diff --git a/Engine/Shaders/Private/ColorMap.ush b/Engine/Shaders/Private/ColorMap.ush
--- a/Engine/Shaders/Private/ColorMap.ush
+++ b/Engine/Shaders/Private/ColorMap.ush
@@ -78,3 +78,15 @@ float3 ColorMapTurbo(float t)
     t = saturate(t);
     return c0+t*(c1+t*(c2+t*(c3+t*(c4+t*(c5+t*c6)))));
 }
+
+/*=============================================================================================
+	Additional color maps (not coming from the ShaderToy mentioned above)
+===============================================================================================*/
+
+#include "ColorSpace.ush"  
+
+float3 GetHSVDebugColor(float x)
+{
+	float c = (1 - saturate(x)) * 0.6; // Remap [0,1] to Blue-Red
+	return x > 0 ? HSV_2_LinearRGB(float3(c, 1, 1)) : float3(0, 0, 0);
+}
diff --git a/Engine/Shaders/Private/ColorSpace.ush b/Engine/Shaders/Private/ColorSpace.ush
--- a/Engine/Shaders/Private/ColorSpace.ush
+++ b/Engine/Shaders/Private/ColorSpace.ush
@@ -188,3 +188,21 @@ float3 LCoCg_2_LinearRGB(float3 LCoCg)
 {
 	return YCoCg_2_LinearRGB(LCoCg_2_YCoCg(LCoCg));
 }
+
+//------------------------------------------------------------------
+// HSV
+//------------------------------------------------------------------
+
+float3 HUE_2_LinearRGB(in float H)
+{
+	float R = abs(H * 6 - 3) - 1;
+	float G = 2 - abs(H * 6 - 2);
+	float B = 2 - abs(H * 6 - 4);
+	return saturate(float3(R, G, B));
+}
+
+float3 HSV_2_LinearRGB(in float3 HSV)
+{
+	float3 RGB = HUE_2_LinearRGB(HSV.x);
+	return ((RGB - 1) * HSV.y + 1) * HSV.z;
+}
diff --git a/Engine/Shaders/Private/HairStrands/HairScatter.usf b/Engine/Shaders/Private/HairStrands/HairScatter.usf
--- a/Engine/Shaders/Private/HairStrands/HairScatter.usf
+++ b/Engine/Shaders/Private/HairStrands/HairScatter.usf
@@ -9,7 +9,6 @@
 #define SUPPORT_CONTACT_SHADOWS 0
 
 #include "../SceneTexturesCommon.ush"
-#include "HairStrandsDebugCommon.ush"
 #include "HairStrandsDeepShadowCommon.ush"
 #include "../DeferredShadingCommon.ush"
 #include "../ShadingModels.ush"
diff --git a/Engine/Shaders/Private/HairStrands/HairStrandsBsdfPlot.usf b/Engine/Shaders/Private/HairStrands/HairStrandsBsdfPlot.usf
--- a/Engine/Shaders/Private/HairStrands/HairStrandsBsdfPlot.usf
+++ b/Engine/Shaders/Private/HairStrands/HairStrandsBsdfPlot.usf
@@ -6,7 +6,6 @@
 #include "../DeferredShadingCommon.ush"
 #include "../HairBsdf.ush"
 #include "HairStrandsCommon.ush"
-#include "HairStrandsDebugCommon.ush"
 #include "HairStrandsVisibilityCommon.ush"
 #include "HairStrandsEnvironmentLightingCommon.ush"
 #include "HairStrandsDeepShadowCommon.ush"
diff --git a/Engine/Shaders/Private/HairStrands/HairStrandsClusterCulling.usf b/Engine/Shaders/Private/HairStrands/HairStrandsClusterCulling.usf
--- a/Engine/Shaders/Private/HairStrands/HairStrandsClusterCulling.usf
+++ b/Engine/Shaders/Private/HairStrands/HairStrandsClusterCulling.usf
@@ -1,7 +1,7 @@
 // Copyright Epic Games, Inc. All Rights Reserved.
 
 #include "/Engine/Public/Platform.ush"
-#include "HairStrandsDebugCommon.ush"
+#include "../ColorMap.ush"
 #include "HairStrandsAABBCommon.ush"
 #include "HairStrandsClusterCommon.ush"
 
@@ -457,7 +457,7 @@ void MainDrawDebugAABBCS(uint3 DispatchThreadId : SV_DispatchThreadID)
 			}
 
 			float SelectedLod = float(DebugInfo.LOD) / D_LOD_DEBUG_SCALE;
-			Color = float4(GetColorCode(SelectedLod), 0.05f);
+			Color = float4(GetHSVDebugColor(SelectedLod), 0.05f);
 		}
 		#endif
 
diff --git a/Engine/Shaders/Private/HairStrands/HairStrandsDebug.usf b/Engine/Shaders/Private/HairStrands/HairStrandsDebug.usf
--- a/Engine/Shaders/Private/HairStrands/HairStrandsDebug.usf
+++ b/Engine/Shaders/Private/HairStrands/HairStrandsDebug.usf
@@ -2,7 +2,7 @@
 
 #include "../Common.ush"
 #include "HairStrandsVisibilityCommon.ush"
-#include "HairStrandsDebugCommon.ush"
+#include "../ColorMap.ush"
 
 float2 OutputResolution;
 uint FastResolveMask;
@@ -47,7 +47,7 @@ void MainPS(
 	//		discard;
 	//	OutColor = float4(float3(1, 0, 0), 1);
 	//
-	//	OutColor = float4(GetColorCode(HairCount / float(32)), 1);
+	//	OutColor = float4(GetHSVDebugColor(HairCount / float(32)), 1);
 	//
 	//	return;
 	//}
@@ -61,7 +61,7 @@ void MainPS(
 	float Alpha = 1;
 	if (DebugMode == DEBUG_MODE_SAMPLE_COUNT)
 	{
-		Color = GetColorCode(CatData.SampleCount/float(MaxSampleCount));
+		Color = GetHSVDebugColor(CatData.SampleCount/float(MaxSampleCount));
 	}
 
 	if (DebugMode == DEBUG_MODE_COVERAGE_TYPE)
@@ -79,7 +79,7 @@ void MainPS(
 
 	if (DebugMode == DEBUG_MODE_COVERAGE)
 	{
-		Color = GetColorCode(saturate(CatData.PixelCoverage));
+		Color = GetHSVDebugColor(saturate(CatData.PixelCoverage));
 	}
 
 	const bool bIsDebugModeMaterial =
diff --git a/Engine/Shaders/Private/HairStrands/HairStrandsDebugCommon.ush b/Engine/Shaders/Private/HairStrands/HairStrandsDebugCommon.ush
deleted file mode 100644
diff --git a/Engine/Shaders/Private/HairStrands/HairStrandsDebugPrint.usf b/Engine/Shaders/Private/HairStrands/HairStrandsDebugPrint.usf
--- a/Engine/Shaders/Private/HairStrands/HairStrandsDebugPrint.usf
+++ b/Engine/Shaders/Private/HairStrands/HairStrandsDebugPrint.usf
@@ -4,7 +4,6 @@
 #include "../ShaderPrintCommon.ush"
 #include "../ShaderDrawDebug.ush"
 #include "HairStrandsVisibilityCommon.ush"
-#include "HairStrandsDebugCommon.ush"
 #include "HairStrandsAABBCommon.ush" 
 
 int2 PixelCoord;
diff --git a/Engine/Shaders/Private/HairStrands/HairStrandsDeepShadowCommon.ush b/Engine/Shaders/Private/HairStrands/HairStrandsDeepShadowCommon.ush
--- a/Engine/Shaders/Private/HairStrands/HairStrandsDeepShadowCommon.ush
+++ b/Engine/Shaders/Private/HairStrands/HairStrandsDeepShadowCommon.ush
@@ -1,6 +1,7 @@
 // Copyright Epic Games, Inc. All Rights Reserved.
 
 #pragma once
+#include "../ColorMap.ush"  
 
 #define DEBUG_MODE 0
 
@@ -260,5 +261,5 @@ float3 ComputeHairCountDebugColor(
 	const float4 DOMValue = DomTexture.SampleLevel(LinearSampler, LightSpacePosition.xy, 0);
 	const float DistanceToFrontDepth = GetDomDistanceToFrontDepth(FrontDepth, LightSpacePosition.z);
 	const float HairCount = ComputeHairCount(DOMValue, DistanceToFrontDepth, InLayerDepths);
-	return GetColorCode(HairCount / MaxHairCount);
+	return GetHSVDebugColor(HairCount / MaxHairCount);
 }
\ No newline at end of file
diff --git a/Engine/Shaders/Private/HairStrands/HairStrandsDeepShadowMask.usf b/Engine/Shaders/Private/HairStrands/HairStrandsDeepShadowMask.usf
--- a/Engine/Shaders/Private/HairStrands/HairStrandsDeepShadowMask.usf
+++ b/Engine/Shaders/Private/HairStrands/HairStrandsDeepShadowMask.usf
@@ -12,7 +12,6 @@
 #include "HairStrandsVoxelPageCommon.ush"
 
 #include "HairStrandsCommon.ush"
-#include "HairStrandsDebugCommon.ush"
 #include "HairStrandsDeepShadowCommon.ush"
 #include "HairStrandsDeepTransmittanceCommon.ush"
 
diff --git a/Engine/Shaders/Private/HairStrands/HairStrandsDeepShadowPS.usf b/Engine/Shaders/Private/HairStrands/HairStrandsDeepShadowPS.usf
--- a/Engine/Shaders/Private/HairStrands/HairStrandsDeepShadowPS.usf
+++ b/Engine/Shaders/Private/HairStrands/HairStrandsDeepShadowPS.usf
@@ -6,7 +6,6 @@
 #include "../Common.ush"
 
 #include "../SceneTexturesCommon.ush"
-#include "HairStrandsDebugCommon.ush"
 #include "HairStrandsDeepShadowCommon.ush"
 #include "HairStrandsVisibilityCommon.ush"
 #include "HairStrandsVoxelPageCommon.ush"
diff --git a/Engine/Shaders/Private/HairStrands/HairStrandsDeepTransmittanceMask.usf b/Engine/Shaders/Private/HairStrands/HairStrandsDeepTransmittanceMask.usf
--- a/Engine/Shaders/Private/HairStrands/HairStrandsDeepTransmittanceMask.usf
+++ b/Engine/Shaders/Private/HairStrands/HairStrandsDeepTransmittanceMask.usf
@@ -8,7 +8,6 @@
 #include "../SceneTextureParameters.ush"
 #include "../DeferredShadingCommon.ush"
 #include "../ShadowFilteringCommon.ush"
-#include "HairStrandsDebugCommon.ush"
 #include "HairStrandsDeepShadowCommon.ush"
 #include "HairStrandsCommon.ush"
 #include "HairStrandsDeepTransmittanceCommon.ush"
diff --git a/Engine/Shaders/Private/HairStrands/HairStrandsEnvironmentLighting.usf b/Engine/Shaders/Private/HairStrands/HairStrandsEnvironmentLighting.usf
--- a/Engine/Shaders/Private/HairStrands/HairStrandsEnvironmentLighting.usf
+++ b/Engine/Shaders/Private/HairStrands/HairStrandsEnvironmentLighting.usf
@@ -35,7 +35,6 @@
 #include "../ReflectionEnvironmentComposite.ush"
 
 #include "../HairBsdf.ush"
-#include "HairStrandsDebugCommon.ush"
 #include "HairStrandsVisibilityCommon.ush"
 #include "HairStrandsEnvironmentLightingCommon.ush"
 #include "HairStrandsVoxelPageCommon.ush"
diff --git a/Engine/Shaders/Private/HairStrands/HairStrandsMeshProjectionHairDebug.usf b/Engine/Shaders/Private/HairStrands/HairStrandsMeshProjectionHairDebug.usf
--- a/Engine/Shaders/Private/HairStrands/HairStrandsMeshProjectionHairDebug.usf
+++ b/Engine/Shaders/Private/HairStrands/HairStrandsMeshProjectionHairDebug.usf
@@ -3,7 +3,6 @@
 //#define SCENE_TEXTURES_DISABLED 1
 
 #include "../Common.ush"
-#include "HairStrandsDebugCommon.ush"
 #include "HairStrandsMeshProjectionCommon.ush" 
 
 #define INPUT_ROOT_FRAME 0
diff --git a/Engine/Shaders/Private/HairStrands/HairStrandsMeshProjectionMeshDebug.usf b/Engine/Shaders/Private/HairStrands/HairStrandsMeshProjectionMeshDebug.usf
--- a/Engine/Shaders/Private/HairStrands/HairStrandsMeshProjectionMeshDebug.usf
+++ b/Engine/Shaders/Private/HairStrands/HairStrandsMeshProjectionMeshDebug.usf
@@ -4,7 +4,6 @@
 
 #include "../Common.ush"
 #include "HairStrandsMeshProjectionCommon.ush"
-#include "HairStrandsDebugCommon.ush"
 
 float4x4 LocalToWorld;
 uint VertexOffset;
diff --git a/Engine/Shaders/Private/HairStrands/HairStrandsVoxelDepthInjection.usf b/Engine/Shaders/Private/HairStrands/HairStrandsVoxelDepthInjection.usf
--- a/Engine/Shaders/Private/HairStrands/HairStrandsVoxelDepthInjection.usf
+++ b/Engine/Shaders/Private/HairStrands/HairStrandsVoxelDepthInjection.usf
@@ -10,7 +10,6 @@
 #include "../ShaderPrintCommon.ush"
 #include "HairStrandsVisibilityCommon.ush"
 #include "HairStrandsVoxelPageCommon.ush"
-#include "HairStrandsDebugCommon.ush"  
 #include "HairStrandsDeepShadowCommonStruct.ush"
 
 #include "HairStrandsVoxelPageCommon.ush"
diff --git a/Engine/Shaders/Private/HairStrands/HairStrandsVoxelPageRayMarching.usf b/Engine/Shaders/Private/HairStrands/HairStrandsVoxelPageRayMarching.usf
--- a/Engine/Shaders/Private/HairStrands/HairStrandsVoxelPageRayMarching.usf
+++ b/Engine/Shaders/Private/HairStrands/HairStrandsVoxelPageRayMarching.usf
@@ -10,7 +10,7 @@
 #include "../ShaderPrintCommon.ush"
 #include "HairStrandsVisibilityCommon.ush"
 #include "HairStrandsVoxelPageCommon.ush"
-#include "HairStrandsDebugCommon.ush"  
+#include "../ColorMap.ush"  
 
 #define VOXEL_TRAVERSAL_DEBUG_PAGEOCCUPANCY 0
 #define VOXEL_TRAVERSAL_DEBUG 0
@@ -607,7 +607,7 @@ void MainCS(uint3 DispatchThreadId : SV_DispatchThreadID)
 				VirtualVoxel.PageTexture,
 				TraversalSettings);
 
-			HitColor = GetColorCode(saturate(TraversalResult.HairCount / 16));
+			HitColor = GetHSVDebugColor(saturate(TraversalResult.HairCount / 16));
 			bIsValid = TraversalResult.HairCount > 0;
 		}
 
diff --git a/Engine/Shaders/Private/VolumetricCloud.usf b/Engine/Shaders/Private/VolumetricCloud.usf
--- a/Engine/Shaders/Private/VolumetricCloud.usf
+++ b/Engine/Shaders/Private/VolumetricCloud.usf
@@ -7,6 +7,7 @@
       
 
 #include "Common.ush"
+#include "ColorMap.ush"
 
 #define SceneTexturesStruct RenderVolumetricCloudParameters.SceneTextures
 
@@ -184,26 +185,6 @@ void MainVS(
 // Common structures and functions
 ////////////////////////////////////////////////////////////
 
-float3 HUEtoRGB(in float H)
-{
-	float R = abs(H * 6 - 3) - 1;
-	float G = 2 - abs(H * 6 - 2);
-	float B = 2 - abs(H * 6 - 4);
-	return saturate(float3(R, G, B));
-}
-
-float3 HSVtoRGB(in float3 HSV)
-{
-	float3 RGB = HUEtoRGB(HSV.x);
-	return ((RGB - 1) * HSV.y + 1) * HSV.z;
-}
-
-float3 GetColorCode(float x)
-{
-	float c = (1 - saturate(x)) * 0.6; // Remap [0,1] to Blue-Red
-	return x > 0 ? HSVtoRGB(float3(c, 1, 1)) : float3(0, 0, 0);
-}
-
 bool RayIntersectSphereSolution(float3 RayOrigin, float3 RayDirection, float4 Sphere, inout float2 Solutions)
 {
 	float3 LocalPosition = RayOrigin - Sphere.xyz;
@@ -1095,7 +1076,7 @@ void MainCommon(in FMaterialPixelParameters MaterialParameters, in float4 SvPosi
 			}
 
 			const bool bMediumHit = CountToDebug > 0;
-			OutColor0 = float4(bMediumHit ? GetColorCode(saturate(float(CountToDebug) / 64.0f)) : float3(0.0f, 0.0f, 0.0), bMediumHit ? 0.0f : 1.0f);
+			OutColor0 = float4(bMediumHit ? GetHSVDebugColor(saturate(float(CountToDebug) / 64.0f)) : float3(0.0f, 0.0f, 0.0), bMediumHit ? 0.0f : 1.0f);
 			OutColor1.x = bMediumHit ? OutColor1.x : MaxHalfFloat;
 		}
 #endif
