Fix a perf regression in the DiffuseIndirectComposite pass

diff --git a/Engine/Shaders/Private/DiffuseIndirectComposite.usf b/Engine/Shaders/Private/DiffuseIndirectComposite.usf
--- a/Engine/Shaders/Private/DiffuseIndirectComposite.usf
+++ b/Engine/Shaders/Private/DiffuseIndirectComposite.usf
@@ -123,8 +123,12 @@ float3 CombineRoughSpecular(FGBufferData GBuffer, float NoV, float3 RayTracedRef
 
 void MainPS(
 	float4 SvPosition : SV_POSITION
+#if DIM_APPLY_DIFFUSE_INDIRECT
 	, out float4 OutAddColor      DUAL_SOURCE_BLENDING_SLOT(0) : SV_Target0
 	, out float4 OutMultiplyColor DUAL_SOURCE_BLENDING_SLOT(1) : SV_Target1
+#else
+	, out float4 OutMultiplyColor : SV_Target0
+#endif
 )
 {
 	float2 SceneBufferUV = SvPositionToBufferUV(SvPosition);
@@ -147,7 +151,7 @@ void MainPS(
 	float NoV = saturate(dot(N, V));
 
 	// Apply diffuse indirect.
-#if DIM_APPLY_DIFFUSE_INDIRECT
+	#if DIM_APPLY_DIFFUSE_INDIRECT
 	{
 		float3 DiffuseIndirectLighting = 0;
 		float3 RoughSpecularIndirectLighting = 0;
@@ -295,10 +299,6 @@ void MainPS(
 			bNeedsSeparateSubsurfaceLightAccumulation);
 		OutAddColor = LightAccumulator_GetResult(LightAccumulator);
 	}
-	#else
-	{
-		OutAddColor = 0;
-	}
 	#endif
 
 	OutMultiplyColor = FinalAmbientOcclusion;
diff --git a/Engine/Source/Runtime/Renderer/Private/IndirectLightRendering.cpp b/Engine/Source/Runtime/Renderer/Private/IndirectLightRendering.cpp
--- a/Engine/Source/Runtime/Renderer/Private/IndirectLightRendering.cpp
+++ b/Engine/Source/Runtime/Renderer/Private/IndirectLightRendering.cpp
@@ -1021,7 +1021,9 @@ void FDeferredShadingSceneRenderer::RenderDiffuseIndirectAndAmbientOcclusion(
 			TShaderMapRef<FDiffuseIndirectCompositePS> PixelShader(View.ShaderMap, PermutationVector);
 			ClearUnusedGraphResources(PixelShader, PassParameters);
 
-			FRHIBlendState* BlendState = TStaticBlendState<CW_RGBA, BO_Add, BF_One, BF_Source1Color, BO_Add, BF_One, BF_Source1Alpha>::GetRHI();
+			FRHIBlendState* BlendState = PermutationVector.Get<FDiffuseIndirectCompositePS::FApplyDiffuseIndirectDim>() > 0 ?
+				TStaticBlendState<CW_RGBA, BO_Add, BF_One, BF_Source1Color, BO_Add, BF_One, BF_Source1Alpha>::GetRHI() :
+				TStaticBlendState<CW_RGBA, BO_Add, BF_Zero, BF_SourceColor, BO_Add, BF_Zero, BF_SourceAlpha>::GetRHI();
 
 			if (bIsVisualizePass)
 			{
