Added non metalness workflow to strata slab node.
- Reworked the packing of slab node to make eveyrhting fit. At the core level, diffuse, F0 and F90 are colored.
- Diffuse and F0 are stored as R7G7B6 to not bloatthe per pixel byte cost. They are dithered to hide the bit count without any loss of quality.
- F90 is stored as chroma+saturation because with force value=1 (HSV space) because this component need to be a high value anyway.
- Fixed related element: material conversion, vizualisation, classification, etc.

Next: implement energy conservation.

diff --git a/Engine/Shaders/Private/BasePassPixelShader.usf b/Engine/Shaders/Private/BasePassPixelShader.usf
--- a/Engine/Shaders/Private/BasePassPixelShader.usf
+++ b/Engine/Shaders/Private/BasePassPixelShader.usf
@@ -910,6 +910,8 @@ void FPixelShaderInOut_MainPS(
 	}
 #endif
 
+	const float Dither = InterleavedGradientNoise(MaterialParameters.SvPosition.xy, View.StateFrameIndexMod8);
+
 #if !STRATA_ENABLED
 	// Store the results in local variables and reuse instead of calling the functions multiple times.
 	half3 BaseColor = GetMaterialBaseColor(PixelMaterialInputs);
@@ -1056,7 +1058,6 @@ void FPixelShaderInOut_MainPS(
 	GBuffer.PrecomputedShadowFactors = GetPrecomputedShadowMasks(LightmapVTPageTableResult, Interpolants, MaterialParameters.PrimitiveId, MaterialParameters.AbsoluteWorldPosition, VolumetricLightmapBrickTextureUVs);
 
 #if !STRATA_ENABLED || STRATA_INLINE_SINGLELAYERWATER
-	const float GBufferDither = InterleavedGradientNoise(MaterialParameters.SvPosition.xy, View.StateFrameIndexMod8);
 	// Use GBuffer.ShadingModelID after SetGBufferForShadingModel(..) because the ShadingModel input might not be the same as the output
 	SetGBufferForShadingModel(
 		GBuffer,
@@ -1069,7 +1070,7 @@ void FPixelShaderInOut_MainPS(
 		Anisotropy,
 		SubsurfaceColor,
 		SubsurfaceProfile,
-		GBufferDither,
+		Dither,
 		ShadingModel
 		);
 #endif // !STRATA_ENABLED
@@ -1303,7 +1304,6 @@ void FPixelShaderInOut_MainPS(
 				GridIndex = ComputeLightGridCellIndex((uint2)((MaterialParameters.SvPosition.xy - ResolvedView.ViewRectMin.xy) * View.LightProbeSizeRatioAndInvSizeRatio.zw), MaterialParameters.SvPosition.w, EyeIndex);
 
 				#if FORWARD_SHADING || TRANSLUCENCY_LIGHTING_SURFACE_FORWARDSHADING || MATERIAL_SHADINGMODEL_SINGLELAYERWATER
-					const float Dither = InterleavedGradientNoise(MaterialParameters.SvPosition.xy, View.StateFrameIndexMod8);
 					FDeferredLightingSplit ForwardDirectLighting = GetForwardDirectLightingSplit(GridIndex, MaterialParameters.AbsoluteWorldPosition, MaterialParameters.CameraVector, GBuffer, NearestResolvedDepthScreenUV, MaterialParameters.PrimitiveId, EyeIndex, Dither, DirectionalLightShadow);
 					
 					#if MATERIAL_SHADINGMODEL_THIN_TRANSLUCENT
@@ -1571,7 +1571,7 @@ void FPixelShaderInOut_MainPS(
 	float3 EmissiveLuminance = 0.0f;
 	const float3 VertexNormal = MaterialParameters.TangentToWorld[2].xyz;
 	PackStrataOut(
-		OpaqueBasePass.Strata.MaterialLobesBufferUAV, 
+		OpaqueBasePass.Strata.MaterialLobesBufferUAV, Dither,
 		GetStrataPixelDataByteOffset(PixelPos, uint2(ResolvedView.BufferSizeAndInvSize.xy), OpaqueBasePass.Strata.MaxBytesPerPixel),
 		StrataPixelHeader, StrataData, MaterialParameters.CameraVector, bStrataSubsurfaceEnable, EmissiveLuminance);
 
diff --git a/Engine/Shaders/Private/ColorSpace.ush b/Engine/Shaders/Private/ColorSpace.ush
--- a/Engine/Shaders/Private/ColorSpace.ush
+++ b/Engine/Shaders/Private/ColorSpace.ush
@@ -206,3 +206,20 @@ float3 HSV_2_LinearRGB(in float3 HSV)
 	float3 RGB = HUE_2_LinearRGB(HSV.x);
 	return ((RGB - 1) * HSV.y + 1) * HSV.z;
 }
+
+float3 RGB_2_HCV(in float3 RGB)
+{
+	// Based on work by Sam Hocevar and Emil Persson
+	float4 P = (RGB.g < RGB.b)	? float4(RGB.bg, -1.0f, 2.0f / 3.0f): float4(RGB.gb, 0.0f, -1.0f / 3.0f);
+	float4 Q = (RGB.r < P.x)	? float4(P.xyw, RGB.r)				: float4(RGB.r, P.yzx);
+	float Chroma = Q.x - min(Q.w, Q.y);
+	float Hue = abs((Q.w - Q.y) / (6.0f * Chroma + 1e-10f) + Q.z);
+	return float3(Hue, Chroma, Q.x);
+}
+
+float3 LinearRGB_2_HSV(in float3 RGB)
+{
+	float3 HCV = RGB_2_HCV(RGB);
+	float s = HCV.y / (HCV.z + 1e-10f);
+	return float3(HCV.x, s, HCV.z);
+}
diff --git a/Engine/Shaders/Private/Strata/Strata.ush b/Engine/Shaders/Private/Strata/Strata.ush
--- a/Engine/Shaders/Private/Strata/Strata.ush
+++ b/Engine/Shaders/Private/Strata/Strata.ush
@@ -2,8 +2,9 @@
 
 #pragma once
 
-#include "../Common.ush"
 #include "../BRDF.ush"
+#include "../ColorSpace.ush"
+#include "../Common.ush"
 #include "../DeferredShadingCommon.ush"
 #include "../ParticipatingMediaCommon.ush"
 #include "../ThinFilmBSDF.ush"
@@ -59,7 +60,7 @@ struct FStrataBSDF
 	float4 InlineVGPRs[3];
 #endif
 
-	float4 VGPRs[5];
+	float4 VGPRs[6];
 };
 
 
@@ -118,7 +119,7 @@ struct FStrataBSDF
 #define STATE_BIT_COUNT_HASSSSPROFILE			1
 #define STATE_BIT_COUNT_GREYWEIGHT				1
 #define STATE_BIT_COUNT_HAZINESS				1
-#define STATE_BIT_COUNT_EDGECOLOR				1
+#define STATE_BIT_COUNT_F90						1
 #define STATE_BIT_COUNT_THINFILM				1
 #define STATE_BIT_COUNT_SIMPLEVOLUME			1
 #define STATE_BIT_COUNT_DMFPPLUGGED				1
@@ -136,8 +137,8 @@ struct FStrataBSDF
 #define STATE_BIT_OFFSET_HASSSSPROFILE	(STATE_BIT_OFFSET_HASSSS			+ STATE_BIT_COUNT_HASSSS)
 #define STATE_BIT_OFFSET_GREYWEIGHT		(STATE_BIT_OFFSET_HASSSSPROFILE		+ STATE_BIT_COUNT_HASSSSPROFILE)
 #define STATE_BIT_OFFSET_HAZINESS		(STATE_BIT_OFFSET_GREYWEIGHT		+ STATE_BIT_COUNT_GREYWEIGHT)
-#define STATE_BIT_OFFSET_EDGECOLOR		(STATE_BIT_OFFSET_HAZINESS			+ STATE_BIT_COUNT_HAZINESS)
-#define STATE_BIT_OFFSET_THINFILM		(STATE_BIT_OFFSET_EDGECOLOR			+ STATE_BIT_COUNT_EDGECOLOR)
+#define STATE_BIT_OFFSET_F90			(STATE_BIT_OFFSET_HAZINESS			+ STATE_BIT_COUNT_HAZINESS)
+#define STATE_BIT_OFFSET_THINFILM		(STATE_BIT_OFFSET_F90				+ STATE_BIT_COUNT_F90)
 #define STATE_BIT_OFFSET_SIMPLEVOLUME	(STATE_BIT_OFFSET_THINFILM			+ STATE_BIT_COUNT_THINFILM)
 #define STATE_BIT_OFFSET_DMFPPLUGGED	(STATE_BIT_OFFSET_SIMPLEVOLUME		+ STATE_BIT_COUNT_SIMPLEVOLUME)
 #define STATE_BIT_OFFSET_HASFUZZ		(STATE_BIT_OFFSET_DMFPPLUGGED		+ STATE_BIT_COUNT_DMFPPLUGGED)
@@ -179,8 +180,8 @@ struct FStrataBSDF
 #define BSDF_SETHASHAZINESS(X, Haze)	WRITE_BITS(X.State, STATE_BIT_COUNT_HAZINESS, STATE_BIT_OFFSET_HAZINESS, Haze)
 
 // Indicates if the BSDF has edge color data
-#define BSDF_GETHASEDGECOLOR(X)			READ_BITS(X.State,  STATE_BIT_COUNT_EDGECOLOR, STATE_BIT_OFFSET_EDGECOLOR)
-#define BSDF_SETHASEDGECOLOR(X, Col)	WRITE_BITS(X.State, STATE_BIT_COUNT_EDGECOLOR, STATE_BIT_OFFSET_EDGECOLOR, Col)
+#define BSDF_GETHASF90(X)				READ_BITS(X.State,  STATE_BIT_COUNT_F90, STATE_BIT_OFFSET_F90)
+#define BSDF_SETHASF90(X, Col)			WRITE_BITS(X.State, STATE_BIT_COUNT_F90, STATE_BIT_OFFSET_F90, Col)
 
 // Indicates if the BSDF has thin film
 #define BSDF_GETHASTHINFILM(X)			READ_BITS(X.State,  STATE_BIT_COUNT_THINFILM, STATE_BIT_OFFSET_THINFILM)
@@ -224,30 +225,32 @@ struct FStrataBSDF
 
 #define SLAB_COMPLEX_PATH_MASK (TO_OFFSET_BITMASK(STATE_BIT_COUNT_ANISOTROPY, STATE_BIT_OFFSET_ANISOTROPY)		|  TO_OFFSET_BITMASK(STATE_BIT_COUNT_HASSSS, STATE_BIT_OFFSET_HASSSS) | \
 								TO_OFFSET_BITMASK(STATE_BIT_COUNT_HASSSSPROFILE, STATE_BIT_OFFSET_HASSSSPROFILE)|  TO_OFFSET_BITMASK(STATE_BIT_COUNT_HAZINESS, STATE_BIT_OFFSET_HAZINESS) | \
-								TO_OFFSET_BITMASK(STATE_BIT_COUNT_THINFILM, STATE_BIT_OFFSET_THINFILM)			|  TO_OFFSET_BITMASK(STATE_BIT_COUNT_EDGECOLOR, STATE_BIT_OFFSET_EDGECOLOR) | \
+								TO_OFFSET_BITMASK(STATE_BIT_COUNT_THINFILM, STATE_BIT_OFFSET_THINFILM)			|  TO_OFFSET_BITMASK(STATE_BIT_COUNT_F90, STATE_BIT_OFFSET_F90) | \
 								TO_OFFSET_BITMASK(STATE_BIT_COUNT_SIMPLEVOLUME, STATE_BIT_OFFSET_SIMPLEVOLUME)	|  TO_OFFSET_BITMASK(STATE_BIT_COUNT_HASFUZZ, STATE_BIT_OFFSET_HASFUZZ) | \
 								TO_OFFSET_BITMASK(STATE_BIT_COUNT_ISTHIN, STATE_BIT_OFFSET_ISTHIN) )
 
 
-#define SLAB_BASECOLOR(X)				X.VGPRs[0].xyz
-#define SLAB_SPECULAR(X)				X.VGPRs[0].w
-#define SLAB_METALLIC(X)				X.VGPRs[1].x
-#define SLAB_ROUGHNESSX(X)				X.VGPRs[1].y
-#define SLAB_ROUGHNESSY(X)				X.VGPRs[1].z
-#define SLAB_HAZINESS(X)				X.VGPRs[1].w
-#define SLAB_EDGECOLOR(X)				X.VGPRs[3].xyz
+#define SLAB_DIFFUSEALBEDO(X)			X.VGPRs[0].xyz
+#define SLAB_ROUGHNESSX(X)				X.VGPRs[0].w
+#define SLAB_F0(X)						X.VGPRs[1].xyz
+#define SLAB_ROUGHNESSY(X)				X.VGPRs[1].w
+#define SLAB_F90(X)						X.VGPRs[2].xyz
+#define SLAB_METALLIC(X)				X.VGPRs[2].w
 // Either with SSS profile
-#define SLAB_SSSPROFILEID(X)			X.VGPRs[2].x
-#define SLAB_SSSPROFILERADIUSSCALE(X)	X.VGPRs[2].y
+#define SLAB_SSSPROFILEID(X)			X.VGPRs[3].x
+#define SLAB_SSSPROFILERADIUSSCALE(X)	X.VGPRs[3].y
 // Or explicit DMFP in centimeter
-#define SLAB_SSSDMFP(X)					X.VGPRs[2].xyz
+#define SLAB_SSSDMFP(X)					X.VGPRs[3].xyz
 // Thin film. IOR is harcoded to oil IOR in order to save VGPR
-#define SLAB_THINFILM_THICKNESS(X)		X.VGPRs[2].w
+#define SLAB_THINFILM_THICKNESS(X)		X.VGPRs[3].w
 #define SLAB_THINFILM_IOR(X)			1.44f
 // Fuzz to simulate cloth
 #define SLAB_FUZZ_COLOR(X)				X.VGPRs[4].xyz
 #define SLAB_FUZZ_AMOUNT(X)				X.VGPRs[4].w
-// NOTE For SLAB, X.VGPRs[3].w is free for another use case.
+// Haziness
+#define SLAB_HAZINESS(X)				X.VGPRs[5].x
+// NOTE For SLAB, X.VGPRs[5].yzw is free for another use case.
+
 
 #define VOLUMETRICFOGCLOUD_ALBEDO(X)	X.VGPRs[0].xyz
 #define VOLUMETRICFOGCLOUD_EXTINCTION(X)X.VGPRs[1].xyz
@@ -460,8 +463,10 @@ FStrataData GetInitialisedStrataDataWithReservedFirstSlot()
 }
 
 FStrataData GetStrataSlabBSDF(
-	float3 BaseColor, float3 EdgeColor, float Specular, float Metallic, float Roughness, 
-	float Anisotropy,
+	float UseMetalness,
+	float3 BaseColor, float3 EdgeColor, float Specular, float Metallic,	//  Metalness workflow
+	float3 DiffuseAlbedo, float3 F0, float3 F90,						// !Metalness workflow
+	float Roughness, float Anisotropy,
 	float SSSProfileID, float3 SSSDMFP, float SSSDMFPScale,
 	float3 Emissive, 
 	float Haziness, 
@@ -472,7 +477,6 @@ FStrataData GetStrataSlabBSDF(
 {
 	FStrataData StrataData = GetInitialisedStrataDataWithReservedFirstSlot();
 
-	const uint bHasEdgeColor = any(EdgeColor < 1);
 	const uint bHasHaziness = Haziness > 0.0f ? 1 : 0;
 	const uint bHasThinFilm = ThinFilmThickness > 0.0f ? 1 : 0;
 	const uint bHasAnisotropy = Anisotropy != 0;
@@ -508,15 +512,36 @@ FStrataData GetStrataSlabBSDF(
 	BSDF_SETISTHIN			(StrataData.Layers[0].BSDFs[0],		bDMFPPlugged && Thickness < STRATA_LAYER_ISTHIN_THICKNESS_THRESHOLD_CM);
 	BSDF_SETHASDMFP			(StrataData.Layers[0].BSDFs[0],		bDMFPPlugged ? 1 : 0);
 	BSDF_SETHASHAZINESS		(StrataData.Layers[0].BSDFs[0],		bHasHaziness);
-	BSDF_SETHASEDGECOLOR	(StrataData.Layers[0].BSDFs[0],		bHasEdgeColor);
 	BSDF_SETHASTHINFILM		(StrataData.Layers[0].BSDFs[0],		bHasThinFilm);
 	BSDF_SETTHICKNESSCM		(StrataData.Layers[0].BSDFs[0],		Thickness);
 	BSDF_SETHASFUZZ			(StrataData.Layers[0].BSDFs[0],		bHasFuzz);
 
-	SLAB_BASECOLOR			(StrataData.Layers[0].BSDFs[0])		= BaseColor;
-	SLAB_EDGECOLOR			(StrataData.Layers[0].BSDFs[0])		= EdgeColor;
-	SLAB_SPECULAR			(StrataData.Layers[0].BSDFs[0])		= Specular;
-	SLAB_METALLIC			(StrataData.Layers[0].BSDFs[0])		= Metallic;
+	if (UseMetalness > 0.0f)
+	{
+		BSDF_SETHASF90(StrataData.Layers[0].BSDFs[0], any(EdgeColor < 1.0f));
+
+		SLAB_DIFFUSEALBEDO(StrataData.Layers[0].BSDFs[0])	= BaseColor - BaseColor * Metallic;
+		SLAB_F0(StrataData.Layers[0].BSDFs[0])				= ComputeF0(Specular, BaseColor, Metallic);
+		SLAB_F90(StrataData.Layers[0].BSDFs[0])				= ComputeF90(F0, EdgeColor, Metallic);
+		SLAB_METALLIC(StrataData.Layers[0].BSDFs[0])		= Metallic;
+	}
+	else
+	{
+		// See how UMaterialExpressionStrataSlabBSDF node interface is changed
+		BSDF_SETHASF90(StrataData.Layers[0].BSDFs[0], any(F90 < 1.0f));
+
+		// Approximate the metallic input from F0.
+		const float DiamondF0 = 0.24f;
+		const float CondutorMinF0 = 0.4f; // roughly the end of semi-conductor
+		const float MaxF0 = max(F0.r, max(F0.g, F0.b));
+
+		SLAB_DIFFUSEALBEDO(StrataData.Layers[0].BSDFs[0])	= DiffuseAlbedo;
+		SLAB_F0(StrataData.Layers[0].BSDFs[0])				= F0;
+		SLAB_F90(StrataData.Layers[0].BSDFs[0])				= F90;
+		// Metalness cannot be recovered from F0 so we use a simple heuristic mapping (based on MaxF0 which should better represent the state of the material)
+		SLAB_METALLIC(StrataData.Layers[0].BSDFs[0])		= saturate((MaxF0 - DiamondF0) / (CondutorMinF0 - DiamondF0));
+	}
+
 	SLAB_ROUGHNESSX			(StrataData.Layers[0].BSDFs[0])		= RoughnessXY.x;
 	SLAB_ROUGHNESSY			(StrataData.Layers[0].BSDFs[0])		= RoughnessXY.y;
 	SLAB_HAZINESS			(StrataData.Layers[0].BSDFs[0])		= Haziness;
@@ -903,15 +928,15 @@ FStrataData StrataHorizontalMixingParameterBlending(FStrataData Foreground, FStr
 	BSDF_SETISTHIN(ResultBSDF,			BSDF_GETISTHIN(ForegroundBSDF)					| BSDF_GETISTHIN(BackgroundBSDF));			// STRATA_TODO fix the MPF/DMFP situation for that to really work
 	BSDF_SETHASDMFP(ResultBSDF,			BSDF_GETHASDMFP(ForegroundBSDF)					| BSDF_GETHASDMFP(BackgroundBSDF));
 	BSDF_SETHASHAZINESS(ResultBSDF,		BSDF_GETHASHAZINESS(ForegroundBSDF)				| BSDF_GETHASHAZINESS(BackgroundBSDF));
-	BSDF_SETHASEDGECOLOR(ResultBSDF,	BSDF_GETHASEDGECOLOR(ForegroundBSDF)			| BSDF_GETHASEDGECOLOR(BackgroundBSDF));
+	BSDF_SETHASF90(ResultBSDF,			BSDF_GETHASF90(ForegroundBSDF)					| BSDF_GETHASF90(BackgroundBSDF));
 	BSDF_SETHASTHINFILM(ResultBSDF,		BSDF_GETHASTHINFILM(ForegroundBSDF)				| BSDF_GETHASTHINFILM(BackgroundBSDF));
 	BSDF_SETHASFUZZ(ResultBSDF,			BSDF_GETHASFUZZ(ForegroundBSDF)					| BSDF_GETHASFUZZ(BackgroundBSDF));
 	BSDF_SETEMISSIVE(ResultBSDF,		lerp(BSDF_GETEMISSIVE(BackgroundBSDF),			BSDF_GETEMISSIVE(ForegroundBSDF),			ForegroundMixFactor));
 	BSDF_SETTHICKNESSCM(ResultBSDF,		lerp(BSDF_GETTHICKNESSCM(BackgroundBSDF),		BSDF_GETTHICKNESSCM(ForegroundBSDF),		ForegroundMixFactor));
 
-	SLAB_BASECOLOR(ResultBSDF) =		lerp(SLAB_BASECOLOR(BackgroundBSDF),			SLAB_BASECOLOR(ForegroundBSDF),				ForegroundMixFactor);
-	SLAB_EDGECOLOR(ResultBSDF)=			lerp(SLAB_EDGECOLOR(BackgroundBSDF),			SLAB_EDGECOLOR(ForegroundBSDF),				ForegroundMixFactor);
-	SLAB_SPECULAR(ResultBSDF) =			lerp(SLAB_SPECULAR(BackgroundBSDF),				SLAB_SPECULAR(ForegroundBSDF),				ForegroundMixFactor);
+	SLAB_DIFFUSEALBEDO(ResultBSDF) =	lerp(SLAB_DIFFUSEALBEDO(BackgroundBSDF),		SLAB_DIFFUSEALBEDO(ForegroundBSDF),			ForegroundMixFactor);
+	SLAB_F0(ResultBSDF) =				lerp(SLAB_F0(BackgroundBSDF),					SLAB_F0(ForegroundBSDF),					ForegroundMixFactor);
+	SLAB_F90(ResultBSDF) =				lerp(SLAB_F90(BackgroundBSDF),					SLAB_F90(ForegroundBSDF),					ForegroundMixFactor);
 	SLAB_METALLIC(ResultBSDF) =			lerp(SLAB_METALLIC(BackgroundBSDF),				SLAB_METALLIC(ForegroundBSDF),				ForegroundMixFactor);
 	SLAB_ROUGHNESSX(ResultBSDF) =		lerp(SLAB_ROUGHNESSX(BackgroundBSDF),			SLAB_ROUGHNESSX(ForegroundBSDF),			ForegroundMixFactor);
 	SLAB_ROUGHNESSY(ResultBSDF) =		lerp(SLAB_ROUGHNESSY(BackgroundBSDF),			SLAB_ROUGHNESSY(ForegroundBSDF),			ForegroundMixFactor);
@@ -953,9 +978,9 @@ FStrataData StrataHorizontalMixingParameterBlending(FStrataData Foreground, FStr
 void SanitizeStrataSlab(inout FStrataBSDF BSDF)
 {
 	BSDF_SETEMISSIVE(BSDF, clamp(BSDF_GETEMISSIVE(BSDF), 0.0, Max10BitsFloat));
-	SLAB_BASECOLOR(BSDF) = saturate(SLAB_BASECOLOR(BSDF));
-	SLAB_EDGECOLOR(BSDF) = saturate(SLAB_EDGECOLOR(BSDF));
-	SLAB_SPECULAR(BSDF) = saturate(SLAB_SPECULAR(BSDF));
+	SLAB_DIFFUSEALBEDO(BSDF) = saturate(SLAB_DIFFUSEALBEDO(BSDF));
+	SLAB_F0(BSDF) = saturate(SLAB_F0(BSDF));
+	SLAB_F90(BSDF) = saturate(SLAB_F90(BSDF));
 	SLAB_METALLIC(BSDF) = saturate(SLAB_METALLIC(BSDF));
 	SLAB_ROUGHNESSX(BSDF) = saturate(SLAB_ROUGHNESSX(BSDF));
 	SLAB_ROUGHNESSY(BSDF) = saturate(SLAB_ROUGHNESSY(BSDF));
@@ -1008,8 +1033,6 @@ void SanitizeStrataSingleLayerWater(inout FStrataBSDF BSDF)
 #endif
 }
 
-
-
 void StrataSanitizeBSDF(inout FStrataBSDF BSDF)
 {
 	switch (BSDF_GETTYPE(BSDF))
@@ -1033,6 +1056,20 @@ void StrataSanitizeBSDF(inout FStrataBSDF BSDF)
 	}
 }
 
+void StrataDitherBSDF(inout FStrataBSDF BSDF, in float Dither)
+{
+	if (BSDF_GETTYPE(BSDF) == STRATA_BSDF_TYPE_SLAB)
+	{
+		// We execute a dither that can go in positive/negative direction to avoid surfaces to looks overall brighter.
+		const float CenteredDither = Dither - 0.5f;
+		const float Dither7bits = CenteredDither * 1.0f / 127.0f;
+		const float Dither6bits = CenteredDither * 1.0f / 63.0f;
+		const float3 Dither776 = float3(Dither7bits, Dither7bits, Dither6bits);
+
+		SLAB_DIFFUSEALBEDO(BSDF) += Dither776;
+		SLAB_F0(BSDF) += Dither776;
+	}
+}
 
 
 ///////////////////////////////////////////////////////////////////////////////
@@ -1193,6 +1230,23 @@ float4 UnpackRGBA8(uint rgba)
 	return Out;
 }
 
+uint PackR7G7B6Gamma2(float3 rgb)
+{
+	uint r = (uint(sqrt(saturate(rgb.r)) * 127.0f) << 0);
+	uint g = (uint(sqrt(saturate(rgb.g)) * 127.0f) << 7);
+	uint b = (uint(sqrt(saturate(rgb.b)) *  63.0f) << 14);
+	return r | g | b;
+}
+
+float3 UnpackR7G7B6Gamma2(uint rgb)
+{
+	float3 Out;
+	Out.r = float((rgb >> 0)  & 0x7F) * (1.0f / 127.0f);
+	Out.g = float((rgb >> 7)  & 0x7F) * (1.0f / 127.0f);
+	Out.b = float((rgb >> 14) & 0x3F) * (1.0f /  63.0f);
+	return Out * Out;
+}
+
 // Simple Octahedral from "A Survey of Efficient Representations for Independent Unit Vectors" http://jcgt.org/published/0003/02/01/
 float SignNotZero(in float k)
 {
@@ -1361,6 +1415,7 @@ uint PackStrataHeaderIn(uint BSDFCount, FStrataPixelHeader In);
 // If this is changed, please update the compiler side material size evaluation in StrataMaterial.cpp
 void PackStrataOut(
 	RWByteAddressBuffer StrataBuffer, 
+	float Dither,
 	FStrataAddressing StrataAddressing,
 	FStrataPixelHeader StrataPixelHeader,
 	FStrataData Strata,
@@ -1395,6 +1450,10 @@ void PackStrataOut(
 		STRATA_UNROLL
 		for (int i = 0; i < Strata.Layers[l].BSDFCount; ++i)
 		{
+			// Dither BSDF data to hide more aggressive bit reduction.
+			// Sanitize happening after will clamp values that could have gone above and below valid range.
+			StrataDitherBSDF(Strata.Layers[l].BSDFs[i], Dither);
+
 			// Sanitize BSDF in place before it is used and packed out to memory
 			StrataSanitizeBSDF(Strata.Layers[l].BSDFs[i]);
 
@@ -1437,10 +1496,8 @@ void PackStrataOut(
 					const float SatNoV = saturate(ContextNoLight.NoV);
 					const float VoH = SatNoV;
 					const float NoL = SatNoV;
-					const float3 BaseColor = SLAB_BASECOLOR(BSDF);
-					const float  Metallic = SLAB_METALLIC(BSDF);
-					const float3 F0 = ComputeF0(SLAB_SPECULAR(BSDF), BaseColor, Metallic);
-					const float3 F90 = ComputeF90(F0, SLAB_EDGECOLOR(BSDF), Metallic);
+					const float3 F0 = SLAB_F0(BSDF);
+					const float3 F90 = SLAB_F90(BSDF);
 					const bool bHasThinFilm = BSDF_GETHASTHINFILM(BSDF);
 					BRANCH
 					if (bHasThinFilm)
@@ -1570,12 +1627,20 @@ void PackStrataOut(
 			{
 			case STRATA_BSDF_TYPE_SLAB:
 			{
-				STRATA_STORE_UINT1(PackColorLinearToGamma2AlphaLinear(float4(SLAB_BASECOLOR(BSDF), SLAB_ROUGHNESSX(BSDF))));
-				STRATA_STORE_UINT1(PackRGBA8(float4(SLAB_SPECULAR(BSDF), SLAB_METALLIC(BSDF), SLAB_ROUGHNESSY(BSDF), SLAB_HAZINESS(BSDF))));
+				// Now we pack diffuse and F0 in a special way: both encoded as R7G7B6 with gamma 2.0. The low bit cound will be hidden by dithering and TAA.
+				// We pack 32 bits in the first uint and the remaining bits are pack is the lowest significant bit of the second uint.
+				const uint PackedDiffuse20Bits	= PackR7G7B6Gamma2(SLAB_DIFFUSEALBEDO(BSDF));
+				const uint PackedF020Bits		= PackR7G7B6Gamma2(SLAB_F0(BSDF));
+				const uint PackedData32Bits		= ((PackedDiffuse20Bits << 12) & 0xFFFFF000) | (PackedF020Bits & 0xFFF);
+				const uint PackedData8Bits		= (PackedF020Bits >> 12) & 0xFF; 
+
+				STRATA_STORE_UINT1(PackedData32Bits);
+				STRATA_STORE_UINT1(PackedData8Bits | PackRGBA8(float4(0.0f , SLAB_ROUGHNESSX(BSDF), SLAB_ROUGHNESSY(BSDF), SLAB_METALLIC(BSDF))));
 
-				if (BSDF_GETHASEDGECOLOR(BSDF) || BSDF_GETHASTHINFILM(BSDF))
+				if (BSDF_GETHASF90(BSDF) || BSDF_GETHASTHINFILM(BSDF) || BSDF_GETHASHAZINESS(BSDF))
 				{
-					STRATA_STORE_UINT1(PackColorLinearToGamma2AlphaLinear(float4(SLAB_EDGECOLOR(BSDF), SLAB_THINFILM_THICKNESS(BSDF))));
+					float3 F90HSV = LinearRGB_2_HSV(SLAB_F90(BSDF));
+					STRATA_STORE_UINT1(PackRGBA8(float4(F90HSV.x, F90HSV.y, SLAB_THINFILM_THICKNESS(BSDF), SLAB_HAZINESS(BSDF))));
 				}
 
 				const bool bIsSimpleVolume = BSDF_GETISSIMPLEVOLUME(BSDF);
@@ -1683,21 +1748,23 @@ bool IsStrataSlabFastPathCompatible(in FStrataBSDF BSDF)
 
 void UnpackFastPathSlabBSDF(inout FStrataBSDF OutBSDF, uint2 PackedData01)
 {
-	float4 Data0 = UnpackColorGamma2ToLinearAlphaLinear(PackedData01.x);
-	SLAB_BASECOLOR(OutBSDF)		= Data0.rgb;
-	SLAB_EDGECOLOR(OutBSDF)		= 1.0f;
-	SLAB_ROUGHNESSX(OutBSDF)	= Data0.a;
+	const uint PackedData8Bits		= PackedData01.y & 0xFF;
+	const uint PackedData32Bits		= PackedData01.x;
+	const uint PackedF020Bits		= ((PackedData8Bits << 12) & 0xFF000) | (PackedData32Bits & 0xFFF);
+	const uint PackedDiffuse20Bits	= (PackedData32Bits >> 12) & 0xFFFFF;
+
+	SLAB_DIFFUSEALBEDO(OutBSDF)	= UnpackR7G7B6Gamma2(PackedDiffuse20Bits);
+	SLAB_F0(OutBSDF)			= UnpackR7G7B6Gamma2(PackedF020Bits);
 
 	float4 Data1 = UnpackRGBA8(PackedData01.y);
-	SLAB_SPECULAR(OutBSDF)		= Data1.x;
-	SLAB_METALLIC(OutBSDF)		= Data1.y;
+	SLAB_ROUGHNESSX(OutBSDF)	= Data1.y;
 	SLAB_ROUGHNESSY(OutBSDF)	= Data1.z;
-	SLAB_HAZINESS(OutBSDF)		= Data1.w;
+	SLAB_METALLIC(OutBSDF)		= Data1.w;
 }
 
 FStrataBSDF UnpackFastPathStrataBSDFIn(ByteAddressBuffer StrataBuffer, inout FStrataAddressing StrataAddressing)
 {
-	FStrataBSDF OutBSDF = (FStrataBSDF)0;
+	FStrataBSDF OutBSDF = (FStrataBSDF)0; 
 
 	OutBSDF.State = StrataLoadUint1(StrataBuffer, StrataAddressing);
 
@@ -1709,10 +1776,13 @@ FStrataBSDF UnpackFastPathStrataBSDFIn(ByteAddressBuffer StrataBuffer, inout FSt
 	uint DataY = StrataLoadUint1(StrataBuffer, StrataAddressing);
 	UnpackFastPathSlabBSDF(OutBSDF, uint2(DataX, DataY));
 
+	// This is convenient to not have to test BSDF_GETHASF90 all over the place in the code using the BSDF for such a core value.
+	SLAB_F90(OutBSDF) = 1.0f;
+
 	// Force a BSDF state to a single slab with isotropic specular, no subsurface scattering or any other options.
 	// This is to help the compiler understand disabled path.
 	BSDF_SETTYPE(OutBSDF, STRATA_BSDF_TYPE_SLAB);
-	BSDF_SETHASEDGECOLOR(OutBSDF, 0);
+	BSDF_SETHASF90(OutBSDF, 0);
 	BSDF_SETANISOTROPY(OutBSDF, 0);
 	BSDF_SETISTOPLAYER(OutBSDF, 1);
 	BSDF_SETHASSSS(OutBSDF, 0);
@@ -1755,11 +1825,20 @@ FStrataBSDF UnpackStrataBSDFIn(ByteAddressBuffer StrataBuffer, inout FStrataAddr
 		uint DataY = StrataLoadUint1(StrataBuffer, StrataAddressing);
 		UnpackFastPathSlabBSDF(OutBSDF, uint2(DataX, DataY));
 
-		SLAB_EDGECOLOR(OutBSDF) = 1.0f;
+		SLAB_F90(OutBSDF) = 1.0f;
 		SLAB_THINFILM_THICKNESS(OutBSDF) = 0.0f;
-		if (BSDF_GETHASEDGECOLOR(OutBSDF) || BSDF_GETHASTHINFILM(OutBSDF))
+		const bool bHasF90 = BSDF_GETHASF90(OutBSDF);
+		if (bHasF90 || BSDF_GETHASTHINFILM(OutBSDF) || BSDF_GETHASHAZINESS(OutBSDF))
 		{
-			StrataLoad_ColorGamma2ToLinearAlphaLinear(StrataBuffer, StrataAddressing, SLAB_EDGECOLOR(OutBSDF), SLAB_THINFILM_THICKNESS(OutBSDF));
+			float4 Data = UnpackRGBA8(StrataLoadUint1(StrataBuffer, StrataAddressing));
+
+			const float F90Hue = Data.x;
+			const float F90Sat = Data.y;
+			const float F90Val = 1.0f;	// Forced to 1 as F90 is not super to have a low value
+			SLAB_F90(OutBSDF) = bHasF90 ? HSV_2_LinearRGB(float3(F90Hue, F90Sat, F90Val)) : 1.0f;
+
+			SLAB_THINFILM_THICKNESS(OutBSDF) = Data.z;
+			SLAB_HAZINESS(OutBSDF) = Data.w;
 		}
 
 		const bool bIsSimpleVolume = BSDF_GETISSIMPLEVOLUME(OutBSDF);
@@ -1783,7 +1862,12 @@ FStrataBSDF UnpackStrataBSDFIn(ByteAddressBuffer StrataBuffer, inout FStrataAddr
 			if (!bIsSimpleVolume && !bIsThin)
 			{
 				const bool bChecker = CheckerFromPixelPos(StrataAddressing.PixelCoords);
-				AdjustBaseColorAndSpecularColorForSubsurfaceProfileLighting(SLAB_BASECOLOR(OutBSDF), SLAB_SPECULAR(OutBSDF), bChecker);
+				float3 DiffuseAlbedo = SLAB_DIFFUSEALBEDO(OutBSDF);
+				float SpecularFactor = 1.0f;
+				AdjustBaseColorAndSpecularColorForSubsurfaceProfileLighting(DiffuseAlbedo, SpecularFactor, bChecker);
+				SLAB_DIFFUSEALBEDO(OutBSDF) = DiffuseAlbedo;
+				SLAB_F0(OutBSDF)  *= SpecularFactor;
+				SLAB_F90(OutBSDF) *= SpecularFactor;
 			}
 			#endif
 		}
diff --git a/Engine/Shaders/Private/Strata/StrataDeferredLighting.ush b/Engine/Shaders/Private/Strata/StrataDeferredLighting.ush
--- a/Engine/Shaders/Private/Strata/StrataDeferredLighting.ush
+++ b/Engine/Shaders/Private/Strata/StrataDeferredLighting.ush
@@ -71,7 +71,7 @@ void StrataPreprocessAreaLightIntegration(
 			{
 				SafeRoughnessX = MakeRoughnessSafe(SLAB_ROUGHNESSX(BSDFContext.BSDF), View.MinRoughness);
 				SafeRoughnessY = MakeRoughnessSafe(SLAB_ROUGHNESSY(BSDFContext.BSDF), View.MinRoughness);
-				SpecularColor = ComputeF0(SLAB_SPECULAR(BSDFContext.BSDF), SLAB_BASECOLOR(BSDFContext.BSDF), SLAB_METALLIC(BSDFContext.BSDF));
+				SpecularColor = SLAB_F0(BSDFContext.BSDF);
 			}
 			break;
 		}
diff --git a/Engine/Shaders/Private/Strata/StrataEvaluation.ush b/Engine/Shaders/Private/Strata/StrataEvaluation.ush
--- a/Engine/Shaders/Private/Strata/StrataEvaluation.ush
+++ b/Engine/Shaders/Private/Strata/StrataEvaluation.ush
@@ -238,9 +238,10 @@ FStrataEvaluateResult StrataEvaluateBSDF(FStrataBSDFContext BSDFContext, FShadow
 		case STRATA_BSDF_TYPE_SLAB:
 		{
 			const bool bHasSubsurface = BSDF_GETHASSSS(BSDFContext.BSDF) > 0;
-			const float3 BaseColor = SLAB_BASECOLOR(BSDFContext.BSDF);
-			const float  Metallic = SLAB_METALLIC(BSDFContext.BSDF);
-			const float3 DiffuseColor = BaseColor - BaseColor * Metallic;
+			const float3 DiffuseColor = SLAB_DIFFUSEALBEDO(BSDFContext.BSDF);
+			const float3 F0			  = SLAB_F0(BSDFContext.BSDF);
+			const float3 F90		  = SLAB_F90(BSDFContext.BSDF);
+			const float Metallic	  = SLAB_METALLIC(BSDFContext.BSDF);
 			const float SafeRoughness = MakeRoughnessSafe(SLAB_ROUGHNESSX(BSDFContext.BSDF));
 			float a2 = Pow4(SafeRoughness);
 
@@ -248,9 +249,6 @@ FStrataEvaluateResult StrataEvaluateBSDF(FStrataBSDFContext BSDFContext, FShadow
 			Sample.DiffusePDF		= BSDFContext.SatNoL / PI;
 			Sample.bSubsurface		= bHasSubsurface;
 
-			const float3 F0  = ComputeF0(SLAB_SPECULAR(BSDFContext.BSDF), BaseColor, Metallic);
-			const float3 F90 = ComputeF90(F0, SLAB_EDGECOLOR(BSDFContext.BSDF), Metallic);
-
 			const float SafeRoughnessX = MakeRoughnessSafe(SLAB_ROUGHNESSX(BSDFContext.BSDF));
 			const float SafeRoughnessY = MakeRoughnessSafe(SLAB_ROUGHNESSY(BSDFContext.BSDF));
 
@@ -419,7 +417,7 @@ FStrataEvaluateResult StrataEvaluateBSDF(FStrataBSDFContext BSDFContext, FShadow
 				}
 				else
 				{
-					const float3 SubsurfaceAlebdo = SLAB_BASECOLOR(BSDFContext.BSDF);
+					const float3 SubsurfaceAlebdo = SLAB_DIFFUSEALBEDO(BSDFContext.BSDF);
 					const float3 MeanFreePathInCm = SLAB_SSSDMFP(BSDFContext.BSDF);
 					TransmissionThroughput	= GetBurleyTransmissionProfile(SubsurfaceAlebdo, MeanFreePathInCm, ThicknessInCm);
 					PhaseFunctionAnisotropy = 0.93f; // STRATA_TODO: expose phase function eccentricity. For now, forward scattering matching default SSS profile value
@@ -547,8 +545,9 @@ FStrataImportanceSampleResult StrataImportanceSampleBSDF(FStrataBSDFContext BSDF
 	{
 		case STRATA_BSDF_TYPE_SLAB:
 		{
-			const float3 BaseColor = SLAB_BASECOLOR(BSDFContext.BSDF);
-			const float  Metallic = SLAB_METALLIC(BSDFContext.BSDF);
+			const float3 DiffuseColor	= SLAB_DIFFUSEALBEDO(BSDFContext.BSDF);
+			const float3 F0 			= SLAB_F0(BSDFContext.BSDF);
+			const float3 F90 			= SLAB_F90(BSDFContext.BSDF);
 
 			// We assume a cosine hemisphere sampling is enough for any roughness.
 			// We also assume that this will work just as well for volumetric (considering it is an isotropic phase function and we are interesting in sampling mostly the hemisphere and not the layers below)
@@ -557,9 +556,6 @@ FStrataImportanceSampleResult StrataImportanceSampleBSDF(FStrataBSDFContext BSDF
 			Sample.DiffusePDF = ImportanceSample.w;
 			PDF = Sample.DiffusePDF;
 
-			const float3 F0 = ComputeF0(SLAB_SPECULAR(BSDF), BaseColor, Metallic);
-			// STRATA_TODO take into account F90 
-
 			const float SafeRoughnessX = MakeRoughnessSafe(SLAB_ROUGHNESSX(BSDFContext.BSDF));
 			const bool bHasAnisotropy = BSDF_GETANISOTROPY(BSDFContext.BSDF);
 			float4 TangentH = 0;
@@ -647,17 +643,16 @@ FStrataEnvLightResult StrataEvaluateForEnvLight(FStrataBSDFContext BSDFContext,
 	{
 		case STRATA_BSDF_TYPE_SLAB:
 		{
-			const float3 BaseColor = SLAB_BASECOLOR(BSDFContext.BSDF);
-			const float  Metallic = SLAB_METALLIC(BSDFContext.BSDF);
+			const float3 DiffuseColor	= SLAB_DIFFUSEALBEDO(BSDFContext.BSDF);
+			float3 F0 					= SLAB_F0(BSDFContext.BSDF);
+			const float3 F90 			= SLAB_F90(BSDFContext.BSDF);
+			const float Metallic 		= SLAB_METALLIC(BSDFContext.BSDF);
 
-			const float3 DiffuseColor = BaseColor - BaseColor * Metallic;
 			StrataEnvLightResult.DiffuseNormal = BSDFContext.N;
 			StrataEnvLightResult.bSubsurface = BSDF_GETHASSSS(BSDFContext.BSDF) > 0;
 			StrataEnvLightResult.DiffuseWeight = DiffuseColor;
 
 			float SafeRoughnessX = MakeRoughnessSafe(SLAB_ROUGHNESSX(BSDFContext.BSDF));
-			float3 F0 = ComputeF0(SLAB_SPECULAR(BSDFContext.BSDF), BaseColor, Metallic);
-			const float3 F90 = ComputeF90(F0, SLAB_EDGECOLOR(BSDFContext.BSDF), Metallic);
 			float3 ThinFilmFresnelWeight = 1;
 
 			const float3 ApproxFresnelForSpecularTransmittance = F_Schlick(F0, F90, BSDFContext.SatNoV);
diff --git a/Engine/Shaders/Private/Strata/StrataFurnaceTest.usf b/Engine/Shaders/Private/Strata/StrataFurnaceTest.usf
--- a/Engine/Shaders/Private/Strata/StrataFurnaceTest.usf
+++ b/Engine/Shaders/Private/Strata/StrataFurnaceTest.usf
@@ -52,6 +52,7 @@ void MainPS(
 		}
 	}
 
+#if 1
 	// Draw a simple 2D grid of spheres 
 	if (bEvaluate)
 	{
@@ -139,7 +140,7 @@ void MainPS(
 			BSDF_SETNORMALID(BSDF, 0);
 			BSDF_SETANISOTROPY(BSDF, Anisoptropy != 0 ? 1 : 0);
 			BSDF_SETISTOPLAYER(BSDF, 1);
-			BSDF_SETHASEDGECOLOR(BSDF, 0);
+			BSDF_SETHASF90(BSDF, 0);
 			BSDF_SETHASSSS(BSDF, 0);
 			BSDF_SETHASSSSPROFILE(BSDF, 0);
 			BSDF_SETHASHAZINESS(BSDF, Haziness != 0 ? 1 : 0);
@@ -148,12 +149,14 @@ void MainPS(
 
 			const float2 RoughnessXY = GetAnisotropicRoughness(Roughness, Anisoptropy);
 
-			SLAB_BASECOLOR(BSDF)	= BaseColor;
-			SLAB_EDGECOLOR(BSDF)	= EdgeColor;
+			const float3 F0 = ComputeF0(Specular.x, BaseColor, Metallic);
+			const float3 F90 = ComputeF90(F0, EdgeColor, Metallic);
+			SLAB_DIFFUSEALBEDO(BSDF)= BaseColor - BaseColor * Metallic;
+			SLAB_F0(BSDF)			= F0;
+			SLAB_F90(BSDF)			= F90;
+			SLAB_METALLIC(BSDF)		= Metallic;
 			SLAB_ROUGHNESSX(BSDF)	= RoughnessXY.x;
 			SLAB_ROUGHNESSY(BSDF)	= RoughnessXY.y;
-			SLAB_SPECULAR(BSDF)		= Specular;
-			SLAB_METALLIC(BSDF)		= Metallic;
 			SLAB_HAZINESS(BSDF)		= Haziness;
 		}
 
@@ -238,6 +241,8 @@ void MainPS(
 				DiffuseSkyLighting;;
 		}
 	}
+#endif
+
 }
 
 #endif // SHADER_FURNACE_ANALYTIC
diff --git a/Engine/Shaders/Private/Strata/StrataMaterialClassification.usf b/Engine/Shaders/Private/Strata/StrataMaterialClassification.usf
--- a/Engine/Shaders/Private/Strata/StrataMaterialClassification.usf
+++ b/Engine/Shaders/Private/Strata/StrataMaterialClassification.usf
@@ -74,9 +74,8 @@ void ComputePostProcessBSDFData(
 		{
 			Classification.bHasSubsurface = true;
 
-			SSS.BaseColor = SLAB_BASECOLOR(BSDF);
-			SSS.Specular  = SLAB_SPECULAR(BSDF);
-			// STRATA_TODO when layered translucency is figured out once again, only bottom layer should run SSS
+			SSS.BaseColor = SLAB_DIFFUSEALBEDO(BSDF);
+			SSS.Specular  = SLAB_F0(BSDF);				// STRATA_TODO: Specular or F0?
 
 			// Profile ID: we only keep the latest ID, and profile will overtake any non-profile based SSS
 			SSS.bIsValid			= true;
diff --git a/Engine/Shaders/Private/Strata/StrataVisualize.usf b/Engine/Shaders/Private/Strata/StrataVisualize.usf
--- a/Engine/Shaders/Private/Strata/StrataVisualize.usf
+++ b/Engine/Shaders/Private/Strata/StrataVisualize.usf
@@ -322,25 +322,25 @@ void VisualizeMaterialPS(
 				PRINTBSDFISTHIN();
 				NEXTLINE;
 
-				HORIZONTALTAB; PRINTPARAMCHAR(_A_); PRINTPARAMCHAR(_L_); PRINTPARAMCHAR(_B_); PRINTPARAMCHAR(_E_); PRINTPARAMCHAR(_D_); PRINTPARAMCHAR(_O_); 
-				HORIZONTALTAB; PRINTFLOAT3(SLAB_BASECOLOR(BSDF));
-				NEXTLINE;
-
-				if (BSDF_GETHASEDGECOLOR(BSDF))
-				{
-					HORIZONTALTAB; PRINTPARAMCHAR(_E_); PRINTPARAMCHAR(_D_); PRINTPARAMCHAR(_G_); PRINTPARAMCHAR(_E_); PRINTPARAMCHAR(_SPC_); PRINTPARAMCHAR(_SPC_);
-					HORIZONTALTAB; PRINTFLOAT3(SLAB_EDGECOLOR(BSDF));
-					NEXTLINE;
-				}
-
-				HORIZONTALTAB; PRINTPARAMCHAR(_R_); PRINTPARAMCHAR(_E_); PRINTPARAMCHAR(_F_); PRINTPARAMCHAR(_L_); PRINTPARAMCHAR(_E_); PRINTPARAMCHAR(_C_);
-				HORIZONTALTAB; PRINTFLOAT1(SLAB_SPECULAR(BSDF));
+				HORIZONTALTAB; PRINTPARAMCHAR(_D_); PRINTPARAMCHAR(_I_); PRINTPARAMCHAR(_F_); PRINTPARAMCHAR(_F_); PRINTPARAMCHAR(_U_); PRINTPARAMCHAR(_S_); 
+				HORIZONTALTAB; PRINTFLOAT3(SLAB_DIFFUSEALBEDO(BSDF));
 				NEXTLINE;
 
 				HORIZONTALTAB; PRINTPARAMCHAR(_M_); PRINTPARAMCHAR(_E_); PRINTPARAMCHAR(_T_); PRINTPARAMCHAR(_A_); PRINTPARAMCHAR(_L_); PRINTPARAMCHAR(_SPC_);
 				HORIZONTALTAB; PRINTFLOAT1(SLAB_METALLIC(BSDF));
 				NEXTLINE;
 
+				HORIZONTALTAB; PRINTPARAMCHAR(_F_); PRINTPARAMCHAR(_0_); PRINTPARAMCHAR(_SPC_); PRINTPARAMCHAR(_SPC_); PRINTPARAMCHAR(_SPC_); PRINTPARAMCHAR(_SPC_);
+				HORIZONTALTAB; PRINTFLOAT3(SLAB_F0(BSDF));
+				NEXTLINE;
+
+				if (BSDF_GETHASF90(BSDF))
+				{
+					HORIZONTALTAB; PRINTPARAMCHAR(_F_); PRINTPARAMCHAR(_9_); PRINTPARAMCHAR(_0_); PRINTPARAMCHAR(_SPC_); PRINTPARAMCHAR(_SPC_); PRINTPARAMCHAR(_SPC_);
+					HORIZONTALTAB; PRINTFLOAT3(SLAB_F90(BSDF));
+					NEXTLINE;
+				}
+
 				if (BSDF_GETANISOTROPY(BSDF))
 				{
 					float Roughness  = 0.f;
@@ -407,12 +407,11 @@ void VisualizeMaterialPS(
 #else
 				NEXTLINE;
 				NEXTLINE;
-				if (BSDF_GETHASEDGECOLOR(BSDF))
+				NEXTLINE;
+				if (BSDF_GETHASF90(BSDF))
 				{
 					NEXTLINE;
 				}
-				NEXTLINE;
-				NEXTLINE;
 				if (BSDF_GETANISOTROPY(BSDF))
 				{
 					NEXTLINE;
diff --git a/Engine/Source/Runtime/Engine/Classes/Materials/MaterialExpressionStrata.h b/Engine/Source/Runtime/Engine/Classes/Materials/MaterialExpressionStrata.h
--- a/Engine/Source/Runtime/Engine/Classes/Materials/MaterialExpressionStrata.h
+++ b/Engine/Source/Runtime/Engine/Classes/Materials/MaterialExpressionStrata.h
@@ -61,8 +61,24 @@ class UMaterialExpressionStrataSlabBSDF : public UMaterialExpressionStrataBSDF
 	 */
 	UPROPERTY()
 	FExpressionInput Specular;
+	
+	/**
+	 * Defines the diffused albedo, the percentage of light reflected as diffuse from the surface. (type = float3, unit = unitless, defaults to 0.18)
+	 */
+	UPROPERTY()
+	FExpressionInput DiffuseAlbedo;
+
+	/**
+	 * Defines F0, the percentage of light reflected as specular from a surface when the view is perpendicular to the surface. (type = float3, unit = unitless, defaults to plastic 0.04)
+	 */
+	UPROPERTY()
+	FExpressionInput F0;
 
-	// STRATA_TODO: edge or F82 EdgeColor?
+	/**
+	 * Defines F90, the percentage of light reflected as specular from a surface when the view is tangent to the surface. (type = float3, unit = unitless, defaults to 1.0f)
+	 */
+	UPROPERTY()
+	FExpressionInput F90;
 
 	/**
 	 * Controls how rough the Material is. Roughness of 0 (smooth) is a mirror reflection and 1 (rough) is completely matte or diffuse. When using anisotropy, it is the roughness used along the Tangent axis. (type = float, unit = unitless, defaults to 0.5)
@@ -140,6 +156,10 @@ class UMaterialExpressionStrataSlabBSDF : public UMaterialExpressionStrataBSDF
 	UPROPERTY(EditAnywhere, BlueprintReadOnly, Category = Material, meta = (DisplayName = "Subsurface Profile"))
 	TObjectPtr<class USubsurfaceProfile> SubsurfaceProfile;
 
+	/** Whether to use the metalness workflow relying on BaseColor, Specular, EdgeColor and Metallic inputs. Or use the DiffuseColor, F0 and F90 specification. */
+	UPROPERTY(EditAnywhere, Category = Mode)
+	uint32 bUseMetalness : 1;
+
 	//~ Begin UMaterialExpression Interface
 #if WITH_EDITOR
 	virtual int32 Compile(class FMaterialCompiler* Compiler, int32 OutputIndex) override;
@@ -149,6 +169,9 @@ class UMaterialExpressionStrataSlabBSDF : public UMaterialExpressionStrataBSDF
 	virtual bool IsResultStrataMaterial(int32 OutputIndex) override;
 	virtual void GatherStrataMaterialInfo(FStrataMaterialInfo& StrataMaterialInfo, int32 OutputIndex) override;
 	virtual FName GetInputName(int32 InputIndex) const override;
+	virtual void GetConnectorToolTip(int32 InputIndex, int32 OutputIndex, TArray<FString>& OutToolTip) override;
+	virtual void PostEditChangeProperty(FPropertyChangedEvent& PropertyChangedEvent) override;
+	virtual const TArray<FExpressionInput*> GetInputs() override;
 
 	bool HasEdgeColor() const;
 	bool HasThinFilm() const;
diff --git a/Engine/Source/Runtime/Engine/Private/Materials/HLSLMaterialTranslator.cpp b/Engine/Source/Runtime/Engine/Private/Materials/HLSLMaterialTranslator.cpp
--- a/Engine/Source/Runtime/Engine/Private/Materials/HLSLMaterialTranslator.cpp
+++ b/Engine/Source/Runtime/Engine/Private/Materials/HLSLMaterialTranslator.cpp
@@ -8951,8 +8951,10 @@ int32 FHLSLMaterialTranslator::StrataCreateAndRegisterNullMaterial()
 }
 
 int32 FHLSLMaterialTranslator::StrataSlabBSDF(
-	int32 BaseColor, int32 EdgeColor, int32 Specular, int32 Metallic, int32 Roughness, 
-	int32 Anisotropy,
+	int32 UseMetalness,
+	int32 BaseColor, int32 EdgeColor, int32 Specular, int32 Metallic,
+	int32 DiffuseAlbedo, int32 F0, int32 F90,
+	int32 Roughness, int32 Anisotropy,
 	int32 SSSProfileId, int32 SSSDMFP, int32 SSSDMFPScale, 
 	int32 EmissiveColor,	
 	int32 Haziness, 
@@ -8961,25 +8963,24 @@ int32 FHLSLMaterialTranslator::StrataSlabBSDF(
 	int32 Thickness, 
 	int32 Normal, int32 Tangent, const FString& SharedNormalIndexMacro)
 {
+	const FString NormalCode = GetParameterCode(Normal);
+
 	return AddCodeChunk(
-		MCT_Strata, TEXT("GetStrataSlabBSDF(%s, %s, %s, %s, %s, %s, %s, %s, %s, %s, %s, %s, %s, %s, %s, %s, Parameters.SharedNormals.NormalTypes) /* Normal = %s ; Tangent = %s */"),
-		*GetParameterCode(BaseColor),
-		*GetParameterCode(EdgeColor),
-		*GetParameterCode(Specular),
-		*GetParameterCode(Metallic),
-		*GetParameterCode(Roughness),
-		*GetParameterCode(Anisotropy),
+		MCT_Strata, TEXT("GetStrataSlabBSDF(%s, %s, %s, %s, %s, %s, %s, %s, %s, %s, %s, %s, %s, %s, %s, %s, %s, %s, %s, %s, Parameters.SharedNormals.NormalTypes) /* Normal = %s ; Tangent = %s */"),
+		* GetParameterCode(UseMetalness),
+		*GetParameterCode(BaseColor),		*GetParameterCode(EdgeColor),	*GetParameterCode(Specular), *GetParameterCode(Metallic),
+		*GetParameterCode(DiffuseAlbedo),	*GetParameterCode(F0),			*GetParameterCode(F90),
+		*GetParameterCode(Roughness),		*GetParameterCode(Anisotropy),
 		*GetParameterCode(SSSProfileId),
 		*GetParameterCode(SSSDMFP),
 		*GetParameterCode(SSSDMFPScale),
 		*GetParameterCode(EmissiveColor),
 		*GetParameterCode(Haziness),
 		*GetParameterCode(ThinFilmThickness),
-		*GetParameterCode(FuzzAmount),
-		*GetParameterCode(FuzzColor),
+		*GetParameterCode(FuzzAmount),		*GetParameterCode(FuzzColor),
 		*GetParameterCode(Thickness),
 		*SharedNormalIndexMacro,
-		*GetParameterCode(Normal),
+		*NormalCode,
 		Tangent!=INDEX_NONE ? *GetParameterCode(Tangent) : TEXT("NONE")
 	);
 }
diff --git a/Engine/Source/Runtime/Engine/Private/Materials/HLSLMaterialTranslator.h b/Engine/Source/Runtime/Engine/Private/Materials/HLSLMaterialTranslator.h
--- a/Engine/Source/Runtime/Engine/Private/Materials/HLSLMaterialTranslator.h
+++ b/Engine/Source/Runtime/Engine/Private/Materials/HLSLMaterialTranslator.h
@@ -977,7 +977,10 @@ protected:
 	// Strata
 	virtual int32 StrataCreateAndRegisterNullMaterial() override;
 	virtual int32 StrataSlabBSDF(
-		int32 BaseColor, int32 EdgeColor, int32 Specular, int32 Metallic, int32 Roughness, int32 Anisotropy,
+		int32 UseMetalness,
+		int32 BaseColor, int32 EdgeColor, int32 Specular, int32 Metallic,
+		int32 DiffuseAlbedo, int32 F0, int32 F90,
+		int32 Roughness, int32 Anisotropy,
 		int32 SSSProfileId, int32 SSSDMFP, int32 SSSDMFPScale,
 		int32 EmissiveColor, 
 		int32 Haziness, 
diff --git a/Engine/Source/Runtime/Engine/Private/Materials/Material.cpp b/Engine/Source/Runtime/Engine/Private/Materials/Material.cpp
--- a/Engine/Source/Runtime/Engine/Private/Materials/Material.cpp
+++ b/Engine/Source/Runtime/Engine/Private/Materials/Material.cpp
@@ -3687,11 +3687,11 @@ void UMaterial::ConvertMaterialToStrataMaterial()
 			MoveConnectionTo(BaseColor, SlabBSDF, 0);		// BaseColor
 			MoveConnectionTo(Metallic, SlabBSDF, 2);		// Metallic
 			MoveConnectionTo(Specular, SlabBSDF, 3);		// Specular
-			MoveConnectionTo(Roughness, SlabBSDF, 4);		// Roughness
-			MoveConnectionTo(Anisotropy, SlabBSDF, 5);		// Anisotropy
-			MoveConnectionTo(Normal, SlabBSDF, 6);			// Normal
-			MoveConnectionTo(Tangent, SlabBSDF, 7);			// Tangent
-			MoveConnectionTo(EmissiveColor, SlabBSDF, 10);	// Emissive
+			MoveConnectionTo(Roughness, SlabBSDF, 7);		// Roughness
+			MoveConnectionTo(Anisotropy, SlabBSDF, 8);		// Anisotropy
+			MoveConnectionTo(Normal, SlabBSDF, 9);			// Normal
+			MoveConnectionTo(Tangent, SlabBSDF, 10);		// Tangent
+			MoveConnectionTo(EmissiveColor, SlabBSDF, 13);	// Emissive
 
 			FrontMaterial.Connect(0, SlabBSDF);
 		}
@@ -3705,17 +3705,17 @@ void UMaterial::ConvertMaterialToStrataMaterial()
 			MoveConnectionTo(Roughness, SlabBSDF, 4);		// Roughness
 			if (Anisotropy.IsConnected())
 			{
-				MoveConnectionTo(Anisotropy, SlabBSDF, 5);	// Anisotropy
+				MoveConnectionTo(Anisotropy, SlabBSDF, 8);	// Anisotropy
 			}
-			MoveConnectionTo(Normal, SlabBSDF, 6);			// Normal
-			MoveConnectionTo(Tangent, SlabBSDF, 7);			// Tangent
+			MoveConnectionTo(Normal, SlabBSDF, 9);			// Normal
+			MoveConnectionTo(Tangent, SlabBSDF, 10);		// Tangent
 			
 			// SSS connections
 			if (ShadingModel == MSM_SubsurfaceProfile)
 			{
 				MoveConnectionTo(BaseColor, SlabBSDF, 0);		// BaseColor
 				SlabBSDF->SubsurfaceProfile = SubsurfaceProfile;// SSS Profile
-				MoveConnectionTo(Opacity, SlabBSDF, 9);			// SSSDMFPScale
+				MoveConnectionTo(Opacity, SlabBSDF, 12);		// SSSDMFPScale
 			}
 			else if (ShadingModel == MSM_Subsurface || ShadingModel == MSM_PreintegratedSkin)
 			{	
@@ -3763,13 +3763,13 @@ void UMaterial::ConvertMaterialToStrataMaterial()
 			MoveConnectionTo(BaseColor, SlabBSDF, 0);		// BaseColor
 			MoveConnectionTo(Metallic, SlabBSDF, 2);		// Metallic
 			MoveConnectionTo(Specular, SlabBSDF, 3);		// Specular
-			MoveConnectionTo(Roughness, SlabBSDF, 4);		// Roughness
-			MoveConnectionTo(Anisotropy, SlabBSDF, 5);		// Anisotropy
-			MoveConnectionTo(Normal, SlabBSDF, 6);			// Normal
-			MoveConnectionTo(Tangent, SlabBSDF, 7);			// Tangent
-			MoveConnectionTo(EmissiveColor, SlabBSDF, 10);	// Emissive
-			MoveConnectionTo(ClearCoat, SlabBSDF, 14);		// FuzzAMount
-			MoveConnectionTo(SubsurfaceColor, SlabBSDF, 15);// FuzzColor
+			MoveConnectionTo(Roughness, SlabBSDF, 7);		// Roughness
+			MoveConnectionTo(Anisotropy, SlabBSDF, 8);		// Anisotropy
+			MoveConnectionTo(Normal, SlabBSDF, 9);			// Normal
+			MoveConnectionTo(Tangent, SlabBSDF, 10);		// Tangent
+			MoveConnectionTo(EmissiveColor, SlabBSDF, 13);	// Emissive
+			MoveConnectionTo(ClearCoat, SlabBSDF, 17);		// FuzzAMount
+			MoveConnectionTo(SubsurfaceColor, SlabBSDF, 18);// FuzzColor
 
 			FrontMaterial.Connect(0, SlabBSDF);
 		}
@@ -3780,17 +3780,16 @@ void UMaterial::ConvertMaterialToStrataMaterial()
 			MoveConnectionTo(BaseColor, SlabBSDF, 0);		// BaseColor
 			MoveConnectionTo(Metallic, SlabBSDF, 2);		// Metallic
 			MoveConnectionTo(Specular, SlabBSDF, 3);		// Specular
-			MoveConnectionTo(Roughness, SlabBSDF, 4);		// Roughness
-			MoveConnectionTo(Anisotropy, SlabBSDF, 5);		// Anisotropy
-			MoveConnectionTo(Normal, SlabBSDF, 6);			// Normal
-			MoveConnectionTo(Tangent, SlabBSDF, 7);			// Tangent
-			MoveConnectionTo(EmissiveColor, SlabBSDF, 10);	// Emissive
-			MoveConnectionTo(Tangent, SlabBSDF, 7);			// Tangent
+			MoveConnectionTo(Roughness, SlabBSDF, 7);		// Roughness
+			MoveConnectionTo(Anisotropy, SlabBSDF, 8);		// Anisotropy
+			MoveConnectionTo(Normal, SlabBSDF, 9);			// Normal
+			MoveConnectionTo(Tangent, SlabBSDF, 10);		// Tangent
+			MoveConnectionTo(EmissiveColor, SlabBSDF, 13);	// Emissive
 
 			UMaterialExpressionStrataTransmittanceToMFP* TransmittanceToMFP = NewObject<UMaterialExpressionStrataTransmittanceToMFP>(this);
 			MoveConnectionTo(SubsurfaceColor, TransmittanceToMFP, 0);	// SubsurfaceColor -> TransmittanceColor
-			SlabBSDF->GetInput(8)->Connect(0, TransmittanceToMFP);		// MFP -> MFP
-			SlabBSDF->GetInput(13)->Connect(1, TransmittanceToMFP);		// Thickness -> Thickness
+			SlabBSDF->GetInput(11)->Connect(0, TransmittanceToMFP);		// MFP -> MFP
+			SlabBSDF->GetInput(16)->Connect(1, TransmittanceToMFP);		// Thickness -> Thickness
 
 			// Set a thickness that will enabled the thin lighting model (corresponding to the legacy two-sided lighting model)
 			UMaterialExpressionConstant* ThicknessConstant = NewObject<UMaterialExpressionConstant>(this);
@@ -3864,9 +3863,9 @@ void UMaterial::ConvertMaterialToStrataMaterial()
 			MoveConnectionTo(BaseColor, TopSlabBSDF, 0);								// BaseColor
 			MoveConnectionTo(Metallic, TopSlabBSDF, 2);									// Metallic
 			CopyConnectionTo(Specular, TopSlabBSDF, 3);									// Specular
-			CopyConnectionTo(Roughness, TopSlabBSDF, 4);								// Roughness
-			CopyConnectionTo(Normal, TopSlabBSDF, 6);									// Normal
-			MoveConnectionTo(EmissiveColor, TopSlabBSDF, 10);							// Emissive
+			CopyConnectionTo(Roughness, TopSlabBSDF, 7);								// Roughness
+			CopyConnectionTo(Normal, TopSlabBSDF, 9);									// Normal
+			MoveConnectionTo(EmissiveColor, TopSlabBSDF, 13);							// Emissive
 
 			// Now weight the top base material by opacity
 			UMaterialExpressionStrataMultiply* TopMaterialWithCoverage = NewObject<UMaterialExpressionStrataMultiply>(this);
@@ -3877,8 +3876,8 @@ void UMaterial::ConvertMaterialToStrataMaterial()
 			UMaterialExpressionStrataSlabBSDF* BottomSlabBSDF = NewObject<UMaterialExpressionStrataSlabBSDF>(this);
 			// Assign specular properties shared with the top layer.
 			MoveConnectionTo(Specular, BottomSlabBSDF, 3);								// Specular
-			MoveConnectionTo(Roughness, BottomSlabBSDF, 4);								// Roughness
-			MoveConnectionTo(Normal, BottomSlabBSDF, 6);								// Normal
+			MoveConnectionTo(Roughness, BottomSlabBSDF, 7);								// Roughness
+			MoveConnectionTo(Normal, BottomSlabBSDF, 9);								// Normal
 
 			TArray<class UMaterialExpressionCustomOutput*> CustomOutputExpressions;
 			GetAllCustomOutputExpressions(CustomOutputExpressions);
@@ -3895,8 +3894,8 @@ void UMaterial::ConvertMaterialToStrataMaterial()
 			{
 				UMaterialExpressionStrataTransmittanceToMFP* TransToMDFP = NewObject<UMaterialExpressionStrataTransmittanceToMFP>(this);
 				MoveConnectionTo(*TTCustomOutput->GetInput(0), TransToMDFP, 0);			// TransmittanceColor -> TransmittanceColor
-				BottomSlabBSDF->GetInput(8 )->Connect(0, TransToMDFP);					// MFP -> MFP
-				BottomSlabBSDF->GetInput(13)->Connect(1, TransToMDFP);					// Thickness -> Thickness
+				BottomSlabBSDF->GetInput(11)->Connect(0, TransToMDFP);					// MFP -> MFP
+				BottomSlabBSDF->GetInput(16)->Connect(1, TransToMDFP);					// Thickness -> Thickness
 
 				UMaterialExpressionStrataVerticalLayering* VerticalLayering = NewObject<UMaterialExpressionStrataVerticalLayering>(this);
 				VerticalLayering->GetInput(0)->Connect(0, TopMaterialWithCoverage);		// Top -> Top
@@ -3917,18 +3916,18 @@ void UMaterial::ConvertMaterialToStrataMaterial()
 			MoveConnectionTo(BaseColor, BottomSlabBSDF, 0);								// BaseColor
 			MoveConnectionTo(Metallic, BottomSlabBSDF, 2);								// Metallic
 			MoveConnectionTo(Specular, BottomSlabBSDF, 3);								// Specular
-			MoveConnectionTo(Roughness, BottomSlabBSDF, 4);								// Roughness
-			CopyConnectionTo(Anisotropy, BottomSlabBSDF, 5);							// Anisotropy
-			CopyConnectionTo(Normal, BottomSlabBSDF, 6);								// Normal
-			CopyConnectionTo(Tangent, BottomSlabBSDF, 7);								// Tangent
+			MoveConnectionTo(Roughness, BottomSlabBSDF, 7);								// Roughness
+			CopyConnectionTo(Anisotropy, BottomSlabBSDF, 8);							// Anisotropy
+			CopyConnectionTo(Normal, BottomSlabBSDF, 9);								// Normal
+			CopyConnectionTo(Tangent, BottomSlabBSDF, 10);								// Tangent
 
 			// Now weight the top base material by opacity.
 			UMaterialExpressionStrataSlabBSDF* TopSlabBSDF = NewObject<UMaterialExpressionStrataSlabBSDF>(this);
-			MoveConnectionTo(EmissiveColor, TopSlabBSDF, 10);							// Emissive
-			MoveConnectionTo(ClearCoatRoughness, TopSlabBSDF, 4);						// ClearCoatRoughness => Roughness
-			MoveConnectionTo(Anisotropy, TopSlabBSDF, 5);								// Anisotropy
-			MoveConnectionTo(Normal, TopSlabBSDF, 6);									// Normal
-			MoveConnectionTo(Tangent, TopSlabBSDF, 7);									// Tangent
+			MoveConnectionTo(EmissiveColor, TopSlabBSDF, 13);							// Emissive
+			MoveConnectionTo(ClearCoatRoughness, TopSlabBSDF, 7);						// ClearCoatRoughness => Roughness
+			MoveConnectionTo(Anisotropy, TopSlabBSDF, 8);								// Anisotropy
+			MoveConnectionTo(Normal, TopSlabBSDF, 9);									// Normal
+			MoveConnectionTo(Tangent, TopSlabBSDF, 10);									// Tangent
 
 			//  The top layer has a hard coded specular value of 0.5 (F0 = 0.04)
 			UMaterialExpressionConstant* ConstantHalf = NewObject<UMaterialExpressionConstant>(this);
@@ -3947,8 +3946,8 @@ void UMaterial::ConvertMaterialToStrataMaterial()
 			Constant075->R = 0.75f;
 			UMaterialExpressionStrataTransmittanceToMFP* TransToMDFP = NewObject<UMaterialExpressionStrataTransmittanceToMFP>(this);
 			TransToMDFP->GetInput(0)->Connect(0, Constant075);
-			TopSlabBSDF->GetInput(8)->Connect(0, TransToMDFP);							// MFP -> MFP
-			TopSlabBSDF->GetInput(13)->Connect(1, TransToMDFP);							// Thickness -> Thickness
+			TopSlabBSDF->GetInput(11)->Connect(0, TransToMDFP);							// MFP -> MFP
+			TopSlabBSDF->GetInput(16)->Connect(1, TransToMDFP);							// Thickness -> Thickness
 
 			// Connect the extra bottom normal if specified
 			TArray<class UMaterialExpressionCustomOutput*> CustomOutputExpressions;
diff --git a/Engine/Source/Runtime/Engine/Private/Materials/MaterialExpressions.cpp b/Engine/Source/Runtime/Engine/Private/Materials/MaterialExpressions.cpp
--- a/Engine/Source/Runtime/Engine/Private/Materials/MaterialExpressions.cpp
+++ b/Engine/Source/Runtime/Engine/Private/Materials/MaterialExpressions.cpp
@@ -19838,6 +19838,7 @@ int32 UMaterialExpressionStrataBSDF::CompilePreview(class FMaterialCompiler* Com
 
 UMaterialExpressionStrataSlabBSDF::UMaterialExpressionStrataSlabBSDF(const FObjectInitializer& ObjectInitializer)
 	: Super(ObjectInitializer)
+	, bUseMetalness(true)
 {
 	struct FConstructorStatics
 	{
@@ -19878,11 +19879,27 @@ int32 UMaterialExpressionStrataSlabBSDF::Compile(class FMaterialCompiler* Compil
 		SSSProfileCodeChunk = Compiler->ForceCast(Compiler->ScalarParameter(NameSubsurfaceProfile, 1.0f), MCT_Float1);
 	}
 
+	auto DielectricSpecularToF0 = [](float Specular)
+	{
+		return 0.08f * Specular;
+	};
+	const float DefaultSpecular = 0.5f;
+	const float DefaultF0 = DielectricSpecularToF0(DefaultSpecular);
+
 	int32 OutputCodeChunk = Compiler->StrataSlabBSDF(
+		bUseMetalness ? Compiler->Constant(1.0f) : Compiler->Constant(0.0f),
+
+		// Metalness workflow
 		CompileWithDefaultFloat3(Compiler, BaseColor, 0.18f, 0.18f, 0.18f),
 		CompileWithDefaultFloat3(Compiler, EdgeColor, 1.0f, 1.0f, 1.0f),
-		CompileWithDefaultFloat1(Compiler, Specular, 0.5f),
-		CompileWithDefaultFloat1(Compiler, Metallic, 0.0f),
+		CompileWithDefaultFloat1(Compiler, Specular, DefaultSpecular),
+		CompileWithDefaultFloat1(Compiler, Metallic,  0.0f),
+
+		// !Metalness workflow
+		CompileWithDefaultFloat3(Compiler, DiffuseAlbedo, 0.18f, 0.18f, 0.18f),
+		CompileWithDefaultFloat3(Compiler, F0, DefaultF0, DefaultF0, DefaultF0),
+		CompileWithDefaultFloat3(Compiler, F90, 1.0f, 1.0f, 1.0f),
+
 		RoughnessCodeChunk,
 		AnisotropyCodeChunk,
 		SSSProfileCodeChunk != INDEX_NONE ? SSSProfileCodeChunk : Compiler->Constant(0.0f),	
@@ -19902,6 +19919,47 @@ int32 UMaterialExpressionStrataSlabBSDF::Compile(class FMaterialCompiler* Compil
 	return OutputCodeChunk;
 }
 
+
+const TArray<FExpressionInput*> UMaterialExpressionStrataSlabBSDF::GetInputs()
+{
+	TArray<FExpressionInput*> Result;
+	if (bUseMetalness)
+	{
+		Result.Add(&BaseColor);
+		Result.Add(&EdgeColor);
+		Result.Add(&Metallic);
+		Result.Add(&Specular);
+	}
+	else
+	{
+		Result.Add(&DiffuseAlbedo);
+		Result.Add(&F0);
+		Result.Add(&F90);
+	}
+	Result.Add(&Roughness);
+	Result.Add(&Anisotropy);
+	Result.Add(&Normal);
+	Result.Add(&Tangent);
+	Result.Add(&SSSDMFP);
+	Result.Add(&SSSDMFPScale);
+	Result.Add(&EmissiveColor);
+	Result.Add(&Haziness);
+	Result.Add(&ThinFilmThickness);
+	Result.Add(&Thickness);
+	Result.Add(&FuzzAmount);
+	Result.Add(&FuzzColor);
+	return Result;
+}
+
+void UMaterialExpressionStrataSlabBSDF::PostEditChangeProperty(FPropertyChangedEvent& PropertyChangedEvent)
+{
+	if(GraphNode)
+	{
+		GraphNode->ReconstructNode();
+	}
+	Super::PostEditChangeProperty(PropertyChangedEvent);
+}
+
 void UMaterialExpressionStrataSlabBSDF::GetCaption(TArray<FString>& OutCaptions) const
 {
 	OutCaptions.Add(TEXT("Strata Slab BSDF"));
@@ -19914,56 +19972,90 @@ uint32 UMaterialExpressionStrataSlabBSDF::GetOutputType(int32 OutputIndex)
 
 uint32 UMaterialExpressionStrataSlabBSDF::GetInputType(int32 InputIndex)
 {
-	switch (InputIndex)
+	const int32 SkipDisabledInputOffset = bUseMetalness ? 0 : -1;
+
+	if (bUseMetalness)
+	{
+		if (InputIndex == 0)
+		{
+			return MCT_Float3; // BaseColor
+		}
+		else if (InputIndex == 1)
+		{
+			return MCT_Float3; // EdgeColor
+		}
+		else if (InputIndex == 2)
+		{
+			return MCT_Float1; // Metallic
+		}
+		else if (InputIndex == 3)
+		{
+			return MCT_Float1; // Specular
+		}
+	}
+	else
+	{
+		if (InputIndex == 0)
+		{
+			return MCT_Float3; // DiffuseAlbedo
+		}
+		else if (InputIndex == 1)
+		{
+			return MCT_Float3; // F0
+		}
+		else if (InputIndex == 2)
+		{
+			return MCT_Float3; // F90
+		}
+	}
+
+	if (InputIndex == (4 + SkipDisabledInputOffset))
 	{
-	case 0:
-		return MCT_Float3; // BaseColor
-		break;
-	case 1:
-		return MCT_Float3; // EdgeColor
-		break;
-	case 2:
-		return MCT_Float1; // Metallic
-		break;
-	case 3:
-		return MCT_Float1; // Specular
-		break;
-	case 4:
 		return MCT_Float1; // Roughness
-		break;
-	case 5:
+	}
+	else if (InputIndex == (5 + SkipDisabledInputOffset))
+	{
 		return MCT_Float1; // Anisotropy
-		break;
-	case 6:
+	}
+	else if (InputIndex == (6 + SkipDisabledInputOffset))
+	{
 		return MCT_Float3; // Normal
-		break;
-	case 7:
+	}
+	else if (InputIndex == (7 + SkipDisabledInputOffset))
+	{
 		return MCT_Float3; // Tangent
-		break;
-	case 8:
+	}
+	else if (InputIndex == (8 + SkipDisabledInputOffset))
+	{
 		return MCT_Float3; // SSSDMFP
-		break;
-	case 9:
+	}
+	else if (InputIndex == (9 + SkipDisabledInputOffset))
+	{
 		return MCT_Float1; // SSSDMFPScale
-		break;
-	case 10:
+	}
+	else if (InputIndex == (10 + SkipDisabledInputOffset))
+	{
 		return MCT_Float3; // Emissive Color
-		break;
-	case 11:
+	}
+	else if (InputIndex == (11 + SkipDisabledInputOffset))
+	{
 		return MCT_Float1; // Haziness
-		break;
-	case 12:
+	}
+	else if (InputIndex == (12 + SkipDisabledInputOffset))
+	{
 		return MCT_Float1; // ThinFilm Thickness
-		break;
-	case 13:
+	}
+	else if (InputIndex == (13 + SkipDisabledInputOffset))
+	{
 		return MCT_Float1; // Thickness
-		break;
-	case 14:
+	}
+	else if (InputIndex == (14 + SkipDisabledInputOffset))
+	{
 		return MCT_Float1; // FuzzAmount
-		break;
-	case 15:
+	}
+	else if (InputIndex == (15 + SkipDisabledInputOffset))
+	{
 		return MCT_Float3; // FuzzColor
-		break;
 	}
 
 	check(false);
@@ -19972,73 +20064,189 @@ uint32 UMaterialExpressionStrataSlabBSDF::GetInputType(int32 InputIndex)
 
 FName UMaterialExpressionStrataSlabBSDF::GetInputName(int32 InputIndex) const
 {
-	if (InputIndex == 0)
-	{
-		return TEXT("BaseColor");
-	}
-	else if (InputIndex == 1)
-	{
-		return TEXT("EdgeColor");
-	}
-	else if (InputIndex == 2)
+	const int32 SkipDisabledInputOffset = bUseMetalness ? 0 : -1;
+
+	if (bUseMetalness)
 	{
-		return TEXT("Metallic");
+		if (InputIndex == 0)
+		{
+			return TEXT("BaseColor");
+		}
+		else if (InputIndex == 1)
+		{
+			return TEXT("EdgeColor");
+		}
+		else if (InputIndex == 2)
+		{
+			return TEXT("Metallic");
+		}
+		else if (InputIndex == 3)
+		{
+			return TEXT("Specular");
+		}
 	}
-	else if (InputIndex == 3)
+	else
 	{
-		return TEXT("Specular");
+		if (InputIndex == 0)
+		{
+			return TEXT("DiffuseAlbedo");
+		}
+		else if (InputIndex == 1)
+		{
+			return TEXT("F0");
+		}
+		else if (InputIndex == 2)
+		{
+			return  TEXT("F90");
+		}
 	}
-	else if (InputIndex == 4)
+
+	if (InputIndex == (4 + SkipDisabledInputOffset))
 	{
 		return TEXT("Roughness");
 	}
-	else if (InputIndex == 5)
+	else if (InputIndex == (5 + SkipDisabledInputOffset))
 	{
 		return TEXT("Anisotropy");
 	}
-	else if (InputIndex == 6)
+	else if (InputIndex == (6 + SkipDisabledInputOffset))
 	{
 		return TEXT("Normal");
 	}
-	else if (InputIndex == 7)
+	else if (InputIndex == (7 + SkipDisabledInputOffset))
 	{
 		return TEXT("Tangent");
 	}
-	else if (InputIndex == 8)
+	else if (InputIndex == (8 + SkipDisabledInputOffset))
 	{
 		return TEXT("SSS Diffuse MFP");
 	}
-	else if (InputIndex == 9)
+	else if (InputIndex == (9 + SkipDisabledInputOffset))
 	{
 		return TEXT("SSS Diffuse MFP Scale");
 	}
-	else if (InputIndex == 10)
+	else if (InputIndex == (10 + SkipDisabledInputOffset))
 	{
 		return TEXT("Emissive Color");
 	}
-	else if (InputIndex == 11)
+	else if (InputIndex == (11 + SkipDisabledInputOffset))
 	{
 		return TEXT("Haziness");
 	}
-	else if (InputIndex == 12)
+	else if (InputIndex == (12 + SkipDisabledInputOffset))
 	{
 		return TEXT("ThinFilm Thickness");
 	}
-	else if (InputIndex == 13)
+	else if (InputIndex == (13 + SkipDisabledInputOffset))
 	{
 		return TEXT("Thickness");
 	}
-	else if (InputIndex == 14)
+	else if (InputIndex == (14 + SkipDisabledInputOffset))
 	{
 		return TEXT("FuzzAmount");
 	}
-	else if (InputIndex == 15)
+	else if (InputIndex == (15 + SkipDisabledInputOffset))
 	{
 		return TEXT("FuzzColor");
 	}
+
 	return TEXT("Unknown");
 }
 
+void UMaterialExpressionStrataSlabBSDF::GetConnectorToolTip(int32 InputIndex, int32 OutputIndex, TArray<FString>& OutToolTip)
+{
+	if (OutputIndex == 0)
+	{
+		OutToolTip.Add(TEXT("TT Ouput"));
+		return;
+	}
+	const int32 SkipDisabledInputOffset = bUseMetalness ? 0 : -1;
+
+	if (bUseMetalness)
+	{
+		if (InputIndex == 0)
+		{
+			Super::GetConnectorToolTip(0, INDEX_NONE, OutToolTip);
+		}
+		else if (InputIndex == 1)
+		{
+			Super::GetConnectorToolTip(1, INDEX_NONE, OutToolTip);
+		}
+		else if (InputIndex == 2)
+		{
+			Super::GetConnectorToolTip(2, INDEX_NONE, OutToolTip);
+		}
+		else if (InputIndex == 3)
+		{
+			Super::GetConnectorToolTip(3, INDEX_NONE, OutToolTip);
+		}
+	}
+	else
+	{
+		if (InputIndex == 0)
+		{
+			Super::GetConnectorToolTip(4, INDEX_NONE, OutToolTip);
+		}
+		else if (InputIndex == 1)
+		{
+			Super::GetConnectorToolTip(5, INDEX_NONE, OutToolTip);
+		}
+		else if (InputIndex == 2)
+		{
+			Super::GetConnectorToolTip(6, INDEX_NONE, OutToolTip);
+		}
+	}
+
+	if (InputIndex == (4 + SkipDisabledInputOffset))
+	{
+		Super::GetConnectorToolTip(7, INDEX_NONE, OutToolTip);
+	}
+	else if (InputIndex == (5 + SkipDisabledInputOffset))
+	{
+		Super::GetConnectorToolTip(8, INDEX_NONE, OutToolTip);
+	}
+	else if (InputIndex == (6 + SkipDisabledInputOffset))
+	{
+		Super::GetConnectorToolTip(9, INDEX_NONE, OutToolTip);
+	}
+	else if (InputIndex == (7 + SkipDisabledInputOffset))
+	{
+		Super::GetConnectorToolTip(10, INDEX_NONE, OutToolTip);
+	}
+	else if (InputIndex == (8 + SkipDisabledInputOffset))
+	{
+		Super::GetConnectorToolTip(11, INDEX_NONE, OutToolTip);
+	}
+	else if (InputIndex == (9 + SkipDisabledInputOffset))
+	{
+		Super::GetConnectorToolTip(12, INDEX_NONE, OutToolTip);
+	}
+	else if (InputIndex == (10 + SkipDisabledInputOffset))
+	{
+		Super::GetConnectorToolTip(13, INDEX_NONE, OutToolTip);
+	}
+	else if (InputIndex == (11 + SkipDisabledInputOffset))
+	{
+		Super::GetConnectorToolTip(14, INDEX_NONE, OutToolTip);
+	}
+	else if (InputIndex == (12 + SkipDisabledInputOffset))
+	{
+		Super::GetConnectorToolTip(15, INDEX_NONE, OutToolTip);
+	}
+	else if (InputIndex == (13 + SkipDisabledInputOffset))
+	{
+		Super::GetConnectorToolTip(16, INDEX_NONE, OutToolTip);
+	}
+	else if (InputIndex == (14 + SkipDisabledInputOffset))
+	{
+		Super::GetConnectorToolTip(17, INDEX_NONE, OutToolTip);
+	}
+	else if (InputIndex == (15 + SkipDisabledInputOffset))
+	{
+		Super::GetConnectorToolTip(18, INDEX_NONE, OutToolTip);
+	}
+}
+
 bool UMaterialExpressionStrataSlabBSDF::IsResultStrataMaterial(int32 OutputIndex)
 {
 	return true;
diff --git a/Engine/Source/Runtime/Engine/Public/MaterialCompiler.h b/Engine/Source/Runtime/Engine/Public/MaterialCompiler.h
--- a/Engine/Source/Runtime/Engine/Public/MaterialCompiler.h
+++ b/Engine/Source/Runtime/Engine/Public/MaterialCompiler.h
@@ -420,7 +420,10 @@ public:
 	// Strata
 	virtual int32 StrataCreateAndRegisterNullMaterial() = 0;
 	virtual int32 StrataSlabBSDF(
-		int32 BaseColor, int32 EdgeColor, int32 Specular, int32 Metallic, int32 Roughness, int32 Anisotropy, 
+		int32 UseMetalness,
+		int32 BaseColor, int32 EdgeColor, int32 Specular, int32 Metallic,
+		int32 DiffuseAlbedo, int32 F0, int32 F90,
+		int32 Roughness, int32 Anisotropy,
 		int32 SSSProfileId, int32 SSSDMFP, int32 SSSDMFPScale,
 		int32 EmissiveColor, 
 		int32 Haziness, 
@@ -914,7 +917,10 @@ public:
 	}
 
 	virtual int32 StrataSlabBSDF(
-		int32 BaseColor, int32 EdgeColor, int32 Specular, int32 Metallic, int32 Roughness, int32 Anisotropy,
+		int32 UseMetalness,
+		int32 BaseColor, int32 EdgeColor, int32 Specular, int32 Metallic,
+		int32 DiffuseAlbedo, int32 F0, int32 F90,
+		int32 Roughness, int32 Anisotropy,
 		int32 SSSProfileId, int32 SSSDMFP, int32 SSSDMFPScale,
 		int32 EmissiveColor, 
 		int32 Haziness, 
@@ -924,7 +930,10 @@ public:
 		int32 Normal, int32 Tangent, const FString& SharedNormalIndexMacro) override
 	{
 		return Compiler->StrataSlabBSDF(
-			BaseColor, EdgeColor, Specular, Metallic, Roughness, Anisotropy,
+			UseMetalness,
+			BaseColor, EdgeColor, Specular, Metallic,
+			DiffuseAlbedo, F0, F90,
+			Roughness, Anisotropy,
 			SSSProfileId, SSSDMFP, SSSDMFPScale,
 			EmissiveColor, 
 			Haziness, 
