Add DBufferTexture material expression to allow DBuffer direct access in material.
This allows decal receiver logic to do funky blends/or other things.
Note there is no guarantee about what is in the DBuffer so this functionality needs to be handled with caution.
If DBuffer is off then the sampling returns the same values as for an "empty" DBuffer (alpha = 1).
The change includes modifying where the MATERIALDECALRESPONSEMASK define is used in the shader. Validated this is working as expected and no unwanted shader samples sneak through the compilation.
Full DBuffer reponse (and associated texture samples) still applies when set on the material, even when also directly sampling DBuffer. But if desired the response can be disabled in the material and the default DBuffer logic applied in the material graph.

diff --git a/Engine/Shaders/Private/BasePassPixelShader.usf b/Engine/Shaders/Private/BasePassPixelShader.usf
--- a/Engine/Shaders/Private/BasePassPixelShader.usf
+++ b/Engine/Shaders/Private/BasePassPixelShader.usf
@@ -1033,7 +1033,7 @@ void FPixelShaderInOut_MainPS(
 			{
 				float2 NDC = MaterialParameters.ScreenPosition.xy / MaterialParameters.ScreenPosition.w;
 				float2 ScreenUV = NDC * ResolvedView.ScreenPositionScaleBias.xy + ResolvedView.ScreenPositionScaleBias.wz;
-				FDBufferData DBufferData = GetDBufferData(ScreenUV, DBufferMask);
+				FDBufferData DBufferData = GetDBufferData(ScreenUV, DBufferMask & MATERIALDECALRESPONSEMASK);
 
 				ApplyDBufferData(DBufferData, MaterialParameters.WorldNormal, SubsurfaceColor, Roughness, BaseColor, Metallic, Specular);
 				DBufferOpacity = (DBufferData.ColorOpacity + DBufferData.NormalOpacity + DBufferData.RoughnessOpacity) * (1.0f / 3.0f);
diff --git a/Engine/Shaders/Private/DBufferDecalShared.ush b/Engine/Shaders/Private/DBufferDecalShared.ush
--- a/Engine/Shaders/Private/DBufferDecalShared.ush
+++ b/Engine/Shaders/Private/DBufferDecalShared.ush
@@ -4,6 +4,8 @@
 	DBufferDecalShared.ush: Common definitions for DBuffer decals
 =============================================================================*/
 
+#pragma once
+
 // all values that are output by the forward rendering pass
 struct FDBufferData
 {
@@ -80,49 +82,53 @@ FDBufferData DecodeDBufferData(
 	#define OpaqueBasePassDBufferCTextureSampler OpaqueBasePass.DBufferCTextureSampler
 #endif
 
-#if MATERIALDECALRESPONSEMASK
 // @param UV - UV space in the DBuffer textures (BufferSize resolution)
-FDBufferData GetDBufferData(float2 UV, uint RTWriteMaskBit)
+FDBufferData GetDBufferData(float2 UV, uint RTMaskBit)
 {
 	float4 DBufferA = float4(0.0, 0.0, 0.0, 1.0);
-	float4 DBufferB = float4(0.5, 0.5, 0.5, 1.0);
+	float4 DBufferB = float4(128.f/255.5f, 128.f / 255.5f, 128.f / 255.5f, 1.0);
 	float4 DBufferC = float4(0.0, 0.0, 0.0, 1.0);
 
-	if ((MATERIALDECALRESPONSEMASK & 0x1) && (RTWriteMaskBit & 0x1))
+#if USE_DBUFFER && ENABLE_DBUFFER_TEXTURES
+	if (RTMaskBit & 0x1)
 	{
 		DBufferA = Texture2DSampleLevel(OpaqueBasePass.DBufferATexture, OpaqueBasePassDBufferATextureSampler, UV, 0);
 	}
 
-	if ((MATERIALDECALRESPONSEMASK & 0x2) && (RTWriteMaskBit & 0x2))
+	if (RTMaskBit & 0x2)
 	{
 		DBufferB = Texture2DSampleLevel(OpaqueBasePass.DBufferBTexture, OpaqueBasePassDBufferBTextureSampler, UV, 0);
 	}
 
-	if ((MATERIALDECALRESPONSEMASK & 0x4) && (RTWriteMaskBit & 0x4))
+	if (RTMaskBit & 0x4)
 	{
 		DBufferC = Texture2DSampleLevel(OpaqueBasePass.DBufferCTexture, OpaqueBasePassDBufferCTextureSampler, UV, 0);
 	}
-	
+#endif
+
 	return DecodeDBufferData(DBufferA, DBufferB, DBufferC);
 }
-#endif
 
 /** Populates DBufferA, DBufferB, DBufferC as float4 and puts opacity in alpha for frame buffer blending */
 void ApplyDBufferData(
 	FDBufferData DBufferData, inout float3 WorldNormal, inout float3 SubsurfaceColor, inout float Roughness, 
 	inout float3 BaseColor, inout float Metallic, inout float Specular )
 {
+#if (MATERIALDECALRESPONSEMASK & 0x1)
+	BaseColor = BaseColor * DBufferData.ColorOpacity + DBufferData.PreMulColor;
+	SubsurfaceColor *= DBufferData.ColorOpacity;
+#endif
+
+#if (MATERIALDECALRESPONSEMASK & 0x2)
 	// We normalise the normal to get smoother visual result (it helps to avoid having D_GGX explodes toward infinity, and matches DecodeGBufferData)
 	WorldNormal = normalize(WorldNormal * DBufferData.NormalOpacity + DBufferData.PreMulWorldNormal);
+#endif
 
+#if (MATERIALDECALRESPONSEMASK & 0x4)
 	Roughness = Roughness * DBufferData.RoughnessOpacity + DBufferData.PreMulRoughness;
-
 	Metallic = Metallic * DBufferData.RoughnessOpacity + DBufferData.PreMulMetallic;
 	Specular = Specular * DBufferData.RoughnessOpacity + DBufferData.PreMulSpecular;
-
-	SubsurfaceColor *= DBufferData.ColorOpacity;
-
-	BaseColor = BaseColor * DBufferData.ColorOpacity + DBufferData.PreMulColor;
+#endif
 }
 
 uint GetDBufferMask(in float2 ScreenPosition, in Texture2D<uint> DBufferMaskTexture)
diff --git a/Engine/Shaders/Private/MaterialTemplate.ush b/Engine/Shaders/Private/MaterialTemplate.ush
--- a/Engine/Shaders/Private/MaterialTemplate.ush
+++ b/Engine/Shaders/Private/MaterialTemplate.ush
@@ -18,6 +18,8 @@
 #include "/Engine/Private/HairShadingCommon.ush"
 #include "/Engine/Private/HairStrands/HairCardsAttributeCommon.ush"
 #include "/Engine/Private/HairStrands/HairStrandsAttributeCommon.ush"
+#include "/Engine/Private/DeferredShadingCommon.ush"
+#include "/Engine/Private/DBufferDecalShared.ush"
 
 #define USE_ANALYTIC_DERIVATIVES	(VF_NANITE)
 
@@ -1929,9 +1931,6 @@ MaterialFloat4 MobileSceneTextureLookup(inout FMaterialPixelParameters Parameter
 
 #if SHADING_PATH_DEFERRED
 
-#include "/Engine/Private/DeferredShadingCommon.ush"		// GetGBufferData()
-
-
 #if POST_PROCESS_MATERIAL
 /** Samples the screen-space velocity for the specified UV coordinates. */
 float2 PostProcessVelocityLookup(float Depth, float2 UV)
@@ -2077,6 +2076,18 @@ float3 DecodeSceneColorForMaterialNode(float2 ScreenUV)
 
 #endif // SHADING_PATH_DEFERRED
 
+float4 MaterialExpressionDBufferTextureLookup(float2 UV, int DBufferTextureIndex)
+{
+	FDBufferData DBufferData = GetDBufferData(UV, 0x07);
+	switch (DBufferTextureIndex)
+	{
+	case 0:	return float4(DBufferData.PreMulColor, DBufferData.ColorOpacity);
+	case 1:	return float4(DBufferData.PreMulWorldNormal, DBufferData.NormalOpacity);
+	case 2:	return float4(DBufferData.PreMulRoughness, DBufferData.PreMulMetallic, DBufferData.PreMulSpecular, DBufferData.RoughnessOpacity);
+	}
+
+	return float4(0, 0, 0, 1);
+}
 
 // DERIV_BASE_VALUE() is a disgusting macro to manage backwards compatibility while changing the generated materials. The existing nodes are not "aware"
 // of partial derivatives, and will use sprintf() to put the same line in both CalcPixelMaterialInputs() and CalcPixelMaterialInputsAnalyticDerivatives().
diff --git a/Engine/Source/Runtime/Engine/Classes/Materials/MaterialExpressionDBufferTexture.h b/Engine/Source/Runtime/Engine/Classes/Materials/MaterialExpressionDBufferTexture.h
new file mode 100644
--- /dev/null
+++ b/Engine/Source/Runtime/Engine/Classes/Materials/MaterialExpressionDBufferTexture.h
@@ -0,0 +1,40 @@
+// Copyright Epic Games, Inc. All Rights Reserved.
+
+#pragma once
+
+#include "CoreMinimal.h"
+#include "UObject/ObjectMacros.h"
+#include "Materials/MaterialExpression.h"
+#include "MaterialExpressionDBufferTexture.generated.h"
+
+UENUM()
+enum EDBufferTextureId
+{
+	/** DBuffer A (Base Color). RGB is pre-multiplied by alpha. A is one minus alpha. */
+	DBT_A UMETA(DisplayName = "DBufferA (BaseColor)"),
+	/** DBuffer B (World Normal). RGB is pre-multiplied by alpha. A is one minus alpha. */
+	DBT_B UMETA(DisplayName = "DBufferB (WorldNormal)"),
+	/** DBuffer C (Roughness). R is Roughness, G is Metallic, B is Specular,  all pre-multiplied by alpha. A is one minus alpha. */
+	DBT_C UMETA(DisplayName = "DBufferC (Roughness)"),
+};
+
+UCLASS(collapsecategories, hidecategories = Object)
+class UMaterialExpressionDBufferTexture : public UMaterialExpression
+{
+	GENERATED_UCLASS_BODY()
+
+	/** UV in 0..1 range */
+	UPROPERTY(meta = (RequiredInput = "false", ToolTip = "Ignored if not specified"))
+	FExpressionInput Coordinates;
+
+	/** Which DBuffer texture we want to make a lookup into. */
+	UPROPERTY(EditAnywhere, Category = UMaterialExpressionDBufferTexture, meta = (DisplayName = "DBuffer Texture"))
+	TEnumAsByte<EDBufferTextureId> DBufferTextureId;
+
+	//~ Begin UMaterialExpression Interface
+#if WITH_EDITOR
+	virtual int32 Compile(class FMaterialCompiler* Compiler, int32 OutputIndex) override;
+	virtual void GetCaption(TArray<FString>& OutCaptions) const override;
+#endif
+	//~ End UMaterialExpression Interface
+};
diff --git a/Engine/Source/Runtime/Engine/Private/Materials/HLSLMaterialTranslator.cpp b/Engine/Source/Runtime/Engine/Private/Materials/HLSLMaterialTranslator.cpp
--- a/Engine/Source/Runtime/Engine/Private/Materials/HLSLMaterialTranslator.cpp
+++ b/Engine/Source/Runtime/Engine/Private/Materials/HLSLMaterialTranslator.cpp
@@ -895,6 +895,23 @@ bool FHLSLMaterialTranslator::Translate()
 		bUsesAnisotropy = IsMaterialPropertyUsed(MP_Anisotropy, Chunk[MP_Anisotropy], FLinearColor(0, 0, 0, 0), 1);
 		MaterialCompilationOutput.bUsesAnisotropy = bUsesAnisotropy;
 
+		EMaterialDecalResponse MDR = (EMaterialDecalResponse)Material->GetMaterialDecalResponse();
+		if (MDR == MDR_Color || MDR == MDR_ColorNormal || MDR == MDR_ColorRoughness || MDR == MDR_ColorNormalRoughness)
+		{
+			MaterialCompilationOutput.SetIsDBufferTextureUsed(0);
+			AddEstimatedTextureSample(1);
+		}
+		if (MDR == MDR_Normal || MDR == MDR_ColorNormal || MDR == MDR_NormalRoughness || MDR == MDR_ColorNormalRoughness)
+		{
+			MaterialCompilationOutput.SetIsDBufferTextureUsed(1);
+			AddEstimatedTextureSample(1);
+		}
+		if (MDR == MDR_Roughness || MDR == MDR_ColorRoughness || MDR == MDR_NormalRoughness || MDR == MDR_ColorNormalRoughness)
+		{
+			MaterialCompilationOutput.SetIsDBufferTextureUsed(2);
+			AddEstimatedTextureSample(1);
+		}
+
 		if (BlendMode == BLEND_Modulate && MaterialShadingModels.IsLit() && !Material->IsDeferredDecal())
 		{
 			Errorf(TEXT("Dynamically lit translucency is not supported for BLEND_Modulate materials."));
@@ -6230,6 +6247,29 @@ int32 FHLSLMaterialTranslator::SceneColor(int32 Offset, int32 ViewportUV, bool b
 		);
 }
 
+int32 FHLSLMaterialTranslator::DBufferTextureLookup(int32 ViewportUV, uint32 DBufferTextureIndex)
+{
+	if (Material->GetMaterialDomain() != MD_Surface || IsTranslucentBlendMode(Material->GetBlendMode()))
+	{
+		Errorf(TEXT("DBuffer scene textures are only available on opaque or masked surfaces."));
+	}
+
+	int32 BufferUV = INDEX_NONE;
+	if (ViewportUV != INDEX_NONE)
+	{
+		BufferUV = AddCodeChunk(MCT_Float2,	TEXT("ClampSceneTextureUV(ViewportUVToBufferUV(%s), 0)"), *CoerceParameter(ViewportUV, MCT_Float2));
+	}
+	else
+	{
+		BufferUV = AddInlinedCodeChunk(MCT_Float2, TEXT("GetDefaultSceneTextureUV(Parameters, 0)"));
+	}
+
+	MaterialCompilationOutput.SetIsDBufferTextureUsed(DBufferTextureIndex);
+	AddEstimatedTextureSample();
+
+	return AddCodeChunk(MCT_Float4, TEXT("MaterialExpressionDBufferTextureLookup(%s, %d)"), *CoerceParameter(BufferUV, MCT_Float2), (int)DBufferTextureIndex);
+}
+
 int32 FHLSLMaterialTranslator::Texture(UTexture* InTexture, int32& TextureReferenceIndex, EMaterialSamplerType SamplerType, ESamplerSourceMode SamplerSource, ETextureMipValueMode MipValueMode)
 {
 	EMaterialValueType ShaderType = InTexture->GetMaterialType();
diff --git a/Engine/Source/Runtime/Engine/Private/Materials/HLSLMaterialTranslator.h b/Engine/Source/Runtime/Engine/Private/Materials/HLSLMaterialTranslator.h
--- a/Engine/Source/Runtime/Engine/Private/Materials/HLSLMaterialTranslator.h
+++ b/Engine/Source/Runtime/Engine/Private/Materials/HLSLMaterialTranslator.h
@@ -811,6 +811,8 @@ protected:
 
 	virtual int32 GetSceneTextureViewSize(int32 SceneTextureId, bool InvProperty) override;
 
+	virtual int32 DBufferTextureLookup(int32 ViewportUV, uint32 DBufferTextureIndex) override;
+
 	// @param bTextureLookup true: texture, false:no texture lookup, usually to get the size
 	void UseSceneTextureId(ESceneTextureId SceneTextureId, bool bTextureLookup);
 
diff --git a/Engine/Source/Runtime/Engine/Private/Materials/MaterialExpressions.cpp b/Engine/Source/Runtime/Engine/Private/Materials/MaterialExpressions.cpp
--- a/Engine/Source/Runtime/Engine/Private/Materials/MaterialExpressions.cpp
+++ b/Engine/Source/Runtime/Engine/Private/Materials/MaterialExpressions.cpp
@@ -77,6 +77,7 @@
 #include "Materials/MaterialExpressionCosine.h"
 #include "Materials/MaterialExpressionCrossProduct.h"
 #include "Materials/MaterialExpressionCustom.h"
+#include "Materials/MaterialExpressionDBufferTexture.h"
 #include "Materials/MaterialExpressionDDX.h"
 #include "Materials/MaterialExpressionDDY.h"
 #include "Materials/MaterialExpressionDecalDerivative.h"
@@ -9842,6 +9843,61 @@ void UMaterialExpressionSceneColor::GetCaption(TArray<FString>& OutCaptions) con
 }
 #endif // WITH_EDITOR
 
+///////////////////////////////////////////////////////////////////////////////
+// UMaterialExpressionDBufferTexture
+///////////////////////////////////////////////////////////////////////////////
+UMaterialExpressionDBufferTexture::UMaterialExpressionDBufferTexture(const FObjectInitializer& ObjectInitializer)
+	: Super(ObjectInitializer)
+{
+#if WITH_EDITORONLY_DATA
+	// Structure to hold one-time initialization
+	struct FConstructorStatics
+	{
+		FText NAME_Decals;
+		FConstructorStatics()
+			: NAME_Decals(LOCTEXT("Decals", "Decals"))
+		{
+		}
+	};
+	static FConstructorStatics ConstructorStatics;
+
+	MenuCategories.Add(ConstructorStatics.NAME_Decals);
+
+	bShaderInputData = true;
+	bShowOutputNameOnPin = true;
+#endif
+
+#if WITH_EDITORONLY_DATA
+	Outputs.Reset();
+	Outputs.Add(FExpressionOutput(TEXT("RGBA"), 1, 1, 1, 1, 1));
+	Outputs.Add(FExpressionOutput(TEXT("RGB"), 1, 1, 1, 1, 0));
+	Outputs.Add(FExpressionOutput(TEXT("A"), 1, 0, 0, 0, 1));
+#endif
+}
+
+#if WITH_EDITOR
+int32 UMaterialExpressionDBufferTexture::Compile(class FMaterialCompiler* Compiler, int32 OutputIndex)
+{
+	int32 ViewportUV = INDEX_NONE;
+
+	if (Coordinates.GetTracedInput().Expression)
+	{
+		ViewportUV = Coordinates.Compile(Compiler);
+	}
+
+	return Compiler->DBufferTextureLookup(ViewportUV, DBufferTextureId);
+}
+
+void UMaterialExpressionDBufferTexture::GetCaption(TArray<FString>& OutCaptions) const
+{
+	UEnum* Enum = StaticEnum<EDBufferTextureId>();
+	check(Enum);
+
+	FString Name = Enum->GetDisplayNameTextByValue(DBufferTextureId).ToString();
+	OutCaptions.Add(Name);
+}
+#endif // WITH_EDITOR
+
 UMaterialExpressionPower::UMaterialExpressionPower(const FObjectInitializer& ObjectInitializer)
 	: Super(ObjectInitializer)
 {
diff --git a/Engine/Source/Runtime/Engine/Public/MaterialCompiler.h b/Engine/Source/Runtime/Engine/Public/MaterialCompiler.h
--- a/Engine/Source/Runtime/Engine/Public/MaterialCompiler.h
+++ b/Engine/Source/Runtime/Engine/Public/MaterialCompiler.h
@@ -277,6 +277,7 @@ public:
 	// @param SceneTextureId of type ESceneTextureId e.g. PPI_SubsurfaceColor
 	virtual int32 SceneTextureLookup(int32 ViewportUV, uint32 SceneTextureId, bool bFiltered) = 0;
 	virtual int32 GetSceneTextureViewSize(int32 SceneTextureId, bool InvProperty) = 0;
+	virtual int32 DBufferTextureLookup(int32 ViewportUV, uint32 DBufferTextureIndex) = 0;
 
 	virtual int32 StaticBool(bool Value) = 0;
 	virtual int32 StaticBoolParameter(FName ParameterName,bool bDefaultValue) = 0;
@@ -623,6 +624,7 @@ public:
 	virtual int32 SceneColor(int32 Offset, int32 ViewportUV, bool bUseOffset) override { return Compiler->SceneColor(Offset, ViewportUV, bUseOffset); }
 	virtual int32 SceneTextureLookup(int32 ViewportUV, uint32 InSceneTextureId, bool bFiltered) override { return Compiler->SceneTextureLookup(ViewportUV, InSceneTextureId, bFiltered); }
 	virtual int32 GetSceneTextureViewSize(int32 SceneTextureId, bool InvProperty) override { return Compiler->GetSceneTextureViewSize(SceneTextureId, InvProperty); }
+	virtual int32 DBufferTextureLookup(int32 ViewportUV, uint32 DBufferTextureIndex) override { return Compiler->DBufferTextureLookup(ViewportUV, DBufferTextureIndex); }
 
 	virtual int32 StaticBool(bool Value) override { return Compiler->StaticBool(Value); }
 	virtual int32 StaticBoolParameter(FName ParameterName, bool bDefaultValue) override { return Compiler->StaticBoolParameter(ParameterName, bDefaultValue); }
diff --git a/Engine/Source/Runtime/Engine/Public/MaterialShared.h b/Engine/Source/Runtime/Engine/Public/MaterialShared.h
--- a/Engine/Source/Runtime/Engine/Public/MaterialShared.h
+++ b/Engine/Source/Runtime/Engine/Public/MaterialShared.h
@@ -554,6 +554,7 @@ public:
 		NumUsedCustomInterpolatorScalars(0),
 		StrataMaterialDescription(),
 #endif
+		UsedDBufferTextures(0),
 		RuntimeVirtualTextureOutputAttributeMask(0),
 		bNeedsSceneTextures(false),
 		bUsesEyeAdaptation(false),
@@ -569,6 +570,8 @@ public:
 	ENGINE_API bool IsSceneTextureUsed(ESceneTextureId TexId) const { return (UsedSceneTextures & (1 << TexId)) != 0; }
 	ENGINE_API void SetIsSceneTextureUsed(ESceneTextureId TexId) { UsedSceneTextures |= (1 << TexId); }
 
+	ENGINE_API void SetIsDBufferTextureUsed(int32 TextureIndex) { UsedDBufferTextures |= (1 << TextureIndex); }
+
 	/** Indicates whether the material uses scene color. */
 	ENGINE_API bool RequiresSceneColorCopy() const { return IsSceneTextureUsed(PPI_SceneColor); }
 
@@ -623,6 +626,9 @@ public:
 	/** The Strata material layout */
 	LAYOUT_FIELD_EDITORONLY(FMemoryImageString, StrataMaterialDescription);
 
+	/** Bitfield of used DBuffer textures . */
+	LAYOUT_FIELD(uint8, UsedDBufferTextures);
+
 	/** Bitfield of runtime virtual texture output attributes. */
 	LAYOUT_FIELD(uint8, RuntimeVirtualTextureOutputAttributeMask);
 
diff --git a/Engine/Source/Runtime/Renderer/Private/BasePassRendering.h b/Engine/Source/Runtime/Renderer/Private/BasePassRendering.h
--- a/Engine/Source/Runtime/Renderer/Private/BasePassRendering.h
+++ b/Engine/Source/Runtime/Renderer/Private/BasePassRendering.h
@@ -424,8 +424,8 @@ public:
 
 	static void ModifyCompilationEnvironment(const FMaterialShaderPermutationParameters& Parameters, FShaderCompilerEnvironment& OutEnvironment)
 	{
-		// For deferred decals, the shader class used is FDeferredDecalPS. the TBasePassPS is only used in the material editor and will read wrong values.
 		OutEnvironment.SetDefine(TEXT("SCENE_TEXTURES_DISABLED"), Parameters.MaterialParameters.MaterialDomain != MD_Surface);
+		OutEnvironment.SetDefine(TEXT("ENABLE_DBUFFER_TEXTURES"), Parameters.MaterialParameters.MaterialDomain == MD_Surface);
 		OutEnvironment.SetDefine(TEXT("COMPILE_BASEPASS_PIXEL_VOLUMETRIC_FOGGING"), DoesPlatformSupportVolumetricFog(Parameters.Platform));
 		OutEnvironment.SetDefine(TEXT("ENABLE_SKY_LIGHT"), bEnableSkyLight);
 		OutEnvironment.SetDefine(TEXT("PLATFORM_FORCE_SIMPLE_SKY_DIFFUSE"), ForceSimpleSkyDiffuse(Parameters.Platform));
