Strata: added parameter blending (first) version of BSDF Vertical layering.

diff --git a/Engine/Shaders/Private/Strata/Strata.ush b/Engine/Shaders/Private/Strata/Strata.ush
--- a/Engine/Shaders/Private/Strata/Strata.ush
+++ b/Engine/Shaders/Private/Strata/Strata.ush
@@ -276,8 +276,13 @@ struct FStrataBSDF
 #define SLW_COLORSCALEBEHINDWATER(X)	X.InlineVGPRs[2].xyz
 #endif
 
-void StrataRequestSharedNormalTangentBasis(inout uint NormalTypes, uint Index);
 
+// Forward declarations
+FParticipatingMedia StrataSlabCreateParticipatingMedia(float3 DiffuseColor, float3 MeanFreePathCentimeters);
+void EnableSlabBSDFSimpleVolumetric(inout FStrataBSDF BSDF);
+void StrataRequestSharedNormalTangentBasis(inout uint NormalTypes, uint Index);
+float3 StrataUnpackNormal(uint PackedNormal);
+void StrataUnpackNormalAndTangent(inout float3 Normal, inout float3 Tangent, in uint InPacked);
 uint   PackRGBA8(float4 rgba);
 float4 UnpackRGBA8(uint rgba);
 uint   PackR11G11B10F(float3 rgb);
@@ -725,9 +730,6 @@ uint  StrataLoadUint1(ByteAddressBuffer StrataBuffer, inout FStrataAddressing St
 ///////////////////////////////////////////////////////////////////////////////
 // Shared normals
 
-float3 StrataUnpackNormal(uint PackedNormal);
-void StrataUnpackNormalAndTangent(inout float3 Normal, inout float3 Tangent, in uint InPacked);
-
 #define STRATA_BASIS_TYPE_NORMAL 0u
 #define STRATA_BASIS_TYPE_TANGENT 1u
 struct FSharedNormals
@@ -903,6 +905,7 @@ FStrataData StrataVerticalLayering(FStrataData Top, FStrataData Base)
 //  - Maybe at some point we can propose a concatenate node, compressing a complex topology into a single one.
 //  - Slab node with SSSprofile cannot be blendend. Those will have to be specifically handled by the user and not mixed with any other slab if it can be demoted to parameter blending
 
+// ==> NOTE: Always pair with the compiler behavior in StrataCompilationInfoHorizontalMixingParamBlend
 FStrataData StrataHorizontalMixingParameterBlending(FStrataData Foreground, FStrataData Background, float Mix, uint NewNormalIndex)
 {
 	const float ForegroundMixFactor = saturate(Mix);
@@ -966,6 +969,7 @@ FStrataData StrataHorizontalMixingParameterBlending(FStrataData Foreground, FStr
 	return Result;
 }
 
+// ==> NOTE: Always pair with the compiler behavior in StrataCompilationInfoAddParamBlend
 FStrataData StrataAddParameterBlending(FStrataData A, FStrataData B, uint NewNormalIndex)
 {
 	const float MixFactor = 0.5f;
@@ -989,8 +993,8 @@ FStrataData StrataAddParameterBlending(FStrataData A, FStrataData B, uint NewNor
 	BSDF_SETHASF90(ResultBSDF,			BSDF_GETHASF90(ABSDF)							| BSDF_GETHASF90(BBSDF));
 	BSDF_SETHASTHINFILM(ResultBSDF,		BSDF_GETHASTHINFILM(ABSDF)						| BSDF_GETHASTHINFILM(BBSDF));
 	BSDF_SETHASFUZZ(ResultBSDF,			BSDF_GETHASFUZZ(ABSDF)							| BSDF_GETHASFUZZ(BBSDF));
-	BSDF_SETEMISSIVE(ResultBSDF,		BSDF_GETEMISSIVE(BBSDF)							+ BSDF_GETEMISSIVE(ABSDF));
-	BSDF_SETTHICKNESSCM(ResultBSDF,		BSDF_GETTHICKNESSCM(BBSDF)						+ BSDF_GETTHICKNESSCM(ABSDF));
+	BSDF_SETEMISSIVE(ResultBSDF,		BSDF_GETEMISSIVE(ABSDF)							+ BSDF_GETEMISSIVE(BBSDF));
+	BSDF_SETTHICKNESSCM(ResultBSDF,		lerp(BSDF_GETTHICKNESSCM(ABSDF)					, BSDF_GETTHICKNESSCM(BBSDF),				MixFactor));
 
 	// Some parameters will contribute to added luminance: in this case we add them together and saturate to not go out of the safe range
 	// Some parameters are not contributing to added luminance: in this case we simply lerp them
@@ -1031,6 +1035,80 @@ FStrataData StrataAddParameterBlending(FStrataData A, FStrataData B, uint NewNor
 	return Result;
 }
 
+// ==> NOTE: Always pair with the compiler behavior in StrataCompilationInfoVerticalLayeringParamBlend
+FStrataData StrataVerticalLayeringParameterBlending(FStrataData Top, FStrataData Base, uint NewNormalIndex, float NoV)
+{
+	FStrataData Result = GetInitialisedStrataDataWithReservedFirstSlot();
+#define ResultBSDF		Result.Layers[0].BSDFs[0]
+#define TopBSDF			Top.Layers[0].BSDFs[0]
+#define BaseBSDF		Base.Layers[0].BSDFs[0]
+
+	// We can only parameter blend the slab 
+	BSDF_SETTYPE(ResultBSDF,			STRATA_BSDF_TYPE_SLAB);
+	BSDF_SETNORMALID(ResultBSDF,		NewNormalIndex);
+
+	BSDF_SETISTOPLAYER(ResultBSDF,		0);
+	BSDF_SETANISOTROPY(ResultBSDF,		0);																					// Anisotropy is disabled because it changes a lot of things
+	BSDF_SETHASSSS(ResultBSDF,			BSDF_GETHASSSS(BaseBSDF));															// We only keep SSS if the bottom layer has it
+	BSDF_SETHASSSSPROFILE(ResultBSDF,	0);																					// Problem: cannot blend SSS profiles so we simply disable SSSProfile
+	BSDF_SETISTHIN(ResultBSDF,			BSDF_GETISTHIN(BaseBSDF));															// STRATA_TODO fix the MPF/DMFP situation for that to really work
+	BSDF_SETHASDMFP(ResultBSDF,			BSDF_GETHASDMFP(BaseBSDF));															//We only keep SSS if the bottom layer has it
+	BSDF_SETHASHAZINESS(ResultBSDF,		BSDF_GETHASHAZINESS(TopBSDF));														// We only keep haziness if the top layer has it
+	BSDF_SETHASF90(ResultBSDF,			BSDF_GETHASF90(BaseBSDF)						| BSDF_GETHASF90(TopBSDF));			// F90: keep union of both even though it will be hard to get a match
+	BSDF_SETHASTHINFILM(ResultBSDF,		BSDF_GETHASTHINFILM(TopBSDF));														// We only keep thinfilme if the top layer has it
+	BSDF_SETHASFUZZ(ResultBSDF,			BSDF_GETHASFUZZ(BaseBSDF)						| BSDF_GETHASFUZZ(TopBSDF));		// Fuzz: keep union of both even though it will be hard to get a match
+
+	// Compute the top slab transmittance that will be used to blend the top and base material parameters.
+	// We do not check BSDF_GETISSIMPLEVOLUME(BSDFContext.BSDF) because that is set when we know the BSDF is not at the bottom later. And here we already know it is not the case.
+	EnableSlabBSDFSimpleVolumetric(TopBSDF);
+	const float3 UnusedDiffuseColor		= 0;	// Unused when only transmittance is needed
+	FParticipatingMedia PM				= StrataSlabCreateParticipatingMedia(UnusedDiffuseColor, SLAB_SSSDMFP(TopBSDF));
+	const float3 SlabTransmittance		= IsotropicMediumSlabTransmittance(PM, STRATA_SIMPLEVOLUME_THICKNESS_M, NoV);
+	const float3 ColorThroughput		= SlabTransmittance;
+	const float GreyThroughput			= dot(ColorThroughput, 1.0f / 3.0f);
+	
+	// Add both layers thickness
+	BSDF_SETTHICKNESSCM(ResultBSDF,		BSDF_GETTHICKNESSCM(BaseBSDF)					+ BSDF_GETTHICKNESSCM(TopBSDF));
+
+	// Now add bottom layer component weights by throughput to the top layer.
+	// Those parameters will contribute to added luminance, in this case we add them together and saturate to not go out of the safe range
+	// Some parameters are not contributing to added luminance: in this case we simply lerp them
+	BSDF_SETEMISSIVE(ResultBSDF,		ColorThroughput * BSDF_GETEMISSIVE(BaseBSDF)	+ BSDF_GETEMISSIVE(TopBSDF));
+
+	// Strata TODO: we should also add the scattering light contribution from the top layer there.
+	SLAB_DIFFUSEALBEDO(ResultBSDF) =	saturate(ColorThroughput * SLAB_DIFFUSEALBEDO(BaseBSDF)+ SLAB_DIFFUSEALBEDO(TopBSDF));
+	SLAB_F0(ResultBSDF) =				saturate(ColorThroughput * SLAB_F0(BaseBSDF)	+ SLAB_F0(TopBSDF));
+	SLAB_F90(ResultBSDF) =				saturate(ColorThroughput * SLAB_F90(BaseBSDF)	+ SLAB_F90(TopBSDF));
+
+	SLAB_FUZZ_AMOUNT(ResultBSDF) =		saturate(GreyThroughput * SLAB_FUZZ_AMOUNT(BaseBSDF)+ SLAB_FUZZ_AMOUNT(TopBSDF));
+	SLAB_FUZZ_COLOR(ResultBSDF) =		saturate(ColorThroughput * SLAB_FUZZ_COLOR(BaseBSDF)+ SLAB_FUZZ_COLOR(TopBSDF));
+
+	// Simple blend for now. 
+	// STRATA_TODO find someting smarter?
+	SLAB_METALLIC(ResultBSDF) =			lerp(SLAB_METALLIC(TopBSDF),					SLAB_METALLIC(BaseBSDF),					GreyThroughput);
+	SLAB_ROUGHNESSX(ResultBSDF) =		lerp(SLAB_ROUGHNESSX(TopBSDF),					SLAB_ROUGHNESSX(BaseBSDF),					GreyThroughput);
+	SLAB_HAZINESS(ResultBSDF) =			lerp(SLAB_HAZINESS(TopBSDF),					SLAB_HAZINESS(BaseBSDF),					GreyThroughput);
+
+	// No anisotropy 
+	SLAB_ROUGHNESSY(ResultBSDF) =		SLAB_ROUGHNESSX(ResultBSDF);
+
+	// Only keep thin film and for the top layer
+	SLAB_THINFILM_THICKNESS(ResultBSDF)=SLAB_THINFILM_THICKNESS(TopBSDF);
+
+#if STRATA_INLINE_SHADING
+	// Only the bottom layer SSS / Thin model is used so we only keep the bottom layer data there. Top layer is converted to throughput.
+	ResultBSDF.TmpDMFP =				BaseBSDF.TmpDMFP;
+#endif
+
+	// When it comes to adding the visual contribution, it is tricky to find a matching operator for coverage. We simply lerp for now.
+	ResultBSDF.Weight =					saturate(GreyThroughput * BaseBSDF.Weight		+ SLAB_DIFFUSEALBEDO(TopBSDF));
+
+#undef ResultBSDF
+#undef TopBSDF
+#undef BaseBSDF
+	return Result;
+}
+
 
 ///////////////////////////////////////////////////////////////////////////////
 // Functions used to sanitize BSDF before they are used for lighting (when forward) or stored to memory (base pass).
diff --git a/Engine/Source/Runtime/Engine/Classes/Materials/MaterialExpressionStrata.h b/Engine/Source/Runtime/Engine/Classes/Materials/MaterialExpressionStrata.h
--- a/Engine/Source/Runtime/Engine/Classes/Materials/MaterialExpressionStrata.h
+++ b/Engine/Source/Runtime/Engine/Classes/Materials/MaterialExpressionStrata.h
@@ -456,6 +456,9 @@ class UMaterialExpressionStrataVerticalLayering : public UMaterialExpressionStra
 	UPROPERTY()
 	FExpressionInput Base;
 
+	UPROPERTY(EditAnywhere, Category = Mode)
+	uint32 bUseParameterBlending : 1;
+
 	//~ Begin UMaterialExpression Interface
 #if WITH_EDITOR
 	virtual int32 Compile(class FMaterialCompiler* Compiler, int32 OutputIndex) override;
diff --git a/Engine/Source/Runtime/Engine/Private/Materials/HLSLMaterialTranslator.cpp b/Engine/Source/Runtime/Engine/Private/Materials/HLSLMaterialTranslator.cpp
--- a/Engine/Source/Runtime/Engine/Private/Materials/HLSLMaterialTranslator.cpp
+++ b/Engine/Source/Runtime/Engine/Private/Materials/HLSLMaterialTranslator.cpp
@@ -9134,10 +9134,21 @@ int32 FHLSLMaterialTranslator::StrataVerticalLayering(int32 Top, int32 Base)
 	);
 }
 
-//virtual int32 FHLSLMaterialTranslator::StrataVerticalLayeringParameterBlending(int32 Top, int32 Base, const FString& SharedNormalIndexMacro)
-//{
-//	// STRATA_TODO
-//}
+int32 FHLSLMaterialTranslator::StrataVerticalLayeringParameterBlending(int32 Top, int32 Base, const FString& SharedNormalIndexMacro, int32 TopBSDFNormalCodeChunk)
+{
+	if (Top == INDEX_NONE || Base == INDEX_NONE)
+	{
+		return INDEX_NONE;
+	}
+	return AddCodeChunk(
+		MCT_Strata, TEXT("StrataVerticalLayeringParameterBlending(%s, %s, %s, dot(%s, %s))"),
+		*GetParameterCode(Top),
+		*GetParameterCode(Base),
+		*SharedNormalIndexMacro,
+		*GetParameterCode(TopBSDFNormalCodeChunk),
+		*GetParameterCode(CameraVector())
+	);
+}
 
 int32 FHLSLMaterialTranslator::StrataAdd(int32 A, int32 B)
 {
diff --git a/Engine/Source/Runtime/Engine/Private/Materials/HLSLMaterialTranslator.h b/Engine/Source/Runtime/Engine/Private/Materials/HLSLMaterialTranslator.h
--- a/Engine/Source/Runtime/Engine/Private/Materials/HLSLMaterialTranslator.h
+++ b/Engine/Source/Runtime/Engine/Private/Materials/HLSLMaterialTranslator.h
@@ -1001,7 +1001,7 @@ protected:
 	virtual int32 StrataHorizontalMixing(int32 Foreground, int32 Background, int32 Mix) override;
 	virtual int32 StrataHorizontalMixingParameterBlending(int32 Foreground, int32 Background, int32 Mix, const FString& SharedNormalIndexMacro) override;
 	virtual int32 StrataVerticalLayering(int32 Top, int32 Base) override;
-	// STRATA_TODO virtual int32 StrataVerticalLayeringParameterBlending(int32 Top, int32 Base, const FString& SharedNormalIndexMacro) override;
+	virtual int32 StrataVerticalLayeringParameterBlending(int32 Top, int32 Base, const FString& SharedNormalIndexMacro, int32 TopBSDFNormalCodeChunk) override;
 	virtual int32 StrataAdd(int32 A, int32 B) override;
 	virtual int32 StrataAddParameterBlending(int32 A, int32 B, const FString& SharedNormalIndexMacro) override;
 	virtual int32 StrataMultiply(int32 A, int32 Weight) override;
diff --git a/Engine/Source/Runtime/Engine/Private/Materials/MaterialExpressions.cpp b/Engine/Source/Runtime/Engine/Private/Materials/MaterialExpressions.cpp
--- a/Engine/Source/Runtime/Engine/Private/Materials/MaterialExpressions.cpp
+++ b/Engine/Source/Runtime/Engine/Private/Materials/MaterialExpressions.cpp
@@ -20841,6 +20841,7 @@ void UMaterialExpressionStrataHorizontalMixing::GatherStrataMaterialInfo(FStrata
 
 UMaterialExpressionStrataVerticalLayering::UMaterialExpressionStrataVerticalLayering(const FObjectInitializer& ObjectInitializer)
 	: Super(ObjectInitializer)
+	, bUseParameterBlending(false)
 {
 	struct FConstructorStatics
 	{
@@ -20867,22 +20868,86 @@ int32 UMaterialExpressionStrataVerticalLayering::Compile(class FMaterialCompiler
 
 	int32 TopCodeChunk = Top.Compile(Compiler);
 	int32 BaseCodeChunk = Base.Compile(Compiler);
+	if (!Compiler->StrataCompilationInfoContainsCodeChunk(TopCodeChunk) || !Compiler->StrataCompilationInfoContainsCodeChunk(BaseCodeChunk))
+	{
+		return Compiler->Errorf(TEXT("Could not find Top CodeChunk or BaseCodeChunk to VerticalLayer"));
+	}
+	const FStrataMaterialCompilationInfo& TopStrataData = Compiler->GetStrataCompilationInfo(TopCodeChunk);
+	const FStrataMaterialCompilationInfo& BaseStrataData = Compiler->GetStrataCompilationInfo(BaseCodeChunk);
 
-	int32 OutputCodeChunk = Compiler->StrataVerticalLayering(TopCodeChunk, BaseCodeChunk);
+	const int32 MixCodeChunk = Compiler->Constant(0.5f);
 
-	if (!Compiler->StrataCompilationInfoContainsCodeChunk(TopCodeChunk) || !Compiler->StrataCompilationInfoContainsCodeChunk(BaseCodeChunk))
+	int32 OutputCodeChunk = INDEX_NONE;
+	if (bUseParameterBlending)
 	{
-		return Compiler->Errorf(TEXT("Could not find TopCodeChunk or BaseCodeChunk to layer"));
+		if (TopStrataData.TotalBSDFCount != 1)
+		{
+			return Compiler->Errorf(TEXT("Top input: cannot vertical layer complex material topology using parameter blending. Only a one-to-one blending is supported today."));
+		}
+		if (BaseStrataData.TotalBSDFCount != 1)
+		{
+			return Compiler->Errorf(TEXT("Base input: cannot vertical layer complex material topology using parameter blending. Only a one-to-one blending is supported today."));
+		}
+
+		const FStrataMaterialCompilationInfo::FBSDF& TopBSDF = TopStrataData.Layers[0].BSDFs[0];
+		const FStrataMaterialCompilationInfo::FBSDF& BaseBSDF = BaseStrataData.Layers[0].BSDFs[0];
+		if (TopBSDF.Type != STRATA_BSDF_TYPE_SLAB)
+		{
+			return Compiler->Errorf(TEXT("Top input: we can only apply parameter blending to Slab BSDF."));
+		}
+		if (BaseBSDF.Type != STRATA_BSDF_TYPE_SLAB)
+		{
+			return Compiler->Errorf(TEXT("Base input: we can only apply parameter blending to Slab BSDF."));
+		}
+
+		// Compute the new Normal and Tangent resulting from the blending using code chunk
+		const int32 NewNormalCodeChunk = StrataBlendNormal(Compiler, TopBSDF.RegisteredSharedNormal.NormalCodeChunk, BaseBSDF.RegisteredSharedNormal.NormalCodeChunk, MixCodeChunk);
+		// The tangent is optional so we treat it differently if INDEX_NONE is specified
+		int32 NewTangentCodeChunk = INDEX_NONE;
+		if (TopBSDF.RegisteredSharedNormal.TangentCodeChunk != INDEX_NONE && BaseBSDF.RegisteredSharedNormal.TangentCodeChunk != INDEX_NONE)
+		{
+			NewTangentCodeChunk = StrataBlendNormal(Compiler, TopBSDF.RegisteredSharedNormal.TangentCodeChunk, BaseBSDF.RegisteredSharedNormal.TangentCodeChunk, MixCodeChunk);
+		}
+		else if (TopBSDF.RegisteredSharedNormal.TangentCodeChunk != INDEX_NONE)
+		{
+			NewTangentCodeChunk = TopBSDF.RegisteredSharedNormal.TangentCodeChunk;
+		}
+		else if (BaseBSDF.RegisteredSharedNormal.TangentCodeChunk != INDEX_NONE)
+		{
+			NewTangentCodeChunk = BaseBSDF.RegisteredSharedNormal.TangentCodeChunk;
+		}
+		const FStrataRegisteredSharedNormal NewRegisteredSharedNormal = StrataCompilationInfoCreateSharedNormal(Compiler, NewNormalCodeChunk, NewTangentCodeChunk);
+
+		OutputCodeChunk = Compiler->StrataVerticalLayeringParameterBlending(TopCodeChunk, BaseCodeChunk, Compiler->GetStrataSharedNormalIndexMacro(NewRegisteredSharedNormal), TopBSDF.RegisteredSharedNormal.NormalCodeChunk);
+
+		StrataCompilationInfoCreateSingleBSDFMaterial(Compiler, OutputCodeChunk, NewRegisteredSharedNormal, STRATA_BSDF_TYPE_SLAB);
+		FStrataMaterialCompilationInfo StrataInfo = StrataCompilationInfoVerticalLayeringParamBlend(Compiler, TopStrataData, BaseStrataData, NewRegisteredSharedNormal);
+
+		Compiler->StrataCompilationInfoRegisterCodeChunk(OutputCodeChunk, StrataInfo);
 	}
-	FStrataMaterialCompilationInfo StrataInfo = StrataCompilationInfoVerticalLayering(Compiler, Compiler->GetStrataCompilationInfo(TopCodeChunk), Compiler->GetStrataCompilationInfo(BaseCodeChunk));
-	Compiler->StrataCompilationInfoRegisterCodeChunk(OutputCodeChunk, StrataInfo);
+	else
+	{
+		OutputCodeChunk = Compiler->StrataVerticalLayering(TopCodeChunk, BaseCodeChunk);
+
+		FStrataMaterialCompilationInfo StrataInfo = StrataCompilationInfoVerticalLayering(Compiler, Compiler->GetStrataCompilationInfo(TopCodeChunk), Compiler->GetStrataCompilationInfo(BaseCodeChunk));
+
+		Compiler->StrataCompilationInfoRegisterCodeChunk(OutputCodeChunk, StrataInfo);
+	}
+
 
 	return OutputCodeChunk;
 }
 
 void UMaterialExpressionStrataVerticalLayering::GetCaption(TArray<FString>& OutCaptions) const
 {
-	OutCaptions.Add(TEXT("Strata BSDF Vertical Layer"));
+	if (bUseParameterBlending)
+	{
+		OutCaptions.Add(TEXT("Strata BSDF Vertical Layer (Parameter Blend)"));
+	}
+	else
+	{
+		OutCaptions.Add(TEXT("Strata BSDF Vertical Layer"));
+	}
 }
 
 uint32 UMaterialExpressionStrataVerticalLayering::GetOutputType(int32 OutputIndex)
diff --git a/Engine/Source/Runtime/Engine/Private/Materials/StrataMaterial.cpp b/Engine/Source/Runtime/Engine/Private/Materials/StrataMaterial.cpp
--- a/Engine/Source/Runtime/Engine/Private/Materials/StrataMaterial.cpp
+++ b/Engine/Source/Runtime/Engine/Private/Materials/StrataMaterial.cpp
@@ -127,6 +127,7 @@ FStrataMaterialCompilationInfo StrataCompilationInfoAdd(FMaterialCompiler* Compi
 }
 
 
+// ==> NOTE: Always pair with the shader behavior in StrataAddParameterBlending.
 FStrataMaterialCompilationInfo StrataCompilationInfoAddParamBlend(FMaterialCompiler* Compiler, const FStrataMaterialCompilationInfo& A, const FStrataMaterialCompilationInfo& B, const FStrataRegisteredSharedNormal& RegisteredSharedNormal)
 {
 	check(A.TotalBSDFCount == 1);
@@ -155,6 +156,7 @@ FStrataMaterialCompilationInfo StrataCompilationInfoHorizontalMixing(FMaterialCo
 }
 
 
+// ==> NOTE: Always pair with the shader behavior in StrataHorizontalMixingParameterBlending
 FStrataMaterialCompilationInfo StrataCompilationInfoHorizontalMixingParamBlend(FMaterialCompiler* Compiler, const FStrataMaterialCompilationInfo& A, const FStrataMaterialCompilationInfo& B, const FStrataRegisteredSharedNormal& RegisteredSharedNormal)
 {
 	check(A.TotalBSDFCount == 1);
@@ -199,6 +201,29 @@ FStrataMaterialCompilationInfo StrataCompilationInfoVerticalLayering(FMaterialCo
 	return StrataInfo;
 }
 
+
+// ==> NOTE: Always pair with the shader behavior in StrataVerticalLayeringParameterBlending
+FStrataMaterialCompilationInfo StrataCompilationInfoVerticalLayeringParamBlend(FMaterialCompiler* Compiler, const FStrataMaterialCompilationInfo& Top, const FStrataMaterialCompilationInfo& Base, const FStrataRegisteredSharedNormal& RegisteredSharedNormal)
+{
+	check(Top.TotalBSDFCount == 1);
+	check(Base.TotalBSDFCount == 1);
+
+	FStrataMaterialCompilationInfo StrataInfo = Base;
+	FStrataMaterialCompilationInfo::FBSDF& NewBSDF = StrataInfo.Layers[0].BSDFs[0];
+	const FStrataMaterialCompilationInfo::FBSDF& TopBSDF = Top.Layers[0].BSDFs[0];
+
+	NewBSDF.RegisteredSharedNormal = RegisteredSharedNormal;
+
+//	NewBSDF.bHasSSS				= TopBSDF.bHasSSS;				// We keep SSS only if the base layer has it. Otherwise it will be simple volume and the throughput wll be applied on the parameters.
+//	NewBSDF.bHasDMFPPluggedIn	= TopBSDF.bHasDMFPPluggedIn;	// Idem
+	NewBSDF.bHasEdgeColor		|= TopBSDF.bHasEdgeColor;		// We keep the union of both, even though it will be hard to get a perfect match.
+	NewBSDF.bHasFuzz			|= TopBSDF.bHasFuzz;			// Idem
+	NewBSDF.bHasThinFilm		 = TopBSDF.bHasThinFilm;		// We only keep thin film from the top layer, because its color is otherwise not controlable
+
+	UpdateTotalBSDFCount(Compiler, StrataInfo);
+	return StrataInfo;
+}
+
 static bool StrataIsSingleBSDF(FMaterialCompiler* Compiler, const FStrataMaterialCompilationInfo& Material)
 {
 	if (Material.TotalBSDFCount == 0 || Material.LayerCount == 0)
diff --git a/Engine/Source/Runtime/Engine/Private/Materials/StrataMaterial.h b/Engine/Source/Runtime/Engine/Private/Materials/StrataMaterial.h
--- a/Engine/Source/Runtime/Engine/Private/Materials/StrataMaterial.h
+++ b/Engine/Source/Runtime/Engine/Private/Materials/StrataMaterial.h
@@ -63,6 +63,7 @@ FStrataMaterialCompilationInfo StrataCompilationInfoAddParamBlend(FMaterialCompi
 FStrataMaterialCompilationInfo StrataCompilationInfoHorizontalMixing(FMaterialCompiler* Compiler, const FStrataMaterialCompilationInfo& A, const FStrataMaterialCompilationInfo& B);
 FStrataMaterialCompilationInfo StrataCompilationInfoHorizontalMixingParamBlend(FMaterialCompiler* Compiler, const FStrataMaterialCompilationInfo& A, const FStrataMaterialCompilationInfo& B, const FStrataRegisteredSharedNormal& RegisteredSharedNormal);
 FStrataMaterialCompilationInfo StrataCompilationInfoVerticalLayering(FMaterialCompiler* Compiler, const FStrataMaterialCompilationInfo& Top, const FStrataMaterialCompilationInfo& Base);
+FStrataMaterialCompilationInfo StrataCompilationInfoVerticalLayeringParamBlend(FMaterialCompiler* Compiler, const FStrataMaterialCompilationInfo& Top, const FStrataMaterialCompilationInfo& Base, const FStrataRegisteredSharedNormal& RegisteredSharedNormal);
 
 bool StrataIsVolumetricFogCloudOnly(FMaterialCompiler* Compiler, const FStrataMaterialCompilationInfo& Material);
 bool StrataIsUnlitOnly(FMaterialCompiler* Compiler, const FStrataMaterialCompilationInfo& Material);
diff --git a/Engine/Source/Runtime/Engine/Public/MaterialCompiler.h b/Engine/Source/Runtime/Engine/Public/MaterialCompiler.h
--- a/Engine/Source/Runtime/Engine/Public/MaterialCompiler.h
+++ b/Engine/Source/Runtime/Engine/Public/MaterialCompiler.h
@@ -442,6 +442,7 @@ public:
 	virtual int32 StrataHorizontalMixing(int32 Foreground, int32 Background, int32 Mix) = 0;
 	virtual int32 StrataHorizontalMixingParameterBlending(int32 Foreground, int32 Background, int32 Mix, const FString& SharedNormalIndexMacro) = 0;
 	virtual int32 StrataVerticalLayering(int32 Top, int32 Base) = 0;
+	virtual int32 StrataVerticalLayeringParameterBlending(int32 Top, int32 Base, const FString& SharedNormalIndexMacro, int32 TopBSDFNormalCodeChunk) = 0;
 	virtual int32 StrataAdd(int32 A, int32 B) = 0;
 	virtual int32 StrataAddParameterBlending(int32 A, int32 B, const FString& SharedNormalIndexMacro) = 0;
 	virtual int32 StrataMultiply(int32 A, int32 Weight) = 0;
@@ -990,6 +991,11 @@ public:
 		return Compiler->StrataVerticalLayering(Top, Base);
 	}
 
+	virtual int32 StrataVerticalLayeringParameterBlending(int32 Top, int32 Base, const FString& SharedNormalIndexMacro, int32 TopBSDFNormalCodeChunk) override
+	{
+		return Compiler->StrataVerticalLayeringParameterBlending(Top, Base, SharedNormalIndexMacro, TopBSDFNormalCodeChunk);
+	}
+
 	virtual int32 StrataAdd(int32 A, int32 B) override
 	{
 		return Compiler->StrataAdd(A, B);
