Strata improved prevision for the thin film feature: gamma2 encoding + dithering hide banding.

diff --git a/Engine/Shaders/Private/Strata/Strata.ush b/Engine/Shaders/Private/Strata/Strata.ush
--- a/Engine/Shaders/Private/Strata/Strata.ush
+++ b/Engine/Shaders/Private/Strata/Strata.ush
@@ -1201,12 +1201,14 @@ void StrataDitherBSDF(inout FStrataBSDF BSDF, in float Dither)
 	{
 		// We execute a dither that can go in positive/negative direction to avoid surfaces to looks overall brighter.
 		const float CenteredDither = Dither - 0.5f;
+		const float Dither8bits = CenteredDither * 1.0f / 128.0f;
 		const float Dither7bits = CenteredDither * 1.0f / 127.0f;
 		const float Dither6bits = CenteredDither * 1.0f / 63.0f;
 		const float3 Dither776 = float3(Dither7bits, Dither7bits, Dither6bits);
 
 		SLAB_DIFFUSEALBEDO(BSDF) += Dither776;
 		SLAB_F0(BSDF) += Dither776;
+		SLAB_THINFILM_THICKNESS(BSDF) += BSDF_GETHASTHINFILM(BSDF) ? Dither8bits : 0.0f;
 	}
 }
 
@@ -1779,7 +1781,8 @@ void PackStrataOut(
 				if (BSDF_GETHASF90(BSDF) || BSDF_GETHASTHINFILM(BSDF) || BSDF_GETHASHAZINESS(BSDF))
 				{
 					float3 F90HSV = LinearRGB_2_HSV(SLAB_F90(BSDF));
-					STRATA_STORE_UINT1(PackRGBA8(float4(F90HSV.x, F90HSV.y, SLAB_THINFILM_THICKNESS(BSDF), SLAB_HAZINESS(BSDF))));
+					const float ThinFilmGamma2 = sqrt(SLAB_THINFILM_THICKNESS(BSDF)); // Using a gamma 2.0 encoding help with precision to be visually more uniform. This is then further improved via dithering.
+					STRATA_STORE_UINT1(PackRGBA8(float4(F90HSV.x, F90HSV.y, ThinFilmGamma2, SLAB_HAZINESS(BSDF))));
 				}
 
 				const bool bIsSimpleVolume = BSDF_GETISSIMPLEVOLUME(BSDF);
@@ -1976,7 +1979,7 @@ FStrataBSDF UnpackStrataBSDFIn(ByteAddressBuffer StrataBuffer, inout FStrataAddr
 			const float F90Val = 1.0f;	// Forced to 1 as F90 is not super to have a low value
 			SLAB_F90(OutBSDF) = bHasF90 ? HSV_2_LinearRGB(float3(F90Hue, F90Sat, F90Val)) : 1.0f;
 
-			SLAB_THINFILM_THICKNESS(OutBSDF) = Data.z;
+			SLAB_THINFILM_THICKNESS(OutBSDF) = Data.z * Data.z;
 			SLAB_HAZINESS(OutBSDF) = Data.w;
 		}
 
