Fixed starta visualize to only show haziness if pertinent.

diff --git a/Engine/Shaders/Private/Strata/StrataVisualize.usf b/Engine/Shaders/Private/Strata/StrataVisualize.usf
--- a/Engine/Shaders/Private/Strata/StrataVisualize.usf
+++ b/Engine/Shaders/Private/Strata/StrataVisualize.usf
@@ -362,9 +362,12 @@ void VisualizeMaterialPS(
 					NEXTLINE;
 				}
 
-				HORIZONTALTAB; PRINTPARAMCHAR(_H_); PRINTPARAMCHAR(_A_); PRINTPARAMCHAR(_Z_); PRINTPARAMCHAR(_I_); PRINTPARAMCHAR(_N_); PRINTPARAMCHAR(_SPC_);
-				HORIZONTALTAB; PRINTFLOAT1(SLAB_HAZINESS(BSDF));
-				NEXTLINE;
+				if (BSDF_GETHASHAZINESS(BSDF))
+				{
+					HORIZONTALTAB; PRINTPARAMCHAR(_H_); PRINTPARAMCHAR(_A_); PRINTPARAMCHAR(_Z_); PRINTPARAMCHAR(_I_); PRINTPARAMCHAR(_N_); PRINTPARAMCHAR(_SPC_);
+					HORIZONTALTAB; PRINTFLOAT1(SLAB_HAZINESS(BSDF));
+					NEXTLINE;
+				}
 
 				if (BSDF_GETHASSSS(BSDF) || BSDF_GETISSIMPLEVOLUME(BSDF) || BSDF_GETISTHIN(BSDF))
 				{
@@ -422,7 +425,10 @@ void VisualizeMaterialPS(
 					NEXTLINE;
 				}
 
-				NEXTLINE;
+				if (BSDF_GETHASHAZINESS(BSDF))
+				{
+					NEXTLINE;
+				}
 
 				if (BSDF_GETHASSSS(BSDF) || BSDF_GETISSIMPLEVOLUME(BSDF) || BSDF_GETISTHIN(BSDF))
 				{
