Improved Strata vertical layering via parameter blending.

diff --git a/Engine/Shaders/Private/Strata/Strata.ush b/Engine/Shaders/Private/Strata/Strata.ush
--- a/Engine/Shaders/Private/Strata/Strata.ush
+++ b/Engine/Shaders/Private/Strata/Strata.ush
@@ -1095,13 +1095,16 @@ FStrataData StrataVerticalLayeringParameterBlending(FStrataData Top, FStrataData
 	// Only keep thin film and for the top layer
 	SLAB_THINFILM_THICKNESS(ResultBSDF)=SLAB_THINFILM_THICKNESS(TopBSDF);
 
+	// Keep the DMFP form the bottom layer only (top layer simple volume is backed in parameters)
+	SLAB_SSSDMFP(ResultBSDF) =			SLAB_SSSDMFP(BaseBSDF);
+
 #if STRATA_INLINE_SHADING
 	// Only the bottom layer SSS / Thin model is used so we only keep the bottom layer data there. Top layer is converted to throughput.
 	ResultBSDF.TmpDMFP =				BaseBSDF.TmpDMFP;
 #endif
 
-	// When it comes to adding the visual contribution, it is tricky to find a matching operator for coverage. We simply lerp for now.
-	ResultBSDF.Weight =					saturate(GreyThroughput * BaseBSDF.Weight		+ SLAB_DIFFUSEALBEDO(TopBSDF));
+	// The weight is the maximum of the two
+	ResultBSDF.Weight =					max(BaseBSDF.Weight,							TopBSDF.Weight);
 
 #undef ResultBSDF
 #undef TopBSDF
