Improve path tracer performance with simplified materials

GPU Lightmass saw a performance regression after using more of the material model for translucent shadow support. Reclaim some performance by special casing more of the material logic for the subset that GPULM needs.

diff --git a/Engine/Plugins/Experimental/GPULightmass/Shaders/Private/LightmapPathTracing.usf b/Engine/Plugins/Experimental/GPULightmass/Shaders/Private/LightmapPathTracing.usf
--- a/Engine/Plugins/Experimental/GPULightmass/Shaders/Private/LightmapPathTracing.usf
+++ b/Engine/Plugins/Experimental/GPULightmass/Shaders/Private/LightmapPathTracing.usf
@@ -5,6 +5,7 @@ PathTracingRayGenShader.usf: Reference path tracing
 ===============================================================================================*/
 
 #define SIMPLIFIED_MATERIAL_SHADER 1
+#define RANDSEQ_UNROLL_SOBOL 0
 
 // Needs to be defined before the uniform buffer struct header is included, very unfortunately
 struct FIrradianceCacheRecord
@@ -124,18 +125,18 @@ FThreeBandSHVectorFloat SHBasisFunction3Float(float3 InputVector)
 	#define USE_IRRADIANCE_CACHING 0
 #endif
 
-uint MaxBounces;
+static const uint MaxBounces = 32;
 
 // 0: only Material sampling
 // 1: only Light sampling
 // 2: both Material and Light
-uint MISMode;
+static const uint MISMode = 2;
 
-uint ApproximateCaustics;
-uint EnableCameraBackfaceCulling;
-uint EnableDirectLighting;
-uint EnableEmissive;
-float MaxNormalBias;
+static const uint ApproximateCaustics = 1;
+static const uint EnableCameraBackfaceCulling = 0;
+static const uint EnableDirectLighting = 1;
+static const uint EnableEmissive = 1;
+static const float MaxNormalBias = 0.1;
 
 RaytracingAccelerationStructure TLAS;
 uint SceneVisibleLightCount;
@@ -405,7 +406,7 @@ void PathTracingKernel(
 		}
 		else
 		{
-			Payload = TraceTransparentRay(Ray, RayCone, bIsCameraRay, bIsLastBounce, bIncludeEmissive, LaunchIndex, NumVisibleLights, RandSequence, PathThroughput, Radiance);
+			Payload = TraceTransparentRay(Ray, RayCone, false , bIsLastBounce, bIncludeEmissive, LaunchIndex, NumVisibleLights, RandSequence, PathThroughput, Radiance);
 		}
 		
 		if (Payload.IsMiss())
diff --git a/Engine/Plugins/Experimental/GPULightmass/Source/GPULightmass/Private/LightmapRayTracing.cpp b/Engine/Plugins/Experimental/GPULightmass/Source/GPULightmass/Private/LightmapRayTracing.cpp
--- a/Engine/Plugins/Experimental/GPULightmass/Source/GPULightmass/Private/LightmapRayTracing.cpp
+++ b/Engine/Plugins/Experimental/GPULightmass/Source/GPULightmass/Private/LightmapRayTracing.cpp
@@ -153,7 +153,7 @@ bool FLightmapRayTracingMeshProcessor::Process(
 
 	FMaterialShaderTypes ShaderTypes;
 
-	if (MaterialResource.IsMasked() || MaterialResource.GetBlendMode() == BLEND_Translucent)
+	if (MaterialResource.IsMasked() || MaterialResource.GetBlendMode() != BLEND_Opaque)
 	{
 		ShaderTypes.AddShaderType<TLightmapMaterialCHS<true>>();
 	}
diff --git a/Engine/Plugins/Experimental/GPULightmass/Source/GPULightmass/Private/LightmapRayTracing.h b/Engine/Plugins/Experimental/GPULightmass/Source/GPULightmass/Private/LightmapRayTracing.h
--- a/Engine/Plugins/Experimental/GPULightmass/Source/GPULightmass/Private/LightmapRayTracing.h
+++ b/Engine/Plugins/Experimental/GPULightmass/Source/GPULightmass/Private/LightmapRayTracing.h
@@ -70,7 +70,6 @@ class FLightmapPathTracingRGS : public FGlobalShader
 		SHADER_PARAMETER_RDG_TEXTURE(Texture2D<float>, RayGuidingCDFY)
 		SHADER_PARAMETER_SRV(StructuredBuffer<FGPUTileDescription>, BatchedTiles)
 
-		SHADER_PARAMETER_STRUCT_INCLUDE(FPathTracingData, PathTracingData)
 		SHADER_PARAMETER_RDG_BUFFER_SRV(StructuredBuffer<FPathTracingLight>, SceneLights)
 		SHADER_PARAMETER(uint32, SceneLightCount)
 		SHADER_PARAMETER(uint32, SceneVisibleLightCount)
@@ -134,7 +133,6 @@ class FVolumetricLightmapPathTracingRGS : public FGlobalShader
 		SHADER_PARAMETER_UAV(RWTexture3D<float4>, OutSHCoefficients0B)
 		SHADER_PARAMETER_UAV(RWTexture3D<float4>, OutSHCoefficients1B)
 
-		SHADER_PARAMETER_STRUCT_INCLUDE(FPathTracingData, PathTracingData)
 		SHADER_PARAMETER_RDG_BUFFER_SRV(StructuredBuffer<FPathTracingLight>, SceneLights)
 		SHADER_PARAMETER(uint32, SceneLightCount)
 		SHADER_PARAMETER(uint32, SceneVisibleLightCount)
@@ -183,7 +181,6 @@ class FStationaryLightShadowTracingRGS : public FGlobalShader
 		SHADER_PARAMETER_RDG_TEXTURE_UAV(RWTexture2D<float4>, ShadowMask)
 		SHADER_PARAMETER_RDG_TEXTURE_UAV(RWTexture2D<float4>, ShadowMaskSampleCount)
 
-		SHADER_PARAMETER_STRUCT_INCLUDE(FPathTracingData, PathTracingData)
 		SHADER_PARAMETER_RDG_BUFFER_SRV(StructuredBuffer<FPathTracingLight>, SceneLights)
 		SHADER_PARAMETER(uint32, SceneLightCount)
 		SHADER_PARAMETER(uint32, SceneVisibleLightCount)
diff --git a/Engine/Plugins/Experimental/GPULightmass/Source/GPULightmass/Private/LightmapRenderer.cpp b/Engine/Plugins/Experimental/GPULightmass/Source/GPULightmass/Private/LightmapRenderer.cpp
--- a/Engine/Plugins/Experimental/GPULightmass/Source/GPULightmass/Private/LightmapRenderer.cpp
+++ b/Engine/Plugins/Experimental/GPULightmass/Source/GPULightmass/Private/LightmapRenderer.cpp
@@ -2251,13 +2251,6 @@ void FLightmapRenderer::Finalize(FRDGBuilder& GraphBuilder)
 									PassParameters->ViewUniformBuffer = Scene->ReferenceView->ViewUniformBuffer;
 									PassParameters->IrradianceCachingParameters = Scene->IrradianceCache->IrradianceCachingParametersUniformBuffer;
 
-									PassParameters->PathTracingData.MaxBounces = 32;
-									PassParameters->PathTracingData.MaxNormalBias = 0.1;
-									PassParameters->PathTracingData.EnableDirectLighting = 1;
-									PassParameters->PathTracingData.EnableEmissive = 1;
-									PassParameters->PathTracingData.MISMode = 2;
-									PassParameters->PathTracingData.ApproximateCaustics = 0;
-
 									SetupPathTracingLightParameters(Scene->LightSceneRenderState, GraphBuilder, PassParameters);
 
 									// TODO: find a way to share IES atlas with path tracer ...
diff --git a/Engine/Plugins/Experimental/GPULightmass/Source/GPULightmass/Private/VolumetricLightmap.cpp b/Engine/Plugins/Experimental/GPULightmass/Source/GPULightmass/Private/VolumetricLightmap.cpp
--- a/Engine/Plugins/Experimental/GPULightmass/Source/GPULightmass/Private/VolumetricLightmap.cpp
+++ b/Engine/Plugins/Experimental/GPULightmass/Source/GPULightmass/Private/VolumetricLightmap.cpp
@@ -559,13 +559,6 @@ void FVolumetricLightmapRenderer::BackgroundTick()
 					PassParameters->ViewUniformBuffer = Scene->ReferenceView->ViewUniformBuffer;
 					PassParameters->IrradianceCachingParameters = Scene->IrradianceCache->IrradianceCachingParametersUniformBuffer;
 
-					PassParameters->PathTracingData.MaxBounces = 32;
-					PassParameters->PathTracingData.MaxNormalBias = 0.1;
-					PassParameters->PathTracingData.EnableDirectLighting = 1;
-					PassParameters->PathTracingData.EnableEmissive = 1;
-					PassParameters->PathTracingData.MISMode = 2;
-					PassParameters->PathTracingData.ApproximateCaustics = 0;
-
 					SetupPathTracingLightParameters(Scene->LightSceneRenderState, GraphBuilder, PassParameters);
 
 					// TODO: find a way to share IES atlas with path tracer ...
diff --git a/Engine/Shaders/Private/PathTracing/Material/PathTracingLambert.ush b/Engine/Shaders/Private/PathTracing/Material/PathTracingLambert.ush
--- a/Engine/Shaders/Private/PathTracing/Material/PathTracingLambert.ush
+++ b/Engine/Shaders/Private/PathTracing/Material/PathTracingLambert.ush
@@ -1,7 +1,7 @@
 // Copyright Epic Games, Inc. All Rights Reserved.
 
 /*=============================================================================================
-PathTracingLambert.usf: Lambertian BRDF sampling functions 
+PathTracingLambert.usf: Lambertian BRDF sampling functions
 ===============================================================================================*/
 #pragma once
 
@@ -16,20 +16,54 @@ FMaterialSample Lambert_SampleMaterial(
 	return CreateMaterialSample(TangentToWorld(SampledValue.xyz, N_World), GetBaseColor(Payload), SampledValue.w, 1.0, 1.0);
 }
 
-float Lambert_Pdf(float3 OutgoingDirection, FMaterialClosestHitPayload Payload)
+FMaterialEval Lambert_EvalMaterial(
+	float3 OutgoingDirection,
+	FMaterialClosestHitPayload Payload
+)
 {
 	float3 N_World = GetWorldNormal(Payload);
 	float NoL = saturate(dot(N_World, OutgoingDirection));
-	return NoL / PI;
+
+	return CreateMaterialEval(GetBaseColor(Payload), NoL / PI);
 }
 
-FMaterialEval Lambert_EvalMaterial(
+
+FMaterialSample Simplified_SampleMaterial(
+	FMaterialClosestHitPayload Payload,
+	float4 RandSample)
+{
+	float3 N_World = GetWorldNormal(Payload);
+
+	float3 FrontColor = GetBaseColor(Payload);
+	float3 BackColor = Payload.ShadingModelID == SHADINGMODELID_TWOSIDED_FOLIAGE ? GetSubsurfaceColor(Payload) : 0.0;
+
+	float4 SampledValue = CosineSampleHemisphere(RandSample.xy);
+
+	float TransmissionProb = LobeSelectionProb(BackColor, FrontColor);
+	bool bDoTransmission = RandSample.z < TransmissionProb;
+	float3 Weight = bDoTransmission ? BackColor / TransmissionProb : FrontColor / (1 - TransmissionProb);
+	float Pdf = SampledValue.w * (bDoTransmission ? TransmissionProb : 1.0 - TransmissionProb);
+
+	return CreateMaterialSample(TangentToWorld(SampledValue.xyz, N_World * (bDoTransmission ? -1.0 : 1.0)), Weight, Pdf, 1.0, 1.0);
+}
+
+FMaterialEval Simplified_EvalMaterial(
 	float3 OutgoingDirection,
 	FMaterialClosestHitPayload Payload
 )
 {
 	float3 N_World = GetWorldNormal(Payload);
-	float NoL = saturate(dot(N_World, OutgoingDirection));
+	float NoL = dot(N_World, OutgoingDirection);
 
-	return CreateMaterialEval(GetBaseColor(Payload), NoL / PI);
+	float3 FrontColor = GetBaseColor(Payload);
+	float3 BackColor = Payload.ShadingModelID == SHADINGMODELID_TWOSIDED_FOLIAGE ? GetSubsurfaceColor(Payload) : 0.0;
+	float TransmissionProb = LobeSelectionProb(BackColor, FrontColor);
+
+	float3 Weight = NoL < 0.0 ? BackColor : FrontColor;
+	float Pdf = NoL < 0.0 ? TransmissionProb : 1.0 - TransmissionProb;
+	if (Pdf > 0.0)
+	{
+		Weight /= Pdf;
+	}
+	return CreateMaterialEval(Weight, Pdf * saturate(abs(NoL)) / PI);
 }
diff --git a/Engine/Shaders/Private/PathTracing/Material/PathTracingMaterialSampling.ush b/Engine/Shaders/Private/PathTracing/Material/PathTracingMaterialSampling.ush
--- a/Engine/Shaders/Private/PathTracing/Material/PathTracingMaterialSampling.ush
+++ b/Engine/Shaders/Private/PathTracing/Material/PathTracingMaterialSampling.ush
@@ -31,11 +31,6 @@ PathTracingMaterialSampling.usf: Material sampling functions
 bool IsNeeValidMaterial(FMaterialClosestHitPayload Payload)
 {
 	// Checks done in the same order as other sampling and evaluation functions to ensure consistency
-	
-#if SIMPLIFIED_MATERIAL_SHADER
-	return true;
-#else
-
 	if (Payload.BlendingMode == RAY_TRACING_BLEND_MODE_MODULATE ||
 		Payload.BlendingMode == RAY_TRACING_BLEND_MODE_ALPHA_HOLDOUT)
 	{
@@ -56,7 +51,6 @@ bool IsNeeValidMaterial(FMaterialClosestHitPayload Payload)
 		default:
 			return true;
 	}
-#endif
 }
 
 float3 GetMaterialTransparency(FMaterialClosestHitPayload Payload, float3 RayDirection)
@@ -65,14 +59,16 @@ float3 GetMaterialTransparency(FMaterialClosestHitPayload Payload, float3 RayDir
 	{
 		case SHADINGMODELID_THIN_TRANSLUCENT:
 		{
-			if (Payload.Ior == 1.0)
+#if !SIMPLIFIED_MATERIAL_SHADER
+			if (Payload.Ior != 1.0)
 			{
-				float VoN = dot(RayDirection, Payload.WorldNormal);
-				return (1 - Payload.Opacity) * pow(Payload.CustomData.xyz, 1.0 / abs(VoN));
+				return 0;
 			}
 			else
+#endif
 			{
-				return 0;
+				float VoN = dot(RayDirection, Payload.WorldNormal);
+				return (1 - Payload.Opacity) * pow(Payload.CustomData.xyz, 1.0 / abs(VoN));
 			}
 		}
 		default:
@@ -82,15 +78,17 @@ float3 GetMaterialTransparency(FMaterialClosestHitPayload Payload, float3 RayDir
 	switch (Payload.BlendingMode)
 	{
 		case RAY_TRACING_BLEND_MODE_TRANSLUCENT:
-			if (Payload.Ior == 1.0)
+#if !SIMPLIFIED_MATERIAL_SHADER
+			if (Payload.Ior != 1.0)
 			{
-				// material is not configured with refraction, can use transparency
-				return 1 - Payload.Opacity;
+				// not using transparency - material will be opaque
+				return 0;
 			}
 			else
+#endif
 			{
-				// not using transparency - material will be opaque
-				return 0;
+				// material is not configured with refraction, can use transparency
+				return 1 - Payload.Opacity;
 			}
 		case RAY_TRACING_BLEND_MODE_ADDITIVE:
 			return 1;
@@ -112,15 +110,17 @@ float3 EstimateMaterialAlbedo(FMaterialClosestHitPayload Payload)
 	switch (Payload.BlendingMode)
 	{
 		case RAY_TRACING_BLEND_MODE_TRANSLUCENT:
-			if (Payload.Ior == 1.0)
+#if !SIMPLIFIED_MATERIAL_SHADER
+			if (Payload.Ior != 1.0)
 			{
-				// material is not configured with refraction, can use transparency
-				// NOTE: the albedo when considering refraction is handled below
-				BrdfWeight = Payload.Opacity;
+				HasRefraction = true;
 			}
 			else
+#endif
 			{
-				HasRefraction = true;
+				// material is not configured with refraction, can use transparency
+				// NOTE: the albedo when considering refraction is handled below
+				BrdfWeight = Payload.Opacity;
 			}
 			break;
 		case RAY_TRACING_BLEND_MODE_ADDITIVE:
@@ -141,6 +141,11 @@ float3 EstimateMaterialAlbedo(FMaterialClosestHitPayload Payload)
 		case SHADINGMODELID_UNLIT:
 			// this object does not scatter light
 			return 0;
+#if SIMPLIFIED_MATERIAL_SHADER
+		case SHADINGMODELID_TWOSIDED_FOLIAGE:
+			Albedo = BrdfWeight * (Payload.BaseColor + GetSubsurfaceColor(Payload));
+			break;
+#else
 		case SHADINGMODELID_TWOSIDED_FOLIAGE:
 		case SHADINGMODELID_SUBSURFACE:
 		case SHADINGMODELID_PREINTEGRATED_SKIN:
@@ -155,10 +160,12 @@ float3 EstimateMaterialAlbedo(FMaterialClosestHitPayload Payload)
 		case SHADINGMODELID_CLOTH:
 			Albedo = BrdfWeight * (Payload.DiffuseColor + lerp(Payload.SpecularColor, Payload.CustomData.xyz, Payload.CustomData.w));
 			break;
+#endif
 		default:
 			Albedo = BrdfWeight * Payload.BaseColor;
 			break;
 	}
+#if !SIMPLIFIED_MATERIAL_SHADER
 	if (HasRefraction)
 	{
 		// glass is currently perfectly clear - so has unit albedo
@@ -166,11 +173,13 @@ float3 EstimateMaterialAlbedo(FMaterialClosestHitPayload Payload)
 		float3 White = 1;
 		Albedo = lerp(White, Albedo, Payload.Opacity);
 	}
+#endif
 	return Albedo;
 }
 
 void AdjustMaterialRoughness(inout FMaterialClosestHitPayload Payload, float PathRoughness)
 {
+#if !SIMPLIFIED_MATERIAL_SHADER
 	// boost material roughness according to what we've seen so far
 	switch (Payload.ShadingModelID)
 	{
@@ -202,11 +211,15 @@ void AdjustMaterialRoughness(inout FMaterialClosestHitPayload Payload, float Pat
 			Payload.Roughness = 1.0;
 			break;
 	}
+#endif
 }
 
 // returns an average roughness for the whole material (for the path roughness of shadow rays from this material)
 float GetAverageRoughness(FMaterialClosestHitPayload Payload)
 {
+#if SIMPLIFIED_MATERIAL_SHADER
+	return 1.0;
+#else
 	const float GlassPercent = (Payload.BlendingMode == RAY_TRACING_BLEND_MODE_TRANSLUCENT && Payload.Ior != 1.0) ? 1 - Payload.Opacity : 0.0;
 	float Roughness = Payload.Roughness;
 	switch (Payload.ShadingModelID)
@@ -224,6 +237,7 @@ float GetAverageRoughness(FMaterialClosestHitPayload Payload)
 	float AvgRoughness = lerp(1.0 /* diffuse lobe */, Roughness, Payload.Metallic);
 	AvgRoughness = lerp(AvgRoughness, Payload.Roughness, GlassPercent);
 	return AvgRoughness;
+#endif
 }
 
 
@@ -238,6 +252,7 @@ FMaterialClosestHitPayload UnpackPathTracingPayload(FPackedMaterialClosestHitPay
 
 	switch (HitPayload.ShadingModelID)
 	{
+#if !SIMPLIFIED_MATERIAL_SHADER
 		case SHADINGMODELID_SUBSURFACE:
 		case SHADINGMODELID_PREINTEGRATED_SKIN:
 		{
@@ -250,6 +265,7 @@ FMaterialClosestHitPayload UnpackPathTracingPayload(FPackedMaterialClosestHitPay
 			
 			// fallthrough to case below to avoid over-bright SSS
 		}
+#endif
 		case SHADINGMODELID_TWOSIDED_FOLIAGE:
 		{
 			float3 Front = HitPayload.DiffuseColor;
@@ -265,6 +281,13 @@ FMaterialClosestHitPayload UnpackPathTracingPayload(FPackedMaterialClosestHitPay
 			}
 			break;
 		}
+#if SIMPLIFIED_MATERIAL_SHADER
+		default:
+		{
+			HitPayload.CustomData = 0;
+			break;
+		}
+#else
 		case SHADINGMODELID_SUBSURFACE_PROFILE:
 		{
 			uint ProfileId = uint(HitPayload.CustomData.a * 255.0 + 0.5);
@@ -315,6 +338,7 @@ FMaterialClosestHitPayload UnpackPathTracingPayload(FPackedMaterialClosestHitPay
 			HitPayload.Opacity = 1;
 			break;
 		}
+#endif
 	}
 
 	return HitPayload;
@@ -332,6 +356,13 @@ FMaterialSample SampleMaterial(
 		return RadianceProbe_SampleMaterial(Payload, RandSample);
 	}
 
+#if SIMPLIFIED_MATERIAL_SHADER
+	FMaterialSample Result = Simplified_SampleMaterial(Payload, RandSample);
+	Result.Weight *= Payload.Opacity;
+	return Result;
+	//return Lambert_SampleMaterial(Payload, RandSample);
+#else
+
 	if (Payload.BlendingMode == RAY_TRACING_BLEND_MODE_MODULATE ||
 		Payload.BlendingMode == RAY_TRACING_BLEND_MODE_ALPHA_HOLDOUT)
 	{
@@ -360,11 +391,9 @@ FMaterialSample SampleMaterial(
 				BrdfProb = Payload.Opacity;
 			}
 			break;
-	#if !SIMPLIFIED_MATERIAL_SHADER
 		case RAY_TRACING_BLEND_MODE_ADDITIVE:
 			BrdfWeight = Payload.Opacity;
 			break;
-	#endif
 		default:
 			// opaque or one of the modes handled above
 			break;
@@ -385,8 +414,7 @@ FMaterialSample SampleMaterial(
 		if (Payload.ShadingModelID == SHADINGMODELID_THIN_TRANSLUCENT)
 		{
 			Result = RoughThinGlass_SampleMaterial(RayDirection, Payload, RandSample);
-		}	
-	#if !SIMPLIFIED_MATERIAL_SHADER
+		}
 		else if (Payload.Roughness == 0)
 		{
 			Result = SpecularRefraction_SampleMaterial(RayDirection, Payload, RandSample);
@@ -395,7 +423,6 @@ FMaterialSample SampleMaterial(
 		{
 			Result = RoughGlass_SampleMaterial(RayDirection, Payload, RandSample);
 		}
-	#endif
 		Result.Weight *= TransWeight;
 		Result.Pdf *= 1 - BrdfProb;
 		return Result;
@@ -403,7 +430,6 @@ FMaterialSample SampleMaterial(
 
 	switch (Payload.ShadingModelID)
 	{
-	#if !SIMPLIFIED_MATERIAL_SHADER
 		case SHADINGMODELID_UNLIT:
 			// Nothing to do, material does not scatter light
 			return Result;
@@ -421,7 +447,6 @@ FMaterialSample SampleMaterial(
 		case SHADINGMODELID_CLEAR_COAT:
 			Result = ClearCoat_SampleMaterial(RayDirection, Payload, RandSample);
 			break;
-	#endif
 		case SHADINGMODELID_TWOSIDED_FOLIAGE:
 			Result = TwoSidedFoliage_SampleMaterial(RayDirection, Payload, RandSample);
 			break;
@@ -434,6 +459,7 @@ FMaterialSample SampleMaterial(
 	Result.Weight *= BrdfWeight;
 	Result.Pdf *= BrdfProb;
 	return Result;
+#endif
 }
 
 FMaterialEval EvalMaterial(
@@ -447,6 +473,13 @@ FMaterialEval EvalMaterial(
 	{
 		return RadianceProbe_EvalMaterial(OutgoingDirection, Payload);
 	}
+#if SIMPLIFIED_MATERIAL_SHADER
+	FMaterialEval Result = Simplified_EvalMaterial(OutgoingDirection, Payload);
+	Result.Weight *= Payload.Opacity;
+	return Result;
+
+	//return Lambert_EvalMaterial(OutgoingDirection, Payload);
+#else
 
 	if (Payload.BlendingMode == RAY_TRACING_BLEND_MODE_MODULATE ||
 		Payload.BlendingMode == RAY_TRACING_BLEND_MODE_ALPHA_HOLDOUT)
@@ -491,7 +524,6 @@ FMaterialEval EvalMaterial(
 	FMaterialEval Result = NullMaterialEval();
 	switch (Payload.ShadingModelID)
 	{
-	#if !SIMPLIFIED_MATERIAL_SHADER
 		case SHADINGMODELID_UNLIT:
 			break;
 		case SHADINGMODELID_DEFAULT_LIT:
@@ -508,7 +540,6 @@ FMaterialEval EvalMaterial(
 		case SHADINGMODELID_CLEAR_COAT:
 			Result = ClearCoat_EvalMaterial(IncomingDirection, OutgoingDirection, Payload);
 			break;
-	#endif
 		case SHADINGMODELID_TWOSIDED_FOLIAGE:
 			Result = TwoSidedFoliage_EvalMaterial(IncomingDirection, OutgoingDirection, Payload);
 			break;
@@ -543,4 +574,5 @@ FMaterialEval EvalMaterial(
 		}
 	}
 	return Result;
+#endif
 }
diff --git a/Engine/Shaders/Private/RayTracing/RayTracingMaterialHitShaders.usf b/Engine/Shaders/Private/RayTracing/RayTracingMaterialHitShaders.usf
--- a/Engine/Shaders/Private/RayTracing/RayTracingMaterialHitShaders.usf
+++ b/Engine/Shaders/Private/RayTracing/RayTracingMaterialHitShaders.usf
@@ -539,7 +539,26 @@ RAY_TRACING_ENTRY_CLOSEST_HIT(MaterialCHS,
 	 * Set material attributes for simplified materials
 	 **/
 	// Anything unsupported will be forced to default lit
-	if (Payload.ShadingModelID != SHADINGMODELID_TWOSIDED_FOLIAGE || Payload.ShadingModelID != SHADINGMODELID_THIN_TRANSLUCENT)
+#if MATERIAL_SHADINGMODEL_TWOSIDED_FOLIAGE
+	if (Payload.ShadingModelID == SHADINGMODELID_TWOSIDED_FOLIAGE)
+	{
+		Payload.CustomData = GetMaterialSubsurfaceData(PixelMaterialInputs);
+	}
+	else
+#endif
+#if MATERIAL_SHADINGMODEL_THIN_TRANSLUCENT
+	if (Payload.ShadingModelID == SHADINGMODELID_THIN_TRANSLUCENT)
+	{
+		Payload.CustomData.xyz = GetThinTranslucentMaterialOutput0(MaterialParameters);
+	}
+	else
+#endif
+#if MATERIAL_SHADINGMODEL_UNLIT
+	if (Payload.ShadingModelID == SHADINGMODELID_UNLIT)
+	{
+	}
+	else
+#endif
 	{
 		Payload.ShadingModelID = SHADINGMODELID_DEFAULT_LIT;
 	}
@@ -550,13 +569,8 @@ RAY_TRACING_ENTRY_CLOSEST_HIT(MaterialCHS,
 	Payload.Specular = 0;
 	Payload.Roughness = 1;
 	Payload.Metallic = 0;
-
-#if MATERIAL_SHADINGMODEL_THIN_TRANSLUCENT
-	if (Payload.ShadingModelID == SHADINGMODELID_THIN_TRANSLUCENT)
-	{
-		Payload.CustomData.xyz = GetThinTranslucentMaterialOutput0(MaterialParameters);
-	}
-#endif
+	Payload.Opacity = GetMaterialOpacity(PixelMaterialInputs);
+	Payload.Ior = 1;
 
 #endif
 
@@ -678,7 +692,7 @@ RAY_TRACING_ENTRY_ANY_HIT(MaterialAHS,
 	// Assume Payload.ShadingModelID == SHADINGMODELID_THIN_TRANSLUCENT
 	float3 Transmission = GetThinTranslucentMaterialOutput0(MaterialParameters);
 	float Opacity = GetMaterialOpacity(PixelMaterialInputs);
-#if REFRACTION_USE_INDEX_OF_REFRACTION
+#if REFRACTION_USE_INDEX_OF_REFRACTION && !SIMPLIFIED_MATERIAL_SHADER
 	float Ior = GetMaterialRefraction(PixelMaterialInputs).x;
 	Ior = Ior == 0 ? 1.0 : Ior;
 #else
@@ -688,6 +702,7 @@ RAY_TRACING_ENTRY_ANY_HIT(MaterialAHS,
 	float3 V = WorldRayDirection();
 	float3 N = normalize(MaterialParameters.WorldNormal);
 	float VoN = abs(dot(V, N));
+#if !SIMPLIFIED_MATERIAL_SHADER
 	if (Opacity < 1.0 && Ior != 1.0)
 	{
 		float PathRoughness = PackedPayload.GetPathRoughness();
@@ -704,6 +719,7 @@ RAY_TRACING_ENTRY_ANY_HIT(MaterialAHS,
 		Transparency *= (1 - Roughness * Roughness) * saturate(PathRoughness - Roughness);
 	}
 	else
+#endif
 	{
 		// no bending at the interface so no fresnel - equation above simplifies to just A
 		float3 A = pow(Transmission, 1.0 / VoN);
@@ -715,12 +731,13 @@ RAY_TRACING_ENTRY_ANY_HIT(MaterialAHS,
 #elif MATERIALBLENDING_TRANSLUCENT
 	float Opacity = GetMaterialOpacity(PixelMaterialInputs);
 	float Transparency = 1 - Opacity;
-#if REFRACTION_USE_INDEX_OF_REFRACTION
+#if REFRACTION_USE_INDEX_OF_REFRACTION && !SIMPLIFIED_MATERIAL_SHADER
 	float Ior = GetMaterialRefraction(PixelMaterialInputs).x;
 #else
 	float Ior = 1.0;
 #endif
 
+#if !SIMPLIFIED_MATERIAL_SHADER
 	// Is refraction enabled?
 	if (Transparency > 0 && Ior != 1.0)
 	{
@@ -744,6 +761,7 @@ RAY_TRACING_ENTRY_ANY_HIT(MaterialAHS,
 		float Fr = FresnelReflectance(VoN, Eta);
 		Transparency *= pow(1 - Fr, 2) * (1 - Roughness * Roughness) * saturate(PathRoughness - Roughness);
 	}
+#endif
 #else
 #error "Control flow does not account for all blending modes!"
 #endif
diff --git a/Engine/Shaders/Shared/PathTracingDefinitions.h b/Engine/Shaders/Shared/PathTracingDefinitions.h
--- a/Engine/Shaders/Shared/PathTracingDefinitions.h
+++ b/Engine/Shaders/Shared/PathTracingDefinitions.h
@@ -28,25 +28,3 @@
 // Constants for the path tracer light grid
 #define PATHTRACER_LIGHT_GRID_SINGULAR_MASK					0x80000000u
 #define PATHTRACER_LIGHT_GRID_LIGHT_COUNT_MASK				0x7FFFFFFFu
-
-#ifdef __cplusplus
-
-BEGIN_SHADER_PARAMETER_STRUCT(FPathTracingData, )
-	SHADER_PARAMETER(uint32, Iteration)
-	SHADER_PARAMETER(uint32, TemporalSeed)
-	SHADER_PARAMETER(uint32, MaxSamples)
-	SHADER_PARAMETER(uint32, MaxBounces)
-	SHADER_PARAMETER(uint32, MaxSSSBounces)
-	SHADER_PARAMETER(uint32, MISMode)
-	SHADER_PARAMETER(uint32, ApproximateCaustics)
-	SHADER_PARAMETER(uint32, EnableCameraBackfaceCulling)
-	SHADER_PARAMETER(uint32, EnableDirectLighting)
-	SHADER_PARAMETER(uint32, EnableEmissive)
-	SHADER_PARAMETER(uint32, SamplerType)
-	SHADER_PARAMETER(uint32, VisualizeLightGrid)
-	SHADER_PARAMETER(float, MaxPathIntensity)
-	SHADER_PARAMETER(float, MaxNormalBias)
-	SHADER_PARAMETER(float, FilterWidth)
-END_SHADER_PARAMETER_STRUCT()
-
-#endif
diff --git a/Engine/Source/Runtime/Renderer/Private/PathTracing.cpp b/Engine/Source/Runtime/Renderer/Private/PathTracing.cpp
--- a/Engine/Source/Runtime/Renderer/Private/PathTracing.cpp
+++ b/Engine/Source/Runtime/Renderer/Private/PathTracing.cpp
@@ -178,6 +178,25 @@ TAutoConsoleVariable<int32> CVarPathTracingLightGridVisualize(
 	ECVF_RenderThreadSafe
 );
 
+BEGIN_SHADER_PARAMETER_STRUCT(FPathTracingData, )
+	SHADER_PARAMETER(uint32, Iteration)
+	SHADER_PARAMETER(uint32, TemporalSeed)
+	SHADER_PARAMETER(uint32, MaxSamples)
+	SHADER_PARAMETER(uint32, MaxBounces)
+	SHADER_PARAMETER(uint32, MaxSSSBounces)
+	SHADER_PARAMETER(uint32, MISMode)
+	SHADER_PARAMETER(uint32, ApproximateCaustics)
+	SHADER_PARAMETER(uint32, EnableCameraBackfaceCulling)
+	SHADER_PARAMETER(uint32, EnableDirectLighting)
+	SHADER_PARAMETER(uint32, EnableEmissive)
+	SHADER_PARAMETER(uint32, SamplerType)
+	SHADER_PARAMETER(uint32, VisualizeLightGrid)
+	SHADER_PARAMETER(float, MaxPathIntensity)
+	SHADER_PARAMETER(float, MaxNormalBias)
+	SHADER_PARAMETER(float, FilterWidth)
+END_SHADER_PARAMETER_STRUCT()
+
+
 // Store the rendering options used on the previous frame so we can correctly invalidate when things change
 struct FPathTracingConfig
 {
