fixed SkyLight IBL looking blurry on Android

diff --git a/Engine/Shaders/Private/ReflectionEnvironmentShaders.usf b/Engine/Shaders/Private/ReflectionEnvironmentShaders.usf
--- a/Engine/Shaders/Private/ReflectionEnvironmentShaders.usf
+++ b/Engine/Shaders/Private/ReflectionEnvironmentShaders.usf
@@ -354,7 +354,7 @@ void FilterPS(
 	}
 
 	uint CubeSize = 1 << ( NumMips - 1 );
-	const float SolidAngleTexel = 4*PI / ( 6 * CubeSize * CubeSize ) * 2;
+	const float SolidAngleTexel = 4.0f * PI / float(6 * CubeSize * CubeSize) * 2.0f;
 
 	//const uint NumSamples = 1024;
 	const uint NumSamples = Roughness < 0.1 ? 32 : 64;
