Workaround for GPU page fault accessing a GPUScene buffer out of bounds.
 - A proper fix is required. but this prevents the title crashing.

diff --git a/Engine/Shaders/Private/SceneData.ush b/Engine/Shaders/Private/SceneData.ush
--- a/Engine/Shaders/Private/SceneData.ush
+++ b/Engine/Shaders/Private/SceneData.ush
@@ -188,7 +188,22 @@ struct FPrimitiveSceneData
 float4 LoadPrimitivePrimitiveSceneDataElement(uint PrimitiveIndex, uint ItemIndex)
 {
 #if USE_GLOBAL_GPU_SCENE_DATA
-	return GPUScenePrimitiveSceneData[PrimitiveIndex + ItemIndex];
+
+	// BEGIN: Workaround for UE-115402
+
+	// Check the read is in-bounds
+	uint NumElements, Stride;
+	GPUScenePrimitiveSceneData.GetDimensions(NumElements, Stride);
+
+	uint TargetIdx = PrimitiveIndex + ItemIndex;
+	if (TargetIdx >= NumElements)
+	{
+		return float4(0, 0, 0, 0);
+	}
+
+	// END: Workaround for UE-115402
+
+	return GPUScenePrimitiveSceneData[TargetIdx];
 #else
 	return View.PrimitiveSceneData[PrimitiveIndex + ItemIndex];
 #endif
