Add pragma once
Add default GPUSKIN_USE_EXTRA_INFLUENCES define

diff --git a/Engine/Shaders/Private/GpuSkinCommon.ush b/Engine/Shaders/Private/GpuSkinCommon.ush
--- a/Engine/Shaders/Private/GpuSkinCommon.ush
+++ b/Engine/Shaders/Private/GpuSkinCommon.ush
@@ -4,6 +4,10 @@
 	GpuSkinCacheCommon.usf: Common properties and defines for GPU Skinning Cache
 =============================================================================*/
 
+#pragma once
+
+#include "Common.ush"
+
 // These offsets are in uint32/floats; they attempt to match TGPUSkinVertexBase<> & derived classes:
 //	TGPUSkinVertexBase
 //	TGPUSkinVertexFloat16Uvs<1..N>
@@ -16,11 +20,13 @@
 //	Position			2,3					2,3								2,3
 //	UVs[1..N]			N/A					5,1..N							5,(1..N)*2
 
-#include "Common.ush"
-
 #define	GPUSKIN_VB_OFFSET_TANGENT_X			0
 #define	GPUSKIN_VB_OFFSET_TANGENT_Z			1
 
+#ifndef GPUSKIN_USE_EXTRA_INFLUENCES
+#define GPUSKIN_USE_EXTRA_INFLUENCES		0
+#endif
+
 #ifndef GPUSKIN_BONE_INDEX_UINT16
 #define GPUSKIN_BONE_INDEX_UINT16			0	// 16-bit bone index if 1, otherwise 8-bit bone index
 #endif
