Fixing a flood of compile warnings caused by macro arguments not being wrapped correctly. This might have been causing unintended bugs.

diff --git a/Engine/Shaders/Private/Strata/Strata.ush b/Engine/Shaders/Private/Strata/Strata.ush
--- a/Engine/Shaders/Private/Strata/Strata.ush
+++ b/Engine/Shaders/Private/Strata/Strata.ush
@@ -68,7 +68,7 @@ struct FStrataBSDF
 #define TO_BITMASK(BitCount)							((1u<<BitCount)-1)
 #define TO_OFFSET_BITMASK(BitCount, BitOffset)			(TO_BITMASK(BitCount)<<BitOffset)
 #define READ_BITS(X, BitCount, BitOffset)				uint((X>>BitOffset) & TO_BITMASK(BitCount))
-#define WRITE_BITS(X, BitCount, BitOffset, NewValue)	X = (X & (~TO_OFFSET_BITMASK(BitCount, BitOffset))) | ((NewValue & TO_BITMASK(BitCount)) << BitOffset)
+#define WRITE_BITS(X, BitCount, BitOffset, NewValue)	X = (X & (~TO_OFFSET_BITMASK(BitCount, BitOffset))) | (((NewValue) & TO_BITMASK(BitCount)) << BitOffset)
 
 // Bit count for each component stored in a Strata header
 // For HEADER_BIT_BSDF_COUNT and  HEADER_BIT_SHAREDNORMAL_X, please refer to values in StrataDefinitions.h
