Fix a perf regression in ReflectionEnvironmentAndSky

diff --git a/Engine/Shaders/Private/ClearCoatCommon.ush b/Engine/Shaders/Private/ClearCoatCommon.ush
--- a/Engine/Shaders/Private/ClearCoatCommon.ush
+++ b/Engine/Shaders/Private/ClearCoatCommon.ush
@@ -44,6 +44,19 @@ void RemapClearCoatDiffuseAndSpecularColor(FGBufferData GBuffer, float NoV, inou
 	}
 }
 
+void RemapClearCoatDiffuseAndSpecularColor(FGBufferData GBuffer, float2 ScreenPosition, inout float3 DiffuseColor, inout float3 SpecularColor)
+{
+	if (GBuffer.ShadingModelID == SHADINGMODELID_CLEAR_COAT)
+	{
+		float3 WorldPosition = mul(float4(ScreenPosition * GBuffer.Depth, GBuffer.Depth, 1), View.ScreenToWorld).xyz;
+		float3 CameraToPixel = normalize(WorldPosition - View.WorldCameraOrigin);
+		float3 V = -CameraToPixel;
+		float NoV = saturate(dot(GBuffer.WorldNormal, V));
+
+		RemapClearCoatDiffuseAndSpecularColor(GBuffer, NoV, DiffuseColor, SpecularColor);
+	}
+}
+
 float3 ClearCoatLayerCombine(FGBufferData GBuffer, float NoV, float3 TopLayerReflections, float3 BottomLayerReflections, float3 SpecularColor)
 {
 	const float ClearCoat			= GBuffer.CustomData.x;
diff --git a/Engine/Shaders/Private/ReflectionEnvironmentPixelShader.usf b/Engine/Shaders/Private/ReflectionEnvironmentPixelShader.usf
--- a/Engine/Shaders/Private/ReflectionEnvironmentPixelShader.usf
+++ b/Engine/Shaders/Private/ReflectionEnvironmentPixelShader.usf
@@ -190,7 +190,7 @@ float3 SkyLightDiffuse(FGBufferData GBuffer, float AmbientOcclusion, float2 Buff
 	return Lighting;
 }
 
-float3 ReflectionEnvironment(FGBufferData GBuffer, float AmbientOcclusion, float2 BufferUV, float3 WorldPosition, float3 V, float4 SvPosition, float3 BentNormal, float3 SpecularColor)
+float3 ReflectionEnvironment(FGBufferData GBuffer, float AmbientOcclusion, float2 BufferUV, float2 ScreenPosition, float4 SvPosition, float3 BentNormal, float3 SpecularColor)
 {
 	float4 Color = float4(0, 0, 0, 1);
 
@@ -205,6 +205,10 @@ float3 ReflectionEnvironment(FGBufferData GBuffer, float AmbientOcclusion, float
 	}
 #endif
 
+	float3 WorldPosition = mul(float4(ScreenPosition * GBuffer.Depth, GBuffer.Depth, 1), View.ScreenToWorld).xyz;
+	float3 CameraToPixel = normalize(WorldPosition - View.WorldCameraOrigin);
+	float3 ReflectionVector = reflect(CameraToPixel, GBuffer.WorldNormal);
+    float3 V = -CameraToPixel;
 	float3 N = GBuffer.WorldNormal;
 
 	const float3 SavedTopLayerNormal = N;
@@ -323,6 +327,7 @@ void ApplyCloudVolumetricShadow(in float3 WorldPosition, inout float3 SkyLightin
 	// To add some spatial variation, we can affect the sky light diffuse contribution according to cloud occlusion from above.
 	// This is an approximation because clouds are also occluding the sky in the sky light capture (a bit of a double contribution then). but it does help by adding spatially varying details.
 #if 0	// DISABLED for now because it has artefact with specular not being affected (and thus looking too bright which can be confusing for users).
+	// NOTE: Consider wrapping this inside a define when enabled. Otherwise, perf will regress on less powerful platforms
 	if (CloudSkyAOEnabled)
 	{
 		float OutOpticalDepth = 0.0f;
@@ -405,15 +410,9 @@ void ReflectionEnvironmentSkyLighting(
 	uint ShadingModelID = GBuffer.ShadingModelID;
 	const bool bUnlitMaterial = ShadingModelID == SHADINGMODELID_UNLIT;
 
-	float3 WorldPosition = mul(float4(ScreenPosition * GBuffer.Depth, GBuffer.Depth, 1), View.ScreenToWorld).xyz;
-	float3 CameraToPixel = normalize(WorldPosition - View.WorldCameraOrigin);
-	float3 N = GBuffer.WorldNormal;
-	float3 V = -CameraToPixel;
-	float NoV = saturate(dot(N, V));
-
 	float3 DiffuseColor = GBuffer.DiffuseColor;
 	float3 SpecularColor = GBuffer.SpecularColor;
-	RemapClearCoatDiffuseAndSpecularColor(GBuffer, NoV, DiffuseColor, SpecularColor);
+	RemapClearCoatDiffuseAndSpecularColor(GBuffer, ScreenPosition, DiffuseColor, SpecularColor);
 
 	// Sample the ambient occlusion that is dynamically generated every frame.
 	float AmbientOcclusion = AmbientOcclusionTexture.SampleLevel(AmbientOcclusionSampler, BufferUV, 0).r;
@@ -432,6 +431,7 @@ void ReflectionEnvironmentSkyLighting(
 	{
 		float3 SkyLighting = SkyLightDiffuse(GBuffer, AmbientOcclusion, BufferUV, ScreenPosition, BentNormal, DiffuseColor);
 
+		float3 WorldPosition = mul(float4(ScreenPosition * GBuffer.Depth, GBuffer.Depth, 1), View.ScreenToWorld).xyz;
 		ApplyCloudVolumetricShadow(WorldPosition, SkyLighting);
 
 		FLightAccumulator LightAccumulator = (FLightAccumulator)0;
@@ -446,7 +446,7 @@ void ReflectionEnvironmentSkyLighting(
 	BRANCH
 	if (!bUnlitMaterial && ShadingModelID != SHADINGMODELID_HAIR)
 	{
-		OutColor.xyz += ReflectionEnvironment(GBuffer, AmbientOcclusion, BufferUV, WorldPosition, V, SvPosition, BentNormal, SpecularColor);
+		OutColor.xyz += ReflectionEnvironment(GBuffer, AmbientOcclusion, BufferUV, ScreenPosition, SvPosition, BentNormal, SpecularColor);
 	}
 
 #endif // STRATA_ENABLED
diff --git a/Engine/Source/Runtime/Renderer/Private/IndirectLightRendering.cpp b/Engine/Source/Runtime/Renderer/Private/IndirectLightRendering.cpp
--- a/Engine/Source/Runtime/Renderer/Private/IndirectLightRendering.cpp
+++ b/Engine/Source/Runtime/Renderer/Private/IndirectLightRendering.cpp
@@ -1139,14 +1139,9 @@ static void AddSkyReflectionPass(
 	bool bHasBoxCaptures = (View.NumBoxReflectionCaptures > 0);
 	bool bHasSphereCaptures = (View.NumSphereReflectionCaptures > 0);
 
-	float DynamicBentNormalAO = 0.0f;
+	float DynamicBentNormalAO = DynamicBentNormalAOTexture ? 1.0f : 0.0f;
 	const FRDGSystemTextures& SystemTextures = FRDGSystemTextures::Get(GraphBuilder);
-	FRDGTextureRef AmbientOcclusionTexture = GetScreenSpaceAOFallback(SystemTextures);
-	if (HasBeenProduced(SceneTextures.ScreenSpaceAO))
-	{
-		DynamicBentNormalAO = 1.0f;
-		AmbientOcclusionTexture = SceneTextures.ScreenSpaceAO;
-	}
+	FRDGTextureRef AmbientOcclusionTexture = HasBeenProduced(SceneTextures.ScreenSpaceAO) ? SceneTextures.ScreenSpaceAO : GetScreenSpaceAOFallback(SystemTextures);
 
 	const auto& SceneColorTexture = SceneTextures.Color;
 
