Implements Dual Pixel Vectorization in TSR's UpdateHistory

diff --git a/Engine/Shaders/Private/TemporalAA/TAACommon.ush b/Engine/Shaders/Private/TemporalAA/TAACommon.ush
--- a/Engine/Shaders/Private/TemporalAA/TAACommon.ush
+++ b/Engine/Shaders/Private/TemporalAA/TAACommon.ush
@@ -6,6 +6,7 @@
 #include "../FastMath.ush"
 #include "../MonteCarlo.ush"
 #include "../ScreenPass.ush"
+#include "/Engine/Public/DualPixelVectorization.ush"
 
 
 //------------------------------------------------------- COMPILER CONFIG
@@ -26,6 +27,13 @@
 
 #define CONFIG_ENABLE_STOCASTIC_QUANTIZATION 1
 
+#if CONFIG_COMPILE_FP16
+	// Take advantage of RDNA's v_pk_*_{uif}16 instructions
+	#define CONFIG_ENABLE_DUAL_PIXEL_VECTORIZATION 1
+#else
+	#define CONFIG_ENABLE_DUAL_PIXEL_VECTORIZATION 0
+#endif
+
 
 //------------------------------------------------------- CONSTANTS
 
@@ -150,6 +158,18 @@ static const uint kPlusIndexes3x3[5] = { 4, 1, 3, 5, 7 };
 	#define taa_ushort3 uint16_t3
 	#define taa_ushort4 uint16_t4
 
+	#define taa_half2x2 half2x2
+	#define taa_half3x2 half3x2
+	#define taa_half4x2 half4x2
+
+	#define taa_short2x2 int16_t2x2
+	#define taa_short3x2 int16_t3x2
+	#define taa_short4x2 int16_t4x2
+
+	#define taa_ushort2x2 uint16_t2x2
+	#define taa_ushort3x2 uint16_t3x2
+	#define taa_ushort4x2 uint16_t4x2
+
 #else
 	#define taa_half  float
 	#define taa_half2 float2
@@ -166,10 +186,24 @@ static const uint kPlusIndexes3x3[5] = { 4, 1, 3, 5, 7 };
 	#define taa_ushort3 uint3
 	#define taa_ushort4 uint4
 
+	#define taa_half2x2 float2x2
+	#define taa_half3x2 float3x2
+	#define taa_half4x2 float4x2
+
+	#define taa_short2x2 int2x2
+	#define taa_short3x2 int3x2
+	#define taa_short4x2 int4x2
+
+	#define taa_ushort2x2 uint2x2
+	#define taa_ushort3x2 uint3x2
+	#define taa_ushort4x2 uint4x2
+
 #endif
 
 #define taa_subpixel_details uint
+#define taa_subpixel_details2 uint2
 #define taa_subpixel_payload uint // TODO: uint16_t
+#define taa_subpixel_payload2 uint2 // TODO: uint16_t
 
 
 //------------------------------------------------------- PARAMETERS
@@ -247,6 +281,21 @@ uint bEnableInterferenceHeuristic;
 
 
 //------------------------------------------------------- FUNCTIONS
+	
+float2x2 ApplyScreenTransform(float2x2 PInA, FScreenTransform AToB)
+{
+	return dpv_add(dpv_mul(PInA, AToB.xy), AToB.zw);
+}
+
+#if PLATFORM_SUPPORTS_REAL_TYPES
+
+half2x2 ApplyScreenTransform(half2x2 PInA, FScreenTransform AToB)
+{
+	return dpv_add(dpv_mul(PInA, half2(AToB.xy)), half2(AToB.zw));
+}
+
+#endif
+
 
 /** Compute the group wave index into SGRP to then recompue the GroupThreadIndex latter */
 uint GetGroupWaveIndex(uint GroupThreadIndex, uint GroupSize)
@@ -313,6 +362,28 @@ taa_short2 ClampPixelOffset(
 	return ClampedOffset;
 }
 
+taa_short2x2 ClampPixelOffset(
+	taa_short2x2 SamplePixelPos,
+	int2 MinPixelPos, int2 MaxPixelPos)
+{
+	taa_short2 SamplePixelPos0 = dpv_lo(SamplePixelPos);
+	taa_short2 SamplePixelPos1 = dpv_hi(SamplePixelPos);
+
+	SamplePixelPos0.x = min(SamplePixelPos0.x, taa_short(MaxPixelPos.x));
+	SamplePixelPos0.x = max(SamplePixelPos0.x, taa_short(MinPixelPos.x));
+		
+	SamplePixelPos0.y = min(SamplePixelPos0.y, taa_short(MaxPixelPos.y));
+	SamplePixelPos0.y = max(SamplePixelPos0.y, taa_short(MinPixelPos.y));
+	
+	SamplePixelPos1.x = min(SamplePixelPos1.x, taa_short(MaxPixelPos.x));
+	SamplePixelPos1.x = max(SamplePixelPos1.x, taa_short(MinPixelPos.x));
+		
+	SamplePixelPos1.y = min(SamplePixelPos1.y, taa_short(MaxPixelPos.y));
+	SamplePixelPos1.y = max(SamplePixelPos1.y, taa_short(MinPixelPos.y));
+
+	return dpv_interleave_registers(SamplePixelPos0, SamplePixelPos1);
+}
+
 taa_short2 ClampPixelOffset(
 	taa_short2 SamplePixelPos,
 	int2 MinPixelPos, int2 MaxPixelPos)
@@ -326,6 +397,42 @@ taa_short2 ClampPixelOffset(
 	return SamplePixelPos;
 }
 
+taa_short2x2 AddAndClampPixelOffset(
+	taa_short2x2 KernelCenterPixelPos,
+	taa_short2x2 Offset,
+	const taa_short2 OffsetDirection,
+	int2 MinPixelPos, int2 MaxPixelPos)
+{
+	taa_short2x2 SamplePixelPos = KernelCenterPixelPos + Offset;
+	
+	taa_short2 SamplePixelPos0 = dpv_lo(SamplePixelPos);
+	taa_short2 SamplePixelPos1 = dpv_hi(SamplePixelPos);
+
+	// Only do clamp based on the compile time known direction of the offset.
+	// This turns only 1 v_pk_max_i16 and v_pk_min_i16 for a 3x3 kernel.
+	if (OffsetDirection.x > 0)
+		SamplePixelPos0.x = min(SamplePixelPos0.x, taa_short(MaxPixelPos.x));
+	else if (OffsetDirection.x < 0)
+		SamplePixelPos0.x = max(SamplePixelPos0.x, taa_short(MinPixelPos.x));
+		
+	if (OffsetDirection.y > 0)
+		SamplePixelPos0.y = min(SamplePixelPos0.y, taa_short(MaxPixelPos.y));
+	else if (OffsetDirection.y < 0)
+		SamplePixelPos0.y = max(SamplePixelPos0.y, taa_short(MinPixelPos.y));
+	
+	if (OffsetDirection.x > 0)
+		SamplePixelPos1.x = min(SamplePixelPos1.x, taa_short(MaxPixelPos.x));
+	else if (OffsetDirection.x < 0)
+		SamplePixelPos1.x = max(SamplePixelPos1.x, taa_short(MinPixelPos.x));
+		
+	if (OffsetDirection.y > 0)
+		SamplePixelPos1.y = min(SamplePixelPos1.y, taa_short(MaxPixelPos.y));
+	else if (OffsetDirection.y < 0)
+		SamplePixelPos1.y = max(SamplePixelPos1.y, taa_short(MinPixelPos.y));
+	
+	return dpv_interleave_registers(SamplePixelPos0, SamplePixelPos1);
+}
+
 taa_short2 AddAndClampPixelOffset(
 	taa_short2 KernelCenterPixelPos,
 	taa_short2 Offset,
@@ -349,6 +456,28 @@ taa_short2 AddAndClampPixelOffset(
 	return SamplePixelPos;
 }
 
+taa_short2x2 InvalidateOutputPixelPos(taa_short2x2 PixelPos, uint2 ViewportMax)
+#if 1
+{
+	taa_short2x2 Subtract = dpv_sub(taa_short2(ViewportMax - 1), PixelPos);
+	taa_ushort2 Override = taa_ushort2(Subtract[0] | Subtract[1]);
+
+	#if CONFIG_COMPILE_FP16
+		PixelPos[0] |= -taa_short2((Override & uint16_t(0x8000)) >> 15);
+	#else
+		PixelPos[0] |= -taa_short2((Override & uint(0x80000000)) >> 31);
+	#endif
+
+	return PixelPos;
+}
+#else
+{
+	bool bIsValidPixel = all(PixelPos < ViewportMax);
+	PixelPos.x = bIsValidPixel ? PixelPos.x : ~taa_short(0);
+	return PixelPos;
+}
+#endif
+
 taa_short2 InvalidateOutputPixelPos(taa_short2 PixelPos, uint2 ViewportMax)
 #if 1
 {
@@ -398,6 +527,11 @@ float SafeRcp(float x)
 	return x > 0.0 ? rcp(x) : 0.0;
 }
 
+float2 SafeRcp(float2 x)
+{
+	return float2(SafeRcp(x.x), SafeRcp(x.y));
+}
+
 #if CONFIG_COMPILE_FP16
 
 half SafeRcp(half x)
@@ -406,99 +540,133 @@ half SafeRcp(half x)
 	return min(rcp(x), half(MaxHalfFloat)) * saturate(x * asfloat16(uint16_t(0x7C00)));
 }
 
+half2 SafeRcp(half2 x)
+{
+	// If x=0.0, then MaxHalfFloat * 0.0 = 0.0
+	return min(rcp(x), half(MaxHalfFloat)) * saturate(x * asfloat16(uint16_t(0x7C00)));
+}
+
 #endif
 
-taa_half3 RGBToYCoCg(taa_half3 RGB)
+taa_half3x2 RGBToYCoCg(taa_half3x2 RGB)
 {
-	taa_half Y = dot(RGB, taa_half3(1, 2, 1));
-	taa_half Co = dot(RGB, taa_half3(2, 0, -2));
-	taa_half Cg = dot(RGB, taa_half3(-1, 2, -1));
+	taa_half2 Y  = mul(taa_half3(1, 2, 1), RGB);
+	taa_half2 Co = mul(taa_half3(2, 0, -2), RGB);
+	taa_half2 Cg = mul(taa_half3(-1, 2, -1), RGB);
 
-	taa_half3 YCoCg = taa_half3(Y, Co, Cg);
+	taa_half3x2 YCoCg = taa_half3x2(Y, Co, Cg);
 	return YCoCg;
 }
 
-taa_half3 YCoCgToRGB(taa_half3 YCoCg)
+taa_half3x2 YCoCgToRGB(taa_half3x2 YCoCg)
 {
-	taa_half Y =  YCoCg.x * taa_half(0.25);
-	taa_half Co = YCoCg.y * taa_half(0.25);
-	taa_half Cg = YCoCg.z * taa_half(0.25);
+	taa_half2 Y =  YCoCg[0] * taa_half(0.25);
+	taa_half2 Co = YCoCg[1] * taa_half(0.25);
+	taa_half2 Cg = YCoCg[2] * taa_half(0.25);
 
-	taa_half R = Y + Co - Cg;
-	taa_half G = Y + Cg;
-	taa_half B = Y - Co - Cg;
+	taa_half2 R = Y + Co - Cg;
+	taa_half2 G = Y + Cg;
+	taa_half2 B = Y - Co - Cg;
 
-	taa_half3 RGB = taa_half3(R, G, B);
+	taa_half3x2 RGB = taa_half3x2(R, G, B);
 	return RGB;
 }
 
-taa_half3 YCoCgToLCoCg(taa_half3 YCoCg)
-{
-	return taa_half3(YCoCg.x, YCoCg.yz * SafeRcp(YCoCg.x));
-}
-
-taa_half3 LCoCgToYCoCg(taa_half3 LCoCg)
+taa_half3 RGBToYCoCg(taa_half3 RGB)
 {
-	return taa_half3(LCoCg.x, LCoCg.x * LCoCg.yz);
+	return dpv_lo(RGBToYCoCg(dpv_interleave_mono_registers(RGB)));
 }
 
-taa_half3 RGBToLCoCg(taa_half3 RGB)
+taa_half3 YCoCgToRGB(taa_half3 YCoCg)
 {
-	return YCoCgToLCoCg(RGBToYCoCg(RGB));
+	return dpv_lo(YCoCgToRGB(dpv_interleave_mono_registers(YCoCg)));
 }
 
-taa_half3 LCoCgToRGB(taa_half3 LCoCg)
-{
-	return YCoCgToRGB(LCoCgToYCoCg(LCoCg));
-}
 
 // Some bright pixel can cause HdrWeight to get nullified under fp16 representation. So clamping this to a value close to the minimum float float positive value (0.000061).
 #define HDR_WEIGHT_SAFE_MIN_VALUE 0.0001
 
 // Faster but less accurate luma computation. 
 // Luma includes a scaling by 4.
+taa_half2 Luma4(taa_half3x2 Color)
+{
+	return (Color[1] * taa_half(2.0)) + (Color[0] + Color[2]);
+}
+
 taa_half Luma4(taa_half3 Color)
 {
-	return (Color.g * taa_half(2.0)) + (Color.r + Color.b);
+	return dpv_lo(Luma4(dpv_interleave_mono_registers(Color)));
 }
 
-taa_half HdrWeightY(taa_half Color)
+taa_half2 HdrWeightY(taa_half2 Luma)
 {
 	taa_half Exposure = taa_half(1.0);
 
-	return max(taa_half(HDR_WEIGHT_SAFE_MIN_VALUE), rcp(Color * Exposure + taa_half(4.0)));
+	return max(taa_half(HDR_WEIGHT_SAFE_MIN_VALUE), rcp(Luma * Exposure + taa_half(4.0)));
 }
 
-taa_half HdrWeightInvY(taa_half Color) 
+taa_half HdrWeightY(taa_half Luma)
 {
-	return taa_half(4.0) * rcp(taa_half(1.0) - Color);
+	return dpv_lo(HdrWeightY(dpv_interleave_mono_registers(Luma)));
+}
+
+taa_half2 HdrWeightInvY(taa_half2 LDRLuma) 
+{
+	return taa_half(4.0) * rcp(taa_half(1.0) - LDRLuma);
+}
+
+taa_half HdrWeightInvY(taa_half LDRLuma) 
+{
+	return dpv_lo(HdrWeightInvY(dpv_interleave_mono_registers(LDRLuma)));
 }
 
 // Optimized HDR weighting function.
+taa_half2 HdrWeight4(taa_half3x2 Color)
+{
+	return HdrWeightY(Luma4(Color));
+}
+
 taa_half HdrWeight4(taa_half3 Color)
 {
 	return HdrWeightY(Luma4(Color));
 }
 
+
 // Returns the weight of a pixels at a coordinate <PixelDelta> from the PDF highest point.
-taa_half ComputeSampleWeigth(taa_half UpscaleFactor, taa_half2 PixelDelta, const float MinimalContribution)
+taa_half2 ComputeSampleWeigth(taa_half2 UpscaleFactor, taa_half2x2 PixelDelta, const float MinimalContribution)
 {
-	taa_half u2 = UpscaleFactor * UpscaleFactor;
+	taa_half2 u2 = UpscaleFactor * UpscaleFactor;
 
 	// 1 - 1.9 * x^2 + 0.9 * x^4
-	taa_half x2 = saturate(u2 * dot(PixelDelta, PixelDelta));
+	taa_half2 x2 = saturate(u2 * dpv_length2(PixelDelta));
 	//return taa_half(((float(0.9) + MinimalContribution) * x2 - float(1.9)) * x2 + float(1.0));
 	return saturate((taa_half(0.9) * x2 - taa_half(1.9)) * x2 + taa_half(1.0 + MinimalContribution));
 }
 
-taa_half2 WeightedLerpFactors(taa_half WeightA, taa_half WeightB, taa_half Blend)
+taa_half ComputeSampleWeigth(taa_half UpscaleFactor, taa_half2 PixelDelta, const float MinimalContribution)
 {
-	taa_half BlendA = (taa_half(1.0) - Blend) * WeightA;
-	taa_half BlendB = Blend * WeightB;
-	taa_half RcpBlend = SafeRcp(BlendA + BlendB);
+	return dpv_lo(ComputeSampleWeigth(
+		dpv_interleave_mono_registers(UpscaleFactor),
+		dpv_interleave_mono_registers(PixelDelta),
+		MinimalContribution));
+}
+
+taa_half2x2 WeightedLerpFactors(taa_half2 WeightA, taa_half2 WeightB, taa_half2 Blend)
+{
+	taa_half2 BlendA = (taa_half(1.0) - Blend) * WeightA;
+	taa_half2 BlendB = Blend * WeightB;
+	taa_half2 RcpBlend = SafeRcp(BlendA + BlendB);
 	BlendA *= RcpBlend;
 	BlendB *= RcpBlend;
-	return taa_half2(BlendA, BlendB);
+	return taa_half2x2(BlendA, BlendB);
+}
+
+taa_half2 WeightedLerpFactors(taa_half WeightA, taa_half WeightB, taa_half Blend)
+{
+	return dpv_lo(WeightedLerpFactors(
+		dpv_interleave_mono_registers(WeightA),
+		dpv_interleave_mono_registers(WeightB),
+		dpv_interleave_mono_registers(Blend)));
 }
 
 bool TakeOnlyOneSamplePair(float2 Offset)
@@ -530,37 +698,47 @@ bool IsOffScreen(uint bCameraCut, float2 ScreenPos, taa_half ParallaxRejectionMa
 	return (bIsCameraCut || bIsOutOfBounds || bIsParallaxRejected);
 }
 
-taa_half MeasureRejectionFactor(taa_half3 PrevYCoCg, taa_half3 ClampedPrevYCoCg, taa_half3 InputCenterYCoCg, taa_half3 InputMinYCoCg, taa_half3 InputMaxYCoCg)
+taa_half2 MeasureRejectionFactor(taa_half3x2 PrevYCoCg, taa_half3x2 ClampedPrevYCoCg, taa_half3x2 InputCenterYCoCg, taa_half3x2 InputMinYCoCg, taa_half3x2 InputMaxYCoCg)
 {
-	taa_half3 BoxSize = InputMaxYCoCg - InputMinYCoCg;
+	taa_half3x2 BoxSize = InputMaxYCoCg - InputMinYCoCg;
 
-	taa_half3 ClampedEnergy = abs(ClampedPrevYCoCg - PrevYCoCg);
-	taa_half3 Delta = abs(InputCenterYCoCg - PrevYCoCg);
+	taa_half3x2 ClampedEnergy = abs(ClampedPrevYCoCg - PrevYCoCg);
+	taa_half3x2 Delta = abs(InputCenterYCoCg - PrevYCoCg);
 
-	taa_half3 ClampError;
+	taa_half3x2 ClampError;
 	#if 0
 	{
 		ClampError = BoxSize * taa_half(0.25);
 	}
 	#else
 	{
-		taa_half Luma = max(ClampedPrevYCoCg.x, PrevYCoCg.x);
+		taa_half2 Luma = max(ClampedPrevYCoCg[0], PrevYCoCg[0]);
 
 		// Measure the error due to the quantization of the color in channels.
-		taa_half Tonemap = Luma * HdrWeightY(Luma) + taa_half(0.5 / 1024.0);
-		taa_half PixelColorQuantizationError = abs(Tonemap * HdrWeightInvY(Tonemap) - Luma);
+		taa_half2 Tonemap = Luma * HdrWeightY(Luma) + taa_half(0.5 / 1024.0);
+		taa_half2 PixelColorQuantizationError = abs(Tonemap * HdrWeightInvY(Tonemap) - Luma);
 		
-		ClampError = max(BoxSize * taa_half(0.25), PixelColorQuantizationError);
+		ClampError = max(BoxSize * taa_half(0.25), taa_half3x2(PixelColorQuantizationError, PixelColorQuantizationError, PixelColorQuantizationError));
 	}
 	#endif
 
-	taa_half3 Factor = taa_half(1.0) - saturate(max(ClampedEnergy - ClampError, taa_half(0.0)) / max(Delta, BoxSize));
+	taa_half3x2 Factor = taa_half(1.0) - saturate(max(ClampedEnergy - ClampError, taa_half3x2(0.0, 0.0, 0.0, 0.0, 0.0, 0.0)) / max(Delta, BoxSize));
 
-	taa_half Rejection = min3(Factor.r, Factor.g, Factor.b);
+	taa_half2 Rejection = min3(Factor[0], Factor[1], Factor[2]);
 	
 	return Rejection;
 }
 
+taa_half MeasureRejectionFactor(taa_half3 PrevYCoCg, taa_half3 ClampedPrevYCoCg, taa_half3 InputCenterYCoCg, taa_half3 InputMinYCoCg, taa_half3 InputMaxYCoCg)
+{
+	return dpv_lo(MeasureRejectionFactor(
+		dpv_interleave_mono_registers(PrevYCoCg), 
+		dpv_interleave_mono_registers(ClampedPrevYCoCg), 
+		dpv_interleave_mono_registers(InputCenterYCoCg), 
+		dpv_interleave_mono_registers(InputMinYCoCg), 
+		dpv_interleave_mono_registers(InputMaxYCoCg)));
+}
+
 struct FSubpixelNeighborhood
 {
 	// 2x2 quad of subpixels grid
@@ -660,29 +838,33 @@ taa_half LumaToInterferenceSeed(taa_half Luma)
 	return Luma * HdrWeightY(Luma);
 }
 
-taa_half3 QuantizeForFloatRenderTarget(taa_half3 Color, taa_half E, float3 QuantizationError)
+taa_half3x2 QuantizeForFloatRenderTarget(taa_half3x2 Color, taa_half E, const float3 QuantizationError)
 {
-	taa_half3 Error = Color * taa_half3(QuantizationError);
+	taa_half3x2 Error = dpv_mul(Color, taa_half3(QuantizationError));
 
 	#if 0
 		// NOP
 	#elif CONFIG_COMPILE_FP16
 	{
-		uint16_t3 UError = asint16(Error);
-
-		Error = asfloat16(UError & uint16_t(~0x03FF));
+		Error[0] = asfloat16(asint16(Error[0]) & uint16_t(~0x03FF));
+		Error[1] = asfloat16(asint16(Error[1]) & uint16_t(~0x03FF));
+		Error[2] = asfloat16(asint16(Error[2]) & uint16_t(~0x03FF));
 	}
 	#else
 	{
-		Error.x = asfloat(asuint(Error.x) & ~0x007FFFFF);
-		Error.y = asfloat(asuint(Error.y) & ~0x007FFFFF);
-		Error.z = asfloat(asuint(Error.z) & ~0x007FFFFF);
+		Error[0] = asfloat(asuint(Error[0]) & ~0x007FFFFF);
+		Error[1] = asfloat(asuint(Error[1]) & ~0x007FFFFF);
+		Error[2] = asfloat(asuint(Error[2]) & ~0x007FFFFF);
 	}
 	#endif
 	
 	return Color + Error * E;
 }
 
+taa_half3 QuantizeForFloatRenderTarget(taa_half3 Color, taa_half E, const float3 QuantizationError)
+{
+	return dpv_lo(QuantizeForFloatRenderTarget(dpv_interleave_mono_registers(Color), E, QuantizationError));
+}
 
 
 struct FSuperResHistory
diff --git a/Engine/Shaders/Private/TemporalAA/TAAUpdateHistory.usf b/Engine/Shaders/Private/TemporalAA/TAAUpdateHistory.usf
--- a/Engine/Shaders/Private/TemporalAA/TAAUpdateHistory.usf
+++ b/Engine/Shaders/Private/TemporalAA/TAAUpdateHistory.usf
@@ -14,7 +14,7 @@
 
 #define CONFIG_RESPONSIVE_STENCIL 1
 
-#define CONFIG_SAMPLES 6
+#define CONFIG_SAMPLES 5
 
 #define CONFIG_MANUAL_LDS_SPILL (CONFIG_COMPILE_FP16)
 
@@ -61,64 +61,88 @@ RWTexture2D<taa_half3> SceneColorOutput;
 
 //------------------------------------------------------- LDS
 
-groupshared taa_half4 SharedArray0[TILE_SIZE * TILE_SIZE];
-groupshared taa_half4 SharedArray1[TILE_SIZE * TILE_SIZE];
-groupshared taa_half4 SharedArray2[TILE_SIZE * TILE_SIZE];
+#if CONFIG_ENABLE_DUAL_PIXEL_VECTORIZATION
+#define GROUP_SHARED_SIZE (TILE_SIZE * TILE_SIZE / 2)
+#else
+#define GROUP_SHARED_SIZE (TILE_SIZE * TILE_SIZE)
+#endif
+
+groupshared taa_half4 SharedArray0[GROUP_SHARED_SIZE];
+groupshared taa_half4 SharedArray1[GROUP_SHARED_SIZE];
+groupshared taa_half4 SharedArray2[GROUP_SHARED_SIZE];
+groupshared taa_half4 SharedArray3[GROUP_SHARED_SIZE];
+groupshared taa_half4 SharedArray4[GROUP_SHARED_SIZE];
+groupshared taa_half4 SharedArray5[GROUP_SHARED_SIZE];
+groupshared taa_subpixel_payload2 SharedParallaxFactorBits[GROUP_SHARED_SIZE];
 
 
 //------------------------------------------------------- FUNCTIONS
 
-taa_half ComputeFallbackContribution(taa_half HistoryValidity)
+taa_ushort2x2 Map8x8Tile2x2LaneDPV(uint GroupThreadIndex)
+{
+	taa_ushort2 GroupId = Map8x8Tile2x2Lane(GroupThreadIndex);
+#if CONFIG_ENABLE_DUAL_PIXEL_VECTORIZATION
+	return dpv_interleave_registers(GroupId, GroupId + taa_ushort2(0,  8 / 2));
+#else
+	return dpv_interleave_mono_registers(GroupId);
+#endif
+}
+
+taa_half2 ComputeFallbackContribution(taa_half2 HistoryValidity)
 {
 	return saturate((HistoryValidity - taa_half(1.0 * rcp(MAX_SAMPLE_COUNT))) * taa_half(6.0));
 }
 
-taa_half3 TransformColorForClampingBox(taa_half3 Color)
+taa_half3x2 TransformColorForClampingBox(taa_half3x2 Color)
 {
 	return RGBToYCoCg(Color);
 }
 
 void ComputeInputKernelSamplePosition(
-	taa_short2 InputPixelPos,
-	taa_half2 dKO,
+	taa_short2x2 InputPixelPos,
+	taa_half2x2 dKO,
 	uint SampleId,
-	out taa_short2 SampleInputPixelPos,
-	out taa_half2 PixelOffset)
+	out taa_short2x2 SampleInputPixelPos,
+	out taa_half2x2 PixelOffset)
 {
 	
 	#if CONFIG_SAMPLES == 9
 	{
 		taa_short2 iPixelOffset = taa_short2(kOffsets3x3[kSquareIndexes3x3[SampleId]]);
-		PixelOffset = taa_half2(iPixelOffset);
+		PixelOffset = dpv_interleave_registers(taa_half2(iPixelOffset), taa_half2(iPixelOffset));
 				
 		SampleInputPixelPos = AddAndClampPixelOffset(
 			InputPixelPos,
-			iPixelOffset, iPixelOffset,
+			dpv_interleave_registers(iPixelOffset, iPixelOffset), iPixelOffset,
 			InputPixelPosMin, InputPixelPosMax);
 	}
 	#elif CONFIG_SAMPLES == 5 || CONFIG_SAMPLES == 6
 	{
 		if (SampleId == 5)
 		{
-			taa_short2 iPixelOffset;
+			taa_short2x2 iPixelOffset;
 			#if CONFIG_COMPILE_FP16
-				iPixelOffset = int16_t2(1, 1) - int16_t2((asuint16(dKO) & uint16_t(0x8000)) >> uint16_t(14));
-				PixelOffset = asfloat16(asuint16(half(1.0)).xx | (asuint16(dKO) & uint16_t(0x8000)));
+				iPixelOffset[0] = int16_t(1) - int16_t2((asuint16(dKO[0]) & uint16_t(0x8000)) >> uint16_t(14));
+				iPixelOffset[1] = int16_t(1) - int16_t2((asuint16(dKO[1]) & uint16_t(0x8000)) >> uint16_t(14));
+				PixelOffset[0] = asfloat16(asuint16(half(1.0)).xx | (asuint16(dKO[0]) & uint16_t(0x8000)));
+				PixelOffset[1] = asfloat16(asuint16(half(1.0)).xx | (asuint16(dKO[1]) & uint16_t(0x8000)));
 			#else
-				iPixelOffset = SignFastInt(dKO);
-				PixelOffset = asfloat(asuint(1.0).xx | (asuint(dKO) & uint(0x80000000)));
+				iPixelOffset = dpv_interleave_registers(
+					SignFastInt(dpv_lo(dKO)), SignFastInt(dpv_hi(dKO)));
+				PixelOffset[0] = asfloat(asuint(1.0).xx | (asuint(dKO[0]) & uint(0x80000000)));
+				PixelOffset[1] = asfloat(asuint(1.0).xx | (asuint(dKO[1]) & uint(0x80000000)));
 			#endif
 						
-			SampleInputPixelPos = ClampPixelOffset(InputPixelPos, InputPixelPosMin, InputPixelPosMax);
+			SampleInputPixelPos = ClampPixelOffset(InputPixelPos + iPixelOffset, InputPixelPosMin, InputPixelPosMax);
 		}
 		else
 		{
 			taa_short2 iPixelOffset = taa_short2(kOffsets3x3[kPlusIndexes3x3[SampleId]]);
-			PixelOffset = taa_half2(iPixelOffset);
+			PixelOffset = dpv_interleave_registers(taa_half2(iPixelOffset), taa_half2(iPixelOffset));
 					
 			SampleInputPixelPos = AddAndClampPixelOffset(
 				InputPixelPos,
-				iPixelOffset, iPixelOffset,
+				dpv_interleave_registers(iPixelOffset, iPixelOffset), iPixelOffset,
 				InputPixelPosMin, InputPixelPosMax);
 		}
 	}
@@ -129,69 +153,91 @@ void ComputeInputKernelSamplePosition(
 
 //------------------------------------------------------- ENTRY POINT
 
-[numthreads(TILE_SIZE, TILE_SIZE, 1)]
+#if CONFIG_ENABLE_DUAL_PIXEL_VECTORIZATION
+	[numthreads(TILE_SIZE, TILE_SIZE / 2, 1)]
+#else
+	[numthreads(TILE_SIZE, TILE_SIZE, 1)]
+#endif
 void MainCS(
 	uint2 GroupId : SV_GroupID,
 	uint GroupThreadIndex : SV_GroupIndex)
 {
-	uint GroupWaveIndex = GetGroupWaveIndex(GroupThreadIndex, /* GroupSize = */ TILE_SIZE * TILE_SIZE);
+	#if CONFIG_ENABLE_DUAL_PIXEL_VECTORIZATION
+		uint GroupWaveIndex = GetGroupWaveIndex(GroupThreadIndex, TILE_SIZE * TILE_SIZE / 2);
+	#else
+		uint GroupWaveIndex = GetGroupWaveIndex(GroupThreadIndex, TILE_SIZE * TILE_SIZE);
+	#endif
 
-	float4 Debug = 0.0;
+	float4x2 Debug = 0.0;
 
-	taa_short2 HistoryPixelPos = (
-		taa_short2(GroupId) * taa_short2(TILE_SIZE, TILE_SIZE) +
-		Map8x8Tile2x2Lane(GroupThreadIndex));
+	taa_short2x2 HistoryPixelPos = dpv_add(
+		taa_short2(GroupId) * taa_short2(TILE_SIZE, TILE_SIZE),
+		taa_short2x2(Map8x8Tile2x2LaneDPV(GroupThreadIndex)));
 
-	float2 ScreenPos = ApplyScreenTransform(float2(HistoryPixelPos), HistoryPixelPosToScreenPos);
+	float2x2 ScreenPos = ApplyScreenTransform(float2x2(HistoryPixelPos), HistoryPixelPosToScreenPos);
 	
 	// Pixel coordinate of the center of output pixel O in the input viewport.
-	float2 PPCo = ApplyScreenTransform(float2(HistoryPixelPos), HistoryPixelPosToPPCo);
+	float2x2 PPCo = ApplyScreenTransform(float2x2(HistoryPixelPos), HistoryPixelPosToPPCo);
 
 	// Pixel coordinate of the center of the nearest input pixel K.
-	float2 PPCk = floor(PPCo) + 0.5;
+	float2x2 PPCk = floor(PPCo) + 0.5;
 	
-	taa_short2 InputPixelPos = ClampPixelOffset(
-		taa_short2(PPCo),
+	taa_short2x2 InputPixelPos = ClampPixelOffset(
+		taa_short2x2(PPCo),
 		InputPixelPosMin, InputPixelPosMax);
 
 	// Fetch reprojections related informations.
-	float2 PrevScreenPos = ScreenPos;
-	taa_half ParallaxRejectionMask = taa_half(1.0);
-	taa_half LowFrequencyRejection = taa_half(1.0);
-	taa_half TranslucencyRejection = taa_half(1.0);
-	taa_half OutputPixelVelocity = taa_half(0.0);
-	taa_subpixel_payload ParallaxFactorBits = 0x0;
-	bool bIsResponsiveAAPixel = false;
+	float2x2 PrevScreenPos = ScreenPos;
+	taa_half2 ParallaxRejectionMask = taa_half(1.0).xx;
+	taa_half2 LowFrequencyRejection = taa_half(1.0).xx;
+	taa_half2 TranslucencyRejection = taa_half(1.0).xx;
+	taa_half2 OutputPixelVelocity = taa_half(0.0).xx;
+	taa_subpixel_payload2 ParallaxFactorBits = taa_subpixel_payload(0x0).xx;
+	bool2 bIsResponsiveAAPixel = false;
 
 	ISOLATE
 	{
 		// Vector in pixel between pixel K -> O.
-		taa_half2 dKO = taa_half2(PPCo - PPCk);
-
-		float2 EncodedVelocity = DilatedVelocityTexture[InputPixelPos];
-		ParallaxRejectionMask = ParallaxRejectionMaskTexture[InputPixelPos];
-		ParallaxFactorBits = ParallaxFactorTexture[InputPixelPos] & SUB_PIXEL_PARALLAX_FACTOR_BIT_MASK;
-
-		float2 ScreenVelocity = DecodeVelocityFromTexture(float4(EncodedVelocity, 0.0, 0.0)).xy;
+		taa_half2x2 dKO = taa_half2x2(PPCo - PPCk);
+
+		float2x2 EncodedVelocity = dpv_interleave_registers(
+			DilatedVelocityTexture[dpv_lo(InputPixelPos)],
+			DilatedVelocityTexture[dpv_hi(InputPixelPos)]);
+		ParallaxRejectionMask = dpv_interleave_registers(
+			ParallaxRejectionMaskTexture[dpv_lo(InputPixelPos)],
+			ParallaxRejectionMaskTexture[dpv_hi(InputPixelPos)]);
+		ParallaxFactorBits = dpv_interleave_registers(
+			ParallaxFactorTexture[dpv_lo(InputPixelPos)],
+			ParallaxFactorTexture[dpv_hi(InputPixelPos)]) & taa_subpixel_payload(SUB_PIXEL_PARALLAX_FACTOR_BIT_MASK);
+
+		float2x2 ScreenVelocity = dpv_interleave_registers(
+			DecodeVelocityFromTexture(float4(dpv_lo(EncodedVelocity), 0.0, 0.0)).xy,
+			DecodeVelocityFromTexture(float4(dpv_hi(EncodedVelocity), 0.0, 0.0)).xy);
 
 		PrevScreenPos = ScreenPos - ScreenVelocity;
-		OutputPixelVelocity = taa_half(length(ScreenVelocity * HistoryInfo_ViewportSize));
-
-		taa_ushort2 RejectionPixelPos = (taa_ushort2(InputPixelPos) - taa_short2(InputPixelPosMin)) / 2;
-		LowFrequencyRejection = HistoryRejectionTexture[RejectionPixelPos];
-		TranslucencyRejection = TranslucencyRejectionTexture[InputPixelPos];
+		OutputPixelVelocity = taa_half2(dpv_length(dpv_mul(ScreenVelocity, HistoryInfo_ViewportSize)));
+
+		taa_ushort2x2 RejectionPixelPos = dpv_sub(taa_ushort2x2(InputPixelPos), taa_ushort2(InputPixelPosMin)) / 2;
+		LowFrequencyRejection = dpv_interleave_registers(
+			HistoryRejectionTexture[dpv_lo(RejectionPixelPos)],
+			HistoryRejectionTexture[dpv_hi(RejectionPixelPos)]);
+		TranslucencyRejection = dpv_interleave_registers(
+			TranslucencyRejectionTexture[dpv_lo(InputPixelPos)],
+			TranslucencyRejectionTexture[dpv_hi(InputPixelPos)]);
 		
 		UNROLL_N(CONFIG_SAMPLES - 1)
 		for (uint SampleId = 1; SampleId < CONFIG_SAMPLES; SampleId++)
 		{
-			taa_short2 SampleInputPixelPos;
-			taa_half2 PixelOffset;
+			taa_short2x2 SampleInputPixelPos;
+			taa_half2x2 PixelOffset;
 			ComputeInputKernelSamplePosition(
 				InputPixelPos, dKO, SampleId,
 				/* out */ SampleInputPixelPos,
 				/* out */ PixelOffset);
 			
-			TranslucencyRejection = min(TranslucencyRejection, TranslucencyRejectionTexture[SampleInputPixelPos]);
+			TranslucencyRejection = min(TranslucencyRejection, dpv_interleave_registers(
+				TranslucencyRejectionTexture[dpv_lo(SampleInputPixelPos)],
+				TranslucencyRejectionTexture[dpv_hi(SampleInputPixelPos)]));
 		}
 
 		TranslucencyRejection = max(TranslucencyRejection, taa_half(MinTranslucencyRejection));
@@ -201,7 +247,9 @@ void MainCS(
 		{
 			const uint kResponsiveStencilMask = 1 << 3;
 			
-			uint SceneStencilRef = InputSceneStencilTexture.Load(int3(InputPixelPos, 0)) STENCIL_COMPONENT_SWIZZLE;
+			uint2 SceneStencilRef = dpv_interleave_registers(
+				InputSceneStencilTexture.Load(int3(dpv_lo(InputPixelPos), 0)) STENCIL_COMPONENT_SWIZZLE, 
+				InputSceneStencilTexture.Load(int3(dpv_lo(InputPixelPos), 0)) STENCIL_COMPONENT_SWIZZLE);
 
 			bIsResponsiveAAPixel = (SceneStencilRef & kResponsiveStencilMask) != 0;
 		}
@@ -209,9 +257,9 @@ void MainCS(
 
 		#if !CONFIG_CLAMP
 		{
-			ParallaxRejectionMask = taa_half(1.0);
-			LowFrequencyRejection = taa_half(1.0);
-			TranslucencyRejection = taa_half(1.0);
+			ParallaxRejectionMask = taa_half(1.0).xx;
+			LowFrequencyRejection = taa_half(1.0).xx;
+			TranslucencyRejection = taa_half(1.0).xx;
 		}
 		#endif
 	}
@@ -222,55 +270,71 @@ void MainCS(
 	#endif
 	
 	// Detect if HistoryBufferUV would be outside of the viewport.
-	bool bOffScreen = IsOffScreen(bCameraCut, PrevScreenPos, ParallaxRejectionMask);
+	bool2 bOffScreen = bool2(
+		IsOffScreen(bCameraCut, dpv_lo(PrevScreenPos), dpv_lo(ParallaxRejectionMask)),
+		IsOffScreen(bCameraCut, dpv_hi(PrevScreenPos), dpv_hi(ParallaxRejectionMask)));
 	
-	taa_half TotalRejection = bOffScreen ? 0.0 : saturate(LowFrequencyRejection * 4.0);
+	taa_half2 TotalRejection = taa_half2(
+		bOffScreen[0] ? 0.0 : saturate(LowFrequencyRejection[0] * 4.0),
+		bOffScreen[1] ? 0.0 : saturate(LowFrequencyRejection[1] * 4.0));
 
 
 	// Filter input scene color at predictor frequency.
-	taa_half3 FilteredInputColor;
-	taa_half3 InputMinColor;
-	taa_half3 InputMaxColor;
-	taa_half InputPixelAlignement;
-	taa_half ClosestInputLuma4;
+	taa_half3x2 FilteredInputColor;
+	taa_half3x2 InputMinColor;
+	taa_half3x2 InputMaxColor;
+	taa_half2 InputPixelAlignement;
+	taa_half2 ClosestInputLuma4;
 	
+	
+	taa_half2x2 dKO = taa_half2x2(PPCo - PPCk);
+		
+	bool2 bUpdateSubpixelDetails;
+	bool2x2 bKOSign;
+
 	ISOLATE
 	{
 		// Vector in pixel between pixel K -> O.
-		taa_half2 dKO = taa_half2(PPCo - PPCk);
+		taa_half2x2 dKO = taa_half2x2(PPCo - PPCk);
+
+		{
+			float SizeOfOutputPixelInInputPixelWidth = 0.5 * (InputInfo_ViewportSize.x * HistoryInfo_ViewportSizeInverse.x);
+
+			bUpdateSubpixelDetails[0] = all(abs(dpv_lo(dKO)) < SizeOfOutputPixelInInputPixelWidth);
+			bUpdateSubpixelDetails[1] = all(abs(dpv_hi(dKO)) < SizeOfOutputPixelInInputPixelWidth);
+		
+			bKOSign = dKO < 0.0;
+		}
+
 
 		FilteredInputColor = taa_half(0.0);
 
-		taa_half FilteredInputColorWeight = taa_half(0.0);
+		taa_half2 FilteredInputColorWeight = taa_half(0.0);
 		
 		// TODO: UECON-464 - this causes shader warnings on XSX DXC.  Revisit once we have aged out the problem GDK versions.
-		#if 0 // shader compiler bug :'(
-			taa_half InputToHistoryFactor = taa_half(HistoryInfo_ViewportSize.x * InputInfo_ViewportSizeInverse.x);
-			taa_half FinalInputToHistoryFactor = bOffScreen ? taa_half(1.0) : InputToHistoryFactor;
-		#else
-			float InputToHistoryFactor = float(HistoryInfo_ViewportSize.x * InputInfo_ViewportSizeInverse.x);
-			float FinalInputToHistoryFactor = lerp(1.0, InputToHistoryFactor, TotalRejection);
-		#endif
-
-		InputMinColor = taa_half(INFINITE_FLOAT);
-		InputMaxColor = taa_half(-INFINITE_FLOAT);
+		taa_half InputToHistoryFactor = taa_half(HistoryInfo_ViewportSize.x * InputInfo_ViewportSizeInverse.x);
+		taa_half2 FinalInputToHistoryFactor = taa_half2(
+			bOffScreen.x ? taa_half(1.0) : InputToHistoryFactor,
+			bOffScreen.y ? taa_half(1.0) : InputToHistoryFactor);
 
 		UNROLL_N(CONFIG_SAMPLES)
 		for (uint SampleId = 0; SampleId < CONFIG_SAMPLES; SampleId++)
 		{
-			taa_short2 SampleInputPixelPos;
-			taa_half2 PixelOffset;
+			taa_short2x2 SampleInputPixelPos;
+			taa_half2x2 PixelOffset;
 			ComputeInputKernelSamplePosition(
 				InputPixelPos, dKO, SampleId,
 				/* out */ SampleInputPixelPos,
 				/* out */ PixelOffset);
 
-			taa_half3 InputColor = InputSceneColorTexture[SampleInputPixelPos];
+			taa_half3x2 InputColor = dpv_interleave_registers(
+				InputSceneColorTexture[dpv_lo(SampleInputPixelPos)],
+				InputSceneColorTexture[dpv_hi(SampleInputPixelPos)]);
 
-			taa_half2 dPP = PixelOffset - dKO;
-			taa_half SampleSpatialWeight = ComputeSampleWeigth(FinalInputToHistoryFactor, dPP, /* MinimalContribution = */ float(0.005));
+			taa_half2x2 dPP = PixelOffset - dKO;
+			taa_half2 SampleSpatialWeight = ComputeSampleWeigth(FinalInputToHistoryFactor, dPP, /* MinimalContribution = */ float(0.005));
 
-			taa_half ToneWeight = HdrWeight4(InputColor);
+			taa_half2 ToneWeight = HdrWeight4(InputColor);
 
 			#if CONFIG_VFX_SURVIVAL
 			if (SampleId == 0)
@@ -279,7 +343,7 @@ void MainCS(
 			}
 			#endif
 
-			FilteredInputColor       += (SampleSpatialWeight * ToneWeight) * InputColor;
+			FilteredInputColor       += dpv_scale((SampleSpatialWeight * ToneWeight), InputColor);
 			FilteredInputColorWeight += (SampleSpatialWeight * ToneWeight);
 
 			if (SampleId == 0)
@@ -295,7 +359,7 @@ void MainCS(
 			}
 		}
 		
-		FilteredInputColor *= rcp(FilteredInputColorWeight);
+		FilteredInputColor = dpv_scale(FilteredInputColor, rcp(FilteredInputColorWeight));
 
 		InputPixelAlignement = ComputeSampleWeigth(InputToHistoryFactor, dKO, /* MinimalContribution = */ float(0.0));
 	}
@@ -306,62 +370,91 @@ void MainCS(
 	{
 		uint LocalGroupThreadIndex = GetGroupThreadIndex(GroupThreadIndex, GroupWaveIndex);
 
-		SharedArray0[LocalGroupThreadIndex] = taa_half4(FilteredInputColor, LowFrequencyRejection);
-		SharedArray1[LocalGroupThreadIndex] = taa_half4(InputMinColor, InputPixelAlignement);
-		SharedArray2[LocalGroupThreadIndex] = taa_half4(InputMaxColor, OutputPixelVelocity);
+		#if CONFIG_ENABLE_DUAL_PIXEL_VECTORIZATION
+			SharedArray0[LocalGroupThreadIndex] = taa_half4(FilteredInputColor[0], FilteredInputColor[1]);
+			SharedArray1[LocalGroupThreadIndex] = taa_half4(FilteredInputColor[2], LowFrequencyRejection);
+			SharedArray2[LocalGroupThreadIndex] = taa_half4(InputMinColor[0], InputMinColor[1]);
+			SharedArray3[LocalGroupThreadIndex] = taa_half4(InputMinColor[2], InputPixelAlignement);
+			SharedArray4[LocalGroupThreadIndex] = taa_half4(InputMaxColor[0], InputMaxColor[1]);
+			SharedArray5[LocalGroupThreadIndex] = taa_half4(InputMaxColor[2], OutputPixelVelocity);
+			SharedParallaxFactorBits[LocalGroupThreadIndex] = ParallaxFactorBits;
+		#else
+			SharedArray0[LocalGroupThreadIndex] = taa_half4(dpv_lo(FilteredInputColor), dpv_lo(LowFrequencyRejection));
+			SharedArray1[LocalGroupThreadIndex] = taa_half4(dpv_lo(InputMinColor), dpv_lo(InputPixelAlignement));
+			SharedArray2[LocalGroupThreadIndex] = taa_half4(dpv_lo(InputMaxColor), dpv_lo(OutputPixelVelocity));
+		#endif
 	}
 	#endif
 	
 	// Reproject history
-	taa_half3 PrevHistoryMoment1;
-	taa_half PrevHistoryValidity;
+	taa_half3x2 PrevHistoryMoment1;
+	taa_half2 PrevHistoryValidity;
 	
-	taa_half3 PrevHistoryMommentMin;
-	taa_half3 PrevHistoryMommentMax;
+	taa_half3x2 PrevHistoryMommentMin;
+	taa_half3x2 PrevHistoryMommentMax;
 
-	taa_half3 PrevFallbackColor;
-	taa_half PrevFallbackWeight;
+	taa_half3x2 PrevFallbackColor;
+	taa_half2 PrevFallbackWeight;
 	
-	taa_subpixel_details PrevSubpixelDetails;
+	taa_subpixel_details2 PrevSubpixelDetails;
 
 	ISOLATE
 	{
 		// Reproject history
-		taa_half3 RawHistory0 = taa_half(0);
-		taa_half3 RawHistory1 = taa_half(0);
-		taa_half2 RawHistory2 = taa_half(0);
+		taa_half3x2 RawHistory0 = taa_half(0);
+		taa_half3x2 RawHistory1 = taa_half(0);
+		taa_half2x2 RawHistory2 = taa_half(0);
 
-		taa_half3 RawHistory1Min = INFINITE_FLOAT;
-		taa_half3 RawHistory1Max = -INFINITE_FLOAT;
+		taa_half3x2 RawHistory1Min = INFINITE_FLOAT;
+		taa_half3x2 RawHistory1Max = -INFINITE_FLOAT;
 
 		// Sample raw history
 		{
-			float2 PrevHistoryBufferUV = ApplyScreenTransform(PrevScreenPos, ScreenPosToPrevHistoryBufferUV);
-			PrevHistoryBufferUV = clamp(PrevHistoryBufferUV, PrevHistoryInfo_UVViewportBilinearMin, PrevHistoryInfo_UVViewportBilinearMax);
+			float2x2 PrevHistoryBufferUV = ApplyScreenTransform(PrevScreenPos, ScreenPosToPrevHistoryBufferUV);
+			PrevHistoryBufferUV[0] = clamp(PrevHistoryBufferUV[0], PrevHistoryInfo_UVViewportBilinearMin[0], PrevHistoryInfo_UVViewportBilinearMax[0]);
+			PrevHistoryBufferUV[1] = clamp(PrevHistoryBufferUV[1], PrevHistoryInfo_UVViewportBilinearMin[1], PrevHistoryInfo_UVViewportBilinearMax[1]);
 
 			#if 1
 			{
-				FCatmullRomSamples Samples = GetBicubic2DCatmullRomSamples(PrevHistoryBufferUV, PrevHistoryInfo_Extent, PrevHistoryInfo_ExtentInverse);
+				FCatmullRomSamples Samples0 = GetBicubic2DCatmullRomSamples(
+					dpv_lo(PrevHistoryBufferUV), PrevHistoryInfo_Extent, PrevHistoryInfo_ExtentInverse);
+				FCatmullRomSamples Samples1 = GetBicubic2DCatmullRomSamples(
+					dpv_hi(PrevHistoryBufferUV), PrevHistoryInfo_Extent, PrevHistoryInfo_ExtentInverse);
 
 				UNROLL
-				for (uint i = 0; i < Samples.Count; i++)
+				for (uint i = 0; i < Samples0.Count; i++)
 				{
-					float2 SampleUV = clamp(Samples.UV[i], PrevHistoryInfo_UVViewportBilinearMin, PrevHistoryInfo_UVViewportBilinearMax);
+					float2 SampleUV0 = clamp(Samples0.UV[i], PrevHistoryInfo_UVViewportBilinearMin, PrevHistoryInfo_UVViewportBilinearMax);
+					float2 SampleUV1 = clamp(Samples1.UV[i], PrevHistoryInfo_UVViewportBilinearMin, PrevHistoryInfo_UVViewportBilinearMax);
 
-					taa_half3 Sample0 = PrevHistory_Textures_0.SampleLevel(GlobalBilinearClampedSampler, SampleUV, 0);
-					taa_half3 Sample1 = PrevHistory_Textures_1.SampleLevel(GlobalBilinearClampedSampler, SampleUV, 0);
-					taa_half2 Sample2 = PrevHistory_Textures_2.SampleLevel(GlobalBilinearClampedSampler, SampleUV, 0);
+					taa_half3x2 Sample0 = dpv_interleave_registers(
+						PrevHistory_Textures_0.SampleLevel(GlobalBilinearClampedSampler, SampleUV0, 0),
+						PrevHistory_Textures_0.SampleLevel(GlobalBilinearClampedSampler, SampleUV1, 0));
+					taa_half3x2 Sample1 = dpv_interleave_registers(
+						PrevHistory_Textures_1.SampleLevel(GlobalBilinearClampedSampler, SampleUV0, 0),
+						PrevHistory_Textures_1.SampleLevel(GlobalBilinearClampedSampler, SampleUV1, 0));
+					taa_half2x2 Sample2 = dpv_interleave_registers(
+						PrevHistory_Textures_2.SampleLevel(GlobalBilinearClampedSampler, SampleUV0, 0),
+						PrevHistory_Textures_2.SampleLevel(GlobalBilinearClampedSampler, SampleUV1, 0));
 
-					RawHistory1Min = min(RawHistory1Min, Sample1 * SafeRcp(Sample2.g));
-					RawHistory1Max = max(RawHistory1Max, Sample1 * SafeRcp(Sample2.g));
+					taa_half2 InvWeight = SafeRcp(Sample2[1]);
 
-					RawHistory0 += Sample0 * taa_half(Samples.Weight[i]);
-					RawHistory1 += Sample1 * taa_half(Samples.Weight[i]);
-					RawHistory2 += Sample2 * taa_half(Samples.Weight[i]);
+					RawHistory1Min = min(RawHistory1Min, dpv_scale(Sample1, InvWeight));
+					RawHistory1Max = max(RawHistory1Max, dpv_scale(Sample1, InvWeight));
+
+					taa_half2 KernelWeight = dpv_interleave_registers(taa_half(Samples0.Weight[i]), taa_half(Samples1.Weight[i]));
+
+					RawHistory0 += dpv_scale(Sample0, KernelWeight);
+					RawHistory1 += dpv_scale(Sample1, KernelWeight);
+					RawHistory2 += dpv_scale(Sample2, KernelWeight);
 				}
-				RawHistory0 *= taa_half(Samples.FinalMultiplier);
-				RawHistory1 *= taa_half(Samples.FinalMultiplier);
-				RawHistory2 *= taa_half(Samples.FinalMultiplier);
+				
+				taa_half2 FinalMultiplier = dpv_interleave_registers(taa_half(Samples0.FinalMultiplier), taa_half(Samples1.FinalMultiplier));
+
+
+				RawHistory0 = dpv_scale(RawHistory0, FinalMultiplier);
+				RawHistory1 = dpv_scale(RawHistory1, FinalMultiplier);
+				RawHistory2 = dpv_scale(RawHistory2, FinalMultiplier);
 			}
 			#else
 			{
@@ -371,32 +464,35 @@ void MainCS(
 			}
 			#endif
 			
-			FSubpixelNeighborhood SubpixelNeighborhood = GatherPrevSubpixelNeighborhood(PrevHistory_Textures_3, PrevHistoryBufferUV);
+			FSubpixelNeighborhood SubpixelNeighborhood0 = GatherPrevSubpixelNeighborhood(PrevHistory_Textures_3, dpv_lo(PrevHistoryBufferUV));
+			FSubpixelNeighborhood SubpixelNeighborhood1 = GatherPrevSubpixelNeighborhood(PrevHistory_Textures_3, dpv_hi(PrevHistoryBufferUV));
 			{
 				PrevSubpixelDetails = 0;
 				UNROLL_N(SUB_PIXEL_COUNT)
 				for (uint SubpixelId = 0; SubpixelId < SUB_PIXEL_COUNT; SubpixelId++)
 				{
-					taa_subpixel_payload SubpixelPayload = GetSubpixelPayload(SubpixelNeighborhood, SubpixelId);
-					PrevSubpixelDetails |= SubpixelPayload << (SUB_PIXEL_BIT_COUNT * SubpixelId);
+					taa_subpixel_payload SubpixelPayload0 = GetSubpixelPayload(SubpixelNeighborhood0, SubpixelId);
+					taa_subpixel_payload SubpixelPayload1 = GetSubpixelPayload(SubpixelNeighborhood1, SubpixelId);
+					PrevSubpixelDetails[0] |= SubpixelPayload0 << taa_subpixel_payload(SUB_PIXEL_BIT_COUNT * SubpixelId);
+					PrevSubpixelDetails[1] |= SubpixelPayload1 << taa_subpixel_payload(SUB_PIXEL_BIT_COUNT * SubpixelId);
 				}
 			}
 
-			RawHistory0 = -min(-RawHistory0, taa_half(0.0));
-			RawHistory1 = -min(-RawHistory1, taa_half(0.0));
-			RawHistory2 = -min(-RawHistory2, taa_half(0.0));
+			RawHistory0 = -dpv_min(-RawHistory0, taa_half(0.0).xxx);
+			RawHistory1 = -dpv_min(-RawHistory1, taa_half(0.0).xxx);
+			RawHistory2 = -dpv_min(-RawHistory2, taa_half(0.0).xx);
 		}
 		
 		// Unpack history.
 		{
 			PrevFallbackColor = RawHistory0;
-			PrevFallbackWeight = RawHistory2.r;
+			PrevFallbackWeight = RawHistory2[0];
 			
 			PrevHistoryMommentMin = RawHistory1Min;
 			PrevHistoryMommentMax = RawHistory1Max;
 
 			PrevHistoryMoment1 = RawHistory1;
-			PrevHistoryValidity = RawHistory2.g;
+			PrevHistoryValidity = RawHistory2[1];
 		}
 
 		// Correct history
@@ -412,33 +508,52 @@ void MainCS(
 	ISOLATE
 	{
 		uint LocalGroupThreadIndex = GetGroupThreadIndex(GroupThreadIndex, GroupWaveIndex);
+	
+		#if CONFIG_ENABLE_DUAL_PIXEL_VECTORIZATION
+			taa_half4 RawLDS0 = SharedArray0[LocalGroupThreadIndex];
+			taa_half4 RawLDS1 = SharedArray1[LocalGroupThreadIndex];
+			taa_half4 RawLDS2 = SharedArray2[LocalGroupThreadIndex];
+			taa_half4 RawLDS3 = SharedArray3[LocalGroupThreadIndex];
+			taa_half4 RawLDS4 = SharedArray4[LocalGroupThreadIndex];
+			taa_half4 RawLDS5 = SharedArray5[LocalGroupThreadIndex];
+
+			ParallaxFactorBits = SharedParallaxFactorBits[LocalGroupThreadIndex];
+
+			FilteredInputColor = taa_half3x2(RawLDS0.rg, RawLDS0.ba, RawLDS1.rg);
+			InputMinColor = taa_half3x2(RawLDS2.rg, RawLDS2.ba, RawLDS3.rg);
+			InputMaxColor = taa_half3x2(RawLDS4.rg, RawLDS4.ba, RawLDS5.rg);
+		
+			LowFrequencyRejection = RawLDS1.ba;
+			InputPixelAlignement = RawLDS3.ba;
+			OutputPixelVelocity = RawLDS5.ba;
+		#else
+			taa_half4 RawLDS0 = SharedArray0[LocalGroupThreadIndex];
+			taa_half4 RawLDS1 = SharedArray1[LocalGroupThreadIndex];
+			taa_half4 RawLDS2 = SharedArray2[LocalGroupThreadIndex];
 
-		taa_half4 RawLDS0 = SharedArray0[LocalGroupThreadIndex];
-		taa_half4 RawLDS1 = SharedArray1[LocalGroupThreadIndex];
-		taa_half4 RawLDS2 = SharedArray2[LocalGroupThreadIndex];
-
-		FilteredInputColor = RawLDS0.rgb;
-		InputMinColor = RawLDS1.rgb;
-		InputMaxColor = RawLDS2.rgb;
+			FilteredInputColor = dpv_interleave_mono_registers(RawLDS0.rgb);
+			InputMinColor = dpv_interleave_mono_registers(RawLDS1.rgb);
+			InputMaxColor = dpv_interleave_mono_registers(RawLDS2.rgb);
 		
-		LowFrequencyRejection = RawLDS0.a;
-		InputPixelAlignement = RawLDS1.a;
-		OutputPixelVelocity = RawLDS2.a;
+			LowFrequencyRejection = dpv_interleave_mono_registers(RawLDS0.a);
+			InputPixelAlignement = dpv_interleave_mono_registers(RawLDS1.a);
+			OutputPixelVelocity = dpv_interleave_mono_registers(RawLDS2.a);
+		#endif
 	}
 	#endif
 
 	// Rejects high frequency details if their low frequencies drift from the low frequency history
 	#if CONFIG_LOW_FREQUENCY_DRIFT_REJECTION
 	{
-		taa_half3 PrevHighFrequencyYCoCg = TransformColorForClampingBox(PrevHistoryMoment1 * SafeRcp(PrevHistoryValidity));
-		taa_half3 PrevYCoCg = TransformColorForClampingBox(PrevFallbackColor);
-		taa_half3 ClampedPrevYCoCg = TransformColorForClampingBox(clamp(PrevFallbackColor, PrevHistoryMommentMin, PrevHistoryMommentMax));
+		taa_half3x2 PrevHighFrequencyYCoCg = TransformColorForClampingBox(dpv_scale(PrevHistoryMoment1, SafeRcp(PrevHistoryValidity)));
+		taa_half3x2 PrevYCoCg = TransformColorForClampingBox(PrevFallbackColor);
+		taa_half3x2 ClampedPrevYCoCg = TransformColorForClampingBox(clamp(PrevFallbackColor, PrevHistoryMommentMin, PrevHistoryMommentMax));
 
-		taa_half HighFrequencyRejection = MeasureRejectionFactor(
+		taa_half2 HighFrequencyRejection = MeasureRejectionFactor(
 			PrevYCoCg, ClampedPrevYCoCg,
 			PrevHighFrequencyYCoCg, InputMinColor, InputMaxColor);
 		
-		PrevHistoryMoment1 *= HighFrequencyRejection;
+		PrevHistoryMoment1 = dpv_scale(PrevHistoryMoment1, HighFrequencyRejection);
 		PrevHistoryValidity *= HighFrequencyRejection;
 	}
 	#endif
@@ -448,17 +563,17 @@ void MainCS(
 	const taa_half PredictionOnlyValidity = Histeresis * taa_half(2.0);
 	
 	// Clamp fallback
-	taa_half LumaMin;
-	taa_half LumaMax;
-	taa_half3 ClampedFallbackColor;
-	taa_half FallbackRejection;
+	taa_half2 LumaMin;
+	taa_half2 LumaMax;
+	taa_half3x2 ClampedFallbackColor;
+	taa_half2 FallbackRejection;
 	{
-		LumaMin = InputMinColor.x;
-		LumaMax = InputMaxColor.x;
+		LumaMin = InputMinColor[0];
+		LumaMax = InputMaxColor[0];
 
-		taa_half3 PrevYCoCg = TransformColorForClampingBox(PrevFallbackColor);
-		taa_half3 ClampedPrevYCoCg = clamp(PrevYCoCg, InputMinColor, InputMaxColor);
-		taa_half3 InputCenterYCoCg = TransformColorForClampingBox(FilteredInputColor);
+		taa_half3x2 PrevYCoCg = TransformColorForClampingBox(PrevFallbackColor);
+		taa_half3x2 ClampedPrevYCoCg = dpv_clamp(PrevYCoCg, InputMinColor, InputMaxColor);
+		taa_half3x2 InputCenterYCoCg = TransformColorForClampingBox(FilteredInputColor);
 
 		ClampedFallbackColor = YCoCgToRGB(ClampedPrevYCoCg);
 		
@@ -481,63 +596,75 @@ void MainCS(
 	}
 	#endif
 	
-	taa_half3 FinalHistoryMoment1;
-	taa_half FinalHistoryValidity;
+	taa_half3x2 FinalHistoryMoment1;
+	taa_half2 FinalHistoryValidity;
 	{
 		// Compute how much the history needs to be rejected, based on the completeness.
-		taa_half PrevHistoryRejectionWeight = LowFrequencyRejection;
+		taa_half2 PrevHistoryRejectionWeight = LowFrequencyRejection;
 			
 		//#if !CONFIG_CLAMP
 		//	PrevHistoryRejectionWeight = 1.0;
 		//#endif
 			
 		FLATTEN
-		if (bOffScreen)
+		if (bOffScreen[0])
 		{
-			PrevHistoryRejectionWeight = taa_half(0.0);
+			PrevHistoryRejectionWeight[0] = taa_half(0.0);
 		}
 
-		taa_half DesiredCurrentContribution = max(Histeresis * InputPixelAlignement, taa_half(0.0));
+		FLATTEN
+		if (bOffScreen[1])
+		{
+			PrevHistoryRejectionWeight[1] = taa_half(0.0);
+		}
+
+		taa_half2 DesiredCurrentContribution = max(Histeresis * InputPixelAlignement, taa_half(0.0));
 
 		// Determine whether the the prediction based rejection was confident enough.
-		taa_half RejectionConfidentEnough = taa_half(1); // saturate(RejectionValidity * MAX_SAMPLE_COUNT - 3.0);
+		taa_half2 RejectionConfidentEnough = taa_half(1); // saturate(RejectionValidity * MAX_SAMPLE_COUNT - 3.0);
 
 		// Compute the newly rejected validity
-		taa_half RejectedValidity = (
+		taa_half2 RejectedValidity = (
 			min(PrevHistoryValidity, PredictionOnlyValidity - DesiredCurrentContribution) +
 			max(PrevHistoryValidity - PredictionOnlyValidity + DesiredCurrentContribution, taa_half(0.0)) * PrevHistoryRejectionWeight);
 
 		RejectedValidity = PrevHistoryValidity * PrevHistoryRejectionWeight;
 
 		// Compute the maximum output validity.
-		taa_half OutputValidity = (
+		taa_half2 OutputValidity = (
 			clamp(RejectedValidity + DesiredCurrentContribution, taa_half(0.0), PredictionOnlyValidity) +
 			clamp(RejectedValidity + DesiredCurrentContribution * PrevHistoryRejectionWeight * RejectionConfidentEnough - PredictionOnlyValidity, 0.0, 1.0 - PredictionOnlyValidity));
 
 		FLATTEN
-		if (bIsResponsiveAAPixel)
+		if (bIsResponsiveAAPixel[0])
+		{
+			OutputValidity[0] = taa_half(0.0);
+		}
+
+		FLATTEN
+		if (bIsResponsiveAAPixel[1])
 		{
-			OutputValidity = taa_half(0.0);
+			OutputValidity[1] = taa_half(0.0);
 		}
 		
-		taa_half InvPrevHistoryValidity = SafeRcp(PrevHistoryValidity);
+		taa_half2 InvPrevHistoryValidity = SafeRcp(PrevHistoryValidity);
 
-		taa_half PrevMomentWeight = max(OutputValidity - DesiredCurrentContribution, taa_half(0.0));
-		taa_half CurrentMomentWeight = min(DesiredCurrentContribution, OutputValidity);
+		taa_half2 PrevMomentWeight = max(OutputValidity - DesiredCurrentContribution, taa_half(0.0));
+		taa_half2 CurrentMomentWeight = min(DesiredCurrentContribution, OutputValidity);
 		
 		#if 1
 		{
-			taa_half PrevHistoryToneWeight = HdrWeightY(Luma4(PrevHistoryMoment1) * InvPrevHistoryValidity);
-			taa_half FilteredInputToneWeight = HdrWeight4(FilteredInputColor);
+			taa_half2 PrevHistoryToneWeight = HdrWeightY(Luma4(PrevHistoryMoment1) * InvPrevHistoryValidity);
+			taa_half2 FilteredInputToneWeight = HdrWeight4(FilteredInputColor);
 			
-			taa_half BlendPrevHistory = PrevMomentWeight * PrevHistoryToneWeight;
-			taa_half BlendFilteredInput = CurrentMomentWeight * FilteredInputToneWeight;
+			taa_half2 BlendPrevHistory = PrevMomentWeight * PrevHistoryToneWeight;
+			taa_half2 BlendFilteredInput = CurrentMomentWeight * FilteredInputToneWeight;
 
-			taa_half CommonWeight = OutputValidity * SafeRcp(BlendPrevHistory + BlendFilteredInput);
+			taa_half2 CommonWeight = OutputValidity * SafeRcp(BlendPrevHistory + BlendFilteredInput);
 
 			FinalHistoryMoment1 = (
-				PrevHistoryMoment1 * (CommonWeight * BlendPrevHistory * InvPrevHistoryValidity) +
-				FilteredInputColor * (CommonWeight * BlendFilteredInput));
+				dpv_scale(PrevHistoryMoment1, CommonWeight * BlendPrevHistory * InvPrevHistoryValidity) +
+				dpv_scale(FilteredInputColor, CommonWeight * BlendFilteredInput));
 		}
 		#else
 		{
@@ -549,9 +676,9 @@ void MainCS(
 		#if 0
 			FinalHistoryValidity = OutputValidity;
 		#else
-			taa_half OutputInvValidity = SafeRcp(OutputValidity);
+			taa_half2 OutputInvValidity = SafeRcp(OutputValidity);
 			FinalHistoryValidity = ceil(taa_half(255.0) * OutputValidity) * rcp(taa_half(255.0));
-			FinalHistoryMoment1 *= FinalHistoryValidity * OutputInvValidity;
+			FinalHistoryMoment1 = dpv_scale(FinalHistoryMoment1, FinalHistoryValidity * OutputInvValidity);
 		#endif
 
 		//Debug.x = RejectionConfidentEnough;
@@ -566,36 +693,46 @@ void MainCS(
 	}
 
 	// Compute the fallback history
-	taa_half3 FinalFallbackColor;
-	taa_half FinalFallbackWeight;
+	taa_half3x2 FinalFallbackColor;
+	taa_half2 FinalFallbackWeight;
 	{
 		const taa_half TargetHesteresisCurrentFrameWeight = rcp(taa_half(MAX_FALLBACK_SAMPLE_COUNT));
 
-		taa_half LumaHistory = Luma4(PrevFallbackColor);
-		taa_half LumaFiltered = Luma4(FilteredInputColor);
+		taa_half2 LumaHistory = Luma4(PrevFallbackColor);
+		taa_half2 LumaFiltered = Luma4(FilteredInputColor);
 
 		{
-			taa_half OutputBlend = ComputeFallbackContribution(FinalHistoryValidity);
+			taa_half2 OutputBlend = ComputeFallbackContribution(FinalHistoryValidity);
 
 			//ClampedFallbackColor = lerp(ClampedFallbackColor, PrevFallbackColor, OutputBlend);
 		}
 
-		taa_half BlendFinal;
+		taa_half2 BlendFinal;
 		#if 1
 		{
-			taa_half CurrentFrameSampleCount = max(InputPixelAlignement, taa_half(0.005));
+			taa_half2 CurrentFrameSampleCount = max(InputPixelAlignement, taa_half(0.005));
 			
 			// Use only one sample count to be extremely fast to recover in history rejection. But then stabilize imediately so sub pixel
 			// frequency can be used as soon as possible.
-			taa_half PrevFallbackSampleCount;
+			taa_half2 PrevFallbackSampleCount;
 			FLATTEN
-			if (PrevFallbackWeight < taa_half(1.0))
+			if (PrevFallbackWeight[0] < taa_half(1.0))
 			{
-				PrevFallbackSampleCount = PrevFallbackWeight;
+				PrevFallbackSampleCount[0] = PrevFallbackWeight[0];
 			}
 			else
 			{
-				PrevFallbackSampleCount = taa_half(MAX_FALLBACK_SAMPLE_COUNT);
+				PrevFallbackSampleCount[0] = taa_half(MAX_FALLBACK_SAMPLE_COUNT);
+			}
+
+			FLATTEN
+			if (PrevFallbackWeight[1] < taa_half(1.0))
+			{
+				PrevFallbackSampleCount[1] = PrevFallbackWeight[1];
+			}
+			else
+			{
+				PrevFallbackSampleCount[1] = taa_half(MAX_FALLBACK_SAMPLE_COUNT);
 			}
 
 			// Reject the history based on the low frequencies.
@@ -604,7 +741,7 @@ void MainCS(
 				// float InvertedBlendFinal = PrevFallbackSampleCount / (CurrentFrameSampleCount + PrevFallbackSampleCount);
 				// float PrevFallbackRejectionFactor = saturate(LowFrequencyRejection / InvertedBlendFinal);
 				
-				taa_half PrevFallbackRejectionFactor = saturate(LowFrequencyRejection * (CurrentFrameSampleCount + PrevFallbackSampleCount) / PrevFallbackSampleCount);
+				taa_half2 PrevFallbackRejectionFactor = saturate(LowFrequencyRejection * (CurrentFrameSampleCount + PrevFallbackSampleCount) / PrevFallbackSampleCount);
 
 				PrevFallbackSampleCount *= PrevFallbackRejectionFactor;
 			}
@@ -622,7 +759,7 @@ void MainCS(
 			// Antiflicker
 			#if 1
 			{
-				taa_half DistToClamp = min( abs(LumaHistory - LumaMin), abs(LumaHistory - LumaMax) ) / max3( LumaHistory, LumaFiltered, taa_half(1e-4) );
+				taa_half2 DistToClamp = min( abs(LumaHistory - LumaMin), abs(LumaHistory - LumaMax) ) / max3( LumaHistory, LumaFiltered, taa_half(1e-4) );
 				BlendFinal *= taa_half(0.2) + taa_half(0.8) * saturate(taa_half(0.5) * DistToClamp);
 			}
 			#endif
@@ -635,7 +772,8 @@ void MainCS(
 			#endif
 
 			// Responsive forces 1/4 of new frame.
-			BlendFinal = bIsResponsiveAAPixel ? taa_half(1.0/4.0) : BlendFinal;
+			BlendFinal[0] = bIsResponsiveAAPixel[0] ? taa_half(1.0/4.0) : BlendFinal[0];
+			BlendFinal[1] = bIsResponsiveAAPixel[1] ? taa_half(1.0/4.0) : BlendFinal[1];
 
 			// Completely discard the history
 			{
@@ -652,34 +790,23 @@ void MainCS(
 		#endif
 
 		{
-			taa_half FilterWeight = HdrWeight4(FilteredInputColor);
-			taa_half ClampedHistoryWeight = HdrWeight4(ClampedFallbackColor);
-
-			#if 0 // DEBUG
-				FilterWeight = 1.0;
-				ClampedHistoryWeight = 1.0;
-			#endif
-
-			taa_half2 Weights = WeightedLerpFactors(ClampedHistoryWeight, FilterWeight, BlendFinal);
+			taa_half2 FilterWeight = HdrWeight4(FilteredInputColor);
+			taa_half2 ClampedHistoryWeight = HdrWeight4(ClampedFallbackColor);
 
-			FinalFallbackColor = ClampedFallbackColor * Weights.x + FilteredInputColor * Weights.y;
+			taa_half2x2 Weights = WeightedLerpFactors(ClampedHistoryWeight, FilterWeight, BlendFinal);
 
-			//Debug.x = BlendFinal;
+			FinalFallbackColor = dpv_scale(ClampedFallbackColor, Weights[0]) + dpv_scale(FilteredInputColor, Weights[1]);
 		}
 	}
 
 	// Update subpixel detail
-	taa_subpixel_details FinalSubpixelDetails;
+	taa_subpixel_details2 FinalSubpixelDetails;
 	{
-		taa_half2 dKO = taa_half2(PPCo - PPCk);
-
-		bool bUpdate = all(abs(dKO) < 0.5 * (InputInfo_ViewportSize.x * HistoryInfo_ViewportSizeInverse.x));
-
 		FinalSubpixelDetails = PrevSubpixelDetails;
 
 
 		{
-			const uint ParallaxFactorMask = (
+			const taa_subpixel_details ParallaxFactorMask = taa_subpixel_details(
 				(SUB_PIXEL_PARALLAX_FACTOR_BIT_MASK << (SUB_PIXEL_PARALLAX_FACTOR_BIT_OFFSET + 0 * SUB_PIXEL_BIT_COUNT)) | 
 				(SUB_PIXEL_PARALLAX_FACTOR_BIT_MASK << (SUB_PIXEL_PARALLAX_FACTOR_BIT_OFFSET + 1 * SUB_PIXEL_BIT_COUNT)) | 
 				(SUB_PIXEL_PARALLAX_FACTOR_BIT_MASK << (SUB_PIXEL_PARALLAX_FACTOR_BIT_OFFSET + 2 * SUB_PIXEL_BIT_COUNT)) | 
@@ -688,31 +815,44 @@ void MainCS(
 			
 			// Reset parallax factor
 			FLATTEN
-			if (bOffScreen)
+			if (bOffScreen.x)
+			{
+				FinalSubpixelDetails.x = FinalSubpixelDetails.x & ~ParallaxFactorMask;
+			}
+			
+			FLATTEN
+			if (bOffScreen.y)
 			{
-				FinalSubpixelDetails = FinalSubpixelDetails & ~ParallaxFactorMask;
+				FinalSubpixelDetails.y = FinalSubpixelDetails.y & ~ParallaxFactorMask;
 			}
 		}
+		
+		uint2 SubpixelId = dpv_dot(uint2x2(dKO), uint2(1, SUB_PIXEL_GRID_SIZE));
 
-		FLATTEN
-		if (bUpdate)
-		{
-			bool2 bBool = dKO < 0.0;
+		taa_subpixel_payload2 SubpixelShift = taa_subpixel_payload2(SubpixelId) * taa_subpixel_payload(SUB_PIXEL_BIT_COUNT);
 
-			uint SubpixelId = dot(uint2(bBool), uint2(1, SUB_PIXEL_GRID_SIZE));
-			uint SubpixelShift = SubpixelId * SUB_PIXEL_BIT_COUNT;
+		taa_subpixel_payload2 SubpixelPayload = (ParallaxFactorBits << taa_subpixel_payload(SUB_PIXEL_PARALLAX_FACTOR_BIT_OFFSET));
 
-			taa_subpixel_payload SubpixelPayload = (ParallaxFactorBits << SUB_PIXEL_PARALLAX_FACTOR_BIT_OFFSET);
+		taa_subpixel_details2 UpdatedFinalSubpixelDetails = (FinalSubpixelDetails & (~(taa_subpixel_payload(SUB_PIXEL_BIT_MASK) << SubpixelShift))) | (SubpixelPayload << SubpixelShift);
 
-			FinalSubpixelDetails = (FinalSubpixelDetails & (~(SUB_PIXEL_BIT_MASK << SubpixelShift))) | (SubpixelPayload << SubpixelShift);
+		FLATTEN
+		if (bUpdateSubpixelDetails[0])
+		{
+			FinalSubpixelDetails.x = UpdatedFinalSubpixelDetails.x;
+		}
+		
+		FLATTEN
+		if (bUpdateSubpixelDetails[1])
+		{
+			FinalSubpixelDetails.y = UpdatedFinalSubpixelDetails.y;
 		}
 	}
 
 	// Compute final output
-	taa_half3 FinalOutputColor;
-	taa_half FinalOutputValidity;
+	taa_half3x2 FinalOutputColor;
+	taa_half2 FinalOutputValidity;
 	{
-		taa_half OutputBlend = ComputeFallbackContribution(FinalHistoryValidity);
+		taa_half2 OutputBlend = ComputeFallbackContribution(FinalHistoryValidity);
 
 		#if 0 // DEBUG
 		{
@@ -733,16 +873,12 @@ void MainCS(
 
 		FinalOutputValidity = lerp(taa_half(1.0), saturate(FinalHistoryValidity), OutputBlend);
 
-		#if 1// Works arround shader compiler bug
-			taa_half3 NormalizedFinalHistoryMoment1 = taa_half3(FinalHistoryMoment1 * float(SafeRcp(FinalHistoryValidity)));
-		#else
-			taa_half3 NormalizedFinalHistoryMoment1 = FinalHistoryMoment1 * SafeRcp(FinalHistoryValidity);
-		#endif
+		taa_half3x2 NormalizedFinalHistoryMoment1 = dpv_scale(FinalHistoryMoment1, SafeRcp(FinalHistoryValidity));
 
-		taa_half FallbackWeight = HdrWeight4(FinalFallbackColor);
-		taa_half Moment1Weight = HdrWeight4(NormalizedFinalHistoryMoment1);
+		taa_half2 FallbackWeight = HdrWeight4(FinalFallbackColor);
+		taa_half2 Moment1Weight = HdrWeight4(NormalizedFinalHistoryMoment1);
 
-		taa_half2 Weights = WeightedLerpFactors(FallbackWeight, Moment1Weight, OutputBlend);
+		taa_half2x2 Weights = WeightedLerpFactors(FallbackWeight, Moment1Weight, OutputBlend);
 
 		#if DEBUG_FALLBACK_BLENDING
 			taa_half3 FallbackColor = taa_half3(1, 0.25, 0.25);
@@ -754,38 +890,21 @@ void MainCS(
 			
 			FinalOutputColor = FinalFallbackColor * Weights.x * DebugColor + NormalizedFinalHistoryMoment1 * Weights.y * DebugColor;
 		#else
-			FinalOutputColor = FinalFallbackColor * Weights.x + NormalizedFinalHistoryMoment1 * Weights.y;
+			FinalOutputColor = dpv_scale(FinalFallbackColor, Weights[0]) + dpv_scale(NormalizedFinalHistoryMoment1, Weights[1]);
 		#endif
-
-		#if 0
-		{
-			FinalOutputColor = FinalFallbackColor;
-			FinalOutputValidity = taa_half(1.0);
-		}
-		#elif 0
-		{
-			FinalOutputColor = NormalizedFinalHistoryMoment1;
-			FinalOutputValidity = FinalHistoryValidity;
-		}
-		#endif
-
-		//Debug.y = OutputBlend;
 	}
 
-	//Debug.x = abs(PrevHistoryRejectionWeight);
-	//Debug.x = abs(Rejection);
-
 
 	ISOLATE
 	{
 		uint LocalGroupThreadIndex = GetGroupThreadIndex(GroupThreadIndex, GroupWaveIndex);
 
 		#if 1
-			taa_short2 LocalHistoryPixelPos = (
-				taa_short2(GroupId) * taa_short2(TILE_SIZE, TILE_SIZE) +
-				Map8x8Tile2x2Lane(LocalGroupThreadIndex));
+			taa_short2x2 LocalHistoryPixelPos = dpv_add(
+				taa_short2(GroupId) * taa_short2(TILE_SIZE, TILE_SIZE),
+				taa_short2x2(Map8x8Tile2x2LaneDPV(LocalGroupThreadIndex)));
 		#else
-			taa_short2 LocalHistoryPixelPos = HistoryPixelPos;
+			taa_short2x2 LocalHistoryPixelPos = HistoryPixelPos;
 		#endif
 			
 		LocalHistoryPixelPos = InvalidateOutputPixelPos(LocalHistoryPixelPos, HistoryInfo_ViewportMax);
@@ -796,11 +915,11 @@ void MainCS(
 			// loss due to pixel format encoding.
 			#if CONFIG_ENABLE_STOCASTIC_QUANTIZATION
 			{
-				uint2 Random = Rand3DPCG16(int3(LocalHistoryPixelPos, View.StateFrameIndexMod8)).xy;
-				float2 E = Hammersley16(0, 1, Random);
+				uint2 Random = Rand3DPCG16(int3(dpv_lo(LocalHistoryPixelPos), View.StateFrameIndexMod8)).xy;
+				float E = Hammersley16(0, 1, Random).x;
 
-				FinalHistoryMoment1 = QuantizeForFloatRenderTarget(FinalHistoryMoment1, E.x, HistoryQuantizationError);
-				FinalFallbackColor = QuantizeForFloatRenderTarget(FinalFallbackColor, E.x, HistoryQuantizationError);
+				FinalHistoryMoment1 = QuantizeForFloatRenderTarget(FinalHistoryMoment1, E, HistoryQuantizationError);
+				FinalFallbackColor = QuantizeForFloatRenderTarget(FinalFallbackColor, E, HistoryQuantizationError);
 			}
 			#endif
 
@@ -816,31 +935,46 @@ void MainCS(
 			}
 			#endif
 			
-			FinalFallbackColor = -min(-FinalFallbackColor, taa_half(0.0));
-			FinalHistoryMoment1 = -min(-FinalHistoryMoment1, taa_half(0.0));
-			FinalFallbackColor = min(FinalFallbackColor, taa_half(Max10BitsFloat));
-			FinalHistoryMoment1 = min(FinalHistoryMoment1, taa_half(Max10BitsFloat));
+			FinalFallbackColor = -dpv_min(-FinalFallbackColor, taa_half(0.0).xxx);
+			FinalHistoryMoment1 = -dpv_min(-FinalHistoryMoment1, taa_half(0.0).xxx);
+			FinalFallbackColor = dpv_min(FinalFallbackColor, taa_half(Max10BitsFloat).xxx);
+			FinalHistoryMoment1 = dpv_min(FinalHistoryMoment1, taa_half(Max10BitsFloat).xxx);
 			
-			HistoryOutput_Textures_0[LocalHistoryPixelPos] = FinalFallbackColor;
-			HistoryOutput_Textures_1[LocalHistoryPixelPos] = FinalHistoryMoment1;
-			HistoryOutput_Textures_2[LocalHistoryPixelPos] = taa_half2(FinalFallbackWeight, FinalHistoryValidity);
-			HistoryOutput_Textures_3[LocalHistoryPixelPos] = FinalSubpixelDetails;
-
-			#if DEBUG_OUTPUT
-			{
-				DebugOutput[LocalHistoryPixelPos] = Debug;
-			}
+			HistoryOutput_Textures_0[dpv_lo(LocalHistoryPixelPos)] = dpv_lo(FinalFallbackColor);
+			HistoryOutput_Textures_1[dpv_lo(LocalHistoryPixelPos)] = dpv_lo(FinalHistoryMoment1);
+			HistoryOutput_Textures_2[dpv_lo(LocalHistoryPixelPos)] = taa_half2(dpv_lo(FinalFallbackWeight), dpv_lo(FinalHistoryValidity));
+			HistoryOutput_Textures_3[dpv_lo(LocalHistoryPixelPos)] = dpv_lo(FinalSubpixelDetails);
+			
+			#if CONFIG_ENABLE_DUAL_PIXEL_VECTORIZATION
+				HistoryOutput_Textures_0[dpv_hi(LocalHistoryPixelPos)] = dpv_hi(FinalFallbackColor);
+				HistoryOutput_Textures_1[dpv_hi(LocalHistoryPixelPos)] = dpv_hi(FinalHistoryMoment1);
+				HistoryOutput_Textures_2[dpv_hi(LocalHistoryPixelPos)] = taa_half2(dpv_hi(FinalFallbackWeight), dpv_hi(FinalHistoryValidity));
+				HistoryOutput_Textures_3[dpv_hi(LocalHistoryPixelPos)] = dpv_hi(FinalSubpixelDetails);
 			#endif
 		}
 
 		// Output final scene color
 		{
-			taa_half3 OutputColor = FinalOutputColor;
-				
-			OutputColor = -min(-OutputColor, taa_half(0.0));
-			OutputColor = min(OutputColor, taa_half(Max10BitsFloat));
+			taa_half3x2 OutputColor = FinalOutputColor;
+
+			OutputColor = -dpv_min(-OutputColor, taa_half(0.0).xxx);
+			OutputColor = dpv_min(OutputColor, taa_half(Max10BitsFloat).xxx);
 
-			SceneColorOutput[LocalHistoryPixelPos] = OutputColor;
+			SceneColorOutput[dpv_lo(LocalHistoryPixelPos)] = dpv_lo(OutputColor);
+			
+			#if CONFIG_ENABLE_DUAL_PIXEL_VECTORIZATION
+				SceneColorOutput[dpv_hi(LocalHistoryPixelPos)] = dpv_hi(OutputColor);
+			#endif
 		}
+
+		#if DEBUG_OUTPUT
+		{
+			DebugOutput[dpv_lo(LocalHistoryPixelPos)] = dpv_lo(Debug);
+
+			#if CONFIG_ENABLE_DUAL_PIXEL_VECTORIZATION
+				DebugOutput[dpv_hi(LocalHistoryPixelPos)] = dpv_hi(Debug);
+			#endif
+		}
+		#endif
 	}
 }
diff --git a/Engine/Shaders/Public/DualPixelVectorization.ush b/Engine/Shaders/Public/DualPixelVectorization.ush
new file mode 100644
--- /dev/null
+++ b/Engine/Shaders/Public/DualPixelVectorization.ush
@@ -0,0 +1,349 @@
+// Copyright Epic Games, Inc. All Rights Reserved.
+
+/*=============================================================================
+	DoublePixelVectorization.ush: API to vectorize a lane to process 2 pixel
+	instead of one, to take full advantage of v_pk_*_[iuf]16 instructions
+=============================================================================*/
+  
+#pragma once
+
+#include "Platform.ush"
+
+#define PLATFORM_SUPPORT_16BITS_DPV (PLATFORM_SUPPORTS_REAL_TYPES)
+
+
+// ------------------------------------------------- DEFINE_DPV_FUNCTION_FOR_ALL_32BITS
+
+#define __DPV_VECTOR(TYPE, COMPONENTS) \
+	TYPE##COMPONENTS
+
+#define __DPV_MULTI_SCALAR(TYPE) \
+	TYPE##2
+
+#define __DPV_MULTI_VECTOR(TYPE, COMPONENTS) \
+	TYPE##COMPONENTS##x2
+
+#define __DPV_MULTI_VECTOR_TRANSPOSED(TYPE, COMPONENTS) \
+	TYPE##2x##COMPONENTS
+
+
+#if PLATFORM_SUPPORT_16BITS_DPV
+
+#define DEFINE_DPV_FUNCTION_FOR_FLOAT_SCALARS(DEFINE_DPV_FUNCTION) \
+	DEFINE_DPV_FUNCTION(float) \
+	DEFINE_DPV_FUNCTION(float16_t) \
+
+#define DEFINE_DPV_FUNCTION_FOR_INT_SCALARS(DEFINE_DPV_FUNCTION) \
+	DEFINE_DPV_FUNCTION(uint) \
+	DEFINE_DPV_FUNCTION(int) \
+	DEFINE_DPV_FUNCTION(uint16_t) \
+	DEFINE_DPV_FUNCTION(int16_t) \
+		
+#define DEFINE_DPV_FUNCTION_FOR_FLOAT_VECTORS(DEFINE_DPV_FUNCTION) \
+	DEFINE_DPV_FUNCTION(float, 2) \
+	DEFINE_DPV_FUNCTION(float, 3) \
+	DEFINE_DPV_FUNCTION(float, 4) \
+	DEFINE_DPV_FUNCTION(float16_t, 2) \
+	DEFINE_DPV_FUNCTION(float16_t, 3) \
+	DEFINE_DPV_FUNCTION(float16_t, 4) \
+
+#define DEFINE_DPV_FUNCTION_FOR_INT_VECTORS(DEFINE_DPV_FUNCTION) \
+	DEFINE_DPV_FUNCTION(uint, 2) \
+	DEFINE_DPV_FUNCTION(uint, 3) \
+	DEFINE_DPV_FUNCTION(uint, 4) \
+	DEFINE_DPV_FUNCTION(int, 2) \
+	DEFINE_DPV_FUNCTION(int, 3) \
+	DEFINE_DPV_FUNCTION(int, 4) \
+	DEFINE_DPV_FUNCTION(uint16_t, 2) \
+	DEFINE_DPV_FUNCTION(uint16_t, 3) \
+	DEFINE_DPV_FUNCTION(uint16_t, 4) \
+	DEFINE_DPV_FUNCTION(int16_t, 2) \
+	DEFINE_DPV_FUNCTION(int16_t, 3) \
+	DEFINE_DPV_FUNCTION(int16_t, 4) \
+		
+#else
+
+#define DEFINE_DPV_FUNCTION_FOR_FLOAT_SCALARS(DEFINE_DPV_FUNCTION) \
+	DEFINE_DPV_FUNCTION(float) \
+
+#define DEFINE_DPV_FUNCTION_FOR_INT_SCALARS(DEFINE_DPV_FUNCTION) \
+	DEFINE_DPV_FUNCTION(uint) \
+	DEFINE_DPV_FUNCTION(int) \
+		
+#define DEFINE_DPV_FUNCTION_FOR_FLOAT_VECTORS(DEFINE_DPV_FUNCTION) \
+	DEFINE_DPV_FUNCTION(float, 2) \
+	DEFINE_DPV_FUNCTION(float, 3) \
+	DEFINE_DPV_FUNCTION(float, 4) \
+
+#define DEFINE_DPV_FUNCTION_FOR_INT_VECTORS(DEFINE_DPV_FUNCTION) \
+	DEFINE_DPV_FUNCTION(uint, 2) \
+	DEFINE_DPV_FUNCTION(uint, 3) \
+	DEFINE_DPV_FUNCTION(uint, 4) \
+	DEFINE_DPV_FUNCTION(int, 2) \
+	DEFINE_DPV_FUNCTION(int, 3) \
+	DEFINE_DPV_FUNCTION(int, 4) \
+		
+#endif
+
+#define DEFINE_DPV_FUNCTION_FOR_ALL_SCALARS(DEFINE_DPV_FUNCTION) \
+	DEFINE_DPV_FUNCTION_FOR_FLOAT_SCALARS(DEFINE_DPV_FUNCTION) \
+	DEFINE_DPV_FUNCTION_FOR_INT_SCALARS(DEFINE_DPV_FUNCTION) \
+
+#define DEFINE_DPV_FUNCTION_FOR_ALL_VECTORS(DEFINE_DPV_FUNCTION) \
+	DEFINE_DPV_FUNCTION_FOR_FLOAT_VECTORS(DEFINE_DPV_FUNCTION) \
+	DEFINE_DPV_FUNCTION_FOR_INT_VECTORS(DEFINE_DPV_FUNCTION) \
+
+
+/* dpv_interleave_registers() & dpv_lo() & dpv_hi()
+ *
+ * Reinterleaves explicit vector floatN for each individual pixel into a floatNx2 matrix.
+ * 
+ * Example:
+ *		half3 Color0 = ... ;
+ *		half3 Color1 = ... ;
+ *		
+ *		half3x2 Color0 = dpv_interleave_registers(Color0, Color1);
+ *		
+ *		// Can access each individual channel with the [0] notation
+ *		half2 Red = Color[0];
+ *		half2 Green = Color[1];
+ *		half2 Blue = Color[2];
+ *		
+ *		// And access back the per vectorized pixel data if need.
+ *		Color0 = dpv_lo(Color);
+ *		Color1 = dpv_hi(Color);
+ */
+
+#define DEFINE_DPV_FUNCTION_SCALAR(TYPE) \
+	__DPV_MULTI_SCALAR(TYPE) dpv_interleave_registers(TYPE lo, TYPE hi) { return __DPV_MULTI_SCALAR(TYPE)(lo, hi); } \
+	TYPE dpv_lo(__DPV_MULTI_SCALAR(TYPE) x) { return x[0]; } \
+	TYPE dpv_hi(__DPV_MULTI_SCALAR(TYPE) x) { return x[1]; } \
+
+#define DEFINE_DPV_FUNCTION_VECTOR(TYPE, COMPONENTS) \
+	__DPV_MULTI_VECTOR(TYPE, COMPONENTS) dpv_interleave_registers(__DPV_VECTOR(TYPE, COMPONENTS) lo, __DPV_VECTOR(TYPE, COMPONENTS) hi) { return transpose(__DPV_MULTI_VECTOR_TRANSPOSED(TYPE, COMPONENTS)(lo, hi)); } \
+	__DPV_VECTOR(TYPE, COMPONENTS) dpv_lo(__DPV_MULTI_VECTOR(TYPE, COMPONENTS) x) { return transpose(x)[0]; } \
+	__DPV_VECTOR(TYPE, COMPONENTS) dpv_hi(__DPV_MULTI_VECTOR(TYPE, COMPONENTS) x) { return transpose(x)[1]; } \
+
+DEFINE_DPV_FUNCTION_FOR_ALL_SCALARS(DEFINE_DPV_FUNCTION_SCALAR)
+DEFINE_DPV_FUNCTION_FOR_ALL_VECTORS(DEFINE_DPV_FUNCTION_VECTOR)
+#undef DEFINE_DPV_FUNCTION_SCALAR
+#undef DEFINE_DPV_FUNCTION_VECTOR
+
+
+/* dpv_interleave_mono_registers()
+ *
+ * Reinterleaves explicit vector floatN for of one individual pixel into a floatNx2 matrix to reuse code for dual pixel vectorization.
+ * 
+ * Example:
+ *		// transform color space a pair of colors
+ *		half3x2 RGBToYCoCg(half3x2 Color) { ... }
+ *		
+ *		// transform color space of a color
+ *		half3 RGBToYCoCg(half3 Color)
+ *		{
+ *			return dpv_lo(RGBToYCoCg(dpv_interleave_mono_registers(Color)));
+ *		}
+ */
+#define DEFINE_DPV_FUNCTION_SCALAR(TYPE) \
+	__DPV_MULTI_SCALAR(TYPE) dpv_interleave_mono_registers(TYPE lo) { return dpv_interleave_registers(lo, lo); }
+
+#define DEFINE_DPV_FUNCTION_VECTOR(TYPE, COMPONENTS) \
+	__DPV_MULTI_VECTOR(TYPE, COMPONENTS) dpv_interleave_mono_registers(__DPV_VECTOR(TYPE, COMPONENTS) lo) { return dpv_interleave_registers(lo, lo); }
+
+DEFINE_DPV_FUNCTION_FOR_ALL_SCALARS(DEFINE_DPV_FUNCTION_SCALAR)
+DEFINE_DPV_FUNCTION_FOR_ALL_VECTORS(DEFINE_DPV_FUNCTION_VECTOR)
+#undef DEFINE_DPV_FUNCTION_SCALAR
+#undef DEFINE_DPV_FUNCTION_VECTOR
+
+
+/* per vector multiplication
+ *
+ * Example:
+ *		half3x2 Color = ...;
+ *		half2 Coverage = ...;
+ * 
+ *		half3x2 BlueishColor = dpv_mul(Color, half3(1.00, 0.25, 0.25));
+ *		half3x2 PremultipliedColor = dpv_scale(Color, Coverage);
+ */
+#define DEFINE_DPV_FUNCTION_VECTOR(TYPE, COMPONENTS) \
+	__DPV_MULTI_VECTOR(TYPE, COMPONENTS) dpv_mul(__DPV_MULTI_VECTOR(TYPE, COMPONENTS) a, __DPV_MULTI_VECTOR(TYPE, COMPONENTS) b) { \
+		__DPV_MULTI_VECTOR(TYPE, COMPONENTS) r; \
+		UNROLL_N(COMPONENTS) for (uint c = 0; c < COMPONENTS; c++) { \
+			r[c] = a[c] * b[c]; \
+		} \
+		return r; \
+	} \
+	__DPV_MULTI_VECTOR(TYPE, COMPONENTS) dpv_mul(__DPV_MULTI_VECTOR(TYPE, COMPONENTS) a, __DPV_VECTOR(TYPE, COMPONENTS) b) { \
+		__DPV_MULTI_VECTOR(TYPE, COMPONENTS) r; \
+		UNROLL_N(COMPONENTS) for (uint c = 0; c < COMPONENTS; c++) { \
+			r[c] = a[c] * b[c]; \
+		} \
+		return r; \
+	} \
+	__DPV_MULTI_VECTOR(TYPE, COMPONENTS) dpv_mul(__DPV_VECTOR(TYPE, COMPONENTS) a, __DPV_MULTI_VECTOR(TYPE, COMPONENTS) b) { \
+		__DPV_MULTI_VECTOR(TYPE, COMPONENTS) r; \
+		UNROLL_N(COMPONENTS) for (uint c = 0; c < COMPONENTS; c++) { \
+			r[c] = a[c] * b[c]; \
+		} \
+		return r; \
+	} \
+	__DPV_MULTI_VECTOR(TYPE, COMPONENTS) dpv_scale(__DPV_MULTI_VECTOR(TYPE, COMPONENTS) a, __DPV_MULTI_SCALAR(TYPE) b) { \
+		__DPV_MULTI_VECTOR(TYPE, COMPONENTS) r; \
+		UNROLL_N(COMPONENTS) for (uint c = 0; c < COMPONENTS; c++) { \
+			r[c] = a[c] * b; \
+		} \
+		return r; \
+	} \
+	__DPV_MULTI_VECTOR(TYPE, COMPONENTS) dpv_scale(__DPV_MULTI_SCALAR(TYPE) a, __DPV_MULTI_VECTOR(TYPE, COMPONENTS) b) { \
+		__DPV_MULTI_VECTOR(TYPE, COMPONENTS) r; \
+		UNROLL_N(COMPONENTS) for (uint c = 0; c < COMPONENTS; c++) { \
+			r[c] = a * b[c]; \
+		} \
+		return r; \
+	} \
+	__DPV_MULTI_VECTOR(TYPE, COMPONENTS) dpv_scale(__DPV_MULTI_VECTOR(TYPE, COMPONENTS) a, TYPE b) { \
+		__DPV_MULTI_VECTOR(TYPE, COMPONENTS) r; \
+		UNROLL_N(COMPONENTS) for (uint c = 0; c < COMPONENTS; c++) { \
+			r[c] = a[c] * b; \
+		} \
+		return r; \
+	} \
+	__DPV_MULTI_VECTOR(TYPE, COMPONENTS) dpv_scale(TYPE a, __DPV_MULTI_VECTOR(TYPE, COMPONENTS) b) { \
+		__DPV_MULTI_VECTOR(TYPE, COMPONENTS) r; \
+		UNROLL_N(COMPONENTS) for (uint c = 0; c < COMPONENTS; c++) { \
+			r[c] = a * b[c]; \
+		} \
+		return r; \
+	} \
+
+DEFINE_DPV_FUNCTION_FOR_ALL_VECTORS(DEFINE_DPV_FUNCTION_VECTOR)
+#undef DEFINE_DPV_FUNCTION_VECTOR
+
+
+/* per vector multiplication
+ *
+ * Example:
+ *		half3x2 PosA = ...;
+ *		half3 PosAToB = ...;
+ * 
+ *		half3x2 PosB = dpv_add(PosA, PosAToB);
+ */
+#define DEFINE_DPV_FUNCTION_VECTOR(TYPE, COMPONENTS) \
+	__DPV_MULTI_VECTOR(TYPE, COMPONENTS) dpv_add(__DPV_MULTI_VECTOR(TYPE, COMPONENTS) a, __DPV_MULTI_VECTOR(TYPE, COMPONENTS) b) { \
+		__DPV_MULTI_VECTOR(TYPE, COMPONENTS) r; \
+		UNROLL_N(COMPONENTS) for (uint c = 0; c < COMPONENTS; c++) { \
+			r[c] = a[c] + b[c]; \
+		} \
+		return r; \
+	} \
+	__DPV_MULTI_VECTOR(TYPE, COMPONENTS) dpv_add(__DPV_MULTI_VECTOR(TYPE, COMPONENTS) a, __DPV_VECTOR(TYPE, COMPONENTS) b) { \
+		__DPV_MULTI_VECTOR(TYPE, COMPONENTS) r; \
+		UNROLL_N(COMPONENTS) for (uint c = 0; c < COMPONENTS; c++) { \
+			r[c] = a[c] + b[c]; \
+		} \
+		return r; \
+	} \
+	__DPV_MULTI_VECTOR(TYPE, COMPONENTS) dpv_add(__DPV_VECTOR(TYPE, COMPONENTS) a, __DPV_MULTI_VECTOR(TYPE, COMPONENTS) b) { \
+		__DPV_MULTI_VECTOR(TYPE, COMPONENTS) r; \
+		UNROLL_N(COMPONENTS) for (uint c = 0; c < COMPONENTS; c++) { \
+			r[c] = a[c] + b[c]; \
+		} \
+		return r; \
+	} \
+	__DPV_MULTI_VECTOR(TYPE, COMPONENTS) dpv_sub(__DPV_MULTI_VECTOR(TYPE, COMPONENTS) a, __DPV_MULTI_VECTOR(TYPE, COMPONENTS) b) { \
+		__DPV_MULTI_VECTOR(TYPE, COMPONENTS) r; \
+		UNROLL_N(COMPONENTS) for (uint c = 0; c < COMPONENTS; c++) { \
+			r[c] = a[c] - b[c]; \
+		} \
+		return r; \
+	} \
+	__DPV_MULTI_VECTOR(TYPE, COMPONENTS) dpv_sub(__DPV_MULTI_VECTOR(TYPE, COMPONENTS) a, __DPV_VECTOR(TYPE, COMPONENTS) b) { \
+		__DPV_MULTI_VECTOR(TYPE, COMPONENTS) r; \
+		UNROLL_N(COMPONENTS) for (uint c = 0; c < COMPONENTS; c++) { \
+			r[c] = a[c] - b[c]; \
+		} \
+		return r; \
+	} \
+	__DPV_MULTI_VECTOR(TYPE, COMPONENTS) dpv_sub(__DPV_VECTOR(TYPE, COMPONENTS) a, __DPV_MULTI_VECTOR(TYPE, COMPONENTS) b) { \
+		__DPV_MULTI_VECTOR(TYPE, COMPONENTS) r; \
+		UNROLL_N(COMPONENTS) for (uint c = 0; c < COMPONENTS; c++) { \
+			r[c] = a[c] - b[c]; \
+		} \
+		return r; \
+	} \
+
+DEFINE_DPV_FUNCTION_FOR_ALL_VECTORS(DEFINE_DPV_FUNCTION_VECTOR)
+#undef DEFINE_DPV_FUNCTION_VECTOR
+
+
+/* per vector multiplication
+ *
+ * Example:
+ *		half3x2 Color = ...;
+ * 
+ *		half2 Luma = dpv_dot(Color, half3(0.25, 0.5, 0.25));
+ */
+#define DEFINE_DPV_FUNCTION_VECTOR(TYPE, COMPONENTS) \
+	__DPV_MULTI_SCALAR(TYPE) dpv_dot(__DPV_MULTI_VECTOR(TYPE, COMPONENTS) a, __DPV_MULTI_VECTOR(TYPE, COMPONENTS) b) { \
+		return dpv_interleave_registers(dot(dpv_lo(a), dpv_lo(b)), dot(dpv_hi(a), dpv_hi(b))); \
+	} \
+	__DPV_MULTI_SCALAR(TYPE) dpv_dot(__DPV_VECTOR(TYPE, COMPONENTS) a, __DPV_MULTI_VECTOR(TYPE, COMPONENTS) b) { \
+		return dpv_interleave_registers(dot(a, dpv_lo(b)), dot(a, dpv_hi(b))); \
+	} \
+	__DPV_MULTI_SCALAR(TYPE) dpv_dot(__DPV_MULTI_VECTOR(TYPE, COMPONENTS) a, __DPV_VECTOR(TYPE, COMPONENTS) b) { \
+		return dpv_interleave_registers(dot(dpv_lo(a), b), dot(dpv_hi(a), b)); \
+	} \
+	__DPV_MULTI_SCALAR(TYPE) dpv_length2(__DPV_MULTI_VECTOR(TYPE, COMPONENTS) v) { \
+		return dpv_dot(v, v); \
+	} \
+	__DPV_MULTI_SCALAR(TYPE) dpv_length(__DPV_MULTI_VECTOR(TYPE, COMPONENTS) v) { \
+		return sqrt(dpv_length2(v)); \
+	} \
+	__DPV_MULTI_VECTOR(TYPE, COMPONENTS) dpv_normalize(__DPV_MULTI_VECTOR(TYPE, COMPONENTS) v) { \
+		return dpv_scale(v, rsqrt(dpv_length2(v))); \
+	} \
+	__DPV_MULTI_VECTOR(TYPE, COMPONENTS) dpv_clamp(__DPV_MULTI_VECTOR(TYPE, COMPONENTS) x, __DPV_MULTI_VECTOR(TYPE, COMPONENTS) min, __DPV_MULTI_VECTOR(TYPE, COMPONENTS) max) { \
+		__DPV_MULTI_VECTOR(TYPE, COMPONENTS) r; \
+		UNROLL_N(COMPONENTS) for (uint c = 0; c < COMPONENTS; c++) { \
+			r[c] = clamp(x[c], min[c], max[c]); \
+		} \
+		return r; \
+	} \
+	__DPV_MULTI_VECTOR(TYPE, COMPONENTS) dpv_min(__DPV_MULTI_VECTOR(TYPE, COMPONENTS) a, __DPV_MULTI_VECTOR(TYPE, COMPONENTS) b) { \
+		return min(a, b); \
+	} \
+	__DPV_MULTI_VECTOR(TYPE, COMPONENTS) dpv_min(__DPV_VECTOR(TYPE, COMPONENTS) a, __DPV_MULTI_VECTOR(TYPE, COMPONENTS) b) { \
+		__DPV_MULTI_VECTOR(TYPE, COMPONENTS) r; \
+		UNROLL_N(COMPONENTS) for (uint c = 0; c < COMPONENTS; c++) { \
+			r[c] = min(a[c], b[c]); \
+		} \
+		return r; \
+	} \
+	__DPV_MULTI_VECTOR(TYPE, COMPONENTS) dpv_min(__DPV_MULTI_VECTOR(TYPE, COMPONENTS) a, __DPV_VECTOR(TYPE, COMPONENTS) b) { \
+		__DPV_MULTI_VECTOR(TYPE, COMPONENTS) r; \
+		UNROLL_N(COMPONENTS) for (uint c = 0; c < COMPONENTS; c++) { \
+			r[c] = min(a[c], b[c]); \
+		} \
+		return r; \
+	} \
+	__DPV_MULTI_VECTOR(TYPE, COMPONENTS) dpv_max(__DPV_MULTI_VECTOR(TYPE, COMPONENTS) a, __DPV_MULTI_VECTOR(TYPE, COMPONENTS) b) { \
+		return max(a, b); \
+	} \
+	__DPV_MULTI_VECTOR(TYPE, COMPONENTS) dpv_max(__DPV_VECTOR(TYPE, COMPONENTS) a, __DPV_MULTI_VECTOR(TYPE, COMPONENTS) b) { \
+		__DPV_MULTI_VECTOR(TYPE, COMPONENTS) r; \
+		UNROLL_N(COMPONENTS) for (uint c = 0; c < COMPONENTS; c++) { \
+			r[c] = max(a[c], b[c]); \
+		} \
+		return r; \
+	} \
+	__DPV_MULTI_VECTOR(TYPE, COMPONENTS) dpv_max(__DPV_MULTI_VECTOR(TYPE, COMPONENTS) a, __DPV_VECTOR(TYPE, COMPONENTS) b) { \
+		__DPV_MULTI_VECTOR(TYPE, COMPONENTS) r; \
+		UNROLL_N(COMPONENTS) for (uint c = 0; c < COMPONENTS; c++) { \
+			r[c] = max(a[c], b[c]); \
+		} \
+		return r; \
+	} \
+
+
+DEFINE_DPV_FUNCTION_FOR_FLOAT_VECTORS(DEFINE_DPV_FUNCTION_VECTOR)
+#undef DEFINE_DPV_FUNCTION_VECTOR
