Fixes shader compilation failure in TSR on vulkan

diff --git a/Engine/Shaders/Private/TemporalAA/TAAUpdateHistory.usf b/Engine/Shaders/Private/TemporalAA/TAAUpdateHistory.usf
--- a/Engine/Shaders/Private/TemporalAA/TAAUpdateHistory.usf
+++ b/Engine/Shaders/Private/TemporalAA/TAAUpdateHistory.usf
@@ -303,7 +303,8 @@ void MainCS(
 			bUpdateSubpixelDetails[0] = all(abs(dpv_lo(dKO)) < SizeOfOutputPixelInInputPixelWidth);
 			bUpdateSubpixelDetails[1] = all(abs(dpv_hi(dKO)) < SizeOfOutputPixelInInputPixelWidth);
 		
-			bKOSign = dKO < 0.0;
+			bKOSign[0] = dKO[0] < 0.0;
+			bKOSign[1] = dKO[1] < 0.0;
 		}
 
 
@@ -827,7 +828,7 @@ void MainCS(
 			}
 		}
 		
-		uint2 SubpixelId = dpv_dot(uint2x2(dKO), uint2(1, SUB_PIXEL_GRID_SIZE));
+		uint2 SubpixelId = dpv_dot(uint2x2(bKOSign), uint2(1, SUB_PIXEL_GRID_SIZE));
 
 		taa_subpixel_payload2 SubpixelShift = taa_subpixel_payload2(SubpixelId) * taa_subpixel_payload(SUB_PIXEL_BIT_COUNT);
 
