- Use 'FLATTEN' instead of 'BRANCH' because of nVidia driver compiler bug while writing velocity to gbuffer
- Disable descriptor heap clobber hack because it should not be needed anymore

diff --git a/Engine/Shaders/Private/BasePassVertexShader.usf b/Engine/Shaders/Private/BasePassVertexShader.usf
--- a/Engine/Shaders/Private/BasePassVertexShader.usf
+++ b/Engine/Shaders/Private/BasePassVertexShader.usf
@@ -203,7 +203,8 @@ void Main(
 	{
 		float4 PrevTranslatedWorldPosition = float4(0, 0, 0, 1);
 #if GPUCULL_TODO
-		BRANCH
+		// Use 'FLATTEN' instead of 'BRANCH' because of nVidia driver compiler bug: see UE-108339 - nVidia is working on a driver side fix right now
+		FLATTEN
 		if (GetPrimitiveData(VFIntermediates).OutputVelocity > 0 || View.ForceDrawAllVelocities != 0)
 #else //!GPUCULL_TODO
 		// Warning: This here is a non-abstracted use of VFIntermediates.PrimitiveId - Only works because GetPrimitiveData(VFIntermediates.PrimitiveId) is a macro.
diff --git a/Engine/Source/Runtime/D3D12RHI/Private/D3D12DescriptorCache.cpp b/Engine/Source/Runtime/D3D12RHI/Private/D3D12DescriptorCache.cpp
--- a/Engine/Source/Runtime/D3D12RHI/Private/D3D12DescriptorCache.cpp
+++ b/Engine/Source/Runtime/D3D12RHI/Private/D3D12DescriptorCache.cpp
@@ -15,7 +15,7 @@ static FAutoConsoleVariableRef CVarGlobalViewHeapBlockSize(
 
 #if PLATFORM_WINDOWS
 
-int32 GGlobalViewHeapClobberAfterUse = 1;
+int32 GGlobalViewHeapClobberAfterUse = 0;
 static FAutoConsoleVariableRef CVarGlobalViewHeapClobberSRVAfterUse(
 	TEXT("D3D12.GlobalViewHeapClobberSRVAfterUse"),
 	GGlobalViewHeapClobberAfterUse,
