More brackets for safer strata macros.

diff --git a/Engine/Shaders/Private/Strata/Strata.ush b/Engine/Shaders/Private/Strata/Strata.ush
--- a/Engine/Shaders/Private/Strata/Strata.ush
+++ b/Engine/Shaders/Private/Strata/Strata.ush
@@ -65,10 +65,10 @@ struct FStrataBSDF
 
 
 
-#define TO_BITMASK(BitCount)							((1u<<BitCount)-1)
-#define TO_OFFSET_BITMASK(BitCount, BitOffset)			(TO_BITMASK(BitCount)<<BitOffset)
-#define READ_BITS(X, BitCount, BitOffset)				uint((X>>BitOffset) & TO_BITMASK(BitCount))
-#define WRITE_BITS(X, BitCount, BitOffset, NewValue)	X = (X & (~TO_OFFSET_BITMASK(BitCount, BitOffset))) | (((NewValue) & TO_BITMASK(BitCount)) << BitOffset)
+#define TO_BITMASK(BitCount)							((1u<<(BitCount))-1)
+#define TO_OFFSET_BITMASK(BitCount, BitOffset)			(TO_BITMASK(BitCount)<<(BitOffset))
+#define READ_BITS(X, BitCount, BitOffset)				uint((X>>(BitOffset)) & TO_BITMASK(BitCount))
+#define WRITE_BITS(X, BitCount, BitOffset, NewValue)	X = (X & (~TO_OFFSET_BITMASK(BitCount, BitOffset))) | (((NewValue) & TO_BITMASK(BitCount)) << (BitOffset))
 
 // Bit count for each component stored in a Strata header
 // For HEADER_BIT_BSDF_COUNT and  HEADER_BIT_SHAREDNORMAL_X, please refer to values in StrataDefinitions.h
