Strata parameter blending fixes:
- Add: SSSProfile is selected from the BSDF trigerring it
- Layer: weight of the top BSDF is not taken into account
Comment.

diff --git a/Engine/Shaders/Private/Strata/Strata.ush b/Engine/Shaders/Private/Strata/Strata.ush
--- a/Engine/Shaders/Private/Strata/Strata.ush
+++ b/Engine/Shaders/Private/Strata/Strata.ush
@@ -961,6 +961,7 @@ FStrataData StrataHorizontalMixingParameterBlending(FStrataData Foreground, FStr
 	ResultBSDF.TmpDMFP =				lerp(BackgroundBSDF.TmpDMFP,					ForegroundBSDF.TmpDMFP,						ForegroundMixFactor);
 #endif
 
+	// We blend the parameter above, and we also need to lerp the weight to match the total energy and coverage.
 	ResultBSDF.Weight =					lerp(BackgroundBSDF.Weight,						ForegroundBSDF.Weight,						ForegroundMixFactor);
 
 #undef ResultBSDF
@@ -1011,9 +1012,17 @@ FStrataData StrataAddParameterBlending(FStrataData A, FStrataData B, uint NewNor
 
 	if (BSDF_GETHASSSSPROFILE(ResultBSDF))
 	{
-		// As mentioned above, we cannot blend SSSProfiles.
-		SLAB_SSSPROFILEID(ResultBSDF) = SLAB_SSSPROFILEID(ABSDF);
-		SLAB_SSSPROFILERADIUSSCALE(ResultBSDF) = SLAB_SSSPROFILERADIUSSCALE(ABSDF);
+		// As mentioned above, we cannot blend SSSProfiles. We do select the appropriate profile, ABSDF being selected if both have SSSProfile.
+		if (BSDF_GETHASSSSPROFILE(ABSDF))
+		{
+			SLAB_SSSPROFILEID(ResultBSDF) = SLAB_SSSPROFILEID(ABSDF);
+			SLAB_SSSPROFILERADIUSSCALE(ResultBSDF) = SLAB_SSSPROFILERADIUSSCALE(ABSDF);
+		}
+		else
+		{
+			SLAB_SSSPROFILEID(ResultBSDF) = SLAB_SSSPROFILEID(BBSDF);
+			SLAB_SSSPROFILERADIUSSCALE(ResultBSDF) = SLAB_SSSPROFILERADIUSSCALE(BBSDF);
+		}
 	}
 	else
 	{
@@ -1064,7 +1073,9 @@ FStrataData StrataVerticalLayeringParameterBlending(FStrataData Top, FStrataData
 	const float3 UnusedDiffuseColor		= 0;	// Unused when only transmittance is needed
 	FParticipatingMedia PM				= StrataSlabCreateParticipatingMedia(UnusedDiffuseColor, SLAB_SSSDMFP(TopBSDF));
 	const float3 SlabTransmittance		= IsotropicMediumSlabTransmittance(PM, STRATA_SIMPLEVOLUME_THICKNESS_M, NoV);
-	const float3 ColorThroughput		= SlabTransmittance;
+	// Compute the colored throughput, taking into account the weight.
+	const float3 SafeTopBSDFWeight		= saturate(TopBSDF.Weight);
+	const float3 ColorThroughput		= (1.0f - SafeTopBSDFWeight) + SafeTopBSDFWeight * SlabTransmittance;
 	const float GreyThroughput			= dot(ColorThroughput, 1.0f / 3.0f);
 	
 	// Add both layers thickness
