Strata - fixed Add operation when using parameter blendign to account for BSDF weight.

diff --git a/Engine/Shaders/Private/Strata/Strata.ush b/Engine/Shaders/Private/Strata/Strata.ush
--- a/Engine/Shaders/Private/Strata/Strata.ush
+++ b/Engine/Shaders/Private/Strata/Strata.ush
@@ -973,42 +973,47 @@ FStrataData StrataHorizontalMixingParameterBlending(FStrataData Foreground, FStr
 // ==> NOTE: Always pair with the compiler behavior in StrataCompilationInfoAddParamBlend
 FStrataData StrataAddParameterBlending(FStrataData A, FStrataData B, uint NewNormalIndex)
 {
-	const float MixFactor = 0.5f;
 
 	FStrataData Result = GetInitialisedStrataDataWithReservedFirstSlot();
 #define ResultBSDF		Result.Layers[0].BSDFs[0]
 #define ABSDF			A.Layers[0].BSDFs[0]
 #define BBSDF			B.Layers[0].BSDFs[0]
 
+	const float3 SafeABSDFWeight = saturate(ABSDF.Weight);
+	const float3 SafeBBSDFWeight = saturate(BBSDF.Weight);
+	const float SafeGreyABSDFWeight = dot(SafeABSDFWeight, 1.0f / 3.0f);
+	const float SafeGreyBBSDFWeight = dot(SafeBBSDFWeight, 1.0f / 3.0f);
+	const float AMixFactor = SafeGreyABSDFWeight / (1e-10 + SafeGreyABSDFWeight + SafeGreyBBSDFWeight);
+
 	// We can only parameter blend the slab 
 	BSDF_SETTYPE(ResultBSDF,			STRATA_BSDF_TYPE_SLAB);
 	BSDF_SETSHAREDLOCALBASISID(ResultBSDF,NewNormalIndex);
 
 	BSDF_SETISTOPLAYER(ResultBSDF,		0);
-	BSDF_SETANISOTROPY(ResultBSDF,		BSDF_GETANISOTROPY(ABSDF)						| BSDF_GETANISOTROPY(BBSDF));
-	BSDF_SETHASSSS(ResultBSDF,			BSDF_GETHASSSS(ABSDF)							| BSDF_GETHASSSS(BBSDF));
-	BSDF_SETHASSSSPROFILE(ResultBSDF,	BSDF_GETHASSSSPROFILE(ABSDF)					| BSDF_GETHASSSSPROFILE(BBSDF));	// Problem: cannot blend SSS profiles...
-	BSDF_SETISTHIN(ResultBSDF,			BSDF_GETISTHIN(ABSDF)							| BSDF_GETISTHIN(BBSDF));			// STRATA_TODO fix the MPF/DMFP situation for that to really work
-	BSDF_SETHASDMFP(ResultBSDF,			BSDF_GETHASDMFP(ABSDF)							| BSDF_GETHASDMFP(BBSDF));
-	BSDF_SETHASHAZINESS(ResultBSDF,		BSDF_GETHASHAZINESS(ABSDF)						| BSDF_GETHASHAZINESS(BBSDF));
-	BSDF_SETHASF90(ResultBSDF,			BSDF_GETHASF90(ABSDF)							| BSDF_GETHASF90(BBSDF));
-	BSDF_SETHASTHINFILM(ResultBSDF,		BSDF_GETHASTHINFILM(ABSDF)						| BSDF_GETHASTHINFILM(BBSDF));
-	BSDF_SETHASFUZZ(ResultBSDF,			BSDF_GETHASFUZZ(ABSDF)							| BSDF_GETHASFUZZ(BBSDF));
-	BSDF_SETEMISSIVE(ResultBSDF,		BSDF_GETEMISSIVE(ABSDF)							+ BSDF_GETEMISSIVE(BBSDF));
-	BSDF_SETTHICKNESSCM(ResultBSDF,		lerp(BSDF_GETTHICKNESSCM(ABSDF)					, BSDF_GETTHICKNESSCM(BBSDF),				MixFactor));
+	BSDF_SETANISOTROPY(ResultBSDF,		BSDF_GETANISOTROPY(BBSDF)								| BSDF_GETANISOTROPY(ABSDF));
+	BSDF_SETHASSSS(ResultBSDF,			BSDF_GETHASSSS(BBSDF)									| BSDF_GETHASSSS(ABSDF));
+	BSDF_SETHASSSSPROFILE(ResultBSDF,	BSDF_GETHASSSSPROFILE(BBSDF)							| BSDF_GETHASSSSPROFILE(ABSDF));	// Problem: cannot blend SSS profiles...
+	BSDF_SETISTHIN(ResultBSDF,			BSDF_GETISTHIN(BBSDF)									| BSDF_GETISTHIN(ABSDF));			// STRATA_TODO fix the MPF/DMFP situation for that to really work
+	BSDF_SETHASDMFP(ResultBSDF,			BSDF_GETHASDMFP(BBSDF)									| BSDF_GETHASDMFP(ABSDF));
+	BSDF_SETHASHAZINESS(ResultBSDF,		BSDF_GETHASHAZINESS(BBSDF)								| BSDF_GETHASHAZINESS(ABSDF));
+	BSDF_SETHASF90(ResultBSDF,			BSDF_GETHASF90(BBSDF)									| BSDF_GETHASF90(ABSDF));
+	BSDF_SETHASTHINFILM(ResultBSDF,		BSDF_GETHASTHINFILM(BBSDF)								| BSDF_GETHASTHINFILM(ABSDF));
+	BSDF_SETHASFUZZ(ResultBSDF,			BSDF_GETHASFUZZ(BBSDF)									| BSDF_GETHASFUZZ(ABSDF));
+	BSDF_SETEMISSIVE(ResultBSDF,		BSDF_GETEMISSIVE(BBSDF) * SafeBBSDFWeight				+ BSDF_GETEMISSIVE(ABSDF) * SafeABSDFWeight);
+	BSDF_SETTHICKNESSCM(ResultBSDF,		lerp(BSDF_GETTHICKNESSCM(BBSDF)							, BSDF_GETTHICKNESSCM(ABSDF),				AMixFactor));
 
 	// Some parameters will contribute to added luminance: in this case we add them together and saturate to not go out of the safe range
 	// Some parameters are not contributing to added luminance: in this case we simply lerp them
-	SLAB_DIFFUSEALBEDO(ResultBSDF) =	saturate(SLAB_DIFFUSEALBEDO(BBSDF)				+ SLAB_DIFFUSEALBEDO(ABSDF));
-	SLAB_F0(ResultBSDF) =				saturate(SLAB_F0(BBSDF)							+ SLAB_F0(ABSDF));
-	SLAB_F90(ResultBSDF) =				saturate(SLAB_F90(BBSDF)						+ SLAB_F90(ABSDF));
-	SLAB_METALLIC(ResultBSDF) =			lerp(SLAB_METALLIC(BBSDF),						SLAB_METALLIC(ABSDF),						MixFactor);
-	SLAB_ROUGHNESSX(ResultBSDF) =		lerp(SLAB_ROUGHNESSX(BBSDF),					SLAB_ROUGHNESSX(ABSDF),						MixFactor);
-	SLAB_ROUGHNESSY(ResultBSDF) =		lerp(SLAB_ROUGHNESSY(BBSDF),					SLAB_ROUGHNESSY(ABSDF),						MixFactor);
-	SLAB_HAZINESS(ResultBSDF) =			lerp(SLAB_HAZINESS(BBSDF),						SLAB_HAZINESS(ABSDF),						MixFactor);
-	SLAB_THINFILM_THICKNESS(ResultBSDF)=lerp(SLAB_THINFILM_THICKNESS(BBSDF),			SLAB_THINFILM_THICKNESS(ABSDF),				MixFactor);
-	SLAB_FUZZ_AMOUNT(ResultBSDF) =		lerp(SLAB_FUZZ_AMOUNT(BBSDF),					SLAB_FUZZ_AMOUNT(ABSDF),					MixFactor);
-	SLAB_FUZZ_COLOR(ResultBSDF) =		saturate(SLAB_FUZZ_COLOR(BBSDF)					+ SLAB_FUZZ_COLOR(ABSDF));
+	SLAB_DIFFUSEALBEDO(ResultBSDF) =	saturate(SLAB_DIFFUSEALBEDO(BBSDF) * SafeBBSDFWeight	+ SLAB_DIFFUSEALBEDO(ABSDF)	* SafeABSDFWeight);
+	SLAB_F0(ResultBSDF) =				saturate(SLAB_F0(BBSDF) * SafeBBSDFWeight				+ SLAB_F0(ABSDF) * SafeABSDFWeight);
+	SLAB_F90(ResultBSDF) =				saturate(SLAB_F90(BBSDF) * SafeBBSDFWeight				+ SLAB_F90(ABSDF) * SafeABSDFWeight);
+	SLAB_METALLIC(ResultBSDF) =			lerp(SLAB_METALLIC(BBSDF),								SLAB_METALLIC(ABSDF),						AMixFactor);
+	SLAB_ROUGHNESSX(ResultBSDF) =		lerp(SLAB_ROUGHNESSX(BBSDF),							SLAB_ROUGHNESSX(ABSDF),						AMixFactor);
+	SLAB_ROUGHNESSY(ResultBSDF) =		lerp(SLAB_ROUGHNESSY(BBSDF),							SLAB_ROUGHNESSY(ABSDF),						AMixFactor);
+	SLAB_HAZINESS(ResultBSDF) =			lerp(SLAB_HAZINESS(BBSDF),								SLAB_HAZINESS(ABSDF),						AMixFactor);
+	SLAB_THINFILM_THICKNESS(ResultBSDF)=lerp(SLAB_THINFILM_THICKNESS(BBSDF),					SLAB_THINFILM_THICKNESS(ABSDF),				AMixFactor);
+	SLAB_FUZZ_AMOUNT(ResultBSDF) =		lerp(SLAB_FUZZ_AMOUNT(BBSDF),							SLAB_FUZZ_AMOUNT(ABSDF),					AMixFactor);
+	SLAB_FUZZ_COLOR(ResultBSDF) =		saturate(SLAB_FUZZ_COLOR(BBSDF) * SafeBBSDFWeight		+ SLAB_FUZZ_COLOR(ABSDF) * SafeABSDFWeight);
 
 	if (BSDF_GETHASSSSPROFILE(ResultBSDF))
 	{
@@ -1026,17 +1031,18 @@ FStrataData StrataAddParameterBlending(FStrataData A, FStrataData B, uint NewNor
 	}
 	else
 	{
-		SLAB_SSSDMFP(ResultBSDF) =		lerp(SLAB_SSSDMFP(BBSDF),						SLAB_SSSDMFP(ABSDF),						MixFactor);
+		SLAB_SSSDMFP(ResultBSDF) =		lerp(SLAB_SSSDMFP(BBSDF),								SLAB_SSSDMFP(ABSDF),						AMixFactor);
 	}
 
 
 #if STRATA_INLINE_SHADING
 	// At this stage we do not know if this strata BSDF is going to be using simple volume lighting so always store the DMFP on the side to not override SSSPROFILE data.
-	ResultBSDF.TmpDMFP =				lerp(BBSDF.TmpDMFP,								ABSDF.TmpDMFP,								MixFactor);
+	ResultBSDF.TmpDMFP =				lerp(BBSDF.TmpDMFP,										ABSDF.TmpDMFP,								AMixFactor);
 #endif
 
-	// When it comes to adding the visual contribution, it is tricky to find a matching operator for coverage. We simply lerp for now.
-	ResultBSDF.Weight =					lerp(BBSDF.Weight,								ABSDF.Weight,								MixFactor);
+	// When it comes to adding the visual contribution, we add the weights without clamping as currently done in the StrataAdd function.
+	// Operations using weights are currently making sure they are in a safe range.
+	ResultBSDF.Weight =	 BBSDF.Weight + ABSDF.Weight;
 
 #undef ResultBSDF
 #undef ABSDF
