Optimises TSR's spatial upscale filter with ClampPixelOffset()

diff --git a/Engine/Shaders/Private/TemporalAA/TAACommon.ush b/Engine/Shaders/Private/TemporalAA/TAACommon.ush
--- a/Engine/Shaders/Private/TemporalAA/TAACommon.ush
+++ b/Engine/Shaders/Private/TemporalAA/TAACommon.ush
@@ -336,6 +336,40 @@ uint GetGroupThreadIndex(uint GroupThreadIndex, uint GroupWaveIndex)
 #endif
 
 // Clamp the offset to be shared across multiple samples
+taa_short2x2 ClampPixelOffset(
+	taa_short2x2 KernelCenterPixelPos,
+	taa_short2x2 Offset,
+	const taa_short2 OffsetDirection,
+	int2 MinPixelPos, int2 MaxPixelPos)
+{
+	taa_short2x2 Min = dpv_sub(taa_short2(MinPixelPos), KernelCenterPixelPos);
+	taa_short2x2 Max = dpv_sub(taa_short2(MaxPixelPos), KernelCenterPixelPos);
+
+	// Only do clamp based on the compile time known direction of the offset.
+	// This turns only 1 v_pk_max_i16 and v_pk_min_i16 for a 3x3 kernel.
+	taa_short2x2 ClampedOffset = 0;
+
+	if (OffsetDirection.x > 0)
+	{
+		ClampedOffset[0] = min(Offset[0], Max[0]);
+	}
+	else if (OffsetDirection.x < 0)
+	{
+		ClampedOffset[0] = max(Offset[0], Min[0]);
+	}
+		
+	if (OffsetDirection.y > 0)
+	{
+		ClampedOffset[1] = min(Offset[1], Max[1]);
+	}
+	else if (OffsetDirection.y < 0)
+	{
+		ClampedOffset[1] = max(Offset[1], Min[1]);
+	}
+
+	return ClampedOffset;
+}
+
 taa_short2 ClampPixelOffset(
 	taa_short2 KernelCenterPixelPos,
 	taa_short2 Offset,
diff --git a/Engine/Shaders/Private/TemporalAA/TAAUpdateHistory.usf b/Engine/Shaders/Private/TemporalAA/TAAUpdateHistory.usf
--- a/Engine/Shaders/Private/TemporalAA/TAAUpdateHistory.usf
+++ b/Engine/Shaders/Private/TemporalAA/TAAUpdateHistory.usf
@@ -111,7 +111,7 @@ void ComputeInputKernelSamplePosition(
 		taa_short2 iPixelOffset = taa_short2(kOffsets3x3[kSquareIndexes3x3[SampleId]]);
 		PixelOffset = dpv_interleave_registers(taa_half2(iPixelOffset), taa_half2(iPixelOffset));
 				
-		SampleInputPixelPos = AddAndClampPixelOffset(
+		SampleInputPixelPos = InputPixelPos + ClampPixelOffset(
 			InputPixelPos,
 			dpv_interleave_registers(iPixelOffset, iPixelOffset), iPixelOffset,
 			InputPixelPosMin, InputPixelPosMax);
@@ -140,7 +140,7 @@ void ComputeInputKernelSamplePosition(
 			taa_short2 iPixelOffset = taa_short2(kOffsets3x3[kPlusIndexes3x3[SampleId]]);
 			PixelOffset = dpv_interleave_registers(taa_half2(iPixelOffset), taa_half2(iPixelOffset));
 					
-			SampleInputPixelPos = AddAndClampPixelOffset(
+			SampleInputPixelPos = InputPixelPos + ClampPixelOffset(
 				InputPixelPos,
 				dpv_interleave_registers(iPixelOffset, iPixelOffset), iPixelOffset,
 				InputPixelPosMin, InputPixelPosMax);
