Implements KERNEL_3X3_PLUS in TSR FilterFrequencies

diff --git a/Engine/Shaders/Private/TemporalAA/TAAFilterFrequencies.usf b/Engine/Shaders/Private/TemporalAA/TAAFilterFrequencies.usf
--- a/Engine/Shaders/Private/TemporalAA/TAAFilterFrequencies.usf
+++ b/Engine/Shaders/Private/TemporalAA/TAAFilterFrequencies.usf
@@ -8,6 +8,7 @@
 /** Kernel enum */
 	#define KERNEL_1X1 0
 	#define KERNEL_3X3 1
+	#define KERNEL_3X3_PLUS 2
 
 
 //------------------------------------------------------- CONFIG
@@ -110,7 +111,7 @@ void MainCS(
 		StartSampleId = 0;
 		SampleCount = 0;
 	}
-	#elif CONFIG_KERNEL == KERNEL_3X3
+	#elif CONFIG_KERNEL == KERNEL_3X3 || CONFIG_KERNEL == KERNEL_3X3_PLUS
 	{
 		StartSampleId = 0;
 		SampleCount = 9;
@@ -131,7 +132,7 @@ void MainCS(
 	UNROLL
 	for (uint SampleId = StartSampleId; SampleId < SampleCount; SampleId++)
 	{
-		int2 Offset;
+		taa_short2 Offset;
 		taa_half SampleSpatialWeight;
 		
 		#if CONFIG_KERNEL == KERNEL_1X1
@@ -140,18 +141,20 @@ void MainCS(
 			SampleSpatialWeight = taa_half(1.0);
 			continue;
 		}
-		#elif CONFIG_KERNEL == KERNEL_3X3
+		#elif CONFIG_KERNEL == KERNEL_3X3 || CONFIG_KERNEL == KERNEL_3X3_PLUS
 		{
 			if (SampleId == 4)
 			{
 				continue;
 			}
-			Offset = kOffsets3x3[SampleId];
+			Offset = taa_short2(kOffsets3x3[SampleId]);
 			
-			//if (abs(Offset.x) + abs(Offset.y) > 1)
-			//{
-			//	continue;
-			//}
+			#if CONFIG_KERNEL == KERNEL_3X3_PLUS
+				if (abs(Offset.x) + abs(Offset.y) > 1)
+				{
+					continue;
+				}
+			#endif
 			
 			#if 0
 				SampleSpatialWeight = taa_half(1.0);
