Fix multi-view rendering for the path tracer

The path tracing radiance texture is allocated based on view rect extents, but was being accessed with aboslute coordinates which didn't work when ViewRectMin was non-zero

diff --git a/Engine/Shaders/Private/PathTracing/PathTracing.usf b/Engine/Shaders/Private/PathTracing/PathTracing.usf
--- a/Engine/Shaders/Private/PathTracing/PathTracing.usf
+++ b/Engine/Shaders/Private/PathTracing/PathTracing.usf
@@ -4,7 +4,7 @@
 PathTracingRayGenShader.usf: Reference path tracing
 ===============================================================================================*/
 
-#define PATH_TRACING  
+#define PATH_TRACING  1
 
 #define PATH_TRACING_SHADER_VERSION 0x51478b68 // Bump to force-compile this shader
 
@@ -488,8 +488,14 @@ RAY_TRACING_ENTRY_RAYGEN(PathTracingMainRG)
 	uint2 DispatchIdx = DispatchRaysIndex().xy;
 	uint2 DispatchDim = DispatchRaysDimensions().xy;
 
-	// shuffle the pixels around so the ones that share a seed land in the same warp
-	uint2 LaunchIndex = (SamplerType == PATHTRACER_SAMPLER_COHERENT ? Reorder(DispatchIdx, DispatchDim) : DispatchIdx) + View.ViewRectMin.xy + TileOffset.xy;
+	uint2 PixelIndex = DispatchIdx;
+	if (SamplerType == PATHTRACER_SAMPLER_COHERENT)
+	{
+		// shuffle the pixels around so the ones that share a seed land in the same warp
+		PixelIndex = Reorder(PixelIndex, DispatchDim);
+	}
+	PixelIndex = PixelIndex + TileOffset.xy;
+	uint2 LaunchIndex = PixelIndex + View.ViewRectMin.xy;
 
 	// Initialize random sequence 
 	RandomSequence RandSequence;
@@ -783,7 +789,7 @@ RAY_TRACING_ENTRY_RAYGEN(PathTracingMainRG)
 	uint NumSamples = Iteration + 1;
 	float BlendFactor = 1.0 / float(NumSamples);
 	// Avoid reading the old pixel on the first sample on the off-chance there is a NaN/Inf pixel ...
-	float4 OldPixel = NumSamples > 1 ? RadianceTexture[LaunchIndex] : 0;
+	float4 OldPixel = NumSamples > 1 ? RadianceTexture[PixelIndex] : 0;
 	float3 OldRadiance = OldPixel.rgb;
 	float OldVariance = OldPixel.a;
 
@@ -795,5 +801,5 @@ RAY_TRACING_ENTRY_RAYGEN(PathTracingMainRG)
 	float DeviationSquared = (RadianceY - Luminance(NewRadiance)) * (RadianceY - Luminance(OldRadiance));
 	float NewVariance = lerp(OldVariance, DeviationSquared, BlendFactor);
 
-	RadianceTexture[LaunchIndex] = float4(NewRadiance, NewVariance);
+	RadianceTexture[PixelIndex] = float4(NewRadiance, NewVariance);
 }
diff --git a/Engine/Shaders/Private/PathTracing/PathTracingCompositingPixelShader.usf b/Engine/Shaders/Private/PathTracing/PathTracingCompositingPixelShader.usf
--- a/Engine/Shaders/Private/PathTracing/PathTracingCompositingPixelShader.usf
+++ b/Engine/Shaders/Private/PathTracing/PathTracingCompositingPixelShader.usf
@@ -41,7 +41,7 @@ void CompositeMain(
 )
 {
 	float2 BufferSize = View.BufferSizeAndInvSize.xy;
-	int3 TexCoord = int3(UV * BufferSize, 0);
+	int3 TexCoord = int3(UV * BufferSize - View.ViewRectMin.xy, 0);
 
 	float4 Radiance = RadianceTexture.Load(TexCoord);
 
