Fixed startra visualise when weight is 10bit float.

diff --git a/Engine/Shaders/Private/Strata/Strata.ush b/Engine/Shaders/Private/Strata/Strata.ush
--- a/Engine/Shaders/Private/Strata/Strata.ush
+++ b/Engine/Shaders/Private/Strata/Strata.ush
@@ -172,8 +172,8 @@ struct FStrataBSDF
 #define BSDF_SETHASSSSPROFILE(X, Profi)	WRITE_BITS(X.State, STATE_BIT_COUNT_HASSSSPROFILE, STATE_BIT_OFFSET_HASSSSPROFILE, Profi)
 
 // Indicates if the BSDF weight is grey scale so that it can be stored in a single float (e.g. top layer BSDFs)
-#define BSDF_GETGREYWEIGHT(X)			READ_BITS(X.State,  STATE_BIT_COUNT_GREYWEIGHT, STATE_BIT_OFFSET_GREYWEIGHT)
-#define BSDF_SETGREYWEIGHT(X, Grey)		WRITE_BITS(X.State, STATE_BIT_COUNT_GREYWEIGHT, STATE_BIT_OFFSET_GREYWEIGHT, Grey)
+#define BSDF_GETHASGREYWEIGHT(X)		READ_BITS(X.State,  STATE_BIT_COUNT_GREYWEIGHT, STATE_BIT_OFFSET_GREYWEIGHT)
+#define BSDF_SETHASGREYWEIGHT(X, Grey)	WRITE_BITS(X.State, STATE_BIT_COUNT_GREYWEIGHT, STATE_BIT_OFFSET_GREYWEIGHT, Grey)
 
 // Indicates if the BSDF has haziness data
 #define BSDF_GETHASHAZINESS(X)			READ_BITS(X.State,  STATE_BIT_COUNT_HAZINESS, STATE_BIT_OFFSET_HAZINESS)
@@ -1777,7 +1777,7 @@ void PackStrataOut(
 			FStrataBSDF BSDF = Layer.BSDFs[i];
 
 			const uint GreyScaleWeight = StrataHasGreyScaleWeight(BSDF.Weight) ? 1 : 0;
-			BSDF_SETGREYWEIGHT(BSDF, GreyScaleWeight);
+			BSDF_SETHASGREYWEIGHT(BSDF, GreyScaleWeight);
 			if (GreyScaleWeight > 0)
 			{
 				BSDF_SETWEIGHT10F(BSDF, Pack10F(BSDF.Weight.x));
@@ -1973,7 +1973,7 @@ FStrataBSDF UnpackStrataBSDFIn(ByteAddressBuffer StrataBuffer, inout FStrataAddr
 	float DummyFloat = 0.0f;
 
 	STRATA_LOAD_UINT1(OutBSDF.State);
-	if (BSDF_GETGREYWEIGHT(OutBSDF))
+	if (BSDF_GETHASGREYWEIGHT(OutBSDF))
 	{
 		OutBSDF.Weight = Unpack10F(BSDF_GETWEIGHT10F(OutBSDF));
 	}
diff --git a/Engine/Shaders/Private/Strata/StrataVisualize.usf b/Engine/Shaders/Private/Strata/StrataVisualize.usf
--- a/Engine/Shaders/Private/Strata/StrataVisualize.usf
+++ b/Engine/Shaders/Private/Strata/StrataVisualize.usf
@@ -187,7 +187,7 @@ void VisualizeMaterialPS(
 #define PRINTBSDFCHARGREY(X) PrintCharacter(PixelPos, OutColor.xyz, BSDFFontColorGrey, TopLeft, X);
 
 #define PRINTBSDFWEIGHT()	PRINTBSDFCHAR(_SPC_); PRINTBSDFCHAR(_MINUS_); PRINTBSDFCHAR(_W_); PRINTBSDFCHAR(_E_); PRINTBSDFCHAR(_I_); PRINTBSDFCHAR(_G_); PRINTBSDFCHAR(_H_);PRINTBSDFCHAR(_T_);  PRINTBSDFCHAR(_SPC_); PrintFloat(PixelPos, OutColor.xyz, FontColor, TopLeft, BSDF.Weight.x); TopLeft.x+=8*7; PRINTBSDFCHAR(_SPC_); PrintFloat(PixelPos, OutColor.xyz, FontColor, TopLeft, BSDF.Weight.y); TopLeft.x+=8*7; PRINTBSDFCHAR(_SPC_); PrintFloat(PixelPos, OutColor.xyz, FontColor, TopLeft, BSDF.Weight.z); TopLeft.x+=8*7; 
-#define PRINTBSDFGREYWEIGHT()PRINTBSDFCHAR(_SPC_);PRINTBSDFCHARGREY(_MINUS_); PRINTBSDFCHARGREY(_W_); PRINTBSDFCHARGREY(_E_); PRINTBSDFCHARGREY(_I_); PRINTBSDFCHARGREY(_G_); PRINTBSDFCHARGREY(_H_);PRINTBSDFCHARGREY(_T_);  PRINTBSDFCHAR(_SPC_); PrintFloat(PixelPos, OutColor.xyz, FontColor, TopLeft, BSDF_GETGREYWEIGHT(BSDF)); TopLeft.x+=8*23; 
+#define PRINTBSDFGREYWEIGHT()PRINTBSDFCHAR(_SPC_);PRINTBSDFCHARGREY(_MINUS_); PRINTBSDFCHARGREY(_W_); PRINTBSDFCHARGREY(_E_); PRINTBSDFCHARGREY(_I_); PRINTBSDFCHARGREY(_G_); PRINTBSDFCHARGREY(_H_);PRINTBSDFCHARGREY(_T_);  PRINTBSDFCHAR(_SPC_); PrintFloat(PixelPos, OutColor.xyz, FontColor, TopLeft, Unpack10F(BSDF_GETWEIGHT10F(BSDF))); TopLeft.x+=8*23; 
 #define PRINTBSDFSHAREDLOCALBASISID()	PRINTBSDFCHAR(_SPC_); PRINTBSDFCHAR(_MINUS_); PRINTBSDFCHAR(_N_); PRINTBSDFCHAR(_R_); PRINTBSDFCHAR(_M_); PRINTBSDFCHAR(_I_); PRINTBSDFCHAR(_D_); PRINTBSDFCHAR(_SPC_); PrintFloatNoFraction(PixelPos, OutColor.xyz, FontColor, TopLeft, float(BSDF_GETSHAREDLOCALBASISID(BSDF)), 1); TopLeft.x+=8*1;
 #define PRINTBSDFANISO()	PRINTBSDFCHAR(_SPC_); PRINTBSDFCHAR(_MINUS_); PRINTBSDFCHAR(_A_); PRINTBSDFCHAR(_N_); PRINTBSDFCHAR(_I_); PRINTBSDFCHAR(_S_); PRINTBSDFCHAR(_O_); PRINTBSDFCHAR(_SPC_); PrintFloatNoFraction(PixelPos, OutColor.xyz, FontColor, TopLeft, float(BSDF_GETANISOTROPY(BSDF)), 1); TopLeft.x+=8*1;
 #define PRINTBSDFTOPLAY()	PRINTBSDFCHAR(_SPC_); PRINTBSDFCHAR(_MINUS_); PRINTBSDFCHAR(_T_); PRINTBSDFCHAR(_O_); PRINTBSDFCHAR(_P_); PRINTBSDFCHAR(_L_); PRINTBSDFCHAR(_A_); PRINTBSDFCHAR(_Y_); PRINTBSDFCHAR(_SPC_); PrintFloatNoFraction(PixelPos, OutColor.xyz, FontColor, TopLeft, float(BSDF_GETISTOPLAYER(BSDF)), 1); TopLeft.x+=8*1;
@@ -306,7 +306,7 @@ void VisualizeMaterialPS(
 #if PERMUTATION_BSDF_PASS==0
 				PRINTBSDFCHAR(_S_); PRINTBSDFCHAR(_L_); PRINTBSDFCHAR(_A_); PRINTBSDFCHAR(_B_); PRINTBSDFCHAR(_SPC_);
 				PRINTBSDFCHAR(_SPC_); PRINTBSDFCHAR(_SPC_); PRINTBSDFCHAR(_SPC_); PRINTBSDFCHAR(_SPC_); PRINTBSDFCHAR(_SPC_);
-				if (BSDF_GETGREYWEIGHT(BSDF))
+				if (BSDF_GETHASGREYWEIGHT(BSDF))
 				{
 					PRINTBSDFGREYWEIGHT();
 				}
@@ -461,7 +461,7 @@ void VisualizeMaterialPS(
 #if PERMUTATION_BSDF_PASS==1
 				PRINTBSDFCHAR(_H_); PRINTBSDFCHAR(_A_); PRINTBSDFCHAR(_I_); PRINTBSDFCHAR(_R_); PRINTBSDFCHAR(_SPC_);
 				PRINTBSDFCHAR(_SPC_); PRINTBSDFCHAR(_SPC_); PRINTBSDFCHAR(_SPC_); PRINTBSDFCHAR(_SPC_); PRINTBSDFCHAR(_SPC_);
-				if (BSDF_GETGREYWEIGHT(BSDF))
+				if (BSDF_GETHASGREYWEIGHT(BSDF))
 				{
 					PRINTBSDFGREYWEIGHT();
 				}
