Fixed desaturation of reflection when blended with a slab blends with another one having SSS.

diff --git a/Engine/Shaders/Private/Strata/Strata.ush b/Engine/Shaders/Private/Strata/Strata.ush
--- a/Engine/Shaders/Private/Strata/Strata.ush
+++ b/Engine/Shaders/Private/Strata/Strata.ush
@@ -949,10 +949,23 @@ FStrataData StrataHorizontalMixingParameterBlending(FStrataData Foreground, FStr
 		// As mentioned above, we cannot blend SSSProfiles.
 		SLAB_SSSPROFILEID(ResultBSDF) = SLAB_SSSPROFILEID(ForegroundBSDF);
 		SLAB_SSSPROFILERADIUSSCALE(ResultBSDF) = SLAB_SSSPROFILERADIUSSCALE(ForegroundBSDF);
+
+		if (all(SLAB_SSSPROFILERADIUSSCALE(ResultBSDF) == 0.0f))
+		{
+			SLAB_SSSPROFILEID(ResultBSDF) = SSS_PROFILE_ID_INVALID;
+			BSDF_SETHASSSS(ResultBSDF, 0);
+			BSDF_SETHASSSSPROFILE(ResultBSDF, 0);
+		}
 	}
 	else
 	{
 		SLAB_SSSDMFP(ResultBSDF) = lerp(SLAB_SSSDMFP(BackgroundBSDF),				SLAB_SSSDMFP(ForegroundBSDF),					ForegroundMixFactor);
+
+		// We still try to disable SSS if possible: for performance and also to preserve the surface color (especially specular reflection color).
+		if (all(SLAB_SSSDMFP(ResultBSDF) == 0.0f))
+		{
+			BSDF_SETHASSSS(ResultBSDF, 0);
+		}
 	}
 
 
