Strata horizontal blending - improved SSS profile selection and blending

diff --git a/Engine/Shaders/Private/Strata/Strata.ush b/Engine/Shaders/Private/Strata/Strata.ush
--- a/Engine/Shaders/Private/Strata/Strata.ush
+++ b/Engine/Shaders/Private/Strata/Strata.ush
@@ -946,9 +946,17 @@ FStrataData StrataHorizontalMixingParameterBlending(FStrataData Foreground, FStr
 
 	if (BSDF_GETHASSSSPROFILE(ResultBSDF))
 	{
-		// As mentioned above, we cannot blend SSSProfiles.
-		SLAB_SSSPROFILEID(ResultBSDF) = SLAB_SSSPROFILEID(ForegroundBSDF);
-		SLAB_SSSPROFILERADIUSSCALE(ResultBSDF) = SLAB_SSSPROFILERADIUSSCALE(ForegroundBSDF);
+		const bool bForegroundHasSSSProfile = BSDF_GETHASSSSPROFILE(ForegroundBSDF);
+		const bool bBackgroundHasSSSProfile = BSDF_GETHASSSSPROFILE(BackgroundBSDF);
+
+		// As mentioned above, we cannot blend SSSProfiles. So we select the one that triggered the profile, foreground having priority.
+		SLAB_SSSPROFILEID(ResultBSDF) = bForegroundHasSSSProfile ? SLAB_SSSPROFILEID(ForegroundBSDF) : SLAB_SSSPROFILEID(BackgroundBSDF);
+
+		// We also lerp the radius scale, while taking into account whether or not the one or the other have SSSProfile.
+		SLAB_SSSPROFILERADIUSSCALE(ResultBSDF) = lerp(
+			bBackgroundHasSSSProfile ? SLAB_SSSPROFILERADIUSSCALE(BackgroundBSDF) : 0.0f,
+			bForegroundHasSSSProfile ? SLAB_SSSPROFILERADIUSSCALE(ForegroundBSDF) : 0.0f,
+			ForegroundMixFactor);
 
 		if (all(SLAB_SSSPROFILERADIUSSCALE(ResultBSDF) == 0.0f))
 		{
