Shader compilation fix for OGL

diff --git a/Engine/Shaders/Private/Strata/Strata.ush b/Engine/Shaders/Private/Strata/Strata.ush
--- a/Engine/Shaders/Private/Strata/Strata.ush
+++ b/Engine/Shaders/Private/Strata/Strata.ush
@@ -958,7 +958,7 @@ FStrataData StrataHorizontalMixingParameterBlending(FStrataData Foreground, FStr
 			bForegroundHasSSSProfile ? SLAB_SSSPROFILERADIUSSCALE(ForegroundBSDF) : 0.0f,
 			ForegroundMixFactor);
 
-		if (all(SLAB_SSSPROFILERADIUSSCALE(ResultBSDF) == 0.0f))
+		if (SLAB_SSSPROFILERADIUSSCALE(ResultBSDF) == 0.0f)
 		{
 			SLAB_SSSPROFILEID(ResultBSDF) = SSS_PROFILE_ID_INVALID;
 			BSDF_SETHASSSS(ResultBSDF, 0);
