Improved stata vertical layer blending by taking top layer weight.
We will rework the weight to calrify coverage, throughput to view and that should help improve parameter blending situation.

diff --git a/Engine/Shaders/Private/Strata/Strata.ush b/Engine/Shaders/Private/Strata/Strata.ush
--- a/Engine/Shaders/Private/Strata/Strata.ush
+++ b/Engine/Shaders/Private/Strata/Strata.ush
@@ -1077,6 +1077,14 @@ FStrataData StrataAddParameterBlending(FStrataData A, FStrataData B, float AMixF
 	return Result;
 }
 
+float VerticalLayeringParameterBlendingBSDFWeightToNormalMix(FStrataData TopStrata)
+{
+	const float3 SafeTopBSDFWeight = saturate(TopStrata.Layers[0].BSDFs[0].Weight);
+	const float SafeGreyTopBSDFWeight = dot(SafeTopBSDFWeight, 1.0f / 3.0f);
+	const float TopNormalContributionForWeightOfOne = 0.5f;
+	return SafeGreyTopBSDFWeight * TopNormalContributionForWeightOfOne;
+}
+
 // ==> NOTE: Always pair with the compiler behavior in StrataCompilationInfoVerticalLayeringParamBlend
 FStrataData StrataVerticalLayeringParameterBlending(FStrataData Top, FStrataData Base, uint NewNormalIndex, float NoV)
 {
@@ -1112,20 +1120,20 @@ FStrataData StrataVerticalLayeringParameterBlending(FStrataData Top, FStrataData
 	const float GreyThroughput			= dot(ColorThroughput, 1.0f / 3.0f);
 	
 	// Add both layers thickness
-	BSDF_SETTHICKNESSCM(ResultBSDF,		BSDF_GETTHICKNESSCM(BaseBSDF)					+ BSDF_GETTHICKNESSCM(TopBSDF));
+	BSDF_SETTHICKNESSCM(ResultBSDF,		BSDF_GETTHICKNESSCM(BaseBSDF)					+ SafeTopBSDFWeight * BSDF_GETTHICKNESSCM(TopBSDF));
 
 	// Now add bottom layer component weights by throughput to the top layer.
 	// Those parameters will contribute to added luminance, in this case we add them together and saturate to not go out of the safe range
 	// Some parameters are not contributing to added luminance: in this case we simply lerp them
-	BSDF_SETEMISSIVE(ResultBSDF,		ColorThroughput * BSDF_GETEMISSIVE(BaseBSDF)	+ BSDF_GETEMISSIVE(TopBSDF));
+	BSDF_SETEMISSIVE(ResultBSDF,		ColorThroughput * BSDF_GETEMISSIVE(BaseBSDF)	+ SafeTopBSDFWeight * BSDF_GETEMISSIVE(TopBSDF));
 
 	// Strata TODO: we should also add the scattering light contribution from the top layer there.
-	SLAB_DIFFUSEALBEDO(ResultBSDF) =	saturate(ColorThroughput * SLAB_DIFFUSEALBEDO(BaseBSDF)+ SLAB_DIFFUSEALBEDO(TopBSDF));
-	SLAB_F0(ResultBSDF) =				saturate(ColorThroughput * SLAB_F0(BaseBSDF)	+ SLAB_F0(TopBSDF));
-	SLAB_F90(ResultBSDF) =				saturate(ColorThroughput * SLAB_F90(BaseBSDF)	+ SLAB_F90(TopBSDF));
+	SLAB_DIFFUSEALBEDO(ResultBSDF) =	saturate(ColorThroughput * SLAB_DIFFUSEALBEDO(BaseBSDF) + SafeTopBSDFWeight * SLAB_DIFFUSEALBEDO(TopBSDF));
+	SLAB_F0(ResultBSDF) =				saturate(ColorThroughput * SLAB_F0(BaseBSDF)	+ SafeTopBSDFWeight * SLAB_F0(TopBSDF));
+	SLAB_F90(ResultBSDF) =				saturate(ColorThroughput * SLAB_F90(BaseBSDF)	+ SafeTopBSDFWeight * SLAB_F90(TopBSDF));
 
-	SLAB_FUZZ_AMOUNT(ResultBSDF) =		saturate(GreyThroughput * SLAB_FUZZ_AMOUNT(BaseBSDF)+ SLAB_FUZZ_AMOUNT(TopBSDF));
-	SLAB_FUZZ_COLOR(ResultBSDF) =		saturate(ColorThroughput * SLAB_FUZZ_COLOR(BaseBSDF)+ SLAB_FUZZ_COLOR(TopBSDF));
+	SLAB_FUZZ_AMOUNT(ResultBSDF) =		saturate(GreyThroughput * SLAB_FUZZ_AMOUNT(BaseBSDF) + SafeTopBSDFWeight * SLAB_FUZZ_AMOUNT(TopBSDF));
+	SLAB_FUZZ_COLOR(ResultBSDF) =		saturate(ColorThroughput * SLAB_FUZZ_COLOR(BaseBSDF) + SafeTopBSDFWeight * SLAB_FUZZ_COLOR(TopBSDF));
 
 	// Simple blend for now. 
 	// STRATA_TODO find someting smarter?
@@ -1133,13 +1141,13 @@ FStrataData StrataVerticalLayeringParameterBlending(FStrataData Top, FStrataData
 	SLAB_ROUGHNESSX(ResultBSDF) =		lerp(SLAB_ROUGHNESSX(TopBSDF),					SLAB_ROUGHNESSX(BaseBSDF),					GreyThroughput);
 	SLAB_HAZINESS(ResultBSDF) =			lerp(SLAB_HAZINESS(TopBSDF),					SLAB_HAZINESS(BaseBSDF),					GreyThroughput);
 
-	// No anisotropy 
+	// Anisotropy is disabled with vertical blending. 
 	SLAB_ROUGHNESSY(ResultBSDF) =		SLAB_ROUGHNESSX(ResultBSDF);
 
-	// Only keep thin film and for the top layer
+	// Only keep thin film and for the top layer.
 	SLAB_THINFILM_THICKNESS(ResultBSDF)=SLAB_THINFILM_THICKNESS(TopBSDF);
 
-	// Keep the DMFP form the bottom layer only (top layer simple volume is backed in parameters)
+	// Keep the DMFP form the bottom layer only (top layer simple volume is backed in parameters).
 	SLAB_SSSDMFP(ResultBSDF) =			SLAB_SSSDMFP(BaseBSDF);
 
 #if STRATA_INLINE_SHADING
diff --git a/Engine/Source/Runtime/Engine/Private/Materials/HLSLMaterialTranslator.cpp b/Engine/Source/Runtime/Engine/Private/Materials/HLSLMaterialTranslator.cpp
--- a/Engine/Source/Runtime/Engine/Private/Materials/HLSLMaterialTranslator.cpp
+++ b/Engine/Source/Runtime/Engine/Private/Materials/HLSLMaterialTranslator.cpp
@@ -8990,6 +8990,11 @@ int32 FHLSLMaterialTranslator::StrataAddParameterBlendingBSDFWeightToNormalMixCo
 	return AddCodeChunk(MCT_Float, TEXT("AddParameterBlendingBSDFWeightToNormalMix(%s, %s)"), *GetParameterCode(ACodeChunk), *GetParameterCode(BCodeChunk));
 }
 
+int32 FHLSLMaterialTranslator::StrataVerticalLayeringParameterBlendingBSDFWeightToNormalMixCodeChunk(int32 TopCodeChunk)
+{
+	return AddCodeChunk(MCT_Float, TEXT("VerticalLayeringParameterBlendingBSDFWeightToNormalMix(%s)"), *GetParameterCode(TopCodeChunk));
+}
+
 int32 FHLSLMaterialTranslator::StrataCreateAndRegisterNullMaterial()
 {
 	int32 OutputCodeChunk = AddInlinedCodeChunk(MCT_Strata, TEXT("GetInitialisedStrataData()"));
diff --git a/Engine/Source/Runtime/Engine/Private/Materials/HLSLMaterialTranslator.h b/Engine/Source/Runtime/Engine/Private/Materials/HLSLMaterialTranslator.h
--- a/Engine/Source/Runtime/Engine/Private/Materials/HLSLMaterialTranslator.h
+++ b/Engine/Source/Runtime/Engine/Private/Materials/HLSLMaterialTranslator.h
@@ -1012,6 +1012,7 @@ protected:
 	virtual FStrataRegisteredSharedLocalBasis StrataCompilationInfoRegisterSharedLocalBasis(int32 NormalCodeChunk, int32 TangentCodeChunk) override;
 	virtual uint8 StrataCompilationInfoGetSharedLocalBasesCount() override;
 	virtual int32 StrataAddParameterBlendingBSDFWeightToNormalMixCodeChunk(int32 ACodeChunk, int32 BCodeChunk) override;
+	virtual int32 StrataVerticalLayeringParameterBlendingBSDFWeightToNormalMixCodeChunk(int32 TopCodeChunk) override;
 
 	FStrataSharedLocalBasesInfo StrataCompilationInfoGetMatchingSharedLocalBasisInfo(const FStrataRegisteredSharedLocalBasis& SearchedSharedLocalBasis);
 
diff --git a/Engine/Source/Runtime/Engine/Private/Materials/MaterialExpressions.cpp b/Engine/Source/Runtime/Engine/Private/Materials/MaterialExpressions.cpp
--- a/Engine/Source/Runtime/Engine/Private/Materials/MaterialExpressions.cpp
+++ b/Engine/Source/Runtime/Engine/Private/Materials/MaterialExpressions.cpp
@@ -20971,8 +20971,7 @@ int32 UMaterialExpressionStrataVerticalLayering::Compile(class FMaterialCompiler
 	const FStrataMaterialCompilationInfo& TopStrataData = Compiler->GetStrataCompilationInfo(TopCodeChunk);
 	const FStrataMaterialCompilationInfo& BaseStrataData = Compiler->GetStrataCompilationInfo(BaseCodeChunk);
 
-	// STRATA_TODO: use a weight accounting for weights and transmittance?
-	const int32 NormalMixCodeChunk = Compiler->Constant(0.5f);
+	const int32 TopNormalMixCodeChunk = Compiler->StrataVerticalLayeringParameterBlendingBSDFWeightToNormalMixCodeChunk(TopCodeChunk);
 
 	int32 OutputCodeChunk = INDEX_NONE;
 	if (bUseParameterBlending)
@@ -20998,12 +20997,12 @@ int32 UMaterialExpressionStrataVerticalLayering::Compile(class FMaterialCompiler
 		}
 
 		// Compute the new Normal and Tangent resulting from the blending using code chunk
-		const int32 NewNormalCodeChunk = StrataBlendNormal(Compiler, TopBSDF.RegisteredSharedLocalBasis.NormalCodeChunk, BaseBSDF.RegisteredSharedLocalBasis.NormalCodeChunk, NormalMixCodeChunk);
+		const int32 NewNormalCodeChunk = StrataBlendNormal(Compiler, BaseBSDF.RegisteredSharedLocalBasis.NormalCodeChunk, TopBSDF.RegisteredSharedLocalBasis.NormalCodeChunk, TopNormalMixCodeChunk);
 		// The tangent is optional so we treat it differently if INDEX_NONE is specified
 		int32 NewTangentCodeChunk = INDEX_NONE;
 		if (TopBSDF.RegisteredSharedLocalBasis.TangentCodeChunk != INDEX_NONE && BaseBSDF.RegisteredSharedLocalBasis.TangentCodeChunk != INDEX_NONE)
 		{
-			NewTangentCodeChunk = StrataBlendNormal(Compiler, TopBSDF.RegisteredSharedLocalBasis.TangentCodeChunk, BaseBSDF.RegisteredSharedLocalBasis.TangentCodeChunk, NormalMixCodeChunk);
+			NewTangentCodeChunk = StrataBlendNormal(Compiler, BaseBSDF.RegisteredSharedLocalBasis.TangentCodeChunk, TopBSDF.RegisteredSharedLocalBasis.TangentCodeChunk, TopNormalMixCodeChunk);
 		}
 		else if (TopBSDF.RegisteredSharedLocalBasis.TangentCodeChunk != INDEX_NONE)
 		{
diff --git a/Engine/Source/Runtime/Engine/Public/MaterialCompiler.h b/Engine/Source/Runtime/Engine/Public/MaterialCompiler.h
--- a/Engine/Source/Runtime/Engine/Public/MaterialCompiler.h
+++ b/Engine/Source/Runtime/Engine/Public/MaterialCompiler.h
@@ -459,6 +459,7 @@ public:
 		return FString::Printf(TEXT("SHAREDLOCALBASIS_INDEX_%u"), SharedLocalBasis.GraphSharedLocalBasisIndex);
 	}
 	virtual int32 StrataAddParameterBlendingBSDFWeightToNormalMixCodeChunk(int32 ACodeChunk, int32 BCodeChunk) = 0;
+	virtual int32 StrataVerticalLayeringParameterBlendingBSDFWeightToNormalMixCodeChunk(int32 TopCodeChunk) = 0;
 
 	// Water
 	virtual int32 SceneDepthWithoutWater(int32 Offset, int32 ViewportUV, bool bUseOffset, float FallbackDepth) = 0;
@@ -1031,6 +1032,11 @@ public:
 	{
 		return Compiler->StrataAddParameterBlendingBSDFWeightToNormalMixCodeChunk(ACodeChunk, BCodeChunk);
 	}
+
+	virtual int32 StrataVerticalLayeringParameterBlendingBSDFWeightToNormalMixCodeChunk(int32 TopCodeChunk) override
+	{
+		return Compiler->StrataVerticalLayeringParameterBlendingBSDFWeightToNormalMixCodeChunk(TopCodeChunk);
+	}
 	
 	virtual const FStrataMaterialCompilationInfo& GetStrataCompilationInfo(int32 CodeChunk) override
 	{
