Path Tracer: clamp path contributions for indirect rays only

Avoid clamping camera hits as these usually have low variance. This avoids effects like bloom and lens flares from dissapearing entirely when clamping is too aggressive.

diff --git a/Engine/Shaders/Private/PathTracing/PathTracing.usf b/Engine/Shaders/Private/PathTracing/PathTracing.usf
--- a/Engine/Shaders/Private/PathTracing/PathTracing.usf
+++ b/Engine/Shaders/Private/PathTracing/PathTracing.usf
@@ -54,12 +54,12 @@ int3 TileOffset;
 uint SceneVisibleLightCount;
 
 
-void AccumulateRadiance(inout float3 TotalRadiance, float3 PathRadiance)
+void AccumulateRadiance(inout float3 TotalRadiance, float3 PathRadiance, bool bIsCameraRay)
 {
 	// User asked for path contributions to be clamped to reduce fireflies.
 	// This puts an upper bound on variance within the pixel at the expense of bias
 	float MaxPathRadiance = max3(PathRadiance.x, PathRadiance.y, PathRadiance.z);
-	if (MaxPathRadiance > MaxPathIntensity)
+	if (!bIsCameraRay && MaxPathRadiance > MaxPathIntensity)
 	{
 		// adjust sample while keeping color
 		PathRadiance *= MaxPathIntensity / MaxPathRadiance;
@@ -69,14 +69,14 @@ void AccumulateRadiance(inout float3 TotalRadiance, float3 PathRadiance)
 	TotalRadiance += PathRadiance;
 }
 
-FMaterialClosestHitPayload TraceTransparentRay(RayDesc Ray, FRayCone RayCone, bool IsCameraRay, bool LastBounce, bool IncludeEmission, uint2 LaunchIndex, uint NumLights, inout RandomSequence RandSequence, inout float3 PathThroughput, inout float3 Radiance)
+FMaterialClosestHitPayload TraceTransparentRay(RayDesc Ray, FRayCone RayCone, bool bIsCameraRay, bool bLastBounce, bool bIncludeEmission, uint2 LaunchIndex, uint NumLights, inout RandomSequence RandSequence, inout float3 PathThroughput, inout float3 Radiance)
 {
-	const uint RayFlags = IsCameraRay && EnableCameraBackfaceCulling ? RAY_FLAG_CULL_BACK_FACING_TRIANGLES : 0;
+	const uint RayFlags = bIsCameraRay && EnableCameraBackfaceCulling ? RAY_FLAG_CULL_BACK_FACING_TRIANGLES : 0;
 	const uint MissShaderIndex = 0;
 	float SelectionWeightSum = 0;
 	float3 PayloadThroughput = PathThroughput;
 	FMaterialClosestHitPayload Payload;
-	if (!IncludeEmission && LastBounce)
+	if (!bIncludeEmission && bLastBounce)
 	{
 		Payload.SetMiss();
 		PathThroughput = 0;
@@ -88,7 +88,7 @@ FMaterialClosestHitPayload TraceTransparentRay(RayDesc Ray, FRayCone RayCone, bo
 		// Set payload inputs
 		PackedPayload.SetRayCone(RayCone);
 		PackedPayload.SetPixelCoord(LaunchIndex);
-		if (!IsCameraRay)
+		if (!bIsCameraRay)
 		{
 			// avoid double counting what was captured by the skylight
 			// also avoid noise from hot spots (they can be properly
@@ -113,7 +113,7 @@ FMaterialClosestHitPayload TraceTransparentRay(RayDesc Ray, FRayCone RayCone, bo
 			RayDesc LightRay = Ray;
 			LightRay.TMax = PackedPayload.IsMiss() ? Ray.TMax : PackedPayload.HitT;
 			float3 LightRadiance = TraceLight(LightRay, LightId).Radiance;
-			AccumulateRadiance(Radiance, PathThroughput * LightRadiance);
+			AccumulateRadiance(Radiance, PathThroughput * LightRadiance, bIsCameraRay);
 		}
 
 		if (PackedPayload.IsMiss())
@@ -133,14 +133,14 @@ FMaterialClosestHitPayload TraceTransparentRay(RayDesc Ray, FRayCone RayCone, bo
 		}
 
 		// add in surface emission (except for modulate blend mode which uses Radiance as the transparency multiplier)
-		if (IncludeEmission && HitPayload.BlendingMode != RAY_TRACING_BLEND_MODE_MODULATE)
+		if (bIncludeEmission && HitPayload.BlendingMode != RAY_TRACING_BLEND_MODE_MODULATE)
 		{
-			AccumulateRadiance(Radiance, PathThroughput * HitPayload.Radiance);
+			AccumulateRadiance(Radiance, PathThroughput * HitPayload.Radiance, bIsCameraRay);
 		}
 
 		float3 Transparency = GetMaterialTransparency(HitPayload, Ray.Direction);
 
-		if (!LastBounce)
+		if (!bLastBounce)
 		{
 			float3 Contrib = PathThroughput * EstimateMaterialAlbedo(HitPayload);
 
@@ -661,7 +661,7 @@ RAY_TRACING_ENTRY_RAYGEN(PathTracingMainRG)
 						{
 							LightContrib *= MISWeightRobust(LightSample.Pdf, MaterialEval.Pdf);
 						}
-						AccumulateRadiance(Radiance, LightContrib);
+						AccumulateRadiance(Radiance, LightContrib, bIsCameraRay);
 					}
 				}
 			}
@@ -770,7 +770,7 @@ RAY_TRACING_ENTRY_RAYGEN(PathTracingMainRG)
 						LightContrib *= TraceTransparentVisibilityRay(LightRay, LaunchIndex, PathRoughness);
 					}
 					// the light made some contribution, and there was nothing along the shadow ray
-					AccumulateRadiance(Radiance, LightContrib);
+					AccumulateRadiance(Radiance, LightContrib, bIsCameraRay);
 				}
 			}
 			
