Add Path Tracer specific material permutation

This permutation will allow the path tracer to evolve independently from other RT passes in the engine. First end-user visible feature is that we disable RayTracingQualitySwitch for this permutation (so that path tracing can see the "full" material quality).

Introduce r.PathTracing cvar to guard the feature so that projects that don't need this functionality don't have to pay for the extra permutation.

Move raytracing clip/ddx/ddy into Common.ush

diff --git a/Engine/Plugins/Experimental/GPULightmass/Source/GPULightmass/Private/LightmapRayTracing.cpp b/Engine/Plugins/Experimental/GPULightmass/Source/GPULightmass/Private/LightmapRayTracing.cpp
--- a/Engine/Plugins/Experimental/GPULightmass/Source/GPULightmass/Private/LightmapRayTracing.cpp
+++ b/Engine/Plugins/Experimental/GPULightmass/Source/GPULightmass/Private/LightmapRayTracing.cpp
@@ -19,72 +19,21 @@ bool IsSupportedVertexFactoryType(const FVertexFactoryType* VertexFactoryType)
 
 #if RHI_RAYTRACING
 
-class FLightmapMaterialCHS : public FMeshMaterialShader, public FUniformLightMapPolicyShaderParametersType
-{
-	DECLARE_INLINE_TYPE_LAYOUT_EXPLICIT_BASES(FLightmapMaterialCHS, NonVirtual, FMeshMaterialShader, FUniformLightMapPolicyShaderParametersType);
-public:
-	FLightmapMaterialCHS(const FMeshMaterialShaderType::CompiledShaderInitializerType& Initializer)
-		: FMeshMaterialShader(Initializer)
-	{
-		PassUniformBuffer.Bind(Initializer.ParameterMap, FSceneTextureUniformParameters::StaticStructMetadata.GetShaderVariableName());
-		FUniformLightMapPolicyShaderParametersType::Bind(Initializer.ParameterMap);
-	}
-
-	FLightmapMaterialCHS() {}
-
-	void GetShaderBindings(
-		const FScene* Scene,
-		ERHIFeatureLevel::Type FeatureLevel,
-		const FPrimitiveSceneProxy* PrimitiveSceneProxy,
-		const FMaterialRenderProxy& MaterialRenderProxy,
-		const FMaterial& Material,
-		const FMeshPassProcessorRenderState& DrawRenderState,
-		const TBasePassShaderElementData<FUniformLightMapPolicy>& ShaderElementData,
-		FMeshDrawSingleShaderBindings& ShaderBindings) const
-	{
-		FMeshMaterialShader::GetShaderBindings(Scene, FeatureLevel, PrimitiveSceneProxy, MaterialRenderProxy, Material, DrawRenderState, ShaderElementData, ShaderBindings);
-
-		FUniformLightMapPolicy::GetPixelShaderBindings(
-			PrimitiveSceneProxy,
-			ShaderElementData.LightMapPolicyElementData,
-			this,
-			ShaderBindings);
-	}
-
-	void GetElementShaderBindings(
-		const FShaderMapPointerTable& PointerTable,
-		const FScene* Scene,
-		const FSceneView* ViewIfDynamicMeshCommand,
-		const FVertexFactory* VertexFactory,
-		const EVertexInputStreamType InputStreamType,
-		ERHIFeatureLevel::Type FeatureLevel,
-		const FPrimitiveSceneProxy* PrimitiveSceneProxy,
-		const FMeshBatch& MeshBatch,
-		const FMeshBatchElement& BatchElement,
-		const TBasePassShaderElementData<FUniformLightMapPolicy>& ShaderElementData,
-		FMeshDrawSingleShaderBindings& ShaderBindings,
-		FVertexInputStreamArray& VertexStreams) const
-	{
-		FMeshMaterialShader::GetElementShaderBindings(PointerTable, Scene, ViewIfDynamicMeshCommand, VertexFactory, InputStreamType, FeatureLevel, PrimitiveSceneProxy, MeshBatch, BatchElement, ShaderElementData, ShaderBindings, VertexStreams);
-	}
-};
-
 template<bool UseAnyHitShader>
-class TLightmapMaterialCHS : public FLightmapMaterialCHS
+class TLightmapMaterial : public FMeshMaterialShader
 {
-	DECLARE_SHADER_TYPE(TLightmapMaterialCHS, MeshMaterial);
+	DECLARE_SHADER_TYPE(TLightmapMaterial, MeshMaterial);
 public:
+	TLightmapMaterial() = default;
 
-	TLightmapMaterialCHS(const FMeshMaterialShaderType::CompiledShaderInitializerType& Initializer)
-		: FLightmapMaterialCHS(Initializer)
+	TLightmapMaterial(const FMeshMaterialShaderType::CompiledShaderInitializerType& Initializer)
+		: FMeshMaterialShader(Initializer)
 	{}
 
-	TLightmapMaterialCHS() {}
-
 	static bool ShouldCompilePermutation(const FMeshMaterialShaderPermutationParameters& Parameters)
 	{
 		return IsSupportedVertexFactoryType(Parameters.VertexFactoryType)
-			&& ((Parameters.MaterialParameters.bIsMasked || Parameters.MaterialParameters.BlendMode == BLEND_Translucent) == UseAnyHitShader)
+			&& ((Parameters.MaterialParameters.bIsMasked || Parameters.MaterialParameters.BlendMode != BLEND_Opaque) == UseAnyHitShader)
 			&& FNoLightMapPolicy::ShouldCompilePermutation(Parameters)
 			&& ShouldCompileRayTracingShadersForProject(Parameters.Platform);
 	}
@@ -96,7 +45,6 @@ public:
 		OutEnvironment.SetDefine(TEXT("USE_RAYTRACED_TEXTURE_RAYCONE_LOD"), 0);
 		OutEnvironment.SetDefine(TEXT("SCENE_TEXTURES_DISABLED"), 1);
 		OutEnvironment.SetDefine(TEXT("SIMPLIFIED_MATERIAL_SHADER"), 1);
-		FNoLightMapPolicy::ModifyCompilationEnvironment(Parameters, OutEnvironment);
 		FMeshMaterialShader::ModifyCompilationEnvironment(Parameters, OutEnvironment);
 	}
 
@@ -123,8 +71,11 @@ public:
 	}
 };
 
-IMPLEMENT_MATERIAL_SHADER_TYPE(template<>, TLightmapMaterialCHS<false>, TEXT("/Engine/Private/RayTracing/RayTracingMaterialHitShaders.usf"), TEXT("closesthit=MaterialCHS"), SF_RayHitGroup);
-IMPLEMENT_MATERIAL_SHADER_TYPE(template<>, TLightmapMaterialCHS<true>, TEXT("/Engine/Private/RayTracing/RayTracingMaterialHitShaders.usf"), TEXT("closesthit=MaterialCHS anyhit=MaterialAHS"), SF_RayHitGroup);
+using FLightmapMaterialCHS     = TLightmapMaterial<false>;
+using FLightmapMaterialCHS_AHS = TLightmapMaterial<true>;
+
+IMPLEMENT_MATERIAL_SHADER_TYPE(template <>, FLightmapMaterialCHS    , TEXT("/Engine/Private/PathTracing/PathTracingMaterialHitShader.usf"), TEXT("closesthit=PathTracingMaterialCHS"), SF_RayHitGroup);
+IMPLEMENT_MATERIAL_SHADER_TYPE(template <>, FLightmapMaterialCHS_AHS, TEXT("/Engine/Private/PathTracing/PathTracingMaterialHitShader.usf"), TEXT("closesthit=PathTracingMaterialCHS anyhit=PathTracingMaterialAHS"), SF_RayHitGroup);
 
 IMPLEMENT_GLOBAL_SHADER(FLightmapPathTracingRGS, "/Plugin/GPULightmass/Private/LightmapPathTracing.usf", "LightmapPathTracingMainRG", SF_RayGen);
 IMPLEMENT_GLOBAL_SHADER(FVolumetricLightmapPathTracingRGS, "/Plugin/GPULightmass/Private/LightmapPathTracing.usf", "VolumetricLightmapPathTracingMainRG", SF_RayGen);
@@ -132,15 +83,14 @@ IMPLEMENT_GLOBAL_SHADER(FStationaryLightShadowTracingRGS, "/Plugin/GPULightmass/
 
 IMPLEMENT_GLOBAL_SHADER(FFirstBounceRayGuidingCDFBuildCS, "/Plugin/GPULightmass/Private/FirstBounceRayGuidingCDFBuild.usf", "FirstBounceRayGuidingCDFBuildCS", SF_Compute);
 
+// TODO: unify this with the PathTracing process call so we can remove the FLightmapRayTracingMeshProcessor class entirely and de-virtualize this function call
 bool FLightmapRayTracingMeshProcessor::Process(
 	const FMeshBatch& RESTRICT MeshBatch,
 	uint64 BatchElementMask,
 	const FPrimitiveSceneProxy* RESTRICT PrimitiveSceneProxy,
 	const FMaterialRenderProxy& RESTRICT MaterialRenderProxy,
 	const FMaterial& RESTRICT MaterialResource,
-	FMaterialShadingModelField ShadingModels,
-	const FUniformLightMapPolicy& RESTRICT LightMapPolicy,
-	const typename FUniformLightMapPolicy::ElementDataType& RESTRICT LightMapElementData)
+	const FUniformLightMapPolicy& RESTRICT LightMapPolicy)
 {
 	const FVertexFactory* VertexFactory = MeshBatch.VertexFactory;
 
@@ -149,17 +99,17 @@ bool FLightmapRayTracingMeshProcessor::Process(
 		FMeshMaterialShader,
 		FMeshMaterialShader,
 		FMeshMaterialShader,
-		FLightmapMaterialCHS> RayTracingShaders;
+		FMeshMaterialShader> RayTracingShaders;
 
 	FMaterialShaderTypes ShaderTypes;
 
 	if (MaterialResource.IsMasked() || MaterialResource.GetBlendMode() != BLEND_Opaque)
 	{
-		ShaderTypes.AddShaderType<TLightmapMaterialCHS<true>>();
+		ShaderTypes.AddShaderType<FLightmapMaterialCHS_AHS>();
 	}
 	else
 	{
-		ShaderTypes.AddShaderType<TLightmapMaterialCHS<false>>();
+		ShaderTypes.AddShaderType<FLightmapMaterialCHS>();
 	}
 
 	FMaterialShaders Shaders;
@@ -170,10 +120,7 @@ bool FLightmapRayTracingMeshProcessor::Process(
 
 	check(Shaders.TryGetShader(SF_RayHitGroup, RayTracingShaders.RayHitGroupShader));
 
-	PassDrawRenderState.SetBlendState(TStaticBlendState<CW_RGBA, BO_Add, BF_One, BF_One, BO_Add, BF_Zero, BF_One>::GetRHI());
-	PassDrawRenderState.SetDepthStencilState(TStaticDepthStencilState<false, CF_DepthNearOrEqual>::GetRHI());
-
-	TBasePassShaderElementData<FUniformLightMapPolicy> ShaderElementData(LightMapElementData);
+	TBasePassShaderElementData<FUniformLightMapPolicy> ShaderElementData(MeshBatch.LCI);
 	ShaderElementData.InitializeMeshMaterialData(ViewIfDynamicMeshCommand, PrimitiveSceneProxy, MeshBatch, -1, true);
 
 	BuildRayTracingMeshCommands(
diff --git a/Engine/Plugins/Experimental/GPULightmass/Source/GPULightmass/Private/LightmapRayTracing.h b/Engine/Plugins/Experimental/GPULightmass/Source/GPULightmass/Private/LightmapRayTracing.h
--- a/Engine/Plugins/Experimental/GPULightmass/Source/GPULightmass/Private/LightmapRayTracing.h
+++ b/Engine/Plugins/Experimental/GPULightmass/Source/GPULightmass/Private/LightmapRayTracing.h
@@ -27,9 +27,7 @@ public:
 		const FPrimitiveSceneProxy* RESTRICT PrimitiveSceneProxy,
 		const FMaterialRenderProxy& RESTRICT MaterialRenderProxy,
 		const FMaterial& RESTRICT MaterialResource,
-		FMaterialShadingModelField ShadingModels,
-		const FUniformLightMapPolicy& RESTRICT LightMapPolicy,
-		const typename FUniformLightMapPolicy::ElementDataType& RESTRICT LightMapElementData) override;
+		const FUniformLightMapPolicy& RESTRICT LightMapPolicy) override;
 };
 
 class FLightmapPathTracingRGS : public FGlobalShader
diff --git a/Engine/Shaders/Private/Common.ush b/Engine/Shaders/Private/Common.ush
--- a/Engine/Shaders/Private/Common.ush
+++ b/Engine/Shaders/Private/Common.ush
@@ -117,6 +117,17 @@ const static float3 Max111110BitsFloat3 = float3(Max11BitsFloat, Max11BitsFloat,
 	#define FORCED_TEXTURE_MIP 0.0f
 #endif
 
+#if !PIXELSHADER
+
+// These built-ins are only valid in pixel shaders.
+// Define dummy versions so that ray-tracing materials will at least compile
+#define clip(x)
+#define ddx(x) 0
+#define ddy(x) 0
+#define fwidth(x) 0
+
+#endif
+
 #ifndef USE_RAYTRACED_TEXTURE_RAYCONE_LOD
 #define USE_RAYTRACED_TEXTURE_RAYCONE_LOD (RAYHITGROUPSHADER)
 #endif // USE_RAYTRACED_TEXTURE_RAYCONE_LOD
@@ -1586,7 +1597,7 @@ bool GetGIReplaceState()
 
 bool GetRayTracingQualitySwitch()
 {
-#if RAYHITGROUPSHADER || LUMEN_CARD_CAPTURE 
+#if (RAYHITGROUPSHADER || LUMEN_CARD_CAPTURE) && !PATH_TRACING
 	return true;
 #else
 	return false;
diff --git a/Engine/Shaders/Private/MaterialTemplate.ush b/Engine/Shaders/Private/MaterialTemplate.ush
--- a/Engine/Shaders/Private/MaterialTemplate.ush
+++ b/Engine/Shaders/Private/MaterialTemplate.ush
@@ -2778,6 +2778,19 @@ void GetMaterialClippingVelocity(FMaterialPixelParameters Parameters, FPixelMate
 	#endif
 }
 
+#if MATERIAL_DITHER_OPACITY_MASK
+float DitheredOpacityMaskToOpacity(float Mask)
+{
+	// This represents the expected value of the function GetMaterialCoverageAndClipping
+	// which randomly dithers the fragment on or off to produce the effect of opacity
+
+	// The expected value of this dithering can be computed as:
+	//    Simplify[Integrate[If[Mask + Dither - 1/2 < 0, 0, 1], {Dither, 0, 1}]]
+	// which is just:
+
+	return saturate(Mask + 0.5);
+}
+#endif
 
 void GetMaterialCoverageAndClipping(FMaterialPixelParameters Parameters, FPixelMaterialInputs PixelMaterialInputs)
 {
diff --git a/Engine/Shaders/Private/PathTracing/PathTracingMaterialHitShader.usf b/Engine/Shaders/Private/PathTracing/PathTracingMaterialHitShader.usf
new file mode 100644
--- /dev/null
+++ b/Engine/Shaders/Private/PathTracing/PathTracingMaterialHitShader.usf
@@ -0,0 +1,451 @@
+// Copyright Epic Games, Inc. All Rights Reserved.
+
+#ifndef USE_MATERIAL_CLOSEST_HIT_SHADER
+	#define USE_MATERIAL_CLOSEST_HIT_SHADER 1
+#endif
+
+#ifndef USE_MATERIAL_ANY_HIT_SHADER
+	#define USE_MATERIAL_ANY_HIT_SHADER 1
+#endif
+
+#define PATH_TRACING 1
+#define ENABLE_SKY_LIGHT 1      
+#define NEEDS_LIGHTMAP_COORDINATE 0
+#ifdef NEEDS_VERTEX_FACTORY_INTERPOLATION
+#undef NEEDS_VERTEX_FACTORY_INTERPOLATION
+#endif
+// Needed for VertexFactoryInterpolate to interpolate attributes from vertices to hit point
+#define NEEDS_VERTEX_FACTORY_INTERPOLATION 1
+
+#include "/Engine/Private/Common.ush"
+#include "/Engine/Private/RayTracing/RayTracingCommon.ush"
+#include "/Engine/Private/RayTracing/RayTracingHitGroupCommon.ush"
+
+#include "/Engine/Generated/Material.ush"   
+#include "/Engine/Generated/VertexFactory.ush"
+
+#include "/Engine/Private/RayTracing/RayTracingCalcInterpolants.ush"
+#include "/Engine/Private/ShadingCommon.ush"
+#include "/Engine/Private/DeferredShadingCommon.ush"
+#include "/Engine/Private/PathTracing/Material/PathTracingFresnel.ush"
+#include "/Engine/Private/PathTracing/Material/PathTracingThinGlass.ush"
+
+
+RAY_TRACING_ENTRY_CLOSEST_HIT(PathTracingMaterialCHS,
+	FPackedMaterialClosestHitPayload, PackedPayload,
+	FDefaultAttributes, Attributes)
+{
+	PackedPayload.HitT = RayTCurrent();
+
+#if USE_MATERIAL_CLOSEST_HIT_SHADER
+	GlobalTextureMipBias = PackedPayload.GetMipBias();
+
+	const uint2 PixelCoord = PackedPayload.GetPixelCoord();
+
+	ResolvedView = ResolveView();
+
+	FVertexFactoryInterpolantsVSToPS Interpolants;
+	float4 SvPosition = float4(PixelCoord.xy, 0.0, 1.0);
+
+	FRayCone PropagatedCone = PropagateRayCone(PackedPayload.GetRayCone(), 0 /* surface curvature */, RayTCurrent());
+	CalcInterpolants(PropagatedCone, Attributes, Interpolants);
+
+	FMaterialPixelParameters MaterialParameters = GetMaterialPixelParameters(Interpolants, SvPosition);
+	FPixelMaterialInputs PixelMaterialInputs;
+
+	{
+		float4 ScreenPosition = SvPositionToResolvedScreenPosition(SvPosition);
+		float3 TranslatedWorldPosition = WorldRayOrigin() + WorldRayDirection() * RayTCurrent() + ResolvedView.PreViewTranslation.xyz;
+
+		bool bIsFrontFace = HitKind() == HIT_KIND_TRIANGLE_FRONT_FACE;
+
+		// #dxr_todo: UE-72130 support world position offset 
+		// #if USE_WORLD_POSITION_EXCLUDING_SHADER_OFFSETS
+		// CalcMaterialParametersEx(MaterialParameters, PixelMaterialInputs, SvPosition, ScreenPosition, bIsFrontFace, TranslatedWorldPosition, BasePassInterpolants.PixelPositionExcludingWPO);
+		// #else
+		CalcMaterialParametersEx(MaterialParameters, PixelMaterialInputs, SvPosition, ScreenPosition, bIsFrontFace, TranslatedWorldPosition, TranslatedWorldPosition);
+		// #endif
+	}
+
+	FMaterialClosestHitPayload Payload = (FMaterialClosestHitPayload)0;
+
+	/**
+	 * Set common material attributes for both full and simplified materials
+	 **/
+	Payload.ShadingModelID = GetMaterialShadingModel(PixelMaterialInputs);
+
+#if MATERIALBLENDING_ALPHACOMPOSITE
+	uint MaterialBlendingMode = RAY_TRACING_BLEND_MODE_ALPHA_COMPOSITE;
+#elif MATERIALBLENDING_ALPHAHOLDOUT
+	uint MaterialBlendingMode = RAY_TRACING_BLEND_MODE_ALPHA_HOLDOUT;
+#elif MATERIALBLENDING_TRANSLUCENT
+	uint MaterialBlendingMode = RAY_TRACING_BLEND_MODE_TRANSLUCENT;
+#elif MATERIALBLENDING_ADDITIVE
+	uint MaterialBlendingMode = RAY_TRACING_BLEND_MODE_ADDITIVE;
+#elif MATERIALBLENDING_MODULATE
+	uint MaterialBlendingMode = RAY_TRACING_BLEND_MODE_MODULATE;
+#else
+	uint MaterialBlendingMode = RAY_TRACING_BLEND_MODE_OPAQUE;
+#endif
+
+	Payload.BlendingMode = MaterialBlendingMode;
+	Payload.PrimitiveLightingChannelMask = GetPrimitiveData(MaterialParameters.PrimitiveId).LightingChannelMask;
+
+	Payload.HitT = RayTCurrent();
+	if (HitKind() == HIT_KIND_TRIANGLE_FRONT_FACE)
+	{
+		Payload.SetFrontFace();
+	}
+
+#if MATERIAL_IS_SKY
+	if (PackedPayload.IsIgnoreSkyMaterials())
+	{
+		PackedPayload = PackRayTracingPayload(Payload, PropagatedCone);
+		return;
+	}
+#endif
+
+	// Store the results in local variables and reuse instead of calling the functions multiple times.
+	half3 BaseColor = GetMaterialBaseColor(PixelMaterialInputs);
+#if MATERIALBLENDING_TRANSLUCENT && REFRACTION_USE_INDEX_OF_REFRACTION
+	half2 Refraction = GetMaterialRefraction(PixelMaterialInputs);
+#else
+	half2 Refraction = 1;
+#endif
+	half  Metallic = GetMaterialMetallic(PixelMaterialInputs);
+	half  Specular = GetMaterialSpecular(PixelMaterialInputs);
+	half Roughness = GetMaterialRoughness(PixelMaterialInputs);
+
+#if !SIMPLIFIED_MATERIAL_SHADER
+
+	/**
+	 * Set material attributes for full materials
+	 **/
+	Payload.WorldPos = MaterialParameters.AbsoluteWorldPosition;
+	Payload.WorldNormal = normalize(MaterialParameters.WorldNormal);
+	Payload.Radiance = GetMaterialEmissive(PixelMaterialInputs);
+	Payload.BaseColor = BaseColor;
+	Payload.Specular = Specular;
+	Payload.Roughness = Roughness;
+	Payload.Metallic = Metallic;
+	Payload.Opacity = GetMaterialOpacity(PixelMaterialInputs);
+
+#if MATERIALBLENDING_MASKED
+#if MATERIAL_DITHER_OPACITY_MASK
+	// dithering emulates real transparency, so switch to translucent
+	Payload.Opacity = DitheredOpacityMaskToOpacity(GetMaterialMask(PixelMaterialInputs));
+	Payload.BlendingMode = RAY_TRACING_BLEND_MODE_TRANSLUCENT;
+#else
+	// regular masked mode - binary decision
+	if (GetMaterialMask(PixelMaterialInputs) < 0)
+	{
+		Payload.Opacity = 0;
+	}
+#endif
+#endif // MATERIALBLENDING_MASKED
+
+	Payload.Ior = Refraction.x == 0 ? 1.0 : Refraction.x;
+	Payload.CustomData = float4(GetMaterialCustomData0(MaterialParameters), GetMaterialCustomData1(MaterialParameters), 0, 0);
+
+#if MATERIAL_SHADINGMODEL_CLEAR_COAT
+#if CLEAR_COAT_BOTTOM_NORMAL
+	if (Payload.ShadingModelID == SHADINGMODELID_CLEAR_COAT)
+	{
+		float2 oct2 = UnitVectorToOctahedron(Payload.WorldNormal);
+
+#if NUM_MATERIAL_OUTPUTS_CLEARCOATBOTTOMNORMAL > 0
+#if MATERIAL_TANGENTSPACENORMAL
+		float3 tempnormal = normalize(TransformTangentVectorToWorld(MaterialParameters.TangentToWorld, ClearCoatBottomNormal0(MaterialParameters)));
+#else
+		float3 tempnormal = ClearCoatBottomNormal0(MaterialParameters);
+#endif
+
+		float2 oct1 = UnitVectorToOctahedron(tempnormal);
+		float2 oct3 = ((oct1 - oct2) * 0.5) + (128.0 / 255.0);
+		Payload.CustomData.a = oct3.x;
+		Payload.CustomData.z = oct3.y;
+#else
+		Payload.CustomData.a = 128.0 / 255.0;
+		Payload.CustomData.z = 128.0 / 255.0;
+#endif
+	}
+#endif
+#endif
+
+	// Override custom data if sub-surface material
+#if MATERIAL_SHADINGMODEL_SUBSURFACE || MATERIAL_SHADINGMODEL_TWOSIDED_FOLIAGE || MATERIAL_SHADINGMODEL_CLOTH || MATERIAL_SHADINGMODEL_CLOTH || MATERIAL_SHADINGMODEL_SUBSURFACE_PROFILE || MATERIAL_SHADINGMODEL_EYE
+	if (Payload.ShadingModelID == SHADINGMODELID_SUBSURFACE ||
+		Payload.ShadingModelID == SHADINGMODELID_TWOSIDED_FOLIAGE ||
+		Payload.ShadingModelID == SHADINGMODELID_CLOTH ||
+		Payload.ShadingModelID == SHADINGMODELID_SUBSURFACE_PROFILE |
+		Payload.ShadingModelID == SHADINGMODELID_EYE)
+	{
+		Payload.CustomData = GetMaterialSubsurfaceData(PixelMaterialInputs);
+	}
+#endif
+
+#if MATERIAL_SHADINGMODEL_CLOTH
+	if (Payload.ShadingModelID == SHADINGMODELID_CLOTH)
+	{
+		Payload.CustomData.a = saturate(GetMaterialCustomData0(MaterialParameters));
+	}
+#endif
+
+#if MATERIAL_SHADINGMODEL_THIN_TRANSLUCENT
+	if (Payload.ShadingModelID == SHADINGMODELID_THIN_TRANSLUCENT)
+	{
+		Payload.CustomData.xyz = GetThinTranslucentMaterialOutput0(MaterialParameters);
+	}
+#endif
+
+#if MATERIAL_USES_ANISOTROPY
+	Payload.WorldTangent = CalculateAnisotropyTangent(MaterialParameters, PixelMaterialInputs);
+	Payload.Anisotropy = GetMaterialAnisotropy(PixelMaterialInputs);
+#else
+	Payload.WorldTangent = 0;
+	Payload.Anisotropy = 0.0f;
+#endif	
+
+	Payload.DiffuseColor = BaseColor - BaseColor * Metallic;
+	Payload.SpecularColor = ComputeF0(Specular, BaseColor, Metallic);
+
+#else
+
+	/**
+	 * Set material attributes for simplified materials
+	 **/
+	 // Anything unsupported will be forced to default lit
+#if MATERIAL_SHADINGMODEL_TWOSIDED_FOLIAGE
+	if (Payload.ShadingModelID == SHADINGMODELID_TWOSIDED_FOLIAGE)
+	{
+		Payload.CustomData = GetMaterialSubsurfaceData(PixelMaterialInputs);
+	}
+	else
+#endif
+#if MATERIAL_SHADINGMODEL_THIN_TRANSLUCENT
+		if (Payload.ShadingModelID == SHADINGMODELID_THIN_TRANSLUCENT)
+		{
+			Payload.CustomData.xyz = GetThinTranslucentMaterialOutput0(MaterialParameters);
+		}
+		else
+#endif
+#if MATERIAL_SHADINGMODEL_UNLIT
+			if (Payload.ShadingModelID == SHADINGMODELID_UNLIT)
+			{
+			}
+			else
+#endif
+			{
+				Payload.ShadingModelID = SHADINGMODELID_DEFAULT_LIT;
+			}
+	Payload.WorldPos = MaterialParameters.AbsoluteWorldPosition;
+	Payload.WorldNormal = normalize(MaterialParameters.WorldNormal);
+	Payload.Radiance = GetMaterialEmissive(PixelMaterialInputs);
+	Payload.BaseColor = BaseColor;
+	Payload.Specular = 0;
+	Payload.Roughness = 1;
+	Payload.Metallic = 0;
+	Payload.Opacity = GetMaterialOpacity(PixelMaterialInputs);
+	Payload.Ior = 1;
+
+#endif
+
+	PackedPayload = PackRayTracingPayload(Payload, PropagatedCone);
+#endif // USE_MATERIAL_CLOSEST_HIT_SHADER
+}
+
+RAY_TRACING_ENTRY_ANY_HIT(PathTracingMaterialAHS,
+	FPackedMaterialClosestHitPayload, PackedPayload,
+	FDefaultAttributes, Attributes)
+{
+#if USE_MATERIAL_ANY_HIT_SHADER
+
+#if MATERIALBLENDING_SOLID
+	// Continue traversal by simply returning from this shader
+	// Usually this path is inactive because opaque materials don't have AHS bound
+	return;
+#elif MATERIALBLENDING_ALPHAHOLDOUT || MATERIALBLENDING_ADDITIVE
+	if (PackedPayload.IsPathTracingVisibilityRay())
+	{
+		// path tracer wants to treat these cases as fully transparent
+		// no material evaluation is required
+		IgnoreHit();
+	}
+	// otherwise accept hit
+#else
+	// All other blending modes _might_ need material evaluation - so we need
+	// to generate some code
+
+#if MATERIALBLENDING_MASKED && !MATERIAL_DITHER_OPACITY_MASK
+	// NOTE: Masked mode execute always - regardless of ray flags
+	// This ensures that both material rays and shadow rays "see" the same thing
+#else
+	// For any other blending models, we have a few options:
+	if (PackedPayload.IsIgnoreTranslucentMaterials())
+	{
+		// special mode used by regular RT shadows - translucent materials are skipped, everything else is treated as opaque
+#if MATERIALBLENDING_TRANSLUCENT || (MATERIALBLENDING_MASKED && MATERIAL_DITHER_OPACITY_MASK)
+		IgnoreHit();
+#endif
+		return;
+	}
+	else if (PackedPayload.IsPathTracingVisibilityRay())
+	{
+		// this case is used by the path tracer only currently
+		// We want to evaluate the material opacity and accumulate it
+		// Keep going and evaluate the material
+	}
+	else
+	{
+		// non-opaque blending mode but no special behavior was requested
+		return;
+	}
+#endif  // !MATERIALBLENDING_MASKED
+
+	// If we got here -- we have a reason to evaluate the material -- do the setup 
+
+	GlobalTextureMipBias = PackedPayload.GetMipBias();
+
+	const uint2 PixelCoord = PackedPayload.GetPixelCoord();
+
+	ResolvedView = ResolveView();
+
+	FVertexFactoryInterpolantsVSToPS Interpolants;
+	float4 SvPosition = float4(PixelCoord.xy, 0.0, 1.0);
+
+	CalcInterpolants(PackedPayload.GetRayCone(), Attributes, Interpolants);
+
+	FMaterialPixelParameters MaterialParameters = GetMaterialPixelParameters(Interpolants, SvPosition);
+	FPixelMaterialInputs PixelMaterialInputs;
+
+	{
+		float4 ScreenPosition = SvPositionToResolvedScreenPosition(SvPosition);
+		float3 TranslatedWorldPosition = WorldRayOrigin() + WorldRayDirection() * RayTCurrent() + ResolvedView.PreViewTranslation.xyz;
+
+		bool bIsFrontFace = HitKind() == HIT_KIND_TRIANGLE_FRONT_FACE;
+
+		MaterialParameters.CameraVector = -WorldRayDirection();
+
+		// #dxr_todo: UE-72130 support world position offset 
+		// #if USE_WORLD_POSITION_EXCLUDING_SHADER_OFFSETS
+		// CalcMaterialParametersEx(MaterialParameters, PixelMaterialInputs, SvPosition, ScreenPosition, bIsFrontFace, TranslatedWorldPosition, BasePassInterpolants.PixelPositionExcludingWPO);
+		// #else
+		CalcMaterialParametersEx(MaterialParameters, PixelMaterialInputs, SvPosition, ScreenPosition, bIsFrontFace, TranslatedWorldPosition, TranslatedWorldPosition);
+		// #endif
+	}
+
+	// Now extract the relevant info from PixelMaterialInputs according to the blending mode
+
+#if MATERIALBLENDING_MASKED && !MATERIAL_DITHER_OPACITY_MASK
+	// Regardless of payload flags -- we always apply this 
+	if (GetMaterialMask(PixelMaterialInputs) < 0)
+	{
+		IgnoreHit();
+	}
+#else
+	// If we got this far -- we have a blending mode that can produce partial opacity _and_ we want to
+	// know how much it contributes to the throughput along this ray. These cases are currently only
+	// exercised by the path tracer.
+#if MATERIALBLENDING_MODULATE
+	// this material is the only one to lead to colored opacity
+	float3 Transparency = GetMaterialEmissive(PixelMaterialInputs);
+#elif MATERIALBLENDING_ALPHACOMPOSITE
+	float Opacity = GetMaterialOpacity(PixelMaterialInputs);
+	float Transparency = 1 - Opacity;
+#elif MATERIAL_SHADINGMODEL_THIN_TRANSLUCENT
+	// Assume Payload.ShadingModelID == SHADINGMODELID_THIN_TRANSLUCENT
+	float3 Transmission = GetThinTranslucentMaterialOutput0(MaterialParameters);
+	float Opacity = GetMaterialOpacity(PixelMaterialInputs);
+#if REFRACTION_USE_INDEX_OF_REFRACTION && !SIMPLIFIED_MATERIAL_SHADER
+	float Ior = GetMaterialRefraction(PixelMaterialInputs).x;
+	Ior = Ior == 0 ? 1.0 : Ior;
+#else
+	float Ior = 1.0;
+#endif
+	float3 Transparency = 1 - Opacity;
+	float3 V = WorldRayDirection();
+	float3 N = normalize(MaterialParameters.WorldNormal);
+	float VoN = abs(dot(V, N));
+#if !SIMPLIFIED_MATERIAL_SHADER
+	if (Opacity < 1.0 && Ior != 1.0)
+	{
+		float PathRoughness = PackedPayload.GetPathRoughness();
+		float Roughness = GetMaterialRoughness(PixelMaterialInputs);
+		if (PathRoughness <= Roughness)
+		{
+			// not using fast caustic approximation - treat as opaque
+			return;
+		}
+		// compute transmission through the slab (fresnel + absorption)
+		Transparency *= ComputeThinSlabWeights(Transmission, VoN, Ior).Transmitted;
+
+		// fake caustic approximation (see comments below)
+		Transparency *= (1 - Roughness * Roughness) * saturate(PathRoughness - Roughness);
+	}
+	else
+#endif
+	{
+		// no bending at the interface so no fresnel - equation above simplifies to just A
+		float3 A = pow(Transmission, 1.0 / VoN);
+		Transparency *= A;
+	}
+#elif MATERIALBLENDING_MASKED && MATERIAL_DITHER_OPACITY_MASK
+	float Opacity = DitheredOpacityMaskToOpacity(GetMaterialMask(PixelMaterialInputs));
+	float Transparency = 1 - Opacity;
+#elif MATERIALBLENDING_TRANSLUCENT
+	float Opacity = GetMaterialOpacity(PixelMaterialInputs);
+	float Transparency = 1 - Opacity;
+#if REFRACTION_USE_INDEX_OF_REFRACTION && !SIMPLIFIED_MATERIAL_SHADER
+	float Ior = GetMaterialRefraction(PixelMaterialInputs).x;
+#else
+	float Ior = 1.0;
+#endif
+
+#if !SIMPLIFIED_MATERIAL_SHADER
+	// Is refraction enabled?
+	if (Transparency > 0 && Ior != 1.0)
+	{
+		float PathRoughness = PackedPayload.GetPathRoughness();
+		float Roughness = GetMaterialRoughness(PixelMaterialInputs);
+		if (PathRoughness <= Roughness)
+		{
+			// not using fast caustic approximation - treat as opaque
+			return;
+		}
+		// The heuristic used here is inspired by the following presentations:
+		//   - Kulla & Conty: Revisiting Physically Based Shading at Imageworks
+		//       https://blog.selfshadow.com/publications/s2017-shading-course/imageworks/s2017_pbs_imageworks_slides_v2.pdf
+		//   - Colin Barré-Brisebois - Pica Pica & Nvidia Turing
+		//       https://www.ea.com/seed/news/siggraph-2018-picapica-nv-turing
+
+		// fake caustics
+		float3 N = normalize(MaterialParameters.WorldNormal);
+		float Eta = Ior;
+		float VoN = abs(dot(WorldRayDirection(), N));
+		float Fr = FresnelReflectance(VoN, Eta);
+		Transparency *= pow(1 - Fr, 2) * (1 - Roughness * Roughness) * saturate(PathRoughness - Roughness);
+	}
+#endif
+#else
+#error "Control flow does not account for all blending modes!"
+#endif
+
+	// Update the ray throughput (it is packed simply into the payload since we don't need to carry any other information across hits)
+	float3 RayThroughput = PackedPayload.GetRayThroughput();
+
+	RayThroughput *= Transparency;
+
+	PackedPayload.SetRayThroughput(RayThroughput);
+
+	if (any(RayThroughput > 0))
+	{
+		// not at full opacity yet -- ignore the hit so we can keep going
+		IgnoreHit();
+	}
+	// If RayThroughput reached 0 -- just return which will accept the hit
+#endif // !MATERIALBLENDING_MASKED
+
+#endif // !MATERIALBLENDING_SOLID && !MATERIALBLENDING_ALPHAHOLDOUT && !MATERIALBLENDING_ADDITIVE
+
+#endif // USE_MATERIAL_ANY_HIT_SHADER
+}
diff --git a/Engine/Shaders/Private/RayTracing/RayTracingCalcInterpolants.ush b/Engine/Shaders/Private/RayTracing/RayTracingCalcInterpolants.ush
new file mode 100644
--- /dev/null
+++ b/Engine/Shaders/Private/RayTracing/RayTracingCalcInterpolants.ush
@@ -0,0 +1,53 @@
+// Copyright Epic Games, Inc. All Rights Reserved.
+
+#pragma once
+
+void CalcInterpolants(in FRayCone RayCone, in FDefaultAttributes Attributes, out FVertexFactoryInterpolantsVSToPS Interpolants)
+{
+	FVertexFactoryRayTracingInterpolants Interpolated = (FVertexFactoryRayTracingInterpolants)0;
+
+	float3 WorldPositions[3];
+	float2 TexCoords[3];
+
+	float3 Weights = float3(
+		1 - Attributes.Barycentrics.x - Attributes.Barycentrics.y,
+		Attributes.Barycentrics.x,
+		Attributes.Barycentrics.y);
+
+	for (int i = 0; i < 3; i++)
+	{
+		FVertexFactoryInput Input = LoadVertexFactoryInputForHGS(PrimitiveIndex(), i);
+
+		FVertexFactoryIntermediates VFIntermediates = GetVertexFactoryIntermediates(Input);
+		float3x3 TangentToLocal = VertexFactoryGetTangentToLocal(Input, VFIntermediates);
+		float4 WorldPositionExcludingWPO = VertexFactoryGetWorldPosition(Input, VFIntermediates);
+		FMaterialVertexParameters VertexParameters = GetMaterialVertexParameters(Input, VFIntermediates, WorldPositionExcludingWPO.xyz, TangentToLocal);
+
+		FVertexFactoryRayTracingInterpolants PerVertexInterpolants = VertexFactoryGetRayTracingInterpolants(Input, VFIntermediates, VertexParameters);
+		Interpolated = VertexFactoryInterpolate(PerVertexInterpolants, Weights[i], Interpolated, 1.0);
+
+		WorldPositions[i] = WorldPositionExcludingWPO.xyz;
+		TexCoords[i] = VertexFactoryGetRayTracingTextureCoordinate(PerVertexInterpolants);
+	}
+
+	Interpolants = VertexFactoryAssignInterpolants(Interpolated);
+
+#if (NUM_TEX_COORD_INTERPOLATORS || USE_PARTICLE_SUBUVS) && !VERTEX_FACTORY_MODIFIES_INTERPOLATION
+	float2 TA = TexCoords[1] - TexCoords[0];
+	float2 TB = TexCoords[2] - TexCoords[0];
+	float UvArea = abs((TA.x * TB.y) - (TA.y * TB.x));
+
+	float3 PA = WorldPositions[1] - WorldPositions[0];
+	float3 PB = WorldPositions[2] - WorldPositions[0];
+	float3 Unnormalized = cross(PB, PA);
+	float InvWorldArea = rsqrt(dot(Unnormalized, Unnormalized));
+	float3 FaceNormal = Unnormalized * InvWorldArea;
+
+	float RayConeCapArea = RayCone.Width * RayCone.Width * (1.0f / (2.0f * PI));
+	float3 RayDir = WorldRayDirection();
+	float RaySlope = dot(RayDir, FaceNormal);
+	const float SmallValue = 1.0 / 16777216.0;
+	GlobalRayCone_TexArea = (UvArea * InvWorldArea * RayConeCapArea) / max(RaySlope * RaySlope, SmallValue);
+	GlobalRayCone_TexArea = max(GlobalRayCone_TexArea, SmallValue);
+#endif
+}
diff --git a/Engine/Shaders/Private/RayTracing/RayTracingMaterialHitShaders.usf b/Engine/Shaders/Private/RayTracing/RayTracingMaterialHitShaders.usf
--- a/Engine/Shaders/Private/RayTracing/RayTracingMaterialHitShaders.usf
+++ b/Engine/Shaders/Private/RayTracing/RayTracingMaterialHitShaders.usf
@@ -14,12 +14,6 @@
 	#define SIMPLIFIED_MATERIAL_SHADER 0
 #endif
 
-// #dxr_todo: UE-72537 move to common.ush
-#define clip(x)
-#define ddx(x) 0
-#define ddy(x) 0
-#define fwidth(x) 0
-
 #include "/Engine/Private/Common.ush"
 #include "RayTracingCommon.ush"
 #include "RayTracingHitGroupCommon.ush"
@@ -27,14 +21,13 @@
 #include "/Engine/Generated/Material.ush"   
 #include "/Engine/Generated/VertexFactory.ush"
 
+#include "RayTracingCalcInterpolants.ush" // must be included after VertexFactory.ush
 #include "/Engine/Private/ShadingCommon.ush"
 #include "/Engine/Private/DeferredShadingCommon.ush"
 #include "/Engine/Private/SHCommon.ush"
 #include "/Engine/Private/ReflectionEnvironmentShared.ush"
 #include "/Engine/Private/VirtualTextureCommon.ush" 
 #include "/Engine/Private/LightmapCommon.ush"
-#include "/Engine/Private/PathTracing/Material/PathTracingFresnel.ush"
-#include "/Engine/Private/PathTracing/Material/PathTracingThinGlass.ush"
 #include "/Engine/Private/Lumen/LumenHardwareRayTracingPayloadCommon.ush"
 
 /** Computes sky diffuse lighting, including precomputed shadowing. */
@@ -214,68 +207,6 @@ void GetPrecomputedIndirectLightingAndSkyLight(
 	#endif
 }
 
-void CalcInterpolants(in FRayCone RayCone, in FDefaultAttributes Attributes, out FVertexFactoryInterpolantsVSToPS Interpolants)
-{
-	FVertexFactoryRayTracingInterpolants Interpolated = (FVertexFactoryRayTracingInterpolants)0;
-
-	float3 WorldPositions[3];
-	float2 TexCoords[3];
-
-	float3 Weights = float3(
-		1 - Attributes.Barycentrics.x - Attributes.Barycentrics.y,
-		Attributes.Barycentrics.x,
-		Attributes.Barycentrics.y);
-
-	for (int i = 0; i < 3; i++)
-	{
-		FVertexFactoryInput Input = LoadVertexFactoryInputForHGS(PrimitiveIndex(), i);
-	
-		FVertexFactoryIntermediates VFIntermediates = GetVertexFactoryIntermediates(Input);
-		float3x3 TangentToLocal = VertexFactoryGetTangentToLocal(Input, VFIntermediates);	
-		float4 WorldPositionExcludingWPO = VertexFactoryGetWorldPosition(Input, VFIntermediates);
-		FMaterialVertexParameters VertexParameters = GetMaterialVertexParameters(Input, VFIntermediates, WorldPositionExcludingWPO.xyz, TangentToLocal);
-
-		FVertexFactoryRayTracingInterpolants PerVertexInterpolants = VertexFactoryGetRayTracingInterpolants(Input, VFIntermediates, VertexParameters);
-		Interpolated = VertexFactoryInterpolate(PerVertexInterpolants, Weights[i], Interpolated, 1.0);
-
-		WorldPositions[i] = WorldPositionExcludingWPO.xyz;
-		TexCoords[i] = VertexFactoryGetRayTracingTextureCoordinate(PerVertexInterpolants);
-	}
-
-	Interpolants = VertexFactoryAssignInterpolants(Interpolated);
-
-#if (NUM_TEX_COORD_INTERPOLATORS || USE_PARTICLE_SUBUVS) && !VERTEX_FACTORY_MODIFIES_INTERPOLATION
-	float2 TA = TexCoords[1] - TexCoords[0]; 
-	float2 TB = TexCoords[2] - TexCoords[0];
-    float UvArea = abs((TA.x * TB.y) - (TA.y * TB.x));
-
-	float3 PA = WorldPositions[1] - WorldPositions[0];
-	float3 PB = WorldPositions[2] - WorldPositions[0];
-    float3 Unnormalized = cross(PB, PA);
-	float InvWorldArea = rsqrt(dot(Unnormalized, Unnormalized));
-    float3 FaceNormal = Unnormalized * InvWorldArea;
-    
-    float RayConeCapArea = RayCone.Width * RayCone.Width * (1.0f / (2.0f * PI));
-    float3 RayDir = WorldRayDirection();
-    float RaySlope = dot(RayDir, FaceNormal);
-    const float SmallValue = 1.0 / 16777216.0;
-    GlobalRayCone_TexArea = (UvArea * InvWorldArea * RayConeCapArea) / max(RaySlope * RaySlope, SmallValue);
-    GlobalRayCone_TexArea = max(GlobalRayCone_TexArea, SmallValue);
-#endif
-}
-
-float DitheredOpacityMaskToOpacity(float Mask)
-{
-	// This represents the expected value of the function GetMaterialCoverageAndClipping
-	// which randomly dithers the fragment on or off to produce the effect of opacity
-
-	// The expected value of this dithering can be computed as:
-	//    Simplify[Integrate[If[Mask + Dither - 1/2 < 0, 0, 1], {Dither, 0, 1}]]
-	// which is just:
-
-	return saturate(Mask + 0.5);
-}
-
 float3 GetLumenGeometryNormal()
 {
 	float3 WorldPositions[3];
@@ -600,14 +531,6 @@ RAY_TRACING_ENTRY_ANY_HIT(MaterialAHS,
 	// Continue traversal by simply returning from this shader
 	// Usually this path is inactive because opaque materials don't have AHS bound
 	return;
-#elif MATERIALBLENDING_ALPHAHOLDOUT || MATERIALBLENDING_ADDITIVE
-	if (PackedPayload.IsPathTracingVisibilityRay())
-	{
-		// path tracer wants to treat these cases as fully transparent
-		// no material evaluation is required
-		IgnoreHit();
-	}
-	// otherwise accept hit
 #else
 	// All other blending modes _might_ need material evaluation - so we need
 	// to generate some code
@@ -625,15 +548,9 @@ RAY_TRACING_ENTRY_ANY_HIT(MaterialAHS,
 #endif
 		return;
 	}
-	else if (PackedPayload.IsPathTracingVisibilityRay())
-	{
-		// this case is used by the path tracer only currently
-		// We want to evaluate the material opacity and accumulate it
-		// Keep going and evaluate the material
-	}
 	else
 	{
-		// non-opaque blending mode but no special behavior was requested
+		// non-opaque blending mode, treat as opaque for now
 		return;
 	}
 #endif  // !MATERIALBLENDING_MASKED
@@ -678,107 +595,6 @@ RAY_TRACING_ENTRY_ANY_HIT(MaterialAHS,
 	{
 		IgnoreHit();
 	}
-#else
-	// If we got this far -- we have a blending mode that can produce partial opacity _and_ we want to
-	// know how much it contributes to the throughput along this ray. These cases are currently only
-	// exercised by the path tracer.
-#if MATERIALBLENDING_MODULATE
-	// this material is the only one to lead to colored opacity
-	float3 Transparency = GetMaterialEmissive(PixelMaterialInputs);
-#elif MATERIALBLENDING_ALPHACOMPOSITE
-	float Opacity = GetMaterialOpacity(PixelMaterialInputs);
-	float Transparency = 1 - Opacity;
-#elif MATERIAL_SHADINGMODEL_THIN_TRANSLUCENT
-	// Assume Payload.ShadingModelID == SHADINGMODELID_THIN_TRANSLUCENT
-	float3 Transmission = GetThinTranslucentMaterialOutput0(MaterialParameters);
-	float Opacity = GetMaterialOpacity(PixelMaterialInputs);
-#if REFRACTION_USE_INDEX_OF_REFRACTION && !SIMPLIFIED_MATERIAL_SHADER
-	float Ior = GetMaterialRefraction(PixelMaterialInputs).x;
-	Ior = Ior == 0 ? 1.0 : Ior;
-#else
-	float Ior = 1.0;
-#endif
-	float3 Transparency = 1 - Opacity;
-	float3 V = WorldRayDirection();
-	float3 N = normalize(MaterialParameters.WorldNormal);
-	float VoN = abs(dot(V, N));
-#if !SIMPLIFIED_MATERIAL_SHADER
-	if (Opacity < 1.0 && Ior != 1.0)
-	{
-		float PathRoughness = PackedPayload.GetPathRoughness();
-		float Roughness = GetMaterialRoughness(PixelMaterialInputs);
-		if (PathRoughness <= Roughness)
-		{
-			// not using fast caustic approximation - treat as opaque
-			return;
-		}
-		// compute transmission through the slab (fresnel + absorption)
-		Transparency *= ComputeThinSlabWeights(Transmission, VoN, Ior).Transmitted;
-
-		// fake caustic approximation (see comments below)
-		Transparency *= (1 - Roughness * Roughness) * saturate(PathRoughness - Roughness);
-	}
-	else
-#endif
-	{
-		// no bending at the interface so no fresnel - equation above simplifies to just A
-		float3 A = pow(Transmission, 1.0 / VoN);
-		Transparency *= A;
-	}
-#elif MATERIALBLENDING_MASKED && MATERIAL_DITHER_OPACITY_MASK
-	float Opacity = DitheredOpacityMaskToOpacity(GetMaterialMask(PixelMaterialInputs));
-	float Transparency = 1 - Opacity;
-#elif MATERIALBLENDING_TRANSLUCENT
-	float Opacity = GetMaterialOpacity(PixelMaterialInputs);
-	float Transparency = 1 - Opacity;
-#if REFRACTION_USE_INDEX_OF_REFRACTION && !SIMPLIFIED_MATERIAL_SHADER
-	float Ior = GetMaterialRefraction(PixelMaterialInputs).x;
-#else
-	float Ior = 1.0;
-#endif
-
-#if !SIMPLIFIED_MATERIAL_SHADER
-	// Is refraction enabled?
-	if (Transparency > 0 && Ior != 1.0)
-	{
-		float PathRoughness = PackedPayload.GetPathRoughness();
-		float Roughness = GetMaterialRoughness(PixelMaterialInputs);
-		if (PathRoughness <= Roughness)
-		{
-			// not using fast caustic approximation - treat as opaque
-			return;
-		}
-		// The heuristic used here is inspired by the following presentations:
-		//   - Kulla & Conty: Revisiting Physically Based Shading at Imageworks
-		//       https://blog.selfshadow.com/publications/s2017-shading-course/imageworks/s2017_pbs_imageworks_slides_v2.pdf
-		//   - Colin Barré-Brisebois - Pica Pica & Nvidia Turing
-		//       https://www.ea.com/seed/news/siggraph-2018-picapica-nv-turing
-
-		// fake caustics
-		float3 N = normalize(MaterialParameters.WorldNormal);
-		float Eta = Ior;
-		float VoN = abs(dot(WorldRayDirection(), N));
-		float Fr = FresnelReflectance(VoN, Eta);
-		Transparency *= pow(1 - Fr, 2) * (1 - Roughness * Roughness) * saturate(PathRoughness - Roughness);
-	}
-#endif
-#else
-#error "Control flow does not account for all blending modes!"
-#endif
-
-	// Update the ray throughput (it is packed simply into the payload since we don't need to carry any other information across hits)
-	float3 RayThroughput = PackedPayload.GetRayThroughput();
-
-	RayThroughput *= Transparency;
-
-	PackedPayload.SetRayThroughput(RayThroughput);
-
-	if (any(RayThroughput > 0))
-	{
-		// not at full opacity yet -- ignore the hit so we can keep going
-		IgnoreHit();
-	}
-	// If RayThroughput reached 0 -- just return which will accept the hit
 #endif // !MATERIALBLENDING_MASKED
 
 #endif // !MATERIALBLENDING_SOLID && !MATERIALBLENDING_ALPHAHOLDOUT && !MATERIALBLENDING_ADDITIVE
diff --git a/Engine/Source/Editor/UnrealEd/Private/SEditorViewportViewMenu.cpp b/Engine/Source/Editor/UnrealEd/Private/SEditorViewportViewMenu.cpp
--- a/Engine/Source/Editor/UnrealEd/Private/SEditorViewportViewMenu.cpp
+++ b/Engine/Source/Editor/UnrealEd/Private/SEditorViewportViewMenu.cpp
@@ -123,7 +123,11 @@ void SEditorViewportViewMenu::FillViewMenu(UToolMenu* Menu) const
 #if RHI_RAYTRACING
 			if (IsRayTracingEnabled())
 			{
-				Section.AddMenuEntry(BaseViewportActions.PathTracingMode, UViewModeUtils::GetViewModeDisplayName(VMI_PathTracing));
+				static auto PathTracingCvar = IConsoleManager::Get().FindTConsoleVariableDataInt(TEXT("r.PathTracing"));
+				if (PathTracingCvar && PathTracingCvar->GetValueOnAnyThread() != 0)
+				{
+					Section.AddMenuEntry(BaseViewportActions.PathTracingMode, UViewModeUtils::GetViewModeDisplayName(VMI_PathTracing));
+				}
 			}
 #endif
 
diff --git a/Engine/Source/Runtime/Engine/Classes/Engine/RendererSettings.h b/Engine/Source/Runtime/Engine/Classes/Engine/RendererSettings.h
--- a/Engine/Source/Runtime/Engine/Classes/Engine/RendererSettings.h
+++ b/Engine/Source/Runtime/Engine/Classes/Engine/RendererSettings.h
@@ -403,6 +403,12 @@ class ENGINE_API URendererSettings : public UDeveloperSettings
 		ConfigRestartRequired = true))
 		uint32 bEnableRayTracingTextureLOD : 1;
 
+	UPROPERTY(config, EditAnywhere, Category = HardwareRayTracing, meta = (
+		ConsoleVariable = "r.PathTracing", DisplayName = "Path Tracing",
+		ToolTip = "Support Path Tracing renderer. This enables additional material permutations.",
+		ConfigRestartRequired = true))
+		uint32 bEnablePathTracing : 1;
+
 	UPROPERTY(config, EditAnywhere, Category=SoftwareRayTracing, meta=(
 		ConsoleVariable="r.GenerateMeshDistanceFields",
 		ToolTip="Whether to build distance fields of static meshes, needed for Software Ray Tracing in Lumen, and distance field AO, which is used to implement Movable SkyLight shadows, and ray traced distance field shadows on directional lights.  Enabling will increase the build times, memory usage and disk size of static meshes.  Changing this setting requires restarting the editor.",
diff --git a/Engine/Source/Runtime/Renderer/Private/PathTracing.cpp b/Engine/Source/Runtime/Renderer/Private/PathTracing.cpp
--- a/Engine/Source/Runtime/Renderer/Private/PathTracing.cpp
+++ b/Engine/Source/Runtime/Renderer/Private/PathTracing.cpp
@@ -12,9 +12,17 @@
 #include "RayTracingTypes.h"
 #include "RayTracingDefinitions.h"
 #include "PathTracingDefinitions.h"
+#include "RayTracing/RayTracingMaterialHitShaders.h"
 #include "RenderCore/Public/GenerateMips.h"
 #include <limits>
 
+TAutoConsoleVariable<int32> CVarPathTracing(
+	TEXT("r.PathTracing"),
+	0,
+	TEXT("Enables the path tracing renderer (to guard the compilation of path tracer specific material permutations)"),
+	ECVF_RenderThreadSafe | ECVF_ReadOnly
+);
+
 TAutoConsoleVariable<int32> CVarPathTracingMaxBounces(
 	TEXT("r.PathTracing.MaxBounces"),
 	-1,
@@ -289,6 +297,13 @@ static void PrepareShaderArgs(const FViewInfo& View, FPathTracingData& PathTraci
 	PathTracingData.FilterWidth = FilterWidth;
 }
 
+static bool ShouldCompilePathTracingShadersForProject(EShaderPlatform ShaderPlatform)
+{
+	return ShouldCompileRayTracingShadersForProject(ShaderPlatform) &&
+			FDataDrivenShaderPlatformInfo::GetSupportsPathTracing(ShaderPlatform) &&
+			CVarPathTracing.GetValueOnAnyThread() != 0;
+}
+
 class FPathTracingSkylightPrepareCS : public FGlobalShader
 {
 	DECLARE_GLOBAL_SHADER(FPathTracingSkylightPrepareCS)
@@ -296,6 +311,7 @@ class FPathTracingSkylightPrepareCS : public FGlobalShader
 
 	static bool ShouldCompilePermutation(const FGlobalShaderPermutationParameters& Parameters)
 	{
+		// NOTE: skylight code is shared with RT passes
 		return ShouldCompileRayTracingShadersForProject(Parameters.Platform);
 	}
 
@@ -327,6 +343,7 @@ class FPathTracingSkylightMISCompensationCS : public FGlobalShader
 
 	static bool ShouldCompilePermutation(const FGlobalShaderPermutationParameters& Parameters)
 	{
+		// NOTE: skylight code is shared with RT passes
 		return ShouldCompileRayTracingShadersForProject(Parameters.Platform);
 	}
 
@@ -395,8 +412,7 @@ class FPathTracingRG : public FGlobalShader
 
 	static bool ShouldCompilePermutation(const FGlobalShaderPermutationParameters& Parameters)
 	{
-		return ShouldCompileRayTracingShadersForProject(Parameters.Platform)
-			&& FDataDrivenShaderPlatformInfo::GetSupportsPathTracing(Parameters.Platform);
+		return ShouldCompilePathTracingShadersForProject(Parameters.Platform);
 	}
 
 	static void ModifyCompilationEnvironment(const FGlobalShaderPermutationParameters& Parameters, FShaderCompilerEnvironment& OutEnvironment)
@@ -462,6 +478,116 @@ class FPathTracingIESAtlasCS : public FGlobalShader
 };
 IMPLEMENT_SHADER_TYPE(, FPathTracingIESAtlasCS, TEXT("/Engine/Private/PathTracing/PathTracingIESAtlas.usf"), TEXT("PathTracingIESAtlasCS"), SF_Compute);
 
+template<bool UseAnyHitShader>
+class TPathTracingMaterial : public FMeshMaterialShader
+{
+	DECLARE_SHADER_TYPE(TPathTracingMaterial, MeshMaterial);
+public:
+	TPathTracingMaterial() = default;
+
+	TPathTracingMaterial(const FMeshMaterialShaderType::CompiledShaderInitializerType& Initializer)
+		: FMeshMaterialShader(Initializer)
+	{}
+
+	static bool ShouldCompilePermutation(const FMeshMaterialShaderPermutationParameters& Parameters)
+	{
+		return Parameters.VertexFactoryType->SupportsRayTracing()
+			&& ((Parameters.MaterialParameters.bIsMasked || Parameters.MaterialParameters.BlendMode != BLEND_Opaque) == UseAnyHitShader)
+			&& ShouldCompilePathTracingShadersForProject(Parameters.Platform);
+	}
+
+	static void ModifyCompilationEnvironment(const FMaterialShaderPermutationParameters& Parameters, FShaderCompilerEnvironment& OutEnvironment)
+	{
+		OutEnvironment.SetDefine(TEXT("USE_MATERIAL_CLOSEST_HIT_SHADER"), 1);
+		OutEnvironment.SetDefine(TEXT("USE_MATERIAL_ANY_HIT_SHADER"), UseAnyHitShader ? 1 : 0);
+		OutEnvironment.SetDefine(TEXT("USE_RAYTRACED_TEXTURE_RAYCONE_LOD"), 0);
+		OutEnvironment.SetDefine(TEXT("SCENE_TEXTURES_DISABLED"), 1);
+		OutEnvironment.SetDefine(TEXT("SIMPLIFIED_MATERIAL_SHADER"), 0); // TODO: expose a permutation so we can unify with GPULightmass?
+		FMeshMaterialShader::ModifyCompilationEnvironment(Parameters, OutEnvironment);
+	}
+
+	static bool ValidateCompiledResult(EShaderPlatform Platform, const FShaderParameterMap& ParameterMap, TArray<FString>& OutError)
+	{
+		if (ParameterMap.ContainsParameterAllocation(FSceneTextureUniformParameters::StaticStructMetadata.GetShaderVariableName()))
+		{
+			OutError.Add(TEXT("Ray tracing closest hit shaders cannot read from the SceneTexturesStruct."));
+			return false;
+		}
+
+		for (const auto& It : ParameterMap.GetParameterMap())
+		{
+			const FParameterAllocation& ParamAllocation = It.Value;
+			if (ParamAllocation.Type != EShaderParameterType::UniformBuffer
+				&& ParamAllocation.Type != EShaderParameterType::LooseData)
+			{
+				OutError.Add(FString::Printf(TEXT("Invalid ray tracing shader parameter '%s'. Only uniform buffers and loose data parameters are supported."), *(It.Key)));
+				return false;
+			}
+		}
+
+		return true;
+	}
+};
+
+using FPathTracingMaterialCHS     = TPathTracingMaterial<false>;
+using FPathTracingMaterialCHS_AHS = TPathTracingMaterial<true>;
+
+
+IMPLEMENT_MATERIAL_SHADER_TYPE(template <>, FPathTracingMaterialCHS    , TEXT("/Engine/Private/PathTracing/PathTracingMaterialHitShader.usf"), TEXT("closesthit=PathTracingMaterialCHS"), SF_RayHitGroup);
+IMPLEMENT_MATERIAL_SHADER_TYPE(template <>, FPathTracingMaterialCHS_AHS, TEXT("/Engine/Private/PathTracing/PathTracingMaterialHitShader.usf"), TEXT("closesthit=PathTracingMaterialCHS anyhit=PathTracingMaterialAHS"), SF_RayHitGroup);
+
+
+bool FRayTracingMeshProcessor::ProcessPathTracing(
+	const FMeshBatch& RESTRICT MeshBatch,
+	uint64 BatchElementMask,
+	const FPrimitiveSceneProxy* RESTRICT PrimitiveSceneProxy,
+	const FMaterialRenderProxy& RESTRICT MaterialRenderProxy,
+	const FMaterial& RESTRICT MaterialResource)
+{
+	const FVertexFactory* VertexFactory = MeshBatch.VertexFactory;
+
+	TMeshProcessorShaders<
+		FMeshMaterialShader,
+		FMeshMaterialShader,
+		FMeshMaterialShader,
+		FMeshMaterialShader,
+		FMeshMaterialShader> RayTracingShaders;
+
+	FMaterialShaderTypes ShaderTypes;
+
+	if (MaterialResource.IsMasked() || MaterialResource.GetBlendMode() != BLEND_Opaque)
+	{
+		ShaderTypes.AddShaderType<FPathTracingMaterialCHS_AHS>();
+	}
+	else
+	{
+		ShaderTypes.AddShaderType<FPathTracingMaterialCHS>();
+	}
+
+	FMaterialShaders Shaders;
+	if (!MaterialResource.TryGetShaders(ShaderTypes, VertexFactory->GetType(), Shaders))
+	{
+		return false;
+	}
+
+	check(Shaders.TryGetShader(SF_RayHitGroup, RayTracingShaders.RayHitGroupShader));
+
+	TBasePassShaderElementData<FUniformLightMapPolicy> ShaderElementData(MeshBatch.LCI);
+	ShaderElementData.InitializeMeshMaterialData(ViewIfDynamicMeshCommand, PrimitiveSceneProxy, MeshBatch, -1, true);
+
+	BuildRayTracingMeshCommands(
+		MeshBatch,
+		BatchElementMask,
+		PrimitiveSceneProxy,
+		MaterialRenderProxy,
+		MaterialResource,
+		PassDrawRenderState,
+		RayTracingShaders,
+		ShaderElementData);
+
+	return true;
+}
+
 RENDERER_API void PrepareSkyTexture_Internal(
 	FRDGBuilder& GraphBuilder,
 	FReflectionUniformParameters& Parameters,
@@ -1054,7 +1180,7 @@ IMPLEMENT_SHADER_TYPE(, FPathTracingCompositorPS, TEXT("/Engine/Private/PathTrac
 void FDeferredShadingSceneRenderer::PreparePathTracing(const FSceneViewFamily& ViewFamily, TArray<FRHIRayTracingShader*>& OutRayGenShaders)
 {
 	if (ViewFamily.EngineShowFlags.PathTracing
-		&& FDataDrivenShaderPlatformInfo::GetSupportsPathTracing(ViewFamily.GetShaderPlatform()))
+		&& ShouldCompilePathTracingShadersForProject(ViewFamily.GetShaderPlatform()))
 	{
 		// Declare all RayGen shaders that require material closest hit shaders to be bound
 		auto RayGenShader = GetGlobalShaderMap(ViewFamily.GetShaderPlatform())->GetShader<FPathTracingRG>();
@@ -1084,6 +1210,13 @@ void FDeferredShadingSceneRenderer::RenderPathTracing(
 		return;
 	}
 
+	if (CVarPathTracing.GetValueOnRenderThread() == 0)
+	{
+		// Path tracing is not enabled on this project (should not be seen by end-users since the menu entry to pick path tracing should be hidden)
+		// If they reach this code through ShowFlag manipulation, they may observe an incomplete image. Is there a way to inform the user here?
+		return;
+	}
+
 	FPathTracingConfig Config = {};
 
 	// Get current value of MaxSPP and reset render if it has changed
diff --git a/Engine/Source/Runtime/Renderer/Private/RayTracing/RayTracingMaterialHitShaders.cpp b/Engine/Source/Runtime/Renderer/Private/RayTracing/RayTracingMaterialHitShaders.cpp
--- a/Engine/Source/Runtime/Renderer/Private/RayTracing/RayTracingMaterialHitShaders.cpp
+++ b/Engine/Source/Runtime/Renderer/Private/RayTracing/RayTracingMaterialHitShaders.cpp
@@ -262,9 +262,7 @@ bool FRayTracingMeshProcessor::Process(
 	const FPrimitiveSceneProxy* RESTRICT PrimitiveSceneProxy,
 	const FMaterialRenderProxy& RESTRICT MaterialRenderProxy,
 	const FMaterial& RESTRICT MaterialResource,
-	FMaterialShadingModelField ShadingModels,
-	const FUniformLightMapPolicy& RESTRICT LightMapPolicy,
-	const typename FUniformLightMapPolicy::ElementDataType& RESTRICT LightMapElementData)
+	const FUniformLightMapPolicy& RESTRICT LightMapPolicy)
 {
 
 	const bool bMaterialsCompiled = GCompileRayTracingMaterialAHS || GCompileRayTracingMaterialCHS;
@@ -331,10 +329,7 @@ bool FRayTracingMeshProcessor::Process(
 		Shaders.TryGetShader(SF_RayHitGroup, RayTracingShaders.RayHitGroupShader);
 	}
 
-	PassDrawRenderState.SetBlendState(TStaticBlendState<CW_RGBA, BO_Add, BF_One, BF_One, BO_Add, BF_Zero, BF_One>::GetRHI());
-	PassDrawRenderState.SetDepthStencilState(TStaticDepthStencilState<false, CF_DepthNearOrEqual>::GetRHI());
-
-	TBasePassShaderElementData<FUniformLightMapPolicy> ShaderElementData(LightMapElementData);
+	TBasePassShaderElementData<FUniformLightMapPolicy> ShaderElementData(MeshBatch.LCI);
 	ShaderElementData.InitializeMeshMaterialData(ViewIfDynamicMeshCommand, PrimitiveSceneProxy, MeshBatch, -1, true);
 
 	BuildRayTracingMeshCommands(
@@ -381,20 +376,20 @@ bool FRayTracingMeshProcessor::TryAddMeshBatch(
 	const FMaterial& Material
 )
 {
-	// Determine the mesh's material and blend mode.
-	const EBlendMode BlendMode = Material.GetBlendMode();
-	const FMaterialShadingModelField ShadingModels = Material.GetShadingModels();
-
 	// Only draw opaque materials.
 	if ((!PrimitiveSceneProxy || PrimitiveSceneProxy->ShouldRenderInMainPass())
 		&& ShouldIncludeDomainInMeshPass(Material.GetMaterialDomain()))
 	{
+		if (RayTracingMeshCommandsMode == ERayTracingMeshCommandsMode::PATH_TRACING)
+		{
+			// Path Tracer has its own process call so that it can attach its own material permutation
+			return ProcessPathTracing(MeshBatch, BatchElementMask, PrimitiveSceneProxy, MaterialRenderProxy, Material);
+		}
+
 		// Check for a cached light-map.
-		const bool bIsLitMaterial = ShadingModels.IsLit();
+		const bool bIsLitMaterial = Material.GetShadingModels().IsLit();
 		static const auto AllowStaticLightingVar = IConsoleManager::Get().FindTConsoleVariableDataInt(TEXT("r.AllowStaticLighting"));
-		// NOTE: path tracing is not interested in baked static lighting (TODO: introduce a call for path tracer specific materials)
-		const bool bAllowStaticLighting = RayTracingMeshCommandsMode == ERayTracingMeshCommandsMode::RAY_TRACING &&
-			(!AllowStaticLightingVar || AllowStaticLightingVar->GetValueOnRenderThread() != 0);
+		const bool bAllowStaticLighting = (!AllowStaticLightingVar || AllowStaticLightingVar->GetValueOnRenderThread() != 0);
 
 		const FLightMapInteraction LightMapInteraction = (bAllowStaticLighting && MeshBatch.LCI && bIsLitMaterial)
 			? MeshBatch.LCI->GetLightMapInteraction(FeatureLevel)
@@ -428,9 +423,7 @@ bool FRayTracingMeshProcessor::TryAddMeshBatch(
 							PrimitiveSceneProxy,
 							MaterialRenderProxy,
 							Material,
-							ShadingModels,
-							FUniformLightMapPolicy(LMP_DISTANCE_FIELD_SHADOWS_AND_HQ_LIGHTMAP),
-							MeshBatch.LCI);
+							FUniformLightMapPolicy(LMP_DISTANCE_FIELD_SHADOWS_AND_HQ_LIGHTMAP));
 					}
 					else
 					{
@@ -440,9 +433,7 @@ bool FRayTracingMeshProcessor::TryAddMeshBatch(
 							PrimitiveSceneProxy,
 							MaterialRenderProxy,
 							Material,
-							ShadingModels,
-							FUniformLightMapPolicy(LMP_HQ_LIGHTMAP),
-							MeshBatch.LCI);
+							FUniformLightMapPolicy(LMP_HQ_LIGHTMAP));
 					}
 				}
 				else if (bAllowLowQualityLightMaps)
@@ -453,9 +444,7 @@ bool FRayTracingMeshProcessor::TryAddMeshBatch(
 						PrimitiveSceneProxy,
 						MaterialRenderProxy,
 						Material,
-						ShadingModels,
-						FUniformLightMapPolicy(LMP_LQ_LIGHTMAP),
-						MeshBatch.LCI);
+						FUniformLightMapPolicy(LMP_LQ_LIGHTMAP));
 				}
 				else
 				{
@@ -465,9 +454,7 @@ bool FRayTracingMeshProcessor::TryAddMeshBatch(
 						PrimitiveSceneProxy,
 						MaterialRenderProxy,
 						Material,
-						ShadingModels,
-						FUniformLightMapPolicy(LMP_NO_LIGHTMAP),
-						MeshBatch.LCI);
+						FUniformLightMapPolicy(LMP_NO_LIGHTMAP));
 				}
 				break;
 			default:
@@ -486,9 +473,7 @@ bool FRayTracingMeshProcessor::TryAddMeshBatch(
 						PrimitiveSceneProxy,
 						MaterialRenderProxy,
 						Material,
-						ShadingModels,
-						FUniformLightMapPolicy(LMP_PRECOMPUTED_IRRADIANCE_VOLUME_INDIRECT_LIGHTING),
-						MeshBatch.LCI);
+						FUniformLightMapPolicy(LMP_PRECOMPUTED_IRRADIANCE_VOLUME_INDIRECT_LIGHTING));
 				}
 				else
 				{
@@ -498,9 +483,7 @@ bool FRayTracingMeshProcessor::TryAddMeshBatch(
 						PrimitiveSceneProxy,
 						MaterialRenderProxy,
 						Material,
-						ShadingModels,
-						FUniformLightMapPolicy(LMP_NO_LIGHTMAP),
-						MeshBatch.LCI);
+						FUniformLightMapPolicy(LMP_NO_LIGHTMAP));
 				}
 				break;
 			};
diff --git a/Engine/Source/Runtime/Renderer/Private/RayTracing/RayTracingMaterialHitShaders.h b/Engine/Source/Runtime/Renderer/Private/RayTracing/RayTracingMaterialHitShaders.h
--- a/Engine/Source/Runtime/Renderer/Private/RayTracing/RayTracingMaterialHitShaders.h
+++ b/Engine/Source/Runtime/Renderer/Private/RayTracing/RayTracingMaterialHitShaders.h
@@ -22,7 +22,10 @@ public:
 		FeatureLevel(InScene ? InScene->GetFeatureLevel() : ERHIFeatureLevel::SM5),
 		PassDrawRenderState(InPassDrawRenderState),
 		RayTracingMeshCommandsMode(InRayTracingMeshCommandsMode)
-	{}
+	{
+		PassDrawRenderState.SetBlendState(TStaticBlendState<CW_RGBA, BO_Add, BF_One, BF_One, BO_Add, BF_Zero, BF_One>::GetRHI());
+		PassDrawRenderState.SetDepthStencilState(TStaticDepthStencilState<false, CF_DepthNearOrEqual>::GetRHI());
+	}
 
 	virtual ~FRayTracingMeshProcessor() = default;
 
@@ -42,9 +45,7 @@ protected:
 		const FPrimitiveSceneProxy* RESTRICT PrimitiveSceneProxy,
 		const FMaterialRenderProxy& RESTRICT MaterialRenderProxy,
 		const FMaterial& RESTRICT MaterialResource,
-		FMaterialShadingModelField ShadingModels,
-		const FUniformLightMapPolicy& RESTRICT LightMapPolicy,
-		const typename FUniformLightMapPolicy::ElementDataType& RESTRICT LightMapElementData);
+		const FUniformLightMapPolicy& RESTRICT LightMapPolicy);
 
 	template<typename PassShadersType, typename ShaderElementDataType>
 	void BuildRayTracingMeshCommands(
@@ -104,6 +105,12 @@ protected:
 	}
 
 private:
+	bool ProcessPathTracing(
+		const FMeshBatch& RESTRICT MeshBatch,
+		uint64 BatchElementMask,
+		const FPrimitiveSceneProxy* RESTRICT PrimitiveSceneProxy,
+		const FMaterialRenderProxy& RESTRICT MaterialRenderProxy,
+		const FMaterial& RESTRICT MaterialResource);
 
 	bool TryAddMeshBatch(
 		const FMeshBatch& RESTRICT MeshBatch,
