Fixed parameter blending with weight applied on BSDF. The problem was similar to not using premultipled alpha. Not when using parameter blended operators, the strata parameters are pre multipled with coverage. It only impacts the base pass, and all not necessary operations should be optimised out by the compiler.

diff --git a/Engine/Shaders/Private/Strata/Strata.ush b/Engine/Shaders/Private/Strata/Strata.ush
--- a/Engine/Shaders/Private/Strata/Strata.ush
+++ b/Engine/Shaders/Private/Strata/Strata.ush
@@ -45,9 +45,16 @@
 // BSDFs representation
 struct FStrataBSDF
 {
-	float3				Weight;		// Weight is a float3 because, when loaded from the material buffer, it contains the BSDF throughput (coverage and transmittance combined)
 	uint				State;		// Extra informations and other enabled BSDF features
 
+#if STRATA_INLINE_SHADING
+	float				Coverage;	// When creating a material, each BSDF/slab can have its coverage updated. 
+									// When this happens, we also update its parameters to account for that, similary to premultiplied alpha blending.
+									// Coverage is kept around because it is needed to evalauted the contribution to throughput.
+	bool				bPremultipliedCoverage;// True when coverage is already weighting all parameters. This is true when parameter blending is used. This allows to use parameter blending or not in the same graph.
+#endif
+	float3				FinalWeight;// Contains the throughput towards the view accounting for BSDF visible above. It also contains coverage if it has not been baked into the BSDF parameters when using operation with parameter blending.
+
 #if STRATA_INLINE_SHADING
 	float3				Emissive;	// All bsdf can potientially have emissive. This is never written out to the strata gbuffer but instead accumulated in the emissive buffer during the base pass.
 	float				ThicknessCm;// This must be kept in order to be able to normalize the mean free path to comply with our volumetric material math all done in a normalised slab of medium.
@@ -568,9 +575,9 @@ FStrataData GetStrataSlabBSDF(
 #if STRATA_INLINE_SHADING
 	// At this stage we do not know if this strata BSDF is going to be using simple volume lighting so always store the DMFP on the side to not override SSSPROFILE data.
 	StrataData.Layers[0].BSDFs[0].TmpDMFP = SSSDMFP * saturate(SSSDMFPScale);
-#endif
 
-	StrataData.Layers[0].BSDFs[0].Weight						= 1.0f;
+	StrataData.Layers[0].BSDFs[0].Coverage = 1.0f;
+#endif
 	return StrataData;
 }
 
@@ -584,7 +591,9 @@ FStrataData GetStrataVolumeFogCloudBSDF(float3 Albedo, float3 Extinction, float3
 	VOLUMETRICFOGCLOUD_EXTINCTION	(StrataData.Layers[0].BSDFs[0]) = Extinction;
 	VOLUMETRICFOGCLOUD_AO			(StrataData.Layers[0].BSDFs[0]) = AmbientOcclusion;
 
-	StrataData.Layers[0].BSDFs[0].Weight = 1.0f;
+#if STRATA_INLINE_SHADING
+	StrataData.Layers[0].BSDFs[0].Coverage = 1.0f;
+#endif
 
 	return StrataData;
 }
@@ -597,7 +606,9 @@ FStrataData GetStrataUnlitBSDF(float3 Emissive, float3 TransmittanceColor)
 	BSDF_SETEMISSIVE				(StrataData.Layers[0].BSDFs[0], Emissive);
 	UNLIT_TRANSMITTANCE				(StrataData.Layers[0].BSDFs[0]) = TransmittanceColor;
 
-	StrataData.Layers[0].BSDFs[0].Weight = 1.0f;
+#if STRATA_INLINE_SHADING
+	StrataData.Layers[0].BSDFs[0].Coverage = 1.0f;
+#endif
 
 	return StrataData;
 }
@@ -615,7 +626,9 @@ FStrataData GetStrataHairBSDF(float3 BaseColor, float Scatter, float Specular, f
 	HAIR_SPECULAR					(StrataData.Layers[0].BSDFs[0]) = Specular;
 	HAIR_BACKLIT					(StrataData.Layers[0].BSDFs[0]) = Backlit;
 
-	StrataData.Layers[0].BSDFs[0].Weight = 1.0f;
+#if STRATA_INLINE_SHADING
+	StrataData.Layers[0].BSDFs[0].Coverage = 1.0f;
+#endif
 
 	return StrataData;
 }
@@ -641,7 +654,9 @@ FStrataData GetStrataSingleLayerWaterBSDF(
 	SLW_COLORSCALEBEHINDWATER		(StrataData.Layers[0].BSDFs[0]) = ColorScaleBehindWater;
 #endif
 
-	StrataData.Layers[0].BSDFs[0].Weight = 1.0f;
+#if STRATA_INLINE_SHADING
+	StrataData.Layers[0].BSDFs[0].Coverage = 1.0f;
+#endif
 
 	return StrataData;
 }
@@ -865,7 +880,18 @@ FStrataData StrataWeight(FStrataData A, float Weight)
 		UNROLL
 		for (int BSDF = 0; BSDF < StrataData.Layers[Layer].BSDFCount; ++BSDF)
 		{
-			StrataData.Layers[Layer].BSDFs[BSDF].Weight *= SafeWeight;
+#if STRATA_INLINE_SHADING
+			StrataData.Layers[Layer].BSDFs[BSDF].Coverage *= SafeWeight;
+
+			if (StrataData.Layers[Layer].BSDFs[BSDF].bPremultipliedCoverage) 
+			{
+				// Only slab can follow parameter blending rule: we can encounter pre weighted lighting contribution according to coverage.
+				SLAB_DIFFUSEALBEDO(StrataData.Layers[Layer].BSDFs[BSDF]) *= SafeWeight;
+				SLAB_F0(StrataData.Layers[Layer].BSDFs[BSDF]) *= SafeWeight;
+				SLAB_F90(StrataData.Layers[Layer].BSDFs[BSDF]) *= SafeWeight;
+				SLAB_FUZZ_COLOR(StrataData.Layers[Layer].BSDFs[BSDF]) *= SafeWeight;
+			}
+#endif
 		}
 	}
 
@@ -909,13 +935,21 @@ FStrataData StrataVerticalLayering(FStrataData Top, FStrataData Base)
 FStrataData StrataHorizontalMixingParameterBlending(FStrataData Foreground, FStrataData Background, float Mix, uint NewNormalIndex)
 {
 	const float ForegroundMixFactor = saturate(Mix);
-	const float BackgroundMixFactor = 1.0 - ForegroundMixFactor;
 
 	FStrataData Result = GetInitialisedStrataDataWithReservedFirstSlot();
 #define ResultBSDF		Result.Layers[0].BSDFs[0]
 #define ForegroundBSDF	Foreground.Layers[0].BSDFs[0]
 #define BackgroundBSDF	Background.Layers[0].BSDFs[0]
 
+#if STRATA_INLINE_SHADING
+	const float ForeLightingFactor = ForegroundBSDF.bPremultipliedCoverage ? 1.0f : saturate(ForegroundBSDF.Coverage);
+	const float BackLightingFactor = BackgroundBSDF.bPremultipliedCoverage ? 1.0f : saturate(BackgroundBSDF.Coverage);
+	ResultBSDF.bPremultipliedCoverage = true;
+#else
+	const float ForeLightingFactor = 1.0f;
+	const float BackLightingFactor = 1.0f;
+#endif
+
 	// We can only parameter blend the slab 
 	BSDF_SETTYPE(ResultBSDF,			STRATA_BSDF_TYPE_SLAB);
 	BSDF_SETSHAREDLOCALBASISID(ResultBSDF,NewNormalIndex);
@@ -930,19 +964,20 @@ FStrataData StrataHorizontalMixingParameterBlending(FStrataData Foreground, FStr
 	BSDF_SETHASF90(ResultBSDF,			BSDF_GETHASF90(ForegroundBSDF)					| BSDF_GETHASF90(BackgroundBSDF));
 	BSDF_SETHASTHINFILM(ResultBSDF,		BSDF_GETHASTHINFILM(ForegroundBSDF)				| BSDF_GETHASTHINFILM(BackgroundBSDF));
 	BSDF_SETHASFUZZ(ResultBSDF,			BSDF_GETHASFUZZ(ForegroundBSDF)					| BSDF_GETHASFUZZ(BackgroundBSDF));
-	BSDF_SETEMISSIVE(ResultBSDF,		lerp(BSDF_GETEMISSIVE(BackgroundBSDF),			BSDF_GETEMISSIVE(ForegroundBSDF),			ForegroundMixFactor));
-	BSDF_SETTHICKNESSCM(ResultBSDF,		lerp(BSDF_GETTHICKNESSCM(BackgroundBSDF),		BSDF_GETTHICKNESSCM(ForegroundBSDF),		ForegroundMixFactor));
-
-	SLAB_DIFFUSEALBEDO(ResultBSDF) =	lerp(SLAB_DIFFUSEALBEDO(BackgroundBSDF),		SLAB_DIFFUSEALBEDO(ForegroundBSDF),			ForegroundMixFactor);
-	SLAB_F0(ResultBSDF) =				lerp(SLAB_F0(BackgroundBSDF),					SLAB_F0(ForegroundBSDF),					ForegroundMixFactor);
-	SLAB_F90(ResultBSDF) =				lerp(SLAB_F90(BackgroundBSDF),					SLAB_F90(ForegroundBSDF),					ForegroundMixFactor);
-	SLAB_METALLIC(ResultBSDF) =			lerp(SLAB_METALLIC(BackgroundBSDF),				SLAB_METALLIC(ForegroundBSDF),				ForegroundMixFactor);
-	SLAB_ROUGHNESSX(ResultBSDF) =		lerp(SLAB_ROUGHNESSX(BackgroundBSDF),			SLAB_ROUGHNESSX(ForegroundBSDF),			ForegroundMixFactor);
-	SLAB_ROUGHNESSY(ResultBSDF) =		lerp(SLAB_ROUGHNESSY(BackgroundBSDF),			SLAB_ROUGHNESSY(ForegroundBSDF),			ForegroundMixFactor);
-	SLAB_HAZINESS(ResultBSDF) =			lerp(SLAB_HAZINESS(BackgroundBSDF),				SLAB_HAZINESS(ForegroundBSDF),				ForegroundMixFactor);
-	SLAB_THINFILM_THICKNESS(ResultBSDF)=lerp(SLAB_THINFILM_THICKNESS(BackgroundBSDF),	SLAB_THINFILM_THICKNESS(ForegroundBSDF),	ForegroundMixFactor);
-	SLAB_FUZZ_AMOUNT(ResultBSDF) =		lerp(SLAB_FUZZ_AMOUNT(BackgroundBSDF),			SLAB_FUZZ_AMOUNT(ForegroundBSDF),			ForegroundMixFactor);
-	SLAB_FUZZ_COLOR(ResultBSDF)	=		lerp(SLAB_FUZZ_COLOR(BackgroundBSDF),			SLAB_FUZZ_COLOR(ForegroundBSDF),			ForegroundMixFactor);
+
+	BSDF_SETEMISSIVE(ResultBSDF,		lerp(BackLightingFactor * BSDF_GETEMISSIVE(BackgroundBSDF),			ForeLightingFactor * BSDF_GETEMISSIVE(ForegroundBSDF),			ForegroundMixFactor));
+	SLAB_DIFFUSEALBEDO(ResultBSDF) =	lerp(BackLightingFactor * SLAB_DIFFUSEALBEDO(BackgroundBSDF),		ForeLightingFactor * SLAB_DIFFUSEALBEDO(ForegroundBSDF),		ForegroundMixFactor);
+	SLAB_F0(ResultBSDF) =				lerp(BackLightingFactor * SLAB_F0(BackgroundBSDF),					ForeLightingFactor * SLAB_F0(ForegroundBSDF),					ForegroundMixFactor);
+	SLAB_F90(ResultBSDF) =				lerp(BackLightingFactor * SLAB_F90(BackgroundBSDF),					ForeLightingFactor * SLAB_F90(ForegroundBSDF),					ForegroundMixFactor);
+	SLAB_FUZZ_COLOR(ResultBSDF)	=		lerp(BackLightingFactor * SLAB_FUZZ_COLOR(BackgroundBSDF),			ForeLightingFactor * SLAB_FUZZ_COLOR(ForegroundBSDF),			ForegroundMixFactor);
+
+	BSDF_SETTHICKNESSCM(ResultBSDF,		lerp(BSDF_GETTHICKNESSCM(BackgroundBSDF),							BSDF_GETTHICKNESSCM(ForegroundBSDF),							ForegroundMixFactor));
+	SLAB_METALLIC(ResultBSDF) =			lerp(SLAB_METALLIC(BackgroundBSDF),									SLAB_METALLIC(ForegroundBSDF),									ForegroundMixFactor);
+	SLAB_ROUGHNESSX(ResultBSDF) =		lerp(SLAB_ROUGHNESSX(BackgroundBSDF),								SLAB_ROUGHNESSX(ForegroundBSDF),								ForegroundMixFactor);
+	SLAB_ROUGHNESSY(ResultBSDF) =		lerp(SLAB_ROUGHNESSY(BackgroundBSDF),								SLAB_ROUGHNESSY(ForegroundBSDF),								ForegroundMixFactor);
+	SLAB_HAZINESS(ResultBSDF) =			lerp(SLAB_HAZINESS(BackgroundBSDF),									SLAB_HAZINESS(ForegroundBSDF),									ForegroundMixFactor);
+	SLAB_THINFILM_THICKNESS(ResultBSDF) = 0;// Thinfilm does not support coverage applied on pararmeters so we simply disable ThinFilm.=lerp(SLAB_THINFILM_THICKNESS(BackgroundBSDF),	SLAB_THINFILM_THICKNESS(ForegroundBSDF),	ForegroundMixFactor);
+	SLAB_FUZZ_AMOUNT(ResultBSDF) =		lerp(SLAB_FUZZ_AMOUNT(BackgroundBSDF),								SLAB_FUZZ_AMOUNT(ForegroundBSDF),								ForegroundMixFactor);
 
 	if (BSDF_GETHASSSSPROFILE(ResultBSDF))
 	{
@@ -967,7 +1002,7 @@ FStrataData StrataHorizontalMixingParameterBlending(FStrataData Foreground, FStr
 	}
 	else
 	{
-		SLAB_SSSDMFP(ResultBSDF) = lerp(SLAB_SSSDMFP(BackgroundBSDF),				SLAB_SSSDMFP(ForegroundBSDF),					ForegroundMixFactor);
+		SLAB_SSSDMFP(ResultBSDF) =		lerp(SLAB_SSSDMFP(BackgroundBSDF),								SLAB_SSSDMFP(ForegroundBSDF),										ForegroundMixFactor);
 
 		// We still try to disable SSS if possible: for performance and also to preserve the surface color (especially specular reflection color).
 		if (all(SLAB_SSSDMFP(ResultBSDF) == 0.0f))
@@ -979,11 +1014,11 @@ FStrataData StrataHorizontalMixingParameterBlending(FStrataData Foreground, FStr
 
 #if STRATA_INLINE_SHADING
 	// At this stage we do not know if this strata BSDF is going to be using simple volume lighting so always store the DMFP on the side to not override SSSPROFILE data.
-	ResultBSDF.TmpDMFP =				lerp(BackgroundBSDF.TmpDMFP,					ForegroundBSDF.TmpDMFP,						ForegroundMixFactor);
-#endif
+	ResultBSDF.TmpDMFP =				lerp(BackgroundBSDF.TmpDMFP,									ForegroundBSDF.TmpDMFP,												ForegroundMixFactor);
 
-	// We blend the parameter above, and we also need to lerp the weight to match the total energy and coverage.
-	ResultBSDF.Weight =					lerp(BackgroundBSDF.Weight,						ForegroundBSDF.Weight,						ForegroundMixFactor);
+	// We blend the parameter above, and we also need to lerp the Coverage to match the total energy.
+	ResultBSDF.Coverage =				lerp(BackgroundBSDF.Coverage,									ForegroundBSDF.Coverage,											ForegroundMixFactor);
+#endif
 
 #undef ResultBSDF
 #undef ForegroundBSDF
@@ -991,14 +1026,16 @@ FStrataData StrataHorizontalMixingParameterBlending(FStrataData Foreground, FStr
 	return Result;
 }
 
-float AddParameterBlendingBSDFWeightToNormalMix(FStrataData AStrata, FStrataData BStrata) 
+float AddParameterBlendingBSDFCoverageToNormalMix(FStrataData AStrata, FStrataData BStrata) 
 {
-	const float3 SafeABSDFWeight = saturate(AStrata.Layers[0].BSDFs[0].Weight);
-	const float3 SafeBBSDFWeight = saturate(BStrata.Layers[0].BSDFs[0].Weight);
-	const float SafeGreyABSDFWeight = dot(SafeABSDFWeight, 1.0f / 3.0f);
-	const float SafeGreyBBSDFWeight = dot(SafeBBSDFWeight, 1.0f / 3.0f);
-	const float AMixFactor = SafeGreyABSDFWeight / (1e-10 + SafeGreyABSDFWeight + SafeGreyBBSDFWeight);
+#if STRATA_INLINE_SHADING
+	const float SafeABSDFCoverage = saturate(AStrata.Layers[0].BSDFs[0].Coverage);
+	const float SafeBBSDFCoverage = saturate(BStrata.Layers[0].BSDFs[0].Coverage);
+	const float AMixFactor = SafeABSDFCoverage / (1e-10 + SafeABSDFCoverage + SafeBBSDFCoverage);
 	return AMixFactor;
+#else
+	return 0.5f;
+#endif
 }
 
 // ==> NOTE: Always pair with the compiler behavior in StrataCompilationInfoAddParamBlend
@@ -1009,8 +1046,14 @@ FStrataData StrataAddParameterBlending(FStrataData A, FStrataData B, float AMixF
 #define ABSDF			A.Layers[0].BSDFs[0]
 #define BBSDF			B.Layers[0].BSDFs[0]
 
-	const float3 SafeABSDFWeight = saturate(ABSDF.Weight);
-	const float3 SafeBBSDFWeight = saturate(BBSDF.Weight);
+#if STRATA_INLINE_SHADING
+	const float ALightingFactor = ABSDF.bPremultipliedCoverage ? 1.0f : saturate(ABSDF.Coverage);
+	const float BLightingFactor = BBSDF.bPremultipliedCoverage ? 1.0f : saturate(BBSDF.Coverage);
+	ResultBSDF.bPremultipliedCoverage = true;
+#else
+	const float ALightingFactor = 1.0f;
+	const float BLightingFactor = 1.0f;
+#endif
 
 	// We can only parameter blend the slab 
 	BSDF_SETTYPE(ResultBSDF,			STRATA_BSDF_TYPE_SLAB);
@@ -1026,21 +1069,23 @@ FStrataData StrataAddParameterBlending(FStrataData A, FStrataData B, float AMixF
 	BSDF_SETHASF90(ResultBSDF,			BSDF_GETHASF90(BBSDF)									| BSDF_GETHASF90(ABSDF));
 	BSDF_SETHASTHINFILM(ResultBSDF,		BSDF_GETHASTHINFILM(BBSDF)								| BSDF_GETHASTHINFILM(ABSDF));
 	BSDF_SETHASFUZZ(ResultBSDF,			BSDF_GETHASFUZZ(BBSDF)									| BSDF_GETHASFUZZ(ABSDF));
-	BSDF_SETEMISSIVE(ResultBSDF,		BSDF_GETEMISSIVE(BBSDF) * SafeBBSDFWeight				+ BSDF_GETEMISSIVE(ABSDF) * SafeABSDFWeight);
-	BSDF_SETTHICKNESSCM(ResultBSDF,		lerp(BSDF_GETTHICKNESSCM(BBSDF)							, BSDF_GETTHICKNESSCM(ABSDF),				AMixFactor));
+
 
 	// Some parameters will contribute to added luminance: in this case we add them together and saturate to not go out of the safe range
+	BSDF_SETEMISSIVE(ResultBSDF,		BLightingFactor * BSDF_GETEMISSIVE(BBSDF)				+ ALightingFactor * BSDF_GETEMISSIVE(ABSDF));
+	SLAB_DIFFUSEALBEDO(ResultBSDF) =	saturate(BLightingFactor * SLAB_DIFFUSEALBEDO(BBSDF)	+ ALightingFactor * SLAB_DIFFUSEALBEDO(ABSDF));
+	SLAB_F0(ResultBSDF) =				saturate(BLightingFactor * SLAB_F0(BBSDF)				+ ALightingFactor * SLAB_F0(ABSDF));
+	SLAB_F90(ResultBSDF) =				saturate(BLightingFactor * SLAB_F90(BBSDF)				+ ALightingFactor * SLAB_F90(ABSDF));
+	SLAB_FUZZ_COLOR(ResultBSDF) =		saturate(BLightingFactor * SLAB_FUZZ_COLOR(BBSDF)		+ ALightingFactor * SLAB_FUZZ_COLOR(ABSDF));
+
 	// Some parameters are not contributing to added luminance: in this case we simply lerp them
-	SLAB_DIFFUSEALBEDO(ResultBSDF) =	saturate(SLAB_DIFFUSEALBEDO(BBSDF) * SafeBBSDFWeight	+ SLAB_DIFFUSEALBEDO(ABSDF)	* SafeABSDFWeight);
-	SLAB_F0(ResultBSDF) =				saturate(SLAB_F0(BBSDF) * SafeBBSDFWeight				+ SLAB_F0(ABSDF) * SafeABSDFWeight);
-	SLAB_F90(ResultBSDF) =				saturate(SLAB_F90(BBSDF) * SafeBBSDFWeight				+ SLAB_F90(ABSDF) * SafeABSDFWeight);
+	BSDF_SETTHICKNESSCM(ResultBSDF,		lerp(BSDF_GETTHICKNESSCM(BBSDF),						BSDF_GETTHICKNESSCM(ABSDF),					AMixFactor));
 	SLAB_METALLIC(ResultBSDF) =			lerp(SLAB_METALLIC(BBSDF),								SLAB_METALLIC(ABSDF),						AMixFactor);
 	SLAB_ROUGHNESSX(ResultBSDF) =		lerp(SLAB_ROUGHNESSX(BBSDF),							SLAB_ROUGHNESSX(ABSDF),						AMixFactor);
 	SLAB_ROUGHNESSY(ResultBSDF) =		lerp(SLAB_ROUGHNESSY(BBSDF),							SLAB_ROUGHNESSY(ABSDF),						AMixFactor);
 	SLAB_HAZINESS(ResultBSDF) =			lerp(SLAB_HAZINESS(BBSDF),								SLAB_HAZINESS(ABSDF),						AMixFactor);
-	SLAB_THINFILM_THICKNESS(ResultBSDF)=lerp(SLAB_THINFILM_THICKNESS(BBSDF),					SLAB_THINFILM_THICKNESS(ABSDF),				AMixFactor);
+	SLAB_THINFILM_THICKNESS(ResultBSDF) = 0;// This does not support coverage applied on pararmeters so we simply disable ThinFilm. lerp(SLAB_THINFILM_THICKNESS(BBSDF), SLAB_THINFILM_THICKNESS(ABSDF), AMixFactor);
 	SLAB_FUZZ_AMOUNT(ResultBSDF) =		lerp(SLAB_FUZZ_AMOUNT(BBSDF),							SLAB_FUZZ_AMOUNT(ABSDF),					AMixFactor);
-	SLAB_FUZZ_COLOR(ResultBSDF) =		saturate(SLAB_FUZZ_COLOR(BBSDF) * SafeBBSDFWeight		+ SLAB_FUZZ_COLOR(ABSDF) * SafeABSDFWeight);
 
 	if (BSDF_GETHASSSSPROFILE(ResultBSDF))
 	{
@@ -1064,12 +1109,12 @@ FStrataData StrataAddParameterBlending(FStrataData A, FStrataData B, float AMixF
 
 #if STRATA_INLINE_SHADING
 	// At this stage we do not know if this strata BSDF is going to be using simple volume lighting so always store the DMFP on the side to not override SSSPROFILE data.
-	ResultBSDF.TmpDMFP =				lerp(BBSDF.TmpDMFP,										ABSDF.TmpDMFP,								AMixFactor);
-#endif
+	ResultBSDF.TmpDMFP =				lerp(BBSDF.TmpDMFP * BLightingFactor,					ABSDF.TmpDMFP * ALightingFactor,			AMixFactor);
 
-	// When it comes to adding the visual contribution, we add the weights without clamping as currently done in the StrataAdd function.
-	// Operations using weights are currently making sure they are in a safe range.
-	ResultBSDF.Weight =	 BBSDF.Weight + ABSDF.Weight;
+	// When it comes to adding the visual contribution, we chose to add the coverage without clamping as currently done in the StrataAdd function.
+	// Operations using Coverage are currently making sure they are in a safe range.
+	ResultBSDF.Coverage =	 BBSDF.Coverage + ABSDF.Coverage;
+#endif
 
 #undef ResultBSDF
 #undef ABSDF
@@ -1077,12 +1122,15 @@ FStrataData StrataAddParameterBlending(FStrataData A, FStrataData B, float AMixF
 	return Result;
 }
 
-float VerticalLayeringParameterBlendingBSDFWeightToNormalMix(FStrataData TopStrata)
+float VerticalLayeringParameterBlendingBSDFCoverageToNormalMix(FStrataData TopStrata)
 {
-	const float3 SafeTopBSDFWeight = saturate(TopStrata.Layers[0].BSDFs[0].Weight);
-	const float SafeGreyTopBSDFWeight = dot(SafeTopBSDFWeight, 1.0f / 3.0f);
+#if STRATA_INLINE_SHADING
+	const float SafeTopBSDFCoverage = saturate(TopStrata.Layers[0].BSDFs[0].Coverage);
 	const float TopNormalContributionForWeightOfOne = 0.5f;
-	return SafeGreyTopBSDFWeight * TopNormalContributionForWeightOfOne;
+	return SafeTopBSDFCoverage * TopNormalContributionForWeightOfOne;
+#else
+	return 0.5f;
+#endif
 }
 
 // ==> NOTE: Always pair with the compiler behavior in StrataCompilationInfoVerticalLayeringParamBlend
@@ -1093,6 +1141,19 @@ FStrataData StrataVerticalLayeringParameterBlending(FStrataData Top, FStrataData
 #define TopBSDF			Top.Layers[0].BSDFs[0]
 #define BaseBSDF		Base.Layers[0].BSDFs[0]
 
+#if STRATA_INLINE_SHADING
+	const float SafeTopCoverage = saturate(TopBSDF.Coverage);
+	const float SafeBaseCoverage = saturate(BaseBSDF.Coverage);
+	const float TopLightingFactor  = TopBSDF.bPremultipliedCoverage  ? 1.0f : SafeTopCoverage;
+	const float BaseLightingFactor = BaseBSDF.bPremultipliedCoverage ? 1.0f : SafeBaseCoverage;
+	ResultBSDF.bPremultipliedCoverage = true;
+#else
+	const float SafeTopCoverage = 1.0f;
+	const float SafeBaseCoverage = 1.0f;
+	const float TopLightingFactor = 1.0f;
+	const float BaseLightingFactor = 1.0f;
+#endif
+
 	// We can only parameter blend the slab 
 	BSDF_SETTYPE(ResultBSDF,			STRATA_BSDF_TYPE_SLAB);
 	BSDF_SETSHAREDLOCALBASISID(ResultBSDF,NewNormalIndex);
@@ -1102,7 +1163,7 @@ FStrataData StrataVerticalLayeringParameterBlending(FStrataData Top, FStrataData
 	BSDF_SETHASSSS(ResultBSDF,			BSDF_GETHASSSS(BaseBSDF));															// We only keep SSS if the bottom layer has it
 	BSDF_SETHASSSSPROFILE(ResultBSDF,	0);																					// Problem: cannot blend SSS profiles so we simply disable SSSProfile
 	BSDF_SETISTHIN(ResultBSDF,			BSDF_GETISTHIN(BaseBSDF));															// STRATA_TODO fix the MPF/DMFP situation for that to really work
-	BSDF_SETHASDMFP(ResultBSDF,			BSDF_GETHASDMFP(BaseBSDF));															//We only keep SSS if the bottom layer has it
+	BSDF_SETHASDMFP(ResultBSDF,			BSDF_GETHASDMFP(BaseBSDF));															// We only keep SSS if the bottom layer has it
 	BSDF_SETHASHAZINESS(ResultBSDF,		BSDF_GETHASHAZINESS(TopBSDF));														// We only keep haziness if the top layer has it
 	BSDF_SETHASF90(ResultBSDF,			BSDF_GETHASF90(BaseBSDF)						| BSDF_GETHASF90(TopBSDF));			// F90: keep union of both even though it will be hard to get a match
 	BSDF_SETHASTHINFILM(ResultBSDF,		BSDF_GETHASTHINFILM(TopBSDF));														// We only keep thinfilme if the top layer has it
@@ -1115,37 +1176,32 @@ FStrataData StrataVerticalLayeringParameterBlending(FStrataData Top, FStrataData
 	FParticipatingMedia PM				= StrataSlabCreateParticipatingMedia(UnusedDiffuseColor, SLAB_SSSDMFP(TopBSDF));
 	const float3 SlabTransmittance		= IsotropicMediumSlabTransmittance(PM, STRATA_SIMPLEVOLUME_THICKNESS_M, NoV);
 	// Compute the colored throughput, taking into account the weight.
-	const float3 SafeTopBSDFWeight		= saturate(TopBSDF.Weight);
-	const float3 ColorThroughput		= (1.0f - SafeTopBSDFWeight) + SafeTopBSDFWeight * SlabTransmittance;
+	const float3 ColorThroughput		= saturate((1.0f - SafeTopCoverage) + SafeTopCoverage * SlabTransmittance);
 	const float GreyThroughput			= dot(ColorThroughput, 1.0f / 3.0f);
 	
 	// Add both layers thickness
-	BSDF_SETTHICKNESSCM(ResultBSDF,		BSDF_GETTHICKNESSCM(BaseBSDF)					+ SafeTopBSDFWeight * BSDF_GETTHICKNESSCM(TopBSDF));
+	BSDF_SETTHICKNESSCM(ResultBSDF,		BSDF_GETTHICKNESSCM(BaseBSDF)												+ TopLightingFactor * BSDF_GETTHICKNESSCM(TopBSDF));
 
 	// Now add bottom layer component weights by throughput to the top layer.
 	// Those parameters will contribute to added luminance, in this case we add them together and saturate to not go out of the safe range
 	// Some parameters are not contributing to added luminance: in this case we simply lerp them
-	BSDF_SETEMISSIVE(ResultBSDF,		ColorThroughput * BSDF_GETEMISSIVE(BaseBSDF)	+ SafeTopBSDFWeight * BSDF_GETEMISSIVE(TopBSDF));
-
-	// Strata TODO: we should also add the scattering light contribution from the top layer there.
-	SLAB_DIFFUSEALBEDO(ResultBSDF) =	saturate(ColorThroughput * SLAB_DIFFUSEALBEDO(BaseBSDF) + SafeTopBSDFWeight * SLAB_DIFFUSEALBEDO(TopBSDF));
-	SLAB_F0(ResultBSDF) =				saturate(ColorThroughput * SLAB_F0(BaseBSDF)	+ SafeTopBSDFWeight * SLAB_F0(TopBSDF));
-	SLAB_F90(ResultBSDF) =				saturate(ColorThroughput * SLAB_F90(BaseBSDF)	+ SafeTopBSDFWeight * SLAB_F90(TopBSDF));
+	BSDF_SETEMISSIVE(ResultBSDF,		ColorThroughput * BaseLightingFactor * BSDF_GETEMISSIVE(BaseBSDF)			+ TopLightingFactor * BSDF_GETEMISSIVE(TopBSDF));
 
-	SLAB_FUZZ_AMOUNT(ResultBSDF) =		saturate(GreyThroughput * SLAB_FUZZ_AMOUNT(BaseBSDF) + SafeTopBSDFWeight * SLAB_FUZZ_AMOUNT(TopBSDF));
-	SLAB_FUZZ_COLOR(ResultBSDF) =		saturate(ColorThroughput * SLAB_FUZZ_COLOR(BaseBSDF) + SafeTopBSDFWeight * SLAB_FUZZ_COLOR(TopBSDF));
-
-	// Simple blend for now. 
-	// STRATA_TODO find someting smarter?
+	SLAB_DIFFUSEALBEDO(ResultBSDF) =	saturate(ColorThroughput * BaseLightingFactor * SLAB_DIFFUSEALBEDO(BaseBSDF)+ TopLightingFactor * SLAB_DIFFUSEALBEDO(TopBSDF));
+	SLAB_F0(ResultBSDF) =				saturate(ColorThroughput * BaseLightingFactor * SLAB_F0(BaseBSDF)			+ TopLightingFactor * SLAB_F0(TopBSDF));
+	SLAB_F90(ResultBSDF) =				saturate(ColorThroughput * BaseLightingFactor * SLAB_F90(BaseBSDF)			+ TopLightingFactor * SLAB_F90(TopBSDF));
+	SLAB_FUZZ_COLOR(ResultBSDF) =		saturate(ColorThroughput * BaseLightingFactor * SLAB_FUZZ_COLOR(BaseBSDF)	+ TopLightingFactor * SLAB_FUZZ_COLOR(TopBSDF));
+		
 	SLAB_METALLIC(ResultBSDF) =			lerp(SLAB_METALLIC(TopBSDF),					SLAB_METALLIC(BaseBSDF),					GreyThroughput);
 	SLAB_ROUGHNESSX(ResultBSDF) =		lerp(SLAB_ROUGHNESSX(TopBSDF),					SLAB_ROUGHNESSX(BaseBSDF),					GreyThroughput);
 	SLAB_HAZINESS(ResultBSDF) =			lerp(SLAB_HAZINESS(TopBSDF),					SLAB_HAZINESS(BaseBSDF),					GreyThroughput);
+	SLAB_METALLIC(ResultBSDF) =			lerp(SLAB_FUZZ_AMOUNT(TopBSDF),					SLAB_FUZZ_AMOUNT(BaseBSDF),					GreyThroughput);
 
 	// Anisotropy is disabled with vertical blending. 
 	SLAB_ROUGHNESSY(ResultBSDF) =		SLAB_ROUGHNESSX(ResultBSDF);
 
-	// Only keep thin film and for the top layer.
-	SLAB_THINFILM_THICKNESS(ResultBSDF)=SLAB_THINFILM_THICKNESS(TopBSDF);
+	// Only keep thin film for the top layer.
+	SLAB_THINFILM_THICKNESS(ResultBSDF)= (SafeTopCoverage >= 1.0f) ? SLAB_THINFILM_THICKNESS(TopBSDF) : 0.0f;
 
 	// Keep the DMFP form the bottom layer only (top layer simple volume is backed in parameters).
 	SLAB_SSSDMFP(ResultBSDF) =			SLAB_SSSDMFP(BaseBSDF);
@@ -1153,10 +1209,9 @@ FStrataData StrataVerticalLayeringParameterBlending(FStrataData Top, FStrataData
 #if STRATA_INLINE_SHADING
 	// Only the bottom layer SSS / Thin model is used so we only keep the bottom layer data there. Top layer is converted to throughput.
 	ResultBSDF.TmpDMFP =				BaseBSDF.TmpDMFP;
-#endif
 
-	// The weight is the maximum of the two
-	ResultBSDF.Weight =					max(BaseBSDF.Weight,							TopBSDF.Weight);
+	ResultBSDF.Coverage =				1.0f - (1.0f - SafeTopCoverage) * (1.0f - SafeBaseCoverage);
+#endif
 
 #undef ResultBSDF
 #undef TopBSDF
@@ -1632,7 +1687,7 @@ void PackStrataOut(
 	// We clamp the layer count to the maximum one according to our byte per pixel budget (removing the bottom layer first)
 	Strata.LayerCount = STRATA_CLAMPED_LAYER_COUNT;
 
-	// For each of the BSDF, we evaluated their combined weight/throughput according to the material operation performed (layering, mixing, etc.).
+	// For each of the BSDF, we evaluated their combined coverage/throughput according to the material operation performed (layering, mixing, etc.).
 	int l;
 	STRATA_UNROLL_N(STRATA_CLAMPED_LAYER_COUNT)
 	for (l = 0; l < Strata.LayerCount; ++l)
@@ -1654,8 +1709,12 @@ void PackStrataOut(
 			StrataSanitizeBSDF(Strata.Layers[l].BSDFs[i]);
 
 			FStrataBSDF BSDF = Strata.Layers[l].BSDFs[i];
-			const float3 BSDFCoverage = BSDF.Weight;
-			const float  BSDFCoverageAvg = dot(BSDF.Weight,1.0f/3.0f);
+
+#if STRATA_INLINE_SHADING
+			const float BSDFCoverage = saturate(BSDF.Coverage);
+#else
+			const float BSDFCoverage = 1.0f;
+#endif
 			BSDFCount++;
 
 			// STRATA_TODO apply the same normal transform as for the material
@@ -1679,7 +1738,7 @@ void PackStrataOut(
 				// We want simple volumetric only if the DMFP input is pluged in, otherwise we avoid storing the DFMP
 				const bool bIsSimpleVolume = !bBottomLayer && BSDF_GETHASDMFP(Strata.Layers[l].BSDFs[i]);
 
-				bStrataSubsurfaceEnable = bStrataSubsurfaceEnable || (BSDF_GETHASSSS(BSDF) && BSDFCoverageAvg > 0 && !bIsSimpleVolume);
+				bStrataSubsurfaceEnable = bStrataSubsurfaceEnable || (BSDF_GETHASSSS(BSDF) && BSDFCoverage > 0 && !bIsSimpleVolume);
 
 				// By default, without coverage, the diffuse BSDF is completely opaque
 				float3 BSDFThroughput = OpaqueBSDFThroughput;
@@ -1759,8 +1818,13 @@ void PackStrataOut(
 
 			} // switch
 
-			// Update the weights to take into account throughput towards the view from layers above
-			Strata.Layers[l].BSDFs[i].Weight *= Throughput;
+
+#if STRATA_INLINE_SHADING
+			// Update the BSDF final weight with throughput towards the view from layers above and coverage if it has not been baked into the parameters.
+			Strata.Layers[l].BSDFs[i].FinalWeight = Throughput * (Strata.Layers[l].BSDFs[i].bPremultipliedCoverage ? 1.0 : Strata.Layers[l].BSDFs[i].Coverage);
+#else
+			Strata.Layers[l].BSDFs[i].FinalWeight = Throughput;
+#endif
 			
 			// Notify that the BSDF is at the top for SSR to only affect reflection there and not on the lower layers
 			BSDF_SETISTOPLAYER(Strata.Layers[l].BSDFs[i], bTopLayer ? 1 : 0);
@@ -1805,17 +1869,17 @@ void PackStrataOut(
 		{
 			FStrataBSDF BSDF = Layer.BSDFs[i];
 
-			const uint GreyScaleWeight = StrataHasGreyScaleWeight(BSDF.Weight) ? 1 : 0;
-			BSDF_SETHASGREYWEIGHT(BSDF, GreyScaleWeight);
-			if (GreyScaleWeight > 0)
+			const uint GreyScaleThroughput = StrataHasGreyScaleWeight(BSDF.FinalWeight) ? 1 : 0;
+			BSDF_SETHASGREYWEIGHT(BSDF, GreyScaleThroughput);
+			if (GreyScaleThroughput > 0)
 			{
-				BSDF_SETWEIGHT10F(BSDF, Pack10F(BSDF.Weight.x));
+				BSDF_SETWEIGHT10F(BSDF, Pack10F(BSDF.FinalWeight.x));
 				STRATA_STORE_UINT1(BSDF.State);
 			}
 			else
 			{
 				STRATA_STORE_UINT1(BSDF.State);
-				STRATA_STORE_UINT1(PackR11G11B10F(BSDF.Weight));
+				STRATA_STORE_UINT1(PackR11G11B10F(BSDF.FinalWeight));
 			}
 
 			const uint BSDFType = BSDF_GETTYPE(BSDF);
@@ -1966,7 +2030,7 @@ FStrataBSDF UnpackFastPathStrataBSDFIn(ByteAddressBuffer StrataBuffer, inout FSt
 	OutBSDF.State = StrataLoadUint1(StrataBuffer, StrataAddressing);
 
 	// In the fast path the single slab BSDF is always on top so the weight must be grey scale.
-	OutBSDF.Weight = Unpack10F(BSDF_GETWEIGHT10F(OutBSDF));
+	OutBSDF.FinalWeight = Unpack10F(BSDF_GETWEIGHT10F(OutBSDF));
 
 	// Now unpack the simple Slab BSDF
 	uint DataX = StrataLoadUint1(StrataBuffer, StrataAddressing);
@@ -2004,11 +2068,11 @@ FStrataBSDF UnpackStrataBSDFIn(ByteAddressBuffer StrataBuffer, inout FStrataAddr
 	STRATA_LOAD_UINT1(OutBSDF.State);
 	if (BSDF_GETHASGREYWEIGHT(OutBSDF))
 	{
-		OutBSDF.Weight = Unpack10F(BSDF_GETWEIGHT10F(OutBSDF));
+		OutBSDF.FinalWeight = Unpack10F(BSDF_GETWEIGHT10F(OutBSDF));
 	}
 	else
 	{
-		StrataLoad_R11G11B10F(StrataBuffer, StrataAddressing, OutBSDF.Weight);
+		StrataLoad_R11G11B10F(StrataBuffer, StrataAddressing, OutBSDF.FinalWeight);
 	}
 
 
diff --git a/Engine/Shaders/Private/Strata/StrataDeferredLighting.ush b/Engine/Shaders/Private/Strata/StrataDeferredLighting.ush
--- a/Engine/Shaders/Private/Strata/StrataDeferredLighting.ush
+++ b/Engine/Shaders/Private/Strata/StrataDeferredLighting.ush
@@ -219,8 +219,8 @@ float4 StrataDeferredLighting(
 			float3 SpecularTotalLight	= ShadowTerms.SurfaceShadow			* (SpecularPathValue		* IntegrationContext.IncidentIrradianceNoShadow					+	SpecularHazePathValue		* IntegrationContext.IncidentIrradianceNoShadowHaze);
 
 			// Apply the BSDF weight containing coverage, transmittance to view and visibility.
-			DiffuseTotalLight  *= BSDF.Weight;
-			SpecularTotalLight *= BSDF.Weight;
+			DiffuseTotalLight  *= BSDF.FinalWeight;
+			SpecularTotalLight *= BSDF.FinalWeight;
 
 			const float IdentityCommonMultiplier = 1.0f;
 			LightAccumulator_AddSplit(Out, DiffuseTotalLight, SpecularTotalLight, DiffuseTotalLight, IdentityCommonMultiplier, bNeedsSeparateSubsurfaceLightAccumulation);
diff --git a/Engine/Shaders/Private/Strata/StrataEnvironmentLighting.ush b/Engine/Shaders/Private/Strata/StrataEnvironmentLighting.ush
--- a/Engine/Shaders/Private/Strata/StrataEnvironmentLighting.ush
+++ b/Engine/Shaders/Private/Strata/StrataEnvironmentLighting.ush
@@ -43,7 +43,7 @@ float4 StrataEnvironmentLighting(
 	#else
 		FStrataBSDF BSDF = UnpackStrataBSDFIn(StrataDataBuffer, StrataAddressing);
 	#endif
-		const float3 BSDFThroughput = BSDF.Weight;
+		const float3 BSDFThroughput = BSDF.FinalWeight;
 
 		// Create the BSDF context
 		FStrataBSDFContext StrataBSDFContext = StrataCreateBSDFContext(StrataPixelHeader, BSDF, StrataAddressing, V);
diff --git a/Engine/Shaders/Private/Strata/StrataForwardLighting.ush b/Engine/Shaders/Private/Strata/StrataForwardLighting.ush
--- a/Engine/Shaders/Private/Strata/StrataForwardLighting.ush
+++ b/Engine/Shaders/Private/Strata/StrataForwardLighting.ush
@@ -52,7 +52,7 @@ float3 StrataForwardLighting(
 				}
 			}
 
-			const float3 BSDFCoverage = BSDF.Weight;
+			const float3 BSDFCoverage = BSDF.Coverage;
 
 			// Create the BSDF context
 			FStrataBSDFContext StrataBSDFContext = StrataCreateBSDFContext(Header, BSDF, NullStrataAddressing, V, L);
diff --git a/Engine/Shaders/Private/Strata/StrataFurnaceTest.usf b/Engine/Shaders/Private/Strata/StrataFurnaceTest.usf
--- a/Engine/Shaders/Private/Strata/StrataFurnaceTest.usf
+++ b/Engine/Shaders/Private/Strata/StrataFurnaceTest.usf
@@ -145,7 +145,7 @@ void MainPS(
 			BSDF_SETHASSSSPROFILE(BSDF, 0);
 			BSDF_SETHASHAZINESS(BSDF, Haziness != 0 ? 1 : 0);
 			BSDF_SETHASTHINFILM(BSDF, 0);
-			BSDF.Weight = 1.0f;
+			BSDF.FinalWeight = 1.0f;
 
 			const float2 RoughnessXY = GetAnisotropicRoughness(Roughness, Anisoptropy);
 
@@ -184,8 +184,8 @@ void MainPS(
 				FStrataBSDFContext BSDFContext = StrataCreateBSDFContext(TangentBasis, BSDF, V, L, true);
 				FStrataEvaluateResult BSDFEvaluate = StrataEvaluateBSDF(BSDFContext, ShadowTerms);
 
-				const float3 DiffusePathWeight  = BSDF.Weight * BSDFEvaluate.DiffusePathValue;
-				const float3 SpecularPathWeight = BSDF.Weight * BSDFEvaluate.SpecularPathValue;
+				const float3 DiffusePathWeight  = BSDF.FinalWeight * BSDFEvaluate.DiffusePathValue;
+				const float3 SpecularPathWeight = BSDF.FinalWeight * BSDFEvaluate.SpecularPathValue;
 
 				Acc += SatNoL * InRadiance * (DiffusePathWeight + SpecularPathWeight) / (S.w != 0 ?  S.w : 1);
 			}
@@ -213,8 +213,8 @@ void MainPS(
 				FStrataBSDFContext BSDFContext = StrataCreateBSDFContext(TangentBasis, BSDF, V, L, true);
 				FStrataEvaluateResult BSDFEvaluate = StrataEvaluateBSDF(BSDFContext, ShadowTerms);
 
-				const float3 DiffusePathWeight  = BSDF.Weight * BSDFEvaluate.DiffusePathValue;
-				const float3 SpecularPathWeight = BSDF.Weight * BSDFEvaluate.SpecularPathValue;
+				const float3 DiffusePathWeight  = BSDF.FinalWeight * BSDFEvaluate.DiffusePathValue;
+				const float3 SpecularPathWeight = BSDF.FinalWeight * BSDFEvaluate.SpecularPathValue;
 
 				Acc += SatNoL * InRadiance * (DiffusePathWeight + SpecularPathWeight) / (PDF != 0 ? PDF : 1);
 
@@ -223,7 +223,7 @@ void MainPS(
 		}
 		else if (IntegratorType == STRATA_INTEGRATION_ENV)
 		{
-			const float3 BSDFThroughput = BSDF.Weight;
+			const float3 BSDFThroughput = BSDF.FinalWeight;
 
 			const float3 UnusedL = float3(0, 0, 1);
 			const float3x3 TangentBasis = GetTangentBasis(N);
@@ -233,8 +233,8 @@ void MainPS(
 			FStrataEnvLightResult StrataEnvLight = StrataEvaluateForEnvLight(BSDFContext, bEnableSpecular);
 
 			const float3 InIrradiance = InRadiance; //It should have a Pi factor though
-			const float3 SpecularSkyLighting = BSDF.Weight * StrataEnvLight.SpecularWeight * InRadiance;
-			const float3 DiffuseSkyLighting  = BSDF.Weight * StrataEnvLight.DiffuseWeight * InIrradiance;
+			const float3 SpecularSkyLighting = BSDF.FinalWeight * StrataEnvLight.SpecularWeight * InRadiance;
+			const float3 DiffuseSkyLighting  = BSDF.FinalWeight * StrataEnvLight.DiffuseWeight * InIrradiance;
 
 			OutColor = 
 				SpecularSkyLighting + 
diff --git a/Engine/Shaders/Private/Strata/StrataMaterialClassification.usf b/Engine/Shaders/Private/Strata/StrataMaterialClassification.usf
--- a/Engine/Shaders/Private/Strata/StrataMaterialClassification.usf
+++ b/Engine/Shaders/Private/Strata/StrataMaterialClassification.usf
@@ -32,8 +32,8 @@ void ComputePostProcessBSDFData(
 	const float FullThroughput = 1.0f;
 	const float FullyRough = 1.0f;
 
-	const float3 BSDFWeight = BSDF.Weight;
-	const float  BSDFWeightAvg = dot(BSDF.Weight,1.0f/3.0f);
+	// FinalWeight does not always contains coverage (see Strata.ush)
+	const float  BSDFWeightAvg = dot(BSDF.FinalWeight, 1.0f/3.0f);
 
 	// STRATA_TODO apply the same normal transform as for the material
 	float3x3 TangentBasis = StrataGetBSDFSharedBasis(Header, BSDF, StrataAddressing);
diff --git a/Engine/Shaders/Private/Strata/StrataVisualize.usf b/Engine/Shaders/Private/Strata/StrataVisualize.usf
--- a/Engine/Shaders/Private/Strata/StrataVisualize.usf
+++ b/Engine/Shaders/Private/Strata/StrataVisualize.usf
@@ -186,7 +186,7 @@ void VisualizeMaterialPS(
 #define PRINTBSDFCHAR(X) PrintCharacter(PixelPos, OutColor.xyz, BSDFFontColor, TopLeft, X);
 #define PRINTBSDFCHARGREY(X) PrintCharacter(PixelPos, OutColor.xyz, BSDFFontColorGrey, TopLeft, X);
 
-#define PRINTBSDFWEIGHT()	PRINTBSDFCHAR(_SPC_); PRINTBSDFCHAR(_MINUS_); PRINTBSDFCHAR(_W_); PRINTBSDFCHAR(_E_); PRINTBSDFCHAR(_I_); PRINTBSDFCHAR(_G_); PRINTBSDFCHAR(_H_);PRINTBSDFCHAR(_T_);  PRINTBSDFCHAR(_SPC_); PrintFloat(PixelPos, OutColor.xyz, FontColor, TopLeft, BSDF.Weight.x); TopLeft.x+=8*7; PRINTBSDFCHAR(_SPC_); PrintFloat(PixelPos, OutColor.xyz, FontColor, TopLeft, BSDF.Weight.y); TopLeft.x+=8*7; PRINTBSDFCHAR(_SPC_); PrintFloat(PixelPos, OutColor.xyz, FontColor, TopLeft, BSDF.Weight.z); TopLeft.x+=8*7; 
+#define PRINTBSDFWEIGHT()	PRINTBSDFCHAR(_SPC_); PRINTBSDFCHAR(_MINUS_); PRINTBSDFCHAR(_W_); PRINTBSDFCHAR(_E_); PRINTBSDFCHAR(_I_); PRINTBSDFCHAR(_G_); PRINTBSDFCHAR(_H_);PRINTBSDFCHAR(_T_);  PRINTBSDFCHAR(_SPC_); PrintFloat(PixelPos, OutColor.xyz, FontColor, TopLeft, BSDF.FinalWeight.x); TopLeft.x+=8*7; PRINTBSDFCHAR(_SPC_); PrintFloat(PixelPos, OutColor.xyz, FontColor, TopLeft, BSDF.FinalWeight.y); TopLeft.x+=8*7; PRINTBSDFCHAR(_SPC_); PrintFloat(PixelPos, OutColor.xyz, FontColor, TopLeft, BSDF.FinalWeight.z); TopLeft.x+=8*7; 
 #define PRINTBSDFGREYWEIGHT()PRINTBSDFCHAR(_SPC_);PRINTBSDFCHARGREY(_MINUS_); PRINTBSDFCHARGREY(_W_); PRINTBSDFCHARGREY(_E_); PRINTBSDFCHARGREY(_I_); PRINTBSDFCHARGREY(_G_); PRINTBSDFCHARGREY(_H_);PRINTBSDFCHARGREY(_T_);  PRINTBSDFCHAR(_SPC_); PrintFloat(PixelPos, OutColor.xyz, FontColor, TopLeft, Unpack10F(BSDF_GETWEIGHT10F(BSDF))); TopLeft.x+=8*23; 
 #define PRINTBSDFSHAREDLOCALBASISID()	PRINTBSDFCHAR(_SPC_); PRINTBSDFCHAR(_MINUS_); PRINTBSDFCHAR(_N_); PRINTBSDFCHAR(_R_); PRINTBSDFCHAR(_M_); PRINTBSDFCHAR(_I_); PRINTBSDFCHAR(_D_); PRINTBSDFCHAR(_SPC_); PrintFloatNoFraction(PixelPos, OutColor.xyz, FontColor, TopLeft, float(BSDF_GETSHAREDLOCALBASISID(BSDF)), 1); TopLeft.x+=8*1;
 #define PRINTBSDFANISO()	PRINTBSDFCHAR(_SPC_); PRINTBSDFCHAR(_MINUS_); PRINTBSDFCHAR(_A_); PRINTBSDFCHAR(_N_); PRINTBSDFCHAR(_I_); PRINTBSDFCHAR(_S_); PRINTBSDFCHAR(_O_); PRINTBSDFCHAR(_SPC_); PrintFloatNoFraction(PixelPos, OutColor.xyz, FontColor, TopLeft, float(BSDF_GETANISOTROPY(BSDF)), 1); TopLeft.x+=8*1;
diff --git a/Engine/Source/Runtime/Engine/Private/Materials/HLSLMaterialTranslator.cpp b/Engine/Source/Runtime/Engine/Private/Materials/HLSLMaterialTranslator.cpp
--- a/Engine/Source/Runtime/Engine/Private/Materials/HLSLMaterialTranslator.cpp
+++ b/Engine/Source/Runtime/Engine/Private/Materials/HLSLMaterialTranslator.cpp
@@ -8985,14 +8985,14 @@ uint8 FHLSLMaterialTranslator::StrataCompilationInfoGetSharedLocalBasesCount()
 	return FinalUsedSharedLocalBasesCount;
 }
 
-int32 FHLSLMaterialTranslator::StrataAddParameterBlendingBSDFWeightToNormalMixCodeChunk(int32 ACodeChunk, int32 BCodeChunk)
+int32 FHLSLMaterialTranslator::StrataAddParameterBlendingBSDFCoverageToNormalMixCodeChunk(int32 ACodeChunk, int32 BCodeChunk)
 {
-	return AddCodeChunk(MCT_Float, TEXT("AddParameterBlendingBSDFWeightToNormalMix(%s, %s)"), *GetParameterCode(ACodeChunk), *GetParameterCode(BCodeChunk));
+	return AddCodeChunk(MCT_Float, TEXT("AddParameterBlendingBSDFCoverageToNormalMix(%s, %s)"), *GetParameterCode(ACodeChunk), *GetParameterCode(BCodeChunk));
 }
 
-int32 FHLSLMaterialTranslator::StrataVerticalLayeringParameterBlendingBSDFWeightToNormalMixCodeChunk(int32 TopCodeChunk)
+int32 FHLSLMaterialTranslator::StrataVerticalLayeringParameterBlendingBSDFCoverageToNormalMixCodeChunk(int32 TopCodeChunk)
 {
-	return AddCodeChunk(MCT_Float, TEXT("VerticalLayeringParameterBlendingBSDFWeightToNormalMix(%s)"), *GetParameterCode(TopCodeChunk));
+	return AddCodeChunk(MCT_Float, TEXT("VerticalLayeringParameterBlendingBSDFCoverageToNormalMix(%s)"), *GetParameterCode(TopCodeChunk));
 }
 
 int32 FHLSLMaterialTranslator::StrataCreateAndRegisterNullMaterial()
diff --git a/Engine/Source/Runtime/Engine/Private/Materials/HLSLMaterialTranslator.h b/Engine/Source/Runtime/Engine/Private/Materials/HLSLMaterialTranslator.h
--- a/Engine/Source/Runtime/Engine/Private/Materials/HLSLMaterialTranslator.h
+++ b/Engine/Source/Runtime/Engine/Private/Materials/HLSLMaterialTranslator.h
@@ -1011,8 +1011,8 @@ protected:
 	virtual FStrataRegisteredSharedLocalBasis StrataCompilationInfoRegisterSharedLocalBasis(int32 NormalCodeChunk) override;
 	virtual FStrataRegisteredSharedLocalBasis StrataCompilationInfoRegisterSharedLocalBasis(int32 NormalCodeChunk, int32 TangentCodeChunk) override;
 	virtual uint8 StrataCompilationInfoGetSharedLocalBasesCount() override;
-	virtual int32 StrataAddParameterBlendingBSDFWeightToNormalMixCodeChunk(int32 ACodeChunk, int32 BCodeChunk) override;
-	virtual int32 StrataVerticalLayeringParameterBlendingBSDFWeightToNormalMixCodeChunk(int32 TopCodeChunk) override;
+	virtual int32 StrataAddParameterBlendingBSDFCoverageToNormalMixCodeChunk(int32 ACodeChunk, int32 BCodeChunk) override;
+	virtual int32 StrataVerticalLayeringParameterBlendingBSDFCoverageToNormalMixCodeChunk(int32 TopCodeChunk) override;
 
 	FStrataSharedLocalBasesInfo StrataCompilationInfoGetMatchingSharedLocalBasisInfo(const FStrataRegisteredSharedLocalBasis& SearchedSharedLocalBasis);
 
diff --git a/Engine/Source/Runtime/Engine/Private/Materials/MaterialExpressions.cpp b/Engine/Source/Runtime/Engine/Private/Materials/MaterialExpressions.cpp
--- a/Engine/Source/Runtime/Engine/Private/Materials/MaterialExpressions.cpp
+++ b/Engine/Source/Runtime/Engine/Private/Materials/MaterialExpressions.cpp
@@ -20971,7 +20971,7 @@ int32 UMaterialExpressionStrataVerticalLayering::Compile(class FMaterialCompiler
 	const FStrataMaterialCompilationInfo& TopStrataData = Compiler->GetStrataCompilationInfo(TopCodeChunk);
 	const FStrataMaterialCompilationInfo& BaseStrataData = Compiler->GetStrataCompilationInfo(BaseCodeChunk);
 
-	const int32 TopNormalMixCodeChunk = Compiler->StrataVerticalLayeringParameterBlendingBSDFWeightToNormalMixCodeChunk(TopCodeChunk);
+	const int32 TopNormalMixCodeChunk = Compiler->StrataVerticalLayeringParameterBlendingBSDFCoverageToNormalMixCodeChunk(TopCodeChunk);
 
 	int32 OutputCodeChunk = INDEX_NONE;
 	if (bUseParameterBlending)
@@ -21112,7 +21112,7 @@ int32 UMaterialExpressionStrataAdd::Compile(class FMaterialCompiler* Compiler, i
 	const FStrataMaterialCompilationInfo& AStrataData = Compiler->GetStrataCompilationInfo(ACodeChunk);
 	const FStrataMaterialCompilationInfo& BStrataData = Compiler->GetStrataCompilationInfo(BCodeChunk);
 
-	const int32 ANormalMixCodeChunk = Compiler->StrataAddParameterBlendingBSDFWeightToNormalMixCodeChunk(ACodeChunk, BCodeChunk);
+	const int32 ANormalMixCodeChunk = Compiler->StrataAddParameterBlendingBSDFCoverageToNormalMixCodeChunk(ACodeChunk, BCodeChunk);
 
 	int32 OutputCodeChunk = INDEX_NONE;
 	if (bUseParameterBlending)
diff --git a/Engine/Source/Runtime/Engine/Public/MaterialCompiler.h b/Engine/Source/Runtime/Engine/Public/MaterialCompiler.h
--- a/Engine/Source/Runtime/Engine/Public/MaterialCompiler.h
+++ b/Engine/Source/Runtime/Engine/Public/MaterialCompiler.h
@@ -458,8 +458,8 @@ public:
 	{
 		return FString::Printf(TEXT("SHAREDLOCALBASIS_INDEX_%u"), SharedLocalBasis.GraphSharedLocalBasisIndex);
 	}
-	virtual int32 StrataAddParameterBlendingBSDFWeightToNormalMixCodeChunk(int32 ACodeChunk, int32 BCodeChunk) = 0;
-	virtual int32 StrataVerticalLayeringParameterBlendingBSDFWeightToNormalMixCodeChunk(int32 TopCodeChunk) = 0;
+	virtual int32 StrataAddParameterBlendingBSDFCoverageToNormalMixCodeChunk(int32 ACodeChunk, int32 BCodeChunk) = 0;
+	virtual int32 StrataVerticalLayeringParameterBlendingBSDFCoverageToNormalMixCodeChunk(int32 TopCodeChunk) = 0;
 
 	// Water
 	virtual int32 SceneDepthWithoutWater(int32 Offset, int32 ViewportUV, bool bUseOffset, float FallbackDepth) = 0;
@@ -1028,14 +1028,14 @@ public:
 		return Compiler->StrataCompilationInfoContainsCodeChunk(CodeChunk);
 	}
 
-	virtual int32 StrataAddParameterBlendingBSDFWeightToNormalMixCodeChunk(int32 ACodeChunk, int32 BCodeChunk) override
+	virtual int32 StrataAddParameterBlendingBSDFCoverageToNormalMixCodeChunk(int32 ACodeChunk, int32 BCodeChunk) override
 	{
-		return Compiler->StrataAddParameterBlendingBSDFWeightToNormalMixCodeChunk(ACodeChunk, BCodeChunk);
+		return Compiler->StrataAddParameterBlendingBSDFCoverageToNormalMixCodeChunk(ACodeChunk, BCodeChunk);
 	}
 
-	virtual int32 StrataVerticalLayeringParameterBlendingBSDFWeightToNormalMixCodeChunk(int32 TopCodeChunk) override
+	virtual int32 StrataVerticalLayeringParameterBlendingBSDFCoverageToNormalMixCodeChunk(int32 TopCodeChunk) override
 	{
-		return Compiler->StrataVerticalLayeringParameterBlendingBSDFWeightToNormalMixCodeChunk(TopCodeChunk);
+		return Compiler->StrataVerticalLayeringParameterBlendingBSDFCoverageToNormalMixCodeChunk(TopCodeChunk);
 	}
 	
 	virtual const FStrataMaterialCompilationInfo& GetStrataCompilationInfo(int32 CodeChunk) override
