Replaced Lumen specific ortho CameraVector calculation with a generic one which works everywhere

diff --git a/Engine/Shaders/Private/MaterialTemplate.ush b/Engine/Shaders/Private/MaterialTemplate.ush
--- a/Engine/Shaders/Private/MaterialTemplate.ush
+++ b/Engine/Shaders/Private/MaterialTemplate.ush
@@ -2878,6 +2878,12 @@ uint DiscardMaterialWithPixelCoverage(FMaterialPixelParameters MaterialParameter
 	#define OPTIONAL_IsFrontFace , in FIsFrontFace bIsFrontFace : FrontFaceSemantic
 #endif
 
+// Return whether View has ortho or perspective projection
+bool IsOrthoProjection(ViewState InView)
+{
+	return InView.ViewToClip[3][3] >= 1.0f;
+}
+
 /** Initializes the subset of Parameters that was not set in GetMaterialPixelParameters. */
 void CalcMaterialParametersEx(
 	in out FMaterialPixelParameters Parameters,
@@ -2902,14 +2908,19 @@ void CalcMaterialParametersEx(
 	Parameters.SvPosition = SvPosition;
 	Parameters.ScreenPosition = ScreenPosition;
 
-	#if LUMEN_CARD_CAPTURE
-		// CameraVector is just ViewForward in an ortho mode
-		Parameters.CameraVector = -ResolvedView.ViewForward;
-	#elif RAYHITGROUPSHADER
+	#if RAYHITGROUPSHADER
 		Parameters.CameraVector = -WorldRayDirection();
 	#else
-		// TranslatedWorldPosition is the world position translated to the camera position, which is just -CameraVector
-		Parameters.CameraVector = normalize(-Parameters.WorldPosition_CamRelative.xyz);
+		if (IsOrthoProjection(ResolvedView))
+		{
+			// CameraVector is just ViewForward in an ortho mode
+			Parameters.CameraVector = -ResolvedView.ViewForward;
+		}
+		else
+		{
+			// TranslatedWorldPosition is the world position translated to the camera position, which is just -CameraVector in perspective projection
+			Parameters.CameraVector = normalize(-Parameters.WorldPosition_CamRelative.xyz);
+		}
 	#endif
 
 	Parameters.LightVector = 0;
