Fixed cloud flickering (as far as my tests goes with low repro...)

diff --git a/Engine/Shaders/Private/VolumetricCloud.usf b/Engine/Shaders/Private/VolumetricCloud.usf
--- a/Engine/Shaders/Private/VolumetricCloud.usf
+++ b/Engine/Shaders/Private/VolumetricCloud.usf
@@ -743,7 +743,7 @@ void MainCommon(in FMaterialPixelParameters MaterialParameters, in float4 SvPosi
 #if 1
 					// Non-linear shadow sample distribution, hardcoded to x^2
 					float PreviousNormT = 0.0f;
-					for (float ShadowT = InvShadowStepCount; ShadowT <= 1.0; ShadowT += InvShadowStepCount)
+					for (float ShadowT = InvShadowStepCount; ShadowT <= 1.00001f; ShadowT += InvShadowStepCount)
 					{
 						DebugGroundShadowMaterialSampleCount++;
 						float CurrentNormT = ShadowT * ShadowT; // CurrentNormT is the end of the considered segment to integrate, PreviousNormT is its beginning.
@@ -847,7 +847,7 @@ void MainCommon(in FMaterialPixelParameters MaterialParameters, in float4 SvPosi
 				// Non-linear shadow sample distribution, hardcoded to x^2
 				const float ShadowDtMeter = ShadowLengthTest * CENTIMETER_TO_METER;
 				float PreviousNormT = 0.0f;
-				for (float ShadowT = InvShadowStepCount; ShadowT <= 1.0; ShadowT += InvShadowStepCount)
+				for (float ShadowT = InvShadowStepCount; ShadowT <= 1.00001f; ShadowT += InvShadowStepCount)
 				{
 					float CurrentNormT = ShadowT * ShadowT; // CurrentNormT is the end of the considered segment to integrate, PreviousNormT is its beginning.
 					const float DetlaNormT = CurrentNormT - PreviousNormT;
@@ -914,7 +914,7 @@ void MainCommon(in FMaterialPixelParameters MaterialParameters, in float4 SvPosi
 					// Non-linear shadow sample distribution, hardcoded to x^2
 					const float ShadowDtMeter = ShadowLengthTest * CENTIMETER_TO_METER;
 					float PreviousNormT = 0.0f;
-					for (float ShadowT = InvShadowStepCount; ShadowT <= 1.0; ShadowT += InvShadowStepCount)
+					for (float ShadowT = InvShadowStepCount; ShadowT <= 1.00001f; ShadowT += InvShadowStepCount)
 					{
 						float CurrentNormT = ShadowT * ShadowT; // CurrentNormT is the end of the considered segment to integrate, PreviousNormT is its beginning.
 						const float DetlaNormT = CurrentNormT - PreviousNormT;
