Fixes some ghosting in TSR on responsive AA

diff --git a/Engine/Shaders/Private/TemporalAA/TAAUpdateHistory.usf b/Engine/Shaders/Private/TemporalAA/TAAUpdateHistory.usf
--- a/Engine/Shaders/Private/TemporalAA/TAAUpdateHistory.usf
+++ b/Engine/Shaders/Private/TemporalAA/TAAUpdateHistory.usf
@@ -225,6 +225,10 @@ void MainCS(
 			TranslucencyRejectionTexture[dpv_lo(InputPixelPos)],
 			TranslucencyRejectionTexture[dpv_hi(InputPixelPos)]);
 		
+		uint2 SceneStencilRef = dpv_interleave_registers(
+			InputSceneStencilTexture.Load(int3(dpv_lo(InputPixelPos), 0)) STENCIL_COMPONENT_SWIZZLE, 
+			InputSceneStencilTexture.Load(int3(dpv_hi(InputPixelPos), 0)) STENCIL_COMPONENT_SWIZZLE);
+
 		UNROLL_N(CONFIG_SAMPLES - 1)
 		for (uint SampleId = 1; SampleId < CONFIG_SAMPLES; SampleId++)
 		{
@@ -238,6 +242,10 @@ void MainCS(
 			TranslucencyRejection = min(TranslucencyRejection, dpv_interleave_registers(
 				TranslucencyRejectionTexture[dpv_lo(SampleInputPixelPos)],
 				TranslucencyRejectionTexture[dpv_hi(SampleInputPixelPos)]));
+			
+			SceneStencilRef |= dpv_interleave_registers(
+				InputSceneStencilTexture.Load(int3(dpv_lo(SampleInputPixelPos), 0)) STENCIL_COMPONENT_SWIZZLE, 
+				InputSceneStencilTexture.Load(int3(dpv_hi(SampleInputPixelPos), 0)) STENCIL_COMPONENT_SWIZZLE);
 		}
 
 		TranslucencyRejection = max(TranslucencyRejection, taa_half(MinTranslucencyRejection));
@@ -247,10 +255,6 @@ void MainCS(
 		{
 			const uint kResponsiveStencilMask = 1 << 3;
 			
-			uint2 SceneStencilRef = dpv_interleave_registers(
-				InputSceneStencilTexture.Load(int3(dpv_lo(InputPixelPos), 0)) STENCIL_COMPONENT_SWIZZLE, 
-				InputSceneStencilTexture.Load(int3(dpv_lo(InputPixelPos), 0)) STENCIL_COMPONENT_SWIZZLE);
-
 			bIsResponsiveAAPixel = (SceneStencilRef & kResponsiveStencilMask) != 0;
 		}
 		#endif
