Fix broken console shader builds

Limit the ddx/ddy workaround defines to ray tracing shaders only like they were originally.

diff --git a/Engine/Shaders/Private/Common.ush b/Engine/Shaders/Private/Common.ush
--- a/Engine/Shaders/Private/Common.ush
+++ b/Engine/Shaders/Private/Common.ush
@@ -117,9 +117,9 @@ const static float3 Max111110BitsFloat3 = float3(Max11BitsFloat, Max11BitsFloat,
 	#define FORCED_TEXTURE_MIP 0.0f
 #endif
 
-#if !PIXELSHADER
+#if RAYHITGROUPSHADER || RAYMISSHADER
 
-// These built-ins are only valid in pixel shaders.
+// These built-ins are not available in ray tracing
 // Define dummy versions so that ray-tracing materials will at least compile
 #define clip(x)
 #define ddx(x) 0
