Rounded corners for background blur

diff --git a/Engine/Shaders/Private/SlateElementPixelShader.usf b/Engine/Shaders/Private/SlateElementPixelShader.usf
--- a/Engine/Shaders/Private/SlateElementPixelShader.usf
+++ b/Engine/Shaders/Private/SlateElementPixelShader.usf
@@ -231,61 +231,11 @@ half4 GetBorderElementColor( VertexToPixelInterpolants InVertex )
 }
 #endif
 
-#if SHADER_TYPE == ST_RoundedBox
-float GetRoundedBoxDistance(float2 pos, float2 center, float radius, float inset)
-{
-	// distance from center
-    pos = abs(pos - center); 
-
-    // distance from the inner corner
-    pos = pos - (center - float2(radius + inset, radius + inset));
-
-    // use distance to nearest edge when not in quadrant with radius
-    // this handles an edge case when radius is very close to thickness
-    // otherwise we're in the quadrant with the radius, 
-    // just use the analytic signed distance function
-    return lerp( length(pos) - radius, max(pos.x - radius, pos.y - radius), float(pos.x <= 0 || pos.y <=0) );
-}
 
+#if SHADER_TYPE == ST_RoundedBox
 half4 GetRoundedBoxElementColor( VertexToPixelInterpolants InVertex )
 {
-	const float2 size = ShaderParams.zw;
-	float2 UV = InVertex.TextureCoordinates[0].xy;
-	float2 pos = size * UV;
-	float2 center = size / 2.0;
-
-	//X = Top Left, Y = Top Right, Z = Bottom Right, W = Bottom Left */
-	float4 cornerRadii = ShaderParams2;
-
-	// figure out which radius to use based on which quadrant we're in
-	float2 quadrant = step(UV, float2(.5,.5));
-
-	float left = lerp(cornerRadii.y, cornerRadii.x, quadrant.x);
-	float right = lerp(cornerRadii.z, cornerRadii.w, quadrant.x);
-	float radius = lerp(right, left, quadrant.y);
-
-	float thickness = ShaderParams.y; 
-
-	// Compute the distances internal and external to the border outline
-	float dext = GetRoundedBoxDistance(pos, center, radius, 0.0);
-	float din  = GetRoundedBoxDistance(pos, center, max(radius - thickness, 0), thickness);
-
-	// Compute the border intensity and fill intensity with a smooth transition
-	float spread = 0.5;
-	float bi = smoothstep(spread, -spread, dext);
-	float fi = smoothstep(spread, -spread, din);
-
-	half4 fill = GetColor(InVertex, GetUV(InVertex, 0) * GetUV(InVertex, 1)); //InVertex.Color;
-	half4 border = InVertex.SecondaryColor;
-
-	// alpha blend the external color 
-	half4 OutColor = lerp(border, fill, half(thickness > radius));
-	OutColor.a = 0.0;
-
-	// blend in the border and fill colors
-	OutColor = lerp(OutColor, border, bi);
-	OutColor = lerp(OutColor, fill, fi);
-	return OutColor;
+	return GetRoundedBoxElementColorInternal(ShaderParams.zw, InVertex.TextureCoordinates[0].xy, ShaderParams.y, ShaderParams2, GetColor(InVertex, GetUV(InVertex, 0) * GetUV(InVertex, 1)),  InVertex.SecondaryColor);
 }
 #endif
 
@@ -396,3 +346,4 @@ float4 DebugBatchingMain(VertexToPixelInterpolants VIn) : SV_Target0
 {
     return BatchColor;
 }
+
diff --git a/Engine/Shaders/Private/SlatePostProcessPixelShader.usf b/Engine/Shaders/Private/SlatePostProcessPixelShader.usf
--- a/Engine/Shaders/Private/SlatePostProcessPixelShader.usf
+++ b/Engine/Shaders/Private/SlatePostProcessPixelShader.usf
@@ -29,6 +29,7 @@ SamplerState ElementTextureSampler;
 float4 BufferSizeAndDirection;
 float4 UVBounds;
 float4 ShaderParams;
+float4 ShaderParams2;
 
 float4 GetSample(float Weight, float Offset, float2 UV)
 {
@@ -107,4 +108,19 @@ float4 DownsampleMain(FScreenVertexOutput Input) : SV_Target0
 
 	return float4(Sample[0] + Sample[1] + Sample[2] + Sample[3]) * 0.25f;
 #endif
+}
+
+float2 Map(float2 value, float2 min1, float2 max1, float min2, float max2)
+{
+  return min2 + (value - min1) * (max2 - min2) / (max1 - min1);
+}
+
+float4 UpsampleMain(FScreenVertexOutput Input) : SV_Target0
+{
+	const float2 LocalSize = ShaderParams.xy;
+	const float2 UV = Input.UV;
+	const float Thickness = 0;
+	const float4 CornerRadii = ShaderParams2;
+	const float4 FillColor = Texture2DSample(ElementTexture, ElementTextureSampler, Input.UV);
+	return GetRoundedBoxElementColorInternal(LocalSize, Map(UV, float2(0,0), ShaderParams.zw, float2(0,0),float2(1,1)), Thickness, CornerRadii, FillColor, FillColor);
 }
\ No newline at end of file
diff --git a/Engine/Shaders/Private/SlateShaderCommon.ush b/Engine/Shaders/Private/SlateShaderCommon.ush
--- a/Engine/Shaders/Private/SlateShaderCommon.ush
+++ b/Engine/Shaders/Private/SlateShaderCommon.ush
@@ -50,30 +50,53 @@ float cross(float2 a, float2 b)
 	return a.x*b.y - a.y*b.x;
 }
 
-/**
- * Given a point p and a parallelogram defined by point a and vectors b and c, determines in p is inside the parallelogram. 
- * returns a 4-vector that can be used with the clip instruction.
- */
-float4 PointInParallelogram(float2 p, float2 a, float4 bc)
+float GetRoundedBoxDistance(float2 pos, float2 center, float radius, float inset)
 {
-	// unoptomized form:
-	//float2 o = p - a;
-	//float2 b = bc.xy;
-	//float2 c = bc.zw;
-	//float d = cross(b, c);
-	//float s = -cross(o, b) / d;
-	//float t = cross(o, c) / d;
-	// test for s and t between 0 and 1
-	//return float4(s, 1 - s, t, 1 - t);
-
-	float2 o = p - a;
-	// precompute 1/d
-	float invD = 1/cross(bc.xy, bc.zw);
-	// Compute an optimized o x b and o x c, leveraging that b and c are in the same vector register already (and free swizzles):
-	//   (o.x * b .y  - o.y * b .x, o.x *  c.y - o.y *  c.x) ==
-	//   (o.x * bc.y  - o.y * bc.x, o.x * bc.w - o.y * bc.z) ==
-	//    o.x * bc.yw - o.y * bc.xz
-	float2 st = (o.x * bc.yw - o.y * bc.xz) * float2(-invD, invD);
-	// test for s and t between 0 and 1
-	return float4(st, float2(1,1) - st);
+	// distance from center
+    pos = abs(pos - center); 
+
+    // distance from the inner corner
+    pos = pos - (center - float2(radius + inset, radius + inset));
+
+    // use distance to nearest edge when not in quadrant with radius
+    // this handles an edge case when radius is very close to thickness
+    // otherwise we're in the quadrant with the radius, 
+    // just use the analytic signed distance function
+    return lerp( length(pos) - radius, max(pos.x - radius, pos.y - radius), float(pos.x <= 0 || pos.y <=0) );
 }
+
+float4 GetRoundedBoxElementColorInternal(float2 size, float2 UV, float thickness, float4 cornerRadii, float4 fillColor, float4 borderColor)
+{
+	float2 pos = size * UV;
+	float2 center = size / 2.0;
+
+	//cornerRadii: X = Top Left, Y = Top Right, Z = Bottom Right, W = Bottom Left */
+
+	// figure out which radius to use based on which quadrant we're in
+	float2 quadrant = step(UV, float2(.5,.5));
+
+	float left = lerp(cornerRadii.y, cornerRadii.x, quadrant.x);
+	float right = lerp(cornerRadii.z, cornerRadii.w, quadrant.x);
+	float radius = lerp(right, left, quadrant.y);
+
+	// Compute the distances internal and external to the border outline
+	float dext = GetRoundedBoxDistance(pos, center, radius, 0.0);
+	float din  = GetRoundedBoxDistance(pos, center, max(radius - thickness, 0), thickness);
+
+	// Compute the border intensity and fill intensity with a smooth transition
+	float spread = 0.5;
+	float bi = smoothstep(spread, -spread, dext);
+	float fi = smoothstep(spread, -spread, din);
+
+	float4 fill = fillColor;
+	float4 border = borderColor;
+
+	// alpha blend the external color 
+	float4 OutColor = lerp(border, fill, float(thickness > radius));
+	OutColor.a = 0.0;
+
+	// blend in the border and fill colors
+	OutColor = lerp(OutColor, border, bi);
+	OutColor = lerp(OutColor, fill, fi);
+	return OutColor;
+}
\ No newline at end of file
diff --git a/Engine/Source/Runtime/Slate/Private/Widgets/Layout/SBackgroundBlur.cpp b/Engine/Source/Runtime/Slate/Private/Widgets/Layout/SBackgroundBlur.cpp
--- a/Engine/Source/Runtime/Slate/Private/Widgets/Layout/SBackgroundBlur.cpp
+++ b/Engine/Source/Runtime/Slate/Private/Widgets/Layout/SBackgroundBlur.cpp
@@ -20,11 +20,13 @@ void SBackgroundBlur::PrivateRegisterAttributes(FSlateAttributeInitializer& Attr
 {
 	SLATE_ADD_MEMBER_ATTRIBUTE_DEFINITION_WITH_NAME(AttributeInitializer, "BlurStrength", BlurStrengthAttribute, EInvalidateWidgetReason::Paint);
 	SLATE_ADD_MEMBER_ATTRIBUTE_DEFINITION_WITH_NAME(AttributeInitializer, "BlurRadius", BlurRadiusAttribute, EInvalidateWidgetReason::Paint);
+	SLATE_ADD_MEMBER_ATTRIBUTE_DEFINITION_WITH_NAME(AttributeInitializer, "CornerRadius", CornerRadiusAttribute, EInvalidateWidgetReason::Paint);
 }
 
 SBackgroundBlur::SBackgroundBlur()
 	: BlurStrengthAttribute(*this)
 	, BlurRadiusAttribute(*this)
+	, CornerRadiusAttribute(*this)
 {
 }
 
@@ -35,6 +37,7 @@ void SBackgroundBlur::Construct(const FArguments& InArgs)
 	LowQualityFallbackBrush = InArgs._LowQualityFallbackBrush;
 	SetBlurStrength(InArgs._BlurStrength);
 	SetBlurRadius(InArgs._BlurRadius);
+	SetCornerRadius(InArgs._CornerRadius);
 
 	ChildSlot
 		.HAlign(InArgs._HAlign)
@@ -80,6 +83,11 @@ void SBackgroundBlur::SetLowQualityBackgroundBrush(const FSlateBrush* InBrush)
 	}
 }
 
+void SBackgroundBlur::SetCornerRadius(TAttribute<FVector4> InCornerRadius)
+{
+	CornerRadiusAttribute.Assign(*this, InCornerRadius);
+}
+
 void SBackgroundBlur::SetHAlign(EHorizontalAlignment HAlign)
 {
 	ChildSlot.SetHorizontalAlignment(HAlign);
@@ -136,7 +144,7 @@ int32 SBackgroundBlur::OnPaint(const FPaintArgs& Args, const FGeometry& Allotted
 				{
 					OutDrawElements.PushClip(FSlateClippingZone(AllottedGeometry));
 
-					FSlateDrawElement::MakePostProcessPass(OutDrawElements, LayerId, PaintGeometry, FVector4((float)KernelSize, ComputedStrength, (float)RenderTargetWidth, (float)RenderTargetHeight), DownsampleAmount);
+					FSlateDrawElement::MakePostProcessPass(OutDrawElements, LayerId, PaintGeometry, FVector4((float)KernelSize, ComputedStrength, (float)RenderTargetWidth, (float)RenderTargetHeight), DownsampleAmount, CornerRadiusAttribute.Get());
 
 					OutDrawElements.PopClip();
 				}
diff --git a/Engine/Source/Runtime/Slate/Public/Widgets/Layout/SBackgroundBlur.h b/Engine/Source/Runtime/Slate/Public/Widgets/Layout/SBackgroundBlur.h
--- a/Engine/Source/Runtime/Slate/Public/Widgets/Layout/SBackgroundBlur.h
+++ b/Engine/Source/Runtime/Slate/Public/Widgets/Layout/SBackgroundBlur.h
@@ -27,6 +27,7 @@ public:
 		, _bApplyAlphaToBlur(true)
 		, _BlurStrength(0.f)
 		, _BlurRadius()
+		, _CornerRadius(FVector4(0,0,0,0))
 		, _LowQualityFallbackBrush(nullptr)
 	{
 		_Visibility = EVisibility::SelfHitTestInvisible;
@@ -40,6 +41,7 @@ public:
 		SLATE_ARGUMENT(bool, bApplyAlphaToBlur)
 		SLATE_ATTRIBUTE(float, BlurStrength)
 		SLATE_ATTRIBUTE(TOptional<int32>, BlurRadius)
+		SLATE_ATTRIBUTE(FVector4, CornerRadius)
 		SLATE_ARGUMENT(const FSlateBrush*, LowQualityFallbackBrush)
 	SLATE_END_ARGS()
 
@@ -52,6 +54,7 @@ public:
 	void SetApplyAlphaToBlur(bool bInApplyAlphaToBlur);
 	void SetBlurRadius(TAttribute<TOptional<int32>> InBlurRadius);
 	void SetBlurStrength(TAttribute<float> InStrength);
+	void SetCornerRadius(TAttribute<FVector4> InCornerRadius);
 	void SetLowQualityBackgroundBrush(const FSlateBrush* InBrush);
 	
 	void SetHAlign(EHorizontalAlignment HAlign);
@@ -81,4 +84,5 @@ protected:
 private:
 	TSlateAttribute<float> BlurStrengthAttribute;
 	TSlateAttribute<TOptional<int32>> BlurRadiusAttribute;
+	TSlateAttribute<FVector4> CornerRadiusAttribute;
 };
diff --git a/Engine/Source/Runtime/SlateCore/Private/Rendering/DrawElements.cpp b/Engine/Source/Runtime/SlateCore/Private/Rendering/DrawElements.cpp
--- a/Engine/Source/Runtime/SlateCore/Private/Rendering/DrawElements.cpp
+++ b/Engine/Source/Runtime/SlateCore/Private/Rendering/DrawElements.cpp
@@ -638,7 +638,7 @@ void FSlateDrawElement::MakeCustomVerts(FSlateWindowElementList& ElementList, ui
 	Element.RenderTransform = FSlateRenderTransform();
 }
 
-void FSlateDrawElement::MakePostProcessPass(FSlateWindowElementList& ElementList, uint32 InLayer, const FPaintGeometry& PaintGeometry, const FVector4& Params, int32 DownsampleAmount)
+void FSlateDrawElement::MakePostProcessPass(FSlateWindowElementList& ElementList, uint32 InLayer, const FPaintGeometry& PaintGeometry, const FVector4& Params, int32 DownsampleAmount, FVector4 CornerRadius)
 {
 	PaintGeometry.CommitTransformsIfUsingLegacyConstructor();
 
@@ -652,6 +652,7 @@ void FSlateDrawElement::MakePostProcessPass(FSlateWindowElementList& ElementList
 	FSlatePostProcessPayload& DataPayload = ElementList.CreatePayload<FSlatePostProcessPayload>(Element);
 	DataPayload.DownsampleAmount = DownsampleAmount;
 	DataPayload.PostProcessData = Params;
+	DataPayload.CornerRadius = CornerRadius;
 
 	Element.Init(ElementList, EElementType::ET_PostProcessPass, InLayer, PaintGeometry, ESlateDrawEffect::None);
 }
diff --git a/Engine/Source/Runtime/SlateCore/Private/Rendering/ElementBatcher.cpp b/Engine/Source/Runtime/SlateCore/Private/Rendering/ElementBatcher.cpp
--- a/Engine/Source/Runtime/SlateCore/Private/Rendering/ElementBatcher.cpp
+++ b/Engine/Source/Runtime/SlateCore/Private/Rendering/ElementBatcher.cpp
@@ -2595,7 +2595,10 @@ void FSlateElementBatcher::AddPostProcessPass(const FSlateDrawElement& DrawEleme
 	// These could be negative with rotation or negative scales.  This is not supported yet
 	if(SizeUV.X > 0 && SizeUV.Y > 0)
 	{
-		FShaderParams Params = FShaderParams::MakePixelShaderParams(FVector4(WorldTopLeft, WorldBotRight), FVector4(Payload.PostProcessData.X, Payload.PostProcessData.Y, (float)Payload.DownsampleAmount, 0.f));
+		FShaderParams Params = FShaderParams::MakePixelShaderParams(
+			FVector4(WorldTopLeft, WorldBotRight),
+			FVector4(Payload.PostProcessData.X, Payload.PostProcessData.Y, (float)Payload.DownsampleAmount, 0.f),
+			FVector4(Payload.CornerRadius));
 
 		CreateRenderBatch(Layer, Params, nullptr, ESlateDrawPrimitive::TriangleList, ESlateShader::PostProcess, ESlateDrawEffect::None, ESlateBatchDrawFlag::None, DrawElement);
 	}
diff --git a/Engine/Source/Runtime/SlateCore/Public/Rendering/DrawElementPayloads.h b/Engine/Source/Runtime/SlateCore/Public/Rendering/DrawElementPayloads.h
--- a/Engine/Source/Runtime/SlateCore/Public/Rendering/DrawElementPayloads.h
+++ b/Engine/Source/Runtime/SlateCore/Public/Rendering/DrawElementPayloads.h
@@ -339,5 +339,6 @@ struct FSlatePostProcessPayload : public FSlateDataPayload
 {
 	// Post Process Data
 	FVector4 PostProcessData;
+	FVector4 CornerRadius;
 	int32 DownsampleAmount;
 };
\ No newline at end of file
diff --git a/Engine/Source/Runtime/SlateCore/Public/Rendering/DrawElements.h b/Engine/Source/Runtime/SlateCore/Public/Rendering/DrawElements.h
--- a/Engine/Source/Runtime/SlateCore/Public/Rendering/DrawElements.h
+++ b/Engine/Source/Runtime/SlateCore/Public/Rendering/DrawElements.h
@@ -252,7 +252,7 @@ public:
 	
 	SLATECORE_API static void MakeCustomVerts(FSlateWindowElementList& ElementList, uint32 InLayer, const FSlateResourceHandle& InRenderResourceHandle, const TArray<FSlateVertex>& InVerts, const TArray<SlateIndex>& InIndexes, ISlateUpdatableInstanceBuffer* InInstanceData, uint32 InInstanceOffset, uint32 InNumInstances, ESlateDrawEffect InDrawEffects = ESlateDrawEffect::None);
 
-	SLATECORE_API static void MakePostProcessPass(FSlateWindowElementList& ElementList, uint32 InLayer, const FPaintGeometry& PaintGeometry, const FVector4& Params, int32 DownsampleAmount);
+	SLATECORE_API static void MakePostProcessPass(FSlateWindowElementList& ElementList, uint32 InLayer, const FPaintGeometry& PaintGeometry, const FVector4& Params, int32 DownsampleAmount, const FVector4 CornerRadius = FVector4(0.0f));
 
 	FSlateDrawElement();
 	SLATECORE_API ~FSlateDrawElement();
diff --git a/Engine/Source/Runtime/SlateCore/Public/Rendering/RenderingCommon.h b/Engine/Source/Runtime/SlateCore/Public/Rendering/RenderingCommon.h
--- a/Engine/Source/Runtime/SlateCore/Public/Rendering/RenderingCommon.h
+++ b/Engine/Source/Runtime/SlateCore/Public/Rendering/RenderingCommon.h
@@ -169,25 +169,28 @@ struct FShaderParams
 	/** Pixel shader parameters */
 	FVector4 PixelParams;
 	FVector4 PixelParams2;
+	FVector4 PixelParams3;
 
 	FShaderParams()
 		: PixelParams(0, 0, 0, 0)
 		, PixelParams2(0, 0, 0, 0)
+		, PixelParams3(0, 0, 0, 0)
 	{}
 
-	FShaderParams(const FVector4& InPixelParams, const FVector4& InPixelParams2 = FVector4(0))
+	FShaderParams(const FVector4& InPixelParams, const FVector4& InPixelParams2 = FVector4(0), const FVector4& InPixelParams3 = FVector4(0))
 		: PixelParams(InPixelParams)
 		, PixelParams2(InPixelParams2)
+		, PixelParams3(InPixelParams3)
 	{}
 
 	bool operator==(const FShaderParams& Other) const
 	{
-		return PixelParams == Other.PixelParams && PixelParams2 == Other.PixelParams2;
+		return PixelParams == Other.PixelParams && PixelParams2 == Other.PixelParams2 && PixelParams3 == Other.PixelParams3;
 	}
 
-	static FShaderParams MakePixelShaderParams(const FVector4& PixelShaderParams, const FVector4& InPixelShaderParams2 = FVector4(0))
+	static FShaderParams MakePixelShaderParams(const FVector4& PixelShaderParams, const FVector4& InPixelShaderParams2 = FVector4(0), const FVector4& InPixelShaderParams3 = FVector4(0))
 	{
-		return FShaderParams(PixelShaderParams, InPixelShaderParams2);
+		return FShaderParams(PixelShaderParams, InPixelShaderParams2, InPixelShaderParams3);
 	}
 };
 
diff --git a/Engine/Source/Runtime/SlateRHIRenderer/Private/SlatePostProcessor.cpp b/Engine/Source/Runtime/SlateRHIRenderer/Private/SlatePostProcessor.cpp
--- a/Engine/Source/Runtime/SlateRHIRenderer/Private/SlatePostProcessor.cpp
+++ b/Engine/Source/Runtime/SlateRHIRenderer/Private/SlatePostProcessor.cpp
@@ -372,8 +372,10 @@ void FSlatePostProcessor::UpsampleRect(FRHICommandListImmediate& RHICmdList, IRe
 {
 	SCOPED_DRAW_EVENT(RHICmdList, SlatePostProcessUpsample);
 
+	const FVector4 Zero(0, 0, 0, 0);
+
 	FGraphicsPipelineStateInitializer GraphicsPSOInit;
-	GraphicsPSOInit.BlendState = TStaticBlendState<>::GetRHI();
+	GraphicsPSOInit.BlendState = Params.CornerRadius == Zero ? TStaticBlendState<>::GetRHI() : TStaticBlendState<CW_RGBA, BO_Add, BF_SourceAlpha, BF_InverseSourceAlpha, BO_Add, BF_One, BF_InverseSourceAlpha>::GetRHI();
 	GraphicsPSOInit.RasterizerState = TStaticRasterizerState<>::GetRHI();
 	GraphicsPSOInit.DepthStencilState = TStaticDepthStencilState<false, CF_Always>::GetRHI();
 
@@ -415,7 +417,7 @@ void FSlatePostProcessor::UpsampleRect(FRHICommandListImmediate& RHICmdList, IRe
 			Params.RestoreStateFunc(RHICmdList, GraphicsPSOInit);
 		}
 
-		TShaderMapRef<FScreenPS> PixelShader(ShaderMap);
+		TShaderMapRef<FSlatePostProcessUpsamplePS> PixelShader(ShaderMap);
 
 		GraphicsPSOInit.BoundShaderState.VertexDeclarationRHI = GFilterVertexDeclaration.VertexDeclarationRHI;
 		GraphicsPSOInit.BoundShaderState.VertexShaderRHI = VertexShader.GetVertexShader();
@@ -427,17 +429,26 @@ void FSlatePostProcessor::UpsampleRect(FRHICommandListImmediate& RHICmdList, IRe
 		{
 			Params.RestoreStateFuncPostPipelineState();
 		}
+		
+
+		const FVector2D SizeUV(
+			DownsampledWidth == SrcTextureWidth ? 1.0f : (DownsampledWidth / (float)SrcTextureWidth) - (1.0f / (float)SrcTextureWidth),
+			DownsampledHeight == SrcTextureHeight ? 1.0f : (DownsampledHeight / (float)SrcTextureHeight) - (1.0f / (float)SrcTextureHeight)
+			);
+
+		const FVector2D Size(DestRect.Right - DestRect.Left, DestRect.Bottom - DestRect.Top);
+		FShaderParams ShaderParams = FShaderParams::MakePixelShaderParams(FVector4(Size, SizeUV), Params.CornerRadius);
+
 
-		PixelShader->SetParameters(RHICmdList, Sampler, SrcTexture);
+		PixelShader->SetShaderParams(RHICmdList, ShaderParams);
+		PixelShader->SetTexture(RHICmdList, SrcTexture, Sampler);
 
-		const float SizeU = (DownsampledWidth == SrcTextureWidth) ? 1.0f : (DownsampledWidth / (float)SrcTextureWidth) - (1.0f / (float)SrcTextureWidth);
-		const float SizeV = (DownsampledHeight == SrcTextureHeight) ? 1.0f : (DownsampledHeight / (float)SrcTextureHeight) - (1.0f / (float)SrcTextureHeight);
 
 		RendererModule.DrawRectangle(RHICmdList,
 			DestRect.Left, DestRect.Top,
-			DestRect.Right - DestRect.Left, DestRect.Bottom - DestRect.Top,
+			Size.X, Size.Y,
 			0, 0,
-			SizeU, SizeV,
+			SizeUV.X, SizeUV.Y,
 			Params.SourceTextureSize,
 			FIntPoint(1, 1),
 			VertexShader,
diff --git a/Engine/Source/Runtime/SlateRHIRenderer/Private/SlatePostProcessor.h b/Engine/Source/Runtime/SlateRHIRenderer/Private/SlatePostProcessor.h
--- a/Engine/Source/Runtime/SlateRHIRenderer/Private/SlatePostProcessor.h
+++ b/Engine/Source/Runtime/SlateRHIRenderer/Private/SlatePostProcessor.h
@@ -14,6 +14,7 @@ struct FPostProcessRectParams
 	FTexture2DRHIRef SourceTexture;
 	FSlateRect SourceRect;
 	FSlateRect DestRect;
+	FVector4 CornerRadius;
 	FIntPoint SourceTextureSize;
 	TFunction<void(FRHICommandListImmediate&, FGraphicsPipelineStateInitializer&)> RestoreStateFunc;
 	TFunction<void()> RestoreStateFuncPostPipelineState;
diff --git a/Engine/Source/Runtime/SlateRHIRenderer/Private/SlateRHIRenderingPolicy.cpp b/Engine/Source/Runtime/SlateRHIRenderer/Private/SlateRHIRenderingPolicy.cpp
--- a/Engine/Source/Runtime/SlateRHIRenderer/Private/SlateRHIRenderingPolicy.cpp
+++ b/Engine/Source/Runtime/SlateRHIRenderer/Private/SlateRHIRenderingPolicy.cpp
@@ -1152,13 +1152,14 @@ void FSlateRHIRenderingPolicy::DrawElements(
 				SLATE_DRAW_EVENT(RHICmdList, PostProcess);
 				RHICmdList.EndRenderPass();
 
-				const FVector4& QuadPositionData = ShaderParams.PixelParams;
+				const FVector4 QuadPositionData = ShaderParams.PixelParams;
 
 				FPostProcessRectParams RectParams;
 				RectParams.SourceTexture = PostProcessTexture->GetTexture2D();
 				RectParams.SourceRect = FSlateRect(0, 0, PostProcessTexture->GetSizeX(), PostProcessTexture->GetSizeY());
 				RectParams.DestRect = FSlateRect(QuadPositionData.X, QuadPositionData.Y, QuadPositionData.Z, QuadPositionData.W);
 				RectParams.SourceTextureSize = PostProcessTexture->GetSizeXY();
+				RectParams.CornerRadius = ShaderParams.PixelParams3;
 
 				RectParams.RestoreStateFunc = [&](FRHICommandListImmediate&InRHICmdList, FGraphicsPipelineStateInitializer& InGraphicsPSOInit) {
 					return UpdateScissorRect(
diff --git a/Engine/Source/Runtime/SlateRHIRenderer/Private/SlateShaders.cpp b/Engine/Source/Runtime/SlateRHIRenderer/Private/SlateShaders.cpp
--- a/Engine/Source/Runtime/SlateRHIRenderer/Private/SlateShaders.cpp
+++ b/Engine/Source/Runtime/SlateRHIRenderer/Private/SlateShaders.cpp
@@ -18,6 +18,7 @@ IMPLEMENT_SHADER_TYPE(, FSlateDebugOverdrawPS, TEXT("/Engine/Private/SlateElemen
 
 IMPLEMENT_SHADER_TYPE(, FSlatePostProcessBlurPS, TEXT("/Engine/Private/SlatePostProcessPixelShader.usf"), TEXT("GaussianBlurMain"), SF_Pixel);
 IMPLEMENT_SHADER_TYPE(, FSlatePostProcessDownsamplePS, TEXT("/Engine/Private/SlatePostProcessPixelShader.usf"), TEXT("DownsampleMain"), SF_Pixel);
+IMPLEMENT_SHADER_TYPE(, FSlatePostProcessUpsamplePS, TEXT("/Engine/Private/SlatePostProcessPixelShader.usf"), TEXT("UpsampleMain"), SF_Pixel);
 IMPLEMENT_SHADER_TYPE(, FSlatePostProcessColorDeficiencyPS, TEXT("/Engine/Private/SlatePostProcessColorDeficiencyPixelShader.usf"), TEXT("ColorDeficiencyMain"), SF_Pixel);
 
 IMPLEMENT_SHADER_TYPE(, FSlateMaskingVS, TEXT("/Engine/Private/SlateMaskingShader.usf"), TEXT("MainVS"), SF_Vertex);
diff --git a/Engine/Source/Runtime/SlateRHIRenderer/Private/SlateShaders.h b/Engine/Source/Runtime/SlateRHIRenderer/Private/SlateShaders.h
--- a/Engine/Source/Runtime/SlateRHIRenderer/Private/SlateShaders.h
+++ b/Engine/Source/Runtime/SlateRHIRenderer/Private/SlateShaders.h
@@ -357,6 +357,30 @@ private:
 };
 
 
+class FSlatePostProcessUpsamplePS : public FSlateElementPS
+{
+	DECLARE_SHADER_TYPE(FSlatePostProcessUpsamplePS, Global);
+public:
+	/** Indicates that this shader should be cached */
+	static bool ShouldCompilePermutation(const FGlobalShaderPermutationParameters& Parameters)
+	{
+		return true;
+	}
+
+	FSlatePostProcessUpsamplePS()
+	{
+	}
+
+	/** Constructor.  Binds all parameters used by the shader */
+	FSlatePostProcessUpsamplePS(const ShaderMetaType::CompiledShaderInitializerType& Initializer)
+		: FSlateElementPS(Initializer)
+	{
+	}
+
+private:
+};
+
+
 class FSlatePostProcessColorDeficiencyPS : public FSlateElementPS
 {
 	DECLARE_SHADER_TYPE(FSlatePostProcessColorDeficiencyPS, Global);
diff --git a/Engine/Source/Runtime/UMG/Private/Components/BackgroundBlur.cpp b/Engine/Source/Runtime/UMG/Private/Components/BackgroundBlur.cpp
--- a/Engine/Source/Runtime/UMG/Private/Components/BackgroundBlur.cpp
+++ b/Engine/Source/Runtime/UMG/Private/Components/BackgroundBlur.cpp
@@ -21,6 +21,7 @@ UBackgroundBlur::UBackgroundBlur(const FObjectInitializer& ObjectInitializer)
 	, BlurStrength(0.f)
 	, bOverrideAutoRadiusCalculation(false)
 	, BlurRadius(0)
+	, CornerRadius(0,0,0,0)
 	, LowQualityFallbackBrush(FSlateNoResource())
 {
 	bIsVariable = false;
@@ -65,6 +66,7 @@ void UBackgroundBlur::SynchronizeProperties()
 		MyBackgroundBlur->SetBlurRadius(bOverrideAutoRadiusCalculation ? BlurRadius : TOptional<int32>());
 		MyBackgroundBlur->SetBlurStrength(BlurStrength);
 		MyBackgroundBlur->SetLowQualityBackgroundBrush(&LowQualityFallbackBrush);
+		MyBackgroundBlur->SetCornerRadius(CornerRadius);
 	}
 }
 
@@ -146,6 +148,15 @@ void UBackgroundBlur::SetBlurStrength(float InStrength)
 	}
 }
 
+void UBackgroundBlur::SetCornerRadius(FVector4 InCornerRadius)
+{
+	CornerRadius = InCornerRadius;
+	if (MyBackgroundBlur.IsValid())
+	{
+		MyBackgroundBlur->SetCornerRadius(InCornerRadius);
+	}
+}
+
 void UBackgroundBlur::SetLowQualityFallbackBrush(const FSlateBrush& InBrush)
 {
 	LowQualityFallbackBrush = InBrush;
diff --git a/Engine/Source/Runtime/UMG/Public/Components/BackgroundBlur.h b/Engine/Source/Runtime/UMG/Public/Components/BackgroundBlur.h
--- a/Engine/Source/Runtime/UMG/Public/Components/BackgroundBlur.h
+++ b/Engine/Source/Runtime/UMG/Public/Components/BackgroundBlur.h
@@ -51,6 +51,13 @@ public:
 	UPROPERTY(EditAnywhere, AdvancedDisplay, BlueprintReadOnly, Category=Appearance, meta=(ClampMin=0, ClampMax=255, EditCondition="bOverrideAutoRadiusCalculation"))
 	int32 BlurRadius;
 
+		/**
+	 * This is the number of pixels which will be weighted in each direction from any given pixel when computing the blur
+	 * A larger value is more costly but allows for stronger blurs.  
+	 */
+	UPROPERTY(EditAnywhere, AdvancedDisplay, BlueprintReadOnly, Category = Appearance)
+	FVector4 CornerRadius;
+
 	/**
 	 * An image to draw instead of applying a blur when low quality override mode is enabled. 
 	 * You can enable low quality mode for background blurs by setting the cvar Slate.ForceBackgroundBlurLowQualityOverride to 1. 
@@ -85,6 +92,9 @@ public:
 	UFUNCTION(BlueprintCallable, Category = "Appearance")
 	virtual void SetBlurStrength(float InStrength);
 
+	UFUNCTION(BlueprintCallable, Category = "Appearance")
+	virtual void SetCornerRadius(FVector4 InCornerRadius);
+
 	UFUNCTION(BlueprintCallable, Category = "Appearance")
 	void SetLowQualityFallbackBrush(const FSlateBrush& InBrush);
 
