Improvement to strata horizontal parameter taking into account BSDF weight. (new function in but improvements still needs tome work so off for now)

diff --git a/Engine/Shaders/Private/Strata/Strata.ush b/Engine/Shaders/Private/Strata/Strata.ush
--- a/Engine/Shaders/Private/Strata/Strata.ush
+++ b/Engine/Shaders/Private/Strata/Strata.ush
@@ -925,6 +925,19 @@ FStrataData StrataVerticalLayering(FStrataData Top, FStrataData Base)
 	return StrataData;
 }
 
+float HorizontalMixingParameterBlendingBSDFCoverageToNormalMix(FStrataData Background, FStrataData Foreground, float HorizontalMixValue)
+{
+	// This is a normalised mix so even though both BSDF have weights, the sum of weights used to combined them should always be 1.
+#if STRATA_INLINE_SHADING
+	const float SafeBackCoverage = saturate(Background.Layers[0].BSDFs[0].Coverage);
+	const float SafeForeCoverage = saturate(Foreground.Layers[0].BSDFs[0].Coverage);
+
+	return HorizontalMixValue;
+#else
+	return 0.5f;
+#endif
+}
+
 // Note about parameter blending
 //  - We can only parameter blend Slab nodes. This error is handled in the compiler.
 //  - We can only parameter blend one BSDF into another single BSDF. More complex topology are not handled and result in a compiler error. This could be extended to a single BSDF per layer.
@@ -932,9 +945,10 @@ FStrataData StrataVerticalLayering(FStrataData Top, FStrataData Base)
 //  - Slab node with SSSprofile cannot be blendend. Those will have to be specifically handled by the user and not mixed with any other slab if it can be demoted to parameter blending
 
 // ==> NOTE: Always pair with the compiler behavior in StrataCompilationInfoHorizontalMixingParamBlend
-FStrataData StrataHorizontalMixingParameterBlending(FStrataData Foreground, FStrataData Background, float Mix, uint NewNormalIndex)
+FStrataData StrataHorizontalMixingParameterBlending(FStrataData Foreground, FStrataData Background, float HorizontalMixValue, float NormalMixCodeChunk, uint NewNormalIndex)
 {
-	const float ForegroundMixFactor = saturate(Mix);
+	const float ForegroundLightingMixFactor = saturate(HorizontalMixValue);	// Only the horizontal mixing value
+	const float ForegroundOtherMixFactor = saturate(NormalMixCodeChunk);	// Horizontal mixing value combined with coverage
 
 	FStrataData Result = GetInitialisedStrataDataWithReservedFirstSlot();
 #define ResultBSDF		Result.Layers[0].BSDFs[0]
@@ -965,19 +979,19 @@ FStrataData StrataHorizontalMixingParameterBlending(FStrataData Foreground, FStr
 	BSDF_SETHASTHINFILM(ResultBSDF,		BSDF_GETHASTHINFILM(ForegroundBSDF)				| BSDF_GETHASTHINFILM(BackgroundBSDF));
 	BSDF_SETHASFUZZ(ResultBSDF,			BSDF_GETHASFUZZ(ForegroundBSDF)					| BSDF_GETHASFUZZ(BackgroundBSDF));
 
-	BSDF_SETEMISSIVE(ResultBSDF,		lerp(BackLightingFactor * BSDF_GETEMISSIVE(BackgroundBSDF),			ForeLightingFactor * BSDF_GETEMISSIVE(ForegroundBSDF),			ForegroundMixFactor));
-	SLAB_DIFFUSEALBEDO(ResultBSDF) =	lerp(BackLightingFactor * SLAB_DIFFUSEALBEDO(BackgroundBSDF),		ForeLightingFactor * SLAB_DIFFUSEALBEDO(ForegroundBSDF),		ForegroundMixFactor);
-	SLAB_F0(ResultBSDF) =				lerp(BackLightingFactor * SLAB_F0(BackgroundBSDF),					ForeLightingFactor * SLAB_F0(ForegroundBSDF),					ForegroundMixFactor);
-	SLAB_F90(ResultBSDF) =				lerp(BackLightingFactor * SLAB_F90(BackgroundBSDF),					ForeLightingFactor * SLAB_F90(ForegroundBSDF),					ForegroundMixFactor);
-	SLAB_FUZZ_COLOR(ResultBSDF)	=		lerp(BackLightingFactor * SLAB_FUZZ_COLOR(BackgroundBSDF),			ForeLightingFactor * SLAB_FUZZ_COLOR(ForegroundBSDF),			ForegroundMixFactor);
+	BSDF_SETEMISSIVE(ResultBSDF,		lerp(BackLightingFactor * BSDF_GETEMISSIVE(BackgroundBSDF),			ForeLightingFactor * BSDF_GETEMISSIVE(ForegroundBSDF),			ForegroundLightingMixFactor));
+	SLAB_DIFFUSEALBEDO(ResultBSDF) =	lerp(BackLightingFactor * SLAB_DIFFUSEALBEDO(BackgroundBSDF),		ForeLightingFactor * SLAB_DIFFUSEALBEDO(ForegroundBSDF),		ForegroundLightingMixFactor);
+	SLAB_F0(ResultBSDF) =				lerp(BackLightingFactor * SLAB_F0(BackgroundBSDF),					ForeLightingFactor * SLAB_F0(ForegroundBSDF),					ForegroundLightingMixFactor);
+	SLAB_F90(ResultBSDF) =				lerp(BackLightingFactor * SLAB_F90(BackgroundBSDF),					ForeLightingFactor * SLAB_F90(ForegroundBSDF),					ForegroundLightingMixFactor);
+	SLAB_FUZZ_COLOR(ResultBSDF)	=		lerp(BackLightingFactor * SLAB_FUZZ_COLOR(BackgroundBSDF),			ForeLightingFactor * SLAB_FUZZ_COLOR(ForegroundBSDF),			ForegroundLightingMixFactor);
 
-	BSDF_SETTHICKNESSCM(ResultBSDF,		lerp(BSDF_GETTHICKNESSCM(BackgroundBSDF),							BSDF_GETTHICKNESSCM(ForegroundBSDF),							ForegroundMixFactor));
-	SLAB_METALLIC(ResultBSDF) =			lerp(SLAB_METALLIC(BackgroundBSDF),									SLAB_METALLIC(ForegroundBSDF),									ForegroundMixFactor);
-	SLAB_ROUGHNESSX(ResultBSDF) =		lerp(SLAB_ROUGHNESSX(BackgroundBSDF),								SLAB_ROUGHNESSX(ForegroundBSDF),								ForegroundMixFactor);
-	SLAB_ROUGHNESSY(ResultBSDF) =		lerp(SLAB_ROUGHNESSY(BackgroundBSDF),								SLAB_ROUGHNESSY(ForegroundBSDF),								ForegroundMixFactor);
-	SLAB_HAZINESS(ResultBSDF) =			lerp(SLAB_HAZINESS(BackgroundBSDF),									SLAB_HAZINESS(ForegroundBSDF),									ForegroundMixFactor);
-	SLAB_THINFILM_THICKNESS(ResultBSDF) = 0;// Thinfilm does not support coverage applied on pararmeters so we simply disable ThinFilm.=lerp(SLAB_THINFILM_THICKNESS(BackgroundBSDF),	SLAB_THINFILM_THICKNESS(ForegroundBSDF),	ForegroundMixFactor);
-	SLAB_FUZZ_AMOUNT(ResultBSDF) =		lerp(SLAB_FUZZ_AMOUNT(BackgroundBSDF),								SLAB_FUZZ_AMOUNT(ForegroundBSDF),								ForegroundMixFactor);
+	BSDF_SETTHICKNESSCM(ResultBSDF,		lerp(BSDF_GETTHICKNESSCM(BackgroundBSDF),							BSDF_GETTHICKNESSCM(ForegroundBSDF),							ForegroundOtherMixFactor));
+	SLAB_METALLIC(ResultBSDF) =			lerp(SLAB_METALLIC(BackgroundBSDF),									SLAB_METALLIC(ForegroundBSDF),									ForegroundOtherMixFactor);
+	SLAB_ROUGHNESSX(ResultBSDF) =		lerp(SLAB_ROUGHNESSX(BackgroundBSDF),								SLAB_ROUGHNESSX(ForegroundBSDF),								ForegroundOtherMixFactor);
+	SLAB_ROUGHNESSY(ResultBSDF) =		lerp(SLAB_ROUGHNESSY(BackgroundBSDF),								SLAB_ROUGHNESSY(ForegroundBSDF),								ForegroundOtherMixFactor);
+	SLAB_HAZINESS(ResultBSDF) =			lerp(SLAB_HAZINESS(BackgroundBSDF),									SLAB_HAZINESS(ForegroundBSDF),									ForegroundOtherMixFactor);
+	SLAB_THINFILM_THICKNESS(ResultBSDF)=lerp(SLAB_THINFILM_THICKNESS(BackgroundBSDF),						SLAB_THINFILM_THICKNESS(ForegroundBSDF),						ForegroundOtherMixFactor);
+	SLAB_FUZZ_AMOUNT(ResultBSDF) =		lerp(SLAB_FUZZ_AMOUNT(BackgroundBSDF),								SLAB_FUZZ_AMOUNT(ForegroundBSDF),								ForegroundOtherMixFactor);
 
 	if (BSDF_GETHASSSSPROFILE(ResultBSDF))
 	{
@@ -991,7 +1005,7 @@ FStrataData StrataHorizontalMixingParameterBlending(FStrataData Foreground, FStr
 		SLAB_SSSPROFILERADIUSSCALE(ResultBSDF) = lerp(
 			bBackgroundHasSSSProfile ? SLAB_SSSPROFILERADIUSSCALE(BackgroundBSDF) : 0.0f,
 			bForegroundHasSSSProfile ? SLAB_SSSPROFILERADIUSSCALE(ForegroundBSDF) : 0.0f,
-			ForegroundMixFactor);
+			ForegroundOtherMixFactor);
 
 		if (SLAB_SSSPROFILERADIUSSCALE(ResultBSDF) == 0.0f)
 		{
@@ -1002,7 +1016,7 @@ FStrataData StrataHorizontalMixingParameterBlending(FStrataData Foreground, FStr
 	}
 	else
 	{
-		SLAB_SSSDMFP(ResultBSDF) =		lerp(SLAB_SSSDMFP(BackgroundBSDF),								SLAB_SSSDMFP(ForegroundBSDF),										ForegroundMixFactor);
+		SLAB_SSSDMFP(ResultBSDF) =		lerp(SLAB_SSSDMFP(BackgroundBSDF),								SLAB_SSSDMFP(ForegroundBSDF),										ForegroundOtherMixFactor);
 
 		// We still try to disable SSS if possible: for performance and also to preserve the surface color (especially specular reflection color).
 		if (all(SLAB_SSSDMFP(ResultBSDF) == 0.0f))
@@ -1014,10 +1028,10 @@ FStrataData StrataHorizontalMixingParameterBlending(FStrataData Foreground, FStr
 
 #if STRATA_INLINE_SHADING
 	// At this stage we do not know if this strata BSDF is going to be using simple volume lighting so always store the DMFP on the side to not override SSSPROFILE data.
-	ResultBSDF.TmpDMFP =				lerp(BackgroundBSDF.TmpDMFP,									ForegroundBSDF.TmpDMFP,												ForegroundMixFactor);
+	ResultBSDF.TmpDMFP =				lerp(BackgroundBSDF.TmpDMFP,									ForegroundBSDF.TmpDMFP,												ForegroundOtherMixFactor);
 
 	// We blend the parameter above, and we also need to lerp the Coverage to match the total energy.
-	ResultBSDF.Coverage =				lerp(BackgroundBSDF.Coverage,									ForegroundBSDF.Coverage,											ForegroundMixFactor);
+	ResultBSDF.Coverage =				lerp(BackgroundBSDF.Coverage,									ForegroundBSDF.Coverage,											ForegroundLightingMixFactor);
 #endif
 
 #undef ResultBSDF
@@ -1026,8 +1040,9 @@ FStrataData StrataHorizontalMixingParameterBlending(FStrataData Foreground, FStr
 	return Result;
 }
 
-float AddParameterBlendingBSDFCoverageToNormalMix(FStrataData AStrata, FStrataData BStrata) 
+float AddParameterBlendingBSDFCoverageToNormalMix(FStrataData AStrata, FStrataData BStrata)
 {
+	// This is a normalised mix so even though both BSDF have weights, the sum of weights used to combined them should always be 1.
 #if STRATA_INLINE_SHADING
 	const float SafeABSDFCoverage = saturate(AStrata.Layers[0].BSDFs[0].Coverage);
 	const float SafeBBSDFCoverage = saturate(BStrata.Layers[0].BSDFs[0].Coverage);
@@ -1124,6 +1139,7 @@ FStrataData StrataAddParameterBlending(FStrataData A, FStrataData B, float AMixF
 
 float VerticalLayeringParameterBlendingBSDFCoverageToNormalMix(FStrataData TopStrata)
 {
+	// This is a normalised mix so even though both BSDF have weights, the sum of weights used to combined them should always be 1.
 #if STRATA_INLINE_SHADING
 	const float SafeTopBSDFCoverage = saturate(TopStrata.Layers[0].BSDFs[0].Coverage);
 	const float TopNormalContributionForWeightOfOne = 0.5f;
diff --git a/Engine/Source/Runtime/Engine/Private/Materials/HLSLMaterialTranslator.cpp b/Engine/Source/Runtime/Engine/Private/Materials/HLSLMaterialTranslator.cpp
--- a/Engine/Source/Runtime/Engine/Private/Materials/HLSLMaterialTranslator.cpp
+++ b/Engine/Source/Runtime/Engine/Private/Materials/HLSLMaterialTranslator.cpp
@@ -8995,6 +8995,11 @@ int32 FHLSLMaterialTranslator::StrataVerticalLayeringParameterBlendingBSDFCovera
 	return AddCodeChunk(MCT_Float, TEXT("VerticalLayeringParameterBlendingBSDFCoverageToNormalMix(%s)"), *GetParameterCode(TopCodeChunk));
 }
 
+int32 FHLSLMaterialTranslator::StrataHorizontalMixingParameterBlendingBSDFCoverageToNormalMixCodeChunk(int32 BackgroundCodeChunk, int32 ForegroundCodeChunk, int32 HorizontalMixCodeChunk)
+{
+	return AddCodeChunk(MCT_Float, TEXT("HorizontalMixingParameterBlendingBSDFCoverageToNormalMix(%s, %s, %s)"), *GetParameterCode(BackgroundCodeChunk), *GetParameterCode(ForegroundCodeChunk), *GetParameterCode(HorizontalMixCodeChunk));
+}
+
 int32 FHLSLMaterialTranslator::StrataCreateAndRegisterNullMaterial()
 {
 	int32 OutputCodeChunk = AddInlinedCodeChunk(MCT_Strata, TEXT("GetInitialisedStrataData()"));
@@ -9118,17 +9123,18 @@ int32 FHLSLMaterialTranslator::StrataHorizontalMixing(int32 Foreground, int32 Ba
 	);
 }
 
-int32 FHLSLMaterialTranslator::StrataHorizontalMixingParameterBlending(int32 Foreground, int32 Background, int32 Mix, const FString& SharedLocalBasisIndexMacro)
+int32 FHLSLMaterialTranslator::StrataHorizontalMixingParameterBlending(int32 Foreground, int32 Background, int32 HorizontalMixCodeChunk, int32 NormalMixCodeChunk, const FString& SharedLocalBasisIndexMacro)
 {
-	if (Foreground == INDEX_NONE || Background == INDEX_NONE || Mix == INDEX_NONE)
+	if (Foreground == INDEX_NONE || Background == INDEX_NONE || HorizontalMixCodeChunk == INDEX_NONE || NormalMixCodeChunk == INDEX_NONE)
 	{
 		return INDEX_NONE;
 	}
 	return AddCodeChunk(
-		MCT_Strata, TEXT("StrataHorizontalMixingParameterBlending(%s, %s, %s, %s)"),
+		MCT_Strata, TEXT("StrataHorizontalMixingParameterBlending(%s, %s, %s, %s, %s)"),
 		*GetParameterCode(Foreground),
 		*GetParameterCode(Background),
-		*GetParameterCode(Mix),
+		*GetParameterCode(HorizontalMixCodeChunk),
+		*GetParameterCode(NormalMixCodeChunk),
 		*SharedLocalBasisIndexMacro
 	);
 }
diff --git a/Engine/Source/Runtime/Engine/Private/Materials/HLSLMaterialTranslator.h b/Engine/Source/Runtime/Engine/Private/Materials/HLSLMaterialTranslator.h
--- a/Engine/Source/Runtime/Engine/Private/Materials/HLSLMaterialTranslator.h
+++ b/Engine/Source/Runtime/Engine/Private/Materials/HLSLMaterialTranslator.h
@@ -997,7 +997,7 @@ protected:
 		int32 EmissiveColor, int32 TopMaterialOpacity, int32 WaterAlbedo, int32 WaterExtinction, int32 WaterPhaseG, 
 		int32 ColorScaleBehindWater, int32 Normal, const FString& SharedLocalBasisIndexMacro) override;
 	virtual int32 StrataHorizontalMixing(int32 Foreground, int32 Background, int32 Mix) override;
-	virtual int32 StrataHorizontalMixingParameterBlending(int32 Foreground, int32 Background, int32 Mix, const FString& SharedLocalBasisIndexMacro) override;
+	virtual int32 StrataHorizontalMixingParameterBlending(int32 Foreground, int32 Background, int32 HorizontalMixCodeChunk, int32 NormalMixCodeChunk, const FString& SharedLocalBasisIndexMacro) override;
 	virtual int32 StrataVerticalLayering(int32 Top, int32 Base) override;
 	virtual int32 StrataVerticalLayeringParameterBlending(int32 Top, int32 Base, const FString& SharedLocalBasisIndexMacro, int32 TopBSDFNormalCodeChunk) override;
 	virtual int32 StrataAdd(int32 A, int32 B) override;
@@ -1013,6 +1013,7 @@ protected:
 	virtual uint8 StrataCompilationInfoGetSharedLocalBasesCount() override;
 	virtual int32 StrataAddParameterBlendingBSDFCoverageToNormalMixCodeChunk(int32 ACodeChunk, int32 BCodeChunk) override;
 	virtual int32 StrataVerticalLayeringParameterBlendingBSDFCoverageToNormalMixCodeChunk(int32 TopCodeChunk) override;
+	virtual int32 StrataHorizontalMixingParameterBlendingBSDFCoverageToNormalMixCodeChunk(int32 BackgroundCodeChunk, int32 ForegroundCodeChunk, int32 HorizontalMixCodeChunk) override;
 
 	FStrataSharedLocalBasesInfo StrataCompilationInfoGetMatchingSharedLocalBasisInfo(const FStrataRegisteredSharedLocalBasis& SearchedSharedLocalBasis);
 
diff --git a/Engine/Source/Runtime/Engine/Private/Materials/MaterialExpressions.cpp b/Engine/Source/Runtime/Engine/Private/Materials/MaterialExpressions.cpp
--- a/Engine/Source/Runtime/Engine/Private/Materials/MaterialExpressions.cpp
+++ b/Engine/Source/Runtime/Engine/Private/Materials/MaterialExpressions.cpp
@@ -20828,11 +20828,13 @@ int32 UMaterialExpressionStrataHorizontalMixing::Compile(class FMaterialCompiler
 	const FStrataMaterialCompilationInfo& ForegroundStrataData = Compiler->GetStrataCompilationInfo(ForegroundCodeChunk);
 	const FStrataMaterialCompilationInfo& BackgroundStrataData = Compiler->GetStrataCompilationInfo(BackgroundCodeChunk);
 
-	const int32 MixCodeChunk = Mix.GetTracedInput().Expression ? Mix.Compile(Compiler) : Compiler->Constant(0.0f);
+	const int32 HorizontalMixCodeChunk = Mix.Compile(Compiler);
 
 	int32 OutputCodeChunk = INDEX_NONE;
 	if (bUseParameterBlending)
 	{
+		const int32 NormalMixCodeChunk = Compiler->StrataHorizontalMixingParameterBlendingBSDFCoverageToNormalMixCodeChunk(BackgroundCodeChunk, ForegroundCodeChunk, HorizontalMixCodeChunk);
+
 		if (ForegroundStrataData.TotalBSDFCount != 1)
 		{
 			return Compiler->Errorf(TEXT("Foreground: cannot parameter horizontal blend complex material topology using parameter blending. Only a one-to-one blending is supported today."));
@@ -20854,12 +20856,12 @@ int32 UMaterialExpressionStrataHorizontalMixing::Compile(class FMaterialCompiler
 		}
 
 		// Compute the new Normal and Tangent resulting from the blending using code chunk
-		const int32 NewNormalCodeChunk = StrataBlendNormal(Compiler, BackgroundBSDF.RegisteredSharedLocalBasis.NormalCodeChunk, ForegroundBSDF.RegisteredSharedLocalBasis.NormalCodeChunk, MixCodeChunk);
+		const int32 NewNormalCodeChunk = StrataBlendNormal(Compiler, BackgroundBSDF.RegisteredSharedLocalBasis.NormalCodeChunk, ForegroundBSDF.RegisteredSharedLocalBasis.NormalCodeChunk, NormalMixCodeChunk);
 		// The tangent is optional so we treat it differently if INDEX_NONE is specified
 		int32 NewTangentCodeChunk = INDEX_NONE;
 		if (ForegroundBSDF.RegisteredSharedLocalBasis.TangentCodeChunk != INDEX_NONE && BackgroundBSDF.RegisteredSharedLocalBasis.TangentCodeChunk != INDEX_NONE)
 		{
-			NewTangentCodeChunk = StrataBlendNormal(Compiler, ForegroundBSDF.RegisteredSharedLocalBasis.TangentCodeChunk, BackgroundBSDF.RegisteredSharedLocalBasis.TangentCodeChunk, MixCodeChunk);
+			NewTangentCodeChunk = StrataBlendNormal(Compiler, BackgroundBSDF.RegisteredSharedLocalBasis.TangentCodeChunk, ForegroundBSDF.RegisteredSharedLocalBasis.TangentCodeChunk, NormalMixCodeChunk);
 		}
 		else if (ForegroundBSDF.RegisteredSharedLocalBasis.TangentCodeChunk != INDEX_NONE)
 		{
@@ -20871,7 +20873,7 @@ int32 UMaterialExpressionStrataHorizontalMixing::Compile(class FMaterialCompiler
 		}
 		const FStrataRegisteredSharedLocalBasis NewRegisteredSharedLocalBasis = StrataCompilationInfoCreateSharedLocalBasis(Compiler, NewNormalCodeChunk, NewTangentCodeChunk);
 
-		OutputCodeChunk = Compiler->StrataHorizontalMixingParameterBlending(ForegroundCodeChunk, BackgroundCodeChunk, MixCodeChunk, Compiler->GetStrataSharedLocalBasisIndexMacro(NewRegisteredSharedLocalBasis));
+		OutputCodeChunk = Compiler->StrataHorizontalMixingParameterBlending(ForegroundCodeChunk, BackgroundCodeChunk, HorizontalMixCodeChunk, NormalMixCodeChunk, Compiler->GetStrataSharedLocalBasisIndexMacro(NewRegisteredSharedLocalBasis));
 
 		StrataCompilationInfoCreateSingleBSDFMaterial(Compiler, OutputCodeChunk, NewRegisteredSharedLocalBasis, STRATA_BSDF_TYPE_SLAB);
 		FStrataMaterialCompilationInfo StrataInfo = StrataCompilationInfoHorizontalMixingParamBlend(Compiler, ForegroundStrataData, BackgroundStrataData, NewRegisteredSharedLocalBasis);
@@ -20883,7 +20885,7 @@ int32 UMaterialExpressionStrataHorizontalMixing::Compile(class FMaterialCompiler
 		OutputCodeChunk = Compiler->StrataHorizontalMixing(
 			ForegroundCodeChunk,
 			BackgroundCodeChunk,
-			MixCodeChunk);
+			HorizontalMixCodeChunk);
 
 		FStrataMaterialCompilationInfo StrataInfo = StrataCompilationInfoHorizontalMixing(Compiler, ForegroundStrataData, BackgroundStrataData);
 
@@ -21143,7 +21145,7 @@ int32 UMaterialExpressionStrataAdd::Compile(class FMaterialCompiler* Compiler, i
 		int32 NewTangentCodeChunk = INDEX_NONE;
 		if (ABSDF.RegisteredSharedLocalBasis.TangentCodeChunk != INDEX_NONE && BBSDF.RegisteredSharedLocalBasis.TangentCodeChunk != INDEX_NONE)
 		{
-			NewTangentCodeChunk = StrataBlendNormal(Compiler, ABSDF.RegisteredSharedLocalBasis.TangentCodeChunk, BBSDF.RegisteredSharedLocalBasis.TangentCodeChunk, ANormalMixCodeChunk);
+			NewTangentCodeChunk = StrataBlendNormal(Compiler, BBSDF.RegisteredSharedLocalBasis.TangentCodeChunk, ABSDF.RegisteredSharedLocalBasis.TangentCodeChunk, ANormalMixCodeChunk);
 		}
 		else if (ABSDF.RegisteredSharedLocalBasis.TangentCodeChunk != INDEX_NONE)
 		{
diff --git a/Engine/Source/Runtime/Engine/Public/MaterialCompiler.h b/Engine/Source/Runtime/Engine/Public/MaterialCompiler.h
--- a/Engine/Source/Runtime/Engine/Public/MaterialCompiler.h
+++ b/Engine/Source/Runtime/Engine/Public/MaterialCompiler.h
@@ -440,7 +440,7 @@ public:
 		int32 BaseColor, int32 Metallic, int32 Specular, int32 Roughness, int32 EmissiveColor, int32 TopMaterialOpacity, 
 		int32 WaterAlbedo, int32 WaterExtinction, int32 WaterPhaseG, int32 ColorScaleBehindWater, int32 Normal, const FString& SharedLocalBasisIndexMacro) = 0;
 	virtual int32 StrataHorizontalMixing(int32 Foreground, int32 Background, int32 Mix) = 0;
-	virtual int32 StrataHorizontalMixingParameterBlending(int32 Foreground, int32 Background, int32 Mix, const FString& SharedLocalBasisIndexMacro) = 0;
+	virtual int32 StrataHorizontalMixingParameterBlending(int32 Foreground, int32 Background, int32 HorizontalMixCodeChunk, int32 NormalMixCodeChunk, const FString& SharedLocalBasisIndexMacro) = 0;
 	virtual int32 StrataVerticalLayering(int32 Top, int32 Base) = 0;
 	virtual int32 StrataVerticalLayeringParameterBlending(int32 Top, int32 Base, const FString& SharedLocalBasisIndexMacro, int32 TopBSDFNormalCodeChunk) = 0;
 	virtual int32 StrataAdd(int32 A, int32 B) = 0;
@@ -460,6 +460,7 @@ public:
 	}
 	virtual int32 StrataAddParameterBlendingBSDFCoverageToNormalMixCodeChunk(int32 ACodeChunk, int32 BCodeChunk) = 0;
 	virtual int32 StrataVerticalLayeringParameterBlendingBSDFCoverageToNormalMixCodeChunk(int32 TopCodeChunk) = 0;
+	virtual int32 StrataHorizontalMixingParameterBlendingBSDFCoverageToNormalMixCodeChunk(int32 BackgroundCodeChunk, int32 ForegroundCodeChunk, int32 HorizontalMixCodeChunk) = 0;
 
 	// Water
 	virtual int32 SceneDepthWithoutWater(int32 Offset, int32 ViewportUV, bool bUseOffset, float FallbackDepth) = 0;
@@ -983,9 +984,9 @@ public:
 		return Compiler->StrataHorizontalMixing(Foreground, Background, Mix);
 	}
 
-	virtual int32 StrataHorizontalMixingParameterBlending(int32 Foreground, int32 Background, int32 Mix, const FString& SharedLocalBasisIndexMacro) override
+	virtual int32 StrataHorizontalMixingParameterBlending(int32 Foreground, int32 Background, int32 HorizontalMixCodeChunk, int32 NormalMixCodeChunk, const FString& SharedLocalBasisIndexMacro) override
 	{
-		return Compiler->StrataHorizontalMixingParameterBlending(Foreground, Background, Mix, SharedLocalBasisIndexMacro);
+		return Compiler->StrataHorizontalMixingParameterBlending(Foreground, Background, HorizontalMixCodeChunk, NormalMixCodeChunk, SharedLocalBasisIndexMacro);
 	}
 
 	virtual int32 StrataVerticalLayering(int32 Top, int32 Base) override
@@ -1037,6 +1038,11 @@ public:
 	{
 		return Compiler->StrataVerticalLayeringParameterBlendingBSDFCoverageToNormalMixCodeChunk(TopCodeChunk);
 	}
+
+	virtual int32 StrataHorizontalMixingParameterBlendingBSDFCoverageToNormalMixCodeChunk(int32 BackgroundCodeChunk, int32 ForegroundCodeChunk, int32 HorizontalMixCodeChunk) override
+	{
+		return Compiler->StrataHorizontalMixingParameterBlendingBSDFCoverageToNormalMixCodeChunk(BackgroundCodeChunk, ForegroundCodeChunk, HorizontalMixCodeChunk);
+	}
 	
 	virtual const FStrataMaterialCompilationInfo& GetStrataCompilationInfo(int32 CodeChunk) override
 	{
