Strata F90 moved from HSV to YCoCg. conversion code is roughly 4 time less instructions.

diff --git a/Engine/Shaders/Private/ColorSpace.ush b/Engine/Shaders/Private/ColorSpace.ush
--- a/Engine/Shaders/Private/ColorSpace.ush
+++ b/Engine/Shaders/Private/ColorSpace.ush
@@ -143,6 +143,11 @@ float3 LAB_2_LinearRGB( float3 LAB, float3 ReferenceWhite)
 // Multiplier that get applied for ALU optimisation purposes on the luma, when converting LinearRGB to YCoCg.
 #define LINEARRGB_2_YCOCG_LUMA_MULTIPLIER 4.0
 
+// YCoCg ranges: 
+//  - Y  in [0, LINEARRGB_2_YCOCG_LUMA_MULTIPLIER]
+//  - Co in [-2, 2]
+//  - Cg in [-2, 2]
+
 float3 LinearRGB_2_YCoCg(float3 RGB)
 {
 	float Y = dot(RGB, float3(1, 2, 1));
@@ -189,6 +194,17 @@ float3 LCoCg_2_LinearRGB(float3 LCoCg)
 	return YCoCg_2_LinearRGB(LCoCg_2_YCoCg(LCoCg));
 }
 
+// NormalisedYCoCg has Y, Co and Cg values mapped to [0, 1]
+float3 LinearRGB_2_NormalisedYCoCg(float3 RGB)
+{
+	return LinearRGB_2_YCoCg(RGB) * float3(1.0f / LINEARRGB_2_YCOCG_LUMA_MULTIPLIER, 0.25f, 0.25f) + float3(0.0f, 0.5f, 0.5f);
+}
+
+float3 NormalisedYCoCg_2_LinearRGB(float3 YCoCg)
+{
+	return YCoCg_2_LinearRGB(YCoCg * float3(LINEARRGB_2_YCOCG_LUMA_MULTIPLIER, 4.0f, 4.0f) + float3(0.0f, -2.0f, -2.0f));
+}
+
 //------------------------------------------------------------------
 // HSV
 //------------------------------------------------------------------
diff --git a/Engine/Shaders/Private/Strata/Strata.ush b/Engine/Shaders/Private/Strata/Strata.ush
--- a/Engine/Shaders/Private/Strata/Strata.ush
+++ b/Engine/Shaders/Private/Strata/Strata.ush
@@ -1915,9 +1915,9 @@ void PackStrataOut(
 
 				if (BSDF_GETHASF90(BSDF) || BSDF_GETHASTHINFILM(BSDF) || BSDF_GETHASHAZINESS(BSDF))
 				{
-					float3 F90HSV = LinearRGB_2_HSV(SLAB_F90(BSDF));
+					float3 F90YCoCg = LinearRGB_2_NormalisedYCoCg(SLAB_F90(BSDF));
 					const float ThinFilmGamma2 = sqrt(SLAB_THINFILM_THICKNESS(BSDF)); // Using a gamma 2.0 encoding help with precision to be visually more uniform. This is then further improved via dithering.
-					STRATA_STORE_UINT1(PackRGBA8(float4(F90HSV.x, F90HSV.y, ThinFilmGamma2, SLAB_HAZINESS(BSDF))));
+					STRATA_STORE_UINT1(PackRGBA8(float4(F90YCoCg.y, F90YCoCg.z, ThinFilmGamma2, SLAB_HAZINESS(BSDF))));
 				}
 
 				const bool bIsSimpleVolume = BSDF_GETISSIMPLEVOLUME(BSDF);
@@ -2109,10 +2109,10 @@ FStrataBSDF UnpackStrataBSDFIn(ByteAddressBuffer StrataBuffer, inout FStrataAddr
 		{
 			float4 Data = UnpackRGBA8(StrataLoadUint1(StrataBuffer, StrataAddressing));
 
-			const float F90Hue = Data.x;
-			const float F90Sat = Data.y;
-			const float F90Val = 1.0f;	// Forced to 1 as F90 is not super to have a low value
-			SLAB_F90(OutBSDF) = bHasF90 ? HSV_2_LinearRGB(float3(F90Hue, F90Sat, F90Val)) : 1.0f;
+			const float F90Co = Data.x;
+			const float F90Cg = Data.y;
+			const float F90Y  = 1.0f;	// Forced to 1 as F90 is not supposed to have a low value
+			SLAB_F90(OutBSDF) = bHasF90 ? NormalisedYCoCg_2_LinearRGB(float3(F90Y, F90Co, F90Cg)) : 1.0f;
 
 			SLAB_THINFILM_THICKNESS(OutBSDF) = Data.z * Data.z;
 			SLAB_HAZINESS(OutBSDF) = Data.w;
