Added Strata backcompatibility mode (disabled multiple scattering and chan diffuse model for Lambert)

diff --git a/Engine/Shaders/Private/Strata/StrataEvaluation.ush b/Engine/Shaders/Private/Strata/StrataEvaluation.ush
--- a/Engine/Shaders/Private/Strata/StrataEvaluation.ush
+++ b/Engine/Shaders/Private/Strata/StrataEvaluation.ush
@@ -26,7 +26,11 @@
 #endif
 
 #ifndef STRATA_MICROFACET_MULTIPLESCATTERING
-#define STRATA_MICROFACET_MULTIPLESCATTERING 0
+#define STRATA_MICROFACET_MULTIPLESCATTERING (PROJECT_STRATA_BACKCOMPATIBILITY > 0 ? 0 : 1)
+#endif
+
+#ifndef STRATA_DIFFUSE_CHAN
+#define STRATA_DIFFUSE_CHAN (PROJECT_STRATA_BACKCOMPATIBILITY > 0 ? 0 : 1)
 #endif
 
 #ifndef STRATA_SSS_TRANSMISSION
@@ -245,7 +249,11 @@ FStrataEvaluateResult StrataEvaluateBSDF(FStrataBSDFContext BSDFContext, FShadow
 			const float SafeRoughness = MakeRoughnessSafe(SLAB_ROUGHNESSX(BSDFContext.BSDF));
 			float a2 = Pow4(SafeRoughness);
 
-			Sample.DiffusePathValue	= Diffuse_Chan(DiffuseColor, a2, BSDFContext.SatNoV, BSDFContext.SatNoL, BSDFContext.Context.VoH, BSDFContext.Context.NoH);
+		#if STRATA_DIFFUSE_CHAN
+			Sample.DiffusePathValue = Diffuse_Chan(DiffuseColor, a2, BSDFContext.SatNoV, BSDFContext.SatNoL, BSDFContext.Context.VoH, BSDFContext.Context.NoH);
+		#else
+			Sample.DiffusePathValue = Diffuse_Lambert(DiffuseColor);
+		#endif
 			Sample.DiffusePDF		= BSDFContext.SatNoL / PI;
 			Sample.bSubsurface		= bHasSubsurface;
 
diff --git a/Engine/Source/Runtime/Engine/Private/ShaderCompiler/ShaderCompiler.cpp b/Engine/Source/Runtime/Engine/Private/ShaderCompiler/ShaderCompiler.cpp
--- a/Engine/Source/Runtime/Engine/Private/ShaderCompiler/ShaderCompiler.cpp
+++ b/Engine/Source/Runtime/Engine/Private/ShaderCompiler/ShaderCompiler.cpp
@@ -5193,8 +5193,12 @@ void GlobalBeginCompileShader(
 
 	{
 		static IConsoleVariable* CVar = IConsoleManager::Get().FindConsoleVariable(TEXT("r.Strata"));
-		bool bStrata = CVar && CVar->GetInt() != 0;
+		const bool bStrata = CVar && CVar->GetInt() != 0;
 		Input.Environment.SetDefine(TEXT("PROJECT_STRATA"), bStrata ? 1 : 0);
+
+		static IConsoleVariable* CVarBackCompatibility = IConsoleManager::Get().FindConsoleVariable(TEXT("r.StrataBackCompatibility"));
+		const bool bStrataBackCompatibility = CVarBackCompatibility && CVarBackCompatibility->GetInt() > 0;
+		Input.Environment.SetDefine(TEXT("PROJECT_STRATA_BACKCOMPATIBILITY"), bStrataBackCompatibility ? 1 : 0);
 	}
 
 	{
diff --git a/Engine/Source/Runtime/RenderCore/Private/Shader.cpp b/Engine/Source/Runtime/RenderCore/Private/Shader.cpp
--- a/Engine/Source/Runtime/RenderCore/Private/Shader.cpp
+++ b/Engine/Source/Runtime/RenderCore/Private/Shader.cpp
@@ -1714,6 +1714,12 @@ void ShaderMapAppendKeyString(EShaderPlatform Platform, FString& KeyString)
 		{
 			KeyString += FString::Printf(TEXT("_BUDGET%u"), CVarBudget->GetValueOnAnyThread());
 		}
+
+		static const auto CVarBackCompatibility = IConsoleManager::Get().FindTConsoleVariableDataInt(TEXT("r.StrataBackCompatibility"));
+		if (bStrataEnabled && CVarBackCompatibility && CVarBackCompatibility->GetValueOnAnyThread()>0)
+		{
+			KeyString += FString::Printf(TEXT("_BACKCOMPAT"));
+		}
 	}
 
 	{
diff --git a/Engine/Source/Runtime/Renderer/Private/Strata/Strata.cpp b/Engine/Source/Runtime/Renderer/Private/Strata/Strata.cpp
--- a/Engine/Source/Runtime/Renderer/Private/Strata/Strata.cpp
+++ b/Engine/Source/Runtime/Renderer/Private/Strata/Strata.cpp
@@ -21,6 +21,12 @@ static TAutoConsoleVariable<int32> CVarStrata(
 	TEXT("Enable Strata materials (Beta)."),
 	ECVF_ReadOnly | ECVF_RenderThreadSafe);
 
+static TAutoConsoleVariable<int32> CVarStrataBackCompatibility(
+	TEXT("r.StrataBackCompatibility"),
+	0,
+	TEXT("Disables Strata multiple scattering and replaces Chan diffuse by Lambert."),
+	ECVF_ReadOnly | ECVF_RenderThreadSafe);
+
 static TAutoConsoleVariable<int32> CVarStrataBytePerPixel(
 	TEXT("r.Strata.BytesPerPixel"),
 	80,
