Respect non-directional "Lighting Mode" settings for translucent materials on mobile.
This changes default look of translucency on mobile (now similar to desktop) as default setting is non-directional

diff --git a/Engine/Shaders/Private/MobileBasePassPixelShader.usf b/Engine/Shaders/Private/MobileBasePassPixelShader.usf
--- a/Engine/Shaders/Private/MobileBasePassPixelShader.usf
+++ b/Engine/Shaders/Private/MobileBasePassPixelShader.usf
@@ -93,6 +93,7 @@ int NumDynamicPointLights;
 #define SUPPORT_SPOTLIGHTS_SHADOW (MATERIALBLENDING_SOLID || MATERIALBLENDING_MASKED) && PROJECT_MOBILE_ENABLE_MOVABLE_SPOTLIGHTS_SHADOW
 #define DEFERRED_SHADING_PATH (MOBILE_DEFERRED_SHADING && ((MATERIALBLENDING_SOLID || MATERIALBLENDING_MASKED) && !MATERIAL_SHADINGMODEL_SINGLELAYERWATER))
 #define USE_SCENE_DEPTH_AUX (DEFERRED_SHADING_PATH && METAL_PROFILE)
+#define TRANSLUCENCY_NON_DIRECTIONAL (MATERIALBLENDING_ANY_TRANSLUCENT && (TRANSLUCENCY_LIGHTING_VOLUMETRIC_NONDIRECTIONAL || TRANSLUCENCY_LIGHTING_VOLUMETRIC_PERVERTEX_NONDIRECTIONAL) && !MATERIAL_SHADINGMODEL_SINGLELAYERWATER)
 
 #define ALLOW_CUBE_REFLECTIONS HQ_REFLECTIONS
 
@@ -776,6 +777,11 @@ void Main(
 	half3 H = normalize(MaterialParameters.CameraVector + MobileDirectionalLight.DirectionalLightDirectionAndShadowTransition.xyz);
 	half NoH = max(0, dot(MaterialParameters.WorldNormal, H));
 
+#if TRANSLUCENCY_NON_DIRECTIONAL
+	NoL = 1.f;
+	ShadingModelContext.SpecularColor = 0.f;
+#endif
+
 	// Direct Lighting (Directional light) + IBL
 #if FULLY_ROUGH
 	Color += (Shadow * NoL) * MobileDirectionalLight.DirectionalLightColor.rgb * ShadingModelContext.DiffuseColor;
