Respect translucency non-directional settings for local lights on mobile

diff --git a/Engine/Shaders/Private/MobileBasePassPixelShader.usf b/Engine/Shaders/Private/MobileBasePassPixelShader.usf
--- a/Engine/Shaders/Private/MobileBasePassPixelShader.usf
+++ b/Engine/Shaders/Private/MobileBasePassPixelShader.usf
@@ -453,6 +453,10 @@ inout half3 Color)
 		half PointNoL = max(0, dot(MaterialParameters.WorldNormal, L));
 		half PointNoH = max(0, dot(MaterialParameters.WorldNormal, PointH));
 
+		#if TRANSLUCENCY_NON_DIRECTIONAL
+			PointNoL = 1.f;
+		#endif
+
 		float Attenuation;
 
 		if (LightColorAndFalloffExponent.w == 0)
