Add debug probes pass to better debug lighting and reflections.

diff --git a/Engine/Shaders/Private/DebugProbes.usf b/Engine/Shaders/Private/DebugProbes.usf
new file mode 100644
--- /dev/null
+++ b/Engine/Shaders/Private/DebugProbes.usf
@@ -0,0 +1,166 @@
+
+
+
+#include "DeferredShadingCommon.ush"
+
+
+
+#define PROBE_MODE_SPECULAR	0
+#define PROBE_MODE_DIFFUSE	1
+#define PROBE_MODE_BOTH		2
+
+
+
+void DebugSphereIntersection(
+	float3 RayOrig, float3 RayDir, float3 SphereCenter, float SphereRadius, int NewProbeMode,
+	inout int ProbeMode, inout float RayT, inout float3 WorlPosition, inout float3 Normal)
+{
+	float2 RayTs = RayIntersectSphere(RayOrig, RayDir, float4(SphereCenter, SphereRadius));
+	float NewRayT = max(0.0f, min(RayTs.x, RayTs.y));
+
+	if (NewRayT > 0.0f && (NewRayT < RayT || RayT==-1.0f))
+	{
+		RayT = NewRayT;
+		ProbeMode = NewProbeMode;
+		WorlPosition = RayOrig + RayDir * RayT;
+		Normal = normalize(WorlPosition - SphereCenter);
+
+		if (ProbeMode == PROBE_MODE_BOTH)
+		{
+			bool bOn = Normal.z > 0.0f;
+			bOn = Normal.y > 0.0f ? !bOn : bOn;
+			bOn = Normal.x > 0.0f ? !bOn : bOn;
+			ProbeMode = bOn ? PROBE_MODE_SPECULAR : PROBE_MODE_DIFFUSE;
+		}
+	}
+}
+
+
+
+int DebugProbesMode;
+
+
+
+void MainPS(
+	  in  float4 SvPosition		: SV_Position
+	, out float4 OutSceneColor	: SV_Target0
+	, out float4 OutGBufferA	: SV_Target1
+	, out float4 OutGBufferB	: SV_Target2
+	, out float4 OutGBufferC	: SV_Target3
+	, out float4 OutGBufferD	: SV_Target4
+	, out float4 OutGBufferE	: SV_Target5
+	, out float  OutDepth		: SV_DEPTH
+)
+{
+	ResolvedView = ResolveView();
+	const uint2 PixelCoord = uint2(SvPosition.xy);
+
+	float2 ScreenPosition = SvPositionToScreenPosition(SvPosition).xy;
+	const float Depth = 1000000.0f;
+	float4 WorldPos = mul(float4(ScreenPosition * Depth, Depth, 1), View.ScreenToWorld);
+
+	const float3 RayDir = normalize(WorldPos.xyz - View.WorldCameraOrigin);
+	const float3 RayOrig = ResolvedView.WorldCameraOrigin;
+
+	int ProbeMode = false;
+	float RayT = -1.0f;
+	float3 WorlPosition = 0.0f;
+	float3 Normal = 0.0f;
+
+
+	// Render the camera probes (this could be done proceppdurally as flat sphere normal on screen rotated by inverse view matrix)
+	if (DebugProbesMode == 1 || DebugProbesMode == 3)
+	{
+		{
+			float3 SphereCenter = ResolvedView.WorldCameraOrigin + (ResolvedView.ViewForward * 1.0f + ResolvedView.ViewRight * 0.5f + ResolvedView.ViewUp * -0.2f) * METER_TO_CENTIMETER;
+			float  SphereRadius = 0.2f * METER_TO_CENTIMETER;
+
+			DebugSphereIntersection(
+				RayOrig, RayDir, SphereCenter, SphereRadius, PROBE_MODE_SPECULAR,
+				ProbeMode, RayT, WorlPosition, Normal);
+		}
+		{
+			float3 SphereCenter = ResolvedView.WorldCameraOrigin + (ResolvedView.ViewForward * 1.0f + ResolvedView.ViewRight * -0.5f + ResolvedView.ViewUp * -0.2f) * METER_TO_CENTIMETER;
+			float  SphereRadius = 0.2f * METER_TO_CENTIMETER;
+
+			DebugSphereIntersection(
+				RayOrig, RayDir, SphereCenter, SphereRadius, PROBE_MODE_DIFFUSE,
+				ProbeMode, RayT, WorlPosition, Normal);
+		}
+	}
+
+
+	// Render the world probes around the camera (this could be accelerated by tracing voxel and only intersecting with spheres of traversed voxel)
+	if (DebugProbesMode == 2 || DebugProbesMode == 3)
+	{
+		const float SphereRadius = 0.5f * METER_TO_CENTIMETER;
+
+		const float3 Steps = float3(5.0f, 5.0f, 2.5f);
+		const float3 SnappedOrigin = floor(ResolvedView.WorldCameraOrigin / (Steps * METER_TO_CENTIMETER) + 0.5) * Steps * METER_TO_CENTIMETER
+								   + float3(0.0f, 0.0f, 1.4 * SphereRadius);	// Small offset to have probes above the ground level 0
+		LOOP
+		for (float x = -15.0f * METER_TO_CENTIMETER; x <= 15.0f * METER_TO_CENTIMETER; x += Steps.x * METER_TO_CENTIMETER)
+		{
+			LOOP
+			for (float y = -15.0f * METER_TO_CENTIMETER; y <= 15.0f * METER_TO_CENTIMETER; y += Steps.y * METER_TO_CENTIMETER)
+			{
+				LOOP
+				for (float z = -5.0f * METER_TO_CENTIMETER; z <= 5.0f * METER_TO_CENTIMETER; z += Steps.z * METER_TO_CENTIMETER)
+				{
+					float3 SphereCenter = SnappedOrigin + float3(x, y, z);
+
+					DebugSphereIntersection(
+						RayOrig, RayDir, SphereCenter, SphereRadius, PROBE_MODE_BOTH,
+						ProbeMode, RayT, WorlPosition, Normal);
+				}
+			}
+		}
+	}
+
+
+	float4 ClipPos = mul(float4(WorlPosition, 1.0f), ResolvedView.WorldToClip);
+	ClipPos /= ClipPos.w;
+	const float DeviceZ = ClipPos.z;
+	OutDepth = DeviceZ;
+
+	if (RayT <= 0.0f)
+	{
+		discard;
+		return;
+	}
+
+
+
+
+	const float PerObjectGBufferData = 0.0f;
+	const float SelectiveOutputMask = 0.0f;
+
+	const float3 BaseColor = 1.0f;
+	const float  Metallic  = ProbeMode == PROBE_MODE_DIFFUSE ? 0.0f : 1.0f;
+	const float  Specular  = ProbeMode == PROBE_MODE_DIFFUSE ? 0.0f : 1.0f;
+	const float  Roughness = 0.01f;
+	const float  IndirectIrradianceAndAO = 0.0f;
+
+	// We write all outputs, the rasterizer will discard what is not bound. That is fine since we do not care that much about performance for this debug view.
+
+	OutSceneColor = 0.0f;
+
+	OutGBufferA = float4(
+		EncodeNormal(Normal), 
+		PerObjectGBufferData);
+
+	OutGBufferB = float4(
+		Metallic, 
+		Specular, 
+		Roughness, 
+		EncodeShadingModelIdAndSelectiveOutputMask(SHADINGMODELID_DEFAULT_LIT, SelectiveOutputMask));
+
+	OutGBufferC = float4(BaseColor, IndirectIrradianceAndAO);
+	OutGBufferD = float4(0.0f, 0.0f, 0.0f, 0.0f);	// custom data
+	OutGBufferE = float4(1.0f, 1.0f, 1.0f, 1.0f);	// Pre computed shadow factor
+
+	// Not writting to the velocity buffer
+
+	OutDepth = DeviceZ;
+
+}
\ No newline at end of file
diff --git a/Engine/Source/Runtime/Renderer/Private/BasePassRendering.cpp b/Engine/Source/Runtime/Renderer/Private/BasePassRendering.cpp
--- a/Engine/Source/Runtime/Renderer/Private/BasePassRendering.cpp
+++ b/Engine/Source/Runtime/Renderer/Private/BasePassRendering.cpp
@@ -18,6 +18,7 @@
 #include "OneColorShader.h"
 #include "ClearQuad.h"
 #include "ProfilingDebugging/CpuProfilerTrace.h"
+#include "DebugProbeRendering.h"
 
 // Changing this causes a full shader recompile
 static TAutoConsoleVariable<int32> CVarSelectiveBasePassOutputs(
@@ -908,6 +909,13 @@ void FDeferredShadingSceneRenderer::RenderBasePass(
 		RenderAnisotropyPass(GraphBuilder, SceneTextures, bEnableParallelBasePasses);
 		GraphBuilder.SetCommandListStat(GET_STATID(STAT_CLM_AfterAnisotropyPass));
 	}
+
+#if !(UE_BUILD_SHIPPING)
+	if (!bForwardShadingEnabled)
+	{
+		StampDeferredDebugProbeMaterialPS(GraphBuilder, Views, BasePassRenderTargets);
+	}
+#endif
 }
 
 BEGIN_SHADER_PARAMETER_STRUCT(FOpaqueBasePassParameters, )
diff --git a/Engine/Source/Runtime/Renderer/Private/DebugProbeRendering.cpp b/Engine/Source/Runtime/Renderer/Private/DebugProbeRendering.cpp
new file mode 100644
--- /dev/null
+++ b/Engine/Source/Runtime/Renderer/Private/DebugProbeRendering.cpp
@@ -0,0 +1,121 @@
+// Copyright Epic Games, Inc. All Rights Reserved.
+
+#include "DebugProbeRendering.h"
+#include "PixelShaderUtils.h"
+#include "ShaderParameterStruct.h"
+
+
+// Changing this causes a full shader recompile
+static TAutoConsoleVariable<int32> CVarDebugProbes(
+	TEXT("r.DebugProbes"),
+	0,
+	TEXT("Enables debug probes rendering to visualise diffuse/specular lighting on simple sphere scattered in the world.") \
+	TEXT(" 0: disabled.\n")
+	TEXT(" 1: camera probes only.\n")
+	TEXT(" 2: world probes only.\n")
+	TEXT(" 3: camera and world probes.\n")
+	,
+	ECVF_RenderThreadSafe);
+
+
+DECLARE_GPU_STAT(StampDeferredDebugProbe);
+
+
+class FStampDeferredDebugProbePS : public FGlobalShader
+{
+	DECLARE_GLOBAL_SHADER(FStampDeferredDebugProbePS);
+	SHADER_USE_PARAMETER_STRUCT(FStampDeferredDebugProbePS, FGlobalShader);
+
+	using FPermutationDomain = TShaderPermutationDomain<>;
+
+	BEGIN_SHADER_PARAMETER_STRUCT(FParameters, )
+		SHADER_PARAMETER_STRUCT_REF(FViewUniformShaderParameters, ViewUniformBuffer)
+		SHADER_PARAMETER(int32, DebugProbesMode)
+		RENDER_TARGET_BINDING_SLOTS()
+	END_SHADER_PARAMETER_STRUCT()
+
+public:
+
+	static FPermutationDomain RemapPermutation(FPermutationDomain PermutationVector)
+	{
+		return PermutationVector;
+	}
+
+	static bool ShouldCompilePermutation(const FGlobalShaderPermutationParameters& Parameters)
+	{
+		return IsFeatureLevelSupported(Parameters.Platform, ERHIFeatureLevel::SM5) && EnumHasAllFlags(Parameters.Flags, EShaderPermutationFlags::HasEditorOnlyData);
+	}
+	static void ModifyCompilationEnvironment(const FGlobalShaderPermutationParameters& Parameters, FShaderCompilerEnvironment& OutEnvironment)
+	{
+		FGlobalShader::ModifyCompilationEnvironment(Parameters, OutEnvironment);
+	}
+};
+
+IMPLEMENT_GLOBAL_SHADER(FStampDeferredDebugProbePS, "/Engine/Private/DebugProbes.usf", "MainPS", SF_Pixel);
+
+ 
+template<bool bEnableDepthWrite>
+static void CommonStampDeferredDebugProbeDrawCall(
+	FRDGBuilder& GraphBuilder,
+	const FViewInfo& View, 
+	FStampDeferredDebugProbePS::FParameters* PassParameters)
+{
+	PassParameters->ViewUniformBuffer = View.ViewUniformBuffer;
+	PassParameters->DebugProbesMode = FMath::Clamp(CVarDebugProbes.GetValueOnRenderThread(), 0, 3);
+
+	FStampDeferredDebugProbePS::FPermutationDomain PermutationVector;
+	TShaderMapRef<FStampDeferredDebugProbePS> PixelShader(View.ShaderMap, PermutationVector);
+
+	FPixelShaderUtils::AddFullscreenPass<FStampDeferredDebugProbePS>(
+		GraphBuilder, View.ShaderMap, RDG_EVENT_NAME("StampDeferredDebugProbePS"),
+		PixelShader, PassParameters, View.ViewRect,
+		TStaticBlendState<>::GetRHI(),
+		TStaticRasterizerState<FM_Solid, CM_None>::GetRHI(),
+		TStaticDepthStencilState<bEnableDepthWrite, CF_DepthNearOrEqual>::GetRHI());
+}
+
+void StampDeferredDebugProbeDepthPS(
+	FRDGBuilder& GraphBuilder,
+	TArrayView<const FViewInfo> Views,
+	const FRDGTextureRef SceneDepthTexture)
+{
+	RDG_EVENT_SCOPE(GraphBuilder, "StampDeferredDebugProbeDepth");
+	RDG_GPU_STAT_SCOPE(GraphBuilder, StampDeferredDebugProbe);
+
+	if (CVarDebugProbes.GetValueOnRenderThread() <= 0)
+	{
+		return;
+	}
+
+	for (const FViewInfo& View : Views)
+	{
+		FStampDeferredDebugProbePS::FParameters* PassParameters = GraphBuilder.AllocParameters<FStampDeferredDebugProbePS::FParameters>();
+		PassParameters->ViewUniformBuffer = View.ViewUniformBuffer;
+		PassParameters->RenderTargets.DepthStencil = FDepthStencilBinding(SceneDepthTexture, ERenderTargetLoadAction::ELoad, ERenderTargetLoadAction::ELoad, FExclusiveDepthStencil::DepthWrite_StencilWrite);
+
+		CommonStampDeferredDebugProbeDrawCall<true>(GraphBuilder, View, PassParameters);
+	}
+}
+
+void StampDeferredDebugProbeMaterialPS(
+	FRDGBuilder& GraphBuilder,
+	TArrayView<const FViewInfo> Views,
+	const FRenderTargetBindingSlots& BasePassRenderTargets)
+{
+	RDG_EVENT_SCOPE(GraphBuilder, "StampDeferredDebugProbeMaterial");
+	RDG_GPU_STAT_SCOPE(GraphBuilder, StampDeferredDebugProbe);
+
+	if (CVarDebugProbes.GetValueOnRenderThread() <= 0)
+	{
+		return;
+	}
+
+	for (const FViewInfo& View : Views)
+	{
+		FStampDeferredDebugProbePS::FParameters* PassParameters = GraphBuilder.AllocParameters<FStampDeferredDebugProbePS::FParameters>();
+		PassParameters->RenderTargets = BasePassRenderTargets;
+
+		CommonStampDeferredDebugProbeDrawCall<false>(GraphBuilder, View, PassParameters);
+	}
+}
+
diff --git a/Engine/Source/Runtime/Renderer/Private/DebugProbeRendering.h b/Engine/Source/Runtime/Renderer/Private/DebugProbeRendering.h
new file mode 100644
--- /dev/null
+++ b/Engine/Source/Runtime/Renderer/Private/DebugProbeRendering.h
@@ -0,0 +1,24 @@
+// Copyright Epic Games, Inc. All Rights Reserved.
+
+/*=============================================================================
+	DebugProbeRendering.h: Debug probe pass rendering definitions.
+=============================================================================*/
+
+#pragma once
+
+#include "CoreMinimal.h"
+#include "RenderGraphBuilder.h"
+#include "SceneRendering.h"
+
+
+void StampDeferredDebugProbeDepthPS(
+	FRDGBuilder& GraphBuilder,
+	TArrayView<const FViewInfo> Views,
+	const FRDGTextureRef SceneDepthTexture);
+
+
+void StampDeferredDebugProbeMaterialPS(
+	FRDGBuilder& GraphBuilder,
+	TArrayView<const FViewInfo> Views,
+	const FRenderTargetBindingSlots& BasePassRenderTargets);
+
diff --git a/Engine/Source/Runtime/Renderer/Private/DepthRendering.cpp b/Engine/Source/Runtime/Renderer/Private/DepthRendering.cpp
--- a/Engine/Source/Runtime/Renderer/Private/DepthRendering.cpp
+++ b/Engine/Source/Runtime/Renderer/Private/DepthRendering.cpp
@@ -31,6 +31,7 @@
 #include "PixelShaderUtils.h"
 #include "RenderGraphUtils.h"
 #include "SceneRenderingUtils.h"
+#include "DebugProbeRendering.h"
 
 static TAutoConsoleVariable<int32> CVarParallelPrePass(
 	TEXT("r.ParallelPrePass"),
@@ -602,6 +603,14 @@ void FDeferredShadingSceneRenderer::RenderPrePass(FRDGBuilder& GraphBuilder, FRD
 			DrawClearQuad(RHICmdList, false, FLinearColor::Transparent, false, 0, true, 0);
 		});
 	}
+
+#if !(UE_BUILD_SHIPPING)
+	const bool bForwardShadingEnabled = IsForwardShadingEnabled(ShaderPlatform);
+	if (!bForwardShadingEnabled)
+	{
+		StampDeferredDebugProbeDepthPS(GraphBuilder, Views, SceneDepthTexture);
+	}
+#endif
 }
 
 bool FMobileSceneRenderer::ShouldRenderPrePass() const
