Debug indirect lighting probes now write velocity for improved stability.

diff --git a/Engine/Shaders/Private/DebugProbes.usf b/Engine/Shaders/Private/DebugProbes.usf
--- a/Engine/Shaders/Private/DebugProbes.usf
+++ b/Engine/Shaders/Private/DebugProbes.usf
@@ -2,6 +2,7 @@
 
 
 #include "DeferredShadingCommon.ush"
+#include "VelocityCommon.ush"
 
 
 
@@ -12,8 +13,8 @@
 
 
 void DebugSphereIntersection(
-	float3 RayOrig, float3 RayDir, float3 SphereCenter, float SphereRadius, int NewProbeMode,
-	inout int ProbeMode, inout float RayT, inout float3 WorlPosition, inout float3 Normal)
+	float3 RayOrig, float3 RayDir, float3 SphereCenter, float SphereRadius, int NewProbeMode, in bool bNewProbeAttachedToCamera,
+	inout int ProbeMode, inout float RayT, inout float3 WorlPosition, inout float3 Normal, inout bool bProbeAttachedToCamera)
 {
 	float2 RayTs = RayIntersectSphere(RayOrig, RayDir, float4(SphereCenter, SphereRadius));
 	float NewRayT = max(0.0f, min(RayTs.x, RayTs.y));
@@ -24,6 +25,7 @@ void DebugSphereIntersection(
 		ProbeMode = NewProbeMode;
 		WorlPosition = RayOrig + RayDir * RayT;
 		Normal = normalize(WorlPosition - SphereCenter);
+		bProbeAttachedToCamera = bNewProbeAttachedToCamera;
 
 		if (ProbeMode == PROBE_MODE_BOTH)
 		{
@@ -43,12 +45,14 @@ int DebugProbesMode;
 
 void MainPS(
 	  in  float4 SvPosition		: SV_Position
-	, out float4 OutSceneColor	: SV_Target0
+	, out float4 OutColor		: SV_Target0
+#if PERMUTATION_RENDERVELOCITY == 0
 	, out float4 OutGBufferA	: SV_Target1
 	, out float4 OutGBufferB	: SV_Target2
 	, out float4 OutGBufferC	: SV_Target3
 	, out float4 OutGBufferD	: SV_Target4
 	, out float4 OutGBufferE	: SV_Target5
+#endif
 	, out float  OutDepth		: SV_DEPTH
 )
 {
@@ -63,6 +67,7 @@ void MainPS(
 	const float3 RayOrig = ResolvedView.WorldCameraOrigin;
 
 	int ProbeMode = false;
+	bool bProbeAttachedToCamera = false;
 	float RayT = -1.0f;
 	float3 WorlPosition = 0.0f;
 	float3 Normal = 0.0f;
@@ -76,16 +81,16 @@ void MainPS(
 			float  SphereRadius = 0.2f * METER_TO_CENTIMETER;
 
 			DebugSphereIntersection(
-				RayOrig, RayDir, SphereCenter, SphereRadius, PROBE_MODE_SPECULAR,
-				ProbeMode, RayT, WorlPosition, Normal);
+				RayOrig, RayDir, SphereCenter, SphereRadius, PROBE_MODE_SPECULAR, true,
+				ProbeMode, RayT, WorlPosition, Normal, bProbeAttachedToCamera);
 		}
 		{
 			float3 SphereCenter = ResolvedView.WorldCameraOrigin + (ResolvedView.ViewForward * 1.0f + ResolvedView.ViewRight * -0.5f + ResolvedView.ViewUp * -0.2f) * METER_TO_CENTIMETER;
 			float  SphereRadius = 0.2f * METER_TO_CENTIMETER;
 
 			DebugSphereIntersection(
-				RayOrig, RayDir, SphereCenter, SphereRadius, PROBE_MODE_DIFFUSE,
-				ProbeMode, RayT, WorlPosition, Normal);
+				RayOrig, RayDir, SphereCenter, SphereRadius, PROBE_MODE_DIFFUSE, true,
+				ProbeMode, RayT, WorlPosition, Normal, bProbeAttachedToCamera);
 		}
 	}
 
@@ -110,8 +115,8 @@ void MainPS(
 					float3 SphereCenter = SnappedOrigin + float3(x, y, z);
 
 					DebugSphereIntersection(
-						RayOrig, RayDir, SphereCenter, SphereRadius, PROBE_MODE_BOTH,
-						ProbeMode, RayT, WorlPosition, Normal);
+						RayOrig, RayDir, SphereCenter, SphereRadius, PROBE_MODE_BOTH, false,
+						ProbeMode, RayT, WorlPosition, Normal, bProbeAttachedToCamera);
 				}
 			}
 		}
@@ -130,7 +135,20 @@ void MainPS(
 	}
 
 
+#if PERMUTATION_RENDERVELOCITY
 
+	float3 Velocity = 0.0f;
+
+	if (!bProbeAttachedToCamera)
+	{
+		float4 ScreenPos	= mul(float4(WorlPosition.xyz + ResolvedView.PreViewTranslation, 1),	ResolvedView.TranslatedWorldToClip);
+		float4 PrevScreenPos= mul(float4(WorlPosition.xyz + ResolvedView.PrevPreViewTranslation, 1),ResolvedView.PrevTranslatedWorldToClip);
+		Velocity = Calculate3DVelocity(ScreenPos, PrevScreenPos);
+	}
+
+	OutColor = EncodeVelocityToTexture(Velocity);
+
+#else
 
 	const float PerObjectGBufferData = 0.0f;
 	const float SelectiveOutputMask = 0.0f;
@@ -143,7 +161,7 @@ void MainPS(
 
 	// We write all outputs, the rasterizer will discard what is not bound. That is fine since we do not care that much about performance for this debug view.
 
-	OutSceneColor = 0.0f;
+	OutColor = 0.0f;
 
 	OutGBufferA = float4(
 		EncodeNormal(Normal), 
@@ -161,6 +179,8 @@ void MainPS(
 
 	// Not writting to the velocity buffer
 
+#endif
+
 	OutDepth = DeviceZ;
 
 }
\ No newline at end of file
diff --git a/Engine/Source/Runtime/Renderer/Private/DebugProbeRendering.cpp b/Engine/Source/Runtime/Renderer/Private/DebugProbeRendering.cpp
--- a/Engine/Source/Runtime/Renderer/Private/DebugProbeRendering.cpp
+++ b/Engine/Source/Runtime/Renderer/Private/DebugProbeRendering.cpp
@@ -26,7 +26,8 @@ class FStampDeferredDebugProbePS : public FGlobalShader
 	DECLARE_GLOBAL_SHADER(FStampDeferredDebugProbePS);
 	SHADER_USE_PARAMETER_STRUCT(FStampDeferredDebugProbePS, FGlobalShader);
 
-	using FPermutationDomain = TShaderPermutationDomain<>;
+	class FRenderVelocity : SHADER_PERMUTATION_BOOL("PERMUTATION_RENDERVELOCITY");
+	using FPermutationDomain = TShaderPermutationDomain<FRenderVelocity>;
 
 	BEGIN_SHADER_PARAMETER_STRUCT(FParameters, )
 		SHADER_PARAMETER_STRUCT_REF(FViewUniformShaderParameters, ViewUniformBuffer)
@@ -58,12 +59,14 @@ template<bool bEnableDepthWrite>
 static void CommonStampDeferredDebugProbeDrawCall(
 	FRDGBuilder& GraphBuilder,
 	const FViewInfo& View, 
-	FStampDeferredDebugProbePS::FParameters* PassParameters)
+	FStampDeferredDebugProbePS::FParameters* PassParameters,
+	bool bRenderVelocity = false)
 {
 	PassParameters->ViewUniformBuffer = View.ViewUniformBuffer;
 	PassParameters->DebugProbesMode = View.Family->EngineShowFlags.VisualizeIndirectLighting ? 3 : FMath::Clamp(CVarDebugIndirectLightingProbes.GetValueOnRenderThread(), 0, 3);
 		
 	FStampDeferredDebugProbePS::FPermutationDomain PermutationVector;
+	PermutationVector.Set<FStampDeferredDebugProbePS::FRenderVelocity>(bRenderVelocity);
 	TShaderMapRef<FStampDeferredDebugProbePS> PixelShader(View.ShaderMap, PermutationVector);
 
 	FPixelShaderUtils::AddFullscreenPass<FStampDeferredDebugProbePS>(
@@ -121,3 +124,27 @@ void StampDeferredDebugProbeMaterialPS(
 	}
 }
 
+void StampDeferredDebugProbeVelocityPS(
+	FRDGBuilder& GraphBuilder,
+	TArrayView<const FViewInfo> Views,
+	const FRenderTargetBindingSlots& BasePassRenderTargets)
+{
+	RDG_EVENT_SCOPE(GraphBuilder, "StampDeferredDebugProbeVelocity");
+	RDG_GPU_STAT_SCOPE(GraphBuilder, StampDeferredDebugProbe);
+
+	const bool bVisualizeIndirectLighting = CVarDebugIndirectLightingProbes.GetValueOnRenderThread() > 0;
+	for (const FViewInfo& View : Views)
+	{
+		if (!(bVisualizeIndirectLighting || View.Family->EngineShowFlags.VisualizeIndirectLighting))
+		{
+			continue;
+		}
+
+		FStampDeferredDebugProbePS::FParameters* PassParameters = GraphBuilder.AllocParameters<FStampDeferredDebugProbePS::FParameters>();
+		PassParameters->RenderTargets = BasePassRenderTargets;
+
+		const bool bRenderVelocity = true;
+		CommonStampDeferredDebugProbeDrawCall<false>(GraphBuilder, View, PassParameters, bRenderVelocity);
+	}
+}
+
diff --git a/Engine/Source/Runtime/Renderer/Private/DebugProbeRendering.h b/Engine/Source/Runtime/Renderer/Private/DebugProbeRendering.h
--- a/Engine/Source/Runtime/Renderer/Private/DebugProbeRendering.h
+++ b/Engine/Source/Runtime/Renderer/Private/DebugProbeRendering.h
@@ -22,3 +22,9 @@ void StampDeferredDebugProbeMaterialPS(
 	TArrayView<const FViewInfo> Views,
 	const FRenderTargetBindingSlots& BasePassRenderTargets);
 
+
+void StampDeferredDebugProbeVelocityPS(
+	FRDGBuilder& GraphBuilder,
+	TArrayView<const FViewInfo> Views,
+	const FRenderTargetBindingSlots& BasePassRenderTargets);
+
diff --git a/Engine/Source/Runtime/Renderer/Private/VelocityRendering.cpp b/Engine/Source/Runtime/Renderer/Private/VelocityRendering.cpp
--- a/Engine/Source/Runtime/Renderer/Private/VelocityRendering.cpp
+++ b/Engine/Source/Runtime/Renderer/Private/VelocityRendering.cpp
@@ -21,6 +21,7 @@
 #endif
 #include "VisualizeTexture.h"
 #include "MeshPassProcessor.inl"
+#include "DebugProbeRendering.h"
 
 // Changing this causes a full shader recompile
 static TAutoConsoleVariable<int32> CVarBasePassOutputsVelocity(
@@ -232,6 +233,9 @@ void FSceneRenderer::RenderVelocities(
 	RDG_WAIT_FOR_TASKS_CONDITIONAL(GraphBuilder, IsVelocityWaitForTasksEnabled(ShaderPlatform));
 
 	const EMeshPass::Type MeshPass = GetMeshPassFromVelocityPass(VelocityPass);
+	FExclusiveDepthStencil ExclusiveDepthStencil = (VelocityPass == EVelocityPass::Opaque && !(Scene->EarlyZPassMode == DDM_AllOpaqueNoVelocity))
+														? FExclusiveDepthStencil::DepthRead_StencilWrite
+														: FExclusiveDepthStencil::DepthWrite_StencilWrite;
 
 	for (int32 ViewIndex = 0; ViewIndex < Views.Num(); ViewIndex++)
 	{
@@ -275,9 +279,7 @@ void FSceneRenderer::RenderVelocities(
 				SceneTextures.Depth.Resolve,
 				ERenderTargetLoadAction::ELoad,
 				ERenderTargetLoadAction::ELoad,
-				(VelocityPass == EVelocityPass::Opaque && !(Scene->EarlyZPassMode == DDM_AllOpaqueNoVelocity))
-					? FExclusiveDepthStencil::DepthRead_StencilWrite
-					: FExclusiveDepthStencil::DepthWrite_StencilWrite);
+				ExclusiveDepthStencil);
 
 			PassParameters->RenderTargets[0] = FRenderTargetBinding(SceneTextures.Velocity, VelocityLoadAction);
 
@@ -308,6 +310,24 @@ void FSceneRenderer::RenderVelocities(
 			}
 		}
 	}
+
+#if !(UE_BUILD_SHIPPING)
+	const bool bForwardShadingEnabled = IsForwardShadingEnabled(ShaderPlatform);
+	if (!bForwardShadingEnabled)
+	{
+		FRenderTargetBindingSlots VelocityRenderTargets;
+		VelocityRenderTargets[0] = FRenderTargetBinding(SceneTextures.Velocity, VelocityLoadAction);
+		VelocityRenderTargets.DepthStencil = FDepthStencilBinding(
+			SceneTextures.Depth.Resolve,
+			ERenderTargetLoadAction::ELoad,
+			ERenderTargetLoadAction::ELoad,
+			ExclusiveDepthStencil);
+
+		StampDeferredDebugProbeVelocityPS(GraphBuilder, Views, VelocityRenderTargets);
+	}
+#endif
+
+
 }
 
 EPixelFormat FVelocityRendering::GetFormat(EShaderPlatform ShaderPlatform)
