Global illumination debug probe now uses Amanatides 3D-DDA to march virtual voxels around the camera and only evaluate intersection for sphere inside each voxel. Instead of evaluating allthe spheres. Reduces the total gbuffer patch code cost by a factor of 10.

diff --git a/Engine/Shaders/Private/DebugProbes.usf b/Engine/Shaders/Private/DebugProbes.usf
--- a/Engine/Shaders/Private/DebugProbes.usf
+++ b/Engine/Shaders/Private/DebugProbes.usf
@@ -95,11 +95,74 @@ void MainPS(
 	}
 
 
-	// Render the world probes around the camera (this could be accelerated by tracing voxel and only intersecting with spheres of traversed voxel)
+	// Render the world probes around the camera
 	if (DebugProbesMode == 2 || DebugProbesMode == 3)
 	{
 		const float SphereRadius = 0.5f * METER_TO_CENTIMETER;
 
+#if 1
+		// Using Amanatides 3D-DDA to march virtual voxels around the camera and only evaluate intersection for sphere inside each voxel.
+		const float3 VoxelSizeCm = float3(5.0f, 5.0f, 5.0f) * METER_TO_CENTIMETER;
+		const float3 SnappedOrigin = (floor(ResolvedView.WorldCameraOrigin / VoxelSizeCm) + 0.5) * VoxelSizeCm;
+			+ float3(0.0f, 0.0f, 1.4 * SphereRadius);	// Small offset to have probes above the ground level 0
+
+		const float3 BoxMin = SnappedOrigin - VoxelSizeCm * 4.0f;
+		const float3 BoxMax = SnappedOrigin + VoxelSizeCm * 4.0f;
+
+		float3 StartP = (RayOrig - BoxMin) / VoxelSizeCm;
+		float3 P = StartP;
+
+		// Amanatides 3D-DDA data preparation
+		float3 stepSign = sign(RayDir);
+		float3 tDelta = abs(1.0f / RayDir);
+		float3 tMax = float3(0.0f, 0.0f, 0.0f);
+		float3 refPoint = floor(P);
+		tMax.x = stepSign.x > 0.0f ? refPoint.x + 1.0f - P.x : P.x - refPoint.x; // floor is more consistent than ceil
+		tMax.y = stepSign.y > 0.0f ? refPoint.y + 1.0f - P.y : P.y - refPoint.y;
+		tMax.z = stepSign.z > 0.0f ? refPoint.z + 1.0f - P.z : P.z - refPoint.z;
+		tMax.x *= tDelta.x;
+		tMax.y *= tDelta.y;
+		tMax.z *= tDelta.z;
+
+		LOOP
+		while (all(dot(StartP - P, StartP - P) < 5.0f * 5.0f))	// stop after a distance of 5 voxels is reached 
+		{
+			// Amanatides 3D-DDA 
+			if (tMax.x < tMax.y)
+			{
+				if (tMax.x < tMax.z)
+				{
+					P.x += stepSign.x;
+					tMax.x += tDelta.x;
+				}
+				else
+				{
+					P.z += stepSign.z;
+					tMax.z += tDelta.z;
+				}
+			}
+			else
+			{
+				if (tMax.y < tMax.z)
+				{
+					P.y += stepSign.y;
+					tMax.y += tDelta.y;
+				}
+				else
+				{
+					P.z += stepSign.z;
+					tMax.z += tDelta.z;
+				}
+			}
+			const float3 SphereCenter = BoxMin + (floor(P) + 0.5) * VoxelSizeCm;
+
+			DebugSphereIntersection(
+				RayOrig, RayDir, SphereCenter, SphereRadius, PROBE_MODE_BOTH, false,
+				ProbeMode, RayT, WorlPosition, Normal, bProbeAttachedToCamera);
+		}
+
+#else
+		// Simple but slow volume around the camera
 		const float3 Steps = float3(5.0f, 5.0f, 2.5f);
 		const float3 SnappedOrigin = floor(ResolvedView.WorldCameraOrigin / (Steps * METER_TO_CENTIMETER) + 0.5) * Steps * METER_TO_CENTIMETER
 								   + float3(0.0f, 0.0f, 1.4 * SphereRadius);	// Small offset to have probes above the ground level 0
@@ -120,6 +183,7 @@ void MainPS(
 				}
 			}
 		}
+#endif
 	}
 
 
