Fixed gi debug viz to not delete probe close to the camera and apply the small offset for nice probes on ground level 0 by default.

diff --git a/Engine/Shaders/Private/DebugProbes.usf b/Engine/Shaders/Private/DebugProbes.usf
--- a/Engine/Shaders/Private/DebugProbes.usf
+++ b/Engine/Shaders/Private/DebugProbes.usf
@@ -103,8 +103,8 @@ void MainPS(
 #if 1
 		// Using Amanatides 3D-DDA to march virtual voxels around the camera and only evaluate intersection for sphere inside each voxel.
 		const float3 VoxelSizeCm = float3(5.0f, 5.0f, 5.0f) * METER_TO_CENTIMETER;
-		const float3 SnappedOrigin = (floor(ResolvedView.WorldCameraOrigin / VoxelSizeCm) + 0.5) * VoxelSizeCm;
-			+ float3(0.0f, 0.0f, 1.4 * SphereRadius);	// Small offset to have probes above the ground level 0
+		const float3 SnappedOrigin = (floor(ResolvedView.WorldCameraOrigin / VoxelSizeCm) + 0.5) * VoxelSizeCm
+								   + float3(0.0f, 0.0f, 1.4 * SphereRadius);	// Small offset to have probes above the ground level 0
 
 		const float3 BoxMin = SnappedOrigin - VoxelSizeCm * 4.0f;
 		const float3 BoxMax = SnappedOrigin + VoxelSizeCm * 4.0f;
@@ -125,8 +125,14 @@ void MainPS(
 		tMax.z *= tDelta.z;
 
 		LOOP
-		while (all(dot(StartP - P, StartP - P) < 5.0f * 5.0f))	// stop after a distance of 5 voxels is reached 
+		while (all(dot(StartP - P, StartP - P) <= 5.1f * 5.1f))	// stop after a distance of 5 voxels is reached 
 		{
+			const float3 SphereCenter = BoxMin + (floor(P) + 0.5) * VoxelSizeCm;
+
+			DebugSphereIntersection(
+				RayOrig, RayDir, SphereCenter, SphereRadius, PROBE_MODE_BOTH, false,
+				ProbeMode, RayT, WorlPosition, Normal, bProbeAttachedToCamera);
+
 			// Amanatides 3D-DDA 
 			if (tMax.x < tMax.y)
 			{
@@ -154,11 +160,6 @@ void MainPS(
 					tMax.z += tDelta.z;
 				}
 			}
-			const float3 SphereCenter = BoxMin + (floor(P) + 0.5) * VoxelSizeCm;
-
-			DebugSphereIntersection(
-				RayOrig, RayDir, SphereCenter, SphereRadius, PROBE_MODE_BOTH, false,
-				ProbeMode, RayT, WorlPosition, Normal, bProbeAttachedToCamera);
 		}
 
 #else
