Clamp high frequency history reprojection to avoid over-sharpening from catrom

diff --git a/Engine/Shaders/Private/TemporalAA/TAACommon.ush b/Engine/Shaders/Private/TemporalAA/TAACommon.ush
--- a/Engine/Shaders/Private/TemporalAA/TAACommon.ush
+++ b/Engine/Shaders/Private/TemporalAA/TAACommon.ush
@@ -732,29 +732,40 @@ bool IsOffScreen(uint bCameraCut, float2 ScreenPos, taa_half ParallaxRejectionMa
 	return (bIsCameraCut || bIsOutOfBounds || bIsParallaxRejected);
 }
 
-taa_half2 MeasureRejectionFactor(taa_half3x2 PrevYCoCg, taa_half3x2 ClampedPrevYCoCg, taa_half3x2 InputCenterYCoCg, taa_half3x2 InputMinYCoCg, taa_half3x2 InputMaxYCoCg)
+taa_half MeasureBackbufferLDRQuantizationError()
+{
+	// Assume the backbuffer is 10bit per channels
+	return taa_half(0.5 / 1024.0);
+}
+
+taa_half2 MeasureBackbufferLDRQuantizationErrorFromHDRLuma(taa_half2 Luma)
+{
+	taa_half2 Tonemap = Luma * HdrWeightY(Luma) + MeasureBackbufferLDRQuantizationError().xx;
+	taa_half2 PixelColorQuantizationError = abs(Tonemap * HdrWeightInvY(Tonemap) - Luma);
+
+	return PixelColorQuantizationError;
+}
+
+taa_half MeasureBackbufferLDRQuantizationErrorFromHDRLuma(taa_half Luma)
+{
+	return dpv_lo(MeasureBackbufferLDRQuantizationErrorFromHDRLuma(dpv_interleave_mono_registers(Luma)));
+}
+
+taa_half2 MeasureRejectionFactor(
+	taa_half3x2 PrevYCoCg,
+	taa_half3x2 ClampedPrevYCoCg,
+	taa_half3x2 InputCenterYCoCg,
+	taa_half3x2 InputMinYCoCg, taa_half3x2 InputMaxYCoCg,
+	taa_half2 BackbufferQuantizationError)
 {
 	taa_half3x2 BoxSize = InputMaxYCoCg - InputMinYCoCg;
 
 	taa_half3x2 ClampedEnergy = abs(ClampedPrevYCoCg - PrevYCoCg);
 	taa_half3x2 Delta = abs(InputCenterYCoCg - PrevYCoCg);
 
-	taa_half3x2 ClampError;
-	#if 0
-	{
-		ClampError = BoxSize * taa_half(0.25);
-	}
-	#else
-	{
-		taa_half2 Luma = max(ClampedPrevYCoCg[0], PrevYCoCg[0]);
-
-		// Measure the error due to the quantization of the color in channels.
-		taa_half2 Tonemap = Luma * HdrWeightY(Luma) + taa_half(0.5 / 1024.0);
-		taa_half2 PixelColorQuantizationError = abs(Tonemap * HdrWeightInvY(Tonemap) - Luma);
-		
-		ClampError = max(BoxSize * taa_half(0.25), taa_half3x2(PixelColorQuantizationError, PixelColorQuantizationError, PixelColorQuantizationError));
-	}
-	#endif
+	taa_half3x2 ClampError = max(
+		BoxSize * taa_half(0.25),
+		taa_half3x2(BackbufferQuantizationError, BackbufferQuantizationError, BackbufferQuantizationError));
 
 	taa_half3x2 Factor = taa_half(1.0) - saturate(max(ClampedEnergy - ClampError, taa_half3x2(0.0, 0.0, 0.0, 0.0, 0.0, 0.0)) / max(Delta, BoxSize));
 
@@ -763,14 +774,20 @@ taa_half2 MeasureRejectionFactor(taa_half3x2 PrevYCoCg, taa_half3x2 ClampedPrevY
 	return Rejection;
 }
 
-taa_half MeasureRejectionFactor(taa_half3 PrevYCoCg, taa_half3 ClampedPrevYCoCg, taa_half3 InputCenterYCoCg, taa_half3 InputMinYCoCg, taa_half3 InputMaxYCoCg)
+taa_half MeasureRejectionFactor(
+	taa_half3 PrevYCoCg,
+	taa_half3 ClampedPrevYCoCg,
+	taa_half3 InputCenterYCoCg,
+	taa_half3 InputMinYCoCg, taa_half3 InputMaxYCoCg,
+	taa_half BackbufferQuantizationError)
 {
 	return dpv_lo(MeasureRejectionFactor(
 		dpv_interleave_mono_registers(PrevYCoCg), 
 		dpv_interleave_mono_registers(ClampedPrevYCoCg), 
 		dpv_interleave_mono_registers(InputCenterYCoCg), 
 		dpv_interleave_mono_registers(InputMinYCoCg), 
-		dpv_interleave_mono_registers(InputMaxYCoCg)));
+		dpv_interleave_mono_registers(InputMaxYCoCg), 
+		dpv_interleave_mono_registers(BackbufferQuantizationError)));
 }
 
 struct FSubpixelNeighborhood
diff --git a/Engine/Shaders/Private/TemporalAA/TAACompareHistory.usf b/Engine/Shaders/Private/TemporalAA/TAACompareHistory.usf
--- a/Engine/Shaders/Private/TemporalAA/TAACompareHistory.usf
+++ b/Engine/Shaders/Private/TemporalAA/TAACompareHistory.usf
@@ -90,7 +90,8 @@ void MainCS(
 
 		Rejection = MeasureRejectionFactor(
 			PrevColor, ClampedPrevColor,
-			InputCenterColor, InputMinColor, InputMaxColor);
+			InputCenterColor, InputMinColor, InputMaxColor,
+			MeasureBackbufferLDRQuantizationErrorFromHDRLuma(max(PrevColor[0], ClampedPrevColor[0])));
 		
 		Debug.rgb = abs(YCoCgToRGB(PrevColor - ClampedPrevColor)) / max3(YCoCgToRGB(InputMaxColor), YCoCgToRGB(InputCenterColor), 0.01);
 	}
diff --git a/Engine/Shaders/Private/TemporalAA/TAACompareTranslucency.usf b/Engine/Shaders/Private/TemporalAA/TAACompareTranslucency.usf
--- a/Engine/Shaders/Private/TemporalAA/TAACompareTranslucency.usf
+++ b/Engine/Shaders/Private/TemporalAA/TAACompareTranslucency.usf
@@ -195,11 +195,13 @@ void MainCS(
 		// TODO: reject with alpha too? Might not need since translucency is premultiplied RGBA 
 		taa_half ForwardRejection = MeasureRejectionFactor(
 			PrevTranslucency.xyz, ClampedPrevTranslucency.xyz,
-			Translucency.xyz, TranslucencyMin.xyz, TranslucencyMax.xyz);
+			Translucency.xyz, TranslucencyMin.xyz, TranslucencyMax.xyz,
+			MeasureBackbufferLDRQuantizationErrorFromHDRLuma(max(PrevTranslucency[0], ClampedPrevTranslucency[0])));
 		
 		taa_half BackwardRejection = MeasureRejectionFactor(
 			Translucency.xyz, ClampedTranslucency.xyz,
-			PrevTranslucency.xyz, PrevTranslucencyMin.xyz, PrevTranslucencyMax.xyz);
+			PrevTranslucency.xyz, PrevTranslucencyMin.xyz, PrevTranslucencyMax.xyz,
+			MeasureBackbufferLDRQuantizationErrorFromHDRLuma(max(Translucency[0], ClampedTranslucency[0])));
 		
 		Rejection = min(ForwardRejection, BackwardRejection);
 	}
diff --git a/Engine/Shaders/Private/TemporalAA/TAAUpdateHistory.usf b/Engine/Shaders/Private/TemporalAA/TAAUpdateHistory.usf
--- a/Engine/Shaders/Private/TemporalAA/TAAUpdateHistory.usf
+++ b/Engine/Shaders/Private/TemporalAA/TAAUpdateHistory.usf
@@ -550,15 +550,19 @@ void MainCS(
 	// Rejects high frequency details if their low frequencies drift from the low frequency history
 	#if CONFIG_LOW_FREQUENCY_DRIFT_REJECTION
 	{
-		taa_half3x2 PrevHighFrequencyYCoCg = TransformColorForClampingBox(dpv_scale(PrevHistoryMoment1, SafeRcp(PrevHistoryValidity)));
+		taa_half3x2 NormalizedPrevHistoryMoment1 = dpv_scale(PrevHistoryMoment1, SafeRcp(PrevHistoryValidity));
+
+		taa_half3x2 PrevHighFrequencyYCoCg = TransformColorForClampingBox(NormalizedPrevHistoryMoment1);
 		taa_half3x2 PrevYCoCg = TransformColorForClampingBox(PrevFallbackColor);
 		taa_half3x2 ClampedPrevYCoCg = TransformColorForClampingBox(clamp(PrevFallbackColor, PrevHistoryMommentMin, PrevHistoryMommentMax));
 
 		taa_half2 HighFrequencyRejection = MeasureRejectionFactor(
 			PrevYCoCg, ClampedPrevYCoCg,
-			PrevHighFrequencyYCoCg, InputMinColor, InputMaxColor);
+			PrevHighFrequencyYCoCg, InputMinColor, InputMaxColor,
+			MeasureBackbufferLDRQuantizationErrorFromHDRLuma(max(PrevYCoCg[0], ClampedPrevYCoCg[0])));
 		
-		PrevHistoryMoment1 = dpv_scale(PrevHistoryMoment1, HighFrequencyRejection);
+		//PrevHistoryMoment1 = dpv_scale(PrevHistoryMoment1, HighFrequencyRejection);
+		PrevHistoryMoment1 = dpv_scale(clamp(NormalizedPrevHistoryMoment1, PrevHistoryMommentMin, PrevHistoryMommentMax), HighFrequencyRejection * PrevHistoryValidity);
 		PrevHistoryValidity *= HighFrequencyRejection;
 	}
 	#endif
@@ -584,7 +588,8 @@ void MainCS(
 		
 		FallbackRejection = MeasureRejectionFactor(
 			PrevYCoCg, ClampedPrevYCoCg,
-			InputCenterYCoCg, InputMinColor, InputMaxColor);
+			InputCenterYCoCg, InputMinColor, InputMaxColor,
+			MeasureBackbufferLDRQuantizationErrorFromHDRLuma(max(PrevYCoCg[0], ClampedPrevYCoCg[0])));
 
 		#if !CONFIG_CLAMP
 		{
