Fix some shader issues in LightmapCommon.ush

diff --git a/Engine/Shaders/Private/LightmapCommon.ush b/Engine/Shaders/Private/LightmapCommon.ush
--- a/Engine/Shaders/Private/LightmapCommon.ush
+++ b/Engine/Shaders/Private/LightmapCommon.ush
@@ -216,14 +216,14 @@ half4 GetPrecomputedShadowMasks(VTPageTableResult LightmapVTPageTableResult, FVe
 		float DirectionalLightShadowing = 1.0f;
 
 	#if CACHED_POINT_INDIRECT_LIGHTING || CACHED_VOLUME_INDIRECT_LIGHTING
-		if (GetPrimitiveData(PrimitiveId).UseSingleSampleShadowFromStationaryLights > 0 && View.IndirectLightingCacheShowFlag > 0.0f)
+		if ((GetPrimitiveData(PrimitiveId).Flags & PRIMITIVE_SCENE_DATA_FLAG_USE_SINGLE_SAMPLE_SHADOW_SL) != 0 && View.IndirectLightingCacheShowFlag > 0.0f)
 		{
 			DirectionalLightShadowing = IndirectLightingCache.DirectionalLightShadowing;
 		}
 	#endif
 
 		BRANCH
-		if (GetPrimitiveData(PrimitiveId).UseVolumetricLightmapShadowFromStationaryLights > 0)
+		if ((GetPrimitiveData(PrimitiveId).Flags & PRIMITIVE_SCENE_DATA_FLAG_USE_VOLUMETRIC_LM_SHADOW_SL) != 0)
 		{
 			#if !PRECOMPUTED_IRRADIANCE_VOLUME_LIGHTING
 				// Compute brick UVs if we haven't already
@@ -273,7 +273,7 @@ half GetPrimaryPrecomputedShadowMask(VTPageTableResult LightmapVTPageTableResult
 	#else
 		#if CACHED_POINT_INDIRECT_LIGHTING || CACHED_VOLUME_INDIRECT_LIGHTING
 			// output per-primitive directional light shadowing if requested
-			if (GetPrimitiveData(MaterialParameters.PrimitiveId).UseSingleSampleShadowFromStationaryLights > 0 && ResolvedView.IndirectLightingCacheShowFlag > 0.0f)
+			if ((GetPrimitiveData(MaterialParameters.PrimitiveId).Flags & PRIMITIVE_SCENE_DATA_FLAG_USE_SINGLE_SAMPLE_SHADOW_SL) != 0 && ResolvedView.IndirectLightingCacheShowFlag > 0.0f)
 			{
 				return IndirectLightingCache.DirectionalLightShadowing;
 			}
