Fix for some Android shader compilation errors

diff --git a/Engine/Shaders/Private/GeometryCollectionVertexFactory.ush b/Engine/Shaders/Private/GeometryCollectionVertexFactory.ush
--- a/Engine/Shaders/Private/GeometryCollectionVertexFactory.ush
+++ b/Engine/Shaders/Private/GeometryCollectionVertexFactory.ush
@@ -394,7 +394,7 @@ FVertexFactoryIntermediates GetVertexFactoryIntermediates(FVertexFactoryInput In
 	float DeterminantSign = Intermediates.SceneData.InstanceData.DeterminantSign;
 #else 
 	FTransformData TransformData = LoadTransformData(0U, 0U, GCVF_GET_INPUT_VERTEX_ID(Input));
-	float DeterminantSign = GetPrimitive_DeterminantSign(Primitive);
+	float DeterminantSign = GetPrimitive_DeterminantSign_FromFlags(Primitive.Flags);
 #endif
 
 	Intermediates.LocalToWorld = TransformData.LocalToWorld;
diff --git a/Engine/Shaders/Private/GpuSkinVertexFactory.ush b/Engine/Shaders/Private/GpuSkinVertexFactory.ush
--- a/Engine/Shaders/Private/GpuSkinVertexFactory.ush
+++ b/Engine/Shaders/Private/GpuSkinVertexFactory.ush
@@ -755,7 +755,7 @@ FVertexFactoryIntermediates GetVertexFactoryIntermediates(FVertexFactoryInput In
 	Intermediates.PrevLocalToWorld = Intermediates.SceneData.InstanceData.PrevLocalToWorld;
 #else
 	Intermediates.InvNonUniformScale = Primitive.InvNonUniformScale;
-	Intermediates.DeterminantSign = GetPrimitive_DeterminantSign(Primitive);
+	Intermediates.DeterminantSign = GetPrimitive_DeterminantSign_FromFlags(Primitive.Flags);
 	Intermediates.LocalToWorld = Primitive.LocalToWorld;
 	Intermediates.WorldToLocal = Primitive.WorldToLocal;
 	Intermediates.PrevLocalToWorld = Primitive.PreviousLocalToWorld;
