Fix insufficient number of bits for lighting channel mask

Fix return type of GetPrimitive_LightingChannelMask

Apply recent refactor of primitive data to PathTracing material

This fixes black frames on two of the entries of PT_CastShadow EngineTest

diff --git a/Engine/Shaders/Private/PathTracing/PathTracingMaterialHitShader.usf b/Engine/Shaders/Private/PathTracing/PathTracingMaterialHitShader.usf
--- a/Engine/Shaders/Private/PathTracing/PathTracingMaterialHitShader.usf
+++ b/Engine/Shaders/Private/PathTracing/PathTracingMaterialHitShader.usf
@@ -89,7 +89,7 @@ RAY_TRACING_ENTRY_CLOSEST_HIT(PathTracingMaterialCHS,
 #endif
 
 	Payload.BlendingMode = MaterialBlendingMode;
-	Payload.PrimitiveLightingChannelMask = GetPrimitiveData(MaterialParameters.PrimitiveId).LightingChannelMask;
+	Payload.PrimitiveLightingChannelMask = GetPrimitive_LightingChannelMask(MaterialParameters.PrimitiveId);
 
 	Payload.HitT = RayTCurrent();
 	if (HitKind() == HIT_KIND_TRIANGLE_FRONT_FACE)
diff --git a/Engine/Shaders/Private/SceneData.ush b/Engine/Shaders/Private/SceneData.ush
--- a/Engine/Shaders/Private/SceneData.ush
+++ b/Engine/Shaders/Private/SceneData.ush
@@ -326,13 +326,13 @@ uint GetPrimitive_LightingChannelMask_FromFlags(uint Flags)
 	return (Channel0 | (Channel1 << 1u) | (Channel2 << 2u));
 }
 
-float GetPrimitive_LightingChannelMask(uint PrimitiveId)
+uint GetPrimitive_LightingChannelMask(uint PrimitiveId)
 {
 	return GetPrimitive_LightingChannelMask_FromFlags(GetPrimitiveData(PrimitiveId).Flags);
 }
 
 #if VF_USE_PRIMITIVE_SCENE_DATA
-float GetPrimitive_LightingChannelMask(FPrimitiveSceneData Primitive)
+uint GetPrimitive_LightingChannelMask(FPrimitiveSceneData Primitive)
 {
 	return GetPrimitive_LightingChannelMask_FromFlags(Primitive.Flags);
 }
diff --git a/Engine/Source/Runtime/Engine/Public/PrimitiveUniformShaderParameters.h b/Engine/Source/Runtime/Engine/Public/PrimitiveUniformShaderParameters.h
--- a/Engine/Source/Runtime/Engine/Public/PrimitiveUniformShaderParameters.h
+++ b/Engine/Source/Runtime/Engine/Public/PrimitiveUniformShaderParameters.h
@@ -247,7 +247,7 @@ public:
 
 private:
 	FPrimitiveUniformShaderParameters Parameters;
-	uint32 LightingChannels : 1;
+	uint32 LightingChannels : 3;
 	uint32 bReceivesDecals : 1;
 	uint32 bUseSingleSampleShadowFromStationaryLights : 1;
 	uint32 bUseVolumetricLightmap : 1;
