Fix depth range precision and offset/clamping issues with VSM cached pages
- Instead, create a reference depth range with a guard band that we reuse frame to frame until we get too close to the edge, at which point we invalidate that clipmap level's cached pages (if any)
- Does not significantly affect invalidations, even when moving quickly near surfaces, since those cases would already get invalidated when those pixels moved to a different clipmap level
- Sets up for not re-copying cached VSM pages every frame

diff --git a/Engine/Shaders/Private/VirtualShadowMaps/Debug.usf b/Engine/Shaders/Private/VirtualShadowMaps/Debug.usf
--- a/Engine/Shaders/Private/VirtualShadowMaps/Debug.usf
+++ b/Engine/Shaders/Private/VirtualShadowMaps/Debug.usf
@@ -151,8 +151,7 @@ void DebugVisualizeVirtualSmCS(uint2 OutputPixel : SV_DispatchThreadID)
 				if (ShadowMapCacheData[VirtualShadowMapId].VirtualShadowMapId != INDEX_NONE)// && MipLevel < VSM_MAX_MIP_LEVELS - 1)
 				{
 					int2 PrevPageAddress = int2(PageAddress) - (ShadowMapCacheData[VirtualShadowMapId].SmPageOffset >> MipLevel);
-					bool bCanCache = MipLevel < VSM_CACHE_ALIGNMENT_LEVEL || all(ShadowMapCacheData[VirtualShadowMapId].SmPageOffset == 0);
-					if (all(PrevPageAddress >= 0 && PrevPageAddress < CalcLevelDimsPages(MipLevel)) && bCanCache)
+					if (all(PrevPageAddress >= 0 && PrevPageAddress < CalcLevelDimsPages(MipLevel)))
 					{
 						uint PrevOffset = CalcPageOffset(ShadowMapCacheData[VirtualShadowMapId].VirtualShadowMapId, MipLevel, PrevPageAddress);
 						PrevPageFlag = PrevPageFlags[PrevOffset];
diff --git a/Engine/Shaders/Private/VirtualShadowMaps/PageCacheCommon.ush b/Engine/Shaders/Private/VirtualShadowMaps/PageCacheCommon.ush
--- a/Engine/Shaders/Private/VirtualShadowMaps/PageCacheCommon.ush
+++ b/Engine/Shaders/Private/VirtualShadowMaps/PageCacheCommon.ush
@@ -18,15 +18,11 @@ struct FShadowMapCacheData
 	int2 SmPageOffset;
 	// ID of the corresponding virtual SM in the chached data 
 	int VirtualShadowMapId;
-	// Depth offset to add to SM texels when copying
-	float DepthOffset;
 };
 
 struct FCachedPageInfo
 {
 	uint2 PhysPageAddress;
-	float DepthOffset;
-	float Padding;
 };
 
 #if HAS_CACHE_DATA
diff --git a/Engine/Shaders/Private/VirtualShadowMaps/PageManagement.usf b/Engine/Shaders/Private/VirtualShadowMaps/PageManagement.usf
--- a/Engine/Shaders/Private/VirtualShadowMaps/PageManagement.usf
+++ b/Engine/Shaders/Private/VirtualShadowMaps/PageManagement.usf
@@ -18,13 +18,26 @@
 #include "ProjectionCommon.ush"
 
 #ifndef VSM_GENERATE_STATS
-	#define VSM_GENERATE_STATS 0
+#define VSM_GENERATE_STATS 0
 #endif // VSM_GENERATE_STATS
 
 #if VSM_GENERATE_STATS
 RWStructuredBuffer<uint> OutStatsBuffer;
 #endif // VSM_GENERATE_STATS
 
+void StatsBufferWrite(uint Index, uint Value)
+{
+#if VSM_GENERATE_STATS
+	OutStatsBuffer[Index] = Value;
+#endif
+}
+void StatsBufferInterlockedAdd(uint Index, uint Value)
+{
+#if VSM_GENERATE_STATS
+	InterlockedAdd(OutStatsBuffer[Index], Value);
+#endif
+}
+
 // Type of input data consume by the page allocation (i.e., data read from the source buffer: Gbuffer, HairStrands data, ...)
 #define INPUT_TYPE_GBUFFER 0
 #define INPUT_TYPE_NANITE 1
@@ -449,10 +462,7 @@ void CreatePageMapping( uint ShadowMapID, uint MipLevel, uint2 PageAddress, uint
 	{
 		uint pPageIndex = 0;
 		InterlockedAdd( AllocatedPagesOffset[0], 1, pPageIndex );
-
-#if VSM_GENERATE_STATS
-		InterlockedAdd(OutStatsBuffer[0], 1);
-#endif // VSM_GENERATE_STATS
+		StatsBufferInterlockedAdd( 0, 1 );
 
 		if (pPageIndex >= VirtualShadowMap.MaxPhysicalPages)
 		{
@@ -474,32 +484,24 @@ void CreatePageMapping( uint ShadowMapID, uint MipLevel, uint2 PageAddress, uint
 		FCachedPageInfo CachedPageInfo;
 
 		CachedPageInfo.PhysPageAddress = uint2(VSM_INVALID_PHYSICAL_PAGE_ADDRESS, VSM_INVALID_PHYSICAL_PAGE_ADDRESS);
-		CachedPageInfo.DepthOffset = 0.0f;
-		CachedPageInfo.Padding = 0.0f;
 
 #if HAS_CACHE_DATA
 		if (ShadowMapCacheData[ShadowMapID].VirtualShadowMapId != INDEX_NONE)
 		{
-			CachedPageInfo.DepthOffset = ShadowMapCacheData[ShadowMapID].DepthOffset;
-
 			int2 PrevPageAddress = int2(PageAddress) - (ShadowMapCacheData[ShadowMapID].SmPageOffset >> MipLevel);
-			// A page cannot be cached if it is above the alignment level, or the light has not moved
-			// This means that for small movement a CSM can cache everything as well.
-			bool bCanCache = MipLevel < VSM_CACHE_ALIGNMENT_LEVEL || all(ShadowMapCacheData[ShadowMapID].SmPageOffset == 0);
-			if (all(uint2(PrevPageAddress) < CalcLevelDimsPages(MipLevel)) && bCanCache)
+			if (all(uint2(PrevPageAddress) < CalcLevelDimsPages(MipLevel)))
 			{
 				uint PrevOffset = CalcPageOffset(ShadowMapCacheData[ShadowMapID].VirtualShadowMapId, MipLevel, PrevPageAddress);
 				uint PrevPageFlag = PrevPageFlags[PrevOffset];
 				// True if the page had moving stuff drawn into it
 				bool PrevDynPageFlag = PrevDynamicCasterPageFlags[PrevOffset] != 0;
 
-#if VSM_GENERATE_STATS
 				// stats for those pages who die from moving stuff
 				if (PrevPageFlag && PrevDynPageFlag)
 				{
-					InterlockedAdd(OutStatsBuffer[2], 1);
+					StatsBufferInterlockedAdd(2, 1);
 				}
-#endif // VSM_GENERATE_STATS
+
 				// If it was either rendered or cached previous frame, and not drawn to by a moving thing we may re-use
 				if (PrevPageFlag && !PrevDynPageFlag)
 				{
@@ -509,53 +511,10 @@ void CreatePageMapping( uint ShadowMapID, uint MipLevel, uint2 PageAddress, uint
 					FPhysicalPageMetaData MetaData = PrevPhysicalPageMetaData[PhysPageAddressToIndex(PrevPhysAddress)];
 					if (MetaData.State == VSM_PHYSICAL_PAGE_STATE_RENDERED)
 					{
-// 'rand robin' invalidation
-#if 0
-						// If cached for too long invalidate
-						if (MetaData.Age > 3 + (PageOffset & 3))
-						{
-#if VSM_GENERATE_STATS
-							InterlockedAdd(OutStatsBuffer[4], 1);
-#endif // VSM_GENERATE_STATS
-						}
-						else
-#endif
-						{
-#if VSM_GENERATE_STATS
-							InterlockedAdd(OutStatsBuffer[1], 1);
-#endif // VSM_GENERATE_STATS
-
-							CachedPageInfo.PhysPageAddress = PrevPhysAddress;
-						}
-					}
-				}
-// Needs more work (must disable page AND ensure fallback page is enabled)!
-#if 0
-				else
-				{
-					// Find fallback
-					uint PrevIncCoarserPageFlag = PrevPageFlag;
-					uint PrevCoarseOffset = 0;
-					// Check coarser levels (up to 4 levels away, arbitrary, hardcoded)
-					for (uint CoarseMipLevel = MipLevel + 1; CoarseMipLevel < min(MipLevel + 4, VSM_MAX_MIP_LEVELS); ++CoarseMipLevel)
-					{
-						PrevCoarseOffset = CalcPageOffset(ShadowMapCacheData[ShadowMapID].VirtualShadowMapId, CoarseMipLevel, PrevPageAddress >> (CoarseMipLevel - MipLevel));
-						uint PrevCoarserPageFlag = PrevPageFlags[PrevCoarseOffset];
-						if (PrevCoarserPageFlag)
-						{
-							PrevIncCoarserPageFlag = 1;
-							break;
-						}
-					}
-					if (PrevIncCoarserPageFlag)
-					{
-#if VSM_GENERATE_STATS
-#endif // VSM_GENERATE_STATS
-						InterlockedAdd(OutStatsBuffer[2], 1);
-						CachedPageInfo.PhysPageAddress = ShadowDecodePageTable(PrevPageTable[PrevCoarseOffset]).PageIndex;
+						StatsBufferInterlockedAdd(1, 1);
+						CachedPageInfo.PhysPageAddress = PrevPhysAddress;
 					}
 				}
-#endif
 			}
 		}
 #endif // HAS_CACHE_DATA
@@ -587,12 +546,10 @@ void CreatePageMappings(uint2 ThreadId : SV_DispatchThreadID, uint2 GroupId : SV
 		return;
 	}
 
-#if VSM_GENERATE_STATS
 	if (all(ThreadId == 0))
 	{
-		OutStatsBuffer[3] = VirtualShadowMap.NumShadowMaps;
+		StatsBufferWrite(3, VirtualShadowMap.NumShadowMaps);
 	}
-#endif // VSM_GENERATE_STATS
 
 	uint Level0RowMask = ((1U << VSM_LOG2_LEVEL0_DIM_PAGES_XY) - 1U);
 
@@ -830,10 +787,10 @@ void ClearPhysicalPages( uint3 TileThreadID : SV_GroupThreadID, uint3 TileIndex
 	{
 		const uint2 TileOffset = (CachedPageInfo.PhysPageAddress << uint2(VSM_LOG2_PAGE_SIZE, VSM_LOG2_PAGE_SIZE)) + (uint2(LocalTileX, LocalTileY) << uint2(LOG2_TILE_SIZE, LOG2_TILE_SIZE));
 		const uint2 BasePos = TileOffset + (TileThreadID.xy << 1u);
-		ClearValue[0] = asuint(max(0.0f, asfloat(CachedPhysicalPagesTexture[BasePos + uint2(0, 0)]) + CachedPageInfo.DepthOffset));
-		ClearValue[1] = asuint(max(0.0f, asfloat(CachedPhysicalPagesTexture[BasePos + uint2(1, 0)]) + CachedPageInfo.DepthOffset));
-		ClearValue[2] = asuint(max(0.0f, asfloat(CachedPhysicalPagesTexture[BasePos + uint2(0, 1)]) + CachedPageInfo.DepthOffset));
-		ClearValue[3] = asuint(max(0.0f, asfloat(CachedPhysicalPagesTexture[BasePos + uint2(1, 1)]) + CachedPageInfo.DepthOffset));
+		ClearValue[0] = asuint(max(0.0f, asfloat(CachedPhysicalPagesTexture[BasePos + uint2(0, 0)])));
+		ClearValue[1] = asuint(max(0.0f, asfloat(CachedPhysicalPagesTexture[BasePos + uint2(1, 0)])));
+		ClearValue[2] = asuint(max(0.0f, asfloat(CachedPhysicalPagesTexture[BasePos + uint2(0, 1)])));
+		ClearValue[3] = asuint(max(0.0f, asfloat(CachedPhysicalPagesTexture[BasePos + uint2(1, 1)])));
 
 		MetaData = PrevPhysicalPageMetaData[PhysPageAddressToIndex(CachedPageInfo.PhysPageAddress)];
 		// This data has survived one more frame
@@ -860,7 +817,7 @@ void ClearPhysicalPages( uint3 TileThreadID : SV_GroupThreadID, uint3 TileIndex
 
 
 /**
- * Invoke to update the sate of all pages after copying & clearing has been performed.
+ * Invoke to update the state of all pages after copying & clearing has been performed.
  */
 [numthreads(64, 1, 1)]
 void UpdatePhysicalPageMetaDataClearState(uint PageIndex : SV_DispatchThreadID)
@@ -941,7 +898,7 @@ void InitializePhysicalPagesPS(uint PageIndex : TEXCOORD0, float4 SVPos : SV_POS
 	if (CachedPageInfo.PhysPageAddress.x != VSM_INVALID_PHYSICAL_PAGE_ADDRESS)
 	{
 		const uint2 CachedPos = CachedPageInfo.PhysPageAddress * VSM_PAGE_SIZE + PagePixelPos;
-		ClearValue = max(0.0f, PrevPhysicalPagePoolHw.Load(int3(CachedPos, 0)) + CachedPageInfo.DepthOffset);
+		ClearValue = max(0.0f, PrevPhysicalPagePoolHw.Load(int3(CachedPos, 0)));
 	}
 #endif // HAS_CACHE_DATA
 
@@ -1209,8 +1166,6 @@ void AllocatePagesUsingRects(uint GroupThreadId : SV_GroupThreadID, uint2 Dispat
 						FCachedPageInfo CachedPageInfo;
 
 						CachedPageInfo.PhysPageAddress = uint2(VSM_INVALID_PHYSICAL_PAGE_ADDRESS, VSM_INVALID_PHYSICAL_PAGE_ADDRESS);
-						CachedPageInfo.DepthOffset = 0.0f;
-						CachedPageInfo.Padding = 0.0f;
 						// Store data that is used later in the process (for each physical page).
 						OutCachedPageInfos[pPage.y * VirtualShadowMap.PhysicalPoolSizePages.x + pPage.x] = CachedPageInfo;
 
diff --git a/Engine/Source/Runtime/Renderer/Private/ShadowSetup.cpp b/Engine/Source/Runtime/Renderer/Private/ShadowSetup.cpp
--- a/Engine/Source/Runtime/Renderer/Private/ShadowSetup.cpp
+++ b/Engine/Source/Runtime/Renderer/Private/ShadowSetup.cpp
@@ -3678,7 +3678,7 @@ void FSceneRenderer::CreateWholeSceneProjectedShadow(
 							// Note: the caching for spot lights only handles static lights and does not need the border, this also means pages can cache at all levels.
 							if( VirtualSmCacheEntry != nullptr )
 							{
-								VirtualSmCacheEntry->Update(VirtualShadowMap->ID, ProjectedShadowInitializer);
+								VirtualSmCacheEntry->UpdateLocal(VirtualShadowMap->ID, ProjectedShadowInitializer);
 								VirtualShadowMap->VirtualShadowMapCacheEntry = VirtualSmCacheEntry;
 							}
 						}
diff --git a/Engine/Source/Runtime/Renderer/Private/VirtualShadowMaps/VirtualShadowMapArray.cpp b/Engine/Source/Runtime/Renderer/Private/VirtualShadowMaps/VirtualShadowMapArray.cpp
--- a/Engine/Source/Runtime/Renderer/Private/VirtualShadowMaps/VirtualShadowMapArray.cpp
+++ b/Engine/Source/Runtime/Renderer/Private/VirtualShadowMaps/VirtualShadowMapArray.cpp
@@ -24,8 +24,6 @@ struct FShadowMapCacheData
 	FIntPoint SmPageOffset;
 	// ID of the corresponding virtual SM in the chached data 
 	int32 VirtualShadowMapId;
-	// Depth offset to add to SM texels when copying
-	float DepthOffset;
 };
 
 
@@ -35,12 +33,9 @@ struct FPhysicalPageMetaData
 	uint32 Age;
 };
 
-
 struct FCachedPageInfo
 {
 	FIntPoint PhysPageAddress;
-	float DepthOffset;
-	float Padding;
 };
 
 int32 GEnableVirtualShadowMaps = 0;
@@ -225,13 +220,11 @@ static void SetCacheDataShaderParameters(FRDGBuilder& GraphBuilder, const TArray
 		{
 			ShadowMapCacheData[SmIndex].SmPageOffset = VirtualShadowMapCacheEntry->GetPageSpaceOffset();
 			ShadowMapCacheData[SmIndex].VirtualShadowMapId = VirtualShadowMapCacheEntry->PrevVirtualShadowMapId;
-			ShadowMapCacheData[SmIndex].DepthOffset = VirtualShadowMapCacheEntry->GetDepthOffset();
 		}
 		else
 		{
 			ShadowMapCacheData[SmIndex].SmPageOffset = FIntPoint(0, 0);
 			ShadowMapCacheData[SmIndex].VirtualShadowMapId = INDEX_NONE;
-			ShadowMapCacheData[SmIndex].DepthOffset = 0.0f;
 		}
 	}
 	CacheDataParameters.ShadowMapCacheData = GraphBuilder.CreateSRV(CreateStructuredBuffer(GraphBuilder, TEXT("Shadow.Virtual.ShadowMapCacheData"), ShadowMapCacheData));
@@ -340,7 +333,6 @@ void FVirtualShadowMapArray::SetShaderDefines(FShaderCompilerEnvironment& OutEnv
 	OutEnvironment.SetDefine(TEXT("VSM_VIRTUAL_MAX_RESOLUTION_XY"), FVirtualShadowMap::VirtualMaxResolutionXY);
 	OutEnvironment.SetDefine(TEXT("VSM_RASTER_WINDOW_PAGES"), FVirtualShadowMap::RasterWindowPages);
 	OutEnvironment.SetDefine(TEXT("VSM_PAGE_TABLE_SIZE"), FVirtualShadowMap::PageTableSize);
-	OutEnvironment.SetDefine(TEXT("VSM_CACHE_ALIGNMENT_LEVEL"), FVirtualShadowMapArrayCacheManager::AlignmentLevel);
 	OutEnvironment.SetDefine(TEXT("INDEX_NONE"), INDEX_NONE);
 }
 
diff --git a/Engine/Source/Runtime/Renderer/Private/VirtualShadowMaps/VirtualShadowMapCacheManager.cpp b/Engine/Source/Runtime/Renderer/Private/VirtualShadowMaps/VirtualShadowMapCacheManager.cpp
--- a/Engine/Source/Runtime/Renderer/Private/VirtualShadowMaps/VirtualShadowMapCacheManager.cpp
+++ b/Engine/Source/Runtime/Renderer/Private/VirtualShadowMaps/VirtualShadowMapCacheManager.cpp
@@ -19,7 +19,6 @@ static TAutoConsoleVariable<int32> CVarAccumulateStats(
 	ECVF_RenderThreadSafe
 );
 
-
 static TAutoConsoleVariable<int32> CVarCacheVirtualSMs(
 	TEXT("r.Shadow.Virtual.Cache"),
 	1,
@@ -32,43 +31,78 @@ void FVirtualShadowMapCacheEntry::UpdateClipmap(
 	int32 VirtualShadowMapId,
 	const FMatrix &WorldToLight,
 	FIntPoint PageSpaceLocation,
-	float GlobalDepth)
+	float LevelRadius,
+	float ViewCenterZ,
+	float ViewRadiusZ)
 {
-	// Swap previous frame data over.
-	PrevVirtualShadowMapId = CurrentVirtualShadowMapId;
-	PrevPageSpaceLocation = CurrentPageSpaceLocation;
-	PrevShadowMapGlobalDepth = CurrentShadowMapGlobalDepth;
+	bool bCacheValid = (CurrentVirtualShadowMapId != INDEX_NONE);
 	
-	if (WorldToLight != ClipmapCacheValidKey.WorldToLight)
+	if (bCacheValid && WorldToLight != Clipmap.WorldToLight)
 	{
-		PrevVirtualShadowMapId = INDEX_NONE;
-		ClipmapCacheValidKey.WorldToLight = WorldToLight;
+		bCacheValid = false;
+		//UE_LOG(LogRenderer, Display, TEXT("Invalidated clipmap level (VSM %d) due to light movement"), VirtualShadowMapId);
 	}
 
+#if 0
+	// Disable cache panning for directional lights (debug)
+	if (bCacheValid)
+	{
+		bCacheValid = bCacheValid && PageSpaceLocation.X == PrevPageSpaceLocation.X;
+		bCacheValid = bCacheValid && PageSpaceLocation.Y == PrevPageSpaceLocation.Y;
+		if (!bCacheValid)
+		{
+			//UE_LOG(LogRenderer, Display, TEXT("Invalidated clipmap level (VSM %d) with page space location %d,%d (Prev %d, %d)"),
+			//	VirtualShadowMapId, PageSpaceLocation.X, PageSpaceLocation.Y, PrevPageSpaceLocation.X, PrevPageSpaceLocation.Y);
+		}
+	}
+#endif
+
+	// Invalidate if the new Z radius strayed too close/outside the guardband of the cached shadow map
+	if (bCacheValid)
+	{
+		float DeltaZ = FMath::Abs(ViewCenterZ - Clipmap.ViewCenterZ);
+		if ((DeltaZ + LevelRadius) > 0.9f * Clipmap.ViewRadiusZ)
+		{
+			bCacheValid = false;
+			//UE_LOG(LogRenderer, Display, TEXT("Invalidated clipmap level (VSM %d) due to depth range movement"), VirtualShadowMapId);
+		}
+	}
+
+	if (bCacheValid)
+	{
+		PrevVirtualShadowMapId = CurrentVirtualShadowMapId;
+	}
+	else
+	{
+		// New cached level
+		PrevVirtualShadowMapId = INDEX_NONE;
+		Clipmap.WorldToLight = WorldToLight;
+		Clipmap.ViewCenterZ = ViewCenterZ;
+		Clipmap.ViewRadiusZ = ViewRadiusZ;
+	}
+		
+	PrevPageSpaceLocation = CurrentPageSpaceLocation;
+	
 	CurrentVirtualShadowMapId = VirtualShadowMapId;
 	CurrentPageSpaceLocation = PageSpaceLocation;
-	CurrentShadowMapGlobalDepth = GlobalDepth;
 }
 
-void FVirtualShadowMapCacheEntry::Update(int32 VirtualShadowMapId, const FWholeSceneProjectedShadowInitializer &InCacheValidKey)
+void FVirtualShadowMapCacheEntry::UpdateLocal(int32 VirtualShadowMapId, const FWholeSceneProjectedShadowInitializer &InCacheValidKey)
 {
 	// Swap previous frame data over.
 	PrevPageSpaceLocation = CurrentPageSpaceLocation;
 	PrevVirtualShadowMapId = CurrentVirtualShadowMapId;
-	PrevShadowMapGlobalDepth = CurrentShadowMapGlobalDepth;
 
 	// Check cache validity based of shadow setup
-	if (!CacheValidKey.IsCachedShadowValid(InCacheValidKey))
+	if (!LocalCacheValidKey.IsCachedShadowValid(InCacheValidKey))
 	{
 		PrevVirtualShadowMapId = INDEX_NONE;
 		//UE_LOG(LogRenderer, Display, TEXT("Invalidated!"));
 	}
-
-	CacheValidKey = InCacheValidKey;
+	LocalCacheValidKey = InCacheValidKey;
 
 	CurrentVirtualShadowMapId = VirtualShadowMapId;
 	PrevPageSpaceLocation = CurrentPageSpaceLocation = FIntPoint(0, 0);
-	PrevShadowMapGlobalDepth = CurrentShadowMapGlobalDepth = 0.0f;
 }
 
 
diff --git a/Engine/Source/Runtime/Renderer/Private/VirtualShadowMaps/VirtualShadowMapCacheManager.h b/Engine/Source/Runtime/Renderer/Private/VirtualShadowMaps/VirtualShadowMapCacheManager.h
--- a/Engine/Source/Runtime/Renderer/Private/VirtualShadowMaps/VirtualShadowMapCacheManager.h
+++ b/Engine/Source/Runtime/Renderer/Private/VirtualShadowMaps/VirtualShadowMapCacheManager.h
@@ -15,17 +15,19 @@ class FGPUScene;
 class FVirtualShadowMapCacheEntry
 {
 public:
-	// Key data (must be same to be valid)
-
 	// True if the cache has been (re)populated, set to false on init and set to true once the cache update process has happened.
 	// Also set to false whenever key data was not valid and all cached data is invalidated.
 	bool IsValid() { return PrevVirtualShadowMapId != INDEX_NONE; }
 
+	void UpdateLocal(int32 VirtualShadowMapId,
+		const FWholeSceneProjectedShadowInitializer &InCacheValidKey);
 
-	void Update(int32 VirtualShadowMapId, const FWholeSceneProjectedShadowInitializer &InCacheValidKey);
-
-	void UpdateClipmap(int32 VirtualShadowMapId, const FMatrix &WorldToLight,
-		FIntPoint PageSpaceLocation, float GlobalDepth);
+	void UpdateClipmap(int32 VirtualShadowMapId,
+		const FMatrix &WorldToLight,
+		FIntPoint PageSpaceLocation,
+		float LevelRadius,
+		float ViewCenterZ,
+		float ViewRadiusZ);
 
 	/**
 	 * Returns the PrevVirtualShadowMapId if cached data is valid (bValidData), or INDEX_NONE otherwise. 
@@ -34,39 +36,24 @@ public:
 
 	FIntPoint GetPageSpaceOffset() const { return PrevPageSpaceLocation - CurrentPageSpaceLocation; }
 
-	/**
-	 * Returns the depth offset to add to the depth of pages as they are copied.
-	 */
-	float GetDepthOffset() const { return PrevShadowMapGlobalDepth - CurrentShadowMapGlobalDepth; }
-
 	// Previous frame data
-
-	// Aligned location in pages previous frame
-	//FIntRect PrevViewPort = FIntRect(0, 0, 0, 0);
-	FIntPoint PrevPageSpaceLocation = FIntPoint(0,0);
-
-	// Set to INDEX_NONE when cache data is invalidated by an external change, like movement.
+	FIntPoint PrevPageSpaceLocation = FIntPoint(0, 0);
 	int32 PrevVirtualShadowMapId = INDEX_NONE;
 
-	// Depth of the world-space origin of the shadow map in shadow maps space.
-	// Used to offset the depth of pages as they are copied when the light moves
-	float PrevShadowMapGlobalDepth = 0.0f;
-
-	// Current frame data 
-	
-	// Aligned location in pages after update.
+	// Current frame data
 	FIntPoint CurrentPageSpaceLocation = FIntPoint(0, 0);
 	int32 CurrentVirtualShadowMapId = INDEX_NONE;
-	float CurrentShadowMapGlobalDepth = 0.0f;
 
 	// TODO: Potentially refactor this to decouple the cache key details
-	FWholeSceneProjectedShadowInitializer CacheValidKey;
+	FWholeSceneProjectedShadowInitializer LocalCacheValidKey;
 
-	struct ClipmapCacheKey
+	struct FClipmapInfo
 	{
 		FMatrix	WorldToLight;
+		float ViewCenterZ;
+		float ViewRadiusZ;
 	};
-	ClipmapCacheKey ClipmapCacheValidKey;
+	FClipmapInfo Clipmap;
 };
 
 // Persistent buffers that we ping pong frame by frame
@@ -92,16 +79,6 @@ struct FVirtualShadowMapArrayFrameData
 class FVirtualShadowMapArrayCacheManager
 {
 public:
-	// Align global coordinates to this mip-level, e.g., 3 and page size 128 => 512 texels 
-	// This is only relevant for directional lights that support scrolling, spot lights are invalidated wholesale on movement
-	static constexpr uint32 AlignmentLevel = 3U;
-	static constexpr uint32 AlignmentPages = (1U << AlignmentLevel);
-	static constexpr uint32 AlignmentTexels = AlignmentPages * FVirtualShadowMap::PageSize;
-	static constexpr uint32 EffectiveCacheResolutionTexels = FVirtualShadowMap::VirtualMaxResolutionXY - AlignmentTexels;
-	static constexpr uint32 EffectiveCacheResolutionPages = FVirtualShadowMap::Level0DimPagesXY - AlignmentPages;
-
-	static constexpr float ClipSpaceScaleFactor = float(EffectiveCacheResolutionTexels) / float(FVirtualShadowMap::VirtualMaxResolutionXY);
-
 	/**
 	 * Call at end of frame to extract resouces from the virtual SM array to preserve to next frame.
 	 * If bCachingEnabled is false, all previous frame data is dropped and cache (and HZB!) data will not be available for the next frame.
diff --git a/Engine/Source/Runtime/Renderer/Private/VirtualShadowMaps/VirtualShadowMapClipmap.cpp b/Engine/Source/Runtime/Renderer/Private/VirtualShadowMaps/VirtualShadowMapClipmap.cpp
--- a/Engine/Source/Runtime/Renderer/Private/VirtualShadowMaps/VirtualShadowMapClipmap.cpp
+++ b/Engine/Source/Runtime/Renderer/Private/VirtualShadowMaps/VirtualShadowMapClipmap.cpp
@@ -67,6 +67,8 @@ FVirtualShadowMapClipmap::FVirtualShadowMapClipmap(
 {
 	check(WorldToLightRotationMatrix.GetOrigin() == FVector(0, 0, 0));	// Should not contain translation or scaling
 
+	const bool bCacheValid = VirtualShadowMapArrayCacheManager && VirtualShadowMapArrayCacheManager->IsValid();
+
 	const FMatrix FaceMatrix(
 		FPlane( 0, 0, 1, 0 ),
 		FPlane( 0, 1, 0, 0 ),
@@ -103,14 +105,13 @@ FVirtualShadowMapClipmap::FVirtualShadowMapClipmap(
 	for (int32 Index = 0; Index < LevelCount; ++Index)
 	{
 		FLevelData& Level = LevelData[Index];
-		const int32 LevelIndex = Index + FirstLevel;		// Absolute (virtual) level index
+		const int32 AbsoluteLevel = Index + FirstLevel;		// Absolute (virtual) level index
 
 		// TODO: Allocate these as a chunk if we continue to use one per clipmap level
-		// This isn't actually required at the moment but enforcing it keeps optimization options open
 		Level.VirtualShadowMap = VirtualShadowMapArray.Allocate();
 		ensure(Index == 0 || (Level.VirtualShadowMap->ID == (LevelData[Index-1].VirtualShadowMap->ID + 1)));
 
-		const float RawLevelRadius = GetLevelRadius(LevelIndex);
+		const float RawLevelRadius = GetLevelRadius(AbsoluteLevel);
 
 		float HalfLevelDim = 2.0f * RawLevelRadius;
 		float SnapSize = RawLevelRadius;
@@ -128,29 +129,50 @@ FVirtualShadowMapClipmap::FVirtualShadowMapClipmap(
 
 		const FVector SnappedWorldCenter = ViewToWorldRotationMatrix.TransformPosition(ViewCenter);
 
-		const float ZScale = 0.5f / RawLevelRadius;
-		const float ZOffset = RawLevelRadius;
-
 		Level.WorldCenter = SnappedWorldCenter;
-		Level.ViewToClip = FReversedZOrthoMatrix(HalfLevelDim, HalfLevelDim, ZScale, ZOffset);
 		Level.CornerOffset = CornerOffset;
 
-		if (VirtualShadowMapArrayCacheManager)
+		// Check if we have a cache entry for this level
+		// If we do and it covers our required depth range, we can use cached pages. Otherwise we need to invalidate.
+		TSharedPtr<FVirtualShadowMapCacheEntry> CacheEntry = nullptr;
+		if (bCacheValid)
 		{
 			// NOTE: We use the absolute (virtual) level index so that the caching is robust against changes to the chosen level range
-			TSharedPtr<FVirtualShadowMapCacheEntry> CacheEntry = VirtualShadowMapArrayCacheManager->FindCreateCacheEntry(LightSceneInfo.Id, LevelIndex);
-			if (CacheEntry)
-			{
-				// We snap to half the size of the VSM at each level
-				check((FVirtualShadowMap::Level0DimPagesXY & 1) == 0);
-				FIntPoint PageOffset(CornerOffset * (FVirtualShadowMap::Level0DimPagesXY >> 2));
-				float DepthOffset = -ViewCenter.Z * ZScale;
-
-				CacheEntry->UpdateClipmap(Level.VirtualShadowMap->ID, WorldToLightRotationMatrix, PageOffset, DepthOffset);
-
-				Level.VirtualShadowMap->VirtualShadowMapCacheEntry = CacheEntry;
-			}
+			CacheEntry = VirtualShadowMapArrayCacheManager->FindCreateCacheEntry(LightSceneInfo.Id, AbsoluteLevel);
+		}
+
+		// We expand the depth range of the clipmap level to allow for camera movement without having to invalidate cached shadow data
+		// (See VirtualShadowMapCacheManager::UpdateClipmap for invalidation logic.)
+		// Hard-coded constant currently; should probably always be at least 2.0f, otherwise the cache will be frequently invalidated.
+		const float ViewRadiusZMultiplier = 5.0f;
+
+		float ViewRadiusZ = RawLevelRadius * ViewRadiusZMultiplier;
+		float ViewCenterDeltaZ = 0.0f;
+
+		if (CacheEntry)
+		{
+			// We snap to half the size of the VSM at each level
+			check((FVirtualShadowMap::Level0DimPagesXY & 1) == 0);
+			FIntPoint PageOffset(CornerOffset * (FVirtualShadowMap::Level0DimPagesXY >> 2));
+
+			CacheEntry->UpdateClipmap(Level.VirtualShadowMap->ID,
+				WorldToLightRotationMatrix,
+				PageOffset,
+				RawLevelRadius,
+				ViewCenter.Z,
+				ViewRadiusZ);
+
+			Level.VirtualShadowMap->VirtualShadowMapCacheEntry = CacheEntry;
+
+			// Update min/max Z based on the cached page (if present and valid)
+			// We need to ensure we use a consistent depth range as the camera moves for each level
+			ViewCenterDeltaZ = ViewCenter.Z - CacheEntry->Clipmap.ViewCenterZ;
+			ViewRadiusZ = CacheEntry->Clipmap.ViewRadiusZ;
 		}
+
+		const float ZScale = 0.5f / ViewRadiusZ;
+		const float ZOffset = ViewRadiusZ + ViewCenterDeltaZ;
+		Level.ViewToClip = FReversedZOrthoMatrix(HalfLevelDim, HalfLevelDim, ZScale, ZOffset);
 	}
 }
 
