Modify ternary operator usage in OctahedralCommon to avoid compilation errors with HLSLCC

diff --git a/Engine/Shaders/Private/OctahedralCommon.ush b/Engine/Shaders/Private/OctahedralCommon.ush
--- a/Engine/Shaders/Private/OctahedralCommon.ush
+++ b/Engine/Shaders/Private/OctahedralCommon.ush
@@ -18,7 +18,7 @@
 float2 UnitVectorToOctahedron(float3 V)
 {
 	V /= (abs(V.x) + abs(V.y) + abs(V.z));
-	V.xy = V.z >= 0.0f ? V.xy : ((1.0f - abs(V.yx)) * (V.xy >= 0.0f ? 1.0f : -1.0f));
+	V.xy = V.z >= 0.0f ? V.xy : ((1.0f - abs(V.yx)) * (V.xy >= 0.0f ? float2(1,1) : float2(-1,-1)));
 	V.xy = V.xy * 0.5f + 0.5f;
 	return V.xy;
 }
@@ -29,7 +29,7 @@ float3 OctahedronToUnitVector(float2 UV)
 
 	float3 V = float3(UV.x, UV.y, 1.0 - abs(UV.x) - abs(UV.y));
 	float T = saturate(-V.z);
-	V.xy += V.xy >= 0.0 ? -T : T;
+	V.xy += V.xy >= 0.0 ? float2(-T,-T) : float2(T,T);
 
 	return normalize(V);
 }
