Fixed some other Strata loop dxc has trouble to unroll...

diff --git a/Engine/Shaders/Private/Strata/Strata.ush b/Engine/Shaders/Private/Strata/Strata.ush
--- a/Engine/Shaders/Private/Strata/Strata.ush
+++ b/Engine/Shaders/Private/Strata/Strata.ush
@@ -12,6 +12,17 @@
 #include "/Engine/Shared/StrataDefinitions.h"
 
 
+// DXC has issue with unrolling certain part of the code. To address temporarly the warning-as-error issues, 
+// we redefine the unroll function and fall back to loop on DXC
+#if COMPILER_DXC == 1
+#define STRATA_UNROLL LOOP
+#define STRATA_UNROLL_N(X) LOOP
+#else
+#define STRATA_UNROLL UNROLL
+#define STRATA_UNROLL_N(X) UNROLL_N(X)
+#endif
+
+
 // During the base pass or forward rendering, shared local bases are simply available from registers.
 // If a BSDF is loaded for processing from the strata buffer, normals are unpacked on demand to reduce VGPR pressure and increase occupancy.
 #ifndef STRATA_INLINE_SHADING
@@ -849,12 +860,12 @@ FStrataData StrataAdd(FStrataData A, FStrataData B)
 	FStrataData StrataData = A;
 
 	// Append each BSDF from B to A, with same layer position
-	UNROLL
+	STRATA_UNROLL_N(STRATA_CLAMPED_LAYER_COUNT)
 	for (int Layer = 0; Layer < B.LayerCount; ++Layer)
 	{
 		FStrataLayer ALayer = A.Layers[Layer];
 		FStrataLayer BLayer = B.Layers[Layer];
-		UNROLL
+		STRATA_UNROLL
 		for (int BSDF = 0; BSDF < BLayer.BSDFCount; BSDF++)
 		{
 			StrataData.Layers[Layer].BSDFs[ALayer.BSDFCount + BSDF] = BLayer.BSDFs[BSDF];
@@ -874,10 +885,10 @@ FStrataData StrataWeight(FStrataData A, float Weight)
 	FStrataData StrataData = A;
 
 	// Apply the weight to each BSDF
-	UNROLL
+	STRATA_UNROLL_N(STRATA_CLAMPED_LAYER_COUNT)
 	for (int Layer = 0; Layer < StrataData.LayerCount; ++Layer)
 	{
-		UNROLL
+		STRATA_UNROLL
 		for (int BSDF = 0; BSDF < StrataData.Layers[Layer].BSDFCount; ++BSDF)
 		{
 #if STRATA_INLINE_SHADING
@@ -914,7 +925,7 @@ FStrataData StrataVerticalLayering(FStrataData Top, FStrataData Base)
 
 	// Add each layer from Base under Top
 	int TopLayerCount = Top.LayerCount;
-	UNROLL
+	STRATA_UNROLL_N(STRATA_CLAMPED_LAYER_COUNT)
 	for (int Layer = 0; Layer < Base.LayerCount; ++Layer)
 	{
 		StrataData.Layers[TopLayerCount + Layer] = Base.Layers[Layer];
@@ -1667,16 +1678,6 @@ void StrataLoad_ColorGamma2ToLinearAlphaLinear(ByteAddressBuffer StrataBuffer, i
 	A = Data.a;
 }
 
-// DXC has issue with unrolling certain part of the code. To address temporarly the warning-as-error issues, 
-// we redefine the unroll function and fall back to loop on DXC
-#if COMPILER_DXC == 1
-#define STRATA_UNROLL LOOP
-#define STRATA_UNROLL_N(X) LOOP
-#else
-#define STRATA_UNROLL UNROLL
-#define STRATA_UNROLL_N(X) UNROLL_N(X)
-#endif
-
 uint PackStrataHeaderIn(uint BSDFCount, FStrataPixelHeader In);
 
 // If this is changed, please update the compiler side material size evaluation in StrataMaterial.cpp
