Fix invalid alignment for VSM cache data.

This causes a vulkan error when packaging a project.

diff --git a/Engine/Shaders/Private/VirtualShadowMaps/PageCacheCommon.ush b/Engine/Shaders/Private/VirtualShadowMaps/PageCacheCommon.ush
--- a/Engine/Shaders/Private/VirtualShadowMaps/PageCacheCommon.ush
+++ b/Engine/Shaders/Private/VirtualShadowMaps/PageCacheCommon.ush
@@ -18,6 +18,7 @@ struct FShadowMapCacheData
 	int2 SmPageOffset;
 	// ID of the corresponding virtual SM in the chached data 
 	int VirtualShadowMapId;
+	int Padding;
 };
 
 struct FCachedPageInfo
