Added support for starta on debug lighting probes.
Shootergame compiled.

diff --git a/Engine/Shaders/Private/BasePassPixelShader.usf b/Engine/Shaders/Private/BasePassPixelShader.usf
--- a/Engine/Shaders/Private/BasePassPixelShader.usf
+++ b/Engine/Shaders/Private/BasePassPixelShader.usf
@@ -1576,7 +1576,6 @@ void FPixelShaderInOut_MainPS(
 	uint2 PixelPos = In.SvPosition.xy;
 
 	float3 EmissiveLuminance = 0.0f;
-	const float3 VertexNormal = MaterialParameters.TangentToWorld[2].xyz;
 	PackStrataOut(
 		OpaqueBasePass.Strata.MaterialLobesBufferUAV, Dither,
 		GetStrataPixelDataByteOffset(PixelPos, uint2(ResolvedView.BufferSizeAndInvSize.xy), OpaqueBasePass.Strata.MaxBytesPerPixel),
diff --git a/Engine/Shaders/Private/DebugProbes.usf b/Engine/Shaders/Private/DebugProbes.usf
--- a/Engine/Shaders/Private/DebugProbes.usf
+++ b/Engine/Shaders/Private/DebugProbes.usf
@@ -1,15 +1,26 @@
-
-
+// Copyright Epic Games, Inc. All Rights Reserved.
 
 #include "DeferredShadingCommon.ush"
 #include "VelocityCommon.ush"
 
+#if PROJECT_STRATA
+#include "Strata/Strata.ush"
+RWByteAddressBuffer MaterialLobesBufferUAV;
+uint MaxBytesPerPixel;
+#endif
+int DebugProbesMode;
+
 
 
 #define PROBE_MODE_SPECULAR	0
 #define PROBE_MODE_DIFFUSE	1
 #define PROBE_MODE_BOTH		2
 
+// Must match DebugProbeRendering.usf
+#define RENDER_DEPTHPREPASS	0
+#define RENDER_BASEPASS		1
+#define RENDER_VELOCITYPASS	2
+
 
 
 void DebugSphereIntersection(
@@ -39,14 +50,10 @@ void DebugSphereIntersection(
 
 
 
-int DebugProbesMode;
-
-
-
 void MainPS(
 	  in  float4 SvPosition		: SV_Position
 	, out float4 OutColor		: SV_Target0
-#if PERMUTATION_RENDERVELOCITY == 0
+#if PERMUTATION_PASS == RENDER_BASEPASS
 	, out float4 OutGBufferA	: SV_Target1
 	, out float4 OutGBufferB	: SV_Target2
 	, out float4 OutGBufferC	: SV_Target3
@@ -200,7 +207,7 @@ void MainPS(
 	}
 
 
-#if PERMUTATION_RENDERVELOCITY
+#if PERMUTATION_PASS == RENDER_VELOCITYPASS
 
 	float3 Velocity = 0.0f;
 
@@ -213,20 +220,62 @@ void MainPS(
 
 	OutColor = EncodeVelocityToTexture(Velocity);
 
-#else
+#elif PERMUTATION_PASS == RENDER_BASEPASS
 
 	const float PerObjectGBufferData = 0.0f;
 	const float SelectiveOutputMask = 0.0f;
 
 	const float3 BaseColor = 1.0f;
-	const float  Metallic  = ProbeMode == PROBE_MODE_DIFFUSE ? 0.0f : 1.0f;
-	const float  Specular  = ProbeMode == PROBE_MODE_DIFFUSE ? 0.0f : 1.0f;
+	const float  Metallic = ProbeMode == PROBE_MODE_DIFFUSE ? 0.0f : 1.0f;
+	const float  Specular = ProbeMode == PROBE_MODE_DIFFUSE ? 0.0f : 1.0f;
 	const float  Roughness = 0.01f;
 	const float  IndirectIrradianceAndAO = 0.0f;
 
-	// We write all outputs, the rasterizer will discard what is not bound. That is fine since we do not care that much about performance for this debug view.
+#if PROJECT_STRATA
+
+	FStrataPixelHeader StrataPixelHeader = InitialiseStrataPixelHeader();
+
+	uint SharedLocalBasisIndex = 0;
+	StrataPixelHeader.SharedLocalBases.Normals[SharedLocalBasisIndex] = normalize(Normal);
+	StrataPixelHeader.SharedLocalBases.Count = 1;
+
+	FStrataData StrataData = GetStrataSlabBSDF(
+		/*float UseMetalness					*/true,
+		/*float3 BaseColor						*/BaseColor, 
+		/*float3 EdgeColor						*/1.0f, 
+		/*float Specular						*/Specular, 
+		/*float Metallic						*/Metallic,
+		/*float3 DiffuseAlbedo					*/0.0f, 
+		/*float3 F0								*/0.0f, 
+		/*float3 F90							*/0.0f,
+		/*float Roughness						*/Roughness, 
+		/*float Anisotropy						*/0.0f,
+		/*float SSSProfileID					*/0.0f, 
+		/*float3 SSSDMFP						*/0.0f, 
+		/*float SSSDMFPScale					*/0.0f,
+		/*float3 Emissive						*/0.0f, 
+		/*float Haziness						*/0.0f, 
+		/*float ThinFilmThickness				*/0.0f,
+		/*float FuzzAmount						*/0.0f, 
+		/*float3 FuzzColor						*/0.0f,
+		/*float Thickness						*/0.0f, 
+		/*uint SharedLocalBasisIndex			*/SharedLocalBasisIndex,
+		/*inout uint SharedLocalBasisTypes		*/StrataPixelHeader.SharedLocalBases.Types);
+
+	StrataPixelHeader.MaterialAO = 1.0f - IndirectIrradianceAndAO; 
+	StrataPixelHeader.BSDFCount = 1;
+
+	float3 EmissiveLuminance = 0.0f;
+	bool bStrataSubsurfaceEnable = false; 
+	const float Dither = 0.0f;
+	PackStrataOut(
+		MaterialLobesBufferUAV, Dither,
+		GetStrataPixelDataByteOffset(PixelCoord, uint2(ResolvedView.BufferSizeAndInvSize.xy), MaxBytesPerPixel),
+		StrataPixelHeader, StrataData, -RayDir, bStrataSubsurfaceEnable, EmissiveLuminance);
+	
+#else // PROJECT_STRATA
 
-	OutColor = 0.0f;
+	// We write all outputs, the rasterizer will discard what is not bound. That is fine since we do not care that much about performance for this debug view.
 
 	OutGBufferA = float4(
 		EncodeNormal(Normal), 
@@ -242,9 +291,11 @@ void MainPS(
 	OutGBufferD = float4(0.0f, 0.0f, 0.0f, 0.0f);	// custom data
 	OutGBufferE = float4(1.0f, 1.0f, 1.0f, 1.0f);	// Pre computed shadow factor
 
-	// Not writting to the velocity buffer
+#endif // PROJECT_STRATA
 
-#endif
+	OutColor = 0.0f;
+
+#endif // PERMUTATION_PASS
 
 	OutDepth = DeviceZ;
 
diff --git a/Engine/Source/Runtime/Renderer/Private/DebugProbeRendering.cpp b/Engine/Source/Runtime/Renderer/Private/DebugProbeRendering.cpp
--- a/Engine/Source/Runtime/Renderer/Private/DebugProbeRendering.cpp
+++ b/Engine/Source/Runtime/Renderer/Private/DebugProbeRendering.cpp
@@ -3,6 +3,8 @@
 #include "DebugProbeRendering.h"
 #include "PixelShaderUtils.h"
 #include "ShaderParameterStruct.h"
+#include "SceneRendering.h"
+#include "Strata/Strata.h"
 
 
 // Changing this causes a full shader recompile
@@ -21,16 +23,24 @@ static TAutoConsoleVariable<int32> CVarVisualizeLightingOnProbes(
 DECLARE_GPU_STAT(StampDeferredDebugProbe);
 
 
+// Must match DebugProbes.usf
+#define RENDER_DEPTHPREPASS  0
+#define RENDER_BASEPASS	     1
+#define RENDER_VELOCITYPASS  2
+
+
 class FStampDeferredDebugProbePS : public FGlobalShader
 {
 	DECLARE_GLOBAL_SHADER(FStampDeferredDebugProbePS);
 	SHADER_USE_PARAMETER_STRUCT(FStampDeferredDebugProbePS, FGlobalShader);
 
-	class FRenderVelocity : SHADER_PERMUTATION_BOOL("PERMUTATION_RENDERVELOCITY");
-	using FPermutationDomain = TShaderPermutationDomain<FRenderVelocity>;
+	class FRenderPass : SHADER_PERMUTATION_RANGE_INT("PERMUTATION_PASS", 0, 3);
+	using FPermutationDomain = TShaderPermutationDomain<FRenderPass>;
 
 	BEGIN_SHADER_PARAMETER_STRUCT(FParameters, )
 		SHADER_PARAMETER_STRUCT_REF(FViewUniformShaderParameters, ViewUniformBuffer)
+		SHADER_PARAMETER_RDG_BUFFER_UAV(RWByteAddressBuffer, MaterialLobesBufferUAV)
+		SHADER_PARAMETER(uint32, MaxBytesPerPixel)
 		SHADER_PARAMETER(int32, DebugProbesMode)
 		RENDER_TARGET_BINDING_SLOTS()
 	END_SHADER_PARAMETER_STRUCT()
@@ -60,13 +70,15 @@ static void CommonStampDeferredDebugProbeDrawCall(
 	FRDGBuilder& GraphBuilder,
 	const FViewInfo& View, 
 	FStampDeferredDebugProbePS::FParameters* PassParameters,
-	bool bRenderVelocity = false)
+	int32 RenderPass)
 {
 	PassParameters->ViewUniformBuffer = View.ViewUniformBuffer;
+	PassParameters->MaterialLobesBufferUAV = View.StrataSceneData->MaterialLobesBufferUAV;
+	PassParameters->MaxBytesPerPixel = View.StrataSceneData->MaxBytesPerPixel;
 	PassParameters->DebugProbesMode = View.Family->EngineShowFlags.VisualizeLightingOnProbes ? 3 : FMath::Clamp(CVarVisualizeLightingOnProbes.GetValueOnRenderThread(), 0, 3);
 		
 	FStampDeferredDebugProbePS::FPermutationDomain PermutationVector;
-	PermutationVector.Set<FStampDeferredDebugProbePS::FRenderVelocity>(bRenderVelocity);
+	PermutationVector.Set<FStampDeferredDebugProbePS::FRenderPass>(RenderPass);
 	TShaderMapRef<FStampDeferredDebugProbePS> PixelShader(View.ShaderMap, PermutationVector);
 
 	FPixelShaderUtils::AddFullscreenPass<FStampDeferredDebugProbePS>(
@@ -94,10 +106,9 @@ void StampDeferredDebugProbeDepthPS(
 		}
 
 		FStampDeferredDebugProbePS::FParameters* PassParameters = GraphBuilder.AllocParameters<FStampDeferredDebugProbePS::FParameters>();
-		PassParameters->ViewUniformBuffer = View.ViewUniformBuffer;
 		PassParameters->RenderTargets.DepthStencil = FDepthStencilBinding(SceneDepthTexture, ERenderTargetLoadAction::ELoad, ERenderTargetLoadAction::ELoad, FExclusiveDepthStencil::DepthWrite_StencilWrite);
 
-		CommonStampDeferredDebugProbeDrawCall<true>(GraphBuilder, View, PassParameters);
+		CommonStampDeferredDebugProbeDrawCall<true>(GraphBuilder, View, PassParameters, RENDER_DEPTHPREPASS);
 	}
 }
 
@@ -120,7 +131,7 @@ void StampDeferredDebugProbeMaterialPS(
 		FStampDeferredDebugProbePS::FParameters* PassParameters = GraphBuilder.AllocParameters<FStampDeferredDebugProbePS::FParameters>();
 		PassParameters->RenderTargets = BasePassRenderTargets;
 
-		CommonStampDeferredDebugProbeDrawCall<false>(GraphBuilder, View, PassParameters);
+		CommonStampDeferredDebugProbeDrawCall<false>(GraphBuilder, View, PassParameters, RENDER_BASEPASS);
 	}
 }
 
@@ -144,7 +155,7 @@ void StampDeferredDebugProbeVelocityPS(
 		PassParameters->RenderTargets = BasePassRenderTargets;
 
 		const bool bRenderVelocity = true;
-		CommonStampDeferredDebugProbeDrawCall<false>(GraphBuilder, View, PassParameters, bRenderVelocity);
+		CommonStampDeferredDebugProbeDrawCall<false>(GraphBuilder, View, PassParameters, RENDER_VELOCITYPASS);
 	}
 }
 
