Added in shader depth test for lighting debug probe with starta.

diff --git a/Engine/Shaders/Private/DebugProbes.usf b/Engine/Shaders/Private/DebugProbes.usf
--- a/Engine/Shaders/Private/DebugProbes.usf
+++ b/Engine/Shaders/Private/DebugProbes.usf
@@ -4,6 +4,7 @@
 #include "VelocityCommon.ush"
 
 #if PROJECT_STRATA
+#include "SceneTexturesCommon.ush"
 #include "Strata/Strata.ush"
 RWByteAddressBuffer MaterialLobesBufferUAV;
 uint MaxBytesPerPixel;
@@ -233,45 +234,53 @@ void MainPS(
 
 #if PROJECT_STRATA
 
-	FStrataPixelHeader StrataPixelHeader = InitialiseStrataPixelHeader();
-
-	uint SharedLocalBasisIndex = 0;
-	StrataPixelHeader.SharedLocalBases.Normals[SharedLocalBasisIndex] = normalize(Normal);
-	StrataPixelHeader.SharedLocalBases.Count = 1;
-
-	FStrataData StrataData = GetStrataSlabBSDF(
-		/*float UseMetalness					*/true,
-		/*float3 BaseColor						*/BaseColor, 
-		/*float3 EdgeColor						*/1.0f, 
-		/*float Specular						*/Specular, 
-		/*float Metallic						*/Metallic,
-		/*float3 DiffuseAlbedo					*/0.0f, 
-		/*float3 F0								*/0.0f, 
-		/*float3 F90							*/0.0f,
-		/*float Roughness						*/Roughness, 
-		/*float Anisotropy						*/0.0f,
-		/*float SSSProfileID					*/0.0f, 
-		/*float3 SSSDMFP						*/0.0f, 
-		/*float SSSDMFPScale					*/0.0f,
-		/*float3 Emissive						*/0.0f, 
-		/*float Haziness						*/0.0f, 
-		/*float ThinFilmThickness				*/0.0f,
-		/*float FuzzAmount						*/0.0f, 
-		/*float3 FuzzColor						*/0.0f,
-		/*float Thickness						*/0.0f, 
-		/*uint SharedLocalBasisIndex			*/SharedLocalBasisIndex,
-		/*inout uint SharedLocalBasisTypes		*/StrataPixelHeader.SharedLocalBases.Types);
-
-	StrataPixelHeader.MaterialAO = 1.0f - IndirectIrradianceAndAO; 
-	StrataPixelHeader.BSDFCount = 1;
-
-	float3 EmissiveLuminance = 0.0f;
-	bool bStrataSubsurfaceEnable = false; 
-	const float Dither = 0.0f;
-	PackStrataOut(
-		MaterialLobesBufferUAV, Dither,
-		GetStrataPixelDataByteOffset(PixelCoord, uint2(ResolvedView.BufferSizeAndInvSize.xy), MaxBytesPerPixel),
-		StrataPixelHeader, StrataData, -RayDir, bStrataSubsurfaceEnable, EmissiveLuminance);
+	// EarlyDepthStencil does not work because we render a simple quad on screen with custom OutDepth.
+	// We cannot rely on late depth test because pixel shader would have already operated the write to strata buffer.
+	// So we must execute a manual depth test in shader (fine because basepass does not write depth)
+	float2 UvBuffer = SvPosition.xy * View.BufferSizeAndInvSize.zw;	// Uv for depth buffer read (size can be larger than viewport)
+	const float DepthBufferDeviceZ = LookupDeviceZ(UvBuffer);
+	if (DepthBufferDeviceZ < (DeviceZ+0.000001))
+	{
+		FStrataPixelHeader StrataPixelHeader = InitialiseStrataPixelHeader();
+
+		uint SharedLocalBasisIndex = 0;
+		StrataPixelHeader.SharedLocalBases.Normals[SharedLocalBasisIndex] = normalize(Normal);
+		StrataPixelHeader.SharedLocalBases.Count = 1;
+
+		FStrataData StrataData = GetStrataSlabBSDF(
+			/*float UseMetalness					*/true,
+			/*float3 BaseColor						*/BaseColor,
+			/*float3 EdgeColor						*/1.0f,
+			/*float Specular						*/Specular,
+			/*float Metallic						*/Metallic,
+			/*float3 DiffuseAlbedo					*/0.0f,
+			/*float3 F0								*/0.0f,
+			/*float3 F90							*/0.0f,
+			/*float Roughness						*/Roughness,
+			/*float Anisotropy						*/0.0f,
+			/*float SSSProfileID					*/0.0f,
+			/*float3 SSSDMFP						*/0.0f,
+			/*float SSSDMFPScale					*/0.0f,
+			/*float3 Emissive						*/0.0f,
+			/*float Haziness						*/0.0f,
+			/*float ThinFilmThickness				*/0.0f,
+			/*float FuzzAmount						*/0.0f,
+			/*float3 FuzzColor						*/0.0f,
+			/*float Thickness						*/0.0f,
+			/*uint SharedLocalBasisIndex			*/SharedLocalBasisIndex,
+			/*inout uint SharedLocalBasisTypes		*/StrataPixelHeader.SharedLocalBases.Types);
+
+		StrataPixelHeader.MaterialAO = 1.0f - IndirectIrradianceAndAO;
+		StrataPixelHeader.BSDFCount = 1;
+
+		float3 EmissiveLuminance = 0.0f;
+		bool bStrataSubsurfaceEnable = false;
+		const float Dither = 0.0f;
+		PackStrataOut(
+			MaterialLobesBufferUAV, Dither,
+			GetStrataPixelDataByteOffset(PixelCoord, uint2(ResolvedView.BufferSizeAndInvSize.xy), MaxBytesPerPixel),
+			StrataPixelHeader, StrataData, -RayDir, bStrataSubsurfaceEnable, EmissiveLuminance);
+	}
 	
 #else // PROJECT_STRATA
 
@@ -297,6 +306,11 @@ void MainPS(
 
 #endif // PERMUTATION_PASS
 
+
+
+#if PERMUTATION_PASS == RENDER_VELOCITYPASS || PERMUTATION_PASS == RENDER_DEPTHPREPASS || (PERMUTATION_PASS == RENDER_BASEPASS && PROJECT_STRATA == 0)
+	// Base pass with strata cannot write to depth because it is going to execute in shader depth testing
 	OutDepth = DeviceZ;
+#endif
 
 }
\ No newline at end of file
diff --git a/Engine/Source/Runtime/Renderer/Private/BasePassRendering.cpp b/Engine/Source/Runtime/Renderer/Private/BasePassRendering.cpp
--- a/Engine/Source/Runtime/Renderer/Private/BasePassRendering.cpp
+++ b/Engine/Source/Runtime/Renderer/Private/BasePassRendering.cpp
@@ -913,7 +913,7 @@ void FDeferredShadingSceneRenderer::RenderBasePass(
 #if !(UE_BUILD_SHIPPING)
 	if (!bForwardShadingEnabled)
 	{
-		StampDeferredDebugProbeMaterialPS(GraphBuilder, Views, BasePassRenderTargets);
+		StampDeferredDebugProbeMaterialPS(GraphBuilder, Views, BasePassRenderTargets, SceneTextures);
 	}
 #endif
 }
diff --git a/Engine/Source/Runtime/Renderer/Private/DebugProbeRendering.cpp b/Engine/Source/Runtime/Renderer/Private/DebugProbeRendering.cpp
--- a/Engine/Source/Runtime/Renderer/Private/DebugProbeRendering.cpp
+++ b/Engine/Source/Runtime/Renderer/Private/DebugProbeRendering.cpp
@@ -4,6 +4,8 @@
 #include "PixelShaderUtils.h"
 #include "ShaderParameterStruct.h"
 #include "SceneRendering.h"
+#include "SceneTextures.h"
+#include "SceneTextureParameters.h"
 #include "Strata/Strata.h"
 
 
@@ -41,6 +43,7 @@ class FStampDeferredDebugProbePS : public FGlobalShader
 		SHADER_PARAMETER_STRUCT_REF(FViewUniformShaderParameters, ViewUniformBuffer)
 		SHADER_PARAMETER_RDG_BUFFER_UAV(RWByteAddressBuffer, MaterialLobesBufferUAV)
 		SHADER_PARAMETER(uint32, MaxBytesPerPixel)
+		SHADER_PARAMETER_STRUCT_INCLUDE(FSceneTextureShaderParameters, SceneTextures)
 		SHADER_PARAMETER(int32, DebugProbesMode)
 		RENDER_TARGET_BINDING_SLOTS()
 	END_SHADER_PARAMETER_STRUCT()
@@ -65,7 +68,7 @@ public:
 IMPLEMENT_GLOBAL_SHADER(FStampDeferredDebugProbePS, "/Engine/Private/DebugProbes.usf", "MainPS", SF_Pixel);
 
  
-template<bool bEnableDepthWrite>
+template<bool bEnableDepthWrite, ECompareFunction CompareFunction>
 static void CommonStampDeferredDebugProbeDrawCall(
 	FRDGBuilder& GraphBuilder,
 	const FViewInfo& View, 
@@ -86,7 +89,7 @@ static void CommonStampDeferredDebugProbeDrawCall(
 		PixelShader, PassParameters, View.ViewRect,
 		TStaticBlendState<>::GetRHI(),
 		TStaticRasterizerState<FM_Solid, CM_None>::GetRHI(),
-		TStaticDepthStencilState<bEnableDepthWrite, CF_DepthNearOrEqual>::GetRHI());
+		TStaticDepthStencilState<bEnableDepthWrite, CompareFunction>::GetRHI());
 }
 
 void StampDeferredDebugProbeDepthPS(
@@ -108,14 +111,15 @@ void StampDeferredDebugProbeDepthPS(
 		FStampDeferredDebugProbePS::FParameters* PassParameters = GraphBuilder.AllocParameters<FStampDeferredDebugProbePS::FParameters>();
 		PassParameters->RenderTargets.DepthStencil = FDepthStencilBinding(SceneDepthTexture, ERenderTargetLoadAction::ELoad, ERenderTargetLoadAction::ELoad, FExclusiveDepthStencil::DepthWrite_StencilWrite);
 
-		CommonStampDeferredDebugProbeDrawCall<true>(GraphBuilder, View, PassParameters, RENDER_DEPTHPREPASS);
+		CommonStampDeferredDebugProbeDrawCall<true, CF_DepthNearOrEqual>(GraphBuilder, View, PassParameters, RENDER_DEPTHPREPASS);
 	}
 }
 
 void StampDeferredDebugProbeMaterialPS(
 	FRDGBuilder& GraphBuilder,
 	TArrayView<const FViewInfo> Views,
-	const FRenderTargetBindingSlots& BasePassRenderTargets)
+	const FRenderTargetBindingSlots& BasePassRenderTargets,
+	const FMinimalSceneTextures& SceneTextures)
 {
 	RDG_EVENT_SCOPE(GraphBuilder, "StampDeferredDebugProbeMaterial");
 	RDG_GPU_STAT_SCOPE(GraphBuilder, StampDeferredDebugProbe);
@@ -130,8 +134,19 @@ void StampDeferredDebugProbeMaterialPS(
 
 		FStampDeferredDebugProbePS::FParameters* PassParameters = GraphBuilder.AllocParameters<FStampDeferredDebugProbePS::FParameters>();
 		PassParameters->RenderTargets = BasePassRenderTargets;
+		if (Strata::IsStrataEnabled())
+		{
+			// Make sure we do not write depth so that we can safely read it from texture parameters
+			PassParameters->RenderTargets.DepthStencil = FDepthStencilBinding();
+			PassParameters->SceneTextures = CreateSceneTextureShaderParameters(GraphBuilder, View.GetFeatureLevel(), ESceneTextureSetupMode::SceneDepth);
+
+			CommonStampDeferredDebugProbeDrawCall<false, CF_Always>(GraphBuilder, View, PassParameters, RENDER_BASEPASS);
+		}
+		else
+		{
+			CommonStampDeferredDebugProbeDrawCall<false, CF_DepthNearOrEqual>(GraphBuilder, View, PassParameters, RENDER_BASEPASS);
+		}
 
-		CommonStampDeferredDebugProbeDrawCall<false>(GraphBuilder, View, PassParameters, RENDER_BASEPASS);
 	}
 }
 
@@ -155,7 +170,7 @@ void StampDeferredDebugProbeVelocityPS(
 		PassParameters->RenderTargets = BasePassRenderTargets;
 
 		const bool bRenderVelocity = true;
-		CommonStampDeferredDebugProbeDrawCall<false>(GraphBuilder, View, PassParameters, RENDER_VELOCITYPASS);
+		CommonStampDeferredDebugProbeDrawCall<false, CF_DepthNearOrEqual>(GraphBuilder, View, PassParameters, RENDER_VELOCITYPASS);
 	}
 }
 
diff --git a/Engine/Source/Runtime/Renderer/Private/DebugProbeRendering.h b/Engine/Source/Runtime/Renderer/Private/DebugProbeRendering.h
--- a/Engine/Source/Runtime/Renderer/Private/DebugProbeRendering.h
+++ b/Engine/Source/Runtime/Renderer/Private/DebugProbeRendering.h
@@ -20,7 +20,8 @@ void StampDeferredDebugProbeDepthPS(
 void StampDeferredDebugProbeMaterialPS(
 	FRDGBuilder& GraphBuilder,
 	TArrayView<const FViewInfo> Views,
-	const FRenderTargetBindingSlots& BasePassRenderTargets);
+	const FRenderTargetBindingSlots& BasePassRenderTargets,
+	const FMinimalSceneTextures& SceneTextures);
 
 
 void StampDeferredDebugProbeVelocityPS(
