Add ShaderPrintText function for print TEXT("") string within shader.

diff --git a/Engine/Shaders/Private/HairStrands/HairStrandsVoxelPageRayMarching.usf b/Engine/Shaders/Private/HairStrands/HairStrandsVoxelPageRayMarching.usf
--- a/Engine/Shaders/Private/HairStrands/HairStrandsVoxelPageRayMarching.usf
+++ b/Engine/Shaders/Private/HairStrands/HairStrandsVoxelPageRayMarching.usf
@@ -106,102 +106,24 @@ void MainCS(uint3 DispatchThreadId : SV_DispatchThreadID)
 			float2 StartPos = float2(50.0f, 100.0f ) / OutputResolution;
 			float2 Pos = StartPos;
 
-			Pos = ShaderPrintSymbol(Pos, _G_);
-			Pos = ShaderPrintSymbol(Pos, _E_);
-			Pos = ShaderPrintSymbol(Pos, _N_);
-			Pos = ShaderPrintSymbol(Pos, _E_);
-			Pos = ShaderPrintSymbol(Pos, _R_);
-			Pos = ShaderPrintSymbol(Pos, _A_);
-			Pos = ShaderPrintSymbol(Pos, _L_);
+			Pos = ShaderPrintText(Pos, TEXT("General"));
 			Pos = ShaderPrintNewline(Pos); Pos.x = StartPos.x;
 
-			Pos = ShaderPrintSymbol(Pos, _SPC_);
-			Pos = ShaderPrintSymbol(Pos, _V_);
-			Pos = ShaderPrintSymbol(Pos, _O_);
-			Pos = ShaderPrintSymbol(Pos, _X_);
-			Pos = ShaderPrintSymbol(Pos, _E_);
-			Pos = ShaderPrintSymbol(Pos, _L_);
-			Pos = ShaderPrintSymbol(Pos, _SPC_);
-			Pos = ShaderPrintSymbol(Pos, _W_);
-			Pos = ShaderPrintSymbol(Pos, _O_);
-			Pos = ShaderPrintSymbol(Pos, _R_);
-			Pos = ShaderPrintSymbol(Pos, _L_);
-			Pos = ShaderPrintSymbol(Pos, _D_);
-			Pos = ShaderPrintSymbol(Pos, _SPC_);
-			Pos = ShaderPrintSymbol(Pos, _S_);
-			Pos = ShaderPrintSymbol(Pos, _I_);
-			Pos = ShaderPrintSymbol(Pos, _Z_);
-			Pos = ShaderPrintSymbol(Pos, _E_);
-			Pos = ShaderPrintSymbol(Pos, _SPC_);
-			Pos = ShaderPrintSymbol(Pos, _SPC_);
+			Pos = ShaderPrintText(Pos, TEXT(" Voxel World Size "));
 			Pos = ShaderPrint(Pos, VirtualVoxel.VoxelWorldSize);
 			Pos = ShaderPrintNewline(Pos); Pos.x = StartPos.x;
-
-			Pos = ShaderPrintSymbol(Pos, _SPC_);
-			Pos = ShaderPrintSymbol(Pos, _P_);
-			Pos = ShaderPrintSymbol(Pos, _A_);
-			Pos = ShaderPrintSymbol(Pos, _G_);
-			Pos = ShaderPrintSymbol(Pos, _E_);
-			Pos = ShaderPrintSymbol(Pos, _SPC_);
-			Pos = ShaderPrintSymbol(Pos, _R_);
-			Pos = ShaderPrintSymbol(Pos, _E_);
-			Pos = ShaderPrintSymbol(Pos, _S_);
-			Pos = ShaderPrintSymbol(Pos, _SPC_);
-			Pos = ShaderPrintSymbol(Pos, _SPC_);
-			Pos = ShaderPrintSymbol(Pos, _SPC_);
-			Pos = ShaderPrintSymbol(Pos, _SPC_);
-			Pos = ShaderPrintSymbol(Pos, _SPC_);
-			Pos = ShaderPrintSymbol(Pos, _SPC_);
-			Pos = ShaderPrintSymbol(Pos, _SPC_);
-			Pos = ShaderPrintSymbol(Pos, _SPC_);
-			Pos = ShaderPrintSymbol(Pos, _SPC_);
-			Pos = ShaderPrintSymbol(Pos, _SPC_);
+			
+			Pos = ShaderPrintText(Pos, TEXT(" Page Res         "));
 			Pos = ShaderPrint(Pos, VirtualVoxel.PageResolution);
 			Pos = ShaderPrintNewline(Pos); Pos.x = StartPos.x;
 
-			Pos = ShaderPrintSymbol(Pos, _SPC_);
-			Pos = ShaderPrintSymbol(Pos, _P_);
-			Pos = ShaderPrintSymbol(Pos, _A_);
-			Pos = ShaderPrintSymbol(Pos, _G_);
-			Pos = ShaderPrintSymbol(Pos, _E_);
-			Pos = ShaderPrintSymbol(Pos, _SPC_);
-			Pos = ShaderPrintSymbol(Pos, _C_);
-			Pos = ShaderPrintSymbol(Pos, _O_);
-			Pos = ShaderPrintSymbol(Pos, _U_);
-			Pos = ShaderPrintSymbol(Pos, _N_);
-			Pos = ShaderPrintSymbol(Pos, _T_);
-			Pos = ShaderPrintSymbol(Pos, _SPC_);
-			Pos = ShaderPrintSymbol(Pos, _R_);
-			Pos = ShaderPrintSymbol(Pos, _E_);
-			Pos = ShaderPrintSymbol(Pos, _S_);
-			Pos = ShaderPrintSymbol(Pos, _SPC_);
-			Pos = ShaderPrintSymbol(Pos, _SPC_);
-			Pos = ShaderPrintSymbol(Pos, _SPC_);
-			Pos = ShaderPrintSymbol(Pos, _SPC_);
+			Pos = ShaderPrintText(Pos, TEXT(" Page Count Res   "));
 			Pos = ShaderPrint(Pos, VirtualVoxel.PageCountResolution.x);
 			Pos = ShaderPrint(Pos, VirtualVoxel.PageCountResolution.y);
 			Pos = ShaderPrint(Pos, VirtualVoxel.PageCountResolution.z);
 			Pos = ShaderPrintNewline(Pos); Pos.x = StartPos.x;
 
-			Pos = ShaderPrintSymbol(Pos, _SPC_);
-			Pos = ShaderPrintSymbol(Pos, _P_);
-			Pos = ShaderPrintSymbol(Pos, _A_);
-			Pos = ShaderPrintSymbol(Pos, _G_);
-			Pos = ShaderPrintSymbol(Pos, _E_);
-			Pos = ShaderPrintSymbol(Pos, _SPC_);
-			Pos = ShaderPrintSymbol(Pos, _T_);
-			Pos = ShaderPrintSymbol(Pos, _E_);
-			Pos = ShaderPrintSymbol(Pos, _X_);
-			Pos = ShaderPrintSymbol(Pos, _T_);
-			Pos = ShaderPrintSymbol(Pos, _U_);
-			Pos = ShaderPrintSymbol(Pos, _R_);
-			Pos = ShaderPrintSymbol(Pos, _E_);
-			Pos = ShaderPrintSymbol(Pos, _SPC_);
-			Pos = ShaderPrintSymbol(Pos, _R_);
-			Pos = ShaderPrintSymbol(Pos, _E_);
-			Pos = ShaderPrintSymbol(Pos, _S_);
-			Pos = ShaderPrintSymbol(Pos, _SPC_);
-			Pos = ShaderPrintSymbol(Pos, _SPC_);
+			Pos = ShaderPrintText(Pos, TEXT(" Page Texture Res  "));
 			Pos = ShaderPrint(Pos, VirtualVoxel.PageTextureResolution.x);
 			Pos = ShaderPrint(Pos, VirtualVoxel.PageTextureResolution.y);
 			Pos = ShaderPrint(Pos, VirtualVoxel.PageTextureResolution.z);
@@ -210,198 +132,30 @@ void MainCS(uint3 DispatchThreadId : SV_DispatchThreadID)
 			// Total page indices stats
 			Pos = ShaderPrintNewline(Pos); Pos.x = StartPos.x;
 
-			Pos = ShaderPrintSymbol(Pos, _SPC_);
-			Pos = ShaderPrintSymbol(Pos, _T_);
-			Pos = ShaderPrintSymbol(Pos, _O_);
-			Pos = ShaderPrintSymbol(Pos, _T_);
-			Pos = ShaderPrintSymbol(Pos, _A_);
-			Pos = ShaderPrintSymbol(Pos, _L_);
-			Pos = ShaderPrintSymbol(Pos, _SPC_);
-			Pos = ShaderPrintSymbol(Pos, _A_);
-			Pos = ShaderPrintSymbol(Pos, _L_);
-			Pos = ShaderPrintSymbol(Pos, _L_);
-			Pos = ShaderPrintSymbol(Pos, _O_);
-			Pos = ShaderPrintSymbol(Pos, _C_);
-			Pos = ShaderPrintSymbol(Pos, _A_);
-			Pos = ShaderPrintSymbol(Pos, _T_);
-			Pos = ShaderPrintSymbol(Pos, _E_);
-			Pos = ShaderPrintSymbol(Pos, _D_);
-			Pos = ShaderPrintSymbol(Pos, _SPC_);
-			Pos = ShaderPrintSymbol(Pos, _P_);
-			Pos = ShaderPrintSymbol(Pos, _A_);
-			Pos = ShaderPrintSymbol(Pos, _G_);
-			Pos = ShaderPrintSymbol(Pos, _E_);
-			Pos = ShaderPrintSymbol(Pos, _SPC_);
-			Pos = ShaderPrintSymbol(Pos, _I_);
-			Pos = ShaderPrintSymbol(Pos, _N_);
-			Pos = ShaderPrintSymbol(Pos, _D_);
-			Pos = ShaderPrintSymbol(Pos, _E_);
-			Pos = ShaderPrintSymbol(Pos, _X_);
-			Pos = ShaderPrintSymbol(Pos, _SPC_);
-			Pos = ShaderPrintSymbol(Pos, _SPC_);
+			Pos = ShaderPrintText(Pos, TEXT(" Total Allocated Page Index  "));
 			Pos = ShaderPrint(Pos, TotalAllocatedPageCount);
 			Pos = ShaderPrintNewline(Pos); Pos.x = StartPos.x;
 
-			Pos = ShaderPrintSymbol(Pos, _SPC_);
-			Pos = ShaderPrintSymbol(Pos, _T_);
-			Pos = ShaderPrintSymbol(Pos, _O_);
-			Pos = ShaderPrintSymbol(Pos, _T_);
-			Pos = ShaderPrintSymbol(Pos, _A_);
-			Pos = ShaderPrintSymbol(Pos, _L_);
-			Pos = ShaderPrintSymbol(Pos, _SPC_);
-			Pos = ShaderPrintSymbol(Pos, _P_);
-			Pos = ShaderPrintSymbol(Pos, _A_);
-			Pos = ShaderPrintSymbol(Pos, _G_);
-			Pos = ShaderPrintSymbol(Pos, _E_);
-			Pos = ShaderPrintSymbol(Pos, _SPC_);
-			Pos = ShaderPrintSymbol(Pos, _I_);
-			Pos = ShaderPrintSymbol(Pos, _N_);
-			Pos = ShaderPrintSymbol(Pos, _D_);
-			Pos = ShaderPrintSymbol(Pos, _E_);
-			Pos = ShaderPrintSymbol(Pos, _X_);
-			Pos = ShaderPrintSymbol(Pos, _SPC_);
-			Pos = ShaderPrintSymbol(Pos, _SPC_);
-			Pos = ShaderPrintSymbol(Pos, _SPC_);
-			Pos = ShaderPrintSymbol(Pos, _SPC_);
-			Pos = ShaderPrintSymbol(Pos, _SPC_);
-			Pos = ShaderPrintSymbol(Pos, _SPC_);
-			Pos = ShaderPrintSymbol(Pos, _SPC_);
-			Pos = ShaderPrintSymbol(Pos, _SPC_);
-			Pos = ShaderPrintSymbol(Pos, _SPC_);
-			Pos = ShaderPrintSymbol(Pos, _SPC_);
-			Pos = ShaderPrintSymbol(Pos, _SPC_);
-			Pos = ShaderPrintSymbol(Pos, _SPC_);
+			Pos = ShaderPrintText(Pos, TEXT(" Total Page Index            "));
 			Pos = ShaderPrint(Pos, TotalPageIndexCount);
 			Pos = ShaderPrintNewline(Pos); Pos.x = StartPos.x;
 
-			Pos = ShaderPrintSymbol(Pos, _SPC_);
-			Pos = ShaderPrintSymbol(Pos, _P_);
-			Pos = ShaderPrintSymbol(Pos, _A_);
-			Pos = ShaderPrintSymbol(Pos, _G_);
-			Pos = ShaderPrintSymbol(Pos, _E_);
-			Pos = ShaderPrintSymbol(Pos, _SPC_);
-			Pos = ShaderPrintSymbol(Pos, _I_);
-			Pos = ShaderPrintSymbol(Pos, _N_);
-			Pos = ShaderPrintSymbol(Pos, _D_);
-			Pos = ShaderPrintSymbol(Pos, _E_);
-			Pos = ShaderPrintSymbol(Pos, _X_);
-			Pos = ShaderPrintSymbol(Pos, _SPC_);
-			Pos = ShaderPrintSymbol(Pos, _O_);
-			Pos = ShaderPrintSymbol(Pos, _C_);
-			Pos = ShaderPrintSymbol(Pos, _C_);
-			Pos = ShaderPrintSymbol(Pos, _U_);
-			Pos = ShaderPrintSymbol(Pos, _P_);
-			Pos = ShaderPrintSymbol(Pos, _A_);
-			Pos = ShaderPrintSymbol(Pos, _N_);
-			Pos = ShaderPrintSymbol(Pos, _C_);
-			Pos = ShaderPrintSymbol(Pos, _Y_);
-			Pos = ShaderPrintSymbol(Pos, _SPC_);
-			Pos = ShaderPrintSymbol(Pos, _SPC_);
-			Pos = ShaderPrintSymbol(Pos, _SPC_);
-			Pos = ShaderPrintSymbol(Pos, _SPC_);
-			Pos = ShaderPrintSymbol(Pos, _SPC_);
-			Pos = ShaderPrintSymbol(Pos, _SPC_);
-			Pos = ShaderPrintSymbol(Pos, _SPC_);
-			Pos = ShaderPrintSymbol(Pos, _SPC_);
+			Pos = ShaderPrintText(Pos, TEXT(" Page Index Occupancy         "));
 			Pos = ShaderPrint(Pos, TotalAllocatedPageCount / float(TotalPageIndexCount) * 100.f);
 			Pos = ShaderPrintNewline(Pos); Pos.x = StartPos.x;
 
 			// Total pages stats
 			Pos = ShaderPrintNewline(Pos); Pos.x = StartPos.x;
 
-			Pos = ShaderPrintSymbol(Pos, _SPC_);
-			Pos = ShaderPrintSymbol(Pos, _T_);
-			Pos = ShaderPrintSymbol(Pos, _O_);
-			Pos = ShaderPrintSymbol(Pos, _T_);
-			Pos = ShaderPrintSymbol(Pos, _A_);
-			Pos = ShaderPrintSymbol(Pos, _L_);
-			Pos = ShaderPrintSymbol(Pos, _SPC_);
-			Pos = ShaderPrintSymbol(Pos, _A_);
-			Pos = ShaderPrintSymbol(Pos, _L_);
-			Pos = ShaderPrintSymbol(Pos, _L_);
-			Pos = ShaderPrintSymbol(Pos, _O_);
-			Pos = ShaderPrintSymbol(Pos, _C_);
-			Pos = ShaderPrintSymbol(Pos, _A_);
-			Pos = ShaderPrintSymbol(Pos, _T_);
-			Pos = ShaderPrintSymbol(Pos, _E_);
-			Pos = ShaderPrintSymbol(Pos, _D_);
-			Pos = ShaderPrintSymbol(Pos, _SPC_);
-			Pos = ShaderPrintSymbol(Pos, _V_);
-			Pos = ShaderPrintSymbol(Pos, _A_);
-			Pos = ShaderPrintSymbol(Pos, _L_);
-			Pos = ShaderPrintSymbol(Pos, _I_);
-			Pos = ShaderPrintSymbol(Pos, _D_);
-			Pos = ShaderPrintSymbol(Pos, _SPC_);
-			Pos = ShaderPrintSymbol(Pos, _P_);
-			Pos = ShaderPrintSymbol(Pos, _A_);
-			Pos = ShaderPrintSymbol(Pos, _G_);
-			Pos = ShaderPrintSymbol(Pos, _E_);
-			Pos = ShaderPrintSymbol(Pos, _SPC_);
-			Pos = ShaderPrintSymbol(Pos, _SPC_);
+			Pos = ShaderPrintText(Pos, TEXT(" Total Allocated Valid Page  "));
 			Pos = ShaderPrint(Pos, TotalAllocatedValidPageCount);
 			Pos = ShaderPrintNewline(Pos); Pos.x = StartPos.x;
 
-			Pos = ShaderPrintSymbol(Pos, _SPC_);
-			Pos = ShaderPrintSymbol(Pos, _T_);
-			Pos = ShaderPrintSymbol(Pos, _O_);
-			Pos = ShaderPrintSymbol(Pos, _T_);
-			Pos = ShaderPrintSymbol(Pos, _A_);
-			Pos = ShaderPrintSymbol(Pos, _L_);
-			Pos = ShaderPrintSymbol(Pos, _SPC_);
-			Pos = ShaderPrintSymbol(Pos, _P_);
-			Pos = ShaderPrintSymbol(Pos, _A_);
-			Pos = ShaderPrintSymbol(Pos, _G_);
-			Pos = ShaderPrintSymbol(Pos, _E_);
-			Pos = ShaderPrintSymbol(Pos, _SPC_);
-			Pos = ShaderPrintSymbol(Pos, _SPC_);
-			Pos = ShaderPrintSymbol(Pos, _SPC_);
-			Pos = ShaderPrintSymbol(Pos, _SPC_);
-			Pos = ShaderPrintSymbol(Pos, _SPC_);
-			Pos = ShaderPrintSymbol(Pos, _SPC_);
-			Pos = ShaderPrintSymbol(Pos, _SPC_);
-			Pos = ShaderPrintSymbol(Pos, _SPC_);
-			Pos = ShaderPrintSymbol(Pos, _SPC_);
-			Pos = ShaderPrintSymbol(Pos, _SPC_);
-			Pos = ShaderPrintSymbol(Pos, _SPC_);
-			Pos = ShaderPrintSymbol(Pos, _SPC_);
-			Pos = ShaderPrintSymbol(Pos, _SPC_);
-			Pos = ShaderPrintSymbol(Pos, _SPC_);
-			Pos = ShaderPrintSymbol(Pos, _SPC_);
-			Pos = ShaderPrintSymbol(Pos, _SPC_);
-			Pos = ShaderPrintSymbol(Pos, _SPC_);
-			Pos = ShaderPrintSymbol(Pos, _SPC_);
+			Pos = ShaderPrintText(Pos, TEXT(" Total Page                  "));
 			Pos = ShaderPrint(Pos, TotalPageCount);
 			Pos = ShaderPrintNewline(Pos); Pos.x = StartPos.x;
 
-			Pos = ShaderPrintSymbol(Pos, _SPC_);
-			Pos = ShaderPrintSymbol(Pos, _P_);
-			Pos = ShaderPrintSymbol(Pos, _A_);
-			Pos = ShaderPrintSymbol(Pos, _G_);
-			Pos = ShaderPrintSymbol(Pos, _E_);
-			Pos = ShaderPrintSymbol(Pos, _SPC_);
-			Pos = ShaderPrintSymbol(Pos, _O_);
-			Pos = ShaderPrintSymbol(Pos, _C_);
-			Pos = ShaderPrintSymbol(Pos, _C_);
-			Pos = ShaderPrintSymbol(Pos, _U_);
-			Pos = ShaderPrintSymbol(Pos, _P_);
-			Pos = ShaderPrintSymbol(Pos, _A_);
-			Pos = ShaderPrintSymbol(Pos, _N_);
-			Pos = ShaderPrintSymbol(Pos, _C_);
-			Pos = ShaderPrintSymbol(Pos, _Y_);
-			Pos = ShaderPrintSymbol(Pos, _SPC_);
-			Pos = ShaderPrintSymbol(Pos, _SPC_);
-			Pos = ShaderPrintSymbol(Pos, _SPC_);
-			Pos = ShaderPrintSymbol(Pos, _SPC_);
-			Pos = ShaderPrintSymbol(Pos, _SPC_);
-			Pos = ShaderPrintSymbol(Pos, _SPC_);
-			Pos = ShaderPrintSymbol(Pos, _SPC_);
-			Pos = ShaderPrintSymbol(Pos, _SPC_);
-			Pos = ShaderPrintSymbol(Pos, _SPC_);
-			Pos = ShaderPrintSymbol(Pos, _SPC_);
-			Pos = ShaderPrintSymbol(Pos, _SPC_);
-			Pos = ShaderPrintSymbol(Pos, _SPC_);
-			Pos = ShaderPrintSymbol(Pos, _SPC_);
-			Pos = ShaderPrintSymbol(Pos, _SPC_);
+			Pos = ShaderPrintText(Pos, TEXT(" Page Occupancy              "));
 			Pos = ShaderPrint(Pos, TotalAllocatedValidPageCount / float(TotalPageCount) * 100.f);
 			Pos = ShaderPrintNewline(Pos); Pos.x = StartPos.x;
 		}
@@ -414,92 +168,24 @@ void MainCS(uint3 DispatchThreadId : SV_DispatchThreadID)
 			float2 StartPos = float2(50.0f, 420.0f + float(MacroGroupId) * 100.0f) / OutputResolution;
 			float2 Pos = StartPos;
 
-			Pos = ShaderPrintSymbol(Pos, _M_);
-			Pos = ShaderPrintSymbol(Pos, _A_);
-			Pos = ShaderPrintSymbol(Pos, _C_);
-			Pos = ShaderPrintSymbol(Pos, _R_);
-			Pos = ShaderPrintSymbol(Pos, _O_);
-			Pos = ShaderPrintSymbol(Pos, _SPC_);
-			Pos = ShaderPrintSymbol(Pos, _G_);
-			Pos = ShaderPrintSymbol(Pos, _R_);
-			Pos = ShaderPrintSymbol(Pos, _O_);
-			Pos = ShaderPrintSymbol(Pos, _U_);
-			Pos = ShaderPrintSymbol(Pos, _P_);
-			Pos = ShaderPrintSymbol(Pos, _SPC_);
-			Pos = ShaderPrintSymbol(Pos, _SPC_);
-			Pos = ShaderPrintSymbol(Pos, _SPC_);
-			Pos = ShaderPrintSymbol(Pos, _SPC_);
+			Pos = ShaderPrintText(Pos, TEXT("Macro Group    "));
 			Pos = ShaderPrint(Pos, MacroGroupId);
 			Pos = ShaderPrintNewline(Pos); Pos.x = StartPos.x;
 
-			Pos = ShaderPrintSymbol(Pos, _SPC_);
-			Pos = ShaderPrintSymbol(Pos, _V_);
-			Pos = ShaderPrintSymbol(Pos, _I_);
-			Pos = ShaderPrintSymbol(Pos, _R_);
-			Pos = ShaderPrintSymbol(Pos, _T_);
-			Pos = ShaderPrintSymbol(Pos, _U_);
-			Pos = ShaderPrintSymbol(Pos, _A_);
-			Pos = ShaderPrintSymbol(Pos, _L_);
-			Pos = ShaderPrintSymbol(Pos, _SPC_);
-			Pos = ShaderPrintSymbol(Pos, _R_);
-			Pos = ShaderPrintSymbol(Pos, _E_);
-			Pos = ShaderPrintSymbol(Pos, _S_);
-			Pos = ShaderPrintSymbol(Pos, _SPC_);
-			Pos = ShaderPrintSymbol(Pos, _SPC_);
-			Pos = ShaderPrintSymbol(Pos, _SPC_);
-			Pos = ShaderPrintSymbol(Pos, _SPC_);
-			Pos = ShaderPrintSymbol(Pos, _SPC_);
-			Pos = ShaderPrintSymbol(Pos, _SPC_);
-			Pos = ShaderPrintSymbol(Pos, _SPC_);
+			Pos = ShaderPrintText(Pos, TEXT(" Virtual Res       "));
 			Pos = ShaderPrint(Pos, NodeDesc.VirtualResolution.x);
 			Pos = ShaderPrint(Pos, NodeDesc.VirtualResolution.y);
 			Pos = ShaderPrint(Pos, NodeDesc.VirtualResolution.z);
 			Pos = ShaderPrintNewline(Pos); Pos.x = StartPos.x;
 
-			Pos = ShaderPrintSymbol(Pos, _SPC_);
-			Pos = ShaderPrintSymbol(Pos, _P_);
-			Pos = ShaderPrintSymbol(Pos, _A_);
-			Pos = ShaderPrintSymbol(Pos, _G_);
-			Pos = ShaderPrintSymbol(Pos, _E_);
-			Pos = ShaderPrintSymbol(Pos, _SPC_);
-			Pos = ShaderPrintSymbol(Pos, _I_);
-			Pos = ShaderPrintSymbol(Pos, _N_);
-			Pos = ShaderPrintSymbol(Pos, _D_);
-			Pos = ShaderPrintSymbol(Pos, _E_);
-			Pos = ShaderPrintSymbol(Pos, _X_);
-			Pos = ShaderPrintSymbol(Pos, _SPC_);
-			Pos = ShaderPrintSymbol(Pos, _R_);
-			Pos = ShaderPrintSymbol(Pos, _E_);
-			Pos = ShaderPrintSymbol(Pos, _S_);
-			Pos = ShaderPrintSymbol(Pos, _SPC_);
-			Pos = ShaderPrintSymbol(Pos, _SPC_);
-			Pos = ShaderPrintSymbol(Pos, _SPC_);
-			Pos = ShaderPrintSymbol(Pos, _SPC_);
+			Pos = ShaderPrintText(Pos, TEXT(" Page Index Res    "));
 			Pos = ShaderPrint(Pos, NodeDesc.PageIndexResolution.x);
 			Pos = ShaderPrint(Pos, NodeDesc.PageIndexResolution.y);
 			Pos = ShaderPrint(Pos, NodeDesc.PageIndexResolution.z);
-			Pos = ShaderPrintSymbol(Pos, _SPC_);
+			Pos = ShaderPrintText(Pos, TEXT(" "));
 			Pos = ShaderPrintNewline(Pos); Pos.x = StartPos.x;
 
-			Pos = ShaderPrintSymbol(Pos, _SPC_);
-			Pos = ShaderPrintSymbol(Pos, _P_);
-			Pos = ShaderPrintSymbol(Pos, _A_);
-			Pos = ShaderPrintSymbol(Pos, _G_);
-			Pos = ShaderPrintSymbol(Pos, _E_);
-			Pos = ShaderPrintSymbol(Pos, _SPC_);
-			Pos = ShaderPrintSymbol(Pos, _I_);
-			Pos = ShaderPrintSymbol(Pos, _N_);
-			Pos = ShaderPrintSymbol(Pos, _D_);
-			Pos = ShaderPrintSymbol(Pos, _E_);
-			Pos = ShaderPrintSymbol(Pos, _X_);
-			Pos = ShaderPrintSymbol(Pos, _SPC_);
-			Pos = ShaderPrintSymbol(Pos, _O_);
-			Pos = ShaderPrintSymbol(Pos, _F_);
-			Pos = ShaderPrintSymbol(Pos, _F_);
-			Pos = ShaderPrintSymbol(Pos, _S_);
-			Pos = ShaderPrintSymbol(Pos, _E_);
-			Pos = ShaderPrintSymbol(Pos, _T_);
-			Pos = ShaderPrintSymbol(Pos, _SPC_);
+			Pos = ShaderPrintText(Pos, TEXT(" Page Index Offset "));
 			Pos = ShaderPrint(Pos, NodeDesc.PageIndexOffset);
 			Pos = ShaderPrintNewline(Pos); Pos.x = StartPos.x;
 
@@ -617,36 +303,15 @@ void MainCS(uint3 DispatchThreadId : SV_DispatchThreadID)
 			float2 StartPos = float2(50.0f, 50.0f) / OutputResolution;
 			float2 Pos = StartPos;
 
-			Pos = ShaderPrintSymbol(Pos, _X_);
-			Pos = ShaderPrintSymbol(Pos, _Y_);
-			Pos = ShaderPrintSymbol(Pos, _SPC_);
+			Pos = ShaderPrintText(Pos, TEXT("X, Y "));
 			Pos = ShaderPrint(Pos, GetCursorPos());
 			Pos = ShaderPrintNewline(Pos); Pos.x = StartPos.x;
 
-			Pos = ShaderPrintSymbol(Pos, _T_);
-			Pos = ShaderPrintSymbol(Pos, _DOT_);
-			Pos = ShaderPrintSymbol(Pos, _C_);
-			Pos = ShaderPrintSymbol(Pos, _O_);
-			Pos = ShaderPrintSymbol(Pos, _U_);
-			Pos = ShaderPrintSymbol(Pos, _N_);
-			Pos = ShaderPrintSymbol(Pos, _T_);
-			Pos = ShaderPrintSymbol(Pos, _SPC_);
+			Pos = ShaderPrintText(Pos, TEXT("T.Count "));
 			Pos = ShaderPrint(Pos, TraversalResult.HairCount);
 			Pos = ShaderPrintNewline(Pos); Pos.x = StartPos.x;
 			
-			Pos = ShaderPrintSymbol(Pos, _T_);
-			Pos = ShaderPrintSymbol(Pos, _DOT_);
-			Pos = ShaderPrintSymbol(Pos, _V_);
-			Pos = ShaderPrintSymbol(Pos, _I_);
-			Pos = ShaderPrintSymbol(Pos, _S_);
-			Pos = ShaderPrintSymbol(Pos, _I_);
-			Pos = ShaderPrintSymbol(Pos, _B_);
-			Pos = ShaderPrintSymbol(Pos, _I_);
-			Pos = ShaderPrintSymbol(Pos, _L_);
-			Pos = ShaderPrintSymbol(Pos, _I_);
-			Pos = ShaderPrintSymbol(Pos, _T_);
-			Pos = ShaderPrintSymbol(Pos, _Y_);
-			Pos = ShaderPrintSymbol(Pos, _SPC_);
+			Pos = ShaderPrintText(Pos, TEXT("T.Visibility "));
 			Pos = ShaderPrint(Pos, TraversalResult.Visibility);
 			Pos = ShaderPrintNewline(Pos); Pos.x = StartPos.x;
 		}
diff --git a/Engine/Shaders/Private/ShaderPrintCommon.ush b/Engine/Shaders/Private/ShaderPrintCommon.ush
--- a/Engine/Shaders/Private/ShaderPrintCommon.ush
+++ b/Engine/Shaders/Private/ShaderPrintCommon.ush
@@ -88,6 +88,11 @@ float2 ShaderPrintSymbol(in float2 ScreenPos, in int Symbol)
 	return ScreenPos;
 }
 
+// Function for reading global TEXT string. 
+// The data and function are generated by the shader compiler.
+float2 ShaderPrintText(float2 Pos, uint InTextEntry);
+GENERATED_SHADER_PRINT
+
 void ShaderPrintSymbol(in int symbol)
 {
 	ShaderPrintCursorPos = ShaderPrintSymbol(ShaderPrintCursorPos, symbol);
diff --git a/Engine/Source/Developer/ShaderCompilerCommon/Private/ShaderCompilerCommon.cpp b/Engine/Source/Developer/ShaderCompilerCommon/Private/ShaderCompilerCommon.cpp
--- a/Engine/Source/Developer/ShaderCompilerCommon/Private/ShaderCompilerCommon.cpp
+++ b/Engine/Source/Developer/ShaderCompilerCommon/Private/ShaderCompilerCommon.cpp
@@ -229,18 +229,18 @@ const TCHAR* FindNextNonWhitespace(const TCHAR* StringPtr)
 	}
 }
 
-const TCHAR* FindMatchingClosingBrace(const TCHAR* OpeningBracePtr)
+const TCHAR* FindMatchingBlock(const TCHAR* OpeningCharPtr, char OpenChar, char CloseChar)
 {
-	const TCHAR* SearchPtr = OpeningBracePtr;
+	const TCHAR* SearchPtr = OpeningCharPtr;
 	int32 Depth = 0;
 
 	while (*SearchPtr)
 	{
-		if (*SearchPtr == '{')
+		if (*SearchPtr == OpenChar)
 		{
 			Depth++;
 		}
-		else if (*SearchPtr == '}')
+		else if (*SearchPtr == CloseChar)
 		{
 			if (Depth == 0)
 			{
@@ -254,6 +254,8 @@ const TCHAR* FindMatchingClosingBrace(const TCHAR* OpeningBracePtr)
 
 	return nullptr;
 }
+const TCHAR* FindMatchingClosingBrace(const TCHAR* OpeningCharPtr)			{ return FindMatchingBlock(OpeningCharPtr, '{', '}'); };
+const TCHAR* FindMatchingClosingParenthesis(const TCHAR* OpeningCharPtr)	{ return FindMatchingBlock(OpeningCharPtr, '(', ')'); };
 
 // See MSDN HLSL 'Symbol Name Restrictions' doc
 inline bool IsValidHLSLIdentifierCharacter(TCHAR Char)
@@ -1127,6 +1129,7 @@ void FShaderParameterParser::ExtractFileAndLine(int32 PragamLineoffset, int32 Li
 	OutLine = FString::FromInt(FinalLine);
 }
 
+void TransformStringIntoCharacterArray(FString& PreprocessedShaderSource);
 
 // The cross compiler doesn't yet support struct initializers needed to construct static structs for uniform buffers
 // Replace all uniform buffer struct member references (View.WorldToClip) with a flattened name that removes the struct dependency (View_WorldToClip)
@@ -1215,8 +1218,274 @@ void RemoveUniformBuffersFromSource(const FShaderCompilerEnvironment& Environmen
 			SearchPtr = FindNextUniformBufferReference(SearchPtr + UniformBufferAccessString.Len(), *UniformBufferName, UniformBufferName.Len());
 		}
 	}
+
+	// Process TEXT macro.
+	// Even this processing is not directly related to the Uniform buffer removal, this processing is 
+	// called here as this function is called by all the platform specific compiler, and avoid to update 
+	// all compiler
+	TransformStringIntoCharacterArray(PreprocessedShaderSource);
 }
 
+///////////////////////////////////////////////////////////////////////////////////////////////////////////////////////////////////////
+// Process TEXT() macro to convert them into GPU MiniFont characters
+
+FString ParseText(const TCHAR* StartPtr, const TCHAR*& EndPtr)
+{
+	const TCHAR* OpeningBracePtr = FCString::Strstr(StartPtr, TEXT("("));
+	check(OpeningBracePtr);
+
+	const TCHAR* ClosingBracePtr = FindMatchingClosingParenthesis(OpeningBracePtr + 1);
+	check(ClosingBracePtr);
+
+	FString Out;
+	if (OpeningBracePtr && ClosingBracePtr)
+	{
+		const TCHAR* CurrPtr = OpeningBracePtr;
+		do
+		{
+			Out += *CurrPtr;
+			CurrPtr++;
+		} while (CurrPtr != ClosingBracePtr+1);
+	}
+	EndPtr = ClosingBracePtr;
+	return Out;
+}
+
+void ConvertTextToMiniFontCharacter(const FString& InText, FString& OutText, FString& OutEncodedText)
+{
+	// From MiniFontCommon.ush - CharacterIDs for the MiniFont
+	// Instead of replicating these values here, we could add them as DEFINE, to insure CPP/Shader are in sync.
+	// TODO: change this to proper ASCII code, and use an ASCII table on the shader side
+	const int32 _0_ = 0;
+	const int32 _1_ = 1;
+	const int32 _2_ = 2;
+	const int32 _3_ = 3;
+	const int32 _4_ = 4;
+	const int32 _5_ = 5;
+	const int32 _6_ = 6;
+	const int32 _7_ = 7;
+	const int32 _8_ = 8;
+	const int32 _9_ = 9;
+	const int32 _A_ = 10;
+	const int32 _B_ = 11;
+	const int32 _C_ = 12;
+	const int32 _D_ = 13;
+	const int32 _E_ = 14;
+	const int32 _F_ = 15;
+	const int32 _G_ = 16;
+	const int32 _H_ = 17;
+	const int32 _I_ = 18;
+	const int32 _J_ = 19;
+	const int32 _K_ = 20;
+	const int32 _L_ = 21;
+	const int32 _M_ = 22;
+	const int32 _N_ = 23;
+	const int32 _O_ = 24;
+	const int32 _P_ = 25;
+	const int32 _Q_ = 26;
+	const int32 _R_ = 27;
+	const int32 _S_ = 28;
+	const int32 _T_ = 29;
+	const int32 _U_ = 30;
+	const int32 _V_ = 31;
+	const int32 _W_ = 32;
+	const int32 _X_ = 33;
+	const int32 _Y_ = 34;
+	const int32 _Z_ = 35;
+	const int32 _MINUS_ = 59;
+	const int32 _COMMA_ = 60;
+	const int32 _DOT_ = 61;
+	const int32 _PLUS_ = 62;
+	const int32 _SPC_ = 63;
+
+	// Trim text to valid character only
+	FString ValidText;
+	{
+		FString Text = InText.ToUpper();
+		const uint32 CharCount = Text.Len();
+
+		ValidText.Reserve(CharCount);
+		for (uint32 CharIt = 0; CharIt < CharCount; ++CharIt)
+		{
+			const char C = Text[CharIt];
+			if ((C >= '0' && C <= '9') || (C >= 'A' && C <= 'Z') || (C == '-') || (C == ',') || (C == '.') || (C == '+') || (C == ' '))
+			{
+				ValidText += C;
+			}
+		}
+	}
+
+	// Translate character into MiniFont encoding
+	{
+		const uint32 CharCount = ValidText.Len();
+		OutEncodedText.Reserve(CharCount * 3); // ~2 digits per character + a comma
+		OutText = ValidText;
+
+		for (uint32 CharIt = 0; CharIt < CharCount; ++CharIt)
+		{
+			const char C = ValidText[CharIt];
+			if ((C >= '0' && C <= '9'))
+			{
+				const int32 Offset = C - '0';
+				OutEncodedText.AppendInt(_0_ + Offset);
+			}
+			if ((C >= 'A' && C <= 'Z'))
+			{
+				const int32 Offset = C - 'A';
+				OutEncodedText.AppendInt(_A_ + Offset);
+			}
+			else if (C == '-')
+			{
+				OutEncodedText.AppendInt(_MINUS_);
+			}
+			else if (C == ',')
+			{
+				OutEncodedText.AppendInt(_COMMA_);
+			}
+			else if (C == '.')
+			{
+				OutEncodedText.AppendInt(_DOT_);
+			}
+			else if (C == '+')
+			{
+				OutEncodedText.AppendInt(_PLUS_);
+			}
+			else if (C == ' ')
+			{
+				OutEncodedText.AppendInt(_SPC_);
+			}
+			if (CharIt + 1 != CharCount)
+			{
+				OutEncodedText += ',';
+			}
+		}
+	}
+}
+
+// Simple token matching and expansion to replace TEXT macro into supported character string
+void TransformStringIntoCharacterArray(FString& PreprocessedShaderSource)
+{
+	struct FTextEntry
+	{
+		uint32  Index;
+		uint32  Hash;
+		uint32  Offset;
+		FString SourceText;
+		FString ConvertedText;
+		FString EncodedText;
+	};
+	TArray<FTextEntry> Entries;
+
+	// 1. Find all TEXT strings
+	// 2. Add a text entry
+	// 3. Replace TEXT by its entry number
+	uint32 GlobalCount = 0;
+	{
+		const TCHAR* TextIdentifier = TEXT("TEXT(\"");
+		const TCHAR* SearchPtr = FCString::Strstr(&PreprocessedShaderSource[0], TextIdentifier);
+		while (SearchPtr)
+		{
+			const TCHAR* EndPtr = nullptr;
+			FString Text = ParseText(SearchPtr, EndPtr);
+			if (EndPtr)
+			{
+				// Trim enclosing
+				Text.RemoveFromEnd("\")");
+				Text.RemoveFromStart("(\"");
+
+				// Register entry and convert text
+				const uint32 EntryIndex = Entries.Num();
+				uint32 ValidCharCount = 0;
+				FTextEntry& Entry = Entries.AddDefaulted_GetRef();
+				Entry.Index			= EntryIndex;
+				Entry.Offset		= GlobalCount;
+				Entry.SourceText	= Text;
+			    ConvertTextToMiniFontCharacter(Entry.SourceText, Entry.ConvertedText, Entry.EncodedText);
+				Entry.Hash			= CityHash32((const char*)&Entry.SourceText.GetCharArray(), sizeof(FString::ElementType) * Entry.SourceText.Len());
+
+				GlobalCount += Entry.ConvertedText.Len();
+
+				// Replace string
+				const TCHAR* StartPtr = &PreprocessedShaderSource[0];
+				const uint32 StartIndex = SearchPtr - StartPtr;
+				const uint32 CharCount = (EndPtr - SearchPtr) + 1;
+				PreprocessedShaderSource.RemoveAt(StartIndex, CharCount);
+				PreprocessedShaderSource.InsertAt(StartIndex, FString::FromInt(EntryIndex));
+
+				// Update SearchPtr, as PreprocessedShaderSource has been modified, and its memory could have been reallocated, causing SearchPtr to be invalid.
+				SearchPtr = &PreprocessedShaderSource[0] + (StartIndex + CharCount);
+			}
+			SearchPtr = FCString::Strstr(SearchPtr, TextIdentifier);
+		}
+	}
+
+	// 4. Write a global struct containing all the entries
+	// 5. Write the function for fetching character for a given entry index
+	const uint32 EntryCount = Entries.Num();
+	FString TextChars;
+	if (EntryCount)
+	{
+		// 1. Encoded character for each text entry within a single global char array
+		TextChars = FString::Printf(TEXT("static const uint TEXT_CHARS[%d] = {\n"), GlobalCount);
+		for (FTextEntry& Entry : Entries)
+		{
+			TextChars += FString::Printf(TEXT("\t%s%s // %d: \"%s\"\n"), *Entry.EncodedText, Entry.Index < EntryCount - 1 ? TEXT(",") : TEXT(""), Entry.Index, * Entry.SourceText);
+		}
+		TextChars += TEXT("};\n\n");
+
+		// 2. Offset within the global array
+		TextChars += FString::Printf(TEXT("static const uint TEXT_OFFSETS[%d] = {\n"), EntryCount+1);
+		for (FTextEntry& Entry : Entries)
+		{
+			TextChars += FString::Printf(TEXT("\t%d, // %d: \"%s\"\n"), Entry.Offset, Entry.Index, *Entry.SourceText);
+		}
+		TextChars += FString::Printf(TEXT("\t%d // end\n"), GlobalCount);
+		TextChars += TEXT("};\n\n");
+
+		// 3. Entry hashes
+		TextChars += TEXT("// Hashes are computed using the CityHash32 function\n");
+		TextChars += FString::Printf(TEXT("static const uint TEXT_HASHES[%d] = {\n"), EntryCount);
+		for (FTextEntry& Entry : Entries)
+		{
+			TextChars += FString::Printf(TEXT("\t0x%x%s // %d: \"%s\"\n"), Entry.Hash, Entry.Index < EntryCount - 1 ? TEXT(",") : TEXT(""), Entry.Index, * Entry.SourceText);
+		}
+		TextChars += TEXT("};\n\n");
+
+		// Function for reading global TEXT string
+		TextChars += TEXT("float2 ShaderPrintText(float2 Pos, uint InTextEntry)\n");
+		TextChars += TEXT("{\n");
+		TextChars += TEXT("\tuint Begin = TEXT_OFFSETS[InTextEntry];\n");
+		TextChars += TEXT("\tuint End = TEXT_OFFSETS[InTextEntry + 1];\n");
+		TextChars += TEXT("\tfor (uint i = Begin; i < End; ++i)\n");
+		TextChars += TEXT("\t{\n");
+		TextChars += TEXT("\t\tPos = ShaderPrintSymbol(Pos, TEXT_CHARS[i]);\n");
+		TextChars += TEXT("\t}\n");
+		TextChars += TEXT("\treturn Pos;\n");
+		TextChars += TEXT("}\n");
+	}
+	else
+	{	
+		TextChars = TEXT("float2 ShaderPrintText(float2 Pos, uint InTextEntry) { return Pos; }\n");
+	}
+	
+	// 6. Insert global struct data + print function
+	{
+		const TCHAR* InsertToken = TEXT("GENERATED_SHADER_PRINT");
+		const TCHAR* SearchPtr = FCString::Strstr(&PreprocessedShaderSource[0], InsertToken);
+		if (SearchPtr)
+		{
+			// Replace string
+			const TCHAR* StartPtr = &PreprocessedShaderSource[0];
+			const uint32 StartIndex = SearchPtr - StartPtr;
+			const uint32 CharCount = FCString::Strlen(InsertToken);
+			PreprocessedShaderSource.RemoveAt(StartIndex, CharCount);
+			PreprocessedShaderSource.InsertAt(StartIndex, TextChars);
+		}
+	}
+}
+
+///////////////////////////////////////////////////////////////////////////////////////////////////////////////////////////////////////
+
 FString CreateShaderCompilerWorkerDirectCommandLine(const FShaderCompilerInput& Input, uint32 CCFlags)
 {
 	FString Text(TEXT("-directcompile -format="));
