Fix for velocity not working on ThinTranslucent shading model.

diff --git a/Engine/Shaders/Private/MaterialTemplate.ush b/Engine/Shaders/Private/MaterialTemplate.ush
--- a/Engine/Shaders/Private/MaterialTemplate.ush
+++ b/Engine/Shaders/Private/MaterialTemplate.ush
@@ -2742,7 +2742,9 @@ void GetMaterialClippingShadowDepth(FMaterialPixelParameters Parameters, FPixelM
 	#endif
 }
 
-void GetMaterialClippingVelocity(FMaterialPixelParameters Parameters, FPixelMaterialInputs PixelMaterialInputs)
+// The material blending mode alone is not enough information. We also need to know if the material is ThinTranslucent because in this case we should never clip.
+// Indeed, ThinTranslucent surfaces always have a coverage of 1 and opacity represent the coverage of the opaque lit material sitting on top of the translucent one.
+void GetMaterialClippingVelocity(FMaterialPixelParameters Parameters, FPixelMaterialInputs PixelMaterialInputs, bool bIsThinTranslucent)
 {
 	ClipLODTransition(Parameters);
 	#if MATERIALBLENDING_MASKED && MATERIAL_DITHER_OPACITY_MASK
@@ -2750,7 +2752,7 @@ void GetMaterialClippingVelocity(FMaterialPixelParameters Parameters, FPixelMate
 	#elif MATERIALBLENDING_MASKED
 		clip(GetMaterialMask(PixelMaterialInputs)); 
 	#elif MATERIALBLENDING_TRANSLUCENT || MATERIALBLENDING_ADDITIVE || MATERIALBLENDING_MODULATE
-		clip(GetMaterialOpacity(PixelMaterialInputs) - 1.0 / 255.0 - GetMaterialOpacityMaskClipValue());
+		clip(bIsThinTranslucent ? 1.0f : GetMaterialOpacity(PixelMaterialInputs) - 1.0 / 255.0 - GetMaterialOpacityMaskClipValue());
 	#endif
 }
 
diff --git a/Engine/Shaders/Private/VelocityShader.usf b/Engine/Shaders/Private/VelocityShader.usf
--- a/Engine/Shaders/Private/VelocityShader.usf
+++ b/Engine/Shaders/Private/VelocityShader.usf
@@ -143,6 +143,10 @@ void MainVertexShader(
  * Pixel Shader
  *============================================================================*/
 
+#ifndef VELOCITY_THIN_TRANSLUCENT_MODE
+#define VELOCITY_THIN_TRANSLUCENT_MODE 0
+#endif
+
 void MainPixelShader(
 	in INPUT_POSITION_QUALIFIERS float4 SvPosition : SV_Position,
 	FVelocityInterpsVSToPS Input
@@ -164,7 +168,7 @@ void MainPixelShader(
 	FMaterialPixelParameters MaterialParameters = GetMaterialPixelParameters(Input.FactoryInterpolants, SvPosition);
 	FPixelMaterialInputs PixelMaterialInputs;
 	CalcMaterialParameters(MaterialParameters, PixelMaterialInputs, SvPosition, bIsFrontFace);
-	GetMaterialClippingVelocity(MaterialParameters, PixelMaterialInputs);
+	GetMaterialClippingVelocity(MaterialParameters, PixelMaterialInputs, VELOCITY_THIN_TRANSLUCENT_MODE > 0);
 
 #if OUTPUT_PIXEL_DEPTH_OFFSET
 	ApplyPixelDepthOffsetToMaterialParameters(MaterialParameters, PixelMaterialInputs, OutDepth);
diff --git a/Engine/Source/Runtime/Renderer/Private/VelocityRendering.cpp b/Engine/Source/Runtime/Renderer/Private/VelocityRendering.cpp
--- a/Engine/Source/Runtime/Renderer/Private/VelocityRendering.cpp
+++ b/Engine/Source/Runtime/Renderer/Private/VelocityRendering.cpp
@@ -124,6 +124,9 @@ public:
 	{
 		FMaterialShader::ModifyCompilationEnvironment(Parameters, OutEnvironment);
 		OutEnvironment.SetRenderTargetOutputFormat(0, PF_A16B16G16R16);
+
+		// We support velocity on thin trnaslucent only with masking, and only if the material is only made of thin translucent shading model.
+		OutEnvironment.SetDefine(TEXT("VELOCITY_THIN_TRANSLUCENT_MODE"), Parameters.MaterialParameters.ShadingModels.HasOnlyShadingModel(MSM_ThinTranslucent));
 	}
 
 	FVelocityPS() = default;
