* Move INDIRECT_GROUP_COUNT_X code into HairStrandsClusterCommon
* Add helper functions for computing (culled) vertex fetch index.

This insures that vertex index are computed the same way and with the same group count when hair culling is enabled.

diff --git a/Engine/Shaders/Private/HairStrands/HairStrandsClusterCommon.ush b/Engine/Shaders/Private/HairStrands/HairStrandsClusterCommon.ush
--- a/Engine/Shaders/Private/HairStrands/HairStrandsClusterCommon.ush
+++ b/Engine/Shaders/Private/HairStrands/HairStrandsClusterCommon.ush
@@ -2,8 +2,33 @@
 
 #include "../Common.ush"
 
+////////////////////////////////////////////////////////////////////////////////
+// Defines
+
+// Define the Group count along X, when the 1D number of groups exceed the 65k limits. 
+// In such a case, the 1D dispatch is wrapped into a 2D dispatch, whose the X dimention 
+// has a size of INDIRECT_GROUP_COUNT_X, with a total of INDIRECT_GROUP_COUNT_X * GROUP_SIZE 
+// thread per line
+#define INDIRECT_GROUP_COUNT_X 16
+
+// Max number of discrete LOD that a hair group can have
 #define MAX_HAIR_LOD 8
 
+////////////////////////////////////////////////////////////////////////////////
+// Utils
+
+// When culling is disabled: return the vertex index at which the control point should be fetched.
+uint GetHairStrandsVertexFetchIndex(uint2 InDispatchThreadId, uint InGroupSize, uint InGroupCountX)
+{
+	return InDispatchThreadId.x + InDispatchThreadId.y * InGroupSize * InGroupCountX;
+}
+
+// When culling is enabled: return the vertex index at which the control point should be fetched.
+uint GetHairStrandsVertexFetchIndex_Culled(uint2 InDispatchThreadId, uint InGroupSize)
+{
+	return InDispatchThreadId.x + InDispatchThreadId.y * InGroupSize * INDIRECT_GROUP_COUNT_X;
+}
+
 ////////////////////////////////////////////////////////////////////////////////
 // Store info about all LODS of a cluster (total LOD count, screen size, ...)
 struct FHairClusterInfo
diff --git a/Engine/Shaders/Private/HairStrands/HairStrandsClusterCulling.usf b/Engine/Shaders/Private/HairStrands/HairStrandsClusterCulling.usf
--- a/Engine/Shaders/Private/HairStrands/HairStrandsClusterCulling.usf
+++ b/Engine/Shaders/Private/HairStrands/HairStrandsClusterCulling.usf
@@ -634,7 +634,7 @@ void MainClusterCullingCS(uint3 DispatchThreadId : SV_DispatchThreadID)
 #endif
 
 /////////////////////////////////////////////////////////////////////////////////////////////////////////////////////////////
-/// 0. Indirect dispatch buffer preparation for compute raster visibility pass
+/// 0. Indirect dispatch buffer preparation for compute raster visibility/voxel pass
 #ifdef SHADER_PREPAREINDIRECTDISPATCH
 
 RWBuffer<uint> DrawIndirectBuffer; 
@@ -649,15 +649,14 @@ void MainClusterCullingPrepareIndirectDispatchCS(
 	if (GroupId.x == 0 && GroupThreadId.x == 0)
 	{
 		// See HairStrandsVisibilityRasterCompute.usf
-		const uint RasterComputeGroupSize = 32;
-		const uint RasterComputeGroupCountX = 16;
+		const uint RasterComputeGroupSize = 32; // TODO change this so that we can have dispatch for different GroupSize
 
 		// DrawIndirectBuffer contains the total number of vertex count needed for drawing quad triangle strips, so 6 vertices per quad / control point
 		const uint VertexCount    = DrawIndirectBuffer[0];
 		const uint ControlPoint   = VertexCount / 6;
 		const uint DispatchCount  = DivideAndRoundUp(ControlPoint, RasterComputeGroupSize);
-		const uint DispatchCountX = DispatchCount >= RasterComputeGroupCountX ? RasterComputeGroupCountX : (DispatchCount % RasterComputeGroupCountX);
-		const uint DispatchCountY = DivideAndRoundUp(DispatchCount, RasterComputeGroupCountX);
+		const uint DispatchCountX = DispatchCount >= INDIRECT_GROUP_COUNT_X ? INDIRECT_GROUP_COUNT_X : (DispatchCount % INDIRECT_GROUP_COUNT_X);
+		const uint DispatchCountY = DivideAndRoundUp(DispatchCount, INDIRECT_GROUP_COUNT_X);
 
 		DispatchIndirectBuffer[0] = DispatchCountX;
 		DispatchIndirectBuffer[1] = DispatchCountY;
diff --git a/Engine/Shaders/Private/HairStrands/HairStrandsGuideDeform.usf b/Engine/Shaders/Private/HairStrands/HairStrandsGuideDeform.usf
--- a/Engine/Shaders/Private/HairStrands/HairStrandsGuideDeform.usf
+++ b/Engine/Shaders/Private/HairStrands/HairStrandsGuideDeform.usf
@@ -1,6 +1,7 @@
 // Copyright Epic Games, Inc. All Rights Reserved.
 
 #include "../Common.ush"
+#include "HairStrandsClusterCommon.ush"
 #include "HairStrandsMeshProjectionCommon.ush"
 #include "HairStrandsMeshRBF.ush"
 
@@ -93,7 +94,7 @@ float3 DisplacePosition(uint VertexIndex, float3 Pos)
 [numthreads(GROUP_SIZE, 1, 1)]
 void MainCS(uint2 DispatchThreadId : SV_DispatchThreadID)
 {
-	const uint VertexIndex = DispatchThreadId.x + DispatchThreadId.y * DispatchCountX * GROUP_SIZE;
+	const uint VertexIndex = GetHairStrandsVertexFetchIndex(DispatchThreadId, GROUP_SIZE, DispatchCountX);
 	if (VertexIndex < VertexCount)
 	{
 		const uint4  Data 				= SimRestPosePositionBuffer[VertexIndex];
diff --git a/Engine/Shaders/Private/HairStrands/HairStrandsInterpolation.usf b/Engine/Shaders/Private/HairStrands/HairStrandsInterpolation.usf
--- a/Engine/Shaders/Private/HairStrands/HairStrandsInterpolation.usf
+++ b/Engine/Shaders/Private/HairStrands/HairStrandsInterpolation.usf
@@ -10,9 +10,6 @@
 ///////////////////////////////////////////////////////////////////////////////////////////////////
 // Common
 
-// These number need to match HairStrandsClusterCulling.usf
-#define GROUP_COUNT_X 16
-
 ///////////////////////////////////////////////////////////////////////////////////////////////////
 // Hair interpolation 
 
@@ -185,11 +182,11 @@ void MainCS(uint2 DispatchThreadId : SV_DispatchThreadID)
 	const float3 SimHairPositionOffset = OutSimHairPositionOffsetBuffer[0].xyz;
 	const float3 OutRenHairPositionOffset = OutRenHairPositionOffsetBuffer[0].xyz;
 
-	uint VertexIndex = DispatchThreadId.x + DispatchThreadId.y * DispatchCountX * GROUP_SIZE;
+	uint VertexIndex = GetHairStrandsVertexFetchIndex(DispatchThreadId, GROUP_SIZE, DispatchCountX);
 #if PERMUTATION_CULLING
 	if (HairStrandsVF_bIsCullingEnable)
 	{
-		const uint FetchIndex = DispatchThreadId.x + DispatchThreadId.y * GROUP_SIZE * GROUP_COUNT_X;
+		const uint FetchIndex = GetHairStrandsVertexFetchIndex_Culled(DispatchThreadId, GROUP_SIZE);
 		const uint VertexCountAfterCulling = HairStrandsVF_CullingIndirectBuffer[3];
 		if (FetchIndex >= VertexCountAfterCulling)
 		{
@@ -724,10 +721,9 @@ RWBuffer<uint4> OutBuffer;
 
 #define GROUP_SIZE 64
 [numthreads(GROUP_SIZE, 1, 1)]
-void MainCS(
-	uint3 DispatchThreadId : SV_DispatchThreadID)
+void MainCS(uint3 DispatchThreadId : SV_DispatchThreadID)
 {
-	const uint Index = DispatchThreadId.x + DispatchThreadId.y * DispatchCountX * GROUP_SIZE;
+	const uint Index = GetHairStrandsVertexFetchIndex(DispatchThreadId, GROUP_SIZE, DispatchCountX);
 	if (Index >= ElementCount)
 	{
 		return;
@@ -749,10 +745,9 @@ RWBuffer<uint4> OutDeformedBuffer;
 
 #define GROUP_SIZE 64
 [numthreads(GROUP_SIZE, 1, 1)]
-void MainCS(
-	uint3 DispatchThreadId : SV_DispatchThreadID)
+void MainCS(uint3 DispatchThreadId : SV_DispatchThreadID)
 {
-	const uint Index = DispatchThreadId.x + DispatchThreadId.y * DispatchCountX * GROUP_SIZE;
+	const uint Index = GetHairStrandsVertexFetchIndex(DispatchThreadId, GROUP_SIZE, DispatchCountX);
 	if (Index >= ElementCount)
 	{
 		return;
diff --git a/Engine/Shaders/Private/HairStrands/HairStrandsRaytracingGeometry.usf b/Engine/Shaders/Private/HairStrands/HairStrandsRaytracingGeometry.usf
--- a/Engine/Shaders/Private/HairStrands/HairStrandsRaytracingGeometry.usf
+++ b/Engine/Shaders/Private/HairStrands/HairStrandsRaytracingGeometry.usf
@@ -1,6 +1,7 @@
 // Copyright Epic Games, Inc. All Rights Reserved.
 
 #include "../Common.ush" 
+#include "HairStrandsClusterCommon.ush"
 #include "HairStrandsVertexFactoryCommon.ush"
 
 uint VertexCount;
@@ -75,7 +76,7 @@ void SetInvalidSegment(uint InVertexIndex)
 [numthreads(GROUP_SIZE, 1, 1)]
 void MainCS(uint2 DispatchThreadId : SV_DispatchThreadID)
 {
-	uint VertexIndex0 = DispatchThreadId.x + DispatchThreadId.y * DispatchCountX * GROUP_SIZE;
+	uint VertexIndex0 = GetHairStrandsVertexFetchIndex(DispatchThreadId, GROUP_SIZE, DispatchCountX);
 	uint VertexIndex1 = VertexIndex0 + 1;
 	float LodRadiusScale0 = 1;
 	float LodRadiusScale1 = 1;
@@ -84,14 +85,13 @@ void MainCS(uint2 DispatchThreadId : SV_DispatchThreadID)
 	if (HairStrandsVF_bIsCullingEnable)
 	{
 		const uint VertexCountAfterCulling = HairStrandsVF_CullingIndirectBuffer[3];
-		uint FetchIndex0 = VertexIndex0;
-		uint FetchIndex1 = VertexIndex1;
+		uint FetchIndex0 = GetHairStrandsVertexFetchIndex_Culled(DispatchThreadId, GROUP_SIZE);
 		if (FetchIndex0 >= VertexCountAfterCulling)
 		{
 			SetInvalidSegment(VertexIndex0);
 			return;
 		}
-		FetchIndex1 = min(FetchIndex0 + 1, VertexCountAfterCulling - 1);
+		uint FetchIndex1 = min(FetchIndex0 + 1, VertexCountAfterCulling - 1);
 
 		VertexIndex0 = HairStrandsVF_CullingIndexBuffer[FetchIndex0];
 		LodRadiusScale0 = HairStrandsVF_CullingRadiusScaleBuffer[FetchIndex0];
diff --git a/Engine/Shaders/Private/HairStrands/HairStrandsTangent.usf b/Engine/Shaders/Private/HairStrands/HairStrandsTangent.usf
--- a/Engine/Shaders/Private/HairStrands/HairStrandsTangent.usf
+++ b/Engine/Shaders/Private/HairStrands/HairStrandsTangent.usf
@@ -1,6 +1,7 @@
 // Copyright Epic Games, Inc. All Rights Reserved.
 
 #include "../Common.ush"
+#include "HairStrandsClusterCommon.ush"
 
 uint VertexCount;
 uint DispatchCountX;
@@ -10,9 +11,6 @@ RWBuffer<float4> OutputTangentBuffer;
 
 uint			HairStrandsVF_bIsCullingEnable;
 
-// These number need to match HairStrandsClusterCulling.usf
-#define GROUP_COUNT_X 16
-
 #if PERMUTATION_CULLING == 1
 Buffer<uint>	HairStrandsVF_CullingIndirectBuffer;
 Buffer<uint>	HairStrandsVF_CullingIndexBuffer;
@@ -36,14 +34,14 @@ float4 PackNormal(float3 V, uint A)
 void MainCS(uint2 DispatchThreadId : SV_DispatchThreadID)
 {
 	const int MaxVertexCount = int(VertexCount);
-	int IndexCurr = DispatchThreadId.x + DispatchThreadId.y * DispatchCountX * GROUP_SIZE;
+	int IndexCurr = GetHairStrandsVertexFetchIndex(DispatchThreadId, GROUP_SIZE, DispatchCountX);
 	int IndexPrev = clamp(IndexCurr-1, 0, MaxVertexCount-1);
 	int IndexNext = clamp(IndexCurr+1, 0, MaxVertexCount-1);
 	bool bValid = IndexCurr < MaxVertexCount;
 #if PERMUTATION_CULLING
 	if (HairStrandsVF_bIsCullingEnable)
 	{
-		const int FetchIndex = DispatchThreadId.x + DispatchThreadId.y * GROUP_SIZE * GROUP_COUNT_X;
+		const int FetchIndex = GetHairStrandsVertexFetchIndex_Culled(DispatchThreadId, GROUP_SIZE);
 		const int VertexCountAfterCulling = clamp((int)HairStrandsVF_CullingIndirectBuffer[3], 0, MaxVertexCount);
 		bValid = false;
 		if (FetchIndex < VertexCountAfterCulling)
diff --git a/Engine/Shaders/Private/HairStrands/HairStrandsVisibilityRasterCompute.usf b/Engine/Shaders/Private/HairStrands/HairStrandsVisibilityRasterCompute.usf
--- a/Engine/Shaders/Private/HairStrands/HairStrandsVisibilityRasterCompute.usf
+++ b/Engine/Shaders/Private/HairStrands/HairStrandsVisibilityRasterCompute.usf
@@ -2,6 +2,7 @@
 
 #include "../Common.ush"
 #include "../MonteCarlo.ush"
+#include "HairStrandsClusterCommon.ush"
 #include "HairStrandsVertexFactoryCommon.ush"
 #include "HairStrandsVisibilityCommon.ush"
 
@@ -262,16 +263,16 @@ uint2 Rasterize(uint2 PixelCoord, bool bIsSteep, float FracY, float RatioScale,
 //////////////////////////////////////////////////////////////////////////////////////////
 
 // These number need to match HairStrandsClusterCulling.usf
+// TODO
 #define GROUP_SIZE 32
-#define GROUP_COUNT_X 16
 
 [numthreads(GROUP_SIZE, 1, 1)]
-void MainCS(uint3 DispatchThreadID : SV_DispatchThreadID)
+void MainCS(uint2 DispatchThreadID : SV_DispatchThreadID)
 {
 #if PERMUTATION_CULLING
-	uint VertexIndex0 = DispatchThreadID.x + DispatchThreadID.y * GROUP_SIZE * GROUP_COUNT_X;
+	uint VertexIndex0 = GetHairStrandsVertexFetchIndex_Culled(DispatchThreadID, GROUP_SIZE);
 #else
-	uint VertexIndex0 = DispatchThreadID.x + DispatchThreadID.y * GROUP_SIZE * DispatchCountX;
+	uint VertexIndex0 = GetHairStrandsVertexFetchIndex(DispatchThreadID, GROUP_SIZE, DispatchCountX);
 #endif
 
 	uint VertexIndex1 = min(VertexIndex0 + 1, HairStrandsVF_VertexCount - 1);
diff --git a/Engine/Shaders/Private/HairStrands/HairStrandsVoxelRasterCompute.usf b/Engine/Shaders/Private/HairStrands/HairStrandsVoxelRasterCompute.usf
--- a/Engine/Shaders/Private/HairStrands/HairStrandsVoxelRasterCompute.usf
+++ b/Engine/Shaders/Private/HairStrands/HairStrandsVoxelRasterCompute.usf
@@ -3,6 +3,7 @@
 #include "HairStrandsVoxelPageCommonStruct.ush"
 #include "../Common.ush"
 
+#include "HairStrandsClusterCommon.ush"
 #include "HairStrandsVertexFactoryCommon.ush"
 #include "HairStrandsVoxelPageCommon.ush" 
 
@@ -53,19 +54,18 @@ float3 GetHairVoxelJitter(uint2 PixelCoord, uint Seed)
 #define JITTER_ENABLE 0
 
 // These number need to match HairStrandsClusterCulling.usf
+// TODO: create permutation for different group size
 #define GROUP_SIZE 32
-#define GROUP_COUNT_X 16
 
 [numthreads(GROUP_SIZE, 1, 1)]
-void MainCS(uint3 DispatchThreadID : SV_DispatchThreadID)
+void MainCS(uint2 DispatchThreadID : SV_DispatchThreadID)
 {
 #if PERMUTATION_CULLING == 1
-	uint VertexIndex0 = DispatchThreadID.x + DispatchThreadID.y * GROUP_SIZE * GROUP_COUNT_X;
+	uint VertexIndex0 = GetHairStrandsVertexFetchIndex_Culled(DispatchThreadID, GROUP_SIZE);
 #else
-	uint VertexIndex0 = DispatchThreadID.x + DispatchThreadID.y * GROUP_SIZE * DispatchCountX;
+	uint VertexIndex0 = GetHairStrandsVertexFetchIndex(DispatchThreadID, GROUP_SIZE, DispatchCountX);
 #endif
 
-	//const uint VertexIndex = DispatchThreadID.x + DispatchThreadID.y * GROUP_SIZE * DispatchCountX;
 	uint VertexIndex1 = min(VertexIndex0 + 1, HairStrandsVF_VertexCount - 1);
 	if (VertexIndex0 >= HairStrandsVF_VertexCount)
 		return;
