Clarify ShaderPrint data structure offset.

diff --git a/Engine/Shaders/Private/MiniFontCommon.ush b/Engine/Shaders/Private/MiniFontCommon.ush
--- a/Engine/Shaders/Private/MiniFontCommon.ush
+++ b/Engine/Shaders/Private/MiniFontCommon.ush
@@ -8,6 +8,8 @@
 
 Texture2D MiniFontTexture;
 
+#define NATIVE_CHARACTER_RES 8
+
 // CharacterIDs for the MiniFont
 #define _0_ 0
 #define _1_ 1
@@ -65,12 +67,12 @@ void PrintCharacter(int2 PixelPos, inout float3 OutColor, float3 FontColor, inou
 {
 	uint2 Rel = (uint2)(PixelPos - LeftTop);
 
-	FLATTEN if (Rel.x < 8 && Rel.y < 8)
+	FLATTEN if (Rel.x < NATIVE_CHARACTER_RES && Rel.y < NATIVE_CHARACTER_RES)
 	{
 		OutColor = lerp(OutColor, FontColor, SampleMiniFont(CharacterID, Rel));
 	}
 
-	LeftTop.x += 8;
+	LeftTop.x += NATIVE_CHARACTER_RES;
 }
 
 // only for positive numbers
@@ -78,7 +80,7 @@ void PrintCharacter(int2 PixelPos, inout float3 OutColor, float3 FontColor, inou
 uint ExtractDigitFromFloat(float Number, float DigitValue)
 {
 	uint Temp = (uint)(Number / DigitValue);
-	return Temp % 10;
+	return (Temp % 10) + _0_;
 }
 
 // for printf debugging in the shader, has to be positive
diff --git a/Engine/Shaders/Private/ShaderPrintCommon.ush b/Engine/Shaders/Private/ShaderPrintCommon.ush
--- a/Engine/Shaders/Private/ShaderPrintCommon.ush
+++ b/Engine/Shaders/Private/ShaderPrintCommon.ush
@@ -28,6 +28,9 @@ RWStructuredBuffer<ShaderPrintItem> RWValuesBuffer;
 #define SHADER_PRINT_TYPE_UINT   3
 #define SHADER_PRINT_TYPE_HEX    4
 
+#define SHADER_PRINT_COUNT_OFFSET 0
+#define SHADER_PRINT_VALUE_OFFSET 1
+
 #include "/Engine/Generated/UniformBuffers/ShaderPrintUniform.ush"
 //float4 FontSize;      // .xy character size, .zw character spacing (in normalized coordinates)
 //int MaxValueCount;    // maximum number of values that we can print in one frame
@@ -53,7 +56,7 @@ void ShaderPrint_Internal(in ShaderPrintItem Item)
 
 	// Buffer counter is stored in first element .Value
 	int IndexToStore = 0;
-	InterlockedAdd(RWValuesBuffer[0].Value, 1, IndexToStore);
+	InterlockedAdd(RWValuesBuffer[SHADER_PRINT_COUNT_OFFSET].Value, 1, IndexToStore);
 
 	// Prevent writing off the buffer
 	// Note that counter still increases so need clamp when reading it in later passes
@@ -63,7 +66,7 @@ void ShaderPrint_Internal(in ShaderPrintItem Item)
 	}
 
 	// Because counter is in first element, read/writes need to be offset by 1
-	RWValuesBuffer[IndexToStore + 1] = Item;
+	RWValuesBuffer[IndexToStore + SHADER_PRINT_VALUE_OFFSET] = Item;
 }
 
 void ShaderPrint_Internal(in float2 ScreenPos, in int Value, in int Type)
@@ -464,6 +467,6 @@ void ShaderPrintFilterOneThread()
 {
 	// Atomic flag is stored in first element .Type Value
 	uint PrevValue;
-	InterlockedCompareExchange(RWValuesBuffer[0].Type, 0, 1, PrevValue);
+	InterlockedCompareExchange(RWValuesBuffer[SHADER_PRINT_COUNT_OFFSET].Type, 0, 1, PrevValue);
 	ShaderPrintFilter(PrevValue == 0);
 }
diff --git a/Engine/Shaders/Private/ShaderPrintDraw.usf b/Engine/Shaders/Private/ShaderPrintDraw.usf
--- a/Engine/Shaders/Private/ShaderPrintDraw.usf
+++ b/Engine/Shaders/Private/ShaderPrintDraw.usf
@@ -23,9 +23,12 @@ RWBuffer<uint> RWIndirectDrawArgsBuffer;
 void InitValueBufferCS()
 {
 	// Clear .Value in the first buffer element which is (ab)used as the buffer counter
-	RWValuesBuffer[0].Value = 0;
 	// Clear .Type in the first buffer element which is (ab)used as the ShaderPrintFilterOneThread() atomic value
-	RWValuesBuffer[0].Type = 0;
+
+	// Symbol/print Counter
+	RWValuesBuffer[SHADER_PRINT_COUNT_OFFSET].Value = 0;
+	RWValuesBuffer[SHADER_PRINT_COUNT_OFFSET].Type = 0;
+
 }
 
 
@@ -36,7 +39,7 @@ void InitValueBufferCS()
 [numthreads(1, 1, 1)]
 void BuildIndirectDispatchArgsCS()
 {
-	uint ValueCount = min(ShaderPrintUniform.MaxValueCount, ValuesBuffer[0].Value);
+	uint ValueCount = min(ShaderPrintUniform.MaxValueCount, ValuesBuffer[SHADER_PRINT_COUNT_OFFSET].Value);
 
 	RWIndirectDispatchArgsBuffer[0] = (ValueCount + 63)/64;
 	RWIndirectDispatchArgsBuffer[1] = 1;
@@ -55,7 +58,7 @@ void AddSymbol(in ShaderPrintItem Item)
 {
 	// Buffer counter is stored in first element .Value
 	int IndexToStore = 0;
-	InterlockedAdd(RWSymbolsBuffer[0].Value, 1, IndexToStore);
+	InterlockedAdd(RWSymbolsBuffer[SHADER_PRINT_COUNT_OFFSET].Value, 1, IndexToStore);
 	
 	// Prevent writing off the buffer
 	// Note that counter still increases so need to clamp when reading it in later passes
@@ -65,7 +68,7 @@ void AddSymbol(in ShaderPrintItem Item)
 	}
 
 	// Because counter is in first element, read/writes need to be offset by 1
-	RWSymbolsBuffer[IndexToStore + 1] = Item;
+	RWSymbolsBuffer[IndexToStore + SHADER_PRINT_VALUE_OFFSET] = Item;
 }
 
 void AddSymbol(inout float2 ScreenPos, in int Symbol)
@@ -73,7 +76,7 @@ void AddSymbol(inout float2 ScreenPos, in int Symbol)
 	ShaderPrintItem Item;
 	Item.ScreenPos = ScreenPos;
 	Item.Value = Symbol;
-	Item.Type = 0;
+	Item.Type = SHADER_PRINT_TYPE_SYMBOL;
 	AddSymbol(Item);
 
 	ScreenPos.x += ShaderPrintUniform.FontSize.z;
@@ -186,13 +189,13 @@ void BuildSymbolBufferCS(
 	uint3 DispatchThreadId : SV_DispatchThreadID
 	)
 {
-	uint ValueCount = min(ShaderPrintUniform.MaxValueCount, ValuesBuffer[0].Value);
+	uint ValueCount = min(ShaderPrintUniform.MaxValueCount, ValuesBuffer[SHADER_PRINT_COUNT_OFFSET].Value);
 	if (DispatchThreadId.x >= ValueCount)
 	{
 		return;
 	}
 
-	ShaderPrintItem Value = ValuesBuffer[DispatchThreadId.x + 1];
+	ShaderPrintItem Value = ValuesBuffer[DispatchThreadId.x + SHADER_PRINT_VALUE_OFFSET];
 	
 	if (Value.Type == SHADER_PRINT_TYPE_SYMBOL)
 		AddSymbol(Value.ScreenPos, Value.Value);
@@ -213,7 +216,7 @@ void BuildSymbolBufferCS(
 [numthreads(1, 1, 1)]
 void BuildIndirectDrawArgsCS()
 {
-	uint SymbolCount = min(ShaderPrintUniform.MaxSymbolCount, SymbolsBuffer[0].Value);
+	uint SymbolCount = min(ShaderPrintUniform.MaxSymbolCount, SymbolsBuffer[SHADER_PRINT_COUNT_OFFSET].Value);
 
 	RWIndirectDrawArgsBuffer[0] = 6;
 	RWIndirectDrawArgsBuffer[1] = SymbolCount * 2; // Main glyph + drop shadow
@@ -231,7 +234,8 @@ void DrawSymbolsVS(
 	out float4 OutPosition : SV_POSITION,
 	out float2 OutUV : TEXCOORD0,
 	out int OutValue : TEXCOORD1,
-	out int OutShadow : TEXCOORD2
+	out int OutShadow : TEXCOORD2,
+	out float3 OutColor : TEXCOORD3
 	)
 {
 	bool bDropShadow = (InstanceId & 1) == 0;
@@ -239,7 +243,7 @@ void DrawSymbolsVS(
 	float2 DropShadowOffset = bDropShadow ? ShaderPrintUniform.FontSize.xy * float2(0.2f, -0.2f) : 0;
 	
 	InstanceId = InstanceId >> 1;
-	ShaderPrintItem Symbol = SymbolsBuffer[InstanceId + 1];
+	ShaderPrintItem Symbol = SymbolsBuffer[InstanceId + SHADER_PRINT_VALUE_OFFSET];
 	OutValue = Symbol.Value;
 
 	OutUV = float2(((VertexId + 1) / 3) & 1, VertexId & 1);
