Fix vulkan compilation issue with new ShaderPrint coloring.

Vulkan/DXC does not support initialization of static const variable having complex type. This change lines change the initialization to some simple define instead.

diff --git a/Engine/Shaders/Private/ShaderPrintCommon.ush b/Engine/Shaders/Private/ShaderPrintCommon.ush
--- a/Engine/Shaders/Private/ShaderPrintCommon.ush
+++ b/Engine/Shaders/Private/ShaderPrintCommon.ush
@@ -50,9 +50,27 @@ struct FFontColor
 	float3 Color;
 };
 
-FFontColor InitFontColor(float InX, float InY, float InZ) { FFontColor Out; Out.Color = float3(InX, InY, InZ); return Out; }
-FFontColor InitFontColor(float3 In)						  { FFontColor Out; Out.Color = In; return Out; }
-
+FFontColor InitFontColor(float InX, float InY, float InZ)	{ FFontColor Out; Out.Color = float3(InX, InY, InZ); return Out; }
+FFontColor InitFontColor(float3 In)							{ FFontColor Out; Out.Color = In; return Out; }
+FFontColor GetDefaultFontColor()							{ FFontColor Out; Out.Color = float3(1,1,1); return Out; }
+
+// Certain shader compiler does not support static initialization with complex type. 
+// In the meantime, using define to provide default color initialization
+#if 1 
+#define FontWhite		InitFontColor(1, 1, 1)
+#define FontBlack		InitFontColor(0, 0, 0)
+#define FontRed			InitFontColor(1, 0, 0)
+#define FontGreen		InitFontColor(0, 1, 0)
+#define FontBlue		InitFontColor(0, 0, 1)
+#define FontYellow		InitFontColor(1, 1, 0)
+#define FontCyan		InitFontColor(0, 1, 1)
+#define FontMagenta		InitFontColor(1, 0, 1)
+#define FontOrange		InitFontColor(243.f / 255.f, 156.f / 255.f, 18.f / 255.f)
+#define FontPurple		InitFontColor(169.f / 255.f, 7.f / 255.f, 228.f / 255.f)
+#define FontTurquoise	InitFontColor(26.f / 255.f, 188.f / 255.f, 156.f / 255.f)
+#define FontSilver		InitFontColor(189.f / 255.f, 195.f / 255.f, 199.f / 255.f)
+#define FontEmerald		InitFontColor(46.f / 255.f, 204.f / 255.f, 113.f / 255.f)
+#else
 static const FFontColor FontWhite		= InitFontColor(1, 1, 1);
 static const FFontColor FontBlack		= InitFontColor(0, 0, 0);
 static const FFontColor FontRed			= InitFontColor(1, 0, 0);
@@ -66,6 +84,7 @@ static const FFontColor FontPurple		= InitFontColor(169.f /255.f,    7.f / 255.f
 static const FFontColor FontTurquoise	= InitFontColor( 26.f / 255.f, 188.f / 255.f, 156.f / 255.f);
 static const FFontColor FontSilver		= InitFontColor(189.f / 255.f, 195.f / 255.f, 199.f / 255.f);
 static const FFontColor FontEmerald		= InitFontColor( 46.f / 255.f, 204.f / 255.f, 113.f / 255.f);
+#endif
 
 // --------------------------------------------------------------------------
 // Util pack/unpac functions
@@ -159,7 +178,7 @@ float2 ShaderPrintSymbol(in float2 ScreenPos, in int Symbol, in FFontColor Color
 
 float2 ShaderPrintSymbol(in float2 ScreenPos, in int Symbol)
 {
-	return ShaderPrintSymbol(ScreenPos, Symbol, FontWhite);
+	return ShaderPrintSymbol(ScreenPos, Symbol, GetDefaultFontColor());
 }
 
 void ShaderPrintSymbol(in int symbol, in FFontColor Color)
@@ -169,7 +188,7 @@ void ShaderPrintSymbol(in int symbol, in FFontColor Color)
 
 void ShaderPrintSymbol(in int symbol)
 {
-	ShaderPrintCursorPos = ShaderPrintSymbol(ShaderPrintCursorPos, symbol, FontWhite);
+	ShaderPrintCursorPos = ShaderPrintSymbol(ShaderPrintCursorPos, symbol, GetDefaultFontColor());
 }
 
 // Function for reading global TEXT string. 
@@ -254,7 +273,7 @@ float2 ShaderPrint(in float2 ScreenPos, in InType Value) \
 { \
 	if (ShaderPrintFilterEnable) \
 	{ \
-		ScreenPos = ShaderPrint(ScreenPos, Value, FontWhite); \
+		ScreenPos = ShaderPrint(ScreenPos, Value, GetDefaultFontColor()); \
 	} \
 	return ScreenPos; \
 } \
@@ -269,7 +288,7 @@ void ShaderPrint(in InType Value) \
 { \
 	if (ShaderPrintFilterEnable) \
 	{ \
-		ShaderPrintCursorPos = ShaderPrint(ShaderPrintCursorPos, Value, FontWhite); \
+		ShaderPrintCursorPos = ShaderPrint(ShaderPrintCursorPos, Value, GetDefaultFontColor()); \
 	} \
 }
 
@@ -296,7 +315,7 @@ float2 ShaderPrint(in float2 ScreenPos, in InType Value) \
 		UNROLL \
 		for (uint CompIt=0;CompIt<InNumComponent;++CompIt) \
 		{ \
-			ScreenPos = ShaderPrint(ScreenPos, Value[CompIt], FontWhite); \
+			ScreenPos = ShaderPrint(ScreenPos, Value[CompIt], GetDefaultFontColor()); \
 		} \
 	} \
 	return ScreenPos; \
@@ -319,7 +338,7 @@ void ShaderPrint(in InType Value) \
 		UNROLL \
 		for (uint CompIt=0;CompIt<InNumComponent;++CompIt) \
 		{ \
-			ShaderPrintCursorPos = ShaderPrint(ShaderPrintCursorPos, Value[CompIt], FontWhite); \
+			ShaderPrintCursorPos = ShaderPrint(ShaderPrintCursorPos, Value[CompIt], GetDefaultFontColor()); \
 		} \
 	} \
 }
@@ -365,7 +384,7 @@ float2 ShaderPrint(in float2 ScreenPos, in InType Value) \
 		UNROLL \
 		for (uint CompY = 0; CompY < InNumComponentY; ++CompY) \
 		{ \
-			ScreenPos = ShaderPrint(ScreenPos, Value[CompY], FontWhite); \
+			ScreenPos = ShaderPrint(ScreenPos, Value[CompY], GetDefaultFontColor()); \
 			ScreenPos += float2(0.f, ShaderPrintUniform.FontSize.w); \
 		} \
 	} \
@@ -390,7 +409,7 @@ void ShaderPrint(in InType Value) \
 		UNROLL \
 		for (uint CompY = 0; CompY < InNumComponentY; ++CompY) \
 		{ \
-			ShaderPrintCursorPos = ShaderPrint(ShaderPrintCursorPos, Value[CompY], FontWhite); \
+			ShaderPrintCursorPos = ShaderPrint(ShaderPrintCursorPos, Value[CompY], GetDefaultFontColor()); \
 			ShaderPrintCursorPos += float2(0.f, ShaderPrintUniform.FontSize.w); \
 		} \
 	} \
diff --git a/Engine/Source/Developer/ShaderCompilerCommon/Private/ShaderCompilerCommon.cpp b/Engine/Source/Developer/ShaderCompilerCommon/Private/ShaderCompilerCommon.cpp
--- a/Engine/Source/Developer/ShaderCompilerCommon/Private/ShaderCompilerCommon.cpp
+++ b/Engine/Source/Developer/ShaderCompilerCommon/Private/ShaderCompilerCommon.cpp
@@ -1370,11 +1370,12 @@ void TransformStringIntoCharacterArray(FString& PreprocessedShaderSource)
 		TextChars += TEXT("\n");
 		TextChars += TEXT("float2 ShaderPrintText(float2 Pos, uint InTextEntry)\n");
 		TextChars += TEXT("{\n");
+		TextChars += TEXT("\tFFontColor DefaultColor = GetDefaultFontColor();\n"); 
 		TextChars += TEXT("\tuint Begin = TEXT_OFFSETS[InTextEntry];\n");
 		TextChars += TEXT("\tuint End = TEXT_OFFSETS[InTextEntry + 1];\n");
 		TextChars += TEXT("\tfor (uint i = Begin; i < End; ++i)\n");
 		TextChars += TEXT("\t{\n");
-		TextChars += TEXT("\t\tPos = ShaderPrintSymbol(Pos, TEXT_CHARS[i], FontWhite);\n");
+		TextChars += TEXT("\t\tPos = ShaderPrintSymbol(Pos, TEXT_CHARS[i], DefaultColor);\n");
 		TextChars += TEXT("\t}\n");
 		TextChars += TEXT("\treturn Pos;\n");
 		TextChars += TEXT("}\n");
@@ -1382,7 +1383,7 @@ void TransformStringIntoCharacterArray(FString& PreprocessedShaderSource)
 	else
 	{	
 		TextChars += TEXT("float2 ShaderPrintText(float2 Pos, uint InTextEntry) { return Pos; }\n");
-		TextChars += TEXT("float2 ShaderPrintText(float2 Pos, uint InTextEntry, float3 InColor) { return Pos; }\n");
+		TextChars += TEXT("float2 ShaderPrintText(float2 Pos, uint InTextEntry, FFontColor InColor) { return Pos; }\n");
 	}
 	
 	// 6. Insert global struct data + print function
