Basic shader check implementation that (on SM6) just puts the line number in the payload (line/file/text is future work)
 - Adds macros check, checkBufferAccess and checkStructuredBufferAccess, along with Slow-suffixed equivalents
 - A new cvar 'r.Shader.CheckLevel' which defaults to 1 (enable check but not Slow). Thus be careful when using check, if unsure use checkSlow!
 - The cvar controls macros DO_CHECK=0 and DO_GUARD_SLOW=0 that mirrow the ones in host code.
 - Note that SM6 version only works for CS and PS.

diff --git a/Engine/Shaders/Private/AssertionMacros.ush b/Engine/Shaders/Private/AssertionMacros.ush
new file mode 100644
--- /dev/null
+++ b/Engine/Shaders/Private/AssertionMacros.ush
@@ -0,0 +1,78 @@
+// Copyright Epic Games, Inc. All Rights Reserved.
+
+/*=============================================================================
+	AssertionMacros.ush: Macros for Assertion
+=============================================================================*/
+
+#pragma once
+
+#include "/Engine/Public/Platform.ush"
+
+#if !defined(DO_CHECK)
+	#define DO_CHECK 0
+#endif 
+
+#if !defined(DO_GUARD_SLOW)
+	#define DO_GUARD_SLOW 0
+#endif 
+
+
+#if DO_CHECK
+
+// Allow platform-specific overrides?
+#if !defined(check)
+	/**
+	 * Check macro, intended to detect serious, execution-disrupting (not continuable) errors in the shader code.
+	 * Is controlled by the read-only CVar `r.Shader.CheckLevel`:
+	 *		0           => DO_CHECK=0 and DO_GUARD_SLOW=0
+	 *		1 (default) => DO_CHECK=1 and DO_GUARD_SLOW=0
+	 *		2           => DO_CHECK=1 and DO_GUARD_SLOW=1
+	 * NOTE: As check() is enabled by default this will cause increased register pressure and have some overhead. Therefore, be careful when using this
+	 * and if in doubt, default to using checkSlow et al.
+	 * When not doing GPU performance profiling it may be a good idea to keep r.Shader.CheckLevel=2 to catch errors early.
+	 */
+	#define check(expr)	 \
+		{ \
+			if (!(expr)) \
+			{ \
+				PLATFORM_BREAK(); \
+			} \
+		}
+#endif // check
+
+
+	#define checkBufferAccess(buffer, index)	\
+	{	\
+		uint BufferCheckNumElements = 0U;	\
+		buffer.GetDimensions(BufferCheckNumElements);	\
+		check(uint(index) < BufferCheckNumElements);	\
+	}
+
+
+	#define checkStructuredBufferAccess(buffer, index)	\
+	{	\
+		uint BufferCheckNumElements = 0U;	\
+		uint BufferCheckStride = 0U;	\
+		buffer.GetDimensions(BufferCheckNumElements, BufferCheckStride);	\
+		check(uint(index) < BufferCheckNumElements);	\
+	}
+
+#else // !DO_CHECK
+	#define check(expr) { }
+
+	#define checkBufferAccess(buffer, index) { }
+	#define checkStructuredBufferAccess(buffer, index) { }
+#endif // DO_CHECK
+
+#if DO_GUARD_SLOW
+	#define checkSlow(expr)					check(expr)
+
+	#define checkBufferAccessSlow(buffer, index) checkBufferAccess(buffer, index)
+	#define checkStructuredBufferAccessSlow(buffer, index)	checkStructuredBufferAccess(buffer, index)
+
+#else // DO_GUARD_SLOW
+	#define checkSlow(expr) { }
+
+	#define checkBufferAccessSlow(buffer, index) { }
+	#define checkStructuredBufferAccessSlow(buffer, index) { }
+#endif // DO_GUARD_SLOW
diff --git a/Engine/Shaders/Private/Common.ush b/Engine/Shaders/Private/Common.ush
--- a/Engine/Shaders/Private/Common.ush
+++ b/Engine/Shaders/Private/Common.ush
@@ -74,6 +74,8 @@
 
 #include "Definitions.usf"
 
+#include "AssertionMacros.ush"
+
 #ifndef VELOCITY_ENCODE_DEPTH
 #define VELOCITY_ENCODE_DEPTH 1
 #endif
diff --git a/Engine/Shaders/Public/Platform/D3D/D3DCommon.ush b/Engine/Shaders/Public/Platform/D3D/D3DCommon.ush
--- a/Engine/Shaders/Public/Platform/D3D/D3DCommon.ush
+++ b/Engine/Shaders/Public/Platform/D3D/D3DCommon.ush
@@ -27,7 +27,8 @@ SamplerState D3DStaticTrilinearClampedSampler   : register(s5, space1000);
 
 #endif // PLATFORM_SUPPORTS_STATIC_SAMPLERS
 
-#if PLATFORM_SUPPORTS_DIAGNOSTIC_BUFFER
+// TODO: Add support for the UEDiagnosticBuffer to other shader types
+#if PLATFORM_SUPPORTS_DIAGNOSTIC_BUFFER && (PIXELSHADER || COMPUTESHADER)
 
 // Diagnostic buffer which will be automatically bound by the RHI as a root parameter
 // Counterpart to FD3D12DiagnosticBufferData in D3D12DirectCommandListManager.h
@@ -55,6 +56,8 @@ void UEReportAssertWithPayload(uint AssertID, uint4 Payload)
 	}
 }
 
+#define PLATFORM_BREAK() UEReportAssertWithPayload(0xCECC, uint4(__LINE__, 0, 0, 0))
+
 #define PLATFORM_ASSERT(condition, assert_id)              { if (!(condition)) UEReportAssertWithPayload(assert_id, uint4(0, 0, 0, 0)); }
 #define PLATFORM_ASSERT1(condition, assert_id, a)          { if (!(condition)) UEReportAssertWithPayload(assert_id, uint4(a, 0, 0, 0));  }
 #define PLATFORM_ASSERT2(condition, assert_id, a, b)       { if (!(condition)) UEReportAssertWithPayload(assert_id, uint4(a, b, 0, 0));  }
diff --git a/Engine/Source/Runtime/Engine/Private/ShaderCompiler/ShaderCompiler.cpp b/Engine/Source/Runtime/Engine/Private/ShaderCompiler/ShaderCompiler.cpp
--- a/Engine/Source/Runtime/Engine/Private/ShaderCompiler/ShaderCompiler.cpp
+++ b/Engine/Source/Runtime/Engine/Private/ShaderCompiler/ShaderCompiler.cpp
@@ -128,6 +128,14 @@ static FAutoConsoleVariableRef CVarShaderCompilerAllowDistributedCompilation(
 /** Maximum number of preprocessed shaders to dump to the log on a crash. Replace with CVar if needed. */
 static constexpr uint32 GMaxNumDumpedShaderSources = 10;
 
+int32 GSShaderCheckLevel = 1;
+static FAutoConsoleVariableRef CVarGSShaderCheckLevel(
+	TEXT("r.Shaders.CheckLevel"),
+	GSShaderCheckLevel,
+	TEXT("0 => DO_CHECK=0, DO_GUARD_SLOW=0, 1 => DO_CHECK=1, DO_GUARD_SLOW=0, 2 => DO_CHECK=1, DO_GUARD_SLOW=1 for all shaders."),
+	ECVF_Default
+);
+
 /** Helper functions for logging more debug info */
 namespace ShaderCompiler
 {
@@ -4957,6 +4965,11 @@ void GlobalBeginCompileShader(
 		}
 	}
 
+	{
+		Input.Environment.SetDefine(TEXT("DO_CHECK"), GSShaderCheckLevel > 0 ? 1 : 0);
+		Input.Environment.SetDefine(TEXT("DO_GUARD_SLOW"), GSShaderCheckLevel > 1 ? 1 : 0);
+	}
+
 	if (CVarShaderWarningsAsErrors.GetValueOnAnyThread())
 	{
 		Input.Environment.CompilerFlags.Add(CFLAG_WarningsAsErrors);
diff --git a/Engine/Source/Runtime/RenderCore/Private/Shader.cpp b/Engine/Source/Runtime/RenderCore/Private/Shader.cpp
--- a/Engine/Source/Runtime/RenderCore/Private/Shader.cpp
+++ b/Engine/Source/Runtime/RenderCore/Private/Shader.cpp
@@ -1430,6 +1430,15 @@ void ShaderMapAppendKeyString(EShaderPlatform Platform, FString& KeyString)
 		KeyString += (CVar && CVar->GetInt() == 1) ? TEXT("_WX") : TEXT("");
 	}
 	
+	{
+		static const auto CVar = IConsoleManager::Get().FindConsoleVariable(TEXT("r.Shaders.CheckLevel"));
+		// Note: Since 1 is the default, we don't modify the hash for this case, so as to not force a rebuild, and to keep the hash shorter.
+		if (CVar && (CVar->GetInt() == 0 || CVar->GetInt() == 2))
+		{
+			KeyString.Appendf(TEXT("_C%d"), CVar->GetInt());
+		}
+	}
+	
 	{
 		static const auto CVar = IConsoleManager::Get().FindConsoleVariable(TEXT("r.Shaders.FlowControlMode"));
 		if (CVar)
