Fix UMaterialExpressionDBufferTexture material node not accounting for dbuffer write mask on supported platforms.

diff --git a/Engine/Shaders/Private/BasePassPixelShader.usf b/Engine/Shaders/Private/BasePassPixelShader.usf
--- a/Engine/Shaders/Private/BasePassPixelShader.usf
+++ b/Engine/Shaders/Private/BasePassPixelShader.usf
@@ -1022,18 +1022,13 @@ void FPixelShaderInOut_MainPS(
 #endif
 		if ((GetPrimitiveData(MaterialParameters.PrimitiveId).Flags & PRIMITIVE_SCENE_DATA_FLAG_DECAL_RECEIVER) != 0 && View.ShowDecalsMask > 0)
 		{
-			uint DBufferMask = 0x07;
-#if PLATFORM_SUPPORTS_RENDERTARGET_WRITE_MASK
-			DBufferMask = DecodeRTWriteMask(In.SvPosition.xy, OpaqueBasePass.DBufferRenderMask, 3);
+			uint ValidDBufferTargetMask = GetDBufferTargetMask(uint2(In.SvPosition.xy)) & MATERIALDECALRESPONSEMASK;
+
 			BRANCH
-#elif PLATFORM_SUPPORTS_PER_PIXEL_DBUFFER_MASK
-			DBufferMask = GetDBufferMask(In.SvPosition.xy, OpaqueBasePass.DBufferRenderMask);
-#endif
-			if (DBufferMask)
+			if (ValidDBufferTargetMask)
 			{
-				float2 NDC = MaterialParameters.ScreenPosition.xy / MaterialParameters.ScreenPosition.w;
-				float2 ScreenUV = NDC * ResolvedView.ScreenPositionScaleBias.xy + ResolvedView.ScreenPositionScaleBias.wz;
-				FDBufferData DBufferData = GetDBufferData(ScreenUV, DBufferMask & MATERIALDECALRESPONSEMASK);
+				float2 BufferUV = SvPositionToBufferUV(In.SvPosition);
+				FDBufferData DBufferData = GetDBufferData(BufferUV, ValidDBufferTargetMask);
 
 				ApplyDBufferData(DBufferData, MaterialParameters.WorldNormal, SubsurfaceColor, Roughness, BaseColor, Metallic, Specular);
 				DBufferOpacity = (DBufferData.ColorOpacity + DBufferData.NormalOpacity + DBufferData.RoughnessOpacity) * (1.0f / 3.0f);
diff --git a/Engine/Shaders/Private/Common.ush b/Engine/Shaders/Private/Common.ush
--- a/Engine/Shaders/Private/Common.ush
+++ b/Engine/Shaders/Private/Common.ush
@@ -1017,12 +1017,10 @@ MaterialFloat3 RGBTDecode8BPC(MaterialFloat4 RGBT, MaterialFloat Range)
   * This gets a bit from a write mask texture created with FRTWriteMaskDecodeCS. Only supprted on some platforms.
   */
 #if PLATFORM_SUPPORTS_RENDERTARGET_WRITE_MASK
-uint DecodeRTWriteMask(float2 ScreenPos, Texture2D<uint> RTWriteMaskTexture, uint NumEncodedTextures)
+uint DecodeRTWriteMask(uint2 PixelPos, Texture2D<uint> RTWriteMaskTexture, uint NumEncodedTextures)
 {
-	uint2 InPos = uint2(ScreenPos);
-
-	uint2 TileIndex = InPos >> 3;
-	uint2 TileOffset = (InPos >> 2) & 1;
+	uint2 TileIndex = PixelPos >> 3;
+	uint2 TileOffset = (PixelPos >> 2) & 1;
 
 	uint Shift = (TileOffset.y * 2 + TileOffset.x) * NumEncodedTextures;
 	uint Mask = ~((~0) << NumEncodedTextures);
diff --git a/Engine/Shaders/Private/CompositeUIPixelShader.usf b/Engine/Shaders/Private/CompositeUIPixelShader.usf
--- a/Engine/Shaders/Private/CompositeUIPixelShader.usf
+++ b/Engine/Shaders/Private/CompositeUIPixelShader.usf
@@ -100,7 +100,7 @@ void Main(
 	)
 {
 #if PLATFORM_SUPPORTS_RENDERTARGET_WRITE_MASK
-	uint CompositeUIMask = DecodeRTWriteMask(Input.Position.xy, UIWriteMaskTexture, 1);
+	uint CompositeUIMask = DecodeRTWriteMask(uint2(Input.Position.xy), UIWriteMaskTexture, 1);
 	BRANCH
 	if (CompositeUIMask == 0)
 	{
diff --git a/Engine/Shaders/Private/DBufferDecalShared.ush b/Engine/Shaders/Private/DBufferDecalShared.ush
--- a/Engine/Shaders/Private/DBufferDecalShared.ush
+++ b/Engine/Shaders/Private/DBufferDecalShared.ush
@@ -82,27 +82,50 @@ FDBufferData DecodeDBufferData(
 	#define OpaqueBasePassDBufferCTextureSampler OpaqueBasePass.DBufferCTextureSampler
 #endif
 
-// @param UV - UV space in the DBuffer textures (BufferSize resolution)
-FDBufferData GetDBufferData(float2 UV, uint RTMaskBit)
+// Returns a bitmask of DBuffer targets that contain valid data for this pixel.
+// @param BufferUV - UV space in the DBuffer textures
+uint GetDBufferTargetMask(uint2 PixelPos)
+{
+#if USE_DBUFFER && ENABLE_DBUFFER_TEXTURES
+
+	#if PLATFORM_SUPPORTS_RENDERTARGET_WRITE_MASK
+		return DecodeRTWriteMask(PixelPos, OpaqueBasePass.DBufferRenderMask, 3);
+	#elif PLATFORM_SUPPORTS_PER_PIXEL_DBUFFER_MASK
+		uint Mask = OpaqueBasePass.DBufferRenderMask.Load(uint3(PixelPos, 0));
+		return Mask > 0 ? 0x07 : 0x00;
+	#else
+		return 0x07;
+	#endif
+
+#else
+	return 0;
+#endif
+}
+
+// @param BufferUV - UV space in the DBuffer textures
+FDBufferData GetDBufferData(float2 BufferUV, uint RTMaskBit)
 {
 	float4 DBufferA = float4(0.0, 0.0, 0.0, 1.0);
 	float4 DBufferB = float4(128.f/255.5f, 128.f / 255.5f, 128.f / 255.5f, 1.0);
 	float4 DBufferC = float4(0.0, 0.0, 0.0, 1.0);
 
 #if USE_DBUFFER && ENABLE_DBUFFER_TEXTURES
+	BRANCH
 	if (RTMaskBit & 0x1)
 	{
-		DBufferA = Texture2DSampleLevel(OpaqueBasePass.DBufferATexture, OpaqueBasePassDBufferATextureSampler, UV, 0);
+		DBufferA = Texture2DSampleLevel(OpaqueBasePass.DBufferATexture, OpaqueBasePassDBufferATextureSampler, BufferUV, 0);
 	}
 
+	BRANCH
 	if (RTMaskBit & 0x2)
 	{
-		DBufferB = Texture2DSampleLevel(OpaqueBasePass.DBufferBTexture, OpaqueBasePassDBufferBTextureSampler, UV, 0);
+		DBufferB = Texture2DSampleLevel(OpaqueBasePass.DBufferBTexture, OpaqueBasePassDBufferBTextureSampler, BufferUV, 0);
 	}
 
+	BRANCH
 	if (RTMaskBit & 0x4)
 	{
-		DBufferC = Texture2DSampleLevel(OpaqueBasePass.DBufferCTexture, OpaqueBasePassDBufferCTextureSampler, UV, 0);
+		DBufferC = Texture2DSampleLevel(OpaqueBasePass.DBufferCTexture, OpaqueBasePassDBufferCTextureSampler, BufferUV, 0);
 	}
 #endif
 
@@ -130,10 +153,3 @@ void ApplyDBufferData(
 	Specular = Specular * DBufferData.RoughnessOpacity + DBufferData.PreMulSpecular;
 #endif
 }
-
-uint GetDBufferMask(in float2 ScreenPosition, in Texture2D<uint> DBufferMaskTexture)
-{
-	int2 PixelPos = int2(ScreenPosition.xy);
-	uint Mask = DBufferMaskTexture.Load( int3(PixelPos, 0) );
-	return Mask > 0 ? 0x07 : 0x00;
-}
diff --git a/Engine/Shaders/Private/MaterialTemplate.ush b/Engine/Shaders/Private/MaterialTemplate.ush
--- a/Engine/Shaders/Private/MaterialTemplate.ush
+++ b/Engine/Shaders/Private/MaterialTemplate.ush
@@ -2083,9 +2083,12 @@ float3 DecodeSceneColorForMaterialNode(float2 ScreenUV)
 
 #endif // SHADING_PATH_DEFERRED
 
-float4 MaterialExpressionDBufferTextureLookup(float2 UV, int DBufferTextureIndex)
+float4 MaterialExpressionDBufferTextureLookup(float2 BufferUV, int DBufferTextureIndex)
 {
-	FDBufferData DBufferData = GetDBufferData(UV, 0x07);
+	uint2 PixelPos = uint2(BufferUV * View.BufferSizeAndInvSize.xy);
+	uint ValidDBufferTargetMask = GetDBufferTargetMask(PixelPos) & (1 << DBufferTextureIndex);
+
+	FDBufferData DBufferData = GetDBufferData(BufferUV, ValidDBufferTargetMask);
 	switch (DBufferTextureIndex)
 	{
 	case 0:	return float4(DBufferData.PreMulColor, DBufferData.ColorOpacity);
