Disable DBuffer sampling in ray tracing hit shaders and Lumen surface cache shaders, as DBuffer is not available in those passes

diff --git a/Engine/Shaders/Private/DBufferDecalShared.ush b/Engine/Shaders/Private/DBufferDecalShared.ush
--- a/Engine/Shaders/Private/DBufferDecalShared.ush
+++ b/Engine/Shaders/Private/DBufferDecalShared.ush
@@ -86,7 +86,7 @@ FDBufferData DecodeDBufferData(
 // @param BufferUV - UV space in the DBuffer textures
 uint GetDBufferTargetMask(uint2 PixelPos)
 {
-#if USE_DBUFFER && ENABLE_DBUFFER_TEXTURES
+#if USE_DBUFFER && ENABLE_DBUFFER_TEXTURES && !RAYHITGROUPSHADER && !LUMEN_CARD_CAPTURE
 
 	#if PLATFORM_SUPPORTS_RENDERTARGET_WRITE_MASK
 		return DecodeRTWriteMask(PixelPos, OpaqueBasePass.DBufferRenderMask, 3);
@@ -105,11 +105,12 @@ uint GetDBufferTargetMask(uint2 PixelPos)
 // @param BufferUV - UV space in the DBuffer textures
 FDBufferData GetDBufferData(float2 BufferUV, uint RTMaskBit)
 {
+	// Setup default values, which mean that no decals are present in DBuffer
 	float4 DBufferA = float4(0.0, 0.0, 0.0, 1.0);
-	float4 DBufferB = float4(128.f/255.5f, 128.f / 255.5f, 128.f / 255.5f, 1.0);
+	float4 DBufferB = float4(128.0f / 255.0f, 128.f / 255.5f, 128.f / 255.5f, 1.0);
 	float4 DBufferC = float4(0.0, 0.0, 0.0, 1.0);
 
-#if USE_DBUFFER && ENABLE_DBUFFER_TEXTURES
+#if USE_DBUFFER && ENABLE_DBUFFER_TEXTURES && !RAYHITGROUPSHADER && !LUMEN_CARD_CAPTURE
 	BRANCH
 	if (RTMaskBit & 0x1)
 	{
