Introduce new RT payload for path tracing shaders

The goal is to simplify the code for path tracing and make room for fields that will be required by future feature additions without impacting the regular RT shaders.

Simplified the packing/unpacking logic a bit since we have more space available in this context.

Remove path tracing specific flags from RT payload and indicate which flags are free for future use

diff --git a/Engine/Plugins/Experimental/GPULightmass/Shaders/Private/LightmapPathTracing.usf b/Engine/Plugins/Experimental/GPULightmass/Shaders/Private/LightmapPathTracing.usf
--- a/Engine/Plugins/Experimental/GPULightmass/Shaders/Private/LightmapPathTracing.usf
+++ b/Engine/Plugins/Experimental/GPULightmass/Shaders/Private/LightmapPathTracing.usf
@@ -76,7 +76,7 @@ void GenerateCosineNormalRay(
 
 	// Perform cosine-hemispherical sampling and convert to world-space
 	float4 Direction_Tangent = CosineSampleHemisphere(RandSample);
-	TangentDirection = Direction_Tangent;
+	TangentDirection = Direction_Tangent.xyz;
 	float3 Direction_World = TangentToWorld(Direction_Tangent.xyz, WorldNormal);
 
 	RayOrigin = WorldPosition;
@@ -157,13 +157,13 @@ void AccumulateRadiance(inout float3 TotalRadiance, float3 PathRadiance)
 	}
 }
 
-FMaterialClosestHitPayload TraceTransparentRay(RayDesc Ray, FRayCone RayCone, bool IsCameraRay, bool LastBounce, bool IncludeEmission, uint2 LaunchIndex, uint NumLights, inout RandomSequence RandSequence, inout float3 PathThroughput, inout float3 Radiance)
+FPathTracingPayload TraceTransparentRay(RayDesc Ray, FRayCone RayCone, bool IsCameraRay, bool LastBounce, bool IncludeEmission, uint2 LaunchIndex, uint NumLights, inout RandomSequence RandSequence, inout float3 PathThroughput, inout float3 Radiance)
 {
 	const uint RayFlags = IsCameraRay && EnableCameraBackfaceCulling ? RAY_FLAG_CULL_BACK_FACING_TRIANGLES : 0;
 	const uint MissShaderIndex = 0;
 	float SelectionWeightSum = 0;
 	float3 PayloadThroughput = PathThroughput;
-	FMaterialClosestHitPayload Payload;
+	FPathTracingPayload Payload;
 	if (!IncludeEmission && LastBounce)
 	{
 		Payload.SetMiss();
@@ -172,17 +172,11 @@ FMaterialClosestHitPayload TraceTransparentRay(RayDesc Ray, FRayCone RayCone, bo
 	}
 	for (;;)
 	{
-		FPackedMaterialClosestHitPayload PackedPayload = (FPackedMaterialClosestHitPayload)0;
-		// Set payload inputs
-		PackedPayload.SetRayCone(RayCone);
-		PackedPayload.SetPixelCoord(LaunchIndex);
-		if (!IsCameraRay)
-		{
-			// avoid double counting what was captured by the skylight
-			// also avoid noise from hot spots (they can be properly
-			// importance sampled if a capturing skylight is added)
-			PackedPayload.SetIgnoreSkyMaterials();
-		}
+		// avoid double counting what was captured by the skylight
+		// also avoid noise from hot spots (they can be properly
+		// importance sampled if a capturing skylight is added)
+		const bool bIgnoreSky = !IsCameraRay;
+		FPackedPathTracingPayload PackedPayload = InitPathTracingPayload(RayCone, LaunchIndex, bIgnoreSky);
 		// Trace the ray
 		TraceRay(
 			TLAS,
@@ -211,7 +205,7 @@ FMaterialClosestHitPayload TraceTransparentRay(RayDesc Ray, FRayCone RayCone, bo
 		}
 
 		// Unpack the payload
-		FMaterialClosestHitPayload HitPayload = UnpackPathTracingPayload(PackedPayload, Ray);
+		FPathTracingPayload HitPayload = UnpackPathTracingPayload(PackedPayload, Ray);
 
 		if (HitPayload.BlendingMode == RAY_TRACING_BLEND_MODE_ADDITIVE ||
 			HitPayload.BlendingMode == RAY_TRACING_BLEND_MODE_TRANSLUCENT)
@@ -293,15 +287,7 @@ float3 TraceTransparentVisibilityRay(RayDesc Ray, uint2 PixelCoord, float PathRo
 	const uint MultiplierForGeometryContributionToShaderIndex = RAY_TRACING_NUM_SHADER_SLOTS;
 	const uint MissShaderIndex = 0;
 
-	FPackedMaterialClosestHitPayload PackedPayload = (FPackedMaterialClosestHitPayload)0;
-
-	// Signal to the Any hit shader we want to evaluate the opacity
-	// The payload is used to carry the path throughput (for transparent shadows)
-	// and current path roughness (for approximate caustics)
-	PackedPayload.SetPathTracingVisibilityRay();
-	PackedPayload.SetRayThroughput(1.0);
-	PackedPayload.SetPathRoughness(PathRoughness);
-	PackedPayload.SetPixelCoord(PixelCoord);
+	FPackedPathTracingPayload PackedPayload = InitPathTracingVisibilityPayload(PixelCoord, PathRoughness);
 
 	// Trace the ray
 	TraceRay(
@@ -380,7 +366,7 @@ void PathTracingKernel(
 		const bool bIsLastBounce = Bounce == MaxBounces;
 		const bool bIncludeEmissive = (EnableDirectLighting != 0 || Bounce > 1) &&
 									  (EnableEmissive       != 0 || bIsCameraRay);
-		FMaterialClosestHitPayload Payload = (FMaterialClosestHitPayload)0;
+		FPathTracingPayload Payload = (FPathTracingPayload)0;
 		
 		if (bIsCameraRay)
 		{
@@ -415,7 +401,7 @@ void PathTracingKernel(
 			break;
 		}
 		
-		if (!Payload.IsFrontFace() && !IsMaterialTransmissive(Payload))
+		if (!Payload.IsFrontFace() && !Payload.IsMaterialTransmissive())
 		{
 			bIsValidSample = false;
 			break;
@@ -518,7 +504,7 @@ void PathTracingKernel(
 			float3 WorldNormal = Payload.WorldNormal;
 			uint PrimitiveLightingChannelMask = Payload.PrimitiveLightingChannelMask;
 
-			bool IsTransmissiveMaterial = ENABLE_TRANSMISSION && IsMaterialTransmissive(Payload);
+			bool IsTransmissiveMaterial = ENABLE_TRANSMISSION && Payload.IsMaterialTransmissive();
 
 			for (uint Index = 0, Num = LightLoopCount.NumLights; Index < Num; ++Index)
 			{
@@ -716,7 +702,7 @@ void PathTracingKernel(
 		NumVisibleLights = 0;
 		
 		// update ray cone for the next bounce
-		RayCone = Payload.GetRayCone();
+		RayCone = Payload.RayCone;
 	}
 	
 #if USE_IRRADIANCE_CACHING
diff --git a/Engine/Shaders/Private/PathTracing/Material/PathTracingClearCoat.ush b/Engine/Shaders/Private/PathTracing/Material/PathTracingClearCoat.ush
--- a/Engine/Shaders/Private/PathTracing/Material/PathTracingClearCoat.ush
+++ b/Engine/Shaders/Private/PathTracing/Material/PathTracingClearCoat.ush
@@ -16,19 +16,19 @@
 #include "PathTracingMaterialCommon.ush"
 #include "PathTracingFresnel.ush"
 
-float3 CalcThinTransmission(float NoL, float NoV, FMaterialClosestHitPayload Payload)
+float3 CalcThinTransmission(float NoL, float NoV, FPathTracingPayload Payload)
 {
 	float3 Transmission = 1.0;
 
-	float AbsorptionMix = GetMetallic(Payload);
+	float AbsorptionMix = Payload.Metallic;
 	if (AbsorptionMix > 0.0)
 	{
-		float LayerThickness = GetClearCoatThickness(Payload);
+		float LayerThickness = 1.0;	// Assume normalized thickness
 		float ThinDistance = LayerThickness * (rcp(NoV) + rcp(NoL));
 
 		// Base color represents reflected color viewed at 0 incidence angle, after being absorbed through the substrate.
 		// Because of this, extinction is normalized by traveling through layer thickness twice
-		float3 TransmissionColor = Diffuse_Lambert(GetBaseColor(Payload) );
+		float3 TransmissionColor = Diffuse_Lambert(Payload.BaseColor);
 		float3 ExtinctionCoefficient = -log(TransmissionColor) / (2.0 * LayerThickness);
 		float3 OpticalDepth = ExtinctionCoefficient * max(ThinDistance - 2.0 * LayerThickness, 0.0);
 		Transmission = exp(-OpticalDepth);
@@ -37,7 +37,7 @@ float3 CalcThinTransmission(float NoL, float NoV, FMaterialClosestHitPayload Pay
 	return Transmission;
 }
 
-float CalcClearCoatSelectionPdf(FMaterialClosestHitPayload Payload)
+float CalcClearCoatSelectionPdf(FPathTracingPayload Payload)
 {
 #if EVAL_TOP_LAYER == 0
 	// never select coating if top is disabled
@@ -48,7 +48,7 @@ float CalcClearCoatSelectionPdf(FMaterialClosestHitPayload Payload)
 #endif
 	// probability of choosing the coat
 	// TODO: this could be improved if we had a directional albedo tabulation
-	float Pdf = GetClearCoat(Payload) * CLEAR_COAT_F0;
+	float Pdf = Payload.GetClearCoat() * CLEAR_COAT_F0;
 	return Pdf;
 }
 
@@ -61,19 +61,19 @@ float ClearCoatFresnel(float VoH) {
 FMaterialEval ClearCoat_EvalMaterial(
 	float3 IncomingDirection,
 	float3 OutgoingDirection,
-	inout FMaterialClosestHitPayload Payload
+	inout FPathTracingPayload Payload
 )
 {
 	FMaterialEval Result = NullMaterialEval();
 	
 	float3 V_World = -IncomingDirection;
 	float3 L_World = OutgoingDirection;
-	float3 N_World = GetWorldNormal(Payload);
+	float3 N_World = Payload.WorldNormal;
 
 	float CoatSelectionPdf = CalcClearCoatSelectionPdf(Payload);
 
 	float2 Alpha = 0;
-	float3x3 Basis = GetGGXBasis(GetClearCoatRoughness(Payload), Payload.WorldNormal, Alpha);
+	float3x3 Basis = GetGGXBasis(Payload.GetClearCoatRoughness(), Payload.WorldNormal, Alpha);
 
 	// move vectors into right shading frame
 	float3 V = mul(Basis, V_World);
@@ -87,7 +87,7 @@ FMaterialEval ClearCoat_EvalMaterial(
 
 	// Specular lobe
 	const float2 GGXResult = GGXEvalReflection(L, V, H, Alpha);
-	const float F = GetClearCoat(Payload) * ClearCoatFresnel(VoH);
+	const float F = Payload.GetClearCoat() * ClearCoatFresnel(VoH);
 	const float3 SpecWeight = F * GGXResult.x;
 	const float SpecPdf = GGXResult.y;
 
@@ -96,7 +96,7 @@ FMaterialEval ClearCoat_EvalMaterial(
 
 	// bottom-layer clear coat evaluation
 	FMaterialEval BottomEval = NullMaterialEval();
-	Payload.WorldNormal = GetClearCoatBottomNormal(Payload);
+	Payload.WorldNormal = Payload.GetClearCoatBottomNormal();
 #if GREY_INNER_LAYER == 0
 	// NOTE: Should we try to account for the IOR ratio between layers? We could try to darken Specular Color by the coating IOR
 	// but this would apply to both the dielctric and metal portions equally ...
@@ -120,7 +120,7 @@ FMaterialEval ClearCoat_EvalMaterial(
 	float CC_F1 = ClearCoatFresnel(NoV);
 	float CC_F2 = ClearCoatFresnel(NoL);
 
-	BottomEval.Weight *= lerp(1.0, Transmission * (1 - CC_F1) * (1 - CC_F2), GetClearCoat(Payload));
+	BottomEval.Weight *= lerp(1.0, Transmission * (1 - CC_F1) * (1 - CC_F2), Payload.GetClearCoat());
 
 	AddLobeWithMIS(Result.Weight, Result.Pdf, BottomEval.Weight, BottomEval.Pdf, 1.0 - CoatSelectionPdf);
 
@@ -129,11 +129,11 @@ FMaterialEval ClearCoat_EvalMaterial(
 
 FMaterialSample ClearCoat_SampleMaterial(
 	float3 RayDirection,
-	FMaterialClosestHitPayload Payload,
+	FPathTracingPayload Payload,
 	float4 RandSample
 )
 {
-	float3 N_World = GetWorldNormal(Payload);
+	float3 N_World = Payload.WorldNormal;
 	float3 V_World = -RayDirection;
 
 	float CoatSelectionPdf = CalcClearCoatSelectionPdf(Payload);
@@ -146,7 +146,7 @@ FMaterialSample ClearCoat_SampleMaterial(
 		RandSample.x /= CoatSelectionPdf;
 
 		// No anisotropy, just pick an arbitrary basis
-		float Roughness = max(GetClearCoatRoughness(Payload), GGX_MIN_ROUGHNESS);
+		float Roughness = max(Payload.GetClearCoatRoughness(), GGX_MIN_ROUGHNESS);
 		float2 Alpha = Roughness * Roughness;
 		float3x3 Basis = GetTangentBasis(N_World);
 
@@ -173,14 +173,14 @@ FMaterialSample ClearCoat_SampleMaterial(
 #if GREY_INNER_LAYER == 0
 		float3 Weight = 0;
 		float Pdf = 0;
-		Payload.WorldNormal = GetClearCoatBottomNormal(Payload);
+		Payload.WorldNormal = Payload.GetClearCoatBottomNormal();
 		FMaterialSample InnerSample = DefaultLit_SampleMaterial(RayDirection, Payload, RandSample);
 		Payload.WorldNormal = N_World;
 
 		OutDirection = InnerSample.Direction;
 		OutRoughness = InnerSample.Roughness;
 #else
-		float3 Nb = GetClearCoatBottomNormal(Payload);
+		float3 Nb = Payload.GetClearCoatBottomNormal();
 		OutDirection = TangentToWorld(CosineSampleHemisphere(RandSample.xy).xyz, Nb);
 		OutRoughness = 1.0;
 #endif
diff --git a/Engine/Shaders/Private/PathTracing/Material/PathTracingCloth.ush b/Engine/Shaders/Private/PathTracing/Material/PathTracingCloth.ush
--- a/Engine/Shaders/Private/PathTracing/Material/PathTracingCloth.ush
+++ b/Engine/Shaders/Private/PathTracing/Material/PathTracingCloth.ush
@@ -8,7 +8,7 @@
 
 #include "PathTracingFresnel.ush"
 
-float ClothLobeSelectionPdf(float3 V, FMaterialClosestHitPayload Payload)
+float ClothLobeSelectionPdf(float3 V, FPathTracingPayload Payload)
 {
 	// TODO: This could be improved if we had a directional albedo tabulation, but this formulation will at least prevent
 	// under-sampling when one lobe is black (like with metals)
@@ -24,14 +24,14 @@ float ClothLobeSelectionPdf(float3 V, FMaterialClosestHitPayload Payload)
 FMaterialEval Cloth_EvalMaterial(
 	float3 IncomingDirection,
 	float3 OutgoingDirection,
-	FMaterialClosestHitPayload Payload
+	FPathTracingPayload Payload
 )
 {
 	FMaterialEval Result = NullMaterialEval();
 
 	const float3 V_World = -IncomingDirection;
 	const float3 L_World = OutgoingDirection;
-	const float3 N_World = GetWorldNormal(Payload);
+	const float3 N_World = Payload.WorldNormal;
 
 	float2 Alpha = 0;
 	const float3x3 Basis = GetGGXBasis(Payload.Roughness, Payload.Anisotropy, Payload.WorldNormal, Payload.WorldTangent, Alpha);
@@ -72,11 +72,11 @@ FMaterialEval Cloth_EvalMaterial(
 
 FMaterialSample Cloth_SampleMaterial(
 	float3 RayDirection,
-	FMaterialClosestHitPayload Payload,
+	FPathTracingPayload Payload,
 	float4 RandSample
 )
 {
-	float3 N = GetWorldNormal(Payload);
+	float3 N = Payload.WorldNormal;
 	float3 V = -RayDirection;
 
 	const float LobeSelectionPdf = ClothLobeSelectionPdf(V, Payload);
diff --git a/Engine/Shaders/Private/PathTracing/Material/PathTracingDefaultLit.ush b/Engine/Shaders/Private/PathTracing/Material/PathTracingDefaultLit.ush
--- a/Engine/Shaders/Private/PathTracing/Material/PathTracingDefaultLit.ush
+++ b/Engine/Shaders/Private/PathTracing/Material/PathTracingDefaultLit.ush
@@ -9,13 +9,13 @@
 #include "PathTracingMaterialCommon.ush"
 #include "PathTracingFresnel.ush"
 
-float CalcLobeSelectionPdf(float3 V, FMaterialClosestHitPayload Payload)
+float CalcLobeSelectionPdf(float3 V, FPathTracingPayload Payload)
 {
 	// TODO: This could be improved if we had a directional albedo tabulation, but this formulation will at least prevent
 	// under-sampling when one lobe is black (like with metals)
 
-	float3 Diff = GetDiffuseColor(Payload);
-	float3 Spec = GetSpecularColor(Payload);
+	float3 Diff = Payload.DiffuseColor;
+	float3 Spec = Payload.SpecularColor;
 
 	return LobeSelectionProb(Diff, Spec);
 }
@@ -23,14 +23,14 @@ float CalcLobeSelectionPdf(float3 V, FMaterialClosestHitPayload Payload)
 FMaterialEval DefaultLit_EvalMaterial(
 	float3 IncomingDirection,
 	float3 OutgoingDirection,
-	FMaterialClosestHitPayload Payload
+	FPathTracingPayload Payload
 )
 {
 	FMaterialEval Result = NullMaterialEval();
 
 	const float3 V_World = -IncomingDirection;
 	const float3 L_World = OutgoingDirection;
-	const float3 N_World = GetWorldNormal(Payload);
+	const float3 N_World = Payload.WorldNormal;
 
 	float2 Alpha = 0;
 	const float3x3 Basis = GetGGXBasis(Payload.Roughness, Payload.Anisotropy, Payload.WorldNormal, Payload.WorldTangent, Alpha);
@@ -47,11 +47,11 @@ FMaterialEval DefaultLit_EvalMaterial(
 	const float LobeSelectionPdf = CalcLobeSelectionPdf(V_World, Payload);
 
 	// Diffuse Lobe
-	AddLobeWithMIS(Result.Weight, Result.Pdf, GetDiffuseColor(Payload), NoL / PI, LobeSelectionPdf);
+	AddLobeWithMIS(Result.Weight, Result.Pdf, Payload.DiffuseColor, NoL / PI, LobeSelectionPdf);
 
 	// Specular lobe
 	const float2 GGXResult = GGXEvalReflection(L, V, H, Alpha);
-	const float3 F = F_Schlick(GetSpecularColor(Payload), VoH);
+	const float3 F = F_Schlick(Payload.SpecularColor, VoH);
 	const float3 SpecWeight = F * GGXResult.x;
 	const float SpecPdf = GGXResult.y;
 
@@ -62,11 +62,11 @@ FMaterialEval DefaultLit_EvalMaterial(
 
 FMaterialSample DefaultLit_SampleMaterial(
 	float3 RayDirection,
-	FMaterialClosestHitPayload Payload,
+	FPathTracingPayload Payload,
 	float4 RandSample
 )
 {
-	float3 N = GetWorldNormal(Payload);
+	float3 N = Payload.WorldNormal;
 	float3 V = -RayDirection;
 	
 	float LobeSelectionPdf = CalcLobeSelectionPdf(V, Payload);
@@ -117,11 +117,11 @@ FMaterialSample DefaultLit_SampleMaterial(
 	float OutPdf = 0;
 
 	// Diffuse Lobe
-	AddLobeWithMIS(OutWeight, OutPdf, GetDiffuseColor(Payload), NoL / PI, LobeSelectionPdf);
+	AddLobeWithMIS(OutWeight, OutPdf, Payload.DiffuseColor, NoL / PI, LobeSelectionPdf);
 
 	// Specular lobe
 	const float2 GGXResult = GGXEvalReflection(L, V, H, Alpha);
-	const float3 F = F_Schlick(GetSpecularColor(Payload), VoH);
+	const float3 F = F_Schlick(Payload.SpecularColor, VoH);
 	const float3 SpecWeight = F * GGXResult.x;
 	const float SpecPdf = GGXResult.y;
 	AddLobeWithMIS(OutWeight, OutPdf, SpecWeight, SpecPdf, 1.0 - LobeSelectionPdf);
diff --git a/Engine/Shaders/Private/PathTracing/Material/PathTracingLambert.ush b/Engine/Shaders/Private/PathTracing/Material/PathTracingLambert.ush
--- a/Engine/Shaders/Private/PathTracing/Material/PathTracingLambert.ush
+++ b/Engine/Shaders/Private/PathTracing/Material/PathTracingLambert.ush
@@ -6,36 +6,36 @@ PathTracingLambert.usf: Lambertian BRDF sampling functions
 #pragma once
 
 FMaterialSample Lambert_SampleMaterial(
-	FMaterialClosestHitPayload Payload,
+	FPathTracingPayload Payload,
 	float4 RandSample)
 {
-	float3 N_World = GetWorldNormal(Payload);
+	float3 N_World = Payload.WorldNormal;
 
 	float4 SampledValue = CosineSampleHemisphere(RandSample.xy);
 
-	return CreateMaterialSample(TangentToWorld(SampledValue.xyz, N_World), GetBaseColor(Payload), SampledValue.w, 1.0, 1.0);
+	return CreateMaterialSample(TangentToWorld(SampledValue.xyz, N_World), Payload.BaseColor, SampledValue.w, 1.0, 1.0);
 }
 
 FMaterialEval Lambert_EvalMaterial(
 	float3 OutgoingDirection,
-	FMaterialClosestHitPayload Payload
+	FPathTracingPayload Payload
 )
 {
-	float3 N_World = GetWorldNormal(Payload);
+	float3 N_World = Payload.WorldNormal;
 	float NoL = saturate(dot(N_World, OutgoingDirection));
 
-	return CreateMaterialEval(GetBaseColor(Payload), NoL / PI);
+	return CreateMaterialEval(Payload.BaseColor, NoL / PI);
 }
 
 
 FMaterialSample Simplified_SampleMaterial(
-	FMaterialClosestHitPayload Payload,
+	FPathTracingPayload Payload,
 	float4 RandSample)
 {
-	float3 N_World = GetWorldNormal(Payload);
+	float3 N_World = Payload.WorldNormal;
 
-	float3 FrontColor = GetBaseColor(Payload);
-	float3 BackColor = Payload.ShadingModelID == SHADINGMODELID_TWOSIDED_FOLIAGE ? GetSubsurfaceColor(Payload) : 0.0;
+	float3 FrontColor = Payload.BaseColor;
+	float3 BackColor = Payload.ShadingModelID == SHADINGMODELID_TWOSIDED_FOLIAGE ? Payload.GetSubsurfaceColor() : 0.0;
 
 	float4 SampledValue = CosineSampleHemisphere(RandSample.xy);
 
@@ -49,14 +49,14 @@ FMaterialSample Simplified_SampleMaterial(
 
 FMaterialEval Simplified_EvalMaterial(
 	float3 OutgoingDirection,
-	FMaterialClosestHitPayload Payload
+	FPathTracingPayload Payload
 )
 {
-	float3 N_World = GetWorldNormal(Payload);
+	float3 N_World = Payload.WorldNormal;
 	float NoL = dot(N_World, OutgoingDirection);
 
-	float3 FrontColor = GetBaseColor(Payload);
-	float3 BackColor = Payload.ShadingModelID == SHADINGMODELID_TWOSIDED_FOLIAGE ? GetSubsurfaceColor(Payload) : 0.0;
+	float3 FrontColor = Payload.BaseColor;
+	float3 BackColor = Payload.ShadingModelID == SHADINGMODELID_TWOSIDED_FOLIAGE ? Payload.GetSubsurfaceColor() : 0.0;
 	float TransmissionProb = LobeSelectionProb(BackColor, FrontColor);
 
 	float3 Weight = NoL < 0.0 ? BackColor : FrontColor;
diff --git a/Engine/Shaders/Private/PathTracing/Material/PathTracingMaterialCommon.ush b/Engine/Shaders/Private/PathTracing/Material/PathTracingMaterialCommon.ush
--- a/Engine/Shaders/Private/PathTracing/Material/PathTracingMaterialCommon.ush
+++ b/Engine/Shaders/Private/PathTracing/Material/PathTracingMaterialCommon.ush
@@ -9,156 +9,6 @@ PathTracingMaterialCommon.usf: Brdf utility functions
 #include "../PathTracingCommon.ush"
 #include "/Engine/Private/BRDF.ush"
 
-float3 GetMaterialEmission(FMaterialClosestHitPayload HitInfo)
-{
-	return HitInfo.Radiance;
-}
-
-float3 GetBaseColor(FMaterialClosestHitPayload HitInfo)
-{
-	return HitInfo.BaseColor;
-}
-
-float3 GetDiffuseColor(FMaterialClosestHitPayload HitInfo)
-{
-	return HitInfo.DiffuseColor;
-}
-
-float3 GetSpecularColor(FMaterialClosestHitPayload HitInfo)
-{
-	return HitInfo.SpecularColor;
-}
-
-float GetRoughness(FMaterialClosestHitPayload HitInfo)
-{
-	return HitInfo.Roughness;
-}
-
-float GetMetallic(FMaterialClosestHitPayload HitInfo)
-{
-	return HitInfo.Metallic;
-}
-
-float GetClearCoat(FMaterialClosestHitPayload HitInfo)
-{
-	return HitInfo.CustomData.x;
-}
-
-float GetClearCoatRoughness(FMaterialClosestHitPayload HitInfo)
-{
-	return max(HitInfo.CustomData.y, 0.02);
-}
-
-float SetClearCoatRoughness(inout FMaterialClosestHitPayload HitInfo, float Roughness)
-{
-	return HitInfo.CustomData.y = Roughness;
-}
-
-
-float3 GetClearCoatBottomNormal(FMaterialClosestHitPayload HitInfo)
-{
-#ifdef PATH_TRACING
-	#if CLEAR_COAT_BOTTOM_NORMAL
-	// Functionally inverse of behavior in RayTracingMaterialHitShader.usf
-	float2 oct2 = UnitVectorToOctahedron(HitInfo.WorldNormal);
-	float2 oct3 = float2(HitInfo.CustomData.a, HitInfo.CustomData.z);
-	float2 oct1 = (oct3 - (128.0 / 255.0)) * 2.0 + oct2;
-
-	return OctahedronToUnitVector(oct1);
-	#endif
-#endif
-	return HitInfo.WorldNormal;
-}
-
-float GetClearCoatThickness(FMaterialClosestHitPayload HitInfo)
-{
-	// Assume normalized thickness
-	return 1.0;
-}
-
-float3 GetSubsurfaceColor(FMaterialClosestHitPayload HitInfo)
-{
-	return HitInfo.CustomData.xyz;
-}
-
-void SetSubsurfaceColor(inout FMaterialClosestHitPayload HitInfo, float3 SSSColor)
-{
-	HitInfo.CustomData.xyz = SSSColor;
-}
-
-// Steal an unused field for storing this information
-float3 GetSubsurfaceRadius(FMaterialClosestHitPayload HitInfo)
-{
-	return HitInfo.IndirectIrradiance;
-}
-
-void SetSubsurfaceRadius(inout FMaterialClosestHitPayload HitInfo, float3 SSSRadius)
-{
-	HitInfo.IndirectIrradiance = SSSRadius;
-}
-
-// Steal an unused float3 (the only material with dual specular does not support anisotropy)
-// Data is encoded as (x: roughness 0, y: roughness 1, z: roughness mix)
-float3 GetDualRoughnessSpecular(FMaterialClosestHitPayload HitInfo)
-{
-	return float3(HitInfo.WorldTangent);
-}
-
-void SetDualRoughnessSpecular(inout FMaterialClosestHitPayload HitInfo, float3 RoughnessData)
-{
-	// store data in unused field and compute average roughness
-	HitInfo.WorldTangent = RoughnessData;
-	HitInfo.Roughness = lerp(RoughnessData.x, RoughnessData.y, RoughnessData.z);
-}
-
-float3 GetWorldNormal(FMaterialClosestHitPayload HitInfo)
-{
-	return HitInfo.WorldNormal;
-}
-
-bool IsPureSpecularTransmission(FMaterialClosestHitPayload HitInfo)
-{
-	return HitInfo.BlendingMode > RAY_TRACING_BLEND_MODE_OPAQUE;
-}
-
-bool IsLambert(FMaterialClosestHitPayload HitInfo)
-{
-	return GetRoughness(HitInfo) > GLOSSY_MAX_ROUGHNESS;
-}
-
-bool IsUnlitMaterial(FMaterialClosestHitPayload HitInfo)
-{
-	return HitInfo.ShadingModelID == SHADINGMODELID_UNLIT;
-}
-
-bool IsDefaultLitMaterial(FMaterialClosestHitPayload HitInfo)
-{
-	return HitInfo.ShadingModelID == SHADINGMODELID_DEFAULT_LIT;
-}
-
-bool IsClearCoatMaterial(FMaterialClosestHitPayload HitInfo)
-{
-	return HitInfo.ShadingModelID == SHADINGMODELID_CLEAR_COAT;
-}
-
-bool IsTwoSidedFoliageMaterial(FMaterialClosestHitPayload HitInfo)
-{
-	return HitInfo.ShadingModelID == SHADINGMODELID_TWOSIDED_FOLIAGE;
-}
-
-bool IsMaterialTransmissive(FMaterialClosestHitPayload HitInfo)
-{
-	return HitInfo.ShadingModelID == SHADINGMODELID_TWOSIDED_FOLIAGE ||
-		   (HitInfo.BlendingMode == RAY_TRACING_BLEND_MODE_TRANSLUCENT && HitInfo.Roughness != 0 && HitInfo.Opacity < 1);
-}
-
-bool IsSubsurfaceMaterial(FMaterialClosestHitPayload HitInfo)
-{
-	return HitInfo.ShadingModelID == SHADINGMODELID_SUBSURFACE ||
-		   HitInfo.ShadingModelID == SHADINGMODELID_PREINTEGRATED_SKIN ||
-		   HitInfo.ShadingModelID == SHADINGMODELID_SUBSURFACE_PROFILE;
-}
-
 // Given two lobes that will roughly contribute colors A and B to the total (estimated for example using directional albedo)
 // return the probability of choosing lobe A
 float LobeSelectionProb(float3 A, float3 B)
diff --git a/Engine/Shaders/Private/PathTracing/Material/PathTracingMaterialSampling.ush b/Engine/Shaders/Private/PathTracing/Material/PathTracingMaterialSampling.ush
--- a/Engine/Shaders/Private/PathTracing/Material/PathTracingMaterialSampling.ush
+++ b/Engine/Shaders/Private/PathTracing/Material/PathTracingMaterialSampling.ush
@@ -28,7 +28,7 @@ PathTracingMaterialSampling.usf: Material sampling functions
 	#define SIMPLIFIED_MATERIAL_SHADER 0
 #endif
 
-bool IsNeeValidMaterial(FMaterialClosestHitPayload Payload)
+bool IsNeeValidMaterial(FPathTracingPayload Payload)
 {
 	// Checks done in the same order as other sampling and evaluation functions to ensure consistency
 	if (Payload.BlendingMode == RAY_TRACING_BLEND_MODE_MODULATE ||
@@ -53,7 +53,7 @@ bool IsNeeValidMaterial(FMaterialClosestHitPayload Payload)
 	}
 }
 
-float3 GetMaterialTransparency(FMaterialClosestHitPayload Payload, float3 RayDirection)
+float3 GetMaterialTransparency(FPathTracingPayload Payload, float3 RayDirection)
 {
 	switch (Payload.ShadingModelID)
 	{
@@ -103,7 +103,7 @@ float3 GetMaterialTransparency(FMaterialClosestHitPayload Payload, float3 RayDir
 	}
 }
 
-float3 EstimateMaterialAlbedo(FMaterialClosestHitPayload Payload)
+float3 EstimateMaterialAlbedo(FPathTracingPayload Payload)
 {
 	float BrdfWeight = 1;
 	bool HasRefraction = false;
@@ -143,14 +143,14 @@ float3 EstimateMaterialAlbedo(FMaterialClosestHitPayload Payload)
 			return 0;
 #if SIMPLIFIED_MATERIAL_SHADER
 		case SHADINGMODELID_TWOSIDED_FOLIAGE:
-			Albedo = BrdfWeight * (Payload.BaseColor + GetSubsurfaceColor(Payload));
+			Albedo = BrdfWeight * (Payload.BaseColor + Payload.GetSubsurfaceColor());
 			break;
 #else
 		case SHADINGMODELID_TWOSIDED_FOLIAGE:
 		case SHADINGMODELID_SUBSURFACE:
 		case SHADINGMODELID_PREINTEGRATED_SKIN:
 		case SHADINGMODELID_SUBSURFACE_PROFILE:
-			Albedo = BrdfWeight * (Payload.DiffuseColor + Payload.SpecularColor + GetSubsurfaceColor(Payload));
+			Albedo = BrdfWeight * (Payload.DiffuseColor + Payload.SpecularColor + Payload.GetSubsurfaceColor());
 			break;
 		case SHADINGMODELID_DEFAULT_LIT:
 		case SHADINGMODELID_CLEAR_COAT:
@@ -177,7 +177,7 @@ float3 EstimateMaterialAlbedo(FMaterialClosestHitPayload Payload)
 	return Albedo;
 }
 
-void AdjustMaterialRoughness(inout FMaterialClosestHitPayload Payload, float PathRoughness)
+void AdjustMaterialRoughness(inout FPathTracingPayload Payload, float PathRoughness)
 {
 #if !SIMPLIFIED_MATERIAL_SHADER
 	// boost material roughness according to what we've seen so far
@@ -188,7 +188,7 @@ void AdjustMaterialRoughness(inout FMaterialClosestHitPayload Payload, float Pat
 			return;
 		case SHADINGMODELID_CLEAR_COAT:
 			// clamp coat roughness
-			SetClearCoatRoughness(Payload, max(GetClearCoatRoughness(Payload), PathRoughness));
+			Payload.SetClearCoatRoughness(max(Payload.GetClearCoatRoughness(), PathRoughness));
 			// fallthrough
 		case SHADINGMODELID_TWOSIDED_FOLIAGE:
 		case SHADINGMODELID_DEFAULT_LIT:
@@ -201,9 +201,9 @@ void AdjustMaterialRoughness(inout FMaterialClosestHitPayload Payload, float Pat
 		case SHADINGMODELID_SUBSURFACE_PROFILE:
 		{
 			// Clamp both specular lobes
-			float3 DualRoughnessData = GetDualRoughnessSpecular(Payload);
+			float3 DualRoughnessData = Payload.GetDualRoughnessSpecular();
 			DualRoughnessData = max(DualRoughnessData, float3(PathRoughness, PathRoughness, 0.0));
-			SetDualRoughnessSpecular(Payload, DualRoughnessData);
+			Payload.SetDualRoughnessSpecular(DualRoughnessData);
 			break;
 		}
 		default:
@@ -215,7 +215,7 @@ void AdjustMaterialRoughness(inout FMaterialClosestHitPayload Payload, float Pat
 }
 
 // returns an average roughness for the whole material (for the path roughness of shadow rays from this material)
-float GetAverageRoughness(FMaterialClosestHitPayload Payload)
+float GetAverageRoughness(FPathTracingPayload Payload)
 {
 #if SIMPLIFIED_MATERIAL_SHADER
 	return 1.0;
@@ -225,11 +225,11 @@ float GetAverageRoughness(FMaterialClosestHitPayload Payload)
 	switch (Payload.ShadingModelID)
 	{
 		case SHADINGMODELID_CLEAR_COAT:
-			Roughness = max(Roughness, GetClearCoatRoughness(Payload)); 
+			Roughness = max(Roughness, Payload.GetClearCoatRoughness());
 			break;
 		case SHADINGMODELID_SUBSURFACE_PROFILE:
 		{
-			float3 DualRoughnessData = GetDualRoughnessSpecular(Payload);
+			float3 DualRoughnessData = Payload.GetDualRoughnessSpecular();
 			Roughness = max(DualRoughnessData.x, DualRoughnessData.y);
 			break;
 		}
@@ -241,14 +241,14 @@ float GetAverageRoughness(FMaterialClosestHitPayload Payload)
 }
 
 
-FMaterialClosestHitPayload UnpackPathTracingPayload(FPackedMaterialClosestHitPayload PackedInput, RayDesc Ray)
+FPathTracingPayload UnpackPathTracingPayloadMaterial(FPackedPathTracingPayload PackedInput, RayDesc Ray)
 {
 	// the path tracer has a few more requirements compared to the ray-tracing path.
 	// Long term, we should reoranize the payload for the path tracer, but for now, just
 	// do a bit of post-processing and re-use fields to encode extracted information
 
 	// This function can also do a bit of per-material post-processing to simplify the rest of the code
-	FMaterialClosestHitPayload HitPayload = UnpackRayTracingPayload(PackedInput, Ray);
+	FPathTracingPayload HitPayload = UnpackPathTracingPayload(PackedInput, Ray);
 
 	switch (HitPayload.ShadingModelID)
 	{
@@ -258,7 +258,7 @@ FMaterialClosestHitPayload UnpackPathTracingPayload(FPackedMaterialClosestHitPay
 		{
 			// TODO: is this accurate enough? hard to measure exactly since the raster path uses a very different approach
 			float SSSRadius = (1 - HitPayload.Opacity) * 10.0; // simple formula, up to 10cm radius
-			SetSubsurfaceRadius(HitPayload, SSSRadius);
+			HitPayload.SetSubsurfaceRadius(SSSRadius);
 
 			// Reset opacity to avoid special cases elsewhere
 			HitPayload.Opacity = 1;
@@ -269,7 +269,7 @@ FMaterialClosestHitPayload UnpackPathTracingPayload(FPackedMaterialClosestHitPay
 		case SHADINGMODELID_TWOSIDED_FOLIAGE:
 		{
 			float3 Front = HitPayload.DiffuseColor;
-			float3 Back = GetSubsurfaceColor(HitPayload);
+			float3 Back = HitPayload.GetSubsurfaceColor();
 			float3 Sum = Front + Back;
 			float S = max(Sum.x, max(Sum.y, Sum.z));
 			if (S > 1)
@@ -277,7 +277,7 @@ FMaterialClosestHitPayload UnpackPathTracingPayload(FPackedMaterialClosestHitPay
 				// nudge the material back to something physically plausible
 				// NOTE: we ignore spec here since it should be accounted for by the brdf model itself
 				HitPayload.DiffuseColor = Front / S;
-				SetSubsurfaceColor(HitPayload, Back / S);
+				HitPayload.SetSubsurfaceColor(Back / S);
 			}
 			break;
 		}
@@ -304,18 +304,18 @@ FMaterialClosestHitPayload UnpackPathTracingPayload(FPackedMaterialClosestHitPay
 				// NOTE: this seems backwards? Opacity=0 acts like default-lit while Opacity=1 acts like SSS?
 				float3 SSSRadius = HitPayload.Opacity * WorldUnitScale * DiffuseMeanFreePath * GetDiffuseSurfaceScalingFactor3D(SurfaceAlbedo);
 
-				SetSubsurfaceRadius(HitPayload, SSSRadius);
+				HitPayload.SetSubsurfaceRadius(SSSRadius);
 
 				// burley mode doesn't have any diffuse component
-				SetSubsurfaceColor(HitPayload, HitPayload.DiffuseColor);
+				HitPayload.SetSubsurfaceColor(HitPayload.DiffuseColor);
 				HitPayload.DiffuseColor = 0;
 			}
 			else
 			{
 				// How to handle this case? a diffusion kernel is tabulated into the texture but I don't see a way to extract the original
 				// radius/color controls
-				SetSubsurfaceRadius(HitPayload, 0);
-				SetSubsurfaceColor(HitPayload, 0);
+				HitPayload.SetSubsurfaceRadius(0);
+				HitPayload.SetSubsurfaceColor(0);
 				// TODO: non-burley profile allows a blend of diffuse and SSS which should be accounted for here
 			}
 			// Decode dual roughness info
@@ -332,7 +332,7 @@ FMaterialClosestHitPayload UnpackPathTracingPayload(FPackedMaterialClosestHitPay
 			Lobe0Roughness = lerp(1.0f, Lobe0Roughness, saturate(HitPayload.Opacity * 10.0f));
 			Lobe1Roughness = lerp(1.0f, Lobe1Roughness, saturate(HitPayload.Opacity * 10.0f));
 
-			SetDualRoughnessSpecular(HitPayload, float3(Lobe0Roughness, Lobe1Roughness, LobeMix));
+			HitPayload.SetDualRoughnessSpecular(float3(Lobe0Roughness, Lobe1Roughness, LobeMix));
 			
 			// Reset opacity to avoid special cases elsewhere
 			HitPayload.Opacity = 1;
@@ -346,7 +346,7 @@ FMaterialClosestHitPayload UnpackPathTracingPayload(FPackedMaterialClosestHitPay
 
 FMaterialSample SampleMaterial(
 	float3 RayDirection,
-	FMaterialClosestHitPayload Payload,
+	FPathTracingPayload Payload,
 	float4 RandSample,
 	bool bForceRadianceProbe
 )
@@ -465,7 +465,7 @@ FMaterialSample SampleMaterial(
 FMaterialEval EvalMaterial(
 	float3 IncomingDirection,
 	float3 OutgoingDirection,
-	FMaterialClosestHitPayload Payload,
+	FPathTracingPayload Payload,
 	bool bForceRadianceProbe
 )
 {
diff --git a/Engine/Shaders/Private/PathTracing/Material/PathTracingRadianceProbe.ush b/Engine/Shaders/Private/PathTracing/Material/PathTracingRadianceProbe.ush
--- a/Engine/Shaders/Private/PathTracing/Material/PathTracingRadianceProbe.ush
+++ b/Engine/Shaders/Private/PathTracing/Material/PathTracingRadianceProbe.ush
@@ -49,7 +49,7 @@ float2 InverseConcentricMapping(float2 p)
 }
 
 FMaterialSample RadianceProbe_SampleMaterial(
-	FMaterialClosestHitPayload Payload,
+	FPathTracingPayload Payload,
 	float4 RandSample
 )
 {
@@ -113,7 +113,7 @@ FMaterialSample RadianceProbe_SampleMaterial(
 #endif
 #endif
 
-	float3 N_World = GetWorldNormal(Payload);
+	float3 N_World = Payload.WorldNormal;
 
 	float4 SampledValue = CosineSampleHemisphereConcentric(RandSample.xy);
 	float3 OutDirection = TangentToWorld(SampledValue.xyz, N_World);
@@ -126,9 +126,9 @@ FMaterialSample RadianceProbe_SampleMaterial(
 	return CreateMaterialSample(OutDirection, OutWeight, OutPdf, 1.0, 1.0);
 }
 
-float RadianceProbe_Pdf(float3 OutgoingDirection, FMaterialClosestHitPayload Payload)
+float RadianceProbe_Pdf(float3 OutgoingDirection, FPathTracingPayload Payload)
 {
-	float3 N_World = GetWorldNormal(Payload);
+	float3 N_World = Payload.WorldNormal;
 	
 	float PdfAdjust = 1.0f;	
 #ifdef LIGHTMAP_PATH_TRACING_MAIN_RG
@@ -194,10 +194,10 @@ float RadianceProbe_Pdf(float3 OutgoingDirection, FMaterialClosestHitPayload Pay
 
 FMaterialEval RadianceProbe_EvalMaterial(
 	float3 OutgoingDirection,
-	FMaterialClosestHitPayload Payload
+	FPathTracingPayload Payload
 )
 {
-	float3 N_World = GetWorldNormal(Payload);
+	float3 N_World = Payload.WorldNormal;
 	
 	float PdfAdjust = 1.0f;
 	
@@ -270,32 +270,32 @@ FMaterialEval RadianceProbe_EvalMaterial(
 #else
 
 FMaterialSample RadianceProbe_SampleMaterial(
-	FMaterialClosestHitPayload Payload,
+	FPathTracingPayload Payload,
 	float4 RandSample)
 {
-	float3 N_World = GetWorldNormal(Payload);
+	float3 N_World = Payload.WorldNormal;
 
 	float4 SampledValue = CosineSampleHemisphere(RandSample.xy);
 
 	return CreateMaterialSample(TangentToWorld(SampledValue.xyz, N_World), 1.0, SampledValue.w, 1.0, 1.0);
 }
 
-float RadianceProbe_Pdf(float3 OutgoingDirection, FMaterialClosestHitPayload Payload)
+float RadianceProbe_Pdf(float3 OutgoingDirection, FPathTracingPayload Payload)
 {
-	float3 N_World = GetWorldNormal(Payload);
+	float3 N_World = Payload.WorldNormal;
 	float NoL = saturate(dot(N_World, OutgoingDirection));
 	return NoL / PI;
 }
 
 FMaterialEval RadianceProbe_EvalMaterial(
 	float3 OutgoingDirection,
-	FMaterialClosestHitPayload Payload
+	FPathTracingPayload Payload
 )
 {
-	float3 N_World = GetWorldNormal(Payload);
+	float3 N_World = Payload.WorldNormal;
 	float NoL = saturate(dot(N_World, OutgoingDirection));
 
-	return CreateMaterialEval(GetDiffuseColor(Payload), NoL / PI);
+	return CreateMaterialEval(Payload.DiffuseColor, NoL / PI);
 }
 
 #endif
\ No newline at end of file
diff --git a/Engine/Shaders/Private/PathTracing/Material/PathTracingSolidGlass.ush b/Engine/Shaders/Private/PathTracing/Material/PathTracingSolidGlass.ush
--- a/Engine/Shaders/Private/PathTracing/Material/PathTracingSolidGlass.ush
+++ b/Engine/Shaders/Private/PathTracing/Material/PathTracingSolidGlass.ush
@@ -12,12 +12,12 @@ PathTracingSolidGlass.ush: Microfacet Refractive BSDF for solid glass
 
 FMaterialSample RoughGlass_SampleMaterial(
 	float3 RayDirection,
-	FMaterialClosestHitPayload Payload,
+	FPathTracingPayload Payload,
 	float4 RandSample)
 {
 	// NOTE: IsFrontFace() determines polygon orientation, because the normal is always flipped towards in the incoming ray
 	bool bIsEntering = Payload.IsFrontFace();
-	float3 N_World = GetWorldNormal(Payload);
+	float3 N_World = Payload.WorldNormal;
 	float3 V_World = -RayDirection;
 
 	// #dxr_todo: Maintain a refraction stack on the path tracing payload
@@ -62,14 +62,14 @@ FMaterialSample RoughGlass_SampleMaterial(
 FMaterialEval RoughGlass_EvalMaterial(
 	float3 IncomingDirection,
 	float3 OutgoingDirection,
-	FMaterialClosestHitPayload Payload
+	FPathTracingPayload Payload
 )
 {
 	// NOTE: IsFrontFace() determines polygon orientation, because the normal is always flipped towards in the incoming ray
 	bool bIsEntering = Payload.IsFrontFace();
 	float3 V_World = -IncomingDirection;
 	float3 L_World = OutgoingDirection;
-	float3 N_World = GetWorldNormal(Payload);
+	float3 N_World = Payload.WorldNormal;
 
 	// #dxr_todo: Maintain a refraction stack on the path tracing payload
 	float Ior = Payload.Ior;
diff --git a/Engine/Shaders/Private/PathTracing/Material/PathTracingSpecularReflection.ush b/Engine/Shaders/Private/PathTracing/Material/PathTracingSpecularReflection.ush
--- a/Engine/Shaders/Private/PathTracing/Material/PathTracingSpecularReflection.ush
+++ b/Engine/Shaders/Private/PathTracing/Material/PathTracingSpecularReflection.ush
@@ -10,20 +10,19 @@ PathTracingSpecularReflection.usf: Singular Mirror BRDF sampling functions
 
 void SpecularReflection_SampleMaterial(
 	float3 RayDirection,
-	FMaterialClosestHitPayload Payload,
+	FPathTracingPayload Payload,
 	out float3 OutDirection,
 	out float3 OutWeight,
 	out float OutPdf)
 {
-	float3 N_World = GetWorldNormal(Payload);
-	float3 SpecularColor = GetSpecularColor(Payload);
+	float3 N_World = Payload.WorldNormal;
 
 	float3 V_World = normalize(-RayDirection);
 	OutDirection = reflect(-V_World, N_World);
 
 	float NoV = abs(dot(V_World, N_World));
 
-	OutWeight = F_Schlick(SpecularColor, NoV);
+	OutWeight = F_Schlick(Payload.SpecularColor, NoV);
 	OutPdf = POSITIVE_INFINITY;
 }
 
diff --git a/Engine/Shaders/Private/PathTracing/Material/PathTracingSpecularTransmission.ush b/Engine/Shaders/Private/PathTracing/Material/PathTracingSpecularTransmission.ush
--- a/Engine/Shaders/Private/PathTracing/Material/PathTracingSpecularTransmission.ush
+++ b/Engine/Shaders/Private/PathTracing/Material/PathTracingSpecularTransmission.ush
@@ -11,12 +11,12 @@ PathTracingSpecularTransmission.ush: Refractive BSDFs (pure singular)
 
 FMaterialSample SpecularRefraction_SampleMaterial(
 	float3 RayDirection,
-	FMaterialClosestHitPayload Payload,
+	FPathTracingPayload Payload,
 	float4 RandSample)
 {
 	// NOTE: IsFrontFace() determines polygon orientation, because the normal is always flipped towards in the incoming ray
 	bool bIsEntering = Payload.IsFrontFace();
-	float3 WorldNormal = GetWorldNormal(Payload);
+	float3 WorldNormal = Payload.WorldNormal;
 
 	// #dxr_todo: Maintain a refraction stack on the path tracing payload
 	float Ior = Payload.Ior;
diff --git a/Engine/Shaders/Private/PathTracing/Material/PathTracingSubsurfaceProfile.ush b/Engine/Shaders/Private/PathTracing/Material/PathTracingSubsurfaceProfile.ush
--- a/Engine/Shaders/Private/PathTracing/Material/PathTracingSubsurfaceProfile.ush
+++ b/Engine/Shaders/Private/PathTracing/Material/PathTracingSubsurfaceProfile.ush
@@ -13,32 +13,32 @@
 #include "PathTracingMaterialCommon.ush"
 #include "PathTracingFresnel.ush"
 
-float2 SubsurfaceProfile_LobeCdf(float3 V, FMaterialClosestHitPayload Payload)
+float2 SubsurfaceProfile_LobeCdf(float3 V, FPathTracingPayload Payload)
 {
 	// TODO: This could be improved if we had a directional albedo tabulation, but this formulation will at least prevent
 	// under-sampling when one lobe is black (like with metals)
 
-	float3 Diff = GetDiffuseColor(Payload);
-	float3 Spec = GetSpecularColor(Payload);
+	float3 Diff = Payload.DiffuseColor;
+	float3 Spec = Payload.SpecularColor;
 
 	float DiffuseProb = LobeSelectionProb(Diff, Spec);
-	float SpecMix = GetDualRoughnessSpecular(Payload).z;
+	float SpecMix = Payload.GetDualRoughnessSpecular().z;
 	return float2(DiffuseProb, DiffuseProb + (1 - DiffuseProb) * (1 - SpecMix));
 }
 
 FMaterialEval SubsurfaceProfile_EvalMaterial(
 	float3 IncomingDirection,
 	float3 OutgoingDirection,
-	FMaterialClosestHitPayload Payload
+	FPathTracingPayload Payload
 )
 {
 	FMaterialEval Result = NullMaterialEval();
 
 	const float3 V_World = -IncomingDirection;
 	const float3 L_World = OutgoingDirection;
-	const float3 N_World = GetWorldNormal(Payload);
+	const float3 N_World = Payload.WorldNormal;
 
-	float3 DualRoughnessData = GetDualRoughnessSpecular(Payload);
+	float3 DualRoughnessData = Payload.GetDualRoughnessSpecular();
 	float2 Alpha0 = 0;
 	float2 Alpha1 = 0;
 	const float3x3 Basis = GetGGXBasis(DualRoughnessData.x, Payload.WorldNormal, Alpha0);
@@ -56,12 +56,12 @@ FMaterialEval SubsurfaceProfile_EvalMaterial(
 	const float3 LobeProb = float3(LobeCdf.x, LobeCdf.y - LobeCdf.x, 1 - LobeCdf.y);
 
 	// Diffuse Lobe
-	AddLobeWithMIS(Result.Weight, Result.Pdf, GetDiffuseColor(Payload), NoL / PI, LobeProb.x);
+	AddLobeWithMIS(Result.Weight, Result.Pdf, Payload.DiffuseColor, NoL / PI, LobeProb.x);
 
 	// Specular lobes
 	const float2 GGXResult0 = GGXEvalReflection(L, V, H, Alpha0);
 	const float2 GGXResult1 = GGXEvalReflection(L, V, H, Alpha1);
-	const float3 F = F_Schlick(GetSpecularColor(Payload), VoH);
+	const float3 F = F_Schlick(Payload.SpecularColor, VoH);
 	AddLobeWithMIS(Result.Weight, Result.Pdf, (1 - DualRoughnessData.z) * F * GGXResult0.x, GGXResult0.y, LobeProb.y);
 	AddLobeWithMIS(Result.Weight, Result.Pdf, (    DualRoughnessData.z) * F * GGXResult1.x, GGXResult1.y, LobeProb.z);
 
@@ -70,17 +70,17 @@ FMaterialEval SubsurfaceProfile_EvalMaterial(
 
 FMaterialSample SubsurfaceProfile_SampleMaterial(
 	float3 RayDirection,
-	FMaterialClosestHitPayload Payload,
+	FPathTracingPayload Payload,
 	float4 RandSample
 )
 {
-	float3 N = GetWorldNormal(Payload);
+	float3 N = Payload.WorldNormal;
 	float3 V = -RayDirection;
 
 	const float2 LobeCdf = SubsurfaceProfile_LobeCdf(V, Payload);
 	const float3 LobeProb = float3(LobeCdf.x, LobeCdf.y - LobeCdf.x, 1 - LobeCdf.y);
 
-	float3 DualRoughnessData = GetDualRoughnessSpecular(Payload);
+	float3 DualRoughnessData = Payload.GetDualRoughnessSpecular();
 	float2 Alpha0 = 0;
 	float2 Alpha1 = 0;
 	const float3x3 Basis = GetGGXBasis(DualRoughnessData.x, Payload.WorldNormal, Alpha0);
@@ -147,12 +147,12 @@ FMaterialSample SubsurfaceProfile_SampleMaterial(
 	float OutPdf = 0;
 
 	// Diffuse Lobe
-	AddLobeWithMIS(OutWeight, OutPdf, GetDiffuseColor(Payload), NoL / PI, LobeProb.x);
+	AddLobeWithMIS(OutWeight, OutPdf, Payload.DiffuseColor, NoL / PI, LobeProb.x);
 
 	// Specular lobe
 	const float2 GGXResult0 = GGXEvalReflection(L, V, H, Alpha0);
 	const float2 GGXResult1 = GGXEvalReflection(L, V, H, Alpha1);
-	const float3 F = F_Schlick(GetSpecularColor(Payload), VoH);
+	const float3 F = F_Schlick(Payload.SpecularColor, VoH);
 	AddLobeWithMIS(OutWeight, OutPdf, (1 - DualRoughnessData.z) * F * GGXResult0.x, GGXResult0.y, LobeProb.y);
 	AddLobeWithMIS(OutWeight, OutPdf, (    DualRoughnessData.z) * F * GGXResult1.x, GGXResult1.y, LobeProb.z);
 
diff --git a/Engine/Shaders/Private/PathTracing/Material/PathTracingThinGlass.ush b/Engine/Shaders/Private/PathTracing/Material/PathTracingThinGlass.ush
--- a/Engine/Shaders/Private/PathTracing/Material/PathTracingThinGlass.ush
+++ b/Engine/Shaders/Private/PathTracing/Material/PathTracingThinGlass.ush
@@ -51,12 +51,12 @@ SlabWeights ComputeThinSlabWeights(float3 SlabColor, float CosV, float Eta)
 FMaterialEval RoughThinGlass_EvalMaterial(
 	float3 IncomingDirection,
 	float3 OutgoingDirection,
-	FMaterialClosestHitPayload Payload
+	FPathTracingPayload Payload
 )
 {
 	float3 V_World = -IncomingDirection;
 	float3 L_World = OutgoingDirection;
-	float3 N_World = GetWorldNormal(Payload);
+	float3 N_World = Payload.WorldNormal;
 
 
 	float2 Alpha;
@@ -100,10 +100,10 @@ FMaterialEval RoughThinGlass_EvalMaterial(
 
 FMaterialSample RoughThinGlass_SampleMaterial(
 	float3 RayDirection,
-	FMaterialClosestHitPayload Payload,
+	FPathTracingPayload Payload,
 	float4 RandSample)
 {
-	float3 N_World = GetWorldNormal(Payload);
+	float3 N_World = Payload.WorldNormal;
 	float3 V_World = -RayDirection;
 
 	float Eta = Payload.Ior;
diff --git a/Engine/Shaders/Private/PathTracing/Material/PathTracingTwoSidedFoliage.ush b/Engine/Shaders/Private/PathTracing/Material/PathTracingTwoSidedFoliage.ush
--- a/Engine/Shaders/Private/PathTracing/Material/PathTracingTwoSidedFoliage.ush
+++ b/Engine/Shaders/Private/PathTracing/Material/PathTracingTwoSidedFoliage.ush
@@ -10,19 +10,19 @@
 
 float CalcTransmissionSelectionPdf(
 	float3 IncomingDirection,
-	FMaterialClosestHitPayload Payload
+	FPathTracingPayload Payload
 )
 {
 	// pick back lobe proportionally to its contribution relative to front side
-	float3 Diff = GetDiffuseColor(Payload);
-	float3 Spec = GetSpecularColor(Payload);
-	float3 Back = GetSubsurfaceColor(Payload);
+	float3 Diff = Payload.DiffuseColor;
+	float3 Spec = Payload.SpecularColor;
+	float3 Back = Payload.GetSubsurfaceColor();
 	return LobeSelectionProb(Back, Diff + Spec);
 }
 
 FMaterialSample TwoSidedFoliage_SampleMaterial(
 	float3 RayDirection,
-	FMaterialClosestHitPayload Payload,
+	FPathTracingPayload Payload,
 	float4 RandSample
 )
 {
@@ -31,8 +31,8 @@ FMaterialSample TwoSidedFoliage_SampleMaterial(
 	{
 		RandSample.x /= TransmissionSelectionPdf;
 
-		float3 N = GetWorldNormal(Payload);
-		float3 SubsurfaceColor = GetSubsurfaceColor(Payload);
+		float3 N = Payload.WorldNormal;
+		float3 SubsurfaceColor = Payload.GetSubsurfaceColor();
 
 		float4 SampledValue = CosineSampleHemisphere(RandSample.yz);
 
@@ -53,12 +53,12 @@ FMaterialSample TwoSidedFoliage_SampleMaterial(
 FMaterialEval TwoSidedFoliage_EvalMaterial(
 	float3 IncomingDirection,
 	float3 OutgoingDirection,
-	FMaterialClosestHitPayload Payload
+	FPathTracingPayload Payload
 )
 {
 	float TransmissionSelectionPdf = CalcTransmissionSelectionPdf(IncomingDirection, Payload);
 
-	float3 N = GetWorldNormal(Payload);
+	float3 N = Payload.WorldNormal;
 	float NoL = dot(N, OutgoingDirection);
 	if (NoL < 0.0)
 	{
@@ -66,7 +66,7 @@ FMaterialEval TwoSidedFoliage_EvalMaterial(
 
 		if (TransmissionSelectionPdf > 0)
 		{
-			float3 SubsurfaceColor = GetSubsurfaceColor(Payload);
+			float3 SubsurfaceColor = Payload.GetSubsurfaceColor();
 			float LambertPdf = -NoL / PI;
 
 			return CreateMaterialEval(SubsurfaceColor / TransmissionSelectionPdf, TransmissionSelectionPdf * LambertPdf);
diff --git a/Engine/Shaders/Private/PathTracing/PathTracing.usf b/Engine/Shaders/Private/PathTracing/PathTracing.usf
--- a/Engine/Shaders/Private/PathTracing/PathTracing.usf
+++ b/Engine/Shaders/Private/PathTracing/PathTracing.usf
@@ -71,13 +71,13 @@ void AccumulateRadiance(inout float3 TotalRadiance, float3 PathRadiance, bool bI
 	TotalRadiance += PathRadiance;
 }
 
-FMaterialClosestHitPayload TraceTransparentRay(RayDesc Ray, FRayCone RayCone, bool bIsCameraRay, bool bLastBounce, bool bIncludeEmission, uint2 LaunchIndex, uint NumLights, inout RandomSequence RandSequence, inout float3 PathThroughput, inout float3 Radiance, inout float3 Albedo, inout float3 Normal)
+FPathTracingPayload TraceTransparentRay(RayDesc Ray, FRayCone RayCone, bool bIsCameraRay, bool bLastBounce, bool bIncludeEmission, uint2 LaunchIndex, uint NumLights, inout RandomSequence RandSequence, inout float3 PathThroughput, inout float3 Radiance, inout float3 Albedo, inout float3 Normal)
 {
 	const uint RayFlags = bIsCameraRay && EnableCameraBackfaceCulling ? RAY_FLAG_CULL_BACK_FACING_TRIANGLES : 0;
 	const uint MissShaderIndex = 0;
 	float SelectionWeightSum = 0;
 	float3 PayloadThroughput = PathThroughput;
-	FMaterialClosestHitPayload Payload;
+	FPathTracingPayload Payload;
 	if (!bIncludeEmission && bLastBounce)
 	{
 		Payload.SetMiss();
@@ -86,17 +86,11 @@ FMaterialClosestHitPayload TraceTransparentRay(RayDesc Ray, FRayCone RayCone, bo
 	}
 	for (;;)
 	{
-		FPackedMaterialClosestHitPayload PackedPayload = (FPackedMaterialClosestHitPayload)0;
-		// Set payload inputs
-		PackedPayload.SetRayCone(RayCone);
-		PackedPayload.SetPixelCoord(LaunchIndex);
-		if (!bIsCameraRay)
-		{
-			// avoid double counting what was captured by the skylight
-			// also avoid noise from hot spots (they can be properly
-			// importance sampled if a capturing skylight is added)
-			PackedPayload.SetIgnoreSkyMaterials();
-		}
+		// avoid double counting what was captured by the skylight
+		// also avoid noise from hot spots (they can be properly
+		// importance sampled if a capturing skylight is added)
+		const bool bIgnoreSky = !bIsCameraRay;
+		FPackedPathTracingPayload PackedPayload = InitPathTracingPayload(RayCone, LaunchIndex, bIgnoreSky);
 		// Trace the ray
 		TraceRay(
 			TLAS,
@@ -125,7 +119,7 @@ FMaterialClosestHitPayload TraceTransparentRay(RayDesc Ray, FRayCone RayCone, bo
 		}
 
 		// Unpack the payload
-		FMaterialClosestHitPayload HitPayload = UnpackPathTracingPayload(PackedPayload, Ray);
+		FPathTracingPayload HitPayload = UnpackPathTracingPayloadMaterial(PackedPayload, Ray);
 
 		if (HitPayload.BlendingMode == RAY_TRACING_BLEND_MODE_ADDITIVE ||
 			HitPayload.BlendingMode == RAY_TRACING_BLEND_MODE_TRANSLUCENT)
@@ -212,15 +206,7 @@ float3 TraceTransparentVisibilityRay(RayDesc Ray, uint2 PixelCoord, float PathRo
 	const uint MultiplierForGeometryContributionToShaderIndex = RAY_TRACING_NUM_SHADER_SLOTS;
 	const uint MissShaderIndex = 0;
 
-	FPackedMaterialClosestHitPayload PackedPayload = (FPackedMaterialClosestHitPayload)0;
-
-	// Signal to the Any hit shader we want to evaluate the opacity
-	// The payload is used to carry the path throughput (for transparent shadows)
-	// and current path roughness (for approximate caustics)
-	PackedPayload.SetPathTracingVisibilityRay();
-	PackedPayload.SetRayThroughput(1.0);
-	PackedPayload.SetPathRoughness(PathRoughness);
-	PackedPayload.SetPixelCoord(PixelCoord);
+	FPackedPathTracingPayload PackedPayload = InitPathTracingVisibilityPayload(PixelCoord, PathRoughness);
 
 	// Trace the ray
 	TraceRay(
@@ -256,10 +242,7 @@ FProbeResult TraceProbeRay(RayDesc Ray, uint2 PixelCoord)
 {
 	// Trace a short ray to see if we escaped the surface
 	// TODO: figure out a way to only hit the object we started on (or handle blocking geometry)
-	FPackedMaterialClosestHitPayload PackedPayload = (FPackedMaterialClosestHitPayload)0;
-	// Set payload inputs
-	PackedPayload.SetPixelCoord(PixelCoord);
-	PackedPayload.SetIgnoreSkyMaterials();
+	FPackedPathTracingPayload PackedPayload = InitPathTracingPayload((FRayCone)0, PixelCoord, true);
 	// TODO: payload flag to skip shading (we only care about the normal)
 	const uint RayFlags = 0;
 	const uint MissShaderIndex = 0;
@@ -281,21 +264,21 @@ FProbeResult TraceProbeRay(RayDesc Ray, uint2 PixelCoord)
 	FProbeResult Result;
 	Result.HitT = PackedPayload.HitT;
 	Result.WorldNormal = PackedPayload.GetWorldNormal();
-	Result.FrontFace = (PackedPayload.GetFlags() & RAY_TRACING_PAYLOAD_OUTPUT_FLAG_FRONT_FACE) != 0;
+	Result.FrontFace = PackedPayload.IsFrontFace();
 	return Result;
 }
 
-bool ProcessSubsurfaceRandomWalk(inout FMaterialClosestHitPayload Payload, inout float3 PathThroughput, inout RandomSequence RandSequence, uint2 PixelCoord, bool SimplifySSS)
+bool ProcessSubsurfaceRandomWalk(inout FPathTracingPayload Payload, inout float3 PathThroughput, inout RandomSequence RandSequence, uint2 PixelCoord, bool SimplifySSS)
 {
-	if (!IsSubsurfaceMaterial(Payload))
+	if (!Payload.IsSubsurfaceMaterial())
 	{
 		// material doesn't have SSS
 		// just return so we can process the rest of the payload's shading
 		return true;
 	}
 
-	float3 SSSColor = GetSubsurfaceColor(Payload);
-	float3 SSSRadius = GetSubsurfaceRadius(Payload);
+	float3 SSSColor = Payload.GetSubsurfaceColor();
+	float3 SSSRadius = Payload.GetSubsurfaceRadius();
 
 	if (SimplifySSS || all(SSSRadius == 0) || all(SSSColor == 0) || MaxSSSBounces == 0)
 	{
@@ -304,7 +287,7 @@ bool ProcessSubsurfaceRandomWalk(inout FMaterialClosestHitPayload Payload, inout
 		// or if the user decided not to perform any SSS bounces
 		Payload.DiffuseColor += SSSColor;
 		Payload.DiffuseColor = saturate(Payload.DiffuseColor);
-		SetSubsurfaceColor(Payload, 0.0);
+		Payload.SetSubsurfaceColor(0.0);
 		return true;
 	}
 	SSSRadius = max(SSSRadius, 0.1); // 1mm minimum to make sure random walk can move forward and to keep SigmaT finite
@@ -560,7 +543,7 @@ RAY_TRACING_ENTRY_RAYGEN(PathTracingMainRG)
 		const bool bIsLastBounce = Bounce == MaxBounces;
 		const bool bIncludeEmissive = (EnableDirectLighting != 0 || Bounce > 1) &&
 									  (EnableEmissive       != 0 || bIsCameraRay);
-		FMaterialClosestHitPayload Payload = TraceTransparentRay(Ray, RayCone, bIsCameraRay, bIsLastBounce, bIncludeEmissive, LaunchIndex, NumVisibleLights, RandSequence, PathThroughput, Radiance, Albedo, Normal);
+		FPathTracingPayload Payload = TraceTransparentRay(Ray, RayCone, bIsCameraRay, bIsLastBounce, bIncludeEmissive, LaunchIndex, NumVisibleLights, RandSequence, PathThroughput, Radiance, Albedo, Normal);
 
 		if (Payload.IsMiss())
 		{
@@ -616,7 +599,7 @@ RAY_TRACING_ENTRY_RAYGEN(PathTracingMainRG)
 			float3 WorldNormal = Payload.WorldNormal;
 			uint PrimitiveLightingChannelMask = Payload.PrimitiveLightingChannelMask;
 
-			bool IsTransmissiveMaterial = ENABLE_TRANSMISSION && IsMaterialTransmissive(Payload);
+			bool IsTransmissiveMaterial = ENABLE_TRANSMISSION && Payload.IsMaterialTransmissive();
 
 			for (uint Index = 0, Num = LightLoopCount.NumLights; Index < Num; ++Index)
 			{
@@ -790,7 +773,7 @@ RAY_TRACING_ENTRY_RAYGEN(PathTracingMainRG)
 		NumVisibleLights = 0;
 		
 		// update ray cone for the next bounce
-		RayCone = Payload.GetRayCone();
+		RayCone = Payload.RayCone;
 	}
 
 
diff --git a/Engine/Shaders/Private/PathTracing/PathTracingCommon.ush b/Engine/Shaders/Private/PathTracing/PathTracingCommon.ush
--- a/Engine/Shaders/Private/PathTracing/PathTracingCommon.ush
+++ b/Engine/Shaders/Private/PathTracing/PathTracingCommon.ush
@@ -1,16 +1,268 @@
 // Copyright Epic Games, Inc. All Rights Reserved.
 
 /*=============================================================================
-	PathTracingCommon.ush: common header used in multiple ray generation shaders
+	PathTracingCommon.ush: path tracing payload structures
 =============================================================================*/
 
 #pragma once
 
-#ifndef PATHTRACINGCOMMON_USH_INCLUDED
-#define PATHTRACINGCOMMON_USH_INCLUDED // Workarround for UE-66460
-
 #include "../RayTracing/RayTracingCommon.ush"
-#include "PathTracingDefinitions.ush"
 
-#endif // PATHTRACINGCOMMON_USH_INCLUDED // Workarround for UE-66460
+// These flags are analogous to RAY_TRACING_PAYLOAD_* flags. There are currently 6 bits available
+#define PATH_TRACING_PAYLOAD_OUTPUT_FLAG_FRONT_FACE                   (1 << 0) // Indicates that ray has hit the front face of a primitive. Set by closest hit shader.
+
+#define PATH_TRACING_PAYLOAD_INPUT_FLAG_IGNORE_SKY_MATERIALS          (1 << 0) // Indicates that the AHS/CHS should skip processing if the material is tagged as sky
+#define PATH_TRACING_PAYLOAD_INPUT_FLAG_PATHTRACER_VISIBILITY         (1 << 1) // Indicates that this is a path tracer visibility ray (which supports material evaluation for transparent shadows and fake caustics)
+
+
+
+// This payload structure is what we transport between RGS/CHS/AHS programs
+struct FPackedPathTracingPayload : FMinimalPayload
+{
+	// float FMinimalPayload::HitT                       // 4  bytes
+	uint PackedRayCone;                                  // 4  bytes (fp16 width, fp16 spread)
+	uint PackedPixelCoord;                               // 4  bytes
+	uint PackedData[13];								 // 52 bytes (encoded data, depends on shading model and ray type)
+														 // 64 bytes total
+
+	void SetPixelCoord(uint2 PixelCoord) { PackedPixelCoord = (PixelCoord.x & 0xFFFF) | (PixelCoord.y << 16); }
+	uint2 GetPixelCoord() { return uint2(PackedPixelCoord & 0xFFFF, PackedPixelCoord >> 16); }
+
+	FRayCone GetRayCone() { return UnpackRayCone(PackedRayCone); }
+
+	float3 GetWorldNormal()
+	{
+		float3 Result;
+		Result.x = f16tof32(PackedData[1] >> 16);
+		Result.y = f16tof32(PackedData[2]);
+		Result.z = f16tof32(PackedData[2] >> 16);
+		return Result;
+	}
+
+	uint GetShadingModelID() { return (PackedData[7] >> 16) & 0xF; }
+	uint GetBlendingMode() { return (PackedData[7] >> 20) & 0x7; }
+	uint GetFlags() { return (PackedData[7] >> 23) & 0x3F; }
+	void SetFlag(uint Flag) { PackedData[7] |= Flag << 23; }
+	uint GetPrimitiveLightingChannelMask() { return (PackedData[7] >> 29) & 0x7; }
+
+	// Flag methods
+	void SetIgnoreSkyMaterials() { SetFlag(PATH_TRACING_PAYLOAD_INPUT_FLAG_IGNORE_SKY_MATERIALS); }
+	bool IsIgnoreSkyMaterials() { return (GetFlags() & PATH_TRACING_PAYLOAD_INPUT_FLAG_IGNORE_SKY_MATERIALS) != 0; }
+
+	void SetVisibilityRay() { SetFlag(PATH_TRACING_PAYLOAD_INPUT_FLAG_PATHTRACER_VISIBILITY); }
+	bool IsVisibilityRay() { return (GetFlags() & PATH_TRACING_PAYLOAD_INPUT_FLAG_PATHTRACER_VISIBILITY) != 0; }
+
+	bool IsFrontFace() { return (GetFlags() & PATH_TRACING_PAYLOAD_OUTPUT_FLAG_FRONT_FACE) != 0; }
+
+
+	// These method are meant to be used only when IsVisibilityRay() is true
+	float3 GetRayThroughput()
+	{
+		return float3(
+			asfloat(PackedData[0]),
+			asfloat(PackedData[1]),
+			asfloat(PackedData[2])
+		);
+	}
+
+	void SetRayThroughput(float3 RayThroughput)
+	{
+		PackedData[0] = asuint(RayThroughput.x);
+		PackedData[1] = asuint(RayThroughput.y);
+		PackedData[2] = asuint(RayThroughput.z);
+	}
+
+	float GetPathRoughness()
+	{
+		return asfloat(PackedData[3]);
+	}
+
+	void SetPathRoughness(float PathRoughness)
+	{
+		PackedData[3] = asuint(PathRoughness);
+	}
+};
+
+// This payload structure is the expanded version of the above which is more convenient to work with
+struct FPathTracingPayload : FMinimalPayload
+{
+	FRayCone RayCone;
+	float3 Radiance;
+	float3 WorldNormal;
+	float3 BaseColor;
+	float3 DiffuseColor;
+	float3 SpecularColor;
+	float  Opacity;
+	float  Metallic;
+	float  Specular;
+	float  Roughness;
+	float  Ior;
+	uint   ShadingModelID;
+	uint   BlendingMode;
+	uint   Flags;
+	uint   PrimitiveLightingChannelMask;
+	float4 CustomData;
+	float3 CustomVector;
+	float3 WorldTangent;
+	float  Anisotropy;
+	float3 WorldPos;
+
+	void SetFrontFace() { Flags |= PATH_TRACING_PAYLOAD_OUTPUT_FLAG_FRONT_FACE; }
+	bool IsFrontFace() { return (Flags & PATH_TRACING_PAYLOAD_OUTPUT_FLAG_FRONT_FACE) != 0; }
+
+
+	// Various ways to interpret CustomData (depending on ShadingModelID)
+	float GetClearCoat() { return CustomData.x; }
+	float GetClearCoatRoughness() { return CustomData.y; }
+	void SetClearCoatRoughness(float ClearCoatRoughness) { CustomData.y = ClearCoatRoughness; }
+	float3 GetClearCoatBottomNormal() { return CustomVector; }
+
+	float3 GetSubsurfaceColor() { return CustomData.xyz; }
+	void SetSubsurfaceColor(float3 SubsurfaceColor) { CustomData.xyz = SubsurfaceColor; }
+	
+	float3 GetSubsurfaceRadius() { return CustomVector.xyz; }
+	void SetSubsurfaceRadius(float3 SubsurfaceRadius) { CustomVector.xyz = SubsurfaceRadius; }
+
+	// Steal an unused float3 (the only material with dual specular does not support anisotropy)
+	// Data is encoded as (x: roughness 0, y: roughness 1, z: roughness mix)
+	float3 GetDualRoughnessSpecular() { return WorldTangent; }
+	void SetDualRoughnessSpecular(float3 RoughnessData)
+	{
+		// store data in unused field and compute average roughness
+		WorldTangent = RoughnessData;
+		Roughness = lerp(RoughnessData.x, RoughnessData.y, RoughnessData.z);
+	}
+
+	bool IsMaterialTransmissive()
+	{
+		return ShadingModelID == SHADINGMODELID_TWOSIDED_FOLIAGE || (BlendingMode == RAY_TRACING_BLEND_MODE_TRANSLUCENT && Ior != 1.0 && Roughness != 0 && Opacity < 1);
+	}
+
+	bool IsSubsurfaceMaterial()
+	{
+		return ShadingModelID == SHADINGMODELID_SUBSURFACE ||
+			   ShadingModelID == SHADINGMODELID_PREINTEGRATED_SKIN ||
+			   ShadingModelID == SHADINGMODELID_SUBSURFACE_PROFILE;
+	}
+};
+
+FPackedPathTracingPayload InitPathTracingPayload(FRayCone RayCone, uint2 PixelCoord, bool bIgnoreSkyMaterials)
+{
+	FPackedPathTracingPayload Output = (FPackedPathTracingPayload)0;
+	Output.PackedRayCone = PackRayCone(RayCone);
+	Output.SetPixelCoord(PixelCoord);
+	if (bIgnoreSkyMaterials)
+	{
+		Output.SetIgnoreSkyMaterials();
+	}
+	return Output;
+}
+
+FPackedPathTracingPayload InitPathTracingVisibilityPayload(uint2 PixelCoord, float PathRoughness)
+{
+	FPackedPathTracingPayload Output = (FPackedPathTracingPayload)0;
+	Output.SetPixelCoord(PixelCoord);
+	// Signal to the AHS we want to evaluate the opacity
+	// The payload is used to carry the path throughput (for transparent shadows)
+	// and current path roughness (for approximate caustics)
+	Output.SetVisibilityRay();
+	Output.SetPathRoughness(PathRoughness);
+	Output.SetRayThroughput(1.0);
+	return Output;
+}
+
+
+FPackedPathTracingPayload PackPathTracingPayload(FPathTracingPayload Input)
+{
+	FPackedPathTracingPayload Output = (FPackedPathTracingPayload)0;
+	Output.HitT = Input.HitT;
+	Output.PackedRayCone = PackRayCone(Input.RayCone);
+
+	Output.PackedData[0] = f32tof16(Input.Radiance.x);
+	Output.PackedData[0] |= f32tof16(Input.Radiance.y) << 16;
+	Output.PackedData[1] = f32tof16(Input.Radiance.z);
+
+	Output.PackedData[1] |= f32tof16(Input.WorldNormal.x) << 16;
+	Output.PackedData[2] = f32tof16(Input.WorldNormal.y);
+	Output.PackedData[2] |= f32tof16(Input.WorldNormal.z) << 16;
+
+	Output.PackedData[3] = f32tof16(Input.BaseColor.x);
+	Output.PackedData[3] |= f32tof16(Input.BaseColor.y) << 16;
+	Output.PackedData[4] = f32tof16(Input.BaseColor.z);
+
+	Output.PackedData[4] |= f32tof16(Input.Opacity) << 16;
+
+	Output.PackedData[5] = f32tof16(Input.Metallic);
+	Output.PackedData[5] |= f32tof16(Input.Specular) << 16;
+	Output.PackedData[6] = f32tof16(Input.Roughness);
+	Output.PackedData[6] |= f32tof16(Input.Anisotropy) << 16;
+
+	Output.PackedData[7] = f32tof16(Input.Ior);                              // 16 bits
+	Output.PackedData[7] |= (Input.ShadingModelID & 0xF) << 16;               // 4 bits
+	Output.PackedData[7] |= (Input.BlendingMode & 0x7) << 20;                 // 3 bits
+	Output.PackedData[7] |= (Input.Flags & 0x3F) << 23;                       // 6 bits
+	Output.PackedData[7] |= (Input.PrimitiveLightingChannelMask & 0x7) << 29; // 3 bits
+                                                                       // total 32 bits
+
+	Output.PackedData[8] = f32tof16(Input.WorldTangent.x);
+	Output.PackedData[8] |= f32tof16(Input.WorldTangent.y) << 16;
+	Output.PackedData[9] = f32tof16(Input.WorldTangent.z);
+
+	Output.PackedData[9] |= f32tof16(Input.CustomVector.x) << 16;
+	Output.PackedData[10] = f32tof16(Input.CustomVector.y);
+	Output.PackedData[10] |= f32tof16(Input.CustomVector.z) << 16;
+
+	Output.PackedData[11] = f32tof16(Input.CustomData.x);
+	Output.PackedData[11] |= f32tof16(Input.CustomData.y) << 16;
+	Output.PackedData[12] = f32tof16(Input.CustomData.z);
+	Output.PackedData[12] |= f32tof16(Input.CustomData.w) << 16;
+
+	return Output;
+}
+
+FPathTracingPayload UnpackPathTracingPayload(FPackedPathTracingPayload Input, RayDesc Ray)
+{
+	FPathTracingPayload Output = (FPathTracingPayload)0;
+
+	Output.HitT = Input.HitT;
+	Output.RayCone = Input.GetRayCone();
+	Output.Radiance.x = f16tof32(Input.PackedData[0]);
+	Output.Radiance.y = f16tof32(Input.PackedData[0] >> 16);
+	Output.Radiance.z = f16tof32(Input.PackedData[1]);
+	Output.WorldNormal.x = f16tof32(Input.PackedData[1] >> 16);
+	Output.WorldNormal.y = f16tof32(Input.PackedData[2]);
+	Output.WorldNormal.z = f16tof32(Input.PackedData[2] >> 16);
+	Output.BaseColor.x = f16tof32(Input.PackedData[3]);
+	Output.BaseColor.y = f16tof32(Input.PackedData[3] >> 16);
+	Output.BaseColor.z = f16tof32(Input.PackedData[4]);
+	Output.Opacity = f16tof32(Input.PackedData[4] >> 16);
+	Output.Metallic = f16tof32(Input.PackedData[5]);
+	Output.Specular = f16tof32(Input.PackedData[5] >> 16);
+	Output.Roughness = f16tof32(Input.PackedData[6]);
+	Output.Anisotropy = f16tof32(Input.PackedData[6] >> 16);
+	Output.Ior = f16tof32(Input.PackedData[7]);
+	Output.ShadingModelID = Input.GetShadingModelID();
+	Output.BlendingMode = Input.GetBlendingMode();
+	Output.Flags = Input.GetFlags();
+	Output.PrimitiveLightingChannelMask = Input.GetPrimitiveLightingChannelMask();
+
+	Output.WorldTangent.x = f16tof32(Input.PackedData[8]);
+	Output.WorldTangent.y = f16tof32(Input.PackedData[8] >> 16);
+	Output.WorldTangent.z = f16tof32(Input.PackedData[9]);
+
+	Output.CustomVector.x = f16tof32(Input.PackedData[9] >> 16);
+	Output.CustomVector.y = f16tof32(Input.PackedData[10]);
+	Output.CustomVector.z = f16tof32(Input.PackedData[10] >> 16);
+
+	Output.CustomData.x = f16tof32(Input.PackedData[11]);
+	Output.CustomData.y = f16tof32(Input.PackedData[11] >> 16);
+	Output.CustomData.z = f16tof32(Input.PackedData[12]);
+	Output.CustomData.w = f16tof32(Input.PackedData[12] >> 16);
+
+	Output.WorldPos = Ray.Origin + Output.HitT * Ray.Direction;
+
+	Output.DiffuseColor = Output.BaseColor - Output.BaseColor * Output.Metallic;
+	Output.SpecularColor = ComputeF0(Output.Specular, Output.BaseColor, Output.Metallic);
 
+	return Output;
+}
diff --git a/Engine/Shaders/Private/PathTracing/PathTracingDefinitions.ush b/Engine/Shaders/Private/PathTracing/PathTracingDefinitions.ush
deleted file mode 100644
diff --git a/Engine/Shaders/Private/PathTracing/PathTracingMaterialHitShader.usf b/Engine/Shaders/Private/PathTracing/PathTracingMaterialHitShader.usf
--- a/Engine/Shaders/Private/PathTracing/PathTracingMaterialHitShader.usf
+++ b/Engine/Shaders/Private/PathTracing/PathTracingMaterialHitShader.usf
@@ -32,14 +32,12 @@
 
 
 RAY_TRACING_ENTRY_CLOSEST_HIT(PathTracingMaterialCHS,
-	FPackedMaterialClosestHitPayload, PackedPayload,
+	FPackedPathTracingPayload, PackedPayload,
 	FDefaultAttributes, Attributes)
 {
 	PackedPayload.HitT = RayTCurrent();
 
 #if USE_MATERIAL_CLOSEST_HIT_SHADER
-	GlobalTextureMipBias = PackedPayload.GetMipBias();
-
 	const uint2 PixelCoord = PackedPayload.GetPixelCoord();
 
 	ResolvedView = ResolveView();
@@ -67,7 +65,8 @@ RAY_TRACING_ENTRY_CLOSEST_HIT(PathTracingMaterialCHS,
 		// #endif
 	}
 
-	FMaterialClosestHitPayload Payload = (FMaterialClosestHitPayload)0;
+	FPathTracingPayload Payload = (FPathTracingPayload)0;
+	Payload.RayCone = PropagatedCone;
 
 	/**
 	 * Set common material attributes for both full and simplified materials
@@ -100,7 +99,7 @@ RAY_TRACING_ENTRY_CLOSEST_HIT(PathTracingMaterialCHS,
 #if MATERIAL_IS_SKY
 	if (PackedPayload.IsIgnoreSkyMaterials())
 	{
-		PackedPayload = PackRayTracingPayload(Payload, PropagatedCone);
+		PackedPayload = PackPathTracingPayload(Payload);
 		return;
 	}
 #endif
@@ -108,9 +107,9 @@ RAY_TRACING_ENTRY_CLOSEST_HIT(PathTracingMaterialCHS,
 	// Store the results in local variables and reuse instead of calling the functions multiple times.
 	half3 BaseColor = GetMaterialBaseColor(PixelMaterialInputs);
 #if MATERIALBLENDING_TRANSLUCENT && REFRACTION_USE_INDEX_OF_REFRACTION
-	half2 Refraction = GetMaterialRefraction(PixelMaterialInputs);
+	half Refraction = GetMaterialRefraction(PixelMaterialInputs).x;
 #else
-	half2 Refraction = 1;
+	half Refraction = 1;
 #endif
 	half  Metallic = GetMaterialMetallic(PixelMaterialInputs);
 	half  Specular = GetMaterialSpecular(PixelMaterialInputs);
@@ -130,6 +129,11 @@ RAY_TRACING_ENTRY_CLOSEST_HIT(PathTracingMaterialCHS,
 	Payload.Metallic = Metallic;
 	Payload.Opacity = GetMaterialOpacity(PixelMaterialInputs);
 
+#if MATERIAL_USES_ANISOTROPY
+	Payload.WorldTangent = CalculateAnisotropyTangent(MaterialParameters, PixelMaterialInputs);
+	Payload.Anisotropy = GetMaterialAnisotropy(PixelMaterialInputs);
+#endif	
+
 #if MATERIALBLENDING_MASKED
 #if MATERIAL_DITHER_OPACITY_MASK
 	// dithering emulates real transparency, so switch to translucent
@@ -144,72 +148,63 @@ RAY_TRACING_ENTRY_CLOSEST_HIT(PathTracingMaterialCHS,
 #endif
 #endif // MATERIALBLENDING_MASKED
 
-	Payload.Ior = Refraction.x == 0 ? 1.0 : Refraction.x;
-	Payload.CustomData = float4(GetMaterialCustomData0(MaterialParameters), GetMaterialCustomData1(MaterialParameters), 0, 0);
+	Payload.Ior = Refraction == 0 ? 1.0 : Refraction;
 
 #if MATERIAL_SHADINGMODEL_CLEAR_COAT
-#if CLEAR_COAT_BOTTOM_NORMAL
 	if (Payload.ShadingModelID == SHADINGMODELID_CLEAR_COAT)
 	{
-		float2 oct2 = UnitVectorToOctahedron(Payload.WorldNormal);
-
+    	Payload.CustomData.x = GetMaterialCustomData0(MaterialParameters);
+		Payload.CustomData.y = GetMaterialCustomData1(MaterialParameters);
+#if CLEAR_COAT_BOTTOM_NORMAL
 #if NUM_MATERIAL_OUTPUTS_CLEARCOATBOTTOMNORMAL > 0
 #if MATERIAL_TANGENTSPACENORMAL
 		float3 tempnormal = normalize(TransformTangentVectorToWorld(MaterialParameters.TangentToWorld, ClearCoatBottomNormal0(MaterialParameters)));
 #else
 		float3 tempnormal = ClearCoatBottomNormal0(MaterialParameters);
 #endif
-
-		float2 oct1 = UnitVectorToOctahedron(tempnormal);
-		float2 oct3 = ((oct1 - oct2) * 0.5) + (128.0 / 255.0);
-		Payload.CustomData.a = oct3.x;
-		Payload.CustomData.z = oct3.y;
+		Payload.CustomVector = tempnormal;
 #else
-		Payload.CustomData.a = 128.0 / 255.0;
-		Payload.CustomData.z = 128.0 / 255.0;
+		Payload.CustomVector = Payload.WorldNormal;
 #endif
-	}
+#else
+		Payload.CustomVector = Payload.WorldNormal;
 #endif
+	}
+	else
 #endif
 
-	// Override custom data if sub-surface material
-#if MATERIAL_SHADINGMODEL_SUBSURFACE || MATERIAL_SHADINGMODEL_TWOSIDED_FOLIAGE || MATERIAL_SHADINGMODEL_CLOTH || MATERIAL_SHADINGMODEL_CLOTH || MATERIAL_SHADINGMODEL_SUBSURFACE_PROFILE || MATERIAL_SHADINGMODEL_EYE
-	if (Payload.ShadingModelID == SHADINGMODELID_SUBSURFACE ||
-		Payload.ShadingModelID == SHADINGMODELID_TWOSIDED_FOLIAGE ||
-		Payload.ShadingModelID == SHADINGMODELID_CLOTH ||
-		Payload.ShadingModelID == SHADINGMODELID_SUBSURFACE_PROFILE |
-		Payload.ShadingModelID == SHADINGMODELID_EYE)
+#if MATERIAL_SHADINGMODEL_CLOTH
+	if (Payload.ShadingModelID == SHADINGMODELID_CLOTH)
 	{
 		Payload.CustomData = GetMaterialSubsurfaceData(PixelMaterialInputs);
+		Payload.CustomData.a = saturate(GetMaterialCustomData0(MaterialParameters)); // Fuzz
 	}
+	else
 #endif
 
-#if MATERIAL_SHADINGMODEL_CLOTH
-	if (Payload.ShadingModelID == SHADINGMODELID_CLOTH)
+#if MATERIAL_SHADINGMODEL_TWOSIDED_FOLIAGE || MATERIAL_SHADINGMODEL_SUBSURFACE || MATERIAL_SHADINGMODEL_PREINTEGRATED_SKIN || MATERIAL_SHADINGMODEL_SUBSURFACE_PROFILE || MATERIAL_SHADINGMODEL_EYE
+	if (Payload.ShadingModelID == SHADINGMODELID_TWOSIDED_FOLIAGE ||
+		Payload.ShadingModelID == SHADINGMODELID_SUBSURFACE ||
+		Payload.ShadingModelID == SHADINGMODELID_PREINTEGRATED_SKIN ||
+		Payload.ShadingModelID == SHADINGMODELID_SUBSURFACE_PROFILE |
+		Payload.ShadingModelID == SHADINGMODELID_EYE)
 	{
-		Payload.CustomData.a = saturate(GetMaterialCustomData0(MaterialParameters));
+		Payload.CustomData = GetMaterialSubsurfaceData(PixelMaterialInputs);
 	}
+	else
 #endif
-
 #if MATERIAL_SHADINGMODEL_THIN_TRANSLUCENT
 	if (Payload.ShadingModelID == SHADINGMODELID_THIN_TRANSLUCENT)
 	{
 		Payload.CustomData.xyz = GetThinTranslucentMaterialOutput0(MaterialParameters);
 	}
+	else
 #endif
+	{
+		// terminal case for when none of the above worked
+	}
 
-#if MATERIAL_USES_ANISOTROPY
-	Payload.WorldTangent = CalculateAnisotropyTangent(MaterialParameters, PixelMaterialInputs);
-	Payload.Anisotropy = GetMaterialAnisotropy(PixelMaterialInputs);
-#else
-	Payload.WorldTangent = 0;
-	Payload.Anisotropy = 0.0f;
-#endif	
-
-	Payload.DiffuseColor = BaseColor - BaseColor * Metallic;
-	Payload.SpecularColor = ComputeF0(Specular, BaseColor, Metallic);
-
-#else
+#else // SIMPLIFIED_MATERIAL_SHADER
 
 	/**
 	 * Set material attributes for simplified materials
@@ -223,21 +218,21 @@ RAY_TRACING_ENTRY_CLOSEST_HIT(PathTracingMaterialCHS,
 	else
 #endif
 #if MATERIAL_SHADINGMODEL_THIN_TRANSLUCENT
-		if (Payload.ShadingModelID == SHADINGMODELID_THIN_TRANSLUCENT)
-		{
-			Payload.CustomData.xyz = GetThinTranslucentMaterialOutput0(MaterialParameters);
-		}
-		else
+	if (Payload.ShadingModelID == SHADINGMODELID_THIN_TRANSLUCENT)
+	{
+		Payload.CustomData.xyz = GetThinTranslucentMaterialOutput0(MaterialParameters);
+	}
+	else
 #endif
 #if MATERIAL_SHADINGMODEL_UNLIT
-			if (Payload.ShadingModelID == SHADINGMODELID_UNLIT)
-			{
-			}
-			else
+	if (Payload.ShadingModelID == SHADINGMODELID_UNLIT)
+	{
+	}
+	else
 #endif
-			{
-				Payload.ShadingModelID = SHADINGMODELID_DEFAULT_LIT;
-			}
+	{
+		Payload.ShadingModelID = SHADINGMODELID_DEFAULT_LIT;
+	}
 	Payload.WorldPos = MaterialParameters.AbsoluteWorldPosition;
 	Payload.WorldNormal = normalize(MaterialParameters.WorldNormal);
 	Payload.Radiance = GetMaterialEmissive(PixelMaterialInputs);
@@ -250,12 +245,12 @@ RAY_TRACING_ENTRY_CLOSEST_HIT(PathTracingMaterialCHS,
 
 #endif
 
-	PackedPayload = PackRayTracingPayload(Payload, PropagatedCone);
+	PackedPayload = PackPathTracingPayload(Payload);
 #endif // USE_MATERIAL_CLOSEST_HIT_SHADER
 }
 
 RAY_TRACING_ENTRY_ANY_HIT(PathTracingMaterialAHS,
-	FPackedMaterialClosestHitPayload, PackedPayload,
+	FPackedPathTracingPayload, PackedPayload,
 	FDefaultAttributes, Attributes)
 {
 #if USE_MATERIAL_ANY_HIT_SHADER
@@ -265,7 +260,7 @@ RAY_TRACING_ENTRY_ANY_HIT(PathTracingMaterialAHS,
 	// Usually this path is inactive because opaque materials don't have AHS bound
 	return;
 #elif MATERIALBLENDING_ALPHAHOLDOUT || MATERIALBLENDING_ADDITIVE
-	if (PackedPayload.IsPathTracingVisibilityRay())
+	if (PackedPayload.IsVisibilityRay())
 	{
 		// path tracer wants to treat these cases as fully transparent
 		// no material evaluation is required
@@ -281,15 +276,7 @@ RAY_TRACING_ENTRY_ANY_HIT(PathTracingMaterialAHS,
 	// This ensures that both material rays and shadow rays "see" the same thing
 #else
 	// For any other blending models, we have a few options:
-	if (PackedPayload.IsIgnoreTranslucentMaterials())
-	{
-		// special mode used by regular RT shadows - translucent materials are skipped, everything else is treated as opaque
-#if MATERIALBLENDING_TRANSLUCENT || (MATERIALBLENDING_MASKED && MATERIAL_DITHER_OPACITY_MASK)
-		IgnoreHit();
-#endif
-		return;
-	}
-	else if (PackedPayload.IsPathTracingVisibilityRay())
+	if (PackedPayload.IsVisibilityRay())
 	{
 		// this case is used by the path tracer only currently
 		// We want to evaluate the material opacity and accumulate it
@@ -304,8 +291,6 @@ RAY_TRACING_ENTRY_ANY_HIT(PathTracingMaterialAHS,
 
 	// If we got here -- we have a reason to evaluate the material -- do the setup 
 
-	GlobalTextureMipBias = PackedPayload.GetMipBias();
-
 	const uint2 PixelCoord = PackedPayload.GetPixelCoord();
 
 	ResolvedView = ResolveView();
diff --git a/Engine/Shaders/Private/RayTracing/RayTracingCommon.ush b/Engine/Shaders/Private/RayTracing/RayTracingCommon.ush
--- a/Engine/Shaders/Private/RayTracing/RayTracingCommon.ush
+++ b/Engine/Shaders/Private/RayTracing/RayTracingCommon.ush
@@ -268,15 +268,21 @@ uint2 GetPixelCoord(uint2 DispatchThreadId, uint UpscaleFactor)
 //   - "inputs" are set before calling TraceRay
 //   - "outputs" are set during execution of the CHS and are valid after TraceRay returns
 // These two sets of masks can overlap because they never co-exist
+#define RAY_TRACING_PAYLOAD_OUTPUT_FLAG_NONE                               0
 #define RAY_TRACING_PAYLOAD_OUTPUT_FLAG_FRONT_FACE                   (1 << 0) // Indicates that ray has hit front face of a primitive. Set by closest hit shader.
+#define RAY_TRACING_PAYLOAD_OUTPUT_FLAG_UNUSED1                      (1 << 1) // Free for future use
+#define RAY_TRACING_PAYLOAD_OUTPUT_FLAG_UNUSED2                      (1 << 2) // Free for future use
+#define RAY_TRACING_PAYLOAD_OUTPUT_FLAG_UNUSED3                      (1 << 3) // Free for future use
+#define RAY_TRACING_PAYLOAD_OUTPUT_FLAG_UNUSED4                      (1 << 4) // Free for future use
+#define RAY_TRACING_PAYLOAD_OUTPUT_FLAG_UNUSED5                      (1 << 5) // Free for future use
 
 #define RAY_TRACING_PAYLOAD_INPUT_FLAG_NONE								   0	
-#define RAY_TRACING_PAYLOAD_INPUT_FLAG_IGNORE_SKY_MATERIALS          (1 << 0) // Indicates that the AHS/CHS should skip processing if the material is tagged as sky
-#define RAY_TRACING_PAYLOAD_INPUT_FLAG_MINIMAL_PAYLOAD               (1 << 1) // Indicates that closest hit shader should only fill FMinimalPayload (HitT), skipping full material evaluation. Set by RayGen shader, before TraceRay().
-#define RAY_TRACING_PAYLOAD_INPUT_FLAG_ENABLE_SKY_LIGHT_CONTRIBUTION (1 << 2) // Indicates that hit shaders should add in the Sky Light contribution in their output. Set by RayGen shader, before TraceRay().
-#define RAY_TRACING_PAYLOAD_INPUT_FLAG_LUMEN_PAYLOAD		         (1 << 3) // Indicates that closest hit shader should evaluate geometry normal and instanceid.
-#define RAY_TRACING_PAYLOAD_INPUT_FLAG_IGNORE_TRANSLUCENT            (1 << 4) // Indicates that translucent materials should be ignored by IgnoreHit() in anyhit shaders.
-#define RAY_TRACING_PAYLOAD_INPUT_FLAG_PATHTRACER_VISIBILITY         (1 << 5) // Indicates that this is a path tracer visibility ray (which supports material evaluation for transparent shadows and fake caustics)
+#define RAY_TRACING_PAYLOAD_INPUT_FLAG_MINIMAL_PAYLOAD               (1 << 0) // Indicates that closest hit shader should only fill FMinimalPayload (HitT), skipping full material evaluation. Set by RayGen shader, before TraceRay().
+#define RAY_TRACING_PAYLOAD_INPUT_FLAG_ENABLE_SKY_LIGHT_CONTRIBUTION (1 << 1) // Indicates that hit shaders should add in the Sky Light contribution in their output. Set by RayGen shader, before TraceRay().
+#define RAY_TRACING_PAYLOAD_INPUT_FLAG_LUMEN_PAYLOAD		         (1 << 2) // Indicates that closest hit shader should evaluate geometry normal and instanceid.
+#define RAY_TRACING_PAYLOAD_INPUT_FLAG_IGNORE_TRANSLUCENT            (1 << 3) // Indicates that translucent materials should be ignored by IgnoreHit() in anyhit shaders.
+#define RAY_TRACING_PAYLOAD_INPUT_FLAG_UNUSED4                       (1 << 4) // Free for future use
+#define RAY_TRACING_PAYLOAD_INPUT_FLAG_UNUSED5                       (1 << 5) // Free for future use
 
 struct FMaterialClosestHitPayload : FMinimalPayload
 {
@@ -392,9 +398,6 @@ struct FPackedMaterialClosestHitPayload : FMinimalPayload
 	void SetPixelCoord(uint2 PixelCoord) { PackedPixelCoord = (PixelCoord.x & 0xFFFF) | (PixelCoord.y << 16); }
 	uint2 GetPixelCoord() { return uint2(PackedPixelCoord & 0xFFFF, PackedPixelCoord >> 16); }
 
-	void SetIgnoreSkyMaterials() { IorAndShadingModelIDAndBlendingModeAndFlagsAndPrimitiveLightingChannelMask |= RAY_TRACING_PAYLOAD_INPUT_FLAG_IGNORE_SKY_MATERIALS << 23; }
-	bool IsIgnoreSkyMaterials() { return ((IorAndShadingModelIDAndBlendingModeAndFlagsAndPrimitiveLightingChannelMask >> 23) & RAY_TRACING_PAYLOAD_INPUT_FLAG_IGNORE_SKY_MATERIALS) != 0; }
-
 	void SetMinimalPayloadMode() { IorAndShadingModelIDAndBlendingModeAndFlagsAndPrimitiveLightingChannelMask |= RAY_TRACING_PAYLOAD_INPUT_FLAG_MINIMAL_PAYLOAD << 23; }
 	bool IsMinimalPayloadMode() { return ((IorAndShadingModelIDAndBlendingModeAndFlagsAndPrimitiveLightingChannelMask >> 23) & RAY_TRACING_PAYLOAD_INPUT_FLAG_MINIMAL_PAYLOAD) != 0; }
 
@@ -407,9 +410,6 @@ struct FPackedMaterialClosestHitPayload : FMinimalPayload
 	void SetIgnoreTranslucentMaterials() { IorAndShadingModelIDAndBlendingModeAndFlagsAndPrimitiveLightingChannelMask |= RAY_TRACING_PAYLOAD_INPUT_FLAG_IGNORE_TRANSLUCENT << 23; }
 	bool IsIgnoreTranslucentMaterials() { return ((IorAndShadingModelIDAndBlendingModeAndFlagsAndPrimitiveLightingChannelMask >> 23) & RAY_TRACING_PAYLOAD_INPUT_FLAG_IGNORE_TRANSLUCENT) != 0; }
 
-	void SetPathTracingVisibilityRay() { IorAndShadingModelIDAndBlendingModeAndFlagsAndPrimitiveLightingChannelMask |= RAY_TRACING_PAYLOAD_INPUT_FLAG_PATHTRACER_VISIBILITY << 23; }
-	bool IsPathTracingVisibilityRay() { return ((IorAndShadingModelIDAndBlendingModeAndFlagsAndPrimitiveLightingChannelMask >> 23) & RAY_TRACING_PAYLOAD_INPUT_FLAG_PATHTRACER_VISIBILITY) != 0; }
-
 	FRayCone GetRayCone() { return UnpackRayCone(PackedRayCone); }
 	void SetRayCone(FRayCone NewRayCone) { PackedRayCone = PackRayCone(NewRayCone); }
 
@@ -515,33 +515,6 @@ struct FPackedMaterialClosestHitPayload : FMinimalPayload
 
 	int GetSceneInstanceIndex() { return PackedCustomData; }
 	float3 GetGeometryNormal() { return GetWorldNormal(); }
-
-	// These method are meant to be used only when IsPathTracingVisibilityRay() is true
-	float3 GetRayThroughput()
-	{
-		return float3(
-			asfloat(RadianceAndNormal[0]),
-			asfloat(RadianceAndNormal[1]),
-			asfloat(RadianceAndNormal[2])
-		);
-	}
-
-	void SetRayThroughput(float3 RayThroughput)
-	{
-		RadianceAndNormal[0] = asuint(RayThroughput.x);
-		RadianceAndNormal[1] = asuint(RayThroughput.y);
-		RadianceAndNormal[2] = asuint(RayThroughput.z);
-	}
-
-	float GetPathRoughness()
-	{
-		return asfloat(MetallicAndSpecularAndRoughness);
-	}
-
-	void SetPathRoughness(float PathRoughness)
-	{
-		MetallicAndSpecularAndRoughness = asuint(PathRoughness);
-	}
 };
 
 FPackedMaterialClosestHitPayload PackRayTracingPayload(FMaterialClosestHitPayload Input, in FRayCone RayCone)
