Fixes TSR's DEBUG_FALLBACK_BLENDING

diff --git a/Engine/Shaders/Private/TemporalAA/TAAUpdateHistory.usf b/Engine/Shaders/Private/TemporalAA/TAAUpdateHistory.usf
--- a/Engine/Shaders/Private/TemporalAA/TAAUpdateHistory.usf
+++ b/Engine/Shaders/Private/TemporalAA/TAAUpdateHistory.usf
@@ -894,7 +894,7 @@ void MainCS(
 			taa_half3 FallbackColor = taa_half3(1, 0.25, 0.25);
 			taa_half3 HighFrequencyColor = taa_half3(0.25, 1, 0.25);
 
-			FinalOutputColor = FinalFallbackColor * Weights.x * FallbackColor + NormalizedFinalHistoryMoment1 * Weights.y * HighFrequencyColor;
+			FinalOutputColor = dpv_scale(FinalFallbackColor * dpv_interleave_mono_registers(FallbackColor), Weights[0]) + dpv_scale(NormalizedFinalHistoryMoment1 * dpv_interleave_mono_registers(HighFrequencyColor), Weights[1]);
 		#elif DEBUG_LOW_FREQUENCY_REJECTION
 			taa_half3 DebugColor = lerp(taa_half3(1, 0.5, 0.5), taa_half3(0.5, 1, 0.5), LowFrequencyRejection);
 			
