Separates TSR source code from TAA

diff --git a/Engine/Shaders/Private/TemporalSuperResolution/TSRClearPrevTextures.usf b/Engine/Shaders/Private/TemporalSuperResolution/TSRClearPrevTextures.usf
new file mode 100644
--- /dev/null
+++ b/Engine/Shaders/Private/TemporalSuperResolution/TSRClearPrevTextures.usf
@@ -0,0 +1,35 @@
+// Copyright Epic Games, Inc. All Rights Reserved.
+
+#include "TSRCommon.ush"
+
+
+//------------------------------------------------------- CONFIG
+
+#define TILE_SIZE 8
+
+
+//------------------------------------------------------- PARAMETERS
+
+RWTexture2D<uint> PrevUseCountOutput;
+RWTexture2D<uint> PrevClosestDepthOutput;
+
+
+//------------------------------------------------------- ENTRY POINT
+
+[numthreads(TILE_SIZE, TILE_SIZE, 1)]
+void MainCS(
+	uint2 GroupId : SV_GroupID,
+	uint GroupThreadIndex : SV_GroupIndex)
+{
+	uint2 DispatchThreadId = (
+		ZOrder2D(GroupThreadIndex, log2(TILE_SIZE)) +
+		GroupId * uint2(TILE_SIZE, TILE_SIZE));
+
+	uint2 InputPixelPos = InputInfo_ViewportMin + DispatchThreadId;
+	
+	if (all(InputPixelPos < InputInfo_ViewportMax))
+	{
+		PrevUseCountOutput[InputPixelPos] = 0;
+		PrevClosestDepthOutput[InputPixelPos] = 0;
+	}
+}
diff --git a/Engine/Shaders/Private/TemporalSuperResolution/TSRCommon.ush b/Engine/Shaders/Private/TemporalSuperResolution/TSRCommon.ush
new file mode 100644
--- /dev/null
+++ b/Engine/Shaders/Private/TemporalSuperResolution/TSRCommon.ush
@@ -0,0 +1,970 @@
+// Copyright Epic Games, Inc. All Rights Reserved.
+
+#include "../Common.ush"
+#include "../Random.ush"
+#include "../TextureSampling.ush"
+#include "../FastMath.ush"
+#include "../MonteCarlo.ush"
+#include "../ScreenPass.ush"
+#include "/Engine/Public/DualPixelVectorization.ush"
+
+
+//------------------------------------------------------- COMPILER CONFIG
+
+// Generate vector truncation warnings to errors.
+#pragma warning(error: 3206)
+
+
+//------------------------------------------------------- CONFIG
+
+#define DEBUG_OUTPUT 0
+
+#if PLATFORM_SUPPORTS_REAL_TYPES && 1
+	#define CONFIG_COMPILE_FP16 1
+#else
+	#define CONFIG_COMPILE_FP16 0
+#endif
+
+#define CONFIG_ENABLE_STOCASTIC_QUANTIZATION 1
+
+#if CONFIG_COMPILE_FP16
+	// Take advantage of RDNA's v_pk_*_{uif}16 instructions
+	#define CONFIG_ENABLE_DUAL_PIXEL_VECTORIZATION 1
+#else
+	#define CONFIG_ENABLE_DUAL_PIXEL_VECTORIZATION 0
+#endif
+
+
+//------------------------------------------------------- CONSTANTS
+
+/* Maximum number of sample. */
+#define MAX_SAMPLE_COUNT 8
+
+#define MAX_FALLBACK_SAMPLE_COUNT 20
+
+#define COMPRESS_PREV_USE_COUNT 1
+#if COMPRESS_PREV_USE_COUNT
+	// COMPRESS_PREV_USE_COUNT compress the use count in 8bit, so need to sum of a 2x2 to not overflow these 8bits
+	#define PREV_USE_COUNT_QUANTIZATION 63
+#else
+	#define PREV_USE_COUNT_QUANTIZATION 256
+#endif
+
+#define PARALLAX_REJECTION_MASK_THRESHOLD 0.5
+
+#define MAX_PARALLAX_FACTOR 8.0
+
+// Size and payload size of the subpixel grid in the history
+#define SUB_PIXEL_GRID_SIZE 2
+#define SUB_PIXEL_COUNT (SUB_PIXEL_GRID_SIZE * SUB_PIXEL_GRID_SIZE)
+
+#define SUB_PIXEL_BIT_COUNT 4
+#define SUB_PIXEL_BIT_MASK ((1 << SUB_PIXEL_BIT_COUNT) - 1)
+
+#define SUB_PIXEL_PARALLAX_FACTOR_BIT_COUNT 4
+#define SUB_PIXEL_PARALLAX_FACTOR_BIT_OFFSET 0
+#define SUB_PIXEL_PARALLAX_FACTOR_BIT_MASK ((1 << SUB_PIXEL_PARALLAX_FACTOR_BIT_COUNT) - 1)
+
+
+// Size and payload of the super resolution history
+#define SUPERRES_RES_MULTIPLIER 2
+
+#define SUPERRES_AGE_BITS 5
+#define SUPERRES_AGE_BITMASK ((1 << SUPERRES_AGE_BITS) - 1)
+
+#define SUPERRES_INTERFERENCE_BITS 4
+#define SUPERRES_INTERFERENCE_BITMASK ((1 << SUPERRES_INTERFERENCE_BITS) - 1)
+
+#define SUPERRES_OFFSET_BITS 4
+#define SUPERRES_OFFSET_BITMASK ((1 << SUPERRES_OFFSET_BITS) - 1)
+#define SUPERRES_OFFSET_OFFSET (SUPERRES_OFFSET_BITMASK / 2)
+#define SUPERRES_OFFSET_MIN (-SUPERRES_OFFSET_OFFSET)
+#define SUPERRES_OFFSET_MAX (SUPERRES_OFFSET_BITMASK - SUPERRES_OFFSET_OFFSET)
+
+
+
+// K = Center of the nearest input pixel.
+// O = Center of the output pixel.
+//
+//          |           |
+//    0     |     1     |     2
+//          |           |
+//          |           |
+//  --------+-----------+--------
+//          |           |
+//          | O         |
+//    3     |     K     |     5
+//          |           |
+//          |           |
+//  --------+-----------+--------
+//          |           |
+//          |           |
+//    6     |     7     |     8
+//          |           |
+//
+static const int2 kOffsets3x3[9] =
+{
+	int2(-1, -1),
+	int2(0, -1),
+	int2(1, -1),
+	int2(-1,  0),
+	int2(0,  0), // K
+	int2(1,  0),
+	int2(-1,  1),
+	int2(0,  1),
+	int2(1,  1),
+};
+
+// T = Center of the nearest top left pixel input pixel.
+// O = Center of the output pixel.
+//
+//          | 
+//    T     |     .
+//          | 
+//       O  | 
+//  --------+--------
+//          | 
+//          | 
+//    .     |     .
+//          | 
+static const int2 Offsets2x2[4] =
+{
+	int2( 0,  0), // T
+	int2( 1,  0),
+	int2( 0,  1),
+	int2( 1,  1),
+};
+
+// Indexes of the 3x3 square.
+static const uint kSquareIndexes3x3[9] = { 4, 0, 1, 2, 3, 5, 6, 7, 8 };
+
+// Indexes of the offsets to have plus + shape.
+static const uint kPlusIndexes3x3[5] = { 4, 1, 3, 5, 7 };
+
+
+#if CONFIG_COMPILE_FP16
+	#define taa_half  half
+	#define taa_half2 half2
+	#define taa_half3 half3
+	#define taa_half4 half4
+
+	#define taa_short  int16_t
+	#define taa_short2 int16_t2
+	#define taa_short3 int16_t3
+	#define taa_short4 int16_t4
+
+	#define taa_ushort  uint16_t
+	#define taa_ushort2 uint16_t2
+	#define taa_ushort3 uint16_t3
+	#define taa_ushort4 uint16_t4
+
+	#define taa_half2x2 half2x2
+	#define taa_half3x2 half3x2
+	#define taa_half4x2 half4x2
+
+	#define taa_short2x2 int16_t2x2
+	#define taa_short3x2 int16_t3x2
+	#define taa_short4x2 int16_t4x2
+
+	#define taa_ushort2x2 uint16_t2x2
+	#define taa_ushort3x2 uint16_t3x2
+	#define taa_ushort4x2 uint16_t4x2
+
+#else
+	#define taa_half  float
+	#define taa_half2 float2
+	#define taa_half3 float3
+	#define taa_half4 float4
+
+	#define taa_short  int
+	#define taa_short2 int2
+	#define taa_short3 int3
+	#define taa_short4 int4
+
+	#define taa_ushort  uint
+	#define taa_ushort2 uint2
+	#define taa_ushort3 uint3
+	#define taa_ushort4 uint4
+
+	#define taa_half2x2 float2x2
+	#define taa_half3x2 float3x2
+	#define taa_half4x2 float4x2
+
+	#define taa_short2x2 int2x2
+	#define taa_short3x2 int3x2
+	#define taa_short4x2 int4x2
+
+	#define taa_ushort2x2 uint2x2
+	#define taa_ushort3x2 uint3x2
+	#define taa_ushort4x2 uint4x2
+
+#endif
+
+#define taa_subpixel_details uint
+#define taa_subpixel_details2 uint2
+#define taa_subpixel_payload uint // TODO: uint16_t
+#define taa_subpixel_payload2 uint2 // TODO: uint16_t
+
+
+//------------------------------------------------------- PARAMETERS
+
+
+float2 InputInfo_Extent;
+float2 InputInfo_ExtentInverse;
+float2 InputInfo_ScreenPosToViewportScale;
+float2 InputInfo_ScreenPosToViewportBias;
+uint2  InputInfo_ViewportMin;
+uint2  InputInfo_ViewportMax;
+float2 InputInfo_ViewportSize;
+float2 InputInfo_ViewportSizeInverse;
+float2 InputInfo_UVViewportMin;
+float2 InputInfo_UVViewportMax;
+float2 InputInfo_UVViewportSize;
+float2 InputInfo_UVViewportSizeInverse;
+float2 InputInfo_UVViewportBilinearMin;
+float2 InputInfo_UVViewportBilinearMax;
+float2 InputJitter;
+int2   InputPixelPosMin;
+int2   InputPixelPosMax;
+
+float2 LowFrequencyInfo_Extent;
+float2 LowFrequencyInfo_ExtentInverse;
+uint2  LowFrequencyInfo_ViewportMin;
+uint2  LowFrequencyInfo_ViewportMax;
+float2 LowFrequencyInfo_ViewportSize;
+float2 LowFrequencyInfo_ViewportSizeInverse;
+float2 LowFrequencyInfo_UVViewportBilinearMin;
+float2 LowFrequencyInfo_UVViewportBilinearMax;
+
+float2 RejectionInfo_Extent;
+float2 RejectionInfo_ExtentInverse;
+uint2  RejectionInfo_ViewportMin;
+uint2  RejectionInfo_ViewportMax;
+float2 RejectionInfo_ViewportSize;
+float2 RejectionInfo_ViewportSizeInverse;
+float2 RejectionInfo_UVViewportBilinearMin;
+float2 RejectionInfo_UVViewportBilinearMax;
+
+float2 HistoryInfo_Extent;
+float2 HistoryInfo_ExtentInverse;
+uint2  HistoryInfo_ViewportMin;
+uint2  HistoryInfo_ViewportMax;
+float2 HistoryInfo_ViewportSize;
+float2 HistoryInfo_ViewportSizeInverse;
+float2 HistoryInfo_UVViewportBilinearMin;
+float2 HistoryInfo_UVViewportBilinearMax;
+
+// FTSRPrevHistoryParameters
+float2 PrevHistoryInfo_Extent;
+float2 PrevHistoryInfo_ExtentInverse;
+float2 PrevHistoryInfo_ScreenPosToViewportScale;
+float2 PrevHistoryInfo_ScreenPosToViewportBias;
+uint2  PrevHistoryInfo_ViewportMin;
+uint2  PrevHistoryInfo_ViewportMax;
+float2 PrevHistoryInfo_ViewportSize;
+float2 PrevHistoryInfo_ViewportSizeInverse;
+float2 PrevHistoryInfo_UVViewportMin;
+float2 PrevHistoryInfo_UVViewportMax;
+float2 PrevHistoryInfo_UVViewportSize;
+float2 PrevHistoryInfo_UVViewportSizeInverse;
+float2 PrevHistoryInfo_UVViewportBilinearMin;
+float2 PrevHistoryInfo_UVViewportBilinearMax;
+FScreenTransform ScreenPosToPrevHistoryBufferUV;
+float HistoryPreExposureCorrection;
+
+uint bCameraCut;
+uint bEnableInterferenceHeuristic;
+
+#if DEBUG_OUTPUT
+	RWTexture2D<float4> DebugOutput;
+#endif
+
+
+//------------------------------------------------------- FUNCTIONS
+	
+float2x2 ApplyScreenTransform(float2x2 PInA, FScreenTransform AToB)
+{
+	return dpv_add(dpv_mul(PInA, AToB.xy), AToB.zw);
+}
+
+#if PLATFORM_SUPPORTS_REAL_TYPES
+
+half2x2 ApplyScreenTransform(half2x2 PInA, FScreenTransform AToB)
+{
+	return dpv_add(dpv_mul(PInA, half2(AToB.xy)), half2(AToB.zw));
+}
+
+#endif
+
+
+/** Compute the group wave index into SGRP to then recompue the GroupThreadIndex latter */
+uint GetGroupWaveIndex(uint GroupThreadIndex, uint GroupSize)
+#if COMPILER_SUPPORTS_WAVE_ONCE
+{
+	uint LaneCountPerWave = WaveGetLaneCount();
+
+	if (LaneCountPerWave >= GroupSize)
+	{
+		return 0;
+	}
+
+	return WaveReadLaneFirst(GroupThreadIndex / LaneCountPerWave);
+}
+#else
+{
+	return 0;
+}
+#endif
+
+/** Force compute the group GroupThreadIndex through lane index and wave index if possible to reduce VGPR pressure. */
+uint GetGroupThreadIndex(uint GroupThreadIndex, uint GroupWaveIndex)
+#if COMPILER_SUPPORTS_WAVE_ONCE
+{
+	// shares GroupWaveOffset to save SALU
+	uint GroupWaveOffset = WaveGetLaneCount() * GroupWaveIndex;
+
+	// Do not share
+	ISOLATE
+	{
+		return GroupWaveOffset + WaveGetLaneIndex();
+	}
+}
+#else
+{
+	return GroupThreadIndex;
+}
+#endif
+
+// Clamp the offset to be shared across multiple samples
+taa_short2x2 ClampPixelOffset(
+	taa_short2x2 KernelCenterPixelPos,
+	taa_short2x2 Offset,
+	const taa_short2 OffsetDirection,
+	int2 MinPixelPos, int2 MaxPixelPos)
+{
+	taa_short2x2 Min = dpv_sub(taa_short2(MinPixelPos), KernelCenterPixelPos);
+	taa_short2x2 Max = dpv_sub(taa_short2(MaxPixelPos), KernelCenterPixelPos);
+
+	// Only do clamp based on the compile time known direction of the offset.
+	// This turns only 1 v_pk_max_i16 and v_pk_min_i16 for a 3x3 kernel.
+	taa_short2x2 ClampedOffset = 0;
+
+	if (OffsetDirection.x > 0)
+	{
+		ClampedOffset[0] = min(Offset[0], Max[0]);
+	}
+	else if (OffsetDirection.x < 0)
+	{
+		ClampedOffset[0] = max(Offset[0], Min[0]);
+	}
+		
+	if (OffsetDirection.y > 0)
+	{
+		ClampedOffset[1] = min(Offset[1], Max[1]);
+	}
+	else if (OffsetDirection.y < 0)
+	{
+		ClampedOffset[1] = max(Offset[1], Min[1]);
+	}
+
+	return ClampedOffset;
+}
+
+taa_short2 ClampPixelOffset(
+	taa_short2 KernelCenterPixelPos,
+	taa_short2 Offset,
+	const taa_short2 OffsetDirection,
+	int2 MinPixelPos, int2 MaxPixelPos)
+{
+	taa_short2 Min = taa_short2(MinPixelPos) - KernelCenterPixelPos;
+	taa_short2 Max = taa_short2(MaxPixelPos) - KernelCenterPixelPos;
+
+	// Only do clamp based on the compile time known direction of the offset.
+	// This turns only 1 v_pk_max_i16 and v_pk_min_i16 for a 3x3 kernel.
+	taa_short2 ClampedOffset = 0;
+
+	if (OffsetDirection.x > 0)
+		ClampedOffset.x = min(Offset.x, Max.x);
+	else if (OffsetDirection.x < 0)
+		ClampedOffset.x = max(Offset.x, Min.x);
+		
+	if (OffsetDirection.y > 0)
+		ClampedOffset.y = min(Offset.y, Max.y);
+	else if (OffsetDirection.y < 0)
+		ClampedOffset.y = max(Offset.y, Min.y);
+
+	return ClampedOffset;
+}
+
+taa_short2x2 ClampPixelOffset(
+	taa_short2x2 SamplePixelPos,
+	int2 MinPixelPos, int2 MaxPixelPos)
+{
+	taa_short2 SamplePixelPos0 = dpv_lo(SamplePixelPos);
+	taa_short2 SamplePixelPos1 = dpv_hi(SamplePixelPos);
+
+	SamplePixelPos0.x = min(SamplePixelPos0.x, taa_short(MaxPixelPos.x));
+	SamplePixelPos0.x = max(SamplePixelPos0.x, taa_short(MinPixelPos.x));
+		
+	SamplePixelPos0.y = min(SamplePixelPos0.y, taa_short(MaxPixelPos.y));
+	SamplePixelPos0.y = max(SamplePixelPos0.y, taa_short(MinPixelPos.y));
+	
+	SamplePixelPos1.x = min(SamplePixelPos1.x, taa_short(MaxPixelPos.x));
+	SamplePixelPos1.x = max(SamplePixelPos1.x, taa_short(MinPixelPos.x));
+		
+	SamplePixelPos1.y = min(SamplePixelPos1.y, taa_short(MaxPixelPos.y));
+	SamplePixelPos1.y = max(SamplePixelPos1.y, taa_short(MinPixelPos.y));
+
+	return dpv_interleave_registers(SamplePixelPos0, SamplePixelPos1);
+}
+
+taa_short2 ClampPixelOffset(
+	taa_short2 SamplePixelPos,
+	int2 MinPixelPos, int2 MaxPixelPos)
+{
+	SamplePixelPos.x = min(SamplePixelPos.x, taa_short(MaxPixelPos.x));
+	SamplePixelPos.x = max(SamplePixelPos.x, taa_short(MinPixelPos.x));
+		
+	SamplePixelPos.y = min(SamplePixelPos.y, taa_short(MaxPixelPos.y));
+	SamplePixelPos.y = max(SamplePixelPos.y, taa_short(MinPixelPos.y));
+
+	return SamplePixelPos;
+}
+
+taa_short2x2 AddAndClampPixelOffset(
+	taa_short2x2 KernelCenterPixelPos,
+	taa_short2x2 Offset,
+	const taa_short2 OffsetDirection,
+	int2 MinPixelPos, int2 MaxPixelPos)
+{
+	taa_short2x2 SamplePixelPos = KernelCenterPixelPos + Offset;
+	
+	taa_short2 SamplePixelPos0 = dpv_lo(SamplePixelPos);
+	taa_short2 SamplePixelPos1 = dpv_hi(SamplePixelPos);
+
+	// Only do clamp based on the compile time known direction of the offset.
+	// This turns only 1 v_pk_max_i16 and v_pk_min_i16 for a 3x3 kernel.
+	if (OffsetDirection.x > 0)
+		SamplePixelPos0.x = min(SamplePixelPos0.x, taa_short(MaxPixelPos.x));
+	else if (OffsetDirection.x < 0)
+		SamplePixelPos0.x = max(SamplePixelPos0.x, taa_short(MinPixelPos.x));
+		
+	if (OffsetDirection.y > 0)
+		SamplePixelPos0.y = min(SamplePixelPos0.y, taa_short(MaxPixelPos.y));
+	else if (OffsetDirection.y < 0)
+		SamplePixelPos0.y = max(SamplePixelPos0.y, taa_short(MinPixelPos.y));
+	
+	if (OffsetDirection.x > 0)
+		SamplePixelPos1.x = min(SamplePixelPos1.x, taa_short(MaxPixelPos.x));
+	else if (OffsetDirection.x < 0)
+		SamplePixelPos1.x = max(SamplePixelPos1.x, taa_short(MinPixelPos.x));
+		
+	if (OffsetDirection.y > 0)
+		SamplePixelPos1.y = min(SamplePixelPos1.y, taa_short(MaxPixelPos.y));
+	else if (OffsetDirection.y < 0)
+		SamplePixelPos1.y = max(SamplePixelPos1.y, taa_short(MinPixelPos.y));
+	
+	return dpv_interleave_registers(SamplePixelPos0, SamplePixelPos1);
+}
+
+taa_short2 AddAndClampPixelOffset(
+	taa_short2 KernelCenterPixelPos,
+	taa_short2 Offset,
+	const taa_short2 OffsetDirection,
+	int2 MinPixelPos, int2 MaxPixelPos)
+{
+	taa_short2 SamplePixelPos = KernelCenterPixelPos + Offset;
+
+	// Only do clamp based on the compile time known direction of the offset.
+	// This turns only 1 v_pk_max_i16 and v_pk_min_i16 for a 3x3 kernel.
+	if (OffsetDirection.x > 0)
+		SamplePixelPos.x = min(SamplePixelPos.x, taa_short(MaxPixelPos.x));
+	else if (OffsetDirection.x < 0)
+		SamplePixelPos.x = max(SamplePixelPos.x, taa_short(MinPixelPos.x));
+		
+	if (OffsetDirection.y > 0)
+		SamplePixelPos.y = min(SamplePixelPos.y, taa_short(MaxPixelPos.y));
+	else if (OffsetDirection.y < 0)
+		SamplePixelPos.y = max(SamplePixelPos.y, taa_short(MinPixelPos.y));
+
+	return SamplePixelPos;
+}
+
+taa_short2x2 InvalidateOutputPixelPos(taa_short2x2 PixelPos, uint2 ViewportMax)
+#if 1
+{
+	taa_short2x2 Subtract = dpv_sub(taa_short2(ViewportMax - 1), PixelPos);
+	taa_ushort2 Override = taa_ushort2(Subtract[0] | Subtract[1]);
+
+	#if CONFIG_COMPILE_FP16
+		PixelPos[0] |= -taa_short2((Override & uint16_t(0x8000)) >> 15);
+	#else
+		PixelPos[0] |= -taa_short2((Override & uint(0x80000000)) >> 31);
+	#endif
+
+	return PixelPos;
+}
+#else
+{
+	bool bIsValidPixel = all(PixelPos < ViewportMax);
+	PixelPos.x = bIsValidPixel ? PixelPos.x : ~taa_short(0);
+	return PixelPos;
+}
+#endif
+
+taa_short2 InvalidateOutputPixelPos(taa_short2 PixelPos, uint2 ViewportMax)
+#if 1
+{
+	taa_short2 Subtract = taa_short2(ViewportMax - 1) - PixelPos;
+	taa_ushort Override = taa_ushort(Subtract.x | Subtract.y);
+
+	#if CONFIG_COMPILE_FP16
+		PixelPos.x |= -taa_short((Override & uint16_t(0x8000)) >> 15);
+	#else
+		PixelPos.x |= -taa_short((Override & uint(0x80000000)) >> 31);
+	#endif
+
+	return PixelPos;
+}
+#else
+{
+	bool bIsValidPixel = all(PixelPos < ViewportMax);
+	PixelPos.x = bIsValidPixel ? PixelPos.x : ~taa_short(0);
+	return PixelPos;
+}
+#endif
+
+taa_ushort2 Map8x8Tile2x2Lane(uint GroupThreadIndex)
+{
+	taa_ushort T = taa_ushort(GroupThreadIndex);
+
+	taa_ushort2 GroupId;
+	GroupId.x = ((T >> taa_ushort(0)) & taa_ushort(0x01)) | ((T >> taa_ushort(2 - 1)) & taa_ushort(0x03 << 1));
+	GroupId.y = ((T >> taa_ushort(1)) & taa_ushort(0x01)) | ((T >> taa_ushort(4 - 1)) & taa_ushort(0x03 << 1));
+
+	return GroupId;
+}
+
+taa_ushort2 Map16x16Tile2x2Lane(uint GroupThreadIndex)
+{
+	taa_ushort T = taa_ushort(GroupThreadIndex);
+
+	taa_ushort2 GroupId;
+	GroupId.x = ((T >> taa_ushort(0)) & taa_ushort(0x01)) | ((T >> taa_ushort(2 - 1)) & taa_ushort(0x07 << 1));
+	GroupId.y = ((T >> taa_ushort(1)) & taa_ushort(0x01)) | ((T >> taa_ushort(5 - 1)) & taa_ushort(0x07 << 1));
+	
+	return GroupId;
+}
+
+float SafeRcp(float x)
+{
+	return x > 0.0 ? rcp(x) : 0.0;
+}
+
+float2 SafeRcp(float2 x)
+{
+	return float2(SafeRcp(x.x), SafeRcp(x.y));
+}
+
+#if CONFIG_COMPILE_FP16
+
+half SafeRcp(half x)
+{
+	// If x=0.0, then MaxHalfFloat * 0.0 = 0.0
+	return min(rcp(x), half(MaxHalfFloat)) * saturate(x * asfloat16(uint16_t(0x7C00)));
+}
+
+half2 SafeRcp(half2 x)
+{
+	// If x=0.0, then MaxHalfFloat * 0.0 = 0.0
+	return min(rcp(x), half(MaxHalfFloat)) * saturate(x * asfloat16(uint16_t(0x7C00)));
+}
+
+#endif
+
+taa_half3x2 RGBToYCoCg(taa_half3x2 RGB)
+{
+	taa_half2 Y  = mul(taa_half3(1, 2, 1), RGB);
+	taa_half2 Co = mul(taa_half3(2, 0, -2), RGB);
+	taa_half2 Cg = mul(taa_half3(-1, 2, -1), RGB);
+
+	taa_half3x2 YCoCg = taa_half3x2(Y, Co, Cg);
+	return YCoCg;
+}
+
+taa_half3x2 YCoCgToRGB(taa_half3x2 YCoCg)
+{
+	taa_half2 Y =  YCoCg[0] * taa_half(0.25);
+	taa_half2 Co = YCoCg[1] * taa_half(0.25);
+	taa_half2 Cg = YCoCg[2] * taa_half(0.25);
+
+	taa_half2 R = Y + Co - Cg;
+	taa_half2 G = Y + Cg;
+	taa_half2 B = Y - Co - Cg;
+
+	taa_half3x2 RGB = taa_half3x2(R, G, B);
+	return RGB;
+}
+
+taa_half3 RGBToYCoCg(taa_half3 RGB)
+{
+	return dpv_lo(RGBToYCoCg(dpv_interleave_mono_registers(RGB)));
+}
+
+taa_half3 YCoCgToRGB(taa_half3 YCoCg)
+{
+	return dpv_lo(YCoCgToRGB(dpv_interleave_mono_registers(YCoCg)));
+}
+
+
+// Some bright pixel can cause HdrWeight to get nullified under fp16 representation. So clamping this to a value close to the minimum float float positive value (0.000061).
+#define HDR_WEIGHT_SAFE_MIN_VALUE 0.0001
+
+// Faster but less accurate luma computation. 
+// Luma includes a scaling by 4.
+taa_half2 Luma4(taa_half3x2 Color)
+{
+	return (Color[1] * taa_half(2.0)) + (Color[0] + Color[2]);
+}
+
+taa_half Luma4(taa_half3 Color)
+{
+	return dpv_lo(Luma4(dpv_interleave_mono_registers(Color)));
+}
+
+taa_half2 HdrWeightY(taa_half2 Luma)
+{
+	taa_half Exposure = taa_half(1.0);
+
+	return max(taa_half(HDR_WEIGHT_SAFE_MIN_VALUE), rcp(Luma * Exposure + taa_half(4.0)));
+}
+
+taa_half HdrWeightY(taa_half Luma)
+{
+	return dpv_lo(HdrWeightY(dpv_interleave_mono_registers(Luma)));
+}
+
+taa_half2 HdrWeightInvY(taa_half2 LDRLuma) 
+{
+	return taa_half(4.0) * rcp(taa_half(1.0) - LDRLuma);
+}
+
+taa_half HdrWeightInvY(taa_half LDRLuma) 
+{
+	return dpv_lo(HdrWeightInvY(dpv_interleave_mono_registers(LDRLuma)));
+}
+
+// Optimized HDR weighting function.
+taa_half2 HdrWeight4(taa_half3x2 Color)
+{
+	return HdrWeightY(Luma4(Color));
+}
+
+taa_half HdrWeight4(taa_half3 Color)
+{
+	return HdrWeightY(Luma4(Color));
+}
+
+
+// Returns the weight of a pixels at a coordinate <PixelDelta> from the PDF highest point.
+taa_half2 ComputeSampleWeigth(taa_half2 UpscaleFactor, taa_half2x2 PixelDelta, const float MinimalContribution)
+{
+	taa_half2 u2 = UpscaleFactor * UpscaleFactor;
+
+	// 1 - 1.9 * x^2 + 0.9 * x^4
+	taa_half2 x2 = saturate(u2 * dpv_length2(PixelDelta));
+	//return taa_half(((float(0.9) + MinimalContribution) * x2 - float(1.9)) * x2 + float(1.0));
+	return saturate((taa_half(0.9) * x2 - taa_half(1.9)) * x2 + taa_half(1.0 + MinimalContribution));
+}
+
+taa_half ComputeSampleWeigth(taa_half UpscaleFactor, taa_half2 PixelDelta, const float MinimalContribution)
+{
+	return dpv_lo(ComputeSampleWeigth(
+		dpv_interleave_mono_registers(UpscaleFactor),
+		dpv_interleave_mono_registers(PixelDelta),
+		MinimalContribution));
+}
+
+taa_half2x2 WeightedLerpFactors(taa_half2 WeightA, taa_half2 WeightB, taa_half2 Blend)
+{
+	taa_half2 BlendA = (taa_half(1.0) - Blend) * WeightA;
+	taa_half2 BlendB = Blend * WeightB;
+	taa_half2 RcpBlend = SafeRcp(BlendA + BlendB);
+	BlendA *= RcpBlend;
+	BlendB *= RcpBlend;
+	return taa_half2x2(BlendA, BlendB);
+}
+
+taa_half2 WeightedLerpFactors(taa_half WeightA, taa_half WeightB, taa_half Blend)
+{
+	return dpv_lo(WeightedLerpFactors(
+		dpv_interleave_mono_registers(WeightA),
+		dpv_interleave_mono_registers(WeightB),
+		dpv_interleave_mono_registers(Blend)));
+}
+
+bool TakeOnlyOneSamplePair(float2 Offset)
+{
+	return Offset.y > 0.0 || (Offset.x > 0.0 && Offset.y == 0.0);
+}
+
+float2 ComputeStaticVelocity(float2 ScreenPos, float DeviceZ)
+{
+	float3 PosN = float3(ScreenPos, DeviceZ);
+
+	float4 ThisClip = float4(PosN, 1);
+	float4 PrevClip = mul( ThisClip, View.ClipToPrevClip );
+	float2 PrevScreen = PrevClip.xy / PrevClip.w;
+	return PosN.xy - PrevScreen;
+}
+
+taa_half ComputePredictionCompleteness(taa_half SampleHistoryValidity)
+{
+	return saturate(SampleHistoryValidity * taa_half(MAX_SAMPLE_COUNT) - taa_half(0.2));
+}
+
+bool IsOffScreen(uint bCameraCut, float2 ScreenPos, taa_half ParallaxRejectionMask)
+{
+	bool bIsCameraCut = bCameraCut != 0;
+	bool bIsOutOfBounds = max(abs(ScreenPos.x), abs(ScreenPos.y)) >= 1.0;
+	bool bIsParallaxRejected = ParallaxRejectionMask < taa_half(PARALLAX_REJECTION_MASK_THRESHOLD);
+
+	return (bIsCameraCut || bIsOutOfBounds || bIsParallaxRejected);
+}
+
+taa_half MeasureBackbufferLDRQuantizationError()
+{
+	// Assume the backbuffer is 10bit per channels
+	return taa_half(0.5 / 1024.0);
+}
+
+taa_half2 MeasureBackbufferLDRQuantizationErrorFromHDRLuma(taa_half2 Luma)
+{
+	taa_half2 Tonemap = Luma * HdrWeightY(Luma) + MeasureBackbufferLDRQuantizationError().xx;
+	taa_half2 PixelColorQuantizationError = abs(Tonemap * HdrWeightInvY(Tonemap) - Luma);
+
+	return PixelColorQuantizationError;
+}
+
+taa_half MeasureBackbufferLDRQuantizationErrorFromHDRLuma(taa_half Luma)
+{
+	return dpv_lo(MeasureBackbufferLDRQuantizationErrorFromHDRLuma(dpv_interleave_mono_registers(Luma)));
+}
+
+taa_half2 MeasureRejectionFactor(
+	taa_half3x2 PrevYCoCg,
+	taa_half3x2 ClampedPrevYCoCg,
+	taa_half3x2 InputCenterYCoCg,
+	taa_half3x2 InputMinYCoCg, taa_half3x2 InputMaxYCoCg,
+	taa_half2 BackbufferQuantizationError)
+{
+	taa_half3x2 BoxSize = InputMaxYCoCg - InputMinYCoCg;
+
+	taa_half3x2 ClampedEnergy = abs(ClampedPrevYCoCg - PrevYCoCg);
+	taa_half3x2 Delta = abs(InputCenterYCoCg - PrevYCoCg);
+
+	taa_half3x2 ClampError = max(
+		BoxSize * taa_half(0.25),
+		taa_half3x2(BackbufferQuantizationError, BackbufferQuantizationError, BackbufferQuantizationError));
+
+	taa_half3x2 Factor = taa_half(1.0) - saturate(max(ClampedEnergy - ClampError, taa_half3x2(0.0, 0.0, 0.0, 0.0, 0.0, 0.0)) / max(Delta, BoxSize));
+
+	taa_half2 Rejection = min3(Factor[0], Factor[1], Factor[2]);
+	
+	return Rejection;
+}
+
+taa_half MeasureRejectionFactor(
+	taa_half3 PrevYCoCg,
+	taa_half3 ClampedPrevYCoCg,
+	taa_half3 InputCenterYCoCg,
+	taa_half3 InputMinYCoCg, taa_half3 InputMaxYCoCg,
+	taa_half BackbufferQuantizationError)
+{
+	return dpv_lo(MeasureRejectionFactor(
+		dpv_interleave_mono_registers(PrevYCoCg), 
+		dpv_interleave_mono_registers(ClampedPrevYCoCg), 
+		dpv_interleave_mono_registers(InputCenterYCoCg), 
+		dpv_interleave_mono_registers(InputMinYCoCg), 
+		dpv_interleave_mono_registers(InputMaxYCoCg), 
+		dpv_interleave_mono_registers(BackbufferQuantizationError)));
+}
+
+struct FSubpixelNeighborhood
+{
+	// 2x2 quad of subpixels grid
+	taa_subpixel_details SampleArray[4];
+
+	// Offset of subpixel
+	uint2 SubpixelOffset;
+
+	// Offset of the closest subpixel
+	bool2 ClosestSubpixelOffset;
+};
+
+FSubpixelNeighborhood GatherPrevSubpixelNeighborhood(Texture2D<uint> PrevSubPixelDetailTexture, float2 PrevHistoryBufferUV)
+{
+	float2 PixelCoord = PrevHistoryBufferUV * PrevHistoryInfo_Extent;
+	float2 TopLeftPixelCoord = floor(PixelCoord - 0.5) + 0.5;
+	
+	FSubpixelNeighborhood Neighborhood;
+	#if 1 // TODO
+	{
+		Neighborhood.SampleArray[0] = PrevSubPixelDetailTexture[uint2(TopLeftPixelCoord + float2(0.0, 0.0))];
+		Neighborhood.SampleArray[1] = PrevSubPixelDetailTexture[uint2(TopLeftPixelCoord + float2(1.0, 0.0))];
+		Neighborhood.SampleArray[2] = PrevSubPixelDetailTexture[uint2(TopLeftPixelCoord + float2(0.0, 1.0))];
+		Neighborhood.SampleArray[3] = PrevSubPixelDetailTexture[uint2(TopLeftPixelCoord + float2(1.0, 1.0))];
+	}
+	#else
+	{
+		// (-,+),(+,+),(+,-),(-,-),
+		/**
+		 *  3 2
+		 *  0 1
+		 */
+		uint4 Samples = PrevSubPixelDetailTexture.Gather(GlobalPointClampedSampler, PrevHistoryBufferUV);
+
+		Neighborhood.SampleArray[0] = Samples[3]; // 00
+		Neighborhood.SampleArray[1] = Samples[2]; // 10
+		Neighborhood.SampleArray[2] = Samples[0]; // 01
+		Neighborhood.SampleArray[3] = Samples[1]; // 11
+	}
+	#endif
+
+	taa_half2 Interp = taa_half2(PixelCoord - TopLeftPixelCoord);
+	Neighborhood.SubpixelOffset = uint2(round(Interp * 2.0));
+	Neighborhood.ClosestSubpixelOffset = Interp > 0.5;
+
+	return Neighborhood;
+}
+
+taa_subpixel_payload GetSubpixelPayload(FSubpixelNeighborhood Neighborhood, uint SubpixelId)
+{
+	const uint2 SubpixelCoord = uint2(SubpixelId % SUB_PIXEL_GRID_SIZE, SubpixelId / SUB_PIXEL_GRID_SIZE);
+
+	uint2 InputSubpixelCoord = Neighborhood.SubpixelOffset + SubpixelCoord;
+	uint2 InputTexelCoord = InputSubpixelCoord / SUB_PIXEL_GRID_SIZE;
+	uint2 QuadSubpixelCoord = InputSubpixelCoord % SUB_PIXEL_GRID_SIZE;
+
+
+	uint QuadSampleId = dot(InputTexelCoord, uint2(1, 2));
+	taa_subpixel_details InputSample = Neighborhood.SampleArray[QuadSampleId];
+
+	uint BitShift = SUB_PIXEL_BIT_COUNT * dot(QuadSubpixelCoord, uint2(1, SUB_PIXEL_GRID_SIZE));
+
+	taa_subpixel_payload SubpixelPayload = taa_subpixel_payload(InputSample >> BitShift) & SUB_PIXEL_BIT_MASK;
+
+	return SubpixelPayload;
+}
+
+taa_subpixel_payload GetClosestSubpixelPayload(FSubpixelNeighborhood Neighborhood)
+{
+	uint QuadSampleId = dot(uint2(Neighborhood.ClosestSubpixelOffset), uint2(1, 2));
+	taa_subpixel_details InputSample = Neighborhood.SampleArray[QuadSampleId];
+
+	uint BitShift = SUB_PIXEL_BIT_COUNT * (3 - QuadSampleId);
+
+	taa_subpixel_payload SubpixelPayload = taa_subpixel_payload(InputSample >> BitShift) & SUB_PIXEL_BIT_MASK;
+
+	return SubpixelPayload;
+}
+
+taa_subpixel_payload CompressParallaxFactor(taa_half ParallaxFactor)
+{
+	return clamp(
+		taa_subpixel_payload((ParallaxFactor - taa_half(1.0)) * taa_half(float(SUB_PIXEL_PARALLAX_FACTOR_BIT_MASK) / MAX_PARALLAX_FACTOR)),
+		taa_subpixel_payload(0), taa_subpixel_payload(SUB_PIXEL_PARALLAX_FACTOR_BIT_MASK));
+}
+
+taa_half UncompressParallaxFactor(taa_subpixel_payload CompressedParallaxFactor)
+{
+	return 
+		taa_half(CompressedParallaxFactor & SUB_PIXEL_PARALLAX_FACTOR_BIT_MASK) *
+		taa_half(MAX_PARALLAX_FACTOR / float(SUB_PIXEL_PARALLAX_FACTOR_BIT_MASK)) +
+		taa_half(1.0);
+}
+
+taa_half LumaToInterferenceSeed(taa_half Luma)
+{
+	return Luma * HdrWeightY(Luma);
+}
+
+taa_half3x2 QuantizeForFloatRenderTarget(taa_half3x2 Color, taa_half E, const float3 QuantizationError)
+{
+	taa_half3x2 Error = dpv_mul(Color, taa_half3(QuantizationError));
+
+	#if 0
+		// NOP
+	#elif CONFIG_COMPILE_FP16
+	{
+		Error[0] = asfloat16(asint16(Error[0]) & uint16_t(~0x03FF));
+		Error[1] = asfloat16(asint16(Error[1]) & uint16_t(~0x03FF));
+		Error[2] = asfloat16(asint16(Error[2]) & uint16_t(~0x03FF));
+	}
+	#else
+	{
+		Error[0] = asfloat(asuint(Error[0]) & ~0x007FFFFF);
+		Error[1] = asfloat(asuint(Error[1]) & ~0x007FFFFF);
+		Error[2] = asfloat(asuint(Error[2]) & ~0x007FFFFF);
+	}
+	#endif
+	
+	return Color + Error * E;
+}
+
+taa_half3 QuantizeForFloatRenderTarget(taa_half3 Color, taa_half E, const float3 QuantizationError)
+{
+	return dpv_lo(QuantizeForFloatRenderTarget(dpv_interleave_mono_registers(Color), E, QuantizationError));
+}
+
+
+struct FSuperResHistory
+{
+	int2 TexelToSample;
+	uint Age;
+	taa_half InterferenceSeed;
+};
+
+uint2 PackSuperResHistory(FSuperResHistory SuperRes)
+{
+	uint CompressedInterferenceSeed = clamp(
+		uint(SuperRes.InterferenceSeed * taa_half(SUPERRES_INTERFERENCE_BITMASK)),
+		uint(0), uint(SUPERRES_INTERFERENCE_BITMASK));
+
+	uint2 Packed = 0;
+
+	Packed[0] |= ((SuperRes.Age) & SUPERRES_AGE_BITMASK) << 0;
+	Packed[0] |= ((CompressedInterferenceSeed) & SUPERRES_INTERFERENCE_BITMASK) << 5;
+
+	Packed[1] |= (uint(SuperRes.TexelToSample.x + SUPERRES_OFFSET_OFFSET) & SUPERRES_OFFSET_BITMASK) << 0 * SUPERRES_OFFSET_BITS;
+	Packed[1] |= (uint(SuperRes.TexelToSample.y + SUPERRES_OFFSET_OFFSET) & SUPERRES_OFFSET_BITMASK) << 1 * SUPERRES_OFFSET_BITS;
+
+	return Packed;
+}
+
+FSuperResHistory UnpackSuperResHistory(uint2 Packed)
+{
+	FSuperResHistory SuperRes;
+
+	SuperRes.Age = ((Packed[0] >> 0) & SUPERRES_AGE_BITMASK);
+	SuperRes.InterferenceSeed = taa_half((Packed[0] >> 5) & SUPERRES_INTERFERENCE_BITMASK) * rcp(taa_half(SUPERRES_INTERFERENCE_BITMASK));
+
+	SuperRes.TexelToSample.x = ((Packed[1] >> 0 * SUPERRES_OFFSET_BITS) & SUPERRES_OFFSET_BITMASK) - SUPERRES_OFFSET_OFFSET;
+	SuperRes.TexelToSample.y = ((Packed[1] >> 1 * SUPERRES_OFFSET_BITS) & SUPERRES_OFFSET_BITMASK) - SUPERRES_OFFSET_OFFSET;
+
+	return SuperRes;
+}
+
+#if COMPRESS_PREV_USE_COUNT
+
+// encode 4 use count that are 8bits shifted in the PF_R32_UINT
+void ComputeCompressedUseCountPixelCoordinates(uint2 PixelPos, out uint2 CompressedUseCountPixelPos, out uint Shift)
+{
+	// Overlap 2x2 consecutive blocks of 8x8 pixel to reduce atomic contention in the scattering pass
+	const uint kTileSize = 8;
+
+	CompressedUseCountPixelPos = (PixelPos % kTileSize) | ((PixelPos >> 1) & ~(kTileSize - 1));
+	Shift = ((PixelPos.x / kTileSize) % 2) | (((PixelPos.y / kTileSize) % 2) << 1);
+}
+	
+#endif
diff --git a/Engine/Shaders/Private/TemporalSuperResolution/TSRCompareHistory.usf b/Engine/Shaders/Private/TemporalSuperResolution/TSRCompareHistory.usf
new file mode 100644
--- /dev/null
+++ b/Engine/Shaders/Private/TemporalSuperResolution/TSRCompareHistory.usf
@@ -0,0 +1,149 @@
+// Copyright Epic Games, Inc. All Rights Reserved.
+
+#include "TSRCommon.ush"
+
+
+//------------------------------------------------------- CONFIG
+
+#define CONFIG_OUTPUT_HALF_RES DIM_OUTPUT_HALF_RES
+
+#define TILE_SIZE 16
+
+
+//------------------------------------------------------- PARAMETERS
+
+Texture2D<taa_half> ParallaxRejectionMaskTexture;
+Texture2D<taa_half3> FilteredInputTexture;
+Texture2D<taa_half3> FilteredPredictionSceneColorTexture;
+Texture2D<taa_half> InterferenceWeightTexture;
+
+RWTexture2D<taa_half> HistoryRejectionOutput;
+
+
+//------------------------------------------------------- LDS
+
+groupshared taa_half SharedRejection[TILE_SIZE * TILE_SIZE];
+
+
+//------------------------------------------------------- FUNCTIONS
+
+taa_half3 TransformColorForClampingBox(taa_half3 Color)
+{
+	return RGBToYCoCg(Color);
+}
+
+
+//------------------------------------------------------- ENTRY POINT
+
+[numthreads(TILE_SIZE, TILE_SIZE, 1)]
+void MainCS(
+	uint2 GroupId : SV_GroupID,
+	uint GroupThreadIndex : SV_GroupIndex)
+{
+	float4 Debug = 0.0;
+
+	taa_short2 InputPixelPos = (
+		taa_short2(InputPixelPosMin) +
+		taa_short2(GroupId) * taa_short2(TILE_SIZE, TILE_SIZE) +
+		Map16x16Tile2x2Lane(GroupThreadIndex));
+
+	taa_half3 InputMinColor;
+	taa_half3 InputMaxColor;
+	taa_half3 InputCenterColor;
+	taa_half ParallaxRejectionMask;
+	taa_half InterferenceWeight;
+	{
+		InputCenterColor = TransformColorForClampingBox(FilteredInputTexture[InputPixelPos]);
+		ParallaxRejectionMask = ParallaxRejectionMaskTexture[InputPixelPos];
+
+		InputMinColor = InputCenterColor;
+		InputMaxColor = InputCenterColor;
+		
+		InterferenceWeight = InterferenceWeightTexture[InputPixelPos];
+	}
+
+	UNROLL
+	for (uint i = 0; i < 9; i++)
+	{
+		if (i == 4)
+		{
+			continue;
+		}
+		
+		taa_short2 SampleInputPixelPos = InputPixelPos + ClampPixelOffset(
+			InputPixelPos,
+			kOffsets3x3[i], kOffsets3x3[i],
+			InputPixelPosMin, InputPixelPosMax);
+
+		taa_half3 Color = TransformColorForClampingBox(FilteredInputTexture[SampleInputPixelPos]);
+
+		InputMinColor = min(InputMinColor, Color);
+		InputMaxColor = max(InputMaxColor, Color);
+	}
+
+	taa_half3 PrevColor = TransformColorForClampingBox(FilteredPredictionSceneColorTexture[InputPixelPos]);
+
+	taa_half Rejection = 1.0;
+	#if 1
+	{
+		taa_half3 ClampedPrevColor = clamp(PrevColor, InputMinColor, InputMaxColor);
+
+		Rejection = MeasureRejectionFactor(
+			PrevColor, ClampedPrevColor,
+			InputCenterColor, InputMinColor, InputMaxColor,
+			MeasureBackbufferLDRQuantizationErrorFromHDRLuma(max(PrevColor[0], ClampedPrevColor[0])));
+		
+		Debug.rgb = abs(YCoCgToRGB(PrevColor - ClampedPrevColor)) / max3(YCoCgToRGB(InputMaxColor), YCoCgToRGB(InputCenterColor), 0.01);
+	}
+	#endif
+	
+	FLATTEN
+	if (bEnableInterferenceHeuristic)
+	{
+		Rejection = saturate(Rejection + 1.0 - InterferenceWeight);
+	}
+
+	if (ParallaxRejectionMask < taa_half(PARALLAX_REJECTION_MASK_THRESHOLD))
+	{
+		Rejection = taa_half(1.0);
+	}
+
+	#if CONFIG_OUTPUT_HALF_RES
+	{
+		SharedRejection[GroupThreadIndex] = Rejection;
+
+		GroupMemoryBarrierWithGroupSync();
+
+		UNROLL
+		for (uint i = 1; i < 4; i++)
+		{
+			uint ButterflySwap = i;
+
+			Rejection = min(Rejection, SharedRejection[GroupThreadIndex ^ ButterflySwap]);
+		}
+
+		bool bIsValid = all(InputPixelPos <= taa_short2(InputPixelPosMax)) && all((taa_ushort2(InputPixelPos - taa_short2(InputPixelPosMin)) % taa_ushort(2)) == taa_ushort(0));
+
+		taa_ushort2 OutputPixelPos = taa_ushort2(InputPixelPos - taa_short2(InputPixelPosMin)) / taa_ushort(2);
+		OutputPixelPos.x = bIsValid ? OutputPixelPos.x : ~taa_ushort(0);
+
+		HistoryRejectionOutput[OutputPixelPos] = Rejection;
+	}
+	#else
+	{
+		bool bIsValid = all(InputPixelPos <= taa_short2(InputPixelPosMax));
+
+		taa_ushort2 OutputPixelPos = taa_ushort2(InputPixelPos - taa_short2(InputPixelPosMin));
+		OutputPixelPos.x = bIsValid ? OutputPixelPos.x : ~taa_ushort(0);
+
+		HistoryRejectionOutput[OutputPixelPos] = Rejection;
+	}
+	#endif
+
+	#if DEBUG_OUTPUT
+	{
+		taa_short2 OutputPixelPos = InvalidateOutputPixelPos(taa_ushort2(InputPixelPos), InputPixelPosMax);
+		DebugOutput[OutputPixelPos] = Debug;
+	}
+	#endif
+}
diff --git a/Engine/Shaders/Private/TemporalSuperResolution/TSRCompareTranslucency.usf b/Engine/Shaders/Private/TemporalSuperResolution/TSRCompareTranslucency.usf
new file mode 100644
--- /dev/null
+++ b/Engine/Shaders/Private/TemporalSuperResolution/TSRCompareTranslucency.usf
@@ -0,0 +1,220 @@
+// Copyright Epic Games, Inc. All Rights Reserved.
+
+#include "TSRCommon.ush"
+
+
+//------------------------------------------------------- CONFIG
+
+#define TILE_SIZE 8
+
+#define CONFIG_KERNEL 0
+
+#define CONFIG_OUTPUT_HALF_RES (DIM_OUTPUT_HALF_RES)
+
+#define CONFIG_OUTPUT_INTERFERENCE_SEEDS 0
+
+#define CONFIG_SAMPLES 5
+
+
+//------------------------------------------------------- PARAMETERS
+
+float2 TranslucencyInfo_Extent;
+float2 TranslucencyInfo_ExtentInverse;
+float2 TranslucencyInfo_ScreenPosToViewportScale;
+float2 TranslucencyInfo_ScreenPosToViewportBias;
+uint2  TranslucencyInfo_ViewportMin;
+uint2  TranslucencyInfo_ViewportMax;
+float2 TranslucencyInfo_ViewportSize;
+float2 TranslucencyInfo_ViewportSizeInverse;
+float2 TranslucencyInfo_UVViewportMin;
+float2 TranslucencyInfo_UVViewportMax;
+float2 TranslucencyInfo_UVViewportSize;
+float2 TranslucencyInfo_UVViewportSizeInverse;
+float2 TranslucencyInfo_UVViewportBilinearMin;
+float2 TranslucencyInfo_UVViewportBilinearMax;
+
+float2 PrevTranslucencyInfo_Extent;
+float2 PrevTranslucencyInfo_ExtentInverse;
+float2 PrevTranslucencyInfo_ScreenPosToViewportScale;
+float2 PrevTranslucencyInfo_ScreenPosToViewportBias;
+uint2  PrevTranslucencyInfo_ViewportMin;
+uint2  PrevTranslucencyInfo_ViewportMax;
+float2 PrevTranslucencyInfo_ViewportSize;
+float2 PrevTranslucencyInfo_ViewportSizeInverse;
+float2 PrevTranslucencyInfo_UVViewportMin;
+float2 PrevTranslucencyInfo_UVViewportMax;
+float2 PrevTranslucencyInfo_UVViewportSize;
+float2 PrevTranslucencyInfo_UVViewportSizeInverse;
+float2 PrevTranslucencyInfo_UVViewportBilinearMin;
+float2 PrevTranslucencyInfo_UVViewportBilinearMax;
+float PrevTranslucencyPreExposureCorrection;
+
+FScreenTransform InputPixelPosToScreenPos;
+FScreenTransform ScreenPosToPrevTranslucencyTextureUV;
+
+Texture2D<taa_half4> TranslucencyTexture;
+Texture2D<float2> DilatedVelocityTexture;
+Texture2D<taa_half4> PrevTranslucencyTexture;
+
+RWTexture2D<taa_half> TranslucencyRejectionOutput;
+
+
+//------------------------------------------------------- ENTRY POINT
+
+taa_half4 TransformToComparisonColorSpace(taa_half4 Color)
+{
+	return taa_half4(RGBToYCoCg(Color.rgb), Color.a);
+}
+
+
+//------------------------------------------------------- ENTRY POINT
+
+[numthreads(TILE_SIZE, TILE_SIZE, 1)]
+void MainCS(
+	uint2 GroupId : SV_GroupID,
+	uint GroupThreadIndex : SV_GroupIndex)
+{
+	float4 Debug = 0;
+	
+	uint2 InputPixelPos = (
+		ZOrder2D(GroupThreadIndex, log2(TILE_SIZE)) +
+		GroupId * uint2(TILE_SIZE, TILE_SIZE));
+
+	float2 ScreenPos = ApplyScreenTransform(float2(InputPixelPos), InputPixelPosToScreenPos);
+	float2 EncodedVelocity = DilatedVelocityTexture[InputPixelPos];
+	
+	// Sample current frame translucency's
+	taa_half4 Translucency;
+	taa_half4 TranslucencyMin;
+	taa_half4 TranslucencyMax;
+	{
+		UNROLL_N(CONFIG_SAMPLES)
+		for (uint SampleId = 0; SampleId < CONFIG_SAMPLES; SampleId++)
+		{
+			taa_short2 SampleInputPixelPos;
+			
+			// TODO: handle SeparateTranslucencyScreenPercentage
+			#if CONFIG_SAMPLES == 9
+			{
+				taa_short2 iPixelOffset = taa_short2(kOffsets3x3[kSquareIndexes3x3[SampleId]]);
+				
+				SampleInputPixelPos = AddAndClampPixelOffset(
+					InputPixelPos,
+					iPixelOffset, iPixelOffset,
+					InputPixelPosMin, InputPixelPosMax);
+			}
+			#elif CONFIG_SAMPLES == 5
+			{
+				taa_short2 iPixelOffset = taa_short2(kOffsets3x3[kPlusIndexes3x3[SampleId]]);
+					
+				SampleInputPixelPos = AddAndClampPixelOffset(
+					InputPixelPos,
+					iPixelOffset, iPixelOffset,
+					InputPixelPosMin, InputPixelPosMax);
+			}
+			#else
+				#error Unknown sample count
+			#endif
+
+			taa_half4 TranslucencyColor = TransformToComparisonColorSpace(TranslucencyTexture[SampleInputPixelPos]);
+
+			if (SampleId == 0)
+			{
+				Translucency = TranslucencyColor;
+				TranslucencyMin = TranslucencyColor;
+				TranslucencyMax = TranslucencyColor;
+			}
+			else
+			{
+				TranslucencyMin = min(TranslucencyMin, TranslucencyColor);
+				TranslucencyMax = max(TranslucencyMax, TranslucencyColor);
+			}
+		}
+	}
+
+	// Sample previous frame's translucency
+	taa_half4 PrevTranslucency;
+	taa_half4 PrevTranslucencyMin;
+	taa_half4 PrevTranslucencyMax;
+	{
+		float2 PrevScreenPos = ScreenPos - DecodeVelocityFromTexture(float4(EncodedVelocity, 0.0, 0.0)).xy;
+		float2 PrevTranslucencyTextureUV = ApplyScreenTransform(PrevScreenPos, ScreenPosToPrevTranslucencyTextureUV);
+	
+		UNROLL_N(CONFIG_SAMPLES)
+		for (uint SampleId = 0; SampleId < CONFIG_SAMPLES; SampleId++)
+		{
+			taa_half2 PixelOffset;
+			
+			// TODO: handle r.SeparateTranslucencyScreenPercentage
+			#if CONFIG_SAMPLES == 9
+			{
+				taa_short2 iPixelOffset = taa_short2(kOffsets3x3[kSquareIndexes3x3[SampleId]]);
+				PixelOffset = taa_half2(iPixelOffset);
+			}
+			#elif CONFIG_SAMPLES == 5
+			{
+				taa_short2 iPixelOffset = taa_short2(kOffsets3x3[kPlusIndexes3x3[SampleId]]);
+				PixelOffset = taa_half2(iPixelOffset);
+			}
+			#else
+				#error Unknown sample count
+			#endif
+
+			float2 SampleTextureUV = clamp(
+				PrevTranslucencyTextureUV + PixelOffset * PrevTranslucencyInfo_ExtentInverse,
+				PrevTranslucencyInfo_UVViewportBilinearMin,
+				PrevTranslucencyInfo_UVViewportBilinearMax);
+		
+			taa_half4 PrevTranslucencyColor = TransformToComparisonColorSpace(
+				PrevTranslucencyTexture.SampleLevel(GlobalPointClampedSampler, SampleTextureUV, 0));
+
+			if (SampleId == 0)
+			{
+				PrevTranslucency = PrevTranslucencyColor;
+				PrevTranslucencyMin = PrevTranslucencyColor;
+				PrevTranslucencyMax = PrevTranslucencyColor;
+			}
+			else
+			{
+				PrevTranslucencyMin = min(PrevTranslucencyMin, PrevTranslucencyColor);
+				PrevTranslucencyMax = max(PrevTranslucencyMax, PrevTranslucencyColor);
+			}
+		}
+
+		PrevTranslucency *= PrevTranslucencyPreExposureCorrection;
+		PrevTranslucencyMin *= PrevTranslucencyPreExposureCorrection;
+		PrevTranslucencyMax *= PrevTranslucencyPreExposureCorrection;
+	}
+
+	taa_half Rejection = 1.0;
+	#if 1
+	{
+		taa_half4 ClampedPrevTranslucency = clamp(PrevTranslucency, TranslucencyMin, TranslucencyMax);
+		taa_half4 ClampedTranslucency = clamp(Translucency, PrevTranslucencyMin, PrevTranslucencyMax);
+
+		// TODO: reject with alpha too? Might not need since translucency is premultiplied RGBA 
+		taa_half ForwardRejection = MeasureRejectionFactor(
+			PrevTranslucency.xyz, ClampedPrevTranslucency.xyz,
+			Translucency.xyz, TranslucencyMin.xyz, TranslucencyMax.xyz,
+			MeasureBackbufferLDRQuantizationErrorFromHDRLuma(max(PrevTranslucency[0], ClampedPrevTranslucency[0])));
+		
+		taa_half BackwardRejection = MeasureRejectionFactor(
+			Translucency.xyz, ClampedTranslucency.xyz,
+			PrevTranslucency.xyz, PrevTranslucencyMin.xyz, PrevTranslucencyMax.xyz,
+			MeasureBackbufferLDRQuantizationErrorFromHDRLuma(max(Translucency[0], ClampedTranslucency[0])));
+		
+		Rejection = min(ForwardRejection, BackwardRejection);
+	}
+	#endif
+
+	if (all(InputPixelPos < InputInfo_ViewportMax))
+	{
+		TranslucencyRejectionOutput[InputPixelPos] = Rejection;
+
+		#if DEBUG_OUTPUT
+		{
+			DebugOutput[InputPixelPos] = Debug;
+		}
+		#endif
+	}
+}
diff --git a/Engine/Shaders/Private/TemporalSuperResolution/TSRDebugHistory.usf b/Engine/Shaders/Private/TemporalSuperResolution/TSRDebugHistory.usf
new file mode 100644
--- /dev/null
+++ b/Engine/Shaders/Private/TemporalSuperResolution/TSRDebugHistory.usf
@@ -0,0 +1,108 @@
+// Copyright Epic Games, Inc. All Rights Reserved.
+
+#include "TSRCommon.ush"
+#include "../Random.ush"
+
+
+//------------------------------------------------------- CONSTANTS
+
+#define TILE_SIZE 8
+
+
+//------------------------------------------------------- PARAMETERS
+
+Texture2D<taa_half3> History_Textures_0;
+Texture2D<taa_half3> History_Textures_1;
+Texture2D<taa_half2> History_Textures_2;
+Texture2D<taa_subpixel_details> History_Textures_3;
+
+Texture2D<taa_half3> PrevHistory_Textures_0;
+Texture2D<taa_half3> PrevHistory_Textures_1;
+Texture2D<taa_half2> PrevHistory_Textures_2;
+Texture2D<taa_subpixel_details> PrevHistory_Textures_3;
+
+
+//------------------------------------------------------- FUNCTION
+
+
+void LMHOperator(float A, float B, float C, out float L, out float M, out float H)
+{
+	float X = min(B, C);
+	float Y = max(B, C);
+	
+	L = min(A, X);
+	float Z = max(A, X);
+
+	M = min(Z, Y);
+	H = max(Z, Y);
+}
+
+float Median(float A, float B, float C)
+{
+	float L, M, H;
+	LMHOperator(A, B, C, L, M, H);
+	return M;
+}
+
+
+taa_subpixel_payload ReadSubpixelDetails(uint2 SubpixelPos)
+{
+	uint2 HistoryPixelPos = SubpixelPos / SUB_PIXEL_GRID_SIZE;
+
+	uint2 SubpixelRelPos = SubpixelPos % SUB_PIXEL_GRID_SIZE;
+	uint SubpixelId = dot(SubpixelRelPos, uint2(1, SUB_PIXEL_GRID_SIZE));
+	uint SubpixelShift = SubpixelId * SUB_PIXEL_BIT_COUNT;
+	
+	taa_subpixel_details RawHistory3 = History_Textures_3[HistoryPixelPos];
+
+	return taa_subpixel_payload(RawHistory3 >> SubpixelShift) & SUB_PIXEL_BIT_MASK;
+}
+
+
+//------------------------------------------------------- ENTRY POINT
+
+[numthreads(TILE_SIZE, TILE_SIZE, 1)]
+void MainCS(
+	uint2 GroupId : SV_GroupID,
+	uint GroupThreadIndex : SV_GroupIndex)
+{
+	float4 Debug = 0.0;
+
+	uint2 DispatchThreadId = (
+		ZOrder2D(GroupThreadIndex, log2(TILE_SIZE)) +
+		GroupId * uint2(TILE_SIZE, TILE_SIZE));
+
+	uint2 HistoryPixelPos = DispatchThreadId / SUB_PIXEL_GRID_SIZE;
+	float2 ViewportUV = (float2(HistoryPixelPos) + 0.5f) * HistoryInfo_ViewportSizeInverse;
+	float2 ScreenPos = ViewportUVToScreenPos(ViewportUV);
+
+	uint2 SubpixelPos = DispatchThreadId % SUB_PIXEL_GRID_SIZE;
+	uint SubpixelId = dot(SubpixelPos, uint2(1, SUB_PIXEL_GRID_SIZE));
+	uint SubpixelShift = SubpixelId * SUB_PIXEL_BIT_COUNT;
+	
+	{
+		#if 1
+			taa_half3 RawHistory0 = History_Textures_0[HistoryPixelPos];
+			taa_half3 RawHistory1 = History_Textures_1[HistoryPixelPos];
+			taa_half2 RawHistory2 = History_Textures_2[HistoryPixelPos];
+			taa_subpixel_details RawHistory3 = History_Textures_3[HistoryPixelPos];
+		#else
+			taa_half3 RawHistory0 = PrevHistory_Textures_0[HistoryPixelPos];
+			taa_half3 RawHistory1 = PrevHistory_Textures_1[HistoryPixelPos];
+			taa_half2 RawHistory2 = PrevHistory_Textures_2[HistoryPixelPos];
+			taa_subpixel_details RawHistory3 = PrevHistory_Textures_3[HistoryPixelPos];
+		#endif
+
+		taa_subpixel_payload SubpixelDetails = taa_subpixel_payload(RawHistory3 >> SubpixelShift) & SUB_PIXEL_BIT_MASK;
+
+		Debug.r = UncompressParallaxFactor(SubpixelDetails);
+		Debug.b = SubpixelId;
+	}
+
+	#if DEBUG_OUTPUT
+	if (all(DispatchThreadId < HistoryInfo_ViewportMax * SUB_PIXEL_GRID_SIZE))
+	{
+		DebugOutput[DispatchThreadId] = Debug;
+	}
+	#endif
+}
diff --git a/Engine/Shaders/Private/TemporalSuperResolution/TSRDecimateHistory.usf b/Engine/Shaders/Private/TemporalSuperResolution/TSRDecimateHistory.usf
new file mode 100644
--- /dev/null
+++ b/Engine/Shaders/Private/TemporalSuperResolution/TSRDecimateHistory.usf
@@ -0,0 +1,421 @@
+// Copyright Epic Games, Inc. All Rights Reserved.
+
+#include "TSRCommon.ush"
+
+
+//------------------------------------------------------- CONFIG
+
+#define TILE_SIZE 8
+
+#define CONFIG_KERNEL 0
+
+#define CONFIG_OUTPUT_HALF_RES (DIM_OUTPUT_HALF_RES)
+
+#define CONFIG_OUTPUT_INTERFERENCE_SEEDS 0
+
+
+//------------------------------------------------------- PARAMETERS
+
+float4x4 RotationalClipToPrevClip;
+float3 OutputQuantizationError;
+
+Texture2D<taa_half3> InputSceneColorTexture;
+Texture2D<float2> DilatedVelocityTexture;
+Texture2D<float> ClosestDepthTexture;
+Texture2D<uint> PrevUseCountTexture;
+Texture2D<uint> PrevClosestDepthTexture;
+Texture2D<uint> ParallaxFactorTexture;
+
+Texture2D<taa_half3> PrevHistory_Textures_0;
+Texture2D<taa_half3> PrevHistory_Textures_1;
+Texture2D<taa_half2> PrevHistory_Textures_2;
+Texture2D<uint>      PrevHistory_Textures_3;
+Texture2D<taa_half>  PrevHistory_LowResTextures_0;
+Texture2D<uint>      PrevHistory_SuperResTextures_0;
+Texture2D<uint>      PrevHistory_SuperResTextures_1;
+
+RWTexture2D<taa_half3> HalfResSceneColorOutput;
+RWTexture2D<taa_half3> HalfResPredictionSceneColorOutput;
+RWTexture2D<taa_half> HalfResParallaxRejectionMaskOutput;
+
+RWTexture2D<taa_half3> PredictionSceneColorOutput;
+RWTexture2D<taa_half> ParallaxRejectionMaskOutput;
+RWTexture2D<taa_half4> InterferenceSeedOutput;
+
+
+//------------------------------------------------------- LDS
+
+groupshared taa_half3 SharedInputColor[TILE_SIZE * TILE_SIZE];
+groupshared taa_half4 SharedHistoryColorAndMask[TILE_SIZE * TILE_SIZE];
+
+
+//------------------------------------------------------- ENTRY POINT
+
+taa_half3 TransformToComparisonColorSpace(taa_half3 Color)
+{
+	return Color;
+	//return log2(Color);
+	//return Color *HdrWeight4(Color); // *float3(HdrWeightY(Color.r), HdrWeightY(Color.g), HdrWeightY(Color.b));
+}
+
+
+//------------------------------------------------------- ENTRY POINT
+
+[numthreads(TILE_SIZE, TILE_SIZE, 1)]
+void MainCS(
+	uint2 GroupId : SV_GroupID,
+	uint GroupThreadIndex : SV_GroupIndex)
+{
+	float4 Debug = 0;
+	
+	uint2 DispatchThreadId = (
+		ZOrder2D(GroupThreadIndex, log2(TILE_SIZE)) +
+		GroupId * uint2(TILE_SIZE, TILE_SIZE));
+
+	uint2 InputPixelPos = InputInfo_ViewportMin + DispatchThreadId;
+	
+	float2 ViewportUV = (float2(DispatchThreadId) + 0.5f) * InputInfo_ViewportSizeInverse;
+	float2 ScreenPos = ViewportUVToScreenPos(ViewportUV);
+	float2 InputBufferUV = (InputPixelPos + 0.5) * InputInfo_ExtentInverse;
+	
+	float2 EncodedVelocity = DilatedVelocityTexture[InputPixelPos];
+	float DeviceZ = ClosestDepthTexture[InputPixelPos];
+	taa_half ParallaxFactor = UncompressParallaxFactor((ParallaxFactorTexture[InputPixelPos] >> 4) & 0xF);
+	taa_half3 InputSceneColor = InputSceneColorTexture[InputPixelPos];
+	
+	taa_half ClosestInputLuma4 = Luma4(InputSceneColor);
+	taa_half InputInterferenceSeed = LumaToInterferenceSeed(ClosestInputLuma4);
+
+	float WorldDepth = ConvertFromDeviceZ(DeviceZ);
+	
+	// Compute the parralax rejection mask
+	taa_half ParallaxRejectionMask = 0.0;
+	{
+		float2 PrevScreenPos = ScreenPos - DecodeVelocityFromTexture(float4(EncodedVelocity, 0.0, 0.0)).xy;
+		float2 PrevInputBufferUV = (InputInfo_ScreenPosToViewportScale * PrevScreenPos + InputInfo_ScreenPosToViewportBias) * InputInfo_ExtentInverse;
+		
+		FBilinearSampleInfos BilinearInter = GetBilinearSampleLevelInfos(PrevInputBufferUV, InputInfo_Extent, InputInfo_ExtentInverse);
+
+		UNROLL_N(4)
+		for (uint i = 0; i < 4; i++)
+		{
+			float BilinearWeight = GetSampleWeight(BilinearInter, i);
+			uint2 PixelPos = GetSamplePixelCoord(BilinearInter, i);
+
+			bool bValidPixelPos = all(PixelPos >= InputInfo_ViewportMin && PixelPos < InputInfo_ViewportMax);
+
+			#if COMPRESS_PREV_USE_COUNT
+				uint2 CompressedUseCountPixelPos;
+				uint Shift;
+				ComputeCompressedUseCountPixelCoordinates(PixelPos, /* out */ CompressedUseCountPixelPos, /* out */ Shift);
+
+				uint HistoryUseCount = (PrevUseCountTexture[CompressedUseCountPixelPos] >> (8 * Shift)) & 0xFF;
+			#else
+				uint HistoryUseCount = PrevUseCountTexture[PixelPos];
+			#endif
+
+			float HistoryClosestDeviceZ = f16tof32(PrevClosestDepthTexture[PixelPos]);
+			float HistoryClosestWorldDepth = ConvertFromDeviceZ(HistoryClosestDeviceZ);
+
+			taa_half DepthRejection;
+			{
+				const float WorldDepthToPixelWorldRadius = GetTanHalfFieldOfView().x * InputInfo_ExtentInverse.x;
+				const float PixelDepthError = 3.0;
+				const float RadiusToDiamater = 2.0;
+			
+				float WorldDepthEpsilon = (WorldDepthToPixelWorldRadius * RadiusToDiamater * PixelDepthError) * HistoryClosestWorldDepth; 
+				float DeltaDepth = abs(HistoryClosestWorldDepth - WorldDepth);
+
+				DepthRejection = saturate(taa_half(2.0) - taa_half(DeltaDepth / WorldDepthEpsilon));
+			}
+
+			taa_half HistoryGhostingRejection = saturate(taa_half(1.0) - taa_half(3.0) * abs(taa_half(HistoryUseCount) * rcp(taa_half(PREV_USE_COUNT_QUANTIZATION)) - taa_half(1.0)));
+
+			FLATTEN
+			if (bValidPixelPos)
+				ParallaxRejectionMask += taa_half(BilinearWeight) * max(DepthRejection, HistoryGhostingRejection);
+		}
+	}
+
+	// Pixel coordinate of the input pixel in the predictor viewport.
+	float2 InputPPCk = float2(InputPixelPos - InputInfo_ViewportMin) + 0.5f - InputJitter;
+	float2 OutputViewportUV = InputPPCk * InputInfo_ViewportSizeInverse;
+	float2 OutputScreenPos = ViewportUVToScreenPos(OutputViewportUV);
+
+	float2 PrevOutputScreenPos = OutputScreenPos - DecodeVelocityFromTexture(float4(EncodedVelocity, 0.0, 0.0)).xy;
+
+	float2 PrevOutputViewportUV = ScreenPosToViewportUV(PrevOutputScreenPos);
+
+	float2 HistoryPPCk = PrevOutputViewportUV * PrevHistoryInfo_ViewportSize;
+
+	// Pixel coordinate of the center of the nearest predictor pixel O.
+	float2 HistoryPPCo = floor(HistoryPPCk) + 0.5;
+	float2 HistoryBufferUV = HistoryPPCo * PrevHistoryInfo_ExtentInverse;
+
+	// Vector in pixel between pixel K -> O.
+	taa_half2 dKO = taa_half2(HistoryPPCo - HistoryPPCk);
+	
+	// Sample the history
+	taa_half3 DecimatedColor;
+	taa_half HistoryMaxRejectionFactor;
+
+	taa_half ClosestInterferenceSeed;
+	taa_half IntegratedInterferenceSeed;
+	taa_half MinInterferenceSeed;
+	taa_half MaxInterferenceSeed;
+	{
+		float2 KernelHistoryBufferUV = clamp(
+			HistoryPPCk * PrevHistoryInfo_ExtentInverse,
+			PrevHistoryInfo_UVViewportBilinearMin,
+			PrevHistoryInfo_UVViewportBilinearMax);
+		
+		float2 PrevInputBufferUV = (InputInfo_ScreenPosToViewportScale * PrevOutputScreenPos + InputInfo_ScreenPosToViewportBias) * InputInfo_ExtentInverse;
+
+		#if CONFIG_OUTPUT_INTERFERENCE_SEEDS
+		{
+			const uint KernelSize = 4;
+
+			uint SuperResPayloads[16];
+
+			uint ClosestPayload = ~0;
+			uint ClosestAge = 0;
+
+			// Sample the super res history and find an aproximative age of the input pixel.
+			UNROLL
+			for (uint y = 0; y < KernelSize; y++)
+			{
+				UNROLL
+				for (uint x = 0; x < KernelSize; x++)
+				{
+					const float2 Offset = float2(x, y) - (KernelSize - 1.0) * 0.5;
+					
+					float2 SampleHistoryBufferUV = clamp(
+						(HistoryPPCk * 2.0 + Offset) * PrevHistoryInfo_ExtentInverse * rcp(float(SUPERRES_RES_MULTIPLIER)),
+						PrevHistoryInfo_UVViewportBilinearMin,
+						PrevHistoryInfo_UVViewportBilinearMax);
+
+					uint2 SuperResPixelCoord = uint2(SampleHistoryBufferUV * PrevHistoryInfo_Extent * float(SUPERRES_RES_MULTIPLIER));
+
+					uint SubpixelPayload = PrevHistory_SuperResTextures_0[SuperResPixelCoord];
+					FSuperResHistory UnpackedSubpixel = UnpackSuperResHistory(SubpixelPayload);
+
+					SuperResPayloads[y * 4 + x] = SubpixelPayload;
+
+					taa_half DistanceToSeed = abs(InputInterferenceSeed - UnpackedSubpixel.InterferenceSeed);
+
+					FLATTEN
+					if (DistanceToSeed < 0.1 && UnpackedSubpixel.Age > ClosestAge)
+					{
+						ClosestPayload = SubpixelPayload;
+						ClosestAge = UnpackedSubpixel.Age;
+					}
+				}
+			}
+			
+			uint AgeMin = UnpackSuperResHistory(ClosestPayload).Age;
+			uint AgeMax = UnpackSuperResHistory(ClosestPayload).Age + 4;
+
+			ClosestInterferenceSeed = UnpackSuperResHistory(ClosestPayload).InterferenceSeed;
+
+			MinInterferenceSeed = 1.0;
+			MaxInterferenceSeed = 0.0;
+
+			// Find out if there is high frequency details in the history with that age.
+			UNROLL
+			for (uint i = 0; i < KernelSize * KernelSize; i++)
+			{
+				FSuperResHistory UnpackedSubpixel = UnpackSuperResHistory(SuperResPayloads[i]);
+
+				bool bAgeBilateralWeight = UnpackedSubpixel.Age >= AgeMin && UnpackedSubpixel.Age <= AgeMax;
+
+				FLATTEN
+				if (bAgeBilateralWeight && ClosestPayload != ~0)
+				{
+					MinInterferenceSeed = min(MinInterferenceSeed, UnpackedSubpixel.InterferenceSeed);
+					MaxInterferenceSeed = max(MaxInterferenceSeed, UnpackedSubpixel.InterferenceSeed);
+				}
+
+			}
+			
+			// actually integrating SubpixelSeed stil suffer interference that causes instability. So instead take
+			// average of min and max that turns out interference free.
+			IntegratedInterferenceSeed = (MinInterferenceSeed + MaxInterferenceSeed) * 0.5;
+		}
+		#endif
+
+		taa_half PrevUninterferenceSeed = PrevHistory_LowResTextures_0.SampleLevel(GlobalBilinearClampedSampler, PrevInputBufferUV, 0);
+		taa_half PrevUninterferenceLuma = PrevUninterferenceSeed * HdrWeightInvY(PrevUninterferenceSeed);
+
+		#if !CONFIG_OUTPUT_INTERFERENCE_SEEDS
+		{
+			DecimatedColor = PrevHistory_Textures_0.SampleLevel(GlobalBilinearClampedSampler, KernelHistoryBufferUV, 0);
+			
+			// Cancels out the flickering Luma component with the interference free luma from the inteference seed of the history.
+			if (bEnableInterferenceHeuristic)
+			{
+				DecimatedColor *= PrevUninterferenceLuma * SafeRcp(Luma4(DecimatedColor));
+			}
+		}
+		#else
+		{
+			DecimatedColor = 0.0;
+			UNROLL_N(4)
+			for (uint i = 0; i < 4; i++)
+			{
+				const float2 PixelOffset = float2(Offsets2x2[i]) - 0.5;
+
+				float2 SampleBufferUV = KernelHistoryBufferUV + PixelOffset * PrevHistoryInfo_ExtentInverse; 
+				SampleBufferUV = clamp(
+					SampleBufferUV,
+					PrevHistoryInfo_UVViewportBilinearMin,
+					PrevHistoryInfo_UVViewportBilinearMax);
+				
+				taa_half3 HistorySample = PrevHistory_Textures_0.SampleLevel(GlobalBilinearClampedSampler, SampleBufferUV, 0);
+
+				// Cancels out the flickering Luma component with the interference free luma from the inteference seed of the history.
+				if (bEnableInterferenceHeuristic)
+				{
+					HistorySample *= PrevUninterferenceLuma * SafeRcp(Luma4(HistorySample));
+				}
+
+				DecimatedColor += HistorySample;
+			}
+
+			DecimatedColor *= rcp(4.0);
+		}
+		#endif
+
+		// Compute the parallax rejection factor of the details in the history.
+		FSubpixelNeighborhood SubpixelNeigborhood = GatherPrevSubpixelNeighborhood(PrevHistory_Textures_3, KernelHistoryBufferUV);
+		{
+			uint EncodedHistoryMaxRejectionFactor = 0;
+
+			UNROLL_N(SUB_PIXEL_COUNT)
+			for (uint SubpixelId = 0; SubpixelId < SUB_PIXEL_COUNT; SubpixelId++)
+			{
+				taa_subpixel_payload SubpixelPayload = GetSubpixelPayload(SubpixelNeigborhood, SubpixelId);
+				taa_subpixel_payload EncodedRejectionFactor = (SubpixelPayload >> SUB_PIXEL_PARALLAX_FACTOR_BIT_OFFSET) & SUB_PIXEL_PARALLAX_FACTOR_BIT_MASK;
+
+				EncodedHistoryMaxRejectionFactor = max(EncodedHistoryMaxRejectionFactor, EncodedRejectionFactor);
+			}
+
+			HistoryMaxRejectionFactor = UncompressParallaxFactor(EncodedHistoryMaxRejectionFactor);
+		}
+	}
+	
+	// Reject subpixel details that might be in history but not input just based on their parrallax
+	{
+		taa_half HistorySubdetailRejectionFactor = max(HistoryMaxRejectionFactor / ParallaxFactor, taa_half(1.0));
+
+		float4 ThisClip = float4(ScreenPos, DeviceZ, 1);
+		float4 PrevClip = mul(ThisClip, View.ClipToPrevClip);
+		float2 PrevScreen = PrevClip.xy / PrevClip.w;
+			
+		float4 RotationalPrevClip = mul(ThisClip, RotationalClipToPrevClip);
+		float2 RotationalPrevScreen = RotationalPrevClip.xy / RotationalPrevClip.w;
+			
+		taa_half PixelParallaxVelocity = taa_half(0.5) * length(taa_half2(RotationalPrevScreen - PrevScreen) * InputInfo_Extent);
+		taa_half HistorySubdetailAdditionalParralax = (HistorySubdetailRejectionFactor - taa_half(1.0)) * PixelParallaxVelocity;
+
+		taa_half HistoryParallaxRejectionMask = saturate(taa_half(1.0) - HistorySubdetailAdditionalParralax);
+
+		ParallaxRejectionMask = min(ParallaxRejectionMask, HistoryParallaxRejectionMask);
+	}
+
+	// Immediately reject based on the rejection mask
+	FLATTEN
+	if ((PrevOutputScreenPos.x < -1.0) ||
+		(PrevOutputScreenPos.y < -1.0) ||
+		(PrevOutputScreenPos.x > 1.0) ||
+		(PrevOutputScreenPos.y > 1.0))
+	//if (IsOffScreen2(bCameraCut, PrevOutputScreenPos, /* ParallaxRejectionMask = */ 1.0))
+	{
+		ParallaxRejectionMask = taa_half(0.0);
+	}
+
+	
+	#if CONFIG_OUTPUT_HALF_RES
+	{
+		SharedInputColor[GroupThreadIndex] = InputSceneColor;
+		SharedHistoryColorAndMask[GroupThreadIndex] = taa_half4(DecimatedColor, ParallaxRejectionMask);
+
+		GroupMemoryBarrierWithGroupSync();
+
+		bool bCheckerboard = (((InputPixelPos.x ^ InputPixelPos.y) >> 1) & 0x1) != 0;
+
+		half3 HalfResInputSceneColor = InputSceneColor;
+		half3 HalfResDecimatedColor = DecimatedColor;
+		half HalfResParallaxRejectionMask = ParallaxRejectionMask;
+
+		UNROLL
+		for (uint i = 1; i < 4; i++)
+		{
+			uint ButterflySwap = i;
+			
+			if (bCheckerboard)
+			{
+				HalfResInputSceneColor = min(HalfResInputSceneColor, SharedInputColor[GroupThreadIndex ^ ButterflySwap]);
+			}
+			else
+			{
+				HalfResInputSceneColor = max(HalfResInputSceneColor, SharedInputColor[GroupThreadIndex ^ ButterflySwap]);
+			}
+			HalfResDecimatedColor += SharedHistoryColorAndMask[GroupThreadIndex ^ ButterflySwap].rgb;
+
+			HalfResParallaxRejectionMask += SharedHistoryColorAndMask[GroupThreadIndex ^ ButterflySwap].a;
+		}
+
+		HalfResDecimatedColor *= 0.25;
+		HalfResParallaxRejectionMask *= 0.25;
+
+		#if CONFIG_ENABLE_STOCASTIC_QUANTIZATION
+		{
+			uint2 Random = Rand3DPCG16(int3((InputPixelPos - InputInfo_ViewportMin) / 2, View.StateFrameIndexMod8)).xy;
+			float2 E = Hammersley16(0, 1, Random);
+
+			HalfResInputSceneColor = QuantizeForFloatRenderTarget(HalfResInputSceneColor, E.x, OutputQuantizationError);
+			HalfResDecimatedColor = QuantizeForFloatRenderTarget(HalfResDecimatedColor, E.x, OutputQuantizationError);
+		}
+		#endif
+
+		if (all(InputPixelPos < InputInfo_ViewportMax) && all(((InputPixelPos - InputInfo_ViewportMin) % 2) == 0))
+		{
+			HalfResSceneColorOutput[InputPixelPos / 2] = HalfResInputSceneColor;
+			HalfResPredictionSceneColorOutput[InputPixelPos / 2] = HalfResDecimatedColor;
+			HalfResParallaxRejectionMaskOutput[InputPixelPos / 2] = HalfResParallaxRejectionMask;
+		}
+	}
+	#endif
+
+	#if CONFIG_ENABLE_STOCASTIC_QUANTIZATION
+	{
+		uint2 Random = Rand3DPCG16(int3(InputPixelPos - InputInfo_ViewportMin, View.StateFrameIndexMod8)).xy;
+		float2 E = Hammersley16(0, 1, Random);
+
+		DecimatedColor = QuantizeForFloatRenderTarget(DecimatedColor, E.x, OutputQuantizationError);
+	}
+	#endif
+
+	if (all(InputPixelPos < InputInfo_ViewportMax))
+	{
+		#if !CONFIG_OUTPUT_HALF_RES
+		{
+			PredictionSceneColorOutput[InputPixelPos] = DecimatedColor;
+		}
+		#endif
+		
+		ParallaxRejectionMaskOutput[InputPixelPos] = ParallaxRejectionMask;
+
+		#if CONFIG_OUTPUT_INTERFERENCE_SEEDS
+		{
+			InterferenceSeedOutput[InputPixelPos] = taa_half4(ClosestInterferenceSeed, IntegratedInterferenceSeed, MinInterferenceSeed, MaxInterferenceSeed);
+		}
+		#endif
+
+		#if DEBUG_OUTPUT
+		{
+			DebugOutput[InputPixelPos] = Debug;
+		}
+		#endif
+	}
+}
diff --git a/Engine/Shaders/Private/TemporalSuperResolution/TSRDetectInterference.usf b/Engine/Shaders/Private/TemporalSuperResolution/TSRDetectInterference.usf
new file mode 100644
--- /dev/null
+++ b/Engine/Shaders/Private/TemporalSuperResolution/TSRDetectInterference.usf
@@ -0,0 +1,207 @@
+// Copyright Epic Games, Inc. All Rights Reserved.
+
+#include "TSRCommon.ush"
+
+
+//------------------------------------------------------- CONFIG
+
+#define TILE_SIZE 8
+
+#define CONFIG_KERNEL 0
+
+#define CONFIG_OUTPUT_HALF_RES (DIM_OUTPUT_HALF_RES)
+
+#define CONFIG_SAMPLES 5
+
+
+//------------------------------------------------------- PARAMETERS
+
+float4x4 RotationalClipToPrevClip;
+float3 OutputQuantizationError;
+
+Texture2D<taa_half3> InputSceneColorTexture;
+Texture2D<taa_half3> PredictionSceneColorTexture;
+Texture2D<taa_half> ParallaxRejectionMaskTexture;
+Texture2D<taa_half4> InterferenceSeedTexture;
+
+Texture2D<float2> DilatedVelocityTexture;
+
+Texture2D<taa_half> PrevHistory_LowResTextures_0;
+
+RWTexture2D<taa_half2> InterferenceWeightOutput;
+RWTexture2D<taa_half2> InterferenceSeedUpdateOutput;
+RWTexture2D<taa_half2> HistoryOutput_LowResTextures_0;
+
+
+//------------------------------------------------------- ENTRY POINT
+
+[numthreads(TILE_SIZE, TILE_SIZE, 1)]
+void MainCS(
+	uint2 GroupId : SV_GroupID,
+	uint GroupThreadIndex : SV_GroupIndex)
+{
+	float4 Debug = 0;
+	
+	uint2 DispatchThreadId = (
+		ZOrder2D(GroupThreadIndex, log2(TILE_SIZE)) +
+		GroupId * uint2(TILE_SIZE, TILE_SIZE));
+
+	uint2 InputPixelPos = InputInfo_ViewportMin + DispatchThreadId;
+	
+	float2 ViewportUV = (float2(DispatchThreadId) + 0.5f) * InputInfo_ViewportSizeInverse;
+	float2 ScreenPos = ViewportUVToScreenPos(ViewportUV);
+
+	float2 EncodedVelocity = DilatedVelocityTexture[InputPixelPos];
+	
+	taa_half ParallaxRejectionMask;
+	taa_half InputLuma;
+	taa_half InputLumaMin;
+	taa_half InputLumaMax;
+	taa_half ClosestInterferenceSeedMin;
+	taa_half ClosestInterferenceSeedMax;
+	taa_half IntegratedInterferenceSeed;
+	taa_half IntegratedInterferenceSeedMin;
+	taa_half IntegratedInterferenceSeedMax;
+	taa_half MinInterferenceSeed;
+	taa_half MaxInterferenceSeed;
+	{
+		ParallaxRejectionMask = ParallaxRejectionMaskTexture[InputPixelPos];
+
+		InputLuma = Luma4(InputSceneColorTexture[InputPixelPos]);
+		InputLumaMin = InputLuma;
+		InputLumaMax = InputLuma;
+
+		taa_half4 InterferenceSeed = InterferenceSeedTexture[InputPixelPos];
+
+		ClosestInterferenceSeedMin = InterferenceSeed.x;
+		ClosestInterferenceSeedMax = InterferenceSeed.x; 
+		IntegratedInterferenceSeed = InterferenceSeed.y;
+		IntegratedInterferenceSeedMin = InterferenceSeed.y;
+		IntegratedInterferenceSeedMax = InterferenceSeed.y;
+		MinInterferenceSeed = InterferenceSeed.z;
+		MaxInterferenceSeed = InterferenceSeed.w;
+	}
+	
+	float2 PrevScreenPos = ScreenPos - DecodeVelocityFromTexture(float4(EncodedVelocity, 0.0, 0.0)).xy;
+	bool bOffScreen = IsOffScreen(bCameraCut, PrevScreenPos, ParallaxRejectionMask);
+
+	taa_half PrevUninterferenceSeed;
+	{
+		float2 PrevInputBufferUV = (InputInfo_ScreenPosToViewportScale * PrevScreenPos + InputInfo_ScreenPosToViewportBias) * InputInfo_ExtentInverse;
+		
+		PrevUninterferenceSeed = PrevHistory_LowResTextures_0.SampleLevel(GlobalBilinearClampedSampler, PrevInputBufferUV, 0);
+	}
+
+
+	taa_half NeigborMaxInterferenceSeedGradient = 0.0;
+	
+	UNROLL_N(CONFIG_SAMPLES - 1)
+	for (uint SampleId = 1; SampleId < CONFIG_SAMPLES; SampleId++)
+	{
+		taa_half2 PixelOffset;
+			
+		#if CONFIG_SAMPLES == 9
+		{
+			PixelOffset = taa_half2(kOffsets3x3[kSquareIndexes3x3[SampleId]]);
+		}
+		#elif CONFIG_SAMPLES == 5 || CONFIG_SAMPLES == 6
+		{
+			if (SampleId == 5)
+			{
+				PixelOffset = taa_half2(SignFastInt(InputJitter)); // TODO: shader parameter
+			}
+			else
+			{
+				PixelOffset = taa_half2(kOffsets3x3[kPlusIndexes3x3[SampleId]]);
+			}
+		}
+		#else
+			#error Unknown sample count
+		#endif
+
+		float2 SampleInputBufferUV = (InputPixelPos + PixelOffset + 0.5) * InputInfo_ExtentInverse;
+		SampleInputBufferUV = clamp(SampleInputBufferUV, InputInfo_UVViewportBilinearMin, InputInfo_UVViewportBilinearMax);
+
+		taa_half NeighborInputLuma = Luma4(InputSceneColorTexture.SampleLevel(GlobalPointClampedSampler, SampleInputBufferUV, 0));
+		taa_half4 InterferenceSeed = InterferenceSeedTexture.SampleLevel(GlobalPointClampedSampler, SampleInputBufferUV, 0);
+		
+		taa_half InterferenceSeedGradient = InterferenceSeed.w - InterferenceSeed.z;
+
+		InputLumaMin = min(InputLumaMin, NeighborInputLuma);
+		InputLumaMax = max(InputLumaMax, NeighborInputLuma);
+
+		ClosestInterferenceSeedMin = min(ClosestInterferenceSeedMin, InterferenceSeed.x);
+		ClosestInterferenceSeedMax = max(ClosestInterferenceSeedMax, InterferenceSeed.x); 
+		
+		// TODO(TSR): this is InterferenceSeed.w - InterferenceSeed.z
+		IntegratedInterferenceSeedMin = min(IntegratedInterferenceSeedMin, InterferenceSeed.y);
+		IntegratedInterferenceSeedMax = max(IntegratedInterferenceSeedMax, InterferenceSeed.y);
+		
+		NeigborMaxInterferenceSeedGradient = max(NeigborMaxInterferenceSeedGradient, InterferenceSeedGradient); 
+	}
+
+	taa_half UninterferenceSeed;
+	bool bIsInterference;
+	{
+		// Error due to the low bitdepth of the inteference seeds in the history.
+		const taa_half InterferenceSeedEncodingError = rcp(taa_half(SUPERRES_INTERFERENCE_BITMASK));
+
+		taa_half InputSeedMin = LumaToInterferenceSeed(InputLumaMin);
+		taa_half InputSeedMax = LumaToInterferenceSeed(InputLumaMax);
+
+		// Clamp the uninteference seed with current frame input to see what happens
+		taa_half ClampedPrevUninterferenceSeed = clamp(PrevUninterferenceSeed, InputSeedMin, InputSeedMax);
+		
+		// Measure the low frequency gradien in the seed. Needs to happen both on current frame input to avoid ghosting on
+		// current frame changes with high contrast geometric edges, and previous frames to avoid ghosting on area that had 
+		// high contrast geometric edges but no longer had in current frame due to shading change.
+		taa_half InterferenceSeedGradient = max(InputSeedMax - InputSeedMin, IntegratedInterferenceSeedMax - IntegratedInterferenceSeedMin);
+
+		// Test whether the current frame will be rejected.
+		bool bRejectPrevScene = abs(ClampedPrevUninterferenceSeed - PrevUninterferenceSeed) > InterferenceSeedEncodingError;
+		
+		// Detect if there is any inteference based on any subpixel detail range from the lowest LDR to highest LDR.
+		bool bRejectInterferenceSeedWithGradient = (
+			(MinInterferenceSeed < ClosestInterferenceSeedMin - InterferenceSeedEncodingError - InterferenceSeedGradient) ||
+			(MaxInterferenceSeed > ClosestInterferenceSeedMax + InterferenceSeedEncodingError + InterferenceSeedGradient));
+			
+		bool bRejectInterferenceSeed = (
+			(MinInterferenceSeed < ClosestInterferenceSeedMin - InterferenceSeedEncodingError) ||
+			(MaxInterferenceSeed > ClosestInterferenceSeedMax + InterferenceSeedEncodingError));
+
+		// Detect interference solely based on the inteference seed, but need to take into account low frequency
+		// gradien that there may be in the neighborhood to avoid ghosting on high contrast geometric edges.
+		bIsInterference = bRejectInterferenceSeedWithGradient;
+
+		// Don't clamp the seed history if there is an interference to keep luma stability in the input
+		// of the shading heuristic next frame.
+		FLATTEN
+		if (bIsInterference)
+		{
+			ClampedPrevUninterferenceSeed = PrevUninterferenceSeed;
+		}
+
+		FLATTEN
+		if (bOffScreen)
+		{
+			bIsInterference = false;
+			ClampedPrevUninterferenceSeed = LumaToInterferenceSeed(InputLuma);
+		}
+
+		UninterferenceSeed = lerp(ClampedPrevUninterferenceSeed, LumaToInterferenceSeed(InputLuma), 0.05);
+	}
+
+	// TODO(TSR): optimise VGPR here
+	if (all(InputPixelPos < InputInfo_ViewportMax))
+	{
+		InterferenceWeightOutput[InputPixelPos] = bIsInterference ? 0.0 : 1.0;
+		
+		HistoryOutput_LowResTextures_0[InputPixelPos] = UninterferenceSeed;
+
+		#if DEBUG_OUTPUT
+		{
+			DebugOutput[InputPixelPos] = Debug;
+		}
+		#endif
+	}
+}
diff --git a/Engine/Shaders/Private/TemporalSuperResolution/TSRDilateRejection.usf b/Engine/Shaders/Private/TemporalSuperResolution/TSRDilateRejection.usf
new file mode 100644
--- /dev/null
+++ b/Engine/Shaders/Private/TemporalSuperResolution/TSRDilateRejection.usf
@@ -0,0 +1,70 @@
+// Copyright Epic Games, Inc. All Rights Reserved.
+
+#include "TSRCommon.ush"
+
+
+//------------------------------------------------------- CONFIG
+
+#define TILE_SIZE 8
+
+
+//------------------------------------------------------- PARAMETERS
+
+Texture2D<taa_half> HistoryRejectionTexture;
+RWTexture2D<taa_half> DilatedHistoryRejectionOutput;
+
+
+//------------------------------------------------------- ENTRY POINT
+
+[numthreads(TILE_SIZE, TILE_SIZE, 1)]
+void MainCS(
+	uint2 GroupId : SV_GroupID,
+	uint GroupThreadIndex : SV_GroupIndex)
+{
+	float4 Debug = 0.0;
+	
+	taa_short2 RejectionPixelPos = (
+		taa_short2(GroupId) * taa_short2(TILE_SIZE, TILE_SIZE) +
+		Map8x8Tile2x2Lane(GroupThreadIndex));
+
+	// Fetch center of the kernel
+	taa_half HistoryRejectionInfo = HistoryRejectionTexture[RejectionPixelPos];
+
+	const int KernelSize = 2;
+
+	// Fetch neighborhood of the kernel
+	UNROLL
+	for (int x = -KernelSize; x <= KernelSize; x++)
+	{
+		UNROLL
+		for (int y = -KernelSize; y <= KernelSize; y++)
+		{
+			taa_short2 Offset = taa_short2(x, y);
+		
+			if (all(Offset == 0))
+			{
+				continue;
+			}
+			
+			taa_short2 SampleInputPixelPos = RejectionPixelPos + ClampPixelOffset(
+				RejectionPixelPos,
+				Offset, Offset,
+				RejectionInfo_ViewportMin, RejectionInfo_ViewportMax - 1);
+
+			taa_half SampleRejectionInfo = HistoryRejectionTexture[SampleInputPixelPos];
+
+			HistoryRejectionInfo = min(HistoryRejectionInfo, SampleRejectionInfo);
+		}
+	}
+	
+	bool bIsValid = all(RejectionPixelPos < taa_short2(RejectionInfo_ViewportMax));
+	RejectionPixelPos.x = bIsValid ? RejectionPixelPos.x : ~taa_ushort(0);
+
+	DilatedHistoryRejectionOutput[RejectionPixelPos] = HistoryRejectionInfo;
+
+	#if DEBUG_OUTPUT
+	{
+		DebugOutput[RejectionPixelPos] = Debug;
+	}
+	#endif
+}
diff --git a/Engine/Shaders/Private/TemporalSuperResolution/TSRDilateVelocity.usf b/Engine/Shaders/Private/TemporalSuperResolution/TSRDilateVelocity.usf
new file mode 100644
--- /dev/null
+++ b/Engine/Shaders/Private/TemporalSuperResolution/TSRDilateVelocity.usf
@@ -0,0 +1,180 @@
+// Copyright Epic Games, Inc. All Rights Reserved.
+
+#include "TSRCommon.ush"
+
+
+//------------------------------------------------------- CONFIG
+
+#define TILE_SIZE 8
+
+#define AA_CROSS 1
+
+
+//------------------------------------------------------- PARAMETERS
+
+float2 PrevOutputBufferUVMin;
+float2 PrevOutputBufferUVMax;
+
+Texture2D SceneDepthTexture;
+Texture2D SceneVelocityTexture;
+
+RWTexture2D<float2> DilatedVelocityOutput;
+RWTexture2D<float> ClosestDepthOutput;
+RWTexture2D<uint> PrevUseCountOutput;
+RWTexture2D<uint> PrevClosestDepthOutput;
+RWTexture2D<uint> ParallaxFactorOutput;
+
+
+//------------------------------------------------------- ENTRY POINT
+
+uint PackDepthAndOffset(const int2 Offset, float DeviceZ)
+{
+	return ((asuint(DeviceZ) << 2) & ~0xF) | (uint(1 + Offset.x) << 0) | (uint(1 + Offset.y) << 2);
+}
+
+void UnpackDepthAndOffset(uint Packed, out int2 Offset, out float DeviceZ)
+{
+	DeviceZ = asfloat((Packed & ~0xF) >> 2);
+	Offset.x = int((Packed >> 0) & 0x3) - 1;
+	Offset.y = int((Packed >> 2) & 0x3) - 1;
+}
+
+
+//------------------------------------------------------- ENTRY POINT
+
+[numthreads(TILE_SIZE, TILE_SIZE, 1)]
+void MainCS(
+	uint2 GroupId : SV_GroupID,
+	uint GroupThreadIndex : SV_GroupIndex)
+{
+	float4 Debug = 0.0;
+
+	uint2 DispatchThreadId = (
+		ZOrder2D(GroupThreadIndex, log2(TILE_SIZE)) +
+		GroupId * uint2(TILE_SIZE, TILE_SIZE));
+
+	uint2 InputPixelPos = InputInfo_ViewportMin + DispatchThreadId;
+	
+	float2 ViewportUV = (float2(DispatchThreadId) + 0.5f) * InputInfo_ViewportSizeInverse;
+	float2 ScreenPos = ViewportUVToScreenPos(ViewportUV);
+	float2 InputBufferUV = (InputPixelPos + 0.5) * InputInfo_ExtentInverse;
+
+	float DeviceZ = SceneDepthTexture[InputPixelPos].r;
+
+	float ClosestDeviceZ = DeviceZ;
+	float2 VelocityPixelOffset = float2(0.0, 0.0);
+	float ParallaxFactorForHistory = 1.0;
+	float ParallaxFactorForHistoryRejection = 1.0;
+	#if AA_CROSS
+	{
+		const int2 kOffsets[4] = {
+			int2( 1,  0),
+			int2( 1,  1),
+			int2( 0,  1),
+			int2(-1,  1),
+		};
+
+		uint PackedDepthOffset = PackDepthAndOffset(/* Offset = */ int2(0, 0), DeviceZ);
+		
+		float FurthestPairDeviceZ = DeviceZ;
+		float FurthestDeviceZ = DeviceZ;
+
+		#if 1
+			UNROLL_N(4)
+			for (uint i = 0; i < 4; i++)
+		#else
+			UNROLL_N(2)
+			for (uint i = 1; i < 4; i+= 2)
+		#endif
+		{
+			const int2 Offset = kOffsets[i];
+
+			float SampleDeviceZ0 = SceneDepthTexture[int2(InputPixelPos) + Offset].r;
+			float SampleDeviceZ1 = SceneDepthTexture[int2(InputPixelPos) - Offset].r;
+			
+			PackedDepthOffset = max(PackedDepthOffset, PackDepthAndOffset(+Offset, SampleDeviceZ0));
+			PackedDepthOffset = max(PackedDepthOffset, PackDepthAndOffset(-Offset, SampleDeviceZ1));
+
+			FurthestPairDeviceZ = min(FurthestPairDeviceZ, max(SampleDeviceZ0, SampleDeviceZ1));
+			FurthestDeviceZ = min(FurthestDeviceZ, min(SampleDeviceZ0, SampleDeviceZ1));
+		}
+
+		ParallaxFactorForHistory = clamp(DeviceZ / FurthestPairDeviceZ, 0.0, MAX_PARALLAX_FACTOR);
+		ParallaxFactorForHistoryRejection = clamp(DeviceZ / FurthestDeviceZ, 0.0, MAX_PARALLAX_FACTOR);
+
+		UnpackDepthAndOffset(PackedDepthOffset, /* out */ VelocityPixelOffset, /* out */ ClosestDeviceZ);
+	}
+	#endif
+	
+	// Camera motion for pixel.
+	float2 FinalScreenVelocity = 0.0;
+
+	// Compute final screen vecocity.
+	#if 1
+	{
+		FinalScreenVelocity = ComputeStaticVelocity(ScreenPos, ClosestDeviceZ);
+
+		#if 1
+		{
+			float4 EncodedVelocity = SceneVelocityTexture[uint2(InputBufferUV * InputInfo_Extent + VelocityPixelOffset)];
+			if (EncodedVelocity.x > 0.0)
+			{
+				FinalScreenVelocity = DecodeVelocityFromTexture(EncodedVelocity).xy;
+			}
+		}
+		#endif
+	}
+	#endif
+
+	// Scatter parallax rejection
+	float2 PrevScreenPos = ScreenPos - FinalScreenVelocity;
+	float2 PrevInputBufferUV = (InputInfo_ScreenPosToViewportScale * PrevScreenPos + InputInfo_ScreenPosToViewportBias) * InputInfo_ExtentInverse;
+
+
+	Debug.rg = VelocityPixelOffset * 0.5 + 0.5;
+
+	if (all(InputPixelPos < InputInfo_ViewportMax))
+	{
+		#if 1
+		{
+			bool bValidHistoryCoord = all(PrevInputBufferUV > PrevOutputBufferUVMin && PrevInputBufferUV < PrevOutputBufferUVMax);
+
+			FBilinearSampleInfos BilinearInter = GetBilinearSampleLevelInfos(PrevInputBufferUV, InputInfo_Extent, InputInfo_ExtentInverse);
+
+			UNROLL_N(4)
+			for (uint i = 0; i < 4; i++)
+			{
+				float Weight = GetSampleWeight(BilinearInter, i);
+				uint2 PixelPos = GetSamplePixelCoord(BilinearInter, i);
+
+				uint FinalAdd = uint(round(Weight * PREV_USE_COUNT_QUANTIZATION));
+			
+				PixelPos.x = ((FinalAdd > 0) && bValidHistoryCoord) ? PixelPos.x : (~0);
+
+				#if COMPRESS_PREV_USE_COUNT
+					uint2 CompressedUseCountPixelPos;
+					uint Shift;
+					ComputeCompressedUseCountPixelCoordinates(PixelPos, /* out */ CompressedUseCountPixelPos, /* out */ Shift);
+
+					FinalAdd = FinalAdd << (8 * Shift);
+				#else
+					uint2 CompressedUseCountPixelPos = PixelPos;
+				#endif
+
+				InterlockedAdd(PrevUseCountOutput[CompressedUseCountPixelPos], FinalAdd);
+				InterlockedMax(PrevClosestDepthOutput[PixelPos], f32tof16(ClosestDeviceZ));
+			}
+		}
+		#endif
+		
+		DilatedVelocityOutput[InputPixelPos] = EncodeVelocityToTexture(float3(FinalScreenVelocity, 0.0)).xy;
+		ClosestDepthOutput[InputPixelPos] = ClosestDeviceZ;
+		ParallaxFactorOutput[InputPixelPos] = CompressParallaxFactor(ParallaxFactorForHistory) | (CompressParallaxFactor(ParallaxFactorForHistoryRejection) << 4);
+
+		#if DEBUG_OUTPUT
+		{
+			DebugOutput[InputPixelPos] = Debug;
+		}
+		#endif
+	}
+}
diff --git a/Engine/Shaders/Private/TemporalSuperResolution/TSRFilterFrequencies.usf b/Engine/Shaders/Private/TemporalSuperResolution/TSRFilterFrequencies.usf
new file mode 100644
--- /dev/null
+++ b/Engine/Shaders/Private/TemporalSuperResolution/TSRFilterFrequencies.usf
@@ -0,0 +1,237 @@
+// Copyright Epic Games, Inc. All Rights Reserved.
+
+#include "TSRCommon.ush"
+
+
+//------------------------------------------------------- ENUMS
+
+/** Kernel enum */
+	#define KERNEL_1X1 0
+	#define KERNEL_3X3 1
+	#define KERNEL_3X3_PLUS 2
+
+
+//------------------------------------------------------- CONFIG
+
+#define TILE_SIZE 16
+
+#define CONFIG_ACCUMULATE_KARIS_WEIGHTING 1
+
+#define CONFIG_KERNEL (KERNEL_3X3)
+
+// Use a spread factor the size of the clamping box in the history comparison to avoid one flickering input pixel
+// to flicker the entire comparison neighborhood
+#define CONFIG_KERNEL_SPREAD_FACTOR 3
+
+// Uses a spread factor that varies slightly between 2x2 pixel neighborhood to reduce chances of inteference
+// against structured input.
+#define CONFIG_USE_DYNAMIC_SPREAD_FACTOR 1
+
+
+//------------------------------------------------------- PARAMETERS
+
+float3 OutputQuantizationError;
+
+Texture2D<taa_half3> InputTexture;
+Texture2D<taa_half3> PredictionSceneColorTexture;
+Texture2D<taa_half> ParallaxRejectionMaskTexture;
+Texture2D<taa_half> InterferenceWeightTexture;
+
+RWTexture2D<taa_half3> FilteredInputOutput;
+RWTexture2D<taa_half3> FilteredPredictionSceneColorOutput;
+
+
+//------------------------------------------------------- FUNCTIONS
+
+taa_half ComputeColorWeight(taa_half3 Color)
+#if CONFIG_ACCUMULATE_KARIS_WEIGHTING
+{
+	return HdrWeight4(Color);
+}
+#else
+{
+	return taa_half(1);
+}
+#endif
+
+taa_half ParallaxRejectionMaskToHistoryValidity(taa_half ParallaxRejectionMask)
+{
+	return (ParallaxRejectionMask < taa_half(PARALLAX_REJECTION_MASK_THRESHOLD)) ? taa_half(0.0) : taa_half(1.0);
+}
+
+//------------------------------------------------------- ENTRY POINT
+
+[numthreads(TILE_SIZE, TILE_SIZE, 1)]
+void MainCS(
+	uint2 GroupId : SV_GroupID,
+	uint GroupThreadIndex : SV_GroupIndex)
+{
+	uint GroupWaveIndex = GetGroupWaveIndex(GroupThreadIndex, /* GroupSize = */ TILE_SIZE * TILE_SIZE);
+
+	float4 Debug = 0.0;
+	
+	taa_short2 InputPixelPos = (
+		taa_short2(InputPixelPosMin) +
+		taa_short2(GroupId) * taa_short2(TILE_SIZE, TILE_SIZE) +
+		Map16x16Tile2x2Lane(GroupThreadIndex));
+
+	// Fetch center of the kernel
+	taa_half3 CenterInput;
+	taa_half3 FilteredInput;
+	taa_half FilteredInputWeight;
+	taa_half3 FilteredHistory;
+	taa_half FilteredHistoryWeight;
+
+	{
+		taa_half SampleSpatialWeight = 1.0;
+
+		taa_half3 InputSample = InputTexture[InputPixelPos];
+		taa_half3 HistorySample = PredictionSceneColorTexture[InputPixelPos];
+		taa_half HistorySampleValidity = ParallaxRejectionMaskToHistoryValidity(
+			ParallaxRejectionMaskTexture[InputPixelPos]);
+		
+		taa_half InputWeight = ComputeColorWeight(InputSample);
+		taa_half HistoryWeight = HistorySampleValidity * ComputeColorWeight(HistorySample);
+
+		CenterInput = InputSample;
+
+		FilteredInput       = (SampleSpatialWeight * InputWeight) * InputSample;
+		FilteredInputWeight = (SampleSpatialWeight * InputWeight);
+		
+		FilteredHistory       = (SampleSpatialWeight * HistoryWeight) * HistorySample;
+		FilteredHistoryWeight = (SampleSpatialWeight * HistoryWeight);
+	}
+	
+	uint StartSampleId;
+	uint SampleCount;
+	uint SampleTrackId;
+	
+	#if CONFIG_KERNEL == KERNEL_1X1
+	{
+		StartSampleId = 0;
+		SampleCount = 0;
+	}
+	#elif CONFIG_KERNEL == KERNEL_3X3 || CONFIG_KERNEL == KERNEL_3X3_PLUS
+	{
+		StartSampleId = 0;
+		SampleCount = 9;
+	}
+	#else
+		#error Unknown kernel
+	#endif
+
+	#if CONFIG_USE_DYNAMIC_SPREAD_FACTOR
+		taa_short2 DynamicSpreadFactor = (
+			taa_short(CONFIG_KERNEL_SPREAD_FACTOR) +
+			taa_short2((InputPixelPos % taa_short(2)) ^ (InputPixelPos.yx % taa_short(2))));
+	#else
+		taa_short2 DynamicSpreadFactor = taa_short(CONFIG_KERNEL_SPREAD_FACTOR);
+	#endif
+
+	// Fetch neighborhood of the kernel
+	UNROLL
+	for (uint SampleId = StartSampleId; SampleId < SampleCount; SampleId++)
+	{
+		taa_short2 Offset;
+		taa_half SampleSpatialWeight;
+		
+		#if CONFIG_KERNEL == KERNEL_1X1
+		{
+			Offset = 0;
+			SampleSpatialWeight = taa_half(1.0);
+			continue;
+		}
+		#elif CONFIG_KERNEL == KERNEL_3X3 || CONFIG_KERNEL == KERNEL_3X3_PLUS
+		{
+			if (SampleId == 4)
+			{
+				continue;
+			}
+			Offset = taa_short2(kOffsets3x3[SampleId]);
+			
+			#if CONFIG_KERNEL == KERNEL_3X3_PLUS
+				if (abs(Offset.x) + abs(Offset.y) > 1)
+				{
+					continue;
+				}
+			#endif
+			
+			#if 0
+				SampleSpatialWeight = taa_half(1.0);
+			#elif 1
+				SampleSpatialWeight = taa_half(rcp(abs(Offset.x) + 1.0) * rcp(abs(Offset.y) + 1.0));
+			#else
+				SampleSpatialWeight = taa_half(rcp(abs(Offset.x) + abs(Offset.y)));
+			#endif
+		}
+		#else
+			#error Unknown kernel
+		#endif
+		
+		taa_short2 SampleInputPixelPos = InputPixelPos + ClampPixelOffset(
+			InputPixelPos,
+			DynamicSpreadFactor * Offset, Offset,
+			InputPixelPosMin, InputPixelPosMax);
+
+		taa_half3 InputSample = InputTexture[SampleInputPixelPos];
+		taa_half3 HistorySample = PredictionSceneColorTexture[SampleInputPixelPos];
+		taa_half HistorySampleValidity = ParallaxRejectionMaskToHistoryValidity(
+			ParallaxRejectionMaskTexture[SampleInputPixelPos]);
+		
+		taa_half SampleInterferenceWeight = InterferenceWeightTexture[SampleInputPixelPos];
+		
+		FLATTEN
+		if (!bEnableInterferenceHeuristic)
+		{
+			SampleInterferenceWeight = 1.0;
+		}
+
+		taa_half InputWeight = (HistorySampleValidity * SampleInterferenceWeight) * ComputeColorWeight(InputSample);
+		taa_half HistoryWeight = (HistorySampleValidity * SampleInterferenceWeight) * ComputeColorWeight(HistorySample);
+
+		FilteredInput       += (SampleSpatialWeight * InputWeight) * InputSample;
+		FilteredInputWeight += (SampleSpatialWeight * InputWeight);
+		
+		FilteredHistory       += (SampleSpatialWeight * HistoryWeight) * HistorySample;
+		FilteredHistoryWeight += (SampleSpatialWeight * HistoryWeight);
+	}
+
+	{
+		FilteredInput *= SafeRcp(FilteredInputWeight);
+		FilteredHistory *= SafeRcp(FilteredHistoryWeight);
+	}
+	
+	taa_short2 OutputPixelPos;
+	#if 1
+		ISOLATE
+		{
+			OutputPixelPos = (
+				taa_short2(InputPixelPosMin) +
+				taa_short2(GroupId) * taa_short2(TILE_SIZE, TILE_SIZE) +
+				Map16x16Tile2x2Lane(GetGroupThreadIndex(GroupThreadIndex, GroupWaveIndex)));
+		}
+	#else
+		OutputPixelPos = InputPixelPos;
+	#endif
+
+	#if CONFIG_ENABLE_STOCASTIC_QUANTIZATION
+	{
+		uint2 Random = Rand3DPCG16(int3(OutputPixelPos, View.StateFrameIndexMod8)).xy;
+		float2 E = Hammersley16(0, 1, Random);
+		
+		FilteredInput = QuantizeForFloatRenderTarget(FilteredInput, E.x, OutputQuantizationError);
+		FilteredHistory = QuantizeForFloatRenderTarget(FilteredHistory, E.x, OutputQuantizationError);
+	}
+	#endif
+
+	OutputPixelPos.x = all(OutputPixelPos < LowFrequencyInfo_ViewportMax) ? OutputPixelPos.x : ~taa_short(0);
+
+	FilteredInputOutput[OutputPixelPos] = FilteredInput;
+	FilteredPredictionSceneColorOutput[OutputPixelPos] = FilteredHistory;
+		
+	#if DEBUG_OUTPUT
+	{
+		DebugOutput[OutputPixelPos - LowFrequencyInfo_ViewportMin] = Debug;
+	}
+	#endif
+}
diff --git a/Engine/Shaders/Private/TemporalSuperResolution/TSRPostfilterRejection.usf b/Engine/Shaders/Private/TemporalSuperResolution/TSRPostfilterRejection.usf
new file mode 100644
--- /dev/null
+++ b/Engine/Shaders/Private/TemporalSuperResolution/TSRPostfilterRejection.usf
@@ -0,0 +1,139 @@
+// Copyright Epic Games, Inc. All Rights Reserved.
+
+#include "TSRCommon.ush"
+
+
+//------------------------------------------------------- CONFIG
+
+#define TILE_SIZE 8
+
+#define CONFIG_OUTPUT_HALF_RES DIM_OUTPUT_HALF_RES
+
+
+//------------------------------------------------------- PARAMETERS
+
+uint4 HistoryRejectionViewport;
+
+Texture2D<taa_half> HistoryRejectionTexture;
+RWTexture2D<taa_half> HistoryRejectionOutput;
+
+
+//------------------------------------------------------- LDS
+
+groupshared taa_half SharedRejection[TILE_SIZE * TILE_SIZE];
+
+
+//------------------------------------------------------- FUNCTIONS
+
+// Operator that output, lowest, median and highest values from 3 input values.
+void LMHOperator(taa_half A, taa_half B, taa_half C, out taa_half L, out taa_half M, out taa_half H)
+{
+	// TODO: min3, max3, med3.
+	taa_half X = min(B, C);
+	taa_half Y = max(B, C);
+	
+	L = min(A, X);
+	taa_half Z = max(A, X);
+
+	M = min(Z, Y);
+	H = max(Z, Y);
+}
+
+// 3 samples median.
+taa_half Median(taa_half A, taa_half B, taa_half C)
+{
+	// TODO: med3.
+	taa_half L, M, H;
+	LMHOperator(A, B, C, L, M, H);
+	return M;
+}
+
+// 9 samples median using Smith1996.
+taa_half Median9(taa_half In[9])
+{
+	// First layer.
+	taa_half L0[3];
+	taa_half M0[3];
+	taa_half H0[3];
+	for (uint j = 0; j < 3; j++)
+	{
+		LMHOperator(In[j * 3 + 0], In[j * 3 + 1], In[j * 3 + 2], L0[j], M0[j], H0[j]);
+	}
+	
+	// Second layer.
+	// TODO: min3, max3.
+	taa_half M1[3];
+	M1[0] = max(max(L0[0], L0[1]), L0[2]);
+	M1[1] = Median(M0[0], M0[1], M0[2]);
+	M1[2] = min(min(H0[0], H0[1]), H0[2]);
+	
+	// Third layer.
+	return Median(M1[0], M1[1], M1[2]);
+}
+
+
+//------------------------------------------------------- ENTRY POINT
+
+[numthreads(TILE_SIZE, TILE_SIZE, 1)]
+void MainCS(
+	uint2 GroupId : SV_GroupID,
+	uint GroupThreadIndex : SV_GroupIndex)
+{
+	float4 Debug = 0.0;
+
+	taa_short2 RejectionPixelPos = (
+		taa_short2(GroupId) * taa_short2(TILE_SIZE, TILE_SIZE) +
+		Map8x8Tile2x2Lane(GroupThreadIndex));
+
+	RejectionPixelPos = ClampPixelOffset(RejectionPixelPos, HistoryRejectionViewport.xy, HistoryRejectionViewport.zw - 1);
+
+	// Fetch neighborhood of the kernel
+	taa_half InSamples[9];
+	UNROLL
+	for (uint i = 0; i < 9; i++)
+	{
+		taa_short2 SampleInputPixelPos = AddAndClampPixelOffset(
+			RejectionPixelPos,
+			kOffsets3x3[i], kOffsets3x3[i],
+			HistoryRejectionViewport.xy, HistoryRejectionViewport.zw - 1);
+
+		InSamples[i] = HistoryRejectionTexture[SampleInputPixelPos];
+	}
+
+	taa_half Rejection = Median9(InSamples);
+
+	#if CONFIG_OUTPUT_HALF_RES
+	{
+		SharedRejection[GroupThreadIndex] = Rejection;
+
+		GroupMemoryBarrierWithGroupSync();
+
+		UNROLL
+		for (uint i = 1; i < 4; i++)
+		{
+			uint ButterflySwap = i;
+
+			Rejection = min(Rejection, SharedRejection[GroupThreadIndex ^ ButterflySwap]);
+		}
+		
+		bool bIsValid = all(RejectionPixelPos <= taa_short2(HistoryRejectionViewport.zw)) && all((taa_ushort2(RejectionPixelPos) % taa_ushort(2)) == taa_ushort(0));
+
+		taa_ushort2 OutputPixelPos = taa_ushort2(RejectionPixelPos) / taa_ushort(2);
+		OutputPixelPos.x = bIsValid ? OutputPixelPos.x : ~taa_ushort(0);
+
+		HistoryRejectionOutput[OutputPixelPos] = Rejection;
+	}
+	#else
+	{
+		taa_short2 OutputPixelPos = InvalidateOutputPixelPos(RejectionPixelPos, HistoryRejectionViewport.zw);
+		HistoryRejectionOutput[OutputPixelPos] = Rejection;
+	}
+	#endif
+	
+	#if DEBUG_OUTPUT
+	{
+		taa_short2 OutputPixelPos = InvalidateOutputPixelPos(RejectionPixelPos, HistoryRejectionViewport.zw);
+		DebugOutput[OutputPixelPos] = Debug;
+	}
+	#endif
+}
