Separates TSR source code from TAA

diff --git a/Engine/Shaders/Private/TemporalSuperResolution/TSRUpdateHistory.usf b/Engine/Shaders/Private/TemporalSuperResolution/TSRUpdateHistory.usf
new file mode 100644
--- /dev/null
+++ b/Engine/Shaders/Private/TemporalSuperResolution/TSRUpdateHistory.usf
@@ -0,0 +1,989 @@
+// Copyright Epic Games, Inc. All Rights Reserved.
+
+#include "TSRCommon.ush"
+
+
+
+//------------------------------------------------------- CONFIG
+
+#define DEBUG_FALLBACK_BLENDING 0
+
+#define DEBUG_LOW_FREQUENCY_REJECTION 0
+
+#define CONFIG_CLAMP 1
+
+#define CONFIG_RESPONSIVE_STENCIL 1
+
+#define CONFIG_SAMPLES 5
+
+#define CONFIG_MANUAL_LDS_SPILL (CONFIG_COMPILE_FP16)
+
+#define CONFIG_LOW_FREQUENCY_DRIFT_REJECTION 1
+
+#define CONFIG_TRANSLUCENCY_REJECTION 1
+
+#define CONFIG_VFX_SURVIVAL 1
+
+
+//------------------------------------------------------- CONSTANTS
+
+#define TILE_SIZE 8
+
+
+//------------------------------------------------------- PARAMETERS
+
+Texture2D<taa_half3> InputSceneColorTexture;
+Texture2D<uint2> InputSceneStencilTexture;
+Texture2D<taa_half> HistoryRejectionTexture;
+Texture2D<taa_half> TranslucencyRejectionTexture;
+
+Texture2D<float2> DilatedVelocityTexture;
+Texture2D<uint> ParallaxFactorTexture;
+Texture2D<taa_half> ParallaxRejectionMaskTexture;
+
+float4 HistoryPixelPosToScreenPos;
+float4 HistoryPixelPosToPPCo;
+float3 HistoryQuantizationError;
+float MinTranslucencyRejection;
+uint ResponsiveStencilMask;
+
+Texture2D<taa_half3> PrevHistory_Textures_0;
+Texture2D<taa_half3> PrevHistory_Textures_1;
+Texture2D<taa_half2> PrevHistory_Textures_2;
+Texture2D<taa_subpixel_details> PrevHistory_Textures_3;
+
+RWTexture2D<taa_half3> HistoryOutput_Textures_0;
+RWTexture2D<taa_half3> HistoryOutput_Textures_1;
+RWTexture2D<taa_half2> HistoryOutput_Textures_2;
+RWTexture2D<taa_subpixel_details> HistoryOutput_Textures_3;
+
+RWTexture2D<taa_half3> SceneColorOutput;
+
+
+//------------------------------------------------------- LDS
+
+#if CONFIG_ENABLE_DUAL_PIXEL_VECTORIZATION
+#define GROUP_SHARED_SIZE (TILE_SIZE * TILE_SIZE / 2)
+#else
+#define GROUP_SHARED_SIZE (TILE_SIZE * TILE_SIZE)
+#endif
+
+groupshared taa_half4 SharedArray0[GROUP_SHARED_SIZE];
+groupshared taa_half4 SharedArray1[GROUP_SHARED_SIZE];
+groupshared taa_half4 SharedArray2[GROUP_SHARED_SIZE];
+groupshared taa_half4 SharedArray3[GROUP_SHARED_SIZE];
+groupshared taa_half4 SharedArray4[GROUP_SHARED_SIZE];
+groupshared taa_half4 SharedArray5[GROUP_SHARED_SIZE];
+groupshared taa_subpixel_payload2 SharedParallaxFactorBits[GROUP_SHARED_SIZE];
+
+
+//------------------------------------------------------- FUNCTIONS
+
+taa_ushort2x2 Map8x8Tile2x2LaneDPV(uint GroupThreadIndex)
+{
+	taa_ushort2 GroupId = Map8x8Tile2x2Lane(GroupThreadIndex);
+#if CONFIG_ENABLE_DUAL_PIXEL_VECTORIZATION
+	return dpv_interleave_registers(GroupId, GroupId + taa_ushort2(0,  8 / 2));
+#else
+	return dpv_interleave_mono_registers(GroupId);
+#endif
+}
+
+taa_half2 ComputeFallbackContribution(taa_half2 HistoryValidity)
+{
+	return saturate((HistoryValidity - taa_half(1.0 * rcp(MAX_SAMPLE_COUNT))) * taa_half(6.0));
+}
+
+taa_half3x2 TransformColorForClampingBox(taa_half3x2 Color)
+{
+	return RGBToYCoCg(Color);
+}
+
+void ComputeInputKernelSamplePosition(
+	taa_short2x2 InputPixelPos,
+	taa_half2x2 dKO,
+	uint SampleId,
+	out taa_short2x2 SampleInputPixelPos,
+	out taa_half2x2 PixelOffset)
+{
+	
+	#if CONFIG_SAMPLES == 9
+	{
+		taa_short2 iPixelOffset = taa_short2(kOffsets3x3[kSquareIndexes3x3[SampleId]]);
+		PixelOffset = dpv_interleave_registers(taa_half2(iPixelOffset), taa_half2(iPixelOffset));
+				
+		SampleInputPixelPos = InputPixelPos + ClampPixelOffset(
+			InputPixelPos,
+			dpv_interleave_registers(iPixelOffset, iPixelOffset), iPixelOffset,
+			InputPixelPosMin, InputPixelPosMax);
+	}
+	#elif CONFIG_SAMPLES == 5 || CONFIG_SAMPLES == 6
+	{
+		if (SampleId == 5)
+		{
+			taa_short2x2 iPixelOffset;
+			#if CONFIG_COMPILE_FP16
+				iPixelOffset[0] = int16_t(1) - int16_t2((asuint16(dKO[0]) & uint16_t(0x8000)) >> uint16_t(14));
+				iPixelOffset[1] = int16_t(1) - int16_t2((asuint16(dKO[1]) & uint16_t(0x8000)) >> uint16_t(14));
+				PixelOffset[0] = asfloat16(asuint16(half(1.0)).xx | (asuint16(dKO[0]) & uint16_t(0x8000)));
+				PixelOffset[1] = asfloat16(asuint16(half(1.0)).xx | (asuint16(dKO[1]) & uint16_t(0x8000)));
+			#else
+				iPixelOffset = dpv_interleave_registers(
+					SignFastInt(dpv_lo(dKO)), SignFastInt(dpv_hi(dKO)));
+				PixelOffset[0] = asfloat(asuint(1.0).xx | (asuint(dKO[0]) & uint(0x80000000)));
+				PixelOffset[1] = asfloat(asuint(1.0).xx | (asuint(dKO[1]) & uint(0x80000000)));
+			#endif
+						
+			SampleInputPixelPos = ClampPixelOffset(InputPixelPos + iPixelOffset, InputPixelPosMin, InputPixelPosMax);
+		}
+		else
+		{
+			taa_short2 iPixelOffset = taa_short2(kOffsets3x3[kPlusIndexes3x3[SampleId]]);
+			PixelOffset = dpv_interleave_registers(taa_half2(iPixelOffset), taa_half2(iPixelOffset));
+					
+			SampleInputPixelPos = InputPixelPos + ClampPixelOffset(
+				InputPixelPos,
+				dpv_interleave_registers(iPixelOffset, iPixelOffset), iPixelOffset,
+				InputPixelPosMin, InputPixelPosMax);
+		}
+	}
+	#else
+		#error Unknown sample count
+	#endif
+}
+
+//------------------------------------------------------- ENTRY POINT
+
+#if CONFIG_ENABLE_DUAL_PIXEL_VECTORIZATION
+	[numthreads(TILE_SIZE, TILE_SIZE / 2, 1)]
+#else
+	[numthreads(TILE_SIZE, TILE_SIZE, 1)]
+#endif
+void MainCS(
+	uint2 GroupId : SV_GroupID,
+	uint GroupThreadIndex : SV_GroupIndex)
+{
+	#if CONFIG_ENABLE_DUAL_PIXEL_VECTORIZATION
+		uint GroupWaveIndex = GetGroupWaveIndex(GroupThreadIndex, TILE_SIZE * TILE_SIZE / 2);
+	#else
+		uint GroupWaveIndex = GetGroupWaveIndex(GroupThreadIndex, TILE_SIZE * TILE_SIZE);
+	#endif
+
+	float4x2 Debug = 0.0;
+
+	taa_short2x2 HistoryPixelPos = dpv_add(
+		taa_short2(GroupId) * taa_short2(TILE_SIZE, TILE_SIZE),
+		taa_short2x2(Map8x8Tile2x2LaneDPV(GroupThreadIndex)));
+
+	float2x2 ScreenPos = ApplyScreenTransform(float2x2(HistoryPixelPos), HistoryPixelPosToScreenPos);
+	
+	// Pixel coordinate of the center of output pixel O in the input viewport.
+	float2x2 PPCo = ApplyScreenTransform(float2x2(HistoryPixelPos), HistoryPixelPosToPPCo);
+
+	// Pixel coordinate of the center of the nearest input pixel K.
+	float2x2 PPCk = floor(PPCo) + 0.5;
+	
+	taa_short2x2 InputPixelPos = ClampPixelOffset(
+		taa_short2x2(PPCo),
+		InputPixelPosMin, InputPixelPosMax);
+
+	// Fetch reprojections related informations.
+	float2x2 PrevScreenPos = ScreenPos;
+	taa_half2 ParallaxRejectionMask = taa_half(1.0).xx;
+	taa_half2 LowFrequencyRejection = taa_half(1.0).xx;
+	taa_half2 TranslucencyRejection = taa_half(1.0).xx;
+	taa_half2 OutputPixelVelocity = taa_half(0.0).xx;
+	taa_subpixel_payload2 ParallaxFactorBits = taa_subpixel_payload(0x0).xx;
+	bool2 bIsResponsiveAAPixel = false;
+
+	ISOLATE
+	{
+		// Vector in pixel between pixel K -> O.
+		taa_half2x2 dKO = taa_half2x2(PPCo - PPCk);
+
+		float2x2 EncodedVelocity = dpv_interleave_registers(
+			DilatedVelocityTexture[dpv_lo(InputPixelPos)],
+			DilatedVelocityTexture[dpv_hi(InputPixelPos)]);
+		ParallaxRejectionMask = dpv_interleave_registers(
+			ParallaxRejectionMaskTexture[dpv_lo(InputPixelPos)],
+			ParallaxRejectionMaskTexture[dpv_hi(InputPixelPos)]);
+		ParallaxFactorBits = dpv_interleave_registers(
+			ParallaxFactorTexture[dpv_lo(InputPixelPos)],
+			ParallaxFactorTexture[dpv_hi(InputPixelPos)]) & taa_subpixel_payload(SUB_PIXEL_PARALLAX_FACTOR_BIT_MASK);
+
+		float2x2 ScreenVelocity = dpv_interleave_registers(
+			DecodeVelocityFromTexture(float4(dpv_lo(EncodedVelocity), 0.0, 0.0)).xy,
+			DecodeVelocityFromTexture(float4(dpv_hi(EncodedVelocity), 0.0, 0.0)).xy);
+
+		PrevScreenPos = ScreenPos - ScreenVelocity;
+		OutputPixelVelocity = taa_half2(dpv_length(dpv_mul(ScreenVelocity, HistoryInfo_ViewportSize)));
+
+		taa_ushort2x2 RejectionPixelPos = dpv_sub(taa_ushort2x2(InputPixelPos), taa_ushort2(InputPixelPosMin)) / 2;
+		LowFrequencyRejection = dpv_interleave_registers(
+			HistoryRejectionTexture[dpv_lo(RejectionPixelPos)],
+			HistoryRejectionTexture[dpv_hi(RejectionPixelPos)]);
+		TranslucencyRejection = dpv_interleave_registers(
+			TranslucencyRejectionTexture[dpv_lo(InputPixelPos)],
+			TranslucencyRejectionTexture[dpv_hi(InputPixelPos)]);
+		
+		uint2 SceneStencilRef = dpv_interleave_registers(
+			InputSceneStencilTexture.Load(int3(dpv_lo(InputPixelPos), 0)) STENCIL_COMPONENT_SWIZZLE, 
+			InputSceneStencilTexture.Load(int3(dpv_hi(InputPixelPos), 0)) STENCIL_COMPONENT_SWIZZLE);
+
+		UNROLL_N(CONFIG_SAMPLES - 1)
+		for (uint SampleId = 1; SampleId < CONFIG_SAMPLES; SampleId++)
+		{
+			taa_short2x2 SampleInputPixelPos;
+			taa_half2x2 PixelOffset;
+			ComputeInputKernelSamplePosition(
+				InputPixelPos, dKO, SampleId,
+				/* out */ SampleInputPixelPos,
+				/* out */ PixelOffset);
+			
+			TranslucencyRejection = min(TranslucencyRejection, dpv_interleave_registers(
+				TranslucencyRejectionTexture[dpv_lo(SampleInputPixelPos)],
+				TranslucencyRejectionTexture[dpv_hi(SampleInputPixelPos)]));
+			
+			SceneStencilRef |= dpv_interleave_registers(
+				InputSceneStencilTexture.Load(int3(dpv_lo(SampleInputPixelPos), 0)) STENCIL_COMPONENT_SWIZZLE, 
+				InputSceneStencilTexture.Load(int3(dpv_hi(SampleInputPixelPos), 0)) STENCIL_COMPONENT_SWIZZLE);
+		}
+
+		TranslucencyRejection = max(TranslucencyRejection, taa_half(MinTranslucencyRejection));
+		
+		// Fetch wehther the pixel is responsive AA or not.
+		#if CONFIG_RESPONSIVE_STENCIL
+		{
+			bIsResponsiveAAPixel = (SceneStencilRef & ResponsiveStencilMask) != 0;
+		}
+		#endif
+
+		#if !CONFIG_CLAMP
+		{
+			ParallaxRejectionMask = taa_half(1.0).xx;
+			LowFrequencyRejection = taa_half(1.0).xx;
+			TranslucencyRejection = taa_half(1.0).xx;
+		}
+		#endif
+	}
+
+	// Reject the translucency to respond quicker on firefliers.
+	#if CONFIG_TRANSLUCENCY_REJECTION
+		LowFrequencyRejection = min(LowFrequencyRejection, TranslucencyRejection);
+	#endif
+	
+	// Detect if HistoryBufferUV would be outside of the viewport.
+	bool2 bOffScreen = bool2(
+		IsOffScreen(bCameraCut, dpv_lo(PrevScreenPos), dpv_lo(ParallaxRejectionMask)),
+		IsOffScreen(bCameraCut, dpv_hi(PrevScreenPos), dpv_hi(ParallaxRejectionMask)));
+	
+	taa_half2 TotalRejection = taa_half2(
+		bOffScreen[0] ? 0.0 : saturate(LowFrequencyRejection[0] * 4.0),
+		bOffScreen[1] ? 0.0 : saturate(LowFrequencyRejection[1] * 4.0));
+
+
+	// Filter input scene color at predictor frequency.
+	taa_half3x2 FilteredInputColor;
+	taa_half3x2 InputMinColor;
+	taa_half3x2 InputMaxColor;
+	taa_half2 InputPixelAlignement;
+	taa_half2 ClosestInputLuma4;
+	
+	
+	taa_half2x2 dKO = taa_half2x2(PPCo - PPCk);
+		
+	bool2 bUpdateSubpixelDetails;
+	bool2x2 bKOSign;
+
+	ISOLATE
+	{
+		// Vector in pixel between pixel K -> O.
+		taa_half2x2 dKO = taa_half2x2(PPCo - PPCk);
+
+		{
+			float SizeOfOutputPixelInInputPixelWidth = 0.5 * (InputInfo_ViewportSize.x * HistoryInfo_ViewportSizeInverse.x);
+
+			bUpdateSubpixelDetails[0] = all(abs(dpv_lo(dKO)) < SizeOfOutputPixelInInputPixelWidth);
+			bUpdateSubpixelDetails[1] = all(abs(dpv_hi(dKO)) < SizeOfOutputPixelInInputPixelWidth);
+		
+			bKOSign[0] = dKO[0] < 0.0;
+			bKOSign[1] = dKO[1] < 0.0;
+		}
+
+
+		FilteredInputColor = taa_half(0.0);
+
+		taa_half2 FilteredInputColorWeight = taa_half(0.0);
+		
+		// TODO: UECON-464 - this causes shader warnings on XSX DXC.  Revisit once we have aged out the problem GDK versions.
+		taa_half InputToHistoryFactor = taa_half(HistoryInfo_ViewportSize.x * InputInfo_ViewportSizeInverse.x);
+		taa_half2 FinalInputToHistoryFactor = taa_half2(
+			bOffScreen.x ? taa_half(1.0) : InputToHistoryFactor,
+			bOffScreen.y ? taa_half(1.0) : InputToHistoryFactor);
+
+		UNROLL_N(CONFIG_SAMPLES)
+		for (uint SampleId = 0; SampleId < CONFIG_SAMPLES; SampleId++)
+		{
+			taa_short2x2 SampleInputPixelPos;
+			taa_half2x2 PixelOffset;
+			ComputeInputKernelSamplePosition(
+				InputPixelPos, dKO, SampleId,
+				/* out */ SampleInputPixelPos,
+				/* out */ PixelOffset);
+
+			taa_half3x2 InputColor = dpv_interleave_registers(
+				InputSceneColorTexture[dpv_lo(SampleInputPixelPos)],
+				InputSceneColorTexture[dpv_hi(SampleInputPixelPos)]);
+
+			taa_half2x2 dPP = PixelOffset - dKO;
+			taa_half2 SampleSpatialWeight = ComputeSampleWeigth(FinalInputToHistoryFactor, dPP, /* MinimalContribution = */ float(0.005));
+
+			taa_half2 ToneWeight = HdrWeight4(InputColor);
+
+			#if CONFIG_VFX_SURVIVAL
+			if (SampleId == 0)
+			{
+				ToneWeight = max(ToneWeight, ComputeSampleWeigth(1.0, dPP, /* MinimalContribution = */ 0.0));
+			}
+			#endif
+
+			FilteredInputColor       += dpv_scale((SampleSpatialWeight * ToneWeight), InputColor);
+			FilteredInputColorWeight += (SampleSpatialWeight * ToneWeight);
+
+			if (SampleId == 0)
+			{
+				ClosestInputLuma4 = Luma4(InputColor);
+				InputMinColor = TransformColorForClampingBox(InputColor);
+				InputMaxColor = TransformColorForClampingBox(InputColor);
+			}
+			else
+			{
+				InputMinColor = min(InputMinColor, TransformColorForClampingBox(InputColor));
+				InputMaxColor = max(InputMaxColor, TransformColorForClampingBox(InputColor));
+			}
+		}
+		
+		FilteredInputColor = dpv_scale(FilteredInputColor, rcp(FilteredInputColorWeight));
+
+		InputPixelAlignement = ComputeSampleWeigth(InputToHistoryFactor, dKO, /* MinimalContribution = */ float(0.0));
+	}
+		
+	// Spills to LDS to make room in VGPR for history sampling.
+	#if CONFIG_MANUAL_LDS_SPILL
+	ISOLATE
+	{
+		uint LocalGroupThreadIndex = GetGroupThreadIndex(GroupThreadIndex, GroupWaveIndex);
+
+		#if CONFIG_ENABLE_DUAL_PIXEL_VECTORIZATION
+			SharedArray0[LocalGroupThreadIndex] = taa_half4(FilteredInputColor[0], FilteredInputColor[1]);
+			SharedArray1[LocalGroupThreadIndex] = taa_half4(FilteredInputColor[2], LowFrequencyRejection);
+			SharedArray2[LocalGroupThreadIndex] = taa_half4(InputMinColor[0], InputMinColor[1]);
+			SharedArray3[LocalGroupThreadIndex] = taa_half4(InputMinColor[2], InputPixelAlignement);
+			SharedArray4[LocalGroupThreadIndex] = taa_half4(InputMaxColor[0], InputMaxColor[1]);
+			SharedArray5[LocalGroupThreadIndex] = taa_half4(InputMaxColor[2], OutputPixelVelocity);
+			SharedParallaxFactorBits[LocalGroupThreadIndex] = ParallaxFactorBits;
+		#else
+			SharedArray0[LocalGroupThreadIndex] = taa_half4(dpv_lo(FilteredInputColor), dpv_lo(LowFrequencyRejection));
+			SharedArray1[LocalGroupThreadIndex] = taa_half4(dpv_lo(InputMinColor), dpv_lo(InputPixelAlignement));
+			SharedArray2[LocalGroupThreadIndex] = taa_half4(dpv_lo(InputMaxColor), dpv_lo(OutputPixelVelocity));
+		#endif
+	}
+	#endif
+	
+	// Reproject history
+	taa_half3x2 PrevHistoryMoment1;
+	taa_half2 PrevHistoryValidity;
+	
+	taa_half3x2 PrevHistoryMommentMin;
+	taa_half3x2 PrevHistoryMommentMax;
+
+	taa_half3x2 PrevFallbackColor;
+	taa_half2 PrevFallbackWeight;
+	
+	taa_subpixel_details2 PrevSubpixelDetails;
+
+	ISOLATE
+	{
+		// Reproject history
+		taa_half3x2 RawHistory0 = taa_half(0);
+		taa_half3x2 RawHistory1 = taa_half(0);
+		taa_half2x2 RawHistory2 = taa_half(0);
+
+		taa_half3x2 RawHistory1Min = INFINITE_FLOAT;
+		taa_half3x2 RawHistory1Max = -INFINITE_FLOAT;
+
+		// Sample raw history
+		{
+			float2x2 PrevHistoryBufferUV = ApplyScreenTransform(PrevScreenPos, ScreenPosToPrevHistoryBufferUV);
+			PrevHistoryBufferUV[0] = clamp(PrevHistoryBufferUV[0], PrevHistoryInfo_UVViewportBilinearMin[0], PrevHistoryInfo_UVViewportBilinearMax[0]);
+			PrevHistoryBufferUV[1] = clamp(PrevHistoryBufferUV[1], PrevHistoryInfo_UVViewportBilinearMin[1], PrevHistoryInfo_UVViewportBilinearMax[1]);
+
+			#if 1
+			{
+				FCatmullRomSamples Samples0 = GetBicubic2DCatmullRomSamples(
+					dpv_lo(PrevHistoryBufferUV), PrevHistoryInfo_Extent, PrevHistoryInfo_ExtentInverse);
+				FCatmullRomSamples Samples1 = GetBicubic2DCatmullRomSamples(
+					dpv_hi(PrevHistoryBufferUV), PrevHistoryInfo_Extent, PrevHistoryInfo_ExtentInverse);
+
+				UNROLL
+				for (uint i = 0; i < Samples0.Count; i++)
+				{
+					float2 SampleUV0 = clamp(Samples0.UV[i], PrevHistoryInfo_UVViewportBilinearMin, PrevHistoryInfo_UVViewportBilinearMax);
+					float2 SampleUV1 = clamp(Samples1.UV[i], PrevHistoryInfo_UVViewportBilinearMin, PrevHistoryInfo_UVViewportBilinearMax);
+
+					taa_half3x2 Sample0 = dpv_interleave_registers(
+						PrevHistory_Textures_0.SampleLevel(GlobalBilinearClampedSampler, SampleUV0, 0),
+						PrevHistory_Textures_0.SampleLevel(GlobalBilinearClampedSampler, SampleUV1, 0));
+					taa_half3x2 Sample1 = dpv_interleave_registers(
+						PrevHistory_Textures_1.SampleLevel(GlobalBilinearClampedSampler, SampleUV0, 0),
+						PrevHistory_Textures_1.SampleLevel(GlobalBilinearClampedSampler, SampleUV1, 0));
+					taa_half2x2 Sample2 = dpv_interleave_registers(
+						PrevHistory_Textures_2.SampleLevel(GlobalBilinearClampedSampler, SampleUV0, 0),
+						PrevHistory_Textures_2.SampleLevel(GlobalBilinearClampedSampler, SampleUV1, 0));
+
+					taa_half2 InvWeight = SafeRcp(Sample2[1]);
+
+					RawHistory1Min = min(RawHistory1Min, dpv_scale(Sample1, InvWeight));
+					RawHistory1Max = max(RawHistory1Max, dpv_scale(Sample1, InvWeight));
+
+					taa_half2 KernelWeight = dpv_interleave_registers(taa_half(Samples0.Weight[i]), taa_half(Samples1.Weight[i]));
+
+					RawHistory0 += dpv_scale(Sample0, KernelWeight);
+					RawHistory1 += dpv_scale(Sample1, KernelWeight);
+					RawHistory2 += dpv_scale(Sample2, KernelWeight);
+				}
+				
+				taa_half2 FinalMultiplier = dpv_interleave_registers(taa_half(Samples0.FinalMultiplier), taa_half(Samples1.FinalMultiplier));
+
+
+				RawHistory0 = dpv_scale(RawHistory0, FinalMultiplier);
+				RawHistory1 = dpv_scale(RawHistory1, FinalMultiplier);
+				RawHistory2 = dpv_scale(RawHistory2, FinalMultiplier);
+			}
+			#else
+			{
+				RawHistory0 = PrevHistory_Textures_0.SampleLevel(GlobalBilinearClampedSampler, PrevHistoryBufferUV, 0);
+				RawHistory1 = PrevHistory_Textures_1.SampleLevel(GlobalBilinearClampedSampler, PrevHistoryBufferUV, 0);
+				RawHistory2 = PrevHistory_Textures_2.SampleLevel(GlobalBilinearClampedSampler, PrevHistoryBufferUV, 0);
+			}
+			#endif
+			
+			FSubpixelNeighborhood SubpixelNeighborhood0 = GatherPrevSubpixelNeighborhood(PrevHistory_Textures_3, dpv_lo(PrevHistoryBufferUV));
+			FSubpixelNeighborhood SubpixelNeighborhood1 = GatherPrevSubpixelNeighborhood(PrevHistory_Textures_3, dpv_hi(PrevHistoryBufferUV));
+			{
+				PrevSubpixelDetails = 0;
+				UNROLL_N(SUB_PIXEL_COUNT)
+				for (uint SubpixelId = 0; SubpixelId < SUB_PIXEL_COUNT; SubpixelId++)
+				{
+					taa_subpixel_payload SubpixelPayload0 = GetSubpixelPayload(SubpixelNeighborhood0, SubpixelId);
+					taa_subpixel_payload SubpixelPayload1 = GetSubpixelPayload(SubpixelNeighborhood1, SubpixelId);
+					PrevSubpixelDetails[0] |= SubpixelPayload0 << taa_subpixel_payload(SUB_PIXEL_BIT_COUNT * SubpixelId);
+					PrevSubpixelDetails[1] |= SubpixelPayload1 << taa_subpixel_payload(SUB_PIXEL_BIT_COUNT * SubpixelId);
+				}
+			}
+
+			RawHistory0 = -dpv_min(-RawHistory0, taa_half(0.0).xxx);
+			RawHistory1 = -dpv_min(-RawHistory1, taa_half(0.0).xxx);
+			RawHistory2 = -dpv_min(-RawHistory2, taa_half(0.0).xx);
+		}
+		
+		// Unpack history.
+		{
+			PrevFallbackColor = RawHistory0;
+			PrevFallbackWeight = RawHistory2[0];
+			
+			PrevHistoryMommentMin = RawHistory1Min;
+			PrevHistoryMommentMax = RawHistory1Max;
+
+			PrevHistoryMoment1 = RawHistory1;
+			PrevHistoryValidity = RawHistory2[1];
+		}
+
+		// Correct history
+		{
+			PrevHistoryMommentMin *= taa_half(HistoryPreExposureCorrection);
+			PrevHistoryMommentMax *= taa_half(HistoryPreExposureCorrection);
+			PrevHistoryMoment1 *= taa_half(HistoryPreExposureCorrection);
+			PrevFallbackColor *= taa_half(HistoryPreExposureCorrection);
+		}
+	}
+		
+	#if CONFIG_MANUAL_LDS_SPILL
+	ISOLATE
+	{
+		uint LocalGroupThreadIndex = GetGroupThreadIndex(GroupThreadIndex, GroupWaveIndex);
+	
+		#if CONFIG_ENABLE_DUAL_PIXEL_VECTORIZATION
+			taa_half4 RawLDS0 = SharedArray0[LocalGroupThreadIndex];
+			taa_half4 RawLDS1 = SharedArray1[LocalGroupThreadIndex];
+			taa_half4 RawLDS2 = SharedArray2[LocalGroupThreadIndex];
+			taa_half4 RawLDS3 = SharedArray3[LocalGroupThreadIndex];
+			taa_half4 RawLDS4 = SharedArray4[LocalGroupThreadIndex];
+			taa_half4 RawLDS5 = SharedArray5[LocalGroupThreadIndex];
+
+			ParallaxFactorBits = SharedParallaxFactorBits[LocalGroupThreadIndex];
+
+			FilteredInputColor = taa_half3x2(RawLDS0.rg, RawLDS0.ba, RawLDS1.rg);
+			InputMinColor = taa_half3x2(RawLDS2.rg, RawLDS2.ba, RawLDS3.rg);
+			InputMaxColor = taa_half3x2(RawLDS4.rg, RawLDS4.ba, RawLDS5.rg);
+		
+			LowFrequencyRejection = RawLDS1.ba;
+			InputPixelAlignement = RawLDS3.ba;
+			OutputPixelVelocity = RawLDS5.ba;
+		#else
+			taa_half4 RawLDS0 = SharedArray0[LocalGroupThreadIndex];
+			taa_half4 RawLDS1 = SharedArray1[LocalGroupThreadIndex];
+			taa_half4 RawLDS2 = SharedArray2[LocalGroupThreadIndex];
+
+			FilteredInputColor = dpv_interleave_mono_registers(RawLDS0.rgb);
+			InputMinColor = dpv_interleave_mono_registers(RawLDS1.rgb);
+			InputMaxColor = dpv_interleave_mono_registers(RawLDS2.rgb);
+		
+			LowFrequencyRejection = dpv_interleave_mono_registers(RawLDS0.a);
+			InputPixelAlignement = dpv_interleave_mono_registers(RawLDS1.a);
+			OutputPixelVelocity = dpv_interleave_mono_registers(RawLDS2.a);
+		#endif
+	}
+	#endif
+
+	// Rejects high frequency details if their low frequencies drift from the low frequency history
+	#if CONFIG_LOW_FREQUENCY_DRIFT_REJECTION
+	{
+		taa_half3x2 NormalizedPrevHistoryMoment1 = dpv_scale(PrevHistoryMoment1, SafeRcp(PrevHistoryValidity));
+
+		taa_half3x2 PrevHighFrequencyYCoCg = TransformColorForClampingBox(NormalizedPrevHistoryMoment1);
+		taa_half3x2 PrevYCoCg = TransformColorForClampingBox(PrevFallbackColor);
+		taa_half3x2 ClampedPrevYCoCg = TransformColorForClampingBox(clamp(PrevFallbackColor, PrevHistoryMommentMin, PrevHistoryMommentMax));
+
+		taa_half2 HighFrequencyRejection = MeasureRejectionFactor(
+			PrevYCoCg, ClampedPrevYCoCg,
+			PrevHighFrequencyYCoCg, InputMinColor, InputMaxColor,
+			MeasureBackbufferLDRQuantizationErrorFromHDRLuma(max(PrevYCoCg[0], ClampedPrevYCoCg[0])));
+		
+		//PrevHistoryMoment1 = dpv_scale(PrevHistoryMoment1, HighFrequencyRejection);
+		PrevHistoryMoment1 = dpv_scale(clamp(NormalizedPrevHistoryMoment1, PrevHistoryMommentMin, PrevHistoryMommentMax), HighFrequencyRejection * PrevHistoryValidity);
+		PrevHistoryValidity *= HighFrequencyRejection;
+	}
+	#endif
+
+	// Contribute current frame input into the predictor for next frame.
+	const taa_half Histeresis = rcp(taa_half(MAX_SAMPLE_COUNT));
+	const taa_half PredictionOnlyValidity = Histeresis * taa_half(2.0);
+	
+	// Clamp fallback
+	taa_half2 LumaMin;
+	taa_half2 LumaMax;
+	taa_half3x2 ClampedFallbackColor;
+	taa_half2 FallbackRejection;
+	{
+		LumaMin = InputMinColor[0];
+		LumaMax = InputMaxColor[0];
+
+		taa_half3x2 PrevYCoCg = TransformColorForClampingBox(PrevFallbackColor);
+		taa_half3x2 ClampedPrevYCoCg = dpv_clamp(PrevYCoCg, InputMinColor, InputMaxColor);
+		taa_half3x2 InputCenterYCoCg = TransformColorForClampingBox(FilteredInputColor);
+
+		ClampedFallbackColor = YCoCgToRGB(ClampedPrevYCoCg);
+		
+		FallbackRejection = MeasureRejectionFactor(
+			PrevYCoCg, ClampedPrevYCoCg,
+			InputCenterYCoCg, InputMinColor, InputMaxColor,
+			MeasureBackbufferLDRQuantizationErrorFromHDRLuma(max(PrevYCoCg[0], ClampedPrevYCoCg[0])));
+
+		#if !CONFIG_CLAMP
+		{
+			ClampedFallbackColor = PrevFallbackColor;
+			FallbackRejection = taa_half(1.0);
+		}
+		#endif
+	}
+
+	// Also reject sub pixel detail low frequency if detect mid frequency changes.
+	#if 0
+	{
+		LowFrequencyRejection = min(LowFrequencyRejection, FallbackRejection);
+	}
+	#endif
+	
+	taa_half3x2 FinalHistoryMoment1;
+	taa_half2 FinalHistoryValidity;
+	{
+		// Compute how much the history needs to be rejected, based on the completeness.
+		taa_half2 PrevHistoryRejectionWeight = LowFrequencyRejection;
+			
+		//#if !CONFIG_CLAMP
+		//	PrevHistoryRejectionWeight = 1.0;
+		//#endif
+			
+		FLATTEN
+		if (bOffScreen[0])
+		{
+			PrevHistoryRejectionWeight[0] = taa_half(0.0);
+		}
+
+		FLATTEN
+		if (bOffScreen[1])
+		{
+			PrevHistoryRejectionWeight[1] = taa_half(0.0);
+		}
+
+		taa_half2 DesiredCurrentContribution = max(Histeresis * InputPixelAlignement, taa_half(0.0));
+
+		// Determine whether the the prediction based rejection was confident enough.
+		taa_half2 RejectionConfidentEnough = taa_half(1); // saturate(RejectionValidity * MAX_SAMPLE_COUNT - 3.0);
+
+		// Compute the newly rejected validity
+		taa_half2 RejectedValidity = (
+			min(PrevHistoryValidity, PredictionOnlyValidity - DesiredCurrentContribution) +
+			max(PrevHistoryValidity - PredictionOnlyValidity + DesiredCurrentContribution, taa_half(0.0)) * PrevHistoryRejectionWeight);
+
+		RejectedValidity = PrevHistoryValidity * PrevHistoryRejectionWeight;
+
+		// Compute the maximum output validity.
+		taa_half2 OutputValidity = (
+			clamp(RejectedValidity + DesiredCurrentContribution, taa_half(0.0), PredictionOnlyValidity) +
+			clamp(RejectedValidity + DesiredCurrentContribution * PrevHistoryRejectionWeight * RejectionConfidentEnough - PredictionOnlyValidity, 0.0, 1.0 - PredictionOnlyValidity));
+
+		FLATTEN
+		if (bIsResponsiveAAPixel[0])
+		{
+			OutputValidity[0] = taa_half(0.0);
+		}
+
+		FLATTEN
+		if (bIsResponsiveAAPixel[1])
+		{
+			OutputValidity[1] = taa_half(0.0);
+		}
+		
+		taa_half2 InvPrevHistoryValidity = SafeRcp(PrevHistoryValidity);
+
+		taa_half2 PrevMomentWeight = max(OutputValidity - DesiredCurrentContribution, taa_half(0.0));
+		taa_half2 CurrentMomentWeight = min(DesiredCurrentContribution, OutputValidity);
+		
+		#if 1
+		{
+			taa_half2 PrevHistoryToneWeight = HdrWeightY(Luma4(PrevHistoryMoment1) * InvPrevHistoryValidity);
+			taa_half2 FilteredInputToneWeight = HdrWeight4(FilteredInputColor);
+			
+			taa_half2 BlendPrevHistory = PrevMomentWeight * PrevHistoryToneWeight;
+			taa_half2 BlendFilteredInput = CurrentMomentWeight * FilteredInputToneWeight;
+
+			taa_half2 CommonWeight = OutputValidity * SafeRcp(BlendPrevHistory + BlendFilteredInput);
+
+			FinalHistoryMoment1 = (
+				dpv_scale(PrevHistoryMoment1, CommonWeight * BlendPrevHistory * InvPrevHistoryValidity) +
+				dpv_scale(FilteredInputColor, CommonWeight * BlendFilteredInput));
+		}
+		#else
+		{
+			FinalHistoryMoment1 = PrevHistoryMoment1 * (PrevMomentWeight * InvPrevHistoryValidity) + FilteredInputColor * (CurrentMomentWeight);
+		}
+		#endif
+
+		// Quantize validity for the 8bit encoding an adjust to avoid numerical shift.
+		#if 0
+			FinalHistoryValidity = OutputValidity;
+		#else
+			taa_half2 OutputInvValidity = SafeRcp(OutputValidity);
+			FinalHistoryValidity = ceil(taa_half(255.0) * OutputValidity) * rcp(taa_half(255.0));
+			FinalHistoryMoment1 = dpv_scale(FinalHistoryMoment1, FinalHistoryValidity * OutputInvValidity);
+		#endif
+
+		//Debug.x = RejectionConfidentEnough;
+		//Debug.x = RejectionValidity;
+
+		#if 0 // DEBUG
+		{
+			FinalHistoryMoment1 = PrevHistoryMoment1;
+			FinalHistoryValidity = PrevHistoryValidity;
+		}
+		#endif
+	}
+
+	// Compute the fallback history
+	taa_half3x2 FinalFallbackColor;
+	taa_half2 FinalFallbackWeight;
+	{
+		const taa_half TargetHesteresisCurrentFrameWeight = rcp(taa_half(MAX_FALLBACK_SAMPLE_COUNT));
+
+		taa_half2 LumaHistory = Luma4(PrevFallbackColor);
+		taa_half2 LumaFiltered = Luma4(FilteredInputColor);
+
+		{
+			taa_half2 OutputBlend = ComputeFallbackContribution(FinalHistoryValidity);
+
+			//ClampedFallbackColor = lerp(ClampedFallbackColor, PrevFallbackColor, OutputBlend);
+		}
+
+		taa_half2 BlendFinal;
+		#if 1
+		{
+			taa_half2 CurrentFrameSampleCount = max(InputPixelAlignement, taa_half(0.005));
+			
+			// Use only one sample count to be extremely fast to recover in history rejection. But then stabilize imediately so sub pixel
+			// frequency can be used as soon as possible.
+			taa_half2 PrevFallbackSampleCount;
+			FLATTEN
+			if (PrevFallbackWeight[0] < taa_half(1.0))
+			{
+				PrevFallbackSampleCount[0] = PrevFallbackWeight[0];
+			}
+			else
+			{
+				PrevFallbackSampleCount[0] = taa_half(MAX_FALLBACK_SAMPLE_COUNT);
+			}
+
+			FLATTEN
+			if (PrevFallbackWeight[1] < taa_half(1.0))
+			{
+				PrevFallbackSampleCount[1] = PrevFallbackWeight[1];
+			}
+			else
+			{
+				PrevFallbackSampleCount[1] = taa_half(MAX_FALLBACK_SAMPLE_COUNT);
+			}
+
+			// Reject the history based on the low frequencies.
+			#if 1
+			{
+				// float InvertedBlendFinal = PrevFallbackSampleCount / (CurrentFrameSampleCount + PrevFallbackSampleCount);
+				// float PrevFallbackRejectionFactor = saturate(LowFrequencyRejection / InvertedBlendFinal);
+				
+				taa_half2 PrevFallbackRejectionFactor = saturate(LowFrequencyRejection * (CurrentFrameSampleCount + PrevFallbackSampleCount) / PrevFallbackSampleCount);
+
+				PrevFallbackSampleCount *= PrevFallbackRejectionFactor;
+			}
+			#endif
+
+			BlendFinal = CurrentFrameSampleCount / (CurrentFrameSampleCount + PrevFallbackSampleCount);
+
+			// Increase the blending on motion
+			#if 1
+			{
+				BlendFinal = lerp(BlendFinal, max(taa_half(0.2), BlendFinal), saturate(OutputPixelVelocity * rcp(taa_half(40.0))));
+			}
+			#endif
+
+			// Antiflicker
+			#if 1
+			{
+				taa_half2 DistToClamp = min( abs(LumaHistory - LumaMin), abs(LumaHistory - LumaMax) ) / max3( LumaHistory, LumaFiltered, taa_half(1e-4) );
+				BlendFinal *= taa_half(0.2) + taa_half(0.8) * saturate(taa_half(0.5) * DistToClamp);
+			}
+			#endif
+			
+			// Make sure to have at least some small contribution
+			#if 1
+			{
+				BlendFinal = max( BlendFinal, saturate( taa_half(0.01) * LumaHistory / abs( LumaFiltered - LumaHistory ) ) );
+			}
+			#endif
+
+			// Responsive forces 1/4 of new frame.
+			BlendFinal[0] = bIsResponsiveAAPixel[0] ? taa_half(1.0/4.0) : BlendFinal[0];
+			BlendFinal[1] = bIsResponsiveAAPixel[1] ? taa_half(1.0/4.0) : BlendFinal[1];
+
+			// Completely discard the history
+			{
+				PrevFallbackSampleCount *= TotalRejection;
+				BlendFinal = lerp(1.0, BlendFinal, TotalRejection);
+			}
+
+			FinalFallbackWeight = saturate(CurrentFrameSampleCount + PrevFallbackSampleCount);
+			
+			#if 1
+				FinalFallbackWeight = saturate(floor(255.0 * (CurrentFrameSampleCount + PrevFallbackSampleCount)) * rcp(255.0));
+			#endif
+		}
+		#endif
+
+		{
+			taa_half2 FilterWeight = HdrWeight4(FilteredInputColor);
+			taa_half2 ClampedHistoryWeight = HdrWeight4(ClampedFallbackColor);
+
+			taa_half2x2 Weights = WeightedLerpFactors(ClampedHistoryWeight, FilterWeight, BlendFinal);
+
+			FinalFallbackColor = dpv_scale(ClampedFallbackColor, Weights[0]) + dpv_scale(FilteredInputColor, Weights[1]);
+		}
+	}
+
+	// Update subpixel detail
+	taa_subpixel_details2 FinalSubpixelDetails;
+	{
+		FinalSubpixelDetails = PrevSubpixelDetails;
+
+
+		{
+			const taa_subpixel_details ParallaxFactorMask = taa_subpixel_details(
+				(SUB_PIXEL_PARALLAX_FACTOR_BIT_MASK << (SUB_PIXEL_PARALLAX_FACTOR_BIT_OFFSET + 0 * SUB_PIXEL_BIT_COUNT)) | 
+				(SUB_PIXEL_PARALLAX_FACTOR_BIT_MASK << (SUB_PIXEL_PARALLAX_FACTOR_BIT_OFFSET + 1 * SUB_PIXEL_BIT_COUNT)) | 
+				(SUB_PIXEL_PARALLAX_FACTOR_BIT_MASK << (SUB_PIXEL_PARALLAX_FACTOR_BIT_OFFSET + 2 * SUB_PIXEL_BIT_COUNT)) | 
+				(SUB_PIXEL_PARALLAX_FACTOR_BIT_MASK << (SUB_PIXEL_PARALLAX_FACTOR_BIT_OFFSET + 3 * SUB_PIXEL_BIT_COUNT)) | 
+				0x0);
+			
+			// Reset parallax factor
+			FLATTEN
+			if (bOffScreen.x)
+			{
+				FinalSubpixelDetails.x = FinalSubpixelDetails.x & ~ParallaxFactorMask;
+			}
+			
+			FLATTEN
+			if (bOffScreen.y)
+			{
+				FinalSubpixelDetails.y = FinalSubpixelDetails.y & ~ParallaxFactorMask;
+			}
+		}
+		
+		uint2 SubpixelId = dpv_dot(uint2x2(bKOSign), uint2(1, SUB_PIXEL_GRID_SIZE));
+
+		taa_subpixel_payload2 SubpixelShift = taa_subpixel_payload2(SubpixelId) * taa_subpixel_payload(SUB_PIXEL_BIT_COUNT);
+
+		taa_subpixel_payload2 SubpixelPayload = (ParallaxFactorBits << taa_subpixel_payload(SUB_PIXEL_PARALLAX_FACTOR_BIT_OFFSET));
+
+		taa_subpixel_details2 UpdatedFinalSubpixelDetails = (FinalSubpixelDetails & (~(taa_subpixel_payload(SUB_PIXEL_BIT_MASK) << SubpixelShift))) | (SubpixelPayload << SubpixelShift);
+
+		FLATTEN
+		if (bUpdateSubpixelDetails[0])
+		{
+			FinalSubpixelDetails.x = UpdatedFinalSubpixelDetails.x;
+		}
+		
+		FLATTEN
+		if (bUpdateSubpixelDetails[1])
+		{
+			FinalSubpixelDetails.y = UpdatedFinalSubpixelDetails.y;
+		}
+	}
+
+	// Compute final output
+	taa_half3x2 FinalOutputColor;
+	taa_half2 FinalOutputValidity;
+	{
+		taa_half2 OutputBlend = ComputeFallbackContribution(FinalHistoryValidity);
+
+		#if 0 // DEBUG
+		{
+			if (View.GeneralPurposeTweak == 2.0)
+				OutputBlend = 0.0;
+			else if (View.GeneralPurposeTweak == 3.0)
+				OutputBlend = 1.0;
+		}
+		#elif 0 // DEBUG
+		{
+			OutputBlend = 0.0;
+		}
+		#elif 0 // DEBUG
+		{
+			OutputBlend = 1.0;
+		}
+		#endif
+
+		FinalOutputValidity = lerp(taa_half(1.0), saturate(FinalHistoryValidity), OutputBlend);
+
+		taa_half3x2 NormalizedFinalHistoryMoment1 = dpv_scale(FinalHistoryMoment1, SafeRcp(FinalHistoryValidity));
+
+		taa_half2 FallbackWeight = HdrWeight4(FinalFallbackColor);
+		taa_half2 Moment1Weight = HdrWeight4(NormalizedFinalHistoryMoment1);
+
+		taa_half2x2 Weights = WeightedLerpFactors(FallbackWeight, Moment1Weight, OutputBlend);
+
+		#if DEBUG_FALLBACK_BLENDING
+			taa_half3 FallbackColor = taa_half3(1, 0.25, 0.25);
+			taa_half3 HighFrequencyColor = taa_half3(0.25, 1, 0.25);
+
+			FinalOutputColor = dpv_scale(FinalFallbackColor * dpv_interleave_mono_registers(FallbackColor), Weights[0]) + dpv_scale(NormalizedFinalHistoryMoment1 * dpv_interleave_mono_registers(HighFrequencyColor), Weights[1]);
+		#elif DEBUG_LOW_FREQUENCY_REJECTION
+			taa_half3 DebugColor = lerp(taa_half3(1, 0.5, 0.5), taa_half3(0.5, 1, 0.5), LowFrequencyRejection);
+			
+			FinalOutputColor = FinalFallbackColor * Weights.x * DebugColor + NormalizedFinalHistoryMoment1 * Weights.y * DebugColor;
+		#else
+			FinalOutputColor = dpv_scale(FinalFallbackColor, Weights[0]) + dpv_scale(NormalizedFinalHistoryMoment1, Weights[1]);
+		#endif
+	}
+
+
+	ISOLATE
+	{
+		uint LocalGroupThreadIndex = GetGroupThreadIndex(GroupThreadIndex, GroupWaveIndex);
+
+		#if 1
+			taa_short2x2 LocalHistoryPixelPos = dpv_add(
+				taa_short2(GroupId) * taa_short2(TILE_SIZE, TILE_SIZE),
+				taa_short2x2(Map8x8Tile2x2LaneDPV(LocalGroupThreadIndex)));
+		#else
+			taa_short2x2 LocalHistoryPixelPos = HistoryPixelPos;
+		#endif
+			
+		LocalHistoryPixelPos = InvalidateOutputPixelPos(LocalHistoryPixelPos, HistoryInfo_ViewportMax);
+
+		// Output final history.
+		{
+			// Stocastically round up or down using the hardware RWTexture2D truncation unit to take into precision
+			// loss due to pixel format encoding.
+			#if CONFIG_ENABLE_STOCASTIC_QUANTIZATION
+			{
+				uint2 Random = Rand3DPCG16(int3(dpv_lo(LocalHistoryPixelPos), View.StateFrameIndexMod8)).xy;
+				float E = Hammersley16(0, 1, Random).x;
+
+				FinalHistoryMoment1 = QuantizeForFloatRenderTarget(FinalHistoryMoment1, E, HistoryQuantizationError);
+				FinalFallbackColor = QuantizeForFloatRenderTarget(FinalFallbackColor, E, HistoryQuantizationError);
+			}
+			#endif
+
+			#if 0 // Manual software mantissa truncation for debugging purposes.
+			{
+				const uint MantissaF32Bits = 23;
+				const uint MantissaF16Bits = 10;
+				const uint3 MantissaR11G11B10Bits = uint3(6, 6, 5);
+
+				//FinalFallbackColor.x = asfloat(asuint(FinalFallbackColor.x) & ((~0) << (MantissaF32Bits - MantissaR11G11B10Bits.r)));
+				//FinalFallbackColor.y = asfloat(asuint(FinalFallbackColor.y) & ((~0) << (MantissaF32Bits - MantissaR11G11B10Bits.g)));
+				//FinalFallbackColor.z = asfloat(asuint(FinalFallbackColor.z) & ((~0) << (MantissaF32Bits - MantissaR11G11B10Bits.b)));
+			}
+			#endif
+			
+			FinalFallbackColor = -dpv_min(-FinalFallbackColor, taa_half(0.0).xxx);
+			FinalHistoryMoment1 = -dpv_min(-FinalHistoryMoment1, taa_half(0.0).xxx);
+			FinalFallbackColor = dpv_min(FinalFallbackColor, taa_half(Max10BitsFloat).xxx);
+			FinalHistoryMoment1 = dpv_min(FinalHistoryMoment1, taa_half(Max10BitsFloat).xxx);
+			
+			HistoryOutput_Textures_0[dpv_lo(LocalHistoryPixelPos)] = dpv_lo(FinalFallbackColor);
+			HistoryOutput_Textures_1[dpv_lo(LocalHistoryPixelPos)] = dpv_lo(FinalHistoryMoment1);
+			HistoryOutput_Textures_2[dpv_lo(LocalHistoryPixelPos)] = taa_half2(dpv_lo(FinalFallbackWeight), dpv_lo(FinalHistoryValidity));
+			HistoryOutput_Textures_3[dpv_lo(LocalHistoryPixelPos)] = dpv_lo(FinalSubpixelDetails);
+			
+			#if CONFIG_ENABLE_DUAL_PIXEL_VECTORIZATION
+				HistoryOutput_Textures_0[dpv_hi(LocalHistoryPixelPos)] = dpv_hi(FinalFallbackColor);
+				HistoryOutput_Textures_1[dpv_hi(LocalHistoryPixelPos)] = dpv_hi(FinalHistoryMoment1);
+				HistoryOutput_Textures_2[dpv_hi(LocalHistoryPixelPos)] = taa_half2(dpv_hi(FinalFallbackWeight), dpv_hi(FinalHistoryValidity));
+				HistoryOutput_Textures_3[dpv_hi(LocalHistoryPixelPos)] = dpv_hi(FinalSubpixelDetails);
+			#endif
+		}
+
+		// Output final scene color
+		{
+			taa_half3x2 OutputColor = FinalOutputColor;
+
+			OutputColor = -dpv_min(-OutputColor, taa_half(0.0).xxx);
+			OutputColor = dpv_min(OutputColor, taa_half(Max10BitsFloat).xxx);
+
+			SceneColorOutput[dpv_lo(LocalHistoryPixelPos)] = dpv_lo(OutputColor);
+			
+			#if CONFIG_ENABLE_DUAL_PIXEL_VECTORIZATION
+				SceneColorOutput[dpv_hi(LocalHistoryPixelPos)] = dpv_hi(OutputColor);
+			#endif
+		}
+
+		#if DEBUG_OUTPUT
+		{
+			DebugOutput[dpv_lo(LocalHistoryPixelPos)] = dpv_lo(Debug);
+
+			#if CONFIG_ENABLE_DUAL_PIXEL_VECTORIZATION
+				DebugOutput[dpv_hi(LocalHistoryPixelPos)] = dpv_hi(Debug);
+			#endif
+		}
+		#endif
+	}
+}
diff --git a/Engine/Shaders/Private/TemporalSuperResolution/TSRUpdateSuperResHistory.usf b/Engine/Shaders/Private/TemporalSuperResolution/TSRUpdateSuperResHistory.usf
new file mode 100644
--- /dev/null
+++ b/Engine/Shaders/Private/TemporalSuperResolution/TSRUpdateSuperResHistory.usf
@@ -0,0 +1,205 @@
+// Copyright Epic Games, Inc. All Rights Reserved.
+
+#include "TSRCommon.ush"
+#include "../Random.ush"
+
+
+//------------------------------------------------------- CONSTANTS
+
+#define TILE_SIZE 8
+
+
+//------------------------------------------------------- PARAMETERS
+
+Texture2D<taa_half3> InputSceneColorTexture;
+
+Texture2D<float2> DilatedVelocityTexture;
+Texture2D<taa_half> ParallaxRejectionMaskTexture;
+
+Texture2D<uint> PrevHistory_SuperResTextures_0;
+Texture2D<uint> PrevHistory_SuperResTextures_1;
+
+RWTexture2D<uint> HistoryOutput_SuperResTextures_0;
+RWTexture2D<uint> HistoryOutput_SuperResTextures_1;
+
+
+//------------------------------------------------------- ENTRY POINT
+
+[numthreads(TILE_SIZE, TILE_SIZE, 1)]
+void MainCS(
+	uint2 GroupId : SV_GroupID,
+	uint GroupThreadIndex : SV_GroupIndex)
+{
+	uint GroupWaveIndex = GetGroupWaveIndex(GroupThreadIndex, /* GroupSize = */ TILE_SIZE * TILE_SIZE);
+
+	float4 Debug = 0.0;
+
+	uint2 DispatchThreadId = (
+		ZOrder2D(GroupThreadIndex, log2(TILE_SIZE)) +
+		GroupId * uint2(TILE_SIZE, TILE_SIZE));
+
+	uint2 HistorySuperPixelPos = DispatchThreadId;
+	float2 ViewportUV = (float2(HistorySuperPixelPos) + 0.5f) * HistoryInfo_ViewportSizeInverse * rcp(taa_half(SUPERRES_RES_MULTIPLIER));
+	float2 ScreenPos = ViewportUVToScreenPos(ViewportUV);
+	
+	// Pixel coordinate of the center of output pixel O in the input viewport.
+	float2 PPCo = ViewportUV * InputInfo_ViewportSize + InputJitter;
+
+	// Pixel coordinate of the center of the nearest input pixel K.
+	float2 PPCk = floor(PPCo) + 0.5;
+	
+	// Vector in pixel between pixel K -> O in input pixel size.
+	taa_half2 idKO = taa_half2(PPCo - PPCk);
+	
+	// Vector in pixel between pixel K -> O in output pixel size.
+	taa_half2 odKO = idKO * (InputInfo_ViewportSizeInverse.x * HistoryInfo_ViewportSize.x * SUPERRES_RES_MULTIPLIER);
+
+	uint2 InputPixelPos = InputInfo_ViewportMin + uint2(PPCk);
+
+	// Fetch reprojections related informations.
+	float2 PrevScreenPos = ScreenPos;
+	taa_half ParallaxRejectionMask = taa_half(1.0);
+	#if 1
+	{
+		float2 RejectionBilinearBufferUV = clamp(
+			PPCo * InputInfo_ExtentInverse,
+			RejectionInfo_UVViewportBilinearMin,
+			RejectionInfo_UVViewportBilinearMax);
+		
+		float2 EncodedVelocity = DilatedVelocityTexture[InputPixelPos];
+		ParallaxRejectionMask = ParallaxRejectionMaskTexture[InputPixelPos];
+		PrevScreenPos = ScreenPos - DecodeVelocityFromTexture(float4(EncodedVelocity, 0.0, 0.0)).xy;
+	}
+	#endif
+
+	float2 ScreenVelocity = ScreenPos - PrevScreenPos;
+	
+	// Detect if HistoryBufferUV would be outside of the viewport.
+	bool bOffScreen = IsOffScreen(bCameraCut, PrevScreenPos, ParallaxRejectionMask);
+	
+	// Read input color
+	taa_half ClosestInterferenceSeed;
+	ISOLATE
+	{
+		float2 SampleInputBufferUV = (InputInfo_ViewportMin + PPCk) * InputInfo_ExtentInverse;
+		SampleInputBufferUV = clamp(SampleInputBufferUV, InputInfo_UVViewportBilinearMin, InputInfo_UVViewportBilinearMax);
+
+		taa_half3 InputColor = InputSceneColorTexture.SampleLevel(GlobalPointClampedSampler, SampleInputBufferUV, 0);
+
+		taa_half ClosestInputLuma4 = Luma4(InputColor);
+		ClosestInterferenceSeed = LumaToInterferenceSeed(ClosestInputLuma4);
+	}
+	
+	// Reproject history
+	uint OldestAge = 0;
+	FSuperResHistory SubpixelHistory;
+	ISOLATE
+	{
+		float2 PrevHistoryBufferUV = ApplyScreenTransform(PrevScreenPos, ScreenPosToPrevHistoryBufferUV);
+		PrevHistoryBufferUV = clamp(PrevHistoryBufferUV, PrevHistoryInfo_UVViewportBilinearMin, PrevHistoryInfo_UVViewportBilinearMax);
+
+		FBilinearSampleInfos SampleInfos = GetBilinearSampleLevelInfos(
+			PrevHistoryBufferUV,
+			PrevHistoryInfo_Extent * float(SUPERRES_RES_MULTIPLIER),
+			PrevHistoryInfo_ExtentInverse * rcp(float(SUPERRES_RES_MULTIPLIER)));
+
+		uint2 BestPayload = 0;
+		int2 BestKernelToSample = 0;
+		uint BestKernelToSampleDistSquare = dot(uint2(SUPERRES_OFFSET_MAX, SUPERRES_OFFSET_MAX), uint2(SUPERRES_OFFSET_MAX, SUPERRES_OFFSET_MAX));
+
+
+		UNROLL_N(4)
+		for (uint SampleId = 0; SampleId < 4; SampleId++)
+		{
+			uint2 PixelCoord = GetSamplePixelCoord(SampleInfos, SampleId);
+
+			uint2 SubpixelPayload;
+			SubpixelPayload[0] = PrevHistory_SuperResTextures_0[PixelCoord];
+			SubpixelPayload[1] = PrevHistory_SuperResTextures_1[PixelCoord];
+
+			FSuperResHistory UnpackedSubpixel = UnpackSuperResHistory(SubpixelPayload);
+
+			const float2 PixelOffset = float2(BilinearSamplingOffsets2x2[SampleId]);
+
+			int2 KernelToTexel = int2(round((PixelOffset - SampleInfos.BilinearInterp) * 2 * SUPERRES_OFFSET_MAX));
+
+			int2 KernelToSample = KernelToTexel + UnpackedSubpixel.TexelToSample;
+
+			uint KernelToSampleDistSquare = uint(dot(KernelToSample, KernelToSample));
+
+			FLATTEN
+			if (KernelToSampleDistSquare < BestKernelToSampleDistSquare || SampleId == 0)
+			{
+				BestPayload = SubpixelPayload;
+				BestKernelToSample = KernelToSample;
+				BestKernelToSampleDistSquare = KernelToSampleDistSquare;
+			}
+
+			bool bIsCloseEnough = abs(UnpackedSubpixel.InterferenceSeed - ClosestInterferenceSeed) < 0.1;
+
+			if (bIsCloseEnough)
+			{
+				OldestAge = max(OldestAge, UnpackedSubpixel.Age);
+			}
+		}
+
+		SubpixelHistory = UnpackSuperResHistory(BestPayload);
+		SubpixelHistory.Age = min(SubpixelHistory.Age + 1, SUPERRES_AGE_BITMASK);
+		SubpixelHistory.TexelToSample = clamp(BestKernelToSample, SUPERRES_OFFSET_MIN, SUPERRES_OFFSET_MAX);
+	}
+		
+	int2 CurrentTexelToSample = int2(round(odKO * (-2 * SUPERRES_OFFSET_MAX)));
+	
+	bool bUpdate = all(CurrentTexelToSample >= SUPERRES_OFFSET_MIN && CurrentTexelToSample <= SUPERRES_OFFSET_MAX);
+
+	FLATTEN
+	if (bOffScreen || bUpdate)
+	{
+		uint2 Random = Rand3DPCG16(int3(InputPixelPos, View.StateFrameIndexMod8)).xy;
+		float2 E = Hammersley16(0, 1, Random);
+
+		SubpixelHistory.TexelToSample = clamp(CurrentTexelToSample, SUPERRES_OFFSET_MIN, SUPERRES_OFFSET_MAX);
+		SubpixelHistory.Age = bOffScreen ? 0 : OldestAge;
+		SubpixelHistory.InterferenceSeed = saturate(ClosestInterferenceSeed + (E.x - 0.5) * rcp(taa_half(SUPERRES_INTERFERENCE_BITMASK)));
+	}
+
+	#if 1
+	{
+		Debug.r = SubpixelHistory.InterferenceSeed;
+		Debug.g = SubpixelHistory.Age * rcp(SUPERRES_AGE_BITMASK);
+		//Debug.b = SubpixelHistory.TexelToSample.x * rcp(16.0);
+		//Debug.a = SubpixelHistory.TexelToSample.y * rcp(16.0);
+
+		Debug.b = bOffScreen ? 1 : 0;
+		//Debug.b = ParallaxRejectionMask;
+	}
+	#endif
+
+	uint2 PackedSubpixelHistory = PackSuperResHistory(SubpixelHistory);
+
+	ISOLATE
+	{
+		#if 1
+			uint LocalGroupThreadIndex = GetGroupThreadIndex(GroupThreadIndex, GroupWaveIndex);
+			uint2 LocalHistoryPixelPos = (
+				ZOrder2D(LocalGroupThreadIndex, log2(TILE_SIZE)) +
+				GroupId * uint2(TILE_SIZE, TILE_SIZE));
+		#else
+			uint LocalGroupThreadIndex = HistorySuperPixelPos;
+		#endif
+			
+		// Output final history.
+		bool bIsValidhistoryPixel = all(LocalHistoryPixelPos < SUPERRES_RES_MULTIPLIER * HistoryInfo_ViewportMax);
+		if (bIsValidhistoryPixel)
+		{
+			HistoryOutput_SuperResTextures_0[LocalHistoryPixelPos] = PackedSubpixelHistory[0];
+			HistoryOutput_SuperResTextures_1[LocalHistoryPixelPos] = PackedSubpixelHistory[1];
+
+			#if DEBUG_OUTPUT
+			{
+				DebugOutput[LocalHistoryPixelPos] = Debug;
+			}
+			#endif
+		}
+	}
+}
