Perform proper gamma correction when rendering unwrapped cubemaps and other simple elements.

diff --git a/Engine/Shaders/Private/SimpleElementPixelShader.usf b/Engine/Shaders/Private/SimpleElementPixelShader.usf
--- a/Engine/Shaders/Private/SimpleElementPixelShader.usf
+++ b/Engine/Shaders/Private/SimpleElementPixelShader.usf
@@ -6,6 +6,7 @@
 
 #include "Common.ush"
 #include "ColorUtils.ush"
+#include "GammaCorrectionCommon.ush"
 
 #define USE_IES_PROFILE 1
 #include "IESLightProfilesCommon.ush"
@@ -136,7 +137,7 @@ void GammaMain(
 	if( Gamma != 1.0 )
 	{
 		// Gamma correct the output color.
-		OutColor.rgb = pow(saturate(OutColor.rgb),Gamma);	
+		OutColor.rgb = ApplyGammaCorrection(saturate(OutColor.rgb), 2.2 * Gamma);
 	}
 	OutColor = RETURN_COLOR(OutColor);
 
@@ -166,7 +167,7 @@ void GammaAlphaOnlyMain(
 	if( Gamma != 1.0 )
 	{
 		// Gamma correct the output color.
-		TmpColor.rgb = pow(saturate(TmpColor.rgb),Gamma);	
+		TmpColor.rgb = ApplyGammaCorrection(saturate(TmpColor.rgb), 2.2 * Gamma);
 	}
 	OutColor = RETURN_COLOR(TmpColor);
 
@@ -199,7 +200,7 @@ void GammaMaskedMain(
 	if( Gamma != 1.0 )
 	{
 		// Gamma correct the output color.
-		OutColor.rgb = pow(saturate(OutColor.rgb),Gamma);	
+		OutColor.rgb = ApplyGammaCorrection(saturate(OutColor.rgb), 2.2 * Gamma);
 	}
 	OutColor = RETURN_COLOR(OutColor);
 
@@ -298,7 +299,7 @@ void GammaDistanceFieldMain(
 	if( Gamma != 1.0 )
 	{
 		// Gamma correct the output color.
-		OutColor.rgb = pow(saturate(OutColor.rgb),Gamma);	
+		OutColor.rgb = ApplyGammaCorrection(saturate(OutColor.rgb), 2.2 * Gamma);
 	}
 
 	OutColor = SimpleElementEncodeFor32BPPHDR(OutColor);
@@ -344,8 +345,11 @@ void CubemapTextureProperties(
 #if HDR_OUTPUT == 1
 	FinalColor.rgb =max(float3(0,0,0),FinalColor.rgb);	
 #else
-	// Gamma correct the output color.
-	FinalColor.rgb = pow(saturate(FinalColor.rgb), Gamma);	
+	if( Gamma != 1.0 )
+	{
+		// Gamma correct the output color.
+		FinalColor.rgb = ApplyGammaCorrection(saturate(FinalColor.rgb), 2.2 * Gamma);
+	}
 #endif
 	
 	OutColor = RETURN_COLOR(FinalColor);	
