Remove sky avoidance logic from RayTracing shaders as this was only used by the path tracer and recently removed from the payload

diff --git a/Engine/Shaders/Private/RayTracing/RayTracingMaterialHitShaders.usf b/Engine/Shaders/Private/RayTracing/RayTracingMaterialHitShaders.usf
--- a/Engine/Shaders/Private/RayTracing/RayTracingMaterialHitShaders.usf
+++ b/Engine/Shaders/Private/RayTracing/RayTracingMaterialHitShaders.usf
@@ -318,14 +318,6 @@ RAY_TRACING_ENTRY_CLOSEST_HIT(MaterialCHS,
 		Payload.SetFrontFace();
 	}
 
-#if MATERIAL_IS_SKY
-	if (PackedPayload.IsIgnoreSkyMaterials())
-	{
-		PackedPayload = PackRayTracingPayload(Payload, PropagatedCone);
-		return;
-	}
-#endif
-
 	// Store the results in local variables and reuse instead of calling the functions multiple times.
 	half3 BaseColor = GetMaterialBaseColor(PixelMaterialInputs);
 #if MATERIALBLENDING_TRANSLUCENT && REFRACTION_USE_INDEX_OF_REFRACTION
