Revert application of bottom context dot(N, L) term to the refracted specular lobe. Apply Fresnel term with without refraction when ClearCoat < 1.

diff --git a/Engine/Shaders/Private/ShadingModels.ush b/Engine/Shaders/Private/ShadingModels.ush
--- a/Engine/Shaders/Private/ShadingModels.ush
+++ b/Engine/Shaders/Private/ShadingModels.ush
@@ -476,6 +476,7 @@ FDirectLighting ClearCoatBxDF( FGBufferData GBuffer, half3 N, half3 V, half3 L,
 			Vis2 = Vis_SmithJointApprox(a2, BottomContext.NoV, NoL);
 		}
 		float3 F = F_Schlick(GBuffer.SpecularColor, BottomContext.VoH);
+		float3 F_DefaultLit = F_Schlick(GBuffer.SpecularColor, Context.VoH);
 		
 		float Energy = 0;
 
@@ -496,10 +497,10 @@ FDirectLighting ClearCoatBxDF( FGBufferData GBuffer, half3 N, half3 V, half3 L,
 			Energy = EnergyNormalization(a2, Context.VoH, AreaLight);
 		}
 
-		// Note: reusing D, V, and F from refracted context to save computation for when ClearCoat < 1
-		float3 CommonSpecular = (Energy * Falloff * D2 * Vis2) * AreaLight.FalloffColor;
-		float3 DefaultSpecular = F * NoL;
-		float3 RefractedSpecular = FresnelCoeff * Transmission * F * BottomContext.NoL;
+		// Note: reusing D and V from refracted context to save computation when ClearCoat < 1
+		float3 CommonSpecular = (Energy * Falloff * NoL * D2 * Vis2) * AreaLight.FalloffColor;
+		float3 DefaultSpecular = F_DefaultLit;
+		float3 RefractedSpecular = FresnelCoeff * Transmission * F;
 		Lighting.Specular += CommonSpecular * lerp(DefaultSpecular, RefractedSpecular, ClearCoat);
 	}
 
