Multiple scattering now handles F90 for environment reflection (following the dielectric/conductor rules)

diff --git a/Engine/Shaders/Private/Strata/StrataEvaluation.ush b/Engine/Shaders/Private/Strata/StrataEvaluation.ush
--- a/Engine/Shaders/Private/Strata/StrataEvaluation.ush
+++ b/Engine/Shaders/Private/Strata/StrataEvaluation.ush
@@ -202,16 +202,16 @@ float StrataBSDFBackFaceProjectedSolidAngleFactor(in FStrataBSDFContext BSDFCont
 
 #if STRATA_MICROFACET_MULTIPLESCATTERING
 
-FStrataGGXLUTResult StrataEvaluateGGXLUT(float3 DiffuseAlbedo, float3 F0, float NoV, float Metallic, float SafeRoughnessX, float SafeRoughnessY)
+FStrataGGXLUTResult StrataEvaluateGGXLUT(float3 DiffuseAlbedo, float3 F0, float3 F90, float NoV, float Metallic, float SafeRoughnessX, float SafeRoughnessY)
 {
 #if STRATA_TRANSLUCENT_ENABLED
 	return StrataSampleGGXLUT(
 		TranslucentBasePass.Strata.GGXEnergyLUT2DTexture, TranslucentBasePass.Strata.GGXEnergyLUT3DTexture, TranslucentBasePass.Strata.GGXEnergyLUTSampler,
-		TranslucentBasePass.Strata.GGXEnergyLUTScaleBias, NoV, SafeRoughnessX, SafeRoughnessY, F0, Metallic, DiffuseAlbedo);
+		TranslucentBasePass.Strata.GGXEnergyLUTScaleBias, NoV, SafeRoughnessX, SafeRoughnessY, F0, F90, Metallic, DiffuseAlbedo);
 #else
 	return StrataSampleGGXLUT(
 		Strata.GGXEnergyLUT2DTexture, Strata.GGXEnergyLUT3DTexture, Strata.GGXEnergyLUTSampler,
-		Strata.GGXEnergyLUTScaleBias, NoV, SafeRoughnessX, SafeRoughnessY, F0, Metallic, DiffuseAlbedo);
+		Strata.GGXEnergyLUTScaleBias, NoV, SafeRoughnessX, SafeRoughnessY, F0, F90, Metallic, DiffuseAlbedo);
 #endif
 }
 
@@ -330,13 +330,13 @@ FStrataEvaluateResult StrataEvaluateBSDF(FStrataBSDFContext BSDFContext, FShadow
 			float3 MSScale = 1;
 			float3 HazeMSScale = 1;
 			#if STRATA_MICROFACET_MULTIPLESCATTERING
-				FStrataGGXLUTResult StrataGGXLUTResult = StrataEvaluateGGXLUT(DiffuseColor, F0, BSDFContext.Context.NoV, Metallic, SafeRoughnessX, SafeRoughnessY);
+				FStrataGGXLUTResult StrataGGXLUTResult = StrataEvaluateGGXLUT(DiffuseColor, F0, F90, BSDFContext.Context.NoV, Metallic, SafeRoughnessX, SafeRoughnessY);
 				MSScale	= StrataGGXLUTResult.SpecularMultipleScatteringScale;
 				const float3 DirectionalAlbedo_SpecularTransmission = StrataGGXLUTResult.DirectionalAlbedo.SpecularT; // We ignore Haziness influence
 				#if STRATA_FASTPATH==0
 				if (bHaziness)
 				{
-					HazeMSScale = StrataEvaluateGGXLUT(DiffuseColor, F0, BSDFContext.Context.NoV, Metallic, HazeSafeRoughnessX, HazeSafeRoughnessY).SpecularMultipleScatteringScale;
+					HazeMSScale = StrataEvaluateGGXLUT(DiffuseColor, F0, F90, BSDFContext.Context.NoV, Metallic, HazeSafeRoughnessX, HazeSafeRoughnessY).SpecularMultipleScatteringScale;
 				}
 				#endif
 			#else
@@ -704,7 +704,7 @@ FStrataEnvLightResult StrataEvaluateForEnvLight(FStrataBSDFContext BSDFContext,
 			#if STRATA_MICROFACET_MULTIPLESCATTERING
 				float3 EvalEnvBRDF = 0;
 				{
-					FStrataGGXLUTResult GGXLUT	= StrataEvaluateGGXLUT(DiffuseColor, F0, BSDFContext.Context.NoV, Metallic, SafeRoughnessX, SafeRoughnessX);
+					FStrataGGXLUTResult GGXLUT	= StrataEvaluateGGXLUT(DiffuseColor, F0, F90, BSDFContext.Context.NoV, Metallic, SafeRoughnessX, SafeRoughnessX);
 					EvalEnvBRDF					= ThinFilmFresnelWeight * GGXLUT.DirectionalAlbedo.SpecularR;
 					StrataEnvLightResult.DiffuseWeight = GGXLUT.DirectionalAlbedo.Diffuse;
 					DirectionalAlbedo_SpecularTransmission = GGXLUT.DirectionalAlbedo.SpecularT; // We ignore Haziness influence
@@ -734,7 +734,7 @@ FStrataEnvLightResult StrataEvaluateForEnvLight(FStrataBSDFContext BSDFContext,
 
 					// Compute the directional albedo for hazy specular
 				#if STRATA_MICROFACET_MULTIPLESCATTERING
-					FStrataGGXLUTResult GGXLUTHaze = StrataEvaluateGGXLUT(DiffuseColor, F0, BSDFContext.Context.NoV, Metallic, SafeRoughnessHaze, SafeRoughnessHaze);
+					FStrataGGXLUTResult GGXLUTHaze = StrataEvaluateGGXLUT(DiffuseColor, F0, F90, BSDFContext.Context.NoV, Metallic, SafeRoughnessHaze, SafeRoughnessHaze);
 					const float3 EvalEnvBRDFHaze = ThinFilmFresnelWeight * GGXLUTHaze.DirectionalAlbedo.SpecularR;
 				#else
 					const float3 EvalEnvBRDFHaze = ThinFilmFresnelWeight * StrataEnvBRDF(F0, F90, SafeRoughnessX, BSDFContext.SatNoV);
@@ -749,7 +749,7 @@ FStrataEnvLightResult StrataEvaluateForEnvLight(FStrataBSDFContext BSDFContext,
 			else
 			{
 				// Compute the directional albedo for diffuse
-				FStrataGGXLUTResult GGXLUT = StrataEvaluateGGXLUT(DiffuseColor, 1.f, BSDFContext.Context.NoV, 0.f, SafeRoughnessX, SafeRoughnessX);
+				FStrataGGXLUTResult GGXLUT = StrataEvaluateGGXLUT(DiffuseColor, 1.0f, 1.0f, BSDFContext.Context.NoV, 0.0f, SafeRoughnessX, SafeRoughnessX);
 				StrataEnvLightResult.DiffuseWeight = GGXLUT.DirectionalAlbedo.Diffuse;
 			}
 			#endif
diff --git a/Engine/Shaders/Private/Strata/StrataMultipleScattering.ush b/Engine/Shaders/Private/Strata/StrataMultipleScattering.ush
--- a/Engine/Shaders/Private/Strata/StrataMultipleScattering.ush
+++ b/Engine/Shaders/Private/Strata/StrataMultipleScattering.ush
@@ -4,6 +4,8 @@
 
 #include "/Engine/Private/Common.ush"
 
+// this multiple scattering implementation is based on "Practical multiple scattering compensation for microfacet models." E. Turquin
+
 // Multiple scattering computation use a precomputed LUT. There are two variant: 
 //  1) Using 2D LUT using Schlick Fresnel decomposition (view angle, roughness)
 //  2) Using 3D LUT building a full 3D LUT (view angle, roughness, F0)
@@ -17,11 +19,11 @@ struct FStrataGGXDirectionalAlbedo
 
 struct FStrataGGXLUTResult
 {
-	// Fresnel A/B term for reconstructing directional albedo for reflection (used to evaluate a chromatic F0). These values are dependent of F0.
+	// Fresnel A/B term for reconstructing directional albedo for reflection (used to evaluate a chromatic specular reflectance). These values are dependent of F0.
 	float A_R;
 	float B_R;
 
-	// Complete directional albedo for reflection, refraction, and diffuse (used to evluate an achromatic F0). These values are independent of F0
+	// Complete directional albedo for reflection, refraction, and diffuse (used to evluate an achromatic specular reflectance). These values are independent of F0
 	float E_R;
 	float E_T;
 	float E_D;
@@ -31,53 +33,54 @@ struct FStrataGGXLUTResult
 	FStrataGGXDirectionalAlbedo DirectionalAlbedo;	// Directional albedo (with multiple scattering for reflection/transmission/diffuse
 };
 
-// Return a scaling factor to multiple a specular GGX interface with, to add multiple scattering contribution (for conductor)
-// From "Practical multiple scattering compensation for microfacet models." E. Turquin
-float3 StrataEvaluateMultipleScatteringConductor(float3 F0, FStrataGGXLUTResult LUT)
+// Return the directional albedo of a specular GGX interface for reflection term without multiple scattering
+float3 StrataEvaluateSpecularGGXDirectionalAlbedo(float3 F0, float3 F90, FStrataGGXLUTResult LUT)
 {
 	// Reconstruction of the directional Albedo for upper hemisphere
 	// E_R = F0.A_R + (1-F0).B_R
-	// 
+	return F0 * LUT.A_R + (F90 - F0) * LUT.B_R;
+}
+
+// Return a scaling factor to multiply a specular GGX interface with, to add multiple scattering contribution (for conductor)
+float3 StrataEvaluateMultipleScatteringConductor(float3 F0, float3 F90, FStrataGGXLUTResult LUT)
+{
 	// Microfacet multiple scattering term for *conductor*
 	//  f_ms(wo,wi) = F_ss . (1 +  F0 (1-E_R)/ E_R)
-	
-	const float3 E_R = F0 * LUT.A_R + (1 - F0) * LUT.B_R;
+
+	const float3 E_R = StrataEvaluateSpecularGGXDirectionalAlbedo(F0, F90, LUT);
 	const float3 Safe_E_R = max(1e-8f, E_R);	// the threshold making E_R safe is lower than any of the non-zero values in the current LUT
 	return (1 + F0 * (1 - E_R) / Safe_E_R);
 }
 
 // Return a scaling factor to multiple a specular GGX interface with, to add multiple scattering contribution (for dieletric)
-float3 StrataEvaluateMultipleScatteringDielectric(float3 F0, FStrataGGXLUTResult LUT)
+float3 StrataEvaluateMultipleScatteringDielectric(FStrataGGXLUTResult LUT)
 {
-	// Microfacet multiple scattering term for *dielietric*
+	// Microfacet multiple scattering term for *dieletric*
 	//  f_ms(wo,wi) = F_ss / (E_T + E_R)
 
 	return 1.f / (LUT.E_R + LUT.E_T);
 }
 
-// Return the directional albedo of a specular GGX interface for reflection term without multiple scattering
-float3 StrataEvaluateSpecularGGXDirectionalAlbedo(float3 F0, FStrataGGXLUTResult LUT)
-{
-	return F0 * LUT.A_R + (1 - F0) * LUT.B_R;
-}
-
 // Return the directional alebedo of a diffuse GGX interface without multiple scattering
 float3 StrataEvaluateDiffuseGGXDirectionalAlbedo(float3 Albedo, FStrataGGXLUTResult LUT)
 {
 	return Albedo * LUT.E_D;
 }
 
-FStrataGGXLUTResult StrataSampleGGXLUT(Texture2D<float4> LUT2D, Texture3D<float2> LUT3D, SamplerState LUTSampler, float2 LUTScaleBias, float NoV, float RoughnessX, float RoughnessY, float3 F0, float Metallic, float3 DiffuseAlbedo)
+FStrataGGXLUTResult StrataSampleGGXLUT(Texture2D<float4> LUT2D, Texture3D<float2> LUT3D, SamplerState LUTSampler, float2 LUTScaleBias, float NoV, float RoughnessX, float RoughnessY, float3 F0, float3 F90, float Metallic, float3 DiffuseAlbedo)
 {
 	// Use the main Roughness term to drive the energy compensation.
 	// "Revisiting Physically Based Shading at Imageworks" Kulla 2017
-	const float Roughness = RoughnessX;
+	const float LUTRoughness = RoughnessX;
+
+	// LUTF0 is used for dieletric directional albedo/multiple scattering, so we expect this terms to be mostly achromatic. If it happens it is not, then we use the max component of the LUT lookup
+	const float LUTF0 = saturate(max(max(F0.x, F0.y), F0.z));
 
 	FStrataGGXLUTResult Out = (FStrataGGXLUTResult)0;
 
 	// 2D LUT sampling
 	{
-		const float2 UV = float2(saturate(NoV), saturate(Roughness)) * LUTScaleBias.x + LUTScaleBias.yy;
+		const float2 UV = float2(saturate(NoV), saturate(LUTRoughness)) * LUTScaleBias.x + LUTScaleBias.yy;
 		const float4 S = LUT2D.SampleLevel(LUTSampler, UV, 0);
 		Out.A_R = S.x;
 		Out.B_R = S.y;
@@ -86,7 +89,7 @@ FStrataGGXLUTResult StrataSampleGGXLUT(Texture2D<float4> LUT2D, Texture3D<float2
 
 	// 3D LUT sampling
 	{
-		const float3 UVW = float3(saturate(NoV), saturate(Roughness), saturate(F0.y)) * LUTScaleBias.x + LUTScaleBias.yyy;
+		const float3 UVW = float3(saturate(NoV), saturate(LUTRoughness), LUTF0) * LUTScaleBias.x + LUTScaleBias.yyy;
 		const float2 S = LUT3D.SampleLevel(LUTSampler, UVW, 0);
 		Out.E_R = S.x;
 		Out.E_T = S.y;
@@ -94,10 +97,10 @@ FStrataGGXLUTResult StrataSampleGGXLUT(Texture2D<float4> LUT2D, Texture3D<float2
 
 	// Derived MS scales
 	{
-		const float3 DA_ConductorSS = StrataEvaluateSpecularGGXDirectionalAlbedo(F0, Out);
+		const float3 DA_ConductorSS = StrataEvaluateSpecularGGXDirectionalAlbedo(F0, F90, Out);
 
-		const float3 MSScale_Dieletric = StrataEvaluateMultipleScatteringDielectric(F0, Out);
-		const float3 MSScale_Conductor = StrataEvaluateMultipleScatteringConductor(F0, Out);
+		const float3 MSScale_Dieletric = StrataEvaluateMultipleScatteringDielectric(Out);
+		const float3 MSScale_Conductor = StrataEvaluateMultipleScatteringConductor(F0, F90, Out);
 		const float3 MSScale		   = lerp(MSScale_Dieletric, MSScale_Conductor, Metallic);
 
 		Out.SpecularMultipleScatteringScale = MSScale;
