Encodes sqrt of velocity in the velocity buffer to have enough precision for history reprojections under slow movements at 8k

diff --git a/Engine/Shaders/Private/Common.ush b/Engine/Shaders/Private/Common.ush
--- a/Engine/Shaders/Private/Common.ush
+++ b/Engine/Shaders/Private/Common.ush
@@ -80,6 +80,13 @@
 #define VELOCITY_ENCODE_DEPTH 1
 #endif
 
+#if FEATURE_LEVEL >= FEATURE_LEVEL_SM5
+	#define VELOCITY_ENCODE_GAMMA 1
+#else
+	#define VELOCITY_ENCODE_GAMMA 0
+#endif
+
+
 // Using SV_ClipDistance has overhead (15% slower base pass in triangle bound test scene on PS4) so projects have to opt-in
 #define USE_GLOBAL_CLIP_PLANE (PLATFORM_SUPPORTS_GLOBAL_CLIP_PLANE && PROJECT_ALLOW_GLOBAL_CLIP_PLANE && !MATERIAL_DOMAIN_POSTPROCESS && !MATERIAL_DOMAIN_UI)
 
@@ -1552,21 +1559,25 @@ struct FScreenVertexOutput
 // texture is 16bit 0..1 range per channel
 float4 EncodeVelocityToTexture(float3 V)
 {
+	#if VELOCITY_ENCODE_GAMMA
+		V.xy = sign(V.xy) * sqrt(abs(V.xy)) * (2.0 / sqrt(2.0));
+	#endif
+
 	// 0.499f is a value smaller than 0.5f to avoid using the full range to use the clear color (0,0) as special value
 	// 0.5f to allow for a range of -2..2 instead of -1..1 for really fast motions for temporal AA
 	float4 EncodedV;
 	EncodedV.xy = V.xy * (0.499f * 0.5f) + 32767.0f / 65535.0f;
 
-#if VELOCITY_ENCODE_DEPTH
-	uint Vz = asuint(V.z);
+	#if VELOCITY_ENCODE_DEPTH
+		uint Vz = asuint(V.z);
 
-	EncodedV.z = saturate(float((Vz >> 16) & 0xFFFF) * rcp(65535.0f) + (0.1 / 65535.0f));
-	EncodedV.w = saturate(float((Vz >>  0) & 0xFFFF) * rcp(65535.0f) + (0.1 / 65535.0f));
+		EncodedV.z = saturate(float((Vz >> 16) & 0xFFFF) * rcp(65535.0f) + (0.1 / 65535.0f));
+		EncodedV.w = saturate(float((Vz >>  0) & 0xFFFF) * rcp(65535.0f) + (0.1 / 65535.0f));
+	#else
+		EncodedV.zw = 0.0;
+	#endif
 
 	return EncodedV;
-#else
-	return float4(EncodedV.x, EncodedV.y, 0, 0);
-#endif
 }
 // see EncodeVelocityToTexture()
 float3 DecodeVelocityFromTexture(float4 EncodedV)
@@ -1576,13 +1587,17 @@ float3 DecodeVelocityFromTexture(float4 EncodedV)
 	float3 V;
 	V.xy = EncodedV.xy * InvDiv - 32767.0f / 65535.0f * InvDiv;
 
-#if VELOCITY_ENCODE_DEPTH
-	V.z = asfloat((uint(round(EncodedV.z * 65535.0f)) << 16) | uint(round(EncodedV.w * 65535.0f)));
+	#if VELOCITY_ENCODE_DEPTH
+		V.z = asfloat((uint(round(EncodedV.z * 65535.0f)) << 16) | uint(round(EncodedV.w * 65535.0f)));
+	#else
+		V.z = 0.0;
+	#endif
+
+	#if VELOCITY_ENCODE_GAMMA
+		V.xy = (V.xy * abs(V.xy)) * 0.5;
+	#endif
 
 	return V;
-#else
-	return float3 (V.x, V.y, 0);
-#endif
 }
 
 // Used for the Global Illumination in the GIReplace material expression
