Changed the VSM instance culling loop to use WaveInterlockedAdd instead of using group shared and doing work twice.
 - simplifies the logic significantly and makes it more efficient

diff --git a/Engine/Shaders/Private/VirtualShadowMaps/BuildPerPageDrawCommands.usf b/Engine/Shaders/Private/VirtualShadowMaps/BuildPerPageDrawCommands.usf
--- a/Engine/Shaders/Private/VirtualShadowMaps/BuildPerPageDrawCommands.usf
+++ b/Engine/Shaders/Private/VirtualShadowMaps/BuildPerPageDrawCommands.usf
@@ -9,6 +9,7 @@
 #include "ProjectionCommon.ush"
 #include "../Nanite/NaniteDataDecode.ush"
 #include "../InstanceCulling/InstanceCullingCommon.ush"
+#include "../WaveOpUtil.ush"
 
 #ifndef NEAR_CLIP
 #define NEAR_CLIP 1
@@ -50,95 +51,24 @@ uint bInstancePerPage;
 RWStructuredBuffer<uint> OutDynamicCasterFlags;
 StructuredBuffer<FPrimCullingCommand> PrimitiveCullingCommands;
 
-groupshared uint GroupTotalVisibleInstances;
-groupshared uint GroupVisibleInstanceOutputOffset;
 
-
-#if 0
-// Ths was intended to replace the duplicated loops below and act as a temlpate (switched using the flag, but shader compiler errors...)
-
-// Detitute persons enum...
-#define PML_COUNT (1)
-#define PML_EMIT (2)
-
-// Poor mans template...
-void ProcessMipLevel(uint PrimaryViewId, uint MipLevel, FFrustumCullData Cull, uint InstanceId, uint DrawCommandId, int ProcessingMode, inout uint MipOverlapCount, inout uint VisibleInstanceOutputOffset)
+void WriteCmd(uint2 VirtualPage, uint MipViewId, uint InstanceId, uint DrawCommandId)
 {
-	uint MipViewId = MipLevel * NumPrimaryViews + PrimaryViewId;
-	FNaniteView MipView = GetNaniteView(MipViewId);
-	uint VirtualShadowMapId = uint(MipView.TargetLayerIndex);
-
-	FScreenRect Rect = GetScreenRect(MipView.ViewRect, Cull, 4);
-
-	// OLA-TODO: Caching is not respected
-	if (OverlapsAnyValidPage(VirtualShadowMapId, MipLevel, MipView.ViewRect.xy, Rect, true))
-	{
-		uint4 RectPages = uint4(Rect.Pixels) >> VSM_LOG2_PAGE_SIZE;
-
-		// Clip to actually allocated pages
-		// TODO: move this to be done as part of or before the overlap test?
-		uint4 AllocatedBounds = PageRectBounds[VirtualShadowMapId * VSM_MAX_MIP_LEVELS + MipLevel];
-		RectPages.xy = max(RectPages.xy, AllocatedBounds.xy);
-		RectPages.zw = min(RectPages.zw, AllocatedBounds.zw);
-		if (all(RectPages.zw >= RectPages.xy))
-		{
-			// 1. Loop over the super-pages, note the implied +1 for the regular level inherent in the >= 0
-			// TODO: May want to get the combined flag from OverlapsAnyValidPage and find the first bit set to not loop where there are no super pages
-			//       but then we can't unroll.
-			for (int SuperPageLevel = VSM_NUM_SUPER_PAGE_LEVELS; SuperPageLevel >= 0; --SuperPageLevel)
-			{
-				// The stride is equal to the size of the super pages
-				uint PageStride = 1U << uint(SuperPageLevel);
-				// Super page mask is stored in the higher bits of the page flags 
-				uint SuperLevelMask = 1U << (uint(SuperPageLevel) + VSM_NUM_SUPER_START_BIT);
-				// Mask to test against to know if there was a higher-level super page for this page
-				// Should be 1 for all bits higher than the current super page
-				uint HigherSuperLevelMask = ~(SuperLevelMask - 1U);
-
-				// Count commands for each overlapped page
-				uint PageTableLevelOffset = CalcPageTableLevelOffset(VirtualShadowMapId, MipLevel);
-				for (uint y = RectPages.y; y <= RectPages.w; y += PageStride)
-				{
-					for (uint x = RectPages.x; x <= RectPages.z; x += PageStride)
-					{
-						uint PageFlagOffset = PageTableLevelOffset + CalcPageOffsetInLevel(MipLevel, uint2(x, y));
-						uint PageFlag = PageFlags[PageFlagOffset];
-
-						// Mask off all bits above and including the current super page level, if this is the first one set, we should emit for this level
-						// otherwise the page is covered by a higher level.
-						if ((PageFlag & HigherSuperLevelMask) == SuperLevelMask)
-						{
-							BRANCH
-							if (ProcessingMode == PML_COUNT)
-							{
-								MipOverlapCount += 1U;
-							}
-							else
-							{
-								FPageInfo PageInfo;
-								PageInfo.VirtualPage = uint2(x, y);
-								PageInfo.ViewId = MipViewId;
-								PageInfo.SuperPageLevel = uint(SuperPageLevel);
-
-								FVisibleInstanceCmd VisibleInstanceCmd;
-								VisibleInstanceCmd.PackedPageInfo = PackPageInfo(PageInfo);
-								VisibleInstanceCmd.InstanceId = InstanceId;
-								VisibleInstanceCmd.DrawCommandId = DrawCommandId;
-
-								VisibleInstancesOut[VisibleInstanceOutputOffset] = VisibleInstanceCmd;
-								VisibleInstanceOutputOffset += 1U;
-							}
-						}
-					}
-				}
-			}
-		}
-	}
-
+	FPageInfo PageInfo;
+	PageInfo.VirtualPage = VirtualPage;
+	PageInfo.ViewId = MipViewId;
+	PageInfo.SuperPageLevel = 0U;
+
+	FVisibleInstanceCmd VisibleInstanceCmd;
+	VisibleInstanceCmd.PackedPageInfo = PackPageInfo(PageInfo);
+	VisibleInstanceCmd.InstanceId = InstanceId;
+	VisibleInstanceCmd.DrawCommandId = DrawCommandId;
+
+	uint VisibleInstanceOutputOffset = 0U;
+	WaveInterlockedAddScalar_(VisibleInstanceCountBufferOut[0], 1U, VisibleInstanceOutputOffset);
+	VisibleInstancesOut[VisibleInstanceOutputOffset] = VisibleInstanceCmd;
 }
 
-#endif // 0
-
 [numthreads(NUM_THREADS_PER_GROUP, 1, 1)]
 void CullPerPageDrawCommandsCs(uint InstanceIdIndex : SV_DispatchThreadID, uint GroupIndex : SV_GroupIndex)
 {
@@ -150,14 +80,9 @@ void CullPerPageDrawCommandsCs(uint InstanceIdIndex : SV_DispatchThreadID, uint
 		return;
 	}
 
-	bool bIsInstanceValid = InstanceIdIndex < NumInstances;
-	// Clamp to a valid index
-	InstanceIdIndex = min(InstanceIdIndex, NumInstances - 1u);
-	if (GroupIndex == 0u)
+	if (InstanceIdIndex >= NumInstances)
 	{
-		GroupTotalVisibleInstances = 0u;
-		// Not strictly needed, but probably good form to initialize anyway
-		GroupVisibleInstanceOutputOffset = 0u;
+		return;
 	}
 
 	// Load the offset (produced by the previous pass on GPU) or maybe use a fresh/recycled buffer for each and set to zero?
@@ -175,8 +100,6 @@ void CullPerPageDrawCommandsCs(uint InstanceIdIndex : SV_DispatchThreadID, uint
 	const bool bHasMoved = GetGPUSceneFrameNumber() == InstanceData.LastUpdateSceneFrameNumber || Cmd.bMaterialMayModifyPosition;
 	const uint PageFlagMask = bHasMoved ? VSM_ALLOCATED_FLAG : VSM_INVALID_FLAG;
 
-	GroupMemoryBarrierWithGroupSync();
-
 	uint ThreadTotalForAllViews = 0;
 #if LOOP_OVER_VIEWS
 	// Loop over views and output visible instance (i.e., those that overlap a valid page)
@@ -185,9 +108,6 @@ void CullPerPageDrawCommandsCs(uint InstanceIdIndex : SV_DispatchThreadID, uint
 	uint PrimaryViewId = ViewIdIndex;
 #endif
 	{
-		uint VisibleMipsMask = 0;
-		uint ThreadTotalPageOverlaps = 0;
-
 		FNaniteView NaniteView = GetNaniteView(PrimaryViewId);
 		const uint2 TargetViewSize = uint2(NaniteView.ViewSizeAndInvSize.xy);
 
@@ -197,14 +117,11 @@ void CullPerPageDrawCommandsCs(uint InstanceIdIndex : SV_DispatchThreadID, uint
 
 		FFrustumCullData Cull = BoxCullFrustum(InstanceData.LocalBoundsCenter, InstanceData.LocalBoundsExtent, LocalToClip, bNearClip, false);
 
-		if (bIsInstanceValid && Cull.bIsVisible)
+		if (Cull.bIsVisible)
 		{
 			// Loop over mip levels and count number of output visible instances, also retain a bit for each mip level (used to skip empty levels in output loop)
 			for (uint MipLevel = 0; MipLevel < uint(NaniteView.TargetNumMipLevels); ++MipLevel)
 			{
-				uint MipPageOverlapCount = 0U;
-				//uint DummyVisibleInstanceOutputOffset = 0U;
-				//ProcessMipLevel(PrimaryViewId, MipLevel, Cull, InstanceId, DrawCommandId, PML_COUNT, MipPageOverlapCount, DummyVisibleInstanceOutputOffset);
 				uint MipViewId = MipLevel * NumPrimaryViews + PrimaryViewId;
 				FNaniteView MipView = GetNaniteView(MipViewId);
 				uint VirtualShadowMapId = uint(MipView.TargetLayerIndex);
@@ -235,7 +152,13 @@ void CullPerPageDrawCommandsCs(uint InstanceIdIndex : SV_DispatchThreadID, uint
 
 									if ((PageFlag & PageFlagMask) != 0)
 									{
-										MipPageOverlapCount += 1U;
+										++ThreadTotalForAllViews;
+										WriteCmd(uint2(x,y), MipViewId, InstanceId, DrawCommandId);
+
+										if (bHasMoved)
+										{
+											OutDynamicCasterFlags[PageFlagOffset] = 1;
+										}
 									}
 								}
 							}
@@ -243,132 +166,22 @@ void CullPerPageDrawCommandsCs(uint InstanceIdIndex : SV_DispatchThreadID, uint
 						else
 						{
 							// bInstancePerPage just one per rect
-							MipPageOverlapCount += 1U;
-						}
-					}
-				}
-
-				if (MipPageOverlapCount > 0u)
-				{
-					VisibleMipsMask |= (1u << MipLevel);
-				}
-				ThreadTotalPageOverlaps += MipPageOverlapCount;
-			}
-		}
-
-		uint VisibleInstanceOutputOffset = 0;
-		// Allocate space within the group to get sub-range for output, and also add to total for whole group
-		BRANCH
-		if (ThreadTotalPageOverlaps > 0)
-		{
-			InterlockedAdd(GroupTotalVisibleInstances, ThreadTotalPageOverlaps, VisibleInstanceOutputOffset);
-		}
-		ThreadTotalForAllViews += ThreadTotalPageOverlaps;
-		GroupMemoryBarrierWithGroupSync();
-
-		// Let one thread per group allocate space in global output buffer, sufficient for all instance draw infos.
-		BRANCH
-		if (GroupIndex == 0)
-		{
-			if (GroupTotalVisibleInstances > 0)
-			{
-				uint NumNodes = GroupTotalVisibleInstances;
-				GroupTotalVisibleInstances = 0;
-				InterlockedAdd(VisibleInstanceCountBufferOut[0], NumNodes, GroupVisibleInstanceOutputOffset);
-			}
-		}
-		GroupMemoryBarrierWithGroupSync();
-		// Add offset in global node output buffer to the output offset for the thread
-		VisibleInstanceOutputOffset += GroupVisibleInstanceOutputOffset;
-
-		// Process mip levels with visible pages and output nodes
-		while (VisibleMipsMask != 0)
-		{
-			uint MipLevel = firstbitlow(VisibleMipsMask);	// Jump to next set bit.
-			VisibleMipsMask &= VisibleMipsMask - 1u;		// Remove bottom bit.
-
-			//uint DummyMipPageOverlapCount = 0U;
-			//ProcessMipLevel(PrimaryViewId, MipLevel, Cull, InstanceId, DrawCommandId, PML_EMIT, DummyMipPageOverlapCount, VisibleInstanceOutputOffset);
-			uint MipViewId = MipLevel * NumPrimaryViews + PrimaryViewId;
-			FNaniteView MipView = GetNaniteView(MipViewId);
-			uint VirtualShadowMapId = uint(MipView.TargetLayerIndex);
-
-			FScreenRect Rect = GetScreenRect(MipView.ViewRect, Cull, 4);
-
-			if (OverlapsAnyValidPage(VirtualShadowMapId, MipLevel, MipView.ViewRect.xy, Rect, bHasMoved))
-			{
-				uint4 RectPages = uint4( MipView.ViewRect.xyxy + Rect.Pixels ) >> VSM_LOG2_PAGE_SIZE;
-
-				// Clip to actually allocated pages
-				// TODO: move this to be done as part of or before the overlap test?
-				uint4 AllocatedBounds = PageRectBounds[VirtualShadowMapId * VSM_MAX_MIP_LEVELS + MipLevel];
-				RectPages.xy = max(RectPages.xy, AllocatedBounds.xy);
-				RectPages.zw = min(RectPages.zw, AllocatedBounds.zw);
-				if (all(RectPages.zw >= RectPages.xy))
-				{
-					if (bInstancePerPage)
-					{
-						// Count commands for each overlapped page
-						uint PageTableLevelOffset = CalcPageTableLevelOffset(VirtualShadowMapId, MipLevel);
-						for (uint y = RectPages.y; y <= RectPages.w; y += 1U)
-						{
-							for (uint x = RectPages.x; x <= RectPages.z; x += 1U)
+							++ThreadTotalForAllViews;
+							WriteCmd(RectPages.xy, MipViewId, InstanceId, DrawCommandId);
+							if (bHasMoved)
 							{
-								uint PageFlagOffset = PageTableLevelOffset + CalcPageOffsetInLevel(MipLevel, uint2(x, y));
-								uint PageFlag = PageFlags[PageFlagOffset];
-
-								if ((PageFlag & PageFlagMask) != 0)
+								uint PageTableLevelOffset = CalcPageTableLevelOffset(VirtualShadowMapId, MipLevel);
+								for (uint y = RectPages.y; y <= RectPages.w; y += 1U)
 								{
-									FPageInfo PageInfo;
-									PageInfo.VirtualPage = uint2(x, y);
-									PageInfo.ViewId = MipViewId;
-									PageInfo.SuperPageLevel = 0U;
-
-									FVisibleInstanceCmd VisibleInstanceCmd;
-									VisibleInstanceCmd.PackedPageInfo = PackPageInfo(PageInfo);
-									VisibleInstanceCmd.InstanceId = InstanceId;
-									VisibleInstanceCmd.DrawCommandId = DrawCommandId;
-
-									VisibleInstancesOut[VisibleInstanceOutputOffset] = VisibleInstanceCmd;
-									VisibleInstanceOutputOffset += 1U;
-
-									if (bHasMoved)
+									for (uint x = RectPages.x; x <= RectPages.z; x += 1U)
 									{
-										OutDynamicCasterFlags[PageFlagOffset] = 1;
-									}
-								}
-							}
-						}
-					}
-					else // !bInstancePerPage
-					{
-						FPageInfo PageInfo;
-						PageInfo.VirtualPage = RectPages.xy;
-						PageInfo.ViewId = MipViewId;
-						PageInfo.SuperPageLevel = 0U;
-
-						FVisibleInstanceCmd VisibleInstanceCmd;
-						VisibleInstanceCmd.PackedPageInfo = PackPageInfo(PageInfo);
-						VisibleInstanceCmd.InstanceId = InstanceId;
-						VisibleInstanceCmd.DrawCommandId = DrawCommandId;
-
-						VisibleInstancesOut[VisibleInstanceOutputOffset] = VisibleInstanceCmd;
-						VisibleInstanceOutputOffset += 1U;
-
-						// Need to perform loop to take care of cache invalidations
-						if (bHasMoved)
-						{
-							uint PageTableLevelOffset = CalcPageTableLevelOffset(VirtualShadowMapId, MipLevel);
-							for (uint y = RectPages.y; y <= RectPages.w; y += 1U)
-							{
-								for (uint x = RectPages.x; x <= RectPages.z; x += 1U)
-								{
-									uint PageFlagOffset = PageTableLevelOffset + CalcPageOffsetInLevel(MipLevel, uint2(x, y));
-									uint PageFlag = PageFlags[PageFlagOffset];
+										uint PageFlagOffset = PageTableLevelOffset + CalcPageOffsetInLevel(MipLevel, uint2(x, y));
+										uint PageFlag = PageFlags[PageFlagOffset];
 
-									if ((PageFlag & PageFlagMask) != 0)
-									{
-										OutDynamicCasterFlags[PageFlagOffset] = 1;
+										if ((PageFlag & PageFlagMask) != 0)
+										{
+											OutDynamicCasterFlags[PageFlagOffset] = 1;
+										}
 									}
 								}
 							}
@@ -376,56 +189,11 @@ void CullPerPageDrawCommandsCs(uint InstanceIdIndex : SV_DispatchThreadID, uint
 					}
 				}
 			}
-			// End ProcessMipLevel
-#if 0
-			uint MipViewId = MipLevel * NumPrimaryViews + PrimaryViewId;
-			FNaniteView MipView = GetNaniteView(MipViewId);
-
-			uint VirtualShadowMapId = uint(MipView.TargetLayerIndex);
-
-			// TODO: minor optimization possible, but need to duplicate setup from CullRasterize for virtual targets
-			//float2 ViewSize = float2( ( TargetViewSize + ( 1u << MipLevel ) - 1u ) >> MipLevel );
-			FScreenRect Rect = GetScreenRect(MipView.ViewRect, Cull, 4);
-
-			uint4 RectPages = uint4(Rect.Pixels) >> VSM_LOG2_PAGE_SIZE;
-
-			// Clip to actually allocated pages
-			uint4 AllocatedBounds = PageRectBounds[VirtualShadowMapId * VSM_MAX_MIP_LEVELS + MipLevel];
-			RectPages.xy = max(RectPages.xy, AllocatedBounds.xy);
-			RectPages.zw = min(RectPages.zw, AllocatedBounds.zw);
-			// Count commands for each overlapped page
-			uint PageTableLevelOffset = CalcPageTableLevelOffset(VirtualShadowMapId, MipLevel);
-			for (uint y = RectPages.y; y <= RectPages.w; y++)
-			{
-				for (uint x = RectPages.x; x <= RectPages.z; x++)
-				{
-					uint PageFlagOffset = PageTableLevelOffset + CalcPageOffsetInLevel(MipLevel, uint2(x, y));
-					uint PageFlag = PageFlags[PageFlagOffset];
-
-					if (PageFlag != 0)
-					{
-						FPageInfo PageInfo;
-						PageInfo.VirtualPage = uint2(x, y);
-						PageInfo.ViewId = MipViewId;
-
-						FVisibleInstanceCmd VisibleInstanceCmd;
-						VisibleInstanceCmd.PackedPageInfo = PackPageInfo(PageInfo);
-						VisibleInstanceCmd.InstanceId = InstanceId;
-						VisibleInstanceCmd.DrawCommandId = DrawCommandId;
-
-						VisibleInstancesOut[NodeOutputOffset] = VisibleInstanceCmd;
-						NodeOutputOffset += 1U;
-					}
-				}
-			}
-#endif
 		}
 	}
-	if (bIsInstanceValid)
-	{
-		// Accumulate total number of instances/page pairs for each draw command, is used to allocate space and output compact range of instances later
-		InterlockedAdd(DrawCommandInstanceCountBufferOut[DrawCommandId], ThreadTotalForAllViews);
-	}
+
+	// Accumulate total number of instances/page pairs for each draw command, is used to allocate space and output compact range of instances later
+	InterlockedAdd(DrawCommandInstanceCountBufferOut[DrawCommandId], ThreadTotalForAllViews);
 }
 
 RWBuffer<uint> DrawIndirectArgsBufferOut;
