Fixed StrataVerticalLayeringParameterBlending

diff --git a/Engine/Shaders/Private/Strata/Strata.ush b/Engine/Shaders/Private/Strata/Strata.ush
--- a/Engine/Shaders/Private/Strata/Strata.ush
+++ b/Engine/Shaders/Private/Strata/Strata.ush
@@ -1222,7 +1222,7 @@ FStrataData StrataVerticalLayeringParameterBlending(FStrataData Top, FStrataData
 	SLAB_METALLIC(ResultBSDF) =			lerp(SLAB_METALLIC(TopBSDF),					SLAB_METALLIC(BaseBSDF),					GreyThroughput);
 	SLAB_ROUGHNESSX(ResultBSDF) =		lerp(SLAB_ROUGHNESSX(TopBSDF),					SLAB_ROUGHNESSX(BaseBSDF),					GreyThroughput);
 	SLAB_HAZINESS(ResultBSDF) =			lerp(SLAB_HAZINESS(TopBSDF),					SLAB_HAZINESS(BaseBSDF),					GreyThroughput);
-	SLAB_METALLIC(ResultBSDF) =			lerp(SLAB_FUZZ_AMOUNT(TopBSDF),					SLAB_FUZZ_AMOUNT(BaseBSDF),					GreyThroughput);
+	SLAB_FUZZ_AMOUNT(ResultBSDF) =		lerp(SLAB_FUZZ_AMOUNT(TopBSDF),					SLAB_FUZZ_AMOUNT(BaseBSDF),					GreyThroughput);
 
 	// Anisotropy is disabled with vertical blending. 
 	SLAB_ROUGHNESSY(ResultBSDF) =		SLAB_ROUGHNESSX(ResultBSDF);
