Fix calculating the origin for instances offset by a mesh pivot. Used when computing dithering LOD transition for instances.

diff --git a/Engine/Shaders/Private/LocalVertexFactory.ush b/Engine/Shaders/Private/LocalVertexFactory.ush
--- a/Engine/Shaders/Private/LocalVertexFactory.ush
+++ b/Engine/Shaders/Private/LocalVertexFactory.ush
@@ -407,7 +407,8 @@ float GetInstanceRandom(FVertexFactoryIntermediates Intermediates)
 
 float3 GetInstanceOrigin(FVertexFactoryIntermediates Intermediates)
 {
-	return Intermediates.InstanceOrigin.xyz + InstancingOffset.xyz;
+	float3 offset = mul(InstancingOffset.xyz, GetInstanceToLocal3x3(Intermediates));
+	return Intermediates.InstanceOrigin.xyz + offset;
 }
 
 #endif	// USE_INSTANCING
