Fixed starta backcompatibility to have authorit in StrataEvaluation.ush

diff --git a/Engine/Shaders/Private/DeferredLightPixelShaders.usf b/Engine/Shaders/Private/DeferredLightPixelShaders.usf
--- a/Engine/Shaders/Private/DeferredLightPixelShaders.usf
+++ b/Engine/Shaders/Private/DeferredLightPixelShaders.usf
@@ -29,7 +29,6 @@
 
 #if STRATA_ENABLED
 #define STRATA_INLINE_SHADING 0
-#define STRATA_MICROFACET_MULTIPLESCATTERING 1
 #define STRATA_SSS_TRANSMISSION USE_TRANSMISSION
 #include "/Engine/Private/Strata/StrataEvaluation.ush"
 #include "/Engine/Private/Strata/StrataDeferredLighting.ush"
