GGX multiple scattering lut now contains specular transmission for F0=0.04 to inverse it on the Chan model.
This makes the non-metalness workflow energy conservative.
Added alternative implementation to StrataEvaluateGGXLUT  when STRATA_MICROFACET_MULTIPLESCATTERING=0 to simplify the code.

diff --git a/Engine/Shaders/Private/Strata/Strata.ush b/Engine/Shaders/Private/Strata/Strata.ush
--- a/Engine/Shaders/Private/Strata/Strata.ush
+++ b/Engine/Shaders/Private/Strata/Strata.ush
@@ -52,6 +52,24 @@
 #define STRATA_SIMPLEVOLUME_THICKNESS_CM	100.0f
 #define STRATA_SIMPLEVOLUME_THICKNESS_M		(STRATA_SIMPLEVOLUME_THICKNESS_CM * CENTIMETER_TO_METER)
 
+///////////////////////////////////////////////////////////////////////////////
+// Utilities
+
+float F0ToMetallic(float F0)
+{
+	// Approximate the metallic input from F0 with a small lerp region
+	const float DiamondF0 = 0.24f;
+	const float CondutorMinF0 = 0.4f;	// roughly the end of semi-conductor
+	// This is compatible with UE shading model mapping allowing F0 to take a value up to 0.08 for dielectric.
+
+	return saturate((F0 - DiamondF0) / (CondutorMinF0 - DiamondF0));
+}
+
+float F0RGBToMetallic(float3 F0)
+{
+	return F0ToMetallic(max(F0.r, max(F0.g, F0.b)));
+}
+
 ///////////////////////////////////////////////////////////////////////////////
 // BSDFs representation
 struct FStrataBSDF
@@ -553,16 +571,10 @@ FStrataData GetStrataSlabBSDF(
 		// See how UMaterialExpressionStrataSlabBSDF node interface is changed
 		BSDF_SETHASF90(StrataData.Layers[0].BSDFs[0], any(F90 < 1.0f));
 
-		// Approximate the metallic input from F0.
-		const float DiamondF0 = 0.24f;
-		const float CondutorMinF0 = 0.4f; // roughly the end of semi-conductor
-		const float MaxF0 = max(F0.r, max(F0.g, F0.b));
-
 		SLAB_DIFFUSEALBEDO(StrataData.Layers[0].BSDFs[0])	= DiffuseAlbedo;
 		SLAB_F0(StrataData.Layers[0].BSDFs[0])				= F0;
 		SLAB_F90(StrataData.Layers[0].BSDFs[0])				= F90;
-		// Metalness cannot be recovered from F0 so we use a simple heuristic mapping (based on MaxF0 which should better represent the state of the material)
-		SLAB_METALLIC(StrataData.Layers[0].BSDFs[0])		= saturate((MaxF0 - DiamondF0) / (CondutorMinF0 - DiamondF0));
+		SLAB_METALLIC(StrataData.Layers[0].BSDFs[0])		= F0RGBToMetallic(F0);
 	}
 
 	SLAB_ROUGHNESSX			(StrataData.Layers[0].BSDFs[0])		= RoughnessXY.x;
diff --git a/Engine/Shaders/Private/Strata/StrataEvaluation.ush b/Engine/Shaders/Private/Strata/StrataEvaluation.ush
--- a/Engine/Shaders/Private/Strata/StrataEvaluation.ush
+++ b/Engine/Shaders/Private/Strata/StrataEvaluation.ush
@@ -26,17 +26,80 @@
 #endif
 
 #ifndef STRATA_MICROFACET_MULTIPLESCATTERING
-#define STRATA_MICROFACET_MULTIPLESCATTERING (PROJECT_STRATA_BACKCOMPATIBILITY > 0 ? 0 : 1)
+	#if PROJECT_STRATA_BACKCOMPATIBILITY > 0
+		#define STRATA_MICROFACET_MULTIPLESCATTERING 0
+	#else
+		#define STRATA_MICROFACET_MULTIPLESCATTERING 1
+	#endif
 #endif
 
 #ifndef STRATA_DIFFUSE_CHAN
-#define STRATA_DIFFUSE_CHAN (PROJECT_STRATA_BACKCOMPATIBILITY > 0 ? 0 : 1)
+	#if PROJECT_STRATA_BACKCOMPATIBILITY > 0
+		#define STRATA_DIFFUSE_CHAN 0
+	#else
+		#define STRATA_DIFFUSE_CHAN 1
+	#endif
 #endif
 
 #ifndef STRATA_SSS_TRANSMISSION
 #define STRATA_SSS_TRANSMISSION 0
 #endif
 
+
+float3 StrataEnvBRDF(float3 F0, float3 F90, float SafeRoughnessX, float SatNoV)
+{
+#if STRATA_USE_ENV_BRDF_APPROX
+	return EnvBRDFApprox(F0, F90, SafeRoughnessX, SatNoV);
+#else
+	return EnvBRDF(F0, F90, SafeRoughnessX, SatNoV);
+#endif
+}
+
+
+#if STRATA_MICROFACET_MULTIPLESCATTERING
+
+FStrataGGXLUTResult StrataEvaluateGGXLUT(float3 DiffuseAlbedo, float3 F0, float3 F90, float NoV, float Metallic, float SafeRoughnessX, float SafeRoughnessY)
+{
+#if STRATA_TRANSLUCENT_ENABLED
+	return StrataSampleGGXLUT(
+		TranslucentBasePass.Strata.GGXEnergyLUT2DTexture, TranslucentBasePass.Strata.GGXEnergyLUT3DTexture, TranslucentBasePass.Strata.GGXEnergyLUTSampler,
+		TranslucentBasePass.Strata.GGXEnergyLUTScaleBias, NoV, SafeRoughnessX, SafeRoughnessY, F0, F90, Metallic, DiffuseAlbedo);
+#else
+	return StrataSampleGGXLUT(
+		Strata.GGXEnergyLUT2DTexture, Strata.GGXEnergyLUT3DTexture, Strata.GGXEnergyLUTSampler,
+		Strata.GGXEnergyLUTScaleBias, NoV, SafeRoughnessX, SafeRoughnessY, F0, F90, Metallic, DiffuseAlbedo);
+#endif
+}
+
+#else // STRATA_MICROFACET_MULTIPLESCATTERING
+
+FStrataGGXLUTResult StrataEvaluateGGXLUT(float3 DiffuseAlbedo, float3 F0, float3 F90, float NoV, float Metallic, float SafeRoughnessX, float SafeRoughnessY)
+{
+	const float SatNoV = saturate(NoV);
+
+	FStrataGGXLUTResult Out;
+	Out.A_R = 0.0f;
+	Out.B_R = 0.0f;
+	Out.E_R = 0.0f;
+	Out.E_T = 0.0f;
+	Out.E_TChan = 0.0f;
+	Out.E_D = 0.0f;
+
+	Out.SpecularMultipleScatteringScale = 1.0f;
+
+	const float3 ApproxFresnelForSpecularTransmittance = F_Schlick(F0, F90, SatNoV);
+	float3 AppoxDirectionalAlbedo_SpecularTransmission = 1.0f - ApproxFresnelForSpecularTransmittance;
+
+	Out.DirectionalAlbedo.SpecularR = StrataEnvBRDF(F0, F90, SafeRoughnessX, SatNoV);
+	Out.DirectionalAlbedo.SpecularT = AppoxDirectionalAlbedo_SpecularTransmission;
+	Out.DirectionalAlbedo.SpecularTChan = 1.0f;
+	Out.DirectionalAlbedo.Diffuse = DiffuseAlbedo;
+
+	return Out;
+}
+
+#endif // STRATA_MICROFACET_MULTIPLESCATTERING
+
 // STRATA_TODO put in a common file
 // Point lobe in off-specular peak direction
 float3 StrataGetOffSpecularPeakReflectionDir(float3 Normal, float3 ReflectionVector, float Roughness)
@@ -200,33 +263,6 @@ float StrataBSDFBackFaceProjectedSolidAngleFactor(in FStrataBSDFContext BSDFCont
 	return BSDFType == STRATA_BSDF_TYPE_HAIR ? 1.0f : saturate(-BSDFContext.Context.NoL);
 }
 
-#if STRATA_MICROFACET_MULTIPLESCATTERING
-
-FStrataGGXLUTResult StrataEvaluateGGXLUT(float3 DiffuseAlbedo, float3 F0, float3 F90, float NoV, float Metallic, float SafeRoughnessX, float SafeRoughnessY)
-{
-#if STRATA_TRANSLUCENT_ENABLED
-	return StrataSampleGGXLUT(
-		TranslucentBasePass.Strata.GGXEnergyLUT2DTexture, TranslucentBasePass.Strata.GGXEnergyLUT3DTexture, TranslucentBasePass.Strata.GGXEnergyLUTSampler,
-		TranslucentBasePass.Strata.GGXEnergyLUTScaleBias, NoV, SafeRoughnessX, SafeRoughnessY, F0, F90, Metallic, DiffuseAlbedo);
-#else
-	return StrataSampleGGXLUT(
-		Strata.GGXEnergyLUT2DTexture, Strata.GGXEnergyLUT3DTexture, Strata.GGXEnergyLUTSampler,
-		Strata.GGXEnergyLUTScaleBias, NoV, SafeRoughnessX, SafeRoughnessY, F0, F90, Metallic, DiffuseAlbedo);
-#endif
-}
-
-#else // STRATA_MICROFACET_MULTIPLESCATTERING
-
-float3 StrataEnvBRDF(float3 F0, float3 F90, float SafeRoughnessX, float SatNoV)
-{
-#if STRATA_USE_ENV_BRDF_APPROX
-	return EnvBRDFApprox(F0, F90, SafeRoughnessX, SatNoV);
-#else
-	return EnvBRDF(F0, F90, SafeRoughnessX, SatNoV);
-#endif
-}
-
-#endif // STRATA_MICROFACET_MULTIPLESCATTERING
 
 
 
@@ -329,20 +365,29 @@ FStrataEvaluateResult StrataEvaluateBSDF(FStrataBSDFContext BSDFContext, FShadow
 
 			float3 MSScale = 1;
 			float3 HazeMSScale = 1;
-			#if STRATA_MICROFACET_MULTIPLESCATTERING
+			float3 DirectionalAlbedo_SpecularTransmission = 1.0f;
+			{
 				FStrataGGXLUTResult StrataGGXLUTResult = StrataEvaluateGGXLUT(DiffuseColor, F0, F90, BSDFContext.Context.NoV, Metallic, SafeRoughnessX, SafeRoughnessY);
 				MSScale	= StrataGGXLUTResult.SpecularMultipleScatteringScale;
-				const float3 DirectionalAlbedo_SpecularTransmission = StrataGGXLUTResult.DirectionalAlbedo.SpecularT; // We ignore Haziness influence
 				#if STRATA_FASTPATH==0
 				if (bHaziness)
 				{
 					HazeMSScale = StrataEvaluateGGXLUT(DiffuseColor, F0, F90, BSDFContext.Context.NoV, Metallic, HazeSafeRoughnessX, HazeSafeRoughnessY).SpecularMultipleScatteringScale;
 				}
 				#endif
-			#else
-				const float3 ApproxFresnelForSpecularTransmittance = F_Schlick(F0, F90, BSDFContext.SatNoV);
-				const float3 DirectionalAlbedo_SpecularTransmission = 1.0f - ApproxFresnelForSpecularTransmittance;
-			#endif
+
+				// We ignore Haziness influence on the specular transmission.
+				DirectionalAlbedo_SpecularTransmission = StrataGGXLUTResult.DirectionalAlbedo.SpecularT;
+
+				#if STRATA_MICROFACET_MULTIPLESCATTERING > 0 && STRATA_DIFFUSE_CHAN > 0
+				// We remove the Chan model baked transmittance specular directional albedo assuming F=0.04.
+				// Then we apply the transmittance according to the current surface setup.
+				// To make this a safe operation:
+				//	- We clamp SpecularTChan to a threshold to avoid div by 0 when the material is a metal. 
+				//	- When the material tends to be metal, DirectionalAlbedo_SpecularTransmission will naturally tend to be 0
+				Sample.DiffusePathValue *= DirectionalAlbedo_SpecularTransmission / max(1e-8f, StrataGGXLUTResult.DirectionalAlbedo.SpecularTChan);
+				#endif
+			}
 
 			float3 FresnelTerm = 0;
 			#if STRATA_FASTPATH==0
@@ -658,13 +703,14 @@ FStrataEnvLightResult StrataEvaluateForEnvLight(FStrataBSDFContext BSDFContext,
 
 			StrataEnvLightResult.DiffuseNormal = BSDFContext.N;
 			StrataEnvLightResult.bSubsurface = BSDF_GETHASSSS(BSDFContext.BSDF) > 0;
-			StrataEnvLightResult.DiffuseWeight = DiffuseColor;
 
 			float SafeRoughnessX = MakeRoughnessSafe(SLAB_ROUGHNESSX(BSDFContext.BSDF));
 			float3 ThinFilmFresnelWeight = 1;
 
-			const float3 ApproxFresnelForSpecularTransmittance = F_Schlick(F0, F90, BSDFContext.SatNoV);
-			float3 DirectionalAlbedo_SpecularTransmission = 1.0f - ApproxFresnelForSpecularTransmittance;
+			FStrataGGXLUTResult GGXLUT							= StrataEvaluateGGXLUT(DiffuseColor, F0, F90, BSDFContext.Context.NoV, Metallic, SafeRoughnessX, SafeRoughnessX);
+			StrataEnvLightResult.DiffuseWeight					= GGXLUT.DirectionalAlbedo.Diffuse;		// This already accounts for specular transmittance
+			float3 EvalEnvBRDF									= GGXLUT.DirectionalAlbedo.SpecularR;
+			const float3 DirectionalAlbedo_SpecularTransmission	= GGXLUT.DirectionalAlbedo.SpecularT;	// We ignore Haziness influence
 
 			BRANCH
 			if (bEnableSpecular)
@@ -680,6 +726,7 @@ FStrataEnvLightResult StrataEvaluateForEnvLight(FStrataBSDFContext BSDFContext,
 					F0 = 1.f;
 				}
 				#endif
+				EvalEnvBRDF *= ThinFilmFresnelWeight;
 
 
 				#if STRATA_FASTPATH==0
@@ -700,19 +747,6 @@ FStrataEnvLightResult StrataEvaluateForEnvLight(FStrataBSDFContext BSDFContext,
 				}
 				#endif
 
-				// Compute the directional albedo for specular and diffuse
-			#if STRATA_MICROFACET_MULTIPLESCATTERING
-				float3 EvalEnvBRDF = 0;
-				{
-					FStrataGGXLUTResult GGXLUT	= StrataEvaluateGGXLUT(DiffuseColor, F0, F90, BSDFContext.Context.NoV, Metallic, SafeRoughnessX, SafeRoughnessX);
-					EvalEnvBRDF					= ThinFilmFresnelWeight * GGXLUT.DirectionalAlbedo.SpecularR;
-					StrataEnvLightResult.DiffuseWeight = GGXLUT.DirectionalAlbedo.Diffuse;
-					DirectionalAlbedo_SpecularTransmission = GGXLUT.DirectionalAlbedo.SpecularT; // We ignore Haziness influence
-				}
-			#else
-				const float3 EvalEnvBRDF = ThinFilmFresnelWeight * StrataEnvBRDF(F0, F90, SafeRoughnessX, BSDFContext.SatNoV);
-			#endif
-
 				StrataEnvLightResult.SpecularDirection = StrataGetOffSpecularPeakReflectionDir(BSDFContext.N, BSDFContext.R, SafeRoughnessX);
 				StrataEnvLightResult.SpecularWeight = EvalEnvBRDF;
 				StrataEnvLightResult.SpecularSafeRoughness = SafeRoughnessX;
@@ -745,14 +779,6 @@ FStrataEnvLightResult StrataEvaluateForEnvLight(FStrataBSDFContext BSDFContext,
 				}
 				#endif
 			}
-			#if STRATA_MICROFACET_MULTIPLESCATTERING
-			else
-			{
-				// Compute the directional albedo for diffuse
-				FStrataGGXLUTResult GGXLUT = StrataEvaluateGGXLUT(DiffuseColor, 1.0f, 1.0f, BSDFContext.Context.NoV, 0.0f, SafeRoughnessX, SafeRoughnessX);
-				StrataEnvLightResult.DiffuseWeight = GGXLUT.DirectionalAlbedo.Diffuse;
-			}
-			#endif
 
 			#if STRATA_FASTPATH==0
 			BRANCH
diff --git a/Engine/Shaders/Private/Strata/StrataLUT.usf b/Engine/Shaders/Private/Strata/StrataLUT.usf
--- a/Engine/Shaders/Private/Strata/StrataLUT.usf
+++ b/Engine/Shaders/Private/Strata/StrataLUT.usf
@@ -58,7 +58,7 @@ uint NumSamples;
 uint EnergyLUTResolution;
 int2 SliceXYCount;
 
-RWTexture3D<float2> OutLUT3D;
+RWTexture3D<float3> OutLUT3D;
 RWTexture2D<float4> OutLUT2D;
 
 void MainPS(
@@ -84,6 +84,9 @@ void MainPS(
 	float E_T = 0.0f; // Directional albedo for 'specular' transmission
 	float E_D = 0.0f; // Directional albedo for 'diffuse'  term
 
+	float E_RChan = 0.0f; // Directional albedo for 'specular' reflection assuming F0=0.04 (to inverse the transmittance baked into the Chan diffuse model)
+	float E_TChan = 0.0f; // Idem but for transmission
+
 	// Roughness
 	const float Roughness = float(Coord.y + 0.5f) / LUTResolution.y;
 	const float m = Roughness * Roughness;
@@ -93,6 +96,11 @@ void MainPS(
 	const float NoV = float(Coord.x + 0.5f) / LUTResolution.x;
 	const float3 V = float3(sqrt(1.0f - NoV * NoV),	0.0f, NoV);
 
+	// F0. TODO: parametrize with something which preserve precision in the lower part e.g., sqrt(F0)
+	const float F0 = float(Coord.z) / LUTResolution.z;
+	const float F0Chan = 0.04f;
+	const float F90 = 1;
+
 	LOOP
 	for (uint i = 0; i < NumSamples; i++)
 	{
@@ -147,34 +155,49 @@ void MainPS(
 		//  f_ss(wo,wi) = F(wo,wi) . D(wo,wi) . G(wo,wi)
 		//  f_ms(wo,wi) = F_ss . (1 +  F0 (1-E_R)/ E_R)
 
-		// TODO: parametrize F0 with something which preserve precision in the lower part e.g., sqrt(F0)
-		const float F0 = float(Coord.z) / LUTResolution.z;
-		const float F90 = 1;
 		{
 			float4 GGXSample = ImportanceSampleVisibleGGX(UniformSampleDisk(E), m2, V);
 			const float3 H  = GGXSample.xyz;
 			const float VoH = saturate(dot(V, H));
-			const float F	= F_Schlick(F0, F90, VoH).x;
+			
+			const float F		= F_Schlick(F0,     F90, VoH).x;
+			const float Fchan	= F_Schlick(F0Chan, F90, VoH).x;
 			
 			const float Eta_Incident = 1.0f;
-			const float Eta_Outgoing = DielectricF0ToIor(F0);
-			const float Eta = Eta_Outgoing / Eta_Incident;
 
-			const FDirectionResult Sample_Lr = ComputeDirection(-V, H, F, Eta, false);
-			const FDirectionResult Sample_Lt = ComputeDirection(-V, H, F, Eta, true);
+			const float Eta_Outgoing	= DielectricF0ToIor(F0);
+			const float Eta_OutgoingChan= DielectricF0ToIor(F0Chan);
+
+			const float Eta		= Eta_Outgoing     / Eta_Incident;
+			const float EtaChan = Eta_OutgoingChan / Eta_Incident;
+
+			const FDirectionResult Sample_Lr	= ComputeDirection(-V, H, F,     Eta, false);
+			const FDirectionResult Sample_LrChan= ComputeDirection(-V, H, Fchan, EtaChan, false);
 
-			const float NoL_r = Sample_Lr.Direction.z;
-			const float NoL_t = Sample_Lt.Direction.z;
+			const FDirectionResult Sample_Lt	= ComputeDirection(-V, H, F,     Eta, true);
+			const FDirectionResult Sample_LtChan= ComputeDirection(-V, H, Fchan, EtaChan, true);
+
+			const float NoL_r		= Sample_Lr.Direction.z;
+			const float NoL_rChan	= Sample_LrChan.Direction.z;
+
+			const float NoL_t		= Sample_Lt.Direction.z;
+			const float NoL_tChan	= Sample_LtChan.Direction.z;
 
 			// Schlick's Fresnel sub-term: Fc = (1-VoH)^5 
 			const float Fc = Pow5(1.0f - VoH);
 
-			const float R_WeightPdf = GGXEvalReflection(Sample_Lr.Direction, V, H, float2(m, m)).x;
-			const float T_WeightPdf = GGXEvalRefraction(Sample_Lt.Direction, V, H, float2(m, m), Eta).x;
+			const float R_WeightPdf		= GGXEvalReflection(Sample_Lr.Direction,     V, H, float2(m, m)).x;
+			const float R_WeightPdfChan = GGXEvalReflection(Sample_LrChan.Direction, V, H, float2(m, m)).x;
+
+			const float T_WeightPdf		= GGXEvalRefraction(Sample_Lt.Direction,     V, H, float2(m, m), Eta).x;
+			const float T_WeightPdfChan = GGXEvalRefraction(Sample_LtChan.Direction, V, H, float2(m, m), EtaChan).x;
 
 			E_R += Sample_Lr.Fresnel * R_WeightPdf;
 			E_T += Sample_Lt.Fresnel * T_WeightPdf;
 
+			E_RChan += Sample_LrChan.Fresnel * R_WeightPdfChan;
+			E_TChan += Sample_LtChan.Fresnel * T_WeightPdfChan;
+
 			A_R += R_WeightPdf;
 			B_R += R_WeightPdf * Fc;
 		}
@@ -187,7 +210,10 @@ void MainPS(
 	E_T /= NumSamples;
 	E_D /= NumSamples;
 
-	const float2 Out3D = float2(E_R, E_T);
+	E_RChan /= NumSamples;
+	E_TChan /= NumSamples;
+
+	const float3 Out3D = float3(E_R, E_T, E_TChan);
 	const float4 Out2D = float4(A_R, B_R, E_D, 0.0f);
 
 	// At the moment ouput 2 LUTs:
diff --git a/Engine/Shaders/Private/Strata/StrataMultipleScattering.ush b/Engine/Shaders/Private/Strata/StrataMultipleScattering.ush
--- a/Engine/Shaders/Private/Strata/StrataMultipleScattering.ush
+++ b/Engine/Shaders/Private/Strata/StrataMultipleScattering.ush
@@ -12,9 +12,10 @@
 
 struct FStrataGGXDirectionalAlbedo
 {
-	float3 SpecularR;	// Specular Directional albedo with multiple scattering (reflection)
-	float3 SpecularT;	// Specular Directional albedo with multiple scattering (transmission)
-	float3 Diffuse;		// Diffuse  Directional albedo with multiple scattering (embedded into Chan's model itself)
+	float3 SpecularR;		// Specular Directional albedo with multiple scattering (reflection)
+	float3 SpecularT;		// Specular Directional albedo with multiple scattering (transmission)
+	float3 SpecularTChan;	// Specular Directional albedo with multiple scattering (transmission) for the Chan model which has a hardcoded F0=0.04 for the diffuse interactions
+	float3 Diffuse;			// Diffuse  Directional albedo with multiple scattering (embedded into Chan's model itself)
 };
 
 struct FStrataGGXLUTResult
@@ -28,6 +29,8 @@ struct FStrataGGXLUTResult
 	float E_T;
 	float E_D;
 
+	float E_TChan;
+
 	// Derived terms
 	float3 SpecularMultipleScatteringScale;			// Multiple scattering scale factor for specular reflection
 	FStrataGGXDirectionalAlbedo DirectionalAlbedo;	// Directional albedo (with multiple scattering for reflection/transmission/diffuse
@@ -67,7 +70,7 @@ float3 StrataEvaluateDiffuseGGXDirectionalAlbedo(float3 Albedo, FStrataGGXLUTRes
 	return Albedo * LUT.E_D;
 }
 
-FStrataGGXLUTResult StrataSampleGGXLUT(Texture2D<float4> LUT2D, Texture3D<float2> LUT3D, SamplerState LUTSampler, float2 LUTScaleBias, float NoV, float RoughnessX, float RoughnessY, float3 F0, float3 F90, float Metallic, float3 DiffuseAlbedo)
+FStrataGGXLUTResult StrataSampleGGXLUT(Texture2D<float4> LUT2D, Texture3D<float3> LUT3D, SamplerState LUTSampler, float2 LUTScaleBias, float NoV, float RoughnessX, float RoughnessY, float3 F0, float3 F90, float Metallic, float3 DiffuseAlbedo)
 {
 	// Use the main Roughness term to drive the energy compensation.
 	// "Revisiting Physically Based Shading at Imageworks" Kulla 2017
@@ -90,9 +93,10 @@ FStrataGGXLUTResult StrataSampleGGXLUT(Texture2D<float4> LUT2D, Texture3D<float2
 	// 3D LUT sampling
 	{
 		const float3 UVW = float3(saturate(NoV), saturate(LUTRoughness), LUTF0) * LUTScaleBias.x + LUTScaleBias.yyy;
-		const float2 S = LUT3D.SampleLevel(LUTSampler, UVW, 0);
+		const float3 S = LUT3D.SampleLevel(LUTSampler, UVW, 0);
 		Out.E_R = S.x;
 		Out.E_T = S.y;
+		Out.E_TChan = S.z;
 	}
 
 	// Derived MS scales
@@ -106,6 +110,7 @@ FStrataGGXLUTResult StrataSampleGGXLUT(Texture2D<float4> LUT2D, Texture3D<float2
 		Out.SpecularMultipleScatteringScale = MSScale;
 		Out.DirectionalAlbedo.SpecularR		= MSScale * lerp(Out.E_R.xxx, DA_ConductorSS, Metallic);
 		Out.DirectionalAlbedo.SpecularT		= MSScale * lerp(Out.E_T, 0.f, Metallic);
+		Out.DirectionalAlbedo.SpecularTChan	= MSScale * lerp(Out.E_TChan, 0.f, Metallic);
 		Out.DirectionalAlbedo.Diffuse		= StrataEvaluateDiffuseGGXDirectionalAlbedo(DiffuseAlbedo, Out);
 	}
 	return Out;
diff --git a/Engine/Source/Runtime/Renderer/Private/Strata/Strata.cpp b/Engine/Source/Runtime/Renderer/Private/Strata/Strata.cpp
--- a/Engine/Source/Runtime/Renderer/Private/Strata/Strata.cpp
+++ b/Engine/Source/Runtime/Renderer/Private/Strata/Strata.cpp
@@ -224,7 +224,7 @@ void InitialiseStrataFrameSceneData(FSceneRenderer& SceneRenderer, FRDGBuilder&
 		bUpdateLUT = StrataSceneData.GGXEnergyLUT2DTexture == nullptr || StrataSceneData.GGXEnergyLUT2DTexture->GetDesc().Extent.X != LUTResolution || CVarStrataLUTContinousUpdate.GetValueOnAnyThread() > 0;
 		if (bUpdateLUT)
 		{
-			FRDGTextureDesc Desc3D = FRDGTextureDesc::Create3D(FIntVector(LUTResolution, LUTResolution, LUTResolution), EPixelFormat::PF_G16R16F, FClearValueBinding::Black, TexCreate_ShaderResource | TexCreate_RenderTargetable | TexCreate_UAV);
+			FRDGTextureDesc Desc3D = FRDGTextureDesc::Create3D(FIntVector(LUTResolution, LUTResolution, LUTResolution), EPixelFormat::PF_FloatR11G11B10, FClearValueBinding::Black, TexCreate_ShaderResource | TexCreate_RenderTargetable | TexCreate_UAV);
 			FRDGTextureDesc Desc2D = FRDGTextureDesc::Create2D(FIntPoint(LUTResolution, LUTResolution), EPixelFormat::PF_FloatRGBA, FClearValueBinding::Black, TexCreate_ShaderResource | TexCreate_RenderTargetable | TexCreate_UAV);
 			FRDGTextureRef OutLUT3D = GraphBuilder.CreateTexture(Desc3D, TEXT("StrataEnergyLUT3D"));
 			FRDGTextureRef OutLUT2D = GraphBuilder.CreateTexture(Desc2D, TEXT("StrataEnergyLUT2D"));
diff --git a/Engine/Source/Runtime/Renderer/Private/Strata/Strata.h b/Engine/Source/Runtime/Renderer/Private/Strata/Strata.h
--- a/Engine/Source/Runtime/Renderer/Private/Strata/Strata.h
+++ b/Engine/Source/Runtime/Renderer/Private/Strata/Strata.h
@@ -17,7 +17,7 @@ struct FMinimalSceneTextures;
 BEGIN_SHADER_PARAMETER_STRUCT(FStrataBasePassUniformParameters, )
 	SHADER_PARAMETER(uint32, MaxBytesPerPixel)
 	SHADER_PARAMETER(FVector2D, GGXEnergyLUTScaleBias)
-	SHADER_PARAMETER_TEXTURE(Texture3D<float2>, GGXEnergyLUT3DTexture)
+	SHADER_PARAMETER_TEXTURE(Texture3D<float3>, GGXEnergyLUT3DTexture)
 	SHADER_PARAMETER_TEXTURE(Texture2D<float4>, GGXEnergyLUT2DTexture)
 	SHADER_PARAMETER_SAMPLER(SamplerState, GGXEnergyLUTSampler)
 	SHADER_PARAMETER_RDG_BUFFER_UAV(RWByteAddressBuffer, MaterialLobesBufferUAV)
@@ -30,7 +30,7 @@ BEGIN_GLOBAL_SHADER_PARAMETER_STRUCT(FStrataGlobalUniformParameters, )
 	SHADER_PARAMETER_RDG_TEXTURE(Texture2D<uint>, ClassificationTexture)
 	SHADER_PARAMETER_RDG_TEXTURE(Texture2D<uint>, TopLayerNormalTexture)
 	SHADER_PARAMETER_RDG_TEXTURE(Texture2D<uint2>, SSSTexture)
-	SHADER_PARAMETER_TEXTURE(Texture3D<float2>, GGXEnergyLUT3DTexture)
+	SHADER_PARAMETER_TEXTURE(Texture3D<float3>, GGXEnergyLUT3DTexture)
 	SHADER_PARAMETER_TEXTURE(Texture2D<float4>, GGXEnergyLUT2DTexture)
 	SHADER_PARAMETER_SAMPLER(SamplerState, GGXEnergyLUTSampler)
 END_GLOBAL_SHADER_PARAMETER_STRUCT()
