Fixed metallic input to Slab node not being saturated before some ops are executed.

diff --git a/Engine/Shaders/Private/Strata/Strata.ush b/Engine/Shaders/Private/Strata/Strata.ush
--- a/Engine/Shaders/Private/Strata/Strata.ush
+++ b/Engine/Shaders/Private/Strata/Strata.ush
@@ -561,10 +561,11 @@ FStrataData GetStrataSlabBSDF(
 	{
 		BSDF_SETHASF90(StrataData.Layers[0].BSDFs[0], any(EdgeColor < 1.0f));
 
-		SLAB_DIFFUSEALBEDO(StrataData.Layers[0].BSDFs[0])	= BaseColor - BaseColor * Metallic;
-		SLAB_F0(StrataData.Layers[0].BSDFs[0])				= ComputeF0(Specular, BaseColor, Metallic);
-		SLAB_F90(StrataData.Layers[0].BSDFs[0])				= ComputeF90(F0, EdgeColor, Metallic);
-		SLAB_METALLIC(StrataData.Layers[0].BSDFs[0])		= Metallic;
+		const float SafeMetallic = saturate(Metallic);
+		SLAB_DIFFUSEALBEDO(StrataData.Layers[0].BSDFs[0])	= BaseColor - BaseColor * SafeMetallic;
+		SLAB_F0(StrataData.Layers[0].BSDFs[0])				= ComputeF0(Specular, BaseColor, SafeMetallic);
+		SLAB_F90(StrataData.Layers[0].BSDFs[0])				= ComputeF90(F0, EdgeColor, SafeMetallic);
+		SLAB_METALLIC(StrataData.Layers[0].BSDFs[0])		= SafeMetallic;
 	}
 	else
 	{
