Fix shader compiler error.

diff --git a/Engine/Shaders/Private/DistanceFieldObjectCulling.usf b/Engine/Shaders/Private/DistanceFieldObjectCulling.usf
--- a/Engine/Shaders/Private/DistanceFieldObjectCulling.usf
+++ b/Engine/Shaders/Private/DistanceFieldObjectCulling.usf
@@ -334,8 +334,8 @@ void ObjectCullPS(
 	
 	uint2 TilePosition = (uint2)SVPos.xy;
 	uint TileIndex = TilePosition.y * NumGroups.x + TilePosition.x;
-	float4 ConeAxisAndCos = TileConeAxisAndCos.Load(TileIndex);
-	float4 ConeAxisDepthRanges = TileConeDepthRanges.Load(TileIndex);
+	float4 ConeAxisAndCos = TileConeAxisAndCos[TileIndex];
+	float4 ConeAxisDepthRanges = TileConeDepthRanges[TileIndex];
 	float3 TileConeVertex = 0;
 	float3 TileConeAxis = ConeAxisAndCos.xyz;
 	float TileConeAngleCos = ConeAxisAndCos.w;
