Fixed RGB<->YCoCg to not result in F90 with component being higher than 1.0 and different chroma.
To this aim, a Y value is set to enforce that and a better match is obtained by devising again the recovered RGB color by the max of its component (to ensure high brightness and saturation). We may revisit that later

diff --git a/Engine/Shaders/Private/Strata/Strata.ush b/Engine/Shaders/Private/Strata/Strata.ush
--- a/Engine/Shaders/Private/Strata/Strata.ush
+++ b/Engine/Shaders/Private/Strata/Strata.ush
@@ -1929,7 +1929,12 @@ void PackStrataOut(
 
 				if (BSDF_GETHASF90(BSDF) || BSDF_GETHASTHINFILM(BSDF) || BSDF_GETHASHAZINESS(BSDF))
 				{
-					float3 F90YCoCg = LinearRGB_2_NormalisedYCoCg(SLAB_F90(BSDF));
+					// What is important is to maintain the hue and saturation, so we scale the color by the maximum of its components
+					float3 F90 = saturate(SLAB_F90(BSDF));
+					const float Divisor = max(F90.r, max(F90.g, F90.b));
+					F90 = Divisor > 0.0f ? F90 / Divisor : 1.0f;
+					float3 F90YCoCg = LinearRGB_2_NormalisedYCoCg(F90);
+
 					const float ThinFilmGamma2 = sqrt(SLAB_THINFILM_THICKNESS(BSDF)); // Using a gamma 2.0 encoding help with precision to be visually more uniform. This is then further improved via dithering.
 					STRATA_STORE_UINT1(PackRGBA8(float4(F90YCoCg.y, F90YCoCg.z, ThinFilmGamma2, SLAB_HAZINESS(BSDF))));
 				}
@@ -2125,8 +2130,15 @@ FStrataBSDF UnpackStrataBSDFIn(ByteAddressBuffer StrataBuffer, inout FStrataAddr
 
 			const float F90Co = Data.x;
 			const float F90Cg = Data.y;
-			const float F90Y  = 1.0f;	// Forced to 1 as F90 is not supposed to have a low value
-			SLAB_F90(OutBSDF) = bHasF90 ? NormalisedYCoCg_2_LinearRGB(float3(F90Y, F90Co, F90Cg)) : 1.0f;
+			// Setting Y=1 when converting back from YCoCg can lead to components being greater than 1. We fix this using two steps:
+			// 1- We apply a simple scale determined manually on Y to keep Hue/Saturation expressivity large and similare to the input color from the base pass,
+			//    even though the mapping from saturate to white is not linear.
+			const float F90Y  = 1.0f / 3.0f; 
+			float3 F90 = bHasF90 ? saturate(NormalisedYCoCg_2_LinearRGB(float3(F90Y, F90Co, F90Cg))) : 1.0f;
+			// 2- We re-scale again the final recovered color according to its max component to maximise the brightness/hue similarity with the input color from the base pass.
+			const float Divisor = max(F90.r, max(F90.g, F90.b));
+			F90 = Divisor > 0.0f ? F90 / Divisor : F90;
+			SLAB_F90(OutBSDF) = F90;
 
 			SLAB_THINFILM_THICKNESS(OutBSDF) = Data.z * Data.z;
 			SLAB_HAZINESS(OutBSDF) = Data.w;
