Support for refraction/distortion to strata material the same way it does with current materials. (basic first version, no rough refraction)

diff --git a/Engine/Shaders/Private/DistortAccumulatePS.usf b/Engine/Shaders/Private/DistortAccumulatePS.usf
--- a/Engine/Shaders/Private/DistortAccumulatePS.usf
+++ b/Engine/Shaders/Private/DistortAccumulatePS.usf
@@ -20,6 +20,8 @@
 #include "/Engine/Generated/VertexFactory.ush"
 #include "DistortionCommon.ush"
 
+#define STRATA_ENABLED PROJECT_STRATA && MATERIAL_IS_STRATA
+
 void ClipOccludedPixel(float PixelZ, float SceneZ)
 {
 #if	HAS_INVERTED_Z_BUFFER
@@ -62,6 +64,34 @@ void Main(
 	
 	// material distortion offset
 	half3 Normal = GetMaterialNormal(MaterialParameters, PixelMaterialInputs);
+#if STRATA_ENABLED
+	FStrataData StrataData = PixelMaterialInputs.FrontMaterial;
+
+	// Use the normal from the top layer to apply distortion
+	// STRATA_TODO compute distortion based on layers, normals and scattering from distribution to achieve rough refraction.
+	if (StrataData.LayerCount > 0)
+	{
+		Normal = 0;
+
+		FStrataLayer Layer = StrataData.Layers[0];
+		for (int i = 0; i < Layer.BSDFCount; ++i)
+		{
+			Normal += Layer.BSDFs[i].Coverage * MaterialParameters.SharedLocalBases.Normals[BSDF_GETSHAREDLOCALBASISID(Layer.BSDFs[i])]
+			#if MATERIAL_TANGENTSPACENORMAL
+				* MaterialParameters.TwoSidedSign
+			#endif
+				;
+		}
+	}
+
+#if MATERIAL_TANGENTSPACENORMAL
+	// Transform back to tangent space as this is what the common function ComputeBufferUVDistortion expects
+	Normal = normalize(TransformWorldVectorToTangent(MaterialParameters.TangentToWorld, Normal));
+#else
+	Normal = normalize(Normal);
+#endif
+
+#endif
 
 	// Prevent silhouettes from geometry that is in front of distortion from being seen in the distortion 
 	float2 NDC = (MaterialParameters.ScreenPosition.xy / MaterialParameters.ScreenPosition.w);
