Fix minor int vs uint shader warning

diff --git a/Engine/Shaders/Private/MaterialTemplate.ush b/Engine/Shaders/Private/MaterialTemplate.ush
--- a/Engine/Shaders/Private/MaterialTemplate.ush
+++ b/Engine/Shaders/Private/MaterialTemplate.ush
@@ -2138,7 +2138,7 @@ float3 DecodeSceneColorForMaterialNode(float2 ScreenUV)
 float4 MaterialExpressionDBufferTextureLookup(float2 BufferUV, int DBufferTextureIndex)
 {
 	uint2 PixelPos = uint2(BufferUV * View.BufferSizeAndInvSize.xy);
-	uint ValidDBufferTargetMask = GetDBufferTargetMask(PixelPos) & (1 << DBufferTextureIndex);
+	uint ValidDBufferTargetMask = GetDBufferTargetMask(PixelPos) & (1u << DBufferTextureIndex);
 
 	FDBufferData DBufferData = GetDBufferData(BufferUV, ValidDBufferTargetMask);
 	switch (DBufferTextureIndex)
