Make shader draw / shader print more consistent
 - add IsSupported to ShaderDebug API
 - cleaned up some potentially unsafe uses of IsSupported/IsEnabled to guards shader parameter setting
 - added AddOBB to ShaderDrawDebug
 - Removed restriction that limited shaderdraw to editor-only

diff --git a/Engine/Shaders/Private/ShaderDrawDebug.ush b/Engine/Shaders/Private/ShaderDrawDebug.ush
--- a/Engine/Shaders/Private/ShaderDrawDebug.ush
+++ b/Engine/Shaders/Private/ShaderDrawDebug.ush
@@ -202,4 +202,44 @@ void AddLineTriangle(float3 P0, float3 P1, float3 P2, float4 Color)
 	AddLine(P1, P2, Color, Color);
 	AddLine(P2, P0, Color, Color);
 }
+
+void AddOBB(float3 Min, float3 Max, float4 Color, float4x4 Transform)
+{
+	uint PreviousPrimitiveCount = 0;
+	InterlockedAdd(OutputShaderDrawIndirect[1], 12, PreviousPrimitiveCount);
+	if ((PreviousPrimitiveCount + 12) <= ShaderDrawMaxElementCount)
+	{
+		float3 P0 = mul(float4(Min.x, Min.y, Min.z, 1.0f), Transform).xyz;
+		float3 P1 = mul(float4(Max.x, Min.y, Min.z, 1.0f), Transform).xyz;
+		float3 P2 = mul(float4(Max.x, Max.y, Min.z, 1.0f), Transform).xyz;
+		float3 P3 = mul(float4(Min.x, Max.y, Min.z, 1.0f), Transform).xyz;
+		float3 P4 = mul(float4(Min.x, Min.y, Max.z, 1.0f), Transform).xyz;
+		float3 P5 = mul(float4(Max.x, Min.y, Max.z, 1.0f), Transform).xyz;
+		float3 P6 = mul(float4(Max.x, Max.y, Max.z, 1.0f), Transform).xyz;
+		float3 P7 = mul(float4(Min.x, Max.y, Max.z, 1.0f), Transform).xyz;
+
+		FShaderDrawElement Element;
+		Element.Color0 = Element.Color1 = Color;
+
+		Element.Pos0 = P0; Element.Pos1 = P1; OutShaderDrawPrimitive[PreviousPrimitiveCount + 0] = PackShaderElement(Element);
+		Element.Pos0 = P1; Element.Pos1 = P2; OutShaderDrawPrimitive[PreviousPrimitiveCount + 1] = PackShaderElement(Element);
+		Element.Pos0 = P2; Element.Pos1 = P3; OutShaderDrawPrimitive[PreviousPrimitiveCount + 2] = PackShaderElement(Element);
+		Element.Pos0 = P3; Element.Pos1 = P0; OutShaderDrawPrimitive[PreviousPrimitiveCount + 3] = PackShaderElement(Element);
+
+		Element.Pos0 = P4; Element.Pos1 = P5; OutShaderDrawPrimitive[PreviousPrimitiveCount + 4] = PackShaderElement(Element);
+		Element.Pos0 = P5; Element.Pos1 = P6; OutShaderDrawPrimitive[PreviousPrimitiveCount + 5] = PackShaderElement(Element);
+		Element.Pos0 = P6; Element.Pos1 = P7; OutShaderDrawPrimitive[PreviousPrimitiveCount + 6] = PackShaderElement(Element);
+		Element.Pos0 = P7; Element.Pos1 = P4; OutShaderDrawPrimitive[PreviousPrimitiveCount + 7] = PackShaderElement(Element);
+
+		Element.Pos0 = P0; Element.Pos1 = P4; OutShaderDrawPrimitive[PreviousPrimitiveCount + 8] = PackShaderElement(Element);
+		Element.Pos0 = P1; Element.Pos1 = P5; OutShaderDrawPrimitive[PreviousPrimitiveCount + 9] = PackShaderElement(Element);
+		Element.Pos0 = P2; Element.Pos1 = P6; OutShaderDrawPrimitive[PreviousPrimitiveCount + 10] = PackShaderElement(Element);
+		Element.Pos0 = P3; Element.Pos1 = P7; OutShaderDrawPrimitive[PreviousPrimitiveCount + 11] = PackShaderElement(Element);
+	}
+	else
+	{
+		InterlockedAdd(OutputShaderDrawIndirect[1], -12, PreviousPrimitiveCount);// Make sure we stay in the safe zone for indirect draw
+	}
+}
+
 #endif
diff --git a/Engine/Source/Runtime/Renderer/Private/ClusteredDeferredShadingPass.cpp b/Engine/Source/Runtime/Renderer/Private/ClusteredDeferredShadingPass.cpp
--- a/Engine/Source/Runtime/Renderer/Private/ClusteredDeferredShadingPass.cpp
+++ b/Engine/Source/Runtime/Renderer/Private/ClusteredDeferredShadingPass.cpp
@@ -169,15 +169,8 @@ static void InternalAddClusteredDeferredShadingPass(
 	PassParameters->SSProfilesTexture = GetSubsufaceProfileTexture_RT(GraphBuilder.RHICmdList)->GetShaderResourceRHI();
 	PassParameters->TransmissionProfilesLinearSampler = TStaticSamplerState<SF_Bilinear, AM_Clamp, AM_Clamp, AM_Clamp>::GetRHI();
 
-	if (ShaderDrawDebug::IsEnabled(View))
-	{
-		ShaderDrawDebug::SetParameters(GraphBuilder, View.ShaderDrawData, PassParameters->ShaderDrawParameters);
-	}
-
-	if (ShaderPrint::IsSupported(View) && ShaderPrint::IsEnabled())
-	{
-		ShaderPrint::SetParameters(GraphBuilder, View, PassParameters->ShaderPrintUniformBuffer);
-	}
+	ShaderDrawDebug::SetParameters(GraphBuilder, View.ShaderDrawData, PassParameters->ShaderDrawParameters);
+	ShaderPrint::SetParameters(GraphBuilder, View, PassParameters->ShaderPrintUniformBuffer);
 
 	PassParameters->RenderTargets[0] = FRenderTargetBinding(SceneTextures.Color.Target, ERenderTargetLoadAction::ELoad);
 	if (bHairStrands)
diff --git a/Engine/Source/Runtime/Renderer/Private/HairStrands/HairStrandsDebug.cpp b/Engine/Source/Runtime/Renderer/Private/HairStrands/HairStrandsDebug.cpp
--- a/Engine/Source/Runtime/Renderer/Private/HairStrands/HairStrandsDebug.cpp
+++ b/Engine/Source/Runtime/Renderer/Private/HairStrands/HairStrandsDebug.cpp
@@ -75,7 +75,7 @@ static FAutoConsoleVariableRef CVarHairStrandsUpdateCullIndex(TEXT("r.HairStrand
 
 static bool IsDebugDrawAndDebugPrintEnabled(const FViewInfo& View)
 {
-	return ShaderDrawDebug::IsEnabled() && ShaderPrint::IsEnabled() && ShaderPrint::IsSupported(View);
+	return ShaderDrawDebug::IsEnabled(View) && ShaderPrint::IsEnabled(View);
 }
 
 bool IsHairStrandsClusterDebugEnable()
diff --git a/Engine/Source/Runtime/Renderer/Private/HairStrands/HairStrandsTransmittance.cpp b/Engine/Source/Runtime/Renderer/Private/HairStrands/HairStrandsTransmittance.cpp
--- a/Engine/Source/Runtime/Renderer/Private/HairStrands/HairStrandsTransmittance.cpp
+++ b/Engine/Source/Runtime/Renderer/Private/HairStrands/HairStrandsTransmittance.cpp
@@ -257,15 +257,8 @@ static FRDGBufferRef AddHairStrandsVoxelTransmittanceMaskPass(
 	Parameters->HairStrands = HairStrands::BindHairStrandsViewUniformParameters(View);
 	Parameters->VirtualVoxel = HairStrands::BindHairStrandsVoxelUniformParameters(View);
 
-	if (ShaderDrawDebug::IsEnabled(View))
-	{
-		ShaderDrawDebug::SetParameters(GraphBuilder, View.ShaderDrawData, Parameters->ShaderDrawParameters);
-	}
-
-	if (ShaderPrint::IsSupported(View) && ShaderPrint::IsEnabled())
-	{
-		ShaderPrint::SetParameters(GraphBuilder, View, Parameters->ShaderPrintUniformBuffer);
-	}
+	ShaderDrawDebug::SetParameters(GraphBuilder, View.ShaderDrawData, Parameters->ShaderDrawParameters);
+	ShaderPrint::SetParameters(GraphBuilder, View, Parameters->ShaderPrintUniformBuffer);
 
 	const bool bIsSuperSampled = GHairStrandsTransmittanceSuperSampling > 0;
 	const bool bIsMipTraversal = GHairStrandsTransmittanceMaskUseMipTraversal > 0;
diff --git a/Engine/Source/Runtime/Renderer/Private/PhysicsFieldRendering.cpp b/Engine/Source/Runtime/Renderer/Private/PhysicsFieldRendering.cpp
--- a/Engine/Source/Runtime/Renderer/Private/PhysicsFieldRendering.cpp
+++ b/Engine/Source/Runtime/Renderer/Private/PhysicsFieldRendering.cpp
@@ -183,8 +183,7 @@ void RenderPhysicsField(
 	{
 		FPhysicsFieldResource* PhysicsFieldResource = (GPhysicsFieldSystemType == 0) ? PhysicsFieldProxy->DebugResource : PhysicsFieldProxy->FieldResource;
 
-		if (Views.Num() > 0 && PhysicsFieldResource && ShaderDrawDebug::IsEnabled() && ShaderPrint::IsEnabled() &&
-			ShaderPrint::IsSupported(Views[0]) && GDynamicRHI->RHIIsTypedUAVLoadSupported(PF_FloatRGBA))
+		if (Views.Num() > 0 && PhysicsFieldResource && ShaderDrawDebug::IsEnabled() && ShaderPrint::IsEnabled(Views[0]) && GDynamicRHI->RHIIsTypedUAVLoadSupported(PF_FloatRGBA))
 		{
 			AddPhysicsFieldRayMarchingPass(GraphBuilder, Views[0], PhysicsFieldResource, SceneColorTexture);
 		}
diff --git a/Engine/Source/Runtime/Renderer/Private/PostProcess/PostProcessing.cpp b/Engine/Source/Runtime/Renderer/Private/PostProcess/PostProcessing.cpp
--- a/Engine/Source/Runtime/Renderer/Private/PostProcess/PostProcessing.cpp
+++ b/Engine/Source/Runtime/Renderer/Private/PostProcess/PostProcessing.cpp
@@ -1141,7 +1141,7 @@ void AddPostProcessingPasses(FRDGBuilder& GraphBuilder, const FViewInfo& View, c
 		ShaderDrawDebug::DrawView(GraphBuilder, View, SceneColor.Texture, SceneDepth.Texture);
 	}
 
-	if (ShaderPrint::IsEnabled() && ShaderPrint::IsSupported(View))
+	if (ShaderPrint::IsEnabled(View))
 	{
 		ShaderPrint::DrawView(GraphBuilder, View, SceneColor.Texture);
 	}
diff --git a/Engine/Source/Runtime/Renderer/Private/ShaderDebug.cpp b/Engine/Source/Runtime/Renderer/Private/ShaderDebug.cpp
--- a/Engine/Source/Runtime/Renderer/Private/ShaderDebug.cpp
+++ b/Engine/Source/Runtime/Renderer/Private/ShaderDebug.cpp
@@ -34,53 +34,37 @@ namespace ShaderDrawDebug
 
 	bool IsEnabled()
 	{
-#if WITH_EDITOR
 		return GShaderDrawDebug_Enable > 0;
-#else
-		return false;
-#endif
 	}
 
 	static bool IsShaderDrawLocked()
 	{
-#if WITH_EDITOR
 		return CVarShaderDrawLock.GetValueOnRenderThread() > 0;
-#else
-		return false;
-#endif
 	}
 
-	static bool IsEnabled(const EShaderPlatform Platform)
+	bool IsSupported(const EShaderPlatform Platform)
 	{
 		return IsFeatureLevelSupported(Platform, ERHIFeatureLevel::SM5) && IsPCPlatform(Platform) && !IsOpenGLPlatform(Platform);
 	}
 
 	void SetEnabled(bool bInEnabled)
 	{
-#if WITH_EDITOR
 		GShaderDrawDebug_Enable = bInEnabled ? 1 : 0;
-#endif		
 	}
 
 	void SetMaxElementCount(uint32 MaxCount)
 	{
-#if WITH_EDITOR
 		GShaderDrawDebug_MaxElementCount = FMath::Max(1024, int32(MaxCount));
-#endif		
 	}
 
 	uint32 GetMaxElementCount()
 	{
-#if WITH_EDITOR
 		return uint32(FMath::Max(1, GShaderDrawDebug_MaxElementCount));
-#else
-		return 0;
-#endif			
 	}
 
 	bool IsEnabled(const FViewInfo& View)
 	{
-		return IsEnabled() && IsEnabled(View.GetShaderPlatform());
+		return IsEnabled() && IsSupported(View.GetShaderPlatform());
 	}
 
 	// Note: Unaligned structures used for structured buffers is an unsupported and/or sparsely
@@ -107,7 +91,7 @@ namespace ShaderDrawDebug
 
 		static bool ShouldCompilePermutation(const FGlobalShaderPermutationParameters& Parameters)
 		{
-			return IsEnabled(Parameters.Platform);
+			return IsSupported(Parameters.Platform);
 		}
 
 		static void ModifyCompilationEnvironment(const FGlobalShaderPermutationParameters& Parameters, FShaderCompilerEnvironment& OutEnvironment)
@@ -136,7 +120,7 @@ namespace ShaderDrawDebug
 
 		static bool ShouldCompilePermutation(const FGlobalShaderPermutationParameters& Parameters)
 		{
-			return IsEnabled(Parameters.Platform);
+			return IsSupported(Parameters.Platform);
 		}
 
 		static void ModifyCompilationEnvironment(const FGlobalShaderPermutationParameters& Parameters, FShaderCompilerEnvironment& OutEnvironment)
@@ -169,7 +153,7 @@ namespace ShaderDrawDebug
 
 		static bool ShouldCompilePermutation(const FGlobalShaderPermutationParameters& Parameters)
 		{
-			return IsEnabled(Parameters.Platform);
+			return IsSupported(Parameters.Platform);
 		}
 
 		static void ModifyCompilationEnvironment(const FGlobalShaderPermutationParameters& Parameters, FShaderCompilerEnvironment& OutEnvironment)
@@ -249,7 +233,7 @@ namespace ShaderDrawDebug
 
 	void BeginView(FRDGBuilder& GraphBuilder, FViewInfo& View)
 	{
-		if (!IsEnabled(View) || !IsEnabled())
+		if (!IsEnabled(View))
 		{
 			return;
 		}
diff --git a/Engine/Source/Runtime/Renderer/Private/ShaderPrint.cpp b/Engine/Source/Runtime/Renderer/Private/ShaderPrint.cpp
--- a/Engine/Source/Runtime/Renderer/Private/ShaderPrint.cpp
+++ b/Engine/Source/Runtime/Renderer/Private/ShaderPrint.cpp
@@ -108,11 +108,6 @@ namespace ShaderPrint
 		return RHISupportsComputeShaders(InShaderPlatform) && !IsHlslccShaderPlatform(InShaderPlatform);
 	}
 
-	bool IsSupported(FViewInfo const& View)
-	{
-		return IsSupported(View.GetShaderPlatform());
-	}
-
 	void SetEnabled(bool bInEnabled)
 	{
 		CVarEnable->Set(bInEnabled);
@@ -133,6 +128,11 @@ namespace ShaderPrint
 		return CVarEnable.GetValueOnAnyThread() != 0;
 	}
 
+	bool IsEnabled(const FViewInfo& View)
+	{
+		return IsEnabled() && IsSupported(View.GetShaderPlatform());
+	}
+
 	// Shader to initialize the output value buffer
 	class FShaderInitValueBufferCS : public FGlobalShader
 	{
@@ -270,7 +270,7 @@ namespace ShaderPrint
 	void BeginView(FRDGBuilder& GraphBuilder, FViewInfo& View)
 	{
 		TRACE_CPUPROFILER_EVENT_SCOPE(ShaderPrint::BeginView);
-		if (!IsSupported(View))
+		if (!IsSupported(View.GetShaderPlatform()))
 		{
 			return;
 		}
diff --git a/Engine/Source/Runtime/Renderer/Private/ShaderPrint.h b/Engine/Source/Runtime/Renderer/Private/ShaderPrint.h
--- a/Engine/Source/Runtime/Renderer/Private/ShaderPrint.h
+++ b/Engine/Source/Runtime/Renderer/Private/ShaderPrint.h
@@ -18,9 +18,11 @@ class FViewInfo;
 namespace ShaderPrint
 {
 	// Does the platform support the ShaderPrint system?
-	bool IsSupported(FViewInfo const& View);
+	bool IsSupported(const EShaderPlatform Platform);
 	// Have we enabled the ShaderPrint system?
 	bool IsEnabled();
+	// Call this to know if a view can render this debug information
+	bool IsEnabled(const FViewInfo& View);
 
 	// Accessors for controlling from code: enabling/disabling, font size, and max drawn element
 	void SetEnabled(bool bInEnabled);
diff --git a/Engine/Source/Runtime/Renderer/Public/ShaderDebug.h b/Engine/Source/Runtime/Renderer/Public/ShaderDebug.h
--- a/Engine/Source/Runtime/Renderer/Public/ShaderDebug.h
+++ b/Engine/Source/Runtime/Renderer/Public/ShaderDebug.h
@@ -19,6 +19,10 @@ namespace ShaderDrawDebug
 {
 	// Call this to know if this is even just available (for exemple in shipping mode buffers won't exists)
 	RENDERER_API bool IsEnabled();
+
+	// Use to disable permutations that should not compile as the shaderdraw is unsupported.
+	RENDERER_API bool IsSupported(const EShaderPlatform Platform);
+
 	RENDERER_API void SetEnabled(bool bEnable);
 	RENDERER_API void SetMaxElementCount(uint32 MaxCount);
 	RENDERER_API uint32 GetMaxElementCount();
