1. Unify the lighting calculation for both mobile deferred and forward.
2. Add an option to use high quality BRDF on mobile as the same with PC.
3. Fallback to default light shading model if it is not supported on mobile.
4. Use EnvBrdf for mobile deferred lighting pass.

diff --git a/Engine/Shaders/Private/MobileBasePassPixelShader.usf b/Engine/Shaders/Private/MobileBasePassPixelShader.usf
--- a/Engine/Shaders/Private/MobileBasePassPixelShader.usf
+++ b/Engine/Shaders/Private/MobileBasePassPixelShader.usf
@@ -112,14 +112,14 @@ int NumDynamicPointLights;
 #include "MobileShadingModels.ush"
 #include "DynamicLightingCommon.ush"
 #include "PlanarReflectionShared.ush"
+#include "ShadingModelsMaterial.ush"
 
 #if MATERIAL_SHADINGMODEL_SINGLELAYERWATER
 	#include "SingleLayerWaterShading.ush"
 #endif
 
 /** Prenormalized capture of the scene that's closest to the object being rendered. */
-#if !FULLY_ROUGH
-	#if HQ_REFLECTIONS
+#if HQ_REFLECTIONS
 	#define MAX_HQ_REFLECTIONS 3
 	TextureCube ReflectionCubemap0;
 	SamplerState ReflectionCubemapSampler0;
@@ -135,7 +135,6 @@ int NumDynamicPointLights;
 		float4x4 CaptureBoxTransformArray[MAX_HQ_REFLECTIONS];
 		float4 CaptureBoxScalesArray[MAX_HQ_REFLECTIONS];
 		#endif
-	#endif
 #endif
 
 #ifndef MOBILE_USE_CSM_BRANCH
@@ -155,7 +154,6 @@ bool IsCSMEnabled()
 #endif
 }
 
-#if !FULLY_ROUGH
 // Half precision version for mobile's high quality reflection
 half MobileComputeMixingWeight(half IndirectIrradiance, half AverageBrightness, half Roughness)
 {
@@ -371,7 +369,6 @@ half3 GetImageBasedReflectionLighting(FMaterialPixelParameters MaterialParameter
 	return ImageBasedReflections.rgb;
 #endif
 }
-#endif //!FULLY_ROUGH
 
 half3 FrameBufferBlendOp(half4 Source)
 {
@@ -418,12 +415,21 @@ float4 SpotLightShadowSharpenAndShadowFadeFraction,
 float4 SpotLightShadowmapMinMax,
 float4x4 SpotLightShadowWorldToShadowMatrix,
 #endif
-inout half3 Color)
+inout FMobileLightAccumulator MobileLightAccumulator)
 {
 	uint LightShadowType = SpotLightAnglesAndSoftTransitionScaleAndLightShadowType.w;
 
 	float FadedShadow = 1.0f;
 
+	float3 ToLight = LightPositionAndInvRadius.xyz - MaterialParameters.AbsoluteWorldPosition;
+	float DistanceSqr = dot(ToLight, ToLight);
+	float3 L = ToLight * rsqrt(DistanceSqr);
+	half3 PointH = normalize(MaterialParameters.CameraVector + L);
+
+	half PointNoL = max(0, dot(MaterialParameters.WorldNormal, L));
+	half PointNoH = max(0, dot(MaterialParameters.WorldNormal, PointH));
+	half PointVoH = max(0, dot(MaterialParameters.CameraVector, PointH));
+
 #if SUPPORT_SPOTLIGHTS_SHADOW
 		if ((LightShadowType & LightShadowType_Shadow) == LightShadowType_Shadow)
 		{
@@ -438,7 +444,7 @@ inout half3 Color)
 				float ShadowZ = 1.0f - HomogeneousShadowPosition.z;
 				float LightSpacePixelDepthForOpaque = min(ShadowZ, 0.99999f);
 				Settings.SceneDepth = LightSpacePixelDepthForOpaque;
-				Settings.TransitionScale = SpotLightAnglesAndSoftTransitionScaleAndLightShadowType.z;
+				Settings.TransitionScale = SpotLightAnglesAndSoftTransitionScaleAndLightShadowType.z * lerp(SpotLightShadowSharpenAndShadowFadeFraction.z, 1.0, PointNoL);
 
 				half Shadow = MobileShadowPCF(ShadowUVs, Settings);
 
@@ -451,14 +457,6 @@ inout half3 Color)
 
 	if ((LightShadowType & ValidLightType) != 0)
 	{
-		float3 ToLight = LightPositionAndInvRadius.xyz - MaterialParameters.AbsoluteWorldPosition;
-		float DistanceSqr = dot(ToLight, ToLight);
-		float3 L = ToLight * rsqrt(DistanceSqr);
-		half3 PointH = normalize(MaterialParameters.CameraVector + L);
-
-		half PointNoL = max(0, dot(MaterialParameters.WorldNormal, L));
-		half PointNoH = max(0, dot(MaterialParameters.WorldNormal, PointH));
-
 		#if TRANSLUCENCY_NON_DIRECTIONAL
 			PointNoL = 1.f;
 		#endif
@@ -481,16 +479,13 @@ inout half3 Color)
 #if PROJECT_MOBILE_ENABLE_MOVABLE_SPOTLIGHTS
 		if ((LightShadowType & LightShadowType_SpotLight) == LightShadowType_SpotLight)
 		{
-			Attenuation *= SpotAttenuation(L, -SpotLightDirectionAndSpecularScale.xyz, SpotLightAnglesAndSoftTransitionScaleAndLightShadowType.xy) * FadedShadow;
+			Attenuation *= SpotAttenuation(L, -SpotLightDirectionAndSpecularScale.xyz, SpotLightAnglesAndSoftTransitionScaleAndLightShadowType.xy);
 		}
 #endif
 
-#if !FULLY_ROUGH
-		FMobileDirectLighting Lighting = MobileIntegrateBxDF(ShadingModelContext, GBuffer, PointNoL, MaterialParameters.CameraVector, PointH, PointNoH);
-		Color += min(65000.0, (Attenuation) * LightColorAndFalloffExponent.rgb * (1.0 / PI) * (Lighting.Diffuse + Lighting.Specular * SpotLightDirectionAndSpecularScale.w));
-#else
-		Color += (Attenuation * PointNoL) * LightColorAndFalloffExponent.rgb * (1.0 / PI) * ShadingModelContext.DiffuseColor;
-#endif
+		FMobileDirectLighting Lighting = MobileIntegrateBxDF(ShadingModelContext, GBuffer, PointNoL, PointNoH, PointVoH, MaterialParameters.CameraVector, L, FadedShadow);
+
+		MobileLightAccumulator_Add(MobileLightAccumulator, Lighting.Diffuse, Lighting.Specular * SpotLightDirectionAndSpecularScale.w, (Attenuation)*LightColorAndFalloffExponent.rgb * (1.0 / PI));
 	}
 }
 #endif
@@ -502,11 +497,11 @@ inout half3 Color)
 	#define PIXELSHADER_EARLYDEPTHSTENCIL 	
 #endif
 
-half ComputeIndirect(VTPageTableResult LightmapVTPageTableResult, FVertexFactoryInterpolantsVSToPS Interpolants, float3 DiffuseDir, FMobileShadingModelContext ShadingModelContext, out half IndirectIrradiance, out half3 Color)
+void ComputeIndirect(VTPageTableResult LightmapVTPageTableResult, FVertexFactoryInterpolantsVSToPS Interpolants, float3 DiffuseDir, half3 DiffuseColorForIndirect, half MaterialAO, out half IndirectIrradiance, inout FMobileLightAccumulator MobileLightAccumulator)
 {
 	//To keep IndirectLightingCache conherence with PC, initialize the IndirectIrradiance to zero.
 	IndirectIrradiance = 0;
-	Color = 0;
+	half3 IndirectColor = 0;
 
 	// Indirect Diffuse
 #if LQ_TEXTURE_LIGHTMAP
@@ -515,7 +510,7 @@ half ComputeIndirect(VTPageTableResult LightmapVTPageTableResult, FVertexFactory
 	GetLightMapCoordinates(Interpolants, LightmapUV0, LightmapUV1, LightmapDataIndex);
 
 	half4 LightmapColor = GetLightMapColorLQ(LightmapVTPageTableResult, LightmapUV0, LightmapUV1, LightmapDataIndex, DiffuseDir);
-	Color += LightmapColor.rgb * ShadingModelContext.DiffuseColor * View.IndirectLightingColorScale;
+	IndirectColor += LightmapColor.rgb * DiffuseColorForIndirect * View.IndirectLightingColorScale;
 	IndirectIrradiance = LightmapColor.a;
 #elif CACHED_POINT_INDIRECT_LIGHTING
 	#if MATERIALBLENDING_MASKED || MATERIALBLENDING_SOLID
@@ -539,7 +534,7 @@ half ComputeIndirect(VTPageTableResult LightmapVTPageTableResult, FVertexFactory
 		half3 DiffuseGI = max(half3(0, 0, 0), DotSH3(PointIndirectLighting, DiffuseTransferSH));
 
 		IndirectIrradiance = Luminance(DiffuseGI);
-		Color += ShadingModelContext.DiffuseColor * DiffuseGI * View.IndirectLightingColorScale;
+		IndirectColor += DiffuseColorForIndirect * DiffuseGI * View.IndirectLightingColorScale;
 	#else 
 		// Non-directional for translucency
 		// Ambient terms packed in xyz
@@ -548,11 +543,12 @@ half ComputeIndirect(VTPageTableResult LightmapVTPageTableResult, FVertexFactory
 		half3 DiffuseGI = PointIndirectLighting;
 
 		IndirectIrradiance = Luminance(DiffuseGI);
-		Color += ShadingModelContext.DiffuseColor * DiffuseGI * View.IndirectLightingColorScale;
+		IndirectColor += DiffuseColorForIndirect * DiffuseGI * View.IndirectLightingColorScale;
 	#endif
 #endif
+	IndirectIrradiance *= MaterialAO;
 
-	return IndirectIrradiance;
+	MobileLightAccumulator_Add(MobileLightAccumulator, IndirectColor, 0.0f, MaterialAO);
 }
 
 PIXELSHADER_EARLYDEPTHSTENCIL
@@ -563,19 +559,19 @@ void Main(
 	OPTIONAL_IsFrontFace
 #if DEFERRED_SHADING_PATH
 #if !USE_GLES_FBF_DEFERRED
-	, out HALF4_TYPE OutColor		: SV_Target0
-	, out HALF4_TYPE OutGBufferA	: SV_Target1
-	, out HALF4_TYPE OutGBufferB	: SV_Target2
+	, out HALF4_TYPE OutColor : SV_Target0
+	, out HALF4_TYPE OutGBufferA : SV_Target1
+	, out HALF4_TYPE OutGBufferB : SV_Target2
 	, out HALF4_TYPE OutGBufferC	: SV_Target3  
 #endif
 #else
-	, out HALF4_TYPE OutColor		: SV_Target0
+	, out HALF4_TYPE OutColor : SV_Target0
 #endif
 #if USE_SCENE_DEPTH_AUX
 	, out float OutSceneDepthAux : SV_TargetDepthAux
 #endif
 #if OUTPUT_PIXEL_DEPTH_OFFSET
-    , out float OutDepth : SV_Depth
+	, out float OutDepth : SV_Depth
 #endif
 	)
 {
@@ -615,7 +611,7 @@ void Main(
 		float3 WorldPosition = BasePassInterpolants.PixelPosition.xyz;
 		float3 WorldPositionExcludingWPO = BasePassInterpolants.PixelPosition.xyz;
 		#if USE_WORLD_POSITION_EXCLUDING_SHADER_OFFSETS
-			WorldPositionExcludingWPO = BasePassInterpolants.PixelPositionExcludingWPO;
+		WorldPositionExcludingWPO = BasePassInterpolants.PixelPositionExcludingWPO;
 		#endif
 		CalcMaterialParametersEx(MaterialParameters, PixelMaterialInputs, SvPosition, ScreenPosition, bIsFrontFace, WorldPosition, WorldPositionExcludingWPO);
 
@@ -636,15 +632,26 @@ void Main(
 #endif
 
 	// Store the results in local variables and reuse instead of calling the functions multiple times.
-	FGBufferData GBuffer = (FGBufferData)0;
-	GBuffer.WorldNormal = MaterialParameters.WorldNormal;
-	GBuffer.BaseColor = GetMaterialBaseColor(PixelMaterialInputs);
-	GBuffer.Metallic = GetMaterialMetallic(PixelMaterialInputs);
-	GBuffer.Specular = GetMaterialSpecular(PixelMaterialInputs);
-	GBuffer.Roughness = GetMaterialRoughness(PixelMaterialInputs);
-	GBuffer.ShadingModelID = GetMaterialShadingModel(PixelMaterialInputs);
+	half3 BaseColor = GetMaterialBaseColor(PixelMaterialInputs);
+	half Metallic = GetMaterialMetallic(PixelMaterialInputs);
+	half Specular = GetMaterialSpecular(PixelMaterialInputs);
+	// The smallest normalized value that can be represented in IEEE 754 (FP16) is 2^-24 = 5.96e-8.
+	// The code will make the following computation involving roughness: 1.0 / Roughness^4.
+	// Therefore to prevent division by zero on devices that do not support denormals, Roughness^4
+	// must be >= 5.96e-8. We will clamp to 0.015625 because 0.015625^4 = 5.96e-8.
+	//
+	// Note that we also clamp to 1.0 to match the deferred renderer on PC where the roughness is 
+	// stored in an 8-bit value and thus automatically clamped at 1.0.
+	half Roughness = max(0.015625f, GetMaterialRoughness(PixelMaterialInputs));
+	half Anisotropy = GetMaterialAnisotropy(PixelMaterialInputs);
+	uint ShadingModelID = GetMaterialShadingModel(PixelMaterialInputs);
+	half Opacity = GetMaterialOpacity(PixelMaterialInputs);
 	half MaterialAO = GetMaterialAmbientOcclusion(PixelMaterialInputs);
-
+	
+	// If we don't use this shading model the color should be black (don't generate shader code for unused data, don't do indirectlighting cache lighting with this color).
+	float3 SubsurfaceColor = 0;
+	// 0..1, SubsurfaceProfileId = int(x * 255)
+	float SubsurfaceProfile = 0;
 
 #if APPLY_AO
 	half4 GatheredAmbientOcclusion = Texture2DSample(AmbientOcclusionTexture, AmbientOcclusionSampler, SvPositionToBufferUV(SvPosition));
@@ -652,18 +659,29 @@ void Main(
 	MaterialAO *= GatheredAmbientOcclusion.r;
 #endif
 
+	FGBufferData GBuffer = (FGBufferData)0;
 	GBuffer.GBufferAO = MaterialAO;
-	// The smallest normalized value that can be represented in IEEE 754 (FP16) is 2^-24 = 5.96e-8.
-	// The code will make the following computation involving roughness: 1.0 / Roughness^4.
-	// Therefore to prevent division by zero on devices that do not support denormals, Roughness^4
-	// must be >= 5.96e-8. We will clamp to 0.015625 because 0.015625^4 = 5.96e-8.
-	//
-	// Note that we also clamp to 1.0 to match the deferred renderer on PC where the roughness is 
-	// stored in an 8-bit value and thus automatically clamped at 1.0.
-	GBuffer.Roughness = max(0.015625, GetMaterialRoughness(PixelMaterialInputs));
+	GBuffer.Depth = MaterialParameters.ScreenPosition.w;
+
+	SetGBufferForShadingModel(
+		GBuffer,
+		MaterialParameters,
+		Opacity,
+		BaseColor,
+		Metallic,
+		Specular,
+		Roughness,
+		Anisotropy,
+		SubsurfaceColor,
+		SubsurfaceProfile,
+		0.0f,
+		ShadingModelID
+	);
+
+	GBuffer.Roughness = max(0.015625f, GBuffer.Roughness);
 	
 	FMobileShadingModelContext ShadingModelContext = (FMobileShadingModelContext)0;
-	ShadingModelContext.Opacity = GetMaterialOpacity(PixelMaterialInputs);
+	ShadingModelContext.Opacity = Opacity;
 
 #if MATERIAL_SHADINGMODEL_THIN_TRANSLUCENT
 	float TopMaterialCoverage = ShadingModelContext.Opacity;
@@ -691,20 +709,19 @@ void Main(
 	GBuffer.BaseColor *= TopMaterialCoverage;
 #endif
 
-	half3 Color = 0;
-
-	half CustomData0 = GetMaterialCustomData0(MaterialParameters);
-	half CustomData1 = GetMaterialCustomData1(MaterialParameters);
-	InitShadingModelContext(ShadingModelContext, GBuffer, MaterialParameters.SvPosition, MaterialParameters.CameraVector, CustomData0, CustomData1);
+	InitShadingModelContext(ShadingModelContext, GBuffer, MaterialParameters.CameraVector);
 	float3 DiffuseDir = MaterialParameters.WorldNormal;
 
-#if MATERIAL_SHADINGMODEL_HAIR
-	FHairTransmittanceData TransmittanceData = InitHairTransmittanceData(true);
-	float3 N = MaterialParameters.WorldNormal;
-	float3 V = MaterialParameters.CameraVector;
-	float3 L = normalize(V - N * dot(V, N));
-	DiffuseDir = L;
-	ShadingModelContext.DiffuseColor = 2 * PI * HairShading(GBuffer, L, V, N, 1, TransmittanceData, 0, 0.2, uint2(0, 0));
+#if MATERIAL_SHADINGMODEL_HAIR && MOBILE_SHADINGMODEL_SUPPORT
+	if (GBuffer.ShadingModelID == SHADINGMODELID_HAIR)
+	{
+		FHairTransmittanceData TransmittanceData = InitHairTransmittanceData(true);
+		float3 N = MaterialParameters.WorldNormal;
+		float3 V = MaterialParameters.CameraVector;
+		float3 L = normalize(V - N * dot(V, N));
+		DiffuseDir = L;
+		ShadingModelContext.DiffuseColor = 2 * PI * HairShading(GBuffer, L, V, N, 1, TransmittanceData, 0, 0.2, uint2(0, 0));
+	}
 #endif
 
 	VTPageTableResult LightmapVTPageTableResult = (VTPageTableResult)0.0f;
@@ -723,15 +740,16 @@ void Main(
 	StoreVirtualTextureFeedback(MaterialParameters.VirtualTextureFeedback, 0, LightmapVTPageTableResult.PackedRequest);
 #endif
 
+	FMobileLightAccumulator MobileLightAccumulator = (FMobileLightAccumulator)0;
+
 	half IndirectIrradiance;
 	half3 IndirectColor;
-	ComputeIndirect(LightmapVTPageTableResult, Interpolants, DiffuseDir, ShadingModelContext, IndirectIrradiance, IndirectColor);
-	Color += IndirectColor;
+	half3 DiffuseColorForIndirect = ShadingModelContext.DiffuseColor;
 
-	half Shadow = GetPrimaryPrecomputedShadowMask(LightmapVTPageTableResult, Interpolants, MaterialParameters).r;
+	// Apply MaterialAO since we don't have the DiffuseIndirectComposite pass on mobile deferred.
+	ComputeIndirect(LightmapVTPageTableResult, Interpolants, DiffuseDir, DiffuseColorForIndirect, MaterialAO, IndirectIrradiance, MobileLightAccumulator);
 
-	// Apply MaterialAO to Color since we don't have the DiffuseIndirectComposite pass on mobile deferred.
-	Color *= MaterialAO;
+	half Shadow = GetPrimaryPrecomputedShadowMask(LightmapVTPageTableResult, Interpolants, MaterialParameters).r;
 
 #if DEFERRED_SHADING_PATH
 	float4 OutGBufferD;
@@ -752,22 +770,24 @@ void Main(
 #endif
 #else
 
-#if !MATERIAL_SHADINGMODEL_UNLIT
-
 #if ENABLE_SKY_LIGHT
 	half3 SkyDiffuseLighting = GetSkySHDiffuseSimple(MaterialParameters.WorldNormal);
 	half3 DiffuseLookup = SkyDiffuseLighting * ResolvedView.SkyLightColor.rgb;
 	IndirectIrradiance += Luminance(DiffuseLookup);
 #endif
 
-	IndirectIrradiance *= MaterialAO;
+	half3 L = MobileDirectionalLight.DirectionalLightDirectionAndShadowTransition.xyz;
+	half NoL = max(0, dot(MaterialParameters.WorldNormal, L));
+	half3 H = normalize(MaterialParameters.CameraVector + L);
+	half NoH = max(0, dot(MaterialParameters.WorldNormal, H));
+	half VoH = max(0, dot(MaterialParameters.CameraVector, H));
 
 	float  ShadowPositionZ = 0;
 #if DIRECTIONAL_LIGHT_CSM && !MATERIAL_SHADINGMODEL_SINGLELAYERWATER && !APPLY_DISTANCE_FIELD
 	// Cascaded Shadow Map
 	if (IsCSMEnabled())
 	{
-		half ShadowMap = MobileDirectionalLightCSM(MaterialParameters.ScreenPosition.xy, MaterialParameters.ScreenPosition.w, ShadowPositionZ);
+		half ShadowMap = MobileDirectionalLightCSM(MaterialParameters.ScreenPosition.xy, MaterialParameters.ScreenPosition.w, NoL, ShadowPositionZ);
 	#if ALLOW_STATIC_LIGHTING
 		Shadow = min(ShadowMap, Shadow);
 	#else
@@ -783,26 +803,20 @@ void Main(
 	}
 #endif
 
-	half NoL = max(0, dot(MaterialParameters.WorldNormal, MobileDirectionalLight.DirectionalLightDirectionAndShadowTransition.xyz));
-	half3 H = normalize(MaterialParameters.CameraVector + MobileDirectionalLight.DirectionalLightDirectionAndShadowTransition.xyz);
-	half NoH = max(0, dot(MaterialParameters.WorldNormal, H));
-
 #if TRANSLUCENCY_NON_DIRECTIONAL
 	NoL = 1.f;
 	ShadingModelContext.SpecularColor = 0.f;
 #endif
 
 	// Direct Lighting (Directional light) + IBL
-#if FULLY_ROUGH
-	Color += (Shadow * NoL) * MobileDirectionalLight.DirectionalLightColor.rgb * ShadingModelContext.DiffuseColor;
-#else
-	FMobileDirectLighting Lighting = MobileIntegrateBxDF(ShadingModelContext, GBuffer, NoL, MaterialParameters.CameraVector, H, NoH);
+	FMobileDirectLighting Lighting = MobileIntegrateBxDF(ShadingModelContext, GBuffer, NoL, NoH, VoH, MaterialParameters.CameraVector, L, Shadow);
 	// MobileDirectionalLight.DirectionalLightDistanceFadeMADAndSpecularScale.z saves SpecularScale for direction light.
-	Color += (Shadow) * MobileDirectionalLight.DirectionalLightColor.rgb * (Lighting.Diffuse + Lighting.Specular * MobileDirectionalLight.DirectionalLightDistanceFadeMADAndSpecularScale.z);
+	MobileLightAccumulator_Add(MobileLightAccumulator, Lighting.Diffuse, Lighting.Specular * MobileDirectionalLight.DirectionalLightDistanceFadeMADAndSpecularScale.z, MobileDirectionalLight.DirectionalLightColor.rgb);
 
 #if	!(MATERIAL_SINGLE_SHADINGMODEL && MATERIAL_SHADINGMODEL_HAIR)
 	// Environment map has been prenormalized, scale by lightmap luminance
 	half3 SpecularIBL = GetImageBasedReflectionLighting(MaterialParameters, GBuffer.Roughness, IndirectIrradiance, 1.0f);
+	half3 SpecularIBLLighting = (half3)0.0f;
 	#if MATERIAL_PLANAR_FORWARD_REFLECTIONS
 		BRANCH
 		if (abs(dot(PlanarReflectionStruct.ReflectionPlane.xyz, 1)) > .0001f)
@@ -823,16 +837,17 @@ void Main(
 		#else
 			half2 AB = EnvBRDFApproxLazarov(GBuffer.Roughness, ShadingModelContext.NoV);
 		#endif
-		Color += SpecularIBL * LayerAttenuation * (ShadingModelContext.SpecPreEnvBrdf * AB.x + AB.y * saturate(50 * ShadingModelContext.SpecPreEnvBrdf.g) * (1 - ShadingModelContext.ClearCoat));
+		SpecularIBLLighting += SpecularIBL * LayerAttenuation * (ShadingModelContext.SpecularColor * AB.x + AB.y * saturate(50 * ShadingModelContext.SpecularColor.g) * (1 - ShadingModelContext.ClearCoat));
 
 		SpecularIBL = GetImageBasedReflectionLighting(MaterialParameters, ShadingModelContext.ClearCoatRoughness, IndirectIrradiance, F);
-		Color += SpecularIBL;
+		SpecularIBLLighting += SpecularIBL;
 
 	#else
-		Color += SpecularIBL * ShadingModelContext.SpecularColor;
+		SpecularIBLLighting += SpecularIBL * ShadingModelContext.EnvBrdf;
 	#endif
+
+		MobileLightAccumulator_Add(MobileLightAccumulator, 0.0f, SpecularIBLLighting, 1.0f);
 #endif
-#endif /* FULLY_ROUGH */
 
 	// Local lights
 #if MAX_DYNAMIC_POINT_LIGHTS > 0 && !MATERIAL_SHADINGMODEL_SINGLELAYERWATER
@@ -861,7 +876,7 @@ void Main(
 												MobileMovablePointLight0.SpotLightShadowmapMinMax,
 												MobileMovablePointLight0.SpotLightShadowWorldToShadowMatrix,
 												#endif
-												Color);
+												MobileLightAccumulator);
 
 		if (MAX_DYNAMIC_POINT_LIGHTS > 1 && NumDynamicPointLights > 1)
 		{
@@ -878,7 +893,7 @@ void Main(
 													MobileMovablePointLight1.SpotLightShadowmapMinMax,
 													MobileMovablePointLight1.SpotLightShadowWorldToShadowMatrix,
 													#endif
-													Color);
+													MobileLightAccumulator);
 
 			if (MAX_DYNAMIC_POINT_LIGHTS > 2 && NumDynamicPointLights > 2)
 			{
@@ -895,7 +910,7 @@ void Main(
 														MobileMovablePointLight2.SpotLightShadowmapMinMax,
 														MobileMovablePointLight2.SpotLightShadowWorldToShadowMatrix,
 														#endif
-														Color);
+														MobileLightAccumulator);
 
 				if (MAX_DYNAMIC_POINT_LIGHTS > 3 && NumDynamicPointLights > 3)
 				{
@@ -912,7 +927,7 @@ void Main(
 															MobileMovablePointLight3.SpotLightShadowmapMinMax,
 															MobileMovablePointLight3.SpotLightShadowWorldToShadowMatrix,
 															#endif
-															Color);
+															MobileLightAccumulator);
 				}
 			}
 		}
@@ -931,20 +946,19 @@ void Main(
 	#if MATERIAL_SHADINGMODEL_SINGLELAYERWATER
 		ShadingModelContext.WaterDiffuseIndirectLuminance += SkyDiffuseLighting;
 	#endif
-		Color += SkyDiffuseLighting * half3(ResolvedView.SkyLightColor.rgb) * ShadingModelContext.DiffuseColor * MaterialAO;
+		half3 DiffuseColorForSky = ShadingModelContext.DiffuseColor;
+		MobileLightAccumulator_Add(MobileLightAccumulator, SkyDiffuseLighting, 0.0f, half3(ResolvedView.SkyLightColor.rgb) * DiffuseColorForSky * MaterialAO);
 	#if MATERIAL_TWOSIDED && LQ_TEXTURE_LIGHTMAP
 	}
 	#endif
 #endif
 
-#endif /* !MATERIAL_SHADINGMODEL_UNLIT */
-
 #if MATERIAL_SHADINGMODEL_SINGLELAYERWATER
 	{
 		const bool CameraIsUnderWater = false;	// Fade out the material contribution over to water contribution according to material opacity.
 		const float3 SunIlluminance = ResolvedView.DirectionalLightColor.rgb * PI;			// times PI because it is divided by PI on CPU (=luminance) and we want illuminance here. 
 		const float3 WaterDiffuseIndirectIlluminance = ShadingModelContext.WaterDiffuseIndirectLuminance * PI;	// DiffuseIndirectLighting is luminance. So we need to multiply by PI to get illuminance.
-		const float3 EnvBrdf = ShadingModelContext.SpecularColor; // SpecularColor is not F0 as in BasePassPixelShader, it is EnvBRDF.
+		const float3 EnvBrdf = ShadingModelContext.EnvBrdf;
 		const uint EyeIndex = 0;
 
 		const float4 NullDistortionParams = 1.0f;
@@ -955,13 +969,19 @@ void Main(
 			CameraIsUnderWater, ShadingModelContext.WaterVisibility, EyeIndex);
 
 		// Accumulate luminance and occlude the background according to transmittance to view and mean transmittance to lights.
-		Color += WaterLighting.Luminance;
+		MobileLightAccumulator_Add(MobileLightAccumulator, WaterLighting.Luminance, 0.0f, 1.0f);
 		ShadingModelContext.Opacity = 1.0 - ((1.0 - ShadingModelContext.Opacity) * dot(WaterLighting.WaterToSceneToLightTransmittance, float3(1.0 / 3.0, 1.0 / 3.0, 1.0 / 3.0)));
 	}
 #endif // MATERIAL_SHADINGMODEL_SINGLELAYERWATER
 
 #endif// DEFERRED_SHADING_PATH
 
+#if !MATERIAL_SHADINGMODEL_UNLIT
+	half3 Color = MobileLightAccumulator_GetResult(MobileLightAccumulator);
+#else
+	half3 Color = 0.0f;
+#endif
+
 	half4 VertexFog = half4(0, 0, 0, 1);
 
 #if USE_VERTEX_FOG
diff --git a/Engine/Shaders/Private/MobileDeferredShading.usf b/Engine/Shaders/Private/MobileDeferredShading.usf
--- a/Engine/Shaders/Private/MobileDeferredShading.usf
+++ b/Engine/Shaders/Private/MobileDeferredShading.usf
@@ -1,4 +1,7 @@
 // Copyright Epic Games, Inc. All Rights Reserved.
+
+#define MOBILE_DEFERRED_LIGHTING 1
+
 #include "Common.ush"
 #include "SHCommon.ush"
 #include "/Engine/Generated/Material.ush"
@@ -89,7 +92,7 @@ FGBufferData DecodeGBufferMobile(
 	float4 InGBufferC,
 	float4 InGBufferD)
 {
-	FGBufferData GBuffer;
+	FGBufferData GBuffer = (FGBufferData)0;
 	GBuffer.WorldNormal = OctahedronToUnitVector( InGBufferA.xy * 2.0f - 1.0f );
 	GBuffer.PrecomputedShadowFactors = InGBufferA.z;
 	GBuffer.PerObjectGBufferData = InGBufferA.a;  
@@ -138,15 +141,14 @@ half ComputeLightFunctionMultiplier(float3 WorldPosition)
 #endif
 }
 
-half3 GetDirectLighting(
+void GetDirectLighting(
 	FMobileLightData LightData, 
 	FMobileShadingModelContext ShadingModelContext, 
 	FGBufferData GBuffer, 
 	float3 WorldPosition, 
-	half3 CameraVector)
+	half3 CameraVector,
+	inout FMobileLightAccumulator MobileLightAccumulator)
 {
-	half3 DirectLighting = 0;
-	
 	float3 ToLight = LightData.Position - WorldPosition;
 	float DistanceSqr = dot(ToLight, ToLight);
 	float3 L = ToLight * rsqrt(DistanceSqr);
@@ -173,26 +175,26 @@ half3 GetDirectLighting(
 		half3 H = normalize(CameraVector + L);
 		half NoL = max(0.0, dot(GBuffer.WorldNormal, L));
 		half NoH = max(0.0, dot(GBuffer.WorldNormal, H));
-		FMobileDirectLighting Lighting = MobileIntegrateBxDF(ShadingModelContext, GBuffer, NoL, CameraVector, H, NoH);
-		DirectLighting = (Lighting.Diffuse + Lighting.Specular * LightData.SpecularScale) * (LightData.Color * (1.0 / PI) * Attenuation);
+		half VoH = max(0.0, dot(CameraVector, H));
+		FMobileDirectLighting Lighting = MobileIntegrateBxDF(ShadingModelContext, GBuffer, NoL, NoH, VoH, CameraVector, L);
+		MobileLightAccumulator_Add(MobileLightAccumulator, Lighting.Diffuse, Lighting.Specular * LightData.SpecularScale, LightData.Color * (1.0 / PI) * Attenuation);
 	}
-	return DirectLighting;
 }
 
 /**
  * Adds local lighting using the light grid, does not apply directional lights, as they are done elsewhere.
  * Does not support dynamic shadows, as these require the per-light shadow mask.
  */
-half3 GetLightGridLocalLighting(
+void GetLightGridLocalLighting(
 	const FCulledLightsGridData InLightGridData, 
 	FMobileShadingModelContext ShadingModelContext, 
 	FGBufferData GBuffer, 
 	float3 WorldPosition, 
 	half3 CameraVector, 
 	uint EyeIndex, 
-	uint FirstNonSimpleLightIndex)
+	uint FirstNonSimpleLightIndex,
+	inout FMobileLightAccumulator MobileLightAccumulator)
 {
-	half3 DirectLighting = 0;
 	// Limit max to ForwardLightData.NumLocalLights.
 	// This prevents GPU hangs when the PS tries to read from uninitialized NumCulledLightsGrid buffer
 	const uint NumLocalLights = min(InLightGridData.NumLocalLights, GetNumLocalLights(EyeIndex));
@@ -231,10 +233,8 @@ half3 GetLightGridLocalLighting(
 		// TODO: not sure if this check is even needed, standard deferred always seems to set it to true?
 		LightData.bSpotLight = (LightType == LIGHT_TYPE_SPOT);
 		
-		DirectLighting+= GetDirectLighting(LightData, ShadingModelContext, GBuffer, WorldPosition, CameraVector);
+		GetDirectLighting(LightData, ShadingModelContext, GBuffer, WorldPosition, CameraVector, MobileLightAccumulator);
 	}
-
-	return DirectLighting;
 }
 
 void MobileDirectLightPS(
@@ -261,7 +261,14 @@ void MobileDirectLightPS(
 	half NoV = max(0, dot(GBuffer.WorldNormal, CameraVector));
 	half3 ReflectionVector = GBuffer.WorldNormal * (NoV * 2.0) - CameraVector;
 	
-	half3 Color = 0;
+	// Directional light
+	half3 L = MobileDirectionalLight.DirectionalLightDirectionAndShadowTransition.xyz;
+	half NoL = max(0, dot(GBuffer.WorldNormal, L));
+	half3 H = normalize(CameraVector + L);
+	half NoH = max(0, dot(GBuffer.WorldNormal, H));
+	half VoH = max(0, dot(CameraVector, H));
+
+	FMobileLightAccumulator MobileLightAccumulator = (FMobileLightAccumulator)0;
 	// Check movable light param to determine if we should be using precomputed shadows
 	half Shadow = LightFunctionParameters2.z > 0.0f ? 1.0f : GBuffer.PrecomputedShadowFactors.r;
 
@@ -269,17 +276,12 @@ void MobileDirectLightPS(
 	// TODO: we don't need to fully compute ScreenPosition here
 	float  ShadowPositionZ = 0;
 	float4 ScreenPosition = SvPositionToScreenPosition(float4(SvPosition.xyz,SceneDepth));
-	float ShadowMap = MobileDirectionalLightCSM(ScreenPosition.xy, SceneDepth, ShadowPositionZ);
+	float ShadowMap = MobileDirectionalLightCSM(ScreenPosition.xy, SceneDepth, NoL, ShadowPositionZ);
 	Shadow = min(ShadowMap, Shadow);
 #endif
 
 	FMobileShadingModelContext ShadingModelContext = (FMobileShadingModelContext)0;
-	{
-		half DielectricSpecular = 0.08 * GBuffer.Specular;
-		ShadingModelContext.DiffuseColor = GBuffer.BaseColor - GBuffer.BaseColor * GBuffer.Metallic;	// 1 mad
-		ShadingModelContext.SpecularColor = (DielectricSpecular - DielectricSpecular * GBuffer.Metallic) + GBuffer.BaseColor * GBuffer.Metallic;	// 2 mad
-		ShadingModelContext.SpecularColor = GetEnvBRDF(ShadingModelContext.SpecularColor, GBuffer.Roughness, NoV);
-	}
+	InitShadingModelContext(ShadingModelContext, GBuffer, CameraVector);
 	
 	float2 LocalPosition = SvPosition.xy - View.ViewRectMin.xy;
 	uint GridIndex = ComputeLightGridCellIndex(uint2(LocalPosition.x, LocalPosition.y), SceneDepth);
@@ -288,18 +290,12 @@ void MobileDirectLightPS(
 	{
 		const uint EyeIndex = 0;
 		const FCulledLightsGridData CulledLightGridData = GetCulledLightsGrid(GridIndex, EyeIndex);
-		Color += GetLightGridLocalLighting(CulledLightGridData, ShadingModelContext, GBuffer, WorldPosition, CameraVector, EyeIndex, 0);
+		GetLightGridLocalLighting(CulledLightGridData, ShadingModelContext, GBuffer, WorldPosition, CameraVector, EyeIndex, 0, MobileLightAccumulator);
 	}
 #endif
-			
-	// Directional light
-	half NoL = max(0, dot(GBuffer.WorldNormal, MobileDirectionalLight.DirectionalLightDirectionAndShadowTransition.xyz));
-	half3 H = normalize(CameraVector + MobileDirectionalLight.DirectionalLightDirectionAndShadowTransition.xyz);
-	half NoH = max(0, dot(GBuffer.WorldNormal, H));
-	FMobileDirectLighting Lighting;
-	Lighting.Specular = ShadingModelContext.SpecularColor * CalcSpecular(GBuffer.Roughness, NoH);
-	Lighting.Diffuse = ShadingModelContext.DiffuseColor;
-	Color += (Shadow * NoL) * MobileDirectionalLight.DirectionalLightColor.rgb * (Lighting.Diffuse + Lighting.Specular * MobileDirectionalLight.DirectionalLightDistanceFadeMADAndSpecularScale.z);
+
+	FMobileDirectLighting Lighting = MobileIntegrateBxDF(ShadingModelContext, GBuffer, NoL, NoH, VoH, CameraVector, L, Shadow);
+	MobileLightAccumulator_Add(MobileLightAccumulator, Lighting.Diffuse, Lighting.Specular * MobileDirectionalLight.DirectionalLightDistanceFadeMADAndSpecularScale.z, MobileDirectionalLight.DirectionalLightColor.rgb);
 
 #if APPLY_REFLECTION
 	uint NumCulledEntryIndex = (ForwardLightData.NumGridCells + GridIndex) * NUM_CULLED_LIGHTS_GRID_STRIDE;
@@ -326,16 +322,20 @@ void MobileDirectLightPS(
 		// Planar reflections win over reflection environment
 		SpecularIBL = lerp(SpecularIBL, PlanarReflection.rgb, PlanarReflection.a);
 	}
-		
-	Color += SpecularIBL * ShadingModelContext.SpecularColor;
+
+	MobileLightAccumulator_Add(MobileLightAccumulator, 0.0f, SpecularIBL * ShadingModelContext.EnvBrdf, 1.0f);
 #elif APPLY_SKY_REFLECTION
 	float SkyAverageBrightness = 1.0f;
 	float3 SpecularIBL = GetSkyLightReflection(ReflectionVector, GBuffer.Roughness, SkyAverageBrightness);
 	SpecularIBL *= ComputeMixingWeight(GBuffer.IndirectIrradiance, SkyAverageBrightness, GBuffer.Roughness);
-	Color += SpecularIBL * ShadingModelContext.SpecularColor;
+	MobileLightAccumulator_Add(MobileLightAccumulator, 0.0f, SpecularIBL * ShadingModelContext.EnvBrdf, 1.0f);
 #endif
 	half3 SkyDiffuseLighting = GetSkySHDiffuseSimple(GBuffer.WorldNormal);
-	Color+= SkyDiffuseLighting * half3(View.SkyLightColor.rgb) * ShadingModelContext.DiffuseColor * GBuffer.GBufferAO;
+	half3 DiffuseColorForSky = ShadingModelContext.DiffuseColor;
+	MobileLightAccumulator_Add(MobileLightAccumulator, SkyDiffuseLighting, 0.0f, half3(View.SkyLightColor.rgb) * DiffuseColorForSky * GBuffer.GBufferAO);
+	
+	half3 Color = MobileLightAccumulator_GetResult(MobileLightAccumulator);
+
 	half LightAttenuation = ComputeLightFunctionMultiplier(WorldPosition);
 
 	// MobileHDR applies PreExposure in tonemapper
@@ -384,14 +384,11 @@ void MobileRadialLightPS(
 	}
 
 	FMobileShadingModelContext ShadingModelContext = (FMobileShadingModelContext)0;
-	{
-		half DielectricSpecular = 0.08 * GBuffer.Specular;
-		ShadingModelContext.DiffuseColor = GBuffer.BaseColor - GBuffer.BaseColor * GBuffer.Metallic;	// 1 mad
-		ShadingModelContext.SpecularColor = (DielectricSpecular - DielectricSpecular * GBuffer.Metallic) + GBuffer.BaseColor * GBuffer.Metallic;	// 2 mad
-		ShadingModelContext.SpecularColor = GetEnvBRDF(ShadingModelContext.SpecularColor, GBuffer.Roughness, NoV);
-	}
+	InitShadingModelContext(ShadingModelContext, GBuffer, CameraVector);
 			
-	half3 Color = GetDirectLighting(LightData, ShadingModelContext, GBuffer, WorldPosition, CameraVector);
+	FMobileLightAccumulator MobileLightAccumulator = (FMobileLightAccumulator)0;
+	GetDirectLighting(LightData, ShadingModelContext, GBuffer, WorldPosition, CameraVector, MobileLightAccumulator);
+	half3 Color = MobileLightAccumulator_GetResult(MobileLightAccumulator);
 	
 	half LightAttenuation = ComputeLightProfileMultiplier(WorldPosition, DeferredLightUniforms.Position, -DeferredLightUniforms.Direction, DeferredLightUniforms.Tangent);
 	LightAttenuation*= ComputeLightFunctionMultiplier(WorldPosition);
diff --git a/Engine/Shaders/Private/MobileShadingModels.ush b/Engine/Shaders/Private/MobileShadingModels.ush
--- a/Engine/Shaders/Private/MobileShadingModels.ush
+++ b/Engine/Shaders/Private/MobileShadingModels.ush
@@ -21,137 +21,198 @@
 #define MOBILE_SHADOW_QUALITY 2
 #endif
 
+#ifndef MOBILE_HIGH_QUALITY_BRDF
+#define MOBILE_HIGH_QUALITY_BRDF 0
+#endif
+
+//It's always 0 in mobile deferred lighting shaders
+#ifndef FULLY_ROUGH
+#define FULLY_ROUGH 0
+#endif
+
+//It's always 0 in mobile deferred lighting shaders
+#ifndef NONMETAL
+#define NONMETAL 0
+#endif
+
+//It's always 0 in mobile deferred lighting shaders
+#ifndef MATERIAL_SHADINGMODEL_SINGLELAYERWATER
+#define MATERIAL_SHADINGMODEL_SINGLELAYERWATER 0
+#endif
+
+//It's always 0 in mobile deferred lighting shaders
+#ifndef FORWARDSHADING_USE_HQ_ENV_BRDF
+#define  FORWARDSHADING_USE_HQ_ENV_BRDF 0
+#endif
+
+//It's always 1 in mobile deferred lighting shaders
+#ifndef MOBILE_DEFERRED_LIGHTING
+#define MOBILE_DEFERRED_LIGHTING 0
+#endif
+
+#ifndef DEFERRED_SHADING_PATH
+#define DEFERRED_SHADING_PATH 0
+#endif
+
+// Enable shading model support for deferred shading once the gbuffer has capacity.
+#define MOBILE_SHADINGMODEL_SUPPORT_DEFERRED 0
+
+#define MOBILE_SHADINGMODEL_SUPPORT (!DEFERRED_SHADING_PATH || SHADINGMODEL_SUPPORT_MOBILE_DEFERRED)
+
+/*------------------------------------------------------------------------------
+	Mobile Shading Models
+------------------------------------------------------------------------------*/
+
 struct FMobileDirectLighting
 {
-	half3 Diffuse;
-	half3 Specular;
+	half3	Diffuse;
+	half3	Specular;
 };
 
 struct FMobileShadingModelContext
 {
 	half Opacity;
 	half3 DiffuseColor;
-#if NONMETAL
-	half SpecularColor;
-#else
 	half3 SpecularColor;
-#endif
-	
-#if MATERIAL_SHADINGMODEL_CLEAR_COAT
-	half ClearCoat;
-	half ClearCoatRoughness;
-	float NoV;
-	half3 SpecPreEnvBrdf;
-#elif MATERIAL_SHADINGMODEL_SINGLELAYERWATER
+	half3 EnvBrdf;
+
+	half NoV;
+
+#if MATERIAL_SHADINGMODEL_SINGLELAYERWATER
 	half BaseMaterialCoverageOverWater;
 	half WaterVisibility;
 	float3 WaterDiffuseIndirectLuminance;
 #endif
+
+#if MOBILE_SHADINGMODEL_SUPPORT
+	half ClearCoat;
+	half ClearCoatRoughness;
+#endif
 };
 
-#if NONMETAL
-half GetEnvBRDF(half SpecularColor, half Roughness, half NoV)
+half3 GetEnvBRDF(half3 SpecularColor, half Roughness, half NoV)
 {
-#if FORWARDSHADING_USE_HQ_ENV_BRDF
-	return EnvBRDF(SpecularColor, Roughness, NoV).x;
-#else
+#if FORWARDSHADING_USE_HQ_ENV_BRDF || MOBILE_DEFERRED_LIGHTING
+	return EnvBRDF(SpecularColor, Roughness, NoV);
+#elif NONMETAL
 	// If nothing is hooked up to Metalic and Specular,
 	// then defaults are the same as a non-metal,
 	// so this define is safe.
-	return EnvBRDFApproxNonmetal(Roughness, NoV);
-#endif
-}
-#else
-half3 GetEnvBRDF(half3 SpecularColor, half Roughness, half NoV)
-{
-#if FORWARDSHADING_USE_HQ_ENV_BRDF
-	return EnvBRDF(SpecularColor, Roughness, NoV);
+	return EnvBRDFApproxNonmetal(Roughness, NoV).xxx;
 #else
 	return EnvBRDFApprox(SpecularColor, Roughness, NoV);
 #endif
 }
-#endif
 
-void InitShadingModelContext(inout FMobileShadingModelContext ShadingModelContext, inout FGBufferData GBuffer, float4 SvPosition, half3 CameraVector, half CustomData0, half CustomData1)
+void InitShadingModelContext(inout FMobileShadingModelContext ShadingModelContext, inout FGBufferData GBuffer, half3 CameraVector)
 {
-#if !FULLY_ROUGH || MATERIAL_SHADINGMODEL_CLEAR_COAT
-	half NoV = max(dot(GBuffer.WorldNormal, CameraVector), 0);
-#endif
+	half NoV = max(dot(GBuffer.WorldNormal, CameraVector), 0.0f);
 
-#if MATERIAL_SHADINGMODEL_CLEAR_COAT
-	ShadingModelContext.ClearCoat = saturate(CustomData0);
-	ShadingModelContext.ClearCoatRoughness = clamp(CustomData1, 0.015625, 1.0);
-
-	// Approximation of refraction's effect on EnvBRDF
-	half RefractionScale = ((NoV * 0.5 + 0.5) * NoV - 1) * saturate(1.25 - 1.25 * GBuffer.Roughness) + 1;
+	// This is to prevent Vis to get inf when both NoL and NoV are 0.
+	ShadingModelContext.NoV = saturate(abs(NoV) + 1e-5);
 
-	GBuffer.Specular *= lerp(1, RefractionScale, ShadingModelContext.ClearCoat);
+#if NONMETAL
+	GBuffer.DiffuseColor = GBuffer.BaseColor;
+	GBuffer.SpecularColor = 0.04;
+#else
+	GBuffer.DiffuseColor = GBuffer.BaseColor - GBuffer.BaseColor * GBuffer.Metallic;	// 1 mad
 	half DielectricSpecular = 0.08 * GBuffer.Specular;
-	ShadingModelContext.SpecularColor = (DielectricSpecular - DielectricSpecular * GBuffer.Metallic) + GBuffer.BaseColor * GBuffer.Metallic;	// 2 mad
-	ShadingModelContext.SpecPreEnvBrdf = ShadingModelContext.SpecularColor;
-	ShadingModelContext.DiffuseColor = GBuffer.BaseColor - GBuffer.BaseColor * GBuffer.Metallic;
+	GBuffer.SpecularColor = (DielectricSpecular - DielectricSpecular * GBuffer.Metallic) + GBuffer.BaseColor * GBuffer.Metallic;	// 2 mad
+#endif
 
-	// This is to prevent Vis to get inf when both NoL and NoV are 0.
-	ShadingModelContext.NoV = saturate(abs(NoV) + 1e-5);
+#if MOBILE_EMULATION && !MOBILE_DEFERRED_LIGHTING
+	{
+		// this feature is only needed for development/editor - we can compile it out for a shipping build (see r.CompileShadersForDevelopment cvar help)
+		GBuffer.DiffuseColor = GBuffer.DiffuseColor * ResolvedView.DiffuseOverrideParameter.w + ResolvedView.DiffuseOverrideParameter.xyz;
+		GBuffer.SpecularColor = GBuffer.SpecularColor * ResolvedView.SpecularOverrideParameter.w + ResolvedView.SpecularOverrideParameter.xyz;
+	}
+#endif
 
-#elif MATERIAL_SHADINGMODEL_SINGLELAYERWATER
+	ShadingModelContext.DiffuseColor = GBuffer.DiffuseColor;
+	ShadingModelContext.SpecularColor = GBuffer.SpecularColor;
+
+#if MATERIAL_SHADINGMODEL_SINGLELAYERWATER
 	ShadingModelContext.BaseMaterialCoverageOverWater = ShadingModelContext.Opacity;
 	ShadingModelContext.WaterVisibility = 1.0 - ShadingModelContext.BaseMaterialCoverageOverWater;
+	// Fade out diffuse as this will be handled by the single scattering lighting. when over the water surface.
+	// We keep the SpecularColor for sun/water interactions
+	ShadingModelContext.WaterDiffuseIndirectLuminance = 0;
+	ShadingModelContext.DiffuseColor *= ShadingModelContext.BaseMaterialCoverageOverWater;
+#endif
 
-	half DielectricSpecular = 0.08 * GBuffer.Specular;
-	ShadingModelContext.DiffuseColor = GBuffer.BaseColor - GBuffer.BaseColor * GBuffer.Metallic;	// 1 mad
-	ShadingModelContext.SpecularColor = (DielectricSpecular - DielectricSpecular * GBuffer.Metallic) + GBuffer.BaseColor * GBuffer.Metallic;	// 2 mad
-
-#else // MATERIAL_SHADINGMODEL_DEFAULTLIT
-	
-	#if NONMETAL
-		ShadingModelContext.DiffuseColor = GBuffer.BaseColor;
-		ShadingModelContext.SpecularColor = 0.04;
-	#else
-		half DielectricSpecular = 0.08 * GBuffer.Specular;
-		ShadingModelContext.DiffuseColor = GBuffer.BaseColor - GBuffer.BaseColor * GBuffer.Metallic;	// 1 mad
-		ShadingModelContext.SpecularColor = (DielectricSpecular - DielectricSpecular * GBuffer.Metallic) + GBuffer.BaseColor * GBuffer.Metallic;	// 2 mad
-	#endif
+#if MOBILE_SHADINGMODEL_SUPPORT
+	if (GBuffer.ShadingModelID == SHADINGMODELID_CLEAR_COAT)
+	{
+		ShadingModelContext.ClearCoat = GBuffer.CustomData.x;
+		ShadingModelContext.ClearCoatRoughness = clamp(GBuffer.CustomData.y, 0.015625, 1.0);
 
+		// Approximation of refraction's effect on EnvBRDF
+		half RefractionScale = ((NoV * 0.5 + 0.5) * NoV - 1) * saturate(1.25 - 1.25 * GBuffer.Roughness) + 1;
+
+		half Specular = lerp(GBuffer.Specular, RefractionScale, ShadingModelContext.ClearCoat);
+		half DielectricSpecular = 0.08 * Specular;
+		ShadingModelContext.SpecularColor = (DielectricSpecular - DielectricSpecular * GBuffer.Metallic) + GBuffer.BaseColor * GBuffer.Metallic;	// 2 mad
+	}
 #endif
 
 #if FULLY_ROUGH
 	EnvBRDFApproxFullyRough(ShadingModelContext.DiffuseColor, ShadingModelContext.SpecularColor);
+	ShadingModelContext.EnvBrdf = 0.0f;
 #else
-	ShadingModelContext.SpecularColor = GetEnvBRDF(ShadingModelContext.SpecularColor, GBuffer.Roughness, NoV);
+	ShadingModelContext.EnvBrdf = GetEnvBRDF(ShadingModelContext.SpecularColor, GBuffer.Roughness, NoV);
 #endif
+}
 
-#if MOBILE_EMULATION
-	{
-		// this feature is only needed for development/editor - we can compile it out for a shipping build (see r.CompileShadersForDevelopment cvar help)
-		ShadingModelContext.DiffuseColor = ShadingModelContext.DiffuseColor * ResolvedView.DiffuseOverrideParameter.w + ResolvedView.DiffuseOverrideParameter.xyz;
-#if NONMETAL
-		ShadingModelContext.SpecularColor = ShadingModelContext.SpecularColor * ResolvedView.SpecularOverrideParameter.w + ResolvedView.SpecularOverrideParameter.x;
+half Vis_SmithJointApprox_Mobile(half a, half NoV, half NoL)
+{
+	half Vis_SmithV = NoL * (NoV * (1 - a) + a);
+	half Vis_SmithL = NoV * (NoL * (1 - a) + a);
+	return 0.5 * rcp(Vis_SmithV + Vis_SmithL);
+}
+
+half3 MobileSpecularGGXInner(half D, half3 EnvBrdf, half3 SpecularColor, half Roughness, half NoV, half NoL, half VoH)
+{
+#if MOBILE_HIGH_QUALITY_BRDF
+	half Vis = Vis_SmithJointApprox_Mobile(Pow2(Roughness), NoV, NoL);
+	half3 F = F_Schlick(SpecularColor, VoH);
 #else
-		ShadingModelContext.SpecularColor = ShadingModelContext.SpecularColor * ResolvedView.SpecularOverrideParameter.w + ResolvedView.SpecularOverrideParameter.xyz;
-#endif
-	}
+	half Vis = (Roughness * 0.25 + 0.25);
+	half3 F = EnvBrdf;
 #endif
 
-#if MATERIAL_SHADINGMODEL_SINGLELAYERWATER
-	// Fade out diffuse as this will be handled by the single scattering lighting. when over the water surface.
-	// We keep the SpecularColor for sun/water interactions
-	ShadingModelContext.WaterDiffuseIndirectLuminance = 0;
-	ShadingModelContext.DiffuseColor *= ShadingModelContext.BaseMaterialCoverageOverWater;
-#endif
+	return (D * Vis) * F;
+}
+
+half3 MobileSpecularGGX(half3 EnvBrdf, half3 SpecularColor, half Roughness, half NoV, half NoH, half NoL, half VoH)
+{
+	half D = GGX_Mobile(Roughness, NoH);
+
+	return MobileSpecularGGXInner(D, EnvBrdf, SpecularColor, Roughness, NoV, NoL, VoH);
+}
+
+half3 MobileDualSpecularGGX(half AverageRoughness, half Lobe0Roughness, half Lobe1Roughness, half LobeMix, half3 EnvBrdf, half3 SpecularColor, half NoV, half NoH, half NoL, half VoH)
+{
+	half D = lerp(GGX_Mobile(Lobe0Roughness, NoH), GGX_Mobile(Lobe1Roughness, NoH), LobeMix);
+
+	return MobileSpecularGGXInner(D, EnvBrdf, SpecularColor, AverageRoughness, NoV, NoL, VoH);
 }
 
-half CalcSpecular(half Roughness, half NoH)
+FMobileDirectLighting MobileDefaultLitBxDF(FMobileShadingModelContext ShadingModelContext, FGBufferData GBuffer, half NoL, half NoH, half VoH, half3 V, half3 L, half Shadow)
 {
-	return (Roughness*0.25 + 0.25) * GGX_Mobile(Roughness, NoH);
+	FMobileDirectLighting Lighting;
+
+	Lighting.Specular = Shadow * NoL * MobileSpecularGGX(ShadingModelContext.EnvBrdf, ShadingModelContext.SpecularColor, GBuffer.Roughness, ShadingModelContext.NoV, NoH, NoL, VoH);
+	Lighting.Diffuse = Shadow * NoL * ShadingModelContext.DiffuseColor;
+
+	return Lighting;
 }
 
-FMobileDirectLighting MobileIntegrateBxDF(FMobileShadingModelContext ShadingModelContext, FGBufferData GBuffer, half NoL, half3 V, half3 H, half NoH)
+#if MOBILE_SHADINGMODEL_SUPPORT
+FMobileDirectLighting MobileClearCoatBxDF(FMobileShadingModelContext ShadingModelContext, FGBufferData GBuffer, half NoL, half NoH, half VoH, half3 V, half3 L, half Shadow)
 {
 	FMobileDirectLighting Lighting;
-#if MATERIAL_SHADINGMODEL_CLEAR_COAT
+
 	half ClearCoatRoughness = ShadingModelContext.ClearCoatRoughness;
-	half VoH = max(0, dot(V, H));
 	half F0 = 0.04;
 	half Fc = Pow5(1 - VoH);
 	half F = Fc + (1 - Fc) * F0;
@@ -195,22 +256,73 @@ FMobileDirectLighting MobileIntegrateBxDF(FMobileShadingModelContext ShadingMode
 	half3 CommonDiffuse = ShadingModelContext.DiffuseColor;
 	half3 DefaultDiffuse = NoL;
 	half3 RefractedDiffuse = (LayerAttenuation * BottomNoL) * Transmission;
-	Lighting.Diffuse = CommonDiffuse * lerp(DefaultDiffuse, RefractedDiffuse, ShadingModelContext.ClearCoat);
+	Lighting.Diffuse = Shadow * CommonDiffuse * lerp(DefaultDiffuse, RefractedDiffuse, ShadingModelContext.ClearCoat);
 
-	half3 CommonSpecular = ShadingModelContext.SpecularColor * CalcSpecular(GBuffer.Roughness, NoH);
+	half3 CommonSpecular = MobileSpecularGGX(ShadingModelContext.EnvBrdf, ShadingModelContext.SpecularColor, GBuffer.Roughness, ShadingModelContext.NoV, NoH, NoL, VoH);
 	half3 DefaultSpecular = NoL;
 	half3 RefractedSpecular = LayerAttenuation * Transmission * BottomNoL;
 	Lighting.Specular += CommonSpecular * lerp(DefaultSpecular, RefractedSpecular, ShadingModelContext.ClearCoat);
-#elif MATERIAL_SHADINGMODEL_HAIR
-	 // TODO transmittance
+
+	Lighting.Specular *= Shadow;
+
+	return Lighting;
+}
+
+FMobileDirectLighting MobileHairBxDF(FMobileShadingModelContext ShadingModelContext, FGBufferData GBuffer, half NoL, half NoH, half VoH, half3 V, half3 L, half Shadow)
+{
+	FMobileDirectLighting Lighting;
+
+	// TODO transmittance
 	Lighting.Specular = half3(0, 0, 0);
 	Lighting.Diffuse = half3(0, 0, 0);
-#else // MATERIAL_SHADINGMODEL_DEFAULT_LIT
-	Lighting.Specular = ShadingModelContext.SpecularColor * (NoL * CalcSpecular(GBuffer.Roughness, NoH));
-	Lighting.Diffuse = NoL * ShadingModelContext.DiffuseColor;
-#endif
+
 	return Lighting;
 }
+#endif
+
+FMobileDirectLighting MobileIntegrateBxDF(FMobileShadingModelContext ShadingModelContext, FGBufferData GBuffer, half NoL, half NoH, half VoH, half3 V, half3 L, half Shadow = 1.0f)
+{
+#if MOBILE_SHADINGMODEL_SUPPORT
+	if (GBuffer.ShadingModelID == SHADINGMODELID_CLEAR_COAT)
+	{
+		return MobileClearCoatBxDF(ShadingModelContext, GBuffer, NoL, NoH, VoH, V, L, Shadow);
+	}
+	else if (GBuffer.ShadingModelID == SHADINGMODELID_HAIR)
+	{
+		return MobileHairBxDF(ShadingModelContext, GBuffer, NoL, NoH, VoH, V, L, Shadow);
+	}
+#endif
+
+	return MobileDefaultLitBxDF(ShadingModelContext, GBuffer, NoL, NoH, VoH, V, L, Shadow);
+}
+
+/*------------------------------------------------------------------------------
+	Mobile Light Accumulation
+------------------------------------------------------------------------------*/
+struct FMobileLightAccumulator
+{
+	float3 TotalLight;
+};
+
+// accumulate light, can be called multiple times
+void MobileLightAccumulator_Add(inout FMobileLightAccumulator In, half3 DiffuseTotalLight, half3 SpecularTotalLight, half3 CommonMultiplier)
+{
+#if FULLY_ROUGH
+	In.TotalLight += DiffuseTotalLight * CommonMultiplier;
+#else
+	In.TotalLight += (DiffuseTotalLight + SpecularTotalLight) * CommonMultiplier;
+#endif
+}
+
+half3 MobileLightAccumulator_GetResult(FMobileLightAccumulator In)
+{
+	return min(65000.0, In.TotalLight);
+}
+
+/*------------------------------------------------------------------------------
+	Mobile Shadow.
+------------------------------------------------------------------------------*/
+
 half MobileShadowPCF(float2 ShadowUVs, FPCFSamplerSettings Settings)
 {
 #if MOBILE_SHADOW_QUALITY == 0
@@ -231,14 +343,14 @@ return ShadowMap;
 // Add fading CSM plane:
 #define FADE_CSM 1
 
-half MobileDirectionalLightCSM(float2 ScreenPosition, float SceneDepth, inout float ShadowPositionZ)
+half MobileDirectionalLightCSM(float2 ScreenPosition, float SceneDepth, half NoL, inout float ShadowPositionZ)
 {
 	half ShadowMap = 1;
 	ShadowPositionZ = 0;
 	FPCFSamplerSettings Settings;
 	Settings.ShadowDepthTexture = MobileDirectionalLight.DirectionalLightShadowTexture;
 	Settings.ShadowDepthTextureSampler = MobileDirectionalLight.DirectionalLightShadowSampler;
-	Settings.TransitionScale = MobileDirectionalLight.DirectionalLightDirectionAndShadowTransition.w;
+	Settings.TransitionScale = MobileDirectionalLight.DirectionalLightDirectionAndShadowTransition.w * lerp(MobileDirectionalLight.DirectionalLightDistanceFadeMADAndSpecularScale.w, 1.0, NoL);
 	Settings.ShadowBufferSize = MobileDirectionalLight.DirectionalLightShadowSize;
 	Settings.bSubsurface = false;
 	Settings.bTreatMaxDepthUnshadowed = false;
diff --git a/Engine/Source/Runtime/Core/Private/HAL/ConsoleManager.cpp b/Engine/Source/Runtime/Core/Private/HAL/ConsoleManager.cpp
--- a/Engine/Source/Runtime/Core/Private/HAL/ConsoleManager.cpp
+++ b/Engine/Source/Runtime/Core/Private/HAL/ConsoleManager.cpp
@@ -2883,4 +2883,13 @@ static TAutoConsoleVariable<int32> CVarMobileSupportsGen4TAA(
 	TEXT("Support desktop Gen4 TAA with mobile rendering\n")
 	TEXT("0: Fallback to FXAA")
 	TEXT("1: Support Desktop Gen4 TAA (default)"),
-	ECVF_ReadOnly | ECVF_RenderThreadSafe);
\ No newline at end of file
+	ECVF_ReadOnly | ECVF_RenderThreadSafe);
+
+static TAutoConsoleVariable<int32> CVarMobileEnableHighQualityBRDF(
+	TEXT("r.Mobile.EnableHighQualityBRDF"),
+	0,
+	TEXT("Use the high quality brdf functions on mobile to get better visual effects but adds GPU cost.\n")
+	TEXT("0: Disabled (More approximations but better performance. [default])\n")
+	TEXT("1: Enabled (More accurate but increases cost on GPU.)\n"),
+	ECVF_ReadOnly | ECVF_RenderThreadSafe
+);
\ No newline at end of file
diff --git a/Engine/Source/Runtime/Engine/Classes/Engine/RendererSettings.h b/Engine/Source/Runtime/Engine/Classes/Engine/RendererSettings.h
--- a/Engine/Source/Runtime/Engine/Classes/Engine/RendererSettings.h
+++ b/Engine/Source/Runtime/Engine/Classes/Engine/RendererSettings.h
@@ -909,6 +909,12 @@ class ENGINE_API URendererSettings : public UDeveloperSettings
 		ConfigRestartRequired = true))
 		uint32 bMobileSupportsGen4TAA : 1;
 
+	UPROPERTY(config, EditAnywhere, Category = Mobile, meta = (
+		ConsoleVariable = "r.Mobile.EnableHighQualityBRDF", DisplayName = "Mobile High Quality BRDF",
+		ToolTip = "Use the high quality brdf functions on mobile to get better visual effects but adds GPU cost. Changing this setting requires restarting the editor.",
+		ConfigRestartRequired = true))
+		uint32 bMobileEnableHighQualityBRDF : 1;
+
 	UPROPERTY(config, EditAnywhere, Category="Experimental|LOD Streaming|Skeletal Mesh", meta=(
 		DisplayName="Stream LODs by default",
 		ToolTip="Whether to stream skeletal mesh LODs by default."))
diff --git a/Engine/Source/Runtime/Engine/Private/ShaderCompiler/ShaderCompiler.cpp b/Engine/Source/Runtime/Engine/Private/ShaderCompiler/ShaderCompiler.cpp
--- a/Engine/Source/Runtime/Engine/Private/ShaderCompiler/ShaderCompiler.cpp
+++ b/Engine/Source/Runtime/Engine/Private/ShaderCompiler/ShaderCompiler.cpp
@@ -5101,7 +5101,9 @@ void GlobalBeginCompileShader(
 		Input.Environment.CompilerFlags.Add(CFLAG_ForceDXC);
 	}
 
-	if (IsMobilePlatform((EShaderPlatform)Target.Platform))
+	bool bIsMobilePlatform = IsMobilePlatform((EShaderPlatform)Target.Platform);
+
+	if (bIsMobilePlatform)
 	{
 		if (IsOpenGLPlatform((EShaderPlatform)Target.Platform))
 		{
@@ -5117,7 +5119,7 @@ void GlobalBeginCompileShader(
 
 	{
 		static const auto CVar = IConsoleManager::Get().FindTConsoleVariableDataInt(TEXT("r.ClearCoatNormal"));
-		Input.Environment.SetDefine(TEXT("CLEAR_COAT_BOTTOM_NORMAL"), CVar ? (CVar->GetValueOnAnyThread() != 0) : 0);
+		Input.Environment.SetDefine(TEXT("CLEAR_COAT_BOTTOM_NORMAL"), CVar ? (CVar->GetValueOnAnyThread() != 0) && !bIsMobilePlatform : 0);
 	}
 
 	{
diff --git a/Engine/Source/Runtime/Engine/Private/ShaderCompiler/ShaderGenerationUtil.cpp b/Engine/Source/Runtime/Engine/Private/ShaderCompiler/ShaderGenerationUtil.cpp
--- a/Engine/Source/Runtime/Engine/Private/ShaderCompiler/ShaderGenerationUtil.cpp
+++ b/Engine/Source/Runtime/Engine/Private/ShaderCompiler/ShaderGenerationUtil.cpp
@@ -256,6 +256,8 @@ static FShaderGlobalDefines FetchShaderGlobalDefines(EShaderPlatform TargetPlatf
 {
 	FShaderGlobalDefines Ret = {};
 
+	bool bIsMobilePlatform = IsMobilePlatform(TargetPlatform);
+
 	{
 		static const auto CVar = IConsoleManager::Get().FindTConsoleVariableDataInt(TEXT("r.AllowStaticLighting"));
 		Ret.ALLOW_STATIC_LIGHTING = (CVar ? (CVar->GetValueOnAnyThread() != 0) : 1);
@@ -266,7 +268,7 @@ static FShaderGlobalDefines FetchShaderGlobalDefines(EShaderPlatform TargetPlatf
 
 	{
 		static const auto CVar = IConsoleManager::Get().FindTConsoleVariableDataInt(TEXT("r.ClearCoatNormal"));
-		Ret.CLEAR_COAT_BOTTOM_NORMAL = CVar ? (CVar->GetValueOnAnyThread() != 0) : 0;
+		Ret.CLEAR_COAT_BOTTOM_NORMAL = CVar ? (CVar->GetValueOnAnyThread() != 0) && !bIsMobilePlatform : 0;
 	}
 
 	{
@@ -367,7 +369,7 @@ static FShaderGlobalDefines FetchShaderGlobalDefines(EShaderPlatform TargetPlatf
 		Ret.VIRTUAL_TEXTURE_FEEDBACK_FACTOR = (CVar ? FMath::Max(CVar->GetInt(), 1) : 1) != 0;
 	}
 
-	if (IsMobilePlatform((EShaderPlatform)TargetPlatform))
+	if (bIsMobilePlatform)
 	{
 		static IConsoleVariable* CVar = IConsoleManager::Get().FindConsoleVariable(TEXT("r.Mobile.EnableMovableSpotlights"));
 		Ret.PROJECT_MOBILE_ENABLE_MOVABLE_SPOTLIGHTS = CVar ? (CVar->GetInt() != 0) : 0;
diff --git a/Engine/Source/Runtime/RenderCore/Private/Shader.cpp b/Engine/Source/Runtime/RenderCore/Private/Shader.cpp
--- a/Engine/Source/Runtime/RenderCore/Private/Shader.cpp
+++ b/Engine/Source/Runtime/RenderCore/Private/Shader.cpp
@@ -1560,6 +1560,11 @@ void ShaderMapAppendKeyString(EShaderPlatform Platform, FString& KeyString)
 		{
 			KeyString += IsMobileDistanceFieldEnabled(Platform) ? TEXT("_MobSDF") : TEXT("");
 		}
+
+		{
+			static IConsoleVariable* MobileEnableHighQualityBRDFCVar = IConsoleManager::Get().FindConsoleVariable(TEXT("r.Mobile.EnableHighQualityBRDF"));
+			KeyString += (MobileEnableHighQualityBRDFCVar && MobileEnableHighQualityBRDFCVar->GetInt() != 0) ? TEXT("_HQBRDF") : TEXT("");
+		}
 	}
 	else
 	{
diff --git a/Engine/Source/Runtime/Renderer/Private/LightSceneInfo.cpp b/Engine/Source/Runtime/Renderer/Private/LightSceneInfo.cpp
--- a/Engine/Source/Runtime/Renderer/Private/LightSceneInfo.cpp
+++ b/Engine/Source/Runtime/Renderer/Private/LightSceneInfo.cpp
@@ -343,7 +343,7 @@ void FLightSceneInfo::ConditionalUpdateMobileMovablePointLightUniformBuffer(cons
 				ShadowFadeFraction = FMath::Max(ShadowFadeFraction, ProjectedShadowInfo->FadeAlphas[ViewIndex]);
 			}
 
-			SpotLightShadowSharpenAndShadowFadeFraction = FVector4(Proxy->GetShadowSharpen() * 7.0f + 1.0f, ShadowFadeFraction, 0.0f, 0.0f);
+			SpotLightShadowSharpenAndShadowFadeFraction = FVector4(Proxy->GetShadowSharpen() * 7.0f + 1.0f, ShadowFadeFraction, ProjectedShadowInfo->GetShaderReceiverDepthBias(), 0.0f);
 			SpotLightWorldToShadowMatrix = ProjectedShadowInfo->GetWorldToShadowMatrix(SpotLightShadowmapMinMax);
 		}
 		
diff --git a/Engine/Source/Runtime/Renderer/Private/MobileBasePassRendering.cpp b/Engine/Source/Runtime/Renderer/Private/MobileBasePassRendering.cpp
--- a/Engine/Source/Runtime/Renderer/Private/MobileBasePassRendering.cpp
+++ b/Engine/Source/Runtime/Renderer/Private/MobileBasePassRendering.cpp
@@ -278,6 +278,7 @@ void SetupMobileDirectionalLightUniformParameters(
 						Params.DirectionalLightShadowTexture = ShadowInfo->RenderTargets.DepthTarget->GetRenderTargetItem().ShaderResourceTexture.GetReference();
 						Params.DirectionalLightDirectionAndShadowTransition.W = 1.0f / ShadowInfo->ComputeTransitionSize();
 						Params.DirectionalLightShadowSize = ShadowBufferSizeValue;
+						Params.DirectionalLightDistanceFadeMADAndSpecularScale.W = ShadowInfo->GetShaderReceiverDepthBias();
 					}
 					Params.DirectionalLightScreenToShadow[OutShadowIndex] = ShadowInfo->GetScreenToShadowMatrix(SceneView);
 					Params.DirectionalLightShadowDistances[OutShadowIndex] = ShadowInfo->CascadeSettings.SplitFar;
@@ -397,4 +398,4 @@ void FMobileSceneRenderer::RenderMobileEditorPrimitives(FRHICommandList& RHICmdL
 		// Draw the view's batched simple elements(lines, sprites, etc).
 		View.TopBatchedViewElements.Draw(RHICmdList, DrawRenderState, FeatureLevel, bNeedToSwitchVerticalAxis, View, false);
 	}
-}
+}
\ No newline at end of file
diff --git a/Engine/Source/Runtime/Renderer/Private/MobileBasePassRendering.h b/Engine/Source/Runtime/Renderer/Private/MobileBasePassRendering.h
--- a/Engine/Source/Runtime/Renderer/Private/MobileBasePassRendering.h
+++ b/Engine/Source/Runtime/Renderer/Private/MobileBasePassRendering.h
@@ -442,6 +442,11 @@ public:
 		OutEnvironment.SetDefine(TEXT("ENABLE_AMBIENT_OCCLUSION"), IsMobileAmbientOcclusionEnabled(Parameters.Platform) ? 1u : 0u);
 
 		OutEnvironment.SetDefine(TEXT("ENABLE_DISTANCE_FIELD"), IsMobileDistanceFieldEnabled(Parameters.Platform));
+
+		static auto* MobileEnableHighQualityBRDFCVAR = IConsoleManager::Get().FindTConsoleVariableDataInt(TEXT("r.Mobile.EnableHighQualityBRDF"));
+		const bool bMobileEnableHighQualityBRDF = (MobileEnableHighQualityBRDFCVAR && MobileEnableHighQualityBRDFCVAR->GetValueOnAnyThread() == 1);
+
+		OutEnvironment.SetDefine(TEXT("MOBILE_HIGH_QUALITY_BRDF"), bMobileEnableHighQualityBRDF ? 1u : 0u);
 	}
 	
 	/** Initialization constructor. */
diff --git a/Engine/Source/Runtime/Renderer/Private/MobileDeferredShadingPass.cpp b/Engine/Source/Runtime/Renderer/Private/MobileDeferredShadingPass.cpp
--- a/Engine/Source/Runtime/Renderer/Private/MobileDeferredShadingPass.cpp
+++ b/Engine/Source/Runtime/Renderer/Private/MobileDeferredShadingPass.cpp
@@ -50,6 +50,8 @@ class FMobileDirectLightFunctionPS : public FMaterialShader
 		SHADER_PARAMETER(FMatrix44f, WorldToLight)
 		SHADER_PARAMETER(FVector4, LightFunctionParameters)
 		SHADER_PARAMETER(FVector3f, LightFunctionParameters2)
+		SHADER_PARAMETER_TEXTURE(Texture2D, PreIntegratedGF)
+		SHADER_PARAMETER_SAMPLER(SamplerState, PreIntegratedGFSampler)
 	END_SHADER_PARAMETER_STRUCT()
 
 	static void ModifyCompilationEnvironment(const FMaterialShaderPermutationParameters& Parameters, FShaderCompilerEnvironment& OutEnvironment)
@@ -137,6 +139,8 @@ public:
 		SHADER_PARAMETER_STRUCT_REF(FDeferredLightUniformStruct, DeferredLightUniforms)
 		SHADER_PARAMETER_TEXTURE(Texture2D, IESTexture)
 		SHADER_PARAMETER_SAMPLER(SamplerState, IESTextureSampler)
+		SHADER_PARAMETER_TEXTURE(Texture2D, PreIntegratedGF)
+		SHADER_PARAMETER_SAMPLER(SamplerState, PreIntegratedGFSampler)
 	END_SHADER_PARAMETER_STRUCT()
 
 	static bool ShouldCompilePermutation(const FMaterialShaderPermutationParameters& Parameters)
@@ -260,6 +264,9 @@ static void RenderDirectLight(FRHICommandListImmediate& RHICmdList, const FScene
 	}
 	PassParameters.PlanarReflection = TUniformBufferRef<FPlanarReflectionUniformParameters>::CreateUniformBufferImmediate(PlanarReflectionUniformParameters, UniformBuffer_SingleDraw);
 
+	PassParameters.PreIntegratedGF = GSystemTextures.PreintegratedGF->GetRenderTargetItem().ShaderResourceTexture;
+	PassParameters.PreIntegratedGFSampler = TStaticSamplerState<SF_Bilinear, AM_Clamp, AM_Clamp, AM_Clamp>::GetRHI();
+
 	if (DirectionalLight)
 	{
 		const bool bUseMovableLight = DirectionalLight && !DirectionalLight->Proxy->HasStaticShadowing();
@@ -443,6 +450,8 @@ static void RenderLocalLight(
 	// Switch x and z so that z of the user specified scale affects the distance along the light direction
 	const FVector InverseScale = FVector(1.f / Scale.Z, 1.f / Scale.Y, 1.f / Scale.X);
 	PassParameters.WorldToLight = LightSceneInfo.Proxy->GetWorldToLight() * FScaleMatrix(FVector(InverseScale));
+	PassParameters.PreIntegratedGF = GSystemTextures.PreintegratedGF->GetRenderTargetItem().ShaderResourceTexture;
+	PassParameters.PreIntegratedGFSampler = TStaticSamplerState<SF_Bilinear, AM_Clamp, AM_Clamp, AM_Clamp>::GetRHI();
 	FMobileRadialLightFunctionPS::SetParameters(RHICmdList, PixelShader, View, LightMaterial.MaterialProxy, *LightMaterial.Material, PassParameters);
 
 	// Shade only MSM_DefaultLit pixels
