Fix ue5 CIS for the field visualizer

diff --git a/Engine/Shaders/Private/HairStrands/HairCardsVoxel.usf b/Engine/Shaders/Private/HairStrands/HairCardsVoxel.usf
--- a/Engine/Shaders/Private/HairStrands/HairCardsVoxel.usf
+++ b/Engine/Shaders/Private/HairStrands/HairCardsVoxel.usf
@@ -5,7 +5,7 @@
 #include "../SceneTextureParameters.ush"
 #include "../DeferredShadingCommon.ush"
 #include "../ShaderDrawDebug.ush"
-
+#include "../PositionReconstructionCommon.ush" 
 
 ///////////////////////////////////////////////////////////////////////////////////////////////////
 uint3 PositionToCoord(float3 P, float3 InMinAABB, float3 InMaxAABB, uint3 InResolution)
@@ -50,15 +50,6 @@ bool IsInVoxelBounds(float3 P, float3 MinP, float3 MaxP)
 		P.x <= MaxP.x && P.y <= MaxP.y && P.z <= MaxP.z;
 }
 
-float3 ReconstructWorldPositionFromDepth(float2 UV, float Depth)
-{
-	float2 ScreenPosition = (UV - View.ScreenPositionScaleBias.wz) / View.ScreenPositionScaleBias.xy;
-	float4 HomogeneousWorldPosition = mul(float4(ScreenPosition * Depth, Depth, 1), View.ScreenToWorld);
-	float3 WorldPosition = HomogeneousWorldPosition.xyz / HomogeneousWorldPosition.w;
-
-	return WorldPosition;
-}
-
 float3 EncodeTangent(float3 N)
 {
 	return N * 0.5 + 0.5;
diff --git a/Engine/Shaders/Private/HairStrands/HairScatter.usf b/Engine/Shaders/Private/HairStrands/HairScatter.usf
--- a/Engine/Shaders/Private/HairStrands/HairScatter.usf
+++ b/Engine/Shaders/Private/HairStrands/HairScatter.usf
@@ -24,16 +24,7 @@
 #include "HairStrandsVoxelPageTraversal.ush"
 #include "../HairBsdf.ush"
 #include "../ShaderDrawDebug.ush"
-
-////////////////////////////////////////////////////////////////////////////////////////////////////////////////////////////////////////////////
-
-float3 ReconstructWorldPositionFromDepth(float2 UV, float Depth)
-{
-	float2 ScreenPosition = (UV - View.ScreenPositionScaleBias.wz) / View.ScreenPositionScaleBias.xy;
-	float4 HomogeneousWorldPosition = mul(float4(ScreenPosition * Depth, Depth, 1), View.ScreenToWorld);
-	float3 WorldPosition = HomogeneousWorldPosition.xyz / HomogeneousWorldPosition.w;
-	return WorldPosition;
-}
+#include "../PositionReconstructionCommon.ush" 
 
 ////////////////////////////////////////////////////////////////////////////////////////////////////////////////////////////////////////////////
 
diff --git a/Engine/Shaders/Private/HairStrands/HairStrandsVoxelOpaque.usf b/Engine/Shaders/Private/HairStrands/HairStrandsVoxelOpaque.usf
--- a/Engine/Shaders/Private/HairStrands/HairStrandsVoxelOpaque.usf
+++ b/Engine/Shaders/Private/HairStrands/HairStrandsVoxelOpaque.usf
@@ -4,6 +4,7 @@
 #include "../Common.ush"
 #include "../CommonViewUniformBuffer.ush"
 #include "../SceneTextureParameters.ush"
+#include "../PositionReconstructionCommon.ush" 
 #include "HairStrandsVoxelPageCommonStruct.ush"
 #include "HairStrandsVoxelPageCommon.ush"
 
@@ -23,15 +24,6 @@ Buffer<uint> VirtualVoxelParams_PageIndexBuffer;
 Buffer<uint4>VirtualVoxelParams_PageIndexCoordBuffer;
 StructuredBuffer<FPackedVirtualVoxelNodeDesc> VirtualVoxelParams_NodeDescBuffer;
 
-float3 ReconstructWorldPositionFromDepth(float2 UV, float Depth)
-{
-	float2 ScreenPosition = (UV - View.ScreenPositionScaleBias.wz) / View.ScreenPositionScaleBias.xy;
-	float4 HomogeneousWorldPosition = mul(float4(ScreenPosition * Depth, Depth, 1), View.ScreenToWorld);
-	float3 WorldPosition = HomogeneousWorldPosition.xyz / HomogeneousWorldPosition.w;
-
-	return WorldPosition;
-}
-
 #define GROUP_SIZE 64
 [numthreads(1, 1, GROUP_SIZE)]
 void MainCS(uint3 DispatchThreadId : SV_DispatchThreadID)
diff --git a/Engine/Shaders/Private/HairStrands/HairStrandsVoxelPageRayMarching.usf b/Engine/Shaders/Private/HairStrands/HairStrandsVoxelPageRayMarching.usf
--- a/Engine/Shaders/Private/HairStrands/HairStrandsVoxelPageRayMarching.usf
+++ b/Engine/Shaders/Private/HairStrands/HairStrandsVoxelPageRayMarching.usf
@@ -8,6 +8,7 @@
 #include "../DeferredShadingCommon.ush"
 #include "../ShaderDrawDebug.ush"
 #include "../ShaderPrintCommon.ush"
+#include "../PositionReconstructionCommon.ush"  
 #include "HairStrandsVisibilityCommon.ush"
 #include "HairStrandsVoxelPageCommon.ush"
 #include "../ColorMap.ush"  
@@ -32,15 +33,6 @@ uint MaxTotalPageIndexCount;
 RWTexture2D<float4> OutputTexture;
 Buffer<uint> TotalValidPageCounter;
 
-float3 ReconstructWorldPositionFromDepth(float2 UV, float Depth)
-{
-	float2 ScreenPosition = (UV - View.ScreenPositionScaleBias.wz) / View.ScreenPositionScaleBias.xy;
-	float4 HomogeneousWorldPosition = mul(float4(ScreenPosition * Depth, Depth, 1), View.ScreenToWorld);
-	float3 WorldPosition = HomogeneousWorldPosition.xyz / HomogeneousWorldPosition.w;
-
-	return WorldPosition;
-}
-
 #define VOXEL_BLOCK 0 
 #define VOXEL_SMOOTH 1
 #define VOXEL_FILTERING VOXEL_BLOCK
diff --git a/Engine/Shaders/Private/PhysicsFieldVisualizer.usf b/Engine/Shaders/Private/PhysicsFieldVisualizer.usf
--- a/Engine/Shaders/Private/PhysicsFieldVisualizer.usf
+++ b/Engine/Shaders/Private/PhysicsFieldVisualizer.usf
@@ -7,10 +7,9 @@
 #include "SceneTextureParameters.ush"
 #include "DeferredShadingCommon.ush"
 #include "ShaderDrawDebug.ush"
-#include "ShaderPrintCommon.ush"
-#include "RayTracing/RayTracingCommon.ush" 
+#include "ShaderPrintCommon.ush" 
+#include "PositionReconstructionCommon.ush" 
  
-
 #include "PhysicsFieldShared.ush"
 #include "PhysicsFieldEval.ush"
 
diff --git a/Engine/Shaders/Private/PositionReconstructionCommon.ush b/Engine/Shaders/Private/PositionReconstructionCommon.ush
new file mode 100644
--- /dev/null
+++ b/Engine/Shaders/Private/PositionReconstructionCommon.ush
@@ -0,0 +1,38 @@
+// Copyright Epic Games, Inc. All Rights Reserved.
+
+/*=============================================================================
+	PositionReconstructionCommon.ush: utils functions used in multiple ray generation shaders
+=============================================================================*/
+
+#pragma once
+
+float3 ReconstructWorldPositionFromDeviceZ(uint2 PixelCoord, float DeviceZ)
+{
+	float4 TranslatedWorldPosition = mul(float4(PixelCoord + 0.5, DeviceZ, 1), View.SVPositionToTranslatedWorld);
+	TranslatedWorldPosition.xyz /= TranslatedWorldPosition.w;
+	return TranslatedWorldPosition.xyz + View.WorldCameraOrigin;
+}
+
+void ReconstructWorldPositionAndCameraDirectionFromDeviceZ(uint2 PixelCoord, float DeviceZ, out float3 OutWorldPosition, out float3 OutCameraDirection)
+{
+	float4 TranslatedWorldPosition = mul(float4(PixelCoord + 0.5, DeviceZ, 1), View.SVPositionToTranslatedWorld);
+	TranslatedWorldPosition.xyz /= TranslatedWorldPosition.w;
+	OutWorldPosition = TranslatedWorldPosition.xyz + View.WorldCameraOrigin;
+	OutCameraDirection = normalize(TranslatedWorldPosition.xyz - View.TranslatedWorldCameraOrigin);
+}
+
+float3 ReconstructTranslatedWorldPositionFromDeviceZ(uint2 PixelCoord, float DeviceZ)
+{
+	float4 TranslatedWorldPosition = mul(float4(PixelCoord + 0.5, DeviceZ, 1), View.SVPositionToTranslatedWorld);
+	TranslatedWorldPosition.xyz /= TranslatedWorldPosition.w;
+	return TranslatedWorldPosition.xyz;
+}
+ 
+float3 ReconstructWorldPositionFromDepth(float2 UV, float Depth)
+{
+	float2 ScreenPosition = (UV - View.ScreenPositionScaleBias.wz) / View.ScreenPositionScaleBias.xy;
+	float4 HomogeneousWorldPosition = mul(float4(ScreenPosition * Depth, Depth, 1), View.ScreenToWorld);
+	float3 WorldPosition = HomogeneousWorldPosition.xyz / HomogeneousWorldPosition.w;
+
+	return WorldPosition;
+}
diff --git a/Engine/Shaders/Private/RayTracing/RayTracingCommon.ush b/Engine/Shaders/Private/RayTracing/RayTracingCommon.ush
--- a/Engine/Shaders/Private/RayTracing/RayTracingCommon.ush
+++ b/Engine/Shaders/Private/RayTracing/RayTracingCommon.ush
@@ -210,36 +210,8 @@ FRayCone PropagateRayCone(in FRayCone Cone, in float SurfaceSpreadAngle, in floa
 	return NewCone;
 }
 
-float3 ReconstructWorldPositionFromDeviceZ(uint2 PixelCoord, float DeviceZ)
-{
-	float4 TranslatedWorldPosition = mul(float4(PixelCoord + 0.5, DeviceZ, 1), View.SVPositionToTranslatedWorld);
-	TranslatedWorldPosition.xyz /= TranslatedWorldPosition.w;
-	return TranslatedWorldPosition.xyz + View.WorldCameraOrigin;
-}
-
-void ReconstructWorldPositionAndCameraDirectionFromDeviceZ(uint2 PixelCoord, float DeviceZ, out float3 OutWorldPosition, out float3 OutCameraDirection)
-{
-	float4 TranslatedWorldPosition = mul(float4(PixelCoord + 0.5, DeviceZ, 1), View.SVPositionToTranslatedWorld);
-	TranslatedWorldPosition.xyz /= TranslatedWorldPosition.w;
-	OutWorldPosition = TranslatedWorldPosition.xyz + View.WorldCameraOrigin;
-	OutCameraDirection = normalize(TranslatedWorldPosition.xyz - View.TranslatedWorldCameraOrigin);
-}
-
-float3 ReconstructTranslatedWorldPositionFromDeviceZ(uint2 PixelCoord, float DeviceZ)
-{
-	float4 TranslatedWorldPosition = mul(float4(PixelCoord + 0.5, DeviceZ, 1), View.SVPositionToTranslatedWorld);
-	TranslatedWorldPosition.xyz /= TranslatedWorldPosition.w;
-	return TranslatedWorldPosition.xyz;
-}
-
-float3 ReconstructWorldPositionFromDepth(float2 UV, float Depth)
-{
-	float2 ScreenPosition = (UV - View.ScreenPositionScaleBias.wz) / View.ScreenPositionScaleBias.xy;
-	float4 HomogeneousWorldPosition = mul(float4(ScreenPosition * Depth, Depth, 1), View.ScreenToWorld);
-	float3 WorldPosition = HomogeneousWorldPosition.xyz / HomogeneousWorldPosition.w;
-
-	return WorldPosition;
-}
+// Include utility fonction for world position reconstruction 
+#include "../PositionReconstructionCommon.ush" 
 
 uint CalcLinearIndex(uint2 PixelCoord)
 {
