Fix for virtual texture sampling reading unbound View uniform buffer in texture editor.

diff --git a/Engine/Shaders/Private/SimpleElementTexture2DPreviewPixelShader.usf b/Engine/Shaders/Private/SimpleElementTexture2DPreviewPixelShader.usf
--- a/Engine/Shaders/Private/SimpleElementTexture2DPreviewPixelShader.usf
+++ b/Engine/Shaders/Private/SimpleElementTexture2DPreviewPixelShader.usf
@@ -18,9 +18,8 @@
 
 #if SAMPLE_VIRTUAL_TEXTURE
 	#define NUM_VIRTUALTEXTURE_SAMPLES 1
-	// Workaround for unbound View uniform buffer which is used for LOD randomization.
-	// View.StateFrameIndexMod8 is used as random seed.
-	#define VT_RANDOMIZE_LOD 0
+	// Workaround for unbound View uniform buffer
+	#define VT_DISABLE_VIEW_UNIFORM_BUFFER 1
 	#include "VirtualTextureCommon.ush"
 #endif
 
diff --git a/Engine/Shaders/Private/VirtualTextureCommon.ush b/Engine/Shaders/Private/VirtualTextureCommon.ush
--- a/Engine/Shaders/Private/VirtualTextureCommon.ush
+++ b/Engine/Shaders/Private/VirtualTextureCommon.ush
@@ -5,8 +5,8 @@
 #include "Common.ush"
 #include "Random.ush"
 
-#ifndef VT_RANDOMIZE_LOD
-#define VT_RANDOMIZE_LOD 1
+#ifndef VT_DISABLE_VIEW_UNIFORM_BUFFER
+#define VT_DISABLE_VIEW_UNIFORM_BUFFER 0
 #endif
 
 #ifndef NUM_VIRTUALTEXTURE_SAMPLES
@@ -200,6 +200,24 @@ struct VTPageTableResult
 	uint PackedRequest;
 };
 
+float GetStochasticMipNoise(float2 SvPositionXY)
+{
+#if VT_DISABLE_VIEW_UNIFORM_BUFFER
+	return 0;
+#else
+	return InterleavedGradientNoise(SvPositionXY, View.StateFrameIndexMod8);
+#endif
+}
+
+float GetGlobalVirtualTextureMipBias()
+{
+#if VT_DISABLE_VIEW_UNIFORM_BUFFER
+	return 0;
+#else
+	return View.GlobalVirtualTextureMipBias;
+#endif
+}
+
 /** Calculate mip level including stochastic noise. Also stores derivatives to OutResult for use when sampling physical texture. */
 int TextureComputeVirtualMipLevel(
 	in out VTPageTableResult OutResult,
@@ -214,13 +232,10 @@ int TextureComputeVirtualMipLevel(
 	// This way the VT mip level selection will come much closer to HW mip selection, even if we're not sampling the texture using anisotropic filtering
 	const float ComputedLevel = MipLevelAniso2D(OutResult.dUVdx, OutResult.dUVdy, PageTableUniform.MaxAnisoLog2);
 
-#if VT_RANDOMIZE_LOD
-	const float Noise = InterleavedGradientNoise(SvPositionXY, View.StateFrameIndexMod8);
-#else
-	const float Noise = 0;
-#endif
+	const float GlobalMipBias = GetGlobalVirtualTextureMipBias();
+	const float Noise = GetStochasticMipNoise(SvPositionXY);
 
-	return (int)floor(ComputedLevel + MipBias + View.GlobalVirtualTextureMipBias + Noise) + int(PageTableUniform.vPageTableMipBias);
+	return (int)floor(ComputedLevel + MipBias + GlobalMipBias + Noise) + int(PageTableUniform.vPageTableMipBias);
 }
 
 /** Samples page table indirection and any apply changes to UV, vLevel and PageTable information for adaptive page tables. */
@@ -332,7 +347,7 @@ VTPageTableResult TextureLoadVirtualPageTable(
 {
 	VTPageTableResult Result = (VTPageTableResult)0;
 	UV = UV * PageTableUniform.UVScale;
-	int vLevel = View.GlobalVirtualTextureMipBias;
+	int vLevel = GetGlobalVirtualTextureMipBias();
 #if PIXELSHADER
 	vLevel = TextureComputeVirtualMipLevel(Result, ddx(UV), ddy(UV), MipBias, SvPositionXY, PageTableUniform);
 #endif // PIXELSHADER
@@ -351,7 +366,7 @@ VTPageTableResult TextureLoadVirtualPageTable(
 {
 	VTPageTableResult Result = (VTPageTableResult)0;
 	UV = UV * PageTableUniform.UVScale;
-	int vLevel = View.GlobalVirtualTextureMipBias;
+	int vLevel = GetGlobalVirtualTextureMipBias();
 #if PIXELSHADER
 	vLevel = TextureComputeVirtualMipLevel(Result, ddx(UV), ddy(UV), MipBias, SvPositionXY, PageTableUniform);
 #endif // PIXELSHADER
@@ -371,7 +386,7 @@ VTPageTableResult TextureLoadVirtualPageTable(
 {
 	VTPageTableResult Result = (VTPageTableResult)0;
 	UV = UV * PageTableUniform.UVScale;
-	int vLevel = View.GlobalVirtualTextureMipBias;
+	int vLevel = GetGlobalVirtualTextureMipBias();
 #if PIXELSHADER
 	vLevel = TextureComputeVirtualMipLevel(Result, ddx(UV), ddy(UV), MipBias, SvPositionXY, PageTableUniform);
 #endif // PIXELSHADER
@@ -390,7 +405,7 @@ VTPageTableResult TextureLoadVirtualPageTable(
 {
 	VTPageTableResult Result = (VTPageTableResult)0;
 	UV = UV * PageTableUniform.UVScale;
-	int vLevel = View.GlobalVirtualTextureMipBias;
+	int vLevel = GetGlobalVirtualTextureMipBias();
 #if PIXELSHADER
 	vLevel = TextureComputeVirtualMipLevel(Result, ddx(UV), ddy(UV), MipBias, SvPositionXY, PageTableUniform);
 #endif // PIXELSHADER
@@ -446,7 +461,7 @@ VTPageTableResult TextureLoadVirtualPageTableLevel(
 {
 	VTPageTableResult Result = (VTPageTableResult)0;
 	UV = UV * PageTableUniform.UVScale;
-	int vLevel = (int)floor(Level + View.GlobalVirtualTextureMipBias);
+	int vLevel = (int)floor(Level + GetGlobalVirtualTextureMipBias());
 	UV = ApplyAddressMode(UV, AddressU, AddressV);
 	TextureLoadVirtualPageTableInternal(Result, PageTable0, PageTableUniform, UV, vLevel);
 	StoreVirtualTextureFeedback(Feedback, SampleIndex, Result.PackedRequest);
@@ -462,7 +477,7 @@ VTPageTableResult TextureLoadVirtualPageTableLevel(
 {
 	VTPageTableResult Result = (VTPageTableResult)0;
 	UV = UV * PageTableUniform.UVScale;
-	int vLevel = (int)floor(Level + View.GlobalVirtualTextureMipBias);
+	int vLevel = (int)floor(Level + GetGlobalVirtualTextureMipBias());
 	UV = ApplyAddressMode(UV, AddressU, AddressV);
 	TextureLoadVirtualPageTableInternal(Result, PageTable0, PageTableUniform, UV, vLevel);
 	return Result;
@@ -479,7 +494,7 @@ VTPageTableResult TextureLoadVirtualPageTableLevel(
 {
 	VTPageTableResult Result = (VTPageTableResult)0;
 	UV = UV * PageTableUniform.UVScale;
-	int vLevel = (int)floor(Level + View.GlobalVirtualTextureMipBias);
+	int vLevel = (int)floor(Level + GetGlobalVirtualTextureMipBias());
 	UV = ApplyAddressMode(UV, AddressU, AddressV);
 	TextureLoadVirtualPageTableInternal(Result, PageTable0, PageTable1, PageTableUniform, UV, vLevel);
 	StoreVirtualTextureFeedback(Feedback, SampleIndex, Result.PackedRequest);
@@ -495,7 +510,7 @@ VTPageTableResult TextureLoadVirtualPageTableLevel(
 {
 	VTPageTableResult Result = (VTPageTableResult)0;
 	UV = UV * PageTableUniform.UVScale;
-	int vLevel = (int)floor(Level + View.GlobalVirtualTextureMipBias);
+	int vLevel = (int)floor(Level + GetGlobalVirtualTextureMipBias());
 	UV = ApplyAddressMode(UV, AddressU, AddressV);
 	TextureLoadVirtualPageTableInternal(Result, PageTable0, PageTable1, PageTableUniform, UV, vLevel);
 	return Result;
@@ -513,7 +528,7 @@ VTPageTableResult TextureLoadVirtualPageTableAdaptive(
 {
 	VTPageTableResult Result = (VTPageTableResult)0;
 	UV = UV * PageTableUniform.UVScale;
-	int vLevel = View.GlobalVirtualTextureMipBias;
+	int vLevel = GetGlobalVirtualTextureMipBias();
 #if PIXELSHADER
 	vLevel = TextureComputeVirtualMipLevel(Result, ddx(UV), ddy(UV), MipBias, SvPositionXY, PageTableUniform);
 #endif // PIXELSHADER
@@ -556,7 +571,7 @@ VTPageTableResult TextureLoadVirtualPageTableAdaptiveLevel(
 {
 	VTPageTableResult Result = (VTPageTableResult)0;
 	// Level is an index into the full size adaptive VT. AdaptiveLevelBias shifts it relative to the low mips allocated VT.
-	int vLevel = (int)floor(max(Level, View.GlobalVirtualTextureMipBias)) - PageTableUniform.AdaptiveLevelBias;
+	int vLevel = (int)floor(Level + GetGlobalVirtualTextureMipBias()) - PageTableUniform.AdaptiveLevelBias;
 	UV = UV * PageTableUniform.UVScale;
 	UV = ApplyAddressMode(UV, AddressU, AddressV);
 	ApplyAdaptivePageTableUniform(PageTableIndirection, Result, PageTableUniform, UV, vLevel);
@@ -575,7 +590,7 @@ VTPageTableResult TextureLoadVirtualPageTableAdaptiveLevel(
 {
 	VTPageTableResult Result = (VTPageTableResult)0;
 	// Level is an index into the full size adaptive VT. AdaptiveLevelBias shifts it relative to the low mips allocated VT.
-	int vLevel = (int)floor(max(Level, View.GlobalVirtualTextureMipBias)) - PageTableUniform.AdaptiveLevelBias;
+	int vLevel = (int)floor(Level + GetGlobalVirtualTextureMipBias()) - PageTableUniform.AdaptiveLevelBias;
 	UV = UV * PageTableUniform.UVScale;
 	UV = ApplyAddressMode(UV, AddressU, AddressV);
 	ApplyAdaptivePageTableUniform(PageTableIndirection, Result, PageTableUniform, UV, vLevel);
