Change scene scattered hair to not receive env lighting, since they alreay received lighting from the scene color.

diff --git a/Engine/Shaders/Private/HairStrands/HairStrandsEnvironmentLighting.usf b/Engine/Shaders/Private/HairStrands/HairStrandsEnvironmentLighting.usf
--- a/Engine/Shaders/Private/HairStrands/HairStrandsEnvironmentLighting.usf
+++ b/Engine/Shaders/Private/HairStrands/HairStrandsEnvironmentLighting.usf
@@ -803,6 +803,8 @@ void MainPS(
 
 #if PERMUTATION_INTEGRATION_TYPE == INTEGRATION_SCENECOLOR
 	if (Sample.bScatterSceneLighting)
+#else
+	if (!Sample.bScatterSceneLighting)
 #endif
 	{
 		const float3 SkyLighting =
