Implements TSR's spatial anti-aliasing for history rejection

diff --git a/Engine/Config/BaseScalability.ini b/Engine/Config/BaseScalability.ini
--- a/Engine/Config/BaseScalability.ini
+++ b/Engine/Config/BaseScalability.ini
@@ -36,18 +36,24 @@ PerfIndexValues_ResolutionQuality="50 71 87 100 100"
 
 [AntiAliasingQuality@0]
 r.PostProcessAAQuality=0
+r.TSR.RejectionAntiAliasingQuality=0
 
 [AntiAliasingQuality@1]
 r.PostProcessAAQuality=2
+r.TSR.RejectionAntiAliasingQuality=0
 
 [AntiAliasingQuality@2]
 r.PostProcessAAQuality=3
+r.TSR.RejectionAntiAliasingQuality=1
 
 [AntiAliasingQuality@3]
 r.PostProcessAAQuality=4
+r.TSR.RejectionAntiAliasingQuality=1
 
 [AntiAliasingQuality@Cine]
 r.PostProcessAAQuality=6
+r.TSR.RejectionAntiAliasingQuality=2
+
 
 ;-----------------------------------------------------------------------------------------------------------------
 
diff --git a/Engine/Shaders/Private/TemporalSuperResolution/TSRCommon.ush b/Engine/Shaders/Private/TemporalSuperResolution/TSRCommon.ush
--- a/Engine/Shaders/Private/TemporalSuperResolution/TSRCommon.ush
+++ b/Engine/Shaders/Private/TemporalSuperResolution/TSRCommon.ush
@@ -27,13 +27,6 @@
 
 #define CONFIG_ENABLE_STOCASTIC_QUANTIZATION 1
 
-#if CONFIG_COMPILE_FP16
-	// Take advantage of RDNA's v_pk_*_{uif}16 instructions
-	#define CONFIG_ENABLE_DUAL_PIXEL_VECTORIZATION 1
-#else
-	#define CONFIG_ENABLE_DUAL_PIXEL_VECTORIZATION 0
-#endif
-
 
 //------------------------------------------------------- CONSTANTS
 
@@ -42,6 +35,10 @@
 
 #define MAX_FALLBACK_SAMPLE_COUNT 20
 
+// Spatial anti-aliasing encoding settings.
+#define SPATIAL_ANTI_ALIASING_OFFSET_RANGE 0.5
+#define SPATIAL_ANTI_ALIASING_OFFSET_BITDEPTH_PER_CHANNEL 4
+
 #define COMPRESS_PREV_USE_COUNT 1
 #if COMPRESS_PREV_USE_COUNT
 	// COMPRESS_PREV_USE_COUNT compress the use count in 8bit, so need to sum of a 2x2 to not overflow these 8bits
@@ -968,3 +965,40 @@ void ComputeCompressedUseCountPixelCoordinates(uint2 PixelPos, out uint2 Compres
 }
 	
 #endif
+
+uint2 EncodeSpatialAntiAliasingOffset(taa_half2x2 TexelOffset)
+{
+	const uint Mask = (1 << SPATIAL_ANTI_ALIASING_OFFSET_BITDEPTH_PER_CHANNEL) - 1;
+	const uint Zero = Mask / 2;
+
+	taa_half Multiply = taa_half(float(Zero) / SPATIAL_ANTI_ALIASING_OFFSET_RANGE);
+
+	uint2 EncodedTexelOffset = 0;
+	EncodedTexelOffset |= clamp(uint2(int2(round(TexelOffset[0] * Multiply)) + int(1 + Zero)), uint(1), Mask) << 0;
+	EncodedTexelOffset |= clamp(uint2(int2(round(TexelOffset[1] * Multiply)) + int(1 + Zero)), uint(1), Mask) << SPATIAL_ANTI_ALIASING_OFFSET_BITDEPTH_PER_CHANNEL;
+	return EncodedTexelOffset;
+}
+
+uint EncodeSpatialAntiAliasingOffset(taa_half2 TexelOffset)
+{
+	return dpv_lo(EncodeSpatialAntiAliasingOffset(dpv_interleave_mono_registers(TexelOffset)));
+}
+
+taa_half2x2 DecodeSpatialAntiAliasingOffset(uint2 EncodedInputTexelOffset)
+{
+	const uint Mask = (1 << SPATIAL_ANTI_ALIASING_OFFSET_BITDEPTH_PER_CHANNEL) - 1;
+	const uint Zero = Mask / 2;
+	
+	taa_half Multiply = taa_half(SPATIAL_ANTI_ALIASING_OFFSET_RANGE / float(Zero));
+
+	taa_half2x2 TexelOffset;
+	TexelOffset[0] = taa_half2((EncodedInputTexelOffset >> 0                                                ) & Mask) * Multiply - (taa_half(1 + Zero) * Multiply);
+	TexelOffset[1] = taa_half2((EncodedInputTexelOffset >> SPATIAL_ANTI_ALIASING_OFFSET_BITDEPTH_PER_CHANNEL) & Mask) * Multiply - (taa_half(1 + Zero) * Multiply);
+
+	return TexelOffset;
+}
+
+taa_half2 DecodeSpatialAntiAliasingOffset(uint EncodedInputTexelOffset)
+{
+	return dpv_lo(DecodeSpatialAntiAliasingOffset(dpv_interleave_mono_registers(EncodedInputTexelOffset)));
+}
diff --git a/Engine/Shaders/Private/TemporalSuperResolution/TSRFilterAntiAliasing.usf b/Engine/Shaders/Private/TemporalSuperResolution/TSRFilterAntiAliasing.usf
new file mode 100644
--- /dev/null
+++ b/Engine/Shaders/Private/TemporalSuperResolution/TSRFilterAntiAliasing.usf
@@ -0,0 +1,126 @@
+// Copyright Epic Games, Inc. All Rights Reserved.
+
+#include "TSRCommon.ush"
+
+
+//------------------------------------------------------- CONFIG
+
+#define TILE_SIZE 8
+
+
+
+//------------------------------------------------------- PARAMETERS
+
+Texture2D<uint> AntiAliasingTexture;
+RWTexture2D<uint> AntiAliasingOutput;
+
+
+//------------------------------------------------------- FUNCTIONS
+
+taa_half2 SampleAntiAliasing(taa_short2 KernelCenter, taa_short2 Offset)
+{
+	taa_short2 SampleInputPixelPos = ClampPixelOffset(
+		KernelCenter + Offset,
+		InputPixelPosMin, InputPixelPosMax);
+	
+	uint EncodedInputTexelOffset = AntiAliasingTexture[SampleInputPixelPos];
+
+	return DecodeSpatialAntiAliasingOffset(EncodedInputTexelOffset);
+}
+
+//------------------------------------------------------- ENTRY POINT
+
+[numthreads(TILE_SIZE, TILE_SIZE, 1)]
+void MainCS(
+	uint2 GroupId : SV_GroupID,
+	uint GroupThreadIndex : SV_GroupIndex)
+{
+	uint GroupWaveIndex = GetGroupWaveIndex(GroupThreadIndex, /* GroupSize = */ TILE_SIZE * TILE_SIZE);
+
+	float4 Debug = 0.0;
+	
+	taa_short2 InputPixelPos = (
+		taa_short2(InputPixelPosMin) +
+		taa_short2(GroupId) * taa_short2(TILE_SIZE, TILE_SIZE) +
+		Map8x8Tile2x2Lane(GroupThreadIndex));
+	
+	taa_half2 TexelOffsetC = SampleAntiAliasing(InputPixelPos, taa_short2( 0,  0));
+	
+	Debug.rg = abs(TexelOffsetC);
+
+	taa_half2 TexelOffsetN = SampleAntiAliasing(InputPixelPos, taa_short2( 0, -1));
+	taa_half2 TexelOffsetS = SampleAntiAliasing(InputPixelPos, taa_short2( 0, +1));
+	taa_half2 TexelOffsetE = SampleAntiAliasing(InputPixelPos, taa_short2(+1,  0));
+	taa_half2 TexelOffsetW = SampleAntiAliasing(InputPixelPos, taa_short2(-1,  0));
+	
+	taa_half2 TexelOffsetNE = SampleAntiAliasing(InputPixelPos, taa_short2(+1, -1));
+	taa_half2 TexelOffsetNW = SampleAntiAliasing(InputPixelPos, taa_short2(-1, -1));
+	taa_half2 TexelOffsetSE = SampleAntiAliasing(InputPixelPos, taa_short2(+1, +1));
+	taa_half2 TexelOffsetSW = SampleAntiAliasing(InputPixelPos, taa_short2(-1, +1));
+	
+	FLATTEN
+	if (TexelOffsetC.x > 0.0)
+	{
+		TexelOffsetC.x = TexelOffsetE.x < 0.0 ? 0.0 : TexelOffsetC.x;
+	}
+	else if (TexelOffsetC.x < 0.0)
+	{
+		TexelOffsetC.x = TexelOffsetW.x > 0.0 ? 0.0 : TexelOffsetC.x;
+	}
+	else if (TexelOffsetC.y > 0.0)
+	{
+		TexelOffsetC.y = TexelOffsetS.y < 0.0 ? 0.0 : TexelOffsetC.y;
+	}
+	else if (TexelOffsetC.y < 0.0)
+	{
+		TexelOffsetC.y = TexelOffsetN.y > 0.0 ? 0.0 : TexelOffsetC.y;
+	}
+#if 0
+	else
+	{
+		TexelOffsetC = TexelOffsetN;
+
+		FLATTEN
+		if (length2(TexelOffsetS) < length2(TexelOffsetC))
+		{
+			TexelOffsetC = TexelOffsetS;
+		}
+		
+		FLATTEN
+		if (length2(TexelOffsetE) < length2(TexelOffsetC))
+		{
+			TexelOffsetC = TexelOffsetE;
+		}
+		
+		FLATTEN
+		if (length2(TexelOffsetW) < length2(TexelOffsetC))
+		{
+			TexelOffsetC = TexelOffsetW;
+		}
+	}
+#endif
+
+	{
+		taa_short2 OutputPixelPos;
+		#if 0
+			ISOLATE
+			{
+				OutputPixelPos = dpv_add(
+					taa_short2(InputPixelPosMin) + taa_short2(GroupId) * taa_short2(TILE_SIZE, TILE_SIZE),
+					taa_short2x2(Map8x8Tile2x2LaneDPV(GroupThreadIndex)));
+			}
+		#else
+			OutputPixelPos = InputPixelPos;
+		#endif
+		
+		OutputPixelPos = InvalidateOutputPixelPos(OutputPixelPos, InputInfo_ViewportMax);
+		
+		AntiAliasingOutput[OutputPixelPos] = EncodeSpatialAntiAliasingOffset(TexelOffsetC);
+		
+		#if DEBUG_OUTPUT
+		{
+			DebugOutput[OutputPixelPos] = Debug;
+		}
+		#endif
+	}
+}
diff --git a/Engine/Shaders/Private/TemporalSuperResolution/TSRFilterFrequencies.usf b/Engine/Shaders/Private/TemporalSuperResolution/TSRFilterFrequencies.usf
--- a/Engine/Shaders/Private/TemporalSuperResolution/TSRFilterFrequencies.usf
+++ b/Engine/Shaders/Private/TemporalSuperResolution/TSRFilterFrequencies.usf
@@ -39,6 +39,7 @@ Texture2D<taa_half> InterferenceWeightTexture;
 
 RWTexture2D<taa_half3> FilteredInputOutput;
 RWTexture2D<taa_half3> FilteredPredictionSceneColorOutput;
+RWTexture2D<taa_half>  InputSceneColorLdrLumaOutput;
 
 
 //------------------------------------------------------- FUNCTIONS
@@ -201,6 +202,25 @@ void MainCS(
 		FilteredHistory *= SafeRcp(FilteredHistoryWeight);
 	}
 	
+	// Compute the LDR luma used for spatial anti-aliasing's computation of the TexelOffset
+	#if DIM_OUTPUT_ANTI_ALIASING_LUMA
+	taa_half CenterLumaLDR;
+	{
+		const taa_half SpatialAAExposure = taa_half(0.5);
+
+		#if 0
+			taa_half3 CenterInputLDR = CenterInput / (SpatialAAExposure + CenterInput);
+			CenterLumaLDR = dot(CenterInputLDR, half3(0.299f, 0.587f, 0.114f));
+		#elif 0
+			taa_half CenterLuma = max(CenterInput.r, max(CenterInput.g, CenterInput.b));
+			CenterLumaLDR = CenterLuma / (SpatialAAExposure + CenterLuma);
+		#else
+			taa_half CenterLuma = dot(CenterInput, half3(0.299f, 0.587f, 0.114f));
+			CenterLumaLDR = CenterLuma / (SpatialAAExposure + CenterLuma);
+		#endif
+	}
+	#endif // DIM_OUTPUT_ANTI_ALIASING_LUMA
+	
 	taa_short2 OutputPixelPos;
 	#if 1
 		ISOLATE
@@ -228,7 +248,13 @@ void MainCS(
 
 	FilteredInputOutput[OutputPixelPos] = FilteredInput;
 	FilteredPredictionSceneColorOutput[OutputPixelPos] = FilteredHistory;
-		
+
+	#if DIM_OUTPUT_ANTI_ALIASING_LUMA
+	{
+		InputSceneColorLdrLumaOutput[OutputPixelPos] = CenterLumaLDR;
+	}
+	#endif
+
 	#if DEBUG_OUTPUT
 	{
 		DebugOutput[OutputPixelPos - LowFrequencyInfo_ViewportMin] = Debug;
diff --git a/Engine/Shaders/Private/TemporalSuperResolution/TSRSpatialAntiAliasing.usf b/Engine/Shaders/Private/TemporalSuperResolution/TSRSpatialAntiAliasing.usf
new file mode 100644
--- /dev/null
+++ b/Engine/Shaders/Private/TemporalSuperResolution/TSRSpatialAntiAliasing.usf
@@ -0,0 +1,503 @@
+// Copyright Epic Games, Inc. All Rights Reserved.
+
+#include "TSRCommon.ush"
+
+
+//------------------------------------------------------- CONFIG
+
+#define TILE_SIZE 8
+
+#define CONFIG_MIN_LUMA_ALIASING 0.05
+
+#if DIM_QUALITY_PRESET == 1
+	#define CONFIG_ITERATIONS 3
+
+#elif DIM_QUALITY_PRESET == 2
+	#define CONFIG_ITERATIONS 8
+
+#else
+	#error unknown DIM_QUALITY_PRESET
+#endif
+
+#if CONFIG_COMPILE_FP16
+	// Take advantage of RDNA's v_pk_*_{uif}16 instructions
+	#define CONFIG_ENABLE_DUAL_PIXEL_VECTORIZATION 1
+#else
+	#define CONFIG_ENABLE_DUAL_PIXEL_VECTORIZATION 0
+#endif
+
+
+//------------------------------------------------------- INCLUDE FXAA
+
+
+//------------------------------------------------------- PARAMETERS
+
+Texture2D<taa_half> InputSceneColorLdrLumaTexture;
+
+RWTexture2D<uint> AntiAliasingOutput;
+
+
+//------------------------------------------------------- FUNCTIONS
+
+taa_ushort2x2 Map8x8Tile2x2LaneDPV(uint GroupThreadIndex)
+#if CONFIG_ENABLE_DUAL_PIXEL_VECTORIZATION
+{
+	taa_ushort T = taa_ushort(GroupThreadIndex);
+
+	taa_ushort2 GroupId;
+	GroupId.x =                                             ((T >> taa_ushort(1 - 1)) & taa_ushort(0x03 << 1));
+	GroupId.y = ((T >> taa_ushort(0)) & taa_ushort(0x01)) | ((T >> taa_ushort(3 - 1)) & taa_ushort(0x03 << 1));
+	
+	return dpv_interleave_registers(GroupId, GroupId + taa_ushort2(1, 0));
+}
+#else
+{
+	taa_ushort2 GroupId = Map8x8Tile2x2Lane(GroupThreadIndex);
+	return dpv_interleave_mono_registers(GroupId);
+}
+#endif
+
+taa_half2 InputToLDRLuma(taa_half3x2 Color)
+{
+	taa_half2 Luma = Luma4(Color);
+	return Luma * HdrWeightY(Luma);
+}
+
+taa_half SampleInputLDRLuma(taa_short2 KernelCenter, taa_short2 Offset)
+{
+	taa_short2 SampleInputPixelPos = KernelCenter + ClampPixelOffset(
+		KernelCenter,
+		Offset, Offset,
+		InputPixelPosMin, InputPixelPosMax);
+	
+	return InputSceneColorLdrLumaTexture[SampleInputPixelPos];
+}
+
+taa_half2 SampleInputLDRLuma(taa_short2x2 KernelCenter, taa_short2 Offset)
+{
+	taa_short2x2 SampleInputPixelPos = KernelCenter + ClampPixelOffset(
+		KernelCenter,
+		dpv_interleave_mono_registers(Offset), Offset,
+		InputPixelPosMin, InputPixelPosMax);
+	
+	taa_half2 InputLuma = dpv_interleave_registers(
+		InputSceneColorLdrLumaTexture[dpv_lo(SampleInputPixelPos)],
+		InputSceneColorLdrLumaTexture[dpv_hi(SampleInputPixelPos)]);
+
+	return InputLuma;
+}
+
+taa_half2 SampleInputLDRLuma(taa_short2x2 KernelCenter, taa_short2x2 Offset)
+{
+	taa_short2x2 SampleInputPixelPos = ClampPixelOffset(
+		KernelCenter + Offset,
+		InputPixelPosMin, InputPixelPosMax);
+	
+	taa_half2 InputLuma = dpv_interleave_registers(
+		InputSceneColorLdrLumaTexture[dpv_lo(SampleInputPixelPos)],
+		InputSceneColorLdrLumaTexture[dpv_hi(SampleInputPixelPos)]);
+
+	return InputLuma;
+}
+
+taa_half2 BilinearSampleInputLDRLuma(float2x2 KernelCenterUV, float2x2 UVOffset, const bool bIsNegativeDirection)
+{
+	float2x2 SampleUV = KernelCenterUV + UVOffset;
+
+	if (bIsNegativeDirection)
+		SampleUV = max(SampleUV, dpv_interleave_mono_registers(InputInfo_UVViewportBilinearMin));
+	else
+		SampleUV = min(SampleUV, dpv_interleave_mono_registers(InputInfo_UVViewportBilinearMax));
+
+	taa_half2 InputLuma = dpv_interleave_registers(
+		InputSceneColorLdrLumaTexture.SampleLevel(GlobalBilinearClampedSampler, dpv_lo(SampleUV), 0),
+		InputSceneColorLdrLumaTexture.SampleLevel(GlobalBilinearClampedSampler, dpv_hi(SampleUV), 0));
+
+	return InputLuma;
+}
+
+void AccumulateEdgeLength(
+	taa_half2 StartLuma, taa_half2 StartEdgeLuma,
+	taa_half2 SampleLuma, taa_half2 SampleEdgeLuma,
+	taa_half2 LumaDelta,
+	inout bool2 bEdgeStopedByIncrement,
+	inout bool2 bEdgeStopedByDecrement,
+	inout taa_ushort2 EdgeLength)
+{
+	bool2 bStopNowIncrement = abs(StartLuma - SampleLuma) > LumaDelta;
+	bool2 bStopNowDecrement = abs(StartEdgeLuma - SampleEdgeLuma) > LumaDelta;
+
+	bEdgeStopedByIncrement = bEdgeStopedByIncrement || bStopNowIncrement;
+	bEdgeStopedByDecrement = bEdgeStopedByDecrement || bStopNowDecrement;
+
+	EdgeLength += (bEdgeStopedByIncrement || bEdgeStopedByDecrement) ? 0 : 1;
+}
+
+taa_half ComputeDistanceToEdge(
+	bool bEdgeStopedByIncrementN,
+	bool bEdgeStopedByIncrementP,
+	bool bEdgeStopedByDecrementN,
+	bool bEdgeStopedByDecrementP,
+	taa_ushort EdgeLengthN,
+	taa_ushort EdgeLengthP)
+{	
+	taa_half fEdgeLengthN = taa_half(EdgeLengthN);
+	taa_half fEdgeLengthP = taa_half(EdgeLengthP);
+	taa_half fEdgeLength = taa_half(1.0) + fEdgeLengthN + fEdgeLengthP;
+	taa_half fEdgeInvLength = rcp(fEdgeLength);
+
+	
+	taa_half TexelCenterOffsetToEdge = 0.0;
+
+	const float Min = -0.5;
+
+	FLATTEN
+	if (!bEdgeStopedByIncrementN && !bEdgeStopedByIncrementP && !bEdgeStopedByDecrementP && !bEdgeStopedByDecrementN)
+	{
+		// No aliasing in the neighborhood.
+		TexelCenterOffsetToEdge = 0.0;
+	}
+	else if (bEdgeStopedByIncrementN && bEdgeStopedByIncrementP)
+	{
+		// Looks like InputC is concave detail
+		TexelCenterOffsetToEdge = 0.5;
+	}
+	else if (bEdgeStopedByDecrementN && bEdgeStopedByDecrementP)
+	{
+		// Looks like InputC is convex detail
+		//TexelCenterOffsetToEdge = 0.5;
+	}
+	else if (bEdgeStopedByIncrementN) // && bEdgeStopedByDecrementP)
+	{
+		// Looks like staircasing from - to +
+		TexelCenterOffsetToEdge = 0.5 - (EdgeLengthN + 0.5) * fEdgeInvLength;
+	}
+	else if (bEdgeStopedByIncrementP) // && bEdgeStopedByDecrementN)
+	{
+		// Looks like staircasing from + to -
+		TexelCenterOffsetToEdge = 0.5 - (EdgeLengthP + 0.5) * fEdgeInvLength;
+	}
+
+	return TexelCenterOffsetToEdge;
+}
+
+
+//------------------------------------------------------- ENTRY POINT
+
+#if CONFIG_ENABLE_DUAL_PIXEL_VECTORIZATION
+	[numthreads(TILE_SIZE, TILE_SIZE / 2, 1)]
+#else
+	[numthreads(TILE_SIZE, TILE_SIZE, 1)]
+#endif
+void MainCS(
+	uint2 GroupId : SV_GroupID,
+	uint GroupThreadIndex : SV_GroupIndex)
+{
+	uint GroupWaveIndex = GetGroupWaveIndex(GroupThreadIndex, /* GroupSize = */ TILE_SIZE * TILE_SIZE);
+
+	float4x2 Debug = 0.0;
+	
+	taa_short2x2 InputPixelPos = dpv_add(
+		taa_short2(InputPixelPosMin) + taa_short2(GroupId) * taa_short2(TILE_SIZE, TILE_SIZE),
+		taa_short2x2(Map8x8Tile2x2LaneDPV(GroupThreadIndex)));
+	
+	//Debug[0] = 1.0;
+
+	taa_half2x2 TexelOffset = 0.0;
+	{
+		#if CONFIG_ENABLE_DUAL_PIXEL_VECTORIZATION && 1
+			taa_half InputCache[3 * 4];
+			UNROLL_N(3)
+			for (int cy = 0; cy < 3; cy++)
+			{
+				UNROLL_N(4)
+				for (int cx = 0; cx < 4; cx++)
+				{
+					const uint ci = cx + 4 * cy;
+					InputCache[ci] = SampleInputLDRLuma(dpv_lo(InputPixelPos), taa_short2(cx - 1, cy - 1));
+				}
+			}
+
+			#define ReadCache(x, y) taa_half2(InputCache[(x + 1) + 4 * (y + 1)], InputCache[(x + 2) + 4 * (y + 1)])
+
+			taa_half2 InputC = ReadCache( 0,  0);
+
+			taa_half2 InputN = ReadCache( 0, -1);
+			taa_half2 InputS = ReadCache( 0, +1);
+			taa_half2 InputE = ReadCache(+1,  0);
+			taa_half2 InputW = ReadCache(-1,  0);
+	
+			taa_half2 InputNE = ReadCache(+1, -1);
+			taa_half2 InputNW = ReadCache(-1, -1);
+			taa_half2 InputSE = ReadCache(+1, +1);
+			taa_half2 InputSW = ReadCache(-1, +1);
+		#else
+			taa_half2 InputC = SampleInputLDRLuma(InputPixelPos, taa_short2( 0,  0));
+
+			taa_half2 InputN = SampleInputLDRLuma(InputPixelPos, taa_short2( 0, -1));
+			taa_half2 InputS = SampleInputLDRLuma(InputPixelPos, taa_short2( 0, +1));
+			taa_half2 InputE = SampleInputLDRLuma(InputPixelPos, taa_short2(+1,  0));
+			taa_half2 InputW = SampleInputLDRLuma(InputPixelPos, taa_short2(-1,  0));
+	
+			taa_half2 InputNE = SampleInputLDRLuma(InputPixelPos, taa_short2(+1, -1));
+			taa_half2 InputNW = SampleInputLDRLuma(InputPixelPos, taa_short2(-1, -1));
+			taa_half2 InputSE = SampleInputLDRLuma(InputPixelPos, taa_short2(+1, +1));
+			taa_half2 InputSW = SampleInputLDRLuma(InputPixelPos, taa_short2(-1, +1));
+		#endif
+
+		bool2 bBrowseVerticaly;
+		taa_short2x2 BrowseDirection;
+		taa_short2x2 EdgeSide = 0;
+		taa_half2 EdgeLuma = 0;
+
+		taa_half2 TotalVarHN = abs(0.5 * (InputNE + InputNW) - InputN);
+		taa_half2 TotalVarH  = abs(0.5 * (InputE  + InputW ) - InputC);
+		taa_half2 TotalVarHS = abs(0.5 * (InputSE + InputSW) - InputS);
+			
+		taa_half2 TotalVarVE = abs(0.5 * (InputNE + InputSE) - InputE);
+		taa_half2 TotalVarV  = abs(0.5 * (InputN  + InputS ) - InputC);
+		taa_half2 TotalVarVW = abs(0.5 * (InputNW + InputSW) - InputW);
+		
+		bBrowseVerticaly = (TotalVarHN + TotalVarH + TotalVarHS) > (TotalVarVE + TotalVarV + TotalVarVW);
+		BrowseDirection[0] = bBrowseVerticaly ? 0 : 1;
+		BrowseDirection[1] = bBrowseVerticaly ? 1 : 0;
+
+		FLATTEN
+		if (bBrowseVerticaly[0])
+		{
+			EdgeSide[0][0] = abs(InputW[0] - InputC[0]) > abs(InputE[0] - InputC[0]) ? -1 : 1;
+			EdgeLuma[0] = abs(InputW[0] - InputC[0]) > abs(InputE[0] - InputC[0]) ? InputW[0] : InputE[0];
+		}
+		else
+		{
+			EdgeSide[1][0] = abs(InputN[0] - InputC[0]) > abs(InputS[0] - InputC[0]) ? -1 : 1;
+			EdgeLuma[0] = abs(InputN[0] - InputC[0]) > abs(InputS[0] - InputC[0]) ? InputN[0] : InputS[0];
+		}
+		
+		FLATTEN
+		if (bBrowseVerticaly[1])
+		{
+			EdgeSide[0][1] = abs(InputW[1] - InputC[1]) > abs(InputE[1] - InputC[1]) ? -1 : 1;
+			EdgeLuma[1] = abs(InputW[1] - InputC[1]) > abs(InputE[1] - InputC[1]) ? InputW[1] : InputE[1];
+		}
+		else
+		{
+			EdgeSide[1][1] = abs(InputN[1] - InputC[1]) > abs(InputS[1] - InputC[1]) ? -1 : 1;
+			EdgeLuma[1] = abs(InputN[1] - InputC[1]) > abs(InputS[1] - InputC[1]) ? InputN[1] : InputS[1];
+		}
+
+		taa_half2 LumaDelta = abs(EdgeLuma - InputC) * 0.5;
+		
+		bool2 WorthBrowsing = LumaDelta > CONFIG_MIN_LUMA_ALIASING;
+
+		taa_half2 BrowseDelta = saturate(LumaDelta * rcp(CONFIG_MIN_LUMA_ALIASING));
+		
+		taa_ushort2 EdgeLengthP = 0;
+		taa_ushort2 EdgeLengthN = 0;
+		
+		/**
+		 * Increment:
+		 *     1 1 1 1 1
+		 *     0 0 X 0 1
+		 * 
+		 * Decrement:
+		 *     1 1 1 1 0
+		 *     0 0 X 0 0
+		 */
+		bool2 bEdgeStopedByIncrementP = false;
+		bool2 bEdgeStopedByIncrementN = false;
+		
+		bool2 bEdgeStopedByDecrementP = false;
+		bool2 bEdgeStopedByDecrementN = false;
+		
+#if 1 // Optimized for RDNA 2
+		const taa_half BilinearInterp = 0.5;
+
+		float2x2 fBrowseDirection = float2x2(BrowseDirection) * dpv_interleave_mono_registers(InputInfo_ExtentInverse);
+		float2x2 KernelUV = (float2x2(InputPixelPos) + 0.5 + float2x2(EdgeSide) * BilinearInterp) * dpv_interleave_mono_registers(InputInfo_ExtentInverse);
+
+		taa_half2 MergedEdgeLuma = BilinearInterp * (EdgeLuma + InputC);
+		
+		taa_half2 MinLuma = MergedEdgeLuma - taa_half(BilinearInterp * 0.5) * abs(EdgeLuma - InputC);
+		taa_half2 MaxLuma = MergedEdgeLuma + taa_half(BilinearInterp * 0.5) * abs(EdgeLuma - InputC);
+
+		bool2 bEdgeMoreLuminous = EdgeLuma > InputC;
+		
+		taa_ushort2 EdgeStopedByMinLumaP = 0;
+		taa_ushort2 EdgeStopedByMinLumaN = 0;
+		
+		taa_ushort2 EdgeStopedByMaxLumaP = 0;
+		taa_ushort2 EdgeStopedByMaxLumaN = 0;
+		
+		bool bStop = false;
+		
+		#if CONFIG_COMPILE_FP16
+			#define asuint(x) asuint16(x)
+			const taa_ushort SignShift = 15;
+			const taa_ushort SignMask = 0x8000;
+		#else
+			const taa_ushort SignShift = 31;
+			const taa_ushort SignMask = 0x80000000;
+		#endif
+			
+		EdgeLengthP = CONFIG_ITERATIONS;
+		EdgeLengthN = CONFIG_ITERATIONS;
+		
+		UNROLL_N(CONFIG_ITERATIONS)
+		for (uint i = 0; i < CONFIG_ITERATIONS; i++)
+		{
+			float2x2 NeighborOffset = fBrowseDirection * float(i + 1);
+
+			taa_half2 SampleLumaP = BilinearSampleInputLDRLuma(KernelUV,  NeighborOffset, /* bIsNegativeDirection = */ false);
+			taa_half2 SampleLumaN = BilinearSampleInputLDRLuma(KernelUV, -NeighborOffset, /* bIsNegativeDirection = */ true);
+				
+			//bEdgeStopedByMinLuma = bEdgeStopedByMinLuma || bool2(SampleLuma < MinLuma) && !bEdgeStopedByMaxLuma;
+			//bEdgeStopedByMaxLuma = bEdgeStopedByMaxLuma || bool2(SampleLuma > MaxLuma) && !bEdgeStopedByMinLuma;
+			taa_ushort2 StopByMinLumaP = asuint(SampleLumaP - MinLuma) & ~EdgeStopedByMaxLumaP;
+			taa_ushort2 StopByMaxLumaP = asuint(MaxLuma - SampleLumaP) & ~EdgeStopedByMinLumaP;
+			taa_ushort2 StopByMinLumaN = asuint(SampleLumaN - MinLuma) & ~EdgeStopedByMaxLumaN;
+			taa_ushort2 StopByMaxLumaN = asuint(MaxLuma - SampleLumaN) & ~EdgeStopedByMinLumaN;
+
+			EdgeStopedByMinLumaP |= StopByMinLumaP;
+			EdgeStopedByMaxLumaP |= StopByMaxLumaP;
+			EdgeStopedByMinLumaN |= StopByMinLumaN;
+			EdgeStopedByMaxLumaN |= StopByMaxLumaN;
+				
+			//EdgeLength += (bEdgeStopedByMaxLuma || bEdgeStopedByMinLuma) ? 0 : 1;
+			EdgeLengthP -= (EdgeStopedByMaxLumaP | EdgeStopedByMinLumaP) >> SignShift;
+			EdgeLengthN -= (EdgeStopedByMaxLumaN | EdgeStopedByMinLumaN) >> SignShift;
+		}
+		
+		bool2 bEdgeStopedByMinLumaP = (EdgeStopedByMinLumaP & SignMask) != 0;
+		bool2 bEdgeStopedByMinLumaN = (EdgeStopedByMinLumaN & SignMask) != 0;
+		
+		bool2 bEdgeStopedByMaxLumaP = (EdgeStopedByMaxLumaP & SignMask) != 0;
+		bool2 bEdgeStopedByMaxLumaN = (EdgeStopedByMaxLumaN & SignMask) != 0;
+		
+		bEdgeStopedByIncrementP = (bEdgeMoreLuminous && bEdgeStopedByMaxLumaP) || ((!bEdgeMoreLuminous) && bEdgeStopedByMinLumaP);
+		bEdgeStopedByIncrementN = (bEdgeMoreLuminous && bEdgeStopedByMaxLumaN) || ((!bEdgeMoreLuminous) && bEdgeStopedByMinLumaN);
+		
+		bEdgeStopedByDecrementP = (bEdgeMoreLuminous && bEdgeStopedByMinLumaP) || ((!bEdgeMoreLuminous) && bEdgeStopedByMaxLumaP);
+		bEdgeStopedByDecrementN = (bEdgeMoreLuminous && bEdgeStopedByMinLumaN) || ((!bEdgeMoreLuminous) && bEdgeStopedByMaxLumaN);
+
+#else // Reference implementation
+		UNROLL_N(CONFIG_ITERATIONS)
+		for (uint i = 0; i < CONFIG_ITERATIONS; i++)
+		{
+			bool2 bIsStoped = (bEdgeStopedByIncrementP || bEdgeStopedByDecrementP) && (bEdgeStopedByIncrementN || bEdgeStopedByDecrementP);
+
+			BRANCH
+			if (!bIsStoped.x)
+			{
+				taa_short2x2 NeighborOffset = BrowseDirection * taa_short(i + 1);
+
+				taa_half2 NeighborP = SampleInputLDRLuma(InputPixelPos, NeighborOffset);
+				taa_half2 EdgeP     = SampleInputLDRLuma(InputPixelPos, NeighborOffset + EdgeSide);
+			
+				taa_half2 NeighborN = SampleInputLDRLuma(InputPixelPos, -NeighborOffset);
+				taa_half2 EdgeN     = SampleInputLDRLuma(InputPixelPos, -NeighborOffset + EdgeSide);
+
+				AccumulateEdgeLength(
+					InputC, EdgeLuma,
+					NeighborP, EdgeP,
+					LumaDelta,
+					/* inout */ bEdgeStopedByIncrementP,
+					/* inout */ bEdgeStopedByDecrementP,
+					/* inout */ EdgeLengthP);
+			
+				AccumulateEdgeLength(
+					InputC, EdgeLuma,
+					NeighborN, EdgeN,
+					LumaDelta,
+					/* inout */ bEdgeStopedByIncrementN,
+					/* inout */ bEdgeStopedByDecrementN,
+					/* inout */ EdgeLengthN);
+			}
+		}
+#endif
+
+		taa_half TexelOffsetLo = ComputeDistanceToEdge(
+			bEdgeStopedByIncrementN[0],
+			bEdgeStopedByIncrementP[0],
+			bEdgeStopedByDecrementN[0],
+			bEdgeStopedByDecrementP[0],
+			dpv_lo(EdgeLengthN),
+			dpv_lo(EdgeLengthP));
+		
+		taa_half TexelOffsetHi = ComputeDistanceToEdge(
+			bEdgeStopedByIncrementN[1],
+			bEdgeStopedByIncrementP[1],
+			bEdgeStopedByDecrementN[1],
+			bEdgeStopedByDecrementP[1],
+			dpv_hi(EdgeLengthN),
+			dpv_hi(EdgeLengthP));
+		
+		FLATTEN
+		if (!WorthBrowsing[0])
+			TexelOffsetLo = 0.0;
+		
+		FLATTEN
+		if (!WorthBrowsing[1])
+			TexelOffsetHi = 0.0;
+
+		TexelOffset = dpv_scale(taa_half2x2(EdgeSide), BrowseDelta * dpv_interleave_registers(TexelOffsetLo, TexelOffsetHi));
+		
+		#if 0
+			if (WorthBrowsing.x)
+			{
+				Debug[0] = abs(TexelOffset[0]);
+				Debug[1] = abs(TexelOffset[1]);
+				//Debug[2] = abs(DecodeSpatialAntiAliasingOffset(EncodeSpatialAntiAliasingOffset(TexelOffset))[0]);
+				//Debug[3] = abs(DecodeSpatialAntiAliasingOffset(EncodeSpatialAntiAliasingOffset(TexelOffset))[1]);
+				//Debug[3] = bBrowseVerticaly ? 1 : 0;
+
+				//Debug[0] = EdgeLengthP * taa_ushort2(WorthBrowsing && bBrowseVerticaly);
+				//Debug[1] = EdgeLengthN * taa_ushort2(WorthBrowsing && bBrowseVerticaly);
+				//Debug[0] = (1 + EdgeLengthP + EdgeLengthN);
+				//Debug[0] = EdgeLengthN;
+		
+				//Debug[0] = WorthBrowsing ? 1 : 0;
+				//Debug[0] = EdgeLuma;
+				//Debug[0] = (EdgeSide[0] + 2) * 0.25;
+				//Debug[1] = (EdgeSide[1] + 2) * 0.25;
+			}
+
+			if (1)
+			{
+				Debug[2][0] = InputSceneColorLdrLumaTexture.SampleLevel(
+					GlobalBilinearClampedSampler,
+					(float2(dpv_lo(InputPixelPos)) + 0.5 + float2(dpv_lo(TexelOffset))) * InputInfo_ExtentInverse,
+					0);
+				
+				Debug[2][1] = InputSceneColorLdrLumaTexture.SampleLevel(
+					GlobalBilinearClampedSampler,
+					(float2(dpv_hi(InputPixelPos)) + 0.5 + float2(dpv_hi(TexelOffset))) * InputInfo_ExtentInverse,
+					0);
+			}
+		#endif
+	}
+
+	uint2 EncodedTexelOffset = EncodeSpatialAntiAliasingOffset(TexelOffset);
+
+	{
+		taa_short2x2 OutputPixelPos = InputPixelPos;
+		
+		OutputPixelPos = InvalidateOutputPixelPos(OutputPixelPos, InputInfo_ViewportMax);
+		
+		AntiAliasingOutput[dpv_lo(OutputPixelPos)] = dpv_lo(EncodedTexelOffset); // dpv_lo(EncodedTexelOffset);
+		
+		#if CONFIG_ENABLE_DUAL_PIXEL_VECTORIZATION
+			AntiAliasingOutput[dpv_hi(OutputPixelPos)] = dpv_hi(EncodedTexelOffset); // dpv_hi(EncodedTexelOffset);
+		#endif
+
+		#if DEBUG_OUTPUT
+		{
+			DebugOutput[dpv_lo(OutputPixelPos)] = dpv_lo(Debug);
+
+			#if CONFIG_ENABLE_DUAL_PIXEL_VECTORIZATION
+				DebugOutput[dpv_hi(OutputPixelPos)] = dpv_hi(Debug);
+			#endif
+		}
+		#endif
+	}
+}
diff --git a/Engine/Shaders/Private/TemporalSuperResolution/TSRUpdateHistory.usf b/Engine/Shaders/Private/TemporalSuperResolution/TSRUpdateHistory.usf
--- a/Engine/Shaders/Private/TemporalSuperResolution/TSRUpdateHistory.usf
+++ b/Engine/Shaders/Private/TemporalSuperResolution/TSRUpdateHistory.usf
@@ -24,6 +24,15 @@
 
 #define CONFIG_VFX_SURVIVAL 1
 
+#define CONFIG_REJECTION_ANTI_ALIASING (DIM_REJECTION_ANTI_ALIASING)
+
+#if CONFIG_COMPILE_FP16
+	// Take advantage of RDNA's v_pk_*_{uif}16 instructions
+	#define CONFIG_ENABLE_DUAL_PIXEL_VECTORIZATION 1
+#else
+	#define CONFIG_ENABLE_DUAL_PIXEL_VECTORIZATION 0
+#endif
+
 
 //------------------------------------------------------- CONSTANTS
 
@@ -40,6 +49,7 @@ Texture2D<taa_half> TranslucencyRejectionTexture;
 Texture2D<float2> DilatedVelocityTexture;
 Texture2D<uint> ParallaxFactorTexture;
 Texture2D<taa_half> ParallaxRejectionMaskTexture;
+Texture2D<uint> AntiAliasingTexture;
 
 float4 HistoryPixelPosToScreenPos;
 float4 HistoryPixelPosToPPCo;
@@ -198,9 +208,6 @@ void MainCS(
 
 	ISOLATE
 	{
-		// Vector in pixel between pixel K -> O.
-		taa_half2x2 dKO = taa_half2x2(PPCo - PPCk);
-
 		float2x2 EncodedVelocity = dpv_interleave_registers(
 			DilatedVelocityTexture[dpv_lo(InputPixelPos)],
 			DilatedVelocityTexture[dpv_hi(InputPixelPos)]);
@@ -229,6 +236,39 @@ void MainCS(
 		uint2 SceneStencilRef = dpv_interleave_registers(
 			InputSceneStencilTexture.Load(int3(dpv_lo(InputPixelPos), 0)) STENCIL_COMPONENT_SWIZZLE, 
 			InputSceneStencilTexture.Load(int3(dpv_hi(InputPixelPos), 0)) STENCIL_COMPONENT_SWIZZLE);
+		
+		#if CONFIG_REJECTION_ANTI_ALIASING
+		{
+			// Detect if HistoryBufferUV would be outside of the viewport.
+			bool2 bOffScreen = bool2(
+				IsOffScreen(bCameraCut, dpv_lo(PrevScreenPos), dpv_lo(ParallaxRejectionMask)),
+				IsOffScreen(bCameraCut, dpv_hi(PrevScreenPos), dpv_hi(ParallaxRejectionMask)));
+			
+			taa_half2 TotalRejection = taa_half2(
+				bOffScreen[0] ? 0.0 : saturate(LowFrequencyRejection[0] * 4.0),
+				bOffScreen[1] ? 0.0 : saturate(LowFrequencyRejection[1] * 4.0));
+
+			uint2 EncodedInputTexelOffset = dpv_interleave_registers(
+				AntiAliasingTexture[dpv_lo(InputPixelPos)],
+				AntiAliasingTexture[dpv_hi(InputPixelPos)]);
+			
+			taa_half2x2 TexelOffset = DecodeSpatialAntiAliasingOffset(EncodedInputTexelOffset);
+			
+			{
+				PPCo += dpv_scale(TexelOffset, taa_half(1.0) - TotalRejection);
+				PPCk = floor(PPCo) + 0.5;
+				InputPixelPos = ClampPixelOffset(
+					taa_short2x2(PPCo),
+					InputPixelPosMin, InputPixelPosMax);
+			}
+			
+			Debug[0] = abs(TexelOffset[0]);
+			Debug[1] = abs(TexelOffset[1]);
+		}
+		#endif
+
+		// Vector in pixel between pixel K -> O.
+		taa_half2x2 dKO = taa_half2x2(PPCo - PPCk);
 
 		UNROLL_N(CONFIG_SAMPLES - 1)
 		for (uint SampleId = 1; SampleId < CONFIG_SAMPLES; SampleId++)
@@ -343,7 +383,7 @@ void MainCS(
 			#if CONFIG_VFX_SURVIVAL
 			if (SampleId == 0)
 			{
-				ToneWeight = max(ToneWeight, ComputeSampleWeigth(1.0, dPP, /* MinimalContribution = */ 0.0));
+				ToneWeight = max(ToneWeight, ToneWeight > 0.2 ? 0.0 : ComputeSampleWeigth(1.0, dPP, /* MinimalContribution = */ 0.0));
 			}
 			#endif
 
@@ -632,12 +672,15 @@ void MainCS(
 		// Determine whether the the prediction based rejection was confident enough.
 		taa_half2 RejectionConfidentEnough = taa_half(1); // saturate(RejectionValidity * MAX_SAMPLE_COUNT - 3.0);
 
+		// Clamp the validity due to motion to maintain better sharpness.
+		taa_half2 ClampedPrevHistoryValidity = min(PrevHistoryValidity, taa_half(1.0) - (taa_half(1.0) - PredictionOnlyValidity) * saturate(OutputPixelVelocity * rcp(taa_half(40.0))));
+
 		// Compute the newly rejected validity
 		taa_half2 RejectedValidity = (
-			min(PrevHistoryValidity, PredictionOnlyValidity - DesiredCurrentContribution) +
-			max(PrevHistoryValidity - PredictionOnlyValidity + DesiredCurrentContribution, taa_half(0.0)) * PrevHistoryRejectionWeight);
+			min(ClampedPrevHistoryValidity, PredictionOnlyValidity - DesiredCurrentContribution) +
+			max(ClampedPrevHistoryValidity - PredictionOnlyValidity + DesiredCurrentContribution, taa_half(0.0)) * PrevHistoryRejectionWeight);
 
-		RejectedValidity = PrevHistoryValidity * PrevHistoryRejectionWeight;
+		RejectedValidity = ClampedPrevHistoryValidity * PrevHistoryRejectionWeight;
 
 		// Compute the maximum output validity.
 		taa_half2 OutputValidity = (
diff --git a/Engine/Source/Runtime/Renderer/Private/PostProcess/PostProcessAA.cpp b/Engine/Source/Runtime/Renderer/Private/PostProcess/PostProcessAA.cpp
--- a/Engine/Source/Runtime/Renderer/Private/PostProcess/PostProcessAA.cpp
+++ b/Engine/Source/Runtime/Renderer/Private/PostProcess/PostProcessAA.cpp
@@ -107,7 +107,7 @@ FScreenPassTexture AddFXAAPass(FRDGBuilder& GraphBuilder, const FViewInfo& View,
 
 	AddDrawScreenPass(
 		GraphBuilder,
-		RDG_EVENT_NAME("FXAA %dx%d (PS)", OutputViewport.Rect.Width(), OutputViewport.Rect.Height()),
+		RDG_EVENT_NAME("FXAA(Quality=%d) %dx%d PS", Inputs.Quality, OutputViewport.Rect.Width(), OutputViewport.Rect.Height()),
 		View,
 		OutputViewport,
 		FScreenPassTextureViewport(Inputs.SceneColor),
diff --git a/Engine/Source/Runtime/Renderer/Private/PostProcess/PostProcessUpscale.cpp b/Engine/Source/Runtime/Renderer/Private/PostProcess/PostProcessUpscale.cpp
--- a/Engine/Source/Runtime/Renderer/Private/PostProcess/PostProcessUpscale.cpp
+++ b/Engine/Source/Runtime/Renderer/Private/PostProcess/PostProcessUpscale.cpp
@@ -228,7 +228,7 @@ FScreenPassTexture AddUpscalePass(FRDGBuilder& GraphBuilder, const FViewInfo& Vi
 	const TCHAR* StageName = StageNames[static_cast<uint32>(Inputs.Stage)];
 
 	GraphBuilder.AddPass(
-		RDG_EVENT_NAME("Upscale (%s) %dx%d -> %dx%d", StageName, Inputs.SceneColor.ViewRect.Width(), Inputs.SceneColor.ViewRect.Height(), Output.ViewRect.Width(), Output.ViewRect.Height()),
+		RDG_EVENT_NAME("Upscale(%s Method=%d) %dx%d -> %dx%d", StageName, int32(Inputs.Method), Inputs.SceneColor.ViewRect.Width(), Inputs.SceneColor.ViewRect.Height(), Output.ViewRect.Width(), Output.ViewRect.Height()),
 		PassParameters,
 		ERDGPassFlags::Raster,
 		[&View, bUsePaniniProjection, PixelShader, PassParameters, InputViewport, OutputViewport](FRHICommandList& RHICmdList)
diff --git a/Engine/Source/Runtime/Renderer/Private/PostProcess/TemporalSuperResolution.cpp b/Engine/Source/Runtime/Renderer/Private/PostProcess/TemporalSuperResolution.cpp
--- a/Engine/Source/Runtime/Renderer/Private/PostProcess/TemporalSuperResolution.cpp
+++ b/Engine/Source/Runtime/Renderer/Private/PostProcess/TemporalSuperResolution.cpp
@@ -50,6 +50,11 @@ TAutoConsoleVariable<int32> CVarTSRFilterShadingRejection(
 	TEXT(" 2: Spatial filter pass is run CompareHistory pass resolution to improve stability."),
 	ECVF_RenderThreadSafe);
 
+TAutoConsoleVariable<int32> CVarTSRRejectionAntiAliasingQuality(
+	TEXT("r.TSR.RejectionAntiAliasingQuality"), 3,
+	TEXT("Controls the quality of spatial anti-aliasing on history rejection (default=1)."),
+	ECVF_Scalability | ECVF_RenderThreadSafe);
+
 TAutoConsoleVariable<int32> CVarTSREnableAntiInterference(
 	TEXT("r.TSR.AntiInterference"), 0,
 	TEXT("Enable heuristic to detect geometric interference between input pixel grid alignement and structured geometry."),
@@ -270,6 +275,10 @@ class FTSRFilterFrequenciesCS : public FTSRShader
 	DECLARE_GLOBAL_SHADER(FTSRFilterFrequenciesCS);
 	SHADER_USE_PARAMETER_STRUCT(FTSRFilterFrequenciesCS, FTSRShader);
 
+	class FOutputAALumaDim : SHADER_PERMUTATION_BOOL("DIM_OUTPUT_ANTI_ALIASING_LUMA");
+
+	using FPermutationDomain = TShaderPermutationDomain<FOutputAALumaDim>;
+
 	BEGIN_SHADER_PARAMETER_STRUCT(FParameters, )
 		SHADER_PARAMETER_STRUCT_INCLUDE(FTSRCommonParameters, CommonParameters)
 		SHADER_PARAMETER(FVector3f, OutputQuantizationError)
@@ -281,6 +290,7 @@ class FTSRFilterFrequenciesCS : public FTSRShader
 
 		SHADER_PARAMETER_RDG_TEXTURE_UAV(RWTexture2D, FilteredInputOutput)
 		SHADER_PARAMETER_RDG_TEXTURE_UAV(RWTexture2D, FilteredPredictionSceneColorOutput)
+		SHADER_PARAMETER_RDG_TEXTURE_UAV(RWTexture2D, InputSceneColorLdrLumaOutput)
 		SHADER_PARAMETER_RDG_TEXTURE_UAV(RWTexture2D, DebugOutput)
 	END_SHADER_PARAMETER_STRUCT()
 }; // class FTSRFilterFrequenciesCS
@@ -333,6 +343,56 @@ class FTSRDilateRejectionCS : public FTSRShader
 	END_SHADER_PARAMETER_STRUCT()
 }; // class FTSRDilateRejectionCS
 
+class FTSRSpatialAntiAliasingCS : public FTSRShader
+{
+	DECLARE_GLOBAL_SHADER(FTSRSpatialAntiAliasingCS);
+	SHADER_USE_PARAMETER_STRUCT(FTSRSpatialAntiAliasingCS, FTSRShader);
+
+	class FQualityDim : SHADER_PERMUTATION_INT("DIM_QUALITY_PRESET", 3);
+
+	using FPermutationDomain = TShaderPermutationDomain<FQualityDim>;
+
+	BEGIN_SHADER_PARAMETER_STRUCT(FParameters, )
+		SHADER_PARAMETER_STRUCT_INCLUDE(FTSRCommonParameters, CommonParameters)
+		SHADER_PARAMETER_RDG_TEXTURE(Texture2D, InputSceneColorTexture)
+		SHADER_PARAMETER_RDG_TEXTURE(Texture2D, InputSceneColorLdrLumaTexture)
+		SHADER_PARAMETER_RDG_TEXTURE_UAV(RWTexture2D, AntiAliasingOutput)
+		SHADER_PARAMETER_RDG_TEXTURE_UAV(RWTexture2D, DebugOutput)
+	END_SHADER_PARAMETER_STRUCT()
+
+	static bool ShouldCompilePermutation(const FGlobalShaderPermutationParameters& Parameters)
+	{
+		FPermutationDomain PermutationVector(Parameters.PermutationId);
+
+		// There is no Quality=0 because the pass doesn't get setup.
+		if (PermutationVector.Get<FQualityDim>() == 0)
+		{
+			return false;
+		}
+
+		return FTSRShader::ShouldCompilePermutation(Parameters);
+	}
+
+	static void ModifyCompilationEnvironment(const FGlobalShaderPermutationParameters& Parameters, FShaderCompilerEnvironment& OutEnvironment)
+	{
+		FTSRShader::ModifyCompilationEnvironment(Parameters, OutEnvironment);
+		OutEnvironment.CompilerFlags.Add(CFLAG_Wave32);
+	}
+}; // class FTSRSpatialAntiAliasingCS
+
+class FTSRFilterAntiAliasingCS : public FTSRShader
+{
+	DECLARE_GLOBAL_SHADER(FTSRFilterAntiAliasingCS);
+	SHADER_USE_PARAMETER_STRUCT(FTSRFilterAntiAliasingCS, FTSRShader);
+
+	BEGIN_SHADER_PARAMETER_STRUCT(FParameters, )
+		SHADER_PARAMETER_STRUCT_INCLUDE(FTSRCommonParameters, CommonParameters)
+		SHADER_PARAMETER_RDG_TEXTURE(Texture2D, AntiAliasingTexture)
+		SHADER_PARAMETER_RDG_TEXTURE_UAV(RWTexture2D, AntiAliasingOutput)
+		SHADER_PARAMETER_RDG_TEXTURE_UAV(RWTexture2D, DebugOutput)
+	END_SHADER_PARAMETER_STRUCT()
+}; // class FTSRFilterAntiAliasingCS
+
 class FTSRUpdateSuperResHistoryCS : public FTSRShader
 {
 	DECLARE_GLOBAL_SHADER(FTSRUpdateSuperResHistoryCS);
@@ -358,6 +418,10 @@ class FTSRUpdateHistoryCS : public FTSRShader
 	DECLARE_GLOBAL_SHADER(FTSRUpdateHistoryCS);
 	SHADER_USE_PARAMETER_STRUCT(FTSRUpdateHistoryCS, FTSRShader);
 
+	class FRejectionAADim : SHADER_PERMUTATION_BOOL("DIM_REJECTION_ANTI_ALIASING");
+
+	using FPermutationDomain = TShaderPermutationDomain<FRejectionAADim>;
+
 	BEGIN_SHADER_PARAMETER_STRUCT(FParameters, )
 		SHADER_PARAMETER_STRUCT_INCLUDE(FTSRCommonParameters, CommonParameters)
 		SHADER_PARAMETER_RDG_TEXTURE(Texture2D, InputSceneColorTexture)
@@ -368,6 +432,7 @@ class FTSRUpdateHistoryCS : public FTSRShader
 		SHADER_PARAMETER_RDG_TEXTURE(Texture2D, DilatedVelocityTexture)
 		SHADER_PARAMETER_RDG_TEXTURE(Texture2D, ParallaxFactorTexture)
 		SHADER_PARAMETER_RDG_TEXTURE(Texture2D, ParallaxRejectionMaskTexture)
+		SHADER_PARAMETER_RDG_TEXTURE(Texture2D, AntiAliasingTexture)
 
 		SHADER_PARAMETER(FScreenTransform, HistoryPixelPosToScreenPos)
 		SHADER_PARAMETER(FScreenTransform, HistoryPixelPosToPPCo)
@@ -416,6 +481,8 @@ IMPLEMENT_GLOBAL_SHADER(FTSRFilterFrequenciesCS,     "/Engine/Private/TemporalSu
 IMPLEMENT_GLOBAL_SHADER(FTSRCompareHistoryCS,        "/Engine/Private/TemporalSuperResolution/TSRCompareHistory.usf",        "MainCS", SF_Compute);
 IMPLEMENT_GLOBAL_SHADER(FTSRPostfilterRejectionCS,   "/Engine/Private/TemporalSuperResolution/TSRPostfilterRejection.usf",   "MainCS", SF_Compute);
 IMPLEMENT_GLOBAL_SHADER(FTSRDilateRejectionCS,       "/Engine/Private/TemporalSuperResolution/TSRDilateRejection.usf",       "MainCS", SF_Compute);
+IMPLEMENT_GLOBAL_SHADER(FTSRSpatialAntiAliasingCS,   "/Engine/Private/TemporalSuperResolution/TSRSpatialAntiAliasing.usf",   "MainCS", SF_Compute);
+IMPLEMENT_GLOBAL_SHADER(FTSRFilterAntiAliasingCS,    "/Engine/Private/TemporalSuperResolution/TSRFilterAntiAliasing.usf",    "MainCS", SF_Compute);
 IMPLEMENT_GLOBAL_SHADER(FTSRUpdateSuperResHistoryCS, "/Engine/Private/TemporalSuperResolution/TSRUpdateSuperResHistory.usf", "MainCS", SF_Compute);
 IMPLEMENT_GLOBAL_SHADER(FTSRUpdateHistoryCS,         "/Engine/Private/TemporalSuperResolution/TSRUpdateHistory.usf",         "MainCS", SF_Compute);
 
@@ -451,6 +518,8 @@ void AddTemporalSuperResolutionPasses(
 
 	bool bRejectSeparateTranslucency = PassInputs.SeparateTranslucencyTextures != nullptr && CVarTSRRejectTranslucency.GetValueOnRenderThread() != 0;
 
+	int32 RejectionAntiAliasingQuality = FMath::Clamp(CVarTSRRejectionAntiAliasingQuality.GetValueOnRenderThread(), 0, 2);
+
 	enum class ERejectionPostFilter : uint8
 	{
 		Disabled,
@@ -875,12 +944,12 @@ void AddTemporalSuperResolutionPasses(
 		PassParameters->PrevTranslucencyTexture = PrevTranslucencyTexture;
 
 		PassParameters->TranslucencyRejectionOutput = GraphBuilder.CreateUAV(TranslucencyRejectionTexture);
-		PassParameters->DebugOutput = CreateDebugUAV(LowFrequencyExtent, TEXT("Debug.TSR.ComparetTranslucency"));
+		PassParameters->DebugOutput = CreateDebugUAV(LowFrequencyExtent, TEXT("Debug.TSR.CompareTranslucency"));
 
 		TShaderMapRef<FTSRCompareTranslucencyCS> ComputeShader(View.ShaderMap);
 		FComputeShaderUtils::AddPass(
 			GraphBuilder,
-			RDG_EVENT_NAME("TSR ComparetTranslucency %dx%d", InputRect.Width(), InputRect.Height()),
+			RDG_EVENT_NAME("TSR CompareTranslucency %dx%d", InputRect.Width(), InputRect.Height()),
 			ComputeShader,
 			PassParameters,
 			FComputeShaderUtils::GetGroupCount(InputRect.Size(), 8));
@@ -965,6 +1034,7 @@ void AddTemporalSuperResolutionPasses(
 
 	// Reject the history with frequency decomposition.
 	FRDGTextureRef HistoryRejectionTexture;
+	FRDGTextureRef InputSceneColorLdrLumaTexture = nullptr;
 	{
 		// Filter out the high frquencies
 		FRDGTextureRef FilteredInputTexture;
@@ -980,6 +1050,17 @@ void AddTemporalSuperResolutionPasses(
 				FilteredInputTexture = GraphBuilder.CreateTexture(Desc, TEXT("TSR.Filtered.SceneColor"));
 				FilteredPredictionSceneColorTexture = GraphBuilder.CreateTexture(Desc, TEXT("TSR.Filtered.Prediction.SceneColor"));
 			}
+			
+			if (RejectionAntiAliasingQuality > 0)
+			{
+				FRDGTextureDesc Desc = FRDGTextureDesc::Create2D(
+					LowFrequencyExtent,
+					PF_R8,
+					FClearValueBinding::None,
+					/* InFlags = */ TexCreate_ShaderResource | TexCreate_UAV);
+
+				InputSceneColorLdrLumaTexture = GraphBuilder.CreateTexture(Desc, TEXT("TSR.SceneColorLdrLuma"));
+			}
 
 			FTSRFilterFrequenciesCS::FParameters* PassParameters = GraphBuilder.AllocParameters<FTSRFilterFrequenciesCS::FParameters>();
 			PassParameters->CommonParameters = CommonParameters;
@@ -1001,12 +1082,18 @@ void AddTemporalSuperResolutionPasses(
 
 			PassParameters->FilteredInputOutput = GraphBuilder.CreateUAV(FilteredInputTexture);
 			PassParameters->FilteredPredictionSceneColorOutput = GraphBuilder.CreateUAV(FilteredPredictionSceneColorTexture);
+			PassParameters->InputSceneColorLdrLumaOutput = InputSceneColorLdrLumaTexture ? GraphBuilder.CreateUAV(InputSceneColorLdrLumaTexture) : nullptr;
 			PassParameters->DebugOutput = CreateDebugUAV(LowFrequencyExtent, TEXT("Debug.TSR.FilterFrequencies"));
 
-			TShaderMapRef<FTSRFilterFrequenciesCS> ComputeShader(View.ShaderMap);
+			FTSRFilterFrequenciesCS::FPermutationDomain PermutationVector;
+			PermutationVector.Set<FTSRFilterFrequenciesCS::FOutputAALumaDim>(RejectionAntiAliasingQuality > 0);
+
+			TShaderMapRef<FTSRFilterFrequenciesCS> ComputeShader(View.ShaderMap, PermutationVector);
 			FComputeShaderUtils::AddPass(
 				GraphBuilder,
-				RDG_EVENT_NAME("TSR FilterFrequencies %dx%d", LowFrequencyRect.Width(), LowFrequencyRect.Height()),
+				RDG_EVENT_NAME("TSR FilterFrequencies(%s) %dx%d",
+					PermutationVector.Get<FTSRFilterFrequenciesCS::FOutputAALumaDim>() ? TEXT("OutputAALuma") : TEXT(""),
+					LowFrequencyRect.Width(), LowFrequencyRect.Height()),
 				ComputeShader,
 				PassParameters,
 				FComputeShaderUtils::GetGroupCount(LowFrequencyRect.Size(), 16));
@@ -1049,6 +1136,61 @@ void AddTemporalSuperResolutionPasses(
 		}
 	}
 
+	// Spatial anti-aliasing when doing history rejection.
+	FRDGTextureRef AntiAliasingTexture = nullptr;
+	if (RejectionAntiAliasingQuality > 0)
+	{
+		FRDGTextureRef RawAntiAliasingTexture;
+		{
+			FRDGTextureDesc Desc = FRDGTextureDesc::Create2D(
+				InputExtent,
+				PF_R8_UINT,
+				FClearValueBinding::None,
+				/* InFlags = */ TexCreate_ShaderResource | TexCreate_UAV);
+
+			RawAntiAliasingTexture = GraphBuilder.CreateTexture(Desc, TEXT("TSR.AntiAliasing.Raw"));
+			AntiAliasingTexture = GraphBuilder.CreateTexture(Desc, TEXT("TSR.AntiAliasing.Filtered"));
+		}
+
+		{
+			FTSRSpatialAntiAliasingCS::FParameters* PassParameters = GraphBuilder.AllocParameters<FTSRSpatialAntiAliasingCS::FParameters>();
+			PassParameters->CommonParameters = CommonParameters;
+			PassParameters->InputSceneColorTexture = PassInputs.SceneColorTexture;
+			PassParameters->InputSceneColorLdrLumaTexture = InputSceneColorLdrLumaTexture;
+			PassParameters->AntiAliasingOutput = GraphBuilder.CreateUAV(RawAntiAliasingTexture);
+			PassParameters->DebugOutput = CreateDebugUAV(InputExtent, TEXT("Debug.TSR.SpatialAntiAliasing"));
+
+			FTSRSpatialAntiAliasingCS::FPermutationDomain PermutationVector;
+			PermutationVector.Set<FTSRSpatialAntiAliasingCS::FQualityDim>(RejectionAntiAliasingQuality);
+
+			TShaderMapRef<FTSRSpatialAntiAliasingCS> ComputeShader(View.ShaderMap, PermutationVector);
+			FComputeShaderUtils::AddPass(
+				GraphBuilder,
+				RDG_EVENT_NAME("TSR SpatialAntiAliasing(Quality=%d) %dx%d",
+					RejectionAntiAliasingQuality,
+					InputRect.Width(), InputRect.Height()),
+				ComputeShader,
+				PassParameters,
+				FComputeShaderUtils::GetGroupCount(InputRect.Size(), 8));
+		}
+
+		{
+			FTSRFilterAntiAliasingCS::FParameters* PassParameters = GraphBuilder.AllocParameters<FTSRFilterAntiAliasingCS::FParameters>();
+			PassParameters->CommonParameters = CommonParameters;
+			PassParameters->AntiAliasingTexture = RawAntiAliasingTexture;
+			PassParameters->AntiAliasingOutput = GraphBuilder.CreateUAV(AntiAliasingTexture);
+			PassParameters->DebugOutput = CreateDebugUAV(InputExtent, TEXT("Debug.TSR.FilterAntiAliasing"));
+
+			TShaderMapRef<FTSRFilterAntiAliasingCS> ComputeShader(View.ShaderMap);
+			FComputeShaderUtils::AddPass(
+				GraphBuilder,
+				RDG_EVENT_NAME("TSR FilterAntiAliasing %dx%d", InputRect.Width(), InputRect.Height()),
+				ComputeShader,
+				PassParameters,
+				FComputeShaderUtils::GetGroupCount(InputRect.Size(), 8));
+		}
+	}
+
 	// Post filter the rejection.
 	if (PostFilter != ERejectionPostFilter::Disabled)
 	{
@@ -1184,6 +1326,7 @@ void AddTemporalSuperResolutionPasses(
 		PassParameters->DilatedVelocityTexture = DilatedVelocityTexture;
 		PassParameters->ParallaxFactorTexture = ParallaxFactorTexture;
 		PassParameters->ParallaxRejectionMaskTexture = ParallaxRejectionMaskTexture;
+		PassParameters->AntiAliasingTexture = AntiAliasingTexture;
 
 		FScreenTransform HistoryPixelPosToViewportUV = (FScreenTransform::Identity + 0.5f) * CommonParameters.HistoryInfo.ViewportSizeInverse;
 		PassParameters->HistoryPixelPosToScreenPos = HistoryPixelPosToViewportUV * FScreenTransform::ViewportUVToScreenPos;
@@ -1199,7 +1342,10 @@ void AddTemporalSuperResolutionPasses(
 		PassParameters->SceneColorOutput = GraphBuilder.CreateUAV(SceneColorOutputTexture);
 		PassParameters->DebugOutput = CreateDebugUAV(HistoryExtent, TEXT("Debug.TSR.UpdateHistory"));
 
-		TShaderMapRef<FTSRUpdateHistoryCS> ComputeShader(View.ShaderMap);
+		FTSRUpdateHistoryCS::FPermutationDomain PermutationVector;
+		PermutationVector.Set<FTSRUpdateHistoryCS::FRejectionAADim>(RejectionAntiAliasingQuality > 0);
+
+		TShaderMapRef<FTSRUpdateHistoryCS> ComputeShader(View.ShaderMap, PermutationVector);
 		ClearUnusedGraphResources(ComputeShader, PassParameters);
 
 		for (int32 i = 0; i < PrevHistory.Textures.Num(); i++)
@@ -1224,8 +1370,9 @@ void AddTemporalSuperResolutionPasses(
 
 		FComputeShaderUtils::AddPass(
 			GraphBuilder,
-			RDG_EVENT_NAME("TSR UpdateHistory(%s) %dx%d", 
+			RDG_EVENT_NAME("TSR UpdateHistory(%s%s) %dx%d", 
 				History.Textures[0]->Desc.Format == PF_FloatR11G11B10 ? TEXT("R11G11B10") : TEXT(""),
+				PermutationVector.Get<FTSRUpdateHistoryCS::FRejectionAADim>() ? TEXT(" RejectionAA") : TEXT(""),
 				HistorySize.X, HistorySize.Y),
 			ComputeShader,
 			PassParameters,
