Strata now ues statistical BRDF lob operators to computerough refraction.
This is the first checkin and lots of use cases still need to be tested.

diff --git a/Engine/Shaders/Private/DistortAccumulatePS.usf b/Engine/Shaders/Private/DistortAccumulatePS.usf
--- a/Engine/Shaders/Private/DistortAccumulatePS.usf
+++ b/Engine/Shaders/Private/DistortAccumulatePS.usf
@@ -22,6 +22,14 @@
 
 #define STRATA_ENABLED PROJECT_STRATA && MATERIAL_IS_STRATA
 
+#if STRATA_ENABLED
+#define STRATA_USE_ENV_BRDF_APPROX 1
+#include "/Engine/Private/Strata/StrataEvaluation.ush"
+
+#include "/Engine/Private/Strata/StrataStatisticalOperators.ush"
+#endif
+
+
 void ClipOccludedPixel(float PixelZ, float SceneZ)
 {
 #if	HAS_INVERTED_Z_BUFFER
@@ -68,41 +76,154 @@ void Main(
 	
 	// material distortion offset
 	half3 Normal = GetMaterialNormal(MaterialParameters, PixelMaterialInputs);
+
+	// Prevent silhouettes from geometry that is in front of distortion from being seen in the distortion 
+	float2 NDC = (MaterialParameters.ScreenPosition.xy / MaterialParameters.ScreenPosition.w);
+	float2 ScreenUV = NDC * ResolvedView.ScreenPositionScaleBias.xy + ResolvedView.ScreenPositionScaleBias.wz;
+
+	// X = IOR, Y = Refraction Bias
+	float2 Refraction = GetMaterialRefraction(PixelMaterialInputs);
+
+
 #if STRATA_ENABLED
+
+
 	FStrataData StrataData = PixelMaterialInputs.FrontMaterial;
 
-	// Use the normal from the top layer to apply distortion
-	// STRATA_TODO compute distortion based on layers, normals and scattering from distribution to achieve rough refraction.
-	if (StrataData.LayerCount > 0)
+	FStrataAddressing NullStrataAddressing = (FStrataAddressing)0;	// Fake unused in StrataCreateBSDFContext when using Forward inline shading
+
+	FStrataPixelHeader StrataPixelHeader = InitialiseStrataPixelHeader();
+	int l;
+	UNROLL
+	for (l = 0; l < StrataData.LayerCount; ++l)
 	{
-		Normal = 0;
+		StrataPixelHeader.BSDFCount += StrataData.Layers[l].BSDFCount;
+	}
+	StrataPixelHeader.SharedLocalBases = MaterialParameters.SharedLocalBases;
+
+	const float3 V = MaterialParameters.CameraVector;
+	
+	const float OpaqueBSDFThroughput = 0.0f;
+	const float FullThroughput = 1.0f;
+
+	float3 Color = 0;
+	float OutThroughput = 1.0f;
+
+	// The current way rough refraction are evaluated is
+	//  1- start with the dirac view lob
+	//	2- refract/scatter/transmit that lob though all BSDF of a layer. Compute the mean of all the resulting lobes.
+	//  3- Use that new lob as input (pointing away from the surface) for the next layer until done.
+	StrataLobStatistic WiLob = StrataGetViewDiracLob(V);
+
+	UNROLL
+	for (l = 0; l < StrataData.LayerCount; ++l)
+	{
+		float3 LayerThroughput = 0.0f;
+		float  LayerCoverage = 0.0f;
+
+		const bool bTopLayer = l == 0;
+		const bool bBottomLayer = l == (StrataData.LayerCount - 1);
+
+		StrataLobStatistic AccumulatedLayerLobs = StrataGetNullLob();
+		float AccumulatedLayerLobsWeights = 0.0f;
+
+		UNROLL
+		for (int i = 0; i < StrataData.Layers[l].BSDFCount; ++i)
+		{
+			FStrataBSDF BSDF = StrataData.Layers[l].BSDFs[i];
+
+			// Sanitize BSDF before it is used for forward shading
+			StrataSanitizeBSDF(BSDF);
+
+			BRANCH
+			if(BSDF_GETTYPE(BSDF) == STRATA_BSDF_TYPE_SLAB)
+			{
+				// We want simple volumetric only if the DMFP input is pluged in, otherwise we avoid dealing with simple volumetric
+				const bool bIsSimpleVolume = BSDF_GETHASDMFP(BSDF);
+				if (bIsSimpleVolume)
+				{
+					EnableSlabBSDFSimpleVolumetric(BSDF);
+				}
+			}
+
+			const float BSDFCoverage = BSDF.Coverage;
+
+			// Create the BSDF context
+			const float3 FakeL = float3(0, 0, 1);
+			FStrataBSDFContext StrataBSDFContext = StrataCreateBSDFContext(StrataPixelHeader, BSDF, NullStrataAddressing, V, FakeL);
+
+			FStrataEvaluateResult StrataEvaluate = StrataEvaluateBSDF(StrataBSDFContext); // Works to call because FakeL is nto used when computing StrataEvaluate.Throughput
 
-		FStrataLayer Layer = StrataData.Layers[0];
-		for (int i = 0; i < Layer.BSDFCount; ++i)
+			// We only support rough refraction from Slab material type for now.
+			if (BSDF_GETTYPE(BSDF) == STRATA_BSDF_TYPE_SLAB)
+			{
+				// Approximation to surface directional albedo (STRATA_TODO: use the LUT for a more preresentative directional albedo)
+				float3 InterfaceFDG = StrataEnvBRDF(SLAB_F0(BSDF), SLAB_F90(BSDF), SLAB_ROUGHNESSX(BSDF), StrataBSDFContext.SatNoV);	
+
+				// The top layer we consider a transition air to water. For layers below, we assume they all have similar index of refraction.
+				const float InterfaceEta12 = i == 0 ? 1.0f / 1.33f : 1.0f; // STRATA_TODO Use Refraction.x in IOR mode?
+
+				// Transform input lob into local space w.r.t. normal
+				const float3 SavedWorldSpaceWiLobMu = WiLob.Mu;
+				WiLob.Mu = float3(dot(WiLob.Mu, StrataBSDFContext.X), dot(WiLob.Mu, StrataBSDFContext.Y), dot(WiLob.Mu, StrataBSDFContext.N));
+
+				// Compute the refracted lob according to roughness . 
+				// STRATA_TODO: take into account scattering and transmittance. Colored transmittance is taken into account already but not scattering influence on variance.
+				StrataLobStatistic NewLob = StrataGetRefractedLob(WiLob, InterfaceFDG * OutThroughput * BSDFCoverage, SLAB_ROUGHNESSX(BSDF), InterfaceEta12);
+
+				// Transform the new lob direction to world space and accumulate
+				NewLob.Mu = NewLob.Mu.x * StrataBSDFContext.X + NewLob.Mu.y * StrataBSDFContext.Y + NewLob.Mu.z * StrataBSDFContext.N;
+				StrataAccumulateWeightedLob(AccumulatedLayerLobs, AccumulatedLayerLobsWeights, NewLob);
+
+				// Restore current layer input lob direction to world space
+				WiLob.Mu = SavedWorldSpaceWiLobMu;
+			}
+			else
+			{
+				// Full roughness to make it visible as something to fix in case a type is not handled.
+				WiLob.Sigma = 1000.0;
+			} 
+
+			LayerCoverage	+= BSDFCoverage;
+			LayerThroughput += BSDFCoverage * StrataEvaluate.Throughput;
+		}
+
+		if (AccumulatedLayerLobsWeights > 0.0f)
 		{
-			Normal += Layer.BSDFs[i].Coverage * MaterialParameters.SharedLocalBases.Normals[BSDF_GETSHAREDLOCALBASISID(Layer.BSDFs[i])]
-			#if MATERIAL_TANGENTSPACENORMAL
-				* MaterialParameters.TwoSidedSign
-			#endif
-				;
+			// Now compute the new input Wi lob from the accumulated refracted lob of the previous layer
+			WiLob = StrataNormaliseLob(AccumulatedLayerLobs, AccumulatedLayerLobsWeights);
+			// And negate the direction to make sure the previously refracted direction is the new incoming
+			WiLob.Mu = -WiLob.Mu;
 		}
+
+		// Take into account the non covered part for this layer and update throughput accordingly
+		LayerThroughput += (1.0 - saturate(LayerCoverage)) * FullThroughput;
+
+		// And combine this layer throughtput with he material throughput
+		OutThroughput *= LayerThroughput;
 	}
 
-#if MATERIAL_TANGENTSPACENORMAL
-	// Transform back to tangent space as this is what the common function ComputeBufferUVDistortion expects
-	Normal = normalize(TransformWorldVectorToTangent(MaterialParameters.TangentToWorld, Normal));
-#else
-	Normal = normalize(Normal);
-#endif
+	// Finally, negate the direction to make sure the new incombing lob is the refracted lob we want to take into account
+	WiLob.Mu = -WiLob.Mu;
 
-#endif
+	// FixUp to reuse the same distortion code we have for now. STRATA_TODO implement a new model?
+	half3 LobWorldDir = normalize(TransformWorldVectorToView(WiLob.Mu));
+	Normal = LobWorldDir;
 
-	// Prevent silhouettes from geometry that is in front of distortion from being seen in the distortion 
-	float2 NDC = (MaterialParameters.ScreenPosition.xy / MaterialParameters.ScreenPosition.w);
-	float2 ScreenUV = NDC * ResolvedView.ScreenPositionScaleBias.xy + ResolvedView.ScreenPositionScaleBias.wz;
+	// Compute UV distortion
+	const bool TryToClip = false; // We cannot clip because the roughness information is always needed.
+	float2 BufferUVDistortion = ComputeBufferUVDistortion(
+		MaterialParameters, PixelMaterialInputs, ResolvedView,
+		Normal, DistortionParams, ScreenUV, Refraction, TryToClip, EyeIndex);
+
+	// Sample depth at distortion offset
+	float2 DistortBufferUV = ScreenUV + BufferUVDistortion;
+
+	OutRough = saturate(StrataLobVarianceToRoughness(WiLob.Sigma));	// STRATA_TODO accumulate variance and convert back to roughness in the apply hader
+
+
+#else // STRATA_ENABLED
 
-	// X = IOR, Y = Refraction Bias
-	float2 Refraction = GetMaterialRefraction(PixelMaterialInputs);
 
 	// Compute UV distortion
 	float2 BufferUVDistortion = ComputeBufferUVDistortion(
@@ -112,6 +233,10 @@ void Main(
 	// Sample depth at distortion offset
 	float2 DistortBufferUV = ScreenUV + BufferUVDistortion;
 
+
+#endif // STRATA_ENABLED
+
+
 #if SHADING_PATH_MOBILE
 	// sample depth from SceneColor.A on Mobile
 	float DistortSceneDepth = ConvertFromDeviceZ(LookupDeviceZ(DistortBufferUV)); 
@@ -128,10 +253,6 @@ void Main(
 
 	// output positives in R|G channels and negatives in B|A channels
 	OutColor = float4(PosOffset.x,PosOffset.y,NegOffset.x,NegOffset.y);
-
-#if STRATA_ENABLED
-	OutRough = saturate( length(abs(PosOffset - NegOffset)) );				// STRATA_TODO use a proper model
-#endif
 }
 
 
diff --git a/Engine/Shaders/Private/DistortApplyScreenPS.usf b/Engine/Shaders/Private/DistortApplyScreenPS.usf
--- a/Engine/Shaders/Private/DistortApplyScreenPS.usf
+++ b/Engine/Shaders/Private/DistortApplyScreenPS.usf
@@ -56,9 +56,9 @@ void Main(
 #if USE_ROUGH_REFRACTION
 
 #if USE_MSAA
-	half RoughnessScatter = RoughnessScatterMSAATexture.Load(int2(TexCoord.xy * View.BufferSizeAndInvSize.xy), SampleIndex) * 10.0f;
+	half RoughnessScatter = RoughnessScatterMSAATexture.Load(int2(TexCoord.xy * View.BufferSizeAndInvSize.xy), SampleIndex) * 5.0f;
 #else
-	half RoughnessScatter = Texture2DSample(RoughnessScatterTexture, RoughnessScatterSampler, TexCoord.xy) * 10.0f;
+	half RoughnessScatter = Texture2DSample(RoughnessScatterTexture, RoughnessScatterSampler, TexCoord.xy) * 5.0f;
 #endif
 
 	const float2 NoDistortionUVs = TexCoord.xy * View.BufferSizeAndInvSize.xy * View.ViewSizeAndInvSize.zw;
diff --git a/Engine/Shaders/Private/Strata/StrataStatisticalOperators.ush b/Engine/Shaders/Private/Strata/StrataStatisticalOperators.ush
new file mode 100644
--- /dev/null
+++ b/Engine/Shaders/Private/Strata/StrataStatisticalOperators.ush
@@ -0,0 +1,143 @@
+// Copyright Epic Games, Inc. All Rights Reserved.
+
+#include "/Engine/Private/Common.ush"
+
+// Statistical operators representing BSDF as a sum of lobes.
+// This implementation is based on [Belcour 2018, "Efficient Rendering of Layered Materials using an Atomic Decomposition with Statistical Operators"]
+
+struct StrataLobStatistic
+{
+	// Mean
+	// xy is the 2d projection of the main direction vector onto the plane defined by the surface normal. 
+	// z is along the normal.
+	float3 Mu;
+
+	// Energy
+	float3 E;
+
+	// Variance
+	float Sigma;
+};
+
+
+
+float StrataLobRoughnessToVariance(float Roughness)
+{
+	const float SafeRoughness = clamp(Roughness, 0.0f, 0.999f);
+	return pow(SafeRoughness, 1.1f) / (1.0f - SafeRoughness);
+}
+
+float StrataLobVarianceToRoughness(float Variance)
+{
+	return pow(Variance / (1.0f + Variance), 1.0f / 1.1f);
+}
+
+
+
+// For the following statistical operators
+//  - WiLob: the lob towards which light is reflected
+//  - InterfaceRoughness = the material layer roughness
+//  - InterfaceFDG = the material layer directional albedo
+//  - InterfaceEta12 = the ratio of refractive mediaEta1 / Eta2
+//  - OpticalDepth = the value impact light transmission (= sigma_t * depth)
+//  - Approximation to in-scattering from the back of the layer (eq.19)
+
+StrataLobStatistic StrataGetNullLob()
+{
+	StrataLobStatistic NullLob = (StrataLobStatistic)0;
+	return NullLob;
+}
+
+StrataLobStatistic StrataGetViewDiracLob(float3 Wi)
+{
+	StrataLobStatistic WiLob;
+	WiLob.E = 1.0f;
+	WiLob.Mu = Wi;
+	WiLob.Sigma = 0.0f;
+	return WiLob;
+}
+
+StrataLobStatistic StrataGetReflectedLob(StrataLobStatistic WiLob, float3 InterfaceFDG, float InterfaceRoughness)
+{
+	StrataLobStatistic WoLob;
+
+	WoLob.E = WiLob.E * InterfaceFDG;
+
+	WoLob.Mu = float3(-WiLob.Mu.xy, WiLob.Mu.z);
+
+	WoLob.Sigma = WiLob.Sigma + StrataLobRoughnessToVariance(InterfaceRoughness);
+
+	return WoLob;
+}
+
+StrataLobStatistic StrataGetRefractedLob(StrataLobStatistic WiLob, float3 InterfaceFDG, float InterfaceRoughness, float InterfaceEta12)
+{
+	StrataLobStatistic WoLob;
+
+	WoLob.E = WiLob.E * (1.0f - InterfaceFDG);
+
+	WoLob.Mu.xy = -WiLob.Mu.xy * InterfaceEta12;
+	const float SquaredProjectedLength = WiLob.Mu.x * WiLob.Mu.x + WiLob.Mu.y * WiLob.Mu.y;
+	WoLob.Mu.z = -sqrt(1.0f - SquaredProjectedLength);
+
+	//	const float S = 0.5f * (1.0f + InterfaceEta12 * (WiLob.Mu.z / WoLob.Mu.z));				// This respect eq.10 but it goes crazy and does not respect the roughness of a single front layer
+	//	const float S = 0.5f * (1.0f + InterfaceEta12 * max(0.0, WiLob.Mu.z / WoLob.Mu.z));		// This respect the roughness range better but looks incorrect
+	const float S = 1.0f;																	// Until this is fully understood, do nto scale anything.
+	WoLob.Sigma = (WiLob.Sigma / InterfaceEta12) + StrataLobRoughnessToVariance(S * InterfaceRoughness);
+
+	return WoLob;
+}
+
+StrataLobStatistic StrataGetTransmittedLob(StrataLobStatistic WiLob, float3 OpticalDepth)
+{
+	StrataLobStatistic WoLob;
+
+	WoLob.E = WiLob.E * exp(-OpticalDepth);
+
+	// No refraction in the paper here
+	WoLob.Mu = -WiLob.Mu;
+
+	WoLob.Sigma = WiLob.Sigma;
+
+	return WoLob;
+}
+
+StrataLobStatistic StrataGetScatteredLob(StrataLobStatistic WiLob, float3 InScattering)
+{
+	StrataLobStatistic WoLob;
+
+	WoLob.E = WiLob.E * InScattering;
+
+	// No refraction in the paper here
+	WoLob.Mu = -WiLob.Mu;
+
+	// In our case, we use an isotropic phase function and the paper do not support that (only HG phase function with g > 0.7).
+	// So we simply assume a roughness of 1 for the lob over which incomling light to scatter is integrated.
+	const float IsotropicPhaseLobVariance = StrataLobRoughnessToVariance(1.0f);
+	WoLob.Sigma = WiLob.Sigma + IsotropicPhaseLobVariance;
+
+	return WoLob;
+}
+
+
+
+// This accumulation function can be used to merge lobes while taking into account their relative importance due to E and a coverage
+void StrataAccumulateWeightedLob(inout StrataLobStatistic LobAccum, inout float WeightAccum, in StrataLobStatistic NewLob)
+{
+	const float NewLobWeight = dot(NewLob.E, (1.0f / 3.0f).xxx);
+	LobAccum.E		+= NewLob.E * NewLobWeight;
+	LobAccum.Mu		+= NewLob.Mu * NewLobWeight;
+	LobAccum.Sigma	+= NewLob.Sigma * NewLobWeight;
+	WeightAccum		+= NewLobWeight;
+}
+
+StrataLobStatistic StrataNormaliseLob(inout StrataLobStatistic LobAccum, inout float WeightAccum)
+{
+	StrataLobStatistic OutputLob;
+	OutputLob.E		= LobAccum.E / WeightAccum;
+	OutputLob.Mu	= normalize(LobAccum.Mu);
+	OutputLob.Sigma	= LobAccum.Sigma / WeightAccum;
+	return OutputLob;
+}
+
+
