Fixed naming Lob => Lobe.

diff --git a/Engine/Shaders/Private/DistortAccumulatePS.usf b/Engine/Shaders/Private/DistortAccumulatePS.usf
--- a/Engine/Shaders/Private/DistortAccumulatePS.usf
+++ b/Engine/Shaders/Private/DistortAccumulatePS.usf
@@ -110,10 +110,10 @@ void Main(
 	float OutThroughput = 1.0f;
 
 	// The current way rough refraction are evaluated is
-	//  1- start with the dirac view lob
-	//	2- refract/scatter/transmit that lob though all BSDF of a layer. Compute the mean of all the resulting lobes.
-	//  3- Use that new lob as input (pointing away from the surface) for the next layer until done.
-	StrataLobStatistic WiLob = StrataGetViewDiracLob(V);
+	//  1- start with the dirac view Lobe
+	//	2- refract/scatter/transmit that Lobe though all BSDF of a layer. Compute the mean of all the resulting Lobes.
+	//  3- Use that new Lobe as input (pointing away from the surface) for the next layer until done.
+	StrataLobeStatistic WiLobe = StrataGetViewDiracLobe(V);
 
 	UNROLL
 	for (l = 0; l < StrataData.LayerCount; ++l)
@@ -124,8 +124,8 @@ void Main(
 		const bool bTopLayer = l == 0;
 		const bool bBottomLayer = l == (StrataData.LayerCount - 1);
 
-		StrataLobStatistic AccumulatedLayerLobs = StrataGetNullLob();
-		float AccumulatedLayerLobsWeights = 0.0f;
+		StrataLobeStatistic AccumulatedLayerLobes = StrataGetNullLobe();
+		float AccumulatedLayerLobesWeights = 0.0f;
 
 		UNROLL
 		for (int i = 0; i < StrataData.Layers[l].BSDFCount; ++i)
@@ -163,37 +163,37 @@ void Main(
 				// The top layer we consider a transition air to water. For layers below, we assume they all have similar index of refraction.
 				const float InterfaceEta12 = i == 0 ? 1.0f / 1.33f : 1.0f; // STRATA_TODO Use Refraction.x in IOR mode?
 
-				// Transform input lob into local space w.r.t. normal
-				const float3 SavedWorldSpaceWiLobMu = WiLob.Mu;
-				WiLob.Mu = float3(dot(WiLob.Mu, StrataBSDFContext.X), dot(WiLob.Mu, StrataBSDFContext.Y), dot(WiLob.Mu, StrataBSDFContext.N));
+				// Transform input Lobe into local space w.r.t. normal
+				const float3 SavedWorldSpaceWiLobeMu = WiLobe.Mu;
+				WiLobe.Mu = float3(dot(WiLobe.Mu, StrataBSDFContext.X), dot(WiLobe.Mu, StrataBSDFContext.Y), dot(WiLobe.Mu, StrataBSDFContext.N));
 
-				// Compute the refracted lob according to roughness . 
+				// Compute the refracted Lobe according to roughness . 
 				// STRATA_TODO: take into account scattering and transmittance. Colored transmittance is taken into account already but not scattering influence on variance.
-				StrataLobStatistic NewLob = StrataGetRefractedLob(WiLob, InterfaceFDG * OutThroughput * BSDFCoverage, SLAB_ROUGHNESSX(BSDF), InterfaceEta12);
+				StrataLobeStatistic NewLobe = StrataGetRefractedLobe(WiLobe, InterfaceFDG * OutThroughput * BSDFCoverage, SLAB_ROUGHNESSX(BSDF), InterfaceEta12);
 
-				// Transform the new lob direction to world space and accumulate
-				NewLob.Mu = NewLob.Mu.x * StrataBSDFContext.X + NewLob.Mu.y * StrataBSDFContext.Y + NewLob.Mu.z * StrataBSDFContext.N;
-				StrataAccumulateWeightedLob(AccumulatedLayerLobs, AccumulatedLayerLobsWeights, NewLob);
+				// Transform the new Lobe direction to world space and accumulate
+				NewLobe.Mu = NewLobe.Mu.x * StrataBSDFContext.X + NewLobe.Mu.y * StrataBSDFContext.Y + NewLobe.Mu.z * StrataBSDFContext.N;
+				StrataAccumulateWeightedLobe(AccumulatedLayerLobes, AccumulatedLayerLobesWeights, NewLobe);
 
-				// Restore current layer input lob direction to world space
-				WiLob.Mu = SavedWorldSpaceWiLobMu;
+				// Restore current layer input Lobe direction to world space
+				WiLobe.Mu = SavedWorldSpaceWiLobeMu;
 			}
 			else
 			{
 				// Full roughness to make it visible as something to fix in case a type is not handled.
-				WiLob.Sigma = 1000.0;
+				WiLobe.Sigma = 1000.0;
 			} 
 
 			LayerCoverage	+= BSDFCoverage;
 			LayerThroughput += BSDFCoverage * StrataEvaluate.Throughput;
 		}
 
-		if (AccumulatedLayerLobsWeights > 0.0f)
+		if (AccumulatedLayerLobesWeights > 0.0f)
 		{
-			// Now compute the new input Wi lob from the accumulated refracted lob of the previous layer
-			WiLob = StrataNormaliseLob(AccumulatedLayerLobs, AccumulatedLayerLobsWeights);
+			// Now compute the new input Wi Lobe from the accumulated refracted Lobe of the previous layer
+			WiLobe = StrataNormaliseLobe(AccumulatedLayerLobes, AccumulatedLayerLobesWeights);
 			// And negate the direction to make sure the previously refracted direction is the new incoming
-			WiLob.Mu = -WiLob.Mu;
+			WiLobe.Mu = -WiLobe.Mu;
 		}
 
 		// Take into account the non covered part for this layer and update throughput accordingly
@@ -204,16 +204,16 @@ void Main(
 	}
 
 	// Not write out roughness
-	OutRough = saturate(StrataLobVarianceToRoughness(WiLob.Sigma));	// STRATA_TODO accumulate variance and convert back to roughness in the apply hader
+	OutRough = saturate(StrataLobeVarianceToRoughness(WiLobe.Sigma));	// STRATA_TODO accumulate variance and convert back to roughness in the apply hader
 	// Only allow clipping if roughness is close to 0.
 	const bool TryToClip = OutRough < 0.01f;
 
-	// Finally, negate the direction to make sure the new incombing lob is the refracted lob we want to take into account
-	WiLob.Mu = -WiLob.Mu;
+	// Finally, negate the direction to make sure the new incombing Lobe is the refracted Lobe we want to take into account
+	WiLobe.Mu = -WiLobe.Mu;
 
 	// FixUp to reuse the same distortion code we have for now. STRATA_TODO implement a new model?
-	half3 LobWorldDir = normalize(TransformWorldVectorToView(WiLob.Mu));
-	Normal = LobWorldDir;
+	half3 LobeWorldDir = normalize(TransformWorldVectorToView(WiLobe.Mu));
+	Normal = LobeWorldDir;
 
 	// Compute UV distortion
 	float2 BufferUVDistortion = ComputeBufferUVDistortion(
diff --git a/Engine/Shaders/Private/Strata/StrataStatisticalOperators.ush b/Engine/Shaders/Private/Strata/StrataStatisticalOperators.ush
--- a/Engine/Shaders/Private/Strata/StrataStatisticalOperators.ush
+++ b/Engine/Shaders/Private/Strata/StrataStatisticalOperators.ush
@@ -5,7 +5,7 @@
 // Statistical operators representing BSDF as a sum of lobes.
 // This implementation is based on [Belcour 2018, "Efficient Rendering of Layered Materials using an Atomic Decomposition with Statistical Operators"]
 
-struct StrataLobStatistic
+struct StrataLobeStatistic
 {
 	// Mean
 	// xy is the 2d projection of the main direction vector onto the plane defined by the surface normal. 
@@ -21,13 +21,13 @@ struct StrataLobStatistic
 
 
 
-float StrataLobRoughnessToVariance(float Roughness)
+float StrataLobeRoughnessToVariance(float Roughness)
 {
 	const float SafeRoughness = clamp(Roughness, 0.0f, 0.999f);
 	return pow(SafeRoughness, 1.1f) / (1.0f - SafeRoughness);
 }
 
-float StrataLobVarianceToRoughness(float Variance)
+float StrataLobeVarianceToRoughness(float Variance)
 {
 	return pow(Variance / (1.0f + Variance), 1.0f / 1.1f);
 }
@@ -35,109 +35,109 @@ float StrataLobVarianceToRoughness(float Variance)
 
 
 // For the following statistical operators
-//  - WiLob: the lob towards which light is reflected
+//  - WiLobe: the lobe towards which light is reflected
 //  - InterfaceRoughness = the material layer roughness
 //  - InterfaceFDG = the material layer directional albedo
 //  - InterfaceEta12 = the ratio of refractive mediaEta1 / Eta2
 //  - OpticalDepth = the value impact light transmission (= sigma_t * depth)
 //  - Approximation to in-scattering from the back of the layer (eq.19)
 
-StrataLobStatistic StrataGetNullLob()
+StrataLobeStatistic StrataGetNullLobe()
 {
-	StrataLobStatistic NullLob = (StrataLobStatistic)0;
-	return NullLob;
+	StrataLobeStatistic NullLobe = (StrataLobeStatistic)0;
+	return NullLobe;
 }
 
-StrataLobStatistic StrataGetViewDiracLob(float3 Wi)
+StrataLobeStatistic StrataGetViewDiracLobe(float3 Wi)
 {
-	StrataLobStatistic WiLob;
-	WiLob.E = 1.0f;
-	WiLob.Mu = Wi;
-	WiLob.Sigma = 0.0f;
-	return WiLob;
+	StrataLobeStatistic WiLobe;
+	WiLobe.E = 1.0f;
+	WiLobe.Mu = Wi;
+	WiLobe.Sigma = 0.0f;
+	return WiLobe;
 }
 
-StrataLobStatistic StrataGetReflectedLob(StrataLobStatistic WiLob, float3 InterfaceFDG, float InterfaceRoughness)
+StrataLobeStatistic StrataGetReflectedLobe(StrataLobeStatistic WiLobe, float3 InterfaceFDG, float InterfaceRoughness)
 {
-	StrataLobStatistic WoLob;
+	StrataLobeStatistic WoLobe;
 
-	WoLob.E = WiLob.E * InterfaceFDG;
+	WoLobe.E = WiLobe.E * InterfaceFDG;
 
-	WoLob.Mu = float3(-WiLob.Mu.xy, WiLob.Mu.z);
+	WoLobe.Mu = float3(-WiLobe.Mu.xy, WiLobe.Mu.z);
 
-	WoLob.Sigma = WiLob.Sigma + StrataLobRoughnessToVariance(InterfaceRoughness);
+	WoLobe.Sigma = WiLobe.Sigma + StrataLobeRoughnessToVariance(InterfaceRoughness);
 
-	return WoLob;
+	return WoLobe;
 }
 
-StrataLobStatistic StrataGetRefractedLob(StrataLobStatistic WiLob, float3 InterfaceFDG, float InterfaceRoughness, float InterfaceEta12)
+StrataLobeStatistic StrataGetRefractedLobe(StrataLobeStatistic WiLobe, float3 InterfaceFDG, float InterfaceRoughness, float InterfaceEta12)
 {
-	StrataLobStatistic WoLob;
+	StrataLobeStatistic WoLobe;
 
-	WoLob.E = WiLob.E * (1.0f - InterfaceFDG);
+	WoLobe.E = WiLobe.E * (1.0f - InterfaceFDG);
 
-	WoLob.Mu.xy = -WiLob.Mu.xy * InterfaceEta12;
-	const float SquaredProjectedLength = WiLob.Mu.x * WiLob.Mu.x + WiLob.Mu.y * WiLob.Mu.y;
-	WoLob.Mu.z = -sqrt(1.0f - SquaredProjectedLength);
+	WoLobe.Mu.xy = -WiLobe.Mu.xy * InterfaceEta12;
+	const float SquaredProjectedLength = WiLobe.Mu.x * WiLobe.Mu.x + WiLobe.Mu.y * WiLobe.Mu.y;
+	WoLobe.Mu.z = -sqrt(1.0f - SquaredProjectedLength);
 
-	//	const float S = 0.5f * (1.0f + InterfaceEta12 * (WiLob.Mu.z / WoLob.Mu.z));				// This respect eq.10 but it goes crazy and does not respect the roughness of a single front layer
-	//	const float S = 0.5f * (1.0f + InterfaceEta12 * max(0.0, WiLob.Mu.z / WoLob.Mu.z));		// This respect the roughness range better but looks incorrect
-	const float S = 1.0f;																	// Until this is fully understood, do nto scale anything.
-	WoLob.Sigma = (WiLob.Sigma / InterfaceEta12) + StrataLobRoughnessToVariance(S * InterfaceRoughness);
+	//	const float S = 0.5f * (1.0f + InterfaceEta12 * (WiLobe.Mu.z / WoLobe.Mu.z));				// This respect eq.10 but it goes crazy and does not respect the roughness of a single front layer
+	//	const float S = 0.5f * (1.0f + InterfaceEta12 * max(0.0, WiLobe.Mu.z / WoLobe.Mu.z));		// This respect the roughness range better but looks incorrect
+	const float S = 1.0f;																			// ==> Until this is fully understood, do not scale anything.
+	WoLobe.Sigma = (WiLobe.Sigma / InterfaceEta12) + StrataLobeRoughnessToVariance(S * InterfaceRoughness);
 
-	return WoLob;
+	return WoLobe;
 }
 
-StrataLobStatistic StrataGetTransmittedLob(StrataLobStatistic WiLob, float3 OpticalDepth)
+StrataLobeStatistic StrataGetTransmittedLobe(StrataLobeStatistic WiLobe, float3 OpticalDepth)
 {
-	StrataLobStatistic WoLob;
+	StrataLobeStatistic WoLobe;
 
-	WoLob.E = WiLob.E * exp(-OpticalDepth);
+	WoLobe.E = WiLobe.E * exp(-OpticalDepth);
 
 	// No refraction in the paper here
-	WoLob.Mu = -WiLob.Mu;
+	WoLobe.Mu = -WiLobe.Mu;
 
-	WoLob.Sigma = WiLob.Sigma;
+	WoLobe.Sigma = WiLobe.Sigma;
 
-	return WoLob;
+	return WoLobe;
 }
 
-StrataLobStatistic StrataGetScatteredLob(StrataLobStatistic WiLob, float3 InScattering)
+StrataLobeStatistic StrataGetScatteredLobe(StrataLobeStatistic WiLobe, float3 InScattering)
 {
-	StrataLobStatistic WoLob;
+	StrataLobeStatistic WoLobe;
 
-	WoLob.E = WiLob.E * InScattering;
+	WoLobe.E = WiLobe.E * InScattering;
 
 	// No refraction in the paper here
-	WoLob.Mu = -WiLob.Mu;
+	WoLobe.Mu = -WiLobe.Mu;
 
 	// In our case, we use an isotropic phase function and the paper do not support that (only HG phase function with g > 0.7).
-	// So we simply assume a roughness of 1 for the lob over which incomling light to scatter is integrated.
-	const float IsotropicPhaseLobVariance = StrataLobRoughnessToVariance(1.0f);
-	WoLob.Sigma = WiLob.Sigma + IsotropicPhaseLobVariance;
+	// So we simply assume a roughness of 1 for the Lobe over which incomling light to scatter is integrated.
+	const float IsotropicPhaseLobeVariance = StrataLobeRoughnessToVariance(1.0f);
+	WoLobe.Sigma = WiLobe.Sigma + IsotropicPhaseLobeVariance;
 
-	return WoLob;
+	return WoLobe;
 }
 
 
 
-// This accumulation function can be used to merge lobes while taking into account their relative importance due to E and a coverage
-void StrataAccumulateWeightedLob(inout StrataLobStatistic LobAccum, inout float WeightAccum, in StrataLobStatistic NewLob)
+// This accumulation function can be used to merge Lobes while taking into account their relative importance due to E and a coverage
+void StrataAccumulateWeightedLobe(inout StrataLobeStatistic LobeAccum, inout float WeightAccum, in StrataLobeStatistic NewLobe)
 {
-	const float NewLobWeight = dot(NewLob.E, (1.0f / 3.0f).xxx);
-	LobAccum.E		+= NewLob.E * NewLobWeight;
-	LobAccum.Mu		+= NewLob.Mu * NewLobWeight;
-	LobAccum.Sigma	+= NewLob.Sigma * NewLobWeight;
-	WeightAccum		+= NewLobWeight;
+	const float NewLobeWeight = dot(NewLobe.E, (1.0f / 3.0f).xxx);
+	LobeAccum.E		+= NewLobe.E * NewLobeWeight;
+	LobeAccum.Mu	+= NewLobe.Mu * NewLobeWeight;
+	LobeAccum.Sigma	+= NewLobe.Sigma * NewLobeWeight;
+	WeightAccum		+= NewLobeWeight;
 }
 
-StrataLobStatistic StrataNormaliseLob(inout StrataLobStatistic LobAccum, inout float WeightAccum)
+StrataLobeStatistic StrataNormaliseLobe(inout StrataLobeStatistic LobeAccum, inout float WeightAccum)
 {
-	StrataLobStatistic OutputLob;
-	OutputLob.E		= LobAccum.E / WeightAccum;
-	OutputLob.Mu	= normalize(LobAccum.Mu);
-	OutputLob.Sigma	= LobAccum.Sigma / WeightAccum;
-	return OutputLob;
+	StrataLobeStatistic OutputLobe;
+	OutputLobe.E	= LobeAccum.E / WeightAccum;
+	OutputLobe.Mu	= normalize(LobeAccum.Mu);
+	OutputLobe.Sigma= LobeAccum.Sigma / WeightAccum;
+	return OutputLobe;
 }
 
 
