Clarification of the current simple mapping from roughness to mip level used to achieve rough refraction with strata.
The render target now sotr the variance for a more linear accumulation.

diff --git a/Engine/Shaders/Private/DistortAccumulatePS.usf b/Engine/Shaders/Private/DistortAccumulatePS.usf
--- a/Engine/Shaders/Private/DistortAccumulatePS.usf
+++ b/Engine/Shaders/Private/DistortAccumulatePS.usf
@@ -49,7 +49,7 @@ void Main(
 	out float4 OutColor		: SV_Target0
 #if STRATA_ENABLED
 	,
-	out float4 OutRough		: SV_Target1
+	out float4 OutVariance	: SV_Target1
 #endif
 	)
 {
@@ -203,10 +203,11 @@ void Main(
 		OutThroughput *= LayerThroughput;
 	}
 
-	// Not write out roughness
-	OutRough = saturate(StrataLobeVarianceToRoughness(WiLobe.Sigma));	// STRATA_TODO accumulate variance and convert back to roughness in the apply hader
+	// Not write out variance (as this is more linear to accumulate than roughness)
+	OutVariance = float4(WiLobe.Sigma, 0.0f, 0.0f, 0.0f);
+
 	// Only allow clipping if roughness is close to 0.
-	const bool TryToClip = OutRough < 0.01f;
+	const bool TryToClip = OutVariance.x < 0.01f;
 
 	// Finally, negate the direction to make sure the new incombing Lobe is the refracted Lobe we want to take into account
 	WiLobe.Mu = -WiLobe.Mu;
diff --git a/Engine/Shaders/Private/DistortApplyScreenPS.usf b/Engine/Shaders/Private/DistortApplyScreenPS.usf
--- a/Engine/Shaders/Private/DistortApplyScreenPS.usf
+++ b/Engine/Shaders/Private/DistortApplyScreenPS.usf
@@ -6,6 +6,12 @@
 #include "Common.ush"
 #include "PostProcessCommon.ush"
 
+
+#if USE_ROUGH_REFRACTION
+#include "/Engine/Private/Strata/StrataStatisticalOperators.ush"
+#endif
+
+
 static const half InvDistortionScaleBias = 1 / 4.0f;
 
 /** 
@@ -30,6 +36,8 @@ Texture2D<float4> RoughnessScatterTexture;
 SamplerState      RoughnessScatterSampler;
 #endif
 
+float RefractionRoughnessToMipLevelFactor;
+
 /** distorts screen texture using accumulated distortion offsets */
 void Main(
 	in noperspective float4 TexCoord: TEXCOORD0,
@@ -56,10 +64,12 @@ void Main(
 #if USE_ROUGH_REFRACTION
 
 #if USE_MSAA
-	half RoughnessScatter = RoughnessScatterMSAATexture.Load(int2(TexCoord.xy * View.BufferSizeAndInvSize.xy), SampleIndex) * 5.0f;
+	half RefractionLobeVariance = RoughnessScatterMSAATexture.Load(int2(TexCoord.xy * View.BufferSizeAndInvSize.xy), SampleIndex).x;
 #else
-	half RoughnessScatter = Texture2DSample(RoughnessScatterTexture, RoughnessScatterSampler, TexCoord.xy) * 5.0f;
+	half RefractionLobeVariance = Texture2DSample(RoughnessScatterTexture, RoughnessScatterSampler, TexCoord.xy).x;
 #endif
+	half RefractionRoughness = saturate(StrataLobeVarianceToRoughness(RefractionLobeVariance));
+	half RefractionMipLevel = RefractionRoughness * RefractionRoughnessToMipLevelFactor;
 
 	const float2 NoDistortionUVs = TexCoord.xy * View.BufferSizeAndInvSize.xy * View.ViewSizeAndInvSize.zw;
 	float2 NewBufferUV = NoDistortionUVs + DistBufferUVOffset;
@@ -70,17 +80,17 @@ void Main(
 	}
 
 #if USE_MSAA
-	if (RoughnessScatter == 0.0f)
+	if (RefractionMipLevel < 0.001f)
 	{
 		// Sample the full res MSAA source texture
 		OutColor = SceneColorMSAATexture.Load(int2(TexCoord * View.BufferSizeAndInvSize.xy), SampleIndex);
 	}
 	else
 	{
-		OutColor = Texture2DSampleLevel(SceneColorTexture, SceneColorTextureSampler, NewBufferUV, RoughnessScatter);
+		OutColor = Texture2DSampleLevel(SceneColorTexture, SceneColorTextureSampler, NewBufferUV, RefractionMipLevel);
 	}
 #else
-	OutColor = Texture2DSampleLevel(SceneColorTexture, SceneColorTextureSampler, NewBufferUV, RoughnessScatter);
+	OutColor = Texture2DSampleLevel(SceneColorTexture, SceneColorTextureSampler, NewBufferUV, RefractionMipLevel);
 #endif
 
 #else // USE_ROUGH_REFRACTION
diff --git a/Engine/Source/Runtime/Renderer/Private/DistortionRendering.cpp b/Engine/Source/Runtime/Renderer/Private/DistortionRendering.cpp
--- a/Engine/Source/Runtime/Renderer/Private/DistortionRendering.cpp
+++ b/Engine/Source/Runtime/Renderer/Private/DistortionRendering.cpp
@@ -50,6 +50,14 @@ static TAutoConsoleVariable<int32> CVarRefractionOffsetQuality(
 	TEXT("When enabled, the offset buffer is made float for higher quality. This is important to maintain the softness of the blurred scene buffer."),
 	ECVF_Scalability | ECVF_RenderThreadSafe);
 
+// Using a simple scale to map roughness to mip level.
+// Later we could remove that and use depth+variance to have a better lob roughness to mip match.
+static TAutoConsoleVariable<float> CVarRefractionRoughnessToMipLevelFactor(
+	TEXT("r.Refraction.RoughnessToMipLevelFactor"),
+	5.0f,
+	TEXT("Factor to translate the roughness factor into a mip level."),
+	ECVF_Scalability | ECVF_RenderThreadSafe);
+
 
 IMPLEMENT_STATIC_UNIFORM_BUFFER_STRUCT(FDistortionPassUniformParameters, "DistortionPass", SceneTextures);
 
@@ -117,6 +125,7 @@ public:
 		SHADER_PARAMETER_SAMPLER(SamplerState, RoughnessScatterSampler)
 		SHADER_PARAMETER_SAMPLER(SamplerState, DistortionTextureSampler)
 		SHADER_PARAMETER_SAMPLER(SamplerState, SceneColorTextureSampler)
+		SHADER_PARAMETER(float, RefractionRoughnessToMipLevelFactor)
 		RENDER_TARGET_BINDING_SLOTS()
 	END_SHADER_PARAMETER_STRUCT()
 
@@ -712,6 +721,7 @@ void FDeferredShadingSceneRenderer::RenderDistortion(FRDGBuilder& GraphBuilder,
 	CommonParameters.SceneColorTextureSampler = bUseRoughRefraction ? TStaticSamplerState<SF_Trilinear>::GetRHI() : TStaticSamplerState<>::GetRHI();
 	CommonParameters.DistortionTextureSampler = TStaticSamplerState<>::GetRHI();
 	CommonParameters.RoughnessScatterSampler = TStaticSamplerState<>::GetRHI();
+	CommonParameters.RefractionRoughnessToMipLevelFactor = FMath::Max(0.0f, CVarRefractionRoughnessToMipLevelFactor.GetValueOnRenderThread());
 
 	FDistortionScreenPS::FPermutationDomain PermutationVector;
 	PermutationVector.Set<FDistortionScreenPS::FUseMSAADim>(SceneColorTexture->Desc.NumSamples > 1);
