Fix compile errors when compiling global shaders: InstanceCullingLoadBalancer.ush(67): error: syntax error, unexpected ';', expecting $end

diff --git a/Engine/Shaders/Private/InstanceCulling/InstanceCullingLoadBalancer.ush b/Engine/Shaders/Private/InstanceCulling/InstanceCullingLoadBalancer.ush
--- a/Engine/Shaders/Private/InstanceCulling/InstanceCullingLoadBalancer.ush
+++ b/Engine/Shaders/Private/InstanceCulling/InstanceCullingLoadBalancer.ush
@@ -44,7 +44,7 @@ FInstanceBatch UnpackBatch(FPackedInstanceBatch PackedBatch, uint ItemDataOffset
 	Result.NumItems = PackedBatch.FirstItem_NumItems & NUM_INSTANCES_ITEM_MASK;
 
 	return Result;
-};
+}
 
 struct FInstanceBatchItem
 {
@@ -64,7 +64,7 @@ FInstanceBatchItem UnpackItem(FPackedInstanceBatchItem PackedItem)
 	Result.BatchPrefixOffset = PackedItem.Payload_BatchPrefixOffset & PREFIX_BIT_MASK;
 
 	return Result;
-};
+}
 
 
 #if USE_WAVE_BIT_PREFIX_WORK_DISTRIBUTION
