Implement path compaction for the path tracer to keep GPU lanes more fully occupied as rays bounce around the scene

The basic idea is to breakup the bounce loop in the path tracer into independent launches so we can group active paths together. This step is done with wave ops + atomics. The implementation is rather naive at the moment, but it already gives a 2x speedup on most scenes. This is on par or better than the speedup from the coherent sampler, but without inposing any constraints on the random numbers.

diff --git a/Engine/Shaders/Private/PathTracing/PathTracing.usf b/Engine/Shaders/Private/PathTracing/PathTracing.usf
--- a/Engine/Shaders/Private/PathTracing/PathTracing.usf
+++ b/Engine/Shaders/Private/PathTracing/PathTracing.usf
@@ -1,801 +1,174 @@
 // Copyright Epic Games, Inc. All Rights Reserved.
 
 /*=============================================================================================
-PathTracingRayGenShader.usf: Reference path tracing
+PathTracing.usf: Reference path tracing
 ===============================================================================================*/
 
 #define PATH_TRACING  1
 
 #define PATH_TRACING_SHADER_VERSION 0x51478b68 // Bump to force-compile this shader
 
-#define RANDSEQ_UNROLL_SOBOL 0 // unrolling sobol sampler doesn't seem to pay off in the path tracer
-
-#define USE_PATH_TRACING_LIGHT_GRID	1
-
-#include "../Common.ush"
-#include "../PostProcessCommon.ush"
-#include "../RectLight.ush"
-#include "../RayTracing/RayTracingCommon.ush"
-
-#include "PathTracingCommon.ush"
-#include "../RayTracing/RayTracingHitGroupCommon.ush"
-
-#include "../ShadingModels.ush"
-#include "./Utilities/PathTracingRandomSequence.ush"
-#include "./Light/PathTracingLightSampling.ush"
-#include "./Light/PathTracingLightGrid.ush"
-#include "./Material/PathTracingMaterialSampling.ush"
-#include "PathTracingCamera.ush"
-
-uint Iteration;
-uint TemporalSeed;
-uint MaxSamples;
-uint MaxBounces;
-uint MaxSSSBounces;
-
-// 0: only Material sampling
-// 1: only Light sampling
-// 2: both Material and Light
-uint MISMode;
-
-uint ApproximateCaustics;
-uint EnableCameraBackfaceCulling;
-uint EnableDirectLighting;
-uint EnableEmissive;
-uint SamplerType;
-uint VisualizeLightGrid;
-float MaxPathIntensity;
-float MaxNormalBias;
-float FilterWidth;
-
-RWTexture2D<float4> RadianceTexture;
-RWTexture2D<float4> AlbedoTexture;
-RWTexture2D<float4> NormalTexture;
-RaytracingAccelerationStructure TLAS;
-int3 TileOffset;
-uint SceneVisibleLightCount;
-
-
-void AccumulateRadiance(inout float3 TotalRadiance, float3 PathRadiance, bool bIsCameraRay)
-{
-	// User asked for path contributions to be clamped to reduce fireflies.
-	// This puts an upper bound on variance within the pixel at the expense of bias
-	float MaxPathRadiance = max3(PathRadiance.x, PathRadiance.y, PathRadiance.z);
-	if (!bIsCameraRay && MaxPathRadiance > MaxPathIntensity)
-	{
-		// adjust sample while keeping color
-		PathRadiance *= MaxPathIntensity / MaxPathRadiance;
-	}
+#include "PathTracingCore.ush"
 
-	// accumulate into the total
-	TotalRadiance += PathRadiance;
-}
+#if PATH_TRACER_USE_COMPACTION
 
-FPathTracingPayload TraceTransparentRay(RayDesc Ray, FRayCone RayCone, bool bIsCameraRay, bool bLastBounce, bool bIncludeEmission, uint2 LaunchIndex, uint NumLights, inout RandomSequence RandSequence, inout float3 PathThroughput, inout float3 Radiance, inout float3 Albedo, inout float3 Normal)
-{
-	const uint RayFlags = bIsCameraRay && EnableCameraBackfaceCulling ? RAY_FLAG_CULL_BACK_FACING_TRIANGLES : 0;
-	const uint MissShaderIndex = 0;
-	float SelectionWeightSum = 0;
-	float3 PayloadThroughput = PathThroughput;
-	FPathTracingPayload Payload;
-	if (!bIncludeEmission && bLastBounce)
-	{
-		Payload.SetMiss();
-		PathThroughput = 0;
-		return Payload;
-	}
-	for (;;)
-	{
-		// avoid double counting what was captured by the skylight
-		// also avoid noise from hot spots (they can be properly
-		// importance sampled if a capturing skylight is added)
-		const bool bIgnoreSky = !bIsCameraRay;
-		FPackedPathTracingPayload PackedPayload = InitPathTracingPayload(RayCone, LaunchIndex, bIgnoreSky);
-		// Trace the ray
-		TraceRay(
-			TLAS,
-			RayFlags,
-			RAY_TRACING_MASK_ALL,
-			RAY_TRACING_SHADER_SLOT_MATERIAL,
-			RAY_TRACING_NUM_SHADER_SLOTS,
-			MissShaderIndex,
-			Ray,
-			PackedPayload);
-
-		// Loop over lights to capture their contribution
-		// #dxr_todo: if we have lots of lights, having some hierarchical structure would be better ....
-		for (uint LightId = 0; LightId < NumLights; ++LightId)
-		{
-			RayDesc LightRay = Ray;
-			LightRay.TMax = PackedPayload.IsMiss() ? Ray.TMax : PackedPayload.HitT;
-			float3 LightRadiance = TraceLight(LightRay, LightId).Radiance;
-			AccumulateRadiance(Radiance, PathThroughput * LightRadiance, bIsCameraRay);
-		}
+int Bounce;
+RWBuffer<uint> PathStateData;
+Buffer<int> ActivePaths;
+RWBuffer<int> NextActivePaths;
+RWBuffer<int> NumPathStates;
 
-		if (PackedPayload.IsMiss())
-		{
-			// Ray didn't hit any real geometry, so nothing left to do
-			break;
-		}
 
-		// Unpack the payload
-		FPathTracingPayload HitPayload = UnpackPathTracingPayloadMaterial(PackedPayload, Ray);
+// TODO: Experiment with AoS vs SoA vs AoSoA?
+#define PATH_STATE_STRIDE   PATHTRACER_PATHSTATE_NUM_WORDS
 
-		if (HitPayload.BlendingMode == RAY_TRACING_BLEND_MODE_ADDITIVE ||
-			HitPayload.BlendingMode == RAY_TRACING_BLEND_MODE_TRANSLUCENT)
-		{
-			// emission needs to be pre-multiplied by the opacity
-			HitPayload.Radiance *= HitPayload.Opacity;
-		}
+FPathState LoadPathStateData(uint Index)
+{
+	FPathState Output = (FPathState)0;
+	uint PackedPixelIndex = PathStateData[PATH_STATE_STRIDE * Index + 0];
+	Output.PixelIndex = uint2(PackedPixelIndex & 0xFFFFu, PackedPixelIndex >> 16);
+
+	Output.RandSequence.SampleIndex = PathStateData[PATH_STATE_STRIDE * Index + 1];
+	Output.RandSequence.SampleSeed  = PathStateData[PATH_STATE_STRIDE * Index + 2];
+
+	Output.Radiance.x = asfloat(PathStateData[PATH_STATE_STRIDE * Index + 3]);
+	Output.Radiance.y = asfloat(PathStateData[PATH_STATE_STRIDE * Index + 4]);
+	Output.Radiance.z = asfloat(PathStateData[PATH_STATE_STRIDE * Index + 5]);
+
+	uint3 PackedAlbedoNormal;
+	PackedAlbedoNormal.x = PathStateData[PATH_STATE_STRIDE * Index + 6];
+	PackedAlbedoNormal.y = PathStateData[PATH_STATE_STRIDE * Index + 7];
+	PackedAlbedoNormal.z = PathStateData[PATH_STATE_STRIDE * Index + 8];
+	Output.Albedo.x = f16tof32(PackedAlbedoNormal.x);
+	Output.Albedo.y = f16tof32(PackedAlbedoNormal.x >> 16);
+	Output.Albedo.z = f16tof32(PackedAlbedoNormal.y);
+	Output.Normal.x = f16tof32(PackedAlbedoNormal.y >> 16);
+	Output.Normal.y = f16tof32(PackedAlbedoNormal.z);
+	Output.Normal.z = f16tof32(PackedAlbedoNormal.z >> 16);
+
+	Output.Ray.Origin.x = asfloat(PathStateData[PATH_STATE_STRIDE * Index + 9]);
+	Output.Ray.Origin.y = asfloat(PathStateData[PATH_STATE_STRIDE * Index + 10]);
+	Output.Ray.Origin.z = asfloat(PathStateData[PATH_STATE_STRIDE * Index + 11]);
+	Output.Ray.Direction.x = asfloat(PathStateData[PATH_STATE_STRIDE * Index + 12]);
+	Output.Ray.Direction.y = asfloat(PathStateData[PATH_STATE_STRIDE * Index + 13]);
+	Output.Ray.Direction.z = asfloat(PathStateData[PATH_STATE_STRIDE * Index + 14]);
+	Output.Ray.TMin = 0;
+	Output.Ray.TMax = POSITIVE_INFINITY;
+
+	uint PackedRayCone = PathStateData[PATH_STATE_STRIDE * Index + 15];
+	Output.RayCone.Width = f16tof32(PackedRayCone);
+	Output.RayCone.SpreadAngle = f16tof32(PackedRayCone >> 16);
+
+	Output.PathThroughput.x = asfloat(PathStateData[PATH_STATE_STRIDE * Index + 16]);
+	Output.PathThroughput.y = asfloat(PathStateData[PATH_STATE_STRIDE * Index + 17]);
+	Output.PathThroughput.z = asfloat(PathStateData[PATH_STATE_STRIDE * Index + 18]);
+
+	Output.PathRoughness = asfloat(PathStateData[PATH_STATE_STRIDE * Index + 19]);
+
+	return Output;
+}
 
-		// add in surface emission (except for modulate blend mode which uses Radiance as the transparency multiplier)
-		if (bIncludeEmission && HitPayload.BlendingMode != RAY_TRACING_BLEND_MODE_MODULATE)
-		{
-			AccumulateRadiance(Radiance, PathThroughput * HitPayload.Radiance, bIsCameraRay);
-		}
+void StorePathStateData(FPathState PathState, uint Index)
+{
+	PathStateData[PATH_STATE_STRIDE * Index + 0] = PathState.PixelIndex.x + PathState.PixelIndex.y * 65536;
 
-		float3 Transparency = GetMaterialTransparency(HitPayload, Ray.Direction);
+	PathStateData[PATH_STATE_STRIDE * Index + 1] = PathState.RandSequence.SampleIndex;
+	PathStateData[PATH_STATE_STRIDE * Index + 2] = PathState.RandSequence.SampleSeed;
 
-		if (!bLastBounce)
-		{
-			float3 Contrib = PathThroughput * EstimateMaterialAlbedo(HitPayload);
-			if (bIsCameraRay)
-			{
-				Albedo += Contrib;
-				Normal += PathThroughput * (1 - Luminance(Transparency)) * HitPayload.WorldNormal;
-			}
-
-			float SelectionWeight = max3(Contrib.x, Contrib.y, Contrib.z);
-			SelectionWeightSum += SelectionWeight;
-			bool AcceptHit = false;
-
-			// weighted reservoir sampling
-			if (SelectionWeight > 0)
-			{
-				if (SelectionWeight < SelectionWeightSum)
-				{
-					// the acceptance probability is not 1.0
-					// generate a random number to see if we should accept this hit
-					float RandValue = RandomSequence_GenerateSample1D(RandSequence);
-					AcceptHit = RandValue * SelectionWeightSum < SelectionWeight;
-				}
-				else
-				{
-					// accept automatically on the first hit
-					AcceptHit = true;
-				}
-			}
-
-			if (AcceptHit)
-			{
-				// stash this hit for next time
-				Payload = HitPayload;
-				PayloadThroughput = PathThroughput / SelectionWeight;
-			}
-		}
+	PathStateData[PATH_STATE_STRIDE * Index + 3] = asuint(PathState.Radiance.x);
+	PathStateData[PATH_STATE_STRIDE * Index + 4] = asuint(PathState.Radiance.y);
+	PathStateData[PATH_STATE_STRIDE * Index + 5] = asuint(PathState.Radiance.z);
 
-		// prepare next step around the loop
-		// retrace the exact same ray with TMin one ulp past the hit we just found
-		PathThroughput *= Transparency;
-		Ray.TMin = asfloat(asuint(HitPayload.HitT) + 1);
+	PathStateData[PATH_STATE_STRIDE * Index + 6] = f32tof16(PathState.Albedo.x) | (f32tof16(PathState.Albedo.y) << 16);
+	PathStateData[PATH_STATE_STRIDE * Index + 7] = f32tof16(PathState.Albedo.z) | (f32tof16(PathState.Normal.x) << 16);
+	PathStateData[PATH_STATE_STRIDE * Index + 8] = f32tof16(PathState.Normal.y) | (f32tof16(PathState.Normal.z) << 16);
 
-		if (all(PathThroughput == 0))
-		{
-			break;
-		}
-	}
+	PathStateData[PATH_STATE_STRIDE * Index + 9]  = asuint(PathState.Ray.Origin.x);
+	PathStateData[PATH_STATE_STRIDE * Index + 10] = asuint(PathState.Ray.Origin.y);
+	PathStateData[PATH_STATE_STRIDE * Index + 11] = asuint(PathState.Ray.Origin.z);
+	PathStateData[PATH_STATE_STRIDE * Index + 12] = asuint(PathState.Ray.Direction.x);
+	PathStateData[PATH_STATE_STRIDE * Index + 13] = asuint(PathState.Ray.Direction.y);
+	PathStateData[PATH_STATE_STRIDE * Index + 14] = asuint(PathState.Ray.Direction.z);
 
-	if (SelectionWeightSum > 0)
-	{
-		// if we stored a valid hit in the payload, reset the path throughput to this point
-		PathThroughput = PayloadThroughput * SelectionWeightSum;
-	}
-	else
-	{
-		PathThroughput = 0;
-		Payload.SetMiss();
-	}
-	return Payload;
-}
+	PathStateData[PATH_STATE_STRIDE * Index + 15] = f32tof16(PathState.RayCone.Width) | (f32tof16(PathState.RayCone.SpreadAngle) << 16);
 
-float3 TraceTransparentVisibilityRay(RayDesc Ray, uint2 PixelCoord, float PathRoughness)
-{
-	const uint RayFlags = RAY_FLAG_ACCEPT_FIRST_HIT_AND_END_SEARCH | RAY_FLAG_SKIP_CLOSEST_HIT_SHADER;
-	const uint InstanceInclusionMask = RAY_TRACING_MASK_SHADOW;
-	const uint RayContributionToHitGroupIndex = RAY_TRACING_SHADER_SLOT_SHADOW;
-	const uint MultiplierForGeometryContributionToShaderIndex = RAY_TRACING_NUM_SHADER_SLOTS;
-	const uint MissShaderIndex = 0;
-
-	FPackedPathTracingPayload PackedPayload = InitPathTracingVisibilityPayload(PixelCoord, PathRoughness);
-
-	// Trace the ray
-	TraceRay(
-		TLAS,
-		RayFlags,
-		InstanceInclusionMask,
-		RayContributionToHitGroupIndex,
-		MultiplierForGeometryContributionToShaderIndex,
-		MissShaderIndex,
-		Ray,
-		PackedPayload);
-
-	if (PackedPayload.IsMiss())
-	{
-		// we didn't hit any geo - just return whatever throughput is left
-		return PackedPayload.GetRayThroughput();
-	}
+	PathStateData[PATH_STATE_STRIDE * Index + 16] = asuint(PathState.PathThroughput.x);
+	PathStateData[PATH_STATE_STRIDE * Index + 17] = asuint(PathState.PathThroughput.y);
+	PathStateData[PATH_STATE_STRIDE * Index + 18] = asuint(PathState.PathThroughput.z);
 
-	// we registered a hit -- assumed to be fully opaque
-	return 0;
+	PathStateData[PATH_STATE_STRIDE * Index + 19] = asuint(PathState.PathRoughness);
 }
 
-struct FProbeResult
+RAY_TRACING_ENTRY_RAYGEN(PathTracingMainRG)
 {
-	float HitT;
-	float3 WorldNormal;
-	int FrontFace;
-
-	bool IsMiss() { return HitT <= 0; }
-};
+	const uint2 DispatchIdx = DispatchRaysIndex().xy;
+	const uint2 DispatchDim = DispatchRaysDimensions().xy;
 
-FProbeResult TraceProbeRay(RayDesc Ray, uint2 PixelCoord)
-{
-	// Trace a short ray to see if we escaped the surface
-	// TODO: figure out a way to only hit the object we started on (or handle blocking geometry)
-	FPackedPathTracingPayload PackedPayload = InitPathTracingPayload((FRayCone)0, PixelCoord, true);
-	// TODO: payload flag to skip shading (we only care about the normal)
-	const uint RayFlags = 0;
-	const uint MissShaderIndex = 0;
-	TraceRay(
-		TLAS,
-		RayFlags,
-		RAY_TRACING_MASK_ALL,
-		RAY_TRACING_SHADER_SLOT_MATERIAL,
-		RAY_TRACING_NUM_SHADER_SLOTS,
-		MissShaderIndex,
-		Ray,
-		PackedPayload);
-
-	if (PackedPayload.IsMiss())
+	int LinearPixelIndex = DispatchIdx.x + DispatchDim.x * DispatchIdx.y;
+	FPathState PathState;
+	if (Bounce == 0)
 	{
-		return (FProbeResult)0;
+		PathState = CreatePathState(DispatchIdx, DispatchDim);
 	}
-
-	FProbeResult Result;
-	Result.HitT = PackedPayload.HitT;
-	Result.WorldNormal = PackedPayload.GetWorldNormal();
-	Result.FrontFace = PackedPayload.IsFrontFace();
-	return Result;
-}
-
-bool ProcessSubsurfaceRandomWalk(inout FPathTracingPayload Payload, inout float3 PathThroughput, inout RandomSequence RandSequence, uint2 PixelCoord, bool SimplifySSS)
-{
-	if (!Payload.IsSubsurfaceMaterial())
+	else
 	{
-		// material doesn't have SSS
-		// just return so we can process the rest of the payload's shading
-		return true;
+		LinearPixelIndex = ActivePaths[LinearPixelIndex];
+		if (LinearPixelIndex < 0)
+		{
+			return; // nothing left to do on this thread
+		}
+		PathState = LoadPathStateData(LinearPixelIndex);
 	}
 
-	float3 SSSColor = Payload.GetSubsurfaceColor();
-	float3 SSSRadius = Payload.GetSubsurfaceRadius();
+	const bool KeepGoing = PathTracingKernel(PathState, Bounce);
+	// check how many lanes in our wave want to keep tracing forward
+	const uint NumActive = WaveActiveCountBits(KeepGoing);
 
-	if (SimplifySSS || all(SSSRadius == 0) || all(SSSColor == 0) || MaxSSSBounces == 0)
+	uint GroupWriteStart;
+	if (WaveIsFirstLane() && NumActive > 0)
 	{
-		// if we are running in a context that won't see the SSS directly -- just turn it into diffuse and skip the random walk
-		// we can also skip random walk if the radius or color is 0
-		// or if the user decided not to perform any SSS bounces
-		Payload.DiffuseColor += SSSColor;
-		Payload.DiffuseColor = saturate(Payload.DiffuseColor);
-		Payload.SetSubsurfaceColor(0.0);
-		return true;
+		// increment the global counter to know which lane to map to
+		// NOTE: we reduce contention by only incrementing once per wave, but this is still sub-optimal
+		InterlockedAdd(NumPathStates[0], NumActive, GroupWriteStart);
 	}
-	SSSRadius = max(SSSRadius, 0.1); // 1mm minimum to make sure random walk can move forward and to keep SigmaT finite
+	// Broadcast the global write offset to all lanes.
+	GroupWriteStart = WaveReadLaneFirst(GroupWriteStart);
 
-	// decide randomly to evaluate the SSS part of the material, or keep shading the current part
-	float SSSProb = LobeSelectionProb(SSSColor, Payload.DiffuseColor + Payload.SpecularColor);
-	float3 RandSample = RandomSequence_GenerateSample3D(RandSequence);
-	if (RandSample.x < SSSProb)
+	if (KeepGoing)
 	{
-		// do random walk
-		PathThroughput *= 1 / SSSProb;
-		RandSample.x /= SSSProb;
+		// the path should still be extended - store the path state
+		uint PathStateIndex = GroupWriteStart + WavePrefixCountBits(KeepGoing);
+		NextActivePaths[PathStateIndex] = LinearPixelIndex;
+		StorePathStateData(PathState, LinearPixelIndex);
 	}
 	else
 	{
-		// skip random walk and shade current point instead
-		PathThroughput *= 1 / (1 - SSSProb);
-		return true;
+		// nothing left to do
+		// Accumulate radiance and update pixel variance
+		PathState.WritePixel(BlendFactor);
 	}
-
-	// create ray to enter the surface (using a diffuse scattering event)
-	RayDesc Ray;
-	Ray.Origin = Payload.WorldPos;
-	Ray.Direction = TangentToWorld(-CosineSampleHemisphere(RandSample.xy).xyz, Payload.WorldNormal);
-	Ray.TMin = 0;
-	ApplyPositionBias(Ray, -Payload.WorldNormal, MaxNormalBias);
-
-	// Use the technique detailed in the Hyperion paper (Section 4.4.2)
-	// https://media.disneyanimation.com/uploads/production/publication_asset/177/asset/a.pdf
-	// This avoids overly bright thin regions by applying a bit of internal reflection when the boundary is hit
-#define SSS_USE_TIR 1
-
-#if SSS_USE_TIR
-	const float3 Albedo = 1 - exp(SSSColor * (-11.43 + SSSColor * (15.38 - 13.91 * SSSColor)));
-#else
-	// Van de-Hulst inverse mapping
-	// https://blog.selfshadow.com/publications/s2017-shading-course/imageworks/s2017_pbs_imageworks_slides_v2.pdf (Slide 44)
-	// [REDACTED_PROFILE_URL] (Section 7.5.3 of v0.1.3)
-	const float3 Albedo = 1 - Pow2(4.09712 + 4.20863 * SSSColor - sqrt(9.59217 + SSSColor * (41.6808 + 17.7126 * SSSColor)));
-	SSSRadius *= 2.0; // roughly match parameterization above
-#endif
-	const int MAX_SSS_BOUNCES = MaxSSSBounces;
-	const float3 SigmaT = 1.0 / SSSRadius;
-	const float3 SigmaS = Albedo * SigmaT;
-	for (int i = 0; i < MAX_SSS_BOUNCES; i++)
-	{
-		// TODO: add Dwivedi sampling for faster escape in low albedo media ...
-
-		// Pick a color channel for transmittance sampling:
-		float3 ColorChannelWeight = PathThroughput * Albedo;
-		float3 ColorChannelCdf = float3(ColorChannelWeight.x,
-										ColorChannelWeight.x + ColorChannelWeight.y,
-										ColorChannelWeight.x + ColorChannelWeight.y + ColorChannelWeight.z);
-		if (ColorChannelCdf.z == 0)
-		{
-			// no more energy left in the path
-			break;
-		}
-		ColorChannelCdf.xy /= ColorChannelCdf.z;
-		float3 ColorChannelPdf = float3(ColorChannelCdf.x, ColorChannelCdf.y - ColorChannelCdf.x, 1 - ColorChannelCdf.y);
-		// pick a maximum distance to scatter
-		const float OneMinusEpsilon = 0.99999994; // 32-bit float just before 1.0
-		if (RandSample.z < ColorChannelCdf.x)
-		{
-			float RescaleRand = min(RandSample.z / ColorChannelPdf.x, OneMinusEpsilon);
-			Ray.TMax = -log(1 - RescaleRand) * SSSRadius.x;
-		}
-		else if (RandSample.z < ColorChannelCdf.y)
-		{
-			float RescaleRand = min((RandSample.z - ColorChannelCdf.x) / ColorChannelPdf.y, OneMinusEpsilon);
-			Ray.TMax = -log(1 - RescaleRand) * SSSRadius.y;
-		}
-		else
-		{
-			float RescaleRand = min((RandSample.z - ColorChannelCdf.y) / ColorChannelPdf.z, OneMinusEpsilon);
-			Ray.TMax = -log(1 - RescaleRand) * SSSRadius.z;
-		}
-
-		FProbeResult ProbeResult = TraceProbeRay(Ray, PixelCoord);
-		if (ProbeResult.IsMiss())
-		{
-			// we didn't hit anything, so scatter randomly in the volume and keep tracing
-			RandSample = RandomSequence_GenerateSample3D(RandSequence);
-			Ray.Origin += Ray.TMax * Ray.Direction;
-			Ray.Direction = UniformSampleSphere(RandSample.xy).xyz;
-			// account for transmittance and scattering up to this point
-			float3 Transmittance = exp(-Ray.TMax * SigmaT);
-			PathThroughput *= SigmaS * Transmittance;
-			// account for transmittance pdf (which is the balance heuristic combination of transmittance sampling in each color channel)
-			float3 TransmittancePdf = SigmaT * Transmittance;
-			PathThroughput /= dot(TransmittancePdf, ColorChannelPdf);
-
-			// NOTE: phase function cancels out since it is being perfectly importance sampled
-		}
-		else
-		{
-			// account for transmittance and scattering up to this point
-			float3 Transmittance = exp(-ProbeResult.HitT * SigmaT);
-			PathThroughput *= Transmittance;
-			// account for transmittance pdf (which is the balance heuristic combination of transmittance sampling in each color channel)
-			float3 TransmittancePdf = Transmittance; // integral of pdf from HitT to infinity
-			PathThroughput /= dot(Transmittance, ColorChannelPdf);
-
-			// our short ray hit the geomtry - decide if we should exit or not
-#if SSS_USE_TIR
-			RandSample = RandomSequence_GenerateSample2D(RandSequence).xxy;
-			float3 WorldNormal = ProbeResult.WorldNormal;
-			float CosTheta = abs(dot(Ray.Direction, WorldNormal));
-			float Fresnel = FresnelReflectance(CosTheta, 1.0 / 1.4);
-			if (RandSample.x < Fresnel)
-			{
-				// internal reflection occured -- reflect and keep tracing
-				// NOTE: weight and probability cancel out, so no throughput adjustment is needed
-				Ray.Origin += ProbeResult.HitT * Ray.Direction;
-				Ray.Direction = reflect(Ray.Direction, WorldNormal);
-				ApplyPositionBias(Ray, WorldNormal, MaxNormalBias);
-				continue;
-			}
-#endif
-			// we hit the boundary! overwrite the current payload and exit the walk with a diffuse scattering event
-			Payload.WorldPos = Ray.Origin + ProbeResult.HitT * Ray.Direction;
-			Payload.WorldNormal = ProbeResult.WorldNormal;
-			// make sure normal is pointed outward from the object so we capture illumination from the exterior
-			if (ProbeResult.FrontFace != Payload.IsFrontFace())
-			{
-				Payload.WorldNormal *= -1.0;
-			}
-			Payload.ShadingModelID = SHADINGMODELID_NUM; // invalid value so that we get diffuse shading
-			Payload.BlendingMode = RAY_TRACING_BLEND_MODE_OPAQUE;
-			Payload.BaseColor = 1;
-			return true;
-		}
-	}
-	// we scattered a bunch of times and never hit anything -- give up
-	return false;
 }
 
-uint2 Shuffle(uint2 Value, uint2 Seed)
-{
-	const uint M = PATHTRACER_COHERENT_TILE_SIZE;
-#if 1
-	// Scramble the MxM tile of pixels with an error-diffusion pattern
-	// visualizing the seeds shows the distribution of seeds is slightly better
-	const uint X = Value.x;
-	const uint Y = Value.y;
-	uint TileID = 0, HilbertState = 0;
-	for (int i = PATHTRACER_COHERENT_TILE_SHIFT - 1; i >= 0; i--) {
-		uint xi = (X >> i) & 1;
-		uint yi = (Y >> i) & 1;
-		uint Row = 8 * HilbertState + 4 * xi + 2 * yi;
-		TileID = TileID * 4 + ((0x361E9CB4u >> Row) & 3);
-		HilbertState = (0x8FE65831u >> Row) & 3;
-	}
-	uint SeedTile = StrongIntegerHash(Seed.x + StrongIntegerHash(Seed.y));
-	uint SobolIndex = FastOwenScrambling(reversebits(TileID), SeedTile);
-	uint SobolX = SobolIndex;
-	uint SobolY = 0;
-	for (uint v = 1; SobolIndex; SobolIndex >>= 1, v ^= v << 1)
-	{
-		SobolY ^= (SobolIndex & 1) * v;
-	}
-	return reversebits(uint2(SobolX, SobolY)) >> (32 - PATHTRACER_COHERENT_TILE_SHIFT);
-#elif 1
-	// randomly shuffle coordinate (x,y) within an RxR tile of pixels according to the given seed
-	// this has a slight directional preference when visualizing seeds
-	uint SeedTile = StrongIntegerHash(Seed.x + StrongIntegerHash(Seed.y));
-	uint Index = Value.x + Value.y * M;
-	uint ShuffledIndex = FastOwenScrambling(reversebits(Index), SeedTile) % (M * M);
-	return uint2(ShuffledIndex % M, ShuffledIndex / M);
 #else
-	// no shuffling (just to show repeating artifacts)
-	return Value;
-#endif
-}
-
-uint2 Reorder(uint2 I, uint2 N)
-{
-	const uint M = PATHTRACER_COHERENT_TILE_SIZE;
-	uint2 GroupId = I * M / N;				// [0, M)  - which group we landed in?
-	uint2 GroupCoord = ((I * M) % N) / M;	// [0, N/M) - which pixel within the group?
-
-	// reconstruct pixel coordinate, shuffling the assignment of pixels within a group to hide repeating structure
-	return GroupCoord * M + Shuffle(GroupId, GroupCoord);
-}
 
 RAY_TRACING_ENTRY_RAYGEN(PathTracingMainRG)
 {
-	uint2 DispatchIdx = DispatchRaysIndex().xy;
-	uint2 DispatchDim = DispatchRaysDimensions().xy;
-
-	uint2 PixelIndex = DispatchIdx;
-	if (SamplerType == PATHTRACER_SAMPLER_COHERENT)
-	{
-		// shuffle the pixels around so the ones that share a seed land in the same warp
-		PixelIndex = Reorder(PixelIndex, DispatchDim);
-	}
-	PixelIndex = PixelIndex + TileOffset.xy;
-	uint2 LaunchIndex = PixelIndex + View.ViewRectMin.xy;
-
-	// Initialize random sequence 
-	RandomSequence RandSequence;
-
-	uint2 RandPixel = SamplerType == PATHTRACER_SAMPLER_COHERENT
-		? (DispatchIdx / PATHTRACER_COHERENT_TILE_SIZE)	// unique seed within small squares of the shuffled sub-images
-		//? (DispatchIdx * PATHTRACER_COHERENT_TILE_SIZE) / DispatchDim // unique seed per sub-image
-		: LaunchIndex; // unique seed per pixel
-	if (SamplerType == PATHTRACER_SAMPLER_ERROR_DIFFUSION)
-	{
-		// z-sampler init
-		RandomSequence_Initialize(RandSequence, RandPixel, Iteration, TemporalSeed - Iteration, MaxSamples);
-	}
-	else
-	{
-		// random sobol init
-		RandomSequence_Initialize(RandSequence, RandPixel.x + RandPixel.y * 65536, TemporalSeed);
-	}
-
-
-	float3 Radiance = 0;
-	float3 Albedo = 0;
-	float3 Normal = 0;
-
-	// Initialize ray and payload
-	RayDesc Ray;
-	uint2 Pixel;
-	float EmitterPdf;
-	SampleEmitter(LaunchIndex, RandSequence, 0, 0, Ray, Pixel, EmitterPdf);
-
-	// This array will hold a CDF for light picking
-	// Seed the array with a uniform CDF at first so that we always have a valid CDF
-	float LightPickingCdf[RAY_TRACING_LIGHT_COUNT_MAXIMUM];
+	const uint2 DispatchIdx = DispatchRaysIndex().xy;
+	const uint2 DispatchDim = DispatchRaysDimensions().xy;
 
-	Ray.Direction = normalize(Ray.Direction);
-
-	FRayCone RayCone = (FRayCone)0;
-	RayCone.SpreadAngle = View.EyeToPixelSpreadAngle;
-
-	// path state variables (these cary information between bounces)
-	float3 PathThroughput = 1.0;
-	float PathRoughness = 0;
-
-	// number of directly visible lights for the first bounce
-	uint NumVisibleLights = SceneVisibleLightCount;
+	FPathState PathState = CreatePathState(DispatchIdx, DispatchDim);
 
 	for (int Bounce = 0; Bounce <= MaxBounces; Bounce++)
 	{
-#if 0
-		// visualize pixel seeds
-		Radiance = HashToColor(StrongIntegerHash(RandPixel.x + StrongIntegerHash(RandPixel.y)));
-		break;
-#endif
-		const bool bIsCameraRay = Bounce == 0;
-		const bool bIsLastBounce = Bounce == MaxBounces;
-		const bool bIncludeEmissive = (EnableDirectLighting != 0 || Bounce > 1) &&
-									  (EnableEmissive       != 0 || bIsCameraRay);
-		FPathTracingPayload Payload = TraceTransparentRay(Ray, RayCone, bIsCameraRay, bIsLastBounce, bIncludeEmissive, LaunchIndex, NumVisibleLights, RandSequence, PathThroughput, Radiance, Albedo, Normal);
-
-		if (Payload.IsMiss())
+		if (!PathTracingKernel(PathState, Bounce))
 		{
-			// we didn't hit anything selectable for further shading, we are done
+			// kernel had nothing more to do
 			break;
 		}
-
-		// As soon as the path is blurry enough, we can get away with diffuse sampling only
-		const bool bSimplifySSS = PathRoughness >= 0.15;
-
-		if (ApproximateCaustics)
-		{
-			// modify the payload roughness to minimize difficult caustics
-			// This is inspired by a trick used in the Arnold renderer:
-			//   https://cgg.mff.cuni.cz/~jaroslav/gicourse2010/giai2010-02-marcos_fajardo-slides.pdf (slide 39)
-			//   https://www.arnoldrenderer.com/research/Arnold_TOG2018.pdf (section 4.2)
-			AdjustMaterialRoughness(Payload, PathRoughness);
-		}
-
-		// Processing the random walk will (stochastically) move the shading point elsewhere on the surface of the object
-		if (!ProcessSubsurfaceRandomWalk(Payload, PathThroughput, RandSequence, LaunchIndex, bSimplifySSS))
-		{
-			// random walk did not terminate at a valid point
-			break;
-		}
-
-		FLightLoopCount LightLoopCount = LightGridLookup(Payload.WorldPos);
-
-		// visualize the # of lights
-		if (VisualizeLightGrid)
-		{
-			Radiance = LightGridVisualize(LightLoopCount, VisualizeLightGrid);
-			Radiance *= abs(dot(Payload.WorldNormal, Ray.Direction));
-			Radiance /= View.PreExposure; // cancel out exposure
-			break;
-		}
-
-		// Choose a random number for both Light sampling and BxDF sampling
-		float4 RandSample = RandomSequence_GenerateSample4D(RandSequence);
-
-		// Does this material require NEE? (will be false if MaterialPdf is always +inf)
-		const bool bIsNeeValid = IsNeeValidMaterial(Payload);
-
-		const bool bDoLightLoop = EnableDirectLighting != 0 || Bounce > 0;
-
-		float LightPickingCdfSum = 0;
-
-		// If we are using Light sampling and the material can use it ...
-		if (MISMode != 0 && bIsNeeValid && SceneLightCount > 0 && bDoLightLoop)
-		{
-			// Choose a light and sample it
-			float3 WorldPos = Payload.WorldPos;
-			float3 WorldNormal = Payload.WorldNormal;
-			uint PrimitiveLightingChannelMask = Payload.PrimitiveLightingChannelMask;
-
-			bool IsTransmissiveMaterial = ENABLE_TRANSMISSION && Payload.IsMaterialTransmissive();
-
-			for (uint Index = 0, Num = LightLoopCount.NumLights; Index < Num; ++Index)
-			{
-				uint LightIndex = GetLightId(Index, LightLoopCount);
-				LightPickingCdfSum += EstimateLight(LightIndex, WorldPos, WorldNormal, PrimitiveLightingChannelMask, IsTransmissiveMaterial);
-				LightPickingCdf[Index] = LightPickingCdfSum;
-			}
-
-			if (LightPickingCdfSum > 0)
-			{
-				// init worked
-				int LightId;
-				float LightPickPdf = 0;
-
-				SelectLight(RandSample.x * LightPickingCdfSum, LightLoopCount.NumLights, LightPickingCdf, LightId, LightPickPdf);
-				LightPickPdf /= LightPickingCdfSum;
-
-				LightId = GetLightId(LightId, LightLoopCount);
-				FLightSample LightSample = SampleLight(LightId, RandSample.yz, WorldPos, WorldNormal);
-
-				LightSample.RadianceOverPdf /= LightPickPdf;
-				LightSample.Pdf *= LightPickPdf;
-				if (LightSample.Pdf > 0)
-				{
-					if (CastsShadow(LightId))
-					{
-						// for transmissive materials, bias the position to the other side of the surface if the light is coming from behind
-						const float SignedPositionBias = IsTransmissiveMaterial ? sign(dot(Payload.WorldNormal, LightSample.Direction)) : 1.0;
-
-						RayDesc LightRay;
-						LightRay.Origin = WorldPos;
-						LightRay.TMin = 0;
-						LightRay.Direction = LightSample.Direction;
-						LightRay.TMax = LightSample.Distance;
-						ApplyPositionBias(LightRay, SignedPositionBias * Payload.WorldNormal, MaxNormalBias);
-
-						float AvgRoughness = ApproximateCaustics ? GetAverageRoughness(Payload) : 0.0;
-
-						LightSample.RadianceOverPdf *= TraceTransparentVisibilityRay(LightRay, LaunchIndex, AvgRoughness);
-					}
-
-					// #dxr_todo: Is it cheaper to fire the ray first? Or eval the material first?
-					if (any(LightSample.RadianceOverPdf > 0))
-					{
-						// Evaluate material
-						FMaterialEval MaterialEval = EvalMaterial(Ray.Direction, LightSample.Direction, Payload, false);
-
-						// Record the contribution
-						float3 LightContrib = PathThroughput * LightSample.RadianceOverPdf * MaterialEval.Weight * MaterialEval.Pdf;
-						if (MISMode == 2)
-						{
-							LightContrib *= MISWeightRobust(LightSample.Pdf, MaterialEval.Pdf);
-						}
-						AccumulateRadiance(Radiance, LightContrib, bIsCameraRay);
-					}
-				}
-			}
-		}
-
-		// Sample material
-		FMaterialSample MaterialSample = SampleMaterial(Ray.Direction, Payload, RandSample, false);
-
-		if (MaterialSample.Pdf < 0 || asuint(MaterialSample.Pdf) > 0x7F800000)
-		{
-			// Pdf became invalid (either negative or NaN)
-			Radiance = float3(1, 0, 1);
-			break;
-		}
-
-		if (!(MaterialSample.Pdf > 0))
-		{
-			// No valid direction -- we are done
-			break;
-		}
-
-		float3 NextPathThroughput = PathThroughput * MaterialSample.Weight;
-		if (!any(NextPathThroughput > 0))
-		{
-			// no energy left in this path
-			break;
-		}
-
-		// Russian roulette:
-		//   The probability of keeping the path should be roughly proportional to the weight at the current shade point,
-		//  but just using MaterialWeight would miss out on cases where the path throughput changes color (like in a cornell
-		//  box when bouncing between walls of different colors). So use the ratio of the brightest color channel in the
-		//  previous and next throughput.
-		//   The second tweak is to add a sqrt() around the probability to soften the termination probability (paths will last
-		//  a little longer). This allows paths to go a bit deeper than the naive heuristic while still allowing them to terminate
-		//  early. This makes RR effective from the very first bounce without needing to delay it.
-		float ContinuationProb = sqrt(saturate(max(NextPathThroughput.x, max(NextPathThroughput.y, NextPathThroughput.z)) / max(PathThroughput.x, max(PathThroughput.y, PathThroughput.z))));
-		if (ContinuationProb < 1)
-		{
-			// If there is some chance we should terminate the ray, draw an extra random value
-			float RussianRouletteRand = RandSample.w; // SampleMaterial does not use this value at the moment
-			//RussianRouletteRand = RandomSequence_GenerateSample1D(RandSequence);
-			if (RussianRouletteRand >= ContinuationProb)
-			{
-				// stochastically terminate the path
-				break;
-			}
-			PathThroughput = NextPathThroughput / ContinuationProb;
-		}
-		else
-		{
-			PathThroughput = NextPathThroughput;
-		}
-
-		// Update ray according to material sample
-		Ray.Origin = Payload.WorldPos;
-		Ray.Direction = MaterialSample.Direction;
-		Ray.TMin = 0;
-		Ray.TMax = POSITIVE_INFINITY;
-		ApplyPositionBias(Ray, MaterialSample.PositionBiasSign * Payload.WorldNormal, MaxNormalBias);
-
-		if (ApproximateCaustics)
-		{
-			// enlarge roughness based on the chosen lobe roughness
-			PathRoughness = max(PathRoughness, MaterialSample.Roughness);
-		}
-
-		// If we are using Material sampling for lights
-		if (MISMode != 1 && bDoLightLoop)
-		{
-			// Check which lights can be seen by the material ray and trace a dedicated shadow ray
-			// While it would be possible to just loop around and use the indirect ray to do this, it would prevent the application
-			// of shadow ray specific logic for transparent shadows or various per light tricks like shadow casting
-			const bool bUseMIS = MISMode == 2 && LightPickingCdfSum > 0;
-			for (uint Index = 0, Num = LightLoopCount.NumMISLights; Index < Num; ++Index)
-			{
-				uint LightId = GetLightId(Index, LightLoopCount);
-				if ((Payload.PrimitiveLightingChannelMask & GetLightingChannelMask(LightId)) == 0)
-				{
-					// light does not affect the current ray
-					continue;
-				}
-
-				FLightHit LightResult = TraceLight(Ray, LightId);
-
-				if (LightResult.IsMiss())
-				{
-					continue;
-				}
-
-				float3 LightContrib = PathThroughput * LightResult.Radiance;
-				if (bUseMIS)
-				{
-					float PreviousCdfValue = Index > 0 ? LightPickingCdf[Index - 1] : 0.0;
-					float LightPickPdf = (LightPickingCdf[Index] - PreviousCdfValue) / LightPickingCdfSum;
-
-					LightContrib *= MISWeightRobust(MaterialSample.Pdf, LightResult.Pdf * LightPickPdf);
-				}
-
-				if (any(LightContrib > 0))
-				{
-					if (CastsShadow(LightId))
-					{
-						RayDesc LightRay = Ray;
-						LightRay.TMax = LightResult.HitT;
-						LightContrib *= TraceTransparentVisibilityRay(LightRay, LaunchIndex, PathRoughness);
-					}
-					// the light made some contribution, and there was nothing along the shadow ray
-					AccumulateRadiance(Radiance, LightContrib, bIsCameraRay);
-				}
-			}
-			
-		}
-
-		// from this point on, we don't need to include lights in the trace call
-		// because NEE handled it for us
-		NumVisibleLights = 0;
-		
-		// update ray cone for the next bounce
-		RayCone = Payload.RayCone;
 	}
 
-
 	// Accumulate radiance and update pixel variance
-	uint NumSamples = Iteration + 1;
-	float BlendFactor = 1.0 / float(NumSamples);
-	// Avoid reading the old pixel on the first sample on the off-chance there is a NaN/Inf pixel ...
-	float4 OldPixel = NumSamples > 1 ? RadianceTexture[PixelIndex] : 0;
-	float4 OldAlbedo = NumSamples > 1 ? AlbedoTexture[PixelIndex] : 0;
-	float4 OldNormal = NumSamples > 1 ? NormalTexture[PixelIndex] : 0;
-	float3 OldRadiance = OldPixel.rgb;
-	float OldVariance = OldPixel.a;
-
-	float3 NewRadiance = lerp(OldRadiance, Radiance, BlendFactor);
-
-	float RadianceY = Luminance(Radiance);
-
-	// https://en.wikipedia.org/wiki/Algorithms_for_calculating_variance#Welford's_online_algorithm
-	float DeviationSquared = (RadianceY - Luminance(NewRadiance)) * (RadianceY - Luminance(OldRadiance));
-	float NewVariance = lerp(OldVariance, DeviationSquared, BlendFactor);
-
-	RadianceTexture[PixelIndex] = float4(NewRadiance, NewVariance);
-	AlbedoTexture[PixelIndex] = lerp(OldAlbedo, float4(Albedo, 0), BlendFactor);
-	NormalTexture[PixelIndex] = lerp(OldNormal, float4(Normal, 0), BlendFactor);
+	PathState.WritePixel(BlendFactor);
 }
+
+#endif
\ No newline at end of file
diff --git a/Engine/Shaders/Private/PathTracing/PathTracingCore.ush b/Engine/Shaders/Private/PathTracing/PathTracingCore.ush
new file mode 100644
--- /dev/null
+++ b/Engine/Shaders/Private/PathTracing/PathTracingCore.ush
@@ -0,0 +1,818 @@
+// Copyright Epic Games, Inc. All Rights Reserved.
+
+#pragma once
+
+#define RANDSEQ_UNROLL_SOBOL PATH_TRACER_USE_COMPACTION // unrolling sobol sampler is best when we can guarantee high occupancy
+
+#define USE_PATH_TRACING_LIGHT_GRID	1
+
+#include "../Common.ush"
+#include "../PostProcessCommon.ush"
+#include "../RectLight.ush"
+#include "../RayTracing/RayTracingCommon.ush"
+
+#include "PathTracingCommon.ush"
+#include "../RayTracing/RayTracingHitGroupCommon.ush"
+
+#include "../ShadingModels.ush"
+#include "./Utilities/PathTracingRandomSequence.ush"
+#include "./Light/PathTracingLightSampling.ush"
+#include "./Light/PathTracingLightGrid.ush"
+#include "./Material/PathTracingMaterialSampling.ush"
+#include "PathTracingCamera.ush"
+
+float BlendFactor;
+uint Iteration;
+uint TemporalSeed;
+uint MaxSamples;
+uint MaxBounces;
+uint MaxSSSBounces;
+
+// 0: only Material sampling
+// 1: only Light sampling
+// 2: both Material and Light
+uint MISMode;
+
+uint ApproximateCaustics;
+uint EnableCameraBackfaceCulling;
+uint EnableDirectLighting;
+uint EnableEmissive;
+uint SamplerType;
+uint VisualizeLightGrid;
+float MaxPathIntensity;
+float MaxNormalBias;
+float FilterWidth;
+
+RWTexture2D<float4> RadianceTexture;
+RWTexture2D<float4> AlbedoTexture;
+RWTexture2D<float4> NormalTexture;
+RaytracingAccelerationStructure TLAS;
+int3 TileOffset;
+uint SceneVisibleLightCount;
+
+
+void AccumulateRadiance(inout float3 TotalRadiance, float3 PathRadiance, bool bIsCameraRay)
+{
+	// User asked for path contributions to be clamped to reduce fireflies.
+	// This puts an upper bound on variance within the pixel at the expense of bias
+	float MaxPathRadiance = max3(PathRadiance.x, PathRadiance.y, PathRadiance.z);
+	if (!bIsCameraRay && MaxPathRadiance > MaxPathIntensity)
+	{
+		// adjust sample while keeping color
+		PathRadiance *= MaxPathIntensity / MaxPathRadiance;
+	}
+
+	// accumulate into the total
+	TotalRadiance += PathRadiance;
+}
+
+FPathTracingPayload TraceTransparentRay(RayDesc Ray, FRayCone RayCone, bool bIsCameraRay, bool bLastBounce, bool bIncludeEmission, uint2 LaunchIndex, uint NumLights, inout RandomSequence RandSequence, inout float3 PathThroughput, inout float3 Radiance, inout float3 Albedo, inout float3 Normal)
+{
+	const uint RayFlags = bIsCameraRay && EnableCameraBackfaceCulling ? RAY_FLAG_CULL_BACK_FACING_TRIANGLES : 0;
+	const uint MissShaderIndex = 0;
+	float SelectionWeightSum = 0;
+	float3 PayloadThroughput = PathThroughput;
+	FPathTracingPayload Payload;
+	if (!bIncludeEmission && bLastBounce)
+	{
+		Payload.SetMiss();
+		PathThroughput = 0;
+		return Payload;
+	}
+	for (;;)
+	{
+		// avoid double counting what was captured by the skylight
+		// also avoid noise from hot spots (they can be properly
+		// importance sampled if a capturing skylight is added)
+		const bool bIgnoreSky = !bIsCameraRay;
+		FPackedPathTracingPayload PackedPayload = InitPathTracingPayload(RayCone, LaunchIndex, bIgnoreSky);
+		// Trace the ray
+		TraceRay(
+			TLAS,
+			RayFlags,
+			RAY_TRACING_MASK_ALL,
+			RAY_TRACING_SHADER_SLOT_MATERIAL,
+			RAY_TRACING_NUM_SHADER_SLOTS,
+			MissShaderIndex,
+			Ray,
+			PackedPayload);
+
+		// Loop over lights to capture their contribution
+		// #dxr_todo: if we have lots of lights, having some hierarchical structure would be better ....
+		for (uint LightId = 0; LightId < NumLights; ++LightId)
+		{
+			RayDesc LightRay = Ray;
+			LightRay.TMax = PackedPayload.IsMiss() ? Ray.TMax : PackedPayload.HitT;
+			float3 LightRadiance = TraceLight(LightRay, LightId).Radiance;
+			AccumulateRadiance(Radiance, PathThroughput * LightRadiance, bIsCameraRay);
+		}
+
+		if (PackedPayload.IsMiss())
+		{
+			// Ray didn't hit any real geometry, so nothing left to do
+			break;
+		}
+
+		// Unpack the payload
+		FPathTracingPayload HitPayload = UnpackPathTracingPayloadMaterial(PackedPayload, Ray);
+
+		if (HitPayload.BlendingMode == RAY_TRACING_BLEND_MODE_ADDITIVE ||
+			HitPayload.BlendingMode == RAY_TRACING_BLEND_MODE_TRANSLUCENT)
+		{
+			// emission needs to be pre-multiplied by the opacity
+			HitPayload.Radiance *= HitPayload.Opacity;
+		}
+
+		// add in surface emission (except for modulate blend mode which uses Radiance as the transparency multiplier)
+		if (bIncludeEmission && HitPayload.BlendingMode != RAY_TRACING_BLEND_MODE_MODULATE)
+		{
+			AccumulateRadiance(Radiance, PathThroughput * HitPayload.Radiance, bIsCameraRay);
+		}
+
+		float3 Transparency = GetMaterialTransparency(HitPayload, Ray.Direction);
+
+		if (!bLastBounce)
+		{
+			float3 Contrib = PathThroughput * EstimateMaterialAlbedo(HitPayload);
+			if (bIsCameraRay)
+			{
+				Albedo += Contrib;
+				Normal += PathThroughput * (1 - Luminance(Transparency)) * HitPayload.WorldNormal;
+			}
+
+			float SelectionWeight = max3(Contrib.x, Contrib.y, Contrib.z);
+			SelectionWeightSum += SelectionWeight;
+			bool AcceptHit = false;
+
+			// weighted reservoir sampling
+			if (SelectionWeight > 0)
+			{
+				if (SelectionWeight < SelectionWeightSum)
+				{
+					// the acceptance probability is not 1.0
+					// generate a random number to see if we should accept this hit
+					float RandValue = RandomSequence_GenerateSample1D(RandSequence);
+					AcceptHit = RandValue * SelectionWeightSum < SelectionWeight;
+				}
+				else
+				{
+					// accept automatically on the first hit
+					AcceptHit = true;
+				}
+			}
+
+			if (AcceptHit)
+			{
+				// stash this hit for next time
+				Payload = HitPayload;
+				PayloadThroughput = PathThroughput / SelectionWeight;
+			}
+		}
+
+		// prepare next step around the loop
+		// retrace the exact same ray with TMin one ulp past the hit we just found
+		PathThroughput *= Transparency;
+		Ray.TMin = asfloat(asuint(HitPayload.HitT) + 1);
+
+		if (all(PathThroughput == 0))
+		{
+			break;
+		}
+	}
+
+	if (SelectionWeightSum > 0)
+	{
+		// if we stored a valid hit in the payload, reset the path throughput to this point
+		PathThroughput = PayloadThroughput * SelectionWeightSum;
+	}
+	else
+	{
+		PathThroughput = 0;
+		Payload.SetMiss();
+	}
+	return Payload;
+}
+
+float3 TraceTransparentVisibilityRay(RayDesc Ray, uint2 PixelCoord, float PathRoughness)
+{
+	const uint RayFlags = RAY_FLAG_ACCEPT_FIRST_HIT_AND_END_SEARCH | RAY_FLAG_SKIP_CLOSEST_HIT_SHADER;
+	const uint InstanceInclusionMask = RAY_TRACING_MASK_SHADOW;
+	const uint RayContributionToHitGroupIndex = RAY_TRACING_SHADER_SLOT_SHADOW;
+	const uint MultiplierForGeometryContributionToShaderIndex = RAY_TRACING_NUM_SHADER_SLOTS;
+	const uint MissShaderIndex = 0;
+
+	FPackedPathTracingPayload PackedPayload = InitPathTracingVisibilityPayload(PixelCoord, PathRoughness);
+
+	// Trace the ray
+	TraceRay(
+		TLAS,
+		RayFlags,
+		InstanceInclusionMask,
+		RayContributionToHitGroupIndex,
+		MultiplierForGeometryContributionToShaderIndex,
+		MissShaderIndex,
+		Ray,
+		PackedPayload);
+
+	if (PackedPayload.IsMiss())
+	{
+		// we didn't hit any geo - just return whatever throughput is left
+		return PackedPayload.GetRayThroughput();
+	}
+
+	// we registered a hit -- assumed to be fully opaque
+	return 0;
+}
+
+struct FProbeResult
+{
+	float HitT;
+	float3 WorldNormal;
+	int FrontFace;
+
+	bool IsMiss() { return HitT <= 0; }
+};
+
+FProbeResult TraceProbeRay(RayDesc Ray, uint2 PixelCoord)
+{
+	// Trace a short ray to see if we escaped the surface
+	// TODO: figure out a way to only hit the object we started on (or handle blocking geometry)
+	FPackedPathTracingPayload PackedPayload = InitPathTracingPayload((FRayCone)0, PixelCoord, true);
+	// TODO: payload flag to skip shading (we only care about the normal)
+	const uint RayFlags = 0;
+	const uint MissShaderIndex = 0;
+	TraceRay(
+		TLAS,
+		RayFlags,
+		RAY_TRACING_MASK_ALL,
+		RAY_TRACING_SHADER_SLOT_MATERIAL,
+		RAY_TRACING_NUM_SHADER_SLOTS,
+		MissShaderIndex,
+		Ray,
+		PackedPayload);
+
+	if (PackedPayload.IsMiss())
+	{
+		return (FProbeResult)0;
+	}
+
+	FProbeResult Result;
+	Result.HitT = PackedPayload.HitT;
+	Result.WorldNormal = PackedPayload.GetWorldNormal();
+	Result.FrontFace = PackedPayload.IsFrontFace();
+	return Result;
+}
+
+bool ProcessSubsurfaceRandomWalk(inout FPathTracingPayload Payload, inout float3 PathThroughput, inout RandomSequence RandSequence, uint2 PixelCoord, bool SimplifySSS)
+{
+	if (!Payload.IsSubsurfaceMaterial())
+	{
+		// material doesn't have SSS
+		// just return so we can process the rest of the payload's shading
+		return true;
+	}
+
+	float3 SSSColor = Payload.GetSubsurfaceColor();
+	float3 SSSRadius = Payload.GetSubsurfaceRadius();
+
+	if (SimplifySSS || all(SSSRadius == 0) || all(SSSColor == 0) || MaxSSSBounces == 0)
+	{
+		// if we are running in a context that won't see the SSS directly -- just turn it into diffuse and skip the random walk
+		// we can also skip random walk if the radius or color is 0
+		// or if the user decided not to perform any SSS bounces
+		Payload.DiffuseColor += SSSColor;
+		Payload.DiffuseColor = saturate(Payload.DiffuseColor);
+		Payload.SetSubsurfaceColor(0.0);
+		return true;
+	}
+	SSSRadius = max(SSSRadius, 0.1); // 1mm minimum to make sure random walk can move forward and to keep SigmaT finite
+
+	// decide randomly to evaluate the SSS part of the material, or keep shading the current part
+	float SSSProb = LobeSelectionProb(SSSColor, Payload.DiffuseColor + Payload.SpecularColor);
+	float3 RandSample = RandomSequence_GenerateSample3D(RandSequence);
+	if (RandSample.x < SSSProb)
+	{
+		// do random walk
+		PathThroughput *= 1 / SSSProb;
+		RandSample.x /= SSSProb;
+	}
+	else
+	{
+		// skip random walk and shade current point instead
+		PathThroughput *= 1 / (1 - SSSProb);
+		return true;
+	}
+
+	// create ray to enter the surface (using a diffuse scattering event)
+	RayDesc Ray;
+	Ray.Origin = Payload.WorldPos;
+	Ray.Direction = TangentToWorld(-CosineSampleHemisphere(RandSample.xy).xyz, Payload.WorldNormal);
+	Ray.TMin = 0;
+	ApplyPositionBias(Ray, -Payload.WorldNormal, MaxNormalBias);
+
+	// Use the technique detailed in the Hyperion paper (Section 4.4.2)
+	// https://media.disneyanimation.com/uploads/production/publication_asset/177/asset/a.pdf
+	// This avoids overly bright thin regions by applying a bit of internal reflection when the boundary is hit
+#define SSS_USE_TIR 1
+
+#if SSS_USE_TIR
+	const float3 Albedo = 1 - exp(SSSColor * (-11.43 + SSSColor * (15.38 - 13.91 * SSSColor)));
+#else
+	// Van de-Hulst inverse mapping
+	// https://blog.selfshadow.com/publications/s2017-shading-course/imageworks/s2017_pbs_imageworks_slides_v2.pdf (Slide 44)
+	// [REDACTED_PROFILE_URL] (Section 7.5.3 of v0.1.3)
+	const float3 Albedo = 1 - Pow2(4.09712 + 4.20863 * SSSColor - sqrt(9.59217 + SSSColor * (41.6808 + 17.7126 * SSSColor)));
+	SSSRadius *= 2.0; // roughly match parameterization above
+#endif
+	const int MAX_SSS_BOUNCES = MaxSSSBounces;
+	const float3 SigmaT = 1.0 / SSSRadius;
+	const float3 SigmaS = Albedo * SigmaT;
+	for (int i = 0; i < MAX_SSS_BOUNCES; i++)
+	{
+		// TODO: add Dwivedi sampling for faster escape in low albedo media ...
+
+		// Pick a color channel for transmittance sampling:
+		float3 ColorChannelWeight = PathThroughput * Albedo;
+		float3 ColorChannelCdf = float3(ColorChannelWeight.x,
+			ColorChannelWeight.x + ColorChannelWeight.y,
+			ColorChannelWeight.x + ColorChannelWeight.y + ColorChannelWeight.z);
+		if (ColorChannelCdf.z == 0)
+		{
+			// no more energy left in the path
+			break;
+		}
+		ColorChannelCdf.xy /= ColorChannelCdf.z;
+		float3 ColorChannelPdf = float3(ColorChannelCdf.x, ColorChannelCdf.y - ColorChannelCdf.x, 1 - ColorChannelCdf.y);
+		// pick a maximum distance to scatter
+		const float OneMinusEpsilon = 0.99999994; // 32-bit float just before 1.0
+		if (RandSample.z < ColorChannelCdf.x)
+		{
+			float RescaleRand = min(RandSample.z / ColorChannelPdf.x, OneMinusEpsilon);
+			Ray.TMax = -log(1 - RescaleRand) * SSSRadius.x;
+		}
+		else if (RandSample.z < ColorChannelCdf.y)
+		{
+			float RescaleRand = min((RandSample.z - ColorChannelCdf.x) / ColorChannelPdf.y, OneMinusEpsilon);
+			Ray.TMax = -log(1 - RescaleRand) * SSSRadius.y;
+		}
+		else
+		{
+			float RescaleRand = min((RandSample.z - ColorChannelCdf.y) / ColorChannelPdf.z, OneMinusEpsilon);
+			Ray.TMax = -log(1 - RescaleRand) * SSSRadius.z;
+		}
+
+		FProbeResult ProbeResult = TraceProbeRay(Ray, PixelCoord);
+		if (ProbeResult.IsMiss())
+		{
+			// we didn't hit anything, so scatter randomly in the volume and keep tracing
+			RandSample = RandomSequence_GenerateSample3D(RandSequence);
+			Ray.Origin += Ray.TMax * Ray.Direction;
+			Ray.Direction = UniformSampleSphere(RandSample.xy).xyz;
+			// account for transmittance and scattering up to this point
+			float3 Transmittance = exp(-Ray.TMax * SigmaT);
+			PathThroughput *= SigmaS * Transmittance;
+			// account for transmittance pdf (which is the balance heuristic combination of transmittance sampling in each color channel)
+			float3 TransmittancePdf = SigmaT * Transmittance;
+			PathThroughput /= dot(TransmittancePdf, ColorChannelPdf);
+
+			// NOTE: phase function cancels out since it is being perfectly importance sampled
+		}
+		else
+		{
+			// account for transmittance and scattering up to this point
+			float3 Transmittance = exp(-ProbeResult.HitT * SigmaT);
+			PathThroughput *= Transmittance;
+			// account for transmittance pdf (which is the balance heuristic combination of transmittance sampling in each color channel)
+			float3 TransmittancePdf = Transmittance; // integral of pdf from HitT to infinity
+			PathThroughput /= dot(Transmittance, ColorChannelPdf);
+
+			// our short ray hit the geomtry - decide if we should exit or not
+#if SSS_USE_TIR
+			RandSample = RandomSequence_GenerateSample2D(RandSequence).xxy;
+			float3 WorldNormal = ProbeResult.WorldNormal;
+			float CosTheta = abs(dot(Ray.Direction, WorldNormal));
+			float Fresnel = FresnelReflectance(CosTheta, 1.0 / 1.4);
+			if (RandSample.x < Fresnel)
+			{
+				// internal reflection occured -- reflect and keep tracing
+				// NOTE: weight and probability cancel out, so no throughput adjustment is needed
+				Ray.Origin += ProbeResult.HitT * Ray.Direction;
+				Ray.Direction = reflect(Ray.Direction, WorldNormal);
+				ApplyPositionBias(Ray, WorldNormal, MaxNormalBias);
+				continue;
+			}
+#endif
+			// we hit the boundary! overwrite the current payload and exit the walk with a diffuse scattering event
+			Payload.WorldPos = Ray.Origin + ProbeResult.HitT * Ray.Direction;
+			Payload.WorldNormal = ProbeResult.WorldNormal;
+			// make sure normal is pointed outward from the object so we capture illumination from the exterior
+			if (ProbeResult.FrontFace != Payload.IsFrontFace())
+			{
+				Payload.WorldNormal *= -1.0;
+			}
+			Payload.ShadingModelID = SHADINGMODELID_NUM; // invalid value so that we get diffuse shading
+			Payload.BlendingMode = RAY_TRACING_BLEND_MODE_OPAQUE;
+			Payload.BaseColor = 1;
+			return true;
+		}
+	}
+	// we scattered a bunch of times and never hit anything -- give up
+	return false;
+}
+
+uint2 Shuffle(uint2 Value, uint2 Seed)
+{
+	const uint M = PATHTRACER_COHERENT_TILE_SIZE;
+#if 1
+	// Scramble the MxM tile of pixels with an error-diffusion pattern
+	// visualizing the seeds shows the distribution of seeds is slightly better
+	const uint X = Value.x;
+	const uint Y = Value.y;
+	uint TileID = 0, HilbertState = 0;
+	for (int i = PATHTRACER_COHERENT_TILE_SHIFT - 1; i >= 0; i--) {
+		uint xi = (X >> i) & 1;
+		uint yi = (Y >> i) & 1;
+		uint Row = 8 * HilbertState + 4 * xi + 2 * yi;
+		TileID = TileID * 4 + ((0x361E9CB4u >> Row) & 3);
+		HilbertState = (0x8FE65831u >> Row) & 3;
+	}
+	uint SeedTile = StrongIntegerHash(Seed.x + StrongIntegerHash(Seed.y));
+	uint SobolIndex = FastOwenScrambling(reversebits(TileID), SeedTile);
+	uint SobolX = SobolIndex;
+	uint SobolY = 0;
+	for (uint v = 1; SobolIndex; SobolIndex >>= 1, v ^= v << 1)
+	{
+		SobolY ^= (SobolIndex & 1) * v;
+	}
+	return reversebits(uint2(SobolX, SobolY)) >> (32 - PATHTRACER_COHERENT_TILE_SHIFT);
+#elif 1
+	// randomly shuffle coordinate (x,y) within an RxR tile of pixels according to the given seed
+	// this has a slight directional preference when visualizing seeds
+	uint SeedTile = StrongIntegerHash(Seed.x + StrongIntegerHash(Seed.y));
+	uint Index = Value.x + Value.y * M;
+	uint ShuffledIndex = FastOwenScrambling(reversebits(Index), SeedTile) % (M * M);
+	return uint2(ShuffledIndex % M, ShuffledIndex / M);
+#else
+	// no shuffling (just to show repeating artifacts)
+	return Value;
+#endif
+}
+
+uint2 Reorder(uint2 I, uint2 N)
+{
+	const uint M = PATHTRACER_COHERENT_TILE_SIZE;
+	uint2 GroupId = I * M / N;				// [0, M)  - which group we landed in?
+	uint2 GroupCoord = ((I * M) % N) / M;	// [0, N/M) - which pixel within the group?
+
+	// reconstruct pixel coordinate, shuffling the assignment of pixels within a group to hide repeating structure
+	return GroupCoord * M + Shuffle(GroupId, GroupCoord);
+}
+
+struct FPathState
+{                                   // packed size
+	uint2 PixelIndex;               //  4 bytes
+	RandomSequence RandSequence;    //  8 bytes
+	float3 Radiance;                // 12 bytes
+	float3 Albedo;                  //  6 bytes // half precision
+	float3 Normal;                  //  6 bytes // half precision
+	RayDesc Ray;                    // 12+12 = 24 bytes (TMin/TMax do not need to be stored)
+	FRayCone RayCone;               //  4 bytes // half precision
+	float3 PathThroughput;          // 12 bytes
+	float PathRoughness;            //  4 bytes
+	                                // 80 bytes total
+
+	void WritePixel(float BlendFactor)
+	{
+		// Avoid reading the old pixel on the first sample on the off-chance there is a NaN/Inf pixel ...
+		float4 OldPixel = Iteration > 0 ? RadianceTexture[PixelIndex] : 0;
+		float4 OldAlbedo = Iteration > 0 ? AlbedoTexture[PixelIndex] : 0;
+		float4 OldNormal = Iteration > 0 ? NormalTexture[PixelIndex] : 0;
+		float3 OldRadiance = OldPixel.rgb;
+		float OldVariance = OldPixel.a;
+
+		float3 NewRadiance = lerp(OldRadiance, Radiance, BlendFactor);
+
+		float RadianceY = Luminance(Radiance);
+
+		// https://en.wikipedia.org/wiki/Algorithms_for_calculating_variance#Welford's_online_algorithm
+		float DeviationSquared = (RadianceY - Luminance(NewRadiance)) * (RadianceY - Luminance(OldRadiance));
+		float NewVariance = lerp(OldVariance, DeviationSquared, BlendFactor);
+
+		RadianceTexture[PixelIndex] = float4(NewRadiance, NewVariance);
+		AlbedoTexture[PixelIndex] = lerp(OldAlbedo, float4(Albedo, 0), BlendFactor);
+		NormalTexture[PixelIndex] = lerp(OldNormal, float4(Normal, 0), BlendFactor);
+	}
+};
+
+FPathState CreatePathState(uint2 DispatchIdx, uint2 DispatchDim)
+{
+	FPathState Output = (FPathState)0;
+
+	Output.PixelIndex = DispatchIdx;
+	if (SamplerType == PATHTRACER_SAMPLER_COHERENT)
+	{
+		// shuffle the pixels around so the ones that share a seed land in the same warp
+		Output.PixelIndex = Reorder(Output.PixelIndex, DispatchDim);
+	}
+	Output.PixelIndex = Output.PixelIndex + TileOffset.xy;
+	uint2 LaunchIndex = Output.PixelIndex + View.ViewRectMin.xy;
+
+	// Initialize random sequence 
+
+	uint2 RandPixel = SamplerType == PATHTRACER_SAMPLER_COHERENT
+		? (DispatchIdx / PATHTRACER_COHERENT_TILE_SIZE)	// unique seed within small squares of the shuffled sub-images
+		//? (DispatchIdx * PATHTRACER_COHERENT_TILE_SIZE) / DispatchDim // unique seed per sub-image
+		: LaunchIndex; // unique seed per pixel
+	if (SamplerType == PATHTRACER_SAMPLER_ERROR_DIFFUSION)
+	{
+		// z-sampler init
+		RandomSequence_Initialize(Output.RandSequence, RandPixel, Iteration, TemporalSeed - Iteration, MaxSamples);
+	}
+	else
+	{
+		// random sobol init
+		RandomSequence_Initialize(Output.RandSequence, RandPixel.x + RandPixel.y * 65536, TemporalSeed);
+	}
+
+
+	// Initialize ray and payload
+	{
+		uint2 Pixel;
+		float EmitterPdf;
+		SampleEmitter(LaunchIndex, Output.RandSequence, 0, 0, Output.Ray, Pixel, EmitterPdf);
+	}
+
+	Output.Ray.Direction = normalize(Output.Ray.Direction);
+
+	Output.RayCone.SpreadAngle = View.EyeToPixelSpreadAngle;
+
+	// path state variables (these cary information between bounces)
+	Output.PathThroughput = 1.0;
+	Output.PathRoughness = 0;
+
+	return Output;
+}
+
+bool PathTracingKernel(inout FPathState PathState, int Bounce)
+{
+	// This array will hold a CDF for light picking
+	float LightPickingCdf[RAY_TRACING_LIGHT_COUNT_MAXIMUM];
+
+#if 0
+	// visualize pixel seeds
+	Radiance = HashToColor(StrongIntegerHash(RandPixel.x + StrongIntegerHash(RandPixel.y)));
+	return false;
+#endif
+	const bool bIsCameraRay = Bounce == 0;
+	const bool bIsLastBounce = Bounce == MaxBounces;
+	const bool bIncludeEmissive = (EnableDirectLighting != 0 || Bounce > 1) &&
+		(EnableEmissive != 0 || bIsCameraRay);
+
+	// number of directly visible lights for the first bounce
+	// after the first bounce, we don't need to include lights in the trace call
+	// because NEE handled it for us
+	uint NumVisibleLights = bIsCameraRay ? SceneVisibleLightCount : 0;
+
+	FPathTracingPayload Payload = TraceTransparentRay(
+		PathState.Ray,
+		PathState.RayCone,
+		bIsCameraRay,
+		bIsLastBounce,
+		bIncludeEmissive,
+		PathState.PixelIndex,
+		NumVisibleLights,
+		PathState.RandSequence,
+		PathState.PathThroughput,
+		PathState.Radiance,
+		PathState.Albedo,
+		PathState.Normal);
+
+	if (Payload.IsMiss())
+	{
+		// we didn't hit anything selectable for further shading, we are done
+		return false;
+	}
+
+	// As soon as the path is blurry enough, we can get away with diffuse sampling only
+	const bool bSimplifySSS = PathState.PathRoughness >= 0.15;
+
+	if (ApproximateCaustics)
+	{
+		// modify the payload roughness to minimize difficult caustics
+		// This is inspired by a trick used in the Arnold renderer:
+		//   https://cgg.mff.cuni.cz/~jaroslav/gicourse2010/giai2010-02-marcos_fajardo-slides.pdf (slide 39)
+		//   https://www.arnoldrenderer.com/research/Arnold_TOG2018.pdf (section 4.2)
+		AdjustMaterialRoughness(Payload, PathState.PathRoughness);
+	}
+
+	// Processing the random walk will (stochastically) move the shading point elsewhere on the surface of the object
+	if (!ProcessSubsurfaceRandomWalk(Payload, PathState.PathThroughput, PathState.RandSequence, PathState.PixelIndex, bSimplifySSS))
+	{
+		// random walk did not terminate at a valid point
+		return false;
+	}
+
+	FLightLoopCount LightLoopCount = LightGridLookup(Payload.WorldPos);
+
+	// visualize the # of lights
+	if (VisualizeLightGrid)
+	{
+		PathState.Radiance = LightGridVisualize(LightLoopCount, VisualizeLightGrid);
+		PathState.Radiance *= abs(dot(Payload.WorldNormal, PathState.Ray.Direction));
+		PathState.Radiance /= View.PreExposure; // cancel out exposure
+		return false;
+	}
+
+	// Choose a random number for both Light sampling and BxDF sampling
+	float4 RandSample = RandomSequence_GenerateSample4D(PathState.RandSequence);
+
+	// Does this material require NEE? (will be false if MaterialPdf is always +inf)
+	const bool bIsNeeValid = IsNeeValidMaterial(Payload);
+
+	const bool bDoLightLoop = EnableDirectLighting != 0 || Bounce > 0;
+
+	float LightPickingCdfSum = 0;
+
+	// If we are using Light sampling and the material can use it ...
+	if (MISMode != 0 && bIsNeeValid && SceneLightCount > 0 && bDoLightLoop)
+	{
+		// Choose a light and sample it
+		float3 WorldPos = Payload.WorldPos;
+		float3 WorldNormal = Payload.WorldNormal;
+		uint PrimitiveLightingChannelMask = Payload.PrimitiveLightingChannelMask;
+
+		bool IsTransmissiveMaterial = ENABLE_TRANSMISSION && Payload.IsMaterialTransmissive();
+
+		for (uint Index = 0, Num = LightLoopCount.NumLights; Index < Num; ++Index)
+		{
+			uint LightIndex = GetLightId(Index, LightLoopCount);
+			LightPickingCdfSum += EstimateLight(LightIndex, WorldPos, WorldNormal, PrimitiveLightingChannelMask, IsTransmissiveMaterial);
+			LightPickingCdf[Index] = LightPickingCdfSum;
+		}
+
+		if (LightPickingCdfSum > 0)
+		{
+			// init worked
+			int LightId;
+			float LightPickPdf = 0;
+
+			SelectLight(RandSample.x * LightPickingCdfSum, LightLoopCount.NumLights, LightPickingCdf, LightId, LightPickPdf);
+			LightPickPdf /= LightPickingCdfSum;
+
+			LightId = GetLightId(LightId, LightLoopCount);
+			FLightSample LightSample = SampleLight(LightId, RandSample.yz, WorldPos, WorldNormal);
+
+			LightSample.RadianceOverPdf /= LightPickPdf;
+			LightSample.Pdf *= LightPickPdf;
+			if (LightSample.Pdf > 0)
+			{
+				if (CastsShadow(LightId))
+				{
+					// for transmissive materials, bias the position to the other side of the surface if the light is coming from behind
+					const float SignedPositionBias = IsTransmissiveMaterial ? sign(dot(Payload.WorldNormal, LightSample.Direction)) : 1.0;
+
+					RayDesc LightRay;
+					LightRay.Origin = WorldPos;
+					LightRay.TMin = 0;
+					LightRay.Direction = LightSample.Direction;
+					LightRay.TMax = LightSample.Distance;
+					ApplyPositionBias(LightRay, SignedPositionBias * Payload.WorldNormal, MaxNormalBias);
+
+					float AvgRoughness = ApproximateCaustics ? GetAverageRoughness(Payload) : 0.0;
+
+					LightSample.RadianceOverPdf *= TraceTransparentVisibilityRay(LightRay, PathState.PixelIndex, AvgRoughness);
+				}
+
+				// #dxr_todo: Is it cheaper to fire the ray first? Or eval the material first?
+				if (any(LightSample.RadianceOverPdf > 0))
+				{
+					// Evaluate material
+					FMaterialEval MaterialEval = EvalMaterial(PathState.Ray.Direction, LightSample.Direction, Payload, false);
+
+					// Record the contribution
+					float3 LightContrib = PathState.PathThroughput * LightSample.RadianceOverPdf * MaterialEval.Weight * MaterialEval.Pdf;
+					if (MISMode == 2)
+					{
+						LightContrib *= MISWeightRobust(LightSample.Pdf, MaterialEval.Pdf);
+					}
+					AccumulateRadiance(PathState.Radiance, LightContrib, bIsCameraRay);
+				}
+			}
+		}
+	}
+
+	// Sample material
+	FMaterialSample MaterialSample = SampleMaterial(PathState.Ray.Direction, Payload, RandSample, false);
+
+	if (MaterialSample.Pdf < 0 || asuint(MaterialSample.Pdf) > 0x7F800000)
+	{
+		// Pdf became invalid (either negative or NaN)
+		PathState.Radiance = float3(1, 0, 1);
+		return false;
+	}
+
+	if (!(MaterialSample.Pdf > 0))
+	{
+		// No valid direction -- we are done
+		return false;
+	}
+
+	float3 NextPathThroughput = PathState.PathThroughput * MaterialSample.Weight;
+	if (!any(NextPathThroughput > 0))
+	{
+		// no energy left in this path
+		return false;
+	}
+
+	// Russian roulette:
+	//   The probability of keeping the path should be roughly proportional to the weight at the current shade point,
+	//  but just using MaterialWeight would miss out on cases where the path throughput changes color (like in a cornell
+	//  box when bouncing between walls of different colors). So use the ratio of the brightest color channel in the
+	//  previous and next throughput.
+	//   The second tweak is to add a sqrt() around the probability to soften the termination probability (paths will last
+	//  a little longer). This allows paths to go a bit deeper than the naive heuristic while still allowing them to terminate
+	//  early. This makes RR effective from the very first bounce without needing to delay it.
+	float ContinuationProb = sqrt(saturate(max(NextPathThroughput.x, max(NextPathThroughput.y, NextPathThroughput.z)) / max(PathState.PathThroughput.x, max(PathState.PathThroughput.y, PathState.PathThroughput.z))));
+	if (ContinuationProb < 1)
+	{
+		// If there is some chance we should terminate the ray, draw an extra random value
+		float RussianRouletteRand = RandSample.w; // SampleMaterial does not use this value at the moment
+		//RussianRouletteRand = RandomSequence_GenerateSample1D(RandSequence);
+		if (RussianRouletteRand >= ContinuationProb)
+		{
+			// stochastically terminate the path
+			return false;
+		}
+		PathState.PathThroughput = NextPathThroughput / ContinuationProb;
+	}
+	else
+	{
+		PathState.PathThroughput = NextPathThroughput;
+	}
+
+	// Update ray according to material sample
+	PathState.Ray.Origin = Payload.WorldPos;
+	PathState.Ray.Direction = MaterialSample.Direction;
+	PathState.Ray.TMin = 0;
+	PathState.Ray.TMax = POSITIVE_INFINITY;
+	ApplyPositionBias(PathState.Ray, MaterialSample.PositionBiasSign * Payload.WorldNormal, MaxNormalBias);
+
+	if (ApproximateCaustics)
+	{
+		// enlarge roughness based on the chosen lobe roughness
+		PathState.PathRoughness = max(PathState.PathRoughness, MaterialSample.Roughness);
+	}
+
+	// If we are using Material sampling for lights
+	if (MISMode != 1 && bDoLightLoop)
+	{
+		// Check which lights can be seen by the material ray and trace a dedicated shadow ray
+		// While it would be possible to just loop around and use the indirect ray to do this, it would prevent the application
+		// of shadow ray specific logic for transparent shadows or various per light tricks like shadow casting
+		const bool bUseMIS = MISMode == 2 && LightPickingCdfSum > 0;
+		for (uint Index = 0, Num = LightLoopCount.NumMISLights; Index < Num; ++Index)
+		{
+			uint LightId = GetLightId(Index, LightLoopCount);
+			if ((Payload.PrimitiveLightingChannelMask & GetLightingChannelMask(LightId)) == 0)
+			{
+				// light does not affect the current ray
+				continue;
+			}
+
+			FLightHit LightResult = TraceLight(PathState.Ray, LightId);
+
+			if (LightResult.IsMiss())
+			{
+				continue;
+			}
+
+			float3 LightContrib = PathState.PathThroughput * LightResult.Radiance;
+			if (bUseMIS)
+			{
+				float PreviousCdfValue = Index > 0 ? LightPickingCdf[Index - 1] : 0.0;
+				float LightPickPdf = (LightPickingCdf[Index] - PreviousCdfValue) / LightPickingCdfSum;
+
+				LightContrib *= MISWeightRobust(MaterialSample.Pdf, LightResult.Pdf * LightPickPdf);
+			}
+
+			if (any(LightContrib > 0))
+			{
+				if (CastsShadow(LightId))
+				{
+					RayDesc LightRay = PathState.Ray;
+					LightRay.TMax = LightResult.HitT;
+					LightContrib *= TraceTransparentVisibilityRay(LightRay, PathState.PixelIndex, PathState.PathRoughness);
+				}
+				// the light made some contribution, and there was nothing along the shadow ray
+				AccumulateRadiance(PathState.Radiance, LightContrib, bIsCameraRay);
+			}
+		}
+
+	}
+
+	// update ray cone for the next bounce
+	PathState.RayCone = Payload.RayCone;
+
+	// this path should continue bouncing
+	return true;
+}
diff --git a/Engine/Shaders/Shared/PathTracingDefinitions.h b/Engine/Shaders/Shared/PathTracingDefinitions.h
--- a/Engine/Shaders/Shared/PathTracingDefinitions.h
+++ b/Engine/Shaders/Shared/PathTracingDefinitions.h
@@ -7,6 +7,8 @@
 #define PATHTRACER_COHERENT_TILE_SHIFT		3
 #define PATHTRACER_COHERENT_TILE_SIZE		(1 << PATHTRACER_COHERENT_TILE_SHIFT)
 
+#define PATHTRACER_PATHSTATE_NUM_WORDS		20
+
 // Constants for 'SamplerType'
 #define PATHTRACER_SAMPLER_DEFAULT			0
 #define PATHTRACER_SAMPLER_ERROR_DIFFUSION	1
diff --git a/Engine/Source/Runtime/Renderer/Private/PathTracing.cpp b/Engine/Source/Runtime/Renderer/Private/PathTracing.cpp
--- a/Engine/Source/Runtime/Renderer/Private/PathTracing.cpp
+++ b/Engine/Source/Runtime/Renderer/Private/PathTracing.cpp
@@ -26,6 +26,14 @@ TAutoConsoleVariable<int32> CVarPathTracing(
 #include "RenderCore/Public/GenerateMips.h"
 #include <limits>
 
+TAutoConsoleVariable<int32> CVarPathTracingCompaction(
+	TEXT("r.PathTracing.Compaction"),
+	1,
+	TEXT("Enables path compaction to improve GPU occupancy for the path tracer (default: 1 (enabled))"),
+	ECVF_RenderThreadSafe
+);
+
+
 TAutoConsoleVariable<int32> CVarPathTracingMaxBounces(
 	TEXT("r.PathTracing.MaxBounces"),
 	-1,
@@ -200,6 +208,7 @@ TAutoConsoleVariable<int32> CVarPathTracingDenoiser(
 );
 
 BEGIN_SHADER_PARAMETER_STRUCT(FPathTracingData, )
+	SHADER_PARAMETER(float, BlendFactor)
 	SHADER_PARAMETER(uint32, Iteration)
 	SHADER_PARAMETER(uint32, TemporalSeed)
 	SHADER_PARAMETER(uint32, MaxSamples)
@@ -425,18 +434,33 @@ class FPathTracingRG : public FGlobalShader
 	DECLARE_GLOBAL_SHADER(FPathTracingRG)
 	SHADER_USE_ROOT_PARAMETER_STRUCT(FPathTracingRG, FGlobalShader)
 
+	class FEnableCompactionDim : SHADER_PERMUTATION_BOOL("PATH_TRACER_USE_COMPACTION");
+
+	using FPermutationDomain = TShaderPermutationDomain<FEnableCompactionDim>;
+
 	static bool ShouldCompilePermutation(const FGlobalShaderPermutationParameters& Parameters)
 	{
+		FPermutationDomain PermutationVector(Parameters.PermutationId);
+		const bool bUseCompaction = PermutationVector.Get<FEnableCompactionDim>();
+		if (bUseCompaction && !GRHISupportsWaveOperations)
+		{
+			// compaction requires wave operations
+			return false;
+		}
 		return ShouldCompilePathTracingShadersForProject(Parameters.Platform);
 	}
 
 	static void ModifyCompilationEnvironment(const FGlobalShaderPermutationParameters& Parameters, FShaderCompilerEnvironment& OutEnvironment)
 	{
-		//OutEnvironment.CompilerFlags.Add(CFLAG_WarningsAsErrors);
 		OutEnvironment.SetDefine(TEXT("USE_RECT_LIGHT_TEXTURES"), 1);
+		FPermutationDomain PermutationVector(Parameters.PermutationId);
+		if (PermutationVector.Get<FEnableCompactionDim>())
+		{
+			OutEnvironment.CompilerFlags.Add(CFLAG_WaveOperations);
+		}
+		OutEnvironment.CompilerFlags.Add(CFLAG_WarningsAsErrors);
 	}
 
-
 	BEGIN_SHADER_PARAMETER_STRUCT(FParameters, )
 		SHADER_PARAMETER_RDG_TEXTURE_UAV(RWTexture2D<float4>, RadianceTexture)
 		SHADER_PARAMETER_RDG_TEXTURE_UAV(RWTexture2D<float4>, AlbedoTexture)
@@ -464,7 +488,14 @@ class FPathTracingRG : public FGlobalShader
 		// Subsurface data
 		SHADER_PARAMETER_TEXTURE(Texture2D, SSProfilesTexture)
 		// Used by multi-GPU rendering
-		SHADER_PARAMETER(FIntVector, TileOffset)	
+		SHADER_PARAMETER(FIntVector, TileOffset)
+
+		// extra parameters required for path compacting kernel
+		SHADER_PARAMETER(int, Bounce)
+		SHADER_PARAMETER_RDG_BUFFER_UAV(RWBuffer<uint>, PathStateData)
+		SHADER_PARAMETER_RDG_BUFFER_SRV(Buffer<int>, ActivePaths)
+		SHADER_PARAMETER_RDG_BUFFER_UAV(RWBuffer<int>, NextActivePaths)
+		SHADER_PARAMETER_RDG_BUFFER_UAV(RWBuffer<int>, NumPathStates)
 	END_SHADER_PARAMETER_STRUCT()
 };
 IMPLEMENT_GLOBAL_SHADER(FPathTracingRG, "/Engine/Private/PathTracing/PathTracing.usf", "PathTracingMainRG", SF_RayGen);
@@ -1205,8 +1236,13 @@ void FDeferredShadingSceneRenderer::PreparePathTracing(const FSceneViewFamily& V
 		&& ShouldCompilePathTracingShadersForProject(ViewFamily.GetShaderPlatform()))
 	{
 		// Declare all RayGen shaders that require material closest hit shaders to be bound
-		auto RayGenShader = GetGlobalShaderMap(ViewFamily.GetShaderPlatform())->GetShader<FPathTracingRG>();
-		OutRayGenShaders.Add(RayGenShader.GetRayTracingShader());
+		for (int EnableCompaction = 0; EnableCompaction < 2; EnableCompaction++)
+		{
+			FPathTracingRG::FPermutationDomain PermutationVector;
+			PermutationVector.Set<FPathTracingRG::FEnableCompactionDim>(EnableCompaction != 0);
+			auto RayGenShader = GetGlobalShaderMap(ViewFamily.GetShaderPlatform())->GetShader<FPathTracingRG>(PermutationVector);
+			OutRayGenShaders.Add(RayGenShader.GetRayTracingShader());
+		}
 	}
 }
 
@@ -1329,6 +1365,7 @@ void FDeferredShadingSceneRenderer::RenderPathTracing(
 		Config.PathTracingData.TemporalSeed = View.ViewState->PathTracingFrameIndex;
 	}
 	Config.PathTracingData.Iteration = View.ViewState->PathTracingSampleIndex;
+	Config.PathTracingData.BlendFactor = 1.0f / (Config.PathTracingData.Iteration + 1);
 
 	// Prepare radiance buffer (will be shared with display pass)
 	FRDGTexture* RadianceTexture = nullptr;
@@ -1357,55 +1394,103 @@ void FDeferredShadingSceneRenderer::RenderPathTracing(
 
 	if (bNeedsMoreRays)
 	{
-		FPathTracingRG::FParameters* PassParameters = GraphBuilder.AllocParameters<FPathTracingRG::FParameters>();
-		PassParameters->TLAS = View.GetRayTracingSceneViewChecked();
-		PassParameters->ViewUniformBuffer = View.ViewUniformBuffer;
-		PassParameters->PathTracingData = Config.PathTracingData;
-		// upload sky/lights data
-		SetLightParameters(GraphBuilder, PassParameters, Scene, View, Config.UseMISCompensation);
-		if (Config.PathTracingData.EnableDirectLighting == 0)
+		// Round up to coherent path tracing tile size to simplify pixel shuffling
+		// TODO: be careful not to write extra pixels past the boundary when using multi-gpu
+		const int32 DispatchSizeX = FMath::DivideAndRoundUp(View.ViewRect.Size().X, PATHTRACER_COHERENT_TILE_SIZE) * PATHTRACER_COHERENT_TILE_SIZE;
+		const int32 DispatchSizeY = FMath::DivideAndRoundUp(View.ViewRect.Size().Y, PATHTRACER_COHERENT_TILE_SIZE) * PATHTRACER_COHERENT_TILE_SIZE;
+
+		// should we use path compaction?
+		const bool bUseCompaction = CVarPathTracingCompaction.GetValueOnRenderThread() != 0 && GRHISupportsWaveOperations;
+		FRDGBuffer* ActivePaths[2] = {};
+		FRDGBuffer* NumActivePaths = nullptr;
+		FRDGBuffer* PathStateData = nullptr;
+		if (bUseCompaction)
 		{
-			PassParameters->SceneVisibleLightCount = 0;
+			ActivePaths[0] = GraphBuilder.CreateBuffer(FRDGBufferDesc::CreateBufferDesc(sizeof(int32), DispatchSizeX * DispatchSizeY), TEXT("PathTracer.ActivePaths0"));
+			ActivePaths[1] = GraphBuilder.CreateBuffer(FRDGBufferDesc::CreateBufferDesc(sizeof(int32), DispatchSizeX * DispatchSizeY), TEXT("PathTracer.ActivePaths1"));
+			NumActivePaths = GraphBuilder.CreateBuffer(FRDGBufferDesc::CreateBufferDesc(sizeof(int32), 1), TEXT("PathTracer.NumActivePaths"));
+			PathStateData = GraphBuilder.CreateBuffer(FRDGBufferDesc::CreateBufferDesc(sizeof(int32) * PATHTRACER_PATHSTATE_NUM_WORDS, DispatchSizeX * DispatchSizeY), TEXT("PathTracer.PathStateData"));
+			AddClearUAVPass(GraphBuilder, GraphBuilder.CreateUAV(ActivePaths[0], PF_R32_UINT), 0);
 		}
+		FPathTracingRG::FPermutationDomain PermutationVector;
+		PermutationVector.Set<FPathTracingRG::FEnableCompactionDim>(bUseCompaction);
+		TShaderMapRef<FPathTracingRG> RayGenShader(View.ShaderMap, PermutationVector);
+		FPathTracingRG::FParameters* PreviousPassParameters = nullptr;
+		// When using path compaction, we need to run the path tracer once per bounce
+		// otherwise, the path tracer is the one doing the bounces
+		for (int Bounce = 0, MaxBounces = bUseCompaction ? Config.PathTracingData.MaxBounces : 0; Bounce <= MaxBounces; Bounce++)
+		{
+			FPathTracingRG::FParameters* PassParameters = GraphBuilder.AllocParameters<FPathTracingRG::FParameters>();
+			PassParameters->TLAS = View.GetRayTracingSceneViewChecked();
+			PassParameters->ViewUniformBuffer = View.ViewUniformBuffer;
+			PassParameters->PathTracingData = Config.PathTracingData;
+			if (PreviousPassParameters == nullptr)
+			{
+				// upload sky/lights data
+				SetLightParameters(GraphBuilder, PassParameters, Scene, View, Config.UseMISCompensation);
+				PreviousPassParameters = PassParameters;
+			}
+			else
+			{
+				// re-use from last iteration
+				PassParameters->IESTexture = PreviousPassParameters->IESTexture;
+				PassParameters->IESTextureSampler = PreviousPassParameters->IESTextureSampler;
+				PassParameters->LightGridParameters = PreviousPassParameters->LightGridParameters;
+				PassParameters->RectLightTexture = PreviousPassParameters->RectLightTexture;
+				PassParameters->RectLightSampler = PreviousPassParameters->RectLightSampler;
+				PassParameters->SceneLightCount = PreviousPassParameters->SceneLightCount;
+				PassParameters->SceneVisibleLightCount = PreviousPassParameters->SceneVisibleLightCount;
+				PassParameters->SceneLights = PreviousPassParameters->SceneLights;
+				PassParameters->SkylightParameters = PreviousPassParameters->SkylightParameters;
+			}
+			if (Config.PathTracingData.EnableDirectLighting == 0)
+			{
+				PassParameters->SceneVisibleLightCount = 0;
+			}
 
-		PassParameters->IESTextureSampler = TStaticSamplerState<SF_Bilinear, AM_Clamp, AM_Clamp, AM_Clamp>::GetRHI();
-		PassParameters->RadianceTexture = GraphBuilder.CreateUAV(RadianceTexture);
-		PassParameters->AlbedoTexture   = GraphBuilder.CreateUAV(AlbedoTexture);
-		PassParameters->NormalTexture   = GraphBuilder.CreateUAV(NormalTexture);
+			PassParameters->IESTextureSampler = TStaticSamplerState<SF_Bilinear, AM_Clamp, AM_Clamp, AM_Clamp>::GetRHI();
+			PassParameters->RadianceTexture = GraphBuilder.CreateUAV(RadianceTexture);
+			PassParameters->AlbedoTexture   = GraphBuilder.CreateUAV(AlbedoTexture);
+			PassParameters->NormalTexture   = GraphBuilder.CreateUAV(NormalTexture);
 
-		PassParameters->SSProfilesTexture = GetSubsufaceProfileTexture_RT(GraphBuilder.RHICmdList)->GetShaderResourceRHI();
+			PassParameters->SSProfilesTexture = GetSubsufaceProfileTexture_RT(GraphBuilder.RHICmdList)->GetShaderResourceRHI();
 
-		// TODO: in multi-gpu case, split image into tiles
-		PassParameters->TileOffset.X = 0;
-		PassParameters->TileOffset.Y = 0;
+			// TODO: in multi-gpu case, split image into tiles
+			PassParameters->TileOffset.X = 0;
+			PassParameters->TileOffset.Y = 0;
 
-		TShaderMapRef<FPathTracingRG> RayGenShader(View.ShaderMap);
-		ClearUnusedGraphResources(RayGenShader, PassParameters);
-		GraphBuilder.AddPass(
-			RDG_EVENT_NAME("Path Tracer Compute (%d x %d) Sample=%d/%d NumLights=%d", View.ViewRect.Size().X, View.ViewRect.Size().Y, View.ViewState->PathTracingSampleIndex, MaxSPP, PassParameters->SceneLightCount),
-			PassParameters,
-			ERDGPassFlags::Compute,
-			[PassParameters, RayGenShader, &View](FRHICommandListImmediate& RHICmdList)
-		{
-			FRHIRayTracingScene* RayTracingSceneRHI = View.GetRayTracingSceneChecked();
+			if (bUseCompaction)
+			{
+				PassParameters->Bounce = Bounce;
+				PassParameters->ActivePaths = GraphBuilder.CreateSRV(ActivePaths[Bounce & 1], PF_R32_UINT);
+				PassParameters->NextActivePaths = GraphBuilder.CreateUAV(ActivePaths[(Bounce & 1) ^ 1], PF_R32_UINT);
+				PassParameters->PathStateData = GraphBuilder.CreateUAV(PathStateData, PF_R32_UINT);
+				PassParameters->NumPathStates = GraphBuilder.CreateUAV(NumActivePaths, PF_R32_UINT);
+				AddClearUAVPass(GraphBuilder, PassParameters->NextActivePaths, ~0u); // make sure everything is initialized to -1 since paths that go inactive don't write anything
+				AddClearUAVPass(GraphBuilder, PassParameters->NumPathStates, 0); // reset number of active paths to 0
+			}
+			ClearUnusedGraphResources(RayGenShader, PassParameters);
+			GraphBuilder.AddPass(
+				bUseCompaction
+					? RDG_EVENT_NAME("Path Tracer Compute (%d x %d) Sample=%d/%d NumLights=%d Bounce=%d", View.ViewRect.Size().X, View.ViewRect.Size().Y, View.ViewState->PathTracingSampleIndex, MaxSPP, PassParameters->SceneLightCount, Bounce)
+					: RDG_EVENT_NAME("Path Tracer Compute (%d x %d) Sample=%d/%d NumLights=%d", View.ViewRect.Size().X, View.ViewRect.Size().Y, View.ViewState->PathTracingSampleIndex, MaxSPP, PassParameters->SceneLightCount),
+				PassParameters,
+				ERDGPassFlags::Compute,
+				[PassParameters, RayGenShader, DispatchSizeX, DispatchSizeY, &View](FRHICommandListImmediate& RHICmdList)
+			{
+				FRHIRayTracingScene* RayTracingSceneRHI = View.GetRayTracingSceneChecked();
 			
-			// Round up to coherent path tracing tile size to simplify pixel shuffling
-			// TODO: be careful not to write extra pixels past the boundary when using multi-gpu
-			const int32 TS = PATHTRACER_COHERENT_TILE_SIZE;
-			int32 DispatchSizeX = FMath::DivideAndRoundUp(View.ViewRect.Size().X, TS) * TS;
-			int32 DispatchSizeY = FMath::DivideAndRoundUp(View.ViewRect.Size().Y, TS) * TS;
-
-			FRayTracingShaderBindingsWriter GlobalResources;
-			SetShaderParameters(GlobalResources, RayGenShader, *PassParameters);
-
-			RHICmdList.RayTraceDispatch(
-				View.RayTracingMaterialPipeline,
-				RayGenShader.GetRayTracingShader(),
-				RayTracingSceneRHI, GlobalResources,
-				DispatchSizeX, DispatchSizeY
-			);
-		});
-
+				FRayTracingShaderBindingsWriter GlobalResources;
+				SetShaderParameters(GlobalResources, RayGenShader, *PassParameters);
+
+				RHICmdList.RayTraceDispatch(
+					View.RayTracingMaterialPipeline,
+					RayGenShader.GetRayTracingShader(),
+					RayTracingSceneRHI, GlobalResources,
+					DispatchSizeX, DispatchSizeY
+				);
+			});
+		}
 		// After we are done, make sure we remember our texture for next time so that we can accumulate samples across frames
 		GraphBuilder.QueueTextureExtraction(RadianceTexture, &View.ViewState->PathTracingRadianceRT);
 		GraphBuilder.QueueTextureExtraction(AlbedoTexture  , &View.ViewState->PathTracingAlbedoRT  );
