Fix invalid pixel coord sampling.

diff --git a/Engine/Shaders/Private/HairStrands/HairStrandsMaterialPS.usf b/Engine/Shaders/Private/HairStrands/HairStrandsMaterialPS.usf
--- a/Engine/Shaders/Private/HairStrands/HairStrandsMaterialPS.usf
+++ b/Engine/Shaders/Private/HairStrands/HairStrandsMaterialPS.usf
@@ -21,9 +21,7 @@ void Main(
 	if (SampleIndex < TotalNodeCount)
 	{
 		const FNodeVis InNode = MaterialPassParameters.NodeVis[SampleIndex];
-		
-		const uint PackedSamplePixelCoord = MaterialPassParameters.NodeCoord[SampleIndex];
-		const float2 SamplePixelCoord = float2(PackedSamplePixelCoord & 0xFFFF, (PackedSamplePixelCoord>>16) & 0xFFFF) + float2(0.5f, 0.5f);
+		const float2 SamplePixelCoord = float2(MaterialPassParameters.NodeCoord[SampleIndex]) + float2(0.5f, 0.5f);
 
 		float SampleDepth = InNode.Depth;
 		uint SampleCoverage8bit = InNode.Coverage8bit;
@@ -32,7 +30,11 @@ void Main(
 
 		if (SampleMaterialId == MaterialPass_MaterialId)
 		{
-			FEvaluateOutput Out = Evaluate(SamplePixelCoord, SampleDepth, SamplePrimitiveId, SampleCoverage8bit, InPrimitiveId, MaterialPassParameters.bUpdateSampleCoverage > 0);
+			FEvaluateOutput Out = (FEvaluateOutput)0;
+			if (SampleDepth > 0)
+			{
+				Out = Evaluate(SamplePixelCoord, SampleDepth, SamplePrimitiveId, SampleCoverage8bit, InPrimitiveId, MaterialPassParameters.bUpdateSampleCoverage > 0);
+			}
 
 			// Note: Use the syntax _ over the syntax . as on certain platform the UAV buffers are 
 			// stored within a static const struct preventing doing a write operation
