Increase limit of hair instance with different material to 256.

diff --git a/Engine/Shaders/Private/HairStrands/HairStrandsVisibilityCommon.ush b/Engine/Shaders/Private/HairStrands/HairStrandsVisibilityCommon.ush
--- a/Engine/Shaders/Private/HairStrands/HairStrandsVisibilityCommon.ush
+++ b/Engine/Shaders/Private/HairStrands/HairStrandsVisibilityCommon.ush
@@ -413,14 +413,14 @@ return 0xFFFFFFFF;
 
 uint PackHairVisSample(uint PrimitiveId, uint MaterialId)
 {
-	return (PrimitiveId & 0x0FFFFFFF) | ((MaterialId & 0xF) << 28);
+	return (PrimitiveId & 0x00FFFFFF) | ((MaterialId & 0xFF) << 24);
 }
 
 FHairVisSample UnpackHairVisSample(uint RawVis)
 {
 	FHairVisSample Out;
-	Out.PrimitiveId = RawVis & 0x0FFFFFFF;
-	Out.MaterialId = (RawVis >> 28) & 0xF;
+	Out.PrimitiveId = RawVis & 0x00FFFFFF;
+	Out.MaterialId = (RawVis >> 24) & 0xFF;
 	return Out;
 }
  
