Make hair debug sample & voxel info more readable when there are several macrogroups visible on screen.

diff --git a/Engine/Shaders/Private/HairStrands/HairStrandsDebugPrint.usf b/Engine/Shaders/Private/HairStrands/HairStrandsDebugPrint.usf
--- a/Engine/Shaders/Private/HairStrands/HairStrandsDebugPrint.usf
+++ b/Engine/Shaders/Private/HairStrands/HairStrandsDebugPrint.usf
@@ -28,6 +28,11 @@ Buffer<int> HairMacroGroupAABBBuffer;
 #define MATERIAL_DATA 1
 #define DETAIL_DATA 1
 
+void AddNewLine(float2 OriginalPos, inout float2 Pos)
+{
+	Pos = ShaderPrintNewline(Pos); Pos.x = OriginalPos.x;
+}
+
 [numthreads(1, 1, 1)]
 void MainCS(uint GroupIndex : SV_GroupIndex, uint3 DispatchThreadId : SV_DispatchThreadID)
 {	
@@ -50,46 +55,17 @@ void MainCS(uint GroupIndex : SV_GroupIndex, uint3 DispatchThreadId : SV_Dispatc
 		// Pixel coord
 		float2 OriginalPos = float2(50, 50) / float2(MaxResolution);
 		float2 Pos = OriginalPos;
-		Pos = ShaderPrintSymbol(Pos, _T_);
-		Pos = ShaderPrintSymbol(Pos, _O_);
-		Pos = ShaderPrintSymbol(Pos, _T_);
-		Pos = ShaderPrintSymbol(Pos, _A_);
-		Pos = ShaderPrintSymbol(Pos, _L_);
-		Pos = ShaderPrintSymbol(Pos, _SPC_);
-		Pos = ShaderPrintSymbol(Pos, _S_);
-		Pos = ShaderPrintSymbol(Pos, _A_);
-		Pos = ShaderPrintSymbol(Pos, _M_);
-		Pos = ShaderPrintSymbol(Pos, _P_);
-		Pos = ShaderPrintSymbol(Pos, _L_);
-		Pos = ShaderPrintSymbol(Pos, _E_);
-		Pos = ShaderPrintSymbol(Pos, _SPC_);
-		Pos = ShaderPrintSymbol(Pos, _C_);
-		Pos = ShaderPrintSymbol(Pos, _O_);
-		Pos = ShaderPrintSymbol(Pos, _U_);
-		Pos = ShaderPrintSymbol(Pos, _N_);
-		Pos = ShaderPrintSymbol(Pos, _T_);
-		Pos = ShaderPrintSymbol(Pos, _SPC_);
-		Pos = ShaderPrint(Pos, TotalSampleCount);
-
-		Pos = ShaderPrintNewline(Pos);
-		Pos.x = OriginalPos.x;
-
-		Pos = ShaderPrintSymbol(Pos, _E_);
-		Pos = ShaderPrintSymbol(Pos, _Q_);
-		Pos = ShaderPrintSymbol(Pos, _U_);
-		Pos = ShaderPrintSymbol(Pos, _I_);
-		Pos = ShaderPrintSymbol(Pos, _V_);
-		Pos = ShaderPrintSymbol(Pos, _A_);
-		Pos = ShaderPrintSymbol(Pos, _L_);
-		Pos = ShaderPrintSymbol(Pos, _E_);
-		Pos = ShaderPrintSymbol(Pos, _N_);
-		Pos = ShaderPrintSymbol(Pos, _T_);
-		Pos = ShaderPrintSymbol(Pos, _SPC_);
-		Pos = ShaderPrintSymbol(Pos, _R_);
-		Pos = ShaderPrintSymbol(Pos, _E_);
-		Pos = ShaderPrintSymbol(Pos, _S_);
-		Pos = ShaderPrintSymbol(Pos, _SPC_);
-		Pos = ShaderPrint(Pos, EquivalentResolution);
+		Pos = ShaderPrintText(Pos, TEXT("Total Sample Count "), FontCyan);
+		Pos = ShaderPrint(Pos, TotalSampleCount, FontCyan);
+		AddNewLine(OriginalPos, Pos);
+
+		Pos = ShaderPrintText(Pos, TEXT("Equivalent Res     "), FontCyan);
+		Pos = ShaderPrint(Pos, EquivalentResolution, FontCyan);
+		AddNewLine(OriginalPos, Pos);
+
+		Pos = ShaderPrintText(Pos, TEXT("MacroGroup Count   "), FontYellow);
+		Pos = ShaderPrint(Pos, HairMacroGroupCount, FontYellow);
+		AddNewLine(OriginalPos, Pos);
 	}
 
 	const float2 PixelPos = PixelCoord;
@@ -106,7 +82,8 @@ void MainCS(uint GroupIndex : SV_GroupIndex, uint3 DispatchThreadId : SV_Dispatc
 	const float ViewHairCount = HairCountTexture.Load(uint3(PixelPos, 0));
 	const float ViewHairCountUint = HairCountUintTexture.Load(uint3(PixelPos, 0)) /  1000.f;
 
-	const float2 OriginalPos = PixelPos / float2(MaxResolution);
+//	const float2 OriginalPos = PixelPos / float2(MaxResolution);
+	const float2 OriginalPos = float2(50,500) / float2(MaxResolution);
 	float2 Pos = OriginalPos;
 
 	uint Total8BitCoverage = 0;
@@ -120,200 +97,95 @@ void MainCS(uint GroupIndex : SV_GroupIndex, uint3 DispatchThreadId : SV_Dispatc
 	}
 
 	// Pixel coord
-	Pos = ShaderPrintSymbol(Pos, _P_);
-	Pos = ShaderPrintSymbol(Pos, _I_);
-	Pos = ShaderPrintSymbol(Pos, _X_);
-	Pos = ShaderPrintSymbol(Pos, _E_);
-	Pos = ShaderPrintSymbol(Pos, _L_);
-	Pos = ShaderPrintSymbol(Pos, _SPC_);
-	Pos = ShaderPrint(Pos, PixelCoord);
-
-	Pos = ShaderPrintNewline(Pos);
-	Pos.x = OriginalPos.x;
-
-#if MATERIAL_DATA
-	// Dispatch thread id
-	Pos = ShaderPrintSymbol(Pos, _T_);
-	Pos = ShaderPrintSymbol(Pos, _H_);
-	Pos = ShaderPrintSymbol(Pos, _R_);
-	Pos = ShaderPrintSymbol(Pos, _E_);
-	Pos = ShaderPrintSymbol(Pos, _A_);
-	Pos = ShaderPrintSymbol(Pos, _D_);
-	Pos = ShaderPrintSymbol(Pos, _SPC_);
-	Pos = ShaderPrintSymbol(Pos, _I_);
-	Pos = ShaderPrintSymbol(Pos, _D_);
-	Pos = ShaderPrintSymbol(Pos, _SPC_);
-	Pos = ShaderPrint(Pos, uint4(PixelCoord.x/GroupSize.x, PixelCoord.x % GroupSize.x, PixelCoord.y / GroupSize.y, PixelCoord.y % GroupSize.y));
-
-	Pos = ShaderPrintNewline(Pos);
-	Pos.x = OriginalPos.x;
-#endif
+	Pos = ShaderPrintText(Pos, TEXT("Pixel           "), FontTurquoise);
+	Pos = ShaderPrint(Pos, PixelCoord, FontTurquoise);
+	AddNewLine(OriginalPos, Pos);
 
 	// Max sample
-	Pos = ShaderPrintSymbol(Pos, _M_);
-	Pos = ShaderPrintSymbol(Pos, _A_);
-	Pos = ShaderPrintSymbol(Pos, _X_);
-	Pos = ShaderPrintSymbol(Pos, _SPC_);
-	Pos = ShaderPrintSymbol(Pos, _S_);
-	Pos = ShaderPrintSymbol(Pos, _A_);
-	Pos = ShaderPrintSymbol(Pos, _M_);
-	Pos = ShaderPrintSymbol(Pos, _P_);
-	Pos = ShaderPrintSymbol(Pos, _L_);
-	Pos = ShaderPrintSymbol(Pos, _E_);
-	Pos = ShaderPrintSymbol(Pos, _SPC_);
+	Pos = ShaderPrintText(Pos, TEXT("Max Sample      "));
 	Pos = ShaderPrint(Pos, MaxSampleCount);
-
-	Pos = ShaderPrintNewline(Pos);
-	Pos.x = OriginalPos.x;
+	AddNewLine(OriginalPos, Pos);
 
 	// Unique sample
-	Pos = ShaderPrintSymbol(Pos, _S_);
-	Pos = ShaderPrintSymbol(Pos, _A_);
-	Pos = ShaderPrintSymbol(Pos, _M_);
-	Pos = ShaderPrintSymbol(Pos, _P_);
-	Pos = ShaderPrintSymbol(Pos, _L_);
-	Pos = ShaderPrintSymbol(Pos, _E_);
-	Pos = ShaderPrintSymbol(Pos, _SPC_);
-	Pos = ShaderPrint(Pos, CatData.SampleCount);
-
-	Pos = ShaderPrintNewline(Pos);
-	Pos.x = OriginalPos.x;
+	Pos = ShaderPrintText(Pos, TEXT("Sample          "), FontOrange);
+	Pos = ShaderPrint(Pos, CatData.SampleCount, FontOrange);
+	AddNewLine(OriginalPos, Pos);
 
 	// Accurate coverage (if coverage pass is activated)
-	Pos = ShaderPrintSymbol(Pos, _H_);
-	Pos = ShaderPrintSymbol(Pos, _A_);
-	Pos = ShaderPrintSymbol(Pos, _I_);
-	Pos = ShaderPrintSymbol(Pos, _R_);
-	Pos = ShaderPrintSymbol(Pos, _SPC_);
-	Pos = ShaderPrintSymbol(Pos, _C_);
-	Pos = ShaderPrintSymbol(Pos, _O_);
-	Pos = ShaderPrintSymbol(Pos, _U_);
-	Pos = ShaderPrintSymbol(Pos, _N_);
-	Pos = ShaderPrintSymbol(Pos, _T_);
-	Pos = ShaderPrintSymbol(Pos, _SPC_);
+	Pos = ShaderPrintText(Pos, TEXT("Hair Count      "));
 	Pos = ShaderPrint(Pos, ViewHairCount);
-
-	Pos = ShaderPrintNewline(Pos);
-	Pos.x = OriginalPos.x;
+	AddNewLine(OriginalPos, Pos);
 
 	// Accurate coverage (if coverage pass is activated) uint
-	Pos = ShaderPrintSymbol(Pos, _H_);
-	Pos = ShaderPrintSymbol(Pos, _A_);
-	Pos = ShaderPrintSymbol(Pos, _I_);
-	Pos = ShaderPrintSymbol(Pos, _R_);
-	Pos = ShaderPrintSymbol(Pos, _SPC_);
-	Pos = ShaderPrintSymbol(Pos, _C_);
-	Pos = ShaderPrintSymbol(Pos, _O_);
-	Pos = ShaderPrintSymbol(Pos, _U_);
-	Pos = ShaderPrintSymbol(Pos, _N_);
-	Pos = ShaderPrintSymbol(Pos, _T_);
-	Pos = ShaderPrintSymbol(Pos, _SPC_);
-	Pos = ShaderPrintSymbol(Pos, _U_);
-	Pos = ShaderPrintSymbol(Pos, _I_);
-	Pos = ShaderPrintSymbol(Pos, _N_);
-	Pos = ShaderPrintSymbol(Pos, _T_);
-	Pos = ShaderPrintSymbol(Pos, _SPC_);
+	Pos = ShaderPrintText(Pos, TEXT("Hair Count Uint "));
 	Pos = ShaderPrint(Pos, ViewHairCountUint);
-
-	Pos = ShaderPrintNewline(Pos);
-	Pos.x = OriginalPos.x;
+	AddNewLine(OriginalPos, Pos);
 	
 	// Accurate coverage (if coverage pass is activated)
-	Pos = ShaderPrintSymbol(Pos, _C_);
-	Pos = ShaderPrintSymbol(Pos, _O_);
-	Pos = ShaderPrintSymbol(Pos, _V_);
-	Pos = ShaderPrintSymbol(Pos, _E_);
-	Pos = ShaderPrintSymbol(Pos, _R_);
-	Pos = ShaderPrintSymbol(Pos, _A_);
-	Pos = ShaderPrintSymbol(Pos, _G_);
-	Pos = ShaderPrintSymbol(Pos, _E_);
-	Pos = ShaderPrintSymbol(Pos, _SPC_);
-	Pos = ShaderPrint(Pos, CatData.PixelCoverage);
-
-	Pos = ShaderPrintNewline(Pos);
-	Pos.x = OriginalPos.x;
+	Pos = ShaderPrintText(Pos, TEXT("Coverage        "), FontOrange);
+	Pos = ShaderPrint(Pos, CatData.PixelCoverage, FontOrange);
+	AddNewLine(OriginalPos, Pos);
 
 	// 8 bit Coverage (cateogory texture)
-	Pos = ShaderPrintSymbol(Pos, _C_);
-	Pos = ShaderPrintSymbol(Pos, _O_);
-	Pos = ShaderPrintSymbol(Pos, _V_);
-	Pos = ShaderPrintSymbol(Pos, _E_);
-	Pos = ShaderPrintSymbol(Pos, _R_);
-	Pos = ShaderPrintSymbol(Pos, _A_);
-	Pos = ShaderPrintSymbol(Pos, _G_);
-	Pos = ShaderPrintSymbol(Pos, _E_);
-	Pos = ShaderPrintSymbol(Pos, _SPC_);
-	Pos = ShaderPrintSymbol(Pos, _8_); 
-	Pos = ShaderPrintSymbol(Pos, _B_);
-	Pos = ShaderPrintSymbol(Pos, _I_);
-	Pos = ShaderPrintSymbol(Pos, _T_);
-	Pos = ShaderPrintSymbol(Pos, _S_);
-	Pos = ShaderPrintSymbol(Pos, _SPC_);
+	Pos = ShaderPrintText(Pos, TEXT("Coverage 8bits  "));
 	Pos = ShaderPrint(Pos, CatData.TotalCoverage8bit);
 	Pos = ShaderPrint(Pos, CatData.TotalCoverage8bit / 255.f);
-	
-	Pos = ShaderPrintNewline(Pos);
-	Pos.x = OriginalPos.x;
+	AddNewLine(OriginalPos, Pos);
 
 	// 8 bit Coverage (summed)
-	Pos = ShaderPrintSymbol(Pos, _S_);
-	Pos = ShaderPrintSymbol(Pos, _U_);
-	Pos = ShaderPrintSymbol(Pos, _M_);
-	Pos = ShaderPrintSymbol(Pos, _M_);
-	Pos = ShaderPrintSymbol(Pos, _E_);
-	Pos = ShaderPrintSymbol(Pos, _D_);
-	Pos = ShaderPrintSymbol(Pos, _SPC_);
-	Pos = ShaderPrintSymbol(Pos, _C_);
-	Pos = ShaderPrintSymbol(Pos, _O_);
-	Pos = ShaderPrintSymbol(Pos, _V_);
-	Pos = ShaderPrintSymbol(Pos, _E_);
-	Pos = ShaderPrintSymbol(Pos, _R_);
-	Pos = ShaderPrintSymbol(Pos, _A_);
-	Pos = ShaderPrintSymbol(Pos, _G_);
-	Pos = ShaderPrintSymbol(Pos, _E_);
-	Pos = ShaderPrintSymbol(Pos, _SPC_);
-	Pos = ShaderPrintSymbol(Pos, _8_); 
-	Pos = ShaderPrintSymbol(Pos, _B_);
-	Pos = ShaderPrintSymbol(Pos, _I_);
-	Pos = ShaderPrintSymbol(Pos, _T_);
-	Pos = ShaderPrintSymbol(Pos, _S_);
-	Pos = ShaderPrintSymbol(Pos, _SPC_);
+	Pos = ShaderPrintText(Pos, TEXT("Sum cov. 8bits  "));
 	Pos = ShaderPrint(Pos, Total8BitCoverage);
-
-	Pos = ShaderPrintNewline(Pos);
-	Pos.x = OriginalPos.x;
+	AddNewLine(OriginalPos, Pos);
 
 	// Resolve type
 	const uint Stencil = DepthStencilTexture.Load(uint3(PixelPos,0)).y;
 	const bool bIsFastResolve = (Stencil & FastResolveMask) != 0;
-	Pos = ShaderPrintSymbol(Pos, _R_);
-	Pos = ShaderPrintSymbol(Pos, _E_);
-	Pos = ShaderPrintSymbol(Pos, _S_);
-	Pos = ShaderPrintSymbol(Pos, _O_);
-	Pos = ShaderPrintSymbol(Pos, _L_);
-	Pos = ShaderPrintSymbol(Pos, _V_);
-	Pos = ShaderPrintSymbol(Pos, _E_);
-	Pos = ShaderPrintSymbol(Pos, _SPC_);
+	Pos = ShaderPrintText(Pos, TEXT("Resolve         "));
+
 	if (bIsFastResolve)
 	{
-		Pos = ShaderPrintSymbol(Pos, _F_);
-		Pos = ShaderPrintSymbol(Pos, _A_);
-		Pos = ShaderPrintSymbol(Pos, _S_);
-		Pos = ShaderPrintSymbol(Pos, _T_);
+		Pos = ShaderPrintText(Pos, TEXT("Fast"), FontEmerald);
 	}
 	else
 	{
-		Pos = ShaderPrintSymbol(Pos, _R_);
-		Pos = ShaderPrintSymbol(Pos, _E_);
-		Pos = ShaderPrintSymbol(Pos, _G_);
-		Pos = ShaderPrintSymbol(Pos, _U_);
-		Pos = ShaderPrintSymbol(Pos, _L_);
-		Pos = ShaderPrintSymbol(Pos, _A_);
-		Pos = ShaderPrintSymbol(Pos, _R_);
+		Pos = ShaderPrintText(Pos, TEXT("Regular"), FontEmerald);
 	}
 
-	Pos = ShaderPrintNewline(Pos);
-	Pos.x = OriginalPos.x;
+	AddNewLine(OriginalPos, Pos);
+	AddNewLine(OriginalPos, Pos);
+
+	const float2 ColumnPos = Pos;
+
+	// Sample Description
+	{	
+		Pos = ShaderPrintText(Pos, TEXT("Sample ID       "), FontWhite);	AddNewLine(OriginalPos, Pos);
+	#if MATERIAL_DATA
+		Pos = ShaderPrintText(Pos, TEXT(".Base Color R   "), FontRed);		AddNewLine(OriginalPos, Pos); 
+		Pos = ShaderPrintText(Pos, TEXT(".Base Color G   "), FontGreen);	AddNewLine(OriginalPos, Pos);
+		Pos = ShaderPrintText(Pos, TEXT(".Base Color B   "), FontBlue);		AddNewLine(OriginalPos, Pos);
+		AddNewLine(OriginalPos, Pos);
+
+		Pos = ShaderPrintText(Pos, TEXT(".Tangent X      "), FontEmerald);	AddNewLine(OriginalPos, Pos); 
+		Pos = ShaderPrintText(Pos, TEXT(".Tangent Y      "), FontEmerald);	AddNewLine(OriginalPos, Pos);
+		Pos = ShaderPrintText(Pos, TEXT(".Tangent Z      "), FontEmerald);	AddNewLine(OriginalPos, Pos);
+		AddNewLine(OriginalPos, Pos);
+
+		Pos = ShaderPrintText(Pos, TEXT(".Roughness      "), FontOrange);	AddNewLine(OriginalPos, Pos);
+		Pos = ShaderPrintText(Pos, TEXT(".Specular       "), FontEmerald);	AddNewLine(OriginalPos, Pos);
+		Pos = ShaderPrintText(Pos, TEXT(".Backlit        "), FontOrange);	AddNewLine(OriginalPos, Pos);
+		AddNewLine(OriginalPos, Pos);
+	#endif 													  
+	#if DETAIL_DATA											  
+		Pos = ShaderPrintText(Pos, TEXT(".Light Channels "), FontEmerald);	AddNewLine(OriginalPos, Pos);
+		Pos = ShaderPrintText(Pos, TEXT(".Scatter Scene  "), FontOrange);	AddNewLine(OriginalPos, Pos);
+		Pos = ShaderPrintText(Pos, TEXT(".Depth          "), FontEmerald);	AddNewLine(OriginalPos, Pos);
+		AddNewLine(OriginalPos, Pos);
+
+		Pos = ShaderPrintText(Pos, TEXT(".Macro Group ID "), FontYellow);	AddNewLine(OriginalPos, Pos);
+		Pos = ShaderPrintText(Pos, TEXT(".Primitive		 "), FontEmerald);	AddNewLine(OriginalPos, Pos);
+	#endif													  
+		Pos = ShaderPrintText(Pos, TEXT(".Coverage       "), FontOrange);	AddNewLine(OriginalPos, Pos);
+	}
 
 	// Per sample description
 	for (uint SampleIt = 0; SampleIt < NodeDesc.Count; SampleIt++)
@@ -322,227 +194,40 @@ void MainCS(uint GroupIndex : SV_GroupIndex, uint3 DispatchThreadId : SV_Dispatc
 		const FPackedHairSample NodeData = HairVisibilityNodeData[LocalOffset];		
 		const FHairSample Sample = UnpackHairSample(NodeData);
 
-		// Base Color
-		Pos = ShaderPrintSymbol(Pos, _MINUS_);
-		Pos = ShaderPrintSymbol(Pos, _MINUS_);
-		Pos = ShaderPrintSymbol(Pos, _MINUS_);
-		Pos = ShaderPrintSymbol(Pos, _MINUS_);
-		Pos = ShaderPrintSymbol(Pos, _MINUS_);
-		Pos = ShaderPrintSymbol(Pos, _MINUS_);
-		Pos = ShaderPrintSymbol(Pos, _MINUS_);
-		Pos = ShaderPrintSymbol(Pos, _MINUS_);
-		Pos = ShaderPrintSymbol(Pos, _MINUS_);
-		Pos = ShaderPrintSymbol(Pos, _MINUS_);
-		Pos = ShaderPrintSymbol(Pos, _MINUS_);
-		Pos = ShaderPrintSymbol(Pos, _MINUS_);
-		Pos = ShaderPrintSymbol(Pos, _MINUS_);
-		Pos = ShaderPrintSymbol(Pos, _MINUS_);
-		Pos = ShaderPrintSymbol(Pos, _MINUS_);
-		Pos = ShaderPrintSymbol(Pos, _MINUS_);
-		Pos = ShaderPrintSymbol(Pos, _MINUS_);
-		Pos = ShaderPrintSymbol(Pos, _MINUS_);
-		Pos = ShaderPrintSymbol(Pos, _MINUS_);
-		Pos = ShaderPrintSymbol(Pos, _MINUS_);
-
-		Pos = ShaderPrintNewline(Pos);
-		Pos.x = OriginalPos.x;
-
-#if MATERIAL_DATA
-		// Base Color
-		Pos = ShaderPrintSymbol(Pos, _DOT_);
-		Pos = ShaderPrintSymbol(Pos, _B_);
-		Pos = ShaderPrintSymbol(Pos, _A_);
-		Pos = ShaderPrintSymbol(Pos, _S_);
-		Pos = ShaderPrintSymbol(Pos, _E_);
-		Pos = ShaderPrintSymbol(Pos, _SPC_);
-		Pos = ShaderPrintSymbol(Pos, _C_);
-		Pos = ShaderPrintSymbol(Pos, _O_);
-		Pos = ShaderPrintSymbol(Pos, _L_);
-		Pos = ShaderPrintSymbol(Pos, _O_);
-		Pos = ShaderPrintSymbol(Pos, _R_);
-		Pos = ShaderPrintSymbol(Pos, _SPC_);
-		Pos = ShaderPrint(Pos, Sample.BaseColor);
-		
-		Pos = ShaderPrintNewline(Pos);
-		Pos.x = OriginalPos.x;
-
-		// Roughness
-		Pos = ShaderPrintSymbol(Pos, _DOT_);
-		Pos = ShaderPrintSymbol(Pos, _R_);
-		Pos = ShaderPrintSymbol(Pos, _O_);
-		Pos = ShaderPrintSymbol(Pos, _U_);
-		Pos = ShaderPrintSymbol(Pos, _G_);
-		Pos = ShaderPrintSymbol(Pos, _H_);
-		Pos = ShaderPrintSymbol(Pos, _N_);
-		Pos = ShaderPrintSymbol(Pos, _E_);
-		Pos = ShaderPrintSymbol(Pos, _S_);
-		Pos = ShaderPrintSymbol(Pos, _S_);
-		Pos = ShaderPrintSymbol(Pos, _SPC_);
-		Pos = ShaderPrint(Pos, Sample.Roughness);
-
-		Pos = ShaderPrintNewline(Pos);
-		Pos.x = OriginalPos.x;
-
-		// Specular
-		Pos = ShaderPrintSymbol(Pos, _DOT_);
-		Pos = ShaderPrintSymbol(Pos, _S_);
-		Pos = ShaderPrintSymbol(Pos, _P_);
-		Pos = ShaderPrintSymbol(Pos, _E_);
-		Pos = ShaderPrintSymbol(Pos, _C_);
-		Pos = ShaderPrintSymbol(Pos, _U_);
-		Pos = ShaderPrintSymbol(Pos, _L_);
-		Pos = ShaderPrintSymbol(Pos, _A_);
-		Pos = ShaderPrintSymbol(Pos, _R_);
-		Pos = ShaderPrintSymbol(Pos, _SPC_);
-		Pos = ShaderPrint(Pos, Sample.Specular);
-
-		Pos = ShaderPrintNewline(Pos);
-		Pos.x = OriginalPos.x;
-
-		// Backlit
-		Pos = ShaderPrintSymbol(Pos, _DOT_);
-		Pos = ShaderPrintSymbol(Pos, _B_);
-		Pos = ShaderPrintSymbol(Pos, _A_);
-		Pos = ShaderPrintSymbol(Pos, _C_);
-		Pos = ShaderPrintSymbol(Pos, _K_);
-		Pos = ShaderPrintSymbol(Pos, _L_);
-		Pos = ShaderPrintSymbol(Pos, _I_);
-		Pos = ShaderPrintSymbol(Pos, _T_);
-		Pos = ShaderPrintSymbol(Pos, _SPC_);
-		Pos = ShaderPrint(Pos, Sample.Backlit);
-
-		Pos = ShaderPrintNewline(Pos);
-		Pos.x = OriginalPos.x;
-
-		// Tangent
-		Pos = ShaderPrintSymbol(Pos, _DOT_);
-		Pos = ShaderPrintSymbol(Pos, _T_);
-		Pos = ShaderPrintSymbol(Pos, _A_);
-		Pos = ShaderPrintSymbol(Pos, _N_);
-		Pos = ShaderPrintSymbol(Pos, _G_);
-		Pos = ShaderPrintSymbol(Pos, _E_);
-		Pos = ShaderPrintSymbol(Pos, _N_);
-		Pos = ShaderPrintSymbol(Pos, _T_);
-		Pos = ShaderPrintSymbol(Pos, _SPC_);
-		Pos = ShaderPrint(Pos, Sample.Tangent);
-
-		Pos = ShaderPrintNewline(Pos);
-		Pos.x = OriginalPos.x;
-#endif 
-
-#if DETAIL_DATA
-		// PrimitiveId
-		Pos = ShaderPrintSymbol(Pos, _DOT_);
-		Pos = ShaderPrintSymbol(Pos, _P_);
-		Pos = ShaderPrintSymbol(Pos, _R_);
-		Pos = ShaderPrintSymbol(Pos, _I_);
-		Pos = ShaderPrintSymbol(Pos, _M_);
-		Pos = ShaderPrintSymbol(Pos, _I_);
-		Pos = ShaderPrintSymbol(Pos, _T_);
-		Pos = ShaderPrintSymbol(Pos, _I_);
-		Pos = ShaderPrintSymbol(Pos, _V_);
-		Pos = ShaderPrintSymbol(Pos, _E_);
-		Pos = ShaderPrintSymbol(Pos, _SPC_);
-		Pos = ShaderPrint(Pos, Sample.PrimitiveId);
-
-		Pos = ShaderPrintNewline(Pos);
-		Pos.x = OriginalPos.x;
-
-		// MacroGroupId
-		Pos = ShaderPrintSymbol(Pos, _DOT_);
-		Pos = ShaderPrintSymbol(Pos, _M_);
-		Pos = ShaderPrintSymbol(Pos, _A_);
-		Pos = ShaderPrintSymbol(Pos, _C_);
-		Pos = ShaderPrintSymbol(Pos, _R_);
-		Pos = ShaderPrintSymbol(Pos, _O_);
-		Pos = ShaderPrintSymbol(Pos, _SPC_);
-		Pos = ShaderPrintSymbol(Pos, _G_);
-		Pos = ShaderPrintSymbol(Pos, _R_);
-		Pos = ShaderPrintSymbol(Pos, _O_);
-		Pos = ShaderPrintSymbol(Pos, _U_);
-		Pos = ShaderPrintSymbol(Pos, _P_);
-		Pos = ShaderPrintSymbol(Pos, _SPC_);
-		Pos = ShaderPrintSymbol(Pos, _I_);
-		Pos = ShaderPrintSymbol(Pos, _D_);
-		Pos = ShaderPrintSymbol(Pos, _SPC_);
-		Pos = ShaderPrint(Pos, Sample.MacroGroupId);
-
-		Pos = ShaderPrintNewline(Pos);
-		Pos.x = OriginalPos.x;		
-
-		// Light channel mask
-		Pos = ShaderPrintSymbol(Pos, _DOT_);
-		Pos = ShaderPrintSymbol(Pos, _L_);
-		Pos = ShaderPrintSymbol(Pos, _I_);
-		Pos = ShaderPrintSymbol(Pos, _G_);
-		Pos = ShaderPrintSymbol(Pos, _H_);
-		Pos = ShaderPrintSymbol(Pos, _T_);
-		Pos = ShaderPrintSymbol(Pos, _SPC_);
-		Pos = ShaderPrintSymbol(Pos, _C_);
-		Pos = ShaderPrintSymbol(Pos, _H_);
-		Pos = ShaderPrintSymbol(Pos, _A_);
-		Pos = ShaderPrintSymbol(Pos, _N_);
-		Pos = ShaderPrintSymbol(Pos, _N_);
-		Pos = ShaderPrintSymbol(Pos, _E_);
-		Pos = ShaderPrintSymbol(Pos, _L_);
-		Pos = ShaderPrintSymbol(Pos, _S_);
-		Pos = ShaderPrintSymbol(Pos, _SPC_);
-		Pos = ShaderPrint(Pos, (Sample.LightChannelMask & 0x1) ? 1u : 0u);
-		Pos = ShaderPrint(Pos, (Sample.LightChannelMask & 0x2) ? 1u : 0u);
-		Pos = ShaderPrint(Pos, (Sample.LightChannelMask & 0x4) ? 1u : 0u);
-
-		Pos = ShaderPrintNewline(Pos);
-		Pos.x = OriginalPos.x;
-
-		// Scene lighting scatter
-		Pos = ShaderPrintSymbol(Pos, _DOT_);
-		Pos = ShaderPrintSymbol(Pos, _S_);
-		Pos = ShaderPrintSymbol(Pos, _C_);
-		Pos = ShaderPrintSymbol(Pos, _A_);
-		Pos = ShaderPrintSymbol(Pos, _T_);
-		Pos = ShaderPrintSymbol(Pos, _T_);
-		Pos = ShaderPrintSymbol(Pos, _E_);
-		Pos = ShaderPrintSymbol(Pos, _R_);
-		Pos = ShaderPrintSymbol(Pos, _SPC_);
-		Pos = ShaderPrintSymbol(Pos, _S_);
-		Pos = ShaderPrintSymbol(Pos, _C_);
-		Pos = ShaderPrintSymbol(Pos, _E_);
-		Pos = ShaderPrintSymbol(Pos, _N_);
-		Pos = ShaderPrintSymbol(Pos, _E_);
-		Pos = ShaderPrintSymbol(Pos, _SPC_);
-		Pos = ShaderPrint(Pos, Sample.bScatterSceneLighting ? 1u : 0u);
-
-		Pos = ShaderPrintNewline(Pos);
-		Pos.x = OriginalPos.x;
-
-		// Depth
-		Pos = ShaderPrintSymbol(Pos, _DOT_);
-		Pos = ShaderPrintSymbol(Pos, _D_);
-		Pos = ShaderPrintSymbol(Pos, _E_);
-		Pos = ShaderPrintSymbol(Pos, _P_);
-		Pos = ShaderPrintSymbol(Pos, _T_);
-		Pos = ShaderPrintSymbol(Pos, _H_);
-		Pos = ShaderPrintSymbol(Pos, _SPC_);
-		Pos = ShaderPrint(Pos, Sample.Depth);
-
-		Pos = ShaderPrintNewline(Pos);
-		Pos.x = OriginalPos.x;
-#endif
-
-		// Coverage8bit
-		Pos = ShaderPrintSymbol(Pos, _DOT_);
-		Pos = ShaderPrintSymbol(Pos, _C_);
-		Pos = ShaderPrintSymbol(Pos, _O_);
-		Pos = ShaderPrintSymbol(Pos, _V_);
-		Pos = ShaderPrintSymbol(Pos, _E_);
-		Pos = ShaderPrintSymbol(Pos, _R_);
-		Pos = ShaderPrintSymbol(Pos, _A_);
-		Pos = ShaderPrintSymbol(Pos, _G_);
-		Pos = ShaderPrintSymbol(Pos, _E_);
-		Pos = ShaderPrintSymbol(Pos, _SPC_);
-		Pos = ShaderPrint(Pos, Sample.Coverage8bit);
-
-		Pos = ShaderPrintNewline(Pos);
-		Pos.x = OriginalPos.x;
+		float2 StartSamplePos = ColumnPos + float2(150.f + 80.f * SampleIt, 0.f) / float2(MaxResolution);
+		Pos = StartSamplePos;
+
+		Pos = ShaderPrint(Pos, SampleIt, FontWhite);									AddNewLine(StartSamplePos, Pos);
+
+	#if MATERIAL_DATA
+		Pos = ShaderPrint(Pos, Sample.BaseColor.x,	FontRed);							AddNewLine(StartSamplePos, Pos);
+		Pos = ShaderPrint(Pos, Sample.BaseColor.y,	FontGreen);							AddNewLine(StartSamplePos, Pos);
+		Pos = ShaderPrint(Pos, Sample.BaseColor.z,	FontBlue);							AddNewLine(StartSamplePos, Pos);
+		AddNewLine(StartSamplePos, Pos);
+
+		Pos = ShaderPrint(Pos, Sample.Tangent.x,	FontEmerald);						AddNewLine(StartSamplePos, Pos);
+		Pos = ShaderPrint(Pos, Sample.Tangent.y,	FontEmerald);						AddNewLine(StartSamplePos, Pos);
+		Pos = ShaderPrint(Pos, Sample.Tangent.z,	FontEmerald);						AddNewLine(StartSamplePos, Pos);
+		AddNewLine(StartSamplePos, Pos);
+
+		Pos = ShaderPrint(Pos, Sample.Roughness,	FontOrange);						AddNewLine(StartSamplePos, Pos);
+		Pos = ShaderPrint(Pos, Sample.Specular,		FontEmerald);						AddNewLine(StartSamplePos, Pos);
+		Pos = ShaderPrint(Pos, Sample.Backlit,		FontOrange);						AddNewLine(StartSamplePos, Pos);
+		AddNewLine(StartSamplePos, Pos);
+	#endif 
+
+	#if DETAIL_DATA
+		Pos = ShaderPrint(Pos,(Sample.LightChannelMask & 0x1) ? 1u : 0u, FontEmerald);	AddNewLine(StartSamplePos, Pos);
+//		Pos = ShaderPrint(Pos,(Sample.LightChannelMask & 0x1) ? 1u : 0u);				AddNewLine(StartSamplePos, Pos);
+//		Pos = ShaderPrint(Pos,(Sample.LightChannelMask & 0x2) ? 1u : 0u);				AddNewLine(StartSamplePos, Pos);
+//		Pos = ShaderPrint(Pos,(Sample.LightChannelMask & 0x4) ? 1u : 0u);				AddNewLine(StartSamplePos, Pos);
+		Pos = ShaderPrint(Pos, Sample.bScatterSceneLighting ? 1u : 0u, FontOrange);		AddNewLine(StartSamplePos, Pos);
+		Pos = ShaderPrint(Pos, Sample.Depth, FontEmerald);								AddNewLine(StartSamplePos, Pos);
+		AddNewLine(StartSamplePos, Pos);
+
+		Pos = ShaderPrint(Pos, Sample.MacroGroupId, FontYellow);						AddNewLine(StartSamplePos, Pos);
+		Pos = ShaderPrint(Pos, Sample.PrimitiveId, FontEmerald);						AddNewLine(StartSamplePos, Pos);
+	#endif
+		Pos = ShaderPrint(Pos, Sample.Coverage8bit, FontOrange);						AddNewLine(StartSamplePos, Pos);
 	}
 }
diff --git a/Engine/Shaders/Private/HairStrands/HairStrandsVoxelPageRayMarching.usf b/Engine/Shaders/Private/HairStrands/HairStrandsVoxelPageRayMarching.usf
--- a/Engine/Shaders/Private/HairStrands/HairStrandsVoxelPageRayMarching.usf
+++ b/Engine/Shaders/Private/HairStrands/HairStrandsVoxelPageRayMarching.usf
@@ -49,6 +49,12 @@ uint2 GetVoxelPageIndexOccupancy(uint3 InPageIndexCoord, uint3 InPageIndexResolu
 	return VirtualVoxel.PageIndexOccupancyBuffer.Load(LinearPageIndexCoord);
 }
 
+FFontColor GetOccupancyColor(float In)
+{
+	float3 Color = lerp(float3(0, 1, 0), float3(1, 0, 0), saturate(In));
+	return InitFontColor(Color);
+}
+
 [numthreads(8, 8, 1)]
 void MainCS(uint3 DispatchThreadId : SV_DispatchThreadID)
 {
@@ -98,57 +104,56 @@ void MainCS(uint3 DispatchThreadId : SV_DispatchThreadID)
 			float2 StartPos = float2(50.0f, 100.0f ) / OutputResolution;
 			float2 Pos = StartPos;
 
-			Pos = ShaderPrintText(Pos, TEXT("General"));
-			Pos = ShaderPrintNewline(Pos); Pos.x = StartPos.x;
-
-			Pos = ShaderPrintText(Pos, TEXT(" Voxel World Size "));
-			Pos = ShaderPrint(Pos, VirtualVoxel.VoxelWorldSize);
+			Pos = ShaderPrintText(Pos, TEXT("Voxel World Size "), FontOrange);
+			Pos = ShaderPrint(Pos, VirtualVoxel.VoxelWorldSize, FontOrange);
 			Pos = ShaderPrintNewline(Pos); Pos.x = StartPos.x;
 			
-			Pos = ShaderPrintText(Pos, TEXT(" Page Res         "));
+			Pos = ShaderPrintText(Pos, TEXT("Page Res         "));
 			Pos = ShaderPrint(Pos, VirtualVoxel.PageResolution);
 			Pos = ShaderPrintNewline(Pos); Pos.x = StartPos.x;
 
-			Pos = ShaderPrintText(Pos, TEXT(" Page Count Res   "));
+			Pos = ShaderPrintText(Pos, TEXT("Page Count Res   "));
 			Pos = ShaderPrint(Pos, VirtualVoxel.PageCountResolution.x);
 			Pos = ShaderPrint(Pos, VirtualVoxel.PageCountResolution.y);
 			Pos = ShaderPrint(Pos, VirtualVoxel.PageCountResolution.z);
 			Pos = ShaderPrintNewline(Pos); Pos.x = StartPos.x;
 
-			Pos = ShaderPrintText(Pos, TEXT(" Page Texture Res  "));
-			Pos = ShaderPrint(Pos, VirtualVoxel.PageTextureResolution.x);
-			Pos = ShaderPrint(Pos, VirtualVoxel.PageTextureResolution.y);
-			Pos = ShaderPrint(Pos, VirtualVoxel.PageTextureResolution.z);
+			Pos = ShaderPrintText(Pos, TEXT("Page Texture Res  "), FontEmerald);
+			Pos = ShaderPrint(Pos, VirtualVoxel.PageTextureResolution.x, FontEmerald);
+			Pos = ShaderPrint(Pos, VirtualVoxel.PageTextureResolution.y, FontEmerald);
+			Pos = ShaderPrint(Pos, VirtualVoxel.PageTextureResolution.z, FontEmerald);
 			Pos = ShaderPrintNewline(Pos); Pos.x = StartPos.x;
 
 			// Total page indices stats
 			Pos = ShaderPrintNewline(Pos); Pos.x = StartPos.x;
 
-			Pos = ShaderPrintText(Pos, TEXT(" Total Allocated Page Index  "));
+			Pos = ShaderPrintText(Pos, TEXT("Total Allocated Page Index  "));
 			Pos = ShaderPrint(Pos, TotalAllocatedPageCount);
 			Pos = ShaderPrintNewline(Pos); Pos.x = StartPos.x;
 
-			Pos = ShaderPrintText(Pos, TEXT(" Total Page Index            "));
+			Pos = ShaderPrintText(Pos, TEXT("Total Page Index            "));
 			Pos = ShaderPrint(Pos, TotalPageIndexCount);
 			Pos = ShaderPrintNewline(Pos); Pos.x = StartPos.x;
 
-			Pos = ShaderPrintText(Pos, TEXT(" Page Index Occupancy         "));
-			Pos = ShaderPrint(Pos, TotalAllocatedPageCount / float(TotalPageIndexCount) * 100.f);
+			const float PageIndexOccupancy = TotalAllocatedPageCount / float(TotalPageIndexCount);
+			Pos = ShaderPrintText(Pos, TEXT("Page Index Occupancy         "), FontOrange);
+			Pos = ShaderPrint(Pos, PageIndexOccupancy * 100.f, GetOccupancyColor(PageIndexOccupancy));
 			Pos = ShaderPrintNewline(Pos); Pos.x = StartPos.x;
 
 			// Total pages stats
 			Pos = ShaderPrintNewline(Pos); Pos.x = StartPos.x;
 
-			Pos = ShaderPrintText(Pos, TEXT(" Total Allocated Valid Page  "));
+			Pos = ShaderPrintText(Pos, TEXT("Total Allocated Valid Page  "));
 			Pos = ShaderPrint(Pos, TotalAllocatedValidPageCount);
 			Pos = ShaderPrintNewline(Pos); Pos.x = StartPos.x;
 
-			Pos = ShaderPrintText(Pos, TEXT(" Total Page                  "));
+			Pos = ShaderPrintText(Pos, TEXT("Total Page                  "));
 			Pos = ShaderPrint(Pos, TotalPageCount);
 			Pos = ShaderPrintNewline(Pos); Pos.x = StartPos.x;
 
-			Pos = ShaderPrintText(Pos, TEXT(" Page Occupancy              "));
-			Pos = ShaderPrint(Pos, TotalAllocatedValidPageCount / float(TotalPageCount) * 100.f);
+			const float PageOccupancy = TotalAllocatedValidPageCount / float(TotalPageCount);
+			Pos = ShaderPrintText(Pos, TEXT("Page Occupancy              "), FontOrange);
+			Pos = ShaderPrint(Pos, PageOccupancy * 100.f, GetOccupancyColor(PageOccupancy));
 			Pos = ShaderPrintNewline(Pos); Pos.x = StartPos.x;
 		}
 
@@ -156,31 +161,62 @@ void MainCS(uint3 DispatchThreadId : SV_DispatchThreadID)
 		//  Page Index Resolution: X Y Z (Linear/Total - %) 
 		//  Virtual Resolution   : X Y Z
 		//  Page Index Offset    : O
+		if (MacroGroupId == 0)
 		{
-			float2 StartPos = float2(50.0f, 420.0f + float(MacroGroupId) * 100.0f) / OutputResolution;
+			float2 StartPos = float2(50.0f, 420.0f) / OutputResolution;
 			float2 Pos = StartPos;
 
-			Pos = ShaderPrintText(Pos, TEXT("Macro Group    "));
-			Pos = ShaderPrint(Pos, MacroGroupId);
+			Pos = ShaderPrintText(Pos, TEXT("Macro Group        "), FontOrange);
+			Pos = ShaderPrintNewline(Pos); Pos.x = StartPos.x;
+
+			Pos = ShaderPrintText(Pos, TEXT(" Virtual Res X     "), FontEmerald);
+			Pos = ShaderPrintNewline(Pos); Pos.x = StartPos.x;
+			Pos = ShaderPrintText(Pos, TEXT(" Virtual Res Y     "), FontEmerald);
+			Pos = ShaderPrintNewline(Pos); Pos.x = StartPos.x;
+			Pos = ShaderPrintText(Pos, TEXT(" Virtual Res Z     "), FontEmerald);
+			Pos = ShaderPrintNewline(Pos); Pos.x = StartPos.x;
+
 			Pos = ShaderPrintNewline(Pos); Pos.x = StartPos.x;
 
-			Pos = ShaderPrintText(Pos, TEXT(" Virtual Res       "));
-			Pos = ShaderPrint(Pos, NodeDesc.VirtualResolution.x);
-			Pos = ShaderPrint(Pos, NodeDesc.VirtualResolution.y);
-			Pos = ShaderPrint(Pos, NodeDesc.VirtualResolution.z);
+			Pos = ShaderPrintText(Pos, TEXT(" Page Index Res X  "));
+			Pos = ShaderPrintNewline(Pos); Pos.x = StartPos.x;
+			Pos = ShaderPrintText(Pos, TEXT(" Page Index Res Y  "));
+			Pos = ShaderPrintNewline(Pos); Pos.x = StartPos.x;
+			Pos = ShaderPrintText(Pos, TEXT(" Page Index Res Z  "));
 			Pos = ShaderPrintNewline(Pos); Pos.x = StartPos.x;
 
-			Pos = ShaderPrintText(Pos, TEXT(" Page Index Res    "));
-			Pos = ShaderPrint(Pos, NodeDesc.PageIndexResolution.x);
-			Pos = ShaderPrint(Pos, NodeDesc.PageIndexResolution.y);
-			Pos = ShaderPrint(Pos, NodeDesc.PageIndexResolution.z);
-			Pos = ShaderPrintText(Pos, TEXT(" "));
 			Pos = ShaderPrintNewline(Pos); Pos.x = StartPos.x;
 
 			Pos = ShaderPrintText(Pos, TEXT(" Page Index Offset "));
-			Pos = ShaderPrint(Pos, NodeDesc.PageIndexOffset);
+			Pos = ShaderPrintNewline(Pos); Pos.x = StartPos.x;
+		}
+
+		{
+			float2 StartPos = float2(250.0f + float(MacroGroupId) * 35.0f, 420.0f) / OutputResolution;
+			float2 Pos = StartPos;
+
+			Pos = ShaderPrint(Pos, MacroGroupId, FontOrange);
+			Pos = ShaderPrintNewline(Pos); Pos.x = StartPos.x;
+
+			Pos = ShaderPrint(Pos, NodeDesc.VirtualResolution.x, FontEmerald);
+			Pos = ShaderPrintNewline(Pos); Pos.x = StartPos.x;
+			Pos = ShaderPrint(Pos, NodeDesc.VirtualResolution.y, FontEmerald);
+			Pos = ShaderPrintNewline(Pos); Pos.x = StartPos.x;
+			Pos = ShaderPrint(Pos, NodeDesc.VirtualResolution.z, FontEmerald);
+			Pos = ShaderPrintNewline(Pos); Pos.x = StartPos.x;
+
+			Pos = ShaderPrintNewline(Pos); Pos.x = StartPos.x;
+
+			Pos = ShaderPrint(Pos, NodeDesc.PageIndexResolution.x, FontYellow);
+			Pos = ShaderPrintNewline(Pos); Pos.x = StartPos.x;
+			Pos = ShaderPrint(Pos, NodeDesc.PageIndexResolution.y, FontYellow);
+			Pos = ShaderPrintNewline(Pos); Pos.x = StartPos.x;
+			Pos = ShaderPrint(Pos, NodeDesc.PageIndexResolution.z, FontYellow);
+			Pos = ShaderPrintNewline(Pos); Pos.x = StartPos.x;
+
 			Pos = ShaderPrintNewline(Pos); Pos.x = StartPos.x;
 
+			Pos = ShaderPrint(Pos, NodeDesc.PageIndexOffset, FontSilver);
 			Pos = ShaderPrintNewline(Pos); Pos.x = StartPos.x;
 		}
 	}
