Change render pass to use hair-only depth texture for pixel rejection instead of the categorization texture.

diff --git a/Engine/Shaders/Private/RayTracing/RayTracingOcclusionRGS.usf b/Engine/Shaders/Private/RayTracing/RayTracingOcclusionRGS.usf
--- a/Engine/Shaders/Private/RayTracing/RayTracingOcclusionRGS.usf
+++ b/Engine/Shaders/Private/RayTracing/RayTracingOcclusionRGS.usf
@@ -60,13 +60,13 @@ void NeedTraceHair(
 	inout bool bTraceHairRay,
 	inout float HairDeviceZ)
 {
-	const FCategorizationData Categorization = DecodeCategorizationData(HairStrands.HairCategorizationTexture.Load(uint3(PixelCoord, 0)));
+	const float HairSampleDepth = DecodeCategorizationData(HairStrands.HairOnlyDepthTexture.Load(uint3(PixelCoord, 0)));
 	bTraceHairRay = false;
 	HairDeviceZ = 0;
-	if (Categorization.PixelCoverage > 0)
+	if (HairSampleDepth > 0)
 	{
 		const uint HairLightChannel = HairLightChannelMaskTexture.Load(uint3(PixelCoord, 0));
-		HairDeviceZ = Categorization.ClosestDepth;
+		HairDeviceZ = HairSampleDepth;
 		bTraceHairRay = (HairLightChannel & LightingChannelMask) != 0;
 	}
 }
diff --git a/Engine/Shaders/Private/ShadowProjectionPixelShader.usf b/Engine/Shaders/Private/ShadowProjectionPixelShader.usf
--- a/Engine/Shaders/Private/ShadowProjectionPixelShader.usf
+++ b/Engine/Shaders/Private/ShadowProjectionPixelShader.usf
@@ -126,8 +126,7 @@ void Main(
 	const uint3 PixelCoord = uint3(floor(SVPos.xy), 0);
 #if SUBPIXEL_SHADOW
 	const bool bUseFadePlane = bCascadeUseFadePlane > 0;
-	const FCategorizationData CategorizationData = DecodeCategorizationData(HairStrands.HairCategorizationTexture.Load(PixelCoord));
-	SceneW = ConvertFromDeviceZ(CategorizationData.ClosestDepth);
+	SceneW = ConvertFromDeviceZ(HairStrands.HairOnlyDepthTexture.Load(PixelCoord));
 	bIsSubsurfaceCompatible = false;
 #endif
 
@@ -414,10 +413,10 @@ void MainOnePassPointLightPS(
 
 #if SUBPIXEL_SHADOW
 	const uint3 PixelCoord = uint3(floor(SVPos.xy), 0);
-	const FCategorizationData CategorizationData = DecodeCategorizationData(HairStrands.HairCategorizationTexture.Load(PixelCoord));
-	if (CategorizationData.PixelCoverage > 0)
+	const float HairSampleDepth = HairStrands.HairOnlyDepthTexture.Load(PixelCoord);
+	if (HairSampleDepth > 0)
 	{
-		SceneW = ConvertFromDeviceZ(CategorizationData.ClosestDepth);
+		SceneW = ConvertFromDeviceZ(HairSampleDepth);
 	}
 	else
 	{
diff --git a/Engine/Shaders/Private/VirtualShadowMaps/PageManagement.usf b/Engine/Shaders/Private/VirtualShadowMaps/PageManagement.usf
--- a/Engine/Shaders/Private/VirtualShadowMaps/PageManagement.usf
+++ b/Engine/Shaders/Private/VirtualShadowMaps/PageManagement.usf
@@ -222,11 +222,11 @@ void GeneratePageFlagsFromPixels(
 	bool   bIsValid		 = false;
 	#if PERMUTATION_INPUT_TYPE == INPUT_TYPE_HAIRSTRANDS
 	{
-		FCategorizationData CatData = DecodeCategorizationData(HairStrands.HairCategorizationTexture.Load(uint3(PixelPos, 0)));
-		FScreenPixelData HairPixel = GetScreenPixelData(PixelPos, CatData.ClosestDepth);
+		const float HairDeviceZ = HairStrands.HairOnlyDepthTexture.Load(uint3(PixelPos, 0)).x;
+		FScreenPixelData HairPixel = GetScreenPixelData(PixelPos, HairDeviceZ);
 		WorldPosition = HairPixel.WorldPosition;
 		SceneDepth = HairPixel.SceneDepth;
-		bIsValid = CatData.PixelCoverage > 0;
+		bIsValid = HairDeviceZ > 0;
 	}
 	#else // PERMUTATION_INPUT_TYPE == INPUT_TYPE_GBUFFER || PERMUTATION_INPUT_TYPE == INPUT_TYPE_NANITE
 	{
diff --git a/Engine/Shaders/Private/VirtualShadowMaps/VirtualShadowMapProjection.usf b/Engine/Shaders/Private/VirtualShadowMaps/VirtualShadowMapProjection.usf
--- a/Engine/Shaders/Private/VirtualShadowMaps/VirtualShadowMapProjection.usf
+++ b/Engine/Shaders/Private/VirtualShadowMaps/VirtualShadowMapProjection.usf
@@ -411,9 +411,8 @@ void VirtualShadowMapProjection(
 	#if HAS_HAIR_STRANDS
 	if (InputType == INPUT_TYPE_HAIRSTRANDS)
 	{
-		FCategorizationData CatData = DecodeCategorizationData(HairStrands.HairCategorizationTexture.Load( int3(PixelPos, 0 ) ) );
-		DeviceZ = CatData.ClosestDepth;
-		if (CatData.PixelCoverage == 0)
+		DeviceZ = HairStrands.HairOnlyDepthTexture.Load(int3(PixelPos, 0));
+		if (DeviceZ == 0)
 		{
 			return;
 		}
