Fix typo in RaytracingOcclusionRGS due to previous refactor.

diff --git a/Engine/Shaders/Private/RayTracing/RayTracingOcclusionRGS.usf b/Engine/Shaders/Private/RayTracing/RayTracingOcclusionRGS.usf
--- a/Engine/Shaders/Private/RayTracing/RayTracingOcclusionRGS.usf
+++ b/Engine/Shaders/Private/RayTracing/RayTracingOcclusionRGS.usf
@@ -60,7 +60,7 @@ void NeedTraceHair(
 	inout bool bTraceHairRay,
 	inout float HairDeviceZ)
 {
-	const float HairSampleDepth = DecodeCategorizationData(HairStrands.HairOnlyDepthTexture.Load(uint3(PixelCoord, 0)));
+	const float HairSampleDepth = HairStrands.HairOnlyDepthTexture.Load(uint3(PixelCoord, 0));
 	bTraceHairRay = false;
 	HairDeviceZ = 0;
 	if (HairSampleDepth > 0)
