Fix Android cook build.

diff --git a/Engine/Shaders/Private/HairStrands/HairStrandsVisibilityCommon.ush b/Engine/Shaders/Private/HairStrands/HairStrandsVisibilityCommon.ush
--- a/Engine/Shaders/Private/HairStrands/HairStrandsVisibilityCommon.ush
+++ b/Engine/Shaders/Private/HairStrands/HairStrandsVisibilityCommon.ush
@@ -226,8 +226,6 @@ uint GetPackedHairSampleCoverage8bit(in FPackedHairSample In)
 ////////////////////////////////////////////////////////////////////////////////////////////////////////////////////////////////
 // Hair velocity
 
-float3 DecodeVelocityFromTexture(float4 EncodedV);
-
 // Return true if the underlying pixel need to run with TAA fast resolve. This is the case when 
 // the velocity is larger than a certain threshold.
 bool NeedFastResolve(float2 InEncodedVelocity, float2 InVelocity, float InVelocityThreshold)
