Fix for packed texture co-ordinates crashing on some Android Vulkan devices

diff --git a/Engine/Shaders/Private/GeometryCollectionVertexFactory.ush b/Engine/Shaders/Private/GeometryCollectionVertexFactory.ush
--- a/Engine/Shaders/Private/GeometryCollectionVertexFactory.ush
+++ b/Engine/Shaders/Private/GeometryCollectionVertexFactory.ush
@@ -61,17 +61,44 @@ struct FVertexFactoryInput
 
 #if NUM_MATERIAL_TEXCOORDS_VERTEX
 	#if !MANUAL_VERTEX_FETCH
-		#if NUM_MATERIAL_TEXCOORDS_VERTEX > 1
-			float4	PackedTexCoords4[NUM_MATERIAL_TEXCOORDS_VERTEX/2] : ATTRIBUTE4;
-		#endif
-		#if NUM_MATERIAL_TEXCOORDS_VERTEX == 1
-			float2	PackedTexCoords2 : ATTRIBUTE4;
-		#elif NUM_MATERIAL_TEXCOORDS_VERTEX == 3
-			float2	PackedTexCoords2 : ATTRIBUTE5;
-		#elif NUM_MATERIAL_TEXCOORDS_VERTEX == 5
-			float2	PackedTexCoords2 : ATTRIBUTE6;
-		#elif NUM_MATERIAL_TEXCOORDS_VERTEX == 7
-			float2	PackedTexCoords2 : ATTRIBUTE7;
+		#if !PACK_TEX_COORDS
+			#if NUM_MATERIAL_TEXCOORDS_VERTEX > 0
+				float2	TexCoords0 : ATTRIBUTE4;
+			#endif
+			#if NUM_MATERIAL_TEXCOORDS_VERTEX > 1
+				float2	TexCoords1 : ATTRIBUTE5;
+			#endif
+			#if NUM_MATERIAL_TEXCOORDS_VERTEX > 2
+				float2	TexCoords2 : ATTRIBUTE6;
+			#endif
+			#if NUM_MATERIAL_TEXCOORDS_VERTEX > 3
+				float2	TexCoords3 : ATTRIBUTE7;
+			#endif
+			#if NUM_MATERIAL_TEXCOORDS_VERTEX > 4
+				float2	TexCoords4 : ATTRIBUTE8;
+			#endif
+			#if NUM_MATERIAL_TEXCOORDS_VERTEX > 5
+				float2	TexCoords5 : ATTRIBUTE9;
+			#endif
+			#if NUM_MATERIAL_TEXCOORDS_VERTEX > 6
+				float2	TexCoords5 : ATTRIBUTE10;
+			#endif
+			#if NUM_MATERIAL_TEXCOORDS_VERTEX > 7
+				float2	TexCoords5 : ATTRIBUTE11;
+			#endif
+		#else
+			#if NUM_MATERIAL_TEXCOORDS_VERTEX > 1
+				float4	PackedTexCoords4[NUM_MATERIAL_TEXCOORDS_VERTEX/2] : ATTRIBUTE4;
+			#endif
+			#if NUM_MATERIAL_TEXCOORDS_VERTEX == 1
+				float2	PackedTexCoords2 : ATTRIBUTE4;
+			#elif NUM_MATERIAL_TEXCOORDS_VERTEX == 3
+				float2	PackedTexCoords2 : ATTRIBUTE5;
+			#elif NUM_MATERIAL_TEXCOORDS_VERTEX == 5
+				float2	PackedTexCoords2 : ATTRIBUTE6;
+			#elif NUM_MATERIAL_TEXCOORDS_VERTEX == 7
+				float2	PackedTexCoords2 : ATTRIBUTE7;
+			#endif
 		#endif
 	#endif
 #elif USE_PARTICLE_SUBUVS && !MANUAL_VERTEX_FETCH
@@ -252,19 +279,46 @@ FMaterialVertexParameters GetMaterialVertexParameters(FVertexFactoryInput Input,
 		Result.TexCoords[CoordinateIndex] = GeometryCollectionVF.VertexFetch_TexCoordBuffer[NumFetchTexCoords * (GeometryCollectionVF.VertexFetch_Parameters[VF_VertexOffset] + Input.VertexId) + ClampedCoordinateIndex];
 	}
 #elif NUM_MATERIAL_TEXCOORDS_VERTEX
-	#if NUM_MATERIAL_TEXCOORDS_VERTEX > 1
-		UNROLL
-		for (int CoordinateIndex = 0; CoordinateIndex < NUM_MATERIAL_TEXCOORDS_VERTEX - 1; CoordinateIndex += 2)
-		{
-			Result.TexCoords[CoordinateIndex] = Input.PackedTexCoords4[CoordinateIndex / 2].xy;
-			if (CoordinateIndex + 1 < NUM_MATERIAL_TEXCOORDS_VERTEX)
+	#if !PACK_TEX_COORDS
+		#if NUM_MATERIAL_TEXCOORDS_VERTEX > 0
+			Result.TexCoords[0] = Input.TexCoords0;
+		#endif
+		#if NUM_MATERIAL_TEXCOORDS_VERTEX > 1
+			Result.TexCoords[1] = Input.TexCoords1;
+		#endif
+		#if NUM_MATERIAL_TEXCOORDS_VERTEX > 2
+			Result.TexCoords[2] = Input.TexCoords2;
+		#endif
+		#if NUM_MATERIAL_TEXCOORDS_VERTEX > 3
+			Result.TexCoords[3] = Input.TexCoords3;
+		#endif
+		#if NUM_MATERIAL_TEXCOORDS_VERTEX > 4
+			Result.TexCoords[4] = Input.TexCoords4;
+		#endif
+		#if NUM_MATERIAL_TEXCOORDS_VERTEX > 5
+			Result.TexCoords[5] = Input.TexCoords5;
+		#endif
+		#if NUM_MATERIAL_TEXCOORDS_VERTEX > 6
+			Result.TexCoords[5] = Input.TexCoords5;
+		#endif
+		#if NUM_MATERIAL_TEXCOORDS_VERTEX > 7
+			Result.TexCoords[7] = Input.TexCoords7;
+		#endif
+	#else
+		#if NUM_MATERIAL_TEXCOORDS_VERTEX > 1
+			UNROLL
+			for(int CoordinateIndex = 0; CoordinateIndex < NUM_MATERIAL_TEXCOORDS_VERTEX-1; CoordinateIndex+=2)
 			{
-				Result.TexCoords[CoordinateIndex+1] = Input.PackedTexCoords4[CoordinateIndex / 2].zw;
+				Result.TexCoords[CoordinateIndex] = Input.PackedTexCoords4[CoordinateIndex/2].xy;
+				if( CoordinateIndex+1 < NUM_MATERIAL_TEXCOORDS_VERTEX )
+				{
+					Result.TexCoords[CoordinateIndex+1] = Input.PackedTexCoords4[CoordinateIndex/2].zw;
+				}
 			}
-		}
-	#endif
-	#if NUM_MATERIAL_TEXCOORDS_VERTEX % 2 == 1
-		Result.TexCoords[NUM_MATERIAL_TEXCOORDS_VERTEX-1] = Input.PackedTexCoords2;
+		#endif	// NUM_MATERIAL_TEXCOORDS_VERTEX > 1
+		#if NUM_MATERIAL_TEXCOORDS_VERTEX % 2 == 1
+			Result.TexCoords[NUM_MATERIAL_TEXCOORDS_VERTEX-1] = Input.PackedTexCoords2;
+		#endif	// NUM_MATERIAL_TEXCOORDS_VERTEX % 2 == 1
 	#endif
 #endif  //MANUAL_VERTEX_FETCH && NUM_MATERIAL_TEXCOORDS_VERTEX
 
diff --git a/Engine/Shaders/Private/LocalVertexFactory.ush b/Engine/Shaders/Private/LocalVertexFactory.ush
--- a/Engine/Shaders/Private/LocalVertexFactory.ush
+++ b/Engine/Shaders/Private/LocalVertexFactory.ush
@@ -110,17 +110,44 @@ struct FVertexFactoryInput
 				#error Too many texture coordinate sets defined on GPUSkin vertex input. Max: 4.
 			#endif
 		#else
-			#if NUM_MATERIAL_TEXCOORDS_VERTEX > 1
-				float4	PackedTexCoords4[NUM_MATERIAL_TEXCOORDS_VERTEX/2] : ATTRIBUTE4;
-			#endif
-			#if NUM_MATERIAL_TEXCOORDS_VERTEX == 1
-				float2	PackedTexCoords2 : ATTRIBUTE4;
-			#elif NUM_MATERIAL_TEXCOORDS_VERTEX == 3
-				float2	PackedTexCoords2 : ATTRIBUTE5;
-			#elif NUM_MATERIAL_TEXCOORDS_VERTEX == 5
-				float2	PackedTexCoords2 : ATTRIBUTE6;
-			#elif NUM_MATERIAL_TEXCOORDS_VERTEX == 7
-				float2	PackedTexCoords2 : ATTRIBUTE7;
+			#if !PACK_TEX_COORDS
+				#if NUM_MATERIAL_TEXCOORDS_VERTEX > 0
+					float2	TexCoords0 : ATTRIBUTE4;
+				#endif
+				#if NUM_MATERIAL_TEXCOORDS_VERTEX > 1
+					float2	TexCoords1 : ATTRIBUTE5;
+				#endif
+				#if NUM_MATERIAL_TEXCOORDS_VERTEX > 2
+					float2	TexCoords2 : ATTRIBUTE6;
+				#endif
+				#if NUM_MATERIAL_TEXCOORDS_VERTEX > 3
+					float2	TexCoords3 : ATTRIBUTE7;
+				#endif
+				#if NUM_MATERIAL_TEXCOORDS_VERTEX > 4
+					float2	TexCoords4 : ATTRIBUTE8;
+				#endif
+				#if NUM_MATERIAL_TEXCOORDS_VERTEX > 5
+					float2	TexCoords5 : ATTRIBUTE9;
+				#endif
+				#if NUM_MATERIAL_TEXCOORDS_VERTEX > 6
+					float2	TexCoords5 : ATTRIBUTE10;
+				#endif
+				#if NUM_MATERIAL_TEXCOORDS_VERTEX > 7
+					float2	TexCoords5 : ATTRIBUTE11;
+				#endif
+			#else
+				#if NUM_MATERIAL_TEXCOORDS_VERTEX > 1
+					float4	PackedTexCoords4[NUM_MATERIAL_TEXCOORDS_VERTEX/2] : ATTRIBUTE4;
+				#endif
+				#if NUM_MATERIAL_TEXCOORDS_VERTEX == 1
+					float2	PackedTexCoords2 : ATTRIBUTE4;
+				#elif NUM_MATERIAL_TEXCOORDS_VERTEX == 3
+					float2	PackedTexCoords2 : ATTRIBUTE5;
+				#elif NUM_MATERIAL_TEXCOORDS_VERTEX == 5
+					float2	PackedTexCoords2 : ATTRIBUTE6;
+				#elif NUM_MATERIAL_TEXCOORDS_VERTEX == 7
+					float2	PackedTexCoords2 : ATTRIBUTE7;
+				#endif
 			#endif
 		#endif
 	#endif
@@ -561,20 +588,47 @@ FMaterialVertexParameters GetMaterialVertexParameters(FVertexFactoryInput Input,
 				Result.TexCoords[CoordinateIndex] = Input.TexCoords[CoordinateIndex].xy;
 			}
 		#else
-			#if NUM_MATERIAL_TEXCOORDS_VERTEX > 1
-				UNROLL
-				for(int CoordinateIndex = 0; CoordinateIndex < NUM_MATERIAL_TEXCOORDS_VERTEX-1; CoordinateIndex+=2)
-				{
-					Result.TexCoords[CoordinateIndex] = Input.PackedTexCoords4[CoordinateIndex/2].xy;
-					if( CoordinateIndex+1 < NUM_MATERIAL_TEXCOORDS_VERTEX )
+			#if !PACK_TEX_COORDS
+				#if NUM_MATERIAL_TEXCOORDS_VERTEX > 0
+					Result.TexCoords[0] = Input.TexCoords0;
+				#endif
+				#if NUM_MATERIAL_TEXCOORDS_VERTEX > 1
+					Result.TexCoords[1] = Input.TexCoords1;
+				#endif
+				#if NUM_MATERIAL_TEXCOORDS_VERTEX > 2
+					Result.TexCoords[2] = Input.TexCoords2;
+				#endif
+				#if NUM_MATERIAL_TEXCOORDS_VERTEX > 3
+					Result.TexCoords[3] = Input.TexCoords3;
+				#endif
+				#if NUM_MATERIAL_TEXCOORDS_VERTEX > 4
+					Result.TexCoords[4] = Input.TexCoords4;
+				#endif
+				#if NUM_MATERIAL_TEXCOORDS_VERTEX > 5
+					Result.TexCoords[5] = Input.TexCoords5;
+				#endif
+				#if NUM_MATERIAL_TEXCOORDS_VERTEX > 6
+					Result.TexCoords[6] = Input.TexCoords6;
+				#endif
+				#if NUM_MATERIAL_TEXCOORDS_VERTEX > 7
+					Result.TexCoords[7] = Input.TexCoords7;
+				#endif
+			#else
+				#if NUM_MATERIAL_TEXCOORDS_VERTEX > 1
+					UNROLL
+					for(int CoordinateIndex = 0; CoordinateIndex < NUM_MATERIAL_TEXCOORDS_VERTEX-1; CoordinateIndex+=2)
 					{
-						Result.TexCoords[CoordinateIndex+1] = Input.PackedTexCoords4[CoordinateIndex/2].zw;
+						Result.TexCoords[CoordinateIndex] = Input.PackedTexCoords4[CoordinateIndex/2].xy;
+						if( CoordinateIndex+1 < NUM_MATERIAL_TEXCOORDS_VERTEX )
+						{
+							Result.TexCoords[CoordinateIndex+1] = Input.PackedTexCoords4[CoordinateIndex/2].zw;
+						}
 					}
-				}
-			#endif	// NUM_MATERIAL_TEXCOORDS_VERTEX > 1
-			#if NUM_MATERIAL_TEXCOORDS_VERTEX % 2 == 1
-				Result.TexCoords[NUM_MATERIAL_TEXCOORDS_VERTEX-1] = Input.PackedTexCoords2;
-			#endif	// NUM_MATERIAL_TEXCOORDS_VERTEX % 2 == 1
+				#endif	// NUM_MATERIAL_TEXCOORDS_VERTEX > 1
+				#if NUM_MATERIAL_TEXCOORDS_VERTEX % 2 == 1
+					Result.TexCoords[NUM_MATERIAL_TEXCOORDS_VERTEX-1] = Input.PackedTexCoords2;
+				#endif	// NUM_MATERIAL_TEXCOORDS_VERTEX % 2 == 1
+			#endif
 		#endif
 #endif  //MANUAL_VERTEX_FETCH && NUM_MATERIAL_TEXCOORDS_VERTEX
 
diff --git a/Engine/Shaders/Private/VertexFactoryCommon.ush b/Engine/Shaders/Private/VertexFactoryCommon.ush
--- a/Engine/Shaders/Private/VertexFactoryCommon.ush
+++ b/Engine/Shaders/Private/VertexFactoryCommon.ush
@@ -3,6 +3,8 @@
 #include "Common.ush"
 #include "SceneData.ush"
 
+#define PACK_TEX_COORDS !VULKAN_PROFILE
+
 float3 TransformLocalToWorld(float3 LocalPosition, float4x4 LocalToWorld)
 {
 	//expand the matrix multiply to avoid a code-gen bug on PS4.
