Implement simple subsurface shading model for virtual shadow maps and related fixes and improvements:
- Add VirtualShadowMapId to forward light data, removing the per-view remapping table. Should fixe a few multi-view/split-screen bugs.
- Minor cleanup of PCF subsurface path; remove dead/broken code.
- Fix up blending of light attenuation into screen shadow mask; disable CSM fading when VSMs are enabled.
- Fix OnePassProjection flag when VSMs are disabled

diff --git a/Engine/Shaders/Private/ClusteredDeferredShadingPixelShader.usf b/Engine/Shaders/Private/ClusteredDeferredShadingPixelShader.usf
--- a/Engine/Shaders/Private/ClusteredDeferredShadingPixelShader.usf
+++ b/Engine/Shaders/Private/ClusteredDeferredShadingPixelShader.usf
@@ -22,6 +22,7 @@
 #include "ShadingModels.ush"
 #include "LightGridCommon.ush"
 #include "DeferredLightingCommon.ush"
+#include "VirtualShadowMaps/TransmissionCommon.ush"
 
 #if USE_HAIR_LIGHTING == 1
 	#include "HairStrands/HairStrandsCommon.ush"
@@ -59,7 +60,9 @@ float4 GetLightGridLocalLighting(const FCulledLightsGridData InLightGridData, fl
 	// This prevents GPU hangs when the PS tries to read from uninitialized NumCulledLightsGrid buffer
 	const uint NumLocalLights = min(InLightGridData.NumLocalLights, GetNumLocalLights(EyeIndex));
 
-	uint ShadowMask = ~ShadowMaskBits[ int2( ScreenUV.xy * View.BufferSizeAndInvSize.xy ) ];
+	int2 PixelPos = int2( ScreenUV.xy * View.BufferSizeAndInvSize.xy );
+	uint ShadowMask = ~ShadowMaskBits[ PixelPos ];
+	half SubsurfaceOpacity = GetSubsurfaceOpacityFromGbuffer( PixelPos );
 
 	LOOP
 	for (uint LocalLightListIndex = FirstNonSimpleLightIndex; LocalLightListIndex < NumLocalLights; LocalLightListIndex++)
@@ -120,15 +123,21 @@ float4 GetLightGridLocalLighting(const FCulledLightsGridData InLightGridData, fl
 		// Static shadowing uses ShadowMapChannel, dynamic shadows are packed into light attenuation using PreviewShadowMapChannel
 		LightData.ShadowMapChannelMask = UnpackShadowMapChannelMask(LightTypeAndPackedShadowMapChannelMask);
 
-		// This matches tiled deferred, since there is no dynamic shadow on the lights used for clustered shading.
-		float4 LightAttenuation = 1.0f;
+		float4 LightAttenuation = float4(1, 1, 1, 1);
 
 		LightData.ShadowedBits = 1;
 		// Be careful with the shift here as the shift amount is mod 32, so >>32 is a noop rather than zero
 		// We treat lights past the 10th as unshadowed
 		if (LocalLightListIndex < GetShadowMaskMaxLightCount())
 		{
-			LightAttenuation = UnpackShadowMask(ShadowMask, LocalLightListIndex);
+			LightAttenuation = UnpackShadowMask(ShadowMask, LocalLightListIndex).xxxx;
+		}
+
+		if ( SubsurfaceOpacity < 1.0f )
+		{
+			// Replace SSS portion for subsurface materials
+			half SSSTransmission = ComputeSimpleSubsurfaceTransmissionFromVirtualShadowMap(LocalLight.VirtualShadowMapId, WorldPosition, SubsurfaceOpacity);
+			LightAttenuation.yw = SSSTransmission;
 		}
 
 		LightData.HairTransmittance = InitHairTransmittanceData();
diff --git a/Engine/Shaders/Private/LightGridCommon.ush b/Engine/Shaders/Private/LightGridCommon.ush
--- a/Engine/Shaders/Private/LightGridCommon.ush
+++ b/Engine/Shaders/Private/LightGridCommon.ush
@@ -177,6 +177,9 @@ struct FLocalLightData
 	bool bClusteredDeferredSupported; 
 	/** Flag is true if it is a simple light. They are always first / together.*/
 	bool bIsSimpleLight;
+	/** Virtual shadow map ID or INDEX_NONE if not present. */
+	int VirtualShadowMapId;
+
 };
 
 FLocalLightData GetLocalLightData(uint GridIndex, uint EyeIndex)
@@ -198,6 +201,7 @@ FLocalLightData GetLocalLightData(uint GridIndex, uint EyeIndex)
 		Result.SpotAnglesAndSourceRadiusPacked = ForwardLightData.ForwardLocalLightBuffer[LocalLightBaseIndex + 3];
 		Result.LightTangentAndSoftSourceRadius = ForwardLightData.ForwardLocalLightBuffer[LocalLightBaseIndex + 4];
 		Result.RectBarnDoor = ForwardLightData.ForwardLocalLightBuffer[LocalLightBaseIndex + 5];
+		Result.VirtualShadowMapId = asint(ForwardLightData.ForwardLocalLightBuffer[LocalLightBaseIndex + 5].z);
 #if INSTANCED_STEREO
 	}
 	else
@@ -211,6 +215,8 @@ FLocalLightData GetLocalLightData(uint GridIndex, uint EyeIndex)
 		Result.LightDirectionAndShadowMask = ForwardLightDataISR.ForwardLocalLightBuffer[LocalLightBaseIndex + 2];
 		Result.SpotAnglesAndSourceRadiusPacked = ForwardLightDataISR.ForwardLocalLightBuffer[LocalLightBaseIndex + 3];
 		Result.LightTangentAndSoftSourceRadius = ForwardLightDataISR.ForwardLocalLightBuffer[LocalLightBaseIndex + 4];
+		Result.RectBarnDoor = ForwardLightData.ForwardLightDataISR[LocalLightBaseIndex + 5];
+		Result.VirtualShadowMapId = asint(ForwardLightData.ForwardLightDataISR[LocalLightBaseIndex + 5].z);
 	}
 #endif
 
@@ -242,4 +248,4 @@ void PackShadowMask(inout uint OutShadowMask, float InShadowFactor, uint InLight
 float UnpackShadowMask(uint InShadowMask, uint InLightIndex)
 {
 	return ((InShadowMask >> (InLightIndex * 3u)) & 7u) / 7.0;
-}
\ No newline at end of file
+}
diff --git a/Engine/Shaders/Private/ShadowProjectionPixelShader.usf b/Engine/Shaders/Private/ShadowProjectionPixelShader.usf
--- a/Engine/Shaders/Private/ShadowProjectionPixelShader.usf
+++ b/Engine/Shaders/Private/ShadowProjectionPixelShader.usf
@@ -282,77 +282,58 @@ void Main(
 		BRANCH
 		if (bIsSubsurfaceCompatible && bHasSubsurface)
 		{
-		#if STRATA_ENABLED
-			const FStrataSubsurfaceData SSSData = StrataUnpackSSSData(Strata.SSSTexture.Load(uint3(SVPos.xy, 0)));
-
-			float Opacity = 1;
 			float Density = 1;
-			if (!bIsEyeOrHair && SSSData.bIsValid)
-			{
-				Opacity = SSSData.bIsProfile ? SSSData.ProfileRadiusScale : 1.0f;
-				Density = SubsurfaceDensityFromOpacity(Opacity);
-			}
-		#else // STRATA_ENABLED
-			float Opacity = GBufferData.CustomData.a;
-			float Density = SubsurfaceDensityFromOpacity(Opacity);
-			if (bIsEyeOrHair)
+			if (!bIsEyeOrHair)
 			{
-				Opacity = 1;
-				Density = 1;
-			}
-		#endif // STRATA_ENABLED
-
-			//@todo - Use a larger kernel filter based on material opacity
-			// This is not being done currently because a varying opacity causes discontinuities in the shadow result
-			// Note: Snapping based on the largest size kernel is done in FProjectedShadowInfo::FProjectedShadowInfo, MaxDownsampleFactor is set to 4
-			// Only factors of MaxDownsampleFactor will be stable
-			float SquareRootFilterScale = lerp(1.999f, 0, Opacity);
-			int SquareRootFilterScaleInt = int(SquareRootFilterScale) + 1;
-
-			#if UNFILTERED_SHADOW_PROJECTION
-				float ShadowMapDepth = Texture2DSampleLevel(ShadowDepthTexture, ShadowDepthTextureSampler, ShadowPosition.xy, 0).x;
-				SSSTransmission = CalculateSubsurfaceOcclusion(Density, LightSpacePixelDepthForSSS, ShadowMapDepth.xxx).x;
-			#else
-		
-				// default code path
-				FPCFSamplerSettings Settings;
-				
-				Settings.ShadowDepthTexture = ShadowDepthTexture;
-				Settings.ShadowDepthTextureSampler = ShadowDepthTextureSampler;
-				Settings.ShadowBufferSize = ShadowBufferSize;
-				Settings.TransitionScale = SoftTransitionScale.z;
-				Settings.SceneDepth = LightSpacePixelDepthForSSS + ProjectionDepthBiasParameters.x;
-				Settings.bSubsurface = true;
-				Settings.bTreatMaxDepthUnshadowed = false;
-				Settings.DensityMulConstant = Density * ProjectionDepthBiasParameters.w;
-				Settings.ProjectionDepthBiasParameters = ProjectionDepthBiasParameters.xw;
-
-			  #if USE_TRANSMISSION
-				if (bHasSubsurfaceTransmission)
+			#if STRATA_ENABLED
+				const FStrataSubsurfaceData SSSData = StrataUnpackSSSData(Strata.SSSTexture.Load(uint3(SVPos.xy, 0)));
+				if (SSSData.bIsValid)
 				{
-				#if STRATA_ENABLED
-					const FStrataTopLayerData TopLayerData = StrataUnpackTopLayerData(Strata.TopLayerNormalTexture.Load(uint3(SVPos.xy, 0)));
-					const float3 WorldNormal = TopLayerData.WorldNormal;
-
-					// For non-profile version, use default transmission parameters (only NormalScale & ExtinctionScale are used)
-					FTransmissionProfileParams TransmissionParams = InitTransmissionProfileParams();
-					if (StrataData.bHasSubsurfaceProfile)
-					{
-						TransmissionParams = GetTransmissionProfileParams(SSSData.ProfileId);
-					}
-				#else // STRATA_ENABLED
-					float3 WorldNormal = GBufferData.WorldNormal;
-					FTransmissionProfileParams TransmissionParams = GetTransmissionProfileParams(GBufferData);
-				#endif // STRATA_ENABLED
-					SSSTransmission = CalcTransmissionThickness(ScreenPosition.xyz, LightPositionAndInvRadius.xyz, WorldNormal, TransmissionParams, Settings);
+					float Opacity = SSSData.bIsProfile ? SSSData.ProfileRadiusScale : 1.0f;
+					Density = SubsurfaceDensityFromOpacity(Opacity);
 				}
-				else
-			  #endif
+			#else // STRATA_ENABLED
+				Density = SubsurfaceDensityFromOpacity(GBufferData.CustomData.a);
+			#endif // STRATA_ENABLED
+			}
+								
+			FPCFSamplerSettings Settings;
+				
+			Settings.ShadowDepthTexture = ShadowDepthTexture;
+			Settings.ShadowDepthTextureSampler = ShadowDepthTextureSampler;
+			Settings.ShadowBufferSize = ShadowBufferSize;
+			Settings.TransitionScale = SoftTransitionScale.z;
+			Settings.SceneDepth = LightSpacePixelDepthForSSS + ProjectionDepthBiasParameters.x;
+			Settings.bSubsurface = true;
+			Settings.bTreatMaxDepthUnshadowed = false;
+			Settings.DensityMulConstant = Density * ProjectionDepthBiasParameters.w;
+			Settings.ProjectionDepthBiasParameters = ProjectionDepthBiasParameters.xw;
+
+			#if USE_TRANSMISSION
+			if (bHasSubsurfaceTransmission)
+			{
+			#if STRATA_ENABLED
+				const FStrataTopLayerData TopLayerData = StrataUnpackTopLayerData(Strata.TopLayerNormalTexture.Load(uint3(SVPos.xy, 0)));
+				const float3 WorldNormal = TopLayerData.WorldNormal;
+
+				// For non-profile version, use default transmission parameters (only NormalScale & ExtinctionScale are used)
+				FTransmissionProfileParams TransmissionParams = InitTransmissionProfileParams();
+				if (StrataData.bHasSubsurfaceProfile)
 				{
-					// ideally we use a larger filter kernel for SSSbut as Gather4 makes that harder
-					SSSTransmission = ManualPCF(ShadowPosition.xy, Settings);
+					TransmissionParams = GetTransmissionProfileParams(SSSData.ProfileId);
 				}
+			#else // STRATA_ENABLED
+				float3 WorldNormal = GBufferData.WorldNormal;
+				FTransmissionProfileParams TransmissionParams = GetTransmissionProfileParams(GBufferData);
+			#endif // STRATA_ENABLED
+				SSSTransmission = CalcTransmissionThickness(ScreenPosition.xyz, LightPositionAndInvRadius.xyz, WorldNormal, TransmissionParams, Settings);
+			}
+			else
 			#endif
+			{
+				// ideally we use a larger filter kernel for SSSbut as Gather4 makes that harder
+				SSSTransmission = ManualPCF(ShadowPosition.xy, Settings);
+			}
 		}
 
 	#endif
diff --git a/Engine/Shaders/Private/TransmissionThickness.ush b/Engine/Shaders/Private/TransmissionThickness.ush
--- a/Engine/Shaders/Private/TransmissionThickness.ush
+++ b/Engine/Shaders/Private/TransmissionThickness.ush
@@ -102,12 +102,12 @@ float CalcShadowLinearDepth(float shadowDepth, float2 param)
 	return 1.0f / (shadowDepth * param.x - param.y);
 }
 
-float CalculateThickness(float ShadowmapDepth, FPCFSamplerSettings Settings, float NormalScale, float NoL)
+float CalculateThickness(float ShadowmapDepth, float SceneDepth, float DensityMulConstant, float NormalScale, float NoL)
 {
 	// Determine the distance that the light traveled through the subsurface object
 	// This assumes that anything between this subsurface pixel and the light was also a subsurface material,
 	// As a result, subsurface materials receive leaked light based on how transparent they are
-	float Thickness = (Settings.SceneDepth - ShadowmapDepth) * Settings.DensityMulConstant;
+	float Thickness = (SceneDepth - ShadowmapDepth) * DensityMulConstant;
 	Thickness = clamp(abs(Thickness > 0 ? Thickness + NormalScale : max(0.0, Thickness * NoL + NormalScale)), 0.15, 5) + 0.25;
 	// Never shadow from depths that were never written to (max depth value)
 	return Thickness;
@@ -138,7 +138,7 @@ float CalcTransmissionThickness(float3 ViewPos, float3 LightPos, float3 WorldNor
 
 
 	ShrinkedScreenPos.xy = ShrinkedScreenPos.xy / ShrinkedScreenPos.w;
-	Settings.DensityMulConstant = DepthScale * ProjectionDepthBiasParameters.w;
+	float DensityMulConstant = DepthScale * ProjectionDepthBiasParameters.w;
 
 
 	Settings.SceneDepth = (1.0f - ShrinkedScreenPos.z) + ProjectionDepthBiasParameters.x;
@@ -153,7 +153,7 @@ float CalcTransmissionThickness(float3 ViewPos, float3 LightPos, float3 WorldNor
 		float2 DeltUv = PoissonDisk[i] * ShadowFilterRadius;
 		float2 ShadowUv = ShrinkedScreenPos.xy + DeltUv;
 		float ShadowmapDepth = Texture2DSampleLevel(Settings.ShadowDepthTexture, Settings.ShadowDepthTextureSampler, ShadowUv.xy, 0).r;
-		float Thickness = CalculateThickness(ShadowmapDepth, Settings, NormalScale, NoL);
+		float Thickness = CalculateThickness(ShadowmapDepth, Settings.SceneDepth, DensityMulConstant, NormalScale, NoL);
 		ThicknessSum += (Thickness);
 	}
 	float Thickness = ThicknessSum / SampleNum;
@@ -195,15 +195,9 @@ float CalcTransmissionThickness(float3 WorldNormal, FTransmissionProfileParams T
 		float CompareDistance = ShadowPosition.z / ShadowPosition.w;
 		float ShadowDepthBias = -DepthBias / ShadowPosition.w;
 		
-		FPCFSamplerSettings Settings;
-
-		Settings.TransitionScale = SoftTransitionScale.z;
-		Settings.bSubsurface = true;
 		//Shadow Depth range Near to far, so no need mul far here.
-		Settings.DensityMulConstant = TransmissionParams.ExtinctionScale * (10.0f / LightInvRadius);
-		Settings.ProjectionDepthBiasParameters = ProjectionDepthBiasParameters.xw;
-		Settings.SceneDepth = 0.0f;
-		Settings.SceneDepth = CalcShadowLinearDepth(CompareDistance + ShadowDepthBias, ProjParam) * LightInvRadius;
+		float DensityMulConstant = TransmissionParams.ExtinctionScale * (10.0f / LightInvRadius);
+		float SceneDepth = CalcShadowLinearDepth(CompareDistance + ShadowDepthBias, ProjParam) * LightInvRadius;
 
 		Thickness = 0;
 
@@ -217,7 +211,7 @@ float CalcTransmissionThickness(float3 WorldNormal, FTransmissionProfileParams T
             float ShadowSample = TextureCubeSampleDepthLevel(ShadowDepthCubeTexture, ShadowDepthTextureSampler, SamplePos, 0.0);
 	#endif
 			float ShadowDepth = CalcShadowLinearDepth(ShadowSample, ProjParam) * LightInvRadius;
-			Thickness += CalculateThickness(ShadowDepth, Settings, NormalScale, NoL);
+			Thickness += CalculateThickness(ShadowDepth, SceneDepth, DensityMulConstant, NormalScale, NoL);
 		}
 		Thickness /= 5;
 
diff --git a/Engine/Shaders/Private/VirtualShadowMaps/PageManagement.usf b/Engine/Shaders/Private/VirtualShadowMaps/PageManagement.usf
--- a/Engine/Shaders/Private/VirtualShadowMaps/PageManagement.usf
+++ b/Engine/Shaders/Private/VirtualShadowMaps/PageManagement.usf
@@ -47,8 +47,8 @@ void StatsBufferInterlockedAdd(uint Index, uint Value)
 RWStructuredBuffer<uint> OutPageRequestFlags;
 
 Texture2D<UlongType> VisBuffer64;
-StructuredBuffer<int> VirtualShadowMapIdRemap;
-uint NumDirectionalLightSmInds;
+StructuredBuffer<int> DirectionalLightIds;
+
 uint bPostBasePass;
 float ResolutionLodBiasLocal;
 float PageDilationBorderSize;
@@ -256,9 +256,9 @@ void GeneratePageFlagsFromPixels(
 	}
 	
 	// Directional lights
-	for (uint Index = 0; Index < NumDirectionalLightSmInds; ++Index)
+	for (uint Index = 0; Index < VirtualShadowMap.NumDirectionalLights; ++Index)
 	{
-		int ClipmapStartId = VirtualShadowMapIdRemap[Index];
+		int ClipmapStartId = DirectionalLightIds[Index];
 		FVirtualShadowMapProjectionShaderData BaseProjectionData = GetVirtualShadowMapProjectionData(ClipmapStartId);
 
 		// Possibly use a single shadow projection data for the whole clipmap with just scales/biases for
@@ -285,15 +285,14 @@ void GeneratePageFlagsFromPixels(
 			const FLocalLightData LightData = GetLocalLightData(CulledLightGridData.DataStartIndex + Index, EyeIndex);
 			const uint LightGridLightIndex = ForwardLightData.CulledLightDataGrid[CulledLightGridData.DataStartIndex + Index];
 
-			float3 ToLight = normalize(LightData.LightPositionAndInvRadius.xyz - WorldPosition);
 			// Also do precise cone test, since froxels are pretty coarse at times.
+			float3 ToLight = normalize(LightData.LightPositionAndInvRadius.xyz - WorldPosition);
 			if (dot(ToLight, LightData.LightDirectionAndShadowMask.xyz) < LightData.SpotAnglesAndSourceRadiusPacked.x)
 			{
 				continue;
 			}
 
-			// The Virtual Shadow Remap stores directional lights first
-			int VirtualShadowMapId = VirtualShadowMapIdRemap[NumDirectionalLightSmInds + LightGridLightIndex];
+			int VirtualShadowMapId = LightData.VirtualShadowMapId;
 			if (VirtualShadowMapId != INDEX_NONE)
 			{
 				bool bSpotLight = LightData.SpotAnglesAndSourceRadiusPacked.x > -2.0f;
diff --git a/Engine/Shaders/Private/VirtualShadowMaps/TransmissionCommon.ush b/Engine/Shaders/Private/VirtualShadowMaps/TransmissionCommon.ush
new file mode 100644
--- /dev/null
+++ b/Engine/Shaders/Private/VirtualShadowMaps/TransmissionCommon.ush
@@ -0,0 +1,85 @@
+// Copyright Epic Games, Inc. All Rights Reserved.
+
+/*=============================================================================
+VirtualShadowMaps/TransmissionCommon.ush:
+=============================================================================*/
+#pragma once
+
+#include "../Common.ush"
+#include "../SceneTexturesCommon.ush"
+#include "../DeferredShadingCommon.ush"
+#include "ProjectionCommon.ush"
+
+// TODO: Support strata
+#ifndef STRATA_ENABLED
+#define STRATA_ENABLED 0
+#endif
+
+// TODO: Once strata is set up this function should be shared with ShadowProjectionPixelShader
+// and possibly also the contact shadow trace
+half GetSubsurfaceOpacityFromGbuffer(uint2 PixelPos)
+{
+	half Opacity = 1.0;
+
+#if STRATA_ENABLED
+	// TODO: VSM projection shaders are not set up for strata yet, but added a skeleton of the logic
+	const FStrataClassification StrataData = StrataUnpackClassificationData(Strata.ClassificationTexture.Load(PixelPos));
+	if (StrataData.bIsStrataMaterial)
+	{
+		if (StrataData.bHasSubsurface)
+		{
+			const FStrataSubsurfaceData SSSData = StrataUnpackSSSData(Strata.SSSTexture.Load(uint3(SVPos.xy, 0)));
+			if (SSSData.IsValid())
+			{
+				Opacity = SSSData.bIsProfile ? SSSData.ProfileRadiusScale : 1.0;
+			}
+		}
+	}
+#else // STRATA_ENABLED
+	FGBufferData GBufferData = GetGBufferDataUint(PixelPos);
+
+	// TODO: SUBSURFACE_PROFILE with fancy transmission
+	if (GBufferData.ShadingModelID == SHADINGMODELID_SUBSURFACE ||
+		GBufferData.ShadingModelID == SHADINGMODELID_PREINTEGRATED_SKIN ||
+		GBufferData.ShadingModelID == SHADINGMODELID_TWOSIDED_FOLIAGE)
+	{
+		Opacity = GBufferData.CustomData.a;
+	}
+#endif // STRATA_ENABLED
+
+	return Opacity;
+}
+
+/**
+* Sample VSM to determine approximate thickness and compute simple subsurface transmission matching the
+* PCF path.
+*/
+half ComputeSimpleSubsurfaceTransmissionFromVirtualShadowMap(
+	int VirtualShadowMapId, 
+	float3 WorldPosition,
+	half SubsurfaceOpacity)
+{
+	// If the receiver is a subsurface material, we do a very simple transmission approximation
+	// This matches the logic in the PCF kernel (ShadowFilteringCommon.ush) and the contact
+	// shadows trace (DeferredLightingCommon.ush). Ideally these should all be unified in a single location.
+	// Could optimize this in the case where we are doing single-sample VSM lookup anyways, but we would need
+	// exclude contact shadows overwriting the occluder distance. Single sample direct shadows are not currently
+	// a common case.
+
+	// Do a separate unfiltered lookup to get approximate occluder distance
+	FVirtualShadowMapSampleResult SSSResult = SampleVirtualShadowMap(
+		VirtualShadowMapId,
+		WorldPosition,
+		0.0f // RayStartDistance
+	);
+
+	if ( SSSResult.bValid && SSSResult.ShadowFactor < 1.0f )
+	{
+		float Density = SubsurfaceDensityFromOpacity( SubsurfaceOpacity );
+		float Occlusion = 1.0f - saturate( exp( -Density * SSSResult.OccluderDistance ) );
+		// NOTE: This square is arbitrary, but we have to be consistent with the PCF path
+		return Square( 1.0f - Occlusion * ( 1.0f - SSSResult.ShadowFactor ) );
+	}
+
+	return 1.0f;
+}
diff --git a/Engine/Shaders/Private/VirtualShadowMaps/VirtualShadowMapProjection.usf b/Engine/Shaders/Private/VirtualShadowMaps/VirtualShadowMapProjection.usf
--- a/Engine/Shaders/Private/VirtualShadowMaps/VirtualShadowMapProjection.usf
+++ b/Engine/Shaders/Private/VirtualShadowMaps/VirtualShadowMapProjection.usf
@@ -20,6 +20,7 @@
 #include "ProjectionDirectional.ush"
 #include "ProjectionSpot.ush"
 #include "PageCacheCommon.ush"
+#include "TransmissionCommon.ush"
 
 #if HAS_HAIR_STRANDS
 #define VOXEL_TRAVERSAL_TYPE VOXEL_TRAVERSAL_LINEAR_MIPMAP
@@ -38,10 +39,9 @@
 
 #define CONTACT_SHADOW_SAMPLES 8
 
-float4 EncodeLightAttenuationFromMask(float ShadowMask)
+float4 EncodeLightAttenuationFromMask(float ShadowMask, float SSSTransmission)
 {
 	const float ShadowFadeFraction = 1;
-	float SSSTransmission = ShadowMask;
 
 	// 0 is shadowed, 1 is unshadowed
 	// RETURN_COLOR not needed unless writing to SceneColor;
@@ -54,14 +54,14 @@ float4 EncodeLightAttenuationFromMask(float ShadowMask)
 	return EncodeLightAttenuation(half4(FadedShadow, FadedSSSShadow, FadedShadow, FadedSSSShadow));
 }
 
-Texture2D<float> InputShadowFactor;
+Texture2D<float2> InputShadowFactor;
 
 void VirtualShadowMapCompositePS(
 	in float4 SvPosition : SV_Position,
 	out float4 OutShadowMask : SV_Target)
 {
-	float ShadowFactor = InputShadowFactor.Load(int3(SvPosition.xy, 0));	
-	OutShadowMask = EncodeLightAttenuationFromMask(ShadowFactor);
+	float2 ShadowFactor = InputShadowFactor.Load(int3(SvPosition.xy, 0));
+	OutShadowMask = EncodeLightAttenuationFromMask(ShadowFactor.x, ShadowFactor.y);
 }
 
 
@@ -142,14 +142,12 @@ float SMRTRayLengthScale;			// Directional lights
 float SMRTCotMaxRayAngleFromLight;	// Spot/point lights
 
 // One pass projection
-StructuredBuffer< int > VirtualShadowMapIdRemap;
-uint NumDirectionalLightSmInds;
 RWTexture2D< uint > RWShadowMaskBits;
 
 // Single light per pass
 // Light parameters loaded via GetRootLightShaderParameters();
 int LightUniformVirtualShadowMapId;
-RWTexture2D< float4 > RWShadowFactor;
+RWTexture2D< float2 > RWShadowFactor;
 
 // Debug output pass
 StructuredBuffer< FPhysicalPageMetaData > PhysicalPageMetaData;
@@ -170,6 +168,7 @@ FVirtualShadowMapSampleResult ProjectLight(
 	float ContactShadowLengthWorld,
 	float3 TranslatedWorldPosition,
 	half3 GBufferNormal,
+	half SubsurfaceOpacity,
 	const half Noise)
 {
 	const float3 BaseWorldPosition = TranslatedWorldPosition - View.PreViewTranslation;
@@ -249,12 +248,16 @@ FVirtualShadowMapSampleResult ProjectLight(
 
 		// Do not run contact shadow when computing the hair shadow mask (i.e. shadow mask applied on hair, has the scene 
 		// depth buffer contains fully opaque depth, which create false positive intersection resulting in wrong self shadowing)
-		if (ContactShadowLengthWorld > 0.0f && !bIsHair)
+		if (!bIsHair)
 		{
-			half ContactShadowHitDistance = VirtualShadowMapScreenRayCast(TranslatedWorldPosition, L, ContactShadowLengthWorld, CONTACT_SHADOW_SAMPLES, Noise);
-			if (ContactShadowHitDistance >= 0.0f)
+			if (ContactShadowLengthWorld > 0.0f)
 			{
-				Result.ShadowFactor = 0.0f;
+				half ContactShadowHitDistance = VirtualShadowMapScreenRayCast(TranslatedWorldPosition, L, ContactShadowLengthWorld, CONTACT_SHADOW_SAMPLES, Noise);
+				if (ContactShadowHitDistance >= 0.0f)
+				{
+					Result.ShadowFactor = 0.0f;
+					Result.OccluderDistance = ContactShadowHitDistance;
+				}
 			}
 		}
 
@@ -386,12 +389,14 @@ void OutputDebug( uint2 PixelPos, FVirtualShadowMapSampleResult Result )
 		else if ( DebugOutputType == 6 )	// OccluderDistance
 		{
 			Output = ColorMapViridis( 0.001f * Result.OccluderDistance );
+			//Output = Result.OccluderDistance > 1000.0f ? float3(1, 1, 1) : ColorMapViridis( 0.001f * Result.OccluderDistance );
 		}
 	}
 
 	RWDebug[ PixelPos ] = float4( Output, 1.0f );
 }
 
+
 [numthreads(8, 8, 1)]
 void VirtualShadowMapProjection(
 	uint3	GroupId				: SV_GroupID,
@@ -408,8 +413,9 @@ void VirtualShadowMapProjection(
 	}
 	
 	float DeviceZ = SceneTexturesStruct.SceneDepthTexture.Load( int3( PixelPos, 0 ) ).r;
+	const bool bIsHair = InputType == INPUT_TYPE_HAIRSTRANDS;
 	#if HAS_HAIR_STRANDS
-	if (InputType == INPUT_TYPE_HAIRSTRANDS)
+	if (bIsHair)
 	{
 		DeviceZ = HairStrands.HairOnlyDepthTexture.Load(int3(PixelPos, 0));
 		if (DeviceZ == 0)
@@ -427,6 +433,7 @@ void VirtualShadowMapProjection(
 	const half Noise = InterleavedGradientNoise( SvPosition.xy, View.StateFrameIndexMod8 );
 
 	const half3 GBufferNormal = GetGBufferDataUint( PixelPos, true ).WorldNormal;
+	const half SubsurfaceOpacity = bIsHair ? 1.0f : GetSubsurfaceOpacityFromGbuffer( PixelPos );
 
 	FVirtualShadowMapSampleResult DebugResult = InitVirtualShadowMapSampleResult();
 	bool bWriteDebug = false;
@@ -444,13 +451,11 @@ void VirtualShadowMapProjection(
 	LOOP
 	for (uint Index = 0; Index < LightCount; ++Index)
 	{
-		// The Virtual Shadow Remap stores directional lights first
-		const uint LocalLightIndex = ForwardLightData.CulledLightDataGrid[CulledLightGridData.DataStartIndex + Index];
-		int VirtualShadowMapId = VirtualShadowMapIdRemap[VirtualShadowMap.NumDirectionalLights + LocalLightIndex];
+		const FLocalLightData LightData = GetLocalLightData( CulledLightGridData.DataStartIndex + Index, EyeIndex );
+		int VirtualShadowMapId = LightData.VirtualShadowMapId;
 
 		if (VirtualShadowMapId != INDEX_NONE)
 		{
-			const FLocalLightData LightData = GetLocalLightData( CulledLightGridData.DataStartIndex + Index, EyeIndex );
 			FLightShaderParameters Light = ConvertFromLocal( LightData );
 
 			FVirtualShadowMapSampleResult Result = ProjectLight(
@@ -461,6 +466,7 @@ void VirtualShadowMapProjection(
 				ContactShadowLengthWorld,
 				TranslatedWorldPosition,
 				GBufferNormal,
+				SubsurfaceOpacity,
 				Noise);
 
 			PackShadowMask(ShadowMask, Result.ShadowFactor, Index);
@@ -488,9 +494,19 @@ void VirtualShadowMapProjection(
 			ContactShadowLengthWorld,
 			TranslatedWorldPosition,
 			GBufferNormal,
+			SubsurfaceOpacity,
 			Noise);
+		
+		half SSSTransmission = Result.ShadowFactor;
+		if ( SubsurfaceOpacity < 1.0f )
+		{
+			SSSTransmission = ComputeSimpleSubsurfaceTransmissionFromVirtualShadowMap(
+				VirtualShadowMapId,
+				TranslatedWorldPosition - View.PreViewTranslation,
+				SubsurfaceOpacity);
+		}
 
-		RWShadowFactor[ PixelPos ] = Result.ShadowFactor;
+		RWShadowFactor[ PixelPos ] = half2( Result.ShadowFactor, SSSTransmission );
 
 		if (DebugVirtualShadowMapId == VirtualShadowMapId)
 		{
diff --git a/Engine/Source/Runtime/Engine/Private/Components/DirectionalLightComponent.cpp b/Engine/Source/Runtime/Engine/Private/Components/DirectionalLightComponent.cpp
--- a/Engine/Source/Runtime/Engine/Private/Components/DirectionalLightComponent.cpp
+++ b/Engine/Source/Runtime/Engine/Private/Components/DirectionalLightComponent.cpp
@@ -483,6 +483,12 @@ public:
 
 	virtual FVector2D GetDirectionalLightDistanceFadeParameters(ERHIFeatureLevel::Type InFeatureLevel, bool bPrecomputedLightingIsValid, int32 MaxNearCascades) const override
 	{
+		// Virtual shadow maps do not use cascade distance fading
+		if (UseVirtualShadowMaps())
+		{
+			return FVector2D(0, 0);
+		}
+
 		float FarDistance = GetCSMMaxDistance(bPrecomputedLightingIsValid, MaxNearCascades);
 		{
 			if (ShouldCreateRayTracedCascade(InFeatureLevel, bPrecomputedLightingIsValid, MaxNearCascades))
diff --git a/Engine/Source/Runtime/Engine/Private/Components/LightComponent.cpp b/Engine/Source/Runtime/Engine/Private/Components/LightComponent.cpp
--- a/Engine/Source/Runtime/Engine/Private/Components/LightComponent.cpp
+++ b/Engine/Source/Runtime/Engine/Private/Components/LightComponent.cpp
@@ -317,6 +317,7 @@ FLightSceneProxy::FLightSceneProxy(const ULightComponent* InLightComponent)
 	, bUsedAsAtmosphereSunLight(InLightComponent->IsUsedAsAtmosphereSunLight())
 	, bAffectDynamicIndirectLighting(InLightComponent->bAffectDynamicIndirectLighting)
 	, bUseRayTracedDistanceFieldShadows(InLightComponent->bUseRayTracedDistanceFieldShadows)
+	, bUseVirtualShadowMaps(false)	// See below
 	, bCastModulatedShadows(false)
 	, bUseWholeSceneCSMForMovableObjects(false)
 	, bTiledDeferredLightingSupported(false)
@@ -339,6 +340,9 @@ FLightSceneProxy::FLightSceneProxy(const ULightComponent* InLightComponent)
 {
 	check(SceneInterface);
 
+	// Currently we use virtual shadows maps for all lights when the global setting is enabled
+	bUseVirtualShadowMaps = ::UseVirtualShadowMaps(SceneInterface->GetShaderPlatform(), SceneInterface->GetFeatureLevel());
+
 	// Treat stationary lights as movable when non-nanite VSMs are enabled
 	const bool bNonNaniteVirtualShadowMaps = UseNonNaniteVirtualShadowMaps(SceneInterface->GetShaderPlatform(), SceneInterface->GetFeatureLevel());
 	if (bNonNaniteVirtualShadowMaps)
diff --git a/Engine/Source/Runtime/Engine/Public/SceneManagement.h b/Engine/Source/Runtime/Engine/Public/SceneManagement.h
--- a/Engine/Source/Runtime/Engine/Public/SceneManagement.h
+++ b/Engine/Source/Runtime/Engine/Public/SceneManagement.h
@@ -1516,6 +1516,7 @@ public:
 	inline bool AffectsTranslucentLighting() const { return bAffectTranslucentLighting; }
 	inline bool Transmission() const { return bTransmission; }
 	inline bool UseRayTracedDistanceFieldShadows() const { return bUseRayTracedDistanceFieldShadows; }
+	inline bool UseVirtualShadowMaps() const { return bUseVirtualShadowMaps; }
 	inline float GetRayStartOffsetDepthScale() const { return RayStartOffsetDepthScale; }
 	inline bool IsTiledDeferredLightingSupported() const { return bTiledDeferredLightingSupported;  }
 	inline uint8 GetLightType() const { return LightType; }
@@ -1726,6 +1727,9 @@ protected:
 	/** Whether to use ray traced distance field area shadows. */
 	const uint8 bUseRayTracedDistanceFieldShadows : 1;
 
+	/** Whether to use virtual shadow maps. */
+	uint8 bUseVirtualShadowMaps : 1;
+
 	/** Whether the light will cast modulated shadows when using the forward renderer (mobile). */
 	uint8 bCastModulatedShadows : 1;
 
diff --git a/Engine/Source/Runtime/Renderer/Private/ClusteredDeferredShadingPass.cpp b/Engine/Source/Runtime/Renderer/Private/ClusteredDeferredShadingPass.cpp
--- a/Engine/Source/Runtime/Renderer/Private/ClusteredDeferredShadingPass.cpp
+++ b/Engine/Source/Runtime/Renderer/Private/ClusteredDeferredShadingPass.cpp
@@ -23,6 +23,7 @@
 #include "ScenePrivate.h"
 #include "ShaderPrintParameters.h"
 #include "ShaderPrint.h"
+#include "VirtualShadowMaps/VirtualShadowMapArray.h"
 
 #include "SceneFilterRendering.h"
 #include "PostProcessing.h"
@@ -103,6 +104,7 @@ class FClusteredShadingPS : public FGlobalShader
 		SHADER_PARAMETER_RDG_UNIFORM_BUFFER(FHairStrandsViewUniformParameters, HairStrands)
 		SHADER_PARAMETER_RDG_UNIFORM_BUFFER(FSceneTextureUniformParameters, SceneTextures)
 		SHADER_PARAMETER_RDG_TEXTURE(Texture2D, ShadowMaskBits)
+		SHADER_PARAMETER_STRUCT_INCLUDE(FVirtualShadowMapSamplingParameters, VirtualShadowMapSamplingParameters)
 		SHADER_PARAMETER_RDG_BUFFER_SRV(Buffer<uint>, HairTransmittanceBuffer)
 
 		SHADER_PARAMETER_TEXTURE(Texture2D, LTCMatTexture)
@@ -127,6 +129,7 @@ class FClusteredShadingPS : public FGlobalShader
 	{
 		FGlobalShader::ModifyCompilationEnvironment(Parameters, OutEnvironment);
 		FForwardLightingParameters::ModifyCompilationEnvironment(Parameters.Platform, OutEnvironment);
+		FVirtualShadowMapArray::SetShaderDefines(OutEnvironment);
 	}
 };
 
@@ -145,6 +148,7 @@ static void InternalAddClusteredDeferredShadingPass(
 	const FSortedLightSetSceneInfo &SortedLightsSet,
 	EClusterPassInputType InputType,
 	FRDGTextureRef ShadowMaskBits,
+	FVirtualShadowMapArray& VirtualShadowMapArray,
 	FRDGBufferSRVRef HairTransmittanceBuffer)
 {
 	check(SortedLightsSet.ClusteredSupportedEnd > 0);
@@ -157,6 +161,7 @@ static void InternalAddClusteredDeferredShadingPass(
 	PassParameters->Forward = View.ForwardLightingResources->ForwardLightDataUniformBuffer;
 	PassParameters->SceneTextures = SceneTextures.UniformBuffer;
 	PassParameters->ShadowMaskBits = ShadowMaskBits;
+	PassParameters->VirtualShadowMapSamplingParameters = VirtualShadowMapArray.GetSamplingParameters(GraphBuilder);
 	PassParameters->HairTransmittanceBuffer = HairTransmittanceBuffer;
 
 	PassParameters->LTCMatTexture = GSystemTextures.LTCMat->GetShaderResourceRHI();
@@ -258,6 +263,7 @@ void FDeferredShadingSceneRenderer::AddClusteredDeferredShadingPass(
 				SortedLightsSet,
 				EClusterPassInputType::GBuffer,
 				ShadowMaskBits,
+				VirtualShadowMapArray,
 				nullptr);
 
 			if (HairStrands::HasViewHairStrandsData(View))
@@ -270,6 +276,7 @@ void FDeferredShadingSceneRenderer::AddClusteredDeferredShadingPass(
 					SortedLightsSet,
 					EClusterPassInputType::HairStrands,
 					HairStrandsShadowMaskBits,
+					VirtualShadowMapArray,
 					GraphBuilder.CreateSRV(TransmittanceMask.TransmittanceMask, FHairStrandsTransmittanceMaskData::Format));
 			}
 		}
diff --git a/Engine/Source/Runtime/Renderer/Private/LightGridInjection.cpp b/Engine/Source/Runtime/Renderer/Private/LightGridInjection.cpp
--- a/Engine/Source/Runtime/Renderer/Private/LightGridInjection.cpp
+++ b/Engine/Source/Runtime/Renderer/Private/LightGridInjection.cpp
@@ -431,6 +431,9 @@ void FSceneRenderer::ComputeLightGrid(FRDGBuilder& GraphBuilder, bool bCullLight
 					// pack light type in this uint32 as well
 					LightTypeAndShadowMapChannelMaskPacked |= SortedLightInfo.SortKey.Fields.LightType << 16;
 
+					const bool bDynamicShadows = ViewFamily.EngineShowFlags.DynamicShadows && VisibleLightInfos.IsValidIndex(LightSceneInfo->Id);
+					const int32 VirtualShadowMapId = bDynamicShadows ? VisibleLightInfos[LightSceneInfo->Id].GetVirtualShadowMapId( &View ) : INDEX_NONE;
+
 					if ((SortedLightInfo.SortKey.Fields.LightType == LightType_Point && ViewFamily.EngineShowFlags.PointLights) ||
 						(SortedLightInfo.SortKey.Fields.LightType == LightType_Spot && ViewFamily.EngineShowFlags.SpotLights) ||
 						(SortedLightInfo.SortKey.Fields.LightType == LightType_Rect && ViewFamily.EngineShowFlags.RectLights))
@@ -455,7 +458,7 @@ void FSceneRenderer::ComputeLightGrid(FRDGBuilder& GraphBuilder, bool bCullLight
 
 						LightData.LightTangentAndSoftSourceRadius = FVector4(LightParameters.Tangent, LightParameters.SoftSourceRadius);
 
-						LightData.RectBarnDoor = FVector4(LightParameters.RectLightBarnCosAngle, LightParameters.RectLightBarnLength, 0, 0);
+						LightData.RectBarnDoor = FVector4(LightParameters.RectLightBarnCosAngle, LightParameters.RectLightBarnLength, *((float*)&VirtualShadowMapId), 0);
 
 						float VolumetricScatteringIntensity = LightProxy->GetVolumetricScatteringIntensity();
 
@@ -509,11 +512,11 @@ void FSceneRenderer::ComputeLightGrid(FRDGBuilder& GraphBuilder, bool bCullLight
 
 							ForwardLightData.DirectionalLightDistanceFadeMAD = FVector2D(FadeParams.Y, -FadeParams.X * FadeParams.Y);
 
-							if (ViewFamily.EngineShowFlags.DynamicShadows && VisibleLightInfos.IsValidIndex(LightSceneInfo->Id) && VisibleLightInfos[LightSceneInfo->Id].AllProjectedShadows.Num() > 0)
+							if (bDynamicShadows)
 							{
 								const TArray<FProjectedShadowInfo*, SceneRenderingAllocator>& DirectionalLightShadowInfos = VisibleLightInfos[LightSceneInfo->Id].AllProjectedShadows;
 
-								ForwardLightData.DirectionalLightVSM = VisibleLightInfos[LightSceneInfo->Id].GetVirtualShadowMapId( &View );
+								ForwardLightData.DirectionalLightVSM = VirtualShadowMapId;
 
 								ForwardLightData.NumDirectionalLightCascades = 0;
 								// Unused cascades should compare > all scene depths
@@ -778,7 +781,9 @@ void FSceneRenderer::ComputeLightGrid(FRDGBuilder& GraphBuilder, bool bCullLight
 
 void FDeferredShadingSceneRenderer::GatherLightsAndComputeLightGrid(FRDGBuilder& GraphBuilder, bool bNeedLightGrid, FSortedLightSetSceneInfo &SortedLightSet)
 {
-	bool bShadowedLightsInClustered = ShouldUseClusteredDeferredShading() && CVarVirtualShadowOnePassProjection.GetValueOnRenderThread();
+	bool bShadowedLightsInClustered = ShouldUseClusteredDeferredShading()
+		&& CVarVirtualShadowOnePassProjection.GetValueOnRenderThread()
+		&& VirtualShadowMapArray.IsAllocated();
 
 	GatherAndSortLights(SortedLightSet, bShadowedLightsInClustered);
 	
diff --git a/Engine/Source/Runtime/Renderer/Private/ShadowRendering.cpp b/Engine/Source/Runtime/Renderer/Private/ShadowRendering.cpp
--- a/Engine/Source/Runtime/Renderer/Private/ShadowRendering.cpp
+++ b/Engine/Source/Runtime/Renderer/Private/ShadowRendering.cpp
@@ -2034,64 +2034,60 @@ void FSceneRenderer::RenderShadowProjections(
 				if (ScissorRect.Area() > 0)
 				{
 					// Project virtual shadow maps
-					FRDGTextureRef ShadowFactorTexture;
 					if (VisibleLightInfo.VirtualShadowMapClipmaps.Num() > 0)
 					{
-						ShadowFactorTexture = RenderVirtualShadowMapProjection(
+						RenderVirtualShadowMapProjection(
 							GraphBuilder,
 							SceneTextures,
 							View,
 							VirtualShadowMapArray,
 							ScissorRect,
 							EVirtualShadowMapProjectionInputType::GBuffer,
-							VisibleLightInfo.FindShadowClipmapForView(&View));
+							VisibleLightInfo.FindShadowClipmapForView(&View),
+							ScreenShadowMaskTexture);
 					}
 					else
 					{
 						check(VirtualShadowMaps.Num() == 1);
-						ShadowFactorTexture = RenderVirtualShadowMapProjection(
+						RenderVirtualShadowMapProjection(
 							GraphBuilder,
 							SceneTextures,
 							View,
 							VirtualShadowMapArray,
 							ScissorRect,
 							EVirtualShadowMapProjectionInputType::GBuffer,
-							VirtualShadowMaps[0]);
+							VirtualShadowMaps[0],
+							ScreenShadowMaskTexture);
 					}
 
-					// Composite into screen shadow mask
-					CompositeVirtualShadowMapMask(GraphBuilder, ScissorRect, ShadowFactorTexture, ScreenShadowMaskTexture);
-
 					// Sub-pixel shadow (no denoising for hair)
 					if (HairStrands::HasViewHairStrandsData(View) && ScreenShadowMaskSubPixelTexture)
 					{
-						FRDGTextureRef HairShadowFactorTexture;
 						if (VisibleLightInfo.VirtualShadowMapClipmaps.Num() > 0)
 						{
-							HairShadowFactorTexture = RenderVirtualShadowMapProjection(
+							RenderVirtualShadowMapProjection(
 								GraphBuilder,
 								SceneTextures,
 								View,
 								VirtualShadowMapArray,
 								ScissorRect,
 								EVirtualShadowMapProjectionInputType::HairStrands,
-								VisibleLightInfo.FindShadowClipmapForView(&View));
+								VisibleLightInfo.FindShadowClipmapForView(&View),
+								ScreenShadowMaskSubPixelTexture);
 						}
 						else
 						{
 							check(VirtualShadowMaps.Num() == 1);
-							HairShadowFactorTexture = RenderVirtualShadowMapProjection(
+							RenderVirtualShadowMapProjection(
 								GraphBuilder,
 								SceneTextures,
 								View,
 								VirtualShadowMapArray,
 								ScissorRect,
 								EVirtualShadowMapProjectionInputType::HairStrands,
-								VirtualShadowMaps[0]);
+								VirtualShadowMaps[0],
+								ScreenShadowMaskSubPixelTexture);
 						}
-
-						// Composite into sub pixel mask
-						CompositeVirtualShadowMapMask(GraphBuilder, ScissorRect, HairShadowFactorTexture, ScreenShadowMaskSubPixelTexture);
 					}
 				}
 			}
diff --git a/Engine/Source/Runtime/Renderer/Private/ShadowSetup.cpp b/Engine/Source/Runtime/Renderer/Private/ShadowSetup.cpp
--- a/Engine/Source/Runtime/Renderer/Private/ShadowSetup.cpp
+++ b/Engine/Source/Runtime/Renderer/Private/ShadowSetup.cpp
@@ -3593,7 +3593,10 @@ void FSceneRenderer::CreateWholeSceneProjectedShadow(
 					SizeX = SizeY = GetCubeShadowDepthZResolution(FeatureLevel, GetCubeShadowDepthZIndex(FeatureLevel, MaxDesiredResolution));
 				}
 
-				const bool bNeedsVirtualShadowMap = VirtualShadowMapArray.IsEnabled() && !ProjectedShadowInitializer.bRayTracedDistanceField;
+				const bool bNeedsVirtualShadowMap =
+					LightSceneInfo->Proxy->UseVirtualShadowMaps() &&
+					VirtualShadowMapArray.IsEnabled() &&
+					!ProjectedShadowInitializer.bRayTracedDistanceField;
 
 				int32 NumShadowMaps = 1;
 				EShadowDepthCacheMode CacheMode[2] = { SDCM_Uncached, SDCM_Uncached };
@@ -4512,7 +4515,10 @@ void FSceneRenderer::AddViewDependentWholeSceneShadowsForView(
 	// Note: it is possible for a non-directional light to set up a proxy that requests a view-dependent SM (or several)
 	// Unmodified UE does not use this, so the virtual SM ignores this path for the time being (more likely to be removed).
 	const bool bDirectionalLight = LightSceneInfo.Proxy->GetLightType() == LightType_Directional;
-	const bool bNeedsVirtualShadowMap = VirtualShadowMapArray.IsEnabled() && bDirectionalLight;
+	const bool bNeedsVirtualShadowMap =
+		LightSceneInfo.Proxy->UseVirtualShadowMaps() &&
+		VirtualShadowMapArray.IsEnabled() &&
+		bDirectionalLight;
 
 	// Allow each view to create a whole scene view dependent shadow
 	for (int32 ViewIndex = 0; ViewIndex < Views.Num(); ViewIndex++)
diff --git a/Engine/Source/Runtime/Renderer/Private/VirtualShadowMaps/VirtualShadowMapArray.cpp b/Engine/Source/Runtime/Renderer/Private/VirtualShadowMaps/VirtualShadowMapArray.cpp
--- a/Engine/Source/Runtime/Renderer/Private/VirtualShadowMaps/VirtualShadowMapArray.cpp
+++ b/Engine/Source/Runtime/Renderer/Private/VirtualShadowMaps/VirtualShadowMapArray.cpp
@@ -410,7 +410,7 @@ class FGeneratePageFlagsFromPixelsCS : public FVirtualPageManagementShader
 		SHADER_PARAMETER_STRUCT_REF(FForwardLightData, ForwardLightData)
 		SHADER_PARAMETER_RDG_TEXTURE(Texture2D<uint2>, VisBuffer64)
 		SHADER_PARAMETER_RDG_BUFFER_UAV(RWStructuredBuffer<uint>, OutPageRequestFlags)
-		SHADER_PARAMETER_RDG_BUFFER_SRV(StructuredBuffer< uint >, VirtualShadowMapIdRemap)
+		SHADER_PARAMETER_RDG_BUFFER_SRV(StructuredBuffer< uint >, DirectionalLightIds)
 		RDG_BUFFER_ACCESS(IndirectBufferArgs, ERHIAccess::IndirectArgs)
 		SHADER_PARAMETER(uint32, InputType)
 		SHADER_PARAMETER(uint32, NumDirectionalLightSmInds)
@@ -873,14 +873,14 @@ void FVirtualShadowMapArray::BuildPageAllocations(
 		ShadowMapProjectionData.AddDefaulted(ShadowMaps.Num());
 
 		// Gather directional light virtual shadow maps
-		TArray<int32, SceneRenderingAllocator> DirectionalLightSmInds;
+		TArray<int32, SceneRenderingAllocator> DirectionalLightIds;
 		for (const FVisibleLightInfo& VisibleLightInfo : VisibleLightInfos)
 		{
 			for (const TSharedPtr<FVirtualShadowMapClipmap>& Clipmap : VisibleLightInfo.VirtualShadowMapClipmaps)
 			{
-				// NOTE: Shader assumes all levels from a given clipmap are contiguous in both the remap and projection arrays
+				// NOTE: Shader assumes all levels from a given clipmap are contiguous
 				int32 ClipmapID = Clipmap->GetVirtualShadowMap()->ID;
-				DirectionalLightSmInds.Add(ClipmapID);
+				DirectionalLightIds.Add(ClipmapID);
 				for (int32 ClipmapLevel = 0; ClipmapLevel < Clipmap->GetLevelCount(); ++ClipmapLevel)
 				{
 					ShadowMapProjectionData[ClipmapID + ClipmapLevel] = Clipmap->GetProjectionShaderData(ClipmapLevel);
@@ -944,7 +944,7 @@ void FVirtualShadowMapArray::BuildPageAllocations(
 	#endif
 
 		UniformParameters.NumShadowMaps = ShadowMaps.Num();
-		UniformParameters.NumDirectionalLights = DirectionalLightSmInds.Num();
+		UniformParameters.NumDirectionalLights = DirectionalLightIds.Num();
 
 		ShadowMapProjectionDataRDG = CreateProjectionDataBuffer(GraphBuilder, TEXT("Shadow.Virtual.ProjectionData"), ShadowMapProjectionData);
 
@@ -993,38 +993,12 @@ void FVirtualShadowMapArray::BuildPageAllocations(
 			FRDGTextureRef NaniteVisBuffer64 = ViewIndex < NaniteRasterResults.Num() ? NaniteRasterResults[ViewIndex].VisBuffer64 : nullptr;
 
 			// This view contained no local lights (that were stored in the light grid), and no directional lights, so nothing to do.
-			if (View.ForwardLightingResources->LocalLightVisibleLightInfosIndex.Num() + DirectionalLightSmInds.Num() == 0)
+			if (View.ForwardLightingResources->LocalLightVisibleLightInfosIndex.Num() + DirectionalLightIds.Num() == 0)
 			{
 				continue;
 			}
 
-			// Build light-index-in-light-grid => virtual-shadow-map-index remap, must be built for each view since they have different sub-sets of lights
-			// TODO: change this engine behaviour and instead upload lights once and for all, such that all indexes can refer to the same light set.
-			TArray<int32, SceneRenderingAllocator> VirtualShadowMapIdRemap = DirectionalLightSmInds;
-			// Note: the remap for the local lights is stored after the directional lights, such that this array is always non-empty...
-			VirtualShadowMapIdRemap.AddDefaulted(View.ForwardLightingResources->LocalLightVisibleLightInfosIndex.Num());
-			for (int32 L = 0; L < View.ForwardLightingResources->LocalLightVisibleLightInfosIndex.Num(); ++L)
-			{
-				// Default value
-				VirtualShadowMapIdRemap[DirectionalLightSmInds.Num() + L] = INDEX_NONE;
-
-				int32 VisibleLightInfosIndex = View.ForwardLightingResources->LocalLightVisibleLightInfosIndex[L];
-				// This can be invalid for example for so-called 'Simple lights' which are injected into the light grid, but not present elsewhere.
-				if (VisibleLightInfosIndex != INDEX_NONE)
-				{
-					const FVisibleLightInfo& VisibleLightInfo = VisibleLightInfos[VisibleLightInfosIndex];
-
-					// Get hold of info about this light to figure out if there is a virtual SM
-					if (FProjectedShadowInfo *ShadowInfo = GetVirtualShadowMapInfo(VisibleLightInfo))
-					{
-						ensure(ShadowInfo->VirtualShadowMaps.Num());
-						ensure(ShadowInfo->VirtualShadowMaps[0]->ID != INDEX_NONE);
-						VirtualShadowMapIdRemap[DirectionalLightSmInds.Num() + L] = ShadowInfo->VirtualShadowMaps[0]->ID;
-					}
-				}
-			}
-
-			VirtualShadowMapIdRemapRDG.Add( CreateStructuredBuffer(GraphBuilder, TEXT("Shadow.Virtual.IdRemap"), VirtualShadowMapIdRemap) );
+			FRDGBufferRef DirectionalLightIdsRDG = CreateStructuredBuffer(GraphBuilder, TEXT("Shadow.Virtual.DirectionalLightIds"), DirectionalLightIds);
 
 			const FRDGSystemTextures& SystemTextures = FRDGSystemTextures::Get(GraphBuilder);
 			FRDGTextureRef VisBuffer64 = NaniteVisBuffer64 != nullptr ? NaniteVisBuffer64 : SystemTextures.Black;
@@ -1056,8 +1030,7 @@ void FVirtualShadowMapArray::BuildPageAllocations(
 						PassParameters->View = View.ViewUniformBuffer;
 						PassParameters->OutPageRequestFlags = PageRequestFlagsUAV;
 						PassParameters->ForwardLightData = View.ForwardLightingResources->ForwardLightDataUniformBuffer;
-						PassParameters->VirtualShadowMapIdRemap = GraphBuilder.CreateSRV(VirtualShadowMapIdRemapRDG[ViewIndex]);
-						PassParameters->NumDirectionalLightSmInds = DirectionalLightSmInds.Num();
+						PassParameters->DirectionalLightIds = GraphBuilder.CreateSRV(DirectionalLightIdsRDG);
 						PassParameters->ResolutionLodBiasLocal = ResolutionLodBiasLocal;
 						PassParameters->PageDilationBorderSize = PageDilationBorderSize;
 
diff --git a/Engine/Source/Runtime/Renderer/Private/VirtualShadowMaps/VirtualShadowMapArray.h b/Engine/Source/Runtime/Renderer/Private/VirtualShadowMaps/VirtualShadowMapArray.h
--- a/Engine/Source/Runtime/Renderer/Private/VirtualShadowMaps/VirtualShadowMapArray.h
+++ b/Engine/Source/Runtime/Renderer/Private/VirtualShadowMaps/VirtualShadowMapArray.h
@@ -267,7 +267,6 @@ public:
 	FRDGBufferRef PageRectBoundsRDG = nullptr;
 	FRDGBufferRef AllocatedPageRectBoundsRDG = nullptr;
 	FRDGBufferRef ShadowMapProjectionDataRDG = nullptr;
-	TArray< FRDGBufferRef, SceneRenderingAllocator > VirtualShadowMapIdRemapRDG;
 
 	// HZB generated for the *current* frame's physical page pool
 	// We use the *previous* frame's HZB (from VirtualShadowMapCacheManager) for culling the current frame
diff --git a/Engine/Source/Runtime/Renderer/Private/VirtualShadowMaps/VirtualShadowMapProjection.cpp b/Engine/Source/Runtime/Renderer/Private/VirtualShadowMaps/VirtualShadowMapProjection.cpp
--- a/Engine/Source/Runtime/Renderer/Private/VirtualShadowMaps/VirtualShadowMapProjection.cpp
+++ b/Engine/Source/Runtime/Renderer/Private/VirtualShadowMaps/VirtualShadowMapProjection.cpp
@@ -130,6 +130,7 @@ void CompositeVirtualShadowMapMask(
 	FRDGBuilder& GraphBuilder,
 	const FIntRect ScissorRect,
 	const FRDGTextureRef Input,
+	bool bDirectionalLight,
 	FRDGTextureRef OutputShadowMaskTexture)
 {
 	auto ShaderMap = GetGlobalShaderMap(GMaxRHIFeatureLevel);
@@ -139,8 +140,13 @@ void CompositeVirtualShadowMapMask(
 
 	PassParameters->RenderTargets[0] = FRenderTargetBinding(OutputShadowMaskTexture, ERenderTargetLoadAction::ELoad);
 
-	// See FProjectedShadowInfo::GetBlendStateForProjection, but we don't want any of the special cascade behavior, etc. as this is a post-denoised mask.
-	FRHIBlendState* BlendState = TStaticBlendState<CW_BA, BO_Min, BF_One, BF_One, BO_Min, BF_One, BF_One>::GetRHI();
+	FRHIBlendState* BlendState = FProjectedShadowInfo::GetBlendStateForProjection(
+		0,					// ShadowMapChannel
+		bDirectionalLight,	// bIsWholeSceneDirectionalShadow,
+		false,				// bUseFadePlane
+		false,				// bProjectingForForwardShading, 
+		false				// bMobileModulatedProjections
+	);
 
 	auto PixelShader = ShaderMap->GetShader<FVirtualShadowMapProjectionCompositePS>();
 	ValidateShaderParameters(PixelShader, *PassParameters);
@@ -191,7 +197,6 @@ class FVirtualShadowMapProjectionCS : public FGlobalShader
 		SHADER_PARAMETER(uint32, InputType)
 		// One pass projection parameters
 		SHADER_PARAMETER_STRUCT_REF(FForwardLightData, ForwardLightData)
-		SHADER_PARAMETER_RDG_BUFFER_SRV(StructuredBuffer< uint >, VirtualShadowMapIdRemap)	// TODO: Move to VSM UB? Per-view though
 		SHADER_PARAMETER_RDG_TEXTURE_UAV(RWTexture2D< uint >, RWShadowMaskBits)
 		// Pass per light parameters
 		SHADER_PARAMETER_STRUCT(FLightShaderParameters, Light)
@@ -280,7 +285,6 @@ static void RenderVirtualShadowMapProjectionCommon(
 	{
 		// One pass projection
 		PassParameters->ForwardLightData = View.ForwardLightingResources->ForwardLightDataUniformBuffer;
-		PassParameters->VirtualShadowMapIdRemap = GraphBuilder.CreateSRV( VirtualShadowMapArray.VirtualShadowMapIdRemapRDG[0] );	// FIXME Index proper view
 		PassParameters->RWShadowMaskBits = GraphBuilder.CreateUAV( OutputTexture );	
 	}
 	else
@@ -381,7 +385,7 @@ static FRDGTextureRef CreateShadowFactorTexture(FRDGBuilder& GraphBuilder, FIntP
 
 	FRDGTextureDesc Desc = FRDGTextureDesc::Create2D(
 		Extent,
-		PF_R32_FLOAT,
+		PF_G16R16,
 		FClearValueBinding(ClearColor),
 		TexCreate_ShaderResource | TexCreate_UAV);
 
@@ -390,16 +394,18 @@ static FRDGTextureRef CreateShadowFactorTexture(FRDGBuilder& GraphBuilder, FIntP
 	return Texture;
 }
 
-FRDGTextureRef RenderVirtualShadowMapProjection(
+void RenderVirtualShadowMapProjection(
 	FRDGBuilder& GraphBuilder,
 	const FMinimalSceneTextures& SceneTextures,
 	const FViewInfo& View,
 	FVirtualShadowMapArray& VirtualShadowMapArray,
 	const FIntRect ScissorRect,
 	EVirtualShadowMapProjectionInputType InputType,
-	FProjectedShadowInfo* ShadowInfo)
+	FProjectedShadowInfo* ShadowInfo,
+	FRDGTextureRef OutputShadowMaskTexture)
 {
-	FRDGTextureRef OutputTexture = CreateShadowFactorTexture(GraphBuilder, SceneTextures.Config.Extent);
+	FRDGTextureRef VirtualShadowMaskTexture = CreateShadowFactorTexture(GraphBuilder, SceneTextures.Config.Extent);
+
 	RenderVirtualShadowMapProjectionCommon(
 		GraphBuilder,
 		SceneTextures,
@@ -407,22 +413,30 @@ FRDGTextureRef RenderVirtualShadowMapProjection(
 		VirtualShadowMapArray,
 		ScissorRect,
 		InputType,
-		OutputTexture,
+		VirtualShadowMaskTexture,
 		ShadowInfo->GetLightSceneInfo().Proxy,
 		ShadowInfo->VirtualShadowMaps[0]->ID);
-	return OutputTexture;
+
+	CompositeVirtualShadowMapMask(
+		GraphBuilder,
+		ScissorRect,
+		VirtualShadowMaskTexture,
+		false,	// bDirectionalLight
+		OutputShadowMaskTexture);
 }
 
-FRDGTextureRef RenderVirtualShadowMapProjection(
+void RenderVirtualShadowMapProjection(
 	FRDGBuilder& GraphBuilder,
 	const FMinimalSceneTextures& SceneTextures,
 	const FViewInfo& View,
 	FVirtualShadowMapArray& VirtualShadowMapArray,
 	const FIntRect ScissorRect,
 	EVirtualShadowMapProjectionInputType InputType,
-	const TSharedPtr<FVirtualShadowMapClipmap>& Clipmap)
+	const TSharedPtr<FVirtualShadowMapClipmap>& Clipmap,
+	FRDGTextureRef OutputShadowMaskTexture)
 {
-	FRDGTextureRef OutputTexture = CreateShadowFactorTexture(GraphBuilder, SceneTextures.Config.Extent);
+	FRDGTextureRef VirtualShadowMaskTexture = CreateShadowFactorTexture(GraphBuilder, SceneTextures.Config.Extent);
+
 	RenderVirtualShadowMapProjectionCommon(
 		GraphBuilder,
 		SceneTextures,
@@ -430,8 +444,14 @@ FRDGTextureRef RenderVirtualShadowMapProjection(
 		VirtualShadowMapArray,		
 		ScissorRect,
 		InputType,
-		OutputTexture,
+		VirtualShadowMaskTexture,
 		Clipmap->GetLightSceneInfo().Proxy,
 		Clipmap->GetVirtualShadowMap()->ID);
-	return OutputTexture;
+	
+	CompositeVirtualShadowMapMask(
+		GraphBuilder,
+		ScissorRect,
+		VirtualShadowMaskTexture,
+		true,	// bDirectionalLight
+		OutputShadowMaskTexture);
 }
diff --git a/Engine/Source/Runtime/Renderer/Private/VirtualShadowMaps/VirtualShadowMapProjection.h b/Engine/Source/Runtime/Renderer/Private/VirtualShadowMaps/VirtualShadowMapProjection.h
--- a/Engine/Source/Runtime/Renderer/Private/VirtualShadowMaps/VirtualShadowMapProjection.h
+++ b/Engine/Source/Runtime/Renderer/Private/VirtualShadowMaps/VirtualShadowMapProjection.h
@@ -18,23 +18,25 @@ enum class EVirtualShadowMapProjectionInputType
 	HairStrands = 1
 };
 
-FRDGTextureRef RenderVirtualShadowMapProjection(
+void RenderVirtualShadowMapProjection(
 	FRDGBuilder& GraphBuilder,
 	const FMinimalSceneTextures& SceneTextures,
 	const FViewInfo& View,
 	FVirtualShadowMapArray& VirtualShadowMapArray,
 	const FIntRect ScissorRect,
 	EVirtualShadowMapProjectionInputType InputType,
-	const TSharedPtr<FVirtualShadowMapClipmap>& Clipmap);
+	const TSharedPtr<FVirtualShadowMapClipmap>& Clipmap,
+	FRDGTextureRef OutputShadowMaskTexture);
 
-FRDGTextureRef RenderVirtualShadowMapProjection(
+void RenderVirtualShadowMapProjection(
 	FRDGBuilder& GraphBuilder,
 	const FMinimalSceneTextures& SceneTextures,
 	const FViewInfo& View,
 	FVirtualShadowMapArray& VirtualShadowMapArray,
 	const FIntRect ScissorRect,
 	EVirtualShadowMapProjectionInputType InputType,
-	FProjectedShadowInfo* ShadowInfo);
+	FProjectedShadowInfo* ShadowInfo,
+	FRDGTextureRef OutputShadowMaskTexture);
 
 FRDGTextureRef RenderVirtualShadowMapProjectionOnePass(
 	FRDGBuilder& GraphBuilder,
@@ -47,4 +49,5 @@ void CompositeVirtualShadowMapMask(
 	FRDGBuilder& GraphBuilder,
 	const FIntRect ScissorRect,
 	const FRDGTextureRef Input,
+	bool bDirectionalLight,
 	FRDGTextureRef OutputShadowMaskTexture);
