Adds DXC bug workarround in TSR when enabling 16bit types

diff --git a/Engine/Shaders/Private/TemporalSuperResolution/TSRUpdateHistory.usf b/Engine/Shaders/Private/TemporalSuperResolution/TSRUpdateHistory.usf
--- a/Engine/Shaders/Private/TemporalSuperResolution/TSRUpdateHistory.usf
+++ b/Engine/Shaders/Private/TemporalSuperResolution/TSRUpdateHistory.usf
@@ -947,7 +947,13 @@ void MainCS(
 			}
 		}
 		
-		uint2 SubpixelId = dpv_dot(uint2x2(bKOSign), uint2(1, SUB_PIXEL_GRID_SIZE));
+		#if COMPILER_DXC == 1 // shader compiler bug
+			uint2x2 bKOSignInt = uint2x2(bKOSign);
+			uint2 KOSignVector = uint2(1, SUB_PIXEL_GRID_SIZE);
+			uint2 SubpixelId = dpv_interleave_registers(dot(dpv_lo(bKOSignInt), KOSignVector), dot(dpv_hi(bKOSignInt), KOSignVector));
+		#else
+			uint2 SubpixelId = dpv_dot(uint2x2(bKOSign), uint2(1, SUB_PIXEL_GRID_SIZE));
+		#endif
 
 		taa_subpixel_payload2 SubpixelShift = taa_subpixel_payload2(SubpixelId) * taa_subpixel_payload(SUB_PIXEL_BIT_COUNT);
 
