Fix for geometry cache vertex factory stomping on TEXCOORDs.

- This follows the same pattern as all the other vertex factories to pack texcoord interpolants into float4s.
- Still could stomp on tex coords, but this buys us some space.

diff --git a/Engine/Shaders/Private/GeometryCacheVertexFactory.ush b/Engine/Shaders/Private/GeometryCacheVertexFactory.ush
--- a/Engine/Shaders/Private/GeometryCacheVertexFactory.ush
+++ b/Engine/Shaders/Private/GeometryCacheVertexFactory.ush
@@ -135,7 +135,8 @@ struct FVertexFactoryInterpolantsVSToPS
 #endif
 
 #if NUM_TEX_COORD_INTERPOLATORS
-	float2	TexCoords[NUM_TEX_COORD_INTERPOLATORS]	: TEXCOORD0;
+	// Pack interpolators to reduce the number of semantics used
+	float4	TexCoords[(NUM_TEX_COORD_INTERPOLATORS + 1) / 2]	: TEXCOORD0;
 #endif
 
 #if INSTANCED_STEREO
@@ -145,6 +146,27 @@ struct FVertexFactoryInterpolantsVSToPS
 	TANGENTTOWORLD_INTERPOLATOR_BLOCK
 };
 
+#if NUM_TEX_COORD_INTERPOLATORS
+float2 GetUV(FVertexFactoryInterpolantsVSToPS Interpolants, int UVIndex)
+{
+	float4 UVVector = Interpolants.TexCoords[UVIndex / 2];
+	return UVIndex % 2 ? UVVector.zw : UVVector.xy;
+}
+
+void SetUV(inout FVertexFactoryInterpolantsVSToPS Interpolants, int UVIndex, float2 InValue)
+{
+	FLATTEN
+	if (UVIndex % 2)
+	{
+		Interpolants.TexCoords[UVIndex / 2].zw = InValue;
+	}
+	else
+	{
+		Interpolants.TexCoords[UVIndex / 2].xy = InValue;
+	}
+}
+#endif
+
 struct FVertexFactoryRayTracingInterpolants
 {
 	FVertexFactoryInterpolantsVSToPS InterpolantsVSToPS;
@@ -160,7 +182,7 @@ FMaterialPixelParameters GetMaterialPixelParameters(FVertexFactoryInterpolantsVS
 	UNROLL
 	for(int CoordinateIndex = 0;CoordinateIndex < NUM_TEX_COORD_INTERPOLATORS;CoordinateIndex++)
 	{
-		Result.TexCoords[CoordinateIndex] = Interpolants.TexCoords[CoordinateIndex];
+		Result.TexCoords[CoordinateIndex] = GetUV(Interpolants, CoordinateIndex);
 	}
 #endif
 
@@ -319,6 +341,9 @@ FVertexFactoryInterpolantsVSToPS VertexFactoryGetInterpolantsVSToPS(FVertexFacto
 {
 	FVertexFactoryInterpolantsVSToPS Interpolants;
 
+	// Initialize the whole struct to 0
+	Interpolants = (FVertexFactoryInterpolantsVSToPS)0;
+
 #if NUM_TEX_COORD_INTERPOLATORS
 	float2 CustomizedUVs[NUM_TEX_COORD_INTERPOLATORS];
 	GetMaterialCustomizedUVs(VertexParameters, CustomizedUVs);
@@ -327,7 +352,7 @@ FVertexFactoryInterpolantsVSToPS VertexFactoryGetInterpolantsVSToPS(FVertexFacto
 	UNROLL
 	for (int CoordinateIndex = 0; CoordinateIndex < NUM_TEX_COORD_INTERPOLATORS; CoordinateIndex++)
 	{
-		Interpolants.TexCoords[CoordinateIndex] = CustomizedUVs[CoordinateIndex];
+		SetUV(Interpolants, CoordinateIndex, CustomizedUVs[CoordinateIndex]);
 	}
 #endif
 
@@ -410,7 +435,7 @@ FVertexFactoryRayTracingInterpolants VertexFactoryInterpolate(FVertexFactoryRayT
 
 #if NUM_TEX_COORD_INTERPOLATORS
 	UNROLL
-	for(int tc = 0; tc < NUM_TEX_COORD_INTERPOLATORS; ++tc)
+	for(int tc = 0; tc < (NUM_TEX_COORD_INTERPOLATORS + 1) / 2; ++tc)
 	{
 		INTERPOLATE_MEMBER(InterpolantsVSToPS.TexCoords[tc]);
 	}
