Fix compile error with instanced stereo rendering

diff --git a/Engine/Shaders/Private/LightGridCommon.ush b/Engine/Shaders/Private/LightGridCommon.ush
--- a/Engine/Shaders/Private/LightGridCommon.ush
+++ b/Engine/Shaders/Private/LightGridCommon.ush
@@ -215,8 +215,8 @@ FLocalLightData GetLocalLightData(uint GridIndex, uint EyeIndex)
 		Result.LightDirectionAndShadowMask = ForwardLightDataISR.ForwardLocalLightBuffer[LocalLightBaseIndex + 2];
 		Result.SpotAnglesAndSourceRadiusPacked = ForwardLightDataISR.ForwardLocalLightBuffer[LocalLightBaseIndex + 3];
 		Result.LightTangentAndSoftSourceRadius = ForwardLightDataISR.ForwardLocalLightBuffer[LocalLightBaseIndex + 4];
-		Result.RectBarnDoor = ForwardLightData.ForwardLightDataISR[LocalLightBaseIndex + 5];
-		Result.VirtualShadowMapId = asint(ForwardLightData.ForwardLightDataISR[LocalLightBaseIndex + 5].z);
+		Result.RectBarnDoor = ForwardLightDataISR.ForwardLocalLightBuffer[LocalLightBaseIndex + 5];
+		Result.VirtualShadowMapId = asint(ForwardLightDataISR.ForwardLocalLightBuffer[LocalLightBaseIndex + 5].z);
 	}
 #endif
 
