Fix CIS shader compilation warning.

diff --git a/Engine/Shaders/Private/HairStrands/HairStrandsVisibilityCommon.ush b/Engine/Shaders/Private/HairStrands/HairStrandsVisibilityCommon.ush
--- a/Engine/Shaders/Private/HairStrands/HairStrandsVisibilityCommon.ush
+++ b/Engine/Shaders/Private/HairStrands/HairStrandsVisibilityCommon.ush
@@ -139,12 +139,12 @@ FNodeDesc DecodeNodeDesc(uint In)
 
 uint PackHairVisDepthCoverage(float InDepth, float InCoverage)
 {
-	return (uint(InDepth * 0x00FFFFFF) << 8) | uint(0xFF * saturate(InCoverage));
+	return (uint(InDepth * 0x00FFFFFF) << 8) | uint(0xFFu * saturate(InCoverage));
 }
 
 uint PackHairVisDepthCoverage(float InDepth, uint InCoverage8bits)
 {
-	return (uint(InDepth * 0x00FFFFFF) << 8) | min(0xFF, InCoverage8bits);
+	return (uint(InDepth * 0x00FFFFFF) << 8) | min(0xFFu, InCoverage8bits);
 }
 
 uint PackHairVisPrimitiveMaterialId(uint PrimitiveId, uint MaterialId)
