Fixed mobile base pass pixel shader compilation warnings

diff --git a/Engine/Shaders/Private/MobileBasePassPixelShader.usf b/Engine/Shaders/Private/MobileBasePassPixelShader.usf
--- a/Engine/Shaders/Private/MobileBasePassPixelShader.usf
+++ b/Engine/Shaders/Private/MobileBasePassPixelShader.usf
@@ -61,7 +61,7 @@ int NumDynamicPointLights;
 #define LightShadowType_PointLight	1
 #define LightShadowType_SpotLight	2
 #define LightShadowType_Shadow		4
-#define ValidLightType LightShadowType_PointLight | LightShadowType_SpotLight
+#define ValidLightType (LightShadowType_PointLight | LightShadowType_SpotLight)
 
 #endif
 
