PR #8181: Fixed incorrect light complexity for lights with IES profiles. (Contributed by [REDACTED_PERSON])

diff --git a/Engine/Shaders/Private/IESLightProfilesCommon.ush b/Engine/Shaders/Private/IESLightProfilesCommon.ush
--- a/Engine/Shaders/Private/IESLightProfilesCommon.ush
+++ b/Engine/Shaders/Private/IESLightProfilesCommon.ush
@@ -11,7 +11,7 @@ SamplerState	IESTextureSampler;
 // Apply IES light profile texture
 float ComputeLightProfileMultiplier(float3 WorldPosition, float3 LightPosition, float3 LightDirection, float3 LightTangent)
 {
-#if USE_IES_PROFILE
+#if USE_IES_PROFILE && !VISUALIZE_LIGHT_CULLING
 	float3 LightBitangent = normalize( cross( LightTangent, LightDirection ) );
 
 	float4x4 LightTransform = float4x4( float4(LightDirection.xyz, 0), float4(LightBitangent.xyz, 0), float4(LightTangent.xyz, 0), float4(0, 0, 0, 1) );
