Added 'precise' qualifier to fix Z-fighting with World Position Offset (WPO) in MaterialTemplate.ush. Only affects D3D11/DXBC shaders on Windows.

#rnx

diff --git a/Engine/Shaders/Private/MaterialTemplate.ush b/Engine/Shaders/Private/MaterialTemplate.ush
--- a/Engine/Shaders/Private/MaterialTemplate.ush
+++ b/Engine/Shaders/Private/MaterialTemplate.ush
@@ -2448,7 +2448,14 @@ half GetMaterialMaskedOpacity(FPixelMaterialInputs PixelMaterialInputs)
 }
 #endif
 
-float3 GetMaterialWorldPositionOffset(FMaterialVertexParameters Parameters)
+// UE-102725: Workaround for precision issue that only occurs with world position offsets in DXBC shaders.
+#if IR_LANGUAGE_DXBC
+#define WPO_PRECISE(T) precise T
+#else
+#define WPO_PRECISE(T) T
+#endif
+
+WPO_PRECISE(float3) GetMaterialWorldPositionOffset(FMaterialVertexParameters Parameters)
 {
 	#if USE_INSTANCING || USE_INSTANCE_CULLING
 		// skip if this instance is hidden
@@ -2460,7 +2467,7 @@ float3 GetMaterialWorldPositionOffset(FMaterialVertexParameters Parameters)
 %s;
 }
 
-float3 GetMaterialPreviousWorldPositionOffset(FMaterialVertexParameters Parameters)
+WPO_PRECISE(float3) GetMaterialPreviousWorldPositionOffset(FMaterialVertexParameters Parameters)
 {
 	#if USE_INSTANCING || USE_INSTANCE_CULLING
 		// skip if this instance is hidden
diff --git a/Engine/Shaders/Public/Platform.ush b/Engine/Shaders/Public/Platform.ush
--- a/Engine/Shaders/Public/Platform.ush
+++ b/Engine/Shaders/Public/Platform.ush
@@ -141,6 +141,10 @@
 #define VECTORVM_PROFILE 0
 #endif
 
+#ifndef IR_LANGUAGE_DXBC
+#define IR_LANGUAGE_DXBC 0
+#endif
+
 // 'static' asserts
 #if COMPILER_GLSL || COMPILER_GLSL_ES3_1 || COMPILER_VULKAN || COMPILER_METAL
 #if !COMPILER_HLSLCC
diff --git a/Engine/Source/Developer/Windows/ShaderFormatD3D/Private/D3DShaderCompiler.cpp b/Engine/Source/Developer/Windows/ShaderFormatD3D/Private/D3DShaderCompiler.cpp
--- a/Engine/Source/Developer/Windows/ShaderFormatD3D/Private/D3DShaderCompiler.cpp
+++ b/Engine/Source/Developer/Windows/ShaderFormatD3D/Private/D3DShaderCompiler.cpp
@@ -1046,6 +1046,10 @@ void CompileD3DShader(const FShaderCompilerInput& Input, FShaderCompilerOutput&
 		AdditionalDefines.SetDefine(TEXT("PLATFORM_SUPPORTS_DIAGNOSTIC_BUFFER"), 1);
 		AdditionalDefines.SetDefine(TEXT("COMPILER_SUPPORTS_NOINLINE"), 1);
 	}
+	else
+	{
+		AdditionalDefines.SetDefine(TEXT("IR_LANGUAGE_DXBC"), 1);
+	}
 
 	if (Input.bSkipPreprocessedCache)
 	{
