Fix strata compilation with shadow transmission
Strata support for VSM transmissions still TODO

diff --git a/Engine/Shaders/Private/ShadowProjectionPixelShader.usf b/Engine/Shaders/Private/ShadowProjectionPixelShader.usf
--- a/Engine/Shaders/Private/ShadowProjectionPixelShader.usf
+++ b/Engine/Shaders/Private/ShadowProjectionPixelShader.usf
@@ -282,11 +282,14 @@ void Main(
 		BRANCH
 		if (bIsSubsurfaceCompatible && bHasSubsurface)
 		{
+			#if STRATA_ENABLED
+			const FStrataSubsurfaceData SSSData = StrataUnpackSSSData(Strata.SSSTexture.Load(uint3(SVPos.xy, 0)));
+			#endif
+
 			float Density = 1;
 			if (!bIsEyeOrHair)
 			{
 			#if STRATA_ENABLED
-				const FStrataSubsurfaceData SSSData = StrataUnpackSSSData(Strata.SSSTexture.Load(uint3(SVPos.xy, 0)));
 				if (SSSData.bIsValid)
 				{
 					float Opacity = SSSData.bIsProfile ? SSSData.ProfileRadiusScale : 1.0f;
@@ -296,7 +299,7 @@ void Main(
 				Density = SubsurfaceDensityFromOpacity(GBufferData.CustomData.a);
 			#endif // STRATA_ENABLED
 			}
-								
+			
 			FPCFSamplerSettings Settings;
 				
 			Settings.ShadowDepthTexture = ShadowDepthTexture;
diff --git a/Engine/Shaders/Private/VirtualShadowMaps/TransmissionCommon.ush b/Engine/Shaders/Private/VirtualShadowMaps/TransmissionCommon.ush
--- a/Engine/Shaders/Private/VirtualShadowMaps/TransmissionCommon.ush
+++ b/Engine/Shaders/Private/VirtualShadowMaps/TransmissionCommon.ush
@@ -15,13 +15,14 @@ VirtualShadowMaps/TransmissionCommon.ush:
 #define STRATA_ENABLED 0
 #endif
 
-// TODO: Once strata is set up this function should be shared with ShadowProjectionPixelShader
-// and possibly also the contact shadow trace
 half GetSubsurfaceOpacityFromGbuffer(uint2 PixelPos)
 {
 	half Opacity = 1.0;
 
-#if STRATA_ENABLED
+	// TODO: Support strata
+	// TODO: Once strata is set up this function should be shared with ShadowProjectionPixelShader
+	// and possibly also the contact shadow trace
+#if STRATA_ENABLED && 0
 	// TODO: VSM projection shaders are not set up for strata yet, but added a skeleton of the logic
 	const FStrataClassification StrataData = StrataUnpackClassificationData(Strata.ClassificationTexture.Load(PixelPos));
 	if (StrataData.bIsStrataMaterial)
