Fix D3D validation error due to a PF_R8 cube texture (GlobalDistanceFieldPageAtlasTexture) being used without the UNORM qualifier in the shader.

#rnx

diff --git a/Engine/Shaders/Private/GlobalDistanceFieldHeightfields.usf b/Engine/Shaders/Private/GlobalDistanceFieldHeightfields.usf
--- a/Engine/Shaders/Private/GlobalDistanceFieldHeightfields.usf
+++ b/Engine/Shaders/Private/GlobalDistanceFieldHeightfields.usf
@@ -12,7 +12,7 @@
 #define GLOBAL_DISTANCE_FIELD_PAGE_HEIGHFIELD_GROUP_NUM (GLOBAL_DISTANCE_FIELD_PAGE_HEIGHFIELD_GROUP_X * GLOBAL_DISTANCE_FIELD_PAGE_HEIGHFIELD_GROUP_X)
 
 RWStructuredBuffer<uint> RWMarkedHeightfieldPageBuffer;
-RWTexture3D<float> RWPageAtlasTexture;
+RWTexture3D<UNORM float> RWPageAtlasTexture;
 
 StructuredBuffer<uint> MarkedHeightfieldPageBuffer;
 StructuredBuffer<uint> PageUpdateTileBuffer;
