Improve several self-intersection behaviors in the path tracer

Add the geometric normal to the path tracing payload. The shading and geometric normal are both packed in 24 bits (oct encoding), which keeps the payload size identical.

Improve the ray biasing strategy in the path tracer to use an adaptive offset instead of a hardcoded epsilon. This improves the behavior far from the origin and near object boundaries. A larger offset is retained for light samples since these can lead to artifacts near the shadow terminator.

Improve the handling of shading normals near grazing angles by tweaking them against the geometric normal. This resolves dark edges on reflective low-poly objects.

Improve the handling of self-intersections in random walk SSS.  Use interface counting to keep track of self intersecting surfaces while performing a random walk. This increases the cost slightly but prevents darkening effects where multiple meshes overlap each other.

diff --git a/Engine/Shaders/Private/PathTracing/PathTracingCommon.ush b/Engine/Shaders/Private/PathTracing/PathTracingCommon.ush
--- a/Engine/Shaders/Private/PathTracing/PathTracingCommon.ush
+++ b/Engine/Shaders/Private/PathTracing/PathTracingCommon.ush
@@ -7,6 +7,7 @@
 #pragma once
 
 #include "../RayTracing/RayTracingCommon.ush"
+#include "../OctahedralCommon.ush"
 
 // These flags are analogous to RAY_TRACING_PAYLOAD_* flags. There are currently 6 bits available
 #define PATH_TRACING_PAYLOAD_OUTPUT_FLAG_FRONT_FACE                   (1 << 0) // Indicates that ray has hit the front face of a primitive. Set by closest hit shader.
@@ -15,6 +16,27 @@
 #define PATH_TRACING_PAYLOAD_INPUT_FLAG_PATHTRACER_VISIBILITY         (1 << 1) // Indicates that this is a path tracer visibility ray (which supports material evaluation for transparent shadows and fake caustics)
 
 
+#define PATH_TRACING_PAYLOAD_NORMAL_BITS  24
+
+uint PayloadEncodeUnitVector(float3 Normal)
+{
+	const int BITS = PATH_TRACING_PAYLOAD_NORMAL_BITS;
+	const float Scale = float((1u << (BITS / 2)) - 2); // largest representable even number (so we can encode 0)
+	float2 OctF = UnitVectorToOctahedron(Normal);
+	int2 Oct = int2(round(OctF * Scale));
+	return Oct.x + (1u << (BITS / 2)) * Oct.y;
+}
+
+float3 PayloadDecodeUnitVector(uint Encoded)
+{
+	const int BITS = PATH_TRACING_PAYLOAD_NORMAL_BITS;
+	const uint Mask = (1u << (BITS / 2)) - 1;
+	const float Scale = float((1u << (BITS / 2)) - 2);
+
+	int2 Oct = int2(Encoded & Mask, (Encoded >> (BITS / 2)) & Mask);
+	float2 OctF = saturate(float2(Oct) / Scale);
+	return OctahedronToUnitVector(OctF);
+}
 
 // This payload structure is what we transport between RGS/CHS/AHS programs
 struct FPackedPathTracingPayload : FMinimalPayload
@@ -30,17 +52,6 @@ struct FPackedPathTracingPayload : FMinimalPayload
 
 	FRayCone GetRayCone() { return UnpackRayCone(PackedRayCone); }
 
-	float3 GetWorldNormal()
-	{
-		float3 Result;
-		Result.x = f16tof32(PackedData[1] >> 16);
-		Result.y = f16tof32(PackedData[2]);
-		Result.z = f16tof32(PackedData[2] >> 16);
-		return Result;
-	}
-
-	uint GetShadingModelID() { return (PackedData[7] >> 16) & 0xF; }
-	uint GetBlendingMode() { return (PackedData[7] >> 20) & 0x7; }
 	uint GetFlags() { return (PackedData[7] >> 23) & 0x3F; }
 	void SetFlag(uint Flag) { PackedData[7] |= Flag << 23; }
 	uint GetPrimitiveLightingChannelMask() { return (PackedData[7] >> 29) & 0x7; }
@@ -88,6 +99,7 @@ struct FPathTracingPayload : FMinimalPayload
 {
 	FRayCone RayCone;
 	float3 Radiance;
+	float3 WorldGeoNormal;
 	float3 WorldNormal;
 	float3 BaseColor;
 	float3 DiffuseColor;
@@ -182,9 +194,11 @@ FPackedPathTracingPayload PackPathTracingPayload(FPathTracingPayload Input)
 	Output.PackedData[0] |= f32tof16(Input.Radiance.y) << 16;
 	Output.PackedData[1] = f32tof16(Input.Radiance.z);
 
-	Output.PackedData[1] |= f32tof16(Input.WorldNormal.x) << 16;
-	Output.PackedData[2] = f32tof16(Input.WorldNormal.y);
-	Output.PackedData[2] |= f32tof16(Input.WorldNormal.z) << 16;
+	uint ShdNormal = PayloadEncodeUnitVector(Input.WorldNormal);
+	uint GeoNormal = PayloadEncodeUnitVector(Input.WorldGeoNormal);
+	Output.PackedData[1] |= ShdNormal << 16; // pack low 16 bits
+	Output.PackedData[2] = ShdNormal >> 16;  // pack upper 8 bits
+	Output.PackedData[2] |= GeoNormal << 8;  // pack all 24 bits
 
 	Output.PackedData[3] = f32tof16(Input.BaseColor.x);
 	Output.PackedData[3] |= f32tof16(Input.BaseColor.y) << 16;
@@ -229,9 +243,8 @@ FPathTracingPayload UnpackPathTracingPayload(FPackedPathTracingPayload Input, Ra
 	Output.Radiance.x = f16tof32(Input.PackedData[0]);
 	Output.Radiance.y = f16tof32(Input.PackedData[0] >> 16);
 	Output.Radiance.z = f16tof32(Input.PackedData[1]);
-	Output.WorldNormal.x = f16tof32(Input.PackedData[1] >> 16);
-	Output.WorldNormal.y = f16tof32(Input.PackedData[2]);
-	Output.WorldNormal.z = f16tof32(Input.PackedData[2] >> 16);
+	Output.WorldNormal = PayloadDecodeUnitVector((Input.PackedData[1] >> 16) | ((Input.PackedData[2] & 0xFF) << 16));
+	Output.WorldGeoNormal = PayloadDecodeUnitVector(Input.PackedData[2] >> 8);
 	Output.BaseColor.x = f16tof32(Input.PackedData[3]);
 	Output.BaseColor.y = f16tof32(Input.PackedData[3] >> 16);
 	Output.BaseColor.z = f16tof32(Input.PackedData[4]);
@@ -241,8 +254,8 @@ FPathTracingPayload UnpackPathTracingPayload(FPackedPathTracingPayload Input, Ra
 	Output.Roughness = f16tof32(Input.PackedData[6]);
 	Output.Anisotropy = f16tof32(Input.PackedData[6] >> 16);
 	Output.Ior = f16tof32(Input.PackedData[7]);
-	Output.ShadingModelID = Input.GetShadingModelID();
-	Output.BlendingMode = Input.GetBlendingMode();
+	Output.ShadingModelID = (Input.PackedData[7] >> 16) & 0xF;
+	Output.BlendingMode = (Input.PackedData[7] >> 20) & 0x7;
 	Output.Flags = Input.GetFlags();
 	Output.PrimitiveLightingChannelMask = Input.GetPrimitiveLightingChannelMask();
 
diff --git a/Engine/Shaders/Private/PathTracing/PathTracingCore.ush b/Engine/Shaders/Private/PathTracing/PathTracingCore.ush
--- a/Engine/Shaders/Private/PathTracing/PathTracingCore.ush
+++ b/Engine/Shaders/Private/PathTracing/PathTracingCore.ush
@@ -228,6 +228,7 @@ struct FProbeResult
 {
 	float HitT;
 	float3 WorldNormal;
+	float3 WorldGeoNormal;
 	int FrontFace;
 
 	bool IsMiss() { return HitT <= 0; }
@@ -236,9 +237,9 @@ struct FProbeResult
 FProbeResult TraceProbeRay(RayDesc Ray, uint2 PixelCoord)
 {
 	// Trace a short ray to see if we escaped the surface
-	// TODO: figure out a way to only hit the object we started on (or handle blocking geometry)
 	FPackedPathTracingPayload PackedPayload = InitPathTracingPayload((FRayCone)0, PixelCoord, true);
-	// TODO: payload flag to skip shading (we only care about the normal)
+	// TODO: use smaller payload + flag to skip shading work? (only want smooth+geo normals)
+	//       however supporting blockers would require some shading ...
 	const uint RayFlags = 0;
 	const uint MissShaderIndex = 0;
 	TraceRay(
@@ -255,14 +256,32 @@ FProbeResult TraceProbeRay(RayDesc Ray, uint2 PixelCoord)
 	{
 		return (FProbeResult)0;
 	}
-
+	FPathTracingPayload Payload = UnpackPathTracingPayload(PackedPayload, Ray);
 	FProbeResult Result;
 	Result.HitT = PackedPayload.HitT;
-	Result.WorldNormal = PackedPayload.GetWorldNormal();
-	Result.FrontFace = PackedPayload.IsFrontFace();
+	Result.WorldNormal = Payload.WorldNormal;
+	Result.WorldGeoNormal = Payload.WorldGeoNormal;
+	Result.FrontFace = Payload.IsFrontFace();
 	return Result;
 }
 
+void ApplyRayBias(inout RayDesc Ray, float HitT, float3 Normal, bool ShadowTerminatorFix = false)
+{
+#if 1
+	// Take maximum of position or hit distance to determine "radius" of hit avoidance
+	const float RefValue = max(max(abs(Ray.Origin.x), abs(Ray.Origin.y)), max(abs(Ray.Origin.z), HitT));
+	const uint UlpRadius = 16; // number of floating point ulps to skip around an intersection
+	const float RelativeOffset = asfloat(asuint(RefValue) + UlpRadius) - RefValue;
+	const float ShadowOffset = 0.1; // 1mm
+	const float BaseOffset = 0.001; // 0.01mm (avoid precision issues very close to the origin where ulps become tiny)
+	// for shadow rays, push out the point a bit more so we can reduce shadow terminator issues from poor tesselations
+	// other rays are less likely to arrive at grazing angles, so don't need the extra offset
+	Ray.Origin += max(ShadowTerminatorFix ? ShadowOffset : BaseOffset, RelativeOffset) * Normal;
+#else
+	ApplyPositionBias(Ray.Origin, Ray.Direction, Normal, MaxNormalBias);
+#endif
+}
+
 bool ProcessSubsurfaceRandomWalk(inout FPathTracingPayload Payload, inout float3 PathThroughput, inout RandomSequence RandSequence, uint2 PixelCoord, bool SimplifySSS)
 {
 	if (!Payload.IsSubsurfaceMaterial())
@@ -308,13 +327,22 @@ bool ProcessSubsurfaceRandomWalk(inout FPathTracingPayload Payload, inout float3
 	Ray.Origin = Payload.WorldPos;
 	Ray.Direction = TangentToWorld(-CosineSampleHemisphere(RandSample.xy).xyz, Payload.WorldNormal);
 	Ray.TMin = 0;
-	ApplyPositionBias(Ray, -Payload.WorldNormal, MaxNormalBias);
+	ApplyRayBias(Ray, Payload.HitT, -Payload.WorldGeoNormal);
+
+    // Use the technique mentioned in "A Hero Beneath the Surface" (Section 6.3.2)
+	// https://jo.dreggn.org/home/2021_spectral_imaging.pdf
+	// This allows multiple overlapping meshes to be treated as a single "volume" by discarding internal hits until we have
+	// crossed the appropriate number of interfaces
+#define SSS_USE_INTERFACE_COUNTING 1
+
+	int InterfaceCounter = Payload.IsFrontFace() ? +1 : -1;
 
 	// Use the technique detailed in the Hyperion paper (Section 4.4.2)
 	// https://media.disneyanimation.com/uploads/production/publication_asset/177/asset/a.pdf
 	// This avoids overly bright thin regions by applying a bit of internal reflection when the boundary is hit
 #define SSS_USE_TIR 1
 
+
 #if SSS_USE_TIR
 	const float3 Albedo = 1 - exp(SSSColor * (-11.43 + SSSColor * (15.38 - 13.91 * SSSColor)));
 #else
@@ -360,60 +388,81 @@ bool ProcessSubsurfaceRandomWalk(inout FPathTracingPayload Payload, inout float3
 			float RescaleRand = min((RandSample.z - ColorChannelCdf.y) / ColorChannelPdf.z, OneMinusEpsilon);
 			Ray.TMax = -log(1 - RescaleRand) * SSSRadius.z;
 		}
-
-		FProbeResult ProbeResult = TraceProbeRay(Ray, PixelCoord);
-		if (ProbeResult.IsMiss())
-		{
-			// we didn't hit anything, so scatter randomly in the volume and keep tracing
-			RandSample = RandomSequence_GenerateSample3D(RandSequence);
-			Ray.Origin += Ray.TMax * Ray.Direction;
-			Ray.Direction = UniformSampleSphere(RandSample.xy).xyz;
-			// account for transmittance and scattering up to this point
-			float3 Transmittance = exp(-Ray.TMax * SigmaT);
-			PathThroughput *= SigmaS * Transmittance;
-			// account for transmittance pdf (which is the balance heuristic combination of transmittance sampling in each color channel)
-			float3 TransmittancePdf = SigmaT * Transmittance;
-			PathThroughput /= dot(TransmittancePdf, ColorChannelPdf);
-
-			// NOTE: phase function cancels out since it is being perfectly importance sampled
-		}
-		else
+		for (;;)
 		{
-			// account for transmittance and scattering up to this point
-			float3 Transmittance = exp(-ProbeResult.HitT * SigmaT);
-			PathThroughput *= Transmittance;
-			// account for transmittance pdf (which is the balance heuristic combination of transmittance sampling in each color channel)
-			float3 TransmittancePdf = Transmittance; // integral of pdf from HitT to infinity
-			PathThroughput /= dot(Transmittance, ColorChannelPdf);
-
-			// our short ray hit the geomtry - decide if we should exit or not
-#if SSS_USE_TIR
-			RandSample = RandomSequence_GenerateSample2D(RandSequence).xxy;
-			float3 WorldNormal = ProbeResult.WorldNormal;
-			float CosTheta = abs(dot(Ray.Direction, WorldNormal));
-			float Fresnel = FresnelReflectance(CosTheta, 1.0 / 1.4);
-			if (RandSample.x < Fresnel)
+			FProbeResult ProbeResult = TraceProbeRay(Ray, PixelCoord);
+			if (ProbeResult.IsMiss())
 			{
-				// internal reflection occured -- reflect and keep tracing
-				// NOTE: weight and probability cancel out, so no throughput adjustment is needed
-				Ray.Origin += ProbeResult.HitT * Ray.Direction;
-				Ray.Direction = reflect(Ray.Direction, WorldNormal);
-				ApplyPositionBias(Ray, WorldNormal, MaxNormalBias);
-				continue;
+				// we didn't hit anything, so scatter randomly in the volume and keep tracing
+				RandSample = RandomSequence_GenerateSample3D(RandSequence);
+				Ray.Origin += Ray.TMax * Ray.Direction;
+				Ray.Direction = UniformSampleSphere(RandSample.xy).xyz;
+				Ray.TMin = 0;
+				// account for transmittance and scattering up to this point
+				float3 Transmittance = exp(-Ray.TMax * SigmaT);
+				PathThroughput *= SigmaS * Transmittance;
+				// account for transmittance pdf (which is the balance heuristic combination of transmittance sampling in each color channel)
+				float3 TransmittancePdf = SigmaT * Transmittance;
+				PathThroughput /= dot(TransmittancePdf, ColorChannelPdf);
+
+				// NOTE: phase function cancels out since it is being perfectly importance sampled
+
+				// exit the Interface counting loop and keep scattering
+				break;
 			}
+#if SSS_USE_INTERFACE_COUNTING
+			// update counter
+			InterfaceCounter += ProbeResult.FrontFace ? +1 : -1;
+			if (InterfaceCounter != 0)
+			{
+				// we have not yet crossed the right number of interfaces, so ignore the current hit and try the segment again
+				Ray.TMin = asfloat(asuint(ProbeResult.HitT) + 1);
+			}
+			else
 #endif
-			// we hit the boundary! overwrite the current payload and exit the walk with a diffuse scattering event
-			Payload.WorldPos = Ray.Origin + ProbeResult.HitT * Ray.Direction;
-			Payload.WorldNormal = ProbeResult.WorldNormal;
-			// make sure normal is pointed outward from the object so we capture illumination from the exterior
-			if (ProbeResult.FrontFace != Payload.IsFrontFace())
 			{
-				Payload.WorldNormal *= -1.0;
+				// account for transmittance and scattering up to this point
+				float3 Transmittance = exp(-ProbeResult.HitT * SigmaT);
+				PathThroughput *= Transmittance;
+				// account for transmittance pdf (which is the balance heuristic combination of transmittance sampling in each color channel)
+				float3 TransmittancePdf = Transmittance; // integral of pdf from HitT to infinity
+				PathThroughput /= dot(Transmittance, ColorChannelPdf);
+
+				// our short ray hit the geomtry - decide if we should exit or not
+#if SSS_USE_TIR
+				RandSample = RandomSequence_GenerateSample2D(RandSequence).xxy;
+				float3 WorldNormal = ProbeResult.WorldNormal;
+				float CosTheta = abs(dot(Ray.Direction, WorldNormal));
+				float Fresnel = FresnelReflectance(CosTheta, 1.0 / 1.4);
+				if (RandSample.x < Fresnel)
+				{
+					// internal reflection occured -- reflect and keep tracing
+					// NOTE: weight and probability cancel out, so no throughput adjustment is needed
+					Ray.Origin += ProbeResult.HitT * Ray.Direction;
+					Ray.Direction = reflect(Ray.Direction, WorldNormal);
+					Ray.TMin = 0;
+					ApplyRayBias(Ray, ProbeResult.HitT, ProbeResult.WorldGeoNormal);
+					// Reset interface counter as if we had just entered the model (taking into account that we are seeing it from the opposite side now)
+					InterfaceCounter = ProbeResult.FrontFace ? -1 : +1;
+					// exit from the interface counting loop and keep scattering
+					break;
+				}
+#endif
+				// we hit the boundary! overwrite the current payload and exit the walk with a diffuse scattering event
+				Payload.WorldPos = Ray.Origin + ProbeResult.HitT * Ray.Direction;
+				Payload.WorldNormal = ProbeResult.WorldNormal;
+				Payload.WorldGeoNormal = ProbeResult.WorldGeoNormal;
+				// make sure normal is pointed outward from the object so we capture illumination from the exterior
+				if (ProbeResult.FrontFace != Payload.IsFrontFace())
+				{
+					Payload.WorldNormal *= -1.0;
+					Payload.WorldGeoNormal *= -1.0;
+				}
+				Payload.ShadingModelID = SHADINGMODELID_NUM; // invalid value so that we get diffuse shading
+				Payload.BlendingMode = RAY_TRACING_BLEND_MODE_OPAQUE;
+				Payload.BaseColor = 1;
+				return true;
 			}
-			Payload.ShadingModelID = SHADINGMODELID_NUM; // invalid value so that we get diffuse shading
-			Payload.BlendingMode = RAY_TRACING_BLEND_MODE_OPAQUE;
-			Payload.BaseColor = 1;
-			return true;
 		}
 	}
 	// we scattered a bunch of times and never hit anything -- give up
@@ -553,6 +602,24 @@ FPathState CreatePathState(uint2 DispatchIdx, uint2 DispatchDim)
 	return Output;
 }
 
+float3 AdjustShadingNormal(float3 ShadingNormal, float3 GeoNormal, float3 RayDirection)
+{
+	// Clip shading normal in a view dependent way such that the reflection stays above the geometric normal
+	// This introduces a bit of view-dependency to the shading normal but fixes dark artifacts around grazing angles
+
+	float3 D = RayDirection;
+	float3 R = reflect(D, ShadingNormal);
+
+	// https://iquilezles.org/www/articles/dontflip/dontflip.htm
+	float k = dot(R, GeoNormal);
+	if (k < 0.0)
+	{
+		return normalize(normalize(R - k * GeoNormal) - D);
+	}
+	return ShadingNormal;
+}
+
+
 bool PathTracingKernel(inout FPathState PathState, int Bounce)
 {
 	// This array will hold a CDF for light picking
@@ -593,6 +660,9 @@ bool PathTracingKernel(inout FPathState PathState, int Bounce)
 		return false;
 	}
 
+	// fix dark edges in reflections from shading normals
+	Payload.WorldNormal = AdjustShadingNormal(Payload.WorldNormal, Payload.WorldGeoNormal, PathState.Ray.Direction);
+
 	// As soon as the path is blurry enough, we can get away with diffuse sampling only
 	const bool bSimplifySSS = PathState.PathRoughness >= 0.15;
 
@@ -669,14 +739,14 @@ bool PathTracingKernel(inout FPathState PathState, int Bounce)
 				if (CastsShadow(LightId))
 				{
 					// for transmissive materials, bias the position to the other side of the surface if the light is coming from behind
-					const float SignedPositionBias = IsTransmissiveMaterial ? sign(dot(Payload.WorldNormal, LightSample.Direction)) : 1.0;
+					const float SignedPositionBias = IsTransmissiveMaterial ? sign(dot(Payload.WorldGeoNormal, LightSample.Direction)) : 1.0;
 
 					RayDesc LightRay;
 					LightRay.Origin = WorldPos;
 					LightRay.TMin = 0;
 					LightRay.Direction = LightSample.Direction;
 					LightRay.TMax = LightSample.Distance;
-					ApplyPositionBias(LightRay, SignedPositionBias * Payload.WorldNormal, MaxNormalBias);
+					ApplyRayBias(LightRay, Payload.HitT, SignedPositionBias * Payload.WorldGeoNormal, true);
 
 					float AvgRoughness = ApproximateCaustics ? GetAverageRoughness(Payload) : 0.0;
 
@@ -755,7 +825,7 @@ bool PathTracingKernel(inout FPathState PathState, int Bounce)
 	PathState.Ray.Direction = MaterialSample.Direction;
 	PathState.Ray.TMin = 0;
 	PathState.Ray.TMax = POSITIVE_INFINITY;
-	ApplyPositionBias(PathState.Ray, MaterialSample.PositionBiasSign * Payload.WorldNormal, MaxNormalBias);
+	ApplyRayBias(PathState.Ray, Payload.HitT, MaterialSample.PositionBiasSign * Payload.WorldGeoNormal);
 
 	if (ApproximateCaustics)
 	{
diff --git a/Engine/Shaders/Private/PathTracing/PathTracingMaterialHitShader.usf b/Engine/Shaders/Private/PathTracing/PathTracingMaterialHitShader.usf
--- a/Engine/Shaders/Private/PathTracing/PathTracingMaterialHitShader.usf
+++ b/Engine/Shaders/Private/PathTracing/PathTracingMaterialHitShader.usf
@@ -46,7 +46,8 @@ RAY_TRACING_ENTRY_CLOSEST_HIT(PathTracingMaterialCHS,
 	float4 SvPosition = float4(PixelCoord.xy, 0.0, 1.0);
 
 	FRayCone PropagatedCone = PropagateRayCone(PackedPayload.GetRayCone(), 0 /* surface curvature */, RayTCurrent());
-	CalcInterpolants(PropagatedCone, Attributes, Interpolants);
+	float3 GeoNormal = 0;
+	CalcInterpolants(PropagatedCone, Attributes, Interpolants, GeoNormal);
 
 	FMaterialPixelParameters MaterialParameters = GetMaterialPixelParameters(Interpolants, SvPosition);
 	FPixelMaterialInputs PixelMaterialInputs;
@@ -121,6 +122,10 @@ RAY_TRACING_ENTRY_CLOSEST_HIT(PathTracingMaterialCHS,
 	 * Set material attributes for full materials
 	 **/
 	Payload.WorldPos = MaterialParameters.AbsoluteWorldPosition;
+	// Because the geometric normal is computed directly in world space
+	// it doesn't reflect the sign flip from the object transform, so apply it here
+	float GeoNormalSign = MaterialParameters.TwoSidedSign * GetPrimitive_DeterminantSign(MaterialParameters.PrimitiveId);
+	Payload.WorldGeoNormal = GeoNormalSign * GeoNormal;
 	Payload.WorldNormal = normalize(MaterialParameters.WorldNormal);
 	Payload.Radiance = GetMaterialEmissive(PixelMaterialInputs);
 	Payload.BaseColor = BaseColor;
diff --git a/Engine/Shaders/Private/RayTracing/RayTracingCalcInterpolants.ush b/Engine/Shaders/Private/RayTracing/RayTracingCalcInterpolants.ush
--- a/Engine/Shaders/Private/RayTracing/RayTracingCalcInterpolants.ush
+++ b/Engine/Shaders/Private/RayTracing/RayTracingCalcInterpolants.ush
@@ -2,7 +2,7 @@
 
 #pragma once
 
-void CalcInterpolants(in FRayCone RayCone, in FDefaultAttributes Attributes, out FVertexFactoryInterpolantsVSToPS Interpolants)
+void CalcInterpolants(in FRayCone RayCone, in FDefaultAttributes Attributes, out FVertexFactoryInterpolantsVSToPS Interpolants, out float3 FaceNormal)
 {
 	FVertexFactoryRayTracingInterpolants Interpolated = (FVertexFactoryRayTracingInterpolants)0;
 
@@ -32,16 +32,16 @@ void CalcInterpolants(in FRayCone RayCone, in FDefaultAttributes Attributes, out
 
 	Interpolants = VertexFactoryAssignInterpolants(Interpolated);
 
-#if (NUM_TEX_COORD_INTERPOLATORS || USE_PARTICLE_SUBUVS) && !VERTEX_FACTORY_MODIFIES_INTERPOLATION
-	float2 TA = TexCoords[1] - TexCoords[0];
-	float2 TB = TexCoords[2] - TexCoords[0];
-	float UvArea = abs((TA.x * TB.y) - (TA.y * TB.x));
-
 	float3 PA = WorldPositions[1] - WorldPositions[0];
 	float3 PB = WorldPositions[2] - WorldPositions[0];
 	float3 Unnormalized = cross(PB, PA);
 	float InvWorldArea = rsqrt(dot(Unnormalized, Unnormalized));
-	float3 FaceNormal = Unnormalized * InvWorldArea;
+	FaceNormal = Unnormalized * InvWorldArea;
+
+#if (NUM_TEX_COORD_INTERPOLATORS || USE_PARTICLE_SUBUVS) && !VERTEX_FACTORY_MODIFIES_INTERPOLATION
+	float2 TA = TexCoords[1] - TexCoords[0];
+	float2 TB = TexCoords[2] - TexCoords[0];
+	float UvArea = abs((TA.x * TB.y) - (TA.y * TB.x));
 
 	float RayConeCapArea = RayCone.Width * RayCone.Width * (1.0f / (2.0f * PI));
 	float3 RayDir = WorldRayDirection();
@@ -51,3 +51,10 @@ void CalcInterpolants(in FRayCone RayCone, in FDefaultAttributes Attributes, out
 	GlobalRayCone_TexArea = max(GlobalRayCone_TexArea, SmallValue);
 #endif
 }
+
+// simplified call when Geometric Normal is not needed
+void CalcInterpolants(in FRayCone RayCone, in FDefaultAttributes Attributes, out FVertexFactoryInterpolantsVSToPS Interpolants)
+{
+	float3 GeoNormal = 0; // unused
+	CalcInterpolants(RayCone, Attributes, Interpolants, GeoNormal);
+}
\ No newline at end of file
