Implement Mixed Radix Convolution Bloom

Original author: Jess.kube

diff --git a/Engine/Shaders/Private/GPUFastFourierTransform.usf b/Engine/Shaders/Private/GPUFastFourierTransform.usf
--- a/Engine/Shaders/Private/GPUFastFourierTransform.usf
+++ b/Engine/Shaders/Private/GPUFastFourierTransform.usf
@@ -14,7 +14,11 @@ GPUFastFourierTransform.usf: Core Fast Fourier Transform Code
 // SCAN_LINE_LENGTH
 
 // Map get the RADIX and NUMTHREADSX as a function of SCAN_LINE_LENGTH
-  
+
+#if MIXED_RADIX == 1 && SCAN_LINE_LENGTH > 8
+	#define RADIX 8
+#else
+
 #if SCAN_LINE_LENGTH == 2
 	#define RADIX  2
 #endif
@@ -63,6 +67,7 @@ GPUFastFourierTransform.usf: Core Fast Fourier Transform Code
 	#define RADIX 8
 #endif
 
+#endif
 
 // NB: This is only needed for the not group shared stuff below.
 // Need to make that work and re-factor to fix this.
diff --git a/Engine/Shaders/Private/GPUFastFourierTransformCore.ush b/Engine/Shaders/Private/GPUFastFourierTransformCore.ush
--- a/Engine/Shaders/Private/GPUFastFourierTransformCore.ush
+++ b/Engine/Shaders/Private/GPUFastFourierTransformCore.ush
@@ -541,6 +541,56 @@ void Butterfly(bool bIsForward, inout Complex Local[RADIX], uint ThreadIdx, uint
 	}
 }
 
+void ButterflyRadix4(bool bIsForward, inout Complex Local[RADIX], uint ThreadIdx, uint Length)
+{
+
+	float AngleEven = TWO_PI * float(ThreadIdx) / float(Length * RADIX);
+	float AngleOdd = TWO_PI * float(ThreadIdx+Length) / float(Length * RADIX);
+	if (!bIsForward)
+	{
+		AngleEven *= -1;
+		AngleOdd *= -1;
+	}
+
+		
+	Complex TwiddleIncEven;
+	sincos(AngleEven, TwiddleIncEven.y, TwiddleIncEven.x);
+	Complex TwiddleIncOdd;
+	sincos(AngleOdd, TwiddleIncOdd.y, TwiddleIncOdd.x);
+	Complex TwiddleEven = TwiddleIncEven;
+	Complex TwiddleOdd = TwiddleIncOdd;
+	for (uint r = 2; r < RADIX; r+=2)
+	{
+		Local[r] = ComplexMult(TwiddleEven, Local[r]);
+		TwiddleEven = ComplexMult(TwiddleEven, TwiddleIncEven);
+		Local[r+1] = ComplexMult(TwiddleOdd, Local[r+1]);
+		TwiddleOdd = ComplexMult(TwiddleOdd, TwiddleIncOdd);
+	}
+}
+
+void ButterflyRadix2(bool bIsForward, inout Complex Local[RADIX], uint ThreadIdx, uint Length)
+{
+
+	float AngleStart = TWO_PI * float(ThreadIdx) / float(Length * RADIX);
+	float AngleInc = TWO_PI * float(Length) / float(Length * RADIX);
+	if (!bIsForward)
+	{
+		AngleStart *= -1;
+		AngleInc *= -1;
+	}
+
+		
+	Complex TwiddleInc;
+	sincos(AngleInc, TwiddleInc.y, TwiddleInc.x);
+	Complex Twiddle;
+	sincos(AngleStart, Twiddle.y, Twiddle.x);
+	for (uint r = 4; r < RADIX; ++r)
+	{
+		Local[r] = ComplexMult(Twiddle, Local[r]);
+		Twiddle = ComplexMult(Twiddle, TwiddleInc);
+	}
+}
+
 // Performs a single pass Stockham FFT using group shared memory.
 void GroupSharedFFT(in const bool bIsForward, inout Complex Local[RADIX], in const uint ArrayLength, in const uint ThreadIdx)
 {
@@ -605,11 +655,27 @@ void GroupSharedFFT(in const bool bIsForward, inout Complex Local[RADIX], in con
 		TransposeData(Local, IdxD, Ns, IdxS, NumCols);
 	}
 	
+#if MIXED_RADIX == 0 || SCAN_LINE_LENGTH == 4096 || SCAN_LINE_LENGTH == 512 || SCAN_LINE_LENGTH == 64 || SCAN_LINE_LENGTH <= 8
 	Butterfly(bIsForward, Local, ThreadIdx, Ns);
 
 	// Transform these RADIX values.
 
 	RadixFFT(bIsForward, Local);
+#elif SCAN_LINE_LENGTH == 2048 || SCAN_LINE_LENGTH == 256 || SCAN_LINE_LENGTH == 32
+	ButterflyRadix4(bIsForward, Local, ThreadIdx, NumCols);
+	
+	Radix4FFT(bIsForward, Local[0], Local[2], Local[4], Local[6]);
+	Radix4FFT(bIsForward, Local[1], Local[3], Local[5], Local[7]);
+#elif SCAN_LINE_LENGTH == 1024 || SCAN_LINE_LENGTH == 128 || SCAN_LINE_LENGTH == 16
+	ButterflyRadix2(bIsForward, Local, ThreadIdx, NumCols);
+	
+	Radix2FFT(bIsForward, Local[0], Local[4]);
+	Radix2FFT(bIsForward, Local[1], Local[5]);
+	Radix2FFT(bIsForward, Local[2], Local[6]);
+	Radix2FFT(bIsForward, Local[3], Local[7]);
+#else
+	#error Unknown SCAN_LINE_LENGTH
+#endif
 	
 	// not needed for the simple case
 	
diff --git a/Engine/Source/Runtime/Renderer/Private/GPUFastFourierTransform.cpp b/Engine/Source/Runtime/Renderer/Private/GPUFastFourierTransform.cpp
--- a/Engine/Source/Runtime/Renderer/Private/GPUFastFourierTransform.cpp
+++ b/Engine/Source/Runtime/Renderer/Private/GPUFastFourierTransform.cpp
@@ -749,6 +749,10 @@ namespace GPUFFT
 			FGlobalShader::ModifyCompilationEnvironment(Parameters, OutEnvironment);
 			OutEnvironment.SetDefine(TEXT("INCLUDE_GROUP_SHARED_COMPLEX_FFT"), 1);
 			OutEnvironment.SetDefine(TEXT("SCAN_LINE_LENGTH"), PowRadixSignalLength);
+			if (PowRadixSignalLength > 8)
+			{
+				OutEnvironment.SetDefine(TEXT("MIXED_RADIX"), 1);
+			}
 		}
 
 		/** Default constructor **/
@@ -786,6 +790,10 @@ namespace GPUFFT
 			FGlobalShader::ModifyCompilationEnvironment(Parameters, OutEnvironment);
 			OutEnvironment.SetDefine(TEXT("INCLUDE_GROUP_SHARED_TWO_FOR_ONE_FFT"), 1);
 			OutEnvironment.SetDefine(TEXT("SCAN_LINE_LENGTH"), PowRadixSignalLength);
+			if (PowRadixSignalLength > 8)
+			{
+				OutEnvironment.SetDefine(TEXT("MIXED_RADIX"), 1);
+			}
 		}
 
 		/** Default constructor **/
@@ -936,6 +944,10 @@ namespace GPUFFT
 			FGlobalShader::ModifyCompilationEnvironment(Parameters, OutEnvironment);
 			OutEnvironment.SetDefine(TEXT("INCLUDE_GROUP_SHARED_CONVOLUTION_WITH_TEXTURE"), 1);
 			OutEnvironment.SetDefine(TEXT("SCAN_LINE_LENGTH"), PowRadixSignalLength);
+			if (PowRadixSignalLength > 8)
+			{
+				OutEnvironment.SetDefine(TEXT("MIXED_RADIX"), 1);
+			}
 		}
 
 		/** Default constructor **/
