Fixes error on some shader compiler in Histogram.usf

diff --git a/Engine/Shaders/Private/Histogram.usf b/Engine/Shaders/Private/Histogram.usf
--- a/Engine/Shaders/Private/Histogram.usf
+++ b/Engine/Shaders/Private/Histogram.usf
@@ -227,13 +227,13 @@ void ScatterWeightsToLDS(uint Bucket0, uint Bucket1, float Weight0, float Weight
 //Common Definitions for Interlocked Adds to textures.
 //This is really just hear so we can call the same function regardless of 64 or 32 bit
 #if DO_64BIT_ATOMICS 
-	typedef uint64_t GatherType;
+	#define GatherType uint64_t
 	void InterlockedAddToTexture(uint index, GatherType V)
 	{
 		InterlockedAdd(HistogramScatter64Output[uint2(index,0)],V);
 	}
 #else
-	typedef uint GatherType;
+	#define GatherType uint
 	void InterlockedAddToTexture(uint index, GatherType V)
 	{
 		uint OldV;
