Fixes FFT bloom sampling outside input boundaries

diff --git a/Engine/Shaders/Private/GPUFastFourierTransform2DCore.ush b/Engine/Shaders/Private/GPUFastFourierTransform2DCore.ush
--- a/Engine/Shaders/Private/GPUFastFourierTransform2DCore.ush
+++ b/Engine/Shaders/Private/GPUFastFourierTransform2DCore.ush
@@ -124,7 +124,7 @@ void CopyDataSrcWindowToLocal(inout Complex LocalBuffer[2][RADIX], bool bIsHoriz
 		UNROLL
 		for (uint i = 0; i < RADIX ; ++i, Pixel.x += Stride)
 		{	
-			bool InWindow = !(Pixel.x > Window.z); 
+			bool InWindow = Pixel.x < Window.z; 
 			if (InWindow)
 			{ 
 				float4 SrcValue = SrcTexture[Pixel];
@@ -141,7 +141,7 @@ void CopyDataSrcWindowToLocal(inout Complex LocalBuffer[2][RADIX], bool bIsHoriz
 		UNROLL
 		for (uint i = 0; i < RADIX; ++i, Pixel.y += Stride)
 		{
-			bool InWindow = !(Pixel.y > Window.w);
+			bool InWindow = Pixel.y < Window.w;
 			if (InWindow)
 			{ 
 				float4 SrcValue = SrcTexture[Pixel];
@@ -167,7 +167,7 @@ void CopyDataSrcWindowToLocal(inout Complex LocalBuffer[2][RADIX], bool bIsHoriz
 		UNROLL
 		for (uint i = 0; i < RADIX; ++i, Pixel.x += Stride)
 		{	
-			bool InWindow = !(Pixel.x > WindowMax.x); 
+			bool InWindow = Pixel.x < WindowMax.x; 
 			if (InWindow)
 			{ 
 				float4 SrcValue = SrcTexture[Pixel];
@@ -182,7 +182,7 @@ void CopyDataSrcWindowToLocal(inout Complex LocalBuffer[2][RADIX], bool bIsHoriz
 		UNROLL
 		for (uint i = 0; i < RADIX; ++i, Pixel.y += Stride)
 		{
-			bool InWindow = !(Pixel.y > WindowMax.y);
+			bool InWindow = Pixel.y < WindowMax.y;
 			if (InWindow)
 			{	
 				float4 SrcValue = SrcTexture[Pixel];
