Fixes r.Histogram.UseAtomic=1

diff --git a/Engine/Shaders/Private/Histogram.usf b/Engine/Shaders/Private/Histogram.usf
--- a/Engine/Shaders/Private/Histogram.usf
+++ b/Engine/Shaders/Private/Histogram.usf
@@ -148,9 +148,8 @@ void CalculateBucketsAndWeights(in uint2 TexelPos, out float Weight0, out float
 {
 
 	float4 SceneColor = InputTexture.Load(int3(TexelPos, 0));
-	#if USE_PREEXPOSURE
-		SceneColor.xyz *= View.OneOverPreExposure;
-	#endif
+	SceneColor.xyz *= View.OneOverPreExposure;
+
 	float LuminanceVal = max(dot(SceneColor.xyz, float3(1.0f, 1.0f, 1.0f)/3.0f), EyeAdaptation_LuminanceMin);
 	float LogLuminance = ComputeHistogramPositionFromLuminance(LuminanceVal);
 	// Map the normalized histogram position into texels.
diff --git a/Engine/Source/Runtime/Renderer/Private/PostProcess/PostProcessHistogram.cpp b/Engine/Source/Runtime/Renderer/Private/PostProcess/PostProcessHistogram.cpp
--- a/Engine/Source/Runtime/Renderer/Private/PostProcess/PostProcessHistogram.cpp
+++ b/Engine/Source/Runtime/Renderer/Private/PostProcess/PostProcessHistogram.cpp
@@ -9,7 +9,7 @@
 #include "ShaderCompilerCore.h"
 
 TAutoConsoleVariable<int32> CVarUseAtomicHistogram(
-	TEXT("r.Histogram.UseAtomic"), 0,
+	TEXT("r.Histogram.UseAtomic"), 1,
 	TEXT("Uses atomic to speed up the generation of the histogram."),
 	ECVF_RenderThreadSafe);
 
