Fix VT feedback in the case where we read DBuffer directly in material instead of through the material decal response.
Previously any DBuffer feedback would be overwritten in the base pass.

diff --git a/Engine/Shaders/Private/VirtualTextureCommon.ush b/Engine/Shaders/Private/VirtualTextureCommon.ush
--- a/Engine/Shaders/Private/VirtualTextureCommon.ush
+++ b/Engine/Shaders/Private/VirtualTextureCommon.ush
@@ -17,10 +17,6 @@
 #define LIGHTMAP_VT_ENABLED 0
 #endif
 
-// Even fully opaque decals may not completely overwrite data from base pass, so don't allow VT feedback to become fully opaque
-// TODO - it may be possible to allow fully opaque decals with certain blend-mode/state combinations, if we can guarantee they overwrite all gbuffer channels
-#define VT_FEEDBACK_DECAL_MAX_OPACITY 0.50f
-
 /** Struct used to store feedback information in. */
 struct FVirtualTextureFeedbackParams
 {
@@ -69,7 +65,7 @@ void FinalizeVirtualTextureFeedback(in FVirtualTextureFeedbackParams Params, flo
 
 		// When using DBuffer, feedback buffer is blending using decal opacity for opaque objects
 		// This way we avoid overwriting values already written to feedback buffer while rendering dbuffer decals
-#if MATERIALBLENDING_ANY_TRANSLUCENT || IS_DECAL || (USE_DBUFFER && MATERIALDECALRESPONSEMASK)
+#if MATERIALBLENDING_ANY_TRANSLUCENT || IS_DECAL || (USE_DBUFFER && MATERIAL_DECAL_READ_MASK)
 
 		// Stochastic alpha testing
 #if MATERIALBLENDING_ADDITIVE || MATERIALBLENDING_MODULATE
@@ -79,13 +75,18 @@ void FinalizeVirtualTextureFeedback(in FVirtualTextureFeedbackParams Params, flo
 #else
 		// Opacity values smaller/larger than the reciprocal of the downscale factor get clamped so at least one pixel will resolve
 		const float VTDownscaleEpsilon = 1.0f / VIRTUAL_TEXTURE_FEEDBACK_FACTOR;
-		#if IS_DECAL
-			const float MaxAlpha = VT_FEEDBACK_DECAL_MAX_OPACITY;
-		#else
-			const float MaxAlpha = 1.0f - VTDownscaleEpsilon;
-		#endif
+		float MaxAlpha = 1.0f - VTDownscaleEpsilon;
+
+		// Even fully opaque gbuffer decals may not completely overwrite data from base pass, so don't allow VT feedback to become fully opaque.
+		// TODO - it may be possible to allow fully opaque decals with certain blend-mode/state combinations, if we can guarantee they overwrite all gbuffer channels
+		// Another similar case is when reading the DBuffer directly (so without a blend alpha available) we don't wan't to allow the base pass to be fully opaque for VT feedback.
+#if (IS_DECAL && !IS_DBUFFER_DECAL) || (USE_DBUFFER && MATERIAL_DECAL_READ_MASK && !MATERIALDECALRESPONSEMASK)
+		MaxAlpha = 0.5f;
+#endif
+
 		const float Alpha =  clamp(Opacity, VTDownscaleEpsilon, MaxAlpha);                  
 #endif
+
         const float AlphaThreshold = frac( PseudoRandom(PixelPos) + // Random  value in 0-1 on 128 x 128 pixel grid
 											SvPosition.w + // Add in depth so we pick different thresholds on different depths 
 											(FrameNumber / (float)VIRTUAL_TEXTURE_FEEDBACK_FACTOR) // Add in framenumber for extra jitter so the pseudorandom pattern changes over time
diff --git a/Engine/Source/Runtime/Engine/Private/Materials/HLSLMaterialTranslator.cpp b/Engine/Source/Runtime/Engine/Private/Materials/HLSLMaterialTranslator.cpp
--- a/Engine/Source/Runtime/Engine/Private/Materials/HLSLMaterialTranslator.cpp
+++ b/Engine/Source/Runtime/Engine/Private/Materials/HLSLMaterialTranslator.cpp
@@ -1466,6 +1466,8 @@ void FHLSLMaterialTranslator::GetMaterialEnvironment(EShaderPlatform InPlatform,
 	OutEnvironment.SetDefine(TEXT("MATERIAL_FULLY_ROUGH"), bIsFullyRough || Material->IsFullyRough());
 	OutEnvironment.SetDefine(TEXT("MATERIAL_USES_ANISOTROPY"), bUsesAnisotropy && FDataDrivenShaderPlatformInfo::GetSupportsAnisotropicMaterials(InPlatform));
 
+	OutEnvironment.SetDefine(TEXT("MATERIAL_DECAL_READ_MASK"), MaterialCompilationOutput.UsedDBufferTextures);
+
 	// Count the number of VTStacks (each stack will allocate a feedback slot)
 	OutEnvironment.SetDefine(TEXT("NUM_VIRTUALTEXTURE_SAMPLES"), VTStacks.Num());
 
