Fix for MH-4139 - removed a clamp when indexing into a non-compacted buffer when evaluating LOD1 cluster AABBs

diff --git a/Engine/Shaders/Private/HairStrands/HairStrandsInterpolation.usf b/Engine/Shaders/Private/HairStrands/HairStrandsInterpolation.usf
--- a/Engine/Shaders/Private/HairStrands/HairStrandsInterpolation.usf
+++ b/Engine/Shaders/Private/HairStrands/HairStrandsInterpolation.usf
@@ -593,7 +593,7 @@ void ClusterAABBEvaluationCS(uint3 DispatchThreadId : SV_DispatchThreadID, uint3
 		if (ThreadReadOffset < IndexCount)
 		{
 			const float3 OutHairPositionOffset = RenderDeformedOffsetBuffer[0].xyz;
-			const uint VertexIndex = clamp(ClusterVertexIdBuffer[IndexOffset + ThreadReadOffset], 0u, VertexCountAfterCulling-1);
+			const uint VertexIndex = ClusterVertexIdBuffer[IndexOffset + ThreadReadOffset];
 			const FHairControlPoint CP = UnpackHairControlPoint(RenderDeformedPositionBuffer[VertexIndex], OutHairPositionOffset, 1, 1, 1);
 
 			if (all(IsFinite(CP.Position)))
