Make mobile HQ BRDF to be enabled per material instead of globally.

diff --git a/Engine/Shaders/Private/MobileShadingModels.ush b/Engine/Shaders/Private/MobileShadingModels.ush
--- a/Engine/Shaders/Private/MobileShadingModels.ush
+++ b/Engine/Shaders/Private/MobileShadingModels.ush
@@ -226,16 +226,29 @@ void InitShadingModelContext(inout FMobileShadingModelContext ShadingModelContex
 
 half Vis_SmithJointApprox_Mobile(half a, half NoV, half NoL)
 {
-	half Vis_SmithV = NoL * (NoV * (1 - a) + a);
-	half Vis_SmithL = NoV * (NoL * (1 - a) + a);
-	return 0.5 * rcp(Vis_SmithV + Vis_SmithL);
+	float Vis_SmithV = NoL * (NoV * (1 - a) + a);
+	float Vis_SmithL = NoV * (NoL * (1 - a) + a);
+	return 0.5f * rcp(max(Vis_SmithV + Vis_SmithL, 0.0000153f));
+}
+
+// [Schlick 1994, "An Inexpensive BRDF Model for Physically-Based Rendering"]
+half3 F_Schlick_Mobile(half3 SpecularColor, half VoH)
+{
+	half OneMinusVoH = 1 - VoH;
+	half Fc = OneMinusVoH * OneMinusVoH;
+	Fc = Fc * Fc;
+	Fc = Fc * OneMinusVoH;
+	//return Fc + (1 - Fc) * SpecularColor;		// 1 add, 3 mad
+
+	// Anything less than 2% is physically impossible and is instead considered to be shadowing
+	return saturate(50.0 * SpecularColor.g) * Fc + (1 - Fc) * SpecularColor;
 }
 
 half3 MobileSpecularGGXInner(half D, half3 EnvBrdf, half3 SpecularColor, half Roughness, half NoV, half NoL, half VoH)
 {
 #if MOBILE_HIGH_QUALITY_BRDF
-	half Vis = Vis_SmithJointApprox_Mobile(Pow2(Roughness), NoV, NoL);
-	half3 F = F_Schlick(SpecularColor, VoH);
+	half Vis = Vis_SmithJointApprox_Mobile(Roughness * Roughness, NoV, NoL);
+	half3 F = F_Schlick_Mobile(SpecularColor, VoH);
 #else
 	half Vis = (Roughness * 0.25 + 0.25);
 	half3 F = EnvBrdf;
diff --git a/Engine/Source/Runtime/Core/Private/HAL/ConsoleManager.cpp b/Engine/Source/Runtime/Core/Private/HAL/ConsoleManager.cpp
--- a/Engine/Source/Runtime/Core/Private/HAL/ConsoleManager.cpp
+++ b/Engine/Source/Runtime/Core/Private/HAL/ConsoleManager.cpp
@@ -2876,13 +2876,4 @@ static TAutoConsoleVariable<int32> CVarMobileSupportsGen4TAA(
 	TEXT("Support desktop Gen4 TAA with mobile rendering\n")
 	TEXT("0: Fallback to FXAA")
 	TEXT("1: Support Desktop Gen4 TAA (default)"),
-	ECVF_ReadOnly | ECVF_RenderThreadSafe);
-
-static TAutoConsoleVariable<int32> CVarMobileEnableHighQualityBRDF(
-	TEXT("r.Mobile.EnableHighQualityBRDF"),
-	0,
-	TEXT("Use the high quality brdf functions on mobile to get better visual effects but adds GPU cost.\n")
-	TEXT("0: Disabled (More approximations but better performance. [default])\n")
-	TEXT("1: Enabled (More accurate but increases cost on GPU.)\n"),
-	ECVF_ReadOnly | ECVF_RenderThreadSafe
-);
\ No newline at end of file
+	ECVF_ReadOnly | ECVF_RenderThreadSafe);
\ No newline at end of file
diff --git a/Engine/Source/Runtime/Engine/Classes/Engine/RendererSettings.h b/Engine/Source/Runtime/Engine/Classes/Engine/RendererSettings.h
--- a/Engine/Source/Runtime/Engine/Classes/Engine/RendererSettings.h
+++ b/Engine/Source/Runtime/Engine/Classes/Engine/RendererSettings.h
@@ -922,12 +922,6 @@ class ENGINE_API URendererSettings : public UDeveloperSettings
 		ConfigRestartRequired = true))
 		uint32 bMobileSupportsGen4TAA : 1;
 
-	UPROPERTY(config, EditAnywhere, Category = Mobile, meta = (
-		ConsoleVariable = "r.Mobile.EnableHighQualityBRDF", DisplayName = "Mobile High Quality BRDF",
-		ToolTip = "Use the high quality brdf functions on mobile to get better visual effects but adds GPU cost. Changing this setting requires restarting the editor.",
-		ConfigRestartRequired = true))
-		uint32 bMobileEnableHighQualityBRDF : 1;
-
 	UPROPERTY(config, EditAnywhere, Category="Experimental|LOD Streaming|Skeletal Mesh", meta=(
 		DisplayName="Stream LODs by default",
 		ToolTip="Whether to stream skeletal mesh LODs by default."))
diff --git a/Engine/Source/Runtime/Engine/Classes/Materials/Material.h b/Engine/Source/Runtime/Engine/Classes/Materials/Material.h
--- a/Engine/Source/Runtime/Engine/Classes/Materials/Material.h
+++ b/Engine/Source/Runtime/Engine/Classes/Materials/Material.h
@@ -784,6 +784,10 @@ public:
 	UPROPERTY(EditAnywhere, BlueprintReadOnly, Category=Mobile)
 	uint8 bUseLightmapDirectionality : 1;
 
+	/* Use the high quality brdf functions on mobile to get better visual effects but adds GPU cost. */
+	UPROPERTY(EditAnywhere, BlueprintReadOnly, Category=Mobile, meta = (DisplayName = "Mobile High Quality BRDF"))
+	uint8 bMobileEnableHighQualityBRDF : 1;
+
 	/* Use alpha to coverage for masked material on mobile, make sure MSAA is enabled as well. */
 	UPROPERTY(EditAnywhere, BlueprintReadOnly, Category=Mobile, AdvancedDisplay, meta = (EditCondition = "BlendMode != EBlendMode::BLEND_Opaque"))
 	uint8 bUseAlphaToCoverage : 1;
diff --git a/Engine/Source/Runtime/Engine/Private/Materials/Material.cpp b/Engine/Source/Runtime/Engine/Private/Materials/Material.cpp
--- a/Engine/Source/Runtime/Engine/Private/Materials/Material.cpp
+++ b/Engine/Source/Runtime/Engine/Private/Materials/Material.cpp
@@ -4527,6 +4527,7 @@ bool UMaterial::CanEditChange(const FProperty* InProperty) const
 			PropertyName == GET_MEMBER_NAME_STRING_CHECKED(UMaterial, bUseLightmapDirectionality) ||
 			PropertyName == GET_MEMBER_NAME_STRING_CHECKED(UMaterial, bUseHQForwardReflections) ||
 			PropertyName == GET_MEMBER_NAME_STRING_CHECKED(UMaterial, bForwardBlendsSkyLightCubemaps) ||
+			PropertyName == GET_MEMBER_NAME_STRING_CHECKED(UMaterial, bMobileEnableHighQualityBRDF) ||
 			PropertyName == GET_MEMBER_NAME_STRING_CHECKED(UMaterial, bUsePlanarForwardReflections)
 			)
 		{
diff --git a/Engine/Source/Runtime/Engine/Private/Materials/MaterialShared.cpp b/Engine/Source/Runtime/Engine/Private/Materials/MaterialShared.cpp
--- a/Engine/Source/Runtime/Engine/Private/Materials/MaterialShared.cpp
+++ b/Engine/Source/Runtime/Engine/Private/Materials/MaterialShared.cpp
@@ -1495,6 +1495,11 @@ bool FMaterialResource::UseLmDirectionality() const
 	return Material->bUseLightmapDirectionality;
 }
 
+bool FMaterialResource::IsMobileHighQualityBRDFEnabled() const
+{
+	return Material->bMobileEnableHighQualityBRDF;
+}
+
 /**
  * Should shaders compiled for this material be saved to disk?
  */
@@ -2047,6 +2052,7 @@ void FMaterial::SetupMaterialEnvironment(
 	OutEnvironment.SetDefine(TEXT("MATERIAL_OUTPUT_OPACITY_AS_ALPHA"), GetBlendableOutputAlpha());
 	OutEnvironment.SetDefine(TEXT("TRANSLUCENT_SHADOW_WITH_MASKED_OPACITY"), GetCastDynamicShadowAsMasked());
 	OutEnvironment.SetDefine(TEXT("MATERIAL_USE_ALPHA_TO_COVERAGE"), IsUsingAlphaToCoverage());
+	OutEnvironment.SetDefine(TEXT("MOBILE_HIGH_QUALITY_BRDF"), IsMobileHighQualityBRDFEnabled());
 
 	EMaterialFloatPrecisionMode FloatPrecisionMode = GetMaterialFloatPrecisionMode();
 	if (FloatPrecisionMode == EMaterialFloatPrecisionMode::MFPM_Full)
diff --git a/Engine/Source/Runtime/Engine/Public/MaterialShared.h b/Engine/Source/Runtime/Engine/Public/MaterialShared.h
--- a/Engine/Source/Runtime/Engine/Public/MaterialShared.h
+++ b/Engine/Source/Runtime/Engine/Public/MaterialShared.h
@@ -1693,6 +1693,7 @@ public:
 	virtual bool IsUsingPlanarForwardReflections() const { return false; }
 	virtual bool IsNonmetal() const { return false; }
 	virtual bool UseLmDirectionality() const { return true; }
+	virtual bool IsMobileHighQualityBRDFEnabled() const { return false; }
 	virtual bool IsMasked() const = 0;
 	virtual bool IsDitherMasked() const { return false; }
 	virtual bool AllowNegativeEmissiveColor() const { return false; }
@@ -2593,6 +2594,7 @@ public:
 	ENGINE_API virtual bool IsUsingPlanarForwardReflections() const override;
 	ENGINE_API virtual bool IsNonmetal() const override;
 	ENGINE_API virtual bool UseLmDirectionality() const override;
+	ENGINE_API virtual bool IsMobileHighQualityBRDFEnabled() const override;
 	ENGINE_API virtual enum EBlendMode GetBlendMode() const override;
 	ENGINE_API virtual enum ERefractionMode GetRefractionMode() const override;
 	ENGINE_API virtual uint32 GetMaterialDecalResponse() const override;
diff --git a/Engine/Source/Runtime/RenderCore/Private/Shader.cpp b/Engine/Source/Runtime/RenderCore/Private/Shader.cpp
--- a/Engine/Source/Runtime/RenderCore/Private/Shader.cpp
+++ b/Engine/Source/Runtime/RenderCore/Private/Shader.cpp
@@ -1560,11 +1560,6 @@ void ShaderMapAppendKeyString(EShaderPlatform Platform, FString& KeyString)
 		{
 			KeyString += IsMobileDistanceFieldEnabled(Platform) ? TEXT("_MobSDF") : TEXT("");
 		}
-
-		{
-			static IConsoleVariable* MobileEnableHighQualityBRDFCVar = IConsoleManager::Get().FindConsoleVariable(TEXT("r.Mobile.EnableHighQualityBRDF"));
-			KeyString += (MobileEnableHighQualityBRDFCVar && MobileEnableHighQualityBRDFCVar->GetInt() != 0) ? TEXT("_HQBRDF") : TEXT("");
-		}
 	}
 	else
 	{
diff --git a/Engine/Source/Runtime/Renderer/Private/MobileBasePassRendering.h b/Engine/Source/Runtime/Renderer/Private/MobileBasePassRendering.h
--- a/Engine/Source/Runtime/Renderer/Private/MobileBasePassRendering.h
+++ b/Engine/Source/Runtime/Renderer/Private/MobileBasePassRendering.h
@@ -447,11 +447,6 @@ public:
 		OutEnvironment.SetDefine(TEXT("ENABLE_AMBIENT_OCCLUSION"), IsMobileAmbientOcclusionEnabled(Parameters.Platform) ? 1u : 0u);
 
 		OutEnvironment.SetDefine(TEXT("ENABLE_DISTANCE_FIELD"), IsMobileDistanceFieldEnabled(Parameters.Platform));
-
-		static auto* MobileEnableHighQualityBRDFCVAR = IConsoleManager::Get().FindTConsoleVariableDataInt(TEXT("r.Mobile.EnableHighQualityBRDF"));
-		const bool bMobileEnableHighQualityBRDF = (MobileEnableHighQualityBRDFCVAR && MobileEnableHighQualityBRDFCVAR->GetValueOnAnyThread() == 1);
-
-		OutEnvironment.SetDefine(TEXT("MOBILE_HIGH_QUALITY_BRDF"), bMobileEnableHighQualityBRDF ? 1u : 0u);
 	}
 	
 	/** Initialization constructor. */
