Fix Android shader compile error.

diff --git a/Engine/Shaders/Private/OctahedralCommon.ush b/Engine/Shaders/Private/OctahedralCommon.ush
--- a/Engine/Shaders/Private/OctahedralCommon.ush
+++ b/Engine/Shaders/Private/OctahedralCommon.ush
@@ -29,7 +29,7 @@ float3 OctahedronToUnitVector( float2 Oct )
 {
 	float3 N = float3( Oct, 1 - dot( 1, abs(Oct) ) );
 	float t = max( -N.z, 0 );
-	N.xy += N.xy >= 0 ? -t : t;
+	N.xy += N.xy >= 0 ? float2(-t, -t) : float2(t, t);
 	return normalize(N);
 }
 
