Update path tracer payload encoding to reflect recent convention change on Octahedral encoding functions

diff --git a/Engine/Shaders/Private/PathTracing/PathTracingCommon.ush b/Engine/Shaders/Private/PathTracing/PathTracingCommon.ush
--- a/Engine/Shaders/Private/PathTracing/PathTracingCommon.ush
+++ b/Engine/Shaders/Private/PathTracing/PathTracingCommon.ush
@@ -21,9 +21,9 @@
 uint PayloadEncodeUnitVector(float3 Normal)
 {
 	const int BITS = PATH_TRACING_PAYLOAD_NORMAL_BITS;
-	const float Scale = float((1u << (BITS / 2)) - 2); // largest representable even number (so we can encode 0)
+	const float Scale = float((1u << (BITS / 2)) - 2) / 2; // largest representable even number (so we can encode 0)
 	float2 OctF = UnitVectorToOctahedron(Normal);
-	int2 Oct = int2(round(OctF * Scale));
+	int2 Oct = int2(round(OctF * Scale + Scale));
 	return Oct.x + (1u << (BITS / 2)) * Oct.y;
 }
 
@@ -35,7 +35,7 @@ float3 PayloadDecodeUnitVector(uint Encoded)
 
 	int2 Oct = int2(Encoded & Mask, (Encoded >> (BITS / 2)) & Mask);
 	float2 OctF = saturate(float2(Oct) / Scale);
-	return OctahedronToUnitVector(OctF);
+	return OctahedronToUnitVector(2 * OctF - 1);
 }
 
 // This payload structure is what we transport between RGS/CHS/AHS programs
