Fix depth range issue when accessing coarsed clipmap levels from cached pages

diff --git a/Engine/Shaders/Private/VirtualShadowMaps/ProjectionCommon.ush b/Engine/Shaders/Private/VirtualShadowMaps/ProjectionCommon.ush
--- a/Engine/Shaders/Private/VirtualShadowMaps/ProjectionCommon.ush
+++ b/Engine/Shaders/Private/VirtualShadowMaps/ProjectionCommon.ush
@@ -106,6 +106,7 @@ struct FVirtualShadowMapClipmapRelativeTransform
 
 // Resulting transform takes UV's/depths from clipmap level ClipmapId to ClipmapId + LevelOffset 
 // ClipmapId + LevelOffset should be a valid level in the same clipmap
+// See VirtualShadowMapClipmap.cpp for construction details
 FVirtualShadowMapClipmapRelativeTransform CalcClipmapRelativeTransform(int ClipmapId, int LevelOffset)
 {
 	const FVirtualShadowMapProjectionShaderData ProjectionDataA = GetVirtualShadowMapProjectionData(ClipmapId);
@@ -117,8 +118,12 @@ FVirtualShadowMapClipmapRelativeTransform CalcClipmapRelativeTransform(int Clipm
 	FVirtualShadowMapClipmapRelativeTransform Result;
 	Result.Scale = exp2(-float(LevelOffset));
 	Result.Bias.xy = 0.25f * (OffsetB - Result.Scale * OffsetA);
-	// NOTE: Currently no snapping in Z
-	Result.Bias.z = 0.5f - Result.Scale * 0.5f;
+
+	// NOTE: relative Z bias can change when caching is enabled due to cached levels pinning the depth range
+	float OffsetZA = ProjectionDataA.ShadowViewToClipMatrix[3][2];
+	float OffsetZB = ProjectionDataB.ShadowViewToClipMatrix[3][2];
+	Result.Bias.z = OffsetZB - Result.Scale * OffsetZA;
+
 	return Result;
 }
 
