Fixed world space normal encoding in FLumenMinimalPayload. +Y vector was wrapping to -Y.

diff --git a/Engine/Shaders/Private/RayTracing/RayTracingCommon.ush b/Engine/Shaders/Private/RayTracing/RayTracingCommon.ush
--- a/Engine/Shaders/Private/RayTracing/RayTracingCommon.ush
+++ b/Engine/Shaders/Private/RayTracing/RayTracingCommon.ush
@@ -157,14 +157,14 @@ struct FLumenMinimalPayload : FMinimalPayload
 	void SetWorldNormal(float3 WorldNormal)
 	{
 		float2 WorldNormalAsOctahedron = UnitVectorToOctahedron(WorldNormal);
-		uint2 QuantizedOctahedron = clamp(round((WorldNormalAsOctahedron * 0.5 + 0.5) * 65536.0), 0.0, 65536.0);
-		PackedData[2] = (QuantizedOctahedron.x & 0xFFFF) << 16 | (QuantizedOctahedron.y & 0xFFFF);
+		uint2 QuantizedOctahedron = clamp((WorldNormalAsOctahedron * 0.5 + 0.5) * 65535.0 + 0.5, 0.0, 65535.0);
+		PackedData[2] = ((QuantizedOctahedron.x & 0xFFFF) << 16) | (QuantizedOctahedron.y & 0xFFFF);
 	}
 
 	float3 GetWorldNormal()
 	{
 		uint2 QuantizedOctahedron = uint2((PackedData[2] >> 16) & 0xFFFF, PackedData[2] & 0xFFFF);
-		float2 WorldNormalAsOctahedron = ((QuantizedOctahedron / 65536.0) - 0.5) * 2.0;
+		float2 WorldNormalAsOctahedron = ((QuantizedOctahedron / 65535.0) - 0.5) * 2.0;
 		float3 WorldNormal = OctahedronToUnitVector(WorldNormalAsOctahedron);
 		return WorldNormal;
 	}
