Do explicit conversion from uint to float in DecodeTransform() to avoid shader complation errors on GL

diff --git a/Engine/Shaders/Private/SceneData.ush b/Engine/Shaders/Private/SceneData.ush
--- a/Engine/Shaders/Private/SceneData.ush
+++ b/Engine/Shaders/Private/SceneData.ush
@@ -370,9 +370,9 @@ float4x4 DecodeTransform( uint4 RotationScale, float3 Translation, inout float3
 	{
 		float3 Rotation =
 		{
-			( RotationScale[0] >>  0 ) & 0xffff,
-			( RotationScale[0] >> 16 ) & 0xffff,
-			( RotationScale[1] >>  0 ) & 0x7fff
+			float(( RotationScale[0] >>  0 ) & 0xffff),
+			float(( RotationScale[0] >> 16 ) & 0xffff),
+			float(( RotationScale[1] >>  0 ) & 0x7fff)
 		};
 
 		float2 OctZ  = ( Rotation.xy - 32768 ) * (1.0f / 32767.0f);
