Fixed crash on iOS GPU caused by out of bounds access to a GPUScene lightmap data

diff --git a/Engine/Shaders/Private/InstanceCulling/BuildInstanceDrawCommands.usf b/Engine/Shaders/Private/InstanceCulling/BuildInstanceDrawCommands.usf
--- a/Engine/Shaders/Private/InstanceCulling/BuildInstanceDrawCommands.usf
+++ b/Engine/Shaders/Private/InstanceCulling/BuildInstanceDrawCommands.usf
@@ -24,6 +24,7 @@
 uint InstanceSceneDataSOAStride;
 uint GPUSceneNumInstances;
 uint GPUSceneNumPrimitives;
+uint GPUSceneNumLightmapDataItems;
 
 #ifndef BATCH_PROCESSING_MODE_NUM
 	#define BATCH_PROCESSING_MODE_NUM 1
@@ -84,15 +85,27 @@ void WritePackedInstanceData(uint Offset, uint InstanceId, FInstanceSceneData In
 {
 	FPrimitiveSceneData PrimitiveData = GetPrimitiveData(InstanceData.PrimitiveId);
 	
+	float4 LightMapUVScaleBias = float4(1,1,0,0);
+	float4 ShadowMapUVScaleBias = float4(1,1,0,0);
+
 	// FIXME: need LODIndex to compute correct LightmapDataIndex
 	uint LightmapDataIndex = PrimitiveData.LightmapDataIndex; // + LODIndex;
-	float4 LightMapUVScaleBias = GetLightmapData(LightmapDataIndex).LightMapCoordinateScaleBias;
-	float4 ShadowMapUVScaleBias = GetLightmapData(LightmapDataIndex).ShadowMapCoordinateScaleBias;
-	if (InstanceData.HasLightShadowUVBias)
+	if (LightmapDataIndex < GPUSceneNumLightmapDataItems)
+	{
+		LightMapUVScaleBias = GetLightmapData(LightmapDataIndex).LightMapCoordinateScaleBias;
+		ShadowMapUVScaleBias = GetLightmapData(LightmapDataIndex).ShadowMapCoordinateScaleBias;
+		if (InstanceData.HasLightShadowUVBias)
+		{
+			LightMapUVScaleBias.zw = InstanceData.LightMapAndShadowMapUVBias.xy;
+			ShadowMapUVScaleBias.zw = InstanceData.LightMapAndShadowMapUVBias.zw;
+		}
+	}
+	else
 	{
-		LightMapUVScaleBias.zw = InstanceData.LightMapAndShadowMapUVBias.xy;
-		ShadowMapUVScaleBias.zw = InstanceData.LightMapAndShadowMapUVBias.zw;
+		// avoid passing invalid lightmap data index into gfx shaders?
+		LightmapDataIndex = 0;
 	}
+
 	float4 PackedLightmapData = PackLightmapData(LightMapUVScaleBias, ShadowMapUVScaleBias);
 
 	const uint PackedPrimitiveFlags = ((PrimitiveData.Flags & 0xffff) | (InstanceData.Flags << 16u));
diff --git a/Engine/Source/Runtime/Renderer/Private/GPUScene.h b/Engine/Source/Runtime/Renderer/Private/GPUScene.h
--- a/Engine/Source/Runtime/Renderer/Private/GPUScene.h
+++ b/Engine/Source/Runtime/Renderer/Private/GPUScene.h
@@ -293,6 +293,7 @@ public:
 
 	int32 GetNumInstances() const { return InstanceSceneDataAllocator.GetMaxSize(); }
 	int32 GetNumPrimitives() const { return DynamicPrimitivesOffset; }
+	int32 GetNumLightmapDataItems() const { return LightmapDataAllocator.GetMaxSize(); }
 
 	bool bUpdateAllPrimitives;
 
diff --git a/Engine/Source/Runtime/Renderer/Private/InstanceCulling/InstanceCullingContext.cpp b/Engine/Source/Runtime/Renderer/Private/InstanceCulling/InstanceCullingContext.cpp
--- a/Engine/Source/Runtime/Renderer/Private/InstanceCulling/InstanceCullingContext.cpp
+++ b/Engine/Source/Runtime/Renderer/Private/InstanceCulling/InstanceCullingContext.cpp
@@ -209,6 +209,7 @@ public:
 		SHADER_PARAMETER(uint32, GPUSceneFrameNumber)
 		SHADER_PARAMETER(uint32, GPUSceneNumInstances)
 		SHADER_PARAMETER(uint32, GPUSceneNumPrimitives)
+		SHADER_PARAMETER(uint32, GPUSceneNumLightmapDataItems)
 
 		SHADER_PARAMETER_STRUCT_INCLUDE(FInstanceProcessingGPULoadBalancer::FShaderParameters, LoadBalancerParameters)
 
@@ -391,6 +392,7 @@ void FInstanceCullingContext::BuildRenderingCommands(
 	PassParametersTmp.GPUSceneFrameNumber = GPUScene.GetSceneFrameNumber();
 	PassParametersTmp.GPUSceneNumInstances = GPUScene.GetNumInstances();
 	PassParametersTmp.GPUSceneNumPrimitives = GPUScene.GetNumPrimitives();
+	PassParametersTmp.GPUSceneNumLightmapDataItems = GPUScene.GetNumLightmapDataItems();
 	PassParametersTmp.DynamicInstanceIdOffset = DynamicInstanceIdOffset;
 	PassParametersTmp.DynamicInstanceIdMax = DynamicInstanceIdOffset + DynamicInstanceIdNum;
 
@@ -688,6 +690,7 @@ FInstanceCullingDeferredContext *FInstanceCullingContext::CreateDeferredContext(
 	PassParametersTmp.GPUSceneFrameNumber = GPUScene.GetSceneFrameNumber();
 	PassParametersTmp.GPUSceneNumInstances = GPUScene.GetNumInstances();
 	PassParametersTmp.GPUSceneNumPrimitives = GPUScene.GetNumPrimitives();
+	PassParametersTmp.GPUSceneNumLightmapDataItems = GPUScene.GetNumLightmapDataItems();
 
 	PassParametersTmp.DrawCommandDescs = GraphBuilder.CreateSRV(DrawCommandDescsRDG);
 	PassParametersTmp.BatchInfos = GraphBuilder.CreateSRV(BatchInfosRDG);
