Fixing volumetric fog sparkle artifacts with virtual shadow map clipmaps
When a coarser clipmap is selected, compute the page address in the target level using only integer operations to avoid precision issues
Clamp the lookup coordinates to the target page (but being careful not to impact the optimal slope bias computation)

diff --git a/Engine/Shaders/Private/VirtualShadowMaps/PageAccessCommon.ush b/Engine/Shaders/Private/VirtualShadowMaps/PageAccessCommon.ush
--- a/Engine/Shaders/Private/VirtualShadowMaps/PageAccessCommon.ush
+++ b/Engine/Shaders/Private/VirtualShadowMaps/PageAccessCommon.ush
@@ -176,9 +176,8 @@ struct FShadowPageTranslationResult
 // Finds the best-resolution mapped page at the given UV
 FShadowPageTranslationResult ShadowVirtualToPhysicalUV(uint VirtualShadowMapID, float2 ShadowMapUV)
 {
-	uint2 vPage = uint2( ShadowMapUV * VSM_LEVEL0_DIM_PAGES_XY );
-
-	FShadowPhysicalPage PhysicalPageEntry = ShadowGetPhysicalPage( VirtualShadowMapID * VSM_PAGE_TABLE_SIZE + vPage.x + vPage.y * VSM_LEVEL0_DIM_PAGES_XY );
+	uint2 vPage = uint2(ShadowMapUV * VSM_LEVEL0_DIM_PAGES_XY);
+	FShadowPhysicalPage PhysicalPageEntry = ShadowGetPhysicalPage(CalcPageOffset(VirtualShadowMapID, 0, vPage));
 
 	FShadowPageTranslationResult Result;
 	Result.bValid = PhysicalPageEntry.bAnyLODValid;
diff --git a/Engine/Shaders/Private/VirtualShadowMaps/PageManagement.usf b/Engine/Shaders/Private/VirtualShadowMaps/PageManagement.usf
--- a/Engine/Shaders/Private/VirtualShadowMaps/PageManagement.usf
+++ b/Engine/Shaders/Private/VirtualShadowMaps/PageManagement.usf
@@ -675,8 +675,7 @@ void PropagateMappedMips(uint2 ThreadId : SV_DispatchThreadID, uint2 GroupId : S
 				FVirtualShadowMapProjectionShaderData LevelProjectionData = GetVirtualShadowMapProjectionData(ClipmapLevelId);								
 				int2 LevelOffset = OffsetScale * LevelProjectionData.ClipmapCornerOffset;
 
-				const int LevelScale = 1u << ClipmapOffset;
-				int2 LevelPage = (BasePage + LevelScale * LevelOffset) / LevelScale;
+				int2 LevelPage = (BasePage + (LevelOffset << ClipmapOffset)) >> ClipmapOffset;
 
 				if (all(LevelPage >= 0) && all(LevelPage < VSM_LEVEL0_DIM_PAGES_XY))
 				{
diff --git a/Engine/Shaders/Private/VirtualShadowMaps/ProjectionCommon.ush b/Engine/Shaders/Private/VirtualShadowMaps/ProjectionCommon.ush
--- a/Engine/Shaders/Private/VirtualShadowMaps/ProjectionCommon.ush
+++ b/Engine/Shaders/Private/VirtualShadowMaps/ProjectionCommon.ush
@@ -127,6 +127,20 @@ FVirtualShadowMapClipmapRelativeTransform CalcClipmapRelativeTransform(int Clipm
 	return Result;
 }
 
+// Transforms a virtual page in a given clipmap level to a page in the offset coarser level
+// This offset must be positive or else the page would not be guaranteed to exist in the target level
+// This is done entirely in integer math as well to avoid precision issues when looking up coarser fallback pages
+uint2 CalcClipmapOffsetLevelPage(uint2 BasePage, int ClipmapId, uint LevelOffset)
+{
+	const FVirtualShadowMapProjectionShaderData ProjectionDataA = GetVirtualShadowMapProjectionData(ClipmapId);
+	const FVirtualShadowMapProjectionShaderData ProjectionDataB = GetVirtualShadowMapProjectionData(ClipmapId + LevelOffset);
+
+	const int OffsetScale = (VSM_LEVEL0_DIM_PAGES_XY >> 2);
+	int2 BasePageOffset  = OffsetScale * ProjectionDataA.ClipmapCornerOffset;
+	int2 LevelPageOffset = OffsetScale * ProjectionDataB.ClipmapCornerOffset;
+	return (BasePage - BasePageOffset + (LevelPageOffset << LevelOffset)) >> LevelOffset;
+}
+
 // Will sample from the given clipmap level Id, onwards to coarser levels if no valid data is present
 // ShadowMapUVs are in terms of the given clipmap level/virtual shadow map
 FVirtualShadowMapSample SampleVirtualShadowMapClipmap(int VirtualShadowMapId, float2 ShadowMapUV)
@@ -135,16 +149,53 @@ FVirtualShadowMapSample SampleVirtualShadowMapClipmap(int VirtualShadowMapId, fl
 
 #define DEBUG_VSM_CLIPMAP_LEVEL_SEARCH 0
 #if !DEBUG_VSM_CLIPMAP_LEVEL_SEARCH
-	FShadowPageTranslationResult Page = ShadowVirtualToPhysicalUV(VirtualShadowMapId, ShadowMapUV);
-	if (Page.bValid)
-	{
-		FVirtualShadowMapClipmapRelativeTransform Transform = CalcClipmapRelativeTransform(VirtualShadowMapId, Page.LODOffset);
-		float2 ClipmapUV = ShadowMapUV * Transform.Scale + Transform.Bias.xy;
-		Result = SampleVirtualShadowMapLevel(VirtualShadowMapId + Page.LODOffset, ClipmapUV, 0);
-		if (Result.bValid)
+	FShadowPageTranslationResult Page = ShadowVirtualToPhysicalUV(VirtualShadowMapId, ShadowMapUV);	
+
+	uint2 BasePage = uint2(ShadowMapUV * VSM_LEVEL0_DIM_PAGES_XY);
+	FShadowPhysicalPage PhysicalPageEntry = ShadowGetPhysicalPage(CalcPageOffset(VirtualShadowMapId, 0, BasePage));
+	if (PhysicalPageEntry.bAnyLODValid)
+	{	
+		uint ClipmapLevelOffset = Page.LODOffset;
+		int ClipmapLevelId = VirtualShadowMapId + ClipmapLevelOffset;
+
+		Result.VirtualTexelAddressFloat = ShadowMapUV * float(CalcLevelDimsTexels(0));
+		Result.VirtualTexelAddress = uint2(Result.VirtualTexelAddressFloat);
+		float DepthLevelScale = 1.0f;
+		float DepthLevelBias = 0.0f;
+
+		// Need to use a coarser clipmap to find a valid page
+		if (ClipmapLevelOffset > 0)
+		{
+			// Compute the virtual page in the offset level by integer math and clamp to the edges to avoid any precision
+			// issues at borders that may cause us to miss the mapped page.
+			uint2 vPage = CalcClipmapOffsetLevelPage(BasePage, VirtualShadowMapId, ClipmapLevelOffset);
+			uint2 VirtualTexelAddressMin = vPage * VSM_PAGE_SIZE;
+			uint2 VirtualTexelAddressMax = VirtualTexelAddressMin + (VSM_PAGE_SIZE - 1);
+
+			FVirtualShadowMapClipmapRelativeTransform Transform = CalcClipmapRelativeTransform(VirtualShadowMapId, Page.LODOffset);
+			float2 ClipmapUV = ShadowMapUV * Transform.Scale + Transform.Bias.xy;
+			DepthLevelScale = Transform.Scale;
+			DepthLevelBias = Transform.Bias.z;
+
+			// NOTE: Do not clamp the float address as it messes with optimal slope bias calculations
+			Result.VirtualTexelAddressFloat = ClipmapUV * float(CalcLevelDimsTexels(0));
+			Result.VirtualTexelAddress = clamp(uint2(Result.VirtualTexelAddressFloat), VirtualTexelAddressMin, VirtualTexelAddressMax);
+
+			PhysicalPageEntry = ShadowGetPhysicalPage(CalcPageOffset(ClipmapLevelId, 0, vPage));
+		}
+
+		// NOTE: This really should be valid if we got here
+		if (PhysicalPageEntry.bThisLODValid)
 		{
+			Result.PhysicalTexelAddress = 
+				PhysicalPageEntry.PhysicalAddress * VSM_PAGE_SIZE +
+				(Result.VirtualTexelAddress & VSM_PAGE_SIZE_MASK);
+
 			// Convert depth into back into the original reference clipmap level range
-			Result.Depth = (Result.Depth - Transform.Bias.z) / Transform.Scale;
+			Result.Depth = (SampleVirtualShadowMapPhysicalDepth(Result.PhysicalTexelAddress) - DepthLevelBias) / DepthLevelScale;
+			Result.MipLevel = 0;
+			Result.VirtualShadowMapId = ClipmapLevelId;
+			Result.bValid = true;
 		}
 	}
 #else // DEBUG_VSM_CLIPMAP_LEVEL_SEARCH
@@ -172,7 +223,8 @@ float ComputeVirtualShadowMapOptimalSlopeBias(
 	int RequestedVirtualShadowMapId,
 	FVirtualShadowMapSample SmSample,
 	float3 WorldPosition,
-	float3 EstimatedGeoWorldNormal)
+	float3 EstimatedGeoWorldNormal,
+	bool bClamp = true)
 {
 	FVirtualShadowMapProjectionShaderData ProjectionData = GetVirtualShadowMapProjectionData(SmSample.VirtualShadowMapId);
 
@@ -189,7 +241,7 @@ float ComputeVirtualShadowMapOptimalSlopeBias(
 	float OptimalSlopeBias = 2.0f * max(0.0f, dot(DepthSlopeUV, TexelCenterOffsetUV));
 	
 	// Clamp to avoid excessive degenerate slope biases causing flickering lit pixels
-	OptimalSlopeBias = min(OptimalSlopeBias, abs(100.0f * ProjectionData.ShadowViewToClipMatrix._33));
+	OptimalSlopeBias = bClamp ? min(OptimalSlopeBias, abs(100.0f * ProjectionData.ShadowViewToClipMatrix._33)) : OptimalSlopeBias;
 
 	// Adjust depth scale if we sampled a different clipmap level
 	// NOTE: Sampled clipmap should always be >= the requested one (coarser)
