Removed array access for tex coords in RT

diff --git a/Engine/Shaders/Private/GeometryCacheVertexFactory.ush b/Engine/Shaders/Private/GeometryCacheVertexFactory.ush
--- a/Engine/Shaders/Private/GeometryCacheVertexFactory.ush
+++ b/Engine/Shaders/Private/GeometryCacheVertexFactory.ush
@@ -91,15 +91,18 @@ FVertexFactoryInput LoadVertexFactoryInputForHGS(uint TriangleIndex, int VertexI
     Input.MotionBlurData.z = GeomCacheMVF.MotionBlurData[VertexId * 3 + 2];
 	
 #if NUM_MATERIAL_TEXCOORDS_VERTEX
-	float TexCoordX = GeomCacheMVF.TexCoords[VertexId * 2 + 0];
-	float TexCoordY = GeomCacheMVF.TexCoords[VertexId * 2 + 1];
+	float2 TexCoordXY = float2(GeomCacheMVF.TexCoords[VertexId * 2 + 0], GeomCacheMVF.TexCoords[VertexId * 2 + 1]);
 
-	UNROLL
-	for (int CoordinateIndex = 0; CoordinateIndex < NUM_MATERIAL_TEXCOORDS_VERTEX; CoordinateIndex++)
-	{
-		Input.TexCoords[CoordinateIndex].x = TexCoordX;
-		Input.TexCoords[CoordinateIndex].y = TexCoordY;
-	}
+	Input.TexCoords0 = TexCoordXY;
+	#if NUM_MATERIAL_TEXCOORDS_VERTEX > 1
+		Input.TexCoords1 = TexCoordXY;
+	#endif
+	#if NUM_MATERIAL_TEXCOORDS_VERTEX > 2
+		Input.TexCoords2 = TexCoordXY;
+	#endif
+	#if NUM_MATERIAL_TEXCOORDS_VERTEX > 3
+		Input.TexCoords3 = TexCoordXY;
+	#endif
 #endif
 
 	return Input;
