Draw rounded boxes with atlased images properly, using the material tex coords (offset within the element) when computing the rounded edge

diff --git a/Engine/Shaders/Private/SlateElementPixelShader.usf b/Engine/Shaders/Private/SlateElementPixelShader.usf
--- a/Engine/Shaders/Private/SlateElementPixelShader.usf
+++ b/Engine/Shaders/Private/SlateElementPixelShader.usf
@@ -235,7 +235,7 @@ half4 GetBorderElementColor( VertexToPixelInterpolants InVertex )
 #if SHADER_TYPE == ST_RoundedBox
 half4 GetRoundedBoxElementColor( VertexToPixelInterpolants InVertex )
 {
-	return GetRoundedBoxElementColorInternal(ShaderParams.zw, InVertex.TextureCoordinates[0].xy, ShaderParams.y, ShaderParams2, GetColor(InVertex, GetUV(InVertex, 0) * GetUV(InVertex, 1)),  InVertex.SecondaryColor);
+	return GetRoundedBoxElementColorInternal(ShaderParams.zw, InVertex.MaterialTexCoords.xy, ShaderParams.y, ShaderParams2, GetColor(InVertex, GetUV(InVertex, 0) * GetUV(InVertex, 1)),  InVertex.SecondaryColor);
 }
 #endif
 
