Implements TSR's CONFIG_SCENE_COLOR_OVERFLOW

diff --git a/Engine/Shaders/Private/TemporalSuperResolution/TSRCommon.ush b/Engine/Shaders/Private/TemporalSuperResolution/TSRCommon.ush
--- a/Engine/Shaders/Private/TemporalSuperResolution/TSRCommon.ush
+++ b/Engine/Shaders/Private/TemporalSuperResolution/TSRCommon.ush
@@ -27,6 +27,8 @@
 
 #define CONFIG_ENABLE_STOCASTIC_QUANTIZATION 1
 
+#define CONFIG_SCENE_COLOR_OVERFLOW 1
+
 
 //------------------------------------------------------- CONSTANTS
 
@@ -202,6 +204,14 @@ static const uint kPlusIndexes3x3[5] = { 4, 1, 3, 5, 7 };
 #define taa_subpixel_payload uint // TODO: uint16_t
 #define taa_subpixel_payload2 uint2 // TODO: uint16_t
 
+// Largest encodable normal number in a half used on console.
+static const taa_half LargestNormalNumber = taa_half(65504.0);
+
+// Largest encodable scene color value to not overflaw in YCoCg with halfs on console
+static const taa_half3 LargestSceneColorRGB = (LargestNormalNumber * taa_half(0.25)).xxx;
+static const taa_half4 LargestSceneColorRGBA = (LargestNormalNumber * taa_half(0.25)).xxxx;
+static const taa_half3 LargestSceneColorYCoCg = LargestNormalNumber.xxx;
+
 
 //------------------------------------------------------- PARAMETERS
 
diff --git a/Engine/Shaders/Private/TemporalSuperResolution/TSRFilterFrequencies.usf b/Engine/Shaders/Private/TemporalSuperResolution/TSRFilterFrequencies.usf
--- a/Engine/Shaders/Private/TemporalSuperResolution/TSRFilterFrequencies.usf
+++ b/Engine/Shaders/Private/TemporalSuperResolution/TSRFilterFrequencies.usf
@@ -181,6 +181,10 @@ void MainCS(
 		
 		taa_half SampleInterferenceWeight = InterferenceWeightTexture[SampleInputPixelPos];
 		
+		#if CONFIG_SCENE_COLOR_OVERFLOW
+			InputSample = min(InputSample, LargestSceneColorRGB);
+		#endif
+
 		FLATTEN
 		if (!bEnableInterferenceHeuristic)
 		{
diff --git a/Engine/Shaders/Private/TemporalSuperResolution/TSRUpdateHistory.usf b/Engine/Shaders/Private/TemporalSuperResolution/TSRUpdateHistory.usf
--- a/Engine/Shaders/Private/TemporalSuperResolution/TSRUpdateHistory.usf
+++ b/Engine/Shaders/Private/TemporalSuperResolution/TSRUpdateHistory.usf
@@ -388,6 +388,10 @@ void MainCS(
 
 
 		FilteredInputColor = taa_half(0.0);
+		InputMinColor = dpv_interleave_mono_registers(+LargestSceneColorYCoCg);
+		InputMaxColor = dpv_interleave_mono_registers(-LargestSceneColorYCoCg);
+		TranslucencyMinColor = dpv_interleave_mono_registers(+LargestSceneColorRGBA);
+		TranslucencyMaxColor = dpv_interleave_mono_registers(-LargestSceneColorRGBA);
 
 		taa_half2 FilteredInputColorWeight = taa_half(0.0);
 		
@@ -415,6 +419,10 @@ void MainCS(
 				InputSceneTranslucencyTexture[dpv_lo(SampleInputPixelPos)],
 				InputSceneTranslucencyTexture[dpv_hi(SampleInputPixelPos)]);
 
+			#if CONFIG_SCENE_COLOR_OVERFLOW
+				InputColor = min(InputColor, dpv_interleave_mono_registers(LargestSceneColorRGB));
+			#endif
+
 			taa_half2x2 dPP = PixelOffset - dKO;
 			taa_half2 SampleSpatialWeight = ComputeSampleWeigth(FinalInputToHistoryFactor, dPP, /* MinimalContribution = */ float(0.005));
 
@@ -433,25 +441,26 @@ void MainCS(
 			if (SampleId == 0)
 			{
 				ClosestInputLuma4 = Luma4(InputColor);
-				InputMinColor = TransformColorForClampingBox(InputColor);
-				InputMaxColor = TransformColorForClampingBox(InputColor);
-
 				FilteredTranslucencyColor = TranslucencyColor;
-				TranslucencyMinColor = TranslucencyColor;
-				TranslucencyMaxColor = TranslucencyColor;
-			}
-			else
-			{
-				InputMinColor = min(InputMinColor, TransformColorForClampingBox(InputColor));
-				InputMaxColor = max(InputMaxColor, TransformColorForClampingBox(InputColor));
-				TranslucencyMinColor = min(TranslucencyMinColor, TranslucencyColor);
-				TranslucencyMaxColor = max(TranslucencyMaxColor, TranslucencyColor);
 			}
+
+			InputMinColor = min(InputMinColor, TransformColorForClampingBox(InputColor));
+			InputMaxColor = max(InputMaxColor, TransformColorForClampingBox(InputColor));
+			TranslucencyMinColor = min(TranslucencyMinColor, TranslucencyColor);
+			TranslucencyMaxColor = max(TranslucencyMaxColor, TranslucencyColor);
 		}
 		
 		FilteredInputColor = dpv_scale(FilteredInputColor, rcp(FilteredInputColorWeight));
 
 		InputPixelAlignement = ComputeSampleWeigth(InputToHistoryFactor, dKO, /* MinimalContribution = */ float(0.0));
+		
+		#if CONFIG_SCENE_COLOR_OVERFLOW
+		{
+			FilteredInputColor = min(FilteredInputColor, dpv_interleave_mono_registers(LargestSceneColorRGB));
+			ClosestInputLuma4 = min(ClosestInputLuma4, LargestSceneColorYCoCg.x);
+			FilteredTranslucencyColor = min(FilteredTranslucencyColor, dpv_interleave_mono_registers(LargestSceneColorRGBA));
+		}
+		#endif
 	}
 		
 	// Spills to LDS to make room in VGPR for history sampling.
@@ -1124,38 +1133,30 @@ void MainCS(
 		#endif
 			
 		LocalHistoryPixelPos = InvalidateOutputPixelPos(LocalHistoryPixelPos, HistoryInfo_ViewportMax);
-
-		// Output final history.
+		
+		// Stocastically round up or down using the hardware RWTexture2D truncation unit to take into precision
+		// loss due to pixel format encoding.
+		#if CONFIG_ENABLE_STOCASTIC_QUANTIZATION
 		{
-			// Stocastically round up or down using the hardware RWTexture2D truncation unit to take into precision
-			// loss due to pixel format encoding.
-			#if CONFIG_ENABLE_STOCASTIC_QUANTIZATION
-			{
-				uint2 Random = Rand3DPCG16(int3(dpv_lo(LocalHistoryPixelPos), View.StateFrameIndexMod8)).xy;
-				float E = Hammersley16(0, 1, Random).x;
+			uint2 Random = Rand3DPCG16(int3(dpv_lo(LocalHistoryPixelPos), View.StateFrameIndexMod8)).xy;
+			float E = Hammersley16(0, 1, Random).x;
 
-				FinalHighFrequencyColor = QuantizeForFloatRenderTarget(FinalHighFrequencyColor, E, HistoryQuantizationError);
-				FinalLowFrequencyColor = QuantizeForFloatRenderTarget(FinalLowFrequencyColor, E, HistoryQuantizationError);
-			}
-			#endif
+			FinalHighFrequencyColor = QuantizeForFloatRenderTarget(FinalHighFrequencyColor, E, HistoryQuantizationError);
+			FinalLowFrequencyColor = QuantizeForFloatRenderTarget(FinalLowFrequencyColor, E, HistoryQuantizationError);
+		}
+		#endif
 
-			#if 0 // Manual software mantissa truncation for debugging purposes.
-			{
-				const uint MantissaF32Bits = 23;
-				const uint MantissaF16Bits = 10;
-				const uint3 MantissaR11G11B10Bits = uint3(6, 6, 5);
+		FinalLowFrequencyColor = -dpv_min(-FinalLowFrequencyColor, taa_half(0.0).xxx);
+		FinalHighFrequencyColor = -dpv_min(-FinalHighFrequencyColor, taa_half(0.0).xxx);
+		FinalLowFrequencyColor = dpv_min(FinalLowFrequencyColor, taa_half(Max10BitsFloat).xxx);
+		FinalHighFrequencyColor = dpv_min(FinalHighFrequencyColor, taa_half(Max10BitsFloat).xxx);		
+		#if CONFIG_ACCUMULATE_TRANSLUCENCY_SEPARATLY
+			FinalTranslucencyColor = -dpv_min(-FinalTranslucencyColor, taa_half(0.0).xxxx);
+			FinalTranslucencyColor = dpv_min(FinalTranslucencyColor, taa_half(Max10BitsFloat).xxxx);
+		#endif
 
-				//FinalLowFrequencyColor.x = asfloat(asuint(FinalLowFrequencyColor.x) & ((~0) << (MantissaF32Bits - MantissaR11G11B10Bits.r)));
-				//FinalLowFrequencyColor.y = asfloat(asuint(FinalLowFrequencyColor.y) & ((~0) << (MantissaF32Bits - MantissaR11G11B10Bits.g)));
-				//FinalLowFrequencyColor.z = asfloat(asuint(FinalLowFrequencyColor.z) & ((~0) << (MantissaF32Bits - MantissaR11G11B10Bits.b)));
-			}
-			#endif
-			
-			FinalLowFrequencyColor = -dpv_min(-FinalLowFrequencyColor, taa_half(0.0).xxx);
-			FinalHighFrequencyColor = -dpv_min(-FinalHighFrequencyColor, taa_half(0.0).xxx);
-			FinalLowFrequencyColor = dpv_min(FinalLowFrequencyColor, taa_half(Max10BitsFloat).xxx);
-			FinalHighFrequencyColor = dpv_min(FinalHighFrequencyColor, taa_half(Max10BitsFloat).xxx);
-			
+		// Output final history lo pixel.
+		{
 			HistoryOutput_LowFrequency[dpv_lo(LocalHistoryPixelPos)]  = dpv_lo(FinalLowFrequencyColor);
 			HistoryOutput_HighFrequency[dpv_lo(LocalHistoryPixelPos)] = dpv_lo(FinalHighFrequencyColor);
 			HistoryOutput_Metadata[dpv_lo(LocalHistoryPixelPos)]      = taa_half2(dpv_lo(FinalLowFrequencyWeight), dpv_lo(FinalHistoryValidity));
@@ -1164,18 +1165,21 @@ void MainCS(
 				HistoryOutput_TranslucencyAlpha[dpv_lo(LocalHistoryPixelPos)] = dpv_lo(FinalTranslucencyColor).a;
 			#endif
 			HistoryOutput_SubpixelDetails[dpv_lo(LocalHistoryPixelPos)] = dpv_lo(FinalSubpixelDetails);
-
-			#if CONFIG_ENABLE_DUAL_PIXEL_VECTORIZATION
-				HistoryOutput_LowFrequency[dpv_hi(LocalHistoryPixelPos)]  = dpv_hi(FinalLowFrequencyColor);
-				HistoryOutput_HighFrequency[dpv_hi(LocalHistoryPixelPos)] = dpv_hi(FinalHighFrequencyColor);
-				HistoryOutput_Metadata[dpv_hi(LocalHistoryPixelPos)]      = taa_half2(dpv_hi(FinalLowFrequencyWeight), dpv_hi(FinalHistoryValidity));
-				#if CONFIG_ACCUMULATE_TRANSLUCENCY_SEPARATLY
-					HistoryOutput_Translucency[dpv_hi(LocalHistoryPixelPos)] = dpv_hi(FinalTranslucencyColor).rgb;
-					HistoryOutput_TranslucencyAlpha[dpv_hi(LocalHistoryPixelPos)] = dpv_hi(FinalTranslucencyColor).a;
-				#endif
-				HistoryOutput_SubpixelDetails[dpv_hi(LocalHistoryPixelPos)] = dpv_hi(FinalSubpixelDetails);
+		}
+		
+		// Output final history hi pixel.
+		#if CONFIG_ENABLE_DUAL_PIXEL_VECTORIZATION
+		{
+			HistoryOutput_LowFrequency[dpv_hi(LocalHistoryPixelPos)]  = dpv_hi(FinalLowFrequencyColor);
+			HistoryOutput_HighFrequency[dpv_hi(LocalHistoryPixelPos)] = dpv_hi(FinalHighFrequencyColor);
+			HistoryOutput_Metadata[dpv_hi(LocalHistoryPixelPos)]      = taa_half2(dpv_hi(FinalLowFrequencyWeight), dpv_hi(FinalHistoryValidity));
+			#if CONFIG_ACCUMULATE_TRANSLUCENCY_SEPARATLY
+				HistoryOutput_Translucency[dpv_hi(LocalHistoryPixelPos)] = dpv_hi(FinalTranslucencyColor).rgb;
+				HistoryOutput_TranslucencyAlpha[dpv_hi(LocalHistoryPixelPos)] = dpv_hi(FinalTranslucencyColor).a;
 			#endif
+			HistoryOutput_SubpixelDetails[dpv_hi(LocalHistoryPixelPos)] = dpv_hi(FinalSubpixelDetails);
 		}
+		#endif
 
 		// Output final scene color
 		{
