Fix the specular artifacts on Snapdragon 845.

#p4v-preflight-copy 17146274

diff --git a/Engine/Shaders/Private/MobileBasePassPixelShader.usf b/Engine/Shaders/Private/MobileBasePassPixelShader.usf
--- a/Engine/Shaders/Private/MobileBasePassPixelShader.usf
+++ b/Engine/Shaders/Private/MobileBasePassPixelShader.usf
@@ -429,11 +429,11 @@ inout FMobileLightAccumulator MobileLightAccumulator)
 
 	float3 ToLight = LightPositionAndInvRadius.xyz - MaterialParameters.AbsoluteWorldPosition;
 	float DistanceSqr = dot(ToLight, ToLight);
-	float3 L = ToLight * rsqrt(DistanceSqr);
-	half3 PointH = normalize(MaterialParameters.CameraVector + L);
+	half3 L = ToLight * rsqrt(DistanceSqr);
+	float3 PointH = normalize(float3(MaterialParameters.CameraVector + L));
 
 	half PointNoL = max(0, dot(MaterialParameters.WorldNormal, L));
-	half PointNoH = max(0, dot(MaterialParameters.WorldNormal, PointH));
+	float PointNoH = max(0, dot(MaterialParameters.WorldNormal, PointH));
 	half PointVoH = max(0, dot(MaterialParameters.CameraVector, PointH));
 
 #if SUPPORT_SPOTLIGHTS_SHADOW
@@ -811,7 +811,7 @@ void Main(
 
 	half3 L = MobileDirectionalLight.DirectionalLightDirectionAndShadowTransition.xyz;
 	half NoL = max(0, dot(MaterialParameters.WorldNormal, L));
-	half3 H = normalize(MaterialParameters.CameraVector + L);
+	float3 H = normalize(float3(MaterialParameters.CameraVector + L));
 	half NoH = max(0, dot(MaterialParameters.WorldNormal, H));
 	half VoH = max(0, dot(MaterialParameters.CameraVector, H));
 
