Mobile: pack LocalObjectBounds into per-instance data instead of lightmap data when project has static lighing disabled. Allows to auto-instance materails that use ObjectBounds and ObjectWorldPostionRadius

diff --git a/Engine/Shaders/Private/InstanceCulling/BuildInstanceDrawCommands.usf b/Engine/Shaders/Private/InstanceCulling/BuildInstanceDrawCommands.usf
--- a/Engine/Shaders/Private/InstanceCulling/BuildInstanceDrawCommands.usf
+++ b/Engine/Shaders/Private/InstanceCulling/BuildInstanceDrawCommands.usf
@@ -70,24 +70,48 @@ RWBuffer<uint> DrawIndirectArgsBufferOut;
 #define USES_PACKED_INSTANCE_DATA (FEATURE_LEVEL == FEATURE_LEVEL_ES3_1 && VF_SUPPORTS_PRIMITIVE_SCENE_DATA)
 
 #if USES_PACKED_INSTANCE_DATA
-
 float4 PackLightmapData(float4 LightmapUVScaleBias, float4 ShadowmapUVScaleBias)
 {
-	float4 Ret;
-	Ret.x = asfloat(PackUnorm2x16(LightmapUVScaleBias.xy));
-	Ret.y = asfloat(PackSnorm2x16(LightmapUVScaleBias.zw));
-	Ret.z = asfloat(PackUnorm2x16(ShadowmapUVScaleBias.xy));
-	Ret.w = asfloat(PackSnorm2x16(ShadowmapUVScaleBias.zw));
-	return Ret;
+	float4 Result;
+	Result.x = asfloat(PackUnorm2x16(LightmapUVScaleBias.xy));
+	Result.y = asfloat(PackSnorm2x16(LightmapUVScaleBias.zw));
+	Result.z = asfloat(PackUnorm2x16(ShadowmapUVScaleBias.xy));
+	Result.w = asfloat(PackSnorm2x16(ShadowmapUVScaleBias.zw));
+	return Result;
+}
+
+float2 PackLocalBoundsCenter(float3 Center)
+{
+	float2 Result = 0;
+	// uses 21 bits for each component, rounded to a 1 unit
+	const int3 Range = ((1u << 20u) - 1u);
+	const int3 Clamped = clamp(int3(round(Center)), -Range, Range);
+	const uint3 BiasClamped = uint3(Clamped + Range);
+	Result.x = asfloat((BiasClamped.x << 11u) | ((BiasClamped.z >> 0u) &  ((1u<<11u)-1u)));
+	Result.y = asfloat((BiasClamped.y << 11u) | ((BiasClamped.z >> 11u) & ((1u<<10u)-1u)));
+	return Result;
+}
+
+float2 PackLocalBoundsExtent(float3 Extent)
+{
+	float2 Result = 0;
+	// uses 21 bits for each component, rounded to a 1 unit
+	const uint3 Range = ((1u << 21u) - 1u);
+	const uint3 Clamped = min(uint3(round(Extent)), Range);
+	Result.x = asfloat((Clamped.x << 11u) | ((Clamped.z >> 0u) &  ((1u<<11u)-1u)));
+	Result.y = asfloat((Clamped.y << 11u) | ((Clamped.z >> 11u) & ((1u<<10u)-1u)));
+	return Result;
 }
 
 void WritePackedInstanceData(uint Offset, uint InstanceId, FInstanceSceneData InstanceData, uint ViewIdIndex)
 {
 	FPrimitiveSceneData PrimitiveData = GetPrimitiveData(InstanceData.PrimitiveId);
 	
+	float InstanceAuxData1 = 0;
+	float4 InstanceAuxData4 = 0;
+#if ALLOW_STATIC_LIGHTING 
 	float4 LightMapUVScaleBias = float4(1,1,0,0);
 	float4 ShadowMapUVScaleBias = float4(1,1,0,0);
-
 	// FIXME: need LODIndex to compute correct LightmapDataIndex
 	uint LightmapDataIndex = PrimitiveData.LightmapDataIndex; // + LODIndex;
 	if (LightmapDataIndex < GPUSceneNumLightmapDataItems)
@@ -105,8 +129,13 @@ void WritePackedInstanceData(uint Offset, uint InstanceId, FInstanceSceneData In
 		// avoid passing invalid lightmap data index into gfx shaders?
 		LightmapDataIndex = 0;
 	}
-
-	float4 PackedLightmapData = PackLightmapData(LightMapUVScaleBias, ShadowMapUVScaleBias);
+	InstanceAuxData4 = PackLightmapData(LightMapUVScaleBias, ShadowMapUVScaleBias);
+	InstanceAuxData1 = asfloat(LightmapDataIndex);
+#else
+	InstanceAuxData4.xy = PackLocalBoundsCenter(InstanceData.LocalBoundsCenter);
+	InstanceAuxData4.zw = PackLocalBoundsExtent(InstanceData.LocalBoundsExtent);
+	InstanceAuxData1 = 0; // TODO: pack something useful here 
+#endif // ALLOW_STATIC_LIGHTING
 
 	const uint PackedPrimitiveFlags = ((PrimitiveData.Flags & 0xffff) | (InstanceData.Flags << 16u));
 
@@ -118,14 +147,14 @@ void WritePackedInstanceData(uint Offset, uint InstanceId, FInstanceSceneData In
 	float4 InstanceTransform3 = InstanceData.LocalToWorld[2];
 	
 	InstanceTransform1.w = asfloat(PackedPrimitiveFlags); 
-	InstanceTransform2.w = asfloat(LightmapDataIndex); // TODO: pack something useful here 
+	InstanceTransform2.w = InstanceAuxData1; 
 	InstanceTransform3.w = InstanceData.PerInstanceRandom;
 		
 	InstanceIdsBufferOutMobile[Offset*5 + 0] = InstanceOriginAndId;
 	InstanceIdsBufferOutMobile[Offset*5 + 1] = InstanceTransform1;
 	InstanceIdsBufferOutMobile[Offset*5 + 2] = InstanceTransform2;
 	InstanceIdsBufferOutMobile[Offset*5 + 3] = InstanceTransform3;
-	InstanceIdsBufferOutMobile[Offset*5 + 4] = PackedLightmapData;
+	InstanceIdsBufferOutMobile[Offset*5 + 4] = InstanceAuxData4;
 }
 #endif
 
diff --git a/Engine/Shaders/Private/SceneData.ush b/Engine/Shaders/Private/SceneData.ush
--- a/Engine/Shaders/Private/SceneData.ush
+++ b/Engine/Shaders/Private/SceneData.ush
@@ -490,7 +490,34 @@ FSceneDataIntermediates GetSceneDataIntermediates(uint InstanceIdOffset, uint Dr
 	return Intermediates;
 }
 #elif (VF_USE_PRIMITIVE_SCENE_DATA == 2)
-FSceneDataIntermediates GetSceneDataIntermediates(uint DrawInstanceId, float4 InstanceOrigin, float4 InstanceTransform1, float4 InstanceTransform2, float4 InstanceTransform3, float4 PackedLightmapData)
+// Must match PackLocalBoundsCenter and PackLocalBoundsExtent
+float3 UnpackLocalBoundsCenter(float2 PackedCenter)
+{
+	float3 Result = 0;
+	// uses 21 bits for each component, rounded to a 1 unit
+	const uint SPLIT_MASK = (1u << 11u) - 1u;
+	const float CenterBias = (1u << 20u) - 1u;
+	uint p0 = asuint(PackedCenter.x);
+	uint p1 = asuint(PackedCenter.y);
+	Result.x = float(p0 >> 11u) - CenterBias;
+	Result.y = float(p1 >> 11u) - CenterBias;
+	Result.z = float((p0 & SPLIT_MASK) | ((p1 & SPLIT_MASK) << 11u)) - CenterBias;
+	return Result;
+}
+float3 UnpackLocalBoundsExtent(float2 PackedExtent)
+{
+	float3 Result = 0;
+	// uses 21 bits for each component, rounded to a 1 unit
+	const uint SPLIT_MASK = (1u << 11u) - 1u;
+	uint p0 = asuint(PackedExtent.x);
+	uint p1 = asuint(PackedExtent.y);
+	Result.x = float(p0 >> 11u);
+	Result.y = float(p1 >> 11u);
+	Result.z = float((p0 & SPLIT_MASK) | ((p1 & SPLIT_MASK) << 11u));
+	return Result;
+}
+
+FSceneDataIntermediates GetSceneDataIntermediates(uint DrawInstanceId, float4 InstanceOrigin, float4 InstanceTransform1, float4 InstanceTransform2, float4 InstanceTransform3, float4 InstanceAuxData)
 {
 	// TODO: pack important primitive and instance flags here
 	const uint PrimitiveFlags = asuint(InstanceTransform1.w);
@@ -511,12 +538,6 @@ FSceneDataIntermediates GetSceneDataIntermediates(uint DrawInstanceId, float4 In
 	InstanceData.PerInstanceRandom = InstanceTransform3.w;
 	#endif
 	
-	InstanceData.LightMapAndShadowMapUVBias = float4(
-		UnpackSnorm2x16(asuint(PackedLightmapData.y)), 
-		UnpackSnorm2x16(asuint(PackedLightmapData.w)));
-	
-	const uint LightmapDataIndex = asuint(InstanceTransform2.w);
-	
 	float3 Scale2;
 	Scale2.x = length2(InstanceData.LocalToWorld[0].xyz);
 	Scale2.y = length2(InstanceData.LocalToWorld[1].xyz);
@@ -552,8 +573,21 @@ FSceneDataIntermediates GetSceneDataIntermediates(uint DrawInstanceId, float4 In
 	Intermediates.Primitive.InvNonUniformScale		= Intermediates.InstanceData.InvNonUniformScale;
 	Intermediates.Primitive.WorldToLocal			= Intermediates.InstanceData.WorldToLocal;
 	Intermediates.Primitive.NonUniformScale			= Intermediates.InstanceData.NonUniformScale;
-	Intermediates.Primitive.LightmapDataIndex		= LightmapDataIndex;
 
+#if ALLOW_STATIC_LIGHTING 
+	Intermediates.Primitive.LightmapDataIndex = asuint(InstanceTransform2.w);
+	InstanceData.LightMapAndShadowMapUVBias = float4(
+		UnpackSnorm2x16(asuint(InstanceAuxData.y)), 
+		UnpackSnorm2x16(asuint(InstanceAuxData.w)));
+#else
+	InstanceData.LocalBoundsCenter = UnpackLocalBoundsCenter(InstanceAuxData.xy);
+	InstanceData.LocalBoundsExtent = UnpackLocalBoundsExtent(InstanceAuxData.zw);
+	Intermediates.Primitive.LocalObjectBoundsMin = InstanceData.LocalBoundsCenter - InstanceData.LocalBoundsExtent; 
+	Intermediates.Primitive.LocalObjectBoundsMax = InstanceData.LocalBoundsCenter + InstanceData.LocalBoundsExtent; 
+	Intermediates.Primitive.ObjectWorldPositionAndRadius.xyz = mul(float4(InstanceData.LocalBoundsCenter.xyz, 0.0f), InstanceData.LocalToWorld).xyz;
+	Intermediates.Primitive.ObjectWorldPositionAndRadius.w = length(InstanceData.LocalBoundsExtent * InstanceData.NonUniformScale.xyz);
+#endif
+	
 	return Intermediates; 
 }
 #else
diff --git a/Engine/Shaders/Private/VertexFactoryCommon.ush b/Engine/Shaders/Private/VertexFactoryCommon.ush
--- a/Engine/Shaders/Private/VertexFactoryCommon.ush
+++ b/Engine/Shaders/Private/VertexFactoryCommon.ush
@@ -28,7 +28,7 @@
 		float4 InstanceTransform1 : ATTRIBUTE9; \
 		float4 InstanceTransform2 : ATTRIBUTE10; \
 		float4 InstanceTransform3 : ATTRIBUTE11; \
-		float4 PackedLightmapData : ATTRIBUTE12; \
+		float4 InstanceAuxData : ATTRIBUTE12; \
 		uint DrawInstanceId : SV_InstanceID;
 
 	#define VF_GPUSCENE_GET_INTERMEDIATES(VFInput) \
@@ -38,13 +38,13 @@
 			VFInput.InstanceTransform1, \
 			VFInput.InstanceTransform2, \
 			VFInput.InstanceTransform3, \
-			VFInput.PackedLightmapData)
+			VFInput.InstanceAuxData)
 
 	#define VF_GPUSCENE_GET_LIGHTMAP_UV_SCALE_BIAS(VFInput, LightmapDataIndex) \
-		float4(UnpackUnorm2x16(asuint(VFInput.PackedLightmapData.x)), UnpackSnorm2x16(asuint(VFInput.PackedLightmapData.y)))
+		float4(UnpackUnorm2x16(asuint(VFInput.InstanceAuxData.x)), UnpackSnorm2x16(asuint(VFInput.InstanceAuxData.y)))
 	
 	#define VF_GPUSCENE_GET_SHADOWMAP_UV_SCALE_BIAS(VFInput, LightmapDataIndex) \
-		float4(UnpackUnorm2x16(asuint(VFInput.PackedLightmapData.z)), UnpackSnorm2x16(asuint(VFInput.PackedLightmapData.w)))
+		float4(UnpackUnorm2x16(asuint(VFInput.InstanceAuxData.z)), UnpackSnorm2x16(asuint(VFInput.InstanceAuxData.w)))
 
 #else
 	#define VF_GPUSCENE_DECLARE_INPUT_BLOCK(Attr0) 
diff --git a/Engine/Source/Runtime/Renderer/Private/GPUScene.cpp b/Engine/Source/Runtime/Renderer/Private/GPUScene.cpp
--- a/Engine/Source/Runtime/Renderer/Private/GPUScene.cpp
+++ b/Engine/Source/Runtime/Renderer/Private/GPUScene.cpp
@@ -1724,7 +1724,7 @@ void FGPUSceneCompactInstanceData::Init(const FGPUScenePrimitiveCollector* Primi
 	InstanceTransform2		= LocalToWorld.GetScaledAxis(EAxis::Y);
 	InstanceTransform3		= LocalToWorld.GetScaledAxis(EAxis::Z);
 	InstanceOriginAndId.W	= *(float*)&DynamicPrimitiveId;
-	PackedLightmapData		= FVector4(0); // FIXME
+	InstanceAuxData			= FVector4(0);
 }
 
 void FGPUSceneCompactInstanceData::Init(const FScene* Scene, int32 PrimitiveId)
@@ -1739,5 +1739,5 @@ void FGPUSceneCompactInstanceData::Init(const FScene* Scene, int32 PrimitiveId)
 	InstanceTransform2		= LocalToWorld.GetScaledAxis(EAxis::Y);
 	InstanceTransform3		= LocalToWorld.GetScaledAxis(EAxis::Z);
 	InstanceOriginAndId.W	= *(float*)&PrimitiveId;
-	PackedLightmapData		= FVector4(0); // FIXME
+	InstanceAuxData			= FVector4(0);
 }
\ No newline at end of file
diff --git a/Engine/Source/Runtime/Renderer/Private/GPUScene.h b/Engine/Source/Runtime/Renderer/Private/GPUScene.h
--- a/Engine/Source/Runtime/Renderer/Private/GPUScene.h
+++ b/Engine/Source/Runtime/Renderer/Private/GPUScene.h
@@ -391,7 +391,7 @@ struct FGPUSceneCompactInstanceData
 	FVector4 InstanceTransform1;
 	FVector4 InstanceTransform2;
 	FVector4 InstanceTransform3;
-	FVector4 PackedLightmapData;
+	FVector4 InstanceAuxData;
 
 	void Init(const FGPUScenePrimitiveCollector* PrimitiveCollector, int32 PrimitiveId);
 	void Init(const FScene* Scene, int32 PrimitiveId);
