Fix hair sometimes not showing up in Vulkan (Indirect Draw parameters were being filled like Indirect Dispatch).

diff --git a/Engine/Shaders/Private/HairStrands/HairStrandsVisibilityTile.usf b/Engine/Shaders/Private/HairStrands/HairStrandsVisibilityTile.usf
--- a/Engine/Shaders/Private/HairStrands/HairStrandsVisibilityTile.usf
+++ b/Engine/Shaders/Private/HairStrands/HairStrandsVisibilityTile.usf
@@ -104,9 +104,10 @@ void MainCS(uint2 DispatchThreadId : SV_DispatchThreadID)
 		const uint TileCount = TileCountBuffer[0];
 
 		// Indirect draw
-		TileIndirectDrawBuffer[0] = bRectPrimitive > 0 ? 4 : 6;
-		TileIndirectDrawBuffer[1] = TileCount;
-		TileIndirectDrawBuffer[2] = 1;
+		TileIndirectDrawBuffer[0] = bRectPrimitive > 0 ? 4 : 6;  // VertexCountPerInstance
+		TileIndirectDrawBuffer[1] = TileCount;  // InstanceCount
+		TileIndirectDrawBuffer[2] = 0;  // StartVertexLocation
+		TileIndirectDrawBuffer[3] = 0;  // StartInstanceLocation
 
 		// Indirect dispatch
 		TileIndirectDispatchBuffer[0] = TileCount >= TileCountXY.x ? TileCountXY.x : TileCount;
@@ -122,9 +123,10 @@ void MainCS(uint2 DispatchThreadId : SV_DispatchThreadID)
 
 		const uint TileClearCount = TileCountBuffer[1];
 
-		TileClearIndirectDrawBuffer[0] = bRectPrimitive > 0 ? 4 : 6;
-		TileClearIndirectDrawBuffer[1] = TileClearCount;
-		TileClearIndirectDrawBuffer[2] = 1;
+		TileClearIndirectDrawBuffer[0] = bRectPrimitive > 0 ? 4 : 6;  // VertexCountPerInstance
+		TileClearIndirectDrawBuffer[1] = TileClearCount;  // InstanceCount
+		TileClearIndirectDrawBuffer[2] = 0;  // StartVertexLocation
+		TileClearIndirectDrawBuffer[3] = 0;  // StartInstanceLocation
 	}
 }
 #endif // SHADER_TILE_COPY_ARGS
