Fixed FX collision normals on PC.

diff --git a/Engine/Shaders/Private/RayTracing/VFXTraceRay.ush b/Engine/Shaders/Private/RayTracing/VFXTraceRay.ush
--- a/Engine/Shaders/Private/RayTracing/VFXTraceRay.ush
+++ b/Engine/Shaders/Private/RayTracing/VFXTraceRay.ush
@@ -20,7 +20,7 @@
 
 float3 GetWorldNormalCommon(FVFXTracePayload Payload)
 {
-	FTriangleBaseAttributes Tri = LoadTriangleBaseAttributes(Payload.PrimitiveIndex);
+	FTriangleBaseAttributes Tri = LoadTriangleBaseAttributes(PrimitiveIndex());
 
 	float3 LocalPosition = (1.0f - Payload.Barycentrics.x - Payload.Barycentrics.y) * Tri.LocalPositions[0]
 		+ Payload.Barycentrics.x * Tri.LocalPositions[1]
