Subsurface scattering performance stability enhancement:
1. use control variates to deal with dynamic lighting change.
2. use bilinear filtered lighting as the control variable.

Performance boost up to 2.4x (3.22->1.34) during dynamic change for large SSS sudden appearance.

diff --git a/Engine/Shaders/Private/PostProcessSubsurface.usf b/Engine/Shaders/Private/PostProcessSubsurface.usf
--- a/Engine/Shaders/Private/PostProcessSubsurface.usf
+++ b/Engine/Shaders/Private/PostProcessSubsurface.usf
@@ -29,19 +29,24 @@ float4 SubsurfaceParams;
 Texture2D SubsurfaceInput0_Texture;
 Texture2D SubsurfaceInput1_Texture;
 Texture2D SubsurfaceInput2_Texture;
+Texture2D SubsurfaceInput3_Texture;
 
 SamplerState SubsurfaceSampler0;
 SamplerState SubsurfaceSampler1;
 SamplerState SubsurfaceSampler2;
+SamplerState SubsurfaceSampler3;
 
 #if SUPPORTS_INDEPENDENT_SAMPLERS
 	#define SharedSubsurfaceSampler0 SubsurfaceSampler0
 	#define SharedSubsurfaceSampler1 SubsurfaceSampler0
 	#define SharedSubsurfaceSampler2 SubsurfaceSampler0
+	#define SharedSubsurfaceSampler3 SubsurfaceSampler0
+
 #else
 	#define SharedSubsurfaceSampler0 SubsurfaceSampler0
 	#define SharedSubsurfaceSampler1 SubsurfaceSampler1
 	#define SharedSubsurfaceSampler2 SubsurfaceSampler2
+	#define SharedSubsurfaceSampler3 SubsurfaceSampler3
 #endif
 
 SCREEN_PASS_TEXTURE_VIEWPORT(SubsurfaceInput0)
diff --git a/Engine/Shaders/Private/SubsurfaceBurleyNormalized.ush b/Engine/Shaders/Private/SubsurfaceBurleyNormalized.ush
--- a/Engine/Shaders/Private/SubsurfaceBurleyNormalized.ush
+++ b/Engine/Shaders/Private/SubsurfaceBurleyNormalized.ush
@@ -446,12 +446,13 @@ int GetNumOfSamplesBasedOnQuality(float2 UV)
 	//1. Estimate the number of samples required to reach the target variance level
 	NumOfSamples = clamp((VariableVar / VARIANCE_LEVEL), BETA_LIMIT, BURLEY_NUM_SAMPLES);//View.GeneralPurposeTweak
 #endif
+	// The following two ad hoc design is not viable when we monitor the variance of control variates residual.
 	//1.1 If the luminance is too low, we should increase several number of samples to oberse the world for high lighting condition
-	NumOfSamples = lerp(NumOfSamples, max(LOW_LUMA_SAMPLE_COUNT, NumOfSamples), step(0.5, QualityMatrix.a));
+	// NumOfSamples = lerp(NumOfSamples, max(LOW_LUMA_SAMPLE_COUNT, NumOfSamples), step(0.5, QualityMatrix.a));
 
 	//2. Clamp down the number of samples when the gamma corrected luminance is too large
 	//   because it does not help to improve the quality, but degrade the performance.
-	NumOfSamples = lerp(NumOfSamples, HIGH_LUMA_SAMPLE_COUNT, step(3.2, QualityMatrix.a)); // 3.2 is an experimental value to have best quality.
+	// NumOfSamples = lerp(NumOfSamples, HIGH_LUMA_SAMPLE_COUNT, step(3.2, QualityMatrix.a)); // 3.2 is an experimental value to have best quality.
 
 	//3. Reduce the number of samples at the edge due to TAA (Flickering)
 #if 0
@@ -673,10 +674,14 @@ BurleySampleDiffuseNormal SampleSharedDiffuseNormal(float2 SampleGroupThreadId)
 
 // If we support independent samplers, use a point sampler as described below. But if they are not supported,
 // then default to the regular sampler which will cause artifacts, but is better than not compiling.
+// The bilinear sampler is required to reduce the variance overestimation with control variates. The point 
+// sampler works but with a little worse performance.
 #if SUPPORTS_INDEPENDENT_SAMPLERS
 	#define SharedBurleyPointSampler SubsurfaceSampler1
+	#define SharedBurleyBilinearSampler SubsurfaceSampler3
 #else
 	#define SharedBurleyPointSampler SubsurfaceSampler0
+	#define SharedBurleyBilinearSampler SharedBurleyPointSampler
 #endif
 
 float3 GetBilateralNormal(float2 ClampedUV, float2 Extent)
@@ -1155,16 +1160,21 @@ float4 UpdateQualityVariance(float4 SubsurfaceColor, float2 BufferUV)
 	float4 History = Texture2DSample(SubsurfaceInput1_Texture, SubsurfaceSampler1, BufferUV);
 #endif
 
+	// Use the bilinear jitter-corrected lighting as the control variable (Constant CV coefficient = 1.0) to remove the variance over-estimation
+	// of spatial features, like peach fuzz, inside subsurface scattering region.
+	float2 BufferUVJitter = (ScreenPosToViewportUV(View.TemporalAAJitter.xy)-float2(0.5, 0.5))* SubsurfaceInput0_UVViewportSize.xy;
+	float4 ControlVariable = Texture2DSampleLevel(SubsurfaceInput3_Texture, SharedBurleyBilinearSampler, BufferUV - BufferUVJitter, 0);//
+
 	// Calculate the illuminance
-	float NewSurfaceLuminance = Luminance(pow(SubsurfaceColor.rgb, 1 / 2.2));
+	float NewSurfaceLuminanceResidual = Luminance(pow(SubsurfaceColor.rgb, 1 / 2.2))- Luminance(pow(ControlVariable.rgb, 1 / 2.2));
 
-	NewHistory.a = (1 - WeightFinal) * History.a + WeightFinal * NewSurfaceLuminance;
+	NewHistory.a = (1 - WeightFinal) * History.a + WeightFinal * NewSurfaceLuminanceResidual;
 
-	float LumDelta = NewSurfaceLuminance - History.a;
+	float LumDelta = NewSurfaceLuminanceResidual - History.a;
 	NewHistory.b = (1 - WeightFinal)*History.b + WeightFinal * (1 - WeightFinal)*LumDelta*LumDelta;
 
 #if USE_TRUE_DISTRIBUTION_VAR
-	NewHistory.g = (1- WeightFinal)*NewHistory.g+WeightFinal * SubsurfaceColor.a;
+	NewHistory.g = (1- WeightFinal)*History.g+WeightFinal * SubsurfaceColor.a;
 #else
 	// Update the random variable variance
 	NewHistory.g = History.b*SubsurfaceColor.a;
@@ -1274,7 +1284,7 @@ void BurleyComputeMain(uint2 DT_ID, uint2 G_ID, uint GI)
 			uint2 TargetGridPos = Pos + uint2(i, j);
 			float2 LBufferUV = ConvertGridPos2UV(TargetGridPos);
 			// We perform quality variance thread-wise in the second pass
-			float4 SurfaceColor = Texture2DSample(SubsurfaceInput0_Texture, SubsurfaceSampler0, LBufferUV);
+			float4 SurfaceColor = SurfaceColor = Texture2DSample(SubsurfaceInput0_Texture, SharedBurleyBilinearSampler, LBufferUV);
 
 			// We have to check per pixel if this is a burley so that we don't overwrite separable pixels
 			uint SelectedProfile = GetSubsurfaceProfileId(LBufferUV, SubsurfaceInput0_Extent);
diff --git a/Engine/Source/Runtime/Renderer/Private/PostProcess/PostProcessSubsurface.cpp b/Engine/Source/Runtime/Renderer/Private/PostProcess/PostProcessSubsurface.cpp
--- a/Engine/Source/Runtime/Renderer/Private/PostProcess/PostProcessSubsurface.cpp
+++ b/Engine/Source/Runtime/Renderer/Private/PostProcess/PostProcessSubsurface.cpp
@@ -121,6 +121,8 @@ namespace
 		TEXT("0: Depth Only. It is more performant (x2 faster for close view).")
 		TEXT("1: Depth and normal. It leads to better quality in regions like eyelids. (default)"),
 		ECVF_RenderThreadSafe);
+	
+	DECLARE_GPU_STAT(SubsurfaceScattering)
 }
 
 // Define to use a custom ps to clear UAV.
@@ -527,6 +529,8 @@ public:
 		SHADER_PARAMETER_SAMPLER(SamplerState, SubsurfaceSampler1)
 		SHADER_PARAMETER_STRUCT(FSubsurfaceInput, SubsurfaceInput2)   // Profile mask | Velocity
 		SHADER_PARAMETER_SAMPLER(SamplerState, SubsurfaceSampler2)
+		SHADER_PARAMETER_STRUCT(FSubsurfaceInput, SubsurfaceInput3)   // Control Variates
+		SHADER_PARAMETER_SAMPLER(SamplerState, SubsurfaceSampler3)
 		SHADER_PARAMETER_RDG_TEXTURE_UAV(RWTexture2D, ProfileIdTexture)
 	END_SHADER_PARAMETER_STRUCT()
 
@@ -907,6 +911,7 @@ void AddSubsurfaceViewPass(
 	FRDGTextureRef QualityHistoryTexture = RegisterExternalRenderTarget(GraphBuilder, QualityHistoryState, SceneColorTextureDescriptor.Extent, TEXT("QualityHistoryTexture"));
 	FRDGTextureRef NewQualityHistoryTexture = nullptr;
 
+	RDG_GPU_STAT_SCOPE(GraphBuilder, SubsurfaceScattering);
 	/**
 	 * When in bypass mode, the setup and convolution passes are skipped, but lighting
 	 * reconstruction is still performed in the recombine pass.
@@ -1097,6 +1102,8 @@ void AddSubsurfaceViewPass(
 					if (!bForceRunningInSeparable)
 					{
 						PassParameters->HistoryUAV = GraphBuilder.CreateUAV(FRDGTextureUAVDesc(NewQualityHistoryTexture, 0));
+						PassParameters->SubsurfaceInput3 = GetSubsurfaceInput(SetupTexture, SubsurfaceViewportParameters);
+						PassParameters->SubsurfaceSampler3 = TStaticSamplerState<SF_Bilinear, AM_Border, AM_Border, AM_Border>::GetRHI();
 					}
 
 					PassParameters->SubsurfaceInput1 = GetSubsurfaceInput(QualityHistoryTexture, SubsurfaceViewportParameters);
