1. Revert changes to the surface color copy from bilinear to point. The CV pass should not modify the output color.
2. Move surface color query for CV into updatevariance.
3. Add UV clamping to avoid introduced variance when projected to boarders during motion.

diff --git a/Engine/Shaders/Private/SubsurfaceBurleyNormalized.ush b/Engine/Shaders/Private/SubsurfaceBurleyNormalized.ush
--- a/Engine/Shaders/Private/SubsurfaceBurleyNormalized.ush
+++ b/Engine/Shaders/Private/SubsurfaceBurleyNormalized.ush
@@ -1161,17 +1161,19 @@ float4 UpdateQualityVariance(float4 SubsurfaceColor, float2 BufferUV)
 #endif
 
 	// Use the bilinear jitter-corrected lighting as the control variable (Constant CV coefficient = 1.0) to remove the variance over-estimation
-	// of spatial features, like peach fuzz, inside subsurface scattering region.
+	// of spatial features, like peach fuzz, inside subsurface scattering region. It will not affect the output surface color.
+	float4 BilinearSurfaceColor = Texture2DSample(SubsurfaceInput0_Texture, SharedBurleyBilinearSampler, BufferUV);
 	float2 BufferUVJitter = (ScreenPosToViewportUV(View.TemporalAAJitter.xy)-float2(0.5, 0.5))* SubsurfaceInput0_UVViewportSize.xy;
-	float4 ControlVariable = Texture2DSampleLevel(SubsurfaceInput3_Texture, SharedBurleyBilinearSampler, BufferUV - BufferUVJitter, 0);//
+	float2 BufferUVJitterRemoved = clamp(BufferUV - BufferUVJitter, SubsurfaceInput0_UVViewportBilinearMin, SubsurfaceInput0_UVViewportBilinearMax);
+	float4 ControlVariable = Texture2DSampleLevel(SubsurfaceInput3_Texture, SharedBurleyBilinearSampler, BufferUVJitterRemoved, 0);
 
 	// Calculate the illuminance
-	float NewSurfaceLuminanceResidual = Luminance(pow(SubsurfaceColor.rgb, 1 / 2.2))- Luminance(pow(ControlVariable.rgb, 1 / 2.2));
+	float NewSurfaceLuminanceResidual = Luminance(pow(BilinearSurfaceColor.rgb, 1 / 2.2))- Luminance(pow(ControlVariable.rgb, 1 / 2.2));
 
 	NewHistory.a = (1 - WeightFinal) * History.a + WeightFinal * NewSurfaceLuminanceResidual;
 
-	float LumDelta = NewSurfaceLuminanceResidual - History.a;
-	NewHistory.b = (1 - WeightFinal)*History.b + WeightFinal * (1 - WeightFinal)*LumDelta*LumDelta;
+	float Delta = NewSurfaceLuminanceResidual - History.a;
+	NewHistory.b = (1 - WeightFinal)*History.b + WeightFinal * (1 - WeightFinal)*Delta*Delta;
 
 #if USE_TRUE_DISTRIBUTION_VAR
 	NewHistory.g = (1- WeightFinal)*History.g+WeightFinal * SubsurfaceColor.a;
@@ -1284,7 +1286,7 @@ void BurleyComputeMain(uint2 DT_ID, uint2 G_ID, uint GI)
 			uint2 TargetGridPos = Pos + uint2(i, j);
 			float2 LBufferUV = ConvertGridPos2UV(TargetGridPos);
 			// We perform quality variance thread-wise in the second pass
-			float4 SurfaceColor = SurfaceColor = Texture2DSample(SubsurfaceInput0_Texture, SharedBurleyBilinearSampler, LBufferUV);
+			float4 SurfaceColor = Texture2DSample(SubsurfaceInput0_Texture, SubsurfaceSampler0, LBufferUV);
 
 			// We have to check per pixel if this is a burley so that we don't overwrite separable pixels
 			uint SelectedProfile = GetSubsurfaceProfileId(LBufferUV, SubsurfaceInput0_Extent);
