Reduce code duplication for instance data loading between mobile and desktop specific implementations

diff --git a/Engine/Shaders/Private/SceneData.ush b/Engine/Shaders/Private/SceneData.ush
--- a/Engine/Shaders/Private/SceneData.ush
+++ b/Engine/Shaders/Private/SceneData.ush
@@ -362,6 +362,41 @@ float4 LoadInstancePayloadDataElement(uint Index)
 #endif
 }
 
+void ComputeIstanceDerivedData(inout FInstanceSceneData InstanceData)
+{
+	//
+	// Do not put any load operations here!
+	//
+	float3 Scale2;
+	Scale2.x = length2(InstanceData.LocalToWorld[0].xyz);
+	Scale2.y = length2(InstanceData.LocalToWorld[1].xyz);
+	Scale2.z = length2(InstanceData.LocalToWorld[2].xyz);
+
+	InstanceData.InvNonUniformScale							= rsqrt(Scale2);
+	InstanceData.DeterminantSign							= CondMask(InstanceData.Flags & INSTANCE_SCENE_DATA_FLAG_DETERMINANT_SIGN, -1.0f, 1.0f);
+
+	InstanceData.NonUniformScale.xyz						= Scale2 * InstanceData.InvNonUniformScale;
+
+	float3 AbsNonUniformScale								= abs(InstanceData.NonUniformScale.xyz);
+	InstanceData.NonUniformScale.w							= max3(AbsNonUniformScale.x, AbsNonUniformScale.y, AbsNonUniformScale.z);
+	
+	InstanceData.WorldToLocal								= InstanceData.LocalToWorld;
+	InstanceData.WorldToLocal[0].xyz					   *= Pow2(InstanceData.InvNonUniformScale.x);
+	InstanceData.WorldToLocal[1].xyz					   *= Pow2(InstanceData.InvNonUniformScale.y);
+	InstanceData.WorldToLocal[2].xyz					   *= Pow2(InstanceData.InvNonUniformScale.z);
+	InstanceData.WorldToLocal[3].xyz						= 0.0f;
+	InstanceData.WorldToLocal								= transpose(InstanceData.WorldToLocal);
+	InstanceData.WorldToLocal[3].xyz						= mul(float4(-InstanceData.LocalToWorld[3].xyz, 0.0f), InstanceData.WorldToLocal).xyz;
+
+	// TODO: Uploading not fully functional
+	// InstanceData.CastShadows								= (InstanceData.Flags & INSTANCE_SCENE_DATA_FLAG_CAST_SHADOWS);
+	InstanceData.NaniteHasImposter							= (InstanceData.Flags & INSTANCE_SCENE_DATA_FLAG_HAS_IMPOSTER);
+	InstanceData.HasRandomID								= (InstanceData.Flags & INSTANCE_SCENE_DATA_FLAG_HAS_RANDOM);
+	InstanceData.HasCustomData								= (InstanceData.Flags & INSTANCE_SCENE_DATA_FLAG_HAS_CUSTOM_DATA);
+	InstanceData.HasDynamicData								= (InstanceData.Flags & INSTANCE_SCENE_DATA_FLAG_HAS_DYNAMIC_DATA);
+	InstanceData.HasLightShadowUVBias						= (InstanceData.Flags & INSTANCE_SCENE_DATA_FLAG_HAS_LIGHTSHADOW_UV_BIAS);
+}
+
 // Fetch from scene primitive buffer
 FInstanceSceneData GetInstanceSceneData(uint InstanceId, uint SOAStride)
 {
@@ -405,34 +440,7 @@ FInstanceSceneData GetInstanceSceneData(uint InstanceId, uint SOAStride)
 
 		InstanceData.LightMapAndShadowMapUVBias					= LoadInstanceSceneDataElement(9 * SOAStride + InstanceId);
 
-		float3 Scale2;
-		Scale2.x = length2(InstanceData.LocalToWorld[0].xyz);
-		Scale2.y = length2(InstanceData.LocalToWorld[1].xyz);
-		Scale2.z = length2(InstanceData.LocalToWorld[2].xyz);
-
-		InstanceData.InvNonUniformScale							= rsqrt(Scale2);
-		InstanceData.DeterminantSign							= CondMask(InstanceData.Flags & INSTANCE_SCENE_DATA_FLAG_DETERMINANT_SIGN, -1.0f, 1.0f);
-
-		InstanceData.NonUniformScale.xyz						= Scale2 * InstanceData.InvNonUniformScale;
-
-		float3 AbsNonUniformScale								= abs(InstanceData.NonUniformScale.xyz);
-		InstanceData.NonUniformScale.w							= max3(AbsNonUniformScale.x, AbsNonUniformScale.y, AbsNonUniformScale.z);
-	
-		InstanceData.WorldToLocal								= InstanceData.LocalToWorld;
-		InstanceData.WorldToLocal[0].xyz					   *= Pow2(InstanceData.InvNonUniformScale.x);
-		InstanceData.WorldToLocal[1].xyz					   *= Pow2(InstanceData.InvNonUniformScale.y);
-		InstanceData.WorldToLocal[2].xyz					   *= Pow2(InstanceData.InvNonUniformScale.z);
-		InstanceData.WorldToLocal[3].xyz						= 0.0f;
-		InstanceData.WorldToLocal								= transpose(InstanceData.WorldToLocal);
-		InstanceData.WorldToLocal[3].xyz						= mul(float4(-InstanceData.LocalToWorld[3].xyz, 0.0f), InstanceData.WorldToLocal).xyz;
-
-		// TODO: Uploading not fully functional
-		// InstanceData.CastShadows								= (InstanceData.Flags & INSTANCE_SCENE_DATA_FLAG_CAST_SHADOWS);
-		InstanceData.NaniteHasImposter							= (InstanceData.Flags & INSTANCE_SCENE_DATA_FLAG_HAS_IMPOSTER);
-		InstanceData.HasRandomID								= (InstanceData.Flags & INSTANCE_SCENE_DATA_FLAG_HAS_RANDOM);
-		InstanceData.HasCustomData								= (InstanceData.Flags & INSTANCE_SCENE_DATA_FLAG_HAS_CUSTOM_DATA);
-		InstanceData.HasDynamicData								= (InstanceData.Flags & INSTANCE_SCENE_DATA_FLAG_HAS_DYNAMIC_DATA);
-		InstanceData.HasLightShadowUVBias						= (InstanceData.Flags & INSTANCE_SCENE_DATA_FLAG_HAS_LIGHTSHADOW_UV_BIAS);
+		ComputeIstanceDerivedData(InstanceData);
 
 		// TODO: Temporary hack!
 	#if USES_PER_INSTANCE_CUSTOM_DATA
@@ -586,27 +594,13 @@ FSceneDataIntermediates GetSceneDataIntermediates(uint DrawInstanceId, float4 In
 
 	InstanceData.Flags = InstanceFlags;
 	InstanceData.PrimitiveId = asuint(InstanceOrigin.w);
+	
 	#if USES_PER_INSTANCE_RANDOM
 	InstanceData.PerInstanceRandom = InstanceTransform3.w;
 	#endif
-	
-	float3 Scale2;
-	Scale2.x = length2(InstanceData.LocalToWorld[0].xyz);
-	Scale2.y = length2(InstanceData.LocalToWorld[1].xyz);
-	Scale2.z = length2(InstanceData.LocalToWorld[2].xyz);
 
-	InstanceData.InvNonUniformScale					= rsqrt(Scale2);
-	InstanceData.DeterminantSign					= CondMask(InstanceData.Flags & INSTANCE_SCENE_DATA_FLAG_DETERMINANT_SIGN, -1.0f, 1.0f);
-	InstanceData.NonUniformScale.xyz				= Scale2 * InstanceData.InvNonUniformScale;
-	float3 AbsNonUniformScale						= abs(InstanceData.NonUniformScale.xyz);
-	InstanceData.NonUniformScale.w					= max3(AbsNonUniformScale.x, AbsNonUniformScale.y, AbsNonUniformScale.z);
-	InstanceData.WorldToLocal						= InstanceData.LocalToWorld;
-	InstanceData.WorldToLocal[0].xyz			   *= Pow2(InstanceData.InvNonUniformScale.x);
-	InstanceData.WorldToLocal[1].xyz			   *= Pow2(InstanceData.InvNonUniformScale.y);
-	InstanceData.WorldToLocal[2].xyz			   *= Pow2(InstanceData.InvNonUniformScale.z);
-	InstanceData.WorldToLocal[3].xyz				= 0.0f;
-	InstanceData.WorldToLocal						= transpose(InstanceData.WorldToLocal);
-	InstanceData.WorldToLocal[3].xyz				= mul(float4(-InstanceData.LocalToWorld[3].xyz, 0.0f), InstanceData.WorldToLocal).xyz;
+	ComputeIstanceDerivedData(InstanceData);
+
 	//
 	FSceneDataIntermediates Intermediates = (FSceneDataIntermediates)0;
 	Intermediates.InstanceData = InstanceData;
