Use triangle strip instead of triangle list for HW hair strands path.

diff --git a/Engine/Plugins/Runtime/HairStrands/Source/HairStrandsCore/Private/GroomComponent.cpp b/Engine/Plugins/Runtime/HairStrands/Source/HairStrandsCore/Private/GroomComponent.cpp
--- a/Engine/Plugins/Runtime/HairStrands/Source/HairStrandsCore/Private/GroomComponent.cpp
+++ b/Engine/Plugins/Runtime/HairStrands/Source/HairStrandsCore/Private/GroomComponent.cpp
@@ -45,6 +45,8 @@ static FAutoConsoleVariableRef CVarGroomCacheStreamingEnable(TEXT("GroomCache.En
 
 #define LOCTEXT_NAMESPACE "GroomComponent"
 
+#define USE_HAIR_TRIANGLE_STRIP 1
+
 ////////////////////////////////////////////////////////////////////////////////////////////////////////////////////////
 
 EHairStrandsDebugMode GetHairStrandsGeometryDebugMode(const FHairGroupInstance* Instance);
@@ -774,7 +776,11 @@ public:
 		{
 			VertexFactory = (FVertexFactory*)Instance->Strands.VertexFactory;
 			HairVertexCount = Instance->Strands.RestResource->GetVertexCount();
+			#if USE_HAIR_TRIANGLE_STRIP
+			MaxVertexIndex = HairVertexCount * 2;
+			#else
 			MaxVertexIndex = HairVertexCount * 6;
+			#endif
 			bUseCulling = Instance->Strands.bIsCullingEnabled;
 			NumPrimitive = bUseCulling ? 0 : HairVertexCount * 2;
 			if (MaterialRenderProxy == nullptr && Instance->Strands.Material)
@@ -856,7 +862,11 @@ public:
 		BatchElement.MaxVertexIndex = MaxVertexIndex;
 		BatchElement.UserData = reinterpret_cast<void*>(uint64(ComponentId));
 		Mesh.ReverseCulling = bUseCardsOrMeshes ? IsLocalToWorldDeterminantNegative() : false;
-		Mesh.Type = PT_TriangleList;
+		#if USE_HAIR_TRIANGLE_STRIP
+		Mesh.Type = GeometryType == EHairGeometryType::Strands ? PT_TriangleStrip : PT_TriangleList;
+		#else
+		Mesh.Type = GeometryType == PT_TriangleList;
+		#endif
 		Mesh.DepthPriorityGroup = SDPG_World;
 		Mesh.bCanApplyViewModeOverrides = false;
 
diff --git a/Engine/Shaders/Private/HairStrands/HairStrandsClusterCulling.usf b/Engine/Shaders/Private/HairStrands/HairStrandsClusterCulling.usf
--- a/Engine/Shaders/Private/HairStrands/HairStrandsClusterCulling.usf
+++ b/Engine/Shaders/Private/HairStrands/HairStrandsClusterCulling.usf
@@ -4,6 +4,11 @@
 #include "../ColorMap.ush"
 #include "HairStrandsAABBCommon.ush"
 #include "HairStrandsClusterCommon.ush"
+#include "HairStrandsVertexFactoryCommon.ush"
+
+#ifndef USE_HAIR_TRIANGLE_STRIP
+	#error Hair triangle geometry type needs to be defined
+#endif
 
 #define COMPACTCS_NUMTHREAD 64
 
@@ -607,7 +612,11 @@ void MainClusterCullingCS(uint3 DispatchThreadId : SV_DispatchThreadID)
 					InterlockedAdd(DispatchIndirectParametersClusterCount[0], 1, VisibleClusterId);
 
 					//uint PreviousVertexCount = 0;
+				#if USE_HAIR_TRIANGLE_STRIP
+					InterlockedAdd(DrawIndirectParameters[0], ClusterLOD.VertexCount * 2/*, PreviousVertexCount*/);
+				#else
 					InterlockedAdd(DrawIndirectParameters[0], ClusterLOD.VertexCount * 6/*, PreviousVertexCount*/);
+				#endif
 
 					GlobalClusterIdBuffer[VisibleClusterId]	  = ClusterId;
 					GlobalIndexStartBuffer[VisibleClusterId]  = ClusterLOD.VertexOffset;
@@ -658,7 +667,11 @@ void MainClusterCullingPrepareIndirectDispatchCS(
 	{
 		// DrawIndirectBuffer contains the total number of vertex count needed for drawing quad triangle strips, so 6 vertices per quad / control point
 		const uint VertexCount    = DrawIndirectBuffer[0];
-		const uint ControlPoint   = VertexCount / 6;
+#if USE_HAIR_TRIANGLE_STRIP
+		const uint ControlPoint   = VertexCount / 2;
+#else
+		const uint ControlPoint	  = VertexCount / 6;
+#endif
 		const uint DispatchCount  = DivideAndRoundUp(ControlPoint, GROUP_SIZE);
 		const uint DispatchCountX = DispatchCount >= INDIRECT_GROUP_COUNT_X ? INDIRECT_GROUP_COUNT_X : (DispatchCount % INDIRECT_GROUP_COUNT_X);
 		const uint DispatchCountY = DivideAndRoundUp(DispatchCount, INDIRECT_GROUP_COUNT_X);
diff --git a/Engine/Shaders/Private/HairStrands/HairStrandsMaterialCommon.ush b/Engine/Shaders/Private/HairStrands/HairStrandsMaterialCommon.ush
--- a/Engine/Shaders/Private/HairStrands/HairStrandsMaterialCommon.ush
+++ b/Engine/Shaders/Private/HairStrands/HairStrandsMaterialCommon.ush
@@ -23,6 +23,12 @@ float3 DecodeVelocityFromTexture(float4 In);
 #include "/Engine/Generated/Material.ush"
 #include "/Engine/Generated/VertexFactory.ush" 
 
+#include "HairStrandsVertexFactoryCommon.ush"
+
+#ifndef USE_HAIR_TRIANGLE_STRIP
+	#error Hair triangle geometry type needs to be defined
+#endif
+
 #define SUPPORT_MATERIAL_PROPERTY 1
 uint MaterialPass_MacroGroupId;
 uint MaterialPass_MaterialId;  
@@ -60,8 +66,11 @@ FVertexFactoryInterpolantsVSToPS GetInterpolants(uint PrimitiveId, uint HairPrim
 	Input.InstanceIdOffset = PrimitiveId | (1U << 31U);
 	Input.DrawInstanceId = 0;
 #endif // VF_USE_PRIMITIVE_SCENE_DATA 
+#if USE_HAIR_TRIANGLE_STRIP
+	Input.VertexId = HairPrimitiveId * 2; // 6: the vertex factory expect triangles strip index. A quad is made of two triangles with indices: 0,1.
+#else
 	Input.VertexId = HairPrimitiveId * 6; // 6: the vertex factory expect triangles list index. A quad is made of two triangles with indices: 0,1,2,3,4,5.
-
+#endif
 	FVertexFactoryIntermediates	Intermediates = GetVertexFactoryIntermediates(Input);
 	FVertexFactoryInterpolantsVSToPS Interpolants = VertexFactoryGetInterpolantsVSToPS(Input, Intermediates, Empty);
 
diff --git a/Engine/Shaders/Private/HairStrands/HairStrandsVertexFactory.ush b/Engine/Shaders/Private/HairStrands/HairStrandsVertexFactory.ush
--- a/Engine/Shaders/Private/HairStrands/HairStrandsVertexFactory.ush
+++ b/Engine/Shaders/Private/HairStrands/HairStrandsVertexFactory.ush
@@ -13,6 +13,10 @@
 
 #define USE_HAIR_COMPLEX_TRANSMITTANCE 1 
 
+#ifndef USE_HAIR_TRIANGLE_STRIP
+	#error Hair triangle geometry type needs to be defined
+#endif
+
 #if MANUAL_VERTEX_FETCH
 	Buffer<float4> VertexFetch_InstanceOriginBuffer;
 	Buffer<float4> VertexFetch_InstanceTransformBuffer;
@@ -118,11 +122,17 @@ FVertexInfo GetVertexInfo(FVertexFactoryInput Input)
 	uint VertexId = Input.VertexId;
 
 	FVertexInfo VertexInfo;
+#if USE_HAIR_TRIANGLE_STRIP
+	VertexInfo.QuadIndex	= VertexId % 2;
+	VertexInfo.BaseIndex	= VertexId / 2;
+	VertexInfo.IsTip		= 0;
+	VertexInfo.IsLeft		= VertexInfo.QuadIndex == 0 ? 1 : 0;
+#else
 	VertexInfo.QuadIndex	= VertexId % 6;
 	VertexInfo.BaseIndex	= VertexId / 6;
 	VertexInfo.IsTip		= VertexInfo.QuadIndex == 0 || VertexInfo.QuadIndex == 2 || VertexInfo.QuadIndex == 4 ? 0 : 1;
 	VertexInfo.IsLeft		= VertexInfo.QuadIndex == 0 || VertexInfo.QuadIndex == 1 || VertexInfo.QuadIndex == 5 ? 1 : 0;
-
+#endif
 	VertexInfo.HairPrimitiveId = VertexInfo.BaseIndex;
 	VertexInfo.VertexIndex = VertexInfo.BaseIndex + VertexInfo.IsTip;
 	VertexInfo.LodRadiusScale = 1.0f;
@@ -198,7 +208,11 @@ FHairControlPoint GetVertexPosition(const FVertexInfo VertexInfo, bool bInvalidJ
 	
 	// Create a degenerated quad the end of each strand to cut between each strands
 	// #hair_todo: This is not efficient for short strand like fur for instance. Need to revisit that at some point
+#if USE_HAIR_TRIANGLE_STRIP
+	const bool bIsInvalidQuad = (bInvalidJointVertex && Out.Type == HAIR_CONTROLPOINT_END) || VertexInfo.bForceInvalidQuad;
+#else
 	const bool bIsInvalidQuad = (bInvalidJointVertex && Out.Type == HAIR_CONTROLPOINT_START && VertexInfo.IsTip == 1) || (bInvalidJointVertex && Out.Type == HAIR_CONTROLPOINT_END && VertexInfo.IsTip == 0) || VertexInfo.bForceInvalidQuad;
+#endif
 	Out.Position = bIsInvalidQuad ? float3(INFINITE_FLOAT, INFINITE_FLOAT, INFINITE_FLOAT) : Out.Position;
 
 	return Out;
@@ -230,7 +244,11 @@ float3 GetVertexPreviousPosition(FVertexFactoryInput Input, bool bInvalidJointVe
 
 	// Create a degenerated quad the end of each strand to cut between each strands
 	// #hair_todo: This is not efficient for short strand like fur for instance. Need to revisit that at some point
+#if USE_HAIR_TRIANGLE_STRIP
+	const bool bIsInvalidQuad = (bInvalidJointVertex && Out.Type == HAIR_CONTROLPOINT_END) || VertexInfo.bForceInvalidQuad;
+#else
 	const bool bIsInvalidQuad = (bInvalidJointVertex && Out.Type == HAIR_CONTROLPOINT_START && VertexInfo.IsTip == 1) || (bInvalidJointVertex && Out.Type == HAIR_CONTROLPOINT_END && VertexInfo.IsTip == 0) || VertexInfo.bForceInvalidQuad;
+#endif
 	Out.Position = bIsInvalidQuad ? float3(INFINITE_FLOAT, INFINITE_FLOAT, INFINITE_FLOAT) : Out.Position;
 
 	return Out.Position;
diff --git a/Engine/Shaders/Private/HairStrands/HairStrandsVertexFactoryCommon.ush b/Engine/Shaders/Private/HairStrands/HairStrandsVertexFactoryCommon.ush
--- a/Engine/Shaders/Private/HairStrands/HairStrandsVertexFactoryCommon.ush
+++ b/Engine/Shaders/Private/HairStrands/HairStrandsVertexFactoryCommon.ush
@@ -6,6 +6,8 @@
 #define HAIR_CONTROLPOINT_START	1
 #define HAIR_CONTROLPOINT_END	2
 
+#define USE_HAIR_TRIANGLE_STRIP 1
+
 ////////////////////////////////////////////////////////////////////////////////////////////////////////////////////////////////
 // Hair control points
 
