Fixes a source of a blur when reprojecting TSR's history

diff --git a/Engine/Shaders/Private/TemporalSuperResolution/TSRDilateVelocity.usf b/Engine/Shaders/Private/TemporalSuperResolution/TSRDilateVelocity.usf
--- a/Engine/Shaders/Private/TemporalSuperResolution/TSRDilateVelocity.usf
+++ b/Engine/Shaders/Private/TemporalSuperResolution/TSRDilateVelocity.usf
@@ -9,11 +9,14 @@
 
 #define AA_CROSS 1
 
+#define CONFIG_DEPTH_EDGE_DETECTION 1
+
 
 //------------------------------------------------------- PARAMETERS
 
 float2 PrevOutputBufferUVMin;
 float2 PrevOutputBufferUVMax;
+float WorldDepthToDepthError;
 
 Texture2D SceneDepthTexture;
 Texture2D SceneVelocityTexture;
@@ -79,6 +82,13 @@ void MainCS(
 		float FurthestPairDeviceZ = DeviceZ;
 		float FurthestDeviceZ = DeviceZ;
 
+		float DeviceZError;
+		{
+			float Depth = ConvertFromDeviceZ(DeviceZ);
+			float ErrorCorrectedDeviceZ = ConvertToDeviceZ(Depth + Depth * WorldDepthToDepthError);
+			DeviceZError = abs(ErrorCorrectedDeviceZ - DeviceZ);
+		}
+
 		#if 1
 			UNROLL_N(4)
 			for (uint i = 0; i < 4; i++)
@@ -92,8 +102,24 @@ void MainCS(
 			float SampleDeviceZ0 = SceneDepthTexture[int2(InputPixelPos) + Offset].r;
 			float SampleDeviceZ1 = SceneDepthTexture[int2(InputPixelPos) - Offset].r;
 			
-			PackedDepthOffset = max(PackedDepthOffset, PackDepthAndOffset(+Offset, SampleDeviceZ0));
-			PackedDepthOffset = max(PackedDepthOffset, PackDepthAndOffset(-Offset, SampleDeviceZ1));
+			#if CONFIG_DEPTH_EDGE_DETECTION
+			{
+				float DepthDiff = abs(SampleDeviceZ1 - SampleDeviceZ0);
+				float DepthVariation = (SampleDeviceZ0 + SampleDeviceZ1) * 0.5 - DeviceZ;
+
+				FLATTEN
+				if (DepthVariation > max(DepthDiff * 0.25, DeviceZError))
+				{
+					PackedDepthOffset = max(PackedDepthOffset, PackDepthAndOffset(+Offset, SampleDeviceZ0));
+					PackedDepthOffset = max(PackedDepthOffset, PackDepthAndOffset(-Offset, SampleDeviceZ1));
+				}
+			}
+			#else
+			{
+				PackedDepthOffset = max(PackedDepthOffset, PackDepthAndOffset(+Offset, SampleDeviceZ0));
+				PackedDepthOffset = max(PackedDepthOffset, PackDepthAndOffset(-Offset, SampleDeviceZ1));
+			}
+			#endif
 
 			FurthestPairDeviceZ = min(FurthestPairDeviceZ, max(SampleDeviceZ0, SampleDeviceZ1));
 			FurthestDeviceZ = min(FurthestDeviceZ, min(SampleDeviceZ0, SampleDeviceZ1));
@@ -167,6 +193,9 @@ void MainCS(
 		}
 		#endif
 		
+		float4 EncodedVelocity = SceneVelocityTexture[uint2(InputBufferUV * InputInfo_Extent)];
+		
+		//DilatedVelocityOutput[InputPixelPos] = EncodedVelocity.xy; //EncodeVelocityToTexture(float3(FinalScreenVelocity, 0.0)).xy;
 		DilatedVelocityOutput[InputPixelPos] = EncodeVelocityToTexture(float3(FinalScreenVelocity, 0.0)).xy;
 		ClosestDepthOutput[InputPixelPos] = ClosestDeviceZ;
 		ParallaxFactorOutput[InputPixelPos] = CompressParallaxFactor(ParallaxFactorForHistory) | (CompressParallaxFactor(ParallaxFactorForHistoryRejection) << 4);
diff --git a/Engine/Source/Runtime/Renderer/Private/PostProcess/TemporalSuperResolution.cpp b/Engine/Source/Runtime/Renderer/Private/PostProcess/TemporalSuperResolution.cpp
--- a/Engine/Source/Runtime/Renderer/Private/PostProcess/TemporalSuperResolution.cpp
+++ b/Engine/Source/Runtime/Renderer/Private/PostProcess/TemporalSuperResolution.cpp
@@ -184,6 +184,7 @@ class FTSRDilateVelocityCS : public FTSRShader
 
 		SHADER_PARAMETER(FVector2D, PrevOutputBufferUVMin)
 		SHADER_PARAMETER(FVector2D, PrevOutputBufferUVMax)
+		SHADER_PARAMETER(float, WorldDepthToDepthError)
 
 		SHADER_PARAMETER_RDG_TEXTURE(Texture2D, SceneDepthTexture)
 		SHADER_PARAMETER_RDG_TEXTURE(Texture2D, SceneVelocityTexture)
@@ -714,6 +715,14 @@ void AddTemporalSuperResolutionPasses(
 		PassParameters->CommonParameters = CommonParameters;
 		PassParameters->PrevOutputBufferUVMin = CommonParameters.InputInfo.UVViewportBilinearMin - CommonParameters.InputInfo.ExtentInverse;
 		PassParameters->PrevOutputBufferUVMax = CommonParameters.InputInfo.UVViewportBilinearMax + CommonParameters.InputInfo.ExtentInverse;
+		{
+			float TanHalfFieldOfView = View.ViewMatrices.GetInvProjectionMatrix().M[0][0];
+
+			// Should be multiplied 0.5* for the diameter to radius, and by 2.0 because GetTanHalfFieldOfView() cover only half of the pixels.
+			float WorldDepthToPixelWorldRadius = TanHalfFieldOfView / float(View.ViewRect.Width());
+
+			PassParameters->WorldDepthToDepthError = WorldDepthToPixelWorldRadius * 2.0f;
+		}
 		PassParameters->SceneDepthTexture = PassInputs.SceneDepthTexture;
 		PassParameters->SceneVelocityTexture = PassInputs.SceneVelocityTexture;
 		PassParameters->DilatedVelocityOutput = GraphBuilder.CreateUAV(DilatedVelocityTexture);
