Fix shader typo.

diff --git a/Engine/Shaders/Private/HairStrands/HairStrandsVisibilityDepthPS.usf b/Engine/Shaders/Private/HairStrands/HairStrandsVisibilityDepthPS.usf
--- a/Engine/Shaders/Private/HairStrands/HairStrandsVisibilityDepthPS.usf
+++ b/Engine/Shaders/Private/HairStrands/HairStrandsVisibilityDepthPS.usf
@@ -90,7 +90,7 @@ void MainPS(
 	in FScreenVertexOutput Input)
 {
 	const uint3 PixelCoord = uint3(Input.Position.xy, 0);
-	OutLightChannelMaskTexture[PixelCoord.xy] = ExtractSampleData(uint3(PixelCoord, 0)).LightChannelMask;
+	OutLightChannelMaskTexture[PixelCoord.xy] = ExtractSampleData(PixelCoord).LightChannelMask;
 }
 #endif
 
