Make DepthOnly pixel shader more likely to generate same code as ShadowDepth pixel shader

diff --git a/Engine/Shaders/Private/DepthOnlyPixelShader.usf b/Engine/Shaders/Private/DepthOnlyPixelShader.usf
--- a/Engine/Shaders/Private/DepthOnlyPixelShader.usf
+++ b/Engine/Shaders/Private/DepthOnlyPixelShader.usf
@@ -19,11 +19,11 @@ void Main(
 #endif
 
 #if !MATERIALBLENDING_SOLID || OUTPUT_PIXEL_DEPTH_OFFSET
+	FVertexFactoryInterpolantsVSToPS FactoryInterpolants
+	// ,float4 PixelPosition : TEXCOORD6 -- Not used currently
 #if USE_WORLD_POSITION_EXCLUDING_SHADER_OFFSETS
-	float3 PixelPositionExcludingWPO : TEXCOORD7,
+	,float3 PixelPositionExcludingWPO : TEXCOORD7
 #endif
-	FVertexFactoryInterpolantsVSToPS FactoryInterpolants
-	// float4 PixelPosition : TEXCOORD6 -- Not used currently
 	OPTIONAL_IsFrontFace
 	OPTIONAL_OutDepthConservative,
 #endif
