Fix invalid hair depth when visualizing Scene Depth World Units in editor.

diff --git a/Engine/Shaders/Private/HairStrands/HairStrandsGBufferWrite.usf b/Engine/Shaders/Private/HairStrands/HairStrandsGBufferWrite.usf
--- a/Engine/Shaders/Private/HairStrands/HairStrandsGBufferWrite.usf
+++ b/Engine/Shaders/Private/HairStrands/HairStrandsGBufferWrite.usf
@@ -86,7 +86,6 @@ void MainPS(
 		Samples_BaseColor	*= InvSampleCount;
 		Samples_Roughnesss	*= InvSampleCount;
 		Samples_Backlit		*= InvSampleCount;
-		Samples_Depth		*= InvSampleCount;
 
 		WriteGBuffer(
 			Samples_Tangent,
