Fix invalid hair material attribute cast.

diff --git a/Engine/Shaders/Private/HairStrands/HairStrandsAttributeCommon.ush b/Engine/Shaders/Private/HairStrands/HairStrandsAttributeCommon.ush
--- a/Engine/Shaders/Private/HairStrands/HairStrandsAttributeCommon.ush
+++ b/Engine/Shaders/Private/HairStrands/HairStrandsAttributeCommon.ush
@@ -42,7 +42,7 @@ uint4 GetHairStrandsOptionalAttribute1(uint InVertexId)
 
 float4 GetHairStrandsOptionalMaterial(uint InVertexId)
 {
-	uint4 Out = 0;
+	float4 Out = 0;
 	if (HairStrandsVF_HasMaterial > 0)
 	{
 		Out = HairStrandsVF_MaterialBuffer[InVertexId];
