Fix shader name in comments
#rnx

diff --git a/Engine/Shaders/Private/VariableRateShading.usf b/Engine/Shaders/Private/VariableRateShading.usf
--- a/Engine/Shaders/Private/VariableRateShading.usf
+++ b/Engine/Shaders/Private/VariableRateShading.usf
@@ -1,7 +1,7 @@
 // Copyright Epic Games, Inc. All Rights Reserved.
 
 /*=============================================================================
-	FixedFoveation.usf
+	VariableRateShading.usf
 =============================================================================*/
 
 #include "Common.ush"
