Workarround for a GL shader compiler bug on Adreno devices related to usage of a struct with static array field in SSBO

diff --git a/Engine/Shaders/Private/InstanceCulling/BuildInstanceDrawCommands.usf b/Engine/Shaders/Private/InstanceCulling/BuildInstanceDrawCommands.usf
--- a/Engine/Shaders/Private/InstanceCulling/BuildInstanceDrawCommands.usf
+++ b/Engine/Shaders/Private/InstanceCulling/BuildInstanceDrawCommands.usf
@@ -48,9 +48,20 @@ struct FContextBatchInfo
 	uint NumViewIds;
 	uint DynamicInstanceIdOffset;
 	uint DynamicInstanceIdMax;
-	uint ItemDataOffset[BATCH_PROCESSING_MODE_NUM];
+	// Must be equal to ItemDataOffset[BATCH_PROCESSING_MODE_NUM]
+	// Not using array here to avoid GL shader compiler bug on Adreno devices
+	uint ItemDataOffset0;
+	uint ItemDataOffset1;
 };
 
+uint GetItemDataOffset(FContextBatchInfo BatchInfo, uint BatchProcessingMode)
+{
+#if BATCH_PROCESSING_MODE_NUM > 2
+	#error Make sure FContextBatchInfo::ItemDataOffset is up to date with BATCH_PROCESSING_MODE_NUM
+#else
+	return (BatchProcessingMode == 0) ? BatchInfo.ItemDataOffset0 : BatchInfo.ItemDataOffset1;
+#endif
+}
 
 #if ENABLE_BATCH_MODE
 StructuredBuffer<FContextBatchInfo> BatchInfos;
@@ -271,7 +282,7 @@ void InstanceCullBuildInstanceIdBufferCS(uint3 GroupId : SV_GroupID, int GroupTh
 	BatchInfo.DynamicInstanceIdMax = DynamicInstanceIdMax;
 #endif // ENABLE_BATCH_MODE
 
-	FInstanceWorkSetup WorkSetup = InstanceCullingLoadBalancer_Setup(GroupId, GroupThreadIndex, BatchInfo.ItemDataOffset[CurrentBatchProcessingMode]);
+	FInstanceWorkSetup WorkSetup = InstanceCullingLoadBalancer_Setup(GroupId, GroupThreadIndex, GetItemDataOffset(BatchInfo, CurrentBatchProcessingMode));
 
 	if (!WorkSetup.bValid)
 	{
