Add GPU-Scene support to GeometryCache VF

diff --git a/Engine/Plugins/Experimental/GeometryCache/Source/GeometryCache/Private/GeometryCacheSceneProxy.cpp b/Engine/Plugins/Experimental/GeometryCache/Source/GeometryCache/Private/GeometryCacheSceneProxy.cpp
--- a/Engine/Plugins/Experimental/GeometryCache/Source/GeometryCache/Private/GeometryCacheSceneProxy.cpp
+++ b/Engine/Plugins/Experimental/GeometryCache/Source/GeometryCache/Private/GeometryCacheSceneProxy.cpp
@@ -62,6 +62,8 @@ FGeometryCacheSceneProxy::FGeometryCacheSceneProxy(UGeometryCacheComponent* Comp
 , MaterialRelevance(Component->GetMaterialRelevance(GetScene().GetFeatureLevel()))
 , CreateTrackProxy(TrackProxyCreator)
 {
+	const ERHIFeatureLevel::Type FeatureLevel = GetScene().GetFeatureLevel();
+
 	Time = Component->GetAnimationTime();
 	bLooping = Component->IsLooping();
 	bExtrapolateFrames = Component->IsExtrapolatingFrames();
@@ -70,6 +72,10 @@ FGeometryCacheSceneProxy::FGeometryCacheSceneProxy(UGeometryCacheComponent* Comp
 	MotionVectorScale = Component->GetMotionVectorScale();
 	UpdatedFrameNum = 0;
 
+	bVFRequiresPrimitiveUniformBuffer = !UseGPUScene(GMaxRHIShaderPlatform, FeatureLevel) || (FeatureLevel == ERHIFeatureLevel::ES3_1);
+
+	bCanSkipRedundantTransformUpdates = false;
+
 	// Copy each section
 	const int32 NumTracks = Component->TrackSections.Num();
 	Tracks.Reserve(NumTracks);
@@ -379,7 +385,7 @@ void FGeometryCacheSceneProxy::CreateMeshBatch(
 	const FMatrix& LocalToWorldTransform = TrackProxy->WorldMatrix * GetLocalToWorld();
 
 	DynamicPrimitiveUniformBuffer.Set(LocalToWorldTransform, LocalToWorldTransform, GetBounds(), GetLocalBounds(), true, false, DrawsVelocity(), false);
-	BatchElement.PrimitiveUniformBuffer = DynamicPrimitiveUniformBuffer.UniformBuffer.GetUniformBufferRHI();
+	BatchElement.PrimitiveUniformBufferResource = &DynamicPrimitiveUniformBuffer.UniformBuffer;
 
 	const FGeometryCacheMeshData* MeshData = TrackProxy->bNextFrameMeshDataSelected ? TrackProxy->NextFrameMeshData : TrackProxy->MeshData;
 
diff --git a/Engine/Shaders/Private/GeometryCacheVertexFactory.ush b/Engine/Shaders/Private/GeometryCacheVertexFactory.ush
--- a/Engine/Shaders/Private/GeometryCacheVertexFactory.ush
+++ b/Engine/Shaders/Private/GeometryCacheVertexFactory.ush
@@ -69,6 +69,8 @@ struct FVertexFactoryInput
 		#error Too many texture coordinate sets defined on Geometry cache vertex input. Max: 4.
 	#endif
 #endif
+
+	VF_GPUSCENE_DECLARE_INPUT_BLOCK(13)
 };
 
 // RHI_RAYTRACING
@@ -105,6 +107,8 @@ FVertexFactoryInput LoadVertexFactoryInputForHGS(uint TriangleIndex, int VertexI
 	#endif
 #endif
 
+	VF_GPUSCENE_SET_INPUT_FOR_RT(Input, GetInstanceUserData(), 0U);
+
 	return Input;
 }
 #endif
@@ -112,32 +116,34 @@ FVertexFactoryInput LoadVertexFactoryInputForHGS(uint TriangleIndex, int VertexI
 struct FPositionOnlyVertexFactoryInput
 {
 	float4	Position	: ATTRIBUTE0;
+	VF_GPUSCENE_DECLARE_INPUT_BLOCK(1)
 };
 
 struct FPositionAndNormalOnlyVertexFactoryInput
 {
 	float4	Position	: ATTRIBUTE0;
 	float4	Normal		: ATTRIBUTE1;
+	VF_GPUSCENE_DECLARE_INPUT_BLOCK(2)
 };
 
 /** for depth-only pass */
 float4 VertexFactoryGetWorldPosition(FPositionOnlyVertexFactoryInput Input)
 {
-	float4 Position = Input.Position;
-	return TransformLocalToTranslatedWorld(Position.xyz);
+	FSceneDataIntermediates SceneData = VF_GPUSCENE_GET_INTERMEDIATES(Input);
+	return TransformLocalToTranslatedWorld(Input.Position.xyz, SceneData.InstanceData.LocalToWorld);
 }
 
 /** for shadow pass (used for slope bias) */
 float4 VertexFactoryGetWorldPosition(FPositionAndNormalOnlyVertexFactoryInput Input)
 {
-	float4 Position = Input.Position;
-	return TransformLocalToTranslatedWorld(Position.xyz);
+	FSceneDataIntermediates SceneData = VF_GPUSCENE_GET_INTERMEDIATES(Input);
+	return TransformLocalToTranslatedWorld(Input.Position.xyz, SceneData.InstanceData.LocalToWorld);
 }
 
 float3 VertexFactoryGetWorldNormal(FPositionAndNormalOnlyVertexFactoryInput Input)
 {
-	float4 Normal = Input.Normal;
-	return RotateLocalToWorld(Normal.xyz);
+	FSceneDataIntermediates SceneData = VF_GPUSCENE_GET_INTERMEDIATES(Input);
+	return RotateLocalToWorld(Input.Normal.xyz, SceneData.InstanceData.LocalToWorld, SceneData.InstanceData.InvNonUniformScale);
 }
 
 struct FVertexFactoryInterpolantsVSToPS
@@ -155,6 +161,10 @@ struct FVertexFactoryInterpolantsVSToPS
 	nointerpolation uint EyeIndex : PACKED_EYE_INDEX;
 #endif
 
+#if VF_USE_PRIMITIVE_SCENE_DATA
+	nointerpolation uint PrimitiveId : PRIMITIVE_ID;
+#endif
+
 	TANGENTTOWORLD_INTERPOLATOR_BLOCK
 };
 
@@ -211,6 +221,11 @@ FMaterialPixelParameters GetMaterialPixelParameters(FVertexFactoryInterpolantsVS
 	Result.VertexColor = 0;
 #endif
 	Result.TwoSidedSign = 1;
+
+#if VF_USE_PRIMITIVE_SCENE_DATA
+	Result.PrimitiveId = Interpolants.PrimitiveId;
+#endif
+
 	return Result;
 }
 
@@ -261,6 +276,11 @@ FMaterialVertexParameters GetMaterialVertexParameters(FVertexFactoryInput Input,
 		Result.TexCoords[3] = Input.TexCoords3;
 	#endif
 #endif
+
+#if VF_USE_PRIMITIVE_SCENE_DATA
+	Result.PrimitiveId = Intermediates.SceneData.PrimitiveId;
+#endif
+
 	return Result;
 }
 
@@ -269,7 +289,7 @@ FMaterialVertexParameters GetMaterialVertexParameters(FVertexFactoryInput Input,
 */
 float3 UnpackedPosition( FVertexFactoryInput Input )
 {
-		return float3(Input.Position.xyz * GeomCache.MeshExtension + GeomCache.MeshOrigin);
+	return float3(Input.Position.xyz * GeomCache.MeshExtension + GeomCache.MeshOrigin);
 }
 
 /**
@@ -298,7 +318,7 @@ half3x3 CalcTangentToLocal(FVertexFactoryInput Input)
 
 void CalcTangentToWorld(FVertexFactoryInput Input, FVertexFactoryIntermediates Intermediates, out float3 TangentToWorld0, out float4 TangentToWorld2)
 {
-	float3x3 LocalToWorld = GetLocalToWorld3x3();
+	float3x3 LocalToWorld = (float3x3)Intermediates.SceneData.InstanceData.LocalToWorld;
 
 	// Remove scaling.
 	half3 InvScale = GetPrimitiveData(Intermediates).InvNonUniformScale;
@@ -350,7 +370,7 @@ float3x3 VertexFactoryGetTangentToLocal( FVertexFactoryInput Input, FVertexFacto
 
 float4 VertexFactoryGetWorldPosition(FVertexFactoryInput Input, FVertexFactoryIntermediates Intermediates)
 {
-	return TransformLocalToTranslatedWorld(Intermediates.UnpackedPosition);
+	return TransformLocalToTranslatedWorld(Intermediates.UnpackedPosition, Intermediates.SceneData.InstanceData.LocalToWorld);
 }
 
 float4 VertexFactoryGetRasterizedWorldPosition(FVertexFactoryInput Input, FVertexFactoryIntermediates Intermediates, float4 InWorldPosition)
@@ -397,6 +417,10 @@ FVertexFactoryInterpolantsVSToPS VertexFactoryGetInterpolantsVSToPS(FVertexFacto
 	Interpolants.EyeIndex = 0;
 #endif
 
+#if VF_USE_PRIMITIVE_SCENE_DATA
+	Interpolants.PrimitiveId = Intermediates.SceneData.PrimitiveId;
+#endif
+
 	return Interpolants;
 }
 
@@ -467,6 +491,10 @@ FVertexFactoryRayTracingInterpolants VertexFactoryInterpolate(FVertexFactoryRayT
 	}
 #endif
 
+#if VF_USE_PRIMITIVE_SCENE_DATA
+	O.InterpolantsVSToPS.PrimitiveId = a.InterpolantsVSToPS.PrimitiveId;
+#endif
+
 	return O;
 }
 #endif // #if NEEDS_VERTEX_FACTORY_INTERPOLATION
@@ -478,7 +506,11 @@ float4 VertexFactoryGetTranslatedPrimitiveVolumeBounds(FVertexFactoryInterpolant
 
 uint VertexFactoryGetPrimitiveId(FVertexFactoryInterpolantsVSToPS Interpolants)
 {
+#if VF_USE_PRIMITIVE_SCENE_DATA
+	return Interpolants.PrimitiveId;
+#else // VF_USE_PRIMITIVE_SCENE_DATA
 	return 0;
+#endif // VF_USE_PRIMITIVE_SCENE_DATA
 }
 
 #include "VertexFactoryDefaultInterface.ush"
\ No newline at end of file
diff --git a/Engine/Shaders/Private/VertexFactoryCommon.ush b/Engine/Shaders/Private/VertexFactoryCommon.ush
--- a/Engine/Shaders/Private/VertexFactoryCommon.ush
+++ b/Engine/Shaders/Private/VertexFactoryCommon.ush
@@ -22,6 +22,11 @@
 	#define VF_GPUSCENE_GET_SHADOWMAP_UV_SCALE_BIAS(VFInput, LightmapDataIndex) \
 		GetLightmapData(LightmapDataIndex).ShadowMapCoordinateScaleBias
 
+	#define VF_GPUSCENE_SET_INPUT_FOR_RT(_VFInputOut_, _PrimitiveId_, _DrawInstanceId_) \
+		_VFInputOut_.InstanceIdOffset = (_PrimitiveId_) | VF_TREAT_INSTANCE_ID_OFFSET_AS_PRIMITIVE_ID_FLAG;	\
+		_VFInputOut_.DrawInstanceId = _DrawInstanceId_;
+
+
 #elif (VF_USE_PRIMITIVE_SCENE_DATA == 2)	
 	#define VF_GPUSCENE_DECLARE_INPUT_BLOCK(Attr0) \
 		float4 InstanceOrigin : ATTRIBUTE8; \
@@ -46,11 +51,14 @@
 	#define VF_GPUSCENE_GET_SHADOWMAP_UV_SCALE_BIAS(VFInput, LightmapDataIndex) \
 		float4(UnpackUnorm2x16(asuint(VFInput.InstanceAuxData.z)), UnpackSnorm2x16(asuint(VFInput.InstanceAuxData.w)))
 
+	#define VF_GPUSCENE_SET_INPUT_FOR_RT(_VFInputOut_, _PrimitiveId_, _DrawInstanceId_)
+
 #else
 	#define VF_GPUSCENE_DECLARE_INPUT_BLOCK(Attr0) 
 	#define VF_GPUSCENE_GET_INTERMEDIATES(VFInput) GetSceneDataIntermediates()
 	#define VF_GPUSCENE_GET_LIGHTMAP_UV_SCALE_BIAS(VFInput, LightmapDataIndex) PrecomputedLightingBuffer.LightMapCoordinateScaleBias
 	#define VF_GPUSCENE_GET_SHADOWMAP_UV_SCALE_BIAS(VFInput, LightmapDataIndex) PrecomputedLightingBuffer.ShadowMapCoordinateScaleBias
+	#define VF_GPUSCENE_SET_INPUT_FOR_RT(_VFInputOut_, _PrimitiveId_, _DrawInstanceId_)
 #endif
 
 
diff --git a/Engine/Source/Runtime/Engine/Private/GeometryCacheVertexFactory.cpp b/Engine/Source/Runtime/Engine/Private/GeometryCacheVertexFactory.cpp
--- a/Engine/Source/Runtime/Engine/Private/GeometryCacheVertexFactory.cpp
+++ b/Engine/Source/Runtime/Engine/Private/GeometryCacheVertexFactory.cpp
@@ -87,6 +87,11 @@ FGPUSkinPassthroughVertexFactory
 void FGeometryCacheVertexVertexFactory::ModifyCompilationEnvironment(const FVertexFactoryShaderPermutationParameters& Parameters, FShaderCompilerEnvironment& OutEnvironment)
 {
 	Super::ModifyCompilationEnvironment(Parameters, OutEnvironment);
+
+	const bool bUseGPUScene = UseGPUScene(Parameters.Platform, GetMaxSupportedFeatureLevel(Parameters.Platform)) && GetMaxSupportedFeatureLevel(Parameters.Platform) > ERHIFeatureLevel::ES3_1;
+	const bool bSupportsPrimitiveIdStream = Parameters.VertexFactoryType->SupportsPrimitiveIdStream();
+
+	OutEnvironment.SetDefine(TEXT("VF_SUPPORTS_PRIMITIVE_SCENE_DATA"), bSupportsPrimitiveIdStream && bUseGPUScene);
 }
 
 void FGeometryCacheVertexVertexFactory::SetData(const FDataType& InData)
@@ -167,6 +172,7 @@ void FGeometryCacheVertexVertexFactory::InitRHI()
 		{
 			FVertexDeclarationElementList PositionOnlyStreamElements;
 			PositionOnlyStreamElements.Add(AccessStreamComponent(Data.PositionComponent, 0, EVertexInputStreamType::PositionOnly));
+			AddPrimitiveIdStreamElement(EVertexInputStreamType::PositionOnly, PositionOnlyStreamElements, 1, 0xff); // TODO: support instancing on mobile
 			InitDeclaration(PositionOnlyStreamElements, EVertexInputStreamType::PositionOnly);
 		}
 
@@ -174,6 +180,7 @@ void FGeometryCacheVertexVertexFactory::InitRHI()
 			FVertexDeclarationElementList PositionAndNormalOnlyStreamElements;
 			PositionAndNormalOnlyStreamElements.Add(AccessStreamComponent(Data.PositionComponent, 0, EVertexInputStreamType::PositionAndNormalOnly));
 			PositionAndNormalOnlyStreamElements.Add(AccessStreamComponent(Data.TangentBasisComponents[1], 1, EVertexInputStreamType::PositionAndNormalOnly));
+			AddPrimitiveIdStreamElement(EVertexInputStreamType::PositionAndNormalOnly, PositionAndNormalOnlyStreamElements, 2, 0xff); // TODO: support instancing on mobile
 			InitDeclaration(PositionAndNormalOnlyStreamElements, EVertexInputStreamType::PositionAndNormalOnly);
 		}
 	}
@@ -237,6 +244,8 @@ void FGeometryCacheVertexVertexFactory::InitRHI()
 		}
 	}
 
+	AddPrimitiveIdStreamElement(EVertexInputStreamType::Default, Elements, 13, 0xff); // TODO: support instancing on mobile
+
 	check(Streams.Num() > 0);
 	check(PositionStreamIndex >= 0);
 	check(MotionBlurDataStreamIndex >= 0);
@@ -342,4 +351,5 @@ IMPLEMENT_VERTEX_FACTORY_TYPE(FGeometryCacheVertexVertexFactory, "/Engine/Privat
 	| EVertexFactoryFlags::SupportsDynamicLighting
 	| EVertexFactoryFlags::SupportsPositionOnly
 	| EVertexFactoryFlags::SupportsRayTracing
+	| EVertexFactoryFlags::SupportsPrimitiveIdStream
 );
