Virtual shadow maps: early out on pixels backfacing the light in both projection and page allocation
Page allocation part is a fairly minor benefit in most scenes but can occasionally make a big difference.
Projection early-out is a pretty uniform benefit all the time, especially with SMRT.

diff --git a/Engine/Shaders/Private/VirtualShadowMaps/PageManagement.usf b/Engine/Shaders/Private/VirtualShadowMaps/PageManagement.usf
--- a/Engine/Shaders/Private/VirtualShadowMaps/PageManagement.usf
+++ b/Engine/Shaders/Private/VirtualShadowMaps/PageManagement.usf
@@ -18,6 +18,8 @@
 #include "PageCacheCommon.ush"
 #include "PageOverlap.ush"
 #include "ProjectionCommon.ush"
+#include "ProjectionSpot.ush"
+#include "ProjectionDirectional.ush"
 
 #ifndef VSM_GENERATE_STATS
 #define VSM_GENERATE_STATS 0
@@ -42,32 +44,24 @@ void StatsBufferInterlockedAdd(uint Index, uint Value)
 
 // Type of input data consume by the page allocation (i.e., data read from the source buffer: Gbuffer, HairStrands data, ...)
 #define INPUT_TYPE_GBUFFER 0
-#define INPUT_TYPE_NANITE 1
-#define INPUT_TYPE_HAIRSTRANDS 2
+#define INPUT_TYPE_HAIRSTRANDS 1
 
 // Flags generated by per-pixel pass to determine which pages are required to provide shadow for the visible geometry
 RWStructuredBuffer<uint> OutPageRequestFlags;
 
-Texture2D<UlongType> VisBuffer64;
 StructuredBuffer<int> DirectionalLightIds;
 
-uint bPostBasePass;
 float ResolutionLodBiasLocal;
 float PageDilationBorderSize;
 uint InputType;
+uint bCullBackfacingPixels;
 
 // Convenience to group up a bunch of data computed for a screen pixel
 struct FScreenPixelData
 {
-	int2 PixelPos;
-	float2 ScreenUV;
-	float4 ClipPosition;
-	float DeviceZ;
-	float SceneDepth;
-	float3 ViewPosition;
-	
 	// Note: Actual world position (not pre-translated)
 	float3 WorldPosition;
+	float SceneDepth;
 };
 
 float2 CalcScreenUV(uint2 PixelPos)
@@ -75,44 +69,22 @@ float2 CalcScreenUV(uint2 PixelPos)
 	return (float2(PixelPos.xy) + View.ViewRectMin.xy + 0.5f) * View.BufferSizeAndInvSize.zw;
 }
 
-// Load DeviceZ from depth buffer (and possibly combine with nanite depth)
-float GetDeviceZFromPixel(uint2 PixelPos)
-{
-	// Merge Nanite depth with depth buffer if enabled (i.e. we're before the base pass)
-	float DeviceZ = 0.0f;
-#if PERMUTATION_INPUT_TYPE == INPUT_TYPE_NANITE
-	// Load Nanite Depth
-	UlongType VisPixel = VisBuffer64[uint2(View.ViewRectMin.xy) + PixelPos];
-	uint DepthInt = 0;
-	uint VisibleClusterIndex = 0;
-	uint TriIndex = 0;
-	UnpackVisPixel(VisPixel, DepthInt, VisibleClusterIndex, TriIndex);
-
-	if (VisibleClusterIndex != 0xFFFFFFFF)
-	{
-		DeviceZ = asfloat(DepthInt);
-	}
-#endif // LOAD_DEPTH_FROM_NANITE_BUFFER
-	return max(DeviceZ, LookupDeviceZ(CalcScreenUV(PixelPos)));
-}
-
 // Use an explicit DeviceZ value for this pixel
 FScreenPixelData GetScreenPixelData(uint2 PixelPos, float DeviceZ)
 {
 	FScreenPixelData Data;
-	Data.PixelPos = PixelPos;
-	Data.DeviceZ = DeviceZ;
-	Data.ScreenUV = CalcScreenUV(PixelPos);
-	Data.SceneDepth = ConvertFromDeviceZ(Data.DeviceZ);
-	Data.ClipPosition = float4(((Data.ScreenUV.xy - View.ScreenPositionScaleBias.wz) / View.ScreenPositionScaleBias.xy), Data.DeviceZ, 1.0f);
+	Data.SceneDepth = ConvertFromDeviceZ(DeviceZ);
+
+	//float4 ClipPosition = float4(((CalcScreenUV(PixelPos) - View.ScreenPositionScaleBias.wz) / View.ScreenPositionScaleBias.xy), DeviceZ, 1.0f);
+	//float4 TranslatedWorldH = mul(ClipPosition, View.ClipToTranslatedWorld);
+	//float3 TranslatedWorldPosition = TranslatedWorldH.xyz / TranslatedWorldH.w;
 
-	float4 ViewH = mul(Data.ClipPosition, View.ClipToView);
-	Data.ViewPosition = ViewH.xyz / ViewH.w;
+	const float4 SvPosition = float4( float2( PixelPos ) + 0.5, DeviceZ, 1.0f );
+	const float3 TranslatedWorldPosition = SvPositionToTranslatedWorld( SvPosition );
 
-	float4 TranslatedWorldH = mul(Data.ClipPosition, View.ClipToTranslatedWorld);
-	float3 TranslatedWorldPosition = TranslatedWorldH.xyz / TranslatedWorldH.w;
 	// Subtract the view pre-translation to get back to world space
 	Data.WorldPosition = TranslatedWorldPosition - View.PreViewTranslation;
+
 	return Data;
 }
 
@@ -198,7 +170,6 @@ void GeneratePageFlagsFromPixels(
 	uint3 GroupThreadId : SV_GroupThreadID,
 	uint3 DispatchThreadId : SV_DispatchThreadID)
 {
-	// Morton order within a group so page access/atomics are more coherent and wave-swizzled gradients are possible.
 #if PERMUTATION_INPUT_TYPE == INPUT_TYPE_HAIRSTRANDS
 	uint2 GroupId		= InGroupId.xy;
 	if (HairStrands.bHairTileValid>0)
@@ -211,19 +182,20 @@ void GeneratePageFlagsFromPixels(
 		}
 		GroupId = HairStrands.HairTileData[LinearIndex];
 	}
-#else // PERMUTATION_INPUT_TYPE == INPUT_TYPE_GBUFFER || PERMUTATION_INPUT_TYPE == INPUT_TYPE_NANITE
+#else // PERMUTATION_INPUT_TYPE == INPUT_TYPE_GBUFFER
 	const uint2 GroupId = InGroupId.xy;
 #endif
+	// Morton order within a group so page access/atomics are more coherent and wave-swizzled gradients are possible.
 	const uint2 PixelPos = VSM_DEFAULT_CS_GROUP_XY * GroupId.xy + MortonDecode(GroupIndex);
-
 	if (any(PixelPos >= uint2(View.ViewSizeAndInvSize.xy)))
 	{
 		return;
 	}
 
 	float3 WorldPosition = 0;
+	half3  WorldNormal   = half3(0, 0, 0);
 	float  SceneDepth    = 0;
-	bool   bIsValid		 = false;
+	bool   bIsValid		 = true;
 	#if PERMUTATION_INPUT_TYPE == INPUT_TYPE_HAIRSTRANDS
 	{
 		const float HairDeviceZ = HairStrands.HairOnlyDepthTexture.Load(uint3(PixelPos, 0)).x;
@@ -232,18 +204,23 @@ void GeneratePageFlagsFromPixels(
 		SceneDepth = HairPixel.SceneDepth;
 		bIsValid = HairDeviceZ > 0;
 	}
-	#else // PERMUTATION_INPUT_TYPE == INPUT_TYPE_GBUFFER || PERMUTATION_INPUT_TYPE == INPUT_TYPE_NANITE
+	#else // PERMUTATION_INPUT_TYPE == INPUT_TYPE_GBUFFER
 	{
-		FScreenPixelData Pixel = GetScreenPixelData(PixelPos, GetDeviceZFromPixel(PixelPos));
+		FScreenPixelData Pixel = GetScreenPixelData(PixelPos, LookupDeviceZ(CalcScreenUV(PixelPos)));
 		WorldPosition = Pixel.WorldPosition;
 		SceneDepth = Pixel.SceneDepth;
-		bIsValid = true;
-		if (bPostBasePass != 0)
-		{
-			bIsValid = GetGBufferData(Pixel.ScreenUV).ShadingModelID != SHADINGMODELID_UNLIT;
-		}
+
+		FGBufferData GBuffer = GetGBufferDataUint(PixelPos, true);
+		WorldNormal = GBuffer.WorldNormal;
+		// Excluding unlit to avoid including processing sky dome
+		bIsValid = GBuffer.ShadingModelID != SHADINGMODELID_UNLIT;
 	}
 	#endif
+	
+	if (!bIsValid)
+	{
+		return;
+	}
 
 	// Dither pattern for page dilation
 	// We don't need to to check all 8 adjacent pages; as long as there's at least a single pixel near the edge
@@ -252,12 +229,6 @@ void GeneratePageFlagsFromPixels(
 	float2 PageDilationOffset = PageDilationBorderSize * 
 		float2(GroupIndex & 1 ? 1.0f : -1.0f,
 			   GroupIndex & 2 ? 1.0f : -1.0f);
-		
-	// excluding unlit to avoid including processing sky dome
-	if (!bIsValid)
-	{
-		return;
-	}
 	
 	// Directional lights
 	for (uint Index = 0; Index < VirtualShadowMap.NumDirectionalLights; ++Index)
@@ -265,6 +236,15 @@ void GeneratePageFlagsFromPixels(
 		int ClipmapStartId = DirectionalLightIds[Index];
 		FVirtualShadowMapProjectionShaderData BaseProjectionData = GetVirtualShadowMapProjectionData(ClipmapStartId);
 
+#if PERMUTATION_INPUT_TYPE != INPUT_TYPE_HAIRSTRANDS
+		// Backface test if requested
+		float3 LightDirection = -BaseProjectionData.TranslatedWorldToShadowViewMatrix._13_23_33;
+		if (bCullBackfacingPixels != 0 && IsBackfaceToDirectionalLight(normalize(WorldNormal), normalize(LightDirection), 0/*BaseProjectionData.LightSourceRadius*/))
+		{
+			continue;
+		}
+#endif
+
 		// Possibly use a single shadow projection data for the whole clipmap with just scales/biases for
 		// each level (to handle snapping). For now we maintain one per level.
 		//int ClipmapLevel = CalcClipmapLevel(length(WorldPosition - ProjectionData.ClipmapWorldOrigin), ClipmapResolutionLodBias);
@@ -289,6 +269,7 @@ void GeneratePageFlagsFromPixels(
 			const FLocalLightData LightData = GetLocalLightData(CulledLightGridData.DataStartIndex + Index, EyeIndex);
 			const uint LightGridLightIndex = ForwardLightData.CulledLightDataGrid[CulledLightGridData.DataStartIndex + Index];
 
+			// TODO: Can optimize these tests by fusing them together
 			// Also do precise cone test, since froxels are pretty coarse at times.
 			float3 ToLight = normalize(LightData.LightPositionAndInvRadius.xyz - WorldPosition);
 			if (dot(ToLight, LightData.LightDirectionAndShadowMask.xyz) < LightData.SpotAnglesAndSourceRadiusPacked.x)
@@ -296,6 +277,15 @@ void GeneratePageFlagsFromPixels(
 				continue;
 			}
 
+			// TODO: Precise radius test necessary?
+
+#if PERMUTATION_INPUT_TYPE != INPUT_TYPE_HAIRSTRANDS
+			// Backface test if requested
+			if (bCullBackfacingPixels != 0 && IsBackfaceToLocalLight(WorldPosition, WorldNormal, LightData.LightPositionAndInvRadius.xyz, LightData.SpotAnglesAndSourceRadiusPacked.z))
+			{
+				continue;
+			}
+#endif
 			int VirtualShadowMapId = LightData.VirtualShadowMapId;
 			if (VirtualShadowMapId != INDEX_NONE)
 			{
diff --git a/Engine/Shaders/Private/VirtualShadowMaps/ProjectionDirectional.ush b/Engine/Shaders/Private/VirtualShadowMaps/ProjectionDirectional.ush
--- a/Engine/Shaders/Private/VirtualShadowMaps/ProjectionDirectional.ush
+++ b/Engine/Shaders/Private/VirtualShadowMaps/ProjectionDirectional.ush
@@ -144,6 +144,13 @@ float3 GetRandomDirectionalLightRayDir(FLightShaderParameters Light, float2 E)
 	return normalize(RayDir);
 }
 
+// Normal and LightDirection should be normalized
+bool IsBackfaceToDirectionalLight(half3 Normal, half3 LightDirection, half LightSourceRadius)
+{
+	half SinAlpha = abs(LightSourceRadius);
+	return dot(Normal, LightDirection) < -SinAlpha;
+}
+
 FVirtualShadowMapSampleResult TraceDirectional(
 	int VirtualShadowMapId,
 	FLightShaderParameters Light,
@@ -155,22 +162,18 @@ FVirtualShadowMapSampleResult TraceDirectional(
 	uint SamplesPerRay,
 	float RayLengthScale,
 	const half Noise,
-	bool bCullBackface)
+	bool bCullBackface,
+	half3 WorldNormal)
 {
 	float3 WorldPosition = TranslatedWorldPosition - View.PreViewTranslation;
 	float3 ViewPosition = mul(float4(TranslatedWorldPosition, 1.0f), View.TranslatedWorldToView).xyz;
 	float DistanceFromViewOrigin = length(ViewPosition);
 
-	const half3 N = GetGBufferDataUint( PixelPos, true ).WorldNormal;
-	half SinAlphaSqr = saturate(Pow2(Light.SourceRadius));
-	half SinAlpha = sqrt(SinAlphaSqr);
-	bool bBackFace = dot( N, Light.Direction ) < -SinAlpha;
-
 	FVirtualShadowMapSampleResult Result = InitVirtualShadowMapSampleResult();
-	Result.bValid = true;	// TODO: false if all samples of all rays miss?
+	Result.bValid = true;	// TODO: false if all samples of all rays miss pages?
 	Result.ShadowFactor = 0.0f;
 
-	if (!bBackFace || !bCullBackface)
+	if (!bCullBackface || !IsBackfaceToDirectionalLight(WorldNormal, Light.Direction, Light.SourceRadius))
 	{
 		// Compute max ray length based on view depth
 		// This affects both how far we are willing to trace across the shadow map (for clipmaps this is related to view depth)
diff --git a/Engine/Shaders/Private/VirtualShadowMaps/ProjectionSpot.ush b/Engine/Shaders/Private/VirtualShadowMaps/ProjectionSpot.ush
--- a/Engine/Shaders/Private/VirtualShadowMaps/ProjectionSpot.ush
+++ b/Engine/Shaders/Private/VirtualShadowMaps/ProjectionSpot.ush
@@ -231,7 +231,40 @@ float4 SobolSampler(uint SampleIndex, inout uint Seed, uint MaxBits)
 	return (Result >> 8) * 5.96046447754e-08; // * 2^-24
 }
 
+// Normal should be normalized
+bool IsBackfaceToLocalLight(
+	float3 Position,
+	half3 Normal,
+	float3 LightPosition,
+	half LightSourceRadius,
+	out half3 OutConeAxis,
+	out half OutConeSin,
+	out half OutDistToLight)
+{
+	float3 ToLight = LightPosition - Position;
+	float DistSqr = dot(ToLight, ToLight);
+	float Falloff = rcp(DistSqr + 1);
+	float InvDist = rsqrt(DistSqr);
+
+	OutConeAxis = ToLight * InvDist;
+	OutConeSin = LightSourceRadius * InvDist;
+	OutDistToLight = DistSqr * InvDist;
+
+	half SinAlphaSqr = saturate(Pow2(LightSourceRadius) * Falloff);
+	half SinAlpha = sqrt(SinAlphaSqr);
+	bool bBackface = dot(Normal, OutConeAxis) < -SinAlpha;
 
+	return bBackface;
+}
+
+// Version without outputs
+bool IsBackfaceToLocalLight(float3 Position, half3 Normal, float3 LightPosition, half LightSourceRadius)
+{
+	half3 Dummy1;
+	half Dummy2;
+	half Dummy3;
+	return IsBackfaceToLocalLight(Position, Normal, LightPosition, LightSourceRadius, Dummy1, Dummy2, Dummy3);
+}
 
 FVirtualShadowMapSampleResult TraceLocalLight(
 	int VirtualShadowMapId,
@@ -245,36 +278,25 @@ FVirtualShadowMapSampleResult TraceLocalLight(
 	uint SamplesPerRay,
 	float CotMaxRayAngleFromLight,
 	const half Noise,
-	const bool bCullBackface)
+	const bool bCullBackface,
+	half3 WorldNormal)
 {
-	const half3 N = GetGBufferDataUint( PixelPos, true ).WorldNormal;
 	const float3 WorldPosition = TranslatedWorldPosition - View.PreViewTranslation;
 
-	float3 ToLight = Light.Position - WorldPosition;
-	float DistSqr = dot(ToLight, ToLight);
-	float Falloff = rcp(DistSqr + 1);
-
-	float InvDist = rsqrt( DistSqr );
-	float DistToLight = DistSqr * InvDist;
-
-	half3 ConeAxis = ToLight * InvDist;
-
-	half SinAlphaSqr = saturate(Pow2(Light.SourceRadius) * Falloff);
-	half SinAlpha = sqrt(SinAlphaSqr);
-	bool bBackFace = dot( N, ConeAxis ) < -SinAlpha;
+	half3 ConeAxis;
+	half ConeSin;
+	half DistToLight;	
+	bool bBackface = IsBackfaceToLocalLight(WorldPosition, WorldNormal, Light.Position, Light.SourceRadius, ConeAxis, ConeSin, DistToLight);
 
 	FVirtualShadowMapSampleResult Result = InitVirtualShadowMapSampleResult();
 	Result.bValid = true;
 	Result.ShadowFactor = 0.0f;
 
 	half RayMissFactor = 0;
-	if (!bBackFace || !bCullBackface)
+	if (!bBackface || !bCullBackface)
 	{
 		float StepOffset = Noise;
 
-		half ConeSin = Light.SourceRadius * InvDist;
-		half ConeCos = sqrt( 1 - Pow2( ConeSin ) );
-
 		uint TileID = MortonEncode( PixelPos );
 
 		uint i = 0;
@@ -398,32 +420,21 @@ bool GenerateRayLocalLight(
 	uint2 PixelPos,
 	float3 TranslatedWorldPosition,
 	float CotMaxRayAngleFromLight,
+	half3 WorldNormal,
 	uint RayIndex,
 	uint RayCount,
 	inout float3 OutRayStart,
 	inout float3 OutRayEnd)
 {
-	const half3 N = GetGBufferDataUint(PixelPos, true).WorldNormal;
 	const float3 WorldPosition = TranslatedWorldPosition - View.PreViewTranslation;
 
-	float3 ToLight = Light.Position - WorldPosition;
-	float DistSqr = dot(ToLight, ToLight);
-	float Falloff = rcp(DistSqr + 1);
-
-	float InvDist = rsqrt(DistSqr);
-	float DistToLight = DistSqr * InvDist;
-
-	half3 ConeAxis = ToLight * InvDist;
-
-	half SinAlphaSqr = saturate(Pow2(Light.SourceRadius) * Falloff);
-	half SinAlpha = sqrt(SinAlphaSqr);
-	bool bBackFace = dot(N, ConeAxis) < -SinAlpha;
-
-	if (!bBackFace)
+	half3 ConeAxis;
+	half ConeSin;
+	half DistToLight;	
+	bool bBackface = IsBackfaceToLocalLight(WorldPosition, WorldNormal, Light.Position, Light.SourceRadius, ConeAxis, ConeSin, DistToLight);
+	
+	if (!bBackface)
 	{
-		half ConeSin = Light.SourceRadius * InvDist;
-		half ConeCos = sqrt(1 - Pow2(ConeSin));
-
 		uint TileID = MortonEncode(PixelPos);
 		RandomSequence RandSequence;
 		// Combine pixel-level and sample-level bits into the sample index (visible structure will be hidden by owen scrambling of the index)
@@ -450,5 +461,5 @@ bool GenerateRayLocalLight(
 		OutRayStart = WorldPosition;
 		OutRayEnd = WorldPosition + Dir * ClampedLength;
 	}
-	return !bBackFace;
+	return !bBackface;
 }
diff --git a/Engine/Shaders/Private/VirtualShadowMaps/ProjectionStructs.ush b/Engine/Shaders/Private/VirtualShadowMaps/ProjectionStructs.ush
--- a/Engine/Shaders/Private/VirtualShadowMaps/ProjectionStructs.ush
+++ b/Engine/Shaders/Private/VirtualShadowMaps/ProjectionStructs.ush
@@ -2,10 +2,7 @@
 
 #pragma once
 
-// Must match C++ structure in VirtualShadowMapArray.h
-// While we currently load/unpack it manually in GetVirtualShadowMapProjectionData,
-// the intention is to turn this back into a simple structure load once the compiler
-// infrastructure is broadly available.
+// Rougly the same structure as in VirtualShadowMapArray.h
 struct FVirtualShadowMapProjectionShaderData
 {
 	float4x4 TranslatedWorldToShadowViewMatrix;
@@ -17,7 +14,7 @@ struct FVirtualShadowMapProjectionShaderData
 	uint LightType;		// Matches ELightComponentType via defines in SSDDefinitions.ush
 
 	float3 ClipmapWorldOrigin;
-	int VirtualShadowMapId;
+	float LightSourceRadius;
 
 	int2 ClipmapCornerOffset;
 	int ClipmapIndex;		// 0..ClipmapLevelCount-1
@@ -25,11 +22,15 @@ struct FVirtualShadowMapProjectionShaderData
 
 	int ClipmapLevelCount;	
 	float ClipmapResolutionLodBias;
+
+	// Derived data for convenience when passing the structure around
+	int VirtualShadowMapId;
 };
 
 FVirtualShadowMapProjectionShaderData GetVirtualShadowMapProjectionData(int VirtualShadowMapId)
 {
 	FVirtualShadowMapProjectionShaderData Result;
+	Result.VirtualShadowMapId = VirtualShadowMapId;
 
 	// Seems the FMatrix forces 16-byte alignment, so ensure padding matches the C++ structure size
 	const uint Stride = 16 * 20;
@@ -58,7 +59,7 @@ FVirtualShadowMapProjectionShaderData GetVirtualShadowMapProjectionData(int Virt
 	Result.LightType						=		 (VirtualShadowMap.ProjectionData.Load (Base + (16*16 + 4*3)));
 
 	Result.ClipmapWorldOrigin				= asfloat(VirtualShadowMap.ProjectionData.Load3(Base + (16*17 + 4*0)));
-	Result.VirtualShadowMapId				= asint	 (VirtualShadowMap.ProjectionData.Load (Base + (16*17 + 4*3)));
+	Result.LightSourceRadius				= asfloat(VirtualShadowMap.ProjectionData.Load (Base + (16*17 + 4*3)));
 
 	Result.ClipmapCornerOffset				= asint  (VirtualShadowMap.ProjectionData.Load2(Base + (16*18 + 4*0)));
 	Result.ClipmapIndex						= asint  (VirtualShadowMap.ProjectionData.Load (Base + (16*18 + 4*2)));
diff --git a/Engine/Shaders/Private/VirtualShadowMaps/VirtualShadowMapProjection.usf b/Engine/Shaders/Private/VirtualShadowMaps/VirtualShadowMapProjection.usf
--- a/Engine/Shaders/Private/VirtualShadowMaps/VirtualShadowMapProjection.usf
+++ b/Engine/Shaders/Private/VirtualShadowMaps/VirtualShadowMapProjection.usf
@@ -157,6 +157,7 @@ int DebugVirtualShadowMapId;
 #define INPUT_TYPE_GBUFFER 0
 #define INPUT_TYPE_HAIRSTRANDS 1
 uint InputType;
+uint bCullBackfacingPixels;
 
 FVirtualShadowMapSampleResult ProjectLight(
 	int VirtualShadowMapId,
@@ -177,8 +178,6 @@ FVirtualShadowMapSampleResult ProjectLight(
 #if DIRECTIONAL_LIGHT
 	uint LightType = LIGHT_TYPE_DIRECTIONAL;
 	float3 L = Light.Direction;
-	// TODO: Find a way to branch out backfacing surfaces again that doesn't make assumptions
-	// about the Gbuffer normal that break with hair
 	bool bInLightRegion = true;
 #else
 	uint LightType = Light.SpotAngles.x > -2.0f ? LIGHT_TYPE_SPOT : LIGHT_TYPE_POINT;
@@ -190,19 +189,19 @@ FVirtualShadowMapSampleResult ProjectLight(
 	bool bInLightRadius = InvDist >= Light.InvRadius;
 	bool bInLightCone = dot( L, Light.Direction ) >= Light.SpotAngles.x;
 	bool bInLightRegion = bInLightRadius && bInLightCone;
-	
 #endif // DIRECTIONAL_LIGHT
 
 	FVirtualShadowMapSampleResult Result = InitVirtualShadowMapSampleResult();
 
+	const bool bIsHair = InputType == INPUT_TYPE_HAIRSTRANDS;
+	const bool bBackfaceCull = !bIsHair && (bCullBackfacingPixels != 0);
+	const bool bValidPixel = bIsHair || GetGBufferDataUint(PixelPos).ShadingModelID != SHADINGMODELID_UNLIT;
+
 	BRANCH
-	if (bInLightRegion)
+	if (bInLightRegion && bValidPixel)
 	{
-		const bool bIsHair = InputType == INPUT_TYPE_HAIRSTRANDS;
-		const bool bBackfaceCull = false;		// TODO
-		const half3 EstimatedGeoWorldNormal = bIsHair ? L : GBufferNormal;
-		
-		TranslatedWorldPosition += EstimatedGeoWorldNormal * NormalBiasLength;
+		const half3 WorldNormal = bIsHair ? L : GBufferNormal;
+		TranslatedWorldPosition += WorldNormal * NormalBiasLength;
 
 		if (SMRTRayCount > 0)
 		{
@@ -218,7 +217,8 @@ FVirtualShadowMapSampleResult ProjectLight(
 				SMRTSamplesPerRay,
 				SMRTRayLengthScale,
 				Noise,
-				bBackfaceCull);
+				bBackfaceCull,
+				WorldNormal);
 		#else
 			Result = TraceLocalLight(
 				VirtualShadowMapId,
@@ -232,7 +232,8 @@ FVirtualShadowMapSampleResult ProjectLight(
 				SMRTSamplesPerRay,
 				SMRTCotMaxRayAngleFromLight,
 				Noise,
-				bBackfaceCull);
+				bBackfaceCull,
+				WorldNormal);
 		#endif // DIRECTIONAL_LIGHT
 		}
 		else
@@ -241,7 +242,7 @@ FVirtualShadowMapSampleResult ProjectLight(
 				VirtualShadowMapId,
 				TranslatedWorldPosition - View.PreViewTranslation,
 				ContactShadowLengthWorld,
-				EstimatedGeoWorldNormal);
+				WorldNormal);
 		}
 
 		// Do not run contact shadow when computing the hair shadow mask (i.e. shadow mask applied on hair, has the scene 
@@ -287,6 +288,7 @@ FVirtualShadowMapSampleResult ProjectLight(
 					Light,
 					PixelPos,
 					TranslatedWorldPosition,
+					WorldNormal,
 					SMRTCotMaxRayAngleFromLight,
 					RayIndex,
 					SMRTRayCount,
diff --git a/Engine/Source/Runtime/Renderer/Private/DeferredShadingRenderer.cpp b/Engine/Source/Runtime/Renderer/Private/DeferredShadingRenderer.cpp
--- a/Engine/Source/Runtime/Renderer/Private/DeferredShadingRenderer.cpp
+++ b/Engine/Source/Runtime/Renderer/Private/DeferredShadingRenderer.cpp
@@ -2656,7 +2656,7 @@ void FDeferredShadingSceneRenderer::Render(FRDGBuilder& GraphBuilder)
 			if (VirtualShadowMapArray.IsEnabled())
 			{
 				ensureMsgf(AreLightsInLightGrid(), TEXT("Virtual shadow map setup requires local lights to be injected into the light grid (this may be caused by 'r.LightCulling.Quality=0')."));
-				VirtualShadowMapArray.BuildPageAllocations(GraphBuilder, SceneTextures, Views, SortedLightSet, VisibleLightInfos, NaniteRasterResults, true);
+				VirtualShadowMapArray.BuildPageAllocations(GraphBuilder, SceneTextures, Views, SortedLightSet, VisibleLightInfos, NaniteRasterResults);
 			}
 
 			RenderShadowDepthMaps(GraphBuilder, InstanceCullingManager);
diff --git a/Engine/Source/Runtime/Renderer/Private/ShadowSetup.cpp b/Engine/Source/Runtime/Renderer/Private/ShadowSetup.cpp
--- a/Engine/Source/Runtime/Renderer/Private/ShadowSetup.cpp
+++ b/Engine/Source/Runtime/Renderer/Private/ShadowSetup.cpp
@@ -4833,7 +4833,8 @@ void FSceneRenderer::AddViewDependentWholeSceneShadowsForView(
 
 			if (bNeedsVirtualShadowMap)
 			{
-				FMatrix WorldToLight = FInverseRotationMatrix(LightSceneInfo.Proxy->GetDirection().GetSafeNormal().Rotation());
+				FVector LightDirection = LightSceneInfo.Proxy->GetDirection();
+				FMatrix WorldToLight = FInverseRotationMatrix(LightDirection.GetSafeNormal().Rotation());
 
 				TSharedPtr<FVirtualShadowMapClipmap> VirtualShadowMapClipmap = TSharedPtr<FVirtualShadowMapClipmap>(new FVirtualShadowMapClipmap(
 					VirtualShadowMapArray,
diff --git a/Engine/Source/Runtime/Renderer/Private/VirtualShadowMaps/VirtualShadowMapArray.cpp b/Engine/Source/Runtime/Renderer/Private/VirtualShadowMaps/VirtualShadowMapArray.cpp
--- a/Engine/Source/Runtime/Renderer/Private/VirtualShadowMaps/VirtualShadowMapArray.cpp
+++ b/Engine/Source/Runtime/Renderer/Private/VirtualShadowMaps/VirtualShadowMapArray.cpp
@@ -114,7 +114,14 @@ TAutoConsoleVariable<int32> CVarCoarsePagesIncludeNonNanite(
 static TAutoConsoleVariable<int32> CVarShowClipmapStats(
 	TEXT("r.Shadow.Virtual.ShowClipmapStats"),
 	-1,
-	TEXT("Set to the number of clipmap you want to show stats for (-1 == off)\n"),
+	TEXT("Set to the number of clipmap you want to show stats for (-1 == off)"),
+	ECVF_RenderThreadSafe
+);
+
+static TAutoConsoleVariable<int32> CVarCullBackfacingPixels(
+	TEXT("r.Shadow.Virtual.CullBackfacingPixels"),
+	1,
+	TEXT("When enabled does not generate shadow data for pixels that are backfacing to the light."),
 	ECVF_RenderThreadSafe
 );
 
@@ -237,6 +244,8 @@ void FVirtualShadowMapArray::Initialize(FRDGBuilder& GraphBuilder, FVirtualShado
 	CacheManager = InCacheManager;
 	check(!bEnabled || CacheManager);
 
+	bCullBackfacingPixels = CVarCullBackfacingPixels.GetValueOnRenderThread() != 0;
+
 	UniformParameters.NumShadowMaps = 0;
 	UniformParameters.NumDirectionalLights = 0;
 
@@ -404,6 +413,7 @@ class FGeneratePageFlagsFromPixelsCS : public FVirtualPageManagementShader
 		SHADER_PARAMETER(uint32, bPostBasePass)
 		SHADER_PARAMETER(float, ResolutionLodBiasLocal)
 		SHADER_PARAMETER(float, PageDilationBorderSize)
+		SHADER_PARAMETER(uint32, bCullBackfacingPixels)
 	END_SHADER_PARAMETER_STRUCT()
 };
 IMPLEMENT_GLOBAL_SHADER(FGeneratePageFlagsFromPixelsCS, "/Engine/Private/VirtualShadowMaps/PageManagement.usf", "GeneratePageFlagsFromPixels", SF_Compute);
@@ -711,8 +721,7 @@ void FVirtualShadowMapArray::BuildPageAllocations(
 	const TArray<FViewInfo>& Views,
 	const FSortedLightSetSceneInfo& SortedLightsInfo,
 	const TArray<FVisibleLightInfo, SceneRenderingAllocator>& VisibleLightInfos,
-	const TArray<Nanite::FRasterResults, TInlineAllocator<2>>& NaniteRasterResults,
-	bool bPostBasePass)
+	const TArray<Nanite::FRasterResults, TInlineAllocator<2>>& NaniteRasterResults)
 {
 	check(IsEnabled());
 	RDG_EVENT_SCOPE(GraphBuilder, "FVirtualShadowMapArray::BuildPageAllocation");
@@ -799,8 +808,8 @@ void FVirtualShadowMapArray::BuildPageAllocations(
 						Data.TranslatedWorldToShadowUVMatrix		= CalcTranslatedWorldToShadowUVMatrix( ViewMatrices.GetTranslatedViewMatrix(), ViewMatrices.GetProjectionMatrix() );
 						Data.TranslatedWorldToShadowUVNormalMatrix	= CalcTranslatedWorldToShadowUVNormalMatrix( ViewMatrices.GetTranslatedViewMatrix(), ViewMatrices.GetProjectionMatrix() );
 						Data.ShadowPreViewTranslation				= FVector(ProjectedShadowInfo->PreShadowTranslation);
-						Data.VirtualShadowMapId						= ID;
 						Data.LightType								= ProjectedShadowInfo->GetLightSceneInfo().Proxy->GetLightType();
+						Data.LightSourceRadius						= ProjectedShadowInfo->GetLightSceneInfo().Proxy->GetSourceRadius();
 					}
 
 				#if !UE_BUILD_SHIPPING
@@ -863,7 +872,6 @@ void FVirtualShadowMapArray::BuildPageAllocations(
 		for (int32 ViewIndex = 0; ViewIndex < Views.Num(); ++ViewIndex)
 		{
 			const FViewInfo &View = Views[ViewIndex];
-			FRDGTextureRef NaniteVisBuffer64 = ViewIndex < NaniteRasterResults.Num() ? NaniteRasterResults[ViewIndex].VisBuffer64 : nullptr;
 
 			// This view contained no local lights (that were stored in the light grid), and no directional lights, so nothing to do.
 			if (View.ForwardLightingResources->LocalLightVisibleLightInfosIndex.Num() + DirectionalLightIds.Num() == 0)
@@ -874,7 +882,6 @@ void FVirtualShadowMapArray::BuildPageAllocations(
 			FRDGBufferRef DirectionalLightIdsRDG = CreateStructuredBuffer(GraphBuilder, TEXT("Shadow.Virtual.DirectionalLightIds"), DirectionalLightIds);
 
 			const FRDGSystemTextures& SystemTextures = FRDGSystemTextures::Get(GraphBuilder);
-			FRDGTextureRef VisBuffer64 = NaniteVisBuffer64 != nullptr ? NaniteVisBuffer64 : SystemTextures.Black;
 
 			FRDGBufferRef ScreenSpaceGridBoundsRDG = nullptr;
 			
@@ -887,8 +894,7 @@ void FVirtualShadowMapArray::BuildPageAllocations(
 				{
 					auto GeneratePageFlags = [&](bool bHairPass)
 					{
-						const bool bUseNaniteDepth = NaniteVisBuffer64 != nullptr && !bPostBasePass;
-						const uint32 InputType = bHairPass ? 2u : (bUseNaniteDepth ? 1u : 0u); // HairStrands, Nanite, or GBuffer
+						const uint32 InputType = bHairPass ? 1U : 0U; // HairStrands or GBuffer
 
 						FGeneratePageFlagsFromPixelsCS::FPermutationDomain PermutationVector;
 						PermutationVector.Set<FGeneratePageFlagsFromPixelsCS::FInputType>(InputType);
@@ -896,9 +902,7 @@ void FVirtualShadowMapArray::BuildPageAllocations(
 						PassParameters->VirtualShadowMap = GetUniformBuffer(GraphBuilder);
 
 						PassParameters->SceneTexturesStruct = SceneTextures.UniformBuffer;
-						PassParameters->bPostBasePass = bPostBasePass;
 
-						PassParameters->VisBuffer64 = VisBuffer64;
 						PassParameters->HairStrands = HairStrands::BindHairStrandsViewUniformParameters(View);
 						PassParameters->View = View.ViewUniformBuffer;
 						PassParameters->OutPageRequestFlags = PageRequestFlagsUAV;
@@ -906,6 +910,7 @@ void FVirtualShadowMapArray::BuildPageAllocations(
 						PassParameters->DirectionalLightIds = GraphBuilder.CreateSRV(DirectionalLightIdsRDG);
 						PassParameters->ResolutionLodBiasLocal = ResolutionLodBiasLocal;
 						PassParameters->PageDilationBorderSize = PageDilationBorderSize;
+						PassParameters->bCullBackfacingPixels = ShouldCullBackfacingPixels() ? 1 : 0;
 
 						auto ComputeShader = View.ShaderMap->GetShader<FGeneratePageFlagsFromPixelsCS>(PermutationVector);
 
@@ -927,7 +932,7 @@ void FVirtualShadowMapArray::BuildPageAllocations(
 						{
 							FComputeShaderUtils::AddPass(
 								GraphBuilder,
-								RDG_EVENT_NAME("GeneratePageFlagsFromPixels(%s)", bHairPass ? TEXT("HairStrands") : (bUseNaniteDepth ? TEXT("Nanite") : TEXT("GBuffer"))),
+								RDG_EVENT_NAME("GeneratePageFlagsFromPixels(%s)", bHairPass ? TEXT("HairStrands") : TEXT("GBuffer")),
 								ComputeShader,
 								PassParameters,
 								FIntVector(GridSize.X, GridSize.Y, 1));
diff --git a/Engine/Source/Runtime/Renderer/Private/VirtualShadowMaps/VirtualShadowMapArray.h b/Engine/Source/Runtime/Renderer/Private/VirtualShadowMaps/VirtualShadowMapArray.h
--- a/Engine/Source/Runtime/Renderer/Private/VirtualShadowMaps/VirtualShadowMapArray.h
+++ b/Engine/Source/Runtime/Renderer/Private/VirtualShadowMaps/VirtualShadowMapArray.h
@@ -88,8 +88,10 @@ struct FVirtualShadowMapProjectionShaderData
 	FVector3f ShadowPreViewTranslation;
 	uint32 LightType = ELightComponentType::LightType_Directional;
 	
+	// TODO: There are more local lights than directional
+	// We should move the directional-specific stuff out to its own structure.
 	FVector3f ClipmapWorldOrigin;
-	int32 VirtualShadowMapId = INDEX_NONE;
+	float LightSourceRadius;				// This should live in shared light structure...
 	
 	FIntPoint ClipmapCornerOffset;
 	int32 ClipmapIndex = 0;					// 0 .. ClipmapLevelCount-1
@@ -181,8 +183,7 @@ public:
 		const TArray<FViewInfo> &Views, 
 		const FSortedLightSetSceneInfo& SortedLights, 
 		const TArray<FVisibleLightInfo, SceneRenderingAllocator> &VisibleLightInfos, 
-		const TArray<Nanite::FRasterResults, TInlineAllocator<2>> &NaniteRasterResults, 
-		bool bPostBasePass);
+		const TArray<Nanite::FRasterResults, TInlineAllocator<2>> &NaniteRasterResults);
 
 	bool IsAllocated() const
 	{
@@ -212,9 +213,8 @@ public:
 
 	void GetPageTableParameters(FRDGBuilder& GraphBuilder, FVirtualShadowMapPageTableParameters& OutParameters);
 
-	bool bInitialized = false;
-	// Are virtual shadow maps enabled? We store this at the start of the frame to centralize the logic.
-	bool bEnabled = false;
+	bool ShouldCullBackfacingPixels() const { return bCullBackfacingPixels; }
+
 	// We keep a reference to the cache manager that was used to initialize this frame as it owns some of the buffers
 	FVirtualShadowMapArrayCacheManager* CacheManager = nullptr;
 
@@ -236,15 +236,6 @@ public:
 	// HPageFlags is a hierarchy over the PageFlags for quick query
 	FRDGBufferRef HPageFlagsRDG = nullptr;
 
-	static constexpr uint32 NumStats = 5;
-	// 0 - allocated pages
-	// 1 - re-usable pages
-	// 2 - Touched by dynamic
-	// 3 - NumSms
-	// 4 - RandRobin invalidated
-
-	FRDGBufferRef StatsBufferRDG = nullptr;
-
 	// Allocation info for each page.
 	FRDGBufferRef CachedPageInfosRDG = nullptr;
 	FRDGBufferRef PhysicalPageMetaDataRDG = nullptr;
@@ -263,11 +254,26 @@ public:
 	FRDGTextureRef HZBPhysical = nullptr;
 	TMap<int32, FVirtualShadowMapHZBMetadata> HZBMetadata;
 
-	// Convert also
-	TRefCountPtr<IPooledRenderTarget> DebugVisualizationOutput;
+	static constexpr uint32 NumStats = 5;
+	// 0 - allocated pages
+	// 1 - re-usable pages
+	// 2 - Touched by dynamic
+	// 3 - NumSms
+	// 4 - RandRobin invalidated
+	FRDGBufferRef StatsBufferRDG = nullptr;
 
+	TRefCountPtr<IPooledRenderTarget> DebugVisualizationOutput;
 	int DebugOutputType = 0;	// 0 = Disabled
-	// Base ID of the light that the user has selected for debug output (if present)
+								// Base ID of the light that the user has selected for debug output (if present)
 	int DebugVirtualShadowMapId = INDEX_NONE;
 	FRDGTextureRef DebugVisualizationProjectionOutput = nullptr;
+
+private:
+	bool bInitialized = false;
+
+	// Are virtual shadow maps enabled? We store this at the start of the frame to centralize the logic.
+	bool bEnabled = false;
+
+	// Is backface culling of pixels enabled? We store this here to keep it consistent between projection and generation
+	bool bCullBackfacingPixels = false;
 };
diff --git a/Engine/Source/Runtime/Renderer/Private/VirtualShadowMaps/VirtualShadowMapClipmap.cpp b/Engine/Source/Runtime/Renderer/Private/VirtualShadowMaps/VirtualShadowMapClipmap.cpp
--- a/Engine/Source/Runtime/Renderer/Private/VirtualShadowMaps/VirtualShadowMapClipmap.cpp
+++ b/Engine/Source/Runtime/Renderer/Private/VirtualShadowMaps/VirtualShadowMapClipmap.cpp
@@ -232,7 +232,6 @@ FVirtualShadowMapProjectionShaderData FVirtualShadowMapClipmap::GetProjectionSha
 	Data.TranslatedWorldToShadowUVMatrix = CalcTranslatedWorldToShadowUVMatrix(WorldToLightViewRotationMatrix, Level.ViewToClip);
 	Data.TranslatedWorldToShadowUVNormalMatrix = CalcTranslatedWorldToShadowUVNormalMatrix(WorldToLightViewRotationMatrix, Level.ViewToClip);
 	Data.ShadowPreViewTranslation = -Level.WorldCenter;
-	Data.VirtualShadowMapId = Level.VirtualShadowMap->ID;
 	Data.LightType = ELightComponentType::LightType_Directional;
 	Data.ClipmapWorldOrigin = WorldOrigin;
 	Data.ClipmapIndex = ClipmapIndex;
@@ -240,6 +239,7 @@ FVirtualShadowMapProjectionShaderData FVirtualShadowMapClipmap::GetProjectionSha
 	Data.ClipmapLevelCount = LevelData.Num();
 	Data.ClipmapResolutionLodBias = ResolutionLodBias;
 	Data.ClipmapCornerOffset = Level.CornerOffset;
+	Data.LightSourceRadius = GetLightSceneInfo().Proxy->GetSourceRadius();
 
 	return Data;
 }
diff --git a/Engine/Source/Runtime/Renderer/Private/VirtualShadowMaps/VirtualShadowMapProjection.cpp b/Engine/Source/Runtime/Renderer/Private/VirtualShadowMaps/VirtualShadowMapProjection.cpp
--- a/Engine/Source/Runtime/Renderer/Private/VirtualShadowMaps/VirtualShadowMapProjection.cpp
+++ b/Engine/Source/Runtime/Renderer/Private/VirtualShadowMaps/VirtualShadowMapProjection.cpp
@@ -193,6 +193,7 @@ class FVirtualShadowMapProjectionCS : public FGlobalShader
 		SHADER_PARAMETER(float, SMRTRayLengthScale)
 		SHADER_PARAMETER(float, SMRTCotMaxRayAngleFromLight)
 		SHADER_PARAMETER(uint32, InputType)
+		SHADER_PARAMETER(uint32, bCullBackfacingPixels)
 		// One pass projection parameters
 		SHADER_PARAMETER_STRUCT_REF(FForwardLightData, ForwardLightData)
 		SHADER_PARAMETER_RDG_TEXTURE_UAV(RWTexture2D< uint >, RWShadowMaskBits)
@@ -269,6 +270,7 @@ static void RenderVirtualShadowMapProjectionCommon(
 	PassParameters->ContactShadowLength = CVarContactShadowLength.GetValueOnRenderThread();
 	PassParameters->NormalBias = GetNormalBiasForShader();
 	PassParameters->InputType = uint32(InputType);
+	PassParameters->bCullBackfacingPixels = VirtualShadowMapArray.ShouldCullBackfacingPixels() ? 1 : 0;
 	if (bHasHairStrandsData)
 	{
 		PassParameters->HairStrands = HairStrands::BindHairStrandsViewUniformParameters(View);
