Controls history recovery under movement also based on highlight difference between input and history to improve stability in motion

diff --git a/Engine/Shaders/Private/TemporalSuperResolution/TSRUpdateHistory.usf b/Engine/Shaders/Private/TemporalSuperResolution/TSRUpdateHistory.usf
--- a/Engine/Shaders/Private/TemporalSuperResolution/TSRUpdateHistory.usf
+++ b/Engine/Shaders/Private/TemporalSuperResolution/TSRUpdateHistory.usf
@@ -755,6 +755,8 @@ void MainCS(
 	taa_half3x2 FinalHighFrequencyColor;
 	taa_half2 FinalHistoryValidity;
 	{
+		taa_half2 InvPrevHistoryValidity = SafeRcp(PrevHistoryValidity);
+
 		// Compute how much the history needs to be rejected, based on the completeness.
 		taa_half2 PrevHistoryRejectionWeight = LowFrequencyRejection;
 			
@@ -778,9 +780,37 @@ void MainCS(
 
 		// Determine whether the the prediction based rejection was confident enough.
 		taa_half2 RejectionConfidentEnough = taa_half(1); // saturate(RejectionValidity * MAX_SAMPLE_COUNT - 3.0);
+		
+		// Clamp the validity due to motion to maintain better sharpness in history reprojection under motion.
+		taa_half2 ClampedPrevHistoryValidity;
+		#if 1
+		{
+			taa_half2 MaxValidity = taa_half(1.0) - (taa_half(1.0) - PredictionOnlyValidity) * saturate(OutputPixelVelocity * taa_half(InvWeightClampingPixelSpeed));
+		
+			// Maintain full stability on interference areas.
+			#if DIM_MOIRE_HEURISTIC
+				MaxValidity = max(MaxValidity, 1.0 - MoireWeight);
+			#endif
+
+			// Clamp up the max validity to favorise stability under motion on high contrast edges.
+			#if 1
+			{
+				taa_half2 PrevHistoryLuma = Luma4(PrevHighFrequencyColor) * InvPrevHistoryValidity;
+				taa_half2 FilteredLuma = Luma4(FilteredInputColor);
 
-		// Clamp the validity due to motion to maintain better sharpness.
-		taa_half2 ClampedPrevHistoryValidity = min(PrevHistoryValidity, taa_half(1.0) - (taa_half(1.0) - PredictionOnlyValidity) * saturate(OutputPixelVelocity * taa_half(InvWeightClampingPixelSpeed)));
+				taa_half2 MinValidityForStability = abs(FilteredLuma - PrevHistoryLuma) / max(FilteredLuma, PrevHistoryLuma);
+				
+				MaxValidity = max(MaxValidity, MinValidityForStability);
+			}
+			#endif
+			
+			ClampedPrevHistoryValidity = min(PrevHistoryValidity, MaxValidity);
+		}
+		#else
+		{
+			ClampedPrevHistoryValidity = PrevHistoryValidity;
+		}
+		#endif
 
 		// Compute the newly rejected validity
 		taa_half2 RejectedValidity = (
@@ -808,8 +838,6 @@ void MainCS(
 			}
 		#endif
 		
-		taa_half2 InvPrevHistoryValidity = SafeRcp(PrevHistoryValidity);
-
 		taa_half2 PrevMomentWeight = max(OutputValidity - DesiredCurrentContribution, taa_half(0.0));
 		taa_half2 CurrentMomentWeight = min(DesiredCurrentContribution, OutputValidity);
 		
