Hair shading model mobile support
Changed Hair Materials to directly use the Hair Shading model, instead of doing so with a FeatureLevelSwitch expression

diff --git a/Engine/Shaders/Private/MobileBasePassPixelShader.usf b/Engine/Shaders/Private/MobileBasePassPixelShader.usf
--- a/Engine/Shaders/Private/MobileBasePassPixelShader.usf
+++ b/Engine/Shaders/Private/MobileBasePassPixelShader.usf
@@ -736,19 +736,6 @@ void Main(
 	}
 #endif
 	InitShadingModelContext(ShadingModelContext, GBuffer, MaterialParameters.CameraVector);
-	float3 DiffuseDir = MaterialParameters.WorldNormal;
-
-#if MATERIAL_SHADINGMODEL_HAIR && MOBILE_SHADINGMODEL_SUPPORT
-	if (GBuffer.ShadingModelID == SHADINGMODELID_HAIR)
-	{
-		FHairTransmittanceData TransmittanceData = InitHairTransmittanceData(true);
-		float3 N = MaterialParameters.WorldNormal;
-		float3 V = MaterialParameters.CameraVector;
-		float3 L = normalize(V - N * dot(V, N));
-		DiffuseDir = L;
-		ShadingModelContext.DiffuseColor = 2 * PI * HairShading(GBuffer, L, V, N, 1, TransmittanceData, 0, 0.2, uint2(0, 0));
-	}
-#endif
 
 	VTPageTableResult LightmapVTPageTableResult = (VTPageTableResult)0.0f;
 #if LIGHTMAP_VT_ENABLED
@@ -780,7 +767,7 @@ void Main(
 #endif
 
 	// Apply MaterialAO since we don't have the DiffuseIndirectComposite pass on mobile deferred.
-	ComputeIndirect(LightmapVTPageTableResult, Interpolants, DiffuseDir, DiffuseColorForIndirect, MaterialAO, IndirectIrradiance, MobileLightAccumulator);
+	ComputeIndirect(LightmapVTPageTableResult, Interpolants, ShadingModelContext.DiffuseDir, DiffuseColorForIndirect, MaterialAO, IndirectIrradiance, MobileLightAccumulator);
 
 	half Shadow = GetPrimaryPrecomputedShadowMask(LightmapVTPageTableResult, Interpolants, MaterialParameters).r;
 
diff --git a/Engine/Shaders/Private/MobileShadingModels.ush b/Engine/Shaders/Private/MobileShadingModels.ush
--- a/Engine/Shaders/Private/MobileShadingModels.ush
+++ b/Engine/Shaders/Private/MobileShadingModels.ush
@@ -78,6 +78,7 @@ struct FMobileShadingModelContext
 	half3 EnvBrdf;
 
 	half NoV;
+	float3 DiffuseDir;
 
 #if MATERIAL_SHADINGMODEL_SINGLELAYERWATER
 	half BaseMaterialCoverageOverWater;
@@ -153,6 +154,7 @@ void InitShadingModelContext(inout FMobileShadingModelContext ShadingModelContex
 
 	// This is to prevent Vis to get inf when both NoL and NoV are 0.
 	ShadingModelContext.NoV = saturate(abs(NoV) + 1e-5);
+	ShadingModelContext.DiffuseDir = GBuffer.WorldNormal;
 
 #if NONMETAL
 	GBuffer.DiffuseColor = GBuffer.BaseColor;
@@ -214,6 +216,15 @@ void InitShadingModelContext(inout FMobileShadingModelContext ShadingModelContex
 		ShadingModelContext.Lobe0Roughness = lerp(1.0f, ShadingModelContext.Lobe0Roughness, RoughnessLerpFactor);
 		ShadingModelContext.Lobe1Roughness = lerp(1.0f, ShadingModelContext.Lobe1Roughness, RoughnessLerpFactor);
 	}
+	else if (GBuffer.ShadingModelID == SHADINGMODELID_HAIR)
+	{
+		FHairTransmittanceData TransmittanceData = InitHairTransmittanceData(true);
+		float3 N = GBuffer.WorldNormal;
+		float3 V = CameraVector;
+		float3 L = normalize(V - N * dot(V, N));
+		ShadingModelContext.DiffuseDir = L;
+		ShadingModelContext.DiffuseColor = 2 * PI * HairShading(GBuffer, L, V, N, 1, TransmittanceData, 0, 0.2, uint2(0, 0));
+	}
 #endif
 
 #if FULLY_ROUGH
@@ -344,12 +355,14 @@ FMobileDirectLighting MobileClearCoatBxDF(FMobileShadingModelContext ShadingMode
 
 FMobileDirectLighting MobileHairBxDF(FMobileShadingModelContext ShadingModelContext, FGBufferData GBuffer, half NoL, half NoH, half VoH, half3 V, half3 L, half Shadow)
 {
-	FMobileDirectLighting Lighting;
+	FHairTransmittanceData HairTransmittance = InitHairTransmittanceData();
+	const float3 BsdfValue = HairShading(GBuffer, L, V, GBuffer.WorldNormal, Shadow, HairTransmittance, 1, 0, uint2(0, 0));
 
-	// TODO transmittance
-	Lighting.Specular = half3(0, 0, 0);
-	Lighting.Diffuse = half3(0, 0, 0);
+	half3 Transmission = BsdfValue * PI; // multiply PI since the LightColor is divided by PI in C++;
 
+	FMobileDirectLighting Lighting;
+	Lighting.Diffuse = Transmission * Shadow;
+	Lighting.Specular = 0;
 	return Lighting;
 }
 
diff --git a/Engine/Source/Runtime/Engine/Private/Materials/HLSLMaterialTranslator.cpp b/Engine/Source/Runtime/Engine/Private/Materials/HLSLMaterialTranslator.cpp
--- a/Engine/Source/Runtime/Engine/Private/Materials/HLSLMaterialTranslator.cpp
+++ b/Engine/Source/Runtime/Engine/Private/Materials/HLSLMaterialTranslator.cpp
@@ -1376,7 +1376,7 @@ void FHLSLMaterialTranslator::ValidateShadingModelsForFeatureLevel(const FMateri
 {
 	if (FeatureLevel <= ERHIFeatureLevel::ES3_1)
 	{
-		const TArray<EMaterialShadingModel>& InvalidShadingModels = { MSM_Hair, MSM_Eye };
+		const TArray<EMaterialShadingModel>& InvalidShadingModels = { MSM_Eye };
 		for (EMaterialShadingModel InvalidShadingModel : InvalidShadingModels)
 		{
 			if (ShadingModels.HasShadingModel(InvalidShadingModel))
