Updated GPUSceneWriter.ush to match new primitive ID format

diff --git a/Engine/Shaders/Private/GPUScene/GPUSceneWriter.ush b/Engine/Shaders/Private/GPUScene/GPUSceneWriter.ush
--- a/Engine/Shaders/Private/GPUScene/GPUSceneWriter.ush
+++ b/Engine/Shaders/Private/GPUScene/GPUSceneWriter.ush
@@ -31,7 +31,9 @@ void StoreInstanceSceneDataElement(uint InstanceId, uint ElementIndex, float4 Da
 void WriteInstancePrimitiveId(uint InstanceId, uint NewPrimitiveId)
 {
 	float4 DataElement = LoadInstanceSceneDataElement(0 * GPUSceneInstanceSceneDataSOAStride + InstanceId);
-	DataElement.y = asfloat(NewPrimitiveId);
+	const uint InstanceFlags = BitFieldExtractU32(asuint(DataElement.x),  12, 20); // Max of 12 instance flags
+	const uint Packed0 = (InstanceFlags << 20u) | NewPrimitiveId;
+	DataElement.x = asfloat(Packed0);
 	checkSlow(InstanceId < GPUSceneNumAllocatedInstances);
 	StoreInstanceSceneDataElement(InstanceId, 0, DataElement);
 }
diff --git a/Engine/Shaders/Private/SceneData.ush b/Engine/Shaders/Private/SceneData.ush
--- a/Engine/Shaders/Private/SceneData.ush
+++ b/Engine/Shaders/Private/SceneData.ush
@@ -412,6 +412,7 @@ FInstanceSceneData GetInstanceSceneData(uint InstanceId, uint SOAStride)
 	//
 
 	// Only process valid instances
+	// Note: keep primitive ID packing in sync with GPUSceneWriter.ush
 	const uint FlagsAndPrimitiveId	= asuint(LoadInstanceSceneDataElement(0 * SOAStride + InstanceId).x);
 	InstanceData.PrimitiveId		= BitFieldExtractU32(FlagsAndPrimitiveId,  20, 0);  // Max of 1,048,576 primitives
 	InstanceData.Flags				= BitFieldExtractU32(FlagsAndPrimitiveId,  12, 20); // Max of 12 instance flags
