Disable VSM backface cull on subsurface pixels

diff --git a/Engine/Shaders/Private/VirtualShadowMaps/PageManagement.usf b/Engine/Shaders/Private/VirtualShadowMaps/PageManagement.usf
--- a/Engine/Shaders/Private/VirtualShadowMaps/PageManagement.usf
+++ b/Engine/Shaders/Private/VirtualShadowMaps/PageManagement.usf
@@ -196,6 +196,7 @@ void GeneratePageFlagsFromPixels(
 	half3  WorldNormal   = half3(0, 0, 0);
 	float  SceneDepth    = 0;
 	bool   bIsValid		 = true;
+	bool   bBackfaceCull = bCullBackfacingPixels != 0;
 	#if PERMUTATION_INPUT_TYPE == INPUT_TYPE_HAIRSTRANDS
 	{
 		const float HairDeviceZ = HairStrands.HairOnlyDepthTexture.Load(uint3(PixelPos, 0)).x;
@@ -203,6 +204,7 @@ void GeneratePageFlagsFromPixels(
 		WorldPosition = HairPixel.WorldPosition;
 		SceneDepth = HairPixel.SceneDepth;
 		bIsValid = HairDeviceZ > 0;
+		bBackfaceCull = false;
 	}
 	#else // PERMUTATION_INPUT_TYPE == INPUT_TYPE_GBUFFER
 	{
@@ -214,6 +216,7 @@ void GeneratePageFlagsFromPixels(
 		WorldNormal = GBuffer.WorldNormal;
 		// Excluding unlit to avoid including processing sky dome
 		bIsValid = GBuffer.ShadingModelID != SHADINGMODELID_UNLIT;
+		bBackfaceCull = bBackfaceCull && !IsSubsurfaceModel(GBuffer.ShadingModelID);
 	}
 	#endif
 	
@@ -239,7 +242,7 @@ void GeneratePageFlagsFromPixels(
 #if PERMUTATION_INPUT_TYPE != INPUT_TYPE_HAIRSTRANDS
 		// Backface test if requested
 		float3 LightDirection = -BaseProjectionData.TranslatedWorldToShadowViewMatrix._13_23_33;
-		if (bCullBackfacingPixels != 0 && IsBackfaceToDirectionalLight(normalize(WorldNormal), normalize(LightDirection), 0/*BaseProjectionData.LightSourceRadius*/))
+		if (bBackfaceCull && IsBackfaceToDirectionalLight(normalize(WorldNormal), normalize(LightDirection), BaseProjectionData.LightSourceRadius))
 		{
 			continue;
 		}
@@ -281,7 +284,7 @@ void GeneratePageFlagsFromPixels(
 
 #if PERMUTATION_INPUT_TYPE != INPUT_TYPE_HAIRSTRANDS
 			// Backface test if requested
-			if (bCullBackfacingPixels != 0 && IsBackfaceToLocalLight(WorldPosition, WorldNormal, LightData.LightPositionAndInvRadius.xyz, LightData.SpotAnglesAndSourceRadiusPacked.z))
+			if (bBackfaceCull && IsBackfaceToLocalLight(WorldPosition, WorldNormal, LightData.LightPositionAndInvRadius.xyz, LightData.SpotAnglesAndSourceRadiusPacked.z))
 			{
 				continue;
 			}
diff --git a/Engine/Shaders/Private/VirtualShadowMaps/VirtualShadowMapProjection.usf b/Engine/Shaders/Private/VirtualShadowMaps/VirtualShadowMapProjection.usf
--- a/Engine/Shaders/Private/VirtualShadowMaps/VirtualShadowMapProjection.usf
+++ b/Engine/Shaders/Private/VirtualShadowMaps/VirtualShadowMapProjection.usf
@@ -193,9 +193,11 @@ FVirtualShadowMapSampleResult ProjectLight(
 
 	FVirtualShadowMapSampleResult Result = InitVirtualShadowMapSampleResult();
 
+	const uint ShadingModelID = GetGBufferDataUint(PixelPos).ShadingModelID;
+
 	const bool bIsHair = InputType == INPUT_TYPE_HAIRSTRANDS;
-	const bool bBackfaceCull = !bIsHair && (bCullBackfacingPixels != 0);
-	const bool bValidPixel = bIsHair || GetGBufferDataUint(PixelPos).ShadingModelID != SHADINGMODELID_UNLIT;
+	const bool bBackfaceCull = (bCullBackfacingPixels != 0) && !bIsHair && !IsSubsurfaceModel(ShadingModelID);
+	const bool bValidPixel = bIsHair || ShadingModelID != SHADINGMODELID_UNLIT;
 
 	BRANCH
 	if (bInLightRegion && bValidPixel)
