Fixed shader compilation error related to mobile auto-instancing

diff --git a/Engine/Shaders/Private/LocalVertexFactory.ush b/Engine/Shaders/Private/LocalVertexFactory.ush
--- a/Engine/Shaders/Private/LocalVertexFactory.ush
+++ b/Engine/Shaders/Private/LocalVertexFactory.ush
@@ -819,7 +819,7 @@ FVertexFactoryIntermediates GetVertexFactoryIntermediates(FVertexFactoryInput In
 	{
 	#if (VF_USE_PRIMITIVE_SCENE_DATA == 2)
 		#if USES_PER_INSTANCE_RANDOM
-		Intermediates.PerInstanceRandom = GetInstanceData(Intermediates).PerInstanceRandom;
+		Intermediates.PerInstanceRandom = GetInstanceData(Intermediates).RandomID;
 		#endif
 		
 		float R = 0; // FIX ME: use SRV for this data?
