Fix missing light attenuation decoding on forward/mobile codepaths.

diff --git a/Engine/Shaders/Private/Common.ush b/Engine/Shaders/Private/Common.ush
--- a/Engine/Shaders/Private/Common.ush
+++ b/Engine/Shaders/Private/Common.ush
@@ -999,6 +999,18 @@ MaterialFloat4 EncodeLightAttenuation(MaterialFloat4 InColor)
 	return sqrt(InColor);
 }
 
+/** Converts value stored in the light attenuation buffer into a linear light attenuation value. */
+MaterialFloat DecodeLightAttenuation(MaterialFloat InColor)
+{
+	return Square(InColor);
+}
+
+/** Converts value stored in the light attenuation buffer into a linear light attenuation value. */
+MaterialFloat4 DecodeLightAttenuation(MaterialFloat4 InColor)
+{
+	return Square(InColor);
+}
+
 // Like RGBM but this can be interpolated.
 MaterialFloat4 RGBTEncode(MaterialFloat3 Color)
 {
@@ -1102,7 +1114,7 @@ float2 CalcScreenUVFromOffsetFraction(float4 ScreenPosition, float2 OffsetFracti
 
 float4 GetPerPixelLightAttenuation(float2 UV)
 {
-	return Square(Texture2DSampleLevel(LightAttenuationTexture, LightAttenuationTextureSampler, UV, 0));
+	return DecodeLightAttenuation(Texture2DSampleLevel(LightAttenuationTexture, LightAttenuationTextureSampler, UV, 0));
 }
 
 // also see ConvertToDeviceZ()
diff --git a/Engine/Shaders/Private/ForwardLightingCommon.ush b/Engine/Shaders/Private/ForwardLightingCommon.ush
--- a/Engine/Shaders/Private/ForwardLightingCommon.ush
+++ b/Engine/Shaders/Private/ForwardLightingCommon.ush
@@ -95,7 +95,7 @@ float4 GetForwardDynamicShadowFactors(float2 ScreenUV)
 	{
 		Value = OpaqueBasePass.ForwardScreenSpaceShadowMaskTexture.Load(int3(IntScreenUV, 0));
 	}
-	return Square(Value);
+	return DecodeLightAttenuation(Value);
 }
 
 float GetIndirectOcclusion(float2 ScreenUV, FGBufferData GBufferData)
diff --git a/Engine/Shaders/Private/MobileBasePassPixelShader.usf b/Engine/Shaders/Private/MobileBasePassPixelShader.usf
--- a/Engine/Shaders/Private/MobileBasePassPixelShader.usf
+++ b/Engine/Shaders/Private/MobileBasePassPixelShader.usf
@@ -820,6 +820,7 @@ void Main(
 	if (ShadowPositionZ == 0)
 	{
 		Shadow = Texture2DSample(MobileBasePass.ScreenSpaceShadowMaskTexture, MobileBasePass.ScreenSpaceShadowMaskSampler, SvPositionToBufferUV(SvPosition)).x;
+		Shadow = DecodeLightAttenuation(Shadow);
 	}
 #endif
 
diff --git a/Engine/Shaders/Private/MobileDeferredShading.usf b/Engine/Shaders/Private/MobileDeferredShading.usf
--- a/Engine/Shaders/Private/MobileDeferredShading.usf
+++ b/Engine/Shaders/Private/MobileDeferredShading.usf
@@ -282,6 +282,7 @@ void MobileDirectLightPS(
 
 #if ENABLE_DISTANCE_FIELD
 	Shadow = Texture2DSample(MobileBasePass.ScreenSpaceShadowMaskTexture, MobileBasePass.ScreenSpaceShadowMaskSampler, SvPositionToBufferUV(SvPosition)).x;
+	Shadow = DecodeLightAttenuation(Shadow);
 #endif
 
 	FMobileShadingModelContext ShadingModelContext = (FMobileShadingModelContext)0;
