Make "r.vt.FeedbackFactor" something that can be changed at runtime, instead of needing a restart/shader recompile.
This should be useful if we want to vary the feedback trade offs for different modes.
Did some comparisons of shader ISA and ALU cost is unchanged.

diff --git a/Engine/Shaders/Private/VirtualTextureCommon.ush b/Engine/Shaders/Private/VirtualTextureCommon.ush
--- a/Engine/Shaders/Private/VirtualTextureCommon.ush
+++ b/Engine/Shaders/Private/VirtualTextureCommon.ush
@@ -31,7 +31,7 @@ void InitializeVirtualTextureFeedback(in out FVirtualTextureFeedbackParams Param
 {
 #if (NUM_VIRTUALTEXTURE_SAMPLES + LIGHTMAP_VT_ENABLED) > 1
 	const uint NumVTSamplesInShader = NUM_VIRTUALTEXTURE_SAMPLES + LIGHTMAP_VT_ENABLED;
-	const uint2 PixelPos = SvPosition / VIRTUAL_TEXTURE_FEEDBACK_FACTOR;
+	const uint2 PixelPos = SvPosition >> View.VirtualTextureFeedbackShift;
 	const uint FeedbackPos = PixelPos.y * View.VirtualTextureFeedbackStride + PixelPos.x;
 	Params.RequestId = (View.VirtualTextureFeedbackSampleOffset + FeedbackPos) % NumVTSamplesInShader;
 	Params.Request = 0xFFFFFFFF;
@@ -53,14 +53,14 @@ void StoreVirtualTextureFeedback(in out FVirtualTextureFeedbackParams Params, ui
 /** This should be called at the end of the pixel shader to write out the gathered VT feedback info to the OutputBuffer. */
 void FinalizeVirtualTextureFeedback(in FVirtualTextureFeedbackParams Params, float4 SvPosition, float Opacity, uint FrameNumber, RWBuffer<uint> OutputBuffer)
 {
-	uint2 PixelTilePos = (uint2)SvPosition.xy % VIRTUAL_TEXTURE_FEEDBACK_FACTOR;
-	uint PixelTileIndex = PixelTilePos.y * VIRTUAL_TEXTURE_FEEDBACK_FACTOR + PixelTilePos.x;
+	uint2 PixelTilePos = (uint2)SvPosition.xy & View.VirtualTextureFeedbackMask;
+	uint PixelTileIndex = (PixelTilePos.y << View.VirtualTextureFeedbackShift) + PixelTilePos.x;
 
 	// This code will only run every few pixels...
 	[branch] if (PixelTileIndex == View.VirtualTextureFeedbackJitterOffset)
 	{
 		// TODO use append buffer ?
-		uint2 PixelPos = (uint2)SvPosition.xy / VIRTUAL_TEXTURE_FEEDBACK_FACTOR;
+		uint2 PixelPos = (uint2)SvPosition.xy >> View.VirtualTextureFeedbackShift;
 		uint FeedbackPos = PixelPos.y * View.VirtualTextureFeedbackStride + PixelPos.x;
 
 		// When using DBuffer, feedback buffer is blending using decal opacity for opaque objects
@@ -73,24 +73,25 @@ void FinalizeVirtualTextureFeedback(in FVirtualTextureFeedbackParams Params, flo
 		// Use a value of 0.25f so we have up to 4 layers of VT transparancy
 		const float Alpha = 0.25f;  
 #else
-		// Opacity values smaller/larger than the reciprocal of the downscale factor get clamped so at least one pixel will resolve
-		const float VTDownscaleEpsilon = 1.0f / VIRTUAL_TEXTURE_FEEDBACK_FACTOR;
-		float MaxAlpha = 1.0f - VTDownscaleEpsilon;
 
+#if (IS_DECAL && !IS_DBUFFER_DECAL)
 		// Even fully opaque gbuffer decals may not completely overwrite data from base pass, so don't allow VT feedback to become fully opaque.
 		// TODO - it may be possible to allow fully opaque decals with certain blend-mode/state combinations, if we can guarantee they overwrite all gbuffer channels
-		// Another similar case is when reading the DBuffer directly (so without a blend alpha available) we don't wan't to allow the base pass to be fully opaque for VT feedback.
-#if (IS_DECAL && !IS_DBUFFER_DECAL) || (USE_DBUFFER && MATERIAL_DECAL_READ_MASK && !MATERIALDECALRESPONSEMASK)
-		MaxAlpha = 0.5f;
+		const float MaxAlpha = 0.5f;
+#elif (USE_DBUFFER && MATERIAL_DECAL_READ_MASK && !MATERIALDECALRESPONSEMASK)
+		// When reading the DBuffer directly (so without a blend alpha available) we don't wan't to allow the base pass to be fully opaque for VT feedback.
+		const float MaxAlpha = 0.5f;
+#else
+		const float MaxAlpha = 1.0f;
 #endif
-
-		const float Alpha =  clamp(Opacity, VTDownscaleEpsilon, MaxAlpha);                  
+		const float MinAlpha = 0.02f;
+		const float Alpha =  clamp(Opacity, MinAlpha, MaxAlpha);
 #endif
 
-        const float AlphaThreshold = frac( PseudoRandom(PixelPos) + // Random  value in 0-1 on 128 x 128 pixel grid
-											SvPosition.w + // Add in depth so we pick different thresholds on different depths 
-											(FrameNumber / (float)VIRTUAL_TEXTURE_FEEDBACK_FACTOR) // Add in framenumber for extra jitter so the pseudorandom pattern changes over time
-											);  
+        const float AlphaThreshold = frac( 
+			PseudoRandom(PixelPos + FrameNumber) + // Random  value in 0-1 on feedback tile. Add in framenumber for extra jitter so the pseudorandom pattern changes over time.
+			SvPosition.w // Add in depth so we pick different thresholds on different depths
+			);  
 											
         // Threshold is uniform pseudo random this conditional has a higher probability to succeed for higher apha values
         // that means more opaque surfaces get more pixels in the feedback buffer which is what we want
diff --git a/Engine/Source/Runtime/Engine/Classes/Engine/RendererSettings.h b/Engine/Source/Runtime/Engine/Classes/Engine/RendererSettings.h
--- a/Engine/Source/Runtime/Engine/Classes/Engine/RendererSettings.h
+++ b/Engine/Source/Runtime/Engine/Classes/Engine/RendererSettings.h
@@ -336,8 +336,7 @@ class ENGINE_API URendererSettings : public UDeveloperSettings
 	UPROPERTY(config, EditAnywhere, Category = VirtualTextures, meta = (
 		EditCondition = "bVirtualTextures",
 		ConsoleVariable = "r.vt.FeedbackFactor", DisplayName = "Feedback resolution factor",
-		ToolTip = "Lower factor will increase virtual texture feedback resolution which increases CPU/GPU overhead, but may decrease streaming latency, especially if materials use many virtual textures. Changing this setting requires restarting the editor.",
-		ConfigRestartRequired = true))
+		ToolTip = "Lower factor will increase virtual texture feedback resolution which increases CPU/GPU overhead, but may decrease streaming latency, especially if materials use many virtual textures."))
 	uint32 VirtualTextureFeedbackFactor;
 
 	UPROPERTY(config, EditAnywhere, Category = VirtualTextures, meta = (
diff --git a/Engine/Source/Runtime/Engine/Private/ShaderCompiler/ShaderCompiler.cpp b/Engine/Source/Runtime/Engine/Private/ShaderCompiler/ShaderCompiler.cpp
--- a/Engine/Source/Runtime/Engine/Private/ShaderCompiler/ShaderCompiler.cpp
+++ b/Engine/Source/Runtime/Engine/Private/ShaderCompiler/ShaderCompiler.cpp
@@ -5763,11 +5763,6 @@ void GlobalBeginCompileShader(
 	Input.Environment.SetDefine(TEXT("PLATFORM_SUPPORTS_MESH_SHADERS_TIER0"), RHISupportsMeshShadersTier0(EShaderPlatform(Target.Platform)) ? 1 : 0);
 	Input.Environment.SetDefine(TEXT("PLATFORM_SUPPORTS_MESH_SHADERS_TIER1"), RHISupportsMeshShadersTier1(EShaderPlatform(Target.Platform)) ? 1 : 0);
 
-	{
-		static const auto CVar = IConsoleManager::Get().FindConsoleVariable(TEXT("r.vt.FeedbackFactor"));
-		Input.Environment.SetDefine(TEXT("VIRTUAL_TEXTURE_FEEDBACK_FACTOR"), CVar ? FMath::RoundUpToPowerOfTwo(FMath::Max(CVar->GetInt(), 1)) : 1);
-	}
-
 	{
 		static const auto CVar = IConsoleManager::Get().FindConsoleVariable(TEXT("r.VT.AnisotropicFiltering"));
 		Input.Environment.SetDefine(TEXT("VIRTUAL_TEXTURE_ANISOTROPIC_FILTERING"), CVar ? (CVar->GetInt() != 0) : 0);
diff --git a/Engine/Source/Runtime/Engine/Private/ShaderCompiler/ShaderGenerationUtil.cpp b/Engine/Source/Runtime/Engine/Private/ShaderCompiler/ShaderGenerationUtil.cpp
--- a/Engine/Source/Runtime/Engine/Private/ShaderCompiler/ShaderGenerationUtil.cpp
+++ b/Engine/Source/Runtime/Engine/Private/ShaderCompiler/ShaderGenerationUtil.cpp
@@ -83,7 +83,6 @@ void FShaderCompileUtilities::ApplyFetchEnvironment(FShaderGlobalDefines& SrcDef
 	FETCH_COMPILE_BOOL(PLATFORM_SUPPORTS_RENDERTARGET_WRITE_MASK);
 	FETCH_COMPILE_BOOL(PLATFORM_SUPPORTS_PER_PIXEL_DBUFFER_MASK);
 	FETCH_COMPILE_BOOL(PLATFORM_SUPPORTS_DISTANCE_FIELDS);
-	FETCH_COMPILE_BOOL(VIRTUAL_TEXTURE_FEEDBACK_FACTOR);
 	FETCH_COMPILE_BOOL(PROJECT_MOBILE_ENABLE_MOVABLE_SPOTLIGHTS);
 	FETCH_COMPILE_BOOL(COMPILE_SHADERS_FOR_DEVELOPMENT_ALLOWED);
 
@@ -364,11 +363,6 @@ static FShaderGlobalDefines FetchShaderGlobalDefines(EShaderPlatform TargetPlatf
 	Ret.PLATFORM_SUPPORTS_PER_PIXEL_DBUFFER_MASK = FDataDrivenShaderPlatformInfo::GetSupportsPerPixelDBufferMask(EShaderPlatform(TargetPlatform)) ? 1 : 0;
 	Ret.PLATFORM_SUPPORTS_DISTANCE_FIELDS = DoesPlatformSupportDistanceFields(EShaderPlatform(TargetPlatform)) ? 1 : 0;
 	
-	{
-		static const auto CVar = IConsoleManager::Get().FindConsoleVariable(TEXT("r.vt.FeedbackFactor"));
-		Ret.VIRTUAL_TEXTURE_FEEDBACK_FACTOR = (CVar ? FMath::Max(CVar->GetInt(), 1) : 1) != 0;
-	}
-
 	if (bIsMobilePlatform)
 	{
 		static IConsoleVariable* CVar = IConsoleManager::Get().FindConsoleVariable(TEXT("r.Mobile.EnableMovableSpotlights"));
diff --git a/Engine/Source/Runtime/Engine/Public/SceneView.h b/Engine/Source/Runtime/Engine/Public/SceneView.h
--- a/Engine/Source/Runtime/Engine/Public/SceneView.h
+++ b/Engine/Source/Runtime/Engine/Public/SceneView.h
@@ -746,6 +746,8 @@ enum ETranslucencyVolumeCascade
 	VIEW_UNIFORM_BUFFER_MEMBER(FMatrix44f, WorldToVirtualTexture) \
 	VIEW_UNIFORM_BUFFER_MEMBER_ARRAY(FVector4, XRPassthroughCameraUVs, [2]) \
 	VIEW_UNIFORM_BUFFER_MEMBER(float, GlobalVirtualTextureMipBias) \
+	VIEW_UNIFORM_BUFFER_MEMBER(uint32, VirtualTextureFeedbackShift) \
+	VIEW_UNIFORM_BUFFER_MEMBER(uint32, VirtualTextureFeedbackMask) \
 	VIEW_UNIFORM_BUFFER_MEMBER(uint32, VirtualTextureFeedbackStride) \
 	VIEW_UNIFORM_BUFFER_MEMBER(uint32, VirtualTextureFeedbackJitterOffset) \
 	VIEW_UNIFORM_BUFFER_MEMBER(uint32, VirtualTextureFeedbackSampleOffset) \
diff --git a/Engine/Source/Runtime/RenderCore/Private/Shader.cpp b/Engine/Source/Runtime/RenderCore/Private/Shader.cpp
--- a/Engine/Source/Runtime/RenderCore/Private/Shader.cpp
+++ b/Engine/Source/Runtime/RenderCore/Private/Shader.cpp
@@ -1808,10 +1808,7 @@ void ShaderMapAppendKeyString(EShaderPlatform Platform, FString& KeyString)
 		static const auto CVarVTAnisotropic = IConsoleManager::Get().FindTConsoleVariableDataInt(TEXT("r.VT.AnisotropicFiltering"));
 		const bool VTAnisotropic = CVarVTAnisotropic && CVarVTAnisotropic->GetValueOnAnyThread() != 0;
 
-		static const auto CVarVTFactor = IConsoleManager::Get().FindConsoleVariable(TEXT("r.vt.FeedbackFactor")); check(CVarVTFactor);
-		const int32 VTFeedbackFactor = CVarVTFactor->GetInt(); 
-
-		auto tt = FString::Printf(TEXT("_VT-%d-%d-%d-%d-%d"), VTLightmaps, VTTextures, VTSupported, VTFeedbackFactor, VTAnisotropic);
+		auto tt = FString::Printf(TEXT("_VT-%d-%d-%d-%d"), VTLightmaps, VTTextures, VTSupported, VTAnisotropic);
  		KeyString += tt;
 	}
 
diff --git a/Engine/Source/Runtime/RenderCore/Public/ShaderMaterial.h b/Engine/Source/Runtime/RenderCore/Public/ShaderMaterial.h
--- a/Engine/Source/Runtime/RenderCore/Public/ShaderMaterial.h
+++ b/Engine/Source/Runtime/RenderCore/Public/ShaderMaterial.h
@@ -36,7 +36,6 @@ struct FShaderGlobalDefines
 	bool PLATFORM_SUPPORTS_RENDERTARGET_WRITE_MASK;
 	bool PLATFORM_SUPPORTS_PER_PIXEL_DBUFFER_MASK;
 	bool PLATFORM_SUPPORTS_DISTANCE_FIELDS;
-	bool VIRTUAL_TEXTURE_FEEDBACK_FACTOR;
 	bool PROJECT_MOBILE_ENABLE_MOVABLE_SPOTLIGHTS;
 
 	// This is a sepcial one. The flag COMPILE_SHADERS_FOR_DEVELOPMENT is set by the cvar r.CompileShadersForDevelopment, but only
diff --git a/Engine/Source/Runtime/Renderer/Private/SceneRendering.cpp b/Engine/Source/Runtime/Renderer/Private/SceneRendering.cpp
--- a/Engine/Source/Runtime/Renderer/Private/SceneRendering.cpp
+++ b/Engine/Source/Runtime/Renderer/Private/SceneRendering.cpp
@@ -1808,11 +1808,15 @@ void FViewInfo::SetupUniformBufferParameters(
 
 	ViewUniformShaderParameters.GlobalVirtualTextureMipBias = FVirtualTextureSystem::Get().GetGlobalMipBias();
 
+	const uint32 VirtualTextureFeedbackScale = GetVirtualTextureFeedbackScale();
+	check(VirtualTextureFeedbackScale == 1 << FMath::FloorLog2(VirtualTextureFeedbackScale));
+	ViewUniformShaderParameters.VirtualTextureFeedbackShift = FMath::FloorLog2(VirtualTextureFeedbackScale);
+	ViewUniformShaderParameters.VirtualTextureFeedbackMask = VirtualTextureFeedbackScale - 1;
 	ViewUniformShaderParameters.VirtualTextureFeedbackStride = GetVirtualTextureFeedbackBufferSize(SceneTexturesConfig.Extent).X;
 	// Use some low(ish) discrepancy sequence to run over every pixel in the virtual texture feedback tile.
 	ViewUniformShaderParameters.VirtualTextureFeedbackJitterOffset = SampleVirtualTextureFeedbackSequence(FrameIndex);
 	// Offset the selected sample index for each frame and add an additional offset each time we iterate over a full virtual texture feedback tile to ensure we get full coverage of sample indices over time.
-	const uint32 NumPixelsInTile = FMath::Square(GetVirtualTextureFeedbackScale());
+	const uint32 NumPixelsInTile = FMath::Square(VirtualTextureFeedbackScale);
 	ViewUniformShaderParameters.VirtualTextureFeedbackSampleOffset = (FrameIndex % NumPixelsInTile) + (FrameIndex / NumPixelsInTile);
 	
 	ViewUniformShaderParameters.RuntimeVirtualTextureMipLevel = FVector4(ForceInitToZero);
diff --git a/Engine/Source/Runtime/Renderer/Private/VT/VirtualTextureFeedbackBuffer.cpp b/Engine/Source/Runtime/Renderer/Private/VT/VirtualTextureFeedbackBuffer.cpp
--- a/Engine/Source/Runtime/Renderer/Private/VT/VirtualTextureFeedbackBuffer.cpp
+++ b/Engine/Source/Runtime/Renderer/Private/VT/VirtualTextureFeedbackBuffer.cpp
@@ -11,7 +11,7 @@ static FAutoConsoleVariableRef CVarVirtualTextureFeedbackFactor(
 	GVirtualTextureFeedbackFactor,
 	TEXT("The size of the VT feedback buffer is calculated by dividing the render resolution by this factor.")
 	TEXT("The value set here is rounded up to the nearest power of two before use."),
-	ECVF_RenderThreadSafe | ECVF_ReadOnly /*Read-only as shaders are compiled with this value*/
+	ECVF_RenderThreadSafe
 );
 
 int32 GetVirtualTextureFeedbackScale()
@@ -22,7 +22,7 @@ int32 GetVirtualTextureFeedbackScale()
 
 FIntPoint GetVirtualTextureFeedbackBufferSize(FIntPoint InSceneTextureExtent)
 {
-	return FIntPoint::DivideAndRoundUp(InSceneTextureExtent, FMath::Max(GetVirtualTextureFeedbackScale(), 1));
+	return FIntPoint::DivideAndRoundUp(InSceneTextureExtent, GetVirtualTextureFeedbackScale());
 }
 
 uint32 SampleVirtualTextureFeedbackSequence(uint32 InFrameIndex)
