Improve shadow terminator of bump/normal mapped surfaces in the path tracer.

Apply the shadow terminator heuristic from RTGems to soften the transition. This required tracking the original smooth normal in the payload, which is now packed together with the tangent vector in an oct24 encoding.

diff --git a/Engine/Shaders/Private/PathTracing/Material/PathTracingCloth.ush b/Engine/Shaders/Private/PathTracing/Material/PathTracingCloth.ush
--- a/Engine/Shaders/Private/PathTracing/Material/PathTracingCloth.ush
+++ b/Engine/Shaders/Private/PathTracing/Material/PathTracingCloth.ush
@@ -52,7 +52,7 @@ FMaterialEval Cloth_EvalMaterial(
 	const float3 FCloth = F_Schlick(FuzzColor, VoH);
 	const float3 ClothWeight = NoL > 0 && NoV > 0 ? (PI * Cloth * DCloth * Vis_Cloth(NoV, NoL)) * FCloth * Fuzz.W : 0.0;
 
-	AddLobeWithMIS(Result.Weight, Result.Pdf, DiffWeight + ClothWeight, NoL / PI, LobeSelectionPdf);
+	AddLobeWithMIS(Result.Weight, Result.Pdf, (DiffWeight + ClothWeight) * ShadowTerminatorTerm(L_World, Payload.WorldNormal, Payload.WorldSmoothNormal), NoL / PI, LobeSelectionPdf);
 
 	// Specular lobe
 	const float2 GGXResult = GGXEvalReflection(L, V, H, Alpha);
@@ -134,6 +134,8 @@ FMaterialSample Cloth_SampleMaterial(
 
 	// With a valid direction in hand -- now evaluate the BxDF (taking advantage of already computed terms)
 
+	const float3 L_World = normalize(mul(L, Basis));
+
 	const float NoL = saturate(L.z);
 	const float VoH = saturate(dot(V, H));
 	const float NoH = saturate(H.z);
@@ -145,7 +147,7 @@ FMaterialSample Cloth_SampleMaterial(
 
 	float3 OutWeight = 0;
 	float OutPdf = 0;
-	AddLobeWithMIS(OutWeight, OutPdf, DiffWeight + ClothWeight, NoL / PI, LobeSelectionPdf);
+	AddLobeWithMIS(OutWeight, OutPdf, (DiffWeight + ClothWeight) * ShadowTerminatorTerm(L_World, Payload.WorldNormal, Payload.WorldSmoothNormal), NoL / PI, LobeSelectionPdf);
 
 	// Specular lobe
 
@@ -156,5 +158,5 @@ FMaterialSample Cloth_SampleMaterial(
 	AddLobeWithMIS(OutWeight, OutPdf, SpecWeight, SpecPdf, 1 - LobeSelectionPdf);
 
 	// transform to world space
-	return CreateMaterialSample(normalize(mul(L, Basis)), OutWeight, OutPdf, 1.0, OutRoughness);
+	return CreateMaterialSample(L_World, OutWeight, OutPdf, 1.0, OutRoughness);
 }
diff --git a/Engine/Shaders/Private/PathTracing/Material/PathTracingDefaultLit.ush b/Engine/Shaders/Private/PathTracing/Material/PathTracingDefaultLit.ush
--- a/Engine/Shaders/Private/PathTracing/Material/PathTracingDefaultLit.ush
+++ b/Engine/Shaders/Private/PathTracing/Material/PathTracingDefaultLit.ush
@@ -40,7 +40,7 @@ FMaterialEval DefaultLit_EvalMaterial(
 	const float LobeSelectionPdf = LobeSelectionProb(Payload.DiffuseColor * (1 - Spec.E), Spec.E);
 
 	// Diffuse Lobe
-	AddLobeWithMIS(Result.Weight, Result.Pdf, (1 - Spec.E) * Payload.DiffuseColor, NoL / PI, LobeSelectionPdf);
+	AddLobeWithMIS(Result.Weight, Result.Pdf, (1 - Spec.E) * Payload.DiffuseColor * ShadowTerminatorTerm(L_World, Payload.WorldNormal, Payload.WorldSmoothNormal), NoL / PI, LobeSelectionPdf);
 
 	// Specular lobe
 	const float2 GGXResult = GGXEvalReflection(L, V, H, Alpha);
@@ -107,6 +107,8 @@ FMaterialSample DefaultLit_SampleMaterial(
 
 	// With a valid direction in hand -- now evaluate the BxDF (taking advantage of already computed terms)
 
+	const float3 L_World = normalize(mul(L, Basis));
+
 	const float NoL = saturate(L.z);
 	const float VoH = saturate(dot(V, H));
 
@@ -114,7 +116,7 @@ FMaterialSample DefaultLit_SampleMaterial(
 	float OutPdf = 0;
 
 	// Diffuse Lobe
-	AddLobeWithMIS(OutWeight, OutPdf, (1 - Spec.E) * Payload.DiffuseColor, NoL / PI, LobeSelectionPdf);
+	AddLobeWithMIS(OutWeight, OutPdf, (1 - Spec.E) * Payload.DiffuseColor * ShadowTerminatorTerm(L_World, Payload.WorldNormal, Payload.WorldSmoothNormal), NoL / PI, LobeSelectionPdf);
 
 	// Specular lobe
 	const float2 GGXResult = GGXEvalReflection(L, V, H, Alpha);
@@ -124,5 +126,5 @@ FMaterialSample DefaultLit_SampleMaterial(
 	AddLobeWithMIS(OutWeight, OutPdf, SpecWeight, SpecPdf, 1.0 - LobeSelectionPdf);
 
 	// transform to world space
-	return CreateMaterialSample(normalize(mul(L, Basis)), OutWeight, OutPdf, 1.0, OutRoughness);
+	return CreateMaterialSample(L_World, OutWeight, OutPdf, 1.0, OutRoughness);
 }
diff --git a/Engine/Shaders/Private/PathTracing/Material/PathTracingLambert.ush b/Engine/Shaders/Private/PathTracing/Material/PathTracingLambert.ush
--- a/Engine/Shaders/Private/PathTracing/Material/PathTracingLambert.ush
+++ b/Engine/Shaders/Private/PathTracing/Material/PathTracingLambert.ush
@@ -13,7 +13,9 @@ FMaterialSample Lambert_SampleMaterial(
 
 	float4 SampledValue = CosineSampleHemisphere(RandSample.xy);
 
-	return CreateMaterialSample(TangentToWorld(SampledValue.xyz, N_World), Payload.BaseColor, SampledValue.w, 1.0, 1.0);
+	float3 L_World = TangentToWorld(SampledValue.xyz, N_World);
+
+	return CreateMaterialSample(L_World, Payload.BaseColor * ShadowTerminatorTerm(L_World, Payload.WorldNormal, Payload.WorldSmoothNormal), SampledValue.w, 1.0, 1.0);
 }
 
 FMaterialEval Lambert_EvalMaterial(
@@ -22,9 +24,10 @@ FMaterialEval Lambert_EvalMaterial(
 )
 {
 	float3 N_World = Payload.WorldNormal;
+	float3 L_World = OutgoingDirection;
 	float NoL = saturate(dot(N_World, OutgoingDirection));
 
-	return CreateMaterialEval(Payload.BaseColor, NoL / PI);
+	return CreateMaterialEval(Payload.BaseColor * ShadowTerminatorTerm(L_World, Payload.WorldNormal, Payload.WorldSmoothNormal), NoL / PI);
 }
 
 
diff --git a/Engine/Shaders/Private/PathTracing/Material/PathTracingMaterialCommon.ush b/Engine/Shaders/Private/PathTracing/Material/PathTracingMaterialCommon.ush
--- a/Engine/Shaders/Private/PathTracing/Material/PathTracingMaterialCommon.ush
+++ b/Engine/Shaders/Private/PathTracing/Material/PathTracingMaterialCommon.ush
@@ -91,3 +91,50 @@ FMaterialEval CreateMaterialEval(float3 Weight, float Pdf)
 	Result.Pdf = Pdf;
 	return Result;
 }
+
+
+
+#define PATH_TRACER_SHADOW_TERMINATOR   1 /// 0: off 1: imageworks 2: disney 3: dreamworks (see paper references below)
+
+float ShadowTerminatorTerm(float3 L, float3 N, float3 Ns)
+{
+	// #jira UE-121401
+#if PATH_TRACER_SHADOW_TERMINATOR == 1
+	// Imageworks terminator softening:
+	// "A Microfacet-Based Shadowing Function to Solve the Bump Terminator Problem"
+	// Alejandro Conty Estevez, Pascal Lecocq, and Clifford Stein.
+	// Ray Tracing Gems(2019), 149 - 158.
+	// http://www.aconty.com/pdf/bump-terminator-nvidia2019.pdf
+	const float Epsilon = 1e-6; // avoid division by 0
+	const float CosD = saturate(abs(dot(Ns, N)));
+	const float Tan2D = (1.0 - CosD * CosD) / (CosD * CosD + Epsilon);
+	const float Alpha2 = saturate(0.125 * Tan2D);
+	const float CosI = saturate(dot(Ns, L));
+	const float Tan2I = (1.0f - CosI * CosI) / (CosI * CosI + Epsilon);
+	return CosI > 0 ? 2.0f / (1.0f + sqrt(1.0f + Alpha2 * Tan2I)) : 0.0;
+#elif PATH_TRACER_SHADOW_TERMINATOR == 2
+	// Disney terminator softening:
+	// "Taming the Shadow Terminator"
+	// Matt Jen-Yuan Chiang, Yining Karl Li, and Brent Burley
+	// SIGGRAPH 2019 Talks
+	// https://www.yiningkarlli.com/projects/shadowterminator.html
+	const float NoL = saturate(dot(N, L));
+	const float NgoL = saturate(dot(Ns, L));
+	const float NgoN = saturate(dot(Ns, N));
+	const float G = saturate(NgoL / (NoL * NgoN + 1e-6));
+	return G + G * (G - G * G); // smooth
+#elif PATH_TRACER_SHADOW_TERMINATOR == 3
+	// Dreamworks terminator softening:
+	// "Predictable and Targeted Softening of the Shadow Terminator"
+	// Priyamvad Deshmukh, Brian Green
+	// SIGGRAPH 2020 Talks
+	// https://research.dreamworks.com/wp-content/uploads/2020/08/talk_shadow_terminator.pdf
+	const float CosD = saturate(abs(dot(Ns, N)));
+	const float d = lerp(0.15, 0.05, CosD);
+	const float t = saturate(dot(Ns, L) / d);
+	return t * t * (3.0 - 2.0 * t);
+#else
+	// no special handling of the terminator, don't do anything
+	return 1.0;
+#endif
+}
diff --git a/Engine/Shaders/Private/PathTracing/Material/PathTracingSubsurfaceProfile.ush b/Engine/Shaders/Private/PathTracing/Material/PathTracingSubsurfaceProfile.ush
--- a/Engine/Shaders/Private/PathTracing/Material/PathTracingSubsurfaceProfile.ush
+++ b/Engine/Shaders/Private/PathTracing/Material/PathTracingSubsurfaceProfile.ush
@@ -56,7 +56,7 @@ FMaterialEval SubsurfaceProfile_EvalMaterial(
 	const float VoH = saturate(dot(V, H));
 
 	// Diffuse Lobe
-	AddLobeWithMIS(Result.Weight, Result.Pdf, Payload.DiffuseColor * (1 - SpecE), NoL / PI, LobeProb.x);
+	AddLobeWithMIS(Result.Weight, Result.Pdf, Payload.DiffuseColor * (1 - SpecE) * ShadowTerminatorTerm(L_World, Payload.WorldNormal, Payload.WorldSmoothNormal), NoL / PI, LobeProb.x);
 
 	// Specular lobes
 	const float2 GGXResult0 = GGXEvalReflection(L, V, H, Alpha0);
@@ -144,6 +144,8 @@ FMaterialSample SubsurfaceProfile_SampleMaterial(
 
 	// With a valid direction in hand -- now evaluate the BxDF (taking advantage of already computed terms)
 
+	const float3 L_World = normalize(mul(L, Basis));
+
 	const float NoL = saturate(L.z);
 	const float VoH = saturate(dot(V, H));
 
@@ -151,7 +153,7 @@ FMaterialSample SubsurfaceProfile_SampleMaterial(
 	float OutPdf = 0;
 
 	// Diffuse Lobe
-	AddLobeWithMIS(OutWeight, OutPdf, Payload.DiffuseColor * (1 - SpecE), NoL / PI, LobeProb.x);
+	AddLobeWithMIS(OutWeight, OutPdf, Payload.DiffuseColor * (1 - SpecE) * ShadowTerminatorTerm(L_World, Payload.WorldNormal, Payload.WorldSmoothNormal), NoL / PI, LobeProb.x);
 
 	// Specular lobe
 	const float2 GGXResult0 = GGXEvalReflection(L, V, H, Alpha0);
@@ -161,5 +163,5 @@ FMaterialSample SubsurfaceProfile_SampleMaterial(
 	AddLobeWithMIS(OutWeight, OutPdf, (    DualRoughnessData.z) * F * GGXResult1.x * Spec1.W, GGXResult1.y, LobeProb.z);
 
 	// transform to world space
-	return CreateMaterialSample(normalize(mul(L, Basis)), OutWeight, OutPdf, 1.0, OutRoughness);
+	return CreateMaterialSample(L_World, OutWeight, OutPdf, 1.0, OutRoughness);
 }
diff --git a/Engine/Shaders/Private/PathTracing/PathTracingCommon.ush b/Engine/Shaders/Private/PathTracing/PathTracingCommon.ush
--- a/Engine/Shaders/Private/PathTracing/PathTracingCommon.ush
+++ b/Engine/Shaders/Private/PathTracing/PathTracingCommon.ush
@@ -99,8 +99,9 @@ struct FPathTracingPayload : FMinimalPayload
 {
 	FRayCone RayCone;
 	float3 Radiance;
-	float3 WorldGeoNormal;
-	float3 WorldNormal;
+	float3 WorldGeoNormal;           // normal of the actual triangle (faceted)
+	float3 WorldNormal;              // normal output of the material (includes normal/bump)
+	float3 WorldSmoothNormal;        // smooth normal before normal/bump is applied
 	float3 BaseColor;
 	float3 DiffuseColor;
 	float3 SpecularColor;
@@ -218,9 +219,11 @@ FPackedPathTracingPayload PackPathTracingPayload(FPathTracingPayload Input)
 	Output.PackedData[7] |= (Input.PrimitiveLightingChannelMask & 0x7) << 29; // 3 bits
                                                                        // total 32 bits
 
-	Output.PackedData[8] = f32tof16(Input.WorldTangent.x);
-	Output.PackedData[8] |= f32tof16(Input.WorldTangent.y) << 16;
-	Output.PackedData[9] = f32tof16(Input.WorldTangent.z);
+	uint SmoothNormal = PayloadEncodeUnitVector(Input.WorldSmoothNormal);
+	uint TangentVec   = PayloadEncodeUnitVector(Input.WorldTangent);
+	Output.PackedData[8] = SmoothNormal;                     // pack all 24 bits
+	Output.PackedData[8] |= (TangentVec << 8) & 0xFF000000u; // pack upper 8 bits
+	Output.PackedData[9] = TangentVec & 0xFFFFu;             // pack lower 16 bits
 
 	Output.PackedData[9] |= f32tof16(Input.CustomVector.x) << 16;
 	Output.PackedData[10] = f32tof16(Input.CustomVector.y);
@@ -259,9 +262,8 @@ FPathTracingPayload UnpackPathTracingPayload(FPackedPathTracingPayload Input, Ra
 	Output.Flags = Input.GetFlags();
 	Output.PrimitiveLightingChannelMask = Input.GetPrimitiveLightingChannelMask();
 
-	Output.WorldTangent.x = f16tof32(Input.PackedData[8]);
-	Output.WorldTangent.y = f16tof32(Input.PackedData[8] >> 16);
-	Output.WorldTangent.z = f16tof32(Input.PackedData[9]);
+	Output.WorldSmoothNormal = PayloadDecodeUnitVector(Input.PackedData[8] & 0x00FFFFFFu);
+	Output.WorldTangent = PayloadDecodeUnitVector((Input.PackedData[8] & 0xFF000000u) >> 8 | (Input.PackedData[9] & 0xFFFFu));
 
 	Output.CustomVector.x = f16tof32(Input.PackedData[9] >> 16);
 	Output.CustomVector.y = f16tof32(Input.PackedData[10]);
diff --git a/Engine/Shaders/Private/PathTracing/PathTracingCore.ush b/Engine/Shaders/Private/PathTracing/PathTracingCore.ush
--- a/Engine/Shaders/Private/PathTracing/PathTracingCore.ush
+++ b/Engine/Shaders/Private/PathTracing/PathTracingCore.ush
@@ -682,6 +682,7 @@ bool PathTracingKernel(inout FPathState PathState, int Bounce)
 
 	// fix dark edges in reflections from shading normals
 	Payload.WorldNormal = AdjustShadingNormal(Payload.WorldNormal, Payload.WorldGeoNormal, PathState.Ray.Direction);
+	Payload.WorldSmoothNormal = AdjustShadingNormal(Payload.WorldSmoothNormal, Payload.WorldGeoNormal, PathState.Ray.Direction);
 
 	// As soon as the path is blurry enough, we can get away with diffuse sampling only
 	const bool bSimplifySSS = PathState.PathRoughness >= 0.15;
diff --git a/Engine/Shaders/Private/PathTracing/PathTracingMaterialHitShader.usf b/Engine/Shaders/Private/PathTracing/PathTracingMaterialHitShader.usf
--- a/Engine/Shaders/Private/PathTracing/PathTracingMaterialHitShader.usf
+++ b/Engine/Shaders/Private/PathTracing/PathTracingMaterialHitShader.usf
@@ -126,7 +126,8 @@ RAY_TRACING_ENTRY_CLOSEST_HIT(PathTracingMaterialCHS,
 	// it doesn't reflect the sign flip from the object transform, so apply it here
 	float GeoNormalSign = MaterialParameters.TwoSidedSign * GetPrimitive_DeterminantSign(MaterialParameters.PrimitiveId);
 	Payload.WorldGeoNormal = GeoNormalSign * GeoNormal;
-	Payload.WorldNormal = normalize(MaterialParameters.WorldNormal);
+	Payload.WorldNormal = MaterialParameters.WorldNormal;
+	Payload.WorldSmoothNormal = MaterialParameters.TwoSidedSign * TransformTangentNormalToWorld(MaterialParameters.TangentToWorld, float3(0, 0, 1));
 	Payload.Radiance = GetMaterialEmissive(PixelMaterialInputs);
 	Payload.BaseColor = BaseColor;
 	Payload.Specular = Specular;
