Implement absorption support for refractive translucent materials in the path tracer

The absorption coefficient is derived from the BaseColor of the material and a global scale factor that gives the inverse distance at which the target BaseColor is achieved. Setting this cvar to 0 disables the effect. The default value of 0.01 corresponds to specifing the target color after travelling through 1 meter of surface. This scale factor will not be needed in future material models that provide more direct control over the absorption color.

The implementation remembers the current absorption value as part of the PathState. For shadow rays, when approximate caustics are being used, the contribution is tracked by simply counting up/down at each interface. An offset is used to account for the possibility of light starting inside an object for finite light sources.

Refactored TraceTransparentRay to directly use the PathState struct since it was accessing most of its fields anyway and the function signature was getting too long

diff --git a/Engine/Shaders/Private/PathTracing/PathTracing.usf b/Engine/Shaders/Private/PathTracing/PathTracing.usf
--- a/Engine/Shaders/Private/PathTracing/PathTracing.usf
+++ b/Engine/Shaders/Private/PathTracing/PathTracing.usf
@@ -63,7 +63,13 @@ FPathState LoadPathStateData(uint Index)
 	Output.PathThroughput.y = asfloat(PathStateData[PATH_STATE_STRIDE * Index + 17]);
 	Output.PathThroughput.z = asfloat(PathStateData[PATH_STATE_STRIDE * Index + 18]);
 
-	Output.PathRoughness = asfloat(PathStateData[PATH_STATE_STRIDE * Index + 19]);
+	uint2 PackedRoughnessSigma;
+	PackedRoughnessSigma.x = PathStateData[PATH_STATE_STRIDE * Index + 19];
+	PackedRoughnessSigma.y = PathStateData[PATH_STATE_STRIDE * Index + 20];
+	Output.PathRoughness = f16tof32(PackedRoughnessSigma.x);
+	Output.SigmaT.x = f16tof32(PackedRoughnessSigma.x >> 16);
+	Output.SigmaT.y = f16tof32(PackedRoughnessSigma.y);
+	Output.SigmaT.z = f16tof32(PackedRoughnessSigma.y >> 16);
 
 	return Output;
 }
@@ -96,7 +102,8 @@ void StorePathStateData(FPathState PathState, uint Index)
 	PathStateData[PATH_STATE_STRIDE * Index + 17] = asuint(PathState.PathThroughput.y);
 	PathStateData[PATH_STATE_STRIDE * Index + 18] = asuint(PathState.PathThroughput.z);
 
-	PathStateData[PATH_STATE_STRIDE * Index + 19] = asuint(PathState.PathRoughness);
+	PathStateData[PATH_STATE_STRIDE * Index + 19] = f32tof16(PathState.PathRoughness) | (f32tof16(PathState.SigmaT.x) << 16);
+	PathStateData[PATH_STATE_STRIDE * Index + 20] = f32tof16(PathState.SigmaT.y) | (f32tof16(PathState.SigmaT.z) << 16);
 }
 
 RAY_TRACING_ENTRY_RAYGEN(PathTracingMainRG)
diff --git a/Engine/Shaders/Private/PathTracing/PathTracingCommon.ush b/Engine/Shaders/Private/PathTracing/PathTracingCommon.ush
--- a/Engine/Shaders/Private/PathTracing/PathTracingCommon.ush
+++ b/Engine/Shaders/Private/PathTracing/PathTracingCommon.ush
@@ -92,6 +92,23 @@ struct FPackedPathTracingPayload : FMinimalPayload
 	{
 		PackedData[3] = asuint(PathRoughness);
 	}
+
+	// tau: optical depth along shadow ray, this is the integral of the extinction coefficient
+	float3 GetTau()
+	{
+		return float3(
+			asfloat(PackedData[4]),
+			asfloat(PackedData[5]),
+			asfloat(PackedData[6])
+		);
+	}
+
+	void SetTau(float3 Tau)
+	{
+		PackedData[4] = asuint(Tau.x);
+		PackedData[5] = asuint(Tau.y);
+		PackedData[6] = asuint(Tau.z);
+	}
 };
 
 // This payload structure is the expanded version of the above which is more convenient to work with
@@ -146,9 +163,16 @@ struct FPathTracingPayload : FMinimalPayload
 		Roughness = lerp(RoughnessData.x, RoughnessData.y, RoughnessData.z);
 	}
 
+	bool IsMaterialSolidGlass()
+	{
+		return ShadingModelID != SHADINGMODELID_THIN_TRANSLUCENT && BlendingMode == RAY_TRACING_BLEND_MODE_TRANSLUCENT && Ior != 1.0 && Opacity < 1;
+	}
+
 	bool IsMaterialTransmissive()
 	{
-		return ShadingModelID == SHADINGMODELID_TWOSIDED_FOLIAGE || (BlendingMode == RAY_TRACING_BLEND_MODE_TRANSLUCENT && Ior != 1.0 && Roughness != 0 && Opacity < 1);
+		return ShadingModelID == SHADINGMODELID_TWOSIDED_FOLIAGE ||
+			   ShadingModelID == SHADINGMODELID_THIN_TRANSLUCENT || 
+			   IsMaterialSolidGlass();
 	}
 
 	bool IsSubsurfaceMaterial()
diff --git a/Engine/Shaders/Private/PathTracing/PathTracingCore.ush b/Engine/Shaders/Private/PathTracing/PathTracingCore.ush
--- a/Engine/Shaders/Private/PathTracing/PathTracingCore.ush
+++ b/Engine/Shaders/Private/PathTracing/PathTracingCore.ush
@@ -43,6 +43,7 @@
 float MaxPathIntensity;
 float MaxNormalBias;
 float FilterWidth;
+float AbsorptionScale; // inverse distance at which we hit BaseColor - exposed mainly so we can disable absorption handling and experiment with the default scale factor
 
 RWTexture2D<float4> RadianceTexture;
 RWTexture2D<float4> AlbedoTexture;
@@ -51,6 +52,43 @@
 int3 TileOffset;
 uint SceneVisibleLightCount;
 
+struct FPathState
+{                                   // packed size
+	uint2 PixelIndex;               //  4 bytes
+	RandomSequence RandSequence;    //  8 bytes
+	float3 Radiance;                // 12 bytes
+	float3 Albedo;                  //  6 bytes // half precision
+	float3 Normal;                  //  6 bytes // half precision
+	RayDesc Ray;                    // 12+12 = 24 bytes (TMin/TMax do not need to be stored)
+	FRayCone RayCone;               //  4 bytes // half precision
+	float3 PathThroughput;          // 12 bytes
+	float PathRoughness;            //  2 bytes // half precision
+	float3 SigmaT; // Extinction    //  6 bytes // half precision 
+									// 84 bytes total
+									// 21 words - sync with PATHTRACER_PATHSTATE_NUM_WORDS
 
+	void WritePixel(float BlendFactor)
+	{
+		// Avoid reading the old pixel on the first sample on the off-chance there is a NaN/Inf pixel ...
+		float4 OldPixel = Iteration > 0 ? RadianceTexture[PixelIndex] : 0;
+		float4 OldAlbedo = Iteration > 0 ? AlbedoTexture[PixelIndex] : 0;
+		float4 OldNormal = Iteration > 0 ? NormalTexture[PixelIndex] : 0;
+		float3 OldRadiance = OldPixel.rgb;
+		float OldVariance = OldPixel.a;
+
+		float3 NewRadiance = lerp(OldRadiance, Radiance, BlendFactor);
+
+		float RadianceY = Luminance(Radiance);
+
+		// https://en.wikipedia.org/wiki/Algorithms_for_calculating_variance#Welford's_online_algorithm
+		float DeviationSquared = (RadianceY - Luminance(NewRadiance)) * (RadianceY - Luminance(OldRadiance));
+		float NewVariance = lerp(OldVariance, DeviationSquared, BlendFactor);
+
+		RadianceTexture[PixelIndex] = float4(NewRadiance, NewVariance);
+		AlbedoTexture[PixelIndex] = lerp(OldAlbedo, float4(Albedo, 0), BlendFactor);
+		NormalTexture[PixelIndex] = lerp(OldNormal, float4(Normal, 0), BlendFactor);
+	}
+};
+
 void AccumulateRadiance(inout float3 TotalRadiance, float3 PathRadiance, bool bIsCameraRay)
 {
@@ -67,17 +105,17 @@
 	TotalRadiance += PathRadiance;
 }
 
-FPathTracingPayload TraceTransparentRay(RayDesc Ray, FRayCone RayCone, bool bIsCameraRay, bool bLastBounce, bool bIncludeEmission, uint2 LaunchIndex, uint NumLights, float PathRoughness, inout RandomSequence RandSequence, inout float3 PathThroughput, inout float3 Radiance, inout float3 Albedo, inout float3 Normal)
+FPathTracingPayload TraceTransparentRay(inout FPathState PathState, bool bIsCameraRay, bool bLastBounce, bool bIncludeEmission, uint NumVisibleLights)
 {
 	const uint RayFlags = bIsCameraRay && EnableCameraBackfaceCulling ? RAY_FLAG_CULL_BACK_FACING_TRIANGLES : 0;
 	const uint MissShaderIndex = 0;
 	float SelectionWeightSum = 0;
-	float3 PayloadThroughput = PathThroughput;
+	float3 PayloadThroughput = PathState.PathThroughput;
 	FPathTracingPayload Payload;
 	if (!bIncludeEmission && bLastBounce)
 	{
 		Payload.SetMiss();
-		PathThroughput = 0;
+		PathState.PathThroughput = 0;
 		return Payload;
 	}
 	for (;;)
@@ -86,7 +124,7 @@
 		// also avoid noise from hot spots (they can be properly
 		// importance sampled if a capturing skylight is added)
 		const bool bIgnoreSky = !bIsCameraRay;
-		FPackedPathTracingPayload PackedPayload = InitPathTracingPayload(RayCone, LaunchIndex, bIgnoreSky);
+		FPackedPathTracingPayload PackedPayload = InitPathTracingPayload(PathState.RayCone, PathState.PixelIndex, bIgnoreSky);
 		// Trace the ray
 		TraceRay(
 			TLAS,
@@ -95,17 +133,17 @@
 			RAY_TRACING_SHADER_SLOT_MATERIAL,
 			RAY_TRACING_NUM_SHADER_SLOTS,
 			MissShaderIndex,
-			Ray,
+			PathState.Ray,
 			PackedPayload);
 
 		// Loop over lights to capture their contribution
 		// #dxr_todo: if we have lots of lights, having some hierarchical structure would be better ....
-		for (uint LightId = 0; LightId < NumLights; ++LightId)
+		for (uint LightId = 0; LightId < NumVisibleLights; ++LightId)
 		{
-			RayDesc LightRay = Ray;
-			LightRay.TMax = PackedPayload.IsMiss() ? Ray.TMax : PackedPayload.HitT;
+			RayDesc LightRay = PathState.Ray;
+			LightRay.TMax = PackedPayload.IsMiss() ? PathState.Ray.TMax : PackedPayload.HitT;
 			float3 LightRadiance = TraceLight(LightRay, LightId).Radiance;
-			AccumulateRadiance(Radiance, PathThroughput * LightRadiance, bIsCameraRay);
+			AccumulateRadiance(PathState.Radiance, PathState.PathThroughput * LightRadiance, bIsCameraRay);
 		}
 
 		if (PackedPayload.IsMiss())
@@ -115,7 +153,7 @@
 		}
 
 		// Unpack the payload
-		FPathTracingPayload HitPayload = UnpackPathTracingPayloadMaterial(PackedPayload, Ray);
+		FPathTracingPayload HitPayload = UnpackPathTracingPayloadMaterial(PackedPayload, PathState.Ray);
 
 		if (HitPayload.BlendingMode == RAY_TRACING_BLEND_MODE_ADDITIVE ||
 			HitPayload.BlendingMode == RAY_TRACING_BLEND_MODE_TRANSLUCENT)
@@ -127,17 +165,20 @@
 		// add in surface emission (except for modulate blend mode which uses Radiance as the transparency multiplier)
 		if (bIncludeEmission && HitPayload.BlendingMode != RAY_TRACING_BLEND_MODE_MODULATE)
 		{
-			AccumulateRadiance(Radiance, PathThroughput * HitPayload.Radiance, bIsCameraRay);
+			AccumulateRadiance(PathState.Radiance, PathState.PathThroughput * HitPayload.Radiance, bIsCameraRay);
 		}
 
-		float3 Transparency = GetMaterialTransparency(HitPayload, Ray.Direction);
+		float3 Transparency = GetMaterialTransparency(HitPayload, PathState.Ray.Direction);
 
+		// account for Beer's law through the currently active medium
+		PathState.PathThroughput *= exp(-PathState.SigmaT * (HitPayload.HitT - PathState.Ray.TMin));
+
 		if (!bLastBounce)
 		{
-			float3 Contrib = PathThroughput * EstimateMaterialAlbedo(HitPayload);
+			float3 Contrib = PathState.PathThroughput * EstimateMaterialAlbedo(HitPayload);
 
 			// accumulate what the denoiser wants into albedo/normal (as long as the current path is rough enough)
-			AccumulateAlbedoNormal(HitPayload, PathThroughput, PathRoughness, Albedo, Normal);
+			AccumulateAlbedoNormal(HitPayload, PathState.PathThroughput, PathState.PathRoughness, PathState.Albedo, PathState.Normal);
 
 			float SelectionWeight = max3(Contrib.x, Contrib.y, Contrib.z);
 			SelectionWeightSum += SelectionWeight;
@@ -150,7 +191,7 @@
 				{
 					// the acceptance probability is not 1.0
 					// generate a random number to see if we should accept this hit
-					float RandValue = RandomSequence_GenerateSample1D(RandSequence);
+					float RandValue = RandomSequence_GenerateSample1D(PathState.RandSequence);
 					AcceptHit = RandValue * SelectionWeightSum < SelectionWeight;
 				}
 				else
@@ -164,16 +205,17 @@
 			{
 				// stash this hit for next time
 				Payload = HitPayload;
-				PayloadThroughput = PathThroughput / SelectionWeight;
+				PayloadThroughput = PathState.PathThroughput / SelectionWeight;
 			}
 		}
 
+		// account for local transparency change
+		PathState.PathThroughput *= Transparency;
 		// prepare next step around the loop
 		// retrace the exact same ray with TMin one ulp past the hit we just found
-		PathThroughput *= Transparency;
-		Ray.TMin = asfloat(asuint(HitPayload.HitT) + 1);
+		PathState.Ray.TMin = asfloat(asuint(HitPayload.HitT) + 1);
 
-		if (all(PathThroughput == 0))
+		if (all(PathState.PathThroughput == 0))
 		{
 			break;
 		}
@@ -182,11 +224,11 @@
 	if (SelectionWeightSum > 0)
 	{
 		// if we stored a valid hit in the payload, reset the path throughput to this point
-		PathThroughput = PayloadThroughput * SelectionWeightSum;
+		PathState.PathThroughput = PayloadThroughput * SelectionWeightSum;
 	}
 	else
 	{
-		PathThroughput = 0;
+		PathState.PathThroughput = 0;
 		Payload.SetMiss();
 	}
 	return Payload;
@@ -216,7 +258,7 @@
 	if (PackedPayload.IsMiss())
 	{
 		// we didn't hit any geo - just return whatever throughput is left
-		return PackedPayload.GetRayThroughput();
+		return PackedPayload.GetRayThroughput() * exp(-max(PackedPayload.GetTau() * AbsorptionScale, 0.0));
 	}
 
 	// we registered a hit -- assumed to be fully opaque
@@ -536,42 +578,6 @@
 	return GroupCoord * M + Shuffle(GroupId, GroupCoord);
 }
 
-struct FPathState
-{                                   // packed size
-	uint2 PixelIndex;               //  4 bytes
-	RandomSequence RandSequence;    //  8 bytes
-	float3 Radiance;                // 12 bytes
-	float3 Albedo;                  //  6 bytes // half precision
-	float3 Normal;                  //  6 bytes // half precision
-	RayDesc Ray;                    // 12+12 = 24 bytes (TMin/TMax do not need to be stored)
-	FRayCone RayCone;               //  4 bytes // half precision
-	float3 PathThroughput;          // 12 bytes
-	float PathRoughness;            //  4 bytes
-	                                // 80 bytes total
-
-	void WritePixel(float BlendFactor)
-	{
-		// Avoid reading the old pixel on the first sample on the off-chance there is a NaN/Inf pixel ...
-		float4 OldPixel = Iteration > 0 ? RadianceTexture[PixelIndex] : 0;
-		float4 OldAlbedo = Iteration > 0 ? AlbedoTexture[PixelIndex] : 0;
-		float4 OldNormal = Iteration > 0 ? NormalTexture[PixelIndex] : 0;
-		float3 OldRadiance = OldPixel.rgb;
-		float OldVariance = OldPixel.a;
-
-		float3 NewRadiance = lerp(OldRadiance, Radiance, BlendFactor);
-
-		float RadianceY = Luminance(Radiance);
-
-		// https://en.wikipedia.org/wiki/Algorithms_for_calculating_variance#Welford's_online_algorithm
-		float DeviationSquared = (RadianceY - Luminance(NewRadiance)) * (RadianceY - Luminance(OldRadiance));
-		float NewVariance = lerp(OldVariance, DeviationSquared, BlendFactor);
-
-		RadianceTexture[PixelIndex] = float4(NewRadiance, NewVariance);
-		AlbedoTexture[PixelIndex] = lerp(OldAlbedo, float4(Albedo, 0), BlendFactor);
-		NormalTexture[PixelIndex] = lerp(OldNormal, float4(Normal, 0), BlendFactor);
-	}
-};
-
 FPathState CreatePathState(uint2 DispatchIdx, uint2 DispatchDim)
 {
 	FPathState Output = (FPathState)0;
@@ -617,6 +623,7 @@
 	// path state variables (these cary information between bounces)
 	Output.PathThroughput = 1.0;
 	Output.PathRoughness = 0;
+	Output.SigmaT = 0; // TODO: figure out the right start value for camera rays
 
 	return Output;
 }
@@ -638,7 +645,6 @@
 	return ShadingNormal;
 }
 
-
 bool PathTracingKernel(inout FPathState PathState, int Bounce)
 {
 	// This array will hold a CDF for light picking
@@ -660,19 +666,11 @@
 	uint NumVisibleLights = bIsCameraRay ? SceneVisibleLightCount : 0;
 
 	FPathTracingPayload Payload = TraceTransparentRay(
-		PathState.Ray,
-		PathState.RayCone,
+		PathState,
 		bIsCameraRay,
 		bIsLastBounce,
 		bIncludeEmissive,
-		PathState.PixelIndex,
-		NumVisibleLights,
-		PathState.PathRoughness,
-		PathState.RandSequence,
-		PathState.PathThroughput,
-		PathState.Radiance,
-		PathState.Albedo,
-		PathState.Normal);
+		NumVisibleLights);
 
 	if (Payload.IsMiss())
 	{
@@ -851,6 +849,25 @@
 	// enlarge roughness based on the chosen lobe roughness
 	PathState.PathRoughness = max(PathState.PathRoughness, MaterialSample.Roughness);
 
+	// update the current extinction if we are crossing a boundary on glass
+	// summing the local extinction gives a rudimentary way of dealing with overlapping regions
+	// long term we will probably want a stack with priorities
+	if (MaterialSample.PositionBiasSign < 0 && Payload.IsMaterialSolidGlass())
+	{
+		float3 LocalSigmaT = -log(max(Payload.BaseColor, 1e-8f)) * AbsorptionScale;
+		if (Payload.IsFrontFace())
+		{
+			// entering
+			PathState.SigmaT += LocalSigmaT;
+		}
+		else
+		{
+			// exiting
+			PathState.SigmaT -= LocalSigmaT;
+			PathState.SigmaT = max(PathState.SigmaT, 0);
+		}
+	}
+
 	// If we are using Material sampling for lights
 	if (MISMode != 1 && bDoLightLoop)
 	{
diff --git a/Engine/Shaders/Private/PathTracing/PathTracingMaterialHitShader.usf b/Engine/Shaders/Private/PathTracing/PathTracingMaterialHitShader.usf
--- a/Engine/Shaders/Private/PathTracing/PathTracingMaterialHitShader.usf
+++ b/Engine/Shaders/Private/PathTracing/PathTracingMaterialHitShader.usf
@@ -309,12 +309,12 @@ RAY_TRACING_ENTRY_ANY_HIT(PathTracingMaterialAHS,
 	FMaterialPixelParameters MaterialParameters = GetMaterialPixelParameters(Interpolants, SvPosition);
 	FPixelMaterialInputs PixelMaterialInputs;
 
+	const bool bIsFrontFace = HitKind() == HIT_KIND_TRIANGLE_FRONT_FACE;
+
 	{
 		float4 ScreenPosition = SvPositionToResolvedScreenPosition(SvPosition);
 		float3 TranslatedWorldPosition = WorldRayOrigin() + WorldRayDirection() * RayTCurrent() + ResolvedView.PreViewTranslation.xyz;
 
-		bool bIsFrontFace = HitKind() == HIT_KIND_TRIANGLE_FRONT_FACE;
-
 		MaterialParameters.CameraVector = -WorldRayDirection();
 
 		// #dxr_todo: UE-72130 support world position offset 
@@ -415,6 +415,20 @@ RAY_TRACING_ENTRY_ANY_HIT(PathTracingMaterialAHS,
 		float VoN = abs(dot(WorldRayDirection(), N));
 		float Fr = FresnelReflectance(VoN, Eta);
 		Transparency *= pow(1 - Fr, 2) * (1 - Roughness * Roughness) * saturate(PathRoughness - Roughness);
+		// Now track the optical thickness so that we can account for Beer's law along the shadow ray
+		// TODO: to support lights inside glass we would need to track an extra offset, but this is hopefully an uncommon scenario
+		float Distance = RayTCurrent();
+		float3 LocalSigmaT = -log(max(GetMaterialBaseColor(PixelMaterialInputs), 1e-8f));
+		if (bIsFrontFace)
+		{
+			// remove contribution from the ray origin to the current hit (assuming the backside will be hit)
+			PackedPayload.SetTau(PackedPayload.GetTau() - LocalSigmaT * Distance);
+		}
+		else
+		{
+			// add contribution from ray origin to current hit (backside) - excess portion not covered by the ray will be removed by the front face hit
+			PackedPayload.SetTau(PackedPayload.GetTau() + LocalSigmaT * Distance);
+		}
 	}
 #endif
 #else
diff --git a/Engine/Shaders/Shared/PathTracingDefinitions.h b/Engine/Shaders/Shared/PathTracingDefinitions.h
--- a/Engine/Shaders/Shared/PathTracingDefinitions.h
+++ b/Engine/Shaders/Shared/PathTracingDefinitions.h
@@ -7,7 +7,7 @@
 #define PATHTRACER_COHERENT_TILE_SHIFT		3
 #define PATHTRACER_COHERENT_TILE_SIZE		(1 << PATHTRACER_COHERENT_TILE_SHIFT)
 
-#define PATHTRACER_PATHSTATE_NUM_WORDS		20
+#define PATHTRACER_PATHSTATE_NUM_WORDS		21
 
 // Constants for 'SamplerType'
 #define PATHTRACER_SAMPLER_DEFAULT			0
diff --git a/Engine/Source/Runtime/Renderer/Private/PathTracing.cpp b/Engine/Source/Runtime/Renderer/Private/PathTracing.cpp
--- a/Engine/Source/Runtime/Renderer/Private/PathTracing.cpp
+++ b/Engine/Source/Runtime/Renderer/Private/PathTracing.cpp
@@ -55,6 +55,14 @@ TAutoConsoleVariable<float> CVarPathTracingFilterWidth(
 	ECVF_RenderThreadSafe
 );
 
+TAutoConsoleVariable<float> CVarPathTracingAbsorptionScale(
+	TEXT("r.PathTracing.AbsorptionScale"),
+	0.01,
+	TEXT("Sets the inverse distance at which BaseColor is reached for transmittance in refractive glass (default = 1/100 units)\n")
+	TEXT("Setting this value to 0 will disable absorption handling for refractive glass\n"),
+	ECVF_RenderThreadSafe
+);
+
 TAutoConsoleVariable<int32> CVarPathTracingMISMode(
 	TEXT("r.PathTracing.MISMode"),
 	2,
@@ -225,6 +233,7 @@ BEGIN_SHADER_PARAMETER_STRUCT(FPathTracingData, )
 	SHADER_PARAMETER(float, MaxPathIntensity)
 	SHADER_PARAMETER(float, MaxNormalBias)
 	SHADER_PARAMETER(float, FilterWidth)
+	SHADER_PARAMETER(float, AbsorptionScale)
 END_SHADER_PARAMETER_STRUCT()
 
 
@@ -257,6 +266,7 @@ struct FPathTracingConfig
 			PathTracingData.VisualizeLightGrid != Other.PathTracingData.VisualizeLightGrid ||
 			PathTracingData.MaxPathIntensity != Other.PathTracingData.MaxPathIntensity ||
 			PathTracingData.FilterWidth != Other.PathTracingData.FilterWidth ||
+			PathTracingData.AbsorptionScale != Other.PathTracingData.AbsorptionScale ||
 			ViewRect != Other.ViewRect ||
 			LightShowFlags != Other.LightShowFlags ||
 			LightGridResolution != Other.LightGridResolution ||
@@ -340,6 +350,7 @@ static void PrepareShaderArgs(const FViewInfo& View, FPathTracingData& PathTraci
 		FilterWidth = View.FinalPostProcessSettings.PathTracingFilterWidth;
 	}
 	PathTracingData.FilterWidth = FilterWidth;
+	PathTracingData.AbsorptionScale = CVarPathTracingAbsorptionScale.GetValueOnRenderThread();
 }
 
 static bool ShouldCompilePathTracingShadersForProject(EShaderPlatform ShaderPlatform)
