Fixed incorrect PackSnorm2x16

diff --git a/Engine/Shaders/Private/Common.ush b/Engine/Shaders/Private/Common.ush
--- a/Engine/Shaders/Private/Common.ush
+++ b/Engine/Shaders/Private/Common.ush
@@ -825,7 +825,7 @@ uint PackUnorm2x16(float2 v)
 
 uint PackSnorm2x16(float2 v)
 {
-	uint2 sv = uint2(round(clamp(v, -1.0, 1.0) * 32767.0));
+	uint2 sv = uint2(round(clamp(v, -1.0, 1.0) * 32767.0) + 32767.0);
 	return (sv.x | (sv.y << 16u));
 }
 
@@ -840,8 +840,8 @@ float2 UnpackUnorm2x16(uint p)
 float2 UnpackSnorm2x16(uint p)
 {
 	float2 Ret;
-	Ret.x = clamp((p & 0xffff) * rcp(32767.0f), -1.0, 1.0);
-	Ret.y = clamp((p >> 16u) * rcp(32767.0f), -1.0, 1.0);
+	Ret.x = clamp((float(p & 0xffff) - 32767.0f) * rcp(32767.0f), -1.0, 1.0);
+	Ret.y = clamp((float(p >> 16u) - 32767.0f) * rcp(32767.0f), -1.0, 1.0);
 	return Ret;
 }
 
