Implement alpha channel support for the path tracer

To maintain consistency with the rasterizer implementation, keep track of background visibility (1.0 where there is background, 0.0 where there are solid objects). MRQ will flip this quantity before outputing to disk.

diff --git a/Engine/Plugins/Experimental/OpenImageDenoise/Source/OpenImageDenoise/Private/Plugin.cpp b/Engine/Plugins/Experimental/OpenImageDenoise/Source/OpenImageDenoise/Private/Plugin.cpp
--- a/Engine/Plugins/Experimental/OpenImageDenoise/Source/OpenImageDenoise/Private/Plugin.cpp
+++ b/Engine/Plugins/Experimental/OpenImageDenoise/Source/OpenImageDenoise/Private/Plugin.cpp
@@ -104,6 +104,16 @@ static void Denoise(FRHICommandListImmediate& RHICmdList, FRHITexture2D* ColorTe
 
 	OIDNFilter.execute();
 
+	// copy alpha channel (TODO: find a way to denoise it as well?)
+	for (int Y = 0, OutIndex = 0, Index = 0; Y < Size.Y; Y++)
+	{
+		for (int X = 0; X < Size.X; X++, OutIndex++, Index++)
+		{
+			DestBuffer[OutIndex].A = RawPixels[Index].A;
+		}
+		OutIndex += DestStride / sizeof(FLinearColor) - Size.X;
+	}
+
 	RHICmdList.UnlockTexture2D(OutputTex, 0, false);
 
 #if WITH_EDITOR
diff --git a/Engine/Shaders/Private/PathTracing/PathTracing.usf b/Engine/Shaders/Private/PathTracing/PathTracing.usf
--- a/Engine/Shaders/Private/PathTracing/PathTracing.usf
+++ b/Engine/Shaders/Private/PathTracing/PathTracing.usf
@@ -35,10 +35,12 @@ FPathState LoadPathStateData(uint Index)
 	Output.Radiance.y = asfloat(PathStateData[PATH_STATE_STRIDE * Index + 4]);
 	Output.Radiance.z = asfloat(PathStateData[PATH_STATE_STRIDE * Index + 5]);
 
+	Output.BackgroundVisibility = asfloat(PathStateData[PATH_STATE_STRIDE * Index + 6]);
+
 	uint3 PackedAlbedoNormal;
-	PackedAlbedoNormal.x = PathStateData[PATH_STATE_STRIDE * Index + 6];
-	PackedAlbedoNormal.y = PathStateData[PATH_STATE_STRIDE * Index + 7];
-	PackedAlbedoNormal.z = PathStateData[PATH_STATE_STRIDE * Index + 8];
+	PackedAlbedoNormal.x = PathStateData[PATH_STATE_STRIDE * Index + 7];
+	PackedAlbedoNormal.y = PathStateData[PATH_STATE_STRIDE * Index + 8];
+	PackedAlbedoNormal.z = PathStateData[PATH_STATE_STRIDE * Index + 9];
 	Output.Albedo.x = f16tof32(PackedAlbedoNormal.x);
 	Output.Albedo.y = f16tof32(PackedAlbedoNormal.x >> 16);
 	Output.Albedo.z = f16tof32(PackedAlbedoNormal.y);
@@ -46,26 +48,26 @@ FPathState LoadPathStateData(uint Index)
 	Output.Normal.y = f16tof32(PackedAlbedoNormal.z);
 	Output.Normal.z = f16tof32(PackedAlbedoNormal.z >> 16);
 
-	Output.Ray.Origin.x = asfloat(PathStateData[PATH_STATE_STRIDE * Index + 9]);
-	Output.Ray.Origin.y = asfloat(PathStateData[PATH_STATE_STRIDE * Index + 10]);
-	Output.Ray.Origin.z = asfloat(PathStateData[PATH_STATE_STRIDE * Index + 11]);
-	Output.Ray.Direction.x = asfloat(PathStateData[PATH_STATE_STRIDE * Index + 12]);
-	Output.Ray.Direction.y = asfloat(PathStateData[PATH_STATE_STRIDE * Index + 13]);
-	Output.Ray.Direction.z = asfloat(PathStateData[PATH_STATE_STRIDE * Index + 14]);
+	Output.Ray.Origin.x = asfloat(PathStateData[PATH_STATE_STRIDE * Index + 10]);
+	Output.Ray.Origin.y = asfloat(PathStateData[PATH_STATE_STRIDE * Index + 11]);
+	Output.Ray.Origin.z = asfloat(PathStateData[PATH_STATE_STRIDE * Index + 12]);
+	Output.Ray.Direction.x = asfloat(PathStateData[PATH_STATE_STRIDE * Index + 13]);
+	Output.Ray.Direction.y = asfloat(PathStateData[PATH_STATE_STRIDE * Index + 14]);
+	Output.Ray.Direction.z = asfloat(PathStateData[PATH_STATE_STRIDE * Index + 15]);
 	Output.Ray.TMin = 0;
 	Output.Ray.TMax = POSITIVE_INFINITY;
 
-	uint PackedRayCone = PathStateData[PATH_STATE_STRIDE * Index + 15];
+	uint PackedRayCone = PathStateData[PATH_STATE_STRIDE * Index + 16];
 	Output.RayCone.Width = f16tof32(PackedRayCone);
 	Output.RayCone.SpreadAngle = f16tof32(PackedRayCone >> 16);
 
-	Output.PathThroughput.x = asfloat(PathStateData[PATH_STATE_STRIDE * Index + 16]);
-	Output.PathThroughput.y = asfloat(PathStateData[PATH_STATE_STRIDE * Index + 17]);
-	Output.PathThroughput.z = asfloat(PathStateData[PATH_STATE_STRIDE * Index + 18]);
+	Output.PathThroughput.x = asfloat(PathStateData[PATH_STATE_STRIDE * Index + 17]);
+	Output.PathThroughput.y = asfloat(PathStateData[PATH_STATE_STRIDE * Index + 18]);
+	Output.PathThroughput.z = asfloat(PathStateData[PATH_STATE_STRIDE * Index + 19]);
 
 	uint2 PackedRoughnessSigma;
-	PackedRoughnessSigma.x = PathStateData[PATH_STATE_STRIDE * Index + 19];
-	PackedRoughnessSigma.y = PathStateData[PATH_STATE_STRIDE * Index + 20];
+	PackedRoughnessSigma.x = PathStateData[PATH_STATE_STRIDE * Index + 20];
+	PackedRoughnessSigma.y = PathStateData[PATH_STATE_STRIDE * Index + 21];
 	Output.PathRoughness = f16tof32(PackedRoughnessSigma.x);
 	Output.SigmaT.x = f16tof32(PackedRoughnessSigma.x >> 16);
 	Output.SigmaT.y = f16tof32(PackedRoughnessSigma.y);
@@ -85,25 +87,27 @@ void StorePathStateData(FPathState PathState, uint Index)
 	PathStateData[PATH_STATE_STRIDE * Index + 4] = asuint(PathState.Radiance.y);
 	PathStateData[PATH_STATE_STRIDE * Index + 5] = asuint(PathState.Radiance.z);
 
-	PathStateData[PATH_STATE_STRIDE * Index + 6] = f32tof16(PathState.Albedo.x) | (f32tof16(PathState.Albedo.y) << 16);
-	PathStateData[PATH_STATE_STRIDE * Index + 7] = f32tof16(PathState.Albedo.z) | (f32tof16(PathState.Normal.x) << 16);
-	PathStateData[PATH_STATE_STRIDE * Index + 8] = f32tof16(PathState.Normal.y) | (f32tof16(PathState.Normal.z) << 16);
+	PathStateData[PATH_STATE_STRIDE * Index + 6] = asuint(PathState.BackgroundVisibility);
+
+	PathStateData[PATH_STATE_STRIDE * Index + 7] = f32tof16(PathState.Albedo.x) | (f32tof16(PathState.Albedo.y) << 16);
+	PathStateData[PATH_STATE_STRIDE * Index + 8] = f32tof16(PathState.Albedo.z) | (f32tof16(PathState.Normal.x) << 16);
+	PathStateData[PATH_STATE_STRIDE * Index + 9] = f32tof16(PathState.Normal.y) | (f32tof16(PathState.Normal.z) << 16);
 
-	PathStateData[PATH_STATE_STRIDE * Index + 9]  = asuint(PathState.Ray.Origin.x);
-	PathStateData[PATH_STATE_STRIDE * Index + 10] = asuint(PathState.Ray.Origin.y);
-	PathStateData[PATH_STATE_STRIDE * Index + 11] = asuint(PathState.Ray.Origin.z);
-	PathStateData[PATH_STATE_STRIDE * Index + 12] = asuint(PathState.Ray.Direction.x);
-	PathStateData[PATH_STATE_STRIDE * Index + 13] = asuint(PathState.Ray.Direction.y);
-	PathStateData[PATH_STATE_STRIDE * Index + 14] = asuint(PathState.Ray.Direction.z);
+	PathStateData[PATH_STATE_STRIDE * Index + 10] = asuint(PathState.Ray.Origin.x);
+	PathStateData[PATH_STATE_STRIDE * Index + 11] = asuint(PathState.Ray.Origin.y);
+	PathStateData[PATH_STATE_STRIDE * Index + 12] = asuint(PathState.Ray.Origin.z);
+	PathStateData[PATH_STATE_STRIDE * Index + 13] = asuint(PathState.Ray.Direction.x);
+	PathStateData[PATH_STATE_STRIDE * Index + 14] = asuint(PathState.Ray.Direction.y);
+	PathStateData[PATH_STATE_STRIDE * Index + 15] = asuint(PathState.Ray.Direction.z);
 
-	PathStateData[PATH_STATE_STRIDE * Index + 15] = f32tof16(PathState.RayCone.Width) | (f32tof16(PathState.RayCone.SpreadAngle) << 16);
+	PathStateData[PATH_STATE_STRIDE * Index + 16] = f32tof16(PathState.RayCone.Width) | (f32tof16(PathState.RayCone.SpreadAngle) << 16);
 
-	PathStateData[PATH_STATE_STRIDE * Index + 16] = asuint(PathState.PathThroughput.x);
-	PathStateData[PATH_STATE_STRIDE * Index + 17] = asuint(PathState.PathThroughput.y);
-	PathStateData[PATH_STATE_STRIDE * Index + 18] = asuint(PathState.PathThroughput.z);
+	PathStateData[PATH_STATE_STRIDE * Index + 17] = asuint(PathState.PathThroughput.x);
+	PathStateData[PATH_STATE_STRIDE * Index + 18] = asuint(PathState.PathThroughput.y);
+	PathStateData[PATH_STATE_STRIDE * Index + 19] = asuint(PathState.PathThroughput.z);
 
-	PathStateData[PATH_STATE_STRIDE * Index + 19] = f32tof16(PathState.PathRoughness) | (f32tof16(PathState.SigmaT.x) << 16);
-	PathStateData[PATH_STATE_STRIDE * Index + 20] = f32tof16(PathState.SigmaT.y) | (f32tof16(PathState.SigmaT.z) << 16);
+	PathStateData[PATH_STATE_STRIDE * Index + 20] = f32tof16(PathState.PathRoughness) | (f32tof16(PathState.SigmaT.x) << 16);
+	PathStateData[PATH_STATE_STRIDE * Index + 21] = f32tof16(PathState.SigmaT.y) | (f32tof16(PathState.SigmaT.z) << 16);
 }
 
 RAY_TRACING_ENTRY_RAYGEN(PathTracingMainRG)
diff --git a/Engine/Shaders/Private/PathTracing/PathTracingCompositingPixelShader.usf b/Engine/Shaders/Private/PathTracing/PathTracingCompositingPixelShader.usf
--- a/Engine/Shaders/Private/PathTracing/PathTracingCompositingPixelShader.usf
+++ b/Engine/Shaders/Private/PathTracing/PathTracingCompositingPixelShader.usf
@@ -45,7 +45,7 @@ void CompositeMain(
 
 	float4 Radiance = RadianceTexture.Load(TexCoord);
 
-	OutColor = float4(Radiance.rgb * View.PreExposure, 1.0);
+	OutColor = float4(Radiance.rgb * View.PreExposure, Radiance.a);
 
 	if (ProgressDisplayEnabled && (Iteration + 1) < MaxSamples)
 	{
@@ -100,13 +100,7 @@ void CompositeMain(
     // split-screen debugging aids
 	if (UV.x < 0.5)
 	{
-		uint NumSamples = Iteration + 1;
-		float OutVariance = Radiance.a;
-		float Lum = Luminance(Radiance.rgb); // mean
-		float StdErr = NumSamples > 0 ? sqrt(OutVariance / NumSamples) : 0; // stderr
-		//OutColor = Tonemap(Lum + 3 * StdErr) - Tonemap(max(Lum - 3 * StdErr, 0.0)); // error estimator
-		OutColor = StdErr * View.PreExposure;
-		//OutColor = StdErr / (Lum + 0.0001) * View.PreExposure;
+		OutColor = 1 - Radiance.a; // visualize alpha channel
 	}
 #endif
 }
diff --git a/Engine/Shaders/Private/PathTracing/PathTracingCore.ush b/Engine/Shaders/Private/PathTracing/PathTracingCore.ush
--- a/Engine/Shaders/Private/PathTracing/PathTracingCore.ush
+++ b/Engine/Shaders/Private/PathTracing/PathTracingCore.ush
@@ -57,6 +57,7 @@ struct FPathState
 	uint2 PixelIndex;               //  4 bytes
 	RandomSequence RandSequence;    //  8 bytes
 	float3 Radiance;                // 12 bytes
+	float  BackgroundVisibility;    //  4 bytes
 	float3 Albedo;                  //  6 bytes // half precision
 	float3 Normal;                  //  6 bytes // half precision
 	RayDesc Ray;                    // 12+12 = 24 bytes (TMin/TMax do not need to be stored)
@@ -64,27 +65,18 @@ struct FPathState
 	float3 PathThroughput;          // 12 bytes
 	float PathRoughness;            //  2 bytes // half precision
 	float3 SigmaT; // Extinction    //  6 bytes // half precision 
-									// 84 bytes total
-									// 21 words - sync with PATHTRACER_PATHSTATE_NUM_WORDS
+									// 88 bytes total
+									// 22 words - sync with PATHTRACER_PATHSTATE_NUM_WORDS
 
 	void WritePixel(float BlendFactor)
 	{
 		// Avoid reading the old pixel on the first sample on the off-chance there is a NaN/Inf pixel ...
-		float4 OldPixel = Iteration > 0 ? RadianceTexture[PixelIndex] : 0;
-		float4 OldAlbedo = Iteration > 0 ? AlbedoTexture[PixelIndex] : 0;
-		float4 OldNormal = Iteration > 0 ? NormalTexture[PixelIndex] : 0;
-		float3 OldRadiance = OldPixel.rgb;
-		float OldVariance = OldPixel.a;
+		float4 OldPixel  = Iteration > 0 ? RadianceTexture[PixelIndex] : 0;
+		float4 OldAlbedo = Iteration > 0 ? AlbedoTexture[PixelIndex]   : 0;
+		float4 OldNormal = Iteration > 0 ? NormalTexture[PixelIndex]   : 0;
 
-		float3 NewRadiance = lerp(OldRadiance, Radiance, BlendFactor);
-
-		float RadianceY = Luminance(Radiance);
-
-		// https://en.wikipedia.org/wiki/Algorithms_for_calculating_variance#Welford's_online_algorithm
-		float DeviationSquared = (RadianceY - Luminance(NewRadiance)) * (RadianceY - Luminance(OldRadiance));
-		float NewVariance = lerp(OldVariance, DeviationSquared, BlendFactor);
-
-		RadianceTexture[PixelIndex] = float4(NewRadiance, NewVariance);
+		const float4 PixelValue = float4(Radiance, saturate(BackgroundVisibility));
+		RadianceTexture[PixelIndex] = lerp(OldPixel, PixelValue, BlendFactor);
 		AlbedoTexture[PixelIndex] = lerp(OldAlbedo, float4(Albedo, 0), BlendFactor);
 		NormalTexture[PixelIndex] = lerp(OldNormal, float4(Normal, 0), BlendFactor);
 	}
@@ -148,7 +140,12 @@ FPathTracingPayload TraceTransparentRay(inout FPathState PathState, bool bIsCame
 
 		if (PackedPayload.IsMiss())
 		{
-			// Ray didn't hit any real geometry, so nothing left to do
+			// Ray didn't hit any real geometry
+			// Accumulate a hit against the sky (for camera rays) and exit
+			if (bIsCameraRay)
+			{
+				PathState.BackgroundVisibility += Luminance(PathState.PathThroughput);
+			}
 			break;
 		}
 
@@ -300,7 +297,7 @@ FProbeResult TraceProbeRay(RayDesc Ray, uint2 PixelCoord)
 	FPathTracingPayload Payload = UnpackPathTracingPayload(PackedPayload, Ray);
 	FProbeResult Result;
 	Result.HitT = PackedPayload.HitT;
-	Result.WorldNormal = Payload.WorldNormal;
+	Result.WorldNormal = Payload.WorldSmoothNormal;
 	Result.WorldGeoNormal = Payload.WorldGeoNormal;
 	Result.FrontFace = Payload.IsFrontFace();
 	return Result;
@@ -512,11 +509,13 @@ bool ProcessSubsurfaceRandomWalk(inout FPathTracingPayload Payload, inout float3
 				// we hit the boundary! overwrite the current payload and exit the walk with a diffuse scattering event
 				Payload.WorldPos = Ray.Origin + ProbeResult.HitT * Ray.Direction;
 				Payload.WorldNormal = ProbeResult.WorldNormal;
+				Payload.WorldSmoothNormal = ProbeResult.WorldNormal;
 				Payload.WorldGeoNormal = ProbeResult.WorldGeoNormal;
 				// make sure normal is pointed outward from the object so we capture illumination from the exterior
 				if (ProbeResult.FrontFace != Payload.IsFrontFace())
 				{
 					Payload.WorldNormal *= -1.0;
+					Payload.WorldSmoothNormal *= -1.0;
 					Payload.WorldGeoNormal *= -1.0;
 				}
 				Payload.ShadingModelID = SHADINGMODELID_NUM; // invalid value so that we get diffuse shading
@@ -682,6 +681,20 @@ bool PathTracingKernel(inout FPathState PathState, int Bounce)
 	Payload.WorldNormal = AdjustShadingNormal(Payload.WorldNormal, Payload.WorldGeoNormal, PathState.Ray.Direction);
 	Payload.WorldSmoothNormal = AdjustShadingNormal(Payload.WorldSmoothNormal, Payload.WorldGeoNormal, PathState.Ray.Direction);
 
+#if 0
+	PathState.Radiance = 0.5 * Payload.WorldNormal + 0.5;
+	return false;
+#endif
+
+#if 0
+	if (Payload.ShadingModelID == SHADINGMODELID_EYE)
+	{
+		//Payload.WorldNormal = Payload.CustomVector;
+		Payload.ShadingModelID = SHADINGMODELID_DEFAULT_LIT;
+		Payload.Anisotropy = 0;
+	}
+#endif
+
 	// As soon as the path is blurry enough, we can get away with diffuse sampling only
 	const bool bSimplifySSS = PathState.PathRoughness >= 0.15;
 
diff --git a/Engine/Shaders/Shared/PathTracingDefinitions.h b/Engine/Shaders/Shared/PathTracingDefinitions.h
--- a/Engine/Shaders/Shared/PathTracingDefinitions.h
+++ b/Engine/Shaders/Shared/PathTracingDefinitions.h
@@ -7,7 +7,7 @@
 #define PATHTRACER_COHERENT_TILE_SHIFT		3
 #define PATHTRACER_COHERENT_TILE_SIZE		(1 << PATHTRACER_COHERENT_TILE_SHIFT)
 
-#define PATHTRACER_PATHSTATE_NUM_WORDS		21
+#define PATHTRACER_PATHSTATE_NUM_WORDS		22
 
 // Constants for 'SamplerType'
 #define PATHTRACER_SAMPLER_DEFAULT			0
