Fix the bug that reflection capture is not working on single layer water material.

diff --git a/Engine/Shaders/Private/SingleLayerWaterComposite.usf b/Engine/Shaders/Private/SingleLayerWaterComposite.usf
--- a/Engine/Shaders/Private/SingleLayerWaterComposite.usf
+++ b/Engine/Shaders/Private/SingleLayerWaterComposite.usf
@@ -147,7 +147,7 @@ SingleLayerWaterCompositeOutput SingleLayerWaterComposite(float2 BufferUV, float
 			float IndirectSpecularOcclusion = 1.0f;
 			float3 ExtraIndirectSpecular = 0;
 			uint CaptureDataStartIndex = 0;
-			uint NumCulledReflectionCaptures = 0;
+			uint NumCulledReflectionCaptures = ForwardLightData.NumReflectionCaptures;
 			//
 			float3 N = GBuffer.WorldNormal;
 			float3 V = -CameraToPixel;
