Fixes FFT Bloom going yellow at lower resolution

diff --git a/Engine/Shaders/Private/Bloom/BloomClampKernel.usf b/Engine/Shaders/Private/Bloom/BloomClampKernel.usf
--- a/Engine/Shaders/Private/Bloom/BloomClampKernel.usf
+++ b/Engine/Shaders/Private/Bloom/BloomClampKernel.usf
@@ -28,5 +28,11 @@ void MainCS(
 
 	float4 KernelTexel = KernelSpatialTexture[PixelCoord];
 
-	ClampedKernelSpatialOutput[PixelCoord] = clamp(KernelTexel, 0.0, OriginalKernelInfo.MaxScatterDispersionEnergy);
+#if POST_PROCESS_ALPHA > 0
+	float4 ClampedKernel = clamp(KernelTexel, 0.0, OriginalKernelInfo.MaxScatterDispersionEnergy);
+#else
+	float4 ClampedKernel = clamp(KernelTexel.rgbb, 0.0, OriginalKernelInfo.MaxScatterDispersionEnergy.rgbb);
+#endif
+
+	ClampedKernelSpatialOutput[PixelCoord] = ClampedKernel;
 }
