Add path tracer support for the eye shading model

The eye shading model includes some custom caustic terms that are not fully physically based, but I am emulating them with a custom diffuse lobe that replicates the logic of the EyeBxDF method.

Also minor cleanups to ShadingModelsMaterial.ush to remove redundant set operations on the gbuffer

diff --git a/Engine/Shaders/Private/PathTracing/Material/PathTracingEye.ush b/Engine/Shaders/Private/PathTracing/Material/PathTracingEye.ush
new file mode 100644
--- /dev/null
+++ b/Engine/Shaders/Private/PathTracing/Material/PathTracingEye.ush
@@ -0,0 +1,154 @@
+// Copyright Epic Games, Inc. All Rights Reserved.
+
+/*=============================================================================================
+	PathTracingEye.usf: Path tracing BRDF model for the eye material
+===============================================================================================*/
+
+#pragma once
+
+#include "PathTracingMaterialCommon.ush"
+#include "PathTracingFresnel.ush"
+#include "PathTracingEnergyConservation.ush"
+
+FMaterialEval Eye_EvalMaterial(
+	float3 IncomingDirection,
+	float3 OutgoingDirection,
+	FPathTracingPayload Payload
+)
+{
+	FMaterialEval Result = NullMaterialEval();
+
+	const float3 V_World = -IncomingDirection;
+	const float3 L_World = OutgoingDirection;
+	const float3 N_World = Payload.WorldNormal;
+
+	float2 Alpha = 0;
+	const float3x3 Basis = GetGGXBasis(Payload.Roughness, Payload.WorldNormal, Alpha);
+
+	// move vectors into right shading frame
+	const float3 V = mul(Basis, V_World);
+	const float3 L = mul(Basis, L_World);
+	const float3 H = normalize(V + L);
+
+	const float NoV = saturate(V.z);
+	const float NoL = saturate(L.z);
+	const float VoH = saturate(dot(V, H));
+
+	const FBxDFEnergyTerms Spec = ComputeGGXSpecEnergyTerms(Payload.Roughness, NoV, Payload.SpecularColor);
+
+	// Probability of picking diffuse lobe vs. specular lobe
+	const float LobeSelectionPdf = LobeSelectionProb(Payload.DiffuseColor * (1 - Spec.E), Spec.E);
+
+
+	const float IrisMask = Payload.Anisotropy;
+	const float3 IrisNormal = Payload.WorldTangent;
+	const float3 CausticNormal = Payload.CustomVector;
+	const float IrisNoL = saturate(dot(IrisNormal, L_World));
+	const float Power = lerp(12, 1, IrisNoL);
+	const float Caustic = 0.8 + 0.2 * (Power + 1) * pow(saturate(dot(CausticNormal, L_World)), Power);
+	const float Iris = IrisNoL * Caustic;
+	const float Sclera = NoL;
+
+	const float EyeDiffuseTweak = 2 * lerp(Sclera, Iris, IrisMask);
+
+
+	// Diffuse Lobe
+	AddLobeWithMIS(Result.Weight, Result.Pdf, (1 - Spec.E) * Payload.DiffuseColor * EyeDiffuseTweak, 1 / (2 * PI), LobeSelectionPdf);
+
+	// Specular lobe
+	const float2 GGXResult = GGXEvalReflection(L, V, H, Alpha);
+	const float3 F = F_Schlick(Payload.SpecularColor, VoH);
+	const float3 SpecWeight = F * GGXResult.x * Spec.W;
+	const float SpecPdf = GGXResult.y;
+
+	AddLobeWithMIS(Result.Weight, Result.Pdf, SpecWeight, SpecPdf, 1.0 - LobeSelectionPdf);
+
+	return Result;
+}
+
+FMaterialSample Eye_SampleMaterial(
+	float3 RayDirection,
+	FPathTracingPayload Payload,
+	float4 RandSample
+)
+{
+	float3 N = Payload.WorldNormal;
+	float3 V = -RayDirection;
+
+	float2 Alpha;
+	float3x3 Basis = GetGGXBasis(Payload.Roughness, Payload.Anisotropy, Payload.WorldNormal, Payload.WorldTangent, Alpha);
+
+	V = mul(Basis, V);
+
+	const float NoV = saturate(V.z);
+	const FBxDFEnergyTerms Spec = ComputeGGXSpecEnergyTerms(Payload.Roughness, NoV, Payload.SpecularColor);
+
+	// Probability of picking diffuse lobe vs. specular lobe
+	const float LobeSelectionPdf = LobeSelectionProb(Payload.DiffuseColor * (1 - Spec.E), Spec.E);
+
+	// Randomly choose to sample diffuse or specular
+	float3 L = 0, H = 0;
+	float OutRoughness = 0;
+	if (RandSample.x < LobeSelectionPdf)
+	{
+		RandSample.x /= LobeSelectionPdf;
+
+		// Lambert
+		// TODO: evaluate CosineSampleHemisphereConcentric
+		float4 Result = UniformSampleHemisphere(RandSample.xy);
+
+		L = Result.xyz;
+		H = normalize(L + V);
+		OutRoughness = 1.0;
+	}
+	else
+	{
+		RandSample.x -= LobeSelectionPdf;
+		RandSample.x /= (1.0 - LobeSelectionPdf);
+
+		// #dxr_todo: Evaluate UniformSampleDiskConcentric
+		H = ImportanceSampleVisibleGGX_aniso(UniformSampleDisk(RandSample.xy), Alpha, V).xyz;
+
+		L = reflect(-V, H);
+		if (L.z <= 0)
+		{
+			// invalid output direction, exit early
+			return NullMaterialSample();
+		}
+		OutRoughness = Payload.Roughness;
+	}
+
+	// With a valid direction in hand -- now evaluate the BxDF (taking advantage of already computed terms)
+
+	const float3 L_World = normalize(mul(L, Basis));
+
+	const float NoL = saturate(L.z);
+	const float VoH = saturate(dot(V, H));
+
+	float3 OutWeight = 0;
+	float OutPdf = 0;
+
+	const float IrisMask = Payload.Anisotropy;
+	const float3 IrisNormal = Payload.WorldTangent;
+	const float3 CausticNormal = Payload.CustomVector;
+	const float IrisNoL = saturate(dot(IrisNormal, L_World));
+	const float Power = lerp(12, 1, IrisNoL);
+	const float Caustic = 0.8 + 0.2 * (Power + 1) * pow(saturate(dot(CausticNormal, L_World)), Power);
+	const float Iris = IrisNoL * Caustic;
+	const float Sclera = NoL;
+
+	const float EyeDiffuseTweak = 2 * lerp(Sclera, Iris, IrisMask);
+
+	// Diffuse Lobe
+	AddLobeWithMIS(OutWeight, OutPdf, (1 - Spec.E) * Payload.DiffuseColor * EyeDiffuseTweak, 1 / (2 * PI), LobeSelectionPdf);
+
+	// Specular lobe
+	const float2 GGXResult = GGXEvalReflection(L, V, H, Alpha);
+	const float3 F = F_Schlick(Payload.SpecularColor, VoH);
+	const float3 SpecWeight = F * GGXResult.x * Spec.W;
+	const float SpecPdf = GGXResult.y;
+	AddLobeWithMIS(OutWeight, OutPdf, SpecWeight, SpecPdf, 1.0 - LobeSelectionPdf);
+
+	// transform to world space
+	return CreateMaterialSample(L_World, OutWeight, OutPdf, 1.0, OutRoughness);
+}
diff --git a/Engine/Shaders/Private/PathTracing/Material/PathTracingMaterialSampling.ush b/Engine/Shaders/Private/PathTracing/Material/PathTracingMaterialSampling.ush
--- a/Engine/Shaders/Private/PathTracing/Material/PathTracingMaterialSampling.ush
+++ b/Engine/Shaders/Private/PathTracing/Material/PathTracingMaterialSampling.ush
@@ -20,6 +20,7 @@ PathTracingMaterialSampling.usf: Material sampling functions
 #include "PathTracingTwoSidedFoliage.ush"
 #include "PathTracingCloth.ush"
 #include "PathTracingSubsurfaceProfile.ush"
+#include "PathTracingEye.ush"
 #include "../../BurleyNormalizedSSSCommon.ush"
 
 // Simplified material shaders for GPU Lightmass, supporting Lambert, Two Sided Foliage and Thin Translucent
@@ -37,20 +38,12 @@ bool IsNeeValidMaterial(FPathTracingPayload Payload)
 		return false;
 	}
 
-	switch (Payload.ShadingModelID)
+	if (Payload.ShadingModelID == SHADINGMODELID_UNLIT)
 	{
-		case SHADINGMODELID_UNLIT:
-			return false;
-		case SHADINGMODELID_DEFAULT_LIT:		
-		case SHADINGMODELID_CLEAR_COAT:			
-		case SHADINGMODELID_TWOSIDED_FOLIAGE:	
-		case SHADINGMODELID_THIN_TRANSLUCENT:
-		case SHADINGMODELID_CLOTH:
-		case SHADINGMODELID_SUBSURFACE:
-		case SHADINGMODELID_PREINTEGRATED_SKIN:
-		default:
-			return true;
+		return false;
 	}
+
+	return true;
 }
 
 float3 GetMaterialTransparency(FPathTracingPayload Payload, float3 RayDirection)
@@ -154,6 +147,7 @@ float3 EstimateMaterialAlbedo(FPathTracingPayload Payload)
 			break;
 		case SHADINGMODELID_DEFAULT_LIT:
 		case SHADINGMODELID_CLEAR_COAT:
+		case SHADINGMODELID_EYE:
 			// TODO: include clearcoat albedo as well
 			Albedo = BrdfWeight * (Payload.DiffuseColor + Payload.SpecularColor);
 			break;
@@ -237,6 +231,7 @@ void AccumulateAlbedoNormal(inout FPathTracingPayload Payload, float3 PathThroug
 		break;
 	case SHADINGMODELID_DEFAULT_LIT:
 	case SHADINGMODELID_CLEAR_COAT:
+	case SHADINGMODELID_EYE:
 		Albedo += BrdfWeight * PathThroughput * (Payload.DiffuseColor + RoughnessWeight * Payload.SpecularColor);
 		Normal += Luminance(BrdfWeight * PathThroughput) * Payload.WorldNormal;
 		break;
@@ -280,6 +275,7 @@ void AdjustMaterialRoughness(inout FPathTracingPayload Payload, float PathRoughn
 		case SHADINGMODELID_CLOTH:
 		case SHADINGMODELID_SUBSURFACE:
 		case SHADINGMODELID_PREINTEGRATED_SKIN:
+		case SHADINGMODELID_EYE:
 			Payload.Roughness = max(Payload.Roughness, PathRoughness);
 			break;
 		case SHADINGMODELID_SUBSURFACE_PROFILE:
@@ -534,6 +530,9 @@ FMaterialSample SampleMaterial(
 		case SHADINGMODELID_TWOSIDED_FOLIAGE:
 			Result = TwoSidedFoliage_SampleMaterial(RayDirection, Payload, RandSample);
 			break;
+		case SHADINGMODELID_EYE:
+			Result = Eye_SampleMaterial(RayDirection, Payload, RandSample);
+			break;
 		default:
 			Result = Lambert_SampleMaterial(Payload, RandSample);
 			break;
@@ -627,6 +626,9 @@ FMaterialEval EvalMaterial(
 		case SHADINGMODELID_TWOSIDED_FOLIAGE:
 			Result = TwoSidedFoliage_EvalMaterial(IncomingDirection, OutgoingDirection, Payload);
 			break;
+		case SHADINGMODELID_EYE:
+			Result = Eye_EvalMaterial(IncomingDirection, OutgoingDirection, Payload);
+			break;
 		default:
 			Result = Lambert_EvalMaterial(OutgoingDirection, Payload);
 			break;
diff --git a/Engine/Shaders/Private/PathTracing/PathTracingMaterialHitShader.usf b/Engine/Shaders/Private/PathTracing/PathTracingMaterialHitShader.usf
--- a/Engine/Shaders/Private/PathTracing/PathTracingMaterialHitShader.usf
+++ b/Engine/Shaders/Private/PathTracing/PathTracingMaterialHitShader.usf
@@ -188,17 +188,17 @@ RAY_TRACING_ENTRY_CLOSEST_HIT(PathTracingMaterialCHS,
 	else
 #endif
 
-#if MATERIAL_SHADINGMODEL_TWOSIDED_FOLIAGE || MATERIAL_SHADINGMODEL_SUBSURFACE || MATERIAL_SHADINGMODEL_PREINTEGRATED_SKIN || MATERIAL_SHADINGMODEL_SUBSURFACE_PROFILE || MATERIAL_SHADINGMODEL_EYE
+#if MATERIAL_SHADINGMODEL_TWOSIDED_FOLIAGE || MATERIAL_SHADINGMODEL_SUBSURFACE || MATERIAL_SHADINGMODEL_PREINTEGRATED_SKIN || MATERIAL_SHADINGMODEL_SUBSURFACE_PROFILE
 	if (Payload.ShadingModelID == SHADINGMODELID_TWOSIDED_FOLIAGE ||
 		Payload.ShadingModelID == SHADINGMODELID_SUBSURFACE ||
 		Payload.ShadingModelID == SHADINGMODELID_PREINTEGRATED_SKIN ||
-		Payload.ShadingModelID == SHADINGMODELID_SUBSURFACE_PROFILE |
-		Payload.ShadingModelID == SHADINGMODELID_EYE)
+		Payload.ShadingModelID == SHADINGMODELID_SUBSURFACE_PROFILE)
 	{
 		Payload.CustomData = GetMaterialSubsurfaceData(PixelMaterialInputs);
 	}
 	else
 #endif
+
 #if MATERIAL_SHADINGMODEL_THIN_TRANSLUCENT
 	if (Payload.ShadingModelID == SHADINGMODELID_THIN_TRANSLUCENT)
 	{
@@ -206,10 +206,53 @@ RAY_TRACING_ENTRY_CLOSEST_HIT(PathTracingMaterialCHS,
 	}
 	else
 #endif
+
+#if MATERIAL_SHADINGMODEL_EYE
+	if (Payload.ShadingModelID == SHADINGMODELID_EYE)
+	{
+		// This is all based on logic from ShadingModelsMaterial.ush
+		const float IrisMask     = saturate(GetMaterialCustomData0(MaterialParameters));
+		const float IrisDistance = saturate(GetMaterialCustomData1(MaterialParameters));
+		Payload.Anisotropy = IrisMask;
+		Payload.CustomData.a = GetMaterialSubsurfaceData(PixelMaterialInputs).a; // SubsurfaceProfile (NOTE: this is not used, EyeBxDF leaves Diffuse=0, so does not actually apply SSS?)
+		Payload.Metallic = 0.0;
+#if NUM_MATERIAL_OUTPUTS_GETTANGENTOUTPUT > 0
+		// Blend in the negative intersection normal to create some concavity
+		// Not great as it ties the concavity to the convexity of the cornea surface
+		// No good justification for that. On the other hand, if we're just looking to
+		// introduce some concavity, this does the job.
+		const float3 PlaneNormal = normalize(GetTangentOutput0(MaterialParameters));
+		const float3 CausticNormal = normalize(lerp(PlaneNormal, -Payload.WorldNormal, IrisMask * IrisDistance));
+		Payload.CustomVector = CausticNormal;
+#else
+		const float3 PlaneNormal = Payload.WorldNormal;
+		Payload.CustomVector = Payload.WorldNormal;
+#endif
+
+#if IRIS_NORMAL // on
+		Payload.Specular = 0.25;
+
+	#if NUM_MATERIAL_OUTPUTS_CLEARCOATBOTTOMNORMAL > 0
+		float3 IrisNormal = normalize( ClearCoatBottomNormal0(MaterialParameters) );
+		#if MATERIAL_TANGENTSPACENORMAL
+		IrisNormal = normalize( TransformTangentVectorToWorld( MaterialParameters.TangentToWorld, IrisNormal ) );
+		#endif
+	#else
+		float3 IrisNormal = PlaneNormal;
+	#endif
+		Payload.WorldTangent = IrisNormal;
+#else
+		// IRIS_NORMAL off
+		Payload.WorldTangent = PlaneNormal;
+#endif
+	}
+	else
+#endif
 	{
 		// terminal case for when none of the above worked
 	}
 
+
 #else // SIMPLIFIED_MATERIAL_SHADER
 
 	/**
diff --git a/Engine/Shaders/Private/ShadingModelsMaterial.ush b/Engine/Shaders/Private/ShadingModelsMaterial.ush
--- a/Engine/Shaders/Private/ShadingModelsMaterial.ush
+++ b/Engine/Shaders/Private/ShadingModelsMaterial.ush
@@ -141,17 +141,13 @@ void SetGBufferForShadingModel(
 #if MATERIAL_SHADINGMODEL_EYE
 	else if (ShadingModel == SHADINGMODELID_EYE)
 	{
-		GBuffer.CustomData.x = EncodeSubsurfaceProfile(SubsurfaceProfile).x;
-		GBuffer.CustomData.w = 1.0f - saturate(GetMaterialCustomData0(MaterialParameters));	// Opacity = 1.0 - Iris Mask
-		GBuffer.Metallic = saturate(GetMaterialCustomData1(MaterialParameters));			// Iris Distance
-
-	#if IRIS_NORMAL
-		float IrisMask		= saturate( GetMaterialCustomData0(MaterialParameters) );
-		float IrisDistance	= saturate( GetMaterialCustomData1(MaterialParameters) );
+		const float IrisMask     = saturate(GetMaterialCustomData0(MaterialParameters));
+		const float IrisDistance = saturate(GetMaterialCustomData1(MaterialParameters));
 
 		GBuffer.CustomData.x = EncodeSubsurfaceProfile(SubsurfaceProfile).x;
-		GBuffer.CustomData.w = 1.0 - IrisMask;	// Opacity
+		GBuffer.CustomData.w = 1.0f - IrisMask;	// Opacity
 
+	#if IRIS_NORMAL
 		float2 WorldNormalOct = UnitVectorToOctahedron( GBuffer.WorldNormal );
 
 		// CausticNormal stored as octahedron
@@ -186,7 +182,7 @@ void SetGBufferForShadingModel(
 		float2 IrisNormalDelta = ( IrisNormalOct - WorldNormalOct ) * 0.5 + (128.0/255.0);
 		GBuffer.CustomData.yz = IrisNormalDelta;
 	#else
-		GBuffer.Metallic = saturate(GetMaterialCustomData1(MaterialParameters));			// Iris Distance
+		GBuffer.Metallic = IrisDistance;
 
 		#if NUM_MATERIAL_OUTPUTS_GETTANGENTOUTPUT > 0
 			float3 Tangent = GetTangentOutput0(MaterialParameters);
