Two fixes for cloth wrap deformer when there are multiple cloth sections and one section has multiple influences enabled
1) we had a "use multiple influences" flag at the FGPUSkinCacheEntry level, but it should be at the DispatchData level
2)  in the shader we were assuming the buffer of skinning weights was either all single-influence or all multiple-influence, but really it can change per section

See CL 15896865 which is a similar fix for the non-skin-cache version.

#rnx

diff --git a/Engine/Shaders/Private/GpuSkinCacheComputeShader.usf b/Engine/Shaders/Private/GpuSkinCacheComputeShader.usf
--- a/Engine/Shaders/Private/GpuSkinCacheComputeShader.usf
+++ b/Engine/Shaders/Private/GpuSkinCacheComputeShader.usf
@@ -96,7 +96,7 @@ struct FVertexTriangleInfluence
 };
 
 // In ClothBuffer:
-#define NUM_FLOAT4S_PER_VERTEX 4;
+#define NUM_FLOAT4S_PER_VERTEX_INFLUENCE 4;
 
 Buffer<float4> ClothBuffer;
 Buffer<float2> ClothPositionsAndNormalsBuffer;
@@ -105,11 +105,12 @@ float ClothBlendWeight;
 float4x4 ClothLocalToWorld;
 float4x4 ClothWorldToLocal;
 
-void GetClothBuffer(uint VertexIndex, out FVertexTriangleInfluence Influences[NUM_INFLUENCES_PER_VERTEX])
+void GetClothBuffer(uint VertexIndex, uint ClothIndexOffset, out FVertexTriangleInfluence Influences[NUM_INFLUENCES_PER_VERTEX])
 {
-    for (int i = 0; i < NUM_INFLUENCES_PER_VERTEX; ++i )
-    {
-        uint Offset = (NUM_INFLUENCES_PER_VERTEX * VertexIndex + i) * NUM_FLOAT4S_PER_VERTEX;
+	const uint VertexOffset = NUM_INFLUENCES_PER_VERTEX * VertexIndex + ClothIndexOffset;
+	for (int i = 0; i < NUM_INFLUENCES_PER_VERTEX; ++i )
+	{
+		const uint Offset = (VertexOffset + i) * NUM_FLOAT4S_PER_VERTEX_INFLUENCE;
 
         Influences[i].PositionBaryCoordsAndDist = ClothBuffer[Offset];
         Influences[i].NormalBaryCoordsAndDist = ClothBuffer[Offset + 1];
@@ -238,7 +239,7 @@ uint InputWeightStride;
 // see r.SkinCache.Debug, default is 1.0f
 float SkinCacheDebug;
 
-FVertexUnpacked UnpackedVertex(uint MorphIndex, uint ClothIndex, uint VertexIndex, uint WeightOffset )
+FVertexUnpacked UnpackedVertex(uint MorphIndex, uint ClothIndex, uint ClothIndexOffset, uint VertexIndex, uint WeightOffset)
 {
 	FVertexUnpacked Unpacked;
 
@@ -295,7 +296,7 @@ FVertexUnpacked UnpackedVertex(uint MorphIndex, uint ClothIndex, uint VertexInde
 	GetVertexOffsets(VertexIndex, Unpacked.PreSkinOffset, Unpacked.PostSkinOffset);
 
 #if GPUSKIN_APEX_CLOTH
-    GetClothBuffer(ClothIndex, Unpacked.Influences);
+    GetClothBuffer(ClothIndex, ClothIndexOffset, Unpacked.Influences);
 #endif
 	return Unpacked;
 }
@@ -316,7 +317,7 @@ void SkinCacheUpdateBatchCS(uint3 GroupID : SV_GroupID,
  uint3 GroupThreadID : SV_GroupThreadID)
 {
 	uint VertexIndex = DispatchThreadID.x;
-		if (VertexIndex >= NumVertices)
+	if (VertexIndex >= NumVertices)
 	{
 		return;
 	}
@@ -330,15 +331,15 @@ void SkinCacheUpdateBatchCS(uint3 GroupID : SV_GroupID,
 
 #if GPUSKIN_MORPH_BLEND
 #if GPUSKIN_APEX_CLOTH
-	FVertexUnpacked	Unpacked = UnpackedVertex(MorphBufferOffset + VertexIndex, ClothBufferOffset + VertexIndex, InputOffset, WeightOffset);
+	FVertexUnpacked Unpacked = UnpackedVertex(MorphBufferOffset + VertexIndex, VertexIndex, ClothBufferOffset, InputOffset, WeightOffset);
 #else
-	FVertexUnpacked	Unpacked = UnpackedVertex(MorphBufferOffset + VertexIndex, VertexIndex, InputOffset, WeightOffset);
+	FVertexUnpacked Unpacked = UnpackedVertex(MorphBufferOffset + VertexIndex, VertexIndex, 0, InputOffset, WeightOffset);
 #endif
 #else
 #if GPUSKIN_APEX_CLOTH
-	FVertexUnpacked	Unpacked = UnpackedVertex(VertexIndex, ClothBufferOffset + VertexIndex, InputOffset, WeightOffset);
+	FVertexUnpacked Unpacked = UnpackedVertex(VertexIndex, VertexIndex, ClothBufferOffset, InputOffset, WeightOffset);
 #else
-	FVertexUnpacked	Unpacked = UnpackedVertex(VertexIndex, VertexIndex, InputOffset, WeightOffset);
+	FVertexUnpacked Unpacked = UnpackedVertex(VertexIndex, VertexIndex, 0, InputOffset, WeightOffset);
 #endif
 #endif
 
diff --git a/Engine/Source/Runtime/Engine/Private/GPUSkinCache.cpp b/Engine/Source/Runtime/Engine/Private/GPUSkinCache.cpp
--- a/Engine/Source/Runtime/Engine/Private/GPUSkinCache.cpp
+++ b/Engine/Source/Runtime/Engine/Private/GPUSkinCache.cpp
@@ -533,7 +533,6 @@ protected:
 	FRHIShaderResourceView* MorphBuffer;
 	FShaderResourceViewRHIRef ClothBuffer;
 	int32 LOD;
-	bool bMultipleClothSkinInfluences;
 
 	friend class FGPUSkinCache;
 	friend class FBaseGPUSkinCacheCS;
@@ -1667,16 +1666,6 @@ bool FGPUSkinCache::ProcessEntry(
         InOutEntry->DispatchData[Section].ClothBlendWeight = ClothBlendWeight;
         InOutEntry->DispatchData[Section].ClothLocalToWorld = ClothLocalToWorld;
         InOutEntry->DispatchData[Section].ClothWorldToLocal = ClothLocalToWorld.Inverse();
-
-		const int32 NumWrapWeights = InOutEntry->DispatchData[Section].Section->ClothMappingData.Num();
-		if (NumWrapWeights > NumVertices)
-		{
-			InOutEntry->bMultipleClothSkinInfluences = true;
-		}
-		else
-		{
-			InOutEntry->bMultipleClothSkinInfluences = false;
-		}
     }
     InOutEntry->DispatchData[Section].SkinType = ClothVertexBuffer ? 2 : (bMorph ? 1 : 0);
 
@@ -1996,6 +1985,9 @@ void FGPUSkinCache::DispatchUpdateSkinning(FRHICommandListImmediate& RHICmdList,
 	FGPUSkinCacheEntry::FSectionDispatchData& DispatchData = Entry->DispatchData[Section];
 	FGPUBaseSkinVertexFactory::FShaderDataType& ShaderData = DispatchData.SourceVertexFactory->GetShaderData();
 	const FString RayTracingTag = (Entry->Mode == EGPUSkinCacheEntryMode::RayTracing ? TEXT("[RT]") : TEXT(""));
+	
+	const uint32 NumWrapDeformerWeights = DispatchData.Section->ClothMappingData.Num();
+	const bool bMultipleWrapDeformerInfluences = (DispatchData.NumVertices < NumWrapDeformerWeights);
 
 	SCOPED_DRAW_EVENTF(RHICmdList, SkinCacheDispatch,
 		TEXT("%sSkinning%d%d%d Mesh=%s LOD=%d Chunk=%d InStreamStart=%d OutStart=%d Vert=%d Morph=%d/%d"),
@@ -2061,7 +2053,7 @@ void FGPUSkinCache::DispatchUpdateSkinning(FRHICommandListImmediate& RHICmdList,
 		}
 		break;
 	case 2:
-		if (Entry->bMultipleClothSkinInfluences)
+		if (bMultipleWrapDeformerInfluences)
 		{
 			// Multiple influences for cloth skinning
 			if (Entry->BoneInfluenceType == 0)
