Fixed usage of PreSkinPosition for SkinCache on platforms that do not support manual vertex fetch

diff --git a/Engine/Shaders/Private/LocalVertexFactory.ush b/Engine/Shaders/Private/LocalVertexFactory.ush
--- a/Engine/Shaders/Private/LocalVertexFactory.ush
+++ b/Engine/Shaders/Private/LocalVertexFactory.ush
@@ -163,6 +163,10 @@ struct FVertexFactoryInput
 	uint InstanceId	: SV_InstanceID;
 #endif
 
+#if GPUSKIN_PASS_THROUGH && !MANUAL_VERTEX_FETCH
+	float4	PreSkinPosition	: ATTRIBUTE14;
+#endif
+
 #if NEEDS_LIGHTMAP_COORDINATE && !MANUAL_VERTEX_FETCH
 	float2	LightMapCoordinate : ATTRIBUTE15;
 #endif
@@ -931,10 +935,14 @@ FVertexFactoryIntermediates GetVertexFactoryIntermediates(FVertexFactoryInput In
 #endif	// USE_INSTANCING
 
 #if GPUSKIN_PASS_THROUGH
-	uint PreSkinVertexOffset = LocalVF.PreSkinBaseVertexIndex + Input.VertexId * 3;
-	Intermediates.PreSkinPosition.x = LocalVF.VertexFetch_PreSkinPositionBuffer[PreSkinVertexOffset + 0];
-	Intermediates.PreSkinPosition.y = LocalVF.VertexFetch_PreSkinPositionBuffer[PreSkinVertexOffset + 1];
-	Intermediates.PreSkinPosition.z = LocalVF.VertexFetch_PreSkinPositionBuffer[PreSkinVertexOffset + 2];
+	#if MANUAL_VERTEX_FETCH
+		uint PreSkinVertexOffset = LocalVF.PreSkinBaseVertexIndex + Input.VertexId * 3;
+		Intermediates.PreSkinPosition.x = LocalVF.VertexFetch_PreSkinPositionBuffer[PreSkinVertexOffset + 0];
+		Intermediates.PreSkinPosition.y = LocalVF.VertexFetch_PreSkinPositionBuffer[PreSkinVertexOffset + 1];
+		Intermediates.PreSkinPosition.z = LocalVF.VertexFetch_PreSkinPositionBuffer[PreSkinVertexOffset + 2];
+	#else
+		Intermediates.PreSkinPosition = Input.PreSkinPosition.xyz;
+	#endif
 #else
 	Intermediates.PreSkinPosition = Input.Position.xyz;
 #endif
diff --git a/Engine/Source/Runtime/Engine/Private/GPUSkinVertexFactory.cpp b/Engine/Source/Runtime/Engine/Private/GPUSkinVertexFactory.cpp
--- a/Engine/Source/Runtime/Engine/Private/GPUSkinVertexFactory.cpp
+++ b/Engine/Source/Runtime/Engine/Private/GPUSkinVertexFactory.cpp
@@ -484,6 +484,7 @@ void TGPUSkinVertexFactory<BoneInfluenceType>::CopyDataTypeForPassthroughFactory
 	DestDataType.TangentBasisComponents[1] = Data.TangentBasisComponents[1];
 	DestDataType.TextureCoordinates = Data.TextureCoordinates;
 	DestDataType.ColorComponent = Data.ColorComponent;
+	DestDataType.PreSkinPositionComponent = Data.PositionComponent;
 	DestDataType.PositionComponentSRV = Data.PositionComponentSRV;
 	DestDataType.PreSkinPositionComponentSRV = Data.PositionComponentSRV;
 	DestDataType.TangentsSRV = Data.TangentsSRV;
@@ -852,8 +853,6 @@ void FGPUSkinPassthroughVertexFactory::InternalUpdateVertexDeclaration(FGPUBaseS
 		Data.TangentBasisComponents[1].VertexStreamUsage = EVertexStreamUsage::Overridden | EVertexStreamUsage::ManualFetch;
 	}
 
-	int32 PrevNumStreams = Streams.Num();
-
 	//hack to allow us to release the alias pointers properly in ReleaseRHI.
 	//To be cleaned up in UE-68826
 	FLocalVertexFactory::ReleaseRHI();
@@ -861,8 +860,6 @@ void FGPUSkinPassthroughVertexFactory::InternalUpdateVertexDeclaration(FGPUBaseS
 	FLocalVertexFactory::InitDynamicRHI();
 	FLocalVertexFactory::InitRHI();
 
-	// Verify no additional stream was created
-	check(Streams.Num() == PrevNumStreams);
 	// Find the added stream (usually at 0)
 	PositionStreamIndex = -1;
 	TangentStreamIndex = -1;
diff --git a/Engine/Source/Runtime/Engine/Private/LocalVertexFactory.cpp b/Engine/Source/Runtime/Engine/Private/LocalVertexFactory.cpp
--- a/Engine/Source/Runtime/Engine/Private/LocalVertexFactory.cpp
+++ b/Engine/Source/Runtime/Engine/Private/LocalVertexFactory.cpp
@@ -356,6 +356,12 @@ void FLocalVertexFactory::InitRHI()
 		}
 	}
 
+	// Only used with FGPUSkinPassthroughVertexFactory on platforms that do not use ManualVertexFetch
+	if (Data.PreSkinPositionComponent.VertexBuffer != nullptr)
+	{
+		Elements.Add(AccessStreamComponent(Data.PreSkinPositionComponent,14));
+	}
+
 	if (Data.LightMapCoordinateComponent.VertexBuffer)
 	{
 		Elements.Add(AccessStreamComponent(Data.LightMapCoordinateComponent,15));
diff --git a/Engine/Source/Runtime/Engine/Public/LocalVertexFactory.h b/Engine/Source/Runtime/Engine/Public/LocalVertexFactory.h
--- a/Engine/Source/Runtime/Engine/Public/LocalVertexFactory.h
+++ b/Engine/Source/Runtime/Engine/Public/LocalVertexFactory.h
@@ -52,6 +52,7 @@ public:
 
 	struct FDataType : public FStaticMeshDataType
 	{
+		FVertexStreamComponent PreSkinPositionComponent;
 		FRHIShaderResourceView* PreSkinPositionComponentSRV = nullptr;
 	};
 
