Added cross-platform support to build Raytracing Instance Descriptor in shader.

diff --git a/Engine/Shaders/Private/RayTracing/RayTracingCommon.ush b/Engine/Shaders/Private/RayTracing/RayTracingCommon.ush
--- a/Engine/Shaders/Private/RayTracing/RayTracingCommon.ush
+++ b/Engine/Shaders/Private/RayTracing/RayTracingCommon.ush
@@ -87,7 +87,42 @@
 [shader("miss")] void name(inout payload_type payload_name)
 #endif //RAY_TRACING_ENTRY_MISS
 
-// 
+#ifndef OVERRIDE_PLATFORM_RAYTRACING_INSTANCE_DESCRITOR
+#define OVERRIDE_PLATFORM_RAYTRACING_INSTANCE_DESCRITOR
+
+// Must match D3D12_RAYTRACING_INSTANCE_FLAGS
+#define PLATFORM_RAYTRACING_INSTANCE_FLAG_NONE								0
+#define PLATFORM_RAYTRACING_INSTANCE_FLAG_TRIANGLE_CULL_DISABLE				0x1
+#define PLATFORM_RAYTRACING_INSTANCE_FLAG_TRIANGLE_FRONT_COUNTERCLOCKWISE	0x2
+#define PLATFORM_RAYTRACING_INSTANCE_FLAG_FORCE_OPAQUE						0x4
+#define PLATFORM_RAYTRACING_INSTANCE_FLAG_FORCE_NON_OPAQUE					0x8
+
+// Must match D3D12_RAYTRACING_INSTANCE_DESC
+struct FRayTracingInstanceDescriptor
+{
+	float3x4 Transform;
+	uint InstanceIdAndMask;
+	uint InstanceContributionToHitGroupIndexAndFlags;
+	uint2 AccelerationStructure; //uint64
+};
+
+void SetRayTracingInstanceTransform(inout FRayTracingInstanceDescriptor Desc, float3x4 LocalToWorld)
+{
+	Desc.Transform = LocalToWorld;
+}
+
+FRayTracingInstanceDescriptor BuildPlatformRayTracingInstanceDesc(uint InstanceMask, uint InstanceId, uint InstanceFlags, uint InstanceContributionToHitGroupIndex, float3x4 LocalToWorld, uint2 AccelerationStructureAddress)
+{
+	FRayTracingInstanceDescriptor Desc;
+	Desc.InstanceIdAndMask = (InstanceMask << 24) | InstanceId;
+	Desc.InstanceContributionToHitGroupIndexAndFlags = (InstanceFlags << 24) | InstanceContributionToHitGroupIndex;
+	Desc.AccelerationStructure = AccelerationStructureAddress;
+	SetRayTracingInstanceTransform(Desc, LocalToWorld);
+
+	return Desc;
+}
+
+#endif // OVERRIDE_PLATFORM_RAYTRACING_INSTANCE_DESCRITOR
 
 struct FBasicRayData
 {
@@ -184,6 +219,41 @@ uint2 GetPixelCoord(uint2 DispatchThreadId, uint UpscaleFactor)
 	return DispatchThreadId * UpscaleFactor + uint2(SubPixelId & (UpscaleFactor - 1), SubPixelId / UpscaleFactor);
 }
 
+// must match ERayTracingInstanceFlags
+#define RAYTRACING_INSTANCE_FLAG_NONE					0
+#define RAYTRACING_INSTANCE_FLAG_TRIANGLE_CULL_DISABLE	(1 << 1)
+#define RAYTRACING_INSTANCE_FLAG_TRIANGLE_CULL_REVERSE	(1 << 2)
+#define RAYTRACING_INSTANCE_FLAG_FORCE_OPAQUE			(1 << 3)
+#define RAYTRACING_INSTANCE_FLAG_FORCE_NON_OPAQUE		(1 << 4)
+
+uint TranslateRayTracingInstanceFlags(uint InFlags)
+{
+	uint Result = PLATFORM_RAYTRACING_INSTANCE_FLAG_NONE;
+
+	if (InFlags & RAYTRACING_INSTANCE_FLAG_TRIANGLE_CULL_DISABLE)
+	{
+		Result |= PLATFORM_RAYTRACING_INSTANCE_FLAG_TRIANGLE_CULL_DISABLE;
+	}
+
+	if ((InFlags & RAYTRACING_INSTANCE_FLAG_TRIANGLE_CULL_REVERSE) == 0)
+	{
+		// Counterclockwise is the default for UE. Reversing culling is achieved by *not* setting this flag.
+		Result |= PLATFORM_RAYTRACING_INSTANCE_FLAG_TRIANGLE_FRONT_COUNTERCLOCKWISE;
+	}
+
+	if (InFlags & RAYTRACING_INSTANCE_FLAG_FORCE_OPAQUE)
+	{
+		Result |= PLATFORM_RAYTRACING_INSTANCE_FLAG_FORCE_OPAQUE;
+	}
+
+	if (InFlags & RAYTRACING_INSTANCE_FLAG_FORCE_NON_OPAQUE)
+	{
+		Result |= PLATFORM_RAYTRACING_INSTANCE_FLAG_FORCE_NON_OPAQUE;
+	}
+
+	return Result;
+}
+
 #define RAY_TRACING_BLEND_MODE_OPAQUE				0
 #define RAY_TRACING_BLEND_MODE_ALPHA_COMPOSITE		1
 #define RAY_TRACING_BLEND_MODE_TRANSLUCENT			2
diff --git a/Engine/Shaders/Private/RayTracing/RayTracingInstanceCopy.usf b/Engine/Shaders/Private/RayTracing/RayTracingInstanceCopy.usf
--- a/Engine/Shaders/Private/RayTracing/RayTracingInstanceCopy.usf
+++ b/Engine/Shaders/Private/RayTracing/RayTracingInstanceCopy.usf
@@ -4,38 +4,29 @@
 	RayTracingInstanceCopy.usf: Build ray tracing instances in the GPU 
 ===============================================================================================*/
 
+#include "RayTracingCommon.ush"
 #include "../Common.ush"
 
 #ifndef THREADGROUP_SIZE
 #define THREADGROUP_SIZE 1
 #endif
 
-// Should match D3D12_RAYTRACING_INSTANCE_DESC size
-//#dxr-todo: support other platforms with a define set when the shader is dispatched
-struct FInstanceDescriptor
-{
-	float3x4 Transform;
-	uint InstanceIdAndMask;
-	uint InstanceContributionToHitGroupIndexAndFlags;
-	uint AccelerationStructure[2]; //uint64
-};
-
 uint InstancesCPUCount;
 uint DescBufferOffset;
 StructuredBuffer<float3x4> InstancesTransforms;
-RWStructuredBuffer<FInstanceDescriptor> InstancesDescriptors;
+RWStructuredBuffer<FRayTracingInstanceDescriptor> InstancesDescriptors;
 
 
 [numthreads(THREADGROUP_SIZE, 1, 1)]
 void RayTracingInstanceCopyShaderCS(uint3 DispatchThreadId : SV_DispatchThreadID)
 {
-	uint InstanceIndex = DispatchThreadId.x;	
+	uint InstanceIdx = DispatchThreadId.x;	
 
-	float3x4 InstanceXForm = InstancesTransforms[InstanceIndex];
-	uint DescIndex = DescBufferOffset + InstanceIndex;
+	float3x4 InstanceXForm = InstancesTransforms[InstanceIdx];
+	uint DescIndex = DescBufferOffset + InstanceIdx;
 
-	if (InstanceIndex < InstancesCPUCount)
+	if (InstanceIdx < InstancesCPUCount)
 	{
-		InstancesDescriptors[DescIndex].Transform = InstanceXForm;
+		SetRayTracingInstanceTransform(InstancesDescriptors[DescIndex], InstanceXForm);
 	}
-}
\ No newline at end of file
+}
