Fixed couple issues with a mobile path in GPUScene (LWC and compressed instance transforms)

diff --git a/Engine/Shaders/Private/SceneData.ush b/Engine/Shaders/Private/SceneData.ush
--- a/Engine/Shaders/Private/SceneData.ush
+++ b/Engine/Shaders/Private/SceneData.ush
@@ -462,6 +462,18 @@ void ComputeInstanceDerivedData(inout FInstanceSceneData InstanceData)
 	//
 	// Do not put any load operations here!
 	//
+
+#if !INSTANCE_COMPRESSED_TRANSFORM
+	float3 Scale2;
+	Scale2.x = length2(InstanceData.LocalToWorld[0].xyz);
+	Scale2.y = length2(InstanceData.LocalToWorld[1].xyz);
+	Scale2.z = length2(InstanceData.LocalToWorld[2].xyz);
+
+	InstanceData.InvNonUniformScale							= rsqrt(Scale2);
+	InstanceData.DeterminantSign							= CondMask(InstanceData.Flags & INSTANCE_SCENE_DATA_FLAG_DETERMINANT_SIGN, -1.0f, 1.0f);
+	InstanceData.NonUniformScale.xyz						= Scale2 * InstanceData.InvNonUniformScale;
+#endif
+
 	InstanceData.NonUniformScale.w							= max3( InstanceData.NonUniformScale.x, InstanceData.NonUniformScale.y, InstanceData.NonUniformScale.z );
 	
 	InstanceData.WorldToLocal								= InstanceData.LocalToWorld;
@@ -527,15 +539,6 @@ FInstanceSceneData GetInstanceSceneData(uint InstanceId, uint SOAStride)
 		InstanceData.LightMapAndShadowMapUVBias					= LoadInstanceSceneDataElement(9 * SOAStride + InstanceId);
 	#endif
 
-		float3 Scale2;
-		Scale2.x = length2(InstanceData.LocalToWorld[0].xyz);
-		Scale2.y = length2(InstanceData.LocalToWorld[1].xyz);
-		Scale2.z = length2(InstanceData.LocalToWorld[2].xyz);
-
-		InstanceData.InvNonUniformScale							= rsqrt(Scale2);
-		InstanceData.DeterminantSign							= CondMask(InstanceData.Flags & INSTANCE_SCENE_DATA_FLAG_DETERMINANT_SIGN, -1.0f, 1.0f);
-
-		InstanceData.NonUniformScale.xyz						= Scale2 * InstanceData.InvNonUniformScale;
 #else
 		uint4  RotationScale									= asuint(LoadInstanceSceneDataElement(1 * SOAStride + InstanceId));
 		float3 Translation										= LoadInstanceSceneDataElement(2 * SOAStride + InstanceId).xyz;
@@ -681,9 +684,13 @@ FSceneDataIntermediates GetSceneDataIntermediates(uint DrawInstanceId, float4 In
 	InstanceData.Flags = InstanceFlags;
 	InstanceData.PrimitiveId = asuint(InstanceOrigin.w);
 	
-	#if USES_PER_INSTANCE_RANDOM
+#if USES_PER_INSTANCE_RANDOM
 	InstanceData.RandomID = InstanceTransform3.w;
-	#endif
+#endif
+
+#if INSTANCE_COMPRESSED_TRANSFORM
+	#error Compressed instance transform is not supported on mobile path
+#endif
 
 	ComputeInstanceDerivedData(InstanceData);
 
diff --git a/Engine/Source/Runtime/Renderer/Private/GPUScene.cpp b/Engine/Source/Runtime/Renderer/Private/GPUScene.cpp
--- a/Engine/Source/Runtime/Renderer/Private/GPUScene.cpp
+++ b/Engine/Source/Runtime/Renderer/Private/GPUScene.cpp
@@ -1910,7 +1910,7 @@ void FGPUScene::AddUpdatePrimitiveIdsPass(FRDGBuilder& GraphBuilder, FInstanceGP
 
 void FGPUSceneCompactInstanceData::Init(const FGPUScenePrimitiveCollector* PrimitiveCollector, int32 PrimitiveId)
 {
-	FMatrix LocalToWorld = FMatrix::Identity;
+	FMatrix44f LocalToWorld = FMatrix44f::Identity;
 	int32 DynamicPrimitiveId = PrimitiveId;
 	if (PrimitiveCollector && PrimitiveCollector->UploadData && !PrimitiveCollector->GetPrimitiveIdRange().IsEmpty())
 	{
@@ -1931,7 +1931,7 @@ void FGPUSceneCompactInstanceData::Init(const FGPUScenePrimitiveCollector* Primi
 
 void FGPUSceneCompactInstanceData::Init(const FScene* Scene, int32 PrimitiveId)
 {
-	FMatrix LocalToWorld = FMatrix::Identity;
+	FMatrix44f LocalToWorld = FMatrix44f::Identity;
 	if (Scene && PrimitiveId >= 0 && PrimitiveId < Scene->PrimitiveTransforms.Num())
 	{
 		LocalToWorld = Scene->PrimitiveTransforms[PrimitiveId];
diff --git a/Engine/Source/Runtime/Renderer/Private/GPUScene.h b/Engine/Source/Runtime/Renderer/Private/GPUScene.h
--- a/Engine/Source/Runtime/Renderer/Private/GPUScene.h
+++ b/Engine/Source/Runtime/Renderer/Private/GPUScene.h
@@ -396,11 +396,11 @@ private:
 
 struct FGPUSceneCompactInstanceData
 {
-	FVector4 InstanceOriginAndId;
-	FVector4 InstanceTransform1;
-	FVector4 InstanceTransform2;
-	FVector4 InstanceTransform3;
-	FVector4 InstanceAuxData;
+	FVector4f InstanceOriginAndId;
+	FVector4f InstanceTransform1;
+	FVector4f InstanceTransform2;
+	FVector4f InstanceTransform3;
+	FVector4f InstanceAuxData;
 
 	void Init(const FGPUScenePrimitiveCollector* PrimitiveCollector, int32 PrimitiveId);
 	void Init(const FScene* Scene, int32 PrimitiveId);
diff --git a/Engine/Source/Runtime/Renderer/Private/MeshDrawCommands.cpp b/Engine/Source/Runtime/Renderer/Private/MeshDrawCommands.cpp
--- a/Engine/Source/Runtime/Renderer/Private/MeshDrawCommands.cpp
+++ b/Engine/Source/Runtime/Renderer/Private/MeshDrawCommands.cpp
@@ -1009,10 +1009,12 @@ void SortAndMergeDynamicPassMeshDrawCommands(
 
 			if (FeatureLevel == ERHIFeatureLevel::ES3_1)
 			{
+				check(sizeof(FGPUSceneCompactInstanceData) == PrimitiveIdBufferStride);
+				
 				auto WritePrimitiveDataFn = [&](int32 PrimitiveIndex, int32 PrimitiveId) 
 				{
 					checkSlow(PrimitiveIndex < MaxNumPrimitives);
-					checkSlow(sizeof(FGPUSceneCompactInstanceData) == PrimitiveIdBufferStride);
+					
 					FGPUSceneCompactInstanceData* PrimitiveData = reinterpret_cast<FGPUSceneCompactInstanceData*>(PrimitiveIdBufferData) + PrimitiveIndex;
 
 					if ((PrimitiveId & GPrimIDDynamicFlag) != 0)
