Prevent uninitialized assignment to SceneDepthAux

diff --git a/Engine/Shaders/Private/MobileBasePassPixelShader.usf b/Engine/Shaders/Private/MobileBasePassPixelShader.usf
--- a/Engine/Shaders/Private/MobileBasePassPixelShader.usf
+++ b/Engine/Shaders/Private/MobileBasePassPixelShader.usf
@@ -99,7 +99,7 @@ int NumDynamicPointLights;
 #define FORCE_VERTEX_NORMAL (MOBILE_QL_DISABLE_MATERIAL_NORMAL)
 #define SUPPORT_SPOTLIGHTS_SHADOW (MATERIALBLENDING_SOLID || MATERIALBLENDING_MASKED) && PROJECT_MOBILE_ENABLE_MOVABLE_SPOTLIGHTS_SHADOW
 #define DEFERRED_SHADING_PATH (MOBILE_DEFERRED_SHADING && ((MATERIALBLENDING_SOLID || MATERIALBLENDING_MASKED) && !MATERIAL_SHADINGMODEL_SINGLELAYERWATER))
-#define USE_SCENE_DEPTH_AUX (DEFERRED_SHADING_PATH && METAL_PROFILE) || !DEFERRED_SHADING_PATH
+#define USE_SCENE_DEPTH_AUX ((DEFERRED_SHADING_PATH && METAL_PROFILE) || !DEFERRED_SHADING_PATH) && (MATERIALBLENDING_SOLID || MATERIALBLENDING_MASKED)
 #if DEFERRED_SHADING_PATH
 	#define SV_TargetDepthAux SV_Target4
 #else
@@ -1047,10 +1047,6 @@ void Main(
 	#else
 		OutColor.rgb = Color * VertexFog.a + VertexFog.rgb;
 
-		#if USE_SCENE_DEPTH_AUX
-			OutSceneDepthAux = SvPosition.z;
-		#endif
-
 		#if !MATERIAL_USE_ALPHA_TO_COVERAGE
 			// Planar reflections and scene captures use scene color alpha to keep track of where content has been rendered, for compositing into a different scene later
 			OutColor.a = 0.0;
@@ -1060,6 +1056,10 @@ void Main(
 		#endif
 	#endif
 
+#if USE_SCENE_DEPTH_AUX
+	OutSceneDepthAux = SvPosition.z;
+#endif
+
 	#if !MATERIALBLENDING_MODULATE
 		OutColor.rgb *= ResolvedView.PreExposure;
 	#endif
