Fixed landscape flicker on mobile

diff --git a/Engine/Shaders/Private/LandscapeVertexFactory.ush b/Engine/Shaders/Private/LandscapeVertexFactory.ush
--- a/Engine/Shaders/Private/LandscapeVertexFactory.ush
+++ b/Engine/Shaders/Private/LandscapeVertexFactory.ush
@@ -624,7 +624,8 @@ FVertexFactoryIntermediates GetVertexFactoryIntermediates(FVertexFactoryInput In
 	
 	float4 LODHeight = LODHeightIndex >> 2 ? Input.LODHeights1 : Input.LODHeights0;
 	float InputHeight = LODHeight[LODHeightIndex & 3];
-	float InputHeightNextLOD = LODHeight[LODHeightIndexNextLod & 3];
+	float4 LODHeightNext = LODHeightIndexNextLod >> 2 ? Input.LODHeights1 : Input.LODHeights0;
+	float InputHeightNextLOD = LODHeightNext[LODHeightIndexNextLod & 3];
 
 	float Height = MinHeight + InputHeight * HeightDelta;
 	float HeightNextLOD = MinHeight + InputHeightNextLOD * HeightDelta;
