- Add instance & particle transform space
- This space allows you to transform in / out of mesh particle & instanced static meshes
- Particle mesh space is now hidden and will be removed later on as it doesn't offer anything over the new space

diff --git a/Engine/Plugins/FX/Niagara/Shaders/Private/NiagaraMeshVertexFactory.ush b/Engine/Plugins/FX/Niagara/Shaders/Private/NiagaraMeshVertexFactory.ush
--- a/Engine/Plugins/FX/Niagara/Shaders/Private/NiagaraMeshVertexFactory.ush
+++ b/Engine/Plugins/FX/Niagara/Shaders/Private/NiagaraMeshVertexFactory.ush
@@ -18,8 +18,8 @@
 #define USE_PARTICLE_TIME (NEEDS_PARTICLE_TIME)
 #define USE_PARTICLE_RANDOM (NEEDS_PARTICLE_RANDOM)
 #define USE_PARTICLE_TRANSFORM (NEEDS_PARTICLE_TRANSFORM)
-#define USE_PARTICLE_LOCAL_TO_WORLD (NEEDS_PARTICLE_LOCAL_TO_WORLD)
-#define USE_PARTICLE_WORLD_TO_LOCAL (NEEDS_PARTICLE_WORLD_TO_LOCAL)
+#define USE_PARTICLE_LOCAL_TO_WORLD (NEEDS_PARTICLE_LOCAL_TO_WORLD || NEEDS_INSTANCE_LOCAL_TO_WORLD_PS)
+#define USE_PARTICLE_WORLD_TO_LOCAL (NEEDS_PARTICLE_WORLD_TO_LOCAL || NEEDS_INSTANCE_WORLD_TO_LOCAL_PS)
 
 #ifndef ENABLE_PRECISE_MOTION_VECTORS
 #define ENABLE_PRECISE_MOTION_VECTORS 0
@@ -367,19 +367,23 @@ FMaterialPixelParameters GetMaterialPixelParameters(FVertexFactoryInterpolantsVS
 #endif
 
 #if USE_PARTICLE_LOCAL_TO_WORLD
-	Result.Particle.ParticleToWorld = transpose(float4x4(
-		Interpolants.ParticleToWorld[0],
-		Interpolants.ParticleToWorld[1],
-		Interpolants.ParticleToWorld[2],
-		float4(0.0f, 0.0f, 0.0f, 1.0f)));
+	float4x4 ParticleToWorld = transpose(float4x4(Interpolants.ParticleToWorld[0], Interpolants.ParticleToWorld[1], Interpolants.ParticleToWorld[2], float4(0.0f, 0.0f, 0.0f, 1.0f)));
+	#if NEEDS_PARTICLE_LOCAL_TO_WORLD
+		Result.Particle.ParticleToWorld = ParticleToWorld;
+	#endif
+	#if NEEDS_INSTANCE_LOCAL_TO_WORLD_PS
+		Result.InstanceLocalToWorld = ParticleToWorld;
+	#endif
 #endif
 
 #if USE_PARTICLE_WORLD_TO_LOCAL
-	Result.Particle.WorldToParticle = transpose(float4x4(
-		Interpolants.WorldToParticle[0],
-		Interpolants.WorldToParticle[1],
-		Interpolants.WorldToParticle[2],
-		float4(0.0f, 0.0f, 0.0f, 1.0f)));
+	float4x4 WorldToParticle = transpose(float4x4(Interpolants.WorldToParticle[0], Interpolants.WorldToParticle[1], Interpolants.WorldToParticle[2], float4(0.0f, 0.0f, 0.0f, 1.0f)));
+	#if NEEDS_PARTICLE_WORLD_TO_LOCAL
+		Result.Particle.WorldToParticle = WorldToParticle;
+	#endif
+	#if NEEDS_INSTANCE_WORLD_TO_LOCAL_PS
+		Result.InstanceWorldToLocal = WorldToParticle;
+	#endif
 #endif
 
 	Result.Particle.MotionBlurFade = 1.0f;
@@ -440,11 +444,11 @@ FMaterialVertexParameters GetMaterialVertexParameters(FVertexFactoryInput Input,
 	Result.Particle.DynamicParameter3 = Intermediates.DynamicParameter3;
 #endif
 
-	float4x4 Transform = GetParticleTransform(Intermediates);
-	Result.InstanceLocalToWorld = Transform;
+	Result.InstanceLocalToWorld = GetParticleTransform(Intermediates);
+	Result.InstanceWorldToLocal = GetParticleInvTransform(Intermediates);
 
-	Result.Particle.ParticleToWorld = Transform;
-	Result.Particle.WorldToParticle = GetParticleInvTransform(Intermediates);
+	Result.Particle.ParticleToWorld = Result.InstanceLocalToWorld;
+	Result.Particle.WorldToParticle = Result.InstanceWorldToLocal;
 
 	// Use the particle's unscaled transform for rotating normals to match static mesh
 	Result.TangentToWorld = mul(TangentToLocal, (float3x3)GetParticleTransformNoScale(Intermediates));
diff --git a/Engine/Shaders/Private/LocalVertexFactory.ush b/Engine/Shaders/Private/LocalVertexFactory.ush
--- a/Engine/Shaders/Private/LocalVertexFactory.ush
+++ b/Engine/Shaders/Private/LocalVertexFactory.ush
@@ -514,6 +514,13 @@ FMaterialPixelParameters GetMaterialPixelParameters(FVertexFactoryInterpolantsVS
 	Result.TwoSidedSign = 1;
 	Result.PrimitiveId = GetPrimitiveId(Interpolants);
 
+#if HAS_INSTANCE_LOCAL_TO_WORLD_PS
+	Result.InstanceLocalToWorld = GetInstanceLocalToWorld(Interpolants);
+#endif
+#if HAS_INSTANCE_WORLD_TO_LOCAL_PS
+	Result.InstanceWorldToLocal = GetInstanceWorldToLocal(Interpolants);
+#endif
+
 #if NEEDS_PARTICLE_LOCAL_TO_WORLD
 	Result.Particle.ParticleToWorld = GetPrimitiveData(Result.PrimitiveId).LocalToWorld;
 #endif
@@ -556,11 +563,30 @@ FMaterialVertexParameters GetMaterialVertexParameters(FVertexFactoryInput Input,
 	Result.InstanceId = GetInstanceId(Input.InstanceId); 
 	Result.InstanceOffset = InstanceOffset;
 	Result.PrevFrameLocalToWorld = mul(GetInstancePrevTransform(Intermediates), GetPrimitiveData(Intermediates).PreviousLocalToWorld);
+
+	// Calculate derived world to local
+	{
+		float3 Scale2;
+		Scale2.x = length2(Result.InstanceLocalToWorld[0].xyz);
+		Scale2.y = length2(Result.InstanceLocalToWorld[1].xyz);
+		Scale2.z = length2(Result.InstanceLocalToWorld[2].xyz);
+
+		float3 InvNonUniformScale = rsqrt(Scale2);
+
+		Result.InstanceWorldToLocal			 = Result.InstanceLocalToWorld;
+		Result.InstanceWorldToLocal[0].xyz	*= Pow2(InvNonUniformScale.x);
+		Result.InstanceWorldToLocal[1].xyz	*= Pow2(InvNonUniformScale.y);
+		Result.InstanceWorldToLocal[2].xyz	*= Pow2(InvNonUniformScale.z);
+		Result.InstanceWorldToLocal[3].xyz	 = 0.0f;
+		Result.InstanceWorldToLocal			 = transpose(Result.InstanceWorldToLocal);
+		Result.InstanceWorldToLocal[3].xyz	 = mul(float4(-Result.InstanceLocalToWorld[3].xyz, 0.0f), Result.InstanceWorldToLocal).xyz;
+	}
 #else
 	Result.PrevFrameLocalToWorld = GetInstanceData(Intermediates).PrevLocalToWorld;
 	#if USE_INSTANCE_CULLING
 		Result.InstanceLocalPosition = Input.Position.xyz;
 		Result.InstanceLocalToWorld = GetInstanceData(Intermediates).LocalToWorld;
+		Result.InstanceWorldToLocal = GetInstanceData(Intermediates).WorldToLocal;
 		Result.PerInstanceParams = Intermediates.PerInstanceParams;
 		Result.PrimitiveLocalInstanceIndex = Intermediates.PrimitiveLocalInstanceIndex;
 	#endif
@@ -1059,6 +1085,13 @@ FVertexFactoryInterpolantsVSToPS VertexFactoryGetInterpolantsVSToPS(FVertexFacto
 
 	SetPrimitiveId(Interpolants, Intermediates.SceneData.PrimitiveId);
 
+#if HAS_INSTANCE_LOCAL_TO_WORLD_PS
+	SetInstanceLocalToWorld(Interpolants, VertexParameters.InstanceLocalToWorld);
+#endif
+#if HAS_INSTANCE_WORLD_TO_LOCAL_PS
+	SetInstanceWorldToLocal(Interpolants, VertexParameters.InstanceWorldToLocal);
+#endif
+
 	return Interpolants;
 }
 
diff --git a/Engine/Shaders/Private/LocalVertexFactoryCommon.ush b/Engine/Shaders/Private/LocalVertexFactoryCommon.ush
--- a/Engine/Shaders/Private/LocalVertexFactoryCommon.ush
+++ b/Engine/Shaders/Private/LocalVertexFactoryCommon.ush
@@ -40,6 +40,13 @@ struct FVertexFactoryInterpolantsVSToPS
 	nointerpolation uint HairPrimitiveId	: HAIR_PRIMITIVE_ID; // Control point ID
 	float2 HairPrimitiveUV					: HAIR_PRIMITIVE_UV; // U: parameteric distance between the two surrounding control points. V: parametric distance along the width.
 #endif
+
+#if HAS_INSTANCE_LOCAL_TO_WORLD_PS
+	nointerpolation float4 InstanceLocalToWorld[3] : INSTANCE_LOCAL_TO_WORLD;
+#endif
+#if HAS_INSTANCE_WORLD_TO_LOCAL_PS
+	nointerpolation float4 InstanceWorldToLocal[3] : INSTANCE_WORLD_TO_LOCAL;
+#endif
 };
 
 #if NUM_TEX_COORD_INTERPOLATORS || USE_PARTICLE_SUBUVS
@@ -150,3 +157,33 @@ void SetLightmapDataIndex(inout FVertexFactoryInterpolantsVSToPS Interpolants, u
 	Interpolants.LightmapDataIndex = LightmapDataIndex;
 #endif
 }
+
+#if HAS_INSTANCE_LOCAL_TO_WORLD_PS
+void SetInstanceLocalToWorld(inout FVertexFactoryInterpolantsVSToPS Interpolants, float4x4 InstanceLocalToWorld)
+{
+	float4x4 InstanceLocalToWorldT = transpose(InstanceLocalToWorld);
+	Interpolants.InstanceLocalToWorld[0] = InstanceLocalToWorldT[0];
+	Interpolants.InstanceLocalToWorld[1] = InstanceLocalToWorldT[1];
+	Interpolants.InstanceLocalToWorld[2] = InstanceLocalToWorldT[2];
+}
+
+float4x4 GetInstanceLocalToWorld(FVertexFactoryInterpolantsVSToPS Interpolants)
+{
+	return transpose(float4x4(Interpolants.InstanceLocalToWorld[0], Interpolants.InstanceLocalToWorld[1], Interpolants.InstanceLocalToWorld[2], float4(0.0f, 0.0f, 0.0f, 1.0f)));
+}
+#endif
+
+#if HAS_INSTANCE_WORLD_TO_LOCAL_PS
+void SetInstanceWorldToLocal(inout FVertexFactoryInterpolantsVSToPS Interpolants, float4x4 InstanceWorldToLocal)
+{
+	float4x4 InstanceWorldToLocalT = transpose(InstanceWorldToLocal);
+	Interpolants.InstanceWorldToLocal[0] = InstanceWorldToLocalT[0];
+	Interpolants.InstanceWorldToLocal[1] = InstanceWorldToLocalT[1];
+	Interpolants.InstanceWorldToLocal[2] = InstanceWorldToLocalT[2];
+}
+
+float4x4 GetInstanceWorldToLocal(FVertexFactoryInterpolantsVSToPS Interpolants)
+{
+	return transpose(float4x4(Interpolants.InstanceWorldToLocal[0], Interpolants.InstanceWorldToLocal[1], Interpolants.InstanceWorldToLocal[2], float4(0.0f, 0.0f, 0.0f, 1.0f)));
+}
+#endif
diff --git a/Engine/Shaders/Private/MaterialTemplate.ush b/Engine/Shaders/Private/MaterialTemplate.ush
--- a/Engine/Shaders/Private/MaterialTemplate.ush
+++ b/Engine/Shaders/Private/MaterialTemplate.ush
@@ -128,6 +128,9 @@
 
 #define IS_MESHPARTICLE_FACTORY (PARTICLE_MESH_FACTORY || NIAGARA_MESH_FACTORY)
 
+#define HAS_INSTANCE_LOCAL_TO_WORLD_PS	(NEEDS_INSTANCE_LOCAL_TO_WORLD_PS && (USE_INSTANCING || USE_INSTANCE_CULLING || IS_MESHPARTICLE_FACTORY))
+#define HAS_INSTANCE_WORLD_TO_LOCAL_PS	(NEEDS_INSTANCE_WORLD_TO_LOCAL_PS && (USE_INSTANCING || USE_INSTANCE_CULLING || IS_MESHPARTICLE_FACTORY))
+
 #define IS_HAIR_FACTORY (HAIR_STRAND_MESH_FACTORY || HAIR_CARD_MESH_FACTORY)
 
 #define TEMPLATE_USES_STRATA			(PROJECT_STRATA && MATERIAL_IS_STRATA)
@@ -362,6 +365,12 @@ struct FMaterialPixelParameters
 	float2	HairPrimitiveUV;	// U: parametric distance between the two surrounding control point. V: parametric distance along hair width
 #endif
 
+#if HAS_INSTANCE_LOCAL_TO_WORLD_PS
+	float4x4 InstanceLocalToWorld;
+#endif
+#if HAS_INSTANCE_WORLD_TO_LOCAL_PS
+	float4x4 InstanceWorldToLocal;
+#endif
 	/** Per-particle properties. Only valid for particle vertex factories. */
 	FMaterialParticleParameters Particle;
 
@@ -462,9 +471,13 @@ struct FMaterialVertexParameters
 	float3 WorldPosition;
 	// TangentToWorld[2] is WorldVertexNormal
 	half3x3 TangentToWorld;
+	
+#if USE_INSTANCING || USE_INSTANCE_CULLING || IS_MESHPARTICLE_FACTORY
+	float4x4 InstanceLocalToWorld;
+	float4x4 InstanceWorldToLocal;
+#endif
 #if USE_INSTANCING || USE_INSTANCE_CULLING
 	/** Per-instance properties. */
-	float4x4 InstanceLocalToWorld;
 	float3 InstanceLocalPosition;
 	float4 PerInstanceParams;
 #if USE_INSTANCE_CULLING
@@ -473,10 +486,6 @@ struct FMaterialVertexParameters
 	uint InstanceId;
 #endif
 	uint InstanceOffset;
-
-#elif IS_MESHPARTICLE_FACTORY 
-	/** Per-particle properties. */
-	float4x4 InstanceLocalToWorld;
 #endif
 	// If either USE_INSTANCING or (IS_MESHPARTICLE_FACTORY && FEATURE_LEVEL >= FEATURE_LEVEL_SM4)
 	// is true, PrevFrameLocalToWorld is a per-instance transform
@@ -548,6 +557,42 @@ MaterialFloat3x3 GetLocalToWorld3x3()
 	return (MaterialFloat3x3)Primitive.LocalToWorld;
 }
 
+float4x4 GetWorldToInstance(FMaterialVertexParameters Parameters)
+{
+	#if USE_INSTANCING || USE_INSTANCE_CULLING || IS_MESHPARTICLE_FACTORY
+		return Parameters.InstanceWorldToLocal;
+	#else
+		return GetPrimitiveData(Parameters).WorldToLocal;
+	#endif
+}
+
+float4x4 GetWorldToInstance(FMaterialPixelParameters Parameters)
+{
+	#if HAS_INSTANCE_WORLD_TO_LOCAL_PS
+		return Parameters.InstanceWorldToLocal;
+	#else
+		return GetPrimitiveData(Parameters).WorldToLocal;
+	#endif
+}
+
+float4x4 GetInstanceToWorld(FMaterialVertexParameters Parameters)
+{
+	#if USE_INSTANCING || USE_INSTANCE_CULLING || IS_MESHPARTICLE_FACTORY
+		return Parameters.InstanceLocalToWorld;
+	#else
+		return GetPrimitiveData(Parameters).LocalToWorld;
+	#endif
+}
+
+float4x4 GetInstanceToWorld(FMaterialPixelParameters Parameters)
+{
+	#if HAS_INSTANCE_LOCAL_TO_WORLD_PS
+		return Parameters.InstanceLocalToWorld;
+	#else
+		return GetPrimitiveData(Parameters).LocalToWorld;
+	#endif
+}
+
 float3 GetTranslatedWorldPosition(FMaterialVertexParameters Parameters)
 {
 	return Parameters.WorldPosition;
diff --git a/Engine/Shaders/Private/MeshParticleVertexFactory.ush b/Engine/Shaders/Private/MeshParticleVertexFactory.ush
--- a/Engine/Shaders/Private/MeshParticleVertexFactory.ush
+++ b/Engine/Shaders/Private/MeshParticleVertexFactory.ush
@@ -10,6 +10,9 @@
 #define USE_PARTICLE_VELOCITY (NEEDS_PARTICLE_VELOCITY)
 #define USE_PARTICLE_TIME (NEEDS_PARTICLE_TIME)
 
+#define USE_PARTICLE_LOCAL_TO_WORLD (NEEDS_PARTICLE_LOCAL_TO_WORLD || NEEDS_INSTANCE_LOCAL_TO_WORLD_PS)
+#define USE_PARTICLE_WORLD_TO_LOCAL (NEEDS_PARTICLE_WORLD_TO_LOCAL || NEEDS_INSTANCE_WORLD_TO_LOCAL_PS)
+
 // Only support previous transform for deferred feature level
 #define USE_PREVIOUS_TRANSFORM (FEATURE_LEVEL >= FEATURE_LEVEL_SM4)
 
@@ -128,8 +131,12 @@ struct FVertexFactoryInterpolantsVSToPS
 	float ParticleTime : PARTICLE_TIME;
 #endif
 
-#if NEEDS_PARTICLE_LOCAL_TO_WORLD || NEEDS_PARTICLE_WORLD_TO_LOCAL
-	float4 LocalToWorld[3] : PARTICLE_TRANSFORM;
+#if USE_PARTICLE_LOCAL_TO_WORLD
+	nointerpolation float4 ParticleToWorld[3] : PARTICLE_LOCAL_TO_WORLD;
+#endif
+
+#if USE_PARTICLE_WORLD_TO_LOCAL
+	nointerpolation float4 WorldToParticle[3] : PARTICLE_WORLD_TO_LOCAL;
 #endif
 
 #if INSTANCED_STEREO
@@ -167,9 +174,9 @@ struct FVertexFactoryIntermediates
 	/** Relative time. */
 	float RelativeTime;
 
-#if NEEDS_PARTICLE_LOCAL_TO_WORLD || NEEDS_PARTICLE_WORLD_TO_LOCAL
-	float4x4 LocalToWorld;
-#endif
+	float4x4 ParticleToWorld;
+	float4x4 WorldToParticle;
+
 	/** Cached primitive and instance data */
 	FSceneDataIntermediates SceneData;
 };
@@ -214,26 +221,23 @@ float4x4 GetParticlePrevTransform(FVertexFactoryInput Input)
 }
 #endif
 
-float4x4 PseudoInverseTransform(float4x4 InTransform)
+float4x4 PseudoInverseTransform(float4x4 LocalToWorld)
 {
-	// we'll assume that the incoming transform is composed of Scale * Rotation * Translation
-	// we'll additionally assume that the scale is uniform
-	float3x3 InvRotScale = transpose((float3x3) InTransform);
-
-	float4x4 InvRotScale44 = float4x4(
-		InvRotScale[0], 0.0f,
-		InvRotScale[1], 0.0f,
-		InvRotScale[2], 0.0f,
-		0.0f, 0.0f, 0.0f, 1.0f);
-
-	float4x4 InvTransMat = float4x4(
-		1.0f, 0.0f, 0.0f, 0.0f,
-		0.0f, 1.0f, 0.0f, 0.0f,
-		0.0f, 0.0f, 1.0f, 0.0f,
-		-InTransform[3].xyz, 1.0f
-	);
-
-	return mul(InvTransMat, InvRotScale44);
+	float3 Scale2;
+	Scale2.x = length2(LocalToWorld[0].xyz);
+	Scale2.y = length2(LocalToWorld[1].xyz);
+	Scale2.z = length2(LocalToWorld[2].xyz);
+
+	float3 InvNonUniformScale = rsqrt(Scale2);
+
+	float4x4 WorldToLocal = 0;
+	WorldToLocal[0].xyz = LocalToWorld[0].xyz * Pow2(InvNonUniformScale.x);
+	WorldToLocal[1].xyz	= LocalToWorld[1].xyz * Pow2(InvNonUniformScale.y);
+	WorldToLocal[2].xyz	= LocalToWorld[2].xyz * Pow2(InvNonUniformScale.z);
+	WorldToLocal[3].w	= 1.0f;
+	WorldToLocal		= transpose(WorldToLocal);
+	WorldToLocal[3].xyz	= mul(float4(-LocalToWorld[3].xyz, 0.0f), WorldToLocal).xyz;
+	return WorldToLocal;
 }
 
 /** Converts from vertex factory specific interpolants to a FMaterialPixelParameters, which is used by material inputs. */
@@ -297,19 +301,23 @@ FMaterialPixelParameters GetMaterialPixelParameters(FVertexFactoryInterpolantsVS
 	Result.Particle.RelativeTime = Interpolants.ParticleTime;
 #endif
 
-#if NEEDS_PARTICLE_LOCAL_TO_WORLD || NEEDS_PARTICLE_WORLD_TO_LOCAL
-	float4x4 ParticleTransform = transpose(float4x4(
-		Interpolants.LocalToWorld[0],
-		Interpolants.LocalToWorld[1],
-		Interpolants.LocalToWorld[2],
-		float4(0.0f, 0.0f, 0.0f, 1.0f)));
-
+#if USE_PARTICLE_LOCAL_TO_WORLD
+	float4x4 ParticleToWorld = transpose(float4x4(Interpolants.ParticleToWorld[0], Interpolants.ParticleToWorld[1], Interpolants.ParticleToWorld[2], float4(0.0f, 0.0f, 0.0f, 1.0f)));
 	#if NEEDS_PARTICLE_LOCAL_TO_WORLD
-		Result.Particle.ParticleToWorld = ParticleTransform;
+		Result.Particle.ParticleToWorld = ParticleToWorld;
 	#endif
+	#if NEEDS_INSTANCE_LOCAL_TO_WORLD_PS
+		Result.InstanceLocalToWorld = ParticleToWorld;
+	#endif
+#endif
 
+#if USE_PARTICLE_WORLD_TO_LOCAL
+	float4x4 WorldToParticle = transpose(float4x4(Interpolants.WorldToParticle[0], Interpolants.WorldToParticle[1], Interpolants.WorldToParticle[2], float4(0.0f, 0.0f, 0.0f, 1.0f)));
 	#if NEEDS_PARTICLE_WORLD_TO_LOCAL
-		Result.Particle.WorldToParticle = PseudoInverseTransform(ParticleTransform);
+		Result.Particle.WorldToParticle = WorldToParticle;
+	#endif
+	#if NEEDS_INSTANCE_WORLD_TO_LOCAL_PS
+		Result.InstanceWorldToLocal = WorldToParticle;
 	#endif
 #endif
 
@@ -345,14 +353,10 @@ FMaterialVertexParameters GetMaterialVertexParameters(FVertexFactoryInput Input,
 
 	float4x4 ParticleTransform = GetParticleTransform(Input);
 	Result.InstanceLocalToWorld = ParticleTransform;
+	Result.InstanceWorldToLocal = PseudoInverseTransform(ParticleTransform);
 
-#if NEEDS_PARTICLE_LOCAL_TO_WORLD
-	Result.Particle.ParticleToWorld = ParticleTransform;
-#endif
-
-#if NEEDS_PARTICLE_WORLD_TO_LOCAL
-	Result.Particle.WorldToParticle = PseudoInverseTransform(ParticleTransform);
-#endif
+	Result.Particle.ParticleToWorld = Result.InstanceLocalToWorld;
+	Result.Particle.WorldToParticle = Result.InstanceWorldToLocal;
 
 	Result.TangentToWorld = mul(TangentToLocal, (float3x3)ParticleTransform ); 
 
@@ -405,9 +409,8 @@ FVertexFactoryIntermediates GetVertexFactoryIntermediates(FVertexFactoryInput In
 	// World position and localtoworld transform.
 	float4x4 LocalToWorldMat = GetParticleTransform(Input);
 
-#if NEEDS_PARTICLE_LOCAL_TO_WORLD || NEEDS_PARTICLE_WORLD_TO_LOCAL
-	Intermediates.LocalToWorld = LocalToWorldMat;
-#endif
+	Intermediates.ParticleToWorld = LocalToWorldMat;
+	Intermediates.WorldToParticle = PseudoInverseTransform(LocalToWorldMat);
 
 	float3 ParticleWorldPosition = LocalToWorldMat[3].xyz;
 
@@ -591,11 +594,18 @@ FVertexFactoryInterpolantsVSToPS VertexFactoryGetInterpolantsVSToPS(FVertexFacto
 	Interpolants.ParticleTime = Intermediates.RelativeTime;
 #endif
 
-#if NEEDS_PARTICLE_LOCAL_TO_WORLD || NEEDS_PARTICLE_WORLD_TO_LOCAL
-	float4x4 LocalToWorldT = transpose(Intermediates.LocalToWorld);
-	Interpolants.LocalToWorld[0] = LocalToWorldT[0];
-	Interpolants.LocalToWorld[1] = LocalToWorldT[1];
-	Interpolants.LocalToWorld[2] = LocalToWorldT[2];
+#if USE_PARTICLE_LOCAL_TO_WORLD
+	float4x4 ParticleToWorldT = transpose(Intermediates.ParticleToWorld);
+	Interpolants.ParticleToWorld[0] = ParticleToWorldT[0];
+	Interpolants.ParticleToWorld[1] = ParticleToWorldT[1];
+	Interpolants.ParticleToWorld[2] = ParticleToWorldT[2];
+#endif
+
+#if USE_PARTICLE_WORLD_TO_LOCAL
+	float4x4 WorldToParticleT = transpose(Intermediates.WorldToParticle);
+	Interpolants.WorldToParticle[0] = WorldToParticleT[0];
+	Interpolants.WorldToParticle[1] = WorldToParticleT[1];
+	Interpolants.WorldToParticle[2] = WorldToParticleT[2];
 #endif
 
 #if INSTANCED_STEREO
diff --git a/Engine/Source/Runtime/Engine/Classes/Materials/MaterialExpressionTransform.h b/Engine/Source/Runtime/Engine/Classes/Materials/MaterialExpressionTransform.h
--- a/Engine/Source/Runtime/Engine/Classes/Materials/MaterialExpressionTransform.h
+++ b/Engine/Source/Runtime/Engine/Classes/Materials/MaterialExpressionTransform.h
@@ -27,8 +27,11 @@ enum EMaterialVectorCoordTransformSource
 	/** Camera space */
 	TRANSFORMSOURCE_Camera  UMETA(DisplayName="Camera Space"),
 
-	/** Particle space */
-	TRANSFORMSOURCE_ParticleWorld  UMETA(DisplayName = "Mesh particle space"),
+	/** Particle space, deprecated value will be removed in a future release use instance space. */
+	TRANSFORMSOURCE_ParticleWorld  UMETA(Hidden, DisplayName = "Mesh particle space"),
+
+	/** Instance space (used to provide per instance transform, i.e. for Instanced Static Mesh / Particles). */
+	TRANSFORMSOURCE_Instance UMETA(DisplayName = "Instance & Particle Space"),
 
 	TRANSFORMSOURCE_MAX,
 };
@@ -51,8 +54,11 @@ enum EMaterialVectorCoordTransform
 	/** Camera space */
 	TRANSFORM_Camera  UMETA(DisplayName="Camera Space"),
 
-	/** Particle space */
-	TRANSFORM_ParticleWorld UMETA(DisplayName = "Mesh particle space"),
+	/** Particle space, deprecated value will be removed in a future release use instance space. */
+	TRANSFORM_ParticleWorld UMETA(Hidden, DisplayName = "Mesh particle space"),
+
+	/** Instance space (used to provide per instance transform, i.e. for Instanced Static Mesh / Particles). */
+	TRANSFORM_Instance UMETA(DisplayName = "Instance & Particle Space"),
 
 	TRANSFORM_MAX,
 };
diff --git a/Engine/Source/Runtime/Engine/Classes/Materials/MaterialExpressionTransformPosition.h b/Engine/Source/Runtime/Engine/Classes/Materials/MaterialExpressionTransformPosition.h
--- a/Engine/Source/Runtime/Engine/Classes/Materials/MaterialExpressionTransformPosition.h
+++ b/Engine/Source/Runtime/Engine/Classes/Materials/MaterialExpressionTransformPosition.h
@@ -27,8 +27,11 @@ enum EMaterialPositionTransformSource
 	/** Camera space */
 	TRANSFORMPOSSOURCE_Camera  UMETA(DisplayName="Camera Space"),
 
-	/** Particle space */
-	TRANSFORMPOSSOURCE_Particle UMETA(DisplayName = "Mesh Particle Space"),
+	/** Particle space, deprecated value will be removed in a future release use instance space. */
+	TRANSFORMPOSSOURCE_Particle UMETA(Hidden, DisplayName = "Mesh Particle Space"),
+
+	/** Instance space (used to provide per instance transform, i.e. for Instanced Static Mesh / Particles). */
+	TRANSFORMPOSSOURCE_Instance UMETA(DisplayName = "Instance & Particle Space"),
 
 	TRANSFORMPOSSOURCE_MAX,
 };
diff --git a/Engine/Source/Runtime/Engine/Private/Materials/HLSLMaterialTranslator.cpp b/Engine/Source/Runtime/Engine/Private/Materials/HLSLMaterialTranslator.cpp
--- a/Engine/Source/Runtime/Engine/Private/Materials/HLSLMaterialTranslator.cpp
+++ b/Engine/Source/Runtime/Engine/Private/Materials/HLSLMaterialTranslator.cpp
@@ -177,6 +177,8 @@ FHLSLMaterialTranslator::FHLSLMaterialTranslator(FMaterial* InMaterial,
 ,	bUsesParticleColor(false)
 ,	bUsesParticleLocalToWorld(false)
 ,	bUsesParticleWorldToLocal(false)
+,	bUsesInstanceLocalToWorldPS(false)
+,	bUsesInstanceWorldToLocalPS(false)
 ,	bUsesVertexPosition(false)
 ,	bUsesTransformVector(false)
 ,	bCompilingPreviousFrame(false)
@@ -1490,6 +1492,8 @@ void FHLSLMaterialTranslator::GetMaterialEnvironment(EShaderPlatform InPlatform,
 	OutEnvironment.SetDefine(TEXT("NEEDS_PARTICLE_COLOR"), bUsesParticleColor); 
 	OutEnvironment.SetDefine(TEXT("NEEDS_PARTICLE_LOCAL_TO_WORLD"), bUsesParticleLocalToWorld);
 	OutEnvironment.SetDefine(TEXT("NEEDS_PARTICLE_WORLD_TO_LOCAL"), bUsesParticleWorldToLocal);
+	OutEnvironment.SetDefine(TEXT("NEEDS_INSTANCE_LOCAL_TO_WORLD_PS"), bUsesInstanceLocalToWorldPS);
+	OutEnvironment.SetDefine(TEXT("NEEDS_INSTANCE_WORLD_TO_LOCAL_PS"), bUsesInstanceWorldToLocalPS);
 	OutEnvironment.SetDefine(TEXT("USES_TRANSFORM_VECTOR"), bUsesTransformVector);
 	OutEnvironment.SetDefine(TEXT("WANT_PIXEL_DEPTH_OFFSET"), bUsesPixelDepthOffset);
 	if (IsMetalPlatform(InPlatform))
@@ -7904,6 +7908,11 @@ int32 FHLSLMaterialTranslator::TransformBase(EMaterialCommonBasis SourceCoordBas
 				CodeStr = TEXT("mul(<A>, <MATRIX>(Parameters.Particle.WorldToParticle))");
 				bUsesParticleWorldToLocal = true;
 			}
+			else if (DestCoordBasis == MCB_Instance)
+			{
+				CodeStr = TEXT("mul(<A>, <MATRIX>GetWorldToInstance(Parameters))");
+				bUsesInstanceWorldToLocalPS = ShaderFrequency == SF_Pixel;
+			}
 
 			// else use MCB_TranslatedWorld as intermediary basis
 			IntermediaryBasis = MCB_TranslatedWorld;
@@ -7939,6 +7948,17 @@ int32 FHLSLMaterialTranslator::TransformBase(EMaterialCommonBasis SourceCoordBas
 			// use World as an intermediary base
 			break;
 		}
+		case MCB_Instance:
+		{
+			if (DestCoordBasis == MCB_World)
+			{
+				CodeStr = TEXT("mul(<A>, <MATRIX>GetInstanceToWorld(Parameters))");
+				bUsesInstanceLocalToWorldPS = ShaderFrequency == SF_Pixel;
+			}
+			// use World as an intermediary base
+			break;
+		}
+
 		default:
 			check(0);
 			break;
diff --git a/Engine/Source/Runtime/Engine/Private/Materials/HLSLMaterialTranslator.h b/Engine/Source/Runtime/Engine/Private/Materials/HLSLMaterialTranslator.h
--- a/Engine/Source/Runtime/Engine/Private/Materials/HLSLMaterialTranslator.h
+++ b/Engine/Source/Runtime/Engine/Private/Materials/HLSLMaterialTranslator.h
@@ -354,6 +354,11 @@ protected:
 	/** true if the material reads mesh particle world to local in the pixel shader. */
 	uint32 bUsesParticleWorldToLocal : 1;
 
+	/** true if the material reads per instance local to world in the pixel shader. */
+	uint32 bUsesInstanceLocalToWorldPS : 1;
+	/** true if the material reads per instance world to local in the pixel shader. */
+	uint32 bUsesInstanceWorldToLocalPS : 1;
+
 	/** true if the material uses any type of vertex position */
 	uint32 bUsesVertexPosition : 1;
 
diff --git a/Engine/Source/Runtime/Engine/Private/Materials/MaterialExpressions.cpp b/Engine/Source/Runtime/Engine/Private/Materials/MaterialExpressions.cpp
--- a/Engine/Source/Runtime/Engine/Private/Materials/MaterialExpressions.cpp
+++ b/Engine/Source/Runtime/Engine/Private/Materials/MaterialExpressions.cpp
@@ -10401,7 +10401,8 @@ static EMaterialCommonBasis GetMaterialCommonBasis(EMaterialVectorCoordTransform
 		MCB_World,						// TRANSFORMSOURCE_World
 		MCB_View,						// TRANSFORMSOURCE_View
 		MCB_Camera,						// TRANSFORMSOURCE_Camera
-		MCB_MeshParticle,
+		MCB_MeshParticle,				// TRANSFORMSOURCE_Particle
+		MCB_Instance,					// TRANSFORMSOURCE_Instance
 	};
 	return ConversionTable[X];
 }
@@ -10414,7 +10415,8 @@ static EMaterialCommonBasis GetMaterialCommonBasis(EMaterialVectorCoordTransform
 		MCB_World,						// TRANSFORM_World
 		MCB_View,						// TRANSFORM_View
 		MCB_Camera,						// TRANSFORM_Camera
-		MCB_MeshParticle,
+		MCB_MeshParticle,				// TRANSFORM_Particle
+		MCB_Instance,					// TRANSFORM_Instance
 	};
 	return ConversionTable[X];
 }
@@ -10516,7 +10518,8 @@ static EMaterialCommonBasis GetMaterialCommonBasis(EMaterialPositionTransformSou
 		MCB_TranslatedWorld,			// TRANSFORMPOSSOURCE_TranslatedWorld
 		MCB_View,						// TRANSFORMPOSSOURCE_View
 		MCB_Camera,						// TRANSFORMPOSSOURCE_Camera
-		MCB_MeshParticle,	
+		MCB_MeshParticle,				// TRANSFORMPOSSOURCE_Particle
+		MCB_Instance,					// TRANSFORMPOSSOURCE_Instance
 	};
 	return ConversionTable[X];
 }
diff --git a/Engine/Source/Runtime/Engine/Public/MaterialShared.h b/Engine/Source/Runtime/Engine/Public/MaterialShared.h
--- a/Engine/Source/Runtime/Engine/Public/MaterialShared.h
+++ b/Engine/Source/Runtime/Engine/Public/MaterialShared.h
@@ -179,6 +179,7 @@ enum EMaterialCommonBasis
 	MCB_View,
 	MCB_Camera,
 	MCB_MeshParticle,
+	MCB_Instance,
 	MCB_MAX,
 };
 
