Fix a bug of DOF jittering when AA method is TAA on mobile.

#p4v-preflight-copy 17759548

diff --git a/Engine/Shaders/Private/PostProcessDOF.usf b/Engine/Shaders/Private/PostProcessDOF.usf
--- a/Engine/Shaders/Private/PostProcessDOF.usf
+++ b/Engine/Shaders/Private/PostProcessDOF.usf
@@ -104,7 +104,7 @@ void SetupPS(
 	// Manually compute UV and ScreenSpacePos from SvPosition to avoid south east shiting of the pixel in bottom right corner
 	// caused by the DestRect.Size() + 1 viewport size done for odd view rect resolution.
 	float2 ViewportUV = min(SvPositionToViewportUV(SvPosition * DoFToScreenScale), 1 - View.ViewSizeAndInvSize.zw * 0.5);
-	float2 UV = ViewportUVToBufferUV(ViewportUV);
+	float2 UV = ViewportUVToBufferUV(ViewportUV) - ((0.5f - View.TemporalAAParams.zw) * BufferSizeAndInvSize.zw);
 	float2 ScreenSpacePos = DoFToScreenScale * float2(ViewportUV.x, 1 - ViewportUV.y) - 1;
 
 	float2 Offset = 0.5f * BufferSizeAndInvSize.zw;
