Fix invalid tile classification with partially covering pixel.

diff --git a/Engine/Shaders/Private/HairStrands/HairStrandsVisibilityTile.usf b/Engine/Shaders/Private/HairStrands/HairStrandsVisibilityTile.usf
--- a/Engine/Shaders/Private/HairStrands/HairStrandsVisibilityTile.usf
+++ b/Engine/Shaders/Private/HairStrands/HairStrandsVisibilityTile.usf
@@ -46,7 +46,7 @@ void TileMainCS(uint2 DispatchThreadId : SV_DispatchThreadID, uint LinearIndex :
 		{
 			const float Transmittance = InputFloatTexture.Load(uint3(PixelCoord, 0));
 			fHair = Transmittance <= TransmittanceThreshold;
-			pHair = Transmittance < 1.f && Transmittance > 0.f;
+			pHair = Transmittance < 1.f;
 		}
 	}
 
@@ -93,7 +93,7 @@ void TileMainCS(uint2 DispatchThreadId : SV_DispatchThreadID, uint LinearIndex :
 
 		const bool bIsFull = all(fMask == 0xFFFFFFFF);
 		const bool bIsPartial = any(pMask > 0);
-
+			
 		if (bIsFull || bIsPartial)
 		{
 			uint WriteToIndex;
