Fixed RVT compression on a windows OpenGL

diff --git a/Engine/Shaders/Private/VirtualTextureCompress.usf b/Engine/Shaders/Private/VirtualTextureCompress.usf
--- a/Engine/Shaders/Private/VirtualTextureCompress.usf
+++ b/Engine/Shaders/Private/VirtualTextureCompress.usf
@@ -44,11 +44,7 @@ uint4 PackRG32(in uint2 u2)
 	uint y0 = (u2.y >> 16) & 0x0000ffff;
 	uint x1 = (u2.x) & 0x0000ffff;
 	uint y1 = (u2.y) & 0x0000ffff;
-#if ETC_PROFILE
 	return uint4(x1,x0,y1,y0);
-#else
-	return uint4(x0,x1,y0,y1);
-#endif
 }
 #else
 uint2 PackRG32(in uint2 u2)
