Update to the hash functions Rand3DPCG32 and Rand4DPCG32 to match the versions presented in the paper Hash Functions for GPU Rendering by [REDACTED_PERSON] and Marc Olano http://jcgt.org/published/0009/03/02/

The initial versions of Rand3DPCG32 and Rand4DPCG32 were implemented in UE4 by coauthor Marc Olano while on sabbatical.

"In our paper we tested the updated versions of Rand3DPCG32 and Rand4DPCG32, respectively called pcg3d and pcg4d in the paper. The versions presented in the paper and included in this PR have higher quality and faster GPU runtime performance than the version currently in UE4. See our paper for testing methods and results."

The changes in this PR come without any IP restrictions or encumbrances.

#GitHub 7892
#7892

diff --git a/Engine/Shaders/Private/Random.ush b/Engine/Shaders/Private/Random.ush
--- a/Engine/Shaders/Private/Random.ush
+++ b/Engine/Shaders/Private/Random.ush
@@ -104,11 +104,9 @@ uint3 Rand3DPCG16(int3 p)
 
 // 3D random number generator inspired by PCGs (permuted congruential generator)
 // Using a **simple** Feistel cipher in place of the usual xor shift permutation step
+// http://jcgt.org/published/0009/03/02/
 // @param v = 3D integer coordinate
 // @return three elements w/ 32 random bits each (0-0xffffffff).
-// ~12 ALU operations for result.x   (10 mad, 3 >>)
-// ~14 ALU operations for result.xy  (11 mad, 3 >>)
-// ~15 ALU operations for result.xyz (12 mad, 3 >>)
 uint3 Rand3DPCG32(int3 p)
 {
 	// taking a signed int then reinterpreting as unsigned gives good behavior for negatives
@@ -117,13 +115,15 @@ uint3 Rand3DPCG32(int3 p)
 	// Linear congruential step.
 	v = v * 1664525u + 1013904223u;
 
-	// swapping low and high bits makes all 32 bits pretty good
-	v = v * (1u << 16u) + (v >> 16u);
-
-	// final shuffle
+	// shuffle
 	v.x += v.y*v.z;
 	v.y += v.z*v.x;
 	v.z += v.x*v.y;
+
+	// xoring high bits into low bits makes all 32 bits pretty good
+	v ^= v >> 16u;
+
+	// final shuffle
 	v.x += v.y*v.z;
 	v.y += v.z*v.x;
 	v.z += v.x*v.y;
@@ -135,11 +135,9 @@ uint3 Rand3DPCG32(int3 p)
 
 // 4D random number generator inspired by PCGs (permuted congruential generator)
 // Using a **simple** Feistel cipher in place of the usual xor shift permutation step
+// http://jcgt.org/published/0009/03/02/
 // @param v = 4D integer coordinate
 // @return four elements w/ 32 random bits each (0-0xffffffff).
-// ~12 ALU operations for result.x   (10 mad, 3 >>)
-// ~14 ALU operations for result.xy  (11 mad, 3 >>)
-// ~15 ALU operations for result.xyz (12 mad, 3 >>)
 uint4 Rand4DPCG32(int4 p)
 {
 	// taking a signed int then reinterpreting as unsigned gives good behavior for negatives
@@ -148,14 +146,16 @@ uint4 Rand4DPCG32(int4 p)
 	// Linear congruential step.
 	v = v * 1664525u + 1013904223u;
 
-	// swapping low and high bits makes all 32 bits pretty good
-	v = v * (1u << 16u) + (v >> 16u);
-
-	// final shuffle
+	// shuffle
 	v.x += v.y*v.w;
 	v.y += v.z*v.x;
 	v.z += v.x*v.y;
 	v.w += v.y*v.z;
+
+	// xoring high bits into low makes all 32 bits pretty good
+	v ^= (v >> 16u);
+
+	// final shuffle
 	v.x += v.y*v.w;
 	v.y += v.z*v.x;
 	v.z += v.x*v.y;
