Fix out of bound write during hair voxelization.

diff --git a/Engine/Shaders/Private/HairStrands/HairStrandsVoxelOpaque.usf b/Engine/Shaders/Private/HairStrands/HairStrandsVoxelOpaque.usf
--- a/Engine/Shaders/Private/HairStrands/HairStrandsVoxelOpaque.usf
+++ b/Engine/Shaders/Private/HairStrands/HairStrandsVoxelOpaque.usf
@@ -47,7 +47,8 @@ void MainCS(uint3 DispatchThreadId : SV_DispatchThreadID)
 	const uint4 PageIndexCoord = VirtualVoxelParams_PageIndexCoordBuffer[NodeDesc.PageIndexOffset + AllocatedPageIndex];
 	const uint  LinearPageIndexCoord = CoordToIndex(PageIndexCoord.xyz, NodeDesc.PageIndexResolution, NodeDesc.PageIndexOffset);
 
-	const bool bIsValid = LinearPageIndexCoord < VirtualVoxelParams_PageIndexCount;
+	// PageIndexCoord have invalid .w component if we run out of available page during the allocation
+	const bool bIsValid = LinearPageIndexCoord < VirtualVoxelParams_PageIndexCount && PageIndexCoord.w != INVALID_MACRO_GROUP_ID;
 	if (bIsValid)
 	{
 		const uint PageIndex = VirtualVoxelParams_PageIndexBuffer.Load(LinearPageIndexCoord);
diff --git a/Engine/Shaders/Private/HairStrands/HairStrandsVoxelPageAllocation.usf b/Engine/Shaders/Private/HairStrands/HairStrandsVoxelPageAllocation.usf
--- a/Engine/Shaders/Private/HairStrands/HairStrandsVoxelPageAllocation.usf
+++ b/Engine/Shaders/Private/HairStrands/HairStrandsVoxelPageAllocation.usf
@@ -593,13 +593,20 @@ void AllocateCS(uint GroupIndex : SV_GroupIndex, uint3 DispatchThreadId : SV_Dis
 	if (bIsValid)
 	{
 		const uint PageIndex = GroupBase + Offset;
-		PageIndexBuffer[GridIndex] = PageIndex < PageCount ? PageIndex : INVALID_VOXEL_PAGE_INDEX;
-		PageToPageIndexBuffer[PageIndex] = GridIndex;
+		const bool bIsAllocationValid = PageIndex < PageCount;
+
+		PageIndexBuffer[GridIndex] = bIsAllocationValid ? PageIndex : INVALID_VOXEL_PAGE_INDEX;
+		if (bIsAllocationValid)
+		{
+			PageToPageIndexBuffer[PageIndex] = GridIndex;
+		}
+
 		// Output the coordinates of the allocated page for indirect dispatch usage
+		// If the allocated failed (run out of page), then we mark the IndexCoord with a invalid GroupID
 		const uint LinearIndex = DispatchThreadId.x;
 		const uint3 PageIndexCoord = IndexToCoord(LinearIndex, PageIndexResolution);
 		const uint IndexRelativeToMacroGroupIndexOffset = PageIndex - PageIndexGlobalCounter[1];
-		PageIndexCoordBuffer[PageIndexOffset + IndexRelativeToMacroGroupIndexOffset] = uint4(PageIndexCoord, MacroGroupId);
+		PageIndexCoordBuffer[PageIndexOffset + IndexRelativeToMacroGroupIndexOffset] = uint4(PageIndexCoord, bIsAllocationValid ? MacroGroupId : INVALID_MACRO_GROUP_ID);
 	}
 	// Mark page index as invalid
 	// Insure that even if write more (due to larger dispatch count that needed), we do not stomp other instance group page index
diff --git a/Engine/Shaders/Private/HairStrands/HairStrandsVoxelPageCommon.ush b/Engine/Shaders/Private/HairStrands/HairStrandsVoxelPageCommon.ush
--- a/Engine/Shaders/Private/HairStrands/HairStrandsVoxelPageCommon.ush
+++ b/Engine/Shaders/Private/HairStrands/HairStrandsVoxelPageCommon.ush
@@ -58,6 +58,7 @@ bool IsInVoxelBounds(float3 P, float3 MinP, float3 MaxP)
 ///////////////////////////////////////////////////////////////////////////////////////////////////
 
 #define INVALID_VOXEL_PAGE_INDEX 0xFFFFFFFF
+#define INVALID_MACRO_GROUP_ID 0xFF
 
 #define VOXEL_HAIR_MASK		0x00FFFFFF
 #define VOXEL_OPAQUE_MASK	0xFF000000
diff --git a/Engine/Shaders/Private/HairStrands/HairStrandsVoxelRasterCompute.usf b/Engine/Shaders/Private/HairStrands/HairStrandsVoxelRasterCompute.usf
--- a/Engine/Shaders/Private/HairStrands/HairStrandsVoxelRasterCompute.usf
+++ b/Engine/Shaders/Private/HairStrands/HairStrandsVoxelRasterCompute.usf
@@ -36,7 +36,6 @@ uint3			VirtualVoxelParams_PageCountResolution;
 uint3			VirtualVoxelParams_PageTextureResolution;
 
 Buffer<uint>	VirtualVoxelParams_PageIndexBuffer;
-Buffer<uint4>	VirtualVoxelParams_PageIndexCoordBuffer;
 StructuredBuffer<FPackedVirtualVoxelNodeDesc> VirtualVoxelParams_NodeDescBuffer;
 
 RWTexture3D<uint> OutPageTexture;
