Fixed camera world position not working for light functions.

diff --git a/Engine/Shaders/Private/LightFunctionPixelShader.usf b/Engine/Shaders/Private/LightFunctionPixelShader.usf
--- a/Engine/Shaders/Private/LightFunctionPixelShader.usf
+++ b/Engine/Shaders/Private/LightFunctionPixelShader.usf
@@ -21,6 +21,8 @@ void Main(
 	out float4 OutColor : SV_Target0
 	)
 {
+	ResolvedView = ResolveView();
+
 	float2 ScreenUV = SvPositionToBufferUV(SvPosition);
 
 	// make SvPosition appear to be rasterized with the depth from the depth buffer
