Fix the transition regions for skinned triangles that are wrap deformed by physics triangles which are partially pinned via zero max distance setting. This allowes for a smoother transition from a purely skinned to cloth simulated domain on the skinned mesh.

diff --git a/Engine/Shaders/Private/GpuSkinCacheComputeShader.usf b/Engine/Shaders/Private/GpuSkinCacheComputeShader.usf
--- a/Engine/Shaders/Private/GpuSkinCacheComputeShader.usf
+++ b/Engine/Shaders/Private/GpuSkinCacheComputeShader.usf
@@ -376,13 +376,13 @@ void SkinCacheUpdateBatchCS(uint3 GroupID : SV_GroupID,
 
     float3 ClothTangentX = float3(0,0,0);
     float3 ClothTangentZ = float3(0,0,0);
-	float TempClothBlendWeight = 0.0f;
 
     Intermediates.SimulatedPosition = float3(0,0,0);
     float3 NormalPosition = float3(0,0,0);
     float3 TangentPosition = float3(0,0,0);
     int NumInfluences = 0;
 	float SumWeights = 0.0;
+	float ClothWeight = 0.0;
 	
     for (int i = 0; i < NUM_INFLUENCES_PER_VERTEX; ++i )
     {
@@ -392,8 +392,6 @@ void SkinCacheUpdateBatchCS(uint3 GroupID : SV_GroupID,
 	    { 
             ++NumInfluences;
 
-		    TempClothBlendWeight = ClothBlendWeight;
-
 		    float3 A = float3(ClothPositionsAndNormalsBuffer[Influence.SourceMeshVertIndices.x * 3], ClothPositionsAndNormalsBuffer[Influence.SourceMeshVertIndices.x * 3 + 1].x);
 		    float3 B = float3(ClothPositionsAndNormalsBuffer[Influence.SourceMeshVertIndices.y * 3], ClothPositionsAndNormalsBuffer[Influence.SourceMeshVertIndices.y * 3 + 1].x);
 		    float3 C = float3(ClothPositionsAndNormalsBuffer[Influence.SourceMeshVertIndices.z * 3], ClothPositionsAndNormalsBuffer[Influence.SourceMeshVertIndices.z * 3 + 1].x);
@@ -402,6 +400,8 @@ void SkinCacheUpdateBatchCS(uint3 GroupID : SV_GroupID,
 		    float3 NB = float3(ClothPositionsAndNormalsBuffer[Influence.SourceMeshVertIndices.y * 3 + 1].y, ClothPositionsAndNormalsBuffer[Influence.SourceMeshVertIndices.y * 3 + 2]);
 		    float3 NC = float3(ClothPositionsAndNormalsBuffer[Influence.SourceMeshVertIndices.z * 3 + 1].y, ClothPositionsAndNormalsBuffer[Influence.SourceMeshVertIndices.z * 3 + 2]);
 
+	    	ClothWeight += ClothBlendWeight * (1.0f - (Influence.SourceMeshVertIndices.w / 65535.0f));
+
 			float Weight = 1.0f;
 			if ( NUM_INFLUENCES_PER_VERTEX > 1 )
 			{
@@ -453,8 +453,9 @@ void SkinCacheUpdateBatchCS(uint3 GroupID : SV_GroupID,
 	else
 	{
 		Intermediates.SimulatedPosition = float3(0, 0, 0);
-		TempClothBlendWeight = 0.0f;
 	}
+	
+	ClothWeight /= NUM_INFLUENCES_PER_VERTEX;
 #endif // GPUSKIN_APEX_CLOTH
 #endif // GPUSKIN_MORPH_BLEND
 
@@ -464,7 +465,7 @@ void SkinCacheUpdateBatchCS(uint3 GroupID : SV_GroupID,
 
 #if GPUSKIN_APEX_CLOTH
     float4 TransformedSkinPos = mul(float4(SPos, 1), ClothLocalToWorld);
-    float3 BlendedPos = lerp(TransformedSkinPos.xyz, Intermediates.SimulatedPosition, TempClothBlendWeight);
+    float3 BlendedPos = lerp(TransformedSkinPos.xyz, Intermediates.SimulatedPosition, ClothWeight);
     SPos = mul(float4(BlendedPos, 1), ClothWorldToLocal).xyz;
 #endif
 
@@ -472,8 +473,8 @@ void SkinCacheUpdateBatchCS(uint3 GroupID : SV_GroupID,
 	float3 TangentZ = normalize(mul((float3x3)Intermediates.BlendMatrix, LocalTangentZ));
 	
 #if GPUSKIN_APEX_CLOTH
-    TangentX = ClothTangentX * TempClothBlendWeight + TangentX * (1.0f - TempClothBlendWeight);
-    TangentZ = ClothTangentZ * TempClothBlendWeight + TangentZ * (1.0f - TempClothBlendWeight);
+    TangentX = ClothTangentX * ClothWeight + TangentX * (1.0f - ClothWeight);
+    TangentZ = ClothTangentZ * ClothWeight + TangentZ * (1.0f - ClothWeight);
 #endif
 
 	PositionBufferUAV[OutputOffset * 3 + 0] = SPos.x;
diff --git a/Engine/Shaders/Private/GpuSkinVertexFactory.ush b/Engine/Shaders/Private/GpuSkinVertexFactory.ush
--- a/Engine/Shaders/Private/GpuSkinVertexFactory.ush
+++ b/Engine/Shaders/Private/GpuSkinVertexFactory.ush
@@ -295,9 +295,9 @@ struct FVertexFactoryIntermediates
 
 #if GPUSKIN_APEX_CLOTH
 	// in world space (non translated)
-	float3 SimulatedPosition;
+	float4 SimulatedPosition;
 	// in world space (non translated)
-	float3 PreviousSimulatedPosition;
+	float4 PreviousSimulatedPosition;
 	
 	FClothVertex ClothVertexInfluences[NUM_INFLUENCES_PER_VERTEX];
 #endif
@@ -342,6 +342,11 @@ bool IsSimulatedVertex(FClothVertex Input)
 	return (Input.SimulIndices.w < FIXED_VERTEX_INDEX);
 }
 
+float GetClothSimulWeight(uint StaticAlpha)
+{
+	return 1.0f - (StaticAlpha / 65535.0f);
+}
+
 float3 GetClothSimulPosition(int Index, bool bPrevious)
 {
 #if FEATURE_LEVEL >= FEATURE_LEVEL_SM4
@@ -378,17 +383,18 @@ float3 GetClothSimulNormal(int Index, bool bPrevious)
 #endif //FEATURE_LEVEL >= FEATURE_LEVEL_SM4
 }
 
-float3 ClothingPosition(FClothVertex ClothVertexInfluences[NUM_INFLUENCES_PER_VERTEX], bool bPrevious)
+float4 ClothingPosition(FClothVertex ClothVertexInfluences[NUM_INFLUENCES_PER_VERTEX], bool bPrevious)
 {
 	int NumInfluences = 0;
 	float SumWeights = 0.0f;
+	float SimulWeight = 0.0f; 
 	float3 AveragedSimulatedPosition = float3(0, 0, 0);
 
     for (int i = 0; i < NUM_INFLUENCES_PER_VERTEX; ++i )
     {
         const FClothVertex Influence = ClothVertexInfluences[i];
 
-	    if( Influence.SimulIndices.w < 0xFFFF )
+	    if( IsSimulatedVertex(Influence) )
 	    { 
             ++NumInfluences;
 
@@ -400,6 +406,8 @@ float3 ClothingPosition(FClothVertex ClothVertexInfluences[NUM_INFLUENCES_PER_VE
 		    float3 NB = GetClothSimulNormal(Influence.SimulIndices.y, bPrevious);
 		    float3 NC = GetClothSimulNormal(Influence.SimulIndices.z, bPrevious);
 
+	    	SimulWeight += GetClothSimulWeight(Influence.SimulIndices.w);
+	    	
 			float Weight = 1.0f;
 			if ( NUM_INFLUENCES_PER_VERTEX > 1 )
 			{
@@ -436,7 +444,9 @@ float3 ClothingPosition(FClothVertex ClothVertexInfluences[NUM_INFLUENCES_PER_VE
 		AveragedSimulatedPosition = float3(0, 0, 0);
 	}
 
-	return AveragedSimulatedPosition;
+	SimulWeight /= NUM_INFLUENCES_PER_VERTEX;
+
+	return float4(AveragedSimulatedPosition, SimulWeight);
 }
 
 #endif //#if GPUSKIN_APEX_CLOTH
@@ -653,12 +663,13 @@ float3x3 SkinTangents( FVertexFactoryInput Input, FVertexFactoryIntermediates In
         float3 TangentPosition = float3(0,0,0);
         int NumInfluences = 0;
 	    float SumWeights = 0.0;
+		float SimulWeight = 0.0;
 
         for (int i = 0; i < NUM_INFLUENCES_PER_VERTEX; ++i )
         {
             const FClothVertex Influence = Intermediates.ClothVertexInfluences[i];
 
-	        if( Influence.SimulIndices.w < 0xFFFF )
+	        if( IsSimulatedVertex(Influence) )
 	        { 
                 ++NumInfluences;
 		        TempClothBlendWeight = ClothBlendWeight;
@@ -671,6 +682,8 @@ float3x3 SkinTangents( FVertexFactoryInput Input, FVertexFactoryIntermediates In
 		        float3 NB = GetClothSimulNormal(Influence.SimulIndices.y, bPrevious);
 		        float3 NC = GetClothSimulNormal(Influence.SimulIndices.z, bPrevious);
 
+	        	SimulWeight += GetClothSimulWeight(Influence.SimulIndices.w);
+	        	
 			    float Weight = 1.0f;
 			    if ( NUM_INFLUENCES_PER_VERTEX > 1 )
 			    {
@@ -701,8 +714,8 @@ float3x3 SkinTangents( FVertexFactoryInput Input, FVertexFactoryIntermediates In
             TangentPosition *= InvWeight;
             NormalPosition *= InvWeight;
 
-		    TangentToLocal[0] = normalize(TangentPosition - Intermediates.SimulatedPosition);
-		    TangentToLocal[2] = normalize(NormalPosition - Intermediates.SimulatedPosition);
+		    TangentToLocal[0] = normalize(TangentPosition - Intermediates.SimulatedPosition.xyz);
+		    TangentToLocal[2] = normalize(NormalPosition - Intermediates.SimulatedPosition.xyz);
 
 		    // cloth data are all in world space so need to change into local space
 		    TangentToLocal[0] = mul(TangentToLocal[0], (half3x3)Intermediates.WorldToLocal);
@@ -710,10 +723,11 @@ float3x3 SkinTangents( FVertexFactoryInput Input, FVertexFactoryIntermediates In
         }
 	    else
 	    {
-		    Intermediates.SimulatedPosition = float3(0, 0, 0);
 		    TempClothBlendWeight = 0.0f;
 	    }
 
+		TempClothBlendWeight *= SimulWeight;
+
 		// derive the new binormal by getting the cross product of the normal and tangent
 		// and flip vector based on sign of tangent basis determinant
 		TangentToLocal[1] = cross(TangentToLocal[2], TangentToLocal[0]) * LocalTangentZ.w;
@@ -768,7 +782,7 @@ FVertexFactoryIntermediates GetVertexFactoryIntermediates(FVertexFactoryInput In
 	}
 	else
 	{
-		Intermediates.PreviousSimulatedPosition = Intermediates.SimulatedPosition = Intermediates.UnpackedPosition;
+		Intermediates.PreviousSimulatedPosition = Intermediates.SimulatedPosition = float4(Intermediates.UnpackedPosition, 0.0f);
 	}
 	
 #endif
@@ -794,7 +808,7 @@ float4 CalcWorldPosition(FVertexFactoryInput Input, FVertexFactoryIntermediates
 		float4 TransformedSkinPos = mul(SkinnedPosition, Intermediates.LocalToWorld);
 
 		// Intermediates.SimulatedPosition is a clothing position in world coord
-		float3 BlendedPos = lerp(TransformedSkinPos.xyz, Intermediates.SimulatedPosition, ClothBlendWeight);
+		float3 BlendedPos = lerp(TransformedSkinPos.xyz, Intermediates.SimulatedPosition.xyz, ClothBlendWeight * Intermediates.SimulatedPosition.w);
 		return float4(BlendedPos + ResolvedView.PreViewTranslation, 1);
 	}
 	else
@@ -914,8 +928,8 @@ float4 VertexFactoryGetPreviousWorldPosition(FVertexFactoryInput Input, FVertexF
 #if GPUSKIN_APEX_CLOTH
 	if (IsSimulatedVertex(Intermediates.ClothVertexInfluences[0]))
 	{
-		float3 PrevSimulatedPos = Intermediates.PreviousSimulatedPosition + ResolvedView.PrevPreViewTranslation;
-		float3 BlendedPrevPos = lerp(PrevSkinPosInWorld.xyz, PrevSimulatedPos.xyz, ClothBlendWeight);
+		float3 PrevSimulatedPos = Intermediates.PreviousSimulatedPosition.xyz + ResolvedView.PrevPreViewTranslation;
+		float3 BlendedPrevPos = lerp(PrevSkinPosInWorld.xyz, PrevSimulatedPos.xyz, ClothBlendWeight * Intermediates.PreviousSimulatedPosition.w);
 		return float4(BlendedPrevPos, 1);
 	}
 #endif // GPUSKIN_APEX_CLOTH
diff --git a/Engine/Source/Editor/ClothPainter/Private/SClothAssetSelector.cpp b/Engine/Source/Editor/ClothPainter/Private/SClothAssetSelector.cpp
--- a/Engine/Source/Editor/ClothPainter/Private/SClothAssetSelector.cpp
+++ b/Engine/Source/Editor/ClothPainter/Private/SClothAssetSelector.cpp
@@ -961,6 +961,16 @@ void SClothAssetSelector::Construct(const FArguments& InArgs, USkeletalMesh* InM
 						.Font(FAppStyle::Get().GetFontStyle("PropertyWindow.NormalFont"))
 						.Text(LOCTEXT("CurrentRadius", "Kernel Radius"))
 					]
+					
+					+ SVerticalBox::Slot()
+					.Padding(0.0f, 2.0f, 0.0f, 2.0f)
+					.AutoHeight()
+					[
+						SNew(STextBlock)
+						.Font(IDetailLayoutBuilder::GetDetailFontBold())
+						.Text(LOCTEXT("SmoothTransition", "Smooth Transition From Skin to Cloth"))
+						.ShadowOffset(FVector2D(1, 1))
+					]
 				]
 
 				+ SSplitter::Slot()
@@ -974,7 +984,7 @@ void SClothAssetSelector::Construct(const FArguments& InArgs, USkeletalMesh* InM
 						SNew(SCheckBox)
 						.IsChecked(this, &SClothAssetSelector::GetCurrentUseMultipleInfluences)
 						.OnCheckStateChanged(this, &SClothAssetSelector::OnCurrentUseMultipleInfluencesChanged)
-						.IsEnabled(this, &SClothAssetSelector::CurrentUseMultipleInfluencesIsEnabled)
+						.IsEnabled(this, &SClothAssetSelector::IsValidClothLodSelected)
 					]
 
 					+ SVerticalBox::Slot()
@@ -993,6 +1003,16 @@ void SClothAssetSelector::Construct(const FArguments& InArgs, USkeletalMesh* InM
 						.OnValueChanged(this, &SClothAssetSelector::OnCurrentKernelRadiusChanged)
 						.LabelPadding(0)
 					]
+					
+					+ SVerticalBox::Slot()
+					.Padding(0.0f, 2.0f, 0.0f, 2.0f)
+					.AutoHeight()
+					[
+						SNew(SCheckBox)
+						.IsChecked(this, &SClothAssetSelector::GetCurrentSmoothTransition)
+						.OnCheckStateChanged(this, &SClothAssetSelector::OnCurrentSmoothTransitionChanged)
+						.IsEnabled(this, &SClothAssetSelector::IsValidClothLodSelected)
+					]
 				]
 
 			]
@@ -1082,7 +1102,43 @@ void SClothAssetSelector::OnCurrentUseMultipleInfluencesChanged(ECheckBoxState I
 	}
 }
 
-bool SClothAssetSelector::CurrentUseMultipleInfluencesIsEnabled() const
+
+ECheckBoxState SClothAssetSelector::GetCurrentSmoothTransition() const
+{
+	UClothingAssetCommon* Asset = SelectedAsset.Get();
+	if (Asset && Asset->IsValidLod(SelectedLod))
+	{
+		const FClothLODDataCommon& LodData = Asset->LodData[SelectedLod];
+		return LodData.bSmoothTransition ? ECheckBoxState::Checked : ECheckBoxState::Unchecked;
+	}
+
+	return ECheckBoxState::Undetermined;
+}
+
+
+void SClothAssetSelector::OnCurrentSmoothTransitionChanged(ECheckBoxState InValue)
+{
+	if (InValue == ECheckBoxState::Undetermined)
+	{
+		return;
+	}
+
+	UClothingAssetCommon* Asset = SelectedAsset.Get();
+	if (Asset && Asset->IsValidLod(SelectedLod))
+	{
+		FClothLODDataCommon& LodData = Asset->LodData[SelectedLod];
+		LodData.bSmoothTransition = (InValue == ECheckBoxState::Checked);
+
+		// Recompute weights
+		if (USkeletalMesh* SkeletalMesh = Cast<USkeletalMesh>(Asset->GetOuter()))
+		{
+			FScopedSkeletalMeshPostEditChange ScopedSkeletalMeshPostEditChange(SkeletalMesh);
+			SkeletalMesh->InvalidateDeriveDataCacheGUID();
+		}
+	}
+}
+
+bool SClothAssetSelector::IsValidClothLodSelected() const
 {
 	UClothingAssetCommon* Asset = SelectedAsset.Get();
 	return (Asset && Asset->IsValidLod(SelectedLod));
diff --git a/Engine/Source/Editor/ClothPainter/Private/SClothAssetSelector.h b/Engine/Source/Editor/ClothPainter/Private/SClothAssetSelector.h
--- a/Engine/Source/Editor/ClothPainter/Private/SClothAssetSelector.h
+++ b/Engine/Source/Editor/ClothPainter/Private/SClothAssetSelector.h
@@ -100,8 +100,12 @@ protected:
 
 	ECheckBoxState GetCurrentUseMultipleInfluences() const;
 	void OnCurrentUseMultipleInfluencesChanged(ECheckBoxState InValue);
-	bool CurrentUseMultipleInfluencesIsEnabled() const;
 
+	ECheckBoxState GetCurrentSmoothTransition() const;
+	void OnCurrentSmoothTransitionChanged(ECheckBoxState InValue);
+
+	bool IsValidClothLodSelected() const;
+	
 	void OnClothingLodSelected(int32 InNewLod);
 
 	// Setters for the list selections so we can handle list selections changing properly
diff --git a/Engine/Source/Editor/UnrealEd/Private/Animation/DebugSkelMeshComponent.cpp b/Engine/Source/Editor/UnrealEd/Private/Animation/DebugSkelMeshComponent.cpp
--- a/Engine/Source/Editor/UnrealEd/Private/Animation/DebugSkelMeshComponent.cpp
+++ b/Engine/Source/Editor/UnrealEd/Private/Animation/DebugSkelMeshComponent.cpp
@@ -753,23 +753,8 @@ void UDebugSkelMeshComponent::RebuildClothingSectionsFixedVerts(bool bInvalidate
 					const FClothLODDataCommon& LodData = ConcreteAsset->LodData[Section.ClothingData.AssetLodIndex];
 					const FPointWeightMap* const MaxDistances = LodData.PhysicalMeshData.FindWeightMap(EWeightMapTargetCommon::MaxDistance);
 
-					if (MaxDistances && MaxDistances->Num())
-					{
-						for (FMeshToMeshVertData& VertData : Section.ClothMappingData)
-						{
-							VertData.SourceMeshVertIndices[3] = MaxDistances->AreAllBelowThreshold(
-								VertData.SourceMeshVertIndices[0],
-								VertData.SourceMeshVertIndices[1],
-								VertData.SourceMeshVertIndices[2]) ? 0xFFFF : 0;
-						}
-					}
-					else
-					{
-						for (FMeshToMeshVertData& VertData : Section.ClothMappingData)
-						{
-							VertData.SourceMeshVertIndices[3] = 0;
-						}
-					}
+					ClothingMeshUtils::ComputeVertexContributions(Section.ClothMappingData, MaxDistances, LodData.bSmoothTransition);
+					
 					if (bInvalidateDerivedDataCache)
 					{
 						// We must always dirty the DDC key unless previewing
diff --git a/Engine/Source/Runtime/ClothingSystemRuntimeCommon/Private/ClothingAsset.cpp b/Engine/Source/Runtime/ClothingSystemRuntimeCommon/Private/ClothingAsset.cpp
--- a/Engine/Source/Runtime/ClothingSystemRuntimeCommon/Private/ClothingAsset.cpp
+++ b/Engine/Source/Runtime/ClothingSystemRuntimeCommon/Private/ClothingAsset.cpp
@@ -328,18 +328,10 @@ bool UClothingAssetCommon::BindToSkeletalMesh(
 		return false;
 	}
 
-	// Calculate fixed verts
+	// Calculate the vertex contribution alpha
 	const FPointWeightMap* const MaxDistances = ClothLodData.PhysicalMeshData.FindWeightMap(EWeightMapTargetCommon::MaxDistance);
-	if (MaxDistances && MaxDistances->Num())
-	{
-		for(FMeshToMeshVertData& VertData : MeshToMeshData)
-		{
-			if (MaxDistances->AreAllBelowThreshold(VertData.SourceMeshVertIndices[0], VertData.SourceMeshVertIndices[1], VertData.SourceMeshVertIndices[2]))
-			{
-				VertData.SourceMeshVertIndices[3] = 0xFFFF;
-			}
-		}
-	}
+
+	ClothingMeshUtils::ComputeVertexContributions(MeshToMeshData, MaxDistances, ClothLodData.bSmoothTransition);
 
 	// We have to copy the bone map to verify we don't exceed the maximum while adding the clothing bones
 	TArray<FBoneIndexType> TempBoneMap = OriginalSection.BoneMap;
@@ -561,23 +553,7 @@ void UClothingAssetCommon::ApplyParameterMasks(bool bUpdateFixedVertData)
 				const FClothLODDataCommon& LodDatum = LodData[Section.ClothingData.AssetLodIndex];
 				const FPointWeightMap* const MaxDistances = LodDatum.PhysicalMeshData.FindWeightMap(EWeightMapTargetCommon::MaxDistance);
 
-				if (MaxDistances && MaxDistances->Num())
-				{
-					for (FMeshToMeshVertData& VertData : Section.ClothMappingData)
-					{
-						VertData.SourceMeshVertIndices[3] = MaxDistances->AreAllBelowThreshold(
-							VertData.SourceMeshVertIndices[0],
-							VertData.SourceMeshVertIndices[1],
-							VertData.SourceMeshVertIndices[2]) ? 0xFFFF : 0;
-					}
-				}
-				else
-				{
-					for (FMeshToMeshVertData& VertData : Section.ClothMappingData)
-					{
-						VertData.SourceMeshVertIndices[3] = 0;
-					}
-				}
+				ClothingMeshUtils::ComputeVertexContributions(Section.ClothMappingData, MaxDistances, LodDatum.bSmoothTransition);
 			}
 		}
 		// We must always dirty the DDC key unless previewing
diff --git a/Engine/Source/Runtime/ClothingSystemRuntimeCommon/Private/Utils/ClothingMeshUtils.cpp b/Engine/Source/Runtime/ClothingSystemRuntimeCommon/Private/Utils/ClothingMeshUtils.cpp
--- a/Engine/Source/Runtime/ClothingSystemRuntimeCommon/Private/Utils/ClothingMeshUtils.cpp
+++ b/Engine/Source/Runtime/ClothingSystemRuntimeCommon/Private/Utils/ClothingMeshUtils.cpp
@@ -580,6 +580,60 @@ namespace ClothingMeshUtils
 		}
 	}
 
+
+	void ComputeVertexContributions(
+		TArray<FMeshToMeshVertData>& InOutSkinningData,
+		const FPointWeightMap* const InMaxDistances,
+		const bool bInSmoothTransition
+		)
+	{
+		if (InMaxDistances && InMaxDistances->Num())
+		{
+			for (FMeshToMeshVertData& VertData : InOutSkinningData)
+			{
+				const bool IsStatic0 = InMaxDistances->IsBelowThreshold(VertData.SourceMeshVertIndices[0]);
+				const bool IsStatic1 = InMaxDistances->IsBelowThreshold(VertData.SourceMeshVertIndices[1]);
+				const bool IsStatic2 = InMaxDistances->IsBelowThreshold(VertData.SourceMeshVertIndices[2]);
+
+				// None of the cloth vertices will move due to max distance constraints.
+				if (IsStatic0 && IsStatic1 && IsStatic2)
+				{
+					VertData.SourceMeshVertIndices[3] = 0xFFFF;
+				}
+				// If all of the vertices are dynamic _or_ if we disallow smooth transition,
+				// ensure there's no blending between cloth and skinned mesh and that the cloth
+				// mesh dominates.
+				else if ((!IsStatic0 && !IsStatic1 && !IsStatic2) || !bInSmoothTransition)
+				{
+					VertData.SourceMeshVertIndices[3] = 0;
+				}
+				else
+				{
+					// Compute how much the vertex actually contributes. A value of 0xFFFF
+					// means that it stays static relative to the skinned mesh, a value of 0x0000
+					// means that only the cloth simulation contributes. 
+					float StaticAlpha = 
+						IsStatic0 * VertData.PositionBaryCoordsAndDist.X +
+						IsStatic1 * VertData.PositionBaryCoordsAndDist.Y +
+						IsStatic2 * VertData.PositionBaryCoordsAndDist.Z;
+					StaticAlpha = FMath::Clamp(StaticAlpha, 0.0f, 1.0f);
+					
+					VertData.SourceMeshVertIndices[3] = static_cast<uint16>(StaticAlpha * 0xFFFF);
+				}	
+			}
+		}
+		else
+		{
+			// Can't determine contribution from the max distance map, so the entire mesh overrides.
+			for (FMeshToMeshVertData& VertData : InOutSkinningData)
+			{
+				VertData.SourceMeshVertIndices[3] = 0;
+			}
+		}
+		
+	}
+
+
 	void FVertexParameterMapper::Map(TArrayView<const float> Source, TArray<float>& Dest)
 	{
 		Map(Source, Dest, [](FVector3f Bary, float A, float B, float C)
diff --git a/Engine/Source/Runtime/ClothingSystemRuntimeCommon/Public/ClothLODData.h b/Engine/Source/Runtime/ClothingSystemRuntimeCommon/Public/ClothLODData.h
--- a/Engine/Source/Runtime/ClothingSystemRuntimeCommon/Public/ClothLODData.h
+++ b/Engine/Source/Runtime/ClothingSystemRuntimeCommon/Public/ClothLODData.h
@@ -27,6 +27,10 @@ struct CLOTHINGSYSTEMRUNTIMECOMMON_API FClothLODDataCommon
 	UPROPERTY()
 	float SkinningKernelRadius = 100.0f;
 
+	// Whether to enable smooth transition from skinned mesh to clothed mesh.
+	UPROPERTY()
+	bool bSmoothTransition = false;
+
 #if WITH_EDITORONLY_DATA
 	UPROPERTY()
 	TArray<FClothParameterMask_Legacy> ParameterMasks_DEPRECATED;
diff --git a/Engine/Source/Runtime/ClothingSystemRuntimeCommon/Public/Utils/ClothingMeshUtils.h b/Engine/Source/Runtime/ClothingSystemRuntimeCommon/Public/Utils/ClothingMeshUtils.h
--- a/Engine/Source/Runtime/ClothingSystemRuntimeCommon/Public/Utils/ClothingMeshUtils.h
+++ b/Engine/Source/Runtime/ClothingSystemRuntimeCommon/Public/Utils/ClothingMeshUtils.h
@@ -9,6 +9,7 @@
 DECLARE_LOG_CATEGORY_EXTERN(LogClothingMeshUtils, Log, All);
 
 struct FClothPhysicalMeshData;
+struct FPointWeightMap;
 
 namespace ClothingMeshUtils
 {
@@ -120,6 +121,16 @@ namespace ClothingMeshUtils
 		TArray<FVector4>& OutEmbeddedPositions, 
 		TArray<int32>& OutSourceIndices);
 
+	/**
+	 * Computes how much each vertex contributes to the final mesh. The final mesh is a blend
+	 * between the cloth and the skinned mesh.
+	 */
+	void CLOTHINGSYSTEMRUNTIMECOMMON_API ComputeVertexContributions(
+		TArray<FMeshToMeshVertData> &InOutSkinningData,
+		const FPointWeightMap* const InMaxDistances,
+		const bool bInSmoothTransition
+		);
+	
 	/**
 	* Given a triangle ABC with normals at each vertex NA, NB and NC, get a barycentric coordinate
 	* and corresponding distance from the triangle encoded in an FVector4 where the components are
diff --git a/Engine/Source/Runtime/Engine/Private/SkeletalRenderCPUSkin.cpp b/Engine/Source/Runtime/Engine/Private/SkeletalRenderCPUSkin.cpp
--- a/Engine/Source/Runtime/Engine/Private/SkeletalRenderCPUSkin.cpp
+++ b/Engine/Source/Runtime/Engine/Private/SkeletalRenderCPUSkin.cpp
@@ -995,8 +995,10 @@ static void SkinVertexSection(
 				// transform back to local space
 				FVector3f SimulatedPosition = WorldToLocal.TransformPosition(SimulatedPositionWorld);
 
+				const float VertexBlend = ClothBlendWeight * (1.0f - (ClothVertData->SourceMeshVertIndices[3] / 65535.0f));
+				
 				// Lerp between skinned and simulated position
-				DestVertex->Position = FMath::Lerp(DestVertex->Position, SimulatedPosition, ClothBlendWeight);
+				DestVertex->Position = FMath::Lerp(DestVertex->Position, SimulatedPosition, VertexBlend);
 
 				// recompute tangent & normal
 				FVector TangentX;
@@ -1006,8 +1008,8 @@ static void SkinVertexSection(
 				// Lerp between skinned and simulated tangents
 				FVector SkinnedTangentX = DestVertex->TangentX.ToFVector();
 				FVector4 SkinnedTangentZ = DestVertex->TangentZ.ToFVector4();
-				DestVertex->TangentX = (TangentX * ClothBlendWeight) + (SkinnedTangentX * (1.0f - ClothBlendWeight));
-				DestVertex->TangentZ = FVector4((TangentZ * ClothBlendWeight) + (SkinnedTangentZ * (1.0f - ClothBlendWeight)), SkinnedTangentZ.W);
+				DestVertex->TangentX = (TangentX * VertexBlend) + (SkinnedTangentX * (1.0f - VertexBlend));
+				DestVertex->TangentZ = FVector4((TangentZ * VertexBlend) + (SkinnedTangentZ * (1.0f - VertexBlend)), SkinnedTangentZ.W);
 			}
 
 			// Copy UVs.
