Fixed water to now really work with Pixel-Depth-Offset.
Thanks to [REDACTED_PERSON]!

diff --git a/Engine/Shaders/Private/MaterialTemplate.ush b/Engine/Shaders/Private/MaterialTemplate.ush
--- a/Engine/Shaders/Private/MaterialTemplate.ush
+++ b/Engine/Shaders/Private/MaterialTemplate.ush
@@ -3056,6 +3056,7 @@ void CalcMaterialParametersEx(
 	Parameters.SvPosition = SvPosition;
 	Parameters.ScreenPosition = ScreenPosition;
 
+	// CameraVector is a normalised vector representing the "from surface to camera" direction.
 	#if RAYHITGROUPSHADER
 		Parameters.CameraVector = -WorldRayDirection();
 	#else
@@ -3204,7 +3205,7 @@ float ApplyPixelDepthOffsetToMaterialParameters(inout FMaterialPixelParameters M
 	// Update positions used for shading
 	MaterialParameters.ScreenPosition.w += PixelDepthOffset;
 	MaterialParameters.SvPosition.w = MaterialParameters.ScreenPosition.w;
-	MaterialParameters.AbsoluteWorldPosition = LWCAdd(MaterialParameters.AbsoluteWorldPosition, MaterialParameters.CameraVector * PixelDepthOffset);
+	MaterialParameters.AbsoluteWorldPosition = LWCAdd(MaterialParameters.AbsoluteWorldPosition, -MaterialParameters.CameraVector * PixelDepthOffset);
 
 	OutDepth = INVARIANT(DeviceDepth);
 
diff --git a/Engine/Shaders/Private/SingleLayerWaterShading.ush b/Engine/Shaders/Private/SingleLayerWaterShading.ush
--- a/Engine/Shaders/Private/SingleLayerWaterShading.ush
+++ b/Engine/Shaders/Private/SingleLayerWaterShading.ush
@@ -43,12 +43,15 @@ WaterVolumeLightingOutput EvaluateWaterVolumeLighting(
 	Output.WaterToSceneTransmittance = 1.0f;
 	Output.WaterToSceneToLightTransmittance = 1.0f;
 
-	float2 WaterNDC = MaterialParameters.ScreenPosition.xy / MaterialParameters.ScreenPosition.w;
+	// Compute the ViewportUV while making sure the WaterNDC is unaffected by the Pixel-Depth-Offset
+	// which would have affected ScreenPosition.w, See ApplyPixelDepthOffsetToMaterialParameters.
+	const float PixelDepthOffset = GetMaterialPixelDepthOffset(PixelMaterialInputs);
+	float2 WaterNDC = MaterialParameters.ScreenPosition.xy / (MaterialParameters.ScreenPosition.w - PixelDepthOffset);
 	float2 ViewportUV = WaterNDC * ResolvedView.ScreenPositionScaleBias.xy + ResolvedView.ScreenPositionScaleBias.wz;
 
 	if (WaterVisibility > 0.0f)
 	{
-		float WaterDepth = ConvertFromDeviceZ(MaterialParameters.SvPosition.z);
+		float WaterDepth = MaterialParameters.ScreenPosition.w; // Scene depth with Pixel-Depth-Offset taken into account is in ScreenPosition.w already
 #if SIMPLE_SINGLE_LAYER_WATER
 
 #if SINGLE_LAYER_WATER_SIMPLE_FORWARD 
