Further improved water choerence with PDO.

diff --git a/Engine/Shaders/Private/MaterialTemplate.ush b/Engine/Shaders/Private/MaterialTemplate.ush
--- a/Engine/Shaders/Private/MaterialTemplate.ush
+++ b/Engine/Shaders/Private/MaterialTemplate.ush
@@ -280,6 +280,12 @@ struct FMaterialPixelParameters
 	/** Post projection position reconstructed from SvPosition, before the divide by W. left..top -1..1, bottom..top -1..1  within the viewport, W is the SceneDepth */
 	float4 ScreenPosition;
 
+	/**
+	 * The pixel UV for a view, considering the buffer resolution containing it.
+	 * This is precomputed because when PixelDepthOffset is used, it affects ScreenPosition which would then prevent us from recovering the correct UV later. 
+	 * See ApplyPixelDepthOffsetToMaterialParameters for the details */
+	float2 ViewBufferUV;
+
 #if IS_NANITE_FACTORY
 	float4 PrevScreenPosition;
 #endif
@@ -3055,6 +3061,7 @@ void CalcMaterialParametersEx(
 
 	Parameters.SvPosition = SvPosition;
 	Parameters.ScreenPosition = ScreenPosition;
+	Parameters.ViewBufferUV = ScreenPositionToBufferUV(ScreenPosition);
 
 	// CameraVector is a normalised vector representing the "from surface to camera" direction.
 	#if RAYHITGROUPSHADER
diff --git a/Engine/Shaders/Private/SingleLayerWaterShading.ush b/Engine/Shaders/Private/SingleLayerWaterShading.ush
--- a/Engine/Shaders/Private/SingleLayerWaterShading.ush
+++ b/Engine/Shaders/Private/SingleLayerWaterShading.ush
@@ -43,11 +43,7 @@ WaterVolumeLightingOutput EvaluateWaterVolumeLighting(
 	Output.WaterToSceneTransmittance = 1.0f;
 	Output.WaterToSceneToLightTransmittance = 1.0f;
 
-	// Compute the ViewportUV while making sure the WaterNDC is unaffected by the Pixel-Depth-Offset
-	// which would have affected ScreenPosition.w, See ApplyPixelDepthOffsetToMaterialParameters.
-	const float PixelDepthOffset = GetMaterialPixelDepthOffset(PixelMaterialInputs);
-	float2 WaterNDC = MaterialParameters.ScreenPosition.xy / (MaterialParameters.ScreenPosition.w - PixelDepthOffset);
-	float2 ViewportUV = WaterNDC * ResolvedView.ScreenPositionScaleBias.xy + ResolvedView.ScreenPositionScaleBias.wz;
+	float2 ViewportUV = MaterialParameters.ViewBufferUV;
 
 	if (WaterVisibility > 0.0f)
 	{
