Fixed wrong refraction when SingleLayerWater tris are not horizontal in the world overall.

diff --git a/Engine/Shaders/Private/SingleLayerWaterShading.ush b/Engine/Shaders/Private/SingleLayerWaterShading.ush
--- a/Engine/Shaders/Private/SingleLayerWaterShading.ush
+++ b/Engine/Shaders/Private/SingleLayerWaterShading.ush
@@ -73,7 +73,7 @@ WaterVolumeLightingOutput EvaluateWaterVolumeLighting(
 		float2 Refraction = GetMaterialRefraction(PixelMaterialInputs); // X = IOR, Y = Refraction Bias
 		float2 BufferUVDistortion = ComputeBufferUVDistortion(
 			MaterialParameters, PixelMaterialInputs, ResolvedView,
-			MaterialParameters.WorldNormal, DistortionParams, ViewportUV, Refraction, false, EyeIndex);
+			GetMaterialNormal(MaterialParameters, PixelMaterialInputs), DistortionParams, ViewportUV, Refraction, false, EyeIndex);
 		PostProcessUVDistortion(MaterialParameters, PixelMaterialInputs, PixelSceneDepth, BufferUVDistortion, Refraction);
 		// Also reduce refraction according to water thickness
 		BufferUVDistortion *= saturate((PixelSceneDepth - WaterDepth) * 1.0 / 30.0f);// *(1 - saturate(WaterDepth * 1.0 / 2000.0f));
