Fix sky light contribution in deferrred ray traced reflection shader.

The ray payload needs to be re-initialized after every trace, otherwise the payload flags may be lost. If bSkyLightAffectReflection is enabled and a ray passes through a transparent part of a masked material, the background would not receive correct sky lighting.
This change initializes sky light contribution, pixel coord and mip bias fields of the ray payload before TraceMaterialRayPacked().

UDN case 00362581

diff --git a/Engine/Shaders/Private/RayTracing/RayTracingDeferredReflections.usf b/Engine/Shaders/Private/RayTracing/RayTracingDeferredReflections.usf
--- a/Engine/Shaders/Private/RayTracing/RayTracingDeferredReflections.usf
+++ b/Engine/Shaders/Private/RayTracing/RayTracingDeferredReflections.usf
@@ -69,6 +69,18 @@ RaytracingAccelerationStructure                 TLAS;
 // used during to reduce ray intersection overhead during Shading.
 RWStructuredBuffer<FRayIntersectionBookmark>    BookmarkBuffer;
 
+FPackedMaterialClosestHitPayload InitializePackedPayload(uint2 PixelPos, bool bSkyLightAffectReflection)
+{
+	FPackedMaterialClosestHitPayload PackedPayload = (FPackedMaterialClosestHitPayload)0;
+	if (bSkyLightAffectReflection)
+	{
+		PackedPayload.SetEnableSkyLightContribution();
+	}
+	PackedPayload.SetPixelCoord(PixelPos);
+	PackedPayload.SetMipBias(TextureMipBias);
+	return PackedPayload;
+}
+
 RAY_TRACING_ENTRY_RAYGEN(RayTracingDeferredReflectionsRGS)
 {
 	const uint DispatchThreadId = DispatchRaysIndex().x; // This shader is dispatched in 1D configuration
@@ -154,7 +166,7 @@ RAY_TRACING_ENTRY_RAYGEN(RayTracingDeferredReflectionsRGS)
 	//float ResultDistance = DENOISER_INVALID_HIT_DISTANCE;
 	float ResultDistance = 1.0e20;
 
-	FPackedMaterialClosestHitPayload PackedPayload = (FPackedMaterialClosestHitPayload)0;
+	FPackedMaterialClosestHitPayload PackedPayload = InitializePackedPayload(PixelPos, bSkyLightAffectReflection);
 	float3 RayHitWorldPos = ReflectionRay.Origin + ReflectionRay.GetDirection() * DENOISER_MISS_HIT_DISTANCE;
 
 	if (DeferredMaterialPayload.SortKey < RAY_TRACING_DEFERRED_MATERIAL_KEY_RAY_MISS)
@@ -167,13 +179,6 @@ RAY_TRACING_ENTRY_RAYGEN(RayTracingDeferredReflectionsRGS)
 
 		FRayIntersectionBookmark Bookmark = BookmarkBuffer[DeferredMaterialPayload.PixelCoordinates];
 
-		if (bSkyLightAffectReflection)
-		{
-			PackedPayload.SetEnableSkyLightContribution();
-		}
-		PackedPayload.SetPixelCoord(PixelPos);
-		PackedPayload.SetMipBias(TextureMipBias);
-
 		TraceDeferredMaterialShadingRay(
 			TLAS,
 			RAY_FLAG_FORCE_OPAQUE,
@@ -191,6 +196,9 @@ RAY_TRACING_ENTRY_RAYGEN(RayTracingDeferredReflectionsRGS)
 		if ((!PackedPayload.IsHit() || PackedPayload.GetOpacity() == 0) 
 			&& ReflectionRay.Roughness <= AnyHitMaxRoughness)
 		{
+			// Reset the payload to a known initial state, as previous Trace() call may have overwritten some of the fields.
+			PackedPayload = InitializePackedPayload(PixelPos, bSkyLightAffectReflection);
+
 			FRayCone RayCone = (FRayCone)0;
 			TraceMaterialRayPacked(
 				PackedPayload,
