Fix shading models bugs for mobile.

#p4v-preflight-copy 17869699

diff --git a/Engine/Shaders/Private/MobileShadingModels.ush b/Engine/Shaders/Private/MobileShadingModels.ush
--- a/Engine/Shaders/Private/MobileShadingModels.ush
+++ b/Engine/Shaders/Private/MobileShadingModels.ush
@@ -58,7 +58,7 @@
 // Enable shading model support for deferred shading once the gbuffer has capacity.
 #define MOBILE_SHADINGMODEL_SUPPORT_DEFERRED 0
 
-#define MOBILE_SHADINGMODEL_SUPPORT (!DEFERRED_SHADING_PATH || SHADINGMODEL_SUPPORT_MOBILE_DEFERRED)
+#define MOBILE_SHADINGMODEL_SUPPORT (!MOBILE_DEFERRED_SHADING || MOBILE_SHADINGMODEL_SUPPORT_DEFERRED)
 
 /*------------------------------------------------------------------------------
 	Mobile Shading Models
@@ -198,6 +198,10 @@ void InitShadingModelContext(inout FMobileShadingModelContext ShadingModelContex
 		half DielectricSpecular = 0.08 * Specular;
 		ShadingModelContext.SpecularColor = (DielectricSpecular - DielectricSpecular * GBuffer.Metallic) + GBuffer.BaseColor * GBuffer.Metallic;	// 2 mad
 	}
+	else if (GBuffer.ShadingModelID == SHADINGMODELID_SUBSURFACE || GBuffer.ShadingModelID == SHADINGMODELID_PREINTEGRATED_SKIN)
+	{
+		ShadingModelContext.SubsurfaceColor = ExtractSubsurfaceColor(GBuffer);
+	}
 	else if (GBuffer.ShadingModelID == SHADINGMODELID_SUBSURFACE_PROFILE)
 	{
 		ShadingModelContext.SubsurfaceProfileInt = ExtractSubsurfaceProfileInt(GBuffer);
