Change hair strands vertex factory to use uniform buffer for binding parameters, instead of binding loose parameters.

Work originally done by Chris.

diff --git a/Engine/Plugins/Runtime/HairStrands/Source/HairStrandsCore/Private/HairStrandsRendering.cpp b/Engine/Plugins/Runtime/HairStrands/Source/HairStrandsCore/Private/HairStrandsRendering.cpp
--- a/Engine/Plugins/Runtime/HairStrands/Source/HairStrandsCore/Private/HairStrandsRendering.cpp
+++ b/Engine/Plugins/Runtime/HairStrands/Source/HairStrandsCore/Private/HairStrandsRendering.cpp
@@ -1963,6 +1963,7 @@ void ComputeHairStrandsInterpolation(
 
 		// Sanity check
 		check(Instance->HairGroupPublicData->VFInput.Strands.PositionBuffer.Buffer);
+		Instance->Strands.UniformBuffer.UpdateUniformBufferImmediate(Instance->GetHairStandsUniformShaderParameters());
 	}
 	else if (InstanceGeometryType == EHairGeometryType::Cards)
 	{	
diff --git a/Engine/Plugins/Runtime/HairStrands/Source/HairStrandsCore/Private/HairStrandsVertexFactory.cpp b/Engine/Plugins/Runtime/HairStrands/Source/HairStrandsCore/Private/HairStrandsVertexFactory.cpp
--- a/Engine/Plugins/Runtime/HairStrands/Source/HairStrandsCore/Private/HairStrandsVertexFactory.cpp
+++ b/Engine/Plugins/Runtime/HairStrands/Source/HairStrandsCore/Private/HairStrandsVertexFactory.cpp
@@ -54,87 +54,66 @@ public:
 };
 TGlobalResource<FDummyCulledDispatchVertexIdsBuffer> GDummyCulledDispatchVertexIdsBuffer;
 
-/////////////////////////////////////////////////////////////////////////////////////////
-inline FRHIShaderResourceView* GetPositionSRV(const FHairGroupPublicData::FVertexFactoryInput& VFInput)			{ return VFInput.Strands.PositionBufferRHISRV; }; 
-inline FRHIShaderResourceView* GetPreviousPositionSRV(const FHairGroupPublicData::FVertexFactoryInput& VFInput)	{ return VFInput.Strands.PrevPositionBufferRHISRV; }
-inline FRHIShaderResourceView* GetAttribute0SRV(const FHairGroupPublicData::FVertexFactoryInput& VFInput)		{ return VFInput.Strands.Attribute0BufferRHISRV; }
-inline FRHIShaderResourceView* GetAttribute1SRV(const FHairGroupPublicData::FVertexFactoryInput& VFInput)		{ return VFInput.Strands.Attribute1BufferRHISRV; }
-inline FRHIShaderResourceView* GetMaterialSRV(const FHairGroupPublicData::FVertexFactoryInput& VFInput)			{ return VFInput.Strands.MaterialBufferRHISRV; }
-inline FRHIShaderResourceView* GetTangentSRV(const FHairGroupPublicData::FVertexFactoryInput& VFInput)			{ return VFInput.Strands.TangentBufferRHISRV; }
-
-inline FRHIShaderResourceView* GetPositionOffsetSRV(const FHairGroupPublicData::FVertexFactoryInput& VFInput)	{ return VFInput.Strands.PositionOffsetBufferRHISRV; };
-inline FRHIShaderResourceView* GetPreviousPositionOffsetSRV(const FHairGroupPublicData::FVertexFactoryInput& VFInput) { return VFInput.Strands.PrevPositionOffsetBufferRHISRV; }
-
-inline bool  UseStableRasterization(const FHairGroupPublicData::FVertexFactoryInput& VFInput)					{ return VFInput.Strands.bUseStableRasterization; };
-inline bool  UseScatterSceneLighting(const FHairGroupPublicData::FVertexFactoryInput& VFInput)					{ return VFInput.Strands.bScatterSceneLighting; };
-inline float GetMaxStrandRadius(const FHairGroupPublicData::FVertexFactoryInput& VFInput)						{ return VFInput.Strands.HairRadius; };
-inline float GetStrandRootScale(const FHairGroupPublicData::FVertexFactoryInput& VFInput)						{ return VFInput.Strands.HairRootScale; };
-inline float GetStrandTipScale(const FHairGroupPublicData::FVertexFactoryInput& VFInput)						{ return VFInput.Strands.HairTipScale; };
-inline float GetStrandLengthScale(const FHairGroupPublicData::FVertexFactoryInput& VFInput)						{ return VFInput.Strands.HairLengthScale; };
-inline float GetMaxStrandLength(const FHairGroupPublicData::FVertexFactoryInput& VFInput)						{ return VFInput.Strands.HairLength; };
-inline float GetHairDensity(const FHairGroupPublicData::FVertexFactoryInput& VFInput)							{ return VFInput.Strands.HairDensity; };
-//inline float GetHasAttribute1(const FHairGroupPublicData::FVertexFactoryInput& VFInput)							{ return VFInput.Strands.HasAttribute1; };
-//inline float GetHasMaterial(const FHairGroupPublicData::FVertexFactoryInput& VFInput)							{ return VFInput.Strands.HasMaterial; };
-
 /////////////////////////////////////////////////////////////////////////////////////////
 
 FHairGroupPublicData::FVertexFactoryInput ComputeHairStrandsVertexInputData(const FHairGroupInstance* Instance);
+FHairStrandsVertexFactoryUniformShaderParameters FHairGroupInstance::GetHairStandsUniformShaderParameters() const
+{
+	const FHairGroupPublicData::FVertexFactoryInput VFInput = ComputeHairStrandsVertexInputData(this);
+
+	FHairStrandsVertexFactoryUniformShaderParameters Out = {};
+	Out.Radius 						= VFInput.Strands.HairRadius;
+	Out.RootScale 					= VFInput.Strands.HairRootScale;
+	Out.TipScale 					= VFInput.Strands.HairTipScale;
+	Out.Length 						= VFInput.Strands.HairLength;
+	Out.Density 					= VFInput.Strands.HairDensity;
+	Out.StableRasterization			= VFInput.Strands.bUseStableRasterization;
+	Out.ScatterSceneLighing			= VFInput.Strands.bScatterSceneLighting;
+	Out.PositionBuffer				= VFInput.Strands.PositionBufferRHISRV;
+	Out.PreviousPositionBuffer		= VFInput.Strands.PrevPositionBufferRHISRV;
+	Out.Attribute0Buffer			= VFInput.Strands.Attribute0BufferRHISRV;
+	Out.Attribute1Buffer			= VFInput.Strands.Attribute1BufferRHISRV;
+	Out.MaterialBuffer 				= VFInput.Strands.MaterialBufferRHISRV;
+	Out.HasAttribute1 				= Out.Attribute1Buffer != nullptr;
+	Out.HasMaterial 				= Out.MaterialBuffer != nullptr;
+	Out.TangentBuffer 				= VFInput.Strands.TangentBufferRHISRV;
+	Out.PositionOffsetBuffer 		= VFInput.Strands.PositionOffsetBufferRHISRV;
+	Out.PreviousPositionOffsetBuffer= VFInput.Strands.PrevPositionOffsetBufferRHISRV;
+
+	// swap in some default data for those buffers that are not valid yet
+	if (!Out.PositionBuffer) 				{ Out.PositionBuffer = GDummyCulledDispatchVertexIdsBuffer.SRVFloat; }
+	if (!Out.PreviousPositionBuffer) 		{ Out.PreviousPositionBuffer = GDummyCulledDispatchVertexIdsBuffer.SRVFloat; }
+	if (!Out.Attribute0Buffer) 				{ Out.Attribute0Buffer = GDummyCulledDispatchVertexIdsBuffer.SRVFloat; }
+	if (!Out.Attribute1Buffer) 				{ Out.Attribute1Buffer = GDummyCulledDispatchVertexIdsBuffer.SRVFloat; }
+	if (!Out.MaterialBuffer) 				{ Out.MaterialBuffer = GDummyCulledDispatchVertexIdsBuffer.SRVRGBA; }
+	if (!Out.TangentBuffer) 				{ Out.TangentBuffer = GDummyCulledDispatchVertexIdsBuffer.SRVFloat; }
+	if (!Out.PositionOffsetBuffer) 			{ Out.PositionOffsetBuffer = GDummyCulledDispatchVertexIdsBuffer.SRVFloat; }
+	if (!Out.PreviousPositionOffsetBuffer) 	{ Out.PreviousPositionOffsetBuffer = GDummyCulledDispatchVertexIdsBuffer.SRVFloat; }
+
+	Out.CullingEnable = HairGroupPublicData->GetCullingResultAvailable();
+	if (Out.CullingEnable)
+	{
+		Out.CulledVertexIdsBuffer = HairGroupPublicData->GetCulledVertexIdBuffer().SRV;
+		Out.CulledVertexRadiusScaleBuffer = HairGroupPublicData->GetCulledVertexRadiusScaleBuffer().SRV;
+	}
+	else
+	{
+		Out.CulledVertexIdsBuffer = GDummyCulledDispatchVertexIdsBuffer.SRVUint;
+		Out.CulledVertexRadiusScaleBuffer = GDummyCulledDispatchVertexIdsBuffer.SRVFloat;
+	}
+	return Out;
+}
+
 
-class FHairStrandsVertexFactoryShaderParameters : public FVertexFactoryShaderParameters
+IMPLEMENT_UNIFORM_BUFFER_STRUCT(FHairStrandsVertexFactoryUniformShaderParameters, "HairStrandsVF");
+
+class HAIRSTRANDSCORE_API FHairStrandsVertexFactoryShaderParameters : public FVertexFactoryShaderParameters
 {
 	DECLARE_TYPE_LAYOUT(FHairStrandsVertexFactoryShaderParameters, NonVirtual);
 public:
 
-	LAYOUT_FIELD(FShaderParameter, Radius);
-	LAYOUT_FIELD(FShaderParameter, RootScale);
-	LAYOUT_FIELD(FShaderParameter, TipScale);
-	LAYOUT_FIELD(FShaderParameter, Length);
-	LAYOUT_FIELD(FShaderParameter, RadiusAtDepth1_Primary);	// unused
-	LAYOUT_FIELD(FShaderParameter, RadiusAtDepth1_Velocity);// unused
-	LAYOUT_FIELD(FShaderParameter, Density);
-	LAYOUT_FIELD(FShaderParameter, Culling);
-	LAYOUT_FIELD(FShaderParameter, HasAttribute1);
-	LAYOUT_FIELD(FShaderParameter, HasMaterial);
-	LAYOUT_FIELD(FShaderParameter, StableRasterization);
-	LAYOUT_FIELD(FShaderParameter, ScatterSceneLighing);
-
-	LAYOUT_FIELD(FShaderResourceParameter, PositionOffsetBuffer);
-	LAYOUT_FIELD(FShaderResourceParameter, PreviousPositionOffsetBuffer);
-	LAYOUT_FIELD(FShaderResourceParameter, PositionBuffer);
-	LAYOUT_FIELD(FShaderResourceParameter, PreviousPositionBuffer);
-	LAYOUT_FIELD(FShaderResourceParameter, Attribute0Buffer);
-	LAYOUT_FIELD(FShaderResourceParameter, Attribute1Buffer);
-	LAYOUT_FIELD(FShaderResourceParameter, MaterialBuffer);
-	LAYOUT_FIELD(FShaderResourceParameter, TangentBuffer);
-
-	LAYOUT_FIELD(FShaderResourceParameter, CulledVertexIdsBuffer);
-	LAYOUT_FIELD(FShaderResourceParameter, CulledVertexRadiusScaleBuffer);
-
 	void Bind(const FShaderParameterMap& ParameterMap)
 	{
-		Radius.Bind(ParameterMap, TEXT("HairStrandsVF_Radius"));
-		RootScale.Bind(ParameterMap, TEXT("HairStrandsVF_RootScale"));
-		TipScale.Bind(ParameterMap, TEXT("HairStrandsVF_TipScale"));
-		Length.Bind(ParameterMap, TEXT("HairStrandsVF_Length"));
-		Density.Bind(ParameterMap, TEXT("HairStrandsVF_Density"));
-		Culling.Bind(ParameterMap, TEXT("HairStrandsVF_CullingEnable"));
-		HasAttribute1.Bind(ParameterMap, TEXT("HairStrandsVF_HasAttribute1"));
-		HasMaterial.Bind(ParameterMap, TEXT("HairStrandsVF_HasMaterial"));
-		StableRasterization.Bind(ParameterMap, TEXT("HairStrandsVF_bUseStableRasterization"));
-		ScatterSceneLighing.Bind(ParameterMap, TEXT("HairStrandsVF_bScatterSceneLighing"));
-
-		PositionOffsetBuffer.Bind(ParameterMap, TEXT("HairStrandsVF_PositionOffsetBuffer"));
-		PreviousPositionOffsetBuffer.Bind(ParameterMap, TEXT("HairStrandsVF_PreviousPositionOffsetBuffer"));
-
-		PositionBuffer.Bind(ParameterMap, TEXT("HairStrandsVF_PositionBuffer"));
-		PreviousPositionBuffer.Bind(ParameterMap, TEXT("HairStrandsVF_PreviousPositionBuffer"));
-		Attribute0Buffer.Bind(ParameterMap, TEXT("HairStrandsVF_Attribute0Buffer"));
-		Attribute1Buffer.Bind(ParameterMap, TEXT("HairStrandsVF_Attribute1Buffer"));
-		MaterialBuffer.Bind(ParameterMap, TEXT("HairStrandsVF_MaterialBuffer"));
-		TangentBuffer.Bind(ParameterMap, TEXT("HairStrandsVF_TangentBuffer"));
-
-		CulledVertexIdsBuffer.Bind(ParameterMap, TEXT("CulledVertexIdsBuffer"));
-		CulledVertexRadiusScaleBuffer.Bind(ParameterMap, TEXT("CulledVertexRadiusScaleBuffer"));
 	}
 
 	void GetElementShaderBindings(
@@ -150,54 +129,7 @@ public:
 	) const
 	{
 		const FHairStrandsVertexFactory* VF = static_cast<const FHairStrandsVertexFactory*>(VertexFactory);
-		const FHairGroupPublicData::FVertexFactoryInput VFInput = ComputeHairStrandsVertexInputData(VF->Data.Instance);
-
-		FShaderResourceViewRHIRef Attribute1SRV = GetAttribute1SRV(VFInput);
-		const bool bHasAttribute1 = Attribute1SRV != nullptr;
-		if (!bHasAttribute1)
-		{
-			Attribute1SRV = GDummyCulledDispatchVertexIdsBuffer.SRVRGBA_Uint;
-		}
-
-		FShaderResourceViewRHIRef MaterialSRV = GetMaterialSRV(VFInput);
-		const bool bHasMaterial = MaterialSRV != nullptr;
-		if (!bHasMaterial)
-		{
-			MaterialSRV = GDummyCulledDispatchVertexIdsBuffer.SRVRGBA;
-		}
-
-		VFS_BindParam(ShaderBindings, PositionBuffer, GetPositionSRV(VFInput));
-		VFS_BindParam(ShaderBindings, PreviousPositionBuffer, GetPreviousPositionSRV(VFInput));
-		VFS_BindParam(ShaderBindings, Attribute0Buffer, GetAttribute0SRV(VFInput));
-		VFS_BindParam(ShaderBindings, Attribute1Buffer, Attribute1SRV.GetReference());
-		VFS_BindParam(ShaderBindings, MaterialBuffer, MaterialSRV.GetReference());
-		VFS_BindParam(ShaderBindings, TangentBuffer, GetTangentSRV(VFInput));
-		VFS_BindParam(ShaderBindings, Radius, GetMaxStrandRadius(VFInput));
-		VFS_BindParam(ShaderBindings, RootScale, GetStrandRootScale(VFInput));
-		VFS_BindParam(ShaderBindings, TipScale, GetStrandTipScale(VFInput));
-		VFS_BindParam(ShaderBindings, Length, GetMaxStrandLength(VFInput));
-		VFS_BindParam(ShaderBindings, PositionOffsetBuffer, GetPositionOffsetSRV(VFInput));
-		VFS_BindParam(ShaderBindings, PreviousPositionOffsetBuffer, GetPreviousPositionOffsetSRV(VFInput));
-		VFS_BindParam(ShaderBindings, Density, GetHairDensity(VFInput));
-		VFS_BindParam(ShaderBindings, StableRasterization, UseStableRasterization(VFInput) ? 1u : 0u);
-		VFS_BindParam(ShaderBindings, ScatterSceneLighing, UseScatterSceneLighting(VFInput) ? 1u : 0u);
-		VFS_BindParam(ShaderBindings, HasAttribute1, bHasAttribute1 ? 1u : 0u);
-		VFS_BindParam(ShaderBindings, HasMaterial, bHasMaterial ? 1u : 0u);
-		
-		FShaderResourceViewRHIRef CulledDispatchVertexIdsSRV = GDummyCulledDispatchVertexIdsBuffer.SRVUint;
-		FShaderResourceViewRHIRef CulledCompactedRadiusScaleBufferSRV = GDummyCulledDispatchVertexIdsBuffer.SRVFloat;
-
-		const FHairGroupPublicData* GroupPublicData = VF->Data.Instance->HairGroupPublicData;
-		const bool bCulling = GroupPublicData->GetCullingResultAvailable();
-		if (bCulling)
-		{
-			CulledDispatchVertexIdsSRV = GroupPublicData->GetCulledVertexIdBuffer().SRV;
-			CulledCompactedRadiusScaleBufferSRV = GroupPublicData->GetCulledVertexRadiusScaleBuffer().SRV;
-		}
-
-		VFS_BindParam(ShaderBindings, Culling, bCulling ? 1 : 0);
-		ShaderBindings.Add(CulledVertexIdsBuffer, CulledDispatchVertexIdsSRV);
-		ShaderBindings.Add(CulledVertexRadiusScaleBuffer, CulledCompactedRadiusScaleBufferSRV);
+		ShaderBindings.Add(Shader->GetUniformBufferParameter<FHairStrandsVertexFactoryUniformShaderParameters>(), VF->Data.Instance->Strands.UniformBuffer);
 	}
 };
 
@@ -216,7 +148,7 @@ void FHairStrandsVertexFactory::ModifyCompilationEnvironment(const FVertexFactor
 	bool bUseGPUSceneAndPrimitiveIdStream = false;
 #if VF_STRANDS_SUPPORT_GPU_SCENE
 	bUseGPUSceneAndPrimitiveIdStream = 
-		&& Parameters.VertexFactoryType->SupportsPrimitiveIdStream() 
+		Parameters.VertexFactoryType->SupportsPrimitiveIdStream() 
 		&& UseGPUScene(Parameters.Platform, GetMaxSupportedFeatureLevel(Parameters.Platform))
 		// TODO: support GPUScene on mobile
 		&& !IsMobilePlatform(Parameters.Platform);
@@ -288,6 +220,10 @@ void FHairStrandsVertexFactory::InitResources()
 	}
 	InitDeclaration(Elements);
 	check(IsValidRef(GetDeclaration()));
+
+	// create the buffer
+	FHairStrandsVertexFactoryUniformShaderParameters Parameters = Data.Instance->GetHairStandsUniformShaderParameters();
+	Data.Instance->Strands.UniformBuffer = FHairStrandsUniformBuffer::CreateUniformBufferImmediate(Parameters, UniformBuffer_MultiFrame);
 }
 
 IMPLEMENT_VERTEX_FACTORY_PARAMETER_TYPE(FHairStrandsVertexFactory, SF_Vertex,		FHairStrandsVertexFactoryShaderParameters);
@@ -307,8 +243,8 @@ void FHairStrandsVertexFactory::ReleaseRHI()
 	FVertexFactory::ReleaseRHI();
 }
 
-IMPLEMENT_VERTEX_FACTORY_TYPE(FHairStrandsVertexFactory,"/Engine/Private/HairStrands/HairStrandsVertexFactory.ush",
-	  EVertexFactoryFlags::UsedWithMaterials
+IMPLEMENT_VERTEX_FACTORY_TYPE(FHairStrandsVertexFactory, "/Engine/Private/HairStrands/HairStrandsVertexFactory.ush",
+	EVertexFactoryFlags::UsedWithMaterials
 	| EVertexFactoryFlags::SupportsDynamicLighting
 	| EVertexFactoryFlags::SupportsPrecisePrevWorldPos
 	| EVertexFactoryFlags::SupportsCachingMeshDrawCommands
diff --git a/Engine/Plugins/Runtime/HairStrands/Source/HairStrandsCore/Public/GroomInstance.h b/Engine/Plugins/Runtime/HairStrands/Source/HairStrandsCore/Public/GroomInstance.h
--- a/Engine/Plugins/Runtime/HairStrands/Source/HairStrandsCore/Public/GroomInstance.h
+++ b/Engine/Plugins/Runtime/HairStrands/Source/HairStrandsCore/Public/GroomInstance.h
@@ -40,6 +40,36 @@ END_GLOBAL_SHADER_PARAMETER_STRUCT()
 
 typedef TUniformBufferRef<FHairCardsVertexFactoryUniformShaderParameters> FHairCardsUniformBuffer;
 
+
+BEGIN_GLOBAL_SHADER_PARAMETER_STRUCT(FHairStrandsVertexFactoryUniformShaderParameters, HAIRSTRANDSCORE_API)
+	SHADER_PARAMETER(float, Radius)
+	SHADER_PARAMETER(float, RootScale)
+	SHADER_PARAMETER(float, TipScale)
+	SHADER_PARAMETER(float, Length)
+	SHADER_PARAMETER(float, Density)
+	SHADER_PARAMETER(uint32, CullingEnable)
+	SHADER_PARAMETER(uint32, HasAttribute1)
+	SHADER_PARAMETER(uint32, HasMaterial)
+	SHADER_PARAMETER(uint32, StableRasterization)
+	SHADER_PARAMETER(uint32, ScatterSceneLighing)
+
+	SHADER_PARAMETER_SRV(Buffer<float4>, PositionOffsetBuffer)
+	SHADER_PARAMETER_SRV(Buffer<float4>, PreviousPositionOffsetBuffer)
+
+	SHADER_PARAMETER_SRV(Buffer<uint4>, PositionBuffer)
+	SHADER_PARAMETER_SRV(Buffer<uint4>, PreviousPositionBuffer)
+
+	SHADER_PARAMETER_SRV(Buffer<float4>, Attribute0Buffer)
+	SHADER_PARAMETER_SRV(Buffer<float4>, Attribute1Buffer)
+	SHADER_PARAMETER_SRV(Buffer<float4>, MaterialBuffer)
+	SHADER_PARAMETER_SRV(Buffer<float4>, TangentBuffer)
+
+	SHADER_PARAMETER_SRV(Buffer<uint>, CulledVertexIdsBuffer)
+	SHADER_PARAMETER_SRV(Buffer<float>, CulledVertexRadiusScaleBuffer)
+END_GLOBAL_SHADER_PARAMETER_STRUCT()
+
+typedef TUniformBufferRef<FHairStrandsVertexFactoryUniformShaderParameters> FHairStrandsUniformBuffer;
+
 enum class EHairLODSelectionType
 {
 	Immediate,	// Done on the rendering thread, prior to render
@@ -113,6 +143,7 @@ struct HAIRSTRANDSCORE_API FHairGroupInstance : public FHairStrandsInstance
 		FRDGExternalBuffer DebugAttributeBuffer;
 		FHairGroupInstanceModifer Modifier;
 
+		FHairStrandsUniformBuffer UniformBuffer;
 		FHairStrandsVertexFactory* VertexFactory = nullptr;
 	} Strands;
 
@@ -253,4 +284,6 @@ struct HAIRSTRANDSCORE_API FHairGroupInstance : public FHairStrandsInstance
 		return (BindingType == EHairBindingType::Skinning) ? Debug.SkinningPreviousLocalToWorld :
 															 Debug.RigidPreviousLocalToWorld;
 	}
+
+	FHairStrandsVertexFactoryUniformShaderParameters GetHairStandsUniformShaderParameters() const;
 };
diff --git a/Engine/Plugins/Runtime/HairStrands/Source/HairStrandsCore/Public/HairStrandsVertexFactory.h b/Engine/Plugins/Runtime/HairStrands/Source/HairStrandsCore/Public/HairStrandsVertexFactory.h
--- a/Engine/Plugins/Runtime/HairStrands/Source/HairStrandsCore/Public/HairStrandsVertexFactory.h
+++ b/Engine/Plugins/Runtime/HairStrands/Source/HairStrandsCore/Public/HairStrandsVertexFactory.h
@@ -26,10 +26,10 @@ public:
 	
 	struct FDataType
 	{
-		const FHairGroupInstance* Instance = nullptr;
+		FHairGroupInstance* Instance = nullptr;
 	};
 
-	FHairStrandsVertexFactory(const FHairGroupInstance* Instance, ERHIFeatureLevel::Type InFeatureLevel, const char* InDebugName)
+	FHairStrandsVertexFactory(FHairGroupInstance* Instance, ERHIFeatureLevel::Type InFeatureLevel, const char* InDebugName)
 		: FVertexFactory(InFeatureLevel)
 		, DebugName(InDebugName)
 	{
diff --git a/Engine/Shaders/Private/HairStrands/HairStrandsAttributeCommon.ush b/Engine/Shaders/Private/HairStrands/HairStrandsAttributeCommon.ush
--- a/Engine/Shaders/Private/HairStrands/HairStrandsAttributeCommon.ush
+++ b/Engine/Shaders/Private/HairStrands/HairStrandsAttributeCommon.ush
@@ -8,17 +8,6 @@
 
 ////////////////////////////////////////////////////////////////////////////////
 
-float  HairStrandsVF_Radius;
-float  HairStrandsVF_RootScale;
-float  HairStrandsVF_TipScale;
-float  HairStrandsVF_Length;
-uint   HairStrandsVF_HasAttribute1;
-uint   HairStrandsVF_HasMaterial;
-
-Buffer<uint4>  HairStrandsVF_PositionBuffer;
-Buffer<float4> HairStrandsVF_Attribute0Buffer;
-Buffer<uint4>  HairStrandsVF_Attribute1Buffer;
-Buffer<float4> HairStrandsVF_MaterialBuffer;
 
 struct FHairVertexAttributes
 {
@@ -33,9 +22,9 @@ struct FHairVertexAttributes
 uint4 GetHairStrandsOptionalAttribute1(uint InVertexId)
 {
 	uint4 Out = 0;
-	if (HairStrandsVF_HasAttribute1 > 0)
+	if (HairStrandsVF.HasAttribute1 > 0)
 	{
-		Out = HairStrandsVF_Attribute1Buffer[InVertexId];
+		Out = HairStrandsVF.Attribute1Buffer[InVertexId];
 	}
 	return Out;
 }
@@ -43,9 +32,9 @@ uint4 GetHairStrandsOptionalAttribute1(uint InVertexId)
 float4 GetHairStrandsOptionalMaterial(uint InVertexId)
 {
 	float4 Out = 0;
-	if (HairStrandsVF_HasMaterial > 0)
+	if (HairStrandsVF.HasMaterial > 0)
 	{
-		Out = HairStrandsVF_MaterialBuffer[InVertexId];
+		Out = HairStrandsVF.MaterialBuffer[InVertexId];
 	}
 	return Out;
 }
@@ -56,11 +45,11 @@ FHairVertexAttributes GetHairStrandsAttributes(uint HairPrimitiveId, float2 Vert
 	const uint VertexId1 = HairPrimitiveId+1;
 
 	// Instead of doing the actual interpolation we could rely on dithering as an optimization
-	const FHairAttributes Attributes0 = DecodeHairAttribute(HairStrandsVF_Attribute0Buffer[VertexId0], GetHairStrandsOptionalAttribute1(VertexId0), HairStrandsVF_Length);
-	const FHairAttributes Attributes1 = DecodeHairAttribute(HairStrandsVF_Attribute0Buffer[VertexId1], GetHairStrandsOptionalAttribute1(VertexId1), HairStrandsVF_Length);
+	const FHairAttributes Attributes0 = DecodeHairAttribute(HairStrandsVF.Attribute0Buffer[VertexId0], GetHairStrandsOptionalAttribute1(VertexId0), HairStrandsVF.Length);
+	const FHairAttributes Attributes1 = DecodeHairAttribute(HairStrandsVF.Attribute0Buffer[VertexId1], GetHairStrandsOptionalAttribute1(VertexId1), HairStrandsVF.Length);
 
-	const FHairControlPoint ControlPoint0 = UnpackHairControlPoint(HairStrandsVF_PositionBuffer[VertexId0], float3(0, 0 ,0), HairStrandsVF_Radius, HairStrandsVF_RootScale, HairStrandsVF_TipScale);
-	const FHairControlPoint ControlPoint1 = UnpackHairControlPoint(HairStrandsVF_PositionBuffer[VertexId1], float3(0, 0, 0), HairStrandsVF_Radius, HairStrandsVF_RootScale, HairStrandsVF_TipScale);
+	const FHairControlPoint ControlPoint0 = UnpackHairControlPoint(HairStrandsVF.PositionBuffer[VertexId0], float3(0, 0 ,0), HairStrandsVF.Radius, HairStrandsVF.RootScale, HairStrandsVF.TipScale);
+	const FHairControlPoint ControlPoint1 = UnpackHairControlPoint(HairStrandsVF.PositionBuffer[VertexId1], float3(0, 0, 0), HairStrandsVF.Radius, HairStrandsVF.RootScale, HairStrandsVF.TipScale);
 	const float WorldLength = lerp(Attributes0.WorldLength, Attributes1.WorldLength, VertexUV.x);
 	const float WorldRadius = lerp(ControlPoint0.WorldRadius, ControlPoint1.WorldRadius, VertexUV.x);
 
diff --git a/Engine/Shaders/Private/HairStrands/HairStrandsVertexFactory.ush b/Engine/Shaders/Private/HairStrands/HairStrandsVertexFactory.ush
--- a/Engine/Shaders/Private/HairStrands/HairStrandsVertexFactory.ush
+++ b/Engine/Shaders/Private/HairStrands/HairStrandsVertexFactory.ush
@@ -29,11 +29,6 @@
 #define USE_CULLED_CLUSTER 1
 #endif
 
-#if USE_CULLED_CLUSTER == 1
-	Buffer<uint> CulledVertexIdsBuffer;
-	Buffer<float> CulledVertexRadiusScaleBuffer;
-#endif
-
 ////////////////////////////////////////////////////////////////////////////////
 // HAIR_STRAND_MESH_FACTORY
 // Define used in certain shader like DeepShadow for running specific code (coverage computation for instance) 
@@ -41,17 +36,6 @@
 // This is set by compilation enviromenent of the vertex factory
 ///////////////////////////////////////////////////////////////////////////////
 
-float  HairStrandsVF_Density;
-uint   HairStrandsVF_CullingEnable;
-uint   HairStrandsVF_bUseStableRasterization;
-uint   HairStrandsVF_bScatterSceneLighing;
-
-// Vertex position (3x16bits) are relative to this offset for improving precision while keeping vertex footprint low
-Buffer<float4> HairStrandsVF_PositionOffsetBuffer;
-Buffer<float4> HairStrandsVF_PreviousPositionOffsetBuffer;
-
-Buffer<uint4>  HairStrandsVF_PreviousPositionBuffer;
-Buffer<float4> HairStrandsVF_TangentBuffer;
 
 // ## A quad is made of the following topology and indexing:
 //
@@ -90,16 +74,22 @@ Buffer<float4> HairStrandsVF_TangentBuffer;
  */
 bool UseStableRasterization()
 {
-	return HairStrandsVF_bUseStableRasterization > 0;
+	return HairStrandsVF.StableRasterization > 0;
 }
 
 bool UseScatterSceneLighting()
 {
-	return HairStrandsVF_bScatterSceneLighing > 0;
+	return HairStrandsVF.ScatterSceneLighing > 0;
 }
 
 struct FVertexFactoryInput
 {
+#if RAYHITGROUPSHADER
+	float4 Position;
+	uint TriangleId;          // from RayHit
+	uint TriangleVertexId;    // 0,1,2
+#endif
+
 	// Dynamic instancing related attributes with InstanceIdOffset : ATTRIBUTE13
 	VF_GPUSCENE_DECLARE_INPUT_BLOCK(13)
 
@@ -138,12 +128,11 @@ FVertexInfo GetVertexInfo(FVertexFactoryInput Input)
 	VertexInfo.LodRadiusScale = 1.0f;
 	VertexInfo.bForceInvalidQuad = false;
 #if USE_CULLED_CLUSTER == 1
-	if (HairStrandsVF_CullingEnable)
+	if (HairStrandsVF.CullingEnable)
 	{
-		VertexInfo.VertexIndex = CulledVertexIdsBuffer[VertexInfo.BaseIndex + VertexInfo.IsTip];
-		VertexInfo.LodRadiusScale = CulledVertexRadiusScaleBuffer[VertexInfo.BaseIndex + VertexInfo.IsTip];
+		VertexInfo.VertexIndex = HairStrandsVF.CulledVertexIdsBuffer[VertexInfo.BaseIndex + VertexInfo.IsTip];
+		VertexInfo.LodRadiusScale = HairStrandsVF.CulledVertexRadiusScaleBuffer[VertexInfo.BaseIndex + VertexInfo.IsTip];
 		VertexInfo.HairPrimitiveId = VertexInfo.VertexIndex - VertexInfo.IsTip;
-		VertexInfo.bForceInvalidQuad = VertexInfo.LodRadiusScale <= 0.0f;
 	}
 #endif
 
@@ -151,14 +140,15 @@ FVertexInfo GetVertexInfo(FVertexFactoryInput Input)
 }
 
 #if RAYHITGROUPSHADER
-FVertexFactoryInput LoadVertexFactoryInput(uint TriangleIndex, int VertexIndex)
+FVertexFactoryInput LoadVertexFactoryInputForHGS(uint TriangleIndex, int VertexIndex)
 {
-	// #hair_todo
 	FTriangleBaseAttributes Triangle = LoadTriangleBaseAttributes(TriangleIndex);
 
 	FVertexFactoryInput Input;
 	Input.Position = float4(Triangle.LocalPositions[VertexIndex], 1.0f);
 	Input.VertexId = Triangle.Indices[VertexIndex];
+	Input.TriangleId = TriangleIndex;
+	Input.TriangleVertexId = VertexIndex;
 #if VF_USE_PRIMITIVE_SCENE_DATA
 	// Note: GetInstanceUserData() stores the GPU-Scene primitive ID
 	Input.InstanceIdOffset = GetInstanceUserData() | VF_TREAT_INSTANCE_ID_OFFSET_AS_PRIMITIVE_ID_FLAG;
@@ -184,11 +174,6 @@ struct FVertexFactoryIntermediates
 	FSceneDataIntermediates SceneData;
 };
 
-FPrimitiveSceneData GetPrimitiveData(FVertexFactoryIntermediates Intermediates)
-{
-	return Intermediates.SceneData.Primitive;
-}
-
 uint GetHairPrimitiveId(FVertexFactoryInput Input)
 {
 	return GetVertexInfo(Input).HairPrimitiveId;
@@ -196,19 +181,25 @@ uint GetHairPrimitiveId(FVertexFactoryInput Input)
 
 FHairControlPoint GetVertexPosition(const FVertexInfo VertexInfo, bool bInvalidJointVertex=true)
 {
-	const float3 PositionOffset = HairStrandsVF_PositionOffsetBuffer[0].xyz;
+	const float3 PositionOffset = HairStrandsVF.PositionOffsetBuffer[0].xyz;
 
-	const uint4 PackedData = HairStrandsVF_PositionBuffer[VertexInfo.VertexIndex];
+	const uint4 PackedData = HairStrandsVF.PositionBuffer[VertexInfo.VertexIndex];
 	FHairControlPoint Out = UnpackHairControlPoint(
 		PackedData, 
 		PositionOffset,
-		HairStrandsVF_Radius * VertexInfo.LodRadiusScale, 
-		HairStrandsVF_RootScale,
-		HairStrandsVF_TipScale);
+#if RAYHITGROUPSHADER
+		HairStrandsVF.Radius,
+#else
+		HairStrandsVF.Radius * VertexInfo.LodRadiusScale,
+#endif
+		HairStrandsVF.RootScale,
+		HairStrandsVF.TipScale);
 	
 	// Create a degenerated quad the end of each strand to cut between each strands
 	// #hair_todo: This is not efficient for short strand like fur for instance. Need to revisit that at some point
-#if USE_HAIR_TRIANGLE_STRIP
+#if RAYHITGROUPSHADER
+	const bool bIsInvalidQuad = false; // can't hit an invalid quad in raytracing (filtered by the acceleration structure)
+#elif USE_HAIR_TRIANGLE_STRIP
 	const bool bIsInvalidQuad = (bInvalidJointVertex && Out.Type == HAIR_CONTROLPOINT_END) || VertexInfo.bForceInvalidQuad;
 #else
 	const bool bIsInvalidQuad = (bInvalidJointVertex && Out.Type == HAIR_CONTROLPOINT_START && VertexInfo.IsTip == 1) || (bInvalidJointVertex && Out.Type == HAIR_CONTROLPOINT_END && VertexInfo.IsTip == 0) || VertexInfo.bForceInvalidQuad;
@@ -231,20 +222,26 @@ float GetWorldStrandRadius(FVertexFactoryInput Input)
 
 float3 GetVertexPreviousPosition(FVertexFactoryInput Input, bool bInvalidJointVertex = true)
 {
-	const float3 PreviousPositionOffset = HairStrandsVF_PreviousPositionOffsetBuffer[0].xyz;
+	const float3 PreviousPositionOffset = HairStrandsVF.PreviousPositionOffsetBuffer[0].xyz;
 
 	FVertexInfo VertexInfo = GetVertexInfo(Input);
-	const uint4 PackedData = HairStrandsVF_PreviousPositionBuffer[VertexInfo.VertexIndex];
+	const uint4 PackedData = HairStrandsVF.PreviousPositionBuffer[VertexInfo.VertexIndex];
 	FHairControlPoint Out = UnpackHairControlPoint(
 		PackedData,
 		PreviousPositionOffset,
-		HairStrandsVF_Radius * VertexInfo.LodRadiusScale,
-		HairStrandsVF_RootScale,
-		HairStrandsVF_TipScale);
+#if RAYHITGROUPSHADER
+		HairStrandsVF.Radius,
+#else
+		HairStrandsVF.Radius * VertexInfo.LodRadiusScale,
+#endif
+		HairStrandsVF.RootScale,
+		HairStrandsVF.TipScale);
 
 	// Create a degenerated quad the end of each strand to cut between each strands
 	// #hair_todo: This is not efficient for short strand like fur for instance. Need to revisit that at some point
-#if USE_HAIR_TRIANGLE_STRIP
+#if RAYHITGROUPSHADER
+	const bool bIsInvalidQuad = false; // can't hit an invalid quad in raytracing (filtered by the acceleration structure)
+#elif USE_HAIR_TRIANGLE_STRIP
 	const bool bIsInvalidQuad = (bInvalidJointVertex && Out.Type == HAIR_CONTROLPOINT_END) || VertexInfo.bForceInvalidQuad;
 #else
 	const bool bIsInvalidQuad = (bInvalidJointVertex && Out.Type == HAIR_CONTROLPOINT_START && VertexInfo.IsTip == 1) || (bInvalidJointVertex && Out.Type == HAIR_CONTROLPOINT_END && VertexInfo.IsTip == 0) || VertexInfo.bForceInvalidQuad;
@@ -258,7 +255,11 @@ float3 GetVertexPreviousPosition(FVertexFactoryInput Input, bool bInvalidJointVe
 float2 GetSegmentUVCoord(FVertexFactoryInput Input)
 {
 	FVertexInfo VertexInfo = GetVertexInfo(Input);
+#if RAYHITGROUPSHADER
+	const float VCoord = 0.5; // TODO: find a way to define this?
+#else
 	const float VCoord = VertexInfo.IsLeft ? 0.0f : 1.0f;
+#endif
 	const float UCoord = VertexInfo.IsTip  ? 1.0f : 0.0f;
 	return float2(UCoord, VCoord);
 }
@@ -289,10 +290,10 @@ FMaterialPixelParameters GetMaterialPixelParameters(FVertexFactoryInterpolantsVS
 	return Result;
 }
 
-half3x3 CalcTangentToWorldNoScale(FVertexFactoryIntermediates Intermediates, half3x3 TangentToLocal)
+half3x3 CalcTangentToWorldNoScale(FPrimitiveSceneData SceneData, half3x3 TangentToLocal)
 {
-	half3x3 LocalToWorld = LWCToFloat3x3(GetPrimitiveData(Intermediates).LocalToWorld);
-	half3 InvScale = GetPrimitiveData(Intermediates).InvNonUniformScale;
+	half3x3 LocalToWorld = LWCToFloat3x3(SceneData.LocalToWorld);
+	half3 InvScale = SceneData.InvNonUniformScale;
 	LocalToWorld[0] *= InvScale.x;
 	LocalToWorld[1] *= InvScale.y;
 	LocalToWorld[2] *= InvScale.z;
@@ -309,7 +310,7 @@ FMaterialVertexParameters GetMaterialVertexParameters(FVertexFactoryInput Input,
 	// does not handle instancing!
 	Result.TangentToWorld = Intermediates.TangentToWorld;
 	
-	Result.PrevFrameLocalToWorld = GetPrimitiveData(Intermediates).PreviousLocalToWorld;
+	Result.PrevFrameLocalToWorld = Intermediates.SceneData.Primitive.PreviousLocalToWorld;
 	Result.PreSkinnedPosition = GetVertexPreviousPosition(Input);
 	Result.PreSkinnedNormal = TangentToLocal[2];
 	return Result;
@@ -320,19 +321,18 @@ float4 CalcWorldPosition(float3 Position, FLWCMatrix LocalToWorld)
 	return TransformLocalToTranslatedWorld(Position, LocalToWorld);
 }
 
-half3x3 CalcTangentToLocal(FVertexFactoryInput Input, inout float TangentSign)
+half3x3 CalcTangentToLocal(uint VertexIndex, inout float TangentSign)
 {
-	FVertexInfo VertexInfo = GetVertexInfo(Input);
-	const half3 TangentInputX = HairStrandsVF_TangentBuffer[VertexInfo.VertexIndex * 2 + 0].xyz;
-	const half4 TangentInputZ = HairStrandsVF_TangentBuffer[VertexInfo.VertexIndex * 2 + 1];
+	const half3 TangentInputX = HairStrandsVF.TangentBuffer[VertexIndex * 2 + 0].xyz;
+	const half4 TangentInputZ = HairStrandsVF.TangentBuffer[VertexIndex * 2 + 1];
 
 	half3 TangentX = TangentBias(TangentInputX);
 	half4 TangentZ = TangentBias(TangentInputZ);
 	TangentSign = TangentZ.w;
-	
+
 	// derive the binormal by getting the cross product of the normal and tangent
 	half3 TangentY = cross(TangentZ.xyz, TangentX) * TangentZ.w;
-	
+
 	// Recalculate TangentX off of the other two vectors
 	// This corrects quantization error since TangentX was passed in as a quantized vertex input
 	// The error shows up most in specular off of a mesh with a smoothed UV seam (normal is smooth, but tangents vary across the seam)
@@ -344,10 +344,15 @@ half3x3 CalcTangentToLocal(FVertexFactoryInput Input, inout float TangentSign)
 	return Result;
 }
 
-half3x3 CalcTangentToWorld(FVertexFactoryIntermediates Intermediates, half3x3 TangentToLocal)
+half3x3 CalcTangentToLocal(FVertexFactoryInput Input, inout float TangentSign)
+{
+	FVertexInfo VertexInfo = GetVertexInfo(Input);
+	return CalcTangentToLocal(VertexInfo.VertexIndex, TangentSign);
+}
+
+half3x3 CalcTangentToWorld(FPrimitiveSceneData SceneData, half3x3 TangentToLocal)
 {
-	half3x3 LocalTangentToWorld = CalcTangentToWorldNoScale(Intermediates, TangentToLocal);
-	return LocalTangentToWorld;
+	return CalcTangentToWorldNoScale(SceneData, TangentToLocal);
 }
 
 float3 GetTangent(FVertexFactoryIntermediates Intermediates)
@@ -365,13 +370,13 @@ FVertexFactoryIntermediates GetVertexFactoryIntermediates(FVertexFactoryInput In
 	Intermediates.SceneData			= VF_GPUSCENE_GET_INTERMEDIATES(Input);
 	Intermediates.HairPrimitiveId	= GetHairPrimitiveId(Input);
 	Intermediates.HairPrimitiveUV	= GetSegmentUVCoord(Input);
-	Intermediates.HairDensity		= HairStrandsVF_Density;
+	Intermediates.HairDensity		= HairStrandsVF.Density;
 	Intermediates.HairDimensions	= float2(0.f, GetWorldStrandRadius(Input));
 
 	float TangentSign;
 	Intermediates.TangentToLocal = CalcTangentToLocal(Input, TangentSign);
-	Intermediates.TangentToWorld = CalcTangentToWorld(Intermediates,Intermediates.TangentToLocal);
-	Intermediates.TangentToWorldSign = TangentSign * GetPrimitive_DeterminantSign_FromFlags(GetPrimitiveData(Intermediates).Flags);	
+	Intermediates.TangentToWorld = CalcTangentToWorld(Intermediates.SceneData.Primitive, Intermediates.TangentToLocal);
+	Intermediates.TangentToWorldSign = TangentSign * GetPrimitive_DeterminantSign_FromFlags(Intermediates.SceneData.Primitive.Flags);
 	return Intermediates;
 }
 
@@ -387,6 +392,15 @@ half3x3 VertexFactoryGetTangentToLocal( FVertexFactoryInput Input, FVertexFactor
 	return Intermediates.TangentToLocal;
 }
 
+#if RAYHITGROUPSHADER
+
+float4 VertexFactoryGetWorldPosition(FVertexFactoryInput Input, FVertexFactoryIntermediates Intermediates)
+{
+	return CalcWorldPosition(Input.Position.xyz, Intermediates.SceneData.Primitive.LocalToWorld);
+}
+
+#else
+
 float4 ComputeViewAlignedWorldPosition(FVertexFactoryInput Input, float3 WorldTangent, float4 WorldPosition, float WorldStrandRadius, FHairViewInfo HairViewInfo)
 {
 	FVertexInfo VertexInfo = GetVertexInfo(Input);
@@ -411,7 +425,7 @@ float4 VertexFactoryGetWorldPosition_Visibility(FVertexFactoryInput Input, FHair
 	FHairControlPoint ControlPoint = GetVertexPosition(VertexInfo);
 
 	FVertexFactoryIntermediates Intermediates = GetVertexFactoryIntermediates(Input);
-	const float4 WorldPosition = CalcWorldPosition(ControlPoint.Position, GetPrimitiveData(Intermediates).LocalToWorld);
+	const float4 WorldPosition = CalcWorldPosition(ControlPoint.Position, Intermediates.SceneData.Primitive.LocalToWorld);
 	const float3 WorldTangent = GetTangent(Intermediates);
 
 	// Minimal radius to snap the strand to a sample/pixel center (to avoid aliasing)
@@ -432,10 +446,10 @@ float4 VertexFactoryGetWorldPosition_Voxelization(FVertexFactoryInput Input, FHa
 
 	FHairControlPoint ControlPoint = GetVertexPosition(VertexInfo);
 	OutHairWorldRadius = ControlPoint.WorldRadius;
-	OutHairDensity = HairStrandsVF_Density;
+	OutHairDensity = HairStrandsVF.Density;
 
 	FVertexFactoryIntermediates Intermediates = GetVertexFactoryIntermediates(Input);
-	const float4 WorldPosition = CalcWorldPosition(ControlPoint.Position, GetPrimitiveData(Intermediates).LocalToWorld);
+	const float4 WorldPosition = CalcWorldPosition(ControlPoint.Position, Intermediates.SceneData.Primitive.LocalToWorld);
 	const float3 WorldTangent = GetTangent(Intermediates);
 
 	// Minimal radius to snap the strand to a sample/pixel center (to avoid aliasing)
@@ -451,7 +465,7 @@ float4 VertexFactoryGetWorldPosition(FVertexFactoryInput Input, FVertexFactoryIn
 {
 	FHairControlPoint ControlPoint = GetVertexPosition(Input);
 	const float3 VertexPosition = ControlPoint.Position;
-	const float4 WorldPosition = CalcWorldPosition(VertexPosition, GetPrimitiveData(Intermediates).LocalToWorld);
+	const float4 WorldPosition = CalcWorldPosition(VertexPosition, Intermediates.SceneData.Primitive.LocalToWorld);
 
 	// Hair shader expect the WorldNormal to be the tangent vector
 	const float3 WorldTangent = GetTangent(Intermediates);
@@ -471,6 +485,7 @@ float4 VertexFactoryGetWorldPosition(FVertexFactoryInput Input, FVertexFactoryIn
 
 	return VertexFactoryGetWorldPosition(Input, Intermediates, HairViewInfo);
 }
+#endif
 
 float4 VertexFactoryGetRasterizedWorldPosition(FVertexFactoryInput Input, FVertexFactoryIntermediates Intermediates, float4 InWorldPosition)
 {
@@ -496,8 +511,8 @@ FVertexFactoryInterpolantsVSToPS VertexFactoryGetInterpolantsVSToPS(FVertexFacto
 	Interpolants.EyeIndex = 0;
 #endif
 	
-	Interpolants.HairPrimitiveId = GetHairPrimitiveId(Input);
-	Interpolants.HairPrimitiveUV = GetSegmentUVCoord(Input);
+	Interpolants.HairPrimitiveId = Intermediates.HairPrimitiveId;
+	Interpolants.HairPrimitiveUV = Intermediates.HairPrimitiveUV;
 	SetPrimitiveId(Interpolants, Intermediates.SceneData.PrimitiveId);
 
 	return Interpolants;
@@ -512,13 +527,13 @@ float4 VertexFactoryGetWorldPosition(FVertexFactoryInput Input)
 // @return translated world position (without quad extension/reorientation).This is used only for velocity computation
 float4 VertexFactoryGetWorldPositionRaw(FVertexFactoryInput Input, FVertexFactoryIntermediates Intermediates, bool bInvalidJointVertex = true)
 {
-	return CalcWorldPosition(GetVertexPosition(Input, bInvalidJointVertex).Position, GetPrimitiveData(Intermediates).LocalToWorld);
+	return CalcWorldPosition(GetVertexPosition(Input, bInvalidJointVertex).Position, Intermediates.SceneData.Primitive.LocalToWorld);
 }
 
 // @return previous translated world position
 float4 VertexFactoryGetPreviousWorldPosition(FVertexFactoryInput Input, FVertexFactoryIntermediates Intermediates, bool bInvalidJointVertex = true)
 {
-	float4x4 PreviousLocalToWorldTranslated = LWCMultiplyTranslation(GetPrimitiveData(Intermediates).PreviousLocalToWorld, ResolvedView.PrevPreViewTranslation);
+	float4x4 PreviousLocalToWorldTranslated = LWCMultiplyTranslation(Intermediates.SceneData.Primitive.PreviousLocalToWorld, ResolvedView.PrevPreViewTranslation);
 
 
 	const float3 VertexPosition = GetVertexPreviousPosition(Input, bInvalidJointVertex);
diff --git a/Engine/Source/Runtime/Renderer/Private/HairStrands/HairStrandsCluster.cpp b/Engine/Source/Runtime/Renderer/Private/HairStrands/HairStrandsCluster.cpp
--- a/Engine/Source/Runtime/Renderer/Private/HairStrands/HairStrandsCluster.cpp
+++ b/Engine/Source/Runtime/Renderer/Private/HairStrands/HairStrandsCluster.cpp
@@ -184,6 +184,12 @@ void CreateHairStrandsMacroGroups(
 		if (!bIsHairStrandsFactory)
 			return;
 
+		// Ensure that the element has been initialized
+		const FHairGroupPublicData* HairGroupPublicData = reinterpret_cast<const FHairGroupPublicData*>(MeshBatchAndRelevance->Mesh->Elements[0].VertexFactoryUserData);
+		const bool bIsValid = HairGroupPublicData->VFInput.Strands.PositionBufferRHISRV != nullptr;
+		if (!bIsValid)
+			return;
+
 		const FBoxSphereBounds& PrimitiveBounds = Proxy->GetBounds();
 
 		bool bFound = false;
