Change hair strands tangent to use common ortho basis function.

diff --git a/Engine/Shaders/Private/HairStrands/HairStrandsTangent.usf b/Engine/Shaders/Private/HairStrands/HairStrandsTangent.usf
--- a/Engine/Shaders/Private/HairStrands/HairStrandsTangent.usf
+++ b/Engine/Shaders/Private/HairStrands/HairStrandsTangent.usf
@@ -1,6 +1,7 @@
 // Copyright Epic Games, Inc. All Rights Reserved.
 
 #include "../Common.ush"
+#include "../MonteCarlo.ush"
 #include "HairStrandsClusterCommon.ush"
 
 uint VertexCount;
@@ -90,22 +91,10 @@ void MainCS(uint2 DispatchThreadId : SV_DispatchThreadID)
 
 		const float3 T = normalize(T0 + T1);
 
-		const float Norm2 = dot(T, float3(0,0,-1));
-		float3 Right;
-		if (Norm2 > 0e-3f)
-		{
-			Right = normalize(cross(T, float3(0,0,-1)));
-		}
-		else
-		{
-			Right = normalize(cross(T, float3(0,1,0)));
-		}
-
-		// Build ortho basis
-		float3 TangentX = normalize(cross(Right, T));
-		float3 TangentZ = T;
+		// TODO: switch to rotation minimizing frames to avoid the discontinuity
+		float3x3 Basis = GetTangentBasis(T);
 
-		OutputTangentBuffer[IndexCurr*2]   = PackNormal(TangentX, 0);
-		OutputTangentBuffer[IndexCurr*2+1] = PackNormal(TangentZ, 127);
+		OutputTangentBuffer[IndexCurr*2]   = PackNormal(Basis[0], 0);
+		OutputTangentBuffer[IndexCurr*2+1] = PackNormal(Basis[2], 127);
 	}
 }
\ No newline at end of file
