Add used/allocation statistic to hair debug mode 1

diff --git a/Engine/Shaders/Private/HairStrands/HairStrandsDebugPrint.usf b/Engine/Shaders/Private/HairStrands/HairStrandsDebugPrint.usf
--- a/Engine/Shaders/Private/HairStrands/HairStrandsDebugPrint.usf
+++ b/Engine/Shaders/Private/HairStrands/HairStrandsDebugPrint.usf
@@ -36,6 +36,7 @@ uint FastResolveMask;
 uint HairMacroGroupCount;
 uint2 GroupSize; 
 uint HairVisibilityNodeGroupSize;
+uint AllocatedSampleCount;
 
 Texture2D<float> HairCountTexture;
 Texture2D<uint> HairCountUintTexture;
@@ -49,6 +50,12 @@ Buffer<int> HairMacroGroupAABBBuffer;
 #define MATERIAL_DATA 1
 #define DETAIL_DATA 1
 
+FFontColor GetOccupancyColor(float In)
+{
+	float3 Color = lerp(float3(0, 1, 0), float3(1, 0, 0), saturate(In));
+	return InitFontColor(Color);
+}
+
 [numthreads(1, 1, 1)]
 void MainCS(uint GroupIndex : SV_GroupIndex, uint3 DispatchThreadId : SV_DispatchThreadID)
 {	
@@ -64,15 +71,24 @@ void MainCS(uint GroupIndex : SV_GroupIndex, uint3 DispatchThreadId : SV_Dispatc
 
 		const uint DispatchX = HairVisibilityIndirectArgsBuffer[0];
 		const uint DispatchY = HairVisibilityIndirectArgsBuffer[1];
-		const uint TotalSampleCount = DispatchX * DispatchY * HairVisibilityNodeGroupSize;
-		const uint2 EquivalentResolution = ceil(sqrt(TotalSampleCount));
+		const uint UsedSampleCount = HairStrands.HairSampleCount.Load(uint3(0, 0, 0));
+		const uint DispatchedSampleCount = DispatchX * DispatchY * HairVisibilityNodeGroupSize;
+
+		const uint2 EquivalentResolution = ceil(sqrt(UsedSampleCount));
+		
 
 		// Pixel coord
 		float2 OriginalPos = float2(50, 50) / float2(MaxResolution);
 		float2 Pos = OriginalPos;
-		Pos = ShaderPrintText(Pos, TEXT("Total Sample Count "), FontCyan);
-		Pos = ShaderPrint(Pos, TotalSampleCount, FontCyan);
-		Pos = ShaderPrint(Pos, HairStrands.HairSampleCount.Load(uint3(0,0,0)), FontCyan);
+		Pos = ShaderPrintText(Pos, TEXT("Sample Count (Used/Allocated)"), FontCyan);
+		Pos = ShaderPrint(Pos, UsedSampleCount, FontCyan);
+		Pos = ShaderPrint(Pos, AllocatedSampleCount, FontCyan);
+
+		const float AllocationPercentage = float(UsedSampleCount) / float(AllocatedSampleCount);
+		const FFontColor OccupancyColor = GetOccupancyColor(AllocationPercentage);
+		Pos = ShaderPrintText(Pos, TEXT(" \% "), OccupancyColor);
+		Pos = ShaderPrint(Pos, AllocationPercentage * 100.f, OccupancyColor);
+
 		AddNewLine(OriginalPos, Pos);
 
 		Pos = ShaderPrintText(Pos, TEXT("Equivalent Res     "), FontCyan);
diff --git a/Engine/Source/Runtime/Renderer/Private/HairStrands/HairStrandsDebug.cpp b/Engine/Source/Runtime/Renderer/Private/HairStrands/HairStrandsDebug.cpp
--- a/Engine/Source/Runtime/Renderer/Private/HairStrands/HairStrandsDebug.cpp
+++ b/Engine/Source/Runtime/Renderer/Private/HairStrands/HairStrandsDebug.cpp
@@ -240,6 +240,7 @@ class FHairDebugPrintCS : public FGlobalShader
 		SHADER_PARAMETER(uint32, FastResolveMask)
 		SHADER_PARAMETER(uint32, HairMacroGroupCount)
 		SHADER_PARAMETER(uint32, HairVisibilityNodeGroupSize)
+		SHADER_PARAMETER(uint32, AllocatedSampleCount)
 		SHADER_PARAMETER_RDG_TEXTURE(Texture2D, HairCountTexture)
 		SHADER_PARAMETER_RDG_TEXTURE(Texture2D, HairCountUintTexture)
 		SHADER_PARAMETER_RDG_BUFFER_SRV(Buffer<uint>, HairVisibilityIndirectArgsBuffer)
@@ -286,6 +287,7 @@ static void AddDebugHairPrintPass(
 	Parameters->ViewUniformBuffer = View->ViewUniformBuffer;
 	Parameters->MaxResolution = VisibilityData.CoverageTexture ? VisibilityData.CoverageTexture->Desc.Extent : FIntPoint(0,0);
 	Parameters->PixelCoord = View->CursorPos;
+	Parameters->AllocatedSampleCount = VisibilityData.MaxNodeCount;
 	Parameters->FastResolveMask = STENCIL_TEMPORAL_RESPONSIVE_AA_MASK;
 	Parameters->HairCountTexture = ViewHairCountTexture;
 	Parameters->HairCountUintTexture = ViewHairCountUintTexture;
