Disable SceneDepthAux when r.MobileHDR=0 for non-IOS platforms

diff --git a/Engine/Shaders/Private/MobileBasePassPixelShader.usf b/Engine/Shaders/Private/MobileBasePassPixelShader.usf
--- a/Engine/Shaders/Private/MobileBasePassPixelShader.usf
+++ b/Engine/Shaders/Private/MobileBasePassPixelShader.usf
@@ -99,7 +99,10 @@ int NumDynamicPointLights;
 #define FORCE_VERTEX_NORMAL (MOBILE_QL_DISABLE_MATERIAL_NORMAL)
 #define SUPPORT_SPOTLIGHTS_SHADOW (MATERIALBLENDING_SOLID || MATERIALBLENDING_MASKED) && PROJECT_MOBILE_ENABLE_MOVABLE_SPOTLIGHTS_SHADOW
 #define DEFERRED_SHADING_PATH (MOBILE_DEFERRED_SHADING && ((MATERIALBLENDING_SOLID || MATERIALBLENDING_MASKED) && !MATERIAL_SHADINGMODEL_SINGLELAYERWATER))
-#define USE_SCENE_DEPTH_AUX ((DEFERRED_SHADING_PATH && METAL_PROFILE) || !DEFERRED_SHADING_PATH) && (MATERIALBLENDING_SOLID || MATERIALBLENDING_MASKED) && !COMPILER_SWITCH
+
+// SceneDepthAux always enabled for IOS, on other platforms only enabled when in forward and MobileHDR=true
+#define USE_SCENE_DEPTH_AUX ((METAL_PROFILE || (!DEFERRED_SHADING_PATH && OUTPUT_MOBILE_HDR)) && (MATERIALBLENDING_SOLID || MATERIALBLENDING_MASKED) && !COMPILER_SWITCH)
+
 #if DEFERRED_SHADING_PATH
 	#define SV_TargetDepthAux SV_Target4
 #else
diff --git a/Engine/Source/Runtime/RenderCore/Private/RenderUtils.cpp b/Engine/Source/Runtime/RenderCore/Private/RenderUtils.cpp
--- a/Engine/Source/Runtime/RenderCore/Private/RenderUtils.cpp
+++ b/Engine/Source/Runtime/RenderCore/Private/RenderUtils.cpp
@@ -1117,13 +1117,17 @@ RENDERCORE_API bool IsMobileDeferredShadingEnabled(const FStaticShaderPlatform P
 
 RENDERCORE_API bool MobileRequiresSceneDepthAux(const FStaticShaderPlatform Platform)
 {
+	static const auto CVarMobileHDR = IConsoleManager::Get().FindTConsoleVariableDataInt(TEXT("r.MobileHDR"));
+	const bool bMobileHDR = (CVarMobileHDR && CVarMobileHDR->GetValueOnAnyThread() != 0);
+
 	// SceneDepth is used on most mobile platforms when forward shading is enabled and always on IOS.
 	if (IsMetalMobilePlatform(Platform))
 	{
 		return true;
 	}
-	else if (!IsMobileDeferredShadingEnabled(Platform))
+	else if (!IsMobileDeferredShadingEnabled(Platform) && bMobileHDR)
 	{
+		// SceneDepthAux disabled when MobileHDR=false for non-IOS
 		return IsAndroidOpenGLESPlatform(Platform) || IsVulkanMobilePlatform(Platform) || IsSimulatedPlatform(Platform);
 	}
 	return false;
