Fixed: gpuscene regression on mobile (shader compilation errors)

diff --git a/Engine/Shaders/Private/SceneData.ush b/Engine/Shaders/Private/SceneData.ush
--- a/Engine/Shaders/Private/SceneData.ush
+++ b/Engine/Shaders/Private/SceneData.ush
@@ -703,6 +703,9 @@ FSceneDataIntermediates GetSceneDataIntermediates(uint DrawInstanceId, float4 In
 	// we supply some of primtive data using per-instance vertex data, and rest of primitive data comes from Primitive UB
 	// If vertex shader uses any of Primitive UB data associated drawcalls will not auto-instance
 	FPrimitiveSceneData Primitive = GetPrimitiveDataFromUniformBuffer();
+	
+	// TODO: add support for LWC, we should pack it as integer tile coordinates 
+	Primitive.TilePosition = 0;
 
 	// TODO: pack important primitive and instance flags here
 	const uint PrimitiveFlags = asuint(InstanceTransform1.w);
@@ -758,8 +761,10 @@ FSceneDataIntermediates GetSceneDataIntermediates(uint DrawInstanceId, float4 In
 	InstanceData.LocalBoundsExtent = UnpackLocalBoundsExtent(InstanceAuxData.zw);
 	Intermediates.Primitive.LocalObjectBoundsMin = InstanceData.LocalBoundsCenter - InstanceData.LocalBoundsExtent; 
 	Intermediates.Primitive.LocalObjectBoundsMax = InstanceData.LocalBoundsCenter + InstanceData.LocalBoundsExtent; 
-	Intermediates.Primitive.ObjectWorldPositionAndRadius.xyz = mul(float4(InstanceData.LocalBoundsCenter.xyz, 0.0f), InstanceData.LocalToWorld).xyz;
-	Intermediates.Primitive.ObjectWorldPositionAndRadius.w = length(InstanceData.LocalBoundsExtent * InstanceData.NonUniformScale.xyz);
+
+	float3 ObjectRelativeWorldPosition = mul(float4(InstanceData.LocalBoundsCenter.xyz, 0.0f), LocalToRelativeWorld).xyz;
+	Intermediates.Primitive.ObjectWorldPosition = MakeLWCVector3(Primitive.TilePosition, ObjectRelativeWorldPosition);
+	Intermediates.Primitive.ObjectRadius = length(InstanceData.LocalBoundsExtent * InstanceData.NonUniformScale.xyz);
 #endif
 	
 	return Intermediates; 
