Fixed: more shader compilation errors related to mobile gpuscene

diff --git a/Engine/Shaders/Private/SceneData.ush b/Engine/Shaders/Private/SceneData.ush
--- a/Engine/Shaders/Private/SceneData.ush
+++ b/Engine/Shaders/Private/SceneData.ush
@@ -705,7 +705,7 @@ FSceneDataIntermediates GetSceneDataIntermediates(uint DrawInstanceId, float4 In
 	FPrimitiveSceneData Primitive = GetPrimitiveDataFromUniformBuffer();
 	
 	// TODO: add support for LWC, we should pack it as integer tile coordinates 
-	Primitive.TilePosition = 0;
+	float3 TilePosition = float3(0,0,0);
 
 	// TODO: pack important primitive and instance flags here
 	const uint PrimitiveFlags = asuint(InstanceTransform1.w);
@@ -720,7 +720,7 @@ FSceneDataIntermediates GetSceneDataIntermediates(uint DrawInstanceId, float4 In
 		float4(InstanceTransform3.xyz, 0.0f), 
 		float4(InstanceOrigin.xyz, 1.0f));
 	
-	InstanceData.LocalToWorld = MakeLWCMatrix(Primitive.TilePosition, LocalToRelativeWorld);
+	InstanceData.LocalToWorld = MakeLWCMatrix(TilePosition, LocalToRelativeWorld);
 
 	InstanceData.Flags = InstanceFlags;
 	InstanceData.PrimitiveId = asuint(InstanceOrigin.w);
@@ -733,7 +733,7 @@ FSceneDataIntermediates GetSceneDataIntermediates(uint DrawInstanceId, float4 In
 	#error Compressed instance transform is not supported on mobile path
 #endif
 
-	ComputeInstanceDerivedData(InstanceData, Primitive.TilePosition, LocalToRelativeWorld);
+	ComputeInstanceDerivedData(InstanceData, TilePosition, LocalToRelativeWorld);
 
 	//
 	FSceneDataIntermediates Intermediates = (FSceneDataIntermediates)0;
@@ -763,7 +763,7 @@ FSceneDataIntermediates GetSceneDataIntermediates(uint DrawInstanceId, float4 In
 	Intermediates.Primitive.LocalObjectBoundsMax = InstanceData.LocalBoundsCenter + InstanceData.LocalBoundsExtent; 
 
 	float3 ObjectRelativeWorldPosition = mul(float4(InstanceData.LocalBoundsCenter.xyz, 0.0f), LocalToRelativeWorld).xyz;
-	Intermediates.Primitive.ObjectWorldPosition = MakeLWCVector3(Primitive.TilePosition, ObjectRelativeWorldPosition);
+	Intermediates.Primitive.ObjectWorldPosition = MakeLWCVector3(TilePosition, ObjectRelativeWorldPosition);
 	Intermediates.Primitive.ObjectRadius = length(InstanceData.LocalBoundsExtent * InstanceData.NonUniformScale.xyz);
 #endif
 	
