Use explicit unsigned integer literal in ComputeGenerateMips.usf to fix compile error with HLSLcc introduced with CL 17983475.

#rnx

diff --git a/Engine/Shaders/Private/ComputeGenerateMips.usf b/Engine/Shaders/Private/ComputeGenerateMips.usf
--- a/Engine/Shaders/Private/ComputeGenerateMips.usf
+++ b/Engine/Shaders/Private/ComputeGenerateMips.usf
@@ -59,8 +59,8 @@ void BuildIndirectDispatchArgsCS(uint DT_ID : SV_DispatchThreadID)
 		{
 
 			uint2 DestTextureSize = uint2(
-				max(TextureSize.x >> (DT_ID + 1), 1),
-				max(TextureSize.y >> (DT_ID + 1), 1));
+				max(TextureSize.x >> (DT_ID + 1u), 1u),
+				max(TextureSize.y >> (DT_ID + 1u), 1u));
 
 			DestTextureSize = (DestTextureSize + 8 - 1) / 8;
 
