Add raytracing material evaluation to hair strands and custom primitive intersection.

Work done by Chris kulla.

diff --git a/Engine/Plugins/Runtime/HairStrands/Source/HairStrandsCore/HairStrandsCore.Build.cs b/Engine/Plugins/Runtime/HairStrands/Source/HairStrandsCore/HairStrandsCore.Build.cs
--- a/Engine/Plugins/Runtime/HairStrands/Source/HairStrandsCore/HairStrandsCore.Build.cs
+++ b/Engine/Plugins/Runtime/HairStrands/Source/HairStrandsCore/HairStrandsCore.Build.cs
@@ -11,6 +11,7 @@ namespace UnrealBuildTool.Rules
 				new string[] {
 					ModuleDirectory + "/Private",
 					EngineDirectory + "/Source/Runtime/Renderer/Private",
+					EngineDirectory + "/Shaders/Shared",
 				});
 			PublicIncludePaths.Add(ModuleDirectory + "/Public");
 
diff --git a/Engine/Plugins/Runtime/HairStrands/Source/HairStrandsCore/Private/GroomComponent.cpp b/Engine/Plugins/Runtime/HairStrands/Source/HairStrandsCore/Private/GroomComponent.cpp
--- a/Engine/Plugins/Runtime/HairStrands/Source/HairStrandsCore/Private/GroomComponent.cpp
+++ b/Engine/Plugins/Runtime/HairStrands/Source/HairStrandsCore/Private/GroomComponent.cpp
@@ -11,7 +11,8 @@
 #include "HairStrandsRendering.h"
 #include "HairCardsVertexFactory.h"
 #include "HairStrandsVertexFactory.h"
-#include "RayTracingInstanceUtils.h"
+#include "RayTracingInstance.h"
+#include "RayTracingDefinitions.h"
 #include "HairStrandsInterface.h"
 #include "UObject/UObjectIterator.h"
 #include "GlobalShader.h"
@@ -253,6 +254,12 @@ static EHairMaterialCompatibility IsHairMaterialCompatible(UMaterialInterface* M
 	return EHairMaterialCompatibility::Valid;
 }
 
+enum class EHairMeshBatchType
+{
+	Raster,
+	Raytracing
+};
+
 /////////////////////////////////////////////////////////////////////////////////////////
 /////////////////////////////////////////////////////////////////////////////////////////
 //  FStrandHairSceneProxy
@@ -534,53 +541,51 @@ public:
 			const EHairGeometryType GeometryType = Instance->HairGroupPublicData->VFInput.GeometryType;
 			const uint32 LODIndex = Instance->HairGroupPublicData->GetIntLODIndex();
 
-			if (GeometryType == EHairGeometryType::Strands)
+			FHairStrandsRaytracingResource* RTGeometry = nullptr;
+			uint8 RayTracingMask = RAY_TRACING_MASK_OPAQUE;
+			switch (GeometryType)
 			{
-				FHairStrandsRaytracingResource* RTGeometry = Instance->Strands.RenRaytracingResource ? Instance->Strands.RenRaytracingResource : nullptr;
-				if (RTGeometry && RTGeometry->RayTracingGeometry.RayTracingGeometryRHI.IsValid())
+				case EHairGeometryType::Strands:
 				{
-					for (const FRayTracingGeometrySegment& Segment : RTGeometry->RayTracingGeometry.Initializer.Segments)
-					{
-						check(Segment.VertexBuffer.IsValid());
-					}
-					AddOpaqueRaytracingInstance(OverrideLocalToWorld, &RTGeometry->RayTracingGeometry, RaytracingInstanceMask_ThinShadow, OutRayTracingInstances);
+					RTGeometry = Instance->Strands.RenRaytracingResource;
+					RayTracingMask = RAY_TRACING_MASK_THIN_SHADOW;
+					break;
 				}
-			}
-			else if (GeometryType == EHairGeometryType::Cards)
-			{
-				FHairStrandsRaytracingResource* RTGeometry = Instance->Cards.LODs[LODIndex].RaytracingResource ? Instance->Cards.LODs[LODIndex].RaytracingResource : nullptr;
-				if (RTGeometry && RTGeometry->RayTracingGeometry.RayTracingGeometryRHI.IsValid())
+				case EHairGeometryType::Cards:
 				{
-					for (const FRayTracingGeometrySegment& Segment : RTGeometry->RayTracingGeometry.Initializer.Segments)
-					{
-						check(Segment.VertexBuffer.IsValid());
-					}
-				
-					if (FMeshBatch* MeshBatch = CreateMeshBatch(Context.ReferenceView, Context.ReferenceViewFamily, Context.RayTracingMeshResourceCollector, Instance, GroupIt, nullptr))
-					{
-						TArray<FMeshBatch> MeshBatches;
-						MeshBatches.Add(*MeshBatch);
-						AddOpaqueRaytracingInstance(OverrideLocalToWorld, &RTGeometry->RayTracingGeometry, RaytracingInstanceMask_Opaque, MeshBatches, GetScene().GetFeatureLevel(), OutRayTracingInstances);
-					}
-
+					RTGeometry = Instance->Cards.LODs[LODIndex].RaytracingResource;
+					break;
+				}
+				case EHairGeometryType::Meshes:
+				{
+					RTGeometry = Instance->Meshes.LODs[LODIndex].RaytracingResource;
+					break;
 				}
 			}
-			else if (GeometryType == EHairGeometryType::Meshes)
+
+			if (RTGeometry && RTGeometry->RayTracingGeometry.RayTracingGeometryRHI.IsValid())
 			{
-				FHairStrandsRaytracingResource* RTGeometry = Instance->Meshes.LODs[LODIndex].RaytracingResource ? Instance->Meshes.LODs[LODIndex].RaytracingResource : nullptr;
-				if (RTGeometry && RTGeometry->RayTracingGeometry.RayTracingGeometryRHI.IsValid())
+				for (const FRayTracingGeometrySegment& Segment : RTGeometry->RayTracingGeometry.Initializer.Segments)
 				{
-					for (const FRayTracingGeometrySegment& Segment : RTGeometry->RayTracingGeometry.Initializer.Segments)
+					check(Segment.VertexBuffer.IsValid());
+				}
+				// ViewFamily.EngineShowFlags.PathTracing
+				if (FMeshBatch* MeshBatch = CreateMeshBatch(Context.ReferenceView, Context.ReferenceViewFamily, Context.RayTracingMeshResourceCollector, EHairMeshBatchType::Raytracing, Instance, GroupIt, nullptr))
+				{
+					FRayTracingInstance RayTracingInstance;
+					RayTracingInstance.Geometry = &RTGeometry->RayTracingGeometry;
+					RayTracingInstance.Materials.Add(*MeshBatch);
+					RayTracingInstance.InstanceTransforms.Add(OverrideLocalToWorld);
+					RayTracingInstance.Mask = RayTracingMask;
+					RayTracingInstance.BuildInstanceMaskAndFlags(GetScene().GetFeatureLevel());
+
+					// If thin shadow is requested, we ensure that regular shadow mask is not added.
+					if (RayTracingMask == RAY_TRACING_MASK_THIN_SHADOW)
 					{
-						check(Segment.VertexBuffer.IsValid());
+						RayTracingInstance.Mask &= ~RAY_TRACING_MASK_SHADOW;
 					}
 
-					if (FMeshBatch* MeshBatch = CreateMeshBatch(Context.ReferenceView, Context.ReferenceViewFamily, Context.RayTracingMeshResourceCollector, Instance, GroupIt, nullptr))
-					{
-						TArray<FMeshBatch> MeshBatches;
-						MeshBatches.Add(*MeshBatch);
-						AddOpaqueRaytracingInstance(OverrideLocalToWorld, &RTGeometry->RayTracingGeometry, RaytracingInstanceMask_Opaque, MeshBatches, GetScene().GetFeatureLevel(), OutRayTracingInstances);
-					}
+					OutRayTracingInstances.Add(RayTracingInstance);
 				}
 			}
 		}
@@ -684,7 +689,7 @@ public:
 						Debug_MaterialProxy = DebugMaterial;
 					}
 
-					if (FMeshBatch* MeshBatch = CreateMeshBatch(View, ViewFamily, Collector, Instances[GroupIt], GroupIt, Debug_MaterialProxy))
+					if (FMeshBatch* MeshBatch = CreateMeshBatch(View, ViewFamily, Collector, EHairMeshBatchType::Raster, Instances[GroupIt], GroupIt, Debug_MaterialProxy))
 					{
 						Collector.AddMesh(ViewIndex, *MeshBatch);
 					}
@@ -706,6 +711,7 @@ public:
 		const FSceneView* View,
 		const FSceneViewFamily& ViewFamily,
 		FMeshElementCollector& Collector,
+		const EHairMeshBatchType MeshBatchType,
 		const FHairGroupInstance* Instance,
 		uint32 GroupIndex,
 		FMaterialRenderProxy* Debug_MaterialProxy) const
@@ -810,7 +816,7 @@ public:
 
 		const bool bUseCardsOrMeshes = GeometryType == EHairGeometryType::Cards || GeometryType == EHairGeometryType::Meshes;
 		Mesh.CastShadow = bUseCardsOrMeshes;
-		Mesh.bUseForMaterial  = bUseCardsOrMeshes;
+		Mesh.bUseForMaterial = MeshBatchType == EHairMeshBatchType::Raytracing || bUseCardsOrMeshes;
 		Mesh.bUseForDepthPass = bUseCardsOrMeshes;
 		Mesh.SegmentIndex = 0;
 		#if !(UE_BUILD_SHIPPING || UE_BUILD_TEST)
diff --git a/Engine/Plugins/Runtime/HairStrands/Source/HairStrandsCore/Private/GroomResources.cpp b/Engine/Plugins/Runtime/HairStrands/Source/HairStrandsCore/Private/GroomResources.cpp
--- a/Engine/Plugins/Runtime/HairStrands/Source/HairStrandsCore/Private/GroomResources.cpp
+++ b/Engine/Plugins/Runtime/HairStrands/Source/HairStrandsCore/Private/GroomResources.cpp
@@ -1501,14 +1501,22 @@ void FHairCardsInterpolationResource::InternalRelease()
 // which used 6*float4 = 48 bytes per vertex
 // NOTE: the vertex buffer is a float4 because it is registered as a UAV for the compute shader to work
 // TODO: use a plain float vertex buffer with 3x the entries instead to save memory? (float3 UAVs are not allowed)
+bool GetSupportHairStrandsProceduralPrimitive(EShaderPlatform InShaderPlatform);
 FHairStrandsRaytracingResource::FHairStrandsRaytracingResource(const FHairStrandsBulkData& InData, const FHairResourceName& ResourceName) :
-	FHairCommonResource(EHairStrandsAllocationType::Deferred, ResourceName),
-#if STRANDS_CUSTOM_INTERSECTOR
-	VertexCount(InData.GetNumPoints() * 2), bOwnBuffers(true) // only allocate space for primitive AABBs
-#else
-	VertexCount(InData.GetNumPoints() * 4), IndexCount(InData.GetNumPoints() * 8 * 3), bOwnBuffers(true)
-#endif
-{}
+	FHairCommonResource(EHairStrandsAllocationType::Deferred, ResourceName)
+{
+	bOwnBuffers = true;
+	bProceduralPrimitive = GetSupportHairStrandsProceduralPrimitive(GMaxRHIShaderPlatform);
+	if (bProceduralPrimitive)
+	{
+		VertexCount = InData.GetNumPoints() * 2; // only allocate space for primitive AABBs
+	}
+	else
+	{
+		VertexCount = InData.GetNumPoints() * 4;
+		IndexCount = InData.GetNumPoints() * 8 * 3;
+	}
+}
 
 FHairStrandsRaytracingResource::FHairStrandsRaytracingResource(const FHairCardsBulkData& InData, const FHairResourceName& ResourceName) :
 	FHairCommonResource(EHairStrandsAllocationType::Deferred, ResourceName),
diff --git a/Engine/Plugins/Runtime/HairStrands/Source/HairStrandsCore/Private/HairStrandsRendering.cpp b/Engine/Plugins/Runtime/HairStrands/Source/HairStrandsCore/Private/HairStrandsRendering.cpp
--- a/Engine/Plugins/Runtime/HairStrands/Source/HairStrandsCore/Private/HairStrandsRendering.cpp
+++ b/Engine/Plugins/Runtime/HairStrands/Source/HairStrandsCore/Private/HairStrandsRendering.cpp
@@ -41,6 +41,9 @@ static FAutoConsoleVariableRef CVarHairCardsInterpolationType(TEXT("r.HairStrand
 static int32 GHairStrandsTransferPositionOnLODChange = 0;
 static FAutoConsoleVariableRef CVarHairStrandsTransferPositionOnLODChange(TEXT("r.HairStrands.Strands.TransferPrevPos"), GHairStrandsTransferPositionOnLODChange, TEXT("Transfer strands prev. position to current position on LOD switching to avoid large discrepancy causing large motion vector"));
 
+static int32 GHairStrands_Raytracing_ForceRebuildBVH = 1;
+static FAutoConsoleVariableRef CVarHairRTGeomForceRebuild(TEXT("r.HairStrands.Strands.Raytracing.ForceRebuildBVH"), GHairStrands_Raytracing_ForceRebuildBVH, TEXT("Force BVH rebuild instead of doing a BVH refit when hair positions changed"));
+
 ////////////////////////////////////////////////////////////////////////////////////////////////////////////////////////
 
 namespace HairTransition
@@ -1224,8 +1227,9 @@ class FHairRaytracingGeometryCS : public FGlobalShader
 	SHADER_USE_PARAMETER_STRUCT(FHairRaytracingGeometryCS, FGlobalShader);
 
 	class FGroupSize : SHADER_PERMUTATION_SPARSE_INT("PERMUTATION_GROUP_SIZE", 32, 64);
-	class FCulling : SHADER_PERMUTATION_INT("PERMUTATION_CULLING", 2);
-	using FPermutationDomain = TShaderPermutationDomain<FGroupSize, FCulling>;
+	class FCulling : SHADER_PERMUTATION_BOOL("PERMUTATION_CULLING");
+	class FProceduralPrimitive : SHADER_PERMUTATION_BOOL("PERMUTATION_PROCEDURAL_PRIMITIVE"); 
+	using FPermutationDomain = TShaderPermutationDomain<FGroupSize, FCulling, FProceduralPrimitive>;
 
 	BEGIN_SHADER_PARAMETER_STRUCT(FParameters, )
 		SHADER_PARAMETER(uint32, VertexCount)
@@ -1244,11 +1248,9 @@ class FHairRaytracingGeometryCS : public FGlobalShader
 		SHADER_PARAMETER_RDG_BUFFER_SRV(Buffer, PositionBuffer)
 		SHADER_PARAMETER_RDG_BUFFER_SRV(Buffer, TangentBuffer)
 		SHADER_PARAMETER_RDG_BUFFER_UAV(RWBuffer, OutputPositionBuffer)
+		SHADER_PARAMETER_RDG_BUFFER_UAV(RWBuffer, OutputIndexBuffer)
 		SHADER_PARAMETER_STRUCT_INCLUDE(ShaderPrint::FShaderParameters, ShaderPrintUniformBuffer)
 		SHADER_PARAMETER_STRUCT_INCLUDE(ShaderDrawDebug::FShaderParameters, ShaderDrawParameters)
-	#if !STRANDS_CUSTOM_INTERSECTOR
-		SHADER_PARAMETER_RDG_BUFFER_UAV(RWBuffer, OutputIndexBuffer)
-	#endif
 	END_SHADER_PARAMETER_STRUCT()
 
 public:
@@ -1257,7 +1259,6 @@ public:
 	{
 		FGlobalShader::ModifyCompilationEnvironment(Parameters, OutEnvironment);
 		OutEnvironment.SetDefine(TEXT("SHADER_RT_GEOMETRY"), 1);
-		OutEnvironment.SetDefine(TEXT("ENABLE_CUSTOM_INTERSECTOR"), STRANDS_CUSTOM_INTERSECTOR);
 	}
 };
 
@@ -1269,6 +1270,7 @@ static void AddGenerateRaytracingGeometryPass(
 	const FShaderPrintData* ShaderPrintData,
 	const FShaderDrawDebugData* ShaderDrawData,
 	uint32 VertexCount,
+	bool bProceduralPrimitive,
 	float HairRadius,
 	float RootScale,
 	float TipScale,
@@ -1300,10 +1302,7 @@ static void AddGenerateRaytracingGeometryPass(
 	{
 		ShaderDrawDebug::SetParameters(GraphBuilder, *ShaderDrawData, Parameters->ShaderDrawParameters);
 	}
-
-#if !STRANDS_CUSTOM_INTERSECTOR
 	Parameters->OutputIndexBuffer = OutIndexBuffer;
-#endif
 
 	const bool bCullingEnable = CullingData.bCullingResultAvailable;
 	Parameters->HairStrandsVF_bIsCullingEnable = bCullingEnable ? 1 : 0;
@@ -1317,12 +1316,12 @@ static void AddGenerateRaytracingGeometryPass(
 
 	FHairRaytracingGeometryCS::FPermutationDomain PermutationVector;
 	PermutationVector.Set<FHairRaytracingGeometryCS::FGroupSize>(GroupSize);
-	PermutationVector.Set<FHairRaytracingGeometryCS::FCulling>(bCullingEnable ? 1 : 0);
-
+	PermutationVector.Set<FHairRaytracingGeometryCS::FCulling>(bCullingEnable);
+	PermutationVector.Set<FHairRaytracingGeometryCS::FProceduralPrimitive>(bProceduralPrimitive);
 	TShaderMapRef<FHairRaytracingGeometryCS> ComputeShader(ShaderMap, PermutationVector);
 	FComputeShaderUtils::AddPass(
 		GraphBuilder,
-		RDG_EVENT_NAME("HairStrands::GenerateRaytracingGeometry"),
+		RDG_EVENT_NAME("HairStrands::GenerateRaytracingGeometry(bCulling:%d, Procedural:%d)", bCullingEnable, bProceduralPrimitive),
 		ComputeShader,
 		Parameters,
 		DispatchCount);
@@ -1410,7 +1409,8 @@ static void BuildHairAccelerationStructure_Strands(
 	FBufferRHIRef& IndexBuffer,
 	FRayTracingGeometry* OutRayTracingGeometry,
 	const FString& AssetDebugName,
-	uint32 LODIndex)
+	uint32 LODIndex,
+	bool bProceduralPrimitive)
 {
 	FRayTracingGeometryInitializer Initializer;
 #if HAIR_RESOURCE_DEBUG_NAME
@@ -1421,33 +1421,37 @@ static void BuildHairAccelerationStructure_Strands(
 	static int32 DebugNumber = 0;
 	Initializer.DebugName = FName(DebugName, DebugNumber++);
 #endif
-#if STRANDS_CUSTOM_INTERSECTOR
-	Initializer.GeometryType = RTGT_Procedural;
-	Initializer.TotalPrimitiveCount = RaytracingVertexCount / 2;
-#else
-	Initializer.IndexBuffer = IndexBuffer;
-	Initializer.IndexBufferOffset = 0;
-	Initializer.GeometryType = RTGT_Triangles;
-	Initializer.TotalPrimitiveCount = RayTracingIndexCount / 3;
-#endif
+
+	if (bProceduralPrimitive)
+	{
+		Initializer.GeometryType = RTGT_Procedural;
+		Initializer.TotalPrimitiveCount = RaytracingVertexCount / 2;
+	}
+	else
+	{
+		Initializer.IndexBuffer = IndexBuffer;
+		Initializer.IndexBufferOffset = 0;
+		Initializer.GeometryType = RTGT_Triangles;
+		Initializer.TotalPrimitiveCount = RayTracingIndexCount / 3;
+	}
 	Initializer.bFastBuild = true;
 	Initializer.bAllowUpdate = false; // we won't be using update, so specify this to allow degenerate primitives to be discarded
 
 	FRayTracingGeometrySegment Segment;
 	Segment.VertexBuffer = PositionBuffer;
-#if STRANDS_CUSTOM_INTERSECTOR
-	// stride covers one AABB (which is made up of two vertices)
-	Segment.VertexBufferStride = FHairStrandsRaytracingFormat::SizeInByte * 2;
-#else
-	Segment.VertexBufferStride = FHairStrandsRaytracingFormat::SizeInByte;
-#endif
 	Segment.VertexBufferElementType = FHairStrandsRaytracingFormat::VertexElementType;
 	Segment.MaxVertices = RaytracingVertexCount / 2;
-#if STRANDS_CUSTOM_INTERSECTOR
-	Segment.NumPrimitives = RaytracingVertexCount / 2;
-#else
-	Segment.NumPrimitives = RayTracingIndexCount / 3;
-#endif
+	if (bProceduralPrimitive)
+	{
+		// stride covers one AABB (which is made up of two vertices)
+		Segment.VertexBufferStride = FHairStrandsRaytracingFormat::SizeInByte * 2;
+		Segment.NumPrimitives = RaytracingVertexCount / 2;
+	}
+	else
+	{
+		Segment.VertexBufferStride = FHairStrandsRaytracingFormat::SizeInByte;
+		Segment.NumPrimitives = RayTracingIndexCount / 3;
+	}
 	Initializer.Segments.Add(Segment);
 
 	OutRayTracingGeometry->SetInitializer(Initializer);
@@ -1678,6 +1682,7 @@ static FHairGroupPublicData::FVertexFactoryInput InternalComputeHairStrandsVerte
 	OutVFInput.Strands.HairRadius = (GStrandHairWidth > 0 ? GStrandHairWidth : Instance->Strands.Modifier.HairWidth) * 0.5f;
 	OutVFInput.Strands.HairRootScale = Instance->Strands.Modifier.HairRootScale;
 	OutVFInput.Strands.HairTipScale = Instance->Strands.Modifier.HairTipScale;
+	OutVFInput.Strands.HairRaytracingRadiusScale = (GHairRaytracingRadiusScale > 0 ? GHairRaytracingRadiusScale : Instance->Strands.Modifier.HairRaytracingRadiusScale);
 	OutVFInput.Strands.HairLength = Instance->Strands.Modifier.HairLength;
 	OutVFInput.Strands.HairDensity = Instance->Strands.Modifier.HairShadowDensity;
 	OutVFInput.Strands.bScatterSceneLighting = Instance->Strands.Modifier.bScatterSceneLighting;
@@ -1959,8 +1964,7 @@ void ComputeHairStrandsInterpolation(
 			#if RHI_RAYTRACING
 			if (Instance->Strands.RenRaytracingResource)
 			{
-				const float HairRadiusScaleRT = (GHairRaytracingRadiusScale > 0 ? GHairRaytracingRadiusScale : Instance->Strands.Modifier.HairRaytracingRadiusScale);
-				const float HairRadiusRT	= Instance->HairGroupPublicData->VFInput.Strands.HairRadius * HairRadiusScaleRT;
+				const float HairRadiusRT	= Instance->HairGroupPublicData->VFInput.Strands.HairRaytracingRadiusScale * Instance->HairGroupPublicData->VFInput.Strands.HairRadius;
 				const float HairRootScaleRT = Instance->HairGroupPublicData->VFInput.Strands.HairRootScale;
 				const float HairTipScaleRT	= Instance->HairGroupPublicData->VFInput.Strands.HairTipScale;
 
@@ -1969,6 +1973,7 @@ void ComputeHairStrandsInterpolation(
 				if (bNeedBuild || bNeedUpdate)
 				{
 					const uint32 HairLODIndex = Instance->HairGroupPublicData->GetIntLODIndex();
+					const bool bProceduralPrimitive = Instance->Strands.RenRaytracingResource->bProceduralPrimitive;
 
 					FRDGImportedBuffer Raytracing_PositionBuffer = Register(GraphBuilder, Instance->Strands.RenRaytracingResource->PositionBuffer, ERDGImportedBufferFlags::CreateViews);
 					FRDGImportedBuffer Raytracing_IndexBuffer = Register(GraphBuilder, Instance->Strands.RenRaytracingResource->IndexBuffer, ERDGImportedBufferFlags::CreateViews);
@@ -1978,6 +1983,7 @@ void ComputeHairStrandsInterpolation(
 						ShaderPrintData,
 						ShaderDrawData,
 						Instance->Strands.RestResource->GetVertexCount(),
+						bProceduralPrimitive,
 						HairRadiusRT,
 						HairRootScaleRT,
 						HairTipScaleRT,
@@ -1995,7 +2001,7 @@ void ComputeHairStrandsInterpolation(
 					GraphBuilder.AddPass(
 						RDG_EVENT_NAME("HairStrands::UpdateBLAS(Strands)"),
 						ERDGPassFlags::NeverCull,
-					[Instance, HairLODIndex, bNeedUpdate](FRHICommandListImmediate& RHICmdList)
+					[Instance, HairLODIndex, bNeedUpdate, bProceduralPrimitive](FRHICommandListImmediate& RHICmdList)
 					{
 						SCOPED_GPU_MASK(RHICmdList, FRHIGPUMask::All());
 
@@ -2003,12 +2009,10 @@ void ComputeHairStrandsInterpolation(
 						if (bLocalNeedBuild)
 						{
 							FBufferRHIRef PositionBuffer(Instance->Strands.RenRaytracingResource->PositionBuffer.Buffer->GetRHI());
-						#if STRANDS_CUSTOM_INTERSECTOR
-							// no index buffer needed
-							FBufferRHIRef IndexBuffer;
-						#else
-							FBufferRHIRef IndexBuffer(Instance->Strands.RenRaytracingResource->IndexBuffer.Buffer->GetRHI());
-						#endif
+
+							// no index buffer needed for procedural primitive
+							FBufferRHIRef IndexBuffer(bProceduralPrimitive ? nullptr : Instance->Strands.RenRaytracingResource->IndexBuffer.Buffer->GetRHI());
+
 							BuildHairAccelerationStructure_Strands(RHICmdList,
 								Instance->Strands.RenRaytracingResource->VertexCount,
 								Instance->Strands.RenRaytracingResource->IndexCount,
@@ -2016,14 +2020,15 @@ void ComputeHairStrandsInterpolation(
 								IndexBuffer,
 								&Instance->Strands.RenRaytracingResource->RayTracingGeometry,
 								Instance->Debug.GroomAssetName,
-								HairLODIndex
+								HairLODIndex,
+								bProceduralPrimitive
 							);
 						}
 						else if (bNeedUpdate)
 						{
 							// hair strands can move chaotically during simulation which can really tank ray tracing performance
 							// even though a rebuild is more expensive, we can still come out ahead overall, even with a single ray cast per pixel
-							UpdateHairAccelerationStructure(RHICmdList, &Instance->Strands.RenRaytracingResource->RayTracingGeometry, EAccelerationStructureBuildMode::Build);
+							UpdateHairAccelerationStructure(RHICmdList, &Instance->Strands.RenRaytracingResource->RayTracingGeometry, GHairStrands_Raytracing_ForceRebuildBVH > 0? EAccelerationStructureBuildMode::Build : EAccelerationStructureBuildMode::Update);
 						}
 						Instance->Strands.RenRaytracingResource->bIsRTGeometryInitialized = true;
 					});
diff --git a/Engine/Plugins/Runtime/HairStrands/Source/HairStrandsCore/Private/HairStrandsVertexFactory.cpp b/Engine/Plugins/Runtime/HairStrands/Source/HairStrandsCore/Private/HairStrandsVertexFactory.cpp
--- a/Engine/Plugins/Runtime/HairStrands/Source/HairStrandsCore/Private/HairStrandsVertexFactory.cpp
+++ b/Engine/Plugins/Runtime/HairStrands/Source/HairStrandsCore/Private/HairStrandsVertexFactory.cpp
@@ -14,6 +14,12 @@
 #include "GroomInstance.h"
 
 #define VF_STRANDS_SUPPORT_GPU_SCENE 0
+#define VF_STRANDS_CUSTOM_INTERSECTOR 1
+
+bool GetSupportHairStrandsProceduralPrimitive(EShaderPlatform InShaderPlatform)
+{
+	return VF_STRANDS_CUSTOM_INTERSECTOR && FDataDrivenShaderPlatformInfo::GetSupportsRayTracingProceduralPrimitive(InShaderPlatform);
+}
 
 /////////////////////////////////////////////////////////////////////////////////////////
 
@@ -65,6 +71,7 @@ FHairStrandsVertexFactoryUniformShaderParameters FHairGroupInstance::GetHairStan
 	Out.Radius 						= VFInput.Strands.HairRadius;
 	Out.RootScale 					= VFInput.Strands.HairRootScale;
 	Out.TipScale 					= VFInput.Strands.HairTipScale;
+	Out.RaytracingRadiusScale		= VFInput.Strands.HairRaytracingRadiusScale;
 	Out.Length 						= VFInput.Strands.HairLength;
 	Out.Density 					= VFInput.Strands.HairDensity;
 	Out.StableRasterization			= VFInput.Strands.bUseStableRasterization;
@@ -152,9 +159,10 @@ void FHairStrandsVertexFactory::ModifyCompilationEnvironment(const FVertexFactor
 		// TODO: support GPUScene on mobile
 		&& !IsMobilePlatform(Parameters.Platform);
 #endif
+	const bool bUseProceduralIntersection = GetSupportHairStrandsProceduralPrimitive(Parameters.Platform);
 	OutEnvironment.SetDefine(TEXT("VF_SUPPORTS_PRIMITIVE_SCENE_DATA"), bUseGPUSceneAndPrimitiveIdStream);
-	OutEnvironment.SetDefine(TEXT("HAIR_STRAND_MESH_FACTORY"), TEXT("1"));
-	OutEnvironment.SetDefine(TEXT("ENABLE_CUSTOM_INTERSECTOR"), STRANDS_CUSTOM_INTERSECTOR);
+	OutEnvironment.SetDefine(TEXT("HAIR_STRAND_MESH_FACTORY"), 1);
+	OutEnvironment.SetDefine(TEXT("ENABLE_CUSTOM_INTERSECTOR"), bUseProceduralIntersection);
 }
 
 void FHairStrandsVertexFactory::ValidateCompiledResult(const FVertexFactoryType* Type, EShaderPlatform Platform, const FShaderParameterMap& ParameterMap, TArray<FString>& OutErrors)
@@ -249,4 +257,7 @@ IMPLEMENT_VERTEX_FACTORY_TYPE(FHairStrandsVertexFactory, "/Engine/Private/HairSt
 	| EVertexFactoryFlags::SupportsPrecisePrevWorldPos
 	| EVertexFactoryFlags::SupportsCachingMeshDrawCommands
 	| (VF_STRANDS_SUPPORT_GPU_SCENE ? EVertexFactoryFlags::SupportsPrimitiveIdStream : EVertexFactoryFlags::None)
+	| EVertexFactoryFlags::SupportsRayTracing
+	| EVertexFactoryFlags::SupportsRayTracingDynamicGeometry
+	| (VF_STRANDS_CUSTOM_INTERSECTOR ? EVertexFactoryFlags::SupportsRayTracingProceduralPrimitive : EVertexFactoryFlags::None)
 );
diff --git a/Engine/Plugins/Runtime/HairStrands/Source/HairStrandsCore/Public/GroomInstance.h b/Engine/Plugins/Runtime/HairStrands/Source/HairStrandsCore/Public/GroomInstance.h
--- a/Engine/Plugins/Runtime/HairStrands/Source/HairStrandsCore/Public/GroomInstance.h
+++ b/Engine/Plugins/Runtime/HairStrands/Source/HairStrandsCore/Public/GroomInstance.h
@@ -47,6 +47,7 @@ BEGIN_GLOBAL_SHADER_PARAMETER_STRUCT(FHairStrandsVertexFactoryUniformShaderParam
 	SHADER_PARAMETER(float, TipScale)
 	SHADER_PARAMETER(float, Length)
 	SHADER_PARAMETER(float, Density)
+	SHADER_PARAMETER(float, RaytracingRadiusScale)
 	SHADER_PARAMETER(uint32, CullingEnable)
 	SHADER_PARAMETER(uint32, HasMaterial)
 	SHADER_PARAMETER(uint32, StableRasterization)
diff --git a/Engine/Plugins/Runtime/HairStrands/Source/HairStrandsCore/Public/GroomResources.h b/Engine/Plugins/Runtime/HairStrands/Source/HairStrandsCore/Public/GroomResources.h
--- a/Engine/Plugins/Runtime/HairStrands/Source/HairStrandsCore/Public/GroomResources.h
+++ b/Engine/Plugins/Runtime/HairStrands/Source/HairStrandsCore/Public/GroomResources.h
@@ -11,8 +11,6 @@
 #include "HairStrandsInterface.h"
 #include "HairStrandsMeshProjection.h"
 
-#define STRANDS_CUSTOM_INTERSECTOR 0
-
 inline uint32 GetBufferTotalNumBytes(const FRDGExternalBuffer& In) 
 {
 	return In.Buffer ? In.Buffer->Desc.GetTotalNumBytes() : 0;
@@ -471,6 +469,7 @@ struct FHairStrandsRaytracingResource : public FHairCommonResource
 	uint32 IndexCount = 0;
 	bool bOwnBuffers = false;
 	bool bIsRTGeometryInitialized = false;
+	bool bProceduralPrimitive = false;
 };
 #endif
 
diff --git a/Engine/Shaders/Private/HairStrands/HairStrandsRaytracingGeometry.usf b/Engine/Shaders/Private/HairStrands/HairStrandsRaytracingGeometry.usf
--- a/Engine/Shaders/Private/HairStrands/HairStrandsRaytracingGeometry.usf
+++ b/Engine/Shaders/Private/HairStrands/HairStrandsRaytracingGeometry.usf
@@ -19,7 +19,7 @@ Buffer<float4>	 PositionOffsetBuffer;
 Buffer<uint4> 	 PositionBuffer;
 Buffer<float4> 	 TangentBuffer;
 RWBuffer<float4> OutputPositionBuffer;
-#if !ENABLE_CUSTOM_INTERSECTOR 
+#if !PERMUTATION_PROCEDURAL_PRIMITIVE 
 RWBuffer<uint>   OutputIndexBuffer;
 #endif
 
@@ -64,7 +64,7 @@ float3x3 GetTangentBasis(uint ControlPointId)
 	return float3x3(Tx, Ty, Tz);
 }
 
-#if ENABLE_CUSTOM_INTERSECTOR
+#if PERMUTATION_PROCEDURAL_PRIMITIVE
 void SetAABB(uint InVertexIndex, float3 Lo, float3 Hi)
 {
 	OutputPositionBuffer[2 * InVertexIndex + 0] = float4(Lo.x, Lo.y, Lo.z, Hi.x);
@@ -75,7 +75,7 @@ void SetAABB(uint InVertexIndex, float3 Lo, float3 Hi)
 // a whole segment of invalid values
 void SetInvalidSegment(uint InVertexIndex)
 {
-#if ENABLE_CUSTOM_INTERSECTOR
+#if PERMUTATION_PROCEDURAL_PRIMITIVE
 	const float NaN = asfloat(0x7FC00000);
 	SetAABB(InVertexIndex, NaN, NaN);
 #else
@@ -118,17 +118,9 @@ void SetCapSegment(uint InVertexIndex)
 #endif
 }
 
+#if !PERMUTATION_PROCEDURAL_PRIMITIVE
 void SetBodySegment(uint InVertexIndex)
 {
-#if ENABLE_CUSTOM_INTERSECTOR
-	FHairControlPoint P0 = GetControlPoint(InVertexIndex);
-	FHairControlPoint P1 = GetControlPoint(InVertexIndex + 1);
-	float4 A = float4(P0.Position, P0.WorldRadius);
-	float4 B = float4(P1.Position, P1.WorldRadius);
-	float3 Lo = min(A.xyz - A.w, B.xyz - B.w);
-	float3 Hi = max(A.xyz + A.w, B.xyz + B.w);
-	SetAABB(InVertexIndex, Lo, Hi);
-#else
 	const uint BotPoint = InVertexIndex * 4;
 	const uint TopPoint = BotPoint + 4;
 
@@ -150,23 +142,25 @@ void SetBodySegment(uint InVertexIndex)
 		OutputIndexBuffer[OutIndex + 4] = BotPoint + Idx;
 		OutputIndexBuffer[OutIndex + 5] = BotPoint + PrevIdx; 
 	}
-#endif
 }
+#endif
 
 
 [numthreads(GROUP_SIZE, 1, 1)]
 void MainCS(uint2 DispatchThreadId : SV_DispatchThreadID)
 {
 	uint VertexIndex0 = GetHairStrandsVertexFetchIndex(DispatchThreadId, GROUP_SIZE, DispatchCountX);
+	uint VertexIndex1 = VertexIndex0 + 1;
 	float LodRadiusScale0 = 1;
 	float LodRadiusScale1 = 1;
 
 	uint SourceIndex0 = VertexIndex0;
+	uint SourceIndex1 = VertexIndex1;
 #if PERMUTATION_CULLING == 1
 	if (HairStrandsVF_bIsCullingEnable)
 	{
 		const uint VertexCountAfterCulling = HairStrandsVF_CullingIndirectBuffer[3];
-		if (VertexIndex0 >= VertexCountAfterCulling)
+		if (VertexIndex1 >= VertexCountAfterCulling)
 		{
 			SetInvalidSegment(VertexIndex0);
 			return;
@@ -174,6 +168,9 @@ void MainCS(uint2 DispatchThreadId : SV_DispatchThreadID)
 
 		SourceIndex0 = HairStrandsVF_CullingIndexBuffer[VertexIndex0];
 		LodRadiusScale0 = HairStrandsVF_CullingRadiusScaleBuffer[VertexIndex0];
+		SourceIndex1 = HairStrandsVF_CullingIndexBuffer[VertexIndex1];
+		LodRadiusScale1 = HairStrandsVF_CullingRadiusScaleBuffer[VertexIndex1];
+
 
 		if (LodRadiusScale0 <= 0.0f)
 		{
@@ -184,8 +181,22 @@ void MainCS(uint2 DispatchThreadId : SV_DispatchThreadID)
 #endif
 
 	FHairControlPoint P0 = GetControlPoint(SourceIndex0);
-#if ENABLE_CUSTOM_INTERSECTOR
-	// aabb position is generated by segment call
+#if PERMUTATION_PROCEDURAL_PRIMITIVE 
+	if (P0.Type == HAIR_CONTROLPOINT_END)
+	{
+		// this vertex doesn't map to a valid segment
+		SetInvalidSegment(VertexIndex0);
+	}
+	else
+	{
+		// fetch the other point and create the aabb around the line
+		FHairControlPoint P1 = GetControlPoint(SourceIndex1);
+		float4 A = float4(P0.Position, P0.WorldRadius * LodRadiusScale0);
+		float4 B = float4(P1.Position, P1.WorldRadius * LodRadiusScale1);
+		float3 Lo = min(A.xyz - A.w, B.xyz - B.w);
+		float3 Hi = max(A.xyz + A.w, B.xyz + B.w);
+		SetAABB(VertexIndex0, Lo, Hi);
+	}
 #else
 	// expand current vertex into 4 points around the curve
 	const float3x3 TangentBasis = GetTangentBasis(SourceIndex0);
@@ -204,7 +215,6 @@ void MainCS(uint2 DispatchThreadId : SV_DispatchThreadID)
 		OutputPositionBuffer[BaseOutIndex + 2] = float4(Px1, 1);
 		OutputPositionBuffer[BaseOutIndex + 3] = float4(Py1, 1);
 	}
-#endif
 
 	if (P0.Type == HAIR_CONTROLPOINT_END)
 	{
@@ -217,6 +227,7 @@ void MainCS(uint2 DispatchThreadId : SV_DispatchThreadID)
 		// regular point, build the body of the cylinder
 		SetBodySegment(VertexIndex0);
 	}
+#endif
 }
 
 #endif // SHADER_RT_GEOMETRY
diff --git a/Engine/Shaders/Private/HairStrands/HairStrandsVertexFactory.ush b/Engine/Shaders/Private/HairStrands/HairStrandsVertexFactory.ush
--- a/Engine/Shaders/Private/HairStrands/HairStrandsVertexFactory.ush
+++ b/Engine/Shaders/Private/HairStrands/HairStrandsVertexFactory.ush
@@ -374,6 +374,276 @@ float3 GetTangent(FVertexFactoryIntermediates Intermediates)
 	return Intermediates.TangentToWorld[2].xyz;
 }
 
+////////////////////////////////////////////////////////////////////////////////////////////////////////////
+// Raytracing custom intersection shader
+
+#if RAYHITGROUPSHADER
+
+#if USE_MATERIAL_INTERSECTION_SHADER
+
+#define STRAND_INTERSECTOR   1 // 0: cylinders 1: ribbons 
+
+float3 isectStrand(float3 ro, float3 rd, float3  pa, float3  pb, float ra, float rb)
+{
+#if STRAND_INTERSECTOR == 0
+	// trick to improve numerical stability - move origin closer to one of the vertices
+	float dd = dot(rd, rd);
+	float idd = rcp(dd);
+	float offset = dot(rd, pa - ro) * idd;
+	ro += offset * rd;
+
+	float3  ba = pb - pa;
+	float3  oa = ro - pa;
+	float3  ob = ro - pb;
+	float rr = ra - rb;
+	float m0 = dot(ba, ba);
+	float m1 = dot(ba, oa);
+	float m2 = dot(ba, rd);
+	float m3 = dot(rd, oa);
+	float m5 = dot(oa, oa);
+	float m6 = dot(ob, rd);
+	float m7 = dot(ob, ob);
+
+	float d2 = m0 - rr * rr;
+
+	float k2 = dd * d2 - m2 * m2;
+	float k1 = d2 * m3 - m1 * m2 + m2 * rr * ra;
+	float k0 = d2 * m5 - m1 * m1 + m1 * rr * ra * 2.0 - m0 * ra * ra;
+
+	float h = k1 * k1 - k0 * k2;
+	if (h > 0.0)
+	{
+		float t = (-sqrt(h) - k1) / k2;
+
+		float y = m1 - ra * rr + t * m2;
+		if (y > 0.0 && y < d2)
+		{
+			return float3(t + offset, y / d2, 0.5);
+		}
+	}
+
+  #if 0
+	// ignore caps: produces cracks between segments
+	return -1.0;
+  #elif 1
+	// end-point cap (only test one since hair strands share vertices (root is less important than the tip)
+	// this keeps a visually continuous curve
+	float h2 = m6 * m6 - dd * (m7 - rb * rb);
+	if (h2 > 0.0)
+	{
+		float t = (-m6 - sqrt(h2)) * idd + offset;
+		return float3(t, 1, 0.5);
+	}
+	return -1.0;
+  #else
+	// caps
+	float h1 = m3 * m3 - dd * (m5 - ra * ra);
+	float h2 = m6 * m6 - dd * (m7 - rb * rb);
+	if (max(h1, h2) < 0.0) return -1.0;
+
+	float3 r = 1e20;
+	if (h1 > 0.0)
+	{
+		float t = (-m3 - sqrt(h1)) * idd + offset;
+		r = float3(t + offset, 0, 0.5);
+	}
+	if (h2 > 0.0)
+	{
+		float t = (-m6 - sqrt(h2)) * idd + offset;
+		if (t > 0.0 && t < r.x)
+			r = float3(t, 1, 0.5);
+	}
+
+	return r;
+  #endif
+#elif STRAND_INTERSECTOR == 1
+	float3 u = pb - pa;
+	float3 w0 = pa - ro;
+	float3 w1 = pb - ro;
+	float3 n = cross(rd, u);
+	float3 n1 = cross(rd, n);
+	float sc = saturate(-dot(w0, n1) / dot(u, n1));
+	float3 cp = lerp(w0, w1, sc);
+	float rc = lerp(ra, rb, sc);
+	float t = dot(cp, rd) / dot(rd, rd);
+	float3 dp = t * rd - cp;
+	float d2 = length2(dp);
+	float r2 = rc * rc;
+	if (d2 < r2 && t > 4. * rc)
+	{
+		float h = sign(dot(dp, n)) * sqrt(d2 / r2);
+		return float3(t, sc, saturate(0.5 * h + 0.5));
+	}
+	return -1.0;
+#else
+	return -1.0;
+#endif
+}
+
+// Entry point
+RAY_TRACING_ENTRY_INTERSECTION(MaterialIS)
+{
+	float3 RayOrg = ObjectRayOrigin();
+	float3 RayDir = ObjectRayDirection();
+	float TMin = RayTMin();
+	float TMax = RayTCurrent();
+
+	uint VertexIndex0 = PrimitiveIndex();
+	uint VertexIndex1 = VertexIndex0 + 1;
+	float ScaleR0 = 1, ScaleR1 = 1;
+	if (HairStrandsVF.CullingEnable)
+	{
+		ScaleR0 = HairStrandsVF.CulledVertexRadiusScaleBuffer[VertexIndex0];
+		ScaleR1 = HairStrandsVF.CulledVertexRadiusScaleBuffer[VertexIndex1];
+		VertexIndex0 = HairStrandsVF.CulledVertexIdsBuffer[VertexIndex0];
+		VertexIndex1 = HairStrandsVF.CulledVertexIdsBuffer[VertexIndex1];
+	}
+
+	const float3 PositionOffset = HairStrandsVF.PositionOffsetBuffer[0].xyz;
+	FHairControlPoint P0 = UnpackHairControlPoint(
+		HairStrandsVF.PositionBuffer[VertexIndex0],
+		PositionOffset,
+		HairStrandsVF.Radius * ScaleR0 * HairStrandsVF.RaytracingRadiusScale,
+		HairStrandsVF.RootScale,
+		HairStrandsVF.TipScale);
+	FHairControlPoint P1 = UnpackHairControlPoint(
+		HairStrandsVF.PositionBuffer[VertexIndex1],
+		PositionOffset,
+		HairStrandsVF.Radius * ScaleR1 * HairStrandsVF.RaytracingRadiusScale,
+		HairStrandsVF.RootScale,
+		HairStrandsVF.TipScale);
+
+	float3 Result = isectStrand(RayOrg, RayDir, P0.Position, P1.Position, P0.WorldRadius, P1.WorldRadius);
+
+	FDefaultAttributes Attributes = (FDefaultAttributes)0;
+	Attributes.Barycentrics = float2(Result.y, Result.z);
+	ReportHit(Result.x, HIT_KIND_TRIANGLE_FRONT_FACE, Attributes);
+}
+#endif // USE_MATERIAL_INTERSECTION_SHADER
+
+#endif // RAYHITGROUPSHADER
+
+////////////////////////////////////////////////////////////////////////////////////////////////////////////
+// Material evaluation for raytracing
+#if RAYHITGROUPSHADER
+
+// define this to let shaders know this alternative method is available
+#define VF_SUPPORTS_RAYTRACING_PREPARE_MATERIAL_PIXEL_PARAMETERS 1 
+
+FMaterialPixelParameters GetMaterialPixelParameters(float3 RayOrigin, float3 RayDirection, float HitT, uint HitPrimitiveIndex, FDefaultAttributes HitAttributes, uint HitKind, out float3 WorldGeoNormal)
+{
+	FMaterialPixelParameters Result = MakeInitializedMaterialPixelParameters();
+
+#if VF_USE_PRIMITIVE_SCENE_DATA
+	FVertexFactoryInput Input;
+	// Note: GetInstanceUserData() stores the GPU-Scene primitive ID
+	Input.InstanceIdOffset = GetInstanceUserData() | VF_TREAT_INSTANCE_ID_OFFSET_AS_PRIMITIVE_ID_FLAG;
+	Input.DrawInstanceId = 0;
+#endif
+
+	FSceneDataIntermediates SceneData = VF_GPUSCENE_GET_INTERMEDIATES(Input); // NOTE: Input is not used when VF_USE_PRIMITIVE_SCENE_DATA == 0
+
+	uint VertexIndex0 = HitPrimitiveIndex;
+	uint VertexIndex1 = VertexIndex0 + 1;
+
+	if (HairStrandsVF.CullingEnable)
+	{
+		VertexIndex0 = HairStrandsVF.CulledVertexIdsBuffer[VertexIndex0];
+		VertexIndex1 = HairStrandsVF.CulledVertexIdsBuffer[VertexIndex1];
+	}
+
+#if ENABLE_CUSTOM_INTERSECTOR
+	float UCoord = HitAttributes.Barycentrics.x;
+	float VCoord = HitAttributes.Barycentrics.y;
+	const float3 PositionOffset = HairStrandsVF.PositionOffsetBuffer[0].xyz;
+	FHairControlPoint P0 = UnpackHairControlPoint(
+		HairStrandsVF.PositionBuffer[VertexIndex0],
+		PositionOffset,
+		HairStrandsVF.Radius,
+		HairStrandsVF.RootScale,
+		HairStrandsVF.TipScale);
+	FHairControlPoint P1 = UnpackHairControlPoint(
+		HairStrandsVF.PositionBuffer[VertexIndex1],
+		PositionOffset,
+		HairStrandsVF.Radius,
+		HairStrandsVF.RootScale,
+		HairStrandsVF.TipScale);
+
+  #if STRAND_INTERSECTOR == 0
+	float3 WorldPosition = RayOrigin + HitT * RayDirection;
+	float4 LocalPosition = mul(float4(WorldPosition, 1), LWCHackToFloat(SceneData.Primitive.WorldToLocal));
+	float3 LocalNormal = normalize(LocalPosition.xyz - lerp(P0.Position, P1.Position, UCoord));
+	// TODO: does this support non-uniform scales correctly?
+	WorldGeoNormal = normalize(RotateLocalToWorld(LocalNormal, SceneData.Primitive.LocalToWorld, SceneData.Primitive.InvNonUniformScale));
+  #endif
+
+	Result.HairPrimitiveId = VertexIndex0;
+#else // ENABLE_CUSTOM_INTERSECTOR
+	const float3 Weights = float3(
+		1 - HitAttributes.Barycentrics.x - HitAttributes.Barycentrics.y,
+		HitAttributes.Barycentrics.x,
+		HitAttributes.Barycentrics.y);
+
+	// fetch triangle
+	FTriangleBaseAttributes Triangle = LoadTriangleBaseAttributes(HitPrimitiveIndex);
+	// Transform vertices into world space
+	const float3 WorldPositions[3] = {
+		CalcWorldPosition(Triangle.LocalPositions[0], SceneData.Primitive.LocalToWorld).xyz,
+		CalcWorldPosition(Triangle.LocalPositions[1], SceneData.Primitive.LocalToWorld).xyz,
+		CalcWorldPosition(Triangle.LocalPositions[2], SceneData.Primitive.LocalToWorld).xyz
+	};
+	// Compute triangle normal
+	{
+		const float3 PA = WorldPositions[1] - WorldPositions[0];
+		const float3 PB = WorldPositions[2] - WorldPositions[0];
+		const float3 Unnormalized = cross(PB, PA);
+		const float InvWorldArea = rsqrt(dot(Unnormalized, Unnormalized));
+		WorldGeoNormal = Unnormalized * InvWorldArea;
+	}
+
+	// NOTE: This must match the tesselation pattern produced by SetBodySegment() in HairStrandsRaytracingGeometry.usf
+	// TODO: is there a way we could handle end-caps elegantly?
+	Result.HairPrimitiveId = Triangle.Indices.z / 4; // last index always belongs to the "bottom" vertex of the hair segment
+
+	float UCoord = dot(float3(1, float((HitPrimitiveIndex % 2) == 0), 0), Weights);
+	float VCoord = 0.5;// can't measure position across the width easily
+#endif // ENABLE_CUSTOM_INTERSECTOR
+
+	Result.HairPrimitiveUV = float2(UCoord, 0.5); // NOTE: don't bother defining a v coordinate since we use a cylinder for RT
+
+	// Required for previewing materials that use ParticleColor
+	Result.Particle.Color = half4(1, 1, 1, 1);
+
+	// Get tangent frame from each vertex
+	float TangentSign_V0 = 0; half3x3 TangentToWorld_V0 = CalcTangentToWorld(SceneData.Primitive, CalcTangentToLocal(VertexIndex0, TangentSign_V0));
+	float TangentSign_V1 = 0; half3x3 TangentToWorld_V1 = CalcTangentToWorld(SceneData.Primitive, CalcTangentToLocal(VertexIndex1, TangentSign_V1));
+	// interpolate to current point within the segment
+	float3 TangentToWorld0 = normalize(lerp(TangentToWorld_V0[0], TangentToWorld_V1[0], UCoord));
+	float3 TangentToWorld2 = normalize(lerp(TangentToWorld_V0[2], TangentToWorld_V1[2], UCoord));
+	float TangentSign = lerp(TangentSign_V0, TangentSign_V1, UCoord) * GetPrimitive_DeterminantSign_FromFlags(SceneData.Primitive.Flags);
+
+	Result.TangentToWorld = AssembleTangentToWorld(TangentToWorld0, float4(TangentToWorld2, TangentSign));
+
+#if STRAND_INTERSECTOR == 1
+	// compute a smooth ribbon normal based on the tangent vector
+	float3 T = Result.TangentToWorld[2];
+	float3 D = RayDirection;
+	WorldGeoNormal = normalize(cross(T, normalize(cross(D, T))));
+	WorldGeoNormal *= dot(RayDirection, WorldGeoNormal) > 0 ? -1.0 : 1.0;
+#endif
+
+	Result.TwoSidedSign = 1;
+
+	Result.PrimitiveId = SceneData.PrimitiveId;
+
+	return Result;
+}
+
+#endif // RAYHITGROUPSHADER
+
+////////////////////////////////////////////////////////////////////////////////////////////////////////////
+// Material evaluation for regular shader
+
 FVertexFactoryIntermediates GetVertexFactoryIntermediates(FVertexFactoryInput Input)
 {
 	FVertexFactoryIntermediates Intermediates;
@@ -570,4 +840,100 @@ float3 VertexFactoryGetWorldNormal(FVertexFactoryInput Input, FVertexFactoryInte
 {
 	return float3(0,0,1);
 }
+
+////////////////////////////////////////////////////////////////////////////////////////////////////////////
+// Regular material evalution for raytracing
+
+#if RAYHITGROUP || COMPUTESHADER
+FVertexFactoryInput LoadVertexFactoryInputForDynamicUpdate(uint TriangleIndex, int VertexIndex, uint PrimitiveId)
+{
+	FVertexFactoryInput Input;
+
+	Input.VertexId = TriangleIndex * 3 + VertexIndex;
+
+#if FEATURE_LEVEL >= FEATURE_LEVEL_SM4 && ONEPASS_POINTLIGHT_SHADOW && USING_VERTEX_SHADER_LAYER
+	Input.InstanceInd = 0;
+#endif
+	return Input;
+}
+
+uint GetNumRayTracingDynamicMeshVerticesIndirect()
+{
+	return 0;
+}
+#endif // RAYHITGROUP || COMPUTESHADER
+
+#if NEEDS_VERTEX_FACTORY_INTERPOLATION
+
+struct FVertexFactoryRayTracingInterpolants
+{
+	FVertexFactoryInterpolantsVSToPS InterpolantsVSToPS;
+};
+
+float2 VertexFactoryGetRayTracingTextureCoordinate(FVertexFactoryRayTracingInterpolants Interpolants)
+{
+#if NUM_MATERIAL_TEXCOORDS
+	return Interpolants.InterpolantsVSToPS.TexCoords[0].xy;
+#else // #if NUM_MATERIAL_TEXCOORDS
+	return float2(0, 0);
+#endif // #if NUM_MATERIAL_TEXCOORDS
+}
+
+FVertexFactoryInterpolantsVSToPS VertexFactoryAssignInterpolants(FVertexFactoryRayTracingInterpolants Input)
+{
+	return Input.InterpolantsVSToPS;
+}
+
+FVertexFactoryRayTracingInterpolants VertexFactoryGetRayTracingInterpolants(FVertexFactoryInput Input, FVertexFactoryIntermediates Intermediates, FMaterialVertexParameters VertexParameters)
+{
+	FVertexFactoryRayTracingInterpolants Interpolants;
+	Interpolants.InterpolantsVSToPS = VertexFactoryGetInterpolantsVSToPS(Input, Intermediates, VertexParameters);
+	return Interpolants;
+}
+
+FVertexFactoryRayTracingInterpolants VertexFactoryInterpolate(FVertexFactoryRayTracingInterpolants a, float aInterp, FVertexFactoryRayTracingInterpolants b, float bInterp)
+{
+	// Default initialize. Otherwise, some graphics pipelines that
+	// couple tessellation with geometry shaders won't write to all TEXCOORD semantics,
+	// but read from them when <FVertexFactoryRayTracingInterpolants> is being copied as a whole.
+	FVertexFactoryRayTracingInterpolants O = (FVertexFactoryRayTracingInterpolants)0;
+
+#if VF_USE_PRIMITIVE_SCENE_DATA
+	O.InterpolantsVSToPS.PrimitiveId = a.InterpolantsVSToPS.PrimitiveId;
+#if NEEDS_LIGHTMAP_COORDINATE
+	O.InterpolantsVSToPS.LightmapDataIndex = a.InterpolantsVSToPS.LightmapDataIndex;
+#endif
+#endif	
+	// Do we really need to interpolate TangentToWorld2 here? It should be replaced by the
+	// interpolated normal from 'whatever' interpolation scheme we're using
+
+	INTERPOLATE_MEMBER(InterpolantsVSToPS.TangentToWorld0.xyz);
+	INTERPOLATE_MEMBER(InterpolantsVSToPS.TangentToWorld2);
+#if INTERPOLATE_VERTEX_COLOR
+	INTERPOLATE_MEMBER(InterpolantsVSToPS.Color);
+#endif
+#if USE_INSTANCING
+	INTERPOLATE_MEMBER(InterpolantsVSToPS.PerInstanceParams);
+#endif
+
+#if NEEDS_LIGHTMAP_COORDINATE
+	INTERPOLATE_MEMBER(InterpolantsVSToPS.LightMapCoordinate);
+#endif
+
+#if NUM_TEX_COORD_INTERPOLATORS
+	UNROLL
+		for (int tc = 0; tc < (NUM_TEX_COORD_INTERPOLATORS + 1) / 2; ++tc)
+		{
+			INTERPOLATE_MEMBER(InterpolantsVSToPS.TexCoords[tc]);
+		}
+#elif USE_PARTICLE_SUBUVS
+	INTERPOLATE_MEMBER(InterpolantsVSToPS.TexCoords[0]);
+#endif
+
+	O.InterpolantsVSToPS.HairPrimitiveId = a.InterpolantsVSToPS.HairPrimitiveId;
+	INTERPOLATE_MEMBER(InterpolantsVSToPS.HairPrimitiveUV);
+	return O;
+}
+#endif // NEEDS_VERTEX_FACTORY_INTERPOLATION
+
 #include "/Engine/Private/VertexFactoryDefaultInterface.ush"
diff --git a/Engine/Source/Runtime/Renderer/Private/DeferredShadingRenderer.cpp b/Engine/Source/Runtime/Renderer/Private/DeferredShadingRenderer.cpp
--- a/Engine/Source/Runtime/Renderer/Private/DeferredShadingRenderer.cpp
+++ b/Engine/Source/Runtime/Renderer/Private/DeferredShadingRenderer.cpp
@@ -936,10 +936,9 @@ bool FDeferredShadingSceneRenderer::GatherRayTracingWorldInstancesForView(FRHICo
 						RayTracingScene.GeometriesToBuild.Add(Geometry);
 					}
 
-					// Thin geometries like hair don't have material, as they only support shadow at the moment.
+					// Validate the material/segment counts
 					if (!ensureMsgf(Instance.GetMaterials().Num() == Geometry->Initializer.Segments.Num() ||
-						(Geometry->Initializer.Segments.Num() == 0 && Instance.GetMaterials().Num() == 1) ||
-						(Instance.GetMaterials().Num() == 0 && (Instance.Mask & RAY_TRACING_MASK_THIN_SHADOW) > 0),
+						(Geometry->Initializer.Segments.Num() == 0 && Instance.GetMaterials().Num() == 1),
 						TEXT("Ray tracing material assignment validation failed for geometry '%s'. "
 							"Instance.GetMaterials().Num() = %d, Geometry->Initializer.Segments.Num() = %d, Instance.Mask = 0x%X."),
 						*Geometry->Initializer.DebugName.ToString(), Instance.GetMaterials().Num(),
@@ -2533,7 +2532,7 @@ void FDeferredShadingSceneRenderer::Render(FRDGBuilder& GraphBuilder)
 		ensureMsgf(!VirtualShadowMapArray.IsEnabled(), TEXT("Virtual shadow maps are not supported in the forward shading path"));
 		RenderShadowDepthMaps(GraphBuilder, InstanceCullingManager);
 
-		if (bHairEnable)
+		if (bHairEnable && !bHasRayTracedOverlay)
 		{
 			RenderHairPrePass(GraphBuilder, Scene, Views, InstanceCullingManager);
 			RenderHairBasePass(GraphBuilder, Scene, SceneTextures, Views, InstanceCullingManager);
@@ -2630,7 +2629,7 @@ void FDeferredShadingSceneRenderer::Render(FRDGBuilder& GraphBuilder)
 	}
 
 	// Render hair
-	if (bHairEnable && !IsForwardShadingEnabled(ShaderPlatform))
+	if (bHairEnable && !IsForwardShadingEnabled(ShaderPlatform) && !bHasRayTracedOverlay)
 	{
 		RenderHairPrePass(GraphBuilder, Scene, Views, InstanceCullingManager);
 		RenderHairBasePass(GraphBuilder, Scene, SceneTextures, Views, InstanceCullingManager);
diff --git a/Engine/Source/Runtime/Renderer/Public/HairStrandsInterface.h b/Engine/Source/Runtime/Renderer/Public/HairStrandsInterface.h
--- a/Engine/Source/Runtime/Renderer/Public/HairStrandsInterface.h
+++ b/Engine/Source/Runtime/Renderer/Public/HairStrandsInterface.h
@@ -262,6 +262,7 @@ public:
 			float HairRadius = 0;
 			float HairRootScale = 0;
 			float HairTipScale = 0;
+			float HairRaytracingRadiusScale = 0;
 			float HairLengthScale = 0;
 			float HairLength = 0;
 			float HairDensity = 0;
