Add voxel traversal shadow flags in order to allow voxelization groom using raytraced geometry.

This allows to enable raytracing geometry and voxelization at the same time for a groom.
 * Shadow casting is handled throught the RT geometry
 * Transmission (lower frequency effect) is handled with voxelization

diff --git a/Engine/Plugins/Runtime/HairStrands/Source/HairStrandsCore/Private/HairStrandsRendering.cpp b/Engine/Plugins/Runtime/HairStrands/Source/HairStrandsCore/Private/HairStrandsRendering.cpp
--- a/Engine/Plugins/Runtime/HairStrands/Source/HairStrandsCore/Private/HairStrandsRendering.cpp
+++ b/Engine/Plugins/Runtime/HairStrands/Source/HairStrandsCore/Private/HairStrandsRendering.cpp
@@ -1687,6 +1687,11 @@ static FHairGroupPublicData::FVertexFactoryInput InternalComputeHairStrandsVerte
 	OutVFInput.Strands.HairDensity = Instance->Strands.Modifier.HairShadowDensity;
 	OutVFInput.Strands.bScatterSceneLighting = Instance->Strands.Modifier.bScatterSceneLighting;
 	OutVFInput.Strands.bUseStableRasterization = Instance->Strands.Modifier.bUseStableRasterization;
+#if RHI_RAYTRACING
+	OutVFInput.Strands.bUseRaytracingGeometry = Instance->Strands.RenRaytracingResource != nullptr;
+#else
+	OutVFInput.Strands.bUseRaytracingGeometry = false;
+#endif
 
 	return OutVFInput;
 }
diff --git a/Engine/Shaders/Private/HairStrands/HairStrandsDeepShadowMask.usf b/Engine/Shaders/Private/HairStrands/HairStrandsDeepShadowMask.usf
--- a/Engine/Shaders/Private/HairStrands/HairStrandsDeepShadowMask.usf
+++ b/Engine/Shaders/Private/HairStrands/HairStrandsDeepShadowMask.usf
@@ -493,7 +493,8 @@ void MainPS(
 		TraversalSettings.SteppingScale = VirtualVoxel.SteppingScale_Shadow;
 		TraversalSettings.Random = Random;
 		TraversalSettings.PixelRadius = SceneDepth * VirtualVoxel.HairCoveragePixelRadiusAtDepth1;
-		
+		TraversalSettings.bCastShadow = true;
+
 		if (bIsValid && RayMarchMaskTexture.Load(uint3(PixelCoord, 0)).r > 0.0f && HairPixelCoverage < CoverageThreshold)
 		{
 			Result = ComputeHairCountVirtualVoxel(
diff --git a/Engine/Shaders/Private/HairStrands/HairStrandsRaytracing.ush b/Engine/Shaders/Private/HairStrands/HairStrandsRaytracing.ush
--- a/Engine/Shaders/Private/HairStrands/HairStrandsRaytracing.ush
+++ b/Engine/Shaders/Private/HairStrands/HairStrandsRaytracing.ush
@@ -54,6 +54,7 @@ float InternalTraverseHair(uint2 PixelCoord, float3 RayOrigin, float3 RayDirecti
 		TraversalSettings.SteppingScale = VirtualVoxel.SteppingScale_Raytracing;
 		TraversalSettings.Random = VoxelRandom;
 		TraversalSettings.bUseOpaqueVisibility = false;
+		TraversalSettings.bCastShadow = true;
 
 		Result = ComputeHairCountVirtualVoxel(
 			P0,
diff --git a/Engine/Shaders/Private/HairStrands/HairStrandsVoxelMip.usf b/Engine/Shaders/Private/HairStrands/HairStrandsVoxelMip.usf
--- a/Engine/Shaders/Private/HairStrands/HairStrandsVoxelMip.usf
+++ b/Engine/Shaders/Private/HairStrands/HairStrandsVoxelMip.usf
@@ -22,7 +22,18 @@ uint ComputeMipDensity(
 		((RawDensity5 & VOXEL_OPAQUE_MASK) >> VOXEL_OPAQUE_SHIFT) +
 		((RawDensity6 & VOXEL_OPAQUE_MASK) >> VOXEL_OPAQUE_SHIFT) +
 		((RawDensity7 & VOXEL_OPAQUE_MASK) >> VOXEL_OPAQUE_SHIFT);
-	const uint OutTotalOpaque = uint(clamp(TotalOpaque / 8.f, TotalOpaque > 0 ? 1 : 0, 0xFF)) << VOXEL_OPAQUE_SHIFT;
+	const uint OutTotalOpaque = uint(clamp(TotalOpaque / 8.f, TotalOpaque > 0 ? 1 : 0, 0x7F)) << VOXEL_OPAQUE_SHIFT;
+
+	// Propagate no shadow casting flag, only if other valid/non-empty voxels are also no-shadow caster
+	const bool bHasShadowCaster =
+		(RawDensity0 & VOXEL_HAIR_MASK) > 0 && (RawDensity0 & VOXEL_CAST_NO_SHADOW_MASK)==0 ||
+		(RawDensity1 & VOXEL_HAIR_MASK) > 0 && (RawDensity1 & VOXEL_CAST_NO_SHADOW_MASK)==0 ||
+		(RawDensity2 & VOXEL_HAIR_MASK) > 0 && (RawDensity2 & VOXEL_CAST_NO_SHADOW_MASK)==0 ||
+		(RawDensity3 & VOXEL_HAIR_MASK) > 0 && (RawDensity3 & VOXEL_CAST_NO_SHADOW_MASK)==0 ||
+		(RawDensity4 & VOXEL_HAIR_MASK) > 0 && (RawDensity4 & VOXEL_CAST_NO_SHADOW_MASK)==0 ||
+		(RawDensity5 & VOXEL_HAIR_MASK) > 0 && (RawDensity5 & VOXEL_CAST_NO_SHADOW_MASK)==0 ||
+		(RawDensity6 & VOXEL_HAIR_MASK) > 0 && (RawDensity6 & VOXEL_CAST_NO_SHADOW_MASK)==0 ||
+		(RawDensity7 & VOXEL_HAIR_MASK) > 0 && (RawDensity7 & VOXEL_CAST_NO_SHADOW_MASK)==0;
 
 	uint TotalHair = 
 		(RawDensity0 & VOXEL_HAIR_MASK) +
@@ -39,7 +50,7 @@ uint ComputeMipDensity(
 	// Insure that if a voxel contains some hair data, its total hair remains > 0 after averaging. 
 	// This is important for pruning invalid/empty page later on, to not remove non-empty page (which 
 	// could arise due to numerical precision).
-	return min(uint(VOXEL_HAIR_MASK), bHasData ? max(TotalHair, 1) : 0) | OutTotalOpaque;
+	return min(uint(VOXEL_HAIR_MASK), bHasData ? max(TotalHair, 1u) : 0u) | OutTotalOpaque | (bHasShadowCaster ? 0u : VOXEL_CAST_NO_SHADOW_MASK);
 }
 
 #if SHADER_MIP_VIRTUALVOXEL
diff --git a/Engine/Shaders/Private/HairStrands/HairStrandsVoxelPageCommon.ush b/Engine/Shaders/Private/HairStrands/HairStrandsVoxelPageCommon.ush
--- a/Engine/Shaders/Private/HairStrands/HairStrandsVoxelPageCommon.ush
+++ b/Engine/Shaders/Private/HairStrands/HairStrandsVoxelPageCommon.ush
@@ -60,10 +60,11 @@ bool IsInVoxelBounds(float3 P, float3 MinP, float3 MaxP)
 #define INVALID_VOXEL_PAGE_INDEX 0xFFFFFFFF
 #define INVALID_MACRO_GROUP_ID 0xFF
 
-#define VOXEL_HAIR_MASK		0x00FFFFFF
-#define VOXEL_OPAQUE_MASK	0xFF000000
-#define VOXEL_OPAQUE_ADD	0xFF000000
-#define VOXEL_OPAQUE_SHIFT	24
+#define VOXEL_CAST_NO_SHADOW_MASK	0x80000000	// Bit mask for holding if the voxel does *not* cast hair shadow (i.e. cast only transmission). The invert/negative logic is for making hair voxelization more efficient.
+#define VOXEL_HAIR_MASK				0x00FFFFFF	// Bit mask for holding hair count/density
+#define VOXEL_OPAQUE_MASK			0x7F000000	// Bit mask for holding opaque geometry
+#define VOXEL_OPAQUE_ADD			0x7F000000
+#define VOXEL_OPAQUE_SHIFT			24
 
 struct FVirtualVoxelCommonDesc
 {
@@ -163,21 +164,17 @@ float GetInternalVoxelHairCount(uint RawDensity)
 	return Raw / VoxelFixPointScale;
 }
 
-FHairTraversalResult GetHairVirtualVoxelDensity(uint3 InVoxelPageCoord, Texture3D<uint> InPageTexture, uint MipIndex, float DensityScale)
+FHairTraversalResult GetHairVirtualVoxelDensity(uint3 InVoxelPageCoord, Texture3D<uint> InPageTexture, uint MipIndex, float DensityScale, bool bCastShadow)
 {
 	FHairTraversalResult Out = InitHairTraversalResult();
 	const uint RawDensity = InPageTexture.Load(uint4(InVoxelPageCoord >> MipIndex, MipIndex));
-	Out.HairCount  = GetInternalVoxelHairCount(RawDensity) * DensityScale;
-	Out.Visibility = GetInternalVoxelOpaqueVisibility(RawDensity);
-	return Out;
-}
 
-FHairTraversalResult GetHairVoxelDensity(uint3 InVolumeCoord, Texture3D<uint> InDensityTexture, uint MipIndex, float DensityScale)
-{
-	FHairTraversalResult Out = InitHairTraversalResult();
-	const uint RawDensity = InDensityTexture.Load(int4(InVolumeCoord, MipIndex));
-	Out.HairCount = GetInternalVoxelHairCount(RawDensity) * DensityScale;
-	Out.Visibility = GetInternalVoxelOpaqueVisibility(RawDensity);
+	const bool bValid = !bCastShadow || (bCastShadow && (RawDensity & VOXEL_CAST_NO_SHADOW_MASK) == 0);
+	if (bValid)
+	{
+		Out.HairCount  = GetInternalVoxelHairCount(RawDensity) * DensityScale;
+		Out.Visibility = GetInternalVoxelOpaqueVisibility(RawDensity);
+	}
 	return Out;
 }
 
diff --git a/Engine/Shaders/Private/HairStrands/HairStrandsVoxelPageTraversal.ush b/Engine/Shaders/Private/HairStrands/HairStrandsVoxelPageTraversal.ush
--- a/Engine/Shaders/Private/HairStrands/HairStrandsVoxelPageTraversal.ush
+++ b/Engine/Shaders/Private/HairStrands/HairStrandsVoxelPageTraversal.ush
@@ -129,6 +129,7 @@ struct FHairTraversalSettings
 	bool	bDebugEnabled;
 	bool	bIsPrimaryRay;
 	bool	bUseOccupancy;
+	bool	bCastShadow;	// True of the traversal is for casting (opaque) shadow, i.e., not for hair transmission
 };
 
 float GetHairTraversalMaxT()
@@ -152,6 +153,7 @@ FHairTraversalSettings InitHairTraversalSettings()
 	Out.bIsPrimaryRay = false;
 	Out.bUseOccupancy = false;
 	Out.ForcedMip = -1;
+	Out.bCastShadow = false;
 	return Out;
 }
 
@@ -270,7 +272,7 @@ FHairTraversalResult ComputeHairCountVirtualVoxel(
 					MipLevel = InSettings.ForcedMip >= 0 ? InSettings.ForcedMip : MipLevel;
 					#endif
 
-					const FHairTraversalResult StepResult = GetHairVirtualVoxelDensity(VoxelPageCoord, InPageTexture, uint(MipLevel), InSettings.DensityScale * HairCountScale);
+					const FHairTraversalResult StepResult = GetHairVirtualVoxelDensity(VoxelPageCoord, InPageTexture, uint(MipLevel), InSettings.DensityScale * HairCountScale, InSettings.bCastShadow);
 					Acc(Out, StepResult);
 
 					if (InSettings.bIsPrimaryRay && StepResult.HairCount > 0)
@@ -468,7 +470,7 @@ FHairTraversalResult ComputeHairCountVirtualVoxel(
 					}
 					#endif
 
-					const FHairTraversalResult StepResult = GetHairVirtualVoxelDensity(VoxelPageCoord, InPageTexture, MipLevel, InSettings.DensityScale);
+					const FHairTraversalResult StepResult = GetHairVirtualVoxelDensity(VoxelPageCoord, InPageTexture, MipLevel, InSettings.DensityScale, InSetttings.bCastShadow);
 					Acc(Out, StepResult);
 
 					if (ShouldStopTraversal(Out, InSettings.CountThreshold))
diff --git a/Engine/Shaders/Private/HairStrands/HairStrandsVoxelRasterCompute.usf b/Engine/Shaders/Private/HairStrands/HairStrandsVoxelRasterCompute.usf
--- a/Engine/Shaders/Private/HairStrands/HairStrandsVoxelRasterCompute.usf
+++ b/Engine/Shaders/Private/HairStrands/HairStrandsVoxelRasterCompute.usf
@@ -13,6 +13,7 @@ uint			MaxRasterCount;
 uint			FrameIdMod8;
 uint			MacroGroupId;
 uint			HairStrandsVF_bIsCullingEnable;
+uint			HairStrandsVF_bHasRaytracedGeometry;
 float			HairStrandsVF_Radius;
 float			HairStrandsVF_RootScale;
 float			HairStrandsVF_TipScale;
@@ -202,12 +203,21 @@ void MainCS(uint2 DispatchThreadID : SV_DispatchThreadID)
 			// Insure we don't write multiple time within the same voxel.
 			// This can happen for small hair segment, where both start & end points could be writtent into the same voxel.
 			const bool bRasterize = any(VoxelPageCoord != PreviousCoord);
-
-			const float HairCoverage = lerp(HairCoverage0, HairCoverage1, U);
 			if (bRasterize)
 			{
 				const float VoxelFixPointScale = GetVoxelDensityFixPointScale();
-				InterlockedAdd(OutPageTexture[VoxelPageCoord], HairCoverage * VoxelFixPointScale * HairStrandsVF_Density);
+				
+				const float HairCoverage = lerp(HairCoverage0, HairCoverage1, U);
+				uint RawData = HairCoverage * VoxelFixPointScale * HairStrandsVF_Density;
+				InterlockedAdd(OutPageTexture[VoxelPageCoord], RawData);
+
+				// Groom having raytraced geometry will cast-shadow on opaque geomtry with their RT-geometry, not with their 
+				// voxelization. To avoid for doubling/incorrect shadowing, we mark voxel with no-shadow casting flag. 
+				// This adds a significant cost when used.
+				if (HairStrandsVF_bHasRaytracedGeometry)
+				{
+					InterlockedOr(OutPageTexture[VoxelPageCoord], VOXEL_CAST_NO_SHADOW_MASK);
+				}
 			}
 			PreviousCoord = VoxelPageCoord;
 		}
diff --git a/Engine/Shaders/Private/VirtualShadowMaps/VirtualShadowMapProjection.usf b/Engine/Shaders/Private/VirtualShadowMaps/VirtualShadowMapProjection.usf
--- a/Engine/Shaders/Private/VirtualShadowMaps/VirtualShadowMapProjection.usf
+++ b/Engine/Shaders/Private/VirtualShadowMaps/VirtualShadowMapProjection.usf
@@ -326,6 +326,7 @@ FVirtualShadowMapSampleResult ProjectLight(
 				TraversalSettings.Random			= Random.xyz;
 				TraversalSettings.PixelRadius		= SceneDepth * VirtualVoxel.HairCoveragePixelRadiusAtDepth1;
 				TraversalSettings.bUseOpaqueVisibility = false;
+				TraversalSettings.bCastShadow		= true;
 
 				const uint VoxelDescCount = VirtualVoxel.NodeDescCount;
 				for (uint VoxelDescIt=0; VoxelDescIt<VoxelDescCount; ++VoxelDescIt)
diff --git a/Engine/Source/Runtime/Renderer/Private/HairStrands/HairStrandsVoxelization.cpp b/Engine/Source/Runtime/Renderer/Private/HairStrands/HairStrandsVoxelization.cpp
--- a/Engine/Source/Runtime/Renderer/Private/HairStrands/HairStrandsVoxelization.cpp
+++ b/Engine/Source/Runtime/Renderer/Private/HairStrands/HairStrandsVoxelization.cpp
@@ -1229,6 +1229,7 @@ class FVoxelRasterComputeCS : public FGlobalShader
 		SHADER_PARAMETER(uint32, MaxRasterCount)
 		SHADER_PARAMETER(uint32, FrameIdMod8)
 		SHADER_PARAMETER(uint32, HairStrandsVF_bIsCullingEnable)
+		SHADER_PARAMETER(uint32, HairStrandsVF_bHasRaytracedGeometry)
 		SHADER_PARAMETER(float,	  HairStrandsVF_Density)
 		SHADER_PARAMETER(float,   HairStrandsVF_Radius)
 		SHADER_PARAMETER(float,   HairStrandsVF_RootScale)
@@ -1319,7 +1320,7 @@ static void AddVirtualVoxelizationComputeRasterPass(
 				PassParameters->HairStrandsVF_TipScale = VFInput.Strands.HairTipScale;
 				PassParameters->HairStrandsVF_Density = VFInput.Strands.HairDensity;
 				PassParameters->HairStrandsVF_LocalToWorldPrimitiveTransform = VFInput.LocalToWorldTransform.ToMatrixWithScale();
-
+				PassParameters->HairStrandsVF_bHasRaytracedGeometry = VFInput.Strands.bUseRaytracingGeometry ? 1u : 0u;
 				const bool bCullingEnable = HairGroupPublicData->GetCullingResultAvailable();
 				PassParameters->HairStrandsVF_bIsCullingEnable = bCullingEnable ? 1 : 0;
 
diff --git a/Engine/Source/Runtime/Renderer/Public/HairStrandsInterface.h b/Engine/Source/Runtime/Renderer/Public/HairStrandsInterface.h
--- a/Engine/Source/Runtime/Renderer/Public/HairStrandsInterface.h
+++ b/Engine/Source/Runtime/Renderer/Public/HairStrandsInterface.h
@@ -268,6 +268,7 @@ public:
 			float HairDensity = 0;
 			bool bUseStableRasterization = false;
 			bool bScatterSceneLighting = false;
+			bool bUseRaytracingGeometry = false;
 		} Strands;
 
 		struct FCards
