Fix shader compilation due to reserved keyword.

diff --git a/Engine/Shaders/Private/HairStrands/HairStrandsVertexFactory.ush b/Engine/Shaders/Private/HairStrands/HairStrandsVertexFactory.ush
--- a/Engine/Shaders/Private/HairStrands/HairStrandsVertexFactory.ush
+++ b/Engine/Shaders/Private/HairStrands/HairStrandsVertexFactory.ush
@@ -152,11 +152,11 @@ FVertexInfo GetVertexInfo(FVertexFactoryInput Input)
 #if RAYHITGROUPSHADER
 FVertexFactoryInput LoadVertexFactoryInputForHGS(uint TriangleIndex, int VertexIndex)
 {
-	FTriangleBaseAttributes Triangle = LoadTriangleBaseAttributes(TriangleIndex);
+	FTriangleBaseAttributes TriangleBase = LoadTriangleBaseAttributes(TriangleIndex);
 
 	FVertexFactoryInput Input;
-	Input.Position = float4(Triangle.LocalPositions[VertexIndex], 1.0f);
-	Input.VertexId = Triangle.Indices[VertexIndex];
+	Input.Position = float4(TriangleBase.LocalPositions[VertexIndex], 1.0f);
+	Input.VertexId = TriangleBase.Indices[VertexIndex];
 	Input.TriangleId = TriangleIndex;
 	Input.TriangleVertexId = VertexIndex;
 #if VF_USE_PRIMITIVE_SCENE_DATA
@@ -585,12 +585,12 @@ FMaterialPixelParameters GetMaterialPixelParameters(float3 RayOrigin, float3 Ray
 		HitAttributes.Barycentrics.y);
 
 	// fetch triangle
-	FTriangleBaseAttributes Triangle = LoadTriangleBaseAttributes(HitPrimitiveIndex);
+	FTriangleBaseAttributes TriangleBase = LoadTriangleBaseAttributes(HitPrimitiveIndex);
 	// Transform vertices into world space
 	const float3 WorldPositions[3] = {
-		CalcWorldPosition(Triangle.LocalPositions[0], SceneData.Primitive.LocalToWorld).xyz,
-		CalcWorldPosition(Triangle.LocalPositions[1], SceneData.Primitive.LocalToWorld).xyz,
-		CalcWorldPosition(Triangle.LocalPositions[2], SceneData.Primitive.LocalToWorld).xyz
+		CalcWorldPosition(TriangleBase.LocalPositions[0], SceneData.Primitive.LocalToWorld).xyz,
+		CalcWorldPosition(TriangleBase.LocalPositions[1], SceneData.Primitive.LocalToWorld).xyz,
+		CalcWorldPosition(TriangleBase.LocalPositions[2], SceneData.Primitive.LocalToWorld).xyz
 	};
 	// Compute triangle normal
 	{
@@ -603,7 +603,7 @@ FMaterialPixelParameters GetMaterialPixelParameters(float3 RayOrigin, float3 Ray
 
 	// NOTE: This must match the tesselation pattern produced by SetBodySegment() in HairStrandsRaytracingGeometry.usf
 	// TODO: is there a way we could handle end-caps elegantly?
-	Result.HairPrimitiveId = Triangle.Indices.z / 4; // last index always belongs to the "bottom" vertex of the hair segment
+	Result.HairPrimitiveId = TriangleBase.Indices.z / 4; // last index always belongs to the "bottom" vertex of the hair segment
 
 	float UCoord = dot(float3(1, float((HitPrimitiveIndex % 2) == 0), 0), Weights);
 	float VCoord = 0.5;// can't measure position across the width easily
