Fix Burley SSS border flickering when mipmaps got garbage data.

diff --git a/Engine/Shaders/Private/SubsurfaceBurleyNormalized.ush b/Engine/Shaders/Private/SubsurfaceBurleyNormalized.ush
--- a/Engine/Shaders/Private/SubsurfaceBurleyNormalized.ush
+++ b/Engine/Shaders/Private/SubsurfaceBurleyNormalized.ush
@@ -777,16 +777,32 @@ float3 GetBilateralNormal(float2 ClampedUV, float2 Extent)
 #endif
 }
 
-BurleySampleDiffuseNormal SampleSSSColorConsideringLocalShared(float2 CenterUV, float2 UVOffset, uint2 CenterGroupThreadID, float MipLevel)
+float2 ClampUVLevel(float2 UV, float2 MinUV, float2 MaxUV, float MipLevel)
 {
-	float2 ClampedUV = clamp(CenterUV+ UVOffset, SubsurfaceInput0_UVViewportBilinearMin, SubsurfaceInput0_UVViewportBilinearMax);
+	uint CeilMipLevel = (uint)(ceil(MipLevel));
+	float2 MipUVCorrection = float( (1u << CeilMipLevel) - 1u ) * SubsurfaceInput0_ExtentInverse;
 	
-	// set mip level to 0 if the mipmap is not generated
+	float2 MinMipUVCorrection = 
+		float2(SubsurfaceInput0_UVViewportMin.x == 0.0f ? 0.0f : MipUVCorrection.x,
+			   SubsurfaceInput0_UVViewportMin.y == 0.0f ? 0.0f : MipUVCorrection.y);
+	float2 MaxMipUVCorrection =
+		float2(SubsurfaceInput0_UVViewportMax.x == 1.0f ? 0.0f : MipUVCorrection.x,
+			   SubsurfaceInput0_UVViewportMax.y == 1.0f ? 0.0f : MipUVCorrection.y);
+
+	return clamp(UV, MinUV + MinMipUVCorrection, MaxUV - MaxMipUVCorrection);
+}
+
+BurleySampleDiffuseNormal SampleSSSColorConsideringLocalShared(float2 CenterUV, float2 UVOffset, uint2 CenterGroupThreadID, float MipLevel)
+{
+	// Set mip level to 0 if the mipmap is not generated
 	if (!ShouldGenerateMipmaps(SUBSURFACE_TILE_TYPE_AFIS))
 	{
 		MipLevel = 0.0f;
 	}
 
+	// Fix border flickering when mipmaps got garbage data.
+	float2 ClampedUV = ClampUVLevel(CenterUV + UVOffset, SubsurfaceInput0_UVViewportBilinearMin, SubsurfaceInput0_UVViewportBilinearMax, MipLevel);
+
 #if !(ENABLE_PROFILE_ID_CACHE && (SUBSURFACE_FORCE_SEPARABLE !=1))
 	#if STRATA_ENABLED
 	const uint2 CenterPos = ClampedUV * SubsurfaceInput0_Extent;
