Fix comments for SSSS_OPACITY_THRESHOLD_EPS to avoid misleading constant value.

#8565

diff --git a/Engine/Shaders/Private/ShadingModelsMaterial.ush b/Engine/Shaders/Private/ShadingModelsMaterial.ush
--- a/Engine/Shaders/Private/ShadingModelsMaterial.ush
+++ b/Engine/Shaders/Private/ShadingModelsMaterial.ush
@@ -68,7 +68,7 @@ void SetGBufferForShadingModel(
 			float MaterialRoughnessToAverage = GetSubsurfaceProfileTexture(SSSS_DUAL_SPECULAR_OFFSET, SubsurfaceProfileUInt).w * SSSS_MAX_DUAL_SPECULAR_ROUGHNESS;
 
 			// Smooth blend out dual specular when opacity is low, we have the extra SSSS_OPACITY_THRESHOLD_EPS so that we fade out by the time we
-			// get to 0.01, as opposed to 0.0.
+			// get to SSSS_OPACITY_THRESHOLD_EPS, as opposed to 0.0.
 			MaterialRoughnessToAverage = lerp(1.0f, MaterialRoughnessToAverage, saturate((Opacity - SSSS_OPACITY_THRESHOLD_EPS) * 10.0f));
 
 			GBuffer.Roughness = saturate(GBuffer.Roughness * MaterialRoughnessToAverage);
diff --git a/Engine/Shaders/Private/SubsurfaceBurleyNormalized.ush b/Engine/Shaders/Private/SubsurfaceBurleyNormalized.ush
--- a/Engine/Shaders/Private/SubsurfaceBurleyNormalized.ush
+++ b/Engine/Shaders/Private/SubsurfaceBurleyNormalized.ush
@@ -1179,7 +1179,7 @@ float4 BurleyNormalizedSS(float2 BufferUV, uint2 GroupThreadID)
 #endif
 
 	// The opacity works by reducing the radius based on opacity, but this runs into precision issues with low opacity values.
-	// So as the opacity goes from 0.02 to 0.01, we transition to fully disabling SSS by the time we get to 0.01.
+	// So as the opacity goes to SSSS_OPACITY_THRESHOLD_EPS, we transition to fully disabling SSS by the time we get there.
 	float LowOpacityEps = SSSS_OPACITY_THRESHOLD_EPS;
 
 	float OriginalLerp = saturate((BurleyParameter.SurfaceOpacity - LowOpacityEps) / LowOpacityEps);
