PR #8559: Fixed clear coat bottom normal artifacts (Contributed by [REDACTED_PERSON])

diff --git a/Engine/Shaders/Private/AmbientCubemapComposite.usf b/Engine/Shaders/Private/AmbientCubemapComposite.usf
--- a/Engine/Shaders/Private/AmbientCubemapComposite.usf
+++ b/Engine/Shaders/Private/AmbientCubemapComposite.usf
@@ -356,7 +356,7 @@ void MainPS(in noperspective float4 UVAndScreenPos : TEXCOORD0, float4 SvPositio
 	
 	BRANCH if( GBuffer.ShadingModelID == SHADINGMODELID_CLEAR_COAT && CLEAR_COAT_BOTTOM_NORMAL)
 	{
-		const float2 oct1 = ((float2(GBuffer.CustomData.a, GBuffer.CustomData.z) * 2) - (256.0/255.0)) + UnitVectorToOctahedron(GBuffer.WorldNormal);
+		const float2 oct1 = ((float2(GBuffer.CustomData.a, GBuffer.CustomData.z) * 4) - (512.0/255.0)) + UnitVectorToOctahedron(GBuffer.WorldNormal);
 		N = OctahedronToUnitVector(oct1);			
 	}
 	float3 DiffuseLookup =  TextureCubeSampleLevel(AmbientCubemap, AmbientCubemapSampler, N, AbsoluteDiffuseMip).rgb;
@@ -377,7 +377,7 @@ void MainPS(in noperspective float4 UVAndScreenPos : TEXCOORD0, float4 SvPositio
 
 			#if CLEAR_COAT_BOTTOM_NORMAL
 				//const float3 ClearCoatUnderNormal = OctahedronToUnitVector((float2(GBuffer.CustomData.a, GBuffer.CustomData.z) * 2) - (256.0/255.0));
-				const float2 oct1 = ((float2(GBuffer.CustomData.a, GBuffer.CustomData.z) * 2) - (256.0/255.0)) + UnitVectorToOctahedron(GBuffer.WorldNormal);
+				const float2 oct1 = ((float2(GBuffer.CustomData.a, GBuffer.CustomData.z) * 4) - (512.0/255.0)) + UnitVectorToOctahedron(GBuffer.WorldNormal);
  				const float3 ClearCoatUnderNormal = OctahedronToUnitVector(oct1);
 				float3 R2 = 2 * dot( V, ClearCoatUnderNormal ) * ClearCoatUnderNormal - V;
 			
diff --git a/Engine/Shaders/Private/BasePassPixelShader.usf b/Engine/Shaders/Private/BasePassPixelShader.usf
--- a/Engine/Shaders/Private/BasePassPixelShader.usf
+++ b/Engine/Shaders/Private/BasePassPixelShader.usf
@@ -1162,7 +1162,7 @@ void FPixelShaderInOut_MainPS(
 	// Clear Coat Bottom Normal
 	BRANCH if( GBuffer.ShadingModelID == SHADINGMODELID_CLEAR_COAT && CLEAR_COAT_BOTTOM_NORMAL)
 	{
-		const float2 oct1 = ((float2(GBuffer.CustomData.a, GBuffer.CustomData.z) * 2) - (256.0/255.0)) + UnitVectorToOctahedron(GBuffer.WorldNormal);
+		const float2 oct1 = ((float2(GBuffer.CustomData.a, GBuffer.CustomData.z) * 4) - (512.0/255.0)) + UnitVectorToOctahedron(GBuffer.WorldNormal);
 		BentNormal = OctahedronToUnitVector(oct1);			
 	}
 	
diff --git a/Engine/Shaders/Private/ClearCoatCommon.ush b/Engine/Shaders/Private/ClearCoatCommon.ush
--- a/Engine/Shaders/Private/ClearCoatCommon.ush
+++ b/Engine/Shaders/Private/ClearCoatCommon.ush
@@ -6,7 +6,7 @@ float3 GetClearCoatBottomNormal(FGBufferData GBuffer, float3 WorldNormal)
 {
 	if (GBuffer.ShadingModelID == SHADINGMODELID_CLEAR_COAT && CLEAR_COAT_BOTTOM_NORMAL)
 	{
-		const float2 oct1 = ((float2(GBuffer.CustomData.a, GBuffer.CustomData.z) * 2) - (256.0 / 255.0)) + UnitVectorToOctahedron(WorldNormal);
+		const float2 oct1 = ((float2(GBuffer.CustomData.a, GBuffer.CustomData.z) * 4) - (512.0 / 255.0)) + UnitVectorToOctahedron(WorldNormal);
 		WorldNormal = OctahedronToUnitVector(oct1);
 	}
 
diff --git a/Engine/Shaders/Private/DeferredLightingCommon.ush b/Engine/Shaders/Private/DeferredLightingCommon.ush
--- a/Engine/Shaders/Private/DeferredLightingCommon.ush
+++ b/Engine/Shaders/Private/DeferredLightingCommon.ush
@@ -342,7 +342,7 @@ FDeferredLightingSplit GetDynamicLightingSplit(
 	float3 N = GBuffer.WorldNormal;
 	BRANCH if( GBuffer.ShadingModelID == SHADINGMODELID_CLEAR_COAT && CLEAR_COAT_BOTTOM_NORMAL)
 	{
-		const float2 oct1 = ((float2(GBuffer.CustomData.a, GBuffer.CustomData.z) * 2) - (256.0/255.0)) + UnitVectorToOctahedron(GBuffer.WorldNormal);
+		const float2 oct1 = ((float2(GBuffer.CustomData.a, GBuffer.CustomData.z) * 4) - (512.0/255.0)) + UnitVectorToOctahedron(GBuffer.WorldNormal);
 		N = OctahedronToUnitVector(oct1);			
 	}
 	
diff --git a/Engine/Shaders/Private/RayTracing/RayTracingMaterialHitShaders.usf b/Engine/Shaders/Private/RayTracing/RayTracingMaterialHitShaders.usf
--- a/Engine/Shaders/Private/RayTracing/RayTracingMaterialHitShaders.usf
+++ b/Engine/Shaders/Private/RayTracing/RayTracingMaterialHitShaders.usf
@@ -380,7 +380,7 @@ RAY_TRACING_ENTRY_CLOSEST_HIT(MaterialCHS,
 			#endif
 
 			float2 oct1 = UnitVectorToOctahedron(tempnormal);
-			float2 oct3 = ( (oct1 - oct2) *  0.5 ) + (128.0/255.0);
+			float2 oct3 = ( (oct1 - oct2) *  0.25 ) + (128.0/255.0);
 			Payload.CustomData.a = oct3.x;
 			Payload.CustomData.z = oct3.y;
 		#else
diff --git a/Engine/Shaders/Private/RayTracing/RayTracingReflections.usf b/Engine/Shaders/Private/RayTracing/RayTracingReflections.usf
--- a/Engine/Shaders/Private/RayTracing/RayTracingReflections.usf
+++ b/Engine/Shaders/Private/RayTracing/RayTracingReflections.usf
@@ -342,7 +342,7 @@ RAY_TRACING_ENTRY_RAYGEN(RayTracingReflectionsRGS)
 			if (ShadingModelID == SHADINGMODELID_CLEAR_COAT)
 			{
 				#if CLEAR_COAT_BOTTOM_NORMAL
-				const float2 oct1 = ((float2(GBufferData.CustomData.a, GBufferData.CustomData.z) * 2) - (256.0/255.0)) + UnitVectorToOctahedron(TopLayerWorldNormal);
+				const float2 oct1 = ((float2(GBufferData.CustomData.a, GBufferData.CustomData.z) * 4) - (512.0/255.0)) + UnitVectorToOctahedron(TopLayerWorldNormal);
 				BottomLayerWorldNormal = OctahedronToUnitVector(oct1);
 				#endif
 
@@ -806,7 +806,7 @@ RAY_TRACING_ENTRY_RAYGEN(RayTracingReflectionsRGS)
 					BottomLayerWorldTangent = TopLayerWorldTangent;
 					float4 CustomData = PackedPayload.GetCustomData();
 					#if CLEAR_COAT_BOTTOM_NORMAL
-					const float2 oct1 = ((float2(CustomData.a, CustomData.z) * 2) - (256.0 / 255.0)) + UnitVectorToOctahedron(TopLayerWorldNormal);
+					const float2 oct1 = ((float2(CustomData.a, CustomData.z) * 4) - (512.0 / 255.0)) + UnitVectorToOctahedron(TopLayerWorldNormal);
 					BottomLayerWorldNormal = OctahedronToUnitVector(oct1);
 					#endif
 
diff --git a/Engine/Shaders/Private/ReflectionEnvironmentPixelShader.usf b/Engine/Shaders/Private/ReflectionEnvironmentPixelShader.usf
--- a/Engine/Shaders/Private/ReflectionEnvironmentPixelShader.usf
+++ b/Engine/Shaders/Private/ReflectionEnvironmentPixelShader.usf
@@ -243,7 +243,7 @@ float3 ReflectionEnvironment(FGBufferData GBuffer, float AmbientOcclusion, float
 		Color = lerp( Color, float4(0,0,0,1), ClearCoat );
 
 #if CLEAR_COAT_BOTTOM_NORMAL
-		const float2 oct1 = ((float2(GBuffer.CustomData.a, GBuffer.CustomData.z) * 2) - (256.0/255.0)) + UnitVectorToOctahedron(GBuffer.WorldNormal);
+		const float2 oct1 = ((float2(GBuffer.CustomData.a, GBuffer.CustomData.z) * 4) - (512.0/255.0)) + UnitVectorToOctahedron(GBuffer.WorldNormal);
 		const float3 ClearCoatUnderNormal = OctahedronToUnitVector(oct1);
 
 		const float3 BottomEffectiveNormal = ClearCoatUnderNormal;			
diff --git a/Engine/Shaders/Private/ShadingModelsMaterial.ush b/Engine/Shaders/Private/ShadingModelsMaterial.ush
--- a/Engine/Shaders/Private/ShadingModelsMaterial.ush
+++ b/Engine/Shaders/Private/ShadingModelsMaterial.ush
@@ -68,7 +68,7 @@ void SetGBufferForShadingModel(
 			float MaterialRoughnessToAverage = GetSubsurfaceProfileTexture(SSSS_DUAL_SPECULAR_OFFSET, SubsurfaceProfileUInt).w * SSSS_MAX_DUAL_SPECULAR_ROUGHNESS;
 
 			// Smooth blend out dual specular when opacity is low, we have the extra SSSS_OPACITY_THRESHOLD_EPS so that we fade out by the time we
-			// get to SSSS_OPACITY_THRESHOLD_EPS, as opposed to 0.0.
+			// get to 0.01, as opposed to 0.0.
 			MaterialRoughnessToAverage = lerp(1.0f, MaterialRoughnessToAverage, saturate((Opacity - SSSS_OPACITY_THRESHOLD_EPS) * 10.0f));
 
 			GBuffer.Roughness = saturate(GBuffer.Roughness * MaterialRoughnessToAverage);
@@ -105,7 +105,7 @@ void SetGBufferForShadingModel(
 				#endif
 
 				float2 oct1 = UnitVectorToOctahedron(tempnormal);
-				float2 oct3 = ( (oct1 - oct2) *  0.5 ) + (128.0/255.0);
+				float2 oct3 = ( (oct1 - oct2) *  0.25 ) + (128.0/255.0);
 				GBuffer.CustomData.a = oct3.x;
 				GBuffer.CustomData.z = oct3.y;
 			#else
