Fixed: Android Materials Sometimes Rendering Black

diff --git a/Engine/Shaders/Private/MaterialTemplate.ush b/Engine/Shaders/Private/MaterialTemplate.ush
--- a/Engine/Shaders/Private/MaterialTemplate.ush
+++ b/Engine/Shaders/Private/MaterialTemplate.ush
@@ -3035,7 +3035,7 @@ uint DiscardMaterialWithPixelCoverage(FMaterialPixelParameters MaterialParameter
 	#define FIsFrontFace bool
 	half GetFloatFacingSign(FIsFrontFace bIsFrontFace)
 	{
-#if COMPILER_DXC && COMPILER_VULKAN
+#if COMPILER_DXC && (COMPILER_VULKAN || COMPILER_GLSL_ES3_1)
 		// We need to flip SV_IsFrontFace for Vulkan when compiling with DXC due to different coordinate systems.
 		// HLSLcc did that by flipping SV_IsFrontFace in the high-level GLSL output.
 		return bIsFrontFace ? -1 : +1;
