Implement depth of field support for the path tracer

This implements a ray-traced sampled depth of field which can serve as a ground truth for the post processed depth of field. Users can select between reference quality or post processing with a path tracing specific checkbox in the post process settings.

diff --git a/Engine/Shaders/Private/PathTracing/PathTracingCamera.ush b/Engine/Shaders/Private/PathTracing/PathTracingCamera.ush
deleted file mode 100644
diff --git a/Engine/Shaders/Private/PathTracing/PathTracingCore.ush b/Engine/Shaders/Private/PathTracing/PathTracingCore.ush
--- a/Engine/Shaders/Private/PathTracing/PathTracingCore.ush
+++ b/Engine/Shaders/Private/PathTracing/PathTracingCore.ush
@@ -20,7 +20,6 @@
 #include "./Light/PathTracingLightSampling.ush"
 #include "./Light/PathTracingLightGrid.ush"
 #include "./Material/PathTracingMaterialSampling.ush"
-#include "PathTracingCamera.ush"
 
 float BlendFactor;
 uint Iteration;
@@ -44,6 +43,8 @@ float MaxPathIntensity;
 float MaxNormalBias;
 float FilterWidth;
 float AbsorptionScale; // inverse distance at which we hit BaseColor - exposed mainly so we can disable absorption handling and experiment with the default scale factor
+float CameraFocusDistance;
+float CameraLensRadius;
 
 RWTexture2D<float4> RadianceTexture;
 RWTexture2D<float4> AlbedoTexture;
@@ -610,13 +611,40 @@ FPathState CreatePathState(uint2 DispatchIdx, uint2 DispatchDim)
 
 	// Initialize ray and payload
 	{
-		uint2 Pixel;
-		float EmitterPdf;
-		SampleEmitter(LaunchIndex, Output.RandSequence, 0, 0, Output.Ray, Pixel, EmitterPdf);
+		float2 AAJitter = RandomSequence_GenerateSample2D(Output.RandSequence);
+
+		// importance sample a gaussian kernel with variable sigma
+		float3 Disk = ConcentricDiskSamplingHelper(AAJitter);
+		float Sigma = FilterWidth / 6.0; // user-provided width covers +/-3*Sigma
+		float R = min(Disk.z, 0.99999994); // avoid log(0) when R=1
+		AAJitter = 0.5 + Sigma * Disk.xy * sqrt(-2.0 * log(1.0 - R * R));
+
+		float2 ViewportUV = (LaunchIndex.xy + AAJitter.xy) * View.BufferSizeAndInvSize.zw;
+		Output.Ray = CreatePrimaryRay(ViewportUV);
+	}
+
+	if (CameraLensRadius > 0)
+	{
+	    // DOF enabled - apply simple thin-lens model
+		float2 LensSample = RandomSequence_GenerateSample2D(Output.RandSequence);
+		float3 ViewX = View.ViewToTranslatedWorld[0].xyz;
+		float3 ViewY = View.ViewToTranslatedWorld[1].xyz;
+		float3 ViewZ = View.ViewToTranslatedWorld[2].xyz;
+		// shift back origin by the near plane amount
+		float ZFactor = rcp(dot(ViewZ, Output.Ray.Direction));
+		float3 NearNudge = (View.NearPlane * ZFactor) * Output.Ray.Direction;
+		float3 Origin = Output.Ray.Origin - NearNudge;
+		// compute focus plane
+		float3 FocusP = Origin + (CameraFocusDistance * ZFactor) * Output.Ray.Direction;
+		// nudge ray origin
+		LensSample = CameraLensRadius * UniformSampleDiskConcentric(LensSample);
+		Origin += LensSample.x * ViewX + LensSample.y * ViewY;
+		// recompute direction
+		Output.Ray.Direction = normalize(FocusP - Origin);
+		// move ray origin back to the near plane for consistency
+		Output.Ray.Origin = Origin + Output.Ray.Direction * (View.NearPlane * rcp(dot(ViewZ, Output.Ray.Direction)));
 	}
 
-	Output.Ray.Direction = normalize(Output.Ray.Direction);
-
 	Output.RayCone.SpreadAngle = View.EyeToPixelSpreadAngle;
 
 	// path state variables (these cary information between bounces)
diff --git a/Engine/Shaders/Private/PathTracing/Sensor/PathTracingPinholeCamera.ush b/Engine/Shaders/Private/PathTracing/Sensor/PathTracingPinholeCamera.ush
deleted file mode 100644
diff --git a/Engine/Source/Runtime/Engine/Classes/Engine/Scene.h b/Engine/Source/Runtime/Engine/Classes/Engine/Scene.h
--- a/Engine/Source/Runtime/Engine/Classes/Engine/Scene.h
+++ b/Engine/Source/Runtime/Engine/Classes/Engine/Scene.h
@@ -1230,6 +1230,9 @@ struct FPostProcessSettings
 	UPROPERTY(EditAnywhere, BlueprintReadWrite, Category = Overrides, meta = (PinHiddenByDefault, InlineEditConditionToggle))
 	uint32 bOverride_PathTracingMaxPathExposure : 1;
 
+	UPROPERTY(EditAnywhere, BlueprintReadWrite, Category = Overrides, meta = (PinHiddenByDefault, InlineEditConditionToggle))
+	uint32 bOverride_PathTracingEnableReferenceDOF : 1;
+
 	UPROPERTY(EditAnywhere, BlueprintReadWrite, Category = Overrides, meta = (PinHiddenByDefault, InlineEditConditionToggle))
 	uint32 bOverride_PathTracingEnableDenoiser : 1;
 
@@ -2018,10 +2021,14 @@ struct FPostProcessSettings
 	UPROPERTY(EditAnywhere, BlueprintReadWrite, Category = "Rendering Features|PathTracing", meta = (editcondition = "bOverride_PathTracingEnableEmissive", DisplayName = "Emissive Materials"))
 	uint32 PathTracingEnableEmissive : 1;
 
-	/** Sets the maximum exposure allowed in the path tracer to reduce fireflies. This should be set a few steps higher than the scene exposure. */
+	/** Sets the maximum exposure allowed in the path tracer to reduce fireflies. This should be set a few stops higher than the scene exposure. */
 	UPROPERTY(interp, EditAnywhere, BlueprintReadWrite, Category = "Rendering Features|PathTracing", meta = (ClampMin = "-10.0", ClampMax = "30.0", editcondition = "bOverride_PathTracingMaxPathExposure", DisplayName = "Max Path Exposure"))
 	float PathTracingMaxPathExposure;
 
+	/** Enables a reference quality depth-of-field which replaces the post-process effect. */
+	UPROPERTY(EditAnywhere, BlueprintReadWrite, Category = "Rendering Features|PathTracing", meta = (editcondition = "bOverride_PathTracingEnableReferenceDOF", DisplayName = "Reference Depth Of Field"))
+	uint32 PathTracingEnableReferenceDOF : 1;
+
 	/** Run the currently loaded denoiser plugin on the last sample to remove noise from the output. Has no effect if a plug-in is not loaded. */
 	UPROPERTY(EditAnywhere, BlueprintReadWrite, Category = "Rendering Features|PathTracing", meta = (editcondition = "bOverride_PathTracingEnableDenoiser", DisplayName = "Denoiser"))
 	uint32 PathTracingEnableDenoiser : 1;
diff --git a/Engine/Source/Runtime/Engine/Private/Scene.cpp b/Engine/Source/Runtime/Engine/Private/Scene.cpp
--- a/Engine/Source/Runtime/Engine/Private/Scene.cpp
+++ b/Engine/Source/Runtime/Engine/Private/Scene.cpp
@@ -585,6 +585,7 @@ FPostProcessSettings::FPostProcessSettings()
 	PathTracingFilterWidth = 3.0f;
 	PathTracingEnableEmissive = 1;
 	PathTracingMaxPathExposure = 30.0f;
+	PathTracingEnableReferenceDOF = 0;
 	PathTracingEnableDenoiser = 1;
 	
 	bMobileHQGaussian = false;
@@ -774,6 +775,7 @@ FPostProcessSettings::FPostProcessSettings(const FPostProcessSettings& Settings)
 	, bOverride_PathTracingFilterWidth(Settings.bOverride_PathTracingFilterWidth)
 	, bOverride_PathTracingEnableEmissive(Settings.bOverride_PathTracingEnableEmissive)
 	, bOverride_PathTracingMaxPathExposure(Settings.bOverride_PathTracingMaxPathExposure)
+	, bOverride_PathTracingEnableReferenceDOF(Settings.bOverride_PathTracingEnableReferenceDOF)
 	, bOverride_PathTracingEnableDenoiser(Settings.bOverride_PathTracingEnableDenoiser)
 
 	, bMobileHQGaussian(Settings.bMobileHQGaussian)
@@ -950,6 +952,7 @@ FPostProcessSettings::FPostProcessSettings(const FPostProcessSettings& Settings)
 	, PathTracingFilterWidth(Settings.PathTracingFilterWidth)
 	, PathTracingEnableEmissive(Settings.PathTracingEnableEmissive)
 	, PathTracingMaxPathExposure(Settings.PathTracingMaxPathExposure)
+	, PathTracingEnableReferenceDOF(Settings.PathTracingEnableReferenceDOF)
 	, PathTracingEnableDenoiser(Settings.PathTracingEnableDenoiser)
 
 	, ScreenPercentage_DEPRECATED(Settings.ScreenPercentage_DEPRECATED)
diff --git a/Engine/Source/Runtime/Engine/Private/SceneView.cpp b/Engine/Source/Runtime/Engine/Private/SceneView.cpp
--- a/Engine/Source/Runtime/Engine/Private/SceneView.cpp
+++ b/Engine/Source/Runtime/Engine/Private/SceneView.cpp
@@ -1613,6 +1613,11 @@ void FSceneView::OverridePostProcessSettings(const FPostProcessSettings& Src, fl
 			Dest.PathTracingMaxPathExposure = Src.PathTracingMaxPathExposure;
 		}
 
+		if (Src.bOverride_PathTracingEnableReferenceDOF)
+		{
+			Dest.PathTracingEnableReferenceDOF = Src.PathTracingEnableReferenceDOF;
+		}
+
 		if (Src.bOverride_PathTracingEnableDenoiser)
 		{
 			Dest.PathTracingEnableDenoiser = Src.PathTracingEnableDenoiser;
diff --git a/Engine/Source/Runtime/Renderer/Private/PathTracing.cpp b/Engine/Source/Runtime/Renderer/Private/PathTracing.cpp
--- a/Engine/Source/Runtime/Renderer/Private/PathTracing.cpp
+++ b/Engine/Source/Runtime/Renderer/Private/PathTracing.cpp
@@ -235,6 +235,8 @@ BEGIN_SHADER_PARAMETER_STRUCT(FPathTracingData, )
 	SHADER_PARAMETER(float, MaxNormalBias)
 	SHADER_PARAMETER(float, FilterWidth)
 	SHADER_PARAMETER(float, AbsorptionScale)
+	SHADER_PARAMETER(float, CameraFocusDistance)
+	SHADER_PARAMETER(float, CameraLensRadius)
 END_SHADER_PARAMETER_STRUCT()
 
 
@@ -268,6 +270,8 @@ struct FPathTracingConfig
 			PathTracingData.MaxPathIntensity != Other.PathTracingData.MaxPathIntensity ||
 			PathTracingData.FilterWidth != Other.PathTracingData.FilterWidth ||
 			PathTracingData.AbsorptionScale != Other.PathTracingData.AbsorptionScale ||
+			PathTracingData.CameraFocusDistance != Other.PathTracingData.CameraFocusDistance ||
+			PathTracingData.CameraLensRadius != Other.PathTracingData.CameraLensRadius ||
 			ViewRect != Other.ViewRect ||
 			LightShowFlags != Other.LightShowFlags ||
 			LightGridResolution != Other.LightGridResolution ||
@@ -348,6 +352,17 @@ static void PrepareShaderArgs(const FViewInfo& View, FPathTracingData& PathTraci
 	}
 	PathTracingData.FilterWidth = FilterWidth;
 	PathTracingData.AbsorptionScale = CVarPathTracingAbsorptionScale.GetValueOnRenderThread();
+	PathTracingData.CameraFocusDistance = 0;
+	PathTracingData.CameraLensRadius = 0;
+	if (View.Family->EngineShowFlags.DepthOfField &&
+		View.FinalPostProcessSettings.PathTracingEnableReferenceDOF &&
+		View.FinalPostProcessSettings.DepthOfFieldFocalDistance > 0 &&
+		View.FinalPostProcessSettings.DepthOfFieldFstop > 0)
+	{
+		const float FocalLengthInCM = 0.05f * View.FinalPostProcessSettings.DepthOfFieldSensorWidth * View.ViewMatrices.GetProjectionMatrix().M[0][0];
+		PathTracingData.CameraFocusDistance = View.FinalPostProcessSettings.DepthOfFieldFocalDistance;
+		PathTracingData.CameraLensRadius = 0.5f * FocalLengthInCM / View.FinalPostProcessSettings.DepthOfFieldFstop;
+	}
 }
 
 static bool ShouldCompilePathTracingShadersForProject(EShaderPlatform ShaderPlatform)
diff --git a/Engine/Source/Runtime/Renderer/Private/PostProcess/DiaphragmDOF.cpp b/Engine/Source/Runtime/Renderer/Private/PostProcess/DiaphragmDOF.cpp
--- a/Engine/Source/Runtime/Renderer/Private/PostProcess/DiaphragmDOF.cpp
+++ b/Engine/Source/Runtime/Renderer/Private/PostProcess/DiaphragmDOF.cpp
@@ -1311,6 +1311,7 @@ bool DiaphragmDOF::IsEnabled(const FViewInfo& View)
 		DiaphragmDOF::IsSupported(View.GetShaderPlatform()) &&
 		View.Family->EngineShowFlags.DepthOfField &&
 		bDepthOfFieldRequestedByCVar &&
+		!(View.Family->EngineShowFlags.PathTracing && View.FinalPostProcessSettings.PathTracingEnableReferenceDOF) &&
 		View.FinalPostProcessSettings.DepthOfFieldFstop > 0 &&
 		View.FinalPostProcessSettings.DepthOfFieldFocalDistance > 0;
 }
