Add Hair strand path tracing support

Add support for binding a ray tracing material to hair strands so that they can show up in path tracing.

Add hair shading model to path tracer

diff --git a/Engine/Plugins/Runtime/HairStrands/Source/HairStrandsCore/Private/GroomResources.cpp b/Engine/Plugins/Runtime/HairStrands/Source/HairStrandsCore/Private/GroomResources.cpp
--- a/Engine/Plugins/Runtime/HairStrands/Source/HairStrandsCore/Private/GroomResources.cpp
+++ b/Engine/Plugins/Runtime/HairStrands/Source/HairStrandsCore/Private/GroomResources.cpp
@@ -1503,13 +1503,13 @@ void FHairCardsInterpolationResource::InternalRelease()
 // TODO: use a plain float vertex buffer with 3x the entries instead to save memory? (float3 UAVs are not allowed)
 bool GetSupportHairStrandsProceduralPrimitive(EShaderPlatform InShaderPlatform);
 FHairStrandsRaytracingResource::FHairStrandsRaytracingResource(const FHairStrandsBulkData& InData, const FHairResourceName& ResourceName) :
-	FHairCommonResource(EHairStrandsAllocationType::Deferred, ResourceName)
+	FHairCommonResource(EHairStrandsAllocationType::Deferred, ResourceName), bOwnBuffers(true)
 {
-	bOwnBuffers = true;
 	bProceduralPrimitive = GetSupportHairStrandsProceduralPrimitive(GMaxRHIShaderPlatform);
 	if (bProceduralPrimitive)
 	{
-		VertexCount = InData.GetNumPoints() * 2; // only allocate space for primitive AABBs
+		// only allocate space for primitive AABBs
+		VertexCount = InData.GetNumPoints() * 2 * STRANDS_PROCEDURAL_INTERSECTOR_MAX_SPLITS;
 	}
 	else
 	{
diff --git a/Engine/Plugins/Runtime/HairStrands/Source/HairStrandsCore/Private/HairStrandsRendering.cpp b/Engine/Plugins/Runtime/HairStrands/Source/HairStrandsCore/Private/HairStrandsRendering.cpp
--- a/Engine/Plugins/Runtime/HairStrands/Source/HairStrandsCore/Private/HairStrandsRendering.cpp
+++ b/Engine/Plugins/Runtime/HairStrands/Source/HairStrandsCore/Private/HairStrandsRendering.cpp
@@ -29,6 +29,10 @@
 static float GHairRaytracingRadiusScale = 0;
 static FAutoConsoleVariableRef CVarHairRaytracingRadiusScale(TEXT("r.HairStrands.RaytracingRadiusScale"), GHairRaytracingRadiusScale, TEXT("Override the per instance scale factor for raytracing hair strands geometry (0: disabled, >0:enabled)"));
 
+static int GHairRaytracingProceduralSplits = 4;
+static FAutoConsoleVariableRef CVarHairRaytracingProceduralScale(TEXT("r.HairStrands.RaytracingProceduralSplits"), GHairRaytracingProceduralSplits, TEXT("Change how many AABBs are used per hair segment to balance between BVH build cost and ray tracing performance. (default: 4)"));
+
+
 static float GStrandHairWidth = 0.0f;
 static FAutoConsoleVariableRef CVarStrandHairWidth(TEXT("r.HairStrands.StrandWidth"), GStrandHairWidth, TEXT("Width of hair strand"));
 
@@ -96,6 +100,11 @@ bool NeedsUpdateCardsMeshTriangles()
 	return GetHairCardsSimulationType() == EHairCardsSimulationType::Guide;
 }
 
+int GetHairRaytracingProceduralSplits()
+{
+	return FMath::Clamp(GHairRaytracingProceduralSplits, 1, STRANDS_PROCEDURAL_INTERSECTOR_MAX_SPLITS);
+}
+
 static FIntVector ComputeDispatchCount(uint32 ItemCount, uint32 GroupSize)
 {
 	const uint32 GroupCount = FMath::DivideAndRoundUp(ItemCount, GroupSize);
@@ -1249,6 +1258,7 @@ class FHairRaytracingGeometryCS : public FGlobalShader
 		SHADER_PARAMETER_RDG_BUFFER_SRV(Buffer, TangentBuffer)
 		SHADER_PARAMETER_RDG_BUFFER_UAV(RWBuffer, OutputPositionBuffer)
 		SHADER_PARAMETER_RDG_BUFFER_UAV(RWBuffer, OutputIndexBuffer)
+		SHADER_PARAMETER(uint32, RaytracingProceduralSplits)
 		SHADER_PARAMETER_STRUCT_INCLUDE(ShaderPrint::FShaderParameters, ShaderPrintUniformBuffer)
 		SHADER_PARAMETER_STRUCT_INCLUDE(ShaderDrawDebug::FShaderParameters, ShaderDrawParameters)
 	END_SHADER_PARAMETER_STRUCT()
@@ -1271,6 +1281,7 @@ static void AddGenerateRaytracingGeometryPass(
 	const FShaderDrawDebugData* ShaderDrawData,
 	uint32 VertexCount,
 	bool bProceduralPrimitive,
+	int ProceduralSplits,
 	float HairRadius,
 	float RootScale,
 	float TipScale,
@@ -1294,6 +1305,7 @@ static void AddGenerateRaytracingGeometryPass(
 	Parameters->PositionBuffer = PositionBuffer;
 	Parameters->TangentBuffer = TangentBuffer;
 	Parameters->OutputPositionBuffer = OutPositionBuffer;
+	Parameters->RaytracingProceduralSplits = GetHairRaytracingProceduralSplits();
 	if (ShaderPrintData)
 	{
 		ShaderPrint::SetParameters(GraphBuilder, *ShaderPrintData, Parameters->ShaderPrintUniformBuffer);
@@ -1303,6 +1315,7 @@ static void AddGenerateRaytracingGeometryPass(
 		ShaderDrawDebug::SetParameters(GraphBuilder, *ShaderDrawData, Parameters->ShaderDrawParameters);
 	}
 	Parameters->OutputIndexBuffer = OutIndexBuffer;
+	Parameters->RaytracingProceduralSplits = ProceduralSplits;
 
 	const bool bCullingEnable = CullingData.bCullingResultAvailable;
 	Parameters->HairStrandsVF_bIsCullingEnable = bCullingEnable ? 1 : 0;
@@ -1410,7 +1423,8 @@ static void BuildHairAccelerationStructure_Strands(
 	FRayTracingGeometry* OutRayTracingGeometry,
 	const FString& AssetDebugName,
 	uint32 LODIndex,
-	bool bProceduralPrimitive)
+	bool bProceduralPrimitive,
+	int ProceduralSplits)
 {
 	FRayTracingGeometryInitializer Initializer;
 #if HAIR_RESOURCE_DEBUG_NAME
@@ -1425,7 +1439,7 @@ static void BuildHairAccelerationStructure_Strands(
 	if (bProceduralPrimitive)
 	{
 		Initializer.GeometryType = RTGT_Procedural;
-		Initializer.TotalPrimitiveCount = RaytracingVertexCount / 2;
+		Initializer.TotalPrimitiveCount = (RaytracingVertexCount / (2 * STRANDS_PROCEDURAL_INTERSECTOR_MAX_SPLITS)) * ProceduralSplits;
 	}
 	else
 	{
@@ -1440,16 +1454,17 @@ static void BuildHairAccelerationStructure_Strands(
 	FRayTracingGeometrySegment Segment;
 	Segment.VertexBuffer = PositionBuffer;
 	Segment.VertexBufferElementType = FHairStrandsRaytracingFormat::VertexElementType;
-	Segment.MaxVertices = RaytracingVertexCount / 2;
 	if (bProceduralPrimitive)
 	{
 		// stride covers one AABB (which is made up of two vertices)
 		Segment.VertexBufferStride = FHairStrandsRaytracingFormat::SizeInByte * 2;
-		Segment.NumPrimitives = RaytracingVertexCount / 2;
+		Segment.MaxVertices = RaytracingVertexCount;
+		Segment.NumPrimitives = Initializer.TotalPrimitiveCount;
 	}
 	else
 	{
 		Segment.VertexBufferStride = FHairStrandsRaytracingFormat::SizeInByte;
+		Segment.MaxVertices = RaytracingVertexCount;
 		Segment.NumPrimitives = RayTracingIndexCount / 3;
 	}
 	Initializer.Segments.Add(Segment);
@@ -1975,8 +1990,13 @@ void ComputeHairStrandsInterpolation(
 				const float HairRadiusRT	= Instance->HairGroupPublicData->VFInput.Strands.HairRaytracingRadiusScale * Instance->HairGroupPublicData->VFInput.Strands.HairRadius;
 				const float HairRootScaleRT = Instance->HairGroupPublicData->VFInput.Strands.HairRootScale;
 				const float HairTipScaleRT	= Instance->HairGroupPublicData->VFInput.Strands.HairTipScale;
+				const int ProceduralSplits = GetHairRaytracingProceduralSplits();
 
-				const bool bNeedUpdate = Instance->Strands.DeformedResource != nullptr ||  Instance->Strands.CachedHairScaledRadius != HairRadiusRT || Instance->Strands.CachedHairRootScale != HairRootScaleRT || Instance->Strands.CachedHairTipScale != HairTipScaleRT;
+				const bool bNeedUpdate = Instance->Strands.DeformedResource != nullptr ||
+					Instance->Strands.CachedHairScaledRadius != HairRadiusRT ||
+					Instance->Strands.CachedHairRootScale != HairRootScaleRT ||
+					Instance->Strands.CachedHairTipScale != HairTipScaleRT ||
+					Instance->Strands.CachedProceduralSplits != ProceduralSplits;
 				const bool bNeedBuild = !Instance->Strands.RenRaytracingResource->bIsRTGeometryInitialized;
 				if (bNeedBuild || bNeedUpdate)
 				{
@@ -1992,6 +2012,7 @@ void ComputeHairStrandsInterpolation(
 						ShaderDrawData,
 						Instance->Strands.RestResource->GetVertexCount(),
 						bProceduralPrimitive,
+						ProceduralSplits,
 						HairRadiusRT,
 						HairRootScaleRT,
 						HairTipScaleRT,
@@ -2009,18 +2030,18 @@ void ComputeHairStrandsInterpolation(
 					GraphBuilder.AddPass(
 						RDG_EVENT_NAME("HairStrands::UpdateBLAS(Strands)"),
 						ERDGPassFlags::NeverCull,
-					[Instance, HairLODIndex, bNeedUpdate, bProceduralPrimitive](FRHICommandListImmediate& RHICmdList)
+					[Instance, HairLODIndex, bNeedUpdate, bProceduralPrimitive, ProceduralSplits](FRHICommandListImmediate& RHICmdList)
 					{
 						SCOPED_GPU_MASK(RHICmdList, FRHIGPUMask::All());
 
-						const bool bLocalNeedBuild = !Instance->Strands.RenRaytracingResource->bIsRTGeometryInitialized;
+						const bool bLocalNeedBuild = !Instance->Strands.RenRaytracingResource->bIsRTGeometryInitialized || Instance->Strands.CachedProceduralSplits != ProceduralSplits;
 						if (bLocalNeedBuild)
 						{
 							FBufferRHIRef PositionBuffer(Instance->Strands.RenRaytracingResource->PositionBuffer.Buffer->GetRHI());
 
 							// no index buffer needed for procedural primitive
 							FBufferRHIRef IndexBuffer(bProceduralPrimitive ? nullptr : Instance->Strands.RenRaytracingResource->IndexBuffer.Buffer->GetRHI());
-
+							Instance->Strands.CachedProceduralSplits = ProceduralSplits;
 							BuildHairAccelerationStructure_Strands(RHICmdList,
 								Instance->Strands.RenRaytracingResource->VertexCount,
 								Instance->Strands.RenRaytracingResource->IndexCount,
@@ -2029,7 +2050,8 @@ void ComputeHairStrandsInterpolation(
 								&Instance->Strands.RenRaytracingResource->RayTracingGeometry,
 								Instance->Debug.GroomAssetName,
 								HairLODIndex,
-								bProceduralPrimitive
+								bProceduralPrimitive,
+								ProceduralSplits
 							);
 						}
 						else if (bNeedUpdate)
diff --git a/Engine/Plugins/Runtime/HairStrands/Source/HairStrandsCore/Private/HairStrandsVertexFactory.cpp b/Engine/Plugins/Runtime/HairStrands/Source/HairStrandsCore/Private/HairStrandsVertexFactory.cpp
--- a/Engine/Plugins/Runtime/HairStrands/Source/HairStrandsCore/Private/HairStrandsVertexFactory.cpp
+++ b/Engine/Plugins/Runtime/HairStrands/Source/HairStrandsCore/Private/HairStrandsVertexFactory.cpp
@@ -14,11 +14,11 @@
 #include "GroomInstance.h"
 
 #define VF_STRANDS_SUPPORT_GPU_SCENE 0
-#define VF_STRANDS_CUSTOM_INTERSECTOR 1
+#define VF_STRANDS_PROCEDURAL_INTERSECTOR 1
 
 bool GetSupportHairStrandsProceduralPrimitive(EShaderPlatform InShaderPlatform)
 {
-	return VF_STRANDS_CUSTOM_INTERSECTOR && FDataDrivenShaderPlatformInfo::GetSupportsRayTracingProceduralPrimitive(InShaderPlatform);
+	return VF_STRANDS_PROCEDURAL_INTERSECTOR && FDataDrivenShaderPlatformInfo::GetSupportsRayTracingProceduralPrimitive(InShaderPlatform);
 }
 
 /////////////////////////////////////////////////////////////////////////////////////////
@@ -63,6 +63,7 @@ TGlobalResource<FDummyCulledDispatchVertexIdsBuffer> GDummyCulledDispatchVertexI
 /////////////////////////////////////////////////////////////////////////////////////////
 
 FHairGroupPublicData::FVertexFactoryInput ComputeHairStrandsVertexInputData(const FHairGroupInstance* Instance);
+int GetHairRaytracingProceduralSplits();
 FHairStrandsVertexFactoryUniformShaderParameters FHairGroupInstance::GetHairStandsUniformShaderParameters() const
 {
 	const FHairGroupPublicData::FVertexFactoryInput VFInput = ComputeHairStrandsVertexInputData(this);
@@ -72,6 +73,7 @@ FHairStrandsVertexFactoryUniformShaderParameters FHairGroupInstance::GetHairStan
 	Out.RootScale 					= VFInput.Strands.HairRootScale;
 	Out.TipScale 					= VFInput.Strands.HairTipScale;
 	Out.RaytracingRadiusScale		= VFInput.Strands.HairRaytracingRadiusScale;
+	Out.RaytracingProceduralSplits  = GetHairRaytracingProceduralSplits();
 	Out.Length 						= VFInput.Strands.HairLength;
 	Out.Density 					= VFInput.Strands.HairDensity;
 	Out.StableRasterization			= VFInput.Strands.bUseStableRasterization;
@@ -162,7 +164,7 @@ void FHairStrandsVertexFactory::ModifyCompilationEnvironment(const FVertexFactor
 	const bool bUseProceduralIntersection = GetSupportHairStrandsProceduralPrimitive(Parameters.Platform);
 	OutEnvironment.SetDefine(TEXT("VF_SUPPORTS_PRIMITIVE_SCENE_DATA"), bUseGPUSceneAndPrimitiveIdStream);
 	OutEnvironment.SetDefine(TEXT("HAIR_STRAND_MESH_FACTORY"), 1);
-	OutEnvironment.SetDefine(TEXT("ENABLE_CUSTOM_INTERSECTOR"), bUseProceduralIntersection);
+	OutEnvironment.SetDefine(TEXT("ENABLE_PROCEDURAL_INTERSECTOR"), bUseProceduralIntersection);
 }
 
 void FHairStrandsVertexFactory::ValidateCompiledResult(const FVertexFactoryType* Type, EShaderPlatform Platform, const FShaderParameterMap& ParameterMap, TArray<FString>& OutErrors)
@@ -259,5 +261,5 @@ IMPLEMENT_VERTEX_FACTORY_TYPE(FHairStrandsVertexFactory, "/Engine/Private/HairSt
 	| (VF_STRANDS_SUPPORT_GPU_SCENE ? EVertexFactoryFlags::SupportsPrimitiveIdStream : EVertexFactoryFlags::None)
 	| EVertexFactoryFlags::SupportsRayTracing
 	| EVertexFactoryFlags::SupportsRayTracingDynamicGeometry
-	| (VF_STRANDS_CUSTOM_INTERSECTOR ? EVertexFactoryFlags::SupportsRayTracingProceduralPrimitive : EVertexFactoryFlags::None)
+	| (VF_STRANDS_PROCEDURAL_INTERSECTOR ? EVertexFactoryFlags::SupportsRayTracingProceduralPrimitive : EVertexFactoryFlags::None)
 );
diff --git a/Engine/Plugins/Runtime/HairStrands/Source/HairStrandsCore/Public/GroomInstance.h b/Engine/Plugins/Runtime/HairStrands/Source/HairStrandsCore/Public/GroomInstance.h
--- a/Engine/Plugins/Runtime/HairStrands/Source/HairStrandsCore/Public/GroomInstance.h
+++ b/Engine/Plugins/Runtime/HairStrands/Source/HairStrandsCore/Public/GroomInstance.h
@@ -52,6 +52,7 @@ BEGIN_GLOBAL_SHADER_PARAMETER_STRUCT(FHairStrandsVertexFactoryUniformShaderParam
 	SHADER_PARAMETER(uint32, HasMaterial)
 	SHADER_PARAMETER(uint32, StableRasterization)
 	SHADER_PARAMETER(uint32, ScatterSceneLighing)
+	SHADER_PARAMETER(uint32, RaytracingProceduralSplits)
 
 	SHADER_PARAMETER_SRV(Buffer<float4>, PositionOffsetBuffer)
 	SHADER_PARAMETER_SRV(Buffer<float4>, PreviousPositionOffsetBuffer)
@@ -145,6 +146,7 @@ struct HAIRSTRANDSCORE_API FHairGroupInstance : public FHairStrandsInstance
 		float CachedHairScaledRadius = 0;
 		float CachedHairRootScale = 0;
 		float CachedHairTipScale = 0;
+		int   CachedProceduralSplits = 0;
 		#endif
 
 		FRDGExternalBuffer DebugAttributeBuffer;
diff --git a/Engine/Plugins/Runtime/HairStrands/Source/HairStrandsCore/Public/GroomResources.h b/Engine/Plugins/Runtime/HairStrands/Source/HairStrandsCore/Public/GroomResources.h
--- a/Engine/Plugins/Runtime/HairStrands/Source/HairStrandsCore/Public/GroomResources.h
+++ b/Engine/Plugins/Runtime/HairStrands/Source/HairStrandsCore/Public/GroomResources.h
@@ -11,6 +11,8 @@
 #include "HairStrandsInterface.h"
 #include "HairStrandsMeshProjection.h"
 
+#define STRANDS_PROCEDURAL_INTERSECTOR_MAX_SPLITS 4
+
 inline uint32 GetBufferTotalNumBytes(const FRDGExternalBuffer& In) 
 {
 	return In.Buffer ? In.Buffer->Desc.GetTotalNumBytes() : 0;
diff --git a/Engine/Shaders/Private/HairStrands/HairStrandsRaytracingGeometry.usf b/Engine/Shaders/Private/HairStrands/HairStrandsRaytracingGeometry.usf
--- a/Engine/Shaders/Private/HairStrands/HairStrandsRaytracingGeometry.usf
+++ b/Engine/Shaders/Private/HairStrands/HairStrandsRaytracingGeometry.usf
@@ -19,7 +19,9 @@ Buffer<float4>	 PositionOffsetBuffer;
 Buffer<uint4> 	 PositionBuffer;
 Buffer<float4> 	 TangentBuffer;
 RWBuffer<float4> OutputPositionBuffer;
-#if !PERMUTATION_PROCEDURAL_PRIMITIVE 
+#if PERMUTATION_PROCEDURAL_PRIMITIVE 
+uint             RaytracingProceduralSplits;
+#else
 RWBuffer<uint>   OutputIndexBuffer;
 #endif
 
@@ -77,7 +79,8 @@ void SetInvalidSegment(uint InVertexIndex)
 {
 #if PERMUTATION_PROCEDURAL_PRIMITIVE
 	const float NaN = asfloat(0x7FC00000);
-	SetAABB(InVertexIndex, NaN, NaN);
+	for (int i = 0; i < RaytracingProceduralSplits; i++)
+		SetAABB(RaytracingProceduralSplits * InVertexIndex + i, NaN, NaN);
 #else
 	const uint BasePoint = InVertexIndex * 4;
 	const uint TrisPerSegment = 8;
@@ -193,9 +196,15 @@ void MainCS(uint2 DispatchThreadId : SV_DispatchThreadID)
 		FHairControlPoint P1 = GetControlPoint(SourceIndex1);
 		float4 A = float4(P0.Position, P0.WorldRadius * LodRadiusScale0);
 		float4 B = float4(P1.Position, P1.WorldRadius * LodRadiusScale1);
-		float3 Lo = min(A.xyz - A.w, B.xyz - B.w);
-		float3 Hi = max(A.xyz + A.w, B.xyz + B.w);
-		SetAABB(VertexIndex0, Lo, Hi);
+		float4 SplitA = A;
+		for (int i = 0; i < RaytracingProceduralSplits; i++)
+		{
+			float4 SplitB = lerp(A, B, float(i + 1) / float(RaytracingProceduralSplits));
+			float3 Lo = min(SplitA.xyz - SplitA.w, SplitB.xyz - SplitB.w);
+			float3 Hi = max(SplitA.xyz + SplitA.w, SplitB.xyz + SplitB.w);
+			SetAABB(RaytracingProceduralSplits * VertexIndex0 + i, Lo, Hi);
+			SplitA = SplitB;
+		}
 	}
 #else
 	// expand current vertex into 4 points around the curve
@@ -215,7 +224,6 @@ void MainCS(uint2 DispatchThreadId : SV_DispatchThreadID)
 		OutputPositionBuffer[BaseOutIndex + 2] = float4(Px1, 1);
 		OutputPositionBuffer[BaseOutIndex + 3] = float4(Py1, 1);
 	}
-
 	if (P0.Type == HAIR_CONTROLPOINT_END)
 	{
 		// last vertex of a curve, we could create a cap, but this would complicate the logic in the vertex factory
diff --git a/Engine/Shaders/Private/HairStrands/HairStrandsVertexFactory.ush b/Engine/Shaders/Private/HairStrands/HairStrandsVertexFactory.ush
--- a/Engine/Shaders/Private/HairStrands/HairStrandsVertexFactory.ush
+++ b/Engine/Shaders/Private/HairStrands/HairStrandsVertexFactory.ush
@@ -381,114 +381,16 @@ float3 GetTangent(FVertexFactoryIntermediates Intermediates)
 
 #if USE_MATERIAL_INTERSECTION_SHADER
 
-#define STRAND_INTERSECTOR   1 // 0: cylinders 1: ribbons 
-
-float3 isectStrand(float3 ro, float3 rd, float3  pa, float3  pb, float ra, float rb)
-{
-#if STRAND_INTERSECTOR == 0
-	// trick to improve numerical stability - move origin closer to one of the vertices
-	float dd = dot(rd, rd);
-	float idd = rcp(dd);
-	float offset = dot(rd, pa - ro) * idd;
-	ro += offset * rd;
-
-	float3  ba = pb - pa;
-	float3  oa = ro - pa;
-	float3  ob = ro - pb;
-	float rr = ra - rb;
-	float m0 = dot(ba, ba);
-	float m1 = dot(ba, oa);
-	float m2 = dot(ba, rd);
-	float m3 = dot(rd, oa);
-	float m5 = dot(oa, oa);
-	float m6 = dot(ob, rd);
-	float m7 = dot(ob, ob);
-
-	float d2 = m0 - rr * rr;
-
-	float k2 = dd * d2 - m2 * m2;
-	float k1 = d2 * m3 - m1 * m2 + m2 * rr * ra;
-	float k0 = d2 * m5 - m1 * m1 + m1 * rr * ra * 2.0 - m0 * ra * ra;
-
-	float h = k1 * k1 - k0 * k2;
-	if (h > 0.0)
-	{
-		float t = (-sqrt(h) - k1) / k2;
-
-		float y = m1 - ra * rr + t * m2;
-		if (y > 0.0 && y < d2)
-		{
-			return float3(t + offset, y / d2, 0.5);
-		}
-	}
-
-  #if 0
-	// ignore caps: produces cracks between segments
-	return -1.0;
-  #elif 1
-	// end-point cap (only test one since hair strands share vertices (root is less important than the tip)
-	// this keeps a visually continuous curve
-	float h2 = m6 * m6 - dd * (m7 - rb * rb);
-	if (h2 > 0.0)
-	{
-		float t = (-m6 - sqrt(h2)) * idd + offset;
-		return float3(t, 1, 0.5);
-	}
-	return -1.0;
-  #else
-	// caps
-	float h1 = m3 * m3 - dd * (m5 - ra * ra);
-	float h2 = m6 * m6 - dd * (m7 - rb * rb);
-	if (max(h1, h2) < 0.0) return -1.0;
-
-	float3 r = 1e20;
-	if (h1 > 0.0)
-	{
-		float t = (-m3 - sqrt(h1)) * idd + offset;
-		r = float3(t + offset, 0, 0.5);
-	}
-	if (h2 > 0.0)
-	{
-		float t = (-m6 - sqrt(h2)) * idd + offset;
-		if (t > 0.0 && t < r.x)
-			r = float3(t, 1, 0.5);
-	}
 
-	return r;
-  #endif
-#elif STRAND_INTERSECTOR == 1
-	float3 u = pb - pa;
-	float3 w0 = pa - ro;
-	float3 w1 = pb - ro;
-	float3 n = cross(rd, u);
-	float3 n1 = cross(rd, n);
-	float sc = saturate(-dot(w0, n1) / dot(u, n1));
-	float3 cp = lerp(w0, w1, sc);
-	float rc = lerp(ra, rb, sc);
-	float t = dot(cp, rd) / dot(rd, rd);
-	float3 dp = t * rd - cp;
-	float d2 = length2(dp);
-	float r2 = rc * rc;
-	if (d2 < r2 && t > 4. * rc)
-	{
-		float h = sign(dot(dp, n)) * sqrt(d2 / r2);
-		return float3(t, sc, saturate(0.5 * h + 0.5));
-	}
-	return -1.0;
-#else
-	return -1.0;
-#endif
-}
-
-// Entry point
 RAY_TRACING_ENTRY_INTERSECTION(MaterialIS)
 {
+
 	float3 RayOrg = ObjectRayOrigin();
 	float3 RayDir = ObjectRayDirection();
 	float TMin = RayTMin();
 	float TMax = RayTCurrent();
 
-	uint VertexIndex0 = PrimitiveIndex();
+	uint VertexIndex0 = PrimitiveIndex() / HairStrandsVF.RaytracingProceduralSplits;
 	uint VertexIndex1 = VertexIndex0 + 1;
 	float ScaleR0 = 1, ScaleR1 = 1;
 	if (HairStrandsVF.CullingEnable)
@@ -503,29 +405,55 @@ RAY_TRACING_ENTRY_INTERSECTION(MaterialIS)
 	FHairControlPoint P0 = UnpackHairControlPoint(
 		HairStrandsVF.PositionBuffer[VertexIndex0],
 		PositionOffset,
-		HairStrandsVF.Radius * ScaleR0 * HairStrandsVF.RaytracingRadiusScale,
+		HairStrandsVF.Radius * ScaleR0,
 		HairStrandsVF.RootScale,
 		HairStrandsVF.TipScale);
 	FHairControlPoint P1 = UnpackHairControlPoint(
 		HairStrandsVF.PositionBuffer[VertexIndex1],
 		PositionOffset,
-		HairStrandsVF.Radius * ScaleR1 * HairStrandsVF.RaytracingRadiusScale,
+		HairStrandsVF.Radius * ScaleR1,
 		HairStrandsVF.RootScale,
 		HairStrandsVF.TipScale);
 
-	float3 Result = isectStrand(RayOrg, RayDir, P0.Position, P1.Position, P0.WorldRadius, P1.WorldRadius);
-
-	FDefaultAttributes Attributes = (FDefaultAttributes)0;
-	Attributes.Barycentrics = float2(Result.y, Result.z);
-	ReportHit(Result.x, HIT_KIND_TRIANGLE_FRONT_FACE, Attributes);
+	// ray/ribbon intersection - this is very approximate compared to a real oriented cone, but sufficient for thin strands
+	float3 pa = P0.Position;
+	float3 pb = P1.Position;
+	float ra = P0.WorldRadius;
+	float rb = P1.WorldRadius;
+	float3 ro = RayOrg;
+	float3 rd = RayDir;
+	float3 u = pb - pa;
+	float3 w0 = pa - ro;
+	float3 w1 = pb - ro;
+	float3 n = cross(rd, u);
+	float3 n1 = cross(rd, n);
+	float sc = saturate(-dot(w0, n1) / dot(u, n1));
+	float3 cp = lerp(w0, w1, sc);
+	float rc = lerp(ra, rb, sc);
+	float t = dot(cp, rd) / dot(rd, rd);
+	float3 dp = t * rd - cp;
+	float d2 = length2(dp);
+	float r2 = rc * rc;
+	if (d2 < r2 && t > 4. * rc)
+	{
+		// NOTE: the tangents are not exactly equal to 'u' at the extremeties due to quantization, could be avoided with better compression
+		float3 ta = P0.Type == HAIR_CONTROLPOINT_START ? u : HairStrandsVF.TangentBuffer[VertexIndex0 * 2 + 1].xyz;
+		float3 tb = P1.Type == HAIR_CONTROLPOINT_END   ? u : HairStrandsVF.TangentBuffer[VertexIndex1 * 2 + 1].xyz;
+		// clip the rounded portion from the end each segment
+		if ((dot(t * rd - w0, ta) * dot(ta, u) < 0)) return;
+		if ((dot(t * rd - w1, tb) * dot(tb, u) > 0)) return;
+
+		float h = sign(dot(dp, n)) * sqrt(saturate(d2 / r2));
+		FDefaultAttributes Attributes = (FDefaultAttributes)0;
+		Attributes.Barycentrics = float2(sc, 0.5 * h + 0.5);
+		ReportHit(t, HIT_KIND_TRIANGLE_FRONT_FACE, Attributes);
+	}
 }
-#endif // USE_MATERIAL_INTERSECTION_SHADER
 
-#endif // RAYHITGROUPSHADER
+#endif // USE_MATERIAL_INTERSECTION_SHADER
 
 ////////////////////////////////////////////////////////////////////////////////////////////////////////////
 // Material evaluation for raytracing
-#if RAYHITGROUPSHADER
 
 // define this to let shaders know this alternative method is available
 #define VF_SUPPORTS_RAYTRACING_PREPARE_MATERIAL_PIXEL_PARAMETERS 1 
@@ -543,7 +471,8 @@ FMaterialPixelParameters GetMaterialPixelParameters(float3 RayOrigin, float3 Ray
 
 	FSceneDataIntermediates SceneData = VF_GPUSCENE_GET_INTERMEDIATES(Input); // NOTE: Input is not used when VF_USE_PRIMITIVE_SCENE_DATA == 0
 
-	uint VertexIndex0 = HitPrimitiveIndex;
+#if ENABLE_PROCEDURAL_INTERSECTOR
+	uint VertexIndex0 = HitPrimitiveIndex / HairStrandsVF.RaytracingProceduralSplits;
 	uint VertexIndex1 = VertexIndex0 + 1;
 
 	if (HairStrandsVF.CullingEnable)
@@ -552,7 +481,6 @@ FMaterialPixelParameters GetMaterialPixelParameters(float3 RayOrigin, float3 Ray
 		VertexIndex1 = HairStrandsVF.CulledVertexIdsBuffer[VertexIndex1];
 	}
 
-#if ENABLE_CUSTOM_INTERSECTOR
 	float UCoord = HitAttributes.Barycentrics.x;
 	float VCoord = HitAttributes.Barycentrics.y;
 	const float3 PositionOffset = HairStrandsVF.PositionOffsetBuffer[0].xyz;
@@ -569,28 +497,19 @@ FMaterialPixelParameters GetMaterialPixelParameters(float3 RayOrigin, float3 Ray
 		HairStrandsVF.RootScale,
 		HairStrandsVF.TipScale);
 
-  #if STRAND_INTERSECTOR == 0
-	float3 WorldPosition = RayOrigin + HitT * RayDirection;
-	float4 LocalPosition = mul(float4(WorldPosition, 1), LWCHackToFloat(SceneData.Primitive.WorldToLocal));
-	float3 LocalNormal = normalize(LocalPosition.xyz - lerp(P0.Position, P1.Position, UCoord));
-	// TODO: does this support non-uniform scales correctly?
-	WorldGeoNormal = normalize(RotateLocalToWorld(LocalNormal, SceneData.Primitive.LocalToWorld, SceneData.Primitive.InvNonUniformScale));
-  #endif
-
-	Result.HairPrimitiveId = VertexIndex0;
-#else // ENABLE_CUSTOM_INTERSECTOR
+#else
 	const float3 Weights = float3(
 		1 - HitAttributes.Barycentrics.x - HitAttributes.Barycentrics.y,
 		HitAttributes.Barycentrics.x,
 		HitAttributes.Barycentrics.y);
 
 	// fetch triangle
-	FTriangleBaseAttributes TriangleBase = LoadTriangleBaseAttributes(HitPrimitiveIndex);
+	FTriangleBaseAttributes Triangle = LoadTriangleBaseAttributes(HitPrimitiveIndex);
 	// Transform vertices into world space
 	const float3 WorldPositions[3] = {
-		CalcWorldPosition(TriangleBase.LocalPositions[0], SceneData.Primitive.LocalToWorld).xyz,
-		CalcWorldPosition(TriangleBase.LocalPositions[1], SceneData.Primitive.LocalToWorld).xyz,
-		CalcWorldPosition(TriangleBase.LocalPositions[2], SceneData.Primitive.LocalToWorld).xyz
+		CalcWorldPosition(Triangle.LocalPositions[0], SceneData.Primitive.LocalToWorld).xyz,
+		CalcWorldPosition(Triangle.LocalPositions[1], SceneData.Primitive.LocalToWorld).xyz,
+		CalcWorldPosition(Triangle.LocalPositions[2], SceneData.Primitive.LocalToWorld).xyz
 	};
 	// Compute triangle normal
 	{
@@ -603,13 +522,54 @@ FMaterialPixelParameters GetMaterialPixelParameters(float3 RayOrigin, float3 Ray
 
 	// NOTE: This must match the tesselation pattern produced by SetBodySegment() in HairStrandsRaytracingGeometry.usf
 	// TODO: is there a way we could handle end-caps elegantly?
-	Result.HairPrimitiveId = TriangleBase.Indices.z / 4; // last index always belongs to the "bottom" vertex of the hair segment
+	uint VertexIndex0 = Triangle.Indices.z / 4; // last index always belongs to the "bottom" vertex of the hair segment
+	uint VertexIndex1 = VertexIndex0 + 1;
+
+	if (HairStrandsVF.CullingEnable)
+	{
+		VertexIndex0 = HairStrandsVF.CulledVertexIdsBuffer[VertexIndex0];
+		VertexIndex1 = HairStrandsVF.CulledVertexIdsBuffer[VertexIndex1];
+	}
 
 	float UCoord = dot(float3(1, float((HitPrimitiveIndex % 2) == 0), 0), Weights);
-	float VCoord = 0.5;// can't measure position across the width easily
-#endif // ENABLE_CUSTOM_INTERSECTOR
+	float VCoord = 0;
+	{
+		// find central point of the hair strand
+		const float3 PositionOffset = HairStrandsVF.PositionOffsetBuffer[0].xyz;
+		FHairControlPoint P0 = UnpackHairControlPoint(
+			HairStrandsVF.PositionBuffer[VertexIndex0],
+			PositionOffset,
+			HairStrandsVF.Radius,
+			HairStrandsVF.RootScale,
+			HairStrandsVF.TipScale);
+		FHairControlPoint P1 = UnpackHairControlPoint(
+			HairStrandsVF.PositionBuffer[VertexIndex1],
+			PositionOffset,
+			HairStrandsVF.Radius,
+			HairStrandsVF.RootScale,
+			HairStrandsVF.TipScale);
+		float4 A = float4(P0.Position, P0.WorldRadius);
+		float4 B = float4(P1.Position, P1.WorldRadius);
+		float4 C = lerp(A, B, UCoord);
+		// approximate position along the width of the ribbon
+		float3 WorldC = LWCToFloat(TransformLocalToWorld(C.xyz, SceneData.Primitive.LocalToWorld));
+		float3 rd = RayDirection;
+		float3 cp = WorldC - RayOrigin;
+		float t = dot(cp, rd) / dot(rd, rd);
+		float3 dp = t * RayDirection - cp;
+		float d2 = length2(dp);
+		float v = saturate(sqrt(d2 / Pow2(C.w))); // TODO: account for scale of radius
+		// figure out sign of v
+		float3 WorldU = LWCToFloat(TransformLocalToWorld(B.xyz - A.xyz, SceneData.Primitive.LocalToWorld));
+		VCoord = v * sign(dot(dp, cross(RayDirection, WorldU)));
+		VCoord = 0.5 * VCoord + 0.5;
+	}
+
 
-	Result.HairPrimitiveUV = float2(UCoord, 0.5); // NOTE: don't bother defining a v coordinate since we use a cylinder for RT
+
+#endif // ENABLE_PROCEDURAL_INTERSECTOR
+	Result.HairPrimitiveId = VertexIndex0;
+	Result.HairPrimitiveUV = float2(UCoord, VCoord);
 
 	// Required for previewing materials that use ParticleColor
 	Result.Particle.Color = half4(1, 1, 1, 1);
@@ -624,15 +584,12 @@ FMaterialPixelParameters GetMaterialPixelParameters(float3 RayOrigin, float3 Ray
 
 	Result.TangentToWorld = AssembleTangentToWorld(TangentToWorld0, float4(TangentToWorld2, TangentSign));
 
-#if STRAND_INTERSECTOR == 1
 	// compute a smooth ribbon normal based on the tangent vector
 	float3 T = Result.TangentToWorld[2];
 	float3 D = RayDirection;
 	WorldGeoNormal = normalize(cross(T, normalize(cross(D, T))));
-	WorldGeoNormal *= dot(RayDirection, WorldGeoNormal) > 0 ? -1.0 : 1.0;
-#endif
 
-	Result.TwoSidedSign = 1;
+	Result.TwoSidedSign = sign(dot(RayDirection, WorldGeoNormal));
 
 	Result.PrimitiveId = SceneData.PrimitiveId;
 
diff --git a/Engine/Shaders/Private/PathTracing/Material/PathTracingHair.ush b/Engine/Shaders/Private/PathTracing/Material/PathTracingHair.ush
new file mode 100644
--- /dev/null
+++ b/Engine/Shaders/Private/PathTracing/Material/PathTracingHair.ush
@@ -0,0 +1,414 @@
+// Copyright Epic Games, Inc. All Rights Reserved.
+
+#pragma once
+
+
+#define PATH_TRACER_HAIR_SHADING_MODEL				1 // 0: zinke diffuse (a simple, cheap model for performance testing)
+                                                      // 1: chiang principled hair (a sophisticated model with specular reflection/transmission)
+
+
+// Shared routines
+
+float3x3 GetHairBasis(float3 T_World, float3 V_World)
+{
+	// a convenient basis for hair shading calculations:
+	//   Z Axis = Hair Tangent
+	//   Y Axis = point toward the viewer
+	//   X Axis = left/right from viewing direction
+	float3 Tx = normalize(cross(T_World, V_World));
+	float3 Ty = cross(Tx, T_World);
+	return float3x3(Tx, Ty, T_World);
+}
+
+
+#if PATH_TRACER_HAIR_SHADING_MODEL == 0
+
+float ZinkeHairDiffuse(float3 L)
+{
+	// Evaluate the zinke diffuse lobe model. This is the far field average of a diffuse (lambertian) cylinder
+	// This was originally derived in: "Light Scattering from Filaments" - Arno Zinke and Andreas Weber - 2007
+	// The exact form was mentioned in: "Importance Sampling for Physically-Based Hair Fiber Models" - Eugene d’Eon, [REDACTED_PERSON], [REDACTED_PERSON] - 2013
+	// 
+	const float CosTheta = sqrt(saturate(1 - L.z * L.z));
+#if 0
+	// exact form - Eq(9) from d'Eon et al.
+	const float CosPhi = clamp(L.y / CosTheta, -1.0, 1.0);
+	const float SinPhi = sqrt(1.0 - CosPhi * CosPhi);
+	const float Phi = acos(CosPhi);
+	return ((PI - Phi) * CosPhi + SinPhi) * CosTheta / (4 * PI);
+#else
+	// optimized fit to avoid acos
+	// error is around ~0.001 
+	const float CosPhiCosTheta = L.y;
+	return 0.0786133 * CosTheta + CosPhiCosTheta * (0.125 + 0.0463867 * CosPhiCosTheta * rcp(CosTheta));
+#endif
+}
+
+FMaterialSample Hair_SampleMaterial(
+	float3 RayDirection,
+	FPathTracingPayload Payload,
+	float4 RandSample)
+{
+	float3x3 HairBasis = GetHairBasis(Payload.WorldTangent, -RayDirection);
+
+	// sample a random position across the width of the fiber
+	float h = 2 * RandSample.x - 1;
+	// create a normal vector in the plane tangent to the fiber
+	float2 N = float2(h, sqrt(1 - h * h));
+	// sample a cosine weighted direction
+	float4 SampledValue = CosineSampleHemisphere(RandSample.yz);
+	// rotate this sampled direction about the (random) normal vector
+	float3 L = float3(SampledValue.x * N.y + SampledValue.z * N.x,
+					  SampledValue.z * N.y - SampledValue.x * N.x,
+					  SampledValue.y);
+	// finally transform into world space
+	float3 L_World = mul(L, HairBasis);
+
+	float Pdf = ZinkeHairDiffuse(L);
+	return CreateMaterialSample(L_World, Payload.BaseColor, Pdf, sign(dot(Payload.WorldGeoNormal, L_World)), 1.0);
+}
+
+FMaterialEval Hair_EvalMaterial(
+	float3 IncomingDirection,
+	float3 OutgoingDirection,
+	FPathTracingPayload Payload
+)
+{
+	float3x3 HairBasis = GetHairBasis(Payload.WorldTangent, -IncomingDirection);
+
+	float3 L_World = OutgoingDirection;
+	float3 L = mul(HairBasis, L_World);
+
+	float Pdf = ZinkeHairDiffuse(L);
+	return CreateMaterialEval(Payload.BaseColor, Pdf);
+}
+
+
+#elif PATH_TRACER_HAIR_SHADING_MODEL == 1
+
+// The implementation below is based on the following paper:
+// "A Practical and Controllable Hair and Fur Model for Production Path Tracing" 
+// Matt Jen-Yuan Chiang, Benedikt Bitterli, Chuck Tappan, Brent Burley
+// EGSR 2016
+// https://media.disneyanimation.com/uploads/production/publication_asset/152/asset/eurographics2016Fur_Smaller.pdf
+
+
+float LogI0(float x)
+{
+	// I0 grows too quickly to evaluate robustly in single precision arithmetic
+	// instead evaluate the log which is much smoother
+	// The fits below were found using Mathematica and have error below 10^-6
+	x = abs(x);
+	if (x < 16.315929055555998) // transition point where the two approximations line up
+	{
+		// rational approximation
+		// 9 fmas + 2 mults + 1 div
+		float num = 0.249915 + x * (0.124919 + x * (0.0486428 + x * (0.011457 + x * 0.000647177)));
+		float den = 1 + x * (0.498089 + x * (0.259956 + x * (0.0741939 + x * (0.0134601 + x * 0.000650605))));
+		return x * x * (num / den);
+	}
+	else
+	{
+		// beyond this point the curve is mostly linear with a very slight offset
+		// 3 fmas + 1 log
+		// TODO: would it be better to use log2 directly and save one mult? can the compiler do this for us?
+		return x - 0.25 * log(-5.68881 + x * (-19.7156 + 39.4779 * x));
+	}
+}
+
+float FastASin(float x) {
+	// Maximum error: 10^-3.35
+	// TODO: consider replacing acosFast/asinFast which have approximation error 10^-2.0 (with one less fma, but also with a slightly sub-optimal constant)
+	float a = abs(x);
+	float r = (0.5 * PI) + x * (-0.207034 + 0.0531013 * x);
+	r = (0.5 * PI) - r * sqrt(1 - a);
+	return x >= 0 ? r : -r;
+}
+
+float Mp(float SinThetaI, float SinThetaO, float v)
+{
+	float rv = rcp(v);
+	float CosThetaO2 = 1 - SinThetaO * SinThetaO;
+	float CosThetaI2 = 1 - SinThetaI * SinThetaI;
+	float CosProduct = sqrt(CosThetaO2 * CosThetaI2);
+	// "Path tracing in Production" - Siggraph 2018 Course
+	// "From Bricks to Bunnies - Adapting a Raytracer to New Requirements" - Luke Emrose
+	// https://jo.dreggn.org/path-tracing-in-production/2018/course-notes.pdf
+	// Simplified expression for better numerical stability - Eq(62)
+	return exp(LogI0(CosProduct * rv) - rv * (SinThetaI * SinThetaO + 1) + log(rv) - log(1 - exp(-2 * rv)));
+}
+
+float TrimmedLogistic(float x, float InvS)
+{
+	// Logistic distribution, re-normalized to have bounds [-Pi,Pi]
+	// The normalization constant below ensures we integrate to 1.0
+	float e = exp(-abs(x * InvS));
+	float a = exp(-PI * InvS); // normalization constant
+	return e * InvS * (1 + a) / ((1 - a) * Pow2(1 + e));
+}
+
+float SampleTrimmedLogistic(float x, float s)
+{
+	float a = exp(-PI * rcp(s));
+	float k = (1 - a) / (1 + a);
+	float r = 2 * x - 1;
+	return s * log((1 + r * k) / (1 - r * k));
+}
+
+float Phi(int p, float GammaO, float GammaT)
+{
+	return 2 * p * GammaT - 2 * GammaO + p * PI;
+}
+
+float Np(float phi, int p, float InvS, float GammaO, float GammaT) {
+	// center around chosen angle
+	float dphi = phi - Phi(p, GammaO, GammaT);
+	// wrap to interval (-PI,PI)
+	dphi /= 2 * PI;
+	dphi -= round(dphi);
+	dphi *= 2 * PI;
+	return TrimmedLogistic(dphi, InvS);
+}
+
+// returns sine from cosine (or vice-versa)
+float TrigInverse(float x)
+{
+	return sqrt(saturate(1 - x * x));
+}
+
+// precomputed terms for evaluating/sampling the hair bxdf
+// TODO: can we somehow do this stuff only once instead of repeating it in eval and sample?
+struct FHairData {
+	// trig that only depends on viewing direction and offset h
+	float4 SinThetaO; // Viewing angles (adjusted by scale tilt per lobe)
+	float GammaO;
+	float GammaT;
+
+	float  v0;    // roughness parameter for Mp
+	float  v1;
+	float  v2;
+	float  v3;
+	float  s;     // azimuthal roughness parameter for Np
+	float  A0;    // attenuation term per lobe (first one is always white)
+	float3 A1;
+	float3 A2;
+	float3 A3;
+	float4 LobePdf;
+	float4 LobeCdf; // Normalized sum to help select a lobe
+};
+
+FHairData PrepareHairData(float3 V, float3 Color, float LongitudinalRoughness, float AzimuthalRoughness, float h, float Specular)
+{
+	FHairData Result = (FHairData)0;
+	// Longitudinal roughness
+	float Bm = max(LongitudinalRoughness, 0.01); // avoid tiny values
+	float Bm2 = Bm * Bm;
+	float Bm4 = Bm2 * Bm2;
+	float Bm8 = Bm4 * Bm4;
+	float Bm10 = Bm8 * Bm2;
+	float Bm20 = Bm10 * Bm10;
+	// Chiang et al. - Eq(7)
+	Result.v0 = Pow2(0.726 * Bm + 0.812 * Bm2 + 3.7 * Bm20);
+	Result.v1 = 0.25 * Result.v0;
+	Result.v2 = 4.00 * Result.v0;
+	Result.v3 = Result.v2;
+
+	// Azimuthal roughness
+	const float SqrtPiOver8 = 0.626657069;
+	float Bn = max(AzimuthalRoughness, 0.01);
+	float Bn2 = Bn * Bn;
+	float Bn4 = Bn2 * Bn2;
+	float Bn8 = Bn4 * Bn4;
+	float Bn11 = Bn8 * Bn2 * Bn;
+	float Bn22 = Bn11 * Bn11;
+	// Chiang et al. - Eq(8) + Eq(12)
+	Result.s = SqrtPiOver8 * (0.265 * Bn + 1.194 * Bn2 + 5.372 * Bn22);
+
+	// Specular to Eta
+	float F0 = Specular * 0.08;
+	float F90 = saturate(Specular * (50.0 * 0.08));
+	float Eta = (1 + sqrt(F0)) / (1 - sqrt(F0));
+	
+	// setup main angles
+	float SinThetaO = V.z;
+	float CosThetaO = TrigInverse(SinThetaO);
+
+	// TODO: should this angle be exposed as a parameter?
+	// Using value from HairBsdf.ush for now
+	float ScaleAngle = 0.035; // about 2 degrees
+	float3 SinScale; // sin(ScaleAngle * 2^k) for k=0,1,2
+	float3 CosScale; // cos(ScaleAngle * 2^k) for k=0,1,2
+	SinScale.x = sin(ScaleAngle);
+	CosScale.x = TrigInverse(SinScale.x);
+	SinScale.y = 2 * CosScale.x * SinScale.x;
+	CosScale.y = Pow2(CosScale.x) - Pow2(SinScale.x);
+	SinScale.z = 2 * CosScale.y * SinScale.y;
+	CosScale.z = Pow2(CosScale.y) - Pow2(SinScale.y);
+
+	// rotate by scale angle (different amount of each lobe)
+	Result.SinThetaO.x = SinThetaO * CosScale.y - CosThetaO * SinScale.y; 	// R: -2*angle
+	Result.SinThetaO.y = SinThetaO * CosScale.x + CosThetaO * SinScale.x;	// TT: +angle
+	Result.SinThetaO.z = SinThetaO * CosScale.z + CosThetaO * SinScale.z;	// TRT: +4*angle
+	Result.SinThetaO.w = SinThetaO;	// TR*T: 0 angle
+
+	// Account for refraction through the cylinder
+	float SinGammaT = h * CosThetaO * rsqrt(Pow2(Eta) - Pow2(SinThetaO));
+	float CosGammaT2 = 1 - SinGammaT * SinGammaT;
+	Result.GammaO = FastASin(h);
+	Result.GammaT = FastASin(SinGammaT);
+
+	float SinThetaT = SinThetaO / Eta;
+	float CosThetaT2 = 1 - SinThetaT * SinThetaT;
+
+	// Tinting factors through the hair
+	float CosGammaO = sqrt(saturate(1 - h * h));
+	float CosTheta = CosThetaO * CosGammaO;
+	float F = lerp(F0, F90, Pow5(1 - CosTheta)); // Schlick approximation
+
+	// Chiang et al. - Eq(9)
+	float3 SigmaA = Pow2(log(max(Color, 1e-4)) / ((((((0.245 * Bn) + 5.574) * Bn - 10.73) * Bn + 2.532) * Bn - 0.215) * Bn + 5.969));
+	
+	// Zinke 2007 - Eq(20)
+	float3 T = exp(-SigmaA * (2 * sqrt(CosGammaT2 / CosThetaT2)));
+
+	// d'Eon et al. - Eq(12) + Eq(13)
+	Result.A0 = F;
+	Result.A1 = Pow2(1 - F) * T;
+	Result.A2 = Result.A1 * F * T;
+	// Chiang et al. - Eq(6)
+	Result.A3 = Result.A2 * F * T * rcp(1.00001 - F * T); // avoid dividing by 0 when F*T=1
+
+	Result.LobePdf = float4(
+		Result.A0,
+		max3(Result.A1.x, Result.A1.y, Result.A1.z),
+		max3(Result.A2.x, Result.A2.y, Result.A2.z),
+		max3(Result.A3.x, Result.A3.y, Result.A3.z)
+	);
+
+	Result.LobeCdf.x = Result.LobePdf.x;
+	Result.LobeCdf.y = Result.LobeCdf.x + Result.LobePdf.y;
+	Result.LobeCdf.z = Result.LobeCdf.y + Result.LobePdf.z;
+	Result.LobeCdf.w = Result.LobeCdf.z + Result.LobePdf.w;
+
+	Result.LobePdf /= Result.LobeCdf.w + 1e-6;
+
+	return Result;
+}
+
+FMaterialSample Hair_SampleMaterial(
+	float3 RayDirection,
+	FPathTracingPayload Payload,
+	float4 RandSample)
+{
+	float3 V_World = -RayDirection;
+	float3x3 HairBasis = GetHairBasis(Payload.WorldTangent, V_World);
+
+	float3 V = mul(HairBasis, V_World); // TODO: simplify this since HairBasis is already sort of aligned with V ?
+	float h = 2 * Payload.CustomData.y - 1; // remap back to [-1,1]
+	FHairData HairData = PrepareHairData(V, Payload.BaseColor, Payload.Roughness, 1.0, h, Payload.Specular);
+
+	// 1) Decide which lobe to sample
+	float u = RandSample.z * HairData.LobeCdf.w;
+	int p = 0;
+	float v = 0, SinThetaOp = 0;
+	if (u < HairData.LobeCdf.x)
+	{
+		u /= HairData.LobeCdf.x;
+		p = 0;
+		v = HairData.v0;
+		SinThetaOp = HairData.SinThetaO.x;
+	}
+	else if (u < HairData.LobeCdf.y)
+	{
+		u -= HairData.LobeCdf.x;
+		u /= HairData.LobeCdf.y - HairData.LobeCdf.x;
+		p = 1;
+		v = HairData.v1;
+		SinThetaOp = HairData.SinThetaO.y;
+	}
+	else if (u < HairData.LobeCdf.z)
+	{
+		u -= HairData.LobeCdf.y;
+		u /= HairData.LobeCdf.z - HairData.LobeCdf.y;
+		p = 2;
+		v = HairData.v2;
+		SinThetaOp = HairData.SinThetaO.z;
+	}
+	else
+	{
+		u -= HairData.LobeCdf.z;
+		u /= HairData.LobeCdf.w - HairData.LobeCdf.z;
+		p = 3;
+		v = HairData.v3;
+		SinThetaOp = HairData.SinThetaO.w;
+	}
+	u = saturate(u);
+
+	// 2) Sample longitudinal angle
+	float CosTheta = 1 + v * log(RandSample.x + (1 - RandSample.x) * exp(-2 / v));
+	float CosPhi = cos(2 * PI * RandSample.y);
+	float SinThetaI = -CosTheta * SinThetaOp + sqrt((1 - CosTheta * CosTheta) * (1 - SinThetaOp * SinThetaOp)) * CosPhi;
+	float CosThetaI = TrigInverse(SinThetaI);
+
+	// 3) Sample Azimuthal angle
+	float PhiI = (p < 3) ? Phi(p, HairData.GammaO, HairData.GammaT) + SampleTrimmedLogistic(u, HairData.s) : 2 * PI * u;
+
+	float3 L = float3(CosThetaI * sin(PhiI), CosThetaI * cos(PhiI), SinThetaI);
+
+	float3 Weight = 0;
+	float Pdf = 0;
+	float InvS = rcp(HairData.s);
+	float Pdf0 = Mp(SinThetaI, HairData.SinThetaO.x, HairData.v0) * Np(PhiI, 0, InvS, HairData.GammaO, HairData.GammaT);
+	float Pdf1 = Mp(SinThetaI, HairData.SinThetaO.y, HairData.v1) * Np(PhiI, 1, InvS, HairData.GammaO, HairData.GammaT);
+	float Pdf2 = Mp(SinThetaI, HairData.SinThetaO.z, HairData.v2) * Np(PhiI, 2, InvS, HairData.GammaO, HairData.GammaT);
+	float Pdf3 = Mp(SinThetaI, HairData.SinThetaO.w, HairData.v3) * (1 / (2 * PI)); // last lobe is isotropic
+
+	AddLobeWithMIS(Weight, Pdf, HairData.A0, Pdf0, HairData.LobePdf.x);
+	AddLobeWithMIS(Weight, Pdf, HairData.A1, Pdf1, HairData.LobePdf.y);
+	AddLobeWithMIS(Weight, Pdf, HairData.A2, Pdf2, HairData.LobePdf.z);
+	AddLobeWithMIS(Weight, Pdf, HairData.A3, Pdf3, HairData.LobePdf.w);
+
+	// finally transform into world space
+	float3 L_World = mul(L, HairBasis);
+
+	return CreateMaterialSample(L_World, Weight, Pdf, sign(dot(Payload.WorldGeoNormal, L_World)), Payload.Roughness);
+}
+
+FMaterialEval Hair_EvalMaterial(
+	float3 IncomingDirection,
+	float3 OutgoingDirection,
+	FPathTracingPayload Payload
+)
+{
+	float3 V_World = -IncomingDirection;
+	float3x3 HairBasis = GetHairBasis(Payload.WorldTangent, V_World);
+
+	float3 L_World = OutgoingDirection;
+	float3 V = mul(HairBasis, V_World);
+	float3 L = mul(HairBasis, L_World);
+
+	float h = 2 * Payload.CustomData.y - 1; // remap back to [-1,1]
+	FHairData HairData = PrepareHairData(V, Payload.BaseColor, Payload.Roughness, 1.0, h, Payload.Specular);
+
+	// precompute data based of lighting direction
+	float SinThetaI = L.z;
+	float CosThetaI = TrigInverse(SinThetaI);
+	float PhiI = atan2(L.x, L.y); // NOTE: in our coordinate system, PhiO == 0, so we only need to account for L
+
+	FMaterialEval Result = NullMaterialEval();
+
+	float InvS = rcp(HairData.s);
+	float Pdf0 = Mp(SinThetaI, HairData.SinThetaO.x, HairData.v0) * Np(PhiI, 0, InvS, HairData.GammaO, HairData.GammaT);
+	float Pdf1 = Mp(SinThetaI, HairData.SinThetaO.y, HairData.v1) * Np(PhiI, 1, InvS, HairData.GammaO, HairData.GammaT);
+	float Pdf2 = Mp(SinThetaI, HairData.SinThetaO.z, HairData.v2) * Np(PhiI, 2, InvS, HairData.GammaO, HairData.GammaT);
+	float Pdf3 = Mp(SinThetaI, HairData.SinThetaO.w, HairData.v3) * (1 / (2 * PI)); // last lobe is isotropic
+
+	AddLobeWithMIS(Result.Weight, Result.Pdf, HairData.A0, Pdf0, HairData.LobePdf.x);
+	AddLobeWithMIS(Result.Weight, Result.Pdf, HairData.A1, Pdf1, HairData.LobePdf.y);
+	AddLobeWithMIS(Result.Weight, Result.Pdf, HairData.A2, Pdf2, HairData.LobePdf.z);
+	AddLobeWithMIS(Result.Weight, Result.Pdf, HairData.A3, Pdf3, HairData.LobePdf.w);
+
+	return Result;
+}
+
+#endif
diff --git a/Engine/Shaders/Private/PathTracing/Material/PathTracingMaterialSampling.ush b/Engine/Shaders/Private/PathTracing/Material/PathTracingMaterialSampling.ush
--- a/Engine/Shaders/Private/PathTracing/Material/PathTracingMaterialSampling.ush
+++ b/Engine/Shaders/Private/PathTracing/Material/PathTracingMaterialSampling.ush
@@ -21,6 +21,7 @@ PathTracingMaterialSampling.usf: Material sampling functions
 #include "PathTracingCloth.ush"
 #include "PathTracingSubsurfaceProfile.ush"
 #include "PathTracingEye.ush"
+#include "PathTracingHair.ush"
 #include "../../BurleyNormalizedSSSCommon.ush"
 
 // Simplified material shaders for GPU Lightmass, supporting Lambert, Two Sided Foliage and Thin Translucent
@@ -154,6 +155,8 @@ float3 EstimateMaterialAlbedo(FPathTracingPayload Payload)
 		case SHADINGMODELID_CLOTH:
 			Albedo = BrdfWeight * (Payload.DiffuseColor + lerp(Payload.SpecularColor, Payload.CustomData.xyz, Payload.CustomData.w));
 			break;
+		case SHADINGMODELID_HAIR:
+			Albedo = BrdfWeight * Payload.BaseColor;
 #endif
 		default:
 			Albedo = BrdfWeight * Payload.BaseColor;
@@ -276,6 +279,7 @@ void AdjustMaterialRoughness(inout FPathTracingPayload Payload, float PathRoughn
 		case SHADINGMODELID_SUBSURFACE:
 		case SHADINGMODELID_PREINTEGRATED_SKIN:
 		case SHADINGMODELID_EYE:
+		case SHADINGMODELID_HAIR:
 			Payload.Roughness = max(Payload.Roughness, PathRoughness);
 			break;
 		case SHADINGMODELID_SUBSURFACE_PROFILE:
@@ -533,6 +537,9 @@ FMaterialSample SampleMaterial(
 		case SHADINGMODELID_EYE:
 			Result = Eye_SampleMaterial(RayDirection, Payload, RandSample);
 			break;
+		case SHADINGMODELID_HAIR:
+			Result = Hair_SampleMaterial(RayDirection, Payload, RandSample);
+			break;
 		default:
 			Result = Lambert_SampleMaterial(Payload, RandSample);
 			break;
@@ -629,6 +636,9 @@ FMaterialEval EvalMaterial(
 		case SHADINGMODELID_EYE:
 			Result = Eye_EvalMaterial(IncomingDirection, OutgoingDirection, Payload);
 			break;
+		case SHADINGMODELID_HAIR:
+			Result = Hair_EvalMaterial(IncomingDirection, OutgoingDirection, Payload);
+			break;
 		default:
 			Result = Lambert_EvalMaterial(OutgoingDirection, Payload);
 			break;
diff --git a/Engine/Shaders/Private/PathTracing/PathTracingCommon.ush b/Engine/Shaders/Private/PathTracing/PathTracingCommon.ush
--- a/Engine/Shaders/Private/PathTracing/PathTracingCommon.ush
+++ b/Engine/Shaders/Private/PathTracing/PathTracingCommon.ush
@@ -172,6 +172,7 @@ struct FPathTracingPayload : FMinimalPayload
 	{
 		return ShadingModelID == SHADINGMODELID_TWOSIDED_FOLIAGE ||
 			   ShadingModelID == SHADINGMODELID_THIN_TRANSLUCENT || 
+			   ShadingModelID == SHADINGMODELID_HAIR ||
 			   IsMaterialSolidGlass();
 	}
 
diff --git a/Engine/Shaders/Private/PathTracing/PathTracingCore.ush b/Engine/Shaders/Private/PathTracing/PathTracingCore.ush
--- a/Engine/Shaders/Private/PathTracing/PathTracingCore.ush
+++ b/Engine/Shaders/Private/PathTracing/PathTracingCore.ush
@@ -235,7 +235,7 @@ FPathTracingPayload TraceTransparentRay(inout FPathState PathState, bool bIsCame
 float3 TraceTransparentVisibilityRay(RayDesc Ray, uint2 PixelCoord, float PathRoughness)
 {
 	const uint RayFlags = RAY_FLAG_ACCEPT_FIRST_HIT_AND_END_SEARCH | RAY_FLAG_SKIP_CLOSEST_HIT_SHADER;
-	const uint InstanceInclusionMask = RAY_TRACING_MASK_SHADOW;
+	const uint InstanceInclusionMask = RAY_TRACING_MASK_SHADOW | RAY_TRACING_MASK_THIN_SHADOW;
 	const uint RayContributionToHitGroupIndex = RAY_TRACING_SHADER_SLOT_SHADOW;
 	const uint MultiplierForGeometryContributionToShaderIndex = RAY_TRACING_NUM_SHADER_SLOTS;
 	const uint MissShaderIndex = 0;
@@ -710,7 +710,7 @@ bool PathTracingKernel(inout FPathState PathState, int Bounce)
 	Payload.WorldSmoothNormal = AdjustShadingNormal(Payload.WorldSmoothNormal, Payload.WorldGeoNormal, PathState.Ray.Direction);
 
 #if 0
-	PathState.Radiance = 0.5 * Payload.WorldNormal + 0.5;
+	PathState.Radiance = 0.5 * Payload.WorldGeoNormal + 0.5;
 	return false;
 #endif
 
@@ -800,13 +800,13 @@ bool PathTracingKernel(inout FPathState PathState, int Bounce)
 				{
 					// for transmissive materials, bias the position to the other side of the surface if the light is coming from behind
 					const float SignedPositionBias = IsTransmissiveMaterial ? sign(dot(Payload.WorldGeoNormal, LightSample.Direction)) : 1.0;
-
+					const bool bEnableShadowTermFix = Payload.ShadingModelID != SHADINGMODELID_HAIR;
 					RayDesc LightRay;
 					LightRay.Origin = WorldPos;
 					LightRay.TMin = 0;
 					LightRay.Direction = LightSample.Direction;
 					LightRay.TMax = LightSample.Distance;
-					ApplyRayBias(LightRay, Payload.HitT, SignedPositionBias * Payload.WorldGeoNormal, true);
+					ApplyRayBias(LightRay, Payload.HitT, SignedPositionBias * Payload.WorldGeoNormal, bEnableShadowTermFix);
 
 					float AvgRoughness = ApproximateCaustics ? GetAverageRoughness(Payload) : 0.0;
 
diff --git a/Engine/Shaders/Private/PathTracing/PathTracingMaterialHitShader.usf b/Engine/Shaders/Private/PathTracing/PathTracingMaterialHitShader.usf
--- a/Engine/Shaders/Private/PathTracing/PathTracingMaterialHitShader.usf
+++ b/Engine/Shaders/Private/PathTracing/PathTracingMaterialHitShader.usf
@@ -42,14 +42,22 @@ RAY_TRACING_ENTRY_CLOSEST_HIT(PathTracingMaterialCHS,
 
 	ResolvedView = ResolveView();
 
-	FVertexFactoryInterpolantsVSToPS Interpolants;
 	float4 SvPosition = float4(PixelCoord.xy, 0.0, 1.0);
 
 	FRayCone PropagatedCone = PropagateRayCone(PackedPayload.GetRayCone(), 0 /* surface curvature */, RayTCurrent());
+
+#if VF_SUPPORTS_RAYTRACING_PREPARE_MATERIAL_PIXEL_PARAMETERS
+	// this is a newer codepath that is both more flexible and allows for more direct calculation compared to the other codepath
+	// TODO: implement such a method for all vertex factories
+	float3 GeoNormal = 0;
+	FMaterialPixelParameters MaterialParameters = GetMaterialPixelParameters(WorldRayOrigin(), WorldRayDirection(), RayTCurrent(), PrimitiveIndex(), Attributes, HitKind(), GeoNormal);
+#else
+	FVertexFactoryInterpolantsVSToPS Interpolants;
 	float3 GeoNormal = 0;
 	CalcInterpolants(PropagatedCone, Attributes, Interpolants, GeoNormal);
 
 	FMaterialPixelParameters MaterialParameters = GetMaterialPixelParameters(Interpolants, SvPosition);
+#endif
 	FPixelMaterialInputs PixelMaterialInputs;
 
 	{
@@ -146,14 +154,15 @@ RAY_TRACING_ENTRY_CLOSEST_HIT(PathTracingMaterialCHS,
 	Payload.Opacity = DitheredOpacityMaskToOpacity(GetMaterialMask(PixelMaterialInputs));
 	Payload.BlendingMode = RAY_TRACING_BLEND_MODE_TRANSLUCENT;
 #else
-	// regular masked mode - binary decision
-	if (GetMaterialMask(PixelMaterialInputs) < 0)
-	{
-		Payload.Opacity = 0;
-	}
+	// regular masked mode - accounted for in AHS
 #endif
 #endif // MATERIALBLENDING_MASKED
 
+#if HAIR_STRAND_MESH_FACTORY
+	Payload.WorldSmoothNormal = Payload.WorldNormal = Payload.WorldGeoNormal;
+	Payload.WorldTangent = MaterialParameters.TangentToWorld[2];
+#endif
+
 	Payload.Ior = Refraction == 0 ? 1.0 : Refraction;
 
 #if MATERIAL_SHADINGMODEL_CLEAR_COAT
@@ -207,6 +216,19 @@ RAY_TRACING_ENTRY_CLOSEST_HIT(PathTracingMaterialCHS,
 	else
 #endif
 
+#if MATERIAL_SHADINGMODEL_HAIR
+	if (Payload.ShadingModelID == SHADINGMODELID_HAIR)
+	{
+		// TODO: encode hair specific info
+#if HAIR_STRAND_MESH_FACTORY
+		Payload.CustomData.xy = MaterialParameters.HairPrimitiveUV;
+#else
+		Payload.CustomData.xy = 0.5;
+#endif
+	}
+	else
+#endif
+
 #if MATERIAL_SHADINGMODEL_EYE
 	if (Payload.ShadingModelID == SHADINGMODELID_EYE)
 	{
diff --git a/Engine/Source/Runtime/Renderer/Private/PathTracing.cpp b/Engine/Source/Runtime/Renderer/Private/PathTracing.cpp
--- a/Engine/Source/Runtime/Renderer/Private/PathTracing.cpp
+++ b/Engine/Source/Runtime/Renderer/Private/PathTracing.cpp
@@ -1431,7 +1431,8 @@ void FDeferredShadingSceneRenderer::RenderPathTracing(
 
 	// If the scene has changed in some way (camera move, object movement, etc ...)
 	// we must invalidate the ViewState to start over from scratch
-	if (FirstTime || Config.IsDifferent(PathTracingState->LastConfig) || HairStrands::HasPositionsChanged(GraphBuilder, View))
+	// NOTE: only check things like hair position changes for interactive viewports, for offline renders we don't want any chance of mid-render invalidation
+	if (FirstTime || Config.IsDifferent(PathTracingState->LastConfig) || (!View.bIsOfflineRender && HairStrands::HasPositionsChanged(GraphBuilder, View)))
 	{
 		// remember the options we used for next time
 		PathTracingState->LastConfig = Config;
