Fix HairStrandsVertexFactory compilation errors on some console platforms

diff --git a/Engine/Shaders/Private/HairStrands/HairStrandsVertexFactory.ush b/Engine/Shaders/Private/HairStrands/HairStrandsVertexFactory.ush
--- a/Engine/Shaders/Private/HairStrands/HairStrandsVertexFactory.ush
+++ b/Engine/Shaders/Private/HairStrands/HairStrandsVertexFactory.ush
@@ -504,12 +504,12 @@ FMaterialPixelParameters GetMaterialPixelParameters(float3 RayOrigin, float3 Ray
 		HitAttributes.Barycentrics.y);
 
 	// fetch triangle
-	FTriangleBaseAttributes Triangle = LoadTriangleBaseAttributes(HitPrimitiveIndex);
+	FTriangleBaseAttributes Tri = LoadTriangleBaseAttributes(HitPrimitiveIndex);
 	// Transform vertices into world space
 	const float3 WorldPositions[3] = {
-		CalcWorldPosition(Triangle.LocalPositions[0], SceneData.Primitive.LocalToWorld).xyz,
-		CalcWorldPosition(Triangle.LocalPositions[1], SceneData.Primitive.LocalToWorld).xyz,
-		CalcWorldPosition(Triangle.LocalPositions[2], SceneData.Primitive.LocalToWorld).xyz
+		CalcWorldPosition(Tri.LocalPositions[0], SceneData.Primitive.LocalToWorld).xyz,
+		CalcWorldPosition(Tri.LocalPositions[1], SceneData.Primitive.LocalToWorld).xyz,
+		CalcWorldPosition(Tri.LocalPositions[2], SceneData.Primitive.LocalToWorld).xyz
 	};
 	// Compute triangle normal
 	{
@@ -522,7 +522,7 @@ FMaterialPixelParameters GetMaterialPixelParameters(float3 RayOrigin, float3 Ray
 
 	// NOTE: This must match the tesselation pattern produced by SetBodySegment() in HairStrandsRaytracingGeometry.usf
 	// TODO: is there a way we could handle end-caps elegantly?
-	uint VertexIndex0 = Triangle.Indices.z / 4; // last index always belongs to the "bottom" vertex of the hair segment
+	uint VertexIndex0 = Tri.Indices.z / 4; // last index always belongs to the "bottom" vertex of the hair segment
 	uint VertexIndex1 = VertexIndex0 + 1;
 
 	if (HairStrandsVF.CullingEnable)
