Fix RTGI final gather reprojection pixel index calculation and add explicit bounds clamp.

* Problem was introduced in CL 16121649 (ScreenUVTime is converted to int2 before multiplication by image resolution)
* Fixes GPU crash

diff --git a/Engine/Shaders/Private/RayTracing/RayTracingFinalGatherRGS.usf b/Engine/Shaders/Private/RayTracing/RayTracingFinalGatherRGS.usf
--- a/Engine/Shaders/Private/RayTracing/RayTracingFinalGatherRGS.usf
+++ b/Engine/Shaders/Private/RayTracing/RayTracingFinalGatherRGS.usf
@@ -164,9 +164,13 @@ RAY_TRACING_ENTRY_RAYGEN(RayTracingFinalGatherRGS)
 			bShouldReject[LocalSampleIndex] = true;
 			//continue;
 		}
-
-		float2 ScreenUVTime = ScreenPositionTime.xy * float2(0.5, -0.5) + 0.5;
-		GatherPointsIndex = DispatchThreadIdToLinearIndex(int2(ScreenUVTime) * GatherPointsResolution + int2(View.ViewRectMin.xy));
+		else
+		{
+			float2 ScreenUVTime = ScreenPositionTime.xy * float2(0.5, -0.5) + 0.5;
+			int2 ScreenPixelIdTime = int2(ScreenUVTime * float2(GatherPointsResolution) + View.ViewRectMin);
+			ScreenPixelIdTime = min(ScreenPixelIdTime, GatherPointsResolution-1);
+			GatherPointsIndex = DispatchThreadIdToLinearIndex(ScreenPixelIdTime);
+		}
 #endif
 
 		FGatherSample GatherSample = ReadGatherSample(GatherPointsBuffer, GatherPointsIndex, LocalSampleIndex, GatherPointsResolution.x * GatherPointsResolution.y);
@@ -280,7 +284,8 @@ RAY_TRACING_ENTRY_RAYGEN(RayTracingFinalGatherRGS)
 			float4 ScreenPositionTime = mul(float4(WorldPosition, 1.0), GatherPointData.ViewMatrices[ParentSampleIndex]);
 			ScreenPositionTime.xyz /= ScreenPositionTime.w;
 			float2 ScreenUVTime = ScreenPositionTime.xy * float2(0.5, -0.5) + 0.5;
-			int2 GatherPointCentroid = int2(ScreenUVTime) * GatherPointsResolution + int2(View.ViewRectMin.xy);
+			int2 GatherPointCentroid = int2(ScreenUVTime * float2(GatherPointsResolution) + View.ViewRectMin.xy);
+			GatherPointCentroid = min(GatherPointCentroid, GatherPointsResolution-1);
 #endif
 
 			for (uint LocalSampleIndex0 = 0; LocalSampleIndex0 < GatherFilterWidth2; ++LocalSampleIndex0)
