Add custom output node for describing medium absorption for the path tracer

Instead of re-using the material's base color for expression absorption through glass, provide a dedicated output, similar to what exists on thin translucent materials.

Also tweak the logic so that solid glass relies on the DefaultLit shading model. This is required for now since we need to put the color in the payload's custom data fields which are otherwise occupied.

diff --git a/Engine/Shaders/Private/PathTracing/PathTracingCommon.ush b/Engine/Shaders/Private/PathTracing/PathTracingCommon.ush
--- a/Engine/Shaders/Private/PathTracing/PathTracingCommon.ush
+++ b/Engine/Shaders/Private/PathTracing/PathTracingCommon.ush
@@ -165,7 +165,7 @@ struct FPathTracingPayload : FMinimalPayload
 
 	bool IsMaterialSolidGlass()
 	{
-		return ShadingModelID != SHADINGMODELID_THIN_TRANSLUCENT && BlendingMode == RAY_TRACING_BLEND_MODE_TRANSLUCENT && Ior != 1.0 && Opacity < 1;
+		return ShadingModelID == SHADINGMODELID_DEFAULT_LIT && BlendingMode == RAY_TRACING_BLEND_MODE_TRANSLUCENT && Ior != 1.0 && Opacity < 1;
 	}
 
 	bool IsMaterialTransmissive()
diff --git a/Engine/Shaders/Private/PathTracing/PathTracingCore.ush b/Engine/Shaders/Private/PathTracing/PathTracingCore.ush
--- a/Engine/Shaders/Private/PathTracing/PathTracingCore.ush
+++ b/Engine/Shaders/Private/PathTracing/PathTracingCore.ush
@@ -895,7 +895,7 @@ bool PathTracingKernel(inout FPathState PathState, int Bounce)
 	// long term we will probably want a stack with priorities
 	if (MaterialSample.PositionBiasSign < 0 && Payload.IsMaterialSolidGlass())
 	{
-		float3 LocalSigmaT = -log(max(Payload.BaseColor, 1e-8f)) * AbsorptionScale;
+		float3 LocalSigmaT = -log(max(Payload.CustomData.xyz, 1e-8f)) * AbsorptionScale;
 		if (Payload.IsFrontFace())
 		{
 			// entering
diff --git a/Engine/Shaders/Private/PathTracing/PathTracingMaterialHitShader.usf b/Engine/Shaders/Private/PathTracing/PathTracingMaterialHitShader.usf
--- a/Engine/Shaders/Private/PathTracing/PathTracingMaterialHitShader.usf
+++ b/Engine/Shaders/Private/PathTracing/PathTracingMaterialHitShader.usf
@@ -116,9 +116,10 @@ RAY_TRACING_ENTRY_CLOSEST_HIT(PathTracingMaterialCHS,
 	// Store the results in local variables and reuse instead of calling the functions multiple times.
 	half3 BaseColor = GetMaterialBaseColor(PixelMaterialInputs);
 #if MATERIALBLENDING_TRANSLUCENT && REFRACTION_USE_INDEX_OF_REFRACTION
-	half Refraction = GetMaterialRefraction(PixelMaterialInputs).x;
+	Payload.Ior = GetMaterialRefraction(PixelMaterialInputs).x;
+	Payload.Ior = (Payload.Ior == 0.0) ? 1.0 : Payload.Ior;
 #else
-	half Refraction = 1;
+	Payload.Ior = 1;
 #endif
 	half  Metallic = GetMaterialMetallic(PixelMaterialInputs);
 	half  Specular = GetMaterialSpecular(PixelMaterialInputs);
@@ -163,8 +164,6 @@ RAY_TRACING_ENTRY_CLOSEST_HIT(PathTracingMaterialCHS,
 	Payload.WorldTangent = MaterialParameters.TangentToWorld[2];
 #endif
 
-	Payload.Ior = Refraction == 0 ? 1.0 : Refraction;
-
 #if MATERIAL_SHADINGMODEL_CLEAR_COAT
 	if (Payload.ShadingModelID == SHADINGMODELID_CLEAR_COAT)
 	{
@@ -270,8 +269,21 @@ RAY_TRACING_ENTRY_CLOSEST_HIT(PathTracingMaterialCHS,
 	}
 	else
 #endif
+
+#if MATERIAL_SHADINGMODEL_DEFAULT_LIT
+	if (Payload.ShadingModelID == SHADINGMODELID_DEFAULT_LIT)
+	{
+		// only allow refraction for default lit materials since we need space for the absorption amount
+#if HAVE_GetAbsorptionMediumMaterialOutput0
+		Payload.CustomData.xyz = GetAbsorptionMediumMaterialOutput0(MaterialParameters);
+#else
+		Payload.CustomData.xyz = 1.0;
+#endif
+	}
+	else
+#endif
 	{
-		// terminal case for when none of the above worked
+		// terminal case for the conditionals above
 	}
 
 
@@ -451,16 +463,15 @@ RAY_TRACING_ENTRY_ANY_HIT(PathTracingMaterialAHS,
 #elif MATERIALBLENDING_TRANSLUCENT
 	float Opacity = GetMaterialOpacity(PixelMaterialInputs);
 	float Transparency = 1 - Opacity;
-#if REFRACTION_USE_INDEX_OF_REFRACTION && !SIMPLIFIED_MATERIAL_SHADER
-	float Ior = GetMaterialRefraction(PixelMaterialInputs).x;
-#else
-	float Ior = 1.0;
-#endif
 
-#if !SIMPLIFIED_MATERIAL_SHADER
 	// Is refraction enabled?
-	if (Transparency > 0 && Ior != 1.0)
+#if !SIMPLIFIED_MATERIAL_SHADER && MATERIAL_SHADINGMODEL_DEFAULT_LIT && REFRACTION_USE_INDEX_OF_REFRACTION
+	float Ior = GetMaterialRefraction(PixelMaterialInputs).x;
+	Ior = Ior == 0.0 ? 1.0 : Ior;
+	int ShadingModelID = GetMaterialShadingModel(PixelMaterialInputs);
+	if (Transparency > 0 && ShadingModelID == SHADINGMODELID_DEFAULT_LIT && Ior != 1.0)
 	{
+		// current hit has some refraction
 		float PathRoughness = PackedPayload.GetPathRoughness();
 		float Roughness = GetMaterialRoughness(PixelMaterialInputs);
 		if (PathRoughness <= Roughness)
@@ -480,10 +491,13 @@ RAY_TRACING_ENTRY_ANY_HIT(PathTracingMaterialAHS,
 		float VoN = abs(dot(WorldRayDirection(), N));
 		float Fr = FresnelReflectance(VoN, Eta);
 		Transparency *= pow(1 - Fr, 2) * (1 - Roughness * Roughness) * saturate(PathRoughness - Roughness);
+
+#if HAVE_GetAbsorptionMediumMaterialOutput0
 		// Now track the optical thickness so that we can account for Beer's law along the shadow ray
 		// TODO: to support lights inside glass we would need to track an extra offset, but this is hopefully an uncommon scenario
 		float Distance = RayTCurrent();
-		float3 LocalSigmaT = -log(max(GetMaterialBaseColor(PixelMaterialInputs), 1e-8f));
+		float3 TransmittanceColor = GetAbsorptionMediumMaterialOutput0(MaterialParameters);
+		float3 LocalSigmaT = -log(max(TransmittanceColor, 1e-8f));
 		if (bIsFrontFace)
 		{
 			// remove contribution from the ray origin to the current hit (assuming the backside will be hit)
@@ -494,8 +508,9 @@ RAY_TRACING_ENTRY_ANY_HIT(PathTracingMaterialAHS,
 			// add contribution from ray origin to current hit (backside) - excess portion not covered by the ray will be removed by the front face hit
 			PackedPayload.SetTau(PackedPayload.GetTau() + LocalSigmaT * Distance);
 		}
+#endif // HAVE_GetAbsorptionMediumMaterialOutput0
 	}
-#endif
+#endif // refraction enabled?
 #else
 #error "Control flow does not account for all blending modes!"
 #endif
diff --git a/Engine/Source/Runtime/Engine/Classes/Materials/MaterialExpressionAbsorptionMediumMaterialOutput.h b/Engine/Source/Runtime/Engine/Classes/Materials/MaterialExpressionAbsorptionMediumMaterialOutput.h
new file mode 100644
--- /dev/null
+++ b/Engine/Source/Runtime/Engine/Classes/Materials/MaterialExpressionAbsorptionMediumMaterialOutput.h
@@ -0,0 +1,33 @@
+// Copyright Epic Games, Inc. All Rights Reserved.
+
+#pragma once
+
+#include "CoreMinimal.h"
+#include "Materials/MaterialExpressionCustomOutput.h"
+#include "UObject/ObjectMacros.h"
+#include "MaterialExpressionAbsorptionMediumMaterialOutput.generated.h"
+
+/** Material output expression for setting absorption properties of solid refractive glass (for the Path Tracer Only). */
+UCLASS(MinimalAPI, collapsecategories, hidecategories = Object)
+class UMaterialExpressionAbsorptionMediumMaterialOutput : public UMaterialExpressionCustomOutput
+{
+	GENERATED_UCLASS_BODY()
+
+	/** Input for the transmittance color seen after travelling a distance of 100 units into the object. Valid range is [0,1]. */
+	UPROPERTY()
+	FExpressionInput TransmittanceColor;
+
+public:
+#if WITH_EDITOR
+	//~ Begin UMaterialExpression Interface
+	virtual int32 Compile(class FMaterialCompiler* Compiler, int32 OutputIndex) override;
+	virtual void GetCaption(TArray<FString>& OutCaptions) const override;
+	//~ End UMaterialExpression Interface
+#endif
+
+	//~ Begin UMaterialExpressionCustomOutput Interface
+	virtual int32 GetNumOutputs() const override;
+	virtual FString GetFunctionName() const override;
+	virtual FString GetDisplayName() const override;
+	//~ End UMaterialExpressionCustomOutput Interface
+};
diff --git a/Engine/Source/Runtime/Engine/Classes/Materials/MaterialExpressionThinTranslucentMaterialOutput.h b/Engine/Source/Runtime/Engine/Classes/Materials/MaterialExpressionThinTranslucentMaterialOutput.h
--- a/Engine/Source/Runtime/Engine/Classes/Materials/MaterialExpressionThinTranslucentMaterialOutput.h
+++ b/Engine/Source/Runtime/Engine/Classes/Materials/MaterialExpressionThinTranslucentMaterialOutput.h
@@ -7,7 +7,7 @@
 #include "UObject/ObjectMacros.h"
 #include "MaterialExpressionThinTranslucentMaterialOutput.generated.h"
 
-/** Material output expression for writing single layer water volume material properties. */
+/** Material output expression for setting absorption properties of thin translucent materials. */
 UCLASS(MinimalAPI, collapsecategories, hidecategories = Object)
 class UMaterialExpressionThinTranslucentMaterialOutput : public UMaterialExpressionCustomOutput
 {
diff --git a/Engine/Source/Runtime/Engine/Private/Materials/HLSLMaterialTranslator.cpp b/Engine/Source/Runtime/Engine/Private/Materials/HLSLMaterialTranslator.cpp
--- a/Engine/Source/Runtime/Engine/Private/Materials/HLSLMaterialTranslator.cpp
+++ b/Engine/Source/Runtime/Engine/Private/Materials/HLSLMaterialTranslator.cpp
@@ -1107,6 +1107,19 @@ bool FHLSLMaterialTranslator::Translate()
 			}
 		}
 
+		if (Material->GetMaterialInterface() && Material->GetMaterialInterface()->GetMaterial()->HasAnyExpressionsInMaterialAndFunctionsOfType<UMaterialExpressionAbsorptionMediumMaterialOutput>())
+		{
+			// User placed an AbsorptionMedium node, make sure default lit is being used
+			if (!MaterialShadingModels.HasShadingModel(MSM_DefaultLit))
+			{
+				Errorf(TEXT("AbsorptionMedium custom output requires Default Lit shading model."));
+			}
+			if (Material->GetRefractionMode() != RM_IndexOfRefraction)
+			{
+				Errorf(TEXT("AbsorptionMedium custom output requires \"Index of Refraction\" as the \"Refraction Mode\"."));
+			}
+		}
+
 		if (Material->IsStrataMaterial())
 		{
 			if (Domain!= MD_Volume && !(BlendMode == BLEND_Translucent || BlendMode == BLEND_Opaque || BlendMode == BLEND_Masked))
@@ -10106,7 +10119,10 @@ int32 FHLSLMaterialTranslator::CustomOutput(class UMaterialExpressionCustomOutpu
 	const FString FunctionNameBase = FString::Printf(TEXT("%s%d"), *Custom->GetFunctionName(), OutputIndex);
 
 	// Primary function will have _LWC suffix if it returns an LWC type
-	FString ImplementationCodeFinite = FString::Printf(TEXT("%s %s%s(%s FMaterial%sParameters Parameters)\r\n{\r\n%s return %s;\r\n}\r\n"),
+	// We also define a pre-processor symbol to indicate the custom output function is available so that
+	// shaders can implement some fallback behavior for cases where the custom output has not been added
+	FString ImplementationCodeFinite = FString::Printf(TEXT("#define HAVE_%s 1\r\n%s %s%s(%s FMaterial%sParameters Parameters)\r\n{\r\n%s return %s;\r\n}\r\n"),
+		*FunctionNameBase,
 		*OutputTypeString,
 		*FunctionNameBase,
 		IsLWCType(OutputType) ? TEXT("_LWC") : TEXT(""),
diff --git a/Engine/Source/Runtime/Engine/Private/Materials/HLSLMaterialTranslator.h b/Engine/Source/Runtime/Engine/Private/Materials/HLSLMaterialTranslator.h
--- a/Engine/Source/Runtime/Engine/Private/Materials/HLSLMaterialTranslator.h
+++ b/Engine/Source/Runtime/Engine/Private/Materials/HLSLMaterialTranslator.h
@@ -20,6 +20,7 @@
 #include "Materials/MaterialExpressionSingleLayerWaterMaterialOutput.h"
 #include "Materials/MaterialExpressionVolumetricAdvancedMaterialOutput.h"
 #include "Materials/MaterialExpressionThinTranslucentMaterialOutput.h"
+#include "Materials/MaterialExpressionAbsorptionMediumMaterialOutput.h"
 #include "Materials/MaterialExpressionVolumetricAdvancedMaterialOutput.h"
 #include "MaterialCompiler.h"
 #include "RenderUtils.h"
diff --git a/Engine/Source/Runtime/Engine/Private/Materials/MaterialExpressions.cpp b/Engine/Source/Runtime/Engine/Private/Materials/MaterialExpressions.cpp
--- a/Engine/Source/Runtime/Engine/Private/Materials/MaterialExpressions.cpp
+++ b/Engine/Source/Runtime/Engine/Private/Materials/MaterialExpressions.cpp
@@ -47,6 +47,7 @@
 #include "ProfilingDebugging/LoadTimeTracker.h"
 
 #include "Materials/MaterialExpressionAbs.h"
+#include "Materials/MaterialExpressionAbsorptionMediumMaterialOutput.h"
 #include "Materials/MaterialExpressionActorPositionWS.h"
 #include "Materials/MaterialExpressionAdd.h"
 #include "Materials/MaterialExpressionAppendVector.h"
@@ -20005,6 +20006,69 @@ FString UMaterialExpressionThinTranslucentMaterialOutput::GetDisplayName() const
 	return TEXT("Thin Translucent Material");
 }
 
+///////////////////////////////////////////////////////////////////////////////
+// UMaterialExpressionAbsorptionMediumMaterialOutput
+///////////////////////////////////////////////////////////////////////////////
+UMaterialExpressionAbsorptionMediumMaterialOutput::UMaterialExpressionAbsorptionMediumMaterialOutput(const FObjectInitializer& ObjectInitializer)
+	: Super(ObjectInitializer)
+{
+	// Structure to hold one-time initialization
+	struct FConstructorStatics
+	{
+		FText NAME_MediumAbsorption;
+		FConstructorStatics()
+			: NAME_MediumAbsorption(LOCTEXT("Shading", "Shading"))
+		{
+		}
+	};
+	static FConstructorStatics ConstructorStatics;
+
+#if WITH_EDITORONLY_DATA
+	MenuCategories.Add(ConstructorStatics.NAME_MediumAbsorption);
+#endif
+
+#if WITH_EDITOR
+	Outputs.Reset();
+#endif
+}
+
+#if WITH_EDITOR
+
+int32 UMaterialExpressionAbsorptionMediumMaterialOutput::Compile(class FMaterialCompiler* Compiler, int32 OutputIndex)
+{
+	int32 CodeInput = INDEX_NONE;
+
+	if (OutputIndex == 0)
+	{
+		CodeInput = TransmittanceColor.IsConnected() ? TransmittanceColor.Compile(Compiler) : Compiler->Constant3(1.0f, 1.0f, 1.0f);
+	}
+
+	return Compiler->CustomOutput(this, OutputIndex, CodeInput);
+}
+
+void UMaterialExpressionAbsorptionMediumMaterialOutput::GetCaption(TArray<FString>& OutCaptions) const
+{
+	OutCaptions.Add(FString(TEXT("Absorption Medium (Path Tracer Only)")));
+}
+
+#endif // WITH_EDITOR
+
+int32 UMaterialExpressionAbsorptionMediumMaterialOutput::GetNumOutputs() const
+{
+	return 1;
+}
+
+FString UMaterialExpressionAbsorptionMediumMaterialOutput::GetFunctionName() const
+{
+	return TEXT("GetAbsorptionMediumMaterialOutput");
+}
+
+FString UMaterialExpressionAbsorptionMediumMaterialOutput::GetDisplayName() const
+{
+	return TEXT("Absorption Medium");
+}
+
+
 
 ///////////////////////////////////////////////////////////////////////////////
 // UMaterialExpressionSceneDepthWithoutWater
