Fix shading compilation warning.

diff --git a/Engine/Shaders/Private/BasePassPixelShader.usf b/Engine/Shaders/Private/BasePassPixelShader.usf
--- a/Engine/Shaders/Private/BasePassPixelShader.usf
+++ b/Engine/Shaders/Private/BasePassPixelShader.usf
@@ -1592,7 +1592,7 @@ void FPixelShaderInOut_MainPS(
 	// Clear only what should be and keep the rest
 	GBuffer = (FGBufferData)0;						// We should not rely on gbuffer for anything else to make any pass work.
 
-	uint2 PixelPos = In.SvPosition.xy;
+	uint2 PixelPos = uint2(In.SvPosition.xy);
 
 	float3 EmissiveLuminance = 0.0f;
 	PackStrataOut(
diff --git a/Engine/Shaders/Private/HairStrands/HairStrandsDeepTransmittanceDualScattering.ush b/Engine/Shaders/Private/HairStrands/HairStrandsDeepTransmittanceDualScattering.ush
--- a/Engine/Shaders/Private/HairStrands/HairStrandsDeepTransmittanceDualScattering.ush
+++ b/Engine/Shaders/Private/HairStrands/HairStrandsDeepTransmittanceDualScattering.ush
@@ -62,14 +62,14 @@ FHairTransmittanceData ComputeDualScatteringTerms(
 
 	// Always shift the hair count by one to remove self-occlusion/shadow aliasing and have smoother transition
 	// This insure the the pow function always starts at 0 for front facing hair
-	const uint HairCount = max(0, TransmittanceMask.HairCount - 1);
+	const float HairCount = max(0, TransmittanceMask.HairCount - 1);
 
 	// This is a coarse approximation of eq. 13. Normally, Beta_f should be weighted by the 'normalized' 
 	// R, TT, and TRT terms 
 	const float3 af_weights = af / (af.r + af.g + af.b);
 	const float3 Beta_f  = dot(float3(Beta_R, Beta_TT, Beta_TRT), af_weights);
 	const float3 Beta_f2 = Beta_f*Beta_f;
-	const float3 sigma_f2 = Beta_f2 * max(1u, HairCount);
+	const float3 sigma_f2 = Beta_f2 * max(1.f, HairCount);
 
 	const float Theta_d = asin(SinThetaL) + asin(SinThetaV);
 	const float Theta_h = Theta_d * 0.5f;
@@ -176,7 +176,7 @@ float3 GetHairTransmittanceOnly(FHairTransmittanceMask TransmittanceMask, FGBuff
 {
 	// Always shift the hair count by one to remove self-occlusion/shadow aliasing and have smoother transition
 	// This insure the the pow function always starts at 0 for front facing hair
-	const uint HairCount = max(0, TransmittanceMask.HairCount - 1);
+	const float HairCount = max(0.f, TransmittanceMask.HairCount - 1.f);
 
 	const FHairAverageScattering AbsorptionData = SampleHairLUT(InHairScatteringLUTTexture, InHairLUTSampler, GBuffer.BaseColor, GBuffer.Roughness, SinLightAngle);
 	return pow(AbsorptionData.A_front, HairCount) * TransmittanceMask.Visibility;
diff --git a/Engine/Shaders/Private/Strata/Strata.ush b/Engine/Shaders/Private/Strata/Strata.ush
--- a/Engine/Shaders/Private/Strata/Strata.ush
+++ b/Engine/Shaders/Private/Strata/Strata.ush
@@ -336,8 +336,8 @@ struct FStrataClassification
 
 uint StrataPackClassificationData(FStrataClassification In)
 {
-	const uint MaterialAO8Bits = sqrt(saturate(In.MaterialAO)) * 255.0f;	// Give more accuracy to the blacks to reduce banding
-	const uint TopLayerRoughness8Bits = In.TopLayerRoughness * 255.0f;
+	const uint MaterialAO8Bits = uint(sqrt(saturate(In.MaterialAO)) * 255.0f);	// Give more accuracy to the blacks to reduce banding
+	const uint TopLayerRoughness8Bits = uint(In.TopLayerRoughness * 255.0f);
 	return 
 		(In.ShadingModels & 0xFF) | 
 		((MaterialAO8Bits & 0xFF)			 << 8  ) |
@@ -408,7 +408,7 @@ uint2 StrataPackSSSData(FStrataSubsurfaceData In)
 uint StrataSubsurfaceProfileIdTo8bits(float In)
 {
 	// Similar encoding than ExtractSubsurfaceProfileInt. Valid profile ID start at 1.
-	return In * 255.0f + 0.5f;
+	return uint(In * 255.0f + 0.5f);
 }
 
 FStrataSubsurfaceData StrataUnpackSSSData(uint2 In)
