Add Chan support for path tracer default lit project, if rough diffuse is activate on the project.

diff --git a/Engine/Shaders/Private/PathTracing/Material/PathTracingDefaultLit.ush b/Engine/Shaders/Private/PathTracing/Material/PathTracingDefaultLit.ush
--- a/Engine/Shaders/Private/PathTracing/Material/PathTracingDefaultLit.ush
+++ b/Engine/Shaders/Private/PathTracing/Material/PathTracingDefaultLit.ush
@@ -33,6 +33,7 @@ FMaterialEval DefaultLit_EvalMaterial(
 	const float NoV = saturate(V.z);
 	const float NoL = saturate(L.z);
 	const float VoH = saturate(dot(V, H));
+	const float NoH = saturate(H.z);
 
 	const FBxDFEnergyTerms Spec = ComputeGGXSpecEnergyTerms(Payload.Roughness, NoV, Payload.SpecularColor);
 
@@ -40,7 +41,8 @@ FMaterialEval DefaultLit_EvalMaterial(
 	const float LobeSelectionPdf = LobeSelectionProb(Payload.DiffuseColor * (1 - Spec.E), Spec.E);
 
 	// Diffuse Lobe
-	AddLobeWithMIS(Result.Weight, Result.Pdf, (1 - Spec.E) * Payload.DiffuseColor * ShadowTerminatorTerm(L_World, Payload.WorldNormal, Payload.WorldSmoothNormal), NoL / PI, LobeSelectionPdf);
+	const float3 Diffuse = GetPathTracingDiffuseModel(Payload.DiffuseColor, Payload.Roughness, NoV, NoL, VoH, NoH);
+	AddLobeWithMIS(Result.Weight, Result.Pdf, (1 - Spec.E) * Diffuse * ShadowTerminatorTerm(L_World, Payload.WorldNormal, Payload.WorldSmoothNormal), NoL / PI, LobeSelectionPdf);
 
 	// Specular lobe
 	const float2 GGXResult = GGXEvalReflection(L, V, H, Alpha);
@@ -111,12 +113,14 @@ FMaterialSample DefaultLit_SampleMaterial(
 
 	const float NoL = saturate(L.z);
 	const float VoH = saturate(dot(V, H));
+	const float NoH = saturate(H.z);
 
 	float3 OutWeight = 0;
 	float OutPdf = 0;
 
 	// Diffuse Lobe
-	AddLobeWithMIS(OutWeight, OutPdf, (1 - Spec.E) * Payload.DiffuseColor * ShadowTerminatorTerm(L_World, Payload.WorldNormal, Payload.WorldSmoothNormal), NoL / PI, LobeSelectionPdf);
+	const float3 Diffuse = GetPathTracingDiffuseModel(Payload.DiffuseColor, Payload.Roughness, NoV, NoL, VoH, NoH);
+	AddLobeWithMIS(OutWeight, OutPdf, (1 - Spec.E) * Diffuse * ShadowTerminatorTerm(L_World, Payload.WorldNormal, Payload.WorldSmoothNormal), NoL / PI, LobeSelectionPdf);
 
 	// Specular lobe
 	const float2 GGXResult = GGXEvalReflection(L, V, H, Alpha);
diff --git a/Engine/Shaders/Private/PathTracing/Material/PathTracingMaterialCommon.ush b/Engine/Shaders/Private/PathTracing/Material/PathTracingMaterialCommon.ush
--- a/Engine/Shaders/Private/PathTracing/Material/PathTracingMaterialCommon.ush
+++ b/Engine/Shaders/Private/PathTracing/Material/PathTracingMaterialCommon.ush
@@ -138,3 +138,12 @@ float ShadowTerminatorTerm(float3 L, float3 N, float3 Ns)
 	return 1.0;
 #endif
 }
+
+float3 GetPathTracingDiffuseModel(float3 DiffuseColor, float Roughness, float NoV, float NoL, float VoH, float NoH)
+{
+#if MATERIAL_ROUGHDIFFUSE
+	return PI * Diffuse_Chan(DiffuseColor, Pow4(Roughness), NoV, NoL, VoH, NoH, 1.0f);
+#else
+	return DiffuseColor;
+#endif
+}
