Fix more strata shading compilation warning.

diff --git a/Engine/Shaders/Private/Strata/StrataForwardLighting.ush b/Engine/Shaders/Private/Strata/StrataForwardLighting.ush
--- a/Engine/Shaders/Private/Strata/StrataForwardLighting.ush
+++ b/Engine/Shaders/Private/Strata/StrataForwardLighting.ush
@@ -467,55 +467,56 @@ float3 StrataForwardLighting(
 			////
 			//// Evaluate Directional light only
 			////
-
-			float3 DirLightL = DirLightData.Direction;	// Already normalized
-			float3 ToDirLight = DirLightL;
-			float DirLightMask = 1;
-			if (DirLightData.bRadialLight)
-			{
-				DirLightMask = GetLocalLightAttenuation(AbsoluteWorldPosition, DirLightData, ToDirLight, DirLightL);
-			}
-
-			bool bEvaluateHasBeenExecuted = false;
-			BRANCH
-			if (DirectionalLightData.HasDirectionalLight)
 			{
-				uint LightChannelMask = DirectionalLightData.DirectionalLightShadowMapChannelMask >> 8;
+				float3 DirLightL = DirLightData.Direction;	// Already normalized
+				float3 ToDirLight = DirLightL;
+				float DirLightMask = 1;
+				if (DirLightData.bRadialLight)
+				{
+					DirLightMask = GetLocalLightAttenuation(AbsoluteWorldPosition, DirLightData, ToDirLight, DirLightL);
+				}
 
-				float DirectionalLightCloudShadow = 1.0f;
-			#if NEEDS_BASEPASS_CLOUD_SHADOW_INTERPOLATOR
-				DirectionalLightCloudShadow = BasePassInterpolants.VertexCloudShadow;
-			#endif
+				bool bEvaluateHasBeenExecuted = false;
+				BRANCH
+				if (DirectionalLightData.HasDirectionalLight)
+				{
+					uint LightChannelMask = DirectionalLightData.DirectionalLightShadowMapChannelMask >> 8;
 
-				FStrataEvaluateResult BSDFEvaluate = (FStrataEvaluateResult)0;
-				Color += StrataForwardLightingCommon(
-					Dither,
-					bForceFullyRough,
-					DirLightData,
-					ToDirLight,
-					DirLightMask,
-					DirLightAttenuation,
-					RectTexture,
-					LightChannelMask,
-					PrimitiveLightingChannelMask,
-					AbsoluteWorldPosition,
-					TranslatedWorldPosition,
-					SceneDepth,
-					BSDFColoredVisibility,
-					StrataPixelHeader,
-					StrataBSDFContext,
-					bEvaluateHasBeenExecuted,
-					BSDFEvaluate) * DirectionalLightCloudShadow;
+					float DirectionalLightCloudShadow = 1.0f;
+				#if NEEDS_BASEPASS_CLOUD_SHADOW_INTERPOLATOR
+					DirectionalLightCloudShadow = BasePassInterpolants.VertexCloudShadow;
+				#endif
 
-				LayerThroughput += BSDFCoverage * BSDFEvaluate.Throughput;
-			}
+					FStrataEvaluateResult BSDFEvaluate = (FStrataEvaluateResult)0;
+					Color += StrataForwardLightingCommon(
+						Dither,
+						bForceFullyRough,
+						DirLightData,
+						ToDirLight,
+						DirLightMask,
+						DirLightAttenuation,
+						RectTexture,
+						LightChannelMask,
+						PrimitiveLightingChannelMask,
+						AbsoluteWorldPosition,
+						TranslatedWorldPosition,
+						SceneDepth,
+						BSDFColoredVisibility,
+						StrataPixelHeader,
+						StrataBSDFContext,
+						bEvaluateHasBeenExecuted,
+						BSDFEvaluate) * DirectionalLightCloudShadow;
+
+					LayerThroughput += BSDFCoverage * BSDFEvaluate.Throughput;
+				}
 
-			BRANCH
-			if(!bEvaluateHasBeenExecuted)
-			{
-				// Only update the throughput
-				FStrataEvaluateResult BSDFEvaluate = StrataEvaluateBSDF(StrataBSDFContext);
-				LayerThroughput += BSDFCoverage * BSDFEvaluate.Throughput;
+				BRANCH
+				if(!bEvaluateHasBeenExecuted)
+				{
+					// Only update the throughput
+					FStrataEvaluateResult BSDFEvaluate = StrataEvaluateBSDF(StrataBSDFContext);
+					LayerThroughput += BSDFCoverage * BSDFEvaluate.Throughput;
+				}
 			}
 
 			////
@@ -524,8 +525,6 @@ float3 StrataForwardLighting(
 
 		#if !DISABLE_FORWARD_LOCAL_LIGHTS && (FORWARD_SHADING || TRANSLUCENCY_LIGHTING_SURFACE_FORWARDSHADING) // Forbidden for Single Layer Water
 			// Basic implementation for FORWARD_LOCAL_LIGHTS where local lights are reloaded per BSDF layer. It is fast by default when layer==1.
-
-			uint GridIndex = ComputeLightGridCellIndex((uint2)((SvPosition.xy - ResolvedView.ViewRectMin.xy) * View.LightProbeSizeRatioAndInvSizeRatio.zw), SvPosition.w, EyeIndex);
 			const FCulledLightsGridData CulledLightsGrid = GetCulledLightsGrid(GridIndex, EyeIndex);
 
 			// Limit max to ForwardLightData.NumLocalLights.
