Fix unroll issue with Strata material packing.

diff --git a/Engine/Shaders/Private/Strata/Strata.ush b/Engine/Shaders/Private/Strata/Strata.ush
--- a/Engine/Shaders/Private/Strata/Strata.ush
+++ b/Engine/Shaders/Private/Strata/Strata.ush
@@ -42,6 +42,10 @@
 #define STRATA_CLAMPED_LAYER_COUNT 1
 #endif
 
+#ifndef STRATA_CLAMPED_MAX_BSDF_COUNT_PER_LAYER
+#define STRATA_CLAMPED_MAX_BSDF_COUNT_PER_LAYER STRATA_MAX_BSDF_COUNT_PER_LAYER
+#endif
+
 #ifndef MATERIAL_SHADINGMODEL_SINGLELAYERWATER
 #define MATERIAL_SHADINGMODEL_SINGLELAYERWATER 0 
 #endif
@@ -1749,7 +1753,7 @@ void PackStrataOut(
 		const bool bTopLayer = l == 0;
 		const bool bBottomLayer = l == (Strata.LayerCount - 1);
 
-		STRATA_UNROLL
+		STRATA_UNROLL_N(STRATA_CLAMPED_MAX_BSDF_COUNT_PER_LAYER)
 		for (int i = 0; i < Strata.Layers[l].BSDFCount; ++i)
 		{
 			// Dither BSDF data to hide more aggressive bit reduction.
@@ -1921,7 +1925,7 @@ void PackStrataOut(
 	{
 		FStrataLayer Layer = Strata.Layers[l];
 
-		STRATA_UNROLL
+		STRATA_UNROLL_N(STRATA_CLAMPED_MAX_BSDF_COUNT_PER_LAYER)
 		for (int i = 0; i < Layer.BSDFCount; ++i)
 		{
 			FStrataBSDF BSDF = Layer.BSDFs[i];
diff --git a/Engine/Source/Runtime/Engine/Private/Materials/HLSLMaterialTranslator.cpp b/Engine/Source/Runtime/Engine/Private/Materials/HLSLMaterialTranslator.cpp
--- a/Engine/Source/Runtime/Engine/Private/Materials/HLSLMaterialTranslator.cpp
+++ b/Engine/Source/Runtime/Engine/Private/Materials/HLSLMaterialTranslator.cpp
@@ -1830,6 +1830,11 @@ void FHLSLMaterialTranslator::GetMaterialEnvironment(EShaderPlatform InPlatform,
 				OutEnvironment.SetDefine(TEXT("STRATA_CLAMPED_LAYER_COUNT"), StrataMaterialAnalysis.ClampedLayerCount);
 			}
 
+			if (StrataMaterialAnalysis.RequestedMaxBSDFCountPerLayer > 0)
+			{
+				OutEnvironment.SetDefine(TEXT("STRATA_CLAMPED_MAX_BSDF_COUNT_PER_LAYER"), StrataMaterialAnalysis.RequestedMaxBSDFCountPerLayer);
+			}
+
 			FString StrataMaterialDescription;
 
 			StrataMaterialDescription += FString::Printf(TEXT("----- STRATA -----\r\n"));
diff --git a/Engine/Source/Runtime/Engine/Private/Materials/StrataMaterial.cpp b/Engine/Source/Runtime/Engine/Private/Materials/StrataMaterial.cpp
--- a/Engine/Source/Runtime/Engine/Private/Materials/StrataMaterial.cpp
+++ b/Engine/Source/Runtime/Engine/Private/Materials/StrataMaterial.cpp
@@ -344,7 +344,7 @@ FStrataMaterialAnalysisResult StrataCompilationInfoMaterialAnalysis(FMaterialCom
 	Result.RequestedByteCount += Compiler->StrataCompilationInfoGetSharedLocalBasesCount() * STRATA_PACKED_SHAREDLOCALBASIS_STRIDE_BYTES;
 	
 	// 2. The list of BSDFs
-
+	Result.RequestedMaxBSDFCountPerLayer = 0;
 	// We process layers from top to bottom to cull the bottom ones in case we run out of pixel bytes
 	for (uint32 LayerIt = 0; LayerIt < Material.LayerCount; LayerIt++)
 	{
@@ -405,6 +405,7 @@ FStrataMaterialAnalysisResult StrataCompilationInfoMaterialAnalysis(FMaterialCom
 			}
 			}
 		}
+		Result.RequestedMaxBSDFCountPerLayer = FMath::Max(Result.RequestedMaxBSDFCountPerLayer, Layer.BSDFCount);
 
 		Result.RequestedLayerCount++;
 		Result.RequestedBSDFCount += Layer.BSDFCount;
diff --git a/Engine/Source/Runtime/Engine/Private/Materials/StrataMaterial.h b/Engine/Source/Runtime/Engine/Private/Materials/StrataMaterial.h
--- a/Engine/Source/Runtime/Engine/Private/Materials/StrataMaterial.h
+++ b/Engine/Source/Runtime/Engine/Private/Materials/StrataMaterial.h
@@ -46,6 +46,16 @@ struct FStrataMaterialCompilationInfo
 		LayerCount = 0;
 		memset(Layers, 0, sizeof(Layers));
 	}
+
+	uint32 GetPerLayerMaxBSDFCount() const 
+	{
+		uint32 Count = 0;
+		for (uint32 LayerIt = 0; LayerIt < LayerCount; ++LayerIt)
+		{
+			Count = FMath::Max(Count, Layers[LayerIt].BSDFCount);
+		}
+		return Count;
+	}
 };
 
 FString GetStrataBSDFName(uint8 BSDFType);
@@ -77,6 +87,7 @@ struct FStrataMaterialAnalysisResult
 	FStrataMaterialAnalysisResult();
 	bool bFitInMemoryBudget;
 	uint32 RequestedLayerCount;
+	uint32 RequestedMaxBSDFCountPerLayer;
 	uint32 RequestedBSDFCount;
 	uint32 RequestedByteCount;
 	uint32 ClampedLayerCount;
