Remove unused materialAO from strata classification data.

diff --git a/Engine/Shaders/Private/Strata/Strata.ush b/Engine/Shaders/Private/Strata/Strata.ush
--- a/Engine/Shaders/Private/Strata/Strata.ush
+++ b/Engine/Shaders/Private/Strata/Strata.ush
@@ -330,7 +330,6 @@ float3 StrataUnpackNormal(uint PackedNormal);
 struct FStrataClassification
 {
 	uint ShadingModels;
-	float MaterialAO;
 	float TopLayerRoughness;
 	bool bIsStrataMaterial;
 	bool bHasSubsurface;
@@ -340,11 +339,9 @@ struct FStrataClassification
 
 uint StrataPackClassificationData(FStrataClassification In)
 {
-	const uint MaterialAO8Bits = uint(sqrt(saturate(In.MaterialAO)) * 255.0f);	// Give more accuracy to the blacks to reduce banding
 	const uint TopLayerRoughness8Bits = uint(In.TopLayerRoughness * 255.0f);
 	return 
 		(In.ShadingModels & 0xFF) | 
-		((MaterialAO8Bits & 0xFF)			 << 8  ) |
 		((TopLayerRoughness8Bits & 0xFF)	 << 16 ) |
 		((In.bIsStrataMaterial		? 1 : 0) << 24 ) |
 		((In.bHasSubsurface			? 1 : 0) << 25 ) | 
@@ -354,18 +351,15 @@ uint StrataPackClassificationData(FStrataClassification In)
 
 FStrataClassification StrataUnpackClassificationData(uint In)
 {
-	uint MaterialAO8Bits = 0;
 	uint TopLayerRoughness8Bits = 0;
 	FStrataClassification Out;
 	Out.ShadingModels			= (In >> 0 ) & 0xFF;
-	MaterialAO8Bits				= (In >> 8 ) & 0xFF;
 	TopLayerRoughness8Bits		= (In >> 16) & 0xFF;
 	Out.bIsStrataMaterial		= (In >> 24) & 0x1;
 	Out.bHasSubsurface			= (In >> 25) & 0x1;
 	Out.bHasSubsurfaceProfile	= (In >> 26) & 0x1;
 	Out.bIsSimple				= (In >> 27) & 0x1;
 
-	Out.MaterialAO = Square(float(MaterialAO8Bits) * (1.0f / 255.0f));
 	Out.TopLayerRoughness = float(TopLayerRoughness8Bits) * (1.0f / 255.0f);
 	return Out;
 }
diff --git a/Engine/Shaders/Private/Strata/StrataMaterialClassification.usf b/Engine/Shaders/Private/Strata/StrataMaterialClassification.usf
--- a/Engine/Shaders/Private/Strata/StrataMaterialClassification.usf
+++ b/Engine/Shaders/Private/Strata/StrataMaterialClassification.usf
@@ -137,7 +137,6 @@ void MainPS(
 	FStrataPixelHeader Header = UnpackStrataHeaderIn(MaterialLobesBuffer, StrataAddressing);
 
 	ClassificationData.bIsStrataMaterial = Header.BSDFCount > 0;
-	ClassificationData.MaterialAO = 1.0;
 	ClassificationData.bIsSimple = true; // non strata material or sky pixels are considered as simple.
 	if (ClassificationData.bIsStrataMaterial)
 	{
@@ -146,7 +145,6 @@ void MainPS(
 		float4 WorldPos = mul(float4(ScreenPosition * Depth, Depth, 1), LWCHackToFloat(PrimaryView.ScreenToWorld));
 		const float3 V = normalize(LWCHackToFloat(PrimaryView.WorldCameraOrigin) - WorldPos.xyz);
 
-		ClassificationData.MaterialAO = Header.MaterialAO;
 		ClassificationData.bIsSimple = Header.BSDFCount == 1;
 
 		float TopLayerRoughness = 0.0f;
