Initial implementation of inline raytracing:
* Added CFLAG_InlineRayTracing and PLATFORM_SUPPORTS_INLINE_RAY_TRACING define.
* Added initial per-platform implementation of TraceRayInline.
* Added FTraceRayInlineResult which is the ouput of TraceRayInline.

diff --git a/Engine/Shaders/Private/RayTracing/TraceRayInline.ush b/Engine/Shaders/Private/RayTracing/TraceRayInline.ush
new file mode 100644
--- /dev/null
+++ b/Engine/Shaders/Private/RayTracing/TraceRayInline.ush
@@ -0,0 +1,23 @@
+// Copyright Epic Games, Inc. All Rights Reserved.
+
+#pragma once
+
+#ifndef PLATFORM_SUPPORTS_INLINE_RAY_TRACING
+#define PLATFORM_SUPPORTS_INLINE_RAY_TRACING 0
+#endif // PLATFORM_SUPPORTS_INLINE_RAY_TRACING
+
+#if PLATFORM_SUPPORTS_INLINE_RAY_TRACING
+
+#if !COMPUTESHADER
+#error Inline ray tracing is only supported in compute shaders
+#endif // !COMPUTESHADER
+
+#if OVERRIDE_TRACERAYINLINE_USH
+#include "/Platform/Private/TraceRayInline.ush"
+#endif // OVERRIDE_TRACERAYINLINE_USH
+
+#ifndef PLATFORM_SUPPORTS_INLINE_RAY_TRACING_TRIANGLE_NORMALS
+#define PLATFORM_SUPPORTS_INLINE_RAY_TRACING_TRIANGLE_NORMALS 0
+#endif // PLATFORM_SUPPORTS_INLINE_RAY_TRACING_TRIANGLE_NORMALS
+
+#endif // PLATFORM_SUPPORTS_INLINE_RAY_TRACING
diff --git a/Engine/Shaders/Private/RayTracing/TraceRayInlineCommon.ush b/Engine/Shaders/Private/RayTracing/TraceRayInlineCommon.ush
new file mode 100644
--- /dev/null
+++ b/Engine/Shaders/Private/RayTracing/TraceRayInlineCommon.ush
@@ -0,0 +1,25 @@
+// Copyright Epic Games, Inc. All Rights Reserved.
+
+#pragma once
+
+struct FTraceRayInlineResult
+{
+	float HitT;
+	float2 Barycentrics;
+	uint InstanceIndex;
+	uint InstanceID;
+	uint InstanceContributionToHitGroupIndex;
+	uint GeometryIndex;
+	uint PrimitiveIndex;
+	bool bIsFrontFace;
+
+	uint Bookmark[2];
+
+	// Only valid when PLATFORM_SUPPORTS_INLINE_RAY_TRACING_TRIANGLE_NORMALS is 1
+	float3 WorldGeometryNormal;
+
+	bool IsMiss() { return HitT < 0; }
+	bool IsHit() { return !IsMiss(); }
+
+	void SetMiss() { HitT = -1; }
+};
diff --git a/Engine/Source/Runtime/RenderCore/Public/ShaderCompilerCore.h b/Engine/Source/Runtime/RenderCore/Public/ShaderCompilerCore.h
--- a/Engine/Source/Runtime/RenderCore/Public/ShaderCompilerCore.h
+++ b/Engine/Source/Runtime/RenderCore/Public/ShaderCompilerCore.h
@@ -106,7 +106,8 @@ enum ECompilerFlags
 	CFLAG_AllowTypedUAVLoads,
 	// Prefer shader execution in wave32 mode if possible.
 	CFLAG_Wave32,
-
+	// Enable support of inline raytracing in compute shader.
+	CFLAG_InlineRayTracing,
 	// Force using the SC rewrite functionality before calling DXC on D3D12
 	CFLAG_D3D12ForceShaderConductorRewrite,
 	// Enable support of C-style data types for platforms that can. Check for PLATFORM_SUPPORTS_REAL_TYPES.
