Change strata classification structure to store data into a bitfield to reduce VGPR pressure.
* VGPR 48 -> 44

diff --git a/Engine/Shaders/Private/PostProcessAmbientOcclusionCommon.ush b/Engine/Shaders/Private/PostProcessAmbientOcclusionCommon.ush
--- a/Engine/Shaders/Private/PostProcessAmbientOcclusionCommon.ush
+++ b/Engine/Shaders/Private/PostProcessAmbientOcclusionCommon.ush
@@ -108,7 +108,7 @@ float3 GetWorldSpaceNormalFromAOInput(float2 UV, float4 SvPosition)
 	#else
 #if STRATA_ENABLED
 		const FStrataClassification StrataData = StrataUnpackClassificationData(Strata.ClassificationTexture.Load(uint3(SvPosition.xy, 0)));
-		if (StrataData.ShadingModels != STRATA_BSDF_TYPE_HAIR)
+		if (!StrataHasShadingModel(StrataData, STRATA_BSDF_TYPE_HAIR))
 		{
 			const FStrataTopLayerData TopLayerData = StrataUnpackTopLayerData(Strata.TopLayerNormalTexture.Load(uint3(SvPosition.xy, 0)));
 			WorldNormal = TopLayerData.WorldNormal;
diff --git a/Engine/Shaders/Private/SSRT/SSRTReflections.usf b/Engine/Shaders/Private/SSRT/SSRTReflections.usf
--- a/Engine/Shaders/Private/SSRT/SSRTReflections.usf
+++ b/Engine/Shaders/Private/SSRT/SSRTReflections.usf
@@ -80,7 +80,7 @@ void ScreenSpaceReflectionsStencilPS(
 	const FStrataClassification StrataData = StrataUnpackClassificationData(Strata.ClassificationTexture.Load(uint3(PixelCoord, 0)));
 
 	const float Roughness = StrataData.TopLayerRoughness;
-	const bool bNoMaterial = !StrataData.bIsStrataMaterial;
+	const bool bNoMaterial = !IsStrataMaterial(StrataData);
 
 #else // STRATA_ENABLED
 
@@ -161,7 +161,7 @@ void ScreenSpaceReflections(
 		// STRATA_TODO try stochastic selection
 		N = TopLayerData.WorldNormal;
 		Roughness = StrataData.TopLayerRoughness;
-		bNoMaterial = !StrataData.bIsStrataMaterial;
+		bNoMaterial = !IsStrataMaterial(StrataData);
 	}
 
 	float DeviceZ = SampleDeviceZFromSceneTextures(UV);
diff --git a/Engine/Shaders/Private/ShadowProjectionPixelShader.usf b/Engine/Shaders/Private/ShadowProjectionPixelShader.usf
--- a/Engine/Shaders/Private/ShadowProjectionPixelShader.usf
+++ b/Engine/Shaders/Private/ShadowProjectionPixelShader.usf
@@ -264,11 +264,11 @@ void Main(
 		bool bIsEyeOrHair 				= false;
 	#if STRATA_ENABLED
 		const FStrataClassification StrataData = StrataUnpackClassificationData(Strata.ClassificationTexture.Load(PixelCoord));
-		if (StrataData.bIsStrataMaterial)
+		if (IsStrataMaterial(StrataData))
 		{
-			bIsEyeOrHair				= StrataHasShadingModel(StrataData.ShadingModels, STRATA_BSDF_TYPE_HAIR);
-			bHasSubsurface				= StrataData.bHasSubsurface || bIsEyeOrHair;
-			bHasSubsurfaceTransmission	= StrataData.bHasSubsurface;
+			bIsEyeOrHair				= StrataHasShadingModel(StrataData, STRATA_BSDF_TYPE_HAIR);
+			bHasSubsurface				= HasSubsurface(StrataData) || bIsEyeOrHair;
+			bHasSubsurfaceTransmission	= HasSubsurface(StrataData);
 		}
 	#else // STRATA_ENABLED
 		FGBufferData GBufferData = GetGBufferData(ScreenUV);
@@ -321,7 +321,7 @@ void Main(
 
 				// For non-profile version, use default transmission parameters (only NormalScale & ExtinctionScale are used)
 				FTransmissionProfileParams TransmissionParams = InitTransmissionProfileParams();
-				if (StrataData.bHasSubsurfaceProfile)
+				if (HasSubsurfaceProfile(StrataData))
 				{
 					TransmissionParams = GetTransmissionProfileParams(SSSData.ProfileId);
 				}
@@ -471,7 +471,7 @@ void MainOnePassPointLightPS(
 	const FStrataTopLayerData TopLayerData = StrataUnpackTopLayerData(Strata.TopLayerNormalTexture.Load(uint3(SVPos.xy, 0)));
 	const FStrataSubsurfaceData SSSData = StrataUnpackSSSData(Strata.SSSTexture.Load(uint3(SVPos.xy, 0)));
 
-	const bool IsSubsurfaceProfile = StrataData.bHasSubsurfaceProfile;
+	const bool IsSubsurfaceProfile = HasSubsurfaceProfile(StrataData);
 	const float3 WorldNormal = TopLayerData.WorldNormal;
 	const FTransmissionProfileParams TransmissionParams = GetTransmissionProfileParams(SSSData.ProfileId);
 
diff --git a/Engine/Shaders/Private/Strata/Strata.ush b/Engine/Shaders/Private/Strata/Strata.ush
--- a/Engine/Shaders/Private/Strata/Strata.ush
+++ b/Engine/Shaders/Private/Strata/Strata.ush
@@ -329,37 +329,41 @@ float3 StrataUnpackNormal(uint PackedNormal);
 
 struct FStrataClassification
 {
-	uint ShadingModels;
+	uint State;
 	float TopLayerRoughness;
-	bool bIsStrataMaterial;
-	bool bHasSubsurface;
-	bool bHasSubsurfaceProfile;
-	bool bIsSimple;
 };
 
+#define CLAS_MASK_STATE					(0xFF)
+#define CLAS_MASK_SHADINGMODELS			(0x1F)
+#define CLAS_MASK_HASSUBSURFACE			(1u << 5)
+#define CLAS_MASK_HASSUBSURFACEPROFILE	(1u << 6)
+#define CLAS_MASK_ISSIMPLE				(1u << 7)
+
+uint GetShadingModels(const FStrataClassification In)		{ return (In.State & CLAS_MASK_SHADINGMODELS); }
+bool IsStrataMaterial(const FStrataClassification In)		{ return (In.State & CLAS_MASK_SHADINGMODELS) > 0; }
+bool HasSubsurface(const FStrataClassification In)			{ return (In.State & CLAS_MASK_HASSUBSURFACE) != 0; }
+bool HasSubsurfaceProfile(const FStrataClassification In)	{ return (In.State & CLAS_MASK_HASSUBSURFACEPROFILE) != 0; }
+bool IsSimpleMaterial(const FStrataClassification In)		{ return (In.State & CLAS_MASK_ISSIMPLE) != 0; }
+
+void AddShadingModels(inout FStrataClassification Out, uint bIn)		{ Out.State |= (bIn & CLAS_MASK_SHADINGMODELS); }
+void SetHasSubsurface(inout FStrataClassification Out, bool bIn)		{ Out.State |= (bIn ? CLAS_MASK_HASSUBSURFACE : 0u); }
+void SetHasSubsurfaceProfile(inout FStrataClassification Out, bool bIn)	{ Out.State |= (bIn ? CLAS_MASK_HASSUBSURFACEPROFILE : 0u); }
+void SetIsSimpleMaterial(inout FStrataClassification Out, bool bIn)		{ Out.State |= (bIn ? CLAS_MASK_ISSIMPLE : 0u); }
+void AndIsSimpleMaterial(inout FStrataClassification Out, bool bIn)		{ Out.State  = (Out.State & ~CLAS_MASK_ISSIMPLE) | ((Out.State & CLAS_MASK_ISSIMPLE) & (bIn ? CLAS_MASK_ISSIMPLE : 0u)); }
+
 uint StrataPackClassificationData(FStrataClassification In)
 {
 	const uint TopLayerRoughness8Bits = uint(In.TopLayerRoughness * 255.0f);
-	return 
-		(In.ShadingModels & 0x1F) | 
-		((TopLayerRoughness8Bits & 0xFF)	 << 8 ) |
-		((In.bHasSubsurface			? 1 : 0) << 5 ) | 
-		((In.bHasSubsurfaceProfile	? 1 : 0) << 6 ) |
-		((In.bIsSimple				? 1 : 0) << 7 );
+	return
+		(In.State & CLAS_MASK_STATE) |
+		((TopLayerRoughness8Bits & 0xFF) << 8);
 }
 
 FStrataClassification StrataUnpackClassificationData(uint In)
 {
-	uint TopLayerRoughness8Bits = 0;
 	FStrataClassification Out;
-	Out.ShadingModels			= (In >> 0 ) & 0x1F;
-	TopLayerRoughness8Bits		= (In >> 8) & 0xFF;
-	Out.bIsStrataMaterial		= Out.ShadingModels > 0;
-	Out.bHasSubsurface			= (In >> 5) & 0x1;
-	Out.bHasSubsurfaceProfile	= (In >> 6) & 0x1;
-	Out.bIsSimple				= (In >> 7) & 0x1;
-
-	Out.TopLayerRoughness = float(TopLayerRoughness8Bits) * (1.0f / 255.0f);
+	Out.State = In & CLAS_MASK_STATE;
+	Out.TopLayerRoughness = float((In >> 8) & 0xFF) * (1.0f / 255.0f);
 	return Out;
 }
 
@@ -444,6 +448,11 @@ bool StrataHasShadingModel(uint ShadingModels, uint ShadingModel)
 	return (ShadingModels & (1u << ShadingModel)) > 0;
 }
 
+bool StrataHasShadingModel(FStrataClassification In, uint ShadingModel)
+{
+	return StrataHasShadingModel(GetShadingModels(In), ShadingModel);
+}
+
 // Return true if a BSDF has a custom area light integrator
 bool StrataHasAreaLightIntegrator(in FStrataBSDF BSDF)
 {
diff --git a/Engine/Shaders/Private/Strata/StrataMaterialClassification.usf b/Engine/Shaders/Private/Strata/StrataMaterialClassification.usf
--- a/Engine/Shaders/Private/Strata/StrataMaterialClassification.usf
+++ b/Engine/Shaders/Private/Strata/StrataMaterialClassification.usf
@@ -53,13 +53,13 @@ void ComputePostProcessBSDFData(
 		TopLayer.WorldNormal += N * max(1e-5, BSDFWeightAvg); 
 	}
 
-	Classification.ShadingModels |= StrataShadingModelBit(BSDF_GETTYPE(BSDF));
+	AddShadingModels(Classification, StrataShadingModelBit(BSDF_GETTYPE(BSDF)));
 
 	switch (BSDF_GETTYPE(BSDF))
 	{
 	case STRATA_BSDF_TYPE_SLAB:
 	{
-		Classification.bIsSimple = Classification.bIsSimple && IsStrataSlabFastPathCompatible(BSDF);
+		AndIsSimpleMaterial(Classification, IsStrataSlabFastPathCompatible(BSDF));
 
 		if (bTopLayer)
 		{
@@ -72,11 +72,11 @@ void ComputePostProcessBSDFData(
 		if (!BSDF_GETISTHIN(BSDF) &&					// STRATA_TODO A material cannot have subsurface if using Thin mode. This is because DMFP does not match MFP at all (right now, SSS would just smears thin lighting model details out)
 			BSDF_GETHASSSS(BSDF) && BSDFWeightAvg > 0)
 		{
-			Classification.bHasSubsurface = true;
-			Classification.bHasSubsurfaceProfile = Classification.bHasSubsurfaceProfile || BSDF_GETHASSSSPROFILE(BSDF);;
+			SetHasSubsurface(Classification, true);
+			SetHasSubsurfaceProfile(Classification, BSDF_GETHASSSSPROFILE(BSDF));
 
 			SSS.BaseColor = SLAB_DIFFUSEALBEDO(BSDF);
-			SSS.Specular  = SLAB_F0(BSDF);				// STRATA_TODO: Specular or F0?
+			SSS.Specular  = SLAB_F0(BSDF).x;				// STRATA_TODO: Specular or F0?
 
 			// Profile ID: we only keep the latest ID, and profile will overtake any non-profile based SSS
 			SSS.bIsValid			= true;
@@ -93,7 +93,7 @@ void ComputePostProcessBSDFData(
 
 	case STRATA_BSDF_TYPE_HAIR:
 	{
-		Classification.bIsSimple = false;
+		AndIsSimpleMaterial(Classification, false);
 
 		// No SSR
 		if (bTopLayer)
@@ -106,7 +106,7 @@ void ComputePostProcessBSDFData(
 
 	case STRATA_BSDF_TYPE_SINGLELAYERWATER:
 	{
-		Classification.bIsSimple = false;
+		AndIsSimpleMaterial(Classification, false);
 
 		if (bTopLayer)
 		{
@@ -136,16 +136,15 @@ void MainPS(
 	FStrataAddressing StrataAddressing = GetStrataPixelDataByteOffset(PixelPos, uint2(View.BufferSizeAndInvSize.xy), MaxBytesPerPixel);
 	FStrataPixelHeader Header = UnpackStrataHeaderIn(MaterialLobesBuffer, StrataAddressing);
 
-	ClassificationData.bIsStrataMaterial = Header.BSDFCount > 0;
-	ClassificationData.bIsSimple = true; // non strata material or sky pixels are considered as simple.
-	if (ClassificationData.bIsStrataMaterial)
+	SetIsSimpleMaterial(ClassificationData, true); // non strata material or sky pixels are considered as simple.
+	if (Header.BSDFCount > 0)
 	{
 		float2 ScreenPosition = SvPositionToScreenPosition(SVPos).xy;
 		const float Depth = 1000000.0f;
 		float4 WorldPos = mul(float4(ScreenPosition * Depth, Depth, 1), LWCHackToFloat(PrimaryView.ScreenToWorld));
 		const float3 V = normalize(LWCHackToFloat(PrimaryView.WorldCameraOrigin) - WorldPos.xyz);
 
-		ClassificationData.bIsSimple = Header.BSDFCount == 1;
+		SetIsSimpleMaterial(ClassificationData, Header.BSDFCount == 1);
 
 		float TopLayerRoughness = 0.0f;
 		float TopLayerRoughnessAverageCount = 0.0f;
@@ -212,8 +211,8 @@ void TileMainCS(uint2 DispatchThreadId : SV_DispatchThreadID, uint LinearIndex :
 	{
 		const uint PackClassification = ClassificationTexture.Load(uint3(PixelCoord, 0));
 		FStrataClassification Data = StrataUnpackClassificationData(PackClassification);
-		bContainsComplexMaterial = !Data.bIsSimple;
-		bContainsStrataMaterial = Data.bIsStrataMaterial;
+		bContainsComplexMaterial = !IsSimpleMaterial(Data);
+		bContainsStrataMaterial = IsStrataMaterial(Data);
 	}
 	s_IsComplexIsMaterialTile[LinearIndex] = (bContainsComplexMaterial ? 1 : 0) | (bContainsStrataMaterial ? 2 : 0);
 
diff --git a/Engine/Shaders/Private/Strata/StrataVisualize.usf b/Engine/Shaders/Private/Strata/StrataVisualize.usf
--- a/Engine/Shaders/Private/Strata/StrataVisualize.usf
+++ b/Engine/Shaders/Private/Strata/StrataVisualize.usf
@@ -239,7 +239,7 @@ void VisualizeMaterialPS(
 
 		PRINTCLASSIFROUGH();
 		NEXTLINE;
-		if (StrataData.bHasSubsurface)
+		if (HasSubsurface(StrataData))
 		{
 			PRINTCLASSIFBASCOL();
 			NEXTLINE;
@@ -278,7 +278,7 @@ void VisualizeMaterialPS(
 		NEXTLINE;
 
 		NEXTLINE;
-		if (StrataData.bHasSubsurface)
+		if (HasSubsurface(StrataData))
 		{
 			NEXTLINE;
 			NEXTLINE;
diff --git a/Engine/Shaders/Private/VirtualShadowMaps/TransmissionCommon.ush b/Engine/Shaders/Private/VirtualShadowMaps/TransmissionCommon.ush
--- a/Engine/Shaders/Private/VirtualShadowMaps/TransmissionCommon.ush
+++ b/Engine/Shaders/Private/VirtualShadowMaps/TransmissionCommon.ush
@@ -25,9 +25,9 @@ half GetSubsurfaceOpacityFromGbuffer(uint2 PixelPos)
 #if STRATA_ENABLED && 0
 	// TODO: VSM projection shaders are not set up for strata yet, but added a skeleton of the logic
 	const FStrataClassification StrataData = StrataUnpackClassificationData(Strata.ClassificationTexture.Load(PixelPos));
-	if (StrataData.bIsStrataMaterial)
+	if (IsStrataMaterial(StrataData))
 	{
-		if (StrataData.bHasSubsurface)
+		if (HasSubsurface(StrataData))
 		{
 			const FStrataSubsurfaceData SSSData = StrataUnpackSSSData(Strata.SSSTexture.Load(uint3(SVPos.xy, 0)));
 			if (SSSData.IsValid())
