Remove deadcode from strata classification.

diff --git a/Engine/Shaders/Private/Strata/StrataMaterialClassification.usf b/Engine/Shaders/Private/Strata/StrataMaterialClassification.usf
--- a/Engine/Shaders/Private/Strata/StrataMaterialClassification.usf
+++ b/Engine/Shaders/Private/Strata/StrataMaterialClassification.usf
@@ -40,11 +40,6 @@ void ComputePostProcessBSDFData(
 	float3 N = TangentBasis[2];
 	//N = normalize(N);
 
-
-	BxDFContext ContextNoLight = (BxDFContext)0;
-	const float3 FakeL = float3(0, 0, 1);
-	Init(ContextNoLight, N, V, FakeL);
-
 	const bool bTopLayer = BSDF_GETISTOPLAYER(BSDF);
 	if (bTopLayer)
 	{
