Fix typo in Strata classification introduced in previous CL.

diff --git a/Engine/Shaders/Private/Strata/StrataMaterialClassification.usf b/Engine/Shaders/Private/Strata/StrataMaterialClassification.usf
--- a/Engine/Shaders/Private/Strata/StrataMaterialClassification.usf
+++ b/Engine/Shaders/Private/Strata/StrataMaterialClassification.usf
@@ -139,7 +139,7 @@ void MainPS(
 		float4 WorldPos = mul(float4(ScreenPosition * Depth, Depth, 1), LWCHackToFloat(PrimaryView.ScreenToWorld));
 		const float3 V = normalize(LWCHackToFloat(PrimaryView.WorldCameraOrigin) - WorldPos.xyz);
 
-		SetIsSimpleMaterial(ClassificationData, Header.BSDFCount == 1);
+		AndIsSimpleMaterial(ClassificationData, Header.BSDFCount == 1);
 
 		float TopLayerRoughness = 0.0f;
 		float TopLayerRoughnessAverageCount = 0.0f;
