Strata store/load functions cleanup

diff --git a/Engine/Shaders/Private/Strata/Strata.ush b/Engine/Shaders/Private/Strata/Strata.ush
--- a/Engine/Shaders/Private/Strata/Strata.ush
+++ b/Engine/Shaders/Private/Strata/Strata.ush
@@ -736,34 +736,27 @@ FStrataAddressing GetStrataPixelDataByteOffset(uint2 PixelPos, uint2 ViewBufferS
 
 #if COMPILER_SWITCH // These platforms are complaining so simply disabling Strata there for now.
 
+void StrataStoreUint1(RWByteAddressBuffer StrataBuffer, inout FStrataAddressing StrataAddressing, uint Data) {}
 #define STRATA_STORE_UINT1(x) 
 
-#define STRATA_LOAD_UINT1(x) x = 0;
-
-#define STRATA_LOAD_FLOAT1(x) x = 0;
-
-uint  StrataLoadUint1(ByteAddressBuffer StrataBuffer, inout FStrataAddressing StrataAddressing) { return 0; }
+uint StrataLoadUint1(ByteAddressBuffer StrataBuffer, inout FStrataAddressing StrataAddressing) { return 0; }
 
 #else
 
-#define STRATA_STORE_UINT1(x)	StrataBuffer.Store (StrataAddressing.CurrentAddr, asuint(x)); \
-								StrataAddressing.CurrentAddr += StrataAddressing.ByteOffsetToNextUint; \
-								StrataAddressing.ReadBytes += 4;
-
-#define STRATA_LOAD_UINT1(x)	x = StrataBuffer.Load (StrataAddressing.CurrentAddr); \
-								StrataAddressing.CurrentAddr += StrataAddressing.ByteOffsetToNextUint; \
-								StrataAddressing.ReadBytes += 4;
-
-#define STRATA_LOAD_FLOAT1(x)	x = asfloat(StrataBuffer.Load (StrataAddressing.CurrentAddr)); \
-								StrataAddressing.CurrentAddr += StrataAddressing.ByteOffsetToNextUint; \
-								StrataAddressing.ReadBytes += 4;
+void StrataStoreUint1(RWByteAddressBuffer StrataBuffer, inout FStrataAddressing StrataAddressing, uint Data)
+{
+	StrataBuffer.Store(StrataAddressing.CurrentAddr, Data);
+	StrataAddressing.CurrentAddr += StrataAddressing.ByteOffsetToNextUint;
+	StrataAddressing.ReadBytes += 4;
+}
+#define STRATA_STORE_UINT1(x)	StrataStoreUint1(StrataBuffer, StrataAddressing, x)
 
-uint  StrataLoadUint1(ByteAddressBuffer StrataBuffer, inout FStrataAddressing StrataAddressing) 
-{ 
-	uint  x = StrataBuffer.Load(StrataAddressing.CurrentAddr); 
-	StrataAddressing.CurrentAddr += StrataAddressing.ByteOffsetToNextUint;	
+uint StrataLoadUint1(ByteAddressBuffer StrataBuffer, inout FStrataAddressing StrataAddressing)
+{
+	uint Data = StrataBuffer.Load(StrataAddressing.CurrentAddr);
+	StrataAddressing.CurrentAddr += StrataAddressing.ByteOffsetToNextUint;
 	StrataAddressing.ReadBytes += 4;
-	return x; 
+	return Data;
 }
 
 #endif
@@ -2051,8 +2044,7 @@ FStrataPixelHeader UnpackStrataHeaderIn(ByteAddressBuffer StrataBuffer, inout FS
 {
 	FStrataPixelHeader Out = InitialiseStrataPixelHeader();
 
-	uint PackedHeader = 0;
-	STRATA_LOAD_UINT1(PackedHeader);
+	uint PackedHeader = StrataLoadUint1(StrataBuffer, StrataAddressing);
 	Out.BSDFCount    = HEADER_GETBSDFCOUNT(PackedHeader);
 	uint SharedLocalBasesCount = HEADER_GETSHAREDLOCALBASESCOUNT(PackedHeader);
 	uint SharedLocalBasesTypes = HEADER_GETSHAREDLOCALBASESTYPE(PackedHeader);
@@ -2143,7 +2135,7 @@ FStrataBSDF UnpackStrataBSDFIn(ByteAddressBuffer StrataBuffer, inout FStrataAddr
 
 	float DummyFloat = 0.0f;
 
-	STRATA_LOAD_UINT1(OutBSDF.State);
+	OutBSDF.State = StrataLoadUint1(StrataBuffer, StrataAddressing);
 	if (BSDF_GETHASGREYWEIGHT(OutBSDF))
 	{
 		OutBSDF.FinalWeight = Unpack10F(BSDF_GETWEIGHT10F(OutBSDF));
