Change Strata header to keep flags compressed to reduce vgpr usage.

diff --git a/Engine/Shaders/Private/BasePassPixelShader.usf b/Engine/Shaders/Private/BasePassPixelShader.usf
--- a/Engine/Shaders/Private/BasePassPixelShader.usf
+++ b/Engine/Shaders/Private/BasePassPixelShader.usf
@@ -1000,20 +1000,19 @@ void FPixelShaderInOut_MainPS(
 	StrataPixelHeader.SharedLocalBases = MaterialParameters.SharedLocalBases;
 
 #if !GBUFFER_HAS_PRECSHADOWFACTOR
-	StrataPixelHeader.bSkipPrecShadowMask = true;
+	SetSkipPrecShadowMask(StrataPixelHeader, true);
 #endif
 #if (GBUFFER_HAS_PRECSHADOWFACTOR && WRITES_PRECSHADOWFACTOR_ZERO)
-	StrataPixelHeader.bZeroPrecShadowMask = true;
+	SetZeroPrecShadowMask(StrataPixelHeader, true);
 #endif
 
 	StrataPixelHeader.MaterialAO = MaterialAO;
-
-	StrataPixelHeader.bCastContactShadow = GetPrimitiveData(MaterialParameters).Flags & PRIMITIVE_SCENE_DATA_FLAG_HAS_CAST_CONTACT_SHADOW;
-	StrataPixelHeader.bHasDynamicIndirectShadowCasterRepresentation = GetPrimitiveData(MaterialParameters).Flags & PRIMITIVE_SCENE_DATA_FLAG_HAS_CAPSULE_REPRESENTATION;
+	SetCastContactShadow(StrataPixelHeader, GetPrimitiveData(MaterialParameters).Flags & PRIMITIVE_SCENE_DATA_FLAG_HAS_CAST_CONTACT_SHADOW);
+	SetDynamicIndirectShadowCasterRepresentation(StrataPixelHeader, GetPrimitiveData(MaterialParameters).Flags & PRIMITIVE_SCENE_DATA_FLAG_HAS_CAPSULE_REPRESENTATION);
 #endif
 
 #if STRATA_ENABLED && STRATA_INLINE_SINGLELAYERWATER
-	StrataPixelHeader.bIsSingleLayerWater = true;
+	SetIsSingleLayerWater(StrataPixelHeader, true);
 
 	// Override GBuffer data with Strata SLW water BSDF to run forward shadfing code.
 	// STRATA_TODO: run the shading through a strata path? (by adding a special BSDF?)
diff --git a/Engine/Shaders/Private/CapsuleShadowShaders.usf b/Engine/Shaders/Private/CapsuleShadowShaders.usf
--- a/Engine/Shaders/Private/CapsuleShadowShaders.usf
+++ b/Engine/Shaders/Private/CapsuleShadowShaders.usf
@@ -965,7 +965,7 @@ void ApplySelfShadowingIntensityForDeferred(float2 ScreenUV, inout float Visibil
 		uint2 PixelPos = View.ViewRectMin.xy + View.ViewSizeAndInvSize.xy * ScreenUV.xy;
 		FStrataAddressing StrataAddressing = GetStrataPixelDataByteOffset(PixelPos, uint2(View.BufferSizeAndInvSize.xy), Strata.MaxBytesPerPixel);
 		FStrataPixelHeader StrataPixelHeader = UnpackStrataHeaderIn(Strata.MaterialLobesBuffer, StrataAddressing);
-		const bool bHasDynamicIndirectShadowCasterRepresentation = StrataPixelHeader.bHasDynamicIndirectShadowCasterRepresentation;
+		const bool bHasDynamicIndirectShadowCasterRepresentation = HasDynamicIndirectShadowCasterRepresentation(StrataPixelHeader);
 #else
 		FGBufferData GBufferData = GetGBufferData(ScreenUV);
 		const bool bHasDynamicIndirectShadowCasterRepresentation = HasDynamicIndirectShadowCasterRepresentation(GBufferData);
diff --git a/Engine/Shaders/Private/DeferredLightPixelShaders.usf b/Engine/Shaders/Private/DeferredLightPixelShaders.usf
--- a/Engine/Shaders/Private/DeferredLightPixelShaders.usf
+++ b/Engine/Shaders/Private/DeferredLightPixelShaders.usf
@@ -170,7 +170,7 @@ float3 GetStableLargeWorldPosition(float2 SVPos, float DeviceZ)
 
 float4 StrataReadPrecomputedShadowFactors(FStrataPixelHeader StrataPixelHeader, int2 PixelPos)
 {
-	if (!StrataPixelHeader.bSkipPrecShadowMask)
+	if (!HasSkipPrecShadowMask(StrataPixelHeader))
 	{
 #if ALLOW_STATIC_LIGHTING
 		float4 GBufferE = SceneTexturesStruct.GBufferETexture.Load(int3(PixelPos, 0));
@@ -179,7 +179,7 @@ float4 StrataReadPrecomputedShadowFactors(FStrataPixelHeader StrataPixelHeader,
 #endif
 		return GBufferE;
 	}
-	return StrataPixelHeader.bZeroPrecShadowMask ? 0.0f : 1.0f;
+	return HasZeroPrecShadowMask(StrataPixelHeader) ? 0.0f : 1.0f;
 }
 
 #endif // STRATA_ENABLED
diff --git a/Engine/Shaders/Private/DeferredLightingCommon.ush b/Engine/Shaders/Private/DeferredLightingCommon.ush
--- a/Engine/Shaders/Private/DeferredLightingCommon.ush
+++ b/Engine/Shaders/Private/DeferredLightingCommon.ush
@@ -150,7 +150,7 @@ float ShadowRayCast(
 		uint2 PixelPos = View.ViewRectMin.xy + View.ViewSizeAndInvSize.xy * SampleUVz.xy;
 		FStrataAddressing StrataAddressing = GetStrataPixelDataByteOffset(PixelPos, uint2(View.BufferSizeAndInvSize.xy), StrataStruct.MaxBytesPerPixel);
 		FStrataPixelHeader StrataPixelHeader = UnpackStrataHeaderIn(StrataStruct.MaterialLobesBuffer, StrataAddressing);
-		bOutHitCastContactShadow = StrataPixelHeader.bCastContactShadow;
+		bOutHitCastContactShadow = DoesCastContactShadow(StrataPixelHeader);
 		#endif
 	#else
 		FGBufferData SampleGBuffer = GetGBufferData( SampleUVz.xy );
diff --git a/Engine/Shaders/Private/SSRT/SSRTReflections.usf b/Engine/Shaders/Private/SSRT/SSRTReflections.usf
--- a/Engine/Shaders/Private/SSRT/SSRTReflections.usf
+++ b/Engine/Shaders/Private/SSRT/SSRTReflections.usf
@@ -144,7 +144,7 @@ void ScreenSpaceReflections(
 		// So in this case we simply load the Header and water BSDF directly because we know there can only be a single one.
 		FStrataAddressing StrataAddressing = GetStrataPixelDataByteOffset(PixelCoord, uint2(View.BufferSizeAndInvSize.xy), Strata.MaxBytesPerPixel);
 		FStrataPixelHeader StrataPixelHeader = UnpackStrataHeaderIn(Strata.MaterialLobesBuffer, StrataAddressing);
-		if (StrataPixelHeader.BSDFCount > 0 && StrataPixelHeader.bIsSingleLayerWater)
+		if (StrataPixelHeader.BSDFCount > 0 && IsSingleLayerWater(StrataPixelHeader))
 		{
 			const FStrataBSDF SLWaterBSDF = UnpackStrataBSDFIn(Strata.MaterialLobesBuffer, StrataAddressing);
 			N = StrataGetBSDFSharedBasis(StrataPixelHeader, SLWaterBSDF, StrataAddressing)[2];
diff --git a/Engine/Shaders/Private/SingleLayerWaterComposite.usf b/Engine/Shaders/Private/SingleLayerWaterComposite.usf
--- a/Engine/Shaders/Private/SingleLayerWaterComposite.usf
+++ b/Engine/Shaders/Private/SingleLayerWaterComposite.usf
@@ -59,7 +59,7 @@ SingleLayerWaterCompositeOutput SingleLayerWaterComposite(float2 BufferUV, float
 	// Water do nto read the classification data to avoid running it a second time when water is enabled.
 	FStrataAddressing StrataAddressing = GetStrataPixelDataByteOffset(SvPosition.xy, uint2(View.BufferSizeAndInvSize.xy), Strata.MaxBytesPerPixel);
 	FStrataPixelHeader StrataPixelHeader = UnpackStrataHeaderIn(Strata.MaterialLobesBuffer, StrataAddressing);
-	if (StrataPixelHeader.BSDFCount > 0 && StrataPixelHeader.bIsSingleLayerWater)
+	if (StrataPixelHeader.BSDFCount > 0 && IsSingleLayerWater(StrataPixelHeader))
 	{
 		float DeviceZ = SampleDeviceZFromSceneTextures(BufferUV);
 		float SceneDepth = ConvertFromDeviceZ(DeviceZ);
@@ -253,7 +253,7 @@ void WaterTileCatergorisationCS(uint3 ThreadId : SV_DispatchThreadID, uint3 Grou
 			#if STRATA_ENABLED
 				FStrataAddressing StrataAddressing = GetStrataPixelDataByteOffset(PixelPos, uint2(View.BufferSizeAndInvSize.xy), Strata.MaxBytesPerPixel);
 				FStrataPixelHeader StrataPixelHeader = UnpackStrataHeaderIn(Strata.MaterialLobesBuffer, StrataAddressing);
-				bContainsWater = bContainsWater || (StrataPixelHeader.BSDFCount > 0 && StrataPixelHeader.bIsSingleLayerWater);
+				bContainsWater = bContainsWater || (StrataPixelHeader.BSDFCount > 0 && IsSingleLayerWater(StrataPixelHeader));
 			#else
 				FGBufferData GBuffer = GetGBufferDataFromSceneTextures(BufferUV);
 				bContainsWater = bContainsWater || GBuffer.ShadingModelID == SHADINGMODELID_SINGLELAYERWATER;
@@ -279,7 +279,7 @@ void WaterTileCatergorisationCS(uint3 ThreadId : SV_DispatchThreadID, uint3 Grou
 	#if STRATA_ENABLED
 		FStrataAddressing StrataAddressing = GetStrataPixelDataByteOffset(PixelPos, uint2(View.BufferSizeAndInvSize.xy), Strata.MaxBytesPerPixel);
 		FStrataPixelHeader StrataPixelHeader = UnpackStrataHeaderIn(Strata.MaterialLobesBuffer, StrataAddressing);
-		bContainsWater = bContainsWater || (StrataPixelHeader.BSDFCount > 0 && StrataPixelHeader.bIsSingleLayerWater);
+		bContainsWater = bContainsWater || (StrataPixelHeader.BSDFCount > 0 && IsSingleLayerWater(StrataPixelHeader));
 	#else
 		FGBufferData GBuffer = GetGBufferDataFromSceneTextures(BufferUV);
 		bContainsWater = bContainsWater || GBuffer.ShadingModelID == SHADINGMODELID_SINGLELAYERWATER;
diff --git a/Engine/Shaders/Private/Strata/Strata.ush b/Engine/Shaders/Private/Strata/Strata.ush
--- a/Engine/Shaders/Private/Strata/Strata.ush
+++ b/Engine/Shaders/Private/Strata/Strata.ush
@@ -103,22 +103,14 @@ struct FStrataBSDF
 #define HEADER_BIT_COUNT_SHAREDLOCALBASES_COUNT	3
 #define HEADER_BIT_COUNT_SHAREDLOCALBASES_TYPE	4
 #define HEADER_BIT_COUNT_AO						8
-#define HEADER_BIT_COUNT_SKIP_PRECSHADOW_MASK	1
-#define HEADER_BIT_COUNT_ZERO_PRECSHADOW_MASK	1
-#define HEADER_BIT_COUNT_ISSINGLELAYERWATER		1
-#define HEADER_BIT_COUNT_CASTCONTACTSHADOW		1
-#define HEADER_BIT_COUNT_DYNINDSHADOWCASTER		1
+#define HEADER_BIT_COUNT_STATE					5
 #define HEADER_BIT_COUNT___UNUSED___			8
 
 #define HEADER_BIT_OFFSET_BSDF_COUNT			0
 #define HEADER_BIT_OFFSET_SHAREDLOCALBASES_COUNT (HEADER_BIT_OFFSET_BSDF_COUNT				+ HEADER_BIT_COUNT_BSDF_COUNT)
 #define HEADER_BIT_OFFSET_SHAREDLOCALBASES_TYPE	(HEADER_BIT_OFFSET_SHAREDLOCALBASES_COUNT	+ HEADER_BIT_COUNT_SHAREDLOCALBASES_COUNT)
 #define HEADER_BIT_OFFSET_AO					(HEADER_BIT_OFFSET_SHAREDLOCALBASES_TYPE	+ HEADER_BIT_COUNT_SHAREDLOCALBASES_TYPE)
-#define HEADER_BIT_OFFSET_SKIP_PRECSHADOW_MASK	(HEADER_BIT_OFFSET_AO						+ HEADER_BIT_COUNT_AO)
-#define HEADER_BIT_OFFSET_ZERO_PRECSHADOW_MASK	(HEADER_BIT_OFFSET_SKIP_PRECSHADOW_MASK		+ HEADER_BIT_COUNT_SKIP_PRECSHADOW_MASK)
-#define HEADER_BIT_OFFSET_ISSINGLELAYERWATER	(HEADER_BIT_OFFSET_ZERO_PRECSHADOW_MASK		+ HEADER_BIT_COUNT_ZERO_PRECSHADOW_MASK)
-#define HEADER_BIT_OFFSET_CASTCONTACTSHADOW		(HEADER_BIT_OFFSET_ISSINGLELAYERWATER		+ HEADER_BIT_COUNT_ISSINGLELAYERWATER)
-#define HEADER_BIT_OFFSET_DYNINDSHADOWCASTER	(HEADER_BIT_OFFSET_CASTCONTACTSHADOW		+ HEADER_BIT_COUNT_CASTCONTACTSHADOW)
+#define HEADER_BIT_OFFSET_STATE					(HEADER_BIT_OFFSET_AO						+ HEADER_BIT_COUNT_AO)
 
 #define HEADER_GETBSDFCOUNT(X)						READ_BITS( X, HEADER_BIT_COUNT_BSDF_COUNT, HEADER_BIT_OFFSET_BSDF_COUNT)
 #define HEADER_SETBSDFCOUNT(X, NewValue)			WRITE_BITS(X, HEADER_BIT_COUNT_BSDF_COUNT, HEADER_BIT_OFFSET_BSDF_COUNT, NewValue)
@@ -132,20 +124,8 @@ struct FStrataBSDF
 #define HEADER_GETAO(X)								READ_BITS( X, HEADER_BIT_COUNT_AO, HEADER_BIT_OFFSET_AO)
 #define HEADER_SETAO(X, NewValue)					WRITE_BITS(X, HEADER_BIT_COUNT_AO, HEADER_BIT_OFFSET_AO, NewValue)
 
-#define HEADER_GETSKIPPRECSHADOWMASK(X)				READ_BITS( X, HEADER_BIT_COUNT_SKIP_PRECSHADOW_MASK, HEADER_BIT_OFFSET_SKIP_PRECSHADOW_MASK)
-#define HEADER_SETSKIPPRECSHADOWMASK(X, NewValue)	WRITE_BITS(X, HEADER_BIT_COUNT_SKIP_PRECSHADOW_MASK, HEADER_BIT_OFFSET_SKIP_PRECSHADOW_MASK, NewValue)
-
-#define HEADER_GETZEROPRECSHADOWMASK(X)				READ_BITS( X, HEADER_BIT_COUNT_ZERO_PRECSHADOW_MASK, HEADER_BIT_OFFSET_ZERO_PRECSHADOW_MASK)
-#define HEADER_SETZEROPRECSHADOWMASK(X, NewValue)	WRITE_BITS(X, HEADER_BIT_COUNT_ZERO_PRECSHADOW_MASK, HEADER_BIT_OFFSET_ZERO_PRECSHADOW_MASK, NewValue)
-
-#define HEADER_GETISSINGLELAYERWATER(X)				READ_BITS( X, HEADER_BIT_COUNT_ISSINGLELAYERWATER, HEADER_BIT_OFFSET_ISSINGLELAYERWATER)
-#define HEADER_SETISSINGLELAYERWATER(X, NewValue)	WRITE_BITS(X, HEADER_BIT_COUNT_ISSINGLELAYERWATER, HEADER_BIT_OFFSET_ISSINGLELAYERWATER, NewValue)
-
-#define HEADER_GETCASTCONTACTSHADOW(X)				READ_BITS( X, HEADER_BIT_COUNT_CASTCONTACTSHADOW, HEADER_BIT_OFFSET_CASTCONTACTSHADOW)
-#define HEADER_SETCASTCONTACTSHADOW(X, NewValue)	WRITE_BITS(X, HEADER_BIT_COUNT_CASTCONTACTSHADOW, HEADER_BIT_OFFSET_CASTCONTACTSHADOW, NewValue)
-
-#define HEADER_GETDYNINDSHADOWCASTER(X)				READ_BITS( X, HEADER_BIT_COUNT_DYNINDSHADOWCASTER, HEADER_BIT_OFFSET_DYNINDSHADOWCASTER)
-#define HEADER_SETDYNINDSHADOWCASTER(X, NewValue)	WRITE_BITS(X, HEADER_BIT_COUNT_DYNINDSHADOWCASTER, HEADER_BIT_OFFSET_DYNINDSHADOWCASTER, NewValue)
+#define HEADER_GETSTATE(X)							READ_BITS( X, HEADER_BIT_COUNT_STATE, HEADER_BIT_OFFSET_STATE)
+#define HEADER_SETSTATE(X, NewValue)				WRITE_BITS(X, HEADER_BIT_COUNT_STATE, HEADER_BIT_OFFSET_STATE, NewValue)
 
 // Bit count for each component stored in a Strata BSDF state
 #define STATE_BIT_COUNT_BSDF					3
@@ -311,8 +291,9 @@ struct FStrataBSDF
 #define SLW_COLORSCALEBEHINDWATER(X)	X.InlineVGPRs[2].xyz
 #endif
 
-
+///////////////////////////////////////////////////////////////////////////////
 // Forward declarations
+
 FParticipatingMedia StrataSlabCreateParticipatingMedia(float3 DiffuseColor, float3 MeanFreePathCentimeters);
 void EnableSlabBSDFSimpleVolumetric(inout FStrataBSDF BSDF);
 void StrataRequestSharedLocalBasisTangent(inout uint Types, uint Index);
@@ -327,6 +308,9 @@ float4 UnpackColorGamma2ToLinearAlphaLinear(uint rgba);
 uint   StrataPackNormal(in float3 Normal);
 float3 StrataUnpackNormal(uint PackedNormal);
 
+///////////////////////////////////////////////////////////////////////////////
+// Classification 
+
 struct FStrataClassification
 {
 	uint State;
@@ -367,6 +351,8 @@ FStrataClassification StrataUnpackClassificationData(uint In)
 	return Out;
 }
 
+///////////////////////////////////////////////////////////////////////////////
+// Top layer
 struct FStrataTopLayerData
 {
 	float3 WorldNormal;
@@ -385,6 +371,9 @@ FStrataTopLayerData StrataUnpackTopLayerData(uint In)
 	return Out;
 }
 
+///////////////////////////////////////////////////////////////////////////////
+// Sub-surface 
+
 struct FStrataSubsurfaceData
 {
 	float3 BaseColor;
@@ -465,6 +454,8 @@ bool StrataHasGreyScaleWeight(in float3 Weight)
 	return Weight.x == Weight.y && Weight.y == Weight.z;
 }
 
+///////////////////////////////////////////////////////////////////////////////
+// Layer & BSDF description
 
 struct FStrataLayer
 {
@@ -500,6 +491,9 @@ FStrataData GetInitialisedStrataDataWithReservedFirstSlot()
 	return StrataData;
 }
 
+///////////////////////////////////////////////////////////////////////////////
+// BSDF
+
 FStrataData GetStrataSlabBSDF(
 	float UseMetalness,
 	float3 BaseColor, float3 EdgeColor, float Specular, float Metallic,	//  Metalness workflow
@@ -791,6 +785,9 @@ uint StrataGetSharedLocalBasisType(in uint Types, uint Index)
 	return (Types >> Index) & 0x1 ? STRATA_BASIS_TYPE_TANGENT : STRATA_BASIS_TYPE_NORMAL;
 }
 
+///////////////////////////////////////////////////////////////////////////////
+// Header
+
 struct FStrataPixelHeader
 {
 	uint BSDFCount;
@@ -805,13 +802,27 @@ struct FStrataPixelHeader
 #endif
 
 	float MaterialAO;
-	bool bSkipPrecShadowMask;
-	bool bZeroPrecShadowMask;
-	bool bIsSingleLayerWater;	// Used to avoid loading BSDF or Categorisation data when generating tiles containing SLW material. Avoid running another categorisation.
-	bool bCastContactShadow;
-	bool bHasDynamicIndirectShadowCasterRepresentation;
+	uint State;
 };
 
+#define HEADER_MASK_SKIPPRECSHADOWMASK			(1u << 0)
+#define HEADER_MASK_ZEROPRECSHADOWMASK			(1u << 1)
+#define HEADER_MASK_ISSINGLELAYERWATER			(1u << 2)	// Used to avoid loading BSDF or Categorisation data when generating tiles containing SLW material. Avoid running another categorisation.
+#define HEADER_MASK_CASTCONTACTSHADOW			(1u << 3)
+#define HEADER_MASK_HASDYNINDIRECTSHADOWCASTER	(1u << 4)
+
+bool HasSkipPrecShadowMask(const FStrataPixelHeader In)										{ return (In.State & HEADER_MASK_SKIPPRECSHADOWMASK) != 0; }
+bool HasZeroPrecShadowMask(const FStrataPixelHeader In)										{ return (In.State & HEADER_MASK_ZEROPRECSHADOWMASK) != 0; }
+bool IsSingleLayerWater(const FStrataPixelHeader In)										{ return (In.State & HEADER_MASK_ISSINGLELAYERWATER) != 0; }
+bool DoesCastContactShadow(const FStrataPixelHeader In)										{ return (In.State & HEADER_MASK_CASTCONTACTSHADOW) != 0; }
+bool HasDynamicIndirectShadowCasterRepresentation(const FStrataPixelHeader In)				{ return (In.State & HEADER_MASK_HASDYNINDIRECTSHADOWCASTER) != 0; }
+
+void SetSkipPrecShadowMask(inout FStrataPixelHeader Out, bool bIn)							{ Out.State |= (bIn ? HEADER_MASK_SKIPPRECSHADOWMASK : 0u); }
+void SetZeroPrecShadowMask(inout FStrataPixelHeader Out, bool bIn)							{ Out.State |= (bIn ? HEADER_MASK_ZEROPRECSHADOWMASK : 0u); }
+void SetIsSingleLayerWater(inout FStrataPixelHeader Out, bool bIn)							{ Out.State |= (bIn ? HEADER_MASK_ISSINGLELAYERWATER : 0u); }
+void SetCastContactShadow (inout FStrataPixelHeader Out, bool bIn)							{ Out.State |= (bIn ? HEADER_MASK_CASTCONTACTSHADOW : 0u); }
+void SetDynamicIndirectShadowCasterRepresentation(inout FStrataPixelHeader Out, bool bIn)	{ Out.State |= (bIn ? HEADER_MASK_HASDYNINDIRECTSHADOWCASTER : 0u); }
+
 FStrataPixelHeader InitialiseStrataPixelHeader()
 {
 	FStrataPixelHeader Out;
@@ -823,11 +834,7 @@ FStrataPixelHeader InitialiseStrataPixelHeader()
 	Out.SharedLocalBasesTypes = 0;
 #endif
 	Out.MaterialAO = 1.0f;
-	Out.bSkipPrecShadowMask = false;
-	Out.bZeroPrecShadowMask = false;
-	Out.bIsSingleLayerWater = false;
-	Out.bCastContactShadow = false;
-	Out.bHasDynamicIndirectShadowCasterRepresentation = false;
+	Out.State = 0;
 	return Out;
 }
 
@@ -2029,12 +2036,7 @@ uint PackStrataHeaderIn(uint InBSDFCount, FStrataPixelHeader InHeader)
 #endif
 
 	HEADER_SETAO(Out, uint(InHeader.MaterialAO * 255.0f));
-
-	HEADER_SETSKIPPRECSHADOWMASK(Out, InHeader.bSkipPrecShadowMask ? 1u : 0u);
-	HEADER_SETZEROPRECSHADOWMASK(Out, InHeader.bZeroPrecShadowMask ? 1u : 0u);
-	HEADER_SETISSINGLELAYERWATER(Out, InHeader.bIsSingleLayerWater ? 1u : 0u);
-	HEADER_SETCASTCONTACTSHADOW(Out, InHeader.bCastContactShadow ? 1u : 0u);
-	HEADER_SETDYNINDSHADOWCASTER(Out, InHeader.bHasDynamicIndirectShadowCasterRepresentation ? 1u : 0u);
+	HEADER_SETSTATE(Out, InHeader.State);
 	return Out;
 }
 
@@ -2063,12 +2065,7 @@ FStrataPixelHeader UnpackStrataHeaderIn(ByteAddressBuffer StrataBuffer, inout FS
 #endif
 
 	Out.MaterialAO = float(HEADER_GETAO(PackedHeader)) * (1.0f / 255.0f);
-
-	Out.bSkipPrecShadowMask = HEADER_GETSKIPPRECSHADOWMASK(PackedHeader);
-	Out.bZeroPrecShadowMask = HEADER_GETZEROPRECSHADOWMASK(PackedHeader);
-	Out.bIsSingleLayerWater = HEADER_GETISSINGLELAYERWATER(PackedHeader);
-	Out.bCastContactShadow = HEADER_GETCASTCONTACTSHADOW(PackedHeader);
-	Out.bHasDynamicIndirectShadowCasterRepresentation = HEADER_GETDYNINDSHADOWCASTER(PackedHeader);
+	Out.State = HEADER_GETSTATE(PackedHeader);
 
 	return Out;
 }
diff --git a/Engine/Shaders/Private/Strata/StrataForwardLighting.ush b/Engine/Shaders/Private/Strata/StrataForwardLighting.ush
--- a/Engine/Shaders/Private/Strata/StrataForwardLighting.ush
+++ b/Engine/Shaders/Private/Strata/StrataForwardLighting.ush
@@ -54,7 +54,7 @@ float3 StrataForwardLightingCommon(
 		// Evaluate the ShadowTerm that can then be used when integrating the lighting
 		FShadowTerms ShadowTerms = { StrataPixelHeader.MaterialAO, 1.0, 1.0, InitHairTransmittanceData() };
 
-		float4 PrecomputedShadowFactors = StrataPixelHeader.bZeroPrecShadowMask ? 0.0f : 1.0f;
+		float4 PrecomputedShadowFactors = HasZeroPrecShadowMask(StrataPixelHeader) ? 0.0f : 1.0f;
 	#if TRANSLUCENCY_LIGHTING_SURFACE_FORWARDSHADING
 		PrecomputedShadowFactors.x = ComputeDirectionalLightStaticShadowing(AbsoluteWorldPosition).x;
 	#else
@@ -342,7 +342,7 @@ float3 StrataForwardLighting(
 				#if FORWARD_SHADING && (MATERIALBLENDING_SOLID || MATERIALBLENDING_MASKED)
 					float2 NearestResolvedDepthScreenUV = CalculateNearestResolvedDepthScreenUV(ScreenUV, MaterialParameters.ScreenPosition.w);
 
-					IndirectOcclusion = GetIndirectOcclusion(NearestResolvedDepthScreenUV, StrataPixelHeader.bHasDynamicIndirectShadowCasterRepresentation);
+					IndirectOcclusion = GetIndirectOcclusion(NearestResolvedDepthScreenUV, HasDynamicIndirectShadowCasterRepresentation(StrataPixelHeader));
 					DiffuseIndirectLighting *= IndirectOcclusion;
 					SubsurfaceIndirectLighting *= IndirectOcclusion;
 					IndirectIrradiance *= IndirectOcclusion;
diff --git a/Engine/Shaders/Private/Strata/StrataVisualize.usf b/Engine/Shaders/Private/Strata/StrataVisualize.usf
--- a/Engine/Shaders/Private/Strata/StrataVisualize.usf
+++ b/Engine/Shaders/Private/Strata/StrataVisualize.usf
@@ -201,11 +201,11 @@ void VisualizeMaterialPS(
 
 #define PRINTBSDFCOUNT()	PRINTBSDFCHAR(min(_0_+Header.BSDFCount,_F_)); PRINTBSDFCHAR(_SPC_); PRINTBSDFCHAR(_B_); PRINTBSDFCHAR(_S_); PRINTBSDFCHAR(_D_); PRINTBSDFCHAR(_F_);
 #define PRINTMATERIALAO()	PRINTBSDFCHAR(_A_); PRINTBSDFCHAR(_O_); PRINTBSDFCHAR(_SPC_); PRINTFLOAT1(Header.MaterialAO);
-#define PRINTMATERIALSKIPPRECSHADOWMASK()	PRINTBSDFCHAR(_S_); PRINTBSDFCHAR(_K_); PRINTBSDFCHAR(_I_); PRINTBSDFCHAR(_P_); PRINTBSDFCHAR(_P_); PRINTBSDFCHAR(_R_); PRINTBSDFCHAR(_E_); PRINTBSDFCHAR(_S_); PRINTBSDFCHAR(_H_); PRINTBSDFCHAR(_A_); PRINTBSDFCHAR(_D_); PRINTBSDFCHAR(_SPC_); PRINTBSDFCHAR(_0_+(Header.bSkipPrecShadowMask?1:0));
-#define PRINTMATERIALZEROPRECSHADOWMASK()	PRINTBSDFCHAR(_Z_); PRINTBSDFCHAR(_E_); PRINTBSDFCHAR(_R_); PRINTBSDFCHAR(_O_); PRINTBSDFCHAR(_P_); PRINTBSDFCHAR(_R_); PRINTBSDFCHAR(_E_); PRINTBSDFCHAR(_S_); PRINTBSDFCHAR(_H_); PRINTBSDFCHAR(_A_); PRINTBSDFCHAR(_D_); PRINTBSDFCHAR(_SPC_); PRINTBSDFCHAR(_0_+(Header.bZeroPrecShadowMask?1:0));
-#define PRINTMATERIALSLW()	PRINTBSDFCHAR(_S_); PRINTBSDFCHAR(_L_); PRINTBSDFCHAR(_W_); PRINTBSDFCHAR(_SPC_); PRINTBSDFCHAR(_0_+(Header.bIsSingleLayerWater?1:0));
-#define PRINTMATERIALCONTSHAD()	PRINTBSDFCHAR(_C_); PRINTBSDFCHAR(_O_); PRINTBSDFCHAR(_N_);  PRINTBSDFCHAR(_T_); PRINTBSDFCHAR(_SPC_);PRINTBSDFCHAR(_S_); PRINTBSDFCHAR(_H_); PRINTBSDFCHAR(_A_);  PRINTBSDFCHAR(_D_); PRINTBSDFCHAR(_SPC_); PRINTBSDFCHAR(_0_+(Header.bCastContactShadow?1:0));
-#define PRINTMATERIALINDSHAD()	PRINTBSDFCHAR(_I_); PRINTBSDFCHAR(_N_); PRINTBSDFCHAR(_D_); PRINTBSDFCHAR(_SPC_); PRINTBSDFCHAR(_S_); PRINTBSDFCHAR(_H_); PRINTBSDFCHAR(_A_);  PRINTBSDFCHAR(_D_); PRINTBSDFCHAR(_SPC_); PRINTBSDFCHAR(_0_+(Header.bHasDynamicIndirectShadowCasterRepresentation?1:0));
+#define PRINTMATERIALSKIPPRECSHADOWMASK()	PRINTBSDFCHAR(_S_); PRINTBSDFCHAR(_K_); PRINTBSDFCHAR(_I_); PRINTBSDFCHAR(_P_); PRINTBSDFCHAR(_P_); PRINTBSDFCHAR(_R_); PRINTBSDFCHAR(_E_); PRINTBSDFCHAR(_S_); PRINTBSDFCHAR(_H_); PRINTBSDFCHAR(_A_); PRINTBSDFCHAR(_D_); PRINTBSDFCHAR(_SPC_); PRINTBSDFCHAR(_0_+(HasSkipPrecShadowMask(Header)?1:0));
+#define PRINTMATERIALZEROPRECSHADOWMASK()	PRINTBSDFCHAR(_Z_); PRINTBSDFCHAR(_E_); PRINTBSDFCHAR(_R_); PRINTBSDFCHAR(_O_); PRINTBSDFCHAR(_P_); PRINTBSDFCHAR(_R_); PRINTBSDFCHAR(_E_); PRINTBSDFCHAR(_S_); PRINTBSDFCHAR(_H_); PRINTBSDFCHAR(_A_); PRINTBSDFCHAR(_D_); PRINTBSDFCHAR(_SPC_); PRINTBSDFCHAR(_0_+(HasZeroPrecShadowMask(Header)?1:0));
+#define PRINTMATERIALSLW()	PRINTBSDFCHAR(_S_); PRINTBSDFCHAR(_L_); PRINTBSDFCHAR(_W_); PRINTBSDFCHAR(_SPC_); PRINTBSDFCHAR(_0_+(IsSingleLayerWater(Header)?1:0));
+#define PRINTMATERIALCONTSHAD()	PRINTBSDFCHAR(_C_); PRINTBSDFCHAR(_O_); PRINTBSDFCHAR(_N_);  PRINTBSDFCHAR(_T_); PRINTBSDFCHAR(_SPC_);PRINTBSDFCHAR(_S_); PRINTBSDFCHAR(_H_); PRINTBSDFCHAR(_A_);  PRINTBSDFCHAR(_D_); PRINTBSDFCHAR(_SPC_); PRINTBSDFCHAR(_0_+(DoesCastContactShadow(Header)?1:0));
+#define PRINTMATERIALINDSHAD()	PRINTBSDFCHAR(_I_); PRINTBSDFCHAR(_N_); PRINTBSDFCHAR(_D_); PRINTBSDFCHAR(_SPC_); PRINTBSDFCHAR(_S_); PRINTBSDFCHAR(_H_); PRINTBSDFCHAR(_A_);  PRINTBSDFCHAR(_D_); PRINTBSDFCHAR(_SPC_); PRINTBSDFCHAR(_0_+(HasDynamicIndirectShadowCasterRepresentation(Header)?1:0));
 
 #define PRINTCLASSIFROUGH()		PRINTBSDFCHAR(_C_); PRINTBSDFCHAR(_L_); PRINTBSDFCHAR(_A_); PRINTBSDFCHAR(_S_); PRINTBSDFCHAR(_DOT_); PRINTBSDFCHAR(_R_) PRINTBSDFCHAR(_O_) PRINTBSDFCHAR(_U_) PRINTBSDFCHAR(_G_) PRINTBSDFCHAR(_H_) PRINTBSDFCHAR(_SPC_);	PRINTBSDFCHAR(_SPC_); PRINTFLOAT1(StrataData.TopLayerRoughness);
 #define PRINTCLASSIFBASCOL()	PRINTBSDFCHAR(_C_); PRINTBSDFCHAR(_L_); PRINTBSDFCHAR(_A_); PRINTBSDFCHAR(_S_); PRINTBSDFCHAR(_DOT_); PRINTBSDFCHAR(_B_) PRINTBSDFCHAR(_A_) PRINTBSDFCHAR(_S_) PRINTBSDFCHAR(_C_) PRINTBSDFCHAR(_O_) PRINTBSDFCHAR(_L_);	PRINTBSDFCHAR(_SPC_); PRINTFLOAT3(SSSData.BaseColor);
