Add default InlineRayTracing support.
* Add shader implementation of TraceRayInline
* Add required RHI changes

diff --git a/Engine/Shaders/Private/RayTracing/TraceRayInline.ush b/Engine/Shaders/Private/RayTracing/TraceRayInline.ush
--- a/Engine/Shaders/Private/RayTracing/TraceRayInline.ush
+++ b/Engine/Shaders/Private/RayTracing/TraceRayInline.ush
@@ -14,10 +14,56 @@
 
 #if OVERRIDE_TRACERAYINLINE_USH
 #include "/Platform/Private/TraceRayInline.ush"
+
+#else // OVERRIDE_TRACERAYINLINE_USH
+
+#include "/Engine/Private/RayTracing/TraceRayInlineCommon.ush"
+
+#define PLATFORM_SUPPORTS_INLINE_RAY_TRACING_TRIANGLE_NORMALS 0
+#define PLATFORM_SUPPORTS_INLINE_RAY_TRACING_MAX_ITERATIONS 0
+
+FTraceRayInlineResult TraceRayInline(RaytracingAccelerationStructure AccelerationStructure,
+	uint RayFlags,
+	uint InstanceInclusionMask,
+	RayDesc Ray)
+{
+	FTraceRayInlineResult TraceRayResult = (FTraceRayInlineResult)0;
+	TraceRayResult.SetMiss();
+
+	RayQuery<RAY_FLAG_SKIP_PROCEDURAL_PRIMITIVES | RAY_FLAG_FORCE_OPAQUE> RayQ;
+
+	RayQ.TraceRayInline(AccelerationStructure, RayFlags, InstanceInclusionMask, Ray);
+	while (RayQ.Proceed());
+	
+	if (RayQ.CommittedStatus() == COMMITTED_TRIANGLE_HIT)
+	{
+		TraceRayResult.HitT = RayQ.CommittedRayT();
+		TraceRayResult.InstanceID = RayQ.CommittedInstanceID();
+		TraceRayResult.InstanceIndex = RayQ.CommittedInstanceIndex();
+		TraceRayResult.PrimitiveIndex = RayQ.CommittedPrimitiveIndex();
+		TraceRayResult.GeometryIndex = RayQ.CommittedGeometryIndex();
+		TraceRayResult.Barycentrics = RayQ.CommittedTriangleBarycentrics();
+		TraceRayResult.bIsFrontFace = RayQ.CommittedTriangleFrontFace();
+		TraceRayResult.InstanceContributionToHitGroupIndex = RayQ.CommittedInstanceContributionToHitGroupIndex();
+
+		// placeholder
+		TraceRayResult.WorldGeometryNormal = RayQ.WorldRayDirection();
+
+		TraceRayResult.Bookmark[0] = asuint(TraceRayResult.HitT);
+		TraceRayResult.Bookmark[1] = asuint(TraceRayResult.HitT);
+	}
+
+	return TraceRayResult;
+}
+
 #endif // OVERRIDE_TRACERAYINLINE_USH
 
 #ifndef PLATFORM_SUPPORTS_INLINE_RAY_TRACING_TRIANGLE_NORMALS
 #define PLATFORM_SUPPORTS_INLINE_RAY_TRACING_TRIANGLE_NORMALS 0
 #endif // PLATFORM_SUPPORTS_INLINE_RAY_TRACING_TRIANGLE_NORMALS
 
+#ifndef PLATFORM_SUPPORTS_INLINE_RAY_TRACING_MAX_ITERATIONS
+#define PLATFORM_SUPPORTS_INLINE_RAY_TRACING_MAX_ITERATIONS 0
+#endif // PLATFORM_SUPPORTS_INLINE_RAY_TRACING_MAX_ITERATIONS
+
 #endif // PLATFORM_SUPPORTS_INLINE_RAY_TRACING
diff --git a/Engine/Source/Developer/Windows/ShaderFormatD3D/Private/D3DShaderCompiler.cpp b/Engine/Source/Developer/Windows/ShaderFormatD3D/Private/D3DShaderCompiler.cpp
--- a/Engine/Source/Developer/Windows/ShaderFormatD3D/Private/D3DShaderCompiler.cpp
+++ b/Engine/Source/Developer/Windows/ShaderFormatD3D/Private/D3DShaderCompiler.cpp
@@ -1072,7 +1072,8 @@ void CompileD3DShader(const FShaderCompilerInput& Input, FShaderCompilerOutput&
 		Language == ELanguage::SM6
 		|| bIsRayTracingShader
 		|| Input.Environment.CompilerFlags.Contains(CFLAG_WaveOperations)
-		|| Input.Environment.CompilerFlags.Contains(CFLAG_ForceDXC);
+		|| Input.Environment.CompilerFlags.Contains(CFLAG_ForceDXC)
+		|| Input.Environment.CompilerFlags.Contains(CFLAG_InlineRayTracing);
 	const TCHAR* ShaderProfile = GetShaderProfileName(Input.Target, bUseDXC);
 
 	if(!ShaderProfile)
@@ -1091,6 +1092,11 @@ void CompileD3DShader(const FShaderCompilerInput& Input, FShaderCompilerOutput&
 		AdditionalDefines.SetDefine(TEXT("PLATFORM_SUPPORTS_DIAGNOSTIC_BUFFER"), 1);
 		AdditionalDefines.SetDefine(TEXT("COMPILER_SUPPORTS_NOINLINE"), 1);
 
+		if (Input.Environment.CompilerFlags.Contains(CFLAG_InlineRayTracing))
+		{
+			AdditionalDefines.SetDefine(TEXT("PLATFORM_SUPPORTS_INLINE_RAY_TRACING"), 1);
+		}
+
 		if (Language == ELanguage::SM6 && Input.Environment.CompilerFlags.Contains(CFLAG_AllowRealTypes))
 		{
 			AdditionalDefines.SetDefine(TEXT("PLATFORM_SUPPORTS_REAL_TYPES"), 1);
diff --git a/Engine/Source/Runtime/D3D12RHI/Private/D3D12Adapter.cpp b/Engine/Source/Runtime/D3D12RHI/Private/D3D12Adapter.cpp
--- a/Engine/Source/Runtime/D3D12RHI/Private/D3D12Adapter.cpp
+++ b/Engine/Source/Runtime/D3D12RHI/Private/D3D12Adapter.cpp
@@ -1003,7 +1003,7 @@ void FD3D12Adapter::InitializeDevices()
 						UE_LOG(LogD3D12RHI, Log, TEXT("D3D12 ray tracing tier 1.1 is supported."));
 
 						GRHISupportsRayTracingPSOAdditions = true;
-						GRHISupportsInlineRayTracing = GRHISupportsRayTracing && RHISupportsInlineRayTracing(GMaxRHIShaderPlatform);
+						GRHISupportsInlineRayTracing = GRHISupportsRayTracing && RHISupportsInlineRayTracing(GMaxRHIShaderPlatform) && (GMaxRHIFeatureLevel >= ERHIFeatureLevel::SM6);
 					}
 				}
 				else if (D3D12Caps5.RaytracingTier == D3D12_RAYTRACING_TIER_NOT_SUPPORTED 
