Clean debug rendering for hair.

diff --git a/Engine/Shaders/Private/HairStrands/HairStrandsClusterCulling.usf b/Engine/Shaders/Private/HairStrands/HairStrandsClusterCulling.usf
--- a/Engine/Shaders/Private/HairStrands/HairStrandsClusterCulling.usf
+++ b/Engine/Shaders/Private/HairStrands/HairStrandsClusterCulling.usf
@@ -115,89 +115,21 @@ void MainDrawDebugAABBCS(uint3 DispatchThreadId : SV_DispatchThreadID)
 			const uint VisibleClusterCount  = CulledDispatchIndirectParametersClusterCountBuffer[0];
 			const uint VisibleTriangleCount = CulledDrawIndirectParameters[0] / 3; // CulledDrawIndirectParameters contains the total number of triangle indices (3 per triangle)
 
-			Pos = ShaderPrintSymbol(Pos, _H_);
-			Pos = ShaderPrintSymbol(Pos, _A_);
-			Pos = ShaderPrintSymbol(Pos, _I_);
-			Pos = ShaderPrintSymbol(Pos, _R_);
-			Pos = ShaderPrintSymbol(Pos, _G_);
-			Pos = ShaderPrintSymbol(Pos, _R_);
-			Pos = ShaderPrintSymbol(Pos, _O_);
-			Pos = ShaderPrintSymbol(Pos, _U_);
-			Pos = ShaderPrintSymbol(Pos, _P_);
-			Pos = ShaderPrintSymbol(Pos, _SPC_);
+			Pos = ShaderPrintText(Pos, TEXT("HairGroup "), FontYellow);
 			Pos = ShaderPrintSymbol(Pos, HairGroupId);
-			Pos = ShaderPrintSymbol(Pos, _SPC_);
-			Pos = ShaderPrintSymbol(Pos, _SPC_);
-			Pos = ShaderPrintSymbol(Pos, _SPC_);
-			Pos = ShaderPrintSymbol(Pos, _SPC_);
-			Pos = ShaderPrintSymbol(Pos, _SPC_);
-			Pos = ShaderPrintSymbol(Pos, _SPC_);
-			Pos = ShaderPrintSymbol(Pos, _SPC_);
-			Pos = ShaderPrintSymbol(Pos, _C_);
-			Pos = ShaderPrintSymbol(Pos, _O_);
-			Pos = ShaderPrintSymbol(Pos, _U_);
-			Pos = ShaderPrintSymbol(Pos, _N_);
-			Pos = ShaderPrintSymbol(Pos, _T_);
-			Pos = ShaderPrintSymbol(Pos, _SPC_);
-			Pos = ShaderPrintSymbol(Pos, _SPC_);
-			Pos = ShaderPrintSymbol(Pos, _SPC_);
-			Pos = ShaderPrintSymbol(Pos, _SPC_);
-			Pos = ShaderPrintSymbol(Pos, _SPC_);
-			Pos = ShaderPrintSymbol(Pos, _SPC_);
-			Pos = ShaderPrintSymbol(Pos, _SPC_);
-			Pos = ShaderPrintSymbol(Pos, _SPC_);
-			Pos = ShaderPrintSymbol(Pos, _P_);
-			Pos = ShaderPrintSymbol(Pos, _E_);
-			Pos = ShaderPrintSymbol(Pos, _R_);
-			Pos = ShaderPrintSymbol(Pos, _C_);
-			Pos = ShaderPrintSymbol(Pos, _E_);
-			Pos = ShaderPrintSymbol(Pos, _N_);
-			Pos = ShaderPrintSymbol(Pos, _T_);
+			Pos = ShaderPrintText(Pos, TEXT("       Count        Percent"), FontYellow);
 
 			Pos = ShaderPrintNewline(Pos); Pos.x = StartPos.x;
 
-			Pos = ShaderPrintSymbol(Pos, _V_);
-			Pos = ShaderPrintSymbol(Pos, _I_);
-			Pos = ShaderPrintSymbol(Pos, _S_);
-			Pos = ShaderPrintSymbol(Pos, _I_);
-			Pos = ShaderPrintSymbol(Pos, _B_);
-			Pos = ShaderPrintSymbol(Pos, _L_);
-			Pos = ShaderPrintSymbol(Pos, _E_);
-			Pos = ShaderPrintSymbol(Pos, _SPC_);
-			Pos = ShaderPrintSymbol(Pos, _C_);
-			Pos = ShaderPrintSymbol(Pos, _L_);
-			Pos = ShaderPrintSymbol(Pos, _U_);
-			Pos = ShaderPrintSymbol(Pos, _S_);
-			Pos = ShaderPrintSymbol(Pos, _T_);
-			Pos = ShaderPrintSymbol(Pos, _E_);
-			Pos = ShaderPrintSymbol(Pos, _R_);
-			Pos = ShaderPrintSymbol(Pos, _S_);
-			Pos = ShaderPrintSymbol(Pos, _SPC_);
-			Pos = ShaderPrintSymbol(Pos, _SPC_);
+			Pos = ShaderPrintText(Pos, TEXT("Visible Clusters  "), FontYellow);
 			Pos = ShaderPrint(Pos, VisibleClusterCount);
+			Pos = ShaderPrintText(Pos, TEXT(" "), FontYellow);
 			Pos = ShaderPrintSymbol(Pos, _SPC_);
 			Pos = ShaderPrint(Pos, 100.0f * float(VisibleClusterCount) / float(ClusterCount));
 
 			Pos = ShaderPrintNewline(Pos); Pos.x = StartPos.x;
 
-			Pos = ShaderPrintSymbol(Pos, _V_);
-			Pos = ShaderPrintSymbol(Pos, _I_);
-			Pos = ShaderPrintSymbol(Pos, _S_);
-			Pos = ShaderPrintSymbol(Pos, _I_);
-			Pos = ShaderPrintSymbol(Pos, _B_);
-			Pos = ShaderPrintSymbol(Pos, _L_);
-			Pos = ShaderPrintSymbol(Pos, _E_);
-			Pos = ShaderPrintSymbol(Pos, _SPC_);
-			Pos = ShaderPrintSymbol(Pos, _T_);
-			Pos = ShaderPrintSymbol(Pos, _R_);
-			Pos = ShaderPrintSymbol(Pos, _I_);
-			Pos = ShaderPrintSymbol(Pos, _A_);
-			Pos = ShaderPrintSymbol(Pos, _N_);
-			Pos = ShaderPrintSymbol(Pos, _G_);
-			Pos = ShaderPrintSymbol(Pos, _L_);
-			Pos = ShaderPrintSymbol(Pos, _E_);
-			Pos = ShaderPrintSymbol(Pos, _S_);
-			Pos = ShaderPrintSymbol(Pos, _SPC_);
+			Pos = ShaderPrintText(Pos, TEXT("Visible Triangles "), FontYellow);
 			Pos = ShaderPrint(Pos, VisibleTriangleCount);
 			Pos = ShaderPrintSymbol(Pos, _SPC_);
 			Pos = ShaderPrint(Pos, 100.0f * float(VisibleTriangleCount) / float(TriangleCount));
@@ -210,24 +142,7 @@ void MainDrawDebugAABBCS(uint3 DispatchThreadId : SV_DispatchThreadID)
 			// Extra new line to seperate from previous block
 			Pos = ShaderPrintNewline(Pos); Pos.x = StartPos.x;
 
-			Pos = ShaderPrintSymbol(Pos, _S_);
-			Pos = ShaderPrintSymbol(Pos, _C_);
-			Pos = ShaderPrintSymbol(Pos, _R_);
-			Pos = ShaderPrintSymbol(Pos, _E_);
-			Pos = ShaderPrintSymbol(Pos, _E_);
-			Pos = ShaderPrintSymbol(Pos, _N_);
-			Pos = ShaderPrintSymbol(Pos, _SPC_);
-			Pos = ShaderPrintSymbol(Pos, _S_);
-			Pos = ShaderPrintSymbol(Pos, _I_);
-			Pos = ShaderPrintSymbol(Pos, _Z_);
-			Pos = ShaderPrintSymbol(Pos, _E_);
-			Pos = ShaderPrintSymbol(Pos, _SPC_);
-			Pos = ShaderPrintSymbol(Pos, _SPC_);
-			Pos = ShaderPrintSymbol(Pos, _SPC_);
-			Pos = ShaderPrintSymbol(Pos, _SPC_);
-			Pos = ShaderPrintSymbol(Pos, _SPC_);
-			Pos = ShaderPrintSymbol(Pos, _SPC_);
-			Pos = ShaderPrintSymbol(Pos, _SPC_);
+			Pos = ShaderPrintText(Pos, TEXT("Screen Size       "), FontYellow);
 
 			float ScreenSize = 1;
 			{
@@ -246,95 +161,7 @@ void MainDrawDebugAABBCS(uint3 DispatchThreadId : SV_DispatchThreadID)
 
 			if (bShowLODHistogram)
 			{
-				Pos = ShaderPrintSymbol(Pos, _L_);
-				Pos = ShaderPrintSymbol(Pos, _O_);
-				Pos = ShaderPrintSymbol(Pos, _D_);
-				Pos = ShaderPrintSymbol(Pos, _0_);
-				Pos = ShaderPrintSymbol(Pos, _SPC_);
-				Pos = ShaderPrintSymbol(Pos, _SPC_);
-				Pos = ShaderPrintSymbol(Pos, _SPC_);
-				Pos = ShaderPrintSymbol(Pos, _SPC_);
-				Pos = ShaderPrintSymbol(Pos, _SPC_);
-				Pos = ShaderPrintSymbol(Pos, _SPC_);
-				Pos = ShaderPrintSymbol(Pos, _SPC_);
-				Pos = ShaderPrintSymbol(Pos, _SPC_);
-				Pos = ShaderPrintSymbol(Pos, _L_);
-				Pos = ShaderPrintSymbol(Pos, _O_);
-				Pos = ShaderPrintSymbol(Pos, _D_);
-				Pos = ShaderPrintSymbol(Pos, _1_);
-				Pos = ShaderPrintSymbol(Pos, _SPC_);
-				Pos = ShaderPrintSymbol(Pos, _SPC_);
-				Pos = ShaderPrintSymbol(Pos, _SPC_);
-				Pos = ShaderPrintSymbol(Pos, _SPC_);
-				Pos = ShaderPrintSymbol(Pos, _SPC_);
-				Pos = ShaderPrintSymbol(Pos, _SPC_);
-				Pos = ShaderPrintSymbol(Pos, _SPC_);
-				Pos = ShaderPrintSymbol(Pos, _SPC_);
-				Pos = ShaderPrintSymbol(Pos, _L_);
-				Pos = ShaderPrintSymbol(Pos, _O_);
-				Pos = ShaderPrintSymbol(Pos, _D_);
-				Pos = ShaderPrintSymbol(Pos, _2_);
-				Pos = ShaderPrintSymbol(Pos, _SPC_);
-				Pos = ShaderPrintSymbol(Pos, _SPC_);
-				Pos = ShaderPrintSymbol(Pos, _SPC_);
-				Pos = ShaderPrintSymbol(Pos, _SPC_);
-				Pos = ShaderPrintSymbol(Pos, _SPC_);
-				Pos = ShaderPrintSymbol(Pos, _SPC_);
-				Pos = ShaderPrintSymbol(Pos, _SPC_);
-				Pos = ShaderPrintSymbol(Pos, _SPC_);
-				Pos = ShaderPrintSymbol(Pos, _L_);
-				Pos = ShaderPrintSymbol(Pos, _O_);
-				Pos = ShaderPrintSymbol(Pos, _D_);
-				Pos = ShaderPrintSymbol(Pos, _3_);
-				Pos = ShaderPrintSymbol(Pos, _SPC_);
-				Pos = ShaderPrintSymbol(Pos, _SPC_);
-				Pos = ShaderPrintSymbol(Pos, _SPC_);
-				Pos = ShaderPrintSymbol(Pos, _SPC_);
-				Pos = ShaderPrintSymbol(Pos, _SPC_);
-				Pos = ShaderPrintSymbol(Pos, _SPC_);
-				Pos = ShaderPrintSymbol(Pos, _SPC_);
-				Pos = ShaderPrintSymbol(Pos, _SPC_);
-				Pos = ShaderPrintSymbol(Pos, _L_);
-				Pos = ShaderPrintSymbol(Pos, _O_);
-				Pos = ShaderPrintSymbol(Pos, _D_);
-				Pos = ShaderPrintSymbol(Pos, _4_);
-				Pos = ShaderPrintSymbol(Pos, _SPC_);
-				Pos = ShaderPrintSymbol(Pos, _SPC_);
-				Pos = ShaderPrintSymbol(Pos, _SPC_);
-				Pos = ShaderPrintSymbol(Pos, _SPC_);
-				Pos = ShaderPrintSymbol(Pos, _SPC_);
-				Pos = ShaderPrintSymbol(Pos, _SPC_);
-				Pos = ShaderPrintSymbol(Pos, _SPC_);
-				Pos = ShaderPrintSymbol(Pos, _SPC_);
-				Pos = ShaderPrintSymbol(Pos, _L_);
-				Pos = ShaderPrintSymbol(Pos, _O_);
-				Pos = ShaderPrintSymbol(Pos, _D_);
-				Pos = ShaderPrintSymbol(Pos, _5_);
-				Pos = ShaderPrintSymbol(Pos, _SPC_);
-				Pos = ShaderPrintSymbol(Pos, _SPC_);
-				Pos = ShaderPrintSymbol(Pos, _SPC_);
-				Pos = ShaderPrintSymbol(Pos, _SPC_);
-				Pos = ShaderPrintSymbol(Pos, _SPC_);
-				Pos = ShaderPrintSymbol(Pos, _SPC_);
-				Pos = ShaderPrintSymbol(Pos, _SPC_);
-				Pos = ShaderPrintSymbol(Pos, _SPC_);
-				Pos = ShaderPrintSymbol(Pos, _L_);
-				Pos = ShaderPrintSymbol(Pos, _O_);
-				Pos = ShaderPrintSymbol(Pos, _D_);
-				Pos = ShaderPrintSymbol(Pos, _6_);
-				Pos = ShaderPrintSymbol(Pos, _SPC_);
-				Pos = ShaderPrintSymbol(Pos, _SPC_);
-				Pos = ShaderPrintSymbol(Pos, _SPC_);
-				Pos = ShaderPrintSymbol(Pos, _SPC_);
-				Pos = ShaderPrintSymbol(Pos, _SPC_);
-				Pos = ShaderPrintSymbol(Pos, _SPC_);
-				Pos = ShaderPrintSymbol(Pos, _SPC_);
-				Pos = ShaderPrintSymbol(Pos, _SPC_);
-				Pos = ShaderPrintSymbol(Pos, _L_);
-				Pos = ShaderPrintSymbol(Pos, _O_);
-				Pos = ShaderPrintSymbol(Pos, _D_);
-				Pos = ShaderPrintSymbol(Pos, _7_);
-
+				Pos = ShaderPrintText(Pos, TEXT("LOD0        LOD1        LOD2        LOD3        LOD4        LOD5        LOD6        LOD7"), FontYellow);
 				Pos = ShaderPrintNewline(Pos); Pos.x = StartPos.x;
 			}
 		
@@ -370,24 +197,7 @@ void MainDrawDebugAABBCS(uint3 DispatchThreadId : SV_DispatchThreadID)
 			}
 			else
 			{
-				Pos = ShaderPrintSymbol(Pos, _L_);
-				Pos = ShaderPrintSymbol(Pos, _O_);
-				Pos = ShaderPrintSymbol(Pos, _D_);
-				Pos = ShaderPrintSymbol(Pos, _SPC_);
-				Pos = ShaderPrintSymbol(Pos, _SPC_);
-				Pos = ShaderPrintSymbol(Pos, _SPC_);
-				Pos = ShaderPrintSymbol(Pos, _SPC_);
-				Pos = ShaderPrintSymbol(Pos, _SPC_);
-				Pos = ShaderPrintSymbol(Pos, _SPC_);
-				Pos = ShaderPrintSymbol(Pos, _SPC_);
-				Pos = ShaderPrintSymbol(Pos, _SPC_);
-				Pos = ShaderPrintSymbol(Pos, _SPC_);
-				Pos = ShaderPrintSymbol(Pos, _SPC_);
-				Pos = ShaderPrintSymbol(Pos, _SPC_);
-				Pos = ShaderPrintSymbol(Pos, _SPC_);
-				Pos = ShaderPrintSymbol(Pos, _SPC_);
-				Pos = ShaderPrintSymbol(Pos, _SPC_);
-				Pos = ShaderPrintSymbol(Pos, _SPC_);
+				Pos = ShaderPrintText(Pos, TEXT("LOD               "), FontYellow);
 
 				if (TotalVisibleCluster > 0)
 				{
@@ -395,6 +205,7 @@ void MainDrawDebugAABBCS(uint3 DispatchThreadId : SV_DispatchThreadID)
 				}
 				else
 				{
+					Pos = ShaderPrintText(Pos, TEXT("Culled"), FontYellow);
 					Pos = ShaderPrintSymbol(Pos, _C_);
 					Pos = ShaderPrintSymbol(Pos, _U_);
 					Pos = ShaderPrintSymbol(Pos, _L_);
