Fixed Strata forward material compilation problems, especially with DXC (long compile time).

diff --git a/Engine/Shaders/Private/Strata/StrataForwardLighting.ush b/Engine/Shaders/Private/Strata/StrataForwardLighting.ush
--- a/Engine/Shaders/Private/Strata/StrataForwardLighting.ush
+++ b/Engine/Shaders/Private/Strata/StrataForwardLighting.ush
@@ -249,7 +249,7 @@ float3 StrataForwardLighting(
 	GridIndex = ComputeLightGridCellIndex((uint2)((MaterialParameters.SvPosition.xy - ResolvedView.ViewRectMin.xy) * View.LightProbeSizeRatioAndInvSizeRatio.zw), MaterialParameters.SvPosition.w, EyeIndex);
 #endif
 
-	UNROLL
+	STRATA_UNROLL_N(STRATA_CLAMPED_LAYER_COUNT)
 	for (int l = 0; l < StrataData.LayerCount; ++l)
 	{
 		float3 LayerThroughput = 0.0f;
@@ -258,7 +258,7 @@ float3 StrataForwardLighting(
 		const bool bTopLayer = l == 0;
 		const bool bBottomLayer = l == (StrataData.LayerCount - 1);
 
-		UNROLL
+		STRATA_UNROLL_N(STRATA_CLAMPED_MAX_BSDF_COUNT_PER_LAYER)
 		for (int i = 0; i < StrataData.Layers[l].BSDFCount; ++i)
 		{
 			FStrataBSDF BSDF = StrataData.Layers[l].BSDFs[i];
