Removed strata special dxc setup that seems unecessary now.

diff --git a/Engine/Shaders/Private/Strata/Strata.ush b/Engine/Shaders/Private/Strata/Strata.ush
--- a/Engine/Shaders/Private/Strata/Strata.ush
+++ b/Engine/Shaders/Private/Strata/Strata.ush
@@ -13,17 +13,6 @@
 #include "/Engine/Shared/StrataDefinitions.h"
 
 
-// DXC has issue with unrolling certain part of the code. To address temporarly the warning-as-error issues, 
-// we redefine the unroll function and fall back to loop on DXC
-#if COMPILER_DXC == 1
-#define STRATA_UNROLL LOOP
-#define STRATA_UNROLL_N(X) LOOP
-#else
-#define STRATA_UNROLL UNROLL
-#define STRATA_UNROLL_N(X) UNROLL_N(X)
-#endif
-
-
 // During the base pass or forward rendering, shared local bases are simply available from registers.
 // If a BSDF is loaded for processing from the strata buffer, normals are unpacked on demand to reduce VGPR pressure and increase occupancy.
 #ifndef STRATA_INLINE_SHADING
@@ -880,12 +869,12 @@ FStrataData StrataAdd(FStrataData A, FStrataData B)
 	FStrataData StrataData = A;
 
 	// Append each BSDF from B to A, with same layer position
-	STRATA_UNROLL_N(STRATA_CLAMPED_LAYER_COUNT)
+	UNROLL_N(STRATA_CLAMPED_LAYER_COUNT)
 	for (int Layer = 0; Layer < B.LayerCount; ++Layer)
 	{
 		FStrataLayer ALayer = A.Layers[Layer];
 		FStrataLayer BLayer = B.Layers[Layer];
-		STRATA_UNROLL
+		UNROLL
 		for (int BSDF = 0; BSDF < BLayer.BSDFCount; BSDF++)
 		{
 			StrataData.Layers[Layer].BSDFs[ALayer.BSDFCount + BSDF] = BLayer.BSDFs[BSDF];
@@ -905,10 +894,10 @@ FStrataData StrataWeight(FStrataData A, float Weight)
 	FStrataData StrataData = A;
 
 	// Apply the weight to each BSDF
-	STRATA_UNROLL_N(STRATA_CLAMPED_LAYER_COUNT)
+	UNROLL_N(STRATA_CLAMPED_LAYER_COUNT)
 	for (int Layer = 0; Layer < StrataData.LayerCount; ++Layer)
 	{
-		STRATA_UNROLL
+		UNROLL
 		for (int BSDF = 0; BSDF < StrataData.Layers[Layer].BSDFCount; ++BSDF)
 		{
 #if STRATA_INLINE_SHADING
@@ -945,7 +934,7 @@ FStrataData StrataVerticalLayering(FStrataData Top, FStrataData Base)
 
 	// Add each layer from Base under Top
 	int TopLayerCount = Top.LayerCount;
-	STRATA_UNROLL_N(STRATA_CLAMPED_LAYER_COUNT)
+	UNROLL_N(STRATA_CLAMPED_LAYER_COUNT)
 	for (int Layer = 0; Layer < Base.LayerCount; ++Layer)
 	{
 		StrataData.Layers[TopLayerCount + Layer] = Base.Layers[Layer];
@@ -1758,7 +1747,7 @@ void PackStrataOut(
 
 	// For each of the BSDF, we evaluated their combined coverage/throughput according to the material operation performed (layering, mixing, etc.).
 	int l;
-	STRATA_UNROLL_N(STRATA_CLAMPED_LAYER_COUNT)
+	UNROLL_N(STRATA_CLAMPED_LAYER_COUNT)
 	for (l = 0; l < Strata.LayerCount; ++l)
 	{
 		float3 LayerThroughput = 0.0f;
@@ -1767,7 +1756,7 @@ void PackStrataOut(
 		const bool bTopLayer = l == 0;
 		const bool bBottomLayer = l == (Strata.LayerCount - 1);
 
-		STRATA_UNROLL_N(STRATA_CLAMPED_MAX_BSDF_COUNT_PER_LAYER)
+		UNROLL_N(STRATA_CLAMPED_MAX_BSDF_COUNT_PER_LAYER)
 		for (int i = 0; i < Strata.Layers[l].BSDFCount; ++i)
 		{
 			// Dither BSDF data to hide more aggressive bit reduction.
@@ -1966,12 +1955,12 @@ void PackStrataOut(
 
 	// 2. the list of BSDFs
 	BSDFCount = 0;
-	STRATA_UNROLL_N(STRATA_CLAMPED_LAYER_COUNT)
+	UNROLL_N(STRATA_CLAMPED_LAYER_COUNT)
 	for (l = 0; l < Strata.LayerCount; ++l)
 	{
 		FStrataLayer Layer = Strata.Layers[l];
 
-		STRATA_UNROLL_N(STRATA_CLAMPED_MAX_BSDF_COUNT_PER_LAYER)
+		UNROLL_N(STRATA_CLAMPED_MAX_BSDF_COUNT_PER_LAYER)
 		for (int i = 0; i < Layer.BSDFCount; ++i)
 		{
 			FStrataBSDF BSDF = Layer.BSDFs[i];
diff --git a/Engine/Shaders/Private/Strata/StrataForwardLighting.ush b/Engine/Shaders/Private/Strata/StrataForwardLighting.ush
--- a/Engine/Shaders/Private/Strata/StrataForwardLighting.ush
+++ b/Engine/Shaders/Private/Strata/StrataForwardLighting.ush
@@ -249,7 +249,7 @@ float3 StrataForwardLighting(
 	GridIndex = ComputeLightGridCellIndex((uint2)((MaterialParameters.SvPosition.xy - ResolvedView.ViewRectMin.xy) * View.LightProbeSizeRatioAndInvSizeRatio.zw), MaterialParameters.SvPosition.w, EyeIndex);
 #endif
 
-	STRATA_UNROLL_N(STRATA_CLAMPED_LAYER_COUNT)
+	UNROLL_N(STRATA_CLAMPED_LAYER_COUNT)
 	for (int l = 0; l < StrataData.LayerCount; ++l)
 	{
 		float3 LayerThroughput = 0.0f;
@@ -258,7 +258,7 @@ float3 StrataForwardLighting(
 		const bool bTopLayer = l == 0;
 		const bool bBottomLayer = l == (StrataData.LayerCount - 1);
 
-		STRATA_UNROLL_N(STRATA_CLAMPED_MAX_BSDF_COUNT_PER_LAYER)
+		UNROLL_N(STRATA_CLAMPED_MAX_BSDF_COUNT_PER_LAYER)
 		for (int i = 0; i < StrataData.Layers[l].BSDFCount; ++i)
 		{
 			FStrataBSDF BSDF = StrataData.Layers[l].BSDFs[i];
