Added morton order to strata buffer tiles. Seems to help an tiny tiny bit. (this revealed the real problem is in fact not there)

diff --git a/Engine/Shaders/Private/Strata/Strata.ush b/Engine/Shaders/Private/Strata/Strata.ush
--- a/Engine/Shaders/Private/Strata/Strata.ush
+++ b/Engine/Shaders/Private/Strata/Strata.ush
@@ -9,7 +9,7 @@
 #include "../DeferredShadingCommon.ush"
 #include "../ParticipatingMediaCommon.ush"
 #include "../ThinFilmBSDF.ush"
- 
+#include "../MortonCode.ush"
 #include "/Engine/Shared/StrataDefinitions.h"
 
 
@@ -705,7 +705,11 @@ FStrataAddressing GetStrataPixelDataByteOffset(uint2 PixelPos, uint2 ViewBufferS
 	uint2 PixelCoordInTile			= PixelPos - TileCoord * STRATA_DATA_TILE_SIZE;
 
 	uint PreviousTilesBytes			= (TileCoord.y * TileCount.x + TileCoord.x) * TileSizeBytes;
+#if 0
 	uint PreviousPixelsInTileBytes	= (PixelCoordInTile.y * STRATA_DATA_TILE_SIZE + PixelCoordInTile.x) * TilePixelSizeBytes;
+#else
+	uint PreviousPixelsInTileBytes	= MortonEncode(PixelCoordInTile) * TilePixelSizeBytes; // Z pattern
+#endif
 	StrataAddressing.CurrentAddr	= PreviousTilesBytes + PreviousPixelsInTileBytes;
 
 	// All
