Add checks and fix out of bounds access in prefix sum

17979382

diff --git a/Engine/Shaders/Private/HairStrands/HairStrandsClusterCulling.usf b/Engine/Shaders/Private/HairStrands/HairStrandsClusterCulling.usf
--- a/Engine/Shaders/Private/HairStrands/HairStrandsClusterCulling.usf
+++ b/Engine/Shaders/Private/HairStrands/HairStrandsClusterCulling.usf
@@ -58,6 +58,7 @@ void MainClearClusterAABBCS(uint2 DispatchThreadId : SV_DispatchThreadID)
 		const int ClusterId = DispatchThreadId.x / 6;
 		const bool bIsMin = (DispatchThreadId.x - ClusterId * 6) < 3 ? 1 : 0;
 		const int Value = bIsMin ? 2147483647 : -2147483648;
+		checkBufferAccessSlow(OutClusterAABBBuffer, DispatchThreadId.x);
 		OutClusterAABBBuffer[DispatchThreadId.x] = Value;
 	}
 
@@ -379,6 +380,7 @@ void MainClusterCullingCS(uint3 DispatchThreadId : SV_DispatchThreadID)
 
 		const uint ClusterId  = DispatchThreadId.x;
 		const uint ClusterId6 = ClusterId * 6;
+		checkBufferAccessSlow(ClusterAABBBuffer, ClusterId6 + 5);
 		const float3 Min = float3(int3(ClusterAABBBuffer[ClusterId6 + 0], ClusterAABBBuffer[ClusterId6 + 1], ClusterAABBBuffer[ClusterId6 + 2]));
 		const float3 Max = float3(int3(ClusterAABBBuffer[ClusterId6 + 3], ClusterAABBBuffer[ClusterId6 + 4], ClusterAABBBuffer[ClusterId6 + 5]));
 
@@ -429,9 +431,13 @@ void MainClusterCullingCS(uint3 DispatchThreadId : SV_DispatchThreadID)
 					InterlockedAdd(DrawIndirectParameters[0], ClusterLOD.VertexCount * 6/*, PreviousVertexCount*/);
 				#endif
 
+					checkBufferAccessSlow(GlobalClusterIdBuffer, VisibleClusterId);
 					GlobalClusterIdBuffer[VisibleClusterId]	  = ClusterId;
+					checkBufferAccessSlow(GlobalIndexStartBuffer, VisibleClusterId);
 					GlobalIndexStartBuffer[VisibleClusterId]  = ClusterLOD.VertexOffset;
+					checkBufferAccessSlow(GlobalIndexCountBuffer, VisibleClusterId);
 					GlobalIndexCountBuffer[VisibleClusterId]  = ClusterLOD.VertexCount;
+					checkBufferAccessSlow(GlobalRadiusScaleBuffer, VisibleClusterId);
 					GlobalRadiusScaleBuffer[VisibleClusterId] = ClusterLOD.RadiusScale;
 				}
 			}
@@ -563,6 +569,8 @@ void MainPerBlockPreFixSumCS(uint3 DispatchThreadId : SV_DispatchThreadID, uint3
 		uint GlobalIndex = WaveStart + GroupThreadId.x;
 		if (GlobalIndex < ItemCount)
 		{
+			checkBufferAccessSlow(GlobalIndexCountBuffer, GlobalIndex);
+
 			const uint IndexCount = GlobalIndexCountBuffer[GlobalIndex];
 			PrefixSumClusterCount[0][GroupThreadId.x] = IndexCount;
 		}
@@ -590,10 +598,12 @@ void MainPerBlockPreFixSumCS(uint3 DispatchThreadId : SV_DispatchThreadID, uint3
 
 		// Write out result, taking into account previous blocks sum
 		uint Result = 0;
-		if (GlobalIndex < ItemCountRoundUp)
+		if (GlobalIndex < ItemCount)
 		{
 			Result = PrefixSumClusterCount[Src][GroupThreadId.x];
 
+			checkBufferAccessSlow(PerBlocklIndexCountPreFixSumBuffer, GlobalIndex);
+
 			if (GroupThreadId.x == 0)
 			{
 				PerBlocklIndexCountPreFixSumBuffer[GlobalIndex] = uint2(0, Result);
@@ -606,6 +616,7 @@ void MainPerBlockPreFixSumCS(uint3 DispatchThreadId : SV_DispatchThreadID, uint3
 
 		if (GroupThreadId.x == PREFIXSUM_LASTTHREAD)
 		{
+			checkBufferAccessSlow(PerBlocklTotalIndexCountBuffer, GroupId.x);
 			PerBlocklTotalIndexCountBuffer[GroupId.x] = PrefixSumClusterCount[Src][GroupThreadId.x];
 		}
 	}
@@ -690,10 +701,11 @@ void MainClusterCullingPreFixSumCS(uint3 DispatchThreadId : SV_DispatchThreadID,
 
 		// Write out result, taking into account previous blocks sum
 		uint Result = 0;
-		if (GlobalIndex < ItemCountRoundUp)
+		if (GlobalIndex < ItemCount)
 		{
 			Result = PrefixSumClusterCount[Src][DispatchThreadId.x];
 
+			checkBufferAccessSlow(GlobalIndexCountPreFixSumBuffer, GlobalIndex);
 			if (DispatchThreadId.x == 0)
 			{
 				GlobalIndexCountPreFixSumBuffer[GlobalIndex] = uint2(0, Result) + PreviousSum;
@@ -773,7 +785,9 @@ void MainClusterCullingCompactVertexIdsLocalBlockCS(uint3 DispatchThreadId : SV_
 
 		if (ThreadWriteOffset < ClusterSrcIndexCount)
 		{
+			checkBufferAccessSlow(CulledCompactedIndexBuffer, WriteToIndex + ThreadWriteOffset);
 			CulledCompactedIndexBuffer[WriteToIndex + ThreadWriteOffset] = ClusterVertexIdBuffer[ClusterSrcStartIndex + ThreadWriteOffset];
+			checkBufferAccessSlow(CulledCompactedRadiusScaleBuffer, WriteToIndex + ThreadWriteOffset);
 			CulledCompactedRadiusScaleBuffer[WriteToIndex + ThreadWriteOffset] = D_USE_CONTINUOUSLOD ? ClusterRadiusScale : 0;
 		}
 	}
