Back out changelist 18293442
causing issue on pc+rt

diff --git a/Engine/Shaders/Private/Strata/Strata.ush b/Engine/Shaders/Private/Strata/Strata.ush
--- a/Engine/Shaders/Private/Strata/Strata.ush
+++ b/Engine/Shaders/Private/Strata/Strata.ush
@@ -13,6 +13,18 @@
 #include "/Engine/Shared/StrataDefinitions.h"
 
 
+// DXC has issue with unrolling certain part of the code. To address temporarly the warning-as-error issues, 
+// we redefine the unroll function and fall back to loop on DXC
+// To remove this, one must test PC + DX12 + RayTracing and make sure raytracing shaders compile.
+#if COMPILER_DXC == 1
+#define STRATA_UNROLL LOOP
+#define STRATA_UNROLL_N(X) LOOP
+#else
+#define STRATA_UNROLL UNROLL
+#define STRATA_UNROLL_N(X) UNROLL_N(X)
+#endif
+
+
 // During the base pass or forward rendering, shared local bases are simply available from registers.
 // If a BSDF is loaded for processing from the strata buffer, normals are unpacked on demand to reduce VGPR pressure and increase occupancy.
 #ifndef STRATA_INLINE_SHADING
@@ -873,12 +885,12 @@ FStrataData StrataAdd(FStrataData A, FStrataData B)
 	FStrataData StrataData = A;
 
 	// Append each BSDF from B to A, with same layer position
-	UNROLL_N(STRATA_CLAMPED_LAYER_COUNT)
+	STRATA_UNROLL_N(STRATA_CLAMPED_LAYER_COUNT)
 	for (int Layer = 0; Layer < B.LayerCount; ++Layer)
 	{
 		FStrataLayer ALayer = A.Layers[Layer];
 		FStrataLayer BLayer = B.Layers[Layer];
-		UNROLL
+		STRATA_UNROLL
 		for (int BSDF = 0; BSDF < BLayer.BSDFCount; BSDF++)
 		{
 			StrataData.Layers[Layer].BSDFs[ALayer.BSDFCount + BSDF] = BLayer.BSDFs[BSDF];
@@ -898,10 +910,10 @@ FStrataData StrataWeight(FStrataData A, float Weight)
 	FStrataData StrataData = A;
 
 	// Apply the weight to each BSDF
-	UNROLL_N(STRATA_CLAMPED_LAYER_COUNT)
+	STRATA_UNROLL_N(STRATA_CLAMPED_LAYER_COUNT)
 	for (int Layer = 0; Layer < StrataData.LayerCount; ++Layer)
 	{
-		UNROLL
+		STRATA_UNROLL
 		for (int BSDF = 0; BSDF < StrataData.Layers[Layer].BSDFCount; ++BSDF)
 		{
 #if STRATA_INLINE_SHADING
@@ -938,7 +950,7 @@ FStrataData StrataVerticalLayering(FStrataData Top, FStrataData Base)
 
 	// Add each layer from Base under Top
 	int TopLayerCount = Top.LayerCount;
-	UNROLL_N(STRATA_CLAMPED_LAYER_COUNT)
+	STRATA_UNROLL_N(STRATA_CLAMPED_LAYER_COUNT)
 	for (int Layer = 0; Layer < Base.LayerCount; ++Layer)
 	{
 		StrataData.Layers[TopLayerCount + Layer] = Base.Layers[Layer];
@@ -1751,7 +1763,7 @@ void PackStrataOut(
 
 	// For each of the BSDF, we evaluated their combined coverage/throughput according to the material operation performed (layering, mixing, etc.).
 	int l;
-	UNROLL_N(STRATA_CLAMPED_LAYER_COUNT)
+	STRATA_UNROLL_N(STRATA_CLAMPED_LAYER_COUNT)
 	for (l = 0; l < Strata.LayerCount; ++l)
 	{
 		float3 LayerThroughput = 0.0f;
@@ -1760,7 +1772,7 @@ void PackStrataOut(
 		const bool bTopLayer = l == 0;
 		const bool bBottomLayer = l == (Strata.LayerCount - 1);
 
-		UNROLL_N(STRATA_CLAMPED_MAX_BSDF_COUNT_PER_LAYER)
+		STRATA_UNROLL_N(STRATA_CLAMPED_MAX_BSDF_COUNT_PER_LAYER)
 		for (int i = 0; i < Strata.Layers[l].BSDFCount; ++i)
 		{
 			// Dither BSDF data to hide more aggressive bit reduction.
@@ -1959,12 +1971,12 @@ void PackStrataOut(
 
 	// 2. the list of BSDFs
 	BSDFCount = 0;
-	UNROLL_N(STRATA_CLAMPED_LAYER_COUNT)
+	STRATA_UNROLL_N(STRATA_CLAMPED_LAYER_COUNT)
 	for (l = 0; l < Strata.LayerCount; ++l)
 	{
 		FStrataLayer Layer = Strata.Layers[l];
 
-		UNROLL_N(STRATA_CLAMPED_MAX_BSDF_COUNT_PER_LAYER)
+		STRATA_UNROLL_N(STRATA_CLAMPED_MAX_BSDF_COUNT_PER_LAYER)
 		for (int i = 0; i < Layer.BSDFCount; ++i)
 		{
 			FStrataBSDF BSDF = Layer.BSDFs[i];
diff --git a/Engine/Shaders/Private/Strata/StrataForwardLighting.ush b/Engine/Shaders/Private/Strata/StrataForwardLighting.ush
--- a/Engine/Shaders/Private/Strata/StrataForwardLighting.ush
+++ b/Engine/Shaders/Private/Strata/StrataForwardLighting.ush
@@ -249,7 +249,7 @@ float3 StrataForwardLighting(
 	GridIndex = ComputeLightGridCellIndex((uint2)((MaterialParameters.SvPosition.xy - ResolvedView.ViewRectMin.xy) * View.LightProbeSizeRatioAndInvSizeRatio.zw), MaterialParameters.SvPosition.w, EyeIndex);
 #endif
 
-	UNROLL_N(STRATA_CLAMPED_LAYER_COUNT)
+	STRATA_UNROLL_N(STRATA_CLAMPED_LAYER_COUNT)
 	for (int l = 0; l < StrataData.LayerCount; ++l)
 	{
 		float3 LayerThroughput = 0.0f;
@@ -258,7 +258,7 @@ float3 StrataForwardLighting(
 		const bool bTopLayer = l == 0;
 		const bool bBottomLayer = l == (StrataData.LayerCount - 1);
 
-		UNROLL_N(STRATA_CLAMPED_MAX_BSDF_COUNT_PER_LAYER)
+		STRATA_UNROLL_N(STRATA_CLAMPED_MAX_BSDF_COUNT_PER_LAYER)
 		for (int i = 0; i < StrataData.Layers[l].BSDFCount; ++i)
 		{
 			FStrataBSDF BSDF = StrataData.Layers[l].BSDFs[i];
