Updating Mobile shaders to only used gather based shadow sampling
Stripped FEATURE_GATHER4 define as all platforms support it

diff --git a/Engine/Shaders/Private/MobileShadingModels.ush b/Engine/Shaders/Private/MobileShadingModels.ush
--- a/Engine/Shaders/Private/MobileShadingModels.ush
+++ b/Engine/Shaders/Private/MobileShadingModels.ush
@@ -10,10 +10,6 @@
 #include "HairBsdf.ush"
 #include "BurleyNormalizedSSSCommon.ush"
 
-#ifndef FEATURE_GATHER4
-#define FEATURE_GATHER4 1
-#endif
-
 #ifndef MAX_MOBILE_SHADOWCASCADES
 #define MAX_MOBILE_SHADOWCASCADES 0
 #endif 
@@ -486,9 +482,9 @@ half MobileShadowPCF(float2 ShadowUVs, FPCFSamplerSettings Settings)
 #elif MOBILE_SHADOW_QUALITY == 1
 	half ShadowMap = Manual1x1PCF(ShadowUVs, Settings);
 #elif MOBILE_SHADOW_QUALITY == 2
-	half ShadowMap = Manual2x2PCF(ShadowUVs, Settings);
-#elif MOBILE_SHADOW_QUALITY == 3
 	half ShadowMap = Manual3x3PCF(ShadowUVs, Settings);
+#elif MOBILE_SHADOW_QUALITY == 3
+	half ShadowMap = Manual5x5PCF(ShadowUVs, Settings);
 #else
 	#error Unsupported MOBILE_SHADOW_QUALITY value.
 #endif
diff --git a/Engine/Shaders/Private/ShadowFilteringCommon.ush b/Engine/Shaders/Private/ShadowFilteringCommon.ush
--- a/Engine/Shaders/Private/ShadowFilteringCommon.ush
+++ b/Engine/Shaders/Private/ShadowFilteringCommon.ush
@@ -262,30 +262,51 @@ float Manual3x3PCF(float2 ShadowPosition, FPCFSamplerSettings Settings)
 		float2 TexelPos = ShadowPosition * Settings.ShadowBufferSize.xy - 0.5f;	// bias to be consistent with texture filtering hardware
 		float2 Fraction = frac(TexelPos);
 		float2 TexelCorner = floor(TexelPos);	// bias to get reliable texel center content
-		{
-
 
-#if FEATURE_GATHER4
-			// Gather4 samples "at the following locations: (-,+),(+,+),(+,-),(-,-), where the magnitude of the deltas are always half a texel" - DX11 Func. Spec.
-			// So we need to offset to the centre of the 2x2 grid we want to sample.
-			float2 SamplePos = (TexelCorner + 1.0f) * Settings.ShadowBufferSize.zw;	// bias to get reliable texel center content
-			float4 SampleValues0 = CalculateOcclusion(Settings.ShadowDepthTexture.Gather(Settings.ShadowDepthTextureSampler, SamplePos, int2(-1, -1)), Settings);
-			float4 SampleValues1 = CalculateOcclusion(Settings.ShadowDepthTexture.Gather(Settings.ShadowDepthTextureSampler, SamplePos, int2(1, -1)), Settings);
-			float4 SampleValues2 = CalculateOcclusion(Settings.ShadowDepthTexture.Gather(Settings.ShadowDepthTextureSampler, SamplePos, int2(-1, 1)), Settings);
-			float4 SampleValues3 = CalculateOcclusion(Settings.ShadowDepthTexture.Gather(Settings.ShadowDepthTextureSampler, SamplePos, int2(1, 1)), Settings);
-			return PCF3x3gather(Fraction, SampleValues0, SampleValues1, SampleValues2, SampleValues3);
-#else // FEATURE_GATHER4 
-			float2 Sample00TexelCenter = TexelCorner - 0.5f;
-			float4 SampleValues0 = FetchRowOfFour(Sample00TexelCenter, 0, Settings);
-			float4 SampleValues1 = FetchRowOfFour(Sample00TexelCenter, 1, Settings);
-			float4 SampleValues2 = FetchRowOfFour(Sample00TexelCenter, 2, Settings);
-			float4 SampleValues3 = FetchRowOfFour(Sample00TexelCenter, 3, Settings);
-			return PCF3x3(Fraction, SampleValues0, SampleValues1, SampleValues2, SampleValues3);
-#endif // FEATURE_GATHER4
-		}
+		// Gather4 samples "at the following locations: (-,+),(+,+),(+,-),(-,-), where the magnitude of the deltas are always half a texel" - DX11 Func. Spec.
+		// So we need to offset to the centre of the 2x2 grid we want to sample.
+		float2 SamplePos = (TexelCorner + 1.0f) * Settings.ShadowBufferSize.zw;	// bias to get reliable texel center content
+		float4 SampleValues0 = CalculateOcclusion(Settings.ShadowDepthTexture.Gather(Settings.ShadowDepthTextureSampler, SamplePos, int2(-1, -1)), Settings);
+		float4 SampleValues1 = CalculateOcclusion(Settings.ShadowDepthTexture.Gather(Settings.ShadowDepthTextureSampler, SamplePos, int2(1, -1)), Settings);
+		float4 SampleValues2 = CalculateOcclusion(Settings.ShadowDepthTexture.Gather(Settings.ShadowDepthTextureSampler, SamplePos, int2(-1, 1)), Settings);
+		float4 SampleValues3 = CalculateOcclusion(Settings.ShadowDepthTexture.Gather(Settings.ShadowDepthTextureSampler, SamplePos, int2(1, 1)), Settings);
+		return PCF3x3gather(Fraction, SampleValues0, SampleValues1, SampleValues2, SampleValues3);
 	}
 }
 
+float Manual5x5PCF(float2 ShadowPosition, FPCFSamplerSettings Settings)
+{
+	// high quality, 6x6 samples, using gather4
+	float2 TexelPos = ShadowPosition * Settings.ShadowBufferSize.xy - 0.5f;	// bias to be consistent with texture filtering hardware
+	float2 Fraction = frac(TexelPos);
+	// Gather4 samples "at the following locations: (-,+),(+,+),(+,-),(-,-), where the magnitude of the deltas are always half a texel" - DX11 Func. Spec.
+	// So we need to offset to the centre of the 2x2 grid we want to sample.
+	float2 SamplePos = (floor(TexelPos) + 1.0f) * Settings.ShadowBufferSize.zw;	// bias to get reliable texel center content
+
+	float4 Values00 = CalculateOcclusion(Settings.ShadowDepthTexture.Gather(Settings.ShadowDepthTextureSampler, SamplePos, int2(-2, -2)), Settings);
+	float4 Values20 = CalculateOcclusion(Settings.ShadowDepthTexture.Gather(Settings.ShadowDepthTextureSampler, SamplePos, int2(0, -2)), Settings);
+	float4 Values40 = CalculateOcclusion(Settings.ShadowDepthTexture.Gather(Settings.ShadowDepthTextureSampler, SamplePos, int2(2, -2)), Settings);
+
+	float2 Row0 = HorizontalPCF5x2(Fraction, Values00, Values20, Values40);
+	float Results = Row0.x * (1.0f - Fraction.y) + Row0.y;
+
+	float4 Values02 = CalculateOcclusion(Settings.ShadowDepthTexture.Gather(Settings.ShadowDepthTextureSampler, SamplePos, int2(-2, 0)), Settings);
+	float4 Values22 = CalculateOcclusion(Settings.ShadowDepthTexture.Gather(Settings.ShadowDepthTextureSampler, SamplePos, int2(0, 0)), Settings);
+	float4 Values42 = CalculateOcclusion(Settings.ShadowDepthTexture.Gather(Settings.ShadowDepthTextureSampler, SamplePos, int2(2, 0)), Settings);
+
+	float2 Row1 = HorizontalPCF5x2(Fraction, Values02, Values22, Values42);
+	Results += Row1.x + Row1.y;
+
+	float4 Values04 = CalculateOcclusion(Settings.ShadowDepthTexture.Gather(Settings.ShadowDepthTextureSampler, SamplePos, int2(-2, 2)), Settings);
+	float4 Values24 = CalculateOcclusion(Settings.ShadowDepthTexture.Gather(Settings.ShadowDepthTextureSampler, SamplePos, int2(0, 2)), Settings);
+	float4 Values44 = CalculateOcclusion(Settings.ShadowDepthTexture.Gather(Settings.ShadowDepthTextureSampler, SamplePos, int2(2, 2)), Settings);
+
+	float2 Row2 = HorizontalPCF5x2(Fraction, Values04, Values24, Values44);
+	Results += Row2.x + Row2.y * Fraction.y;
+
+	return 0.04f * Results;
+}
+
 // break this out for forward rendering as it's not part of ManualPCF's set of shadowquality settings.
 float Manual2x2PCF(float2 ShadowPosition, FPCFSamplerSettings Settings)
 {
@@ -332,18 +353,10 @@ float Manual1x1PCF(float2 ShadowPosition, FPCFSamplerSettings Settings)
 	// using Gather: xyzw in counter clockwise order starting with the sample to the lower left of the queried location
 	float4 Samples;
 
-#if FEATURE_GATHER4
 	// Gather4 samples "at the following locations: (-,+),(+,+),(+,-),(-,-), where the magnitude of the deltas are always half a texel" - DX11 Func. Spec.
 	// So we need to offset to the centre of the 2x2 grid we want to sample.
 	float2 QuadCenter = floor(TexelPos) + 1.0f;	
 	Samples = Settings.ShadowDepthTexture.Gather(Settings.ShadowDepthTextureSampler, QuadCenter * Settings.ShadowBufferSize.zw);
-#else
-	float2 TexelCenter = floor(TexelPos) + 0.5f;	// bias to get reliable texel center content
-	Samples.x = Texture2DSampleLevel(Settings.ShadowDepthTexture, Settings.ShadowDepthTextureSampler, (TexelCenter.xy + float2(0, 1)) * Settings.ShadowBufferSize.zw, 0).r;
-	Samples.y = Texture2DSampleLevel(Settings.ShadowDepthTexture, Settings.ShadowDepthTextureSampler, (TexelCenter.xy + float2(1, 1)) * Settings.ShadowBufferSize.zw, 0).r;
-	Samples.z = Texture2DSampleLevel(Settings.ShadowDepthTexture, Settings.ShadowDepthTextureSampler, (TexelCenter.xy + float2(1, 0)) * Settings.ShadowBufferSize.zw, 0).r;
-	Samples.w = Texture2DSampleLevel(Settings.ShadowDepthTexture, Settings.ShadowDepthTextureSampler, (TexelCenter.xy + float2(0, 0)) * Settings.ShadowBufferSize.zw, 0).r;
-#endif
 
 	float4 Values00 = CalculateOcclusion(Samples, Settings);
 	return PCF1x1(Fraction, Values00);
@@ -359,111 +372,7 @@ float ManualPCF(float2 ShadowPosition, FPCFSamplerSettings Settings)
 	return Manual1x1PCF(ShadowPosition, Settings);
 #elif SHADOW_QUALITY == 3
 	return Manual3x3PCF(ShadowPosition, Settings);
-#elif FEATURE_GATHER4
-	// high quality, 6x6 samples, using gather4
-	{
-		float2 TexelPos = ShadowPosition * Settings.ShadowBufferSize.xy - 0.5f;	// bias to be consistent with texture filtering hardware
-		float2 Fraction = frac(TexelPos);
-		// Gather4 samples "at the following locations: (-,+),(+,+),(+,-),(-,-), where the magnitude of the deltas are always half a texel" - DX11 Func. Spec.
-		// So we need to offset to the centre of the 2x2 grid we want to sample.
-		float2 SamplePos = (floor(TexelPos) + 1.0f) * Settings.ShadowBufferSize.zw;	// bias to get reliable texel center content
-
-		float4 Values00 = CalculateOcclusion(Settings.ShadowDepthTexture.Gather(Settings.ShadowDepthTextureSampler, SamplePos, int2(-2, -2)), Settings);
-		float4 Values20 = CalculateOcclusion(Settings.ShadowDepthTexture.Gather(Settings.ShadowDepthTextureSampler, SamplePos, int2(0, -2)), Settings);
-		float4 Values40 = CalculateOcclusion(Settings.ShadowDepthTexture.Gather(Settings.ShadowDepthTextureSampler, SamplePos, int2(2, -2)), Settings);
-
-		float2 Row0 = HorizontalPCF5x2(Fraction, Values00, Values20, Values40);
-		float Results = Row0.x * (1.0f - Fraction.y) + Row0.y;
-
-		float4 Values02 = CalculateOcclusion(Settings.ShadowDepthTexture.Gather(Settings.ShadowDepthTextureSampler, SamplePos, int2(-2, 0)), Settings);
-		float4 Values22 = CalculateOcclusion(Settings.ShadowDepthTexture.Gather(Settings.ShadowDepthTextureSampler, SamplePos, int2(0, 0)), Settings);
-		float4 Values42 = CalculateOcclusion(Settings.ShadowDepthTexture.Gather(Settings.ShadowDepthTextureSampler, SamplePos, int2(2, 0)), Settings);
-
-		float2 Row1 = HorizontalPCF5x2(Fraction, Values02, Values22, Values42);
-		Results += Row1.x + Row1.y;
-
-		float4 Values04 = CalculateOcclusion(Settings.ShadowDepthTexture.Gather(Settings.ShadowDepthTextureSampler, SamplePos, int2(-2, 2)), Settings);
-		float4 Values24 = CalculateOcclusion(Settings.ShadowDepthTexture.Gather(Settings.ShadowDepthTextureSampler, SamplePos, int2(0, 2)), Settings);
-		float4 Values44 = CalculateOcclusion(Settings.ShadowDepthTexture.Gather(Settings.ShadowDepthTextureSampler, SamplePos, int2(2, 2)), Settings);
-
-		float2 Row2 = HorizontalPCF5x2(Fraction, Values04, Values24, Values44);
-		Results += Row2.x + Row2.y * Fraction.y;
-
-		return 0.04f * Results;
-	}
-
-#else // FEATURE_GATHER4
-
-	// high quality, 7x7 samples, not using gather4 (todo: ideally we make this 6x6 to get same results with gather code)
-	{
-		float2 Fraction = frac(ShadowPosition * Settings.ShadowBufferSize.xy);
-		float2 Sample00TexelCenter = floor(ShadowPosition * Settings.ShadowBufferSize.xy) - float2(3, 3);
-
-		// Fetch 7x7 shadowmap point samples
-		// Do 6x6 PCF samples, sharing the point samples between neighboring PCF samples
-		float4 Results;
-
-		float4 SampleValues03;
-		float4 SampleValues13;
-
-		{
-			float4 SampleValues10;
-			float4 SampleValues11;
-			float4 SampleValues12;
-
-			// Group work to minimize temporary registers and to split texture work with PCF ALU operations to hide texture latency
-			// Without this layout (all texture lookups at the beginning, PCF ALU's at the end) this shader was 4x slower on Nvidia cards
-			{
-				float4 SampleValues00 =  FetchRowOfFour(Sample00TexelCenter, 0, Settings);
-				SampleValues10.x = SampleValues00.w;
-
-				float4 SampleValues01 =  FetchRowOfFour(Sample00TexelCenter, 1, Settings);
-				SampleValues11.x = SampleValues01.w;
-
-				float4 SampleValues02 = FetchRowOfFour(Sample00TexelCenter, 2, Settings);
-				SampleValues12.x = SampleValues02.w;
-
-				SampleValues03 = FetchRowOfFour(Sample00TexelCenter, 3, Settings);
-				SampleValues13.x = SampleValues03.w;
-				Results.x = PCF3x3(Fraction, SampleValues00, SampleValues01, SampleValues02, SampleValues03);
-			}
-
-			{
-				SampleValues10.yzw = FetchRowOfThreeAfterFour(Sample00TexelCenter, 0, Settings);
-				SampleValues11.yzw = FetchRowOfThreeAfterFour(Sample00TexelCenter, 1, Settings);
-				SampleValues12.yzw = FetchRowOfThreeAfterFour(Sample00TexelCenter, 2, Settings);
-				SampleValues13.yzw = FetchRowOfThreeAfterFour(Sample00TexelCenter, 3, Settings);
-				Results.y = PCF3x3(Fraction, SampleValues10, SampleValues11, SampleValues12, SampleValues13);
-			}
-		}
-
-		{
-			float4 SampleValues14;
-			float4 SampleValues15;
-			float4 SampleValues16;
-
-			{
-				float4 SampleValues04 =  FetchRowOfFour(Sample00TexelCenter, 4, Settings);
-				SampleValues14.x = SampleValues04.w;
-
-				float4 SampleValues05 = FetchRowOfFour(Sample00TexelCenter, 5, Settings);
-				SampleValues15.x = SampleValues05.w;
-
-				float4 SampleValues06 = FetchRowOfFour(Sample00TexelCenter, 6, Settings);
-				SampleValues16.x = SampleValues06.w;
-
-				Results.z = PCF3x3(Fraction, SampleValues03, SampleValues04, SampleValues05, SampleValues06);
-			}
-
-			{
-				SampleValues14.yzw = FetchRowOfThreeAfterFour(Sample00TexelCenter, 4, Settings);
-				SampleValues15.yzw = FetchRowOfThreeAfterFour(Sample00TexelCenter, 5, Settings);
-				SampleValues16.yzw = FetchRowOfThreeAfterFour(Sample00TexelCenter, 6, Settings);
-				Results.w = PCF3x3(Fraction, SampleValues13, SampleValues14, SampleValues15, SampleValues16);
-			}
-		}
-
-		return dot(Results, .25f);
-	}
-#endif	// FEATURE_GATHER4
+#else
+	return Manual5x5PCF(ShadowPosition, Settings);
+#endif
 }
diff --git a/Engine/Shaders/Private/ShadowProjectionPixelShader.usf b/Engine/Shaders/Private/ShadowProjectionPixelShader.usf
--- a/Engine/Shaders/Private/ShadowProjectionPixelShader.usf
+++ b/Engine/Shaders/Private/ShadowProjectionPixelShader.usf
@@ -31,14 +31,6 @@
 #include "HairStrands/HairStrandsVisibilityCommonStruct.ush"
 #include "Common.ush"
 
-#if FEATURE_LEVEL >= FEATURE_LEVEL_SM5
-	// Gather() is supported
-	#define FEATURE_GATHER4 1
-#else
-	// Gather() is not supported
-	#define FEATURE_GATHER4 0
-#endif
-
 #if MODULATED_SHADOWS
 	#define MobileSceneTextures MobileBasePass.SceneTextures
 #endif
diff --git a/Engine/Source/Runtime/MaterialShaderQualitySettings/Classes/ShaderPlatformQualitySettings.h b/Engine/Source/Runtime/MaterialShaderQualitySettings/Classes/ShaderPlatformQualitySettings.h
--- a/Engine/Source/Runtime/MaterialShaderQualitySettings/Classes/ShaderPlatformQualitySettings.h
+++ b/Engine/Source/Runtime/MaterialShaderQualitySettings/Classes/ShaderPlatformQualitySettings.h
@@ -19,10 +19,10 @@ enum class EMobileShadowQuality : uint8
 	NoFiltering,
 	// Medium quality, 1x1 PCF filtering.
 	PCF_1x1 UMETA(DisplayName = "1x1 PCF"),
-	// Medium/High quality, 2x2 PCF filtering.
-	PCF_2x2 UMETA(DisplayName = "2x2 PCF"),
-	// Highest quality, 3x3 PCF filtering.
-	PCF_3x3 UMETA(DisplayName = "3x3 PCF")
+	// Medium/High quality, 3x3 PCF filtering.
+	PCF_3x3 UMETA(DisplayName = "3x3 PCF"),
+	// Highest quality, 5x5 PCF filtering.
+	PCF_5x5 UMETA(DisplayName = "5x5 PCF")
 };
 
 // FMaterialQualityOverrides represents the full set of possible material overrides per quality level.
@@ -41,7 +41,7 @@ public:
 		, bForceLQReflections(false)
 		, bForceDisablePreintegratedGF(false)
 		, bDisableMaterialNormalCalculation(false)
-		, MobileShadowQuality(EMobileShadowQuality::PCF_2x2)
+		, MobileShadowQuality(EMobileShadowQuality::PCF_3x3)
 	{
 	}
 
