Move (optionally) precomputed shadow factor into Strata header.

diff --git a/Engine/Shaders/Private/BasePassPixelShader.usf b/Engine/Shaders/Private/BasePassPixelShader.usf
--- a/Engine/Shaders/Private/BasePassPixelShader.usf
+++ b/Engine/Shaders/Private/BasePassPixelShader.usf
@@ -998,14 +998,6 @@ void FPixelShaderInOut_MainPS(
 		StrataPixelHeader.BSDFCount += StrataData.Layers[l].BSDFCount;
 	}
 	StrataPixelHeader.SharedLocalBases = MaterialParameters.SharedLocalBases;
-
-#if !GBUFFER_HAS_PRECSHADOWFACTOR
-	SetSkipPrecShadowMask(StrataPixelHeader, true);
-#endif
-#if (GBUFFER_HAS_PRECSHADOWFACTOR && WRITES_PRECSHADOWFACTOR_ZERO)
-	SetZeroPrecShadowMask(StrataPixelHeader, true);
-#endif
-
 	StrataPixelHeader.MaterialAO = MaterialAO;
 	SetCastContactShadow(StrataPixelHeader, GetPrimitiveData(MaterialParameters).Flags & PRIMITIVE_SCENE_DATA_FLAG_HAS_CAST_CONTACT_SHADOW);
 	SetDynamicIndirectShadowCasterRepresentation(StrataPixelHeader, GetPrimitiveData(MaterialParameters).Flags & PRIMITIVE_SCENE_DATA_FLAG_HAS_CAPSULE_REPRESENTATION);
@@ -1089,6 +1081,37 @@ void FPixelShaderInOut_MainPS(
 		);
 #endif // !STRATA_ENABLED
 
+	// Static shadow mask
+#if STRATA_ENABLED && !MATERIAL_SHADINGMODEL_UNLIT && GBUFFER_HAS_PRECSHADOWFACTOR
+	{
+		// Encode shadow mask only if the shadow mask is entirely non-zero and non-one
+	#if WRITES_PRECSHADOWFACTOR_ZERO
+		SetHasPrecShadowMask(StrataPixelHeader, false);
+		SetZeroPrecShadowMask(StrataPixelHeader, true);
+	#else
+	  #if ALLOW_STATIC_LIGHTING
+		const bool bAllZero = all(GBuffer.PrecomputedShadowFactors == 0);
+		const bool bAllOne  = all(GBuffer.PrecomputedShadowFactors == 1);
+		if (!bAllZero && !bAllOne)
+		{
+			SetHasPrecShadowMask(StrataPixelHeader, true);
+			StrataPixelHeader.PrecomputedShadowMask = GBuffer.PrecomputedShadowFactors;
+		}
+		else if (bAllZero)
+		{
+			SetHasPrecShadowMask(StrataPixelHeader, false);
+			SetZeroPrecShadowMask(StrataPixelHeader, true);
+		}
+		else if (bAllOne)
+	  #endif
+		{
+			SetHasPrecShadowMask(StrataPixelHeader, false);
+			SetZeroPrecShadowMask(StrataPixelHeader, false);
+		}
+	#endif
+	}
+#endif
+
 #if USES_GBUFFER
 	// This requires cleanup. Shader code that uses GBuffer.SelectiveOutputMask expects the outputmask to be in
 	// bits [4:7], but it gets packed as bits [0:3] in the flexible gbuffer since we might move it around.
diff --git a/Engine/Shaders/Private/DeferredLightPixelShaders.usf b/Engine/Shaders/Private/DeferredLightPixelShaders.usf
--- a/Engine/Shaders/Private/DeferredLightPixelShaders.usf
+++ b/Engine/Shaders/Private/DeferredLightPixelShaders.usf
@@ -166,26 +166,6 @@ float3 GetStableLargeWorldPosition(float2 SVPos, float DeviceZ)
 
 ////////////////////////////////////////////////////////////////////////////////////////////////////
 
-#if STRATA_ENABLED
-
-float4 StrataReadPrecomputedShadowFactors(FStrataPixelHeader StrataPixelHeader, int2 PixelPos)
-{
-	if (!HasSkipPrecShadowMask(StrataPixelHeader))
-	{
-#if ALLOW_STATIC_LIGHTING
-		float4 GBufferE = SceneTexturesStruct.GBufferETexture.Load(int3(PixelPos, 0));
-#else
-		float4 GBufferE = 1;
-#endif
-		return GBufferE;
-	}
-	return HasZeroPrecShadowMask(StrataPixelHeader) ? 0.0f : 1.0f;
-}
-
-#endif // STRATA_ENABLED
-
-////////////////////////////////////////////////////////////////////////////////////////////////////
-
 #if USE_HAIR_LIGHTING == 0
 
 void DeferredLightPixelMain(
@@ -250,10 +230,9 @@ void DeferredLightPixelMain(
 			float4 LightAttenuation = GetPerPixelLightAttenuation(InputParams.ScreenUV);
 
 			float Dither = InterleavedGradientNoise(InputParams.PixelPos, View.StateFrameIndexMod8);
-			float4 PrecomputedShadowFactors = StrataReadPrecomputedShadowFactors(StrataPixelHeader, PixelPos);
 			const uint FakeShadingModelID = 0;
 			const float FakeContactShadowOpacity = 1.0f;
-			GetShadowTerms(SceneDepth, PrecomputedShadowFactors, FakeShadingModelID, FakeContactShadowOpacity,
+			GetShadowTerms(SceneDepth, StrataPixelHeader.PrecomputedShadowMask, FakeShadingModelID, FakeContactShadowOpacity,
 				LightData, DerivedParams.TranslatedWorldPosition, L, LightAttenuation, Dither, ShadowTerms);
 
 			Radiance = StrataDeferredLighting(
diff --git a/Engine/Shaders/Private/ShadingFurnaceTest.usf b/Engine/Shaders/Private/ShadingFurnaceTest.usf
--- a/Engine/Shaders/Private/ShadingFurnaceTest.usf
+++ b/Engine/Shaders/Private/ShadingFurnaceTest.usf
@@ -224,7 +224,7 @@ void MainPS(
 						RectTexture, 
 						SurfaceShadow).xyz;
 				}
-			#else
+			#elif 0
 				{
 					FStrataAddressing StrataAddressing = GetStrataPixelDataByteOffset(PixelPos, uint2(View.BufferSizeAndInvSize.xy), Strata.MaxBytesPerPixel);
 					FStrataPixelHeader StrataPixelHeader = UnpackStrataHeaderIn(Strata.MaterialLobesBuffer, StrataAddressing);
diff --git a/Engine/Shaders/Private/Strata/Strata.ush b/Engine/Shaders/Private/Strata/Strata.ush
--- a/Engine/Shaders/Private/Strata/Strata.ush
+++ b/Engine/Shaders/Private/Strata/Strata.ush
@@ -808,22 +808,24 @@ struct FStrataPixelHeader
 #endif
 
 	float MaterialAO;
+	float4 PrecomputedShadowMask;
+
 	uint State;
 };
 
-#define HEADER_MASK_SKIPPRECSHADOWMASK			(1u << 0)
+#define HEADER_MASK_HASPRECSHADOWMASK			(1u << 0)
 #define HEADER_MASK_ZEROPRECSHADOWMASK			(1u << 1)
 #define HEADER_MASK_ISSINGLELAYERWATER			(1u << 2)	// Used to avoid loading BSDF or Categorisation data when generating tiles containing SLW material. Avoid running another categorisation.
 #define HEADER_MASK_CASTCONTACTSHADOW			(1u << 3)
 #define HEADER_MASK_HASDYNINDIRECTSHADOWCASTER	(1u << 4)
 
-bool HasSkipPrecShadowMask(const FStrataPixelHeader In)										{ return (In.State & HEADER_MASK_SKIPPRECSHADOWMASK) != 0; }
+bool HasPrecShadowMask(const FStrataPixelHeader In)											{ return (In.State & HEADER_MASK_HASPRECSHADOWMASK) != 0; }
 bool HasZeroPrecShadowMask(const FStrataPixelHeader In)										{ return (In.State & HEADER_MASK_ZEROPRECSHADOWMASK) != 0; }
 bool IsSingleLayerWater(const FStrataPixelHeader In)										{ return (In.State & HEADER_MASK_ISSINGLELAYERWATER) != 0; }
 bool DoesCastContactShadow(const FStrataPixelHeader In)										{ return (In.State & HEADER_MASK_CASTCONTACTSHADOW) != 0; }
 bool HasDynamicIndirectShadowCasterRepresentation(const FStrataPixelHeader In)				{ return (In.State & HEADER_MASK_HASDYNINDIRECTSHADOWCASTER) != 0; }
 
-void SetSkipPrecShadowMask(inout FStrataPixelHeader Out, bool bIn)							{ Out.State |= (bIn ? HEADER_MASK_SKIPPRECSHADOWMASK : 0u); }
+void SetHasPrecShadowMask(inout FStrataPixelHeader Out, bool bIn)							{ Out.State |= (bIn ? HEADER_MASK_HASPRECSHADOWMASK : 0u); }
 void SetZeroPrecShadowMask(inout FStrataPixelHeader Out, bool bIn)							{ Out.State |= (bIn ? HEADER_MASK_ZEROPRECSHADOWMASK : 0u); }
 void SetIsSingleLayerWater(inout FStrataPixelHeader Out, bool bIn)							{ Out.State |= (bIn ? HEADER_MASK_ISSINGLELAYERWATER : 0u); }
 void SetCastContactShadow (inout FStrataPixelHeader Out, bool bIn)							{ Out.State |= (bIn ? HEADER_MASK_CASTCONTACTSHADOW : 0u); }
@@ -840,6 +842,7 @@ FStrataPixelHeader InitialiseStrataPixelHeader()
 	Out.SharedLocalBasesTypes = 0;
 #endif
 	Out.MaterialAO = 1.0f;
+	Out.PrecomputedShadowMask = 1.0f;
 	Out.State = 0;
 	return Out;
 }
@@ -1723,7 +1726,7 @@ void StrataLoad_ColorGamma2ToLinearAlphaLinear(ByteAddressBuffer StrataBuffer, i
 	A = Data.a;
 }
 
-uint PackStrataHeaderIn(uint BSDFCount, FStrataPixelHeader In);
+uint2 PackStrataHeaderIn(uint BSDFCount, FStrataPixelHeader In);
 
 // If this is changed, please update the compiler side material size evaluation in StrataMaterial.cpp
 void PackStrataOut(
@@ -1952,7 +1955,15 @@ void PackStrataOut(
 
 	// 1. the header
 	AndIsSimpleMaterial(ClassificationData, BSDFVisibleCount == 1);
-	STRATA_STORE_UINT1(PackStrataHeaderIn(BSDFVisibleCount, StrataPixelHeader));
+	const uint2 PackedHeader = PackStrataHeaderIn(BSDFVisibleCount, StrataPixelHeader);
+	STRATA_STORE_UINT1(PackedHeader.x);
+#if ALLOW_STATIC_LIGHTING
+	if (HasPrecShadowMask(StrataPixelHeader))
+	{
+		STRATA_STORE_UINT1(PackedHeader.y);
+	}
+#endif
+
 #if STRATA_INLINE_SHADING
 	UNROLL
 	for (uint i = 0; i < StrataPixelHeader.SharedLocalBases.Count; ++i)
@@ -2099,18 +2110,25 @@ void PackStrataOut(
 }
 
 // Pack only strata header
-uint PackStrataHeaderIn(uint InBSDFCount, FStrataPixelHeader InHeader)
+uint2 PackStrataHeaderIn(uint InBSDFCount, FStrataPixelHeader InHeader)
 {
-	uint Out = 0;
-	HEADER_SETBSDFCOUNT(Out, InBSDFCount);
+	uint2 Out = 0;
+	HEADER_SETBSDFCOUNT(Out.x, InBSDFCount);
 
 #if STRATA_INLINE_SHADING
-	HEADER_SETSHAREDLOCALBASESCOUNT(Out, InHeader.SharedLocalBases.Count);
-	HEADER_SETSHAREDLOCALBASESTYPE(Out, InHeader.SharedLocalBases.Types);
+	HEADER_SETSHAREDLOCALBASESCOUNT(Out.x, InHeader.SharedLocalBases.Count);
+	HEADER_SETSHAREDLOCALBASESTYPE(Out.x, InHeader.SharedLocalBases.Types);
 #endif
 
-	HEADER_SETAO(Out, uint(InHeader.MaterialAO * 255.0f));
-	HEADER_SETSTATE(Out, InHeader.State);
+	HEADER_SETAO(Out.x, uint(InHeader.MaterialAO * 255.0f));
+	HEADER_SETSTATE(Out.x, InHeader.State);
+
+#if ALLOW_STATIC_LIGHTING
+	if (HasPrecShadowMask(InHeader))
+	{
+		Out.y = PackRGBA8(InHeader.PrecomputedShadowMask);
+	}
+#endif
 	return Out;
 }
 
@@ -2141,6 +2159,14 @@ FStrataPixelHeader UnpackStrataHeaderIn(ByteAddressBuffer StrataBuffer, inout FS
 	Out.MaterialAO = float(HEADER_GETAO(PackedHeader)) * (1.0f / 255.0f);
 	Out.State = HEADER_GETSTATE(PackedHeader);
 
+	// Pre-shadow mask
+#if ALLOW_STATIC_LIGHTING
+	Out.PrecomputedShadowMask = HasZeroPrecShadowMask(Out) ? 0.f : 1.f;
+	if (HasPrecShadowMask(Out))
+	{
+		Out.PrecomputedShadowMask = UnpackRGBA8(StrataLoadUint1(StrataBuffer, StrataAddressing));
+	}
+#endif
 	return Out;
 }
 
diff --git a/Engine/Shaders/Private/Strata/StrataVisualize.usf b/Engine/Shaders/Private/Strata/StrataVisualize.usf
--- a/Engine/Shaders/Private/Strata/StrataVisualize.usf
+++ b/Engine/Shaders/Private/Strata/StrataVisualize.usf
@@ -201,7 +201,7 @@ void VisualizeMaterialPS(
 
 #define PRINTBSDFCOUNT()	PRINTBSDFCHAR(min(_0_+Header.BSDFCount,_F_)); PRINTBSDFCHAR(_SPC_); PRINTBSDFCHAR(_B_); PRINTBSDFCHAR(_S_); PRINTBSDFCHAR(_D_); PRINTBSDFCHAR(_F_);
 #define PRINTMATERIALAO()	PRINTBSDFCHAR(_A_); PRINTBSDFCHAR(_O_); PRINTBSDFCHAR(_SPC_); PRINTFLOAT1(Header.MaterialAO);
-#define PRINTMATERIALSKIPPRECSHADOWMASK()	PRINTBSDFCHAR(_S_); PRINTBSDFCHAR(_K_); PRINTBSDFCHAR(_I_); PRINTBSDFCHAR(_P_); PRINTBSDFCHAR(_P_); PRINTBSDFCHAR(_R_); PRINTBSDFCHAR(_E_); PRINTBSDFCHAR(_S_); PRINTBSDFCHAR(_H_); PRINTBSDFCHAR(_A_); PRINTBSDFCHAR(_D_); PRINTBSDFCHAR(_SPC_); PRINTBSDFCHAR(_0_+(HasSkipPrecShadowMask(Header)?1:0));
+#define PRINTMATERIALSKIPPRECSHADOWMASK()	PRINTBSDFCHAR(_H_); PRINTBSDFCHAR(_A_); PRINTBSDFCHAR(_S_); PRINTBSDFCHAR(_SPC_); PRINTBSDFCHAR(_P_); PRINTBSDFCHAR(_R_); PRINTBSDFCHAR(_E_); PRINTBSDFCHAR(_S_); PRINTBSDFCHAR(_H_); PRINTBSDFCHAR(_A_); PRINTBSDFCHAR(_D_); PRINTBSDFCHAR(_SPC_); PRINTBSDFCHAR(_0_+(HasPrecShadowMask(Header)?1:0));
 #define PRINTMATERIALZEROPRECSHADOWMASK()	PRINTBSDFCHAR(_Z_); PRINTBSDFCHAR(_E_); PRINTBSDFCHAR(_R_); PRINTBSDFCHAR(_O_); PRINTBSDFCHAR(_P_); PRINTBSDFCHAR(_R_); PRINTBSDFCHAR(_E_); PRINTBSDFCHAR(_S_); PRINTBSDFCHAR(_H_); PRINTBSDFCHAR(_A_); PRINTBSDFCHAR(_D_); PRINTBSDFCHAR(_SPC_); PRINTBSDFCHAR(_0_+(HasZeroPrecShadowMask(Header)?1:0));
 #define PRINTMATERIALSLW()	PRINTBSDFCHAR(_S_); PRINTBSDFCHAR(_L_); PRINTBSDFCHAR(_W_); PRINTBSDFCHAR(_SPC_); PRINTBSDFCHAR(_0_+(IsSingleLayerWater(Header)?1:0));
 #define PRINTMATERIALCONTSHAD()	PRINTBSDFCHAR(_C_); PRINTBSDFCHAR(_O_); PRINTBSDFCHAR(_N_);  PRINTBSDFCHAR(_T_); PRINTBSDFCHAR(_SPC_);PRINTBSDFCHAR(_S_); PRINTBSDFCHAR(_H_); PRINTBSDFCHAR(_A_);  PRINTBSDFCHAR(_D_); PRINTBSDFCHAR(_SPC_); PRINTBSDFCHAR(_0_+(DoesCastContactShadow(Header)?1:0));
diff --git a/Engine/Source/Runtime/Engine/Private/Materials/StrataMaterial.cpp b/Engine/Source/Runtime/Engine/Private/Materials/StrataMaterial.cpp
--- a/Engine/Source/Runtime/Engine/Private/Materials/StrataMaterial.cpp
+++ b/Engine/Source/Runtime/Engine/Private/Materials/StrataMaterial.cpp
@@ -333,6 +333,8 @@ FStrataMaterialAnalysisResult::FStrataMaterialAnalysisResult()
 FStrataMaterialAnalysisResult StrataCompilationInfoMaterialAnalysis(FMaterialCompiler* Compiler, const FStrataMaterialCompilationInfo& Material, const uint32 StrataBytePerPixel)
 {
 	const uint32 UintByteSize = sizeof(uint32);
+	static const auto CVar = IConsoleManager::Get().FindTConsoleVariableDataInt(TEXT("r.AllowStaticLighting"));
+	const bool bAllowStaticLighting = CVar->GetValueOnAnyThread() > 0;
 
 	FStrataMaterialAnalysisResult Result;
 
@@ -340,6 +342,12 @@ FStrataMaterialAnalysisResult StrataCompilationInfoMaterialAnalysis(FMaterialCom
 
 	// SharedLocalBases_BSDFCount
 	Result.RequestedByteCount += UintByteSize;
+
+	// Pre-calc shadow mask
+	if (bAllowStaticLighting)
+	{
+		Result.RequestedByteCount += UintByteSize;
+	}
 	// shared local bases between BSDFs
 	Result.RequestedByteCount += Compiler->StrataCompilationInfoGetSharedLocalBasesCount() * STRATA_PACKED_SHAREDLOCALBASIS_STRIDE_BYTES;
 	
diff --git a/Engine/Source/Runtime/Engine/Private/ShaderCompiler/ShaderGenerationUtil.cpp b/Engine/Source/Runtime/Engine/Private/ShaderCompiler/ShaderGenerationUtil.cpp
--- a/Engine/Source/Runtime/Engine/Private/ShaderCompiler/ShaderGenerationUtil.cpp
+++ b/Engine/Source/Runtime/Engine/Private/ShaderCompiler/ShaderGenerationUtil.cpp
@@ -1092,7 +1092,11 @@ static FString CreateGBufferDecodeFunctionDirect(const FGBufferInfo& BufferInfo)
 			FullStr += CurrLine;
 		}
 	}
-	FullStr += TEXT(",\n\tfloat CustomNativeDepth");
+	if (!bFirst)
+	{
+		FullStr += TEXT(",\n\t\t");
+	}
+	FullStr += TEXT(" \n\tfloat CustomNativeDepth");
 	FullStr += TEXT(",\n\tfloat4 AnisotropicData");
 	FullStr += TEXT(",\n\tuint CustomStencil");
 	FullStr += TEXT(",\n\tfloat SceneDepth");
@@ -1502,7 +1506,12 @@ static FString CreateGBufferDecodeFunctionVariation(const FGBufferInfo& BufferIn
 			FullStr += CurrLine;
 		}
 	}
-	FullStr += TEXT(",\n\t\tCustomNativeDepth");
+
+	if (!bFirst)
+	{
+		FullStr += TEXT(",\n\t\t");
+	}
+	FullStr += TEXT(" \n\t\tCustomNativeDepth");
 	FullStr += TEXT(",\n\t\tAnisotropicData");
 	FullStr += TEXT(",\n\t\tCustomStencil");
 	FullStr += TEXT(",\n\t\tSceneDepth");
@@ -1672,7 +1681,7 @@ static void SetStandardGBufferSlots(bool Slots[], bool bWriteEmissive, bool bHas
 {
 	Slots[GBS_SceneColor] = bWriteEmissive;
 	Slots[GBS_Velocity] = bHasVelocity;
-	Slots[GBS_PrecomputedShadowFactor] = bHasStaticLighting;
+	Slots[GBS_PrecomputedShadowFactor] = bIsStrataMaterial ? false : true;
 
 	Slots[GBS_WorldNormal] =			bIsStrataMaterial ? false : true;
 	Slots[GBS_PerObjectGBufferData] =	bIsStrataMaterial ? false : true;
diff --git a/Engine/Source/Runtime/RenderCore/Private/GBufferInfo.cpp b/Engine/Source/Runtime/RenderCore/Private/GBufferInfo.cpp
--- a/Engine/Source/Runtime/RenderCore/Private/GBufferInfo.cpp
+++ b/Engine/Source/Runtime/RenderCore/Private/GBufferInfo.cpp
@@ -246,11 +246,6 @@ FGBufferInfo RENDERCORE_API FetchLegacyGBufferInfo(const FGBufferParams& Params)
 		TargetGBufferC = -1;
 
 		Info.NumTargets = 1;
-		if (Params.bHasPrecShadowFactor)
-		{
-			TargetGBufferE = Info.NumTargets++;
-
-		}
 		if (Params.bHasVelocity)
 		{
 			TargetVelocity = Info.NumTargets++;
@@ -278,16 +273,6 @@ FGBufferInfo RENDERCORE_API FetchLegacyGBufferInfo(const FGBufferParams& Params)
 			}
 		}
 
-		if (Params.bHasPrecShadowFactor)
-		{
-			Info.Targets[TargetGBufferE].Init(GBT_Unorm_8_8_8_8, TEXT("GBufferE"), false, true, true, false); // Precalc
-			Info.Slots[GBS_PrecomputedShadowFactor] = FGBufferItem(GBS_PrecomputedShadowFactor, GBC_Raw_Unorm_8_8_8_8, GBCH_Both);
-			Info.Slots[GBS_PrecomputedShadowFactor].Packing[0] = FGBufferPacking(TargetGBufferE, 0, 0);
-			Info.Slots[GBS_PrecomputedShadowFactor].Packing[1] = FGBufferPacking(TargetGBufferE, 1, 1);
-			Info.Slots[GBS_PrecomputedShadowFactor].Packing[2] = FGBufferPacking(TargetGBufferE, 2, 2);
-			Info.Slots[GBS_PrecomputedShadowFactor].Packing[3] = FGBufferPacking(TargetGBufferE, 3, 3);
-		}
-
 		return Info;
 	}
 
