Strata uint4 based load/store. Disabled for now.

diff --git a/Engine/Shaders/Private/Strata/Strata.ush b/Engine/Shaders/Private/Strata/Strata.ush
--- a/Engine/Shaders/Private/Strata/Strata.ush
+++ b/Engine/Shaders/Private/Strata/Strata.ush
@@ -24,6 +24,9 @@
 #define STRATA_UNROLL_N(X) UNROLL_N(X)
 #endif
 
+#define STRATA_MEMACCESS_UINT  0
+#define STRATA_MEMACCESS_UINT4 1
+#define STRATA_MEMACCESS STRATA_MEMACCESS_UINT
 
 // During the base pass or forward rendering, shared local bases are simply available from registers.
 // If a BSDF is loaded for processing from the strata buffer, normals are unpacked on demand to reduce VGPR pressure and increase occupancy.
@@ -693,24 +696,40 @@ uint DivideAndRoundUp(uint Dividend, uint Divisor, uint DivisorAsBitShift)
 struct FStrataAddressing
 {
 	uint  CurrentAddr;
+
+#if STRATA_MEMACCESS == STRATA_MEMACCESS_UINT4
+	uint  ByteOffsetToNextUint4;
+	uint4 Uint4Cache;
+	int   UintProcessed;	// uints remaining to read, or uint appended before write
+
+	// Packedlocal basis in a uint4 available for unpacking later.
+	uint4 PackedLocalBasis;
+#if STRATA_MAX_SHAREDLOCALBASES_REGISTERS!=4
+#error We do not support more than 4 PackedLocalBasis as of today.
+#endif
+
+#else
 	uint  ByteOffsetToNextUint;
-	uint  ReadBytes;
+#endif
+
+	// This is to apply checkboard pattern when needed
 	uint2 PixelCoords;
+
+	// This should be compiled out as not used in most cases expect for debug and visualization.
+	uint  ReadBytes;
 };
 FStrataAddressing GetStrataPixelDataByteOffset(uint2 PixelPos, uint2 ViewBufferSize, uint StrataMaxBytesPerPixel)
 {
 	FStrataAddressing StrataAddressing = (FStrataAddressing)0;
 
-#if 0
-	// Linear storage, pixel0 bytes, pixel1 bytes, etc.
-	StrataAddressing.CurrentAddr = (PixelPos.y * ViewBufferSize.x + PixelPos.x) * StrataMaxBytesPerPixel;
-	StrataAddressing.ByteOffsetToNextUint = 4;
-#else
-
 	// A tile of pixels has each uint that can be read packed in contiguous memory to be cache efficient. Then, the next set of uint to be read is in another "plane", after all uint from all the pixel/tile.
 	// Without that, L2 is rarely hit and GCN TA FIFO command buffer is getting full in no time because memory fetches are not served in time (GPU cost then increase by at least x2)
 
-	uint TilePixelSizeBytes			= 4;	// sizeof(uint)
+#if STRATA_MEMACCESS == STRATA_MEMACCESS_UINT4
+	uint TilePixelSizeBytes			= 4*4;	// 4*sizeof(uint)
+#else
+	uint TilePixelSizeBytes			= 4;
+#endif
 	uint TileSizeBytes				= STRATA_DATA_TILE_SIZE * STRATA_DATA_TILE_SIZE * TilePixelSizeBytes;
 	uint2 TileCount					= uint2(DivideAndRoundUp(ViewBufferSize.x, STRATA_DATA_TILE_SIZE, STRATA_DATA_TILE_SIZE_DIV_AS_SHIFT), DivideAndRoundUp(ViewBufferSize.y, STRATA_DATA_TILE_SIZE, STRATA_DATA_TILE_SIZE_DIV_AS_SHIFT));
 	uint2 TileCoord					= uint2(PixelPos.x >> STRATA_DATA_TILE_SIZE_DIV_AS_SHIFT, PixelPos.y >> STRATA_DATA_TILE_SIZE_DIV_AS_SHIFT);
@@ -724,9 +743,12 @@ FStrataAddressing GetStrataPixelDataByteOffset(uint2 PixelPos, uint2 ViewBufferS
 #endif
 	StrataAddressing.CurrentAddr	= PreviousTilesBytes + PreviousPixelsInTileBytes;
 
-	// All
+#if STRATA_MEMACCESS == STRATA_MEMACCESS_UINT4
+	StrataAddressing.ByteOffsetToNextUint4 = TileCount.x * TileCount.y * TileSizeBytes;
+	StrataAddressing.Uint4Cache = 0;
+	StrataAddressing.UintProcessed = 0;
+#else
 	StrataAddressing.ByteOffsetToNextUint = TileCount.x * TileCount.y * TileSizeBytes;
-
 #endif
 
 	StrataAddressing.PixelCoords = PixelPos;
@@ -739,23 +761,76 @@ FStrataAddressing GetStrataPixelDataByteOffset(uint2 PixelPos, uint2 ViewBufferS
 void StrataStoreUint1(RWByteAddressBuffer StrataBuffer, inout FStrataAddressing StrataAddressing, uint Data) {}
 #define STRATA_STORE_UINT1(x) 
 
+void FinalizeWrites(RWByteAddressBuffer StrataBuffer, inout FStrataAddressing StrataAddressing) {}
+
 uint StrataLoadUint1(ByteAddressBuffer StrataBuffer, inout FStrataAddressing StrataAddressing) { return 0; }
 
 #else
 
 void StrataStoreUint1(RWByteAddressBuffer StrataBuffer, inout FStrataAddressing StrataAddressing, uint Data)
 {
+#if STRATA_MEMACCESS == STRATA_MEMACCESS_UINT4
+	BRANCH
+	if(StrataAddressing.UintProcessed == 4)
+	{
+		StrataBuffer.Store4(StrataAddressing.CurrentAddr, StrataAddressing.Uint4Cache);
+		StrataAddressing.CurrentAddr += StrataAddressing.ByteOffsetToNextUint4;
+		StrataAddressing.UintProcessed = 0;
+		StrataAddressing.ReadBytes += 16;
+	}
+
+	// TODO check what is done during this addressing... It should be a simple bit select
+	StrataAddressing.Uint4Cache[0] = StrataAddressing.UintProcessed == 0 ? Data : StrataAddressing.Uint4Cache[0];
+	StrataAddressing.Uint4Cache[1] = StrataAddressing.UintProcessed == 1 ? Data : StrataAddressing.Uint4Cache[1];
+	StrataAddressing.Uint4Cache[2] = StrataAddressing.UintProcessed == 2 ? Data : StrataAddressing.Uint4Cache[2];
+	StrataAddressing.Uint4Cache[3] = StrataAddressing.UintProcessed == 3 ? Data : StrataAddressing.Uint4Cache[3];
+	StrataAddressing.UintProcessed++;
+#else
 	StrataBuffer.Store(StrataAddressing.CurrentAddr, Data);
 	StrataAddressing.CurrentAddr += StrataAddressing.ByteOffsetToNextUint;
 	StrataAddressing.ReadBytes += 4;
+#endif
 }
 #define STRATA_STORE_UINT1(x)	StrataStoreUint1(StrataBuffer, StrataAddressing, x)
 
+void FinalizeWrites(RWByteAddressBuffer StrataBuffer, inout FStrataAddressing StrataAddressing)
+{
+#if STRATA_MEMACCESS == STRATA_MEMACCESS_UINT4
+	BRANCH
+	if(StrataAddressing.UintProcessed > 0)
+	{
+		StrataBuffer.Store4(StrataAddressing.CurrentAddr, StrataAddressing.Uint4Cache);
+		StrataAddressing.CurrentAddr += StrataAddressing.ByteOffsetToNextUint4;
+		StrataAddressing.UintProcessed = 0;
+		StrataAddressing.ReadBytes += 16;
+	}
+#endif
+}
+
 uint StrataLoadUint1(ByteAddressBuffer StrataBuffer, inout FStrataAddressing StrataAddressing)
 {
+#if STRATA_MEMACCESS == STRATA_MEMACCESS_UINT4
+	BRANCH
+	if(StrataAddressing.UintProcessed == 0)
+	{
+		StrataAddressing.Uint4Cache = StrataBuffer.Load4(StrataAddressing.CurrentAddr);
+		StrataAddressing.CurrentAddr += StrataAddressing.ByteOffsetToNextUint4;
+		StrataAddressing.UintProcessed = 4;
+		StrataAddressing.ReadBytes += 16;
+	}
+
+	// TODO check what is done during this addressing... It should be a simple bit select
+	uint Data = StrataAddressing.Uint4Cache[0];
+	Data = StrataAddressing.UintProcessed == 3 ? StrataAddressing.Uint4Cache[1] : Data;
+	Data = StrataAddressing.UintProcessed == 2 ? StrataAddressing.Uint4Cache[2] : Data;
+	Data = StrataAddressing.UintProcessed == 1 ? StrataAddressing.Uint4Cache[3] : Data;
+	StrataAddressing.UintProcessed--;
+#else
 	uint Data = StrataBuffer.Load(StrataAddressing.CurrentAddr);
 	StrataAddressing.CurrentAddr += StrataAddressing.ByteOffsetToNextUint;
 	StrataAddressing.ReadBytes += 4;
+#endif
+
 	return Data;
 }
 
@@ -803,7 +878,9 @@ struct FStrataPixelHeader
 	FSharedLocalBases SharedLocalBases;
 #else
 	ByteAddressBuffer StrataBuffer;
+	#if STRATA_MEMACCESS == STRATA_MEMACCESS_UINT
 	uint SharedLocalBasesByteOffset;
+	#endif
 	uint SharedLocalBasesTypes;
 #endif
 
@@ -838,7 +915,9 @@ FStrataPixelHeader InitialiseStrataPixelHeader()
 #if STRATA_INLINE_SHADING
 	Out.SharedLocalBases = (FSharedLocalBases)0;
 #else
+	#if STRATA_MEMACCESS == STRATA_MEMACCESS_UINT
 	Out.SharedLocalBasesByteOffset = 0;
+	#endif
 	Out.SharedLocalBasesTypes = 0;
 #endif
 	Out.MaterialAO = 1.0f;
@@ -856,7 +935,12 @@ float3x3 StrataGetBSDFSharedBasis(in FStrataPixelHeader StrataPixelHeader, in FS
 	float3 Bitangent = cross(Normal, Tangent);
 	return float3x3(Tangent, Bitangent, Normal);
 #else
+
+#if STRATA_MEMACCESS == STRATA_MEMACCESS_UINT4
+	uint PackedBasis = StrataAddressing.PackedLocalBasis[BSDF_GETSHAREDLOCALBASISID(BSDF)];
+#else
 	uint PackedBasis = StrataPixelHeader.StrataBuffer.Load(StrataPixelHeader.SharedLocalBasesByteOffset + BSDF_GETSHAREDLOCALBASISID(BSDF) * StrataAddressing.ByteOffsetToNextUint);
+#endif
 	uint BasisType   = StrataGetSharedLocalBasisType(StrataPixelHeader.SharedLocalBasesTypes, BSDF_GETSHAREDLOCALBASISID(BSDF));
 	float3x3 OutTangentBasis;
 	if (BasisType == STRATA_BASIS_TYPE_NORMAL)
@@ -2106,6 +2190,8 @@ void PackStrataOut(
 		}
 	}
 
+	FinalizeWrites(StrataBuffer, StrataAddressing);
+
 	checkSlow(BSDFCount == BSDFVisibleCount);
 }
 
@@ -2144,13 +2230,23 @@ FStrataPixelHeader UnpackStrataHeaderIn(ByteAddressBuffer StrataBuffer, inout FS
 	uint SharedLocalBasesTypes = HEADER_GETSHAREDLOCALBASESTYPE(PackedHeader);
 
 #if STRATA_INLINE_SHADING==0
-	// Only keep the offset to the shared local bases memory. We are going to load them on demand.
-	Out.SharedLocalBasesByteOffset = StrataAddressing.CurrentAddr;
+
 	Out.StrataBuffer = StrataBuffer;
 	Out.SharedLocalBasesTypes = SharedLocalBasesTypes;
 
+#if STRATA_MEMACCESS == STRATA_MEMACCESS_UINT4
+	// Now load the packed local bases (stocked on the addressin part because it is only available while loading strata data)
+	for (uint i = 0; i < SharedLocalBasesCount; ++i)
+	{
+		StrataAddressing.PackedLocalBasis[i] = StrataLoadUint1(Out.StrataBuffer, StrataAddressing);
+	}
+#else
+	// Only keep the offset to the shared local bases memory. We are going to load them on demand.
+	Out.SharedLocalBasesByteOffset = StrataAddressing.CurrentAddr;
+
 	// Now skip over the shared local basis memory to be able to load the BSDF content
 	StrataAddressing.CurrentAddr += SharedLocalBasesCount * StrataAddressing.ByteOffsetToNextUint;
+#endif
 
 	// And account for the space used by normal for the debug information
 	StrataAddressing.ReadBytes += SharedLocalBasesCount * STRATA_PACKED_SHAREDLOCALBASIS_STRIDE_BYTES;
