Strata tiny renaming.

diff --git a/Engine/Shaders/Private/Strata/Strata.ush b/Engine/Shaders/Private/Strata/Strata.ush
--- a/Engine/Shaders/Private/Strata/Strata.ush
+++ b/Engine/Shaders/Private/Strata/Strata.ush
@@ -1810,7 +1810,7 @@ void StrataLoad_ColorGamma2ToLinearAlphaLinear(ByteAddressBuffer StrataBuffer, i
 	A = Data.a;
 }
 
-uint2 PackStrataHeaderIn(uint BSDFCount, FStrataPixelHeader In);
+uint2 PackStrataHeader(uint BSDFCount, FStrataPixelHeader In);
 
 // If this is changed, please update the compiler side material size evaluation in StrataMaterial.cpp
 void PackStrataOut(
@@ -2039,7 +2039,7 @@ void PackStrataOut(
 
 	// 1. the header
 	AndIsSimpleMaterial(ClassificationData, BSDFVisibleCount == 1);
-	const uint2 PackedHeader = PackStrataHeaderIn(BSDFVisibleCount, StrataPixelHeader);
+	const uint2 PackedHeader = PackStrataHeader(BSDFVisibleCount, StrataPixelHeader);
 	STRATA_STORE_UINT1(PackedHeader.x);
 #if ALLOW_STATIC_LIGHTING
 	if (HasPrecShadowMask(StrataPixelHeader))
@@ -2196,7 +2196,7 @@ void PackStrataOut(
 }
 
 // Pack only strata header
-uint2 PackStrataHeaderIn(uint InBSDFCount, FStrataPixelHeader InHeader)
+uint2 PackStrataHeader(uint InBSDFCount, FStrataPixelHeader InHeader)
 {
 	uint2 Out = 0;
 	HEADER_SETBSDFCOUNT(Out.x, InBSDFCount);
diff --git a/Engine/Shaders/Private/Strata/StrataMaterialClassification.usf b/Engine/Shaders/Private/Strata/StrataMaterialClassification.usf
--- a/Engine/Shaders/Private/Strata/StrataMaterialClassification.usf
+++ b/Engine/Shaders/Private/Strata/StrataMaterialClassification.usf
@@ -364,7 +364,7 @@ void ClearMaterialBufferMainCS(uint2 DispatchThreadId : SV_DispatchThreadID)
 
 	FStrataPixelHeader StrataHeader = InitialiseStrataPixelHeader();
 	uint BSDFCount = 0;
-	uint PackedStrataHeader = PackStrataHeaderIn(BSDFCount, StrataHeader);
+	uint PackedStrataHeader = PackStrataHeader(BSDFCount, StrataHeader);
 
 	FStrataAddressing StrataAddressing = GetStrataPixelDataByteOffset(DispatchThreadId.xy, TiledViewBufferResolution, MaxBytesPerPixel);
 	MaterialLobesBufferUAV.Store(StrataAddressing.CurrentAddr, PackedStrataHeader);
