Add un/packR8 to strata encoding, header sanity check, and comment/sections.

diff --git a/Engine/Shaders/Private/Strata/Strata.ush b/Engine/Shaders/Private/Strata/Strata.ush
--- a/Engine/Shaders/Private/Strata/Strata.ush
+++ b/Engine/Shaders/Private/Strata/Strata.ush
@@ -95,12 +95,17 @@ struct FStrataBSDF
 };
 
 
-
+////////////////////////////////////////////////////////////////////////////////////////////////////////////////////////////////////////////////////////////////////////////////////////////////
+// Bit utils functions
 #define TO_BITMASK(BitCount)							((1u<<(BitCount))-1)
 #define TO_OFFSET_BITMASK(BitCount, BitOffset)			(TO_BITMASK(BitCount)<<(BitOffset))
 #define READ_BITS(X, BitCount, BitOffset)				uint((X>>(BitOffset)) & TO_BITMASK(BitCount))
 #define WRITE_BITS(X, BitCount, BitOffset, NewValue)	X = (X & (~TO_OFFSET_BITMASK(BitCount, BitOffset))) | (((NewValue) & TO_BITMASK(BitCount)) << (BitOffset))
 
+
+///////////////////////////////////////////////////////////////////////////////
+// Material Header
+// 
 // Bit count for each component stored in a Strata header
 // For HEADER_BIT_BSDF_COUNT and  HEADER_BIT_SHAREDLOCALBASES_X, please refer to values in StrataDefinitions.h
 #define HEADER_BIT_COUNT_BSDF_COUNT				4
@@ -116,6 +121,18 @@ struct FStrataBSDF
 #define HEADER_BIT_OFFSET_AO					(HEADER_BIT_OFFSET_SHAREDLOCALBASES_TYPE	+ HEADER_BIT_COUNT_SHAREDLOCALBASES_TYPE)
 #define HEADER_BIT_OFFSET_STATE					(HEADER_BIT_OFFSET_AO						+ HEADER_BIT_COUNT_AO)
 
+// Sanity check
+#if (HEADER_BIT_COUNT_STATE + HEADER_BIT_OFFSET_STATE) > 32
+#error Strata Header is larger than 32 bits
+#endif
+
+// State bit masks
+#define HEADER_MASK_HASPRECSHADOWMASK			(1u << 0)
+#define HEADER_MASK_ZEROPRECSHADOWMASK			(1u << 1)
+#define HEADER_MASK_ISSINGLELAYERWATER			(1u << 2)	// Used to avoid loading BSDF or Categorisation data when generating tiles containing SLW material. Avoid running another categorisation.
+#define HEADER_MASK_CASTCONTACTSHADOW			(1u << 3)
+#define HEADER_MASK_HASDYNINDIRECTSHADOWCASTER	(1u << 4)
+
 #define HEADER_GETBSDFCOUNT(X)						READ_BITS( X, HEADER_BIT_COUNT_BSDF_COUNT, HEADER_BIT_OFFSET_BSDF_COUNT)
 #define HEADER_SETBSDFCOUNT(X, NewValue)			WRITE_BITS(X, HEADER_BIT_COUNT_BSDF_COUNT, HEADER_BIT_OFFSET_BSDF_COUNT, NewValue)
 
@@ -131,6 +148,9 @@ struct FStrataBSDF
 #define HEADER_GETSTATE(X)							READ_BITS( X, HEADER_BIT_COUNT_STATE, HEADER_BIT_OFFSET_STATE)
 #define HEADER_SETSTATE(X, NewValue)				WRITE_BITS(X, HEADER_BIT_COUNT_STATE, HEADER_BIT_OFFSET_STATE, NewValue)
 
+///////////////////////////////////////////////////////////////////////////////
+// BSDF State
+ 
 // Bit count for each component stored in a Strata BSDF state
 #define STATE_BIT_COUNT_BSDF					3
 #define STATE_BIT_COUNT_SHAREDLOCALBASESID		2
@@ -166,6 +186,11 @@ struct FStrataBSDF
 #define STATE_BIT_OFFSET_ISTHIN			(STATE_BIT_OFFSET_HASFUZZ			+ STATE_BIT_COUNT_HASFUZZ)
 #define STATE_BIT_OFFSET_WEIGHT10F		22
 
+// Sanity check
+#if (STATE_BIT_OFFSET_WEIGHT10F + STATE_BIT_COUNT_WEIGHT10F) > 32
+#error Strata BSDF state is larger than 32 bits
+#endif
+
 // Accessors for the BSDF type
 #define BSDF_GETTYPE(X)					READ_BITS(X.State,  STATE_BIT_COUNT_BSDF, STATE_BIT_OFFSET_BSDF)
 #define BSDF_SETTYPE(X, NewType)		WRITE_BITS(X.State, STATE_BIT_COUNT_BSDF, STATE_BIT_OFFSET_BSDF, NewType)
@@ -244,13 +269,14 @@ struct FStrataBSDF
 #define BSDF_GETTHICKNESSCM(X)			STRATA_SIMPLEVOLUME_THICKNESS_CM
 #endif
 
+///////////////////////////////////////////////////////////////////////////////
+// Slab BSDF members
 #define SLAB_COMPLEX_PATH_MASK (TO_OFFSET_BITMASK(STATE_BIT_COUNT_ANISOTROPY, STATE_BIT_OFFSET_ANISOTROPY)		|  TO_OFFSET_BITMASK(STATE_BIT_COUNT_HASSSS, STATE_BIT_OFFSET_HASSSS) | \
 								TO_OFFSET_BITMASK(STATE_BIT_COUNT_HASSSSPROFILE, STATE_BIT_OFFSET_HASSSSPROFILE)|  TO_OFFSET_BITMASK(STATE_BIT_COUNT_HAZINESS, STATE_BIT_OFFSET_HAZINESS) | \
 								TO_OFFSET_BITMASK(STATE_BIT_COUNT_THINFILM, STATE_BIT_OFFSET_THINFILM)			|  TO_OFFSET_BITMASK(STATE_BIT_COUNT_F90, STATE_BIT_OFFSET_F90) | \
 								TO_OFFSET_BITMASK(STATE_BIT_COUNT_SIMPLEVOLUME, STATE_BIT_OFFSET_SIMPLEVOLUME)	|  TO_OFFSET_BITMASK(STATE_BIT_COUNT_HASFUZZ, STATE_BIT_OFFSET_HASFUZZ) | \
 								TO_OFFSET_BITMASK(STATE_BIT_COUNT_ISTHIN, STATE_BIT_OFFSET_ISTHIN) )
 
-
 #define SLAB_DIFFUSEALBEDO(X)			X.VGPRs[0].xyz
 #define SLAB_ROUGHNESS(X)				X.VGPRs[0].w
 #define SLAB_F0(X)						X.VGPRs[1].xyz
@@ -270,12 +296,18 @@ struct FStrataBSDF
 #define SLAB_FUZZ_COLOR(X)				X.VGPRs[4].xyz
 #define SLAB_FUZZ_AMOUNT(X)				X.VGPRs[4].w
 
+///////////////////////////////////////////////////////////////////////////////
+// Volumetric BSDF members
 #define VOLUMETRICFOGCLOUD_ALBEDO(X)	X.VGPRs[0].xyz
 #define VOLUMETRICFOGCLOUD_EXTINCTION(X)X.VGPRs[1].xyz
 #define VOLUMETRICFOGCLOUD_AO(X)		X.VGPRs[0].w
 
+///////////////////////////////////////////////////////////////////////////////
+// Unlit BSDF members
 #define UNLIT_TRANSMITTANCE(X)			X.VGPRs[0].xyz
 
+///////////////////////////////////////////////////////////////////////////////
+// Hair BSDF members
 #define HAIR_BASECOLOR(X)				X.VGPRs[0].xyz
 #define HAIR_SCATTER(X)					X.VGPRs[0].w
 #define HAIR_ROUGHNESS(X)				X.VGPRs[1].x
@@ -283,6 +315,8 @@ struct FStrataBSDF
 #define HAIR_BACKLIT(X)					X.VGPRs[1].z
 #define HAIR_COMPLEXTRANSMITTANCE(X)	X.VGPRs[1].w
 
+///////////////////////////////////////////////////////////////////////////////
+// Single Layer Water BSDF members
 #define SLW_BASECOLOR(X)				X.VGPRs[0].xyz
 #define SLW_METALLIC(X)					X.VGPRs[1].x
 #define SLW_SPECULAR(X)					X.VGPRs[1].y
@@ -299,10 +333,10 @@ struct FStrataBSDF
 // Forward declarations
 
 FParticipatingMedia StrataSlabCreateParticipatingMedia(float3 DiffuseColor, float3 MeanFreePathCentimeters);
-void EnableSlabBSDFSimpleVolumetric(inout FStrataBSDF BSDF);
-void StrataRequestSharedLocalBasisTangent(inout uint Types, uint Index);
+void   EnableSlabBSDFSimpleVolumetric(inout FStrataBSDF BSDF);
+void   StrataRequestSharedLocalBasisTangent(inout uint Types, uint Index);
 float3 StrataUnpackNormal(uint PackedNormal);
-void StrataUnpackNormalAndTangent(inout float3 Normal, inout float3 Tangent, in uint InPacked);
+void   StrataUnpackNormalAndTangent(inout float3 Normal, inout float3 Tangent, in uint InPacked);
 uint   PackRGBA8(float4 rgba);
 float4 UnpackRGBA8(uint rgba);
 uint   PackR11G11B10F(float3 rgb);
@@ -311,7 +345,7 @@ uint   PackColorLinearToGamma2AlphaLinear(float4 rgba);
 float4 UnpackColorGamma2ToLinearAlphaLinear(uint rgba);
 uint   StrataPackNormal(in float3 Normal);
 float3 StrataUnpackNormal(uint PackedNormal);
-bool IsStrataSlabFastPathCompatible(in FStrataBSDF BSDF);
+bool   IsStrataSlabFastPathCompatible(in FStrataBSDF BSDF);
 
 ///////////////////////////////////////////////////////////////////////////////
 // Classification 
@@ -432,6 +466,8 @@ FStrataSubsurfaceData StrataUnpackSSSData(uint2 In)
 	return Out;
 }
 
+///////////////////////////////////////////////////////////////////////////////
+// Shading models
 uint StrataShadingModelBit(uint ShadingModel)
 {
 	return 1u << ShadingModel;
@@ -890,12 +926,6 @@ struct FStrataPixelHeader
 	uint State;
 };
 
-#define HEADER_MASK_HASPRECSHADOWMASK			(1u << 0)
-#define HEADER_MASK_ZEROPRECSHADOWMASK			(1u << 1)
-#define HEADER_MASK_ISSINGLELAYERWATER			(1u << 2)	// Used to avoid loading BSDF or Categorisation data when generating tiles containing SLW material. Avoid running another categorisation.
-#define HEADER_MASK_CASTCONTACTSHADOW			(1u << 3)
-#define HEADER_MASK_HASDYNINDIRECTSHADOWCASTER	(1u << 4)
-
 bool HasPrecShadowMask(const FStrataPixelHeader In)											{ return (In.State & HEADER_MASK_HASPRECSHADOWMASK) != 0; }
 bool HasZeroPrecShadowMask(const FStrataPixelHeader In)										{ return (In.State & HEADER_MASK_ZEROPRECSHADOWMASK) != 0; }
 bool IsSingleLayerWater(const FStrataPixelHeader In)										{ return (In.State & HEADER_MASK_ISSINGLELAYERWATER) != 0; }
@@ -1621,6 +1651,16 @@ float4 UnpackColorGamma2ToLinearAlphaLinear(uint rgba)
 	return float4(Gamma2 * Gamma2, a * (1.0f / 255.0f));	// Simple gamma 2.0
 }
 
+uint PackR8(float Value)
+{
+	return uint(saturate(Value) * 255.0f);
+}
+
+float UnpackR8(uint Value)
+{
+	return float(Value & 0xFF) * (1.0f / 255.0f);
+}
+
 uint PackRGBA8(float4 rgba)
 {
 	uint r = (uint(saturate(rgba.r) * 255.0f) << 0);
@@ -2206,7 +2246,7 @@ uint2 PackStrataHeader(uint InBSDFCount, FStrataPixelHeader InHeader)
 	HEADER_SETSHAREDLOCALBASESTYPE(Out.x, InHeader.SharedLocalBases.Types);
 #endif
 
-	HEADER_SETAO(Out.x, uint(InHeader.MaterialAO * 255.0f));
+	HEADER_SETAO(Out.x, PackR8(InHeader.MaterialAO));
 	HEADER_SETSTATE(Out.x, InHeader.State);
 
 #if ALLOW_STATIC_LIGHTING
@@ -2229,7 +2269,7 @@ FStrataPixelHeader UnpackStrataHeaderIn(ByteAddressBuffer StrataBuffer, inout FS
 	uint SharedLocalBasesCount = HEADER_GETSHAREDLOCALBASESCOUNT(PackedHeader);
 	uint SharedLocalBasesTypes = HEADER_GETSHAREDLOCALBASESTYPE(PackedHeader);
 
-	Out.MaterialAO = float(HEADER_GETAO(PackedHeader)) * (1.0f / 255.0f);
+	Out.MaterialAO = UnpackR8(HEADER_GETAO(PackedHeader));
 	Out.State = HEADER_GETSTATE(PackedHeader);
 
 	// Pre-shadow mask
