Fixes regressions in highres FFT

diff --git a/Engine/Shaders/Private/GPUFastFourierTransform.usf b/Engine/Shaders/Private/GPUFastFourierTransform.usf
--- a/Engine/Shaders/Private/GPUFastFourierTransform.usf
+++ b/Engine/Shaders/Private/GPUFastFourierTransform.usf
@@ -1382,7 +1382,7 @@ void ComplexMultiplyImagesCS(uint3 GroupID : SV_GroupID, uint3 GroupThreadID : S
 		// redSum = Norm[0].x;  greenSum = Norm[0].y; blueSum = Norm[1].x; alphaSum = Norm[1].y
 
 		// We normalize each RGBA channel independently.
-		InvNorm = float4(((bIsGAGroup) ? Norm[0].y : Norm[0].x).xx, ((bIsGAGroup) ? Norm[1].y : Norm[1].x).xx);
+		InvNorm = float4(rcp((bIsGAGroup) ? Norm[0].y : Norm[0].x).xx, rcp((bIsGAGroup) ? Norm[1].y : Norm[1].x).xx);
 	}
 
 
diff --git a/Engine/Shaders/Private/GPUFastFourierTransform2DCore.ush b/Engine/Shaders/Private/GPUFastFourierTransform2DCore.ush
--- a/Engine/Shaders/Private/GPUFastFourierTransform2DCore.ush
+++ b/Engine/Shaders/Private/GPUFastFourierTransform2DCore.ush
@@ -131,6 +131,11 @@ void CopyDataSrcWindowToLocal(inout Complex LocalBuffer[2][RADIX], bool bIsHoriz
 				LocalBuffer[0][ i ] = SrcValue.xy;
 				LocalBuffer[1][ i ] = SrcValue.zw;
 			}
+			else
+			{
+				LocalBuffer[0][ i ] = 0.0;
+				LocalBuffer[1][ i ] = 0.0;
+			}
 
 		}
 	}
@@ -148,6 +153,11 @@ void CopyDataSrcWindowToLocal(inout Complex LocalBuffer[2][RADIX], bool bIsHoriz
 				LocalBuffer[0][ i ] = SrcValue.xy;
 				LocalBuffer[1][ i ] = SrcValue.zw;
 			}
+			else
+			{
+				LocalBuffer[0][ i ] = 0.0;
+				LocalBuffer[1][ i ] = 0.0;
+			}
 		}
 	}
 }
@@ -173,7 +183,12 @@ void CopyDataSrcWindowToLocal(inout Complex LocalBuffer[2][RADIX], bool bIsHoriz
 				float4 SrcValue = SrcTexture[Pixel];
 				LocalBuffer[0][ i ] = SrcValue.xy;
 				LocalBuffer[1][ i ] = SrcValue.zw;
-			}		
+			}	
+			else
+			{
+				LocalBuffer[0][ i ] = 0.0;
+				LocalBuffer[1][ i ] = 0.0;
+			}	
 		}
 	}
 	else 
@@ -189,6 +204,11 @@ void CopyDataSrcWindowToLocal(inout Complex LocalBuffer[2][RADIX], bool bIsHoriz
 				LocalBuffer[0][ i ] = SrcValue.xy;
 				LocalBuffer[1][ i ] = SrcValue.zw;
 			}
+			else
+			{
+				LocalBuffer[0][ i ] = 0.0;
+				LocalBuffer[1][ i ] = 0.0;
+			}
 		}
 	}
 }
diff --git a/Engine/Source/Runtime/Renderer/Private/GPUFastFourierTransform.cpp b/Engine/Source/Runtime/Renderer/Private/GPUFastFourierTransform.cpp
--- a/Engine/Source/Runtime/Renderer/Private/GPUFastFourierTransform.cpp
+++ b/Engine/Source/Runtime/Renderer/Private/GPUFastFourierTransform.cpp
@@ -1023,7 +1023,7 @@ void GPUFFT::ComplexFFTImage1D::MultiPass(
 				FIntPoint SrcBufferSize = PrevSpectralTexture->Desc.Extent;
 				checkf(!(SrcBufferSize.X < DstExtent.X) && !(SrcBufferSize.Y < DstExtent.Y), TEXT("FFT: Allocated Buffers too small."));
 
-				OpDstTexture = GraphBuilder.CreateTexture(DstTexture->Desc, TEXT("FFT.Spectral"));
+				OpDstTexture = GraphBuilder.CreateTexture(SpectralDesc, TEXT("FFT.Spectral"));
 			}
 			else
 			{
