Strata - move back PrecShadow back to render target for upcoming optimisation we have to do

diff --git a/Engine/Shaders/Private/BasePassPixelShader.usf b/Engine/Shaders/Private/BasePassPixelShader.usf
--- a/Engine/Shaders/Private/BasePassPixelShader.usf
+++ b/Engine/Shaders/Private/BasePassPixelShader.usf
@@ -1095,7 +1095,6 @@ void FPixelShaderInOut_MainPS(
 		if (!bAllZero && !bAllOne)
 		{
 			SetHasPrecShadowMask(StrataPixelHeader, true);
-			StrataPixelHeader.PrecomputedShadowMask = GBuffer.PrecomputedShadowFactors;
 		}
 		else if (bAllZero)
 		{
@@ -1611,8 +1610,12 @@ void FPixelShaderInOut_MainPS(
 	// STRATA_TODO GBUFFER_HAS_DIFFUSE_SAMPLE_OCCLUSION and ALLOW_STATIC_LIGHTING cases
 	StrataPixelHeader.MaterialAO = GBuffer.GBufferAO;
 
-	// Clear only what should be and keep the rest
-	GBuffer = (FGBufferData)0;						// We should not rely on gbuffer for anything else to make any pass work.
+	// We only rely on this GBuffer structure for the write out to legacy render target such as velocity or precomputed shadow factors
+	const float4 PrecomputedShadowFactors = GBuffer.PrecomputedShadowFactors;
+	const float4 EncodedVelocity = GBuffer.Velocity;
+	GBuffer = (FGBufferData)0;
+	GBuffer.PrecomputedShadowFactors = PrecomputedShadowFactors;
+	GBuffer.Velocity = EncodedVelocity;
 
 	uint2 PixelPos = uint2(In.SvPosition.xy);
 
@@ -1801,7 +1804,7 @@ void FPixelShaderInOut_MainPS(
 
 			EncodeGBufferToMRT(Out, GBuffer, QuantizationBias);
 
-			if (GBuffer.ShadingModelID == SHADINGMODELID_UNLIT)
+			if (GBuffer.ShadingModelID == SHADINGMODELID_UNLIT && !STRATA_ENABLED) // Do not touch what strata outputs
 			{
 				Out.MRT[1] = 0;
 				SetGBufferForUnlit(Out.MRT[2]);
diff --git a/Engine/Shaders/Private/DeferredLightPixelShaders.usf b/Engine/Shaders/Private/DeferredLightPixelShaders.usf
--- a/Engine/Shaders/Private/DeferredLightPixelShaders.usf
+++ b/Engine/Shaders/Private/DeferredLightPixelShaders.usf
@@ -166,6 +166,26 @@ float3 GetStableLargeWorldPosition(float2 SVPos, float DeviceZ)
 
 ////////////////////////////////////////////////////////////////////////////////////////////////////
 
+#if STRATA_ENABLED
+
+float4 StrataReadPrecomputedShadowFactors(FStrataPixelHeader StrataPixelHeader, int2 PixelPos)
+{
+	if (HasPrecShadowMask(StrataPixelHeader))
+	{
+#if ALLOW_STATIC_LIGHTING
+		float4 GBufferE = SceneTexturesStruct.GBufferETexture.Load(int3(PixelPos, 0));
+#else
+		float4 GBufferE = 1;
+#endif
+		return GBufferE;
+	}
+	return HasZeroPrecShadowMask(StrataPixelHeader) ? 0.0f : 1.0f;
+}
+
+#endif // STRATA_ENABLED
+
+////////////////////////////////////////////////////////////////////////////////////////////////////
+
 #if USE_HAIR_LIGHTING == 0
 
 void DeferredLightPixelMain(
@@ -232,7 +252,8 @@ void DeferredLightPixelMain(
 			float Dither = InterleavedGradientNoise(InputParams.PixelPos, View.StateFrameIndexMod8);
 			const uint FakeShadingModelID = 0;
 			const float FakeContactShadowOpacity = 1.0f;
-			GetShadowTerms(SceneDepth, StrataPixelHeader.PrecomputedShadowMask, FakeShadingModelID, FakeContactShadowOpacity,
+			float4 PrecomputedShadowFactors = StrataReadPrecomputedShadowFactors(StrataPixelHeader, PixelPos);
+			GetShadowTerms(SceneDepth, PrecomputedShadowFactors, FakeShadingModelID, FakeContactShadowOpacity,
 				LightData, DerivedParams.TranslatedWorldPosition, L, LightAttenuation, Dither, ShadowTerms);
 
 			Radiance = StrataDeferredLighting(
diff --git a/Engine/Shaders/Private/Strata/Strata.ush b/Engine/Shaders/Private/Strata/Strata.ush
--- a/Engine/Shaders/Private/Strata/Strata.ush
+++ b/Engine/Shaders/Private/Strata/Strata.ush
@@ -921,7 +921,6 @@ struct FStrataPixelHeader
 #endif
 
 	float MaterialAO;
-	float4 PrecomputedShadowMask;
 
 	uint State;
 };
@@ -951,7 +950,6 @@ FStrataPixelHeader InitialiseStrataPixelHeader()
 	Out.SharedLocalBasesTypes = 0;
 #endif
 	Out.MaterialAO = 1.0f;
-	Out.PrecomputedShadowMask = 1.0f;
 	Out.State = 0;
 	return Out;
 }
@@ -2238,23 +2236,17 @@ void PackStrataOut(
 // Pack only strata header
 uint2 PackStrataHeader(uint InBSDFCount, FStrataPixelHeader InHeader)
 {
-	uint2 Out = 0;
-	HEADER_SETBSDFCOUNT(Out.x, InBSDFCount);
+	uint Out = 0;
+	HEADER_SETBSDFCOUNT(Out, InBSDFCount);
 
 #if STRATA_INLINE_SHADING
-	HEADER_SETSHAREDLOCALBASESCOUNT(Out.x, InHeader.SharedLocalBases.Count);
-	HEADER_SETSHAREDLOCALBASESTYPE(Out.x, InHeader.SharedLocalBases.Types);
+	HEADER_SETSHAREDLOCALBASESCOUNT(Out, InHeader.SharedLocalBases.Count);
+	HEADER_SETSHAREDLOCALBASESTYPE(Out, InHeader.SharedLocalBases.Types);
 #endif
 
-	HEADER_SETAO(Out.x, PackR8(InHeader.MaterialAO));
-	HEADER_SETSTATE(Out.x, InHeader.State);
+	HEADER_SETAO(Out, PackR8(InHeader.MaterialAO));
+	HEADER_SETSTATE(Out, InHeader.State);
 
-#if ALLOW_STATIC_LIGHTING
-	if (HasPrecShadowMask(InHeader))
-	{
-		Out.y = PackRGBA8(InHeader.PrecomputedShadowMask);
-	}
-#endif
 	return Out;
 }
 
@@ -2272,15 +2264,6 @@ FStrataPixelHeader UnpackStrataHeaderIn(ByteAddressBuffer StrataBuffer, inout FS
 	Out.MaterialAO = UnpackR8(HEADER_GETAO(PackedHeader));
 	Out.State = HEADER_GETSTATE(PackedHeader);
 
-	// Pre-shadow mask
-#if ALLOW_STATIC_LIGHTING
-	Out.PrecomputedShadowMask = HasZeroPrecShadowMask(Out) ? 0.f : 1.f;
-	if (HasPrecShadowMask(Out))
-	{
-		Out.PrecomputedShadowMask = UnpackRGBA8(StrataLoadUint1(StrataBuffer, StrataAddressing));
-	}
-#endif
-
 #if STRATA_INLINE_SHADING==0
 
 	Out.StrataBuffer = StrataBuffer;
diff --git a/Engine/Source/Runtime/Engine/Private/Materials/StrataMaterial.cpp b/Engine/Source/Runtime/Engine/Private/Materials/StrataMaterial.cpp
--- a/Engine/Source/Runtime/Engine/Private/Materials/StrataMaterial.cpp
+++ b/Engine/Source/Runtime/Engine/Private/Materials/StrataMaterial.cpp
@@ -333,8 +333,6 @@ FStrataMaterialAnalysisResult::FStrataMaterialAnalysisResult()
 FStrataMaterialAnalysisResult StrataCompilationInfoMaterialAnalysis(FMaterialCompiler* Compiler, const FStrataMaterialCompilationInfo& Material, const uint32 StrataBytePerPixel)
 {
 	const uint32 UintByteSize = sizeof(uint32);
-	static const auto CVar = IConsoleManager::Get().FindTConsoleVariableDataInt(TEXT("r.AllowStaticLighting"));
-	const bool bAllowStaticLighting = CVar->GetValueOnAnyThread() > 0;
 
 	FStrataMaterialAnalysisResult Result;
 
@@ -343,11 +341,6 @@ FStrataMaterialAnalysisResult StrataCompilationInfoMaterialAnalysis(FMaterialCom
 	// SharedLocalBases_BSDFCount
 	Result.RequestedByteCount += UintByteSize;
 
-	// Pre-calc shadow mask
-	if (bAllowStaticLighting)
-	{
-		Result.RequestedByteCount += UintByteSize;
-	}
 	// shared local bases between BSDFs
 	Result.RequestedByteCount += Compiler->StrataCompilationInfoGetSharedLocalBasesCount() * STRATA_PACKED_SHAREDLOCALBASIS_STRIDE_BYTES;
 	
diff --git a/Engine/Source/Runtime/Engine/Private/ShaderCompiler/ShaderGenerationUtil.cpp b/Engine/Source/Runtime/Engine/Private/ShaderCompiler/ShaderGenerationUtil.cpp
--- a/Engine/Source/Runtime/Engine/Private/ShaderCompiler/ShaderGenerationUtil.cpp
+++ b/Engine/Source/Runtime/Engine/Private/ShaderCompiler/ShaderGenerationUtil.cpp
@@ -1681,7 +1681,7 @@ static void SetStandardGBufferSlots(bool Slots[], bool bWriteEmissive, bool bHas
 {
 	Slots[GBS_SceneColor] = bWriteEmissive;
 	Slots[GBS_Velocity] = bHasVelocity;
-	Slots[GBS_PrecomputedShadowFactor] = bIsStrataMaterial ? false : bHasStaticLighting;
+	Slots[GBS_PrecomputedShadowFactor] = bHasStaticLighting;
 
 	Slots[GBS_WorldNormal] =			bIsStrataMaterial ? false : true;
 	Slots[GBS_PerObjectGBufferData] =	bIsStrataMaterial ? false : true;
diff --git a/Engine/Source/Runtime/RenderCore/Private/GBufferInfo.cpp b/Engine/Source/Runtime/RenderCore/Private/GBufferInfo.cpp
--- a/Engine/Source/Runtime/RenderCore/Private/GBufferInfo.cpp
+++ b/Engine/Source/Runtime/RenderCore/Private/GBufferInfo.cpp
@@ -246,6 +246,10 @@ FGBufferInfo RENDERCORE_API FetchLegacyGBufferInfo(const FGBufferParams& Params)
 		TargetGBufferC = -1;
 
 		Info.NumTargets = 1;
+		if (Params.bHasPrecShadowFactor)
+		{
+			TargetGBufferE = Info.NumTargets++;
+		}
 		if (Params.bHasVelocity)
 		{
 			TargetVelocity = Info.NumTargets++;
@@ -273,6 +277,17 @@ FGBufferInfo RENDERCORE_API FetchLegacyGBufferInfo(const FGBufferParams& Params)
 			}
 		}
 
+		if (Params.bHasPrecShadowFactor)
+		{
+			Info.Targets[TargetGBufferE].Init(GBT_Unorm_8_8_8_8, TEXT("GBufferE"), false, true, true, false); // Precalc
+			Info.Slots[GBS_PrecomputedShadowFactor] = FGBufferItem(GBS_PrecomputedShadowFactor, GBC_Raw_Unorm_8_8_8_8, GBCH_Both);
+			Info.Slots[GBS_PrecomputedShadowFactor].Packing[0] = FGBufferPacking(TargetGBufferE, 0, 0);
+			Info.Slots[GBS_PrecomputedShadowFactor].Packing[1] = FGBufferPacking(TargetGBufferE, 1, 1);
+			Info.Slots[GBS_PrecomputedShadowFactor].Packing[2] = FGBufferPacking(TargetGBufferE, 2, 2);
+			Info.Slots[GBS_PrecomputedShadowFactor].Packing[3] = FGBufferPacking(TargetGBufferE, 3, 3);
+		}
+
+
 		return Info;
 	}
 
