Remove strata classification texture.

* All state into are now stored into the header
* Toplayer data now contains both top layer normal & top layer roughness

diff --git a/Engine/Shaders/Private/BasePassPixelShader.usf b/Engine/Shaders/Private/BasePassPixelShader.usf
--- a/Engine/Shaders/Private/BasePassPixelShader.usf
+++ b/Engine/Shaders/Private/BasePassPixelShader.usf
@@ -1620,7 +1620,6 @@ void FPixelShaderInOut_MainPS(
 	uint2 PixelPos = uint2(In.SvPosition.xy);
 
 	// classificaiton & pre-preprocessed data
-	FStrataClassification ClassificationData = (FStrataClassification)0;
 	FStrataSubsurfaceData SSSData = (FStrataSubsurfaceData)0;
 	FStrataTopLayerData TopLayerData = (FStrataTopLayerData)0;
 
@@ -1630,11 +1629,10 @@ void FPixelShaderInOut_MainPS(
 		Dither,
 		GetStrataPixelDataByteOffset(PixelPos, uint2(ResolvedView.BufferSizeAndInvSize.xy), OpaqueBasePass.Strata.MaxBytesPerPixel),
 		StrataPixelHeader, StrataData, MaterialParameters.CameraVector, bStrataSubsurfaceEnable, EmissiveLuminance,
-		ClassificationData, SSSData, TopLayerData);
+		SSSData, TopLayerData);
 
 	// Write out classificaiton & pre-preprocessed data
-	OpaqueBasePass.Strata.ClassificationTextureUAV[PixelPos] = StrataPackClassificationData(ClassificationData);
-	OpaqueBasePass.Strata.TopLayerNormalTextureUAV[PixelPos] = StrataPackTopLayerData(TopLayerData);
+	OpaqueBasePass.Strata.TopLayerTextureUAV[PixelPos] = StrataPackTopLayerData(TopLayerData);
 
 	// Only write SSS datta if needed
 	BRANCH
diff --git a/Engine/Shaders/Private/DebugProbes.usf b/Engine/Shaders/Private/DebugProbes.usf
--- a/Engine/Shaders/Private/DebugProbes.usf
+++ b/Engine/Shaders/Private/DebugProbes.usf
@@ -273,7 +273,6 @@ void MainPS(
 		StrataPixelHeader.MaterialAO = 1.0f - IndirectIrradianceAndAO;
 		StrataPixelHeader.BSDFCount = 1;
 
-		FStrataClassification ClassificationData = (FStrataClassification)0;
 		FStrataSubsurfaceData SSSData = (FStrataSubsurfaceData)0;
 		FStrataTopLayerData TopLayerData = (FStrataTopLayerData)0;
 
@@ -284,7 +283,7 @@ void MainPS(
 			MaterialLobesBufferUAV, Dither,
 			GetStrataPixelDataByteOffset(PixelCoord, uint2(ResolvedView.BufferSizeAndInvSize.xy), MaxBytesPerPixel),
 			StrataPixelHeader, StrataData, -RayDir, bStrataSubsurfaceEnable, EmissiveLuminance,
-			ClassificationData, SSSData, TopLayerData);
+			SSSData, TopLayerData);
 	}
 	
 #else // PROJECT_STRATA
diff --git a/Engine/Shaders/Private/DistanceFieldScreenGridLighting.usf b/Engine/Shaders/Private/DistanceFieldScreenGridLighting.usf
--- a/Engine/Shaders/Private/DistanceFieldScreenGridLighting.usf
+++ b/Engine/Shaders/Private/DistanceFieldScreenGridLighting.usf
@@ -31,7 +31,7 @@ void ComputeDistanceFieldNormalPS(
 	float SceneDepth = CalcSceneDepth(ScreenUV);
 
 #if STRATA_ENABLED
-	float3 WorldNormal = StrataUnpackTopLayerData(Strata.TopLayerNormalTexture.Load(uint3(PixelCoord, 0))).WorldNormal;
+	float3 WorldNormal = StrataUnpackTopLayerData(Strata.TopLayerTexture.Load(uint3(PixelCoord, 0))).WorldNormal;
 #else
 	float3 WorldNormal = GetGBufferData(ScreenUV).WorldNormal.xyz;
 #endif
@@ -53,7 +53,7 @@ void ComputeDistanceFieldNormalCS(
 	float SceneDepth = CalcSceneDepth(ScreenUV);
 
 #if STRATA_ENABLED
-	float3 WorldNormal = StrataUnpackTopLayerData(Strata.TopLayerNormalTexture.Load(uint3(PixelCoord, 0))).WorldNormal;
+	float3 WorldNormal = StrataUnpackTopLayerData(Strata.TopLayerTexture.Load(uint3(PixelCoord, 0))).WorldNormal;
 #else
 	float3 WorldNormal = GetGBufferData(ScreenUV).WorldNormal.xyz;
 #endif
diff --git a/Engine/Shaders/Private/PostProcessAmbientOcclusion.usf b/Engine/Shaders/Private/PostProcessAmbientOcclusion.usf
--- a/Engine/Shaders/Private/PostProcessAmbientOcclusion.usf
+++ b/Engine/Shaders/Private/PostProcessAmbientOcclusion.usf
@@ -180,7 +180,7 @@ void MainSetupPS(in noperspective float4 UVAndScreenPos : TEXCOORD0, float4 SvPo
 		Samples[i].rgb = normalize(ReconstructNormalFromDepthBuffer(float4(UV[i] * ViewPortSize, SvPosition.zw))) * 0.5f + 0.5f;
 #else
 #if STRATA_ENABLED
-		const FStrataTopLayerData TopLayerData = StrataUnpackTopLayerData(Strata.TopLayerNormalTexture.Load(uint3(clamp(UV[i] * View.BufferSizeAndInvSize.xy, 0.0f, View.BufferSizeAndInvSize.xy - 1.0f), 0)));
+		const FStrataTopLayerData TopLayerData = StrataUnpackTopLayerData(Strata.TopLayerTexture.Load(uint3(clamp(UV[i] * View.BufferSizeAndInvSize.xy, 0.0f, View.BufferSizeAndInvSize.xy - 1.0f), 0)));
 		Samples[i].rgb = TopLayerData.WorldNormal * 0.5f + 0.5f;
 #else
 		Samples[i].rgb = GetGBufferData(UV[i], true).WorldNormal * 0.5f + 0.5f;
@@ -607,7 +607,7 @@ float3 GetNormal(float2 UV, float3 ViewSpacePosMid)
 #if USE_NORMALBUFFER
 
 #if STRATA_ENABLED
-	const FStrataTopLayerData TopLayerData = StrataUnpackTopLayerData(Strata.TopLayerNormalTexture.Load(uint3(clamp(UV * View.BufferSizeAndInvSize.xy, 0.0f, View.BufferSizeAndInvSize.xy-1.0f), 0)));
+	const FStrataTopLayerData TopLayerData = StrataUnpackTopLayerData(Strata.TopLayerTexture.Load(uint3(clamp(UV * View.BufferSizeAndInvSize.xy, 0.0f, View.BufferSizeAndInvSize.xy-1.0f), 0)));
 	float3 WorldNormal = TopLayerData.WorldNormal;
 #else
 	// Get the normal from the normal buffer
@@ -985,7 +985,7 @@ float GTAOInnerIntegratePSandCS(in float2 UV, in uint2 iPos)
 	// ViewspacePos and Normal
 	float3 ViewSpacePos		= ScreenToViewPos(UV, SceneDepth);
 #if STRATA_ENABLED
-	const FStrataTopLayerData TopLayerData = StrataUnpackTopLayerData(Strata.TopLayerNormalTexture.Load(uint3(iPos, 0)));
+	const FStrataTopLayerData TopLayerData = StrataUnpackTopLayerData(Strata.TopLayerTexture.Load(uint3(iPos, 0)));
 	float3 WorldNormal = TopLayerData.WorldNormal;
 #else
 	float3 WorldNormal		= GetGBufferData(UV, false).WorldNormal;
diff --git a/Engine/Shaders/Private/PostProcessAmbientOcclusionCommon.ush b/Engine/Shaders/Private/PostProcessAmbientOcclusionCommon.ush
--- a/Engine/Shaders/Private/PostProcessAmbientOcclusionCommon.ush
+++ b/Engine/Shaders/Private/PostProcessAmbientOcclusionCommon.ush
@@ -107,10 +107,11 @@ float3 GetWorldSpaceNormalFromAOInput(float2 UV, float4 SvPosition)
 		WorldNormal = ReconstructNormalFromDepthBuffer(SvPosition);
 	#else
 #if STRATA_ENABLED
-		const FStrataClassification StrataData = StrataUnpackClassificationData(Strata.ClassificationTexture.Load(uint3(SvPosition.xy, 0)));
-		if (!StrataHasShadingModel(StrataData, STRATA_BSDF_TYPE_HAIR))
+		FStrataAddressing StrataAddressing = GetStrataPixelDataByteOffset(SvPosition.xy, uint2(View.BufferSizeAndInvSize.xy), Strata.MaxBytesPerPixel);
+		FStrataPixelHeader StrataHeader = UnpackStrataHeaderIn(Strata.MaterialLobesBuffer, StrataAddressing);
+		if (!StrataHasShadingModel(StrataHeader, STRATA_BSDF_TYPE_HAIR))
 		{
-			const FStrataTopLayerData TopLayerData = StrataUnpackTopLayerData(Strata.TopLayerNormalTexture.Load(uint3(SvPosition.xy, 0)));
+			const FStrataTopLayerData TopLayerData = StrataUnpackTopLayerData(Strata.TopLayerTexture.Load(uint3(SvPosition.xy, 0)));
 			WorldNormal = TopLayerData.WorldNormal;
 		}
 #else
diff --git a/Engine/Shaders/Private/SSRT/SSRTReflections.usf b/Engine/Shaders/Private/SSRT/SSRTReflections.usf
--- a/Engine/Shaders/Private/SSRT/SSRTReflections.usf
+++ b/Engine/Shaders/Private/SSRT/SSRTReflections.usf
@@ -77,10 +77,10 @@ void ScreenSpaceReflectionsStencilPS(
 #if STRATA_ENABLED
 
 	uint2 PixelCoord = SvPosition.xy;
-	const FStrataClassification StrataData = StrataUnpackClassificationData(Strata.ClassificationTexture.Load(uint3(PixelCoord, 0)));
+	const FStrataTopLayerData TopLayerData = StrataUnpackTopLayerData(Strata.TopLayerTexture.Load(uint3(PixelCoord, 0)));
 
-	const float Roughness = StrataData.TopLayerRoughness;
-	const bool bNoMaterial = !IsStrataMaterial(StrataData);
+	const float Roughness = TopLayerData.Roughness;
+	const bool bNoMaterial = !IsStrataMaterial(TopLayerData);
 
 #else // STRATA_ENABLED
 
@@ -154,14 +154,13 @@ void ScreenSpaceReflections(
 	}
 	else
 	{
-		const FStrataClassification StrataData = StrataUnpackClassificationData(Strata.ClassificationTexture.Load(uint3(PixelCoord, 0)));
-		const FStrataTopLayerData TopLayerData = StrataUnpackTopLayerData(Strata.TopLayerNormalTexture.Load(uint3(PixelCoord, 0)));
+		const FStrataTopLayerData TopLayerData = StrataUnpackTopLayerData(Strata.TopLayerTexture.Load(uint3(PixelCoord, 0)));
 
 		// For now, SSR is only used to trace the top material layer.
 		// STRATA_TODO try stochastic selection
 		N = TopLayerData.WorldNormal;
-		Roughness = StrataData.TopLayerRoughness;
-		bNoMaterial = !IsStrataMaterial(StrataData);
+		Roughness = TopLayerData.Roughness;
+		bNoMaterial = !IsStrataMaterial(TopLayerData);
 	}
 
 	float DeviceZ = SampleDeviceZFromSceneTextures(UV);
diff --git a/Engine/Shaders/Private/ShadowProjectionPixelShader.usf b/Engine/Shaders/Private/ShadowProjectionPixelShader.usf
--- a/Engine/Shaders/Private/ShadowProjectionPixelShader.usf
+++ b/Engine/Shaders/Private/ShadowProjectionPixelShader.usf
@@ -211,7 +211,7 @@ void Main(
 			const float3 LightDirection = bIsDirectional ? -LightPositionOrDirection.xyz : normalize(LightPositionOrDirection.xyz - TranslateWorldPosition);
 
 		#if STRATA_ENABLED
-			const FStrataTopLayerData TopLayerData = StrataUnpackTopLayerData(Strata.TopLayerNormalTexture.Load(uint3(SVPos.xy, 0)));
+			const FStrataTopLayerData TopLayerData = StrataUnpackTopLayerData(Strata.TopLayerTexture.Load(uint3(SVPos.xy, 0)));
 			const float3 WorldNormal = TopLayerData.WorldNormal;
 		#else
 			FGBufferData GBufferData = GetGBufferData(ScreenUV);
@@ -255,12 +255,13 @@ void Main(
 		bool bHasSubsurfaceTransmission	= false;
 		bool bIsEyeOrHair 				= false;
 	#if STRATA_ENABLED
-		const FStrataClassification StrataData = StrataUnpackClassificationData(Strata.ClassificationTexture.Load(PixelCoord));
-		if (IsStrataMaterial(StrataData))
+		FStrataAddressing StrataAddressing = GetStrataPixelDataByteOffset(PixelCoord.xy, uint2(View.BufferSizeAndInvSize.xy), Strata.MaxBytesPerPixel);
+		FStrataPixelHeader StrataHeader = UnpackStrataHeaderIn(Strata.MaterialLobesBuffer, StrataAddressing);
+		if (IsStrataMaterial(StrataHeader))
 		{
-			bIsEyeOrHair				= StrataHasShadingModel(StrataData, STRATA_BSDF_TYPE_HAIR);
-			bHasSubsurface				= HasSubsurface(StrataData) || bIsEyeOrHair;
-			bHasSubsurfaceTransmission	= HasSubsurface(StrataData);
+			bIsEyeOrHair				= StrataHasShadingModel(StrataHeader, STRATA_BSDF_TYPE_HAIR);
+			bHasSubsurface				= HasSubsurface(StrataHeader) || bIsEyeOrHair;
+			bHasSubsurfaceTransmission	= HasSubsurface(StrataHeader);
 		}
 	#else // STRATA_ENABLED
 		FGBufferData GBufferData = GetGBufferData(ScreenUV);
@@ -308,12 +309,12 @@ void Main(
 			if (bHasSubsurfaceTransmission)
 			{
 			#if STRATA_ENABLED
-				const FStrataTopLayerData TopLayerData = StrataUnpackTopLayerData(Strata.TopLayerNormalTexture.Load(uint3(SVPos.xy, 0)));
+				const FStrataTopLayerData TopLayerData = StrataUnpackTopLayerData(Strata.TopLayerTexture.Load(uint3(SVPos.xy, 0)));
 				const float3 WorldNormal = TopLayerData.WorldNormal;
 
 				// For non-profile version, use default transmission parameters (only NormalScale & ExtinctionScale are used)
 				FTransmissionProfileParams TransmissionParams = InitTransmissionProfileParams();
-				if (HasSubsurfaceProfile(StrataData))
+				if (HasSubsurfaceProfile(StrataHeader))
 				{
 					TransmissionParams = GetTransmissionProfileParams(SSSData.ProfileId);
 				}
@@ -421,7 +422,7 @@ void MainOnePassPointLightPS(
 #if !SUBPIXEL_SHADOW
 	{
 #if STRATA_ENABLED
-		const FStrataTopLayerData TopLayerData = StrataUnpackTopLayerData(Strata.TopLayerNormalTexture.Load(uint3(SVPos.xy, 0)));
+		const FStrataTopLayerData TopLayerData = StrataUnpackTopLayerData(Strata.TopLayerTexture.Load(uint3(SVPos.xy, 0)));
 		const float3 WorldNormal = TopLayerData.WorldNormal;
 #else // STRATA_ENABLED
 		FGBufferData GBufferData = GetGBufferData(ScreenUV);
@@ -459,11 +460,10 @@ void MainOnePassPointLightPS(
 
 #if STRATA_ENABLED
 
-	const FStrataClassification StrataData = StrataUnpackClassificationData(Strata.ClassificationTexture.Load(uint3(SVPos.xy, 0)));
-	const FStrataTopLayerData TopLayerData = StrataUnpackTopLayerData(Strata.TopLayerNormalTexture.Load(uint3(SVPos.xy, 0)));
+	const FStrataTopLayerData TopLayerData = StrataUnpackTopLayerData(Strata.TopLayerTexture.Load(uint3(SVPos.xy, 0)));
 	const FStrataSubsurfaceData SSSData = StrataUnpackSSSData(Strata.SSSTexture.Load(uint3(SVPos.xy, 0)));
 
-	const bool IsSubsurfaceProfile = HasSubsurfaceProfile(StrataData);
+	const bool IsSubsurfaceProfile = SSSData.bIsValid;
 	const float3 WorldNormal = TopLayerData.WorldNormal;
 	const FTransmissionProfileParams TransmissionParams = GetTransmissionProfileParams(SSSData.ProfileId);
 
diff --git a/Engine/Shaders/Private/Strata/Strata.ush b/Engine/Shaders/Private/Strata/Strata.ush
--- a/Engine/Shaders/Private/Strata/Strata.ush
+++ b/Engine/Shaders/Private/Strata/Strata.ush
@@ -102,42 +102,50 @@ struct FStrataBSDF
 #define READ_BITS(X, BitCount, BitOffset)				uint((X>>(BitOffset)) & TO_BITMASK(BitCount))
 #define WRITE_BITS(X, BitCount, BitOffset, NewValue)	X = (X & (~TO_OFFSET_BITMASK(BitCount, BitOffset))) | (((NewValue) & TO_BITMASK(BitCount)) << (BitOffset))
 
-
 ///////////////////////////////////////////////////////////////////////////////
 // Material Header
 // 
 // Bit count for each component stored in a Strata header
+// 
+// State bit masks
+#define HEADER_MASK_ISSIMPLEMATERIAL			(1u << 0)
+#define HEADER_MASK_HASPRECSHADOWMASK			(1u << 1)
+#define HEADER_MASK_ZEROPRECSHADOWMASK			(1u << 2)
+#define HEADER_MASK_CASTCONTACTSHADOW			(1u << 3)
+#define HEADER_MASK_HASDYNINDIRECTSHADOWCASTER	(1u << 4)
+#define HEADER_MASK_ISSINGLELAYERWATER			(1u << 5)	// Used to avoid loading BSDF or Categorisation data when generating tiles containing SLW material. Avoid running another categorisation.
+#define HEADER_MASK_HASSUBSURFACE				(1u << 6)
+#define HEADER_MASK_HASSUBSURFACEPROFILE		(1u << 7)
+#define HEADER_MASK_BIT_COUNT					8
+
 // For HEADER_BIT_BSDF_COUNT and  HEADER_BIT_SHAREDLOCALBASES_X, please refer to values in StrataDefinitions.h
-#define HEADER_BIT_COUNT_BSDF_COUNT				4
-#define HEADER_BIT_COUNT_SHAREDLOCALBASES_COUNT	3
-#define HEADER_BIT_COUNT_SHAREDLOCALBASES_TYPE	4
-#define HEADER_BIT_COUNT_AO						8
-#define HEADER_BIT_COUNT_STATE					5
-#define HEADER_BIT_COUNT___UNUSED___			8
-
-#define HEADER_BIT_OFFSET_BSDF_COUNT			0
-#define HEADER_BIT_OFFSET_SHAREDLOCALBASES_COUNT (HEADER_BIT_OFFSET_BSDF_COUNT				+ HEADER_BIT_COUNT_BSDF_COUNT)
-#define HEADER_BIT_OFFSET_SHAREDLOCALBASES_TYPE	(HEADER_BIT_OFFSET_SHAREDLOCALBASES_COUNT	+ HEADER_BIT_COUNT_SHAREDLOCALBASES_COUNT)
-#define HEADER_BIT_OFFSET_AO					(HEADER_BIT_OFFSET_SHAREDLOCALBASES_TYPE	+ HEADER_BIT_COUNT_SHAREDLOCALBASES_TYPE)
-#define HEADER_BIT_OFFSET_STATE					(HEADER_BIT_OFFSET_AO						+ HEADER_BIT_COUNT_AO)
+#define HEADER_BIT_COUNT_STATE						HEADER_MASK_BIT_COUNT
+#define HEADER_BIT_COUNT_AO							8
+#define HEADER_BIT_COUNT_BSDF_COUNT					4
+#define HEADER_BIT_COUNT_SHADINGMODELS				4
+#define HEADER_BIT_COUNT_SHAREDLOCALBASES_COUNT		3
+#define HEADER_BIT_COUNT_SHAREDLOCALBASES_TYPE		4
+#define HEADER_BIT_COUNT___UNUSED___				1
+
+#define HEADER_BIT_OFFSET_STATE						0
+#define HEADER_BIT_OFFSET_AO						(HEADER_BIT_OFFSET_STATE					+ HEADER_BIT_COUNT_STATE)
+#define HEADER_BIT_OFFSET_BSDF_COUNT				(HEADER_BIT_OFFSET_AO						+ HEADER_BIT_COUNT_AO)
+#define HEADER_BIT_OFFSET_SHADINGMODELS				(HEADER_BIT_OFFSET_BSDF_COUNT				+ HEADER_BIT_COUNT_BSDF_COUNT)
+#define HEADER_BIT_OFFSET_SHAREDLOCALBASES_COUNT	(HEADER_BIT_OFFSET_SHADINGMODELS			+ HEADER_BIT_COUNT_SHADINGMODELS)
+#define HEADER_BIT_OFFSET_SHAREDLOCALBASES_TYPE		(HEADER_BIT_OFFSET_SHAREDLOCALBASES_COUNT	+ HEADER_BIT_COUNT_SHAREDLOCALBASES_COUNT)
+
+#define HEADER_MASK_SHADINGMODELS					(((1u<<HEADER_BIT_COUNT_SHADINGMODELS)-1) << HEADER_BIT_OFFSET_SHADINGMODELS)
 
 // Sanity check
-#if (HEADER_BIT_COUNT_STATE + HEADER_BIT_OFFSET_STATE) > 32
+#if (HEADER_BIT_OFFSET_SHAREDLOCALBASES_TYPE + HEADER_BIT_COUNT_SHAREDLOCALBASES_TYPE) > 32
 #error Strata Header is larger than 32 bits
 #endif
 
-// State bit masks
-#define HEADER_MASK_HASPRECSHADOWMASK			(1u << 0)
-#define HEADER_MASK_ZEROPRECSHADOWMASK			(1u << 1)
-#define HEADER_MASK_ISSINGLELAYERWATER			(1u << 2)	// Used to avoid loading BSDF or Categorisation data when generating tiles containing SLW material. Avoid running another categorisation.
-#define HEADER_MASK_CASTCONTACTSHADOW			(1u << 3)
-#define HEADER_MASK_HASDYNINDIRECTSHADOWCASTER	(1u << 4)
-
 #define HEADER_GETBSDFCOUNT(X)						READ_BITS( X, HEADER_BIT_COUNT_BSDF_COUNT, HEADER_BIT_OFFSET_BSDF_COUNT)
 #define HEADER_SETBSDFCOUNT(X, NewValue)			WRITE_BITS(X, HEADER_BIT_COUNT_BSDF_COUNT, HEADER_BIT_OFFSET_BSDF_COUNT, NewValue)
 
 #define HEADER_GETSHAREDLOCALBASESCOUNT(X)			READ_BITS( X, HEADER_BIT_COUNT_SHAREDLOCALBASES_COUNT, HEADER_BIT_OFFSET_SHAREDLOCALBASES_COUNT)
-#define HEADER_SETSHAREDLOCALBASESCOUNT(X, NewValue)WRITE_BITS(X, HEADER_BIT_COUNT_SHAREDLOCALBASES_COUNT, HEADER_BIT_OFFSET_SHAREDLOCALBASES_COUNT, NewValue)
+#define HEADER_SETSHAREDLOCALBASESCOUNT(X, NewValue) WRITE_BITS(X, HEADER_BIT_COUNT_SHAREDLOCALBASES_COUNT, HEADER_BIT_OFFSET_SHAREDLOCALBASES_COUNT, NewValue)
 
 #define HEADER_GETSHAREDLOCALBASESTYPE(X)			READ_BITS( X, HEADER_BIT_COUNT_SHAREDLOCALBASES_TYPE, HEADER_BIT_OFFSET_SHAREDLOCALBASES_TYPE)
 #define HEADER_SETSHAREDLOCALBASESTYPE(X, NewValue)	WRITE_BITS(X, HEADER_BIT_COUNT_SHAREDLOCALBASES_TYPE, HEADER_BIT_OFFSET_SHAREDLOCALBASES_TYPE, NewValue)
@@ -148,11 +156,14 @@ struct FStrataBSDF
 #define HEADER_GETSTATE(X)							READ_BITS( X, HEADER_BIT_COUNT_STATE, HEADER_BIT_OFFSET_STATE)
 #define HEADER_SETSTATE(X, NewValue)				WRITE_BITS(X, HEADER_BIT_COUNT_STATE, HEADER_BIT_OFFSET_STATE, NewValue)
 
+#define HEADER_GETSHADINGMODELS(X)					READ_BITS( X, HEADER_BIT_COUNT_SHADINGMODELS, HEADER_BIT_OFFSET_SHADINGMODELS)
+#define HEADER_SETSHADINGMODELS(X, NewValue)		WRITE_BITS(X, HEADER_BIT_COUNT_SHADINGMODELS, HEADER_BIT_OFFSET_SHADINGMODELS, NewValue)
+
 ///////////////////////////////////////////////////////////////////////////////
 // BSDF State
  
 // Bit count for each component stored in a Strata BSDF state
-#define STATE_BIT_COUNT_BSDF					3
+#define STATE_BIT_COUNT_BSDF_TYPE				3
 #define STATE_BIT_COUNT_SHAREDLOCALBASESID		2
 #define STATE_BIT_COUNT_ANISOTROPY				1
 #define STATE_BIT_COUNT_TOPLAYER				1
@@ -170,21 +181,21 @@ struct FStrataBSDF
 #define STATE_BIT_COUNT_WEIGHT10F				10
 
 // Bit offset of each component store in BSDF state 
-#define STATE_BIT_OFFSET_BSDF			0
-#define STATE_BIT_OFFSET_SHAREDLOCALBASESID (STATE_BIT_OFFSET_BSDF			+ STATE_BIT_COUNT_BSDF)
-#define STATE_BIT_OFFSET_ANISOTROPY		(STATE_BIT_OFFSET_SHAREDLOCALBASESID+ STATE_BIT_COUNT_SHAREDLOCALBASESID)
-#define STATE_BIT_OFFSET_TOPLAYER		(STATE_BIT_OFFSET_ANISOTROPY		+ STATE_BIT_COUNT_ANISOTROPY)
-#define STATE_BIT_OFFSET_HASSSS			(STATE_BIT_OFFSET_TOPLAYER			+ STATE_BIT_COUNT_TOPLAYER)
-#define STATE_BIT_OFFSET_HASSSSPROFILE	(STATE_BIT_OFFSET_HASSSS			+ STATE_BIT_COUNT_HASSSS)
-#define STATE_BIT_OFFSET_GREYWEIGHT		(STATE_BIT_OFFSET_HASSSSPROFILE		+ STATE_BIT_COUNT_HASSSSPROFILE)
-#define STATE_BIT_OFFSET_HAZINESS		(STATE_BIT_OFFSET_GREYWEIGHT		+ STATE_BIT_COUNT_GREYWEIGHT)
-#define STATE_BIT_OFFSET_F90			(STATE_BIT_OFFSET_HAZINESS			+ STATE_BIT_COUNT_HAZINESS)
-#define STATE_BIT_OFFSET_THINFILM		(STATE_BIT_OFFSET_F90				+ STATE_BIT_COUNT_F90)
-#define STATE_BIT_OFFSET_SIMPLEVOLUME	(STATE_BIT_OFFSET_THINFILM			+ STATE_BIT_COUNT_THINFILM)
-#define STATE_BIT_OFFSET_DMFPPLUGGED	(STATE_BIT_OFFSET_SIMPLEVOLUME		+ STATE_BIT_COUNT_SIMPLEVOLUME)
-#define STATE_BIT_OFFSET_HASFUZZ		(STATE_BIT_OFFSET_DMFPPLUGGED		+ STATE_BIT_COUNT_DMFPPLUGGED)
-#define STATE_BIT_OFFSET_ISTHIN			(STATE_BIT_OFFSET_HASFUZZ			+ STATE_BIT_COUNT_HASFUZZ)
-#define STATE_BIT_OFFSET_WEIGHT10F		22
+#define STATE_BIT_OFFSET_BSDF_TYPE				0
+#define STATE_BIT_OFFSET_SHAREDLOCALBASESID		(STATE_BIT_OFFSET_BSDF_TYPE			+ STATE_BIT_COUNT_BSDF_TYPE)
+#define STATE_BIT_OFFSET_ANISOTROPY				(STATE_BIT_OFFSET_SHAREDLOCALBASESID+ STATE_BIT_COUNT_SHAREDLOCALBASESID)
+#define STATE_BIT_OFFSET_TOPLAYER				(STATE_BIT_OFFSET_ANISOTROPY		+ STATE_BIT_COUNT_ANISOTROPY)
+#define STATE_BIT_OFFSET_HASSSS					(STATE_BIT_OFFSET_TOPLAYER			+ STATE_BIT_COUNT_TOPLAYER)
+#define STATE_BIT_OFFSET_HASSSSPROFILE			(STATE_BIT_OFFSET_HASSSS			+ STATE_BIT_COUNT_HASSSS)
+#define STATE_BIT_OFFSET_GREYWEIGHT				(STATE_BIT_OFFSET_HASSSSPROFILE		+ STATE_BIT_COUNT_HASSSSPROFILE)
+#define STATE_BIT_OFFSET_HAZINESS				(STATE_BIT_OFFSET_GREYWEIGHT		+ STATE_BIT_COUNT_GREYWEIGHT)
+#define STATE_BIT_OFFSET_F90					(STATE_BIT_OFFSET_HAZINESS			+ STATE_BIT_COUNT_HAZINESS)
+#define STATE_BIT_OFFSET_THINFILM				(STATE_BIT_OFFSET_F90				+ STATE_BIT_COUNT_F90)
+#define STATE_BIT_OFFSET_SIMPLEVOLUME			(STATE_BIT_OFFSET_THINFILM			+ STATE_BIT_COUNT_THINFILM)
+#define STATE_BIT_OFFSET_DMFPPLUGGED			(STATE_BIT_OFFSET_SIMPLEVOLUME		+ STATE_BIT_COUNT_SIMPLEVOLUME)
+#define STATE_BIT_OFFSET_HASFUZZ				(STATE_BIT_OFFSET_DMFPPLUGGED		+ STATE_BIT_COUNT_DMFPPLUGGED)
+#define STATE_BIT_OFFSET_ISTHIN					(STATE_BIT_OFFSET_HASFUZZ			+ STATE_BIT_COUNT_HASFUZZ)
+#define STATE_BIT_OFFSET_WEIGHT10F				22
 
 // Sanity check
 #if (STATE_BIT_OFFSET_WEIGHT10F + STATE_BIT_COUNT_WEIGHT10F) > 32
@@ -192,8 +203,8 @@ struct FStrataBSDF
 #endif
 
 // Accessors for the BSDF type
-#define BSDF_GETTYPE(X)					READ_BITS(X.State,  STATE_BIT_COUNT_BSDF, STATE_BIT_OFFSET_BSDF)
-#define BSDF_SETTYPE(X, NewType)		WRITE_BITS(X.State, STATE_BIT_COUNT_BSDF, STATE_BIT_OFFSET_BSDF, NewType)
+#define BSDF_GETTYPE(X)					READ_BITS(X.State,  STATE_BIT_COUNT_BSDF_TYPE, STATE_BIT_OFFSET_BSDF_TYPE)
+#define BSDF_SETTYPE(X, NewType)		WRITE_BITS(X.State, STATE_BIT_COUNT_BSDF_TYPE, STATE_BIT_OFFSET_BSDF_TYPE, NewType)
 
 // Accessors for shared local bases
 // This assumes Normals[STRATA_MAX_SHAREDLOCALBASES_REGISTERS]
@@ -239,7 +250,7 @@ struct FStrataBSDF
 
 // Indicates if the BSDF is part of the top layer
 #define BSDF_GETISSIMPLEVOLUME(X)		READ_BITS(X.State,  STATE_BIT_COUNT_SIMPLEVOLUME, STATE_BIT_OFFSET_SIMPLEVOLUME)
-#define BSDF_SETISSIMPLEVOLUME(X, Value)WRITE_BITS(X.State, STATE_BIT_COUNT_SIMPLEVOLUME, STATE_BIT_OFFSET_SIMPLEVOLUME, Value)
+#define BSDF_SETISSIMPLEVOLUME(X, Value) WRITE_BITS(X.State, STATE_BIT_COUNT_SIMPLEVOLUME, STATE_BIT_OFFSET_SIMPLEVOLUME, Value)
 
 // Indicates if the BSDF is part of the top layer
 #define BSDF_GETHASDMFP(X)				READ_BITS(X.State,  STATE_BIT_COUNT_DMFPPLUGGED, STATE_BIT_OFFSET_DMFPPLUGGED)
@@ -337,6 +348,8 @@ void   EnableSlabBSDFSimpleVolumetric(inout FStrataBSDF BSDF);
 void   StrataRequestSharedLocalBasisTangent(inout uint Types, uint Index);
 float3 StrataUnpackNormal(uint PackedNormal);
 void   StrataUnpackNormalAndTangent(inout float3 Normal, inout float3 Tangent, in uint InPacked);
+uint   PackR8(float Value);
+float  UnpackR8(uint Value);
 uint   PackRGBA8(float4 rgba);
 float4 UnpackRGBA8(uint rgba);
 uint   PackR11G11B10F(float3 rgb);
@@ -345,71 +358,39 @@ uint   PackColorLinearToGamma2AlphaLinear(float4 rgba);
 float4 UnpackColorGamma2ToLinearAlphaLinear(uint rgba);
 uint   StrataPackNormal(in float3 Normal);
 float3 StrataUnpackNormal(uint PackedNormal);
+uint   StrataPackNormal24(in float3 Normal);
+float3 StrataUnpackNormal24(uint PackedNormal);
 bool   IsStrataSlabFastPathCompatible(in FStrataBSDF BSDF);
 
-///////////////////////////////////////////////////////////////////////////////
-// Classification 
-
-struct FStrataClassification
-{
-	uint State;
-	float TopLayerRoughness;
-};
-
-#define CLAS_MASK_STATE					(0xFF)
-#define CLAS_MASK_SHADINGMODELS			(0x1F)
-#define CLAS_MASK_HASSUBSURFACE			(1u << 5)
-#define CLAS_MASK_HASSUBSURFACEPROFILE	(1u << 6)
-#define CLAS_MASK_ISSIMPLE				(1u << 7)
-
-uint GetShadingModels(const FStrataClassification In)		{ return (In.State & CLAS_MASK_SHADINGMODELS); }
-bool IsStrataMaterial(const FStrataClassification In)		{ return (In.State & CLAS_MASK_SHADINGMODELS) > 0; }
-bool HasSubsurface(const FStrataClassification In)			{ return (In.State & CLAS_MASK_HASSUBSURFACE) != 0; }
-bool HasSubsurfaceProfile(const FStrataClassification In)	{ return (In.State & CLAS_MASK_HASSUBSURFACEPROFILE) != 0; }
-bool IsSimpleMaterial(const FStrataClassification In)		{ return (In.State & CLAS_MASK_ISSIMPLE) != 0; }
-
-void AddShadingModels(inout FStrataClassification Out, uint bIn)		{ Out.State |= (bIn & CLAS_MASK_SHADINGMODELS); }
-void SetHasSubsurface(inout FStrataClassification Out, bool bIn)		{ Out.State |= (bIn ? CLAS_MASK_HASSUBSURFACE : 0u); }
-void SetHasSubsurfaceProfile(inout FStrataClassification Out, bool bIn)	{ Out.State |= (bIn ? CLAS_MASK_HASSUBSURFACEPROFILE : 0u); }
-void SetIsSimpleMaterial(inout FStrataClassification Out, bool bIn)		{ Out.State |= (bIn ? CLAS_MASK_ISSIMPLE : 0u); }
-void AndIsSimpleMaterial(inout FStrataClassification Out, bool bIn)		{ Out.State  = (Out.State & ~CLAS_MASK_ISSIMPLE) | ((Out.State & CLAS_MASK_ISSIMPLE) & (bIn ? CLAS_MASK_ISSIMPLE : 0u)); }
-
-uint StrataPackClassificationData(FStrataClassification In)
-{
-	const uint TopLayerRoughness8Bits = uint(In.TopLayerRoughness * 255.0f);
-	return
-		(In.State & CLAS_MASK_STATE) |
-		((TopLayerRoughness8Bits & 0xFF) << 8);
-}
-
-FStrataClassification StrataUnpackClassificationData(uint In)
-{
-	FStrataClassification Out;
-	Out.State = In & CLAS_MASK_STATE;
-	Out.TopLayerRoughness = float((In >> 8) & 0xFF) * (1.0f / 255.0f);
-	return Out;
-}
-
 ///////////////////////////////////////////////////////////////////////////////
 // Top layer
 struct FStrataTopLayerData
 {
 	float3 WorldNormal;
+	float  Roughness;
+	bool   bIsValid;
 };
 
 uint StrataPackTopLayerData(FStrataTopLayerData In)
 {
-	uint Out = StrataPackNormal(In.WorldNormal);
-	return Out;
+	// Output is min. by 1 to flag data has valid. This creates a small bias for one particular orientation with a mirror roughness.
+	return max((StrataPackNormal24(In.WorldNormal)<<8) | PackR8(In.Roughness), 1u);
 }
 
 FStrataTopLayerData StrataUnpackTopLayerData(uint In)
 {
 	FStrataTopLayerData Out = (FStrataTopLayerData)0;
-	Out.WorldNormal = StrataUnpackNormal(In);
+	Out.bIsValid = In > 0;
+	Out.WorldNormal = StrataUnpackNormal24(In>>8);
+	Out.Roughness = UnpackR8(In);
 	return Out;
 }
 
+bool IsStrataMaterial(FStrataTopLayerData In)
+{
+	return In.bIsValid;
+}
+
 ///////////////////////////////////////////////////////////////////////////////
 // Sub-surface 
 
@@ -478,11 +459,6 @@ bool StrataHasShadingModel(uint ShadingModels, uint ShadingModel)
 	return (ShadingModels & (1u << ShadingModel)) > 0;
 }
 
-bool StrataHasShadingModel(FStrataClassification In, uint ShadingModel)
-{
-	return StrataHasShadingModel(GetShadingModels(In), ShadingModel);
-}
-
 // Return true if a BSDF has a custom area light integrator
 bool StrataHasAreaLightIntegrator(in FStrataBSDF BSDF)
 {
@@ -925,17 +901,30 @@ struct FStrataPixelHeader
 	uint State;
 };
 
+bool IsSimpleMaterial(const FStrataPixelHeader In)											{ return (In.State & HEADER_MASK_ISSIMPLEMATERIAL) != 0; }
 bool HasPrecShadowMask(const FStrataPixelHeader In)											{ return (In.State & HEADER_MASK_HASPRECSHADOWMASK) != 0; }
 bool HasZeroPrecShadowMask(const FStrataPixelHeader In)										{ return (In.State & HEADER_MASK_ZEROPRECSHADOWMASK) != 0; }
 bool IsSingleLayerWater(const FStrataPixelHeader In)										{ return (In.State & HEADER_MASK_ISSINGLELAYERWATER) != 0; }
 bool DoesCastContactShadow(const FStrataPixelHeader In)										{ return (In.State & HEADER_MASK_CASTCONTACTSHADOW) != 0; }
+uint GetShadingModels(const FStrataPixelHeader In)											{ return (In.State & HEADER_MASK_SHADINGMODELS)>>HEADER_BIT_OFFSET_SHADINGMODELS; }
+bool IsStrataMaterial(const FStrataPixelHeader In)											{ return (In.State & HEADER_MASK_SHADINGMODELS) > 0; }
+bool HasSubsurface(const FStrataPixelHeader In)												{ return (In.State & HEADER_MASK_HASSUBSURFACE) != 0; }
+bool HasSubsurfaceProfile(const FStrataPixelHeader In)										{ return (In.State & HEADER_MASK_HASSUBSURFACEPROFILE) != 0; }
 bool HasDynamicIndirectShadowCasterRepresentation(const FStrataPixelHeader In)				{ return (In.State & HEADER_MASK_HASDYNINDIRECTSHADOWCASTER) != 0; }
 
+bool StrataHasShadingModel(const FStrataPixelHeader In, uint ShadingModel)					{ return StrataHasShadingModel(GetShadingModels(In), ShadingModel); }
+
+void SetIsSimpleMaterial(inout FStrataPixelHeader Out, bool bIn)							{ Out.State |= (bIn ? HEADER_MASK_ISSIMPLEMATERIAL : 0u); }
 void SetHasPrecShadowMask(inout FStrataPixelHeader Out, bool bIn)							{ Out.State |= (bIn ? HEADER_MASK_HASPRECSHADOWMASK : 0u); }
 void SetZeroPrecShadowMask(inout FStrataPixelHeader Out, bool bIn)							{ Out.State |= (bIn ? HEADER_MASK_ZEROPRECSHADOWMASK : 0u); }
 void SetIsSingleLayerWater(inout FStrataPixelHeader Out, bool bIn)							{ Out.State |= (bIn ? HEADER_MASK_ISSINGLELAYERWATER : 0u); }
 void SetCastContactShadow (inout FStrataPixelHeader Out, bool bIn)							{ Out.State |= (bIn ? HEADER_MASK_CASTCONTACTSHADOW : 0u); }
 void SetDynamicIndirectShadowCasterRepresentation(inout FStrataPixelHeader Out, bool bIn)	{ Out.State |= (bIn ? HEADER_MASK_HASDYNINDIRECTSHADOWCASTER : 0u); }
+void SetHasSubsurface(inout FStrataPixelHeader Out, bool bIn)								{ Out.State |= (bIn ? HEADER_MASK_HASSUBSURFACE : 0u); }
+void SetHasSubsurfaceProfile(inout FStrataPixelHeader Out, bool bIn)						{ Out.State |= (bIn ? HEADER_MASK_HASSUBSURFACEPROFILE : 0u); }
+
+void AddShadingModels(inout FStrataPixelHeader Out, uint bIn)								{ Out.State |= ((bIn<<HEADER_BIT_OFFSET_SHADINGMODELS) & HEADER_MASK_SHADINGMODELS); }
+void AndIsSimpleMaterial(inout FStrataPixelHeader Out, bool bIn)								{ Out.State = (Out.State & ~HEADER_MASK_ISSIMPLEMATERIAL) | ((Out.State & HEADER_MASK_ISSIMPLEMATERIAL) & (bIn ? HEADER_MASK_ISSIMPLEMATERIAL : 0u)); }
 
 FStrataPixelHeader InitialiseStrataPixelHeader()
 {
@@ -1731,6 +1720,36 @@ float3 StrataUnpackNormal(uint PackedNormal)
 	return normalize(Normal);
 }
 
+uint StrataPackNormal24(in float3 Normal)
+{
+	float l1norm = abs(Normal.x) + abs(Normal.y) + abs(Normal.z);
+	float2 result = Normal.xy * (1.0 / l1norm);
+	if (Normal.z < 0.0)
+	{
+		result = (1.0 - abs(result.yx)) * SignNotZero(result.xy);
+	}
+
+	uint PackedX = uint(clamp(result.x * 2047.0f + 2048.0f, 0.0f, 4095.0f));
+	uint PackedY = uint(clamp(result.y * 2047.0f + 2048.0f, 0.0f, 4095.0f));
+	return PackedX | (PackedY << 12);
+}
+
+float3 StrataUnpackNormal24(uint PackedNormal)
+{	
+	int X12Bits = 0xFFF & PackedNormal;
+	int Y12Bits = 0xFFF & (PackedNormal >> 12);
+
+	float x = float(X12Bits - 2048) / 2047.0f;
+	float y = float(Y12Bits - 2048) / 2047.0f;
+
+	float3 Normal = float3(x, y, 1.0f - abs(x) - abs(y));
+	if (Normal.z < 0.0f)
+	{
+		Normal.xy = (1.0f - abs(Normal.yx)) * SignNotZero(Normal.xy);
+	}
+	return normalize(Normal);
+}
+
 #define STRATA_TANGENT_SYMMETRY 1
 uint StrataPackNormalAndTangent(in float3 Normal, in float3 Tangent)
 {
@@ -1860,7 +1879,6 @@ void PackStrataOut(
 	float3 V,
 	inout bool bStrataSubsurfaceEnable,
 	inout float3 EmissiveColor,
-	inout FStrataClassification ClassificationData,
 	inout FStrataSubsurfaceData SSSData,
 	inout FStrataTopLayerData TopLayerData)
 {
@@ -1873,11 +1891,10 @@ void PackStrataOut(
 
 	// While packing Strata layer data, Classification/SSS/TopLayer data are extracted & stored for dedicated passes
 	// This avoid to run a post-basepass which would re-read all the material data
-	ClassificationData = (FStrataClassification)0;
 	SSSData = (FStrataSubsurfaceData)0;
 	TopLayerData = (FStrataTopLayerData)0;
 
-	SetIsSimpleMaterial(ClassificationData, true); // non strata material or sky pixels are considered as simple (so that (strata simple pixels) + (non-strata pixels) can build a simple tile).
+	SetIsSimpleMaterial(StrataPixelHeader, true); // non strata material or sky pixels are considered as simple (so that (strata simple pixels) + (non-strata pixels) can build a simple tile).
 
 	uint BSDFCount = 0;
 	uint BSDFVisibleCount = 0;
@@ -1925,16 +1942,13 @@ void PackStrataOut(
 			const float3 FakeL = float3(0, 0, 1);
 			Init(ContextNoLight, N, V, FakeL);
 
-			if (bTopLayer)
-			{
-				// together with a threshold to avoid 0-length normal after normalization in case influence is fadeout using the Weight node.
-				TopLayerData.WorldNormal += N * max(1e-5, BSDFCoverage);
-			}
-
+			bool bIsSimpleMaterial = true;
 			switch (BSDF_GETTYPE(BSDF))
 			{
 			case STRATA_BSDF_TYPE_SLAB:
 			{
+				bIsSimpleMaterial = IsStrataSlabFastPathCompatible(BSDF);
+
 				EmissiveColor += BSDF_GETEMISSIVE(BSDF) * Throughput;
 
 				// Two sided lighting can only be evalauted on the bottom layer, together with option SSS. Otherwise, the simple volume will be used.
@@ -1997,7 +2011,7 @@ void PackStrataOut(
 
 				if (bTopLayer)
 				{
-					ClassificationData.TopLayerRoughness += BSDFCoverage * SLAB_ROUGHNESS(BSDF);
+					TopLayerData.Roughness += BSDFCoverage * SLAB_ROUGHNESS(BSDF);
 				}
 
 				// STRATA_TODO	Should we have the Fuzz influence the next layter visibility accoding to fuzz amount? Should it also occlude diffuse?
@@ -2010,9 +2024,11 @@ void PackStrataOut(
 
 			case STRATA_BSDF_TYPE_HAIR:
 			{
+				bIsSimpleMaterial = false;
+
 				if (bTopLayer)
 				{
-					ClassificationData.TopLayerRoughness += BSDFCoverage * FullyRough;
+					TopLayerData.Roughness += BSDFCoverage * FullyRough;
 				}
 
 				// Hair are considered opaque
@@ -2023,11 +2039,13 @@ void PackStrataOut(
 
 			case STRATA_BSDF_TYPE_SINGLELAYERWATER:
 			{
+				bIsSimpleMaterial = false;
+
 				EmissiveColor += BSDF_GETEMISSIVE(BSDF) * Throughput;
 
 				if (bTopLayer)
 				{
-					ClassificationData.TopLayerRoughness += BSDFCoverage * SLW_ROUGHNESS(BSDF);
+					TopLayerData.Roughness += BSDFCoverage * SLW_ROUGHNESS(BSDF);
 				}
 
 				// Water is considered opaque, it has a custom process to simulate translucency
@@ -2038,7 +2056,6 @@ void PackStrataOut(
 
 			} // switch
 
-
 		#if STRATA_INLINE_SHADING
 			// Update the BSDF final weight with throughput towards the view from layers above and coverage if it has not been baked into the parameters.
 			Strata.Layers[l].BSDFs[i].FinalWeight = Throughput * (Strata.Layers[l].BSDFs[i].bPremultipliedCoverage ? 1.0 : Strata.Layers[l].BSDFs[i].Coverage);
@@ -2047,9 +2064,19 @@ void PackStrataOut(
 		#endif
 			
 			// Cull non visible BSDF
-			if (any(Strata.Layers[l].BSDFs[i].FinalWeight > 0))
+			const bool bIsVisible = any(Strata.Layers[l].BSDFs[i].FinalWeight > 0);
+			if (bIsVisible)
 			{
 				++BSDFVisibleCount;
+
+				AddShadingModels(StrataPixelHeader, StrataShadingModelBit(BSDF_GETTYPE(BSDF)));
+				AndIsSimpleMaterial(StrataPixelHeader, bIsSimpleMaterial);
+
+				if (bTopLayer)
+				{
+					// together with a threshold to avoid 0-length normal after normalization in case influence is fadeout using the Weight node.
+					TopLayerData.WorldNormal += N * max(1e-5, BSDFCoverage);
+				}
 			}
 
 			// Notify that the BSDF is at the top for SSR to only affect reflection there and not on the lower layers
@@ -2063,7 +2090,7 @@ void PackStrataOut(
 			const float BSDFWeightAvg = dot(LayerCoverage, 1.0f / 3.0f);
 
 			TopLayerData.WorldNormal = normalize(TopLayerData.WorldNormal);
-			ClassificationData.TopLayerRoughness = BSDFWeightAvg > 0.0f ? ClassificationData.TopLayerRoughness / BSDFWeightAvg : 0.0f;
+			TopLayerData.Roughness = BSDFWeightAvg > 0.0f ? TopLayerData.Roughness / BSDFWeightAvg : 0.0f;
 		}
 
 		// Take into account the non covered part for this layer and update throughput accordingly
@@ -2076,7 +2103,7 @@ void PackStrataOut(
 	// Now write out Strata data
 
 	// 1. the header
-	AndIsSimpleMaterial(ClassificationData, BSDFVisibleCount == 1);
+	AndIsSimpleMaterial(StrataPixelHeader, BSDFVisibleCount == 1);
 	const uint2 PackedHeader = PackStrataHeader(BSDFVisibleCount, StrataPixelHeader);
 	STRATA_STORE_UINT1(PackedHeader.x);
 #if ALLOW_STATIC_LIGHTING
@@ -2132,14 +2159,11 @@ void PackStrataOut(
 				}
 
 				const uint BSDFType = BSDF_GETTYPE(BSDF);
-				AddShadingModels(ClassificationData, StrataShadingModelBit(BSDFType));
 
 				switch (BSDF_GETTYPE(BSDF))
 				{
 				case STRATA_BSDF_TYPE_SLAB:
 				{
-					AndIsSimpleMaterial(ClassificationData, IsStrataSlabFastPathCompatible(BSDF));
-
 					// Now we pack diffuse and F0 in a special way: both encoded as R7G7B6 with gamma 2.0. The low bit cound will be hidden by dithering and TAA.
 					// We pack 32 bits in the first uint and the remaining bits are pack is the lowest significant bit of the second uint.
 					const uint PackedDiffuse20Bits	= PackR7G7B6Gamma2(SLAB_DIFFUSEALBEDO(BSDF));
@@ -2166,8 +2190,8 @@ void PackStrataOut(
 					const bool bIsThin = BSDF_GETISTHIN(BSDF);
 					if (BSDF_GETHASSSS(BSDF) || bIsSimpleVolume || bIsThin)
 					{
-						SetHasSubsurface(ClassificationData, !bIsThin && BSDF_GETHASSSS(BSDF));
-						SetHasSubsurfaceProfile(ClassificationData, !bIsThin && BSDF_GETHASSSSPROFILE(BSDF));
+						SetHasSubsurface(StrataPixelHeader, !bIsThin && BSDF_GETHASSSS(BSDF));
+						SetHasSubsurfaceProfile(StrataPixelHeader, !bIsThin && BSDF_GETHASSSSPROFILE(BSDF));
 
 						if (!bIsThin && BSDF_GETHASSSS(BSDF))
 						{
@@ -2207,7 +2231,6 @@ void PackStrataOut(
 				break;
 				case STRATA_BSDF_TYPE_HAIR:
 				{
-					AndIsSimpleMaterial(ClassificationData, false);
 					STRATA_STORE_UINT1(PackColorLinearToGamma2AlphaLinear(float4(HAIR_BASECOLOR(BSDF), HAIR_ROUGHNESS(BSDF))));
 					STRATA_STORE_UINT1(PackRGBA8(float4(HAIR_SCATTER(BSDF), HAIR_SPECULAR(BSDF), HAIR_BACKLIT(BSDF), HAIR_COMPLEXTRANSMITTANCE(BSDF))));
 					// 8 bytes
@@ -2215,7 +2238,6 @@ void PackStrataOut(
 				break;
 				case STRATA_BSDF_TYPE_SINGLELAYERWATER:
 				{
-					AndIsSimpleMaterial(ClassificationData, false);
 					STRATA_STORE_UINT1(PackColorLinearToGamma2AlphaLinear(float4(SLW_BASECOLOR(BSDF), 0.0f)));
 					STRATA_STORE_UINT1(PackRGBA8(float4(SLW_METALLIC(BSDF), SLW_SPECULAR(BSDF), SLW_ROUGHNESS(BSDF), SLW_TOPMATERIALOPACITY(BSDF))));
 				}
diff --git a/Engine/Shaders/Private/Strata/StrataMaterialClassification.usf b/Engine/Shaders/Private/Strata/StrataMaterialClassification.usf
--- a/Engine/Shaders/Private/Strata/StrataMaterialClassification.usf
+++ b/Engine/Shaders/Private/Strata/StrataMaterialClassification.usf
@@ -128,7 +128,7 @@ void MainPS(
 	FStrataSubsurfaceData SSSData				= (FStrataSubsurfaceData)0;
 	FStrataTopLayerData TopLayerData			= (FStrataTopLayerData)0;
 
-	FStrataAddressing StrataAddressing = GetStrataPixelDataByteOffset(PixelPos, uint2(View.BufferSizeAndInvSize.xy), MaxBytesPerPixel);
+	FStrataAddressing StrataAddressing = GetStrataPixelDataByteOffset(PixelPos, uint2(PrimaryView.BufferSizeAndInvSize.xy), MaxBytesPerPixel);
 	FStrataPixelHeader Header = UnpackStrataHeaderIn(MaterialLobesBuffer, StrataAddressing);
 
 	SetIsSimpleMaterial(ClassificationData, true); // non strata material or sky pixels are considered as simple.
@@ -179,7 +179,8 @@ uint2 DecodeTile(uint In)
 int TileSize;
 int bRectPrimitive;
 int2 ViewResolution;
-Texture2D<uint> ClassificationTexture;
+uint MaxBytesPerPixel;
+ByteAddressBuffer MaterialBuffer;
 RWBuffer<uint> SimpleTileIndirectDataBuffer;
 RWBuffer<uint> SimpleTileListDataBuffer;
 RWBuffer<uint> ComplexTileIndirectDataBuffer;
@@ -206,10 +207,11 @@ void TileMainCS(uint2 DispatchThreadId : SV_DispatchThreadID, uint LinearIndex :
 	bool bContainsStrataMaterial = false;
 	if (all(PixelCoord < uint2(ViewResolution)))
 	{
-		const uint PackClassification = ClassificationTexture.Load(uint3(PixelCoord, 0));
-		FStrataClassification Data = StrataUnpackClassificationData(PackClassification);
-		bContainsComplexMaterial = !IsSimpleMaterial(Data);
-		bContainsStrataMaterial = IsStrataMaterial(Data);
+		FStrataAddressing StrataAddressing = GetStrataPixelDataByteOffset(PixelCoord, uint2(View.BufferSizeAndInvSize.xy), MaxBytesPerPixel);
+		FStrataPixelHeader StrataPixelHeader = UnpackStrataHeaderIn(MaterialBuffer, StrataAddressing);
+
+		bContainsStrataMaterial  = StrataPixelHeader.BSDFCount > 0;
+		bContainsComplexMaterial = !IsSimpleMaterial(StrataPixelHeader);
 	}
 
 #if PERMUTATION_WAVE_OPS
@@ -347,7 +349,7 @@ void StencilMainPS(
 #if SHADER_CLEAR_MATERIAL_BUFFER
 
 RWByteAddressBuffer MaterialLobesBufferUAV;
-RWTexture2D<uint> ClassificationTextureUAV;
+RWTexture2D<uint> TopLayerTextureUAV;
 RWTexture2D<uint2> SSSTextureUAV;
 uint MaxBytesPerPixel;
 int2 TiledViewBufferResolution;
@@ -364,15 +366,14 @@ void ClearMaterialBufferMainCS(uint2 DispatchThreadId : SV_DispatchThreadID)
 	// The first layer of tiled uints contains the header that we need to clear so we only write a single uint per pixel instead of clearing the entire buffer which would be too slow.
 
 	FStrataPixelHeader StrataHeader = InitialiseStrataPixelHeader();
+	StrataHeader.State |= HEADER_MASK_ISSIMPLEMATERIAL;
 	uint BSDFCount = 0;
 	uint PackedStrataHeader = PackStrataHeader(BSDFCount, StrataHeader);
 
 	FStrataAddressing StrataAddressing = GetStrataPixelDataByteOffset(DispatchThreadId.xy, TiledViewBufferResolution, MaxBytesPerPixel);
 	MaterialLobesBufferUAV.Store(StrataAddressing.CurrentAddr, PackedStrataHeader);
 
-	FStrataClassification ClassificationData = (FStrataClassification)0;
-	SetIsSimpleMaterial(ClassificationData, true);
-	ClassificationTextureUAV[DispatchThreadId.xy] = StrataPackClassificationData(ClassificationData);
+	TopLayerTextureUAV[DispatchThreadId.xy] = 0;
 
 	FStrataSubsurfaceData StrataSubsurfaceData = (FStrataSubsurfaceData)0;
 	SSSTextureUAV[DispatchThreadId.xy] = StrataPackSSSData(StrataSubsurfaceData);
diff --git a/Engine/Shaders/Private/Strata/StrataVisualize.usf b/Engine/Shaders/Private/Strata/StrataVisualize.usf
--- a/Engine/Shaders/Private/Strata/StrataVisualize.usf
+++ b/Engine/Shaders/Private/Strata/StrataVisualize.usf
@@ -164,7 +164,7 @@ void VisualizeMaterialPS(
 		uint2 RectMin = uint2(MaterialPos + InspectSquare);
 		uint2 RectMax = RectMin;
 
-		const FStrataClassification StrataData = StrataUnpackClassificationData(Strata.ClassificationTexture.Load(uint3(MaterialPos, 0)));
+		const FStrataTopLayerData TopLayerData = StrataUnpackTopLayerData(Strata.TopLayerTexture.Load(uint3(MaterialPos, 0)));
 		const FStrataSubsurfaceData SSSData = StrataUnpackSSSData(Strata.SSSTexture.Load(uint3(MaterialPos, 0)));
 
 		FStrataAddressing StrataAddressing = GetStrataPixelDataByteOffset(MaterialPos, uint2(View.BufferSizeAndInvSize.xy), Strata.MaxBytesPerPixel);
@@ -207,7 +207,7 @@ void VisualizeMaterialPS(
 #define PRINTMATERIALCONTSHAD()	PRINTBSDFCHAR(_C_); PRINTBSDFCHAR(_O_); PRINTBSDFCHAR(_N_);  PRINTBSDFCHAR(_T_); PRINTBSDFCHAR(_SPC_);PRINTBSDFCHAR(_S_); PRINTBSDFCHAR(_H_); PRINTBSDFCHAR(_A_);  PRINTBSDFCHAR(_D_); PRINTBSDFCHAR(_SPC_); PRINTBSDFCHAR(_0_+(DoesCastContactShadow(Header)?1:0));
 #define PRINTMATERIALINDSHAD()	PRINTBSDFCHAR(_I_); PRINTBSDFCHAR(_N_); PRINTBSDFCHAR(_D_); PRINTBSDFCHAR(_SPC_); PRINTBSDFCHAR(_S_); PRINTBSDFCHAR(_H_); PRINTBSDFCHAR(_A_);  PRINTBSDFCHAR(_D_); PRINTBSDFCHAR(_SPC_); PRINTBSDFCHAR(_0_+(HasDynamicIndirectShadowCasterRepresentation(Header)?1:0));
 
-#define PRINTCLASSIFROUGH()		PRINTBSDFCHAR(_C_); PRINTBSDFCHAR(_L_); PRINTBSDFCHAR(_A_); PRINTBSDFCHAR(_S_); PRINTBSDFCHAR(_DOT_); PRINTBSDFCHAR(_R_) PRINTBSDFCHAR(_O_) PRINTBSDFCHAR(_U_) PRINTBSDFCHAR(_G_) PRINTBSDFCHAR(_H_) PRINTBSDFCHAR(_SPC_);	PRINTBSDFCHAR(_SPC_); PRINTFLOAT1(StrataData.TopLayerRoughness);
+#define PRINTCLASSIFROUGH()		PRINTBSDFCHAR(_C_); PRINTBSDFCHAR(_L_); PRINTBSDFCHAR(_A_); PRINTBSDFCHAR(_S_); PRINTBSDFCHAR(_DOT_); PRINTBSDFCHAR(_R_) PRINTBSDFCHAR(_O_) PRINTBSDFCHAR(_U_) PRINTBSDFCHAR(_G_) PRINTBSDFCHAR(_H_) PRINTBSDFCHAR(_SPC_);	PRINTBSDFCHAR(_SPC_); PRINTFLOAT1(TopLayerData.Roughness);
 #define PRINTCLASSIFBASCOL()	PRINTBSDFCHAR(_C_); PRINTBSDFCHAR(_L_); PRINTBSDFCHAR(_A_); PRINTBSDFCHAR(_S_); PRINTBSDFCHAR(_DOT_); PRINTBSDFCHAR(_B_) PRINTBSDFCHAR(_A_) PRINTBSDFCHAR(_S_) PRINTBSDFCHAR(_C_) PRINTBSDFCHAR(_O_) PRINTBSDFCHAR(_L_);	PRINTBSDFCHAR(_SPC_); PRINTFLOAT3(SSSData.BaseColor);
 #define PRINTCLASSIFISPROF()	PRINTBSDFCHAR(_C_); PRINTBSDFCHAR(_L_); PRINTBSDFCHAR(_A_); PRINTBSDFCHAR(_S_); PRINTBSDFCHAR(_DOT_); PRINTBSDFCHAR(_I_) PRINTBSDFCHAR(_S_) PRINTBSDFCHAR(_P_) PRINTBSDFCHAR(_R_) PRINTBSDFCHAR(_O_) PRINTBSDFCHAR(_F_);	PRINTBSDFCHAR(_SPC_); PRINTBSDFCHAR(SSSData.bIsProfile?_1_:_0_);
 #define PRINTCLASSIFPROID()		PRINTBSDFCHAR(_C_); PRINTBSDFCHAR(_L_); PRINTBSDFCHAR(_A_); PRINTBSDFCHAR(_S_); PRINTBSDFCHAR(_DOT_); PRINTBSDFCHAR(_P_) PRINTBSDFCHAR(_R_) PRINTBSDFCHAR(_O_) PRINTBSDFCHAR(_I_) PRINTBSDFCHAR(_D_) PRINTBSDFCHAR(_SPC_);	PRINTBSDFCHAR(_SPC_); PRINTFLOATNOFRAC(SSSData.ProfileId, 3);
@@ -239,7 +239,7 @@ void VisualizeMaterialPS(
 
 		PRINTCLASSIFROUGH();
 		NEXTLINE;
-		if (HasSubsurface(StrataData))
+		if (HasSubsurface(Header))
 		{
 			PRINTCLASSIFBASCOL();
 			NEXTLINE;
@@ -278,7 +278,7 @@ void VisualizeMaterialPS(
 		NEXTLINE;
 
 		NEXTLINE;
-		if (HasSubsurface(StrataData))
+		if (HasSubsurface(Header))
 		{
 			NEXTLINE;
 			NEXTLINE;
diff --git a/Engine/Shaders/Private/SubsurfaceBurleyNormalized.ush b/Engine/Shaders/Private/SubsurfaceBurleyNormalized.ush
--- a/Engine/Shaders/Private/SubsurfaceBurleyNormalized.ush
+++ b/Engine/Shaders/Private/SubsurfaceBurleyNormalized.ush
@@ -391,7 +391,7 @@ FStrataSubsurfaceData LoadStrataSSSData(float2 UV)
 FStrataTopLayerData LoadStrataTopLayerData(float2 UV)
 {
 	const float2 PixelPos = UV.xy * View.BufferSizeAndInvSize.xy - View.ViewRectMin.xy;
-	return StrataUnpackTopLayerData(Strata.TopLayerNormalTexture.Load(uint3(PixelPos, 0)));
+	return StrataUnpackTopLayerData(Strata.TopLayerTexture.Load(uint3(PixelPos, 0)));
 }
 #endif
 
@@ -781,7 +781,7 @@ float3 GetBilateralNormal(float2 ClampedUV, float2 Extent)
 #if BILATERAL_FILTER_KERNEL_FUNCTION_TYPE == BILATERAL_FILTER_KERNEL_FUNCTION_DEPTH_AND_NORMAL
 	#if STRATA_ENABLED
 	const uint2 BufferPos = ClampedUV * Extent;
-	const FStrataTopLayerData TopLayerData = StrataUnpackTopLayerData(Strata.TopLayerNormalTexture.Load(uint3(BufferPos, 0)));
+	const FStrataTopLayerData TopLayerData = StrataUnpackTopLayerData(Strata.TopLayerTexture.Load(uint3(BufferPos, 0)));
 	return TopLayerData.WorldNormal;
 	#else
 	return GetScreenSpaceData(ClampedUV).GBuffer.WorldNormal;
diff --git a/Engine/Source/Runtime/Renderer/Private/Strata/Strata.cpp b/Engine/Source/Runtime/Renderer/Private/Strata/Strata.cpp
--- a/Engine/Source/Runtime/Renderer/Private/Strata/Strata.cpp
+++ b/Engine/Source/Runtime/Renderer/Private/Strata/Strata.cpp
@@ -73,12 +73,10 @@ IMPLEMENT_GLOBAL_SHADER_PARAMETER_STRUCT(FStrataGlobalUniformParameters, "Strata
 
 void FStrataSceneData::Reset()
 {
-	ClassificationTexture = nullptr;
-	TopLayerNormalTexture = nullptr;
+	TopLayerTexture = nullptr;
 	SSSTexture = nullptr;
 
-	ClassificationTextureUAV = nullptr;
-	TopLayerNormalTextureUAV = nullptr;
+	TopLayerTextureUAV = nullptr;
 	SSSTextureUAV = nullptr;
 
 	MaterialLobesBuffer = nullptr;
@@ -117,7 +115,7 @@ namespace Strata
 static void AddStrataClearMaterialBufferPass(
 	FRDGBuilder& GraphBuilder, 
 	FRDGBufferUAVRef MaterialLobesBufferUAV,
-	FRDGTextureUAVRef ClassificationTextureUAV,
+	FRDGTextureUAVRef TopLayerTextureUAV,
 	FRDGTextureUAVRef SSSTextureUAV,
 	uint32 MaxBytesPerPixel, 
 	FIntPoint TiledViewBufferResolution);
@@ -168,12 +166,6 @@ void InitialiseStrataFrameSceneData(FSceneRenderer& SceneRenderer, FRDGBuilder&
 		const uint32 RoundToValue = 4u;
 		StrataSceneData.MaxBytesPerPixel = FMath::DivideAndRoundUp(MaterialConservativeByteCountPerPixel, RoundToValue) * RoundToValue;
 
-		// Classification texture
-		{
-			StrataSceneData.ClassificationTexture = GraphBuilder.CreateTexture(FRDGTextureDesc::Create2D(SceneTextureExtent, PF_R16_UINT, FClearValueBinding::Black, TexCreate_ShaderResource | TexCreate_UAV), TEXT("Strata.ClassificationTexture"));
-			StrataSceneData.ClassificationTextureUAV = GraphBuilder.CreateUAV(StrataSceneData.ClassificationTexture);
-		}
-
 		// Tile classification buffers
 		{
 			const int32 TileInPixel = GetStrataBufferTileSize();
@@ -197,8 +189,8 @@ void InitialiseStrataFrameSceneData(FSceneRenderer& SceneRenderer, FRDGBuilder&
 
 		// Top layer texture
 		{
-			StrataSceneData.TopLayerNormalTexture = GraphBuilder.CreateTexture(FRDGTextureDesc::Create2D(SceneTextureExtent, PF_R32_UINT, FClearValueBinding::Black, TexCreate_ShaderResource | TexCreate_UAV), TEXT("Strata.TopLayerNormalTexture"));
-			StrataSceneData.TopLayerNormalTextureUAV = GraphBuilder.CreateUAV(StrataSceneData.TopLayerNormalTexture);
+			StrataSceneData.TopLayerTexture = GraphBuilder.CreateTexture(FRDGTextureDesc::Create2D(SceneTextureExtent, PF_R32_UINT, FClearValueBinding::Black, TexCreate_ShaderResource | TexCreate_UAV), TEXT("Strata.TopLayerTexture"));
+			StrataSceneData.TopLayerTextureUAV = GraphBuilder.CreateUAV(StrataSceneData.TopLayerTexture);
 		}
 
 		// SSS texture
@@ -233,7 +225,7 @@ void InitialiseStrataFrameSceneData(FSceneRenderer& SceneRenderer, FRDGBuilder&
 		AddStrataClearMaterialBufferPass(
 			GraphBuilder, 
 			StrataSceneData.MaterialLobesBufferUAV,
-			StrataSceneData.ClassificationTextureUAV,
+			StrataSceneData.TopLayerTextureUAV,
 			StrataSceneData.SSSTextureUAV,
 			StrataSceneData.MaxBytesPerPixel, 
 			MaterialBufferSizeXY);
@@ -256,8 +248,7 @@ void BindStrataBasePassUniformParameters(FRDGBuilder& GraphBuilder, FStrataScene
 		OutStrataUniformParameters.MaxBytesPerPixel = StrataSceneData->MaxBytesPerPixel;
 		OutStrataUniformParameters.MaterialLobesBufferUAV = StrataSceneData->MaterialLobesBufferUAV;
 
-		OutStrataUniformParameters.ClassificationTextureUAV = StrataSceneData->ClassificationTextureUAV;
-		OutStrataUniformParameters.TopLayerNormalTextureUAV = StrataSceneData->TopLayerNormalTextureUAV;
+		OutStrataUniformParameters.TopLayerTextureUAV = StrataSceneData->TopLayerTextureUAV;
 		OutStrataUniformParameters.SSSTextureUAV = StrataSceneData->SSSTextureUAV;
 	}
 	else
@@ -268,8 +259,7 @@ void BindStrataBasePassUniformParameters(FRDGBuilder& GraphBuilder, FStrataScene
 		OutStrataUniformParameters.bRoughDiffuse = 0u;
 		OutStrataUniformParameters.MaxBytesPerPixel = 0;
 		OutStrataUniformParameters.MaterialLobesBufferUAV = GraphBuilder.CreateUAV(GraphBuilder.RegisterExternalBuffer(GWhiteVertexBufferWithRDG->Buffer), PF_R32_UINT);
-		OutStrataUniformParameters.ClassificationTextureUAV = DummyWritableUAV;
-		OutStrataUniformParameters.TopLayerNormalTextureUAV = DummyWritableUAV;
+		OutStrataUniformParameters.TopLayerTextureUAV = DummyWritableUAV;
 		OutStrataUniformParameters.SSSTextureUAV = DummyWritableUAV;
 	}
 }
@@ -281,8 +271,7 @@ void BindStrataGlobalUniformParameters(FRDGBuilder& GraphBuilder, FStrataSceneDa
 		OutStrataUniformParameters.bRoughDiffuse = StrataSceneData->bRoughDiffuse ? 1u : 0u;
 		OutStrataUniformParameters.MaxBytesPerPixel = StrataSceneData->MaxBytesPerPixel;
 		OutStrataUniformParameters.MaterialLobesBuffer = StrataSceneData->MaterialLobesBufferSRV;
-		OutStrataUniformParameters.ClassificationTexture = StrataSceneData->ClassificationTexture;
-		OutStrataUniformParameters.TopLayerNormalTexture = StrataSceneData->TopLayerNormalTexture;
+		OutStrataUniformParameters.TopLayerTexture = StrataSceneData->TopLayerTexture;
 		OutStrataUniformParameters.SSSTexture = StrataSceneData->SSSTexture;
 	}
 	else
@@ -291,8 +280,7 @@ void BindStrataGlobalUniformParameters(FRDGBuilder& GraphBuilder, FStrataSceneDa
 		OutStrataUniformParameters.bRoughDiffuse = 0;
 		OutStrataUniformParameters.MaxBytesPerPixel = 0;
 		OutStrataUniformParameters.MaterialLobesBuffer = GraphBuilder.CreateSRV(GraphBuilder.RegisterExternalBuffer(GWhiteVertexBufferWithRDG->Buffer), PF_R32_UINT);
-		OutStrataUniformParameters.ClassificationTexture = SystemTextures.Black;
-		OutStrataUniformParameters.TopLayerNormalTexture = SystemTextures.DefaultNormal8Bit;
+		OutStrataUniformParameters.TopLayerTexture = SystemTextures.DefaultNormal8Bit;
 		OutStrataUniformParameters.SSSTexture = SystemTextures.Black;
 	}
 }
@@ -393,7 +381,7 @@ class FStrataClearMaterialBufferCS : public FGlobalShader
 
 	BEGIN_SHADER_PARAMETER_STRUCT(FParameters, )
 		SHADER_PARAMETER_RDG_BUFFER_UAV(RWByteAddressBuffer, MaterialLobesBufferUAV)
-		SHADER_PARAMETER_RDG_TEXTURE_UAV(RWTexture2D<uint>, ClassificationTextureUAV)
+		SHADER_PARAMETER_RDG_TEXTURE_UAV(RWTexture2D<uint>, TopLayerTextureUAV)
 		SHADER_PARAMETER_RDG_TEXTURE_UAV(RWTexture2D<uint2>, SSSTextureUAV)
 		SHADER_PARAMETER(uint32, MaxBytesPerPixel)
 		SHADER_PARAMETER(FIntPoint, TiledViewBufferResolution)
@@ -427,7 +415,8 @@ class FStrataMaterialTileClassificationPassCS : public FGlobalShader
 		SHADER_PARAMETER(int32, TileSize)
 		SHADER_PARAMETER(int32, bRectPrimitive)
 		SHADER_PARAMETER(FIntPoint, ViewResolution)
-		SHADER_PARAMETER_RDG_TEXTURE(Texture2D<uint>, ClassificationTexture)
+		SHADER_PARAMETER(uint32, MaxBytesPerPixel)
+		SHADER_PARAMETER_RDG_BUFFER_SRV(ByteAddressBuffer, MaterialBuffer)
 		SHADER_PARAMETER_RDG_BUFFER_UAV(RWBuffer, SimpleTileIndirectDataBuffer)
 		SHADER_PARAMETER_RDG_BUFFER_UAV(RWBuffer, SimpleTileListDataBuffer)
 		SHADER_PARAMETER_RDG_BUFFER_UAV(RWBuffer, ComplexTileIndirectDataBuffer)
@@ -657,10 +646,12 @@ void AddStrataMaterialClassificationPass(FRDGBuilder& GraphBuilder, const FMinim
 			PermutationVector.Set< FStrataMaterialTileClassificationPassCS::FWaveOps >(bWaveOps);
 			TShaderMapRef<FStrataMaterialTileClassificationPassCS> ComputeShader(View.ShaderMap, PermutationVector);
 			FStrataMaterialTileClassificationPassCS::FParameters* PassParameters = GraphBuilder.AllocParameters<FStrataMaterialTileClassificationPassCS::FParameters>();
+			PassParameters->ViewUniformBuffer = View.ViewUniformBuffer;
 			PassParameters->TileSize = GetStrataBufferTileSize();	// STRATA_TODO not sure we want to tie the buffer tile optimisation for cache and the Categotisation tile size?
 			PassParameters->bRectPrimitive = GRHISupportsRectTopology ? 1 : 0;
 			PassParameters->ViewResolution = View.ViewRect.Size();
-			PassParameters->ClassificationTexture = View.StrataSceneData->ClassificationTexture;
+			PassParameters->MaxBytesPerPixel = View.StrataSceneData->MaxBytesPerPixel;
+			PassParameters->MaterialBuffer = View.StrataSceneData->MaterialLobesBufferSRV;
 			PassParameters->SimpleTileListDataBuffer = View.StrataSceneData->ClassificationTileListBufferUAV[EStrataTileMaterialType::ESimple];
 			PassParameters->SimpleTileIndirectDataBuffer = View.StrataSceneData->ClassificationTileIndirectBufferUAV[EStrataTileMaterialType::ESimple];
 			PassParameters->ComplexTileListDataBuffer = View.StrataSceneData->ClassificationTileListBufferUAV[EStrataTileMaterialType::EComplex];
@@ -672,7 +663,7 @@ void AddStrataMaterialClassificationPass(FRDGBuilder& GraphBuilder, const FMinim
 				RDG_EVENT_NAME("Strata::MaterialTileClassification(%s)", bWaveOps ? TEXT("Wave") : TEXT("SharedMemory")),
 				ComputeShader,
 				PassParameters,
-				FComputeShaderUtils::GetGroupCount(View.StrataSceneData->ClassificationTexture->Desc.Extent, GroupSize));
+				FComputeShaderUtils::GetGroupCount(PassParameters->ViewResolution, GroupSize));
 		}
 	}
 }
@@ -680,7 +671,7 @@ void AddStrataMaterialClassificationPass(FRDGBuilder& GraphBuilder, const FMinim
 static void AddStrataClearMaterialBufferPass(
 	FRDGBuilder& GraphBuilder,
 	FRDGBufferUAVRef MaterialLobesBufferUAV,
-	FRDGTextureUAVRef ClassificationTextureUAV,
+	FRDGTextureUAVRef TopLayerTextureUAV,
 	FRDGTextureUAVRef SSSTextureUAV,
 	uint32 MaxBytesPerPixel,
 	FIntPoint TiledViewBufferResolution)
@@ -688,7 +679,7 @@ static void AddStrataClearMaterialBufferPass(
 	TShaderMapRef<FStrataClearMaterialBufferCS> ComputeShader(GetGlobalShaderMap(GMaxRHIFeatureLevel));
 	FStrataClearMaterialBufferCS::FParameters* PassParameters = GraphBuilder.AllocParameters<FStrataClearMaterialBufferCS::FParameters>();
 	PassParameters->MaterialLobesBufferUAV = MaterialLobesBufferUAV;
-	PassParameters->ClassificationTextureUAV = ClassificationTextureUAV;
+	PassParameters->TopLayerTextureUAV = TopLayerTextureUAV;
 	PassParameters->SSSTextureUAV = SSSTextureUAV;
 	PassParameters->MaxBytesPerPixel = MaxBytesPerPixel;
 	PassParameters->TiledViewBufferResolution = TiledViewBufferResolution;
@@ -736,6 +727,4 @@ FScreenPassTexture AddStrataDebugPasses(FRDGBuilder& GraphBuilder, const FViewIn
 	return MoveTemp(ScreenPassSceneColor);
 }
 
-} // namespace Strata
-
-
+} // namespace Strata
\ No newline at end of file
diff --git a/Engine/Source/Runtime/Renderer/Private/Strata/Strata.h b/Engine/Source/Runtime/Renderer/Private/Strata/Strata.h
--- a/Engine/Source/Runtime/Renderer/Private/Strata/Strata.h
+++ b/Engine/Source/Runtime/Renderer/Private/Strata/Strata.h
@@ -19,8 +19,7 @@ BEGIN_SHADER_PARAMETER_STRUCT(FStrataBasePassUniformParameters, )
 	SHADER_PARAMETER(uint32, MaxBytesPerPixel)
 	SHADER_PARAMETER(uint32, bRoughDiffuse)
 	SHADER_PARAMETER_RDG_BUFFER_UAV(RWByteAddressBuffer, MaterialLobesBufferUAV)
-	SHADER_PARAMETER_RDG_TEXTURE_UAV(RWTexture2D<uint>, ClassificationTextureUAV)
-	SHADER_PARAMETER_RDG_TEXTURE_UAV(RWTexture2D<uint>, TopLayerNormalTextureUAV)
+	SHADER_PARAMETER_RDG_TEXTURE_UAV(RWTexture2D<uint>, TopLayerTextureUAV)
 	SHADER_PARAMETER_RDG_TEXTURE_UAV(RWTexture2D<uint2>, SSSTextureUAV)
 END_SHADER_PARAMETER_STRUCT()
 
@@ -28,8 +27,7 @@ BEGIN_GLOBAL_SHADER_PARAMETER_STRUCT(FStrataGlobalUniformParameters, )
 	SHADER_PARAMETER(uint32, MaxBytesPerPixel)
 	SHADER_PARAMETER(uint32, bRoughDiffuse)
 	SHADER_PARAMETER_RDG_BUFFER_SRV(ByteAddressBuffer, MaterialLobesBuffer)
-	SHADER_PARAMETER_RDG_TEXTURE(Texture2D<uint>, ClassificationTexture)
-	SHADER_PARAMETER_RDG_TEXTURE(Texture2D<uint>, TopLayerNormalTexture)
+	SHADER_PARAMETER_RDG_TEXTURE(Texture2D<uint>, TopLayerTexture)
 	SHADER_PARAMETER_RDG_TEXTURE(Texture2D<uint2>, SSSTexture)
 END_GLOBAL_SHADER_PARAMETER_STRUCT()
 
@@ -59,12 +57,10 @@ struct FStrataSceneData
 	FRDGBufferUAVRef ClassificationTileIndirectBufferUAV[EStrataTileMaterialType::ECount];
 	FRDGBufferSRVRef ClassificationTileIndirectBufferSRV[EStrataTileMaterialType::ECount];
 
-	FRDGTextureRef ClassificationTexture;
-	FRDGTextureRef TopLayerNormalTexture;
+	FRDGTextureRef TopLayerTexture;
 	FRDGTextureRef SSSTexture;
 
-	FRDGTextureUAVRef ClassificationTextureUAV;
-	FRDGTextureUAVRef TopLayerNormalTextureUAV;
+	FRDGTextureUAVRef TopLayerTextureUAV;
 	FRDGTextureUAVRef SSSTextureUAV;
 
 	TRDGUniformBufferRef<FStrataGlobalUniformParameters> StrataGlobalUniformParameters{};
