Fix typo into unpacking function for 16-bits normal encoding

diff --git a/Engine/Shaders/Private/Strata/Strata.ush b/Engine/Shaders/Private/Strata/Strata.ush
--- a/Engine/Shaders/Private/Strata/Strata.ush
+++ b/Engine/Shaders/Private/Strata/Strata.ush
@@ -1709,8 +1709,8 @@ uint StrataPackNormal(in float3 Normal)
 
 float3 StrataUnpackNormal(uint PackedNormal)
 {
-	float x = f16tof32(0xFFFFFFFF &  PackedNormal);
-	float y = f16tof32(0xFFFFFFFF & (PackedNormal>>16));
+	float x = f16tof32(0xFFFF &  PackedNormal);
+	float y = f16tof32(0xFFFF & (PackedNormal>>16));
 
 	float3 Normal = float3(x, y, 1.0f - abs(x) - abs(y));
 	if (Normal.z < 0.0f)
