New start material buffer: 2 uint using rop and the rest is written through UAV.
Huge optimisation for console.

diff --git a/Engine/Shaders/Private/BasePassPixelShader.usf b/Engine/Shaders/Private/BasePassPixelShader.usf
--- a/Engine/Shaders/Private/BasePassPixelShader.usf
+++ b/Engine/Shaders/Private/BasePassPixelShader.usf
@@ -1619,20 +1619,24 @@ void FPixelShaderInOut_MainPS(
 
 	uint2 PixelPos = uint2(In.SvPosition.xy);
 
-	// classificaiton & pre-preprocessed data
+	float3 EmissiveLuminance = 0.0f;
 	FStrataSubsurfaceData SSSData = (FStrataSubsurfaceData)0;
 	FStrataTopLayerData TopLayerData = (FStrataTopLayerData)0;
 
-	float3 EmissiveLuminance = 0.0f;
+	// Generate the strata material data to write out
+	FStrataAddressing StrataAddressing = GetStrataPixelDataByteOffset(PixelPos, uint2(ResolvedView.BufferSizeAndInvSize.xy), OpaqueBasePass.Strata.MaxBytesPerPixel);
+	FRWStrataMaterialContainer RWStrataMaterialContainer = InitialiseRWStrataMaterialContainer(OpaqueBasePass.Strata.MaterialTextureArrayUAV);
 	PackStrataOut(
-		OpaqueBasePass.Strata.MaterialLobesBufferUAV,
+		RWStrataMaterialContainer,
 		Dither,
-		GetStrataPixelDataByteOffset(PixelPos, uint2(ResolvedView.BufferSizeAndInvSize.xy), OpaqueBasePass.Strata.MaxBytesPerPixel),
+		StrataAddressing,
 		StrataPixelHeader, StrataData, MaterialParameters.CameraVector, bStrataSubsurfaceEnable, EmissiveLuminance,
 		SSSData, TopLayerData);
 
-	// Write out classificaiton & pre-preprocessed data
-	OpaqueBasePass.Strata.TopLayerTextureUAV[PixelPos] = StrataPackTopLayerData(TopLayerData);
+	// Write out MRT data
+	Out.StrataOutput[0] = RWStrataMaterialContainer.MaterialRenderTargets[0];
+	Out.StrataOutput[1] = RWStrataMaterialContainer.MaterialRenderTargets[1];
+	Out.StrataOutput[2] = StrataPackTopLayerData(TopLayerData);
 
 	// Only write SSS datta if needed
 	BRANCH
diff --git a/Engine/Shaders/Private/CapsuleShadowShaders.usf b/Engine/Shaders/Private/CapsuleShadowShaders.usf
--- a/Engine/Shaders/Private/CapsuleShadowShaders.usf
+++ b/Engine/Shaders/Private/CapsuleShadowShaders.usf
@@ -964,7 +964,7 @@ void ApplySelfShadowingIntensityForDeferred(float2 ScreenUV, inout float Visibil
 #if STRATA_ENABLED
 		uint2 PixelPos = View.ViewRectMin.xy + View.ViewSizeAndInvSize.xy * ScreenUV.xy;
 		FStrataAddressing StrataAddressing = GetStrataPixelDataByteOffset(PixelPos, uint2(View.BufferSizeAndInvSize.xy), Strata.MaxBytesPerPixel);
-		FStrataPixelHeader StrataPixelHeader = UnpackStrataHeaderIn(Strata.MaterialLobesBuffer, StrataAddressing);
+		FStrataPixelHeader StrataPixelHeader = UnpackStrataHeaderIn(Strata.MaterialTextureArray, StrataAddressing);
 		const bool bHasDynamicIndirectShadowCasterRepresentation = HasDynamicIndirectShadowCasterRepresentation(StrataPixelHeader);
 #else
 		FGBufferData GBufferData = GetGBufferData(ScreenUV);
diff --git a/Engine/Shaders/Private/Common.ush b/Engine/Shaders/Private/Common.ush
--- a/Engine/Shaders/Private/Common.ush
+++ b/Engine/Shaders/Private/Common.ush
@@ -1902,6 +1902,9 @@ struct FPixelShaderOut
 {
 	// [0..7], only usable if PIXELSHADEROUTPUT_MRT0, PIXELSHADEROUTPUT_MRT1, ... is 1
 	float4 MRT[8];
+	
+	// Explicit uint output specific to strata.
+	uint StrataOutput[3];
 
 	// Pixel Shader OutCoverage, only usable if PIXELSHADEROUTPUT_COVERAGE is 1
 	uint Coverage;
diff --git a/Engine/Shaders/Private/DebugProbes.usf b/Engine/Shaders/Private/DebugProbes.usf
--- a/Engine/Shaders/Private/DebugProbes.usf
+++ b/Engine/Shaders/Private/DebugProbes.usf
@@ -6,7 +6,7 @@
 #if PROJECT_STRATA
 #include "SceneTexturesCommon.ush"
 #include "Strata/Strata.ush"
-RWByteAddressBuffer MaterialLobesBufferUAV;
+RWTexture2DArray<uint> MaterialTextureArrayUAV;
 uint MaxBytesPerPixel;
 #endif
 int DebugProbesMode;
@@ -279,11 +279,15 @@ void MainPS(
 		float3 EmissiveLuminance = 0.0f;
 		bool bStrataSubsurfaceEnable = false;
 		const float Dither = 0.0f;
+		FStrataAddressing StrataAddressing = GetStrataPixelDataByteOffset(PixelCoord, uint2(ResolvedView.BufferSizeAndInvSize.xy), MaxBytesPerPixel);
+		FRWStrataMaterialContainer RWStrataMaterialContainer = InitialiseRWStrataMaterialContainer(MaterialTextureArrayUAV);
 		PackStrataOut(
-			MaterialLobesBufferUAV, Dither,
-			GetStrataPixelDataByteOffset(PixelCoord, uint2(ResolvedView.BufferSizeAndInvSize.xy), MaxBytesPerPixel),
+			RWStrataMaterialContainer, 
+			Dither,
+			StrataAddressing,
 			StrataPixelHeader, StrataData, -RayDir, bStrataSubsurfaceEnable, EmissiveLuminance,
 			SSSData, TopLayerData);
+		// TODO write out the two MRT uint?
 	}
 	
 #else // PROJECT_STRATA
diff --git a/Engine/Shaders/Private/DeferredLightPixelShaders.usf b/Engine/Shaders/Private/DeferredLightPixelShaders.usf
--- a/Engine/Shaders/Private/DeferredLightPixelShaders.usf
+++ b/Engine/Shaders/Private/DeferredLightPixelShaders.usf
@@ -218,7 +218,7 @@ void DeferredLightPixelMain(
 #if STRATA_ENABLED
 
 	FStrataAddressing StrataAddressing = GetStrataPixelDataByteOffset(PixelPos, uint2(View.BufferSizeAndInvSize.xy), Strata.MaxBytesPerPixel);
-	FStrataPixelHeader StrataPixelHeader = UnpackStrataHeaderIn(Strata.MaterialLobesBuffer, StrataAddressing);
+	FStrataPixelHeader StrataPixelHeader = UnpackStrataHeaderIn(Strata.MaterialTextureArray, StrataAddressing);
 
 	BRANCH
 	if (StrataPixelHeader.BSDFCount > 0	// This test is also enough to exclude sky pixels
@@ -264,7 +264,7 @@ void DeferredLightPixelMain(
 				LightMask,
 				ShadowTerms,
 				RectTexture,
-				Strata.MaterialLobesBuffer,
+				Strata.MaterialTextureArray,
 				StrataAddressing,
 				StrataPixelHeader);
 		}
diff --git a/Engine/Shaders/Private/DeferredLightingCommon.ush b/Engine/Shaders/Private/DeferredLightingCommon.ush
--- a/Engine/Shaders/Private/DeferredLightingCommon.ush
+++ b/Engine/Shaders/Private/DeferredLightingCommon.ush
@@ -145,11 +145,11 @@ float ShadowRayCast(
 		float3 SampleUVz = RayStartUVz + RayStepUVz * FirstHitTime;
 	#if STRATA_ENABLED
 		#if (!MATERIAL_IS_STRATA || STRATA_TRANSLUCENT_ENABLED)
-		// StrataStruct.MaterialLobesBuffer is only available to translucent materials or global shaders computing contact shadows.
+		// StrataStruct.MaterialTextureArray is only available to translucent materials or global shaders computing contact shadows.
 		// It is not defined for opaque materials, only the UAV is available to write the strata buffer.
 		uint2 PixelPos = View.ViewRectMin.xy + View.ViewSizeAndInvSize.xy * SampleUVz.xy;
 		FStrataAddressing StrataAddressing = GetStrataPixelDataByteOffset(PixelPos, uint2(View.BufferSizeAndInvSize.xy), StrataStruct.MaxBytesPerPixel);
-		FStrataPixelHeader StrataPixelHeader = UnpackStrataHeaderIn(StrataStruct.MaterialLobesBuffer, StrataAddressing);
+		FStrataPixelHeader StrataPixelHeader = UnpackStrataHeaderIn(StrataStruct.MaterialTextureArray, StrataAddressing);
 		bOutHitCastContactShadow = DoesCastContactShadow(StrataPixelHeader);
 		#endif
 	#else
diff --git a/Engine/Shaders/Private/PixelShaderOutputCommon.ush b/Engine/Shaders/Private/PixelShaderOutputCommon.ush
--- a/Engine/Shaders/Private/PixelShaderOutputCommon.ush
+++ b/Engine/Shaders/Private/PixelShaderOutputCommon.ush
@@ -77,6 +77,25 @@ void MainPS
 		, out float4 OutTarget2 : SV_Target2
 #endif
 
+#if STRATA_OPAQUE_DEFERRED
+#if PIXELSHADEROUTPUT_MRT3
+#error Strata cannot map to such a case
+#endif
+#if PIXELSHADEROUTPUT_MRT2
+		, out uint StrataOutput3 : SV_Target3
+		, out uint StrataOutput4 : SV_Target4
+		, out uint StrataOutput5 : SV_Target5
+#elif PIXELSHADEROUTPUT_MRT1
+		, out uint StrataOutput2 : SV_Target2
+		, out uint StrataOutput3 : SV_Target3
+		, out uint StrataOutput4 : SV_Target4
+#else
+		, out uint StrataOutput1 : SV_Target1
+		, out uint StrataOutput2 : SV_Target2
+		, out uint StrataOutput3 : SV_Target3
+#endif
+#else // STRATA_OPAQUE_DEFERRED
+
 #if PIXELSHADEROUTPUT_MRT3
 		, out float4 OutTarget3 : SV_Target3
 #endif
@@ -97,6 +116,8 @@ void MainPS
 		, out float4 OutTarget7 : SV_Target7
 #endif
 
+#endif // STRATA_OPAQUE_DEFERRED
+
 		// todo: if we are not inside a volumetric decal we could use OPTIONAL_OutDepthConservative for better performance, this would require a shader permutation
 		OPTIONAL_OutDepthConservative
 
@@ -147,6 +168,30 @@ void MainPS
 	OutTarget2 = PixelShaderOut.MRT[2];
 #endif
 
+#if STRATA_OPAQUE_DEFERRED
+
+	// In this case, here is the gbuffer pattern
+	// MRT0 is pixel color
+	// MRT1 is velocity if enabled or precomputed shadow if velocity if disabled and precomputed shadow enabled
+	// MRT2 is precomputed shadow if both velocity and prec shadow are enabled.
+	// After, Strata output are appended (3 uint output)
+
+#if PIXELSHADEROUTPUT_MRT2
+	StrataOutput3 = PixelShaderOut.StrataOutput[0];
+	StrataOutput4 = PixelShaderOut.StrataOutput[1];
+	StrataOutput5 = PixelShaderOut.StrataOutput[2];
+#elif PIXELSHADEROUTPUT_MRT1
+	StrataOutput2 = PixelShaderOut.StrataOutput[0];
+	StrataOutput3 = PixelShaderOut.StrataOutput[1];
+	StrataOutput4 = PixelShaderOut.StrataOutput[2];
+#else
+	StrataOutput1 = PixelShaderOut.StrataOutput[0];
+	StrataOutput2 = PixelShaderOut.StrataOutput[1];
+	StrataOutput3 = PixelShaderOut.StrataOutput[2];
+#endif
+
+#else // STRATA_OPAQUE_DEFERRED
+
 #if PIXELSHADEROUTPUT_MRT3
 	OutTarget3 = PixelShaderOut.MRT[3];
 #endif
@@ -167,6 +212,8 @@ void MainPS
 	OutTarget7 = PixelShaderOut.MRT[7];
 #endif
 
+#endif // STRATA_OPAQUE_DEFERRED
+
 #if PIXELSHADEROUTPUT_COVERAGE || PIXELSHADEROUTPUT_A2C
 	OutCoverage = PixelShaderOut.Coverage;
 #endif 
diff --git a/Engine/Shaders/Private/PostProcessAmbientOcclusionCommon.ush b/Engine/Shaders/Private/PostProcessAmbientOcclusionCommon.ush
--- a/Engine/Shaders/Private/PostProcessAmbientOcclusionCommon.ush
+++ b/Engine/Shaders/Private/PostProcessAmbientOcclusionCommon.ush
@@ -108,7 +108,7 @@ float3 GetWorldSpaceNormalFromAOInput(float2 UV, float4 SvPosition)
 	#else
 #if STRATA_ENABLED
 		FStrataAddressing StrataAddressing = GetStrataPixelDataByteOffset(SvPosition.xy, uint2(View.BufferSizeAndInvSize.xy), Strata.MaxBytesPerPixel);
-		FStrataPixelHeader StrataHeader = UnpackStrataHeaderIn(Strata.MaterialLobesBuffer, StrataAddressing);
+		FStrataPixelHeader StrataHeader = UnpackStrataHeaderIn(Strata.MaterialTextureArray, StrataAddressing);
 		if (!StrataHasShadingModel(StrataHeader, STRATA_BSDF_TYPE_HAIR))
 		{
 			const FStrataTopLayerData TopLayerData = StrataUnpackTopLayerData(Strata.TopLayerTexture.Load(uint3(SvPosition.xy, 0)));
diff --git a/Engine/Shaders/Private/ReflectionEnvironmentPixelShader.usf b/Engine/Shaders/Private/ReflectionEnvironmentPixelShader.usf
--- a/Engine/Shaders/Private/ReflectionEnvironmentPixelShader.usf
+++ b/Engine/Shaders/Private/ReflectionEnvironmentPixelShader.usf
@@ -369,7 +369,7 @@ void ReflectionEnvironmentSkyLighting(
 #if STRATA_ENABLED
 
 	FStrataAddressing StrataAddressing = GetStrataPixelDataByteOffset(PixelPos, uint2(View.BufferSizeAndInvSize.xy), Strata.MaxBytesPerPixel);
-	FStrataPixelHeader StrataPixelHeader = UnpackStrataHeaderIn(Strata.MaterialLobesBuffer, StrataAddressing);
+	FStrataPixelHeader StrataPixelHeader = UnpackStrataHeaderIn(Strata.MaterialTextureArray, StrataAddressing);
 	BRANCH
 	if (StrataPixelHeader.BSDFCount > 0) // This test is also enough to exclude sky pixels
 	{
@@ -404,7 +404,7 @@ void ReflectionEnvironmentSkyLighting(
 			V,
 			AmbientOcclusion,
 			TopLayerSpecularContributionFactor,
-			Strata.MaterialLobesBuffer,
+			Strata.MaterialTextureArray,
 			StrataAddressing,
 			StrataPixelHeader,
 			SSRReductionFactor,
diff --git a/Engine/Shaders/Private/SSRT/SSRTReflections.usf b/Engine/Shaders/Private/SSRT/SSRTReflections.usf
--- a/Engine/Shaders/Private/SSRT/SSRTReflections.usf
+++ b/Engine/Shaders/Private/SSRT/SSRTReflections.usf
@@ -143,10 +143,10 @@ void ScreenSpaceReflections(
 		// We do not rerun material classification after the Strata material buffer has been patched with SLW material.
 		// So in this case we simply load the Header and water BSDF directly because we know there can only be a single one.
 		FStrataAddressing StrataAddressing = GetStrataPixelDataByteOffset(PixelCoord, uint2(View.BufferSizeAndInvSize.xy), Strata.MaxBytesPerPixel);
-		FStrataPixelHeader StrataPixelHeader = UnpackStrataHeaderIn(Strata.MaterialLobesBuffer, StrataAddressing);
+		FStrataPixelHeader StrataPixelHeader = UnpackStrataHeaderIn(Strata.MaterialTextureArray, StrataAddressing);
 		if (StrataPixelHeader.BSDFCount > 0 && IsSingleLayerWater(StrataPixelHeader))
 		{
-			const FStrataBSDF SLWaterBSDF = UnpackStrataBSDFIn(Strata.MaterialLobesBuffer, StrataAddressing);
+			const FStrataBSDF SLWaterBSDF = UnpackStrataBSDFIn(Strata.MaterialTextureArray, StrataAddressing);
 			N = StrataGetBSDFSharedBasis(StrataPixelHeader, SLWaterBSDF, StrataAddressing)[2];
 			Roughness = SLW_ROUGHNESS(SLWaterBSDF);
 			bNoMaterial = false;
diff --git a/Engine/Shaders/Private/ShadingFurnaceTest.usf b/Engine/Shaders/Private/ShadingFurnaceTest.usf
--- a/Engine/Shaders/Private/ShadingFurnaceTest.usf
+++ b/Engine/Shaders/Private/ShadingFurnaceTest.usf
@@ -56,13 +56,13 @@ void MainPS(
 	{
 		// Fetch the first layer of roughness & normal only. This fine for simple material but will not work for clear coat.
 		FStrataAddressing StrataAddressing = GetStrataPixelDataByteOffset(PixelPos, uint2(View.BufferSizeAndInvSize.xy), Strata.MaxBytesPerPixel);
-		FStrataPixelHeader StrataPixelHeader = UnpackStrataHeaderIn(Strata.MaterialLobesBuffer, StrataAddressing);
+		FStrataPixelHeader StrataPixelHeader = UnpackStrataHeaderIn(Strata.MaterialTextureArray, StrataAddressing);
 		bIsValid = StrataPixelHeader.BSDFCount > 0;
 		if (bIsValid)
 		{
 			const float3 DummyV = float3(0, 0, 1);
 			const float3 DummyL = float3(0, 0, 1);
-			FStrataBSDF BSDF = UnpackStrataBSDFIn(Strata.MaterialLobesBuffer, StrataAddressing);
+			FStrataBSDF BSDF = UnpackStrataBSDFIn(Strata.MaterialTextureArray, StrataAddressing);
 			FStrataBSDFContext BSDFContext = StrataCreateBSDFContext(StrataPixelHeader, BSDF, StrataAddressing, DummyV, DummyL);
 
 			SceneDepth = CalcSceneDepth(ScreenUV);
@@ -227,7 +227,7 @@ void MainPS(
 			#elif 0
 				{
 					FStrataAddressing StrataAddressing = GetStrataPixelDataByteOffset(PixelPos, uint2(View.BufferSizeAndInvSize.xy), Strata.MaxBytesPerPixel);
-					FStrataPixelHeader StrataPixelHeader = UnpackStrataHeaderIn(Strata.MaterialLobesBuffer, StrataAddressing);
+					FStrataPixelHeader StrataPixelHeader = UnpackStrataHeaderIn(Strata.MaterialTextureArray, StrataAddressing);
 					const float LightMask = 1;
 
 					BRANCH
@@ -242,7 +242,7 @@ void MainPS(
 							LightMask,
 							ShadowTerms,
 							RectTexture,
-							Strata.MaterialLobesBuffer,
+							Strata.MaterialTextureArray,
 							StrataAddressing,
 							StrataPixelHeader).xyz;
 					}
diff --git a/Engine/Shaders/Private/ShadowProjectionPixelShader.usf b/Engine/Shaders/Private/ShadowProjectionPixelShader.usf
--- a/Engine/Shaders/Private/ShadowProjectionPixelShader.usf
+++ b/Engine/Shaders/Private/ShadowProjectionPixelShader.usf
@@ -256,7 +256,7 @@ void Main(
 		bool bIsEyeOrHair 				= false;
 	#if STRATA_ENABLED
 		FStrataAddressing StrataAddressing = GetStrataPixelDataByteOffset(PixelCoord.xy, uint2(View.BufferSizeAndInvSize.xy), Strata.MaxBytesPerPixel);
-		FStrataPixelHeader StrataHeader = UnpackStrataHeaderIn(Strata.MaterialLobesBuffer, StrataAddressing);
+		FStrataPixelHeader StrataHeader = UnpackStrataHeaderIn(Strata.MaterialTextureArray, StrataAddressing);
 		if (IsStrataMaterial(StrataHeader))
 		{
 			bIsEyeOrHair				= StrataHasShadingModel(StrataHeader, STRATA_BSDF_TYPE_HAIR);
diff --git a/Engine/Shaders/Private/SingleLayerWaterComposite.usf b/Engine/Shaders/Private/SingleLayerWaterComposite.usf
--- a/Engine/Shaders/Private/SingleLayerWaterComposite.usf
+++ b/Engine/Shaders/Private/SingleLayerWaterComposite.usf
@@ -58,7 +58,7 @@ SingleLayerWaterCompositeOutput SingleLayerWaterComposite(float2 BufferUV, float
 
 	// Water do nto read the classification data to avoid running it a second time when water is enabled.
 	FStrataAddressing StrataAddressing = GetStrataPixelDataByteOffset(SvPosition.xy, uint2(View.BufferSizeAndInvSize.xy), Strata.MaxBytesPerPixel);
-	FStrataPixelHeader StrataPixelHeader = UnpackStrataHeaderIn(Strata.MaterialLobesBuffer, StrataAddressing);
+	FStrataPixelHeader StrataPixelHeader = UnpackStrataHeaderIn(Strata.MaterialTextureArray, StrataAddressing);
 	if (StrataPixelHeader.BSDFCount > 0 && IsSingleLayerWater(StrataPixelHeader))
 	{
 		float DeviceZ = SampleDeviceZFromSceneTextures(BufferUV);
@@ -72,7 +72,7 @@ SingleLayerWaterCompositeOutput SingleLayerWaterComposite(float2 BufferUV, float
 		float DeltaDepth = CameraIsUnderWater ? WaterDepth : OpaqueDepth - WaterDepth;	// Inverted depth
 		if (DeltaDepth > 0.0)
 		{
-			const FStrataBSDF SLWaterBSDF = UnpackStrataBSDFIn(Strata.MaterialLobesBuffer, StrataAddressing);
+			const FStrataBSDF SLWaterBSDF = UnpackStrataBSDFIn(Strata.MaterialTextureArray, StrataAddressing);
 			const float3 WorldNormal = StrataGetBSDFSharedBasis(StrataPixelHeader, SLWaterBSDF, StrataAddressing)[2];
 			const float3 F0 = ComputeF0(SLW_SPECULAR(SLWaterBSDF), SLW_BASECOLOR(SLWaterBSDF), SLW_METALLIC(SLWaterBSDF));
 			const float Roughness = SLW_ROUGHNESS(SLWaterBSDF);
@@ -252,7 +252,7 @@ void WaterTileCatergorisationCS(uint3 ThreadId : SV_DispatchThreadID, uint3 Grou
 
 			#if STRATA_ENABLED
 				FStrataAddressing StrataAddressing = GetStrataPixelDataByteOffset(PixelPos, uint2(View.BufferSizeAndInvSize.xy), Strata.MaxBytesPerPixel);
-				FStrataPixelHeader StrataPixelHeader = UnpackStrataHeaderIn(Strata.MaterialLobesBuffer, StrataAddressing);
+				FStrataPixelHeader StrataPixelHeader = UnpackStrataHeaderIn(Strata.MaterialTextureArray, StrataAddressing);
 				bContainsWater = bContainsWater || (StrataPixelHeader.BSDFCount > 0 && IsSingleLayerWater(StrataPixelHeader));
 			#else
 				FGBufferData GBuffer = GetGBufferDataFromSceneTextures(BufferUV);
@@ -278,7 +278,7 @@ void WaterTileCatergorisationCS(uint3 ThreadId : SV_DispatchThreadID, uint3 Grou
 
 	#if STRATA_ENABLED
 		FStrataAddressing StrataAddressing = GetStrataPixelDataByteOffset(PixelPos, uint2(View.BufferSizeAndInvSize.xy), Strata.MaxBytesPerPixel);
-		FStrataPixelHeader StrataPixelHeader = UnpackStrataHeaderIn(Strata.MaterialLobesBuffer, StrataAddressing);
+		FStrataPixelHeader StrataPixelHeader = UnpackStrataHeaderIn(Strata.MaterialTextureArray, StrataAddressing);
 		bContainsWater = bContainsWater || (StrataPixelHeader.BSDFCount > 0 && IsSingleLayerWater(StrataPixelHeader));
 	#else
 		FGBufferData GBuffer = GetGBufferDataFromSceneTextures(BufferUV);
diff --git a/Engine/Shaders/Private/Strata/Strata.ush b/Engine/Shaders/Private/Strata/Strata.ush
--- a/Engine/Shaders/Private/Strata/Strata.ush
+++ b/Engine/Shaders/Private/Strata/Strata.ush
@@ -24,9 +24,35 @@
 #define STRATA_UNROLL_N(X) UNROLL_N(X)
 #endif
 
-#define STRATA_MEMACCESS_UINT  0
-#define STRATA_MEMACCESS_UINT4 1
-#define STRATA_MEMACCESS STRATA_MEMACCESS_UINT
+// Strata Rendertarget / UAV setup
+// RT0: Color
+// [insert here potential velocity rendertarget and prec shadow]
+// RT2: Strata uint
+// RT3: Strata uint
+// RT4: Top normal
+//
+// UAV SSS data written out if needed
+// UAV if strata data overrun the MRT output size
+
+#define STARTA_MATERIAL_RENDERTARGET_COUNT 2
+struct FRWStrataMaterialContainerStruct
+{
+	uint MaterialRenderTargets[STARTA_MATERIAL_RENDERTARGET_COUNT];
+	RWTexture2DArray<uint> ExtraMaterialDataUAV;
+};
+#define FRWStrataMaterialContainer	FRWStrataMaterialContainerStruct
+#define FStrataMaterialContainer	Texture2DArray<uint>
+
+FRWStrataMaterialContainer InitialiseRWStrataMaterialContainer(RWTexture2DArray<uint> ExtraMaterialDataUAV)
+{
+	FRWStrataMaterialContainer RWStrataMaterialContainer;
+	for (int i = 0; i < STARTA_MATERIAL_RENDERTARGET_COUNT; ++i)
+	{
+		RWStrataMaterialContainer.MaterialRenderTargets[i] = 0;
+	}
+	RWStrataMaterialContainer.ExtraMaterialDataUAV = ExtraMaterialDataUAV;
+	return RWStrataMaterialContainer;
+}
 
 // During the base pass or forward rendering, shared local bases are simply available from registers.
 // If a BSDF is loaded for processing from the strata buffer, normals are unpacked on demand to reduce VGPR pressure and increase occupancy.
@@ -707,62 +733,16 @@ uint DivideAndRoundUp(uint Dividend, uint Divisor, uint DivisorAsBitShift)
 
 struct FStrataAddressing
 {
-	uint  CurrentAddr;
-
-#if STRATA_MEMACCESS == STRATA_MEMACCESS_UINT4
-	uint  ByteOffsetToNextUint4;
-	uint4 Uint4Cache;
-	int   UintProcessed;	// uints remaining to read, or uint appended before write
-
-	// Packedlocal basis in a uint4 available for unpacking later.
-	uint4 PackedLocalBasis;
-#if STRATA_MAX_SHAREDLOCALBASES_REGISTERS!=4
-#error We do not support more than 4 PackedLocalBasis as of today.
-#endif
-
-#else
-	uint  ByteOffsetToNextUint;
-#endif
+	uint  CurrentIndex;
 
-	// This is to apply checkboard pattern when needed
 	uint2 PixelCoords;
 
-	// This should be compiled out as not used in most cases expect for debug and visualization.
 	uint  ReadBytes;
 };
 FStrataAddressing GetStrataPixelDataByteOffset(uint2 PixelPos, uint2 ViewBufferSize, uint StrataMaxBytesPerPixel)
 {
 	FStrataAddressing StrataAddressing = (FStrataAddressing)0;
-
-	// A tile of pixels has each uint that can be read packed in contiguous memory to be cache efficient. Then, the next set of uint to be read is in another "plane", after all uint from all the pixel/tile.
-	// Without that, L2 is rarely hit and GCN TA FIFO command buffer is getting full in no time because memory fetches are not served in time (GPU cost then increase by at least x2)
-
-#if STRATA_MEMACCESS == STRATA_MEMACCESS_UINT4
-	uint TilePixelSizeBytes			= 4*4;	// 4*sizeof(uint)
-#else
-	uint TilePixelSizeBytes			= 4;
-#endif
-	uint TileSizeBytes				= STRATA_DATA_TILE_SIZE * STRATA_DATA_TILE_SIZE * TilePixelSizeBytes;
-	uint2 TileCount					= uint2(DivideAndRoundUp(ViewBufferSize.x, STRATA_DATA_TILE_SIZE, STRATA_DATA_TILE_SIZE_DIV_AS_SHIFT), DivideAndRoundUp(ViewBufferSize.y, STRATA_DATA_TILE_SIZE, STRATA_DATA_TILE_SIZE_DIV_AS_SHIFT));
-	uint2 TileCoord					= uint2(PixelPos.x >> STRATA_DATA_TILE_SIZE_DIV_AS_SHIFT, PixelPos.y >> STRATA_DATA_TILE_SIZE_DIV_AS_SHIFT);
-	uint2 PixelCoordInTile			= PixelPos - TileCoord * STRATA_DATA_TILE_SIZE;
-
-	uint PreviousTilesBytes			= (TileCoord.y * TileCount.x + TileCoord.x) * TileSizeBytes;
-#if 0
-	uint PreviousPixelsInTileBytes	= (PixelCoordInTile.y * STRATA_DATA_TILE_SIZE + PixelCoordInTile.x) * TilePixelSizeBytes;
-#else
-	uint PreviousPixelsInTileBytes	= MortonEncode(PixelCoordInTile) * TilePixelSizeBytes; // Z pattern
-#endif
-	StrataAddressing.CurrentAddr	= PreviousTilesBytes + PreviousPixelsInTileBytes;
-
-#if STRATA_MEMACCESS == STRATA_MEMACCESS_UINT4
-	StrataAddressing.ByteOffsetToNextUint4 = TileCount.x * TileCount.y * TileSizeBytes;
-	StrataAddressing.Uint4Cache = 0;
-	StrataAddressing.UintProcessed = 0;
-#else
-	StrataAddressing.ByteOffsetToNextUint = TileCount.x * TileCount.y * TileSizeBytes;
-#endif
-
+	StrataAddressing.CurrentIndex = 0;
 	StrataAddressing.PixelCoords = PixelPos;
 	StrataAddressing.ReadBytes = 0;
 	return StrataAddressing;
@@ -770,79 +750,40 @@ FStrataAddressing GetStrataPixelDataByteOffset(uint2 PixelPos, uint2 ViewBufferS
 
 #if COMPILER_SWITCH // These platforms are complaining so simply disabling Strata there for now.
 
-void StrataStoreUint1(RWByteAddressBuffer StrataBuffer, inout FStrataAddressing StrataAddressing, uint Data) {}
+void StrataStoreUint1(inout FRWStrataMaterialContainer StrataBuffer, inout FStrataAddressing StrataAddressing, uint Data) {}
 #define STRATA_STORE_UINT1(x) 
 
-void FinalizeWrites(RWByteAddressBuffer StrataBuffer, inout FStrataAddressing StrataAddressing) {}
+void FinalizeWrites(inout FRWStrataMaterialContainer StrataBuffer, inout FStrataAddressing StrataAddressing) {}
 
-uint StrataLoadUint1(ByteAddressBuffer StrataBuffer, inout FStrataAddressing StrataAddressing) { return 0; }
+uint StrataLoadUint1(FStrataMaterialContainer StrataBuffer, inout FStrataAddressing StrataAddressing) { return 0; }
 
 #else
 
-void StrataStoreUint1(RWByteAddressBuffer StrataBuffer, inout FStrataAddressing StrataAddressing, uint Data)
+void StrataStoreUint1(inout FRWStrataMaterialContainer StrataBuffer, inout FStrataAddressing StrataAddressing, uint Data)
 {
-#if STRATA_MEMACCESS == STRATA_MEMACCESS_UINT4
-	BRANCH
-	if(StrataAddressing.UintProcessed == 4)
+	// We heavily rely on the compiler to optimize out the if code
+	if (StrataAddressing.CurrentIndex < STARTA_MATERIAL_RENDERTARGET_COUNT)
 	{
-		StrataBuffer.Store4(StrataAddressing.CurrentAddr, StrataAddressing.Uint4Cache);
-		StrataAddressing.CurrentAddr += StrataAddressing.ByteOffsetToNextUint4;
-		StrataAddressing.UintProcessed = 0;
-		StrataAddressing.ReadBytes += 16;
+		StrataBuffer.MaterialRenderTargets[0] = StrataAddressing.CurrentIndex == 0 ? Data : StrataBuffer.MaterialRenderTargets[0];
+		StrataBuffer.MaterialRenderTargets[1] = StrataAddressing.CurrentIndex == 1 ? Data : StrataBuffer.MaterialRenderTargets[1];
 	}
-
-	// TODO check what is done during this addressing... It should be a simple bit select
-	StrataAddressing.Uint4Cache[0] = StrataAddressing.UintProcessed == 0 ? Data : StrataAddressing.Uint4Cache[0];
-	StrataAddressing.Uint4Cache[1] = StrataAddressing.UintProcessed == 1 ? Data : StrataAddressing.Uint4Cache[1];
-	StrataAddressing.Uint4Cache[2] = StrataAddressing.UintProcessed == 2 ? Data : StrataAddressing.Uint4Cache[2];
-	StrataAddressing.Uint4Cache[3] = StrataAddressing.UintProcessed == 3 ? Data : StrataAddressing.Uint4Cache[3];
-	StrataAddressing.UintProcessed++;
-#else
-	StrataBuffer.Store(StrataAddressing.CurrentAddr, Data);
-	StrataAddressing.CurrentAddr += StrataAddressing.ByteOffsetToNextUint;
-	StrataAddressing.ReadBytes += 4;
-#endif
+	else
+	{
+		StrataBuffer.ExtraMaterialDataUAV[uint3(StrataAddressing.PixelCoords, StrataAddressing.CurrentIndex)] = Data;
+	}
+	StrataAddressing.CurrentIndex++;
 }
 #define STRATA_STORE_UINT1(x)	StrataStoreUint1(StrataBuffer, StrataAddressing, x)
 
-void FinalizeWrites(RWByteAddressBuffer StrataBuffer, inout FStrataAddressing StrataAddressing)
+void FinalizeWrites(inout FRWStrataMaterialContainer StrataBuffer, inout FStrataAddressing StrataAddressing)
 {
-#if STRATA_MEMACCESS == STRATA_MEMACCESS_UINT4
-	BRANCH
-	if(StrataAddressing.UintProcessed > 0)
-	{
-		StrataBuffer.Store4(StrataAddressing.CurrentAddr, StrataAddressing.Uint4Cache);
-		StrataAddressing.CurrentAddr += StrataAddressing.ByteOffsetToNextUint4;
-		StrataAddressing.UintProcessed = 0;
-		StrataAddressing.ReadBytes += 16;
-	}
-#endif
+	// NOP
 }
 
-uint StrataLoadUint1(ByteAddressBuffer StrataBuffer, inout FStrataAddressing StrataAddressing)
+uint StrataLoadUint1(FStrataMaterialContainer StrataBuffer, inout FStrataAddressing StrataAddressing)
 {
-#if STRATA_MEMACCESS == STRATA_MEMACCESS_UINT4
-	BRANCH
-	if(StrataAddressing.UintProcessed == 0)
-	{
-		StrataAddressing.Uint4Cache = StrataBuffer.Load4(StrataAddressing.CurrentAddr);
-		StrataAddressing.CurrentAddr += StrataAddressing.ByteOffsetToNextUint4;
-		StrataAddressing.UintProcessed = 4;
-		StrataAddressing.ReadBytes += 16;
-	}
-
-	// TODO check what is done during this addressing... It should be a simple bit select
-	uint Data = StrataAddressing.Uint4Cache[0];
-	Data = StrataAddressing.UintProcessed == 3 ? StrataAddressing.Uint4Cache[1] : Data;
-	Data = StrataAddressing.UintProcessed == 2 ? StrataAddressing.Uint4Cache[2] : Data;
-	Data = StrataAddressing.UintProcessed == 1 ? StrataAddressing.Uint4Cache[3] : Data;
-	StrataAddressing.UintProcessed--;
-#else
-	uint Data = StrataBuffer.Load(StrataAddressing.CurrentAddr);
-	StrataAddressing.CurrentAddr += StrataAddressing.ByteOffsetToNextUint;
-	StrataAddressing.ReadBytes += 4;
-#endif
-
+	uint Data = StrataBuffer[uint3(StrataAddressing.PixelCoords, StrataAddressing.CurrentIndex)];
+	StrataAddressing.CurrentIndex++;
 	return Data;
 }
 
@@ -889,10 +830,8 @@ struct FStrataPixelHeader
 	// All the shared local bases in VGPRS. This is only used in basepass and during forward rendering.
 	FSharedLocalBases SharedLocalBases;
 #else
-	ByteAddressBuffer StrataBuffer;
-	#if STRATA_MEMACCESS == STRATA_MEMACCESS_UINT
-	uint SharedLocalBasesByteOffset;
-	#endif
+	FStrataMaterialContainer StrataBuffer;
+	uint SharedLocalBasesIndexOffset;
 	uint SharedLocalBasesTypes;
 #endif
 
@@ -933,9 +872,7 @@ FStrataPixelHeader InitialiseStrataPixelHeader()
 #if STRATA_INLINE_SHADING
 	Out.SharedLocalBases = (FSharedLocalBases)0;
 #else
-	#if STRATA_MEMACCESS == STRATA_MEMACCESS_UINT
-	Out.SharedLocalBasesByteOffset = 0;
-	#endif
+	Out.SharedLocalBasesIndexOffset = 0;
 	Out.SharedLocalBasesTypes = 0;
 #endif
 	Out.MaterialAO = 1.0f;
@@ -953,11 +890,8 @@ float3x3 StrataGetBSDFSharedBasis(in FStrataPixelHeader StrataPixelHeader, in FS
 	return float3x3(Tangent, Bitangent, Normal);
 #else
 
-#if STRATA_MEMACCESS == STRATA_MEMACCESS_UINT4
-	uint PackedBasis = StrataAddressing.PackedLocalBasis[BSDF_GETSHAREDLOCALBASISID(BSDF)];
-#else
-	uint PackedBasis = StrataPixelHeader.StrataBuffer.Load(StrataPixelHeader.SharedLocalBasesByteOffset + BSDF_GETSHAREDLOCALBASISID(BSDF) * StrataAddressing.ByteOffsetToNextUint);
-#endif
+	uint PackedBasis = StrataPixelHeader.StrataBuffer[uint3(StrataAddressing.PixelCoords, StrataPixelHeader.SharedLocalBasesIndexOffset + BSDF_GETSHAREDLOCALBASISID(BSDF))];
+
 	uint BasisType   = StrataGetSharedLocalBasisType(StrataPixelHeader.SharedLocalBasesTypes, BSDF_GETSHAREDLOCALBASISID(BSDF));
 	float3x3 OutTangentBasis;
 	if (BasisType == STRATA_BASIS_TYPE_NORMAL)
@@ -1842,12 +1776,12 @@ void StrataUnpackNormalAndTangent(inout float3 Normal, inout float3 Tangent, in
 	}
 }
 
-void StrataLoad_R11G11B10F(ByteAddressBuffer StrataBuffer, inout FStrataAddressing StrataAddressing, inout float3 A)
+void StrataLoad_R11G11B10F(FStrataMaterialContainer StrataBuffer, inout FStrataAddressing StrataAddressing, inout float3 A)
 {
 	A = UnpackR11G11B10F(StrataLoadUint1(StrataBuffer, StrataAddressing));
 }
 
-void StrataLoad_R11_G11_B10F(ByteAddressBuffer StrataBuffer, inout FStrataAddressing StrataAddressing, inout float A, inout float B, inout float C)
+void StrataLoad_R11_G11_B10F(FStrataMaterialContainer StrataBuffer, inout FStrataAddressing StrataAddressing, inout float A, inout float B, inout float C)
 {
 	float3 temp = UnpackR11G11B10F(StrataLoadUint1(StrataBuffer, StrataAddressing));
 	A = temp.x;
@@ -1855,25 +1789,28 @@ void StrataLoad_R11_G11_B10F(ByteAddressBuffer StrataBuffer, inout FStrataAddres
 	C = temp.z;
 }
 
-void StrataLoad_ColorGamma2ToLinear(ByteAddressBuffer StrataBuffer, inout FStrataAddressing StrataAddressing, inout float3 A)
+void StrataLoad_ColorGamma2ToLinear(FStrataMaterialContainer StrataBuffer, inout FStrataAddressing StrataAddressing, inout float3 A)
 {
 	A = UnpackColorGamma2ToLinear(StrataLoadUint1(StrataBuffer, StrataAddressing));
 }
 
-void StrataLoad_ColorGamma2ToLinearAlphaLinear(ByteAddressBuffer StrataBuffer, inout FStrataAddressing StrataAddressing, inout float3 RGB, inout float A)
+void StrataLoad_ColorGamma2ToLinearAlphaLinear(FStrataMaterialContainer StrataBuffer, inout FStrataAddressing StrataAddressing, inout float3 RGB, inout float A)
 {
 	float4 Data = UnpackColorGamma2ToLinearAlphaLinear(StrataLoadUint1(StrataBuffer, StrataAddressing));
 	RGB = Data.rgb;
 	A = Data.a;
 }
 
+///////////////////////////////////////////////////////////////////////////////
+// Strata deferred base-pass output
+
 uint2 PackStrataHeader(uint BSDFCount, FStrataPixelHeader In);
 
 // If this is changed, please update the compiler side material size evaluation in StrataMaterial.cpp
 void PackStrataOut(
-	RWByteAddressBuffer StrataBuffer, 
+	inout FRWStrataMaterialContainer StrataBuffer,
 	float Dither,
-	FStrataAddressing StrataAddressing,
+	inout FStrataAddressing StrataAddressing,
 	FStrataPixelHeader StrataPixelHeader,
 	FStrataData Strata,
 	float3 V,
@@ -2274,7 +2211,7 @@ uint2 PackStrataHeader(uint InBSDFCount, FStrataPixelHeader InHeader)
 
 // Unpack only strata header
 // If this is changed, please update the compiler side material size evaluation in StrataMaterial.cpp
-FStrataPixelHeader UnpackStrataHeaderIn(ByteAddressBuffer StrataBuffer, inout FStrataAddressing StrataAddressing)
+FStrataPixelHeader UnpackStrataHeaderIn(FStrataMaterialContainer StrataBuffer, inout FStrataAddressing StrataAddressing)
 {
 	FStrataPixelHeader Out = InitialiseStrataPixelHeader();
 
@@ -2291,19 +2228,11 @@ FStrataPixelHeader UnpackStrataHeaderIn(ByteAddressBuffer StrataBuffer, inout FS
 	Out.StrataBuffer = StrataBuffer;
 	Out.SharedLocalBasesTypes = SharedLocalBasesTypes;
 
-#if STRATA_MEMACCESS == STRATA_MEMACCESS_UINT4
-	// Now load the packed local bases (stocked on the addressin part because it is only available while loading strata data)
-	for (uint i = 0; i < SharedLocalBasesCount; ++i)
-	{
-		StrataAddressing.PackedLocalBasis[i] = StrataLoadUint1(Out.StrataBuffer, StrataAddressing);
-	}
-#else
 	// Only keep the offset to the shared local bases memory. We are going to load them on demand.
-	Out.SharedLocalBasesByteOffset = StrataAddressing.CurrentAddr;
+	Out.SharedLocalBasesIndexOffset = StrataAddressing.CurrentIndex;
 
 	// Now skip over the shared local basis memory to be able to load the BSDF content
-	StrataAddressing.CurrentAddr += SharedLocalBasesCount * StrataAddressing.ByteOffsetToNextUint;
-#endif
+	StrataAddressing.CurrentIndex += SharedLocalBasesCount;
 
 	// And account for the space used by normal for the debug information
 	StrataAddressing.ReadBytes += SharedLocalBasesCount * STRATA_PACKED_SHAREDLOCALBASIS_STRIDE_BYTES;
@@ -2332,7 +2261,7 @@ void UnpackFastPathSlabBSDF(inout FStrataBSDF OutBSDF, uint2 PackedData01)
 	SLAB_ANISOTROPY(OutBSDF)	= Data1.z * 2.f - 1.f;
 }
 
-FStrataBSDF UnpackFastPathStrataBSDFIn(ByteAddressBuffer StrataBuffer, inout FStrataAddressing StrataAddressing)
+FStrataBSDF UnpackFastPathStrataBSDFIn(FStrataMaterialContainer StrataBuffer, inout FStrataAddressing StrataAddressing)
 {
 	FStrataBSDF OutBSDF = (FStrataBSDF)0; 
 
@@ -2368,7 +2297,7 @@ FStrataBSDF UnpackFastPathStrataBSDFIn(ByteAddressBuffer StrataBuffer, inout FSt
 // Unpack a single BSDF
 // Note: All BSDF lobes needs to be unpack in a sequential manner since each BSDF has a variable footprint.
 // If this is changed, please update the compiler side material size evaluation in StrataMaterial.cpp
-FStrataBSDF UnpackStrataBSDFIn(ByteAddressBuffer StrataBuffer, inout FStrataAddressing StrataAddressing)
+FStrataBSDF UnpackStrataBSDFIn(FStrataMaterialContainer StrataBuffer, inout FStrataAddressing StrataAddressing)
 {
 	FStrataBSDF OutBSDF = (FStrataBSDF)0;
 
diff --git a/Engine/Shaders/Private/Strata/StrataDeferredLighting.ush b/Engine/Shaders/Private/Strata/StrataDeferredLighting.ush
--- a/Engine/Shaders/Private/Strata/StrataDeferredLighting.ush
+++ b/Engine/Shaders/Private/Strata/StrataDeferredLighting.ush
@@ -18,7 +18,7 @@ float4 StrataDeferredLighting(
 	float LightMask,
 	FShadowTerms ShadowTerms,
 	FRectTexture RectTexture,
-	ByteAddressBuffer StrataDataBuffer,
+	FStrataMaterialContainer MaterialBuffer,
 	FStrataAddressing StrataAddressing,
 	FStrataPixelHeader StrataPixelHeader)
 {
@@ -32,9 +32,9 @@ float4 StrataDeferredLighting(
 	{
 		// Unpack BSDF data
 #if STRATA_FASTPATH
-		FStrataBSDF BSDF = UnpackFastPathStrataBSDFIn(StrataDataBuffer, StrataAddressing);
+		FStrataBSDF BSDF = UnpackFastPathStrataBSDFIn(MaterialBuffer, StrataAddressing);
 #else
-		FStrataBSDF BSDF = UnpackStrataBSDFIn(StrataDataBuffer, StrataAddressing);
+		FStrataBSDF BSDF = UnpackStrataBSDFIn(MaterialBuffer, StrataAddressing);
 #endif
 		FStrataBSDFContext BSDFContext = StrataCreateBSDFContext(StrataPixelHeader, BSDF, StrataAddressing, V, L);
 
diff --git a/Engine/Shaders/Private/Strata/StrataEnvironmentLighting.ush b/Engine/Shaders/Private/Strata/StrataEnvironmentLighting.ush
--- a/Engine/Shaders/Private/Strata/StrataEnvironmentLighting.ush
+++ b/Engine/Shaders/Private/Strata/StrataEnvironmentLighting.ush
@@ -18,7 +18,7 @@ float4 StrataEnvironmentLighting(
 	float3 V,
 	float AmbientOcclusion,
 	float TopLayerSpecularContributionFactor,
-	ByteAddressBuffer StrataDataBuffer,
+	FStrataMaterialContainer MaterialBuffer,
 	FStrataAddressing StrataAddressing,
 	FStrataPixelHeader StrataPixelHeader,
 	inout float SSRReductionFactor,
@@ -53,9 +53,9 @@ float4 StrataEnvironmentLighting(
 #endif
 	{
 	#if STRATA_FASTPATH
-		FStrataBSDF BSDF = UnpackFastPathStrataBSDFIn(StrataDataBuffer, StrataAddressing);
+		FStrataBSDF BSDF = UnpackFastPathStrataBSDFIn(MaterialBuffer, StrataAddressing);
 	#else
-		FStrataBSDF BSDF = UnpackStrataBSDFIn(StrataDataBuffer, StrataAddressing);
+		FStrataBSDF BSDF = UnpackStrataBSDFIn(MaterialBuffer, StrataAddressing);
 	#endif
 		const float3 BSDFThroughput = BSDF.FinalWeight;
 
diff --git a/Engine/Shaders/Private/Strata/StrataMaterialClassification.usf b/Engine/Shaders/Private/Strata/StrataMaterialClassification.usf
--- a/Engine/Shaders/Private/Strata/StrataMaterialClassification.usf
+++ b/Engine/Shaders/Private/Strata/StrataMaterialClassification.usf
@@ -180,7 +180,7 @@ int TileSize;
 int bRectPrimitive;
 int2 ViewResolution;
 uint MaxBytesPerPixel;
-ByteAddressBuffer MaterialBuffer;
+Texture2DArray<uint> MaterialTextureArray;
 RWBuffer<uint> SimpleTileIndirectDataBuffer;
 RWBuffer<uint> SimpleTileListDataBuffer;
 RWBuffer<uint> ComplexTileIndirectDataBuffer;
@@ -208,7 +208,7 @@ void TileMainCS(uint2 DispatchThreadId : SV_DispatchThreadID, uint LinearIndex :
 	if (all(PixelCoord < uint2(ViewResolution)))
 	{
 		FStrataAddressing StrataAddressing = GetStrataPixelDataByteOffset(PixelCoord, uint2(View.BufferSizeAndInvSize.xy), MaxBytesPerPixel);
-		FStrataPixelHeader StrataPixelHeader = UnpackStrataHeaderIn(MaterialBuffer, StrataAddressing);
+		FStrataPixelHeader StrataPixelHeader = UnpackStrataHeaderIn(MaterialTextureArray, StrataAddressing);
 
 		bContainsStrataMaterial  = StrataPixelHeader.BSDFCount > 0;
 		bContainsComplexMaterial = !IsSimpleMaterial(StrataPixelHeader);
@@ -348,7 +348,7 @@ void StencilMainPS(
 
 #if SHADER_CLEAR_MATERIAL_BUFFER
 
-RWByteAddressBuffer MaterialLobesBufferUAV;
+RWTexture2DArray<uint> MaterialTextureArrayUAV;
 RWTexture2D<uint> TopLayerTextureUAV;
 RWTexture2D<uint2> SSSTextureUAV;
 uint MaxBytesPerPixel;
@@ -357,7 +357,7 @@ int2 TiledViewBufferResolution;
 [numthreads(8, 8, 1)]
 void ClearMaterialBufferMainCS(uint2 DispatchThreadId : SV_DispatchThreadID)
 {
-	if (any(int2(DispatchThreadId.xy) > TiledViewBufferResolution))
+	if (any(int2(DispatchThreadId.xy) >= TiledViewBufferResolution))
 	{
 		return;
 	}
@@ -370,8 +370,8 @@ void ClearMaterialBufferMainCS(uint2 DispatchThreadId : SV_DispatchThreadID)
 	uint BSDFCount = 0;
 	uint PackedStrataHeader = PackStrataHeader(BSDFCount, StrataHeader);
 
-	FStrataAddressing StrataAddressing = GetStrataPixelDataByteOffset(DispatchThreadId.xy, TiledViewBufferResolution, MaxBytesPerPixel);
-	MaterialLobesBufferUAV.Store(StrataAddressing.CurrentAddr, PackedStrataHeader);
+	const uint FirstSlice = 0;
+	MaterialTextureArrayUAV[uint3(DispatchThreadId.xy, FirstSlice)] = PackedStrataHeader;
 
 	TopLayerTextureUAV[DispatchThreadId.xy] = 0;
 
@@ -380,3 +380,28 @@ void ClearMaterialBufferMainCS(uint2 DispatchThreadId : SV_DispatchThreadID)
 }
 
 #endif // SHADER_CLEAR_MATERIAL_BUFFER
+
+////////////////////////////////////////////////////////////////////////////////////////////////////////////
+
+#if SHADER_POST_BASEPASS
+
+Texture2DArray<uint> MaterialTextureArrayMRTs;
+RWTexture2DArray<uint> MaterialTextureArrayUAV;
+
+[numthreads(8, 8, 1)]
+void PostBasePassCS(uint2 DispatchThreadId : SV_DispatchThreadID)
+{
+	if (any(int2(DispatchThreadId.xy) >= View.ViewSizeAndInvSize.xy))
+	{
+		return;
+	}
+
+	int3 CoordData0 = int3(View.ViewRectMin.xy + DispatchThreadId.xy, 0);
+	int3 CoordData1 = int3(View.ViewRectMin.xy + DispatchThreadId.xy, 1);
+
+	MaterialTextureArrayUAV[CoordData0] = MaterialTextureArrayMRTs[CoordData0];
+	MaterialTextureArrayUAV[CoordData1] = MaterialTextureArrayMRTs[CoordData1];
+}
+
+#endif // SHADER_POST_BASEPASS
+
diff --git a/Engine/Shaders/Private/Strata/StrataVisualize.usf b/Engine/Shaders/Private/Strata/StrataVisualize.usf
--- a/Engine/Shaders/Private/Strata/StrataVisualize.usf
+++ b/Engine/Shaders/Private/Strata/StrataVisualize.usf
@@ -81,7 +81,7 @@ void VisualizeMaterialPS(
 	const int  InspectSquare = 10;
 
 	FStrataAddressing StrataAddressing = GetStrataPixelDataByteOffset(MaterialPos, uint2(View.BufferSizeAndInvSize.xy), Strata.MaxBytesPerPixel);
-	FStrataPixelHeader Header = UnpackStrataHeaderIn(Strata.MaterialLobesBuffer, StrataAddressing);
+	FStrataPixelHeader Header = UnpackStrataHeaderIn(Strata.MaterialTextureArray, StrataAddressing);
 	bool bStrataMaterial = Header.BSDFCount > 0;
 
 	const float2 ScreenPosition = SvPositionToScreenPosition(SVPos).xy;
@@ -128,11 +128,11 @@ void VisualizeMaterialPS(
 	{
 		FStrataAddressing StrataAddressing = GetStrataPixelDataByteOffset(PixelPos, uint2(View.BufferSizeAndInvSize.xy), Strata.MaxBytesPerPixel);
 
-		FStrataPixelHeader Header = UnpackStrataHeaderIn(Strata.MaterialLobesBuffer, StrataAddressing);
+		FStrataPixelHeader Header = UnpackStrataHeaderIn(Strata.MaterialTextureArray, StrataAddressing);
 
 		for (uint BSDFIndex = 0; BSDFIndex < Header.BSDFCount; ++BSDFIndex)
 		{
-			const FStrataBSDF BSDF = UnpackStrataBSDFIn(Strata.MaterialLobesBuffer, StrataAddressing);
+			const FStrataBSDF BSDF = UnpackStrataBSDFIn(Strata.MaterialTextureArray, StrataAddressing);
 
 			// Referential
 			if (all(PixelPos == MaterialPos))
@@ -170,7 +170,7 @@ void VisualizeMaterialPS(
 		FStrataAddressing StrataAddressing = GetStrataPixelDataByteOffset(MaterialPos, uint2(View.BufferSizeAndInvSize.xy), Strata.MaxBytesPerPixel);
 		const uint PixelStrataDataByteOffsetStart = StrataAddressing.ReadBytes;
 
-		FStrataPixelHeader Header = UnpackStrataHeaderIn(Strata.MaterialLobesBuffer, StrataAddressing);
+		FStrataPixelHeader Header = UnpackStrataHeaderIn(Strata.MaterialTextureArray, StrataAddressing);
 		const uint PixelStrataDataByteOffsetPostHeader = StrataAddressing.ReadBytes;
 
 		float3 BSDFFontColor = float3(0.25, 1.0, 0.25);
@@ -304,7 +304,7 @@ void VisualizeMaterialPS(
 		LOOP
 		for (uint BSDFIndex = 0; BSDFIndex < Header.BSDFCount; ++BSDFIndex)
 		{
-			const FStrataBSDF BSDF = UnpackStrataBSDFIn(Strata.MaterialLobesBuffer, StrataAddressing);
+			const FStrataBSDF BSDF = UnpackStrataBSDFIn(Strata.MaterialTextureArray, StrataAddressing);
 
 			TopLeft.x = TopLeftInit.x;
 
diff --git a/Engine/Source/Runtime/RenderCore/Public/RenderGraphUtils.h b/Engine/Source/Runtime/RenderCore/Public/RenderGraphUtils.h
--- a/Engine/Source/Runtime/RenderCore/Public/RenderGraphUtils.h
+++ b/Engine/Source/Runtime/RenderCore/Public/RenderGraphUtils.h
@@ -69,6 +69,34 @@ inline FRenderTargetBindingSlots GetRenderTargetBindings(ERenderTargetLoadAction
 	return BindingSlots;
 }
 
+struct FTextureRenderTargetBinding
+{
+	FRDGTextureRef Texture;
+	int16 ArraySlice;
+
+	FTextureRenderTargetBinding()
+		: Texture(nullptr)
+		, ArraySlice(-1)
+	{}
+
+	FTextureRenderTargetBinding(FRDGTextureRef InTexture, int16 InArraySlice = -1)
+		: Texture(InTexture)
+		, ArraySlice(InArraySlice)
+	{}
+};
+inline FRenderTargetBindingSlots GetRenderTargetBindings(ERenderTargetLoadAction ColorLoadAction, TArrayView<FTextureRenderTargetBinding> ColorTextures)
+{
+	check(ColorTextures.Num() <= MaxSimultaneousRenderTargets);
+
+	FRenderTargetBindingSlots BindingSlots;
+	for (int32 Index = 0, Count = ColorTextures.Num(); Index < Count; ++Index)
+	{
+		check(ColorTextures[Index].Texture);
+		BindingSlots[Index] = FRenderTargetBinding(ColorTextures[Index].Texture, ColorLoadAction, 0, ColorTextures[Index].ArraySlice);
+	}
+	return BindingSlots;
+}
+
 /**
  * Clears all render graph tracked resources that are not bound by a shader.
  * Excludes any resources on the ExcludeList from being cleared regardless of whether the 
diff --git a/Engine/Source/Runtime/Renderer/Private/BasePassRendering.cpp b/Engine/Source/Runtime/Renderer/Private/BasePassRendering.cpp
--- a/Engine/Source/Runtime/Renderer/Private/BasePassRendering.cpp
+++ b/Engine/Source/Runtime/Renderer/Private/BasePassRendering.cpp
@@ -748,9 +748,10 @@ void FDeferredShadingSceneRenderer::RenderBasePass(
 
 	const FExclusiveDepthStencil ExclusiveDepthStencil(BasePassDepthStencilAccess);
 
-	TStaticArray<FRDGTextureRef, MaxSimultaneousRenderTargets> BasePassTextures;
+	TStaticArray<FTextureRenderTargetBinding, MaxSimultaneousRenderTargets> BasePassTextures;
 	uint32 BasePassTextureCount = SceneTextures.GetGBufferRenderTargets(BasePassTextures);
-	TArrayView<FRDGTextureRef> BasePassTexturesView = MakeArrayView(BasePassTextures.GetData(), BasePassTextureCount);
+	Strata::ApprendStrataMRTs(*this, BasePassTextureCount, BasePassTextures);
+	TArrayView<FTextureRenderTargetBinding> BasePassTexturesView = MakeArrayView(BasePassTextures.GetData(), BasePassTextureCount);
 	FRDGTextureRef BasePassDepthTexture = SceneTextures.Depth.Target;
 	FLinearColor SceneColorClearValue;
 
@@ -901,6 +902,8 @@ void FDeferredShadingSceneRenderer::RenderBasePass(
 		StampDeferredDebugProbeMaterialPS(GraphBuilder, Views, BasePassRenderTargets, SceneTextures);
 	}
 #endif
+
+	Strata::PostBasePass(GraphBuilder, Views);
 }
 
 BEGIN_SHADER_PARAMETER_STRUCT(FOpaqueBasePassParameters, )
diff --git a/Engine/Source/Runtime/Renderer/Private/DebugProbeRendering.cpp b/Engine/Source/Runtime/Renderer/Private/DebugProbeRendering.cpp
--- a/Engine/Source/Runtime/Renderer/Private/DebugProbeRendering.cpp
+++ b/Engine/Source/Runtime/Renderer/Private/DebugProbeRendering.cpp
@@ -41,7 +41,7 @@ class FStampDeferredDebugProbePS : public FGlobalShader
 
 	BEGIN_SHADER_PARAMETER_STRUCT(FParameters, )
 		SHADER_PARAMETER_STRUCT_REF(FViewUniformShaderParameters, ViewUniformBuffer)
-		SHADER_PARAMETER_RDG_BUFFER_UAV(RWByteAddressBuffer, MaterialLobesBufferUAV)
+		SHADER_PARAMETER_RDG_TEXTURE_UAV(RWTexture2DArray<uint>, MaterialTextureArrayUAV)
 		SHADER_PARAMETER(uint32, MaxBytesPerPixel)
 		SHADER_PARAMETER_STRUCT_INCLUDE(FSceneTextureShaderParameters, SceneTextures)
 		SHADER_PARAMETER(int32, DebugProbesMode)
@@ -76,7 +76,7 @@ static void CommonStampDeferredDebugProbeDrawCall(
 	int32 RenderPass)
 {
 	PassParameters->ViewUniformBuffer = View.ViewUniformBuffer;
-	PassParameters->MaterialLobesBufferUAV = View.StrataSceneData->MaterialLobesBufferUAV;
+	PassParameters->MaterialTextureArrayUAV = View.StrataSceneData->MaterialTextureArrayUAV;
 	PassParameters->MaxBytesPerPixel = View.StrataSceneData->MaxBytesPerPixel;
 	PassParameters->DebugProbesMode = View.Family->EngineShowFlags.VisualizeLightingOnProbes ? 3 : FMath::Clamp(CVarVisualizeLightingOnProbes.GetValueOnRenderThread(), 0, 3);
 		
diff --git a/Engine/Source/Runtime/Renderer/Private/SceneTextures.cpp b/Engine/Source/Runtime/Renderer/Private/SceneTextures.cpp
--- a/Engine/Source/Runtime/Renderer/Private/SceneTextures.cpp
+++ b/Engine/Source/Runtime/Renderer/Private/SceneTextures.cpp
@@ -19,6 +19,7 @@
 #include "PostProcess/PostProcessAmbientOcclusionMobile.h"
 #include "PostProcess/PostProcessPixelProjectedReflectionMobile.h"
 #include "IHeadMountedDisplayModule.h"
+#include "Strata/Strata.h"
 
 static TAutoConsoleVariable<int32> CVarSceneTargetsResizeMethod(
 	TEXT("r.SceneRenderTargetResizeMethod"),
@@ -852,12 +853,12 @@ const FSceneTextures& FSceneTextures::Get(FRDGBuilder& GraphBuilder)
 	return *SceneTextures;
 }
 
-uint32 FSceneTextures::GetGBufferRenderTargets(TStaticArray<FRDGTextureRef, MaxSimultaneousRenderTargets>& RenderTargets) const
+uint32 FSceneTextures::GetGBufferRenderTargets(TStaticArray<FTextureRenderTargetBinding, MaxSimultaneousRenderTargets>& RenderTargets) const
 {
 	uint32 RenderTargetCount = 0;
 
 	// All configurations use scene color in the first slot.
-	RenderTargets[RenderTargetCount++] = Color.Target;
+	RenderTargets[RenderTargetCount++] = FTextureRenderTargetBinding(Color.Target);
 
 	if (Config.bIsUsingGBuffers)
 	{
@@ -889,7 +890,7 @@ uint32 FSceneTextures::GetGBufferRenderTargets(TStaticArray<FRDGTextureRef, MaxS
 			checkf(Entry.Index <= 0 || Entry.Texture != nullptr, TEXT("Texture '%s' was requested by FGBufferInfo, but it is null."), Entry.Name);
 			if (Entry.Index > 0)
 			{
-				RenderTargets[Entry.Index] = Entry.Texture;
+				RenderTargets[Entry.Index] = FTextureRenderTargetBinding(Entry.Texture);
 				RenderTargetCount = FMath::Max(RenderTargetCount, uint32(Entry.Index + 1));
 			}
 		}
@@ -897,7 +898,7 @@ uint32 FSceneTextures::GetGBufferRenderTargets(TStaticArray<FRDGTextureRef, MaxS
 	// Forward shading path. Simple forward shading does not use velocity.
 	else if (IsUsingBasePassVelocity(Config.ShaderPlatform) && !IsSimpleForwardShadingEnabled(Config.ShaderPlatform))
 	{
-		RenderTargets[RenderTargetCount++] = Velocity;
+		RenderTargets[RenderTargetCount++] = FTextureRenderTargetBinding(Velocity);
 	}
 
 	return RenderTargetCount;
@@ -905,11 +906,11 @@ uint32 FSceneTextures::GetGBufferRenderTargets(TStaticArray<FRDGTextureRef, MaxS
 
 uint32 FSceneTextures::GetGBufferRenderTargets(ERenderTargetLoadAction LoadAction, FRenderTargetBindingSlots& RenderTargetBindingSlots) const
 {
-	TStaticArray<FRDGTextureRef, MaxSimultaneousRenderTargets> RenderTargets;
+	TStaticArray<FTextureRenderTargetBinding, MaxSimultaneousRenderTargets> RenderTargets;
 	const uint32 RenderTargetCount = GetGBufferRenderTargets(RenderTargets);
 	for (uint32 Index = 0; Index < RenderTargetCount; ++Index)
 	{
-		RenderTargetBindingSlots[Index] = FRenderTargetBinding(RenderTargets[Index], LoadAction);
+		RenderTargetBindingSlots[Index] = FRenderTargetBinding(RenderTargets[Index].Texture, LoadAction);
 	}
 	return RenderTargetCount;
 }
diff --git a/Engine/Source/Runtime/Renderer/Private/SceneTextures.h b/Engine/Source/Runtime/Renderer/Private/SceneTextures.h
--- a/Engine/Source/Runtime/Renderer/Private/SceneTextures.h
+++ b/Engine/Source/Runtime/Renderer/Private/SceneTextures.h
@@ -147,7 +147,7 @@ struct RENDERER_API FSceneTextures : public FMinimalSceneTextures
 	static const FSceneTextures& Get(FRDGBuilder& GraphBuilder);
 
 	// Configures an array of render targets for the GBuffer pass.
-	uint32 GetGBufferRenderTargets(TStaticArray<FRDGTextureRef, MaxSimultaneousRenderTargets>& RenderTargets) const;
+	uint32 GetGBufferRenderTargets(TStaticArray<FTextureRenderTargetBinding, MaxSimultaneousRenderTargets>& RenderTargets) const;
 	uint32 GetGBufferRenderTargets(ERenderTargetLoadAction LoadAction, FRenderTargetBindingSlots& RenderTargets) const;
 
 	// (Deferred) Texture containing conservative downsampled depth for occlusion.
diff --git a/Engine/Source/Runtime/Renderer/Private/Strata/Strata.cpp b/Engine/Source/Runtime/Renderer/Private/Strata/Strata.cpp
--- a/Engine/Source/Runtime/Renderer/Private/Strata/Strata.cpp
+++ b/Engine/Source/Runtime/Renderer/Private/Strata/Strata.cpp
@@ -79,9 +79,9 @@ void FStrataSceneData::Reset()
 	TopLayerTextureUAV = nullptr;
 	SSSTextureUAV = nullptr;
 
-	MaterialLobesBuffer = nullptr;
-	MaterialLobesBufferUAV = nullptr;
-	MaterialLobesBufferSRV = nullptr;
+	MaterialTextureArray = nullptr;
+	MaterialTextureArrayUAV = nullptr;
+	MaterialTextureArraySRV = nullptr;
 
 	for (uint32 i = 0; i < EStrataTileMaterialType::ECount; ++i)
 	{
@@ -114,7 +114,7 @@ namespace Strata
 // Forward declaration
 static void AddStrataClearMaterialBufferPass(
 	FRDGBuilder& GraphBuilder, 
-	FRDGBufferUAVRef MaterialLobesBufferUAV,
+	FRDGTextureUAVRef MaterialTextureArrayUAV,
 	FRDGTextureUAVRef TopLayerTextureUAV,
 	FRDGTextureUAVRef SSSTextureUAV,
 	uint32 MaxBytesPerPixel, 
@@ -189,7 +189,7 @@ void InitialiseStrataFrameSceneData(FSceneRenderer& SceneRenderer, FRDGBuilder&
 
 		// Top layer texture
 		{
-			StrataSceneData.TopLayerTexture = GraphBuilder.CreateTexture(FRDGTextureDesc::Create2D(SceneTextureExtent, PF_R32_UINT, FClearValueBinding::Black, TexCreate_ShaderResource | TexCreate_UAV), TEXT("Strata.TopLayerTexture"));
+			StrataSceneData.TopLayerTexture = GraphBuilder.CreateTexture(FRDGTextureDesc::Create2D(SceneTextureExtent, PF_R32_UINT, FClearValueBinding::Black, TexCreate_RenderTargetable | TexCreate_ShaderResource | TexCreate_UAV), TEXT("Strata.TopLayerTexture"));
 			StrataSceneData.TopLayerTextureUAV = GraphBuilder.CreateUAV(StrataSceneData.TopLayerTexture);
 		}
 
@@ -204,11 +204,14 @@ void InitialiseStrataFrameSceneData(FSceneRenderer& SceneRenderer, FRDGBuilder&
 		StrataSceneData.MaxBytesPerPixel = 4u;
 	}
 
-	// create the material lob buffer for all views
-	const uint32 MaterialLobesBufferByteSize = FMath::Max(4u, MaterialBufferSizeXY.X * MaterialBufferSizeXY.Y * StrataSceneData.MaxBytesPerPixel);
-	StrataSceneData.MaterialLobesBuffer = GraphBuilder.CreateBuffer(FRDGBufferDesc::CreateByteAddressDesc(MaterialLobesBufferByteSize), TEXT("Strata.MaterialBuffer"));
-	StrataSceneData.MaterialLobesBufferSRV = GraphBuilder.CreateSRV(StrataSceneData.MaterialLobesBuffer);
-	StrataSceneData.MaterialLobesBufferUAV = GraphBuilder.CreateUAV(StrataSceneData.MaterialLobesBuffer);
+	// Create the material data container
+	FIntPoint SceneTextureExtent = IsStrataEnabled() ? GetSceneTextureExtent() : FIntPoint(2, 2);
+	const FRDGTextureDesc MaterialTextureDesc = FRHITextureCreateInfo::Create2DArray(SceneTextureExtent, PF_R32_UINT, FClearValueBinding::Transparent, TexCreate_TargetArraySlicesIndependently | TexCreate_RenderTargetable | TexCreate_ShaderResource | TexCreate_UAV, StrataSceneData.MaxBytesPerPixel / 4, 1, 1);
+	const FRDGTextureDesc MaterialTextureDescMRTs = FRHITextureCreateInfo::Create2DArray(SceneTextureExtent, PF_R32_UINT, FClearValueBinding::Transparent, TexCreate_TargetArraySlicesIndependently | TexCreate_RenderTargetable | TexCreate_ShaderResource | TexCreate_UAV, 2, 1, 1);
+	StrataSceneData.MaterialTextureArrayMRTs = GraphBuilder.CreateTexture(MaterialTextureDescMRTs, TEXT("Strata.MaterialRT"));
+	StrataSceneData.MaterialTextureArray = GraphBuilder.CreateTexture(MaterialTextureDesc, TEXT("Strata.Material"));
+	StrataSceneData.MaterialTextureArraySRV = GraphBuilder.CreateSRV(FRDGTextureSRVDesc::Create(StrataSceneData.MaterialTextureArray));
+	StrataSceneData.MaterialTextureArrayUAV = GraphBuilder.CreateUAV(FRDGTextureUAVDesc(StrataSceneData.MaterialTextureArray, 0));
 
 	// Rough diffuse model
 	StrataSceneData.bRoughDiffuse = CVarStrataRoughDiffuse.GetValueOnRenderThread() > 0 ? 1u : 0u;
@@ -224,7 +227,7 @@ void InitialiseStrataFrameSceneData(FSceneRenderer& SceneRenderer, FRDGBuilder&
 	{
 		AddStrataClearMaterialBufferPass(
 			GraphBuilder, 
-			StrataSceneData.MaterialLobesBufferUAV,
+			GraphBuilder.CreateUAV(FRDGTextureUAVDesc(StrataSceneData.MaterialTextureArrayMRTs, 0)),
 			StrataSceneData.TopLayerTextureUAV,
 			StrataSceneData.SSSTextureUAV,
 			StrataSceneData.MaxBytesPerPixel, 
@@ -246,21 +249,22 @@ void BindStrataBasePassUniformParameters(FRDGBuilder& GraphBuilder, FStrataScene
 	{
 		OutStrataUniformParameters.bRoughDiffuse = StrataSceneData->bRoughDiffuse ? 1u : 0u;
 		OutStrataUniformParameters.MaxBytesPerPixel = StrataSceneData->MaxBytesPerPixel;
-		OutStrataUniformParameters.MaterialLobesBufferUAV = StrataSceneData->MaterialLobesBufferUAV;
-
-		OutStrataUniformParameters.TopLayerTextureUAV = StrataSceneData->TopLayerTextureUAV;
+		OutStrataUniformParameters.MaterialTextureArrayUAV = StrataSceneData->MaterialTextureArrayUAV;
 		OutStrataUniformParameters.SSSTextureUAV = StrataSceneData->SSSTextureUAV;
 	}
 	else
 	{
-		FRDGTextureRef DummyWritableTexture = GraphBuilder.CreateTexture(FRDGTextureDesc::Create2D(FIntPoint(1,1), PF_R32_UINT, FClearValueBinding::None, TexCreate_ShaderResource | TexCreate_UAV), TEXT("Strata.DummyWritableTexture"));
-		FRDGTextureUAVRef DummyWritableUAV = GraphBuilder.CreateUAV(DummyWritableTexture);
+		FRDGTextureRef DummyWritableTexture = GraphBuilder.CreateTexture(FRDGTextureDesc::Create2D(FIntPoint(1, 1), PF_R32_UINT, FClearValueBinding::None, TexCreate_ShaderResource | TexCreate_UAV), TEXT("Strata.DummyWritableTexture"));
+		FRDGTextureUAVRef DummyWritableTextureUAV = GraphBuilder.CreateUAV(FRDGTextureUAVDesc(DummyWritableTexture));
+
+		FRDGTextureRef DummyWritableTextureArray = GraphBuilder.CreateTexture(FRDGTextureDesc::Create2DArray(FIntPoint(1, 1), PF_R32_UINT, FClearValueBinding::None, TexCreate_ShaderResource | TexCreate_UAV, 1), TEXT("Strata.DummyWritableTexture"));
+		FRDGTextureUAVRef DummyWritableTextureArrayUAV = GraphBuilder.CreateUAV(FRDGTextureUAVDesc(DummyWritableTextureArray));
+
 		const FRDGSystemTextures& SystemTextures = FRDGSystemTextures::Get(GraphBuilder);
 		OutStrataUniformParameters.bRoughDiffuse = 0u;
 		OutStrataUniformParameters.MaxBytesPerPixel = 0;
-		OutStrataUniformParameters.MaterialLobesBufferUAV = GraphBuilder.CreateUAV(GraphBuilder.RegisterExternalBuffer(GWhiteVertexBufferWithRDG->Buffer), PF_R32_UINT);
-		OutStrataUniformParameters.TopLayerTextureUAV = DummyWritableUAV;
-		OutStrataUniformParameters.SSSTextureUAV = DummyWritableUAV;
+		OutStrataUniformParameters.MaterialTextureArrayUAV = DummyWritableTextureArrayUAV;
+		OutStrataUniformParameters.SSSTextureUAV = DummyWritableTextureUAV;
 	}
 }
 
@@ -270,16 +274,17 @@ void BindStrataGlobalUniformParameters(FRDGBuilder& GraphBuilder, FStrataSceneDa
 	{
 		OutStrataUniformParameters.bRoughDiffuse = StrataSceneData->bRoughDiffuse ? 1u : 0u;
 		OutStrataUniformParameters.MaxBytesPerPixel = StrataSceneData->MaxBytesPerPixel;
-		OutStrataUniformParameters.MaterialLobesBuffer = StrataSceneData->MaterialLobesBufferSRV;
+		OutStrataUniformParameters.MaterialTextureArray = StrataSceneData->MaterialTextureArraySRV;
 		OutStrataUniformParameters.TopLayerTexture = StrataSceneData->TopLayerTexture;
 		OutStrataUniformParameters.SSSTexture = StrataSceneData->SSSTexture;
 	}
 	else
 	{
 		const FRDGSystemTextures& SystemTextures = FRDGSystemTextures::Get(GraphBuilder);
+		FRDGTextureRef DefaultTextureArray = GSystemTextures.GetDefaultTexture(GraphBuilder, ETextureDimension::Texture2DArray, EPixelFormat::PF_R32_UINT, FClearValueBinding::Transparent);
 		OutStrataUniformParameters.bRoughDiffuse = 0;
 		OutStrataUniformParameters.MaxBytesPerPixel = 0;
-		OutStrataUniformParameters.MaterialLobesBuffer = GraphBuilder.CreateSRV(GraphBuilder.RegisterExternalBuffer(GWhiteVertexBufferWithRDG->Buffer), PF_R32_UINT);
+		OutStrataUniformParameters.MaterialTextureArray = GraphBuilder.CreateSRV(FRDGTextureSRVDesc::Create(DefaultTextureArray));
 		OutStrataUniformParameters.TopLayerTexture = SystemTextures.DefaultNormal8Bit;
 		OutStrataUniformParameters.SSSTexture = SystemTextures.Black;
 	}
@@ -380,7 +385,7 @@ class FStrataClearMaterialBufferCS : public FGlobalShader
 	using FPermutationDomain = TShaderPermutationDomain<>;
 
 	BEGIN_SHADER_PARAMETER_STRUCT(FParameters, )
-		SHADER_PARAMETER_RDG_BUFFER_UAV(RWByteAddressBuffer, MaterialLobesBufferUAV)
+		SHADER_PARAMETER_RDG_TEXTURE_UAV(RWTexture2DArray<uint>, MaterialTextureArrayUAV)
 		SHADER_PARAMETER_RDG_TEXTURE_UAV(RWTexture2D<uint>, TopLayerTextureUAV)
 		SHADER_PARAMETER_RDG_TEXTURE_UAV(RWTexture2D<uint2>, SSSTextureUAV)
 		SHADER_PARAMETER(uint32, MaxBytesPerPixel)
@@ -416,7 +421,7 @@ class FStrataMaterialTileClassificationPassCS : public FGlobalShader
 		SHADER_PARAMETER(int32, bRectPrimitive)
 		SHADER_PARAMETER(FIntPoint, ViewResolution)
 		SHADER_PARAMETER(uint32, MaxBytesPerPixel)
-		SHADER_PARAMETER_RDG_BUFFER_SRV(ByteAddressBuffer, MaterialBuffer)
+		SHADER_PARAMETER_RDG_TEXTURE_SRV(Texture2DArray<uint>, MaterialTextureArray)
 		SHADER_PARAMETER_RDG_BUFFER_UAV(RWBuffer, SimpleTileIndirectDataBuffer)
 		SHADER_PARAMETER_RDG_BUFFER_UAV(RWBuffer, SimpleTileListDataBuffer)
 		SHADER_PARAMETER_RDG_BUFFER_UAV(RWBuffer, ComplexTileIndirectDataBuffer)
@@ -621,6 +626,83 @@ void AddStrataStencilPass(
 
 ////////////////////////////////////////////////////////////////////////////////////////////////////////////////////////////////////////////////
 
+void ApprendStrataMRTs(FSceneRenderer& SceneRenderer, uint32& RenderTargetCount, TStaticArray<FTextureRenderTargetBinding, MaxSimultaneousRenderTargets>& RenderTargets)
+{
+	if (Strata::IsStrataEnabled() && SceneRenderer.Scene)
+	{
+		// Add 2 uint for Strata fast path
+		auto AddStrataOutputTarget = [&](int16 StrataMaterialArraySlice)
+		{
+			RenderTargets[RenderTargetCount] = FTextureRenderTargetBinding(SceneRenderer.Scene->StrataSceneData.MaterialTextureArrayMRTs, StrataMaterialArraySlice);
+			RenderTargetCount++;
+		};
+		AddStrataOutputTarget(0);
+		AddStrataOutputTarget(1);
+
+		// Add another MRT for Strata top layer information
+		{
+			RenderTargets[RenderTargetCount] = FTextureRenderTargetBinding(SceneRenderer.Scene->StrataSceneData.TopLayerTexture);
+			RenderTargetCount++;
+		};
+	}
+}
+
+////////////////////////////////////////////////////////////////////////////////////////////////////////////////////////////////////////////////
+
+class FStrataPostBasePassCS : public FGlobalShader
+{
+	DECLARE_GLOBAL_SHADER(FStrataPostBasePassCS);
+	SHADER_USE_PARAMETER_STRUCT(FStrataPostBasePassCS, FGlobalShader);
+
+	BEGIN_SHADER_PARAMETER_STRUCT(FParameters, )
+		SHADER_PARAMETER_STRUCT_REF(FViewUniformShaderParameters, ViewUniformBuffer)
+		SHADER_PARAMETER_RDG_TEXTURE_SRV(Texture2DArray<uint>, MaterialTextureArrayMRTs)
+		SHADER_PARAMETER_RDG_TEXTURE_UAV(RWTexture2DArray<uint>, MaterialTextureArrayUAV)
+	END_SHADER_PARAMETER_STRUCT()
+
+	static bool ShouldCompilePermutation(const FGlobalShaderPermutationParameters& Parameters)
+	{
+		return GetMaxSupportedFeatureLevel(Parameters.Platform) >= ERHIFeatureLevel::SM5 && Strata::IsStrataEnabled();
+	}
+
+	static void ModifyCompilationEnvironment(const FGlobalShaderPermutationParameters& Parameters, FShaderCompilerEnvironment& OutEnvironment)
+	{
+		FGlobalShader::ModifyCompilationEnvironment(Parameters, OutEnvironment);
+		OutEnvironment.SetDefine(TEXT("SHADER_POST_BASEPASS"), 1);
+	}
+};
+IMPLEMENT_GLOBAL_SHADER(FStrataPostBasePassCS, "/Engine/Private/Strata/StrataMaterialClassification.usf", "PostBasePassCS", SF_Compute);
+
+void PostBasePass(FRDGBuilder& GraphBuilder, const TArray<FViewInfo>& Views)
+{
+	RDG_EVENT_SCOPE_CONDITIONAL(GraphBuilder, IsStrataEnabled() && Views.Num() > 0, "StrataPostBasePass");
+	if (!IsStrataEnabled())
+	{
+		return;
+	}
+
+	for (int32 i = 0; i < Views.Num(); ++i)
+	{
+		const FViewInfo& View = Views[i];
+
+		TShaderMapRef<FStrataPostBasePassCS> ComputeShader(View.ShaderMap);
+		FStrataPostBasePassCS::FParameters* PassParameters = GraphBuilder.AllocParameters<FStrataPostBasePassCS::FParameters>();
+		PassParameters->ViewUniformBuffer = View.ViewUniformBuffer;
+		PassParameters->MaterialTextureArrayMRTs = GraphBuilder.CreateSRV(FRDGTextureSRVDesc::Create(View.StrataSceneData->MaterialTextureArrayMRTs));
+		PassParameters->MaterialTextureArrayUAV = View.StrataSceneData->MaterialTextureArrayUAV;
+
+		const uint32 GroupSize = 8;
+		FComputeShaderUtils::AddPass(
+			GraphBuilder,
+			RDG_EVENT_NAME("Strata::PostBasePassCS"),
+			ComputeShader,
+			PassParameters,
+			FComputeShaderUtils::GetGroupCount(View.ViewRect.Size(), GroupSize));
+	}
+}
+
+////////////////////////////////////////////////////////////////////////////////////////////////////////////////////////////////////////////////
+
 void AddStrataMaterialClassificationPass(FRDGBuilder& GraphBuilder, const FMinimalSceneTextures& SceneTextures, const TArray<FViewInfo>& Views)
 {
 	RDG_EVENT_SCOPE_CONDITIONAL(GraphBuilder, IsStrataEnabled() && Views.Num() > 0, "StrataMaterialClassification");
@@ -651,7 +733,7 @@ void AddStrataMaterialClassificationPass(FRDGBuilder& GraphBuilder, const FMinim
 			PassParameters->bRectPrimitive = GRHISupportsRectTopology ? 1 : 0;
 			PassParameters->ViewResolution = View.ViewRect.Size();
 			PassParameters->MaxBytesPerPixel = View.StrataSceneData->MaxBytesPerPixel;
-			PassParameters->MaterialBuffer = View.StrataSceneData->MaterialLobesBufferSRV;
+			PassParameters->MaterialTextureArray = View.StrataSceneData->MaterialTextureArraySRV;
 			PassParameters->SimpleTileListDataBuffer = View.StrataSceneData->ClassificationTileListBufferUAV[EStrataTileMaterialType::ESimple];
 			PassParameters->SimpleTileIndirectDataBuffer = View.StrataSceneData->ClassificationTileIndirectBufferUAV[EStrataTileMaterialType::ESimple];
 			PassParameters->ComplexTileListDataBuffer = View.StrataSceneData->ClassificationTileListBufferUAV[EStrataTileMaterialType::EComplex];
@@ -670,7 +752,7 @@ void AddStrataMaterialClassificationPass(FRDGBuilder& GraphBuilder, const FMinim
 
 static void AddStrataClearMaterialBufferPass(
 	FRDGBuilder& GraphBuilder,
-	FRDGBufferUAVRef MaterialLobesBufferUAV,
+	FRDGTextureUAVRef MaterialTextureArrayUAV,
 	FRDGTextureUAVRef TopLayerTextureUAV,
 	FRDGTextureUAVRef SSSTextureUAV,
 	uint32 MaxBytesPerPixel,
@@ -678,7 +760,7 @@ static void AddStrataClearMaterialBufferPass(
 {
 	TShaderMapRef<FStrataClearMaterialBufferCS> ComputeShader(GetGlobalShaderMap(GMaxRHIFeatureLevel));
 	FStrataClearMaterialBufferCS::FParameters* PassParameters = GraphBuilder.AllocParameters<FStrataClearMaterialBufferCS::FParameters>();
-	PassParameters->MaterialLobesBufferUAV = MaterialLobesBufferUAV;
+	PassParameters->MaterialTextureArrayUAV = MaterialTextureArrayUAV;
 	PassParameters->TopLayerTextureUAV = TopLayerTextureUAV;
 	PassParameters->SSSTextureUAV = SSSTextureUAV;
 	PassParameters->MaxBytesPerPixel = MaxBytesPerPixel;
diff --git a/Engine/Source/Runtime/Renderer/Private/Strata/Strata.h b/Engine/Source/Runtime/Renderer/Private/Strata/Strata.h
--- a/Engine/Source/Runtime/Renderer/Private/Strata/Strata.h
+++ b/Engine/Source/Runtime/Renderer/Private/Strata/Strata.h
@@ -18,15 +18,14 @@ struct FScreenPassTexture;
 BEGIN_SHADER_PARAMETER_STRUCT(FStrataBasePassUniformParameters, )
 	SHADER_PARAMETER(uint32, MaxBytesPerPixel)
 	SHADER_PARAMETER(uint32, bRoughDiffuse)
-	SHADER_PARAMETER_RDG_BUFFER_UAV(RWByteAddressBuffer, MaterialLobesBufferUAV)
-	SHADER_PARAMETER_RDG_TEXTURE_UAV(RWTexture2D<uint>, TopLayerTextureUAV)
+	SHADER_PARAMETER_RDG_TEXTURE_UAV(RWTexture2DArray<uint>, MaterialTextureArrayUAV)
 	SHADER_PARAMETER_RDG_TEXTURE_UAV(RWTexture2D<uint2>, SSSTextureUAV)
 END_SHADER_PARAMETER_STRUCT()
 
 BEGIN_GLOBAL_SHADER_PARAMETER_STRUCT(FStrataGlobalUniformParameters, )
 	SHADER_PARAMETER(uint32, MaxBytesPerPixel)
 	SHADER_PARAMETER(uint32, bRoughDiffuse)
-	SHADER_PARAMETER_RDG_BUFFER_SRV(ByteAddressBuffer, MaterialLobesBuffer)
+	SHADER_PARAMETER_RDG_TEXTURE_SRV(Texture2DArray<uint>, MaterialTextureArray)
 	SHADER_PARAMETER_RDG_TEXTURE(Texture2D<uint>, TopLayerTexture)
 	SHADER_PARAMETER_RDG_TEXTURE(Texture2D<uint2>, SSSTexture)
 END_GLOBAL_SHADER_PARAMETER_STRUCT()
@@ -47,9 +46,13 @@ struct FStrataSceneData
 
 	// Resources allocated and updated each frame
 
-	FRDGBufferRef MaterialLobesBuffer;
-	FRDGBufferUAVRef MaterialLobesBufferUAV;
-	FRDGBufferSRVRef MaterialLobesBufferSRV;
+
+	FRDGTextureRef MaterialTextureArrayMRTs;
+
+	FRDGTextureRef MaterialTextureArray;
+	FRDGTextureUAVRef MaterialTextureArrayUAV;
+	FRDGTextureSRVRef MaterialTextureArraySRV;
+
 	FRDGBufferRef ClassificationTileListBuffer[EStrataTileMaterialType::ECount];
 	FRDGBufferUAVRef ClassificationTileListBufferUAV[EStrataTileMaterialType::ECount];
 	FRDGBufferSRVRef ClassificationTileListBufferSRV[EStrataTileMaterialType::ECount];
@@ -87,6 +90,9 @@ void InitialiseStrataFrameSceneData(FSceneRenderer& SceneRenderer, FRDGBuilder&
 void BindStrataBasePassUniformParameters(FRDGBuilder& GraphBuilder, FStrataSceneData* StrataSceneData, FStrataBasePassUniformParameters& OutStrataUniformParameters);
 void BindStrataGlobalUniformParameters(FRDGBuilder& GraphBuilder, FStrataSceneData* StrataSceneData, FStrataGlobalUniformParameters& OutStrataUniformParameters);
 
+void ApprendStrataMRTs(FSceneRenderer& SceneRenderer, uint32& BasePassTextureCount, TStaticArray<FTextureRenderTargetBinding, MaxSimultaneousRenderTargets>& BasePassTextures);
+void PostBasePass(FRDGBuilder& GraphBuilder, const TArray<FViewInfo>& Views);
+
 TRDGUniformBufferRef<FStrataGlobalUniformParameters> BindStrataGlobalUniformParameters(FStrataSceneData* StrataSceneData);
 
 void AddStrataMaterialClassificationPass(FRDGBuilder& GraphBuilder, const FMinimalSceneTextures& SceneTextures, const TArray<FViewInfo>& Views);
