Strata - removed no longer needed CATEGORIZATION shader

diff --git a/Engine/Shaders/Private/Strata/StrataMaterialClassification.usf b/Engine/Shaders/Private/Strata/StrataMaterialClassification.usf
--- a/Engine/Shaders/Private/Strata/StrataMaterialClassification.usf
+++ b/Engine/Shaders/Private/Strata/StrataMaterialClassification.usf
@@ -9,160 +9,6 @@
 #define TILE_SIZE 8
 #define GROUP_THREAD_COUNT (TILE_SIZE*TILE_SIZE)
 
-#if SHADER_CATEGORIZATION
-
-uint MaxBytesPerPixel;
-ByteAddressBuffer MaterialLobesBuffer;
-
-// Note: 
-// This code is partial copy of PackStrataOut in Strata.ush. 
-// The command part will progressively be removed as we remove data from the gbuffer.
-void ComputePostProcessBSDFData(
-	const FStrataPixelHeader Header,
-	const FStrataBSDF BSDF,
-	const FStrataAddressing StrataAddressing,
-	float3 V,
-	inout FStrataClassification Classification,
-	inout FStrataSubsurfaceData SSS,
-	inout FStrataTopLayerData TopLayer,
-	inout float TopLayerRoughness, 
-	inout float TopLayerRoughnessAverageCount)
-{
-	const float OpaqueBSDFThroughput = 0.0f;
-	const float FullThroughput = 1.0f;
-	const float FullyRough = 1.0f;
-
-	// FinalWeight does not always contains coverage (see Strata.ush)
-	const float  BSDFWeightAvg = dot(BSDF.FinalWeight, 1.0f/3.0f);
-
-	// STRATA_TODO apply the same normal transform as for the material
-	float3x3 TangentBasis = StrataGetBSDFSharedBasis(Header, BSDF, StrataAddressing);
-	float3 N = TangentBasis[2];
-	//N = normalize(N);
-
-	const bool bTopLayer = BSDF_GETISTOPLAYER(BSDF);
-	if (bTopLayer)
-	{
-		// We must take into account the weight in case horizontal blending is used,
-		// together with a threshold to avoid 0-length normal after normalization in case influence is fadeout using the Weight node.
-		TopLayer.WorldNormal += N * max(1e-5, BSDFWeightAvg); 
-	}
-
-	AddShadingModels(Classification, StrataShadingModelBit(BSDF_GETTYPE(BSDF)));
-
-	switch (BSDF_GETTYPE(BSDF))
-	{
-	case STRATA_BSDF_TYPE_SLAB:
-	{
-		AndIsSimpleMaterial(Classification, IsStrataSlabFastPathCompatible(BSDF));
-
-		if (bTopLayer)
-		{
-			// Setup SSR for the top layer
-			TopLayerRoughness += BSDFWeightAvg * SLAB_ROUGHNESS(BSDF);
-			TopLayerRoughnessAverageCount += BSDFWeightAvg;
-		}
-
-		// Select the latest encoutered SSS BSDF
-		if (!BSDF_GETISTHIN(BSDF) &&					// STRATA_TODO A material cannot have subsurface if using Thin mode. This is because DMFP does not match MFP at all (right now, SSS would just smears thin lighting model details out)
-			BSDF_GETHASSSS(BSDF) && BSDFWeightAvg > 0)
-		{
-			SetHasSubsurface(Classification, true);
-			SetHasSubsurfaceProfile(Classification, BSDF_GETHASSSSPROFILE(BSDF));
-
-			SSS.BaseColor = SLAB_DIFFUSEALBEDO(BSDF);
-			SSS.Specular  = SLAB_F0(BSDF).x;				// STRATA_TODO: Specular or F0?
-
-			// Profile ID: we only keep the latest ID, and profile will overtake any non-profile based SSS
-			SSS.bIsValid			= true;
-			SSS.bIsProfile			= SSS.bIsProfile || BSDF_GETHASSSSPROFILE(BSDF);
-			SSS.ProfileId			= StrataSubsurfaceProfileIdTo8bits(SLAB_SSSPROFILEID(BSDF));
-			SSS.ProfileRadiusScale	= SLAB_SSSPROFILERADIUSSCALE(BSDF);
-
-			// STRATA_TODO This should be accumulated according to the BSDF weight when we run this in the base pass.
-			// We cannot do that now because BSDFWeight contains weight*coverage*trhougput all mixed together.
-			SSS.DMFP				= SLAB_SSSDMFP(BSDF);
-		}
-		break;
-	}
-
-	case STRATA_BSDF_TYPE_HAIR:
-	{
-		AndIsSimpleMaterial(Classification, false);
-
-		// No SSR
-		if (bTopLayer)
-		{
-			TopLayerRoughness += BSDFWeightAvg * FullyRough; // Do not trigger SSR by using fully rough
-			TopLayerRoughnessAverageCount += BSDFWeightAvg;
-		}
-		break;
-	}
-
-	case STRATA_BSDF_TYPE_SINGLELAYERWATER:
-	{
-		AndIsSimpleMaterial(Classification, false);
-
-		if (bTopLayer)
-		{
-			const float Roughness = SLW_ROUGHNESS(BSDF);
-
-			TopLayerRoughness += BSDFWeightAvg * Roughness;
-			TopLayerRoughnessAverageCount += BSDFWeightAvg;
-		}
-		break;
-	}
-
-	} // switch		
-}
-
-void MainPS(
-	float4 SVPos : SV_POSITION, 
-	out uint  OutColor0 : SV_Target0,
-	out uint  OutColor1 : SV_Target1,
-	out uint2 OutColor2 : SV_Target2)
-{
-	const uint2 PixelPos = uint2(SVPos.xy);
-	
-	FStrataClassification ClassificationData	= (FStrataClassification)0;
-	FStrataSubsurfaceData SSSData				= (FStrataSubsurfaceData)0;
-	FStrataTopLayerData TopLayerData			= (FStrataTopLayerData)0;
-
-	FStrataAddressing StrataAddressing = GetStrataPixelDataByteOffset(PixelPos, uint2(PrimaryView.BufferSizeAndInvSize.xy), MaxBytesPerPixel);
-	FStrataPixelHeader Header = UnpackStrataHeaderIn(MaterialLobesBuffer, StrataAddressing);
-
-	SetIsSimpleMaterial(ClassificationData, true); // non strata material or sky pixels are considered as simple.
-	if (Header.BSDFCount > 0)
-	{
-		float2 ScreenPosition = SvPositionToScreenPosition(SVPos).xy;
-		const float Depth = 1000000.0f;
-		float4 WorldPos = mul(float4(ScreenPosition * Depth, Depth, 1), LWCHackToFloat(PrimaryView.ScreenToWorld));
-		const float3 V = normalize(LWCHackToFloat(PrimaryView.WorldCameraOrigin) - WorldPos.xyz);
-
-		AndIsSimpleMaterial(ClassificationData, Header.BSDFCount == 1);
-
-		float TopLayerRoughness = 0.0f;
-		float TopLayerRoughnessAverageCount = 0.0f;
-		for (uint BSDFIndex = 0; BSDFIndex < Header.BSDFCount; ++BSDFIndex)
-		{
-			const FStrataBSDF BSDF = UnpackStrataBSDFIn(MaterialLobesBuffer, StrataAddressing);
-			ComputePostProcessBSDFData(Header, BSDF, StrataAddressing, V, ClassificationData, SSSData, TopLayerData, TopLayerRoughness, TopLayerRoughnessAverageCount);
-		}
-
-		TopLayerData.WorldNormal = normalize(TopLayerData.WorldNormal);
-		ClassificationData.TopLayerRoughness = TopLayerRoughnessAverageCount > 0.0f ? TopLayerRoughness / TopLayerRoughnessAverageCount : 0.0f;
-	}
-	else
-	{
-		TopLayerData.WorldNormal = normalize(float3(-1.0f, -1.0f, -1.0f));
-	}
-
-	OutColor0 = StrataPackClassificationData(ClassificationData);
-	OutColor1 = StrataPackTopLayerData(TopLayerData);
-	OutColor2 = StrataPackSSSData(SSSData);
-}
-#endif // SHADER_CATEGORIZATION
-
 ////////////////////////////////////////////////////////////////////////////////////////////////////////////
 
 uint EncodeTile(uint2 TileCoord)
