Fixed reflection environment PS warning.

diff --git a/Engine/Shaders/Private/ReflectionEnvironmentPixelShader.usf b/Engine/Shaders/Private/ReflectionEnvironmentPixelShader.usf
--- a/Engine/Shaders/Private/ReflectionEnvironmentPixelShader.usf
+++ b/Engine/Shaders/Private/ReflectionEnvironmentPixelShader.usf
@@ -302,7 +302,7 @@ float3 ReflectionEnvironment(FGBufferData GBuffer, float AmbientOcclusion, float
 		// F_Schlick
 		const float CoatF0 = 0.04f;
 		#if USE_ENERGY_CONSERVATION
-		float F = ComputeGGXSpecEnergyTerms(ClearCoatRoughness, NoV, CoatF0).E;
+		float F = ComputeGGXSpecEnergyTerms(ClearCoatRoughness, NoV, CoatF0).E.x;
 		#else
 		float F = EnvBRDF(CoatF0, ClearCoatRoughness, NoV).x;
 		#endif
