Remove legacy precomputedshadow storage in header data.

diff --git a/Engine/Shaders/Private/Strata/Strata.ush b/Engine/Shaders/Private/Strata/Strata.ush
--- a/Engine/Shaders/Private/Strata/Strata.ush
+++ b/Engine/Shaders/Private/Strata/Strata.ush
@@ -1804,7 +1804,7 @@ void StrataLoad_ColorGamma2ToLinearAlphaLinear(FStrataMaterialContainer StrataBu
 ///////////////////////////////////////////////////////////////////////////////
 // Strata deferred base-pass output
 
-uint2 PackStrataHeader(uint BSDFCount, FStrataPixelHeader In);
+uint PackStrataHeader(uint BSDFCount, FStrataPixelHeader In);
 
 // If this is changed, please update the compiler side material size evaluation in StrataMaterial.cpp
 void PackStrataOut(
@@ -2041,14 +2041,8 @@ void PackStrataOut(
 
 	// 1. the header
 	AndIsSimpleMaterial(StrataPixelHeader, BSDFVisibleCount == 1);
-	const uint2 PackedHeader = PackStrataHeader(BSDFVisibleCount, StrataPixelHeader);
-	STRATA_STORE_UINT1(PackedHeader.x);
-#if ALLOW_STATIC_LIGHTING
-	if (HasPrecShadowMask(StrataPixelHeader))
-	{
-		STRATA_STORE_UINT1(PackedHeader.y);
-	}
-#endif
+	const uint PackedHeader = PackStrataHeader(BSDFVisibleCount, StrataPixelHeader);
+	STRATA_STORE_UINT1(PackedHeader);
 
 #if STRATA_INLINE_SHADING
 	UNROLL
@@ -2193,7 +2187,7 @@ void PackStrataOut(
 }
 
 // Pack only strata header
-uint2 PackStrataHeader(uint InBSDFCount, FStrataPixelHeader InHeader)
+uint PackStrataHeader(uint InBSDFCount, FStrataPixelHeader InHeader)
 {
 	uint Out = 0;
 	HEADER_SETBSDFCOUNT(Out, InBSDFCount);
