Add compact fast-path material for Strata.

diff --git a/Engine/Shaders/Private/CapsuleShadowShaders.usf b/Engine/Shaders/Private/CapsuleShadowShaders.usf
--- a/Engine/Shaders/Private/CapsuleShadowShaders.usf
+++ b/Engine/Shaders/Private/CapsuleShadowShaders.usf
@@ -964,7 +964,7 @@ void ApplySelfShadowingIntensityForDeferred(float2 ScreenUV, inout float Visibil
 #if STRATA_ENABLED
 		uint2 PixelPos = View.ViewRectMin.xy + View.ViewSizeAndInvSize.xy * ScreenUV.xy;
 		FStrataAddressing StrataAddressing = GetStrataPixelDataByteOffset(PixelPos, uint2(View.BufferSizeAndInvSize.xy), Strata.MaxBytesPerPixel);
-		FStrataPixelHeader StrataPixelHeader = UnpackStrataHeaderIn(Strata.MaterialTextureArray, StrataAddressing);
+		FStrataPixelHeader StrataPixelHeader = UnpackStrataHeaderIn(Strata.MaterialTextureArray, StrataAddressing, Strata.TopLayerTexture);
 		const bool bHasDynamicIndirectShadowCasterRepresentation = HasDynamicIndirectShadowCasterRepresentation(StrataPixelHeader);
 #else
 		FGBufferData GBufferData = GetGBufferData(ScreenUV);
diff --git a/Engine/Shaders/Private/DeferredLightPixelShaders.usf b/Engine/Shaders/Private/DeferredLightPixelShaders.usf
--- a/Engine/Shaders/Private/DeferredLightPixelShaders.usf
+++ b/Engine/Shaders/Private/DeferredLightPixelShaders.usf
@@ -218,7 +218,7 @@ void DeferredLightPixelMain(
 #if STRATA_ENABLED
 
 	FStrataAddressing StrataAddressing = GetStrataPixelDataByteOffset(PixelPos, uint2(View.BufferSizeAndInvSize.xy), Strata.MaxBytesPerPixel);
-	FStrataPixelHeader StrataPixelHeader = UnpackStrataHeaderIn(Strata.MaterialTextureArray, StrataAddressing);
+	FStrataPixelHeader StrataPixelHeader = UnpackStrataHeaderIn(Strata.MaterialTextureArray, StrataAddressing, Strata.TopLayerTexture);
 
 	BRANCH
 	if (StrataPixelHeader.BSDFCount > 0	// This test is also enough to exclude sky pixels
diff --git a/Engine/Shaders/Private/DeferredLightingCommon.ush b/Engine/Shaders/Private/DeferredLightingCommon.ush
--- a/Engine/Shaders/Private/DeferredLightingCommon.ush
+++ b/Engine/Shaders/Private/DeferredLightingCommon.ush
@@ -149,7 +149,7 @@ float ShadowRayCast(
 		// It is not defined for opaque materials, only the UAV is available to write the strata buffer.
 		uint2 PixelPos = View.ViewRectMin.xy + View.ViewSizeAndInvSize.xy * SampleUVz.xy;
 		FStrataAddressing StrataAddressing = GetStrataPixelDataByteOffset(PixelPos, uint2(View.BufferSizeAndInvSize.xy), StrataStruct.MaxBytesPerPixel);
-		FStrataPixelHeader StrataPixelHeader = UnpackStrataHeaderIn(StrataStruct.MaterialTextureArray, StrataAddressing);
+		FStrataPixelHeader StrataPixelHeader = UnpackStrataHeaderIn(StrataStruct.MaterialTextureArray, StrataAddressing, StrataStruct.TopLayerTexture);
 		bOutHitCastContactShadow = DoesCastContactShadow(StrataPixelHeader);
 		#endif
 	#else
diff --git a/Engine/Shaders/Private/PostProcessAmbientOcclusionCommon.ush b/Engine/Shaders/Private/PostProcessAmbientOcclusionCommon.ush
--- a/Engine/Shaders/Private/PostProcessAmbientOcclusionCommon.ush
+++ b/Engine/Shaders/Private/PostProcessAmbientOcclusionCommon.ush
@@ -108,7 +108,7 @@ float3 GetWorldSpaceNormalFromAOInput(float2 UV, float4 SvPosition)
 	#else
 #if STRATA_ENABLED
 		FStrataAddressing StrataAddressing = GetStrataPixelDataByteOffset(SvPosition.xy, uint2(View.BufferSizeAndInvSize.xy), Strata.MaxBytesPerPixel);
-		FStrataPixelHeader StrataHeader = UnpackStrataHeaderIn(Strata.MaterialTextureArray, StrataAddressing);
+		FStrataPixelHeader StrataHeader = UnpackStrataHeaderIn(Strata.MaterialTextureArray, StrataAddressing, Strata.TopLayerTexture);
 		if (!StrataHasShadingModel(StrataHeader, STRATA_BSDF_TYPE_HAIR))
 		{
 			const FStrataTopLayerData TopLayerData = StrataUnpackTopLayerData(Strata.TopLayerTexture.Load(uint3(SvPosition.xy, 0)));
diff --git a/Engine/Shaders/Private/ReflectionEnvironmentPixelShader.usf b/Engine/Shaders/Private/ReflectionEnvironmentPixelShader.usf
--- a/Engine/Shaders/Private/ReflectionEnvironmentPixelShader.usf
+++ b/Engine/Shaders/Private/ReflectionEnvironmentPixelShader.usf
@@ -369,7 +369,7 @@ void ReflectionEnvironmentSkyLighting(
 #if STRATA_ENABLED
 
 	FStrataAddressing StrataAddressing = GetStrataPixelDataByteOffset(PixelPos, uint2(View.BufferSizeAndInvSize.xy), Strata.MaxBytesPerPixel);
-	FStrataPixelHeader StrataPixelHeader = UnpackStrataHeaderIn(Strata.MaterialTextureArray, StrataAddressing);
+	FStrataPixelHeader StrataPixelHeader = UnpackStrataHeaderIn(Strata.MaterialTextureArray, StrataAddressing, Strata.TopLayerTexture);
 	BRANCH
 	if (StrataPixelHeader.BSDFCount > 0) // This test is also enough to exclude sky pixels
 	{
diff --git a/Engine/Shaders/Private/SSRT/SSRTReflections.usf b/Engine/Shaders/Private/SSRT/SSRTReflections.usf
--- a/Engine/Shaders/Private/SSRT/SSRTReflections.usf
+++ b/Engine/Shaders/Private/SSRT/SSRTReflections.usf
@@ -143,10 +143,10 @@ void ScreenSpaceReflections(
 		// We do not rerun material classification after the Strata material buffer has been patched with SLW material.
 		// So in this case we simply load the Header and water BSDF directly because we know there can only be a single one.
 		FStrataAddressing StrataAddressing = GetStrataPixelDataByteOffset(PixelCoord, uint2(View.BufferSizeAndInvSize.xy), Strata.MaxBytesPerPixel);
-		FStrataPixelHeader StrataPixelHeader = UnpackStrataHeaderIn(Strata.MaterialTextureArray, StrataAddressing);
+		FStrataPixelHeader StrataPixelHeader = UnpackStrataHeaderIn(Strata.MaterialTextureArray, StrataAddressing, Strata.TopLayerTexture);
 		if (StrataPixelHeader.BSDFCount > 0 && IsSingleLayerWater(StrataPixelHeader))
 		{
-			const FStrataBSDF SLWaterBSDF = UnpackStrataBSDFIn(Strata.MaterialTextureArray, StrataAddressing);
+			const FStrataBSDF SLWaterBSDF = UnpackStrataBSDFIn(Strata.MaterialTextureArray, StrataAddressing, StrataPixelHeader);
 			N = StrataGetBSDFSharedBasis(StrataPixelHeader, SLWaterBSDF, StrataAddressing)[2];
 			Roughness = SLW_ROUGHNESS(SLWaterBSDF);
 			bNoMaterial = false;
diff --git a/Engine/Shaders/Private/ShadingFurnaceTest.usf b/Engine/Shaders/Private/ShadingFurnaceTest.usf
--- a/Engine/Shaders/Private/ShadingFurnaceTest.usf
+++ b/Engine/Shaders/Private/ShadingFurnaceTest.usf
@@ -56,18 +56,18 @@ void MainPS(
 	{
 		// Fetch the first layer of roughness & normal only. This fine for simple material but will not work for clear coat.
 		FStrataAddressing StrataAddressing = GetStrataPixelDataByteOffset(PixelPos, uint2(View.BufferSizeAndInvSize.xy), Strata.MaxBytesPerPixel);
-		FStrataPixelHeader StrataPixelHeader = UnpackStrataHeaderIn(Strata.MaterialTextureArray, StrataAddressing);
+		FStrataPixelHeader StrataPixelHeader = UnpackStrataHeaderIn(Strata.MaterialTextureArray, StrataAddressing, Strata.TopLayerTexture);
 		bIsValid = StrataPixelHeader.BSDFCount > 0;
 		if (bIsValid)
 		{
 			const float3 DummyV = float3(0, 0, 1);
 			const float3 DummyL = float3(0, 0, 1);
-			FStrataBSDF BSDF = UnpackStrataBSDFIn(Strata.MaterialTextureArray, StrataAddressing);
+			FStrataBSDF BSDF = UnpackStrataBSDFIn(Strata.MaterialTextureArray, StrataAddressing, StrataPixelHeader);
 			FStrataBSDFContext BSDFContext = StrataCreateBSDFContext(StrataPixelHeader, BSDF, StrataAddressing, DummyV, DummyL);
 
 			SceneDepth = CalcSceneDepth(ScreenUV);
 			N = BSDFContext.N;
-			StrataGetBSDFRoughness(BSDFContext.BSDF, Roughness);
+			Roughness = StrataGetBSDFRoughness(BSDFContext.BSDF);
 			ClearCoatRoughness = Roughness;
 		}
 	}
@@ -227,7 +227,7 @@ void MainPS(
 			#elif 0
 				{
 					FStrataAddressing StrataAddressing = GetStrataPixelDataByteOffset(PixelPos, uint2(View.BufferSizeAndInvSize.xy), Strata.MaxBytesPerPixel);
-					FStrataPixelHeader StrataPixelHeader = UnpackStrataHeaderIn(Strata.MaterialTextureArray, StrataAddressing);
+					FStrataPixelHeader StrataPixelHeader = UnpackStrataHeaderIn(Strata.MaterialTextureArray, StrataAddressing, Strata.TopLayerTexture);
 					const float LightMask = 1;
 
 					BRANCH
diff --git a/Engine/Shaders/Private/ShadowProjectionPixelShader.usf b/Engine/Shaders/Private/ShadowProjectionPixelShader.usf
--- a/Engine/Shaders/Private/ShadowProjectionPixelShader.usf
+++ b/Engine/Shaders/Private/ShadowProjectionPixelShader.usf
@@ -256,7 +256,7 @@ void Main(
 		bool bIsEyeOrHair 				= false;
 	#if STRATA_ENABLED
 		FStrataAddressing StrataAddressing = GetStrataPixelDataByteOffset(PixelCoord.xy, uint2(View.BufferSizeAndInvSize.xy), Strata.MaxBytesPerPixel);
-		FStrataPixelHeader StrataHeader = UnpackStrataHeaderIn(Strata.MaterialTextureArray, StrataAddressing);
+		FStrataPixelHeader StrataHeader = UnpackStrataHeaderIn(Strata.MaterialTextureArray, StrataAddressing, Strata.TopLayerTexture);
 		if (IsStrataMaterial(StrataHeader))
 		{
 			bIsEyeOrHair				= StrataHasShadingModel(StrataHeader, STRATA_BSDF_TYPE_HAIR);
diff --git a/Engine/Shaders/Private/SingleLayerWaterComposite.usf b/Engine/Shaders/Private/SingleLayerWaterComposite.usf
--- a/Engine/Shaders/Private/SingleLayerWaterComposite.usf
+++ b/Engine/Shaders/Private/SingleLayerWaterComposite.usf
@@ -58,7 +58,7 @@ SingleLayerWaterCompositeOutput SingleLayerWaterComposite(float2 BufferUV, float
 
 	// Water do nto read the classification data to avoid running it a second time when water is enabled.
 	FStrataAddressing StrataAddressing = GetStrataPixelDataByteOffset(SvPosition.xy, uint2(View.BufferSizeAndInvSize.xy), Strata.MaxBytesPerPixel);
-	FStrataPixelHeader StrataPixelHeader = UnpackStrataHeaderIn(Strata.MaterialTextureArray, StrataAddressing);
+	FStrataPixelHeader StrataPixelHeader = UnpackStrataHeaderIn(Strata.MaterialTextureArray, StrataAddressing, Strata.TopLayerTexture);
 	if (StrataPixelHeader.BSDFCount > 0 && IsSingleLayerWater(StrataPixelHeader))
 	{
 		float DeviceZ = SampleDeviceZFromSceneTextures(BufferUV);
@@ -72,7 +72,7 @@ SingleLayerWaterCompositeOutput SingleLayerWaterComposite(float2 BufferUV, float
 		float DeltaDepth = CameraIsUnderWater ? WaterDepth : OpaqueDepth - WaterDepth;	// Inverted depth
 		if (DeltaDepth > 0.0)
 		{
-			const FStrataBSDF SLWaterBSDF = UnpackStrataBSDFIn(Strata.MaterialTextureArray, StrataAddressing);
+			const FStrataBSDF SLWaterBSDF = UnpackStrataBSDFIn(Strata.MaterialTextureArray, StrataAddressing, StrataPixelHeader);
 			const float3 WorldNormal = StrataGetBSDFSharedBasis(StrataPixelHeader, SLWaterBSDF, StrataAddressing)[2];
 			const float3 F0 = ComputeF0(SLW_SPECULAR(SLWaterBSDF), SLW_BASECOLOR(SLWaterBSDF), SLW_METALLIC(SLWaterBSDF));
 			const float Roughness = SLW_ROUGHNESS(SLWaterBSDF);
@@ -252,7 +252,7 @@ void WaterTileCatergorisationCS(uint3 ThreadId : SV_DispatchThreadID, uint3 Grou
 
 			#if STRATA_ENABLED
 				FStrataAddressing StrataAddressing = GetStrataPixelDataByteOffset(PixelPos, uint2(View.BufferSizeAndInvSize.xy), Strata.MaxBytesPerPixel);
-				FStrataPixelHeader StrataPixelHeader = UnpackStrataHeaderIn(Strata.MaterialTextureArray, StrataAddressing);
+				FStrataPixelHeader StrataPixelHeader = UnpackStrataHeaderIn(Strata.MaterialTextureArray, StrataAddressing, Strata.TopLayerTexture);
 				bContainsWater = bContainsWater || (StrataPixelHeader.BSDFCount > 0 && IsSingleLayerWater(StrataPixelHeader));
 			#else
 				FGBufferData GBuffer = GetGBufferDataFromSceneTextures(BufferUV);
@@ -278,7 +278,7 @@ void WaterTileCatergorisationCS(uint3 ThreadId : SV_DispatchThreadID, uint3 Grou
 
 	#if STRATA_ENABLED
 		FStrataAddressing StrataAddressing = GetStrataPixelDataByteOffset(PixelPos, uint2(View.BufferSizeAndInvSize.xy), Strata.MaxBytesPerPixel);
-		FStrataPixelHeader StrataPixelHeader = UnpackStrataHeaderIn(Strata.MaterialTextureArray, StrataAddressing);
+		FStrataPixelHeader StrataPixelHeader = UnpackStrataHeaderIn(Strata.MaterialTextureArray, StrataAddressing, Strata.TopLayerTexture);
 		bContainsWater = bContainsWater || (StrataPixelHeader.BSDFCount > 0 && IsSingleLayerWater(StrataPixelHeader));
 	#else
 		FGBufferData GBuffer = GetGBufferDataFromSceneTextures(BufferUV);
diff --git a/Engine/Shaders/Private/Strata/Strata.ush b/Engine/Shaders/Private/Strata/Strata.ush
--- a/Engine/Shaders/Private/Strata/Strata.ush
+++ b/Engine/Shaders/Private/Strata/Strata.ush
@@ -120,6 +120,45 @@ struct FStrataBSDF
 	float4 VGPRs[5];
 };
 
+////////////////////////////////////////////////////////////////////////////////////////////////////////////////////////////////////////////////////////////////////////////////////////////////
+// Storage layouts
+// ---------------
+// 2 types of storages layout:
+// * Complex: store many (mixed) BSDFs
+// * Fast   : store a single Slab BSDF with no extra feature
+// 
+// Header State (8bits):
+// * ISSIMPLEMATERIAL			
+// * HASPRECSHADOWMASK			
+// * ZEROPRECSHADOWMASK			
+// * CASTCONTACTSHADOW			
+// * HASDYNINDIRECTSHADOWCASTER	
+// * ISSINGLELAYERWATER			
+// * HASSUBSURFACE				
+// * HASSUBSURFACEPROFILE
+// 
+// * Complex: -------------------------------------------------------------------------------------------------------------------------------------------------------------------------
+//            HEADER        31 |  State | AO | BSDF Count | BSDF Types | Basis count | Basis type |
+//                             |  8       8    4            4            3             4           
+//            BASIS(0)      32 |  Normal: Oct 16/16 | Tangent: 11/11/10
+//            BASIS(1)      32 |  Normal: Oct 16/16 | Tangent: 11/11/10
+//			  ...			   |
+//            STATE(0)      17 |  BSDF Type | Normal Id | Aniso | TopLayer | SSS | SSSProfile | GreyWeigh | Haze | F90 | ThinFilm | Simple vol. | MFP Plugged | Fuzz | Thin
+//                             |  3           2           1       1          1     1            1           1      1     1          1             1             1      1 
+//            BSDF(0)       ?? |  Adaptive
+//            STATE(1)      17 |  BSDF Type | Normal Id | Aniso | TopLayer | SSS | SSSProfile | GreyWeigh | Haze | F90 | ThinFilm | Simple vol. | MFP Plugged | Fuzz | Thin
+//                             |  3           2           1       1          1     1            1           1      1     1          1             1             1      1 
+//            BSDF(1)       ?? |  Adaptive
+//			  ...			   |
+// * Fast   : -------------------------------------------------------------------------------------------------------------------------------------------------------------------------          
+//            HEADER	    29 | State | A0 | Roughness | DiffuseAbedo(8/20)
+//                             | 8       8    8           8
+//            BASIS(0)       0 | (Use top layer normal)
+//            STATE(0)	     0 | 
+//            BSDF(0)       32 | F0_RGB | DiffuseAbedo(12/20)
+//                             | 20       12
+//            -------------------------------------------------------------------------------------------------------------------------------------------------------------------------
+////////////////////////////////////////////////////////////////////////////////////////////////////////////////////////////////////////////////////////////////////////////////////////////////
 
 ////////////////////////////////////////////////////////////////////////////////////////////////////////////////////////////////////////////////////////////////////////////////////////////////
 // Bit utils functions
@@ -162,6 +201,9 @@ struct FStrataBSDF
 
 #define HEADER_MASK_SHADINGMODELS					(((1u<<HEADER_BIT_COUNT_SHADINGMODELS)-1) << HEADER_BIT_OFFSET_SHADINGMODELS)
 
+#define HEADER_FASTENCODING_BIT_COUNT				(HEADER_BIT_COUNT_STATE + HEADER_BIT_COUNT_AO)
+#define HEADER_FASTENCODING_MASK					((1u<<HEADER_FASTENCODING_BIT_COUNT)-1u)
+
 // Sanity check
 #if (HEADER_BIT_OFFSET_SHAREDLOCALBASES_TYPE + HEADER_BIT_COUNT_SHAREDLOCALBASES_TYPE) > 32
 #error Strata Header is larger than 32 bits
@@ -832,7 +874,8 @@ struct FStrataPixelHeader
 #else
 	FStrataMaterialContainer StrataBuffer;
 	uint SharedLocalBasesIndexOffset;
-	uint SharedLocalBasesTypes;
+	uint SharedLocalBasesTypes_PackedHeader; // Stores SharedLocalBasesTypes for non-simple material, and PackedHeader data for simple material
+	uint PackedTopLayerData;
 #endif
 
 	float MaterialAO;
@@ -846,7 +889,7 @@ bool HasZeroPrecShadowMask(const FStrataPixelHeader In)										{ return (In.St
 bool IsSingleLayerWater(const FStrataPixelHeader In)										{ return (In.State & HEADER_MASK_ISSINGLELAYERWATER) != 0; }
 bool DoesCastContactShadow(const FStrataPixelHeader In)										{ return (In.State & HEADER_MASK_CASTCONTACTSHADOW) != 0; }
 uint GetShadingModels(const FStrataPixelHeader In)											{ return (In.State & HEADER_MASK_SHADINGMODELS)>>HEADER_BIT_OFFSET_SHADINGMODELS; }
-bool IsStrataMaterial(const FStrataPixelHeader In)											{ return (In.State & HEADER_MASK_SHADINGMODELS) > 0; }
+bool IsStrataMaterial(const FStrataPixelHeader In)											{ return (In.State & HEADER_MASK_SHADINGMODELS) > 0 || IsSimpleMaterial(In); }
 bool HasSubsurface(const FStrataPixelHeader In)												{ return (In.State & HEADER_MASK_HASSUBSURFACE) != 0; }
 bool HasSubsurfaceProfile(const FStrataPixelHeader In)										{ return (In.State & HEADER_MASK_HASSUBSURFACEPROFILE) != 0; }
 bool HasDynamicIndirectShadowCasterRepresentation(const FStrataPixelHeader In)				{ return (In.State & HEADER_MASK_HASDYNINDIRECTSHADOWCASTER) != 0; }
@@ -863,7 +906,7 @@ void SetHasSubsurface(inout FStrataPixelHeader Out, bool bIn)								{ Out.State
 void SetHasSubsurfaceProfile(inout FStrataPixelHeader Out, bool bIn)						{ Out.State |= (bIn ? HEADER_MASK_HASSUBSURFACEPROFILE : 0u); }
 
 void AddShadingModels(inout FStrataPixelHeader Out, uint bIn)								{ Out.State |= ((bIn<<HEADER_BIT_OFFSET_SHADINGMODELS) & HEADER_MASK_SHADINGMODELS); }
-void AndIsSimpleMaterial(inout FStrataPixelHeader Out, bool bIn)								{ Out.State = (Out.State & ~HEADER_MASK_ISSIMPLEMATERIAL) | ((Out.State & HEADER_MASK_ISSIMPLEMATERIAL) & (bIn ? HEADER_MASK_ISSIMPLEMATERIAL : 0u)); }
+void AndIsSimpleMaterial(inout FStrataPixelHeader Out, bool bIn)							{ Out.State = (Out.State & ~HEADER_MASK_ISSIMPLEMATERIAL) | ((Out.State & HEADER_MASK_ISSIMPLEMATERIAL) & (bIn ? HEADER_MASK_ISSIMPLEMATERIAL : 0u)); }
 
 FStrataPixelHeader InitialiseStrataPixelHeader()
 {
@@ -873,7 +916,8 @@ FStrataPixelHeader InitialiseStrataPixelHeader()
 	Out.SharedLocalBases = (FSharedLocalBases)0;
 #else
 	Out.SharedLocalBasesIndexOffset = 0;
-	Out.SharedLocalBasesTypes = 0;
+	Out.SharedLocalBasesTypes_PackedHeader = 0;
+	Out.PackedTopLayerData = 0;
 #endif
 	Out.MaterialAO = 1.0f;
 	Out.State = 0;
@@ -889,10 +933,14 @@ float3x3 StrataGetBSDFSharedBasis(in FStrataPixelHeader StrataPixelHeader, in FS
 	float3 Bitangent = cross(Normal, Tangent);
 	return float3x3(Tangent, Bitangent, Normal);
 #else
+	if (IsSimpleMaterial(StrataPixelHeader))
+	{		
+		return GetTangentBasis(StrataUnpackTopLayerData(StrataPixelHeader.PackedTopLayerData).WorldNormal);
+	}
 
 	uint PackedBasis = StrataPixelHeader.StrataBuffer[uint3(StrataAddressing.PixelCoords, StrataPixelHeader.SharedLocalBasesIndexOffset + BSDF_GETSHAREDLOCALBASISID(BSDF))];
 
-	uint BasisType   = StrataGetSharedLocalBasisType(StrataPixelHeader.SharedLocalBasesTypes, BSDF_GETSHAREDLOCALBASISID(BSDF));
+	uint BasisType   = StrataGetSharedLocalBasisType(StrataPixelHeader.SharedLocalBasesTypes_PackedHeader, BSDF_GETSHAREDLOCALBASISID(BSDF));
 	float3x3 OutTangentBasis;
 	if (BasisType == STRATA_BASIS_TYPE_NORMAL)
 	{
@@ -1833,6 +1881,8 @@ void PackStrataOut(
 
 	SetIsSimpleMaterial(StrataPixelHeader, true); // non strata material or sky pixels are considered as simple (so that (strata simple pixels) + (non-strata pixels) can build a simple tile).
 
+	uint SimpleMaterial_LayerIndex = 0;
+	uint SimpleMaterial_BSDFIndex = 0;
 	uint BSDFCount = 0;
 	uint BSDFVisibleCount = 0;
 	float3 Throughput = 1.0f;
@@ -2014,6 +2064,12 @@ void PackStrataOut(
 					// together with a threshold to avoid 0-length normal after normalization in case influence is fadeout using the Weight node.
 					TopLayerData.WorldNormal += N * max(1e-5, BSDFCoverage);
 				}
+
+				if (bIsSimpleMaterial)
+				{
+					SimpleMaterial_LayerIndex = l;
+					SimpleMaterial_BSDFIndex = i;
+				}
 			}
 
 			// Notify that the BSDF is at the top for SSR to only affect reflection there and not on the lower layers
@@ -2036,11 +2092,51 @@ void PackStrataOut(
 		// And combine this layer throughtput with he material throughput
 		Throughput *= LayerThroughput;
 	}
+	AndIsSimpleMaterial(StrataPixelHeader, BSDFVisibleCount == 1);
 
 	// Now write out Strata data
 
+	///////////////////////////////////////////////////////////////////////////
+	// A. Fast encoding
+	const bool bHasFastEncoding = IsSimpleMaterial(StrataPixelHeader);
+	if (bHasFastEncoding)
+	{		
+		FStrataBSDF BSDF = Strata.Layers[SimpleMaterial_LayerIndex].BSDFs[SimpleMaterial_BSDFIndex];
+
+		// Now we pack diffuse and F0 in a special way: both encoded as R7G7B6 with gamma 2.0. The low bit cound will be hidden by dithering and TAA.
+		// We pack 32 bits in the first uint and the remaining bits are pack is the lowest significant bit of the second uint.
+		const uint PackedDiffuse20Bits	= PackR7G7B6Gamma2(SLAB_DIFFUSEALBEDO(BSDF));
+		const uint PackedDiffuse12Bits	= PackedDiffuse20Bits & 0xFFF;
+		const uint PackedDiffuse8Bits	= (PackedDiffuse20Bits>>12) & 0xFF;
+		const uint PackedF020Bits		= PackR7G7B6Gamma2(SLAB_F0(BSDF));
+		const uint PackedRoughness8bits = PackR8(SLAB_ROUGHNESS(BSDF));
+
+		// Data0 (State|AO|Roughness|Diffuse8bits)
+		{
+			uint Out = 0;
+			HEADER_SETSTATE(Out, StrataPixelHeader.State);
+			HEADER_SETAO(Out, PackR8(StrataPixelHeader.MaterialAO));
+
+			uint Packed16 = (PackedDiffuse8Bits << 8) | PackedRoughness8bits;
+			Out = (Out & HEADER_FASTENCODING_MASK) | (Packed16 << HEADER_FASTENCODING_BIT_COUNT);
+			STRATA_STORE_UINT1(Out);
+
+			#if (HEADER_FASTENCODING_BIT_COUNT + 16) > 32
+			#error Strata fast path header is > 32bits
+			#endif
+		}
+		// Data1 (F0|Diffuse12bits)
+		{
+			uint Out = PackedF020Bits | (PackedDiffuse12Bits << 20);
+			STRATA_STORE_UINT1(Out);
+		}
+	}
+	///////////////////////////////////////////////////////////////////////////
+	// B. Regular/Complex encoding
+	else // (!bHasFastEncoding)
+	{
+
 	// 1. the header
-	AndIsSimpleMaterial(StrataPixelHeader, BSDFVisibleCount == 1);
 	const uint PackedHeader = PackStrataHeader(BSDFVisibleCount, StrataPixelHeader);
 	STRATA_STORE_UINT1(PackedHeader);
 
@@ -2180,6 +2276,7 @@ void PackStrataOut(
 			}
 		}
 	}
+	} // if (bHasFastEncoding)
 
 	FinalizeWrites(StrataBuffer, StrataAddressing);
 
@@ -2205,32 +2302,54 @@ uint PackStrataHeader(uint InBSDFCount, FStrataPixelHeader InHeader)
 
 // Unpack only strata header
 // If this is changed, please update the compiler side material size evaluation in StrataMaterial.cpp
-FStrataPixelHeader UnpackStrataHeaderIn(FStrataMaterialContainer StrataBuffer, inout FStrataAddressing StrataAddressing)
+FStrataPixelHeader UnpackStrataHeaderIn(FStrataMaterialContainer StrataBuffer, inout FStrataAddressing StrataAddressing, Texture2D<uint> InStrataTopLayerTexture)
 {
 	FStrataPixelHeader Out = InitialiseStrataPixelHeader();
 
 	uint PackedHeader = StrataLoadUint1(StrataBuffer, StrataAddressing);
-	Out.BSDFCount    = HEADER_GETBSDFCOUNT(PackedHeader);
-	uint SharedLocalBasesCount = HEADER_GETSHAREDLOCALBASESCOUNT(PackedHeader);
-	uint SharedLocalBasesTypes = HEADER_GETSHAREDLOCALBASESTYPE(PackedHeader);
+	Out.State		= HEADER_GETSTATE(PackedHeader);
+	Out.MaterialAO	= UnpackR8(HEADER_GETAO(PackedHeader));
+	Out.BSDFCount   = HEADER_GETBSDFCOUNT(PackedHeader);
 
-	Out.MaterialAO = UnpackR8(HEADER_GETAO(PackedHeader));
-	Out.State = HEADER_GETSTATE(PackedHeader);
+	const bool bIsSimpleMaterial = (PackedHeader & HEADER_MASK_ISSIMPLEMATERIAL) > 0;
 
-#if STRATA_INLINE_SHADING==0
+	// Use fast path for decoding data
+	BRANCH
+	if (bIsSimpleMaterial)
+	{
+		Out.BSDFCount	= 1;
+
+		// Clear shading model, and ensurewe have only STRATA_BSDF_TYPE_SLAB
+		Out.State = Out.State & (~HEADER_MASK_SHADINGMODELS);
+		HEADER_SETSHADINGMODELS(Out.State, StrataShadingModelBit(STRATA_BSDF_TYPE_SLAB));
+
+	#if STRATA_INLINE_SHADING==0
+		Out.StrataBuffer = StrataBuffer;
+		Out.SharedLocalBasesTypes_PackedHeader = PackedHeader;
+		Out.PackedTopLayerData = InStrataTopLayerTexture.Load(uint3(StrataAddressing.PixelCoords, 0));
+		Out.SharedLocalBasesIndexOffset = 0; // Unused
+	  #endif
+	}
+	else
+	{
+		uint SharedLocalBasesCount = HEADER_GETSHAREDLOCALBASESCOUNT(PackedHeader);
+		uint SharedLocalBasesTypes = HEADER_GETSHAREDLOCALBASESTYPE(PackedHeader);
 
-	Out.StrataBuffer = StrataBuffer;
-	Out.SharedLocalBasesTypes = SharedLocalBasesTypes;
+	#if STRATA_INLINE_SHADING==0
+		Out.StrataBuffer = StrataBuffer;
+		Out.SharedLocalBasesTypes_PackedHeader = SharedLocalBasesTypes;
+		Out.PackedTopLayerData = 0;
 
-	// Only keep the offset to the shared local bases memory. We are going to load them on demand.
-	Out.SharedLocalBasesIndexOffset = StrataAddressing.CurrentIndex;
+		// Only keep the offset to the shared local bases memory. We are going to load them on demand.
+		Out.SharedLocalBasesIndexOffset = StrataAddressing.CurrentIndex;
 
-	// Now skip over the shared local basis memory to be able to load the BSDF content
-	StrataAddressing.CurrentIndex += SharedLocalBasesCount;
+		// Now skip over the shared local basis memory to be able to load the BSDF content
+		StrataAddressing.CurrentIndex += SharedLocalBasesCount;
 
-	// And account for the space used by normal for the debug information
-	StrataAddressing.ReadBytes += SharedLocalBasesCount * STRATA_PACKED_SHAREDLOCALBASIS_STRIDE_BYTES;
-#endif
+		// And account for the space used by normal for the debug information
+		StrataAddressing.ReadBytes += SharedLocalBasesCount * STRATA_PACKED_SHAREDLOCALBASIS_STRIDE_BYTES;
+	#endif
+	}
 
 	return Out;
 }
@@ -2255,19 +2374,27 @@ void UnpackFastPathSlabBSDF(inout FStrataBSDF OutBSDF, uint2 PackedData01)
 	SLAB_ANISOTROPY(OutBSDF)	= Data1.z * 2.f - 1.f;
 }
 
-FStrataBSDF UnpackFastPathStrataBSDFIn(FStrataMaterialContainer StrataBuffer, inout FStrataAddressing StrataAddressing)
+FStrataBSDF UnpackFastPathStrataBSDFIn(FStrataMaterialContainer StrataBuffer, inout FStrataAddressing StrataAddressing, FStrataPixelHeader StrataPixelHeader)
 {
-	FStrataBSDF OutBSDF = (FStrataBSDF)0; 
+	FStrataBSDF OutBSDF = (FStrataBSDF)0;
 
-	OutBSDF.State = StrataLoadUint1(StrataBuffer, StrataAddressing);
+#if STRATA_INLINE_SHADING==0
+	uint Data0 = StrataPixelHeader.SharedLocalBasesTypes_PackedHeader;
+	uint Data1 = StrataLoadUint1(StrataBuffer, StrataAddressing);
+	OutBSDF.State = 0;
+	OutBSDF.FinalWeight = 1.0f;
 
-	// In the fast path the single slab BSDF is always on top so the weight must be grey scale.
-	OutBSDF.FinalWeight = Unpack10F(BSDF_GETWEIGHT10F(OutBSDF));
+	const uint PackedRoughness8bits = 0xFF    & (Data0 >>  HEADER_FASTENCODING_BIT_COUNT);
+	const uint PackedData8Bits		= 0xFF    & (Data0 >> (HEADER_FASTENCODING_BIT_COUNT + 8));
+	const uint PackedData12Bits		= 0xFFF   & (Data1 >> 20);
+	const uint PackedF020Bits		= 0xFFFFF & (Data1);
+	const uint PackedDiffuse20Bits	= PackedData12Bits | (PackedData8Bits<<12);
 
-	// Now unpack the simple Slab BSDF
-	uint DataX = StrataLoadUint1(StrataBuffer, StrataAddressing);
-	uint DataY = StrataLoadUint1(StrataBuffer, StrataAddressing);
-	UnpackFastPathSlabBSDF(OutBSDF, uint2(DataX, DataY));
+	SLAB_DIFFUSEALBEDO(OutBSDF)	= UnpackR7G7B6Gamma2(PackedDiffuse20Bits);
+	SLAB_F0(OutBSDF)			= UnpackR7G7B6Gamma2(PackedF020Bits);
+	SLAB_ROUGHNESS(OutBSDF)		= UnpackR8(PackedRoughness8bits);
+	SLAB_ANISOTROPY(OutBSDF)	= 0.0f;
+#endif
 
 	// This is convenient to not have to test BSDF_GETHASF90 all over the place in the code using the BSDF for such a core value.
 	SLAB_F90(OutBSDF) = 1.0f;
@@ -2275,6 +2402,7 @@ FStrataBSDF UnpackFastPathStrataBSDFIn(FStrataMaterialContainer StrataBuffer, in
 	// Force a BSDF state to a single slab with isotropic specular, no subsurface scattering or any other options.
 	// This is to help the compiler understand disabled path.
 	BSDF_SETTYPE(OutBSDF, STRATA_BSDF_TYPE_SLAB);
+	BSDF_SETHASGREYWEIGHT(OutBSDF, 1);
 	BSDF_SETHASF90(OutBSDF, 0);
 	BSDF_SETHASANISOTROPY(OutBSDF, 0);
 	BSDF_SETISTOPLAYER(OutBSDF, 1);
@@ -2291,8 +2419,14 @@ FStrataBSDF UnpackFastPathStrataBSDFIn(FStrataMaterialContainer StrataBuffer, in
 // Unpack a single BSDF
 // Note: All BSDF lobes needs to be unpack in a sequential manner since each BSDF has a variable footprint.
 // If this is changed, please update the compiler side material size evaluation in StrataMaterial.cpp
-FStrataBSDF UnpackStrataBSDFIn(FStrataMaterialContainer StrataBuffer, inout FStrataAddressing StrataAddressing)
+FStrataBSDF UnpackStrataBSDFIn(FStrataMaterialContainer StrataBuffer, inout FStrataAddressing StrataAddressing, in FStrataPixelHeader StrataHeader)
 {
+	BRANCH
+	if (IsSimpleMaterial(StrataHeader))
+	{
+		return UnpackFastPathStrataBSDFIn(StrataBuffer, StrataAddressing, StrataHeader);
+	}
+
 	FStrataBSDF OutBSDF = (FStrataBSDF)0;
 
 	float DummyFloat = 0.0f;
diff --git a/Engine/Shaders/Private/Strata/StrataDeferredLighting.ush b/Engine/Shaders/Private/Strata/StrataDeferredLighting.ush
--- a/Engine/Shaders/Private/Strata/StrataDeferredLighting.ush
+++ b/Engine/Shaders/Private/Strata/StrataDeferredLighting.ush
@@ -32,14 +32,13 @@ float4 StrataDeferredLighting(
 	{
 		// Unpack BSDF data
 #if STRATA_FASTPATH
-		FStrataBSDF BSDF = UnpackFastPathStrataBSDFIn(MaterialBuffer, StrataAddressing);
+		FStrataBSDF BSDF = UnpackFastPathStrataBSDFIn(MaterialBuffer, StrataAddressing, StrataPixelHeader);
 #else
-		FStrataBSDF BSDF = UnpackStrataBSDFIn(MaterialBuffer, StrataAddressing);
+		FStrataBSDF BSDF = UnpackStrataBSDFIn(MaterialBuffer, StrataAddressing, StrataPixelHeader);
 #endif
 		FStrataBSDFContext BSDFContext = StrataCreateBSDFContext(StrataPixelHeader, BSDF, StrataAddressing, V, L);
 
-		float Roughness = 0;
-		StrataGetBSDFRoughness(BSDFContext.BSDF, Roughness);
+		float Roughness = StrataGetBSDFRoughness(BSDFContext.BSDF);
 
 		FAreaLightIntegrateContext AreaLightContext = InitAreaLightIntegrateContext();
 		FStrataEvaluateResult BSDFEvaluate = (FStrataEvaluateResult)0;
diff --git a/Engine/Shaders/Private/Strata/StrataEnvironmentLighting.ush b/Engine/Shaders/Private/Strata/StrataEnvironmentLighting.ush
--- a/Engine/Shaders/Private/Strata/StrataEnvironmentLighting.ush
+++ b/Engine/Shaders/Private/Strata/StrataEnvironmentLighting.ush
@@ -53,9 +53,9 @@ float4 StrataEnvironmentLighting(
 #endif
 	{
 	#if STRATA_FASTPATH
-		FStrataBSDF BSDF = UnpackFastPathStrataBSDFIn(MaterialBuffer, StrataAddressing);
+		FStrataBSDF BSDF = UnpackFastPathStrataBSDFIn(MaterialBuffer, StrataAddressing, StrataPixelHeader);
 	#else
-		FStrataBSDF BSDF = UnpackStrataBSDFIn(MaterialBuffer, StrataAddressing);
+		FStrataBSDF BSDF = UnpackStrataBSDFIn(MaterialBuffer, StrataAddressing, StrataPixelHeader);
 	#endif
 		const float3 BSDFThroughput = BSDF.FinalWeight;
 
diff --git a/Engine/Shaders/Private/Strata/StrataEvaluation.ush b/Engine/Shaders/Private/Strata/StrataEvaluation.ush
--- a/Engine/Shaders/Private/Strata/StrataEvaluation.ush
+++ b/Engine/Shaders/Private/Strata/StrataEvaluation.ush
@@ -77,9 +77,9 @@ struct FStrataBSDFContext
 	float3 TangentH;
 };
 
-void StrataGetBSDFRoughness(in FStrataBSDF BSDF, inout float R)
+float StrataGetBSDFRoughness(in FStrataBSDF BSDF)
 {
-	R = 0;
+	float R = 0;
 	const uint BSDFType = BSDF_GETTYPE(BSDF);
 	switch (BSDFType)
 	{
@@ -99,6 +99,7 @@ void StrataGetBSDFRoughness(in FStrataBSDF BSDF, inout float R)
 		break;
 	}
 	}
+	return R;
 }
 
 FStrataBSDFContext StrataCreateBSDFContext(float3x3 TangentBasis, FStrataBSDF BSDF, float3 V, float3 L, bool bHasValidL=true)
diff --git a/Engine/Shaders/Private/Strata/StrataForwardLighting.ush b/Engine/Shaders/Private/Strata/StrataForwardLighting.ush
--- a/Engine/Shaders/Private/Strata/StrataForwardLighting.ush
+++ b/Engine/Shaders/Private/Strata/StrataForwardLighting.ush
@@ -66,8 +66,7 @@ float3 StrataForwardLightingCommon(
 		GetShadowTerms(SceneDepth, PrecomputedShadowFactors, FakeShadingModelID, FakeContactShadowOpacity,
 			LightData, AbsoluteWorldPosition, ToLight, LightAttenuation, Dither, ShadowTerms);
 
-		float Roughness = 0;
-		StrataGetBSDFRoughness(StrataBSDFContext.BSDF, Roughness);
+		float Roughness = StrataGetBSDFRoughness(StrataBSDFContext.BSDF);
 		FAreaLightIntegrateContext AreaLightContext = InitAreaLightIntegrateContext();
 		BRANCH
 		if (ShadowTerms.SurfaceShadow + ShadowTerms.TransmissionShadow > 0)
@@ -302,8 +301,7 @@ float3 StrataForwardLighting(
 				float DiffOcclusion = MaterialAO;
 				float SpecOcclusion = MaterialAO;
 				float3 BentNormal = StrataBSDFContext.N;
-				float Roughness = 0;
-				StrataGetBSDFRoughness(BSDF, Roughness);
+				float Roughness = StrataGetBSDFRoughness(BSDF);
 				ApplyBentNormal(MaterialParameters.CameraVector, MaterialParameters.WorldNormal, MaterialParameters.TangentToWorld, GetBentNormalZero(MaterialParameters), Roughness, BentNormal, DiffOcclusion, SpecOcclusion);
 
 				// Update a context specific for environment lighting
diff --git a/Engine/Shaders/Private/Strata/StrataMaterialClassification.usf b/Engine/Shaders/Private/Strata/StrataMaterialClassification.usf
--- a/Engine/Shaders/Private/Strata/StrataMaterialClassification.usf
+++ b/Engine/Shaders/Private/Strata/StrataMaterialClassification.usf
@@ -26,6 +26,7 @@ int TileSize;
 int bRectPrimitive;
 int2 ViewResolution;
 uint MaxBytesPerPixel;
+Texture2D<uint> TopLayerTexture;
 Texture2DArray<uint> MaterialTextureArray;
 RWBuffer<uint> SimpleTileIndirectDataBuffer;
 RWBuffer<uint> SimpleTileListDataBuffer;
@@ -54,10 +55,11 @@ void TileMainCS(uint2 DispatchThreadId : SV_DispatchThreadID, uint LinearIndex :
 	if (all(PixelCoord < uint2(ViewResolution)))
 	{
 		FStrataAddressing StrataAddressing = GetStrataPixelDataByteOffset(PixelCoord, uint2(View.BufferSizeAndInvSize.xy), MaxBytesPerPixel);
-		FStrataPixelHeader StrataPixelHeader = UnpackStrataHeaderIn(MaterialTextureArray, StrataAddressing);
+		FStrataPixelHeader StrataPixelHeader = UnpackStrataHeaderIn(MaterialTextureArray, StrataAddressing, TopLayerTexture);
 
+		const bool bIsSimple = IsSimpleMaterial(StrataPixelHeader) || StrataPixelHeader.BSDFCount == 0; // BSDFCount == 0 ensures that non-strata pixel, like sky pixels, won't make a simple tile flagged as complex
 		bContainsStrataMaterial  = StrataPixelHeader.BSDFCount > 0;
-		bContainsComplexMaterial = !IsSimpleMaterial(StrataPixelHeader);
+		bContainsComplexMaterial = !bIsSimple;
 	}
 
 #if PERMUTATION_WAVE_OPS
@@ -212,7 +214,6 @@ void ClearMaterialBufferMainCS(uint2 DispatchThreadId : SV_DispatchThreadID)
 	// The first layer of tiled uints contains the header that we need to clear so we only write a single uint per pixel instead of clearing the entire buffer which would be too slow.
 
 	FStrataPixelHeader StrataHeader = InitialiseStrataPixelHeader();
-	StrataHeader.State |= HEADER_MASK_ISSIMPLEMATERIAL;
 	uint BSDFCount = 0;
 	uint PackedStrataHeader = PackStrataHeader(BSDFCount, StrataHeader);
 
diff --git a/Engine/Shaders/Private/Strata/StrataVisualize.usf b/Engine/Shaders/Private/Strata/StrataVisualize.usf
--- a/Engine/Shaders/Private/Strata/StrataVisualize.usf
+++ b/Engine/Shaders/Private/Strata/StrataVisualize.usf
@@ -81,7 +81,7 @@ void VisualizeMaterialPS(
 	const int  InspectSquare = 10;
 
 	FStrataAddressing StrataAddressing = GetStrataPixelDataByteOffset(MaterialPos, uint2(View.BufferSizeAndInvSize.xy), Strata.MaxBytesPerPixel);
-	FStrataPixelHeader Header = UnpackStrataHeaderIn(Strata.MaterialTextureArray, StrataAddressing);
+	FStrataPixelHeader Header = UnpackStrataHeaderIn(Strata.MaterialTextureArray, StrataAddressing, Strata.TopLayerTexture);
 	bool bStrataMaterial = Header.BSDFCount > 0;
 
 	const float2 ScreenPosition = SvPositionToScreenPosition(SVPos).xy;
@@ -128,11 +128,11 @@ void VisualizeMaterialPS(
 	{
 		FStrataAddressing StrataAddressing = GetStrataPixelDataByteOffset(PixelPos, uint2(View.BufferSizeAndInvSize.xy), Strata.MaxBytesPerPixel);
 
-		FStrataPixelHeader Header = UnpackStrataHeaderIn(Strata.MaterialTextureArray, StrataAddressing);
+		FStrataPixelHeader Header = UnpackStrataHeaderIn(Strata.MaterialTextureArray, StrataAddressing, Strata.TopLayerTexture);
 
 		for (uint BSDFIndex = 0; BSDFIndex < Header.BSDFCount; ++BSDFIndex)
 		{
-			const FStrataBSDF BSDF = UnpackStrataBSDFIn(Strata.MaterialTextureArray, StrataAddressing);
+			const FStrataBSDF BSDF = UnpackStrataBSDFIn(Strata.MaterialTextureArray, StrataAddressing, Header);
 
 			// Referential
 			if (all(PixelPos == MaterialPos))
@@ -170,7 +170,7 @@ void VisualizeMaterialPS(
 		FStrataAddressing StrataAddressing = GetStrataPixelDataByteOffset(MaterialPos, uint2(View.BufferSizeAndInvSize.xy), Strata.MaxBytesPerPixel);
 		const uint PixelStrataDataByteOffsetStart = StrataAddressing.ReadBytes;
 
-		FStrataPixelHeader Header = UnpackStrataHeaderIn(Strata.MaterialTextureArray, StrataAddressing);
+		FStrataPixelHeader Header = UnpackStrataHeaderIn(Strata.MaterialTextureArray, StrataAddressing, Strata.TopLayerTexture);
 		const uint PixelStrataDataByteOffsetPostHeader = StrataAddressing.ReadBytes;
 
 		float3 BSDFFontColor = float3(0.25, 1.0, 0.25);
@@ -304,7 +304,7 @@ void VisualizeMaterialPS(
 		LOOP
 		for (uint BSDFIndex = 0; BSDFIndex < Header.BSDFCount; ++BSDFIndex)
 		{
-			const FStrataBSDF BSDF = UnpackStrataBSDFIn(Strata.MaterialTextureArray, StrataAddressing);
+			const FStrataBSDF BSDF = UnpackStrataBSDFIn(Strata.MaterialTextureArray, StrataAddressing, Header);
 
 			TopLeft.x = TopLeftInit.x;
 
diff --git a/Engine/Source/Runtime/Renderer/Private/Strata/Strata.cpp b/Engine/Source/Runtime/Renderer/Private/Strata/Strata.cpp
--- a/Engine/Source/Runtime/Renderer/Private/Strata/Strata.cpp
+++ b/Engine/Source/Runtime/Renderer/Private/Strata/Strata.cpp
@@ -424,6 +424,7 @@ class FStrataMaterialTileClassificationPassCS : public FGlobalShader
 		SHADER_PARAMETER(int32, bRectPrimitive)
 		SHADER_PARAMETER(FIntPoint, ViewResolution)
 		SHADER_PARAMETER(uint32, MaxBytesPerPixel)
+		SHADER_PARAMETER_RDG_TEXTURE(Texture2D, TopLayerTexture)
 		SHADER_PARAMETER_RDG_TEXTURE_SRV(Texture2DArray<uint>, MaterialTextureArray)
 		SHADER_PARAMETER_RDG_BUFFER_UAV(RWBuffer, SimpleTileIndirectDataBuffer)
 		SHADER_PARAMETER_RDG_BUFFER_UAV(RWBuffer, SimpleTileListDataBuffer)
@@ -754,6 +755,7 @@ void AddStrataMaterialClassificationPass(FRDGBuilder& GraphBuilder, const FMinim
 			PassParameters->bRectPrimitive = GRHISupportsRectTopology ? 1 : 0;
 			PassParameters->ViewResolution = View.ViewRect.Size();
 			PassParameters->MaxBytesPerPixel = View.StrataSceneData->MaxBytesPerPixel;
+			PassParameters->TopLayerTexture = View.StrataSceneData->TopLayerTexture;
 			PassParameters->MaterialTextureArray = View.StrataSceneData->MaterialTextureArraySRV;
 			PassParameters->SimpleTileListDataBuffer = View.StrataSceneData->ClassificationTileListBufferUAV[EStrataTileMaterialType::ESimple];
 			PassParameters->SimpleTileIndirectDataBuffer = View.StrataSceneData->ClassificationTileIndirectBufferUAV[EStrataTileMaterialType::ESimple];
