Fix incorrect DF shadowing caused by invalid SDF bricks.
- When SphereRadius > MaxEncodedDistance ray is incorrectly occluded.
- Limit SphereRadius to MaxEncodedDistance.

diff --git a/Engine/Shaders/Private/DistanceFieldShadowingShared.ush b/Engine/Shaders/Private/DistanceFieldShadowingShared.ush
--- a/Engine/Shaders/Private/DistanceFieldShadowingShared.ush
+++ b/Engine/Shaders/Private/DistanceFieldShadowingShared.ush
@@ -117,6 +117,16 @@ float ShadowRayTraceThroughCulledObjects(
 			{
 				FDFAssetData DFAssetData = LoadDFAssetDataHighestResolution(DFObjectData.AssetIndex);
 
+				float MaxEncodedDistance = DFAssetData.DistanceFieldToVolumeScaleBias.x + DFAssetData.DistanceFieldToVolumeScaleBias.y;
+
+				if (DFObjectData.bMostlyTwoSided)
+				{
+					MaxEncodedDistance -= VolumeTwoSidedMeshDistanceBias;
+				}
+
+				// Prevent incorrect shadowing when sampling invalid bricks by limiting VolumeMaxSphereRadius to MaxEncodedDistance
+				VolumeMaxSphereRadius = min(VolumeMaxSphereRadius, MaxEncodedDistance);
+
 				float SampleRayTime = IntersectionTimes.x * VolumeRayLength;
 #if DF_SHADOW_QUALITY == 2
 				uint MaxSteps = 64;
