Strata - fix for unrol warning and pc debug shader compilation

diff --git a/Engine/Shaders/Private/BasePassPixelShader.usf b/Engine/Shaders/Private/BasePassPixelShader.usf
--- a/Engine/Shaders/Private/BasePassPixelShader.usf
+++ b/Engine/Shaders/Private/BasePassPixelShader.usf
@@ -1628,6 +1628,7 @@ void FPixelShaderInOut_MainPS(
 	FRWStrataMaterialContainer RWStrataMaterialContainer = InitialiseRWStrataMaterialContainer(OpaqueBasePass.Strata.MaterialTextureArrayUAV);
 	PackStrataOut(
 		RWStrataMaterialContainer,
+		OpaqueBasePass.Strata.MaterialTextureArrayUAV,
 		Dither,
 		StrataAddressing,
 		StrataPixelHeader, StrataData, MaterialParameters.CameraVector, bStrataSubsurfaceEnable, EmissiveLuminance,
diff --git a/Engine/Shaders/Private/DebugProbes.usf b/Engine/Shaders/Private/DebugProbes.usf
--- a/Engine/Shaders/Private/DebugProbes.usf
+++ b/Engine/Shaders/Private/DebugProbes.usf
@@ -283,6 +283,7 @@ void MainPS(
 		FRWStrataMaterialContainer RWStrataMaterialContainer = InitialiseRWStrataMaterialContainer(MaterialTextureArrayUAV);
 		PackStrataOut(
 			RWStrataMaterialContainer, 
+			MaterialTextureArrayUAV,
 			Dither,
 			StrataAddressing,
 			StrataPixelHeader, StrataData, -RayDir, bStrataSubsurfaceEnable, EmissiveLuminance,
diff --git a/Engine/Shaders/Private/Strata/Strata.ush b/Engine/Shaders/Private/Strata/Strata.ush
--- a/Engine/Shaders/Private/Strata/Strata.ush
+++ b/Engine/Shaders/Private/Strata/Strata.ush
@@ -38,7 +38,7 @@
 struct FRWStrataMaterialContainerStruct
 {
 	uint MaterialRenderTargets[STARTA_MATERIAL_RENDERTARGET_COUNT];
-	RWTexture2DArray<uint> ExtraMaterialDataUAV;
+//	RWTexture2DArray<uint> ExtraMaterialDataUAV; // If this is used, we cannot run shaders on debug...
 };
 #define FRWStrataMaterialContainer	FRWStrataMaterialContainerStruct
 #define FStrataMaterialContainer	Texture2DArray<uint>
@@ -46,11 +46,12 @@ struct FRWStrataMaterialContainerStruct
 FRWStrataMaterialContainer InitialiseRWStrataMaterialContainer(RWTexture2DArray<uint> ExtraMaterialDataUAV)
 {
 	FRWStrataMaterialContainer RWStrataMaterialContainer;
+	UNROLL
 	for (int i = 0; i < STARTA_MATERIAL_RENDERTARGET_COUNT; ++i)
 	{
 		RWStrataMaterialContainer.MaterialRenderTargets[i] = 0;
 	}
-	RWStrataMaterialContainer.ExtraMaterialDataUAV = ExtraMaterialDataUAV;
+//	RWStrataMaterialContainer.ExtraMaterialDataUAV = ExtraMaterialDataUAV;
 	return RWStrataMaterialContainer;
 }
 
@@ -792,7 +793,7 @@ FStrataAddressing GetStrataPixelDataByteOffset(uint2 PixelPos, uint2 ViewBufferS
 
 #if COMPILER_SWITCH // These platforms are complaining so simply disabling Strata there for now.
 
-void StrataStoreUint1(inout FRWStrataMaterialContainer StrataBuffer, inout FStrataAddressing StrataAddressing, uint Data) {}
+void StrataStoreUint1(inout FRWStrataMaterialContainer StrataBuffer, RWTexture2DArray<uint> ExtraMaterialDataUAV, inout FStrataAddressing StrataAddressing, uint Data) {}
 #define STRATA_STORE_UINT1(x) 
 
 void FinalizeWrites(inout FRWStrataMaterialContainer StrataBuffer, inout FStrataAddressing StrataAddressing) {}
@@ -801,7 +802,7 @@ uint StrataLoadUint1(FStrataMaterialContainer StrataBuffer, inout FStrataAddress
 
 #else
 
-void StrataStoreUint1(inout FRWStrataMaterialContainer StrataBuffer, inout FStrataAddressing StrataAddressing, uint Data)
+void StrataStoreUint1(inout FRWStrataMaterialContainer StrataBuffer, RWTexture2DArray<uint> ExtraMaterialDataUAV, inout FStrataAddressing StrataAddressing, uint Data)
 {
 	// We heavily rely on the compiler to optimize out the if code
 	if (StrataAddressing.CurrentIndex < STARTA_MATERIAL_RENDERTARGET_COUNT)
@@ -811,11 +812,11 @@ void StrataStoreUint1(inout FRWStrataMaterialContainer StrataBuffer, inout FStra
 	}
 	else
 	{
-		StrataBuffer.ExtraMaterialDataUAV[uint3(StrataAddressing.PixelCoords, StrataAddressing.CurrentIndex)] = Data;
+		ExtraMaterialDataUAV[uint3(StrataAddressing.PixelCoords, StrataAddressing.CurrentIndex)] = Data;
 	}
 	StrataAddressing.CurrentIndex++;
 }
-#define STRATA_STORE_UINT1(x)	StrataStoreUint1(StrataBuffer, StrataAddressing, x)
+#define STRATA_STORE_UINT1(x)	StrataStoreUint1(StrataBuffer, ExtraMaterialDataUAV, StrataAddressing, x)
 
 void FinalizeWrites(inout FRWStrataMaterialContainer StrataBuffer, inout FStrataAddressing StrataAddressing)
 {
@@ -1857,6 +1858,7 @@ uint PackStrataHeader(uint BSDFCount, FStrataPixelHeader In);
 // If this is changed, please update the compiler side material size evaluation in StrataMaterial.cpp
 void PackStrataOut(
 	inout FRWStrataMaterialContainer StrataBuffer,
+	RWTexture2DArray<uint> ExtraMaterialDataUAV,
 	float Dither,
 	inout FStrataAddressing StrataAddressing,
 	FStrataPixelHeader StrataPixelHeader,
