Fix short-term fix to allow legacy-eye conversion material to compile.

Compiler doesn't like the dynamic layering structure of eye. This needs to be revisited.

diff --git a/Engine/Shaders/Private/Strata/StrataLegacyConversion.ush b/Engine/Shaders/Private/Strata/StrataLegacyConversion.ush
--- a/Engine/Shaders/Private/Strata/StrataLegacyConversion.ush
+++ b/Engine/Shaders/Private/Strata/StrataLegacyConversion.ush
@@ -279,60 +279,13 @@ FStrataData StrataConvertLegacyMaterial(
 		const float IrisDistance = ClearCoatRoughness;
 		const float IrisOpacity = saturate(1.0f - IrisMask);
 
-		// * Either only the scelar is visible: only a single layer is needed
-		// * Or both the sclera and the iris are visible on the transition area: two layer mixing iris/scelar
-		// * Or only the iris is visible but due to the cornea having a different normal than the iris, we need two layers: two layer vertically composed for cornea and iris
-		if (IrisMask <= 0.0f)
+		// Incorrect composition, but allows to compile shaders
 		{
-			Out = GetStrataSlabBSDF(
-				true,				// UseMetalness
-				BaseColor,			// BaseColor		-  Metalness workflow
-				Ones,				// EdgeColor		-  Metalness workflow
-				Specular,			// Specular			-  Metalness workflow
-				0.f,				// Metallic			-  Metalness workflow
-				Zeros,				// DiffuseAlbedo	- !Metalness workflow
-				Zeros, 				// F0				- !Metalness workflow
-				Zeros,				// F90				- !Metalness workflow
-				Roughness,			// Roughness
-				0.f,				// Anisotropy
-				SubSurfaceProfileId,// SSSProfileID
-				Zeros,				// SSSDMFP
-				Opacity,			// SSSDMFPScale
-				Emissive,			// Emissive
-				0.f,				// Haziness
-				0.f,				// ThinFilmThickness
-				0.f,				// FuzzAmount
-				Zeros,				// FuzzColor
-				DefaultThickness,
-				SharedLocalBasisIndex, SharedLocalBasisTypes);
-		}
-		else if (IrisMask < 1.0f)
-		{
-			FStrataData Top = GetStrataSlabBSDF(
-				true,				// UseMetalness
-				BaseColor,			// BaseColor		-  Metalness workflow
-				Ones,				// EdgeColor		-  Metalness workflow
-				Specular,			// Specular			-  Metalness workflow
-				0.f,				// Metallic			-  Metalness workflow
-				Zeros,				// DiffuseAlbedo	- !Metalness workflow
-				Zeros, 				// F0				- !Metalness workflow
-				Zeros,				// F90				- !Metalness workflow
-				Roughness,			// Roughness
-				0.f,				// Anisotropy
-				SubSurfaceProfileId,// SSSProfileID
-				Zeros,				// SSSDMFP
-				Opacity,			// SSSDMFPScale
-				Emissive,			// Emissive
-				0.f,				// Haziness
-				0.f,				// ThinFilmThickness
-				0.f,				// FuzzAmount
-				Zeros,				// FuzzColor
-				DefaultThickness,
-				SharedLocalBasisIndex, SharedLocalBasisTypes);
+			const float3 TopBaseColor = BaseColor; // lerp(BaseColor, Zeros, IrisMask);
 
-			FStrataData IrisSlab = GetStrataSlabBSDF(
+			FStrataData CoatSlab = GetStrataSlabBSDF(
 				true,				// UseMetalness
-				BaseColor,			// BaseColor		-  Metalness workflow
+				TopBaseColor,		// BaseColor		-  Metalness workflow
 				Ones,				// EdgeColor		-  Metalness workflow
 				Specular,			// Specular			-  Metalness workflow
 				0.f,				// Metallic			-  Metalness workflow
@@ -350,32 +303,6 @@ FStrataData StrataConvertLegacyMaterial(
 				0.f,				// FuzzAmount
 				Zeros,				// FuzzColor
 				DefaultThickness,
-				ClearCoatBottomNormal_SharedLocalBasisIndex, /*ClearCoatBottomNormal_*/SharedLocalBasisTypes);
-
-			Out = StrataHorizontalMixing(IrisSlab, Top, IrisOpacity);
-		}
-		else // (IrisMask >= 1.0f)
-		{
-			FStrataData CoatSlab = GetStrataSlabBSDF(
-				true,				// UseMetalness
-				Zeros,				// BaseColor		-  Metalness workflow
-				Ones,				// EdgeColor		-  Metalness workflow
-				Specular,			// Specular			-  Metalness workflow
-				0.f,				// Metallic			-  Metalness workflow
-				Zeros,				// DiffuseAlbedo	- !Metalness workflow
-				Zeros, 				// F0				- !Metalness workflow
-				Zeros,				// F90				- !Metalness workflow
-				Roughness,			// Roughness
-				0.f,				// Anisotropy
-				0.f,				// SSSProfileID
-				10000.f,			// SSSDMFP
-				1.f,				// SSSDMFPScale
-				Emissive,			// Emissive
-				0.f,				// Haziness
-				0.f,				// ThinFilmThickness
-				0.f,				// FuzzAmount
-				Zeros,				// FuzzColor
-				DefaultThickness,
 				SharedLocalBasisIndex, SharedLocalBasisTypes);
 
 			FStrataData IrisSlab = GetStrataSlabBSDF(
