Restructure how the material data is written for the fast path to avoid register spilling when compiler more complex materials.

diff --git a/Engine/Shaders/Private/Strata/Strata.ush b/Engine/Shaders/Private/Strata/Strata.ush
--- a/Engine/Shaders/Private/Strata/Strata.ush
+++ b/Engine/Shaders/Private/Strata/Strata.ush
@@ -2099,64 +2099,34 @@ void PackStrataOut(
 	// Now write out Strata data
 
 	///////////////////////////////////////////////////////////////////////////
-	// A. Fast encoding
+	// Two types of encodings
+	// * A. Fast encoding
+	// * B. Regular complex
 	const bool bHasFastEncoding = IsSimpleMaterial(StrataPixelHeader);
-	if (bHasFastEncoding)
-	{		
-		FStrataBSDF BSDF = Strata.Layers[SimpleMaterial_LayerIndex].BSDFs[SimpleMaterial_BSDFIndex];
-
-		// Now we pack diffuse and F0 in a special way: both encoded as R7G7B6 with gamma 2.0. The low bit cound will be hidden by dithering and TAA.
-		// We pack 32 bits in the first uint and the remaining bits are pack is the lowest significant bit of the second uint.
-		const uint PackedDiffuse20Bits	= PackR7G7B6Gamma2(SLAB_DIFFUSEALBEDO(BSDF));
-		const uint PackedDiffuse12Bits	= PackedDiffuse20Bits & 0xFFF;
-		const uint PackedDiffuse8Bits	= (PackedDiffuse20Bits>>12) & 0xFF;
-		const uint PackedF020Bits		= PackR7G7B6Gamma2(SLAB_F0(BSDF));
-		const uint PackedRoughness8bits = PackR8(SLAB_ROUGHNESS(BSDF));
-
-		// Data0 (State|AO|Roughness|Diffuse8bits)
-		{
-			uint Out = 0;
-			HEADER_SETSTATE(Out, StrataPixelHeader.State);
-			HEADER_SETAO(Out, PackR8(StrataPixelHeader.MaterialAO));
-
-			uint Packed16 = (PackedDiffuse8Bits << 8) | PackedRoughness8bits;
-			Out = (Out & HEADER_FASTENCODING_MASK) | (Packed16 << HEADER_FASTENCODING_BIT_COUNT);
-			STRATA_STORE_UINT1(Out);
 
-			#if (HEADER_FASTENCODING_BIT_COUNT + 16) > 32
-			#error Strata fast path header is > 32bits
-			#endif
-		}
-		// Data1 (F0|Diffuse12bits)
-		{
-			uint Out = PackedF020Bits | (PackedDiffuse12Bits << 20);
-			STRATA_STORE_UINT1(Out);
-		}
-	}
-	///////////////////////////////////////////////////////////////////////////
-	// B. Regular/Complex encoding
-	else // (!bHasFastEncoding)
+	if (!bHasFastEncoding)
 	{
-
-	// 1. the header
-	const uint PackedHeader = PackStrataHeader(BSDFVisibleCount, StrataPixelHeader);
-	STRATA_STORE_UINT1(PackedHeader);
-
-#if STRATA_INLINE_SHADING
-	UNROLL
-	for (uint i = 0; i < StrataPixelHeader.SharedLocalBases.Count; ++i)
-	{
-		const uint BasisType = StrataGetSharedLocalBasisType(StrataPixelHeader.SharedLocalBases.Types, i);
-		if (BasisType == STRATA_BASIS_TYPE_NORMAL)
-		{
-			STRATA_STORE_UINT1(StrataPackNormal(StrataPixelHeader.SharedLocalBases.Normals[i]));
-		}
-		else // if (BasisType == STRATA_BASIS_TYPE_TANGENT)
+		// 1. the header (Regular/Complex encoding)
+		const uint PackedHeader = PackStrataHeader(BSDFVisibleCount, StrataPixelHeader);
+		STRATA_STORE_UINT1(PackedHeader);
+
+		// 1.1 tangent basis (Regular/Complex encoding)
+	#if STRATA_INLINE_SHADING
+		UNROLL
+		for (uint i = 0; i < StrataPixelHeader.SharedLocalBases.Count; ++i)
 		{
-			STRATA_STORE_UINT1(StrataPackNormalAndTangent(StrataPixelHeader.SharedLocalBases.Normals[i], StrataPixelHeader.SharedLocalBases.Tangents[i]));
+			const uint BasisType = StrataGetSharedLocalBasisType(StrataPixelHeader.SharedLocalBases.Types, i);
+			if (BasisType == STRATA_BASIS_TYPE_NORMAL)
+			{
+				STRATA_STORE_UINT1(StrataPackNormal(StrataPixelHeader.SharedLocalBases.Normals[i]));
+			}
+			else // if (BasisType == STRATA_BASIS_TYPE_TANGENT)
+			{
+				STRATA_STORE_UINT1(StrataPackNormalAndTangent(StrataPixelHeader.SharedLocalBases.Normals[i], StrataPixelHeader.SharedLocalBases.Tangents[i]));
+			}
 		}
+	#endif
 	}
-#endif
 
 	// 2. the list of BSDFs
 	BSDFCount = 0;
@@ -2176,7 +2146,42 @@ void PackStrataOut(
 			{
 				const uint GreyScaleThroughput = StrataHasGreyScaleWeight(BSDF.FinalWeight) ? 1 : 0;
 				BSDF_SETHASGREYWEIGHT(BSDF, GreyScaleThroughput);
-				if (GreyScaleThroughput > 0)
+
+				// A. Fast path encoding: store header & data
+				if (bHasFastEncoding && l == SimpleMaterial_LayerIndex && i == SimpleMaterial_BSDFIndex)
+				{		
+					// Now we pack diffuse and F0 in a special way: both encoded as R7G7B6 with gamma 2.0. The low bit cound will be hidden by dithering and TAA.
+					// We pack 32 bits in the first uint and the remaining bits are pack is the lowest significant bit of the second uint.
+					const uint PackedDiffuse20Bits	= PackR7G7B6Gamma2(SLAB_DIFFUSEALBEDO(BSDF));
+					const uint PackedDiffuse12Bits	= PackedDiffuse20Bits & 0xFFF;
+					const uint PackedDiffuse8Bits	= (PackedDiffuse20Bits>>12) & 0xFF;
+					const uint PackedF020Bits		= PackR7G7B6Gamma2(SLAB_F0(BSDF));
+					const uint PackedRoughness8bits = PackR8(SLAB_ROUGHNESS(BSDF));
+
+					// Data0 (State|AO|Roughness|Diffuse8bits)
+					{
+						uint Out = 0;
+						HEADER_SETSTATE(Out, StrataPixelHeader.State);
+						HEADER_SETAO(Out, PackR8(StrataPixelHeader.MaterialAO));
+
+						uint Packed16 = (PackedDiffuse8Bits << 8) | PackedRoughness8bits;
+						Out = (Out & HEADER_FASTENCODING_MASK) | (Packed16 << HEADER_FASTENCODING_BIT_COUNT);
+						STRATA_STORE_UINT1(Out);
+
+						#if (HEADER_FASTENCODING_BIT_COUNT + 16) > 32
+						#error Strata fast path header is > 32bits
+						#endif
+					}
+					// Data1 (F0|Diffuse12bits)
+					{
+						uint Out = PackedF020Bits | (PackedDiffuse12Bits << 20);
+						STRATA_STORE_UINT1(Out);
+					}
+					// Ensure the rest of the BSDF is not stored with the regular path
+					BSDF.State = 0;
+				}
+				// B. Weight for Regular/Complex path
+				else if (GreyScaleThroughput > 0)
 				{
 					BSDF_SETWEIGHT10F(BSDF, Pack10F(BSDF.FinalWeight.x));
 					STRATA_STORE_UINT1(BSDF.State);
@@ -2278,7 +2283,6 @@ void PackStrataOut(
 			}
 		}
 	}
-	} // if (bHasFastEncoding)
 
 	FinalizeWrites(StrataBuffer, StrataAddressing);
 
