Change Strata pack out loop to have local layer index.

diff --git a/Engine/Shaders/Private/Strata/Strata.ush b/Engine/Shaders/Private/Strata/Strata.ush
--- a/Engine/Shaders/Private/Strata/Strata.ush
+++ b/Engine/Shaders/Private/Strata/Strata.ush
@@ -1893,206 +1893,207 @@ void PackStrataOut(
 	Strata.LayerCount = min(Strata.LayerCount, STRATA_CLAMPED_LAYER_COUNT);
 
 	// For each of the BSDF, we evaluated their combined coverage/throughput according to the material operation performed (layering, mixing, etc.).
-	int l;
-	STRATA_UNROLL_N(STRATA_CLAMPED_LAYER_COUNT)
-	for (l = 0; l < Strata.LayerCount; ++l)
 	{
-		float3 LayerThroughput = 0.0f;
-		float3 LayerCoverage = 0.0f;
+		STRATA_UNROLL_N(STRATA_CLAMPED_LAYER_COUNT)
+		for (int l = 0; l < Strata.LayerCount; ++l)
+		{
+			float3 LayerThroughput = 0.0f;
+			float3 LayerCoverage = 0.0f;
 
-		const bool bTopLayer = l == 0;
-		const bool bBottomLayer = l == (Strata.LayerCount - 1);
+			const bool bTopLayer = l == 0;
+			const bool bBottomLayer = l == (Strata.LayerCount - 1);
 
-		STRATA_UNROLL_N(STRATA_CLAMPED_MAX_BSDF_COUNT_PER_LAYER)
-		for (int i = 0; i < Strata.Layers[l].BSDFCount; ++i)
-		{
-			// Dither BSDF data to hide more aggressive bit reduction.
-			// Sanitize happening after will clamp values that could have gone above and below valid range.
-			StrataDitherBSDF(Strata.Layers[l].BSDFs[i], Dither);
+			STRATA_UNROLL_N(STRATA_CLAMPED_MAX_BSDF_COUNT_PER_LAYER)
+			for (int i = 0; i < Strata.Layers[l].BSDFCount; ++i)
+			{
+				// Dither BSDF data to hide more aggressive bit reduction.
+				// Sanitize happening after will clamp values that could have gone above and below valid range.
+				StrataDitherBSDF(Strata.Layers[l].BSDFs[i], Dither);
 
-			// Sanitize BSDF in place before it is used and packed out to memory
-			StrataSanitizeBSDF(Strata.Layers[l].BSDFs[i]);
+				// Sanitize BSDF in place before it is used and packed out to memory
+				StrataSanitizeBSDF(Strata.Layers[l].BSDFs[i]);
 
-			FStrataBSDF BSDF = Strata.Layers[l].BSDFs[i];
+				FStrataBSDF BSDF = Strata.Layers[l].BSDFs[i];
 
-		#if STRATA_INLINE_SHADING
-			const float BSDFCoverage = saturate(BSDF.Coverage);
-		#else
-			const float BSDFCoverage = 1.0f;
-		#endif
-			BSDFCount++;
+			#if STRATA_INLINE_SHADING
+				const float BSDFCoverage = saturate(BSDF.Coverage);
+			#else
+				const float BSDFCoverage = 1.0f;
+			#endif
+				BSDFCount++;
 
-			// STRATA_TODO apply the same normal transform as for the material
-			float3x3 TangentBasis = StrataGetBSDFSharedBasis(StrataPixelHeader, BSDF, StrataAddressing);
-			float3 N = TangentBasis[2];
-			//N = normalize(N);
+				// STRATA_TODO apply the same normal transform as for the material
+				float3x3 TangentBasis = StrataGetBSDFSharedBasis(StrataPixelHeader, BSDF, StrataAddressing);
+				float3 N = TangentBasis[2];
+				//N = normalize(N);
 
-			BxDFContext ContextNoLight = (BxDFContext)0;
-			const float3 FakeL = float3(0, 0, 1);
-			Init(ContextNoLight, N, V, FakeL);
+				BxDFContext ContextNoLight = (BxDFContext)0;
+				const float3 FakeL = float3(0, 0, 1);
+				Init(ContextNoLight, N, V, FakeL);
 
-			bool bIsSimpleMaterial = true;
-			switch (BSDF_GETTYPE(BSDF))
-			{
-			case STRATA_BSDF_TYPE_SLAB:
-			{
-				bIsSimpleMaterial = IsStrataSlabFastPathCompatible(BSDF);
+				bool bIsSimpleMaterial = true;
+				switch (BSDF_GETTYPE(BSDF))
+				{
+				case STRATA_BSDF_TYPE_SLAB:
+				{
+					bIsSimpleMaterial = IsStrataSlabFastPathCompatible(BSDF);
 
-				EmissiveColor += BSDF_GETEMISSIVE(BSDF) * Throughput;
+					EmissiveColor += BSDF_GETEMISSIVE(BSDF) * Throughput;
 
-				// Two sided lighting can only be evalauted on the bottom layer, together with option SSS. Otherwise, the simple volume will be used.
-				const bool bIsThin = bBottomLayer && BSDF_GETISTHIN(Strata.Layers[l].BSDFs[i]);
+					// Two sided lighting can only be evalauted on the bottom layer, together with option SSS. Otherwise, the simple volume will be used.
+					const bool bIsThin = bBottomLayer && BSDF_GETISTHIN(Strata.Layers[l].BSDFs[i]);
 
-				// We want simple volumetric only if the DMFP input is pluged in, otherwise we avoid storing the DFMP
-				const bool bIsSimpleVolume = !bBottomLayer && BSDF_GETHASDMFP(Strata.Layers[l].BSDFs[i]);
+					// We want simple volumetric only if the DMFP input is pluged in, otherwise we avoid storing the DFMP
+					const bool bIsSimpleVolume = !bBottomLayer && BSDF_GETHASDMFP(Strata.Layers[l].BSDFs[i]);
 
-				bStrataSubsurfaceEnable = bStrataSubsurfaceEnable || (BSDF_GETHASSSS(BSDF) && BSDFCoverage > 0 && !bIsSimpleVolume);
+					bStrataSubsurfaceEnable = bStrataSubsurfaceEnable || (BSDF_GETHASSSS(BSDF) && BSDFCoverage > 0 && !bIsSimpleVolume);
 
-				// By default, without coverage, the diffuse BSDF is completely opaque
-				float3 BSDFThroughput = OpaqueBSDFThroughput;
+					// By default, without coverage, the diffuse BSDF is completely opaque
+					float3 BSDFThroughput = OpaqueBSDFThroughput;
 
-				if (bIsSimpleVolume)
-				{
-					float3 ApproxFresnelForTransmittance = 0;
-					// STRATA_TODO Fresnel throughput, Use a LUT, this is currently in progress and will be coming soon.
-					// In the meantime, we use a simple approach to feed Fresnel function and compute an approximation of the general transmition/througput.
-					const float SatNoV = saturate(ContextNoLight.NoV);
-					const float VoH = SatNoV;
-					const float NoL = SatNoV;
-					const float3 F0 = SLAB_F0(BSDF);
-					const float3 F90 = SLAB_F90(BSDF);
-					const bool bHasThinFilm = BSDF_GETHASTHINFILM(BSDF);
-					BRANCH
-					if (bHasThinFilm)
+					if (bIsSimpleVolume)
+					{
+						float3 ApproxFresnelForTransmittance = 0;
+						// STRATA_TODO Fresnel throughput, Use a LUT, this is currently in progress and will be coming soon.
+						// In the meantime, we use a simple approach to feed Fresnel function and compute an approximation of the general transmition/througput.
+						const float SatNoV = saturate(ContextNoLight.NoV);
+						const float VoH = SatNoV;
+						const float NoL = SatNoV;
+						const float3 F0 = SLAB_F0(BSDF);
+						const float3 F90 = SLAB_F90(BSDF);
+						const bool bHasThinFilm = BSDF_GETHASTHINFILM(BSDF);
+						BRANCH
+						if (bHasThinFilm)
+						{
+							float ThinFilmIOR = SLAB_THINFILM_IOR(BSDF);
+							float ThinFilmThickness = SLAB_THINFILM_THICKNESS(BSDF);
+							ApproxFresnelForTransmittance = F_ThinFilm(SatNoV, NoL, VoH, F0, F90, ThinFilmIOR, ThinFilmThickness);
+						}
+						else
+						{
+							ApproxFresnelForTransmittance = F_Schlick(F0, F90, VoH);
+						}
+
+						EnableSlabBSDFSimpleVolumetric(Strata.Layers[l].BSDFs[i]);
+
+						const float3 UnusedDiffuseColor = 0; // Unused when only trasmittance is needed
+						FParticipatingMedia PM = StrataSlabCreateParticipatingMedia(UnusedDiffuseColor, SLAB_SSSDMFP(Strata.Layers[l].BSDFs[i]));
+						const float DiffuseToVolumeBlend = StrataSlabDiffuseToVolumeBlend(PM);
+						const float3 SlabTransmittance = IsotropicMediumSlabTransmittance(PM, STRATA_SIMPLEVOLUME_THICKNESS_M, ContextNoLight.NoV);
+
+						// Update the see-through throughput accoding to the medium transmittance
+						BSDFThroughput = lerp(OpaqueBSDFThroughput, (1- ApproxFresnelForTransmittance) * SlabTransmittance, DiffuseToVolumeBlend);
+
+						// Mark the BSDF has not having any SSS through post process.
+						bStrataSubsurfaceEnable = false;
+					}
+
+					if (bIsThin)
 					{
-						float ThinFilmIOR = SLAB_THINFILM_IOR(BSDF);
-						float ThinFilmThickness = SLAB_THINFILM_THICKNESS(BSDF);
-						ApproxFresnelForTransmittance = F_ThinFilm(SatNoV, NoL, VoH, F0, F90, ThinFilmIOR, ThinFilmThickness);
+						EnableSlabBSDFIsThin(Strata.Layers[l].BSDFs[i]);
 					}
 					else
 					{
-						ApproxFresnelForTransmittance = F_Schlick(F0, F90, VoH);
+						// If not selected or not possible to have two sided lighting enabled (not bottom layer or SSSProfile is enabled), let's make sure we disable the feature.
+						BSDF_SETISTHIN(Strata.Layers[l].BSDFs[i], 0);
 					}
 
-					EnableSlabBSDFSimpleVolumetric(Strata.Layers[l].BSDFs[i]);
-
-					const float3 UnusedDiffuseColor = 0; // Unused when only trasmittance is needed
-					FParticipatingMedia PM = StrataSlabCreateParticipatingMedia(UnusedDiffuseColor, SLAB_SSSDMFP(Strata.Layers[l].BSDFs[i]));
-					const float DiffuseToVolumeBlend = StrataSlabDiffuseToVolumeBlend(PM);
-					const float3 SlabTransmittance = IsotropicMediumSlabTransmittance(PM, STRATA_SIMPLEVOLUME_THICKNESS_M, ContextNoLight.NoV);
+					if (bTopLayer)
+					{
+						TopLayerData.Roughness += BSDFCoverage * SLAB_ROUGHNESS(BSDF);
+					}
 
-					// Update the see-through throughput accoding to the medium transmittance
-					BSDFThroughput = lerp(OpaqueBSDFThroughput, (1- ApproxFresnelForTransmittance) * SlabTransmittance, DiffuseToVolumeBlend);
+					// STRATA_TODO	Should we have the Fuzz influence the next layter visibility accoding to fuzz amount? Should it also occlude diffuse?
+					//				If this is done, we should also update ComputePostProcessBSDFData.
 
-					// Mark the BSDF has not having any SSS through post process.
-					bStrataSubsurfaceEnable = false;
+					LayerCoverage	+= BSDFCoverage;
+					LayerThroughput += BSDFCoverage * BSDFThroughput;
+					break;
 				}
 
-				if (bIsThin)
-				{
-					EnableSlabBSDFIsThin(Strata.Layers[l].BSDFs[i]);
-				}
-				else
+				case STRATA_BSDF_TYPE_HAIR:
 				{
-					// If not selected or not possible to have two sided lighting enabled (not bottom layer or SSSProfile is enabled), let's make sure we disable the feature.
-					BSDF_SETISTHIN(Strata.Layers[l].BSDFs[i], 0);
-				}
+					bIsSimpleMaterial = false;
 
-				if (bTopLayer)
-				{
-					TopLayerData.Roughness += BSDFCoverage * SLAB_ROUGHNESS(BSDF);
+					if (bTopLayer)
+					{
+						TopLayerData.Roughness += BSDFCoverage * FullyRough;
+					}
+
+					// Hair are considered opaque
+					LayerCoverage	+= BSDFCoverage;
+					LayerThroughput += BSDFCoverage * OpaqueBSDFThroughput;
+					break;
 				}
 
-				// STRATA_TODO	Should we have the Fuzz influence the next layter visibility accoding to fuzz amount? Should it also occlude diffuse?
-				//				If this is done, we should also update ComputePostProcessBSDFData.
+				case STRATA_BSDF_TYPE_SINGLELAYERWATER:
+				{
+					bIsSimpleMaterial = false;
 
-				LayerCoverage	+= BSDFCoverage;
-				LayerThroughput += BSDFCoverage * BSDFThroughput;
-				break;
-			}
+					EmissiveColor += BSDF_GETEMISSIVE(BSDF) * Throughput;
 
-			case STRATA_BSDF_TYPE_HAIR:
-			{
-				bIsSimpleMaterial = false;
+					if (bTopLayer)
+					{
+						TopLayerData.Roughness += BSDFCoverage * SLW_ROUGHNESS(BSDF);
+					}
 
-				if (bTopLayer)
-				{
-					TopLayerData.Roughness += BSDFCoverage * FullyRough;
+					// Water is considered opaque, it has a custom process to simulate translucency
+					LayerCoverage += BSDFCoverage;
+					LayerThroughput += BSDFCoverage * OpaqueBSDFThroughput;
+					break;
 				}
 
-				// Hair are considered opaque
-				LayerCoverage	+= BSDFCoverage;
-				LayerThroughput += BSDFCoverage * OpaqueBSDFThroughput;
-				break;
-			}
+				} // switch
 
-			case STRATA_BSDF_TYPE_SINGLELAYERWATER:
-			{
-				bIsSimpleMaterial = false;
+			#if STRATA_INLINE_SHADING
+				// Update the BSDF final weight with throughput towards the view from layers above and coverage if it has not been baked into the parameters.
+				Strata.Layers[l].BSDFs[i].FinalWeight = Throughput * (Strata.Layers[l].BSDFs[i].bPremultipliedCoverage ? 1.0 : Strata.Layers[l].BSDFs[i].Coverage);
+			#else
+				Strata.Layers[l].BSDFs[i].FinalWeight = Throughput;
+			#endif
+			
+				// Cull non visible BSDF
+				const bool bIsVisible = any(Strata.Layers[l].BSDFs[i].FinalWeight > 0);
+				if (bIsVisible)
+				{
+					++BSDFVisibleCount;
 
-				EmissiveColor += BSDF_GETEMISSIVE(BSDF) * Throughput;
+					AddShadingModels(StrataPixelHeader, StrataShadingModelBit(BSDF_GETTYPE(BSDF)));
+					AndIsSimpleMaterial(StrataPixelHeader, bIsSimpleMaterial);
 
-				if (bTopLayer)
-				{
-					TopLayerData.Roughness += BSDFCoverage * SLW_ROUGHNESS(BSDF);
+					if (bTopLayer)
+					{
+						// together with a threshold to avoid 0-length normal after normalization in case influence is fadeout using the Weight node.
+						TopLayerData.WorldNormal += N * max(1e-5, BSDFCoverage);
+					}
+
+					if (bIsSimpleMaterial)
+					{
+						SimpleMaterial_LayerIndex = l;
+						SimpleMaterial_BSDFIndex = i;
+					}
 				}
 
-				// Water is considered opaque, it has a custom process to simulate translucency
-				LayerCoverage += BSDFCoverage;
-				LayerThroughput += BSDFCoverage * OpaqueBSDFThroughput;
-				break;
+				// Notify that the BSDF is at the top for SSR to only affect reflection there and not on the lower layers
+				BSDF_SETISTOPLAYER(Strata.Layers[l].BSDFs[i], bTopLayer ? 1 : 0);
 			}
 
-			} // switch
-
-		#if STRATA_INLINE_SHADING
-			// Update the BSDF final weight with throughput towards the view from layers above and coverage if it has not been baked into the parameters.
-			Strata.Layers[l].BSDFs[i].FinalWeight = Throughput * (Strata.Layers[l].BSDFs[i].bPremultipliedCoverage ? 1.0 : Strata.Layers[l].BSDFs[i].Coverage);
-		#else
-			Strata.Layers[l].BSDFs[i].FinalWeight = Throughput;
-		#endif
-			
-			// Cull non visible BSDF
-			const bool bIsVisible = any(Strata.Layers[l].BSDFs[i].FinalWeight > 0);
-			if (bIsVisible)
+			// Average the TopLayer data once we have read all the data
+			if (bTopLayer)
 			{
-				++BSDFVisibleCount;
-
-				AddShadingModels(StrataPixelHeader, StrataShadingModelBit(BSDF_GETTYPE(BSDF)));
-				AndIsSimpleMaterial(StrataPixelHeader, bIsSimpleMaterial);
+				// Top layer roughness & normal have been weighted with BSDF Coverage only (since the accumulated throughput is 1.f for the top-layer)
+				const float BSDFWeightAvg = dot(LayerCoverage, 1.0f / 3.0f);
 
-				if (bTopLayer)
-				{
-					// together with a threshold to avoid 0-length normal after normalization in case influence is fadeout using the Weight node.
-					TopLayerData.WorldNormal += N * max(1e-5, BSDFCoverage);
-				}
-
-				if (bIsSimpleMaterial)
-				{
-					SimpleMaterial_LayerIndex = l;
-					SimpleMaterial_BSDFIndex = i;
-				}
+				TopLayerData.WorldNormal = normalize(TopLayerData.WorldNormal);
+				TopLayerData.Roughness = BSDFWeightAvg > 0.0f ? TopLayerData.Roughness / BSDFWeightAvg : 0.0f;
 			}
 
-			// Notify that the BSDF is at the top for SSR to only affect reflection there and not on the lower layers
-			BSDF_SETISTOPLAYER(Strata.Layers[l].BSDFs[i], bTopLayer ? 1 : 0);
-		}
-
-		// Average the TopLayer data once we have read all the data
-		if (bTopLayer)
-		{
-			// Top layer roughness & normal have been weighted with BSDF Coverage only (since the accumulated throughput is 1.f for the top-layer)
-			const float BSDFWeightAvg = dot(LayerCoverage, 1.0f / 3.0f);
+			// Take into account the non covered part for this layer and update throughput accordingly
+			LayerThroughput += (1.0 - saturate(LayerCoverage)) * FullThroughput;
 
-			TopLayerData.WorldNormal = normalize(TopLayerData.WorldNormal);
-			TopLayerData.Roughness = BSDFWeightAvg > 0.0f ? TopLayerData.Roughness / BSDFWeightAvg : 0.0f;
+			// And combine this layer throughtput with he material throughput
+			Throughput *= LayerThroughput;
 		}
-
-		// Take into account the non covered part for this layer and update throughput accordingly
-		LayerThroughput += (1.0 - saturate(LayerCoverage)) * FullThroughput;
-
-		// And combine this layer throughtput with he material throughput
-		Throughput *= LayerThroughput;
 	}
 	AndIsSimpleMaterial(StrataPixelHeader, BSDFVisibleCount == 1);
 
@@ -2130,156 +2131,158 @@ void PackStrataOut(
 
 	// 2. the list of BSDFs
 	BSDFCount = 0;
-	STRATA_UNROLL_N(STRATA_CLAMPED_LAYER_COUNT)
-	for (l = 0; l < Strata.LayerCount; ++l)
 	{
-		FStrataLayer Layer = Strata.Layers[l];
-
-		STRATA_UNROLL_N(STRATA_CLAMPED_MAX_BSDF_COUNT_PER_LAYER)
-		for (int i = 0; i < Layer.BSDFCount; ++i)
+		STRATA_UNROLL_N(STRATA_CLAMPED_LAYER_COUNT)
+		for (int l = 0; l < Strata.LayerCount; ++l)
 		{
-			FStrataBSDF BSDF = Layer.BSDFs[i];
+			FStrataLayer Layer = Strata.Layers[l];
 
-			// Only write visible BSDF
-			BRANCH
-			if (any(BSDF.FinalWeight > 0))
+			STRATA_UNROLL_N(STRATA_CLAMPED_MAX_BSDF_COUNT_PER_LAYER)
+			for (int i = 0; i < Layer.BSDFCount; ++i)
 			{
-				const uint GreyScaleThroughput = StrataHasGreyScaleWeight(BSDF.FinalWeight) ? 1 : 0;
-				BSDF_SETHASGREYWEIGHT(BSDF, GreyScaleThroughput);
-
-				// A. Fast path encoding: store header & data
-				if (bHasFastEncoding && l == SimpleMaterial_LayerIndex && i == SimpleMaterial_BSDFIndex)
-				{		
-					// Now we pack diffuse and F0 in a special way: both encoded as R7G7B6 with gamma 2.0. The low bit cound will be hidden by dithering and TAA.
-					// We pack 32 bits in the first uint and the remaining bits are pack is the lowest significant bit of the second uint.
-					const uint PackedDiffuse20Bits	= PackR7G7B6Gamma2(SLAB_DIFFUSEALBEDO(BSDF));
-					const uint PackedDiffuse12Bits	= PackedDiffuse20Bits & 0xFFF;
-					const uint PackedDiffuse8Bits	= (PackedDiffuse20Bits>>12) & 0xFF;
-					const uint PackedF020Bits		= PackR7G7B6Gamma2(SLAB_F0(BSDF));
-					const uint PackedRoughness8bits = PackR8(SLAB_ROUGHNESS(BSDF));
-
-					// Data0 (State|AO|Roughness|Diffuse8bits)
-					{
-						uint Out = 0;
-						HEADER_SETSTATE(Out, StrataPixelHeader.State);
-						HEADER_SETAO(Out, PackR8(StrataPixelHeader.MaterialAO));
+				FStrataBSDF BSDF = Layer.BSDFs[i];
+
+				// Only write visible BSDF
+				BRANCH
+				if (any(BSDF.FinalWeight > 0))
+				{
+					const uint GreyScaleThroughput = StrataHasGreyScaleWeight(BSDF.FinalWeight) ? 1 : 0;
+					BSDF_SETHASGREYWEIGHT(BSDF, GreyScaleThroughput);
+
+					// A. Fast path encoding: store header & data
+					if (bHasFastEncoding && l == SimpleMaterial_LayerIndex && i == SimpleMaterial_BSDFIndex)
+					{		
+						// Now we pack diffuse and F0 in a special way: both encoded as R7G7B6 with gamma 2.0. The low bit cound will be hidden by dithering and TAA.
+						// We pack 32 bits in the first uint and the remaining bits are pack is the lowest significant bit of the second uint.
+						const uint PackedDiffuse20Bits	= PackR7G7B6Gamma2(SLAB_DIFFUSEALBEDO(BSDF));
+						const uint PackedDiffuse12Bits	= PackedDiffuse20Bits & 0xFFF;
+						const uint PackedDiffuse8Bits	= (PackedDiffuse20Bits>>12) & 0xFF;
+						const uint PackedF020Bits		= PackR7G7B6Gamma2(SLAB_F0(BSDF));
+						const uint PackedRoughness8bits = PackR8(SLAB_ROUGHNESS(BSDF));
+
+						// Data0 (State|AO|Roughness|Diffuse8bits)
+						{
+							uint Out = 0;
+							HEADER_SETSTATE(Out, StrataPixelHeader.State);
+							HEADER_SETAO(Out, PackR8(StrataPixelHeader.MaterialAO));
 
-						uint Packed16 = (PackedDiffuse8Bits << 8) | PackedRoughness8bits;
-						Out = (Out & HEADER_FASTENCODING_MASK) | (Packed16 << HEADER_FASTENCODING_BIT_COUNT);
-						STRATA_STORE_UINT1(Out);
+							uint Packed16 = (PackedDiffuse8Bits << 8) | PackedRoughness8bits;
+							Out = (Out & HEADER_FASTENCODING_MASK) | (Packed16 << HEADER_FASTENCODING_BIT_COUNT);
+							STRATA_STORE_UINT1(Out);
 
-						#if (HEADER_FASTENCODING_BIT_COUNT + 16) > 32
-						#error Strata fast path header is > 32bits
-						#endif
+							#if (HEADER_FASTENCODING_BIT_COUNT + 16) > 32
+							#error Strata fast path header is > 32bits
+							#endif
+						}
+						// Data1 (F0|Diffuse12bits)
+						{
+							uint Out = PackedF020Bits | (PackedDiffuse12Bits << 20);
+							STRATA_STORE_UINT1(Out);
+						}
+						// Ensure the rest of the BSDF is not stored with the regular path
+						BSDF.State = 0;
 					}
-					// Data1 (F0|Diffuse12bits)
+					// B. Weight for Regular/Complex path
+					else if (GreyScaleThroughput > 0)
 					{
-						uint Out = PackedF020Bits | (PackedDiffuse12Bits << 20);
-						STRATA_STORE_UINT1(Out);
+						BSDF_SETWEIGHT10F(BSDF, Pack10F(BSDF.FinalWeight.x));
+						STRATA_STORE_UINT1(BSDF.State);
 					}
-					// Ensure the rest of the BSDF is not stored with the regular path
-					BSDF.State = 0;
-				}
-				// B. Weight for Regular/Complex path
-				else if (GreyScaleThroughput > 0)
-				{
-					BSDF_SETWEIGHT10F(BSDF, Pack10F(BSDF.FinalWeight.x));
-					STRATA_STORE_UINT1(BSDF.State);
-				}
-				else
-				{
-					STRATA_STORE_UINT1(BSDF.State);
-					STRATA_STORE_UINT1(PackR11G11B10F(BSDF.FinalWeight));
-				}
-
-				const uint BSDFType = BSDF_GETTYPE(BSDF);
-
-				switch (BSDF_GETTYPE(BSDF))
-				{
-				case STRATA_BSDF_TYPE_SLAB:
-				{
-					// Now we pack diffuse and F0 in a special way: both encoded as R7G7B6 with gamma 2.0. The low bit cound will be hidden by dithering and TAA.
-					// We pack 32 bits in the first uint and the remaining bits are pack is the lowest significant bit of the second uint.
-					const uint PackedDiffuse20Bits	= PackR7G7B6Gamma2(SLAB_DIFFUSEALBEDO(BSDF));
-					const uint PackedF020Bits		= PackR7G7B6Gamma2(SLAB_F0(BSDF));
-					const uint PackedData32Bits		= ((PackedDiffuse20Bits << 12) & 0xFFFFF000) | (PackedF020Bits & 0xFFF);
-					const uint PackedData8Bits		= (PackedF020Bits >> 12) & 0xFF; 
-
-					STRATA_STORE_UINT1(PackedData32Bits);
-					STRATA_STORE_UINT1(PackedData8Bits | PackRGBA8(float4(0.0f , SLAB_ROUGHNESS(BSDF), (SLAB_ANISOTROPY(BSDF) + 1.f) * 0.5f, 0.f)));
-
-					if (BSDF_GETHASF90(BSDF) || BSDF_GETHASTHINFILM(BSDF) || BSDF_GETHASHAZINESS(BSDF))
+					else
 					{
-						// What is important is to maintain the hue and saturation, so we scale the color by the maximum of its components
-						float3 F90 = saturate(SLAB_F90(BSDF));
-						const float Divisor = max(F90.r, max(F90.g, F90.b));
-						F90 = Divisor > 0.0f ? F90 / Divisor : 1.0f;
-						float3 F90YCoCg = LinearRGB_2_NormalisedYCoCg(F90);
-
-						const float ThinFilmGamma2 = sqrt(SLAB_THINFILM_THICKNESS(BSDF)); // Using a gamma 2.0 encoding help with precision to be visually more uniform. This is then further improved via dithering.
-						STRATA_STORE_UINT1(PackRGBA8(float4(F90YCoCg.y, F90YCoCg.z, ThinFilmGamma2, SLAB_HAZINESS(BSDF))));
+						STRATA_STORE_UINT1(BSDF.State);
+						STRATA_STORE_UINT1(PackR11G11B10F(BSDF.FinalWeight));
 					}
 
-					const bool bIsSimpleVolume = BSDF_GETISSIMPLEVOLUME(BSDF);
-					const bool bIsThin = BSDF_GETISTHIN(BSDF);
-					if (BSDF_GETHASSSS(BSDF) || bIsSimpleVolume || bIsThin)
+					const uint BSDFType = BSDF_GETTYPE(BSDF);
+
+					switch (BSDF_GETTYPE(BSDF))
 					{
-						SetHasSubsurface(StrataPixelHeader, !bIsThin && BSDF_GETHASSSS(BSDF));
-						SetHasSubsurfaceProfile(StrataPixelHeader, !bIsThin && BSDF_GETHASSSSPROFILE(BSDF));
+					case STRATA_BSDF_TYPE_SLAB:
+					{
+						// Now we pack diffuse and F0 in a special way: both encoded as R7G7B6 with gamma 2.0. The low bit cound will be hidden by dithering and TAA.
+						// We pack 32 bits in the first uint and the remaining bits are pack is the lowest significant bit of the second uint.
+						const uint PackedDiffuse20Bits	= PackR7G7B6Gamma2(SLAB_DIFFUSEALBEDO(BSDF));
+						const uint PackedF020Bits		= PackR7G7B6Gamma2(SLAB_F0(BSDF));
+						const uint PackedData32Bits		= ((PackedDiffuse20Bits << 12) & 0xFFFFF000) | (PackedF020Bits & 0xFFF);
+						const uint PackedData8Bits		= (PackedF020Bits >> 12) & 0xFF; 
+
+						STRATA_STORE_UINT1(PackedData32Bits);
+						STRATA_STORE_UINT1(PackedData8Bits | PackRGBA8(float4(0.0f , SLAB_ROUGHNESS(BSDF), (SLAB_ANISOTROPY(BSDF) + 1.f) * 0.5f, 0.f)));
 
-						if (!bIsThin && BSDF_GETHASSSS(BSDF))
+						if (BSDF_GETHASF90(BSDF) || BSDF_GETHASTHINFILM(BSDF) || BSDF_GETHASHAZINESS(BSDF))
 						{
-							SSSData.BaseColor = SLAB_DIFFUSEALBEDO(BSDF);
-							SSSData.Specular  = F0ToDielectricSpecular(SLAB_F0(BSDF).x);
-
-							// Profile ID: we only keep the latest ID, and profile will overtake any non-profile based SSS
-							SSSData.bIsValid 			= true;
-							SSSData.bIsProfile 			= SSSData.bIsProfile || BSDF_GETHASSSSPROFILE(BSDF);
-							SSSData.ProfileId 			= StrataSubsurfaceProfileIdTo8bits(SLAB_SSSPROFILEID(BSDF));
-							SSSData.ProfileRadiusScale 	= SLAB_SSSPROFILERADIUSSCALE(BSDF);
-
-							// STRATA_TODO This should be accumulated according to the BSDF weight when we run this in the base pass.
-							// We cannot do that now because BSDFWeight contains weight*coverage*trhougput all mixed together.
-							SSSData.DMFP = SLAB_SSSDMFP(BSDF);
+							// What is important is to maintain the hue and saturation, so we scale the color by the maximum of its components
+							float3 F90 = saturate(SLAB_F90(BSDF));
+							const float Divisor = max(F90.r, max(F90.g, F90.b));
+							F90 = Divisor > 0.0f ? F90 / Divisor : 1.0f;
+							float3 F90YCoCg = LinearRGB_2_NormalisedYCoCg(F90);
+
+							const float ThinFilmGamma2 = sqrt(SLAB_THINFILM_THICKNESS(BSDF)); // Using a gamma 2.0 encoding help with precision to be visually more uniform. This is then further improved via dithering.
+							STRATA_STORE_UINT1(PackRGBA8(float4(F90YCoCg.y, F90YCoCg.z, ThinFilmGamma2, SLAB_HAZINESS(BSDF))));
 						}
 
-						if (BSDF_GETHASSSSPROFILE(BSDF) && !bIsSimpleVolume)
+						const bool bIsSimpleVolume = BSDF_GETISSIMPLEVOLUME(BSDF);
+						const bool bIsThin = BSDF_GETISTHIN(BSDF);
+						if (BSDF_GETHASSSS(BSDF) || bIsSimpleVolume || bIsThin)
 						{
-							// Simple volume overrides SSS profile because it takes over for when not at the bottom of the BSDF layer, or during forward rendering.
-							STRATA_STORE_UINT1(PackRGBA8(float4(SLAB_SSSPROFILEID(BSDF), SLAB_SSSPROFILERADIUSSCALE(BSDF), 0, 0)));
+							SetHasSubsurface(StrataPixelHeader, !bIsThin && BSDF_GETHASSSS(BSDF));
+							SetHasSubsurfaceProfile(StrataPixelHeader, !bIsThin && BSDF_GETHASSSSPROFILE(BSDF));
+
+							if (!bIsThin && BSDF_GETHASSSS(BSDF))
+							{
+								SSSData.BaseColor = SLAB_DIFFUSEALBEDO(BSDF);
+								SSSData.Specular  = F0ToDielectricSpecular(SLAB_F0(BSDF).x);
+
+								// Profile ID: we only keep the latest ID, and profile will overtake any non-profile based SSS
+								SSSData.bIsValid 			= true;
+								SSSData.bIsProfile 			= SSSData.bIsProfile || BSDF_GETHASSSSPROFILE(BSDF);
+								SSSData.ProfileId 			= StrataSubsurfaceProfileIdTo8bits(SLAB_SSSPROFILEID(BSDF));
+								SSSData.ProfileRadiusScale 	= SLAB_SSSPROFILERADIUSSCALE(BSDF);
+
+								// STRATA_TODO This should be accumulated according to the BSDF weight when we run this in the base pass.
+								// We cannot do that now because BSDFWeight contains weight*coverage*trhougput all mixed together.
+								SSSData.DMFP = SLAB_SSSDMFP(BSDF);
+							}
+
+							if (BSDF_GETHASSSSPROFILE(BSDF) && !bIsSimpleVolume)
+							{
+								// Simple volume overrides SSS profile because it takes over for when not at the bottom of the BSDF layer, or during forward rendering.
+								STRATA_STORE_UINT1(PackRGBA8(float4(SLAB_SSSPROFILEID(BSDF), SLAB_SSSPROFILERADIUSSCALE(BSDF), 0, 0)));
+							}
+							else
+							{
+								// Path used for bottom most layer SSS, simple volume and two sided lighting.
+								STRATA_STORE_UINT1(PackR11G11B10F(SLAB_SSSDMFP(BSDF)));
+							}
 						}
-						else
+
+						if (BSDF_GETHASFUZZ(BSDF))
 						{
-							// Path used for bottom most layer SSS, simple volume and two sided lighting.
-							STRATA_STORE_UINT1(PackR11G11B10F(SLAB_SSSDMFP(BSDF)));
+							STRATA_STORE_UINT1(PackColorLinearToGamma2AlphaLinear(float4(SLAB_FUZZ_COLOR(BSDF), SLAB_FUZZ_AMOUNT(BSDF))));
 						}
-					}
 
-					if (BSDF_GETHASFUZZ(BSDF))
+						// 8-20 bytes
+					}
+					break;
+					case STRATA_BSDF_TYPE_HAIR:
+					{
+						STRATA_STORE_UINT1(PackColorLinearToGamma2AlphaLinear(float4(HAIR_BASECOLOR(BSDF), HAIR_ROUGHNESS(BSDF))));
+						STRATA_STORE_UINT1(PackRGBA8(float4(HAIR_SCATTER(BSDF), HAIR_SPECULAR(BSDF), HAIR_BACKLIT(BSDF), HAIR_COMPLEXTRANSMITTANCE(BSDF))));
+						// 8 bytes
+					}
+					break;
+					case STRATA_BSDF_TYPE_SINGLELAYERWATER:
 					{
-						STRATA_STORE_UINT1(PackColorLinearToGamma2AlphaLinear(float4(SLAB_FUZZ_COLOR(BSDF), SLAB_FUZZ_AMOUNT(BSDF))));
+						STRATA_STORE_UINT1(PackColorLinearToGamma2AlphaLinear(float4(SLW_BASECOLOR(BSDF), 0.0f)));
+						STRATA_STORE_UINT1(PackRGBA8(float4(SLW_METALLIC(BSDF), SLW_SPECULAR(BSDF), SLW_ROUGHNESS(BSDF), SLW_TOPMATERIALOPACITY(BSDF))));
+					}
+					break;
 					}
 
-					// 8-20 bytes
-				}
-				break;
-				case STRATA_BSDF_TYPE_HAIR:
-				{
-					STRATA_STORE_UINT1(PackColorLinearToGamma2AlphaLinear(float4(HAIR_BASECOLOR(BSDF), HAIR_ROUGHNESS(BSDF))));
-					STRATA_STORE_UINT1(PackRGBA8(float4(HAIR_SCATTER(BSDF), HAIR_SPECULAR(BSDF), HAIR_BACKLIT(BSDF), HAIR_COMPLEXTRANSMITTANCE(BSDF))));
-					// 8 bytes
+					// Stop writing if above budget.
+					BSDFCount++;
 				}
-				break;
-				case STRATA_BSDF_TYPE_SINGLELAYERWATER:
-				{
-					STRATA_STORE_UINT1(PackColorLinearToGamma2AlphaLinear(float4(SLW_BASECOLOR(BSDF), 0.0f)));
-					STRATA_STORE_UINT1(PackRGBA8(float4(SLW_METALLIC(BSDF), SLW_SPECULAR(BSDF), SLW_ROUGHNESS(BSDF), SLW_TOPMATERIALOPACITY(BSDF))));
-				}
-				break;
-				}
-
-				// Stop writing if above budget.
-				BSDFCount++;
 			}
 		}
 	}
