Add separable filter support for non-Subsurface profile subsurface scattering material using strata.

diff --git a/Engine/Shaders/Private/BurleyNormalizedSSSCommon.ush b/Engine/Shaders/Private/BurleyNormalizedSSSCommon.ush
--- a/Engine/Shaders/Private/BurleyNormalizedSSSCommon.ush
+++ b/Engine/Shaders/Private/BurleyNormalizedSSSCommon.ush
@@ -175,6 +175,11 @@ float3 GetScalingFactor3D(float3 SurfaceAlbedo)
 	return S3D;
 }
 
+float3 GetCDF3D(float3 D, float X)
+{
+	return 1 - 0.25 * exp(-X / D) - 0.75 * exp(-X / (3 * D));
+}
+
 //@TODO Revisit here after offline comparison.
 //call this function in TransmissionCommon.ush for Debug. No transmission tint is applied.
 
@@ -236,3 +241,27 @@ float3 GetBurleyTransmissionProfile(
 
 	return TransmissionThroughput * FadeDistance;
 }
+
+struct FBurleyParameter
+{
+	float4 SurfaceAlbedo;
+	float4 DiffuseMeanFreePath;
+	float  WorldUnitScale;
+	float  SurfaceOpacity; // 1.0 means full Burley, 0.0 means full Default Lit
+};
+
+//Surface albedo and mean free path length
+float4 GetSubsurfaceProfileSurfaceAlbedo(uint SubsurfaceProfileInt)
+{
+	return View.SSProfilesTexture.Load(int3(BSSS_SURFACEALBEDO_OFFSET, SubsurfaceProfileInt, 0));
+}
+
+float4 GetSubsurfaceProfileDiffuseMeanFreePath(uint SubsurfaceProfileInt)
+{
+	return View.SSProfilesTexture.Load(int3(BSSS_DMFP_OFFSET, SubsurfaceProfileInt, 0));
+}
+
+float GetSubsurfaceProfileWorldUnitScale(uint SubsurfaceProfileInt)
+{
+	return View.SSProfilesTexture.Load(int3(SSSS_SUBSURFACE_COLOR_OFFSET, SubsurfaceProfileInt, 0)).a;
+}
diff --git a/Engine/Shaders/Private/PostProcessSubsurface.usf b/Engine/Shaders/Private/PostProcessSubsurface.usf
--- a/Engine/Shaders/Private/PostProcessSubsurface.usf
+++ b/Engine/Shaders/Private/PostProcessSubsurface.usf
@@ -187,6 +187,108 @@ float4 GetSceneColor(float2 BufferUV)
 	return Texture2DSampleLevel(SubsurfaceInput0_Texture, SharedSubsurfaceSampler0, BufferUV, 0);
 }
 
+#include "BurleyNormalizedSSSCommon.ush"
+
+#if STRATA_ENABLED
+
+// Calculate the Subsurface profile kernel on the fly using integration instead of LUT.
+half4 GetSubsurfaceProfileKernel(int SampleIndex, float3 D, float ScatterRadiusInMM, float WorldUnitScale)
+{
+	// See notes in SubsurfaceProfile.cpp. Scatter radius is in mm instead of cm here.
+	ScatterRadiusInMM = max(ScatterRadiusInMM/2.229f,1.0f);
+	
+	const float Range = ((2 * SSSS_N_KERNELWEIGHTCOUNT - 1) > 20) ? 3.0f : 2.0f;
+
+	const float r = Range * float(SampleIndex) / (SSSS_N_KERNELWEIGHTCOUNT - 1);
+	const float r0 = Range * (float(SampleIndex) - 0.5f) / (SSSS_N_KERNELWEIGHTCOUNT - 1);
+	const float r1 = Range * (float(SampleIndex) + 0.5f) / (SSSS_N_KERNELWEIGHTCOUNT - 1);
+
+	const float x = r * r / (Range );
+	const float x0 = r0 * r0 / (Range );
+	const float x1 = r1 * r1 / (Range );
+	const float3 MaxCDF = GetCDF3D(D, ScatterRadiusInMM * Range * (1.0f + 0.5f / (SSSS_N_KERNELWEIGHTCOUNT - 1)));
+
+	half4 Kernel = (half4)0.0f;
+
+	BRANCH
+		if (SampleIndex == 0)
+		{
+			Kernel.rgb = GetCDF3D(D, ScatterRadiusInMM * x1) / MaxCDF;
+			Kernel.a = 0.0f;
+		}
+		else
+		{
+			// Need to multiply by 0.5f, and modify MaxCDF to match as close as offline LUT. Here favors integration 
+			// instead of discrete approximation.
+			Kernel.rgb = (GetCDF3D(D, ScatterRadiusInMM * x1) - GetCDF3D(D, ScatterRadiusInMM * x0)) / MaxCDF;
+			Kernel.a =  4.0f * x; // 4.0f is consistent to the scale in LUT
+		}
+
+	// Scale to match offline encoding
+	const float TableMaxA = 3.0f;
+	// const FLinearColor TableColorScale = FLinearColor(
+	//		1.0f / TableMaxRGB,
+	//		1.0f / TableMaxRGB,
+	//		1.0f / TableMaxRGB,
+	//		1.0f / TableMaxA);
+	// C *= TableColorScale
+	// C.W *= Data.ScatterRadius / SUBSURFACE_RADIUS_SCALE;     
+	Kernel.a *= (ScatterRadiusInMM * BURLEY_MM_2_CM) * WorldUnitScale * 2.229f * 0.707f / (TableMaxA * SUBSURFACE_RADIUS_SCALE);
+
+	return Kernel;
+}
+#endif
+
+// fetch surface albedo and diffuse mean free path
+// The diffuse mean free path is modulated by world scale and a custom term passed in
+// which is the opacity texture.
+#if STRATA_ENABLED
+FBurleyParameter GetBurleyParameters(FStrataSubsurfaceData SSSData)
+{
+	FBurleyParameter Out = (FBurleyParameter)0;
+	if (SSSData.bIsProfile)
+	{
+		Out.SurfaceAlbedo = GetSubsurfaceProfileSurfaceAlbedo(SSSData.ProfileId);
+		Out.DiffuseMeanFreePath = DecodeDiffuseMeanFreePath(GetSubsurfaceProfileDiffuseMeanFreePath(SSSData.ProfileId));
+		Out.WorldUnitScale = DecodeWorldUnitScale(GetSubsurfaceProfileWorldUnitScale(SSSData.ProfileId)) * BURLEY_CM_2_MM;
+
+		Out.SurfaceOpacity = SSSData.ProfileRadiusScale;
+		Out.DiffuseMeanFreePath *= Out.SurfaceOpacity;
+	}
+	else if (SSSData.bIsValid)
+	{
+		// Decode the DMFP stored as Float111110 in a uint over a RGBA8 unorm
+		const float3 DMFP = SSSData.DMFP * BURLEY_CM_2_MM; // Stored DMFP is in cm, converted into mm
+
+		const float MaxDMFP = max3(DMFP.x, DMFP.y, DMFP.z);
+		const uint IndexOfMaxDmfp = (DMFP.r == MaxDMFP) ? 0 : ((DMFP.g == MaxDMFP) ? 1 : 2);
+
+		Out.SurfaceAlbedo = float4(SSSData.BaseColor, SSSData.BaseColor[IndexOfMaxDmfp]);
+		Out.DiffuseMeanFreePath = float4(DMFP, MaxDMFP);
+		Out.WorldUnitScale = BURLEY_CM_2_MM;
+
+		// STRATA_TODO smooth transistion based on pixel footprint w.r.t. MFPRadius and pixel footprint
+		Out.SurfaceOpacity = 1.0f;
+		Out.DiffuseMeanFreePath *= Out.SurfaceOpacity;
+	}
+	return Out;
+}
+#else
+FBurleyParameter GetBurleyParameters(uint SubsurfaceProfileInt, FGBufferData GBuffer)
+{
+	FBurleyParameter Out = (FBurleyParameter)0;
+	Out.SurfaceAlbedo = GetSubsurfaceProfileSurfaceAlbedo(SubsurfaceProfileInt);
+	Out.DiffuseMeanFreePath = DecodeDiffuseMeanFreePath(GetSubsurfaceProfileDiffuseMeanFreePath(SubsurfaceProfileInt));
+	Out.WorldUnitScale = DecodeWorldUnitScale(GetSubsurfaceProfileWorldUnitScale(SubsurfaceProfileInt)) * BURLEY_CM_2_MM;
+
+	// Mask should always be true, we can probably ignore it?
+	Out.SurfaceOpacity = UseSubsurfaceProfile(GBuffer.ShadingModelID) ? GBuffer.CustomData.a : 0.0f;
+	Out.DiffuseMeanFreePath *= Out.SurfaceOpacity;
+
+	return Out;
+}
+#endif
+
 // from [REDACTED_PROFILE_URL]/separable-sss/tree/master/Demo
 // [REDACTED_PERSON] [REDACTED_PROFILE_URL]
 // [REDACTED_PROFILE_URL]
diff --git a/Engine/Shaders/Private/SeparableSSS.ush b/Engine/Shaders/Private/SeparableSSS.ush
--- a/Engine/Shaders/Private/SeparableSSS.ush
+++ b/Engine/Shaders/Private/SeparableSSS.ush
@@ -176,6 +176,53 @@ float3 SSSSTransmittance(float translucency, float sssWidth, float3 worldPositio
     return profile * saturate(0.3 + dot(light, -worldNormal));
 }
 
+//-----------------------------------------------------------------------------
+// Query the kernel for both strata and the SSS profile
+struct FSeparableFilterParameters
+{
+#if STRATA_ENABLED
+	float3 D3D;
+	float ScatterRadius;
+	float WorldUnitScale;
+#endif
+	uint SubsurfaceProfileInt;
+};
+
+#if STRATA_ENABLED
+void FillBurleyParameters(FStrataSubsurfaceData SSSData,
+	inout FSeparableFilterParameters Parameters)
+{
+	const FBurleyParameter BurleyParameter = GetBurleyParameters(SSSData);
+	const float3 S3D = GetScalingFactor3D(BurleyParameter.SurfaceAlbedo.xyz);
+
+	Parameters.D3D = BurleyParameter.DiffuseMeanFreePath.xyz / S3D.xyz;
+	Parameters.ScatterRadius = BurleyParameter.DiffuseMeanFreePath.w;
+	Parameters.WorldUnitScale = BurleyParameter.WorldUnitScale;
+}
+#endif
+
+half4 GetSubsurfaceProfileKernel(uint index, FSeparableFilterParameters Parameters)
+{
+	half4 Kernel = (half4)0.0f;
+#if STRATA_ENABLED
+	if (Parameters.SubsurfaceProfileInt == SSS_PROFILE_ID_PERPIXEL)
+	{
+		Kernel = GetSubsurfaceProfileKernel(index, Parameters.D3D, 
+											Parameters.ScatterRadius, 
+											Parameters.WorldUnitScale);
+	}
+	else
+	{
+		Kernel = GetSubsurfaceProfileKernel(SSSS_N_KERNELWEIGHTOFFSET + index,
+											Parameters.SubsurfaceProfileInt);
+	}
+#else
+	Kernel = GetSubsurfaceProfileKernel(SSSS_N_KERNELWEIGHTOFFSET + index, 
+										Parameters.SubsurfaceProfileInt);
+#endif
+	return Kernel;
+}
+
 //-----------------------------------------------------------------------------
 // Separable SSS Reflectance Pixel Shader
 
@@ -218,27 +265,32 @@ float4 SSSSBlurPS(uint2 BufferPos, float2 BufferUV, float2 dir, bool initStencil
 	// ideally this comes from a half res buffer as well - there are some minor artifacts
 	finalStep *= SSSStrength; // Modulate it using the opacity (0..1 range)
     
+	FSeparableFilterParameters SeparableFilterParameters;
+
 #if STRATA_ENABLED
 	const FStrataSubsurfaceData SSSData = LoadStataSSSData(BufferUV);
 	// For SSS_PROFILE_ID_PERPIXEL is managed throught the burley passes
-	const uint SubsurfaceProfileInt = SSSData.bIsProfile ? SSSData.ProfileId : SSS_PROFILE_ID_INVALID;
+	const uint SubsurfaceProfileInt = SSSData.bIsProfile ? SSSData.ProfileId : SSS_PROFILE_ID_PERPIXEL;
+	SeparableFilterParameters.SubsurfaceProfileInt = SubsurfaceProfileInt;
+	FillBurleyParameters(SSSData, SeparableFilterParameters);
 #else
 	const FGBufferData GBufferData = GetGBufferData(BufferUV);
 	// 0..255, which SubSurface profile to pick
 	// ideally this comes from a half res buffer as well - there are some minor artifacts
 	const uint SubsurfaceProfileInt = ExtractSubsurfaceProfileInt(GBufferData);
+	SeparableFilterParameters.SubsurfaceProfileInt = SubsurfaceProfileInt;
 #endif
 
-
     // Accumulate the center sample:
     float3 colorAccum = 0;
 	// >0 to avoid division by 0, not 100% correct to not visible
 	float3 colorInvDiv = 0.00001f;
 
 	// center sample
-	colorInvDiv += GetSubsurfaceProfileKernel(SSSS_N_KERNELWEIGHTOFFSET, SubsurfaceProfileInt).rgb;
-	colorAccum = colorM.rgb * GetSubsurfaceProfileKernel(SSSS_N_KERNELWEIGHTOFFSET, SubsurfaceProfileInt).rgb;
+	half3  CentralKernelWeight = GetSubsurfaceProfileKernel(0, SeparableFilterParameters).rgb;
 
+	colorInvDiv += CentralKernelWeight;
+	colorAccum = colorM.rgb * CentralKernelWeight;
 	float3 BoundaryColorBleed = GetSubsurfaceProfileBoundaryColorBleed(SubsurfaceProfileInt);
 
 	// Accumulate the other samples:
@@ -246,7 +298,7 @@ float4 SSSSBlurPS(uint2 BufferPos, float2 BufferUV, float2 dir, bool initStencil
 	for (int i = 1; i < SSSS_N_KERNELWEIGHTCOUNT; i++) 
 	{
 		// Kernel.a = 0..SUBSURFACE_KERNEL_SIZE (radius)
-		half4 Kernel = GetSubsurfaceProfileKernel(SSSS_N_KERNELWEIGHTOFFSET + i, SubsurfaceProfileInt);
+		half4 Kernel = GetSubsurfaceProfileKernel(i, SeparableFilterParameters);
 
 		float4 LocalAccum = 0;
 
diff --git a/Engine/Shaders/Private/SubsurfaceBurleyNormalized.ush b/Engine/Shaders/Private/SubsurfaceBurleyNormalized.ush
--- a/Engine/Shaders/Private/SubsurfaceBurleyNormalized.ush
+++ b/Engine/Shaders/Private/SubsurfaceBurleyNormalized.ush
@@ -6,7 +6,6 @@
 #pragma once
 
 #include "Random.ush"
-#include "BurleyNormalizedSSSCommon.ush"
 #include "DeferredShadingCommon.ush"
 #include "MonteCarlo.ush"
 
@@ -141,30 +140,6 @@
 #define SUBSURFACE_SAMPLER_TYPE 0
 #endif
 
-struct FBurleyParameter
-{
-	float4 SurfaceAlbedo;
-	float4 DiffuseMeanFreePath;
-	float  WorldUnitScale;
-	float  SurfaceOpacity; // 1.0 means full Burley, 0.0 means full Default Lit
-};
-
-//Surface albedo and mean free path length
-float4 GetSubsurfaceProfileSurfaceAlbedo(uint SubsurfaceProfileInt)
-{
-	return View.SSProfilesTexture.Load(int3(BSSS_SURFACEALBEDO_OFFSET, SubsurfaceProfileInt, 0));
-}
-
-float4 GetSubsurfaceProfileDiffuseMeanFreePath(uint SubsurfaceProfileInt)
-{
-	return View.SSProfilesTexture.Load(int3(BSSS_DMFP_OFFSET, SubsurfaceProfileInt, 0));
-}
-
-float GetSubsurfaceProfileWorldUnitScale(uint SubsurfaceProfileInt)
-{
-	return View.SSProfilesTexture.Load(int3(SSSS_SUBSURFACE_COLOR_OFFSET, SubsurfaceProfileInt, 0)).a;
-}
-
 float SampleDepthTexture(float2 ScreenUV)
 {
 #if SUBSURFACE_HALF_RES && !(SUBSURFACE_SINGLE_PASS)
@@ -176,56 +151,6 @@ float SampleDepthTexture(float2 ScreenUV)
 	return Texture2DSampleLevel(SceneTexturesStruct.SceneDepthTexture, SceneTexturesStruct_SceneDepthTextureSampler, ScreenUV, 0).r;
 }
 
-// fetch surface albedo and diffuse mean free path
-// The diffuse mean free path is modulated by world scale and a custom term passed in
-// which is the opacity texture.
-#if STRATA_ENABLED
-FBurleyParameter GetBurleyParameters(FStrataSubsurfaceData SSSData)
-{
-	FBurleyParameter Out = (FBurleyParameter)0;
-	if (SSSData.bIsProfile)
-	{
-		Out.SurfaceAlbedo		= GetSubsurfaceProfileSurfaceAlbedo(SSSData.ProfileId);
-		Out.DiffuseMeanFreePath = DecodeDiffuseMeanFreePath(GetSubsurfaceProfileDiffuseMeanFreePath(SSSData.ProfileId));
-		Out.WorldUnitScale		= DecodeWorldUnitScale(GetSubsurfaceProfileWorldUnitScale(SSSData.ProfileId)) * BURLEY_CM_2_MM;
-
-		Out.SurfaceOpacity		= SSSData.ProfileRadiusScale;
-		Out.DiffuseMeanFreePath*= Out.SurfaceOpacity;
-	}
-	else if(SSSData.bIsValid)
-	{
-		// Decode the DMFP stored as Float111110 in a uint over a RGBA8 unorm
-		const float3 DMFP				= SSSData.DMFP * BURLEY_CM_2_MM; // Stored DMFP is in cm, converted into mm
-
-		const float MaxDMFP				= max3(DMFP.x, DMFP.y, DMFP.z);
-		const uint IndexOfMaxDmfp		= (DMFP.r == MaxDMFP) ? 0 : ((DMFP.g == MaxDMFP) ? 1 : 2);
-
-		Out.SurfaceAlbedo				= float4(SSSData.BaseColor, SSSData.BaseColor[IndexOfMaxDmfp]);
-		Out.DiffuseMeanFreePath			= float4(DMFP, MaxDMFP);
-		Out.WorldUnitScale				= BURLEY_CM_2_MM;
-
-		// STRATA_TODO smooth transistion based on pixel footprint w.r.t. MFPRadius and pixel footprint
-		Out.SurfaceOpacity				= 1.0f;
-		Out.DiffuseMeanFreePath			*= Out.SurfaceOpacity;
-	}
-	return Out;
-}
-#else
-FBurleyParameter GetBurleyParameters(uint SubsurfaceProfileInt, FGBufferData GBuffer)
-{
-	FBurleyParameter Out = (FBurleyParameter)0;
-	Out.SurfaceAlbedo		= GetSubsurfaceProfileSurfaceAlbedo(SubsurfaceProfileInt);
-	Out.DiffuseMeanFreePath = DecodeDiffuseMeanFreePath(GetSubsurfaceProfileDiffuseMeanFreePath(SubsurfaceProfileInt));
-	Out.WorldUnitScale		= DecodeWorldUnitScale(GetSubsurfaceProfileWorldUnitScale(SubsurfaceProfileInt)) * BURLEY_CM_2_MM;
-
-	// Mask should always be true, we can probably ignore it?
-	Out.SurfaceOpacity = UseSubsurfaceProfile(GBuffer.ShadingModelID) ? GBuffer.CustomData.a : 0.0f;
-	Out.DiffuseMeanFreePath *= Out.SurfaceOpacity;
-
-	return Out;
-}
-#endif
-
 float2 Generate2DRandomNumber(int3 Seed)
 {
 #if SAMPLE_ROOT_ANGLE_R2SEQUENCE
@@ -1354,4 +1279,4 @@ void BurleyComputeMain(uint2 DT_ID, uint2 G_ID, uint GI)
 #endif
 }
 
-#endif
\ No newline at end of file
+#endif
diff --git a/Engine/Source/Runtime/Engine/Private/Rendering/BurleyNormalizedSSS.cpp b/Engine/Source/Runtime/Engine/Private/Rendering/BurleyNormalizedSSS.cpp
--- a/Engine/Source/Runtime/Engine/Private/Rendering/BurleyNormalizedSSS.cpp
+++ b/Engine/Source/Runtime/Engine/Private/Rendering/BurleyNormalizedSSS.cpp
@@ -69,7 +69,8 @@ void MapFallOffColor2SurfaceAlbedoAndDiffuseMeanFreePath(float FalloffColor, flo
 // Average relative error: 5.5% (reference to MC)
 float GetPerpendicularScalingFactor(float SurfaceAlbedo)
 {
-	return 1.85 - SurfaceAlbedo + 7 * FMath::Pow(SurfaceAlbedo - 0.8, 3);
+	// add abs() to match the formula in the original paper. 
+	return 1.85 - SurfaceAlbedo + 7 * FMath::Pow(FMath::Abs(SurfaceAlbedo - 0.8), 3);
 }
 
 FVector GetPerpendicularScalingFactor(FLinearColor SurfaceAlbedo)
