Change legacy conversion to help compiler to statically unroll strata packing code.

diff --git a/Engine/Shaders/Private/Strata/StrataLegacyConversion.ush b/Engine/Shaders/Private/Strata/StrataLegacyConversion.ush
--- a/Engine/Shaders/Private/Strata/StrataLegacyConversion.ush
+++ b/Engine/Shaders/Private/Strata/StrataLegacyConversion.ush
@@ -2,15 +2,90 @@
 
 #pragma once
 
-#include "../Common.ush"        
+#include "../Common.ush"
 
-FStrataData StrataConvertLegacyMaterial(
+struct FStrataLegacyParameters
+{
+	float3 BaseColor;
+	float  Specular;
+	float  Metallic;
+	float  Roughness;
+	float  Anisotropy;
+	float  SSSProfileID;
+	float3 SSSDMFP;
+	float  SSSDMFPScale;
+	float3 Emissive;
+	float  FuzzAmount;
+	float3 FuzzColor;
+	float  Thickness;
+	uint   SharedLocalBasisIndex;
+
+	float Weight;
+};
+
+FStrataLegacyParameters InitStrataLegacyParameters(uint InSharedLocalBasisIndex, float InWeight)
+{
+	const float3 Zeros = float3(0, 0, 0);
+
+	FStrataLegacyParameters Out;
+	Out.BaseColor = Zeros;
+	Out.Specular = 0.5f;
+	Out.Metallic = 0.f;
+	Out.Roughness = 0.5f;
+	Out.Anisotropy = 0.f;
+	Out.SSSProfileID = 0.f;
+	Out.SSSDMFP = Zeros;
+	Out.SSSDMFPScale = 1.f;
+	Out.Emissive = Zeros;
+	Out.FuzzAmount = 0.f;
+	Out.FuzzColor = Zeros;
+	Out.Thickness = STRATA_LAYER_DEFAULT_THICKNESS_CM;
+	Out.SharedLocalBasisIndex = InSharedLocalBasisIndex;
+	Out.Weight = InWeight;
+	return Out;
+}
+
+FStrataData CreateLegacySlab(FStrataLegacyParameters In, inout uint SharedLocalBasisTypes)
+{
+	const float3 Zeros = float3(0, 0, 0);
+	const float3 Ones  = float3(1, 1, 1);
+
+	// Fixed layer structure for helping compiler to unroll and optimize shader
+	return GetStrataSlabBSDF(
+		true,				// UseMetalness
+		In.BaseColor,		// BaseColor		-  Metalness workflow
+		Ones,				// EdgeColor		-  Metalness workflow
+		In.Specular,		// Specular			-  Metalness workflow
+		In.Metallic,		// Metallic			-  Metalness workflow
+		Zeros,				// DiffuseAlbedo	- !Metalness workflow
+		Zeros, 				// F0				- !Metalness workflow
+		Zeros,				// F90				- !Metalness workflow
+		In.Roughness,		// Roughness
+		In.Anisotropy,		// Anisotropy
+		In.SSSProfileID,	// SSSProfileID
+		In.SSSDMFP,			// SSSDMFP
+		In.SSSDMFPScale,	// SSSDMFPScale
+		In.Emissive,		// Emissive
+		0.f,				// Haziness
+		0.f,				// ThinFilmThickness
+		In.FuzzAmount,		// FuzzAmount
+		In.FuzzColor,		// FuzzColor
+		In.Thickness,
+		In.SharedLocalBasisIndex, SharedLocalBasisTypes);
+}
+
+// Convert legacy shading models - Dynamic
+// This function can handle dynamic shading models (i.e., known only at runtime). 
+// For this, the layer topology is 'fixed' and composed of two slabs vertically layered. This is done for 
+// helping the compiler to unroll Strata BSDF traversing and packing. In most cases, the bottom slab is 
+// weighted by 0 and it will be removed once the data are written-out/packed.
+FStrataData StrataConvertLegacyMaterialDynamic(
 	float3 BaseColor, float Specular, float Metallic,
 	float Roughness, float Anisotropy,
 	float3 SubSurfaceColor, float SubSurfaceProfileId,
 	float ClearCoat, float ClearCoatRoughness,
 	float3 Emissive,
-	float Opacity, 
+	float Opacity,
 	float3 TransmittanceColor,
 	float3 WaterScatteringCoefficients, float3 WaterAbsorptionCoefficients, float WaterPhaseG, float3 ColorScaleBehindWater,
 	uint ShadingModel,
@@ -18,45 +93,43 @@ FStrataData StrataConvertLegacyMaterial(
 	uint ClearCoatBottomNormal_SharedLocalBasisIndex,
 	inout uint SharedLocalBasisTypes)
 {
-	FStrataData Out = GetInitialisedStrataData();
-
 	const float DefaultThickness = STRATA_LAYER_DEFAULT_THICKNESS_CM;
 	const float3 Zeros = float3(0, 0, 0);
 	const float3 Ones  = float3(1, 1, 1);
+
+	// Can only mix Unlit / DefaultLit / Sub-Surface / Preintegrated-Skin / Subsurface-Profile / ClearCoat / Foliage / Cloth / Eye/ Thin
+	// No supported : Water/Hair
+
+	// Top/Bottom
+	FStrataLegacyParameters Top = InitStrataLegacyParameters(SharedLocalBasisIndex, 1.0f);
+	FStrataLegacyParameters Bot = InitStrataLegacyParameters(SharedLocalBasisIndex, 1.0f);
+
+	float FinalWeight = 1.0f;
+
 	if (ShadingModel == SHADINGMODELID_UNLIT)
 	{
-		Out = GetStrataUnlitBSDF(Emissive, Ones /*TransnmittanceColor*/);
+		// Unlit is handle with a emissive slab
+		Top.BaseColor = Zeros;
+		Top.Specular = 0.f;
+		Top.Metallic = 0.f;
+		Top.Roughness = 0.f;
+		Top.Anisotropy = Anisotropy;
+		Top.Emissive = Emissive;
+
+		Bot.Weight = 0.f;
+		FinalWeight = Opacity;
 	}
 	else if (ShadingModel == SHADINGMODELID_DEFAULT_LIT)
 	{
-		Out = GetStrataSlabBSDF(
-			true,				// UseMetalness
-			BaseColor,			// BaseColor		-  Metalness workflow
-			Ones,				// EdgeColor		-  Metalness workflow
-			Specular,			// Specular			-  Metalness workflow
-			Metallic,			// Metallic			-  Metalness workflow
-			Zeros,				// DiffuseAlbedo	- !Metalness workflow
-			Zeros, 				// F0				- !Metalness workflow
-			Zeros,				// F90				- !Metalness workflow
-			Roughness,			// Roughness
-			Anisotropy,			// Anisotropy
-			0.f,				// SSSProfileID
-			Zeros,				// SSSDMFP
-			0.f,				// SSSDMFPScale
-			Emissive,			// Emissive
-			0.f,				// Haziness
-			0.f,				// ThinFilmThickness
-			0.f,				// FuzzAmount
-			Zeros,				// FuzzColor
-			DefaultThickness,
-			SharedLocalBasisIndex, SharedLocalBasisTypes);
+		Top.BaseColor = BaseColor;
+		Top.Specular = Specular;
+		Top.Metallic = Metallic;
+		Top.Roughness = Roughness;
+		Top.Anisotropy = Anisotropy;
+		Top.Emissive = Emissive;
 
-		// Take in account the legacy opacity weighting for translucent material
-		#if !MATERIALBLENDING_MASKED && !MATERIALBLENDING_SOLID
-		{
-			Out = StrataWeight(Out, Opacity);
-		}
-		#endif
+		Bot.Weight = 0.f;
+		FinalWeight = Opacity;
 	}
 	else if (ShadingModel == SHADINGMODELID_SUBSURFACE || ShadingModel == SHADINGMODELID_PREINTEGRATED_SKIN)
 	{
@@ -72,82 +145,48 @@ FStrataData StrataConvertLegacyMaterial(
 		// and a constant scaling for matching roughly the behavior
 		SSSMFP = SSSMFP * 0.5f;
 
-		Out = GetStrataSlabBSDF(
-			true,				// UseMetalness
-			ModifiedBaseColor,	// BaseColor		-  Metalness workflow
-			Ones,				// EdgeColor		-  Metalness workflow
-			Specular,			// Specular			-  Metalness workflow
-			Metallic,			// Metallic			-  Metalness workflow
-			Zeros,				// DiffuseAlbedo	- !Metalness workflow
-			Zeros, 				// F0				- !Metalness workflow
-			Zeros,				// F90				- !Metalness workflow
-			Roughness,			// Roughness
-			Anisotropy,			// Anisotropy
-			0.f,				// SSSProfileID
-			SSSMFP,				// SSSDMFP
-			Opacity,			// SSSDMFPScale
-			Emissive,			// Emissive
-			0.f,				// Haziness
-			0.f,				// ThinFilmThickness
-			0.f,				// FuzzAmount
-			Zeros,				// FuzzColor
-			DefaultThickness,
-			SharedLocalBasisIndex, SharedLocalBasisTypes);
+		Top.BaseColor = ModifiedBaseColor;
+		Top.Specular = Specular;
+		Top.Metallic = Metallic;
+		Top.Roughness = Roughness;
+		Top.Anisotropy = Anisotropy;
+		Top.SSSDMFP = SSSMFP;
+		Top.SSSDMFPScale = Opacity;
+		Top.Emissive = Emissive;
 
 		// No need to take into account opacity weighting for translucent material, as the coverage is assumed to be 1
+		Bot.Weight = 0.f;
+		FinalWeight = 1.f;
 	}
 	else if (ShadingModel == SHADINGMODELID_SUBSURFACE_PROFILE)
 	{
-		Out = GetStrataSlabBSDF(
-			true,				// UseMetalness
-			BaseColor,			// BaseColor		-  Metalness workflow
-			Ones,				// EdgeColor		-  Metalness workflow
-			Specular,			// Specular			-  Metalness workflow
-			Metallic,			// Metallic			-  Metalness workflow
-			Zeros,				// DiffuseAlbedo	- !Metalness workflow
-			Zeros, 				// F0				- !Metalness workflow
-			Zeros,				// F90				- !Metalness workflow
-			Roughness,			// Roughness
-			Anisotropy,			// Anisotropy
-			SubSurfaceProfileId,// SSSProfileID
-			Zeros,				// SSSDMFP
-			Opacity,			// SSSDMFPScale
-			Emissive,			// Emissive
-			0.f,				// Haziness
-			0.f,				// ThinFilmThickness
-			0.f,				// FuzzAmount
-			Zeros,				// FuzzColor
-			DefaultThickness,
-			SharedLocalBasisIndex, SharedLocalBasisTypes);
+		Top.BaseColor = BaseColor;
+		Top.Specular = Specular;
+		Top.Metallic = Metallic;
+		Top.Roughness = Roughness;
+		Top.Anisotropy = Anisotropy;
+		Top.SSSProfileID = SubSurfaceProfileId;
+		Top.SSSDMFPScale = Opacity;
+		Top.Emissive = Emissive;
 
 		// No need to take into account opacity weighting for translucent material, as the coverage is assumed to be 1
+		Bot.Weight = 0.f;
+		FinalWeight = 1.f;
 	}
 	else if (ShadingModel == SHADINGMODELID_CLEAR_COAT)
 	{
 		// === BOTTOM layer === 
-
-		FStrataData BottomSlab = GetStrataSlabBSDF(
-			true,				// UseMetalness
-			BaseColor,			// BaseColor		-  Metalness workflow
-			Ones,				// EdgeColor		-  Metalness workflow
-			Specular,			// Specular			-  Metalness workflow
-			Metallic,			// Metallic			-  Metalness workflow
-			Zeros,				// DiffuseAlbedo	- !Metalness workflow
-			Zeros, 				// F0				- !Metalness workflow
-			Zeros,				// F90				- !Metalness workflow
-			Roughness,			// Roughness
-			Anisotropy,			// Anisotropy
-			0.f,				// SSSProfileID
-			Zeros,				// SSSDMFP
-			0.f,				// SSSDMFPScale
-			Zeros,				// Emissive
-			0.f,				// Haziness
-			0.f,				// ThinFilmThickness
-			0.f,				// FuzzAmount
-			Zeros,				// FuzzColor
-			DefaultThickness,
-			ClearCoatBottomNormal_SharedLocalBasisIndex, /*ClearCoatBottomNormal_*/SharedLocalBasisTypes); // TODO normally we should be able to flag the bottom layer as anisotropic
-		
+		Bot.BaseColor = BaseColor;
+		Bot.Specular = Specular;
+		Bot.Metallic = Metallic;
+		Bot.Roughness = Roughness;
+		Bot.Anisotropy = Anisotropy;
+		Bot.SSSProfileID = 0.f;
+		Bot.SSSDMFP = Zeros;
+		Bot.SSSDMFPScale = 0.f;
+		Bot.Emissive = Zeros;
+		Bot.SharedLocalBasisIndex = ClearCoatBottomNormal_SharedLocalBasisIndex;
+		Bot.Weight = 1.f;
 		// === TOP layer === 
 		 																								   
 		// The top layer has a hard coded specular value of 0.5 (F0 = 0.04)
@@ -161,117 +200,58 @@ FStrataData StrataConvertLegacyMaterial(
 		const float Thickness = DefaultThickness;
 		const float3 MFP = TransmittanceToMeanFreePath(0.75f, Thickness * CENTIMETER_TO_METER) * METER_TO_CENTIMETER;
 
-		FStrataData TopSlab = GetStrataSlabBSDF(
-			true,				// UseMetalness
-			Zeros,				// BaseColor		-  Metalness workflow
-			Ones,				// EdgeColor		-  Metalness workflow
-			0.5f,				// Specular			-  Metalness workflow
-			0.f,				// Metallic			-  Metalness workflow
-			Zeros,				// DiffuseAlbedo	- !Metalness workflow
-			Zeros, 				// F0				- !Metalness workflow
-			Zeros,				// F90				- !Metalness workflow
-			ClearCoatRoughness, // Roughness
-			Anisotropy,			// Anisotropy
-			0.f,				// SSSProfileID
-			MFP,				// SSSDMFP
-			1.f,				// SSSDMFPScale
-			Emissive,			// Emissive
-			0.f,				// Haziness
-			0.f,				// ThinFilmThickness
-			0.f,				// FuzzAmount
-			Zeros,				// FuzzColor
-			Thickness,
-			SharedLocalBasisIndex, SharedLocalBasisTypes);
-
-		// === Layering === 
-		FStrataData WeightedTopSlab = StrataWeight(TopSlab, ClearCoat);
-		Out = StrataVerticalLayering(WeightedTopSlab, BottomSlab);
+		Top.BaseColor = Zeros;
+		Top.Specular = 0.5f;
+		Top.Metallic = 0.f;
+		Top.Roughness = ClearCoatRoughness;
+		Top.Anisotropy = Anisotropy;
+		Top.SSSProfileID = 0.f;
+		Top.SSSDMFP = MFP;
+		Top.SSSDMFPScale = 1.f;
+		Top.Emissive = Emissive;
+		Top.Thickness = Thickness;
+		Top.SharedLocalBasisIndex = SharedLocalBasisIndex;
+		Top.Weight = ClearCoat;
 
-		// Take in account the legacy opacity weighting for translucent material
-		#if !MATERIALBLENDING_MASKED && !MATERIALBLENDING_SOLID
-		{
-			Out = StrataWeight(Out, Opacity);
-		}
-		#endif
+		FinalWeight = Opacity;
 	}
 	else if (ShadingModel == SHADINGMODELID_TWOSIDED_FOLIAGE)
 	{
-	// Set a thickness that will enabled the thin lighting model (corresponding to the legacy two-sided lighting model)
+		// Set a thickness that will enabled the thin lighting model (corresponding to the legacy two-sided lighting model)
 		const float Thickness = STRATA_LAYER_ISTHIN_THICKNESS_THRESHOLD_CM - 1e-5f;
 		const float3 MFP = TransmittanceToMeanFreePath(SubSurfaceColor /*TransmittanceColor*/, Thickness * CENTIMETER_TO_METER) * METER_TO_CENTIMETER;
 
-		Out = GetStrataSlabBSDF(
-			true,				// UseMetalness
-			BaseColor,			// BaseColor		-  Metalness workflow
-			Ones,				// EdgeColor		-  Metalness workflow
-			Specular,			// Specular			-  Metalness workflow
-			Metallic,			// Metallic			-  Metalness workflow
-			Zeros,				// DiffuseAlbedo	- !Metalness workflow
-			Zeros, 				// F0				- !Metalness workflow
-			Zeros,				// F90				- !Metalness workflow
-			Roughness,			// Roughness
-			Anisotropy,			// Anisotropy
-			0.f,				// SSSProfileID
-			MFP,				// SSSDMFP
-			1.f,				// SSSDMFPScale
-			Emissive,			// Emissive
-			0.f,				// Haziness
-			0.f,				// ThinFilmThickness
-			0.f,				// FuzzAmount
-			Zeros,				// FuzzColor
-			Thickness,
-			SharedLocalBasisIndex, SharedLocalBasisTypes);
+		Top.BaseColor = BaseColor;
+		Top.Specular = Specular;
+		Top.Metallic = Metallic;
+		Top.Roughness = Roughness;
+		Top.Anisotropy = Anisotropy;
+		Top.SSSProfileID = 0.f;
+		Top.SSSDMFP = MFP;
+		Top.SSSDMFPScale = 1.f;
+		Top.Emissive = Emissive;
+		Top.Thickness = Thickness;
 
 		// Take in account the legacy opacity weighting for translucent material
-		#if !MATERIALBLENDING_MASKED && !MATERIALBLENDING_SOLID
-		{
-			Out = StrataWeight(Out, Opacity);
-		}
-		#endif
-	}
-	else if (ShadingModel == SHADINGMODELID_HAIR)
-	{
-		Out = GetStrataHairBSDF(
-			BaseColor,	   // BaseColor
-			Metallic,	   // Scatter
-			Specular,	   // Specular
-			Roughness,	   // Roughness
-			ClearCoat,	   // Backlit
-			Emissive,		// EmissiveColor
-			SharedLocalBasisIndex);
-
-		// No need to take into account opacity weighting for translucent material, as the coverage is assumed to be 1
+		Bot.Weight = 0.f;
+		FinalWeight = Opacity;
 	}
 	else if (ShadingModel == SHADINGMODELID_CLOTH)
 	{
-		Out = GetStrataSlabBSDF(
-			true,				// UseMetalness
-			BaseColor,			// BaseColor		-  Metalness workflow
-			Ones,				// EdgeColor		-  Metalness workflow
-			Specular,			// Specular			-  Metalness workflow
-			Metallic,			// Metallic			-  Metalness workflow
-			Zeros,				// DiffuseAlbedo	- !Metalness workflow
-			Zeros, 				// F0				- !Metalness workflow
-			Zeros,				// F90				- !Metalness workflow
-			Roughness,			// Roughness
-			Anisotropy,			// Anisotropy
-			0.f,				// SSSProfileID
-			0.f,				// SSSDMFP
-			0.f,				// SSSDMFPScale
-			Emissive,			// Emissive
-			0.f,				// Haziness
-			0.f,				// ThinFilmThickness
-			ClearCoat,			// FuzzAmount
-			SubSurfaceColor,	// FuzzColor
-			DefaultThickness,
-			SharedLocalBasisIndex, SharedLocalBasisTypes);
+		Top.BaseColor = BaseColor;
+		Top.Specular = Specular;
+		Top.Metallic = Metallic;
+		Top.Roughness = Roughness;
+		Top.Anisotropy = Anisotropy;
+		Top.SSSProfileID = 0.f;
+		Top.SSSDMFP = 0.f;
+		Top.SSSDMFPScale = 0.f;
+		Top.Emissive = Emissive;
+		Top.FuzzAmount = ClearCoat;
+		Top.FuzzColor = SubSurfaceColor;
 
-		// Take in account the legacy opacity weighting for translucent material
-		#if !MATERIALBLENDING_MASKED && !MATERIALBLENDING_SOLID
-		{
-			Out = StrataWeight(Out, Opacity);
-		}
-		#endif
+		Bot.Weight = 0.f;
+		FinalWeight = Opacity;
 	}
 	else if (ShadingModel == SHADINGMODELID_EYE)
 	{
@@ -280,55 +260,126 @@ FStrataData StrataConvertLegacyMaterial(
 		const float IrisOpacity = saturate(1.0f - IrisMask);
 
 		// Incorrect composition, but allows to compile shaders
-		{
-			const float3 TopBaseColor = BaseColor; // lerp(BaseColor, Zeros, IrisMask);
-
-			FStrataData CoatSlab = GetStrataSlabBSDF(
-				true,				// UseMetalness
-				TopBaseColor,		// BaseColor		-  Metalness workflow
-				Ones,				// EdgeColor		-  Metalness workflow
-				Specular,			// Specular			-  Metalness workflow
-				0.f,				// Metallic			-  Metalness workflow
-				Zeros,				// DiffuseAlbedo	- !Metalness workflow
-				Zeros, 				// F0				- !Metalness workflow
-				Zeros,				// F90				- !Metalness workflow
-				Roughness,			// Roughness
-				0.f,				// Anisotropy
-				SubSurfaceProfileId,// SSSProfileID
-				Zeros,				// SSSDMFP
-				Opacity,			// SSSDMFPScale
-				Emissive,			// Emissive
-				0.f,				// Haziness
-				0.f,				// ThinFilmThickness
-				0.f,				// FuzzAmount
-				Zeros,				// FuzzColor
-				DefaultThickness,
-				SharedLocalBasisIndex, SharedLocalBasisTypes);
-
-			FStrataData IrisSlab = GetStrataSlabBSDF(
-				true,				// UseMetalness
-				BaseColor,			// BaseColor		-  Metalness workflow
-				Ones,				// EdgeColor		-  Metalness workflow
-				Specular,			// Specular			-  Metalness workflow
-				0.f,				// Metallic			-  Metalness workflow
-				Zeros,				// DiffuseAlbedo	- !Metalness workflow
-				Zeros, 				// F0				- !Metalness workflow
-				Zeros,				// F90				- !Metalness workflow
-				Roughness,			// Roughness
-				0.f,				// Anisotropy
-				SubSurfaceProfileId,// SSSProfileID
-				Zeros,				// SSSDMFP
-				Opacity,			// SSSDMFPScale
-				Emissive,			// Emissive
-				0.f,				// Haziness
-				0.f,				// ThinFilmThickness
-				0.f,				// FuzzAmount
-				Zeros,				// FuzzColor
-				DefaultThickness,
-				ClearCoatBottomNormal_SharedLocalBasisIndex, /*ClearCoatBottomNormal_*/SharedLocalBasisTypes);
-
-			Out = StrataVerticalLayering(CoatSlab, IrisSlab);
-		}
+		
+		// Cornea/Sclera
+		Top.BaseColor = BaseColor;
+		Top.Specular = Specular;
+		Top.Metallic = 0.f;
+		Top.Roughness = Roughness;
+		Top.Anisotropy = 0.f;
+		Top.SSSProfileID = SubSurfaceProfileId;
+		Top.SSSDMFP = Zeros;
+		Top.SSSDMFPScale = Opacity;
+		Top.Emissive = Emissive;
+		Top.SharedLocalBasisIndex = SharedLocalBasisIndex;
+		Top.Weight = 1.f;
+
+		// Iris
+		Bot.BaseColor = BaseColor;
+		Bot.Specular = Specular;
+		Bot.Metallic = .0f;
+		Bot.Roughness = Roughness;
+		Bot.Anisotropy = 0.f;
+		Bot.SSSProfileID = SubSurfaceProfileId;
+		Bot.SSSDMFP = Zeros;
+		Bot.SSSDMFPScale = Opacity;
+		Bot.Emissive = Emissive;
+		Bot.SharedLocalBasisIndex = ClearCoatBottomNormal_SharedLocalBasisIndex;
+		Bot.Weight = 1.f;
+
+		// No need to take into account opacity weighting for translucent material, as the coverage is assumed to be 1
+		FinalWeight = 1.0f;
+	}
+	else if (ShadingModel == SHADINGMODELID_THIN_TRANSLUCENT)	
+	{
+		// Top slab BSDF as a simple Disney material
+		Top.BaseColor = BaseColor;
+		Top.Specular = Specular;
+		Top.Metallic = Metallic;
+		Top.Roughness = Roughness;
+		Top.Anisotropy = Anisotropy;
+		Top.SSSProfileID = 0.f;
+		Top.SSSDMFP = Zeros;
+		Top.SSSDMFPScale = 0.f;
+		Top.Emissive = Emissive;
+		Top.SharedLocalBasisIndex = SharedLocalBasisIndex;
+
+		// Bottom slab BSDF will be a simple absorption only layer
+		const float Thickness = DefaultThickness;
+		const float3 MFP = TransmittanceToMeanFreePath(TransmittanceColor, Thickness * CENTIMETER_TO_METER) * METER_TO_CENTIMETER;
+
+		// Assign specular properties shared with the top layer.
+		Bot.BaseColor = Zeros;
+		Bot.Specular = Specular;
+		Bot.Metallic = .0f;
+		Bot.Roughness = Roughness;
+		Bot.Anisotropy = 0.f;
+		Bot.SSSProfileID = .0f;
+		Bot.SSSDMFP = MFP;
+		Bot.SSSDMFPScale = 1.f;
+		Bot.Emissive = Zeros;
+		Bot.Thickness = Thickness;
+		Bot.SharedLocalBasisIndex = ClearCoatBottomNormal_SharedLocalBasisIndex;
+
+		// No need to take into account opacity weighting for translucent material, as the coverage is assumed to be 1
+		// Now weight the top base material by opacity
+		Top.Weight = Opacity;
+	}
+
+	// Fixed layer structure for helping compiler to unroll and optimize shader
+	FStrataData Top_Slab = CreateLegacySlab(Top, SharedLocalBasisTypes);
+	FStrataData Bot_Slab = CreateLegacySlab(Bot, SharedLocalBasisTypes);
+	FStrataData Out = StrataVerticalLayering(StrataWeight(Top_Slab, Top.Weight), StrataWeight(Bot_Slab, Bot.Weight));
+
+	// Take in account the legacy opacity weighting for translucent material
+	#if !MATERIALBLENDING_MASKED && !MATERIALBLENDING_SOLID
+	{
+		Out = StrataWeight(Out, FinalWeight);
+	}
+	#endif
+
+	return Out;
+}
+
+// Convert legacy shading models - Static
+// This function is for static single shading model (i.e., known at shader compilation time). 
+// It reuses the dynamic version for most part, but for special node like Unlit/Hair/Water, 
+// we use the dedicated node
+FStrataData StrataConvertLegacyMaterialStatic(
+	float3 BaseColor, float Specular, float Metallic,
+	float Roughness, float Anisotropy,
+	float3 SubSurfaceColor, float SubSurfaceProfileId,
+	float ClearCoat, float ClearCoatRoughness,
+	float3 Emissive,
+	float Opacity,
+	float3 TransmittanceColor,
+	float3 WaterScatteringCoefficients, float3 WaterAbsorptionCoefficients, float WaterPhaseG, float3 ColorScaleBehindWater,
+	uint ShadingModel,
+	uint SharedLocalBasisIndex,
+	uint ClearCoatBottomNormal_SharedLocalBasisIndex,
+	inout uint SharedLocalBasisTypes)
+{
+	FStrataData Out = GetInitialisedStrataData();
+
+	const float DefaultThickness = STRATA_LAYER_DEFAULT_THICKNESS_CM;
+	const float3 Zeros = float3(0, 0, 0);
+	const float3 Ones = float3(1, 1, 1);
+
+	if (ShadingModel == SHADINGMODELID_UNLIT)
+	{
+		Out = GetStrataUnlitBSDF(Emissive, Ones /*TransnmittanceColor*/);
+	}
+	else if (ShadingModel == SHADINGMODELID_HAIR)
+	{
+		Out = GetStrataHairBSDF(
+			BaseColor,	   // BaseColor
+			Metallic,	   // Scatter
+			Specular,	   // Specular
+			Roughness,	   // Roughness
+			ClearCoat,	   // Backlit
+			Emissive,	   // EmissiveColor
+			SharedLocalBasisIndex);
+
 		// No need to take into account opacity weighting for translucent material, as the coverage is assumed to be 1
 	}
 	else if (ShadingModel == SHADINGMODELID_SINGLELAYERWATER)
@@ -351,72 +402,22 @@ FStrataData StrataConvertLegacyMaterial(
 
 		// No need to take into account opacity weighting for translucent material, as the coverage is assumed to be 1
 	}
-	else if (ShadingModel == SHADINGMODELID_THIN_TRANSLUCENT)	
+	else
 	{
-		// Top slab BSDF as a simple Disney material
-		FStrataData TopSlab = GetStrataSlabBSDF(
-			true,				// UseMetalness
-			BaseColor,			// BaseColor		-  Metalness workflow
-			Ones,				// EdgeColor		-  Metalness workflow
-			Specular,			// Specular			-  Metalness workflow
-			Metallic,			// Metallic			-  Metalness workflow
-			Zeros,				// DiffuseAlbedo	- !Metalness workflow
-			Zeros, 				// F0				- !Metalness workflow
-			Zeros,				// F90				- !Metalness workflow
-			Roughness,			// Roughness
-			Anisotropy,			// Anisotropy
-			0.f,				// SSSProfileID
-			Zeros,				// SSSDMFP
-			0.f,				// SSSDMFPScale
-			Emissive,			// Emissive
-			0.f,				// Haziness
-			0.f,				// ThinFilmThickness
-			0.f,				// FuzzAmount
-			Zeros,				// FuzzColor
-			DefaultThickness,
-			SharedLocalBasisIndex, SharedLocalBasisTypes);
-
-		// Now weight the top base material by opacity
-		FStrataData WeightedTopSlab = StrataWeight(TopSlab, Opacity);
-
-		// Bottom slab BSDF will be a simple absorption only layer
-		const bool bHasTransmittance = any(TransmittanceColor > 0.f);
-		if (bHasTransmittance)
-		{
-			const float Thickness = DefaultThickness;
-			const float3 MFP = TransmittanceToMeanFreePath(TransmittanceColor, Thickness * CENTIMETER_TO_METER) * METER_TO_CENTIMETER;
-
-			// Assign specular properties shared with the top layer.
-			FStrataData BottomSlab = GetStrataSlabBSDF(
-				true,				// UseMetalness
-				Zeros,	// BaseColor		-  Metalness workflow
-				Ones,				// EdgeColor		-  Metalness workflow
-				Specular,			// Specular			-  Metalness workflow
-				0.f,				// Metallic			-  Metalness workflow
-				Zeros,				// DiffuseAlbedo	- !Metalness workflow
-				Zeros, 				// F0				- !Metalness workflow
-				Zeros,				// F90				- !Metalness workflow
-				Roughness,			// Roughness
-				0.f,				// Anisotropy
-				0.f,				// SSSProfileID
-				MFP,				// SSSDMFP
-				1.f,				// SSSDMFPScale
-				Zeros,				// Emissive
-				0.f,				// Haziness
-				0.f,				// ThinFilmThickness
-				0.f,				// FuzzAmount
-				Zeros,				// FuzzColor
-				Thickness,
-				SharedLocalBasisIndex, SharedLocalBasisTypes);
-
-			Out = StrataVerticalLayering(WeightedTopSlab, BottomSlab);
-		}
-		else
-		{
-			Out = TopSlab;
-		}
-
-		// No need to take into account opacity weighting for translucent material, as the coverage is assumed to be 1
+		Out = StrataConvertLegacyMaterialDynamic(
+			BaseColor, Specular, Metallic,
+			Roughness, Anisotropy,
+			SubSurfaceColor, SubSurfaceProfileId,
+			ClearCoat, ClearCoatRoughness,
+			Emissive,
+			Opacity,
+			TransmittanceColor,
+			WaterScatteringCoefficients, WaterAbsorptionCoefficients, WaterPhaseG, ColorScaleBehindWater,
+			ShadingModel,
+			SharedLocalBasisIndex,
+			ClearCoatBottomNormal_SharedLocalBasisIndex,
+			SharedLocalBasisTypes);
 	}
+
 	return Out;
-}
\ No newline at end of file
+}
diff --git a/Engine/Source/Runtime/Engine/Private/Materials/HLSLMaterialTranslator.cpp b/Engine/Source/Runtime/Engine/Private/Materials/HLSLMaterialTranslator.cpp
--- a/Engine/Source/Runtime/Engine/Private/Materials/HLSLMaterialTranslator.cpp
+++ b/Engine/Source/Runtime/Engine/Private/Materials/HLSLMaterialTranslator.cpp
@@ -9536,6 +9536,7 @@ int32 FHLSLMaterialTranslator::StrataSlabBSDF(
 }
 
 int32 FHLSLMaterialTranslator::StrataConversionFromLegacy(
+	bool bHasDynamicShadingModels,
 	int32 BaseColor, int32 Specular, int32 Metallic,
 	int32 Roughness, int32 Anisotropy,
 	int32 SubSurfaceColor, int32 SubSurfaceProfileId,
@@ -9555,7 +9556,8 @@ int32 FHLSLMaterialTranslator::StrataConversionFromLegacy(
 	const FString ClearCoat_TangentCode = Tangent != INDEX_NONE ? *GetParameterCode(ClearCoat_Tangent) : TEXT("NONE");
 
 	return AddCodeChunk(
-		MCT_Strata, TEXT("StrataConvertLegacyMaterial(%s, %s, %s, %s, %s, %s, %s, %s, %s, %s, %s, %s, %s, %s, %s, %s, %s, %s, %s, Parameters.SharedLocalBases.Types) /* Normal = %s ; Tangent = %s ; ClearCoat_Normal = %s ; ClearCoat_Tangent = %s */"),
+		MCT_Strata, TEXT("StrataConvertLegacyMaterial%s(%s, %s, %s, %s, %s, %s, %s, %s, %s, %s, %s, %s, %s, %s, %s, %s, %s, %s, %s, Parameters.SharedLocalBases.Types) /* Normal = %s ; Tangent = %s ; ClearCoat_Normal = %s ; ClearCoat_Tangent = %s */"),
+		bHasDynamicShadingModels ? TEXT("Dynamic") : TEXT("Static"),
 		*GetParameterCode(BaseColor), *GetParameterCode(Specular), *GetParameterCode(Metallic),
 		*GetParameterCode(Roughness), *GetParameterCode(Anisotropy),
 		*GetParameterCode(SubSurfaceColor), *GetParameterCode(SubSurfaceProfileId),
diff --git a/Engine/Source/Runtime/Engine/Private/Materials/HLSLMaterialTranslator.h b/Engine/Source/Runtime/Engine/Private/Materials/HLSLMaterialTranslator.h
--- a/Engine/Source/Runtime/Engine/Private/Materials/HLSLMaterialTranslator.h
+++ b/Engine/Source/Runtime/Engine/Private/Materials/HLSLMaterialTranslator.h
@@ -1014,6 +1014,7 @@ protected:
 		int32 Thickness,
 		int32 Normal, int32 Tangent, const FString& SharedLocalBasisIndexMacro) override;
 	virtual int32 StrataConversionFromLegacy(
+		bool bHasDynamicShadingModel,
 		int32 BaseColor, int32 Specular, int32 Metallic,
 		int32 Roughness, int32 Anisotropy,
 		int32 SubSurfaceColor, int32 SubSurfaceProfileId,
diff --git a/Engine/Source/Runtime/Engine/Private/Materials/MaterialExpressions.cpp b/Engine/Source/Runtime/Engine/Private/Materials/MaterialExpressions.cpp
--- a/Engine/Source/Runtime/Engine/Private/Materials/MaterialExpressions.cpp
+++ b/Engine/Source/Runtime/Engine/Private/Materials/MaterialExpressions.cpp
@@ -42,6 +42,7 @@
 #include "Engine/TextureCubeArray.h"
 #include "Engine/TextureRenderTargetCube.h"
 #include "Engine/VolumeTexture.h"
+#include "Engine/SubsurfaceProfile.h"
 #include "Styling/CoreStyle.h"
 #include "VT/RuntimeVirtualTexture.h"
 #include "ProfilingDebugging/LoadTimeTracker.h"
@@ -20235,12 +20236,13 @@ int32 UMaterialExpressionStrataLegacyConversion::Compile(class FMaterialCompiler
 	const bool bHasSSS = HasSSS();
 	if (bHasSSS)
 	{
-		FName NameSubsurfaceProfile(FString(TEXT("__SubsurfaceProfile")));
-		SSSProfileCodeChunk = Compiler->ForceCast(Compiler->ScalarParameter(NameSubsurfaceProfile, 1.0f), MCT_Float1);
+		SSSProfileCodeChunk = Compiler->ForceCast(Compiler->ScalarParameter(GetSubsurfaceProfileParameterName(), 1.0f), MCT_Float1);
 	}
 
+	const bool bHasDynamicShadingModels = ConvertedStrataMaterialInfo.CountShadingModels() > 1;
 	// We probably need to do something along these line as well :::
 	int32 OutputCodeChunk = Compiler->StrataConversionFromLegacy(
+		bHasDynamicShadingModels,
 		// Metalness workflow
 		CompileWithDefaultFloat3(Compiler, BaseColor, 0.0f, 0.0f, 0.0f),
 		CompileWithDefaultFloat1(Compiler, Specular, 0.5f),
@@ -20549,8 +20551,7 @@ int32 UMaterialExpressionStrataSlabBSDF::Compile(class FMaterialCompiler* Compil
 	const bool bHasSSS = HasSSS();
 	if (bHasSSS)
 	{
-		FName NameSubsurfaceProfile(FString(TEXT("__SubsurfaceProfile")));
-		SSSProfileCodeChunk = Compiler->ForceCast(Compiler->ScalarParameter(NameSubsurfaceProfile, 1.0f), MCT_Float1);
+		SSSProfileCodeChunk = Compiler->ForceCast(Compiler->ScalarParameter(GetSubsurfaceProfileParameterName(), 1.0f), MCT_Float1);
 	}
 
 	auto DielectricSpecularToF0 = [](float SpecularIn)
diff --git a/Engine/Source/Runtime/Engine/Public/MaterialCompiler.h b/Engine/Source/Runtime/Engine/Public/MaterialCompiler.h
--- a/Engine/Source/Runtime/Engine/Public/MaterialCompiler.h
+++ b/Engine/Source/Runtime/Engine/Public/MaterialCompiler.h
@@ -455,6 +455,7 @@ public:
 		int32 Thickness,
 		int32 Normal, int32 Tangent, const FString& SharedLocalBasisIndexMacro) = 0;
 	virtual int32 StrataConversionFromLegacy(
+		bool bHasDynamicShadingModels,
 		int32 BaseColor, int32 Specular, int32 Metallic,
 		int32 Roughness, int32 Anisotropy,
 		int32 SubSurfaceColor, int32 SubSurfaceProfileId,
@@ -994,6 +995,7 @@ public:
 	}
 
 	virtual int32 StrataConversionFromLegacy(
+		bool bHasDynamicShadingModels,
 		int32 BaseColor, int32 Specular, int32 Metallic,
 		int32 Roughness, int32 Anisotropy,
 		int32 SubSurfaceColor, int32 SubSurfaceProfileId,
@@ -1007,6 +1009,7 @@ public:
 		int32 ClearCoat_Normal, int32 ClearCoat_Tangent, const FString& ClearCoat_SharedLocalBasisIndexMacro) override
 	{
 		return Compiler->StrataConversionFromLegacy(
+			bHasDynamicShadingModels,
 			BaseColor, Specular, Metallic,
 			Roughness, Anisotropy,
 			SubSurfaceColor, SubSurfaceProfileId,
