Fix GPU crash during translucent pass.

diff --git a/Engine/Shaders/Private/BasePassPixelShader.usf b/Engine/Shaders/Private/BasePassPixelShader.usf
--- a/Engine/Shaders/Private/BasePassPixelShader.usf
+++ b/Engine/Shaders/Private/BasePassPixelShader.usf
@@ -1662,6 +1662,7 @@ void FPixelShaderInOut_MainPS(
 		Color = StrataForwardLighting(
 			EyeIndex,
 			In.SvPosition,
+			TranslucentBasePass.Strata.bRoughDiffuse,
 			BasePassInterpolants,
 			Interpolants,
 			LightmapVTPageTableResult,
diff --git a/Engine/Shaders/Private/DistortAccumulatePS.usf b/Engine/Shaders/Private/DistortAccumulatePS.usf
--- a/Engine/Shaders/Private/DistortAccumulatePS.usf
+++ b/Engine/Shaders/Private/DistortAccumulatePS.usf
@@ -153,7 +153,7 @@ void Main(
 			const float3 FakeL = float3(0, 0, 1);
 			FStrataBSDFContext StrataBSDFContext = StrataCreateBSDFContext(StrataPixelHeader, BSDF, NullStrataAddressing, V, FakeL);
 
-			FStrataEvaluateResult StrataEvaluate = StrataEvaluateBSDF(StrataBSDFContext); // Works to call because FakeL is nto used when computing StrataEvaluate.Throughput
+			FStrataEvaluateResult StrataEvaluate = StrataEvaluateBSDF(StrataBSDFContext, true); // Works to call because FakeL is nto used when computing StrataEvaluate.Throughput
 
 			// We only support rough refraction from Slab material type for now.
 			if (BSDF_GETTYPE(BSDF) == STRATA_BSDF_TYPE_SLAB)
diff --git a/Engine/Shaders/Private/Strata/StrataDeferredLighting.ush b/Engine/Shaders/Private/Strata/StrataDeferredLighting.ush
--- a/Engine/Shaders/Private/Strata/StrataDeferredLighting.ush
+++ b/Engine/Shaders/Private/Strata/StrataDeferredLighting.ush
@@ -53,14 +53,14 @@ float4 StrataDeferredLighting(
 
 			// We must have the StrataIntegrateBSDF inside the if due to the rectlight texture: it must be non ambiguous which texture is going to be used.
 			// After the compilation, a local resource must map to a unique global resource (the default or the actual rect light texture).
-			BSDFEvaluate = StrataIntegrateBSDF(BSDFContext, ShadowTerms, AreaLightContext, bForceFullyRough);
+			BSDFEvaluate = StrataIntegrateBSDF(BSDFContext, ShadowTerms, AreaLightContext, bForceFullyRough, Strata.bRoughDiffuse);
 		}
 		else
 		{
 			FCapsuleLight Capsule = GetCapsule(ToLight, LightData);
 			AreaLightContext = CreateCapsuleIntegrateContext(Roughness, BSDFContext.N, BSDFContext.V, Capsule, LightData.bInverseSquared);
 
-			BSDFEvaluate = StrataIntegrateBSDF(BSDFContext, ShadowTerms, AreaLightContext, bForceFullyRough);
+			BSDFEvaluate = StrataIntegrateBSDF(BSDFContext, ShadowTerms, AreaLightContext, bForceFullyRough, Strata.bRoughDiffuse);
 		}
 
 		float3 DiffuseLuminance = BSDFEvaluate.IntegratedDiffuseValue;
diff --git a/Engine/Shaders/Private/Strata/StrataEvaluation.ush b/Engine/Shaders/Private/Strata/StrataEvaluation.ush
--- a/Engine/Shaders/Private/Strata/StrataEvaluation.ush
+++ b/Engine/Shaders/Private/Strata/StrataEvaluation.ush
@@ -199,7 +199,7 @@ struct FStrataEvaluateResult
 #define INTEGRATION_AREA_LIGHT_CAPSULE	1
 #define INTEGRATION_AREA_LIGHT_RECT		2
 
-FStrataEvaluateResult StrataEvaluateBSDFCommon(FStrataBSDFContext BSDFContext, FShadowTerms ShadowTerms, FAreaLightIntegrateContext AreaLightContext, int IntegrationType, bool bForceFullyRough)
+FStrataEvaluateResult StrataEvaluateBSDFCommon(FStrataBSDFContext BSDFContext, FShadowTerms ShadowTerms, FAreaLightIntegrateContext AreaLightContext, int IntegrationType, bool bForceFullyRough, bool bRoughDiffuseEnabled)
 {
 	FStrataEvaluateResult Sample = (FStrataEvaluateResult)0;
 
@@ -255,7 +255,7 @@ FStrataEvaluateResult StrataEvaluateBSDFCommon(FStrataBSDFContext BSDFContext, F
 			//// Evaluate the diffuse component.
 			////
 
-			if (STRATA_DIFFUSE_CHAN && Strata.bRoughDiffuse)
+			if (STRATA_DIFFUSE_CHAN && bRoughDiffuseEnabled)
 			{
 				// If the specular layer is anisotropic, the diffuse takes the 'main' roughness value rather than the tangent/bitangent value
 				Sample.DiffusePathValue = Diffuse_Chan(DiffuseColor, Alpha2, NoV, AreaLightContext.NoL, VoH, NoH, GetAreaLightDiffuseMicroReflWeight(AreaLightContext.AreaLight));
@@ -298,7 +298,7 @@ FStrataEvaluateResult StrataEvaluateBSDFCommon(FStrataBSDFContext BSDFContext, F
 				// To make this a safe operation:
 				//	- We clamp SpecularTChan to a threshold to avoid div by 0 when the material is a metal. 
 				//	- When the material tends to be metal, DirectionalAlbedo_SpecularTransmission will naturally tend to be 0
-				if (STRATA_DIFFUSE_CHAN > 0 && Strata.bRoughDiffuse)
+				if (STRATA_DIFFUSE_CHAN > 0 && bRoughDiffuseEnabled)
 				{
 					const float ChanF0 = 0.04;
 					FBxDFEnergyTerms EnergyTerms = ComputeGGXSpecEnergyTerms(SafeRoughness, BSDFContext.Context.NoV, ChanF0, 1.0f);
@@ -666,7 +666,7 @@ FStrataEvaluateResult StrataEvaluateBSDFCommon(FStrataBSDFContext BSDFContext, F
 	return Sample;
 }
 
-FStrataEvaluateResult StrataEvaluateBSDF(FStrataBSDFContext BSDFContext)
+FStrataEvaluateResult StrataEvaluateBSDF(FStrataBSDFContext BSDFContext, bool bRoughDiffuseEnabled)
 {
 	FShadowTerms IdentityShadow = { 1, 1, 1, InitHairTransmittanceData() };
 	FAreaLightIntegrateContext DummyAreaLightContext = InitAreaLightIntegrateContext();
@@ -674,10 +674,10 @@ FStrataEvaluateResult StrataEvaluateBSDF(FStrataBSDFContext BSDFContext)
 	DummyAreaLightContext.NoL = BSDFContext.Context.NoL;
 	DummyAreaLightContext.Falloff = 1;
 	const bool bForceFullyRough = false;
-	return StrataEvaluateBSDFCommon(BSDFContext, IdentityShadow, DummyAreaLightContext, INTEGRATION_PUNCTUAL_LIGHT, bForceFullyRough);
+	return StrataEvaluateBSDFCommon(BSDFContext, IdentityShadow, DummyAreaLightContext, INTEGRATION_PUNCTUAL_LIGHT, bForceFullyRough, bRoughDiffuseEnabled);
 }
 
-FStrataEvaluateResult StrataIntegrateBSDF(FStrataBSDFContext BSDFContext, FShadowTerms Shadow, FAreaLightIntegrateContext AreaLightContext, bool bForceFullyRough)
+FStrataEvaluateResult StrataIntegrateBSDF(FStrataBSDFContext BSDFContext, FShadowTerms Shadow, FAreaLightIntegrateContext AreaLightContext, bool bForceFullyRough, bool bRoughDiffuseEnabled)
 {
 	BRANCH
 	if (IsAreaLight(AreaLightContext.AreaLight))
@@ -685,15 +685,15 @@ FStrataEvaluateResult StrataIntegrateBSDF(FStrataBSDFContext BSDFContext, FShado
 		BRANCH
 		if(IsRectLight(AreaLightContext.AreaLight))
 		{
-			return StrataEvaluateBSDFCommon(BSDFContext, Shadow, AreaLightContext, INTEGRATION_AREA_LIGHT_RECT, bForceFullyRough);
+			return StrataEvaluateBSDFCommon(BSDFContext, Shadow, AreaLightContext, INTEGRATION_AREA_LIGHT_RECT, bForceFullyRough, bRoughDiffuseEnabled);
 		}
 		else
 		{
-			return StrataEvaluateBSDFCommon(BSDFContext, Shadow, AreaLightContext, INTEGRATION_AREA_LIGHT_CAPSULE, bForceFullyRough);
+			return StrataEvaluateBSDFCommon(BSDFContext, Shadow, AreaLightContext, INTEGRATION_AREA_LIGHT_CAPSULE, bForceFullyRough, bRoughDiffuseEnabled);
 		}
 	}
 
-	return StrataEvaluateBSDFCommon(BSDFContext, Shadow, AreaLightContext, INTEGRATION_PUNCTUAL_LIGHT, bForceFullyRough);
+	return StrataEvaluateBSDFCommon(BSDFContext, Shadow, AreaLightContext, INTEGRATION_PUNCTUAL_LIGHT, bForceFullyRough, bRoughDiffuseEnabled);
 }
 
 struct FStrataImportanceSampleResult
diff --git a/Engine/Shaders/Private/Strata/StrataForwardLighting.ush b/Engine/Shaders/Private/Strata/StrataForwardLighting.ush
--- a/Engine/Shaders/Private/Strata/StrataForwardLighting.ush
+++ b/Engine/Shaders/Private/Strata/StrataForwardLighting.ush
@@ -30,6 +30,7 @@ void GetPrecomputedIndirectLightingAndSkyLight(
 float3 StrataForwardLightingCommon(
 	float Dither,
 	bool bForceFullyRough,
+	bool bRoughDiffuseEnabled, 
 	FDeferredLightData LightData,
 	float3 ToLight,
 	float LightMask,
@@ -87,7 +88,7 @@ float3 StrataForwardLightingCommon(
 
 				// We must have the evaluate inside the if due to the rectlight texture: it must be now be ambiguous which texture is going ot be used.
 				// After te compilation, a local resource must map to a unique global resource (the default or the actual rect light texture).
-				BSDFEvaluate = StrataIntegrateBSDF(StrataBSDFContext, ShadowTerms, AreaLightContext, bForceFullyRough);
+				BSDFEvaluate = StrataIntegrateBSDF(StrataBSDFContext, ShadowTerms, AreaLightContext, bForceFullyRough, bRoughDiffuseEnabled);
 				bEvaluateHasBeenExecuted = true;
 			}
 			else
@@ -96,7 +97,7 @@ float3 StrataForwardLightingCommon(
 				AreaLightContext = CreateCapsuleIntegrateContext(Roughness, StrataBSDFContext.N, StrataBSDFContext.V, Capsule, LightData.bInverseSquared);
 				Lighting = IntegrateLight(Capsule, LightData.bInverseSquared);
 
-				BSDFEvaluate = StrataIntegrateBSDF(StrataBSDFContext, ShadowTerms, AreaLightContext, bForceFullyRough);
+				BSDFEvaluate = StrataIntegrateBSDF(StrataBSDFContext, ShadowTerms, AreaLightContext, bForceFullyRough, bRoughDiffuseEnabled);
 				bEvaluateHasBeenExecuted = true;
 			}
 
@@ -119,7 +120,7 @@ float3 StrataForwardLightingCommon(
 
 				// We must have the evaluate inside the if due to the rectlight texture: it must be now be ambiguous which texture is going ot be used.
 				// After te compilation, a local resource must map to a unique global resource (the default or the actual rect light texture).
-				BSDFEvaluate = StrataIntegrateBSDF(StrataBSDFContext, ShadowTerms, AreaLightContext, bForceFullyRough);
+				BSDFEvaluate = StrataIntegrateBSDF(StrataBSDFContext, ShadowTerms, AreaLightContext, bForceFullyRough, bRoughDiffuseEnabled);
 				bEvaluateHasBeenExecuted = true;
 			}
 			else
@@ -127,7 +128,7 @@ float3 StrataForwardLightingCommon(
 				FCapsuleLight Capsule = GetCapsule(ToLight, LightData);
 				AreaLightContext = CreateCapsuleIntegrateContext(Roughness, StrataBSDFContext.N, StrataBSDFContext.V, Capsule, LightData.bInverseSquared);
 
-				BSDFEvaluate = StrataIntegrateBSDF(StrataBSDFContext, ShadowTerms, AreaLightContext, bForceFullyRough);
+				BSDFEvaluate = StrataIntegrateBSDF(StrataBSDFContext, ShadowTerms, AreaLightContext, bForceFullyRough, bRoughDiffuseEnabled);
 				bEvaluateHasBeenExecuted = true;
 			}
 
@@ -150,6 +151,7 @@ float3 StrataForwardLightingCommon(
 float3 StrataForwardLighting(
 	uint EyeIndex,
 	float4 SvPosition,
+	bool bRoughDiffuseEnabled,
 	FBasePassInterpolantsVSToPS BasePassInterpolants,
 	FVertexFactoryInterpolantsVSToPS Interpolants,
 	VTPageTableResult LightmapVTPageTableResult,
@@ -390,7 +392,7 @@ float3 StrataForwardLighting(
 			float FinalLerpFactor;
 			ComputeVolumeUVs(AbsoluteWorldPosition, LightingPositionOffset, InnerVolumeUVs, OuterVolumeUVs, FinalLerpFactor);
 
-			FStrataEvaluateResult BSDFEvaluate = StrataEvaluateBSDF(StrataBSDFContext);
+			FStrataEvaluateResult BSDFEvaluate = StrataEvaluateBSDF(StrataBSDFContext, bRoughDiffuseEnabled);
 			LayerThroughput += BSDFCoverage * BSDFEvaluate.Throughput;
 
 		#if TRANSLUCENCY_PERVERTEX_FORWARD_SHADING
@@ -489,6 +491,7 @@ float3 StrataForwardLighting(
 					Color += StrataForwardLightingCommon(
 						Dither,
 						bForceFullyRough,
+						bRoughDiffuseEnabled,
 						DirLightData,
 						ToDirLight,
 						DirLightMask,
@@ -512,7 +515,7 @@ float3 StrataForwardLighting(
 				if(!bEvaluateHasBeenExecuted)
 				{
 					// Only update the throughput
-					FStrataEvaluateResult BSDFEvaluate = StrataEvaluateBSDF(StrataBSDFContext);
+					FStrataEvaluateResult BSDFEvaluate = StrataEvaluateBSDF(StrataBSDFContext, bRoughDiffuseEnabled);
 					LayerThroughput += BSDFCoverage * BSDFEvaluate.Throughput;
 				}
 			}
@@ -583,6 +586,7 @@ float3 StrataForwardLighting(
 				Color += StrataForwardLightingCommon(
 					Dither,
 					bForceFullyRough,
+					bRoughDiffuseEnabled,
 					LightData,
 					ToLight,
 					LightMask,
diff --git a/Engine/Shaders/Private/Strata/StrataVisualize.usf b/Engine/Shaders/Private/Strata/StrataVisualize.usf
--- a/Engine/Shaders/Private/Strata/StrataVisualize.usf
+++ b/Engine/Shaders/Private/Strata/StrataVisualize.usf
@@ -65,7 +65,7 @@ float3 DrawPlot(
 		TangentBasis[2] = float3(0, 0, 1);
 
 		FStrataBSDFContext BSDFContext = StrataCreateBSDFContext(TangentBasis, BSDF, LocalV, L, true);
-		FStrataEvaluateResult BSDFEvaluate = StrataEvaluateBSDF(BSDFContext);
+		FStrataEvaluateResult BSDFEvaluate = StrataEvaluateBSDF(BSDFContext, Strata.bRoughDiffuse);
 		Out = (BSDFEvaluate.DiffusePathValue + BSDFEvaluate.SpecularPathValue + BSDFEvaluate.TransmissionPathValue) * BSDFContext.SatNoL; // * BSDF.Weight
 	}
 	return Out;
diff --git a/Engine/Source/Runtime/Renderer/Private/BasePassRendering.h b/Engine/Source/Runtime/Renderer/Private/BasePassRendering.h
--- a/Engine/Source/Runtime/Renderer/Private/BasePassRendering.h
+++ b/Engine/Source/Runtime/Renderer/Private/BasePassRendering.h
@@ -97,7 +97,7 @@ END_GLOBAL_SHADER_PARAMETER_STRUCT()
 BEGIN_GLOBAL_SHADER_PARAMETER_STRUCT(FTranslucentBasePassUniformParameters,)
 	SHADER_PARAMETER_STRUCT(FSharedBasePassUniformParameters, Shared)
 	SHADER_PARAMETER_STRUCT(FSceneTextureUniformParameters, SceneTextures)
-	SHADER_PARAMETER_STRUCT(FStrataGlobalUniformParameters, Strata)
+	SHADER_PARAMETER_STRUCT(FStrataForwardPassUniformParameters, Strata)
 	SHADER_PARAMETER_STRUCT(FLightCloudTransmittanceParameters, ForwardDirLightCloudShadow)
 	// Material SSR
 	SHADER_PARAMETER(FVector4f, HZBUvFactorAndInvFactor)
diff --git a/Engine/Source/Runtime/Renderer/Private/Strata/Strata.cpp b/Engine/Source/Runtime/Renderer/Private/Strata/Strata.cpp
--- a/Engine/Source/Runtime/Renderer/Private/Strata/Strata.cpp
+++ b/Engine/Source/Runtime/Renderer/Private/Strata/Strata.cpp
@@ -277,7 +277,7 @@ void BindStrataGlobalUniformParameters(FRDGBuilder& GraphBuilder, FStrataSceneDa
 	{
 		OutStrataUniformParameters.bRoughDiffuse = StrataSceneData->bRoughDiffuse ? 1u : 0u;
 		OutStrataUniformParameters.MaxBytesPerPixel = StrataSceneData->MaxBytesPerPixel;
-		OutStrataUniformParameters.MaterialTextureArray = StrataSceneData->MaterialTextureArraySRV;
+		OutStrataUniformParameters.MaterialTextureArray = StrataSceneData->MaterialTextureArray;
 		OutStrataUniformParameters.TopLayerTexture = StrataSceneData->TopLayerTexture;
 		OutStrataUniformParameters.SSSTexture = StrataSceneData->SSSTexture;
 	}
@@ -287,12 +287,24 @@ void BindStrataGlobalUniformParameters(FRDGBuilder& GraphBuilder, FStrataSceneDa
 		FRDGTextureRef DefaultTextureArray = GSystemTextures.GetDefaultTexture(GraphBuilder, ETextureDimension::Texture2DArray, EPixelFormat::PF_R32_UINT, FClearValueBinding::Transparent);
 		OutStrataUniformParameters.bRoughDiffuse = 0;
 		OutStrataUniformParameters.MaxBytesPerPixel = 0;
-		OutStrataUniformParameters.MaterialTextureArray = GraphBuilder.CreateSRV(FRDGTextureSRVDesc::Create(DefaultTextureArray));
+		OutStrataUniformParameters.MaterialTextureArray = DefaultTextureArray;
 		OutStrataUniformParameters.TopLayerTexture = SystemTextures.DefaultNormal8Bit;
 		OutStrataUniformParameters.SSSTexture = SystemTextures.Black;
 	}
 }
 
+void BindStrataForwardPasslUniformParameters(FRDGBuilder& GraphBuilder, FStrataSceneData* StrataSceneData, FStrataForwardPassUniformParameters& OutStrataUniformParameters)
+{
+	if (IsStrataEnabled() && StrataSceneData)
+	{
+		OutStrataUniformParameters.bRoughDiffuse = StrataSceneData->bRoughDiffuse ? 1u : 0u;
+	}
+	else
+	{
+		OutStrataUniformParameters.bRoughDiffuse = 0;
+	}
+}
+
 TRDGUniformBufferRef<FStrataGlobalUniformParameters> BindStrataGlobalUniformParameters(FStrataSceneData* StrataSceneData)
 {
 	check(StrataSceneData);
diff --git a/Engine/Source/Runtime/Renderer/Private/Strata/Strata.h b/Engine/Source/Runtime/Renderer/Private/Strata/Strata.h
--- a/Engine/Source/Runtime/Renderer/Private/Strata/Strata.h
+++ b/Engine/Source/Runtime/Renderer/Private/Strata/Strata.h
@@ -22,10 +22,14 @@ BEGIN_SHADER_PARAMETER_STRUCT(FStrataBasePassUniformParameters, )
 	SHADER_PARAMETER_RDG_TEXTURE_UAV(RWTexture2D<uint2>, SSSTextureUAV)
 END_SHADER_PARAMETER_STRUCT()
 
+BEGIN_SHADER_PARAMETER_STRUCT(FStrataForwardPassUniformParameters, )
+	SHADER_PARAMETER(uint32, bRoughDiffuse)
+END_SHADER_PARAMETER_STRUCT()
+
 BEGIN_GLOBAL_SHADER_PARAMETER_STRUCT(FStrataGlobalUniformParameters, )
 	SHADER_PARAMETER(uint32, MaxBytesPerPixel)
 	SHADER_PARAMETER(uint32, bRoughDiffuse)
-	SHADER_PARAMETER_RDG_TEXTURE_SRV(Texture2DArray<uint>, MaterialTextureArray)
+	SHADER_PARAMETER_RDG_TEXTURE(Texture2DArray<uint>, MaterialTextureArray)
 	SHADER_PARAMETER_RDG_TEXTURE(Texture2D<uint>, TopLayerTexture)
 	SHADER_PARAMETER_RDG_TEXTURE(Texture2D<uint2>, SSSTexture)
 END_GLOBAL_SHADER_PARAMETER_STRUCT()
@@ -89,6 +93,7 @@ void InitialiseStrataFrameSceneData(FSceneRenderer& SceneRenderer, FRDGBuilder&
 
 void BindStrataBasePassUniformParameters(FRDGBuilder& GraphBuilder, FStrataSceneData* StrataSceneData, FStrataBasePassUniformParameters& OutStrataUniformParameters);
 void BindStrataGlobalUniformParameters(FRDGBuilder& GraphBuilder, FStrataSceneData* StrataSceneData, FStrataGlobalUniformParameters& OutStrataUniformParameters);
+void BindStrataForwardPasslUniformParameters(FRDGBuilder& GraphBuilder, FStrataSceneData* StrataSceneData, FStrataForwardPassUniformParameters& OutStrataUniformParameters);
 
 void ApprendStrataMRTs(FSceneRenderer& SceneRenderer, uint32& BasePassTextureCount, TStaticArray<FTextureRenderTargetBinding, MaxSimultaneousRenderTargets>& BasePassTextures);
 void SetBasePassRenderTargetOutputFormat(const EShaderPlatform Platform, FShaderCompilerEnvironment& OutEnvironment);
diff --git a/Engine/Source/Runtime/Renderer/Private/TranslucentRendering.cpp b/Engine/Source/Runtime/Renderer/Private/TranslucentRendering.cpp
--- a/Engine/Source/Runtime/Renderer/Private/TranslucentRendering.cpp
+++ b/Engine/Source/Runtime/Renderer/Private/TranslucentRendering.cpp
@@ -760,7 +760,7 @@ TRDGUniformBufferRef<FTranslucentBasePassUniformParameters> CreateTranslucentBas
 
 	SetupSharedBasePassParameters(GraphBuilder, View, BasePassParameters.Shared);
 	SetupSceneTextureUniformParameters(GraphBuilder, View.FeatureLevel, SceneTextureSetupMode, BasePassParameters.SceneTextures);
-	Strata::BindStrataGlobalUniformParameters(GraphBuilder, View.StrataSceneData, BasePassParameters.Strata);
+	Strata::BindStrataForwardPasslUniformParameters(GraphBuilder, View.StrataSceneData, BasePassParameters.Strata);
 
 	const FLightSceneProxy* SelectedForwardDirectionalLightProxy = View.ForwardLightingResources ? View.ForwardLightingResources->SelectedForwardDirectionalLightProxy : nullptr;
 	SetupLightCloudTransmittanceParameters(GraphBuilder, Scene, View, SelectedForwardDirectionalLightProxy ? SelectedForwardDirectionalLightProxy->GetLightSceneInfo() : nullptr, BasePassParameters.ForwardDirLightCloudShadow);
