Strata - Fixed wrong SSS data on header

diff --git a/Engine/Shaders/Private/Strata/Strata.ush b/Engine/Shaders/Private/Strata/Strata.ush
--- a/Engine/Shaders/Private/Strata/Strata.ush
+++ b/Engine/Shaders/Private/Strata/Strata.ush
@@ -2003,6 +2003,10 @@ void PackStrataOut(
 						TopLayerData.Roughness += BSDFCoverage * SLAB_ROUGHNESS(BSDF);
 					}
 
+					const bool bSSSMask = !bIsThin;	// SSS and Thin can only be on the bottom layer and Simple can only be for non bottom layer: so ignoring bIsSimpleVolume here.
+					SetHasSubsurface(StrataPixelHeader,			(bSSSMask && BSDF_GETHASSSS(BSDF)));
+					SetHasSubsurfaceProfile(StrataPixelHeader,	(bSSSMask && BSDF_GETHASSSSPROFILE(BSDF)));
+
 					// STRATA_TODO	Should we have the Fuzz influence the next layter visibility accoding to fuzz amount? Should it also occlude diffuse?
 					//				If this is done, we should also update ComputePostProcessBSDFData.
 
@@ -2226,9 +2230,6 @@ void PackStrataOut(
 						const bool bIsThin = BSDF_GETISTHIN(BSDF);
 						if (BSDF_GETHASSSS(BSDF) || bIsSimpleVolume || bIsThin)
 						{
-							SetHasSubsurface(StrataPixelHeader, !bIsThin && BSDF_GETHASSSS(BSDF));
-							SetHasSubsurfaceProfile(StrataPixelHeader, !bIsThin && BSDF_GETHASSSSPROFILE(BSDF));
-
 							if (!bIsThin && BSDF_GETHASSSS(BSDF))
 							{
 								SSSData.BaseColor = SLAB_DIFFUSEALBEDO(BSDF);
