Change thin translucent material conversion to use a single slab to save performance on console.

diff --git a/Engine/Shaders/Private/Strata/StrataLegacyConversion.ush b/Engine/Shaders/Private/Strata/StrataLegacyConversion.ush
--- a/Engine/Shaders/Private/Strata/StrataLegacyConversion.ush
+++ b/Engine/Shaders/Private/Strata/StrataLegacyConversion.ush
@@ -6,6 +6,11 @@
 
 struct FStrataLegacyParameters
 {
+	bool UseMetalness;
+	float3 DiffuseAlbedo;
+	float3 F0;
+	float3 F90;
+
 	float3 BaseColor;
 	float  Specular;
 	float  Metallic;
@@ -28,6 +33,12 @@ FStrataLegacyParameters InitStrataLegacyParameters(uint InSharedLocalBasisIndex,
 	const float3 Zeros = float3(0, 0, 0);
 
 	FStrataLegacyParameters Out;
+
+	Out.UseMetalness = true;
+	Out.DiffuseAlbedo = Zeros;
+	Out.F0 = Zeros;
+	Out.F90 = Zeros;
+
 	Out.BaseColor = Zeros;
 	Out.Specular = 0.5f;
 	Out.Metallic = 0.f;
@@ -52,14 +63,14 @@ FStrataData CreateLegacySlab(FStrataLegacyParameters In, inout uint SharedLocalB
 
 	// Fixed layer structure for helping compiler to unroll and optimize shader
 	return GetStrataSlabBSDF(
-		true,				// UseMetalness
+		In.UseMetalness,	// UseMetalness
 		In.BaseColor,		// BaseColor		-  Metalness workflow
 		Ones,				// EdgeColor		-  Metalness workflow
 		In.Specular,		// Specular			-  Metalness workflow
 		In.Metallic,		// Metallic			-  Metalness workflow
-		Zeros,				// DiffuseAlbedo	- !Metalness workflow
-		Zeros, 				// F0				- !Metalness workflow
-		Zeros,				// F90				- !Metalness workflow
+		In.DiffuseAlbedo,	// DiffuseAlbedo	- !Metalness workflow
+		In.F0,				// F0				- !Metalness workflow
+		In.F90,				// F90				- !Metalness workflow
 		In.Roughness,		// Roughness
 		In.Anisotropy,		// Anisotropy
 		In.SSSProfileID,	// SSSProfileID
@@ -292,6 +303,51 @@ FStrataData StrataConvertLegacyMaterialDynamic(
 	}
 	else if (ShadingModel == SHADINGMODELID_THIN_TRANSLUCENT)	
 	{
+		const float Thickness = DefaultThickness;
+
+		// If the material is dieletric, bias the opacity so that the dieletric part extend further. This better matches the two layer approaches.
+		Opacity = lerp(1.f - Square(1.f - Opacity), Opacity, Metallic);
+
+		// Because we are using a single slab for simulating a 'coverage mix' between two BSDFs (a opaque metallic/conductor and a non-scattering translucent medium), 
+		// we need to do some approximation. In order to reduce avoid hue shift due to difference between 'solid' base-color/diffuse-albedo and 'transparent' transmission 
+		// color, we lerp BaseColor with transmittance color in LAB space so that the interpolated color minized its 'distance' 
+		const float3 BaseColor_LAB			= LinearRGB_2_LAB(BaseColor, XYZ_WHITE_REF_D65);
+		const float3 TransmittanceColor_LAB = LinearRGB_2_LAB(TransmittanceColor, XYZ_WHITE_REF_D65);
+		const float3 Color_LAB			    = lerp(TransmittanceColor_LAB, BaseColor_LAB, Opacity);
+		const float3 Color					= LAB_2_LinearRGB(Color_LAB, XYZ_WHITE_REF_D65);
+
+		// Lerp the transmittance color rather than MFP as it is more perceptually linear
+		const float3 MFP = TransmittanceToMeanFreePath(lerp(Color, Zeros, Opacity), Thickness * CENTIMETER_TO_METER) * METER_TO_CENTIMETER;
+
+		const float F0Dieletrict = DielectricSpecularToF0(Specular);
+
+		const float3 TopF0     = lerp(F0Dieletrict, BaseColor, Metallic);
+		const float3 TopAlbedo = Color * (1.f - Metallic);
+		const float3 TopMFP    = 0;
+
+		const float3 BotF0     = F0Dieletrict;
+		const float3 BotAlbedo = Zeros;
+		const float3 BotMFP    = MFP;
+
+		Top.DiffuseAlbedo	= lerp(BotAlbedo, TopAlbedo, Opacity);
+		Top.SSSDMFP			= lerp(BotMFP,    TopMFP,    Opacity);
+		Top.F0				= lerp(BotF0,     TopF0,     Opacity);
+		Top.F90				= Ones;
+		Top.UseMetalness	= false;
+		Top.Roughness 		= Roughness;
+		Top.Anisotropy 		= 0.f;
+		Top.SSSProfileID 	= 0.f;
+		Top.SSSDMFPScale 	= 1.f;
+		Top.Emissive 		= Emissive;
+		Top.Thickness 		= Thickness;
+		Top.SharedLocalBasisIndex = SharedLocalBasisIndex;
+
+		Top.Weight  = 1.0f;
+		Bot.Weight  = 0.f;
+		FinalWeight = 1.0f;
+
+		// Reference two layers approach
+	#if 0
 		// Top slab BSDF as a simple Disney material
 		Top.BaseColor = BaseColor;
 		Top.Specular = Specular;
@@ -324,6 +380,7 @@ FStrataData StrataConvertLegacyMaterialDynamic(
 		// No need to take into account opacity weighting for translucent material, as the coverage is assumed to be 1
 		// Now weight the top base material by opacity
 		Top.Weight = Opacity;
+	#endif
 	}
 
 	// Fixed layer structure for helping compiler to unroll and optimize shader
