LWCDdx/y respect FORCED_TEXTURE_MIP being 0, fix compile failure if these functions are included in compute shaders

diff --git a/Engine/Shaders/Private/LWCOperations.ush b/Engine/Shaders/Private/LWCOperations.ush
--- a/Engine/Shaders/Private/LWCOperations.ush
+++ b/Engine/Shaders/Private/LWCOperations.ush
@@ -56,7 +56,7 @@ LWC_COMPARE_OP(LWCLessEqual, <=, FLWCType, FFloatType, FBoolType)
 #undef LWC_COMPARE_OP
 
 // No ddxy inside ray tracing shaders
-#if RAYHITGROUPSHADER || RAYMISSHADER
+#if RAYHITGROUPSHADER || RAYMISSHADER || USE_FORCE_TEXTURE_MIP
 FFloatType LWCDdx(FLWCType V) { return (FFloatType)0; }
 FFloatType LWCDdy(FLWCType V) { return (FFloatType)0; }
 #else
