Make sure SSS material are at the bottom of the legacy to starta converted material for SSS to be detected.
Fixed SSS to only be enabled for bottom layer.
Added simple visualisation mode in strata for that. It will be extended later.
This highlight an issue: we do not take into account weight when compute the top and bottom layer.

diff --git a/Engine/Shaders/Private/Strata/Strata.ush b/Engine/Shaders/Private/Strata/Strata.ush
--- a/Engine/Shaders/Private/Strata/Strata.ush
+++ b/Engine/Shaders/Private/Strata/Strata.ush
@@ -2003,7 +2003,7 @@ void PackStrataOut(
 						TopLayerData.Roughness += BSDFCoverage * SLAB_ROUGHNESS(BSDF);
 					}
 
-					const bool bSSSMask = !bIsThin;	// SSS and Thin can only be on the bottom layer and Simple can only be for non bottom layer: so ignoring bIsSimpleVolume here.
+					const bool bSSSMask = !bIsThin && bBottomLayer;	// SSS and Thin can only be on the bottom layer, so ignoring bIsSimpleVolume here.
 					SetHasSubsurface(StrataPixelHeader,			(bSSSMask && BSDF_GETHASSSS(BSDF)));
 					SetHasSubsurfaceProfile(StrataPixelHeader,	(bSSSMask && BSDF_GETHASSSSPROFILE(BSDF)));
 
diff --git a/Engine/Shaders/Private/Strata/StrataLegacyConversion.ush b/Engine/Shaders/Private/Strata/StrataLegacyConversion.ush
--- a/Engine/Shaders/Private/Strata/StrataLegacyConversion.ush
+++ b/Engine/Shaders/Private/Strata/StrataLegacyConversion.ush
@@ -386,7 +386,15 @@ FStrataData StrataConvertLegacyMaterialDynamic(
 	// Fixed layer structure for helping compiler to unroll and optimize shader
 	FStrataData Top_Slab = CreateLegacySlab(Top, SharedLocalBasisTypes);
 	FStrataData Bot_Slab = CreateLegacySlab(Bot, SharedLocalBasisTypes);
-	FStrataData Out = StrataVerticalLayering(StrataWeight(Top_Slab, Top.Weight), StrataWeight(Bot_Slab, Bot.Weight));
+	Top_Slab = StrataWeight(Top_Slab, Top.Weight);
+	FStrataData Out = Top_Slab; 
+	if (Bot.Weight > 0.0f)
+	{
+		// Default BSDF as set to be on top so that their information is stored as part of the Strata top layer data.
+		// When bottom weight is 0, we must set the top BRDF at the bottom to make sure SSS can be identified correctly.
+		Out = StrataVerticalLayering(Out, StrataWeight(Bot_Slab, Bot.Weight));
+	}
+	
 
 	// Take in account the legacy opacity weighting for translucent material
 	#if !MATERIALBLENDING_MASKED && !MATERIALBLENDING_SOLID
diff --git a/Engine/Shaders/Private/Strata/StrataVisualize.usf b/Engine/Shaders/Private/Strata/StrataVisualize.usf
--- a/Engine/Shaders/Private/Strata/StrataVisualize.usf
+++ b/Engine/Shaders/Private/Strata/StrataVisualize.usf
@@ -120,6 +120,21 @@ void VisualizeMaterialPS(
 	return;
 #endif
 
+#if 0	// STRATA_TODO generalise this visualisation to more properties
+	{
+		FStrataAddressing StrataAddressing = GetStrataPixelDataByteOffset(PixelPos, uint2(View.BufferSizeAndInvSize.xy), Strata.MaxBytesPerPixel);
+		const uint PixelStrataDataByteOffsetStart = StrataAddressing.ReadBytes;
+
+		FStrataPixelHeader Header2 = UnpackStrataHeaderIn(Strata.MaterialTextureArray, StrataAddressing, Strata.TopLayerTexture);
+		const uint PixelStrataDataByteOffsetPostHeader = StrataAddressing.ReadBytes;
+
+		if (HasSubsurface(Header2) || HasSubsurfaceProfile(Header2))
+		{
+			OutColor = float4(HasSubsurfaceProfile(Header2) ? 1 : 0, HasSubsurface(Header2) ? 1 : 0, 0, 0.5);
+		}
+	}
+#endif
+
 	OutColor = float4(0.0f, 0.0f, 0.0f, 1.0f);
 
 	// Plot BSDFs
