Add texel-scale ray start dither for VSM directional lights
Helps to hide aliasing due to insufficient shadow resolution, albeit at the cost of requiring an optimal bias (which can occasionally increase peter-panning).
Default on as it is generally a net quality improvement in most case, but can be disabled via cvar.

diff --git a/Engine/Shaders/Private/VirtualShadowMaps/ProjectionDirectional.ush b/Engine/Shaders/Private/VirtualShadowMaps/ProjectionDirectional.ush
--- a/Engine/Shaders/Private/VirtualShadowMaps/ProjectionDirectional.ush
+++ b/Engine/Shaders/Private/VirtualShadowMaps/ProjectionDirectional.ush
@@ -26,16 +26,41 @@ struct FSMRTClipmapRayState
 	uint2 VirtualTexelAddress;
 };
 
+
+float ComputeOptimalSlopeBiasDirectional(
+	FVirtualShadowMapProjectionShaderData ProjectionData,
+	float2 TexelOffset,
+	float3 EstimatedGeoWorldNormal,
+	bool bClamp = true)
+{
+	float4 NormalPlaneUV = mul(float4(EstimatedGeoWorldNormal, 0.0f), ProjectionData.TranslatedWorldToShadowUVNormalMatrix);
+	float2 DepthSlopeUV = -NormalPlaneUV.xy / NormalPlaneUV.z;
+
+	// Clamp to avoid excessive degenerate slope biases causing flickering lit pixels
+	float2 Clamp = 20.0f;
+	DepthSlopeUV = bClamp ? clamp(DepthSlopeUV, -Clamp, Clamp) : DepthSlopeUV;
+		
+	float2 OffsetUV = TexelOffset / float(CalcLevelDimsTexels(0));
+	float OptimalSlopeBias = 2.0f * max(0.0f, dot(DepthSlopeUV, OffsetUV));
+
+	return OptimalSlopeBias;
+}
+
+
 FSMRTClipmapRayState SMRTClipmapRayInitialize(
 	int VirtualShadowMapId,
-	float3 RayStart,
-	float3 RayVector)
+	float3 RayOriginWorld,
+	float3 RayDir,
+	float RayLength,
+	float RayStartOffset,
+	half3 WorldNormal,
+	float2 TexelDitherOffset)
 {
 	const FVirtualShadowMapProjectionShaderData BaseProjectionData = GetVirtualShadowMapProjectionData(VirtualShadowMapId);
 
-	// The page at RayStart has a high chance of being mapped since that is close to the receiver point
+	// The page at RayOriginWorld has a high chance of being mapped since that is close or equal to to the receiver point
 	// Thus we start our clipmap sampling around that page even as we walk through the depth space
-	const float ClipmapLevelFloat = CalcClipmapLevelFloat(BaseProjectionData, RayStart);	
+	const float ClipmapLevelFloat = CalcClipmapLevelFloat(BaseProjectionData, RayOriginWorld);	
 	const int ClipmapIndex = max(0, int(floor(ClipmapLevelFloat)) - BaseProjectionData.ClipmapLevel);
 	const float ClipmapLevelFrac = frac(ClipmapLevelFloat);
 	int ClipmapIndexAdjacent = ClipmapIndex + (ClipmapLevelFrac > 0.5f ? 1 : -1);
@@ -43,6 +68,9 @@ FSMRTClipmapRayState SMRTClipmapRayInitialize(
 
 	const FVirtualShadowMapProjectionShaderData ProjectionData = GetVirtualShadowMapProjectionData(VirtualShadowMapId + ClipmapIndex);
 
+	float3 RayStart = RayOriginWorld + RayDir * RayStartOffset;
+	float3 RayVector = RayDir * RayLength;
+
 	const float3 RayStartUV = mul(float4(RayStart + ProjectionData.ShadowPreViewTranslation, 1.0f), ProjectionData.TranslatedWorldToShadowUVMatrix).xyz;
 	const float3 RayStepUV  = mul(float4(RayVector, 0.0f), ProjectionData.TranslatedWorldToShadowUVMatrix).xyz;
 
@@ -50,6 +78,22 @@ FSMRTClipmapRayState SMRTClipmapRayInitialize(
 	float3 RayStartShadowTexelsZ = float3(LevelDims, LevelDims, 1.0f) * RayStartUV.xyz;
 	float3 RayStepShadowTexelsZ  = float3(LevelDims, LevelDims, 1.0f) * RayStepUV.xyz;
 
+	// Texel dither to hide aliasing
+	// Note that this is directly scaled in texel space of the clipmap level, so it can create visual
+	// discontinuities at clipmap boundaries. We could scale this by distance instead of texel size but
+	// in the case where texels are large enough to be obvious there are already visual discontinuities.
+	if (any(TexelDitherOffset != 0))
+	{
+		// When using a texel bias we have to compute a plane-based offset; with any significant magnification
+		// moving the sampling location like this can produce self-shadowing artifacts.
+		float OptimalBias = ComputeOptimalSlopeBiasDirectional(ProjectionData, TexelDitherOffset, WorldNormal);
+		// Subtract off any portion of the bias that was already covered by the ray start offset (usually contact shadows)
+		OptimalBias = max(0.0f, OptimalBias - abs(RayStartOffset * ProjectionData.ShadowViewToClipMatrix._33));
+		
+		RayStartShadowTexelsZ.xy += TexelDitherOffset;
+		RayStartShadowTexelsZ.z += OptimalBias;
+	}
+
 	// Work out scale/bias to convert to the adjacent UVs from the regular UVs
 	const FVirtualShadowMapProjectionShaderData ProjectionDataAdjacent = GetVirtualShadowMapProjectionData(VirtualShadowMapId + ClipmapIndexAdjacent);
 
@@ -163,7 +207,8 @@ FVirtualShadowMapSampleResult TraceDirectional(
 	float RayLengthScale,
 	const half Noise,
 	bool bCullBackface,
-	half3 WorldNormal)
+	half3 WorldNormal,
+	float TexelDitherScale)
 {
 	float3 WorldPosition = TranslatedWorldPosition - LWCHackToFloat(PrimaryView.PreViewTranslation);
 	float3 ViewPosition = mul(float4(TranslatedWorldPosition, 1.0f), View.TranslatedWorldToView).xyz;
@@ -188,16 +233,17 @@ FVirtualShadowMapSampleResult TraceDirectional(
 		uint i = 0;
 		for ( ; i < MaxRayCount; i++)
 		{
+			// One sample for ray, one for texel dither
 			RandomSequence RandSequence;
-			RandomSequence_Initialize(RandSequence, PixelPos, i, View.StateFrameIndexMod8, MaxRayCount);
-			float2 E = RandomSequence_GenerateSample2D(RandSequence);
+			RandomSequence_Initialize(RandSequence, PixelPos, i, View.StateFrameIndex, MaxRayCount);
+			float4 RandSample = RandomSequence_GenerateSample4D(RandSequence);
+			float2 E = RandSample.xy;
 
 			float3 RayDir = GetRandomDirectionalLightRayDir(Light, E);
-			// TODO, move the clipmap selection outside/before this offset
-			float3 RayStart = WorldPosition + RayDir * RayStartOffset;
-			float3 RayVector = RayDir * RayLength;
 
-			FSMRTClipmapRayState RayState = SMRTClipmapRayInitialize(VirtualShadowMapId, RayStart, RayVector);
+			float2 TexelDitherOffset = (RandSample.zw - 0.5f) * TexelDitherScale;
+
+			FSMRTClipmapRayState RayState = SMRTClipmapRayInitialize(VirtualShadowMapId, WorldPosition, RayDir, RayLength, RayStartOffset, WorldNormal, TexelDitherOffset);
 			FSMRTResult SMRTResult = SMRTRayCast(RayState, SamplesPerRay, StepOffset);
 
 			RayMissCount += SMRTResult.bValidHit ? 0 : 1;
diff --git a/Engine/Shaders/Private/VirtualShadowMaps/VirtualShadowMapProjection.usf b/Engine/Shaders/Private/VirtualShadowMaps/VirtualShadowMapProjection.usf
--- a/Engine/Shaders/Private/VirtualShadowMaps/VirtualShadowMapProjection.usf
+++ b/Engine/Shaders/Private/VirtualShadowMaps/VirtualShadowMapProjection.usf
@@ -138,6 +138,7 @@ int SMRTRayCount;					// 0 = off
 int SMRTSamplesPerRay;
 float SMRTRayLengthScale;			// Directional lights
 float SMRTCotMaxRayAngleFromLight;	// Spot/point lights
+float SMRTTexelDitherScale;			// Currently only directional lights
 
 // One pass projection
 RWTexture2D< uint > RWShadowMaskBits;
@@ -220,7 +221,8 @@ FVirtualShadowMapSampleResult ProjectLight(
 				SMRTRayLengthScale,
 				Noise,
 				bBackfaceCull,
-				WorldNormal);
+				WorldNormal,
+				SMRTTexelDitherScale);
 		#else
 			Result = TraceLocalLight(
 				VirtualShadowMapId,
diff --git a/Engine/Source/Runtime/Renderer/Private/VirtualShadowMaps/VirtualShadowMapProjection.cpp b/Engine/Source/Runtime/Renderer/Private/VirtualShadowMaps/VirtualShadowMapProjection.cpp
--- a/Engine/Source/Runtime/Renderer/Private/VirtualShadowMaps/VirtualShadowMapProjection.cpp
+++ b/Engine/Source/Runtime/Renderer/Private/VirtualShadowMaps/VirtualShadowMapProjection.cpp
@@ -100,6 +100,14 @@ static TAutoConsoleVariable<int32> CVarSMRTAdaptiveRayCount(
 	ECVF_RenderThreadSafe
 );
 
+static TAutoConsoleVariable<float> CVarSMRTTexelDitherScale(
+	TEXT( "r.Shadow.Virtual.SMRT.TexelDitherScale" ),
+	2.0f,
+	TEXT( "Applies a dither to the shadow map ray casts to help hide aliasing due to insufficient shadow resolution.\n" )
+	TEXT( "This is usually desirable, but it can occasionally cause shadows from thin geometry to separate from their casters at shallow light angles." ),
+	ECVF_RenderThreadSafe
+);
+
 
 // Composite denoised shadow projection mask onto the light's shadow mask
 // Basically just a copy shader with a special blend mode
@@ -192,6 +200,7 @@ class FVirtualShadowMapProjectionCS : public FGlobalShader
 		SHADER_PARAMETER(uint32, SMRTSamplesPerRay)
 		SHADER_PARAMETER(float, SMRTRayLengthScale)
 		SHADER_PARAMETER(float, SMRTCotMaxRayAngleFromLight)
+		SHADER_PARAMETER(float, SMRTTexelDitherScale)
 		SHADER_PARAMETER(uint32, InputType)
 		SHADER_PARAMETER(uint32, bCullBackfacingPixels)
 		// One pass projection parameters
@@ -271,6 +280,7 @@ static void RenderVirtualShadowMapProjectionCommon(
 	PassParameters->NormalBias = GetNormalBiasForShader();
 	PassParameters->InputType = uint32(InputType);
 	PassParameters->bCullBackfacingPixels = VirtualShadowMapArray.ShouldCullBackfacingPixels() ? 1 : 0;
+	PassParameters->SMRTTexelDitherScale = CVarSMRTTexelDitherScale.GetValueOnRenderThread();
 	if (bHasHairStrandsData)
 	{
 		PassParameters->HairStrands = HairStrands::BindHairStrandsViewUniformParameters(View);
