Fix r.RayTracing.Shadows.EnableTwoSidedGeometry behavior. The logic was incorrectly inverted in CL 17418008.

diff --git a/Engine/Shaders/Private/RayTracing/RayTracingOcclusionRGS.usf b/Engine/Shaders/Private/RayTracing/RayTracingOcclusionRGS.usf
--- a/Engine/Shaders/Private/RayTracing/RayTracingOcclusionRGS.usf
+++ b/Engine/Shaders/Private/RayTracing/RayTracingOcclusionRGS.usf
@@ -422,7 +422,8 @@ FOcclusionResult ComputeOcclusion(
 
 		uint RayFlags = 0;
 
-		if (bTwoSidedGeometry != 0)
+		// Enable back face culling if not used two-sided shadow casting mode.
+		if (bTwoSidedGeometry != 1)
 		{
 			RayFlags |= RAY_FLAG_CULL_BACK_FACING_TRIANGLES;
 		}
