Fix invalid velocity write when hair coverage threshold is set.

diff --git a/Engine/Shaders/Private/HairStrands/HairStrandsVelocity.usf b/Engine/Shaders/Private/HairStrands/HairStrandsVelocity.usf
--- a/Engine/Shaders/Private/HairStrands/HairStrandsVelocity.usf
+++ b/Engine/Shaders/Private/HairStrands/HairStrandsVelocity.usf
@@ -69,7 +69,7 @@ void MainCS(uint3 DispatchThreadId : SV_DispatchThreadID, uint2 InGroupId : SV_G
 
 	const FNodeDesc NodeDesc = DecodeNodeDesc(NodeIndex.Load(uint3(PixelCoord, 0)));
 	bool bNeedFastResolve = false;
-	bool bValidVelocity = NodeDesc.Count > 0;
+	bool bValidVelocity = false;
 	float2 OutEncodedVelocity = 0;
 	if (NodeDesc.Count > 0)
 	{
@@ -91,6 +91,7 @@ void MainCS(uint3 DispatchThreadId : SV_DispatchThreadID, uint2 InGroupId : SV_G
 		const float HairPixelCoverage = CoverageTexture.Load(uint3(PixelCoord,0));
 		if (HairPixelCoverage >= CoverageThreshold)
 		{
+			bValidVelocity = true;
 			for (uint NodeIt = 0; NodeIt < NodeDesc.Count; ++NodeIt)
 			{
 				const uint SampleIndex = NodeDesc.Offset + NodeIt;
@@ -114,7 +115,7 @@ void MainCS(uint3 DispatchThreadId : SV_DispatchThreadID, uint2 InGroupId : SV_G
 				const float CurrentVelocityMagnitude2 = dot(CurrentVelocity, CurrentVelocity);
 				if (CurrentVelocityMagnitude2 > MaxVelocityMagnitude2)
 				{
-				MaxEncodedVelocity = EncodedVelocity;
+					MaxEncodedVelocity = EncodedVelocity;
 					MaxVelocityMagnitude2 = CurrentVelocityMagnitude2;
 				}
 				#endif			
