Optimises TSR on some platforms

diff --git a/Engine/Shaders/Private/TemporalSuperResolution/TSRCommon.ush b/Engine/Shaders/Private/TemporalSuperResolution/TSRCommon.ush
--- a/Engine/Shaders/Private/TemporalSuperResolution/TSRCommon.ush
+++ b/Engine/Shaders/Private/TemporalSuperResolution/TSRCommon.ush
@@ -9,6 +9,10 @@
 #include "/Engine/Public/DualPixelVectorization.ush"
 
 
+//------------------------------------------------------- RECOMPILE HASH
+// 4E268118-A455-494C-BA81-DCFF1C9528CB
+
+
 //------------------------------------------------------- COMPILER CONFIG
 
 // Generate vector truncation warnings to errors.
@@ -543,12 +547,8 @@ taa_short2 InvalidateOutputPixelPos(taa_short2 PixelPos, uint2 ViewportMax)
 
 taa_ushort2 Map8x8Tile2x2Lane(uint GroupThreadIndex)
 {
-	taa_ushort T = taa_ushort(GroupThreadIndex);
-
-	taa_ushort2 GroupId;
-	GroupId.x = ((T >> taa_ushort(0)) & taa_ushort(0x01)) | ((T >> taa_ushort(2 - 1)) & taa_ushort(0x03 << 1));
-	GroupId.y = ((T >> taa_ushort(1)) & taa_ushort(0x01)) | ((T >> taa_ushort(4 - 1)) & taa_ushort(0x03 << 1));
-
+	taa_ushort2 T = taa_ushort(GroupThreadIndex).xx;
+	taa_ushort2 GroupId = (T >> taa_ushort2(0, 1) & taa_ushort2(0x01, 0x01)) | ((T >> taa_ushort2(2 - 1, 4 - 1)) & taa_ushort2(0x03 << 1, 0x03 << 1));
 	return GroupId;
 }
 
@@ -591,12 +591,17 @@ half2 SafeRcp(half2 x)
 
 taa_half3x2 RGBToYCoCg(taa_half3x2 RGB)
 {
-	taa_half2 Y  = mul(taa_half3(1, 2, 1), RGB);
-	taa_half2 Co = mul(taa_half3(2, 0, -2), RGB);
-	taa_half2 Cg = mul(taa_half3(-1, 2, -1), RGB);
+    taa_half2 R = RGB[0];
+    taa_half2 G = RGB[1];
+    taa_half2 B = RGB[2];       
+
+    taa_half2 RplusB = R + B;
+    taa_half2 Y = 2.0 * G + RplusB;
+    taa_half2 Cg = 2.0 * G - RplusB;
+    taa_half2 Co = 2.0 * (R - B);
+    taa_half3x2 YCoCg = taa_half3x2(Y, Co, Cg);
 
-	taa_half3x2 YCoCg = taa_half3x2(Y, Co, Cg);
-	return YCoCg;
+    return YCoCg;
 }
 
 taa_half3x2 YCoCgToRGB(taa_half3x2 YCoCg)
@@ -774,7 +779,7 @@ taa_half2 MeasureRejectionFactor(
 		BoxSize * taa_half(0.25),
 		taa_half3x2(BackbufferQuantizationError, BackbufferQuantizationError, BackbufferQuantizationError));
 
-	taa_half3x2 Factor = taa_half(1.0) - saturate(max(ClampedEnergy - ClampError, taa_half3x2(0.0, 0.0, 0.0, 0.0, 0.0, 0.0)) / max(Delta, BoxSize));
+	taa_half3x2 Factor = taa_half(1.0) - saturate(max(ClampedEnergy - ClampError, taa_half3x2(0.0, 0.0, 0.0, 0.0, 0.0, 0.0)) * rcp(max(Delta, BoxSize)));
 
 	taa_half2 Rejection = min3(Factor[0], Factor[1], Factor[2]);
 	
diff --git a/Engine/Shaders/Private/TemporalSuperResolution/TSRUpdateHistory.usf b/Engine/Shaders/Private/TemporalSuperResolution/TSRUpdateHistory.usf
--- a/Engine/Shaders/Private/TemporalSuperResolution/TSRUpdateHistory.usf
+++ b/Engine/Shaders/Private/TemporalSuperResolution/TSRUpdateHistory.usf
@@ -3,7 +3,6 @@
 #include "TSRCommon.ush"
 
 
-
 //------------------------------------------------------- CONFIG
 
 #define DEBUG_FALLBACK_BLENDING 0
@@ -41,6 +40,10 @@
 	#define CONFIG_ENABLE_DUAL_PIXEL_VECTORIZATION 0
 #endif
 
+#ifndef PLATFORM_SPECIFIC_ISOLATE
+	#define PLATFORM_SPECIFIC_ISOLATE ISOLATE
+#endif
+
 
 //------------------------------------------------------- CONSTANTS
 
@@ -90,6 +93,8 @@ RWTexture2D<taa_half3> SceneColorOutput;
 
 //------------------------------------------------------- LDS
 
+#if CONFIG_MANUAL_LDS_SPILL
+
 #if CONFIG_ENABLE_DUAL_PIXEL_VECTORIZATION
 #define GROUP_SHARED_SIZE (TILE_SIZE * TILE_SIZE / 2)
 #else
@@ -112,6 +117,8 @@ groupshared taa_half4 SharedArray11[GROUP_SHARED_SIZE];
 
 groupshared taa_subpixel_payload2 SharedParallaxFactorBits[GROUP_SHARED_SIZE];
 
+#endif // CONFIG_MANUAL_LDS_SPILL
+
 
 //------------------------------------------------------- FUNCTIONS
 
@@ -241,37 +248,99 @@ void MainCS(
 	taa_subpixel_payload2 ParallaxFactorBits = taa_subpixel_payload(0x0).xx;
 	bool2 bIsResponsiveAAPixel = false;
 
-	ISOLATE
+	PLATFORM_SPECIFIC_ISOLATE
 	{
-		float2x2 EncodedVelocity = dpv_interleave_registers(
-			DilatedVelocityTexture[dpv_lo(InputPixelPos)],
-			DilatedVelocityTexture[dpv_hi(InputPixelPos)]);
-		ParallaxRejectionMask = dpv_interleave_registers(
-			ParallaxRejectionMaskTexture[dpv_lo(InputPixelPos)],
-			ParallaxRejectionMaskTexture[dpv_hi(InputPixelPos)]);
-		ParallaxFactorBits = dpv_interleave_registers(
-			ParallaxFactorTexture[dpv_lo(InputPixelPos)],
-			ParallaxFactorTexture[dpv_hi(InputPixelPos)]) & taa_subpixel_payload(SUB_PIXEL_PARALLAX_FACTOR_BIT_MASK);
-
-		float2x2 ScreenVelocity = dpv_interleave_registers(
-			DecodeVelocityFromTexture(float4(dpv_lo(EncodedVelocity), 0.0, 0.0)).xy,
-			DecodeVelocityFromTexture(float4(dpv_hi(EncodedVelocity), 0.0, 0.0)).xy);
-
-		PrevScreenPos = ScreenPos - ScreenVelocity;
-		OutputPixelVelocity = taa_half2(dpv_length(dpv_mul(ScreenVelocity, HistoryInfo_ViewportSize)));
-
-		taa_ushort2x2 RejectionPixelPos = dpv_sub(taa_ushort2x2(InputPixelPos), taa_ushort2(InputPixelPosMin)) / 2;
-		LowFrequencyRejection = dpv_interleave_registers(
-			HistoryRejectionTexture[dpv_lo(RejectionPixelPos)],
-			HistoryRejectionTexture[dpv_hi(RejectionPixelPos)]);
-		TranslucencyRejection = dpv_interleave_registers(
-			TranslucencyRejectionTexture[dpv_lo(InputPixelPos)],
-			TranslucencyRejectionTexture[dpv_hi(InputPixelPos)]);
+		float2 RawEncodedVelocity[DPV_PIXEL_PER_LANE];
+		half RawParallaxRejectionMask[DPV_PIXEL_PER_LANE];
+		uint RawParallaxFactorBits[DPV_PIXEL_PER_LANE];
+		half RawLowFrequencyRejection[DPV_PIXEL_PER_LANE];
+		half RawTranslucencyRejection[DPV_PIXEL_PER_LANE];
+		half RawTranslucencyRejectionNeighborhood[4][DPV_PIXEL_PER_LANE];
+		half RawMoireWeightNeighborhood[4][DPV_PIXEL_PER_LANE]; // TODO(TSR): bug
+		
+		uint RawSceneStencilRef[DPV_PIXEL_PER_LANE];
+		uint RawSceneStencilRefNeighborhood[4][DPV_PIXEL_PER_LANE];
 		
-		uint2 SceneStencilRef = dpv_interleave_registers(
-			InputSceneStencilTexture.Load(int3(dpv_lo(InputPixelPos), 0)) STENCIL_COMPONENT_SWIZZLE, 
-			InputSceneStencilTexture.Load(int3(dpv_hi(InputPixelPos), 0)) STENCIL_COMPONENT_SWIZZLE);
+		// Issue overlapped texture fetches
+		PLATFORM_SPECIFIC_ISOLATE
+		{
+			taa_ushort2x2 RejectionPixelPos = dpv_sub(taa_ushort2x2(InputPixelPos), taa_ushort2(InputPixelPosMin)) / 2;
+
+			UNROLL_N(2)
+			for (uint PixelId = 0; PixelId < DPV_PIXEL_PER_LANE; PixelId++)
+			{
+				RawEncodedVelocity[PixelId]       = DilatedVelocityTexture[dpv_access_pixel(InputPixelPos, PixelId)];
+				RawParallaxRejectionMask[PixelId] = ParallaxRejectionMaskTexture[dpv_access_pixel(InputPixelPos, PixelId)];
+				RawParallaxFactorBits[PixelId]    = ParallaxFactorTexture[dpv_access_pixel(InputPixelPos, PixelId)];
+				RawTranslucencyRejection[PixelId] = TranslucencyRejectionTexture[dpv_access_pixel(InputPixelPos, PixelId)];
+				RawSceneStencilRef[PixelId]       = InputSceneStencilTexture.Load(int3(dpv_access_pixel(InputPixelPos, PixelId), 0)) STENCIL_COMPONENT_SWIZZLE;
+				
+				RawLowFrequencyRejection[PixelId] = HistoryRejectionTexture[dpv_access_pixel(RejectionPixelPos, PixelId)];
+			}
+
+			// Vector in pixel between pixel K -> O.
+			taa_half2x2 dKO = taa_half2x2(PPCo - PPCk);
 		
+			UNROLL_N(CONFIG_SAMPLES - 1)
+			for (uint NeighborId = 0; NeighborId < (CONFIG_SAMPLES - 1); NeighborId++)
+			{
+				taa_short2x2 SampleInputPixelPos;
+				taa_half2x2 PixelOffset;
+				ComputeInputKernelSamplePosition(
+					InputPixelPos, dKO, /* SampleId = */ NeighborId + 1,
+					/* out */ SampleInputPixelPos,
+					/* out */ PixelOffset);
+			
+				UNROLL_N(2)
+				for (uint PixelId = 0; PixelId < DPV_PIXEL_PER_LANE; PixelId++)
+				{
+					RawTranslucencyRejectionNeighborhood[NeighborId][PixelId] = TranslucencyRejectionTexture[dpv_access_pixel(SampleInputPixelPos, PixelId)];
+					RawSceneStencilRefNeighborhood[NeighborId][PixelId] = InputSceneStencilTexture.Load(int3(dpv_access_pixel(SampleInputPixelPos, PixelId), 0)) STENCIL_COMPONENT_SWIZZLE;
+			
+					#if DIM_MOIRE_HEURISTIC
+						RawMoireWeightNeighborhood[NeighborId][PixelId] = InterferenceWeightTexture[dpv_lo(SampleInputPixelPos)];
+					#endif
+				}
+			} // for (uint NeighborId = 0; NeighborId < (CONFIG_SAMPLES - 1); NeighborId++)
+		}
+		
+		uint2 SceneStencilRef;
+
+		// Process texture fetches.
+		PLATFORM_SPECIFIC_ISOLATE
+		{
+			// Process velocity
+			{
+				float2x2 ScreenVelocity = dpv_interleave_registers(
+					DecodeVelocityFromTexture(float4(RawEncodedVelocity[0], 0.0, 0.0)).xy,
+					DecodeVelocityFromTexture(float4(RawEncodedVelocity[1], 0.0, 0.0)).xy);
+
+				PrevScreenPos = ScreenPos - ScreenVelocity;
+				OutputPixelVelocity = taa_half2(dpv_length(dpv_mul(ScreenVelocity, HistoryInfo_ViewportSize)));
+			}
+
+			// Process input texel, forcing to pack register.
+			ParallaxRejectionMask = dpv_force_interleave_registers_array(RawParallaxRejectionMask);
+			ParallaxFactorBits    = dpv_force_interleave_registers_array(RawParallaxFactorBits);
+			LowFrequencyRejection = dpv_force_interleave_registers_array(RawLowFrequencyRejection);
+			TranslucencyRejection = dpv_force_interleave_registers_array(RawTranslucencyRejection);
+			SceneStencilRef       = dpv_force_interleave_registers_array(RawSceneStencilRef);
+			
+			// Process neighborhood
+			UNROLL_N(CONFIG_SAMPLES - 1)
+			for (uint NeighborId = 0; NeighborId < (CONFIG_SAMPLES - 1); NeighborId++)
+			{
+				TranslucencyRejection = min(TranslucencyRejection, dpv_interleave_registers_array(RawTranslucencyRejectionNeighborhood[NeighborId]));
+			
+				SceneStencilRef |= dpv_interleave_registers_array(RawSceneStencilRefNeighborhood[NeighborId]);
+			
+				#if DIM_MOIRE_HEURISTIC
+					MoireWeight = min(MoireWeight, dpv_interleave_registers_array(RawMoireWeightNeighborhood[NeighborId]));
+				#endif
+			}
+		}
+
+		// TODO(TSR): Integrate in the manual overlapped texture fetches.
 		#if CONFIG_REJECTION_ANTI_ALIASING
 		{
 			// Detect if HistoryBufferUV would be outside of the viewport.
@@ -302,55 +371,30 @@ void MainCS(
 			}
 		}
 		#endif
-
-		// Vector in pixel between pixel K -> O.
-		taa_half2x2 dKO = taa_half2x2(PPCo - PPCk);
-
-		UNROLL_N(CONFIG_SAMPLES - 1)
-		for (uint SampleId = 1; SampleId < CONFIG_SAMPLES; SampleId++)
-		{
-			taa_short2x2 SampleInputPixelPos;
-			taa_half2x2 PixelOffset;
-			ComputeInputKernelSamplePosition(
-				InputPixelPos, dKO, SampleId,
-				/* out */ SampleInputPixelPos,
-				/* out */ PixelOffset);
-			
-			TranslucencyRejection = min(TranslucencyRejection, dpv_interleave_registers(
-				TranslucencyRejectionTexture[dpv_lo(SampleInputPixelPos)],
-				TranslucencyRejectionTexture[dpv_hi(SampleInputPixelPos)]));
-			
-			SceneStencilRef |= dpv_interleave_registers(
-				InputSceneStencilTexture.Load(int3(dpv_lo(SampleInputPixelPos), 0)) STENCIL_COMPONENT_SWIZZLE, 
-				InputSceneStencilTexture.Load(int3(dpv_hi(SampleInputPixelPos), 0)) STENCIL_COMPONENT_SWIZZLE);
-			
-			#if DIM_MOIRE_HEURISTIC
-				MoireWeight = min(MoireWeight, dpv_interleave_registers(
-					InterferenceWeightTexture[dpv_lo(SampleInputPixelPos)],
-					InterferenceWeightTexture[dpv_hi(SampleInputPixelPos)]));
-			#endif
-		}
-
-		TranslucencyRejection = max(TranslucencyRejection, taa_half(MinTranslucencyRejection));
 		
-		// Fetch wehther the pixel is responsive AA or not.
-		#if CONFIG_RESPONSIVE_STENCIL
+		// Final post processing.
 		{
-			bIsResponsiveAAPixel = (SceneStencilRef & ResponsiveStencilMask) != 0;
+			TranslucencyRejection = max(TranslucencyRejection, taa_half(MinTranslucencyRejection));
+		
+			// Fetch wehther the pixel is responsive AA or not.
+			#if CONFIG_RESPONSIVE_STENCIL
+			{
+				bIsResponsiveAAPixel = (SceneStencilRef & ResponsiveStencilMask) != 0;
 
-			#if CONFIG_REJECT_TRAMSLUCENCY_ON_RESPONSIVE_AA_ONLY
-				TranslucencyRejection = max(TranslucencyRejection, bIsResponsiveAAPixel ? taa_half(0.0) : taa_half(1.0));
+				#if CONFIG_REJECT_TRAMSLUCENCY_ON_RESPONSIVE_AA_ONLY
+					TranslucencyRejection = max(TranslucencyRejection, bIsResponsiveAAPixel ? taa_half(0.0) : taa_half(1.0));
+				#endif
+			}
 			#endif
-		}
-		#endif
 
-		#if !CONFIG_CLAMP
-		{
-			ParallaxRejectionMask = taa_half(1.0).xx;
-			LowFrequencyRejection = taa_half(1.0).xx;
-			TranslucencyRejection = taa_half(1.0).xx;
+			#if !CONFIG_CLAMP
+			{
+				ParallaxRejectionMask = taa_half(1.0).xx;
+				LowFrequencyRejection = taa_half(1.0).xx;
+				TranslucencyRejection = taa_half(1.0).xx;
+			}
+			#endif
 		}
-		#endif
 	}
 
 	// Reject the translucency to respond quicker on firefliers.
@@ -386,11 +430,12 @@ void MainCS(
 	bool2 bUpdateSubpixelDetails;
 	bool2x2 bKOSign;
 
-	ISOLATE
+	PLATFORM_SPECIFIC_ISOLATE
 	{
 		// Vector in pixel between pixel K -> O.
 		taa_half2x2 dKO = taa_half2x2(PPCo - PPCk);
 
+		// Compresses the sign of dKO into SGPR to keep VGPR pressure low.
 		{
 			float SizeOfOutputPixelInInputPixelWidth = 0.5 * (InputInfo_ViewportSize.x * HistoryInfo_ViewportSizeInverse.x);
 
@@ -400,53 +445,75 @@ void MainCS(
 			bKOSign[0] = dKO[0] < 0.0;
 			bKOSign[1] = dKO[1] < 0.0;
 		}
+		
+		taa_half3 RawInputColorArray[CONFIG_SAMPLES][DPV_PIXEL_PER_LANE];
+		taa_half4 RawTranslucencyColorArray[CONFIG_SAMPLES][DPV_PIXEL_PER_LANE];
 
+		taa_half2x2 RawdPPArray[CONFIG_SAMPLES];
+		taa_half2   RawSampleSpatialWeightArray[CONFIG_SAMPLES];
+		
+		// Issues overlapped texture fetches
+		PLATFORM_SPECIFIC_ISOLATE
+		{
+			// Soften the spatial kernel when there is high noise in the input to make parrallax disocclusion of detailed geometry less distracting.
+			taa_half2 OffScreenInputToHistoryFactor = taa_half(1.0) - taa_half(0.5) * NoiseFiltering;
+
+			taa_half2 KernelInputToHistoryFactor = taa_half2(
+				bOffScreen.x ? OffScreenInputToHistoryFactor.x : taa_half(InputToHistoryFactor),
+				bOffScreen.y ? OffScreenInputToHistoryFactor.y : taa_half(InputToHistoryFactor));
+
+			UNROLL_N(CONFIG_SAMPLES)
+			for (uint SampleId = 0; SampleId < CONFIG_SAMPLES; SampleId++)
+			{
+				taa_short2x2 SampleInputPixelPos;
+				taa_half2x2 PixelOffset;
+				ComputeInputKernelSamplePosition(
+					InputPixelPos, dKO, SampleId,
+					/* out */ SampleInputPixelPos,
+					/* out */ PixelOffset);
+				
+				taa_half2x2 dPP = PixelOffset - dKO;
+				taa_half2 SampleSpatialWeight = ComputeSampleWeigth(KernelInputToHistoryFactor, dPP, /* MinimalContribution = */ float(0.005));
+				
+				RawdPPArray[SampleId] = dPP; 
+				RawSampleSpatialWeightArray[SampleId] = SampleSpatialWeight;
 
+				UNROLL_N(2)
+				for (uint PixelId = 0; PixelId < DPV_PIXEL_PER_LANE; PixelId++)
+				{
+					RawInputColorArray[SampleId][PixelId] = InputSceneColorTexture[dpv_access_pixel(SampleInputPixelPos, PixelId)];
+					#if CONFIG_ACCUMULATE_TRANSLUCENCY_SEPARATLY
+						RawTranslucencyColorArray[SampleId][PixelId] = InputSceneTranslucencyTexture[dpv_access_pixel(SampleInputPixelPos, PixelId)];
+					#endif
+				}
+			} // for (uint SampleId = 0; SampleId < CONFIG_SAMPLES; SampleId++)
+		}
+			
+		taa_half2 FilteredInputColorWeight = taa_half(0.0);
+		
 		FilteredInputColor = taa_half(0.0);
 		InputMinColor = dpv_interleave_mono_registers(+LargestSceneColorYCoCg);
 		InputMaxColor = dpv_interleave_mono_registers(-LargestSceneColorYCoCg);
 		TranslucencyMinColor = dpv_interleave_mono_registers(+LargestSceneColorRGBA);
 		TranslucencyMaxColor = dpv_interleave_mono_registers(-LargestSceneColorRGBA);
 
-		taa_half2 FilteredInputColorWeight = taa_half(0.0);
-		
-		// Soften the spatial kernel when there is high noise in the input to make parrallax disocclusion of detailed geometry less distracting.
-		taa_half2 OffScreenInputToHistoryFactor = taa_half(1.0) - taa_half(0.5) * NoiseFiltering;
-
-		taa_half2 KernelInputToHistoryFactor = taa_half2(
-			bOffScreen.x ? OffScreenInputToHistoryFactor.x : taa_half(InputToHistoryFactor),
-			bOffScreen.y ? OffScreenInputToHistoryFactor.y : taa_half(InputToHistoryFactor));
-
 		UNROLL_N(CONFIG_SAMPLES)
 		for (uint SampleId = 0; SampleId < CONFIG_SAMPLES; SampleId++)
 		{
-			taa_short2x2 SampleInputPixelPos;
-			taa_half2x2 PixelOffset;
-			ComputeInputKernelSamplePosition(
-				InputPixelPos, dKO, SampleId,
-				/* out */ SampleInputPixelPos,
-				/* out */ PixelOffset);
-
-			taa_half3x2 InputColor = dpv_interleave_registers(
-				InputSceneColorTexture[dpv_lo(SampleInputPixelPos)],
-				InputSceneColorTexture[dpv_hi(SampleInputPixelPos)]);
-			
-			taa_half4x2 TranslucencyColor = dpv_interleave_registers(
-				InputSceneTranslucencyTexture[dpv_lo(SampleInputPixelPos)],
-				InputSceneTranslucencyTexture[dpv_hi(SampleInputPixelPos)]);
+			taa_half2 SampleSpatialWeight = RawSampleSpatialWeightArray[SampleId];
+			taa_half3x2 InputColor = dpv_force_interleave_registers_array(RawInputColorArray[SampleId]);
+			taa_half4x2 TranslucencyColor = dpv_force_interleave_registers_array(RawTranslucencyColorArray[SampleId]);
 
 			#if CONFIG_SCENE_COLOR_OVERFLOW
 				InputColor = min(InputColor, dpv_interleave_mono_registers(LargestSceneColorRGB));
 			#endif
 
-			taa_half2x2 dPP = PixelOffset - dKO;
-			taa_half2 SampleSpatialWeight = ComputeSampleWeigth(KernelInputToHistoryFactor, dPP, /* MinimalContribution = */ float(0.005));
-
 			taa_half2 ToneWeight = HdrWeight4(InputColor);
 
 			#if CONFIG_VFX_SURVIVAL
 			if (SampleId == 0)
 			{
+				taa_half2x2 dPP =  RawdPPArray[SampleId];
 				ToneWeight = max(ToneWeight, ToneWeight > 0.2 ? 0.0 : ComputeSampleWeigth(1.0, dPP, /* MinimalContribution = */ 0.0));
 			}
 			#endif
@@ -456,14 +523,23 @@ void MainCS(
 
 			if (SampleId == 0)
 			{
-				ClosestInputLuma4 = Luma4(InputColor);
 				FilteredTranslucencyColor = TranslucencyColor;
 			}
 
 			InputMinColor = min(InputMinColor, TransformColorForClampingBox(InputColor));
 			InputMaxColor = max(InputMaxColor, TransformColorForClampingBox(InputColor));
-			TranslucencyMinColor = min(TranslucencyMinColor, TranslucencyColor);
-			TranslucencyMaxColor = max(TranslucencyMaxColor, TranslucencyColor);
+
+			#if CONFIG_ACCUMULATE_TRANSLUCENCY_SEPARATLY
+			{
+				TranslucencyMinColor = min(TranslucencyMinColor, TranslucencyColor);
+				TranslucencyMaxColor = max(TranslucencyMaxColor, TranslucencyColor);
+				
+				if (SampleId == 0)
+				{
+					FilteredTranslucencyColor = TranslucencyColor;
+				}
+			}
+			#endif
 		}
 		
 		FilteredInputColor = dpv_scale(FilteredInputColor, rcp(FilteredInputColorWeight));
@@ -481,7 +557,7 @@ void MainCS(
 		
 	// Spills to LDS to make room in VGPR for history sampling.
 	#if CONFIG_MANUAL_LDS_SPILL
-	ISOLATE
+	PLATFORM_SPECIFIC_ISOLATE
 	{
 		uint LocalGroupThreadIndex = GetGroupThreadIndex(GroupThreadIndex, GroupWaveIndex);
 
@@ -525,7 +601,7 @@ void MainCS(
 	
 	taa_subpixel_details2 PrevSubpixelDetails;
 
-	ISOLATE
+	PLATFORM_SPECIFIC_ISOLATE
 	{
 		// Reproject history
 		taa_half3x2 AccumulateLowFrequency = taa_half(0);
@@ -542,47 +618,80 @@ void MainCS(
 			PrevHistoryBufferUV[0] = clamp(PrevHistoryBufferUV[0], PrevHistoryInfo_UVViewportBilinearMin[0], PrevHistoryInfo_UVViewportBilinearMax[0]);
 			PrevHistoryBufferUV[1] = clamp(PrevHistoryBufferUV[1], PrevHistoryInfo_UVViewportBilinearMin[1], PrevHistoryInfo_UVViewportBilinearMax[1]);
 
-			#if 1
-			{
-				FCatmullRomSamples Samples0 = GetBicubic2DCatmullRomSamples(
-					dpv_lo(PrevHistoryBufferUV), PrevHistoryInfo_Extent, PrevHistoryInfo_ExtentInverse);
-				FCatmullRomSamples Samples1 = GetBicubic2DCatmullRomSamples(
-					dpv_hi(PrevHistoryBufferUV), PrevHistoryInfo_Extent, PrevHistoryInfo_ExtentInverse);
+			FCatmullRomSamples Samples0 = GetBicubic2DCatmullRomSamples(
+				dpv_lo(PrevHistoryBufferUV), PrevHistoryInfo_Extent, PrevHistoryInfo_ExtentInverse);
+			FCatmullRomSamples Samples1 = GetBicubic2DCatmullRomSamples(
+				dpv_hi(PrevHistoryBufferUV), PrevHistoryInfo_Extent, PrevHistoryInfo_ExtentInverse);
+				
+			taa_half3 RawLowFrequency[BICUBIC_CATMULL_ROM_SAMPLES][DPV_PIXEL_PER_LANE];
+			taa_half3 RawHighFrequency[BICUBIC_CATMULL_ROM_SAMPLES][DPV_PIXEL_PER_LANE];
+			taa_half2 RawMetadata[BICUBIC_CATMULL_ROM_SAMPLES][DPV_PIXEL_PER_LANE];
+			taa_half3 RawTranslucencyColor[BICUBIC_CATMULL_ROM_SAMPLES][DPV_PIXEL_PER_LANE];
+			taa_half  RawTranslucencyAlpha[BICUBIC_CATMULL_ROM_SAMPLES][DPV_PIXEL_PER_LANE];
+
+			taa_half2 RawKernelWeight[BICUBIC_CATMULL_ROM_SAMPLES];
 
-				UNROLL
-				for (uint i = 0; i < Samples0.Count; i++)
+			// Issues texture fetches.
+			PLATFORM_SPECIFIC_ISOLATE
+			{
+				UNROLL_N(BICUBIC_CATMULL_ROM_SAMPLES)
+				for (uint i = 0; i < BICUBIC_CATMULL_ROM_SAMPLES; i++)
 				{
 					float2 SampleUV0 = clamp(Samples0.UV[i], PrevHistoryInfo_UVViewportBilinearMin, PrevHistoryInfo_UVViewportBilinearMax);
 					float2 SampleUV1 = clamp(Samples1.UV[i], PrevHistoryInfo_UVViewportBilinearMin, PrevHistoryInfo_UVViewportBilinearMax);
+					
+					taa_half2 KernelWeight = dpv_force_interleave_registers(taa_half(Samples0.Weight[i]), taa_half(Samples1.Weight[i]));
+
+					RawKernelWeight[i] = KernelWeight;
+
+					UNROLL_N(2)
+					for (uint PixelId = 0; PixelId < DPV_PIXEL_PER_LANE; PixelId++)
+					{
+						float2 SampleUV0 = clamp(Samples0.UV[i], PrevHistoryInfo_UVViewportBilinearMin, PrevHistoryInfo_UVViewportBilinearMax);
+						float2 SampleUV1 = clamp(Samples1.UV[i], PrevHistoryInfo_UVViewportBilinearMin, PrevHistoryInfo_UVViewportBilinearMax);
+					
+						float2 SampleUV = PixelId == 0 ? SampleUV0 : SampleUV1;
+						
+						RawLowFrequency[i][PixelId]      = PrevHistory_LowFrequency.SampleLevel(GlobalBilinearClampedSampler, SampleUV, 0);
+						RawHighFrequency[i][PixelId]     = PrevHistory_HighFrequency.SampleLevel(GlobalBilinearClampedSampler, SampleUV, 0);
+						RawMetadata[i][PixelId]	         = PrevHistory_Metadata.SampleLevel(GlobalBilinearClampedSampler, SampleUV, 0);
+						
+						#if CONFIG_ACCUMULATE_TRANSLUCENCY_SEPARATLY
+							RawTranslucencyColor[i][PixelId] = PrevHistory_Translucency.SampleLevel(GlobalBilinearClampedSampler, SampleUV, 0);
+							RawTranslucencyAlpha[i][PixelId] = PrevHistory_TranslucencyAlpha.SampleLevel(GlobalBilinearClampedSampler, SampleUV, 0);
+						#endif
+					}
 
-					taa_half3x2 SampleLowFrequency = dpv_interleave_registers(
-						PrevHistory_LowFrequency.SampleLevel(GlobalBilinearClampedSampler, SampleUV0, 0),
-						PrevHistory_LowFrequency.SampleLevel(GlobalBilinearClampedSampler, SampleUV1, 0));
-					taa_half3x2 SampleHighFrequency = dpv_interleave_registers(
-						PrevHistory_HighFrequency.SampleLevel(GlobalBilinearClampedSampler, SampleUV0, 0),
-						PrevHistory_HighFrequency.SampleLevel(GlobalBilinearClampedSampler, SampleUV1, 0));
-					taa_half2x2 SampleMetadata = dpv_interleave_registers(
-						PrevHistory_Metadata.SampleLevel(GlobalBilinearClampedSampler, SampleUV0, 0),
-						PrevHistory_Metadata.SampleLevel(GlobalBilinearClampedSampler, SampleUV1, 0));
+				} // for (uint i = 0; i < BICUBIC_CATMULL_ROM_SAMPLES; i++)
+			}
+
+			// Process history texture fetches.
+			PLATFORM_SPECIFIC_ISOLATE
+			{
+				UNROLL_N(BICUBIC_CATMULL_ROM_SAMPLES)
+				for (uint i = 0; i < BICUBIC_CATMULL_ROM_SAMPLES; i++)
+				{
+					taa_half3x2 SampleLowFrequency = dpv_interleave_registers_array(RawLowFrequency[i]);
+					taa_half3x2 SampleHighFrequency = dpv_interleave_registers_array(RawHighFrequency[i]);
+					taa_half2x2 SampleMetadata = dpv_interleave_registers_array(RawMetadata[i]);
 					taa_half4x2 SampleTranslucency = dpv_interleave_registers(
-						taa_half4(
-							PrevHistory_Translucency.SampleLevel(GlobalBilinearClampedSampler, SampleUV0, 0),
-							PrevHistory_TranslucencyAlpha.SampleLevel(GlobalBilinearClampedSampler, SampleUV0, 0)),
-						taa_half4(
-							PrevHistory_Translucency.SampleLevel(GlobalBilinearClampedSampler, SampleUV1, 0),
-							PrevHistory_TranslucencyAlpha.SampleLevel(GlobalBilinearClampedSampler, SampleUV1, 0)));
+						taa_half4(RawTranslucencyColor[i][0], RawTranslucencyAlpha[i][0]),
+						taa_half4(RawTranslucencyColor[i][1], RawTranslucencyAlpha[i][1]));
+					
+					taa_half2 KernelWeight = RawKernelWeight[i];
 
 					taa_half2 InvWeight = SafeRcp(SampleMetadata[1]);
 
 					HighFrequencyMin = min(HighFrequencyMin, dpv_scale(SampleHighFrequency, InvWeight));
 					HighFrequencyMax = max(HighFrequencyMax, dpv_scale(SampleHighFrequency, InvWeight));
 
-					taa_half2 KernelWeight = dpv_interleave_registers(taa_half(Samples0.Weight[i]), taa_half(Samples1.Weight[i]));
-
 					AccumulateLowFrequency  += dpv_scale(SampleLowFrequency, KernelWeight);
 					AccumulateHighFrequency += dpv_scale(SampleHighFrequency, KernelWeight);
 					AccumulateMetadata      += dpv_scale(SampleMetadata, KernelWeight);
-					AccumulateTranslucency  += dpv_scale(SampleTranslucency, KernelWeight);
+					
+					#if CONFIG_ACCUMULATE_TRANSLUCENCY_SEPARATLY
+						AccumulateTranslucency  += dpv_scale(SampleTranslucency, KernelWeight);
+					#endif
 				}
 				
 				taa_half2 FinalMultiplier = dpv_interleave_registers(taa_half(Samples0.FinalMultiplier), taa_half(Samples1.FinalMultiplier));
@@ -592,15 +701,7 @@ void MainCS(
 				AccumulateMetadata      = dpv_scale(AccumulateMetadata, FinalMultiplier);
 				AccumulateTranslucency  = dpv_scale(AccumulateTranslucency, FinalMultiplier);
 			}
-			#else
-			{
-				AccumulateLowFrequency  = PrevHistory_LowFrequency.SampleLevel(GlobalBilinearClampedSampler, PrevHistoryBufferUV, 0);
-				AccumulateHighFrequency = PrevHistory_HighFrequency.SampleLevel(GlobalBilinearClampedSampler, PrevHistoryBufferUV, 0);
-				AccumulateMetadata      = PrevHistory_Metadata.SampleLevel(GlobalBilinearClampedSampler, PrevHistoryBufferUV, 0);
-				AccumulateTranslucency  = PrevHistory_Translucency.SampleLevel(GlobalBilinearClampedSampler, PrevHistoryBufferUV, 0);
-			}
-			#endif
-			
+
 			FSubpixelNeighborhood SubpixelNeighborhood0 = GatherPrevSubpixelNeighborhood(PrevHistory_SubpixelDetails, dpv_lo(PrevHistoryBufferUV));
 			FSubpixelNeighborhood SubpixelNeighborhood1 = GatherPrevSubpixelNeighborhood(PrevHistory_SubpixelDetails, dpv_hi(PrevHistoryBufferUV));
 			{
@@ -645,7 +746,7 @@ void MainCS(
 	}
 		
 	#if CONFIG_MANUAL_LDS_SPILL
-	ISOLATE
+	PLATFORM_SPECIFIC_ISOLATE
 	{
 		uint LocalGroupThreadIndex = GetGroupThreadIndex(GroupThreadIndex, GroupWaveIndex);
 	
@@ -938,13 +1039,13 @@ void MainCS(
 				// float InvertedBlendFinal = PrevLowFrequencySampleCount / (CurrentFrameSampleCount + PrevLowFrequencySampleCount);
 				// float PrevLowFrequencyRejectionFactor = saturate(LowFrequencyRejection / InvertedBlendFinal);
 				
-				taa_half2 PrevLowFrequencyRejectionFactor = saturate(LowFrequencyRejection * (CurrentFrameSampleCount + PrevLowFrequencySampleCount) / PrevLowFrequencySampleCount);
+				taa_half2 PrevLowFrequencyRejectionFactor = saturate(LowFrequencyRejection * (CurrentFrameSampleCount + PrevLowFrequencySampleCount) * rcp(PrevLowFrequencySampleCount));
 
 				PrevLowFrequencySampleCount *= PrevLowFrequencyRejectionFactor;
 			}
 			#endif
 
-			BlendFinal = CurrentFrameSampleCount / (CurrentFrameSampleCount + PrevLowFrequencySampleCount);
+			BlendFinal = CurrentFrameSampleCount * rcp(CurrentFrameSampleCount + PrevLowFrequencySampleCount);
 
 			// Increase the blending on motion
 			#if 1
@@ -956,7 +1057,7 @@ void MainCS(
 			// Antiflicker
 			#if 1
 			{
-				taa_half2 DistToClamp = min( abs(LumaHistory - LumaMin), abs(LumaHistory - LumaMax) ) / max3( LumaHistory, LumaFiltered, taa_half(1e-4) );
+				taa_half2 DistToClamp = min( abs(LumaHistory - LumaMin), abs(LumaHistory - LumaMax) ) * rcp(max3( LumaHistory, LumaFiltered, taa_half(1e-4) ));
 				BlendFinal *= taa_half(0.2) + taa_half(0.8) * saturate(taa_half(0.5) * DistToClamp);
 			}
 			#endif
@@ -964,7 +1065,7 @@ void MainCS(
 			// Make sure to have at least some small contribution
 			#if 1
 			{
-				BlendFinal = max( BlendFinal, saturate( taa_half(0.01) * LumaHistory / abs( LumaFiltered - LumaHistory ) ) );
+				BlendFinal = max(BlendFinal, saturate(taa_half(0.01) * LumaHistory * rcp(abs(LumaFiltered - LumaHistory))));
 			}
 			#endif
 
@@ -1028,7 +1129,7 @@ void MainCS(
 		// Make sure to have at least some small contribution
 		#if 0
 		{
-			BlendFinal = max( BlendFinal, saturate( taa_half(0.01) * TranslucencyLumaHistory / abs( TranslucencyLumaFiltered - TranslucencyLumaHistory ) ) );
+			BlendFinal = max( BlendFinal, saturate( taa_half(0.01) * TranslucencyLumaHistory * rcp(abs(TranslucencyLumaFiltered - TranslucencyLumaHistory))));
 		}
 		#endif
 		
@@ -1168,7 +1269,7 @@ void MainCS(
 	}
 
 
-	ISOLATE
+	PLATFORM_SPECIFIC_ISOLATE
 	{
 		uint LocalGroupThreadIndex = GetGroupThreadIndex(GroupThreadIndex, GroupWaveIndex);
 
diff --git a/Engine/Shaders/Private/TextureSampling.ush b/Engine/Shaders/Private/TextureSampling.ush
--- a/Engine/Shaders/Private/TextureSampling.ush
+++ b/Engine/Shaders/Private/TextureSampling.ush
@@ -105,7 +105,7 @@ void Bicubic2DCatmullRom(in float2 UV, in float2 Size, in float2 InvSize, out fl
 	Weight[2] = w3;
 
 	Sample[0] = tc - 1;
-	Sample[1] = tc + w2 / Weight[1];
+	Sample[1] = tc + w2 * rcp(Weight[1]);
 	Sample[2] = tc + 2;
 
 	Sample[0] *= InvSize;
diff --git a/Engine/Shaders/Public/DualPixelVectorization.ush b/Engine/Shaders/Public/DualPixelVectorization.ush
--- a/Engine/Shaders/Public/DualPixelVectorization.ush
+++ b/Engine/Shaders/Public/DualPixelVectorization.ush
@@ -11,88 +11,174 @@
 
 #define PLATFORM_SUPPORT_16BITS_DPV (PLATFORM_SUPPORTS_REAL_TYPES)
 
+// Dual Pixel vectoring processes two pixel per lane.
+#define DPV_PIXEL_PER_LANE 2
+
 
 // ------------------------------------------------- DEFINE_DPV_FUNCTION_FOR_ALL_32BITS
 
+// Helper to get the type for a single pixel vector
 #define __DPV_VECTOR(TYPE, COMPONENTS) \
 	TYPE##COMPONENTS
 
+// Helper to get the type for a dual pixel scalar
 #define __DPV_MULTI_SCALAR(TYPE) \
 	TYPE##2
 
+// Helper to get the type for a dual pixel vector
 #define __DPV_MULTI_VECTOR(TYPE, COMPONENTS) \
 	TYPE##COMPONENTS##x2
 
+// Helper to get the transposed type for a dual pixel vector
 #define __DPV_MULTI_VECTOR_TRANSPOSED(TYPE, COMPONENTS) \
 	TYPE##2x##COMPONENTS
 
 
-#if PLATFORM_SUPPORT_16BITS_DPV
+// __DPV_FOREACH_COMPONENTS() helper to write per component code with abritrary separator (like an operator for dot product)
+#define __DPV_FOREACH_COMPONENTS_2(COMPONENT_CODE, COMPONENT_SEPARATOR, ...) \
+	COMPONENT_CODE(0, __VA_ARGS__) \
+	COMPONENT_SEPARATOR \
+	COMPONENT_CODE(1, __VA_ARGS__) \
+
+#define __DPV_FOREACH_COMPONENTS_3(COMPONENT_CODE, COMPONENT_SEPARATOR, ...) \
+	COMPONENT_CODE(0, __VA_ARGS__) \
+	COMPONENT_SEPARATOR \
+	COMPONENT_CODE(1, __VA_ARGS__) \
+	COMPONENT_SEPARATOR \
+	COMPONENT_CODE(2, __VA_ARGS__) \
+
+#define __DPV_FOREACH_COMPONENTS_4(COMPONENT_CODE, COMPONENT_SEPARATOR, ...) \
+	COMPONENT_CODE(0, __VA_ARGS__) \
+	COMPONENT_SEPARATOR \
+	COMPONENT_CODE(1, __VA_ARGS__) \
+	COMPONENT_SEPARATOR \
+	COMPONENT_CODE(2, __VA_ARGS__) \
+	COMPONENT_SEPARATOR \
+	COMPONENT_CODE(3, __VA_ARGS__) \
+
+#define __DPV_FOREACH_COMPONENTS(COMPONENTS, COMPONENT_CODE, COMPONENT_SEPARATOR, ...) \
+	__DPV_FOREACH_COMPONENTS_##COMPONENTS(COMPONENT_CODE, COMPONENT_SEPARATOR, __VA_ARGS__)
+	
+
+// __DPV_FOREACH_COMPONENTS() helper to write per component code to be used as comma separated parameters to a function
+#define __DPV_FOREACH_COMPONENTS_ARGS_2(COMPONENT_CODE, ...) \
+	COMPONENT_CODE(0, __VA_ARGS__), \
+	COMPONENT_CODE(1, __VA_ARGS__) \
+
+#define __DPV_FOREACH_COMPONENTS_ARGS_3(COMPONENT_CODE, ...) \
+	COMPONENT_CODE(0, __VA_ARGS__), \
+	COMPONENT_CODE(1, __VA_ARGS__), \
+	COMPONENT_CODE(2, __VA_ARGS__) \
+
+#define __DPV_FOREACH_COMPONENTS_ARGS_4(COMPONENT_CODE, ...) \
+	COMPONENT_CODE(0, __VA_ARGS__), \
+	COMPONENT_CODE(1, __VA_ARGS__), \
+	COMPONENT_CODE(2, __VA_ARGS__), \
+	COMPONENT_CODE(3, __VA_ARGS__) \
+
+#define __DPV_FOREACH_COMPONENTS_ARGS(COMPONENTS, COMPONENT_CODE, ...) \
+	__DPV_FOREACH_COMPONENTS_ARGS_##COMPONENTS(COMPONENT_CODE, __VA_ARGS__)
+	
 
+// __DPV_FOREACH_COMPONENTS_CONSTRUCTOR() helper to write constructor for a dual pixel vector
+#define __DPV_FOREACH_COMPONENTS_CONSTRUCTOR(TYPE, COMPONENTS, COMPONENT_CODE, ...) \
+	__DPV_MULTI_VECTOR(TYPE, COMPONENTS)(__DPV_FOREACH_COMPONENTS_ARGS(COMPONENTS, COMPONENT_CODE, __VA_ARGS__))
+
+
+// Helper to be used with __DPV_FOREACH_COMPONENTS() to call a function
+#define __DPV_CODE_FUNCTION_2_ARGS(c, FUNCTION_NAME, ARG_NAME_0, ARG_NAME_1) FUNCTION_NAME(ARG_NAME_0[c], ARG_NAME_1[c])
+#define __DPV_CODE_FUNCTION_3_ARGS(c, FUNCTION_NAME, ARG_NAME_0, ARG_NAME_1, ARG_NAME_2) FUNCTION_NAME(ARG_NAME_0[c], ARG_NAME_1[c], ARG_NAME_2[c])
+
+// helper to be used with __DPV_FOREACH_COMPONENTS() to call an operator
+#define __DPV_CODE_OPERATOR_VECTOR_VECTOR(c, ARG_NAME_0, OPERATOR, ARG_NAME_1) (ARG_NAME_0[c] OPERATOR ARG_NAME_1[c])
+#define __DPV_CODE_OPERATOR_VECTOR_SCALAR(c, ARG_NAME_0, OPERATOR, ARG_NAME_1) (ARG_NAME_0[c] OPERATOR ARG_NAME_1)
+
+
+// Helpers to define for each individual scalar and vector type
 #define DEFINE_DPV_FUNCTION_FOR_FLOAT_SCALARS(DEFINE_DPV_FUNCTION) \
 	DEFINE_DPV_FUNCTION(float) \
-	DEFINE_DPV_FUNCTION(float16_t) \
 
 #define DEFINE_DPV_FUNCTION_FOR_INT_SCALARS(DEFINE_DPV_FUNCTION) \
 	DEFINE_DPV_FUNCTION(uint) \
 	DEFINE_DPV_FUNCTION(int) \
-	DEFINE_DPV_FUNCTION(uint16_t) \
-	DEFINE_DPV_FUNCTION(int16_t) \
 		
+#define DEFINE_DPV_FUNCTION_FOR_FLOAT_VECTOR(DEFINE_DPV_FUNCTION, COMPONENTS) \
+	DEFINE_DPV_FUNCTION(float, COMPONENTS) \
+
+#define DEFINE_DPV_FUNCTION_FOR_INT_VECTOR(DEFINE_DPV_FUNCTION, COMPONENTS) \
+	DEFINE_DPV_FUNCTION(uint, COMPONENTS) \
+	DEFINE_DPV_FUNCTION(int, COMPONENTS) \
+
 #define DEFINE_DPV_FUNCTION_FOR_FLOAT_VECTORS(DEFINE_DPV_FUNCTION) \
-	DEFINE_DPV_FUNCTION(float, 2) \
-	DEFINE_DPV_FUNCTION(float, 3) \
-	DEFINE_DPV_FUNCTION(float, 4) \
-	DEFINE_DPV_FUNCTION(float16_t, 2) \
-	DEFINE_DPV_FUNCTION(float16_t, 3) \
-	DEFINE_DPV_FUNCTION(float16_t, 4) \
+	DEFINE_DPV_FUNCTION_FOR_FLOAT_VECTOR(DEFINE_DPV_FUNCTION, 2) \
+	DEFINE_DPV_FUNCTION_FOR_FLOAT_VECTOR(DEFINE_DPV_FUNCTION, 3) \
+	DEFINE_DPV_FUNCTION_FOR_FLOAT_VECTOR(DEFINE_DPV_FUNCTION, 4) \
 
 #define DEFINE_DPV_FUNCTION_FOR_INT_VECTORS(DEFINE_DPV_FUNCTION) \
-	DEFINE_DPV_FUNCTION(uint, 2) \
-	DEFINE_DPV_FUNCTION(uint, 3) \
-	DEFINE_DPV_FUNCTION(uint, 4) \
-	DEFINE_DPV_FUNCTION(int, 2) \
-	DEFINE_DPV_FUNCTION(int, 3) \
-	DEFINE_DPV_FUNCTION(int, 4) \
-	DEFINE_DPV_FUNCTION(uint16_t, 2) \
-	DEFINE_DPV_FUNCTION(uint16_t, 3) \
-	DEFINE_DPV_FUNCTION(uint16_t, 4) \
-	DEFINE_DPV_FUNCTION(int16_t, 2) \
-	DEFINE_DPV_FUNCTION(int16_t, 3) \
-	DEFINE_DPV_FUNCTION(int16_t, 4) \
+	DEFINE_DPV_FUNCTION_FOR_INT_VECTOR(DEFINE_DPV_FUNCTION, 2) \
+	DEFINE_DPV_FUNCTION_FOR_INT_VECTOR(DEFINE_DPV_FUNCTION, 3) \
+	DEFINE_DPV_FUNCTION_FOR_INT_VECTOR(DEFINE_DPV_FUNCTION, 4) \
+		
+#if PLATFORM_SUPPORT_16BITS_DPV
+
+	#define DEFINE_DPV_FUNCTION_FOR_HALF_SCALARS(DEFINE_DPV_FUNCTION) \
+		DEFINE_DPV_FUNCTION(float16_t) \
+
+	#define DEFINE_DPV_FUNCTION_FOR_SHORT_SCALARS(DEFINE_DPV_FUNCTION) \
+		DEFINE_DPV_FUNCTION(uint16_t) \
+		DEFINE_DPV_FUNCTION(int16_t) \
+		
+	#define DEFINE_DPV_FUNCTION_FOR_HALF_VECTORS(DEFINE_DPV_FUNCTION) \
+		DEFINE_DPV_FUNCTION(float16_t, 2) \
+		DEFINE_DPV_FUNCTION(float16_t, 3) \
+		DEFINE_DPV_FUNCTION(float16_t, 4) \
+
+	#define DEFINE_DPV_FUNCTION_FOR_SHORT_VECTORS(DEFINE_DPV_FUNCTION) \
+		DEFINE_DPV_FUNCTION(uint16_t, 2) \
+		DEFINE_DPV_FUNCTION(uint16_t, 3) \
+		DEFINE_DPV_FUNCTION(uint16_t, 4) \
+		DEFINE_DPV_FUNCTION(int16_t, 2) \
+		DEFINE_DPV_FUNCTION(int16_t, 3) \
+		DEFINE_DPV_FUNCTION(int16_t, 4) \
 		
 #else
 
-#define DEFINE_DPV_FUNCTION_FOR_FLOAT_SCALARS(DEFINE_DPV_FUNCTION) \
-	DEFINE_DPV_FUNCTION(float) \
+	#define DEFINE_DPV_FUNCTION_FOR_HALF_SCALARS(DEFINE_DPV_FUNCTION)
 
-#define DEFINE_DPV_FUNCTION_FOR_INT_SCALARS(DEFINE_DPV_FUNCTION) \
-	DEFINE_DPV_FUNCTION(uint) \
-	DEFINE_DPV_FUNCTION(int) \
+	#define DEFINE_DPV_FUNCTION_FOR_SHORT_SCALARS(DEFINE_DPV_FUNCTION)
 		
-#define DEFINE_DPV_FUNCTION_FOR_FLOAT_VECTORS(DEFINE_DPV_FUNCTION) \
-	DEFINE_DPV_FUNCTION(float, 2) \
-	DEFINE_DPV_FUNCTION(float, 3) \
-	DEFINE_DPV_FUNCTION(float, 4) \
+	#define DEFINE_DPV_FUNCTION_FOR_HALF_VECTORS(DEFINE_DPV_FUNCTION)
 
-#define DEFINE_DPV_FUNCTION_FOR_INT_VECTORS(DEFINE_DPV_FUNCTION) \
-	DEFINE_DPV_FUNCTION(uint, 2) \
-	DEFINE_DPV_FUNCTION(uint, 3) \
-	DEFINE_DPV_FUNCTION(uint, 4) \
-	DEFINE_DPV_FUNCTION(int, 2) \
-	DEFINE_DPV_FUNCTION(int, 3) \
-	DEFINE_DPV_FUNCTION(int, 4) \
+	#define DEFINE_DPV_FUNCTION_FOR_SHORT_VECTORS(DEFINE_DPV_FUNCTION)
 		
 #endif
 
-#define DEFINE_DPV_FUNCTION_FOR_ALL_SCALARS(DEFINE_DPV_FUNCTION) \
+#define DEFINE_DPV_FUNCTION_FOR_ALL_32BIT_SCALARS(DEFINE_DPV_FUNCTION) \
 	DEFINE_DPV_FUNCTION_FOR_FLOAT_SCALARS(DEFINE_DPV_FUNCTION) \
 	DEFINE_DPV_FUNCTION_FOR_INT_SCALARS(DEFINE_DPV_FUNCTION) \
 
-#define DEFINE_DPV_FUNCTION_FOR_ALL_VECTORS(DEFINE_DPV_FUNCTION) \
+#define DEFINE_DPV_FUNCTION_FOR_ALL_16BIT_SCALARS(DEFINE_DPV_FUNCTION) \
+	DEFINE_DPV_FUNCTION_FOR_HALF_SCALARS(DEFINE_DPV_FUNCTION) \
+	DEFINE_DPV_FUNCTION_FOR_SHORT_SCALARS(DEFINE_DPV_FUNCTION) \
+
+#define DEFINE_DPV_FUNCTION_FOR_ALL_SCALARS(DEFINE_DPV_FUNCTION) \
+	DEFINE_DPV_FUNCTION_FOR_ALL_32BIT_SCALARS(DEFINE_DPV_FUNCTION) \
+	DEFINE_DPV_FUNCTION_FOR_ALL_16BIT_SCALARS(DEFINE_DPV_FUNCTION) \
+
+
+#define DEFINE_DPV_FUNCTION_FOR_ALL_32BIT_VECTORS(DEFINE_DPV_FUNCTION) \
 	DEFINE_DPV_FUNCTION_FOR_FLOAT_VECTORS(DEFINE_DPV_FUNCTION) \
 	DEFINE_DPV_FUNCTION_FOR_INT_VECTORS(DEFINE_DPV_FUNCTION) \
 
+#define DEFINE_DPV_FUNCTION_FOR_ALL_16BIT_VECTORS(DEFINE_DPV_FUNCTION) \
+	DEFINE_DPV_FUNCTION_FOR_HALF_VECTORS(DEFINE_DPV_FUNCTION) \
+	DEFINE_DPV_FUNCTION_FOR_SHORT_VECTORS(DEFINE_DPV_FUNCTION) \
+
+#define DEFINE_DPV_FUNCTION_FOR_ALL_VECTORS(DEFINE_DPV_FUNCTION) \
+	DEFINE_DPV_FUNCTION_FOR_ALL_32BIT_VECTORS(DEFINE_DPV_FUNCTION) \
+	DEFINE_DPV_FUNCTION_FOR_ALL_16BIT_VECTORS(DEFINE_DPV_FUNCTION) \
+
+
 
 /* dpv_interleave_registers() & dpv_lo() & dpv_hi()
  *
@@ -102,7 +188,7 @@
  *		half3 Color0 = ... ;
  *		half3 Color1 = ... ;
  *		
- *		half3x2 Color0 = dpv_interleave_registers(Color0, Color1);
+ *		half3x2 Color = dpv_interleave_registers(Color0, Color1);
  *		
  *		// Can access each individual channel with the [0] notation
  *		half2 Red = Color[0];
@@ -118,11 +204,13 @@
 	__DPV_MULTI_SCALAR(TYPE) dpv_interleave_registers(TYPE lo, TYPE hi) { return __DPV_MULTI_SCALAR(TYPE)(lo, hi); } \
 	TYPE dpv_lo(__DPV_MULTI_SCALAR(TYPE) x) { return x[0]; } \
 	TYPE dpv_hi(__DPV_MULTI_SCALAR(TYPE) x) { return x[1]; } \
+	TYPE dpv_access_pixel(__DPV_MULTI_SCALAR(TYPE) x, const uint DualPixelId) { return x[DualPixelId]; } \
 
 #define DEFINE_DPV_FUNCTION_VECTOR(TYPE, COMPONENTS) \
 	__DPV_MULTI_VECTOR(TYPE, COMPONENTS) dpv_interleave_registers(__DPV_VECTOR(TYPE, COMPONENTS) lo, __DPV_VECTOR(TYPE, COMPONENTS) hi) { return transpose(__DPV_MULTI_VECTOR_TRANSPOSED(TYPE, COMPONENTS)(lo, hi)); } \
 	__DPV_VECTOR(TYPE, COMPONENTS) dpv_lo(__DPV_MULTI_VECTOR(TYPE, COMPONENTS) x) { return transpose(x)[0]; } \
 	__DPV_VECTOR(TYPE, COMPONENTS) dpv_hi(__DPV_MULTI_VECTOR(TYPE, COMPONENTS) x) { return transpose(x)[1]; } \
+	__DPV_VECTOR(TYPE, COMPONENTS) dpv_access_pixel(__DPV_MULTI_VECTOR(TYPE, COMPONENTS) x, const uint DualPixelId) { return transpose(x)[DualPixelId]; } \
 
 DEFINE_DPV_FUNCTION_FOR_ALL_SCALARS(DEFINE_DPV_FUNCTION_SCALAR)
 DEFINE_DPV_FUNCTION_FOR_ALL_VECTORS(DEFINE_DPV_FUNCTION_VECTOR)
@@ -130,6 +218,55 @@ DEFINE_DPV_FUNCTION_FOR_ALL_VECTORS(DEFINE_DPV_FUNCTION_VECTOR)
 #undef DEFINE_DPV_FUNCTION_VECTOR
 
 
+/* dpv_force_interleave_registers() to force using RDNA's v_pack_b32_b16
+ *
+ * Reinterleaves explicit vector floatN for each individual pixel into a floatNx2 matrix.
+ * 
+ * Example:
+ *		half3x2 Color = dpv_force_interleave_registers(Color0, Color1);
+ */
+
+#define DEFINE_DPV_FUNCTION_SCALAR(TYPE) \
+	__DPV_MULTI_SCALAR(TYPE) dpv_force_interleave_registers(TYPE lo, TYPE hi) { return dpv_interleave_registers(lo, hi); } \
+
+#define DEFINE_DPV_FUNCTION_VECTOR(TYPE, COMPONENTS) \
+	__DPV_MULTI_VECTOR(TYPE, COMPONENTS) dpv_force_interleave_registers(__DPV_VECTOR(TYPE, COMPONENTS) lo, __DPV_VECTOR(TYPE, COMPONENTS) hi) { return dpv_interleave_registers(lo, hi); } \
+
+DEFINE_DPV_FUNCTION_FOR_ALL_32BIT_SCALARS(DEFINE_DPV_FUNCTION_SCALAR)
+DEFINE_DPV_FUNCTION_FOR_ALL_32BIT_VECTORS(DEFINE_DPV_FUNCTION_VECTOR)
+#undef DEFINE_DPV_FUNCTION_SCALAR
+#undef DEFINE_DPV_FUNCTION_VECTOR
+
+
+#define DEFINE_DPV_FUNCTION_SCALAR(TYPE) \
+	__DPV_MULTI_SCALAR(TYPE) dpv_force_interleave_registers(TYPE lo, TYPE hi) { return v_pack_b32_b16(lo, hi); } \
+
+#define DEFINE_DPV_FUNCTION_VECTOR(TYPE, COMPONENTS) \
+	__DPV_MULTI_VECTOR(TYPE, COMPONENTS) dpv_force_interleave_registers(__DPV_VECTOR(TYPE, COMPONENTS) lo, __DPV_VECTOR(TYPE, COMPONENTS) hi) { \
+		return __DPV_FOREACH_COMPONENTS_CONSTRUCTOR(TYPE, COMPONENTS, __DPV_CODE_FUNCTION_2_ARGS, dpv_force_interleave_registers, lo, hi); \
+	} \
+
+DEFINE_DPV_FUNCTION_FOR_ALL_16BIT_SCALARS(DEFINE_DPV_FUNCTION_SCALAR)
+DEFINE_DPV_FUNCTION_FOR_ALL_16BIT_VECTORS(DEFINE_DPV_FUNCTION_VECTOR)
+#undef DEFINE_DPV_FUNCTION_SCALAR
+#undef DEFINE_DPV_FUNCTION_VECTOR
+
+
+/* dpv_interleave_array_registers() to interleave an array
+ *
+ * Reinterleaves explicit vector floatN for each individual pixel into a floatNx2 matrix.
+ * 
+ * Example:
+ *		half3 ColorArray[2];
+ *		ColorArray[0] = ... ;
+ *		ColorArray[1] = ... ;
+ *		
+ *		half3x2 Color = dpv_interleave_registers_array(ColorArray);
+ */
+#define dpv_interleave_registers_array(ArrayName) dpv_interleave_registers((ArrayName)[0], (ArrayName)[1])
+#define dpv_force_interleave_registers_array(ArrayName) dpv_force_interleave_registers((ArrayName)[0], (ArrayName)[1])
+
+
 /* dpv_interleave_mono_registers()
  *
  * Reinterleaves explicit vector floatN for of one individual pixel into a floatNx2 matrix to reuse code for dual pixel vectorization.
@@ -167,53 +304,25 @@ DEFINE_DPV_FUNCTION_FOR_ALL_VECTORS(DEFINE_DPV_FUNCTION_VECTOR)
  */
 #define DEFINE_DPV_FUNCTION_VECTOR(TYPE, COMPONENTS) \
 	__DPV_MULTI_VECTOR(TYPE, COMPONENTS) dpv_mul(__DPV_MULTI_VECTOR(TYPE, COMPONENTS) a, __DPV_MULTI_VECTOR(TYPE, COMPONENTS) b) { \
-		__DPV_MULTI_VECTOR(TYPE, COMPONENTS) r; \
-		UNROLL_N(COMPONENTS) for (uint c = 0; c < COMPONENTS; c++) { \
-			r[c] = a[c] * b[c]; \
-		} \
-		return r; \
+		return __DPV_FOREACH_COMPONENTS_CONSTRUCTOR(TYPE, COMPONENTS, __DPV_CODE_OPERATOR_VECTOR_VECTOR, a, *, b); \
 	} \
 	__DPV_MULTI_VECTOR(TYPE, COMPONENTS) dpv_mul(__DPV_MULTI_VECTOR(TYPE, COMPONENTS) a, __DPV_VECTOR(TYPE, COMPONENTS) b) { \
-		__DPV_MULTI_VECTOR(TYPE, COMPONENTS) r; \
-		UNROLL_N(COMPONENTS) for (uint c = 0; c < COMPONENTS; c++) { \
-			r[c] = a[c] * b[c]; \
-		} \
-		return r; \
+		return __DPV_FOREACH_COMPONENTS_CONSTRUCTOR(TYPE, COMPONENTS, __DPV_CODE_OPERATOR_VECTOR_VECTOR, a, *, b); \
 	} \
 	__DPV_MULTI_VECTOR(TYPE, COMPONENTS) dpv_mul(__DPV_VECTOR(TYPE, COMPONENTS) a, __DPV_MULTI_VECTOR(TYPE, COMPONENTS) b) { \
-		__DPV_MULTI_VECTOR(TYPE, COMPONENTS) r; \
-		UNROLL_N(COMPONENTS) for (uint c = 0; c < COMPONENTS; c++) { \
-			r[c] = a[c] * b[c]; \
-		} \
-		return r; \
+		return __DPV_FOREACH_COMPONENTS_CONSTRUCTOR(TYPE, COMPONENTS, __DPV_CODE_OPERATOR_VECTOR_VECTOR, a, *, b); \
 	} \
 	__DPV_MULTI_VECTOR(TYPE, COMPONENTS) dpv_scale(__DPV_MULTI_VECTOR(TYPE, COMPONENTS) a, __DPV_MULTI_SCALAR(TYPE) b) { \
-		__DPV_MULTI_VECTOR(TYPE, COMPONENTS) r; \
-		UNROLL_N(COMPONENTS) for (uint c = 0; c < COMPONENTS; c++) { \
-			r[c] = a[c] * b; \
-		} \
-		return r; \
+		return __DPV_FOREACH_COMPONENTS_CONSTRUCTOR(TYPE, COMPONENTS, __DPV_CODE_OPERATOR_VECTOR_SCALAR, a, *, b); \
 	} \
 	__DPV_MULTI_VECTOR(TYPE, COMPONENTS) dpv_scale(__DPV_MULTI_SCALAR(TYPE) a, __DPV_MULTI_VECTOR(TYPE, COMPONENTS) b) { \
-		__DPV_MULTI_VECTOR(TYPE, COMPONENTS) r; \
-		UNROLL_N(COMPONENTS) for (uint c = 0; c < COMPONENTS; c++) { \
-			r[c] = a * b[c]; \
-		} \
-		return r; \
+		return __DPV_FOREACH_COMPONENTS_CONSTRUCTOR(TYPE, COMPONENTS, __DPV_CODE_OPERATOR_VECTOR_SCALAR, b, *, a); \
 	} \
 	__DPV_MULTI_VECTOR(TYPE, COMPONENTS) dpv_scale(__DPV_MULTI_VECTOR(TYPE, COMPONENTS) a, TYPE b) { \
-		__DPV_MULTI_VECTOR(TYPE, COMPONENTS) r; \
-		UNROLL_N(COMPONENTS) for (uint c = 0; c < COMPONENTS; c++) { \
-			r[c] = a[c] * b; \
-		} \
-		return r; \
+		return __DPV_FOREACH_COMPONENTS_CONSTRUCTOR(TYPE, COMPONENTS, __DPV_CODE_OPERATOR_VECTOR_SCALAR, a, *, b); \
 	} \
 	__DPV_MULTI_VECTOR(TYPE, COMPONENTS) dpv_scale(TYPE a, __DPV_MULTI_VECTOR(TYPE, COMPONENTS) b) { \
-		__DPV_MULTI_VECTOR(TYPE, COMPONENTS) r; \
-		UNROLL_N(COMPONENTS) for (uint c = 0; c < COMPONENTS; c++) { \
-			r[c] = a * b[c]; \
-		} \
-		return r; \
+		return __DPV_FOREACH_COMPONENTS_CONSTRUCTOR(TYPE, COMPONENTS, __DPV_CODE_OPERATOR_VECTOR_SCALAR, b, *, a); \
 	} \
 
 DEFINE_DPV_FUNCTION_FOR_ALL_VECTORS(DEFINE_DPV_FUNCTION_VECTOR)
@@ -230,46 +339,22 @@ DEFINE_DPV_FUNCTION_FOR_ALL_VECTORS(DEFINE_DPV_FUNCTION_VECTOR)
  */
 #define DEFINE_DPV_FUNCTION_VECTOR(TYPE, COMPONENTS) \
 	__DPV_MULTI_VECTOR(TYPE, COMPONENTS) dpv_add(__DPV_MULTI_VECTOR(TYPE, COMPONENTS) a, __DPV_MULTI_VECTOR(TYPE, COMPONENTS) b) { \
-		__DPV_MULTI_VECTOR(TYPE, COMPONENTS) r; \
-		UNROLL_N(COMPONENTS) for (uint c = 0; c < COMPONENTS; c++) { \
-			r[c] = a[c] + b[c]; \
-		} \
-		return r; \
+		return __DPV_FOREACH_COMPONENTS_CONSTRUCTOR(TYPE, COMPONENTS, __DPV_CODE_OPERATOR_VECTOR_VECTOR, a, +, b); \
 	} \
 	__DPV_MULTI_VECTOR(TYPE, COMPONENTS) dpv_add(__DPV_MULTI_VECTOR(TYPE, COMPONENTS) a, __DPV_VECTOR(TYPE, COMPONENTS) b) { \
-		__DPV_MULTI_VECTOR(TYPE, COMPONENTS) r; \
-		UNROLL_N(COMPONENTS) for (uint c = 0; c < COMPONENTS; c++) { \
-			r[c] = a[c] + b[c]; \
-		} \
-		return r; \
+		return __DPV_FOREACH_COMPONENTS_CONSTRUCTOR(TYPE, COMPONENTS, __DPV_CODE_OPERATOR_VECTOR_VECTOR, a, +, b); \
 	} \
 	__DPV_MULTI_VECTOR(TYPE, COMPONENTS) dpv_add(__DPV_VECTOR(TYPE, COMPONENTS) a, __DPV_MULTI_VECTOR(TYPE, COMPONENTS) b) { \
-		__DPV_MULTI_VECTOR(TYPE, COMPONENTS) r; \
-		UNROLL_N(COMPONENTS) for (uint c = 0; c < COMPONENTS; c++) { \
-			r[c] = a[c] + b[c]; \
-		} \
-		return r; \
+		return __DPV_FOREACH_COMPONENTS_CONSTRUCTOR(TYPE, COMPONENTS, __DPV_CODE_OPERATOR_VECTOR_VECTOR, a, +, b); \
 	} \
 	__DPV_MULTI_VECTOR(TYPE, COMPONENTS) dpv_sub(__DPV_MULTI_VECTOR(TYPE, COMPONENTS) a, __DPV_MULTI_VECTOR(TYPE, COMPONENTS) b) { \
-		__DPV_MULTI_VECTOR(TYPE, COMPONENTS) r; \
-		UNROLL_N(COMPONENTS) for (uint c = 0; c < COMPONENTS; c++) { \
-			r[c] = a[c] - b[c]; \
-		} \
-		return r; \
+		return __DPV_FOREACH_COMPONENTS_CONSTRUCTOR(TYPE, COMPONENTS, __DPV_CODE_OPERATOR_VECTOR_VECTOR, a, -, b); \
 	} \
 	__DPV_MULTI_VECTOR(TYPE, COMPONENTS) dpv_sub(__DPV_MULTI_VECTOR(TYPE, COMPONENTS) a, __DPV_VECTOR(TYPE, COMPONENTS) b) { \
-		__DPV_MULTI_VECTOR(TYPE, COMPONENTS) r; \
-		UNROLL_N(COMPONENTS) for (uint c = 0; c < COMPONENTS; c++) { \
-			r[c] = a[c] - b[c]; \
-		} \
-		return r; \
+		return __DPV_FOREACH_COMPONENTS_CONSTRUCTOR(TYPE, COMPONENTS, __DPV_CODE_OPERATOR_VECTOR_VECTOR, a, -, b); \
 	} \
 	__DPV_MULTI_VECTOR(TYPE, COMPONENTS) dpv_sub(__DPV_VECTOR(TYPE, COMPONENTS) a, __DPV_MULTI_VECTOR(TYPE, COMPONENTS) b) { \
-		__DPV_MULTI_VECTOR(TYPE, COMPONENTS) r; \
-		UNROLL_N(COMPONENTS) for (uint c = 0; c < COMPONENTS; c++) { \
-			r[c] = a[c] - b[c]; \
-		} \
-		return r; \
+		return __DPV_FOREACH_COMPONENTS_CONSTRUCTOR(TYPE, COMPONENTS, __DPV_CODE_OPERATOR_VECTOR_VECTOR, a, -, b); \
 	} \
 
 DEFINE_DPV_FUNCTION_FOR_ALL_VECTORS(DEFINE_DPV_FUNCTION_VECTOR)
@@ -285,54 +370,34 @@ DEFINE_DPV_FUNCTION_FOR_ALL_VECTORS(DEFINE_DPV_FUNCTION_VECTOR)
  */
 #define DEFINE_DPV_FUNCTION_VECTOR(TYPE, COMPONENTS) \
 	__DPV_MULTI_SCALAR(TYPE) dpv_dot(__DPV_MULTI_VECTOR(TYPE, COMPONENTS) a, __DPV_MULTI_VECTOR(TYPE, COMPONENTS) b) { \
-		return dpv_interleave_registers(dot(dpv_lo(a), dpv_lo(b)), dot(dpv_hi(a), dpv_hi(b))); \
+		return __DPV_FOREACH_COMPONENTS(COMPONENTS, __DPV_CODE_OPERATOR_VECTOR_VECTOR, +, a, *, b); \
 	} \
 	__DPV_MULTI_SCALAR(TYPE) dpv_dot(__DPV_VECTOR(TYPE, COMPONENTS) a, __DPV_MULTI_VECTOR(TYPE, COMPONENTS) b) { \
-		return dpv_interleave_registers(dot(a, dpv_lo(b)), dot(a, dpv_hi(b))); \
+		return __DPV_FOREACH_COMPONENTS(COMPONENTS, __DPV_CODE_OPERATOR_VECTOR_VECTOR, +, a, *, b); \
 	} \
 	__DPV_MULTI_SCALAR(TYPE) dpv_dot(__DPV_MULTI_VECTOR(TYPE, COMPONENTS) a, __DPV_VECTOR(TYPE, COMPONENTS) b) { \
-		return dpv_interleave_registers(dot(dpv_lo(a), b), dot(dpv_hi(a), b)); \
+		return __DPV_FOREACH_COMPONENTS(COMPONENTS, __DPV_CODE_OPERATOR_VECTOR_VECTOR, +, a, *, b); \
 	} \
 	__DPV_MULTI_VECTOR(TYPE, COMPONENTS) dpv_clamp(__DPV_MULTI_VECTOR(TYPE, COMPONENTS) x, __DPV_MULTI_VECTOR(TYPE, COMPONENTS) min, __DPV_MULTI_VECTOR(TYPE, COMPONENTS) max) { \
-		__DPV_MULTI_VECTOR(TYPE, COMPONENTS) r; \
-		UNROLL_N(COMPONENTS) for (uint c = 0; c < COMPONENTS; c++) { \
-			r[c] = clamp(x[c], min[c], max[c]); \
-		} \
-		return r; \
+		return __DPV_FOREACH_COMPONENTS_CONSTRUCTOR(TYPE, COMPONENTS, __DPV_CODE_FUNCTION_3_ARGS, clamp, x, min, max); \
 	} \
 	__DPV_MULTI_VECTOR(TYPE, COMPONENTS) dpv_min(__DPV_MULTI_VECTOR(TYPE, COMPONENTS) a, __DPV_MULTI_VECTOR(TYPE, COMPONENTS) b) { \
 		return min(a, b); \
 	} \
 	__DPV_MULTI_VECTOR(TYPE, COMPONENTS) dpv_min(__DPV_VECTOR(TYPE, COMPONENTS) a, __DPV_MULTI_VECTOR(TYPE, COMPONENTS) b) { \
-		__DPV_MULTI_VECTOR(TYPE, COMPONENTS) r; \
-		UNROLL_N(COMPONENTS) for (uint c = 0; c < COMPONENTS; c++) { \
-			r[c] = min(a[c], b[c]); \
-		} \
-		return r; \
+		return __DPV_FOREACH_COMPONENTS_CONSTRUCTOR(TYPE, COMPONENTS, __DPV_CODE_FUNCTION_2_ARGS, min, a, b); \
 	} \
 	__DPV_MULTI_VECTOR(TYPE, COMPONENTS) dpv_min(__DPV_MULTI_VECTOR(TYPE, COMPONENTS) a, __DPV_VECTOR(TYPE, COMPONENTS) b) { \
-		__DPV_MULTI_VECTOR(TYPE, COMPONENTS) r; \
-		UNROLL_N(COMPONENTS) for (uint c = 0; c < COMPONENTS; c++) { \
-			r[c] = min(a[c], b[c]); \
-		} \
-		return r; \
+		return __DPV_FOREACH_COMPONENTS_CONSTRUCTOR(TYPE, COMPONENTS, __DPV_CODE_FUNCTION_2_ARGS, min, a, b); \
 	} \
 	__DPV_MULTI_VECTOR(TYPE, COMPONENTS) dpv_max(__DPV_MULTI_VECTOR(TYPE, COMPONENTS) a, __DPV_MULTI_VECTOR(TYPE, COMPONENTS) b) { \
 		return max(a, b); \
 	} \
 	__DPV_MULTI_VECTOR(TYPE, COMPONENTS) dpv_max(__DPV_VECTOR(TYPE, COMPONENTS) a, __DPV_MULTI_VECTOR(TYPE, COMPONENTS) b) { \
-		__DPV_MULTI_VECTOR(TYPE, COMPONENTS) r; \
-		UNROLL_N(COMPONENTS) for (uint c = 0; c < COMPONENTS; c++) { \
-			r[c] = max(a[c], b[c]); \
-		} \
-		return r; \
+		return __DPV_FOREACH_COMPONENTS_CONSTRUCTOR(TYPE, COMPONENTS, __DPV_CODE_FUNCTION_2_ARGS, max, a, b); \
 	} \
 	__DPV_MULTI_VECTOR(TYPE, COMPONENTS) dpv_max(__DPV_MULTI_VECTOR(TYPE, COMPONENTS) a, __DPV_VECTOR(TYPE, COMPONENTS) b) { \
-		__DPV_MULTI_VECTOR(TYPE, COMPONENTS) r; \
-		UNROLL_N(COMPONENTS) for (uint c = 0; c < COMPONENTS; c++) { \
-			r[c] = max(a[c], b[c]); \
-		} \
-		return r; \
+		return __DPV_FOREACH_COMPONENTS_CONSTRUCTOR(TYPE, COMPONENTS, __DPV_CODE_FUNCTION_2_ARGS, max, a, b); \
 	} \
 
 DEFINE_DPV_FUNCTION_FOR_ALL_VECTORS(DEFINE_DPV_FUNCTION_VECTOR)
@@ -359,4 +424,5 @@ DEFINE_DPV_FUNCTION_FOR_ALL_VECTORS(DEFINE_DPV_FUNCTION_VECTOR)
 
 
 DEFINE_DPV_FUNCTION_FOR_FLOAT_VECTORS(DEFINE_DPV_FUNCTION_VECTOR)
+DEFINE_DPV_FUNCTION_FOR_HALF_VECTORS(DEFINE_DPV_FUNCTION_VECTOR)
 #undef DEFINE_DPV_FUNCTION_VECTOR
diff --git a/Engine/Shaders/Public/Platform.ush b/Engine/Shaders/Public/Platform.ush
--- a/Engine/Shaders/Public/Platform.ush
+++ b/Engine/Shaders/Public/Platform.ush
@@ -752,6 +752,95 @@ float4 max3( float4 a, float4 b, float4 c )
 	);
 }
 
+#if PLATFORM_SUPPORTS_REAL_TYPES
+
+half min3( half a, half b, half c )
+{
+	return min( a, min( b, c ) );
+}
+
+half max3( half a, half b, half c )
+{
+	return max( a, max( b, c ) );
+}
+
+half2 min3( half2 a, half2 b, half2 c )
+{
+	return half2(
+		min3( a.x, b.x, c.x ),
+		min3( a.y, b.y, c.y )
+	);
+}
+
+half2 max3( half2 a, half2 b, half2 c )
+{
+	return half2(
+		max3( a.x, b.x, c.x ),
+		max3( a.y, b.y, c.y )
+	);
+}
+
+half3 max3( half3 a, half3 b, half3 c )
+{
+	return half3(
+		max3( a.x, b.x, c.x ),
+		max3( a.y, b.y, c.y ),
+		max3( a.z, b.z, c.z )
+	);
+}
+
+half3 min3( half3 a, half3 b, half3 c )
+{
+	return half3(
+		min3( a.x, b.x, c.x ),
+		min3( a.y, b.y, c.y ),
+		min3( a.z, b.z, c.z )
+	);
+}
+
+half4 min3( half4 a, half4 b, half4 c )
+{
+	return half4(
+		min3( a.x, b.x, c.x ),
+		min3( a.y, b.y, c.y ),
+		min3( a.z, b.z, c.z ),
+		min3( a.w, b.w, c.w )
+	);
+}
+
+half4 max3( half4 a, half4 b, half4 c )
+{
+	return half4(
+		max3( a.x, b.x, c.x ),
+		max3( a.y, b.y, c.y ),
+		max3( a.z, b.z, c.z ),
+		max3( a.w, b.w, c.w )
+	);
+}
+
+#endif
+
+#endif
+
+#if PLATFORM_SUPPORTS_REAL_TYPES && !defined(COMPILER_SUPPORTS_PACK_B32_B16)
+
+// Function that explicitly use RDNA's v_pack_b32_f16 on supported platform. Note that RDNA's documentation call this instruction v_pack_b32_f16
+// but really is a v_pack_b32_b16.
+half2 v_pack_b32_b16(half a, half b)
+{
+	return half2(a, b);
+}
+
+int16_t2 v_pack_b32_b16(int16_t a, int16_t b)
+{
+	return int16_t2(a, b);
+}
+
+uint16_t2 v_pack_b32_b16(uint16_t a, uint16_t b)
+{
+	return uint16_t2(a, b);
+}
+
 #endif
 
 #if !defined(COMPILER_SUPPORTS_COND_MASK)
