Fixes black ghosting on parallax rejection

diff --git a/Engine/Shaders/Private/TemporalSuperResolution/TSRUpdateHistory.usf b/Engine/Shaders/Private/TemporalSuperResolution/TSRUpdateHistory.usf
--- a/Engine/Shaders/Private/TemporalSuperResolution/TSRUpdateHistory.usf
+++ b/Engine/Shaders/Private/TemporalSuperResolution/TSRUpdateHistory.usf
@@ -458,9 +458,19 @@ void MainCS(
 			// Soften the spatial kernel when there is high noise in the input to make parrallax disocclusion of detailed geometry less distracting.
 			taa_half2 OffScreenInputToHistoryFactor = taa_half(1.0) - taa_half(0.5) * NoiseFiltering;
 
-			taa_half2 KernelInputToHistoryFactor = taa_half2(
-				bOffScreen.x ? OffScreenInputToHistoryFactor.x : taa_half(InputToHistoryFactor),
-				bOffScreen.y ? OffScreenInputToHistoryFactor.y : taa_half(InputToHistoryFactor));
+			taa_half2 KernelInputToHistoryFactor = OffScreenInputToHistoryFactor;
+			
+			FLATTEN
+			if (bOffScreen.x)
+			{
+				KernelInputToHistoryFactor.x = taa_half(InputToHistoryFactor);
+			}
+			
+			FLATTEN
+			if (bOffScreen.y)
+			{
+				KernelInputToHistoryFactor.y = taa_half(InputToHistoryFactor);
+			}
 
 			UNROLL_N(CONFIG_SAMPLES)
 			for (uint SampleId = 0; SampleId < CONFIG_SAMPLES; SampleId++)
@@ -742,6 +752,9 @@ void MainCS(
 			PrevHistoryMommentMax *= taa_half(HistoryPreExposureCorrection);
 			PrevHighFrequencyColor *= taa_half(HistoryPreExposureCorrection);
 			PrevLowFrequencyColor *= taa_half(HistoryPreExposureCorrection);
+
+			// Fixes negative validity that can happen due to negative lobes of Catrom.
+			PrevHistoryValidity = max(PrevHistoryValidity, taa_half(0.0).xx);
 		}
 	}
 		
@@ -958,7 +971,7 @@ void MainCS(
 			taa_half2 BlendPrevHistory = PrevMomentWeight * PrevHistoryToneWeight;
 			taa_half2 BlendFilteredInput = CurrentMomentWeight * FilteredInputToneWeight;
 
-			taa_half2 CommonWeight = OutputValidity * SafeRcp(BlendPrevHistory + BlendFilteredInput);
+			taa_half2 CommonWeight = SafeRcp(BlendPrevHistory + BlendFilteredInput);
 
 			FinalHighFrequencyColor = (
 				dpv_scale(PrevHighFrequencyColor, CommonWeight * BlendPrevHistory * InvPrevHistoryValidity) +
@@ -970,21 +983,16 @@ void MainCS(
 		}
 		#endif
 
-		// Quantize validity for the 8bit encoding an adjust to avoid numerical shift.
+		// Quantize validity for the 8bit encoding to avoid numerical shift between color and validity.
 		#if 0
 			FinalHistoryValidity = OutputValidity;
 		#else
-			taa_half2 OutputInvValidity = SafeRcp(OutputValidity);
 			FinalHistoryValidity = ceil(taa_half(255.0) * OutputValidity) * rcp(taa_half(255.0));
-			FinalHighFrequencyColor = dpv_scale(FinalHighFrequencyColor, FinalHistoryValidity * OutputInvValidity);
 		#endif
 
-		//Debug.x = RejectionConfidentEnough;
-		//Debug.x = RejectionValidity;
-
 		#if 0 // DEBUG
 		{
-			FinalHighFrequencyColor = PrevHighFrequencyColor;
+			FinalHighFrequencyColor = dpv_scale(PrevHighFrequencyColor, InvPrevHistoryValidity);
 			FinalHistoryValidity = PrevHistoryValidity;
 		}
 		#endif
@@ -1242,10 +1250,8 @@ void MainCS(
 
 		FinalOutputValidity = lerp(taa_half(1.0), saturate(FinalHistoryValidity), OutputBlend);
 
-		taa_half3x2 NormalizedFinalHighFrequencyColor = dpv_scale(FinalHighFrequencyColor, SafeRcp(FinalHistoryValidity));
-
 		taa_half2 LowFrequencyWeight = HdrWeight4(FinalLowFrequencyColor);
-		taa_half2 Moment1Weight = HdrWeight4(NormalizedFinalHighFrequencyColor);
+		taa_half2 Moment1Weight = HdrWeight4(FinalHighFrequencyColor);
 
 		taa_half2x2 Weights = WeightedLerpFactors(LowFrequencyWeight, Moment1Weight, OutputBlend);
 
@@ -1253,13 +1259,13 @@ void MainCS(
 			taa_half3 LowFrequencyColor = taa_half3(1, 0.25, 0.25);
 			taa_half3 HighFrequencyColor = taa_half3(0.25, 1, 0.25);
 
-			FinalOutputColor = dpv_scale(FinalLowFrequencyColor * dpv_interleave_mono_registers(LowFrequencyColor), Weights[0]) + dpv_scale(NormalizedFinalHighFrequencyColor * dpv_interleave_mono_registers(HighFrequencyColor), Weights[1]);
+			FinalOutputColor = dpv_scale(FinalLowFrequencyColor * dpv_interleave_mono_registers(LowFrequencyColor), Weights[0]) + dpv_scale(FinalHighFrequencyColor * dpv_interleave_mono_registers(HighFrequencyColor), Weights[1]);
 		#elif DEBUG_LOW_FREQUENCY_REJECTION
 			taa_half3 DebugColor = lerp(taa_half3(1, 0.5, 0.5), taa_half3(0.5, 1, 0.5), LowFrequencyRejection);
 			
-			FinalOutputColor = FinalLowFrequencyColor * Weights.x * DebugColor + NormalizedFinalHighFrequencyColor * Weights.y * DebugColor;
+			FinalOutputColor = FinalLowFrequencyColor * Weights.x * DebugColor + FinalHighFrequencyColor * Weights.y * DebugColor;
 		#else
-			FinalOutputColor = dpv_scale(FinalLowFrequencyColor, Weights[0]) + dpv_scale(NormalizedFinalHighFrequencyColor, Weights[1]);
+			FinalOutputColor = dpv_scale(FinalLowFrequencyColor, Weights[0]) + dpv_scale(FinalHighFrequencyColor, Weights[1]);
 		#endif
 
 		#if CONFIG_ACCUMULATE_TRANSLUCENCY_SEPARATLY
@@ -1283,6 +1289,9 @@ void MainCS(
 			
 		LocalHistoryPixelPos = InvalidateOutputPixelPos(LocalHistoryPixelPos, HistoryInfo_ViewportMax);
 		
+		// Premultiply FinalHighFrequencyColor to have validity weighted average in next frame history reprojection.
+		taa_half3x2 PremultipliedHighFrequencyColor = dpv_scale(FinalHighFrequencyColor, FinalHistoryValidity);
+
 		// Stocastically round up or down using the hardware RWTexture2D truncation unit to take into precision
 		// loss due to pixel format encoding.
 		#if CONFIG_ENABLE_STOCASTIC_QUANTIZATION
@@ -1290,24 +1299,29 @@ void MainCS(
 			uint2 Random = Rand3DPCG16(int3(dpv_lo(LocalHistoryPixelPos), View.StateFrameIndexMod8)).xy;
 			float E = Hammersley16(0, 1, Random).x;
 
-			FinalHighFrequencyColor = QuantizeForFloatRenderTarget(FinalHighFrequencyColor, E, HistoryQuantizationError);
+			PremultipliedHighFrequencyColor = QuantizeForFloatRenderTarget(PremultipliedHighFrequencyColor, E, HistoryQuantizationError);
 			FinalLowFrequencyColor = QuantizeForFloatRenderTarget(FinalLowFrequencyColor, E, HistoryQuantizationError);
 		}
 		#endif
 
-		FinalLowFrequencyColor = -dpv_min(-FinalLowFrequencyColor, taa_half(0.0).xxx);
-		FinalHighFrequencyColor = -dpv_min(-FinalHighFrequencyColor, taa_half(0.0).xxx);
-		FinalLowFrequencyColor = dpv_min(FinalLowFrequencyColor, taa_half(Max10BitsFloat).xxx);
-		FinalHighFrequencyColor = dpv_min(FinalHighFrequencyColor, taa_half(Max10BitsFloat).xxx);		
-		#if CONFIG_ACCUMULATE_TRANSLUCENCY_SEPARATLY
-			FinalTranslucencyColor = -dpv_min(-FinalTranslucencyColor, taa_half(0.0).xxxx);
-			FinalTranslucencyColor = dpv_min(FinalTranslucencyColor, taa_half(Max10BitsFloat).xxxx);
-		#endif
+		// Protext from NaN and +Inf when writing out the history.
+		{
+			FinalLowFrequencyColor = -dpv_min(-FinalLowFrequencyColor, taa_half(0.0).xxx);
+			FinalLowFrequencyColor = dpv_min(FinalLowFrequencyColor, taa_half(Max10BitsFloat).xxx);
+
+			PremultipliedHighFrequencyColor = -dpv_min(-PremultipliedHighFrequencyColor, taa_half(0.0).xxx);
+			PremultipliedHighFrequencyColor = dpv_min(PremultipliedHighFrequencyColor, taa_half(Max10BitsFloat).xxx);	
+
+			#if CONFIG_ACCUMULATE_TRANSLUCENCY_SEPARATLY
+				FinalTranslucencyColor = -dpv_min(-FinalTranslucencyColor, taa_half(0.0).xxxx);
+				FinalTranslucencyColor = dpv_min(FinalTranslucencyColor, taa_half(Max10BitsFloat).xxxx);
+			#endif
+		}
 
 		// Output final history lo pixel.
 		{
 			HistoryOutput_LowFrequency[dpv_lo(LocalHistoryPixelPos)]  = dpv_lo(FinalLowFrequencyColor);
-			HistoryOutput_HighFrequency[dpv_lo(LocalHistoryPixelPos)] = dpv_lo(FinalHighFrequencyColor);
+			HistoryOutput_HighFrequency[dpv_lo(LocalHistoryPixelPos)] = dpv_lo(PremultipliedHighFrequencyColor);
 			HistoryOutput_Metadata[dpv_lo(LocalHistoryPixelPos)]      = taa_half2(dpv_lo(FinalLowFrequencyWeight), dpv_lo(FinalHistoryValidity));
 			#if CONFIG_ACCUMULATE_TRANSLUCENCY_SEPARATLY
 				HistoryOutput_Translucency[dpv_lo(LocalHistoryPixelPos)] = dpv_lo(FinalTranslucencyColor).rgb;
@@ -1320,7 +1334,7 @@ void MainCS(
 		#if CONFIG_ENABLE_DUAL_PIXEL_VECTORIZATION
 		{
 			HistoryOutput_LowFrequency[dpv_hi(LocalHistoryPixelPos)]  = dpv_hi(FinalLowFrequencyColor);
-			HistoryOutput_HighFrequency[dpv_hi(LocalHistoryPixelPos)] = dpv_hi(FinalHighFrequencyColor);
+			HistoryOutput_HighFrequency[dpv_hi(LocalHistoryPixelPos)] = dpv_hi(PremultipliedHighFrequencyColor);
 			HistoryOutput_Metadata[dpv_hi(LocalHistoryPixelPos)]      = taa_half2(dpv_hi(FinalLowFrequencyWeight), dpv_hi(FinalHistoryValidity));
 			#if CONFIG_ACCUMULATE_TRANSLUCENCY_SEPARATLY
 				HistoryOutput_Translucency[dpv_hi(LocalHistoryPixelPos)] = dpv_hi(FinalTranslucencyColor).rgb;
