Fixes stability regression on trees in TSR's explicit texture fetches overlaping

diff --git a/Engine/Shaders/Private/TemporalSuperResolution/TSRUpdateHistory.usf b/Engine/Shaders/Private/TemporalSuperResolution/TSRUpdateHistory.usf
--- a/Engine/Shaders/Private/TemporalSuperResolution/TSRUpdateHistory.usf
+++ b/Engine/Shaders/Private/TemporalSuperResolution/TSRUpdateHistory.usf
@@ -255,18 +255,18 @@ void MainCS(
 		uint RawParallaxFactorBits[DPV_PIXEL_PER_LANE];
 		half RawLowFrequencyRejection[DPV_PIXEL_PER_LANE];
 		half RawTranslucencyRejection[DPV_PIXEL_PER_LANE];
-		half RawTranslucencyRejectionNeighborhood[4][DPV_PIXEL_PER_LANE];
-		half RawMoireWeightNeighborhood[4][DPV_PIXEL_PER_LANE]; // TODO(TSR): bug
+		half RawTranslucencyRejectionNeighborhood[CONFIG_SAMPLES - 1][DPV_PIXEL_PER_LANE];
+		half RawMoireWeightNeighborhood[CONFIG_SAMPLES - 1][DPV_PIXEL_PER_LANE]; // TODO(TSR): bug
 		
 		uint RawSceneStencilRef[DPV_PIXEL_PER_LANE];
-		uint RawSceneStencilRefNeighborhood[4][DPV_PIXEL_PER_LANE];
+		uint RawSceneStencilRefNeighborhood[CONFIG_SAMPLES - 1][DPV_PIXEL_PER_LANE];
 		
 		// Issue overlapped texture fetches
 		PLATFORM_SPECIFIC_ISOLATE
 		{
 			taa_ushort2x2 RejectionPixelPos = dpv_sub(taa_ushort2x2(InputPixelPos), taa_ushort2(InputPixelPosMin)) / 2;
 
-			UNROLL_N(2)
+			UNROLL_N(DPV_PIXEL_PER_LANE)
 			for (uint PixelId = 0; PixelId < DPV_PIXEL_PER_LANE; PixelId++)
 			{
 				RawEncodedVelocity[PixelId]       = DilatedVelocityTexture[dpv_access_pixel(InputPixelPos, PixelId)];
@@ -291,14 +291,14 @@ void MainCS(
 					/* out */ SampleInputPixelPos,
 					/* out */ PixelOffset);
 			
-				UNROLL_N(2)
+				UNROLL_N(DPV_PIXEL_PER_LANE)
 				for (uint PixelId = 0; PixelId < DPV_PIXEL_PER_LANE; PixelId++)
 				{
 					RawTranslucencyRejectionNeighborhood[NeighborId][PixelId] = TranslucencyRejectionTexture[dpv_access_pixel(SampleInputPixelPos, PixelId)];
 					RawSceneStencilRefNeighborhood[NeighborId][PixelId] = InputSceneStencilTexture.Load(int3(dpv_access_pixel(SampleInputPixelPos, PixelId), 0)) STENCIL_COMPONENT_SWIZZLE;
 			
 					#if DIM_MOIRE_HEURISTIC
-						RawMoireWeightNeighborhood[NeighborId][PixelId] = InterferenceWeightTexture[dpv_lo(SampleInputPixelPos)];
+						RawMoireWeightNeighborhood[NeighborId][PixelId] = InterferenceWeightTexture[dpv_access_pixel(SampleInputPixelPos, PixelId)];
 					#endif
 				}
 			} // for (uint NeighborId = 0; NeighborId < (CONFIG_SAMPLES - 1); NeighborId++)
@@ -321,7 +321,7 @@ void MainCS(
 
 			// Process input texel, forcing to pack register.
 			ParallaxRejectionMask = dpv_force_interleave_registers_array(RawParallaxRejectionMask);
-			ParallaxFactorBits    = dpv_force_interleave_registers_array(RawParallaxFactorBits);
+			ParallaxFactorBits    = dpv_force_interleave_registers_array(RawParallaxFactorBits) & taa_subpixel_payload(SUB_PIXEL_PARALLAX_FACTOR_BIT_MASK);
 			LowFrequencyRejection = dpv_force_interleave_registers_array(RawLowFrequencyRejection);
 			TranslucencyRejection = dpv_force_interleave_registers_array(RawTranslucencyRejection);
 			SceneStencilRef       = dpv_force_interleave_registers_array(RawSceneStencilRef);
@@ -488,7 +488,7 @@ void MainCS(
 				RawdPPArray[SampleId] = dPP; 
 				RawSampleSpatialWeightArray[SampleId] = SampleSpatialWeight;
 
-				UNROLL_N(2)
+				UNROLL_N(DPV_PIXEL_PER_LANE)
 				for (uint PixelId = 0; PixelId < DPV_PIXEL_PER_LANE; PixelId++)
 				{
 					RawInputColorArray[SampleId][PixelId] = InputSceneColorTexture[dpv_access_pixel(SampleInputPixelPos, PixelId)];
