Increase max number of hair instance.

diff --git a/Engine/Shaders/Private/HairStrands/HairStrandsVisibilityCommon.ush b/Engine/Shaders/Private/HairStrands/HairStrandsVisibilityCommon.ush
--- a/Engine/Shaders/Private/HairStrands/HairStrandsVisibilityCommon.ush
+++ b/Engine/Shaders/Private/HairStrands/HairStrandsVisibilityCommon.ush
@@ -151,17 +151,17 @@ uint PackHairVisDepthCoverage(float InDepth, uint InCoverage8bits)
 
 uint PackHairVisPrimitiveMaterialId(uint PrimitiveId, uint MaterialId)
 {
-	return (PrimitiveId & 0x00FFFFFF) | ((MaterialId & 0xFF) << 24);
+	return (PrimitiveId & 0x007FFFFF) | ((MaterialId & 0x1FF) << 23);
 }
 
 uint UnpackHairVisPrimitiveId(uint RawVis)
 {
-	return RawVis & 0x00FFFFFF;
+	return RawVis & 0x007FFFFF;
 }
 
 uint UnpackHairVisMaterialId(uint RawVis)
 {
-	return (RawVis >> 24) & 0xFF;
+	return (RawVis >> 23) & 0x1FF;
 }
 
 float UnpackHairVisDepth(uint InPacked)
