Restructures motion blur shader files for better maintainability of dual motion blur directions change

This change is only move code arround to make the following changes easier to diff.

diff --git a/Engine/Shaders/Private/MotionBlur/MotionBlurApply.usf b/Engine/Shaders/Private/MotionBlur/MotionBlurApply.usf
new file mode 100644
--- /dev/null
+++ b/Engine/Shaders/Private/MotionBlur/MotionBlurApply.usf
@@ -0,0 +1,326 @@
+// Copyright Epic Games, Inc. All Rights Reserved.
+
+#include "MotionBlurCommon.ush"
+#include "../Random.ush"
+
+
+//------------------------------------------------------- CONSTANTS
+
+// from the paper: We use SOFT Z EXTENT = 1mm to 10cm for our results
+#define SOFT_Z_EXTENT 1
+
+#define MOTION_BLUR_QUALITY_LOW 0
+#define MOTION_BLUR_QUALITY_MEDIUM 1
+#define MOTION_BLUR_QUALITY_HIGH 2
+#define MOTION_BLUR_QUALITY_VERY_HIGH 2
+
+
+//------------------------------------------------------- PARAMETERS
+
+float MotionBlur_AspectRatio;
+float MotionBlur_VelocityScale;
+float MotionBlur_VelocityScaleForTiles;
+float MotionBlur_VelocityMax;
+
+SCREEN_PASS_TEXTURE_VIEWPORT(Color)
+SCREEN_PASS_TEXTURE_VIEWPORT(Velocity)
+SCREEN_PASS_TEXTURE_VIEWPORT(VelocityTile)
+FScreenTransform ColorToVelocity;
+
+Texture2D ColorTexture;
+Texture2D VelocityTexture;
+Texture2D VelocityTileTexture;
+Texture2D VelocityFlatTexture;
+Texture2D DepthTexture;
+
+SamplerState ColorSampler;
+SamplerState VelocitySampler;
+SamplerState VelocityTileSampler;
+SamplerState VelocityFlatSampler;
+SamplerState DepthSampler;
+
+#if USE_POST_MOTION_BLUR_TRANSLUCENCY
+
+Texture2D PostMotionBlurTranslucencyTexture;
+SamplerState PostMotionBlurTranslucencySampler;
+
+float2 ScreenPosToPostMotionBlurTranslucencyUV;
+float2 PostMotionBlurTranslucencyUVMax;
+
+#endif // USE_POST_MOTION_BLUR_TRANSLUCENCY
+
+#if SUPPORTS_INDEPENDENT_SAMPLERS
+	#define SharedVelocitySampler VelocitySampler
+	#define SharedVelocityTileSampler VelocitySampler
+	#define SharedVelocityFlatSampler VelocitySampler
+	#define SharedDepthSampler VelocitySampler
+#else
+	#define SharedVelocitySampler VelocitySampler
+	#define SharedVelocityTileSampler VelocityTileSampler
+	#define SharedVelocityFlatSampler VelocityFlatSampler
+	#define SharedDepthSampler DepthSampler
+#endif
+
+//------------------------------------------------------- LDS
+
+#if COMPUTESHADER && !COMPILER_SUPPORTS_WAVE_VOTE
+
+groupshared uint GroupSharedFastPath;
+
+#endif
+
+
+//------------------------------------------------------- FUNCTIONS
+
+uint GetStepCountFromQuality()
+{
+#if MOTION_BLUR_QUALITY == MOTION_BLUR_QUALITY_LOW
+	return 4;
+#elif MOTION_BLUR_QUALITY == MOTION_BLUR_QUALITY_MEDIUM
+	return 6;
+#elif MOTION_BLUR_QUALITY == MOTION_BLUR_QUALITY_HIGH
+	return 8;
+#else // MOTION_BLUR_QUALITY == MOTION_BLUR_QUALITY_VERY_HIGH
+	return 16;
+#endif
+}
+
+float2 DepthCmp(float CenterDepth, float SampleDepth, float DepthScale)
+{
+	return saturate(0.5 + float2(DepthScale, -DepthScale) * (SampleDepth - CenterDepth));
+}
+
+float2 SpreadCmp(float OffsetLength, float2 SpreadLength, float PixelToSampleScale)
+{
+	return saturate(PixelToSampleScale * SpreadLength - max(OffsetLength - 1, 0));
+}
+
+float SampleWeight(float CenterDepth, float SampleDepth, float OffsetLength, float CenterSpreadLength, float SampleSpreadLength, float PixelToSampleScale, float DepthScale)
+{
+	float2 DepthWeights = DepthCmp(CenterDepth, SampleDepth, DepthScale);
+	float2 SpreadWeights = SpreadCmp(OffsetLength, float2(CenterSpreadLength, SampleSpreadLength), PixelToSampleScale);
+	return dot(DepthWeights, SpreadWeights);
+}
+
+float2 DecodeVelocity(float2 Velocity)
+{
+	// 11:11:10  (VelocityLength, VelocityAngle, Depth)
+	float VelocityLength = Velocity.x;
+	float VelocityAngle = Velocity.y * (2 * PI) - PI;
+	sincos(VelocityAngle, Velocity.x, Velocity.y);
+	Velocity *= VelocityLength;
+	return Velocity;
+}
+
+float GetVelocityLengthPixels(float2 EncodedVelocity)
+{
+	// 11:11:10  (VelocityLength, VelocityAngle, Depth)
+	float VelocityLength = EncodedVelocity.x;
+	VelocityLength *= MotionBlur_VelocityScale;
+	return min(VelocityLength, MotionBlur_VelocityMax);
+}
+
+
+//------------------------------------------------------- ENTRY POINT
+
+float4 MainMotionBlurCommon(float2 ColorUV, float2 ColorPixelPos)
+{
+	const uint StepCount = GetStepCountFromQuality() / 2;
+
+	const float PixelToTileScale = (1.0 / 16.0);
+
+	float Random  = InterleavedGradientNoise(ColorPixelPos, 0);
+	float Random2 = InterleavedGradientNoise(ColorPixelPos, 1);
+
+	// [-0.25, 0.25]
+	float2 TileJitter = (float2(Random, Random2) - 0.5) * 0.5;
+
+	// Map color UV to velocity UV space.
+	float2 VelocityUV = ApplyScreenTransform(ColorUV, ColorToVelocity);
+
+	// Map velocity UV to velocity tile UV space with jitter.
+	float2 NearestVelocitySvPosition = floor(VelocityUV * Velocity_Extent) + 0.5;
+	float2 VelocityTileUV = ((NearestVelocitySvPosition - Velocity_ViewportMin) * PixelToTileScale + TileJitter) * VelocityTile_ExtentInverse;
+
+	// Velocity tile UV originates at [0,0]; only need to clamp max.
+	VelocityTileUV = min(VelocityTileUV, VelocityTile_UVViewportBilinearMax);
+
+	float4 MinMaxVelocity = VelocityTileTexture.SampleLevel(SharedVelocityTileSampler, VelocityTileUV, 0);
+
+	float2 MinVelocityPixels = MinMaxVelocity.xy * MotionBlur_VelocityScale;
+	float2 MaxVelocityPixels = MinMaxVelocity.zw * MotionBlur_VelocityScale;
+
+	float MinVelocityLengthSqrPixels = dot(MinVelocityPixels, MinVelocityPixels);
+	float MaxVelocityLengthSqrPixels = dot(MaxVelocityPixels, MaxVelocityPixels);
+
+	// Input buffer 0 as same viewport as output buffer.
+	float4 CenterColor = ColorTexture.SampleLevel(ColorSampler, ColorUV, 0);
+
+	float4 SearchVectorPixels = float4(MaxVelocityPixels, -MaxVelocityPixels);
+	float4 SearchVector = SearchVectorPixels * Color_ExtentInverse.xyxy;
+
+	// converts pixel length to sample steps
+	float PixelToSampleScale = StepCount * rsqrt(dot(MaxVelocityPixels, MaxVelocityPixels));	
+		
+	#if USE_POST_MOTION_BLUR_TRANSLUCENCY
+		const float2 PostMotionBlurTranslucencyUV = min((ColorPixelPos + (float2)0.5f) * ScreenPosToPostMotionBlurTranslucencyUV, PostMotionBlurTranslucencyUVMax);
+		float4 PostMotionBlurTranslucency = PostMotionBlurTranslucencyTexture.SampleLevel(PostMotionBlurTranslucencySampler, PostMotionBlurTranslucencyUV, 0);
+	#endif
+
+	
+	// TODO expose cvars
+	bool bSkipPath = MaxVelocityLengthSqrPixels < 0.25;
+	bool bFastPath = MinVelocityLengthSqrPixels > 0.4 * MaxVelocityLengthSqrPixels;
+
+	// Only use fast path if all threads of the compute shader would.
+	#if COMPILER_SUPPORTS_WAVE_VOTE
+	{
+		bFastPath = WaveActiveAllTrue(bFastPath);
+	}
+	#elif COMPUTESHADER
+	{
+		GroupSharedFastPath = 0;
+		GroupMemoryBarrierWithGroupSync();
+
+		uint IgnoredOut;
+		InterlockedAdd(GroupSharedFastPath, bFastPath ? 1 : 0);
+		
+		GroupMemoryBarrierWithGroupSync();
+		bFastPath = (GroupSharedFastPath == (THREADGROUP_SIZEX * THREADGROUP_SIZEY));
+	}
+	#endif
+	
+	BRANCH
+	if (bSkipPath)
+	{
+		#if USE_POST_MOTION_BLUR_TRANSLUCENCY
+			CenterColor.rgb = CenterColor.rgb * PostMotionBlurTranslucency.a + PostMotionBlurTranslucency.rgb;
+		#endif
+		#if !POST_PROCESS_ALPHA
+			CenterColor.a = 0;
+		#endif
+		return CenterColor;
+	}
+
+	float4 OutColor = 0;
+
+	BRANCH
+	if (bFastPath)
+	{
+		float4 ColorAccum = 0;
+
+		UNROLL
+		for (uint i = 0; i < StepCount; i++)
+		{
+			float2 OffsetLength = (float)i + 0.5 + float2(Random - 0.5, 0.5 - Random);
+			float2 OffsetFraction = OffsetLength / StepCount;
+
+			float2 SampleUV[2];
+			SampleUV[0] = ColorUV + OffsetFraction.x * SearchVector.xy;
+			SampleUV[1] = ColorUV + OffsetFraction.y * SearchVector.zw;
+
+			SampleUV[0] = clamp(SampleUV[0], Color_UVViewportBilinearMin, Color_UVViewportBilinearMax);
+			SampleUV[1] = clamp(SampleUV[1], Color_UVViewportBilinearMin, Color_UVViewportBilinearMax);
+
+			ColorAccum += ColorTexture.SampleLevel(ColorSampler, SampleUV[0], 0);
+			ColorAccum += ColorTexture.SampleLevel(ColorSampler, SampleUV[1], 0);
+		}
+	
+		ColorAccum *= 0.5 / StepCount;
+
+		OutColor = ColorAccum;
+	}
+	else
+	{
+		float3 CenterVelocityDepth = VelocityFlatTexture.SampleLevel(SharedVelocityFlatSampler, VelocityUV, 0).xyz;
+		float  CenterDepth = CenterVelocityDepth.z;
+		float  CenterVelocityLength = GetVelocityLengthPixels(CenterVelocityDepth.xy);
+
+		float4 ColorAccum = 0;
+		float  ColorAccumWeight = 0;
+
+		UNROLL
+		for (uint i = 0; i < StepCount; i++)
+		{
+			float2 SampleUV[2];
+			float4 SampleColor[2];
+			float  SampleDepth[2];
+			float  SampleVelocityLength[2];
+			float  Weight[2];
+
+			float2 OffsetLength = (float)i + 0.5 + float2(Random - 0.5, 0.5 - Random);
+			float2 OffsetFraction = OffsetLength / StepCount;
+
+			SampleUV[0] = ColorUV + OffsetFraction.x * SearchVector.xy;
+			SampleUV[1] = ColorUV + OffsetFraction.y * SearchVector.zw;
+
+			SampleUV[0] = clamp(SampleUV[0], Color_UVViewportBilinearMin, Color_UVViewportBilinearMax);
+			SampleUV[1] = clamp(SampleUV[1], Color_UVViewportBilinearMin, Color_UVViewportBilinearMax);
+
+			UNROLL
+			for (uint j = 0; j < 2; j++)
+			{
+				float3 SampleVelocityDepth = VelocityFlatTexture.SampleLevel(
+					SharedVelocityFlatSampler, ApplyScreenTransform(SampleUV[j], ColorToVelocity), 0).xyz;
+
+				SampleColor[j] = ColorTexture.SampleLevel(ColorSampler, SampleUV[j], 0);
+				SampleDepth[j] = SampleVelocityDepth.z;
+
+				// in pixels
+				SampleVelocityLength[j] = GetVelocityLengthPixels(SampleVelocityDepth.xy);
+
+				Weight[j] = SampleWeight(CenterDepth, SampleDepth[j], OffsetLength.x, CenterVelocityLength, SampleVelocityLength[j], PixelToSampleScale, SOFT_Z_EXTENT);
+			}
+		
+			bool2 Mirror = bool2(SampleDepth[0] > SampleDepth[1], SampleVelocityLength[1] > SampleVelocityLength[0]);
+			Weight[0] = all(Mirror) ? Weight[1] : Weight[0];
+			Weight[1] = any(Mirror) ? Weight[1] : Weight[0];
+			
+			ColorAccum += Weight[0] * SampleColor[0] + Weight[1] * SampleColor[1];
+			ColorAccumWeight += Weight[0] + Weight[1];
+		}
+	
+		ColorAccum			*= 0.5 / StepCount;
+		ColorAccumWeight	*= 0.5 / StepCount;
+
+		OutColor = ColorAccum + (1 - ColorAccumWeight) * CenterColor;
+	}
+	
+#if USE_POST_MOTION_BLUR_TRANSLUCENCY
+	OutColor.rgb = OutColor.rgb * PostMotionBlurTranslucency.a + PostMotionBlurTranslucency.rgb;
+#endif
+
+#if !POST_PROCESS_ALPHA
+	OutColor.a = 0;
+#endif
+
+	return OutColor;
+}
+
+void MainPS(
+	in noperspective float4 UVAndScreenPos : TEXCOORD0,
+	in float4 SvPosition : SV_Position,
+	out float4 OutColor : SV_Target0)
+{
+	float2 ColorUV = UVAndScreenPos.xy;
+	OutColor = MainMotionBlurCommon(ColorUV, floor(SvPosition.xy));
+}
+
+#if COMPUTESHADER
+RWTexture2D<float4> OutColorTexture;
+
+[numthreads(THREADGROUP_SIZEX, THREADGROUP_SIZEY, 1)]
+void MainCS(uint2 DispatchThreadId : SV_DispatchThreadID)
+{
+	uint2 iColorPixelPos = DispatchThreadId + Color_ViewportMin;
+	float2 ColorPixelPos = float2(iColorPixelPos);
+	float2 ColorUV = (ColorPixelPos + 0.5) * Color_ExtentInverse;
+
+	float4 OutColor = MainMotionBlurCommon(ColorUV, ColorPixelPos);
+
+	if (all(iColorPixelPos < Color_ViewportMax))
+	{
+		OutColorTexture[iColorPixelPos] = OutColor;
+	}
+}
+#endif
diff --git a/Engine/Shaders/Private/MotionBlur/MotionBlurCommon.ush b/Engine/Shaders/Private/MotionBlur/MotionBlurCommon.ush
new file mode 100644
--- /dev/null
+++ b/Engine/Shaders/Private/MotionBlur/MotionBlurCommon.ush
@@ -0,0 +1,37 @@
+// Copyright Epic Games, Inc. All Rights Reserved.
+
+#include "../Common.ush"
+#include "../FastMath.ush"
+#include "../ScreenPass.ush"
+
+
+//------------------------------------------------------- FUNCTIONS
+
+float4 MinMaxLength(float4 v0, float4 v1)
+{
+	float2 Min = dot(v0.xy, v0.xy) < dot(v1.xy, v1.xy) ? v0.xy : v1.xy;
+	float2 Max = dot(v0.zw, v0.zw) > dot(v1.zw, v1.zw) ? v0.zw : v1.zw;
+	return float4(Min, Max);
+}
+
+float4 MinMaxLengthPolar(float4 v0, float4 v1)
+{
+	float2 Min = v0.x < v1.x ? v0.xy : v1.xy;
+	float2 Max = v0.z > v1.z ? v0.zw : v1.zw;
+	return float4(Min, Max);
+}
+
+float2 CartesianToPolar(float2 Velocity)
+{
+	float Length = length(Velocity);
+	float Angle = atan2Fast(Velocity.y, Velocity.x);
+	return float2(Length, Angle);
+}
+
+float2 PolarToCartesian(float2 Velocity)
+{
+	float Length = Velocity.x;
+	float Angle = Velocity.y;
+	sincos(Angle, Velocity.y, Velocity.x);
+	return Velocity * Length;
+}
diff --git a/Engine/Shaders/Private/MotionBlur/MotionBlurTileGather.usf b/Engine/Shaders/Private/MotionBlur/MotionBlurTileGather.usf
new file mode 100644
--- /dev/null
+++ b/Engine/Shaders/Private/MotionBlur/MotionBlurTileGather.usf
@@ -0,0 +1,91 @@
+// Copyright Epic Games, Inc. All Rights Reserved.
+
+#include "MotionBlurCommon.ush"
+
+
+//------------------------------------------------------- CONFIG
+
+#define THREADGROUP_SIZEX		16
+#define THREADGROUP_SIZEY		16
+#define THREADGROUP_TOTALSIZE	(THREADGROUP_SIZEX * THREADGROUP_SIZEY)
+
+
+//------------------------------------------------------- PARAMETERS
+
+SCREEN_PASS_TEXTURE_VIEWPORT(VelocityTile)
+float MotionBlur_VelocityScaleForTiles;
+
+Texture2D VelocityTileTexture;
+
+RWTexture2D<float4>	OutVelocityTileTexture;
+
+
+//------------------------------------------------------- ENTRY POINT
+
+[numthreads(16, 16, 1)]
+void MainCS(
+	uint3 GroupId : SV_GroupID,
+	uint3 DispatchThreadId : SV_DispatchThreadID,
+	uint3 GroupThreadId : SV_GroupThreadID,
+	uint GroupIndex : SV_GroupIndex)
+{
+	uint2 PixelPos = DispatchThreadId.xy;
+	float2 UV = ((float2)PixelPos.xy + 0.5) * VelocityTile_ExtentInverse;
+
+	float4 MinMaxVelocity = VelocityTileTexture[PixelPos];
+
+	// Scatter as gather
+	for(int x = -3; x <= 3; x++)
+	{
+		for(int y = -3; y <= 3; y++)
+		{
+			if (x == 0 && y == 0)
+				continue;
+
+			int2 Offset = int2(x,y);
+			int2 SampleIndex = PixelPos + Offset;
+			
+			bool2 bInsideViewport = 0 <= SampleIndex && SampleIndex < (int2)VelocityTile_ViewportMax;
+			if (!all(bInsideViewport))
+				continue;
+
+			float4 ScatterMinMax = VelocityTileTexture[SampleIndex];
+			float2 MaxVelocity = ScatterMinMax.zw;
+
+			float2 VelocityPixels = MaxVelocity * MotionBlur_VelocityScaleForTiles;
+
+			float  VelocityLengthPixelsSqr = dot(VelocityPixels, VelocityPixels);
+			float  VelocityLengthPixelsInv = rsqrtFast(VelocityLengthPixelsSqr + 1e-8);
+			float  VelocityLengthPixels = VelocityLengthPixelsSqr * VelocityLengthPixelsInv;
+			float2 VelocityDir = VelocityPixels * VelocityLengthPixelsInv;
+
+			// Project pixel corner on to dir. This is the oriented extent of a pixel.
+			// 1/2 pixel because shape is swept tile
+			// +1/2 pixel for conservative rasterization
+			// 99% to give epsilon before neighbor is filled. Otherwise all neighbors lie on edges of quad when no velocity in their direction.
+			float PixelExtent = abs(VelocityDir.x) + abs(VelocityDir.y);
+			float2 QuadExtent = float2(VelocityLengthPixels, 0) + PixelExtent.xx * 0.99;
+
+			// Orient quad along velocity direction
+			float2 AxisX = VelocityDir;
+			float2 AxisY = float2(-VelocityDir.y, VelocityDir.x);
+
+			// Project this pixel center onto scatter quad
+			float2 PixelCenterOnQuad;
+			PixelCenterOnQuad.x = dot(AxisX, Offset);
+			PixelCenterOnQuad.y = dot(AxisY, Offset);
+
+			bool2 bInsideQuad = abs(PixelCenterOnQuad) < QuadExtent;
+			if (all(bInsideQuad)) 
+			{
+				MinMaxVelocity = MinMaxLength(MinMaxVelocity, ScatterMinMax);
+			}
+		}
+	}
+
+	BRANCH
+	if (all(PixelPos.xy < VelocityTile_ViewportMax))
+	{
+		OutVelocityTileTexture[PixelPos] = MinMaxVelocity;
+	}
+}
diff --git a/Engine/Shaders/Private/MotionBlur/MotionBlurTileScatter.usf b/Engine/Shaders/Private/MotionBlur/MotionBlurTileScatter.usf
new file mode 100644
--- /dev/null
+++ b/Engine/Shaders/Private/MotionBlur/MotionBlurTileScatter.usf
@@ -0,0 +1,97 @@
+// Copyright Epic Games, Inc. All Rights Reserved.
+
+#include "MotionBlurCommon.ush"
+
+//------------------------------------------------------- CONFIG
+
+#define VELOCITY_SCATTER_PASS_MIN 0
+#define VELOCITY_SCATTER_PASS_MAX 1
+
+
+//------------------------------------------------------- PARAMETERS
+
+SCREEN_PASS_TEXTURE_VIEWPORT(VelocityTile)
+float MotionBlur_VelocityScaleForTiles;
+uint ScatterPass;
+
+Texture2D VelocityTileTexture;
+
+
+//------------------------------------------------------- ENTRY POINT
+
+void MainVS(
+	uint VId : SV_VertexID,
+	uint IId : SV_InstanceID,
+	out nointerpolation float4 OutColor : TEXCOORD0,
+	out float4 OutPosition : SV_POSITION)
+{
+	OutPosition = float4(0, 0, 0, 1);
+
+	// needs to be the same on C++ side (faster on NVIDIA and AMD)
+	uint QuadsPerInstance = 8;
+
+	// remap the indices to get vertexid to VId and quadid into IId
+	IId = IId * QuadsPerInstance + (VId / 4);
+	VId = VId % 4;
+
+	// triangle A: 0:left top, 1:right top, 2: left bottom
+	// triangle B: 3:right bottom, 4:left bottom, 5: right top
+	float2 CornerOffset = float2(VId % 2, VId / 2) * 2 - 1;
+
+	uint2 PixelPos = uint2(IId % VelocityTile_ViewportMax.x, IId / VelocityTile_ViewportMax.x);
+
+	BRANCH
+	if (PixelPos.y >= VelocityTile_ViewportMax.y)
+	{
+		OutColor = 0;
+		return;
+	}
+
+	float2 SvPosition = PixelPos + 0.5;
+
+	float4 MinMaxVelocity = VelocityTileTexture[PixelPos];
+	OutColor = MinMaxVelocity;
+
+	float4 MinMaxVelocityPixels = MinMaxVelocity * MotionBlur_VelocityScaleForTiles;
+	float2 VelocityPixels = MinMaxVelocityPixels.zw;
+
+	// Is the velocity small enough not to cover adjacent tiles?
+	BRANCH
+	if (dot(VelocityPixels, VelocityPixels) * 16 * 16 <= 0.25)
+	{
+		OutPosition.xy = (SvPosition + CornerOffset * 0.5 - VelocityTile_ScreenPosToViewportBias) / VelocityTile_ScreenPosToViewportScale.xy;
+		OutPosition.z = 0.0002; // zero clips
+		return;
+	}
+
+	float  VelocityLengthPixelsSqr = dot(VelocityPixels, VelocityPixels);
+	float  VelocityLengthPixelsInv = rsqrtFast(VelocityLengthPixelsSqr);
+	float  VelocityLengthPixels = VelocityLengthPixelsSqr * VelocityLengthPixelsInv;
+	float2 VelocityDir = VelocityPixels * VelocityLengthPixelsInv;
+
+	// Project pixel corner on to dir. This is the oriented extent of a pixel.
+	// 1/2 pixel because shape is swept tile
+	// +1/2 pixel for conservative rasterization
+	// 99% to give epsilon before neighbor is filled. Otherwise all neighbors lie on edges of quad when no velocity in their direction.
+	// dot(abs(VelocityDir), float2(1, 1))
+	float Extent = abs(VelocityDir.x) + abs(VelocityDir.y);
+	CornerOffset *= float2(VelocityLengthPixels, 0) + Extent.xx * 0.99;
+
+	// Orient along velocity direction
+	float2 AxisX = VelocityDir;
+	float2 AxisY = float2(-VelocityDir.y, VelocityDir.x);
+	CornerOffset = AxisX * CornerOffset.x + AxisY * CornerOffset.y;
+
+	OutPosition.xy = (SvPosition + CornerOffset - VelocityTile_ScreenPosToViewportBias) / VelocityTile_ScreenPosToViewportScale;
+
+	// Depth ordered by velocity length
+	OutPosition.z = (ScatterPass == VELOCITY_SCATTER_PASS_MAX) ? VelocityLengthPixels : length(MinMaxVelocityPixels.xy);
+	OutPosition.z = clamp(OutPosition.z / VelocityTile_ScreenPosToViewportScale.x * 0.5, 0.0002, 0.999);
+}
+
+void MainPS(
+	nointerpolation float4 InColor : TEXCOORD0,
+	out float4 OutColor : SV_Target0)
+{
+	OutColor = InColor;
+}
diff --git a/Engine/Shaders/Private/MotionBlur/MotionBlurVelocityFlatten.usf b/Engine/Shaders/Private/MotionBlur/MotionBlurVelocityFlatten.usf
new file mode 100644
--- /dev/null
+++ b/Engine/Shaders/Private/MotionBlur/MotionBlurVelocityFlatten.usf
@@ -0,0 +1,195 @@
+// Copyright Epic Games, Inc. All Rights Reserved.
+
+#include "MotionBlurVelocityFlatten.ush"
+
+
+//------------------------------------------------------- CONSTANTS
+
+#define THREADGROUP_SIZEX		16
+#define THREADGROUP_SIZEY		16
+#define THREADGROUP_TOTALSIZE	(THREADGROUP_SIZEX * THREADGROUP_SIZEY)
+
+#ifndef CAMERA_MOTION_BLUR_MODE
+// 0-off; 1-on; 2-override
+#define CAMERA_MOTION_BLUR_MODE 1
+#endif
+
+
+//------------------------------------------------------- PARAMETERS
+
+SCREEN_PASS_TEXTURE_VIEWPORT(Velocity)
+SCREEN_PASS_TEXTURE_VIEWPORT(VelocityTile)
+
+#if CAMERA_MOTION_BLUR_MODE == 2
+float4x4 ClipToPrevClipOverride;
+#endif
+
+float MotionBlur_AspectRatio;
+float MotionBlur_VelocityScale;
+float MotionBlur_VelocityScaleForTiles;
+float MotionBlur_VelocityMax;
+
+Texture2D VelocityTexture;
+Texture2D VelocityTileTexture;
+Texture2D DepthTexture;
+
+RWTexture2D<float4>	OutVelocityFlatTexture;
+RWTexture2D<float4>	OutVelocityTileTexture;
+
+//------------------------------------------------------- LDS
+
+groupshared float4 Shared[THREADGROUP_TOTALSIZE];
+
+
+//------------------------------------------------------- FUNCTIONS
+
+float2 GetCameraMotionVelocity(uint3 DispatchThreadId, float Depth, float4x4 ClipToPrevClip)
+{
+	// Compute velocity due to camera motion.
+	float2 ViewportUV = ((float2)DispatchThreadId.xy + 0.5) / Velocity_ViewportSize;
+	float2 ScreenPos = 2 * float2(ViewportUV.x, 1 - ViewportUV.y) - 1;
+	float4 ThisClip = float4(ScreenPos, Depth, 1);
+	float4 PrevClip = mul(ThisClip, ClipToPrevClip);
+	float2 PrevScreen = PrevClip.xy / PrevClip.w;
+	return ScreenPos - PrevScreen;
+}
+
+
+//------------------------------------------------------- ENTRY POINT
+
+[numthreads(THREADGROUP_SIZEX, THREADGROUP_SIZEY, 1)]
+void MainCS(
+	uint3 GroupId : SV_GroupID,
+	uint3 DispatchThreadId : SV_DispatchThreadID,
+	uint3 GroupThreadId : SV_GroupThreadID,
+	uint GroupIndex : SV_GroupIndex) 
+{ 
+	int2 PixelPos = min(DispatchThreadId.xy + Velocity_ViewportMin, Velocity_ViewportMax - 1);
+	int2 DepthOffset = 0;
+
+	{
+		const int Cross = 1;
+		// For motion vector, use camera/dynamic motion from min depth pixel in pattern around pixel.
+		// This enables better quality outline on foreground against different motion background.
+		
+		half4 Depths;
+		Depths.x = DepthTexture[ PixelPos + int2(-Cross, -Cross) ].x;
+		Depths.y = DepthTexture[ PixelPos + int2( Cross, -Cross) ].x;
+		Depths.z = DepthTexture[ PixelPos + int2(-Cross,  Cross) ].x;
+		Depths.w = DepthTexture[ PixelPos + int2( Cross,  Cross) ].x;
+
+		DepthOffset = Cross;
+		int DepthOffsetXx = Cross;
+
+		// Nearest depth is the largest depth (depth surface 0=far, 1=near).
+		if(Depths.x > Depths.y) 
+		{
+			DepthOffsetXx = -Cross;
+		}
+		if(Depths.z > Depths.w) 
+		{
+			DepthOffset.x = -Cross;
+		}
+		half DepthsXY = max(Depths.x, Depths.y);
+		half DepthsZW = max(Depths.z, Depths.w);
+		if(DepthsXY > DepthsZW) 
+		{
+			DepthOffset.y = -Cross;
+			DepthOffset.x = DepthOffsetXx; 
+		}
+	}
+
+	float4 EncodedVelocity = VelocityTexture[ PixelPos + DepthOffset ];
+	float Depth = DepthTexture[ PixelPos + DepthOffset ].x;
+
+	float2 Velocity = 0.0f;
+	if (EncodedVelocity.x > 0.0)
+	{
+		Velocity = DecodeVelocityFromTexture(EncodedVelocity).xy;
+#if CAMERA_MOTION_BLUR_MODE != 1
+		// Cancel motion blur caused by the current camera movements
+		Velocity -= GetCameraMotionVelocity(DispatchThreadId, Depth, View.ClipToPrevClip);
+#endif
+#if CAMERA_MOTION_BLUR_MODE == 2
+		// Apply motion blur caused by the custom camera movements
+		Velocity += GetCameraMotionVelocity(DispatchThreadId, Depth, ClipToPrevClipOverride);
+#endif
+	}
+#if CAMERA_MOTION_BLUR_MODE > 0
+	else
+	{
+#if CAMERA_MOTION_BLUR_MODE == 2
+		// Apply motion blur caused by the custom camera movements
+		Velocity = GetCameraMotionVelocity(DispatchThreadId, Depth, ClipToPrevClipOverride);
+#else
+		// Apply motion blur caused by the current camera movements
+		Velocity = GetCameraMotionVelocity(DispatchThreadId, Depth, View.ClipToPrevClip);
+#endif
+	}
+#endif
+
+	Velocity.y *= -MotionBlur_AspectRatio;
+
+	float2 VelocityPolar = CartesianToPolar(Velocity);
+
+	// If the velocity vector was zero length, VelocityPolar will contain NaNs.
+	if (any(isnan(VelocityPolar)))
+	{
+		VelocityPolar = float2(0.0f, 0.0f);
+	}
+
+	bool bInsideViewport = all(PixelPos.xy < Velocity_ViewportMax);
+
+	// 11:11:10  (VelocityLength, VelocityAngle, Depth)
+	float2 EncodedPolarVelocity;
+	EncodedPolarVelocity.x = VelocityPolar.x;
+	EncodedPolarVelocity.y = VelocityPolar.y * (0.5 / PI) + 0.5;
+
+	BRANCH
+	if (bInsideViewport)
+	{
+		OutVelocityFlatTexture[PixelPos] = float3(EncodedPolarVelocity, ConvertFromDeviceZ(Depth)).xyzz;
+	}
+
+	// Limit velocity
+	VelocityPolar.x = min(VelocityPolar.x, MotionBlur_VelocityMax / MotionBlur_VelocityScale);
+
+	float4 VelocityMinMax = VelocityPolar.xyxy;
+	VelocityMinMax.x = bInsideViewport ? VelocityMinMax.x : 2;
+	VelocityMinMax.z = bInsideViewport ? VelocityMinMax.z : 0;
+
+	Shared[GroupIndex] = VelocityMinMax;
+
+	GroupMemoryBarrierWithGroupSync();
+
+#if THREADGROUP_TOTALSIZE > 512
+	if (GroupIndex < 512) Shared[GroupIndex] = MinMaxLengthPolar(Shared[GroupIndex], Shared[GroupIndex + 512]);
+	GroupMemoryBarrierWithGroupSync();
+#endif
+#if THREADGROUP_TOTALSIZE > 256
+	if (GroupIndex < 256) Shared[GroupIndex] = MinMaxLengthPolar(Shared[GroupIndex], Shared[GroupIndex + 256]);
+	GroupMemoryBarrierWithGroupSync();
+#endif
+#if THREADGROUP_TOTALSIZE > 128
+	if (GroupIndex < 128) Shared[GroupIndex] = MinMaxLengthPolar(Shared[GroupIndex], Shared[GroupIndex + 128]);
+	GroupMemoryBarrierWithGroupSync();
+#endif
+#if THREADGROUP_TOTALSIZE > 64
+	if (GroupIndex <  64) Shared[GroupIndex] = MinMaxLengthPolar(Shared[GroupIndex], Shared[GroupIndex +  64]);
+	GroupMemoryBarrierWithGroupSync();
+#endif
+
+	// Safe for vector sizes 32 or larger, AMD and NV
+	// TODO Intel variable size vector
+	if (GroupIndex < 32) Shared[GroupIndex] = MinMaxLengthPolar(Shared[GroupIndex], Shared[GroupIndex + 32]);
+	if (GroupIndex < 16) Shared[GroupIndex] = MinMaxLengthPolar(Shared[GroupIndex], Shared[GroupIndex + 16]);
+	if (GroupIndex <  8) Shared[GroupIndex] = MinMaxLengthPolar(Shared[GroupIndex], Shared[GroupIndex +  8]);
+	if (GroupIndex <  4) Shared[GroupIndex] = MinMaxLengthPolar(Shared[GroupIndex], Shared[GroupIndex +  4]);
+	if (GroupIndex <  2) Shared[GroupIndex] = MinMaxLengthPolar(Shared[GroupIndex], Shared[GroupIndex +  2]);
+	if (GroupIndex <  1) Shared[GroupIndex] = MinMaxLengthPolar(Shared[GroupIndex], Shared[GroupIndex +  1]);
+
+	if (GroupIndex == 0)
+	{
+		OutVelocityTileTexture[GroupId.xy] = float4(PolarToCartesian(Shared[0].xy), PolarToCartesian(Shared[0].zw));
+	}
+}
diff --git a/Engine/Shaders/Private/MotionBlur/MotionBlurVelocityFlatten.ush b/Engine/Shaders/Private/MotionBlur/MotionBlurVelocityFlatten.ush
new file mode 100644
--- /dev/null
+++ b/Engine/Shaders/Private/MotionBlur/MotionBlurVelocityFlatten.ush
@@ -0,0 +1,6 @@
+// Copyright Epic Games, Inc. All Rights Reserved.
+
+#include "MotionBlurCommon.ush"
+
+
+// TODO: merge velocity flatten pass within TSR's dilate velocity pass.
diff --git a/Engine/Shaders/Private/MotionBlur/MotionBlurVisualize.usf b/Engine/Shaders/Private/MotionBlur/MotionBlurVisualize.usf
new file mode 100644
--- /dev/null
+++ b/Engine/Shaders/Private/MotionBlur/MotionBlurVisualize.usf
@@ -0,0 +1,217 @@
+// Copyright Epic Games, Inc. All Rights Reserved.
+
+#include "MotionBlurCommon.ush"
+//#include "Random.ush"
+
+
+//------------------------------------------------------- CONSTANTS
+
+// 0:off / 1:on, useful to debug the motionblur algorithm
+#define MOTIONBLUR_TESTCHART 0
+
+
+//------------------------------------------------------- PARAMETERS
+
+SCREEN_PASS_TEXTURE_VIEWPORT(Velocity)
+
+Texture2D ColorTexture;
+Texture2D VelocityTexture;
+Texture2D DepthTexture;
+
+SamplerState ColorSampler;
+SamplerState VelocitySampler;
+SamplerState DepthSampler;
+
+#if USE_POST_MOTION_BLUR_TRANSLUCENCY
+
+Texture2D PostMotionBlurTranslucencyTexture;
+SamplerState PostMotionBlurTranslucencySampler;
+
+#endif // USE_POST_MOTION_BLUR_TRANSLUCENCY
+
+#if SUPPORTS_INDEPENDENT_SAMPLERS
+	#define SharedVelocitySampler VelocitySampler
+	#define SharedDepthSampler VelocitySampler
+#else
+	#define SharedVelocitySampler VelocitySampler
+	#define SharedDepthSampler DepthSampler
+#endif
+
+
+//------------------------------------------------------- FUNCTIONS
+
+// debug motionblur (very useful, keep)
+// @param ScreenPos -1..1 -1..1 for viewport
+// @param Velocity in -1..1 range for full motionblur
+// @apram Color RGB and depth in alpha
+// @param AvgObject 0:background, 1:foregound
+void OverrideWithTestChart(float2 ScreenPos, inout float2 ObjectVelocity, inout float2 BackgroundVelocity, inout float4 Color, inout float AvgObject)
+{
+#if MOTIONBLUR_TESTCHART == 1
+	// needs to be inside the loop to prevent NVIDIA driver optimizetion (blinking)
+	float2 PixelPos = ScreenPos * Velocity_ScreenPosToViewportScale.xy + Velocity_ScreenPosToViewportBias - 25;
+	float3 BackgroundColor = lerp(0.0, 0.3f, PseudoRandom(PixelPos)); 
+	float3 ForegroundColor = lerp(float3(1, 0, 0), float3(1, 1, 0), PseudoRandom(PixelPos)); 
+
+	int2 tile = (int2)floor(PixelPos / 12.0f);
+	int2 experiment = (int2)floor(tile / 5.0f);
+
+	if(experiment.x >= 0 && experiment.y >= 0 && experiment.x < 10 && experiment.y < 5)
+	{
+		int2 localtile = uint2(tile) % 5;
+
+		bool bForeground = localtile.x == 2 && localtile.y == 2;
+
+		Color.rgb = bForeground ? ForegroundColor : BackgroundColor;
+
+		// depth
+		Color.a = bForeground ? 100.0f : 1000.0f;
+
+		bool bLeftSide = experiment.x < 5;
+
+		if(!bLeftSide)
+		{
+			experiment.x -= 5;
+		}
+
+		float ForegroundAngle = (experiment.x - 1) * (6.283f / 12);
+		float BackgroundAngle = (experiment.y - 1) * (6.283f / 12) + 3.1415f/2;
+
+		// ForegroundR with very small amounts needs extra testing so we do a non linear scale
+		float ForegroundR = pow(experiment.x / 5.0f, 2);
+		float BackgroundR = pow(experiment.y / 5.0f, 2);
+
+		float2 ForegroundXY = ForegroundR * float2(sin(ForegroundAngle), cos(ForegroundAngle));
+		float2 BackgroundXY = BackgroundR * float2(sin(BackgroundAngle), cos(BackgroundAngle));
+
+		BackgroundVelocity.xy = BackgroundXY;
+
+		if(bLeftSide)
+		{
+			ObjectVelocity.xy = ForegroundXY;
+			AvgObject = bForeground;
+		}
+		else
+		{
+			ObjectVelocity.xy = bForeground ? ForegroundXY : BackgroundXY;
+			AvgObject = 1.0f;
+		}
+	}
+
+#endif
+}
+
+// used to visualize the motion blur
+// @return 0/1
+float Compute3DCheckerBoard(float3 Pos)
+{
+	float3 TiledWorldPos = frac(Pos) > 0.5f;
+	return (float)((uint)dot(float3(1,1,1), TiledWorldPos) % 2);
+}
+
+
+//------------------------------------------------------- ENTRY POINT
+
+void MainPS(
+	noperspective float4 UVAndScreenPos : TEXCOORD0,
+	out float4 OutColor : SV_Target0)
+{
+	const float TileSize = 48;
+
+	float2 ScreenSpacePos = UVAndScreenPos.zw;
+	float2 PixelPos = ScreenSpacePos * Velocity_ScreenPosToViewportScale + Velocity_ScreenPosToViewportBias;
+	float2 UV = PixelPos * Velocity_ExtentInverse;
+
+	float2 PixelPosAtTileCenter = PixelPos - (frac(PixelPos / TileSize) - 0.5f) * TileSize;
+	float2 UVAtTileCenter = PixelPosAtTileCenter * Velocity_ExtentInverse;
+	float2 ScreenPosAtTileCenter = (PixelPosAtTileCenter - Velocity_ScreenPosToViewportBias) / Velocity_ScreenPosToViewportScale;
+
+	// World aligned checkerboards
+	{
+		float DeviceZ = DepthTexture.SampleLevel(SharedDepthSampler, UV, 0).r;
+		float SceneDepth = ConvertFromDeviceZ(DeviceZ);
+
+		float3 ScreenVector = mul(float4(ScreenSpacePos, 1, 0), View.ScreenToWorld).xyz;
+
+		// world space position of the current pixel
+		float3 OffsetWorldPos = ScreenVector * SceneDepth;
+
+		float3 AbsWorldPos = View.WorldCameraOrigin + OffsetWorldPos;
+		float3 WorldCheckerboard = Compute3DCheckerBoard(AbsWorldPos * 0.02f) * 0.1f + Compute3DCheckerBoard(AbsWorldPos * 0.002f) * 0.3f + Compute3DCheckerBoard(AbsWorldPos * 0.0002f) * 0.6f;
+		OutColor = float4(lerp(WorldCheckerboard, float3(0,0,0), 0.7f), 1);
+	}
+
+	bool bSelectorOpaqueAtTileCenter = VelocityTexture.SampleLevel(SharedVelocitySampler, UVAtTileCenter, 0).x > 0;
+	bool bSelectorOpaque = VelocityTexture.SampleLevel(SharedVelocitySampler, UV, 0).x > 0;
+
+	// relative, in screen space -1...1 -1..1, can be even outside of that range, points into the movement direction
+	float2 VelocityAtTileCenter;
+	{
+		VelocityAtTileCenter = DecodeVelocityFromTexture(VelocityTexture.SampleLevel(SharedVelocitySampler, UVAtTileCenter, 0)).xy;
+
+		// reconstruct from camera motion if we don't have velocity data
+		if(!bSelectorOpaqueAtTileCenter)
+		{
+			float DeviceZ = DepthTexture.SampleLevel(SharedDepthSampler, UVAtTileCenter, 0).r;
+
+			// UV is per tile
+			float4 ThisClip = float4(ScreenPosAtTileCenter, DeviceZ, 1);
+			float4 PrevClip = mul(ThisClip, View.ClipToPrevClip);
+			float2 PrevScreen = PrevClip.xy / PrevClip.w;
+
+			// points into the movement direction
+			VelocityAtTileCenter = ScreenPosAtTileCenter - PrevScreen;
+		}
+	}
+
+	// tint yellow if velocity data is stored in texture
+	{
+		float3 TintColor = bSelectorOpaque ? float3(0.5f, 0.5f, 0.2f) : float3(0.5f, 0.5f, 0.5f);
+	
+		OutColor.rgb = lerp(OutColor.rgb, TintColor, 0.45f);
+	}
+
+	// tile center
+	{
+		float2 Delta = PixelPos - PixelPosAtTileCenter;
+
+		float HMask = lerp(saturate(abs(Delta.y)), 1, saturate(abs(Delta.x) / TileSize * 4));
+		float VMask = lerp(saturate(abs(Delta.x)), 1, saturate(abs(Delta.y) / TileSize * 4));
+
+		float Dist = length(Delta);
+		OutColor.rgb *= lerp(1, 0.3f, saturate(3 - Dist));
+	}
+
+	float3 LineColor = bSelectorOpaqueAtTileCenter ? float3(1,1,0) : float3(0.7,0.7,0.7);
+
+	// points into the movement direction
+	float2 PixelDirection = VelocityAtTileCenter * Velocity_ScreenPosToViewportScale;
+
+	// arrow
+	{
+		float2 PerpPixelDirection = float2(PixelDirection.y, -PixelDirection.x);
+		float2 DirectionInTile = PixelPos - PixelPosAtTileCenter;
+
+		float DistOnLine = dot(normalize(-PixelDirection), DirectionInTile) + length(PixelDirection);
+
+		bool bArrowHead = DistOnLine < 8;
+
+		float LocalThickness = 1 + (frac(DistOnLine/8)*8)*0.25f;
+
+		float PerpDirectionMask = saturate(LocalThickness - abs(dot(normalize(PerpPixelDirection), DirectionInTile)));
+		float DirectionMask = saturate(length(PixelDirection) - length(DirectionInTile));
+
+		float3 LineMask = PerpDirectionMask * DirectionMask;
+		OutColor.rgb = lerp(OutColor.rgb, LineColor, LineMask);
+	}
+
+	// previous pos is a dot
+	{
+		float3 DotColor = float3(0,1,0);
+		// PixelPos of the previous position
+		float2 PreviousPixelPos = PixelPosAtTileCenter - PixelDirection;
+		float Dist = length(PreviousPixelPos - PixelPos);
+		OutColor.rgb = lerp(OutColor.rgb, LineColor, saturate(4 - Dist));
+		OutColor.rgb = lerp(OutColor.rgb, 0, saturate(2.5f - Dist));
+	}
+}
diff --git a/Engine/Shaders/Private/PostProcessMotionBlur.usf b/Engine/Shaders/Private/PostProcessMotionBlur.usf
deleted file mode 100644
diff --git a/Engine/Shaders/Private/PostProcessVelocityFlatten.usf b/Engine/Shaders/Private/PostProcessVelocityFlatten.usf
deleted file mode 100644
diff --git a/Engine/Source/Runtime/Renderer/Private/PostProcess/PostProcessMotionBlur.cpp b/Engine/Source/Runtime/Renderer/Private/PostProcess/PostProcessMotionBlur.cpp
--- a/Engine/Source/Runtime/Renderer/Private/PostProcess/PostProcessMotionBlur.cpp
+++ b/Engine/Source/Runtime/Renderer/Private/PostProcess/PostProcessMotionBlur.cpp
@@ -257,8 +257,6 @@ public:
 	END_SHADER_PARAMETER_STRUCT()
 };
 
-IMPLEMENT_GLOBAL_SHADER(FMotionBlurVelocityFlattenCS, "/Engine/Private/PostProcessVelocityFlatten.usf", "VelocityFlattenMain", SF_Compute);
-
 BEGIN_SHADER_PARAMETER_STRUCT(FMotionBlurVelocityDilateParameters, )
 	SHADER_PARAMETER_STRUCT(FMotionBlurParameters, MotionBlur)
 	SHADER_PARAMETER_STRUCT(FScreenPassTextureViewportParameters, VelocityTile)
@@ -279,8 +277,6 @@ public:
 	END_SHADER_PARAMETER_STRUCT()
 };
 
-IMPLEMENT_GLOBAL_SHADER(FMotionBlurVelocityDilateGatherCS, "/Engine/Private/PostProcessVelocityFlatten.usf", "VelocityGatherCS", SF_Compute);
-
 enum class EMotionBlurVelocityScatterPass : uint32
 {
 	DrawMin,
@@ -296,24 +292,20 @@ END_SHADER_PARAMETER_STRUCT()
 
 class FMotionBlurVelocityDilateScatterVS : public FMotionBlurShader
 {
-public:
 	DECLARE_GLOBAL_SHADER(FMotionBlurVelocityDilateScatterVS);
 	SHADER_USE_PARAMETER_STRUCT(FMotionBlurVelocityDilateScatterVS, FMotionBlurShader);
+
 	using FParameters = FMotionBlurVelocityDilateScatterParameters;
 };
 
-IMPLEMENT_GLOBAL_SHADER(FMotionBlurVelocityDilateScatterVS, "/Engine/Private/PostProcessVelocityFlatten.usf", "VelocityScatterVS", SF_Vertex);
-
 class FMotionBlurVelocityDilateScatterPS : public FMotionBlurShader
 {
-public:
 	DECLARE_GLOBAL_SHADER(FMotionBlurVelocityDilateScatterPS);
 	SHADER_USE_PARAMETER_STRUCT(FMotionBlurVelocityDilateScatterPS, FMotionBlurShader);
+
 	using FParameters = FMotionBlurVelocityDilateScatterParameters;
 };
 
-IMPLEMENT_GLOBAL_SHADER(FMotionBlurVelocityDilateScatterPS, "/Engine/Private/PostProcessVelocityFlatten.usf", "VelocityScatterPS", SF_Pixel);
-
 class FMotionBlurQualityDimension : SHADER_PERMUTATION_ENUM_CLASS("MOTION_BLUR_QUALITY", EMotionBlurQuality);
 class FPostMotionBlurTranslucencyDimension : SHADER_PERMUTATION_BOOL("USE_POST_MOTION_BLUR_TRANSLUCENCY");
 
@@ -357,8 +349,6 @@ public:
 	using FPermutationDomain = FMotionBlurFilterPermutationDomain;
 };
 
-IMPLEMENT_GLOBAL_SHADER(FMotionBlurFilterPS, "/Engine/Private/PostProcessMotionBlur.usf", "MainPS", SF_Pixel);
-
 class FMotionBlurFilterCS : public FMotionBlurShader
 {
 public:
@@ -380,8 +370,6 @@ public:
 	using FPermutationDomain = FMotionBlurFilterPermutationDomain;
 };
 
-IMPLEMENT_GLOBAL_SHADER(FMotionBlurFilterCS, "/Engine/Private/PostProcessMotionBlur.usf", "MainCS", SF_Compute);
-
 class FMotionBlurVisualizePS : public FMotionBlurShader
 {
 public:
@@ -404,14 +392,15 @@ public:
 
 		RENDER_TARGET_BINDING_SLOTS()
 	END_SHADER_PARAMETER_STRUCT()
-
-	static bool ShouldCompilePermutation(const FGlobalShaderPermutationParameters& Parameters)
-	{
-		return IsFeatureLevelSupported(Parameters.Platform, ERHIFeatureLevel::SM5);
-	}
 };
 
-IMPLEMENT_GLOBAL_SHADER(FMotionBlurVisualizePS, "/Engine/Private/PostProcessMotionBlur.usf", "VisualizeMotionBlurPS", SF_Pixel);
+IMPLEMENT_GLOBAL_SHADER(FMotionBlurVelocityFlattenCS,       "/Engine/Private/MotionBlur/MotionBlurVelocityFlatten.usf", "MainCS", SF_Compute);
+IMPLEMENT_GLOBAL_SHADER(FMotionBlurVelocityDilateGatherCS,  "/Engine/Private/MotionBlur/MotionBlurTileGather.usf",      "MainCS", SF_Compute);
+IMPLEMENT_GLOBAL_SHADER(FMotionBlurVelocityDilateScatterVS, "/Engine/Private/MotionBlur/MotionBlurTileScatter.usf",     "MainVS", SF_Vertex);
+IMPLEMENT_GLOBAL_SHADER(FMotionBlurVelocityDilateScatterPS, "/Engine/Private/MotionBlur/MotionBlurTileScatter.usf",     "MainPS", SF_Pixel);
+IMPLEMENT_GLOBAL_SHADER(FMotionBlurFilterPS,                "/Engine/Private/MotionBlur/MotionBlurApply.usf",           "MainPS", SF_Pixel);
+IMPLEMENT_GLOBAL_SHADER(FMotionBlurFilterCS,                "/Engine/Private/MotionBlur/MotionBlurApply.usf",           "MainCS", SF_Compute);
+IMPLEMENT_GLOBAL_SHADER(FMotionBlurVisualizePS,             "/Engine/Private/MotionBlur/MotionBlurVisualize.usf",       "MainPS", SF_Pixel);
 
 TGlobalResource<FSpriteIndexBuffer<8>> GScatterQuadIndexBuffer;
 
