Reapply LWC Changes to MotionBlurVisualize
Integration: CL-17962937

diff --git a/Engine/Shaders/Private/MotionBlur/MotionBlurVisualize.usf b/Engine/Shaders/Private/MotionBlur/MotionBlurVisualize.usf
--- a/Engine/Shaders/Private/MotionBlur/MotionBlurVisualize.usf
+++ b/Engine/Shaders/Private/MotionBlur/MotionBlurVisualize.usf
@@ -131,12 +131,12 @@ void MainPS(
 		float DeviceZ = DepthTexture.SampleLevel(SharedDepthSampler, UV, 0).r;
 		float SceneDepth = ConvertFromDeviceZ(DeviceZ);
 
-		float3 ScreenVector = mul(float4(ScreenSpacePos, 1, 0), View.ScreenToWorld).xyz;
+		float3 ScreenVector = mul(float3(ScreenSpacePos, 1), LWCToFloat3x3(PrimaryView.ScreenToWorld)).xyz;
 
 		// world space position of the current pixel
 		float3 OffsetWorldPos = ScreenVector * SceneDepth;
 
-		float3 AbsWorldPos = View.WorldCameraOrigin + OffsetWorldPos;
+		float3 AbsWorldPos = LWCHackToFloat(PrimaryView.WorldCameraOrigin) + OffsetWorldPos;
 		float3 WorldCheckerboard = Compute3DCheckerBoard(AbsWorldPos * 0.02f) * 0.1f + Compute3DCheckerBoard(AbsWorldPos * 0.002f) * 0.3f + Compute3DCheckerBoard(AbsWorldPos * 0.0002f) * 0.6f;
 		OutColor = float4(lerp(WorldCheckerboard, float3(0,0,0), 0.7f), 1);
 	}
