Implements r.MotionBlur.Directions=2

diff --git a/Engine/Shaders/Private/MotionBlur/MotionBlurApply.usf b/Engine/Shaders/Private/MotionBlur/MotionBlurApply.usf
--- a/Engine/Shaders/Private/MotionBlur/MotionBlurApply.usf
+++ b/Engine/Shaders/Private/MotionBlur/MotionBlurApply.usf
@@ -1,5 +1,8 @@
 // Copyright Epic Games, Inc. All Rights Reserved.
 
+#define CONFIG_MAX_RANGE_SIZE DIM_BLUR_DIRECTIONS
+
+
 #include "MotionBlurCommon.ush"
 #include "../Random.ush"
 
@@ -12,7 +15,22 @@
 #define MOTION_BLUR_QUALITY_LOW 0
 #define MOTION_BLUR_QUALITY_MEDIUM 1
 #define MOTION_BLUR_QUALITY_HIGH 2
-#define MOTION_BLUR_QUALITY_VERY_HIGH 2
+#define MOTION_BLUR_QUALITY_VERY_HIGH 3
+
+
+#define CONFIG_HOLE_FILL_SEPARATLY 1
+
+#if MOTION_BLUR_QUALITY == MOTION_BLUR_QUALITY_LOW
+	#define CONFIG_SAMPLE_COUNT 4
+#elif MOTION_BLUR_QUALITY == MOTION_BLUR_QUALITY_MEDIUM
+	#define CONFIG_SAMPLE_COUNT 6
+#elif MOTION_BLUR_QUALITY == MOTION_BLUR_QUALITY_HIGH
+	#define CONFIG_SAMPLE_COUNT 8
+#elif MOTION_BLUR_QUALITY == MOTION_BLUR_QUALITY_VERY_HIGH
+	#define CONFIG_SAMPLE_COUNT 16
+#else
+	#error Unknown motion blur quality
+#endif
 
 
 //------------------------------------------------------- PARAMETERS
@@ -29,7 +47,8 @@ FScreenTransform ColorToVelocity;
 
 Texture2D ColorTexture;
 Texture2D VelocityTexture;
-Texture2D VelocityTileTexture;
+Texture2D VelocityTileTextures_Textures_0;
+Texture2D VelocityTileTextures_Textures_1;
 Texture2D VelocityFlatTexture;
 Texture2D DepthTexture;
 
@@ -61,9 +80,13 @@ float2 PostMotionBlurTranslucencyUVMax;
 	#define SharedDepthSampler DepthSampler
 #endif
 
+RWTexture2D<float4> OutColorTexture;
+RWTexture2D<float4> DebugOutput;
+
+
 //------------------------------------------------------- LDS
 
-#if COMPUTESHADER && !COMPILER_SUPPORTS_WAVE_VOTE
+#if !COMPILER_SUPPORTS_WAVE_VOTE
 
 groupshared uint GroupSharedFastPath;
 
@@ -72,46 +95,84 @@ groupshared uint GroupSharedFastPath;
 
 //------------------------------------------------------- FUNCTIONS
 
-uint GetStepCountFromQuality()
+float NormalizeAccumulator(float Accumulator, float AccumulatorWeight)
 {
-#if MOTION_BLUR_QUALITY == MOTION_BLUR_QUALITY_LOW
-	return 4;
-#elif MOTION_BLUR_QUALITY == MOTION_BLUR_QUALITY_MEDIUM
-	return 6;
-#elif MOTION_BLUR_QUALITY == MOTION_BLUR_QUALITY_HIGH
-	return 8;
-#else // MOTION_BLUR_QUALITY == MOTION_BLUR_QUALITY_VERY_HIGH
-	return 16;
-#endif
+	return -min(-Accumulator * rcp(AccumulatorWeight), 0.0);
 }
 
-float2 DepthCmp(float CenterDepth, float SampleDepth, float DepthScale)
+float4 NormalizeAccumulator(float4 Accumulator, float AccumulatorWeight)
 {
-	return saturate(0.5 + float2(DepthScale, -DepthScale) * (SampleDepth - CenterDepth));
+	return -min(-Accumulator * rcp(AccumulatorWeight), 0.0);
 }
 
-float2 SpreadCmp(float OffsetLength, float2 SpreadLength, float PixelToSampleScale)
+void NormalizeAccumulatorWithHoleFill(
+	float4 Color,
+	float ColorWeight,
+	float HoleFillWeight,
+	float DepthAware,
+	float4 HoleFillingColor,
+	float InvSampleCount,
+	out float4 ColorOutput,
+	out float OpacityOutput,
+	out bool bValidColorOutput)
 {
-	return saturate(PixelToSampleScale * SpreadLength - max(OffsetLength - 1, 0));
+	OpacityOutput = min(saturate(ColorWeight * InvSampleCount * 2), (DepthAware * InvSampleCount * 2));
+	
+	float4 FinalAccumulator = Color + HoleFillWeight * HoleFillingColor;
+	
+	ColorOutput = NormalizeAccumulator(FinalAccumulator, ColorWeight + HoleFillWeight);
+	bValidColorOutput = (ColorWeight + HoleFillWeight) > 0;
 }
 
-float SampleWeight(float CenterDepth, float SampleDepth, float OffsetLength, float CenterSpreadLength, float SampleSpreadLength, float PixelToSampleScale, float DepthScale)
+
+// Compute the weight of the sample for hole filling.
+float ComputeSampleHoleFillWeight(float CenterDepth, float SampleDepth, float DepthScale)
 {
-	float2 DepthWeights = DepthCmp(CenterDepth, SampleDepth, DepthScale);
-	float2 SpreadWeights = SpreadCmp(OffsetLength, float2(CenterSpreadLength, SampleSpreadLength), PixelToSampleScale);
-	return dot(DepthWeights, SpreadWeights);
+	return saturate(DepthScale * (SampleDepth - CenterDepth));
 }
 
-float2 DecodeVelocity(float2 Velocity)
+// Computes the contribution weight of one sample.
+float ComputeSampleConvolutionWeight(
+	float SampleDepth, float SampleSpreadLength, float SampleVelocityAngle,
+	float OffsetLength, float BlurAngle, float PixelToSampleScale)
 {
-	// 11:11:10  (VelocityLength, VelocityAngle, Depth)
-	float VelocityLength = Velocity.x;
-	float VelocityAngle = Velocity.y * (2 * PI) - PI;
-	sincos(VelocityAngle, Velocity.x, Velocity.y);
-	Velocity *= VelocityLength;
-	return Velocity;
+	// Compare the length 
+	float SpreadWeights = saturate(PixelToSampleScale * SampleSpreadLength - max(OffsetLength - 1, 0));
+	
+	// Compare the directions
+	#if CONFIG_MAX_RANGE_SIZE > 1
+		float DirectionWeights = saturate(1.0 - max(GetPolarRelativeAngle(SampleVelocityAngle, BlurAngle) - 0.1, 0.0) * 4.0);
+	#else
+		float DirectionWeights = 1.0;
+	#endif
+
+	return SpreadWeights * DirectionWeights;
+}
+
+// Selectively computes the contribution weight of the center or the sample based on whether sample is respectively behind center or not
+float ComputeCenterOrSampleWeight(
+	float CenterDepth, float CenterSpreadLength, float CenterVelocityAngle,
+	float SampleDepth, float SampleSpreadLength, float SampleVelocityAngle,
+	float OffsetLength, float BlurAngle, float PixelToSampleScale, float DepthScale)
+{
+	// Compute weight to use the center data if center is closer than the sample.
+	float CenterWeight = saturate(0.5 + DepthScale * (SampleDepth - CenterDepth));
+	
+	// Compute weight to use the sample data if sample is closer than the center.
+	float SampleWeight = saturate(0.5 - DepthScale * (SampleDepth - CenterDepth));
+
+	float CenterConvolutionWeight = ComputeSampleConvolutionWeight(
+		CenterDepth, CenterSpreadLength, CenterVelocityAngle,
+		OffsetLength, BlurAngle, PixelToSampleScale);
+	
+	float SampleConvolutionWeight = ComputeSampleConvolutionWeight(
+		SampleDepth, SampleSpreadLength, SampleVelocityAngle,
+		OffsetLength, BlurAngle, PixelToSampleScale);
+
+	return CenterWeight * CenterConvolutionWeight + SampleWeight * SampleConvolutionWeight;
 }
 
+// TODO: move that to velocity flatten.
 float GetVelocityLengthPixels(float2 EncodedVelocity)
 {
 	// 11:11:10  (VelocityLength, VelocityAngle, Depth)
@@ -123,9 +184,14 @@ float GetVelocityLengthPixels(float2 EncodedVelocity)
 
 //------------------------------------------------------- ENTRY POINT
 
-float4 MainMotionBlurCommon(float2 ColorUV, float2 ColorPixelPos)
+[numthreads(THREADGROUP_SIZEX, THREADGROUP_SIZEY, 1)]
+void MainCS(uint2 DispatchThreadId : SV_DispatchThreadID)
 {
-	const uint StepCount = GetStepCountFromQuality() / 2;
+	float4 Debug = 0;
+
+	uint2 iColorPixelPos = DispatchThreadId + Color_ViewportMin;
+	float2 ColorPixelPos = float2(iColorPixelPos);
+	float2 ColorUV = (ColorPixelPos + 0.5) * Color_ExtentInverse;
 
 	const float PixelToTileScale = (1.0 / 16.0);
 
@@ -145,23 +211,25 @@ float4 MainMotionBlurCommon(float2 ColorUV, float2 ColorPixelPos)
 	// Velocity tile UV originates at [0,0]; only need to clamp max.
 	VelocityTileUV = min(VelocityTileUV, VelocityTile_UVViewportBilinearMax);
 
-	float4 MinMaxVelocity = VelocityTileTexture.SampleLevel(SharedVelocityTileSampler, VelocityTileUV, 0);
+	FVelocityRange VelocityRange = DecodeVelocityRange(
+		VelocityTileTextures_Textures_0.SampleLevel(SharedVelocityTileSampler, VelocityTileUV, 0),
+		VelocityTileTextures_Textures_1.SampleLevel(SharedVelocityTileSampler, VelocityTileUV, 0));
 
-	float2 MinVelocityPixels = MinMaxVelocity.xy * MotionBlur_VelocityScale;
-	float2 MaxVelocityPixels = MinMaxVelocity.zw * MotionBlur_VelocityScale;
+	float2 MinVelocityPixels = VelocityRange.Min * MotionBlur_VelocityScale;
+	float2 Max0VelocityPixels = VelocityRange.Max[0] * MotionBlur_VelocityScale;
+
+	#if CONFIG_MAX_RANGE_SIZE >= 2
+		float2 Max1VelocityPixels = VelocityRange.Max[1] * MotionBlur_VelocityScale;
+	#endif
 
 	float MinVelocityLengthSqrPixels = dot(MinVelocityPixels, MinVelocityPixels);
-	float MaxVelocityLengthSqrPixels = dot(MaxVelocityPixels, MaxVelocityPixels);
+	float Max0VelocityLengthSqrPixels = dot(Max0VelocityPixels, Max0VelocityPixels);
 
 	// Input buffer 0 as same viewport as output buffer.
 	float4 CenterColor = ColorTexture.SampleLevel(ColorSampler, ColorUV, 0);
 
-	float4 SearchVectorPixels = float4(MaxVelocityPixels, -MaxVelocityPixels);
-	float4 SearchVector = SearchVectorPixels * Color_ExtentInverse.xyxy;
+	float2 SearchVector0 = Max0VelocityPixels * Color_ExtentInverse.xyxy;
 
-	// converts pixel length to sample steps
-	float PixelToSampleScale = StepCount * rsqrt(dot(MaxVelocityPixels, MaxVelocityPixels));	
-		
 	#if USE_POST_MOTION_BLUR_TRANSLUCENCY
 		const float2 PostMotionBlurTranslucencyUV = min((ColorPixelPos + (float2)0.5f) * ScreenPosToPostMotionBlurTranslucencyUV, PostMotionBlurTranslucencyUVMax);
 		float4 PostMotionBlurTranslucency = PostMotionBlurTranslucencyTexture.SampleLevel(PostMotionBlurTranslucencySampler, PostMotionBlurTranslucencyUV, 0);
@@ -169,15 +237,15 @@ float4 MainMotionBlurCommon(float2 ColorUV, float2 ColorPixelPos)
 
 	
 	// TODO expose cvars
-	bool bSkipPath = MaxVelocityLengthSqrPixels < 0.25;
-	bool bFastPath = MinVelocityLengthSqrPixels > 0.4 * MaxVelocityLengthSqrPixels;
-
+	bool bSkipPath = Max0VelocityLengthSqrPixels < 0.25;
+	bool bFastPath = MinVelocityLengthSqrPixels > 0.4 * Max0VelocityLengthSqrPixels;
+	
 	// Only use fast path if all threads of the compute shader would.
 	#if COMPILER_SUPPORTS_WAVE_VOTE
 	{
 		bFastPath = WaveActiveAllTrue(bFastPath);
 	}
-	#elif COMPUTESHADER
+	#else
 	{
 		GroupSharedFastPath = 0;
 		GroupMemoryBarrierWithGroupSync();
@@ -190,34 +258,29 @@ float4 MainMotionBlurCommon(float2 ColorUV, float2 ColorPixelPos)
 	}
 	#endif
 	
+	float4 OutColor;
+
 	BRANCH
 	if (bSkipPath)
 	{
-		#if USE_POST_MOTION_BLUR_TRANSLUCENCY
-			CenterColor.rgb = CenterColor.rgb * PostMotionBlurTranslucency.a + PostMotionBlurTranslucency.rgb;
-		#endif
-		#if !POST_PROCESS_ALPHA
-			CenterColor.a = 0;
-		#endif
-		return CenterColor;
+		OutColor = CenterColor;
 	}
+	else if (bFastPath)
+	{
+		const float InvSampleCount = rcp(float(CONFIG_SAMPLE_COUNT));
 
-	float4 OutColor = 0;
 
-	BRANCH
-	if (bFastPath)
-	{
 		float4 ColorAccum = 0;
-
-		UNROLL
-		for (uint i = 0; i < StepCount; i++)
+		
+		UNROLL_N(CONFIG_SAMPLE_COUNT / 2)
+		for (uint i = 0; i < (CONFIG_SAMPLE_COUNT / 2); i++)
 		{
 			float2 OffsetLength = (float)i + 0.5 + float2(Random - 0.5, 0.5 - Random);
-			float2 OffsetFraction = OffsetLength / StepCount;
+			float2 OffsetFraction = OffsetLength * (2.0 * InvSampleCount);
 
 			float2 SampleUV[2];
-			SampleUV[0] = ColorUV + OffsetFraction.x * SearchVector.xy;
-			SampleUV[1] = ColorUV + OffsetFraction.y * SearchVector.zw;
+			SampleUV[0] = ColorUV + OffsetFraction.x * SearchVector0;
+			SampleUV[1] = ColorUV - OffsetFraction.y * SearchVector0;
 
 			SampleUV[0] = clamp(SampleUV[0], Color_UVViewportBilinearMin, Color_UVViewportBilinearMax);
 			SampleUV[1] = clamp(SampleUV[1], Color_UVViewportBilinearMin, Color_UVViewportBilinearMax);
@@ -226,101 +289,269 @@ float4 MainMotionBlurCommon(float2 ColorUV, float2 ColorPixelPos)
 			ColorAccum += ColorTexture.SampleLevel(ColorSampler, SampleUV[1], 0);
 		}
 	
-		ColorAccum *= 0.5 / StepCount;
+		ColorAccum *= InvSampleCount;
 
 		OutColor = ColorAccum;
 	}
 	else
 	{
+		const uint TotalSteps = CONFIG_SAMPLE_COUNT / 2;
+		const uint DirectionCount = CONFIG_MAX_RANGE_SIZE;
+		const uint StepPerDirectionCount = TotalSteps / DirectionCount;
+		const uint SamplePerDirection = CONFIG_SAMPLE_COUNT / DirectionCount;
+
+		const uint StepCount = CONFIG_SAMPLE_COUNT / 2;
+
+		// Decide whether should only accumulate along one direction.
+		#if CONFIG_MAX_RANGE_SIZE > 1
+			float Max1VelocityLengthSqrPixels = dot(Max1VelocityPixels, Max1VelocityPixels);
+			bool bDoOneDirectionOnly = Max1VelocityLengthSqrPixels < 0.25;
+			if (bDoOneDirectionOnly)
+			{
+				Max1VelocityPixels = Max0VelocityPixels;
+			}
+		#else
+			const bool bDoOneDirectionOnly = true;
+		#endif
+		
 		float3 CenterVelocityDepth = VelocityFlatTexture.SampleLevel(SharedVelocityFlatSampler, VelocityUV, 0).xyz;
 		float  CenterDepth = CenterVelocityDepth.z;
 		float  CenterVelocityLength = GetVelocityLengthPixels(CenterVelocityDepth.xy);
+		float  CenterVelocityAngle = CenterVelocityDepth.y * (2.0 * PI) - PI;
+		
+		#if CONFIG_MAX_RANGE_SIZE > 1
+			float4 HoleFillColor = 0;
+			float  HoleFillColorWeight = 0;
+		#endif
 
-		float4 ColorAccum = 0;
-		float  ColorAccumWeight = 0;
+		float4 DirectionalColor[CONFIG_MAX_RANGE_SIZE];
+		float  DirectionalColorWeight[CONFIG_MAX_RANGE_SIZE];
+		
+		#if CONFIG_MAX_RANGE_SIZE > 1
+			float  HoleFillWeightAccum[CONFIG_MAX_RANGE_SIZE];
 
-		UNROLL
-		for (uint i = 0; i < StepCount; i++)
-		{
-			float2 SampleUV[2];
-			float4 SampleColor[2];
-			float  SampleDepth[2];
-			float  SampleVelocityLength[2];
-			float  Weight[2];
+			float  DepthAccum[CONFIG_MAX_RANGE_SIZE];
+			float  DepthSquareAccum[CONFIG_MAX_RANGE_SIZE];
+			float  DepthAwareWeight[CONFIG_MAX_RANGE_SIZE];
+		#endif
 
-			float2 OffsetLength = (float)i + 0.5 + float2(Random - 0.5, 0.5 - Random);
-			float2 OffsetFraction = OffsetLength / StepCount;
+		// Iterate over the different directions.
+		UNROLL_N(CONFIG_MAX_RANGE_SIZE)
+		for (uint DirectionId = 0; DirectionId < CONFIG_MAX_RANGE_SIZE; DirectionId++)
+		{
+			float PixelToSampleScale = TotalSteps * rsqrt(dot(Max0VelocityPixels, Max0VelocityPixels));	
+			float2 SearchVector = SearchVector0;
+			float BlurAngle = CartesianToPolar(Max0VelocityPixels).y;
 
-			SampleUV[0] = ColorUV + OffsetFraction.x * SearchVector.xy;
-			SampleUV[1] = ColorUV + OffsetFraction.y * SearchVector.zw;
+			bool bAccumulateHoleFillColor = true;
 
-			SampleUV[0] = clamp(SampleUV[0], Color_UVViewportBilinearMin, Color_UVViewportBilinearMax);
-			SampleUV[1] = clamp(SampleUV[1], Color_UVViewportBilinearMin, Color_UVViewportBilinearMax);
+			#if CONFIG_MAX_RANGE_SIZE > 1
+			if (DirectionId == 1)
+			{
+				PixelToSampleScale = TotalSteps * rsqrt(dot(Max1VelocityPixels, Max1VelocityPixels));	
+				SearchVector = Max1VelocityPixels * Color_ExtentInverse.xy;
+				BlurAngle = CartesianToPolar(Max1VelocityPixels).y;
+				bAccumulateHoleFillColor = HoleFillColorWeight == 0.0;
+			}
+			#endif
+		
+			DirectionalColor[DirectionId] = 0;
+			DirectionalColorWeight[DirectionId] = 0;
+			
+			#if CONFIG_MAX_RANGE_SIZE > 1
+			{
+				HoleFillWeightAccum[DirectionId] = 0;
 
+				DepthAccum[DirectionId] = 0;
+				DepthSquareAccum[DirectionId] = 0;
+				DepthAwareWeight[DirectionId] = 0;
+			}
+			#endif
+			
+			// Iterate over steps of 2 samples in each directions.
 			UNROLL
-			for (uint j = 0; j < 2; j++)
+			for (uint StepId = DirectionId; StepId < TotalSteps; StepId += CONFIG_MAX_RANGE_SIZE)
 			{
-				float3 SampleVelocityDepth = VelocityFlatTexture.SampleLevel(
-					SharedVelocityFlatSampler, ApplyScreenTransform(SampleUV[j], ColorToVelocity), 0).xyz;
-
-				SampleColor[j] = ColorTexture.SampleLevel(ColorSampler, SampleUV[j], 0);
-				SampleDepth[j] = SampleVelocityDepth.z;
-
-				// in pixels
-				SampleVelocityLength[j] = GetVelocityLengthPixels(SampleVelocityDepth.xy);
-
-				Weight[j] = SampleWeight(CenterDepth, SampleDepth[j], OffsetLength.x, CenterVelocityLength, SampleVelocityLength[j], PixelToSampleScale, SOFT_Z_EXTENT);
+				float2 SampleUV[2];
+				float4 SampleColor[2];
+				float  SampleDepth[2];
+				float  SampleVelocityLength[2];
+				float  ConvolutionWeight[2];
+				float  HoleFillingWeight[2];
+
+				float2 OffsetLength = (float(StepId) + 0.5) + float2(Random - 0.5, 0.5 - Random) * (bDoOneDirectionOnly ? 1.0 : 2.0);
+				float2 OffsetFraction = OffsetLength / TotalSteps;
+
+				float WeightOffsetLength = (float)(DirectionId + StepId) + 0.5;
+
+				SampleUV[0] = ColorUV + OffsetFraction.x * SearchVector;
+				SampleUV[1] = ColorUV - OffsetFraction.y * SearchVector;
+
+				SampleUV[0] = clamp(SampleUV[0], Color_UVViewportBilinearMin, Color_UVViewportBilinearMax);
+				SampleUV[1] = clamp(SampleUV[1], Color_UVViewportBilinearMin, Color_UVViewportBilinearMax);
+
+				UNROLL
+				for (uint j = 0; j < 2; j++)
+				{
+					float3 SampleVelocityDepth = VelocityFlatTexture.SampleLevel(
+						SharedVelocityFlatSampler, ApplyScreenTransform(SampleUV[j], ColorToVelocity), 0).xyz;
+
+					SampleColor[j] = ColorTexture.SampleLevel(ColorSampler, SampleUV[j], 0);
+					SampleDepth[j] = SampleVelocityDepth.z;
+
+					// Decode
+					SampleVelocityDepth.x = GetVelocityLengthPixels(SampleVelocityDepth.x); // TODO: move in velocity flatten
+					SampleVelocityDepth.y = SampleVelocityDepth.y * (2.0 * PI) - PI;
+
+					// in pixels
+					SampleVelocityLength[j] = SampleVelocityDepth.x;
+					
+					ConvolutionWeight[j] = ComputeSampleConvolutionWeight(
+						SampleDepth[j], SampleVelocityDepth.x, SampleVelocityDepth.y,
+						OffsetLength, BlurAngle, PixelToSampleScale);
+
+					HoleFillingWeight[j] = ComputeCenterOrSampleWeight(
+						CenterDepth, CenterVelocityLength, CenterVelocityAngle,
+						SampleDepth[j], SampleVelocityDepth.x, SampleVelocityDepth.y,
+						WeightOffsetLength, BlurAngle, PixelToSampleScale, SOFT_Z_EXTENT);
+				}
+		
+				{
+					bool2 Mirror = bool2(SampleDepth[0] > SampleDepth[1], SampleVelocityLength[0] < SampleVelocityLength[1]);
+					HoleFillingWeight[0] = all(Mirror) ? HoleFillingWeight[1] : HoleFillingWeight[0];
+					HoleFillingWeight[1] = any(Mirror) ? HoleFillingWeight[1] : HoleFillingWeight[0];
+				}
+				
+				#if CONFIG_MAX_RANGE_SIZE > 1
+				{
+					DirectionalColor[DirectionId] += ConvolutionWeight[0] * SampleColor[0] + ConvolutionWeight[1] * SampleColor[1];
+					DirectionalColorWeight[DirectionId] += ConvolutionWeight[0] + ConvolutionWeight[1];
+				
+					DepthAccum[DirectionId] += HoleFillingWeight[0] * SampleDepth[0] + HoleFillingWeight[1] * SampleDepth[1];
+					DepthSquareAccum[DirectionId] += HoleFillingWeight[0] * Square(SampleDepth[0]) + HoleFillingWeight[1] * Square(SampleDepth[1]);
+					DepthAwareWeight[DirectionId] += HoleFillingWeight[0] + HoleFillingWeight[1];
+
+				
+					// Measure how much should hole fill DirectionalColor with.
+					{
+						float HoleFillingWeight0 = saturate(HoleFillingWeight[0] - ConvolutionWeight[0]);
+						float HoleFillingWeight1 = saturate(HoleFillingWeight[1] - ConvolutionWeight[1]);
+
+						//float HoleFillingWeight = ComputeSampleHoleFillWeight(CenterDepth, SampleDepth[j], SOFT_Z_EXTENT);
+						
+						HoleFillWeightAccum[DirectionId] += HoleFillingWeight0 + HoleFillingWeight1;
+					
+						// Build a hole filling along the major directional blur.
+						if (bAccumulateHoleFillColor)
+						{
+							HoleFillColor += HoleFillingWeight0 * SampleColor[0] + HoleFillingWeight1 * SampleColor[1];
+							HoleFillColorWeight += HoleFillingWeight0 + HoleFillingWeight1;
+						}
+					}
+
+				}
+				#else
+				{
+					DirectionalColor[DirectionId] += HoleFillingWeight[0] * SampleColor[0] + HoleFillingWeight[1] * SampleColor[1];
+					DirectionalColorWeight[DirectionId] += HoleFillingWeight[0] + HoleFillingWeight[1];
+				}
+				#endif
+			} // for (uint StepId = 0; StepId < StepCount; StepId += CONFIG_MAX_RANGE_SIZE)
+		} // for (uint DirectionId = 0; DirectionId < CONFIG_MAX_RANGE_SIZE; DirectionId++)
+		
+		#if CONFIG_MAX_RANGE_SIZE > 1
+		{
+			float InvStepCount = rcp(float(StepPerDirectionCount));
+			float InvSampleCount = rcp(float(SamplePerDirection));
+			if (bDoOneDirectionOnly)
+			{
+				DirectionalColor[0] += DirectionalColor[1];
+				DirectionalColorWeight[0] += DirectionalColorWeight[1];
+				HoleFillWeightAccum[0] += HoleFillWeightAccum[1];
+		
+				DirectionalColor[1] = 0.0;
+				DirectionalColorWeight[1] = 0.0;
+				HoleFillWeightAccum[1] = 0.0;
+		
+				InvStepCount = rcp(float(TotalSteps));
+				InvSampleCount = rcp(float(CONFIG_SAMPLE_COUNT));
 			}
 		
-			bool2 Mirror = bool2(SampleDepth[0] > SampleDepth[1], SampleVelocityLength[1] > SampleVelocityLength[0]);
-			Weight[0] = all(Mirror) ? Weight[1] : Weight[0];
-			Weight[1] = any(Mirror) ? Weight[1] : Weight[0];
-			
-			ColorAccum += Weight[0] * SampleColor[0] + Weight[1] * SampleColor[1];
-			ColorAccumWeight += Weight[0] + Weight[1];
-		}
-	
-		ColorAccum			*= 0.5 / StepCount;
-		ColorAccumWeight	*= 0.5 / StepCount;
+			float Velocity1Translucency;
+			{
+				float AvgDepthSquare0 = NormalizeAccumulator(DepthSquareAccum[0], DepthAwareWeight[0]);
+				float AvgDepth0 = NormalizeAccumulator(DepthAccum[0], DepthAwareWeight[0]);
+				float AvgDepth1 = NormalizeAccumulator(DepthAccum[1], DepthAwareWeight[1]);
 
-		OutColor = ColorAccum + (1 - ColorAccumWeight) * CenterColor;
-	}
-	
-#if USE_POST_MOTION_BLUR_TRANSLUCENCY
-	OutColor.rgb = OutColor.rgb * PostMotionBlurTranslucency.a + PostMotionBlurTranslucency.rgb;
-#endif
+				float Variance0 = AvgDepthSquare0 - Square(AvgDepth0);
 
-#if !POST_PROCESS_ALPHA
-	OutColor.a = 0;
-#endif
+				Velocity1Translucency = saturate(2.0 * Variance0 / (Variance0 + Square(AvgDepth1 - AvgDepth0)));
+			}
 
-	return OutColor;
-}
+			float4 NormalizedHoleFillColor0 = NormalizeAccumulator(HoleFillColor, HoleFillColorWeight);
+			float4 NormalizedColor0;
+			float Opacity0;
+			bool bValidColorOutput0;
+			NormalizeAccumulatorWithHoleFill(
+				DirectionalColor[0],
+				DirectionalColorWeight[0],
+				HoleFillWeightAccum[0],
+				DepthAwareWeight[0],
+				NormalizedHoleFillColor0,
+				InvSampleCount,
+				/* out */ NormalizedColor0,
+				/* out */ Opacity0,
+				/* out */ bValidColorOutput0);
+
+			float4 NormalizedHoleFillColor1 = lerp(NormalizedColor0, NormalizedHoleFillColor0, bValidColorOutput0 ? Velocity1Translucency : 1.0);
+			float4 NormalizedColor1;
+			float Opacity1;
+			bool bValidColorOutput1;
+			NormalizeAccumulatorWithHoleFill(
+				DirectionalColor[1],
+				DirectionalColorWeight[1],
+				HoleFillWeightAccum[1],
+				DepthAwareWeight[1],
+				NormalizedHoleFillColor1,
+				InvSampleCount,
+				/* out */ NormalizedColor1,
+				/* out */ Opacity1,
+				/* out */ bValidColorOutput1);
+		
+			Opacity1 *= saturate(1.0 - Velocity1Translucency);
 
-void MainPS(
-	in noperspective float4 UVAndScreenPos : TEXCOORD0,
-	in float4 SvPosition : SV_Position,
-	out float4 OutColor : SV_Target0)
-{
-	float2 ColorUV = UVAndScreenPos.xy;
-	OutColor = MainMotionBlurCommon(ColorUV, floor(SvPosition.xy));
-}
+			float CenterColorContribution = saturate(1.0 - Opacity0) * saturate(1.0 - Opacity1);
+			float InvTotalWeight = rcp(Opacity0 + Opacity1);
 
-#if COMPUTESHADER
-RWTexture2D<float4> OutColorTexture;
+			OutColor = (NormalizedColor0 * Opacity0 + NormalizedColor1 * Opacity1) * -min(-InvTotalWeight * (1.0 - CenterColorContribution), 0.0) + CenterColorContribution * CenterColor;
+		}
+		#else // CONFIG_MAX_RANGE_SIZE == 1
+		{
+			const uint DirectionId = 0;
 
-[numthreads(THREADGROUP_SIZEX, THREADGROUP_SIZEY, 1)]
-void MainCS(uint2 DispatchThreadId : SV_DispatchThreadID)
-{
-	uint2 iColorPixelPos = DispatchThreadId + Color_ViewportMin;
-	float2 ColorPixelPos = float2(iColorPixelPos);
-	float2 ColorUV = (ColorPixelPos + 0.5) * Color_ExtentInverse;
+			DirectionalColor[DirectionId] *= 0.5 / float(StepPerDirectionCount);
+			DirectionalColorWeight[DirectionId] *= 0.5 / float(StepPerDirectionCount);
 
-	float4 OutColor = MainMotionBlurCommon(ColorUV, ColorPixelPos);
+			OutColor = DirectionalColor[DirectionId] + (1.0 - DirectionalColorWeight[DirectionId]) * CenterColor;
+		}
+		#endif
+	}
+	
+	#if USE_POST_MOTION_BLUR_TRANSLUCENCY
+		OutColor.rgb = OutColor.rgb * PostMotionBlurTranslucency.a + PostMotionBlurTranslucency.rgb;
+	#endif
 
+	#if !POST_PROCESS_ALPHA
+		OutColor.a = 0;
+	#endif
+	
 	if (all(iColorPixelPos < Color_ViewportMax))
 	{
 		OutColorTexture[iColorPixelPos] = OutColor;
+
+		#if DEBUG_OUTPUT
+			DebugOutput[iColorPixelPos] = Debug;
+		#endif
 	}
 }
-#endif
+
diff --git a/Engine/Shaders/Private/MotionBlur/MotionBlurCommon.ush b/Engine/Shaders/Private/MotionBlur/MotionBlurCommon.ush
--- a/Engine/Shaders/Private/MotionBlur/MotionBlurCommon.ush
+++ b/Engine/Shaders/Private/MotionBlur/MotionBlurCommon.ush
@@ -5,6 +5,15 @@
 #include "../ScreenPass.ush"
 
 
+//------------------------------------------------------- CONFIG
+
+#define DEBUG_OUTPUT 0
+
+#ifndef CONFIG_MAX_RANGE_SIZE
+	#error need to set CONFIG_MAX_RANGE_SIZE
+#endif
+
+
 //------------------------------------------------------- FUNCTIONS
 
 float4 MinMaxLength(float4 v0, float4 v1)
@@ -35,3 +44,131 @@ float2 PolarToCartesian(float2 Velocity)
 	sincos(Angle, Velocity.y, Velocity.x);
 	return Velocity * Length;
 }
+
+float GetPolarRelativeAngle(float V0, float V1)
+{
+	float AbsDiff = abs(V0 - V1);
+
+	if (AbsDiff > PI)
+	{
+		AbsDiff = 2.0 * PI - AbsDiff;
+	}
+
+	return AbsDiff;
+}
+
+float2 GetMaxPolarvelocity(float2 V0, float2 V1)
+{
+	return V0.x > V1.x ? V0 : V1;
+}
+
+float2 GetMinPolarvelocity(float2 V0, float2 V1)
+{
+	return V0.x > V1.x ? V1 : V0;
+}
+
+float2 GetMaxPolarvelocity(float2 V0, float2 V1, float2 V2)
+{
+	return GetMaxPolarvelocity(GetMaxPolarvelocity(V0, V1), V2);
+}
+
+
+//------------------------------------------------------- TILED VELOCITY RANGE FUNCTIONS
+
+struct FVelocityRange
+{
+	float2 Max[CONFIG_MAX_RANGE_SIZE];
+	float2 Min;
+};
+
+FVelocityRange SetupPolarVelocityRange(float2 PolarVelocity)
+{
+	FVelocityRange V;
+	V.Max[0] = PolarVelocity;
+
+	#if CONFIG_MAX_RANGE_SIZE > 1
+		V.Max[1] = 0.0;
+	#endif
+
+	V.Min = PolarVelocity;
+	return V;
+}
+
+FVelocityRange ReducePolarVelocityRange(FVelocityRange V0, FVelocityRange V1)
+{
+	FVelocityRange V;
+	V.Max[0] = GetMaxPolarvelocity(V0.Max[0], V1.Max[0]);
+	
+	#if CONFIG_MAX_RANGE_SIZE > 1
+	if (GetPolarRelativeAngle(V0.Max[0].y, V1.Max[0].y) < ((PI / 180) * 5.0) || 0)
+	{
+		V.Max[1] = GetMaxPolarvelocity(V0.Max[1], V1.Max[1]);
+	}
+	else
+	{
+		V.Max[1] = GetMaxPolarvelocity(GetMinPolarvelocity(V0.Max[0], V1.Max[0]), V0.Max[1], V1.Max[1]);
+	}
+	#endif
+	
+	V.Min = GetMinPolarvelocity(V0.Min, V1.Min);
+
+	return V;
+}
+
+FVelocityRange CartesianToPolar(FVelocityRange CartesianRange)
+{
+	FVelocityRange PolarRange;
+	PolarRange.Max[0] = CartesianToPolar(CartesianRange.Max[0]);
+	#if CONFIG_MAX_RANGE_SIZE > 1
+		PolarRange.Max[1] = CartesianToPolar(CartesianRange.Max[1]);
+	#endif
+	PolarRange.Min = CartesianToPolar(CartesianRange.Min);
+	return PolarRange;
+}
+
+FVelocityRange PolarToCartesian(FVelocityRange PolarRange)
+{
+	FVelocityRange CartesianRange;
+	CartesianRange.Max[0] = PolarToCartesian(PolarRange.Max[0]);
+	#if CONFIG_MAX_RANGE_SIZE > 1
+		CartesianRange.Max[1] = PolarToCartesian(PolarRange.Max[1]);
+	#endif
+	CartesianRange.Min = PolarToCartesian(PolarRange.Min);
+	return CartesianRange;
+}
+
+void StoreVelocityRange(
+	RWTexture2D<float4> VelocityFlatTexture_0,
+	RWTexture2D<float4>	VelocityTileTexture_1,
+	uint2 PixelCoord,
+	FVelocityRange CartesianRange)
+{
+	VelocityFlatTexture_0[PixelCoord] = float4(CartesianRange.Min, CartesianRange.Max[0]);
+	#if CONFIG_MAX_RANGE_SIZE > 1
+		VelocityTileTexture_1[PixelCoord] = float4(CartesianRange.Max[1], 0.0, 0.0);
+	#endif
+}
+
+FVelocityRange DecodeVelocityRange(
+	float4 Raw0,
+	float4 Raw1)
+{
+	FVelocityRange CartesianRange;
+	CartesianRange.Max[0] = Raw0.zw;
+	#if CONFIG_MAX_RANGE_SIZE > 1
+		CartesianRange.Max[1] = Raw1.xy;
+	#endif
+	CartesianRange.Min = Raw0.xy;
+	return CartesianRange;
+}
+
+FVelocityRange LoadVelocityRange(
+	Texture2D<float4> VelocityFlatTexture_0,
+	Texture2D<float4> VelocityTileTexture_1,
+	uint2 PixelCoord)
+{
+	float4 Raw0 = VelocityFlatTexture_0[PixelCoord];
+	float4 Raw1 = VelocityTileTexture_1[PixelCoord];
+
+	return DecodeVelocityRange(Raw0, Raw1);
+}
diff --git a/Engine/Shaders/Private/MotionBlur/MotionBlurTileGather.usf b/Engine/Shaders/Private/MotionBlur/MotionBlurTileGather.usf
--- a/Engine/Shaders/Private/MotionBlur/MotionBlurTileGather.usf
+++ b/Engine/Shaders/Private/MotionBlur/MotionBlurTileGather.usf
@@ -1,13 +1,8 @@
 // Copyright Epic Games, Inc. All Rights Reserved.
 
-#include "MotionBlurCommon.ush"
-
+#define CONFIG_MAX_RANGE_SIZE DIM_BLUR_DIRECTIONS
 
-//------------------------------------------------------- CONFIG
-
-#define THREADGROUP_SIZEX		16
-#define THREADGROUP_SIZEY		16
-#define THREADGROUP_TOTALSIZE	(THREADGROUP_SIZEX * THREADGROUP_SIZEY)
+#include "MotionBlurCommon.ush"
 
 
 //------------------------------------------------------- PARAMETERS
@@ -15,9 +10,14 @@
 SCREEN_PASS_TEXTURE_VIEWPORT(VelocityTile)
 float MotionBlur_VelocityScaleForTiles;
 
-Texture2D VelocityTileTexture;
+Texture2D CenterVelocityTileTextures_Textures_0;
+Texture2D CenterVelocityTileTextures_Textures_1;
 
-RWTexture2D<float4>	OutVelocityTileTexture;
+Texture2D VelocityTileTextures_Textures_0;
+Texture2D VelocityTileTextures_Textures_1;
+
+RWTexture2D<float4>	OutVelocityTile_Textures_0;
+RWTexture2D<float4>	OutVelocityTile_Textures_1;
 
 
 //------------------------------------------------------- ENTRY POINT
@@ -31,8 +31,8 @@ void MainCS(
 {
 	uint2 PixelPos = DispatchThreadId.xy;
 	float2 UV = ((float2)PixelPos.xy + 0.5) * VelocityTile_ExtentInverse;
-
-	float4 MinMaxVelocity = VelocityTileTexture[PixelPos];
+	
+	FVelocityRange VelocityPolarRange = CartesianToPolar(LoadVelocityRange(CenterVelocityTileTextures_Textures_0, CenterVelocityTileTextures_Textures_1, PixelPos));
 
 	// Scatter as gather
 	for(int x = -3; x <= 3; x++)
@@ -48,9 +48,10 @@ void MainCS(
 			bool2 bInsideViewport = 0 <= SampleIndex && SampleIndex < (int2)VelocityTile_ViewportMax;
 			if (!all(bInsideViewport))
 				continue;
+			
+			FVelocityRange NeighborVelocityRange = LoadVelocityRange(VelocityTileTextures_Textures_0, VelocityTileTextures_Textures_1, SampleIndex);
 
-			float4 ScatterMinMax = VelocityTileTexture[SampleIndex];
-			float2 MaxVelocity = ScatterMinMax.zw;
+			float2 MaxVelocity = NeighborVelocityRange.Max[0];
 
 			float2 VelocityPixels = MaxVelocity * MotionBlur_VelocityScaleForTiles;
 
@@ -78,7 +79,7 @@ void MainCS(
 			bool2 bInsideQuad = abs(PixelCenterOnQuad) < QuadExtent;
 			if (all(bInsideQuad)) 
 			{
-				MinMaxVelocity = MinMaxLength(MinMaxVelocity, ScatterMinMax);
+				VelocityPolarRange = ReducePolarVelocityRange(VelocityPolarRange, CartesianToPolar(NeighborVelocityRange));
 			}
 		}
 	}
@@ -86,6 +87,6 @@ void MainCS(
 	BRANCH
 	if (all(PixelPos.xy < VelocityTile_ViewportMax))
 	{
-		OutVelocityTileTexture[PixelPos] = MinMaxVelocity;
+		StoreVelocityRange(OutVelocityTile_Textures_0, OutVelocityTile_Textures_1, PixelPos, PolarToCartesian(VelocityPolarRange));
 	}
 }
diff --git a/Engine/Shaders/Private/MotionBlur/MotionBlurTileScatter.usf b/Engine/Shaders/Private/MotionBlur/MotionBlurTileScatter.usf
--- a/Engine/Shaders/Private/MotionBlur/MotionBlurTileScatter.usf
+++ b/Engine/Shaders/Private/MotionBlur/MotionBlurTileScatter.usf
@@ -1,5 +1,8 @@
 // Copyright Epic Games, Inc. All Rights Reserved.
 
+#define CONFIG_MAX_RANGE_SIZE 1
+
+
 #include "MotionBlurCommon.ush"
 
 //------------------------------------------------------- CONFIG
@@ -14,7 +17,8 @@ SCREEN_PASS_TEXTURE_VIEWPORT(VelocityTile)
 float MotionBlur_VelocityScaleForTiles;
 uint ScatterPass;
 
-Texture2D VelocityTileTexture;
+Texture2D VelocityTileTextures_Textures_0;
+Texture2D VelocityTileTextures_Textures_1;
 
 
 //------------------------------------------------------- ENTRY POINT
@@ -49,7 +53,7 @@ void MainVS(
 
 	float2 SvPosition = PixelPos + 0.5;
 
-	float4 MinMaxVelocity = VelocityTileTexture[PixelPos];
+	float4 MinMaxVelocity = VelocityTileTextures_Textures_0[PixelPos];
 	OutColor = MinMaxVelocity;
 
 	float4 MinMaxVelocityPixels = MinMaxVelocity * MotionBlur_VelocityScaleForTiles;
diff --git a/Engine/Shaders/Private/MotionBlur/MotionBlurVelocityFlatten.usf b/Engine/Shaders/Private/MotionBlur/MotionBlurVelocityFlatten.usf
--- a/Engine/Shaders/Private/MotionBlur/MotionBlurVelocityFlatten.usf
+++ b/Engine/Shaders/Private/MotionBlur/MotionBlurVelocityFlatten.usf
@@ -1,5 +1,7 @@
 // Copyright Epic Games, Inc. All Rights Reserved.
 
+#define CONFIG_MAX_RANGE_SIZE DIM_BLUR_DIRECTIONS
+
 #include "MotionBlurVelocityFlatten.ush"
 
 
@@ -34,15 +36,66 @@ Texture2D VelocityTileTexture;
 Texture2D DepthTexture;
 
 RWTexture2D<float4>	OutVelocityFlatTexture;
-RWTexture2D<float4>	OutVelocityTileTexture;
+RWTexture2D<float4>	OutVelocityTile_Textures_0;
+RWTexture2D<float4>	OutVelocityTile_Textures_1;
+
+RWTexture2D<float4>	DebugOutput;
+
 
 //------------------------------------------------------- LDS
 
-groupshared float4 Shared[THREADGROUP_TOTALSIZE];
+#if DEBUG_OUTPUT
+groupshared float4 Debug[THREADGROUP_TOTALSIZE];
+#endif
+
+groupshared float4 SharedArray0[THREADGROUP_TOTALSIZE];
+groupshared float2 SharedArray1[THREADGROUP_TOTALSIZE];
 
 
 //------------------------------------------------------- FUNCTIONS
 
+void WritePolarVelocityRangeToLDS(uint LDSSlotId, FVelocityRange V)
+{
+	SharedArray0[LDSSlotId] = float4(V.Max[0], V.Min);
+	
+	#if CONFIG_MAX_RANGE_SIZE > 1
+		SharedArray1[LDSSlotId] = V.Max[1];
+	#endif
+}
+
+FVelocityRange ReadPolarVelocityRangeFromLDS(uint LDSSlotId)
+{
+	FVelocityRange V;
+	V.Max[0] = SharedArray0[LDSSlotId].xy;
+	#if CONFIG_MAX_RANGE_SIZE > 1
+		V.Max[1] = SharedArray1[LDSSlotId].xy;
+	#endif
+	V.Min = SharedArray0[LDSSlotId].zw;
+	return V;
+}
+
+void VelovityFlattenStep(uint GroupIndex, const uint GroupSize, inout FVelocityRange V)
+{
+	if (GroupIndex < GroupSize)
+	{
+		FVelocityRange V1 = ReadPolarVelocityRangeFromLDS(GroupIndex + GroupSize);
+		
+		//if (GroupSize == uint(View.GeneralPurposeTweak))
+		//{
+		//	Debug[GroupIndex].r = V.Max[1].g;
+		//}
+
+		V = ReducePolarVelocityRange(V, V1);
+		
+		//if (GroupSize == uint(View.GeneralPurposeTweak))
+		//{
+		//	Debug[GroupIndex].g = V.Max[1].g;
+		//}
+
+		WritePolarVelocityRangeToLDS(GroupIndex, V);
+	}
+}
+
 float2 GetCameraMotionVelocity(uint3 DispatchThreadId, float Depth, float4x4 ClipToPrevClip)
 {
 	// Compute velocity due to camera motion.
@@ -64,6 +117,10 @@ void MainCS(
 	uint3 GroupThreadId : SV_GroupThreadID,
 	uint GroupIndex : SV_GroupIndex) 
 { 
+#if DEBUG_OUTPUT
+	Debug[GroupIndex] = 0.0;
+#endif
+
 	int2 PixelPos = min(DispatchThreadId.xy + Velocity_ViewportMin, Velocity_ViewportMax - 1);
 	int2 DepthOffset = 0;
 
@@ -141,55 +198,64 @@ void MainCS(
 	bool bInsideViewport = all(PixelPos.xy < Velocity_ViewportMax);
 
 	// 11:11:10  (VelocityLength, VelocityAngle, Depth)
-	float2 EncodedPolarVelocity;
-	EncodedPolarVelocity.x = VelocityPolar.x;
-	EncodedPolarVelocity.y = VelocityPolar.y * (0.5 / PI) + 0.5;
-
-	BRANCH
-	if (bInsideViewport)
 	{
-		OutVelocityFlatTexture[PixelPos] = float3(EncodedPolarVelocity, ConvertFromDeviceZ(Depth)).xyzz;
+		float2 EncodedPolarVelocity;
+		EncodedPolarVelocity.x = VelocityPolar.x;
+		EncodedPolarVelocity.y = VelocityPolar.y * (0.5 / PI) + 0.5;
+
+		BRANCH
+		if (bInsideViewport)
+		{
+			OutVelocityFlatTexture[PixelPos] = float3(EncodedPolarVelocity, ConvertFromDeviceZ(Depth)).xyzz;
+		}
 	}
 
 	// Limit velocity
 	VelocityPolar.x = min(VelocityPolar.x, MotionBlur_VelocityMax / MotionBlur_VelocityScale);
 
-	float4 VelocityMinMax = VelocityPolar.xyxy;
-	VelocityMinMax.x = bInsideViewport ? VelocityMinMax.x : 2;
-	VelocityMinMax.z = bInsideViewport ? VelocityMinMax.z : 0;
+	FVelocityRange VelocityPolarRange = SetupPolarVelocityRange(VelocityPolar);
+
+	// TODO
+	//float4 VelocityMinMax = VelocityPolar.xyxy;
+	//VelocityMinMax.x = bInsideViewport ? VelocityMinMax.x : 2;
+	//VelocityMinMax.z = bInsideViewport ? VelocityMinMax.z : 0;
 
-	Shared[GroupIndex] = VelocityMinMax;
+	WritePolarVelocityRangeToLDS(GroupIndex, VelocityPolarRange);
 
 	GroupMemoryBarrierWithGroupSync();
 
 #if THREADGROUP_TOTALSIZE > 512
-	if (GroupIndex < 512) Shared[GroupIndex] = MinMaxLengthPolar(Shared[GroupIndex], Shared[GroupIndex + 512]);
+	VelovityFlattenStep(GroupIndex, /* GroupSize = */ 512, /* inout */ VelocityPolarRange);
 	GroupMemoryBarrierWithGroupSync();
 #endif
 #if THREADGROUP_TOTALSIZE > 256
-	if (GroupIndex < 256) Shared[GroupIndex] = MinMaxLengthPolar(Shared[GroupIndex], Shared[GroupIndex + 256]);
+	VelovityFlattenStep(GroupIndex, /* GroupSize = */ 256, /* inout */ VelocityPolarRange);
 	GroupMemoryBarrierWithGroupSync();
 #endif
 #if THREADGROUP_TOTALSIZE > 128
-	if (GroupIndex < 128) Shared[GroupIndex] = MinMaxLengthPolar(Shared[GroupIndex], Shared[GroupIndex + 128]);
+	VelovityFlattenStep(GroupIndex, /* GroupSize = */ 128, /* inout */ VelocityPolarRange);
 	GroupMemoryBarrierWithGroupSync();
 #endif
 #if THREADGROUP_TOTALSIZE > 64
-	if (GroupIndex <  64) Shared[GroupIndex] = MinMaxLengthPolar(Shared[GroupIndex], Shared[GroupIndex +  64]);
+	VelovityFlattenStep(GroupIndex, /* GroupSize = */ 64, /* inout */ VelocityPolarRange);
 	GroupMemoryBarrierWithGroupSync();
 #endif
 
 	// Safe for vector sizes 32 or larger, AMD and NV
 	// TODO Intel variable size vector
-	if (GroupIndex < 32) Shared[GroupIndex] = MinMaxLengthPolar(Shared[GroupIndex], Shared[GroupIndex + 32]);
-	if (GroupIndex < 16) Shared[GroupIndex] = MinMaxLengthPolar(Shared[GroupIndex], Shared[GroupIndex + 16]);
-	if (GroupIndex <  8) Shared[GroupIndex] = MinMaxLengthPolar(Shared[GroupIndex], Shared[GroupIndex +  8]);
-	if (GroupIndex <  4) Shared[GroupIndex] = MinMaxLengthPolar(Shared[GroupIndex], Shared[GroupIndex +  4]);
-	if (GroupIndex <  2) Shared[GroupIndex] = MinMaxLengthPolar(Shared[GroupIndex], Shared[GroupIndex +  2]);
-	if (GroupIndex <  1) Shared[GroupIndex] = MinMaxLengthPolar(Shared[GroupIndex], Shared[GroupIndex +  1]);
+	VelovityFlattenStep(GroupIndex, /* GroupSize = */ 32, /* inout */ VelocityPolarRange);
+	VelovityFlattenStep(GroupIndex, /* GroupSize = */ 16, /* inout */ VelocityPolarRange);
+	VelovityFlattenStep(GroupIndex, /* GroupSize = */  8, /* inout */ VelocityPolarRange);
+	VelovityFlattenStep(GroupIndex, /* GroupSize = */  4, /* inout */ VelocityPolarRange);
+	VelovityFlattenStep(GroupIndex, /* GroupSize = */  2, /* inout */ VelocityPolarRange);
+	VelovityFlattenStep(GroupIndex, /* GroupSize = */  1, /* inout */ VelocityPolarRange);
 
 	if (GroupIndex == 0)
 	{
-		OutVelocityTileTexture[GroupId.xy] = float4(PolarToCartesian(Shared[0].xy), PolarToCartesian(Shared[0].zw));
+		StoreVelocityRange(OutVelocityTile_Textures_0, OutVelocityTile_Textures_1, GroupId, PolarToCartesian(VelocityPolarRange));
 	}
+	
+#if DEBUG_OUTPUT
+	DebugOutput[PixelPos] = Debug[GroupIndex];
+#endif
 }
diff --git a/Engine/Shaders/Private/MotionBlur/MotionBlurVisualize.usf b/Engine/Shaders/Private/MotionBlur/MotionBlurVisualize.usf
--- a/Engine/Shaders/Private/MotionBlur/MotionBlurVisualize.usf
+++ b/Engine/Shaders/Private/MotionBlur/MotionBlurVisualize.usf
@@ -1,7 +1,8 @@
 // Copyright Epic Games, Inc. All Rights Reserved.
 
+#define CONFIG_MAX_RANGE_SIZE 1
+
 #include "MotionBlurCommon.ush"
-//#include "Random.ush"
 
 
 //------------------------------------------------------- CONSTANTS
diff --git a/Engine/Source/Runtime/Renderer/Private/PostProcess/PostProcessMotionBlur.cpp b/Engine/Source/Runtime/Renderer/Private/PostProcess/PostProcessMotionBlur.cpp
--- a/Engine/Source/Runtime/Renderer/Private/PostProcess/PostProcessMotionBlur.cpp
+++ b/Engine/Source/Runtime/Renderer/Private/PostProcess/PostProcessMotionBlur.cpp
@@ -38,10 +38,10 @@ namespace
 		TEXT("Adds a second motion blur pass that smooths noise for a higher quality blur."),
 		ECVF_RenderThreadSafe);
 
-	TAutoConsoleVariable<int32> CVarMotionBlurPreferCompute(
-		TEXT("r.MotionBlur.PreferCompute"),
-		0,
-		TEXT("Will use compute shaders for motion blur pass."),
+	TAutoConsoleVariable<int32> CVarMotionBlurDirections(
+		TEXT("r.MotionBlur.Directions"),
+		1,
+		TEXT("Number of bluring direction (default = 1)."),
 		ECVF_RenderThreadSafe);
 
 	TAutoConsoleVariable<int32> CVarAllowMotionBlurInVR(
@@ -152,6 +152,11 @@ EMotionBlurFilter GetMotionBlurFilter()
 	return CVarMotionBlurSeparable.GetValueOnRenderThread() != 0 ? EMotionBlurFilter::Separable : EMotionBlurFilter::Unified;
 }
 
+int32 GetMotionBlurDirections()
+{
+	return FMath::Clamp(CVarMotionBlurDirections.GetValueOnRenderThread(), 1, 2);
+}
+
 FRHISamplerState* GetMotionBlurColorSampler()
 {
 	bool bFiltered = false;
@@ -187,6 +192,60 @@ FRHISamplerState* GetPostMotionBlurTranslucencySampler(bool bUpscale)
 	}
 }
 
+FRDGTextureUAVRef CreateDebugUAV(FRDGBuilder& GraphBuilder, const FIntPoint& Extent, const TCHAR* DebugName)
+{
+#if !UE_BUILD_SHIPPING
+	FRDGTextureDesc Desc = FRDGTextureDesc::Create2D(
+		Extent,
+		PF_FloatRGBA,
+		FClearValueBinding::None,
+		TexCreate_ShaderResource | TexCreate_UAV);
+
+	FRDGTextureRef Texture = GraphBuilder.CreateTexture(Desc, DebugName);
+
+	return GraphBuilder.CreateUAV(Texture);
+#else
+	return nullptr;
+#endif
+}
+
+BEGIN_SHADER_PARAMETER_STRUCT(FVecocityTileTextures, )
+	SHADER_PARAMETER_RDG_TEXTURE_ARRAY(Texture2D, Textures, [2])
+END_SHADER_PARAMETER_STRUCT()
+
+BEGIN_SHADER_PARAMETER_STRUCT(FVecocityTileUAVs, )
+	SHADER_PARAMETER_RDG_TEXTURE_UAV_ARRAY(Texture2D, Textures, [2])
+END_SHADER_PARAMETER_STRUCT()
+
+FVecocityTileTextures CreateVecocityTileTextures(FRDGBuilder& GraphBuilder, FIntPoint VelocityTileCount, const TCHAR* DebugName, bool ScatterDilatation = false)
+{
+	FRDGTextureDesc Desc = FRDGTextureDesc::Create2D(
+		VelocityTileCount,
+		PF_FloatRGBA,
+		FClearValueBinding::None,
+		GFastVRamConfig.MotionBlur | TexCreate_ShaderResource | (ScatterDilatation ? TexCreate_RenderTargetable : TexCreate_UAV));
+
+	FVecocityTileTextures Textures;
+	Textures.Textures[0] = GraphBuilder.CreateTexture(Desc, DebugName);
+
+	Desc.Format = PF_G16R16F;
+	Textures.Textures[1] = GraphBuilder.CreateTexture(Desc, DebugName);
+	return Textures;
+}
+
+FVecocityTileUAVs CreateUAVs(FRDGBuilder& GraphBuilder, FVecocityTileTextures VecocityTileTextures)
+{
+	FVecocityTileUAVs UAVs;
+	for (int32 i = 0; i < UAVs.Textures.Num(); i++)
+	{
+		if (VecocityTileTextures.Textures[i])
+		{
+			UAVs.Textures[i] = GraphBuilder.CreateUAV(VecocityTileTextures.Textures[i]);
+		}
+	}
+	return UAVs;
+}
+
 // Set of common shader parameters shared by all motion blur shaders.
 BEGIN_SHADER_PARAMETER_STRUCT(FMotionBlurParameters, )
 	SHADER_PARAMETER(float, AspectRatio)
@@ -234,13 +293,15 @@ public:
 	{}
 };
 
+class FMotionBlurDirections : SHADER_PERMUTATION_SPARSE_INT("DIM_BLUR_DIRECTIONS", 1, 2);
+
 class FMotionBlurVelocityFlattenCS : public FMotionBlurShader
 {
 public:
 	static const uint32 ThreadGroupSize = 16;
 
 	class FCameraMotionBlurMode : SHADER_PERMUTATION_INT("CAMERA_MOTION_BLUR_MODE", 3);
-	using FPermutationDomain = TShaderPermutationDomain<FCameraMotionBlurMode>;
+	using FPermutationDomain = TShaderPermutationDomain<FCameraMotionBlurMode, FMotionBlurDirections>;
 
 	DECLARE_GLOBAL_SHADER(FMotionBlurVelocityFlattenCS);
 	SHADER_USE_PARAMETER_STRUCT(FMotionBlurVelocityFlattenCS, FMotionBlurShader);
@@ -249,18 +310,21 @@ public:
 		SHADER_PARAMETER_STRUCT_REF(FViewUniformShaderParameters, View)
 		SHADER_PARAMETER_STRUCT(FMotionBlurParameters, MotionBlur)
 		SHADER_PARAMETER_STRUCT(FScreenPassTextureViewportParameters, Velocity)
+		SHADER_PARAMETER(FMatrix44f, ClipToPrevClipOverride)
+
 		SHADER_PARAMETER_RDG_TEXTURE(Texture2D, VelocityTexture)
 		SHADER_PARAMETER_RDG_TEXTURE(Texture2D, DepthTexture)
+
 		SHADER_PARAMETER_RDG_TEXTURE_UAV(RWTexture2D<float4>, OutVelocityFlatTexture)
-		SHADER_PARAMETER_RDG_TEXTURE_UAV(RWTexture2D<float4>, OutVelocityTileTexture)
-		SHADER_PARAMETER(FMatrix44f, ClipToPrevClipOverride)
+		SHADER_PARAMETER_STRUCT(FVecocityTileUAVs, OutVelocityTile)
+		SHADER_PARAMETER_RDG_TEXTURE_UAV(RWTexture2D<float4>, DebugOutput)
 	END_SHADER_PARAMETER_STRUCT()
 };
 
 BEGIN_SHADER_PARAMETER_STRUCT(FMotionBlurVelocityDilateParameters, )
 	SHADER_PARAMETER_STRUCT(FMotionBlurParameters, MotionBlur)
 	SHADER_PARAMETER_STRUCT(FScreenPassTextureViewportParameters, VelocityTile)
-	SHADER_PARAMETER_RDG_TEXTURE(Texture2D, VelocityTileTexture)
+	SHADER_PARAMETER_STRUCT(FVecocityTileTextures, VelocityTileTextures)
 END_SHADER_PARAMETER_STRUCT()
 
 class FMotionBlurVelocityDilateGatherCS : public FMotionBlurShader
@@ -273,8 +337,12 @@ public:
 
 	BEGIN_SHADER_PARAMETER_STRUCT(FParameters, )
 		SHADER_PARAMETER_STRUCT_INCLUDE(FMotionBlurVelocityDilateParameters, Dilate)
-		SHADER_PARAMETER_RDG_TEXTURE_UAV(RWTexture2D<float4>, OutVelocityTileTexture)
+		SHADER_PARAMETER_STRUCT(FVecocityTileTextures, CenterVelocityTileTextures)
+		SHADER_PARAMETER_STRUCT(FVecocityTileUAVs, OutVelocityTile)
+		SHADER_PARAMETER_RDG_TEXTURE_UAV(RWTexture2D<float4>, DebugOutput)
 	END_SHADER_PARAMETER_STRUCT()
+
+	using FPermutationDomain = TShaderPermutationDomain<FMotionBlurDirections>;
 };
 
 enum class EMotionBlurVelocityScatterPass : uint32
@@ -309,7 +377,7 @@ class FMotionBlurVelocityDilateScatterPS : public FMotionBlurShader
 class FMotionBlurQualityDimension : SHADER_PERMUTATION_ENUM_CLASS("MOTION_BLUR_QUALITY", EMotionBlurQuality);
 class FPostMotionBlurTranslucencyDimension : SHADER_PERMUTATION_BOOL("USE_POST_MOTION_BLUR_TRANSLUCENCY");
 
-using FMotionBlurFilterPermutationDomain = TShaderPermutationDomain<FMotionBlurQualityDimension, FPostMotionBlurTranslucencyDimension>;
+using FMotionBlurFilterPermutationDomain = TShaderPermutationDomain<FMotionBlurQualityDimension, FPostMotionBlurTranslucencyDimension, FMotionBlurDirections>;
 
 BEGIN_SHADER_PARAMETER_STRUCT(FMotionBlurFilterParameters, )
 	SHADER_PARAMETER_STRUCT(FMotionBlurParameters, MotionBlur)
@@ -325,7 +393,7 @@ BEGIN_SHADER_PARAMETER_STRUCT(FMotionBlurFilterParameters, )
 
 	SHADER_PARAMETER_RDG_TEXTURE(Texture2D, ColorTexture)
 	SHADER_PARAMETER_RDG_TEXTURE(Texture2D, VelocityFlatTexture)
-	SHADER_PARAMETER_RDG_TEXTURE(Texture2D, VelocityTileTexture)
+	SHADER_PARAMETER_STRUCT(FVecocityTileTextures, VelocityTileTextures)
 	SHADER_PARAMETER_RDG_TEXTURE(Texture2D, PostMotionBlurTranslucencyTexture)
 
 	SHADER_PARAMETER_SAMPLER(SamplerState, ColorSampler)
@@ -335,20 +403,6 @@ BEGIN_SHADER_PARAMETER_STRUCT(FMotionBlurFilterParameters, )
 	SHADER_PARAMETER_SAMPLER(SamplerState, PostMotionBlurTranslucencySampler)
 END_SHADER_PARAMETER_STRUCT()
 
-class FMotionBlurFilterPS : public FMotionBlurShader
-{
-public:
-	DECLARE_GLOBAL_SHADER(FMotionBlurFilterPS);
-	SHADER_USE_PARAMETER_STRUCT(FMotionBlurFilterPS, FMotionBlurShader);
-
-	BEGIN_SHADER_PARAMETER_STRUCT(FParameters, )
-		SHADER_PARAMETER_STRUCT_INCLUDE(FMotionBlurFilterParameters, Filter)
-		RENDER_TARGET_BINDING_SLOTS()
-	END_SHADER_PARAMETER_STRUCT()
-
-	using FPermutationDomain = FMotionBlurFilterPermutationDomain;
-};
-
 class FMotionBlurFilterCS : public FMotionBlurShader
 {
 public:
@@ -358,6 +412,7 @@ public:
 	BEGIN_SHADER_PARAMETER_STRUCT(FParameters, )
 		SHADER_PARAMETER_STRUCT_INCLUDE(FMotionBlurFilterParameters, Filter)
 		SHADER_PARAMETER_RDG_TEXTURE_UAV(RWTexture2D<float4>, OutColorTexture)
+		SHADER_PARAMETER_RDG_TEXTURE_UAV(RWTexture2D<float4>, DebugOutput)
 	END_SHADER_PARAMETER_STRUCT()
 
 	static void ModifyCompilationEnvironment(const FGlobalShaderPermutationParameters& Parameters, FShaderCompilerEnvironment& OutEnvironment)
@@ -398,7 +453,6 @@ IMPLEMENT_GLOBAL_SHADER(FMotionBlurVelocityFlattenCS,       "/Engine/Private/Mot
 IMPLEMENT_GLOBAL_SHADER(FMotionBlurVelocityDilateGatherCS,  "/Engine/Private/MotionBlur/MotionBlurTileGather.usf",      "MainCS", SF_Compute);
 IMPLEMENT_GLOBAL_SHADER(FMotionBlurVelocityDilateScatterVS, "/Engine/Private/MotionBlur/MotionBlurTileScatter.usf",     "MainVS", SF_Vertex);
 IMPLEMENT_GLOBAL_SHADER(FMotionBlurVelocityDilateScatterPS, "/Engine/Private/MotionBlur/MotionBlurTileScatter.usf",     "MainPS", SF_Pixel);
-IMPLEMENT_GLOBAL_SHADER(FMotionBlurFilterPS,                "/Engine/Private/MotionBlur/MotionBlurApply.usf",           "MainPS", SF_Pixel);
 IMPLEMENT_GLOBAL_SHADER(FMotionBlurFilterCS,                "/Engine/Private/MotionBlur/MotionBlurApply.usf",           "MainCS", SF_Compute);
 IMPLEMENT_GLOBAL_SHADER(FMotionBlurVisualizePS,             "/Engine/Private/MotionBlur/MotionBlurVisualize.usf",       "MainPS", SF_Pixel);
 
@@ -451,13 +505,15 @@ void AddMotionBlurVelocityPass(
 	FRDGTextureRef DepthTexture,
 	FRDGTextureRef VelocityTexture,
 	FRDGTextureRef* VelocityFlatTextureOutput,
-	FRDGTextureRef* VelocityTileTextureOutput)
+	FVecocityTileTextures* VelocityTileTexturesOutput)
 {
 	check(ColorTexture);
 	check(DepthTexture);
 	check(VelocityTexture);
 	check(VelocityFlatTextureOutput);
-	check(VelocityTileTextureOutput);
+	check(VelocityTileTexturesOutput);
+
+	const int32 BlurDirections = GetMotionBlurDirections();
 
 	const FIntPoint VelocityTileCount = Viewports.VelocityTile.Extent;
 
@@ -471,14 +527,8 @@ void AddMotionBlurVelocityPass(
 				GFastVRamConfig.VelocityFlat | TexCreate_ShaderResource | TexCreate_UAV),
 			TEXT("MotionBlur.VelocityFlat"));
 
-	FRDGTextureRef VelocityTileTextureSetup =
-		GraphBuilder.CreateTexture(
-			FRDGTextureDesc::Create2D(
-				VelocityTileCount,
-				PF_FloatRGBA,
-				FClearValueBinding::None,
-				GFastVRamConfig.VelocityMax | TexCreate_ShaderResource | TexCreate_UAV),
-			TEXT("MotionBlur.VelocityTile"));
+	FVecocityTileTextures VelocityTileTexturesSetup = CreateVecocityTileTextures(
+		GraphBuilder, VelocityTileCount, TEXT("MotionBlur.VelocityTile"));
 
 	const FMotionBlurParameters MotionBlurParametersNoScale = GetMotionBlurParameters(View, Viewports.Color.Rect.Size(), 1.0f);
 
@@ -490,18 +540,22 @@ void AddMotionBlurVelocityPass(
 		FMotionBlurVelocityFlattenCS::FParameters* PassParameters = GraphBuilder.AllocParameters<FMotionBlurVelocityFlattenCS::FParameters>();
 		PassParameters->MotionBlur = MotionBlurParametersNoScale;
 		PassParameters->View = View.ViewUniformBuffer;
-		PassParameters->Velocity = Viewports.VelocityParameters;
-		PassParameters->DepthTexture = DepthTexture;
-		PassParameters->VelocityTexture = VelocityTexture;
-		PassParameters->OutVelocityFlatTexture = GraphBuilder.CreateUAV(VelocityFlatTexture);
-		PassParameters->OutVelocityTileTexture = GraphBuilder.CreateUAV(VelocityTileTextureSetup);
 		if (bOverrideCameraMotionBlur)
 		{
 			PassParameters->ClipToPrevClipOverride = View.ClipToPrevClipOverride.GetValue();
 		}
 
+		PassParameters->Velocity = Viewports.VelocityParameters;
+		PassParameters->DepthTexture = DepthTexture;
+		PassParameters->VelocityTexture = VelocityTexture;
+
+		PassParameters->OutVelocityFlatTexture = GraphBuilder.CreateUAV(VelocityFlatTexture);
+		PassParameters->OutVelocityTile = CreateUAVs(GraphBuilder, VelocityTileTexturesSetup);
+		PassParameters->DebugOutput = CreateDebugUAV(GraphBuilder, VelocityFlatTexture->Desc.Extent, TEXT("Debug.MotionBlur.Flatten"));
+
 		FMotionBlurVelocityFlattenCS::FPermutationDomain PermutationVector;
 		PermutationVector.Set<FMotionBlurVelocityFlattenCS::FCameraMotionBlurMode>(bEnableCameraMotionBlur ? (bOverrideCameraMotionBlur ? 2 : 1) : 0);
+		PermutationVector.Set<FMotionBlurDirections>(BlurDirections);
 
 		TShaderMapRef<FMotionBlurVelocityFlattenCS> ComputeShader(View.ShaderMap, PermutationVector);
 		FComputeShaderUtils::AddPass(
@@ -514,22 +568,17 @@ void AddMotionBlurVelocityPass(
 
 	bool ScatterDilatation = IsMotionBlurScatterRequired(View, Viewports.Color);
 
-	FRDGTextureRef VelocityTileTexture =
-		GraphBuilder.CreateTexture(
-			FRDGTextureDesc::Create2D(
-				VelocityTileCount,
-				PF_FloatRGBA,
-				FClearValueBinding::None,
-				GFastVRamConfig.MotionBlur | TexCreate_ShaderResource | (ScatterDilatation ? TexCreate_RenderTargetable : TexCreate_UAV)),
-			TEXT("MotionBlur.DilatedVelocityTile"));
-
 	FMotionBlurVelocityDilateParameters VelocityDilateParameters;
 	VelocityDilateParameters.MotionBlur = MotionBlurParametersNoScale;
 	VelocityDilateParameters.VelocityTile = Viewports.VelocityTileParameters;
-	VelocityDilateParameters.VelocityTileTexture = VelocityTileTextureSetup;
+	VelocityDilateParameters.VelocityTileTextures = VelocityTileTexturesSetup;
 
+	FVecocityTileTextures VelocityTileTextures;
 	if (ScatterDilatation)
 	{
+		VelocityTileTextures = CreateVecocityTileTextures(
+			GraphBuilder, VelocityTileCount, TEXT("MotionBlur.ScatteredVelocityTile"), /* ScatterDilatation = */ true);
+
 		FRDGTextureRef VelocityTileDepthTexture =
 			GraphBuilder.CreateTexture(
 				FRDGTextureDesc::Create2D(
@@ -551,14 +600,13 @@ void AddMotionBlurVelocityPass(
 
 		PassParameters->RenderTargets[0] =
 			FRenderTargetBinding(
-				VelocityTileTexture,
+				VelocityTileTextures.Textures[0],
 				ERenderTargetLoadAction::ENoAction);
 
 		TShaderMapRef<FMotionBlurVelocityDilateScatterVS> VertexShader(View.ShaderMap);
 		TShaderMapRef<FMotionBlurVelocityDilateScatterPS> PixelShader(View.ShaderMap);
 		
-		ValidateShaderParameters(VertexShader, *PassParameters);
-		ValidateShaderParameters(PixelShader, *PassParameters);
+		ClearUnusedGraphResources(VertexShader, PixelShader, PassParameters);
 
 		GraphBuilder.AddPass(
 			RDG_EVENT_NAME("VelocityTileScatter %dx%d", VelocityTileCount.X, VelocityTileCount.Y),
@@ -609,23 +657,53 @@ void AddMotionBlurVelocityPass(
 			}
 		});
 	}
-	else
+
+	// ScatterGather the dilatation
+	if (!ScatterDilatation || BlurDirections > 1)
 	{
+		FVecocityTileTextures DilatedTileTextures = CreateVecocityTileTextures(
+			GraphBuilder, VelocityTileCount, TEXT("MotionBlur.GatheredVelocityTile"));
+
 		FMotionBlurVelocityDilateGatherCS::FParameters* PassParameters = GraphBuilder.AllocParameters<FMotionBlurVelocityDilateGatherCS::FParameters>();
 		PassParameters->Dilate = VelocityDilateParameters;
-		PassParameters->OutVelocityTileTexture = GraphBuilder.CreateUAV(VelocityTileTexture);
+		if (ScatterDilatation)
+		{
+			check(BlurDirections > 1);
+			// Feed the scattered min max to make sure ReducePolarVelocityRange() is aware the highest velocity might be coming from tiles further away
+			PassParameters->CenterVelocityTileTextures.Textures[0] = VelocityTileTextures.Textures[0];
+			PassParameters->CenterVelocityTileTextures.Textures[1] = VelocityTileTexturesSetup.Textures[1];
+		}
+		else
+		{
+			PassParameters->CenterVelocityTileTextures = VelocityTileTexturesSetup;
+		}
+		PassParameters->OutVelocityTile = CreateUAVs(GraphBuilder, DilatedTileTextures);
 
-		TShaderMapRef<FMotionBlurVelocityDilateGatherCS> ComputeShader(View.ShaderMap);
+		FMotionBlurVelocityDilateGatherCS::FPermutationDomain PermutationVector;
+		PermutationVector.Set<FMotionBlurDirections>(BlurDirections);
+
+		TShaderMapRef<FMotionBlurVelocityDilateGatherCS> ComputeShader(View.ShaderMap, PermutationVector);
 		FComputeShaderUtils::AddPass(
 			GraphBuilder,
 			RDG_EVENT_NAME("VelocityTileGatherCS %dx%d", VelocityTileCount.X, VelocityTileCount.Y),
 			ComputeShader,
 			PassParameters,
 			FComputeShaderUtils::GetGroupCount(VelocityTileCount, FMotionBlurVelocityDilateGatherCS::ThreadGroupSize));
+
+		if (!ScatterDilatation)
+		{
+			VelocityTileTextures = DilatedTileTextures;
+		}
+		else if (BlurDirections > 1)
+		{
+			// VelocityTileTextures.Textures[0] is actually scattered.
+			GraphBuilder.RemoveUnusedTextureWarning(DilatedTileTextures.Textures[0]);
+			VelocityTileTextures.Textures[1] = DilatedTileTextures.Textures[1];
+		}
 	}
 
 	*VelocityFlatTextureOutput = VelocityFlatTexture;
-	*VelocityTileTextureOutput = VelocityTileTexture;
+	*VelocityTileTexturesOutput = VelocityTileTextures;
 }
 
 FRDGTextureRef AddMotionBlurFilterPass(
@@ -634,7 +712,7 @@ FRDGTextureRef AddMotionBlurFilterPass(
 	const FMotionBlurViewports& Viewports,
 	FRDGTextureRef ColorTexture,
 	FRDGTextureRef VelocityFlatTexture,
-	FRDGTextureRef VelocityTileTexture,
+	FVecocityTileTextures VelocityTileTextures,
 	FRDGTextureRef PostMotionBlurTranslucency,
 	const FIntPoint& PostMotionBlurTranslucencySize,
 	EMotionBlurFilterPass MotionBlurFilterPass,
@@ -642,10 +720,11 @@ FRDGTextureRef AddMotionBlurFilterPass(
 {
 	check(ColorTexture);
 	check(VelocityFlatTexture);
-	check(VelocityTileTexture);
 	check(MotionBlurFilterPass != EMotionBlurFilterPass::MAX);
 	check(MotionBlurQuality != EMotionBlurQuality::MAX);
 
+	const int32 BlurDirections = GetMotionBlurDirections();
+
 	const float MotionBlur2ndScale = CVarMotionBlur2ndScale.GetValueOnRenderThread();
 
 	const bool bUseCompute = View.bUseComputePasses;
@@ -687,7 +766,7 @@ FRDGTextureRef AddMotionBlurFilterPass(
 	MotionBlurFilterParameters.ColorToVelocity = Viewports.ColorToVelocityTransform;
 	MotionBlurFilterParameters.ColorTexture = ColorTexture;
 	MotionBlurFilterParameters.VelocityFlatTexture = VelocityFlatTexture;
-	MotionBlurFilterParameters.VelocityTileTexture = VelocityTileTexture;
+	MotionBlurFilterParameters.VelocityTileTextures = VelocityTileTextures;
 	MotionBlurFilterParameters.ColorSampler = GetMotionBlurColorSampler();
 	MotionBlurFilterParameters.VelocitySampler = GetMotionBlurVelocitySampler();
 	MotionBlurFilterParameters.VelocityTileSampler = GetMotionBlurVelocitySampler();
@@ -711,48 +790,31 @@ FRDGTextureRef AddMotionBlurFilterPass(
 	FRDGTextureDesc OutColorDesc = FRDGTextureDesc::Create2D(
 		ColorTexture->Desc.Extent,
 		IsPostProcessingWithAlphaChannelSupported() ? PF_FloatRGBA : PF_FloatRGB,
-		(!bUseCompute && View.GetOverwriteLoadAction() == ERenderTargetLoadAction::EClear) ? FClearValueBinding::Black : FClearValueBinding::None,
-		(bUseCompute ? TexCreate_UAV : TexCreate_RenderTargetable) | TexCreate_ShaderResource | GFastVRamConfig.MotionBlur);
+		FClearValueBinding::None,
+		TexCreate_UAV | TexCreate_ShaderResource | GFastVRamConfig.MotionBlur);
 
 	FRDGTextureRef ColorTextureOutput = GraphBuilder.CreateTexture(OutColorDesc, TEXT("MotionBlur.SceneColor"));
 
 	FMotionBlurFilterPermutationDomain PermutationVector;
 	PermutationVector.Set<FMotionBlurQualityDimension>(MotionBlurQuality);
 	PermutationVector.Set<FPostMotionBlurTranslucencyDimension>(PostMotionBlurTranslucency != nullptr);
+	PermutationVector.Set<FMotionBlurDirections>(BlurDirections);
 
-	if (bUseCompute)
-	{
-		FMotionBlurFilterCS::FParameters* PassParameters = GraphBuilder.AllocParameters<FMotionBlurFilterCS::FParameters>();
-		PassParameters->Filter = MotionBlurFilterParameters;
-		PassParameters->OutColorTexture = GraphBuilder.CreateUAV(ColorTextureOutput);
-
-		TShaderMapRef<FMotionBlurFilterCS> ComputeShader(View.ShaderMap, PermutationVector);
-
-		FComputeShaderUtils::AddPass(
-			GraphBuilder,
-			RDG_EVENT_NAME("Motion Blur %dx%d (CS)", Viewports.Color.Rect.Width(), Viewports.Color.Rect.Height()),
-			ComputeShader,
-			PassParameters,
-			FComputeShaderUtils::GetGroupCount(Viewports.Color.Rect.Size(), FComputeShaderUtils::kGolden2DGroupSize));
-	}
-	else
-	{
-		FMotionBlurFilterPS::FParameters* PassParameters = GraphBuilder.AllocParameters<FMotionBlurFilterPS::FParameters>();
-		PassParameters->Filter = MotionBlurFilterParameters;
-		PassParameters->RenderTargets[0] = FRenderTargetBinding(ColorTextureOutput, View.GetOverwriteLoadAction());
+	FMotionBlurFilterCS::FParameters* PassParameters = GraphBuilder.AllocParameters<FMotionBlurFilterCS::FParameters>();
+	PassParameters->Filter = MotionBlurFilterParameters;
+	PassParameters->OutColorTexture = GraphBuilder.CreateUAV(ColorTextureOutput);
+	PassParameters->DebugOutput = CreateDebugUAV(GraphBuilder, ColorTextureOutput->Desc.Extent, TEXT("Debug.MotionBlur.Apply"));
 
-		TShaderMapRef<FMotionBlurFilterPS> PixelShader(View.ShaderMap, PermutationVector);
+	TShaderMapRef<FMotionBlurFilterCS> ComputeShader(View.ShaderMap, PermutationVector);
 
-		AddDrawScreenPass(
-			GraphBuilder,
-			RDG_EVENT_NAME("Motion Blur %dx%d (PS)", Viewports.Color.Rect.Width(), Viewports.Color.Rect.Height()),
-			View,
-			Viewports.Color,
-			Viewports.Color,
-			PixelShader,
-			PassParameters,
-			EScreenPassDrawFlags::AllowHMDHiddenAreaMask);
-	}
+	FComputeShaderUtils::AddPass(
+		GraphBuilder,
+		RDG_EVENT_NAME("MotionBlur Apply(Quality=%d Directions=%d) %dx%d (CS)",
+			int32(MotionBlurQuality), BlurDirections,
+			Viewports.Color.Rect.Width(), Viewports.Color.Rect.Height()),
+		ComputeShader,
+		PassParameters,
+		FComputeShaderUtils::GetGroupCount(Viewports.Color.Rect.Size(), FComputeShaderUtils::kGolden2DGroupSize));
 
 	return ColorTextureOutput;
 }
@@ -849,7 +911,7 @@ FScreenPassTexture AddMotionBlurPass(FRDGBuilder& GraphBuilder, const FViewInfo&
 	RDG_EVENT_SCOPE(GraphBuilder, "MotionBlur");
 
 	FRDGTextureRef VelocityFlatTexture = nullptr;
-	FRDGTextureRef VelocityTileTexture = nullptr;
+	FVecocityTileTextures VelocityTileTextures;
 	AddMotionBlurVelocityPass(
 		GraphBuilder,
 		View,
@@ -858,7 +920,7 @@ FScreenPassTexture AddMotionBlurPass(FRDGBuilder& GraphBuilder, const FViewInfo&
 		Inputs.SceneDepth.Texture,
 		Inputs.SceneVelocity.Texture,
 		&VelocityFlatTexture,
-		&VelocityTileTexture);
+		&VelocityTileTextures);
 
 	FScreenPassTexture Output;
 	Output.ViewRect = Viewports.Color.Rect;
@@ -871,7 +933,7 @@ FScreenPassTexture AddMotionBlurPass(FRDGBuilder& GraphBuilder, const FViewInfo&
 			Viewports,
 			Inputs.SceneColor.Texture,
 			VelocityFlatTexture,
-			VelocityTileTexture,
+			VelocityTileTextures,
 			nullptr,
 			FIntPoint(0, 0),
 			EMotionBlurFilterPass::Separable0,
@@ -883,7 +945,7 @@ FScreenPassTexture AddMotionBlurPass(FRDGBuilder& GraphBuilder, const FViewInfo&
 			Viewports,
 			MotionBlurFilterTexture,
 			VelocityFlatTexture,
-			VelocityTileTexture,
+			VelocityTileTextures,
 			Inputs.PostMotionBlurTranslucency.Texture,
 			Inputs.PostMotionBlurTranslucency.ViewRect.Size(),
 			EMotionBlurFilterPass::Separable1,
@@ -897,7 +959,7 @@ FScreenPassTexture AddMotionBlurPass(FRDGBuilder& GraphBuilder, const FViewInfo&
 			Viewports,
 			Inputs.SceneColor.Texture,
 			VelocityFlatTexture,
-			VelocityTileTexture,
+			VelocityTileTextures,
 			Inputs.PostMotionBlurTranslucency.Texture,
 			Inputs.PostMotionBlurTranslucency.ViewRect.Size(),
 			EMotionBlurFilterPass::Unified,
