HZB invalidation cull using static pages (when enabled)

diff --git a/Engine/Shaders/Private/VirtualShadowMaps/CacheManagement.usf b/Engine/Shaders/Private/VirtualShadowMaps/CacheManagement.usf
--- a/Engine/Shaders/Private/VirtualShadowMaps/CacheManagement.usf
+++ b/Engine/Shaders/Private/VirtualShadowMaps/CacheManagement.usf
@@ -136,7 +136,13 @@ void VirtualSmInvalidateInstancePagesCS(
 											if (PageFlag != 0)
 											{
 #if USE_HZB_OCCLUSION
-												if (!IsPageVisibleHZB(uint2(x, y), PageFlagOffset, HZBTestRect))
+												// When separate static caching is enabled we check against the static page
+												// rather than the dynamic one. This is not totally ideal for a number of reasons
+												// but it is generally better than checking against the dynamic one in most scenes.
+												// Ideally in fully cached cases we would check the cached/merged (dynamic) page at
+												// the very least.
+												const bool bSampleStaticPage = true;
+												if (!IsPageVisibleHZB(uint2(x, y), PageFlagOffset, HZBTestRect, bSampleStaticPage))
 												{
 													continue;
 												}
