Remove dithered opacity evaluation in raytracing shader, as it cause invalid shadow casting.

The original change was converting dithered opacity into translucent evaluation. But this cause shadow evaluation to be incorrect.

diff --git a/Engine/Shaders/Private/RayTracing/RayTracingMaterialHitShaders.usf b/Engine/Shaders/Private/RayTracing/RayTracingMaterialHitShaders.usf
--- a/Engine/Shaders/Private/RayTracing/RayTracingMaterialHitShaders.usf
+++ b/Engine/Shaders/Private/RayTracing/RayTracingMaterialHitShaders.usf
@@ -350,17 +350,11 @@ RAY_TRACING_ENTRY_CLOSEST_HIT(MaterialCHS,
 	Payload.Opacity = GetMaterialOpacity(PixelMaterialInputs);
 
 #if MATERIALBLENDING_MASKED
-#if MATERIAL_DITHER_OPACITY_MASK
-	// dithering emulates real transparency, so switch to translucent
-	Payload.Opacity = DitheredOpacityMaskToOpacity(GetMaterialMask(PixelMaterialInputs));
-	Payload.BlendingMode = RAY_TRACING_BLEND_MODE_TRANSLUCENT;
-#else
 	// regular masked mode - binary decision
 	if (GetMaterialMask(PixelMaterialInputs) < 0)
 	{
 		Payload.Opacity = 0;
 	}
-#endif
 #endif // MATERIALBLENDING_MASKED
 
 	Payload.Ior = Refraction.x == 0 ? 1.0 : Refraction.x;
@@ -533,7 +527,7 @@ RAY_TRACING_ENTRY_ANY_HIT(MaterialAHS,
 	// All other blending modes _might_ need material evaluation - so we need
 	// to generate some code
 
-#if MATERIALBLENDING_MASKED && !MATERIAL_DITHER_OPACITY_MASK
+#if MATERIALBLENDING_MASKED
 	// NOTE: Masked mode execute always - regardless of ray flags
 	// This ensures that both material rays and shadow rays "see" the same thing
 #else
@@ -541,7 +535,7 @@ RAY_TRACING_ENTRY_ANY_HIT(MaterialAHS,
 	if (PackedPayload.IsIgnoreTranslucentMaterials())
 	{
 		// special mode used by regular RT shadows - translucent materials are skipped, everything else is treated as opaque
-#if MATERIALBLENDING_TRANSLUCENT || (MATERIALBLENDING_MASKED && MATERIAL_DITHER_OPACITY_MASK)
+#if MATERIALBLENDING_TRANSLUCENT
 		IgnoreHit();
 #endif
 		return;
@@ -592,7 +586,7 @@ RAY_TRACING_ENTRY_ANY_HIT(MaterialAHS,
 
 	// Now extract the relevant info from PixelMaterialInputs according to the blending mode
 
-#if MATERIALBLENDING_MASKED && !MATERIAL_DITHER_OPACITY_MASK
+#if MATERIALBLENDING_MASKED
 	// Regardless of payload flags -- we always apply this 
 	if (GetMaterialMask(PixelMaterialInputs) < 0)
 	{
