Added VirtualTexturePendingMips debug view mode.
For each pixel it shows the worst case missing mip level in all of the virtual textures sampled at that pixel.
Useful to analyze where late virtual texture page streaming is happening.

diff --git a/Engine/Shaders/Private/DebugViewModeVertexShader.usf b/Engine/Shaders/Private/DebugViewModeVertexShader.usf
--- a/Engine/Shaders/Private/DebugViewModeVertexShader.usf
+++ b/Engine/Shaders/Private/DebugViewModeVertexShader.usf
@@ -9,6 +9,7 @@
 #define SCENE_TEXTURES_DISABLED 1
 
 #include "Common.ush"
+#include "VirtualTextureCommon.ush"
 
 struct FTexCoordScalesParams
 {
@@ -25,6 +26,11 @@ MaterialFloat StoreTexSample(in out FTexCoordScalesParams Params, float4 C, int
 	return 1.f;
 }
 
+MaterialFloat StoreVTSampleInfo(in out FTexCoordScalesParams Params, VTPageTableResult PageTableResult, uint LayerIndex, int TextureReferenceIndex)
+{
+	return 1.f;
+}
+
 #include "/Engine/Generated/Material.ush"
 #include "/Engine/Generated/VertexFactory.ush"
 
diff --git a/Engine/Shaders/Private/MaterialTemplate.ush b/Engine/Shaders/Private/MaterialTemplate.ush
--- a/Engine/Shaders/Private/MaterialTemplate.ush
+++ b/Engine/Shaders/Private/MaterialTemplate.ush
@@ -451,9 +451,11 @@ struct FMaterialPixelParameters
 #if TEX_COORD_SCALE_ANALYSIS
 #define MaterialStoreTexCoordScale(Parameters, UV, TextureReferenceIndex) StoreTexCoordScale(Parameters.TexCoordScalesParams, UV, TextureReferenceIndex)
 #define MaterialStoreTexSample(Parameters, UV, TextureReferenceIndex) StoreTexSample(Parameters.TexCoordScalesParams, UV, TextureReferenceIndex)
+#define MaterialStoreVTSampleInfo(Parameters, PageTableResult, LayerIndex, TextureReferenceIndex) StoreVTSampleInfo(Parameters, PageTableResult, LayerIndex, TextureReferenceIndex)
 #else
 #define MaterialStoreTexCoordScale(Parameters, UV, TextureReferenceIndex) 1.0f
 #define MaterialStoreTexSample(Parameters, UV, TextureReferenceIndex) 1.0f
+#define MaterialStoreVTSampleInfo(Parameters, PageTableResult, LayerIndex, TextureReferenceIndex) 1.0f
 #endif
 
 // @todo compat hack
diff --git a/Engine/Shaders/Private/MaterialTexCoordScalesPixelShader.usf b/Engine/Shaders/Private/MaterialTexCoordScalesPixelShader.usf
--- a/Engine/Shaders/Private/MaterialTexCoordScalesPixelShader.usf
+++ b/Engine/Shaders/Private/MaterialTexCoordScalesPixelShader.usf
@@ -17,6 +17,7 @@
 
 #include "Common.ush"
 #include "SHCommon.ush"
+#include "VirtualTextureCommon.ush"
 
  // If -1, output everything to the RW texture.
 int2 AnalysisParams; // (TextureAnalysisIndex, bOutputScales)
@@ -95,6 +96,11 @@ MaterialFloat StoreTexSample(in out FTexCoordScalesParams Params, float4 C, int
 	return 1.f;
 }
 
+MaterialFloat StoreVTSampleInfo(in out FTexCoordScalesParams Params, VTPageTableResult PageTableResult, uint LayerIndex, int TextureReferenceIndex)
+{
+	return 1.f; 
+}
+
 #include "/Engine/Generated/Material.ush"
 #include "DebugViewModeCommon.ush"
 
diff --git a/Engine/Shaders/Private/RequiredTextureResolutionPixelShader.usf b/Engine/Shaders/Private/RequiredTextureResolutionPixelShader.usf
--- a/Engine/Shaders/Private/RequiredTextureResolutionPixelShader.usf
+++ b/Engine/Shaders/Private/RequiredTextureResolutionPixelShader.usf
@@ -9,31 +9,26 @@
 
 #include "Common.ush"
 #include "SHCommon.ush"
+#include "VirtualTextureCommon.ush"
 
- // If -1, output everything to the RW texture.
-int2 AnalysisParams; // (TextureAnalysisIndex, TextureResolution)
+int3 AnalysisParams; // (AnalysisType, TextureAnalysisIndex, TextureResolution)
 float PrimitiveAlpha;
 
+// Analysis type will be ANALYSIS_TEXTURE_RESOLUTION or ANALYSIS_VIRTUAL_TEXTURE_PENDING_MIPS
+#define TextureAnalysisType (AnalysisParams.x)
 // Analysis index will be in [0, 31] if analysing a single texture
-#define TextureAnalysisIndex (AnalysisParams.x)
-#define TextureResolution (AnalysisParams.y)
+#define TextureAnalysisIndex (AnalysisParams.y)
+// Texture resolution if we're analysing a single texture
+#define TextureResolution (AnalysisParams.z)
 
-float GetComponent(float4 V, int Index)
-{
-	FLATTEN
-	if (Index == 0) return V.x;
-	FLATTEN
-	if (Index == 1) return V.y;
-	FLATTEN
-	if (Index == 2) return V.z;
-	return V.w;
-}
+#define ANALYSIS_TEXTURE_RESOLUTION				0
+#define ANALYSIS_VIRTUAL_TEXTURE_PENDING_MIPS	1
 
 struct FTexCoordScalesParams
 {
 	float RequiredResolution;
 	float TexSample;
-	float AvailableVTLevelSize; // Size in pixels of the VT level of the texture that's resident (of "the" texture, not the whole atlas, not the physical texture, not the cache texture, the actual virtualized texture)
+	int VTMissingMipCount;
 };
 
 MaterialFloat StoreTexCoordScale(in out FTexCoordScalesParams Params, float2 UV, int TextureReferenceIndex)
@@ -57,19 +52,51 @@ MaterialFloat StoreTexCoordScale(in out FTexCoordScalesParams Params, float2 UV,
 
 MaterialFloat StoreTexSample(in out FTexCoordScalesParams Params, float4 C, int TextureReferenceIndex)
 {
-	Params.TexSample = TextureAnalysisIndex == TextureReferenceIndex ? lerp(.4f, 1.f, saturate((C.r + C.g + C.b) / 3)) : Params.TexSample;
+	if (TextureAnalysisIndex == TextureReferenceIndex)
+	{
+		Params.TexSample = lerp(.4f, 1.f, saturate((C.r + C.g + C.b) / 3));
+	}
 	return 1.f;
 }
 
-MaterialFloat StoreAvailableVTLevel(in out FTexCoordScalesParams Params, float Level, int TextureReferenceIndex)
+MaterialFloat StoreVTSampleInfo(in out FTexCoordScalesParams Params, VTPageTableResult PageTableResult, uint LayerIndex, int TextureReferenceIndex)
 {
-	Params.AvailableVTLevelSize = (TextureAnalysisIndex == TextureReferenceIndex) ? Level : Params.AvailableVTLevelSize;
+	if (TextureAnalysisType == ANALYSIS_VIRTUAL_TEXTURE_PENDING_MIPS)
+	{
+		uint RequestLevel = GetRequestedLevel(PageTableResult);
+
+		// Find the highest resolution from all the page table entries sampled for this VT stack.
+		// We do that because a single VT stack can cover textures of different resolutions.
+		// Taking info from the lower resolution textures will always show undersampling compared to the requested level that we are not interested in.
+		uint SampledLevel = 0xffff;
+		for (uint Index = 0; Index < 8; ++Index)
+		{
+			SampledLevel = min(SampledLevel, GetSampledLevel(PageTableResult, Index, 0xffff));
+		}
+
+		if (SampledLevel < 0xffff)
+		{
+			// Accumulate the worst case that we find for the current pixel.
+			// Multiple VT stacks can exist, and we want to visualize the worst case.
+			int Diff = max((int)SampledLevel - (int)RequestLevel, 0);
+			Params.VTMissingMipCount = max(Diff, Params.VTMissingMipCount);
+		}
+	}
 	return 1.f;
 }
 
 #include "/Engine/Generated/Material.ush"
 #include "DebugViewModeCommon.ush"
 
+float3 AccuracyColorLookup(float Accuracy)
+{
+	Accuracy = clamp(Accuracy, -1.99, 1.99);
+	int ColorIndex = floor(Accuracy) + 2;
+	float3 Color0 = DebugViewModePass.AccuracyColors[ColorIndex].rgb;
+	float3 Color1 = DebugViewModePass.AccuracyColors[ColorIndex + 1].rgb;
+	float ColorLerp = frac(Accuracy);
+	return lerp(Color0, Color1, ColorLerp);
+}
 
 #if LIGHTMAP_VT_ENABLED || NUM_VIRTUALTEXTURE_SAMPLES
 EARLYDEPTHSTENCIL
@@ -86,21 +113,18 @@ void Main(
 	ResolvedView = ResolveView();
 #endif
 
-	// This default value will make it dark grey when nothing updates the texSample.
-	float3 Result = float3(UNDEFINED_ACCURACY, UNDEFINED_ACCURACY, UNDEFINED_ACCURACY);
-
 	FMaterialPixelParameters MaterialParameters = GetMaterialPixelParameters(DebugInputs, DebugInputs.SvPosition);
 
 	MaterialParameters.TexCoordScalesParams.RequiredResolution = 0;
 	MaterialParameters.TexCoordScalesParams.TexSample = 0;
-	MaterialParameters.TexCoordScalesParams.AvailableVTLevelSize = -1;
+	MaterialParameters.TexCoordScalesParams.VTMissingMipCount = -1;
 
 	half3 BaseColor;
 	{
 		FPixelMaterialInputs PixelMaterialInputs;
 		CalcMaterialParameters(MaterialParameters, PixelMaterialInputs, DebugInputs.SvPosition, bIsFrontFace);
 
-		// Sample material properties. The results are not used, but the calls to StoreTexCoordScale are still be made.
+		// Sample material properties. The results are not used, but the calls to StoreTexCoordScale will still be made.
 		BaseColor = GetMaterialBaseColorRaw(PixelMaterialInputs);
 		half  Metallic = GetMaterialMetallicRaw(PixelMaterialInputs);
 		half  Specular = GetMaterialSpecularRaw(PixelMaterialInputs);
@@ -122,25 +146,28 @@ void Main(
 #endif
 	}
 
+	float3 Result = float3(UNDEFINED_ACCURACY, UNDEFINED_ACCURACY, UNDEFINED_ACCURACY);
 	Result *= saturate(Luminance(BaseColor));
 
-	if (MaterialParameters.TexCoordScalesParams.RequiredResolution > 0)
+	if (TextureAnalysisType == ANALYSIS_TEXTURE_RESOLUTION)
 	{
-		float Accuracy = 0.0f;
-#if NUM_VIRTUALTEXTURE_SAMPLES || LIGHTMAP_VT_ENABLED
-		// For VT sampling we use the dynamically sampled miplevel. So this could vary on a tile per tile basis.
-		if (MaterialParameters.TexCoordScalesParams.AvailableVTLevelSize > 0)
+		if (MaterialParameters.TexCoordScalesParams.RequiredResolution > 0)
 		{
-			Accuracy = clamp(log2(	MaterialParameters.TexCoordScalesParams.AvailableVTLevelSize / MaterialParameters.TexCoordScalesParams.RequiredResolution), -1.99, 1.99);			
+			float Accuracy = log2(TextureResolution / MaterialParameters.TexCoordScalesParams.RequiredResolution);
+			Result = AccuracyColorLookup(Accuracy);
+			Result *= MaterialParameters.TexCoordScalesParams.TexSample;
 		}
-		else
-#endif
+	}
+	else if (TextureAnalysisType == ANALYSIS_VIRTUAL_TEXTURE_PENDING_MIPS)
+	{
+		if (MaterialParameters.TexCoordScalesParams.VTMissingMipCount >= 0)
 		{
-			Accuracy = clamp(log2(TextureResolution / MaterialParameters.TexCoordScalesParams.RequiredResolution), -1.99, 1.99);
+			float Accuracy = -MaterialParameters.TexCoordScalesParams.VTMissingMipCount;
+			Result = AccuracyColorLookup(Accuracy);
+			Result *= 0.5f;
 		}
-		int ColorIndex = floor(Accuracy) + 2;
-		Result = MaterialParameters.TexCoordScalesParams.TexSample * lerp(DebugViewModeStruct.AccuracyColors[ColorIndex].rgb, DebugViewModeStruct.AccuracyColors[ColorIndex + 1].rgb, frac(Accuracy));
 	}
+
 	OutColor = RETURN_COLOR(float4(Result, PrimitiveAlpha));
 	
 #if NUM_VIRTUALTEXTURE_SAMPLES || LIGHTMAP_VT_ENABLED
@@ -153,9 +180,3 @@ void Main(
 	);
 #endif	
 }
-
-
-
-
-
-
diff --git a/Engine/Shaders/Private/ViewMode/LODColorationPixelShader.usf b/Engine/Shaders/Private/ViewMode/LODColorationPixelShader.usf
--- a/Engine/Shaders/Private/ViewMode/LODColorationPixelShader.usf
+++ b/Engine/Shaders/Private/ViewMode/LODColorationPixelShader.usf
@@ -15,10 +15,12 @@
 
 #include "../Common.ush"
 #include "../SHCommon.ush"
+#include "../VirtualTextureCommon.ush"
 
 struct FTexCoordScalesParams { float Dummy; };
 MaterialFloat StoreTexCoordScale(in out FTexCoordScalesParams Params, float2 UV, int TextureReferenceIndex) { return 1.f; }
 MaterialFloat StoreTexSample(in out FTexCoordScalesParams Params, float4 C, int TextureReferenceIndex) { return 1.f; }
+MaterialFloat StoreVTSampleInfo(in out FTexCoordScalesParams Params, VTPageTableResult PageTableResult, uint LayerIndex, int TextureReferenceIndex) { return 1.f; }
 
 #include "/Engine/Generated/Material.ush"
 #include "../DebugViewModeCommon.ush"
diff --git a/Engine/Shaders/Private/VirtualTextureCommon.ush b/Engine/Shaders/Private/VirtualTextureCommon.ush
--- a/Engine/Shaders/Private/VirtualTextureCommon.ush
+++ b/Engine/Shaders/Private/VirtualTextureCommon.ush
@@ -202,6 +202,31 @@ struct VTPageTableResult
 	uint PackedRequest;
 };
 
+uint GetRequestedLevel(VTPageTableResult PageTableResult)
+{
+	uint RequestLevelPlusOne = (PageTableResult.PackedRequest >> 24) & 0xf;
+	uint RequestLevel = max(RequestLevelPlusOne, 1) - 1;
+	return RequestLevel;
+}
+
+uint GetSampledLevel(VTPageTableResult PageTableResult, uint LayerIndex, uint FallbackValue)
+{
+	uint PackedPageTableValue = PageTableResult.PageTableValue[LayerIndex / 4u][LayerIndex & 3u];
+	
+	// 0 is an empty value in the PageTableResult.
+	// We must also consider 1 as an empty value since it will appear in the .w channel page table textures with less than 4 channels.
+	// todo: Inject the number of layers inside VTPageTableResult and avoid the need to special case values here.
+	// HLSLMaterialTranslator could pass the value to TextureLoadVirtualPageTable through a macro similar to VIRTUALTEXTURE_PAGETABLE_%d.
+	if (PackedPageTableValue != 0 && PackedPageTableValue != 1)
+	{
+		uint SampleLevel = PackedPageTableValue & 0xf;
+		return SampleLevel;
+	}
+
+	// Return whatever the caller can interpret as empty.
+	return FallbackValue;
+}
+
 float GetStochasticMipNoise(float2 SvPositionXY)
 {
 #if VT_DISABLE_VIEW_UNIFORM_BUFFER
diff --git a/Engine/Source/Editor/EditorStyle/Private/SlateEditorStyle.cpp b/Engine/Source/Editor/EditorStyle/Private/SlateEditorStyle.cpp
--- a/Engine/Source/Editor/EditorStyle/Private/SlateEditorStyle.cpp
+++ b/Engine/Source/Editor/EditorStyle/Private/SlateEditorStyle.cpp
@@ -4800,6 +4800,7 @@ void FSlateEditorStyle::FStyle::SetupLevelEditorStyle()
 		Set( "EditorViewport.TexStreamAccMeshUVDensityMode", new IMAGE_BRUSH("Icons/icon_ViewMode_TextureStreamingAccuracy_16px", Icon16x16));
 		Set( "EditorViewport.TexStreamAccMaterialTextureScaleMode", new IMAGE_BRUSH( "Icons/icon_ViewMode_TextureStreamingAccuracy_16px", Icon16x16 ) );
 		Set( "EditorViewport.RequiredTextureResolutionMode", new IMAGE_BRUSH( "Icons/icon_ViewMode_TextureStreamingAccuracy_16px", Icon16x16 ) );
+		Set( "EditorViewport.VirtualTexturePendingMipsMode", new IMAGE_BRUSH( "Icons/icon_ViewMode_TextureStreamingAccuracy_16px", Icon16x16 ) );
 		Set( "EditorViewport.StationaryLightOverlapMode", new IMAGE_BRUSH( "Icons/icon_ViewMode_StationaryLightOverlap_16px", Icon16x16 ) );
 		Set( "EditorViewport.LightmapDensityMode", new IMAGE_BRUSH( "Icons/icon_ViewMode_LightmapDensity_16px", Icon16x16 ) );
 
diff --git a/Engine/Source/Editor/EditorStyle/Private/StarshipStyle.cpp b/Engine/Source/Editor/EditorStyle/Private/StarshipStyle.cpp
--- a/Engine/Source/Editor/EditorStyle/Private/StarshipStyle.cpp
+++ b/Engine/Source/Editor/EditorStyle/Private/StarshipStyle.cpp
@@ -4352,6 +4352,7 @@ void FStarshipEditorStyle::FStyle::SetupLevelEditorStyle()
 		Set( "EditorViewport.TexStreamAccMeshUVDensityMode", new IMAGE_BRUSH("Icons/icon_ViewMode_TextureStreamingAccuracy_16px", Icon16x16));
 		Set( "EditorViewport.TexStreamAccMaterialTextureScaleMode", new IMAGE_BRUSH( "Icons/icon_ViewMode_TextureStreamingAccuracy_16px", Icon16x16 ) );
 		Set( "EditorViewport.RequiredTextureResolutionMode", new IMAGE_BRUSH( "Icons/icon_ViewMode_TextureStreamingAccuracy_16px", Icon16x16 ) );
+		Set( "EditorViewport.VirtualTexturePendingMipsMode", new IMAGE_BRUSH( "Icons/icon_ViewMode_TextureStreamingAccuracy_16px", Icon16x16 ) );
 		Set( "EditorViewport.StationaryLightOverlapMode", new IMAGE_BRUSH( "Icons/icon_ViewMode_StationaryLightOverlap_16px", Icon16x16 ) );
 		Set( "EditorViewport.LightmapDensityMode", new IMAGE_BRUSH( "Icons/icon_ViewMode_LightmapDensity_16px", Icon16x16 ) );
 
diff --git a/Engine/Source/Editor/UnrealEd/Private/EditorBuildUtils.cpp b/Engine/Source/Editor/UnrealEd/Private/EditorBuildUtils.cpp
--- a/Engine/Source/Editor/UnrealEd/Private/EditorBuildUtils.cpp
+++ b/Engine/Source/Editor/UnrealEd/Private/EditorBuildUtils.cpp
@@ -1219,6 +1219,8 @@ EDebugViewShaderMode ViewModeIndexToDebugViewShaderMode(EViewModeIndex SelectedV
 		return DVSM_MaterialTextureScaleAccuracy;
 	case VMI_RequiredTextureResolution:
 		return DVSM_RequiredTextureResolution;
+	case VMI_VirtualTexturePendingMips:
+		return DVSM_VirtualTexturePendingMips;
 	case VMI_RayTracingDebug:
 		return DVSM_RayTracingDebug;
 	case VMI_LODColoration:
diff --git a/Engine/Source/Editor/UnrealEd/Private/EditorViewportCommands.cpp b/Engine/Source/Editor/UnrealEd/Private/EditorViewportCommands.cpp
--- a/Engine/Source/Editor/UnrealEd/Private/EditorViewportCommands.cpp
+++ b/Engine/Source/Editor/UnrealEd/Private/EditorViewportCommands.cpp
@@ -91,6 +91,7 @@ void FEditorViewportCommands::RegisterCommands()
 	UI_COMMAND( TexStreamAccMaterialTextureScaleMode, "Material Texture Scales Accuracy View Mode", "Visualize the accuracy of the material texture scales used for texture streaming", EUserInterfaceActionType::RadioButton, FInputChord() );
 	UI_COMMAND( TexStreamAccMaterialTextureScaleAll, "All Textures", "Visualize the scales accuracy of all textures", EUserInterfaceActionType::RadioButton, FInputChord() );
 	UI_COMMAND( RequiredTextureResolutionMode, "Required Texture Resolution View Mode", "Visualize the ratio between the currently streamed texture resolution and the resolution wanted by the GPU", EUserInterfaceActionType::RadioButton, FInputChord() );
+	UI_COMMAND( VirtualTexturePendingMipsMode, "Virtual Texture Pending Mips View Mode", "Visualize the difference between the currently streamed virtual texture level and the level wanted by the GPU", EUserInterfaceActionType::RadioButton, FInputChord() );
 
 	for (int32 TextureIndex = 0; TextureIndex < TEXSTREAM_MAX_NUM_TEXTURES_PER_MATERIAL; ++TextureIndex)
 	{
diff --git a/Engine/Source/Editor/UnrealEd/Private/SEditorViewport.cpp b/Engine/Source/Editor/UnrealEd/Private/SEditorViewport.cpp
--- a/Engine/Source/Editor/UnrealEd/Private/SEditorViewport.cpp
+++ b/Engine/Source/Editor/UnrealEd/Private/SEditorViewport.cpp
@@ -412,7 +412,8 @@ void SEditorViewport::BindCommands()
 	MAP_VIEWMODE_ACTION( Commands.TexStreamAccPrimitiveDistanceMode, VMI_PrimitiveDistanceAccuracy );
 	MAP_VIEWMODE_ACTION( Commands.TexStreamAccMeshUVDensityMode, VMI_MeshUVDensityAccuracy);
 	MAP_VIEWMODE_ACTION( Commands.TexStreamAccMaterialTextureScaleMode, VMI_MaterialTextureScaleAccuracy );
-	MAP_VIEWMODE_ACTION( Commands.RequiredTextureResolutionMode, VMI_RequiredTextureResolution );
+	MAP_VIEWMODE_ACTION( Commands.RequiredTextureResolutionMode, VMI_RequiredTextureResolution);
+	MAP_VIEWMODE_ACTION( Commands.VirtualTexturePendingMipsMode, VMI_VirtualTexturePendingMips );
 	MAP_VIEWMODE_ACTION( Commands.StationaryLightOverlapMode, VMI_StationaryLightOverlap );
 
 	if (IConsoleManager::Get().FindTConsoleVariableDataInt(TEXT("r.AllowStaticLighting"))->GetValueOnAnyThread() != 0)
diff --git a/Engine/Source/Editor/UnrealEd/Private/SEditorViewportViewMenu.cpp b/Engine/Source/Editor/UnrealEd/Private/SEditorViewportViewMenu.cpp
--- a/Engine/Source/Editor/UnrealEd/Private/SEditorViewportViewMenu.cpp
+++ b/Engine/Source/Editor/UnrealEd/Private/SEditorViewportViewMenu.cpp
@@ -181,10 +181,14 @@ void SEditorViewportViewMenu::FillViewMenu(UToolMenu* Menu) const
 							{
 								Section.AddMenuEntry(BaseViewportCommands.TexStreamAccMaterialTextureScaleMode, UViewModeUtils::GetViewModeDisplayName(VMI_MaterialTextureScaleAccuracy));
 							}
-							if (AllowDebugViewShaderMode(DVSM_RequiredTextureResolution, GMaxRHIShaderPlatform, FeatureLevel) && (!InParentToolBar.IsValid() || InParentToolBar.Pin()->IsViewModeSupported(VMI_MaterialTextureScaleAccuracy)))
+							if (AllowDebugViewShaderMode(DVSM_RequiredTextureResolution, GMaxRHIShaderPlatform, FeatureLevel) && (!InParentToolBar.IsValid() || InParentToolBar.Pin()->IsViewModeSupported(VMI_RequiredTextureResolution)))
 							{
 								Section.AddMenuEntry(BaseViewportCommands.RequiredTextureResolutionMode, UViewModeUtils::GetViewModeDisplayName(VMI_RequiredTextureResolution));
 							}
+							if (AllowDebugViewShaderMode(DVSM_RequiredTextureResolution, GMaxRHIShaderPlatform, FeatureLevel) && (!InParentToolBar.IsValid() || InParentToolBar.Pin()->IsViewModeSupported(VMI_VirtualTexturePendingMips)))
+							{
+								Section.AddMenuEntry(BaseViewportCommands.VirtualTexturePendingMipsMode, UViewModeUtils::GetViewModeDisplayName(VMI_VirtualTexturePendingMips));
+							}
 						}
 					}
 				};
@@ -205,7 +209,8 @@ void SEditorViewportViewMenu::FillViewMenu(UToolMenu* Menu) const
 									ViewMode == VMI_LightComplexity || ViewMode == VMI_LightmapDensity || ViewMode == VMI_StationaryLightOverlap
 									|| ViewMode == VMI_ShaderComplexity || ViewMode == VMI_ShaderComplexityWithQuadOverdraw || ViewMode == VMI_QuadOverdraw
 									// Texture Streaming Accuracy
-									|| ViewMode == VMI_PrimitiveDistanceAccuracy || ViewMode == VMI_MeshUVDensityAccuracy || ViewMode == VMI_MaterialTextureScaleAccuracy || ViewMode == VMI_RequiredTextureResolution
+									|| ViewMode == VMI_PrimitiveDistanceAccuracy || ViewMode == VMI_MeshUVDensityAccuracy || ViewMode == VMI_MaterialTextureScaleAccuracy 
+									|| ViewMode == VMI_RequiredTextureResolution || ViewMode == VMI_VirtualTexturePendingMips
 								);
 						})),
 					EUserInterfaceActionType::RadioButton,
diff --git a/Engine/Source/Editor/UnrealEd/Public/EditorViewportCommands.h b/Engine/Source/Editor/UnrealEd/Public/EditorViewportCommands.h
--- a/Engine/Source/Editor/UnrealEd/Public/EditorViewportCommands.h
+++ b/Engine/Source/Editor/UnrealEd/Public/EditorViewportCommands.h
@@ -110,6 +110,8 @@ public:
 	TSharedPtr< FUICommandInfo > RequiredTextureResolutionMode;
 	TSharedPtr< FUICommandInfo > RequiredTextureResolutionSingle[TEXSTREAM_MAX_NUM_TEXTURES_PER_MATERIAL];
 
+	TSharedPtr< FUICommandInfo > VirtualTexturePendingMipsMode;
+
 	/** Changes the viewport to stationary light overlap mode */
 	TSharedPtr< FUICommandInfo > StationaryLightOverlapMode;
 
diff --git a/Engine/Source/Runtime/Engine/Classes/Engine/EngineBaseTypes.h b/Engine/Source/Runtime/Engine/Classes/Engine/EngineBaseTypes.h
--- a/Engine/Source/Runtime/Engine/Classes/Engine/EngineBaseTypes.h
+++ b/Engine/Source/Runtime/Engine/Classes/Engine/EngineBaseTypes.h
@@ -907,6 +907,9 @@ enum EViewModeIndex
 	/** Visualize various aspects of Nanite */
 	VMI_VisualizeNanite = 29 UMETA(DisplayName = "Nanite Visualization"),
 
+	/** Compare the required texture resolution to the actual resolution. */
+	VMI_VirtualTexturePendingMips = 30 UMETA(DisplayName = "Virtual Texture Pending Mips"),
+
 	VMI_Max UMETA(Hidden),
 
 	// VMI_Unknown - The value assigned to VMI_Unknown must be the highest possible of any member of EViewModeIndex, or GetViewModeName might seg-fault
diff --git a/Engine/Source/Runtime/Engine/Private/DebugViewModeHelpers.cpp b/Engine/Source/Runtime/Engine/Private/DebugViewModeHelpers.cpp
--- a/Engine/Source/Runtime/Engine/Private/DebugViewModeHelpers.cpp
+++ b/Engine/Source/Runtime/Engine/Private/DebugViewModeHelpers.cpp
@@ -56,6 +56,7 @@ bool AllowDebugViewShaderMode(EDebugViewShaderMode ShaderMode, EShaderPlatform P
 	case DVSM_MaterialTextureScaleAccuracy:
 	case DVSM_RequiredTextureResolution:
 	case DVSM_OutputMaterialTextureScales:
+	case DVSM_VirtualTexturePendingMips:
 		return FeatureLevel >= ERHIFeatureLevel::SM5 && (bForceTextureStreamingBuild || PlatformSupportsDebugViewShaders(Platform));
 	case DVSM_RayTracingDebug:
 		return FeatureLevel >= ERHIFeatureLevel::SM5 ;
diff --git a/Engine/Source/Runtime/Engine/Private/Materials/HLSLMaterialTranslator.cpp b/Engine/Source/Runtime/Engine/Private/Materials/HLSLMaterialTranslator.cpp
--- a/Engine/Source/Runtime/Engine/Private/Materials/HLSLMaterialTranslator.cpp
+++ b/Engine/Source/Runtime/Engine/Private/Materials/HLSLMaterialTranslator.cpp
@@ -6074,9 +6074,13 @@ int32 FHLSLMaterialTranslator::TextureSample(
 			FString SampleCodeAnalytic;
 
 			if (SamplerDebugSupported(TextureType, bVirtualTexture, bDecal))
+			{
 				SampleCodeAnalytic = FString::Printf(TEXT("Debug%s(%s, %s, %d, VTUniform_Unpack(Material.VTPackedUniform[%d]), %s)"), *TextureTypeName, *TextureName, *VTPageTableResult_Analytic, VTPageTableIndex, VirtualTextureIndex, *UV_Scale);
+			}
 			else
+			{
 				SampleCodeAnalytic = FString::Printf(TEXT("%s(%s, %s, %d, VTUniform_Unpack(Material.VTPackedUniform[%d]))"), *TextureTypeName, *TextureName, *VTPageTableResult_Analytic, VTPageTableIndex, VirtualTextureIndex);
+			}
 
 			SampleCodeAnalytic = ApplySamplerType(SampleCodeAnalytic, SamplerType);
 			SamplingCodeIndex = AddCodeChunkInnerDeriv(*SampleCodeFinite, *SampleCodeAnalytic, MCT_Float4, false, EDerivativeStatus::NotValid);
@@ -6085,20 +6089,12 @@ int32 FHLSLMaterialTranslator::TextureSample(
 		{
 			SamplingCodeIndex = AddCodeChunk(MCT_Float4, *SampleCodeFinite, *TextureName, *VTPageTableResult_Finite, VTPageTableIndex, VirtualTextureIndex);
 		}
-		// TODO
-		/*
-		const bool bStoreAvailableVTLevel = ShaderFrequency == SF_Pixel && TextureReferenceIndex != INDEX_NONE;
-		if (bStoreAvailableVTLevel)
+	
+		const bool bStoreVTSampleInfo = ShaderFrequency == SF_Pixel && TextureReferenceIndex != INDEX_NONE;
+		if (bStoreVTSampleInfo)
 		{
-			check(VirtualTextureUniformExpressionIndex >= 0);
-			check(VirtualTextureIndex >= 0);
-
-			AddCodeChunk(MCT_Float, TEXT("StoreAvailableVTLevel(Parameters.TexCoordScalesParams, TextureVirtualGetSampledLevelSize(Material.VirtualTexturePageTable_%d, Material.VirtualTextureUniformData[%d], Parameters.SvPosition.xy, %s), %d)"),
-				VirtualTextureUniformExpressionIndex,
-				VirtualTextureIndex,
-				*UVs,
-				(int)TextureReferenceIndex);
-		}*/
+			AddCodeChunk(MCT_Float, TEXT("MaterialStoreVTSampleInfo(Parameters.TexCoordScalesParams, %s, %d, %d)"), *VTPageTableResult_Finite, VTPageTableIndex, TextureReferenceIndex);
+		}
 	}
 	else
 	{
@@ -6140,6 +6136,7 @@ int32 FHLSLMaterialTranslator::TextureSample(
 			FString SampleCodeAnalytic = SampleCodeFinite;
 			if (MipValueMode == TMVM_None || MipValueMode == TMVM_MipBias)
 			{
+				
 				if (SamplerDebugSupported(TextureType, bVirtualTexture, bDecal))
 					SampleCodeAnalytic = FString::Printf(TEXT("Debug%sGrad(%s,%s,%s,%s,%s,%s)"), *TextureTypeName, *TextureName, *SamplerStateCode, *UVs, *UV_Ddx, *UV_Ddy, *UV_Scale);
 				else
@@ -6174,9 +6171,9 @@ int32 FHLSLMaterialTranslator::TextureProperty(int32 TextureIndex, EMaterialExpo
 		TextureType != MCT_VolumeTexture &&
 		TextureType != MCT_Texture2DArray)
 	{
-		return Errorf(TEXT("Texture size only available for Texture2D, TextureVirtual, Texture2DArray, and VolumeTexture, not %s"),DescribeType(TextureType));
+		return Errorf(TEXT("Texture size only available for Texture2D, TextureVirtual, Texture2DArray, and VolumeTexture, not %s"), DescribeType(TextureType));
 	}
-		
+
 	FMaterialUniformExpressionTexture* TextureExpression = (*CurrentScopeChunks)[TextureIndex].UniformExpression->GetTextureUniformExpression();
 	if (!TextureExpression)
 	{
diff --git a/Engine/Source/Runtime/Engine/Private/SceneView.cpp b/Engine/Source/Runtime/Engine/Private/SceneView.cpp
--- a/Engine/Source/Runtime/Engine/Private/SceneView.cpp
+++ b/Engine/Source/Runtime/Engine/Private/SceneView.cpp
@@ -2791,6 +2791,10 @@ EDebugViewShaderMode FSceneViewFamily::ChooseDebugViewShaderMode() const
 	{
 		return DVSM_RequiredTextureResolution;
 	}
+	else if (EngineShowFlags.VirtualTexturePendingMips)
+	{
+		return DVSM_VirtualTexturePendingMips;
+	}
 	else if (EngineShowFlags.RayTracingDebug)
 	{
 		return DVSM_RayTracingDebug;
diff --git a/Engine/Source/Runtime/Engine/Private/ShowFlags.cpp b/Engine/Source/Runtime/Engine/Private/ShowFlags.cpp
--- a/Engine/Source/Runtime/Engine/Private/ShowFlags.cpp
+++ b/Engine/Source/Runtime/Engine/Private/ShowFlags.cpp
@@ -312,6 +312,7 @@ void ApplyViewMode(EViewModeIndex ViewModeIndex, bool bPerspective, FEngineShowF
 		case VMI_MeshUVDensityAccuracy:
 		case VMI_MaterialTextureScaleAccuracy:
 		case VMI_RequiredTextureResolution:
+		case VMI_VirtualTexturePendingMips:
 		case VMI_LODColoration:
 		case VMI_HLODColoration:
 			bPostProcessing = false;
@@ -369,6 +370,7 @@ void ApplyViewMode(EViewModeIndex ViewModeIndex, bool bPerspective, FEngineShowF
 	EngineShowFlags.SetMeshUVDensityAccuracy(ViewModeIndex == VMI_MeshUVDensityAccuracy);
 	EngineShowFlags.SetMaterialTextureScaleAccuracy(ViewModeIndex == VMI_MaterialTextureScaleAccuracy);
 	EngineShowFlags.SetRequiredTextureResolution(ViewModeIndex == VMI_RequiredTextureResolution);
+	EngineShowFlags.SetVirtualTexturePendingMips(ViewModeIndex == VMI_VirtualTexturePendingMips);
 	EngineShowFlags.SetStationaryLightOverlap(ViewModeIndex == VMI_StationaryLightOverlap);
 	EngineShowFlags.SetLightMapDensity(ViewModeIndex == VMI_LightmapDensity || ViewModeIndex == VMI_LitLightmapDensity);
 	EngineShowFlags.SetPostProcessing(bPostProcessing);
@@ -472,6 +474,7 @@ void EngineShowFlagOverride(EShowFlagInitMode ShowFlagInitMode, EViewModeIndex V
 			ViewModeIndex == VMI_MeshUVDensityAccuracy ||
 			ViewModeIndex == VMI_MaterialTextureScaleAccuracy ||
 			ViewModeIndex == VMI_RequiredTextureResolution ||
+			ViewModeIndex == VMI_VirtualTexturePendingMips ||
 			ViewModeIndex == VMI_LightmapDensity ||
 			ViewModeIndex == VMI_LitLightmapDensity)
 		{
@@ -496,6 +499,7 @@ void EngineShowFlagOverride(EShowFlagInitMode ShowFlagInitMode, EViewModeIndex V
 			ViewModeIndex == VMI_MeshUVDensityAccuracy ||
 			ViewModeIndex == VMI_MaterialTextureScaleAccuracy ||
 			ViewModeIndex == VMI_RequiredTextureResolution ||
+			ViewModeIndex == VMI_VirtualTexturePendingMips ||
 			ViewModeIndex == VMI_LODColoration ||
 			ViewModeIndex == VMI_HLODColoration ||
 			ViewModeIndex == VMI_LightmapDensity)
@@ -529,7 +533,8 @@ void EngineShowFlagOverride(EShowFlagInitMode ShowFlagInitMode, EViewModeIndex V
 		if (ViewModeIndex == VMI_PrimitiveDistanceAccuracy ||
 			ViewModeIndex == VMI_MeshUVDensityAccuracy ||
 			ViewModeIndex == VMI_MaterialTextureScaleAccuracy || 
-			ViewModeIndex == VMI_RequiredTextureResolution)
+			ViewModeIndex == VMI_RequiredTextureResolution ||
+			ViewModeIndex == VMI_VirtualTexturePendingMips)
 		{
 			EngineShowFlags.SetDecals(false); // Decals require the use of FDebugPSInLean.
 			EngineShowFlags.SetParticles(false); // FX are fully streamed.
@@ -699,6 +704,10 @@ EViewModeIndex FindViewMode(const FEngineShowFlags& EngineShowFlags)
 	{
 		return VMI_RequiredTextureResolution;
 	}
+	else if (EngineShowFlags.VirtualTexturePendingMips)
+	{
+		return VMI_RequiredTextureResolution;
+	}
 	else if (EngineShowFlags.ShaderComplexity)
 	{
 		return VMI_ShaderComplexity;
@@ -792,6 +801,7 @@ const TCHAR* GetViewModeName(EViewModeIndex ViewModeIndex)
 		case VMI_MeshUVDensityAccuracy:		return TEXT("MeshUVDensityAccuracy");
 		case VMI_MaterialTextureScaleAccuracy: return TEXT("MaterialTextureScaleAccuracy");
 		case VMI_RequiredTextureResolution: return TEXT("RequiredTextureResolution");
+		case VMI_VirtualTexturePendingMips:	return TEXT("VirtualTexturePendingMips");
 		case VMI_StationaryLightOverlap:	return TEXT("StationaryLightOverlap");
 		case VMI_LightmapDensity:			return TEXT("LightmapDensity");
 		case VMI_LitLightmapDensity:		return TEXT("LitLightmapDensity");
diff --git a/Engine/Source/Runtime/Engine/Private/ViewModeNames.cpp b/Engine/Source/Runtime/Engine/Private/ViewModeNames.cpp
--- a/Engine/Source/Runtime/Engine/Private/ViewModeNames.cpp
+++ b/Engine/Source/Runtime/Engine/Private/ViewModeNames.cpp
@@ -153,6 +153,11 @@ TArray<FText> FillViewModeDisplayNames()
 			ViewModeDisplayNames.Emplace(LOCTEXT("UViewModeUtils_VMI_VisualizeNanite", "Nanite Visualization"));
 		}
 
+		else if (ViewModeIndex == VMI_VirtualTexturePendingMips)
+		{
+			ViewModeDisplayNames.Emplace(LOCTEXT("UViewModeUtils_VMI_VirtualTexturePendingMips", "Virtual Texture Pending Mips"));
+		}
+
 		// VMI_Max
 		else if (ViewModeIndex == VMI_Max)
 		{
@@ -319,7 +324,7 @@ TArray<const FSlateBrush*> FillViewModeDisplayIcons()
 		{
 			ViewModeDisplayIcons.Emplace(FAppStyle::Get().GetBrush("EditorViewport.RequiredTextureResolutionMode"));
 		}
-	
+
 		// Ray tracing modes
 		// Run path tracing pipeline
 		else if (ViewModeIndex == VMI_PathTracing)
@@ -332,6 +337,12 @@ TArray<const FSlateBrush*> FillViewModeDisplayIcons()
 			ViewModeDisplayIcons.Emplace(FAppStyle::Get().GetBrush("EditorViewport.RayTracingDebugMode"));
 		}
 
+		// Compare the required texture resolution to the actual resolution
+		else if (ViewModeIndex == VMI_VirtualTexturePendingMips)
+		{
+			ViewModeDisplayIcons.Emplace(FAppStyle::Get().GetBrush("EditorViewport.VirtualTexturePendingMipsMode"));
+		}
+
 		// VMI_Max
 		else if (ViewModeIndex == VMI_Max)
 		{
diff --git a/Engine/Source/Runtime/Engine/Public/DebugViewModeHelpers.h b/Engine/Source/Runtime/Engine/Public/DebugViewModeHelpers.h
--- a/Engine/Source/Runtime/Engine/Public/DebugViewModeHelpers.h
+++ b/Engine/Source/Runtime/Engine/Public/DebugViewModeHelpers.h
@@ -29,8 +29,9 @@ enum EDebugViewShaderMode
 	DVSM_PrimitiveDistanceAccuracy,	// Visualize the accuracy of the primitive distance computed for texture streaming.
 	DVSM_MeshUVDensityAccuracy,		// Visualize the accuracy of the mesh UV densities computed for texture streaming.
 	DVSM_MaterialTextureScaleAccuracy, // Visualize the accuracy of the material texture scales used for texture streaming.
-	DVSM_OutputMaterialTextureScales, // Outputs the material texture scales.
-	DVSM_RequiredTextureResolution, // Visualize the accuracy of the material texture scales used for texture streaming.
+	DVSM_OutputMaterialTextureScales,  // Outputs the material texture scales.
+	DVSM_RequiredTextureResolution, // Visualize the accuracy of the streamed texture resolution.
+	DVSM_VirtualTexturePendingMips,	// Visualize the pending virtual texture mips.
 	DVSM_RayTracingDebug,			// Visualize ray tracing debug modes.
 	DVSM_LODColoration,				// Visualize primitive LOD .
 	DVSM_MAX
diff --git a/Engine/Source/Runtime/Engine/Public/ShowFlags.h b/Engine/Source/Runtime/Engine/Public/ShowFlags.h
--- a/Engine/Source/Runtime/Engine/Public/ShowFlags.h
+++ b/Engine/Source/Runtime/Engine/Public/ShowFlags.h
@@ -452,6 +452,7 @@ private:
 		SetMaterialTextureScaleAccuracy(false);
 		SetOutputMaterialTextureScales(false);
 		SetRequiredTextureResolution(false);
+		SetVirtualTexturePendingMips(false);
 		SetMotionBlur(InitMode != ESFIM_Editor && InitMode != ESFIM_VREditing);
 		SetBones(false);
 		SetServerDrawDebug(false);
diff --git a/Engine/Source/Runtime/Engine/Public/ShowFlagsValues.inl b/Engine/Source/Runtime/Engine/Public/ShowFlagsValues.inl
--- a/Engine/Source/Runtime/Engine/Public/ShowFlagsValues.inl
+++ b/Engine/Source/Runtime/Engine/Public/ShowFlagsValues.inl
@@ -353,6 +353,8 @@ SHOWFLAG_FIXED_IN_SHIPPING(0, MaterialTextureScaleAccuracy, SFG_Hidden, NSLOCTEX
 SHOWFLAG_FIXED_IN_SHIPPING(0, OutputMaterialTextureScales, SFG_Hidden, NSLOCTEXT("UnrealEd", "OutputMaterialTextureScales", "Output Material Texture Scales"))
 /** Compare the required texture resolution to the actual resolution. */
 SHOWFLAG_FIXED_IN_SHIPPING(0, RequiredTextureResolution, SFG_Hidden, NSLOCTEXT("UnrealEd", "RequiredTextureResolution", "Required Texture Resolution"))
+/** Compare the required mip level to the resident virtual texture mip level. */
+SHOWFLAG_FIXED_IN_SHIPPING(0, VirtualTexturePendingMips, SFG_Hidden, NSLOCTEXT("UnrealEd", "VirtualTexturePendingMips", "Virtual Texture Pending Mips"))
 /** If WidgetComponents should be rendered in the scene */
 SHOWFLAG_ALWAYS_ACCESSIBLE(WidgetComponents, SFG_Normal, NSLOCTEXT("UnrealEd", "WidgetComponentsSF", "Widget Components"))
 /** Draw the bones of all skeletal meshes */
diff --git a/Engine/Source/Runtime/Renderer/Private/BasePassRendering.cpp b/Engine/Source/Runtime/Renderer/Private/BasePassRendering.cpp
--- a/Engine/Source/Runtime/Renderer/Private/BasePassRendering.cpp
+++ b/Engine/Source/Runtime/Renderer/Private/BasePassRendering.cpp
@@ -735,6 +735,7 @@ void FDeferredShadingSceneRenderer::RenderBasePass(
 	const bool bDoParallelBasePass = bEnableParallelBasePasses && !bDebugViewMode && !bRenderLightmapDensity; // DebugView and LightmapDensity are non-parallel substitutions inside BasePass
 	const bool bNeedsBeginRender = AllowDebugViewmodes() &&
 		(ViewFamily.EngineShowFlags.RequiredTextureResolution ||
+			ViewFamily.EngineShowFlags.VirtualTexturePendingMips ||
 			ViewFamily.EngineShowFlags.MaterialTextureScaleAccuracy ||
 			ViewFamily.EngineShowFlags.MeshUVDensityAccuracy ||
 			ViewFamily.EngineShowFlags.PrimitiveDistanceAccuracy ||
diff --git a/Engine/Source/Runtime/Renderer/Private/DebugViewModeRendering.cpp b/Engine/Source/Runtime/Renderer/Private/DebugViewModeRendering.cpp
--- a/Engine/Source/Runtime/Renderer/Private/DebugViewModeRendering.cpp
+++ b/Engine/Source/Runtime/Renderer/Private/DebugViewModeRendering.cpp
@@ -380,6 +380,7 @@ void InitDebugViewModeInterfaces()
 	FDebugViewModeInterface::SetInterface(DVSM_MaterialTextureScaleAccuracy, new FMaterialTexCoordScaleAccuracyInterface());
 	FDebugViewModeInterface::SetInterface(DVSM_OutputMaterialTextureScales, new FOutputMaterialTexCoordScaleInterface());
 	FDebugViewModeInterface::SetInterface(DVSM_RequiredTextureResolution, new FRequiredTextureResolutionInterface());
+	FDebugViewModeInterface::SetInterface(DVSM_VirtualTexturePendingMips, new FRequiredTextureResolutionInterface());
 
 	FDebugViewModeInterface::SetInterface(DVSM_LODColoration, new FLODColorationInterface());
 }
diff --git a/Engine/Source/Runtime/Renderer/Private/PostProcess/PostProcessStreamingAccuracyLegend.cpp b/Engine/Source/Runtime/Renderer/Private/PostProcess/PostProcessStreamingAccuracyLegend.cpp
--- a/Engine/Source/Runtime/Renderer/Private/PostProcess/PostProcessStreamingAccuracyLegend.cpp
+++ b/Engine/Source/Runtime/Renderer/Private/PostProcess/PostProcessStreamingAccuracyLegend.cpp
@@ -54,6 +54,10 @@ FScreenPassTexture AddStreamingAccuracyLegendPass(FRDGBuilder& GraphBuilder, con
 		{
 			DrawDesc(Canvas, OutputViewRect.Min.X + 115, OutputViewRect.Max.Y - 75, LOCTEXT("DescRequiredTextureResolution", "Shows the ratio between the currently streamed texture resolution and the resolution wanted by the GPU."));
 		}
+		else if (ViewFamily.GetDebugViewShaderMode() == DVSM_VirtualTexturePendingMips)
+		{
+			DrawDesc(Canvas, OutputViewRect.Min.X + 115, OutputViewRect.Max.Y - 75, LOCTEXT("DescVirtualTexturePendingMips", "Shows the number of pending virtual texture mips to reach the resolution wanted by the GPU."));
+		}
 		else if (ViewFamily.GetDebugViewShaderMode() == DVSM_MaterialTextureScaleAccuracy)
 		{
 			DrawDesc(Canvas, OutputViewRect.Min.X + 115, OutputViewRect.Max.Y - 75, LOCTEXT("DescMaterialTextureScaleAccuracy", "Shows under/over texture streaming caused by the material texture scales applied when sampling."));
@@ -71,11 +75,15 @@ FScreenPassTexture AddStreamingAccuracyLegendPass(FRDGBuilder& GraphBuilder, con
 		DrawBox(Canvas, OutputViewRect.Min.X + 115, OutputViewRect.Max.Y - 25, Colors[0], LOCTEXT("2XUnder", "2X+ Under"));
 		DrawBox(Canvas, OutputViewRect.Min.X + 215, OutputViewRect.Max.Y - 25, Colors[1], LOCTEXT("1XUnder", "1X Under"));
 		DrawBox(Canvas, OutputViewRect.Min.X + 315, OutputViewRect.Max.Y - 25, Colors[2], LOCTEXT("Good", "Good"));
-		DrawBox(Canvas, OutputViewRect.Min.X + 415, OutputViewRect.Max.Y - 25, Colors[3], LOCTEXT("1xOver", "1X Over"));
-		DrawBox(Canvas, OutputViewRect.Min.X + 515, OutputViewRect.Max.Y - 25, Colors[4], LOCTEXT("2XOver", "2X+ Over"));
 
-		const FLinearColor UndefColor(UndefinedStreamingAccuracyIntensity, UndefinedStreamingAccuracyIntensity, UndefinedStreamingAccuracyIntensity, 1.f);
-		DrawBox(Canvas, OutputViewRect.Min.X + 615, OutputViewRect.Max.Y - 25, UndefColor, LOCTEXT("Undefined", "Undefined"));
+		if (ViewFamily.GetDebugViewShaderMode() != DVSM_VirtualTexturePendingMips)
+		{
+			DrawBox(Canvas, OutputViewRect.Min.X + 415, OutputViewRect.Max.Y - 25, Colors[3], LOCTEXT("1xOver", "1X Over"));
+			DrawBox(Canvas, OutputViewRect.Min.X + 515, OutputViewRect.Max.Y - 25, Colors[4], LOCTEXT("2XOver", "2X+ Over"));
+
+			const FLinearColor UndefColor(UndefinedStreamingAccuracyIntensity, UndefinedStreamingAccuracyIntensity, UndefinedStreamingAccuracyIntensity, 1.f);
+			DrawBox(Canvas, OutputViewRect.Min.X + 615, OutputViewRect.Max.Y - 25, UndefColor, LOCTEXT("Undefined", "Undefined"));
+		}
 
 		if (ViewFamily.GetDebugViewShaderMode() == DVSM_MaterialTextureScaleAccuracy || ViewFamily.GetDebugViewShaderMode() == DVSM_MeshUVDensityAccuracy)
 		{
diff --git a/Engine/Source/Runtime/Renderer/Private/PostProcess/PostProcessing.cpp b/Engine/Source/Runtime/Renderer/Private/PostProcess/PostProcessing.cpp
--- a/Engine/Source/Runtime/Renderer/Private/PostProcess/PostProcessing.cpp
+++ b/Engine/Source/Runtime/Renderer/Private/PostProcess/PostProcessing.cpp
@@ -1469,6 +1469,7 @@ void AddDebugViewPostProcessingPasses(FRDGBuilder& GraphBuilder, const FViewInfo
 		case DVSM_MeshUVDensityAccuracy:
 		case DVSM_MaterialTextureScaleAccuracy:
 		case DVSM_RequiredTextureResolution:
+		case DVSM_VirtualTexturePendingMips:
 		{
 			FStreamingAccuracyLegendInputs PassInputs;
 			PassInputs.OverrideOutput = OverrideOutput;
diff --git a/Engine/Source/Runtime/Renderer/Private/RequiredTextureResolutionRendering.cpp b/Engine/Source/Runtime/Renderer/Private/RequiredTextureResolutionRendering.cpp
--- a/Engine/Source/Runtime/Renderer/Private/RequiredTextureResolutionRendering.cpp
+++ b/Engine/Source/Runtime/Renderer/Private/RequiredTextureResolutionRendering.cpp
@@ -29,6 +29,7 @@ void FRequiredTextureResolutionInterface::GetDebugViewModeShaderBindings(
 ) const
 {
 	const FRequiredTextureResolutionPS& Shader = static_cast<const FRequiredTextureResolutionPS&>(ShaderBase);
+	const int32 AnalysisType = DebugViewMode == DVSM_VirtualTexturePendingMips ? 1 : 0;
 	int32 AnalysisIndex = INDEX_NONE;
 	int32 TextureResolution = 64;
 	FMaterialRenderContext MaterialContext(&MaterialRenderProxy, Material, nullptr);
@@ -88,7 +89,7 @@ void FRequiredTextureResolutionInterface::GetDebugViewModeShaderBindings(
 		}
 	}
 
-	ShaderBindings.Add(Shader.AnalysisParamsParameter, FIntPoint(AnalysisIndex, TextureResolution));
+	ShaderBindings.Add(Shader.AnalysisParamsParameter, FIntVector(AnalysisType, AnalysisIndex, TextureResolution));
 	ShaderBindings.Add(Shader.PrimitiveAlphaParameter, (!PrimitiveSceneProxy || PrimitiveSceneProxy->IsSelected()) ? 1.f : .2f);
 }
 
diff --git a/Engine/Source/Runtime/Renderer/Private/RequiredTextureResolutionRendering.h b/Engine/Source/Runtime/Renderer/Private/RequiredTextureResolutionRendering.h
--- a/Engine/Source/Runtime/Renderer/Private/RequiredTextureResolutionRendering.h
+++ b/Engine/Source/Runtime/Renderer/Private/RequiredTextureResolutionRendering.h
@@ -26,7 +26,9 @@ public:
 
 	static bool ShouldCompilePermutation(const FMeshMaterialShaderPermutationParameters& Parameters)
 	{
-		return ShouldCompileDebugViewModeShader(DVSM_RequiredTextureResolution, Parameters);
+		return 
+			ShouldCompileDebugViewModeShader(DVSM_RequiredTextureResolution, Parameters) ||
+			ShouldCompileDebugViewModeShader(DVSM_VirtualTexturePendingMips, Parameters);
 	}
 
 	FRequiredTextureResolutionPS(const ShaderMetaType::CompiledShaderInitializerType& Initializer):
