Fixes a source of nan in the history parralax rejection factor heuristic that can cause random history rejection under high movements

diff --git a/Engine/Shaders/Private/TemporalSuperResolution/TSRDecimateHistory.usf b/Engine/Shaders/Private/TemporalSuperResolution/TSRDecimateHistory.usf
--- a/Engine/Shaders/Private/TemporalSuperResolution/TSRDecimateHistory.usf
+++ b/Engine/Shaders/Private/TemporalSuperResolution/TSRDecimateHistory.usf
@@ -326,6 +326,9 @@ void MainCS(
 		taa_half PixelParallaxVelocity = taa_half(0.5) * length(taa_half2(RotationalPrevScreen - PrevScreen) * InputInfo_Extent);
 		taa_half HistorySubdetailAdditionalParralax = (HistorySubdetailRejectionFactor - taa_half(1.0)) * PixelParallaxVelocity;
 
+		// PixelParallaxVelocity can be +inf and generate NaN if HistorySubdetailRejectionFactor = 1.0;
+		HistorySubdetailAdditionalParralax = -min(-HistorySubdetailAdditionalParralax, taa_half(0.0));
+
 		taa_half HistoryParallaxRejectionMask = saturate(taa_half(1.0) - HistorySubdetailAdditionalParralax);
 
 		ParallaxRejectionMask = min(ParallaxRejectionMask, HistoryParallaxRejectionMask);
