Implements TSR velocity hole filling to have better guess of what parallax disoccluded area might look like at history reprojection

diff --git a/Engine/Shaders/Private/TemporalSuperResolution/TSRCommon.ush b/Engine/Shaders/Private/TemporalSuperResolution/TSRCommon.ush
--- a/Engine/Shaders/Private/TemporalSuperResolution/TSRCommon.ush
+++ b/Engine/Shaders/Private/TemporalSuperResolution/TSRCommon.ush
@@ -139,10 +139,10 @@ static const int2 Offsets2x2[4] =
 };
 
 // Indexes of the 3x3 square.
-static const uint kSquareIndexes3x3[9] = { 4, 0, 1, 2, 3, 5, 6, 7, 8 };
+static const uint kSquareIndexes3x3[9] = { 4, 0, 1, 2, 3, 8, 7, 6, 5 };
 
 // Indexes of the offsets to have plus + shape.
-static const uint kPlusIndexes3x3[5] = { 4, 1, 3, 5, 7 };
+static const uint kPlusIndexes3x3[5] = { 4, 1, 3, 7, 5 };
 
 
 #if CONFIG_COMPILE_FP16
@@ -1019,3 +1019,47 @@ taa_half2 DecodeSpatialAntiAliasingOffset(uint EncodedInputTexelOffset)
 {
 	return dpv_lo(DecodeSpatialAntiAliasingOffset(dpv_interleave_mono_registers(EncodedInputTexelOffset)));
 }
+
+
+#define VELOCITY_HOLE_FILLING_BITS 18
+#define VELOCITY_HOLE_FILLING_LENGTH_BITS 13
+#define VELOCITY_HOLE_FILLING_ANGLE_BITS (VELOCITY_HOLE_FILLING_BITS - VELOCITY_HOLE_FILLING_LENGTH_BITS)
+#define VELOCITY_HOLE_FILLING_LENGTH_PRECISION (pow(2.0, 5))
+
+uint EncodeHoleFillingVelocity(float2 PixelVelocity)
+{
+	const uint AngleBitDepth = 1 << VELOCITY_HOLE_FILLING_ANGLE_BITS;
+	const uint LengthBitDepth = 1 << VELOCITY_HOLE_FILLING_LENGTH_BITS;
+	const float LengthPixelPrecision = VELOCITY_HOLE_FILLING_LENGTH_PRECISION;
+
+	float EncodedAngle = atan2Fast(PixelVelocity.y, PixelVelocity.x) * (0.5 * AngleBitDepth / PI) + (0.5 * AngleBitDepth);
+	uint iEncodedAngle = uint(round(EncodedAngle)) & (AngleBitDepth - 1);
+
+	float EncodedLength = length(PixelVelocity) * LengthPixelPrecision;
+	uint iEncodedLength = clamp(ceil(EncodedLength), 0, LengthBitDepth - 1);
+
+	uint EncodedHoleFillingVelocity = iEncodedAngle * LengthBitDepth + iEncodedLength;
+	return EncodedHoleFillingVelocity;
+}
+
+void DecodeHoleFillingVelocity(uint EncodedHoleFillingVelocity, out float2 PixelVelocity, out float OutOfRangeFactor)
+{
+	const uint AngleBitDepth = 1 << VELOCITY_HOLE_FILLING_ANGLE_BITS;
+	const uint LengthBitDepth = 1 << VELOCITY_HOLE_FILLING_LENGTH_BITS;
+	const float LengthPixelPrecision = VELOCITY_HOLE_FILLING_LENGTH_PRECISION;
+	const float LengthRange = float(LengthBitDepth) / LengthPixelPrecision;
+
+	uint iEncodedLength = EncodedHoleFillingVelocity % LengthBitDepth;
+	uint iEncodedAngle = (EncodedHoleFillingVelocity >> VELOCITY_HOLE_FILLING_LENGTH_BITS) % AngleBitDepth;
+
+	float EncodedLength = float(iEncodedLength);
+	float Length = EncodedLength * rcp(LengthPixelPrecision);
+
+	float EncodedAngle = float(iEncodedAngle);
+	float Angle = EncodedAngle * (PI * 2.0 / float(AngleBitDepth)) - PI;
+
+	sincos(Angle, /* out */ PixelVelocity.y, /* out */ PixelVelocity.x);
+	PixelVelocity *= Length;
+
+	OutOfRangeFactor = saturate(LengthRange - Length);
+}
diff --git a/Engine/Shaders/Private/TemporalSuperResolution/TSRDecimateHistory.usf b/Engine/Shaders/Private/TemporalSuperResolution/TSRDecimateHistory.usf
--- a/Engine/Shaders/Private/TemporalSuperResolution/TSRDecimateHistory.usf
+++ b/Engine/Shaders/Private/TemporalSuperResolution/TSRDecimateHistory.usf
@@ -113,7 +113,7 @@ void MainCS(
 				uint HistoryUseCount = PrevUseCountTexture[PixelPos];
 			#endif
 
-			float HistoryClosestDeviceZ = f16tof32(PrevClosestDepthTexture[PixelPos]);
+			float HistoryClosestDeviceZ = f16tof32(PrevClosestDepthTexture[PixelPos] >> VELOCITY_HOLE_FILLING_BITS);
 			float HistoryClosestWorldDepth = ConvertFromDeviceZ(HistoryClosestDeviceZ);
 
 			taa_half DepthRejection;
diff --git a/Engine/Shaders/Private/TemporalSuperResolution/TSRDilateVelocity.usf b/Engine/Shaders/Private/TemporalSuperResolution/TSRDilateVelocity.usf
--- a/Engine/Shaders/Private/TemporalSuperResolution/TSRDilateVelocity.usf
+++ b/Engine/Shaders/Private/TemporalSuperResolution/TSRDilateVelocity.usf
@@ -163,6 +163,11 @@ void MainCS(
 	{
 		#if 1
 		{
+			float2 PixelVelocity = FinalScreenVelocity * 0.5 * InputInfo_ViewportSize;
+			uint EncodedHoleFillingVelocity = EncodeHoleFillingVelocity(PixelVelocity);
+
+			uint PrevClosestDepth = (f32tof16(ClosestDeviceZ) << VELOCITY_HOLE_FILLING_BITS) | EncodedHoleFillingVelocity;
+			
 			bool bValidHistoryCoord = all(PrevInputBufferUV > PrevOutputBufferUVMin && PrevInputBufferUV < PrevOutputBufferUVMax);
 
 			FBilinearSampleInfos BilinearInter = GetBilinearSampleLevelInfos(PrevInputBufferUV, InputInfo_Extent, InputInfo_ExtentInverse);
@@ -188,7 +193,7 @@ void MainCS(
 				#endif
 
 				InterlockedAdd(PrevUseCountOutput[CompressedUseCountPixelPos], FinalAdd);
-				InterlockedMax(PrevClosestDepthOutput[PixelPos], f32tof16(ClosestDeviceZ));
+				InterlockedMax(PrevClosestDepthOutput[PixelPos], PrevClosestDepth);
 			}
 		}
 		#endif
diff --git a/Engine/Shaders/Private/TemporalSuperResolution/TSRHoleFillVelocity.usf b/Engine/Shaders/Private/TemporalSuperResolution/TSRHoleFillVelocity.usf
new file mode 100644
--- /dev/null
+++ b/Engine/Shaders/Private/TemporalSuperResolution/TSRHoleFillVelocity.usf
@@ -0,0 +1,171 @@
+// Copyright Epic Games, Inc. All Rights Reserved.
+
+#include "TSRCommon.ush"
+
+
+//------------------------------------------------------- CONFIG
+
+#define TILE_SIZE 8
+
+#define CONFIG_SAMPLES 9
+
+
+//------------------------------------------------------- PARAMETERS
+
+Texture2D<uint> PrevClosestDepthTexture;
+Texture2D<float2> DilatedVelocityTexture;
+Texture2D<taa_half> ParallaxRejectionMaskTexture;
+
+RWTexture2D<float2> HoleFilledVelocityOutput;
+RWTexture2D<taa_half> HoleFilledVelocityMaskOutput;
+
+
+//------------------------------------------------------- ENTRY POINT
+
+[numthreads(TILE_SIZE, TILE_SIZE, 1)]
+void MainCS(
+	uint2 GroupId : SV_GroupID,
+	uint GroupThreadIndex : SV_GroupIndex)
+{
+	float4 Debug = 0.0;
+
+	uint2 DispatchThreadId = (
+		ZOrder2D(GroupThreadIndex, log2(TILE_SIZE)) +
+		GroupId * uint2(TILE_SIZE, TILE_SIZE));
+
+	taa_short2 InputPixelPos = taa_short2(InputInfo_ViewportMin + DispatchThreadId);
+	
+	float2 ViewportUV = (float2(DispatchThreadId) + 0.5f) * InputInfo_ViewportSizeInverse;
+	float2 ScreenPos = ViewportUVToScreenPos(ViewportUV);
+
+	float2 EncodedVelocityNeighborhood[CONFIG_SAMPLES];
+	taa_half ParallaxRejectionMaskNeighborhood[CONFIG_SAMPLES];
+	ISOLATE
+	{
+		UNROLL_N(CONFIG_SAMPLES)
+		for (uint i = 0; i < CONFIG_SAMPLES; i++)
+		{
+			taa_short2 Offset;
+			#if CONFIG_SAMPLES == 9
+				Offset = kOffsets3x3[kSquareIndexes3x3[i]];
+			#elif CONFIG_SAMPLES == 5
+				Offset = kOffsets3x3[kPlusIndexes3x3[i]];
+			#endif
+
+			taa_short2 SampleInputPixelPos = InputPixelPos + ClampPixelOffset(
+				InputPixelPos,
+				Offset, Offset,
+				InputPixelPosMin, InputPixelPosMax);
+
+			EncodedVelocityNeighborhood[i] = DilatedVelocityTexture[SampleInputPixelPos];
+			ParallaxRejectionMaskNeighborhood[i] = ParallaxRejectionMaskTexture[SampleInputPixelPos];
+		}
+	}
+	
+	float2 ScreenVelocity = DecodeVelocityFromTexture(float4(EncodedVelocityNeighborhood[0], 0.0, 0.0)).xy;
+	
+	// Get the closest hole filling velocity.
+	float2 HoleFillingPixelVelocity;
+	float HoleFillingOutOfRange;
+	{
+		float2 PrevScreenPos = ScreenPos - ScreenVelocity;
+		float2 PrevInputBufferUV = (InputInfo_ScreenPosToViewportScale * PrevScreenPos + InputInfo_ScreenPosToViewportBias) * InputInfo_ExtentInverse;
+		
+		FBilinearSampleInfos BilinearInter = GetBilinearSampleLevelInfos(PrevInputBufferUV, InputInfo_Extent, InputInfo_ExtentInverse);
+		
+		bool bIsValidHoleFillingPixelVelocity = false;
+		uint EncodedHoleFillingVelocity = ~0;
+		UNROLL_N(4)
+		for (uint i = 0; i < 4; i++)
+		{
+			uint2 PixelPos = GetSamplePixelCoord(BilinearInter, i);
+
+			bool bValidPixelPos = all(PixelPos >= InputInfo_ViewportMin && PixelPos < InputInfo_ViewportMax);
+
+			uint SampleHistoryClosestDeviceZ = PrevClosestDepthTexture[PixelPos];
+
+			bIsValidHoleFillingPixelVelocity = bIsValidHoleFillingPixelVelocity || bValidPixelPos;
+
+			FLATTEN
+			if (bValidPixelPos)
+			{
+				EncodedHoleFillingVelocity = min(EncodedHoleFillingVelocity, SampleHistoryClosestDeviceZ);
+			}
+		}
+
+		DecodeHoleFillingVelocity(EncodedHoleFillingVelocity, /* out */ HoleFillingPixelVelocity, /* out */ HoleFillingOutOfRange);
+
+		if (!bIsValidHoleFillingPixelVelocity)
+		{
+			HoleFillingOutOfRange = 0.0;
+		}
+		
+		//Debug.r = EncodedHoleFillingVelocity;
+		//Debug.r = f16tof32(EncodedHoleFillingVelocity >> VELOCITY_HOLE_FILLING_BITS);
+		//Debug.r = f16tof32(EncodedHoleFillingVelocity);
+		//Debug.r = abs(EncodedHoleFillingVelocity % (1 << VELOCITY_HOLE_FILLING_LENGTH_BITS));
+	}
+	
+	bool bIsParallaxRejected = ParallaxRejectionMaskNeighborhood[0] < taa_half(PARALLAX_REJECTION_MASK_THRESHOLD);
+
+	taa_half HoleFilledVelocityMask = taa_half(0.0);
+
+	FLATTEN
+	if (bIsParallaxRejected && View.GeneralPurposeTweak >= 1)
+	{
+		float2 ScreenPixelVelocity = ScreenVelocity * InputInfo_ViewportSize * 0.5;
+		float ScreenPixelVelocityLength = length2(ScreenPixelVelocity);
+		
+		UNROLL_N((CONFIG_SAMPLES - 1) / 2)
+		for (uint i = 0; i < (CONFIG_SAMPLES - 1) / 2; i++)
+		{
+			uint NeightbordId0 = i + 1;
+			uint NeightbordId1 = i + 1 + (CONFIG_SAMPLES - 1) / 2;
+
+			float2 NeighborVelocity0 = DecodeVelocityFromTexture(float4(EncodedVelocityNeighborhood[NeightbordId0], 0.0, 0.0)).xy * InputInfo_ViewportSize * 0.5;
+			float2 NeighborVelocity1 = DecodeVelocityFromTexture(float4(EncodedVelocityNeighborhood[NeightbordId1], 0.0, 0.0)).xy * InputInfo_ViewportSize * 0.5;
+			bool bNeighborIsParallaxRejected0 = ParallaxRejectionMaskNeighborhood[NeightbordId0] < taa_half(PARALLAX_REJECTION_MASK_THRESHOLD);
+			bool bNeighborIsParallaxRejected1 = ParallaxRejectionMaskNeighborhood[NeightbordId1] < taa_half(PARALLAX_REJECTION_MASK_THRESHOLD);
+			
+			bool bIsPixelThinRejection = !bNeighborIsParallaxRejected0 && !bNeighborIsParallaxRejected1; // && bIsParallaxRejected;
+
+			float IsSomeWhatIdentical0 = saturate(1.0 - length2(NeighborVelocity0 - ScreenPixelVelocity));
+			float IsSomeWhatIdentical1 = saturate(1.0 - length2(NeighborVelocity1 - ScreenPixelVelocity));
+
+			float NeighborPixelVelocityLength0 = length2(NeighborVelocity0);
+			float NeighborPixelVelocityLength1 = length2(NeighborVelocity1);
+			
+			if (bIsPixelThinRejection && NeighborPixelVelocityLength0 > ScreenPixelVelocityLength && IsSomeWhatIdentical1 > 0.0)
+			{
+				EncodedVelocityNeighborhood[0] = EncodedVelocityNeighborhood[NeightbordId0];
+				HoleFilledVelocityMask = taa_half(1.0);
+			}
+			else if (bIsPixelThinRejection && NeighborPixelVelocityLength1 > ScreenPixelVelocityLength && IsSomeWhatIdentical0 > 0.0)
+			{
+				EncodedVelocityNeighborhood[0] = EncodedVelocityNeighborhood[NeightbordId1];
+				HoleFilledVelocityMask = taa_half(1.0);
+			}
+		}
+
+		// If the hole filling velocity has better result that the neighborhood, use that.
+		if (HoleFillingOutOfRange > HoleFilledVelocityMask)
+		{
+			float2 HoleFillingVelocity = HoleFillingPixelVelocity * 2.0 * InputInfo_ViewportSizeInverse;
+
+			EncodedVelocityNeighborhood[0] = EncodeVelocityToTexture(float3(HoleFillingVelocity, 0.0)).xy;
+			HoleFilledVelocityMask = HoleFillingOutOfRange;
+		}
+	}
+
+	if (all(InputPixelPos < InputInfo_ViewportMax))
+	{
+		HoleFilledVelocityOutput[InputPixelPos] = EncodedVelocityNeighborhood[0];
+		HoleFilledVelocityMaskOutput[InputPixelPos] = HoleFilledVelocityMask;
+
+		#if DEBUG_OUTPUT
+		{
+			DebugOutput[InputPixelPos] = Debug;
+		}
+		#endif
+	}
+}
diff --git a/Engine/Shaders/Private/TemporalSuperResolution/TSRUpdateHistory.usf b/Engine/Shaders/Private/TemporalSuperResolution/TSRUpdateHistory.usf
--- a/Engine/Shaders/Private/TemporalSuperResolution/TSRUpdateHistory.usf
+++ b/Engine/Shaders/Private/TemporalSuperResolution/TSRUpdateHistory.usf
@@ -65,6 +65,7 @@ Texture2D<uint> ParallaxFactorTexture;
 Texture2D<taa_half> ParallaxRejectionMaskTexture;
 Texture2D<uint> AntiAliasingTexture;
 Texture2D<taa_half> NoiseFilteringTexture;
+Texture2D<taa_half> HoleFilledVelocityMaskTexture;
 
 float4 HistoryPixelPosToScreenPos;
 float4 HistoryPixelPosToPPCo;
@@ -245,6 +246,7 @@ void MainCS(
 	taa_half2 OutputPixelVelocity = taa_half(0.0).xx;
 	taa_half2 MoireWeight = taa_half(1.0).xx;
 	taa_half2 NoiseFiltering = taa_half(0.0).xx;
+	taa_half2 HoleFilledVelocityMask = taa_half(0.0).xx;
 	taa_subpixel_payload2 ParallaxFactorBits = taa_subpixel_payload(0x0).xx;
 	bool2 bIsResponsiveAAPixel = false;
 
@@ -359,6 +361,10 @@ void MainCS(
 			NoiseFiltering = dpv_interleave_registers(
 				NoiseFilteringTexture[dpv_lo(InputPixelPos)],
 				NoiseFilteringTexture[dpv_hi(InputPixelPos)]);
+			
+			HoleFilledVelocityMask = dpv_interleave_registers(
+				HoleFilledVelocityMaskTexture[dpv_lo(InputPixelPos)],
+				HoleFilledVelocityMaskTexture[dpv_hi(InputPixelPos)]);
 
 			taa_half2x2 TexelOffset = DecodeSpatialAntiAliasingOffset(EncodedInputTexelOffset);
 			
@@ -1085,8 +1091,10 @@ void MainCS(
 
 			// Completely discard the history
 			{
+				taa_half2 TotalRejectionBlending = max(BlendFinal, taa_half(1.0) - HoleFilledVelocityMask);
+
 				PrevLowFrequencySampleCount *= TotalRejection;
-				BlendFinal = lerp(1.0, BlendFinal, TotalRejection);
+				BlendFinal = lerp(TotalRejectionBlending, BlendFinal, TotalRejection);
 			}
 
 			FinalLowFrequencyWeight = saturate(CurrentFrameSampleCount + PrevLowFrequencySampleCount);
diff --git a/Engine/Source/Runtime/Renderer/Private/PostProcess/TemporalSuperResolution.cpp b/Engine/Source/Runtime/Renderer/Private/PostProcess/TemporalSuperResolution.cpp
--- a/Engine/Source/Runtime/Renderer/Private/PostProcess/TemporalSuperResolution.cpp
+++ b/Engine/Source/Runtime/Renderer/Private/PostProcess/TemporalSuperResolution.cpp
@@ -82,6 +82,11 @@ TAutoConsoleVariable<float> CVarTSRWeightClampingPixelSpeed(
 	TEXT("Smallest reduce blur in movement (Default = 1.0f)."),
 	ECVF_RenderThreadSafe);
 
+TAutoConsoleVariable<int32> CVarTSRVelocityHoleFill(
+	TEXT("r.TSR.Velocity.HoleFill"), 1,
+	TEXT("Whether to holl-fill the velocity buffer on parralax disocclusion to reduce boiling of disoccluded areas."),
+	ECVF_Scalability | ECVF_RenderThreadSafe);
+
 #if COMPILE_TSR_DEBUG_PASSES
 
 TAutoConsoleVariable<int32> CVarTSRSetupDebugPasses(
@@ -309,6 +314,24 @@ class FTSRDetectInterferenceCS : public FTSRShader
 	END_SHADER_PARAMETER_STRUCT()
 }; // class FTSRDetectInterferenceCS
 
+class FTSRHoleFillVelocityCS : public FTSRShader
+{
+	DECLARE_GLOBAL_SHADER(FTSRHoleFillVelocityCS);
+	SHADER_USE_PARAMETER_STRUCT(FTSRHoleFillVelocityCS, FTSRShader);
+
+	BEGIN_SHADER_PARAMETER_STRUCT(FParameters, )
+		SHADER_PARAMETER_STRUCT_INCLUDE(FTSRCommonParameters, CommonParameters)
+
+		SHADER_PARAMETER_RDG_TEXTURE(Texture2D, PrevClosestDepthTexture)
+		SHADER_PARAMETER_RDG_TEXTURE(Texture2D, DilatedVelocityTexture)
+		SHADER_PARAMETER_RDG_TEXTURE(Texture2D, ParallaxRejectionMaskTexture)
+
+		SHADER_PARAMETER_RDG_TEXTURE_UAV(RWTexture2D, HoleFilledVelocityOutput)
+		SHADER_PARAMETER_RDG_TEXTURE_UAV(RWTexture2D, HoleFilledVelocityMaskOutput)
+		SHADER_PARAMETER_RDG_TEXTURE_UAV(RWTexture2D, DebugOutput)
+	END_SHADER_PARAMETER_STRUCT()
+}; // class FTSRHoleFillVelocityCS
+
 class FTSRFilterFrequenciesCS : public FTSRShader
 {
 	DECLARE_GLOBAL_SHADER(FTSRFilterFrequenciesCS);
@@ -481,6 +504,7 @@ class FTSRUpdateHistoryCS : public FTSRShader
 		SHADER_PARAMETER_RDG_TEXTURE(Texture2D, ParallaxRejectionMaskTexture)
 		SHADER_PARAMETER_RDG_TEXTURE(Texture2D, AntiAliasingTexture)
 		SHADER_PARAMETER_RDG_TEXTURE(Texture2D, NoiseFilteringTexture)
+		SHADER_PARAMETER_RDG_TEXTURE(Texture2D, HoleFilledVelocityMaskTexture)
 		
 		SHADER_PARAMETER(FScreenTransform, HistoryPixelPosToScreenPos)
 		SHADER_PARAMETER(FScreenTransform, HistoryPixelPosToPPCo)
@@ -526,6 +550,7 @@ IMPLEMENT_GLOBAL_SHADER(FTSRClearPrevTexturesCS,     "/Engine/Private/TemporalSu
 IMPLEMENT_GLOBAL_SHADER(FTSRDilateVelocityCS,        "/Engine/Private/TemporalSuperResolution/TSRDilateVelocity.usf",        "MainCS", SF_Compute);
 IMPLEMENT_GLOBAL_SHADER(FTSRDecimateHistoryCS,       "/Engine/Private/TemporalSuperResolution/TSRDecimateHistory.usf",       "MainCS", SF_Compute);
 IMPLEMENT_GLOBAL_SHADER(FTSRCompareTranslucencyCS,   "/Engine/Private/TemporalSuperResolution/TSRCompareTranslucency.usf",   "MainCS", SF_Compute);
+IMPLEMENT_GLOBAL_SHADER(FTSRHoleFillVelocityCS,      "/Engine/Private/TemporalSuperResolution/TSRHoleFillVelocity.usf",      "MainCS", SF_Compute);
 IMPLEMENT_GLOBAL_SHADER(FTSRDetectInterferenceCS,    "/Engine/Private/TemporalSuperResolution/TSRDetectInterference.usf",    "MainCS", SF_Compute);
 IMPLEMENT_GLOBAL_SHADER(FTSRFilterFrequenciesCS,     "/Engine/Private/TemporalSuperResolution/TSRFilterFrequencies.usf",     "MainCS", SF_Compute);
 IMPLEMENT_GLOBAL_SHADER(FTSRCompareHistoryCS,        "/Engine/Private/TemporalSuperResolution/TSRCompareHistory.usf",        "MainCS", SF_Compute);
@@ -579,6 +604,8 @@ void AddTemporalSuperResolutionPasses(
 
 	int32 RejectionAntiAliasingQuality = FMath::Clamp(CVarTSRRejectionAntiAliasingQuality.GetValueOnRenderThread(), 0, 2);
 
+	bool bHoleFillVelocity = CVarTSRVelocityHoleFill.GetValueOnRenderThread() != 0 && RejectionAntiAliasingQuality > 0;
+
 	enum class ERejectionPostFilter : uint8
 	{
 		Disabled,
@@ -1090,6 +1117,45 @@ void AddTemporalSuperResolutionPasses(
 		InterferenceWeightTexture = WhiteDummy;
 	}
 
+	FRDGTextureRef HoleFilledVelocityMaskTexture = BlackDummy;
+	if (bHoleFillVelocity)
+	{
+		FRDGTextureRef HoleFilledVelocityTexture;
+		{
+			FRDGTextureDesc Desc = FRDGTextureDesc::Create2D(
+				InputExtent,
+				PF_G16R16,
+				FClearValueBinding::None,
+				/* InFlags = */ TexCreate_ShaderResource | TexCreate_UAV);
+
+			HoleFilledVelocityTexture = GraphBuilder.CreateTexture(Desc, TEXT("TSR.Velocity.HoleFilled"));
+
+			Desc.Format = PF_R8;
+			HoleFilledVelocityMaskTexture = GraphBuilder.CreateTexture(Desc, TEXT("TSR.velocity.HoleFillMask"));
+		}
+
+		FTSRHoleFillVelocityCS::FParameters* PassParameters = GraphBuilder.AllocParameters<FTSRHoleFillVelocityCS::FParameters>();
+		PassParameters->CommonParameters = CommonParameters;
+
+		PassParameters->PrevClosestDepthTexture = PrevClosestDepthTexture;
+		PassParameters->DilatedVelocityTexture = DilatedVelocityTexture;
+		PassParameters->ParallaxRejectionMaskTexture = ParallaxRejectionMaskTexture;
+
+		PassParameters->HoleFilledVelocityOutput = GraphBuilder.CreateUAV(HoleFilledVelocityTexture);
+		PassParameters->HoleFilledVelocityMaskOutput = GraphBuilder.CreateUAV(HoleFilledVelocityMaskTexture);
+		PassParameters->DebugOutput = CreateDebugUAV(LowFrequencyExtent, TEXT("Debug.TSR.HoleFillelocity"));
+
+		TShaderMapRef<FTSRHoleFillVelocityCS> ComputeShader(View.ShaderMap);
+		FComputeShaderUtils::AddPass(
+			GraphBuilder,
+			RDG_EVENT_NAME("TSR HoleFillVelocity %dx%d", InputRect.Width(), InputRect.Height()),
+			ComputeShader,
+			PassParameters,
+			FComputeShaderUtils::GetGroupCount(InputRect.Size(), 8));
+
+		DilatedVelocityTexture = HoleFilledVelocityTexture;
+	}
+
 	// Reject the history with frequency decomposition.
 	FRDGTextureRef HistoryRejectionTexture;
 	FRDGTextureRef InputSceneColorLdrLumaTexture = nullptr;
@@ -1378,6 +1444,7 @@ void AddTemporalSuperResolutionPasses(
 		PassParameters->ParallaxRejectionMaskTexture = ParallaxRejectionMaskTexture;
 		PassParameters->AntiAliasingTexture = AntiAliasingTexture;
 		PassParameters->NoiseFilteringTexture = NoiseFilteringTexture;
+		PassParameters->HoleFilledVelocityMaskTexture = HoleFilledVelocityMaskTexture;
 
 		FScreenTransform HistoryPixelPosToViewportUV = (FScreenTransform::Identity + 0.5f) * CommonParameters.HistoryInfo.ViewportSizeInverse;
 		PassParameters->HistoryPixelPosToScreenPos = HistoryPixelPosToViewportUV * FScreenTransform::ViewportUVToScreenPos;
