Reduces TSR ghosting on lumen reflections

diff --git a/Engine/Shaders/Private/TemporalSuperResolution/TSRHoleFillVelocity.usf b/Engine/Shaders/Private/TemporalSuperResolution/TSRHoleFillVelocity.usf
--- a/Engine/Shaders/Private/TemporalSuperResolution/TSRHoleFillVelocity.usf
+++ b/Engine/Shaders/Private/TemporalSuperResolution/TSRHoleFillVelocity.usf
@@ -12,6 +12,8 @@
 
 //------------------------------------------------------- PARAMETERS
 
+float MaxHollFillPixelVelocity;
+
 Texture2D<uint> PrevClosestDepthTexture;
 Texture2D<float2> DilatedVelocityTexture;
 Texture2D<taa_half> ParallaxRejectionMaskTexture;
@@ -74,7 +76,7 @@ void MainCS(
 		FBilinearSampleInfos BilinearInter = GetBilinearSampleLevelInfos(PrevInputBufferUV, InputInfo_Extent, InputInfo_ExtentInverse);
 		
 		bool bIsValidHoleFillingPixelVelocity = false;
-		uint EncodedHoleFillingVelocity = ~0;
+		uint EncodedHoleFillingVelocity = 0;
 		UNROLL_N(4)
 		for (uint i = 0; i < 4; i++)
 		{
@@ -89,7 +91,7 @@ void MainCS(
 			FLATTEN
 			if (bValidPixelPos)
 			{
-				EncodedHoleFillingVelocity = min(EncodedHoleFillingVelocity, SampleHistoryClosestDeviceZ);
+				EncodedHoleFillingVelocity = max(EncodedHoleFillingVelocity, SampleHistoryClosestDeviceZ);
 			}
 		}
 
@@ -111,7 +113,7 @@ void MainCS(
 	taa_half HoleFilledVelocityMask = taa_half(0.0);
 
 	FLATTEN
-	if (bIsParallaxRejected && View.GeneralPurposeTweak >= 1)
+	if (bIsParallaxRejected)
 	{
 		float2 ScreenPixelVelocity = ScreenVelocity * InputInfo_ViewportSize * 0.5;
 		float ScreenPixelVelocityLength = length2(ScreenPixelVelocity);
@@ -150,10 +152,12 @@ void MainCS(
 		// If the hole filling velocity has better result that the neighborhood, use that.
 		if (HoleFillingOutOfRange > HoleFilledVelocityMask)
 		{
+			float VelocityDelta = length(ScreenPixelVelocity - HoleFillingPixelVelocity);
+
 			float2 HoleFillingVelocity = HoleFillingPixelVelocity * 2.0 * InputInfo_ViewportSizeInverse;
 
 			EncodedVelocityNeighborhood[0] = EncodeVelocityToTexture(float3(HoleFillingVelocity, 0.0)).xy;
-			HoleFilledVelocityMask = HoleFillingOutOfRange;
+			HoleFilledVelocityMask = HoleFillingOutOfRange * saturate(MaxHollFillPixelVelocity - VelocityDelta);
 		}
 	}
 
diff --git a/Engine/Source/Runtime/Renderer/Private/PostProcess/TemporalSuperResolution.cpp b/Engine/Source/Runtime/Renderer/Private/PostProcess/TemporalSuperResolution.cpp
--- a/Engine/Source/Runtime/Renderer/Private/PostProcess/TemporalSuperResolution.cpp
+++ b/Engine/Source/Runtime/Renderer/Private/PostProcess/TemporalSuperResolution.cpp
@@ -87,6 +87,11 @@ TAutoConsoleVariable<int32> CVarTSRVelocityHoleFill(
 	TEXT("Whether to holl-fill the velocity buffer on parralax disocclusion to reduce boiling of disoccluded areas."),
 	ECVF_Scalability | ECVF_RenderThreadSafe);
 
+TAutoConsoleVariable<float> CVarTSRVelocityMaxHoleFillScatterVelocity(
+	TEXT("r.TSR.Velocity.HoleFill.MaxScatterVelocity"), 4.0f,
+	TEXT("Maximum pixel velocity difference tolerated between the disoccluded and scattered disoccluding geometries."),
+	ECVF_RenderThreadSafe);
+
 #if COMPILE_TSR_DEBUG_PASSES
 
 TAutoConsoleVariable<int32> CVarTSRSetupDebugPasses(
@@ -321,6 +326,7 @@ class FTSRHoleFillVelocityCS : public FTSRShader
 
 	BEGIN_SHADER_PARAMETER_STRUCT(FParameters, )
 		SHADER_PARAMETER_STRUCT_INCLUDE(FTSRCommonParameters, CommonParameters)
+		SHADER_PARAMETER(float, MaxHollFillPixelVelocity)
 
 		SHADER_PARAMETER_RDG_TEXTURE(Texture2D, PrevClosestDepthTexture)
 		SHADER_PARAMETER_RDG_TEXTURE(Texture2D, DilatedVelocityTexture)
@@ -1137,6 +1143,7 @@ void AddTemporalSuperResolutionPasses(
 		FTSRHoleFillVelocityCS::FParameters* PassParameters = GraphBuilder.AllocParameters<FTSRHoleFillVelocityCS::FParameters>();
 		PassParameters->CommonParameters = CommonParameters;
 
+		PassParameters->MaxHollFillPixelVelocity = CVarTSRVelocityMaxHoleFillScatterVelocity.GetValueOnRenderThread();
 		PassParameters->PrevClosestDepthTexture = PrevClosestDepthTexture;
 		PassParameters->DilatedVelocityTexture = DilatedVelocityTexture;
 		PassParameters->ParallaxRejectionMaskTexture = ParallaxRejectionMaskTexture;
