Improves AA quality of geometric velocity edges under motion to sit better in the background.

diff --git a/Engine/Shaders/Private/TemporalSuperResolution/TSRCommon.ush b/Engine/Shaders/Private/TemporalSuperResolution/TSRCommon.ush
--- a/Engine/Shaders/Private/TemporalSuperResolution/TSRCommon.ush
+++ b/Engine/Shaders/Private/TemporalSuperResolution/TSRCommon.ush
@@ -7,6 +7,7 @@
 #include "../MonteCarlo.ush"
 #include "../ScreenPass.ush"
 #include "/Engine/Public/DualPixelVectorization.ush"
+#include "/Engine/Public/WaveBroadcastIntrinsics.ush"
 
 
 //------------------------------------------------------- RECOMPILE HASH
@@ -29,6 +30,12 @@
 	#define CONFIG_COMPILE_FP16 0
 #endif
 
+#if PLATFORM_SUPPORTS_WAVE_BROADCAST && 1
+	#define CONFIG_BUTTERFLY_KERNEL 1
+#else
+	#define CONFIG_BUTTERFLY_KERNEL 0
+#endif
+
 #define CONFIG_ENABLE_STOCASTIC_QUANTIZATION 1
 
 #define CONFIG_SCENE_COLOR_OVERFLOW 1
@@ -237,6 +244,9 @@ float2 InputInfo_UVViewportBilinearMax;
 float2 InputJitter;
 int2   InputPixelPosMin;
 int2   InputPixelPosMax;
+FScreenTransform InputPixelPosToScreenPos;
+float2 ScreenVelocityToInputPixelVelocity;
+float2 InputPixelVelocityToScreenVelocity;
 
 float2 LowFrequencyInfo_Extent;
 float2 LowFrequencyInfo_ExtentInverse;
@@ -1063,3 +1073,16 @@ void DecodeHoleFillingVelocity(uint EncodedHoleFillingVelocity, out float2 Pixel
 
 	OutOfRangeFactor = saturate(LengthRange - Length);
 }
+
+taa_short2 GetLaneOffsetSign()
+#if CONFIG_BUTTERFLY_KERNEL
+{
+	uint LaneIndex = WaveGetLaneIndex();
+
+	return taa_short(-1) + taa_short2(taa_ushort2(taa_ushort(LaneIndex) << taa_ushort(1), LaneIndex) & taa_ushort(0x2));
+}
+#else
+{
+	return taa_short(1).xx;
+}
+#endif
diff --git a/Engine/Shaders/Private/TemporalSuperResolution/TSRCompareTranslucency.usf b/Engine/Shaders/Private/TemporalSuperResolution/TSRCompareTranslucency.usf
--- a/Engine/Shaders/Private/TemporalSuperResolution/TSRCompareTranslucency.usf
+++ b/Engine/Shaders/Private/TemporalSuperResolution/TSRCompareTranslucency.usf
@@ -51,7 +51,6 @@ float PrevTranslucencyPreExposureCorrection;
 
 float TranslucencyHighlightLuminance;
 
-FScreenTransform InputPixelPosToScreenPos;
 FScreenTransform ScreenPosToPrevTranslucencyTextureUV;
 
 Texture2D<taa_half4> TranslucencyTexture;
diff --git a/Engine/Shaders/Private/TemporalSuperResolution/TSRHoleFillVelocity.usf b/Engine/Shaders/Private/TemporalSuperResolution/TSRHoleFillVelocity.usf
--- a/Engine/Shaders/Private/TemporalSuperResolution/TSRHoleFillVelocity.usf
+++ b/Engine/Shaders/Private/TemporalSuperResolution/TSRHoleFillVelocity.usf
@@ -8,6 +8,17 @@
 #define TILE_SIZE 8
 
 #define CONFIG_SAMPLES 9
+#define CONFIG_SAMPLE_PAIRS ((CONFIG_SAMPLES - 1) / 2)
+
+#if CONFIG_BUTTERFLY_KERNEL
+	#if CONFIG_SAMPLES == 9
+		#define CONFIG_PRE_FETCH_SAMPLES 4
+	#else
+		#define CONFIG_PRE_FETCH_SAMPLES 3
+	#endif
+#else
+	#define CONFIG_PRE_FETCH_SAMPLES (CONFIG_SAMPLES)
+#endif
 
 
 //------------------------------------------------------- PARAMETERS
@@ -19,7 +30,7 @@ Texture2D<float2> DilatedVelocityTexture;
 Texture2D<taa_half> ParallaxRejectionMaskTexture;
 
 RWTexture2D<float2> HoleFilledVelocityOutput;
-RWTexture2D<taa_half> HoleFilledVelocityMaskOutput;
+RWTexture2D<taa_half2> HoleFilledVelocityMaskOutput;
 
 
 //------------------------------------------------------- ENTRY POINT
@@ -34,31 +45,42 @@ void MainCS(
 	uint2 DispatchThreadId = (
 		ZOrder2D(GroupThreadIndex, log2(TILE_SIZE)) +
 		GroupId * uint2(TILE_SIZE, TILE_SIZE));
+	taa_short2 LaneOffsetSign = GetLaneOffsetSign();
 
 	taa_short2 InputPixelPos = taa_short2(InputInfo_ViewportMin + DispatchThreadId);
 	
-	float2 ViewportUV = (float2(DispatchThreadId) + 0.5f) * InputInfo_ViewportSizeInverse;
-	float2 ScreenPos = ViewportUVToScreenPos(ViewportUV);
+	float2 ScreenPos = ApplyScreenTransform(float2(InputPixelPos), InputPixelPosToScreenPos);
 
-	float2 EncodedVelocityNeighborhood[CONFIG_SAMPLES];
-	taa_half ParallaxRejectionMaskNeighborhood[CONFIG_SAMPLES];
+	float2 EncodedVelocityNeighborhood[CONFIG_PRE_FETCH_SAMPLES];
+	taa_half ParallaxRejectionMaskNeighborhood[CONFIG_PRE_FETCH_SAMPLES];
 	ISOLATE
 	{
-		UNROLL_N(CONFIG_SAMPLES)
-		for (uint i = 0; i < CONFIG_SAMPLES; i++)
+		UNROLL_N(CONFIG_PRE_FETCH_SAMPLES)
+		for (uint i = 0; i < CONFIG_PRE_FETCH_SAMPLES; i++)
 		{
-			taa_short2 Offset;
-			#if CONFIG_SAMPLES == 9
-				Offset = kOffsets3x3[kSquareIndexes3x3[i]];
-			#elif CONFIG_SAMPLES == 5
-				Offset = kOffsets3x3[kPlusIndexes3x3[i]];
+			taa_short2 SampleInputPixelPos;
+			#if CONFIG_BUTTERFLY_KERNEL
+			{
+				taa_short2 Offset = taa_short2(Offsets2x2[i]) * LaneOffsetSign;
+				
+				SampleInputPixelPos = InputPixelPos + Offset;
+			}
+			#else
+			{
+				taa_short2 Offset;
+				#if CONFIG_SAMPLES == 9
+					Offset = kOffsets3x3[kSquareIndexes3x3[i]];
+				#elif CONFIG_SAMPLES == 5
+					Offset = kOffsets3x3[kPlusIndexes3x3[i]];
+				#endif
+					
+				SampleInputPixelPos = InputPixelPos + ClampPixelOffset(
+					InputPixelPos,
+					Offset, Offset,
+					InputPixelPosMin, InputPixelPosMax);
+			}
 			#endif
 
-			taa_short2 SampleInputPixelPos = InputPixelPos + ClampPixelOffset(
-				InputPixelPos,
-				Offset, Offset,
-				InputPixelPosMin, InputPixelPosMax);
-
 			EncodedVelocityNeighborhood[i] = DilatedVelocityTexture[SampleInputPixelPos];
 			ParallaxRejectionMaskNeighborhood[i] = ParallaxRejectionMaskTexture[SampleInputPixelPos];
 		}
@@ -108,63 +130,186 @@ void MainCS(
 		//Debug.r = abs(EncodedHoleFillingVelocity % (1 << VELOCITY_HOLE_FILLING_LENGTH_BITS));
 	}
 	
+	// Preprocessed the prefetched samples
+	float2 PixelVelocityNeighborhood[CONFIG_PRE_FETCH_SAMPLES];
+	float PixelVelocityLengthNeighborhood[CONFIG_PRE_FETCH_SAMPLES];
+	{
+		UNROLL_N(CONFIG_PRE_FETCH_SAMPLES)
+		for (uint i = 0; i < CONFIG_PRE_FETCH_SAMPLES; i++)
+		{
+			PixelVelocityNeighborhood[i] = DecodeVelocityFromTexture(float4(EncodedVelocityNeighborhood[i], 0.0, 0.0)).xy * ScreenVelocityToInputPixelVelocity;
+			PixelVelocityLengthNeighborhood[i] = dot(PixelVelocityNeighborhood[i], PixelVelocityNeighborhood[i]);
+		}
+	}
+	
+	float2 EndodedScreenVelocity = EncodedVelocityNeighborhood[0];
+
 	bool bIsParallaxRejected = ParallaxRejectionMaskNeighborhood[0] < taa_half(PARALLAX_REJECTION_MASK_THRESHOLD);
 
 	taa_half HoleFilledVelocityMask = taa_half(0.0);
+	taa_half RejectHighFrequency = taa_half(0.0);
+	//taa_half IsVelocityEdge = taa_half(0.0);
 
-	FLATTEN
-	if (bIsParallaxRejected)
 	{
-		float2 ScreenPixelVelocity = ScreenVelocity * InputInfo_ViewportSize * 0.5;
-		float ScreenPixelVelocityLength = length2(ScreenPixelVelocity);
+		float2 ScreenPixelVelocity = ScreenVelocity * ScreenVelocityToInputPixelVelocity;
+		float ScreenPixelVelocityLength = dot(ScreenPixelVelocity, ScreenPixelVelocity);
 		
-		UNROLL_N((CONFIG_SAMPLES - 1) / 2)
-		for (uint i = 0; i < (CONFIG_SAMPLES - 1) / 2; i++)
+		UNROLL_N(CONFIG_SAMPLE_PAIRS)
+		for (uint i = 0; i < CONFIG_SAMPLE_PAIRS; i++)
+		ISOLATE
 		{
-			uint NeightbordId0 = i + 1;
-			uint NeightbordId1 = i + 1 + (CONFIG_SAMPLES - 1) / 2;
-
-			float2 NeighborVelocity0 = DecodeVelocityFromTexture(float4(EncodedVelocityNeighborhood[NeightbordId0], 0.0, 0.0)).xy * InputInfo_ViewportSize * 0.5;
-			float2 NeighborVelocity1 = DecodeVelocityFromTexture(float4(EncodedVelocityNeighborhood[NeightbordId1], 0.0, 0.0)).xy * InputInfo_ViewportSize * 0.5;
-			bool bNeighborIsParallaxRejected0 = ParallaxRejectionMaskNeighborhood[NeightbordId0] < taa_half(PARALLAX_REJECTION_MASK_THRESHOLD);
-			bool bNeighborIsParallaxRejected1 = ParallaxRejectionMaskNeighborhood[NeightbordId1] < taa_half(PARALLAX_REJECTION_MASK_THRESHOLD);
+			taa_short2 Offset0;
+			taa_short2 Offset1;
 			
-			bool bIsPixelThinRejection = !bNeighborIsParallaxRejected0 && !bNeighborIsParallaxRejected1; // && bIsParallaxRejected;
+			float2 NeighborEncodedVelocity0;
+			float2 NeighborEncodedVelocity1;
 
-			float IsSomeWhatIdentical0 = saturate(1.0 - length2(NeighborVelocity0 - ScreenPixelVelocity));
-			float IsSomeWhatIdentical1 = saturate(1.0 - length2(NeighborVelocity1 - ScreenPixelVelocity));
+			float2 NeighborPixelVelocity0;
+			float2 NeighborPixelVelocity1;
 
-			float NeighborPixelVelocityLength0 = length2(NeighborVelocity0);
-			float NeighborPixelVelocityLength1 = length2(NeighborVelocity1);
+			taa_half NeighborPixelVelocityLength0;
+			taa_half NeighborPixelVelocityLength1;
+			
+			bool bNeighborIsParallaxRejected0;
+			bool bNeighborIsParallaxRejected1;
 			
-			if (bIsPixelThinRejection && NeighborPixelVelocityLength0 > ScreenPixelVelocityLength && IsSomeWhatIdentical1 > 0.0)
+			#if CONFIG_BUTTERFLY_KERNEL
 			{
-				EncodedVelocityNeighborhood[0] = EncodedVelocityNeighborhood[NeightbordId0];
-				HoleFilledVelocityMask = taa_half(1.0);
+				const uint PairIndex = i;
+				
+				const uint ButterflyArray0[] = { 0x0, 0x0, 0x0, 0x1 };
+				const uint ButterflyArray1[] = { 0x1, 0x2, 0x3, 0x2 };
+				
+				uint FetchIndex[2];
+				FWaveBroadcastSettings BroadcastSettings[2];
+				BroadcastSettings[0] = InitWaveXorButterfly(ButterflyArray0[PairIndex]);
+				BroadcastSettings[1] = InitWaveXorButterfly(ButterflyArray1[PairIndex]);
+				
+				if (PairIndex == 3)
+				{
+					FetchIndex[0] = 2;
+					FetchIndex[1] = 1;
+					Offset0 = taa_short2(-1, 1);
+				}
+				else
+				{
+					FetchIndex[0] = PairIndex + 1;
+					FetchIndex[1] = PairIndex + 1;
+					Offset0 = taa_short2(Offsets2x2[PairIndex + 1]);
+				}
+
+				Offset0 *= LaneOffsetSign;
+				Offset1 = -Offset0;
+
+				NeighborEncodedVelocity0 = WaveBroadcast(BroadcastSettings[0], EncodedVelocityNeighborhood[FetchIndex[0]]);
+				NeighborEncodedVelocity1 = WaveBroadcast(BroadcastSettings[1], EncodedVelocityNeighborhood[FetchIndex[1]]);
+			
+				NeighborPixelVelocity0 = WaveBroadcast(BroadcastSettings[0], PixelVelocityNeighborhood[FetchIndex[0]]);
+				NeighborPixelVelocity1 = WaveBroadcast(BroadcastSettings[1], PixelVelocityNeighborhood[FetchIndex[1]]);
+			
+				NeighborPixelVelocityLength0 = WaveBroadcast(BroadcastSettings[0], PixelVelocityLengthNeighborhood[FetchIndex[0]]);
+				NeighborPixelVelocityLength1 = WaveBroadcast(BroadcastSettings[1], PixelVelocityLengthNeighborhood[FetchIndex[1]]);
+				
+				bNeighborIsParallaxRejected0 = WaveBroadcast(BroadcastSettings[0], ParallaxRejectionMaskNeighborhood[FetchIndex[0]]) < taa_half(PARALLAX_REJECTION_MASK_THRESHOLD);
+				bNeighborIsParallaxRejected1 = WaveBroadcast(BroadcastSettings[1], ParallaxRejectionMaskNeighborhood[FetchIndex[1]]) < taa_half(PARALLAX_REJECTION_MASK_THRESHOLD);
+			}
+			#else
+			{
+				const uint NeightbordId0 = i + 1;
+				const uint NeightbordId1 = i + 1 + CONFIG_SAMPLE_PAIRS;
+			
+				#if CONFIG_SAMPLES == 9
+					Offset0 = kOffsets3x3[kSquareIndexes3x3[NeightbordId0]];
+					Offset1 = kOffsets3x3[kSquareIndexes3x3[NeightbordId1]];
+				#elif CONFIG_SAMPLES == 5
+					Offset0 = kOffsets3x3[kPlusIndexes3x3[NeightbordId0]];
+					Offset1 = kOffsets3x3[kPlusIndexes3x3[NeightbordId1]];
+				#endif
+					
+				NeighborEncodedVelocity0 = EncodedVelocityNeighborhood[NeightbordId0];
+				NeighborEncodedVelocity1 = EncodedVelocityNeighborhood[NeightbordId1];
+			
+				NeighborPixelVelocity0 = PixelVelocityNeighborhood[NeightbordId0];
+				NeighborPixelVelocity1 = PixelVelocityNeighborhood[NeightbordId1];
+			
+				NeighborPixelVelocityLength0 = dot(NeighborPixelVelocity0, NeighborPixelVelocity0);
+				NeighborPixelVelocityLength1 = dot(NeighborPixelVelocity1, NeighborPixelVelocity1);
+				
+				bNeighborIsParallaxRejected0 = ParallaxRejectionMaskNeighborhood[NeightbordId0] < taa_half(PARALLAX_REJECTION_MASK_THRESHOLD);
+				bNeighborIsParallaxRejected1 = ParallaxRejectionMaskNeighborhood[NeightbordId1] < taa_half(PARALLAX_REJECTION_MASK_THRESHOLD);
+			}
+			#endif
+
+			taa_half2 NeighborDeltaPixelVelocity0 = taa_half2(ScreenPixelVelocity - NeighborPixelVelocity0);
+			taa_half2 NeighborDeltaPixelVelocity1 = taa_half2(ScreenPixelVelocity - NeighborPixelVelocity1);
+			
+			// Detect edge of the velocity
+			//{
+			//	taa_half MaxDeltaVelocityLength = saturate(max(NeighborPixelVelocityLength0, NeighborPixelVelocityLength1));
+			//
+			//	if (i == 0)
+			//		IsVelocityEdge = MaxDeltaVelocityLength;
+			//	else
+			//		IsVelocityEdge = max(IsVelocityEdge, MaxDeltaVelocityLength);
+			//}
+
+			// Detect and hole fill pixel thin parallax rejection to reduce distractnesss.
+			{
+				bool bIsPixelThinRejection = bIsParallaxRejected && !bNeighborIsParallaxRejected0 && !bNeighborIsParallaxRejected1;
+
+				taa_half IsSomeWhatIdentical0 = saturate(taa_half(1.0) - dot(NeighborDeltaPixelVelocity0, NeighborDeltaPixelVelocity0));
+				taa_half IsSomeWhatIdentical1 = saturate(taa_half(1.0) - dot(NeighborDeltaPixelVelocity1, NeighborDeltaPixelVelocity1));
+
+				FLATTEN
+				if (bIsPixelThinRejection && NeighborPixelVelocityLength0 > ScreenPixelVelocityLength && IsSomeWhatIdentical1 > 0.0)
+				{
+					EndodedScreenVelocity = NeighborEncodedVelocity0;
+					HoleFilledVelocityMask = taa_half(1.0);
+				}
+				else if (bIsPixelThinRejection && NeighborPixelVelocityLength1 > ScreenPixelVelocityLength && IsSomeWhatIdentical0 > 0.0)
+				{
+					EndodedScreenVelocity = NeighborEncodedVelocity1;
+					HoleFilledVelocityMask = taa_half(1.0);
+				}
 			}
-			else if (bIsPixelThinRejection && NeighborPixelVelocityLength1 > ScreenPixelVelocityLength && IsSomeWhatIdentical0 > 0.0)
+
+			// Reject high frequency on velocity edge discontinuity.
 			{
-				EncodedVelocityNeighborhood[0] = EncodedVelocityNeighborhood[NeightbordId1];
-				HoleFilledVelocityMask = taa_half(1.0);
+				taa_half IsOverlapingEdge0 = taa_half(1.0) - dot(saturate(NeighborDeltaPixelVelocity0 * Offset0), taa_half(1.0).xx); // * saturate(ScreenPixelVelocityLength - NeighborPixelVelocityLength0);
+				taa_half IsOverlapingEdge1 = taa_half(1.0) - dot(saturate(NeighborDeltaPixelVelocity1 * Offset1), taa_half(1.0).xx); // * saturate(ScreenPixelVelocityLength - NeighborPixelVelocityLength1);
+				
+				taa_half IsOverlapingEdge = min(IsOverlapingEdge0, IsOverlapingEdge1);
+
+				if (i == 0)
+					RejectHighFrequency = IsOverlapingEdge;
+				else
+					RejectHighFrequency = min(RejectHighFrequency, IsOverlapingEdge);
 			}
 		}
 
 		// If the hole filling velocity has better result that the neighborhood, use that.
-		if (HoleFillingOutOfRange > HoleFilledVelocityMask)
+		FLATTEN
+		if (taa_half(HoleFillingOutOfRange) > HoleFilledVelocityMask && bIsParallaxRejected)
 		{
 			float VelocityDelta = length(ScreenPixelVelocity - HoleFillingPixelVelocity);
 
-			float2 HoleFillingVelocity = HoleFillingPixelVelocity * 2.0 * InputInfo_ViewportSizeInverse;
+			float2 HoleFillingVelocity = HoleFillingPixelVelocity * InputPixelVelocityToScreenVelocity;
 
-			EncodedVelocityNeighborhood[0] = EncodeVelocityToTexture(float3(HoleFillingVelocity, 0.0)).xy;
-			HoleFilledVelocityMask = HoleFillingOutOfRange * saturate(MaxHollFillPixelVelocity - VelocityDelta);
+			EndodedScreenVelocity = EncodeVelocityToTexture(float3(HoleFillingVelocity, 0.0)).xy;
+			HoleFilledVelocityMask = taa_half(HoleFillingOutOfRange) * saturate(MaxHollFillPixelVelocity - VelocityDelta);
 		}
+
+		// Remove noise from high frequency rejections.
+		RejectHighFrequency = saturate(taa_half(1.0) - RejectHighFrequency * taa_half(1.1));
+
+		// Do not reject high frequency additionally where hole filling happen
+		RejectHighFrequency = max(RejectHighFrequency, HoleFilledVelocityMask);
 	}
 
 	if (all(InputPixelPos < InputInfo_ViewportMax))
 	{
-		HoleFilledVelocityOutput[InputPixelPos] = EncodedVelocityNeighborhood[0];
-		HoleFilledVelocityMaskOutput[InputPixelPos] = HoleFilledVelocityMask;
+		HoleFilledVelocityOutput[InputPixelPos] = EndodedScreenVelocity;
+		HoleFilledVelocityMaskOutput[InputPixelPos] = taa_half2(HoleFilledVelocityMask, RejectHighFrequency);
 
 		#if DEBUG_OUTPUT
 		{
diff --git a/Engine/Shaders/Private/TemporalSuperResolution/TSRUpdateHistory.usf b/Engine/Shaders/Private/TemporalSuperResolution/TSRUpdateHistory.usf
--- a/Engine/Shaders/Private/TemporalSuperResolution/TSRUpdateHistory.usf
+++ b/Engine/Shaders/Private/TemporalSuperResolution/TSRUpdateHistory.usf
@@ -65,7 +65,7 @@ Texture2D<uint> ParallaxFactorTexture;
 Texture2D<taa_half> ParallaxRejectionMaskTexture;
 Texture2D<uint> AntiAliasingTexture;
 Texture2D<taa_half> NoiseFilteringTexture;
-Texture2D<taa_half> HoleFilledVelocityMaskTexture;
+Texture2D<taa_half2> HoleFilledVelocityMaskTexture;
 
 float4 HistoryPixelPosToScreenPos;
 float4 HistoryPixelPosToPPCo;
@@ -246,7 +246,7 @@ void MainCS(
 	taa_half2 OutputPixelVelocity = taa_half(0.0).xx;
 	taa_half2 MoireWeight = taa_half(1.0).xx;
 	taa_half2 NoiseFiltering = taa_half(0.0).xx;
-	taa_half2 HoleFilledVelocityMask = taa_half(0.0).xx;
+	taa_half2x2 HoleFilledVelocityMask = dpv_interleave_mono_registers(taa_half(1.0).xx);
 	taa_subpixel_payload2 ParallaxFactorBits = taa_subpixel_payload(0x0).xx;
 	bool2 bIsResponsiveAAPixel = false;
 
@@ -362,7 +362,7 @@ void MainCS(
 				NoiseFilteringTexture[dpv_lo(InputPixelPos)],
 				NoiseFilteringTexture[dpv_hi(InputPixelPos)]);
 			
-			HoleFilledVelocityMask = dpv_interleave_registers(
+			HoleFilledVelocityMask = dpv_interleave_mono_registers(taa_half(1.0).xx) - dpv_interleave_registers(
 				HoleFilledVelocityMaskTexture[dpv_lo(InputPixelPos)],
 				HoleFilledVelocityMaskTexture[dpv_hi(InputPixelPos)]);
 
@@ -886,7 +886,7 @@ void MainCS(
 		taa_half2 InvPrevHistoryValidity = SafeRcp(PrevHistoryValidity);
 
 		// Compute how much the history needs to be rejected, based on the completeness.
-		taa_half2 PrevHistoryRejectionWeight = LowFrequencyRejection;
+		taa_half2 PrevHistoryRejectionWeight = LowFrequencyRejection * HoleFilledVelocityMask[1];
 			
 		//#if !CONFIG_CLAMP
 		//	PrevHistoryRejectionWeight = 1.0;
@@ -1091,7 +1091,7 @@ void MainCS(
 
 			// Completely discard the history
 			{
-				taa_half2 TotalRejectionBlending = max(BlendFinal, taa_half(1.0) - HoleFilledVelocityMask);
+				taa_half2 TotalRejectionBlending = max(BlendFinal, HoleFilledVelocityMask[0]);
 
 				PrevLowFrequencySampleCount *= TotalRejection;
 				BlendFinal = lerp(TotalRejectionBlending, BlendFinal, TotalRejection);
diff --git a/Engine/Shaders/Public/WaveBroadcastIntrinsics.ush b/Engine/Shaders/Public/WaveBroadcastIntrinsics.ush
--- a/Engine/Shaders/Public/WaveBroadcastIntrinsics.ush
+++ b/Engine/Shaders/Public/WaveBroadcastIntrinsics.ush
@@ -16,8 +16,7 @@
 /** Different instruction available for WaveBroadcast() */
 	#define WAVE_BROADCAST_NOP 0
 	#if PLATFORM_SUPPORTS_WAVE_BROADCAST
-		#define WAVE_BROADCAST_SWAP_WITHIN_LANE_GROUP 1
-		#define WAVE_BROADCAST_LANE_GROUP 2
+		#define WAVE_BROADCAST_GCN_SWIZZLE 1
 	#endif
 
 /** Compile time structure to choose which broadcasting should be done. */
@@ -27,9 +26,9 @@ struct FWaveBroadcastSettings
 	uint Operation;
 
 	// Informations about the broadcast.
-	uint LaneGroupSize;
-	uint InnerLaneGroupSize;
-	uint InnerLaneGroupId;
+	uint SwizzleAnd;
+	uint SwizzleOr;
+	uint SwizzleXor;
 };
 
 /** Generic broadcast instruction. */
@@ -40,21 +39,9 @@ uint WaveBroadcast(const FWaveBroadcastSettings Settings, uint x)
 		 return x;
 	}
 	#if PLATFORM_SUPPORTS_WAVE_BROADCAST
-	else if (Settings.Operation == WAVE_BROADCAST_SWAP_WITHIN_LANE_GROUP)
+	else if (Settings.Operation == WAVE_BROADCAST_GCN_SWIZZLE)
 	{
-		const uint and_mask = 0x1F;
-		const uint or_mask  = 0x00;
-		const uint xor_mask = Settings.LaneGroupSize >> 1;
-		return WaveLaneSwizzleGCN(x, and_mask, or_mask, xor_mask);
-	}
-	else if (Settings.Operation == WAVE_BROADCAST_LANE_GROUP)
-	{
-		const uint InnerGroupCount = Settings.LaneGroupSize / Settings.InnerLaneGroupSize;
-
-		const uint and_mask = ~((InnerGroupCount - 1) * Settings.InnerLaneGroupSize);
-		const uint or_mask  = Settings.InnerLaneGroupId * Settings.InnerLaneGroupSize;
-		const uint xor_mask = 0x00;
-		return WaveLaneSwizzleGCN(x, and_mask, or_mask, xor_mask);
+		return WaveLaneSwizzleGCN(x, Settings.SwizzleAnd, Settings.SwizzleOr, Settings.SwizzleXor);
 	}
 	#endif
 
@@ -145,9 +132,34 @@ float4 WaveBroadcast(const FWaveBroadcastSettings Settings, float4 v)
 
 
 
+/** Init a nop broadcast */
+FWaveBroadcastSettings InitNopBroadcast()
+{
+	FWaveBroadcastSettings Settings;
+	Settings.Operation = WAVE_BROADCAST_NOP;
+	Settings.SwizzleAnd = 0x00;
+	Settings.SwizzleOr  = 0x00;
+	Settings.SwizzleXor = 0x00;
+	return Settings;
+}
+
 
 #if PLATFORM_SUPPORTS_WAVE_BROADCAST
 
+/**  Broadcast in butterfly pattern.
+ *
+ *	return src[laneId ^ XorButterFly];
+ */
+FWaveBroadcastSettings InitWaveXorButterfly(const uint XorButterFly)
+{
+	FWaveBroadcastSettings Settings;
+	Settings.Operation = WAVE_BROADCAST_GCN_SWIZZLE;
+	Settings.SwizzleAnd = 0x1F;
+	Settings.SwizzleOr  = 0x00;
+	Settings.SwizzleXor = XorButterFly;
+	return Settings;
+}
+
 /** Swap left lane with righ lanes within lane group (size is power of two in [2; 64]). 
  *
  *  If a lane is not active, the VGPR value returned is 0.
@@ -162,12 +174,7 @@ float4 WaveBroadcast(const FWaveBroadcastSettings Settings, float4 v)
  */
 FWaveBroadcastSettings InitWaveSwapWithinLaneGroup(const uint LaneGroupSize)
 {
-	FWaveBroadcastSettings Settings;
-	Settings.Operation = WAVE_BROADCAST_SWAP_WITHIN_LANE_GROUP;
-	Settings.LaneGroupSize = LaneGroupSize;
-	Settings.InnerLaneGroupSize = 0;
-	Settings.InnerLaneGroupId = 0;
-	return Settings;
+	return InitWaveXorButterfly(/* XorButterFly = */ LaneGroupSize >> 1);
 }
 
 /**  Broadcast inner lane group over a lane group (size is power of two in [2; 64]).
@@ -185,11 +192,13 @@ FWaveBroadcastSettings InitWaveSwapWithinLaneGroup(const uint LaneGroupSize)
  */
 FWaveBroadcastSettings InitWaveBroadcastLaneGroup(const uint LaneGroupSize, const uint InnerLaneGroupSize, const uint InnerLaneGroupId)
 {
+	const uint InnerGroupCount = LaneGroupSize / InnerLaneGroupSize;
+
 	FWaveBroadcastSettings Settings;
-	Settings.Operation = WAVE_BROADCAST_LANE_GROUP;
-	Settings.LaneGroupSize = LaneGroupSize;
-	Settings.InnerLaneGroupSize = InnerLaneGroupSize;
-	Settings.InnerLaneGroupId = InnerLaneGroupId;
+	Settings.Operation = WAVE_BROADCAST_GCN_SWIZZLE;
+	Settings.SwizzleAnd = ~((InnerGroupCount - 1) * InnerLaneGroupSize);
+	Settings.SwizzleOr  = InnerLaneGroupId * InnerLaneGroupSize;
+	Settings.SwizzleXor = 0x00;
 	return Settings;
 }
 
diff --git a/Engine/Source/Runtime/Renderer/Private/PostProcess/TemporalSuperResolution.cpp b/Engine/Source/Runtime/Renderer/Private/PostProcess/TemporalSuperResolution.cpp
--- a/Engine/Source/Runtime/Renderer/Private/PostProcess/TemporalSuperResolution.cpp
+++ b/Engine/Source/Runtime/Renderer/Private/PostProcess/TemporalSuperResolution.cpp
@@ -11,6 +11,8 @@
 #include "ScenePrivate.h"
 #include "RendererModule.h"
 
+PRAGMA_DISABLE_OPTIMIZATION
+
 #define COMPILE_TSR_DEBUG_PASSES (!UE_BUILD_SHIPPING)
 
 namespace
@@ -88,8 +90,8 @@ TAutoConsoleVariable<int32> CVarTSRVelocityHoleFill(
 	ECVF_Scalability | ECVF_RenderThreadSafe);
 
 TAutoConsoleVariable<float> CVarTSRVelocityMaxHoleFillScatterVelocity(
-	TEXT("r.TSR.Velocity.HoleFill.MaxScatterVelocity"), 4.0f,
-	TEXT("Maximum pixel velocity difference tolerated between the disoccluded and scattered disoccluding geometries."),
+	TEXT("r.TSR.Velocity.HoleFill.MaxScatterVelocity"), 8.0f,
+	TEXT("Maximum output pixel velocity difference tolerated between the disoccluded and scattered disoccluding geometries."),
 	ECVF_RenderThreadSafe);
 
 #if COMPILE_TSR_DEBUG_PASSES
@@ -109,10 +111,13 @@ BEGIN_SHADER_PARAMETER_STRUCT(FTSRCommonParameters, )
 
 	SHADER_PARAMETER(FIntPoint, InputPixelPosMin)
 	SHADER_PARAMETER(FIntPoint, InputPixelPosMax)
+	SHADER_PARAMETER(FScreenTransform, InputPixelPosToScreenPos)
 
 	SHADER_PARAMETER(FVector2f, InputJitter)
 	SHADER_PARAMETER(int32, bCameraCut)
 	SHADER_PARAMETER(float, MoireInvExposure)
+	SHADER_PARAMETER(FVector2D, ScreenVelocityToInputPixelVelocity)
+	SHADER_PARAMETER(FVector2D, InputPixelVelocityToScreenVelocity)
 
 	SHADER_PARAMETER_STRUCT_REF(FViewUniformShaderParameters, ViewUniformBuffer)
 END_SHADER_PARAMETER_STRUCT()
@@ -284,7 +289,6 @@ class FTSRCompareTranslucencyCS : public FTSRShader
 		SHADER_PARAMETER(float, PrevTranslucencyPreExposureCorrection)
 		SHADER_PARAMETER(float, TranslucencyHighlightLuminance)
 
-		SHADER_PARAMETER(FScreenTransform, InputPixelPosToScreenPos)
 		SHADER_PARAMETER(FScreenTransform, ScreenPosToPrevTranslucencyTextureUV)
 
 		SHADER_PARAMETER_RDG_TEXTURE(Texture2D, DilatedVelocityTexture)
@@ -674,12 +678,15 @@ void AddTemporalSuperResolutionPasses(
 			FMath::Max(InputExtent.Y, QuantizedHistoryViewSize.Y));
 	}
 
+	float ScreenPercentage = float(InputRect.Width()) / float(OutputRect.Width());
+	float InvScreenPercentage = float(OutputRect.Width()) / float(InputRect.Width());
+
 	RDG_EVENT_SCOPE(GraphBuilder, "TemporalSuperResolution %dx%d -> %dx%d", InputRect.Width(), InputRect.Height(), OutputRect.Width(), OutputRect.Height());
 	RDG_GPU_STAT_SCOPE(GraphBuilder, TemporalSuperResolution);
 
 	FRDGTextureRef BlackUintDummy = GSystemTextures.GetZeroUIntDummy(GraphBuilder);
-	FRDGTextureRef BlackDummy = GraphBuilder.RegisterExternalTexture(GSystemTextures.BlackDummy);
-	FRDGTextureRef WhiteDummy = GraphBuilder.RegisterExternalTexture(GSystemTextures.WhiteDummy);
+	FRDGTextureRef BlackDummy = GSystemTextures.GetBlackDummy(GraphBuilder);
+	FRDGTextureRef WhiteDummy = GSystemTextures.GetWhiteDummy(GraphBuilder);
 
 	FRDGTextureRef SeparateTranslucencyTexture = (bRejectSeparateTranslucency || bAccumulateSeparateTranslucency) ?
 		PassInputs.SeparateTranslucencyTextures->GetColorForRead(GraphBuilder) : nullptr;
@@ -690,6 +697,9 @@ void AddTemporalSuperResolutionPasses(
 			InputExtent, InputRect));
 		CommonParameters.InputPixelPosMin = CommonParameters.InputInfo.ViewportMin;
 		CommonParameters.InputPixelPosMax = CommonParameters.InputInfo.ViewportMax - 1;
+		CommonParameters.InputPixelPosToScreenPos = (FScreenTransform::Identity + 0.5) * CommonParameters.InputInfo.ViewportSizeInverse * FScreenTransform::ViewportUVToScreenPos;
+		CommonParameters.ScreenVelocityToInputPixelVelocity = (FScreenTransform::Identity / CommonParameters.InputPixelPosToScreenPos).Scale;
+		CommonParameters.InputPixelVelocityToScreenVelocity = CommonParameters.InputPixelPosToScreenPos.Scale;
 
 		CommonParameters.LowFrequencyInfo = GetScreenPassTextureViewportParameters(FScreenPassTextureViewport(
 			LowFrequencyExtent, LowFrequencyRect));
@@ -1055,7 +1065,6 @@ void AddTemporalSuperResolutionPasses(
 		PassParameters->PrevTranslucencyPreExposureCorrection = PrevHistoryParameters.HistoryPreExposureCorrection;
 		PassParameters->TranslucencyHighlightLuminance = CVarTSRTranslucencyHighlightLuminance.GetValueOnRenderThread();
 
-		PassParameters->InputPixelPosToScreenPos = (FScreenTransform::Identity + 0.5) * CommonParameters.InputInfo.ViewportSizeInverse * FScreenTransform::ViewportUVToScreenPos;
 		PassParameters->ScreenPosToPrevTranslucencyTextureUV = FScreenTransform::ChangeTextureBasisFromTo(
 			PrevTranslucencyViewport, FScreenTransform::ETextureBasis::ScreenPosition, FScreenTransform::ETextureBasis::TextureUV);
 
@@ -1136,21 +1145,21 @@ void AddTemporalSuperResolutionPasses(
 
 			HoleFilledVelocityTexture = GraphBuilder.CreateTexture(Desc, TEXT("TSR.Velocity.HoleFilled"));
 
-			Desc.Format = PF_R8;
+			Desc.Format = PF_R8G8;
 			HoleFilledVelocityMaskTexture = GraphBuilder.CreateTexture(Desc, TEXT("TSR.velocity.HoleFillMask"));
 		}
 
 		FTSRHoleFillVelocityCS::FParameters* PassParameters = GraphBuilder.AllocParameters<FTSRHoleFillVelocityCS::FParameters>();
 		PassParameters->CommonParameters = CommonParameters;
 
-		PassParameters->MaxHollFillPixelVelocity = CVarTSRVelocityMaxHoleFillScatterVelocity.GetValueOnRenderThread();
+		PassParameters->MaxHollFillPixelVelocity = CVarTSRVelocityMaxHoleFillScatterVelocity.GetValueOnRenderThread() * ScreenPercentage;
 		PassParameters->PrevClosestDepthTexture = PrevClosestDepthTexture;
 		PassParameters->DilatedVelocityTexture = DilatedVelocityTexture;
 		PassParameters->ParallaxRejectionMaskTexture = ParallaxRejectionMaskTexture;
 
 		PassParameters->HoleFilledVelocityOutput = GraphBuilder.CreateUAV(HoleFilledVelocityTexture);
 		PassParameters->HoleFilledVelocityMaskOutput = GraphBuilder.CreateUAV(HoleFilledVelocityMaskTexture);
-		PassParameters->DebugOutput = CreateDebugUAV(LowFrequencyExtent, TEXT("Debug.TSR.HoleFillelocity"));
+		PassParameters->DebugOutput = CreateDebugUAV(LowFrequencyExtent, TEXT("Debug.TSR.HoleFillVelocity"));
 
 		TShaderMapRef<FTSRHoleFillVelocityCS> ComputeShader(View.ShaderMap);
 		FComputeShaderUtils::AddPass(
