Optimises TSR UpdateHistory's RejectionAA permutation's texture fetches

diff --git a/Engine/Shaders/Private/TemporalSuperResolution/TSRCommon.ush b/Engine/Shaders/Private/TemporalSuperResolution/TSRCommon.ush
--- a/Engine/Shaders/Private/TemporalSuperResolution/TSRCommon.ush
+++ b/Engine/Shaders/Private/TemporalSuperResolution/TSRCommon.ush
@@ -1011,21 +1011,21 @@ uint EncodeSpatialAntiAliasingOffset(taa_half2 TexelOffset)
 	return dpv_lo(EncodeSpatialAntiAliasingOffset(dpv_interleave_mono_registers(TexelOffset)));
 }
 
-taa_half2x2 DecodeSpatialAntiAliasingOffset(uint2 EncodedInputTexelOffset)
+taa_half2x2 DecodeSpatialAntiAliasingOffset(taa_ushort2 EncodedInputTexelOffset)
 {
-	const uint Mask = (1 << SPATIAL_ANTI_ALIASING_OFFSET_BITDEPTH_PER_CHANNEL) - 1;
-	const uint Zero = Mask / 2;
+	const taa_ushort Mask = taa_ushort(1 << SPATIAL_ANTI_ALIASING_OFFSET_BITDEPTH_PER_CHANNEL) - taa_ushort(1);
+	const taa_ushort Zero = Mask / taa_ushort(2);
 	
 	taa_half Multiply = taa_half(SPATIAL_ANTI_ALIASING_OFFSET_RANGE / float(Zero));
 
 	taa_half2x2 TexelOffset;
-	TexelOffset[0] = taa_half2((EncodedInputTexelOffset >> 0                                                ) & Mask) * Multiply - (taa_half(1 + Zero) * Multiply);
-	TexelOffset[1] = taa_half2((EncodedInputTexelOffset >> SPATIAL_ANTI_ALIASING_OFFSET_BITDEPTH_PER_CHANNEL) & Mask) * Multiply - (taa_half(1 + Zero) * Multiply);
+	TexelOffset[0] = taa_half2((EncodedInputTexelOffset >> taa_ushort(0                                                )) & Mask) * Multiply - (taa_half(1 + Zero) * Multiply);
+	TexelOffset[1] = taa_half2((EncodedInputTexelOffset >> taa_ushort(SPATIAL_ANTI_ALIASING_OFFSET_BITDEPTH_PER_CHANNEL)) & Mask) * Multiply - (taa_half(1 + Zero) * Multiply);
 
 	return TexelOffset;
 }
 
-taa_half2 DecodeSpatialAntiAliasingOffset(uint EncodedInputTexelOffset)
+taa_half2 DecodeSpatialAntiAliasingOffset(taa_ushort EncodedInputTexelOffset)
 {
 	return dpv_lo(DecodeSpatialAntiAliasingOffset(dpv_interleave_mono_registers(EncodedInputTexelOffset)));
 }
diff --git a/Engine/Shaders/Private/TemporalSuperResolution/TSRUpdateHistory.usf b/Engine/Shaders/Private/TemporalSuperResolution/TSRUpdateHistory.usf
--- a/Engine/Shaders/Private/TemporalSuperResolution/TSRUpdateHistory.usf
+++ b/Engine/Shaders/Private/TemporalSuperResolution/TSRUpdateHistory.usf
@@ -63,7 +63,7 @@ Texture2D<taa_half> TranslucencyRejectionTexture;
 Texture2D<float2> DilatedVelocityTexture;
 Texture2D<uint> ParallaxFactorTexture;
 Texture2D<taa_half> ParallaxRejectionMaskTexture;
-Texture2D<uint> AntiAliasingTexture;
+Texture2D<taa_ushort> AntiAliasingTexture;
 Texture2D<taa_half> NoiseFilteringTexture;
 Texture2D<taa_half2> HoleFilledVelocityMaskTexture;
 
@@ -253,16 +253,22 @@ void MainCS(
 	PLATFORM_SPECIFIC_ISOLATE
 	{
 		float2 RawEncodedVelocity[DPV_PIXEL_PER_LANE];
-		half RawParallaxRejectionMask[DPV_PIXEL_PER_LANE];
+		taa_half RawParallaxRejectionMask[DPV_PIXEL_PER_LANE];
 		uint RawParallaxFactorBits[DPV_PIXEL_PER_LANE];
-		half RawLowFrequencyRejection[DPV_PIXEL_PER_LANE];
-		half RawTranslucencyRejection[DPV_PIXEL_PER_LANE];
-		half RawTranslucencyRejectionNeighborhood[CONFIG_SAMPLES - 1][DPV_PIXEL_PER_LANE];
-		half RawMoireWeightNeighborhood[CONFIG_SAMPLES - 1][DPV_PIXEL_PER_LANE]; // TODO(TSR): bug
+		taa_half RawLowFrequencyRejection[DPV_PIXEL_PER_LANE];
+		taa_half RawTranslucencyRejection[DPV_PIXEL_PER_LANE];
+		taa_half RawTranslucencyRejectionNeighborhood[CONFIG_SAMPLES - 1][DPV_PIXEL_PER_LANE];
+		taa_half RawMoireWeightNeighborhood[CONFIG_SAMPLES - 1][DPV_PIXEL_PER_LANE]; // TODO(TSR): bug
 		
 		uint RawSceneStencilRef[DPV_PIXEL_PER_LANE];
 		uint RawSceneStencilRefNeighborhood[CONFIG_SAMPLES - 1][DPV_PIXEL_PER_LANE];
 		
+		#if CONFIG_REJECTION_ANTI_ALIASING
+			taa_ushort RawEncodedInputTexelOffset[DPV_PIXEL_PER_LANE];
+			taa_half RawNoiseFiltering[DPV_PIXEL_PER_LANE]; 
+			taa_half2 RawHoleFilledVelocityMask[DPV_PIXEL_PER_LANE];
+		#endif
+
 		// Issue overlapped texture fetches
 		PLATFORM_SPECIFIC_ISOLATE
 		{
@@ -271,13 +277,23 @@ void MainCS(
 			UNROLL_N(DPV_PIXEL_PER_LANE)
 			for (uint PixelId = 0; PixelId < DPV_PIXEL_PER_LANE; PixelId++)
 			{
-				RawEncodedVelocity[PixelId]       = DilatedVelocityTexture[dpv_access_pixel(InputPixelPos, PixelId)];
-				RawParallaxRejectionMask[PixelId] = ParallaxRejectionMaskTexture[dpv_access_pixel(InputPixelPos, PixelId)];
-				RawParallaxFactorBits[PixelId]    = ParallaxFactorTexture[dpv_access_pixel(InputPixelPos, PixelId)];
-				RawTranslucencyRejection[PixelId] = TranslucencyRejectionTexture[dpv_access_pixel(InputPixelPos, PixelId)];
-				RawSceneStencilRef[PixelId]       = InputSceneStencilTexture.Load(int3(dpv_access_pixel(InputPixelPos, PixelId), 0)) STENCIL_COMPONENT_SWIZZLE;
+				taa_ushort2 LocalInputPixelPos = dpv_access_pixel(InputPixelPos, PixelId);
+
+				RawEncodedVelocity[PixelId]       = DilatedVelocityTexture[LocalInputPixelPos];
+				RawParallaxRejectionMask[PixelId] = ParallaxRejectionMaskTexture[LocalInputPixelPos];
+				RawParallaxFactorBits[PixelId]    = ParallaxFactorTexture[LocalInputPixelPos];
+				RawTranslucencyRejection[PixelId] = TranslucencyRejectionTexture[LocalInputPixelPos];
+				RawSceneStencilRef[PixelId]       = InputSceneStencilTexture.Load(int3(LocalInputPixelPos, 0)) STENCIL_COMPONENT_SWIZZLE;
 				
 				RawLowFrequencyRejection[PixelId] = HistoryRejectionTexture[dpv_access_pixel(RejectionPixelPos, PixelId)];
+				
+				#if CONFIG_REJECTION_ANTI_ALIASING
+				{
+					RawEncodedInputTexelOffset[PixelId] = AntiAliasingTexture[LocalInputPixelPos];
+					RawNoiseFiltering[PixelId] = NoiseFilteringTexture[LocalInputPixelPos];
+					RawHoleFilledVelocityMask[PixelId] = HoleFilledVelocityMaskTexture[LocalInputPixelPos];
+				}
+				#endif
 			}
 
 			// Vector in pixel between pixel K -> O.
@@ -342,7 +358,6 @@ void MainCS(
 			}
 		}
 
-		// TODO(TSR): Integrate in the manual overlapped texture fetches.
 		#if CONFIG_REJECTION_ANTI_ALIASING
 		{
 			// Detect if HistoryBufferUV would be outside of the viewport.
@@ -354,17 +369,10 @@ void MainCS(
 				bOffScreen[0] ? 0.0 : saturate(LowFrequencyRejection[0] * 4.0),
 				bOffScreen[1] ? 0.0 : saturate(LowFrequencyRejection[1] * 4.0));
 
-			uint2 EncodedInputTexelOffset = dpv_interleave_registers(
-				AntiAliasingTexture[dpv_lo(InputPixelPos)],
-				AntiAliasingTexture[dpv_hi(InputPixelPos)]);
-			
-			NoiseFiltering = dpv_interleave_registers(
-				NoiseFilteringTexture[dpv_lo(InputPixelPos)],
-				NoiseFilteringTexture[dpv_hi(InputPixelPos)]);
+			taa_ushort2 EncodedInputTexelOffset = dpv_force_interleave_registers_array(RawEncodedInputTexelOffset);
 			
-			HoleFilledVelocityMask = dpv_interleave_mono_registers(taa_half(1.0).xx) - dpv_interleave_registers(
-				HoleFilledVelocityMaskTexture[dpv_lo(InputPixelPos)],
-				HoleFilledVelocityMaskTexture[dpv_hi(InputPixelPos)]);
+			NoiseFiltering = dpv_force_interleave_registers_array(RawNoiseFiltering);
+			HoleFilledVelocityMask = dpv_interleave_mono_registers(taa_half(1.0).xx) - dpv_force_interleave_registers_array(RawHoleFilledVelocityMask);
 
 			taa_half2x2 TexelOffset = DecodeSpatialAntiAliasingOffset(EncodedInputTexelOffset);
 			
