Add some wiggle room to backfacing test for shadow filtering to account for error/shading normals.

diff --git a/Engine/Shaders/Private/VirtualShadowMaps/ProjectionDirectional.ush b/Engine/Shaders/Private/VirtualShadowMaps/ProjectionDirectional.ush
--- a/Engine/Shaders/Private/VirtualShadowMaps/ProjectionDirectional.ush
+++ b/Engine/Shaders/Private/VirtualShadowMaps/ProjectionDirectional.ush
@@ -191,7 +191,9 @@ float3 GetRandomDirectionalLightRayDir(FLightShaderParameters Light, float2 E)
 // Normal and LightDirection should be normalized
 bool IsBackfaceToDirectionalLight(half3 Normal, half3 LightDirection, half LightSourceRadius)
 {
-	half SinAlpha = abs(LightSourceRadius);
+	// Allow a minimum of ~5 degrees of wiggle room to account for normal issues
+	half MinSinAlpha = 0.1;
+	half SinAlpha = max(abs(LightSourceRadius), MinSinAlpha);
 	return dot(Normal, LightDirection) < -SinAlpha;
 }
 
