Implements dynamic half-res motion blur

1) Motionblur have tile classification to improve VGPR pressure on fast motion blur gathering tiles
2) Motionblur tile classification can lower the gathering resolution to half res with r.MotionBlur.HalfResGather=1 on part of the screen with large motion vectors to reduce VALU costs under large movements, but still can compose with full res scene color without round trip to main memory to avoid adding memory bandwidth costs
3) TSR output a half res scene color in Mip1 so motion blur can dynamically sample half res under largement movements with r.MotionBlur.HalfResInput=1 to save memory bandwidth in the motion blur convolution
4) Number of direction samples becomes dynamic (and clamped to exisiting motion blur quality) based on length of motion vector to saves VALU with slow pixel velocities
5) Motion blur automatically avoid sampling full res center if the blur kernel completly occlude its to completly stop fetching full res to save memory bandwidth cost under large motions

Drops motion blur costs from 753us to 339us of motion blur costs + 40us additional cost in TSR to output half res in sandbox under movements.

And saves an average 0.3ms in the entire chase gameplay sequence.

diff --git a/Engine/Shaders/Private/MotionBlur/MotionBlurApply.usf b/Engine/Shaders/Private/MotionBlur/MotionBlurApply.usf
--- a/Engine/Shaders/Private/MotionBlur/MotionBlurApply.usf
+++ b/Engine/Shaders/Private/MotionBlur/MotionBlurApply.usf
@@ -7,50 +7,90 @@
 #include "../Random.ush"
 
 
-//------------------------------------------------------- CONSTANTS
+//------------------------------------------------------- DEBUG
+
+// Debug the input scene color resolution.
+#define DEBUG_OVERLAY_INPUT_RES 0
+
+// Debug the gathered resolution.
+#define DEBUG_OVERLAY_GATHER_RES 0
+
+// Debug the tile classification
+#define DEBUG_OVERLAY_TILE_CLASSIFICATION 0
+
+// Debug whether the full res center is fecthed or not.
+#define DEBUG_OVERLAY_SKIP_CENTER 0
+
+// Debug the number of samples
+//	1: show number of samples
+//  2: show the density of samples (lower at half res)
+#define DEBUG_OVERLAY_SAMPLES 0
+
+
+//------------------------------------------------------- CONFIG
+
+#if DIM_TILE_CLASSIFICATION == FILTER_TILE_CLASSIFY_GATHER_HALF_RES
+	#define CONFIG_IS_HALF_RES 1
+
+#elif DIM_TILE_CLASSIFICATION == FILTER_TILE_CLASSIFY_GATHER_FULL_RES
+	#define CONFIG_IS_HALF_RES 0
+
+#elif DIM_TILE_CLASSIFICATION == FILTER_TILE_CLASSIFY_SCATTER_AS_GATHER_1_VELOCITY_HALF_RES
+	#define CONFIG_IS_HALF_RES 1
+
+#elif DIM_TILE_CLASSIFICATION == FILTER_TILE_CLASSIFY_SCATTER_AS_GATHER_1_VELOCITY_FULL_RES
+	#define CONFIG_IS_HALF_RES 0
+
+#elif DIM_TILE_CLASSIFICATION == FILTER_TILE_CLASSIFY_SCATTER_AS_GATHER_2_VELOCITY_FULL_RES
+	#define CONFIG_IS_HALF_RES 0
+
+#else
+	#error unknown tile classification
+#endif
 
 // from the paper: We use SOFT Z EXTENT = 1mm to 10cm for our results
 #define SOFT_Z_EXTENT 1
 
-#define MOTION_BLUR_QUALITY_LOW 0
-#define MOTION_BLUR_QUALITY_MEDIUM 1
-#define MOTION_BLUR_QUALITY_HIGH 2
-#define MOTION_BLUR_QUALITY_VERY_HIGH 3
+// Whether should interleave the half res motion blur with the full res to reduce the size of the grain on screen.
+#define CONFIG_SHUFFLE_HALF_RES 1
+
+// Whether should do post motion blur translucency.
+#define CONFIG_POST_MOTIONBLUR_TRANSLUCENCY 1
+
+// Save memory bandwidth by not fetching full res center if output is fully gathered.
+#if COMPILER_SUPPORTS_WAVE_MINMAX
+	#define CONFIG_SKIP_CENTER 1
+#else
+	#define CONFIG_SKIP_CENTER 0
+#endif
 
 
-#define CONFIG_HOLE_FILL_SEPARATLY 1
+//------------------------------------------------------- CONSTANTS
 
-#if MOTION_BLUR_QUALITY == MOTION_BLUR_QUALITY_LOW
-	#define CONFIG_SAMPLE_COUNT 4
-#elif MOTION_BLUR_QUALITY == MOTION_BLUR_QUALITY_MEDIUM
-	#define CONFIG_SAMPLE_COUNT 6
-#elif MOTION_BLUR_QUALITY == MOTION_BLUR_QUALITY_HIGH
-	#define CONFIG_SAMPLE_COUNT 8
-#elif MOTION_BLUR_QUALITY == MOTION_BLUR_QUALITY_VERY_HIGH
-	#define CONFIG_SAMPLE_COUNT 16
+#if CONFIG_IS_HALF_RES
+	#define TILE_SIZE (VELOCITY_FILTER_TILE_SIZE / 2)
 #else
-	#error Unknown motion blur quality
+	#define TILE_SIZE (VELOCITY_FILTER_TILE_SIZE)
 #endif
 
 
 //------------------------------------------------------- PARAMETERS
 
-float MotionBlur_AspectRatio;
-float MotionBlur_VelocityScale;
-float MotionBlur_VelocityScaleForTiles;
-float MotionBlur_VelocityMax;
-
 SCREEN_PASS_TEXTURE_VIEWPORT(Color)
 SCREEN_PASS_TEXTURE_VIEWPORT(Velocity)
 SCREEN_PASS_TEXTURE_VIEWPORT(VelocityTile)
 FScreenTransform ColorToVelocity;
+uint MaxSampleCount;
 
-Texture2D ColorTexture;
-Texture2D VelocityTexture;
+uint TileListOffset;
+StructuredBuffer<uint> TileListsBuffer;
+StructuredBuffer<uint> TileListsSizeBuffer;
+
+Texture2D<mb_half4> ColorTexture;
+Texture2D<mb_half3> VelocityFlatTexture;
 Texture2D VelocityTileTextures_Textures_0;
 Texture2D VelocityTileTextures_Textures_1;
-Texture2D VelocityFlatTexture;
-Texture2D DepthTexture;
+Texture2D<mb_half4> HalfResMotionBlurTexture;
 
 SamplerState ColorSampler;
 SamplerState VelocitySampler;
@@ -58,15 +98,12 @@ SamplerState VelocityTileSampler;
 SamplerState VelocityFlatSampler;
 SamplerState DepthSampler;
 
-#if USE_POST_MOTION_BLUR_TRANSLUCENCY
-
-Texture2D PostMotionBlurTranslucencyTexture;
-SamplerState PostMotionBlurTranslucencySampler;
-
-float2 ScreenPosToPostMotionBlurTranslucencyUV;
-float2 PostMotionBlurTranslucencyUVMax;
-
-#endif // USE_POST_MOTION_BLUR_TRANSLUCENCY
+Texture2D TranslucencyTexture;
+SamplerState TranslucencySampler;
+FScreenTransform ColorToTranslucency;
+float2 TranslucencyUVMin;
+float2 TranslucencyUVMax;
+float2 TranslucencyExtentInverse;
 
 #if SUPPORTS_INDEPENDENT_SAMPLERS
 	#define SharedVelocitySampler VelocitySampler
@@ -80,7 +117,7 @@ float2 PostMotionBlurTranslucencyUVMax;
 	#define SharedDepthSampler DepthSampler
 #endif
 
-RWTexture2D<float4> OutColorTexture;
+RWTexture2D<mb_half4> OutColorTexture;
 RWTexture2D<float4> DebugOutput;
 
 
@@ -88,84 +125,99 @@ RWTexture2D<float4> DebugOutput;
 
 #if !COMPILER_SUPPORTS_WAVE_VOTE
 
-groupshared uint GroupSharedFastPath;
+groupshared uint SharedSampleCount;
+groupshared uint SharedFastPath;
 
 #endif
 
 
 //------------------------------------------------------- FUNCTIONS
 
-float NormalizeAccumulator(float Accumulator, float AccumulatorWeight)
+float NormalizeAccumulator(float Accumulator, mb_half AccumulatorWeight)
 {
-	return -min(-Accumulator * rcp(AccumulatorWeight), 0.0);
+	return -min(-Accumulator * rcp(AccumulatorWeight), mb_half(0.0));
 }
 
 float4 NormalizeAccumulator(float4 Accumulator, float AccumulatorWeight)
 {
-	return -min(-Accumulator * rcp(AccumulatorWeight), 0.0);
+	return -min(-Accumulator * rcp(AccumulatorWeight), float(0.0));
+}
+
+#if CONFIG_COMPILE_FP16
+
+mb_half NormalizeAccumulator(mb_half Accumulator, mb_half AccumulatorWeight)
+{
+	return -min(-Accumulator * rcp(AccumulatorWeight), mb_half(0.0));
+}
+
+mb_half4 NormalizeAccumulator(mb_half4 Accumulator, mb_half AccumulatorWeight)
+{
+	return -min(-Accumulator * rcp(AccumulatorWeight), mb_half(0.0));
 }
 
+#endif // CONFIG_COMPILE_FP16
+
 void NormalizeAccumulatorWithHoleFill(
-	float4 Color,
-	float ColorWeight,
-	float HoleFillWeight,
-	float DepthAware,
-	float4 HoleFillingColor,
-	float InvSampleCount,
-	out float4 ColorOutput,
-	out float OpacityOutput,
+	mb_half4 Color,
+	mb_half ColorWeight,
+	mb_half HoleFillWeight,
+	mb_half DepthAware,
+	mb_half4 HoleFillingColor,
+	mb_half InvSampleCount,
+	out mb_half4 ColorOutput,
+	out mb_half OpacityOutput,
 	out bool bValidColorOutput)
 {
-	OpacityOutput = min(saturate(ColorWeight * InvSampleCount * 2), (DepthAware * InvSampleCount * 2));
+	OpacityOutput = min(saturate(ColorWeight * InvSampleCount * mb_half(2)), (DepthAware * InvSampleCount * mb_half(2)));
 	
-	float4 FinalAccumulator = Color + HoleFillWeight * HoleFillingColor;
+	mb_half4 FinalAccumulator = Color + HoleFillWeight * HoleFillingColor;
 	
 	ColorOutput = NormalizeAccumulator(FinalAccumulator, ColorWeight + HoleFillWeight);
-	bValidColorOutput = (ColorWeight + HoleFillWeight) > 0;
+	bValidColorOutput = (ColorWeight + HoleFillWeight) > mb_half(0.0);
 }
 
 
 // Compute the weight of the sample for hole filling.
-float ComputeSampleHoleFillWeight(float CenterDepth, float SampleDepth, float DepthScale)
+mb_half ComputeSampleHoleFillWeight(mb_half CenterDepth, mb_half SampleDepth, mb_half DepthScale)
 {
 	return saturate(DepthScale * (SampleDepth - CenterDepth));
 }
 
 // Computes the contribution weight of one sample.
-float ComputeSampleConvolutionWeight(
-	float SampleDepth, float SampleSpreadLength, float SampleVelocityAngle,
-	float OffsetLength, float BlurAngle, float PixelToSampleScale)
+mb_half ComputeSampleConvolutionWeight(
+	mb_half SampleDepth, mb_half SampleSpreadLength, mb_half SampleVelocityAngle,
+	mb_half OffsetLength, mb_half BlurAngle, mb_half PixelToSampleScale)
 {
 	// Compare the length 
-	float SpreadWeights = saturate(PixelToSampleScale * SampleSpreadLength - max(OffsetLength - 1, 0));
+	mb_half SpreadWeights = saturate(PixelToSampleScale * SampleSpreadLength - max(OffsetLength - mb_half(1), mb_half(0)));
 	
 	// Compare the directions
 	#if CONFIG_MAX_RANGE_SIZE > 1
-		float DirectionWeights = saturate(1.0 - max(GetPolarRelativeAngle(SampleVelocityAngle, BlurAngle) - 0.1, 0.0) * 4.0);
+		mb_half DirectionWeights = saturate(mb_half(1.0) - max(GetPolarRelativeAngle(SampleVelocityAngle, BlurAngle) - mb_half(0.1), mb_half(0.0)) * mb_half(4.0));
 	#else
-		float DirectionWeights = 1.0;
+		mb_half DirectionWeights = mb_half(1.0);
 	#endif
 
 	return SpreadWeights * DirectionWeights;
 }
 
 // Selectively computes the contribution weight of the center or the sample based on whether sample is respectively behind center or not
-float ComputeCenterOrSampleWeight(
-	float CenterDepth, float CenterSpreadLength, float CenterVelocityAngle,
-	float SampleDepth, float SampleSpreadLength, float SampleVelocityAngle,
-	float OffsetLength, float BlurAngle, float PixelToSampleScale, float DepthScale)
+mb_half ComputeCenterOrSampleWeight(
+	mb_half CenterDepth, mb_half CenterSpreadLength, mb_half CenterVelocityAngle,
+	mb_half SampleDepth, mb_half SampleSpreadLength, mb_half SampleVelocityAngle,
+	mb_half OffsetLength, mb_half BlurAngle, mb_half PixelToSampleScale, mb_half DepthScale)
 {
 	// Compute weight to use the center data if center is closer than the sample.
-	float CenterWeight = saturate(0.5 + DepthScale * (SampleDepth - CenterDepth));
+	mb_half CenterWeight = saturate(0.5 + DepthScale * (SampleDepth - CenterDepth));
 	
 	// Compute weight to use the sample data if sample is closer than the center.
-	float SampleWeight = saturate(0.5 - DepthScale * (SampleDepth - CenterDepth));
+	mb_half SampleWeight = saturate(0.5 - DepthScale * (SampleDepth - CenterDepth));
 
-	float CenterConvolutionWeight = ComputeSampleConvolutionWeight(
+	mb_half CenterConvolutionWeight = ComputeSampleConvolutionWeight(
 		CenterDepth, CenterSpreadLength, CenterVelocityAngle,
 		OffsetLength, BlurAngle, PixelToSampleScale);
 	
-	float SampleConvolutionWeight = ComputeSampleConvolutionWeight(
+	mb_half SampleConvolutionWeight = ComputeSampleConvolutionWeight(
 		SampleDepth, SampleSpreadLength, SampleVelocityAngle,
 		OffsetLength, BlurAngle, PixelToSampleScale);
 
@@ -173,30 +225,49 @@ float ComputeCenterOrSampleWeight(
 }
 
 // TODO: move that to velocity flatten.
-float GetVelocityLengthPixels(float2 EncodedVelocity)
+mb_half GetVelocityLengthPixels(mb_half2 EncodedVelocity)
 {
 	// 11:11:10  (VelocityLength, VelocityAngle, Depth)
-	float VelocityLength = EncodedVelocity.x;
-	VelocityLength *= MotionBlur_VelocityScale;
-	return min(VelocityLength, MotionBlur_VelocityMax);
+	return EncodedVelocity.x;
 }
 
 
 //------------------------------------------------------- ENTRY POINT
 
-[numthreads(THREADGROUP_SIZEX, THREADGROUP_SIZEY, 1)]
-void MainCS(uint2 DispatchThreadId : SV_DispatchThreadID)
+[numthreads(TILE_SIZE, TILE_SIZE, 1)]
+void MainCS(
+	uint GroupId : SV_GroupID,
+	uint GroupThreadIndex : SV_GroupIndex)
 {
 	float4 Debug = 0;
+	mb_half4 DebugOverlay = mb_half(1.0).xxxx;
+	
+	uint PackedGroupOffset = TileListsBuffer[GroupId + TileListOffset];
+	uint2 GroupOffset = uint2(PackedGroupOffset, PackedGroupOffset >> 16) & 0xFFFF;
+	
+	const bool bIsHalfRes = CONFIG_IS_HALF_RES ? true : false;
+
+	uint2 DispatchThreadId = (
+		ZOrder2D(GroupThreadIndex, log2(TILE_SIZE)) +
+		GroupOffset * uint2(TILE_SIZE, TILE_SIZE));
 
 	uint2 iColorPixelPos = DispatchThreadId + Color_ViewportMin;
-	float2 ColorPixelPos = float2(iColorPixelPos);
-	float2 ColorUV = (ColorPixelPos + 0.5) * Color_ExtentInverse;
 
-	const float PixelToTileScale = (1.0 / 16.0);
+	float2 ColorUV;
+	BRANCH
+	if (bIsHalfRes)
+	{
+		ColorUV = (float2(DispatchThreadId * 2 + Color_ViewportMin) + 1.0) * Color_ExtentInverse;
+	}
+	else
+	{
+		ColorUV = (float2(iColorPixelPos) + 0.5) * Color_ExtentInverse;
+	}
 
-	float Random  = InterleavedGradientNoise(ColorPixelPos, 0);
-	float Random2 = InterleavedGradientNoise(ColorPixelPos, 1);
+	const float PixelToTileScale = rcp(float(VELOCITY_FLATTEN_TILE_SIZE));
+
+	float Random  = InterleavedGradientNoise(iColorPixelPos, 0);
+	float Random2 = InterleavedGradientNoise(iColorPixelPos, 1);
 
 	// [-0.25, 0.25]
 	float2 TileJitter = (float2(Random, Random2) - 0.5) * 0.5;
@@ -215,124 +286,178 @@ void MainCS(uint2 DispatchThreadId : SV_DispatchThreadID)
 		VelocityTileTextures_Textures_0.SampleLevel(SharedVelocityTileSampler, VelocityTileUV, 0),
 		VelocityTileTextures_Textures_1.SampleLevel(SharedVelocityTileSampler, VelocityTileUV, 0));
 
-	float2 MinVelocityPixels = VelocityRange.Min * MotionBlur_VelocityScale;
-	float2 Max0VelocityPixels = VelocityRange.Max[0] * MotionBlur_VelocityScale;
+	float2 MinVelocityPixels = VelocityRange.Min;
+	float2 Max0VelocityPixels = VelocityRange.Max[0];
+	#if CONFIG_MAX_RANGE_SIZE > 1
+		float2 Max1VelocityPixels = VelocityRange.Max[1];
+	#endif
 
-	#if CONFIG_MAX_RANGE_SIZE >= 2
-		float2 Max1VelocityPixels = VelocityRange.Max[1] * MotionBlur_VelocityScale;
+	// Compute how many sample should be applied.
+	float MipLevel0;
+	float MipLevel1;
+	uint SampleCount;
+	float InvSampleCount;
+	bool bSkip;
+	bool bFastPath;
+	#if CONFIG_MAX_RANGE_SIZE > 1
+		bool bDoOneDirectionOnly;
+	#else
+		const bool bDoOneDirectionOnly = true;
 	#endif
+	{
+		const uint SampleCountFactor = 4;
 
-	float MinVelocityLengthSqrPixels = dot(MinVelocityPixels, MinVelocityPixels);
-	float Max0VelocityLengthSqrPixels = dot(Max0VelocityPixels, Max0VelocityPixels);
+		float MaxPixelLength0 = length(Max0VelocityPixels);
 
-	// Input buffer 0 as same viewport as output buffer.
-	float4 CenterColor = ColorTexture.SampleLevel(ColorSampler, ColorUV, 0);
+		uint RecommendedSampleCount = clamp(SampleCountFactor * ceil(MaxPixelLength0 * rcp(float(SampleCountFactor))), SampleCountFactor, MaxSampleCount);
+		
+		bFastPath = length2(MinVelocityPixels) > 0.4 * (MaxPixelLength0 * MaxPixelLength0);
 
-	float2 SearchVector0 = Max0VelocityPixels * Color_ExtentInverse.xyxy;
+		#if COMPILER_SUPPORTS_WAVE_MINMAX
+		{
+			SampleCount = ToScalarMemory(WaveActiveMax(RecommendedSampleCount));
+			bFastPath = WaveActiveAllTrue(bFastPath);
+		}
+		#else
+		{
+			if (GroupThreadIndex == 0)
+			{
+				SharedSampleCount = 0;
+				SharedFastPath = 0;
+			}
+			GroupMemoryBarrierWithGroupSync();
 
-	#if USE_POST_MOTION_BLUR_TRANSLUCENCY
-		const float2 PostMotionBlurTranslucencyUV = min((ColorPixelPos + (float2)0.5f) * ScreenPosToPostMotionBlurTranslucencyUV, PostMotionBlurTranslucencyUVMax);
-		float4 PostMotionBlurTranslucency = PostMotionBlurTranslucencyTexture.SampleLevel(PostMotionBlurTranslucencySampler, PostMotionBlurTranslucencyUV, 0);
-	#endif
+			InterlockedMax(SharedSampleCount, RecommendedSampleCount);
+			InterlockedAdd(SharedFastPath, bFastPath ? 1 : 0);
 
+			GroupMemoryBarrierWithGroupSync();
+			SampleCount = ToScalarMemory(SharedSampleCount);
+			bFastPath = (SharedFastPath == (TILE_SIZE * TILE_SIZE));
+		}
+		#endif
 	
-	// TODO expose cvars
-	bool bSkipPath = Max0VelocityLengthSqrPixels < 0.25;
-	bool bFastPath = MinVelocityLengthSqrPixels > 0.4 * Max0VelocityLengthSqrPixels;
-	
-	// Only use fast path if all threads of the compute shader would.
-	#if COMPILER_SUPPORTS_WAVE_VOTE
-	{
-		bFastPath = WaveActiveAllTrue(bFastPath);
-	}
-	#else
-	{
-		GroupSharedFastPath = 0;
-		GroupMemoryBarrierWithGroupSync();
+		InvSampleCount = ToScalarMemory(rcp(float(SampleCount)));
+		bSkip = MaxPixelLength0 < CONFIG_MINIMAL_PIXEL_VELOCITY;
 
-		uint IgnoredOut;
-		InterlockedAdd(GroupSharedFastPath, bFastPath ? 1 : 0);
+		#if CONFIG_MAX_RANGE_SIZE > 1
+		{
+			float MaxPixelLength1 = length(Max1VelocityPixels);
+			bDoOneDirectionOnly = MaxPixelLength1 < CONFIG_MINIMAL_PIXEL_VELOCITY;
+			if (bDoOneDirectionOnly)
+			{
+				MipLevel0 = saturate(MaxPixelLength0 * InvSampleCount - 1.0);
+				MipLevel1 = MipLevel0;
+				Max1VelocityPixels = Max0VelocityPixels;
+			}
+			else
+			{
+				MipLevel0 = saturate(MaxPixelLength0 * (InvSampleCount * 2.0) - 1.0);
+				MipLevel1 = saturate(MaxPixelLength1 * (InvSampleCount * 2.0) - 1.0);
+			}
+		}
+		#else
+		{
+			MipLevel0 = saturate(MaxPixelLength0 * InvSampleCount - 1.0);
+			MipLevel1 = MipLevel0;
+		}
+		#endif
 		
-		GroupMemoryBarrierWithGroupSync();
-		bFastPath = (GroupSharedFastPath == (THREADGROUP_SIZEX * THREADGROUP_SIZEY));
+		#if DEBUG_OVERLAY_SAMPLES == 1
+			DebugOverlay = lerp(mb_half4(0.5, 1.0, 0.5, 1.0), mb_half4(1.0, 0.5, 0.5, 1.0), float(SampleCount) / float(MaxSampleCount));
+		#elif DEBUG_OVERLAY_SAMPLES == 2
+			DebugOverlay = lerp(mb_half4(0.5, 1.0, 0.5, 1.0), mb_half4(1.0, 0.5, 0.5, 1.0), float(SampleCount) / float(MaxSampleCount * (bIsHalfRes ? 4 : 1)));
+		#endif
 	}
-	#endif
-	
-	float4 OutColor;
 
+	mb_half2 SearchVector0 = mb_half2(Max0VelocityPixels * Color_ExtentInverse.xy);
+
+	mb_half4 MotionBlurColor;
+	mb_half FullResBlend;
+	
 	BRANCH
-	if (bSkipPath)
+	if (bSkip)
 	{
-		OutColor = CenterColor;
+		MotionBlurColor = mb_half(0.0).xxxx;
+		FullResBlend = 1.0;
+		
+		#if DEBUG_OVERLAY_TILE_CLASSIFICATION
+			DebugOverlay = mb_half4(0.5, 0.5, 1.0, 1.0);
+		#endif
+
+		#if DEBUG_OVERLAY_SAMPLES
+			DebugOverlay = mb_half4(0.5, 1.0, 0.5, 1.0);
+		#endif
 	}
+	#if DIM_TILE_CLASSIFICATION == FILTER_TILE_CLASSIFY_GATHER_HALF_RES || DIM_TILE_CLASSIFICATION == FILTER_TILE_CLASSIFY_GATHER_FULL_RES
+	else
+	#else
 	else if (bFastPath)
+	#endif
 	{
-		const float InvSampleCount = rcp(float(CONFIG_SAMPLE_COUNT));
-
-
-		float4 ColorAccum = 0;
+		mb_half4 ColorAccum = 0;
 		
-		UNROLL_N(CONFIG_SAMPLE_COUNT / 2)
-		for (uint i = 0; i < (CONFIG_SAMPLE_COUNT / 2); i++)
+		LOOP
+		for (uint i = 0; i < SampleCount; i += 4)
 		{
-			float2 OffsetLength = (float)i + 0.5 + float2(Random - 0.5, 0.5 - Random);
-			float2 OffsetFraction = OffsetLength * (2.0 * InvSampleCount);
+			UNROLL_N(2)
+			for (uint j = 0; j < 2; j ++)
+			{
+				float2 OffsetLength = float(i / 2 + j) + (0.5 + float2(Random - 0.5, 0.5 - Random));
+				float2 OffsetFraction = OffsetLength * (2.0 * InvSampleCount);
 
-			float2 SampleUV[2];
-			SampleUV[0] = ColorUV + OffsetFraction.x * SearchVector0;
-			SampleUV[1] = ColorUV - OffsetFraction.y * SearchVector0;
+				float2 SampleUV[2];
+				SampleUV[0] = ColorUV + OffsetFraction.x * SearchVector0;
+				SampleUV[1] = ColorUV - OffsetFraction.y * SearchVector0;
 
-			SampleUV[0] = clamp(SampleUV[0], Color_UVViewportBilinearMin, Color_UVViewportBilinearMax);
-			SampleUV[1] = clamp(SampleUV[1], Color_UVViewportBilinearMin, Color_UVViewportBilinearMax);
+				SampleUV[0] = clamp(SampleUV[0], Color_UVViewportBilinearMin, Color_UVViewportBilinearMax);
+				SampleUV[1] = clamp(SampleUV[1], Color_UVViewportBilinearMin, Color_UVViewportBilinearMax);
 
-			ColorAccum += ColorTexture.SampleLevel(ColorSampler, SampleUV[0], 0);
-			ColorAccum += ColorTexture.SampleLevel(ColorSampler, SampleUV[1], 0);
+				ColorAccum += ColorTexture.SampleLevel(ColorSampler, SampleUV[0], MipLevel0);
+				ColorAccum += ColorTexture.SampleLevel(ColorSampler, SampleUV[1], MipLevel0);
+			}
 		}
-	
-		ColorAccum *= InvSampleCount;
 
-		OutColor = ColorAccum;
+		MotionBlurColor = ColorAccum * InvSampleCount;
+		FullResBlend = 0.0;
+		
+		#if DEBUG_OVERLAY_TILE_CLASSIFICATION
+			DebugOverlay = mb_half4(0.5, 1.0, 0.5, 1.0);
+		#endif
+			
+		#if DEBUG_OVERLAY_INPUT_RES
+		{
+			mb_half4 DebugMipLevel = lerp(mb_half4(1.0, 0.5, 0.5, 1.0), mb_half4(0.5, 1.0, 0.5, 1.0), mb_half(MipLevel0));
+			MotionBlurColor *= DebugMipLevel;
+		}
+		#endif
 	}
+	#if DIM_TILE_CLASSIFICATION == FILTER_TILE_CLASSIFY_SCATTER_AS_GATHER_1_VELOCITY_HALF_RES || DIM_TILE_CLASSIFICATION == FILTER_TILE_CLASSIFY_SCATTER_AS_GATHER_1_VELOCITY_FULL_RES || DIM_TILE_CLASSIFICATION == FILTER_TILE_CLASSIFY_SCATTER_AS_GATHER_2_VELOCITY_FULL_RES
 	else
 	{
-		const uint TotalSteps = CONFIG_SAMPLE_COUNT / 2;
+		const uint TotalSteps = SampleCount / 2;
 		const uint DirectionCount = CONFIG_MAX_RANGE_SIZE;
 		const uint StepPerDirectionCount = TotalSteps / DirectionCount;
-		const uint SamplePerDirection = CONFIG_SAMPLE_COUNT / DirectionCount;
-
-		const uint StepCount = CONFIG_SAMPLE_COUNT / 2;
+		const uint SamplePerDirection = SampleCount / DirectionCount;
 
-		// Decide whether should only accumulate along one direction.
-		#if CONFIG_MAX_RANGE_SIZE > 1
-			float Max1VelocityLengthSqrPixels = dot(Max1VelocityPixels, Max1VelocityPixels);
-			bool bDoOneDirectionOnly = Max1VelocityLengthSqrPixels < 0.25;
-			if (bDoOneDirectionOnly)
-			{
-				Max1VelocityPixels = Max0VelocityPixels;
-			}
-		#else
-			const bool bDoOneDirectionOnly = true;
-		#endif
-		
-		float3 CenterVelocityDepth = VelocityFlatTexture.SampleLevel(SharedVelocityFlatSampler, VelocityUV, 0).xyz;
-		float  CenterDepth = CenterVelocityDepth.z;
-		float  CenterVelocityLength = GetVelocityLengthPixels(CenterVelocityDepth.xy);
-		float  CenterVelocityAngle = CenterVelocityDepth.y * (2.0 * PI) - PI;
+		mb_half3 CenterVelocityDepth = VelocityFlatTexture.SampleLevel(SharedVelocityFlatSampler, VelocityUV, 0).xyz;
+		mb_half  CenterDepth = CenterVelocityDepth.z;
+		mb_half  CenterVelocityLength = GetVelocityLengthPixels(CenterVelocityDepth.xy);
+		mb_half  CenterVelocityAngle = CenterVelocityDepth.y * (2.0 * PI) - PI;
 		
 		#if CONFIG_MAX_RANGE_SIZE > 1
-			float4 HoleFillColor = 0;
-			float  HoleFillColorWeight = 0;
+			mb_half4 HoleFillColor = 0;
+			mb_half  HoleFillColorWeight = 0;
 		#endif
 
-		float4 DirectionalColor[CONFIG_MAX_RANGE_SIZE];
-		float  DirectionalColorWeight[CONFIG_MAX_RANGE_SIZE];
+		mb_half4 DirectionalColor[CONFIG_MAX_RANGE_SIZE];
+		mb_half  DirectionalColorWeight[CONFIG_MAX_RANGE_SIZE];
 		
 		#if CONFIG_MAX_RANGE_SIZE > 1
-			float  HoleFillWeightAccum[CONFIG_MAX_RANGE_SIZE];
+			mb_half HoleFillWeightAccum[CONFIG_MAX_RANGE_SIZE];
 
-			float  DepthAccum[CONFIG_MAX_RANGE_SIZE];
-			float  DepthSquareAccum[CONFIG_MAX_RANGE_SIZE];
-			float  DepthAwareWeight[CONFIG_MAX_RANGE_SIZE];
+			mb_half DepthAccum[CONFIG_MAX_RANGE_SIZE];
+			mb_half DepthSquareAccum[CONFIG_MAX_RANGE_SIZE];
+			mb_half DepthAwareWeight[CONFIG_MAX_RANGE_SIZE];
 		#endif
 
 		// Iterate over the different directions.
@@ -342,6 +467,7 @@ void MainCS(uint2 DispatchThreadId : SV_DispatchThreadID)
 			float PixelToSampleScale = TotalSteps * rsqrt(dot(Max0VelocityPixels, Max0VelocityPixels));	
 			float2 SearchVector = SearchVector0;
 			float BlurAngle = CartesianToPolar(Max0VelocityPixels).y;
+			float MipLevel = MipLevel0;
 
 			bool bAccumulateHoleFillColor = true;
 
@@ -352,6 +478,7 @@ void MainCS(uint2 DispatchThreadId : SV_DispatchThreadID)
 				SearchVector = Max1VelocityPixels * Color_ExtentInverse.xy;
 				BlurAngle = CartesianToPolar(Max1VelocityPixels).y;
 				bAccumulateHoleFillColor = HoleFillColorWeight == 0.0;
+				MipLevel = MipLevel1;
 			}
 			#endif
 		
@@ -369,20 +496,20 @@ void MainCS(uint2 DispatchThreadId : SV_DispatchThreadID)
 			#endif
 			
 			// Iterate over steps of 2 samples in each directions.
-			UNROLL
+			LOOP
 			for (uint StepId = DirectionId; StepId < TotalSteps; StepId += CONFIG_MAX_RANGE_SIZE)
 			{
 				float2 SampleUV[2];
-				float4 SampleColor[2];
-				float  SampleDepth[2];
-				float  SampleVelocityLength[2];
-				float  ConvolutionWeight[2];
-				float  HoleFillingWeight[2];
+				mb_half4 SampleColor[2];
+				mb_half SampleDepth[2];
+				mb_half SampleVelocityLength[2];
+				mb_half ConvolutionWeight[2];
+				mb_half HoleFillingWeight[2];
 
-				float2 OffsetLength = (float(StepId) + 0.5) + float2(Random - 0.5, 0.5 - Random) * (bDoOneDirectionOnly ? 1.0 : 2.0);
-				float2 OffsetFraction = OffsetLength / TotalSteps;
+				mb_half2 OffsetLength = mb_half(StepId + 0.5) + mb_half2(Random - 0.5, 0.5 - Random) * (bDoOneDirectionOnly ? mb_half(1.0) : mb_half(2.0));
+				mb_half2 OffsetFraction = OffsetLength * rcp(mb_half(TotalSteps));
 
-				float WeightOffsetLength = (float)(DirectionId + StepId) + 0.5;
+				mb_half WeightOffsetLength = mb_half(DirectionId + StepId) + 0.5;
 
 				SampleUV[0] = ColorUV + OffsetFraction.x * SearchVector;
 				SampleUV[1] = ColorUV - OffsetFraction.y * SearchVector;
@@ -393,26 +520,26 @@ void MainCS(uint2 DispatchThreadId : SV_DispatchThreadID)
 				UNROLL
 				for (uint j = 0; j < 2; j++)
 				{
-					float3 SampleVelocityDepth = VelocityFlatTexture.SampleLevel(
+					mb_half3 SampleVelocityDepth = VelocityFlatTexture.SampleLevel(
 						SharedVelocityFlatSampler, ApplyScreenTransform(SampleUV[j], ColorToVelocity), 0).xyz;
 
-					SampleColor[j] = ColorTexture.SampleLevel(ColorSampler, SampleUV[j], 0);
-					SampleDepth[j] = SampleVelocityDepth.z;
+					SampleColor[j] = ColorTexture.SampleLevel(ColorSampler, SampleUV[j], MipLevel);
+					SampleDepth[j] = mb_half(SampleVelocityDepth.z);
 
 					// Decode
 					SampleVelocityDepth.x = GetVelocityLengthPixels(SampleVelocityDepth.x); // TODO: move in velocity flatten
-					SampleVelocityDepth.y = SampleVelocityDepth.y * (2.0 * PI) - PI;
+					SampleVelocityDepth.y = SampleVelocityDepth.y * mb_half(2.0 * PI) - mb_half(PI);
 
 					// in pixels
 					SampleVelocityLength[j] = SampleVelocityDepth.x;
 					
 					ConvolutionWeight[j] = ComputeSampleConvolutionWeight(
-						SampleDepth[j], SampleVelocityDepth.x, SampleVelocityDepth.y,
-						OffsetLength, BlurAngle, PixelToSampleScale);
+						SampleVelocityDepth.z, SampleVelocityDepth.x, SampleVelocityDepth.y,
+						OffsetLength.x, BlurAngle, PixelToSampleScale);
 
 					HoleFillingWeight[j] = ComputeCenterOrSampleWeight(
 						CenterDepth, CenterVelocityLength, CenterVelocityAngle,
-						SampleDepth[j], SampleVelocityDepth.x, SampleVelocityDepth.y,
+						SampleVelocityDepth.z, SampleVelocityDepth.x, SampleVelocityDepth.y,
 						WeightOffsetLength, BlurAngle, PixelToSampleScale, SOFT_Z_EXTENT);
 				}
 		
@@ -434,10 +561,10 @@ void MainCS(uint2 DispatchThreadId : SV_DispatchThreadID)
 				
 					// Measure how much should hole fill DirectionalColor with.
 					{
-						float HoleFillingWeight0 = saturate(HoleFillingWeight[0] - ConvolutionWeight[0]);
-						float HoleFillingWeight1 = saturate(HoleFillingWeight[1] - ConvolutionWeight[1]);
+						mb_half HoleFillingWeight0 = saturate(HoleFillingWeight[0] - ConvolutionWeight[0]);
+						mb_half HoleFillingWeight1 = saturate(HoleFillingWeight[1] - ConvolutionWeight[1]);
 
-						//float HoleFillingWeight = ComputeSampleHoleFillWeight(CenterDepth, SampleDepth[j], SOFT_Z_EXTENT);
+						//mb_half HoleFillingWeight = ComputeSampleHoleFillWeight(CenterDepth, SampleDepth[j], SOFT_Z_EXTENT);
 						
 						HoleFillWeightAccum[DirectionId] += HoleFillingWeight0 + HoleFillingWeight1;
 					
@@ -457,40 +584,45 @@ void MainCS(uint2 DispatchThreadId : SV_DispatchThreadID)
 				}
 				#endif
 			} // for (uint StepId = 0; StepId < StepCount; StepId += CONFIG_MAX_RANGE_SIZE)
+
+			#if DEBUG_OVERLAY_INPUT_RES
+			{
+				mb_half4 DebugMipLevel = lerp(mb_half4(1.0, 0.5, 0.5, 1.0), mb_half4(0.5, 1.0, 0.5, 1.0), mb_half(MipLevel0));
+				DirectionalColor[DirectionId] *= DebugMipLevel;
+			}
+			#endif
 		} // for (uint DirectionId = 0; DirectionId < CONFIG_MAX_RANGE_SIZE; DirectionId++)
 		
 		#if CONFIG_MAX_RANGE_SIZE > 1
 		{
-			float InvStepCount = rcp(float(StepPerDirectionCount));
-			float InvSampleCount = rcp(float(SamplePerDirection));
+			mb_half InvDirectionSampleCount = InvSampleCount * 2.0;
 			if (bDoOneDirectionOnly)
 			{
 				DirectionalColor[0] += DirectionalColor[1];
 				DirectionalColorWeight[0] += DirectionalColorWeight[1];
 				HoleFillWeightAccum[0] += HoleFillWeightAccum[1];
 		
-				DirectionalColor[1] = 0.0;
-				DirectionalColorWeight[1] = 0.0;
-				HoleFillWeightAccum[1] = 0.0;
+				DirectionalColor[1] = mb_half(0.0);
+				DirectionalColorWeight[1] = mb_half(0.0);
+				HoleFillWeightAccum[1] = mb_half(0.0);
 		
-				InvStepCount = rcp(float(TotalSteps));
-				InvSampleCount = rcp(float(CONFIG_SAMPLE_COUNT));
+				InvDirectionSampleCount = InvSampleCount;
 			}
 		
-			float Velocity1Translucency;
+			mb_half Velocity1Translucency;
 			{
-				float AvgDepthSquare0 = NormalizeAccumulator(DepthSquareAccum[0], DepthAwareWeight[0]);
-				float AvgDepth0 = NormalizeAccumulator(DepthAccum[0], DepthAwareWeight[0]);
-				float AvgDepth1 = NormalizeAccumulator(DepthAccum[1], DepthAwareWeight[1]);
+				mb_half AvgDepthSquare0 = NormalizeAccumulator(DepthSquareAccum[0], mb_half(DepthAwareWeight[0]));
+				mb_half AvgDepth0 = NormalizeAccumulator(DepthAccum[0], mb_half(DepthAwareWeight[0]));
+				mb_half AvgDepth1 = NormalizeAccumulator(DepthAccum[1], mb_half(DepthAwareWeight[1]));
 
-				float Variance0 = AvgDepthSquare0 - Square(AvgDepth0);
+				mb_half Variance0 = AvgDepthSquare0 - Square(AvgDepth0);
 
-				Velocity1Translucency = saturate(2.0 * Variance0 / (Variance0 + Square(AvgDepth1 - AvgDepth0)));
+				Velocity1Translucency = mb_half(saturate(2.0 * Variance0 / (Variance0 + Square(AvgDepth1 - AvgDepth0))));
 			}
 
-			float4 NormalizedHoleFillColor0 = NormalizeAccumulator(HoleFillColor, HoleFillColorWeight);
-			float4 NormalizedColor0;
-			float Opacity0;
+			mb_half4 NormalizedHoleFillColor0 = NormalizeAccumulator(HoleFillColor, HoleFillColorWeight);
+			mb_half4 NormalizedColor0;
+			mb_half Opacity0;
 			bool bValidColorOutput0;
 			NormalizeAccumulatorWithHoleFill(
 				DirectionalColor[0],
@@ -498,14 +630,14 @@ void MainCS(uint2 DispatchThreadId : SV_DispatchThreadID)
 				HoleFillWeightAccum[0],
 				DepthAwareWeight[0],
 				NormalizedHoleFillColor0,
-				InvSampleCount,
+				InvDirectionSampleCount,
 				/* out */ NormalizedColor0,
 				/* out */ Opacity0,
 				/* out */ bValidColorOutput0);
 
-			float4 NormalizedHoleFillColor1 = lerp(NormalizedColor0, NormalizedHoleFillColor0, bValidColorOutput0 ? Velocity1Translucency : 1.0);
-			float4 NormalizedColor1;
-			float Opacity1;
+			mb_half4 NormalizedHoleFillColor1 = lerp(NormalizedColor0, NormalizedHoleFillColor0, bValidColorOutput0 ? Velocity1Translucency : 1.0);
+			mb_half4 NormalizedColor1;
+			mb_half Opacity1;
 			bool bValidColorOutput1;
 			NormalizeAccumulatorWithHoleFill(
 				DirectionalColor[1],
@@ -513,45 +645,255 @@ void MainCS(uint2 DispatchThreadId : SV_DispatchThreadID)
 				HoleFillWeightAccum[1],
 				DepthAwareWeight[1],
 				NormalizedHoleFillColor1,
-				InvSampleCount,
+				InvDirectionSampleCount,
 				/* out */ NormalizedColor1,
 				/* out */ Opacity1,
 				/* out */ bValidColorOutput1);
 		
 			Opacity1 *= saturate(1.0 - Velocity1Translucency);
 
-			float CenterColorContribution = saturate(1.0 - Opacity0) * saturate(1.0 - Opacity1);
-			float InvTotalWeight = rcp(Opacity0 + Opacity1);
+			mb_half CenterColorContribution = saturate(1.0 - Opacity0) * saturate(1.0 - Opacity1);
+			mb_half InvTotalWeight = rcp(Opacity0 + Opacity1);
 
-			OutColor = (NormalizedColor0 * Opacity0 + NormalizedColor1 * Opacity1) * -min(-InvTotalWeight * (1.0 - CenterColorContribution), 0.0) + CenterColorContribution * CenterColor;
+			MotionBlurColor = (NormalizedColor0 * Opacity0 + NormalizedColor1 * Opacity1) * -min(-InvTotalWeight * (1.0 - CenterColorContribution), 0.0);
+			FullResBlend = CenterColorContribution;
+			
+			#if DEBUG_OVERLAY_TILE_CLASSIFICATION
+				DebugOverlay = mb_half4(1.0, 0.5, 0.5, 1.0);
+			#endif
 		}
 		#else // CONFIG_MAX_RANGE_SIZE == 1
 		{
 			const uint DirectionId = 0;
 
-			DirectionalColor[DirectionId] *= 0.5 / float(StepPerDirectionCount);
-			DirectionalColorWeight[DirectionId] *= 0.5 / float(StepPerDirectionCount);
+			DirectionalColor[DirectionId] *= mb_half(0.5) / mb_half(StepPerDirectionCount);
+			DirectionalColorWeight[DirectionId] *= mb_half(0.5) / mb_half(StepPerDirectionCount);
 
-			OutColor = DirectionalColor[DirectionId] + (1.0 - DirectionalColorWeight[DirectionId]) * CenterColor;
+			MotionBlurColor = DirectionalColor[DirectionId];
+			FullResBlend = (mb_half(1.0) - DirectionalColorWeight[DirectionId]);
+			
+			#if DEBUG_OVERLAY_TILE_CLASSIFICATION
+				DebugOverlay = mb_half4(1.0, 1.0, 0.5, 1.0);
+			#endif
 		}
 		#endif
 	}
-	
-	#if USE_POST_MOTION_BLUR_TRANSLUCENCY
-		OutColor.rgb = OutColor.rgb * PostMotionBlurTranslucency.a + PostMotionBlurTranslucency.rgb;
 	#endif
 
-	#if !POST_PROCESS_ALPHA
-		OutColor.a = 0;
+	#if CONFIG_SKIP_CENTER
+		bool bFetchCenter = WaveActiveAnyTrue(FullResBlend > 0.0);
+	#else
+		const bool bFetchCenter = true;
 	#endif
+
+	#if DEBUG_OVERLAY_GATHER_RES
+	{
+		mb_half4 DebugHalfResGatherLevel = lerp(mb_half4(1.0, 0.5, 0.5, 1.0), mb_half4(0.5, 1.0, 0.5, 1.0), bIsHalfRes ? mb_half(1.0) : mb_half(0.0));
+		MotionBlurColor *= DebugHalfResGatherLevel;
+	}
+	#endif
+
+	#if DEBUG_OVERLAY_SKIP_CENTER
+		DebugOverlay = lerp(mb_half4(1.0, 0.5, 0.5, 1.0), mb_half4(0.5, 1.0, 0.5, 1.0), bFetchCenter ? mb_half(0.0) : mb_half(1.0));
+	#endif
+
+	BRANCH
+	if (bIsHalfRes)
+	{
+		#if CONFIG_SHUFFLE_HALF_RES
+			mb_short2 OutputPixelCoord = mb_short2(ColorUV * Color_Extent);
+			mb_short2 OutputPixelCoord0 = (OutputPixelCoord & mb_short(~0x3)) | ((OutputPixelCoord >> mb_short(1)) & mb_short(0x1));
+			mb_short2 OutputPixelCoord1 = OutputPixelCoord0 + mb_short2(2, 0);
+			mb_short2 OutputPixelCoord2 = OutputPixelCoord0 + mb_short2(0, 2);
+			mb_short2 OutputPixelCoord3 = OutputPixelCoord0 + mb_short2(2, 2);
+		#else
+			mb_short2 OutputPixelCoord0 = mb_short2(uint2(ColorUV * Color_Extent)) & mb_short(~0x1);
+			mb_short2 OutputPixelCoord1 = OutputPixelCoord0 + mb_short2(1, 0);
+			mb_short2 OutputPixelCoord2 = OutputPixelCoord0 + mb_short2(0, 1);
+			mb_short2 OutputPixelCoord3 = OutputPixelCoord0 + mb_short2(1, 1);
+		#endif
+
+		#if 1
+			float2 PostMotionBlurTranslucencyUV0 = ApplyScreenTransform((float2(uint2(OutputPixelCoord0)) + 0.5) * Color_ExtentInverse, ColorToTranslucency);
+			float2 PostMotionBlurTranslucencyUV1 = ApplyScreenTransform((float2(uint2(OutputPixelCoord1)) + 0.5) * Color_ExtentInverse, ColorToTranslucency);
+			float2 PostMotionBlurTranslucencyUV2 = ApplyScreenTransform((float2(uint2(OutputPixelCoord2)) + 0.5) * Color_ExtentInverse, ColorToTranslucency);
+			float2 PostMotionBlurTranslucencyUV3 = ApplyScreenTransform((float2(uint2(OutputPixelCoord3)) + 0.5) * Color_ExtentInverse, ColorToTranslucency);
+		#elif CONFIG_SHUFFLE_HALF_RES
+			float2 ColorUV0 = (float2(uint2(OutputPixelCoord0)) + 0.5) * Color_ExtentInverse;
+
+			float2 PostMotionBlurTranslucencyUV0 = ApplyScreenTransform(ColorUV0, ColorToTranslucency);
+			float2 PostMotionBlurTranslucencyUV1 = PostMotionBlurTranslucencyUV0 + float2(2.0, 0.0) * TranslucencyExtentInverse;
+			float2 PostMotionBlurTranslucencyUV2 = PostMotionBlurTranslucencyUV0 + float2(0.0, 2.0) * TranslucencyExtentInverse;
+			float2 PostMotionBlurTranslucencyUV3 = PostMotionBlurTranslucencyUV0 + float2(2.0, 2.0) * TranslucencyExtentInverse;
+		#else
+			float2 PostMotionBlurTranslucencyUV = ApplyScreenTransform(ColorUV, ColorToTranslucency);
+			float2 PostMotionBlurTranslucencyUV0 = PostMotionBlurTranslucencyUV + float2(-0.5, -0.5) * TranslucencyExtentInverse;
+			float2 PostMotionBlurTranslucencyUV1 = PostMotionBlurTranslucencyUV + float2(+0.5, -0.5) * TranslucencyExtentInverse;
+			float2 PostMotionBlurTranslucencyUV2 = PostMotionBlurTranslucencyUV + float2(-0.5, +0.5) * TranslucencyExtentInverse;
+			float2 PostMotionBlurTranslucencyUV3 = PostMotionBlurTranslucencyUV + float2(+0.5, +0.5) * TranslucencyExtentInverse;
+		#endif
+			
+		PostMotionBlurTranslucencyUV0 = clamp(PostMotionBlurTranslucencyUV0, TranslucencyUVMin, TranslucencyUVMax);
+		PostMotionBlurTranslucencyUV1 = clamp(PostMotionBlurTranslucencyUV1, TranslucencyUVMin, TranslucencyUVMax);
+		PostMotionBlurTranslucencyUV2 = clamp(PostMotionBlurTranslucencyUV2, TranslucencyUVMin, TranslucencyUVMax);
+		PostMotionBlurTranslucencyUV3 = clamp(PostMotionBlurTranslucencyUV3, TranslucencyUVMin, TranslucencyUVMax);
+		
+		mb_half4 CenterColor0;
+		mb_half4 CenterColor1;
+		mb_half4 CenterColor2;
+		mb_half4 CenterColor3;
+		
+		mb_half4 PostMotionBlurTranslucency0;
+		mb_half4 PostMotionBlurTranslucency1;
+		mb_half4 PostMotionBlurTranslucency2;
+		mb_half4 PostMotionBlurTranslucency3;
+		
+		// Fetch center and motion blur translucency with overlapped texture fetches
+		BRANCH
+		if (bFetchCenter)
+		{
+			CenterColor0 = ColorTexture[OutputPixelCoord0];
+			CenterColor1 = ColorTexture[OutputPixelCoord1];
+			CenterColor2 = ColorTexture[OutputPixelCoord2];
+			CenterColor3 = ColorTexture[OutputPixelCoord3];
+		
+			PostMotionBlurTranslucency0 = TranslucencyTexture.SampleLevel(TranslucencySampler, PostMotionBlurTranslucencyUV0, 0);
+			PostMotionBlurTranslucency1 = TranslucencyTexture.SampleLevel(TranslucencySampler, PostMotionBlurTranslucencyUV1, 0);
+			PostMotionBlurTranslucency2 = TranslucencyTexture.SampleLevel(TranslucencySampler, PostMotionBlurTranslucencyUV2, 0);
+			PostMotionBlurTranslucency3 = TranslucencyTexture.SampleLevel(TranslucencySampler, PostMotionBlurTranslucencyUV3, 0);
+		}
+		else
+		{
+			CenterColor0 = mb_half(0.0).xxxx;
+			CenterColor1 = mb_half(0.0).xxxx;
+			CenterColor2 = mb_half(0.0).xxxx;
+			CenterColor3 = mb_half(0.0).xxxx;
+		
+			PostMotionBlurTranslucency0 = TranslucencyTexture.SampleLevel(TranslucencySampler, PostMotionBlurTranslucencyUV0, 0);
+			PostMotionBlurTranslucency1 = TranslucencyTexture.SampleLevel(TranslucencySampler, PostMotionBlurTranslucencyUV1, 0);
+			PostMotionBlurTranslucency2 = TranslucencyTexture.SampleLevel(TranslucencySampler, PostMotionBlurTranslucencyUV2, 0);
+			PostMotionBlurTranslucency3 = TranslucencyTexture.SampleLevel(TranslucencySampler, PostMotionBlurTranslucencyUV3, 0);
+		}
+		
+		#if DEBUG_OVERLAY_INPUT_RES
+		{
+			CenterColor0 *= mb_half4(1.0, 0.5, 0.5, 1.0);
+			CenterColor1 *= mb_half4(1.0, 0.5, 0.5, 1.0);
+			CenterColor2 *= mb_half4(1.0, 0.5, 0.5, 1.0);
+			CenterColor3 *= mb_half4(1.0, 0.5, 0.5, 1.0);
+		}
+		#endif
+
+		// Blend full res and motion blur
+		mb_half4 OutputColor0 = CenterColor0 * FullResBlend + MotionBlurColor;
+		mb_half4 OutputColor1 = CenterColor1 * FullResBlend + MotionBlurColor;
+		mb_half4 OutputColor2 = CenterColor2 * FullResBlend + MotionBlurColor;
+		mb_half4 OutputColor3 = CenterColor3 * FullResBlend + MotionBlurColor;
+		
+		// Adds debug overlay
+		#if DEBUG_OVERLAY_TILE_CLASSIFICATION || DEBUG_OVERLAY_SKIP_CENTER || DEBUG_OVERLAY_SAMPLES
+		{
+			OutputColor0 *= DebugOverlay;
+			OutputColor1 *= DebugOverlay;
+			OutputColor2 *= DebugOverlay;
+			OutputColor3 *= DebugOverlay;
+		}
+		#endif
+
+		// Blend post motion blur translucency
+		#if CONFIG_POST_MOTIONBLUR_TRANSLUCENCY
+		{
+			OutputColor0.rgb = OutputColor0.rgb * PostMotionBlurTranslucency0.a + PostMotionBlurTranslucency0.rgb;
+			OutputColor1.rgb = OutputColor1.rgb * PostMotionBlurTranslucency1.a + PostMotionBlurTranslucency1.rgb;
+			OutputColor2.rgb = OutputColor2.rgb * PostMotionBlurTranslucency2.a + PostMotionBlurTranslucency2.rgb;
+			OutputColor3.rgb = OutputColor3.rgb * PostMotionBlurTranslucency3.a + PostMotionBlurTranslucency3.rgb;
+		}
+		#endif
+
+		#if !POST_PROCESS_ALPHA
+		{
+			OutputColor0.a = 0;
+			OutputColor1.a = 0;
+			OutputColor2.a = 0;
+			OutputColor3.a = 0;
+		}
+		#endif
 	
-	if (all(iColorPixelPos < Color_ViewportMax))
+		if (all(OutputPixelCoord0 < mb_short2(Color_ViewportMax)))
+		{
+			OutColorTexture[OutputPixelCoord0] = OutputColor0;
+			OutColorTexture[OutputPixelCoord1] = OutputColor1;
+			OutColorTexture[OutputPixelCoord2] = OutputColor2;
+			OutColorTexture[OutputPixelCoord3] = OutputColor3;
+			
+			#if DEBUG_OUTPUT
+				DebugOutput[OutputPixelCoord0] = Debug;
+				DebugOutput[OutputPixelCoord1] = Debug;
+				DebugOutput[OutputPixelCoord2] = Debug;
+				DebugOutput[OutputPixelCoord3] = Debug;
+			#endif
+		}
+	}
+	else
 	{
-		OutColorTexture[iColorPixelPos] = OutColor;
+		mb_short2 OutputPixelCoord = mb_short2(ColorUV * Color_Extent);
 
-		#if DEBUG_OUTPUT
-			DebugOutput[iColorPixelPos] = Debug;
+		float2 PostMotionBlurTranslucencyUV = ApplyScreenTransform(ColorUV, ColorToTranslucency);
+		PostMotionBlurTranslucencyUV = clamp(PostMotionBlurTranslucencyUV, TranslucencyUVMin, TranslucencyUVMax);
+		
+		mb_half4 CenterColor;
+		mb_half4 PostMotionBlurTranslucency;
+		
+		// Fetch center and motion blur translucency with overlapped texture fetches
+		BRANCH
+		if (bFetchCenter)
+		{
+			CenterColor = ColorTexture[OutputPixelCoord];
+			PostMotionBlurTranslucency = TranslucencyTexture.SampleLevel(TranslucencySampler, PostMotionBlurTranslucencyUV, 0);
+		}
+		else
+		{
+			CenterColor = mb_half(0.0).xxxx;
+			PostMotionBlurTranslucency = TranslucencyTexture.SampleLevel(TranslucencySampler, PostMotionBlurTranslucencyUV, 0);
+		}
+		
+		#if DEBUG_OVERLAY_INPUT_RES
+		{
+			CenterColor *= mb_half4(1.0, 0.5, 0.5, 1.0);
+		}
+		#endif
+
+		// Blend full res and motion blur
+		mb_half4 OutputColor = CenterColor * FullResBlend + MotionBlurColor;
+		
+		// Adds debug overlay
+		#if DEBUG_OVERLAY_TILE_CLASSIFICATION || DEBUG_OVERLAY_SKIP_CENTER || DEBUG_OVERLAY_SAMPLES
+		{
+			OutputColor *= DebugOverlay;
+		}
+		#endif
+
+		// Blend post motion blur translucency
+		#if CONFIG_POST_MOTIONBLUR_TRANSLUCENCY
+		{
+			OutputColor.rgb = OutputColor.rgb * PostMotionBlurTranslucency.a + PostMotionBlurTranslucency.rgb;
+		}
+		#endif
+
+		#if !POST_PROCESS_ALPHA
+		{
+			OutputColor.a = 0;
+		}
 		#endif
+	
+		if (all(OutputPixelCoord < mb_short2(Color_ViewportMax)))
+		{
+			OutColorTexture[OutputPixelCoord] = OutputColor;
+
+			#if DEBUG_OUTPUT
+				DebugOutput[OutputPixelCoord] = Debug;
+			#endif
+		}
 	}
 }
 
diff --git a/Engine/Shaders/Private/MotionBlur/MotionBlurCommon.ush b/Engine/Shaders/Private/MotionBlur/MotionBlurCommon.ush
--- a/Engine/Shaders/Private/MotionBlur/MotionBlurCommon.ush
+++ b/Engine/Shaders/Private/MotionBlur/MotionBlurCommon.ush
@@ -5,6 +5,12 @@
 #include "../ScreenPass.ush"
 
 
+//------------------------------------------------------- COMPILER CONFIG
+
+// Generate vector truncation warnings to errors.
+#pragma warning(error: 3206)
+
+
 //------------------------------------------------------- CONFIG
 
 #define DEBUG_OUTPUT 0
@@ -13,6 +19,92 @@
 	#error need to set CONFIG_MAX_RANGE_SIZE
 #endif
 
+// Size in input pixel of the velocity flatten tiling
+#define VELOCITY_FLATTEN_TILE_SIZE 16
+
+// Size of the tile at full resolution
+#define VELOCITY_FILTER_TILE_SIZE 16
+
+#if PLATFORM_SUPPORTS_REAL_TYPES && 0
+	#define CONFIG_COMPILE_FP16 1
+#else
+	#define CONFIG_COMPILE_FP16 0
+#endif
+
+#define CONFIG_MINIMAL_PIXEL_VELOCITY 0.5
+
+
+//------------------------------------------------------- CONSTANTS
+
+#if CONFIG_COMPILE_FP16
+	#define mb_half  half
+	#define mb_half2 half2
+	#define mb_half3 half3
+	#define mb_half4 half4
+
+	#define mb_short  int16_t
+	#define mb_short2 int16_t2
+	#define mb_short3 int16_t3
+	#define mb_short4 int16_t4
+
+	#define mb_ushort  uint16_t
+	#define mb_ushort2 uint16_t2
+	#define mb_ushort3 uint16_t3
+	#define mb_ushort4 uint16_t4
+
+	#define mb_half2x2 half2x2
+	#define mb_half3x2 half3x2
+	#define mb_half4x2 half4x2
+
+	#define mb_short2x2 int16_t2x2
+	#define mb_short3x2 int16_t3x2
+	#define mb_short4x2 int16_t4x2
+
+	#define mb_ushort2x2 uint16_t2x2
+	#define mb_ushort3x2 uint16_t3x2
+	#define mb_ushort4x2 uint16_t4x2
+
+#else
+	#define mb_half  float
+	#define mb_half2 float2
+	#define mb_half3 float3
+	#define mb_half4 float4
+
+	#define mb_short  int
+	#define mb_short2 int2
+	#define mb_short3 int3
+	#define mb_short4 int4
+
+	#define mb_ushort  uint
+	#define mb_ushort2 uint2
+	#define mb_ushort3 uint3
+	#define mb_ushort4 uint4
+
+	#define mb_half2x2 float2x2
+	#define mb_half3x2 float3x2
+	#define mb_half4x2 float4x2
+
+	#define mb_short2x2 int2x2
+	#define mb_short3x2 int3x2
+	#define mb_short4x2 int4x2
+
+	#define mb_ushort2x2 uint2x2
+	#define mb_ushort3x2 uint3x2
+	#define mb_ushort4x2 uint4x2
+
+#endif
+
+//------------------------------------------------------- DEFINES
+
+/** Filter tile classification. Matches FMotionBlurFilterCS::ETileClassification */
+#define FILTER_TILE_CLASSIFY_GATHER_HALF_RES 0
+#define FILTER_TILE_CLASSIFY_GATHER_FULL_RES 1
+#define FILTER_TILE_CLASSIFY_SCATTER_AS_GATHER_1_VELOCITY_HALF_RES 2
+#define FILTER_TILE_CLASSIFY_SCATTER_AS_GATHER_1_VELOCITY_FULL_RES 3
+#define FILTER_TILE_CLASSIFY_SCATTER_AS_GATHER_2_VELOCITY_FULL_RES 4
+#define FILTER_TILE_CLASSIFY_COUNT 5
+
+
 
 //------------------------------------------------------- FUNCTIONS
 
diff --git a/Engine/Shaders/Private/MotionBlur/MotionBlurFilterTileClassify.usf b/Engine/Shaders/Private/MotionBlur/MotionBlurFilterTileClassify.usf
new file mode 100644
--- /dev/null
+++ b/Engine/Shaders/Private/MotionBlur/MotionBlurFilterTileClassify.usf
@@ -0,0 +1,187 @@
+// Copyright Epic Games, Inc. All Rights Reserved.
+
+#define CONFIG_MAX_RANGE_SIZE DIM_BLUR_DIRECTIONS
+
+
+#include "MotionBlurCommon.ush"
+
+//------------------------------------------------------- CONFIG
+
+#define TILE_SIZE 16
+
+
+//------------------------------------------------------- PARAMETERS
+
+FScreenTransform FilterTileIdToFlattenTileId;
+uint2 FlattenTileMaxId;
+uint2 FilterTileCount;
+uint bAllowHalfResGather;
+uint TileListMaxSize;
+float HalfResPixelVelocityThreshold;
+
+Texture2D VelocityTileTextures_Textures_0;
+Texture2D VelocityTileTextures_Textures_1;
+
+RWCoherentStructuredBuffer(uint) TileListsOutput;
+RWCoherentStructuredBuffer(uint) TileListsSizeOutput;
+
+RWTexture2D<float4> DebugOutput;
+
+
+//------------------------------------------------------- LDS
+
+groupshared uint SharedTileCount[FILTER_TILE_CLASSIFY_COUNT];
+
+
+//------------------------------------------------------- ENTRY POINT
+
+#ifdef SETUP_PASS
+
+StructuredBuffer<uint> TileListsSizeBuffer;
+RWBuffer<uint> DispatchParametersOutput;
+
+[numthreads(FILTER_TILE_CLASSIFY_COUNT, 1, 1)]
+void MainCS(uint GroupThreadIndex : SV_GroupIndex)
+{
+	if (GroupThreadIndex < FILTER_TILE_CLASSIFY_COUNT)
+	{
+		DispatchParametersOutput[3 * GroupThreadIndex + 0] = TileListsSizeBuffer[GroupThreadIndex];
+		DispatchParametersOutput[3 * GroupThreadIndex + 1] = 1;
+		DispatchParametersOutput[3 * GroupThreadIndex + 2] = 1;
+	}
+}
+
+#else // !SETUP_PASS
+
+[numthreads(TILE_SIZE, TILE_SIZE, 1)]
+void MainCS(
+	uint2 GroupId : SV_GroupID,
+	uint GroupThreadIndex : SV_GroupIndex)
+{
+	float4 Debug = 0;
+	
+	uint2 DispatchThreadId = (
+		ZOrder2D(GroupThreadIndex, log2(TILE_SIZE)) +
+		GroupId * uint2(TILE_SIZE, TILE_SIZE));
+
+	mb_short2 FilterTileCoord = DispatchThreadId;
+	bool bIsValidTileCoord = all(FilterTileCoord < mb_short2(FilterTileCount));
+
+	mb_short2 FlattenTileCoord = mb_short2(ApplyScreenTransform(mb_half2(FilterTileCoord), FilterTileIdToFlattenTileId));
+
+	uint TileClassify = FILTER_TILE_CLASSIFY_GATHER_HALF_RES;
+
+	UNROLL_N(3)
+	for (int y = -1; y <= 1; y++)
+	{
+		UNROLL_N(3)
+		for (int x = -1; x <= 1; x++)
+		{
+			const mb_short2 Offset = mb_short2(x, y);
+
+			mb_short2 NeighborFlattenTileCoord = FlattenTileCoord + Offset;
+			NeighborFlattenTileCoord = clamp(NeighborFlattenTileCoord, mb_short(0).xx, mb_short2(FlattenTileMaxId));
+
+			FVelocityRange VelocityRange = DecodeVelocityRange(
+				VelocityTileTextures_Textures_0[NeighborFlattenTileCoord],
+				VelocityTileTextures_Textures_1[NeighborFlattenTileCoord]);
+			
+			float MinVelocityLengthSqrPixels = dot(VelocityRange.Min, VelocityRange.Min);
+			float Max0VelocityLengthSqrPixels = dot(VelocityRange.Max[0], VelocityRange.Max[0]);
+			
+			bool bHalfRes = (Max0VelocityLengthSqrPixels > HalfResPixelVelocityThreshold) && (bAllowHalfResGather != 0);
+
+			#if CONFIG_MAX_RANGE_SIZE > 1
+				float Max1VelocityLengthSqrPixels = dot(VelocityRange.Max[1], VelocityRange.Max[1]);
+				bool bDoOneDirectionOnly = Max1VelocityLengthSqrPixels < (CONFIG_MINIMAL_PIXEL_VELOCITY * CONFIG_MINIMAL_PIXEL_VELOCITY);
+			#else
+				const bool bDoOneDirectionOnly = true;
+			#endif
+
+			bool bSkip0 = Max0VelocityLengthSqrPixels < (CONFIG_MINIMAL_PIXEL_VELOCITY * CONFIG_MINIMAL_PIXEL_VELOCITY);
+			bool bGatherOnly = MinVelocityLengthSqrPixels > 0.4 * Max0VelocityLengthSqrPixels;
+
+			uint SuggestedTileClassification = FILTER_TILE_CLASSIFY_SCATTER_AS_GATHER_1_VELOCITY_FULL_RES;
+			if (bSkip0)
+			{
+				// Dissregard bAllowHalfRes because FILTER_TILE_CLASSIFY_GATHER_HALF_RES is still the quickest when no motion blur is to applied.
+				SuggestedTileClassification = FILTER_TILE_CLASSIFY_GATHER_HALF_RES;
+			}
+			#if CONFIG_MAX_RANGE_SIZE > 1 && 1
+			else if (!bDoOneDirectionOnly)
+			{
+				SuggestedTileClassification = FILTER_TILE_CLASSIFY_SCATTER_AS_GATHER_2_VELOCITY_FULL_RES;
+			}
+			#endif
+			else if (bGatherOnly)
+			{
+				SuggestedTileClassification = bHalfRes ? FILTER_TILE_CLASSIFY_GATHER_HALF_RES : FILTER_TILE_CLASSIFY_GATHER_FULL_RES;
+			}
+			else
+			{
+				SuggestedTileClassification = bHalfRes ? FILTER_TILE_CLASSIFY_SCATTER_AS_GATHER_1_VELOCITY_HALF_RES : FILTER_TILE_CLASSIFY_SCATTER_AS_GATHER_1_VELOCITY_FULL_RES;
+			}
+
+			TileClassify = max(TileClassify, SuggestedTileClassification);
+		}
+	}
+
+	if (GroupThreadIndex < FILTER_TILE_CLASSIFY_COUNT)
+	{
+		SharedTileCount[GroupThreadIndex] = 0;
+	}
+	
+	GroupMemoryBarrierWithGroupSync();
+	
+	// Alloc within the group
+	uint LocalAllocation = 0;
+	if (bIsValidTileCoord)
+	{
+		InterlockedAdd(SharedTileCount[TileClassify], 1, /* out */ LocalAllocation);
+	}
+	
+	GroupMemoryBarrierWithGroupSync();
+	
+	// Alloc globally
+	if (GroupThreadIndex < FILTER_TILE_CLASSIFY_COUNT)
+	{
+		uint GlobalGroupAllocation;
+		InterlockedAdd(TileListsSizeOutput[GroupThreadIndex], SharedTileCount[GroupThreadIndex], /* out */ GlobalGroupAllocation);
+
+		SharedTileCount[GroupThreadIndex] = GlobalGroupAllocation;
+	}
+	
+	GroupMemoryBarrierWithGroupSync();
+	
+	// Output
+	uint GlobalAllocation = SharedTileCount[TileClassify] + LocalAllocation;
+
+	if (TileClassify == FILTER_TILE_CLASSIFY_GATHER_HALF_RES)
+	{
+		Debug = float4(0, 1, 0, 0);
+	}
+	else if (TileClassify == FILTER_TILE_CLASSIFY_GATHER_FULL_RES)
+	{
+		Debug = float4(0, 0, 1, 0);
+	}
+	else if (TileClassify == FILTER_TILE_CLASSIFY_SCATTER_AS_GATHER_1_VELOCITY_HALF_RES || TileClassify == FILTER_TILE_CLASSIFY_SCATTER_AS_GATHER_1_VELOCITY_FULL_RES)
+	{
+		Debug = float4(1, 1, 0, 0);
+	}
+	else if (TileClassify == FILTER_TILE_CLASSIFY_SCATTER_AS_GATHER_2_VELOCITY_FULL_RES)
+	{
+		Debug = float4(1, 0, 0, 0);
+	}
+
+
+	if (bIsValidTileCoord)
+	{
+		TileListsOutput[GlobalAllocation + TileListMaxSize * TileClassify] = uint(FilterTileCoord.x) | (uint(FilterTileCoord.y) << 16);
+
+		#if DEBUG_OUTPUT
+			DebugOutput[FilterTileCoord] = Debug;
+		#endif
+	}
+}
+
+#endif // !SETUP_PASS
diff --git a/Engine/Shaders/Private/MotionBlur/MotionBlurTileGather.usf b/Engine/Shaders/Private/MotionBlur/MotionBlurTileGather.usf
--- a/Engine/Shaders/Private/MotionBlur/MotionBlurTileGather.usf
+++ b/Engine/Shaders/Private/MotionBlur/MotionBlurTileGather.usf
@@ -8,7 +8,7 @@
 //------------------------------------------------------- PARAMETERS
 
 SCREEN_PASS_TEXTURE_VIEWPORT(VelocityTile)
-float MotionBlur_VelocityScaleForTiles;
+float VelocityScaleForFlattenTiles;
 
 Texture2D CenterVelocityTileTextures_Textures_0;
 Texture2D CenterVelocityTileTextures_Textures_1;
@@ -53,7 +53,7 @@ void MainCS(
 
 			float2 MaxVelocity = NeighborVelocityRange.Max[0];
 
-			float2 VelocityPixels = MaxVelocity * MotionBlur_VelocityScaleForTiles;
+			float2 VelocityPixels = MaxVelocity * VelocityScaleForFlattenTiles;
 
 			float  VelocityLengthPixelsSqr = dot(VelocityPixels, VelocityPixels);
 			float  VelocityLengthPixelsInv = rsqrtFast(VelocityLengthPixelsSqr + 1e-8);
diff --git a/Engine/Shaders/Private/MotionBlur/MotionBlurTileScatter.usf b/Engine/Shaders/Private/MotionBlur/MotionBlurTileScatter.usf
--- a/Engine/Shaders/Private/MotionBlur/MotionBlurTileScatter.usf
+++ b/Engine/Shaders/Private/MotionBlur/MotionBlurTileScatter.usf
@@ -14,7 +14,7 @@
 //------------------------------------------------------- PARAMETERS
 
 SCREEN_PASS_TEXTURE_VIEWPORT(VelocityTile)
-float MotionBlur_VelocityScaleForTiles;
+float VelocityScaleForFlattenTiles;
 uint ScatterPass;
 
 Texture2D VelocityTileTextures_Textures_0;
@@ -56,7 +56,7 @@ void MainVS(
 	float4 MinMaxVelocity = VelocityTileTextures_Textures_0[PixelPos];
 	OutColor = MinMaxVelocity;
 
-	float4 MinMaxVelocityPixels = MinMaxVelocity * MotionBlur_VelocityScaleForTiles;
+	float4 MinMaxVelocityPixels = MinMaxVelocity * VelocityScaleForFlattenTiles;
 	float2 VelocityPixels = MinMaxVelocityPixels.zw;
 
 	// Is the velocity small enough not to cover adjacent tiles?
diff --git a/Engine/Shaders/Private/MotionBlur/MotionBlurVelocityFlatten.usf b/Engine/Shaders/Private/MotionBlur/MotionBlurVelocityFlatten.usf
--- a/Engine/Shaders/Private/MotionBlur/MotionBlurVelocityFlatten.usf
+++ b/Engine/Shaders/Private/MotionBlur/MotionBlurVelocityFlatten.usf
@@ -7,9 +7,8 @@
 
 //------------------------------------------------------- CONSTANTS
 
-#define THREADGROUP_SIZEX		16
-#define THREADGROUP_SIZEY		16
-#define THREADGROUP_TOTALSIZE	(THREADGROUP_SIZEX * THREADGROUP_SIZEY)
+#define TILE_SIZE				(VELOCITY_FLATTEN_TILE_SIZE)
+#define THREADGROUP_TOTALSIZE	(TILE_SIZE * TILE_SIZE)
 
 #ifndef CAMERA_MOTION_BLUR_MODE
 // 0-off; 1-on; 2-override
@@ -26,10 +25,8 @@ SCREEN_PASS_TEXTURE_VIEWPORT(VelocityTile)
 float4x4 ClipToPrevClipOverride;
 #endif
 
-float MotionBlur_AspectRatio;
-float MotionBlur_VelocityScale;
-float MotionBlur_VelocityScaleForTiles;
-float MotionBlur_VelocityMax;
+float2 VelocityScale;
+float VelocityMax;
 
 Texture2D VelocityTexture;
 Texture2D VelocityTileTexture;
@@ -96,7 +93,7 @@ void VelovityFlattenStep(uint GroupIndex, const uint GroupSize, inout FVelocityR
 	}
 }
 
-float2 GetCameraMotionVelocity(uint3 DispatchThreadId, float Depth, float4x4 ClipToPrevClip)
+float2 GetCameraMotionVelocity(uint2 DispatchThreadId, float Depth, float4x4 ClipToPrevClip)
 {
 	// Compute velocity due to camera motion.
 	float2 ViewportUV = ((float2)DispatchThreadId.xy + 0.5) / Velocity_ViewportSize;
@@ -110,11 +107,11 @@ float2 GetCameraMotionVelocity(uint3 DispatchThreadId, float Depth, float4x4 Cli
 
 //------------------------------------------------------- ENTRY POINT
 
-[numthreads(THREADGROUP_SIZEX, THREADGROUP_SIZEY, 1)]
+[numthreads(TILE_SIZE, TILE_SIZE, 1)]
 void MainCS(
-	uint3 GroupId : SV_GroupID,
-	uint3 DispatchThreadId : SV_DispatchThreadID,
-	uint3 GroupThreadId : SV_GroupThreadID,
+	uint2 GroupId : SV_GroupID,
+	uint2 DispatchThreadId : SV_DispatchThreadID,
+	uint2 GroupThreadId : SV_GroupThreadID,
 	uint GroupIndex : SV_GroupIndex) 
 { 
 #if DEBUG_OUTPUT
@@ -185,7 +182,7 @@ void MainCS(
 	}
 #endif
 
-	Velocity.y *= -MotionBlur_AspectRatio;
+	Velocity *= VelocityScale;
 
 	float2 VelocityPolar = CartesianToPolar(Velocity);
 
@@ -211,7 +208,7 @@ void MainCS(
 	}
 
 	// Limit velocity
-	VelocityPolar.x = min(VelocityPolar.x, MotionBlur_VelocityMax / MotionBlur_VelocityScale);
+	VelocityPolar.x = min(VelocityPolar.x, VelocityMax);
 
 	FVelocityRange VelocityPolarRange = SetupPolarVelocityRange(VelocityPolar);
 
diff --git a/Engine/Shaders/Private/TemporalSuperResolution/TSRUpdateHistory.usf b/Engine/Shaders/Private/TemporalSuperResolution/TSRUpdateHistory.usf
--- a/Engine/Shaders/Private/TemporalSuperResolution/TSRUpdateHistory.usf
+++ b/Engine/Shaders/Private/TemporalSuperResolution/TSRUpdateHistory.usf
@@ -74,6 +74,7 @@ float MinTranslucencyRejection;
 float InvWeightClampingPixelSpeed;
 float InputToHistoryFactor;
 uint ResponsiveStencilMask;
+uint bGenerateOutputMip1;
 
 Texture2D<taa_half3> PrevHistory_LowFrequency;
 Texture2D<taa_half3> PrevHistory_HighFrequency;
@@ -89,13 +90,12 @@ RWTexture2D<taa_half3> HistoryOutput_Translucency;
 RWTexture2D<taa_half>  HistoryOutput_TranslucencyAlpha;
 RWTexture2D<taa_subpixel_details> HistoryOutput_SubpixelDetails;
 
-RWTexture2D<taa_half3> SceneColorOutput;
+RWTexture2D<taa_half3> SceneColorOutputMip0;
+RWTexture2D<taa_half3> SceneColorOutputMip1;
 
 
 //------------------------------------------------------- LDS
 
-#if CONFIG_MANUAL_LDS_SPILL
-
 #if CONFIG_ENABLE_DUAL_PIXEL_VECTORIZATION
 #define GROUP_SHARED_SIZE (TILE_SIZE * TILE_SIZE / 2)
 #else
@@ -118,8 +118,6 @@ groupshared taa_half4 SharedArray11[GROUP_SHARED_SIZE];
 
 groupshared taa_subpixel_payload2 SharedParallaxFactorBits[GROUP_SHARED_SIZE];
 
-#endif // CONFIG_MANUAL_LDS_SPILL
-
 
 //------------------------------------------------------- FUNCTIONS
 
@@ -1360,17 +1358,69 @@ void MainCS(
 		}
 		#endif
 
-		// Output final scene color
+		// Output final scene color Mip0
+		taa_half3x2 OutputColor;
 		{
-			taa_half3x2 OutputColor = FinalOutputColor;
+			OutputColor = FinalOutputColor;
 
 			OutputColor = -dpv_min(-OutputColor, taa_half(0.0).xxx);
 			OutputColor = dpv_min(OutputColor, taa_half(Max10BitsFloat).xxx);
 
-			SceneColorOutput[dpv_lo(LocalHistoryPixelPos)] = dpv_lo(OutputColor);
+			SceneColorOutputMip0[dpv_lo(LocalHistoryPixelPos)] = dpv_lo(OutputColor);
+			
+			#if CONFIG_ENABLE_DUAL_PIXEL_VECTORIZATION
+				SceneColorOutputMip0[dpv_hi(LocalHistoryPixelPos)] = dpv_hi(OutputColor);
+			#endif
+		}
+		
+		// Output final scene color Mip1
+		{
+			taa_half3x2 HalfResOutput = OutputColor;
+			taa_short2x2 HalfResOutputPixelPos = (LocalHistoryPixelPos >> taa_short(1)) | (((LocalHistoryPixelPos & taa_short(0x1)) | !bGenerateOutputMip1) * taa_short(~0));
+			
+			//HalfResOutputPixelPos = (LocalHistoryPixelPos >> taa_short(1));
+			
+			#if PLATFORM_SUPPORTS_WAVE_BROADCAST && 0 // Support WaveBroadcast with halfs
+			BRANCH
+			if (bGenerateOutputMip1)
+			{
+				FWaveBroadcastSettings Horizontal = InitWaveXorButterfly(/* XorButterFly = */ 0x1);
+				FWaveBroadcastSettings Vertical = InitWaveXorButterfly(/* XorButterFly = */ 0x2);
+
+				HalfResOutput += WaveBroadcast(Horizontal, HalfResOutput);
+				HalfResOutput = HalfResOutput * taa_half(0.25) + taa_half(0.25) * WaveBroadcast(Vertical, HalfResOutput);
+			}
+			#else
+			BRANCH
+			if (bGenerateOutputMip1)
+			{
+				SharedArray0[LocalGroupThreadIndex] = taa_half4(dpv_lo(HalfResOutput), 0.0);
+				#if CONFIG_ENABLE_DUAL_PIXEL_VECTORIZATION
+					SharedArray1[LocalGroupThreadIndex] = taa_half4(dpv_hi(HalfResOutput), 0.0);
+				#endif
+
+				SharedArray0[LocalGroupThreadIndex] += SharedArray0[LocalGroupThreadIndex ^ 0x1];
+				#if CONFIG_ENABLE_DUAL_PIXEL_VECTORIZATION
+					SharedArray1[LocalGroupThreadIndex] += SharedArray1[LocalGroupThreadIndex ^ 0x1];
+				#endif
+					
+				SharedArray0[LocalGroupThreadIndex] += SharedArray0[LocalGroupThreadIndex ^ 0x2];
+				#if CONFIG_ENABLE_DUAL_PIXEL_VECTORIZATION
+					SharedArray1[LocalGroupThreadIndex] += SharedArray1[LocalGroupThreadIndex ^ 0x2];
+				#endif
+					
+				#if CONFIG_ENABLE_DUAL_PIXEL_VECTORIZATION
+					HalfResOutput = dpv_interleave_registers(SharedArray0[LocalGroupThreadIndex].rgb, SharedArray1[LocalGroupThreadIndex].rgb) * 0.25;
+				#else
+					HalfResOutput = dpv_interleave_mono_registers(SharedArray0[LocalGroupThreadIndex].rgb) * 0.25;
+				#endif
+			}
+			#endif
+
+			SceneColorOutputMip1[dpv_lo(HalfResOutputPixelPos)] = dpv_lo(HalfResOutput);
 			
 			#if CONFIG_ENABLE_DUAL_PIXEL_VECTORIZATION
-				SceneColorOutput[dpv_hi(LocalHistoryPixelPos)] = dpv_hi(OutputColor);
+				SceneColorOutputMip1[dpv_hi(HalfResOutputPixelPos)] = dpv_hi(HalfResOutput);
 			#endif
 		}
 
diff --git a/Engine/Source/Runtime/Renderer/Private/PostProcess/PostProcessMotionBlur.cpp b/Engine/Source/Runtime/Renderer/Private/PostProcess/PostProcessMotionBlur.cpp
--- a/Engine/Source/Runtime/Renderer/Private/PostProcess/PostProcessMotionBlur.cpp
+++ b/Engine/Source/Runtime/Renderer/Private/PostProcess/PostProcessMotionBlur.cpp
@@ -7,7 +7,6 @@
 #include "SpriteIndexBuffer.h"
 #include "PostProcessing.h"
 
-
 namespace
 {
 #if !(UE_BUILD_SHIPPING || UE_BUILD_TEST)
@@ -44,6 +43,16 @@ namespace
 		TEXT("Number of bluring direction (default = 1)."),
 		ECVF_RenderThreadSafe);
 
+	TAutoConsoleVariable<int32> CVarMotionBlurHalfResInput(
+		TEXT("r.MotionBlur.HalfResInput"), 1,
+		TEXT("Whether motion blur also blur with a half resolution input."),
+		ECVF_RenderThreadSafe);
+
+	TAutoConsoleVariable<int32> CVarMotionBlurHalfResGather(
+		TEXT("r.MotionBlur.HalfResGather"), 1,
+		TEXT("Whether to do motion blur filter dynamically at half res under heavy motion."),
+		ECVF_RenderThreadSafe);
+
 	TAutoConsoleVariable<int32> CVarAllowMotionBlurInVR(
 		TEXT("vr.AllowMotionBlurInVR"),
 		0,
@@ -79,7 +88,8 @@ DECLARE_GPU_STAT(MotionBlur)
 
 }
 
-const int32 kMotionBlurTileSize = 16;
+const int32 kMotionBlurFlattenTileSize = 16;
+const int32 kMotionBlurFilterTileSize = 16;
 const int32 kMotionBlurComputeTileSizeX = 8;
 const int32 kMotionBlurComputeTileSizeY = 8;
 
@@ -137,11 +147,16 @@ bool IsMotionBlurScatterRequired(const FViewInfo& View, const FScreenPassTexture
 
 FIntPoint GetMotionBlurTileCount(FIntPoint SizeInPixels)
 {
-	const uint32 TilesX = FMath::DivideAndRoundUp(SizeInPixels.X, kMotionBlurTileSize);
-	const uint32 TilesY = FMath::DivideAndRoundUp(SizeInPixels.Y, kMotionBlurTileSize);
+	const uint32 TilesX = FMath::DivideAndRoundUp(SizeInPixels.X, kMotionBlurFlattenTileSize);
+	const uint32 TilesY = FMath::DivideAndRoundUp(SizeInPixels.Y, kMotionBlurFlattenTileSize);
 	return FIntPoint(TilesX, TilesY);
 }
 
+bool DoesMotionBlurNeedsHalfResInput()
+{
+	return CVarMotionBlurHalfResInput.GetValueOnRenderThread() != 0;
+}
+
 EMotionBlurQuality GetMotionBlurQuality()
 {
 	// Quality levels begin at 1. 0 is reserved for 'off'.
@@ -249,38 +264,6 @@ FVecocityTileUAVs CreateUAVs(FRDGBuilder& GraphBuilder, FVecocityTileTextures Ve
 	return UAVs;
 }
 
-// Set of common shader parameters shared by all motion blur shaders.
-BEGIN_SHADER_PARAMETER_STRUCT(FMotionBlurParameters, )
-	SHADER_PARAMETER(float, AspectRatio)
-	SHADER_PARAMETER(float, VelocityScale)
-	SHADER_PARAMETER(float, VelocityScaleForTiles)
-	SHADER_PARAMETER(float, VelocityMax)
-END_SHADER_PARAMETER_STRUCT()
-
-FMotionBlurParameters GetMotionBlurParameters(const FViewInfo& View, FIntPoint SceneViewportSize, float BlurScale)
-{
-	const FSceneViewState* ViewState = View.ViewState;
-
-	const float TileSize = kMotionBlurTileSize;
-	const float SceneViewportSizeX = SceneViewportSize.X;
-	const float SceneViewportSizeY = SceneViewportSize.Y;
-	const float MotionBlurTimeScale = ViewState ? ViewState->MotionBlurTimeScale : 1.0f;
-
-	// Scale by 0.5 due to blur samples going both ways.
-	const float VelocityScale = MotionBlurTimeScale * View.FinalPostProcessSettings.MotionBlurAmount * 0.5f;
-	const float VelocityUVToPixel = BlurScale * SceneViewportSizeX * 0.5f;
-
-	// 0:no 1:full screen width, percent conversion
-	const float UVVelocityMax = View.FinalPostProcessSettings.MotionBlurMax / 100.0f;
-
-	FMotionBlurParameters MotionBlurParameters;
-	MotionBlurParameters.AspectRatio = SceneViewportSizeY / SceneViewportSizeX;
-	MotionBlurParameters.VelocityScale = VelocityUVToPixel * VelocityScale;
-	MotionBlurParameters.VelocityScaleForTiles = MotionBlurParameters.VelocityScale / TileSize;
-	MotionBlurParameters.VelocityMax = FMath::Abs(VelocityUVToPixel) * UVVelocityMax;
-	return MotionBlurParameters;
-}
-
 // Base class for a motion blur / velocity shader.
 class FMotionBlurShader : public FGlobalShader
 {
@@ -311,9 +294,10 @@ public:
 
 	BEGIN_SHADER_PARAMETER_STRUCT(FParameters, )
 		SHADER_PARAMETER_STRUCT_REF(FViewUniformShaderParameters, View)
-		SHADER_PARAMETER_STRUCT(FMotionBlurParameters, MotionBlur)
 		SHADER_PARAMETER_STRUCT(FScreenPassTextureViewportParameters, Velocity)
 		SHADER_PARAMETER(FMatrix44f, ClipToPrevClipOverride)
+		SHADER_PARAMETER(FVector2f, VelocityScale)
+		SHADER_PARAMETER(float, VelocityMax)
 
 		SHADER_PARAMETER_RDG_TEXTURE(Texture2D, VelocityTexture)
 		SHADER_PARAMETER_RDG_TEXTURE(Texture2D, DepthTexture)
@@ -325,9 +309,9 @@ public:
 };
 
 BEGIN_SHADER_PARAMETER_STRUCT(FMotionBlurVelocityDilateParameters, )
-	SHADER_PARAMETER_STRUCT(FMotionBlurParameters, MotionBlur)
 	SHADER_PARAMETER_STRUCT(FScreenPassTextureViewportParameters, VelocityTile)
 	SHADER_PARAMETER_STRUCT(FVecocityTileTextures, VelocityTileTextures)
+	SHADER_PARAMETER(float, VelocityScaleForFlattenTiles)
 END_SHADER_PARAMETER_STRUCT()
 
 class FMotionBlurVelocityDilateGatherCS : public FMotionBlurShader
@@ -355,18 +339,16 @@ enum class EMotionBlurVelocityScatterPass : uint32
 	MAX
 };
 
-BEGIN_SHADER_PARAMETER_STRUCT(FMotionBlurVelocityDilateScatterParameters, )
-	SHADER_PARAMETER(uint32, ScatterPass)
-	SHADER_PARAMETER_STRUCT_INCLUDE(FMotionBlurVelocityDilateParameters, Dilate)
-	RENDER_TARGET_BINDING_SLOTS()
-END_SHADER_PARAMETER_STRUCT()
-
 class FMotionBlurVelocityDilateScatterVS : public FMotionBlurShader
 {
 	DECLARE_GLOBAL_SHADER(FMotionBlurVelocityDilateScatterVS);
 	SHADER_USE_PARAMETER_STRUCT(FMotionBlurVelocityDilateScatterVS, FMotionBlurShader);
-
-	using FParameters = FMotionBlurVelocityDilateScatterParameters;
+	
+	BEGIN_SHADER_PARAMETER_STRUCT(FParameters, )
+		SHADER_PARAMETER(uint32, ScatterPass)
+		SHADER_PARAMETER_STRUCT_INCLUDE(FMotionBlurVelocityDilateParameters, Dilate)
+		RENDER_TARGET_BINDING_SLOTS()
+	END_SHADER_PARAMETER_STRUCT()
 };
 
 class FMotionBlurVelocityDilateScatterPS : public FMotionBlurShader
@@ -374,63 +356,132 @@ class FMotionBlurVelocityDilateScatterPS : public FMotionBlurShader
 	DECLARE_GLOBAL_SHADER(FMotionBlurVelocityDilateScatterPS);
 	SHADER_USE_PARAMETER_STRUCT(FMotionBlurVelocityDilateScatterPS, FMotionBlurShader);
 
-	using FParameters = FMotionBlurVelocityDilateScatterParameters;
+	using FParameters = FEmptyShaderParameters;
 };
 
-class FMotionBlurQualityDimension : SHADER_PERMUTATION_ENUM_CLASS("MOTION_BLUR_QUALITY", EMotionBlurQuality);
-class FPostMotionBlurTranslucencyDimension : SHADER_PERMUTATION_BOOL("USE_POST_MOTION_BLUR_TRANSLUCENCY");
-
-using FMotionBlurFilterPermutationDomain = TShaderPermutationDomain<FMotionBlurQualityDimension, FPostMotionBlurTranslucencyDimension, FMotionBlurDirections>;
-
-BEGIN_SHADER_PARAMETER_STRUCT(FMotionBlurFilterParameters, )
-	SHADER_PARAMETER_STRUCT(FMotionBlurParameters, MotionBlur)
+class FMotionBlurFilterTileClassifyCS : public FMotionBlurShader
+{
+	DECLARE_GLOBAL_SHADER(FMotionBlurFilterTileClassifyCS);
+	SHADER_USE_PARAMETER_STRUCT(FMotionBlurFilterTileClassifyCS, FMotionBlurShader);
 
-	SHADER_PARAMETER_STRUCT(FScreenPassTextureViewportParameters, Color)
-	SHADER_PARAMETER_STRUCT(FScreenPassTextureViewportParameters, Velocity)
-	SHADER_PARAMETER_STRUCT(FScreenPassTextureViewportParameters, VelocityTile)
+	using FPermutationDomain = TShaderPermutationDomain<FMotionBlurDirections>;
 
-	SHADER_PARAMETER(FScreenTransform, ColorToVelocity)
+	BEGIN_SHADER_PARAMETER_STRUCT(FParameters, )
+		SHADER_PARAMETER(FScreenTransform, FilterTileIdToFlattenTileId)
+		SHADER_PARAMETER(FIntPoint, FlattenTileMaxId)
+		SHADER_PARAMETER(FIntPoint, FilterTileCount)
+		SHADER_PARAMETER(int32, bAllowHalfResGather)
+		SHADER_PARAMETER(int32, TileListMaxSize)
+		SHADER_PARAMETER(float, HalfResPixelVelocityThreshold)
+		SHADER_PARAMETER_STRUCT(FVecocityTileTextures, VelocityTileTextures)
+		SHADER_PARAMETER_RDG_BUFFER_UAV(RWStructuredBuffer<uint>, TileListsOutput)
+		SHADER_PARAMETER_RDG_BUFFER_UAV(RWStructuredBuffer<uint>, TileListsSizeOutput)
+		SHADER_PARAMETER_RDG_TEXTURE_UAV(RWTexture2D<float4>, DebugOutput)
+	END_SHADER_PARAMETER_STRUCT()
+};
 
-	SHADER_PARAMETER(FVector2f, ScreenPosToPostMotionBlurTranslucencyUV)
-	SHADER_PARAMETER(FVector2f, PostMotionBlurTranslucencyUVMax)
+class FSetupMotionBlurFilterCS : public FMotionBlurShader
+{
+	DECLARE_GLOBAL_SHADER(FSetupMotionBlurFilterCS);
+	SHADER_USE_PARAMETER_STRUCT(FSetupMotionBlurFilterCS, FMotionBlurShader);
 
-	SHADER_PARAMETER_RDG_TEXTURE(Texture2D, ColorTexture)
-	SHADER_PARAMETER_RDG_TEXTURE(Texture2D, VelocityFlatTexture)
-	SHADER_PARAMETER_STRUCT(FVecocityTileTextures, VelocityTileTextures)
-	SHADER_PARAMETER_RDG_TEXTURE(Texture2D, PostMotionBlurTranslucencyTexture)
+	BEGIN_SHADER_PARAMETER_STRUCT(FParameters, )
+		SHADER_PARAMETER_RDG_BUFFER_SRV(RWStructuredBuffer<uint>, TileListsSizeBuffer)
+		SHADER_PARAMETER_RDG_BUFFER_UAV(RWStructuredBuffer<uint>, DispatchParametersOutput)
+	END_SHADER_PARAMETER_STRUCT()
 
-	SHADER_PARAMETER_SAMPLER(SamplerState, ColorSampler)
-	SHADER_PARAMETER_SAMPLER(SamplerState, VelocitySampler)
-	SHADER_PARAMETER_SAMPLER(SamplerState, VelocityTileSampler)
-	SHADER_PARAMETER_SAMPLER(SamplerState, VelocityFlatSampler)
-	SHADER_PARAMETER_SAMPLER(SamplerState, PostMotionBlurTranslucencySampler)
-END_SHADER_PARAMETER_STRUCT()
+	static void ModifyCompilationEnvironment(const FGlobalShaderPermutationParameters& Parameters, FShaderCompilerEnvironment& OutEnvironment)
+	{
+		OutEnvironment.SetDefine(TEXT("DIM_BLUR_DIRECTIONS"), 1);
+		OutEnvironment.SetDefine(TEXT("SETUP_PASS"), 1);
+	}
+};
 
 class FMotionBlurFilterCS : public FMotionBlurShader
 {
-public:
 	DECLARE_GLOBAL_SHADER(FMotionBlurFilterCS);
 	SHADER_USE_PARAMETER_STRUCT(FMotionBlurFilterCS, FMotionBlurShader);
 
+	enum class ETileClassification
+	{
+		GatherHalfRes,
+		GatherFullRes,
+		ScatterAsGatherOneVelocityHalfRes,
+		ScatterAsGatherOneVelocityFullRes,
+		ScatterAsGatherTwoVelocityFullRes,
+		MAX
+	};
+
+	static bool IsHalfResTileClassification(ETileClassification TileClassification)
+	{
+		return TileClassification == ETileClassification::GatherHalfRes || TileClassification == ETileClassification::ScatterAsGatherOneVelocityHalfRes;
+	}
+	
+	class FTileClassificationDim : SHADER_PERMUTATION_ENUM_CLASS("DIM_TILE_CLASSIFICATION", ETileClassification);
+
+	using FPermutationDomain = TShaderPermutationDomain<FMotionBlurDirections, FTileClassificationDim>;
+
 	BEGIN_SHADER_PARAMETER_STRUCT(FParameters, )
-		SHADER_PARAMETER_STRUCT_INCLUDE(FMotionBlurFilterParameters, Filter)
+		SHADER_PARAMETER_STRUCT(FScreenPassTextureViewportParameters, Color)
+		SHADER_PARAMETER_STRUCT(FScreenPassTextureViewportParameters, Velocity)
+		SHADER_PARAMETER_STRUCT(FScreenPassTextureViewportParameters, VelocityTile)
+
+		SHADER_PARAMETER(FScreenTransform, ColorToVelocity)
+		SHADER_PARAMETER(int32, MaxSampleCount)
+
+		SHADER_PARAMETER_RDG_TEXTURE(Texture2D, ColorTexture)
+		SHADER_PARAMETER_RDG_TEXTURE(Texture2D, VelocityFlatTexture)
+		SHADER_PARAMETER_STRUCT(FVecocityTileTextures, VelocityTileTextures)
+
+		SHADER_PARAMETER_SAMPLER(SamplerState, ColorSampler)
+		SHADER_PARAMETER_SAMPLER(SamplerState, VelocitySampler)
+		SHADER_PARAMETER_SAMPLER(SamplerState, VelocityTileSampler)
+		SHADER_PARAMETER_SAMPLER(SamplerState, VelocityFlatSampler)
+
+		SHADER_PARAMETER_RDG_TEXTURE(Texture2D, TranslucencyTexture)
+		SHADER_PARAMETER_SAMPLER(SamplerState, TranslucencySampler)
+		SHADER_PARAMETER(FScreenTransform, ColorToTranslucency)
+		SHADER_PARAMETER(FVector2f, TranslucencyUVMin)
+		SHADER_PARAMETER(FVector2f, TranslucencyUVMax)
+		SHADER_PARAMETER(FVector2f, TranslucencyExtentInverse)
+
+		SHADER_PARAMETER(int32, TileListOffset)
+		SHADER_PARAMETER_RDG_BUFFER_SRV(RWStructuredBuffer<uint>, TileListsBuffer)
+		SHADER_PARAMETER_RDG_BUFFER_SRV(RWStructuredBuffer<uint>, TileListsSizeBuffer)
+		RDG_BUFFER_ACCESS(DispatchParameters, ERHIAccess::IndirectArgs)
+
 		SHADER_PARAMETER_RDG_TEXTURE_UAV(RWTexture2D<float4>, OutColorTexture)
 		SHADER_PARAMETER_RDG_TEXTURE_UAV(RWTexture2D<float4>, DebugOutput)
 	END_SHADER_PARAMETER_STRUCT()
 
-	static void ModifyCompilationEnvironment(const FGlobalShaderPermutationParameters& Parameters, FShaderCompilerEnvironment& OutEnvironment)
+	static FPermutationDomain RemapPermutation(FPermutationDomain PermutationVector)
 	{
-		FGlobalShader::ModifyCompilationEnvironment(Parameters, OutEnvironment);
-		OutEnvironment.SetDefine(TEXT("THREADGROUP_SIZEX"), kMotionBlurComputeTileSizeX);
-		OutEnvironment.SetDefine(TEXT("THREADGROUP_SIZEY"), kMotionBlurComputeTileSizeY);
+		if (PermutationVector.Get<FTileClassificationDim>() == ETileClassification::ScatterAsGatherTwoVelocityFullRes)
+		{
+			PermutationVector.Set<FMotionBlurDirections>(2);
+		}
+		else
+		{
+			PermutationVector.Set<FMotionBlurDirections>(1);
+		}
+
+		return PermutationVector;
+	}
+
+	static bool ShouldCompilePermutation(const FGlobalShaderPermutationParameters& Parameters)
+	{
+		FPermutationDomain PermutationVector(Parameters.PermutationId);
+		if (RemapPermutation(PermutationVector) != PermutationVector)
+		{
+			return false;
+		}
+		return FMotionBlurShader::ShouldCompilePermutation(Parameters);
 	}
 
-	using FPermutationDomain = FMotionBlurFilterPermutationDomain;
 };
 
 class FMotionBlurVisualizePS : public FMotionBlurShader
 {
-public:
 	DECLARE_GLOBAL_SHADER(FMotionBlurVisualizePS);
 	SHADER_USE_PARAMETER_STRUCT(FMotionBlurVisualizePS, FMotionBlurShader);
 
@@ -452,12 +503,14 @@ public:
 	END_SHADER_PARAMETER_STRUCT()
 };
 
-IMPLEMENT_GLOBAL_SHADER(FMotionBlurVelocityFlattenCS,       "/Engine/Private/MotionBlur/MotionBlurVelocityFlatten.usf", "MainCS", SF_Compute);
-IMPLEMENT_GLOBAL_SHADER(FMotionBlurVelocityDilateGatherCS,  "/Engine/Private/MotionBlur/MotionBlurTileGather.usf",      "MainCS", SF_Compute);
-IMPLEMENT_GLOBAL_SHADER(FMotionBlurVelocityDilateScatterVS, "/Engine/Private/MotionBlur/MotionBlurTileScatter.usf",     "MainVS", SF_Vertex);
-IMPLEMENT_GLOBAL_SHADER(FMotionBlurVelocityDilateScatterPS, "/Engine/Private/MotionBlur/MotionBlurTileScatter.usf",     "MainPS", SF_Pixel);
-IMPLEMENT_GLOBAL_SHADER(FMotionBlurFilterCS,                "/Engine/Private/MotionBlur/MotionBlurApply.usf",           "MainCS", SF_Compute);
-IMPLEMENT_GLOBAL_SHADER(FMotionBlurVisualizePS,             "/Engine/Private/MotionBlur/MotionBlurVisualize.usf",       "MainPS", SF_Pixel);
+IMPLEMENT_GLOBAL_SHADER(FMotionBlurVelocityFlattenCS,       "/Engine/Private/MotionBlur/MotionBlurVelocityFlatten.usf",    "MainCS", SF_Compute);
+IMPLEMENT_GLOBAL_SHADER(FMotionBlurVelocityDilateGatherCS,  "/Engine/Private/MotionBlur/MotionBlurTileGather.usf",         "MainCS", SF_Compute);
+IMPLEMENT_GLOBAL_SHADER(FMotionBlurVelocityDilateScatterVS, "/Engine/Private/MotionBlur/MotionBlurTileScatter.usf",        "MainVS", SF_Vertex);
+IMPLEMENT_GLOBAL_SHADER(FMotionBlurVelocityDilateScatterPS, "/Engine/Private/MotionBlur/MotionBlurTileScatter.usf",        "MainPS", SF_Pixel);
+IMPLEMENT_GLOBAL_SHADER(FMotionBlurFilterTileClassifyCS,    "/Engine/Private/MotionBlur/MotionBlurFilterTileClassify.usf", "MainCS", SF_Compute);
+IMPLEMENT_GLOBAL_SHADER(FSetupMotionBlurFilterCS,           "/Engine/Private/MotionBlur/MotionBlurFilterTileClassify.usf", "MainCS", SF_Compute);
+IMPLEMENT_GLOBAL_SHADER(FMotionBlurFilterCS,                "/Engine/Private/MotionBlur/MotionBlurApply.usf",              "MainCS", SF_Compute);
+IMPLEMENT_GLOBAL_SHADER(FMotionBlurVisualizePS,             "/Engine/Private/MotionBlur/MotionBlurVisualize.usf",          "MainPS", SF_Pixel);
 
 TGlobalResource<FSpriteIndexBuffer<8>> GScatterQuadIndexBuffer;
 
@@ -533,15 +586,32 @@ void AddMotionBlurVelocityPass(
 	FVecocityTileTextures VelocityTileTexturesSetup = CreateVecocityTileTextures(
 		GraphBuilder, VelocityTileCount, TEXT("MotionBlur.VelocityTile"));
 
-	const FMotionBlurParameters MotionBlurParametersNoScale = GetMotionBlurParameters(View, Viewports.Color.Rect.Size(), 1.0f);
-
 	// Velocity flatten pass: combines depth / velocity into a single target for sampling efficiency.
 	{
 		const bool bEnableCameraMotionBlur = View.bCameraMotionBlur.Get(true);
 		const bool bOverrideCameraMotionBlur = View.ClipToPrevClipOverride.IsSet();
 
 		FMotionBlurVelocityFlattenCS::FParameters* PassParameters = GraphBuilder.AllocParameters<FMotionBlurVelocityFlattenCS::FParameters>();
-		PassParameters->MotionBlur = MotionBlurParametersNoScale;
+
+		{
+			const FSceneViewState* ViewState = View.ViewState;
+
+			const float TileSize = kMotionBlurFlattenTileSize;
+			const float SceneViewportSizeX = Viewports.Color.Rect.Width();
+			const float SceneViewportSizeY = Viewports.Color.Rect.Height();
+			const float MotionBlurTimeScale = ViewState ? ViewState->MotionBlurTimeScale : 1.0f;
+
+			// Scale by 0.5 due to blur samples going both ways.
+			const float VelocityScale = MotionBlurTimeScale * View.FinalPostProcessSettings.MotionBlurAmount * 0.5f;
+
+			// 0:no 1:full screen width, percent conversion
+			const float UVVelocityMax = View.FinalPostProcessSettings.MotionBlurMax / 100.0f;
+
+			PassParameters->VelocityScale.X = SceneViewportSizeX * 0.5f * VelocityScale;
+			PassParameters->VelocityScale.Y = - SceneViewportSizeY * 0.5f * VelocityScale;
+			PassParameters->VelocityMax = SceneViewportSizeX * 0.5f * UVVelocityMax;
+		}
+
 		PassParameters->View = View.ViewUniformBuffer;
 		if (bOverrideCameraMotionBlur)
 		{
@@ -563,7 +633,9 @@ void AddMotionBlurVelocityPass(
 		TShaderMapRef<FMotionBlurVelocityFlattenCS> ComputeShader(View.ShaderMap, PermutationVector);
 		FComputeShaderUtils::AddPass(
 			GraphBuilder,
-			RDG_EVENT_NAME("Velocity Flatten(CameraMotionBlur%s)", bEnableCameraMotionBlur ? (bOverrideCameraMotionBlur ? TEXT("Override") : TEXT("On")) : TEXT("Off")),
+			RDG_EVENT_NAME("Velocity Flatten(CameraMotionBlur%s) %dx%d",
+				bEnableCameraMotionBlur ? (bOverrideCameraMotionBlur ? TEXT("Override") : TEXT("On")) : TEXT("Off"),
+				Viewports.Velocity.Rect.Width(), Viewports.Velocity.Rect.Height()),
 			ComputeShader,
 			PassParameters,
 			FComputeShaderUtils::GetGroupCount(Viewports.Velocity.Rect.Size(), FMotionBlurVelocityFlattenCS::ThreadGroupSize));
@@ -572,9 +644,9 @@ void AddMotionBlurVelocityPass(
 	bool ScatterDilatation = IsMotionBlurScatterRequired(View, Viewports.Color);
 
 	FMotionBlurVelocityDilateParameters VelocityDilateParameters;
-	VelocityDilateParameters.MotionBlur = MotionBlurParametersNoScale;
 	VelocityDilateParameters.VelocityTile = Viewports.VelocityTileParameters;
 	VelocityDilateParameters.VelocityTileTextures = VelocityTileTexturesSetup;
+	VelocityDilateParameters.VelocityScaleForFlattenTiles = (1.0f / float(kMotionBlurFlattenTileSize)) * (float(Viewports.Velocity.Rect.Width()) / float(Viewports.Color.Rect.Width()));
 
 	FVecocityTileTextures VelocityTileTextures;
 	if (ScatterDilatation)
@@ -591,7 +663,7 @@ void AddMotionBlurVelocityPass(
 					TexCreate_DepthStencilTargetable),
 				TEXT("MotionBlur.DilatedVelocityDepth"));
 
-		FMotionBlurVelocityDilateScatterParameters* PassParameters = GraphBuilder.AllocParameters<FMotionBlurVelocityDilateScatterParameters>();
+		FMotionBlurVelocityDilateScatterVS::FParameters* PassParameters = GraphBuilder.AllocParameters<FMotionBlurVelocityDilateScatterVS::FParameters>();
 		PassParameters->Dilate = VelocityDilateParameters;
 
 		PassParameters->RenderTargets.DepthStencil =
@@ -609,7 +681,7 @@ void AddMotionBlurVelocityPass(
 		TShaderMapRef<FMotionBlurVelocityDilateScatterVS> VertexShader(View.ShaderMap);
 		TShaderMapRef<FMotionBlurVelocityDilateScatterPS> PixelShader(View.ShaderMap);
 		
-		ClearUnusedGraphResources(VertexShader, PixelShader, PassParameters);
+		ClearUnusedGraphResources(VertexShader, PassParameters);
 
 		GraphBuilder.AddPass(
 			RDG_EVENT_NAME("VelocityTileScatter %dx%d", VelocityTileCount.X, VelocityTileCount.Y),
@@ -759,65 +831,179 @@ FRDGTextureRef AddMotionBlurFilterPass(
 		}
 	};
 
+	static int32 kMaxSampleCountPerQuality[] = { 4, 8, 12, 16 };
+	static_assert(UE_ARRAY_COUNT(kMaxSampleCountPerQuality) == int32(EMotionBlurQuality::MAX), "Fix me!");
+
+	const int32 MaxSampleCount = kMaxSampleCountPerQuality[int32(MotionBlurQuality)];
+
 	const float BlurScale = BlurScaleLUT[static_cast<uint32>(MotionBlurFilterPass)][static_cast<uint32>(MotionBlurQuality)];
 
-	FMotionBlurFilterParameters MotionBlurFilterParameters;
-	MotionBlurFilterParameters.MotionBlur = GetMotionBlurParameters(View, Viewports.Color.Rect.Size(), BlurScale);
-	MotionBlurFilterParameters.Color = Viewports.ColorParameters;
-	MotionBlurFilterParameters.Velocity = Viewports.VelocityParameters;
-	MotionBlurFilterParameters.VelocityTile = Viewports.VelocityTileParameters;
-	MotionBlurFilterParameters.ColorToVelocity = Viewports.ColorToVelocityTransform;
-	MotionBlurFilterParameters.ColorTexture = ColorTexture;
-	MotionBlurFilterParameters.VelocityFlatTexture = VelocityFlatTexture;
-	MotionBlurFilterParameters.VelocityTileTextures = VelocityTileTextures;
-	MotionBlurFilterParameters.ColorSampler = GetMotionBlurColorSampler();
-	MotionBlurFilterParameters.VelocitySampler = GetMotionBlurVelocitySampler();
-	MotionBlurFilterParameters.VelocityTileSampler = GetMotionBlurVelocitySampler();
-	MotionBlurFilterParameters.VelocityFlatSampler = GetMotionBlurVelocitySampler();
-
-	const bool bPostMotionBlurTranslucency = (PostMotionBlurTranslucency != nullptr);
-	if (bPostMotionBlurTranslucency)
+	const FIntVector FilterTileCount = FComputeShaderUtils::GetGroupCount(Viewports.Color.Rect.Size(), kMotionBlurFilterTileSize);
+
+	int32 TileListMaxSize = FilterTileCount.X * FilterTileCount.Y;
+
+	// Tile classify the filtering
+	FRDGBufferRef TileListsBuffer;
+	FRDGBufferRef TileListsSizeBuffer;
+	{
+
+		TileListsBuffer = GraphBuilder.CreateBuffer(
+			FRDGBufferDesc::CreateStructuredDesc(sizeof(uint32), TileListMaxSize * int32(FMotionBlurFilterCS::ETileClassification::MAX)),
+			TEXT("MotionBlur.TileOffsets"));
+
+		TileListsSizeBuffer = GraphBuilder.CreateBuffer(
+			FRDGBufferDesc::CreateStructuredDesc(sizeof(uint32), int32(FMotionBlurFilterCS::ETileClassification::MAX)),
+			TEXT("MotionBlur.TileCounters"));
+
+		AddClearUAVPass(GraphBuilder, GraphBuilder.CreateUAV(TileListsSizeBuffer), 0);
+
+		FMotionBlurFilterTileClassifyCS::FParameters* PassParameters = GraphBuilder.AllocParameters<FMotionBlurFilterTileClassifyCS::FParameters>();
+		PassParameters->FilterTileIdToFlattenTileId = (FScreenTransform::Identity + 0.5f) * (float(kMotionBlurFilterTileSize * Viewports.Velocity.Rect.Width()) / float(kMotionBlurFlattenTileSize * Viewports.Color.Rect.Width()));
+		PassParameters->FlattenTileMaxId = Viewports.VelocityTile.Rect.Size() - 1;
+		PassParameters->FilterTileCount = FIntPoint(FilterTileCount.X, FilterTileCount.Y);
+		PassParameters->TileListMaxSize = TileListMaxSize;
+		PassParameters->HalfResPixelVelocityThreshold = FMath::Square(float(MaxSampleCount));
+		PassParameters->bAllowHalfResGather = CVarMotionBlurHalfResGather.GetValueOnRenderThread() ? 1 : 0;
+		PassParameters->VelocityTileTextures = VelocityTileTextures;
+		PassParameters->TileListsOutput = GraphBuilder.CreateUAV(TileListsBuffer);
+		PassParameters->TileListsSizeOutput = GraphBuilder.CreateUAV(TileListsSizeBuffer);
+		PassParameters->DebugOutput = CreateDebugUAV(GraphBuilder, PassParameters->FilterTileCount, TEXT("Debug.MotionBlur.FilterTileClassify"));
+
+		FMotionBlurFilterTileClassifyCS::FPermutationDomain PermutationVector;
+		PermutationVector.Set<FMotionBlurDirections>(BlurDirections);
+
+		TShaderMapRef<FMotionBlurFilterTileClassifyCS> ComputeShader(View.ShaderMap, PermutationVector);
+		FComputeShaderUtils::AddPass(
+			GraphBuilder,
+			RDG_EVENT_NAME("MotionBlur FilterTileClassify %dx%d", FilterTileCount.X, FilterTileCount.Y),
+			ComputeShader,
+			PassParameters,
+			FComputeShaderUtils::GetGroupCount(FilterTileCount, 8));
+	}
+
+	// Setup the filter's dispatch parameters.
+	FRDGBufferRef DispatchParameters;
 	{
-		const bool bScaleTranslucency = Viewports.Color.Rect.Size() != PostMotionBlurTranslucencySize;
-		const FVector2f OutputSize(Viewports.Color.Rect.Size());
-		const FVector2f OutputSizeInv = FVector2f(1.0f, 1.0f) / OutputSize;
-		const FVector2f PostMotionBlurTranslucencyExtent(PostMotionBlurTranslucency->Desc.Extent);
-		const FVector2f PostMotionBlurTranslucencyExtentInv = FVector2f(1.0f, 1.0f) / PostMotionBlurTranslucencyExtent;
-
-		MotionBlurFilterParameters.PostMotionBlurTranslucencyTexture = PostMotionBlurTranslucency;
-		MotionBlurFilterParameters.PostMotionBlurTranslucencySampler = GetPostMotionBlurTranslucencySampler(bScaleTranslucency);
-		MotionBlurFilterParameters.ScreenPosToPostMotionBlurTranslucencyUV = OutputSizeInv * FVector2f(PostMotionBlurTranslucencySize) * PostMotionBlurTranslucencyExtentInv;
-		MotionBlurFilterParameters.PostMotionBlurTranslucencyUVMax = (FVector2f(PostMotionBlurTranslucencySize) - FVector2f(0.5f, 0.5f)) * PostMotionBlurTranslucencyExtentInv;
+		DispatchParameters = GraphBuilder.CreateBuffer(
+			FRDGBufferDesc::CreateIndirectDesc<FRHIDispatchIndirectParameters>(int32(FMotionBlurFilterCS::ETileClassification::MAX)),
+			TEXT("MotionBlur.FilterDispatchParameters"));
+
+		FSetupMotionBlurFilterCS::FParameters* PassParameters = GraphBuilder.AllocParameters<FSetupMotionBlurFilterCS::FParameters>();
+		PassParameters->TileListsSizeBuffer = GraphBuilder.CreateSRV(TileListsSizeBuffer);
+		PassParameters->DispatchParametersOutput = GraphBuilder.CreateUAV(DispatchParameters);
+
+		TShaderMapRef<FSetupMotionBlurFilterCS> ComputeShader(View.ShaderMap);
+		FComputeShaderUtils::AddPass(
+			GraphBuilder,
+			RDG_EVENT_NAME("MotionBlur SetupFilter"),
+			ComputeShader,
+			PassParameters,
+			FIntVector(1, 1, 1));
 	}
 
-	FRDGTextureDesc OutColorDesc = FRDGTextureDesc::Create2D(
-		ColorTexture->Desc.Extent,
-		IsPostProcessingWithAlphaChannelSupported() ? PF_FloatRGBA : PF_FloatRGB,
-		FClearValueBinding::None,
-		TexCreate_UAV | TexCreate_ShaderResource | GFastVRamConfig.MotionBlur);
+	// Apply the filter.
+	FRDGTextureRef ColorTextureOutput;
+	{
+		static const TCHAR* kTileClassificationNames[] = {
+			TEXT("GatherHalfRes"),
+			TEXT("GatherFullRes"),
+			TEXT("ScatterAsGatherOneVelocityHalfRes"),
+			TEXT("ScatterAsGatherOneVelocityFullRes"),
+			TEXT("ScatterAsGatherTwoVelocityFullRes"),
+		};
+		static_assert(UE_ARRAY_COUNT(kTileClassificationNames) == int32(FMotionBlurFilterCS::ETileClassification::MAX), "Fix me!");
+
+		FMotionBlurFilterCS::FParameters OriginalPassParameters;
+		OriginalPassParameters.Color = Viewports.ColorParameters;
+		OriginalPassParameters.Velocity = Viewports.VelocityParameters;
+		OriginalPassParameters.VelocityTile = Viewports.VelocityTileParameters;
+		OriginalPassParameters.ColorToVelocity = Viewports.ColorToVelocityTransform;
+		OriginalPassParameters.MaxSampleCount = MaxSampleCount;
+
+		OriginalPassParameters.ColorTexture = ColorTexture;
+		OriginalPassParameters.VelocityFlatTexture = VelocityFlatTexture;
+		OriginalPassParameters.VelocityTileTextures = VelocityTileTextures;
+		OriginalPassParameters.ColorSampler = GetMotionBlurColorSampler();
+		OriginalPassParameters.VelocitySampler = GetMotionBlurVelocitySampler();
+		OriginalPassParameters.VelocityTileSampler = GetMotionBlurVelocitySampler();
+		OriginalPassParameters.VelocityFlatSampler = GetMotionBlurVelocitySampler();
+
+		OriginalPassParameters.TileListsBuffer = GraphBuilder.CreateSRV(TileListsBuffer);
+		OriginalPassParameters.TileListsSizeBuffer = GraphBuilder.CreateSRV(TileListsSizeBuffer);
+		OriginalPassParameters.DispatchParameters = DispatchParameters;
+
+		if (PostMotionBlurTranslucency != nullptr)
+		{
+			// TODO: broken with split screen
+			const bool bScaleTranslucency = Viewports.Color.Rect.Size() != PostMotionBlurTranslucencySize;
+			const FVector2f OutputSize(Viewports.Color.Rect.Size());
+			const FVector2f OutputSizeInv = FVector2f(1.0f, 1.0f) / OutputSize;
+			const FVector2f PostMotionBlurTranslucencyExtent(PostMotionBlurTranslucency->Desc.Extent);
+			const FVector2f PostMotionBlurTranslucencyExtentInv = FVector2f(1.0f, 1.0f) / PostMotionBlurTranslucencyExtent;
+		
+			OriginalPassParameters.TranslucencyTexture = PostMotionBlurTranslucency;
+			OriginalPassParameters.TranslucencySampler = GetPostMotionBlurTranslucencySampler(bScaleTranslucency);
+			OriginalPassParameters.ColorToTranslucency = FScreenTransform::ChangeTextureUVCoordinateFromTo(
+				Viewports.Color,
+				FScreenPassTextureViewport(PostMotionBlurTranslucency->Desc.Extent, FIntRect(FIntPoint::ZeroValue, PostMotionBlurTranslucencySize)));
+			OriginalPassParameters.TranslucencyUVMin = FVector2f(0.0f, 0.0f);
+			OriginalPassParameters.TranslucencyUVMax = (FVector2f(PostMotionBlurTranslucencySize) - FVector2f(0.5f, 0.5f)) * PostMotionBlurTranslucencyExtentInv;
+			OriginalPassParameters.TranslucencyExtentInverse = PostMotionBlurTranslucencyExtentInv;
+		}
+		else
+		{
+			OriginalPassParameters.TranslucencyTexture = GSystemTextures.GetBlackAlphaOneDummy(GraphBuilder);
+			OriginalPassParameters.TranslucencySampler = TStaticSamplerState<SF_Point, AM_Clamp, AM_Clamp, AM_Clamp>::GetRHI();
+			OriginalPassParameters.ColorToTranslucency = FScreenTransform::Identity;
+			OriginalPassParameters.TranslucencyUVMin = FVector2f(0.0f, 0.0f);
+			OriginalPassParameters.TranslucencyUVMax = FVector2f(0.0f, 0.0f);
+			OriginalPassParameters.TranslucencyExtentInverse = FVector2f(0.0f, 0.0f);
+		}
 
-	FRDGTextureRef ColorTextureOutput = GraphBuilder.CreateTexture(OutColorDesc, TEXT("MotionBlur.SceneColor"));
+		{
+			FRDGTextureDesc OutColorDesc = FRDGTextureDesc::Create2D(
+				ColorTexture->Desc.Extent,
+				IsPostProcessingWithAlphaChannelSupported() ? PF_FloatRGBA : PF_FloatRGB,
+				FClearValueBinding::None,
+				TexCreate_UAV | TexCreate_ShaderResource | GFastVRamConfig.MotionBlur);
 
-	FMotionBlurFilterPermutationDomain PermutationVector;
-	PermutationVector.Set<FMotionBlurQualityDimension>(MotionBlurQuality);
-	PermutationVector.Set<FPostMotionBlurTranslucencyDimension>(PostMotionBlurTranslucency != nullptr);
-	PermutationVector.Set<FMotionBlurDirections>(BlurDirections);
+			ColorTextureOutput = GraphBuilder.CreateTexture(OutColorDesc, TEXT("MotionBlur.SceneColor"));
+		}
 
-	FMotionBlurFilterCS::FParameters* PassParameters = GraphBuilder.AllocParameters<FMotionBlurFilterCS::FParameters>();
-	PassParameters->Filter = MotionBlurFilterParameters;
-	PassParameters->OutColorTexture = GraphBuilder.CreateUAV(ColorTextureOutput);
-	PassParameters->DebugOutput = CreateDebugUAV(GraphBuilder, ColorTextureOutput->Desc.Extent, TEXT("Debug.MotionBlur.Apply"));
+		OriginalPassParameters.OutColorTexture = GraphBuilder.CreateUAV(ColorTextureOutput, ERDGUnorderedAccessViewFlags::SkipBarrier);
+		OriginalPassParameters.DebugOutput = CreateDebugUAV(GraphBuilder, ColorTextureOutput->Desc.Extent, TEXT("Debug.MotionBlur.Filter"));
 
-	TShaderMapRef<FMotionBlurFilterCS> ComputeShader(View.ShaderMap, PermutationVector);
+		RDG_EVENT_SCOPE(GraphBuilder, "MotionBlur FullResFilter(BlurDirections=%d MaxSamples=%d%s) %dx%d",
+			BlurDirections, OriginalPassParameters.MaxSampleCount,
+			PostMotionBlurTranslucency ? TEXT(" ComposeTranslucency") : TEXT(""),
+			Viewports.Color.Rect.Width(), Viewports.Color.Rect.Height());
 
-	FComputeShaderUtils::AddPass(
-		GraphBuilder,
-		RDG_EVENT_NAME("MotionBlur Apply(Quality=%d Directions=%d) %dx%d (CS)",
-			int32(MotionBlurQuality), BlurDirections,
-			Viewports.Color.Rect.Width(), Viewports.Color.Rect.Height()),
-		ComputeShader,
-		PassParameters,
-		FComputeShaderUtils::GetGroupCount(Viewports.Color.Rect.Size(), FComputeShaderUtils::kGolden2DGroupSize));
+		for (int32 TileClassifcation = 0; TileClassifcation < int32(FMotionBlurFilterCS::ETileClassification::MAX); TileClassifcation++)
+		{
+			FMotionBlurFilterCS::FParameters* PassParameters = GraphBuilder.AllocParameters<FMotionBlurFilterCS::FParameters>();
+			*PassParameters = OriginalPassParameters;
+			PassParameters->TileListOffset = TileListMaxSize * TileClassifcation;
+
+			FMotionBlurFilterCS::FPermutationDomain PermutationVector;
+			PermutationVector.Set<FMotionBlurFilterCS::FTileClassificationDim>(FMotionBlurFilterCS::ETileClassification(TileClassifcation));
+			PermutationVector.Set<FMotionBlurDirections>(BlurDirections);
+			PermutationVector = FMotionBlurFilterCS::RemapPermutation(PermutationVector);
+
+			if (PermutationVector.Get<FMotionBlurDirections>() > BlurDirections)
+			{
+				continue;
+			}
+
+			TShaderMapRef<FMotionBlurFilterCS> ComputeShader(View.ShaderMap, PermutationVector);
+			FComputeShaderUtils::AddPass(
+				GraphBuilder,
+				RDG_EVENT_NAME("MotionBlur Filter(%s)", kTileClassificationNames[int32(PermutationVector.Get<FMotionBlurFilterCS::FTileClassificationDim>())]),
+				ComputeShader,
+				PassParameters,
+				DispatchParameters,
+				/* IndirectArgsOffset = */ sizeof(FRHIDispatchIndirectParameters) * TileClassifcation);
+		}
+	}
 
 	return ColorTextureOutput;
 }
diff --git a/Engine/Source/Runtime/Renderer/Private/PostProcess/PostProcessMotionBlur.h b/Engine/Source/Runtime/Renderer/Private/PostProcess/PostProcessMotionBlur.h
--- a/Engine/Source/Runtime/Renderer/Private/PostProcess/PostProcessMotionBlur.h
+++ b/Engine/Source/Runtime/Renderer/Private/PostProcess/PostProcessMotionBlur.h
@@ -30,6 +30,9 @@ enum class EMotionBlurFilter : uint32
 	Separable
 };
 
+// Returns whether motion blur also want half res input the scene color's MIP1.
+bool DoesMotionBlurNeedsHalfResInput();
+
 // Returns the global setting for motion blur quality.
 EMotionBlurQuality GetMotionBlurQuality();
 
diff --git a/Engine/Source/Runtime/Renderer/Private/PostProcess/PostProcessing.cpp b/Engine/Source/Runtime/Renderer/Private/PostProcess/PostProcessing.cpp
--- a/Engine/Source/Runtime/Renderer/Private/PostProcess/PostProcessing.cpp
+++ b/Engine/Source/Runtime/Renderer/Private/PostProcess/PostProcessing.cpp
@@ -598,6 +598,8 @@ void AddPostProcessingPasses(FRDGBuilder& GraphBuilder, const FViewInfo& View, c
 
 		PassSequence.Finalize();
 
+		bool bMotionBlurNeedsHalfResInput = PassSequence.IsEnabled(EPass::MotionBlur) && DoesMotionBlurNeedsHalfResInput();
+
 		// Post Process Material Chain - Before Translucency
 		{
 			const FPostProcessMaterialChain MaterialChain = GetPostProcessMaterialChain(View, BL_BeforeTranslucency);
@@ -687,6 +689,7 @@ void AddPostProcessingPasses(FRDGBuilder& GraphBuilder, const FViewInfo& View, c
 			ITemporalUpscaler::FPassInputs UpscalerPassInputs;
 
 			UpscalerPassInputs.bAllowDownsampleSceneColor = bAllowSceneDownsample;
+			UpscalerPassInputs.bGenerateOutputMip1 = bMotionBlurNeedsHalfResInput;
 			UpscalerPassInputs.DownsampleOverrideFormat = DownsampleOverrideFormat;
 			UpscalerPassInputs.SceneColorTexture = SceneColor.Texture;
 			UpscalerPassInputs.SceneDepthTexture = SceneDepth.Texture;
diff --git a/Engine/Source/Runtime/Renderer/Private/PostProcess/TemporalAA.h b/Engine/Source/Runtime/Renderer/Private/PostProcess/TemporalAA.h
--- a/Engine/Source/Runtime/Renderer/Private/PostProcess/TemporalAA.h
+++ b/Engine/Source/Runtime/Renderer/Private/PostProcess/TemporalAA.h
@@ -177,7 +177,8 @@ public:
 
 	struct FPassInputs
 	{
-		bool bAllowDownsampleSceneColor;
+		bool bAllowDownsampleSceneColor = false;
+		bool bGenerateOutputMip1 = false;
 		EPixelFormat DownsampleOverrideFormat;
 		FRDGTextureRef SceneColorTexture = nullptr;
 		FRDGTextureRef SceneDepthTexture = nullptr;
diff --git a/Engine/Source/Runtime/Renderer/Private/PostProcess/TemporalSuperResolution.cpp b/Engine/Source/Runtime/Renderer/Private/PostProcess/TemporalSuperResolution.cpp
--- a/Engine/Source/Runtime/Renderer/Private/PostProcess/TemporalSuperResolution.cpp
+++ b/Engine/Source/Runtime/Renderer/Private/PostProcess/TemporalSuperResolution.cpp
@@ -521,12 +521,14 @@ class FTSRUpdateHistoryCS : public FTSRShader
 		SHADER_PARAMETER(float, InvWeightClampingPixelSpeed)
 		SHADER_PARAMETER(float, InputToHistoryFactor)
 		SHADER_PARAMETER(int32, ResponsiveStencilMask)
+		SHADER_PARAMETER(int32, bGenerateOutputMip1)
 
 		SHADER_PARAMETER_STRUCT_INCLUDE(FTSRPrevHistoryParameters, PrevHistoryParameters)
 		SHADER_PARAMETER_STRUCT(FTSRHistoryTextures, PrevHistory)
 
 		SHADER_PARAMETER_STRUCT(FTSRHistoryUAVs, HistoryOutput)
-		SHADER_PARAMETER_RDG_TEXTURE_UAV(RWTexture2D, SceneColorOutput)
+		SHADER_PARAMETER_RDG_TEXTURE_UAV(RWTexture2D, SceneColorOutputMip0)
+		SHADER_PARAMETER_RDG_TEXTURE_UAV(RWTexture2D, SceneColorOutputMip1)
 		SHADER_PARAMETER_RDG_TEXTURE_UAV(RWTexture2D, DebugOutput)
 	END_SHADER_PARAMETER_STRUCT()
 
@@ -1438,7 +1440,8 @@ void AddTemporalSuperResolutionPasses(
 				HistoryExtent,
 				PF_FloatR11G11B10,
 				FClearValueBinding::None,
-				/* InFlags = */ TexCreate_ShaderResource | TexCreate_UAV);
+				/* InFlags = */ TexCreate_ShaderResource | TexCreate_UAV,
+				/* NumMips = */ PassInputs.bGenerateOutputMip1 ? 2 : 1);
 
 			SceneColorOutputTexture = GraphBuilder.CreateTexture(Desc, TEXT("TSR.Output"));
 		}
@@ -1468,12 +1471,21 @@ void AddTemporalSuperResolutionPasses(
 		PassParameters->InvWeightClampingPixelSpeed = 1.0f / CVarTSRWeightClampingPixelSpeed.GetValueOnRenderThread();
 		PassParameters->InputToHistoryFactor = float(HistorySize.X) / float(InputRect.Width());
 		PassParameters->ResponsiveStencilMask = CVarTSREnableResponiveAA.GetValueOnRenderThread() ? (STENCIL_TEMPORAL_RESPONSIVE_AA_MASK) : 0;
+		PassParameters->bGenerateOutputMip1 = PassInputs.bGenerateOutputMip1 ? 1 : 0;
 
 		PassParameters->PrevHistoryParameters = PrevHistoryParameters;
 		PassParameters->PrevHistory = PrevHistory;
 
 		PassParameters->HistoryOutput = CreateUAVs(GraphBuilder, History);
-		PassParameters->SceneColorOutput = GraphBuilder.CreateUAV(SceneColorOutputTexture);
+		PassParameters->SceneColorOutputMip0 = GraphBuilder.CreateUAV(FRDGTextureUAVDesc(SceneColorOutputTexture, /* InMipLevel = */ 0));
+		if (PassInputs.bGenerateOutputMip1)
+		{
+			PassParameters->SceneColorOutputMip1 = GraphBuilder.CreateUAV(FRDGTextureUAVDesc(SceneColorOutputTexture, /* InMipLevel = */ 1));
+		}
+		else
+		{
+			PassParameters->SceneColorOutputMip1 = PassParameters->SceneColorOutputMip0;
+		}
 		PassParameters->DebugOutput = CreateDebugUAV(HistoryExtent, TEXT("Debug.TSR.UpdateHistory"));
 
 		FTSRUpdateHistoryCS::FPermutationDomain PermutationVector;
@@ -1484,10 +1496,11 @@ void AddTemporalSuperResolutionPasses(
 		TShaderMapRef<FTSRUpdateHistoryCS> ComputeShader(View.ShaderMap, PermutationVector);
 		FComputeShaderUtils::AddPass(
 			GraphBuilder,
-			RDG_EVENT_NAME("TSR UpdateHistory(%s%s%s) %dx%d", 
+			RDG_EVENT_NAME("TSR UpdateHistory(%s%s%s%s) %dx%d", 
 				History.LowFrequency->Desc.Format == PF_FloatR11G11B10 ? TEXT("R11G11B10") : TEXT(""),
 				PermutationVector.Get<FTSRUpdateHistoryCS::FRejectionAADim>() ? TEXT(" RejectionAA") : TEXT(""),
 				PermutationVector.Get<FTSRUpdateHistoryCS::FSeparateTranslucencyDim>() ? TEXT(" SeparateTranslucency") : TEXT(""),
+				PassInputs.bGenerateOutputMip1 ? TEXT(" OutputMip1") : TEXT(""),
 				HistorySize.X, HistorySize.Y),
 			ComputeShader,
 			PassParameters,
