Fixes r.MotionBlur.HalfResInput=1 overflow on the sun

diff --git a/Engine/Shaders/Private/TemporalSuperResolution/TSRUpdateHistory.usf b/Engine/Shaders/Private/TemporalSuperResolution/TSRUpdateHistory.usf
--- a/Engine/Shaders/Private/TemporalSuperResolution/TSRUpdateHistory.usf
+++ b/Engine/Shaders/Private/TemporalSuperResolution/TSRUpdateHistory.usf
@@ -1375,11 +1375,9 @@ void MainCS(
 		
 		// Output final scene color Mip1
 		{
-			taa_half3x2 HalfResOutput = OutputColor;
+			taa_half3x2 HalfResOutput = OutputColor * taa_half(0.25);
 			taa_short2x2 HalfResOutputPixelPos = (LocalHistoryPixelPos >> taa_short(1)) | (((LocalHistoryPixelPos & taa_short(0x1)) | !bGenerateOutputMip1) * taa_short(~0));
 			
-			//HalfResOutputPixelPos = (LocalHistoryPixelPos >> taa_short(1));
-			
 			#if PLATFORM_SUPPORTS_WAVE_BROADCAST && 0 // Support WaveBroadcast with halfs
 			BRANCH
 			if (bGenerateOutputMip1)
@@ -1388,7 +1386,7 @@ void MainCS(
 				FWaveBroadcastSettings Vertical = InitWaveXorButterfly(/* XorButterFly = */ 0x2);
 
 				HalfResOutput += WaveBroadcast(Horizontal, HalfResOutput);
-				HalfResOutput = HalfResOutput * taa_half(0.25) + taa_half(0.25) * WaveBroadcast(Vertical, HalfResOutput);
+				HalfResOutput += WaveBroadcast(Vertical, HalfResOutput);
 			}
 			#else
 			BRANCH
@@ -1410,9 +1408,9 @@ void MainCS(
 				#endif
 					
 				#if CONFIG_ENABLE_DUAL_PIXEL_VECTORIZATION
-					HalfResOutput = dpv_interleave_registers(SharedArray0[LocalGroupThreadIndex].rgb, SharedArray1[LocalGroupThreadIndex].rgb) * 0.25;
+					HalfResOutput = dpv_interleave_registers(SharedArray0[LocalGroupThreadIndex].rgb, SharedArray1[LocalGroupThreadIndex].rgb);
 				#else
-					HalfResOutput = dpv_interleave_mono_registers(SharedArray0[LocalGroupThreadIndex].rgb) * 0.25;
+					HalfResOutput = dpv_interleave_mono_registers(SharedArray0[LocalGroupThreadIndex].rgb);
 				#endif
 			}
 			#endif
