Fixes MB quarter res on XSS letter box

diff --git a/Engine/Shaders/Private/MotionBlur/MotionBlurApply.usf b/Engine/Shaders/Private/MotionBlur/MotionBlurApply.usf
--- a/Engine/Shaders/Private/MotionBlur/MotionBlurApply.usf
+++ b/Engine/Shaders/Private/MotionBlur/MotionBlurApply.usf
@@ -289,6 +289,7 @@ void MainCS(
 
 	// Map color UV to velocity UV space.
 	float2 VelocityUV = ApplyScreenTransform(ColorUV, ColorToVelocity);
+	VelocityUV = clamp(VelocityUV, Velocity_UVViewportBilinearMin, Velocity_UVViewportBilinearMax);
 
 	// Map velocity UV to velocity tile UV space with jitter.
 	float2 NearestVelocitySvPosition = floor(VelocityUV * Velocity_Extent) + 0.5;
@@ -715,10 +716,11 @@ void MainCS(
 	BRANCH
 	if (bIsHalfRes)
 	{
-		mb_short2 OutputPixelCoord0 = mb_short2(uint2(ColorUV * Color_Extent)) & mb_short(~0x1);
-		mb_short2 OutputPixelCoord1 = OutputPixelCoord0 + mb_short2(1, 0);
-		mb_short2 OutputPixelCoord2 = OutputPixelCoord0 + mb_short2(0, 1);
-		mb_short2 OutputPixelCoord3 = OutputPixelCoord0 + mb_short2(1, 1);
+		mb_short2 OutputPixelCoord = mb_short2(uint2(ColorUV * Color_Extent)) & mb_short(~0x1);
+		mb_short2 OutputPixelCoord0 = min(OutputPixelCoord + mb_short2(0, 0), mb_short2(Color_ViewportMax - 1));
+		mb_short2 OutputPixelCoord1 = min(OutputPixelCoord + mb_short2(1, 0), mb_short2(Color_ViewportMax - 1));
+		mb_short2 OutputPixelCoord2 = min(OutputPixelCoord + mb_short2(0, 1), mb_short2(Color_ViewportMax - 1));
+		mb_short2 OutputPixelCoord3 = min(OutputPixelCoord + mb_short2(1, 1), mb_short2(Color_ViewportMax - 1));
 
 		#if 1
 			float2 PostMotionBlurTranslucencyUV0 = ApplyScreenTransform((float2(uint2(OutputPixelCoord0)) + 0.5) * Color_ExtentInverse, ColorToTranslucency);
@@ -924,13 +926,13 @@ void MainCS(
 			QuarterResOutput = 0.0;
 		}
 		
-		bool bIsValid = all(OutputPixelCoord0 < mb_short2(Color_ViewportMax));
-		mb_short2 OutputPixelCoord0Mip0 = (bIsValid) ? OutputPixelCoord0 : mb_short(~0).xx;
+		bool bIsValid = all(OutputPixelCoord < mb_short2(Color_ViewportMax));
+		mb_short2 OutputPixelCoord0Mip0 = (bIsValid) ? OutputPixelCoord : mb_short(~0).xx;
 		mb_short2 OutputPixelCoord1Mip0 = (bIsValid) ? OutputPixelCoord1 : mb_short(~0).xx;
 		mb_short2 OutputPixelCoord2Mip0 = (bIsValid) ? OutputPixelCoord2 : mb_short(~0).xx;
 		mb_short2 OutputPixelCoord3Mip0 = (bIsValid) ? OutputPixelCoord3 : mb_short(~0).xx;
-		mb_short2 OutputPixelCoordMip1 = (bIsValid && (OutputMip1 != 0)) ? (OutputPixelCoord0 >> mb_short(1)) : mb_short(~0).xx;
-		mb_short2 OutputPixelCoordMip2 = (bIsValid && (OutputMip2 != 0) && (OutputPixelCoord0 & 0x3) == 0) ? (OutputPixelCoord0 >> mb_short(2)) : mb_short(~0).xx;
+		mb_short2 OutputPixelCoordMip1 = (bIsValid && (OutputMip1 != 0)) ? (OutputPixelCoord >> mb_short(1)) : mb_short(~0).xx;
+		mb_short2 OutputPixelCoordMip2 = (bIsValid && (OutputMip2 != 0) && (OutputPixelCoord & 0x3) == 0) ? (OutputPixelCoord >> mb_short(2)) : mb_short(~0).xx;
 
 		SceneColorOutputMip0[OutputPixelCoord0Mip0] = OutputColor0;
 		SceneColorOutputMip0[OutputPixelCoord1Mip0] = OutputColor1;
@@ -961,7 +963,7 @@ void MainCS(
 		BRANCH
 		if (bFetchCenter)
 		{
-			CenterColor = ColorTexture[OutputPixelCoord];
+			CenterColor = ColorTexture[min(OutputPixelCoord, mb_short2(Color_ViewportMax - 1))];
 			PostMotionBlurTranslucency = TranslucencyTexture.SampleLevel(TranslucencySampler, PostMotionBlurTranslucencyUV, 0);
 		}
 		else
