Fix invalid screen space ray cast bias for hair material

diff --git a/Engine/Shaders/Private/VirtualShadowMaps/VirtualShadowMapProjection.usf b/Engine/Shaders/Private/VirtualShadowMaps/VirtualShadowMapProjection.usf
--- a/Engine/Shaders/Private/VirtualShadowMaps/VirtualShadowMapProjection.usf
+++ b/Engine/Shaders/Private/VirtualShadowMaps/VirtualShadowMapProjection.usf
@@ -202,7 +202,7 @@ FVirtualShadowMapSampleResult ProjectLight(
 	BRANCH
 	if (bInLightRegion && bValidPixel)
 	{
-		const half3 WorldNormal = bIsHair ? L : GBufferNormal;
+		const half3 WorldNormal = (bIsHair || ShadingModelID == SHADINGMODELID_HAIR) ? L : GBufferNormal; 
 		TranslatedWorldPosition += WorldNormal * NormalBiasLength;
 
 		if (SMRTRayCount > 0)
