Shader merge fixes

diff --git a/Engine/Shaders/Private/LocalVertexFactory.ush b/Engine/Shaders/Private/LocalVertexFactory.ush
--- a/Engine/Shaders/Private/LocalVertexFactory.ush
+++ b/Engine/Shaders/Private/LocalVertexFactory.ush
@@ -591,12 +591,11 @@ FMaterialVertexParameters GetMaterialVertexParameters(FVertexFactoryInput Input,
 		Result.InstanceLocalToWorld = GetInstanceData(Intermediates).LocalToWorld;
 		Result.InstanceWorldToLocal = GetInstanceData(Intermediates).WorldToLocal;
 		Result.PerInstanceParams = Intermediates.PerInstanceParams;
-		Result.PrimitiveLocalInstanceIndex = Intermediates.PrimitiveLocalInstanceIndex;
 	#endif
-#if USES_PER_INSTANCE_CUSTOM_DATA
-	Result.CustomDataOffset = Intermediates.SceneData.InstanceData.CustomDataOffset;
-	Result.CustomDataCount = Intermediates.SceneData.InstanceData.CustomDataCount;
-#endif
+	#if USES_PER_INSTANCE_CUSTOM_DATA
+		Result.CustomDataOffset = Intermediates.SceneData.InstanceData.CustomDataOffset;
+		Result.CustomDataCount = Intermediates.SceneData.InstanceData.CustomDataCount;
+	#endif
 #endif
 
 	Result.PreSkinnedPosition = Intermediates.PreSkinPosition.xyz;
@@ -853,18 +852,6 @@ FVertexFactoryIntermediates GetVertexFactoryIntermediates(FVertexFactoryInput In
 	Intermediates.PrimitiveLocalInstanceIndex = Intermediates.SceneData.InstanceId - GetPrimitiveData(Intermediates).InstanceSceneDataOffset;
 
 #if USE_INSTANCE_CULLING
-	{
-	#if (VF_USE_PRIMITIVE_SCENE_DATA == 2)
-		#if USES_PER_INSTANCE_RANDOM
-		Intermediates.PerInstanceRandom = GetInstanceData(Intermediates).RandomID;
-		#endif
-		
-		float R = 0; // FIX ME: use SRV for this data?
-		float G = 0;
-		float B = 0;
-
-		Intermediates.InstanceLightmapAndShadowMapUVBias = GetInstanceData(Intermediates).LightMapAndShadowMapUVBias;
-	#else
 	Intermediates.InstanceLightmapAndShadowMapUVBias = Intermediates.SceneData.InstanceData.LightMapAndShadowMapUVBias;
 	Intermediates.PerInstanceRandom = Intermediates.SceneData.InstanceData.RandomID;
 #if USE_EDITOR_SHADERS
@@ -874,11 +861,6 @@ FVertexFactoryIntermediates GetVertexFactoryIntermediates(FVertexFactoryInput In
 	Intermediates.IsSelected = 0.0f;
 	Intermediates.HitProxyId = float4(0.0f, 0.0f, 0.0f, 0.0f);
 #endif // USE_EDITOR_SHADERS
-	#endif
-
-		Intermediates.IsSelected = trunc(R * (1.0 / 256.0));;
-		R = R - 256.0 * Intermediates.IsSelected;
-		Intermediates.HitProxyId = float4(R / 255.0, G / 255.0, B / 255.0, 0);
 
 	{
 		// This repackaging is stupid-seeming but these params are propably packed off to some interpolator or something
diff --git a/Engine/Shaders/Private/MaterialTemplate.ush b/Engine/Shaders/Private/MaterialTemplate.ush
--- a/Engine/Shaders/Private/MaterialTemplate.ush
+++ b/Engine/Shaders/Private/MaterialTemplate.ush
@@ -471,7 +471,6 @@ struct FMaterialVertexParameters
 	float3 WorldPosition;
 	// TangentToWorld[2] is WorldVertexNormal
 	half3x3 TangentToWorld;
-	half3x3 TangentToWorld;
 
 #if USES_PER_INSTANCE_CUSTOM_DATA
 	uint CustomDataOffset;
diff --git a/Engine/Shaders/Private/RequiredTextureResolutionPixelShader.usf b/Engine/Shaders/Private/RequiredTextureResolutionPixelShader.usf
--- a/Engine/Shaders/Private/RequiredTextureResolutionPixelShader.usf
+++ b/Engine/Shaders/Private/RequiredTextureResolutionPixelShader.usf
@@ -92,8 +92,8 @@ float3 AccuracyColorLookup(float Accuracy)
 {
 	Accuracy = clamp(Accuracy, -1.99, 1.99);
 	int ColorIndex = floor(Accuracy) + 2;
-	float3 Color0 = DebugViewModePass.AccuracyColors[ColorIndex].rgb;
-	float3 Color1 = DebugViewModePass.AccuracyColors[ColorIndex + 1].rgb;
+	float3 Color0 = DebugViewModeStruct.AccuracyColors[ColorIndex].rgb;
+	float3 Color1 = DebugViewModeStruct.AccuracyColors[ColorIndex + 1].rgb;
 	float ColorLerp = frac(Accuracy);
 	return lerp(Color0, Color1, ColorLerp);
 }
diff --git a/Engine/Shaders/Private/SceneData.ush b/Engine/Shaders/Private/SceneData.ush
--- a/Engine/Shaders/Private/SceneData.ush
+++ b/Engine/Shaders/Private/SceneData.ush
@@ -526,10 +526,13 @@ void ComputeInstanceDerivedData(inout FInstanceSceneData InstanceData, float3 Ti
 	RelativeWorldToLocal[3].xyz						        = mul(float4(-LocalToRelativeWorld[3].xyz, 0.0f), RelativeWorldToLocal).xyz;
 	InstanceData.WorldToLocal								= MakeLWCInverseMatrix(TilePosition, RelativeWorldToLocal);
 
+// See defines in FInstanceSceneData
+#if 0
 	InstanceData.HasRandomID								= (InstanceData.Flags & INSTANCE_SCENE_DATA_FLAG_HAS_RANDOM);
 	InstanceData.HasCustomData								= (InstanceData.Flags & INSTANCE_SCENE_DATA_FLAG_HAS_CUSTOM_DATA);
 	InstanceData.HasDynamicData								= (InstanceData.Flags & INSTANCE_SCENE_DATA_FLAG_HAS_DYNAMIC_DATA);
 	InstanceData.HasLightShadowUVBias						= (InstanceData.Flags & INSTANCE_SCENE_DATA_FLAG_HAS_LIGHTSHADOW_UV_BIAS);
+#endif
 }
 
 // Fetch from scene primitive buffer
@@ -589,17 +592,22 @@ FInstanceSceneData GetInstanceSceneData(uint InstanceId, uint SOAStride, bool bC
 		InstanceData.RandomID									= LoadInstanceSceneDataElement(0 * SOAStride + InstanceId).w;
 	#endif
 
+		float3 TilePosition 									= GetPrimitiveData(InstanceData.PrimitiveId).TilePosition;
+
 #if INSTANCE_COMPRESSED_TRANSFORMS
 		uint4  RotationScale									= asuint(LoadInstanceSceneDataElement(1 * SOAStride + InstanceId));
 		float3 Translation										= LoadInstanceSceneDataElement(2 * SOAStride + InstanceId).xyz;
 		float3 Scale = 0;
-		InstanceData.LocalToWorld								= DecodeTransform(RotationScale, Translation, Scale);
+		float4x4 LocalToRelativeWorld							= DecodeTransform( RotationScale, Translation, Scale );
+
+		uint4  PrevRotationScale								= asuint(LoadInstanceSceneDataElement(3 * SOAStride + InstanceId));
+		float3 PrevTranslation									= LoadInstanceSceneDataElement(4 * SOAStride + InstanceId).xyz;
+		float3 PrevScale = 0;
+		float4x4 PrevLocalToRelativeWorld						= DecodeTransform( PrevRotationScale, PrevTranslation, PrevScale );
 
 		InstanceData.NonUniformScale.xyz						= abs(Scale);
 		InstanceData.InvNonUniformScale							= rcp(InstanceData.NonUniformScale.xyz);
 #else
-		float3 TilePosition 								= GetPrimitiveData(InstanceData.PrimitiveId).TilePosition;	
-
 		float4x4 LocalToRelativeWorld							= transpose(float4x4(LoadInstanceSceneDataElement(1 * SOAStride + InstanceId),
 																			LoadInstanceSceneDataElement(2 * SOAStride + InstanceId),
 																	                LoadInstanceSceneDataElement(3 * SOAStride + InstanceId),
@@ -609,32 +617,13 @@ FInstanceSceneData GetInstanceSceneData(uint InstanceId, uint SOAStride, bool bC
 																	                LoadInstanceSceneDataElement(5 * SOAStride + InstanceId),
 																	                LoadInstanceSceneDataElement(6 * SOAStride + InstanceId),
 																	                float4(0.0f, 0.0f, 0.0f, 1.0f)));
-		
-		InstanceData.LocalToWorld							= MakeLWCMatrix(TilePosition, LocalToRelativeWorld);
-		InstanceData.PrevLocalToWorld							= MakeLWCMatrix(TilePosition, PrevLocalToRelativeWorld);		
-
-		ComputeInstanceDerivedData(InstanceData, TilePosition, LocalToRelativeWorld);
-
-		float3 Scale2;
-		Scale2.x = length2(InstanceData.LocalToWorld[0].xyz);
-		Scale2.y = length2(InstanceData.LocalToWorld[1].xyz);
-		Scale2.z = length2(InstanceData.LocalToWorld[2].xyz);
 
-		InstanceData.InvNonUniformScale  = rsqrt(Scale2);
-		InstanceData.NonUniformScale.xyz = Scale2 * InstanceData.InvNonUniformScale;
 #endif
 
-		InstanceData.DeterminantSign			= CondMask(InstanceData.Flags & INSTANCE_SCENE_DATA_FLAG_DETERMINANT_SIGN, -1.0f, 1.0f);	// Scale.x * Scale.y * Scale.z < 0.0 ? -1.0 : 1.0;
 
-		InstanceData.NonUniformScale.w			= max3(InstanceData.NonUniformScale.x, InstanceData.NonUniformScale.y, InstanceData.NonUniformScale.z);
-	
-		InstanceData.WorldToLocal				= InstanceData.LocalToWorld;
-		InstanceData.WorldToLocal[0].xyz	   *= Pow2(InstanceData.InvNonUniformScale.x);
-		InstanceData.WorldToLocal[1].xyz	   *= Pow2(InstanceData.InvNonUniformScale.y);
-		InstanceData.WorldToLocal[2].xyz	   *= Pow2(InstanceData.InvNonUniformScale.z);
-		InstanceData.WorldToLocal[3].xyz		= 0.0f;
-		InstanceData.WorldToLocal				= transpose(InstanceData.WorldToLocal);
-		InstanceData.WorldToLocal[3].xyz		= mul(float4(-InstanceData.LocalToWorld[3].xyz, 0.0f), InstanceData.WorldToLocal).xyz;
+		InstanceData.LocalToWorld	= MakeLWCMatrix(TilePosition, LocalToRelativeWorld);
+		InstanceData.PrevLocalToWorld	= MakeLWCMatrix(TilePosition, PrevLocalToRelativeWorld);
+		ComputeInstanceDerivedData(InstanceData, TilePosition, LocalToRelativeWorld);
 
 		const uint PayloadDataRelativeOffset	= InstanceData.RelativeId   * GetPrimitiveData(InstanceData.PrimitiveId).InstancePayloadDataStride;
 		const uint PayloadDataGlobalOffset		= PayloadDataRelativeOffset + GetPrimitiveData(InstanceData.PrimitiveId).InstancePayloadDataOffset;
@@ -700,28 +689,34 @@ FInstanceSceneData GetInstanceSceneData(uint InstanceId, uint SOAStride, bool bC
 		if (bHasDynamicData)
 		{
 		#if INSTANCE_COMPRESSED_TRANSFORMS
-			uint4  PrevRotationScale		= asuint(LoadInstancePayloadDataElement(Offset_DynamicData + 0));
-			float3 PrevTranslation			= LoadInstancePayloadDataElement(Offset_DynamicData + 1).xyz;
-			float3 PrevScale = 0;
-			InstanceData.PrevLocalToWorld	= DecodeTransform(PrevRotationScale, PrevTranslation, PrevScale);
+			uint4  PrevRotationScale			= asuint(LoadInstancePayloadDataElement(Offset_DynamicData + 0));
+			float3 PrevTranslation				= LoadInstancePayloadDataElement(Offset_DynamicData + 1).xyz;
+			float3 PrevScale					= 0;
+			float4x4 PrevLocalToRelativeWorld	= DecodeTransform(PrevRotationScale, PrevTranslation, PrevScale);
 		#else
-			InstanceData.PrevLocalToWorld	= transpose(float4x4(LoadInstancePayloadDataElement(Offset_DynamicData + 0),
+			float4x4 PrevLocalToRelativeWorld	= transpose(float4x4(LoadInstancePayloadDataElement(Offset_DynamicData + 0),
 																 LoadInstancePayloadDataElement(Offset_DynamicData + 1),
 																 LoadInstancePayloadDataElement(Offset_DynamicData + 2),
 																 float4(0.0f, 0.0f, 0.0f, 1.0f)));
 		#endif
+
+			float3 TilePosition 				= GetPrimitiveData(InstanceData.PrimitiveId).TilePosition;
+			InstanceData.PrevLocalToWorld		= MakeLWCMatrix(TilePosition, PrevLocalToRelativeWorld);
 		}
 		else
 		{
 		#if INSTANCE_COMPRESSED_TRANSFORMS
 			// TODO: Temporary PrevVelocityHack
-			uint4  PrevRotationScale		= asuint(LoadInstanceSceneDataElement(3 * SOAStride + InstanceId));
-			float3 PrevTranslation			= LoadInstanceSceneDataElement(4 * SOAStride + InstanceId).xyz;
+			uint4  PrevRotationScale			= asuint(LoadInstanceSceneDataElement(3 * SOAStride + InstanceId));
+			float3 PrevTranslation				= LoadInstanceSceneDataElement(4 * SOAStride + InstanceId).xyz;
 			float3 PrevScale = 0;
-			InstanceData.PrevLocalToWorld	= DecodeTransform(PrevRotationScale, PrevTranslation, PrevScale);
+			float4x4 PrevLocalToRelativeWorld	= DecodeTransform(PrevRotationScale, PrevTranslation, PrevScale);
 		#else
-			InstanceData.PrevLocalToWorld	= InstanceData.LocalToWorld;
+			float4x4 PrevLocalToRelativeWorld	= InstanceData.LocalToWorld;
 		#endif
+
+			float3 TilePosition 				= GetPrimitiveData(InstanceData.PrimitiveId).TilePosition;
+			InstanceData.PrevLocalToWorld		= MakeLWCMatrix(TilePosition, PrevLocalToRelativeWorld);
 		}
 
 	#if 1 //NEEDS_LIGHTMAP_COORDINATE
