Add a far field visualization mode for ray tracing debugging.

diff --git a/Engine/Shaders/Private/RayTracing/RayTracingDebug.usf b/Engine/Shaders/Private/RayTracing/RayTracingDebug.usf
--- a/Engine/Shaders/Private/RayTracing/RayTracingDebug.usf
+++ b/Engine/Shaders/Private/RayTracing/RayTracingDebug.usf
@@ -17,6 +17,10 @@ uint ShouldUsePreExposure;
 uint OpaqueOnly;
 float TimingScale;
 
+float MaxTraceDistance;
+float FarFieldMaxTraceDistance;
+float3 FarFieldReferencePos;
+
 // TODO: move MurmurMix and IntToColor into a common header, shared with Nanite/Visualize.usf
 
 uint MurmurMix(uint Hash)
@@ -43,7 +47,7 @@ float3 IntToColor(uint Index)
 	return Color * (1.0f / 255.0f);
 }
 
-float4 DebugRayTracingMaterial(RayDesc Ray)
+float4 DebugRayTracingMaterial(RayDesc Ray, bool bIsFarField)
 {
 	uint2 PixelCoord = DispatchRaysIndex().xy + View.ViewRectMin;
 
@@ -51,7 +55,7 @@ float4 DebugRayTracingMaterial(RayDesc Ray)
 	RayCone.SpreadAngle = View.EyeToPixelSpreadAngle;
 
 	const uint RayFlags = RAY_FLAG_CULL_BACK_FACING_TRIANGLES;
-	const uint InstanceInclusionMask = OpaqueOnly ? RAY_TRACING_MASK_OPAQUE : RAY_TRACING_MASK_ALL;
+	const uint InstanceInclusionMask = (bIsFarField) ? RAY_TRACING_MASK_FAR_FIELD : (OpaqueOnly ? RAY_TRACING_MASK_OPAQUE : RAY_TRACING_MASK_ALL);
 	const bool bEnableSkyLightContribution = true;
 	const bool bIgnoreTranslucentMaterials = false;
 
@@ -59,6 +63,17 @@ float4 DebugRayTracingMaterial(RayDesc Ray)
 	uint64_t TimeBegin = GetShaderTimestamp();
 #endif
 
+	if (bIsFarField)
+	{
+		Ray.TMin = MaxTraceDistance;
+		Ray.TMax = FarFieldMaxTraceDistance;
+		Ray.Origin += FarFieldReferencePos;
+	}
+	else
+	{
+		Ray.TMax = (MaxTraceDistance > 0.0f) ? MaxTraceDistance : Ray.TMax;
+	}
+	
 	FMaterialClosestHitPayload Payload = TraceMaterialRay(
 		TLAS,
 		RayFlags,
@@ -73,7 +88,8 @@ float4 DebugRayTracingMaterial(RayDesc Ray)
 	uint64_t TimeEnd = GetShaderTimestamp();
 #endif
 
-	float4 Result = float4(0, 0, 0, 1);
+	// HitT is packed into the alpha for far field compositing
+	float4 Result = (RAY_TRACING_DEBUG_VIZ_FAR_FIELD) ? float4(0, 0, 0, -1) : float4(0, 0, 0, 1);
 
 	if (Payload.IsHit())
 	{
@@ -145,6 +161,9 @@ float4 DebugRayTracingMaterial(RayDesc Ray)
 			Result = float4(0.0, pow(AniG, 2.2), pow(AniB, 2.2), 1.0f);
 			break;
 			}
+		case RAY_TRACING_DEBUG_VIZ_FAR_FIELD:
+			Result = float4(Payload.BaseColor.rgb, Payload.HitT);
+			break;
 		#if PLATFORM_SUPPORTS_SHADER_TIMESTAMP
 		case RAY_TRACING_DEBUG_VIZ_PERFORMANCE:
 			Result.rgb = ColorMapInferno(float(TimeEnd - TimeBegin) * TimingScale);
@@ -211,7 +230,16 @@ RAY_TRACING_ENTRY_RAYGEN(RayTracingDebugMainRGS)
 	}
 	else
 	{
-		Result = DebugRayTracingMaterial(Ray);
+		bool bIsFarField = false;
+		Result = DebugRayTracingMaterial(Ray, bIsFarField);
+		
+		const bool bNearFieldMiss = Result.w <= 0.0;
+		if (VisualizationMode == RAY_TRACING_DEBUG_VIZ_FAR_FIELD && bNearFieldMiss)
+		{
+			bIsFarField = true;
+			const float4 FarFieldTint = float4(1.0f, 0.5f, 0.5f, 1.0f);
+			Result = DebugRayTracingMaterial(Ray, bIsFarField) * FarFieldTint;
+		}
 	}
 
 	if (ShouldUsePreExposure)
diff --git a/Engine/Shaders/Shared/RaytracingDebugDefinitions.h b/Engine/Shaders/Shared/RaytracingDebugDefinitions.h
--- a/Engine/Shaders/Shared/RaytracingDebugDefinitions.h
+++ b/Engine/Shaders/Shared/RaytracingDebugDefinitions.h
@@ -32,4 +32,4 @@
 #define RAY_TRACING_DEBUG_VIZ_INSTANCES					22
 #define RAY_TRACING_DEBUG_VIZ_PERFORMANCE				23 
 #define RAY_TRACING_DEBUG_VIZ_TRIANGLES					24
-
+#define RAY_TRACING_DEBUG_VIZ_FAR_FIELD					25
diff --git a/Engine/Source/Editor/UnrealEd/Private/RayTracingDebugVisualizationMenuCommands.cpp b/Engine/Source/Editor/UnrealEd/Private/RayTracingDebugVisualizationMenuCommands.cpp
--- a/Engine/Source/Editor/UnrealEd/Private/RayTracingDebugVisualizationMenuCommands.cpp
+++ b/Engine/Source/Editor/UnrealEd/Private/RayTracingDebugVisualizationMenuCommands.cpp
@@ -61,6 +61,7 @@ void FRayTracingDebugVisualizationMenuCommands::CreateRayTracingDebugVisualizati
 	}
 
 	RayTracingDebugModeNames.Add(LOCTEXT("Triangles", "Triangles"));
+	RayTracingDebugModeNames.Add(LOCTEXT("FarField", "FarField"));
 
 	for ( int32 RayTracingDebugIndex = 0; RayTracingDebugIndex < RayTracingDebugModeNames.Num(); ++RayTracingDebugIndex)
 	{
diff --git a/Engine/Source/Runtime/Renderer/Private/RayTracing/RayTracingDebug.cpp b/Engine/Source/Runtime/Renderer/Private/RayTracing/RayTracingDebug.cpp
--- a/Engine/Source/Runtime/Renderer/Private/RayTracing/RayTracingDebug.cpp
+++ b/Engine/Source/Runtime/Renderer/Private/RayTracing/RayTracingDebug.cpp
@@ -47,6 +47,9 @@ class FRayTracingDebugRGS : public FGlobalShader
 		SHADER_PARAMETER(uint32, VisualizationMode)
 		SHADER_PARAMETER(int32, ShouldUsePreExposure)
 		SHADER_PARAMETER(float, TimingScale)
+		SHADER_PARAMETER(float, MaxTraceDistance)
+		SHADER_PARAMETER(float, FarFieldMaxTraceDistance)
+		SHADER_PARAMETER(FVector3f, FarFieldReferencePos)
 		SHADER_PARAMETER(int32, OpaqueOnly)
 		SHADER_PARAMETER_SRV(RaytracingAccelerationStructure, TLAS)
 		SHADER_PARAMETER_RDG_TEXTURE_UAV(RWTexture2D<float4>, Output)
@@ -129,6 +132,7 @@ void FDeferredShadingSceneRenderer::RenderRayTracingDebug(FRDGBuilder& GraphBuil
 		RayTracingDebugVisualizationModes.Emplace(FName(*LOCTEXT("Instances", "Instances").ToString()),											RAY_TRACING_DEBUG_VIZ_INSTANCES);
 		RayTracingDebugVisualizationModes.Emplace(FName(*LOCTEXT("Performance", "Performance").ToString()),										RAY_TRACING_DEBUG_VIZ_PERFORMANCE);
 		RayTracingDebugVisualizationModes.Emplace(FName(*LOCTEXT("Triangles", "Triangles").ToString()),											RAY_TRACING_DEBUG_VIZ_TRIANGLES);
+		RayTracingDebugVisualizationModes.Emplace(FName(*LOCTEXT("FarField", "FarField").ToString()),											RAY_TRACING_DEBUG_VIZ_FAR_FIELD);
 	}
 
 	uint32 DebugVisualizationMode;
@@ -201,6 +205,20 @@ void FDeferredShadingSceneRenderer::RenderRayTracingDebug(FRDGBuilder& GraphBuil
 	RayGenParameters->ShouldUsePreExposure = View.Family->EngineShowFlags.Tonemapper;
 	RayGenParameters->TimingScale = CVarRayTracingDebugTimingScale.GetValueOnAnyThread() / 25000.0f;
 	RayGenParameters->OpaqueOnly = CVarRayTracingDebugModeOpaqueOnly.GetValueOnRenderThread();
+	
+	if (Lumen::UseFarField())
+	{
+		RayGenParameters->MaxTraceDistance = Lumen::GetMaxTraceDistance();
+		RayGenParameters->FarFieldMaxTraceDistance = Lumen::GetFarFieldMaxTraceDistance();
+		RayGenParameters->FarFieldReferencePos = Lumen::GetFarFieldReferencePos();
+	}
+	else
+	{
+		RayGenParameters->MaxTraceDistance = 0.0f;
+		RayGenParameters->FarFieldMaxTraceDistance = 0.0f;
+		RayGenParameters->FarFieldReferencePos = FVector(0.0f);
+	}
+	
 	RayGenParameters->TLAS = View.GetRayTracingSceneViewChecked();
 	RayGenParameters->ViewUniformBuffer = View.ViewUniformBuffer;
 	RayGenParameters->Output = GraphBuilder.CreateUAV(SceneColorTexture);
