Works arround vulkan shader compilation failure in TSR

diff --git a/Engine/Shaders/Private/TemporalSuperResolution/TSRUpdateHistory.usf b/Engine/Shaders/Private/TemporalSuperResolution/TSRUpdateHistory.usf
--- a/Engine/Shaders/Private/TemporalSuperResolution/TSRUpdateHistory.usf
+++ b/Engine/Shaders/Private/TemporalSuperResolution/TSRUpdateHistory.usf
@@ -1376,7 +1376,10 @@ void MainCS(
 		// Output final scene color Mip1
 		{
 			taa_half3x2 HalfResOutput = OutputColor * taa_half(0.25);
-			taa_short2x2 HalfResOutputPixelPos = (LocalHistoryPixelPos >> taa_short(1)) | (((LocalHistoryPixelPos & taa_short(0x1)) | !bGenerateOutputMip1) * taa_short(~0));
+			//taa_short2x2 HalfResOutputPixelPos = (LocalHistoryPixelPos >> taa_short(1)) | (((LocalHistoryPixelPos & taa_short(0x1)) | !bGenerateOutputMip1) * taa_short(~0));
+			taa_short2x2 HalfResOutputPixelPos;
+			HalfResOutputPixelPos[0] = (LocalHistoryPixelPos[0] >> taa_short(1)) | (((LocalHistoryPixelPos[0] & taa_short(0x1)) | !bGenerateOutputMip1) * taa_short(~0));
+			HalfResOutputPixelPos[1] = (LocalHistoryPixelPos[1] >> taa_short(1)) | (((LocalHistoryPixelPos[1] & taa_short(0x1)) | !bGenerateOutputMip1) * taa_short(~0));
 			
 			#if PLATFORM_SUPPORTS_WAVE_BROADCAST && 0 // Support WaveBroadcast with halfs
 			BRANCH
