Fix shader compilation error for strict compliers.

diff --git a/Engine/Shaders/Private/VirtualTextureCommon.ush b/Engine/Shaders/Private/VirtualTextureCommon.ush
--- a/Engine/Shaders/Private/VirtualTextureCommon.ush
+++ b/Engine/Shaders/Private/VirtualTextureCommon.ush
@@ -205,7 +205,7 @@ struct VTPageTableResult
 uint GetRequestedLevel(VTPageTableResult PageTableResult)
 {
 	uint RequestLevelPlusOne = (PageTableResult.PackedRequest >> 24) & 0xf;
-	uint RequestLevel = max(RequestLevelPlusOne, 1) - 1;
+	uint RequestLevel = max(RequestLevelPlusOne, 1u) - 1u;
 	return RequestLevel;
 }
 
