Fixed various typos

diff --git a/Engine/Shaders/Private/DebugViewModeCommon.ush b/Engine/Shaders/Private/DebugViewModeCommon.ush
--- a/Engine/Shaders/Private/DebugViewModeCommon.ush
+++ b/Engine/Shaders/Private/DebugViewModeCommon.ush
@@ -1,12 +1,12 @@
 // Copyright Epic Games, Inc. All Rights Reserved.
 
 /*=============================================================================
-	DebugViewModeCommon.hlsl: Debug pixel shader interpolants.
+	DebugViewModeCommon.ush: Debug pixel shader interpolants.
 =============================================================================*/
 
 // This define the interpolant needed by the pixel shader used for debug viewmode.
 // Those use a global shader that must have exactly the same inputs as the VS outputs.
-// This is also compatible with shaders that modify the mesh position (like tesselation).
+// This is also compatible with shaders that modify the mesh position.
 // Which is better in the end than simply defining a material override.
 
 struct FDebugVSToPS
diff --git a/Engine/Shaders/Private/QuadOverdraw.ush b/Engine/Shaders/Private/QuadOverdraw.ush
--- a/Engine/Shaders/Private/QuadOverdraw.ush
+++ b/Engine/Shaders/Private/QuadOverdraw.ush
@@ -1,12 +1,12 @@
 // Copyright Epic Games, Inc. All Rights Reserved.
 
 /*=============================================================================
-	QuadOverdraw.usf: Utils to compute quad coverage
+	QuadOverdraw.ush: Utils to compute quad coverage
 =============================================================================*/
 
 #pragma once
 
-// Because we use the same UAV for Descritors and Complexity, we have the same clear value.
+// Because we use the same UAV for Descriptors and Complexity, we have the same clear value.
 // In this case, the invalid descriptor is 0, so we need to add 1 to the primitive ID to prevent clash.
 const static uint QO_DESC_NONE = 0;
 
@@ -74,7 +74,7 @@ uint2 QuadComplexityOffset()
 	
 	bOutputToComplexity : if complexity is written to the UAV instead of RT. 
 		This is the only code path enabling early out in the costly loop. 
-		Use with [earlydepthstencil] to keep depth write and udpate.
+		Use with [earlydepthstencil] to keep depth write and update.
 		A major behavior difference is that complexity is always additive. (can be changed)
 	
 	QuadComplexity : the value of the quad complexity.
diff --git a/Engine/Shaders/Private/ShaderComplexityAccumulatePixelShader.usf b/Engine/Shaders/Private/ShaderComplexityAccumulatePixelShader.usf
--- a/Engine/Shaders/Private/ShaderComplexityAccumulatePixelShader.usf
+++ b/Engine/Shaders/Private/ShaderComplexityAccumulatePixelShader.usf
@@ -1,7 +1,7 @@
 // Copyright Epic Games, Inc. All Rights Reserved.
 
 /*=============================================================================
-	ShaderComplexityAccumulatePixelShader.usf: Outputs complexity
+	ShaderComplexityAccumulatePixelShader.usf: Outputs shader complexity
 =============================================================================*/
 
 #include "Common.ush"
@@ -24,7 +24,7 @@ void Main(
 {
 	float3 FinalComplexity = NormalizedComplexity.xyz;
 
-	[branch]
+	BRANCH
 	if (bShowQuadOverdraw && NormalizedComplexity.x > 0)
 	{
 		uint Coverage = ComputeQuadCoverage(Input.SvPosition.xy, Input.PrimitiveID, 24, false, false, 0);
diff --git a/Engine/Source/Runtime/Renderer/Private/PostProcess/DiaphragmDOF.cpp b/Engine/Source/Runtime/Renderer/Private/PostProcess/DiaphragmDOF.cpp
--- a/Engine/Source/Runtime/Renderer/Private/PostProcess/DiaphragmDOF.cpp
+++ b/Engine/Source/Runtime/Renderer/Private/PostProcess/DiaphragmDOF.cpp
@@ -1487,12 +1487,12 @@ FRDGTextureRef DiaphragmDOF::AddPasses(
 	const bool bEnableScatterBokehSettings = CVarEnableScatterBokehSettings.GetValueOnRenderThread() == 1;
 	const bool bEnableSlightOutOfFocusBokeh = bSupportGatheringBokehSimulation && bRecombineDoesSlightOutOfFocus && CVarEnableRecombineBokehSettings.GetValueOnRenderThread();
 		
-	// Setup all the desciptors.
+	// Setup all the descriptors.
 	FRDGTextureDesc FullResDesc;
 	{
 		FullResDesc = InputSceneColor->Desc;
 
-		// Reset so that the number of samples of decsriptor becomes 1, which is totally legal still with MSAA because
+		// Reset so that the number of samples of descriptor becomes 1, which is totally legal still with MSAA because
 		// the scene color will already be resolved to ShaderResource texture that is always 1. This is to work around
 		// hack that MSAA will have targetable texture with MSAA != shader resource, and still have descriptor indicating
 		// the number of samples of the targetable resource.
