Fixes FFT bloom sampling some NaN in the kernel sampling

diff --git a/Engine/Shaders/Private/Bloom/BloomReduceKernelSurvey.usf b/Engine/Shaders/Private/Bloom/BloomReduceKernelSurvey.usf
--- a/Engine/Shaders/Private/Bloom/BloomReduceKernelSurvey.usf
+++ b/Engine/Shaders/Private/Bloom/BloomReduceKernelSurvey.usf
@@ -34,6 +34,11 @@ void MainCS(uint GroupThreadIndex : SV_GroupIndex)
 		{
 			float4 LoadedColor = SurveyOutput[GroupThreadIndex + i * TILE_SIZE];
 
+			if ((GroupThreadIndex + i * TILE_SIZE) >= SurveyGroupCount)
+			{
+				LoadedColor = 0.0;
+			}
+
 			if (SurveyReduceOp == 0)
 			{
 				Color = max(Color, LoadedColor);
