Fix shader compilation error for some platforms that need explicit uint type.

diff --git a/Engine/Shaders/Private/Common.ush b/Engine/Shaders/Private/Common.ush
--- a/Engine/Shaders/Private/Common.ush
+++ b/Engine/Shaders/Private/Common.ush
@@ -1125,7 +1125,7 @@ uint DecodeRTWriteMask(uint2 PixelPos, Texture2D<uint> RTWriteMaskTexture, uint
 	uint2 TileOffset = (PixelPos >> 2) & 1;
 
 	uint Shift = (TileOffset.y * 2 + TileOffset.x) * NumEncodedTextures;
-	uint Mask = ~((~0) << NumEncodedTextures);
+	uint Mask = ~((~0u) << NumEncodedTextures);
 
 	return (RTWriteMaskTexture.Load(uint3(TileIndex, 0)) >> Shift) & Mask;
 }
