Fix shader compiler error

diff --git a/Engine/Shaders/Private/VirtualShadowMaps/PageManagement.usf b/Engine/Shaders/Private/VirtualShadowMaps/PageManagement.usf
--- a/Engine/Shaders/Private/VirtualShadowMaps/PageManagement.usf
+++ b/Engine/Shaders/Private/VirtualShadowMaps/PageManagement.usf
@@ -136,13 +136,13 @@ void MarkPage(uint VirtualShadowMapId, uint MipLevel, float3 WorldPosition, bool
 	// PageDilationBorderSize == 0 implies PageDilationOffset.xy == 0
 	if (bUsePageDilation)
 	{
-		uint2 PageAddress2 = clamp(uint2(PageAddressFloat + PageDilationOffset), 0, MaxPageAddress);
+		uint2 PageAddress2 = clamp(uint2(PageAddressFloat + PageDilationOffset), 0U, MaxPageAddress);
 		uint PageOffset2 = CalcPageOffset(VirtualShadowMapId, MipLevel, PageAddress2);
 		if (PageOffset2 != PageOffset)
 		{
 			OutPageRequestFlags[PageOffset2] = Flags;
 		}
-		uint2 PageAddress3 = clamp(uint2(PageAddressFloat - PageDilationOffset), 0, MaxPageAddress);
+		uint2 PageAddress3 = clamp(uint2(PageAddressFloat - PageDilationOffset), 0U, MaxPageAddress);
 		uint PageOffset3 = CalcPageOffset(VirtualShadowMapId, MipLevel, PageAddress3);
 		if (PageOffset3 != PageOffset)
 		{
