Fixes tonemapper sampling outside viewports on even sized resolution

diff --git a/Engine/Shaders/Private/PostProcessTonemap.usf b/Engine/Shaders/Private/PostProcessTonemap.usf
--- a/Engine/Shaders/Private/PostProcessTonemap.usf
+++ b/Engine/Shaders/Private/PostProcessTonemap.usf
@@ -387,25 +387,16 @@ float2 ComputeCounterStepForTileDebugging(float4 SvPosition, int2 TileSize, out
 #define OffsetInterpolator(CounterStep, InterpolatorName) \
 		InterpolatorName += ddx(InterpolatorName) * CounterStep.x + ddy(InterpolatorName) * CounterStep.y;
 
-half3 LookupSceneColor(float2 SceneUV, int2 PixelOffset)
+float4 SampleSceneColor(float2 SceneUV)
 {
-#if (ES3_1_PROFILE && COMPILER_GLSL_ES3_1)
-	// slower but always works
-	// to prevent "error: Texture offset not supported on GLSL ES"
-	return Texture2DSample(ColorTexture, ColorSampler, SceneUV + PixelOffset * Color_ExtentInverse).rgb;
-#else
-	// faster
-	return ColorTexture.SampleLevel(ColorSampler, SceneUV, 0, PixelOffset).rgb;
-#endif
+	SceneUV = clamp(SceneUV, Color_UVViewportBilinearMin, Color_UVViewportBilinearMax);
+	return Texture2DSample(ColorTexture, ColorSampler, SceneUV);
 }
 
-float4 SampleSceneColor(float2 SceneUV)
+half3 LookupSceneColor(float2 SceneUV, float2 PixelOffset)
 {
-	#if USE_COLOR_FRINGE
-		SceneUV = clamp(SceneUV, Color_UVViewportBilinearMin, Color_UVViewportBilinearMax);
-	#endif
-
-	return Texture2DSample(ColorTexture, ColorSampler, SceneUV);
+	float2 SampleUV = SceneUV + PixelOffset * Color_ExtentInverse;
+	return SampleSceneColor(SampleUV);
 }
 
 float4 TonemapCommonPS(
@@ -514,14 +505,14 @@ float4 TonemapCommonPS(
 
 		half A0 = Luminance(SceneColor.rgb);
 
-	#if HAS_PIXEL_QUAD_MESSAGE_PASSING_SUPPORT 
+	#if HAS_PIXEL_QUAD_MESSAGE_PASSING_SUPPORT
 		// Use Wave Intrinsics to reduce texture taps
 		FPQMPContext PQMP = PQMPInit(floor(SceneUV* Color_Extent));
 				
 		half4 LuminanceNeighbors;
 
-		half3 C1 = Texture2DSample(ColorTexture, ColorSampler, SceneUV + float2(PQMP.QuadVector.x, 0) * Color_ExtentInverse).rgb * View.OneOverPreExposure;
-		half3 C3 = Texture2DSample(ColorTexture, ColorSampler, SceneUV + float2(0, PQMP.QuadVector.y) * Color_ExtentInverse).rgb * View.OneOverPreExposure;
+		half3 C1 = LookupSceneColor(SceneUV, float2(PQMP.QuadVector.x, 0)).rgb * View.OneOverPreExposure;
+		half3 C3 = LookupSceneColor(SceneUV, float2(0, PQMP.QuadVector.y)).rgb * View.OneOverPreExposure;
 		LuminanceNeighbors.x = Luminance(C1);
 		LuminanceNeighbors.y = Luminance(C3);
 		
@@ -530,10 +521,10 @@ float4 TonemapCommonPS(
 		LuminanceNeighbors.z =  PQMPReadX(PQMP,A0);
 		LuminanceNeighbors.w =  PQMPReadY(PQMP,A0);
 	#else
-		half3 C1 = LookupSceneColor(SceneUV, int2(-1,  0)) * View.OneOverPreExposure;
-		half3 C2 = LookupSceneColor(SceneUV, int2( 1,  0)) * View.OneOverPreExposure;
-		half3 C3 = LookupSceneColor(SceneUV, int2( 0, -1)) * View.OneOverPreExposure;
-		half3 C4 = LookupSceneColor(SceneUV, int2( 0,  1)) * View.OneOverPreExposure;
+		half3 C1 = LookupSceneColor(SceneUV, float2(-1,  0)) * View.OneOverPreExposure;
+		half3 C2 = LookupSceneColor(SceneUV, float2( 1,  0)) * View.OneOverPreExposure;
+		half3 C3 = LookupSceneColor(SceneUV, float2( 0, -1)) * View.OneOverPreExposure;
+		half3 C4 = LookupSceneColor(SceneUV, float2( 0,  1)) * View.OneOverPreExposure;
 		half4 LuminanceNeighbors = half4(Luminance(C1), Luminance(C2), Luminance(C3), Luminance(C4));
 	#endif
 
