[Splitscreen] Landscape flickering for player 1

Landscape LOD uniform buffer for player 1 was getting overwritten with data from player 2, causing visual glitches

Also fixed an issue where the Vulkan RHI wasn't copying the data provided at the resource creation when BUF_Volatile was involved (thanks [REDACTED_HANDLE])

#localization none
#tests PF tested on multiple platforms on 19.00. Only issue found was related to another change (CL18272160) that has now been backed out.

diff --git a/Engine/Shaders/Private/LandscapeVertexFactory.ush b/Engine/Shaders/Private/LandscapeVertexFactory.ush
--- a/Engine/Shaders/Private/LandscapeVertexFactory.ush
+++ b/Engine/Shaders/Private/LandscapeVertexFactory.ush
@@ -31,6 +31,17 @@ int GetComponentLinearIndex()
 #endif
 }
 
+float GetViewLODData(uint ComponentIndex)
+{
+#if FIXED_GRID
+	// Fixed grid doesn't use LandscapeContinuousLODParameters buffers
+	return 0;
+#else
+	const int AbsoluteSectionLODIndex = View.LandscapeIndirection[LandscapeContinuousLODParameters.LandscapeIndex] + ComponentIndex;
+	return View.LandscapePerComponentData[AbsoluteSectionLODIndex];
+#endif
+}
+
 /* Get current Lod for this section instance. */
 float GetSectionLod(uint ComponentIndex)
 {
@@ -106,7 +117,11 @@ float4 GetNeighborSectionLod(uint ComponentIndex)
 		int2 NeighborPos = NeighborOffset[i] + int2(LandscapeParameters.ComponentBaseX, LandscapeParameters.ComponentBaseY);
 		NeighborPos = max(LandscapeContinuousLODParameters.Min                                                     , NeighborPos);
 		NeighborPos = min(LandscapeContinuousLODParameters.Min + LandscapeContinuousLODParameters.Size - int2(1, 1), NeighborPos);
-		NeighborLod[i] = max(LandscapeContinuousLODParameters.SectionLOD[(NeighborPos.y - LandscapeContinuousLODParameters.Min.y) * LandscapeContinuousLODParameters.Size.x + (NeighborPos.x - LandscapeContinuousLODParameters.Min.x)], GetSectionLod(ComponentIndex));
+
+		uint NeighborComponentIndex = (NeighborPos.y - LandscapeContinuousLODParameters.Min.y) * LandscapeContinuousLODParameters.Size.x + (NeighborPos.x - LandscapeContinuousLODParameters.Min.x);
+		float SectionLOD = GetViewLODData(NeighborComponentIndex);
+
+		NeighborLod[i] = max(SectionLOD, GetSectionLod(ComponentIndex));
 		NeighborLod[i] = View.OverrideLandscapeLOD >= 0.0f ? View.OverrideLandscapeLOD : NeighborLod[i] + View.FarShadowStaticMeshLODBias;
 		NeighborLod[i] = min(NeighborLod[i], LandscapeParameters.LastLOD);
 	}
