Strata - Programmatically compute the bottom layer index.
Strata now compiles again on console after 18359454.
Cost looks similar.

diff --git a/Engine/Shaders/Private/Strata/Strata.ush b/Engine/Shaders/Private/Strata/Strata.ush
--- a/Engine/Shaders/Private/Strata/Strata.ush
+++ b/Engine/Shaders/Private/Strata/Strata.ush
@@ -1892,6 +1892,34 @@ void PackStrataOut(
 	// We clamp the layer count to the maximum one according to our byte per pixel budget (removing the bottom layer first)
 	Strata.LayerCount = min(Strata.LayerCount, STRATA_CLAMPED_LAYER_COUNT);
 
+	// Evaluate the bottom layer index programmatically according to coverage.
+	// This is needed for instance when doing legacy to strata material conversion: a static topology with 2 layers is used for the compiler to
+	// successfully unroll the loops below. Typically the BSDF is set as the top layer. 
+	// SSS would not work if not the bottom layer so that is why we recover the bottom layer index this way.
+	int BottomLayerIndex = Strata.LayerCount - 1;
+#if STRATA_INLINE_SHADING
+	{
+		STRATA_UNROLL_N(STRATA_CLAMPED_LAYER_COUNT)
+		for (int l = BottomLayerIndex; l > 0; --l)	// Skipping layer 0 to process one less layer when there are two.
+		{
+			float TotalWeight = 0.0f;
+			STRATA_UNROLL_N(STRATA_CLAMPED_MAX_BSDF_COUNT_PER_LAYER)
+			for (int i = 0; i < Strata.Layers[l].BSDFCount; ++i)
+			{
+				TotalWeight += /*Throughput * */(Strata.Layers[l].BSDFs[i].bPremultipliedCoverage ? 1.0 : Strata.Layers[l].BSDFs[i].Coverage);
+			}
+			if (TotalWeight == 0.0f)
+			{
+				BottomLayerIndex--;
+			}
+			else
+			{
+				break;
+			}
+		}
+	}
+#endif
+
 	// For each of the BSDF, we evaluated their combined coverage/throughput according to the material operation performed (layering, mixing, etc.).
 	{
 		STRATA_UNROLL_N(STRATA_CLAMPED_LAYER_COUNT)
@@ -1901,7 +1929,7 @@ void PackStrataOut(
 			float3 LayerCoverage = 0.0f;
 
 			const bool bTopLayer = l == 0;
-			const bool bBottomLayer = l == (Strata.LayerCount - 1);
+			const bool bBottomLayer = l == BottomLayerIndex;
 
 			STRATA_UNROLL_N(STRATA_CLAMPED_MAX_BSDF_COUNT_PER_LAYER)
 			for (int i = 0; i < Strata.Layers[l].BSDFCount; ++i)
diff --git a/Engine/Shaders/Private/Strata/StrataLegacyConversion.ush b/Engine/Shaders/Private/Strata/StrataLegacyConversion.ush
--- a/Engine/Shaders/Private/Strata/StrataLegacyConversion.ush
+++ b/Engine/Shaders/Private/Strata/StrataLegacyConversion.ush
@@ -386,15 +386,7 @@ FStrataData StrataConvertLegacyMaterialDynamic(
 	// Fixed layer structure for helping compiler to unroll and optimize shader
 	FStrataData Top_Slab = CreateLegacySlab(Top, SharedLocalBasisTypes);
 	FStrataData Bot_Slab = CreateLegacySlab(Bot, SharedLocalBasisTypes);
-	Top_Slab = StrataWeight(Top_Slab, Top.Weight);
-	FStrataData Out = Top_Slab; 
-	if (Bot.Weight > 0.0f)
-	{
-		// Default BSDF as set to be on top so that their information is stored as part of the Strata top layer data.
-		// When bottom weight is 0, we must set the top BRDF at the bottom to make sure SSS can be identified correctly.
-		Out = StrataVerticalLayering(Out, StrataWeight(Bot_Slab, Bot.Weight));
-	}
-	
+	FStrataData Out = StrataVerticalLayering(StrataWeight(Top_Slab, Top.Weight), StrataWeight(Bot_Slab, Bot.Weight));
 
 	// Take in account the legacy opacity weighting for translucent material
 	#if !MATERIALBLENDING_MASKED && !MATERIALBLENDING_SOLID
