Strata - min fix for debug viz

diff --git a/Engine/Shaders/Private/Strata/StrataVisualize.usf b/Engine/Shaders/Private/Strata/StrataVisualize.usf
--- a/Engine/Shaders/Private/Strata/StrataVisualize.usf
+++ b/Engine/Shaders/Private/Strata/StrataVisualize.usf
@@ -131,6 +131,7 @@ void VisualizeMaterialPS(
 		if (HasSubsurface(Header2) || HasSubsurfaceProfile(Header2))
 		{
 			OutColor = float4(HasSubsurfaceProfile(Header2) ? 1 : 0, HasSubsurface(Header2) ? 1 : 0, 0, 0.5);
+			return;
 		}
 	}
 #endif
