Fix subsurface profile id visualization.

diff --git a/Engine/Shaders/Private/MiniFontCommon.ush b/Engine/Shaders/Private/MiniFontCommon.ush
--- a/Engine/Shaders/Private/MiniFontCommon.ush
+++ b/Engine/Shaders/Private/MiniFontCommon.ush
@@ -52,6 +52,7 @@ Texture2D MiniFontTexture;
 #define _DOT_ 46
 #define _PLUS_ 43
 #define _SPC_ 32
+#define _QUESTIONMARK_ 63
 
 // sample from the MiniFont texture
 // @param CharacterID - the character to print
diff --git a/Engine/Shaders/Private/PostProcessSubsurface.usf b/Engine/Shaders/Private/PostProcessSubsurface.usf
--- a/Engine/Shaders/Private/PostProcessSubsurface.usf
+++ b/Engine/Shaders/Private/PostProcessSubsurface.usf
@@ -329,6 +329,23 @@ float ComputeFullResLerp(float ProfileRadiusScale, float PixelRadiusScale, float
 	return Ret;
 }
 
+uint MiniFontCharFromSubsurfaceProfile(uint SubsurfaceProfileInt)
+{
+	if (SubsurfaceProfileInt <= 9)
+	{
+		return _0_ + SubsurfaceProfileInt;
+	}
+
+	SubsurfaceProfileInt -= 10;
+
+	if (SubsurfaceProfileInt <= _Z_ - _A_)
+	{
+		return _A_ + SubsurfaceProfileInt;
+	}
+
+	return _QUESTIONMARK_;
+}
+
 // visualization (doesn't have to be fast)
 void VisualizePS(in noperspective float4 UVAndScreenPos : TEXCOORD0, float4 SvPosition : SV_POSITION, out float4 OutColor : SV_Target0)
 {
@@ -402,7 +419,7 @@ void VisualizePS(in noperspective float4 UVAndScreenPos : TEXCOORD0, float4 SvPo
 			}
 
 			int2 LeftTop = (PixelPos / 8) * 8;
-			PrintCharacter(PixelPos, OutColor.rgb, float3(1, 1, 1), LeftTop, SubsurfaceProfileInt);
+			PrintCharacter(PixelPos, OutColor.rgb, float3(1, 1, 1), LeftTop, MiniFontCharFromSubsurfaceProfile(SubsurfaceProfileInt));
 
 			OutColor.rgb *= ComputeFullResLerp(ProfileRadiusScale, PixelRadiusScale, UV, SubsurfaceInput0_ExtentInverse);
 		}
