Fixing smooth transition of planar reflection.
It will start to fade off at roughness= 0.2, and finish at roughness = 0.3 as originally expected.

diff --git a/Engine/Shaders/Private/PlanarReflectionShared.ush b/Engine/Shaders/Private/PlanarReflectionShared.ush
--- a/Engine/Shaders/Private/PlanarReflectionShared.ush
+++ b/Engine/Shaders/Private/PlanarReflectionShared.ush
@@ -66,7 +66,7 @@ half4 ComputePlanarReflections(float3 WorldPosition, half3 WorldNormal, half Rou
 
 		// Fade out in regions of the planar reflection that weren't written to, so we can composite with other reflection methods
 		FinalFade *= PlanarReflectionTextureValue.a;
-		OutPlanarReflection.rgb = PlanarReflectionTextureValue.rgb;
+		OutPlanarReflection.rgb = PlanarReflectionTextureValue.rgb * RoughnessFade; // Add roughness fade to color to provide smooth color transition.
 		OutPlanarReflection.a = FinalFade;
 	}
 
