Reduce scope of forward sky evaluation to help compilers.

There is no performance saving related to that, but foward shaders compilation seems faster.

diff --git a/Engine/Shaders/Private/Strata/StrataForwardLighting.ush b/Engine/Shaders/Private/Strata/StrataForwardLighting.ush
--- a/Engine/Shaders/Private/Strata/StrataForwardLighting.ush
+++ b/Engine/Shaders/Private/Strata/StrataForwardLighting.ush
@@ -298,72 +298,80 @@ float3 StrataForwardLighting(
 
 			if (UseBasePassSkylight > 0)
 			{
-				// Compute the bent normal from the BSDF normal
-				const float MaterialAO = StrataPixelHeader.MaterialAO;
-				float DiffOcclusion = MaterialAO;
-				float SpecOcclusion = MaterialAO;
-				float3 BentNormal = StrataBSDFContext.N;
-				float Roughness = StrataGetBSDFRoughness(BSDF);
-				ApplyBentNormal(MaterialParameters.CameraVector, MaterialParameters.WorldNormal, MaterialParameters.TangentToWorld, GetBentNormalZero(MaterialParameters), Roughness, BentNormal, DiffOcclusion, SpecOcclusion);
-
-				// Update a context specific for environment lighting
-				FStrataBSDFContext EnvBSDFContext = StrataBSDFContext;
-				EnvBSDFContext.N = BentNormal;
-				StrataUpdateBSDFContext(EnvBSDFContext, EnvBSDFContext.L);
-
-				// And evaluate diffuse parameters
 				const bool bEnableSpecular = ReflectionStruct.SkyLightParameters.y > 0.0f;
-				FStrataEnvLightResult StrataEnvLight = StrataEvaluateForEnvLight(EnvBSDFContext, bEnableSpecular, bForceFullyRough);
-				float3 DiffuseNormal = StrataEnvLight.DiffuseNormal;
-				float3 DiffuseColorForIndirect = StrataEnvLight.DiffuseWeight;
-				float3 SubsurfaceColorForIndirect = StrataEnvLight.DiffuseBackFaceWeight;
+				float Roughness = StrataGetBSDFRoughness(BSDF);
 				float IndirectOcclusion = 1.0f;
 				float IndirectIrradiance = 0.0f; 
 
-				// Evaluate diffuse lighting
-				if (any((DiffuseColorForIndirect + SubsurfaceColorForIndirect) > 0.0f))
-				{
-					const bool bEvaluateBackface = any(SubsurfaceColorForIndirect > 0.0);
-
-					float3 DiffuseIndirectLighting = 0.0f;
-					float3 SubsurfaceIndirectLighting = 0.0f;
-					GetPrecomputedIndirectLightingAndSkyLight(
-						MaterialParameters,
-						Interpolants,
-						BasePassInterpolants,
-						LightmapVTPageTableResult,
-						bEvaluateBackface,
-						DiffuseNormal,
-						VolumetricLightmapBrickTextureUVs,
-						DiffuseIndirectLighting,
-						SubsurfaceIndirectLighting,
-						IndirectIrradiance);
-
-				#if FORWARD_SHADING && (MATERIALBLENDING_SOLID || MATERIALBLENDING_MASKED)
-					float2 NearestResolvedDepthScreenUV = CalculateNearestResolvedDepthScreenUV(ScreenUV, MaterialParameters.ScreenPosition.w);
-
-					IndirectOcclusion = GetIndirectOcclusion(NearestResolvedDepthScreenUV, HasDynamicIndirectShadowCasterRepresentation(StrataPixelHeader));
-					DiffuseIndirectLighting *= IndirectOcclusion;
-					SubsurfaceIndirectLighting *= IndirectOcclusion;
-					IndirectIrradiance *= IndirectOcclusion;
-				#endif
+				const float MaterialAO = StrataPixelHeader.MaterialAO;
+				float DiffOcclusion = MaterialAO;
+				float SpecOcclusion = MaterialAO;
 
-					Color += (DiffuseIndirectLighting * DiffuseColorForIndirect + SubsurfaceIndirectLighting * SubsurfaceColorForIndirect) * 
-						BSDFColoredVisibility * AOMultiBounce(StrataEnvLight.DiffuseColor, DiffOcclusion);
+				// 1. Diffuse Sky evaluation using bent normal
+				{
+					// Compute the bent normal from the BSDF normal
+					float3 BentNormal = StrataBSDFContext.N;
+					ApplyBentNormal(MaterialParameters.CameraVector, MaterialParameters.WorldNormal, MaterialParameters.TangentToWorld, GetBentNormalZero(MaterialParameters), Roughness, BentNormal, DiffOcclusion, SpecOcclusion);
+
+					// Update a context specific for environment lighting
+					FStrataBSDFContext EnvBSDFContext = StrataBSDFContext;
+					EnvBSDFContext.N = BentNormal;
+					StrataUpdateBSDFContext(EnvBSDFContext, EnvBSDFContext.L);
+
+					// And evaluate diffuse parameters
+					FStrataEnvLightResult StrataEnvLight = StrataEvaluateForEnvLight(EnvBSDFContext, false/*bEnableSpecular*/, bForceFullyRough);
+					float3 DiffuseNormal = StrataEnvLight.DiffuseNormal;
+					float3 DiffuseColorForIndirect = StrataEnvLight.DiffuseWeight;
+					float3 SubsurfaceColorForIndirect = StrataEnvLight.DiffuseBackFaceWeight;
+
+					// Evaluate diffuse lighting
+					if (any((DiffuseColorForIndirect + SubsurfaceColorForIndirect) > 0.0f))
+					{
+						const bool bEvaluateBackface = any(SubsurfaceColorForIndirect > 0.0);
+
+						float3 DiffuseIndirectLighting = 0.0f;
+						float3 SubsurfaceIndirectLighting = 0.0f;
+						GetPrecomputedIndirectLightingAndSkyLight(
+							MaterialParameters,
+							Interpolants,
+							BasePassInterpolants,
+							LightmapVTPageTableResult,
+							bEvaluateBackface,
+							DiffuseNormal,
+							VolumetricLightmapBrickTextureUVs,
+							DiffuseIndirectLighting,
+							SubsurfaceIndirectLighting,
+							IndirectIrradiance);
+
+					#if FORWARD_SHADING && (MATERIALBLENDING_SOLID || MATERIALBLENDING_MASKED)
+						float2 NearestResolvedDepthScreenUV = CalculateNearestResolvedDepthScreenUV(ScreenUV, MaterialParameters.ScreenPosition.w);
+
+						IndirectOcclusion = GetIndirectOcclusion(NearestResolvedDepthScreenUV, HasDynamicIndirectShadowCasterRepresentation(StrataPixelHeader));
+						DiffuseIndirectLighting *= IndirectOcclusion;
+						SubsurfaceIndirectLighting *= IndirectOcclusion;
+						IndirectIrradiance *= IndirectOcclusion;
+					#endif
+
+						Color += (DiffuseIndirectLighting * DiffuseColorForIndirect + SubsurfaceIndirectLighting * SubsurfaceColorForIndirect) * 
+							BSDFColoredVisibility * AOMultiBounce(StrataEnvLight.DiffuseColor, DiffOcclusion);
+					}
 				}
 
+				// 2. Specular evaluation (using regular normal)
 			#if FORWARD_PER_PIXEL_SHADING || TRANSLUCENCY_LIGHTING_SURFACE_LIGHTINGVOLUME
-				// Evaluate specular parameters without bent normal
-				StrataEnvLight = StrataEvaluateForEnvLight(StrataBSDFContext, bEnableSpecular, bForceFullyRough);
-
-				// Evaluate specular lighting
-				if (any(StrataEnvLight.SpecularWeight > 0.0f))
 				{
-					int SingleCaptureIndex = GetPrimitiveData(MaterialParameters).SingleCaptureIndex;
-					Color += GetImageBasedReflectionSpecular(MaterialParameters, StrataEnvLight.SpecularDirection, Roughness, IndirectIrradiance, GridIndex, SingleCaptureIndex, EyeIndex) *
-						 StrataEnvLight.SpecularWeight * IndirectOcclusion * BSDFColoredVisibility * AOMultiBounce(StrataEnvLight.SpecularColor, SpecOcclusion);
+					// Evaluate specular parameters without bent normal
+					FStrataEnvLightResult StrataEnvLight = StrataEvaluateForEnvLight(StrataBSDFContext, bEnableSpecular, bForceFullyRough);
+
+					// Evaluate specular lighting
+					if (any(StrataEnvLight.SpecularWeight > 0.0f))
+					{
+						int SingleCaptureIndex = GetPrimitiveData(MaterialParameters).SingleCaptureIndex;
+						Color += GetImageBasedReflectionSpecular(MaterialParameters, StrataEnvLight.SpecularDirection, Roughness, IndirectIrradiance, GridIndex, SingleCaptureIndex, EyeIndex) *
+							 StrataEnvLight.SpecularWeight * IndirectOcclusion * BSDFColoredVisibility * AOMultiBounce(StrataEnvLight.SpecularColor, SpecOcclusion);
+					}
+					// We skip hazy lob evaluation in forward.
 				}
-				// We skip hazy lob evaluation in forward.
 			#endif
 			}
 
