[ShaderReduction] - Remove 7,688 shaders needed in the editor by combining separate debug viewmode shaders into a single ubershader controlled dynamically.  Debug shaders are now only ~7% of the shaders in EngineTest instead of 15% of all shaders.

- Stats were run on EngineTest and results will vary from project to project.
- Removes FComplexityAccumulatePS, FLODColorationPS, FRequiredTextureResolutionPS, FMeshTexCoordSizeAccuracyPS, FPrimitiveDistanceAccuracyPS, FMaterialTexCoordScalePS
- All shader code is now contained in DebugViewmodePixelShader.usf.  Individual visualizations are contained within separate functions.
- All the C++ interface code has now been folded into a single implementation.
- Previously each implementation kept track of its name, and whether it needed local VF, if it needed material properties and needing instruction count.  Since every implementation is in the same shader we need all those features.  To simplify the code all this has been removed.
- We only keep track of a single FDebugViewModeInterface singleton to use for all debug modes.

diff --git a/Engine/Shaders/Private/DebugViewModePixelShader.usf b/Engine/Shaders/Private/DebugViewModePixelShader.usf
new file mode 100644
--- /dev/null
+++ b/Engine/Shaders/Private/DebugViewModePixelShader.usf
@@ -0,0 +1,585 @@
+// Copyright Epic Games, Inc. All Rights Reserved.
+
+/*=============================================================================
+	DebugViewModePixelShader.usf: Ubershader for debug pixel modes.
+=============================================================================*/
+
+#define TEX_COORD_SCALE_ANALYSIS 1
+#define DEBUG_MATERIAL_PARAMETERS 1
+
+// Don't allow scene texture as we use ESceneRenderTargetsMode::DontSet in FMaterialTexCoordScalePS::SetParameters
+// to avoid crashes in commandlet mode.
+#undef SCENE_TEXTURES_DISABLED
+#define SCENE_TEXTURES_DISABLED 1
+
+#undef EYE_ADAPTATION_DISABLED
+#define EYE_ADAPTATION_DISABLED 1
+
+#include "Common.ush"
+#include "SHCommon.ush"
+#include "QuadOverdraw.ush"
+#include "VirtualTextureCommon.ush"
+
+#define UNDEF_ACCURACY 1024
+
+// If -1, output everything to the RW texture.
+float4 OneOverCPUTexCoordScales[MAX_NUM_TEXTURE_REGISTER / 4];
+int4 TexCoordIndices[MAX_NUM_TEXTURE_REGISTER / 4];
+float4 CPUTexelFactor; // In world space units.
+float4 NormalizedComplexity;
+int3 AnalysisParams; // (TextureAnalysisIndex, bOutputScales) or // (TextureAnalysisIndex, TextureResolution)
+float PrimitiveAlpha;
+int TexCoordAnalysisIndex;
+float CPULogDistance;
+uint bShowQuadOverdraw;
+uint bOutputQuadOverdraw;
+int LODIndex;
+int VisualizeMode;
+
+// Analysis index will be in [0, 31] if analysing a single texture, otherwise -1
+#define TextureAnalysisIndex (AnalysisParams.x)
+#define bOutputScales (AnalysisParams.y)
+// Texture resolution if we're analysing a single texture
+#define TextureResolution (AnalysisParams.y)
+// Analysis type will be ANALYSIS_TEXTURE_RESOLUTION or ANALYSIS_VIRTUAL_TEXTURE_PENDING_MIPS
+#define TextureAnalysisType (AnalysisParams.z)
+
+#define ANALYSIS_TEXTURE_RESOLUTION				0
+#define ANALYSIS_VIRTUAL_TEXTURE_PENDING_MIPS	1
+
+// Updated MAX_NUM_TEXTURE_REGISTER to 256 
+
+float GetComponent(float4 V, int Index)
+{
+	FLATTEN
+		if (Index == 0) return V.x;
+	FLATTEN
+		if (Index == 1) return V.y;
+	FLATTEN
+		if (Index == 2) return V.z;
+	return V.w;
+}
+
+struct FTexCoordScalesParams
+{
+	// Global var used to reduce the number of param passed to the shaders.
+	int2 PixelPosition;
+
+	// ddU, ddV for each coord index (up to 4 currently)
+	float4 OneOverDDU;
+	float4 OneOverDDV;
+
+	// Used when inspecting single elements.
+	float MinScale;
+	float MaxScale;
+
+	float TexSample;
+	float TexSampleAverage;
+
+	// When computing the texture scales, this holds 4 scales for consecutive texture register indices.
+	float4 ScalesPerIndex;
+
+	float RequiredResolution;
+	int VTMissingMipCount;
+};
+
+MaterialFloat StoreTexCoordScale(in out FTexCoordScalesParams Params, float2 UV, int TextureReferenceIndex)
+{
+	// Take the minimum scale since high scale requires less resolution.
+	float GPUScaleX = length(ddx(UV));
+	float GPUScaleY = length(ddy(UV));
+
+	if (TextureReferenceIndex >= 0 && TextureReferenceIndex < MAX_NUM_TEXTURE_REGISTER)
+	{
+		float OneOverCPUScale = OneOverCPUTexCoordScales[TextureReferenceIndex / 4][TextureReferenceIndex % 4];
+
+		int TexCoordIndex = TexCoordIndices[TextureReferenceIndex / 4][TextureReferenceIndex % 4];
+
+		float GPUScale = min(GPUScaleX * GetComponent(Params.OneOverDDU, TexCoordIndex), GPUScaleY * GetComponent(Params.OneOverDDV, TexCoordIndex));
+
+		// Code path when searching the worst scales. MinScale and MaxScale are only used when rendering for accuracies (!bOutputScales)
+		const bool bUpdateMinMax = (OneOverCPUScale > 0 && (TextureAnalysisIndex == -1 || TextureAnalysisIndex == TextureReferenceIndex));
+		Params.MinScale = bUpdateMinMax ? min(Params.MinScale, GPUScale * OneOverCPUScale) : Params.MinScale;
+		Params.MaxScale = bUpdateMinMax ? max(Params.MaxScale, GPUScale * OneOverCPUScale) : Params.MaxScale;
+
+		// Code path when rendering the material in a render target. Each TILE_RESOLUTION X TILE_RESOLUTION maps to different texture.
+		const bool bUpdateScale = (bOutputScales && Params.PixelPosition.x / TILE_RESOLUTION == TextureReferenceIndex / 4);
+		Params.ScalesPerIndex[TextureReferenceIndex % 4] = bUpdateScale ? min(Params.ScalesPerIndex[TextureReferenceIndex % 4], GPUScale) : Params.ScalesPerIndex[TextureReferenceIndex % 4];
+	}
+
+#if FEATURE_LEVEL >= FEATURE_LEVEL_SM5
+	float2 CoordDDX = ddx_fine(UV);
+	float2 CoordDDY = ddy_fine(UV);
+#else
+	float2 CoordDDX = ddx(UV);
+	float2 CoordDDY = ddy(UV);
+#endif 
+
+	if (TextureAnalysisIndex == TextureReferenceIndex)
+	{
+		float MinDelta = min(length(CoordDDX), length(CoordDDY));
+		float RequiredResolution = (1 / max(MinDelta, 0.0000000001f));
+		Params.RequiredResolution = max(Params.RequiredResolution, RequiredResolution);
+	}
+
+	return 1.f;
+}
+
+MaterialFloat StoreTexSample(in out FTexCoordScalesParams Params, float4 C, int TextureReferenceIndex)
+{
+	// Alpha?
+	if (TextureAnalysisIndex == TextureReferenceIndex)
+	{
+		Params.TexSample = TextureAnalysisIndex == TextureReferenceIndex ? lerp(.4f, 1.f, saturate(Luminance(C.rgb))) : Params.TexSample;
+	}
+	Params.TexSampleAverage = TextureAnalysisIndex == TextureReferenceIndex ? lerp(.4f, 1.f, saturate((C.r + C.g + C.b) / 3)) : Params.TexSample;
+
+	return 1.f;
+}
+
+MaterialFloat StoreVTSampleInfo(in out FTexCoordScalesParams Params, VTPageTableResult PageTableResult, uint LayerIndex, int TextureReferenceIndex)
+{
+	if (TextureAnalysisType == ANALYSIS_VIRTUAL_TEXTURE_PENDING_MIPS)
+	{
+		uint RequestLevel = GetRequestedLevel(PageTableResult);
+
+		// Find the highest resolution from all the page table entries sampled for this VT stack.
+		// We do that because a single VT stack can cover textures of different resolutions.
+		// Taking info from the lower resolution textures will always show undersampling compared to the requested level that we are not interested in.
+		uint SampledLevel = 0xffff;
+		for (uint Index = 0; Index < 8; ++Index)
+		{
+			SampledLevel = min(SampledLevel, GetSampledLevel(PageTableResult, Index, 0xffff));
+		}
+
+		if (SampledLevel < 0xffff)
+		{
+			// Accumulate the worst case that we find for the current pixel.
+			// Multiple VT stacks can exist, and we want to visualize the worst case.
+			int Diff = max((int)SampledLevel - (int)RequestLevel, 0);
+			Params.VTMissingMipCount = max(Diff, Params.VTMissingMipCount);
+		}
+	}
+	return 1.f;
+}
+
+#include "/Engine/Generated/Material.ush"
+#include "DebugViewModeCommon.ush"
+
+float3 AccuracyColorLookup(float Accuracy)
+{
+	Accuracy = clamp(Accuracy, -1.99, 1.99);
+	int ColorIndex = floor(Accuracy) + 2;
+	float3 Color0 = DebugViewModeStruct.AccuracyColors[ColorIndex].rgb;
+	float3 Color1 = DebugViewModeStruct.AccuracyColors[ColorIndex + 1].rgb;
+	float ColorLerp = frac(Accuracy);
+	return lerp(Color0, Color1, ColorLerp);
+}
+
+float GetCPUTexelFactor(int CoordIndex)
+{
+	FLATTEN
+		if (CoordIndex == 0) return CPUTexelFactor.x;
+	FLATTEN
+		if (CoordIndex == 1) return CPUTexelFactor.y;
+	FLATTEN
+		if (CoordIndex == 2) return CPUTexelFactor.z;
+	return CPUTexelFactor.w;
+}
+
+float2 GetTexCoord(in FDebugPSIn Inputs, int CoordIndex)
+{
+	FLATTEN
+		if (CoordIndex == 0) return Inputs.TexCoord01.xy;
+	FLATTEN
+		if (CoordIndex == 1) return Inputs.TexCoord01.zw;
+	FLATTEN
+		if (CoordIndex == 2) return Inputs.TexCoord23.xy;
+	return Inputs.TexCoord23.zw;
+}
+
+float GetTexCoordSizeAccuracy(in FDebugPSIn Inputs, int CoordIndex, float CPUSize)
+{
+	float3 WorldPosition = SvPositionToResolvedTranslatedWorld(Inputs.SvPosition);
+
+	float2 TexCoord = GetTexCoord(Inputs, CoordIndex);
+
+#if FEATURE_LEVEL >= FEATURE_LEVEL_SM5
+	float2 CoordDDX = ddx_fine(TexCoord);
+	float2 CoordDDY = ddy_fine(TexCoord);
+	float3 WorldPosDDX = ddx_fine(WorldPosition);
+	float3 WorldPosDDY = ddy_fine(WorldPosition);
+#else
+	float2 CoordDDX = ddx(TexCoord);
+	float2 CoordDDY = ddy(TexCoord);
+	float3 WorldPosDDX = ddx(WorldPosition);
+	float3 WorldPosDDY = ddy(WorldPosition);
+#endif 
+
+	float UVAera = abs(CoordDDX.x * CoordDDY.y - CoordDDX.y * CoordDDY.x);
+	float WorldAera = length(cross(WorldPosDDX, WorldPosDDY));
+
+	float GPUTexelFactor = sqrt(WorldAera / max(UVAera, 0.0000000001f));
+
+	return clamp(log2(CPUSize) - log2(GPUTexelFactor), -1.99, 1.99);
+}
+
+float4 VisualizeMaterialTexCoordScales(in FDebugPSIn DebugInputs, in bool bIsFrontFace)
+{
+	// This default value will make it dark grey when nothing updates the texSample.
+	float3 Result = float3(UNDEFINED_ACCURACY, UNDEFINED_ACCURACY, UNDEFINED_ACCURACY);
+
+	FTexCoordScalesParams Params;
+
+	// Set global var used for outputting data
+	Params.PixelPosition = DebugInputs.SvPosition.xy;
+	Params.OneOverDDU = 1 / float4(length(ddx(DebugInputs.TexCoord01.xy)), length(ddx(DebugInputs.TexCoord01.zw)), length(ddx(DebugInputs.TexCoord23.xy)), length(ddx(DebugInputs.TexCoord23.zw)));
+	Params.OneOverDDV = 1 / float4(length(ddy(DebugInputs.TexCoord01.xy)), length(ddy(DebugInputs.TexCoord01.zw)), length(ddy(DebugInputs.TexCoord23.xy)), length(ddy(DebugInputs.TexCoord23.zw)));
+
+	Params.TexSample = 1.f;
+	Params.TexSampleAverage = 1.f;
+	Params.MinScale = INITIAL_GPU_SCALE;
+	Params.MaxScale = 0;
+	Params.ScalesPerIndex = INITIAL_GPU_SCALE;
+	Params.RequiredResolution = 0;
+	Params.VTMissingMipCount = -1;
+
+	FMaterialPixelParameters MaterialParameters = GetMaterialPixelParameters(DebugInputs, DebugInputs.SvPosition);
+
+	// This is used to patch invalid parameters when computing the scales with FMeshMaterialRenderItem::EnqueueMaterialRender
+	FLATTEN
+		if (bOutputScales)
+		{
+#if NUM_MATERIAL_TEXCOORDS > 0
+			// Set all texcoord but one to 0. This allows the analysis to figure out which texcoord is driving the lookup.
+			int CurrentCoordIndex = (Params.PixelPosition.y / TILE_RESOLUTION) % MAX_NUM_TEX_COORD;
+
+			UNROLL
+				for (int CoordinateIndex = 0; CoordinateIndex < NUM_MATERIAL_TEXCOORDS; ++CoordinateIndex)
+				{
+					// The quad material used only has a valid TexCoord for UV 0
+					MaterialParameters.TexCoords[CoordinateIndex] = CoordinateIndex != CurrentCoordIndex ? 0 : DebugInputs.TexCoord01.xy;
+				}
+
+			// The scale must also be computed to match the actual coord inspected. That would be required when rendering the mesh in world space
+			// Although not required here as using the material quad only use texcoord0.
+			// Params.OneOverDDU = GetComponent(Params.OneOverDDU, CurrentCoordIndex).xxxx;
+			// Params.OneOverDDV = GetComponent(Params.OneOverDDV, CurrentCoordIndex).xxxx;
+#endif
+		}
+
+	half3 BaseColor;
+	MaterialParameters.TexCoordScalesParams = Params;
+	{
+		FPixelMaterialInputs PixelMaterialInputs;
+		CalcMaterialParameters(MaterialParameters, PixelMaterialInputs, DebugInputs.SvPosition, bIsFrontFace);
+
+		// Sample material properties. The results are not used, but the calls to StoreTexCoordScale will still be made.
+		BaseColor = GetMaterialBaseColorRaw(PixelMaterialInputs);
+		half  Metallic = GetMaterialMetallicRaw(PixelMaterialInputs);
+		half  Specular = GetMaterialSpecularRaw(PixelMaterialInputs);
+		float Roughness = GetMaterialRoughnessRaw(PixelMaterialInputs);
+		half3 Normal = GetMaterialNormalRaw(PixelMaterialInputs);
+		half3 Emissive = GetMaterialEmissiveRaw(PixelMaterialInputs);
+		half Opacity = GetMaterialOpacityRaw(PixelMaterialInputs);
+#if MATERIALBLENDING_MASKED
+		half Mask = GetMaterialMask(PixelMaterialInputs);
+		if (!bOutputScales)
+		{
+			clip(GetMaterialMask(PixelMaterialInputs));
+		}
+#endif
+		half4 SSData = GetMaterialSubsurfaceDataRaw(PixelMaterialInputs);
+		float Custom0 = GetMaterialCustomData0(MaterialParameters);
+		float Custom1 = GetMaterialCustomData1(MaterialParameters);
+		float MaterialAO = GetMaterialAmbientOcclusionRaw(PixelMaterialInputs);
+		float PixelDepthOffset = GetMaterialPixelDepthOffset(PixelMaterialInputs);
+#if CLEAR_COAT_BOTTOM_NORMAL && NUM_MATERIAL_OUTPUTS_CLEARCOATBOTTOMNORMAL > 0
+		float3 BottomNormal = ClearCoatBottomNormal0(MaterialParameters);
+#endif
+	}
+	Params = MaterialParameters.TexCoordScalesParams;
+
+	Result *= saturate(Luminance(BaseColor));
+
+	float4 OutColor;
+	if (bOutputScales)
+	{
+		// Output the scales
+		OutColor = Params.ScalesPerIndex;
+	}
+	else
+	{
+		float PixelScale = ((Params.PixelPosition.x & 0x08) == (Params.PixelPosition.y & 0x08) || TextureAnalysisIndex != -1) ? Params.MinScale : Params.MaxScale;
+
+		// Output accuracy. If the GPU Scale is INITIAL_GPU_SCALE, then the texcoord was not updated.
+		if (Params.MinScale != INITIAL_GPU_SCALE)
+		{
+			Params.TexSample = TextureAnalysisIndex == -1 ? lerp(.4f, 1.f, saturate(Luminance(BaseColor))) : Params.TexSample;
+
+			float Accuracy = clamp(log2(PixelScale), -1.99, 1.99);
+			int ColorIndex = floor(Accuracy) + 2;
+			Result = Params.TexSample * lerp(DebugViewModeStruct.AccuracyColors[ColorIndex].rgb, DebugViewModeStruct.AccuracyColors[ColorIndex + 1].rgb, frac(Accuracy));
+		}
+
+		OutColor = float4(Result, PrimitiveAlpha);
+	}
+	return OutColor;
+}
+
+float4 VisualizeMeshTexCoordSizeAccuracy(in FDebugPSIn DebugInputs, in bool bIsFrontFace)
+{
+	float3 Result = float3(UNDEFINED_ACCURACY, UNDEFINED_ACCURACY, UNDEFINED_ACCURACY);
+
+	if (TexCoordAnalysisIndex >= 0)
+	{
+		float CPUSize = GetCPUTexelFactor(TexCoordAnalysisIndex);
+		if (CPUSize > 0)
+		{
+			float Accuracy = GetTexCoordSizeAccuracy(DebugInputs, TexCoordAnalysisIndex, CPUSize);
+			int ColorIndex = floor(Accuracy) + 2;
+			Result = lerp(DebugViewModeStruct.AccuracyColors[ColorIndex].rgb, DebugViewModeStruct.AccuracyColors[ColorIndex + 1].rgb, frac(Accuracy));
+		}
+	}
+	else
+	{
+		float MinAccuracy = UNDEF_ACCURACY;
+		float MaxAccuracy = -UNDEF_ACCURACY;
+
+		[unroll]
+		for (int CoordIndex = 0; CoordIndex < 4; ++CoordIndex)
+		{
+			float CPUSize = GetCPUTexelFactor(CoordIndex);
+			if (CPUSize > 0)
+			{
+				float Accuracy = GetTexCoordSizeAccuracy(DebugInputs, CoordIndex, CPUSize);
+				MinAccuracy = min(Accuracy, MinAccuracy);
+				MaxAccuracy = max(Accuracy, MaxAccuracy);
+			}
+		}
+
+		int2 PixelPosition = DebugInputs.SvPosition.xy;
+		float Accuracy = (PixelPosition.x & 0x08) == (PixelPosition.y & 0x08) ? MinAccuracy : MaxAccuracy;
+
+		if (abs(Accuracy) != UNDEF_ACCURACY)
+		{
+			int ColorIndex = floor(Accuracy) + 2;
+			Result = lerp(DebugViewModeStruct.AccuracyColors[ColorIndex].rgb, DebugViewModeStruct.AccuracyColors[ColorIndex + 1].rgb, frac(Accuracy));
+		}
+	}
+
+	return float4(Result, PrimitiveAlpha);
+}
+
+float4 VisualizePrimitiveDistanceAccuracy(in FDebugPSIn DebugInputs, in bool bIsFrontFace)
+{
+	float3 Result = float3(UNDEFINED_ACCURACY, UNDEFINED_ACCURACY, UNDEFINED_ACCURACY);
+
+	if (CPULogDistance >= 0)
+	{
+		float ViewDistance = length(SvPositionToResolvedTranslatedWorld(DebugInputs.SvPosition));
+		float GPULogDistance = log2(max(1, ViewDistance));
+
+		float Accuracy = clamp(GPULogDistance - CPULogDistance, -1.99, 1.99);
+		int ColorIndex = floor(Accuracy) + 2;
+
+		Result = lerp(DebugViewModeStruct.AccuracyColors[ColorIndex].rgb, DebugViewModeStruct.AccuracyColors[ColorIndex + 1].rgb, frac(Accuracy));
+	}
+
+	return float4(Result, PrimitiveAlpha);
+}
+
+float4 VisualizeRequiredTextureResolution(in FDebugPSIn DebugInputs, in bool bIsFrontFace)
+{
+	FMaterialPixelParameters MaterialParameters = GetMaterialPixelParameters(DebugInputs, DebugInputs.SvPosition);
+
+	MaterialParameters.TexCoordScalesParams.RequiredResolution = 0;
+	MaterialParameters.TexCoordScalesParams.TexSample = 0;
+	MaterialParameters.TexCoordScalesParams.VTMissingMipCount = -1;
+
+	half3 BaseColor;
+	{
+		FPixelMaterialInputs PixelMaterialInputs;
+		CalcMaterialParameters(MaterialParameters, PixelMaterialInputs, DebugInputs.SvPosition, bIsFrontFace);
+
+		// Sample material properties. The results are not used, but the calls to StoreTexCoordScale are still be made.
+		BaseColor = GetMaterialBaseColorRaw(PixelMaterialInputs);
+		half  Metallic = GetMaterialMetallicRaw(PixelMaterialInputs);
+		half  Specular = GetMaterialSpecularRaw(PixelMaterialInputs);
+		float Roughness = GetMaterialRoughnessRaw(PixelMaterialInputs);
+		half3 Normal = GetMaterialNormalRaw(PixelMaterialInputs);
+		half3 Emissive = GetMaterialEmissiveRaw(PixelMaterialInputs);
+		half Opacity = GetMaterialOpacityRaw(PixelMaterialInputs);
+#if MATERIALBLENDING_MASKED
+		half Mask = GetMaterialMask(PixelMaterialInputs);
+		clip(GetMaterialMask(PixelMaterialInputs));
+#endif
+		half4 SSData = GetMaterialSubsurfaceDataRaw(PixelMaterialInputs);
+		float Custom0 = GetMaterialCustomData0(MaterialParameters);
+		float Custom1 = GetMaterialCustomData1(MaterialParameters);
+		float MaterialAO = GetMaterialAmbientOcclusionRaw(PixelMaterialInputs);
+		float PixelDepthOffset = GetMaterialPixelDepthOffset(PixelMaterialInputs);
+#if CLEAR_COAT_BOTTOM_NORMAL && NUM_MATERIAL_OUTPUTS_CLEARCOATBOTTOMNORMAL > 0
+		float3 BottomNormal = ClearCoatBottomNormal0(MaterialParameters);
+#endif
+	}
+
+	float3 Result = float3(UNDEFINED_ACCURACY, UNDEFINED_ACCURACY, UNDEFINED_ACCURACY);
+	Result *= saturate(Luminance(BaseColor));
+
+	if (TextureAnalysisType == ANALYSIS_TEXTURE_RESOLUTION)
+	{
+		if (MaterialParameters.TexCoordScalesParams.RequiredResolution > 0)
+		{
+			float Accuracy = log2(TextureResolution / MaterialParameters.TexCoordScalesParams.RequiredResolution);
+			Result = AccuracyColorLookup(Accuracy);
+			Result *= MaterialParameters.TexCoordScalesParams.TexSample;
+		}
+	}
+	else if (TextureAnalysisType == ANALYSIS_VIRTUAL_TEXTURE_PENDING_MIPS)
+	{
+		if (MaterialParameters.TexCoordScalesParams.VTMissingMipCount >= 0)
+		{
+			float Accuracy = -MaterialParameters.TexCoordScalesParams.VTMissingMipCount;
+			Result = AccuracyColorLookup(Accuracy);
+			Result *= 0.5f;
+		}
+	}
+
+#if NUM_VIRTUALTEXTURE_SAMPLES || LIGHTMAP_VT_ENABLED
+	FinalizeVirtualTextureFeedback(
+		MaterialParameters.VirtualTextureFeedback,
+		MaterialParameters.SvPosition,
+		1.0,
+		View.FrameNumber,
+		View.VTFeedbackBuffer
+	);
+#endif
+
+	return float4(Result, PrimitiveAlpha);
+}
+
+float4 VisualizeShaderComplexityAccumulate(in FDebugPSIn DebugInputs, in bool bIsFrontFace)
+{
+	float3 Result = float3(1, 0, 1);
+	if (bOutputQuadOverdraw)
+	{
+		float3 FinalComplexity = NormalizedComplexity.xyz;
+
+		[branch]
+		if (bShowQuadOverdraw && NormalizedComplexity.x > 0)
+		{
+			uint Coverage = ComputeQuadCoverage(DebugInputs.SvPosition.xy, DebugInputs.SvPrimitiveID, 24, false, false, 0);
+			// The extra cost from quad overdraw is assumed to be linear.
+			FinalComplexity.x *= 4.f / (float)(Coverage);
+		}
+
+		// use the maximum range allowed for scene color
+		Result = FinalComplexity;
+	}
+	else
+	{
+		Result = NormalizedComplexity.xyz;
+	}
+
+	// alpha channel needs to be 1 to have decals working where the framebuffer blend is setup to multiply with alpha
+	return float4(Result, 1);
+}
+
+float4 VisualizeLODColoration(in FDebugPSIn DebugInputs, in bool bIsFrontFace)
+{
+	float3 Result = DebugViewModeStruct.LODColors[LODIndex].rgb;
+
+	FMaterialPixelParameters MaterialParameters = GetMaterialPixelParameters(DebugInputs, DebugInputs.SvPosition);
+
+	half3 BaseColor;
+	half Opacity;
+	{
+		FPixelMaterialInputs PixelMaterialInputs;
+		CalcMaterialParameters(MaterialParameters, PixelMaterialInputs, DebugInputs.SvPosition, bIsFrontFace);
+
+		// Sample material properties. The results are not used, but the calls to StoreTexCoordScale are still be made.
+		BaseColor = GetMaterialBaseColorRaw(PixelMaterialInputs);
+		Opacity = GetMaterialOpacityRaw(PixelMaterialInputs);
+
+		half3 Emissive = GetMaterialEmissiveRaw(PixelMaterialInputs);
+#if MATERIALBLENDING_ADDITIVE
+		Opacity = (.05 + .95 * Luminance(Emissive));
+#else
+		BaseColor += Emissive;
+#endif
+
+#if MATERIALBLENDING_MASKED
+		clip(GetMaterialMask(PixelMaterialInputs));
+#endif
+	}
+
+	Result *= (.05 + .95 * Luminance(BaseColor));
+
+#if NUM_VIRTUALTEXTURE_SAMPLES || LIGHTMAP_VT_ENABLED
+	FinalizeVirtualTextureFeedback(
+		MaterialParameters.VirtualTextureFeedback,
+		MaterialParameters.SvPosition,
+		1.0,
+		View.FrameNumber,
+		View.VTFeedbackBuffer
+	);
+#endif
+
+	return float4(Result, Opacity);
+}
+
+void Main(
+	in FDebugPSIn DebugInputs
+	OPTIONAL_IsFrontFace,
+	out float4 OutColor : SV_Target0
+	)
+{
+#if INSTANCED_STEREO
+	ResolvedView = ResolveView(DebugInputs.EyeIndex);
+#else
+	ResolvedView = ResolveView();
+#endif
+
+	// If we see pink something went wrong.
+	float4 VisualizeColor = float4(1, 0, 1, 1);
+
+	BRANCH
+	if (VisualizeMode == DVSM_ShaderComplexity ||
+		VisualizeMode == DVSM_ShaderComplexityContainedQuadOverhead ||
+		VisualizeMode == DVSM_ShaderComplexityBleedingQuadOverhead ||
+		VisualizeMode == DVSM_QuadComplexity)
+	{
+		VisualizeColor = VisualizeShaderComplexityAccumulate(DebugInputs, bIsFrontFace);
+	}
+	else if (VisualizeMode == DVSM_PrimitiveDistanceAccuracy)
+	{
+		VisualizeColor = VisualizePrimitiveDistanceAccuracy(DebugInputs, bIsFrontFace);
+	}
+	else if (VisualizeMode == DVSM_MeshUVDensityAccuracy)
+	{
+		VisualizeColor = VisualizeMeshTexCoordSizeAccuracy(DebugInputs, bIsFrontFace);
+	}
+	else if (VisualizeMode == DVSM_MaterialTextureScaleAccuracy ||
+			 VisualizeMode == DVSM_OutputMaterialTextureScales)
+	{
+		VisualizeColor = VisualizeMaterialTexCoordScales(DebugInputs, bIsFrontFace);
+	}
+	else if (VisualizeMode == DVSM_RequiredTextureResolution ||
+			 VisualizeMode == DVSM_VirtualTexturePendingMips)
+	{
+		VisualizeColor = VisualizeRequiredTextureResolution(DebugInputs, bIsFrontFace);
+	}
+	else if (VisualizeMode == DVSM_LODColoration)
+	{
+		VisualizeColor = VisualizeLODColoration(DebugInputs, bIsFrontFace);
+	}
+
+	OutColor = RETURN_COLOR(VisualizeColor);	
+}
+
+
+
+
+
+
diff --git a/Engine/Shaders/Private/MaterialTexCoordScalesPixelShader.usf b/Engine/Shaders/Private/MaterialTexCoordScalesPixelShader.usf
deleted file mode 100644
diff --git a/Engine/Shaders/Private/MeshTexCoordSizeAccuracyPixelShader.usf b/Engine/Shaders/Private/MeshTexCoordSizeAccuracyPixelShader.usf
deleted file mode 100644
diff --git a/Engine/Shaders/Private/PrimitiveDistanceAccuracyPixelShader.usf b/Engine/Shaders/Private/PrimitiveDistanceAccuracyPixelShader.usf
deleted file mode 100644
diff --git a/Engine/Shaders/Private/RequiredTextureResolutionPixelShader.usf b/Engine/Shaders/Private/RequiredTextureResolutionPixelShader.usf
deleted file mode 100644
diff --git a/Engine/Shaders/Private/ShaderComplexityAccumulatePixelShader.usf b/Engine/Shaders/Private/ShaderComplexityAccumulatePixelShader.usf
deleted file mode 100644
diff --git a/Engine/Shaders/Private/ViewMode/LODColorationPixelShader.usf b/Engine/Shaders/Private/ViewMode/LODColorationPixelShader.usf
deleted file mode 100644
diff --git a/Engine/Source/Runtime/Engine/Private/DebugViewModeHelpers.cpp b/Engine/Source/Runtime/Engine/Private/DebugViewModeHelpers.cpp
--- a/Engine/Source/Runtime/Engine/Private/DebugViewModeHelpers.cpp
+++ b/Engine/Source/Runtime/Engine/Private/DebugViewModeHelpers.cpp
@@ -18,6 +18,41 @@
 
 #define LOCTEXT_NAMESPACE "LogDebugViewMode"
 
+const TCHAR* DebugViewShaderModeToString(EDebugViewShaderMode InShaderMode)
+{
+	switch (InShaderMode)
+	{
+	case DVSM_None:
+		return TEXT("DVSM_None");
+	case DVSM_ShaderComplexity:
+		return TEXT("DVSM_ShaderComplexity");
+	case DVSM_ShaderComplexityContainedQuadOverhead:
+		return TEXT("DVSM_ShaderComplexityContainedQuadOverhead");
+	case DVSM_ShaderComplexityBleedingQuadOverhead:
+		return TEXT("DVSM_ShaderComplexityBleedingQuadOverhead");
+	case DVSM_QuadComplexity:
+		return TEXT("DVSM_QuadComplexity");
+	case DVSM_PrimitiveDistanceAccuracy:
+		return TEXT("DVSM_PrimitiveDistanceAccuracy");
+	case DVSM_MeshUVDensityAccuracy:
+		return TEXT("DVSM_MeshUVDensityAccuracy");
+	case DVSM_MaterialTextureScaleAccuracy:
+		return TEXT("DVSM_MaterialTextureScaleAccuracy");
+	case DVSM_OutputMaterialTextureScales:
+		return TEXT("DVSM_OutputMaterialTextureScales");
+	case DVSM_RequiredTextureResolution:
+		return TEXT("DVSM_RequiredTextureResolution");
+	case DVSM_VirtualTexturePendingMips:
+		return TEXT("DVSM_VirtualTexturePendingMips");
+	case DVSM_RayTracingDebug:
+		return TEXT("DVSM_RayTracingDebug");
+	case DVSM_LODColoration:
+		return TEXT("DVSM_LODColoration");
+	default:
+		return TEXT("DVSM_None");
+	}
+}
+
 #if WITH_DEBUG_VIEW_MODES
 
 static bool PlatformSupportsDebugViewShaders(EShaderPlatform Platform)
@@ -68,32 +103,21 @@ bool AllowDebugViewShaderMode(EDebugViewShaderMode ShaderMode, EShaderPlatform P
 #endif
 }
 
-bool ShouldCompileDebugViewModeShader(EDebugViewShaderMode ShaderMode, const FMeshMaterialShaderPermutationParameters& Parameters)
+bool ShouldCompileDebugViewModeShader(const FMeshMaterialShaderPermutationParameters& Parameters)
 {
-	if (!EnumHasAllFlags(Parameters.Flags, EShaderPermutationFlags::HasEditorOnlyData))
+	if (!PlatformSupportsDebugViewShaders(Parameters.Platform))
 	{
-		// Debug view shaders only in editor
 		return false;
 	}
 
-	if (!AllowDebugViewShaderMode(ShaderMode, Parameters.Platform, Parameters.MaterialParameters.FeatureLevel))
+	if (Parameters.MaterialParameters.FeatureLevel < ERHIFeatureLevel::SM5)
 	{
-		// Don't support this mode
 		return false;
 	}
 
-	const FDebugViewModeInterface* DebugViewModeInterface = FDebugViewModeInterface::GetInterface(ShaderMode);
-	if (!DebugViewModeInterface->bNeedsMaterialProperties &&
-		!Parameters.MaterialParameters.bIsDefaultMaterial &&
-		FDebugViewModeInterface::AllowFallbackToDefaultMaterial(Parameters.MaterialParameters.bHasVertexPositionOffsetConnected, Parameters.MaterialParameters.bHasPixelDepthOffsetConnected))
-	{
-		// We can replace this material with the default material
-		return false;
-	}
-
-	if (DebugViewModeInterface->bNeedsOnlyLocalVertexFactor && Parameters.VertexFactoryType->GetFName() != TEXT("FLocalVertexFactory"))
+	if (!EnumHasAllFlags(Parameters.Flags, EShaderPermutationFlags::HasEditorOnlyData))
 	{
-		// This debug view mode only needed for local vertex factory
+		// Debug view shaders only in editor
 		return false;
 	}
 
@@ -303,13 +327,6 @@ bool CompileDebugViewModeShaders(EDebugViewShaderMode ShaderMode, EMaterialQuali
 			bool bMaterialFinished = true;
 			if (Material && Material->GetGameThreadShaderMap())
 			{
-				if (!DebugViewModeInterface->bNeedsMaterialProperties &&
-					FDebugViewModeInterface::AllowFallbackToDefaultMaterial(Material))
-				{
-					Material = UMaterial::GetDefaultMaterial(MD_Surface)->GetMaterialResource(FeatureLevel, QualityLevel);
-					check(Material);
-				}
-
 				FMaterialShaderTypes ShaderTypes;
 				DebugViewModeInterface->AddShaderTypes(FeatureLevel, LocalVertexFactory, ShaderTypes);
 				if (Material->ShouldCacheShaders(ShaderTypes, LocalVertexFactory) && !Material->HasShaders(ShaderTypes, LocalVertexFactory))
diff --git a/Engine/Source/Runtime/Engine/Private/DebugViewModeInterface.cpp b/Engine/Source/Runtime/Engine/Private/DebugViewModeInterface.cpp
--- a/Engine/Source/Runtime/Engine/Private/DebugViewModeInterface.cpp
+++ b/Engine/Source/Runtime/Engine/Private/DebugViewModeInterface.cpp
@@ -10,39 +10,67 @@ DebugViewModeInterface.cpp: Contains definitions for rendering debug viewmodes.
 
 #if ENABLE_DRAW_DEBUG
 
-FDebugViewModeInterface* FDebugViewModeInterface::Singletons[DVSM_MAX] = {}; // Init to null.
+FDebugViewModeInterface* FDebugViewModeInterface::Singleton = nullptr;
 
-void FDebugViewModeInterface::SetDrawRenderState(EBlendMode BlendMode, FRenderState& DrawRenderState, bool bHasDepthPrepassForMaskedMaterial) const
+void FDebugViewModeInterface::SetDrawRenderState(EDebugViewShaderMode DebugViewMode, EBlendMode BlendMode, FRenderState& DrawRenderState, bool bHasDepthPrepassForMaskedMaterial) const
 {
-	if (IsTranslucentBlendMode(BlendMode))
+	if (DebugViewMode == DVSM_QuadComplexity || DebugViewMode == DVSM_ShaderComplexityBleedingQuadOverhead || DebugViewMode == DVSM_ShaderComplexityContainedQuadOverhead || DebugViewMode == DVSM_ShaderComplexity)
 	{
-		// Otherwise, force translucent blend mode (shaders will use an hardcoded alpha).
-		DrawRenderState.BlendState = TStaticBlendState<CW_RGBA, BO_Add, BF_SourceAlpha, BF_InverseSourceAlpha, BO_Add, BF_Zero, BF_InverseSourceAlpha>::GetRHI();
-		DrawRenderState.DepthStencilState = TStaticDepthStencilState<false, CF_DepthNearOrEqual>::GetRHI();
+		if (BlendMode == BLEND_Opaque)
+		{
+			DrawRenderState.DepthStencilState = TStaticDepthStencilState<true, CF_DepthNearOrEqual>::GetRHI();
+		}
+		else if (BlendMode == BLEND_Masked)
+		{
+			if (bHasDepthPrepassForMaskedMaterial)
+			{
+				DrawRenderState.DepthStencilState = TStaticDepthStencilState<false, CF_Equal>::GetRHI();
+			}
+			else
+			{
+				DrawRenderState.DepthStencilState = TStaticDepthStencilState<false, CF_DepthNearOrEqual>::GetRHI();
+			}
+		}
+		else // Translucent
+		{
+			DrawRenderState.DepthStencilState = TStaticDepthStencilState<false, CF_DepthNearOrEqual>::GetRHI();
+		}
+		DrawRenderState.BlendState = TStaticBlendState<CW_RGBA, BO_Add, BF_One, BF_One, BO_Add, BF_Zero, BF_One>::GetRHI();
 	}
-	else
+	else if (DebugViewMode == DVSM_OutputMaterialTextureScales)
 	{
 		DrawRenderState.BlendState = TStaticBlendState<>::GetRHI();
-
-		// If not selected, use depth equal to make alpha test stand out (goes with EarlyZPassMode = DDM_AllOpaque) 
-		if (BlendMode == BLEND_Masked && bHasDepthPrepassForMaskedMaterial)
+		DrawRenderState.DepthStencilState = TStaticDepthStencilState<false, CF_Always>::GetRHI();
+	}
+	else
+	{
+		if (IsTranslucentBlendMode(BlendMode))
 		{
-			DrawRenderState.DepthStencilState = TStaticDepthStencilState<false, CF_Equal>::GetRHI();
+			// Otherwise, force translucent blend mode (shaders will use an hardcoded alpha).
+			DrawRenderState.BlendState = TStaticBlendState<CW_RGBA, BO_Add, BF_SourceAlpha, BF_InverseSourceAlpha, BO_Add, BF_Zero, BF_InverseSourceAlpha>::GetRHI();
+			DrawRenderState.DepthStencilState = TStaticDepthStencilState<false, CF_DepthNearOrEqual>::GetRHI();
 		}
 		else
 		{
-			DrawRenderState.DepthStencilState = TStaticDepthStencilState<>::GetRHI();
+			DrawRenderState.BlendState = TStaticBlendState<>::GetRHI();
+
+			// If not selected, use depth equal to make alpha test stand out (goes with EarlyZPassMode = DDM_AllOpaque) 
+			if (BlendMode == BLEND_Masked && bHasDepthPrepassForMaskedMaterial)
+			{
+				DrawRenderState.DepthStencilState = TStaticDepthStencilState<false, CF_Equal>::GetRHI();
+			}
+			else
+			{
+				DrawRenderState.DepthStencilState = TStaticDepthStencilState<>::GetRHI();
+			}
 		}
 	}
 }
 
-void FDebugViewModeInterface::SetInterface(EDebugViewShaderMode InDebugViewMode, FDebugViewModeInterface* Interface)
+void FDebugViewModeInterface::SetInterface(FDebugViewModeInterface* Interface)
 {
-	if ((uint32)InDebugViewMode < DVSM_MAX)
-	{
-		ensure(!Singletons[InDebugViewMode]);
-		Singletons[InDebugViewMode] = Interface;
-	}
+	ensure(!Singleton);
+	Singleton = Interface;
 }
 
 bool FDebugViewModeInterface::AllowFallbackToDefaultMaterial(bool bHasVertexPositionOffsetConnected, bool bHasPixelDepthOffsetConnected)
diff --git a/Engine/Source/Runtime/Engine/Public/DebugViewModeHelpers.h b/Engine/Source/Runtime/Engine/Public/DebugViewModeHelpers.h
--- a/Engine/Source/Runtime/Engine/Public/DebugViewModeHelpers.h
+++ b/Engine/Source/Runtime/Engine/Public/DebugViewModeHelpers.h
@@ -37,12 +37,14 @@ enum EDebugViewShaderMode
 	DVSM_MAX
 };
 
+ENGINE_API const TCHAR* DebugViewShaderModeToString(EDebugViewShaderMode InShaderMode);
+
 #if WITH_DEBUG_VIEW_MODES
 /** Returns true if the vertex shader (and potential hull and domain) should be compiled on the given platform. */
 ENGINE_API bool AllowDebugViewVSDSHS(EShaderPlatform Platform);
 /** Returns true if the shader mode can be enabled. This is only for UI elements as no shader platform is actually passed. */
 ENGINE_API bool AllowDebugViewShaderMode(EDebugViewShaderMode ShaderMode, EShaderPlatform Platform, ERHIFeatureLevel::Type FeatureLevel);
-ENGINE_API bool ShouldCompileDebugViewModeShader(EDebugViewShaderMode ShaderMode, const FMeshMaterialShaderPermutationParameters& Parameters);
+ENGINE_API bool ShouldCompileDebugViewModeShader(const FMeshMaterialShaderPermutationParameters& Parameters);
 #else
 FORCEINLINE bool AllowDebugViewVSDSHS(EShaderPlatform Platform)  { return false; }
 FORCEINLINE bool AllowDebugViewShaderMode(EDebugViewShaderMode ShaderMode, EShaderPlatform Platform, ERHIFeatureLevel::Type FeatureLevel) { return false; }
diff --git a/Engine/Source/Runtime/Engine/Public/DebugViewModeInterface.h b/Engine/Source/Runtime/Engine/Public/DebugViewModeInterface.h
--- a/Engine/Source/Runtime/Engine/Public/DebugViewModeInterface.h
+++ b/Engine/Source/Runtime/Engine/Public/DebugViewModeInterface.h
@@ -39,16 +39,7 @@ public:
 		FRHIDepthStencilState*	DepthStencilState;
 	};
 
-	FDebugViewModeInterface(
-		const TCHAR* InPixelShaderName,		// Shader class name implementing the PS
-		bool InNeedsOnlyLocalVertexFactor,  // Whether this viewmode will only be used with the local vertex factory (for draw tiled mesh).
-		bool InNeedsMaterialProperties,		// Whether the PS use any of the material properties (otherwise default material will be used, reducing shader compilation).
-		bool InNeedsInstructionCount		// Whether FDebugViewModePS::GetDebugViewModeShaderBindings() will use the num of instructions.
-	)
-		: PixelShaderName(InPixelShaderName)
-		, bNeedsOnlyLocalVertexFactor(InNeedsOnlyLocalVertexFactor)
-		, bNeedsMaterialProperties(InNeedsMaterialProperties)
-		, bNeedsInstructionCount(InNeedsInstructionCount)
+	FDebugViewModeInterface()
 	{}
 
 	virtual ~FDebugViewModeInterface() {}
@@ -57,7 +48,7 @@ public:
 		const FVertexFactoryType* InVertexFactoryType,
 		FMaterialShaderTypes& OutShaderTypes) const = 0;
 
-	virtual void SetDrawRenderState(EBlendMode BlendMode, FRenderState& DrawRenderState, bool bHasDepthPrepassForMaskedMaterial) const;
+	virtual void SetDrawRenderState(EDebugViewShaderMode DebugViewMode, EBlendMode BlendMode, FRenderState& DrawRenderState, bool bHasDepthPrepassForMaskedMaterial) const;
 
 	virtual void GetDebugViewModeShaderBindings(
 		const FDebugViewModePS& Shader,
@@ -75,26 +66,14 @@ public:
 		FMeshDrawSingleShaderBindings& ShaderBindings
 	) const {}
 
-	/** The shader class name, used to filter out shaders that need to be compiled. */
-	const TCHAR* PixelShaderName;
-
-	/** Whether only tiled mesh draw will be required. */
-	const bool bNeedsOnlyLocalVertexFactor;
-
-	/** Whether the viewmode any of material properties (otherwise it can fallback to using the default material) */
-	const bool bNeedsMaterialProperties;
-
-	/** Whether FDebugViewModePS::GetDebugViewModeShaderBindings() will use the num of instructions. */
-	const bool bNeedsInstructionCount;
-
 	/** Return the interface object for the given viewmode. */
 	static const FDebugViewModeInterface* GetInterface(EDebugViewShaderMode InDebugViewMode) 
 	{
-		return (uint32)InDebugViewMode < DVSM_MAX ? Singletons[InDebugViewMode] : nullptr;
+		return (uint32)InDebugViewMode < DVSM_MAX ? Singleton : nullptr;
 	}
 
 	/** Return the interface object for the given viewmode. */
-	static void SetInterface(EDebugViewShaderMode InDebugViewMode, FDebugViewModeInterface* Interface);
+	static void SetInterface(FDebugViewModeInterface* Interface);
 	
 	/** Whether this material can be substituted by the default material. */
 	static bool AllowFallbackToDefaultMaterial(const FMaterial* InMaterial);
@@ -102,7 +81,7 @@ public:
 
 private:
 	
-	static FDebugViewModeInterface* Singletons[DVSM_MAX];
+	static FDebugViewModeInterface* Singleton;
 };
 
 #endif // ENABLE_DRAW_DEBUG
diff --git a/Engine/Source/Runtime/Renderer/Private/DebugViewModeRendering.cpp b/Engine/Source/Runtime/Renderer/Private/DebugViewModeRendering.cpp
--- a/Engine/Source/Runtime/Renderer/Private/DebugViewModeRendering.cpp
+++ b/Engine/Source/Runtime/Renderer/Private/DebugViewModeRendering.cpp
@@ -6,12 +6,6 @@ DebugViewModeRendering.cpp: Contains definitions for rendering debug viewmodes.
 
 #include "DebugViewModeRendering.h"
 #include "Materials/Material.h"
-#include "ShaderComplexityRendering.h"
-#include "PrimitiveDistanceAccuracyRendering.h"
-#include "MeshTexCoordSizeAccuracyRendering.h"
-#include "MaterialTexCoordScalesRendering.h"
-#include "RequiredTextureResolutionRendering.h"
-#include "ViewMode/LODColorationRendering.h"
 #include "PrimitiveSceneInfo.h"
 #include "ScenePrivate.h"
 #include "PostProcessing.h"
@@ -30,6 +24,64 @@ IMPLEMENT_STATIC_UNIFORM_BUFFER_STRUCT(FDebugViewModePassUniformParameters, "Deb
 
 #if WITH_DEBUG_VIEW_MODES
 
+int32 GCacheShaderComplexityShaders = 0;
+static FAutoConsoleVariableRef CVarCacheShaderComplexityShaders(
+	TEXT("r.ShaderComplexity.CacheShaders"),
+	GCacheShaderComplexityShaders,
+	TEXT("If non zero, store the shader complexity shaders in the material shader map, to prevent compile on-the-fly lag. (default=0)"),
+	ECVF_ReadOnly
+);
+
+int32 GShaderComplexityBaselineForwardVS = 134;
+static FAutoConsoleVariableRef CVarShaderComplexityBaselineForwardVS(
+	TEXT("r.ShaderComplexity.Baseline.Forward.VS"),
+	GShaderComplexityBaselineForwardVS,
+	TEXT("Minimum number of instructions for vertex shaders in forward shading (default=134)"),
+	ECVF_Default
+);
+
+int32 GShaderComplexityBaselineForwardPS = 635;
+static FAutoConsoleVariableRef CVarShaderComplexityBaselineForwardPS(
+	TEXT("r.ShaderComplexity.Baseline.Forward.PS"),
+	GShaderComplexityBaselineForwardPS,
+	TEXT("Minimum number of instructions for pixel shaders in forward shading (default=635)"),
+	ECVF_Default
+);
+
+int32 GShaderComplexityBaselineForwardUnlitPS = 47;
+static FAutoConsoleVariableRef CVarShaderComplexityBaselineForwardUnlitPS(
+	TEXT("r.ShaderComplexity.Baseline.Forward.UnlitPS"),
+	GShaderComplexityBaselineForwardUnlitPS,
+	TEXT("Minimum number of instructions for unlit material pixel shaders in forward shading (default=47)"),
+	ECVF_Default
+);
+
+int32 GShaderComplexityBaselineDeferredVS = 41;
+static FAutoConsoleVariableRef CVarShaderComplexityBaselineDeferredVS(
+	TEXT("r.ShaderComplexity.Baseline.Deferred.VS"),
+	GShaderComplexityBaselineDeferredVS,
+	TEXT("Minimum number of instructions for vertex shaders in deferred shading (default=41)"),
+	ECVF_Default
+);
+
+int32 GShaderComplexityBaselineDeferredPS = 111;
+static FAutoConsoleVariableRef CVarShaderComplexityBaselineDeferredPS(
+	TEXT("r.ShaderComplexity.Baseline.Deferred.PS"),
+	GShaderComplexityBaselineDeferredPS,
+	TEXT("Minimum number of instructions for pixel shaders in deferred shading (default=111)"),
+	ECVF_Default
+);
+
+int32 GShaderComplexityBaselineDeferredUnlitPS = 33;
+static FAutoConsoleVariableRef CVarShaderComplexityBaselineDeferredUnlitPS(
+	TEXT("r.ShaderComplexity.Baseline.Deferred.UnlitPS"),
+	GShaderComplexityBaselineDeferredUnlitPS,
+	TEXT("Minimum number of instructions for unlit material pixel shaders in deferred shading (default=33)"),
+	ECVF_Default
+);
+
+IMPLEMENT_MATERIAL_SHADER_TYPE(, FDebugViewModePS, TEXT("/Engine/Private/DebugViewModePixelShader.usf"), TEXT("Main"), SF_Pixel);
+
 int32 GetQuadOverdrawUAVIndex(EShaderPlatform Platform, ERHIFeatureLevel::Type FeatureLevel)
 {
 	if (IsSimpleForwardShadingEnabled(Platform))
@@ -146,6 +198,11 @@ void RenderDebugViewMode(FRDGBuilder& GraphBuilder, TArrayView<FViewInfo> Views,
 	}
 }
 
+bool FDebugViewModePS::ShouldCompilePermutation(const FMeshMaterialShaderPermutationParameters& Parameters)
+{
+	return ShouldCompileDebugViewModeShader(Parameters);
+}
+
 void FDebugViewModePS::GetElementShaderBindings(
 	const FShaderMapPointerTable& PointerTable,
 	const FScene* Scene,
@@ -210,13 +267,6 @@ FDebugViewModeMeshProcessor::FDebugViewModeMeshProcessor(
 	}
 }
 
-void AddDebugViewModeShaderTypes(ERHIFeatureLevel::Type FeatureLevel,
-	const FVertexFactoryType* VertexFactoryType,
-	FMaterialShaderTypes& OutShaderTypes)
-{
-	OutShaderTypes.AddShaderType<FDebugViewModeVS>();
-}
-
 void FDebugViewModeMeshProcessor::AddMeshBatch(const FMeshBatch& RESTRICT MeshBatch, uint64 BatchElementMask, const FPrimitiveSceneProxy* RESTRICT PrimitiveSceneProxy, int32 StaticMeshId)
 {
 	if (!DebugViewModeInterface)
@@ -232,12 +282,6 @@ void FDebugViewModeMeshProcessor::AddMeshBatch(const FMeshBatch& RESTRICT MeshBa
 	}
 
 	const FMaterial* Material = BatchMaterial;
-	if (!DebugViewModeInterface->bNeedsMaterialProperties && FDebugViewModeInterface::AllowFallbackToDefaultMaterial(Material))
-	{
-		MaterialRenderProxy = UMaterial::GetDefaultMaterial(MD_Surface)->GetRenderProxy();
-		Material = MaterialRenderProxy->GetMaterialNoFallback(FeatureLevel);
-		check(Material);
-	}
 
 	FVertexFactoryType* VertexFactoryType = MeshBatch.VertexFactory->GetType();
 
@@ -265,7 +309,7 @@ void FDebugViewModeMeshProcessor::AddMeshBatch(const FMeshBatch& RESTRICT MeshBa
 	FMeshPassProcessorRenderState DrawRenderState;
 
 	FDebugViewModeInterface::FRenderState InterfaceRenderState;
-	DebugViewModeInterface->SetDrawRenderState(BatchMaterial->GetBlendMode(), InterfaceRenderState, Scene ? (Scene->GetShadingPath() == EShadingPath::Deferred && Scene->EarlyZPassMode != DDM_NonMaskedOnly) : false);
+	DebugViewModeInterface->SetDrawRenderState(DebugViewMode, BatchMaterial->GetBlendMode(), InterfaceRenderState, Scene ? (Scene->GetShadingPath() == EShadingPath::Deferred && Scene->EarlyZPassMode != DDM_NonMaskedOnly) : false);
 	DrawRenderState.SetBlendState(InterfaceRenderState.BlendState);
 	DrawRenderState.SetDepthStencilState(InterfaceRenderState.DepthStencilState);
 
@@ -279,10 +323,7 @@ void FDebugViewModeMeshProcessor::AddMeshBatch(const FMeshBatch& RESTRICT MeshBa
 		ViewModeParamName);
 
 	// Shadermap can be null while shaders are compiling.
-	if (DebugViewModeInterface->bNeedsInstructionCount)
-	{
-		UpdateInstructionCount(ShaderElementData, BatchMaterial, VertexFactoryType);
-	}
+	UpdateInstructionCount(ShaderElementData, BatchMaterial, VertexFactoryType);
 
 	ShaderElementData.InitializeMeshMaterialData(ViewIfDynamicMeshCommand, PrimitiveSceneProxy, MeshBatch, StaticMeshId, true);
 
@@ -359,6 +400,163 @@ void FDebugViewModeMeshProcessor::UpdateInstructionCount(FDebugViewModeShaderEle
 	}
 }
 
+void FDebugViewModeImplementation::AddShaderTypes(ERHIFeatureLevel::Type InFeatureLevel,
+	const FVertexFactoryType* InVertexFactoryType,
+	FMaterialShaderTypes& OutShaderTypes) const
+{
+	OutShaderTypes.AddShaderType<FDebugViewModeVS>();
+	OutShaderTypes.AddShaderType<FDebugViewModePS>();
+}
+
+void FDebugViewModeImplementation::GetDebugViewModeShaderBindings(
+	const FDebugViewModePS& Shader,
+	const FPrimitiveSceneProxy* PrimitiveSceneProxy,
+	const FMaterialRenderProxy& MaterialRenderProxy,
+	const FMaterial& Material,
+	EDebugViewShaderMode DebugViewMode,
+	const FVector& ViewOrigin,
+	int32 VisualizeLODIndex,
+	int32 VisualizeElementIndex,
+	int32 NumVSInstructions,
+	int32 NumPSInstructions,
+	int32 ViewModeParam,
+	FName ViewModeParamName,
+	FMeshDrawSingleShaderBindings& ShaderBindings) const
+{
+	FVector4f OneOverCPUTexCoordScales[TEXSTREAM_MAX_NUM_TEXTURES_PER_MATERIAL / 4];
+	FIntVector4 TexCoordIndices[TEXSTREAM_MAX_NUM_TEXTURES_PER_MATERIAL / 4];
+	FVector4 WorldUVDensities;
+	FVector4f NormalizedComplexityValue;
+	FIntVector AnalysisParameter;
+	const float PrimitiveAlpha = (!PrimitiveSceneProxy || PrimitiveSceneProxy->IsSelected()) ? 1.f : .2f;
+	const int32 TexCoordAnalysisIndex = ViewModeParam >= 0 ? FMath::Clamp<int32>(ViewModeParam, 0, MAX_TEXCOORDS - 1) : -1;
+	float CPULogDistance = -1.0f;
+	int32 bShowQuadOverdraw = 0;
+	const int32 boolOutputQuadOverdraw = (DebugViewMode == DVSM_QuadComplexity) || (DebugViewMode == DVSM_ShaderComplexityContainedQuadOverhead) ? 1 : 0;
+	const int32 LODIndex = FMath::Clamp(VisualizeLODIndex, 0, NumLODColorationColors - 1);
+	const int32 AnalysisType = DebugViewMode == DVSM_VirtualTexturePendingMips ? 1 : 0;
+
+	FMemory::Memzero(OneOverCPUTexCoordScales); // 0 remap to irrelevant data.
+	FMemory::Memzero(TexCoordIndices);
+	FMemory::Memzero(WorldUVDensities);
+
+	// Gather Data
+#if WITH_EDITORONLY_DATA
+	if (PrimitiveSceneProxy)
+	{
+		PrimitiveSceneProxy->GetMaterialTextureScales(VisualizeLODIndex, VisualizeElementIndex, nullptr, OneOverCPUTexCoordScales, TexCoordIndices);
+		PrimitiveSceneProxy->GetMeshUVDensities(VisualizeLODIndex, VisualizeElementIndex, WorldUVDensities);
+
+		float Distance = 0;
+		if (PrimitiveSceneProxy->GetPrimitiveDistance(VisualizeLODIndex, VisualizeElementIndex, ViewOrigin, Distance))
+		{
+			// Because the streamer use FMath::FloorToFloat, here we need to use -1 to have a useful result.
+			CPULogDistance = FMath::Max<float>(0.f, FMath::Log2(FMath::Max<float>(1.f, Distance)));
+		}
+	}
+#endif
+
+	if (DebugViewMode == DVSM_OutputMaterialTextureScales || DebugViewMode == DVSM_MaterialTextureScaleAccuracy)
+	{
+		const bool bOutputScales = DebugViewMode == DVSM_OutputMaterialTextureScales;
+		const int32 AnalysisIndex = ViewModeParam >= 0 ? FMath::Clamp<int32>(ViewModeParam, 0, TEXSTREAM_MAX_NUM_TEXTURES_PER_MATERIAL - 1) : -1;
+		AnalysisParameter = FIntVector(bOutputScales ? -1 : AnalysisIndex, bOutputScales ? 1 : 0, AnalysisType);
+	}
+	else if (DebugViewMode == DVSM_RequiredTextureResolution || DebugViewMode == DVSM_VirtualTexturePendingMips)
+	{
+		int32 AnalysisIndex = INDEX_NONE;
+		int32 TextureResolution = 64;
+		FMaterialRenderContext MaterialContext(&MaterialRenderProxy, Material, nullptr);
+		const FUniformExpressionSet& UniformExpressions = Material.GetUniformExpressions();
+		EMaterialTextureParameterType TextureTypes[] = { EMaterialTextureParameterType::Standard2D, EMaterialTextureParameterType::Virtual };
+		if (ViewModeParam != INDEX_NONE && ViewModeParamName == NAME_None) // If displaying texture per texture indices
+		{
+			for (EMaterialTextureParameterType TextureType : TextureTypes)
+			{
+				for (int32 ParameterIndex = 0; ParameterIndex < UniformExpressions.GetNumTextures(TextureType); ++ParameterIndex)
+				{
+					const FMaterialTextureParameterInfo& Parameter = UniformExpressions.GetTextureParameter(TextureType, ParameterIndex);
+					if (Parameter.TextureIndex == ViewModeParam)
+					{
+						const UTexture* Texture = nullptr;
+						UniformExpressions.GetTextureValue(TextureType, ParameterIndex, MaterialContext, Material, Texture);
+						if (Texture && Texture->GetResource())
+						{
+							AnalysisIndex = ViewModeParam;
+
+							if (Texture->IsStreamable())
+							{
+								TextureResolution = 1 << (Texture->GetResource()->GetCurrentMipCount() - 1);
+							}
+							else
+							{
+								TextureResolution = FMath::Max(Texture->GetResource()->GetSizeX(), Texture->GetResource()->GetSizeY());
+							}
+						}
+					}
+				}
+			}
+		}
+		else if (ViewModeParam != INDEX_NONE) // Otherwise show only texture matching the given name
+		{
+			for (EMaterialTextureParameterType TextureType : TextureTypes)
+			{
+				for (int32 ParameterIndex = 0; ParameterIndex < UniformExpressions.GetNumTextures(TextureType); ++ParameterIndex)
+				{
+					const UTexture* Texture = nullptr;
+					UniformExpressions.GetTextureValue(TextureType, ParameterIndex, MaterialContext, Material, Texture);
+					if (Texture && Texture->GetResource() && Texture->GetFName() == ViewModeParamName)
+					{
+						const FMaterialTextureParameterInfo& Parameter = UniformExpressions.GetTextureParameter(TextureType, ParameterIndex);
+						AnalysisIndex = Parameter.TextureIndex;
+
+						if (Texture->IsStreamable())
+						{
+							TextureResolution = 1 << (Texture->GetResource()->GetCurrentMipCount() - 1);
+						}
+						else
+						{
+							TextureResolution = FMath::Max(Texture->GetResource()->GetSizeX(), Texture->GetResource()->GetSizeY());
+						}
+					}
+				}
+			}
+		}
+
+		AnalysisParameter = FIntVector(AnalysisIndex, TextureResolution, AnalysisType);
+	}
+
+	if (DebugViewMode == DVSM_QuadComplexity)
+	{
+		NormalizedComplexityValue = FVector4f(NormalizedQuadComplexityValue);
+		bShowQuadOverdraw = true;
+	}
+	else
+	{
+		// normalize the complexity so we can fit it in a low precision scene color which is necessary on some platforms
+		// late value is for overdraw which can be problematic with a low precision float format, at some point the precision isn't there any more and it doesn't accumulate
+		const float NormalizeMul = 1.0f / GetMaxShaderComplexityCount(Material.GetFeatureLevel());
+		NormalizedComplexityValue = FVector4f(NumPSInstructions * NormalizeMul, NumVSInstructions * NormalizeMul, 1 / 32.0f);
+		ShaderBindings.Add(Shader.ShowQuadOverdraw, DebugViewMode != DVSM_ShaderComplexity ? 1 : 0);
+
+		bShowQuadOverdraw = DebugViewMode != DVSM_ShaderComplexity ? 1 : 0;
+	}
+
+	// Bind Data
+	ShaderBindings.Add(Shader.OneOverCPUTexCoordScalesParameter, OneOverCPUTexCoordScales);
+	ShaderBindings.Add(Shader.TexCoordIndicesParameter, TexCoordIndices);
+	ShaderBindings.Add(Shader.CPUTexelFactorParameter, FVector4f(WorldUVDensities));
+	ShaderBindings.Add(Shader.NormalizedComplexity, NormalizedComplexityValue);
+	ShaderBindings.Add(Shader.AnalysisParamsParameter, AnalysisParameter);
+	ShaderBindings.Add(Shader.PrimitiveAlphaParameter, PrimitiveAlpha);
+	ShaderBindings.Add(Shader.TexCoordAnalysisIndexParameter, TexCoordAnalysisIndex);
+	ShaderBindings.Add(Shader.CPULogDistanceParameter, CPULogDistance);
+	ShaderBindings.Add(Shader.ShowQuadOverdraw, bShowQuadOverdraw);
+	ShaderBindings.Add(Shader.LODIndexParameter, LODIndex);
+	ShaderBindings.Add(Shader.OutputQuadOverdrawParameter, boolOutputQuadOverdraw);
+	ShaderBindings.Add(Shader.VisualizeModeParameter, DebugViewMode);
+}
+
 FMeshPassProcessor* CreateDebugViewModePassProcessor(const FScene* Scene, const FSceneView* InViewIfDynamicMeshCommand, FMeshPassDrawListContext* InDrawListContext)
 {
 	const ERHIFeatureLevel::Type FeatureLevel = Scene ? Scene->GetFeatureLevel() : (InViewIfDynamicMeshCommand ? InViewIfDynamicMeshCommand->GetFeatureLevel() : GMaxRHIFeatureLevel);
@@ -368,21 +566,9 @@ FMeshPassProcessor* CreateDebugViewModePassProcessor(const FScene* Scene, const
 FRegisterPassProcessorCreateFunction RegisterDebugViewModeMobilePass(&CreateDebugViewModePassProcessor, EShadingPath::Mobile, EMeshPass::DebugViewMode, EMeshPassFlags::MainView);
 FRegisterPassProcessorCreateFunction RegisterDebugViewModePass(&CreateDebugViewModePassProcessor, EShadingPath::Deferred, EMeshPass::DebugViewMode, EMeshPassFlags::MainView);
 
-void InitDebugViewModeInterfaces()
+void InitDebugViewModeInterface()
 {
-	FDebugViewModeInterface::SetInterface(DVSM_ShaderComplexity, new FComplexityAccumulateInterface(true, false));
-	FDebugViewModeInterface::SetInterface(DVSM_ShaderComplexityContainedQuadOverhead, new FComplexityAccumulateInterface(true, false));
-	FDebugViewModeInterface::SetInterface(DVSM_ShaderComplexityBleedingQuadOverhead, new FComplexityAccumulateInterface(true, true));
-	FDebugViewModeInterface::SetInterface(DVSM_QuadComplexity, new FComplexityAccumulateInterface(false, false));
-
-	FDebugViewModeInterface::SetInterface(DVSM_PrimitiveDistanceAccuracy, new FPrimitiveDistanceAccuracyInterface());
-	FDebugViewModeInterface::SetInterface(DVSM_MeshUVDensityAccuracy, new FMeshTexCoordSizeAccuracyInterface());
-	FDebugViewModeInterface::SetInterface(DVSM_MaterialTextureScaleAccuracy, new FMaterialTexCoordScaleAccuracyInterface());
-	FDebugViewModeInterface::SetInterface(DVSM_OutputMaterialTextureScales, new FOutputMaterialTexCoordScaleInterface());
-	FDebugViewModeInterface::SetInterface(DVSM_RequiredTextureResolution, new FRequiredTextureResolutionInterface());
-	FDebugViewModeInterface::SetInterface(DVSM_VirtualTexturePendingMips, new FRequiredTextureResolutionInterface());
-
-	FDebugViewModeInterface::SetInterface(DVSM_LODColoration, new FLODColorationInterface());
+	FDebugViewModeInterface::SetInterface(new FDebugViewModeImplementation());
 }
 
 #else // !WITH_DEBUG_VIEW_MODES
diff --git a/Engine/Source/Runtime/Renderer/Private/DebugViewModeRendering.h b/Engine/Source/Runtime/Renderer/Private/DebugViewModeRendering.h
--- a/Engine/Source/Runtime/Renderer/Private/DebugViewModeRendering.h
+++ b/Engine/Source/Runtime/Renderer/Private/DebugViewModeRendering.h
@@ -15,6 +15,7 @@ DebugViewModeRendering.h: Contains definitions for rendering debug viewmodes.
 #include "MeshMaterialShader.h"
 #include "ShaderBaseClasses.h"
 #include "MeshPassProcessor.h"
+#include "DebugViewModeInterface.h"
 
 class FPrimitiveSceneProxy;
 struct FMeshBatchElement;
@@ -135,10 +136,49 @@ public:
 class FDebugViewModePS : public FMeshMaterialShader
 {
 public:
+	DECLARE_SHADER_TYPE(FDebugViewModePS, MeshMaterial);
+
 	FDebugViewModePS() = default;
 	FDebugViewModePS(const FMeshMaterialShaderType::CompiledShaderInitializerType& Initializer)
 		: FMeshMaterialShader(Initializer)
-	{}
+	{
+		OneOverCPUTexCoordScalesParameter.Bind(Initializer.ParameterMap, TEXT("OneOverCPUTexCoordScales"));
+		TexCoordIndicesParameter.Bind(Initializer.ParameterMap, TEXT("TexCoordIndices"));
+		CPUTexelFactorParameter.Bind(Initializer.ParameterMap, TEXT("CPUTexelFactor"));
+		NormalizedComplexity.Bind(Initializer.ParameterMap, TEXT("NormalizedComplexity"));
+		AnalysisParamsParameter.Bind(Initializer.ParameterMap, TEXT("AnalysisParams"));
+		PrimitiveAlphaParameter.Bind(Initializer.ParameterMap, TEXT("PrimitiveAlpha"));
+		TexCoordAnalysisIndexParameter.Bind(Initializer.ParameterMap, TEXT("TexCoordAnalysisIndex"));
+		CPULogDistanceParameter.Bind(Initializer.ParameterMap, TEXT("CPULogDistance"));
+		ShowQuadOverdraw.Bind(Initializer.ParameterMap, TEXT("bShowQuadOverdraw"));
+		OutputQuadOverdrawParameter.Bind(Initializer.ParameterMap, TEXT("bOutputQuadOverdraw"));
+		LODIndexParameter.Bind(Initializer.ParameterMap, TEXT("LODIndex"));
+		VisualizeModeParameter.Bind(Initializer.ParameterMap, TEXT("VisualizeMode"));
+		QuadBufferUAV.Bind(Initializer.ParameterMap, TEXT("RWQuadBuffer"));
+	}
+
+	static bool ShouldCompilePermutation(const FMeshMaterialShaderPermutationParameters& Parameters);
+
+	static void ModifyCompilationEnvironment(const FMaterialShaderPermutationParameters& Parameters, FShaderCompilerEnvironment& OutEnvironment)
+	{
+		OutEnvironment.SetDefine(TEXT("UNDEFINED_ACCURACY"), UndefinedStreamingAccuracyIntensity);
+		OutEnvironment.SetDefine(TEXT("MAX_NUM_TEX_COORD"), (uint32)TEXSTREAM_MAX_NUM_UVCHANNELS);
+		OutEnvironment.SetDefine(TEXT("INITIAL_GPU_SCALE"), (uint32)TEXSTREAM_INITIAL_GPU_SCALE);
+		OutEnvironment.SetDefine(TEXT("TILE_RESOLUTION"), (uint32)TEXSTREAM_TILE_RESOLUTION);
+		OutEnvironment.SetDefine(TEXT("MAX_NUM_TEXTURE_REGISTER"), (uint32)TEXSTREAM_MAX_NUM_TEXTURES_PER_MATERIAL);
+		OutEnvironment.SetDefine(TEXT("SCENE_TEXTURES_DISABLED"), 1u);
+		TCHAR BufferRegister[] = { 'u', '0', 0 };
+		BufferRegister[1] += GetQuadOverdrawUAVIndex(Parameters.Platform, Parameters.MaterialParameters.FeatureLevel);
+		OutEnvironment.SetDefine(TEXT("QUAD_BUFFER_REGISTER"), BufferRegister);
+		OutEnvironment.SetDefine(TEXT("OUTPUT_QUAD_OVERDRAW"), 1);
+
+		for (int i = 0; i < DVSM_MAX; ++i)
+		{
+			OutEnvironment.SetDefine(DebugViewShaderModeToString(static_cast<EDebugViewShaderMode>(i)), i);
+		}
+
+		FMeshMaterialShader::ModifyCompilationEnvironment(Parameters, OutEnvironment);
+	}
 
 	void GetElementShaderBindings(
 		const FShaderMapPointerTable& PointerTable,
@@ -153,6 +193,20 @@ public:
 		const FDebugViewModeShaderElementData& ShaderElementData,
 		FMeshDrawSingleShaderBindings& ShaderBindings,
 		FVertexInputStreamArray& VertexStreams) const;
+
+	LAYOUT_FIELD(FShaderParameter, OneOverCPUTexCoordScalesParameter);
+	LAYOUT_FIELD(FShaderParameter, TexCoordIndicesParameter);
+	LAYOUT_FIELD(FShaderParameter, CPUTexelFactorParameter);
+	LAYOUT_FIELD(FShaderParameter, NormalizedComplexity);
+	LAYOUT_FIELD(FShaderParameter, AnalysisParamsParameter);
+	LAYOUT_FIELD(FShaderParameter, PrimitiveAlphaParameter);
+	LAYOUT_FIELD(FShaderParameter, TexCoordAnalysisIndexParameter);
+	LAYOUT_FIELD(FShaderParameter, CPULogDistanceParameter);
+	LAYOUT_FIELD(FShaderParameter, ShowQuadOverdraw);
+	LAYOUT_FIELD(FShaderParameter, LODIndexParameter);
+	LAYOUT_FIELD(FShaderParameter, OutputQuadOverdrawParameter);
+	LAYOUT_FIELD(FShaderParameter, VisualizeModeParameter);
+	LAYOUT_FIELD(FShaderResourceParameter, QuadBufferUAV);
 };
 
 class FDebugViewModeMeshProcessor : public FMeshPassProcessor
@@ -172,9 +226,31 @@ private:
 	const FDebugViewModeInterface* DebugViewModeInterface;
 };
 
-void AddDebugViewModeShaderTypes(ERHIFeatureLevel::Type FeatureLevel,
-	const FVertexFactoryType* VertexFactoryType,
-	FMaterialShaderTypes& OutShaderTypes);
+class FDebugViewModeImplementation : public FDebugViewModeInterface
+{
+public:
+	FDebugViewModeImplementation() : FDebugViewModeInterface() {}
+
+	virtual void AddShaderTypes(ERHIFeatureLevel::Type InFeatureLevel,
+		const FVertexFactoryType* InVertexFactoryType,
+		FMaterialShaderTypes& OutShaderTypes) const override;
+
+	virtual void GetDebugViewModeShaderBindings(
+		const FDebugViewModePS& BaseShader,
+		const FPrimitiveSceneProxy* RESTRICT PrimitiveSceneProxy,
+		const FMaterialRenderProxy& RESTRICT MaterialRenderProxy,
+		const FMaterial& RESTRICT Material,
+		EDebugViewShaderMode DebugViewMode,
+		const FVector& ViewOrigin,
+		int32 VisualizeLODIndex,
+		int32 VisualizeElementIndex,
+		int32 NumVSInstructions,
+		int32 NumPSInstructions,
+		int32 ViewModeParam,
+		FName ViewModeParamName,
+		FMeshDrawSingleShaderBindings& ShaderBindings
+	) const override;
+};
 
 #endif // WITH_DEBUG_VIEW_MODES
 
diff --git a/Engine/Source/Runtime/Renderer/Private/MaterialTexCoordScalesRendering.cpp b/Engine/Source/Runtime/Renderer/Private/MaterialTexCoordScalesRendering.cpp
deleted file mode 100644
diff --git a/Engine/Source/Runtime/Renderer/Private/MaterialTexCoordScalesRendering.h b/Engine/Source/Runtime/Renderer/Private/MaterialTexCoordScalesRendering.h
deleted file mode 100644
diff --git a/Engine/Source/Runtime/Renderer/Private/MeshTexCoordSizeAccuracyRendering.cpp b/Engine/Source/Runtime/Renderer/Private/MeshTexCoordSizeAccuracyRendering.cpp
deleted file mode 100644
diff --git a/Engine/Source/Runtime/Renderer/Private/MeshTexCoordSizeAccuracyRendering.h b/Engine/Source/Runtime/Renderer/Private/MeshTexCoordSizeAccuracyRendering.h
deleted file mode 100644
diff --git a/Engine/Source/Runtime/Renderer/Private/PrimitiveDistanceAccuracyRendering.cpp b/Engine/Source/Runtime/Renderer/Private/PrimitiveDistanceAccuracyRendering.cpp
deleted file mode 100644
diff --git a/Engine/Source/Runtime/Renderer/Private/PrimitiveDistanceAccuracyRendering.h b/Engine/Source/Runtime/Renderer/Private/PrimitiveDistanceAccuracyRendering.h
deleted file mode 100644
diff --git a/Engine/Source/Runtime/Renderer/Private/RendererPrivate.h b/Engine/Source/Runtime/Renderer/Private/RendererPrivate.h
--- a/Engine/Source/Runtime/Renderer/Private/RendererPrivate.h
+++ b/Engine/Source/Runtime/Renderer/Private/RendererPrivate.h
@@ -7,7 +7,6 @@
 #pragma once
 
 #include "CoreMinimal.h"
-#include "ShaderComplexityRendering.h"
 #include "Async/ParallelFor.h"
 #include "DeferredShadingRenderer.h"
 #include "DynamicPrimitiveDrawing.h"
diff --git a/Engine/Source/Runtime/Renderer/Private/RequiredTextureResolutionRendering.cpp b/Engine/Source/Runtime/Renderer/Private/RequiredTextureResolutionRendering.cpp
deleted file mode 100644
diff --git a/Engine/Source/Runtime/Renderer/Private/RequiredTextureResolutionRendering.h b/Engine/Source/Runtime/Renderer/Private/RequiredTextureResolutionRendering.h
deleted file mode 100644
diff --git a/Engine/Source/Runtime/Renderer/Private/SceneRendering.cpp b/Engine/Source/Runtime/Renderer/Private/SceneRendering.cpp
--- a/Engine/Source/Runtime/Renderer/Private/SceneRendering.cpp
+++ b/Engine/Source/Runtime/Renderer/Private/SceneRendering.cpp
@@ -4215,8 +4215,8 @@ FRendererModule::FRendererModule()
 	CVarVertexDeformationOutputsVelocity->SetOnChangedCallback(FConsoleVariableDelegate::CreateStatic(&OnChangeCVarRequiringRecreateRenderState));
 
 #if !(UE_BUILD_SHIPPING || UE_BUILD_TEST)
-	void InitDebugViewModeInterfaces();
-	InitDebugViewModeInterfaces();
+	void InitDebugViewModeInterface();
+	InitDebugViewModeInterface();
 #endif // !(UE_BUILD_SHIPPING || UE_BUILD_TEST)
 }
 
diff --git a/Engine/Source/Runtime/Renderer/Private/ShaderComplexityRendering.cpp b/Engine/Source/Runtime/Renderer/Private/ShaderComplexityRendering.cpp
deleted file mode 100644
diff --git a/Engine/Source/Runtime/Renderer/Private/ShaderComplexityRendering.h b/Engine/Source/Runtime/Renderer/Private/ShaderComplexityRendering.h
deleted file mode 100644
diff --git a/Engine/Source/Runtime/Renderer/Private/ViewMode/LODColorationRendering.cpp b/Engine/Source/Runtime/Renderer/Private/ViewMode/LODColorationRendering.cpp
deleted file mode 100644
diff --git a/Engine/Source/Runtime/Renderer/Private/ViewMode/LODColorationRendering.h b/Engine/Source/Runtime/Renderer/Private/ViewMode/LODColorationRendering.h
deleted file mode 100644
