Improve strata normal code gen.

For instance: reduce AO setup pass from 0.3ms -> 0.2ms and base pass by 0.05ms.

diff --git a/Engine/Shaders/Private/Strata/Strata.ush b/Engine/Shaders/Private/Strata/Strata.ush
--- a/Engine/Shaders/Private/Strata/Strata.ush
+++ b/Engine/Shaders/Private/Strata/Strata.ush
@@ -10,6 +10,7 @@
 #include "../ParticipatingMediaCommon.ush"
 #include "../ThinFilmBSDF.ush"
 #include "../MortonCode.ush"
+#include "../OctahedralCommon.ush"
 #include "/Engine/Shared/StrataDefinitions.h"     
 
 
@@ -1679,60 +1680,31 @@ float2 SignNotZero(in float2 v)
 
 uint StrataPackNormal(in float3 Normal)
 {
-	float l1norm = abs(Normal.x) + abs(Normal.y) + abs(Normal.z);
-	float2 result = Normal.xy * (1.0 / l1norm);
-	if (Normal.z < 0.0)
-	{
-		result = (1.0 - abs(result.yx)) * SignNotZero(result.xy);
-	}
-	uint PackedX = f32tof16(result.x);
-	uint PackedY = f32tof16(result.y);
-	return PackedX | (PackedY << 16);
+	const float2 Result = UnitVectorToOctahedron(Normal);
+	return f32tof16(Result.x) | (f32tof16(Result.y) << 16);
 }
 
 float3 StrataUnpackNormal(uint PackedNormal)
 {
-	float x = f16tof32(0xFFFF &  PackedNormal);
-	float y = f16tof32(0xFFFF & (PackedNormal>>16));
-
-	float3 Normal = float3(x, y, 1.0f - abs(x) - abs(y));
-	if (Normal.z < 0.0f)
-	{
-		Normal.xy = (1.0f - abs(Normal.yx)) * SignNotZero(Normal.xy);
-	}
-	return normalize(Normal);
+	const float2 xy = float2(f16tof32(0xFFFF & PackedNormal), f16tof32(0xFFFF & (PackedNormal >> 16)));
+	return OctahedronToUnitVector(xy);
 }
 
 uint StrataPackNormal24(in float3 Normal)
 {
-	float l1norm = abs(Normal.x) + abs(Normal.y) + abs(Normal.z);
-	float2 result = Normal.xy * (1.0 / l1norm);
-	if (Normal.z < 0.0)
-	{
-		result = (1.0 - abs(result.yx)) * SignNotZero(result.xy);
-	}
-
-	uint PackedX = uint(clamp(result.x * 2047.0f + 2048.0f, 0.0f, 4095.0f));
-	uint PackedY = uint(clamp(result.y * 2047.0f + 2048.0f, 0.0f, 4095.0f));
-	return PackedX | (PackedY << 12);
+	const float2 Result = UnitVectorToOctahedron(Normal);
+	const uint2 PackedXY = uint2(clamp(Result * 2047.0f + 2048.0f, 0.0f, 4095.0f));
+	return PackedXY.x | (PackedXY.y << 12);
 }
 
 float3 StrataUnpackNormal24(uint PackedNormal)
-{	
-	int X12Bits = 0xFFF & PackedNormal;
-	int Y12Bits = 0xFFF & (PackedNormal >> 12);
-
-	float x = float(X12Bits - 2048) / 2047.0f;
-	float y = float(Y12Bits - 2048) / 2047.0f;
-
-	float3 Normal = float3(x, y, 1.0f - abs(x) - abs(y));
-	if (Normal.z < 0.0f)
-	{
-		Normal.xy = (1.0f - abs(Normal.yx)) * SignNotZero(Normal.xy);
-	}
-	return normalize(Normal);
+{
+	const int2 XY12Bits = int2(0xFFF & PackedNormal, 0xFFF & (PackedNormal >> 12));
+	const float2 xy = float2(XY12Bits - 2048) / 2047.0f;
+	return OctahedronToUnitVector(xy);
 }
 
+
 #define STRATA_TANGENT_SYMMETRY 1
 uint StrataPackNormalAndTangent(in float3 Normal, in float3 Tangent)
 {
