Fix invalid entrypoint for deferred light overlap shader.

diff --git a/Engine/Shaders/Private/StationaryLightOverlapShaders.usf b/Engine/Shaders/Private/StationaryLightOverlapShaders.usf
--- a/Engine/Shaders/Private/StationaryLightOverlapShaders.usf
+++ b/Engine/Shaders/Private/StationaryLightOverlapShaders.usf
@@ -11,7 +11,8 @@
 /** 1 if the light has a valid shadowmap channel. */
 float HasValidChannel;
 
-void OverlapDirectionalPixelMain(
+#if RADIAL_ATTENUATION == 0
+void OverlapPixelMain(
 	float2 InUV : TEXCOORD0,
 	float3 ScreenVector : TEXCOORD1,
 	out float4 OutColor : SV_Target0
@@ -34,8 +35,10 @@ void OverlapDirectionalPixelMain(
 		}
 	}
 }
+#endif
 
-void OverlapRadialPixelMain(
+#if RADIAL_ATTENUATION == 1
+void OverlapPixelMain(
 	float4 InScreenPosition : TEXCOORD0,
 	in float4 SvPosition : SV_Position,		// after all interpolators
 	out float4 OutColor : SV_Target0
@@ -53,16 +56,9 @@ void OverlapRadialPixelMain(
 {
 		float3 WorldPosition = LWCHackToFloat(SvPositionToWorld(SvPosition));
 
-		float3 ToLight = DeferredLightUniforms.Direction;
-
-		#if RADIAL_ATTENUATION
-			ToLight = DeferredLightUniforms.Position - WorldPosition;
-		#endif
+		float3 ToLight = DeferredLightUniforms.Position - WorldPosition;
 
-		float SpotFalloff = 1;
-		#if RADIAL_ATTENUATION
-			SpotFalloff = SpotAttenuation(normalize(ToLight), -DeferredLightUniforms.Direction, DeferredLightUniforms.SpotAngles);
-		#endif
+		float SpotFalloff = SpotAttenuation(normalize(ToLight), -DeferredLightUniforms.Direction, DeferredLightUniforms.SpotAngles);
 
 		float RadialAttenuationRatio = length(WorldPosition - DeferredLightUniforms.Position) * DeferredLightUniforms.InvRadius;
 
@@ -78,4 +74,5 @@ void OverlapRadialPixelMain(
 			}
 		}
 	}
-}
\ No newline at end of file
+}
+#endif 
\ No newline at end of file
diff --git a/Engine/Source/Runtime/Renderer/Private/LightRendering.cpp b/Engine/Source/Runtime/Renderer/Private/LightRendering.cpp
--- a/Engine/Source/Runtime/Renderer/Private/LightRendering.cpp
+++ b/Engine/Source/Runtime/Renderer/Private/LightRendering.cpp
@@ -770,7 +770,6 @@ class FDeferredLightOverlapPS : public FGlobalShader
 	DECLARE_SHADER_TYPE(FDeferredLightOverlapPS,Global)
 
 	class FRadialAttenuation : SHADER_PERMUTATION_BOOL("RADIAL_ATTENUATION");
-
 	using FPermutationDomain = TShaderPermutationDomain<FRadialAttenuation>;
 
 public:
@@ -808,7 +807,7 @@ private:
 	LAYOUT_FIELD(FShaderParameter, HasValidChannel);
 };
 
-IMPLEMENT_GLOBAL_SHADER(FDeferredLightOverlapPS, "/Engine/Private/StationaryLightOverlapShaders.usf", "OverlapRadialPixelMain", SF_Pixel);
+IMPLEMENT_GLOBAL_SHADER(FDeferredLightOverlapPS, "/Engine/Private/StationaryLightOverlapShaders.usf", "OverlapPixelMain", SF_Pixel);
 
 static void SplitSimpleLightsByView(TArrayView<const FViewInfo> Views, const FSimpleLightArray& SimpleLights, TArrayView<FSimpleLightArray> SimpleLightsByView)
 {
