Don't ouput quad overdraw visualization for shader platforms that don't support the required instructions.

diff --git a/Engine/Shaders/Private/DebugViewModePixelShader.usf b/Engine/Shaders/Private/DebugViewModePixelShader.usf
--- a/Engine/Shaders/Private/DebugViewModePixelShader.usf
+++ b/Engine/Shaders/Private/DebugViewModePixelShader.usf
@@ -466,6 +466,7 @@ float4 VisualizeShaderComplexityAccumulate(in FDebugPSIn DebugInputs, in bool bI
 	{
 		float3 FinalComplexity = NormalizedComplexity.xyz;
 
+#if !(COMPILER_METAL)
 		[branch]
 		if (bShowQuadOverdraw && NormalizedComplexity.x > 0)
 		{
@@ -473,6 +474,7 @@ float4 VisualizeShaderComplexityAccumulate(in FDebugPSIn DebugInputs, in bool bI
 			// The extra cost from quad overdraw is assumed to be linear.
 			FinalComplexity.x *= 4.f / (float)(Coverage);
 		}
+#endif // !(COMPILER_METAL)
 
 		// use the maximum range allowed for scene color
 		Result = FinalComplexity;
