Fix shader compiler error packaging with static lighting (InstanceData.HasLightShadowUVBias should be InstanceData.Flags & INSTANCE_SCENE_DATA_FLAG_HAS_LIGHTSHADOW_UV_BIAS).

diff --git a/Engine/Shaders/Private/InstanceCulling/BuildInstanceDrawCommands.usf b/Engine/Shaders/Private/InstanceCulling/BuildInstanceDrawCommands.usf
--- a/Engine/Shaders/Private/InstanceCulling/BuildInstanceDrawCommands.usf
+++ b/Engine/Shaders/Private/InstanceCulling/BuildInstanceDrawCommands.usf
@@ -98,7 +98,7 @@ void WritePackedInstanceData(uint Offset, uint InstanceId, FInstanceSceneData In
 	{
 		LightMapUVScaleBias = GetLightmapData(LightmapDataIndex).LightMapCoordinateScaleBias;
 		ShadowMapUVScaleBias = GetLightmapData(LightmapDataIndex).ShadowMapCoordinateScaleBias;
-		if (InstanceData.HasLightShadowUVBias)
+		if (InstanceData.Flags & INSTANCE_SCENE_DATA_FLAG_HAS_LIGHTSHADOW_UV_BIAS)
 		{
 			LightMapUVScaleBias.zw = InstanceData.LightMapAndShadowMapUVBias.xy;
 			ShadowMapUVScaleBias.zw = InstanceData.LightMapAndShadowMapUVBias.zw;
