One-liner to fix 2 paint layer issues :
* No weight blend layers should never be normalized (in case of multiple edit layers, it could happen that all AllLayerWeightSum != 1, which would wrongly attenuate the no weight blend layer's weight)
* For the same reason, the visibility layer, which is not weight blended, could end up being blended with weight-blended layers, leaving to non-fully masked area where there is paint from another layer

#tests editor

diff --git a/Engine/Shaders/Private/LandscapeLayersCS.usf b/Engine/Shaders/Private/LandscapeLayersCS.usf
--- a/Engine/Shaders/Private/LandscapeLayersCS.usf
+++ b/Engine/Shaders/Private/LandscapeLayersCS.usf
@@ -43,6 +43,12 @@ Buffer<float2> InWeightmapTextureOutputOffset;
 Texture2DArray<float> InAtlasPaintLayers;
 RWTexture2D<UNORM float4> OutComponentWeightMaps;
 
+bool IsVisibilityLayer(int InSourcePaintLayerIndex)
+{
+	// Paint layer 0 is reserved for visibility : 
+	return (InSourcePaintLayerIndex == 0);
+}
+
 float ComputeWeightmapChannel(int ChannelIndexToProcess, uint2 TexelPosition, uint SourcePaintLayerIndex)
 {
 	float Channel = 0;
@@ -50,6 +56,13 @@ float ComputeWeightmapChannel(int ChannelIndexToProcess, uint2 TexelPosition, ui
 	if (ChannelIndexToProcess != -1)
 	{
 		Channel = InAtlasPaintLayers.Load(int4(TexelPosition, SourcePaintLayerIndex, 0));
+		
+		// If this isn't a weight blended channel, it shouldn't be normalized (it's true for the visibility layer as well, which is not weight blended) : 
+		if ((InWeightmapWeightBlendMode[SourcePaintLayerIndex] == 0.0f) || IsVisibilityLayer(SourcePaintLayerIndex))
+		{
+			return Channel;
+		}
+		
 		float AllLayerWeightSum = 0;
 
 		float Unused = 0.0;
