Removed redundant value that was being used to distinguish debug modes.
In the new uber shader setup we can just use VisualizeMode for that.

diff --git a/Engine/Shaders/Private/DebugViewModePixelShader.usf b/Engine/Shaders/Private/DebugViewModePixelShader.usf
--- a/Engine/Shaders/Private/DebugViewModePixelShader.usf
+++ b/Engine/Shaders/Private/DebugViewModePixelShader.usf
@@ -27,7 +27,7 @@ float4 OneOverCPUTexCoordScales[MAX_NUM_TEXTURE_REGISTER / 4];
 int4 TexCoordIndices[MAX_NUM_TEXTURE_REGISTER / 4];
 float4 CPUTexelFactor; // In world space units.
 float4 NormalizedComplexity;
-int3 AnalysisParams; // (TextureAnalysisIndex, bOutputScales) or // (TextureAnalysisIndex, TextureResolution)
+int2 AnalysisParams; // (TextureAnalysisIndex, bOutputScales) or // (TextureAnalysisIndex, TextureResolution)
 float PrimitiveAlpha;
 int TexCoordAnalysisIndex;
 float CPULogDistance;
@@ -41,11 +41,6 @@ int VisualizeMode;
 #define bOutputScales (AnalysisParams.y)
 // Texture resolution if we're analysing a single texture
 #define TextureResolution (AnalysisParams.y)
-// Analysis type will be ANALYSIS_TEXTURE_RESOLUTION or ANALYSIS_VIRTUAL_TEXTURE_PENDING_MIPS
-#define TextureAnalysisType (AnalysisParams.z)
-
-#define ANALYSIS_TEXTURE_RESOLUTION				0
-#define ANALYSIS_VIRTUAL_TEXTURE_PENDING_MIPS	1
 
 // Updated MAX_NUM_TEXTURE_REGISTER to 256 
 
@@ -139,7 +134,7 @@ MaterialFloat StoreTexSample(in out FTexCoordScalesParams Params, float4 C, int
 
 MaterialFloat StoreVTSampleInfo(in out FTexCoordScalesParams Params, VTPageTableResult PageTableResult, uint LayerIndex, int TextureReferenceIndex)
 {
-	if (TextureAnalysisType == ANALYSIS_VIRTUAL_TEXTURE_PENDING_MIPS)
+	if (VisualizeMode == DVSM_VirtualTexturePendingMips)
 	{
 		uint RequestLevel = GetRequestedLevel(PageTableResult);
 
@@ -427,7 +422,7 @@ float4 VisualizeRequiredTextureResolution(in FDebugPSIn DebugInputs, in bool bIs
 	float3 Result = float3(UNDEFINED_ACCURACY, UNDEFINED_ACCURACY, UNDEFINED_ACCURACY);
 	Result *= saturate(Luminance(BaseColor));
 
-	if (TextureAnalysisType == ANALYSIS_TEXTURE_RESOLUTION)
+	if (VisualizeMode == DVSM_RequiredTextureResolution)
 	{
 		if (MaterialParameters.TexCoordScalesParams.RequiredResolution > 0)
 		{
@@ -436,7 +431,7 @@ float4 VisualizeRequiredTextureResolution(in FDebugPSIn DebugInputs, in bool bIs
 			Result *= MaterialParameters.TexCoordScalesParams.TexSample;
 		}
 	}
-	else if (TextureAnalysisType == ANALYSIS_VIRTUAL_TEXTURE_PENDING_MIPS)
+	else if (VisualizeMode == DVSM_VirtualTexturePendingMips)
 	{
 		if (MaterialParameters.TexCoordScalesParams.VTMissingMipCount >= 0)
 		{
diff --git a/Engine/Source/Runtime/Renderer/Private/DebugViewModeRendering.cpp b/Engine/Source/Runtime/Renderer/Private/DebugViewModeRendering.cpp
--- a/Engine/Source/Runtime/Renderer/Private/DebugViewModeRendering.cpp
+++ b/Engine/Source/Runtime/Renderer/Private/DebugViewModeRendering.cpp
@@ -427,14 +427,13 @@ void FDebugViewModeImplementation::GetDebugViewModeShaderBindings(
 	FIntVector4 TexCoordIndices[TEXSTREAM_MAX_NUM_TEXTURES_PER_MATERIAL / 4];
 	FVector4 WorldUVDensities;
 	FVector4f NormalizedComplexityValue;
-	FIntVector AnalysisParameter;
+	FIntPoint AnalysisParameter;
 	const float PrimitiveAlpha = (!PrimitiveSceneProxy || PrimitiveSceneProxy->IsSelected()) ? 1.f : .2f;
 	const int32 TexCoordAnalysisIndex = ViewModeParam >= 0 ? FMath::Clamp<int32>(ViewModeParam, 0, MAX_TEXCOORDS - 1) : -1;
 	float CPULogDistance = -1.0f;
 	int32 bShowQuadOverdraw = 0;
 	const int32 boolOutputQuadOverdraw = (DebugViewMode == DVSM_QuadComplexity) || (DebugViewMode == DVSM_ShaderComplexityContainedQuadOverhead) ? 1 : 0;
 	const int32 LODIndex = FMath::Clamp(VisualizeLODIndex, 0, NumLODColorationColors - 1);
-	const int32 AnalysisType = DebugViewMode == DVSM_VirtualTexturePendingMips ? 1 : 0;
 
 	FMemory::Memzero(OneOverCPUTexCoordScales); // 0 remap to irrelevant data.
 	FMemory::Memzero(TexCoordIndices);
@@ -460,7 +459,7 @@ void FDebugViewModeImplementation::GetDebugViewModeShaderBindings(
 	{
 		const bool bOutputScales = DebugViewMode == DVSM_OutputMaterialTextureScales;
 		const int32 AnalysisIndex = ViewModeParam >= 0 ? FMath::Clamp<int32>(ViewModeParam, 0, TEXSTREAM_MAX_NUM_TEXTURES_PER_MATERIAL - 1) : -1;
-		AnalysisParameter = FIntVector(bOutputScales ? -1 : AnalysisIndex, bOutputScales ? 1 : 0, AnalysisType);
+		AnalysisParameter = FIntPoint(bOutputScales ? -1 : AnalysisIndex, bOutputScales ? 1 : 0);
 	}
 	else if (DebugViewMode == DVSM_RequiredTextureResolution || DebugViewMode == DVSM_VirtualTexturePendingMips)
 	{
@@ -523,7 +522,7 @@ void FDebugViewModeImplementation::GetDebugViewModeShaderBindings(
 			}
 		}
 
-		AnalysisParameter = FIntVector(AnalysisIndex, TextureResolution, AnalysisType);
+		AnalysisParameter = FIntPoint(AnalysisIndex, TextureResolution);
 	}
 
 	if (DebugViewMode == DVSM_QuadComplexity)
