Fix strata eye conversion

diff --git a/Engine/Shaders/Private/Strata/StrataLegacyConversion.ush b/Engine/Shaders/Private/Strata/StrataLegacyConversion.ush
--- a/Engine/Shaders/Private/Strata/StrataLegacyConversion.ush
+++ b/Engine/Shaders/Private/Strata/StrataLegacyConversion.ush
@@ -268,10 +268,11 @@ FStrataData StrataConvertLegacyMaterialDynamic(
 	{
 		const float IrisMask = ClearCoat;
 		const float IrisDistance = ClearCoatRoughness;
-		const float IrisOpacity = saturate(1.0f - IrisMask);
 
-		// Incorrect composition, but allows to compile shaders
-		
+		// Lerp transmittance rather MFP, as transmittance is perceptually linear unlike MFP
+		const float Thickness = DefaultThickness;
+		const float3 MFP = TransmittanceToMeanFreePath(IrisMask.xxx, Thickness * CENTIMETER_TO_METER) * METER_TO_CENTIMETER;
+
 		// Cornea/Sclera
 		Top.BaseColor = BaseColor;
 		Top.Specular = Specular;
@@ -279,7 +280,7 @@ FStrataData StrataConvertLegacyMaterialDynamic(
 		Top.Roughness = Roughness;
 		Top.Anisotropy = 0.f;
 		Top.SSSProfileID = SubSurfaceProfileId;
-		Top.SSSDMFP = Zeros;
+		Top.SSSDMFP = MFP;
 		Top.SSSDMFPScale = Opacity;
 		Top.Emissive = Emissive;
 		Top.SharedLocalBasisIndex = SharedLocalBasisIndex;
