Remove HasSSSProfile from material header.

diff --git a/Engine/Shaders/Private/ShadowProjectionPixelShader.usf b/Engine/Shaders/Private/ShadowProjectionPixelShader.usf
--- a/Engine/Shaders/Private/ShadowProjectionPixelShader.usf
+++ b/Engine/Shaders/Private/ShadowProjectionPixelShader.usf
@@ -314,7 +314,7 @@ void Main(
 
 				// For non-profile version, use default transmission parameters (only NormalScale & ExtinctionScale are used)
 				FTransmissionProfileParams TransmissionParams = InitTransmissionProfileParams();
-				if (HasSubsurfaceProfile(StrataHeader))
+				if (SSSData.bIsProfile)
 				{
 					TransmissionParams = GetTransmissionProfileParams(SSSData.ProfileId);
 				}
diff --git a/Engine/Shaders/Private/Strata/Strata.ush b/Engine/Shaders/Private/Strata/Strata.ush
--- a/Engine/Shaders/Private/Strata/Strata.ush
+++ b/Engine/Shaders/Private/Strata/Strata.ush
@@ -137,7 +137,7 @@ struct FStrataBSDF
 // * HASDYNINDIRECTSHADOWCASTER	
 // * ISSINGLELAYERWATER			
 // * HASSUBSURFACE				
-// * HASSUBSURFACEPROFILE
+// * UNUSED
 // 
 // * Complex: -------------------------------------------------------------------------------------------------------------------------------------------------------------------------
 //            HEADER        31 |  State | AO | BSDF Count | BSDF Types | Basis count | Basis type |
@@ -182,7 +182,7 @@ struct FStrataBSDF
 #define HEADER_MASK_HASDYNINDIRECTSHADOWCASTER	(1u << 4)
 #define HEADER_MASK_ISSINGLELAYERWATER			(1u << 5)	// Used to avoid loading BSDF or Categorisation data when generating tiles containing SLW material. Avoid running another categorisation.
 #define HEADER_MASK_HASSUBSURFACE				(1u << 6)
-#define HEADER_MASK_HASSUBSURFACEPROFILE		(1u << 7)
+#define HEADER_MASK_UNUSED						(1u << 7)
 #define HEADER_MASK_BIT_COUNT					8
 
 // For HEADER_BIT_BSDF_COUNT and  HEADER_BIT_SHAREDLOCALBASES_X, please refer to values in StrataDefinitions.h
@@ -893,7 +893,6 @@ bool DoesCastContactShadow(const FStrataPixelHeader In)										{ return (In.St
 uint GetShadingModels(const FStrataPixelHeader In)											{ return (In.State & HEADER_MASK_SHADINGMODELS)>>HEADER_BIT_OFFSET_SHADINGMODELS; }
 bool IsStrataMaterial(const FStrataPixelHeader In)											{ return (In.State & HEADER_MASK_SHADINGMODELS) > 0 || IsSimpleMaterial(In); }
 bool HasSubsurface(const FStrataPixelHeader In)												{ return (In.State & HEADER_MASK_HASSUBSURFACE) != 0; }
-bool HasSubsurfaceProfile(const FStrataPixelHeader In)										{ return (In.State & HEADER_MASK_HASSUBSURFACEPROFILE) != 0; }
 bool HasDynamicIndirectShadowCasterRepresentation(const FStrataPixelHeader In)				{ return (In.State & HEADER_MASK_HASDYNINDIRECTSHADOWCASTER) != 0; }
 
 bool StrataHasShadingModel(const FStrataPixelHeader In, uint ShadingModel)					{ return StrataHasShadingModel(GetShadingModels(In), ShadingModel); }
@@ -905,7 +904,6 @@ void SetIsSingleLayerWater(inout FStrataPixelHeader Out, bool bIn)							{ Out.S
 void SetCastContactShadow (inout FStrataPixelHeader Out, bool bIn)							{ Out.State |= (bIn ? HEADER_MASK_CASTCONTACTSHADOW : 0u); }
 void SetDynamicIndirectShadowCasterRepresentation(inout FStrataPixelHeader Out, bool bIn)	{ Out.State |= (bIn ? HEADER_MASK_HASDYNINDIRECTSHADOWCASTER : 0u); }
 void SetHasSubsurface(inout FStrataPixelHeader Out, bool bIn)								{ Out.State |= (bIn ? HEADER_MASK_HASSUBSURFACE : 0u); }
-void SetHasSubsurfaceProfile(inout FStrataPixelHeader Out, bool bIn)						{ Out.State |= (bIn ? HEADER_MASK_HASSUBSURFACEPROFILE : 0u); }
 
 void AddShadingModels(inout FStrataPixelHeader Out, uint bIn)								{ Out.State |= ((bIn<<HEADER_BIT_OFFSET_SHADINGMODELS) & HEADER_MASK_SHADINGMODELS); }
 void AndIsSimpleMaterial(inout FStrataPixelHeader Out, bool bIn)							{ Out.State = (Out.State & ~HEADER_MASK_ISSIMPLEMATERIAL) | ((Out.State & HEADER_MASK_ISSIMPLEMATERIAL) & (bIn ? HEADER_MASK_ISSIMPLEMATERIAL : 0u)); }
@@ -2005,7 +2003,6 @@ void PackStrataOut(
 
 					const bool bSSSMask = !bIsThin && bBottomLayer;	// SSS and Thin can only be on the bottom layer, so ignoring bIsSimpleVolume here.
 					SetHasSubsurface(StrataPixelHeader,			(bSSSMask && BSDF_GETHASSSS(BSDF)));
-					SetHasSubsurfaceProfile(StrataPixelHeader,	(bSSSMask && BSDF_GETHASSSSPROFILE(BSDF)));
 
 					// STRATA_TODO	Should we have the Fuzz influence the next layter visibility accoding to fuzz amount? Should it also occlude diffuse?
 					//				If this is done, we should also update ComputePostProcessBSDFData.
