Strata - fixed shader compilation warning

diff --git a/Engine/Shaders/Private/Strata/Strata.ush b/Engine/Shaders/Private/Strata/Strata.ush
--- a/Engine/Shaders/Private/Strata/Strata.ush
+++ b/Engine/Shaders/Private/Strata/Strata.ush
@@ -1853,8 +1853,8 @@ void PackStrataOut(
 
 	SetIsSimpleMaterial(StrataPixelHeader, true); // non strata material or sky pixels are considered as simple (so that (strata simple pixels) + (non-strata pixels) can build a simple tile).
 
-	uint SimpleMaterial_LayerIndex = 0;
-	uint SimpleMaterial_BSDFIndex = 0;
+	int SimpleMaterial_LayerIndex = 0;
+	int SimpleMaterial_BSDFIndex = 0;
 	uint BSDFCount = 0;
 	uint BSDFVisibleCount = 0;
 	float3 Throughput = 1.0f;
