* Add 'is single material' bit into strata material header (not used yet)
* Shuffle BSDF type & basis bits to compact them next to other bit relevant to single material

diff --git a/Engine/Shaders/Private/Strata/Strata.ush b/Engine/Shaders/Private/Strata/Strata.ush
--- a/Engine/Shaders/Private/Strata/Strata.ush
+++ b/Engine/Shaders/Private/Strata/Strata.ush
@@ -131,32 +131,32 @@ struct FStrataBSDF
 // 
 // Header State (8bits):
 // * ISSIMPLEMATERIAL			
+// * ISSINGLEMATERIAL			
 // * HASPRECSHADOWMASK			
 // * ZEROPRECSHADOWMASK			
 // * CASTCONTACTSHADOW			
-// * HASDYNINDIRECTSHADOWCASTER	
+// * HASDYNINDIRECTSHADOWCASTER
 // * ISSINGLELAYERWATER			
 // * HASSUBSURFACE				
-// * UNUSED
 // 
 // * Complex: -------------------------------------------------------------------------------------------------------------------------------------------------------------------------
 //            HEADER        31 |  State | AO | BSDF Count | BSDF Types | Basis count | Basis type |
 //                             |  8       8    4            4            3             4           
 //            BASIS(0)      32 |  Normal: Oct 16/16 | Tangent: 11/11/10
 //            BASIS(1)      32 |  Normal: Oct 16/16 | Tangent: 11/11/10
-//			  ...			   |
+//                             |
 //            STATE(0)      17 |  BSDF Type | Normal Id | Aniso | TopLayer | SSS | SSSProfile | GreyWeigh | Haze | F90 | ThinFilm | Simple vol. | MFP Plugged | Fuzz | Thin
 //                             |  3           2           1       1          1     1            1           1      1     1          1             1             1      1 
 //            BSDF(0)       ?? |  Adaptive
 //            STATE(1)      17 |  BSDF Type | Normal Id | Aniso | TopLayer | SSS | SSSProfile | GreyWeigh | Haze | F90 | ThinFilm | Simple vol. | MFP Plugged | Fuzz | Thin
 //                             |  3           2           1       1          1     1            1           1      1     1          1             1             1      1 
 //            BSDF(1)       ?? |  Adaptive
-//			  ...			   |
+//                             |
 // * Fast   : -------------------------------------------------------------------------------------------------------------------------------------------------------------------------          
-//            HEADER	    29 | State | A0 | Roughness | DiffuseAbedo(8/20)
+//            HEADER        29 | State | A0 | Roughness | DiffuseAbedo(8/20)
 //                             | 8       8    8           8
 //            BASIS(0)       0 | (Use top layer normal)
-//            STATE(0)	     0 | 
+//            STATE(0)       0 | 
 //            BSDF(0)       32 | F0_RGB | DiffuseAbedo(12/20)
 //                             | 20       12
 //            -------------------------------------------------------------------------------------------------------------------------------------------------------------------------
@@ -171,34 +171,34 @@ struct FStrataBSDF
 
 ///////////////////////////////////////////////////////////////////////////////
 // Material Header
-// 
+//
 // Bit count for each component stored in a Strata header
 // 
 // State bit masks
 #define HEADER_MASK_ISSIMPLEMATERIAL			(1u << 0)
-#define HEADER_MASK_HASPRECSHADOWMASK			(1u << 1)
-#define HEADER_MASK_ZEROPRECSHADOWMASK			(1u << 2)
-#define HEADER_MASK_CASTCONTACTSHADOW			(1u << 3)
-#define HEADER_MASK_HASDYNINDIRECTSHADOWCASTER	(1u << 4)
-#define HEADER_MASK_ISSINGLELAYERWATER			(1u << 5)	// Used to avoid loading BSDF or Categorisation data when generating tiles containing SLW material. Avoid running another categorisation.
-#define HEADER_MASK_HASSUBSURFACE				(1u << 6)
-#define HEADER_MASK_UNUSED						(1u << 7)
+#define HEADER_MASK_ISSINGLEBSDF				(1u << 1)
+#define HEADER_MASK_HASPRECSHADOWMASK			(1u << 2)
+#define HEADER_MASK_ZEROPRECSHADOWMASK			(1u << 3)
+#define HEADER_MASK_CASTCONTACTSHADOW			(1u << 4)
+#define HEADER_MASK_HASDYNINDIRECTSHADOWCASTER	(1u << 5)
+#define HEADER_MASK_ISSINGLELAYERWATER			(1u << 6)	// Used to avoid loading BSDF or Categorisation data when generating tiles containing SLW material. Avoid running another categorisation.
+#define HEADER_MASK_HASSUBSURFACE				(1u << 7)
 #define HEADER_MASK_BIT_COUNT					8
 
 // For HEADER_BIT_BSDF_COUNT and  HEADER_BIT_SHAREDLOCALBASES_X, please refer to values in StrataDefinitions.h
 #define HEADER_BIT_COUNT_STATE						HEADER_MASK_BIT_COUNT
 #define HEADER_BIT_COUNT_AO							8
-#define HEADER_BIT_COUNT_BSDF_COUNT					4
 #define HEADER_BIT_COUNT_SHADINGMODELS				4
+#define HEADER_BIT_COUNT_BSDF_COUNT					4
 #define HEADER_BIT_COUNT_SHAREDLOCALBASES_COUNT		3
 #define HEADER_BIT_COUNT_SHAREDLOCALBASES_TYPE		4
 #define HEADER_BIT_COUNT___UNUSED___				1
 
 #define HEADER_BIT_OFFSET_STATE						0
 #define HEADER_BIT_OFFSET_AO						(HEADER_BIT_OFFSET_STATE					+ HEADER_BIT_COUNT_STATE)
-#define HEADER_BIT_OFFSET_BSDF_COUNT				(HEADER_BIT_OFFSET_AO						+ HEADER_BIT_COUNT_AO)
-#define HEADER_BIT_OFFSET_SHADINGMODELS				(HEADER_BIT_OFFSET_BSDF_COUNT				+ HEADER_BIT_COUNT_BSDF_COUNT)
-#define HEADER_BIT_OFFSET_SHAREDLOCALBASES_COUNT	(HEADER_BIT_OFFSET_SHADINGMODELS			+ HEADER_BIT_COUNT_SHADINGMODELS)
+#define HEADER_BIT_OFFSET_SHADINGMODELS				(HEADER_BIT_OFFSET_AO						+ HEADER_BIT_COUNT_AO)
+#define HEADER_BIT_OFFSET_BSDF_COUNT				(HEADER_BIT_OFFSET_SHADINGMODELS			+ HEADER_BIT_COUNT_SHADINGMODELS)
+#define HEADER_BIT_OFFSET_SHAREDLOCALBASES_COUNT	(HEADER_BIT_OFFSET_BSDF_COUNT				+ HEADER_BIT_COUNT_BSDF_COUNT)
 #define HEADER_BIT_OFFSET_SHAREDLOCALBASES_TYPE		(HEADER_BIT_OFFSET_SHAREDLOCALBASES_COUNT	+ HEADER_BIT_COUNT_SHAREDLOCALBASES_COUNT)
 
 #define HEADER_MASK_SHADINGMODELS					(((1u<<HEADER_BIT_COUNT_SHADINGMODELS)-1) << HEADER_BIT_OFFSET_SHADINGMODELS)
@@ -233,8 +233,6 @@ struct FStrataBSDF
 // BSDF State
  
 // Bit count for each component stored in a Strata BSDF state
-#define STATE_BIT_COUNT_BSDF_TYPE				3
-#define STATE_BIT_COUNT_SHAREDLOCALBASESID		2
 #define STATE_BIT_COUNT_ANISOTROPY				1
 #define STATE_BIT_COUNT_TOPLAYER				1
 #define STATE_BIT_COUNT_HASSSS					1
@@ -247,13 +245,13 @@ struct FStrataBSDF
 #define STATE_BIT_COUNT_DMFPPLUGGED				1
 #define STATE_BIT_COUNT_HASFUZZ					1
 #define STATE_BIT_COUNT_ISTHIN					1
+#define STATE_BIT_COUNT_BSDF_TYPE				3
+#define STATE_BIT_COUNT_SHAREDLOCALBASESID		2
 #define STATE_BIT_COUNT___UNUSED___				5
 #define STATE_BIT_COUNT_WEIGHT10F				10
 
 // Bit offset of each component store in BSDF state 
-#define STATE_BIT_OFFSET_BSDF_TYPE				0
-#define STATE_BIT_OFFSET_SHAREDLOCALBASESID		(STATE_BIT_OFFSET_BSDF_TYPE			+ STATE_BIT_COUNT_BSDF_TYPE)
-#define STATE_BIT_OFFSET_ANISOTROPY				(STATE_BIT_OFFSET_SHAREDLOCALBASESID+ STATE_BIT_COUNT_SHAREDLOCALBASESID)
+#define STATE_BIT_OFFSET_ANISOTROPY				0
 #define STATE_BIT_OFFSET_TOPLAYER				(STATE_BIT_OFFSET_ANISOTROPY		+ STATE_BIT_COUNT_ANISOTROPY)
 #define STATE_BIT_OFFSET_HASSSS					(STATE_BIT_OFFSET_TOPLAYER			+ STATE_BIT_COUNT_TOPLAYER)
 #define STATE_BIT_OFFSET_HASSSSPROFILE			(STATE_BIT_OFFSET_HASSSS			+ STATE_BIT_COUNT_HASSSS)
@@ -265,6 +263,8 @@ struct FStrataBSDF
 #define STATE_BIT_OFFSET_DMFPPLUGGED			(STATE_BIT_OFFSET_SIMPLEVOLUME		+ STATE_BIT_COUNT_SIMPLEVOLUME)
 #define STATE_BIT_OFFSET_HASFUZZ				(STATE_BIT_OFFSET_DMFPPLUGGED		+ STATE_BIT_COUNT_DMFPPLUGGED)
 #define STATE_BIT_OFFSET_ISTHIN					(STATE_BIT_OFFSET_HASFUZZ			+ STATE_BIT_COUNT_HASFUZZ)
+#define STATE_BIT_OFFSET_BSDF_TYPE				(STATE_BIT_OFFSET_ISTHIN			+ STATE_BIT_COUNT_ISTHIN)
+#define STATE_BIT_OFFSET_SHAREDLOCALBASESID		(STATE_BIT_OFFSET_BSDF_TYPE			+ STATE_BIT_COUNT_BSDF_TYPE)
 #define STATE_BIT_OFFSET_WEIGHT10F				22
 
 // Sanity check
@@ -886,6 +886,7 @@ struct FStrataPixelHeader
 };
 
 bool IsSimpleMaterial(const FStrataPixelHeader In)											{ return (In.State & HEADER_MASK_ISSIMPLEMATERIAL) != 0; }
+bool IsSingleMaterial(const FStrataPixelHeader In)											{ return (In.State & HEADER_MASK_ISSINGLEBSDF) != 0; }
 bool HasPrecShadowMask(const FStrataPixelHeader In)											{ return (In.State & HEADER_MASK_HASPRECSHADOWMASK) != 0; }
 bool HasZeroPrecShadowMask(const FStrataPixelHeader In)										{ return (In.State & HEADER_MASK_ZEROPRECSHADOWMASK) != 0; }
 bool IsSingleLayerWater(const FStrataPixelHeader In)										{ return (In.State & HEADER_MASK_ISSINGLELAYERWATER) != 0; }
@@ -898,6 +899,7 @@ bool HasDynamicIndirectShadowCasterRepresentation(const FStrataPixelHeader In)
 bool StrataHasShadingModel(const FStrataPixelHeader In, uint ShadingModel)					{ return StrataHasShadingModel(GetShadingModels(In), ShadingModel); }
 
 void SetIsSimpleMaterial(inout FStrataPixelHeader Out, bool bIn)							{ Out.State |= (bIn ? HEADER_MASK_ISSIMPLEMATERIAL : 0u); }
+void SetIsSingleMaterial(inout FStrataPixelHeader Out, bool bIn)							{ Out.State |= (bIn ? HEADER_MASK_ISSINGLEBSDF : 0u); }
 void SetHasPrecShadowMask(inout FStrataPixelHeader Out, bool bIn)							{ Out.State |= (bIn ? HEADER_MASK_HASPRECSHADOWMASK : 0u); }
 void SetZeroPrecShadowMask(inout FStrataPixelHeader Out, bool bIn)							{ Out.State |= (bIn ? HEADER_MASK_ZEROPRECSHADOWMASK : 0u); }
 void SetIsSingleLayerWater(inout FStrataPixelHeader Out, bool bIn)							{ Out.State |= (bIn ? HEADER_MASK_ISSINGLELAYERWATER : 0u); }
@@ -2002,7 +2004,7 @@ void PackStrataOut(
 					}
 
 					const bool bSSSMask = !bIsThin && bBottomLayer;	// SSS and Thin can only be on the bottom layer, so ignoring bIsSimpleVolume here.
-					SetHasSubsurface(StrataPixelHeader,			(bSSSMask && BSDF_GETHASSSS(BSDF)));
+					SetHasSubsurface(StrataPixelHeader,	(bSSSMask && BSDF_GETHASSSS(BSDF)));
 
 					// STRATA_TODO	Should we have the Fuzz influence the next layter visibility accoding to fuzz amount? Should it also occlude diffuse?
 					//				If this is done, we should also update ComputePostProcessBSDFData.
@@ -2068,7 +2070,7 @@ void PackStrataOut(
 						TopLayerData.WorldNormal += N * max(1e-5, BSDFCoverage);
 					}
 
-					if (bIsSimpleMaterial)
+					// if a material is simple/single, it is necessary the only one
 					{
 						SimpleMaterial_LayerIndex = l;
 						SimpleMaterial_BSDFIndex = i;
@@ -2096,7 +2098,10 @@ void PackStrataOut(
 			Throughput *= LayerThroughput;
 		}
 	}
+
+	// Set storage layout as either: fast(0), single(1), or complex(2)
 	AndIsSimpleMaterial(StrataPixelHeader, BSDFVisibleCount == 1);
+	SetIsSingleMaterial(StrataPixelHeader, BSDFVisibleCount == 1);
 
 	// Now write out Strata data
 
