Improve sample pattern in PostProcessHistogram.usf
- Use Gather
- Change iteration stride from 1 to THREADGROUP_SIZE
- Increase L1 hit rate (ie: from 13% to 90% generating bilateral grid)

diff --git a/Engine/Shaders/Private/PostProcessHistogram.usf b/Engine/Shaders/Private/PostProcessHistogram.usf
--- a/Engine/Shaders/Private/PostProcessHistogram.usf
+++ b/Engine/Shaders/Private/PostProcessHistogram.usf
@@ -8,9 +8,14 @@
 #include "ScreenPass.ush"
 #include "PostProcessHistogramCommon.ush"
 
+#if LOOP_SIZE_X % 2 || LOOP_SIZE_Y % 2
+#	 error "Must be even to use gather"
+#endif
+
 SCREEN_PASS_TEXTURE_VIEWPORT(Input)
 
 Texture2D InputTexture;
+SamplerState InputSampler;
 
 #if BILATERAL_GRID
 // Output bilateral grid texture (UAV)
@@ -49,72 +54,91 @@ void MainCS(
 	
 	GroupMemoryBarrierWithGroupSync();
 
-	// Each thread group processes LoopX * LoopY texels of the input.
-	uint2 TileSize = uint2(LOOP_SIZEX, LOOP_SIZEY);
+	// Each thread in the group processes LoopX * LoopY texels of the input.
+	const uint2 TileSize = uint2(THREADGROUP_SIZEX * LOOP_SIZEX, THREADGROUP_SIZEY * LOOP_SIZEY);
 
 	// Top left input texel for this group.
-	uint2 LeftTop = Input_ViewportMin + DispatchThreadId.xy * TileSize;
+	const uint2 LeftTop = Input_ViewportMin + GroupId.xy * TileSize;
 
-	float2 InvViewSize = Input_ViewportSizeInverse.xy;
+	const float2 InvViewSize = Input_ViewportSizeInverse.xy;
+	const float2 InvExtent = Input_ExtentInverse.xy;
+
+	const uint2 NumPixelsPerIter = uint2(2, 2);
 
 	// Accumulate all pixels into THREADGROUP_SIZEX*THREADGROUP_SIZEY histograms
-	LOOP for (uint y = 0; y < LOOP_SIZEY; ++y)
+	for (uint y = 0; y < THREADGROUP_SIZEY * LOOP_SIZEY; y += LOOP_SIZEY * NumPixelsPerIter.y)
 	{
-		LOOP for (uint x = 0; x < LOOP_SIZEX; ++x)
+		for (uint x = 0; x < THREADGROUP_SIZEX * LOOP_SIZEX; x += LOOP_SIZEX * NumPixelsPerIter.x)
 		{
-			uint2 Tile = uint2(x, y);
-			uint2 TexelPos = LeftTop + Tile;
+			uint2 TexelPos = LeftTop + uint2(x, y) + GroupThreadId.xy * NumPixelsPerIter;
 
-			if(TexelPos.x < Input_ViewportMax.x && TexelPos.y < Input_ViewportMax.y)
+			// don't include last column if viewport max is odd
+			if(all(TexelPos < Input_ViewportMax - 1))
 			{
-				float4 SceneColor = InputTexture.Load(int3(TexelPos, 0));
+				// sample at intersection of 4 texels
+				float2 UV = (TexelPos + 1.0f) * InvExtent;
 
-				SceneColor.xyz *= View.OneOverPreExposure;
+				float4 SceneColorR = InputTexture.GatherRed(InputSampler, UV);
+				float4 SceneColorG = InputTexture.GatherGreen(InputSampler, UV);
+				float4 SceneColorB = InputTexture.GatherBlue(InputSampler, UV);
 
-				float LuminanceVal = CalculateEyeAdaptationLuminance(SceneColor.xyz);
+				SceneColorR *= View.OneOverPreExposure;
+				SceneColorG *= View.OneOverPreExposure;
+				SceneColorB *= View.OneOverPreExposure;
 
-				float LogLuminance = log2(LuminanceVal);
-				float LogLuminanceHist = ComputeHistogramPositionFromLogLuminance(LogLuminance);
+				float4 LuminanceVal;
+				LuminanceVal.x = CalculateEyeAdaptationLuminance(float3(SceneColorR.x, SceneColorG.x, SceneColorB.x));
+				LuminanceVal.y = CalculateEyeAdaptationLuminance(float3(SceneColorR.y, SceneColorG.y, SceneColorB.y));
+				LuminanceVal.z = CalculateEyeAdaptationLuminance(float3(SceneColorR.z, SceneColorG.z, SceneColorB.z));
+				LuminanceVal.w = CalculateEyeAdaptationLuminance(float3(SceneColorR.w, SceneColorG.w, SceneColorB.w));
 
+				// only sample screen weight once per quad
 				float2 ScreenUV = (TexelPos.xy - Input_ViewportMin) * InvViewSize;
 				float ScreenWeight = AdaptationWeightTexture(ScreenUV);
 
-				// Map the normalized histogram position into texels.
-				float fBucket = saturate(LogLuminanceHist) * (HISTOGRAM_SIZE-1);// * 0.9999f;
+				UNROLL_N(4)
+				for (uint i = 0; i < 4; ++i)
+				{
+					float LogLuminance = log2(LuminanceVal[i]);
+					float LogLuminanceHist = ComputeHistogramPositionFromLogLuminance(LogLuminance);
+
+					// Map the normalized histogram position into texels.
+					float fBucket = saturate(LogLuminanceHist) * (HISTOGRAM_SIZE - 1);// * 0.9999f;
 
-				// Find two discrete buckets that straddle the continuous histogram position.
-				uint Bucket0 = (uint)(fBucket);
-				uint Bucket1 = Bucket0 + 1;
+					// Find two discrete buckets that straddle the continuous histogram position.
+					uint Bucket0 = (uint)(fBucket);
+					uint Bucket1 = Bucket0 + 1;
 
-				Bucket0 = min(Bucket0, uint(HISTOGRAM_SIZE - 1));
-				Bucket1 = min(Bucket1, uint(HISTOGRAM_SIZE - 1));
+					Bucket0 = min(Bucket0, uint(HISTOGRAM_SIZE - 1));
+					Bucket1 = min(Bucket1, uint(HISTOGRAM_SIZE - 1));
 
-				// Weighted blend between the two buckets.
-				float Weight1 = frac(fBucket);
-				float Weight0 = 1.0f - Weight1;
+					// Weighted blend between the two buckets.
+					float Weight1 = frac(fBucket);
+					float Weight0 = 1.0f - Weight1;
 
 #if !BILATERAL_GRID
-				// When EyeAdaptation_BlackHistogramBucketInfluence=.0, we will ignore the last bucket. The main use
-				// case is so the black background pixels in the editor have no effect. But if we have cases where
-				// pixel values can actually be black, we want to set EyeAdaptation_LastHistogramBucketInfluence=1.0.
-				// This value is controlled by the cvar "r.EyeAdaptation.BlackHistogramBucketInfluence"
-				if (Bucket0 == 0)
-				{
-					Weight0 *= EyeAdaptation_BlackHistogramBucketInfluence;
-				}
-
-				Weight0 *= ScreenWeight;
-				Weight1 *= ScreenWeight;
+					// When EyeAdaptation_BlackHistogramBucketInfluence=.0, we will ignore the last bucket. The main use
+					// case is so the black background pixels in the editor have no effect. But if we have cases where
+					// pixel values can actually be black, we want to set EyeAdaptation_LastHistogramBucketInfluence=1.0.
+					// This value is controlled by the cvar "r.EyeAdaptation.BlackHistogramBucketInfluence"
+					if (Bucket0 == 0)
+					{
+						Weight0 *= EyeAdaptation_BlackHistogramBucketInfluence;
+					}
+
+					Weight0 *= ScreenWeight;
+					Weight1 *= ScreenWeight;
 #endif
 
-				// Accumulate the weight to the nearby history buckets.
-				SharedHistogram[Bucket0][GroupThreadId.x][GroupThreadId.y] += Weight0;
-				SharedHistogram[Bucket1][GroupThreadId.x][GroupThreadId.y] += Weight1;
+					// Accumulate the weight to the nearby history buckets.
+					SharedHistogram[Bucket0][GroupThreadId.x][GroupThreadId.y] += Weight0;
+					SharedHistogram[Bucket1][GroupThreadId.x][GroupThreadId.y] += Weight1;
 
 #if BILATERAL_GRID
-				SharedLuminanceSum[Bucket0][GroupThreadId.x][GroupThreadId.y] += LogLuminance * Weight0;
-				SharedLuminanceSum[Bucket1][GroupThreadId.x][GroupThreadId.y] += LogLuminance * Weight1;
+					SharedLuminanceSum[Bucket0][GroupThreadId.x][GroupThreadId.y] += LogLuminance * Weight0;
+					SharedLuminanceSum[Bucket1][GroupThreadId.x][GroupThreadId.y] += LogLuminance * Weight1;
 #endif
+				}
 			}
 		}
 	}
@@ -122,12 +146,32 @@ void MainCS(
 	GroupMemoryBarrierWithGroupSync();
 
 	// Accumulate all histograms into one.
+
+#if BILATERAL_GRID
+#	if THREADGROUP_SIZEX * THREADGROUP_SIZEY < HISTOGRAM_SIZE
+#		error "Not enough threads to output complete histogram"
+#	endif
+
+	if (GroupIndex < HISTOGRAM_SIZE)
+	{
+		float Sum = 0;
+		float SumLuminance = 0;
+
+		LOOP for (uint y = 0; y < THREADGROUP_SIZEY; ++y)
+		{
+			LOOP for (uint x = 0; x < THREADGROUP_SIZEX; ++x)
+			{
+				Sum += SharedHistogram[GroupIndex][x][y];
+				SumLuminance += SharedLuminanceSum[GroupIndex][x][y];
+			}
+		}
+
+		BilateralGridRWTexture[uint3(GroupId.xy, GroupIndex)] = float2(SumLuminance, Sum);
+	}
+#else	
 	if (GroupIndex < HISTOGRAM_SIZE / 4) 
 	{
 		float4 Sum = 0;
-#if BILATERAL_GRID
-		float4 SumLuminance = 0;
-#endif
 
 		LOOP for (uint y = 0; y < THREADGROUP_SIZEY; ++y)
 		{
@@ -138,23 +182,9 @@ void MainCS(
 					SharedHistogram[GroupIndex * 4 + 1][x][y],
 					SharedHistogram[GroupIndex * 4 + 2][x][y],
 					SharedHistogram[GroupIndex * 4 + 3][x][y]);
-
-#if BILATERAL_GRID
-				SumLuminance += float4(
-					SharedLuminanceSum[GroupIndex * 4 + 0][x][y],
-					SharedLuminanceSum[GroupIndex * 4 + 1][x][y],
-					SharedLuminanceSum[GroupIndex * 4 + 2][x][y],
-					SharedLuminanceSum[GroupIndex * 4 + 3][x][y]);
-#endif
 			}
 		}
 
-#if BILATERAL_GRID
-		BilateralGridRWTexture[uint3(GroupId.xy, GroupIndex * 4 + 0)] = float2(SumLuminance.x, Sum.x);
-		BilateralGridRWTexture[uint3(GroupId.xy, GroupIndex * 4 + 1)] = float2(SumLuminance.y, Sum.y);
-		BilateralGridRWTexture[uint3(GroupId.xy, GroupIndex * 4 + 2)] = float2(SumLuminance.z, Sum.z);
-		BilateralGridRWTexture[uint3(GroupId.xy, GroupIndex * 4 + 3)] = float2(SumLuminance.w, Sum.w);
-#else
 		float2 MaxExtent = Input_ViewportSize;
 		float Area = MaxExtent.x * MaxExtent.y;
 	
@@ -163,6 +193,6 @@ void MainCS(
 
 		// Output texture with one histogram per line, x and y unwrapped into all the lines
 		HistogramRWTexture[uint2(GroupIndex, GroupId.x + GroupId.y * ThreadGroupCount.x)] = Sum * NormalizeFactor;
-#endif
 	}
+#endif
 }
diff --git a/Engine/Source/Runtime/Renderer/Private/PostProcess/PostProcessHistogram.cpp b/Engine/Source/Runtime/Renderer/Private/PostProcess/PostProcessHistogram.cpp
--- a/Engine/Source/Runtime/Renderer/Private/PostProcess/PostProcessHistogram.cpp
+++ b/Engine/Source/Runtime/Renderer/Private/PostProcess/PostProcessHistogram.cpp
@@ -37,6 +37,7 @@ public:
 		SHADER_PARAMETER_STRUCT_REF(FViewUniformShaderParameters, View)
 		SHADER_PARAMETER_STRUCT(FScreenPassTextureViewportParameters, Input)
 		SHADER_PARAMETER_STRUCT(FEyeAdaptationParameters, EyeAdaptation)
+		SHADER_PARAMETER_SAMPLER(SamplerState, InputSampler)
 		SHADER_PARAMETER_RDG_TEXTURE(Texture2D, InputTexture)
 		SHADER_PARAMETER_RDG_TEXTURE_UAV(RWTexture2D<float4>, HistogramRWTexture)
 		SHADER_PARAMETER_RDG_TEXTURE_UAV(RWTexture3D<float2>, BilateralGridRWTexture)
@@ -225,6 +226,7 @@ static FRDGTextureRef AddHistogramLegacyPass(
 		FHistogramCS::FParameters* PassParameters = GraphBuilder.AllocParameters<FHistogramCS::FParameters>();
 		PassParameters->View = View.ViewUniformBuffer;
 		PassParameters->Input = GetScreenPassTextureViewportParameters(FScreenPassTextureViewport(SceneColor));
+		PassParameters->InputSampler = TStaticSamplerState<SF_Point, AM_Clamp, AM_Clamp, AM_Clamp>::GetRHI();
 		PassParameters->InputTexture = SceneColor.Texture;
 		PassParameters->HistogramRWTexture = GraphBuilder.CreateUAV(HistogramTexture);
 		PassParameters->ThreadGroupCount = HistogramThreadGroupCount;
@@ -434,6 +436,7 @@ FRDGTextureRef AddLocalExposurePass(
 		PassParameters->View = View.ViewUniformBuffer;
 		PassParameters->EyeAdaptation = EyeAdaptationParameters;
 		PassParameters->Input = GetScreenPassTextureViewportParameters(FScreenPassTextureViewport(SceneColor));
+		PassParameters->InputSampler = TStaticSamplerState<SF_Point, AM_Clamp, AM_Clamp, AM_Clamp>::GetRHI();
 		PassParameters->InputTexture = SceneColor.Texture;
 		PassParameters->BilateralGridRWTexture = GraphBuilder.CreateUAV(LocalExposureTexture);
 		PassParameters->ThreadGroupCount = ThreadGroupCount;
