Double occupancy of bilaterial grid generation
- Pack 2 floats into 1 uint
- Use InterlockedAdd

diff --git a/Engine/Shaders/Private/PostProcessHistogram.usf b/Engine/Shaders/Private/PostProcessHistogram.usf
--- a/Engine/Shaders/Private/PostProcessHistogram.usf
+++ b/Engine/Shaders/Private/PostProcessHistogram.usf
@@ -17,6 +17,42 @@ SCREEN_PASS_TEXTURE_VIEWPORT(Input)
 Texture2D InputTexture;
 SamplerState InputSampler;
 
+uint PackFloatToUINT32(float v)
+{
+	return (uint)(v * float(1 << 19));
+}
+
+float UnpackUINT32ToFloat(uint v)
+{
+	float fv = (float)(v);
+	return fv * (1.0f / float(1 << 19));
+}
+
+// Want to pack 2 floats in a uint (16-bits for each)
+// each float is [0,1]
+// Each thread samples 8*8 tile = 64 texels
+// log2(64) = 6 => 7 bits
+// shift float by 16-7 = 9 bits
+#define PACK_SCALE float(1 << 9)
+
+uint PackTwoFloatToUINT32(float a, float b)
+{
+	uint ua = (uint)(a * PACK_SCALE) & 0xFFFF;
+	uint ub = (uint)(b * PACK_SCALE) & 0xFFFF;
+	return (ua << 16) | ub;
+}
+
+void UnpackUINT32ToTwoFloat(uint v, out float a, out float b)
+{
+	uint ua = v >> 16;
+	uint ub = v & 0xFFFF;
+
+	const float h = (1.0f / PACK_SCALE);
+
+	a = ua * h;
+	b = ub * h;
+}
+
 #if BILATERAL_GRID
 // Output bilateral grid texture (UAV)
 RWTexture3D<float2> BilateralGridRWTexture;
@@ -29,11 +65,7 @@ RWTexture2D<float4> HistogramRWTexture;
 uint2 ThreadGroupCount;
 
 // THREADGROUP_SIZEX*THREADGROUP_SIZEY histograms of the size HISTOGRAM_SIZE
-groupshared float SharedHistogram[HISTOGRAM_SIZE][THREADGROUP_SIZEX][THREADGROUP_SIZEY];
-
-#if BILATERAL_GRID
-groupshared float SharedLuminanceSum[HISTOGRAM_SIZE][THREADGROUP_SIZEX][THREADGROUP_SIZEY];
-#endif
+groupshared uint SharedHistogram[HISTOGRAM_SIZE][THREADGROUP_SIZEX][THREADGROUP_SIZEY];
 
 [numthreads(THREADGROUP_SIZEX, THREADGROUP_SIZEY, 1)]
 void MainCS(
@@ -46,10 +78,7 @@ void MainCS(
 	// clear all THREADGROUP_SIZEX*THREADGROUP_SIZEY histograms
 	UNROLL for (uint i = 0; i < HISTOGRAM_SIZE; ++i)
 	{
-		SharedHistogram[i][GroupThreadId.x][GroupThreadId.y] = 0.0f;
-#if BILATERAL_GRID
-		SharedLuminanceSum[i][GroupThreadId.x][GroupThreadId.y] = 0.0f;
-#endif
+		SharedHistogram[i][GroupThreadId.x][GroupThreadId.y] = 0;
 	}
 	
 	GroupMemoryBarrierWithGroupSync();
@@ -131,12 +160,19 @@ void MainCS(
 #endif
 
 					// Accumulate the weight to the nearby history buckets.
-					SharedHistogram[Bucket0][GroupThreadId.x][GroupThreadId.y] += Weight0;
-					SharedHistogram[Bucket1][GroupThreadId.x][GroupThreadId.y] += Weight1;
 
 #if BILATERAL_GRID
-					SharedLuminanceSum[Bucket0][GroupThreadId.x][GroupThreadId.y] += LogLuminance * Weight0;
-					SharedLuminanceSum[Bucket1][GroupThreadId.x][GroupThreadId.y] += LogLuminance * Weight1;
+					uint LogLuminanceHist0 = PackTwoFloatToUINT32(Weight0, LogLuminanceHist * Weight0);
+					uint LogLuminanceHist1 = PackTwoFloatToUINT32(Weight1, LogLuminanceHist * Weight1);
+
+					InterlockedAdd(SharedHistogram[Bucket0][GroupThreadId.x][GroupThreadId.y], LogLuminanceHist0);
+					InterlockedAdd(SharedHistogram[Bucket1][GroupThreadId.x][GroupThreadId.y], LogLuminanceHist1);
+#else
+					uint Weight0Int = PackFloatToUINT32(Weight0);
+					uint Weight1Int = PackFloatToUINT32(Weight1);
+
+					InterlockedAdd(SharedHistogram[Bucket0][GroupThreadId.x][GroupThreadId.y], Weight0Int);
+					InterlockedAdd(SharedHistogram[Bucket1][GroupThreadId.x][GroupThreadId.y], Weight1Int);
 #endif
 				}
 			}
@@ -161,8 +197,14 @@ void MainCS(
 		{
 			LOOP for (uint x = 0; x < THREADGROUP_SIZEX; ++x)
 			{
-				Sum += SharedHistogram[GroupIndex][x][y];
-				SumLuminance += SharedLuminanceSum[GroupIndex][x][y];
+				float SumWeight;
+				float SumLumHist;
+				UnpackUINT32ToTwoFloat(SharedHistogram[GroupIndex][x][y], SumWeight, SumLumHist);
+				// ComputeLogLuminanceFromHistogramPosition but include SumWeight
+				float SumLum = (SumLumHist - SumWeight * EyeAdaptation_HistogramBias) / EyeAdaptation_HistogramScale;
+
+				Sum += SumWeight;
+				SumLuminance += SumLum;
 			}
 		}
 
