Fix constant used for iteration stride in PostProcessHistogram.usf
- didn't cause issues atm since it's the same value.

diff --git a/Engine/Shaders/Private/PostProcessHistogram.usf b/Engine/Shaders/Private/PostProcessHistogram.usf
--- a/Engine/Shaders/Private/PostProcessHistogram.usf
+++ b/Engine/Shaders/Private/PostProcessHistogram.usf
@@ -95,9 +95,9 @@ void MainCS(
 	const uint2 NumPixelsPerIter = uint2(2, 2);
 
 	// Accumulate all pixels into THREADGROUP_SIZEX*THREADGROUP_SIZEY histograms
-	for (uint y = 0; y < THREADGROUP_SIZEY * LOOP_SIZEY; y += LOOP_SIZEY * NumPixelsPerIter.y)
+	for (uint y = 0; y < THREADGROUP_SIZEY * LOOP_SIZEY; y += THREADGROUP_SIZEY * NumPixelsPerIter.y)
 	{
-		for (uint x = 0; x < THREADGROUP_SIZEX * LOOP_SIZEX; x += LOOP_SIZEX * NumPixelsPerIter.x)
+		for (uint x = 0; x < THREADGROUP_SIZEX * LOOP_SIZEX; x += THREADGROUP_SIZEX * NumPixelsPerIter.x)
 		{
 			uint2 TexelPos = LeftTop + uint2(x, y) + GroupThreadId.xy * NumPixelsPerIter;
 
