Add single BSDF permuation for strata material for direct & sky lighting.

diff --git a/Engine/Shaders/Private/DeferredLightPixelShaders.usf b/Engine/Shaders/Private/DeferredLightPixelShaders.usf
--- a/Engine/Shaders/Private/DeferredLightPixelShaders.usf
+++ b/Engine/Shaders/Private/DeferredLightPixelShaders.usf
@@ -36,6 +36,15 @@
 #endif
 
 #if STRATA_ENABLED
+  #if STRATA_TILETYPE == 0
+	#define STRATA_FASTPATH 1
+  #elif STRATA_TILETYPE == 1
+	#define STRATA_SINGLEPATH 1
+  #elif STRATA_TILETYPE == 2
+	// COMPLEX PATH
+  #else
+	#error Strata tile type non-implemented
+  #endif
 #include "/Engine/Private/Strata/StrataEvaluation.ush"
 #include "/Engine/Private/Strata/StrataDeferredLighting.ush"
 #endif
diff --git a/Engine/Shaders/Private/ReflectionEnvironmentPixelShader.usf b/Engine/Shaders/Private/ReflectionEnvironmentPixelShader.usf
--- a/Engine/Shaders/Private/ReflectionEnvironmentPixelShader.usf
+++ b/Engine/Shaders/Private/ReflectionEnvironmentPixelShader.usf
@@ -105,6 +105,15 @@ FSkyLightVisibilityData GetSkyLightVisibilityData(float3 SkyLightingNormal, cons
 
 #if STRATA_ENABLED
 #define STRATA_INLINE_SHADING 0
+#if STRATA_TILETYPE == 0
+	#define STRATA_FASTPATH 1
+#elif STRATA_TILETYPE == 1
+	#define STRATA_SINGLEPATH 1
+#elif STRATA_TILETYPE == 2
+	// COMPLEX PATH
+#else
+	#error Strata tile type non-implemented
+#endif
 #include "/Engine/Private/Strata/StrataEvaluation.ush"
 #include "/Engine/Private/Strata/StrataEnvironmentLighting.ush"
 #endif
diff --git a/Engine/Shaders/Private/Strata/Strata.ush b/Engine/Shaders/Private/Strata/Strata.ush
--- a/Engine/Shaders/Private/Strata/Strata.ush
+++ b/Engine/Shaders/Private/Strata/Strata.ush
@@ -125,9 +125,10 @@ struct FStrataBSDF
 ////////////////////////////////////////////////////////////////////////////////////////////////////////////////////////////////////////////////////////////////////////////////////////////////
 // Storage layouts
 // ---------------
-// 2 types of storages layout:
+// 3 types of storages layout:
 // * Complex: store many (mixed) BSDFs
-// * Fast   : store a single Slab BSDF with no extra feature
+// * Single : store a single BSDF (with all feature apart from anisotropy)
+// * Simple : store a single Slab BSDF (with no extra feature)
 // 
 // Header State (8bits):
 // * ISSIMPLEMATERIAL			
@@ -139,8 +140,9 @@ struct FStrataBSDF
 // * ISSINGLELAYERWATER			
 // * HASSUBSURFACE				
 // 
-// * Complex: -------------------------------------------------------------------------------------------------------------------------------------------------------------------------
-//            HEADER        31 |  State | AO | BSDF Count | BSDF Types | Basis count | Basis type |
+// Layouts:
+// * LAYOUT 2 -------------------------------------------------------------------------------------------------------------------------------------------------------------------------
+//  (Complex) HEADER        31 |  State | AO | BSDF Types | BSDF Count | Basis count | Basis type |
 //                             |  8       8    4            4            3             4           
 //            BASIS(0)      32 |  Normal: Oct 16/16 | Tangent: 11/11/10
 //            BASIS(1)      32 |  Normal: Oct 16/16 | Tangent: 11/11/10
@@ -152,9 +154,14 @@ struct FStrataBSDF
 //                             |  3           2           1       1          1     1            1           1      1     1          1             1             1      1 
 //            BSDF(1)       ?? |  Adaptive
 //                             |
-// * Fast   : -------------------------------------------------------------------------------------------------------------------------------------------------------------------------          
-//            HEADER        29 | State | A0 | Roughness | DiffuseAbedo(8/20)
-//                             | 8       8    8           8
+// * LAYOUT 1 -------------------------------------------------------------------------------------------------------------------------------------------------------------------------
+//  (Single)  HEADER        31 |  State | AO | BSDF Types | Aniso | TopLayer | SSS | SSSProfile | GreyWeigh | Haze | F90 | ThinFilm | Simple vol. | MFP Plugged | Fuzz | Thin
+//                             |  8       8    4            1       1          1     1            1           1      1     1          1             1             1      1            
+//            BASIS(0)       0 | (Use top layer normal)
+//            BSDF(0)       ?? |  Adaptive
+// * LAYOUT 0 -------------------------------------------------------------------------------------------------------------------------------------------------------------------------          
+//  (Simple)  HEADER        29 | State | A0 | Roughness | DiffuseAbedo(8/20)
+//   (Fast)                    | 8       8    8           8
 //            BASIS(0)       0 | (Use top layer normal)
 //            STATE(0)       0 | 
 //            BSDF(0)       32 | F0_RGB | DiffuseAbedo(12/20)
@@ -176,7 +183,7 @@ struct FStrataBSDF
 // 
 // State bit masks
 #define HEADER_MASK_ISSIMPLEMATERIAL			(1u << 0)
-#define HEADER_MASK_ISSINGLEBSDF				(1u << 1)
+#define HEADER_MASK_ISSINGLEMATERIAL			(1u << 1)
 #define HEADER_MASK_HASPRECSHADOWMASK			(1u << 2)
 #define HEADER_MASK_ZEROPRECSHADOWMASK			(1u << 3)
 #define HEADER_MASK_CASTCONTACTSHADOW			(1u << 4)
@@ -203,8 +210,17 @@ struct FStrataBSDF
 
 #define HEADER_MASK_SHADINGMODELS					(((1u<<HEADER_BIT_COUNT_SHADINGMODELS)-1) << HEADER_BIT_OFFSET_SHADINGMODELS)
 
+// The different layout 'aliases' differently the material header
+// 
+// Layout 0 (Fast)
 #define HEADER_FASTENCODING_BIT_COUNT				(HEADER_BIT_COUNT_STATE + HEADER_BIT_COUNT_AO)
 #define HEADER_FASTENCODING_MASK					((1u<<HEADER_FASTENCODING_BIT_COUNT)-1u)
+// Layout 1 (Single BSDF)
+#define HEADER_SINGLEENCODING_BIT_COUNT				(HEADER_BIT_COUNT_STATE + HEADER_BIT_COUNT_AO + HEADER_BIT_COUNT_SHADINGMODELS)
+#define HEADER_SINGLEENCODING_MASK					((1u<<HEADER_SINGLEENCODING_BIT_COUNT)-1u)
+// Layout 2 (Complex)
+#define HEADER_COMPLEXENCODING_BIT_COUNT			(32u)
+#define HEADER_COMPLEXENCODING_MASK					(0xFFFFFFFFu)
 
 // Sanity check
 #if (HEADER_BIT_OFFSET_SHAREDLOCALBASES_TYPE + HEADER_BIT_COUNT_SHAREDLOCALBASES_TYPE) > 32
@@ -267,6 +283,16 @@ struct FStrataBSDF
 #define STATE_BIT_OFFSET_SHAREDLOCALBASESID		(STATE_BIT_OFFSET_BSDF_TYPE			+ STATE_BIT_COUNT_BSDF_TYPE)
 #define STATE_BIT_OFFSET_WEIGHT10F				22
 
+// Layout 0 (Fast)
+#define STATE_FASTENCODING_BIT_COUNT			(0u)
+#define STATE_FASTENCODING_MASK					(0u)
+// Layout 1 (Single BSDF)
+#define STATE_SINGLEENCODING_BIT_COUNT			(STATE_BIT_OFFSET_ISTHIN + STATE_BIT_COUNT_ISTHIN)
+#define STATE_SINGLEENCODING_MASK				((1u<<STATE_SINGLEENCODING_BIT_COUNT)-1u)
+// Layout 2 (Complex)
+#define STATE_COMPLEXCODING_BIT_COUNT			(32u)
+#define STATE_COMPLEXFASTENCODING_MASK			(0xFFFFFFFFu)
+
 // Sanity check
 #if (STATE_BIT_OFFSET_WEIGHT10F + STATE_BIT_COUNT_WEIGHT10F) > 32
 #error Strata BSDF state is larger than 32 bits
@@ -886,7 +912,7 @@ struct FStrataPixelHeader
 };
 
 bool IsSimpleMaterial(const FStrataPixelHeader In)											{ return (In.State & HEADER_MASK_ISSIMPLEMATERIAL) != 0; }
-bool IsSingleMaterial(const FStrataPixelHeader In)											{ return (In.State & HEADER_MASK_ISSINGLEBSDF) != 0; }
+bool IsSingleMaterial(const FStrataPixelHeader In)											{ return (In.State & HEADER_MASK_ISSINGLEMATERIAL) != 0; }
 bool HasPrecShadowMask(const FStrataPixelHeader In)											{ return (In.State & HEADER_MASK_HASPRECSHADOWMASK) != 0; }
 bool HasZeroPrecShadowMask(const FStrataPixelHeader In)										{ return (In.State & HEADER_MASK_ZEROPRECSHADOWMASK) != 0; }
 bool IsSingleLayerWater(const FStrataPixelHeader In)										{ return (In.State & HEADER_MASK_ISSINGLELAYERWATER) != 0; }
@@ -899,7 +925,7 @@ bool HasDynamicIndirectShadowCasterRepresentation(const FStrataPixelHeader In)
 bool StrataHasShadingModel(const FStrataPixelHeader In, uint ShadingModel)					{ return StrataHasShadingModel(GetShadingModels(In), ShadingModel); }
 
 void SetIsSimpleMaterial(inout FStrataPixelHeader Out, bool bIn)							{ Out.State |= (bIn ? HEADER_MASK_ISSIMPLEMATERIAL : 0u); }
-void SetIsSingleMaterial(inout FStrataPixelHeader Out, bool bIn)							{ Out.State |= (bIn ? HEADER_MASK_ISSINGLEBSDF : 0u); }
+void SetIsSingleMaterial(inout FStrataPixelHeader Out, bool bIn)							{ Out.State |= (bIn ? HEADER_MASK_ISSINGLEMATERIAL : 0u); }
 void SetHasPrecShadowMask(inout FStrataPixelHeader Out, bool bIn)							{ Out.State |= (bIn ? HEADER_MASK_HASPRECSHADOWMASK : 0u); }
 void SetZeroPrecShadowMask(inout FStrataPixelHeader Out, bool bIn)							{ Out.State |= (bIn ? HEADER_MASK_ZEROPRECSHADOWMASK : 0u); }
 void SetIsSingleLayerWater(inout FStrataPixelHeader Out, bool bIn)							{ Out.State |= (bIn ? HEADER_MASK_ISSINGLELAYERWATER : 0u); }
@@ -909,6 +935,7 @@ void SetHasSubsurface(inout FStrataPixelHeader Out, bool bIn)								{ Out.State
 
 void AddShadingModels(inout FStrataPixelHeader Out, uint bIn)								{ Out.State |= ((bIn<<HEADER_BIT_OFFSET_SHADINGMODELS) & HEADER_MASK_SHADINGMODELS); }
 void AndIsSimpleMaterial(inout FStrataPixelHeader Out, bool bIn)							{ Out.State = (Out.State & ~HEADER_MASK_ISSIMPLEMATERIAL) | ((Out.State & HEADER_MASK_ISSIMPLEMATERIAL) & (bIn ? HEADER_MASK_ISSIMPLEMATERIAL : 0u)); }
+void AndIsSingleMaterial(inout FStrataPixelHeader Out, bool bIn)							{ Out.State = (Out.State & ~HEADER_MASK_ISSINGLEMATERIAL) | ((Out.State & HEADER_MASK_ISSINGLEMATERIAL) & (bIn ? HEADER_MASK_ISSINGLEMATERIAL : 0u)); }
 
 FStrataPixelHeader InitialiseStrataPixelHeader()
 {
@@ -935,7 +962,7 @@ float3x3 StrataGetBSDFSharedBasis(in FStrataPixelHeader StrataPixelHeader, in FS
 	float3 Bitangent = cross(Normal, Tangent);
 	return float3x3(Tangent, Bitangent, Normal);
 #else
-	if (IsSimpleMaterial(StrataPixelHeader))
+	if (IsSimpleMaterial(StrataPixelHeader) || IsSingleMaterial(StrataPixelHeader))
 	{		
 		return GetTangentBasis(StrataUnpackTopLayerData(StrataPixelHeader.PackedTopLayerData).WorldNormal);
 	}
@@ -1854,6 +1881,7 @@ void PackStrataOut(
 	TopLayerData = (FStrataTopLayerData)0;
 
 	SetIsSimpleMaterial(StrataPixelHeader, true); // non strata material or sky pixels are considered as simple (so that (strata simple pixels) + (non-strata pixels) can build a simple tile).
+	SetIsSingleMaterial(StrataPixelHeader, true); // non strata material or sky pixels are considered as simple (so that (strata simple pixels) + (non-strata pixels) can build a simple tile).
 
 	int SimpleMaterial_LayerIndex = 0;
 	int SimpleMaterial_BSDFIndex = 0;
@@ -1932,12 +1960,18 @@ void PackStrataOut(
 				Init(ContextNoLight, N, V, FakeL);
 
 				bool bIsSimpleMaterial = true;
+				bool bIsSingleMaterial = true;
 				switch (BSDF_GETTYPE(BSDF))
 				{
 				case STRATA_BSDF_TYPE_SLAB:
 				{
+					// Update simple material compatibility
 					bIsSimpleMaterial = IsStrataSlabFastPathCompatible(BSDF);
 
+					// Update single material compatibility. For now, single materials don't support anisotropy, 
+					// they use the complex pass, as the toplayer data does not contain the full frame basis, only the top normal.
+					bIsSingleMaterial = BSDF_GETHASANISOTROPY(BSDF) == 0;
+
 					EmissiveColor += BSDF_GETEMISSIVE(BSDF) * Throughput;
 
 					// Two sided lighting can only be evalauted on the bottom layer, together with option SSS. Otherwise, the simple volume will be used.
@@ -2063,6 +2097,7 @@ void PackStrataOut(
 
 					AddShadingModels(StrataPixelHeader, StrataShadingModelBit(BSDF_GETTYPE(BSDF)));
 					AndIsSimpleMaterial(StrataPixelHeader, bIsSimpleMaterial);
+					AndIsSingleMaterial(StrataPixelHeader, bIsSingleMaterial);
 
 					if (bTopLayer)
 					{
@@ -2101,17 +2136,20 @@ void PackStrataOut(
 
 	// Set storage layout as either: fast(0), single(1), or complex(2)
 	AndIsSimpleMaterial(StrataPixelHeader, BSDFVisibleCount == 1);
-	SetIsSingleMaterial(StrataPixelHeader, BSDFVisibleCount == 1);
+	AndIsSingleMaterial(StrataPixelHeader, BSDFVisibleCount == 1);
 
 	// Now write out Strata data
 
 	///////////////////////////////////////////////////////////////////////////
-	// Two types of encodings
-	// * A. Fast encoding
-	// * B. Regular complex
-	const bool bHasFastEncoding = IsSimpleMaterial(StrataPixelHeader);
-
-	if (!bHasFastEncoding)
+	// 3 types of encodings
+	// * A. Layout0: Simple encoding	(use top layer normal)
+	// * B. Layout1: Single encoding	(use top layer normal)
+	// * C. Layout2: Complex encoding	(use basis)
+	const bool bHasFastEncoding   = IsSimpleMaterial(StrataPixelHeader);
+	const bool bHasSingleEncoding = IsSingleMaterial(StrataPixelHeader);
+
+	// (Layout2)
+	if (!bHasFastEncoding && !bHasSingleEncoding)
 	{
 		// 1. the header (Regular/Complex encoding)
 		const uint PackedHeader = PackStrataHeader(BSDFVisibleCount, StrataPixelHeader);
@@ -2155,7 +2193,7 @@ void PackStrataOut(
 					const uint GreyScaleThroughput = StrataHasGreyScaleWeight(BSDF.FinalWeight) ? 1 : 0;
 					BSDF_SETHASGREYWEIGHT(BSDF, GreyScaleThroughput);
 
-					// A. Fast path encoding: store header & data
+					// A. Layout0 - Simple encoding (aka. fast-path): store header & data
 					if (bHasFastEncoding && l == SimpleMaterial_LayerIndex && i == SimpleMaterial_BSDFIndex)
 					{		
 						// Now we pack diffuse and F0 in a special way: both encoded as R7G7B6 with gamma 2.0. The low bit cound will be hidden by dithering and TAA.
@@ -2166,7 +2204,7 @@ void PackStrataOut(
 						const uint PackedF020Bits		= PackR7G7B6Gamma2(SLAB_F0(BSDF));
 						const uint PackedRoughness8bits = PackR8(SLAB_ROUGHNESS(BSDF));
 
-						// Data0 (State|AO|Roughness|Diffuse8bits)
+						// Data0 (Header_State|Header_AO|Roughness|Diffuse8bits)
 						{
 							uint Out = 0;
 							HEADER_SETSTATE(Out, StrataPixelHeader.State);
@@ -2188,7 +2226,25 @@ void PackStrataOut(
 						// Ensure the rest of the BSDF is not stored with the regular path
 						BSDF.State = 0;
 					}
-					// B. Weight for Regular/Complex path
+					// B. Layout1 - Simple encoding: single BSDF, whose header & BSDF state are merged
+					else if (bHasSingleEncoding && l == SimpleMaterial_LayerIndex && i == SimpleMaterial_BSDFIndex)
+					{
+						// Data0 (Header_State|Header_AO|Header_BSDFTypes|BSDF_State)
+						{
+							uint Out = 0;
+							HEADER_SETSTATE(Out, StrataPixelHeader.State);
+							HEADER_SETAO(Out, PackR8(StrataPixelHeader.MaterialAO));
+							HEADER_SETSHADINGMODELS(Out, GetShadingModels(StrataPixelHeader));
+							Out = Out & HEADER_SINGLEENCODING_MASK;							
+							Out = Out | ((BSDF.State & STATE_SINGLEENCODING_MASK) << HEADER_SINGLEENCODING_BIT_COUNT);
+							STRATA_STORE_UINT1(Out);
+
+							#if (HEADER_SIMPLEENCODING_BIT_COUNT) > 32
+							#error Strata fast path header is > 32bits
+							#endif
+						}
+					}
+					// C. Layout2 - Weight for Regular/Complex path
 					else if (GreyScaleThroughput > 0)
 					{
 						BSDF_SETWEIGHT10F(BSDF, Pack10F(BSDF.FinalWeight.x));
@@ -2200,8 +2256,8 @@ void PackStrataOut(
 						STRATA_STORE_UINT1(PackR11G11B10F(BSDF.FinalWeight));
 					}
 
+					// Layout1 & Layout2
 					const uint BSDFType = BSDF_GETTYPE(BSDF);
-
 					switch (BSDF_GETTYPE(BSDF))
 					{
 					case STRATA_BSDF_TYPE_SLAB:
@@ -2324,16 +2380,25 @@ FStrataPixelHeader UnpackStrataHeaderIn(FStrataMaterialContainer StrataBuffer, i
 	Out.BSDFCount   = HEADER_GETBSDFCOUNT(PackedHeader);
 
 	const bool bIsSimpleMaterial = (PackedHeader & HEADER_MASK_ISSIMPLEMATERIAL) > 0;
+	const bool bIsSingleMaterial = (PackedHeader & HEADER_MASK_ISSINGLEMATERIAL) > 0;
 
-	// Use fast path for decoding data
+	// Use simple or single layout for decoding data
 	BRANCH
-	if (bIsSimpleMaterial)
+	if (bIsSimpleMaterial || bIsSingleMaterial)
 	{
 		Out.BSDFCount	= 1;
 
-		// Clear shading model, and ensurewe have only STRATA_BSDF_TYPE_SLAB
-		Out.State = Out.State & (~HEADER_MASK_SHADINGMODELS);
-		HEADER_SETSHADINGMODELS(Out.State, StrataShadingModelBit(STRATA_BSDF_TYPE_SLAB));
+		if (bIsSimpleMaterial)
+		{
+			// Clear shading model, and ensurewe have only STRATA_BSDF_TYPE_SLAB
+			Out.State = Out.State & (~HEADER_MASK_SHADINGMODELS);
+			HEADER_SETSHADINGMODELS(Out.State, StrataShadingModelBit(STRATA_BSDF_TYPE_SLAB));
+		}
+		else // if (bIsSingleMaterial)
+		{
+			const uint ShadingModels = HEADER_GETSHADINGMODELS(PackedHeader);
+			HEADER_SETSHADINGMODELS(Out.State, ShadingModels);
+		}
 
 	#if STRATA_INLINE_SHADING==0
 		Out.StrataBuffer = StrataBuffer;
@@ -2443,17 +2508,32 @@ FStrataBSDF UnpackStrataBSDFIn(FStrataMaterialContainer StrataBuffer, inout FStr
 
 	float DummyFloat = 0.0f;
 
-	OutBSDF.State = StrataLoadUint1(StrataBuffer, StrataAddressing);
-	if (BSDF_GETHASGREYWEIGHT(OutBSDF))
+	BRANCH
+	if (IsSingleMaterial(StrataHeader))
 	{
-		OutBSDF.FinalWeight = Unpack10F(BSDF_GETWEIGHT10F(OutBSDF));
+		uint Data = 0;
+	#if STRATA_INLINE_SHADING==0
+		Data = StrataHeader.SharedLocalBasesTypes_PackedHeader;
+	#endif
+		OutBSDF.State = (Data>>HEADER_SINGLEENCODING_BIT_COUNT) & STATE_SINGLEENCODING_MASK;
+		OutBSDF.FinalWeight = 1.0f;
+
+		// For now, single encode doesn't support anisotropy, as the top layer data only contains the normal, not the full frame basis
+		BSDF_SETHASANISOTROPY(OutBSDF, 0);
 	}
 	else
 	{
-		StrataLoad_R11G11B10F(StrataBuffer, StrataAddressing, OutBSDF.FinalWeight);
+		OutBSDF.State = StrataLoadUint1(StrataBuffer, StrataAddressing);
+		if (BSDF_GETHASGREYWEIGHT(OutBSDF))
+		{
+			OutBSDF.FinalWeight = Unpack10F(BSDF_GETWEIGHT10F(OutBSDF));
+		}
+		else
+		{
+			StrataLoad_R11G11B10F(StrataBuffer, StrataAddressing, OutBSDF.FinalWeight);
+		}
 	}
 
-
 	const uint BSDFType = BSDF_GETTYPE(OutBSDF);
 	BSDF_SETTYPE(OutBSDF, BSDFType);
 	switch (BSDFType)
diff --git a/Engine/Shaders/Private/Strata/StrataDeferredLighting.ush b/Engine/Shaders/Private/Strata/StrataDeferredLighting.ush
--- a/Engine/Shaders/Private/Strata/StrataDeferredLighting.ush
+++ b/Engine/Shaders/Private/Strata/StrataDeferredLighting.ush
@@ -8,6 +8,10 @@
 #define STRATA_FASTPATH 0
 #endif
 
+#ifndef STRATA_SINGLEPATH 
+#define STRATA_SINGLEPATH 0
+#endif
+
 // Analytical lighting evaluation for strata material.
 // Unpack BSDF on-the-fly
 float4 StrataDeferredLighting(
@@ -25,7 +29,7 @@ float4 StrataDeferredLighting(
 	FLightAccumulator Out = (FLightAccumulator)0;
 	const bool bForceFullyRough = false;
 
-#if STRATA_FASTPATH == 0
+#if STRATA_FASTPATH == 0 && STRATA_SINGLEPATH == 0
 	LOOP
 	for (uint BSDFIndex = 0; BSDFIndex < StrataPixelHeader.BSDFCount; ++BSDFIndex)
 #endif
diff --git a/Engine/Shaders/Private/Strata/StrataEnvironmentLighting.ush b/Engine/Shaders/Private/Strata/StrataEnvironmentLighting.ush
--- a/Engine/Shaders/Private/Strata/StrataEnvironmentLighting.ush
+++ b/Engine/Shaders/Private/Strata/StrataEnvironmentLighting.ush
@@ -6,6 +6,10 @@
 #define STRATA_FASTPATH 0
 #endif
 
+#ifndef STRATA_SINGLEPATH 
+#define STRATA_SINGLEPATH 0
+#endif
+
 // Env. lighting evaluation for strata material.
 // Unpack BSDF on-the-fly
 float4 StrataEnvironmentLighting(
@@ -48,7 +52,7 @@ float4 StrataEnvironmentLighting(
 	uint NumCulledReflectionCaptures = 0;
 #endif
 
-#if STRATA_FASTPATH == 0
+#if STRATA_FASTPATH == 0 && STRATA_SINGLEPATH == 0
 	for (uint BSDFIndex = 0; BSDFIndex < StrataPixelHeader.BSDFCount; ++BSDFIndex)
 #endif
 	{
diff --git a/Engine/Shaders/Private/Strata/StrataMaterialClassification.usf b/Engine/Shaders/Private/Strata/StrataMaterialClassification.usf
--- a/Engine/Shaders/Private/Strata/StrataMaterialClassification.usf
+++ b/Engine/Shaders/Private/Strata/StrataMaterialClassification.usf
@@ -28,13 +28,18 @@ int2 ViewResolution;
 uint MaxBytesPerPixel;
 Texture2D<uint> TopLayerTexture;
 Texture2DArray<uint> MaterialTextureArray;
+// Fast/Simple
 RWBuffer<uint> SimpleTileIndirectDataBuffer;
 RWBuffer<uint> SimpleTileListDataBuffer;
+// Single
+RWBuffer<uint> SingleTileIndirectDataBuffer;
+RWBuffer<uint> SingleTileListDataBuffer;
+// Complex
 RWBuffer<uint> ComplexTileIndirectDataBuffer;
 RWBuffer<uint> ComplexTileListDataBuffer;
 
 #if !PERMUTATION_WAVE_OPS
-groupshared uint s_IsComplexIsMaterialTile[GROUP_THREAD_COUNT];
+groupshared uint s_TileFlags[GROUP_THREAD_COUNT];
 #endif
 
 #if PERMUTATION_STRATA_CLEAR_DURING_CATEGORIZATION
@@ -47,6 +52,7 @@ void TileMainCS(uint2 DispatchThreadId : SV_DispatchThreadID, uint LinearIndex :
 	if (all(DispatchThreadId == 0))
 	{
 		SimpleTileIndirectDataBuffer[0] = bRectPrimitive > 0 ? 4 : 6;
+		SingleTileIndirectDataBuffer[0] = bRectPrimitive > 0 ? 4 : 6;
 		ComplexTileIndirectDataBuffer[0] = bRectPrimitive > 0 ? 4 : 6;
 	}
 
@@ -55,16 +61,21 @@ void TileMainCS(uint2 DispatchThreadId : SV_DispatchThreadID, uint LinearIndex :
 
 	// Pixels outside of the view area are considered simple to enable screen borders to receive the simple permutation when not aligned to shader group size.
 	bool bContainsComplexMaterial = false;
-	bool bContainsStrataMaterial = false;
+	bool bContainsSimpleMaterial  = false;
+	bool bContainsSingleMaterial  = false;
+	bool bContainsStrataMaterial  = false;
 	bool bHasScreenSpaceSubsurfaceScattering = false;
 	if (all(PixelCoord < uint2(ViewResolution)))
 	{
 		FStrataAddressing StrataAddressing = GetStrataPixelDataByteOffset(PixelCoord, uint2(View.BufferSizeAndInvSize.xy), MaxBytesPerPixel);
 		FStrataPixelHeader StrataPixelHeader = UnpackStrataHeaderIn(MaterialTextureArray, StrataAddressing, TopLayerTexture);
 
-		const bool bIsSimple = IsSimpleMaterial(StrataPixelHeader) || StrataPixelHeader.BSDFCount == 0; // BSDFCount == 0 ensures that non-strata pixel, like sky pixels, won't make a simple tile flagged as complex
+		const bool bIsSimple =  IsSimpleMaterial(StrataPixelHeader) || StrataPixelHeader.BSDFCount == 0; // BSDFCount == 0 ensures that non-strata pixel, like sky pixels, won't make a simple tile flagged as complex
+		const bool bIsSingle = !IsSimpleMaterial(StrataPixelHeader) && IsSingleMaterial(StrataPixelHeader);
 		bContainsStrataMaterial  = StrataPixelHeader.BSDFCount > 0;
-		bContainsComplexMaterial = !bIsSimple;
+		bContainsSimpleMaterial  = bIsSimple;
+		bContainsSingleMaterial  = bIsSingle;
+		bContainsComplexMaterial = !bIsSingle & !bIsSimple;
 		bHasScreenSpaceSubsurfaceScattering = HasSubsurface(StrataPixelHeader);
 	}
 
@@ -79,78 +90,75 @@ void TileMainCS(uint2 DispatchThreadId : SV_DispatchThreadID, uint LinearIndex :
 #endif
 
 #if PERMUTATION_WAVE_OPS
-
 	const bool bHasAnyStrata  = WaveActiveAnyTrue(bContainsStrataMaterial);
+//	const bool bHasAnySimple  = WaveActiveAnyTrue(bContainsSimpleMaterial);
+	const bool bHasAnySingle  = WaveActiveAnyTrue(bContainsSingleMaterial);
 	const bool bHasAnyComplex = WaveActiveAnyTrue(bContainsComplexMaterial);
-
-	if (LinearIndex < 1 && bHasAnyStrata)
-	{
-		if (bHasAnyComplex)
-		{
-			uint WriteToIndex;
-			InterlockedAdd(ComplexTileIndirectDataBuffer[1], 1, WriteToIndex);
-			ComplexTileListDataBuffer[WriteToIndex] = EncodeTile(GroupId.xy);
-		}
-		else
-		{
-			uint WriteToIndex;
-			InterlockedAdd(SimpleTileIndirectDataBuffer[1], 1, WriteToIndex);
-			SimpleTileListDataBuffer[WriteToIndex] = EncodeTile(GroupId.xy);
-		}
-	}
-
 #else // PERMUTATION_WAVE_OPS
 
-	s_IsComplexIsMaterialTile[LinearIndex] = (bContainsComplexMaterial ? 1 : 0) | (bContainsStrataMaterial ? 2 : 0);
+	s_TileFlags[LinearIndex] = 
+		(bContainsStrataMaterial  ? 0x1u : 0u) | 
+//		(bContainsSimpleMaterial  ? 0x2u : 0u) | 
+		(bContainsSingleMaterial  ? 0x4u : 0u) | 
+		(bContainsComplexMaterial ? 0x8u : 0u);
 
 	GroupMemoryBarrierWithGroupSync();
 	if (LinearIndex < 32)
 	{
-		s_IsComplexIsMaterialTile[LinearIndex] = s_IsComplexIsMaterialTile[LinearIndex] | s_IsComplexIsMaterialTile[LinearIndex + 32];
+		s_TileFlags[LinearIndex] = s_TileFlags[LinearIndex] | s_TileFlags[LinearIndex + 32];
 	}
 	GroupMemoryBarrierWithGroupSync();
 	if (LinearIndex < 16)
 	{
-		s_IsComplexIsMaterialTile[LinearIndex] = s_IsComplexIsMaterialTile[LinearIndex] | s_IsComplexIsMaterialTile[LinearIndex + 16];
+		s_TileFlags[LinearIndex] = s_TileFlags[LinearIndex] | s_TileFlags[LinearIndex + 16];
 	}
 	GroupMemoryBarrierWithGroupSync();
 
 	if (LinearIndex < 8)
 	{
-		s_IsComplexIsMaterialTile[LinearIndex] = s_IsComplexIsMaterialTile[LinearIndex] | s_IsComplexIsMaterialTile[LinearIndex + 8];
+		s_TileFlags[LinearIndex] = s_TileFlags[LinearIndex] | s_TileFlags[LinearIndex + 8];
 	}
 	GroupMemoryBarrierWithGroupSync();
 	if (LinearIndex < 4)
 	{
-		s_IsComplexIsMaterialTile[LinearIndex] = s_IsComplexIsMaterialTile[LinearIndex] | s_IsComplexIsMaterialTile[LinearIndex + 4];
+		s_TileFlags[LinearIndex] = s_TileFlags[LinearIndex] | s_TileFlags[LinearIndex + 4];
 	}
 	GroupMemoryBarrierWithGroupSync();
 	if (LinearIndex < 2)
 	{
-		s_IsComplexIsMaterialTile[LinearIndex] = s_IsComplexIsMaterialTile[LinearIndex] | s_IsComplexIsMaterialTile[LinearIndex + 2];
+		s_TileFlags[LinearIndex] = s_TileFlags[LinearIndex] | s_TileFlags[LinearIndex + 2];
 	}
 	GroupMemoryBarrierWithGroupSync();
 
-	const uint FinalIsSimpleIsMaterialTile = s_IsComplexIsMaterialTile[LinearIndex] | s_IsComplexIsMaterialTile[LinearIndex + 1];
-	bContainsComplexMaterial = (FinalIsSimpleIsMaterialTile & 1) > 0;
-	bContainsStrataMaterial = (FinalIsSimpleIsMaterialTile & 2) > 0;
+	const uint FinalTileFlags = s_TileFlags[LinearIndex] | s_TileFlags[LinearIndex + 1];
+
+	const bool bHasAnyStrata  = (FinalTileFlags & 0x1u) > 0;
+//	const bool bHasAnySimple  = (FinalTileFlags & 0x2u) > 0;
+	const bool bHasAnySingle  = (FinalTileFlags & 0x4u) > 0;
+	const bool bHasAnyComplex = (FinalTileFlags & 0x8u) > 0;
+#endif // PERMUTATION_WAVE_OPS
 
-	if (LinearIndex < 1 && bContainsStrataMaterial)
+	if (LinearIndex < 1 && bHasAnyStrata)
 	{
-		if (bContainsComplexMaterial)
+		if (bHasAnyComplex)
 		{
 			uint WriteToIndex;
 			InterlockedAdd(ComplexTileIndirectDataBuffer[1], 1, WriteToIndex);
 			ComplexTileListDataBuffer[WriteToIndex] = EncodeTile(GroupId.xy);
 		}
-		else
+		else if (bHasAnySingle)
+		{
+			uint WriteToIndex;
+			InterlockedAdd(SingleTileIndirectDataBuffer[1], 1, WriteToIndex);
+			SingleTileListDataBuffer[WriteToIndex] = EncodeTile(GroupId.xy);
+		}
+		else // (bHasAnySimple)
 		{
 			uint WriteToIndex;
 			InterlockedAdd(SimpleTileIndirectDataBuffer[1], 1, WriteToIndex);
 			SimpleTileListDataBuffer[WriteToIndex] = EncodeTile(GroupId.xy);
 		}
 	}
-#endif // PERMUTATION_WAVE_OPS
 }
 #endif // SHADER_TILE_CATEGORIZATION
 
diff --git a/Engine/Source/Runtime/Renderer/Private/IndirectLightRendering.cpp b/Engine/Source/Runtime/Renderer/Private/IndirectLightRendering.cpp
--- a/Engine/Source/Runtime/Renderer/Private/IndirectLightRendering.cpp
+++ b/Engine/Source/Runtime/Renderer/Private/IndirectLightRendering.cpp
@@ -204,7 +204,7 @@ class FReflectionEnvironmentSkyLightingPS : public FGlobalShader
 	class FDynamicSkyLight : SHADER_PERMUTATION_BOOL("ENABLE_DYNAMIC_SKY_LIGHT");
 	class FSkyShadowing : SHADER_PERMUTATION_BOOL("APPLY_SKY_SHADOWING");
 	class FRayTracedReflections : SHADER_PERMUTATION_BOOL("RAY_TRACED_REFLECTIONS");
-	class FStrataFastPath : SHADER_PERMUTATION_BOOL("STRATA_FASTPATH");
+	class FStrataTileType : SHADER_PERMUTATION_INT("STRATA_TILETYPE", 3);
 
 	using FPermutationDomain = TShaderPermutationDomain<
 		FHasBoxCaptures,
@@ -214,7 +214,7 @@ class FReflectionEnvironmentSkyLightingPS : public FGlobalShader
 		FDynamicSkyLight,
 		FSkyShadowing,
 		FRayTracedReflections,
-		FStrataFastPath>;
+		FStrataTileType>;
 
 	static FPermutationDomain RemapPermutation(FPermutationDomain PermutationVector)
 	{
@@ -230,15 +230,15 @@ class FReflectionEnvironmentSkyLightingPS : public FGlobalShader
 			PermutationVector.Set<FSkyShadowing>(false);
 		}
 
-		if (PermutationVector.Get<FStrataFastPath>() && !Strata::IsStrataEnabled())
+		if (PermutationVector.Get<FStrataTileType>() && !Strata::IsStrataEnabled())
 		{
-			PermutationVector.Set<FStrataFastPath>(false);
+			PermutationVector.Set<FStrataTileType>(0);
 		}
 
 		return PermutationVector;
 	}
 
-	static FPermutationDomain BuildPermutationVector(const FViewInfo& View, bool bBoxCapturesOnly, bool bSphereCapturesOnly, bool bSupportDFAOIndirectOcclusion, bool bEnableSkyLight, bool bEnableDynamicSkyLight, bool bApplySkyShadowing, bool bRayTracedReflections, bool bStrataFastPath)
+	static FPermutationDomain BuildPermutationVector(const FViewInfo& View, bool bBoxCapturesOnly, bool bSphereCapturesOnly, bool bSupportDFAOIndirectOcclusion, bool bEnableSkyLight, bool bEnableDynamicSkyLight, bool bApplySkyShadowing, bool bRayTracedReflections, EStrataTileMaterialType TileType)
 	{
 		FPermutationDomain PermutationVector;
 
@@ -249,8 +249,12 @@ class FReflectionEnvironmentSkyLightingPS : public FGlobalShader
 		PermutationVector.Set<FDynamicSkyLight>(bEnableDynamicSkyLight);
 		PermutationVector.Set<FSkyShadowing>(bApplySkyShadowing);
 		PermutationVector.Set<FRayTracedReflections>(bRayTracedReflections);
-		PermutationVector.Set<FStrataFastPath>(Strata::IsStrataEnabled() && bStrataFastPath);
-
+		PermutationVector.Set<FStrataTileType>(0);
+		if (Strata::IsStrataEnabled())
+		{
+			check(EStrataTileMaterialType::ECount != TileType);
+			PermutationVector.Set<FStrataTileType>(TileType);
+		}
 		return RemapPermutation(PermutationVector);
 	}
 
@@ -1278,7 +1282,7 @@ static void AddSkyReflectionPass(
 		View, bHasBoxCaptures, bHasSphereCaptures, DynamicBentNormalAO != 0.0f,
 		bSkyLight, bDynamicSkyLight, bApplySkyShadowing,
 		bRequiresSpecializedReflectionEnvironmentShader,
-		StrataTileMaterialType == EStrataTileMaterialType::ESimple);
+		StrataTileMaterialType);
 
 	TShaderMapRef<FReflectionEnvironmentSkyLightingPS> PixelShader(View.ShaderMap, PermutationVector);
 	ClearUnusedGraphResources(PixelShader, &PassParameters->PS);
@@ -1596,6 +1600,20 @@ void FDeferredShadingSceneRenderer::RenderDeferredReflectionsAndSkyLighting(
 					bApplySkyShadowing,
 					EStrataTileMaterialType::ESimple);
 
+				AddSkyReflectionPass(
+					GraphBuilder,
+					View,
+					Scene,
+					SceneTextures,
+					DynamicBentNormalAOTexture,
+					ReflectionsColor,
+					RayTracingReflectionOptions,
+					SceneTextureParameters,
+					bSkyLight,
+					bDynamicSkyLight,
+					bApplySkyShadowing,
+					EStrataTileMaterialType::ESingle);
+
 				AddSkyReflectionPass(
 					GraphBuilder,
 					View,
diff --git a/Engine/Source/Runtime/Renderer/Private/LightRendering.cpp b/Engine/Source/Runtime/Renderer/Private/LightRendering.cpp
--- a/Engine/Source/Runtime/Renderer/Private/LightRendering.cpp
+++ b/Engine/Source/Runtime/Renderer/Private/LightRendering.cpp
@@ -425,7 +425,7 @@ class FDeferredLightPS : public FGlobalShader
 	class FAtmosphereTransmittance : SHADER_PERMUTATION_BOOL("USE_ATMOSPHERE_TRANSMITTANCE");
 	class FCloudTransmittance 	: SHADER_PERMUTATION_BOOL("USE_CLOUD_TRANSMITTANCE");
 	class FAnistropicMaterials 	: SHADER_PERMUTATION_BOOL("SUPPORTS_ANISOTROPIC_MATERIALS");
-	class FStrataFastPath		: SHADER_PERMUTATION_BOOL("STRATA_FASTPATH");
+	class FStrataTileType		: SHADER_PERMUTATION_INT("STRATA_TILETYPE", 3);
 
 	using FPermutationDomain = TShaderPermutationDomain<
 		FSourceShapeDim,
@@ -438,7 +438,7 @@ class FDeferredLightPS : public FGlobalShader
 		FAtmosphereTransmittance,
 		FCloudTransmittance,
 		FAnistropicMaterials,
-		FStrataFastPath>;
+		FStrataTileType>;
 
 	BEGIN_SHADER_PARAMETER_STRUCT(FParameters, )
 		SHADER_PARAMETER_RDG_UNIFORM_BUFFER(FSceneTextureUniformParameters, SceneTextures)
@@ -526,7 +526,7 @@ class FDeferredLightPS : public FGlobalShader
 			}
 		}
 
-		if (PermutationVector.Get<FStrataFastPath>() && !Strata::IsStrataEnabled())
+		if (!Strata::IsStrataEnabled() && PermutationVector.Get<FStrataTileType>() != 0)
 		{
 			return false;
 		}
@@ -1687,7 +1687,9 @@ static uint32 InternalSetBoundingGeometryDepthState(FGraphicsPipelineStateInitia
 		switch (TileType)
 		{
 		case EStrataTileMaterialType::ESimple : StencilRef = Strata::StencilBit_Fast;    GraphicsPSOInit.DepthStencilState = TStaticDepthStencilState<false, CompareFunction, true, CF_Equal, SO_Keep, SO_Keep, SO_Keep, true, CF_Equal, SO_Keep, SO_Keep, SO_Keep, Strata::StencilBit_Fast, 0x0>::GetRHI(); break;
+		case EStrataTileMaterialType::ESingle : StencilRef = Strata::StencilBit_Single;  GraphicsPSOInit.DepthStencilState = TStaticDepthStencilState<false, CompareFunction, true, CF_Equal, SO_Keep, SO_Keep, SO_Keep, true, CF_Equal, SO_Keep, SO_Keep, SO_Keep, Strata::StencilBit_Single, 0x0>::GetRHI(); break;
 		case EStrataTileMaterialType::EComplex: StencilRef = Strata::StencilBit_Complex; GraphicsPSOInit.DepthStencilState = TStaticDepthStencilState<false, CompareFunction, true, CF_Equal, SO_Keep, SO_Keep, SO_Keep, true, CF_Equal, SO_Keep, SO_Keep, SO_Keep, 0x0, 0x0>::GetRHI(); break;
+		default: check(false);
 		}
 	}
 	else
@@ -1968,12 +1970,14 @@ static void InternalRenderLight(
 /** Shader parameters for Standard Deferred Light Overlap Debug pass. */
 BEGIN_SHADER_PARAMETER_STRUCT(FRenderLightParameters, )
 	// PS/VS parameter structs
-	SHADER_PARAMETER_STRUCT_INCLUDE(FDeferredLightPS::FParameters, PS)	
+	SHADER_PARAMETER_STRUCT_INCLUDE(FDeferredLightPS::FParameters, PS)
 	SHADER_PARAMETER_STRUCT_INCLUDE(FDeferredLightVS::FParameters, VS)
 	// Strata tiles
 	SHADER_PARAMETER_RDG_BUFFER_SRV(Buffer<uint>, TileListBufferSimple)
+	SHADER_PARAMETER_RDG_BUFFER_SRV(Buffer<uint>, TileListBufferSingle)
 	SHADER_PARAMETER_RDG_BUFFER_SRV(Buffer<uint>, TileListBufferComplex)
 	RDG_BUFFER_ACCESS(TileIndirectBufferSimple, ERHIAccess::IndirectArgs)
+	RDG_BUFFER_ACCESS(TileIndirectBufferSingle, ERHIAccess::IndirectArgs)
 	RDG_BUFFER_ACCESS(TileIndirectBufferComplex, ERHIAccess::IndirectArgs)
 END_SHADER_PARAMETER_STRUCT()
 
@@ -2062,8 +2066,10 @@ static void RenderLight(
 		if (Strata::IsStrataEnabled())
 		{
 			PassParameters->TileListBufferSimple = View.StrataSceneData->ClassificationTileListBufferSRV[EStrataTileMaterialType::ESimple];
+			PassParameters->TileListBufferSingle = View.StrataSceneData->ClassificationTileListBufferSRV[EStrataTileMaterialType::ESingle];
 			PassParameters->TileListBufferComplex = View.StrataSceneData->ClassificationTileListBufferSRV[EStrataTileMaterialType::EComplex];
 			PassParameters->TileIndirectBufferSimple = View.StrataSceneData->ClassificationTileIndirectBuffer[EStrataTileMaterialType::ESimple];
+			PassParameters->TileIndirectBufferSingle = View.StrataSceneData->ClassificationTileIndirectBuffer[EStrataTileMaterialType::ESingle];
 			PassParameters->TileIndirectBufferComplex = View.StrataSceneData->ClassificationTileIndirectBuffer[EStrataTileMaterialType::EComplex];
 		}
 		else
@@ -2081,7 +2087,7 @@ static void RenderLight(
 		PermutationVector.Set< FDeferredLightPS::FHairLighting>(0);
 		PermutationVector.Set< FDeferredLightPS::FLightingChannelsDim >(View.bUsesLightingChannels);
 		PermutationVector.Set< FDeferredLightPS::FVisualizeCullingDim >(View.Family->EngineShowFlags.VisualizeLightCulling);
-		PermutationVector.Set< FDeferredLightPS::FStrataFastPath >(false);
+		PermutationVector.Set< FDeferredLightPS::FStrataTileType >(0);
 		if (bIsRadial)
 		{
 			PermutationVector.Set< FDeferredLightPS::FSourceShapeDim >(LightProxy->IsRectLight() ? ELightSourceShape::Rect : ELightSourceShape::Capsule);
@@ -2111,20 +2117,29 @@ static void RenderLight(
 		{
 			// Simple tiles
 			{
-				PermutationVector.Set< FDeferredLightPS::FStrataFastPath>(true);
+				const EStrataTileMaterialType TileType = EStrataTileMaterialType::ESimple;
+				PermutationVector.Set<FDeferredLightPS::FStrataTileType>(TileType);
+				TShaderMapRef< FDeferredLightPS > PixelShader(View.ShaderMap, PermutationVector);
+				InternalRenderLight(GraphBuilder, Scene, View, LightSceneInfo, PixelShader, PassParameters, TileType, TEXT("Light::StandardDeferred(Simple)"));
+			}
+			// Single tiles
+			{
+				const EStrataTileMaterialType TileType = EStrataTileMaterialType::ESingle;
+				PermutationVector.Set<FDeferredLightPS::FStrataTileType>(TileType);
 				TShaderMapRef< FDeferredLightPS > PixelShader(View.ShaderMap, PermutationVector);
-				InternalRenderLight(GraphBuilder, Scene, View, LightSceneInfo, PixelShader, PassParameters, EStrataTileMaterialType::ESimple, TEXT("Light::StandardDeferred(Simple)"));
+				InternalRenderLight(GraphBuilder, Scene, View, LightSceneInfo, PixelShader, PassParameters, TileType, TEXT("Light::StandardDeferred(Single)"));
 			}
 			// Complex tiles
 			{
-				PermutationVector.Set< FDeferredLightPS::FStrataFastPath>(false);
+				const EStrataTileMaterialType TileType = EStrataTileMaterialType::EComplex;
+				PermutationVector.Set<FDeferredLightPS::FStrataTileType>(TileType);
 				TShaderMapRef< FDeferredLightPS > PixelShader(View.ShaderMap, PermutationVector);
-				InternalRenderLight(GraphBuilder, Scene, View, LightSceneInfo, PixelShader, PassParameters, EStrataTileMaterialType::EComplex, TEXT("Light::StandardDeferred(Complex)"));
+				InternalRenderLight(GraphBuilder, Scene, View, LightSceneInfo, PixelShader, PassParameters, TileType, TEXT("Light::StandardDeferred(Complex)"));
 			}
 		}
 		else
 		{
-			PermutationVector.Set< FDeferredLightPS::FStrataFastPath>(false);
+			PermutationVector.Set< FDeferredLightPS::FStrataTileType>(0);
 			TShaderMapRef< FDeferredLightPS > PixelShader(View.ShaderMap, PermutationVector);
 			InternalRenderLight(GraphBuilder, Scene, View, LightSceneInfo, PixelShader, PassParameters, EStrataTileMaterialType::ECount, TEXT("Light::StandardDeferred"));
 		}
@@ -2402,7 +2417,7 @@ void FDeferredShadingSceneRenderer::RenderSimpleLightsStandardDeferred(
 		PermutationVector.Set< FDeferredLightPS::FHairLighting>(0);
 		PermutationVector.Set< FDeferredLightPS::FAtmosphereTransmittance >(false);
 		PermutationVector.Set< FDeferredLightPS::FCloudTransmittance >(false);
-		PermutationVector.Set< FDeferredLightPS::FStrataFastPath>(false);
+		PermutationVector.Set< FDeferredLightPS::FStrataTileType>(0);
 		TShaderMapRef<FDeferredLightPS> PixelShader(View.ShaderMap, PermutationVector);
 
 		FDeferredLightVS::FPermutationDomain PermutationVectorVS;
@@ -2441,7 +2456,7 @@ void FDeferredShadingSceneRenderer::RenderSimpleLightsStandardDeferred(
 								//@todo - accurate ortho camera / light intersection
 								|| !View.IsPerspectiveProjection();
 
-				SetBoundingGeometryRasterizerAndDepthState(GraphicsPSOInit, View, bCameraInsideLightGeometry, TileType);
+				SetBoundingGeometryRasterizerAndDepthState(GraphicsPSOInit, View, bCameraInsideLightGeometry, EStrataTileMaterialType::EComplex);
 				GraphicsPSOInit.BoundShaderState.VertexDeclarationRHI = GetVertexDeclarationFVector4();
 				GraphicsPSOInit.BoundShaderState.VertexShaderRHI = VertexShader.GetVertexShader();
 				GraphicsPSOInit.BoundShaderState.PixelShaderRHI = PixelShader.GetPixelShader();
diff --git a/Engine/Source/Runtime/Renderer/Private/PostProcess/SceneRenderTargets.h b/Engine/Source/Runtime/Renderer/Private/PostProcess/SceneRenderTargets.h
--- a/Engine/Source/Runtime/Renderer/Private/PostProcess/SceneRenderTargets.h
+++ b/Engine/Source/Runtime/Renderer/Private/PostProcess/SceneRenderTargets.h
@@ -26,6 +26,7 @@
 #define STENCIL_RECEIVE_DECAL_BIT_ID		7
 // Used only during the lighting pass - alias/reuse decal bit
 #define STENCIL_STRATA_FASTPATH				7 
+#define STENCIL_STRATA_SINGLEPATH			1
 
 // Outputs a compile-time constant stencil's bit mask ready to be used
 // in TStaticDepthStencilState<> template parameter. It also takes care
diff --git a/Engine/Source/Runtime/Renderer/Private/Strata/Strata.cpp b/Engine/Source/Runtime/Renderer/Private/Strata/Strata.cpp
--- a/Engine/Source/Runtime/Renderer/Private/Strata/Strata.cpp
+++ b/Engine/Source/Runtime/Renderer/Private/Strata/Strata.cpp
@@ -93,6 +93,7 @@ const TCHAR* ToString(EStrataTileMaterialType Type)
 	switch (Type)
 	{
 	case EStrataTileMaterialType::ESimple:		return TEXT("Simple");
+	case EStrataTileMaterialType::ESingle:		return TEXT("Single");
 	case EStrataTileMaterialType::EComplex:		return TEXT("Complex");
 	}
 	return TEXT("Unknown");
@@ -155,9 +156,10 @@ void InitialiseStrataFrameSceneData(FSceneRenderer& SceneRenderer, FRDGBuilder&
 			const int32 TileInPixel = GetStrataBufferTileSize();
 			const FIntPoint TileResolution(FMath::DivideAndRoundUp(SceneTextureExtent.X, TileInPixel), FMath::DivideAndRoundUp(SceneTextureExtent.Y, TileInPixel));
 
-			const TCHAR* StrataTileResourceNames[EStrataTileMaterialType::ECount][2] =
+			const TCHAR* StrataTileResourceNames[EStrataTileMaterialType::ECount][3] =
 			{
 				{ TEXT("Strata.StrataTileListBuffer(Simple)"),	TEXT("Strata.StrataTileIndirectBuffer(Simple)") },
+				{ TEXT("Strata.StrataTileListBuffer(Single)"), 	TEXT("Strata.StrataTileIndirectBuffer(Single)") },
 				{ TEXT("Strata.StrataTileListBuffer(Complex)"),	TEXT("Strata.StrataTileIndirectBuffer(Complex)")}
 			};
 
@@ -427,6 +429,8 @@ class FStrataMaterialTileClassificationPassCS : public FGlobalShader
 		SHADER_PARAMETER_RDG_TEXTURE_SRV(Texture2DArray<uint>, MaterialTextureArray)
 		SHADER_PARAMETER_RDG_BUFFER_UAV(RWBuffer, SimpleTileIndirectDataBuffer)
 		SHADER_PARAMETER_RDG_BUFFER_UAV(RWBuffer, SimpleTileListDataBuffer)
+		SHADER_PARAMETER_RDG_BUFFER_UAV(RWBuffer, SingleTileIndirectDataBuffer)
+		SHADER_PARAMETER_RDG_BUFFER_UAV(RWBuffer, SingleTileListDataBuffer)
 		SHADER_PARAMETER_RDG_BUFFER_UAV(RWBuffer, ComplexTileIndirectDataBuffer)
 		SHADER_PARAMETER_RDG_BUFFER_UAV(RWBuffer, ComplexTileListDataBuffer)
 		SHADER_PARAMETER_RDG_TEXTURE_UAV(RWTexture2D<uint2>, SSSTextureUAV)
@@ -516,12 +520,6 @@ static void AddStrataInternalClassificationTilePass(
 	EStrataTileMaterialType TileMaterialType,
 	const bool bDebug = false)
 {
-	// We cannot early exit due to the fact that the local light are still rendered as mesh volumes (so cannot be tiled as such)
-	//if (ShouldPassesReadingStrataBeTiled())
-	//{
-	//	return;
-	//}
-
 	EPrimitiveType StrataTilePrimitiveType = PT_TriangleList;
 	const FIntPoint OutputResolution = View.ViewRect.Size();
 	FVector4f OutputResolutionAndInv = FVector4f(OutputResolution.X, OutputResolution.Y, 1.0f / float(OutputResolution.X), 1.0f / float(OutputResolution.Y));
@@ -542,14 +540,10 @@ static void AddStrataInternalClassificationTilePass(
 		ParametersPS->RenderTargets[0] = FRenderTargetBinding(*ColorTexture, ERenderTargetLoadAction::ELoad);
 		switch (TileMaterialType)
 		{
-		case EStrataTileMaterialType::ESimple:
-			ParametersPS->DebugTileColor = FVector4f(0.0f, 1.0f, 0.0f, 1.0);
-			break;
-		case EStrataTileMaterialType::EComplex:
-			ParametersPS->DebugTileColor = FVector4f(1.0f, 0.0f, 0.0f, 1.0);
-			break;
-		default:
-			check(false);
+		case EStrataTileMaterialType::ESimple: ParametersPS->DebugTileColor = FVector4f(0.0f, 1.0f, 0.0f, 1.0); break;
+		case EStrataTileMaterialType::ESingle: ParametersPS->DebugTileColor = FVector4f(1.0f, 1.0f, 0.0f, 1.0); break;
+		case EStrataTileMaterialType::EComplex: ParametersPS->DebugTileColor = FVector4f(1.0f, 0.0f, 0.0f, 1.0); break;
+		default: check(false);
 		}
 	}
 	else
@@ -564,9 +558,7 @@ static void AddStrataInternalClassificationTilePass(
 	}
 	
 	GraphBuilder.AddPass(
-		bDebug ? 
-		RDG_EVENT_NAME("Strata::DebugClassificationPass") :
-		RDG_EVENT_NAME("Strata::StencilClassificationPass"),
+		RDG_EVENT_NAME("Strata::%sClassificationPass(%s)", bDebug ? TEXT("Debug") : TEXT("Stencil"), ToString(TileMaterialType)),
 		ParametersPS,
 		ERDGPassFlags::Raster,
 		[ParametersPS, VertexShader, PixelShader, OutputResolution, StrataTilePrimitiveType, bDebug](FRHICommandList& RHICmdList)
@@ -614,6 +606,7 @@ void AddStrataStencilPass(
 	const FMinimalSceneTextures& SceneTextures)
 {
 	AddStrataInternalClassificationTilePass(GraphBuilder, View, &SceneTextures.Depth.Target, nullptr, EStrataTileMaterialType::ESimple);
+	AddStrataInternalClassificationTilePass(GraphBuilder, View, &SceneTextures.Depth.Target, nullptr, EStrataTileMaterialType::ESingle);
 }
 
 void AddStrataStencilPass(
@@ -625,6 +618,7 @@ void AddStrataStencilPass(
 	{
 		const FViewInfo& View = Views[i];
 		AddStrataInternalClassificationTilePass(GraphBuilder, View, &SceneTextures.Depth.Target, nullptr, EStrataTileMaterialType::ESimple);
+		AddStrataInternalClassificationTilePass(GraphBuilder, View, &SceneTextures.Depth.Target, nullptr, EStrataTileMaterialType::ESingle);
 	}
 }
 
@@ -759,11 +753,16 @@ void AddStrataMaterialClassificationPass(FRDGBuilder& GraphBuilder, const FMinim
 			PassParameters->MaxBytesPerPixel = View.StrataSceneData->MaxBytesPerPixel;
 			PassParameters->TopLayerTexture = View.StrataSceneData->TopLayerTexture;
 			PassParameters->MaterialTextureArray = View.StrataSceneData->MaterialTextureArraySRV;
+			PassParameters->SSSTextureUAV = View.StrataSceneData->SSSTextureUAV;
+			// Simple
 			PassParameters->SimpleTileListDataBuffer = View.StrataSceneData->ClassificationTileListBufferUAV[EStrataTileMaterialType::ESimple];
 			PassParameters->SimpleTileIndirectDataBuffer = View.StrataSceneData->ClassificationTileIndirectBufferUAV[EStrataTileMaterialType::ESimple];
+			// Single
+			PassParameters->SingleTileListDataBuffer = View.StrataSceneData->ClassificationTileListBufferUAV[EStrataTileMaterialType::ESingle];
+			PassParameters->SingleTileIndirectDataBuffer = View.StrataSceneData->ClassificationTileIndirectBufferUAV[EStrataTileMaterialType::ESingle];
+			// Complex
 			PassParameters->ComplexTileListDataBuffer = View.StrataSceneData->ClassificationTileListBufferUAV[EStrataTileMaterialType::EComplex];
 			PassParameters->ComplexTileIndirectDataBuffer = View.StrataSceneData->ClassificationTileIndirectBufferUAV[EStrataTileMaterialType::EComplex];
-			PassParameters->SSSTextureUAV = View.StrataSceneData->SSSTextureUAV;
 
 			const uint32 GroupSize = 8;
 			FComputeShaderUtils::AddPass(
@@ -829,10 +828,9 @@ FScreenPassTexture AddStrataDebugPasses(FRDGBuilder& GraphBuilder, const FViewIn
 	{
 		RDG_EVENT_SCOPE(GraphBuilder, "Strata::VisualizeClassification");
 		const bool bDebugPass = true;
-		AddStrataInternalClassificationTilePass(
-			GraphBuilder, View, nullptr, &ScreenPassSceneColor.Texture, EStrataTileMaterialType::ESimple, bDebugPass);
-		AddStrataInternalClassificationTilePass(
-			GraphBuilder, View, nullptr, &ScreenPassSceneColor.Texture, EStrataTileMaterialType::EComplex, bDebugPass);
+		AddStrataInternalClassificationTilePass(GraphBuilder, View, nullptr, &ScreenPassSceneColor.Texture, EStrataTileMaterialType::ESimple, bDebugPass);
+		AddStrataInternalClassificationTilePass(GraphBuilder, View, nullptr, &ScreenPassSceneColor.Texture, EStrataTileMaterialType::ESingle, bDebugPass);
+		AddStrataInternalClassificationTilePass(GraphBuilder, View, nullptr, &ScreenPassSceneColor.Texture, EStrataTileMaterialType::EComplex, bDebugPass);
 	}
 
 	return MoveTemp(ScreenPassSceneColor);
diff --git a/Engine/Source/Runtime/Renderer/Private/Strata/Strata.h b/Engine/Source/Runtime/Renderer/Private/Strata/Strata.h
--- a/Engine/Source/Runtime/Renderer/Private/Strata/Strata.h
+++ b/Engine/Source/Runtime/Renderer/Private/Strata/Strata.h
@@ -36,8 +36,9 @@ END_GLOBAL_SHADER_PARAMETER_STRUCT()
 
 enum EStrataTileMaterialType : uint32
 {
-	ESimple = 0,
-	EComplex = 1,
+	ESimple  = 0,
+	ESingle  = 1,
+	EComplex = 2,
 	ECount
 };
 
@@ -84,6 +85,7 @@ struct FStrataSceneData
 namespace Strata
 {
 constexpr uint32 StencilBit_Complex= 0x00; // No stencil bit
+constexpr uint32 StencilBit_Single = 0x02; // In sync with SceneRenderTargets.h - GET_STENCIL_BIT_MASK(STENCIL_STRATA_SINGLEPATH)
 constexpr uint32 StencilBit_Fast   = 0x80; // In sync with SceneRenderTargets.h - GET_STENCIL_BIT_MASK(STENCIL_STRATA_FASTPATH)
 
 bool IsStrataEnabled();
