Strata
- unlit fast now use opacity
- forward translucent now has correct fogging (needed to separate transmittance)
- clarified some code

preflight https://horde.devtools.epicgames.com/job/61b8b1dbb026ce352ac5efc8

diff --git a/Engine/Shaders/Private/BasePassPixelShader.usf b/Engine/Shaders/Private/BasePassPixelShader.usf
--- a/Engine/Shaders/Private/BasePassPixelShader.usf
+++ b/Engine/Shaders/Private/BasePassPixelShader.usf
@@ -88,6 +88,11 @@
 
 #define STRATA_ENABLED					(PROJECT_STRATA && MATERIAL_IS_STRATA)
 
+// We detect the isolated unlit BSDF node case to make sure lots of code is optimised out, when used, there can only be one.
+// When unlit is selected as part of dynamic shading model, the SLAB model is used instead.
+// This allows us to avoid having the unlit node case inside the general BSDF loop and 0 BSDF to process in such pixels.
+#define STRATA_OPTIMIZED_UNLIT			MATERIAL_SHADINGMODEL_UNLIT
+
 float NormalCurvatureToRoughness(float3 WorldNormal)
 {
     float3 dNdx = ddx(WorldNormal);
@@ -1030,7 +1035,7 @@ void FPixelShaderInOut_MainPS(
 #endif
 #endif // !STRATA_ENABLED
 
-#if STRATA_ENABLED && !MATERIAL_SHADINGMODEL_UNLIT
+#if STRATA_ENABLED && !STRATA_OPTIMIZED_UNLIT
 	FStrataData StrataData = PixelMaterialInputs.FrontMaterial;
 
 	// Initialise a Strata header with normal in registers
@@ -1125,7 +1130,7 @@ void FPixelShaderInOut_MainPS(
 #endif // !STRATA_ENABLED
 
 	// Static shadow mask
-#if STRATA_ENABLED && !MATERIAL_SHADINGMODEL_UNLIT && GBUFFER_HAS_PRECSHADOWFACTOR
+#if STRATA_ENABLED && !STRATA_OPTIMIZED_UNLIT && GBUFFER_HAS_PRECSHADOWFACTOR
 	{
 		// Encode shadow mask only if the shadow mask is entirely non-zero and non-one
 	#if WRITES_PRECSHADOWFACTOR_ZERO
@@ -1511,9 +1516,9 @@ void FPixelShaderInOut_MainPS(
 #endif
 
 #if MATERIAL_WORKS_WITH_DUAL_SOURCE_COLOR_BLENDING
-	float3 DualBlendColorAdd = 0.0f;
-	float3 DualBlendColorMul = 1.0f;
-	float3 DualBlendColorCov = 1.0f;
+	float3 DualBlendSurfaceLuminancePostCoverage	= 0.0f;
+	float3 DualBlendSurfaceTransmittancePreCoverage = 1.0f;
+	float DualBlendSurfaceCoverage					= 1.0f;
 #endif
 
 #if !STRATA_ENABLED || STRATA_INLINE_SINGLELAYERWATER
@@ -1573,8 +1578,9 @@ void FPixelShaderInOut_MainPS(
 #if MATERIAL_SHADINGMODEL_THIN_TRANSLUCENT && !STRATA_ENABLED
 	{
 		AccumulateThinTranslucentModel(
-								DualBlendColorAdd,
-								DualBlendColorMul,
+								DualBlendSurfaceLuminancePostCoverage,
+								DualBlendSurfaceTransmittancePreCoverage,
+								DualBlendSurfaceCoverage,
 								MaterialParameters,
 								GBuffer,
 								DiffuseColor,
@@ -1589,7 +1595,7 @@ void FPixelShaderInOut_MainPS(
 
 
 	bool bStrataSubsurfaceEnable = false;
-#if STRATA_ENABLED && !MATERIAL_SHADINGMODEL_UNLIT
+#if STRATA_ENABLED && !STRATA_OPTIMIZED_UNLIT
 
 #if STRATA_INLINE_SHADING
 	// We must normalize each normal and tangent to avoid non normalised vectors due to per vertex interpolation or texture filtering,
@@ -1655,7 +1661,7 @@ void FPixelShaderInOut_MainPS(
 	Color += EmissiveLuminance;
 #endif // STRATA_OPAQUE_DEFERRED
 
-#if (STRATA_TRANSLUCENT_FORWARD || STRATA_FORWARD_SHADING) && !STRATA_INLINE_SINGLELAYERWATER
+#if (STRATA_TRANSLUCENT_FORWARD || STRATA_FORWARD_SHADING) && !STRATA_INLINE_SINGLELAYERWATER && !STRATA_OPTIMIZED_UNLIT
 	//FORWARD_SHADING
 	if (StrataData.LayerCount > 0)
 	{
@@ -1679,16 +1685,15 @@ void FPixelShaderInOut_MainPS(
 			L,
 			StrataPixelHeader,
 			StrataData,
-			Throughput);
+			DualBlendSurfaceTransmittancePreCoverage,
+			DualBlendSurfaceCoverage);
 
 	#if STRATA_TRANSLUCENT_MATERIAL
-		DualBlendColorCov = GetMaterialOpacity(PixelMaterialInputs); // Opacity represent coverage STRATA_TODO after validation, rename everything
-		DualBlendColorAdd = Color;
-		DualBlendColorMul = Throughput;
+		DualBlendSurfaceLuminancePostCoverage = Color;
 		Color = 0.0f;
 		Opacity = 1.0f; // nullify following operation
 	#elif STRATA_FORWARD_SHADING
-		// Color unchanged
+		// Color unchanged for opaque materials
 		Opacity = 1.0f;
 	#else
 		#error Unhandled Strata forward shading mode
@@ -1696,18 +1701,20 @@ void FPixelShaderInOut_MainPS(
 	}
 #endif // STRATA_TRANSLUCENT_FORWARD || STRATA_FORWARD_SHADING
 
-#elif STRATA_ENABLED && MATERIAL_SHADINGMODEL_UNLIT
+#elif STRATA_ENABLED && STRATA_OPTIMIZED_UNLIT
 
 	// Unlit forces a single BSDF
 	FStrataBSDF UnlitBSDF = PixelMaterialInputs.FrontMaterial.Layers[0].BSDFs[0];
-#if STRATA_OPAQUE_DEFERRED || STRATA_FORWARD_SHADING
+
+#if STRATA_OPAQUE_DEFERRED
 	Color = BSDF_GETEMISSIVE(UnlitBSDF);
-#else // STRATA_TRANSLUCENT_FORWARD
-	DualBlendColorCov = GetMaterialOpacity(PixelMaterialInputs); // Opacity represent coverage STRATA_TODO after validation, rename everything
-	DualBlendColorAdd = BSDF_GETEMISSIVE(UnlitBSDF);
-	DualBlendColorMul = UNLIT_TRANSMITTANCE(UnlitBSDF);
+	Opacity = 1.0f;
+#else // STRATA_TRANSLUCENT_FORWARD || STRATA_FORWARD_SHADING
+	DualBlendSurfaceLuminancePostCoverage		= BSDF_GETEMISSIVE(UnlitBSDF);
+	DualBlendSurfaceTransmittancePreCoverage	= UNLIT_TRANSMITTANCE(UnlitBSDF);
+	DualBlendSurfaceCoverage					= UnlitBSDF.Coverage;
 	Color = 0.0f;
-	Opacity = 1.0f; // nullify following operation
+	Opacity = 1.0f;
 #endif
 	
 #endif
@@ -1721,18 +1728,11 @@ void FPixelShaderInOut_MainPS(
 		Out.MRT[0] = RETURN_COLOR(Out.MRT[0]);
 	// MATERIAL_SHADINGMODEL_THIN_TRANSLUCENT must come first because it also has MATERIALBLENDING_TRANSLUCENT defined
 	#elif MATERIAL_WORKS_WITH_DUAL_SOURCE_COLOR_BLENDING
-		// After thin translucency, the final color is going to be:
-		//    FinalColor = DualBlendColorAdd + DualBlendColorMul * BackgroundColor;
-	    // To apply fogging, we want the final equation to be:
-		//    FinalColor = Fogging.rgb + Fogging.a * (DualBlendColorAdd + DualBlendColorMul * BackgroundColor);
-		//    FinalColor = (Fogging.rgb + Fogging.a * DualBlendColorAdd) + Fogging.a * DualBlendColorMul * BackgroundColor;
-		// Or in other words:
-		//    AdjustedDualBlendAdd = Fogging.rgb + Fogging.a * DualBlendColorAdd;
-		//    AdjustedDualBlendMul = Fogging.a * DualBlendColorMul;
-		//    FinalColor = AdjustedDualBlendAdd + AdjustedDualBlendMul * BackgroundColor;
-
-		float3 AdjustedDualBlendAdd = DualBlendColorCov * (Fogging.rgb + Fogging.a * DualBlendColorAdd);
-		float3 AdjustedDualBlendMul = lerp(1.0f, Fogging.a * DualBlendColorMul, DualBlendColorCov);
+
+		// Add fog luminance according to surfacecoverage and reduce surface luminance according to fog coverage.
+		float3 AdjustedDualBlendAdd = DualBlendSurfaceCoverage * Fogging.rgb + Fogging.a * DualBlendSurfaceLuminancePostCoverage;
+		// Fade the surface color transmittance out to 1 according to the surface coverage, and take into account the fog coverage to the surface.
+		float3 AdjustedDualBlendMul = lerp(1.0f, Fogging.a * DualBlendSurfaceTransmittancePreCoverage, DualBlendSurfaceCoverage);
 
 		#if DUAL_SOURCE_COLOR_BLENDING_ENABLED
 			// no RETURN_COLOR because these values are explicit multiplies and adds
diff --git a/Engine/Shaders/Private/Strata/StrataForwardLighting.ush b/Engine/Shaders/Private/Strata/StrataForwardLighting.ush
--- a/Engine/Shaders/Private/Strata/StrataForwardLighting.ush
+++ b/Engine/Shaders/Private/Strata/StrataForwardLighting.ush
@@ -162,7 +162,8 @@ float3 StrataForwardLighting(
 	float3 L,
 	FStrataPixelHeader StrataPixelHeader,
 	FStrataData StrataData,
-	inout float3 OutThroughput
+	inout float3 OutTransmittancePreCoverage,
+	inout float OutCoverage
 	)
 {
 	const uint PrimitiveId = MaterialParameters.PrimitiveId;
@@ -179,7 +180,9 @@ float3 StrataForwardLighting(
 	const float FullThroughput = 1.0f;
 
 	float3 Color = 0;
-	OutThroughput = 1.0f;
+	OutCoverage = 0.0f;
+	float3 TransmittancePostCoverage = 1.0f;
+	OutTransmittancePreCoverage = 1.0f;
 
 	float4 DynamicShadowFactors = 1;
 #if MATERIALBLENDING_SOLID || MATERIALBLENDING_MASKED
@@ -281,7 +284,7 @@ float3 StrataForwardLighting(
 			}
 
 			const float3 BSDFCoverage = BSDF.Coverage;
-			const float3 BSDFColoredVisibility = OutThroughput * BSDFCoverage;
+			const float3 BSDFColoredVisibility = TransmittancePostCoverage * BSDFCoverage;
 
 			// Create the BSDF context
 			FStrataBSDFContext StrataBSDFContext = StrataCreateBSDFContext(StrataPixelHeader, BSDF, NullStrataAddressing, V, L);
@@ -619,12 +622,35 @@ float3 StrataForwardLighting(
 			LayerCoverage	+= BSDFCoverage;
 		}
 
+
+
+		// Now we are going to compute TransmittancePreCoverage and Coverage by combining the top layer0 and bottom layer1.
+		const float3 Layer0Meant = OutTransmittancePreCoverage;
+		const float3 Layer1Meant = LayerThroughput / max(1e-5, LayerCoverage);
+
+		const float Layer0Coverage = OutCoverage;
+		const float Layer1Coverage = LayerCoverage;
+		const float LayerMaxCoverage = max(Layer0Coverage, Layer1Coverage);
+		const float LayerMaxCoverageInv = LayerMaxCoverage <= 0.0f ? 1.0f : 1.0f / LayerMaxCoverage;
+
+		// If a layer has a lower coverage than the other, we account for that when compuing the TransmittancePreCoverage.
+		const float TransmittanceOne = 1.0f;
+		const float3 LayertT0 = lerp(1.0, Layer0Meant, Layer0Coverage * LayerMaxCoverageInv); // Layer0Meant * (Layer0Coverage * LayerMaxCoverageInv) + TransmittanceOne * saturate(1.0f - Layer0Coverage * LayerMaxCoverageInv);
+		const float3 LayertT1 = lerp(1.0, Layer1Meant, Layer1Coverage * LayerMaxCoverageInv); // Layer1Meant * (Layer1Coverage * LayerMaxCoverageInv) + TransmittanceOne * saturate(1.0f - Layer1Coverage * LayerMaxCoverageInv);
+
+		// Now we update the material transmittance pre coverage
+		OutTransmittancePreCoverage = LayertT0 * LayertT1;
+		OutCoverage = LayerMaxCoverage;
+
+
+
 		// Take into account the non covered part for this layer and update throughput accordingly
 		LayerThroughput += (1.0 - saturate(LayerCoverage)) * FullThroughput;
 
 		// And combine this layer throughtput with the full material throughput
-		OutThroughput *= LayerThroughput;
+		TransmittancePostCoverage *= LayerThroughput;
 	}
+
 	return Color;
 }
 
diff --git a/Engine/Shaders/Private/Strata/StrataLegacyConversion.ush b/Engine/Shaders/Private/Strata/StrataLegacyConversion.ush
--- a/Engine/Shaders/Private/Strata/StrataLegacyConversion.ush
+++ b/Engine/Shaders/Private/Strata/StrataLegacyConversion.ush
@@ -425,7 +425,8 @@ FStrataData StrataConvertLegacyMaterialStatic(
 
 	if (ShadingModel == SHADINGMODELID_UNLIT)
 	{
-		Out = GetStrataUnlitBSDF(Emissive, Ones /*TransnmittanceColor*/);
+		Out = GetStrataUnlitBSDF(Emissive, Zeros);
+		Out = StrataWeight(Out, Opacity);
 	}
 	else if (ShadingModel == SHADINGMODELID_HAIR)
 	{
diff --git a/Engine/Shaders/Private/ThinTranslucentCommon.ush b/Engine/Shaders/Private/ThinTranslucentCommon.ush
--- a/Engine/Shaders/Private/ThinTranslucentCommon.ush
+++ b/Engine/Shaders/Private/ThinTranslucentCommon.ush
@@ -9,8 +9,9 @@
 #define THIN_TRANSLUCENT_ACCURATE_FRESNEL 1
 
 #if MATERIAL_SHADINGMODEL_THIN_TRANSLUCENT
-void AccumulateThinTranslucentModel(inout float3 DualBlendColorAdd,
-								inout float3 DualBlendColorMul,
+void AccumulateThinTranslucentModel(inout float3 DualBlendSurfaceLuminancePostCoverage,
+								inout float3 DualBlendSurfaceTransmittancePreCoverage,
+								inout float DualBlendSurfaceCoverage,
 								FMaterialPixelParameters MaterialParams,
 								FGBufferData GBuffer,
 								float3 DiffuseColor,
@@ -70,8 +71,10 @@ void AccumulateThinTranslucentModel(inout float3 DualBlendColorAdd,
 
 	SurfaceColor += SpecularColor;
 
-	DualBlendColorAdd = SurfaceColor;
-	DualBlendColorMul = Transmittance;
+	// Luminance and transmitance assumin a full coverage of 1.
+	DualBlendSurfaceCoverage = 1.0f;
+	DualBlendSurfaceLuminancePostCoverage = SurfaceColor * DualBlendSurfaceCoverage;
+	DualBlendSurfaceTransmittancePreCoverage = Transmittance;
 }
 #endif
 
