Normalize opacity clamping logic between PCF and VSM path. The former was basically treating relevant materials as having a max of 0.999 (changed now to 0.99 to make a few other cases robust) opacity always and thus never occluded backfaces on two-sided materials.
Fix warning triggering incorrectly when enabling/disabling VSMs.
Make contact shadow stepoffset the same between VSM and engine contact shadows paths.

diff --git a/Engine/Shaders/Private/DeferredShadingCommon.ush b/Engine/Shaders/Private/DeferredShadingCommon.ush
--- a/Engine/Shaders/Private/DeferredShadingCommon.ush
+++ b/Engine/Shaders/Private/DeferredShadingCommon.ush
@@ -200,7 +200,7 @@ float3 EncodeSubsurfaceProfile(float SubsurfaceProfile)
 // Derive density from a heuristic using opacity, tweaked for useful falloff ranges and to give a linear depth falloff with opacity
 float SubsurfaceDensityFromOpacity(float Opacity)
 {
-	return (-0.05f * log(1.0f - min(Opacity, 0.999f)));
+	return (-0.05f * log(1.0f - min(Opacity, 0.99f)));
 }
 
 float EncodeIndirectIrradiance(float IndirectIrradiance)
diff --git a/Engine/Shaders/Private/VirtualShadowMaps/TransmissionCommon.ush b/Engine/Shaders/Private/VirtualShadowMaps/TransmissionCommon.ush
--- a/Engine/Shaders/Private/VirtualShadowMaps/TransmissionCommon.ush
+++ b/Engine/Shaders/Private/VirtualShadowMaps/TransmissionCommon.ush
@@ -44,7 +44,12 @@ half GetSubsurfaceOpacityFromGbuffer(uint2 PixelPos)
 		GBufferData.ShadingModelID == SHADINGMODELID_PREINTEGRATED_SKIN ||
 		GBufferData.ShadingModelID == SHADINGMODELID_TWOSIDED_FOLIAGE)
 	{
-		Opacity = GBufferData.CustomData.a;
+		// This clamp aligns with SubsurfaceDensityFromOpacity
+		// Various engine paths treat these subsurface materials differently
+		// even when they have Opacity = 1 in the material shader, so this is
+		// important to avoid things like backface transmission being shadowed by
+		// contact shadows and so on.
+		Opacity = min(GBufferData.CustomData.a, 0.99f);
 	}
 #endif // STRATA_ENABLED
 
diff --git a/Engine/Shaders/Private/VirtualShadowMaps/VirtualShadowMapProjection.usf b/Engine/Shaders/Private/VirtualShadowMaps/VirtualShadowMapProjection.usf
--- a/Engine/Shaders/Private/VirtualShadowMaps/VirtualShadowMapProjection.usf
+++ b/Engine/Shaders/Private/VirtualShadowMaps/VirtualShadowMapProjection.usf
@@ -68,9 +68,11 @@ MAX_OCCUPENCY
 // Returns hit distance or negative if miss
 float VirtualShadowMapScreenRayCast(
 	float3 RayOriginTranslatedWorld, float3 RayDirection, float RayLength,
-	int NumSteps, float StepOffset)
+	int NumSteps, float Dither)
 {
-	bool bHitCastContactShadow = true;
+	float StepOffset = Dither - 0.5f;
+
+	bool bHitCastContactShadow = false;
 	float HitDistance = ShadowRayCast(
 		RayOriginTranslatedWorld, RayDirection, RayLength,
 		NumSteps, StepOffset, bHitCastContactShadow);
@@ -252,7 +254,7 @@ FVirtualShadowMapSampleResult ProjectLight(
 			if (ContactShadowLengthWorld > 0.0f)
 			{
 				half ContactShadowHitDistance = VirtualShadowMapScreenRayCast(TranslatedWorldPosition, L, ContactShadowLengthWorld, CONTACT_SHADOW_SAMPLES, Noise);
-				if (ContactShadowHitDistance >= 0.0f)
+				if (ContactShadowHitDistance > 0.0f)
 				{
 					Result.ShadowFactor = 0.0f;
 					Result.OccluderDistance = ContactShadowHitDistance;
diff --git a/Engine/Source/Runtime/Renderer/Private/VirtualShadowMaps/VirtualShadowMapCacheManager.cpp b/Engine/Source/Runtime/Renderer/Private/VirtualShadowMaps/VirtualShadowMapCacheManager.cpp
--- a/Engine/Source/Runtime/Renderer/Private/VirtualShadowMaps/VirtualShadowMapCacheManager.cpp
+++ b/Engine/Source/Runtime/Renderer/Private/VirtualShadowMaps/VirtualShadowMapCacheManager.cpp
@@ -84,11 +84,6 @@ void FVirtualShadowMapCacheEntry::UpdateClipmap(
 	}
 
 	bool bRadiusMatches = (ViewRadiusZ == Clipmap.ViewRadiusZ);
-	if (!bRadiusMatches)
-	{
-		// These really should be exact by construction currently
-		UE_LOG(LogRenderer, Warning, TEXT("Invalidated clipmap level (VSM %d) due to Z radius mismatch"), VirtualShadowMapId);
-	}
 
 	if (bCacheValid && bRadiusMatches)
 	{
@@ -96,6 +91,12 @@ void FVirtualShadowMapCacheEntry::UpdateClipmap(
 	}
 	else
 	{
+		if (bCacheValid && !bRadiusMatches)
+		{
+			// These really should be exact by construction currently
+			UE_LOG(LogRenderer, Warning, TEXT("Invalidated clipmap level (VSM %d) due to Z radius mismatch"), VirtualShadowMapId);
+		}
+
 		// New cached level
 		PrevVirtualShadowMapId = INDEX_NONE;
 		Clipmap.WorldToLight = WorldToLight;
