Fix inverse squared falloff being applied onto directional light.

diff --git a/Engine/Shaders/Private/DeferredLightPixelShaders.usf b/Engine/Shaders/Private/DeferredLightPixelShaders.usf
--- a/Engine/Shaders/Private/DeferredLightPixelShaders.usf
+++ b/Engine/Shaders/Private/DeferredLightPixelShaders.usf
@@ -103,7 +103,7 @@ FDeferredLightData SetupLightDataForStandardDeferred()
 	LightData.DistanceFadeMAD = DeferredLightUniforms.DistanceFadeMAD;
 	LightData.ShadowMapChannelMask = DeferredLightUniforms.ShadowMapChannelMask;
 	LightData.ShadowedBits = DeferredLightUniforms.ShadowedBits;
-	LightData.bInverseSquared = DeferredLightUniforms.FalloffExponent == 0;
+	LightData.bInverseSquared = LIGHT_SOURCE_SHAPE > 0 && DeferredLightUniforms.FalloffExponent == 0; // Directional lights don't use 'inverse squared attenuation'
 	LightData.bRadialLight = LIGHT_SOURCE_SHAPE > 0;
 	//@todo - permutation opportunity
 	LightData.bSpotLight = LIGHT_SOURCE_SHAPE > 0;
