Implements TSR's velocity buffer interpolation

diff --git a/Engine/Shaders/Private/TemporalSuperResolution/TSRUpdateHistory.usf b/Engine/Shaders/Private/TemporalSuperResolution/TSRUpdateHistory.usf
--- a/Engine/Shaders/Private/TemporalSuperResolution/TSRUpdateHistory.usf
+++ b/Engine/Shaders/Private/TemporalSuperResolution/TSRUpdateHistory.usf
@@ -3,12 +3,15 @@
 #include "TSRCommon.ush"
 
 
-//------------------------------------------------------- CONFIG
+//------------------------------------------------------- DEBUG
 
 #define DEBUG_FALLBACK_BLENDING 0
 
 #define DEBUG_LOW_FREQUENCY_REJECTION 0
 
+
+//------------------------------------------------------- CONFIG
+
 #define CONFIG_CLAMP 1
 
 #define CONFIG_SAMPLES 5
@@ -21,6 +24,8 @@
 
 #define CONFIG_REJECTION_ANTI_ALIASING (DIM_REJECTION_ANTI_ALIASING)
 
+#define CONFIG_INTERPOLATE_VELOCITY 1
+
 #if DIM_SEPARATE_TRANSLUCENCY == 1
 	#define CONFIG_ACCUMULATE_TRANSLUCENCY_SEPARATLY 1
 	#define CONFIG_VFX_SURVIVAL 0
@@ -168,6 +173,7 @@ void ComputeInputKernelSamplePosition(
 	}
 	#elif CONFIG_SAMPLES == 5 || CONFIG_SAMPLES == 6
 	{
+		#if CONFIG_SAMPLES == 6
 		if (SampleId == 5)
 		{
 			taa_short2x2 iPixelOffset;
@@ -186,6 +192,7 @@ void ComputeInputKernelSamplePosition(
 			SampleInputPixelPos = ClampPixelOffset(InputPixelPos + iPixelOffset, InputPixelPosMin, InputPixelPosMax);
 		}
 		else
+		#endif
 		{
 			taa_short2 iPixelOffset = taa_short2(kOffsets3x3[kPlusIndexes3x3[SampleId]]);
 			PixelOffset = dpv_interleave_registers(taa_half2(iPixelOffset), taa_half2(iPixelOffset));
@@ -196,6 +203,11 @@ void ComputeInputKernelSamplePosition(
 				InputPixelPosMin, InputPixelPosMax);
 		}
 	}
+	#elif CONFIG_SAMPLES == 2
+	{
+		SampleInputPixelPos = InputPixelPos;
+		PixelOffset = 0;
+	}
 	#else
 		#error Unknown sample count
 	#endif
@@ -250,7 +262,12 @@ void MainCS(
 
 	PLATFORM_SPECIFIC_ISOLATE
 	{
-		float2 RawEncodedVelocity[DPV_PIXEL_PER_LANE];
+		#if CONFIG_INTERPOLATE_VELOCITY
+			float2 RawEncodedVelocityNeighborhood[4][DPV_PIXEL_PER_LANE];
+		#else
+			float2 RawEncodedVelocity[DPV_PIXEL_PER_LANE];
+		#endif
+
 		taa_half RawParallaxRejectionMask[DPV_PIXEL_PER_LANE];
 		uint RawParallaxFactorBits[DPV_PIXEL_PER_LANE];
 		taa_half RawLowFrequencyRejection[DPV_PIXEL_PER_LANE];
@@ -276,8 +293,10 @@ void MainCS(
 			for (uint PixelId = 0; PixelId < DPV_PIXEL_PER_LANE; PixelId++)
 			{
 				taa_ushort2 LocalInputPixelPos = dpv_access_pixel(InputPixelPos, PixelId);
-
-				RawEncodedVelocity[PixelId]       = DilatedVelocityTexture[LocalInputPixelPos];
+				
+				#if !CONFIG_INTERPOLATE_VELOCITY
+					RawEncodedVelocity[PixelId]   = DilatedVelocityTexture[LocalInputPixelPos];
+				#endif
 				RawParallaxRejectionMask[PixelId] = ParallaxRejectionMaskTexture[LocalInputPixelPos];
 				RawParallaxFactorBits[PixelId]    = ParallaxFactorTexture[LocalInputPixelPos];
 				RawTranslucencyRejection[PixelId] = TranslucencyRejectionTexture[LocalInputPixelPos];
@@ -318,6 +337,22 @@ void MainCS(
 					#endif
 				}
 			} // for (uint NeighborId = 0; NeighborId < (CONFIG_SAMPLES - 1); NeighborId++)
+			
+			// Fetch the 2x2 neighborhood of the velocity
+			#if CONFIG_INTERPOLATE_VELOCITY
+			UNROLL_N(DPV_PIXEL_PER_LANE)
+			for (uint PixelId = 0; PixelId < DPV_PIXEL_PER_LANE; PixelId++)
+			{
+				taa_short2 LocalInputPixelPos0 = dpv_access_pixel(InputPixelPos, PixelId);
+				taa_short2 LocalInputPixelPos1 = LocalInputPixelPos0 + taa_short2(sign(dpv_access_pixel(dKO, PixelId)));
+				LocalInputPixelPos1 = ClampPixelOffset(LocalInputPixelPos1, InputPixelPosMin, InputPixelPosMax);
+
+				RawEncodedVelocityNeighborhood[0][PixelId] = DilatedVelocityTexture[taa_short2(LocalInputPixelPos0.x, LocalInputPixelPos0.y)];
+				RawEncodedVelocityNeighborhood[1][PixelId] = DilatedVelocityTexture[taa_short2(LocalInputPixelPos1.x, LocalInputPixelPos0.y)];
+				RawEncodedVelocityNeighborhood[2][PixelId] = DilatedVelocityTexture[taa_short2(LocalInputPixelPos0.x, LocalInputPixelPos1.y)];
+				RawEncodedVelocityNeighborhood[3][PixelId] = DilatedVelocityTexture[taa_short2(LocalInputPixelPos1.x, LocalInputPixelPos1.y)];
+			}
+			#endif
 		}
 		
 		uint2 SceneStencilRef;
@@ -325,11 +360,47 @@ void MainCS(
 		// Process texture fetches.
 		PLATFORM_SPECIFIC_ISOLATE
 		{
+			// Vector in pixel between pixel K -> O.
+			taa_half2x2 dKO = taa_half2x2(PPCo - PPCk);
+		
 			// Process velocity
 			{
-				float2x2 ScreenVelocity = dpv_interleave_registers(
-					DecodeVelocityFromTexture(float4(RawEncodedVelocity[0], 0.0, 0.0)).xy,
-					DecodeVelocityFromTexture(float4(RawEncodedVelocity[1], 0.0, 0.0)).xy);
+				float2x2 ScreenVelocity;
+				#if CONFIG_INTERPOLATE_VELOCITY
+				UNROLL_N(DPV_PIXEL_PER_LANE)
+				for (uint PixelId = 0; PixelId < DPV_PIXEL_PER_LANE; PixelId++)
+				{
+					float2 V00 = DecodeVelocityFromTexture(float4(RawEncodedVelocityNeighborhood[0][PixelId], 0.0, 0.0)).xy;
+					float2 V10 = DecodeVelocityFromTexture(float4(RawEncodedVelocityNeighborhood[1][PixelId], 0.0, 0.0)).xy;
+					float2 V01 = DecodeVelocityFromTexture(float4(RawEncodedVelocityNeighborhood[2][PixelId], 0.0, 0.0)).xy;
+					float2 V11 = DecodeVelocityFromTexture(float4(RawEncodedVelocityNeighborhood[3][PixelId], 0.0, 0.0)).xy;
+
+					taa_half2 BilinearInterp = abs(dpv_access_pixel(dKO, PixelId));
+					
+					float Bilinear00 = saturate(1.0 - BilinearInterp.x) * saturate(1.0 - BilinearInterp.y);
+					float Bilinear10 = saturate(0.0 + BilinearInterp.x) * saturate(1.0 - BilinearInterp.y);
+					float Bilinear01 = saturate(1.0 - BilinearInterp.x) * saturate(0.0 + BilinearInterp.y);
+					float Bilinear11 = saturate(0.0 + BilinearInterp.x) * saturate(0.0 + BilinearInterp.y);
+					
+					float Bilateral00 = Bilinear00;
+					float Bilateral10 = Bilinear10 * saturate(1.0 - length2((V10 - V00) * InputInfo_ViewportSize * 0.5));
+					float Bilateral01 = Bilinear01 * saturate(1.0 - length2((V01 - V00) * InputInfo_ViewportSize * 0.5));
+					float Bilateral11 = Bilinear11 * saturate(1.0 - length2((V11 - V00) * InputInfo_ViewportSize * 0.5));
+
+					float NormalizeBilateral = rcp(Bilateral00 + Bilateral10 + Bilateral01 + Bilateral11);
+
+					float2 LocalScreenVelocity = (V00 * Bilateral00 + V10 * Bilateral10 + V01 * Bilateral01 + V11 * Bilateral11) * NormalizeBilateral;
+					
+					ScreenVelocity[0][PixelId] = LocalScreenVelocity.x;
+					ScreenVelocity[1][PixelId] = LocalScreenVelocity.y;
+				}
+				#else
+				{
+					ScreenVelocity = dpv_interleave_registers(
+						DecodeVelocityFromTexture(float4(RawEncodedVelocity[0], 0.0, 0.0)).xy,
+						DecodeVelocityFromTexture(float4(RawEncodedVelocity[1], 0.0, 0.0)).xy);
+				}
+				#endif
 
 				PrevScreenPos = ScreenPos - ScreenVelocity;
 				OutputPixelVelocity = taa_half2(dpv_length(dpv_mul(ScreenVelocity, HistoryInfo_ViewportSize)));
@@ -470,18 +541,18 @@ void MainCS(
 			// Soften the spatial kernel when there is high noise in the input to make parrallax disocclusion of detailed geometry less distracting.
 			taa_half2 OffScreenInputToHistoryFactor = taa_half(1.0) - taa_half(0.5) * NoiseFiltering;
 
-			taa_half2 KernelInputToHistoryFactor = OffScreenInputToHistoryFactor;
+			taa_half2 KernelInputToHistoryFactor = taa_half(InputToHistoryFactor).xx;
 			
 			FLATTEN
 			if (bOffScreen.x)
 			{
-				KernelInputToHistoryFactor.x = taa_half(InputToHistoryFactor);
+				KernelInputToHistoryFactor.x = taa_half(OffScreenInputToHistoryFactor.x);
 			}
 			
 			FLATTEN
 			if (bOffScreen.y)
 			{
-				KernelInputToHistoryFactor.y = taa_half(InputToHistoryFactor);
+				KernelInputToHistoryFactor.y = taa_half(OffScreenInputToHistoryFactor.y);
 			}
 
 			UNROLL_N(CONFIG_SAMPLES)
diff --git a/Engine/Source/Runtime/Renderer/Private/PostProcess/TemporalSuperResolution.cpp b/Engine/Source/Runtime/Renderer/Private/PostProcess/TemporalSuperResolution.cpp
--- a/Engine/Source/Runtime/Renderer/Private/PostProcess/TemporalSuperResolution.cpp
+++ b/Engine/Source/Runtime/Renderer/Private/PostProcess/TemporalSuperResolution.cpp
@@ -706,7 +706,7 @@ ITemporalUpscaler::FOutputs AddTemporalSuperResolutionPasses(
 		CommonParameters.InputPixelPosMax = CommonParameters.InputInfo.ViewportMax - 1;
 		CommonParameters.InputPixelPosToScreenPos = (FScreenTransform::Identity + 0.5) * CommonParameters.InputInfo.ViewportSizeInverse * FScreenTransform::ViewportUVToScreenPos;
 		CommonParameters.ScreenVelocityToInputPixelVelocity = (FScreenTransform::Identity / CommonParameters.InputPixelPosToScreenPos).Scale;
-		CommonParameters.InputPixelVelocityToScreenVelocity = CommonParameters.InputPixelPosToScreenPos.Scale;
+		CommonParameters.InputPixelVelocityToScreenVelocity = CommonParameters.InputPixelPosToScreenPos.Scale.GetAbs();
 
 		CommonParameters.LowFrequencyInfo = GetScreenPassTextureViewportParameters(FScreenPassTextureViewport(
 			LowFrequencyExtent, LowFrequencyRect));
