Fix half resolution vertical artifacts at borders for SSS.

diff --git a/Engine/Shaders/Private/PostProcessSubsurfaceTile.usf b/Engine/Shaders/Private/PostProcessSubsurfaceTile.usf
--- a/Engine/Shaders/Private/PostProcessSubsurfaceTile.usf
+++ b/Engine/Shaders/Private/PostProcessSubsurfaceTile.usf
@@ -93,7 +93,7 @@ void MainVS(
 	uint2 TileVertex = TileCoord * SUBSURFACE_TILE_SIZE;
 	TileVertex.x += InVertexId == 1 || InVertexId == 2 || InVertexId == 4 ? SUBSURFACE_TILE_SIZE : 0;
 	TileVertex.y += InVertexId == 2 || InVertexId == 4 || InVertexId == 5 ? SUBSURFACE_TILE_SIZE : 0;
-	Out.UV = (float2(TileVertex + ViewMin) + 0.5f) * ExtentInverse;
+	Out.UV = float2(TileVertex + ViewMin) * ExtentInverse;
 	Out.Position = float4(Out.UV * float2(2.0f, -2.0f) + float2(-1.0, 1.0f), 0.0f, 1.0f);
 }
 
