Fixes r.TSR.Translucency.SeparateTemporalAccumulation with r.TSR.HistoryScreenPercentage

diff --git a/Engine/Shaders/Private/TemporalSuperResolution/TSRResolveHistory.usf b/Engine/Shaders/Private/TemporalSuperResolution/TSRResolveHistory.usf
--- a/Engine/Shaders/Private/TemporalSuperResolution/TSRResolveHistory.usf
+++ b/Engine/Shaders/Private/TemporalSuperResolution/TSRResolveHistory.usf
@@ -141,16 +141,16 @@ void MainCS(
 				
 				taa_half3 SampleTranslucencyColor = History_Translucency.SampleLevel(GlobalBilinearClampedSampler, SampleUV, 0);
 				taa_half  SampleTranslucencyAlpha = History_TranslucencyAlpha.SampleLevel(GlobalBilinearClampedSampler, SampleUV, 0);
+				
+				taa_half SampleHistoryValidity = SampleMetadata.y;
 
 				#if 1
 				{
 					SampleLowFrequencyColor  = SampleLowFrequencyColor  * SampleTranslucencyAlpha + SampleTranslucencyColor;
-					SampleHighFrequencyColor = SampleHighFrequencyColor * SampleTranslucencyAlpha + SampleTranslucencyColor;
+					SampleHighFrequencyColor = SampleHighFrequencyColor * SampleTranslucencyAlpha + SampleTranslucencyColor * SampleHistoryValidity;
 				}
 				#endif
 
-				taa_half SampleHistoryValidity = SampleMetadata.y;
-
 				float2 PixelOffset = abs(SamplePixelPos - OutputPixelPos);
 
 				float kernelWeight = MitchellNetravali(PixelOffset.x, B, C) * MitchellNetravali(PixelOffset.y, B, C);
