Fixes a source of aliasing under movement in TSR

diff --git a/Engine/Shaders/Private/TemporalSuperResolution/TSRDecimateHistory.usf b/Engine/Shaders/Private/TemporalSuperResolution/TSRDecimateHistory.usf
--- a/Engine/Shaders/Private/TemporalSuperResolution/TSRDecimateHistory.usf
+++ b/Engine/Shaders/Private/TemporalSuperResolution/TSRDecimateHistory.usf
@@ -111,7 +111,7 @@ void MainCS(
 
 			taa_half DepthRejection;
 			{
-				const float WorldDepthToPixelWorldRadius = GetTanHalfFieldOfView().x * InputInfo_ExtentInverse.x;
+				const float WorldDepthToPixelWorldRadius = GetTanHalfFieldOfView().x * InputInfo_ViewportSizeInverse.x;
 				const float PixelDepthError = 3.0;
 				const float RadiusToDiamater = 2.0;
 			
@@ -219,7 +219,7 @@ void MainCS(
 		float4 RotationalPrevClip = mul(ThisClip, RotationalClipToPrevClip);
 		float2 RotationalPrevScreen = RotationalPrevClip.xy / RotationalPrevClip.w;
 			
-		taa_half PixelParallaxVelocity = taa_half(0.5) * length(taa_half2(RotationalPrevScreen - PrevScreen) * InputInfo_Extent);
+		taa_half PixelParallaxVelocity = taa_half(0.5) * length(taa_half2(RotationalPrevScreen - PrevScreen) * InputInfo_ViewportSize);
 		taa_half HistorySubdetailAdditionalParralax = (HistorySubdetailRejectionFactor - taa_half(1.0)) * PixelParallaxVelocity;
 
 		// PixelParallaxVelocity can be +inf and generate NaN if HistorySubdetailRejectionFactor = 1.0;
