Works arround TSR running out of UAVs when enabling its DEBUG_OUTPUT

diff --git a/Engine/Shaders/Private/TemporalSuperResolution/TSRDilateVelocity.usf b/Engine/Shaders/Private/TemporalSuperResolution/TSRDilateVelocity.usf
--- a/Engine/Shaders/Private/TemporalSuperResolution/TSRDilateVelocity.usf
+++ b/Engine/Shaders/Private/TemporalSuperResolution/TSRDilateVelocity.usf
@@ -237,7 +237,7 @@ void MainCS(
 		}
 		#endif
 
-		#if DEBUG_OUTPUT
+		#if DEBUG_OUTPUT && 0 // running out of UAVs
 		{
 			DebugOutput[OutputPixelPos] = Debug;
 		}
diff --git a/Engine/Shaders/Private/TemporalSuperResolution/TSRUpdateHistory.usf b/Engine/Shaders/Private/TemporalSuperResolution/TSRUpdateHistory.usf
--- a/Engine/Shaders/Private/TemporalSuperResolution/TSRUpdateHistory.usf
+++ b/Engine/Shaders/Private/TemporalSuperResolution/TSRUpdateHistory.usf
@@ -1417,6 +1417,7 @@ void MainCS(
 		}
 		
 		// Output final scene color Mip1
+		#if !DEBUG_OUTPUT // running out of UAVs
 		{
 			taa_half3x2 HalfResOutput = OutputColor * taa_half(0.25);
 			//taa_short2x2 HalfResOutputPixelPos = (LocalHistoryPixelPos >> taa_short(1)) | (((LocalHistoryPixelPos & taa_short(0x1)) | !bGenerateOutputMip1) * taa_short(~0));
@@ -1467,6 +1468,7 @@ void MainCS(
 				SceneColorOutputMip1[dpv_hi(HalfResOutputPixelPos)] = dpv_hi(HalfResOutput);
 			#endif
 		}
+		#endif
 
 		#if DEBUG_OUTPUT
 		{
