Simplify code for primary ray generation precision and improve precision with LWC

diff --git a/Engine/Shaders/Private/RayTracing/RayTracingCommon.ush b/Engine/Shaders/Private/RayTracing/RayTracingCommon.ush
--- a/Engine/Shaders/Private/RayTracing/RayTracingCommon.ush
+++ b/Engine/Shaders/Private/RayTracing/RayTracingCommon.ush
@@ -367,11 +367,10 @@ RayDesc CreatePrimaryRay(float2 UV)
 	float2 ScreenPosition = (UV - View.ScreenPositionScaleBias.wz) / View.ScreenPositionScaleBias.xy;
 
 	RayDesc Ray;
-	float4 O = mul(float4(ScreenPosition, 1, 1), LWCHackToFloat(PrimaryView.ClipToWorld)); // near (world space)
-	float4 A = mul(float4(ScreenPosition, 1, 1), View.ClipToTranslatedWorld); // near 
-	float4 B = mul(float4(ScreenPosition, 0, 1), View.ClipToTranslatedWorld); // far
-	Ray.Origin = O.xyz / O.w;
-	Ray.Direction = normalize(A.w * B.xyz - B.w * A.xyz);
+	float4 NearPoint = mul(float4(ScreenPosition, 1, 1), View.ClipToTranslatedWorld);
+	float4 FarPoint  = mul(float4(ScreenPosition, 0, 1), View.ClipToTranslatedWorld);
+	Ray.Origin = NearPoint.xyz * rcp(NearPoint.w) + LWCHackToFloat(MakeLWCVector3(View.ViewTilePosition, View.RelativeWorldCameraOrigin));
+	Ray.Direction = normalize(NearPoint.w * FarPoint.xyz - FarPoint.w * NearPoint.xyz);
 
 	// extract far plane (ortho only, so ignore distinction between T and Z)
 	float4 Back = mul(float4(0, 0, 0, 1), View.ClipToView);
