Improve the quality of AlphaToCoverage on mobile.

diff --git a/Engine/Shaders/Private/MobileBasePassPixelShader.usf b/Engine/Shaders/Private/MobileBasePassPixelShader.usf
--- a/Engine/Shaders/Private/MobileBasePassPixelShader.usf
+++ b/Engine/Shaders/Private/MobileBasePassPixelShader.usf
@@ -1068,8 +1068,9 @@ void Main(
 			// Planar reflections and scene captures use scene color alpha to keep track of where content has been rendered, for compositing into a different scene later
 			OutColor.a = 0.0;
 		#else
-			half MaterialOpacityMask = GetMaterialMaskInputRaw(PixelMaterialInputs);
-			OutColor.a = GetMaterialMask(PixelMaterialInputs) / max(abs(ddx(MaterialOpacityMask)) + abs(ddy(MaterialOpacityMask)), 0.0001f) + 0.5f;
+			half MaterialOpacityClip = GetMaterialOpacityMaskClipValue();
+			float Mask = GetMaterialMask(PixelMaterialInputs) / (1.0 - MaterialOpacityClip);
+			OutColor.a = uint((Mask + 0.25f) * 4.0f) / 4.0f;
 		#endif
 	#endif
 
