Implements r.TSR.Velocity.Extrapolation

diff --git a/Engine/Shaders/Private/TemporalSuperResolution/TSRCommon.ush b/Engine/Shaders/Private/TemporalSuperResolution/TSRCommon.ush
--- a/Engine/Shaders/Private/TemporalSuperResolution/TSRCommon.ush
+++ b/Engine/Shaders/Private/TemporalSuperResolution/TSRCommon.ush
@@ -1034,3 +1034,15 @@ taa_short2 GetLaneOffsetSign()
 	return taa_short(1).xx;
 }
 #endif
+
+float ComputePixelVelocityBilateralWeight(float2 PixelVelocityDelta)
+{
+	return saturate(1.0 - length2(PixelVelocityDelta));
+}
+
+float ComputeScreenVelocityBilateralWeight(float2 ScreenVelocityDelta)
+{
+	return ComputePixelVelocityBilateralWeight(ScreenVelocityDelta * ScreenVelocityToInputPixelVelocity);
+}
+
+
diff --git a/Engine/Shaders/Private/TemporalSuperResolution/TSRDilateVelocity.usf b/Engine/Shaders/Private/TemporalSuperResolution/TSRDilateVelocity.usf
--- a/Engine/Shaders/Private/TemporalSuperResolution/TSRDilateVelocity.usf
+++ b/Engine/Shaders/Private/TemporalSuperResolution/TSRDilateVelocity.usf
@@ -18,16 +18,21 @@
 
 #define AA_CROSS 1
 
+#define CONFIG_USE_SCENE_VELOCITY 1
+
 #define CONFIG_DEPTH_EDGE_DETECTION 1
 
+#define CONFIG_VELOCITY_EXTRAPOLATION 1
+
 
 //------------------------------------------------------- PARAMETERS
 
 float2 PrevOutputBufferUVMin;
 float2 PrevOutputBufferUVMax;
 float WorldDepthToDepthError;
+float VelocityExtrapolationMultiplier;
 
-Texture2D SceneDepthTexture;
+Texture2D<float> SceneDepthTexture;
 Texture2D SceneVelocityTexture;
 
 RWTexture2D<float2> DilatedVelocityOutput;
@@ -75,7 +80,7 @@ void MainCS(
 	float2 ScreenPos = ViewportUVToScreenPos(ViewportUV);
 	float2 InputBufferUV = (InputPixelPos + 0.5) * InputInfo_ExtentInverse;
 
-	float DeviceZ = SceneDepthTexture[InputPixelPos].r;
+	float DeviceZ = SceneDepthTexture[InputPixelPos];
 
 	float ClosestDeviceZ = DeviceZ;
 	float2 VelocityPixelOffset = float2(0.0, 0.0);
@@ -112,8 +117,8 @@ void MainCS(
 		{
 			const int2 Offset = kOffsets[i];
 
-			float SampleDeviceZ0 = SceneDepthTexture[int2(InputPixelPos) + Offset].r;
-			float SampleDeviceZ1 = SceneDepthTexture[int2(InputPixelPos) - Offset].r;
+			float SampleDeviceZ0 = SceneDepthTexture[int2(InputPixelPos) + Offset];
+			float SampleDeviceZ1 = SceneDepthTexture[int2(InputPixelPos) - Offset];
 			
 			#if CONFIG_DEPTH_EDGE_DETECTION
 			{
@@ -149,13 +154,51 @@ void MainCS(
 	float2 FinalScreenVelocity = 0.0;
 
 	// Compute final screen vecocity.
-	#if 1
+	#if CONFIG_VELOCITY_EXTRAPOLATION
+	BRANCH
+	if (VelocityExtrapolationMultiplier > 0.0)
+	{
+		int2 PixelCoord1 = int2(InputBufferUV * InputInfo_Extent + VelocityPixelOffset * 1.0);
+		int2 PixelCoord2 = int2(InputBufferUV * InputInfo_Extent + VelocityPixelOffset * 2.0);
+
+		PixelCoord2 = clamp(PixelCoord2, InputPixelPosMin, InputPixelPosMax);
+
+		float ClosestDeviceZ1 = ClosestDeviceZ;
+		float ClosestDeviceZ2 = SceneDepthTexture[PixelCoord2];
+
+		float4 EncodedVelocity1 = SceneVelocityTexture[PixelCoord1];
+		float4 EncodedVelocity2 = SceneVelocityTexture[PixelCoord2];
+
+		float2 ScreenVelocity1 = ComputeStaticVelocity(ScreenPos, ClosestDeviceZ1);
+		float2 ScreenVelocity2 = ComputeStaticVelocity(ScreenPos, ClosestDeviceZ2);
+		
+		#if CONFIG_USE_SCENE_VELOCITY
+		{
+			if (EncodedVelocity1.x > 0.0)
+			{
+				ScreenVelocity1 = DecodeVelocityFromTexture(EncodedVelocity1).xy;
+			}
+			if (EncodedVelocity2.x > 0.0)
+			{
+				ScreenVelocity2 = DecodeVelocityFromTexture(EncodedVelocity2).xy;
+			}
+		}
+		#endif
+
+		float2 VelocityDelta = ScreenVelocity1 - ScreenVelocity2;
+		float ExtrapolationBilateral = ComputeScreenVelocityBilateralWeight(VelocityDelta);
+
+		FinalScreenVelocity = ScreenVelocity1 + VelocityDelta * ExtrapolationBilateral * VelocityExtrapolationMultiplier;
+	}
+	else
+	#endif
 	{
 		FinalScreenVelocity = ComputeStaticVelocity(ScreenPos, ClosestDeviceZ);
+		
+		float4 EncodedVelocity = SceneVelocityTexture[uint2(InputBufferUV * InputInfo_Extent + VelocityPixelOffset)];
 
-		#if 1
+		#if CONFIG_USE_SCENE_VELOCITY
 		{
-			float4 EncodedVelocity = SceneVelocityTexture[uint2(InputBufferUV * InputInfo_Extent + VelocityPixelOffset)];
 			if (EncodedVelocity.x > 0.0)
 			{
 				FinalScreenVelocity = DecodeVelocityFromTexture(EncodedVelocity).xy;
@@ -163,7 +206,6 @@ void MainCS(
 		}
 		#endif
 	}
-	#endif
 
 	// Do motion blur velocity flatten.
 	#if DIM_MOTION_BLUR_DIRECTIONS > 0
diff --git a/Engine/Shaders/Private/TemporalSuperResolution/TSRUpdateHistory.usf b/Engine/Shaders/Private/TemporalSuperResolution/TSRUpdateHistory.usf
--- a/Engine/Shaders/Private/TemporalSuperResolution/TSRUpdateHistory.usf
+++ b/Engine/Shaders/Private/TemporalSuperResolution/TSRUpdateHistory.usf
@@ -371,9 +371,9 @@ void MainCS(
 					float Bilinear11 = saturate(0.0 + BilinearInterp.x) * saturate(0.0 + BilinearInterp.y);
 					
 					float Bilateral00 = Bilinear00;
-					float Bilateral10 = Bilinear10 * saturate(1.0 - length2((V10 - V00) * InputInfo_ViewportSize * 0.5));
-					float Bilateral01 = Bilinear01 * saturate(1.0 - length2((V01 - V00) * InputInfo_ViewportSize * 0.5));
-					float Bilateral11 = Bilinear11 * saturate(1.0 - length2((V11 - V00) * InputInfo_ViewportSize * 0.5));
+					float Bilateral10 = Bilinear10 * ComputeScreenVelocityBilateralWeight(V10 - V00);
+					float Bilateral01 = Bilinear01 * ComputeScreenVelocityBilateralWeight(V01 - V00);
+					float Bilateral11 = Bilinear11 * ComputeScreenVelocityBilateralWeight(V11 - V00);
 
 					float NormalizeBilateral = rcp(Bilateral00 + Bilateral10 + Bilateral01 + Bilateral11);
 
diff --git a/Engine/Source/Runtime/Renderer/Private/PostProcess/TemporalSuperResolution.cpp b/Engine/Source/Runtime/Renderer/Private/PostProcess/TemporalSuperResolution.cpp
--- a/Engine/Source/Runtime/Renderer/Private/PostProcess/TemporalSuperResolution.cpp
+++ b/Engine/Source/Runtime/Renderer/Private/PostProcess/TemporalSuperResolution.cpp
@@ -70,6 +70,11 @@ TAutoConsoleVariable<float> CVarTSRWeightClampingPixelSpeed(
 	TEXT("Smallest reduce blur in movement (Default = 1.0f)."),
 	ECVF_RenderThreadSafe);
 
+TAutoConsoleVariable<float> CVarTSRVelocityExtrapolation(
+	TEXT("r.TSR.Velocity.Extrapolation"), 1.0f,
+	TEXT("Defines how much the velocity should be extrapolated on geometric discontinuities (Default = 1.0f)."),
+	ECVF_RenderThreadSafe);
+
 TAutoConsoleVariable<int32> CVarTSRVelocityHoleFill(
 	TEXT("r.TSR.Velocity.HoleFill"), 1,
 	TEXT("Whether to holl-fill the velocity buffer on parralax disocclusion to reduce boiling of disoccluded areas."),
@@ -195,6 +200,7 @@ class FTSRDilateVelocityCS : public FTSRShader
 		SHADER_PARAMETER(FVector2f, PrevOutputBufferUVMin)
 		SHADER_PARAMETER(FVector2f, PrevOutputBufferUVMax)
 		SHADER_PARAMETER(float, WorldDepthToDepthError)
+		SHADER_PARAMETER(float, VelocityExtrapolationMultiplier)
 
 		SHADER_PARAMETER_RDG_TEXTURE(Texture2D, SceneDepthTexture)
 		SHADER_PARAMETER_RDG_TEXTURE(Texture2D, SceneVelocityTexture)
@@ -753,6 +759,8 @@ ITemporalUpscaler::FOutputs AddTemporalSuperResolutionPasses(
 
 			PassParameters->WorldDepthToDepthError = WorldDepthToPixelWorldRadius * 2.0f;
 		}
+		PassParameters->VelocityExtrapolationMultiplier = FMath::Clamp(CVarTSRVelocityExtrapolation.GetValueOnRenderThread(), 0.0f, 1.0f);
+
 		PassParameters->SceneDepthTexture = PassInputs.SceneDepthTexture;
 		PassParameters->SceneVelocityTexture = PassInputs.SceneVelocityTexture;
 
