Fix for a pltform to compile when strata is off.
We always compile the tile vertex shader and this includes  strata.ush. We will fix compilation issue there when we get to it.

diff --git a/Engine/Shaders/Private/Strata/StrataMaterialClassification.usf b/Engine/Shaders/Private/Strata/StrataMaterialClassification.usf
--- a/Engine/Shaders/Private/Strata/StrataMaterialClassification.usf
+++ b/Engine/Shaders/Private/Strata/StrataMaterialClassification.usf
@@ -6,21 +6,17 @@
 #define STRATA_SSS_MATERIAL_OVERRIDE 0
 #include "/Engine/Private/Strata/Strata.ush"
 
-#define TILE_SIZE 8
-#define GROUP_THREAD_COUNT (TILE_SIZE*TILE_SIZE)
+#define STRATA_TILE_SIZE 8
+#define GROUP_THREAD_COUNT (STRATA_TILE_SIZE * STRATA_TILE_SIZE)
 
 ////////////////////////////////////////////////////////////////////////////////////////////////////////////
 
+// DecodeTile is in StrataTiling.usf.
 uint EncodeTile(uint2 TileCoord)
 {
 	return TileCoord.x | (TileCoord.y << 16); // assumes 16bit is enough to represent a tiled resolution up to 65,535 :)
 }
 
-uint2 DecodeTile(uint In)
-{
-	return uint2(In & 0xFFFF, In >> 16);
-}
-
 #if SHADER_TILE_CATEGORIZATION
 int TileSize;
 int bRectPrimitive;
@@ -46,7 +42,7 @@ groupshared uint s_TileFlags[GROUP_THREAD_COUNT];
 RWTexture2D<uint2> SSSTextureUAV;
 #endif
 
-[numthreads(TILE_SIZE, TILE_SIZE, 1)]
+[numthreads(STRATA_TILE_SIZE, STRATA_TILE_SIZE, 1)]
 void TileMainCS(uint2 DispatchThreadId : SV_DispatchThreadID, uint LinearIndex : SV_GroupIndex, uint3 GroupId : SV_GroupID)
 {
 	if (all(DispatchThreadId == 0))
@@ -164,60 +160,6 @@ void TileMainCS(uint2 DispatchThreadId : SV_DispatchThreadID, uint LinearIndex :
 
 ////////////////////////////////////////////////////////////////////////////////////////////////////////////
 
-#if SHADER_STENCIL_CATEGORIZATION
-
-float4 OutputViewSizeAndInvSize;
-float4 OutputBufferSizeAndInvSize;
-float4x4 ViewScreenToTranslatedWorld;
-Buffer<uint> TileListBuffer;
-float4 DebugTileColor;
-
-void StrataTilePassVS(
-	in uint InVertexId : SV_VertexID,
-	in uint InInstanceId : SV_InstanceID,
-#if PERMUTATION_ENABLE_TEXCOORD_SCREENVECTOR
-	out float2 OutTexCoord : TEXCOORD0,		// Used by DeferredLightPixelMain directional light
-	out float3 OutScreenVector : TEXCOORD1,	// Used by DeferredLightPixelMain directional light
-#endif
-#if PERMUTATION_ENABLE_DEBUG
-	out uint2 OutTileCoord : TILE_COORD,	// Used by StencilMainPS below
-#endif
-	out float4 OutPosition : SV_POSITION)	// Use by StencilMainPS
-{
-	const uint2 TileCoord = DecodeTile(TileListBuffer[InInstanceId]);
-
-	uint2 TileVertex = TileCoord * TILE_SIZE;
-	TileVertex.x += InVertexId == 1 || InVertexId == 2 || InVertexId == 4 ? TILE_SIZE : 0;
-	TileVertex.y += InVertexId == 2 || InVertexId == 4 || InVertexId == 5 ? TILE_SIZE : 0;
-
-	OutPosition = float4(float2(TileVertex) * OutputViewSizeAndInvSize.zw * float2(2.0f, -2.0f) + float2(-1.0, 1.0f), 0.5f, 1.0f);
-
-#if PERMUTATION_ENABLE_TEXCOORD_SCREENVECTOR
-	OutTexCoord = float2(TileVertex) * OutputBufferSizeAndInvSize.zw;
-	OutScreenVector = mul(float4(OutPosition.xy, 1, 0), ViewScreenToTranslatedWorld).xyz;
-#endif
-
-#if PERMUTATION_ENABLE_DEBUG
-	OutTileCoord = TileCoord;
-#endif
-}
-
-void StencilMainPS(
-	in uint2 InTileCoord : TILE_COORD,
-	in float4 SVPos : SV_POSITION,
-	out float4 OutColor0 : SV_Target0)
-{
-	const bool bTileX = (InTileCoord.x & 1) == 0;
-	const bool bTileY = (InTileCoord.y & 1) == 0;
-	const bool bChecker = (bTileX && bTileY) || (!bTileX && !bTileY);
-	OutColor0 = DebugTileColor * float4(1.0f, 1.0f, 1.0f, bChecker ? 0.33f : 0.66f);
-	OutColor0.rgb *= OutColor0.a; // premultiplied alpha blending assuming alpha is opacity==coverage
-}
-
-#endif //SHADER_STENCIL_CATEGORIZATION
-
-////////////////////////////////////////////////////////////////////////////////////////////////////////////
-
 #if SHADER_CLEAR_MATERIAL_BUFFER
 
 RWTexture2DArray<uint> MaterialTextureArrayUAV;
diff --git a/Engine/Shaders/Private/Strata/StrataTiling.usf b/Engine/Shaders/Private/Strata/StrataTiling.usf
new file mode 100644
--- /dev/null
+++ b/Engine/Shaders/Private/Strata/StrataTiling.usf
@@ -0,0 +1,76 @@
+// Copyright Epic Games, Inc. All Rights Reserved.
+
+#include "/Engine/Private/Common.ush"
+#include "/Engine/Shared/StrataDefinitions.h"   
+
+
+// EncodeTile is in StrataMaterialClassification.usf.
+uint2 DecodeTile(uint In)
+{
+	return uint2(In & 0xFFFF, In >> 16);
+}
+
+
+////////////////////////////////////////////////////////////////////////////////////////////////////////////
+
+#if SHADER_TILE_VS
+
+float4 OutputViewSizeAndInvSize;
+float4 OutputBufferSizeAndInvSize;
+float4x4 ViewScreenToTranslatedWorld;
+
+Buffer<uint> TileListBuffer;
+
+void StrataTilePassVS(
+	in uint InVertexId : SV_VertexID,
+	in uint InInstanceId : SV_InstanceID,
+#if PERMUTATION_ENABLE_TEXCOORD_SCREENVECTOR
+	out float2 OutTexCoord : TEXCOORD0,		// Used by DeferredLightPixelMain directional light
+	out float3 OutScreenVector : TEXCOORD1,	// Used by DeferredLightPixelMain directional light
+#endif
+#if PERMUTATION_ENABLE_DEBUG
+	out uint2 OutTileCoord : TILE_COORD,	// Used by StencilMainPS below
+#endif
+	out float4 OutPosition : SV_POSITION)	// Use by StencilMainPS
+{
+	const uint2 TileCoord = DecodeTile(TileListBuffer[InInstanceId]);
+
+	uint2 TileVertex = TileCoord * STRATA_TILE_SIZE;
+	TileVertex.x += InVertexId == 1 || InVertexId == 2 || InVertexId == 4 ? STRATA_TILE_SIZE : 0;
+	TileVertex.y += InVertexId == 2 || InVertexId == 4 || InVertexId == 5 ? STRATA_TILE_SIZE : 0;
+
+	OutPosition = float4(float2(TileVertex) * OutputViewSizeAndInvSize.zw * float2(2.0f, -2.0f) + float2(-1.0, 1.0f), 0.5f, 1.0f);
+
+#if PERMUTATION_ENABLE_TEXCOORD_SCREENVECTOR
+	OutTexCoord = float2(TileVertex) * OutputBufferSizeAndInvSize.zw;
+	OutScreenVector = mul(float4(OutPosition.xy, 1, 0), ViewScreenToTranslatedWorld).xyz;
+#endif
+
+#if PERMUTATION_ENABLE_DEBUG
+	OutTileCoord = TileCoord;
+#endif
+}
+
+#endif // SHADER_TILE_VS
+
+////////////////////////////////////////////////////////////////////////////////////////////////////////////
+
+#if SHADER_STENCIL_TAGGING_PS
+
+float4 DebugTileColor;
+
+void StencilTaggingMainPS(
+	in uint2 InTileCoord : TILE_COORD,
+	in float4 SVPos : SV_POSITION,
+	out float4 OutColor0 : SV_Target0)
+{
+	const bool bTileX = (InTileCoord.x & 1) == 0;
+	const bool bTileY = (InTileCoord.y & 1) == 0;
+	const bool bChecker = (bTileX && bTileY) || (!bTileX && !bTileY);
+	OutColor0 = DebugTileColor * float4(1.0f, 1.0f, 1.0f, bChecker ? 0.33f : 0.66f);
+	OutColor0.rgb *= OutColor0.a; // premultiplied alpha blending assuming alpha is opacity==coverage
+}
+
+#endif // SHADER_STENCIL_PS
+
+
diff --git a/Engine/Shaders/Shared/StrataDefinitions.h b/Engine/Shaders/Shared/StrataDefinitions.h
--- a/Engine/Shaders/Shared/StrataDefinitions.h
+++ b/Engine/Shaders/Shared/StrataDefinitions.h
@@ -29,6 +29,10 @@
 #define STRATA_BSDF_TYPE_SINGLELAYERWATER	4
 // When more than 5 BSDF must exists, please update STATE_BIT_COUNT_BSDF and FStrataClassification.ShadingModels packing in Strata.ush
 
+// The size of strata material classification tiles on screen
+#define STRATA_TILE_SIZE					8
+
+// This is used to round up the strata material buffer for better cache locality on read/write.
 #define STRATA_DATA_TILE_SIZE				8
 #define STRATA_DATA_TILE_SIZE_DIV_AS_SHIFT	3
 
diff --git a/Engine/Source/Runtime/Renderer/Private/Strata/Strata.cpp b/Engine/Source/Runtime/Renderer/Private/Strata/Strata.cpp
--- a/Engine/Source/Runtime/Renderer/Private/Strata/Strata.cpp
+++ b/Engine/Source/Runtime/Renderer/Private/Strata/Strata.cpp
@@ -489,12 +489,12 @@ class FStrataMaterialStencilTaggingPassPS : public FGlobalShader
 	static void ModifyCompilationEnvironment(const FGlobalShaderPermutationParameters& Parameters, FShaderCompilerEnvironment& OutEnvironment)
 	{
 		FGlobalShader::ModifyCompilationEnvironment(Parameters, OutEnvironment);
-		OutEnvironment.SetDefine(TEXT("SHADER_STENCIL_CATEGORIZATION"), 1); 
+		OutEnvironment.SetDefine(TEXT("SHADER_STENCIL_TAGGING_PS"), 1); 
 	}
 };
 
-IMPLEMENT_GLOBAL_SHADER(FStrataTilePassVS, "/Engine/Private/Strata/StrataMaterialClassification.usf", "StrataTilePassVS", SF_Vertex);
-IMPLEMENT_GLOBAL_SHADER(FStrataMaterialStencilTaggingPassPS, "/Engine/Private/Strata/StrataMaterialClassification.usf", "StencilMainPS", SF_Pixel);
+IMPLEMENT_GLOBAL_SHADER(FStrataTilePassVS, "/Engine/Private/Strata/StrataTiling.usf", "StrataTilePassVS", SF_Vertex);
+IMPLEMENT_GLOBAL_SHADER(FStrataMaterialStencilTaggingPassPS, "/Engine/Private/Strata/StrataTiling.usf", "StencilTaggingMainPS", SF_Pixel);
 
 void FillUpTiledPassData(
 	EStrataTileMaterialType Type, 
diff --git a/Engine/Source/Runtime/Renderer/Private/Strata/Strata.h b/Engine/Source/Runtime/Renderer/Private/Strata/Strata.h
--- a/Engine/Source/Runtime/Renderer/Private/Strata/Strata.h
+++ b/Engine/Source/Runtime/Renderer/Private/Strata/Strata.h
@@ -139,7 +139,7 @@ class FStrataTilePassVS : public FGlobalShader
 	static void ModifyCompilationEnvironment(const FGlobalShaderPermutationParameters& Parameters, FShaderCompilerEnvironment& OutEnvironment)
 	{
 		FGlobalShader::ModifyCompilationEnvironment(Parameters, OutEnvironment);
-		OutEnvironment.SetDefine(TEXT("SHADER_STENCIL_CATEGORIZATION"), 1);
+		OutEnvironment.SetDefine(TEXT("SHADER_TILE_VS"), 1);
 	}
 };
 
