Fix for forward rendererer incorrectly premulitplying transparent surfaces by alpha

diff --git a/Engine/Shaders/Private/MobileBasePassPixelShader.usf b/Engine/Shaders/Private/MobileBasePassPixelShader.usf
--- a/Engine/Shaders/Private/MobileBasePassPixelShader.usf
+++ b/Engine/Shaders/Private/MobileBasePassPixelShader.usf
@@ -1101,9 +1101,7 @@ void Main(
 		OutColor.rgb = sqrt(OutColor.rgb);
 	#endif
 #endif
-#if (MATERIALBLENDING_TRANSLUCENT || MATERIALBLENDING_ALPHACOMPOSITE)
-		OutColor.rgb = OutColor.rgb * OutColor.a;
-#endif
+
 #if NUM_VIRTUALTEXTURE_SAMPLES || LIGHTMAP_VT_ENABLED
 	FinalizeVirtualTextureFeedback(
 		MaterialParameters.VirtualTextureFeedback,
