Fix a bug of dark screen with vignette enabled on android.

diff --git a/Engine/Shaders/Private/PostProcessTonemap.usf b/Engine/Shaders/Private/PostProcessTonemap.usf
--- a/Engine/Shaders/Private/PostProcessTonemap.usf
+++ b/Engine/Shaders/Private/PostProcessTonemap.usf
@@ -208,19 +208,11 @@ float2 ScreenPosToUV(float2 ScreenPos, float2 ExtentInverse)
 {
 	float2 UV = (ScreenPos * Color_ScreenPosToViewportScale + Color_ScreenPosToViewportBias) * ExtentInverse;
 
-#if NEEDTOSWITCHVERTICLEAXIS
-	UV.y = 1.0 - UV.y;
-#endif
-
 	return UV;
 }
 
 float2 UVToScreenPos(float2 UV, float2 Extent)
 {
-#if NEEDTOSWITCHVERTICLEAXIS
-	UV.y = 1.0 - UV.y;
-#endif
-
 	return (UV * Extent - Color_ScreenPosToViewportBias) / Color_ScreenPosToViewportScale;
 }
 
