* Fixed incorrect landscape normals computation
* Removed dead code

#tests editor

diff --git a/Engine/Shaders/Private/LandscapeLayersPS.usf b/Engine/Shaders/Private/LandscapeLayersPS.usf
--- a/Engine/Shaders/Private/LandscapeLayersPS.usf
+++ b/Engine/Shaders/Private/LandscapeLayersPS.usf
@@ -104,8 +104,8 @@ void PSHeightmapMain(float2 InTextureCoordinates : TEXCOORD0, out float4 OutColo
 
 		bool bIsMinBorderTexelX = InTextureCoordinates.x <= TexelSize.x;
 		bool bIsMinBorderTexelY = InTextureCoordinates.y <= TexelSize.y;
-		bool bIsMaxBorderTexelX = InTextureCoordinates.x >= TexelSize.x * max(CurrentMipComponentVertexCount - 1, 1.0);
-		bool bIsMaxBorderTexelY = InTextureCoordinates.y >= TexelSize.y * max(CurrentMipComponentVertexCount - 1, 1.0);
+		bool bIsMaxBorderTexelX = InTextureCoordinates.x >= 1.0f - TexelSize.x;
+		bool bIsMaxBorderTexelY = InTextureCoordinates.y >= 1.0f - TexelSize.y;
 
 		// The triangle topology is the following (where C = center, T = top, B = bottom, L = left, R = right and Nx the normals we need to interpolate):
 		// TL ------ TT
@@ -302,8 +302,6 @@ void PSHeightmapMainMips(float2 InTextureCoordinates : TEXCOORD0, out float4 Out
 	OutColor = float4(PackedHeightResult.xy, NormalResult.xy);
 }
 
-float4 WeightmapLayerWeightBlend; // each xyzw == 1 texture channel
-
 void PSWeightmapMain(float2 InTextureCoordinates : TEXCOORD0, out float4 OutColor : SV_Target0)
 {
 	float4 SourceColor = ReadTexture1.Sample(ReadTexture1Sampler, InTextureCoordinates);
@@ -359,8 +357,6 @@ void PSWeightmapMainMips(float2 InTextureCoordinates : TEXCOORD0, out float4 Out
 
 	float2 ParentMipTexelSize = 1.0 / ParentMipTextureSize;
 	float2 SourceUV = InTextureCoordinates - (ParentMipTexelSize * 0.5);
-	float2 HeightResult = 0;
-	float2 NormalResult = 0;
 	float Alpha = 0.5;
 	OutColor = 0; // Default to 0
 
