SSD AO: Disable history accumulation on pixels with UNLIT shading model.

diff --git a/Engine/Shaders/Private/ScreenSpaceDenoise/SSDTemporalAccumulation.usf b/Engine/Shaders/Private/ScreenSpaceDenoise/SSDTemporalAccumulation.usf
--- a/Engine/Shaders/Private/ScreenSpaceDenoise/SSDTemporalAccumulation.usf
+++ b/Engine/Shaders/Private/ScreenSpaceDenoise/SSDTemporalAccumulation.usf
@@ -118,6 +118,7 @@
 	#define CONFIG_BILATERAL_PRESET BILATERAL_PRESET_AO_HISTORY
 	#define CONFIG_HISTORY_BILATERAL_PRESET BILATERAL_PRESET_AO_HISTORY
 	#define CONFIG_PATCH_PREV_SCENE_DEPTH 1
+	#define CONFIG_USE_UNLIT_HISTORY_REJECTION 1
 
 	// Use pre transformed rejection buffer that contains pre transformed momment 1 & 2.
 	#if 1
@@ -318,6 +319,10 @@
 	#define CONFIG_NORMALIZED_INPUT 0
 #endif
 
+/** Whether should reject history on pixels with UNLIT shading model. */
+#ifndef CONFIG_USE_UNLIT_HISTORY_REJECTION
+	#define CONFIG_USE_UNLIT_HISTORY_REJECTION 0
+#endif
 
 //------------------------------------------------------- CONFIG ENABLED DEFAULTS
 
@@ -497,6 +502,7 @@ void TemporallyAccumulate(
 		SceneBufferUV, BufferUVToBufferPixelCoord(SceneBufferUV));
 
 	float DeviceZ;
+	uint ShadingModelID;
 	{
 		FSSDSampleSceneInfos RefInfo = UncompressSampleSceneInfo(
 			CONFIG_METADATA_BUFFER_LAYOUT, /* bIsPrevFrame = */ false,
@@ -504,6 +510,7 @@ void TemporallyAccumulate(
 			CompressedRefSceneMetadata);
 
 		DeviceZ = RefInfo.DeviceZ;
+		ShadingModelID = RefInfo.ShadingModelID;
 	}
 
 	// Reproject to previous frame.
@@ -966,7 +973,14 @@ void TemporallyAccumulate(
 
 					// When doing upscaling, may have an invalid input sample.
 					FinalHistoryWeight = max(FinalHistoryWeight, NeighborMoment1.SampleCount * 0.1);
-			
+
+#if CONFIG_USE_UNLIT_HISTORY_REJECTION
+					bool bValid = DenoiseShadingModelID(ShadingModelID);
+					if (!bValid)
+					{
+						FinalHistoryWeight = 0.0;
+					}
+#endif
 					// Apply the history weight.
 					HistorySample = MulSignal(ClampedHistorySample, FinalHistoryWeight);
 					HistorySample.SampleCount = FinalHistoryWeight;
