Fix invalid BSDF type when using strata single permutation.

diff --git a/Engine/Shaders/Private/Strata/Strata.ush b/Engine/Shaders/Private/Strata/Strata.ush
--- a/Engine/Shaders/Private/Strata/Strata.ush
+++ b/Engine/Shaders/Private/Strata/Strata.ush
@@ -2519,6 +2519,8 @@ FStrataBSDF UnpackStrataBSDFIn(FStrataMaterialContainer StrataBuffer, inout FStr
 	#endif
 		OutBSDF.State = (Data>>HEADER_SINGLEENCODING_BIT_COUNT) & STATE_SINGLEENCODING_MASK;
 		OutBSDF.FinalWeight = 1.0f;
+		const uint BSDFType = firstbitlow(GetShadingModels(StrataHeader));
+		BSDF_SETTYPE(OutBSDF, BSDFType);
 
 		// For now, single encode doesn't support anisotropy, as the top layer data only contains the normal, not the full frame basis
 		BSDF_SETHASANISOTROPY(OutBSDF, 0);
