Fix subsurface scattering alpha issue with alpha compositing on

diff --git a/Engine/Shaders/Private/PostProcessSubsurface.usf b/Engine/Shaders/Private/PostProcessSubsurface.usf
--- a/Engine/Shaders/Private/PostProcessSubsurface.usf
+++ b/Engine/Shaders/Private/PostProcessSubsurface.usf
@@ -917,9 +917,8 @@ void SubsurfaceRecombinePS(in FScreenVertexOutput Input, out float4 OutColor : S
 #endif
 
 	if (!bUseSSSProfile)
-	{
-		// Ingore the copy of non-scattering pixels with zero alpha.  
-		OutColor = float4(0.0f, 0.0f, 0.0f, 0.0f);
+	{  
+		OutColor = Texture2DSample(SubsurfaceInput0_Texture, SharedSubsurfaceSampler0, BufferUV);
 		return;
 	}
 
@@ -1007,29 +1006,14 @@ void SubsurfaceRecombinePS(in FScreenVertexOutput Input, out float4 OutColor : S
 	// combine potentially half res with full res
 	float3 SubsurfaceLighting = lerp(DiffuseAndSpecular.Diffuse, SSSColor, FadedTint);
 
-	OutColor = float4(SubsurfaceLighting * StoredBaseColor + ExtractedNonSubsurface, 1.0f);
+	OutColor = float4(SubsurfaceLighting * StoredBaseColor + ExtractedNonSubsurface, 0.0f);
 }
 
 void SubsurfaceRecombineCopyPS(in FScreenVertexOutput Input, out float4 OutColor : SV_Target0)
 {
 	const float2 BufferUV = Input.UV.xy;
 
-#if STRATA_ENABLED
-	const FStrataSubsurfaceData SSSData = LoadStataSSSData(BufferUV);
-	const bool bUseSSSProfile = SSSData.bIsValid;
-#else
-	FScreenSpaceData ScreenSpaceData = GetScreenSpaceData(BufferUV);
-	const bool bUseSSSProfile = UseSubsurfaceProfile(ScreenSpaceData.GBuffer.ShadingModelID);
-#endif
-
-	if (bUseSSSProfile)
-	{
-		OutColor = float4(Texture2DSampleLevel(SubsurfaceInput0_Texture, SharedSubsurfaceSampler0, BufferUV,0).rgb,1.0f);
-	}
-	else
-	{
-		OutColor = float4(0.0f, 0.0f, 0.0f, 0.0f);
-	}
+	OutColor = Texture2DSampleLevel(SubsurfaceInput0_Texture, SharedSubsurfaceSampler0, BufferUV, 0);
 }
 
 void SubsurfaceViewportCopyPS(noperspective float4 InUV : TEXCOORD0, out float4 OutColor : SV_Target0)
diff --git a/Engine/Source/Runtime/Renderer/Private/PostProcess/PostProcessSubsurface.cpp b/Engine/Source/Runtime/Renderer/Private/PostProcess/PostProcessSubsurface.cpp
--- a/Engine/Source/Runtime/Renderer/Private/PostProcess/PostProcessSubsurface.cpp
+++ b/Engine/Source/Runtime/Renderer/Private/PostProcess/PostProcessSubsurface.cpp
@@ -1374,7 +1374,7 @@ void AddSubsurfaceViewPass(
 			TShaderMapRef<FSubsurfaceRecombineCopyPS> PixelShader(View.ShaderMap);			
 			TShaderMapRef<FSubsurfaceTilePassVS> VertexShader(View.ShaderMap);
 			
-			FRHIBlendState* BlendState = TStaticBlendState<CW_RGBA, BO_Add, BF_SourceAlpha, BF_InverseSourceAlpha, BO_Add, BF_Zero, BF_One>::GetRHI();
+			FRHIBlendState* BlendState = TStaticBlendState<CW_RGBA, BO_Add, BF_One, BF_Zero, BO_Add, BF_One, BF_Zero>::GetRHI();
 			
 			AddSubsurfaceTiledScreenPass(
 				GraphBuilder,
