Fix strata hair model conversion when shading model is driven by a material expression

diff --git a/Engine/Shaders/Private/Strata/StrataLegacyConversion.ush b/Engine/Shaders/Private/Strata/StrataLegacyConversion.ush
--- a/Engine/Shaders/Private/Strata/StrataLegacyConversion.ush
+++ b/Engine/Shaders/Private/Strata/StrataLegacyConversion.ush
@@ -109,7 +109,45 @@ FStrataData StrataConvertLegacyMaterialDynamic(
 	const float3 Ones  = float3(1, 1, 1);
 
 	// Can only mix Unlit / DefaultLit / Sub-Surface / Preintegrated-Skin / Subsurface-Profile / ClearCoat / Foliage / Cloth / Eye/ Thin
-	// No supported : Water/Hair
+
+	// Note: If hair or single layer water are enabled with other shading model, the code generation won't be ideal, as the compiler will struggle with unrolling the BSDFs
+	#if MATERIAL_SHADINGMODEL_SINGLELAYERWATER
+	if (ShadingModel == SHADINGMODELID_SINGLELAYERWATER)
+	{
+		const float3 WaterExtinction = WaterScatteringCoefficients + WaterAbsorptionCoefficients;
+		const float3 WaterAlbedo = WaterScatteringCoefficients / WaterExtinction;
+
+		return GetStrataSingleLayerWaterBSDF(
+			BaseColor, 				// BaseColor
+			Metallic, 				// Metallic 
+			Specular, 				// Specular
+			Roughness, 				// Roughness 
+			Emissive, 				// Emissive 
+			Opacity,				// TopMaterialOpacity
+			WaterAlbedo, 			// WaterAlbedo
+			WaterExtinction, 		// WaterExtinction
+			WaterPhaseG, 			// WaterPhaseG
+			ColorScaleBehindWater,	// ColorScaleBehindWater
+			SharedLocalBasisIndex);
+
+		// No need to take into account opacity weighting for translucent material, as the coverage is assumed to be 1
+	}
+	#endif
+	#if MATERIAL_SHADINGMODEL_HAIR
+	if (ShadingModel == SHADINGMODELID_HAIR)
+	{
+		return GetStrataHairBSDF(
+			BaseColor,	   // BaseColor
+			Metallic,	   // Scatter
+			Specular,	   // Specular
+			Roughness,	   // Roughness
+			ClearCoat,	   // Backlit
+			Emissive,	   // EmissiveColor
+			SharedLocalBasisIndex);
+
+		// No need to take into account opacity weighting for translucent material, as the coverage is assumed to be 1
+	}
+	#endif
 
 	// Top/Bottom
 	FStrataLegacyParameters Top = InitStrataLegacyParameters(SharedLocalBasisIndex, 1.0f);
