Add Strata support for fetching post-process input.

diff --git a/Engine/Shaders/Private/MaterialTemplate.ush b/Engine/Shaders/Private/MaterialTemplate.ush
--- a/Engine/Shaders/Private/MaterialTemplate.ush
+++ b/Engine/Shaders/Private/MaterialTemplate.ush
@@ -2188,7 +2188,65 @@ float4 SceneTextureLookup(float2 UV, int SceneTextureIndex, bool bFiltered)
 	}
 
 #else // !SCENE_TEXTURES_DISABLED
+#if STRATA_ENABLED
 
+	const uint2 PixelPos = uint2(UV * View.BufferSizeAndInvSize.xy);
+
+	// 1. Post-process inputs which are independent on material data
+	switch (SceneTextureIndex)
+	{
+	case PPI_SceneColor:			return float4(CalcSceneColor(UV), 0);
+	case PPI_SceneDepth:			return CalcSceneDepth(UV);
+	case PPI_SeparateTranslucency:	return float4(1, 1, 1, 1);	// todo
+	case PPI_Opacity:				return 1.f; // Opacity has different meaning depending on the legacy shading model.
+	case PPI_CustomDepth:			return CalcSceneCustomDepth(UV);
+#if POST_PROCESS_MATERIAL
+	case PPI_PostProcessInput0:		return Texture2DSample(PostProcessInput_0_Texture, bFiltered ? PostProcessInput_BilinearSampler : PostProcessInput_0_SharedSampler, UV);
+	case PPI_PostProcessInput1:		return Texture2DSample(PostProcessInput_1_Texture, bFiltered ? PostProcessInput_BilinearSampler : PostProcessInput_1_SharedSampler, UV);
+	case PPI_PostProcessInput2:		return Texture2DSample(PostProcessInput_2_Texture, bFiltered ? PostProcessInput_BilinearSampler : PostProcessInput_2_SharedSampler, UV);
+	case PPI_PostProcessInput3:		return Texture2DSample(PostProcessInput_3_Texture, bFiltered ? PostProcessInput_BilinearSampler : PostProcessInput_3_SharedSampler, UV);
+	case PPI_PostProcessInput4:		return Texture2DSample(PostProcessInput_4_Texture, bFiltered ? PostProcessInput_BilinearSampler : PostProcessInput_4_SharedSampler, UV);
+#endif // __POST_PROCESS_COMMON__
+	case PPI_DecalMask:				return 0;  // material compiler will return an error
+	case PPI_AmbientOcclusion:		return CalcSceneAO(UV);
+	case PPI_CustomStencil:			return CalcSceneCustomStencil(PixelPos);
+#if POST_PROCESS_MATERIAL
+	case PPI_Velocity:				return float4(PostProcessVelocityLookup(ConvertToDeviceZ(CalcSceneDepth(UV)), UV), 0, 0);
+#endif
+	}
+
+	// 2. Post-process inputs which are dependent on material data
+	//    Note: For PostProcess material we always return the first layer dataCalcSceneCustomStencil
+	FStrataAddressing StrataAddressing   = GetStrataPixelDataByteOffset(PixelPos, uint2(View.BufferSizeAndInvSize.xy), Strata.MaxBytesPerPixel);
+	FStrataPixelHeader StrataPixelHeader = UnpackStrataHeaderIn(Strata.MaterialTextureArray, StrataAddressing, Strata.TopLayerTexture);
+	BRANCH
+	if (StrataPixelHeader.BSDFCount > 0)
+	{
+		const FStrataBSDF BSDF = UnpackStrataBSDFIn(Strata.MaterialTextureArray, StrataAddressing, StrataPixelHeader);
+		switch (SceneTextureIndex)
+		{
+			// order needs to match to ESceneTextureId
+			case PPI_DiffuseColor:			return float4(StrataGetBSDFDiffuseColor(BSDF), 0);
+			case PPI_SpecularColor:			return float4(StrataGetBSDFSpecularColor(BSDF), 0);
+			case PPI_SubsurfaceColor:		return float4(StrataGetBSDFSubSurfaceColor(BSDF), 0);
+			case PPI_BaseColor:				return float4(StrataGetBSDFBaseColor(BSDF), 0);
+			case PPI_Specular:				return StrataGetBSDFSpecular(BSDF);
+			case PPI_Metallic:				return StrataGetBSDFMetallic(BSDF);
+			case PPI_WorldNormal:			return float4(StrataGetWorldNormal(StrataPixelHeader, BSDF, StrataAddressing), 0);
+			case PPI_Roughness:				return StrataGetBSDFRoughness(BSDF);
+			case PPI_MaterialAO:			return StrataGetAO(StrataPixelHeader);
+			case PPI_ShadingModelColor:		return float4(GetShadingModelColor(StrataGetLegacyShadingModels(BSDF)), 1);
+			case PPI_ShadingModelID:		return float4(StrataGetLegacyShadingModels(BSDF), 0, 0, 0);
+			case PPI_StoredBaseColor:		return float4(StrataGetBSDFBaseColor(BSDF), 0);
+			case PPI_StoredSpecular:		return float4(StrataGetBSDFSpecular(BSDF).rrr, 0);
+			case PPI_WorldTangent:			return float4(StrataGetWorldTangent(StrataPixelHeader, BSDF, StrataAddressing), 0);
+			case PPI_Anisotropy:			return StrataGetBSDFAnisotropy(BSDF);
+		}
+	}
+
+	return float4(0, 0, 0, 0);
+
+#else // STRATA_ENABLED
 	FScreenSpaceData ScreenSpaceData = GetScreenSpaceData(UV, false);
 	switch(SceneTextureIndex)
 	{
@@ -2259,6 +2317,7 @@ float4 SceneTextureLookup(float2 UV, int SceneTextureIndex, bool bFiltered)
 		default:
 			return float4(0, 0, 0, 0);
 	}
+#endif // STRATA_ENABLED
 #endif // SCENE_TEXTURES_DISABLED
 }
 
diff --git a/Engine/Shaders/Private/PostProcessMaterialShaders.usf b/Engine/Shaders/Private/PostProcessMaterialShaders.usf
--- a/Engine/Shaders/Private/PostProcessMaterialShaders.usf
+++ b/Engine/Shaders/Private/PostProcessMaterialShaders.usf
@@ -4,6 +4,9 @@
 	PostProcessMaterialShaders.usf: Shaders for rendering post process materials
 =============================================================================*/
 
+#define STRATA_INLINE_SHADING 1
+#define STRATA_DEFERRED_SHADING 1
+
 #include "Common.ush"
 #include "ScreenPass.ush"
 
@@ -203,7 +206,7 @@ void MainPS(
 	CalcMaterialParametersPost(Parameters, PixelMaterialInputs, SvPosition, true);
 
 	float Alpha = 1.0f;
-#if STRATA_ENABLED  
+#if STRATA_ENABLED   
 	// Unlit forces a single BSDF
 	const FStrataBSDF UnlitBSDF = PixelMaterialInputs.FrontMaterial.Layers[0].BSDFs[0];
 
diff --git a/Engine/Shaders/Private/SceneTexturesCommon.ush b/Engine/Shaders/Private/SceneTexturesCommon.ush
--- a/Engine/Shaders/Private/SceneTexturesCommon.ush
+++ b/Engine/Shaders/Private/SceneTexturesCommon.ush
@@ -102,6 +102,34 @@ float4 GatherSceneDepth(float2 UV, float2 InvBufferSize)
 #endif
 }
 
+/** Fetches custom device depth and converts to linear. */
+float CalcSceneCustomDepth(float2 ScreenUV)
+{
+#if SCENE_TEXTURES_DISABLED
+	return SCENE_TEXTURES_DISABLED_SCENE_DEPTH_VALUE;
+#else
+	return ConvertFromDeviceZ(Texture2DSampleLevel(SceneTexturesStruct.CustomDepthTexture, SceneTexturesStruct_CustomDepthTextureSampler, ScreenUV, 0).r);
+#endif
+}
+
+uint CalcSceneCustomStencil(uint2 PixelPos)
+{
+#if SCENE_TEXTURES_DISABLED
+	return 0;
+#else
+	return SceneTexturesStruct.CustomStencilTexture.Load(uint3(PixelPos, 0)).r;
+#endif
+}
+
+float CalcSceneAO(float2 ScreenUV)
+{
+#if SCENE_TEXTURES_DISABLED
+	return 1.0f;
+#else
+	return Texture2DSampleLevel(SceneTexturesStruct.ScreenSpaceAOTexture, SceneTexturesStruct_ScreenSpaceAOTextureSampler, ScreenUV, 0).r;
+#endif
+}
+
 #endif // SHADING_PATH_DEFERRED
 
 // Only reference MobileSceneTextures uniform buffer if SHADING_PATH_MOBILE
diff --git a/Engine/Shaders/Private/ShadingCommon.ush b/Engine/Shaders/Private/ShadingCommon.ush
--- a/Engine/Shaders/Private/ShadingCommon.ush
+++ b/Engine/Shaders/Private/ShadingCommon.ush
@@ -138,3 +138,18 @@ float ComputeHazyLobeRoughness(float Roughness, float Haziness)
 {
 	return lerp(Roughness, 1.0f, Haziness);
 }
+
+float F0ToMetallic(float F0)
+{
+	// Approximate the metallic input from F0 with a small lerp region
+	const float DiamondF0 = 0.24f;
+	const float CondutorMinF0 = 0.4f;	// roughly the end of semi-conductor
+	// This is compatible with UE shading model mapping allowing F0 to take a value up to 0.08 for dielectric.
+
+	return saturate((F0 - DiamondF0) / (CondutorMinF0 - DiamondF0));
+}
+
+float F0RGBToMetallic(float3 F0)
+{
+	return F0ToMetallic(max(F0.r, max(F0.g, F0.b)));
+}
diff --git a/Engine/Shaders/Private/Strata/Strata.ush b/Engine/Shaders/Private/Strata/Strata.ush
--- a/Engine/Shaders/Private/Strata/Strata.ush
+++ b/Engine/Shaders/Private/Strata/Strata.ush
@@ -61,6 +61,10 @@ FRWStrataMaterialContainer InitialiseRWStrataMaterialContainer(RWTexture2DArray<
 #define STRATA_INLINE_SHADING 1
 #endif
 
+#ifndef STRATA_DEFERRED_SHADING
+#define STRATA_DEFERRED_SHADING (STRATA_INLINE_SHADING == 0)
+#endif 
+
 // During the unpacking of the Strata material data, if a slab has some sub-surface scattering, its BaseColor/Specular value will be 
 // patched/overriden with white albedo (and optionally no specular) for computing the incoming irradiance without any albedo information. 
 // This is required by the post-process/screen-space sub-surface methods. This overriden can be opt-out. This is used during the classification 
@@ -98,15 +102,12 @@ struct FStrataBSDF
 {
 	uint				State;		// Extra informations and other enabled BSDF features
 
+	float3				FinalWeight;// Contains the throughput towards the view accounting for BSDF visible above. It also contains coverage if it has not been baked into the BSDF parameters when using operation with parameter blending.
 #if STRATA_INLINE_SHADING
 	float				Coverage;	// When creating a material, each BSDF/slab can have its coverage updated. 
 									// When this happens, we also update its parameters to account for that, similary to premultiplied alpha blending.
 									// Coverage is kept around because it is needed to evalauted the contribution to throughput.
 	bool				bPremultipliedCoverage;// True when coverage is already weighting all parameters. This is true when parameter blending is used. This allows to use parameter blending or not in the same graph.
-#endif
-	float3				FinalWeight;// Contains the throughput towards the view accounting for BSDF visible above. It also contains coverage if it has not been baked into the BSDF parameters when using operation with parameter blending.
-
-#if STRATA_INLINE_SHADING
 	float3				Emissive;	// All bsdf can potientially have emissive. This is never written out to the strata gbuffer but instead accumulated in the emissive buffer during the base pass.
 	float				ThicknessCm;// This must be kept in order to be able to normalize the mean free path to comply with our volumetric material math all done in a normalised slab of medium.
 	float3				TmpDMFP;
@@ -901,7 +902,8 @@ struct FStrataPixelHeader
 #if STRATA_INLINE_SHADING
 	// All the shared local bases in VGPRS. This is only used in basepass and during forward rendering.
 	FSharedLocalBases SharedLocalBases;
-#else
+#endif
+#if STRATA_DEFERRED_SHADING
 	FStrataMaterialContainer StrataBuffer;
 	uint SharedLocalBasesIndexOffset;
 	uint SharedLocalBasesTypes_PackedHeader; // Stores SharedLocalBasesTypes for non-simple material, and PackedHeader data for simple material
@@ -945,7 +947,8 @@ FStrataPixelHeader InitialiseStrataPixelHeader()
 	Out.BSDFCount = 0;
 #if STRATA_INLINE_SHADING
 	Out.SharedLocalBases = (FSharedLocalBases)0;
-#else
+#endif
+#if STRATA_DEFERRED_SHADING
 	Out.SharedLocalBasesIndexOffset = 0;
 	Out.SharedLocalBasesTypes_PackedHeader = 0;
 	Out.PackedTopLayerData = 0;
@@ -955,7 +958,7 @@ FStrataPixelHeader InitialiseStrataPixelHeader()
 	return Out;
 }
 
-float3x3 StrataGetBSDFSharedBasis(in FStrataPixelHeader StrataPixelHeader, in FStrataBSDF BSDF, in FStrataAddressing StrataAddressing)
+float3x3 StrataGetBSDFSharedBasis_InlineShading(in FStrataPixelHeader StrataPixelHeader, in FStrataBSDF BSDF, in FStrataAddressing StrataAddressing)
 {
 #if STRATA_INLINE_SHADING
 	// Assume Normal & Tangent are already been normalized
@@ -963,7 +966,12 @@ float3x3 StrataGetBSDFSharedBasis(in FStrataPixelHeader StrataPixelHeader, in FS
 	float3 Tangent = StrataPixelHeader.SharedLocalBases.Tangents[BSDF_GETSHAREDLOCALBASISID(BSDF)];
 	float3 Bitangent = cross(Normal, Tangent);
 	return float3x3(Tangent, Bitangent, Normal);
-#else
+#endif
+}
+
+float3x3 StrataGetBSDFSharedBasis_DeferredShading(in FStrataPixelHeader StrataPixelHeader, in FStrataBSDF BSDF, in FStrataAddressing StrataAddressing)
+{
+#if STRATA_DEFERRED_SHADING
 	if (IsSimpleMaterial(StrataPixelHeader) || IsSingleMaterial(StrataPixelHeader))
 	{		
 		return GetTangentBasis(StrataUnpackTopLayerData(StrataPixelHeader.PackedTopLayerData).WorldNormal);
@@ -992,7 +1000,10 @@ float3x3 StrataGetBSDFSharedBasis(in FStrataPixelHeader StrataPixelHeader, in FS
 #endif
 }
 
-
+float3x3 StrataGetBSDFSharedBasis(in FStrataPixelHeader StrataPixelHeader, in FStrataBSDF BSDF, in FStrataAddressing StrataAddressing)
+{
+	return StrataGetBSDFSharedBasis_DeferredShading(StrataPixelHeader, BSDF, StrataAddressing);
+}
 
 ///////////////////////////////////////////////////////////////////////////////
 // Operator nodes
@@ -1568,6 +1579,130 @@ void EnableSlabBSDFIsThin(inout FStrataBSDF BSDF)
 	EnableSlabMFPComputationSpace(BSDF);
 }
 
+///////////////////////////////////////////////////////////////////////////////
+// Generic BSDF accessors
+
+float3 StrataGetBSDFBaseColor(in FStrataBSDF BSDF)
+{
+	switch (BSDF_GETTYPE(BSDF))
+	{
+	case STRATA_BSDF_TYPE_SLAB: 			return lerp(SLAB_DIFFUSEALBEDO(BSDF), SLAB_F0(BSDF), F0RGBToMetallic(SLAB_F0(BSDF)));  // This BaseColor recovery based on Metallic recovers is not accurate, and only used for legacy conversion purpose (e.g., PostProcessMaterial texture access)
+	case STRATA_BSDF_TYPE_HAIR: 			return HAIR_BASECOLOR(BSDF);
+	case STRATA_BSDF_TYPE_SINGLELAYERWATER: return SLW_BASECOLOR(BSDF);
+	}
+	return 0.f;
+}
+
+float3 StrataGetBSDFDiffuseColor(in FStrataBSDF BSDF)
+{
+	switch (BSDF_GETTYPE(BSDF))
+	{
+	case STRATA_BSDF_TYPE_SLAB: 			return SLAB_DIFFUSEALBEDO(BSDF);
+	case STRATA_BSDF_TYPE_HAIR: 			return HAIR_BASECOLOR(BSDF);
+	case STRATA_BSDF_TYPE_SINGLELAYERWATER: return SLW_BASECOLOR(BSDF);
+	}
+	return 0.f;
+}
+
+float3 StrataGetBSDFSpecularColor(in FStrataBSDF BSDF)
+{
+	switch (BSDF_GETTYPE(BSDF))
+	{
+	case STRATA_BSDF_TYPE_SLAB: 			return SLAB_F0(BSDF);
+	case STRATA_BSDF_TYPE_HAIR: 			return DielectricSpecularToF0(HAIR_SPECULAR(BSDF));
+	case STRATA_BSDF_TYPE_SINGLELAYERWATER: return DielectricSpecularToF0(SLW_SPECULAR(BSDF));
+	}
+	return 0.f;
+}
+
+float3 StrataGetBSDFSubSurfaceColor(in FStrataBSDF BSDF)
+{
+	switch (BSDF_GETTYPE(BSDF))
+	{
+	case STRATA_BSDF_TYPE_SLAB: 			return BSDF_GETHASSSS(BSDF) ? SLAB_DIFFUSEALBEDO(BSDF) : 0; // A more accurate version would be IsotropicMediumSlabPunctualDirectionalAlbedo(StrataSlabCreateParticipatingMedia(SLAB_DIFFUSEALBEDO(BSDF), SLAB_SSSDMFP(BSDFContext.BSDF)));
+	case STRATA_BSDF_TYPE_HAIR: 			return 0;
+	case STRATA_BSDF_TYPE_SINGLELAYERWATER: return 0;
+	}
+	return 0.f;
+}
+
+float StrataGetBSDFSpecular(in FStrataBSDF BSDF)
+{
+	switch (BSDF_GETTYPE(BSDF))
+	{
+	case STRATA_BSDF_TYPE_SLAB: 			return F0ToDielectricSpecular(SLAB_F0(BSDF).x);
+	case STRATA_BSDF_TYPE_HAIR: 			return HAIR_SPECULAR(BSDF);
+	case STRATA_BSDF_TYPE_SINGLELAYERWATER: return SLW_SPECULAR(BSDF);
+	}
+	return 0.f;
+}
+
+float StrataGetBSDFMetallic(in FStrataBSDF BSDF)
+{
+	switch (BSDF_GETTYPE(BSDF))
+	{
+	case STRATA_BSDF_TYPE_SLAB: 			return F0RGBToMetallic(SLAB_F0(BSDF)); // Metallic recovers is not accurate, and only used for legacy conversion purpose (e.g., PostProcessMaterial texture access)
+	case STRATA_BSDF_TYPE_HAIR: 			return 0;
+	case STRATA_BSDF_TYPE_SINGLELAYERWATER: return 0;
+	}
+	return 0.f;
+}
+
+float StrataGetBSDFRoughness(in FStrataBSDF BSDF)
+{
+	switch (BSDF_GETTYPE(BSDF))
+	{
+	case STRATA_BSDF_TYPE_SLAB: 			return SLAB_ROUGHNESS(BSDF);
+	case STRATA_BSDF_TYPE_HAIR: 			return HAIR_ROUGHNESS(BSDF);
+	case STRATA_BSDF_TYPE_SINGLELAYERWATER: return SLW_ROUGHNESS(BSDF);
+	}
+	return 0.f;
+}
+
+float StrataGetBSDFAnisotropy(in FStrataBSDF BSDF)
+{
+	switch (BSDF_GETTYPE(BSDF))
+	{
+	case STRATA_BSDF_TYPE_SLAB: 			return SLAB_ANISOTROPY(BSDF);
+	case STRATA_BSDF_TYPE_HAIR: 			return 0;
+	case STRATA_BSDF_TYPE_SINGLELAYERWATER: return 0;
+	}
+	return 0.f;
+}
+
+float3 StrataGetWorldNormal(in FStrataPixelHeader Header, in FStrataBSDF BSDF, in const FStrataAddressing Addressing)
+{
+	return StrataGetBSDFSharedBasis_DeferredShading(Header, BSDF, Addressing)[2];
+}
+
+float3 StrataGetWorldTangent(in FStrataPixelHeader Header, in FStrataBSDF BSDF, in const FStrataAddressing Addressing)
+{
+	return StrataGetBSDFSharedBasis_DeferredShading(Header, BSDF, Addressing)[0];
+}
+
+float StrataGetAO(in FStrataPixelHeader Header)
+{
+	return HEADER_GETAO(Header.State);
+}
+
+float StrataGetLegacyShadingModels(in FStrataBSDF BSDF)
+{
+	switch (BSDF_GETTYPE(BSDF))
+	{
+	case STRATA_BSDF_TYPE_SLAB:
+	{
+		// Base on features, we can only partially recovers the original shading models. Shading models like foliage, eye, clear coat are impossible/harder to infer.
+		if (BSDF_GETHASSSSPROFILE(BSDF))	{ return SHADINGMODELID_SUBSURFACE_PROFILE; }
+		else if (BSDF_GETHASSSS(BSDF))		{ return SHADINGMODELID_SUBSURFACE; }
+		else if (BSDF_GETHASFUZZ(BSDF))		{ return SHADINGMODELID_CLOTH; }
+		else if (BSDF_GETHASFUZZ(BSDF))		{ return SHADINGMODELID_TWOSIDED_FOLIAGE; }
+		else								{ return SHADINGMODELID_DEFAULT_LIT; }
+	}
+	case STRATA_BSDF_TYPE_HAIR: 			{ return SHADINGMODELID_HAIR; }
+	case STRATA_BSDF_TYPE_SINGLELAYERWATER: { return SHADINGMODELID_SINGLELAYERWATER; }
+	}
+	return 0.f;
+}
 
 ///////////////////////////////////////////////////////////////////////////////
 // BSDFs packing/unpacking
@@ -1953,7 +2088,7 @@ void PackStrataOut(
 				BSDFCount++;
 
 				// STRATA_TODO apply the same normal transform as for the material
-				float3x3 TangentBasis = StrataGetBSDFSharedBasis(StrataPixelHeader, BSDF, StrataAddressing);
+				float3x3 TangentBasis = StrataGetBSDFSharedBasis_InlineShading(StrataPixelHeader, BSDF, StrataAddressing);
 				float3 N = TangentBasis[2];
 				//N = normalize(N);
 
@@ -2402,7 +2537,7 @@ FStrataPixelHeader UnpackStrataHeaderIn(FStrataMaterialContainer StrataBuffer, i
 			HEADER_SETSHADINGMODELS(Out.State, ShadingModels);
 		}
 
-	#if STRATA_INLINE_SHADING==0
+	#if STRATA_DEFERRED_SHADING
 		Out.StrataBuffer = StrataBuffer;
 		Out.SharedLocalBasesTypes_PackedHeader = PackedHeader;
 		Out.PackedTopLayerData = InStrataTopLayerTexture.Load(uint3(StrataAddressing.PixelCoords, 0));
@@ -2414,7 +2549,7 @@ FStrataPixelHeader UnpackStrataHeaderIn(FStrataMaterialContainer StrataBuffer, i
 		uint SharedLocalBasesCount = HEADER_GETSHAREDLOCALBASESCOUNT(PackedHeader);
 		uint SharedLocalBasesTypes = HEADER_GETSHAREDLOCALBASESTYPE(PackedHeader);
 
-	#if STRATA_INLINE_SHADING==0
+	#if STRATA_DEFERRED_SHADING
 		Out.StrataBuffer = StrataBuffer;
 		Out.SharedLocalBasesTypes_PackedHeader = SharedLocalBasesTypes;
 		Out.PackedTopLayerData = 0;
@@ -2457,7 +2592,7 @@ FStrataBSDF UnpackFastPathStrataBSDFIn(FStrataMaterialContainer StrataBuffer, in
 {
 	FStrataBSDF OutBSDF = (FStrataBSDF)0;
 
-#if STRATA_INLINE_SHADING==0
+#if STRATA_DEFERRED_SHADING
 	uint Data0 = StrataPixelHeader.SharedLocalBasesTypes_PackedHeader;
 	uint Data1 = StrataLoadUint1(StrataBuffer, StrataAddressing);
 	OutBSDF.State = 0;
@@ -2514,7 +2649,7 @@ FStrataBSDF UnpackStrataBSDFIn(FStrataMaterialContainer StrataBuffer, inout FStr
 	if (IsSingleMaterial(StrataHeader))
 	{
 		uint Data = 0;
-	#if STRATA_INLINE_SHADING==0
+	#if STRATA_DEFERRED_SHADING
 		Data = StrataHeader.SharedLocalBasesTypes_PackedHeader;
 	#endif
 		OutBSDF.State = (Data>>HEADER_SINGLEENCODING_BIT_COUNT) & STATE_SINGLEENCODING_MASK;
diff --git a/Engine/Shaders/Private/Strata/StrataEvaluation.ush b/Engine/Shaders/Private/Strata/StrataEvaluation.ush
--- a/Engine/Shaders/Private/Strata/StrataEvaluation.ush
+++ b/Engine/Shaders/Private/Strata/StrataEvaluation.ush
@@ -77,31 +77,6 @@ struct FStrataBSDFContext
 	float3 TangentH;
 };
 
-float StrataGetBSDFRoughness(in FStrataBSDF BSDF)
-{
-	float R = 0;
-	const uint BSDFType = BSDF_GETTYPE(BSDF);
-	switch (BSDFType)
-	{
-	case STRATA_BSDF_TYPE_SLAB:
-	{
-		R = SLAB_ROUGHNESS(BSDF);
-		break;
-	}
-	case STRATA_BSDF_TYPE_HAIR:
-	{
-		R = HAIR_ROUGHNESS(BSDF);
-		break;
-	}
-	case STRATA_BSDF_TYPE_SINGLELAYERWATER:
-	{
-		R = SLW_ROUGHNESS(BSDF);
-		break;
-	}
-	}
-	return R;
-}
-
 FStrataBSDFContext StrataCreateBSDFContext(float3x3 TangentBasis, FStrataBSDF BSDF, float3 V, float3 L, bool bHasValidL=true)
 {
 	FStrataBSDFContext BSDFContext = (FStrataBSDFContext)0;
diff --git a/Engine/Source/Runtime/Renderer/Private/PostProcess/PostProcessMaterial.cpp b/Engine/Source/Runtime/Renderer/Private/PostProcess/PostProcessMaterial.cpp
--- a/Engine/Source/Runtime/Renderer/Private/PostProcess/PostProcessMaterial.cpp
+++ b/Engine/Source/Runtime/Renderer/Private/PostProcess/PostProcessMaterial.cpp
@@ -553,6 +553,8 @@ FScreenPassTexture AddPostProcessMaterialPass(
 	}
 
 	PostProcessMaterialParameters->bFlipYAxis = Inputs.bFlipYAxis && !bForceIntermediateTarget;
+	PostProcessMaterialParameters->Strata = Strata::BindStrataGlobalUniformParameters(View.StrataSceneData);
+
 
 	ClearUnusedGraphResources(VertexShader, PixelShader, PostProcessMaterialParameters);
 
diff --git a/Engine/Source/Runtime/Renderer/Private/PostProcess/PostProcessMaterial.h b/Engine/Source/Runtime/Renderer/Private/PostProcess/PostProcessMaterial.h
--- a/Engine/Source/Runtime/Renderer/Private/PostProcess/PostProcessMaterial.h
+++ b/Engine/Source/Runtime/Renderer/Private/PostProcess/PostProcessMaterial.h
@@ -4,6 +4,7 @@
 
 #include "ScreenPass.h"
 #include "OverridePassSequence.h"
+#include "Strata/Strata.h"
 
 class UMaterialInterface;
 
@@ -38,6 +39,7 @@ enum class EPostProcessMaterialInput : uint32
 BEGIN_SHADER_PARAMETER_STRUCT(FPostProcessMaterialParameters, )
 	SHADER_PARAMETER_STRUCT_REF(FViewUniformShaderParameters, View)
 	SHADER_PARAMETER_STRUCT_INCLUDE(FSceneTextureShaderParameters, SceneTextures)
+	SHADER_PARAMETER_RDG_UNIFORM_BUFFER(FStrataGlobalUniformParameters, Strata)
 	SHADER_PARAMETER_STRUCT(FScreenPassTextureViewportParameters, PostProcessOutput)
 	SHADER_PARAMETER_STRUCT_ARRAY(FScreenPassTextureInput, PostProcessInput, [kPostProcessMaterialInputCountMax])
 	SHADER_PARAMETER_SAMPLER(SamplerState, PostProcessInput_BilinearSampler)
diff --git a/Engine/Source/Runtime/Renderer/Private/Strata/Strata.h b/Engine/Source/Runtime/Renderer/Private/Strata/Strata.h
--- a/Engine/Source/Runtime/Renderer/Private/Strata/Strata.h
+++ b/Engine/Source/Runtime/Renderer/Private/Strata/Strata.h
@@ -26,7 +26,9 @@ BEGIN_SHADER_PARAMETER_STRUCT(FStrataForwardPassUniformParameters, )
 	SHADER_PARAMETER(uint32, bRoughDiffuse)
 END_SHADER_PARAMETER_STRUCT()
 
-BEGIN_GLOBAL_SHADER_PARAMETER_STRUCT(FStrataGlobalUniformParameters, )
+// This paramater struct is declared with RENDERER_API even though it is not public. This is
+// to workaround other modules doing 'private include' of the Renderer module
+BEGIN_GLOBAL_SHADER_PARAMETER_STRUCT(FStrataGlobalUniformParameters, RENDERER_API)
 	SHADER_PARAMETER(uint32, MaxBytesPerPixel)
 	SHADER_PARAMETER(uint32, bRoughDiffuse)
 	SHADER_PARAMETER_RDG_TEXTURE(Texture2DArray<uint>, MaterialTextureArray)
