Fix shader compile error for targets that don't use strata deferred or inline shading.

diff --git a/Engine/Shaders/Private/Strata/Strata.ush b/Engine/Shaders/Private/Strata/Strata.ush
--- a/Engine/Shaders/Private/Strata/Strata.ush
+++ b/Engine/Shaders/Private/Strata/Strata.ush
@@ -966,6 +966,8 @@ float3x3 StrataGetBSDFSharedBasis_InlineShading(in FStrataPixelHeader StrataPixe
 	float3 Tangent = StrataPixelHeader.SharedLocalBases.Tangents[BSDF_GETSHAREDLOCALBASISID(BSDF)];
 	float3 Bitangent = cross(Normal, Tangent);
 	return float3x3(Tangent, Bitangent, Normal);
+#else
+	return float3x3(float3(1,0,0), float3(0,1,0), float3(0,0,1));
 #endif
 }
 
@@ -997,6 +999,8 @@ float3x3 StrataGetBSDFSharedBasis_DeferredShading(in FStrataPixelHeader StrataPi
 		OutTangentBasis[2] = Normal;
 	}
 	return OutTangentBasis;
+#else
+	return float3x3(float3(1,0,0), float3(0,1,0), float3(0,0,1));
 #endif
 }
 
