- Ensure LocalVertexFactory sets particle color for VS
- Fix for Cascade mesh VF not reading particle color
- Fix for Cascade trail VF not handling custom interpolators

#rnx

diff --git a/Engine/Shaders/Private/LocalVertexFactory.ush b/Engine/Shaders/Private/LocalVertexFactory.ush
--- a/Engine/Shaders/Private/LocalVertexFactory.ush
+++ b/Engine/Shaders/Private/LocalVertexFactory.ush
@@ -650,6 +650,8 @@ FMaterialVertexParameters GetMaterialVertexParameters(FVertexFactoryInput Input,
 		#endif
 #endif  //MANUAL_VERTEX_FETCH && NUM_MATERIAL_TEXCOORDS_VERTEX
 
+	Result.Particle.Color = half4(1,1,1,1);
+
 #if NEEDS_PARTICLE_LOCAL_TO_WORLD
 	Result.Particle.ParticleToWorld = GetPrimitiveData(Intermediates).LocalToWorld;
 #endif
diff --git a/Engine/Shaders/Private/MeshParticleVertexFactory.ush b/Engine/Shaders/Private/MeshParticleVertexFactory.ush
--- a/Engine/Shaders/Private/MeshParticleVertexFactory.ush
+++ b/Engine/Shaders/Private/MeshParticleVertexFactory.ush
@@ -355,6 +355,7 @@ FMaterialVertexParameters GetMaterialVertexParameters(FVertexFactoryInput Input,
 	Result.InstanceLocalToWorld = LWCPromote(ParticleTransform);
 	Result.InstanceWorldToLocal = LWCPromoteInverse(PseudoInverseTransform(ParticleTransform));
 
+	Result.Particle.Color = Intermediates.ParticleColor;
 	Result.Particle.ParticleToWorld = Result.InstanceLocalToWorld;
 	Result.Particle.WorldToParticle = Result.InstanceWorldToLocal;
 
diff --git a/Engine/Shaders/Private/ParticleBeamTrailVertexFactory.ush b/Engine/Shaders/Private/ParticleBeamTrailVertexFactory.ush
--- a/Engine/Shaders/Private/ParticleBeamTrailVertexFactory.ush
+++ b/Engine/Shaders/Private/ParticleBeamTrailVertexFactory.ush
@@ -22,7 +22,9 @@ struct FVertexFactoryInterpolantsVSToPS
 {
 	TANGENTTOWORLD_INTERPOLATOR_BLOCK
 
-	float4	TexCoord			: TEXCOORD2;
+#if NUM_TEX_COORD_INTERPOLATORS
+	float4	TexCoords[(NUM_TEX_COORD_INTERPOLATORS+1)/2] : TEXCOORD2;
+#endif
 
 #if NEEDS_PARTICLE_COLOR
 	float4  Color				: TEXCOORD3;
@@ -56,21 +58,18 @@ FMaterialPixelParameters GetMaterialPixelParameters(FVertexFactoryInterpolantsVS
 	// GetMaterialPixelParameters is responsible for fully initializing the result
 	FMaterialPixelParameters Result = MakeInitializedMaterialPixelParameters();
 
-#if NUM_MATERIAL_TEXCOORDS
-	#if NUM_MATERIAL_TEXCOORDS >= 1
-		Result.TexCoords[0] = Interpolants.TexCoord.xy;
-	#if NUM_MATERIAL_TEXCOORDS >= 2
-		Result.TexCoords[1] = Interpolants.TexCoord.zw;
-	#endif	// >= 2
-	#endif	// >= 1
-	#if NUM_MATERIAL_TEXCOORDS > 2
-		UNROLL
-		for(int CoordinateIndex = 2;CoordinateIndex < NUM_MATERIAL_TEXCOORDS;CoordinateIndex++)
-		{
-			Result.TexCoords[CoordinateIndex] = Interpolants.TexCoord.xy;
-		}
-	#endif	// > 2
-#endif
+#if NUM_TEX_COORD_INTERPOLATORS
+	UNROLL
+	for (int UVIndex=0; UVIndex < NUM_TEX_COORD_INTERPOLATORS / 2; ++UVIndex)
+	{
+		Result.TexCoords[(UVIndex * 2) + 0] = Interpolants.TexCoords[UVIndex].xy;
+		Result.TexCoords[(UVIndex * 2) + 1] = Interpolants.TexCoords[UVIndex].zw;
+	}
+	#if (NUM_TEX_COORD_INTERPOLATORS & 1) == 1
+		Result.TexCoords[NUM_TEX_COORD_INTERPOLATORS - 1] = Interpolants.TexCoords[NUM_TEX_COORD_INTERPOLATORS / 2].xy;
+	#endif
+#endif // NUM_TEX_COORD_INTERPOLATORS
+
 	Result.VertexColor = 1;
 
 #if NEEDS_PARTICLE_COLOR
@@ -201,7 +200,23 @@ float3 VertexFactoryGetPositionForVertexLighting(FVertexFactoryInput Input, FVer
 FVertexFactoryInterpolantsVSToPS VertexFactoryGetInterpolantsVSToPS(FVertexFactoryInput Input, FVertexFactoryIntermediates Intermediates, FMaterialVertexParameters VertexParameters)
 {
 	FVertexFactoryInterpolantsVSToPS Interpolants;
-	Interpolants.TexCoord	= Input.TexCoord;
+
+#if NUM_TEX_COORD_INTERPOLATORS
+	float2 CustomizedUVs[NUM_TEX_COORD_INTERPOLATORS];
+	GetMaterialCustomizedUVs(VertexParameters, CustomizedUVs);
+	GetCustomInterpolators(VertexParameters, CustomizedUVs);
+
+	UNROLL
+	for (int UVIndex=0; UVIndex < NUM_TEX_COORD_INTERPOLATORS / 2; ++UVIndex)
+	{
+		Interpolants.TexCoords[UVIndex].xy = CustomizedUVs[UVIndex * 2 + 0];
+		Interpolants.TexCoords[UVIndex].zw = CustomizedUVs[UVIndex * 2 + 1];
+	}
+	#if (NUM_TEX_COORD_INTERPOLATORS & 1) == 1
+		Interpolants.TexCoords[NUM_TEX_COORD_INTERPOLATORS / 2].xy = CustomizedUVs[NUM_TEX_COORD_INTERPOLATORS - 1];
+		Interpolants.TexCoords[NUM_TEX_COORD_INTERPOLATORS / 2].zw = 0.0f;
+	#endif
+#endif	// NUM_MATERIAL_TEXCOORDS
 
 #if NEEDS_PARTICLE_COLOR
 	Interpolants.Color		= Intermediates.Color;
