- Fix for LightmapUV access

#rnx

diff --git a/Engine/Shaders/Private/ParticleBeamTrailVertexFactory.ush b/Engine/Shaders/Private/ParticleBeamTrailVertexFactory.ush
--- a/Engine/Shaders/Private/ParticleBeamTrailVertexFactory.ush
+++ b/Engine/Shaders/Private/ParticleBeamTrailVertexFactory.ush
@@ -82,7 +82,11 @@ FMaterialPixelParameters GetMaterialPixelParameters(FVertexFactoryInterpolantsVS
 	Result.UnMirrored = 1;
 
 #if LIGHTMAP_UV_ACCESS
-	Result.LightmapUVs = Interpolants.TexCoord.xy;
+	#if NUM_TEX_COORD_INTERPOLATORS > 0
+		Result.LightmapUVs = Interpolants.TexCoords[0].xy;
+	#else
+		Result.LightmapUVs = 0;
+	#endif
 #endif
 
 	Result.TwoSidedSign = 1;
