Make DF shadowing skip Nanite meshes if VSM is enabled (Nanite renders into VSM unconditionally causing double shadowing).
 - needed for VSM / DF integration

diff --git a/Engine/Shaders/Private/DistanceFieldLightingShared.ush b/Engine/Shaders/Private/DistanceFieldLightingShared.ush
--- a/Engine/Shaders/Private/DistanceFieldLightingShared.ush
+++ b/Engine/Shaders/Private/DistanceFieldLightingShared.ush
@@ -155,6 +155,7 @@ struct FDFObjectBounds
 	float3 BoxExtent;	// World space AABB extent
 	uint OftenMoving;
 	uint bCastShadow;
+	bool bIsNaniteMesh;
 };
 
 uint NumSceneObjects;
@@ -173,8 +174,9 @@ FDFObjectBounds LoadDFObjectBounds(uint ObjectIndex)
 	Bounds.BoxExtent = Vector1.xyz;
 
 	uint Flags = asuint(Vector1.w);
-	Bounds.OftenMoving = Flags & 1;
-	Bounds.bCastShadow = Flags & 2;
+	Bounds.OftenMoving = Flags & 1U;
+	Bounds.bCastShadow = Flags & 2U;
+	Bounds.bIsNaniteMesh = (Flags & 4U) != 0U;
 
 	return Bounds;
 }
diff --git a/Engine/Shaders/Private/DistanceFieldShadowing.usf b/Engine/Shaders/Private/DistanceFieldShadowing.usf
--- a/Engine/Shaders/Private/DistanceFieldShadowing.usf
+++ b/Engine/Shaders/Private/DistanceFieldShadowing.usf
@@ -85,6 +85,7 @@ void CopyCulledObjectData(uint DestIndex, uint SourceIndex)
 float4 ShadowConvexHull[12];
 float4 ShadowBoundingSphere;
 uint NumShadowHullPlanes;
+uint bDrawNaniteMeshes;
 
 bool ShadowConvexHullIntersectSphere(float3 SphereOrigin, float SphereRadius)
 {
@@ -151,7 +152,9 @@ void CullObjectsForShadowCS(
 		const bool bConvexHullIntersect = ShadowBoundingSphere.w == 0 && ShadowConvexHullIntersectSphere(DFObjectBounds.Center, DFObjectBounds.SphereRadius);
 		const bool bBoundingSphereIntersect = ShadowBoundingSphere.w > 0 && dot(ShadowBoundingSphere.xyz - DFObjectBounds.Center, ShadowBoundingSphere.xyz - DFObjectBounds.Center) < Square(ShadowBoundingSphere.w + DFObjectBounds.SphereRadius);
 
-        if (DFObjectBounds.bCastShadow && (bConvexHullIntersect || bBoundingSphereIntersect))
+        if (DFObjectBounds.bCastShadow 
+			&& (bConvexHullIntersect || bBoundingSphereIntersect)
+			&& (bDrawNaniteMeshes || !DFObjectBounds.bIsNaniteMesh))
 		{
 			FDFObjectData DFObjectData = LoadDFObjectData(ObjectIndex);
 
diff --git a/Engine/Source/Runtime/Renderer/Private/DistanceFieldObjectManagement.cpp b/Engine/Source/Runtime/Renderer/Private/DistanceFieldObjectManagement.cpp
--- a/Engine/Source/Runtime/Renderer/Private/DistanceFieldObjectManagement.cpp
+++ b/Engine/Source/Runtime/Renderer/Private/DistanceFieldObjectManagement.cpp
@@ -710,8 +710,9 @@ void FDistanceFieldSceneData::UpdateDistanceFieldObjectBuffers(
 									const FGlobalDFCacheType CacheType = PrimitiveSceneProxy->IsOftenMoving() ? GDF_Full : GDF_MostlyStatic;
 									const uint32 bOftenMoving = CacheType == GDF_Full;
 									const uint32 bCastShadow = PrimitiveSceneProxy->CastsDynamicShadow();
+									const uint32 bIsNaniteMesh = PrimitiveSceneProxy->IsNaniteMesh() ? 1U : 0U;
 
-									const uint32 Flags = bOftenMoving | (bCastShadow << 1);
+									const uint32 Flags = bOftenMoving | (bCastShadow << 1U) | (bIsNaniteMesh << 2U);
 
 									FVector4f ObjectWorldExtentAndFlags(WorldSpaceMeshBounds.GetExtent(), 0.0f);
 									ObjectWorldExtentAndFlags.W = *(const float*)&Flags;
diff --git a/Engine/Source/Runtime/Renderer/Private/DistanceFieldShadowing.cpp b/Engine/Source/Runtime/Renderer/Private/DistanceFieldShadowing.cpp
--- a/Engine/Source/Runtime/Renderer/Private/DistanceFieldShadowing.cpp
+++ b/Engine/Source/Runtime/Renderer/Private/DistanceFieldShadowing.cpp
@@ -164,6 +164,7 @@ class FCullObjectsForShadowCS : public FGlobalShader
 		SHADER_PARAMETER(FMatrix44f, WorldToShadow)
 		SHADER_PARAMETER(float, ObjectExpandScale)
 		SHADER_PARAMETER(uint32, NumShadowHullPlanes)
+		SHADER_PARAMETER(uint32, bDrawNaniteMeshes)
 		SHADER_PARAMETER(FVector4f, ShadowBoundingSphere)
 		SHADER_PARAMETER_ARRAY(FVector4f,ShadowConvexHull,[12])
 	END_SHADER_PARAMETER_STRUCT()
@@ -558,6 +559,7 @@ void CullDistanceFieldObjectsForLight(
 		PassParameters->ObjectExpandScale = bIsHeightfield ? 0.f : WorldToShadowValue.GetMaximumAxisScale();
 		PassParameters->NumShadowHullPlanes = NumPlanes;
 		PassParameters->ShadowBoundingSphere = ShadowBoundingSphereValue;
+		PassParameters->bDrawNaniteMeshes = !LightSceneProxy->UseVirtualShadowMaps();
 
 		check(NumPlanes <= 12);
 
