Remove unused WorldNormal computation from "one pass" point light shadow VS
In theory this could be DCE'd with cross-stage compilation (VS/GS), but probably better to be explicit than to rely on that.

diff --git a/Engine/Shaders/Private/ShadowDepthCommon.ush b/Engine/Shaders/Private/ShadowDepthCommon.ush
--- a/Engine/Shaders/Private/ShadowDepthCommon.ush
+++ b/Engine/Shaders/Private/ShadowDepthCommon.ush
@@ -36,7 +36,6 @@ struct FShadowDepthVSToPS
 struct FShadowDepthVSToGS : FShadowDepthVSToPS
 {
 	float4 GSPosition : TEXCOORD6;
-	float3 WorldNormal : TEXCOORD8;
 #if ENABLE_FALLBACK_POINTLIGHT_SHADOW_GS
 	nointerpolation uint ViewIndex : TEXCOORD9;
 #endif // ENABLE_FALLBACK_POINTLIGHT_SHADOW_GS
diff --git a/Engine/Shaders/Private/ShadowDepthVertexShader.usf b/Engine/Shaders/Private/ShadowDepthVertexShader.usf
--- a/Engine/Shaders/Private/ShadowDepthVertexShader.usf
+++ b/Engine/Shaders/Private/ShadowDepthVertexShader.usf
@@ -215,24 +215,7 @@ void MainOnePassPointLightGS(triangle FShadowDepthVSToGS Input[3], inout Triangl
 			// Calculate the clip space position for each cube face
 			// Note: We don't apply any constant/slop bias here. A constant bias is applied in the shadow projection 
 			// code (see CubemapHardwarePCF() in ShadowProjectionCommon.ush)
-			// If one want to use this code path, don't forget to remove the bias in CubemapHardwarePCF().
-			#if 0
-			const float4 WorldPosition = Input[VertexIndex].GSPosition;
-			const float3 WorldVertexNormal = Input[VertexIndex].WorldNormal;
-			
-			float4 ClipSpacePosition;
-			float ShadowDepth;
-			SetShadowDepthOutputs(
-				PassStruct.ShadowViewProjectionMatrices[CubeFaceIndex],
-				PassStruct.ShadowViewMatrices[CubeFaceIndex], 
-				WorldPosition, 
-				WorldVertexNormal, 
-				ClipSpacePosition, 
-				ShadowDepth);
-			ClipSpacePositions[VertexIndex] = ClipSpacePosition;
-			#else
 			ClipSpacePositions[VertexIndex] = mul(Input[VertexIndex].GSPosition, PassStruct.ShadowViewProjectionMatrices[CubeFaceIndex]);
-			#endif
 		}
 		// Frustum culling, saves GPU time with high poly meshes
 		BRANCH
@@ -303,7 +286,6 @@ void Main(
 
 	#if IS_FOR_GEOMETRY_SHADER
 		OutParameters.GSPosition = WorldPos;
-		OutParameters.WorldNormal = WorldNormal;
 		#if ENABLE_FALLBACK_POINTLIGHT_SHADOW_GS
 			OutParameters.ViewIndex = VertexFactoryGetViewIndex(VFIntermediates);
 		#endif // ENABLE_FALLBACK_POINTLIGHT_SHADOW_GS
@@ -445,7 +427,6 @@ void PositionOnlyMain(
 
 #if ONEPASS_POINTLIGHT_SHADOW
 	OutPosition = OutParameters.GSPosition = WorldPos;
-	OutParameters.WorldNormal = WorldNormal;
 #if ENABLE_FALLBACK_POINTLIGHT_SHADOW_GS
 	OutParameters.ViewIndex = VertexFactoryGetViewIndex(Input);
 #endif // !ENABLE_FALLBACK_POINTLIGHT_SHADOW_GS
