Adds support for alpha channel in TSR

diff --git a/Engine/Shaders/Private/TemporalSuperResolution/TSRCommon.ush b/Engine/Shaders/Private/TemporalSuperResolution/TSRCommon.ush
--- a/Engine/Shaders/Private/TemporalSuperResolution/TSRCommon.ush
+++ b/Engine/Shaders/Private/TemporalSuperResolution/TSRCommon.ush
@@ -36,10 +36,15 @@
 	#define CONFIG_BUTTERFLY_KERNEL 0
 #endif
 
-#define CONFIG_ENABLE_STOCASTIC_QUANTIZATION 1
-
 #define CONFIG_SCENE_COLOR_OVERFLOW 1
 
+#if POST_PROCESS_ALPHA
+	#define CONFIG_SCENE_COLOR_APLHA 1
+#else
+	#define CONFIG_SCENE_COLOR_APLHA 0
+#endif
+
+#define CONFIG_ENABLE_STOCASTIC_QUANTIZATION (!CONFIG_SCENE_COLOR_APLHA)
 
 //------------------------------------------------------- CONSTANTS
 
@@ -193,6 +198,17 @@ static const uint kPlusIndexes3x3[5] = { 4, 1, 3, 7, 5 };
 
 #endif
 
+#if CONFIG_SCENE_COLOR_APLHA
+	#define taa_halfC taa_half4
+	#define taa_halfCx2 taa_half4x2
+
+#else
+	#define taa_halfC taa_half3
+	#define taa_halfCx2 taa_half3x2
+
+#endif
+
+
 #define taa_subpixel_details uint
 #define taa_subpixel_details2 uint2
 #define taa_subpixel_payload uint // TODO: uint16_t
@@ -205,6 +221,7 @@ static const taa_half LargestNormalNumber = taa_half(65504.0);
 static const taa_half3 LargestSceneColorRGB = (LargestNormalNumber * taa_half(0.25)).xxx;
 static const taa_half4 LargestSceneColorRGBA = (LargestNormalNumber * taa_half(0.25)).xxxx;
 static const taa_half3 LargestSceneColorYCoCg = LargestNormalNumber.xxx;
+static const taa_half4 LargestSceneColorYCoCgA = LargestNormalNumber.xxxx;
 
 
 //------------------------------------------------------- PARAMETERS
@@ -642,6 +659,18 @@ taa_half3x2 YCoCgToRGB(taa_half3x2 YCoCg)
 	return RGB;
 }
 
+CALL_SITE_DEBUGLOC
+taa_half4x2 RGBToYCoCg(taa_half4x2 RGBA)
+{
+	return taa_half4x2(RGBToYCoCg(taa_half3x2(RGBA[0], RGBA[1], RGBA[2])), RGBA[3]);
+}
+
+CALL_SITE_DEBUGLOC
+taa_half4x2 YCoCgToRGB(taa_half4x2 YCoCgA)
+{
+	return taa_half4x2(YCoCgToRGB(taa_half3x2(YCoCgA[0], YCoCgA[1], YCoCgA[2])), YCoCgA[3]);
+}
+
 CALL_SITE_DEBUGLOC
 taa_half3 RGBToYCoCg(taa_half3 RGB)
 {
@@ -654,6 +683,18 @@ taa_half3 YCoCgToRGB(taa_half3 YCoCg)
 	return dpv_lo(YCoCgToRGB(dpv_interleave_mono_registers(YCoCg)));
 }
 
+CALL_SITE_DEBUGLOC
+taa_half4 RGBToYCoCg(taa_half4 RGBA)
+{
+	return dpv_lo(RGBToYCoCg(dpv_interleave_mono_registers(RGBA)));
+}
+
+CALL_SITE_DEBUGLOC
+taa_half4 YCoCgToRGB(taa_half4 YCoCgA)
+{
+	return dpv_lo(YCoCgToRGB(dpv_interleave_mono_registers(YCoCgA)));
+}
+
 
 // Some bright pixel can cause HdrWeight to get nullified under fp16 representation. So clamping this to a value close to the minimum float float positive value (0.000061).
 #define HDR_WEIGHT_SAFE_MIN_VALUE 0.0001
@@ -666,12 +707,24 @@ taa_half2 Luma4(taa_half3x2 Color)
 	return (Color[1] * taa_half(2.0)) + (Color[0] + Color[2]);
 }
 
+CALL_SITE_DEBUGLOC
+taa_half2 Luma4(taa_half4x2 Color)
+{
+	return Luma4(taa_half3x2(Color[0], Color[1], Color[2]));
+}
+
 CALL_SITE_DEBUGLOC
 taa_half Luma4(taa_half3 Color)
 {
 	return dpv_lo(Luma4(dpv_interleave_mono_registers(Color)));
 }
 
+CALL_SITE_DEBUGLOC
+taa_half Luma4(taa_half4 Color)
+{
+	return Luma4(Color.rgb);
+}
+
 CALL_SITE_DEBUGLOC
 taa_half2 HdrWeightY(taa_half2 Luma)
 {
@@ -705,12 +758,24 @@ taa_half2 HdrWeight4(taa_half3x2 Color)
 	return HdrWeightY(Luma4(Color));
 }
 
+CALL_SITE_DEBUGLOC
+taa_half2 HdrWeight4(taa_half4x2 Color)
+{
+	return HdrWeightY(Luma4(Color));
+}
+
 CALL_SITE_DEBUGLOC
 taa_half HdrWeight4(taa_half3 Color)
 {
 	return HdrWeightY(Luma4(Color));
 }
 
+CALL_SITE_DEBUGLOC
+taa_half HdrWeight4(taa_half4 Color)
+{
+	return HdrWeightY(Luma4(Color));
+}
+
 
 // Returns the weight of a pixels at a coordinate <PixelDelta> from the PDF highest point.
 CALL_SITE_DEBUGLOC
@@ -804,33 +869,45 @@ taa_half MeasureBackbufferLDRQuantizationErrorFromHDRLuma(taa_half Luma)
 }
 
 taa_half2 MeasureRejectionFactor(
-	taa_half3x2 PrevYCoCg,
-	taa_half3x2 ClampedPrevYCoCg,
-	taa_half3x2 InputCenterYCoCg,
-	taa_half3x2 InputMinYCoCg, taa_half3x2 InputMaxYCoCg,
+	taa_halfCx2 PrevYCoCg,
+	taa_halfCx2 ClampedPrevYCoCg,
+	taa_halfCx2 InputCenterYCoCg,
+	taa_halfCx2 InputMinYCoCg, taa_halfCx2 InputMaxYCoCg,
 	taa_half2 BackbufferQuantizationError)
 {
-	taa_half3x2 BoxSize = InputMaxYCoCg - InputMinYCoCg;
+	const taa_halfC Null = taa_half(0.0);
+	
+	#if CONFIG_SCENE_COLOR_APLHA
+		taa_halfCx2 BackbufferQuantizationErrorVector = taa_half4x2(BackbufferQuantizationError, BackbufferQuantizationError, BackbufferQuantizationError, BackbufferQuantizationError);
+	#else
+		taa_halfCx2 BackbufferQuantizationErrorVector = taa_half3x2(BackbufferQuantizationError, BackbufferQuantizationError, BackbufferQuantizationError);
+	#endif
 
-	taa_half3x2 ClampedEnergy = abs(ClampedPrevYCoCg - PrevYCoCg);
-	taa_half3x2 Delta = abs(InputCenterYCoCg - PrevYCoCg);
+	taa_halfCx2 BoxSize = InputMaxYCoCg - InputMinYCoCg;
 
-	taa_half3x2 ClampError = max(
+	taa_halfCx2 ClampedEnergy = abs(ClampedPrevYCoCg - PrevYCoCg);
+	taa_halfCx2 Delta = abs(InputCenterYCoCg - PrevYCoCg);
+
+	taa_halfCx2 ClampError = max(
 		BoxSize * taa_half(0.25),
-		taa_half3x2(BackbufferQuantizationError, BackbufferQuantizationError, BackbufferQuantizationError));
+		BackbufferQuantizationErrorVector);
 
-	taa_half3x2 Factor = taa_half(1.0) - saturate(max(ClampedEnergy - ClampError, taa_half3x2(0.0, 0.0, 0.0, 0.0, 0.0, 0.0)) * rcp(max(Delta, BoxSize)));
+	taa_halfCx2 Factor = taa_half(1.0) - saturate(max(ClampedEnergy - ClampError, dpv_interleave_mono_registers(Null)) * rcp(max(Delta, BoxSize)));
 
 	taa_half2 Rejection = min3(Factor[0], Factor[1], Factor[2]);
 	
+	#if CONFIG_SCENE_COLOR_APLHA
+		Rejection = min(Rejection, Factor[3]);
+	#endif
+
 	return Rejection;
 }
 
 taa_half MeasureRejectionFactor(
-	taa_half3 PrevYCoCg,
-	taa_half3 ClampedPrevYCoCg,
-	taa_half3 InputCenterYCoCg,
-	taa_half3 InputMinYCoCg, taa_half3 InputMaxYCoCg,
+	taa_halfC PrevYCoCg,
+	taa_halfC ClampedPrevYCoCg,
+	taa_halfC InputCenterYCoCg,
+	taa_halfC InputMinYCoCg, taa_halfC InputMaxYCoCg,
 	taa_half BackbufferQuantizationError)
 {
 	return dpv_lo(MeasureRejectionFactor(
diff --git a/Engine/Shaders/Private/TemporalSuperResolution/TSRCompareHistory.usf b/Engine/Shaders/Private/TemporalSuperResolution/TSRCompareHistory.usf
--- a/Engine/Shaders/Private/TemporalSuperResolution/TSRCompareHistory.usf
+++ b/Engine/Shaders/Private/TemporalSuperResolution/TSRCompareHistory.usf
@@ -13,8 +13,8 @@
 //------------------------------------------------------- PARAMETERS
 
 Texture2D<taa_half> ParallaxRejectionMaskTexture;
-Texture2D<taa_half3> FilteredInputTexture;
-Texture2D<taa_half3> FilteredPredictionSceneColorTexture;
+Texture2D<taa_halfC> FilteredInputTexture;
+Texture2D<taa_halfC> FilteredPredictionSceneColorTexture;
 
 RWTexture2D<taa_half> HistoryRejectionOutput;
 
@@ -26,7 +26,7 @@ groupshared taa_half SharedRejection[TILE_SIZE * TILE_SIZE];
 
 //------------------------------------------------------- FUNCTIONS
 
-taa_half3 TransformColorForClampingBox(taa_half3 Color)
+taa_halfC TransformColorForClampingBox(taa_halfC Color)
 {
 	return RGBToYCoCg(Color);
 }
@@ -46,9 +46,9 @@ void MainCS(
 		taa_short2(GroupId) * taa_short2(TILE_SIZE, TILE_SIZE) +
 		Map16x16Tile2x2Lane(GroupThreadIndex));
 
-	taa_half3 InputMinColor;
-	taa_half3 InputMaxColor;
-	taa_half3 InputCenterColor;
+	taa_halfC InputMinColor;
+	taa_halfC InputMaxColor;
+	taa_halfC InputCenterColor;
 	taa_half ParallaxRejectionMask;
 
 	{
@@ -74,19 +74,19 @@ void MainCS(
 			kOffsets3x3[i], kOffsets3x3[i],
 			InputPixelPosMin, InputPixelPosMax);
 
-		taa_half3 Color = TransformColorForClampingBox(FilteredInputTexture[SampleInputPixelPos]);
+		taa_halfC Color = TransformColorForClampingBox(FilteredInputTexture[SampleInputPixelPos]);
 
 		InputMinColor = min(InputMinColor, Color);
 		InputMaxColor = max(InputMaxColor, Color);
 	}
 	
 
-	taa_half3 PrevColor = TransformColorForClampingBox(FilteredPredictionSceneColorTexture[InputPixelPos]);
+	taa_halfC PrevColor = TransformColorForClampingBox(FilteredPredictionSceneColorTexture[InputPixelPos]);
 
 	taa_half Rejection = 1.0;
 	#if 1
 	{
-		taa_half3 ClampedPrevColor = clamp(PrevColor, InputMinColor, InputMaxColor);
+		taa_halfC ClampedPrevColor = clamp(PrevColor, InputMinColor, InputMaxColor);
 
 		Rejection = MeasureRejectionFactor(
 			PrevColor, ClampedPrevColor,
diff --git a/Engine/Shaders/Private/TemporalSuperResolution/TSRCompareTranslucency.usf b/Engine/Shaders/Private/TemporalSuperResolution/TSRCompareTranslucency.usf
--- a/Engine/Shaders/Private/TemporalSuperResolution/TSRCompareTranslucency.usf
+++ b/Engine/Shaders/Private/TemporalSuperResolution/TSRCompareTranslucency.usf
@@ -78,11 +78,18 @@ taa_half MeasureTranslucencyRejectionFactor(
 	taa_half4 ClampedPrevTranslucency = clamp(PrevTranslucency, TranslucencyMin, TranslucencyMax);
 	
 	taa_half BackbufferQuantizationError = 0.0; //MeasureBackbufferLDRQuantizationErrorFromHDRLuma(max(PrevTranslucency[0], Translucency[0]));
-
-	taa_half Rejection = MeasureRejectionFactor(
-		PrevTranslucency.xyz, ClampedPrevTranslucency.xyz,
-		Translucency.xyz, TranslucencyMin.xyz, TranslucencyMax.xyz,
-		BackbufferQuantizationError);
+	
+	#if CONFIG_SCENE_COLOR_APLHA
+		taa_half Rejection = MeasureRejectionFactor(
+			PrevTranslucency, ClampedPrevTranslucency,
+			Translucency, TranslucencyMin, TranslucencyMax,
+			BackbufferQuantizationError);
+	#else
+		taa_half Rejection = MeasureRejectionFactor(
+			PrevTranslucency.xyz, ClampedPrevTranslucency.xyz,
+			Translucency.xyz, TranslucencyMin.xyz, TranslucencyMax.xyz,
+			BackbufferQuantizationError);
+	#endif
 
 	Rejection = lerp(1.0, Rejection, IsHighlight);
 
diff --git a/Engine/Shaders/Private/TemporalSuperResolution/TSRDecimateHistory.usf b/Engine/Shaders/Private/TemporalSuperResolution/TSRDecimateHistory.usf
--- a/Engine/Shaders/Private/TemporalSuperResolution/TSRDecimateHistory.usf
+++ b/Engine/Shaders/Private/TemporalSuperResolution/TSRDecimateHistory.usf
@@ -17,31 +17,37 @@
 float4x4 RotationalClipToPrevClip;
 float3 OutputQuantizationError;
 
-Texture2D<taa_half3> InputSceneColorTexture;
+Texture2D<taa_halfC> InputSceneColorTexture;
 Texture2D<float2> DilatedVelocityTexture;
 Texture2D<float> ClosestDepthTexture;
 Texture2D<uint> PrevUseCountTexture;
 Texture2D<uint> PrevClosestDepthTexture;
 Texture2D<uint> ParallaxFactorTexture;
 
-Texture2D<taa_half3> PrevHistory_LowFrequency;
-Texture2D<taa_half3> PrevHistory_HighFrequency;
+Texture2D<taa_halfC> PrevHistory_LowFrequency;
+Texture2D<taa_halfC> PrevHistory_HighFrequency;
 Texture2D<taa_half2> PrevHistory_Metadata;
 Texture2D<taa_half4> PrevHistory_Translucency;
 Texture2D<taa_subpixel_details> PrevHistory_SubpixelDetails;
 
-RWTexture2D<taa_half3> HalfResSceneColorOutput;
-RWTexture2D<taa_half3> HalfResPredictionSceneColorOutput;
+RWTexture2D<taa_halfC> HalfResSceneColorOutput;
+RWTexture2D<taa_halfC> HalfResPredictionSceneColorOutput;
 RWTexture2D<taa_half> HalfResParallaxRejectionMaskOutput;
 
-RWTexture2D<taa_half3> PredictionSceneColorOutput;
+RWTexture2D<taa_halfC> PredictionSceneColorOutput;
 RWTexture2D<taa_half> ParallaxRejectionMaskOutput;
 
 
 //------------------------------------------------------- LDS
 
-groupshared taa_half3 SharedInputColor[TILE_SIZE * TILE_SIZE];
+groupshared taa_halfC SharedInputColor[TILE_SIZE * TILE_SIZE];
+
+#if CONFIG_SCENE_COLOR_APLHA
+groupshared taa_halfC SharedHistoryColor[TILE_SIZE * TILE_SIZE];
+groupshared taa_half  SharedMask[TILE_SIZE * TILE_SIZE];
+#else
 groupshared taa_half4 SharedHistoryColorAndMask[TILE_SIZE * TILE_SIZE];
+#endif
 
 
 //------------------------------------------------------- ENTRY POINT
@@ -76,7 +82,7 @@ void MainCS(
 	float2 EncodedVelocity = DilatedVelocityTexture[InputPixelPos];
 	float DeviceZ = ClosestDepthTexture[InputPixelPos];
 	taa_half ParallaxFactor = UncompressParallaxFactor((ParallaxFactorTexture[InputPixelPos] >> 4) & 0xF);
-	taa_half3 InputSceneColor = InputSceneColorTexture[InputPixelPos];
+	taa_halfC InputSceneColor = InputSceneColorTexture[InputPixelPos];
 
 	float WorldDepth = ConvertFromDeviceZ(DeviceZ);
 	
@@ -148,7 +154,7 @@ void MainCS(
 	taa_half2 dKO = taa_half2(HistoryPPCo - HistoryPPCk);
 	
 	// Sample the history
-	taa_half3 DecimatedColor;
+	taa_halfC DecimatedColor;
 	taa_half HistoryMaxRejectionFactor;
 
 	taa_half ClosestInterferenceSeed;
@@ -181,7 +187,7 @@ void MainCS(
 					PrevHistoryInfo_UVViewportBilinearMin,
 					PrevHistoryInfo_UVViewportBilinearMax);
 				
-				taa_half3 HistorySample = PrevHistory_LowFrequency.SampleLevel(GlobalBilinearClampedSampler, SampleBufferUV, 0);
+				taa_halfC HistorySample = PrevHistory_LowFrequency.SampleLevel(GlobalBilinearClampedSampler, SampleBufferUV, 0);
 
 				DecimatedColor += HistorySample;
 			}
@@ -245,15 +251,20 @@ void MainCS(
 	#if CONFIG_OUTPUT_HALF_RES
 	{
 		SharedInputColor[GroupThreadIndex] = InputSceneColor;
-		SharedHistoryColorAndMask[GroupThreadIndex] = taa_half4(DecimatedColor, ParallaxRejectionMask);
+		#if CONFIG_SCENE_COLOR_APLHA
+			SharedHistoryColor[GroupThreadIndex] = DecimatedColor;
+			SharedMask[GroupThreadIndex] = ParallaxRejectionMask;
+		#else
+			SharedHistoryColorAndMask[GroupThreadIndex] = taa_half4(DecimatedColor, ParallaxRejectionMask);
+		#endif
 
 		GroupMemoryBarrierWithGroupSync();
 
 		bool bCheckerboard = (((InputPixelPos.x ^ InputPixelPos.y) >> 1) & 0x1) != 0;
 
-		half3 HalfResInputSceneColor = InputSceneColor;
-		half3 HalfResDecimatedColor = DecimatedColor;
-		half HalfResParallaxRejectionMask = ParallaxRejectionMask;
+		taa_halfC HalfResInputSceneColor = InputSceneColor;
+		taa_halfC HalfResDecimatedColor = DecimatedColor;
+		taa_half HalfResParallaxRejectionMask = ParallaxRejectionMask;
 
 		UNROLL
 		for (uint i = 1; i < 4; i++)
@@ -268,13 +279,18 @@ void MainCS(
 			{
 				HalfResInputSceneColor = max(HalfResInputSceneColor, SharedInputColor[GroupThreadIndex ^ ButterflySwap]);
 			}
-			HalfResDecimatedColor += SharedHistoryColorAndMask[GroupThreadIndex ^ ButterflySwap].rgb;
-
-			HalfResParallaxRejectionMask += SharedHistoryColorAndMask[GroupThreadIndex ^ ButterflySwap].a;
+			
+			#if CONFIG_SCENE_COLOR_APLHA
+				HalfResDecimatedColor += SharedHistoryColor[GroupThreadIndex ^ ButterflySwap];
+				HalfResParallaxRejectionMask += SharedMask[GroupThreadIndex ^ ButterflySwap];
+			#else
+				HalfResDecimatedColor += SharedHistoryColorAndMask[GroupThreadIndex ^ ButterflySwap].rgb;
+				HalfResParallaxRejectionMask += SharedHistoryColorAndMask[GroupThreadIndex ^ ButterflySwap].a;
+			#endif
 		}
 
-		HalfResDecimatedColor *= 0.25;
-		HalfResParallaxRejectionMask *= 0.25;
+		HalfResDecimatedColor *= taa_half(0.25);
+		HalfResParallaxRejectionMask *= taa_half(0.25);
 
 		#if CONFIG_ENABLE_STOCASTIC_QUANTIZATION
 		{
diff --git a/Engine/Shaders/Private/TemporalSuperResolution/TSRFilterFrequencies.usf b/Engine/Shaders/Private/TemporalSuperResolution/TSRFilterFrequencies.usf
--- a/Engine/Shaders/Private/TemporalSuperResolution/TSRFilterFrequencies.usf
+++ b/Engine/Shaders/Private/TemporalSuperResolution/TSRFilterFrequencies.usf
@@ -32,18 +32,18 @@
 
 float3 OutputQuantizationError;
 
-Texture2D<taa_half3> InputTexture;
-Texture2D<taa_half3> PredictionSceneColorTexture;
+Texture2D<taa_halfC> InputTexture;
+Texture2D<taa_halfC> PredictionSceneColorTexture;
 Texture2D<taa_half> ParallaxRejectionMaskTexture;
 
-RWTexture2D<taa_half3> FilteredInputOutput;
-RWTexture2D<taa_half3> FilteredPredictionSceneColorOutput;
+RWTexture2D<taa_halfC> FilteredInputOutput;
+RWTexture2D<taa_halfC> FilteredPredictionSceneColorOutput;
 RWTexture2D<taa_half>  InputSceneColorLdrLumaOutput;
 
 
 //------------------------------------------------------- FUNCTIONS
 
-taa_half ComputeColorWeight(taa_half3 Color)
+taa_half ComputeColorWeight(taa_halfC Color)
 #if CONFIG_ACCUMULATE_KARIS_WEIGHTING
 {
 	return HdrWeight4(Color);
@@ -76,19 +76,17 @@ void MainCS(
 		Map16x16Tile2x2Lane(GroupThreadIndex));
 
 	// Fetch center of the kernel
-	taa_half3 CenterInput;
-	taa_half3 FilteredInput;
+	taa_halfC CenterInput;
+	taa_halfC FilteredInput;
 	taa_half FilteredInputWeight;
-	taa_half3 FilteredHistory;
+	taa_halfC FilteredHistory;
 	taa_half FilteredHistoryWeight;
-	taa_half4 FilteredInterferenceSeed;
-	taa_half FilteredInterferenceSeedWeight;
 
 	{
 		taa_half SampleSpatialWeight = 1.0;
 
-		taa_half3 InputSample = InputTexture[InputPixelPos];
-		taa_half3 HistorySample = PredictionSceneColorTexture[InputPixelPos];
+		taa_halfC InputSample = InputTexture[InputPixelPos];
+		taa_halfC HistorySample = PredictionSceneColorTexture[InputPixelPos];
 		taa_half HistorySampleValidity = ParallaxRejectionMaskToHistoryValidity(
 			ParallaxRejectionMaskTexture[InputPixelPos]);
 		
@@ -175,13 +173,17 @@ void MainCS(
 			DynamicSpreadFactor * Offset, Offset,
 			InputPixelPosMin, InputPixelPosMax);
 
-		taa_half3 InputSample = InputTexture[SampleInputPixelPos];
-		taa_half3 HistorySample = PredictionSceneColorTexture[SampleInputPixelPos];
+		taa_halfC InputSample = InputTexture[SampleInputPixelPos];
+		taa_halfC HistorySample = PredictionSceneColorTexture[SampleInputPixelPos];
 		taa_half HistorySampleValidity = ParallaxRejectionMaskToHistoryValidity(
 			ParallaxRejectionMaskTexture[SampleInputPixelPos]);
 
 		#if CONFIG_SCENE_COLOR_OVERFLOW
-			InputSample = min(InputSample, LargestSceneColorRGB);
+			#if CONFIG_SCENE_COLOR_APLHA
+				InputSample = min(InputSample, LargestSceneColorRGBA);
+			#else
+				InputSample = min(InputSample, LargestSceneColorRGB);
+			#endif
 		#endif
 
 		taa_half InputWeight = HistorySampleValidity * ComputeColorWeight(InputSample);
@@ -212,7 +214,7 @@ void MainCS(
 			taa_half CenterLuma = max(CenterInput.r, max(CenterInput.g, CenterInput.b));
 			CenterLumaLDR = CenterLuma / (SpatialAAExposure + CenterLuma);
 		#else
-			taa_half CenterLuma = dot(CenterInput, half3(0.299f, 0.587f, 0.114f));
+			taa_half CenterLuma = dot(CenterInput.rgb, half3(0.299f, 0.587f, 0.114f));
 			CenterLumaLDR = CenterLuma / (SpatialAAExposure + CenterLuma);
 		#endif
 	}
diff --git a/Engine/Shaders/Private/TemporalSuperResolution/TSRResolveHistory.usf b/Engine/Shaders/Private/TemporalSuperResolution/TSRResolveHistory.usf
--- a/Engine/Shaders/Private/TemporalSuperResolution/TSRResolveHistory.usf
+++ b/Engine/Shaders/Private/TemporalSuperResolution/TSRResolveHistory.usf
@@ -28,7 +28,7 @@ RWTexture2D<uint> PrevClosestDepthOutput;
 
 //------------------------------------------------------- LDS
 
-groupshared taa_half3 SharedArray0[TILE_SIZE * TILE_SIZE];
+groupshared taa_halfC SharedArray0[TILE_SIZE * TILE_SIZE];
 
 
 //------------------------------------------------------- FUNCTIONS
@@ -83,14 +83,14 @@ uint2 OutputViewRectMax;
 uint bGenerateOutputMip1;
 float HistoryValidityMultiply;
 
-Texture2D<taa_half3> History_LowFrequency;
-Texture2D<taa_half3> History_HighFrequency;
+Texture2D<taa_halfC> History_LowFrequency;
+Texture2D<taa_halfC> History_HighFrequency;
 Texture2D<taa_half2> History_Metadata;
 Texture2D<taa_half3> History_Translucency;
 Texture2D<taa_half>  History_TranslucencyAlpha;
 
-RWTexture2D<taa_half3> SceneColorOutputMip0;
-RWTexture2D<taa_half3> SceneColorOutputMip1;
+RWTexture2D<taa_halfC> SceneColorOutputMip0;
+RWTexture2D<taa_halfC> SceneColorOutputMip1;
 
 
 //------------------------------------------------------- ENTRY POINT
@@ -106,8 +106,8 @@ void MainCS(
 		ZOrder2D(GroupThreadIndex, log2(TILE_SIZE)) +
 		GroupId * uint2(TILE_SIZE, TILE_SIZE));
 
-	taa_half3 LowFrequencyColor = 0.0;
-	taa_half3 HighFrequencyColor = 0.0;
+	taa_halfC LowFrequencyColor = 0.0;
+	taa_halfC HighFrequencyColor = 0.0;
 	taa_half  HistoryValidity = 0.0;
 	
 	// Downsample the history
@@ -135,19 +135,25 @@ void MainCS(
 
 				SampleUV = clamp(SampleUV, HistoryInfo_UVViewportBilinearMin, HistoryInfo_UVViewportBilinearMax);
 
-				taa_half3 SampleLowFrequencyColor  = History_LowFrequency.SampleLevel(GlobalBilinearClampedSampler, SampleUV, 0);
-				taa_half3 SampleHighFrequencyColor = History_HighFrequency.SampleLevel(GlobalBilinearClampedSampler, SampleUV, 0);
+				taa_halfC SampleLowFrequencyColor  = History_LowFrequency.SampleLevel(GlobalBilinearClampedSampler, SampleUV, 0);
+				taa_halfC SampleHighFrequencyColor = History_HighFrequency.SampleLevel(GlobalBilinearClampedSampler, SampleUV, 0);
 				taa_half2 SampleMetadata   	       = History_Metadata.SampleLevel(GlobalBilinearClampedSampler, SampleUV, 0);
 				
 				taa_half3 SampleTranslucencyColor = History_Translucency.SampleLevel(GlobalBilinearClampedSampler, SampleUV, 0);
 				taa_half  SampleTranslucencyAlpha = History_TranslucencyAlpha.SampleLevel(GlobalBilinearClampedSampler, SampleUV, 0);
 				
 				taa_half SampleHistoryValidity = SampleMetadata.y;
+				taa_half InvSampleHistoryValidity = SafeRcp(SampleHistoryValidity);
 
 				#if 1
 				{
-					SampleLowFrequencyColor  = SampleLowFrequencyColor  * SampleTranslucencyAlpha + SampleTranslucencyColor;
-					SampleHighFrequencyColor = SampleHighFrequencyColor * SampleTranslucencyAlpha + SampleTranslucencyColor * SampleHistoryValidity;
+					SampleLowFrequencyColor.rgb  = SampleLowFrequencyColor.rgb  * SampleTranslucencyAlpha + SampleTranslucencyColor.rgb;
+					SampleHighFrequencyColor.rgb = SampleHighFrequencyColor.rgb * SampleTranslucencyAlpha + SampleTranslucencyColor.rgb * SampleHistoryValidity;
+				
+					#if CONFIG_SCENE_COLOR_APLHA
+						SampleLowFrequencyColor.a = saturate(taa_half(1.0) - saturate(taa_half(1.0) - SampleLowFrequencyColor.a) * SampleTranslucencyAlpha);
+						SampleHighFrequencyColor.a = saturate(taa_half(1.0) - saturate(taa_half(1.0) - SampleHighFrequencyColor.a * InvSampleHistoryValidity) * SampleTranslucencyAlpha) * SampleHistoryValidity;
+					#endif
 				}
 				#endif
 
@@ -156,7 +162,7 @@ void MainCS(
 				float kernelWeight = MitchellNetravali(PixelOffset.x, B, C) * MitchellNetravali(PixelOffset.y, B, C);
 				
 				float LowFrequencyHdrWeight  = HdrWeight4(SampleLowFrequencyColor);
-				float HighFrequencyHdrWeight = HdrWeight4(SampleHighFrequencyColor * SafeRcp(SampleHistoryValidity));
+				float HighFrequencyHdrWeight = HdrWeight4(SampleHighFrequencyColor * InvSampleHistoryValidity);
 
 				LowFrequencyColor   += (kernelWeight * LowFrequencyHdrWeight)  * SampleLowFrequencyColor;
 				LowFrequencyWeight  += (kernelWeight * LowFrequencyHdrWeight);
@@ -175,7 +181,7 @@ void MainCS(
 	}
 	
 	// Compute final output
-	taa_half3 FinalOutputColor;
+	taa_halfC FinalOutputColor;
 	{
 		taa_half OutputBlend = ComputeLowFrequencyContribution(HistoryValidity);
 		
@@ -198,7 +204,7 @@ void MainCS(
 	{
 		uint LocalGroupThreadIndex = GroupThreadIndex;
 
-		taa_half3 HalfResOutput = FinalOutputColor * taa_half(0.25);
+		taa_halfC HalfResOutput = FinalOutputColor * taa_half(0.25);
 		taa_short2 HalfResOutputPixelPos;
 		HalfResOutputPixelPos[0] = (OutputPixelPosition[0] >> taa_short(1)) | (((OutputPixelPosition[0] & taa_short(0x1)) | !bGenerateOutputMip1) * taa_short(~0));
 		HalfResOutputPixelPos[1] = (OutputPixelPosition[1] >> taa_short(1)) | (((OutputPixelPosition[1] & taa_short(0x1)) | !bGenerateOutputMip1) * taa_short(~0));
diff --git a/Engine/Shaders/Private/TemporalSuperResolution/TSRUpdateHistory.usf b/Engine/Shaders/Private/TemporalSuperResolution/TSRUpdateHistory.usf
--- a/Engine/Shaders/Private/TemporalSuperResolution/TSRUpdateHistory.usf
+++ b/Engine/Shaders/Private/TemporalSuperResolution/TSRUpdateHistory.usf
@@ -58,7 +58,11 @@
 
 #define CONFIG_RESPONSIVE_STENCIL 1
 
-#define CONFIG_MANUAL_LDS_SPILL (CONFIG_COMPILE_FP16)
+#if CONFIG_COMPILE_FP16 && !CONFIG_SCENE_COLOR_APLHA
+	#define CONFIG_MANUAL_LDS_SPILL 1
+#else
+	#define CONFIG_MANUAL_LDS_SPILL 0
+#endif
 
 #if CONFIG_COMPILE_FP16
 	// Take advantage of RDNA's v_pk_*_{uif}16 instructions
@@ -79,7 +83,7 @@
 
 //------------------------------------------------------- PARAMETERS
 
-Texture2D<taa_half3> InputSceneColorTexture;
+Texture2D<taa_halfC> InputSceneColorTexture;
 Texture2D<uint2> InputSceneStencilTexture;
 Texture2D<taa_half4> InputSceneTranslucencyTexture;
 
@@ -102,22 +106,22 @@ float InputToHistoryFactor;
 uint ResponsiveStencilMask;
 uint bGenerateOutputMip1;
 
-Texture2D<taa_half3> PrevHistory_LowFrequency;
-Texture2D<taa_half3> PrevHistory_HighFrequency;
+Texture2D<taa_halfC> PrevHistory_LowFrequency;
+Texture2D<taa_halfC> PrevHistory_HighFrequency;
 Texture2D<taa_half2> PrevHistory_Metadata;
 Texture2D<taa_half3> PrevHistory_Translucency;
 Texture2D<taa_half>  PrevHistory_TranslucencyAlpha;
 Texture2D<taa_subpixel_details> PrevHistory_SubpixelDetails;
 
-RWTexture2D<taa_half3> HistoryOutput_LowFrequency;
-RWTexture2D<taa_half3> HistoryOutput_HighFrequency;
+RWTexture2D<taa_halfC> HistoryOutput_LowFrequency;
+RWTexture2D<taa_halfC> HistoryOutput_HighFrequency;
 RWTexture2D<taa_half2> HistoryOutput_Metadata;
 RWTexture2D<taa_half3> HistoryOutput_Translucency;
 RWTexture2D<taa_half>  HistoryOutput_TranslucencyAlpha;
 RWTexture2D<taa_subpixel_details> HistoryOutput_SubpixelDetails;
 
-RWTexture2D<taa_half3> SceneColorOutputMip0;
-RWTexture2D<taa_half3> SceneColorOutputMip1;
+RWTexture2D<taa_halfC> SceneColorOutputMip0;
+RWTexture2D<taa_halfC> SceneColorOutputMip1;
 
 
 //------------------------------------------------------- LDS
@@ -159,7 +163,7 @@ taa_ushort2x2 Map8x8Tile2x2LaneDPV(uint GroupThreadIndex)
 }
 
 CALL_SITE_DEBUGLOC 
-taa_half3x2 TransformColorForClampingBox(taa_half3x2 Color)
+taa_halfCx2 TransformColorForClampingBox(taa_halfCx2 Color)
 {
 	return RGBToYCoCg(Color);
 }
@@ -430,7 +434,7 @@ void MainCS(
 			UNROLL_N(CONFIG_SAMPLES - 1)
 			for (uint NeighborId = 0; NeighborId < (CONFIG_SAMPLES - 1); NeighborId++)
 			{
-				TranslucencyRejection = min(TranslucencyRejection, dpv_interleave_registers_array(RawTranslucencyRejectionNeighborhood[NeighborId]));
+				TranslucencyRejection = min(TranslucencyRejection, dpv_force_interleave_registers_array(RawTranslucencyRejectionNeighborhood[NeighborId]));
 			
 				SceneStencilRef |= dpv_interleave_registers_array(RawSceneStencilRefNeighborhood[NeighborId]);
 			}
@@ -505,9 +509,9 @@ void MainCS(
 
 
 	// Filter input scene color at predictor frequency.
-	taa_half3x2 FilteredInputColor;
-	taa_half3x2 InputMinColor;
-	taa_half3x2 InputMaxColor;
+	taa_halfCx2 FilteredInputColor;
+	taa_halfCx2 InputMinColor;
+	taa_halfCx2 InputMaxColor;
 	
 	taa_half4x2 FilteredTranslucencyColor;
 	taa_half4x2 TranslucencyMinColor;
@@ -538,7 +542,7 @@ void MainCS(
 			bKOSign[1] = dKO[1] < 0.0;
 		}
 		
-		taa_half3 RawInputColorArray[CONFIG_SAMPLES][DPV_PIXEL_PER_LANE];
+		taa_halfC RawInputColorArray[CONFIG_SAMPLES][DPV_PIXEL_PER_LANE];
 		taa_half4 RawTranslucencyColorArray[CONFIG_SAMPLES][DPV_PIXEL_PER_LANE];
 
 		taa_half2x2 RawdPPArray[CONFIG_SAMPLES];
@@ -594,8 +598,13 @@ void MainCS(
 		taa_half2 FilteredInputColorWeight = taa_half(0.0);
 		
 		FilteredInputColor = taa_half(0.0);
-		InputMinColor = dpv_interleave_mono_registers(+LargestSceneColorYCoCg);
-		InputMaxColor = dpv_interleave_mono_registers(-LargestSceneColorYCoCg);
+		#if CONFIG_SCENE_COLOR_APLHA
+			InputMinColor = dpv_interleave_mono_registers(+LargestSceneColorYCoCgA);
+			InputMaxColor = dpv_interleave_mono_registers(-LargestSceneColorYCoCgA);
+		#else
+			InputMinColor = dpv_interleave_mono_registers(+LargestSceneColorYCoCg);
+			InputMaxColor = dpv_interleave_mono_registers(-LargestSceneColorYCoCg);
+		#endif
 		TranslucencyMinColor = dpv_interleave_mono_registers(+LargestSceneColorRGBA);
 		TranslucencyMaxColor = dpv_interleave_mono_registers(-LargestSceneColorRGBA);
 
@@ -603,11 +612,15 @@ void MainCS(
 		for (uint SampleId = 0; SampleId < CONFIG_SAMPLES; SampleId++)
 		{
 			taa_half2 SampleSpatialWeight = RawSampleSpatialWeightArray[SampleId];
-			taa_half3x2 InputColor = dpv_force_interleave_registers_array(RawInputColorArray[SampleId]);
+			taa_halfCx2 InputColor = dpv_force_interleave_registers_array(RawInputColorArray[SampleId]);
 			taa_half4x2 TranslucencyColor = dpv_force_interleave_registers_array(RawTranslucencyColorArray[SampleId]);
 
 			#if CONFIG_SCENE_COLOR_OVERFLOW
-				InputColor = min(InputColor, dpv_interleave_mono_registers(LargestSceneColorRGB));
+				#if CONFIG_SCENE_COLOR_APLHA
+					InputColor = min(InputColor, dpv_interleave_mono_registers(LargestSceneColorRGBA));
+				#else
+					InputColor = min(InputColor, dpv_interleave_mono_registers(LargestSceneColorRGB));
+				#endif
 			#endif
 
 			taa_half2 ToneWeight = HdrWeight4(InputColor);
@@ -650,7 +663,11 @@ void MainCS(
 		
 		#if CONFIG_SCENE_COLOR_OVERFLOW
 		{
-			FilteredInputColor = min(FilteredInputColor, dpv_interleave_mono_registers(LargestSceneColorRGB));
+			#if CONFIG_SCENE_COLOR_APLHA
+				FilteredInputColor = min(FilteredInputColor, dpv_interleave_mono_registers(LargestSceneColorRGBA));
+			#else
+				FilteredInputColor = min(FilteredInputColor, dpv_interleave_mono_registers(LargestSceneColorRGB));
+			#endif
 			ClosestInputLuma4 = min(ClosestInputLuma4, LargestSceneColorYCoCg.x);
 			FilteredTranslucencyColor = min(FilteredTranslucencyColor, dpv_interleave_mono_registers(LargestSceneColorRGBA));
 		}
@@ -690,14 +707,14 @@ void MainCS(
 	#endif
 	
 	// Reproject history
-	taa_half3x2 PrevLowFrequencyColor;
+	taa_halfCx2 PrevLowFrequencyColor;
 	taa_half2 PrevLowFrequencyWeight;
 	
-	taa_half3x2 PrevHighFrequencyColor;
+	taa_halfCx2 PrevHighFrequencyColor;
 	taa_half2 PrevHistoryValidity;
 	
-	taa_half3x2 PrevHistoryMommentMin;
-	taa_half3x2 PrevHistoryMommentMax;
+	taa_halfCx2 PrevHistoryMommentMin;
+	taa_halfCx2 PrevHistoryMommentMax;
 
 	taa_half4x2 PrevHistoryTranslucency;
 	
@@ -706,13 +723,13 @@ void MainCS(
 	PLATFORM_SPECIFIC_ISOLATE
 	{
 		// Reproject history
-		taa_half3x2 AccumulateLowFrequency = taa_half(0);
-		taa_half3x2 AccumulateHighFrequency = taa_half(0);
+		taa_halfCx2 AccumulateLowFrequency = taa_half(0);
+		taa_halfCx2 AccumulateHighFrequency = taa_half(0);
 		taa_half2x2 AccumulateMetadata = taa_half(0);
 		taa_half4x2 AccumulateTranslucency = taa_half(0);
 
-		taa_half3x2 HighFrequencyMin = INFINITE_FLOAT;
-		taa_half3x2 HighFrequencyMax = -INFINITE_FLOAT;
+		taa_halfCx2 HighFrequencyMin = INFINITE_FLOAT;
+		taa_halfCx2 HighFrequencyMax = -INFINITE_FLOAT;
 
 		// Sample raw history
 		{
@@ -725,8 +742,8 @@ void MainCS(
 			FCatmullRomSamples Samples1 = GetBicubic2DCatmullRomSamples(
 				dpv_hi(PrevHistoryBufferUV), PrevHistoryInfo_Extent, PrevHistoryInfo_ExtentInverse);
 				
-			taa_half3 RawLowFrequency[BICUBIC_CATMULL_ROM_SAMPLES][DPV_PIXEL_PER_LANE];
-			taa_half3 RawHighFrequency[BICUBIC_CATMULL_ROM_SAMPLES][DPV_PIXEL_PER_LANE];
+			taa_halfC RawLowFrequency[BICUBIC_CATMULL_ROM_SAMPLES][DPV_PIXEL_PER_LANE];
+			taa_halfC RawHighFrequency[BICUBIC_CATMULL_ROM_SAMPLES][DPV_PIXEL_PER_LANE];
 			taa_half2 RawMetadata[BICUBIC_CATMULL_ROM_SAMPLES][DPV_PIXEL_PER_LANE];
 			taa_half3 RawTranslucencyColor[BICUBIC_CATMULL_ROM_SAMPLES][DPV_PIXEL_PER_LANE];
 			taa_half  RawTranslucencyAlpha[BICUBIC_CATMULL_ROM_SAMPLES][DPV_PIXEL_PER_LANE];
@@ -773,10 +790,10 @@ void MainCS(
 				UNROLL_N(BICUBIC_CATMULL_ROM_SAMPLES)
 				for (uint i = 0; i < BICUBIC_CATMULL_ROM_SAMPLES; i++)
 				{
-					taa_half3x2 SampleLowFrequency = dpv_interleave_registers_array(RawLowFrequency[i]);
-					taa_half3x2 SampleHighFrequency = dpv_interleave_registers_array(RawHighFrequency[i]);
-					taa_half2x2 SampleMetadata = dpv_interleave_registers_array(RawMetadata[i]);
-					taa_half4x2 SampleTranslucency = dpv_interleave_registers(
+					taa_halfCx2 SampleLowFrequency = dpv_force_interleave_registers_array(RawLowFrequency[i]);
+					taa_halfCx2 SampleHighFrequency = dpv_force_interleave_registers_array(RawHighFrequency[i]);
+					taa_half2x2 SampleMetadata = dpv_force_interleave_registers_array(RawMetadata[i]);
+					taa_half4x2 SampleTranslucency = dpv_force_interleave_registers(
 						taa_half4(RawTranslucencyColor[i][0], RawTranslucencyAlpha[i][0]),
 						taa_half4(RawTranslucencyColor[i][1], RawTranslucencyAlpha[i][1]));
 					
@@ -796,7 +813,7 @@ void MainCS(
 					#endif
 				}
 				
-				taa_half2 FinalMultiplier = dpv_interleave_registers(taa_half(Samples0.FinalMultiplier), taa_half(Samples1.FinalMultiplier));
+				taa_half2 FinalMultiplier = dpv_force_interleave_registers(taa_half(Samples0.FinalMultiplier), taa_half(Samples1.FinalMultiplier));
 
 				AccumulateLowFrequency  = dpv_scale(AccumulateLowFrequency, FinalMultiplier);
 				AccumulateHighFrequency = dpv_scale(AccumulateHighFrequency, FinalMultiplier);
@@ -817,9 +834,14 @@ void MainCS(
 					PrevSubpixelDetails[1] |= SubpixelPayload1 << taa_subpixel_payload(SUB_PIXEL_BIT_COUNT * SubpixelId);
 				}
 			}
-
-			AccumulateLowFrequency  = -dpv_min(-AccumulateLowFrequency, taa_half(0.0).xxx);
-			AccumulateHighFrequency = -dpv_min(-AccumulateHighFrequency, taa_half(0.0).xxx);
+			
+			#if CONFIG_SCENE_COLOR_APLHA
+				AccumulateLowFrequency  = -dpv_min(-AccumulateLowFrequency, taa_half(0.0).xxxx);
+				AccumulateHighFrequency = -dpv_min(-AccumulateHighFrequency, taa_half(0.0).xxxx);
+			#else
+				AccumulateLowFrequency  = -dpv_min(-AccumulateLowFrequency, taa_half(0.0).xxx);
+				AccumulateHighFrequency = -dpv_min(-AccumulateHighFrequency, taa_half(0.0).xxx);
+			#endif
 			AccumulateMetadata      = -dpv_min(-AccumulateMetadata, taa_half(0.0).xx);
 			AccumulateTranslucency  = -dpv_min(-AccumulateTranslucency, taa_half(0.0).xxxx);
 		}
@@ -906,11 +928,11 @@ void MainCS(
 	// Rejects high frequency details if their low frequencies drift from the low frequency history
 	#if CONFIG_LOW_FREQUENCY_DRIFT_REJECTION
 	{
-		taa_half3x2 NormalizedPrevHighFrequencyColor = dpv_scale(PrevHighFrequencyColor, SafeRcp(PrevHistoryValidity));
+		taa_halfCx2 NormalizedPrevHighFrequencyColor = dpv_scale(PrevHighFrequencyColor, SafeRcp(PrevHistoryValidity));
 
-		taa_half3x2 PrevHighFrequencyYCoCg = TransformColorForClampingBox(NormalizedPrevHighFrequencyColor);
-		taa_half3x2 PrevYCoCg = TransformColorForClampingBox(PrevLowFrequencyColor);
-		taa_half3x2 ClampedPrevYCoCg = TransformColorForClampingBox(clamp(PrevLowFrequencyColor, PrevHistoryMommentMin, PrevHistoryMommentMax));
+		taa_halfCx2 PrevHighFrequencyYCoCg = TransformColorForClampingBox(NormalizedPrevHighFrequencyColor);
+		taa_halfCx2 PrevYCoCg = TransformColorForClampingBox(PrevLowFrequencyColor);
+		taa_halfCx2 ClampedPrevYCoCg = TransformColorForClampingBox(clamp(PrevLowFrequencyColor, PrevHistoryMommentMin, PrevHistoryMommentMax));
 
 		taa_half2 HighFrequencyRejection = MeasureRejectionFactor(
 			PrevYCoCg, ClampedPrevYCoCg,
@@ -930,15 +952,15 @@ void MainCS(
 	// Clamp low frequency
 	taa_half2 LumaMin;
 	taa_half2 LumaMax;
-	taa_half3x2 ClampedLowFrequencyColor;
+	taa_halfCx2 ClampedLowFrequencyColor;
 	taa_half2 MidFrequencyRejection;
 	{
 		LumaMin = InputMinColor[0];
 		LumaMax = InputMaxColor[0];
 
-		taa_half3x2 PrevYCoCg = TransformColorForClampingBox(PrevLowFrequencyColor);
-		taa_half3x2 ClampedPrevYCoCg = dpv_clamp(PrevYCoCg, InputMinColor, InputMaxColor);
-		taa_half3x2 InputCenterYCoCg = TransformColorForClampingBox(FilteredInputColor);
+		taa_halfCx2 PrevYCoCg = TransformColorForClampingBox(PrevLowFrequencyColor);
+		taa_halfCx2 ClampedPrevYCoCg = dpv_clamp(PrevYCoCg, InputMinColor, InputMaxColor);
+		taa_halfCx2 InputCenterYCoCg = TransformColorForClampingBox(FilteredInputColor);
 
 		ClampedLowFrequencyColor = YCoCgToRGB(ClampedPrevYCoCg);
 		
@@ -962,7 +984,7 @@ void MainCS(
 	}
 	#endif
 	
-	taa_half3x2 FinalHighFrequencyColor;
+	taa_halfCx2 FinalHighFrequencyColor;
 	taa_half2 FinalHistoryValidity;
 	{
 		taa_half2 InvPrevHistoryValidity = SafeRcp(PrevHistoryValidity);
@@ -1082,7 +1104,7 @@ void MainCS(
 	}
 
 	// Compute the low frequency history
-	taa_half3x2 FinalLowFrequencyColor;
+	taa_halfCx2 FinalLowFrequencyColor;
 	taa_half2 FinalLowFrequencyWeight;
 	{
 		const taa_half TargetHesteresisCurrentFrameWeight = rcp(taa_half(MAX_FALLBACK_SAMPLE_COUNT));
@@ -1204,10 +1226,17 @@ void MainCS(
 
 		taa_half4x2 ClampedPrevHistoryTranslucency = clamp(PrevHistoryTranslucency, TranslucencyMinColor, TranslucencyMaxColor);
 		
-		taa_half2 HighFrequencyTranslucencyRejection = MeasureRejectionFactor(
-			GetColorChannel(PrevHistoryTranslucency), GetColorChannel(ClampedPrevHistoryTranslucency),
-			GetColorChannel(FilteredTranslucencyColor), GetColorChannel(TranslucencyMinColor), GetColorChannel(TranslucencyMaxColor),
-			taa_half(0.0));
+		#if CONFIG_SCENE_COLOR_APLHA
+			taa_half2 HighFrequencyTranslucencyRejection = MeasureRejectionFactor(
+				PrevHistoryTranslucency, ClampedPrevHistoryTranslucency,
+				FilteredTranslucencyColor, TranslucencyMinColor, TranslucencyMaxColor,
+				taa_half(0.0));
+		#else
+			taa_half2 HighFrequencyTranslucencyRejection = MeasureRejectionFactor(
+				GetColorChannel(PrevHistoryTranslucency), GetColorChannel(ClampedPrevHistoryTranslucency),
+				GetColorChannel(FilteredTranslucencyColor), GetColorChannel(TranslucencyMinColor), GetColorChannel(TranslucencyMaxColor),
+				taa_half(0.0));
+		#endif
 
 		taa_half2 BlendFinal = taa_half(0.2) * InputPixelAlignement;
 
@@ -1311,7 +1340,7 @@ void MainCS(
 	}
 
 	// Compute final output
-	taa_half3x2 FinalOutputColor;
+	taa_halfCx2 FinalOutputColor;
 	{
 		taa_half2 OutputBlend = ComputeLowFrequencyContribution(FinalHistoryValidity);
 
@@ -1351,14 +1380,26 @@ void MainCS(
 		#endif
 
 		#if CONFIG_ACCUMULATE_TRANSLUCENCY_SEPARATLY
-			FinalOutputColor = dpv_scale(FinalOutputColor, FinalTranslucencyColor[3]) + dpv_interleave_registers(
-				dpv_lo(FinalTranslucencyColor).rgb, dpv_hi(FinalTranslucencyColor).rgb);
+			FinalOutputColor[0] = FinalOutputColor[0] * FinalTranslucencyColor[3] + FinalTranslucencyColor[0];
+			FinalOutputColor[1] = FinalOutputColor[1] * FinalTranslucencyColor[3] + FinalTranslucencyColor[1];
+			FinalOutputColor[2] = FinalOutputColor[2] * FinalTranslucencyColor[3] + FinalTranslucencyColor[2];
+			#if CONFIG_SCENE_COLOR_APLHA
+				FinalOutputColor[3] = saturate(taa_half(1.0) - (taa_half(1.0) - FinalOutputColor[3]) * FinalTranslucencyColor[3]);
+			#endif
 		#endif
 	}
 
 
 	PLATFORM_SPECIFIC_ISOLATE
 	{
+		#if CONFIG_SCENE_COLOR_APLHA
+			const taa_halfC ColorNull = taa_half(0.0).xxxx;
+			const taa_halfC ColorMax10BitsFloat = taa_half(Max10BitsFloat).xxxx;
+		#else
+			const taa_halfC ColorNull = taa_half(0.0).xxx;
+			const taa_halfC ColorMax10BitsFloat = taa_half(Max10BitsFloat).xxx;
+		#endif
+
 		uint LocalGroupThreadIndex = GetGroupThreadIndex(GroupThreadIndex, GroupWaveIndex);
 
 		#if 1
@@ -1372,7 +1413,7 @@ void MainCS(
 		LocalHistoryPixelPos = InvalidateOutputPixelPos(LocalHistoryPixelPos, HistoryInfo_ViewportMax);
 		
 		// Premultiply FinalHighFrequencyColor to have validity weighted average in next frame history reprojection.
-		taa_half3x2 PremultipliedHighFrequencyColor = dpv_scale(FinalHighFrequencyColor, FinalHistoryValidity);
+		taa_halfCx2 PremultipliedHighFrequencyColor = dpv_scale(FinalHighFrequencyColor, FinalHistoryValidity);
 
 		// Stocastically round up or down using the hardware RWTexture2D truncation unit to take into precision
 		// loss due to pixel format encoding.
@@ -1388,11 +1429,11 @@ void MainCS(
 
 		// Protext from NaN and +Inf when writing out the history.
 		{
-			FinalLowFrequencyColor = -dpv_min(-FinalLowFrequencyColor, taa_half(0.0).xxx);
-			FinalLowFrequencyColor = dpv_min(FinalLowFrequencyColor, taa_half(Max10BitsFloat).xxx);
+			FinalLowFrequencyColor = -dpv_min(-FinalLowFrequencyColor, ColorNull);
+			FinalLowFrequencyColor = dpv_min(FinalLowFrequencyColor, ColorMax10BitsFloat);
 
-			PremultipliedHighFrequencyColor = -dpv_min(-PremultipliedHighFrequencyColor, taa_half(0.0).xxx);
-			PremultipliedHighFrequencyColor = dpv_min(PremultipliedHighFrequencyColor, taa_half(Max10BitsFloat).xxx);	
+			PremultipliedHighFrequencyColor = -dpv_min(-PremultipliedHighFrequencyColor, ColorNull);
+			PremultipliedHighFrequencyColor = dpv_min(PremultipliedHighFrequencyColor, ColorMax10BitsFloat);	
 
 			#if CONFIG_ACCUMULATE_TRANSLUCENCY_SEPARATLY
 				FinalTranslucencyColor = -dpv_min(-FinalTranslucencyColor, taa_half(0.0).xxxx);
@@ -1427,12 +1468,12 @@ void MainCS(
 		#endif
 
 		// Output final scene color Mip0
-		taa_half3x2 OutputColor;
+		taa_halfCx2 OutputColor;
 		{
 			OutputColor = FinalOutputColor;
 
-			OutputColor = -dpv_min(-OutputColor, taa_half(0.0).xxx);
-			OutputColor = dpv_min(OutputColor, taa_half(Max10BitsFloat).xxx);
+			OutputColor = -dpv_min(-OutputColor, ColorNull);
+			OutputColor = dpv_min(OutputColor, ColorMax10BitsFloat);
 
 			SceneColorOutputMip0[dpv_lo(LocalHistoryPixelPos)] = dpv_lo(OutputColor);
 			
@@ -1444,7 +1485,7 @@ void MainCS(
 		// Output final scene color Mip1
 		#if !DEBUG_OUTPUT // running out of UAVs
 		{
-			taa_half3x2 HalfResOutput = OutputColor * taa_half(0.25);
+			taa_halfCx2 HalfResOutput = OutputColor * taa_half(0.25);
 			//taa_short2x2 HalfResOutputPixelPos = (LocalHistoryPixelPos >> taa_short(1)) | (((LocalHistoryPixelPos & taa_short(0x1)) | !bGenerateOutputMip1) * taa_short(~0));
 			taa_short2x2 HalfResOutputPixelPos;
 			HalfResOutputPixelPos[0] = (LocalHistoryPixelPos[0] >> taa_short(1)) | (((LocalHistoryPixelPos[0] & taa_short(0x1)) | !bGenerateOutputMip1) * taa_short(~0));
@@ -1464,9 +1505,16 @@ void MainCS(
 			BRANCH
 			if (bGenerateOutputMip1)
 			{
-				SharedArray0[LocalGroupThreadIndex] = taa_half4(dpv_lo(HalfResOutput), 0.0);
-				#if CONFIG_ENABLE_DUAL_PIXEL_VECTORIZATION
-					SharedArray1[LocalGroupThreadIndex] = taa_half4(dpv_hi(HalfResOutput), 0.0);
+				#if CONFIG_SCENE_COLOR_APLHA
+					SharedArray0[LocalGroupThreadIndex] = dpv_lo(HalfResOutput);
+					#if CONFIG_ENABLE_DUAL_PIXEL_VECTORIZATION
+						SharedArray1[LocalGroupThreadIndex] = dpv_hi(HalfResOutput);
+					#endif
+				#else
+					SharedArray0[LocalGroupThreadIndex] = taa_half4(dpv_lo(HalfResOutput), 0.0);
+					#if CONFIG_ENABLE_DUAL_PIXEL_VECTORIZATION
+						SharedArray1[LocalGroupThreadIndex] = taa_half4(dpv_hi(HalfResOutput), 0.0);
+					#endif
 				#endif
 
 				SharedArray0[LocalGroupThreadIndex] += SharedArray0[LocalGroupThreadIndex ^ 0x1];
@@ -1479,10 +1527,18 @@ void MainCS(
 					SharedArray1[LocalGroupThreadIndex] += SharedArray1[LocalGroupThreadIndex ^ 0x2];
 				#endif
 					
-				#if CONFIG_ENABLE_DUAL_PIXEL_VECTORIZATION
-					HalfResOutput = dpv_interleave_registers(SharedArray0[LocalGroupThreadIndex].rgb, SharedArray1[LocalGroupThreadIndex].rgb);
+				#if CONFIG_SCENE_COLOR_APLHA
+					#if CONFIG_ENABLE_DUAL_PIXEL_VECTORIZATION
+						HalfResOutput = dpv_interleave_registers(SharedArray0[LocalGroupThreadIndex], SharedArray1[LocalGroupThreadIndex]);
+					#else
+						HalfResOutput = dpv_interleave_mono_registers(SharedArray0[LocalGroupThreadIndex]);
+					#endif
 				#else
-					HalfResOutput = dpv_interleave_mono_registers(SharedArray0[LocalGroupThreadIndex].rgb);
+					#if CONFIG_ENABLE_DUAL_PIXEL_VECTORIZATION
+						HalfResOutput = dpv_interleave_registers(SharedArray0[LocalGroupThreadIndex].rgb, SharedArray1[LocalGroupThreadIndex].rgb);
+					#else
+						HalfResOutput = dpv_interleave_mono_registers(SharedArray0[LocalGroupThreadIndex].rgb);
+					#endif
 				#endif
 			}
 			#endif
diff --git a/Engine/Source/Runtime/Renderer/Private/PostProcess/TemporalSuperResolution.cpp b/Engine/Source/Runtime/Renderer/Private/PostProcess/TemporalSuperResolution.cpp
--- a/Engine/Source/Runtime/Renderer/Private/PostProcess/TemporalSuperResolution.cpp
+++ b/Engine/Source/Runtime/Renderer/Private/PostProcess/TemporalSuperResolution.cpp
@@ -575,6 +575,9 @@ ITemporalUpscaler::FOutputs AddTemporalSuperResolutionPasses(
 	const bool bSetupDebugPasses = CVarTSRSetupDebugPasses.GetValueOnRenderThread() != 0;
 #endif
 
+	// Whether alpha channel is supported.
+	const bool bSupportsAlpha = IsPostProcessingWithAlphaChannelSupported();
+
 	// Whether to use camera cut shader permutation or not.
 	bool bCameraCut = !InputHistory.IsValid() || View.bCameraCut;
 
@@ -588,6 +591,8 @@ ITemporalUpscaler::FOutputs AddTemporalSuperResolutionPasses(
 
 	bool bAccumulateSeparateTranslucency = bIsSeperateTranslucyTexturesValid && CVarTSRTranslucencySeparateTemporalAccumulation.GetValueOnRenderThread() != 0;
 
+	EPixelFormat ColorFormat = bSupportsAlpha ? PF_FloatRGBA : PF_FloatR11G11B10;
+
 	int32 RejectionAntiAliasingQuality = FMath::Clamp(CVarTSRRejectionAntiAliasingQuality.GetValueOnRenderThread(), 1, 2);
 	if (UpdateHistoryQuality == FTSRUpdateHistoryCS::EQuality::Low)
 	{
@@ -667,7 +672,10 @@ ITemporalUpscaler::FOutputs AddTemporalSuperResolutionPasses(
 	float ScreenPercentage = float(InputRect.Width()) / float(OutputRect.Width());
 	float InvScreenPercentage = float(OutputRect.Width()) / float(InputRect.Width());
 
-	RDG_EVENT_SCOPE(GraphBuilder, "TemporalSuperResolution %dx%d -> %dx%d", InputRect.Width(), InputRect.Height(), OutputRect.Width(), OutputRect.Height());
+	RDG_EVENT_SCOPE(GraphBuilder, "TemporalSuperResolution(%s) %dx%d -> %dx%d",
+		bSupportsAlpha ? TEXT("Alpha") : TEXT(""),
+		InputRect.Width(), InputRect.Height(),
+		OutputRect.Width(), OutputRect.Height());
 	RDG_GPU_STAT_SCOPE(GraphBuilder, TemporalSuperResolution);
 
 	FRDGTextureRef BlackUintDummy = GSystemTextures.GetZeroUIntDummy(GraphBuilder);
@@ -895,7 +903,7 @@ ITemporalUpscaler::FOutputs AddTemporalSuperResolutionPasses(
 	{
 		FRDGTextureDesc Desc = FRDGTextureDesc::Create2D(
 			HistoryExtent,
-			(CVarTSRR11G11B10History.GetValueOnRenderThread() != 0) ? PF_FloatR11G11B10 : PF_FloatRGBA,
+			(CVarTSRR11G11B10History.GetValueOnRenderThread() != 0 && !bSupportsAlpha) ? PF_FloatR11G11B10 : PF_FloatRGBA,
 			FClearValueBinding::None,
 			TexCreate_ShaderResource | TexCreate_UAV);
 
@@ -907,7 +915,7 @@ ITemporalUpscaler::FOutputs AddTemporalSuperResolutionPasses(
 
 		if (bAccumulateSeparateTranslucency)
 		{
-			Desc.Format = (CVarTSRR11G11B10History.GetValueOnRenderThread() != 0) ? PF_FloatR11G11B10 : PF_FloatRGBA;
+			Desc.Format = (CVarTSRR11G11B10History.GetValueOnRenderThread() != 0 && !bSupportsAlpha) ? PF_FloatR11G11B10 : PF_FloatRGBA;
 			History.Translucency = GraphBuilder.CreateTexture(Desc, TEXT("TSR.History.Translucency"));
 
 			Desc.Format = PF_R8;
@@ -946,7 +954,7 @@ ITemporalUpscaler::FOutputs AddTemporalSuperResolutionPasses(
 
 			PassParameters->RotationalClipToPrevClip = RotationalInvViewProj * RotationalPrevViewProj;
 		}
-		PassParameters->OutputQuantizationError = ComputePixelFormatQuantizationError(PF_FloatR11G11B10);
+		PassParameters->OutputQuantizationError = ComputePixelFormatQuantizationError(ColorFormat);
 		{
 			float TanHalfFieldOfView = View.ViewMatrices.GetInvProjectionMatrix().M[0][0];
 
@@ -968,7 +976,7 @@ ITemporalUpscaler::FOutputs AddTemporalSuperResolutionPasses(
 		{
 			FRDGTextureDesc Desc = FRDGTextureDesc::Create2D(
 				LowFrequencyExtent,
-				PF_FloatR11G11B10,
+				ColorFormat,
 				FClearValueBinding::None,
 				/* InFlags = */ TexCreate_ShaderResource | TexCreate_UAV);
 
@@ -986,7 +994,7 @@ ITemporalUpscaler::FOutputs AddTemporalSuperResolutionPasses(
 		{
 			FRDGTextureDesc Desc = FRDGTextureDesc::Create2D(
 				InputExtent,
-				PF_FloatR11G11B10,
+				ColorFormat,
 				FClearValueBinding::None,
 				/* InFlags = */ TexCreate_ShaderResource | TexCreate_UAV);
 
@@ -1120,7 +1128,7 @@ ITemporalUpscaler::FOutputs AddTemporalSuperResolutionPasses(
 			{
 				FRDGTextureDesc Desc = FRDGTextureDesc::Create2D(
 					LowFrequencyExtent,
-					PF_FloatR11G11B10,
+					ColorFormat,
 					FClearValueBinding::None,
 					/* InFlags = */ TexCreate_ShaderResource | TexCreate_UAV);
 
@@ -1333,7 +1341,7 @@ ITemporalUpscaler::FOutputs AddTemporalSuperResolutionPasses(
 	{
 		FRDGTextureDesc Desc = FRDGTextureDesc::Create2D(
 			OutputExtent,
-			PF_FloatR11G11B10,
+			ColorFormat,
 			FClearValueBinding::None,
 			/* InFlags = */ TexCreate_ShaderResource | TexCreate_UAV,
 			/* NumMips = */ PassInputs.bGenerateOutputMip1 ? 2 : 1);
