Fix splitscreen subsurface scattering issue, and reduce unnecessary calculation in splitscreen settings.

diff --git a/Engine/Shaders/Private/PostProcessSubsurface.usf b/Engine/Shaders/Private/PostProcessSubsurface.usf
--- a/Engine/Shaders/Private/PostProcessSubsurface.usf
+++ b/Engine/Shaders/Private/PostProcessSubsurface.usf
@@ -624,6 +624,7 @@ void BuildIndirectDispatchArgsCS(uint2 DispatchThreadId : SV_DispatchThreadID)
 
 #define TYPE_SEPARABLE 0x1
 #define TYPE_BURLEY    0x2
+#define TYPE_IN_VIEWPORT 0x4
 
 // 0, no subsurface
 // 1, Separable
@@ -690,17 +691,19 @@ void AddSubsurfaceComputeGroup(uint GI, uint Type, uint2 GroupXY)
 
 	GroupMemoryBarrierWithGroupSync();
 
-	if (GI == 0 && (SubsurfaceTypeFlag & TYPE_SEPARABLE))
+	if (GI == 0 && ((SubsurfaceTypeFlag & TYPE_SEPARABLE) && (SubsurfaceTypeFlag & TYPE_IN_VIEWPORT)))
 	{
 		AddSeparableGroup(GroupXY);
 	}
 
-	if (GI == 0 && (SubsurfaceTypeFlag & TYPE_BURLEY))
+	if (GI == 0 && ((SubsurfaceTypeFlag & TYPE_BURLEY) && (SubsurfaceTypeFlag & TYPE_IN_VIEWPORT)))
 	{
 		AddBurleyGroup(GroupXY);
 	}
 
-	if (GI == 0 && SubsurfaceTypeFlag > 0)
+	if (GI == 0 && ((SubsurfaceTypeFlag & TYPE_IN_VIEWPORT) &&
+				   ((SubsurfaceTypeFlag &	TYPE_SEPARABLE) ||
+					(SubsurfaceTypeFlag &	   TYPE_BURLEY))))
 	{
 		AddAllGroup(GroupXY);
 	}
@@ -827,6 +830,7 @@ void SetupIndirectCS(uint3 DT_ID : SV_DispatchThreadID,
 
 	if (all(Pos < Output_ViewportMax))
 	{
+		Type |= TYPE_IN_VIEWPORT;
 		SetupTexture[Pos] = OutColor;
 
 #if !SUBSURFACE_FORCE_SEPARABLE && ENABLE_PROFILE_ID_CACHE
diff --git a/Engine/Shaders/Private/PostProcessSubsurfaceTile.usf b/Engine/Shaders/Private/PostProcessSubsurfaceTile.usf
--- a/Engine/Shaders/Private/PostProcessSubsurfaceTile.usf
+++ b/Engine/Shaders/Private/PostProcessSubsurfaceTile.usf
@@ -78,6 +78,7 @@ void ClearUAV(uint2 DT_ID : SV_DispatchThreadID)
 #if SHADER_TILE_VS
 
 int2 ViewMin;
+int2 ViewMax;
 float2 ExtentInverse;
 uint TileType;
 Buffer<uint> TileDataBuffer;
@@ -93,7 +94,7 @@ void MainVS(
 	uint2 TileVertex = TileCoord * SUBSURFACE_TILE_SIZE;
 	TileVertex.x += InVertexId == 1 || InVertexId == 2 || InVertexId == 4 ? SUBSURFACE_TILE_SIZE : 0;
 	TileVertex.y += InVertexId == 2 || InVertexId == 4 || InVertexId == 5 ? SUBSURFACE_TILE_SIZE : 0;
-	Out.UV = float2(TileVertex + ViewMin) * ExtentInverse;
+	Out.UV = float2(min((int2)TileVertex + ViewMin, ViewMax)) * ExtentInverse;
 	Out.Position = float4(Out.UV * float2(2.0f, -2.0f) + float2(-1.0, 1.0f), 0.0f, 1.0f);
 }
 
diff --git a/Engine/Source/Runtime/Renderer/Private/SubsurfaceTiles.cpp b/Engine/Source/Runtime/Renderer/Private/SubsurfaceTiles.cpp
--- a/Engine/Source/Runtime/Renderer/Private/SubsurfaceTiles.cpp
+++ b/Engine/Source/Runtime/Renderer/Private/SubsurfaceTiles.cpp
@@ -28,6 +28,7 @@ FSubsurfaceTilePassVS::FParameters GetSubsurfaceTileParameters(const FViewInfo&
 	Out.TileType = uint32(TileType);
 	Out.bRectPrimitive = InTile.bRectPrimitive ? 1 : 0;
 	Out.ViewMin = InView.ViewRect.Min;
+	Out.ViewMax = InView.ViewRect.Max;
 	Out.ExtentInverse = FVector2f(1.f / BufferExtent.X, 1.f / BufferExtent.Y);
 	Out.TileDataBuffer = InTile.GetTileBufferSRV(TileType);
 	Out.TileIndirectBuffer = InTile.TileIndirectDrawBuffer;
diff --git a/Engine/Source/Runtime/Renderer/Private/SubsurfaceTiles.h b/Engine/Source/Runtime/Renderer/Private/SubsurfaceTiles.h
--- a/Engine/Source/Runtime/Renderer/Private/SubsurfaceTiles.h
+++ b/Engine/Source/Runtime/Renderer/Private/SubsurfaceTiles.h
@@ -64,6 +64,7 @@ class FSubsurfaceTilePassVS : public FGlobalShader
 
 	BEGIN_SHADER_PARAMETER_STRUCT(FParameters, )
 		SHADER_PARAMETER(FIntPoint, ViewMin)
+		SHADER_PARAMETER(FIntPoint, ViewMax)
 		SHADER_PARAMETER(int32, bRectPrimitive)
 		SHADER_PARAMETER(uint32, TileType)
 		SHADER_PARAMETER(FVector2f, ExtentInverse)
