Adjust LWC->float conversion for LWC UVs, avoid the need to explicitly compute derivatives
UV discontinuity and possibility of seam returns, but should be much more rare than previously

diff --git a/Engine/Shaders/Private/MaterialTemplate.ush b/Engine/Shaders/Private/MaterialTemplate.ush
--- a/Engine/Shaders/Private/MaterialTemplate.ush
+++ b/Engine/Shaders/Private/MaterialTemplate.ush
@@ -1845,9 +1845,14 @@ float LWCApplyAddressModeMirror(FLWCScalar v)
 
 float LWCApplyAddressMode(FLWCScalar v, uint AddressMode)
 {
-	if(AddressMode == LWCADDRESSMODE_WRAP) return LWCFrac(v);
-	else if(AddressMode == LWCADDRESSMODE_MIRROR) return LWCApplyAddressModeMirror(v);
-	else return LWCSaturate(v);
+	// For wrap/mirror, we can take the offset of the normalized tile
+	// This lets the GPU handle wrapping in most cases which will preserve correct derivatives.  There will be a seam along tile boundaries in the UV however
+	// May be possible to improve this by wrapping at higher values than TileSize, which would reduce frequency (but not eliminate) the seam
+	// For the CLAMP case, a simple LWCToFloat() is sufficient.  This will lose a ton of precision for large values, but we don't care about this since the GPU will clamp to [0,1) anyway
+	// It's possible certain GPUs might need a special case if the ToFloat() conversion overflows
+	if(AddressMode == LWCADDRESSMODE_WRAP) return LWCNormalizeTile(v).Offset;
+	else if(AddressMode == LWCADDRESSMODE_MIRROR) LWCNormalizeTile(v).Offset;
+	else return LWCToFloat(v);
 }
 float2 LWCApplyAddressMode(FLWCVector2 UV, uint AddressX, uint AddressY)
 {
diff --git a/Engine/Source/Runtime/Engine/Private/Materials/HLSLMaterialTranslator.cpp b/Engine/Source/Runtime/Engine/Private/Materials/HLSLMaterialTranslator.cpp
--- a/Engine/Source/Runtime/Engine/Private/Materials/HLSLMaterialTranslator.cpp
+++ b/Engine/Source/Runtime/Engine/Private/Materials/HLSLMaterialTranslator.cpp
@@ -5859,7 +5859,12 @@ int32 FHLSLMaterialTranslator::TextureSample(
 			break;
 		}
 
-		if (MipValueMode == TMVM_None || MipValueMode == TMVM_MipBias)
+		// Explicitly compute the derivatives for LWC UVs
+		// This is needed for 100% correct functionality, otherwise filtering seams are possible where there is discontinuity in LWC->float UV conversion
+		// This is expensive though, and discontinuities can be minimized by carefully choosing conversion operation
+		// Disabled for now, may enable as an option in the future
+		const bool bExplicitLWCDerivatives = false;
+		if (bExplicitLWCDerivatives && (MipValueMode == TMVM_None || MipValueMode == TMVM_MipBias))
 		{
 			int32 MipScaleIndex = INDEX_NONE;
 			if (MipValueMode == TMVM_MipBias)
