Missing return statement

diff --git a/Engine/Shaders/Private/MaterialTemplate.ush b/Engine/Shaders/Private/MaterialTemplate.ush
--- a/Engine/Shaders/Private/MaterialTemplate.ush
+++ b/Engine/Shaders/Private/MaterialTemplate.ush
@@ -1851,7 +1851,7 @@ float LWCApplyAddressMode(FLWCScalar v, uint AddressMode)
 	// For the CLAMP case, a simple LWCToFloat() is sufficient.  This will lose a ton of precision for large values, but we don't care about this since the GPU will clamp to [0,1) anyway
 	// It's possible certain GPUs might need a special case if the ToFloat() conversion overflows
 	if(AddressMode == LWCADDRESSMODE_WRAP) return LWCNormalizeTile(v).Offset;
-	else if(AddressMode == LWCADDRESSMODE_MIRROR) LWCNormalizeTile(v).Offset;
+	else if(AddressMode == LWCADDRESSMODE_MIRROR) return LWCNormalizeTile(v).Offset;
 	else return LWCToFloat(v);
 }
 float2 LWCApplyAddressMode(FLWCVector2 UV, uint AddressX, uint AddressY)
