Compile out realtime capture related code from atmosphere shading on mobile platforms

diff --git a/Engine/Shaders/Private/SkyAtmosphereCommon.ush b/Engine/Shaders/Private/SkyAtmosphereCommon.ush
--- a/Engine/Shaders/Private/SkyAtmosphereCommon.ush
+++ b/Engine/Shaders/Private/SkyAtmosphereCommon.ush
@@ -8,6 +8,8 @@
 #ifndef D_SKY_ATMOSPHERE_COMMON
 #define D_SKY_ATMOSPHERE_COMMON
 
+// Disabled for all mobile platforms, see USkyLightComponent::IsRealTimeCaptureEnabled
+#define AP_SUPPORT_REALTIME_REFLECTION_CAPTURE (FEATURE_LEVEL > FEATURE_LEVEL_ES3_1)
 
 // The constants below should match the one in SceneRendering.cpp
 // Kilometers as unit for computations related to the sky and its atmosphere
@@ -56,6 +58,7 @@ float4 GetAerialPerspectiveLuminanceTransmittance(
 		Weight = saturate(NonLinSlice*NonLinSlice * 2.0f); // Square to have a linear falloff from the change of distribution above
 	}
 
+#if AP_SUPPORT_REALTIME_REFLECTION_CAPTURE
 	if (ViewIsRealTimeReflectionCapture)
 	{
 		// We modify ScreenUv to sample the correct AP accordign to 360 reflection AP LUT.
@@ -74,6 +77,7 @@ float4 GetAerialPerspectiveLuminanceTransmittance(
 
 		ScreenUv = FromUnitToSubUvs(ScreenUv, CameraAerialPerspectiveVolumeSizeAndInvSize);
 	}
+#endif
 
 	float4 AP = Texture3DSampleLevel(AerialPerspectiveVolumeTexture, AerialPerspectiveVolumeTextureSampler, float3(ScreenUv, NonLinW), 0.0f);
 
