Fiz HZB culling of invalidations issue with boxes parallel to light near plane

diff --git a/Engine/Shaders/Private/VirtualShadowMaps/CacheManagement.usf b/Engine/Shaders/Private/VirtualShadowMaps/CacheManagement.usf
--- a/Engine/Shaders/Private/VirtualShadowMaps/CacheManagement.usf
+++ b/Engine/Shaders/Private/VirtualShadowMaps/CacheManagement.usf
@@ -136,6 +136,14 @@ void VirtualSmInvalidateInstancePagesCS(
 						{
 							int ViewDim = int(uint(VSM_VIRTUAL_MAX_RESOLUTION_XY) >> MipLevel);
 							FScreenRect Rect = GetScreenRect(int4(0, 0, ViewDim, ViewDim), Cull, 4);
+
+							// Add a small epsilon to the HZB depth test
+							// This is to handle the rare case where an object that is fully parallel to the 
+							// light's near plane might self-occlude the HZB test due to minor precision differences
+							// in the computation. While rare, this can come up with things like point lights and
+							// axis aligned boxes.
+							Rect.Depth += 1e-8f;
+
 							// Use Hierarchical mip test to speed up (allows skipping invalidating areas that don't have any flags anyway)
 							if (OverlapsAnyValidPage(ProjectionData.VirtualShadowMapId, MipLevel, float2(0.0f, 0.0f), Rect, VSM_ALLOCATED_FLAG))
 							{
