Mobile Integrate Static_Light when DistanceField is on

diff --git a/Engine/Shaders/Private/MobileBasePassPixelShader.usf b/Engine/Shaders/Private/MobileBasePassPixelShader.usf
--- a/Engine/Shaders/Private/MobileBasePassPixelShader.usf
+++ b/Engine/Shaders/Private/MobileBasePassPixelShader.usf
@@ -416,6 +416,15 @@ void ApplyPixelDepthOffsetForMobileBasePass(inout FMaterialPixelParameters Mater
     float PixelDepthOffset = ApplyPixelDepthOffsetToMaterialParameters(MaterialParameters, PixelMaterialInputs, OutDepth);
 }
 
+half ApplyPrecomputedShadowMask(half ShadowMap, half Shadow)
+{
+#if ALLOW_STATIC_LIGHTING
+	return min(ShadowMap, Shadow);
+#else
+	return ShadowMap;
+#endif
+}
+
 #if MAX_DYNAMIC_POINT_LIGHTS > 0
 void AccumulateLightingOfDynamicPointLight(
 FMaterialPixelParameters MaterialParameters, 
@@ -822,19 +831,16 @@ void Main(
 	if (IsCSMEnabled())
 	{
 		half ShadowMap = MobileDirectionalLightCSM(MaterialParameters.ScreenPosition.xy, MaterialParameters.ScreenPosition.w, NoL, ShadowPositionZ);
-	#if ALLOW_STATIC_LIGHTING
-		Shadow = min(ShadowMap, Shadow);
-	#else
-		Shadow = ShadowMap;
-	#endif
+		Shadow = ApplyPrecomputedShadowMask(ShadowMap, Shadow);
 	}
 #endif /* DIRECTIONAL_LIGHT_CSM */
 
 #if APPLY_DISTANCE_FIELD
 	if (ShadowPositionZ == 0)
 	{
-		Shadow = Texture2DSample(MobileBasePass.ScreenSpaceShadowMaskTexture, MobileBasePass.ScreenSpaceShadowMaskSampler, SvPositionToBufferUV(SvPosition)).x;
-		Shadow = DecodeLightAttenuation(Shadow);
+		half ShadowMap = Texture2DSample(MobileBasePass.ScreenSpaceShadowMaskTexture, MobileBasePass.ScreenSpaceShadowMaskSampler, SvPositionToBufferUV(SvPosition)).x;
+		ShadowMap = DecodeLightAttenuation(ShadowMap);
+		Shadow = ApplyPrecomputedShadowMask(ShadowMap, Shadow);
 	}
 #endif
 
