Fix discrepancy in (dynamic) sky lighting between strands and cards with.

This changes ensure that the sky lighting evaluation is identical for the two types of geometry. The shader evaluation is now common to both paths.

diff --git a/Engine/Shaders/Private/HairStrands/HairStrandsEnvironmentLighting.usf b/Engine/Shaders/Private/HairStrands/HairStrandsEnvironmentLighting.usf
--- a/Engine/Shaders/Private/HairStrands/HairStrandsEnvironmentLighting.usf
+++ b/Engine/Shaders/Private/HairStrands/HairStrandsEnvironmentLighting.usf
@@ -10,6 +10,7 @@
 #include "../BRDF.ush"
 #include "../ReflectionEnvironmentShared.ush"
 #include "../SkyLightingShared.ush"
+#include "../SkyLightingDiffuseShared.ush"
 #include "../DistanceFieldAOShared.ush"
 #include "../ShadingModels.ush"
 #include "../LightGridCommon.ush"
@@ -152,11 +153,6 @@ void AddSample(uint SampleOffset, uint SampleIt, const float3 D, const float3 W)
 
 //////////////////////////////////////////////////////////////////////////////////////////////////////////////////////////////
 
-
-float4 SkyLight_OcclusionTintAndMinOcclusion;
-float SkyLight_OcclusionExponent;
-uint SkyLight_OcclusionCombineMode;
-
 float HairDistanceThreshold;
 uint bHairUseViewHairCount;
 Texture2D<float> HairCountTexture;
@@ -631,6 +627,7 @@ float3 ReflectionEnvironmentUniformSampling(
 // Spherical based harmonic integration
 
 #if PERMUTATION_INTEGRATION_TYPE == INTEGRATION_SH
+
 float3 ReflectionEnvironmentSH(
 	FGBufferData GBuffer,
 	float2 ScreenPosition,
@@ -639,62 +636,13 @@ float3 ReflectionEnvironmentSH(
 	uint MacroGroupId,
 	bool bDebugEnabled)
 {
-	const float PreExposure = GetPreExposure();
-	const float3 WorldPosition = mul(float4(ScreenPosition * GBuffer.Depth, GBuffer.Depth, 1), LWCHackToFloat(PrimaryView.ScreenToWorld)).xyz;
-	const float3 CameraToPixel = normalize(WorldPosition - LWCHackToFloat(PrimaryView.WorldCameraOrigin));
+	const float AmbientOcclusion = 1.0f;
+	const float3 BentNormal = GBuffer.WorldNormal;
+	const float3 DiffuseColor = GBuffer.BaseColor; // not used
+	const float2 BufferUV = SvPosition * View.BufferSizeAndInvSize.zw; // not used
 
-	const float3 V = -CameraToPixel;
-	const float3 T = GBuffer.WorldNormal;
-	const float3 N = T;
-	const float3 L = normalize(V - N * dot(V, N));
-
-	FHairTransmittanceData TransmittanceData = GetHairTransmittance(
-		V,
-		L,
-		GBuffer,
-		InitHairTransmittanceMask(),
-		View.HairScatteringLUTTexture, 
-		View.HairScatteringLUTSampler,
-		View.HairComponents);
-
-	bool bEvalMultiScatter = true;
-	float3 f_s = PI * HairShading(GBuffer, L, V, N, 1, TransmittanceData, 1, 0.2, uint2(0, 0));
-
-	// This logic mimic the ReflecitonEnvironemtPixelShader.usf. Ideallyw we should share the same 
-	// code/path, but the various permutation make the code a bit obscure
-	// Need to take into account DFAO?
-	float3 Lighting = 0;
-	{
-	#if 0
-		// Diffuse SH lighting. Require to have GI data laydown during the based pass
-		float AmbientOcclusion = 1;
-		float SkyVisibility = 1;
-		float3 SkyLightingNormal = GBuffer.WorldNormal;
-
-		// Apply DFAO controls
-		SkyVisibility = pow(SkyVisibility, SkyLight_OcclusionExponent);
-		SkyVisibility = lerp(SkyVisibility, 1, SkyLight_OcclusionTintAndMinOcclusion.w);
-
-		// Combine with other AO sources
-		if (SkyLight_OcclusionCombineMode == 0)
-		{
-			// Combine with min which nicely avoids over-occlusion in cases where strong DFAO is present along with strong SSAO (distant trees)
-			SkyVisibility = min(SkyVisibility, min(GBuffer.GBufferAO, AmbientOcclusion));
-		}
-		else
-		{
-			// Combine with mul, which continues to add SSAO depth even indoors.  SSAO will need to be tweaked to be less strong.
-			SkyVisibility = SkyVisibility * min(GBuffer.GBufferAO, AmbientOcclusion);
-		}
-		const float3 DiffuseLighting = PreExposure * GetSkySHDiffuse(L) * View.SkyLightColor.rgb;
-		Lighting = SkyVisibility * DiffuseLighting * f_s + (1 - SkyVisibility) * SkyLight_OcclusionTintAndMinOcclusion.xyz * f_s;
-	#else
-		// Probe lighting
-		const uint2  PixelPosition = SvPosition - ResolvedView.ViewRectMin.xy;
-		const float  PixelSceneDepth = GBuffer.Depth;
-		Lighting = PreExposure * GatherRadiance(L, GBuffer.Roughness, PixelPosition, PixelSceneDepth, WorldPosition) * f_s;
-	#endif
-	}
+	// Note this code is enabled only if sky lighting is dynamic
+	const float3 Lighting = SkyLightDiffuse(GBuffer, AmbientOcclusion, BufferUV, ScreenPosition, BentNormal, DiffuseColor);
 
 	// Transform NaNs to black, transform negative colors to black.
 	return -min(-Lighting, 0.0);
diff --git a/Engine/Shaders/Private/ReflectionEnvironmentPixelShader.usf b/Engine/Shaders/Private/ReflectionEnvironmentPixelShader.usf
--- a/Engine/Shaders/Private/ReflectionEnvironmentPixelShader.usf
+++ b/Engine/Shaders/Private/ReflectionEnvironmentPixelShader.usf
@@ -9,6 +9,7 @@
 #include "BRDF.ush"
 #include "ReflectionEnvironmentShared.ush"
 #include "SkyLightingShared.ush"
+#include "SkyLightingDiffuseShared.ush"
 #include "DistanceFieldAOShared.ush"
 #include "ShadingModels.ush"
 #include "LightGridCommon.ush"
@@ -35,73 +36,6 @@ int CloudSkyAOEnabled;
 Texture2D AmbientOcclusionTexture;
 SamplerState AmbientOcclusionSampler;
 
-float3 ContrastAndNormalizeMulAdd;
-float OcclusionExponent;
-float OcclusionCombineMode;
-
-
-
-struct FSkyLightVisibilityData
-{
-	float  SkyDiffuseLookUpMul;
-	float3 SkyDiffuseLookUpAdd;
-	float3 SkyDiffuseLookUpNormal;
-};
-FSkyLightVisibilityData GetSkyLightVisibilityData(float3 SkyLightingNormal, const float3 WorldNormal, const float GBufferAO, float ScreenAO, const float3 BentNormal)
-{
-	float SkyVisibility = 1;
-	float DotProductFactor = 1;
-
-#if APPLY_SKY_SHADOWING
-#define USE_DIRECTIONAL_OCCLUSION_ON_SKY_DIFFUSE 1
-
-#if USE_DIRECTIONAL_OCCLUSION_ON_SKY_DIFFUSE
-	{
-		SkyVisibility = length(BentNormal);
-		float3 NormalizedBentNormal = BentNormal / (max(SkyVisibility, .00001f));
-
-		// Use more bent normal in corners
-		float BentNormalWeightFactor = SkyVisibility;
-
-		SkyLightingNormal = lerp(NormalizedBentNormal, WorldNormal, BentNormalWeightFactor);
-		DotProductFactor = lerp(dot(NormalizedBentNormal, WorldNormal), 1, BentNormalWeightFactor);
-	}
-#else
-	{
-		SkyVisibility = length(BentNormal);
-	}
-#endif
-
-	float ContrastCurve = 1 / (1 + exp(-ContrastAndNormalizeMulAdd.x * (SkyVisibility * 10 - 5)));
-	SkyVisibility = saturate(ContrastCurve * ContrastAndNormalizeMulAdd.y + ContrastAndNormalizeMulAdd.z);
-
-#endif
-
-	// Apply DFAO controls
-	SkyVisibility = pow(SkyVisibility, OcclusionExponent);
-	SkyVisibility = lerp(SkyVisibility, 1, OcclusionTintAndMinOcclusion.w);
-
-	// Combine with other AO sources
-	if (OcclusionCombineMode == 0)
-	{
-		// Combine with min which nicely avoids over-occlusion in cases where strong DFAO is present along with strong SSAO (distant trees)
-		SkyVisibility = min(SkyVisibility, min(GBufferAO, ScreenAO));
-	}
-	else
-	{
-		// Combine with mul, which continues to add SSAO depth even indoors.  SSAO will need to be tweaked to be less strong.
-		SkyVisibility = SkyVisibility * min(GBufferAO, ScreenAO);
-	}
-
-	// Apply AO to the sky diffuse and account for darkening due to the geometry term
-	// apply the Diffuse color to the lighting (including OcclusionTintAndMinOcclusion as it's considered another light, that fixes SubsurfaceProfile being too dark)
-	FSkyLightVisibilityData SkyVisData;
-	SkyVisData.SkyDiffuseLookUpMul = SkyVisibility * DotProductFactor;
-	SkyVisData.SkyDiffuseLookUpAdd = (1 - SkyVisibility) * OcclusionTintAndMinOcclusion.xyz;
-	SkyVisData.SkyDiffuseLookUpNormal = SkyLightingNormal;
-	return SkyVisData;
-}
-
 
 #if STRATA_ENABLED
 #define STRATA_INLINE_SHADING 0
@@ -147,62 +81,6 @@ float3 GatherRadiance(float CompositeAlpha, float3 TranslatedWorldPosition, floa
 		bCompositeSkylight);
 }
 
-float3 SkyLightDiffuse(FGBufferData GBuffer, float AmbientOcclusion, float2 BufferUV, float2 ScreenPosition, float3 BentNormal, float3 DiffuseColor)
-{
-	float2 UV = BufferUV;
-	float3 Lighting = 0;
-
-	float3 SkyLightingNormal = GetClearCoatBottomNormal(GBuffer, GBuffer.WorldNormal);
-	FSkyLightVisibilityData SkyVisData = GetSkyLightVisibilityData(SkyLightingNormal, GBuffer.WorldNormal, GBuffer.GBufferAO, AmbientOcclusion, BentNormal);
-
-	const float3 N = GBuffer.WorldNormal;
-	const float3 V = -normalize(mul(float3(ScreenPosition, 1), LWCToFloat3x3(PrimaryView.ScreenToWorld)).xyz);
-	const float NoV = saturate(dot(V, N));
-	const FBxDFEnergyTerms EnergyTerms = ComputeGGXSpecEnergyTerms(GBuffer.Roughness, NoV, GBuffer.SpecularColor);
-	float3 DiffuseWeight = ComputeEnergyPreservation(EnergyTerms);
-
-	BRANCH
-	if (GBuffer.ShadingModelID == SHADINGMODELID_TWOSIDED_FOLIAGE)
-	{
-		float3 SubsurfaceLookup = GetSkySHDiffuse(-GBuffer.WorldNormal) * View.SkyLightColor.rgb;
-		float3 SubsurfaceColor = ExtractSubsurfaceColor(GBuffer);
-		Lighting += SkyVisData.SkyDiffuseLookUpMul * SubsurfaceLookup * SubsurfaceColor;
-	}
-
-	if (GBuffer.ShadingModelID == SHADINGMODELID_SUBSURFACE || GBuffer.ShadingModelID == SHADINGMODELID_PREINTEGRATED_SKIN)
-	{
-		float3 SubsurfaceColor = ExtractSubsurfaceColor(GBuffer);
-		// Add subsurface energy to diffuse
-		DiffuseColor += SubsurfaceColor;
-	}
-
-	BRANCH
-	if (GBuffer.ShadingModelID == SHADINGMODELID_HAIR)
-	{
-		float3 L = normalize(V - N * dot(V, N));
-		SkyVisData.SkyDiffuseLookUpNormal = L;
-		FHairTransmittanceData TransmittanceData = InitHairTransmittanceData(true);
-		DiffuseColor = PI * HairShading(GBuffer, L, V, N, 1, TransmittanceData, 0, 0.2, uint2(0, 0));
-		DiffuseWeight = 1.0f;
-	}
-
-	if (GBuffer.ShadingModelID == SHADINGMODELID_CLOTH)
-	{
-		float3 ClothFuzz = ExtractSubsurfaceColor(GBuffer);
-		DiffuseColor += ClothFuzz * GBuffer.CustomData.a;
-	}
-
-	// Compute the preconvolved incoming lighting with the bent normal direction
-	float3 DiffuseLookup = GetSkySHDiffuse(SkyVisData.SkyDiffuseLookUpNormal) * View.SkyLightColor.rgb;
-
-	// And accumulate the lighting
-	Lighting += (SkyVisData.SkyDiffuseLookUpMul * DiffuseLookup + SkyVisData.SkyDiffuseLookUpAdd) * DiffuseColor * DiffuseWeight;
-	Lighting *= View.PreExposure;
-
-	//Lighting = (Texture2DSampleLevel(BentNormalAOTexture, BentNormalAOSampler, UV, 0).xyz);
-
-	return Lighting;
-}
 
 float3 ReflectionEnvironment(FGBufferData GBuffer, float AmbientOcclusion, float2 BufferUV, float2 ScreenPosition, float4 SvPosition, float3 BentNormal, float3 SpecularColor)
 {
diff --git a/Engine/Shaders/Private/SkyLightingDiffuseShared.ush b/Engine/Shaders/Private/SkyLightingDiffuseShared.ush
new file mode 100644
--- /dev/null
+++ b/Engine/Shaders/Private/SkyLightingDiffuseShared.ush
@@ -0,0 +1,139 @@
+// Copyright Epic Games, Inc. All Rights Reserved.
+
+/*=============================================================================
+	SkyLightingDiffuseShared.usf
+=============================================================================*/
+
+#pragma once
+
+#include "SkyLightingShared.ush"
+#include "ClearCoatCommon.ush"
+#include "ShadingEnergyConservation.ush"
+#include "HairBsdf.ush"
+
+#ifndef APPLY_SKY_SHADOWING
+#define APPLY_SKY_SHADOWING 0
+#endif
+
+float3 ContrastAndNormalizeMulAdd;
+float OcclusionExponent;
+float OcclusionCombineMode;
+// The following paramaters are declared in SkyLightingShared.usf
+// float4 OcclusionTintAndMinOcclusion;
+
+struct FSkyLightVisibilityData
+{
+	float  SkyDiffuseLookUpMul;
+	float3 SkyDiffuseLookUpAdd;
+	float3 SkyDiffuseLookUpNormal;
+};
+
+FSkyLightVisibilityData GetSkyLightVisibilityData(float3 SkyLightingNormal, const float3 WorldNormal, const float GBufferAO, float ScreenAO, const float3 BentNormal)
+{
+	float SkyVisibility = 1;
+	float DotProductFactor = 1;
+
+#if APPLY_SKY_SHADOWING
+	#define USE_DIRECTIONAL_OCCLUSION_ON_SKY_DIFFUSE 1
+	#if USE_DIRECTIONAL_OCCLUSION_ON_SKY_DIFFUSE
+	{
+		SkyVisibility = length(BentNormal);
+		float3 NormalizedBentNormal = BentNormal / (max(SkyVisibility, .00001f));
+
+		// Use more bent normal in corners
+		float BentNormalWeightFactor = SkyVisibility;
+
+		SkyLightingNormal = lerp(NormalizedBentNormal, WorldNormal, BentNormalWeightFactor);
+		DotProductFactor = lerp(dot(NormalizedBentNormal, WorldNormal), 1, BentNormalWeightFactor);
+	}
+	#else
+	{
+		SkyVisibility = length(BentNormal);
+	}
+	#endif
+
+	float ContrastCurve = 1 / (1 + exp(-ContrastAndNormalizeMulAdd.x * (SkyVisibility * 10 - 5)));
+	SkyVisibility = saturate(ContrastCurve * ContrastAndNormalizeMulAdd.y + ContrastAndNormalizeMulAdd.z);
+#endif
+
+	// Apply DFAO controls
+	SkyVisibility = pow(SkyVisibility, OcclusionExponent);
+	SkyVisibility = lerp(SkyVisibility, 1, OcclusionTintAndMinOcclusion.w);
+
+	// Combine with other AO sources
+	if (OcclusionCombineMode == 0)
+	{
+		// Combine with min which nicely avoids over-occlusion in cases where strong DFAO is present along with strong SSAO (distant trees)
+		SkyVisibility = min(SkyVisibility, min(GBufferAO, ScreenAO));
+	}
+	else
+	{
+		// Combine with mul, which continues to add SSAO depth even indoors.  SSAO will need to be tweaked to be less strong.
+		SkyVisibility = SkyVisibility * min(GBufferAO, ScreenAO);
+	}
+
+	// Apply AO to the sky diffuse and account for darkening due to the geometry term
+	// apply the Diffuse color to the lighting (including OcclusionTintAndMinOcclusion as it's considered another light, that fixes SubsurfaceProfile being too dark)
+	FSkyLightVisibilityData SkyVisData;
+	SkyVisData.SkyDiffuseLookUpMul = SkyVisibility * DotProductFactor;
+	SkyVisData.SkyDiffuseLookUpAdd = (1 - SkyVisibility) * OcclusionTintAndMinOcclusion.xyz;
+	SkyVisData.SkyDiffuseLookUpNormal = SkyLightingNormal;
+	return SkyVisData;
+}
+
+float3 SkyLightDiffuse(FGBufferData GBuffer, float AmbientOcclusion, float2 BufferUV, float2 ScreenPosition, float3 BentNormal, float3 DiffuseColor)
+{
+	float2 UV = BufferUV;
+	float3 Lighting = 0;
+
+	float3 SkyLightingNormal = GetClearCoatBottomNormal(GBuffer, GBuffer.WorldNormal);
+	FSkyLightVisibilityData SkyVisData = GetSkyLightVisibilityData(SkyLightingNormal, GBuffer.WorldNormal, GBuffer.GBufferAO, AmbientOcclusion, BentNormal);
+
+	const float3 N = GBuffer.WorldNormal;
+	const float3 V = -normalize(mul(float3(ScreenPosition, 1), LWCToFloat3x3(PrimaryView.ScreenToWorld)).xyz);
+	const float NoV = saturate(dot(V, N));
+	const FBxDFEnergyTerms EnergyTerms = ComputeGGXSpecEnergyTerms(GBuffer.Roughness, NoV, GBuffer.SpecularColor);
+	float3 DiffuseWeight = ComputeEnergyPreservation(EnergyTerms);
+
+	BRANCH
+	if (GBuffer.ShadingModelID == SHADINGMODELID_TWOSIDED_FOLIAGE)
+	{
+		float3 SubsurfaceLookup = GetSkySHDiffuse(-GBuffer.WorldNormal) * View.SkyLightColor.rgb;
+		float3 SubsurfaceColor = ExtractSubsurfaceColor(GBuffer);
+		Lighting += SkyVisData.SkyDiffuseLookUpMul * SubsurfaceLookup * SubsurfaceColor;
+	}
+
+	if (GBuffer.ShadingModelID == SHADINGMODELID_SUBSURFACE || GBuffer.ShadingModelID == SHADINGMODELID_PREINTEGRATED_SKIN)
+	{
+		float3 SubsurfaceColor = ExtractSubsurfaceColor(GBuffer);
+		// Add subsurface energy to diffuse
+		DiffuseColor += SubsurfaceColor;
+	}
+
+	BRANCH
+	if (GBuffer.ShadingModelID == SHADINGMODELID_HAIR)
+	{
+		float3 L = normalize(V - N * dot(V, N));
+		SkyVisData.SkyDiffuseLookUpNormal = L;
+		FHairTransmittanceData TransmittanceData = InitHairTransmittanceData(true);
+		DiffuseColor = PI * HairShading(GBuffer, L, V, N, 1, TransmittanceData, 0, 0.2, uint2(0, 0));
+		DiffuseWeight = 1.0f;
+	}
+
+	if (GBuffer.ShadingModelID == SHADINGMODELID_CLOTH)
+	{
+		float3 ClothFuzz = ExtractSubsurfaceColor(GBuffer);
+		DiffuseColor += ClothFuzz * GBuffer.CustomData.a;
+	}
+
+	// Compute the preconvolved incoming lighting with the bent normal direction
+	float3 DiffuseLookup = GetSkySHDiffuse(SkyVisData.SkyDiffuseLookUpNormal) * View.SkyLightColor.rgb;
+
+	// And accumulate the lighting
+	Lighting += (SkyVisData.SkyDiffuseLookUpMul * DiffuseLookup + SkyVisData.SkyDiffuseLookUpAdd) * DiffuseColor * DiffuseWeight;
+	Lighting *= View.PreExposure;
+
+	//Lighting = (Texture2DSampleLevel(BentNormalAOTexture, BentNormalAOSampler, UV, 0).xyz);
+
+	return Lighting;
+}
\ No newline at end of file
diff --git a/Engine/Shaders/Private/SkyLightingShared.ush b/Engine/Shaders/Private/SkyLightingShared.ush
--- a/Engine/Shaders/Private/SkyLightingShared.ush
+++ b/Engine/Shaders/Private/SkyLightingShared.ush
@@ -4,6 +4,8 @@
 	SkyLightingShared.usf
 =============================================================================*/
 
+#pragma once
+
 float ApproximateConeConeIntersection(float ArcLength0, float ArcLength1, float AngleBetweenCones)
 {
 	float AngleDifference = abs(ArcLength0 - ArcLength1);
diff --git a/Engine/Source/Runtime/Renderer/Private/HairStrands/HairStrandsEnvironment.cpp b/Engine/Source/Runtime/Renderer/Private/HairStrands/HairStrandsEnvironment.cpp
--- a/Engine/Source/Runtime/Renderer/Private/HairStrands/HairStrandsEnvironment.cpp
+++ b/Engine/Source/Runtime/Renderer/Private/HairStrands/HairStrandsEnvironment.cpp
@@ -21,6 +21,7 @@
 #include "ShaderDebug.h"
 #include "Lumen/LumenRadianceCache.h"
 #include "Lumen/LumenScreenProbeGather.h"
+#include "IndirectLightRendering.h"
 
 ///////////////////////////////////////////////////////////////////////////////////////////////////
 
@@ -230,22 +231,16 @@ public:
 
 	BEGIN_SHADER_PARAMETER_STRUCT(FParameters, )
 		SHADER_PARAMETER(float, Voxel_TanConeAngle)
-
 		SHADER_PARAMETER(uint32, MultipleScatterSampleCount)
-
 		SHADER_PARAMETER(float,  HairDualScatteringRoughnessOverride)
 		SHADER_PARAMETER(float, TransmissionDensityScaleFactor)
 		SHADER_PARAMETER(float, HairDistanceThreshold)
-
-		SHADER_PARAMETER(FVector4f, SkyLight_OcclusionTintAndMinOcclusion)
-
-		SHADER_PARAMETER(uint32, SkyLight_OcclusionCombineMode)
-		SHADER_PARAMETER(float, SkyLight_OcclusionExponent)
 		SHADER_PARAMETER(uint32, bHairUseViewHairCount)
 
 		SHADER_PARAMETER_TEXTURE(Texture2D, PreIntegratedGF)
 		SHADER_PARAMETER_SAMPLER(SamplerState, PreIntegratedGFSampler)
 
+		SHADER_PARAMETER_STRUCT_INCLUDE(FSkyDiffuseLightingParameters, SkyDiffuseLighting)
 		SHADER_PARAMETER_STRUCT_INCLUDE(FSceneTextureParameters, SceneTextures)
 
 		SHADER_PARAMETER_RDG_UNIFORM_BUFFER(FHairStrandsViewUniformParameters, HairStrands)
@@ -352,22 +347,8 @@ static void AddHairStrandsEnvironmentLightingPassPS(
 			case 2: IntegrationType = EHairLightingIntegrationType::SH; break;
 		}
 	}
-
-	float SkyLightContrast = 0.01f;
-	float SkyLightOcclusionExponent = 1.0f;
-	FVector4f SkyLightOcclusionTintAndMinOcclusion(0.0f, 0.0f, 0.0f, 0.0f);
-	EOcclusionCombineMode SkyLightOcclusionCombineMode = EOcclusionCombineMode::OCM_MAX;
-	if (FSkyLightSceneProxy* SkyLight = Scene->SkyLight)
-	{
-		SkyLightOcclusionExponent = SkyLight->OcclusionExponent;
-		SkyLightOcclusionTintAndMinOcclusion = FVector4f(SkyLight->OcclusionTint);
-		SkyLightOcclusionTintAndMinOcclusion.W = SkyLight->MinOcclusion;
-		SkyLightOcclusionCombineMode = SkyLight->OcclusionCombineMode;
-	}
-
-	PassParameters->SkyLight_OcclusionCombineMode = SkyLightOcclusionCombineMode == OCM_Minimum ? 0.0f : 1.0f;
-	PassParameters->SkyLight_OcclusionExponent = SkyLightOcclusionExponent;
-	PassParameters->SkyLight_OcclusionTintAndMinOcclusion = SkyLightOcclusionTintAndMinOcclusion;
+	
+	PassParameters->SkyDiffuseLighting = GetSkyDiffuseLightingParameters(Scene->SkyLight, 1.0f /*DynamicBentNormalAO*/);
 	PassParameters->Voxel_TanConeAngle = FMath::Tan(FMath::DegreesToRadians(GetHairStrandsSkyLightingConeAngle()));
 	PassParameters->HairDistanceThreshold = FMath::Max(GHairStrandsSkyLightingDistanceThreshold, 1.f);
 	PassParameters->bHairUseViewHairCount = VisibilityData.ViewHairCountTexture && GHairStrandsSkyLightingUseHairCountTexture ? 1 : 0;
diff --git a/Engine/Source/Runtime/Renderer/Private/IndirectLightRendering.cpp b/Engine/Source/Runtime/Renderer/Private/IndirectLightRendering.cpp
--- a/Engine/Source/Runtime/Renderer/Private/IndirectLightRendering.cpp
+++ b/Engine/Source/Runtime/Renderer/Private/IndirectLightRendering.cpp
@@ -191,6 +191,39 @@ class FAmbientCubemapCompositePS : public FGlobalShader
 	END_SHADER_PARAMETER_STRUCT()
 };
 
+// Setups all shader parameters related to skylight.
+FSkyDiffuseLightingParameters GetSkyDiffuseLightingParameters(const FSkyLightSceneProxy* SkyLight, float DynamicBentNormalAO)
+{
+	float SkyLightContrast = 0.01f;
+	float SkyLightOcclusionExponent = 1.0f;
+	FVector4f SkyLightOcclusionTintAndMinOcclusion(0.0f, 0.0f, 0.0f, 0.0f);
+	EOcclusionCombineMode SkyLightOcclusionCombineMode = EOcclusionCombineMode::OCM_MAX;
+	if (SkyLight)
+	{
+		FDistanceFieldAOParameters Parameters(SkyLight->OcclusionMaxDistance, SkyLight->Contrast);
+		SkyLightContrast = Parameters.Contrast;
+		SkyLightOcclusionExponent = SkyLight->OcclusionExponent;
+		SkyLightOcclusionTintAndMinOcclusion = FVector4f(SkyLight->OcclusionTint);
+		SkyLightOcclusionTintAndMinOcclusion.W = SkyLight->MinOcclusion;
+		SkyLightOcclusionCombineMode = SkyLight->OcclusionCombineMode;
+	}
+
+	// Scale and bias to remap the contrast curve to [0,1]
+	const float Min = 1 / (1 + FMath::Exp(-SkyLightContrast * (0 * 10 - 5)));
+	const float Max = 1 / (1 + FMath::Exp(-SkyLightContrast * (1 * 10 - 5)));
+	const float Mul = 1.0f / (Max - Min);
+	const float Add = -Min / (Max - Min);
+
+	FSkyDiffuseLightingParameters Out;
+	Out.OcclusionTintAndMinOcclusion = SkyLightOcclusionTintAndMinOcclusion;
+	Out.ContrastAndNormalizeMulAdd = FVector(SkyLightContrast, Mul, Add);
+	Out.OcclusionExponent = SkyLightOcclusionExponent;
+	Out.OcclusionCombineMode = SkyLightOcclusionCombineMode == OCM_Minimum ? 0.0f : 1.0f;
+	Out.ApplyBentNormalAO = DynamicBentNormalAO;
+	Out.InvSkySpecularOcclusionStrength = 1.0f / FMath::Max(CVarSkySpecularOcclusionStrength.GetValueOnRenderThread(), 0.1f);
+	return Out;
+}
+
 /** Pixel shader that does tiled deferred culling of reflection captures, then sorts and composites them. */
 class FReflectionEnvironmentSkyLightingPS : public FGlobalShader
 {
@@ -280,15 +313,6 @@ class FReflectionEnvironmentSkyLightingPS : public FGlobalShader
 	}
 
 	BEGIN_SHADER_PARAMETER_STRUCT(FParameters, )
-
-		// Sky light parameters.
-		SHADER_PARAMETER(FVector4f, OcclusionTintAndMinOcclusion)
-		SHADER_PARAMETER(FVector3f, ContrastAndNormalizeMulAdd)
-		SHADER_PARAMETER(float, ApplyBentNormalAO)
-		SHADER_PARAMETER(float, InvSkySpecularOcclusionStrength)
-		SHADER_PARAMETER(float, OcclusionExponent)
-		SHADER_PARAMETER(float, OcclusionCombineMode)
-
 		// Distance field AO parameters.
 		// TODO. FDFAOUpsampleParameters
 		SHADER_PARAMETER(FVector2f, AOBufferBilinearUVMax)
@@ -313,6 +337,7 @@ class FReflectionEnvironmentSkyLightingPS : public FGlobalShader
 		SHADER_PARAMETER(float, CloudSkyAOFarDepthKm)
 		SHADER_PARAMETER(int32, CloudSkyAOEnabled)
 
+		SHADER_PARAMETER_STRUCT_INCLUDE(FSkyDiffuseLightingParameters, SkyDiffuseLighting)
 		SHADER_PARAMETER_STRUCT_INCLUDE(FSceneTextureParameters, SceneTextures)
 
 		SHADER_PARAMETER_STRUCT_REF(FViewUniformShaderParameters, ViewUniformBuffer)
@@ -1184,36 +1209,7 @@ static void AddSkyReflectionPass(
 	// Setup the parameters of the shader.
 	{
 		// Setups all shader parameters related to skylight.
-		{
-			FSkyLightSceneProxy* SkyLight = Scene->SkyLight;
-
-			float SkyLightContrast = 0.01f;
-			float SkyLightOcclusionExponent = 1.0f;
-			FVector4f SkyLightOcclusionTintAndMinOcclusion(0.0f, 0.0f, 0.0f, 0.0f);
-			EOcclusionCombineMode SkyLightOcclusionCombineMode = EOcclusionCombineMode::OCM_MAX;
-			if (SkyLight)
-			{
-				FDistanceFieldAOParameters Parameters(SkyLight->OcclusionMaxDistance, SkyLight->Contrast);
-				SkyLightContrast = Parameters.Contrast;
-				SkyLightOcclusionExponent = SkyLight->OcclusionExponent;
-				SkyLightOcclusionTintAndMinOcclusion = FVector4f(SkyLight->OcclusionTint);
-				SkyLightOcclusionTintAndMinOcclusion.W = SkyLight->MinOcclusion;
-				SkyLightOcclusionCombineMode = SkyLight->OcclusionCombineMode;
-			}
-
-			// Scale and bias to remap the contrast curve to [0,1]
-			const float Min = 1 / (1 + FMath::Exp(-SkyLightContrast * (0 * 10 - 5)));
-			const float Max = 1 / (1 + FMath::Exp(-SkyLightContrast * (1 * 10 - 5)));
-			const float Mul = 1.0f / (Max - Min);
-			const float Add = -Min / (Max - Min);
-
-			PassParameters->PS.OcclusionTintAndMinOcclusion = SkyLightOcclusionTintAndMinOcclusion;
-			PassParameters->PS.ContrastAndNormalizeMulAdd = FVector(SkyLightContrast, Mul, Add);
-			PassParameters->PS.OcclusionExponent = SkyLightOcclusionExponent;
-			PassParameters->PS.OcclusionCombineMode = SkyLightOcclusionCombineMode == OCM_Minimum ? 0.0f : 1.0f;
-			PassParameters->PS.ApplyBentNormalAO = DynamicBentNormalAO;
-			PassParameters->PS.InvSkySpecularOcclusionStrength = 1.0f / FMath::Max(CVarSkySpecularOcclusionStrength.GetValueOnRenderThread(), 0.1f);
-		}
+		PassParameters->PS.SkyDiffuseLighting = GetSkyDiffuseLightingParameters(Scene->SkyLight, DynamicBentNormalAO);
 
 		// Setups all shader parameters related to distance field AO
 		{
diff --git a/Engine/Source/Runtime/Renderer/Private/IndirectLightRendering.h b/Engine/Source/Runtime/Renderer/Private/IndirectLightRendering.h
--- a/Engine/Source/Runtime/Renderer/Private/IndirectLightRendering.h
+++ b/Engine/Source/Runtime/Renderer/Private/IndirectLightRendering.h
@@ -52,3 +52,15 @@ BEGIN_SHADER_PARAMETER_STRUCT(FCommonParameters, )
 END_SHADER_PARAMETER_STRUCT()
 
 } // HybridIndirectLighting
+
+// Sky light parameters.
+BEGIN_SHADER_PARAMETER_STRUCT(FSkyDiffuseLightingParameters, )
+	SHADER_PARAMETER(FVector4f, OcclusionTintAndMinOcclusion)
+	SHADER_PARAMETER(FVector3f, ContrastAndNormalizeMulAdd)
+	SHADER_PARAMETER(float, ApplyBentNormalAO)
+	SHADER_PARAMETER(float, InvSkySpecularOcclusionStrength)
+	SHADER_PARAMETER(float, OcclusionExponent)
+	SHADER_PARAMETER(float, OcclusionCombineMode)
+END_SHADER_PARAMETER_STRUCT()
+
+FSkyDiffuseLightingParameters GetSkyDiffuseLightingParameters(const FSkyLightSceneProxy* SkyLight, float DynamicBentNormalAO);
