Add utility function to convert FHairSample into FGBufferData.

diff --git a/Engine/Shaders/Private/ClusteredDeferredShadingPixelShader.usf b/Engine/Shaders/Private/ClusteredDeferredShadingPixelShader.usf
--- a/Engine/Shaders/Private/ClusteredDeferredShadingPixelShader.usf
+++ b/Engine/Shaders/Private/ClusteredDeferredShadingPixelShader.usf
@@ -27,6 +27,7 @@
 #if USE_HAIR_LIGHTING == 1
 	#include "HairStrands/HairStrandsCommon.ush"
 	#include "HairStrands/HairStrandsVisibilityCommon.ush"
+	#include "HairStrands/HairStrandsVisibilityUtils.ush"
 	#include "HairStrands/HairStrandsDeepTransmittanceCommon.ush"
 	#include "HairStrands/HairStrandsDeepTransmittanceDualScattering.ush" 
 #endif // USE_HAIR_LIGHTING
@@ -391,22 +392,8 @@ void ClusteredShadingPixelShader(
 	const float2 LocalPosition = PixelCoord - View.ViewRectMin.xy;
 	FScreenSpaceData ScreenSpaceData = (FScreenSpaceData)0;
 	{
-		// HAIR_TODO: make a HairSampleToGbuffer helper function
 		ScreenSpaceData.AmbientOcclusion = 1;
-		ScreenSpaceData.GBuffer.WorldNormal = Sample.Tangent;
-		ScreenSpaceData.GBuffer.BaseColor = Sample.BaseColor;
-		ScreenSpaceData.GBuffer.Roughness = Sample.Roughness;
-		ScreenSpaceData.GBuffer.ShadingModelID = SHADINGMODELID_HAIR;
-		ScreenSpaceData.GBuffer.DiffuseColor = 0;
-		ScreenSpaceData.GBuffer.SpecularColor = 0;
-		ScreenSpaceData.GBuffer.Specular = Sample.Specular;
-		ScreenSpaceData.GBuffer.Metallic = 0;
-		ScreenSpaceData.GBuffer.Depth = SceneDepth;
-		ScreenSpaceData.GBuffer.GBufferAO = 1;
-		ScreenSpaceData.GBuffer.CustomData = float4(HairStrands.HairDualScatteringRoughnessOverride, 0, Sample.Backlit, 0);
-		ScreenSpaceData.GBuffer.IndirectIrradiance = 0;
-		ScreenSpaceData.GBuffer.PrecomputedShadowFactors = 1;
-		ScreenSpaceData.GBuffer.PerObjectGBufferData = 0;
+		ScreenSpaceData.GBuffer = HairSampleToGBufferData(Sample, HairStrands.HairDualScatteringRoughnessOverride);
 	}
 
 	const float Dither = InterleavedGradientNoise(SvPosition.xy, View.StateFrameIndexMod8);
diff --git a/Engine/Shaders/Private/DeferredLightPixelShaders.usf b/Engine/Shaders/Private/DeferredLightPixelShaders.usf
--- a/Engine/Shaders/Private/DeferredLightPixelShaders.usf
+++ b/Engine/Shaders/Private/DeferredLightPixelShaders.usf
@@ -22,6 +22,7 @@
 
 #if USE_HAIR_LIGHTING == 1
 #include "HairStrands/HairStrandsVisibilityCommon.ush"
+#include "HairStrands/HairStrandsVisibilityUtils.ush"
 #endif
 #include "HairStrands/HairStrandsCommon.ush"
 #include "HairStrands/HairStrandsDeepTransmittanceCommon.ush"
@@ -364,20 +365,7 @@ void DeferredLightPixelMain(
 	const float SceneDepth = ConvertFromDeviceZ(Sample.Depth);
 	FScreenSpaceData HairScreenSpaceData = (FScreenSpaceData)0;
 	HairScreenSpaceData.AmbientOcclusion = 1;
-	HairScreenSpaceData.GBuffer.WorldNormal = Sample.Tangent;
-	HairScreenSpaceData.GBuffer.BaseColor = Sample.BaseColor;
-	HairScreenSpaceData.GBuffer.Roughness = Sample.Roughness;
-	HairScreenSpaceData.GBuffer.ShadingModelID = SHADINGMODELID_HAIR;
-	HairScreenSpaceData.GBuffer.DiffuseColor = 0;
-	HairScreenSpaceData.GBuffer.SpecularColor = 0;
-	HairScreenSpaceData.GBuffer.Specular = Sample.Specular;
-	HairScreenSpaceData.GBuffer.Metallic = 0;
-	HairScreenSpaceData.GBuffer.Depth = SceneDepth;
-	HairScreenSpaceData.GBuffer.GBufferAO = 1;
-	HairScreenSpaceData.GBuffer.CustomData = float4(HairStrands.HairDualScatteringRoughnessOverride, 0, Sample.Backlit, 0);
-	HairScreenSpaceData.GBuffer.IndirectIrradiance = 0;
-	HairScreenSpaceData.GBuffer.PrecomputedShadowFactors = 1;
-	HairScreenSpaceData.GBuffer.PerObjectGBufferData = 0;
+	HairScreenSpaceData.GBuffer = HairSampleToGBufferData(Sample, HairStrands.HairDualScatteringRoughnessOverride);
 
 	const float3 TranslatedWorldPosition = mul(float4(ScreenPosition * HairScreenSpaceData.GBuffer.Depth, HairScreenSpaceData.GBuffer.Depth, 1), View.ScreenToTranslatedWorld).xyz;
 	const float3 CameraVector = normalize(TranslatedWorldPosition - View.TranslatedWorldCameraOrigin);
diff --git a/Engine/Shaders/Private/HairStrands/HairStrandsDeepTransmittanceMask.usf b/Engine/Shaders/Private/HairStrands/HairStrandsDeepTransmittanceMask.usf
--- a/Engine/Shaders/Private/HairStrands/HairStrandsDeepTransmittanceMask.usf
+++ b/Engine/Shaders/Private/HairStrands/HairStrandsDeepTransmittanceMask.usf
@@ -541,13 +541,6 @@ void MainCS(uint2 DispatchThreadId : SV_DispatchThreadID)
 			return;
 		}
 
-		FGBufferData GBuffer = (FGBufferData)0;
-		GBuffer.ShadingModelID = SHADINGMODELID_HAIR;
-		GBuffer.WorldNormal = Sample.Tangent;
-		GBuffer.BaseColor = Sample.BaseColor;
-		GBuffer.Roughness = Sample.Roughness;
-		GBuffer.Specular = Sample.Specular;
-
 		const uint MacroGroupIndex = clamp(Sample.MacroGroupId, 0, MAX_HAIR_MACROGROUP_COUNT - 1);
 
 		FTransmittanceSettings Settings;
diff --git a/Engine/Shaders/Private/HairStrands/HairStrandsEnvironmentLighting.usf b/Engine/Shaders/Private/HairStrands/HairStrandsEnvironmentLighting.usf
--- a/Engine/Shaders/Private/HairStrands/HairStrandsEnvironmentLighting.usf
+++ b/Engine/Shaders/Private/HairStrands/HairStrandsEnvironmentLighting.usf
@@ -34,6 +34,7 @@
 
 #include "../HairBsdf.ush"
 #include "HairStrandsVisibilityCommon.ush"
+#include "HairStrandsVisibilityUtils.ush"
 #include "HairStrandsEnvironmentLightingCommon.ush"
 #include "HairStrandsVoxelPageCommon.ush"
 #include "HairStrandsDeepShadowCommon.ush"
@@ -730,21 +731,8 @@ void MainPS(
 	const float2 UV = (PixelCoord + float2(0.5f, 0.5f)) / float2(View.BufferSizeAndInvSize.xy);
 	const float2 ScreenPosition = (UV - View.ScreenPositionScaleBias.wz) / View.ScreenPositionScaleBias.xy;
 
-
-
 	const FHairSample Sample = UnpackHairSample(HairStrands.HairSampleData[LocalOffset]);
-
-	FGBufferData GBuffer = (FGBufferData)0;
-	GBuffer.ShadingModelID = SHADINGMODELID_HAIR;
-	GBuffer.WorldNormal = Sample.Tangent;
-	GBuffer.BaseColor = Sample.BaseColor;
-	GBuffer.Roughness = Sample.Roughness;
-	GBuffer.Specular = Sample.Specular;
-	GBuffer.Metallic = 0; // Scattering;
-	GBuffer.Depth = ConvertFromDeviceZ(Sample.Depth);
-	GBuffer.GBufferAO = 1;
-	GBuffer.CustomData = float4(HairStrands.HairDualScatteringRoughnessOverride, 0, Sample.Backlit, 0);
-	GBuffer.IndirectIrradiance = 1000000;
+	const FGBufferData GBuffer = HairSampleToGBufferData(Sample, HairStrands.HairDualScatteringRoughnessOverride);
 
 #if PERMUTATION_DEBUG
 	const bool bDebugEnabled = all(int2(PixelCoord) == GetCursorPos());
diff --git a/Engine/Shaders/Private/HairStrands/HairStrandsVisibilityUtils.ush b/Engine/Shaders/Private/HairStrands/HairStrandsVisibilityUtils.ush
new file mode 100644
--- /dev/null
+++ b/Engine/Shaders/Private/HairStrands/HairStrandsVisibilityUtils.ush
@@ -0,0 +1,24 @@
+// Copyright Epic Games, Inc. All Rights Reserved.
+
+#pragma once
+
+// Convert a FHairSample into a FGufferData
+FGBufferData HairSampleToGBufferData(in FHairSample In, in float InDualScatteringRoughnessOverride=0)
+{
+	FGBufferData Out = (FGBufferData)0;
+	Out.ShadingModelID = SHADINGMODELID_HAIR;
+	Out.WorldNormal = In.Tangent;
+	Out.BaseColor = In.BaseColor;
+	Out.Roughness = In.Roughness;
+	Out.Specular = In.Specular;
+	Out.Metallic = 0; // Scattering;
+	Out.Depth = ConvertFromDeviceZ(In.Depth);
+	Out.GBufferAO = 1;
+	Out.CustomData = float4(InDualScatteringRoughnessOverride, 0, In.Backlit, 0);
+	Out.IndirectIrradiance = 1000000;
+	Out.DiffuseColor = 0;
+	Out.SpecularColor = 0;
+	Out.PrecomputedShadowFactors = 1;
+	Out.PerObjectGBufferData = 0;
+	return Out;
+}
\ No newline at end of file
