Improve precision for LWCFrac, LWC wrap texture address mode

diff --git a/Engine/Shaders/Private/LWCOperations.ush b/Engine/Shaders/Private/LWCOperations.ush
--- a/Engine/Shaders/Private/LWCOperations.ush
+++ b/Engine/Shaders/Private/LWCOperations.ush
@@ -111,8 +111,15 @@ FLWCType LWCSelect(FBoolType S, FLWCType Lhs, FFloatType Rhs) { return LWCConstr
 
 FLWCType LWCNegate(FLWCType V) { return LWCConstructor(-LWCGetTile(V), -V.Offset); }
 
-// FFloatType LWCFrac(FLWCType V) { return frac(frac(LWCGetTile(V)) * UE_LWC_RENDER_TILE_SIZE + V.Offset); }
-FFloatType LWCFrac(FLWCType V) { FLWCType VN = LWCNormalizeTile(V); return frac(VN.Offset); }
+FFloatType LWCFrac(FLWCType V)
+{
+	// Inner frac() avoids overflow if tile is large
+	// Outer frac() avoids helps avoid precision loss in the 'FracTile + V.Offset' addition
+	FFloatType FracTile = frac(frac(LWCGetTile(V)) * UE_LWC_RENDER_TILE_SIZE);
+
+	return frac(FracTile + V.Offset);
+}
+//FFloatType LWCFrac(FLWCType V) { FLWCType VN = LWCNormalizeTile(V); return frac(VN.Offset); }
 
 FLWCType LWCFloor(FLWCType V) { FLWCType VN = LWCNormalizeTile(V); return LWCConstructor(LWCGetTile(VN), floor(VN.Offset)); }
 FLWCType LWCCeil(FLWCType V) { FLWCType VN = LWCNormalizeTile(V); return LWCConstructor(LWCGetTile(VN), ceil(VN.Offset)); }
diff --git a/Engine/Shaders/Private/MaterialTemplate.ush b/Engine/Shaders/Private/MaterialTemplate.ush
--- a/Engine/Shaders/Private/MaterialTemplate.ush
+++ b/Engine/Shaders/Private/MaterialTemplate.ush
@@ -1838,22 +1838,39 @@ float GetSphericalParticleOpacity(FMaterialPixelParameters Parameters, float Den
 #define LWCADDRESSMODE_WRAP 1u
 #define LWCADDRESSMODE_MIRROR 2u
 
+float LWCApplyAddressModeWrap(FLWCScalar V)
+{
+	// Compute the fractional part of the tile, then add to the offset
+	// Let the texture unit apply the final coordinate wrapping, which will allow derivatives to work correctly unless we cross a tile boundary
+	const float FracTile = frac(frac(LWCGetTile(V)) * UE_LWC_RENDER_TILE_SIZE);
+	return FracTile + V.Offset;
+}
+
 float LWCApplyAddressModeMirror(FLWCScalar v)
 {
-	float t = LWCFrac(LWCMultiply(v, 0.5f)) * 2.0f;
-	return 1.0f - abs(t - 1.0f);
+	// Unclear what the best option is for MIRROR
+	// We can apply the mirror logic directly, but that will break derivatives
+	// We can use similar logic as WRAP, but in that case results will actually be incorrect at tile boundaries (not just wrong derivatives)
+	// Or we can convert to float and accept precision loss for large values (but otherwise correct)
+	// TODO - something better?
+
+	//float t = LWCFrac(LWCMultiply(v, 0.5f)) * 2.0f;
+	//return 1.0f - abs(t - 1.0f);
+	return LWCToFloat(v);
 }
 
-float LWCApplyAddressMode(FLWCScalar v, uint AddressMode)
+float LWCApplyAddressModeClamp(FLWCScalar v)
 {
-	// For wrap/mirror, we can take the offset of the normalized tile
-	// This lets the GPU handle wrapping in most cases which will preserve correct derivatives.  There will be a seam along tile boundaries in the UV however
-	// May be possible to improve this by wrapping at higher values than TileSize, which would reduce frequency (but not eliminate) the seam
 	// For the CLAMP case, a simple LWCToFloat() is sufficient.  This will lose a ton of precision for large values, but we don't care about this since the GPU will clamp to [0,1) anyway
 	// It's possible certain GPUs might need a special case if the ToFloat() conversion overflows
-	if(AddressMode == LWCADDRESSMODE_WRAP) return LWCNormalizeTile(v).Offset;
-	else if(AddressMode == LWCADDRESSMODE_MIRROR) return LWCNormalizeTile(v).Offset;
-	else return LWCToFloat(v);
+	return LWCToFloat(v);
+}
+
+float LWCApplyAddressMode(FLWCScalar v, uint AddressMode)
+{
+	if(AddressMode == LWCADDRESSMODE_WRAP) return LWCApplyAddressModeWrap(v);
+	else if(AddressMode == LWCADDRESSMODE_MIRROR) return LWCApplyAddressModeMirror(v);
+	else return LWCApplyAddressModeClamp(v);
 }
 float2 LWCApplyAddressMode(FLWCVector2 UV, uint AddressX, uint AddressY)
 {
