Compiles TSR shaders with CFLAG_WarningsAsErrors

diff --git a/Engine/Shaders/Private/LWCOperations.ush b/Engine/Shaders/Private/LWCOperations.ush
--- a/Engine/Shaders/Private/LWCOperations.ush
+++ b/Engine/Shaders/Private/LWCOperations.ush
@@ -127,7 +127,7 @@ FLWCType LWCRound(FLWCType V) { FLWCType VN = LWCNormalizeTile(V); return LWCCon
 FLWCType LWCTrunc(FLWCType V) { FLWCType VN = LWCNormalizeTile(V); return LWCConstructor(LWCGetTile(VN), trunc(VN.Offset)); }
 
 // These rely on +/-inf being handled properly by various min/max operations
-FFloatType LWCSign(FLWCType V) { return sign(LWCToFloat(V)); }
+FFloatType LWCSign(FLWCType V) { return FFloatType(sign(LWCToFloat(V))); }
 FFloatType LWCSaturate(FLWCType V) { return saturate(LWCToFloat(V)); }
 FFloatType LWCClampScalar(FLWCType V, float Low, float High) { return clamp(LWCToFloat(V), Low, High); }
 
diff --git a/Engine/Shaders/Private/MonteCarlo.ush b/Engine/Shaders/Private/MonteCarlo.ush
--- a/Engine/Shaders/Private/MonteCarlo.ush
+++ b/Engine/Shaders/Private/MonteCarlo.ush
@@ -86,7 +86,7 @@ float2 JitteredR2( uint Index, uint NumSamples, float2 Jitter, float JitterAmoun
 	const float d0 = 0.76;
 	const float i0 = 0.7;
 
-	return frac( a * Index + ( JitterAmount * 0.5 * d0 * sqrt(PI) * rsqrt( NumSamples ) ) * Jitter );
+	return frac( a * float(Index) + ( JitterAmount * 0.5 * d0 * sqrt(PI) * rsqrt( float(NumSamples) ) ) * Jitter );
 }
 
 // R2 Jittered point sequence. Progressive
diff --git a/Engine/Shaders/Private/TemporalSuperResolution/TSRClearPrevTextures.usf b/Engine/Shaders/Private/TemporalSuperResolution/TSRClearPrevTextures.usf
--- a/Engine/Shaders/Private/TemporalSuperResolution/TSRClearPrevTextures.usf
+++ b/Engine/Shaders/Private/TemporalSuperResolution/TSRClearPrevTextures.usf
@@ -22,7 +22,7 @@ void MainCS(
 	uint GroupThreadIndex : SV_GroupIndex)
 {
 	uint2 DispatchThreadId = (
-		ZOrder2D(GroupThreadIndex, log2(TILE_SIZE)) +
+		ZOrder2D(GroupThreadIndex, uint(log2(float(TILE_SIZE)))) +
 		GroupId * uint2(TILE_SIZE, TILE_SIZE));
 
 	uint2 InputPixelPos = InputInfo_ViewportMin + DispatchThreadId;
diff --git a/Engine/Shaders/Private/TemporalSuperResolution/TSRCommon.ush b/Engine/Shaders/Private/TemporalSuperResolution/TSRCommon.ush
--- a/Engine/Shaders/Private/TemporalSuperResolution/TSRCommon.ush
+++ b/Engine/Shaders/Private/TemporalSuperResolution/TSRCommon.ush
@@ -348,10 +348,13 @@ uint GetGroupThreadIndex(uint GroupThreadIndex, uint GroupWaveIndex)
 	uint GroupWaveOffset = WaveGetLaneCount() * GroupWaveIndex;
 
 	// Do not share
+	uint ComputedGroupThreadIndex;
 	ISOLATE
 	{
-		return GroupWaveOffset + WaveGetLaneIndex();
+		ComputedGroupThreadIndex = GroupWaveOffset + WaveGetLaneIndex();
 	}
+
+	return ComputedGroupThreadIndex;
 }
 #else
 {
@@ -1118,7 +1121,7 @@ uint EncodeHoleFillingVelocity(float2 PixelVelocity)
 	uint iEncodedAngle = uint(round(EncodedAngle)) & (AngleBitDepth - 1);
 
 	float EncodedLength = length(PixelVelocity) * LengthPixelPrecision;
-	uint iEncodedLength = clamp(ceil(EncodedLength), 0, LengthBitDepth - 1);
+	uint iEncodedLength = clamp(uint(ceil(EncodedLength)), 0u, LengthBitDepth - 1u);
 
 	uint EncodedHoleFillingVelocity = iEncodedAngle * LengthBitDepth + iEncodedLength;
 	return EncodedHoleFillingVelocity;
diff --git a/Engine/Shaders/Private/TemporalSuperResolution/TSRCompareHistory.usf b/Engine/Shaders/Private/TemporalSuperResolution/TSRCompareHistory.usf
--- a/Engine/Shaders/Private/TemporalSuperResolution/TSRCompareHistory.usf
+++ b/Engine/Shaders/Private/TemporalSuperResolution/TSRCompareHistory.usf
@@ -71,7 +71,7 @@ void MainCS(
 		
 		taa_short2 SampleInputPixelPos = InputPixelPos + ClampPixelOffset(
 			InputPixelPos,
-			kOffsets3x3[i], kOffsets3x3[i],
+			taa_short2(kOffsets3x3[i]), taa_short2(kOffsets3x3[i]),
 			InputPixelPosMin, InputPixelPosMax);
 
 		taa_halfC Color = TransformColorForClampingBox(FilteredInputTexture[SampleInputPixelPos]);
diff --git a/Engine/Shaders/Private/TemporalSuperResolution/TSRCompareTranslucency.usf b/Engine/Shaders/Private/TemporalSuperResolution/TSRCompareTranslucency.usf
--- a/Engine/Shaders/Private/TemporalSuperResolution/TSRCompareTranslucency.usf
+++ b/Engine/Shaders/Private/TemporalSuperResolution/TSRCompareTranslucency.usf
@@ -105,8 +105,8 @@ void MainCS(
 {
 	float4 Debug = 0;
 	
-	uint2 InputPixelPos = (
-		ZOrder2D(GroupThreadIndex, log2(TILE_SIZE)) +
+	taa_short2 InputPixelPos = taa_short2(
+		ZOrder2D(GroupThreadIndex, uint(log2(float(TILE_SIZE)))) +
 		GroupId * uint2(TILE_SIZE, TILE_SIZE));
 
 	float2 ScreenPos = ApplyScreenTransform(float2(InputPixelPos), InputPixelPosToScreenPos);
@@ -226,9 +226,9 @@ void MainCS(
 			PrevTranslucency *= rcp(taa_half(CONFIG_SAMPLES));
 		#endif
 
-		PrevTranslucency *= PrevTranslucencyPreExposureCorrection;
-		PrevTranslucencyMin *= PrevTranslucencyPreExposureCorrection;
-		PrevTranslucencyMax *= PrevTranslucencyPreExposureCorrection;
+		PrevTranslucency    *= taa_half(PrevTranslucencyPreExposureCorrection);
+		PrevTranslucencyMin *= taa_half(PrevTranslucencyPreExposureCorrection);
+		PrevTranslucencyMax *= taa_half(PrevTranslucencyPreExposureCorrection);
 	}
 
 	taa_half Rejection = 1.0;
@@ -251,7 +251,7 @@ void MainCS(
 	}
 	#endif
 
-	if (all(InputPixelPos < InputInfo_ViewportMax))
+	if (all(InputPixelPos < taa_short2(InputInfo_ViewportMax)))
 	{
 		TranslucencyRejectionOutput[InputPixelPos] = Rejection;
 
diff --git a/Engine/Shaders/Private/TemporalSuperResolution/TSRDebugHistory.usf b/Engine/Shaders/Private/TemporalSuperResolution/TSRDebugHistory.usf
--- a/Engine/Shaders/Private/TemporalSuperResolution/TSRDebugHistory.usf
+++ b/Engine/Shaders/Private/TemporalSuperResolution/TSRDebugHistory.usf
@@ -69,7 +69,7 @@ void MainCS(
 	float4 Debug = 0.0;
 
 	uint2 DispatchThreadId = (
-		ZOrder2D(GroupThreadIndex, log2(TILE_SIZE)) +
+		ZOrder2D(GroupThreadIndex, uint(log2(float(TILE_SIZE)))) +
 		GroupId * uint2(TILE_SIZE, TILE_SIZE));
 
 	uint2 HistoryPixelPos = DispatchThreadId / SUB_PIXEL_GRID_SIZE;
@@ -96,7 +96,7 @@ void MainCS(
 		taa_subpixel_payload SubpixelDetails = taa_subpixel_payload(RawHistory3 >> SubpixelShift) & SUB_PIXEL_BIT_MASK;
 
 		Debug.r = UncompressParallaxFactor(SubpixelDetails);
-		Debug.b = SubpixelId;
+		Debug.b = float(SubpixelId);
 	}
 
 	#if DEBUG_OUTPUT && 0
diff --git a/Engine/Shaders/Private/TemporalSuperResolution/TSRDecimateHistory.usf b/Engine/Shaders/Private/TemporalSuperResolution/TSRDecimateHistory.usf
--- a/Engine/Shaders/Private/TemporalSuperResolution/TSRDecimateHistory.usf
+++ b/Engine/Shaders/Private/TemporalSuperResolution/TSRDecimateHistory.usf
@@ -70,7 +70,7 @@ void MainCS(
 	float4 Debug = 0;
 	
 	uint2 DispatchThreadId = (
-		ZOrder2D(GroupThreadIndex, log2(TILE_SIZE)) +
+		ZOrder2D(GroupThreadIndex, uint(log2(float(TILE_SIZE)))) +
 		GroupId * uint2(TILE_SIZE, TILE_SIZE));
 
 	uint2 InputPixelPos = InputInfo_ViewportMin + DispatchThreadId;
@@ -225,7 +225,7 @@ void MainCS(
 		float4 RotationalPrevClip = mul(ThisClip, RotationalClipToPrevClip);
 		float2 RotationalPrevScreen = RotationalPrevClip.xy / RotationalPrevClip.w;
 			
-		taa_half PixelParallaxVelocity = taa_half(0.5) * length(taa_half2(RotationalPrevScreen - PrevScreen) * InputInfo_ViewportSize);
+		taa_half PixelParallaxVelocity = taa_half(0.5) * length(taa_half2(RotationalPrevScreen - PrevScreen) * taa_half2(InputInfo_ViewportSize));
 		taa_half HistorySubdetailAdditionalParralax = (HistorySubdetailRejectionFactor - taa_half(1.0)) * PixelParallaxVelocity;
 
 		// PixelParallaxVelocity can be +inf and generate NaN if HistorySubdetailRejectionFactor = 1.0;
@@ -295,10 +295,10 @@ void MainCS(
 		#if CONFIG_ENABLE_STOCASTIC_QUANTIZATION
 		{
 			uint2 Random = Rand3DPCG16(int3((InputPixelPos - InputInfo_ViewportMin) / 2, View.StateFrameIndexMod8)).xy;
-			float2 E = Hammersley16(0, 1, Random);
+			taa_half E = taa_half(Hammersley16(0, 1, Random).x);
 
-			HalfResInputSceneColor = QuantizeForFloatRenderTarget(HalfResInputSceneColor, E.x, OutputQuantizationError);
-			HalfResDecimatedColor = QuantizeForFloatRenderTarget(HalfResDecimatedColor, E.x, OutputQuantizationError);
+			HalfResInputSceneColor = QuantizeForFloatRenderTarget(HalfResInputSceneColor, E, OutputQuantizationError);
+			HalfResDecimatedColor = QuantizeForFloatRenderTarget(HalfResDecimatedColor, E, OutputQuantizationError);
 		}
 		#endif
 
@@ -314,9 +314,9 @@ void MainCS(
 	#if CONFIG_ENABLE_STOCASTIC_QUANTIZATION
 	{
 		uint2 Random = Rand3DPCG16(int3(InputPixelPos - InputInfo_ViewportMin, View.StateFrameIndexMod8)).xy;
-		float2 E = Hammersley16(0, 1, Random);
+		taa_half E = taa_half(Hammersley16(0, 1, Random).x);
 
-		DecimatedColor = QuantizeForFloatRenderTarget(DecimatedColor, E.x, OutputQuantizationError);
+		DecimatedColor = QuantizeForFloatRenderTarget(DecimatedColor, E, OutputQuantizationError);
 	}
 	#endif
 
diff --git a/Engine/Shaders/Private/TemporalSuperResolution/TSRDilateVelocity.usf b/Engine/Shaders/Private/TemporalSuperResolution/TSRDilateVelocity.usf
--- a/Engine/Shaders/Private/TemporalSuperResolution/TSRDilateVelocity.usf
+++ b/Engine/Shaders/Private/TemporalSuperResolution/TSRDilateVelocity.usf
@@ -71,7 +71,7 @@ void MainCS(
 	float4 Debug = 0.0;
 
 	uint2 DispatchThreadId = (
-		ZOrder2D(GroupThreadIndex, log2(TILE_SIZE)) +
+		ZOrder2D(GroupThreadIndex, uint(log2(float(TILE_SIZE)))) +
 		GroupId * uint2(TILE_SIZE, TILE_SIZE));
 
 	uint2 InputPixelPos = InputInfo_ViewportMin + DispatchThreadId;
@@ -83,7 +83,7 @@ void MainCS(
 	float DeviceZ = SceneDepthTexture[InputPixelPos];
 
 	float ClosestDeviceZ = DeviceZ;
-	float2 VelocityPixelOffset = float2(0.0, 0.0);
+	int2 VelocityPixelOffset = int2(0, 0);
 	float ParallaxFactorForHistory = 1.0;
 	float ParallaxFactorForHistoryRejection = 1.0;
 	#if AA_CROSS
@@ -158,8 +158,8 @@ void MainCS(
 	BRANCH
 	if (VelocityExtrapolationMultiplier > 0.0)
 	{
-		int2 PixelCoord1 = int2(InputBufferUV * InputInfo_Extent + VelocityPixelOffset * 1.0);
-		int2 PixelCoord2 = int2(InputBufferUV * InputInfo_Extent + VelocityPixelOffset * 2.0);
+		int2 PixelCoord1 = int2(InputBufferUV * InputInfo_Extent) + VelocityPixelOffset * 1;
+		int2 PixelCoord2 = int2(InputBufferUV * InputInfo_Extent) + VelocityPixelOffset * 2;
 
 		PixelCoord2 = clamp(PixelCoord2, InputPixelPosMin, InputPixelPosMax);
 
@@ -195,7 +195,7 @@ void MainCS(
 	{
 		FinalScreenVelocity = ComputeStaticVelocity(ScreenPos, ClosestDeviceZ);
 		
-		float4 EncodedVelocity = SceneVelocityTexture[uint2(InputBufferUV * InputInfo_Extent + VelocityPixelOffset)];
+		float4 EncodedVelocity = SceneVelocityTexture[uint2(InputBufferUV * InputInfo_Extent) + VelocityPixelOffset];
 
 		#if CONFIG_USE_SCENE_VELOCITY
 		{
@@ -267,7 +267,7 @@ void MainCS(
 		//DilatedVelocityOutput[OutputPixelPos] = EncodedVelocity.xy; //EncodeVelocityToTexture(float3(FinalScreenVelocity, 0.0)).xy;
 		DilatedVelocityOutput[OutputPixelPos] = EncodeVelocityToTexture(float3(FinalScreenVelocity, 0.0)).xy;
 		ClosestDepthOutput[OutputPixelPos] = ClosestDeviceZ;
-		ParallaxFactorOutput[OutputPixelPos] = CompressParallaxFactor(ParallaxFactorForHistory) | (CompressParallaxFactor(ParallaxFactorForHistoryRejection) << 4);
+		ParallaxFactorOutput[OutputPixelPos] = CompressParallaxFactor(taa_half(ParallaxFactorForHistory)) | (CompressParallaxFactor(taa_half(ParallaxFactorForHistoryRejection)) << 4);
 
 		// Output motion blur velocity flatten
 		#if DIM_MOTION_BLUR_DIRECTIONS > 0
diff --git a/Engine/Shaders/Private/TemporalSuperResolution/TSRFilterAntiAliasing.usf b/Engine/Shaders/Private/TemporalSuperResolution/TSRFilterAntiAliasing.usf
--- a/Engine/Shaders/Private/TemporalSuperResolution/TSRFilterAntiAliasing.usf
+++ b/Engine/Shaders/Private/TemporalSuperResolution/TSRFilterAntiAliasing.usf
@@ -11,7 +11,7 @@
 
 //------------------------------------------------------- PARAMETERS
 
-Texture2D<uint> AntiAliasingTexture;
+Texture2D<taa_ushort> AntiAliasingTexture;
 RWTexture2D<uint> AntiAliasingOutput;
 
 
@@ -23,7 +23,7 @@ taa_half2 SampleAntiAliasing(taa_short2 KernelCenter, taa_short2 Offset)
 		KernelCenter + Offset,
 		InputPixelPosMin, InputPixelPosMax);
 	
-	uint EncodedInputTexelOffset = AntiAliasingTexture[SampleInputPixelPos];
+	taa_ushort EncodedInputTexelOffset = AntiAliasingTexture[SampleInputPixelPos];
 
 	return DecodeSpatialAntiAliasingOffset(EncodedInputTexelOffset);
 }
diff --git a/Engine/Shaders/Private/TemporalSuperResolution/TSRFilterFrequencies.usf b/Engine/Shaders/Private/TemporalSuperResolution/TSRFilterFrequencies.usf
--- a/Engine/Shaders/Private/TemporalSuperResolution/TSRFilterFrequencies.usf
+++ b/Engine/Shaders/Private/TemporalSuperResolution/TSRFilterFrequencies.usf
@@ -123,7 +123,7 @@ void MainCS(
 	#if CONFIG_USE_DYNAMIC_SPREAD_FACTOR
 		taa_short2 DynamicSpreadFactor = (
 			taa_short(CONFIG_KERNEL_SPREAD_FACTOR) +
-			taa_short2((InputPixelPos % taa_short(2)) ^ (InputPixelPos.yx % taa_short(2))));
+			taa_short2((taa_ushort2(InputPixelPos) % taa_ushort(2)) ^ (taa_ushort2(InputPixelPos.yx) % taa_ushort(2))));
 	#else
 		taa_short2 DynamicSpreadFactor = taa_short(CONFIG_KERNEL_SPREAD_FACTOR);
 	#endif
@@ -236,14 +236,14 @@ void MainCS(
 	#if CONFIG_ENABLE_STOCASTIC_QUANTIZATION
 	{
 		uint2 Random = Rand3DPCG16(int3(OutputPixelPos, View.StateFrameIndexMod8)).xy;
-		float2 E = Hammersley16(0, 1, Random);
+		taa_half E = taa_half(Hammersley16(0, 1, Random).x);
 		
 		FilteredInput = QuantizeForFloatRenderTarget(FilteredInput, E.x, OutputQuantizationError);
 		FilteredHistory = QuantizeForFloatRenderTarget(FilteredHistory, E.x, OutputQuantizationError);
 	}
 	#endif
 
-	OutputPixelPos.x = all(OutputPixelPos < LowFrequencyInfo_ViewportMax) ? OutputPixelPos.x : ~taa_short(0);
+	OutputPixelPos.x = all(OutputPixelPos < taa_short2(LowFrequencyInfo_ViewportMax)) ? OutputPixelPos.x : ~taa_short(0);
 
 	FilteredInputOutput[OutputPixelPos] = FilteredInput;
 	FilteredPredictionSceneColorOutput[OutputPixelPos] = FilteredHistory;
diff --git a/Engine/Shaders/Private/TemporalSuperResolution/TSRHoleFillVelocity.usf b/Engine/Shaders/Private/TemporalSuperResolution/TSRHoleFillVelocity.usf
--- a/Engine/Shaders/Private/TemporalSuperResolution/TSRHoleFillVelocity.usf
+++ b/Engine/Shaders/Private/TemporalSuperResolution/TSRHoleFillVelocity.usf
@@ -43,7 +43,7 @@ void MainCS(
 	float4 Debug = 0.0;
 
 	uint2 DispatchThreadId = (
-		ZOrder2D(GroupThreadIndex, log2(TILE_SIZE)) +
+		ZOrder2D(GroupThreadIndex, uint(log2(float(TILE_SIZE)))) +
 		GroupId * uint2(TILE_SIZE, TILE_SIZE));
 	taa_short2 LaneOffsetSign = GetLaneOffsetSign();
 
@@ -69,9 +69,9 @@ void MainCS(
 			{
 				taa_short2 Offset;
 				#if CONFIG_SAMPLES == 9
-					Offset = kOffsets3x3[kSquareIndexes3x3[i]];
+					Offset = taa_short2(kOffsets3x3[kSquareIndexes3x3[i]]);
 				#elif CONFIG_SAMPLES == 5
-					Offset = kOffsets3x3[kPlusIndexes3x3[i]];
+					Offset = taa_short2(kOffsets3x3[kPlusIndexes3x3[i]]);
 				#endif
 					
 				SampleInputPixelPos = InputPixelPos + ClampPixelOffset(
@@ -167,8 +167,8 @@ void MainCS(
 			float2 NeighborPixelVelocity0;
 			float2 NeighborPixelVelocity1;
 
-			taa_half NeighborPixelVelocityLength0;
-			taa_half NeighborPixelVelocityLength1;
+			float NeighborPixelVelocityLength0;
+			float NeighborPixelVelocityLength1;
 			
 			bool bNeighborIsParallaxRejected0;
 			bool bNeighborIsParallaxRejected1;
@@ -219,11 +219,11 @@ void MainCS(
 				const uint NeightbordId1 = i + 1 + CONFIG_SAMPLE_PAIRS;
 			
 				#if CONFIG_SAMPLES == 9
-					Offset0 = kOffsets3x3[kSquareIndexes3x3[NeightbordId0]];
-					Offset1 = kOffsets3x3[kSquareIndexes3x3[NeightbordId1]];
+					Offset0 = taa_short2(kOffsets3x3[kSquareIndexes3x3[NeightbordId0]]);
+					Offset1 = taa_short2(kOffsets3x3[kSquareIndexes3x3[NeightbordId1]]);
 				#elif CONFIG_SAMPLES == 5
-					Offset0 = kOffsets3x3[kPlusIndexes3x3[NeightbordId0]];
-					Offset1 = kOffsets3x3[kPlusIndexes3x3[NeightbordId1]];
+					Offset0 = taa_short2(kOffsets3x3[kPlusIndexes3x3[NeightbordId0]]);
+					Offset1 = taa_short2(kOffsets3x3[kPlusIndexes3x3[NeightbordId1]]);
 				#endif
 					
 				NeighborEncodedVelocity0 = EncodedVelocityNeighborhood[NeightbordId0];
@@ -296,7 +296,7 @@ void MainCS(
 			float2 HoleFillingVelocity = HoleFillingPixelVelocity * InputPixelVelocityToScreenVelocity;
 
 			EndodedScreenVelocity = EncodeVelocityToTexture(float3(HoleFillingVelocity, 0.0)).xy;
-			HoleFilledVelocityMask = taa_half(HoleFillingOutOfRange) * saturate(MaxHollFillPixelVelocity - VelocityDelta);
+			HoleFilledVelocityMask = taa_half(HoleFillingOutOfRange) * taa_half(saturate(MaxHollFillPixelVelocity - VelocityDelta));
 		}
 
 		// Remove noise from high frequency rejections.
@@ -306,7 +306,7 @@ void MainCS(
 		RejectHighFrequency = max(RejectHighFrequency, HoleFilledVelocityMask);
 	}
 
-	if (all(InputPixelPos < InputInfo_ViewportMax))
+	if (all(InputPixelPos < taa_short2(InputInfo_ViewportMax)))
 	{
 		HoleFilledVelocityOutput[InputPixelPos] = EndodedScreenVelocity;
 		HoleFilledVelocityMaskOutput[InputPixelPos] = taa_half2(HoleFilledVelocityMask, RejectHighFrequency);
diff --git a/Engine/Shaders/Private/TemporalSuperResolution/TSRPostfilterRejection.usf b/Engine/Shaders/Private/TemporalSuperResolution/TSRPostfilterRejection.usf
--- a/Engine/Shaders/Private/TemporalSuperResolution/TSRPostfilterRejection.usf
+++ b/Engine/Shaders/Private/TemporalSuperResolution/TSRPostfilterRejection.usf
@@ -94,7 +94,7 @@ void MainCS(
 	{
 		taa_short2 SampleInputPixelPos = AddAndClampPixelOffset(
 			RejectionPixelPos,
-			kOffsets3x3[i], kOffsets3x3[i],
+			taa_short2(kOffsets3x3[i]), taa_short2(kOffsets3x3[i]),
 			HistoryRejectionViewport.xy, HistoryRejectionViewport.zw - 1);
 
 		InSamples[i] = HistoryRejectionTexture[SampleInputPixelPos];
diff --git a/Engine/Shaders/Private/TemporalSuperResolution/TSRResolveHistory.usf b/Engine/Shaders/Private/TemporalSuperResolution/TSRResolveHistory.usf
--- a/Engine/Shaders/Private/TemporalSuperResolution/TSRResolveHistory.usf
+++ b/Engine/Shaders/Private/TemporalSuperResolution/TSRResolveHistory.usf
@@ -103,7 +103,7 @@ void MainCS(
 	float4 Debug = 0.0;
 
 	uint2 DispatchThreadId = (
-		ZOrder2D(GroupThreadIndex, log2(TILE_SIZE)) +
+		ZOrder2D(GroupThreadIndex, uint(log2(float(TILE_SIZE)))) +
 		GroupId * uint2(TILE_SIZE, TILE_SIZE));
 
 	taa_halfC LowFrequencyColor = 0.0;
@@ -159,10 +159,10 @@ void MainCS(
 
 				float2 PixelOffset = abs(SamplePixelPos - OutputPixelPos);
 
-				float kernelWeight = MitchellNetravali(PixelOffset.x, B, C) * MitchellNetravali(PixelOffset.y, B, C);
+				taa_half kernelWeight = taa_half(MitchellNetravali(PixelOffset.x, B, C) * MitchellNetravali(PixelOffset.y, B, C));
 				
-				float LowFrequencyHdrWeight  = HdrWeight4(SampleLowFrequencyColor);
-				float HighFrequencyHdrWeight = HdrWeight4(SampleHighFrequencyColor * InvSampleHistoryValidity);
+				taa_half LowFrequencyHdrWeight  = HdrWeight4(SampleLowFrequencyColor);
+				taa_half HighFrequencyHdrWeight = HdrWeight4(SampleHighFrequencyColor * InvSampleHistoryValidity);
 
 				LowFrequencyColor   += (kernelWeight * LowFrequencyHdrWeight)  * SampleLowFrequencyColor;
 				LowFrequencyWeight  += (kernelWeight * LowFrequencyHdrWeight);
@@ -170,14 +170,14 @@ void MainCS(
 				HighFrequencyColor  += (kernelWeight * HighFrequencyHdrWeight) * SampleHighFrequencyColor;
 				HighFrequencyWeight += (kernelWeight * HighFrequencyHdrWeight) * SampleHistoryValidity;
 
-				HistoryValidity       += max(kernelWeight, 0.0) * SampleHistoryValidity;
-				HistoryValidityWeight += max(kernelWeight, 0.0);
+				HistoryValidity       += max(kernelWeight, taa_half(0.0)) * SampleHistoryValidity;
+				HistoryValidityWeight += max(kernelWeight, taa_half(0.0));
 			}
 		}
 
 		LowFrequencyColor  *= rcp(LowFrequencyWeight);
 		HighFrequencyColor *= SafeRcp(HighFrequencyWeight);
-		HistoryValidity *= HistoryValidityMultiply / HistoryValidityWeight;
+		HistoryValidity *= taa_half(HistoryValidityMultiply) / HistoryValidityWeight;
 	}
 	
 	// Compute final output
@@ -193,7 +193,7 @@ void MainCS(
 		FinalOutputColor = LowFrequencyColor * Weights[0] + HighFrequencyColor * Weights[1];
 	}
 
-	taa_short2 OutputPixelPosition = InvalidateOutputPixelPos(OutputViewRectMin + DispatchThreadId, OutputViewRectMax);
+	taa_short2 OutputPixelPosition = InvalidateOutputPixelPos(taa_short2(OutputViewRectMin + DispatchThreadId), OutputViewRectMax);
 	
 	// Output final scene color Mip0
 	{
@@ -206,8 +206,8 @@ void MainCS(
 
 		taa_halfC HalfResOutput = FinalOutputColor * taa_half(0.25);
 		taa_short2 HalfResOutputPixelPos;
-		HalfResOutputPixelPos[0] = (OutputPixelPosition[0] >> taa_short(1)) | (((OutputPixelPosition[0] & taa_short(0x1)) | !bGenerateOutputMip1) * taa_short(~0));
-		HalfResOutputPixelPos[1] = (OutputPixelPosition[1] >> taa_short(1)) | (((OutputPixelPosition[1] & taa_short(0x1)) | !bGenerateOutputMip1) * taa_short(~0));
+		HalfResOutputPixelPos[0] = (OutputPixelPosition[0] >> taa_short(1)) | (((OutputPixelPosition[0] & taa_short(0x1)) | taa_short(!bGenerateOutputMip1)) * taa_short(~0));
+		HalfResOutputPixelPos[1] = (OutputPixelPosition[1] >> taa_short(1)) | (((OutputPixelPosition[1] & taa_short(0x1)) | taa_short(!bGenerateOutputMip1)) * taa_short(~0));
 			
 		#if PLATFORM_SUPPORTS_WAVE_BROADCAST && 0 // Support WaveBroadcast with halfs
 		BRANCH
diff --git a/Engine/Shaders/Private/TemporalSuperResolution/TSRSpatialAntiAliasing.usf b/Engine/Shaders/Private/TemporalSuperResolution/TSRSpatialAntiAliasing.usf
--- a/Engine/Shaders/Private/TemporalSuperResolution/TSRSpatialAntiAliasing.usf
+++ b/Engine/Shaders/Private/TemporalSuperResolution/TSRSpatialAntiAliasing.usf
@@ -133,7 +133,7 @@ void AccumulateEdgeLength(
 	bEdgeStopedByIncrement = bEdgeStopedByIncrement || bStopNowIncrement;
 	bEdgeStopedByDecrement = bEdgeStopedByDecrement || bStopNowDecrement;
 
-	EdgeLength += (bEdgeStopedByIncrement || bEdgeStopedByDecrement) ? 0 : 1;
+	EdgeLength += (bEdgeStopedByIncrement || bEdgeStopedByDecrement) ? taa_ushort(0) : taa_ushort(1);
 }
 
 taa_half ComputeDistanceToEdge(
@@ -158,12 +158,12 @@ taa_half ComputeDistanceToEdge(
 	if (!bEdgeStopedByIncrementN && !bEdgeStopedByIncrementP && !bEdgeStopedByDecrementP && !bEdgeStopedByDecrementN)
 	{
 		// No aliasing in the neighborhood.
-		TexelCenterOffsetToEdge = 0.0;
+		TexelCenterOffsetToEdge = taa_half(0.0);
 	}
 	else if (bEdgeStopedByIncrementN && bEdgeStopedByIncrementP)
 	{
 		// Looks like InputC is concave detail
-		TexelCenterOffsetToEdge = 0.5;
+		TexelCenterOffsetToEdge = taa_half(0.5);
 	}
 	else if (bEdgeStopedByDecrementN && bEdgeStopedByDecrementP)
 	{
@@ -173,12 +173,12 @@ taa_half ComputeDistanceToEdge(
 	else if (bEdgeStopedByIncrementN) // && bEdgeStopedByDecrementP)
 	{
 		// Looks like staircasing from - to +
-		TexelCenterOffsetToEdge = 0.5 - (EdgeLengthN + 0.5) * fEdgeInvLength;
+		TexelCenterOffsetToEdge = taa_half(0.5) - (EdgeLengthN + taa_half(0.5)) * fEdgeInvLength;
 	}
 	else if (bEdgeStopedByIncrementP) // && bEdgeStopedByDecrementN)
 	{
 		// Looks like staircasing from + to -
-		TexelCenterOffsetToEdge = 0.5 - (EdgeLengthP + 0.5) * fEdgeInvLength;
+		TexelCenterOffsetToEdge = taa_half(0.5) - (EdgeLengthP + taa_half(0.5)) * fEdgeInvLength;
 	}
 
 	return TexelCenterOffsetToEdge;
@@ -275,8 +275,8 @@ void MainCS(
 		#endif
 
 		bBrowseVerticaly = (TotalVarHN + TotalVarH + TotalVarHS) > (TotalVarVE + TotalVarV + TotalVarVW);
-		BrowseDirection[0] = bBrowseVerticaly ? 0 : 1;
-		BrowseDirection[1] = bBrowseVerticaly ? 1 : 0;
+		BrowseDirection[0] = bBrowseVerticaly ? taa_short(0) : taa_short(1);
+		BrowseDirection[1] = bBrowseVerticaly ? taa_short(1) : taa_short(0);
 
 		FLATTEN
 		if (bBrowseVerticaly[0])
@@ -306,7 +306,7 @@ void MainCS(
 		
 		bool2 WorthBrowsing = LumaDelta > CONFIG_MIN_LUMA_ALIASING;
 
-		taa_half2 BrowseDelta = saturate(LumaDelta * rcp(CONFIG_MIN_LUMA_ALIASING));
+		taa_half2 BrowseDelta = saturate(LumaDelta * rcp(taa_half(CONFIG_MIN_LUMA_ALIASING)));
 		
 		taa_ushort2 EdgeLengthP = 0;
 		taa_ushort2 EdgeLengthN = 0;
diff --git a/Engine/Shaders/Private/TemporalSuperResolution/TSRUpdateHistory.usf b/Engine/Shaders/Private/TemporalSuperResolution/TSRUpdateHistory.usf
--- a/Engine/Shaders/Private/TemporalSuperResolution/TSRUpdateHistory.usf
+++ b/Engine/Shaders/Private/TemporalSuperResolution/TSRUpdateHistory.usf
@@ -358,16 +358,16 @@ void MainCS(
 			// Fetch the 2x2 neighborhood of the velocity
 			#if CONFIG_INTERPOLATE_VELOCITY
 			UNROLL_N(DPV_PIXEL_PER_LANE)
-			for (uint PixelId = 0; PixelId < DPV_PIXEL_PER_LANE; PixelId++)
+			for (uint PixelId2 = 0; PixelId2 < DPV_PIXEL_PER_LANE; PixelId2++)
 			{
-				taa_short2 LocalInputPixelPos0 = dpv_access_pixel(InputPixelPos, PixelId);
-				taa_short2 LocalInputPixelPos1 = LocalInputPixelPos0 + taa_short2(sign(dpv_access_pixel(dKO, PixelId)));
+				taa_short2 LocalInputPixelPos0 = dpv_access_pixel(InputPixelPos, PixelId2);
+				taa_short2 LocalInputPixelPos1 = LocalInputPixelPos0 + taa_short2(sign(dpv_access_pixel(dKO, PixelId2)));
 				LocalInputPixelPos1 = ClampPixelOffset(LocalInputPixelPos1, InputPixelPosMin, InputPixelPosMax);
 
-				RawEncodedVelocityNeighborhood[0][PixelId] = DilatedVelocityTexture[taa_short2(LocalInputPixelPos0.x, LocalInputPixelPos0.y)];
-				RawEncodedVelocityNeighborhood[1][PixelId] = DilatedVelocityTexture[taa_short2(LocalInputPixelPos1.x, LocalInputPixelPos0.y)];
-				RawEncodedVelocityNeighborhood[2][PixelId] = DilatedVelocityTexture[taa_short2(LocalInputPixelPos0.x, LocalInputPixelPos1.y)];
-				RawEncodedVelocityNeighborhood[3][PixelId] = DilatedVelocityTexture[taa_short2(LocalInputPixelPos1.x, LocalInputPixelPos1.y)];
+				RawEncodedVelocityNeighborhood[0][PixelId2] = DilatedVelocityTexture[taa_short2(LocalInputPixelPos0.x, LocalInputPixelPos0.y)];
+				RawEncodedVelocityNeighborhood[1][PixelId2] = DilatedVelocityTexture[taa_short2(LocalInputPixelPos1.x, LocalInputPixelPos0.y)];
+				RawEncodedVelocityNeighborhood[2][PixelId2] = DilatedVelocityTexture[taa_short2(LocalInputPixelPos0.x, LocalInputPixelPos1.y)];
+				RawEncodedVelocityNeighborhood[3][PixelId2] = DilatedVelocityTexture[taa_short2(LocalInputPixelPos1.x, LocalInputPixelPos1.y)];
 			}
 			#endif
 		}
@@ -521,8 +521,6 @@ void MainCS(
 	taa_half2 ClosestInputLuma4;
 	
 	
-	taa_half2x2 dKO = taa_half2x2(PPCo - PPCk);
-		
 	bool2 bUpdateSubpixelDetails;
 	bool2x2 bKOSign;
 
@@ -629,7 +627,7 @@ void MainCS(
 			if (SampleId == 0)
 			{
 				taa_half2x2 dPP =  RawdPPArray[SampleId];
-				ToneWeight = max(ToneWeight, ToneWeight > 0.2 ? 0.0 : ComputeSampleWeigth(1.0, dPP, /* MinimalContribution = */ 0.0));
+				ToneWeight = max(ToneWeight, ToneWeight > taa_half(0.2) ? taa_half(0.0).xx : ComputeSampleWeigth(taa_half(1.0), dPP, /* MinimalContribution = */ 0.0));
 			}
 			#endif
 
@@ -741,7 +739,7 @@ void MainCS(
 				dpv_lo(PrevHistoryBufferUV), PrevHistoryInfo_Extent, PrevHistoryInfo_ExtentInverse);
 			FCatmullRomSamples Samples1 = GetBicubic2DCatmullRomSamples(
 				dpv_hi(PrevHistoryBufferUV), PrevHistoryInfo_Extent, PrevHistoryInfo_ExtentInverse);
-				
+			
 			taa_halfC RawLowFrequency[BICUBIC_CATMULL_ROM_SAMPLES][DPV_PIXEL_PER_LANE];
 			taa_halfC RawHighFrequency[BICUBIC_CATMULL_ROM_SAMPLES][DPV_PIXEL_PER_LANE];
 			taa_half2 RawMetadata[BICUBIC_CATMULL_ROM_SAMPLES][DPV_PIXEL_PER_LANE];
@@ -766,9 +764,6 @@ void MainCS(
 					UNROLL_N(2)
 					for (uint PixelId = 0; PixelId < DPV_PIXEL_PER_LANE; PixelId++)
 					{
-						float2 SampleUV0 = clamp(Samples0.UV[i], PrevHistoryInfo_UVViewportBilinearMin, PrevHistoryInfo_UVViewportBilinearMax);
-						float2 SampleUV1 = clamp(Samples1.UV[i], PrevHistoryInfo_UVViewportBilinearMin, PrevHistoryInfo_UVViewportBilinearMax);
-					
 						float2 SampleUV = PixelId == 0 ? SampleUV0 : SampleUV1;
 						
 						RawLowFrequency[i][PixelId]      = PrevHistory_LowFrequency.SampleLevel(GlobalBilinearClampedSampler, SampleUV, 0);
@@ -1049,7 +1044,7 @@ void MainCS(
 		// Compute the maximum output validity.
 		taa_half2 OutputValidity = (
 			clamp(RejectedValidity + DesiredCurrentContribution, taa_half(0.0), PredictionOnlyValidity) +
-			clamp(RejectedValidity + DesiredCurrentContribution * PrevHistoryRejectionWeight * RejectionConfidentEnough - PredictionOnlyValidity, 0.0, 1.0 - PredictionOnlyValidity));
+			clamp(RejectedValidity + DesiredCurrentContribution * PrevHistoryRejectionWeight * RejectionConfidentEnough - PredictionOnlyValidity, taa_half(0.0), taa_half(1.0) - PredictionOnlyValidity));
 
 		#if !CONFIG_ACCUMULATE_TRANSLUCENCY_SEPARATLY
 			FLATTEN
@@ -1199,7 +1194,7 @@ void MainCS(
 			FinalLowFrequencyWeight = saturate(CurrentFrameSampleCount + PrevLowFrequencySampleCount);
 			
 			#if 1
-				FinalLowFrequencyWeight = saturate(floor(255.0 * (CurrentFrameSampleCount + PrevLowFrequencySampleCount)) * rcp(255.0));
+				FinalLowFrequencyWeight = saturate(floor(taa_half(255.0) * (CurrentFrameSampleCount + PrevLowFrequencySampleCount)) * rcp(taa_half(255.0)));
 			#endif
 		}
 		#endif
@@ -1420,7 +1415,7 @@ void MainCS(
 		#if CONFIG_ENABLE_STOCASTIC_QUANTIZATION
 		{
 			uint2 Random = Rand3DPCG16(int3(dpv_lo(LocalHistoryPixelPos), View.StateFrameIndexMod8)).xy;
-			float E = Hammersley16(0, 1, Random).x;
+			taa_half E = taa_half(Hammersley16(0, 1, Random).x);
 
 			PremultipliedHighFrequencyColor = QuantizeForFloatRenderTarget(PremultipliedHighFrequencyColor, E, HistoryQuantizationError);
 			FinalLowFrequencyColor = QuantizeForFloatRenderTarget(FinalLowFrequencyColor, E, HistoryQuantizationError);
@@ -1488,8 +1483,8 @@ void MainCS(
 			taa_halfCx2 HalfResOutput = OutputColor * taa_half(0.25);
 			//taa_short2x2 HalfResOutputPixelPos = (LocalHistoryPixelPos >> taa_short(1)) | (((LocalHistoryPixelPos & taa_short(0x1)) | !bGenerateOutputMip1) * taa_short(~0));
 			taa_short2x2 HalfResOutputPixelPos;
-			HalfResOutputPixelPos[0] = (LocalHistoryPixelPos[0] >> taa_short(1)) | (((LocalHistoryPixelPos[0] & taa_short(0x1)) | !bGenerateOutputMip1) * taa_short(~0));
-			HalfResOutputPixelPos[1] = (LocalHistoryPixelPos[1] >> taa_short(1)) | (((LocalHistoryPixelPos[1] & taa_short(0x1)) | !bGenerateOutputMip1) * taa_short(~0));
+			HalfResOutputPixelPos[0] = (LocalHistoryPixelPos[0] >> taa_short(1)) | (((LocalHistoryPixelPos[0] & taa_short(0x1)) | taa_short(!bGenerateOutputMip1)) * taa_short(~0));
+			HalfResOutputPixelPos[1] = (LocalHistoryPixelPos[1] >> taa_short(1)) | (((LocalHistoryPixelPos[1] & taa_short(0x1)) | taa_short(!bGenerateOutputMip1)) * taa_short(~0));
 			
 			#if PLATFORM_SUPPORTS_WAVE_BROADCAST && 0 // Support WaveBroadcast with halfs
 			BRANCH
diff --git a/Engine/Shaders/Private/TextureSampling.ush b/Engine/Shaders/Private/TextureSampling.ush
--- a/Engine/Shaders/Private/TextureSampling.ush
+++ b/Engine/Shaders/Private/TextureSampling.ush
@@ -15,7 +15,7 @@ float2 SnapToClosestPixelCenterUVEx(float2 SampleUV, float2 TextureSize, float2
 
 float2 SnapToClosestPixelCenterUV(float2 SampleUV, float2 TextureSize, float2 TextureInvSize, const uint MipLevel = 0)
 {
-	return SnapToClosestPixelCenterUVEx(SampleUV, TextureSize, TextureInvSize, 1u << MipLevel, rcp(float(1u << MipLevel)));
+	return SnapToClosestPixelCenterUVEx(SampleUV, TextureSize, TextureInvSize, float(1u << MipLevel), rcp(float(1u << MipLevel)));
 }
 
 
@@ -58,7 +58,7 @@ FBilinearSampleInfos GetBilinearSampleLevelInfosEx(float2 SampleUV, float2 Textu
 
 FBilinearSampleInfos GetBilinearSampleLevelInfos(float2 SampleUV, float2 TextureSize, float2 TextureInvSize, const uint MipLevel = 0)
 {
-	return GetBilinearSampleLevelInfosEx(SampleUV, TextureSize, TextureInvSize, 1u << MipLevel, rcp(float(1u << MipLevel)));
+	return GetBilinearSampleLevelInfosEx(SampleUV, TextureSize, TextureInvSize, float(1u << MipLevel), rcp(float(1u << MipLevel)));
 }
 
 /** Return the weight of a sample <SampleId>. */
@@ -91,7 +91,7 @@ void Bicubic2DCatmullRom(in float2 UV, in float2 Size, in float2 InvSize, out fl
 	UV *= Size;
 
 	float2 tc = floor(UV - 0.5) + 0.5;
-	half2 f = UV - tc;
+	half2 f = half2(UV - tc);
 	half2 f2 = f * f;
 	half2 f3 = f2 * f;
 
diff --git a/Engine/Source/Runtime/Renderer/Private/PostProcess/TemporalSuperResolution.cpp b/Engine/Source/Runtime/Renderer/Private/PostProcess/TemporalSuperResolution.cpp
--- a/Engine/Source/Runtime/Renderer/Private/PostProcess/TemporalSuperResolution.cpp
+++ b/Engine/Source/Runtime/Renderer/Private/PostProcess/TemporalSuperResolution.cpp
@@ -174,6 +174,7 @@ public:
 	static void ModifyCompilationEnvironment(const FGlobalShaderPermutationParameters& Parameters, FShaderCompilerEnvironment& OutEnvironment)
 	{
 		OutEnvironment.CompilerFlags.Add(CFLAG_AllowRealTypes);
+		OutEnvironment.CompilerFlags.Add(CFLAG_WarningsAsErrors);
 	}
 }; // class FTemporalSuperResolutionShader
 
