Add support for accessing triangle indices and vertex positions in Vulkan inline ray tracing shaders.

Also add FTraceRayInlineResult::ObjectToWorld3x4 on some platforms (where PLATFORM_SUPPORTS_INLINE_RAY_TRACING_INSTANCE_MATRIX=1).

diff --git a/Engine/Shaders/Private/RayTracing/RayTracingCommon.ush b/Engine/Shaders/Private/RayTracing/RayTracingCommon.ush
--- a/Engine/Shaders/Private/RayTracing/RayTracingCommon.ush
+++ b/Engine/Shaders/Private/RayTracing/RayTracingCommon.ush
@@ -129,6 +129,14 @@ FRayTracingInstanceDescriptor BuildPlatformRayTracingInstanceDesc(uint InstanceM
 
 #endif // OVERRIDE_PLATFORM_RAYTRACING_INSTANCE_DESCRITOR
 
+// Indices and vertex positions of the ray tracing geometry triangle.
+// Used for loading and interpolation of per-vertex attributes and computing per-triangle normals.
+struct FTriangleBaseAttributes
+{
+	uint3 Indices;
+	float3 LocalPositions[3];
+};
+
 struct FBasicRayData
 {
 	float3 Origin;
diff --git a/Engine/Shaders/Private/RayTracing/RayTracingHitGroupCommon.ush b/Engine/Shaders/Private/RayTracing/RayTracingHitGroupCommon.ush
--- a/Engine/Shaders/Private/RayTracing/RayTracingHitGroupCommon.ush
+++ b/Engine/Shaders/Private/RayTracing/RayTracingHitGroupCommon.ush
@@ -52,12 +52,6 @@ float3 LoadVertexPositionFloat3(ByteAddressBuffer SourceBuffer, uint Index, uint
 
 // Explicitly declare "system" resources that will be available to all hit group shaders.
 
-struct FTriangleBaseAttributes
-{
-	uint3 Indices;
-	float3 LocalPositions[3];
-};
-
 FTriangleBaseAttributes LoadTriangleBaseAttributes(uint PrimitiveId)
 {
 	FTriangleBaseAttributes Result = (FTriangleBaseAttributes)0;
diff --git a/Engine/Shaders/Private/RayTracing/TraceRayInline.ush b/Engine/Shaders/Private/RayTracing/TraceRayInline.ush
--- a/Engine/Shaders/Private/RayTracing/TraceRayInline.ush
+++ b/Engine/Shaders/Private/RayTracing/TraceRayInline.ush
@@ -19,8 +19,13 @@
 
 #include "/Engine/Private/RayTracing/TraceRayInlineCommon.ush"
 
+#if COMPILER_VULKAN
+#include "/Engine/Private/RayTracing/TraceRayInlineVulkan.ush"
+#endif // COMPILER_VULKAN
+
 #define PLATFORM_SUPPORTS_INLINE_RAY_TRACING_TRIANGLE_NORMALS 0
 #define PLATFORM_SUPPORTS_INLINE_RAY_TRACING_MAX_ITERATIONS 0
+#define PLATFORM_SUPPORTS_INLINE_RAY_TRACING_INSTANCE_MATRIX 1
 
 FTraceRayInlineResult TraceRayInline(RaytracingAccelerationStructure AccelerationStructure,
 	uint RayFlags,
@@ -52,6 +57,10 @@ FTraceRayInlineResult TraceRayInline(RaytracingAccelerationStructure Acceleratio
 
 		TraceRayResult.Bookmark[0] = asuint(TraceRayResult.HitT);
 		TraceRayResult.Bookmark[1] = asuint(TraceRayResult.HitT);
+
+	#if PLATFORM_SUPPORTS_INLINE_RAY_TRACING_INSTANCE_MATRIX
+		TraceRayResult.ObjectToWorld3x4 = RayQ.CommittedObjectToWorld3x4();
+	#endif // PLATFORM_SUPPORTS_INLINE_RAY_TRACING_INSTANCE_MATRIX
 	}
 
 	return TraceRayResult;
@@ -80,4 +89,17 @@ FTraceRayInlineResult TraceRayInline(RaytracingAccelerationStructure Acceleratio
 #define PLATFORM_SUPPORTS_INLINE_RAY_TRACING_MAX_ITERATIONS 0
 #endif // PLATFORM_SUPPORTS_INLINE_RAY_TRACING_MAX_ITERATIONS
 
+#ifndef PLATFORM_SUPPORTS_INLINE_RAY_TRACING_TRIANGLE_ATTRIBUTES
+#define PLATFORM_SUPPORTS_INLINE_RAY_TRACING_TRIANGLE_ATTRIBUTES 0
+#endif // PLATFORM_SUPPORTS_INLINE_RAY_TRACING_TRIANGLE_ATTRIBUTES
+
+#ifndef PLATFORM_SUPPORTS_INLINE_RAY_TRACING_INSTANCE_MATRIX
+#define PLATFORM_SUPPORTS_INLINE_RAY_TRACING_INSTANCE_MATRIX 0
+#endif // PLATFORM_SUPPORTS_INLINE_RAY_TRACING_INSTANCE_MATRIX
+
+#ifndef FRayTracingSceneMetadataRecord
+// Not used, but must be a valid structure declaration as it may appear in interfaces, resource declarations, etc.
+struct FRayTracingSceneMetadataRecord { uint Dummy; };
+#endif // FRayTracingSceneMetadataRecord
+
 #endif // PLATFORM_SUPPORTS_INLINE_RAY_TRACING
diff --git a/Engine/Shaders/Private/RayTracing/TraceRayInlineCommon.ush b/Engine/Shaders/Private/RayTracing/TraceRayInlineCommon.ush
--- a/Engine/Shaders/Private/RayTracing/TraceRayInlineCommon.ush
+++ b/Engine/Shaders/Private/RayTracing/TraceRayInlineCommon.ush
@@ -17,6 +17,9 @@ struct FTraceRayInlineResult
 
 	uint Bookmark[2];
 
+	// Only valid when PLATFORM_SUPPORTS_INLINE_RAY_TRACING_INSTANCE_MATRIX is 1
+	float3x4 ObjectToWorld3x4;
+
 	// Only valid when PLATFORM_SUPPORTS_INLINE_RAY_TRACING_TRIANGLE_NORMALS is 1
 	float3 WorldGeometryNormal;
 
diff --git a/Engine/Shaders/Private/RayTracing/TraceRayInlineVulkan.ush b/Engine/Shaders/Private/RayTracing/TraceRayInlineVulkan.ush
new file mode 100644
--- /dev/null
+++ b/Engine/Shaders/Private/RayTracing/TraceRayInlineVulkan.ush
@@ -0,0 +1,140 @@
+// Copyright Epic Games, Inc. All Rights Reserved.
+
+/*==============================================================================
+	TraceRayInlineVulkan.ush: Vulkan-specific inline ray tracing functionality,
+	such as utility functions for accessing triangle indices and vertices.
+==============================================================================*/
+
+#pragma once
+
+#include "/Engine/Shared/RayTracingDefinitions.h"
+#include "/Engine/Private/RayTracing/RayTracingCommon.ush"
+
+#if COMPILER_DXC && COMPILER_VULKAN && PLATFORM_SUPPORTS_INLINE_RAY_TRACING
+
+#define PLATFORM_SUPPORTS_INLINE_RAY_TRACING_TRIANGLE_ATTRIBUTES 1
+
+uint2 VulkanRawBufferLoad2(uint64_t BaseAddress, uint OffsetInBytes)
+{
+	uint2 Result;
+	Result.x = vk::RawBufferLoad(BaseAddress + OffsetInBytes);
+	Result.y = vk::RawBufferLoad(BaseAddress + OffsetInBytes + 4);
+	return Result;
+}
+
+uint3 VulkanRawBufferLoad3(uint64_t BaseAddress, uint OffsetInBytes)
+{
+	uint3 Result;
+	Result.x = vk::RawBufferLoad(BaseAddress + OffsetInBytes);
+	Result.y = vk::RawBufferLoad(BaseAddress + OffsetInBytes + 4);
+	Result.z = vk::RawBufferLoad(BaseAddress + OffsetInBytes + 8);
+	return Result;
+}
+
+uint4 VulkanRawBufferLoad4(uint64_t BaseAddress, uint OffsetInBytes)
+{
+	uint4 Result;
+	Result.x = vk::RawBufferLoad(BaseAddress + OffsetInBytes);
+	Result.y = vk::RawBufferLoad(BaseAddress + OffsetInBytes + 4);
+	Result.z = vk::RawBufferLoad(BaseAddress + OffsetInBytes + 8);
+	Result.w = vk::RawBufferLoad(BaseAddress + OffsetInBytes + 12);
+	return Result;
+}
+
+uint3 VulkanLoadIndices16Bit(uint64_t SourceBufferBaseAddress, uint OffsetInBytes)
+{
+	uint3 Result;
+
+	// ByteAddressBuffer loads must be aligned to DWORD boundary.
+	// We can load 2 DWORDs (4 SHORTs) at a time and extract 3 SHORTs (first 3 or second 3).
+	uint AlignedOffsetInBytes = OffsetInBytes & (~3);
+	const uint2 PackedIndices = VulkanRawBufferLoad2(SourceBufferBaseAddress, AlignedOffsetInBytes);
+
+	if (AlignedOffsetInBytes == OffsetInBytes)
+	{
+		// Extract first 3 SHORTs from 2 DWORDs
+		Result[0] = PackedIndices[0] & 0xffff;
+		Result[1] = PackedIndices[0] >> 16;
+		Result[2] = PackedIndices[1] & 0xffff;
+	}
+	else
+	{
+		// Extract second 3 SHORTs from 2 DWORDs
+		Result[0] = PackedIndices[0] >> 16;
+		Result[1] = PackedIndices[1] & 0xffff;
+		Result[2] = PackedIndices[1] >> 16;
+	}
+
+	return Result;
+}
+
+uint3 VulkanLoadIndices32Bit(uint64_t SourceBufferBaseAddress, uint OffsetInBytes)
+{
+	return VulkanRawBufferLoad3(SourceBufferBaseAddress, OffsetInBytes);
+}
+
+float3 VulkanLoadVertexPositionFloat3(uint64_t SourceBufferBaseAddress, uint Index, uint StrideInBytes)
+{
+	uint OffsetInBytes = Index * StrideInBytes;
+	return asfloat(VulkanRawBufferLoad3(SourceBufferBaseAddress, OffsetInBytes));
+}
+
+// Counterpart of FVulkanRayTracingGeometryParameters in VulkanRayTracing.cpp
+struct FVulkanRayTracingGeometryParameters
+{
+	uint Config;
+	uint IndexBufferOffsetInBytes;
+	uint64_t IndexBuffer;
+	uint64_t VertexBuffer;
+};
+
+// Create a generic alias for the Vulkan-specific metadata type
+#define FRayTracingSceneMetadataRecord FVulkanRayTracingGeometryParameters
+
+FTriangleBaseAttributes LoadInlineRayTracingTriangleAttributes(
+	StructuredBuffer<FRayTracingSceneMetadataRecord> RayTracingSceneMetadata,
+	uint InstanceContributionToHitGroupIndex,
+	uint MultiplierForGeometryContributionToHitGroupIndex,
+	uint GeometryIndex,
+	uint PrimitiveIndex)
+{
+	uint RecordIndex = InstanceContributionToHitGroupIndex / MultiplierForGeometryContributionToHitGroupIndex + GeometryIndex;
+	FVulkanRayTracingGeometryParameters Params = RayTracingSceneMetadata[RecordIndex];
+
+	FTriangleBaseAttributes Result = (FTriangleBaseAttributes)0;
+
+	// Fetch vertex indices and positions, then compute local space normal and transform it to world space
+
+	const uint IndexBufferOffsetInBytes = Params.IndexBufferOffsetInBytes;
+	const uint IndexBufferStride = Params.Config & 0xFF;
+	const uint BaseIndex = PrimitiveIndex * 3;
+
+	uint64_t IndexBuffer = Params.IndexBuffer;
+
+	if (IndexBufferStride == 0)
+	{
+		// Non-indexed geometry (implicit triangle list indices)
+		Result.Indices = uint3(BaseIndex, BaseIndex + 1, BaseIndex + 2);
+	}
+	else if (IndexBufferStride == 2)
+	{
+		Result.Indices = VulkanLoadIndices16Bit(IndexBuffer, IndexBufferOffsetInBytes + BaseIndex * IndexBufferStride);
+	}
+	else
+	{
+		Result.Indices = VulkanLoadIndices32Bit(IndexBuffer, IndexBufferOffsetInBytes + BaseIndex * IndexBufferStride);
+	}
+
+	// Fetch vertex positions (in local space)
+	const uint VertexStride = (Params.Config >> 8) & 0xFF;
+
+	// #dxr_todo: UE-72160 handle various vertex formats, for now only supporting float3
+	uint64_t VertexBuffer = Params.VertexBuffer;
+	Result.LocalPositions[0] = VulkanLoadVertexPositionFloat3(VertexBuffer, Result.Indices[0], VertexStride);
+	Result.LocalPositions[1] = VulkanLoadVertexPositionFloat3(VertexBuffer, Result.Indices[1], VertexStride);
+	Result.LocalPositions[2] = VulkanLoadVertexPositionFloat3(VertexBuffer, Result.Indices[2], VertexStride);
+
+	return Result;
+}
+
+#endif // COMPILER_DXC && COMPILER_VULKAN && PLATFORM_SUPPORTS_INLINE_RAY_TRACING
