Planar reflections operate in translated world space, fix problems at LWC-scale

diff --git a/Engine/Shaders/Private/ForwardLightingCommon.ush b/Engine/Shaders/Private/ForwardLightingCommon.ush
--- a/Engine/Shaders/Private/ForwardLightingCommon.ush
+++ b/Engine/Shaders/Private/ForwardLightingCommon.ush
@@ -499,7 +499,7 @@ half3 GetImageBasedReflectionSpecular(FMaterialPixelParameters MaterialParameter
 	if (abs(dot(PlanarReflectionStruct.ReflectionPlane.xyz, 1)) > .0001f)
 	{
 		// Reuse ReflectionCubemapSampler to avoid reducing the sampler count available to artists
-		float4 PlanarReflection = ComputePlanarReflections(LWCHackToFloat(MaterialParameters.AbsoluteWorldPosition), MaterialParameters.WorldNormal, Roughness, ReflectionStruct.ReflectionCubemapSampler);
+		float4 PlanarReflection = ComputePlanarReflections(MaterialParameters.WorldPosition_CamRelative, MaterialParameters.WorldNormal, Roughness, ReflectionStruct.ReflectionCubemapSampler);
 		// Planar reflections win over SSR and reflection environment
 		SpecularLighting = PlanarReflection.rgb + (1 - PlanarReflection.a) * SpecularLighting;
 	}
diff --git a/Engine/Shaders/Private/MobileBasePassPixelShader.usf b/Engine/Shaders/Private/MobileBasePassPixelShader.usf
--- a/Engine/Shaders/Private/MobileBasePassPixelShader.usf
+++ b/Engine/Shaders/Private/MobileBasePassPixelShader.usf
@@ -862,7 +862,7 @@ void Main(
 		BRANCH
 		if (abs(dot(PlanarReflectionStruct.ReflectionPlane.xyz, 1)) > .0001f)
 		{
-			half4 PlanarReflection = GetPlanarReflection(LWCHackToFloat(MaterialParameters.AbsoluteWorldPosition), MaterialParameters.WorldNormal, GBuffer.Roughness);
+			half4 PlanarReflection = GetPlanarReflection(MaterialParameters.WorldPosition_CamRelative, MaterialParameters.WorldNormal, GBuffer.Roughness);
 			// Planar reflections win over reflection environment
 			SpecularIBL = lerp(SpecularIBL, PlanarReflection.rgb, PlanarReflection.a);
 		}
diff --git a/Engine/Shaders/Private/MobileDeferredShading.usf b/Engine/Shaders/Private/MobileDeferredShading.usf
--- a/Engine/Shaders/Private/MobileDeferredShading.usf
+++ b/Engine/Shaders/Private/MobileDeferredShading.usf
@@ -252,8 +252,10 @@ void MobileDirectLightPS(
 	}
 	
 	float2 ScreenPos = UVAndScreenPos.zw;
-	float3 WorldPosition = mul(float4(ScreenPos * SceneDepth, SceneDepth, 1), LWCHackToFloat(PrimaryView.ScreenToWorld)).xyz;
-	half3 CameraVector = normalize(LWCHackToFloat(PrimaryView.WorldCameraOrigin) - WorldPosition);
+	float3 TranslatedWorldPosition = mul(float4(ScreenPos * SceneDepth, SceneDepth, 1), PrimaryView.ScreenToTranslatedWorld).xyz;
+	float3 WorldPosition = TranslatedWorldPosition - LWCHackToFloat(PrimaryView.PreViewTranslation);
+	
+	half3 CameraVector = normalize(PrimaryView.TranslatedWorldCameraOrigin - TranslatedWorldPosition);
 	half NoV = max(0, dot(GBuffer.WorldNormal, CameraVector));
 	half3 ReflectionVector = GBuffer.WorldNormal * (NoV * 2.0) - CameraVector;
 	
@@ -303,9 +305,9 @@ void MobileDirectLightPS(
 	uint NumLocalReflectionCaptures = min(ForwardLightData.NumCulledLightsGrid[NumCulledEntryIndex + 0], ForwardLightData.NumReflectionCaptures);
 	uint DataStartIndex = ForwardLightData.NumCulledLightsGrid[NumCulledEntryIndex + 1];
 
-	float3 SpecularIBL = CompositeReflectionCapturesAndSkylight(
+	float3 SpecularIBL = CompositeReflectionCapturesAndSkylightTWS(
 		1.0f,
-		WorldPosition,
+		TranslatedWorldPosition,
 		ReflectionVector,//RayDirection,
 		GBuffer.Roughness,
 		GBuffer.IndirectIrradiance,
@@ -319,7 +321,7 @@ void MobileDirectLightPS(
 	BRANCH
 	if (abs(dot(PlanarReflectionStruct.ReflectionPlane.xyz, 1)) > .0001f)
 	{
-		half4 PlanarReflection = GetPlanarReflection(WorldPosition, GBuffer.WorldNormal, GBuffer.Roughness);
+		half4 PlanarReflection = GetPlanarReflection(TranslatedWorldPosition, GBuffer.WorldNormal, GBuffer.Roughness);
 		// Planar reflections win over reflection environment
 		SpecularIBL = lerp(SpecularIBL, PlanarReflection.rgb, PlanarReflection.a);
 	}
diff --git a/Engine/Shaders/Private/PlanarReflectionShaders.usf b/Engine/Shaders/Private/PlanarReflectionShaders.usf
--- a/Engine/Shaders/Private/PlanarReflectionShaders.usf
+++ b/Engine/Shaders/Private/PlanarReflectionShaders.usf
@@ -76,7 +76,7 @@ void PrefilterPlanarReflectionPS(
 	if (KernelRadiusY > 0)
 	{
 		// Compute distance from camera origin to the camera vector intersection with the reflection plane
-		float CenterDistanceToPlane = -(dot(LWCHackToFloat(PrimaryView.WorldCameraOrigin), PlanarReflectionStruct.ReflectionPlane.xyz) - PlanarReflectionStruct.ReflectionPlane.w) / dot(normalize(ScreenVector), PlanarReflectionStruct.ReflectionPlane.xyz);
+		float CenterDistanceToPlane = -(dot(PrimaryView.TranslatedWorldCameraOrigin, PlanarReflectionStruct.ReflectionPlane.xyz) - PlanarReflectionStruct.ReflectionPlane.w) / dot(normalize(ScreenVector), PlanarReflectionStruct.ReflectionPlane.xyz);
 		// Note: reading from depth after the temporal AA pass
 		// This works without jittering because we disable temporal AA jittering for planar reflections
 		float CenterSceneDepth = CalcSceneDepth(BufferUV);
@@ -141,13 +141,13 @@ void PlanarReflectionPS(
 
 	FGBufferData GBuffer = GetGBufferDataFromSceneTextures(BufferUV);
 
-	float3 WorldPosition = ScreenVector * GBuffer.Depth + LWCHackToFloat(PrimaryView.WorldCameraOrigin);
+	float3 TranslatedWorldPosition = ResolvedView.TranslatedWorldCameraOrigin + ScreenVector * GBuffer.Depth;
 
 	// Skip unlit pixels (skybox)
 	BRANCH
 	if (GBuffer.ShadingModelID != SHADINGMODELID_UNLIT)
 	{
-		OutColor = ComputePlanarReflections(WorldPosition, GBuffer.WorldNormal, GBuffer.Roughness, PlanarReflectionStruct.PlanarReflectionSampler);
+		OutColor = ComputePlanarReflections(TranslatedWorldPosition, GBuffer.WorldNormal, GBuffer.Roughness, PlanarReflectionStruct.PlanarReflectionSampler);
 
 		OutColor.xyz *= View.PreExposure;
 	}
diff --git a/Engine/Shaders/Private/PlanarReflectionShared.ush b/Engine/Shaders/Private/PlanarReflectionShared.ush
--- a/Engine/Shaders/Private/PlanarReflectionShared.ush
+++ b/Engine/Shaders/Private/PlanarReflectionShared.ush
@@ -4,14 +4,14 @@
  PlanarReflectionShared.usf
 =============================================================================*/
 
-half4 ComputePlanarReflections(float3 WorldPosition, half3 WorldNormal, half Roughness, SamplerState SharedClampSampler)
+half4 ComputePlanarReflections(float3 TranslatedWorldPosition, half3 WorldNormal, half Roughness, SamplerState SharedClampSampler)
 {
 	half4 OutPlanarReflection = 0;
 
-	float PlaneDistance = dot(PlanarReflectionStruct.ReflectionPlane, float4(WorldPosition, -1));
+	float PlaneDistance = dot(PlanarReflectionStruct.ReflectionPlane, float4(TranslatedWorldPosition, -1));
 	half DistanceFade = 1 - saturate(abs(PlaneDistance) * PlanarReflectionStruct.PlanarReflectionParameters.x + PlanarReflectionStruct.PlanarReflectionParameters.y);
 
-	float3 PlaneOriginToWorldPosition = WorldPosition - PlanarReflectionStruct.PlanarReflectionOrigin.xyz;
+	float3 PlaneOriginToWorldPosition = TranslatedWorldPosition - PlanarReflectionStruct.PlanarReflectionOrigin.xyz;
 	float XAxisDistance = dot(PlaneOriginToWorldPosition, PlanarReflectionStruct.PlanarReflectionXAxis.xyz);
 	half XAxisFade = saturate((PlanarReflectionStruct.PlanarReflectionXAxis.w - abs(XAxisDistance)) * PlanarReflectionStruct.PlanarReflectionParameters.x);
 	float YAxisDistance = dot(PlaneOriginToWorldPosition, PlanarReflectionStruct.PlanarReflectionYAxis.xyz);
@@ -26,7 +26,7 @@ half4 ComputePlanarReflections(float3 WorldPosition, half3 WorldNormal, half Rou
 	if (FinalFade > 0)
 	{
 		// CameraToPixel in the main view is what we used as ReflectionVector when rendering the reflection pass to PlanarReflectionTexture
-		float3 CameraToPixel = normalize(WorldPosition - LWCHackToFloat(ResolvedView.WorldCameraOrigin));
+		float3 CameraToPixel = normalize(TranslatedWorldPosition - ResolvedView.TranslatedWorldCameraOrigin);
 		// Reflect the effective ReflectionVector in mirrored space to get the original camera vector
 		float3 MirroredCameraVector = reflect(CameraToPixel, -PlanarReflectionStruct.ReflectionPlane.xyz);
 		// Transform the GBuffer normal into mirrored space
@@ -36,9 +36,9 @@ half4 ComputePlanarReflections(float3 WorldPosition, half3 WorldNormal, half Rou
 		// At this point we have a new reflection vector off of the GBuffer normal, and we need to approximate its intersection with the scene
 		// An accurate intersection would ray trace the planar reflection depth buffer
 		// As an approximation we are just intersecting with a user defined sphere
-		float3 VirtualReflectionSpherePosition = WorldPosition + MirroredReflectionVectorOffNormal * PlanarReflectionStruct.PlanarReflectionParameters.z;
+		float3 VirtualReflectionSpherePosition = TranslatedWorldPosition + MirroredReflectionVectorOffNormal * PlanarReflectionStruct.PlanarReflectionParameters.z;
 		// Transform the intersection position into view space
-		float3 ViewVirtualReflectionSpherePosition = mul(float4(VirtualReflectionSpherePosition + LWCHackToFloat(ResolvedView.PreViewTranslation), 1), ResolvedView.TranslatedWorldToView).xyz;
+		float3 ViewVirtualReflectionSpherePosition = mul(float4(VirtualReflectionSpherePosition, 1), ResolvedView.TranslatedWorldToView).xyz;
 		// Transform the intersection position into clip space using the same projection matrix used to render PlanarReflectionTexture
 		float4 ClipVirtualReflectionSpherePosition = mul(float4(ViewVirtualReflectionSpherePosition, 1), PlanarReflectionStruct.ProjectionWithExtraFOV[ResolvedView.StereoPassIndex]);
 
@@ -74,9 +74,9 @@ half4 ComputePlanarReflections(float3 WorldPosition, half3 WorldNormal, half Rou
 }
 
 #if (FEATURE_LEVEL <= FEATURE_LEVEL_ES3_1)
-half4 GetPlanarReflection(float3 WorldPosition, half3 WorldNormal, half Roughness)
+half4 GetPlanarReflection(float3 TranslatedWorldPosition, half3 WorldNormal, half Roughness)
 {
-	half4 PlanarReflection = ComputePlanarReflections(WorldPosition, WorldNormal, Roughness, PlanarReflectionStruct.PlanarReflectionSampler);
+	half4 PlanarReflection = ComputePlanarReflections(TranslatedWorldPosition, WorldNormal, Roughness, PlanarReflectionStruct.PlanarReflectionSampler);
 #if OUTPUT_GAMMA_SPACE
 	// the capture will also be in gamma space, convert to linear:
 	PlanarReflection.rgb *= PlanarReflection.rgb;
diff --git a/Engine/Source/Runtime/Core/Public/Math/Plane.h b/Engine/Source/Runtime/Core/Public/Math/Plane.h
--- a/Engine/Source/Runtime/Core/Public/Math/Plane.h
+++ b/Engine/Source/Runtime/Core/Public/Math/Plane.h
@@ -177,6 +177,14 @@ public:
 	 */
 	TPlane<T> TransformByUsingAdjointT(const TMatrix<T>& M, T DetM, const TMatrix<T>& TA) const;
 
+	/**
+	 * Get the result of translating the plane by the given offset
+	 *
+	 * @param V The translation amount
+	 * @return The result of transform.
+	 */
+	TPlane<T> TranslateBy(const TVector<T>& V) const;
+
 	/**
 	 * Check if two planes are identical.
 	 *
@@ -477,6 +485,12 @@ FORCEINLINE TPlane<T> TPlane<T>::Flip() const
 	return TPlane<T>(-X, -Y, -Z, -W);
 }
 
+template<typename T>
+FORCEINLINE TPlane<T> TPlane<T>::TranslateBy(const TVector<T>& V) const
+{
+	return TPlane<T>(GetOrigin() + V, GetNormal());
+}
+
 template<typename T>
 FORCEINLINE bool TPlane<T>::operator==(const TPlane<T>& V) const
 {
diff --git a/Engine/Source/Runtime/Engine/Private/SceneView.cpp b/Engine/Source/Runtime/Engine/Private/SceneView.cpp
--- a/Engine/Source/Runtime/Engine/Private/SceneView.cpp
+++ b/Engine/Source/Runtime/Engine/Private/SceneView.cpp
@@ -2490,9 +2490,7 @@ void FSceneView::SetupCommonViewUniformBufferParameters(
 	ViewUniformShaderParameters.PrevClipToRelativeWorld = RelativeMatrices.PrevClipToRelativeWorld;
 
 	// Convert global clipping plane to translated world space
-	const FVector3f TranslatedGlobalClippingPlaneOrigin = GlobalClippingPlane.GetOrigin() + InViewMatrices.GetPreViewTranslation();
-	const FVector3f TranslatedGlobalClippingPlaneNormal = GlobalClippingPlane.GetNormal();
-	const FPlane4f TranslatedGlobalClippingPlane(TranslatedGlobalClippingPlaneOrigin, TranslatedGlobalClippingPlaneNormal);
+	const FPlane4f TranslatedGlobalClippingPlane(GlobalClippingPlane.TranslateBy(InViewMatrices.GetPreViewTranslation()));
 
 	ViewUniformShaderParameters.GlobalClippingPlane = FVector4(TranslatedGlobalClippingPlane.X, TranslatedGlobalClippingPlane.Y, TranslatedGlobalClippingPlane.Z, -TranslatedGlobalClippingPlane.W);
 
diff --git a/Engine/Source/Runtime/Renderer/Private/PlanarReflectionRendering.cpp b/Engine/Source/Runtime/Renderer/Private/PlanarReflectionRendering.cpp
--- a/Engine/Source/Runtime/Renderer/Private/PlanarReflectionRendering.cpp
+++ b/Engine/Source/Runtime/Renderer/Private/PlanarReflectionRendering.cpp
@@ -43,9 +43,12 @@ void SetupPlanarReflectionUniformParameters(const class FSceneView& View, const
 	{
 		ensure(ReflectionSceneProxy->ViewRect[0].Min.X >= 0);
 
+		const FVector PreViewTranslation = View.ViewMatrices.GetPreViewTranslation();
+		const FPlane4f TranslatedReflectionPlane(ReflectionSceneProxy->ReflectionPlane.TranslateBy(PreViewTranslation));
+
 		// Need to set W separately due to FVector = FPlane, which sets W to 1.0.
-		OutParameters.ReflectionPlane = ReflectionSceneProxy->ReflectionPlane;
-		OutParameters.ReflectionPlane.W = ReflectionSceneProxy->ReflectionPlane.W;
+		OutParameters.ReflectionPlane = TranslatedReflectionPlane;
+		OutParameters.ReflectionPlane.W = TranslatedReflectionPlane.W;
 
 		PlanarReflectionTextureValue = ReflectionSceneProxy->RenderTarget;
 
@@ -79,7 +82,7 @@ void SetupPlanarReflectionUniformParameters(const class FSceneView& View, const
 				(ViewRect.Height() / 2.0f + ViewRect.Min.Y) * InvBufferSizeY);
 		}
 
-		OutParameters.PlanarReflectionOrigin = ReflectionSceneProxy->PlanarReflectionOrigin;
+		OutParameters.PlanarReflectionOrigin = PreViewTranslation + ReflectionSceneProxy->PlanarReflectionOrigin;
 		OutParameters.PlanarReflectionXAxis = ReflectionSceneProxy->PlanarReflectionXAxis;
 		OutParameters.PlanarReflectionYAxis = ReflectionSceneProxy->PlanarReflectionYAxis;
 		OutParameters.InverseTransposeMirrorMatrix = ReflectionSceneProxy->InverseTransposeMirrorMatrix;
