Fix a cubemap shadow depth vs-ps linkage error

diff --git a/Engine/Shaders/Private/ShadowDepthVertexShader.usf b/Engine/Shaders/Private/ShadowDepthVertexShader.usf
--- a/Engine/Shaders/Private/ShadowDepthVertexShader.usf
+++ b/Engine/Shaders/Private/ShadowDepthVertexShader.usf
@@ -253,9 +253,6 @@ void Main(
   #else
 	out FShadowDepthVSToPS OutParameters,
   #endif
-#if FEATURE_LEVEL >= FEATURE_LEVEL_SM4 && USING_VERTEX_SHADER_LAYER
-	out uint LayerIndex : SV_RenderTargetArrayIndex,
-#endif
 #else
 	out FShadowDepthVSToPS OutParameters,
 #endif
@@ -263,6 +260,9 @@ void Main(
 	out nointerpolation uint PackedPageInfo : TEXCOORD8,
 #endif
 	out float4 OutPosition : SV_POSITION
+#if ONEPASS_POINTLIGHT_SHADOW && FEATURE_LEVEL >= FEATURE_LEVEL_SM4 && USING_VERTEX_SHADER_LAYER
+	,out uint LayerIndex : SV_RenderTargetArrayIndex
+#endif
 #if VIRTUAL_SM_ENABLED	
 	// OLA-TODO: this collides with instanced stereo, which thankfully is not used with shadow maps, so should be fine, presumably.
 	, out float4 OutVirtualSmPageClip : SV_ClipDistance
