Fix invalid Strata normal in forward lighting.

diff --git a/Engine/Shaders/Private/Strata/Strata.ush b/Engine/Shaders/Private/Strata/Strata.ush
--- a/Engine/Shaders/Private/Strata/Strata.ush
+++ b/Engine/Shaders/Private/Strata/Strata.ush
@@ -1006,7 +1006,11 @@ float3x3 StrataGetBSDFSharedBasis_DeferredShading(in FStrataPixelHeader StrataPi
 
 float3x3 StrataGetBSDFSharedBasis(in FStrataPixelHeader StrataPixelHeader, in FStrataBSDF BSDF, in FStrataAddressing StrataAddressing)
 {
+#if STRATA_INLINE_SHADING
+	return StrataGetBSDFSharedBasis_InlineShading(StrataPixelHeader, BSDF, StrataAddressing);
+#else 
 	return StrataGetBSDFSharedBasis_DeferredShading(StrataPixelHeader, BSDF, StrataAddressing);
+#endif
 }
 
 ///////////////////////////////////////////////////////////////////////////////
