Fixed error in LandscapeGrassWeight.usf when GPUScene is disabled

diff --git a/Engine/Shaders/Private/LandscapeGrassWeight.usf b/Engine/Shaders/Private/LandscapeGrassWeight.usf
--- a/Engine/Shaders/Private/LandscapeGrassWeight.usf
+++ b/Engine/Shaders/Private/LandscapeGrassWeight.usf
@@ -76,7 +76,7 @@ void PSMain(
 #if VF_USE_PRIMITIVE_SCENE_DATA
 	float Height = clamp(round(Interpolants.LocalPosZ / TERRAIN_ZSCALE) + 32768.0, 0, 65535);
 #else // !VF_USE_PRIMITIVE_SCENE_DATA
-	float3 LocalPos = LWCMultiply(MaterialParameters.AbsoluteWorldPosition, Primitive.WorldToLocal);
+	float3 LocalPos = LWCMultiply(MaterialParameters.AbsoluteWorldPosition, GetWorldToInstance(MaterialParameters));
 	float Height = clamp(round(LocalPos.z / TERRAIN_ZSCALE) + 32768.0, 0, 65535);
 #endif // VF_USE_PRIMITIVE_SCENE_DATA
 	float2 PackedHeight = float2(floor(Height / 256), fmod(Height, 256)) / 255;
