Support LWC on GPUSkinnedVertexFactory and SkinCache shaders
*) change GPUSkinVF to do cloth simulation in component space instead of world space
*) transform simulated pos from component space into local space and blend with skinned pos

diff --git a/Engine/Shaders/Private/GpuSkinCacheComputeShader.usf b/Engine/Shaders/Private/GpuSkinCacheComputeShader.usf
--- a/Engine/Shaders/Private/GpuSkinCacheComputeShader.usf
+++ b/Engine/Shaders/Private/GpuSkinCacheComputeShader.usf
@@ -73,8 +73,8 @@ Buffer<float4> ClothBuffer;
 Buffer<float2> ClothPositionsAndNormalsBuffer;
 uint ClothBufferOffset;
 float ClothBlendWeight;
-float4x4 ClothLocalToWorld;
-float4x4 ClothWorldToLocal;
+/** Transform from cloth space (relative to cloth root bone) to local/component space */
+float4x4 ClothToLocal;
 
 void GetClothBuffer(uint VertexIndex, uint ClothIndexOffset, out FVertexTriangleInfluence Influences[NUM_INFLUENCES_PER_VERTEX])
 {
@@ -413,9 +413,9 @@ void SkinCacheUpdateBatchCS(uint3 GroupID : SV_GroupID,
         ClothTangentX = normalize(TangentPosition - Intermediates.SimulatedPosition);
         ClothTangentZ = normalize(NormalPosition - Intermediates.SimulatedPosition);
 
-	    // cloth data are all in world space so need to change into local space
-	    ClothTangentX = mul(ClothTangentX, (half3x3)ClothWorldToLocal);
-	    ClothTangentZ = mul(ClothTangentZ, (half3x3)ClothWorldToLocal);
+	    // Simulated cloth data are in cloth space so need to change into local space
+	    ClothTangentX = mul(ClothTangentX, (half3x3)ClothToLocal);
+	    ClothTangentZ = mul(ClothTangentZ, (half3x3)ClothToLocal);
     }
 	else
 	{
@@ -429,9 +429,9 @@ void SkinCacheUpdateBatchCS(uint3 GroupID : SV_GroupID,
 	float3 SPos = SkinPosition(Intermediates) * SkinCacheDebug;
 
 #if GPUSKIN_APEX_CLOTH
-    float4 TransformedSkinPos = mul(float4(SPos, 1), ClothLocalToWorld);
-    float3 BlendedPos = lerp(TransformedSkinPos.xyz, Intermediates.SimulatedPosition, ClothWeight);
-    SPos = mul(float4(BlendedPos, 1), ClothWorldToLocal).xyz;
+	// Transform simulated position from cloth space into local space and blend with skinned position
+    float4 TransformedSimulatedPos = mul(float4(Intermediates.SimulatedPosition.xyz, 1), ClothToLocal);
+    SPos = lerp(SPos.xyz, TransformedSimulatedPos, ClothWeight);
 #endif
 
 	float3 TangentX = normalize(mul((float3x3)Intermediates.BlendMatrix, LocalTangentX));
diff --git a/Engine/Shaders/Private/GpuSkinVertexFactory.ush b/Engine/Shaders/Private/GpuSkinVertexFactory.ush
--- a/Engine/Shaders/Private/GpuSkinVertexFactory.ush
+++ b/Engine/Shaders/Private/GpuSkinVertexFactory.ush
@@ -47,13 +47,13 @@ uint NumBoneInfluencesParam;
 #if GPUSKIN_APEX_CLOTH
 /** Vertex buffer from which to read simulated positions of clothing. */
 Buffer<float2> ClothSimulVertsPositionsNormals;
-/** Transform to apply when positions/normals are loaded */
-float4x4 ClothLocalToWorld;
+/** Transform from cloth space (relative to cloth root bone) to local(component) space */
+float4x4 ClothToLocal;
 /** blend weight between simulated positions and original key-framed animation */
 float ClothBlendWeight;
 #if GPUSKIN_APEX_CLOTH_PREVIOUS
 Buffer<float2> PreviousClothSimulVertsPositionsNormals;
-float4x4 PreviousClothLocalToWorld;
+float4x4 PreviousClothToLocal;
 #endif //GPUSKIN_APEX_CLOTH_PREVIOUS
 #endif// #if GPUSKIN_APEX_CLOTH
 
@@ -375,14 +375,12 @@ float3 GetClothSimulPosition(int Index, bool bPrevious)
 #if GPUSKIN_APEX_CLOTH_PREVIOUS
 	if(bPrevious)
 	{
-		Position = float3(PreviousClothSimulVertsPositionsNormals[Index * 3], PreviousClothSimulVertsPositionsNormals[Index * 3 + 1].x);
-		return float3(mul(float4(Position, 1), PreviousClothLocalToWorld).xyz);
+		return float3(PreviousClothSimulVertsPositionsNormals[Index * 3], PreviousClothSimulVertsPositionsNormals[Index * 3 + 1].x);
 	}
 	else
 #endif
 	{
-		Position = float3(ClothSimulVertsPositionsNormals[Index * 3], ClothSimulVertsPositionsNormals[Index * 3 + 1].x);
-		return float3(mul(float4(Position, 1), ClothLocalToWorld).xyz);
+		return float3(ClothSimulVertsPositionsNormals[Index * 3], ClothSimulVertsPositionsNormals[Index * 3 + 1].x);
 	}
 }
 
@@ -392,14 +390,12 @@ float3 GetClothSimulNormal(int Index, bool bPrevious)
 #if GPUSKIN_APEX_CLOTH_PREVIOUS
 	if(bPrevious)
 	{
-		Normal = float3(PreviousClothSimulVertsPositionsNormals[Index * 3 + 1].y, PreviousClothSimulVertsPositionsNormals[Index * 3 + 2]);
-		return float3(mul(float4(Normal, 0), PreviousClothLocalToWorld).xyz);
+		return float3(PreviousClothSimulVertsPositionsNormals[Index * 3 + 1].y, PreviousClothSimulVertsPositionsNormals[Index * 3 + 2]);
 	}
 	else
 #endif
 	{
-		Normal = float3(ClothSimulVertsPositionsNormals[Index * 3 + 1].y, ClothSimulVertsPositionsNormals[Index * 3 + 2]);
-		return float3(mul(float4(Normal, 0), ClothLocalToWorld).xyz);
+		return float3(ClothSimulVertsPositionsNormals[Index * 3 + 1].y, ClothSimulVertsPositionsNormals[Index * 3 + 2]);
 	}
 }
 
@@ -752,9 +748,9 @@ float3x3 SkinTangents( FVertexFactoryInput Input, FVertexFactoryIntermediates In
 		    TangentToLocal[0] = normalize(TangentPosition - Intermediates.SimulatedPosition.xyz);
 		    TangentToLocal[2] = normalize(NormalPosition - Intermediates.SimulatedPosition.xyz);
 
-		    // cloth data are all in world space so need to change into local space
-		    TangentToLocal[0] = LWCMultiplyVector(TangentToLocal[0], Intermediates.WorldToLocal);
-		    TangentToLocal[2] = LWCMultiplyVector(TangentToLocal[2], Intermediates.WorldToLocal);
+			// Simulated cloth data are in cloth space so need to change into local space
+			TangentToLocal[0] = mul(float4(TangentToLocal[0], 0.0f), ClothToLocal).xyz;
+			TangentToLocal[2] = mul(float4(TangentToLocal[2], 0.0f), ClothToLocal).xyz;
 	
 			// derive the new binormal by getting the cross product of the normal and tangent
 			// and flip vector based on sign of tangent basis determinant
@@ -832,25 +828,16 @@ FVertexFactoryIntermediates GetVertexFactoryIntermediates(FVertexFactoryInput In
 
 float4 CalcWorldPosition(FVertexFactoryInput Input, FVertexFactoryIntermediates Intermediates)
 {
+	float3 SkinnedPosition = SkinPosition(Input, Intermediates);
 #if GPUSKIN_APEX_CLOTH
 	if (IsSimulatedVertex(Intermediates.ClothVertexInfluences[0]))
 	{
-		// skinned positions
-		float4 SkinnedPosition = float4(SkinPosition(Input, Intermediates), 1);
-
-		// LWC_TODO - This is computing a value in TranslatedWorldSpace which is good, but 'SimulatedPosition' is currently in AbsoluteWorldSpace which is bad, need to update cloth sim to work in TLS
-		float4 TransformedSkinPos = mul(SkinnedPosition, LWCHackToFloat(Intermediates.LocalToWorld));
-
-		// Intermediates.SimulatedPosition is a clothing position in world coord
-		float3 BlendedPos = lerp(TransformedSkinPos.xyz, Intermediates.SimulatedPosition.xyz, ClothBlendWeight * Intermediates.SimulatedPosition.w);
-		return float4(BlendedPos + LWCHackToFloat(ResolvedView.PreViewTranslation), 1);
+		// Transform simulated position from cloth space into local space and blend with skinned position
+		float4 TransformedSimulatedPos = mul(float4(Intermediates.SimulatedPosition.xyz, 1), ClothToLocal);
+		SkinnedPosition = lerp(SkinnedPosition, TransformedSimulatedPos.xyz, ClothBlendWeight * Intermediates.SimulatedPosition.w);
 	}
-	else
 #endif
-	{
-		FLWCMatrix LocalToWorld = Intermediates.SceneData.InstanceData.LocalToWorld;
-		return TransformLocalToTranslatedWorld(SkinPosition(Input, Intermediates), LocalToWorld);
-	}
+	return TransformLocalToTranslatedWorld(SkinnedPosition, Intermediates.LocalToWorld);
 }
 
 /**
@@ -950,18 +937,16 @@ FVertexFactoryInterpolantsVSToPS VertexFactoryGetInterpolantsVSToPS(FVertexFacto
 float4 VertexFactoryGetPreviousWorldPosition(FVertexFactoryInput Input, FVertexFactoryIntermediates Intermediates)
 {
 	float4x4 PreviousLocalToWorldTranslated = LWCMultiplyTranslation(Intermediates.PrevLocalToWorld, ResolvedView.PrevPreViewTranslation);
-	float4 PrevSkinPosInWorld = mul(float4(SkinPreviousPosition(Input, Intermediates),1), PreviousLocalToWorldTranslated);
+	float3 PrevSkinnedPosition = SkinPreviousPosition(Input, Intermediates);
 #if GPUSKIN_APEX_CLOTH
 	if (IsSimulatedVertex(Intermediates.ClothVertexInfluences[0]))
 	{
-		// TODO LWC
-		float3 PrevSimulatedPos = Intermediates.PreviousSimulatedPosition.xyz + LWCHackToFloat(ResolvedView.PrevPreViewTranslation);
-		float3 BlendedPrevPos = lerp(PrevSkinPosInWorld.xyz, PrevSimulatedPos.xyz, ClothBlendWeight * Intermediates.PreviousSimulatedPosition.w);
-		return float4(BlendedPrevPos, 1);
+		// Transform simulated position from cloth space into local space and blend with skinned position
+		float4 PrevTransformedSimulatedPos = mul(float4(Intermediates.PreviousSimulatedPosition.xyz, 1), PreviousClothToLocal);
+		PrevSkinnedPosition = lerp(PrevSkinnedPosition, PrevTransformedSimulatedPos.xyz, ClothBlendWeight * Intermediates.PreviousSimulatedPosition.w);
 	}
-#endif // GPUSKIN_APEX_CLOTH
-
-	return PrevSkinPosInWorld;
+#endif
+	return mul(float4(PrevSkinnedPosition,1), PreviousLocalToWorldTranslated);
 }
 
 
diff --git a/Engine/Source/Runtime/Engine/Private/GPUSkinCache.cpp b/Engine/Source/Runtime/Engine/Private/GPUSkinCache.cpp
--- a/Engine/Source/Runtime/Engine/Private/GPUSkinCache.cpp
+++ b/Engine/Source/Runtime/Engine/Private/GPUSkinCache.cpp
@@ -270,9 +270,7 @@ public:
 		uint32 ClothBufferOffset = 0;
         float ClothBlendWeight = 0.0f;
 
-		// LWC_TODO: Cloth LocalToWorld/WorldToLocal matrices forced to float.
-        FMatrix44f ClothLocalToWorld = FMatrix44f::Identity;
-        FMatrix44f ClothWorldToLocal = FMatrix44f::Identity;
+        FMatrix44f ClothToLocal = FMatrix44f::Identity;
 
 		// triangle index buffer (input for the RecomputeSkinTangents, might need special index buffer unique to position and normal, not considering UV/vertex color)
 		uint32 IndexBufferOffsetValue = 0;
@@ -556,8 +554,7 @@ public:
 		ClothPositionsAndNormalsBuffer.Bind(Initializer.ParameterMap, TEXT("ClothPositionsAndNormalsBuffer"));
 		ClothBufferOffset.Bind(Initializer.ParameterMap, TEXT("ClothBufferOffset"));
 		ClothBlendWeight.Bind(Initializer.ParameterMap, TEXT("ClothBlendWeight"));
-		ClothLocalToWorld.Bind(Initializer.ParameterMap, TEXT("ClothLocalToWorld"));
-		ClothWorldToLocal.Bind(Initializer.ParameterMap, TEXT("ClothWorldToLocal"));
+		ClothToLocal.Bind(Initializer.ParameterMap, TEXT("ClothToLocal"));
 	}
 
 	void SetParameters(
@@ -605,8 +602,7 @@ public:
 			SetSRVParameter(RHICmdList, ShaderRHI, ClothPositionsAndNormalsBuffer, DispatchData.ClothPositionsAndNormalsBuffer);
 			SetShaderValue(RHICmdList, ShaderRHI, ClothBufferOffset, DispatchData.ClothBufferOffset);
 			SetShaderValue(RHICmdList, ShaderRHI, ClothBlendWeight, DispatchData.ClothBlendWeight);
-			SetShaderValue(RHICmdList, ShaderRHI, ClothLocalToWorld, DispatchData.ClothLocalToWorld);
-			SetShaderValue(RHICmdList, ShaderRHI, ClothWorldToLocal, DispatchData.ClothWorldToLocal);
+			SetShaderValue(RHICmdList, ShaderRHI, ClothToLocal, DispatchData.ClothToLocal);
 		}
 
 		SetShaderValue(RHICmdList, ShaderRHI, SkinCacheDebug, CVarGPUSkinCacheDebug.GetValueOnRenderThread());
@@ -652,8 +648,7 @@ private:
 	LAYOUT_FIELD(FShaderResourceParameter, ClothPositionsAndNormalsBuffer)
 	LAYOUT_FIELD(FShaderParameter, ClothBufferOffset)
 	LAYOUT_FIELD(FShaderParameter, ClothBlendWeight)
-	LAYOUT_FIELD(FShaderParameter, ClothLocalToWorld)
-	LAYOUT_FIELD(FShaderParameter, ClothWorldToLocal)
+	LAYOUT_FIELD(FShaderParameter, ClothToLocal)
 
 };
 
@@ -1359,7 +1354,7 @@ bool FGPUSkinCache::ProcessEntry(
 	const FMorphVertexBuffer* MorphVertexBuffer,
 	const FSkeletalMeshVertexClothBuffer* ClothVertexBuffer, 
 	const FClothSimulData* SimData,
-	const FMatrix44f& ClothLocalToWorld, 
+	const FMatrix44f& ClothToLocal,
 	float ClothBlendWeight, 
 	uint32 RevisionNumber, 
 	int32 Section,
@@ -1545,8 +1540,7 @@ bool FGPUSkinCache::ProcessEntry(
 		}
 
         InOutEntry->DispatchData[Section].ClothBlendWeight = ClothBlendWeight;
-        InOutEntry->DispatchData[Section].ClothLocalToWorld = ClothLocalToWorld;
-        InOutEntry->DispatchData[Section].ClothWorldToLocal = ClothLocalToWorld.Inverse();
+        InOutEntry->DispatchData[Section].ClothToLocal = ClothToLocal;
     }
     InOutEntry->DispatchData[Section].SkinType = ClothVertexBuffer && InOutEntry->DispatchData[Section].ClothPositionsAndNormalsBuffer ? 2 : (bMorph ? 1 : 0);
 
diff --git a/Engine/Source/Runtime/Engine/Private/GPUSkinVertexFactory.cpp b/Engine/Source/Runtime/Engine/Private/GPUSkinVertexFactory.cpp
--- a/Engine/Source/Runtime/Engine/Private/GPUSkinVertexFactory.cpp
+++ b/Engine/Source/Runtime/Engine/Private/GPUSkinVertexFactory.cpp
@@ -1247,8 +1247,8 @@ public:
 		FGPUSkinVertexFactoryShaderParameters::Bind(ParameterMap);
 		ClothSimulVertsPositionsNormalsParameter.Bind(ParameterMap,TEXT("ClothSimulVertsPositionsNormals"));
 		PreviousClothSimulVertsPositionsNormalsParameter.Bind(ParameterMap,TEXT("PreviousClothSimulVertsPositionsNormals"));
-		ClothLocalToWorldParameter.Bind(ParameterMap, TEXT("ClothLocalToWorld"));
-		PreviousClothLocalToWorldParameter.Bind(ParameterMap, TEXT("PreviousClothLocalToWorld"));
+		ClothToLocalParameter.Bind(ParameterMap, TEXT("ClothToLocal"));
+		PreviousClothToLocalParameter.Bind(ParameterMap, TEXT("PreviousClothToLocal"));
 		ClothBlendWeightParameter.Bind(ParameterMap, TEXT("ClothBlendWeight"));
 		GPUSkinApexClothParameter.Bind(ParameterMap, TEXT("GPUSkinApexCloth"));
 		GPUSkinApexClothStartIndexOffsetParameter.Bind(ParameterMap, TEXT("GPUSkinApexClothStartIndexOffset"));
@@ -1278,7 +1278,7 @@ public:
 		uint32 FrameNumber = View->Family->FrameNumber;
 
 		ShaderBindings.Add(ClothSimulVertsPositionsNormalsParameter, ClothShaderData.GetClothBufferForReading(false, FrameNumber).VertexBufferSRV);
-		ShaderBindings.Add(ClothLocalToWorldParameter, ClothShaderData.GetClothLocalToWorldForReading(false, FrameNumber));
+		ShaderBindings.Add(ClothToLocalParameter, ClothShaderData.GetClothToLocalForReading(false, FrameNumber));
 		ShaderBindings.Add(ClothBlendWeightParameter,ClothShaderData.ClothBlendWeight);
 
 		// Mobile doesn't support motion blur, no need to feed the previous frame cloth data
@@ -1287,7 +1287,7 @@ public:
 		if (!bIsMobile)
 		{
 			ShaderBindings.Add(PreviousClothSimulVertsPositionsNormalsParameter, ClothShaderData.GetClothBufferForReading(true, FrameNumber).VertexBufferSRV);
-			ShaderBindings.Add(PreviousClothLocalToWorldParameter, ClothShaderData.GetClothLocalToWorldForReading(true, FrameNumber));
+			ShaderBindings.Add(PreviousClothToLocalParameter, ClothShaderData.GetClothToLocalForReading(true, FrameNumber));
 		}
 
 		ShaderBindings.Add(GPUSkinApexClothParameter, ClothVertexFactory->GetClothBuffer());
@@ -1300,8 +1300,8 @@ protected:
 	
 	LAYOUT_FIELD(FShaderResourceParameter, ClothSimulVertsPositionsNormalsParameter);
 	LAYOUT_FIELD(FShaderResourceParameter, PreviousClothSimulVertsPositionsNormalsParameter);
-	LAYOUT_FIELD(FShaderParameter, ClothLocalToWorldParameter);
-	LAYOUT_FIELD(FShaderParameter, PreviousClothLocalToWorldParameter);
+	LAYOUT_FIELD(FShaderParameter, ClothToLocalParameter);
+	LAYOUT_FIELD(FShaderParameter, PreviousClothToLocalParameter);
 	LAYOUT_FIELD(FShaderParameter, ClothBlendWeightParameter);
 	LAYOUT_FIELD(FShaderResourceParameter, GPUSkinApexClothParameter);
 	LAYOUT_FIELD(FShaderParameter, GPUSkinApexClothStartIndexOffsetParameter);
diff --git a/Engine/Source/Runtime/Engine/Private/SkeletalRenderGPUSkin.cpp b/Engine/Source/Runtime/Engine/Private/SkeletalRenderGPUSkin.cpp
--- a/Engine/Source/Runtime/Engine/Private/SkeletalRenderGPUSkin.cpp
+++ b/Engine/Source/Runtime/Engine/Private/SkeletalRenderGPUSkin.cpp
@@ -751,6 +751,8 @@ void FSkeletalMeshObjectGPUSkin::ProcessUpdatedDynamicData(EGPUSkinCacheEntryMod
 				ClothShaderData.ClothBlendWeight = DynamicData->ClothBlendWeight;
 
 				bNeedFence = ClothShaderData.UpdateClothSimulData(RHICmdList, SimData->Positions, SimData->Normals, FrameNumberToPrepare, FeatureLevel) || bNeedFence;
+				// Transform from cloth space to local space. Cloth space is relative to cloth root bone, local space is component space.
+				ClothShaderData.GetClothToLocalForWriting(FrameNumberToPrepare) = SimData->ComponentRelativeTransform.ToMatrixWithScale();
 				ClothShaderData.GetClothLocalToWorldForWriting(FrameNumberToPrepare) = SimData->ComponentRelativeTransform.ToMatrixWithScale() * DynamicData->ClothObjectLocalToWorld;
 			}
 #endif // WITH_APEX_CLOTHING || WITH_CHAOS_CLOTHING
@@ -758,10 +760,7 @@ void FSkeletalMeshObjectGPUSkin::ProcessUpdatedDynamicData(EGPUSkinCacheEntryMod
 			// Try to use the GPU skinning cache if possible
 			if (bUseSkinCache)
 			{
-				// This takes the cloth positions from cloth space into world space
-				FMatrix44f ClothLocalToWorld = bClothFactory ? VertexFactoryData.ClothVertexFactories[SectionIdx]->GetClothShaderData().GetClothLocalToWorldForWriting(FrameNumberToPrepare) : FMatrix44f::Identity;
-				// Matrices are transposed in UE meaning matrix multiples need to happen in reverse ((AB)x = b becomes xTBTAT = b).
-				FMatrix44f LocalToCloth = DynamicData->ClothObjectLocalToWorld * ClothLocalToWorld.Inverse();
+				FMatrix44f ClothToLocal = bClothFactory ? VertexFactoryData.ClothVertexFactories[SectionIdx]->GetClothShaderData().GetClothToLocalForWriting(FrameNumberToPrepare) : FMatrix44f::Identity;
 
 				// ProcessEntry returns false if not enough memory is left in skin cache to allocate for the mesh, if that happens don't try to process subsequent sections because they will also fail.
 				if (bSkinCacheResult)
@@ -776,7 +775,7 @@ void FSkeletalMeshObjectGPUSkin::ProcessUpdatedDynamicData(EGPUSkinCacheEntryMod
 						bMorph ? &MorphVertexBuffer : 0,
 						bClothFactory ? &LODData.ClothVertexBuffer : 0,
 						bClothFactory ? DynamicData->ClothingSimData.Find(Section.CorrespondClothAssetIndex) : 0,
-						LocalToCloth,
+						ClothToLocal,
 						DynamicData->ClothBlendWeight,
 						RevisionNumber,
 						SectionIdx,
@@ -1475,6 +1474,7 @@ void FSkeletalMeshObjectGPUSkin::RefreshClothingTransforms(const FMatrix& InNewL
 
 					if(FClothSimulData* SimData = DynamicData->ClothingSimData.Find(ActorIdx))
 					{
+						ClothShaderData.GetClothToLocalForWriting(FrameNumber) = SimData->ComponentRelativeTransform.ToMatrixWithScale();
 						ClothShaderData.GetClothLocalToWorldForWriting(FrameNumber) = SimData->ComponentRelativeTransform.ToMatrixWithScale() * InNewLocalToWorld;
 					}
 				}
diff --git a/Engine/Source/Runtime/Engine/Public/GPUSkinCache.h b/Engine/Source/Runtime/Engine/Public/GPUSkinCache.h
--- a/Engine/Source/Runtime/Engine/Public/GPUSkinCache.h
+++ b/Engine/Source/Runtime/Engine/Public/GPUSkinCache.h
@@ -167,7 +167,7 @@ public:
 		const FMorphVertexBuffer* MorphVertexBuffer, 
 		const FSkeletalMeshVertexClothBuffer* ClothVertexBuffer, 
 		const FClothSimulData* SimData,
-		const FMatrix44f& ClothLocalToWorld,
+		const FMatrix44f& ClothToLocal,
 		float ClothBlendWeight, 
 		uint32 RevisionNumber, 
 		int32 Section,
diff --git a/Engine/Source/Runtime/Engine/Public/GPUSkinVertexFactory.h b/Engine/Source/Runtime/Engine/Public/GPUSkinVertexFactory.h
--- a/Engine/Source/Runtime/Engine/Public/GPUSkinVertexFactory.h
+++ b/Engine/Source/Runtime/Engine/Public/GPUSkinVertexFactory.h
@@ -721,6 +721,27 @@ public:
 			return ClothSimulPositionNormalBuffer[Index];
 		}
 		
+		FMatrix44f& GetClothToLocalForWriting(uint32 FrameNumber)
+		{
+			uint32 Index = GetOldestIndex(FrameNumber);
+			Index = (BufferFrameNumber[0] == FrameNumber) ? 0 : Index;
+			Index = (BufferFrameNumber[1] == FrameNumber) ? 1 : Index;
+
+			return ClothToLocal[Index];
+		}
+
+		const FMatrix44f& GetClothToLocalForReading(bool bPrevious, uint32 FrameNumber) const
+		{
+			int32 Index = GetMostRecentIndex(FrameNumber);
+
+			if(bPrevious && DoWeHavePreviousData())
+			{
+				Index = 1 - Index;
+			}
+
+			return ClothToLocal[Index];
+		}
+
 		FMatrix44f& GetClothLocalToWorldForWriting(uint32 FrameNumber)
 		{
 			uint32 Index = GetOldestIndex(FrameNumber);
@@ -760,6 +781,7 @@ public:
 		 * Matrix to apply to positions/normals
 		 */
 		FMatrix44f ClothLocalToWorld[2];
+		FMatrix44f ClothToLocal[2];
 
 		/** Whether to double buffer. */
 		bool bDoubleBuffer = false;
@@ -839,6 +861,8 @@ public:
 
 			ClothLocalToWorld[0] = FMatrix44f::Identity;
 			ClothLocalToWorld[1] = FMatrix44f::Identity;
+			ClothToLocal[0] = FMatrix44f::Identity;
+			ClothToLocal[1] = FMatrix44f::Identity;
 
 			bDoubleBuffer = false;
 		}
