Fix for LWC mod/trig operations

diff --git a/Engine/Shaders/Private/LWCOperations.ush b/Engine/Shaders/Private/LWCOperations.ush
--- a/Engine/Shaders/Private/LWCOperations.ush
+++ b/Engine/Shaders/Private/LWCOperations.ush
@@ -10,6 +10,13 @@ FLWCType LWCNormalizeTile(FLWCType V)
 	return LWCConstructor(IntTile, (V.Tile - IntTile) * UE_LWC_RENDER_TILE_SIZE + V.Offset);
 }
 
+// Forces the tile to be an integral value, but doesn't adjust range of offset
+FLWCType LWCMakeIntTile(FLWCType V)
+{
+	FFloatType IntTile = floor(V.Tile);
+	return LWCConstructor(IntTile, (V.Tile - IntTile) * UE_LWC_RENDER_TILE_SIZE + V.Offset);
+}
+
 FFloatType LWCSqrtUnscaled(FLWCType V) { return sqrt(V.Offset * UE_LWC_RENDER_TILE_SIZE_RCP + LWCGetTile(V)); }
 FFloatType LWCRsqrtUnscaled(FLWCType V) { return rsqrt(V.Offset * UE_LWC_RENDER_TILE_SIZE_RCP + LWCGetTile(V)); }
 FFloatType LWCRcpUnscaled(FLWCType V) { return rcp(V.Offset * UE_LWC_RENDER_TILE_SIZE_RCP + LWCGetTile(V)); }
@@ -23,6 +30,7 @@ FFloatType LWCRcp(FLWCType V) { return LWCRcpScaled(V, UE_LWC_RENDER_TILE_SIZE_R
 #else // !UE_LARGE_WORLD_COORDINATES_DISABLED
 
 FLWCType LWCNormalizeTile(FLWCType V) { return V; }
+FLWCType LWCMakeIntTile(FLWCType V) { return V; }
 
 FFloatType LWCSqrtUnscaled(FLWCType V) { return sqrt(V.Offset); }
 FFloatType LWCRsqrtUnscaled(FLWCType V) { return rsqrt(V.Offset); }
@@ -119,12 +127,12 @@ FFloatType LWCFrac(FLWCType V)
 
 	return frac(FracTile + V.Offset);
 }
-//FFloatType LWCFrac(FLWCType V) { FLWCType VN = LWCNormalizeTile(V); return frac(VN.Offset); }
+//FFloatType LWCFrac(FLWCType V) { FLWCType VN = LWCMakeIntTile(V); return frac(VN.Offset); }
 
-FLWCType LWCFloor(FLWCType V) { FLWCType VN = LWCNormalizeTile(V); return LWCConstructor(LWCGetTile(VN), floor(VN.Offset)); }
-FLWCType LWCCeil(FLWCType V) { FLWCType VN = LWCNormalizeTile(V); return LWCConstructor(LWCGetTile(VN), ceil(VN.Offset)); }
-FLWCType LWCRound(FLWCType V) { FLWCType VN = LWCNormalizeTile(V); return LWCConstructor(LWCGetTile(VN), round(VN.Offset)); }
-FLWCType LWCTrunc(FLWCType V) { FLWCType VN = LWCNormalizeTile(V); return LWCConstructor(LWCGetTile(VN), trunc(VN.Offset)); }
+FLWCType LWCFloor(FLWCType V) { FLWCType VN = LWCMakeIntTile(V); return LWCConstructor(LWCGetTile(VN), floor(VN.Offset)); }
+FLWCType LWCCeil(FLWCType V) { FLWCType VN = LWCMakeIntTile(V); return LWCConstructor(LWCGetTile(VN), ceil(VN.Offset)); }
+FLWCType LWCRound(FLWCType V) { FLWCType VN = LWCMakeIntTile(V); return LWCConstructor(LWCGetTile(VN), round(VN.Offset)); }
+FLWCType LWCTrunc(FLWCType V) { FLWCType VN = LWCMakeIntTile(V); return LWCConstructor(LWCGetTile(VN), trunc(VN.Offset)); }
 
 // These rely on +/-inf being handled properly by various min/max operations
 FFloatType LWCSign(FLWCType V) { return FFloatType(sign(LWCToFloat(V))); }
@@ -150,21 +158,25 @@ FLWCType LWCLerp(FLWCType Lhs, FLWCType Rhs, FFloatType S)
 
 FFloatType LWCFmod(FLWCType Lhs, FFloatType Rhs)
 {
-	FFloatType ModTileSize = fmod(UE_LWC_RENDER_TILE_SIZE, Rhs);
-	return fmod(ModTileSize * LWCGetTile(Lhs) + Lhs.Offset, Rhs);
+	return LWCToFloat(LWCSubtract(Lhs, LWCMultiply(LWCTrunc(LWCDivide(Lhs, Rhs)), Rhs)));
+	//FFloatType ModTileSize = fmod(UE_LWC_RENDER_TILE_SIZE, Rhs);
+	//return fmod(ModTileSize * LWCGetTile(Lhs) + Lhs.Offset, Rhs);
 }
 FFloatType LWCFmodFloor(FLWCType Lhs, FFloatType Rhs)
 {
-	FFloatType ModTileSize = FmodFloor(UE_LWC_RENDER_TILE_SIZE, Rhs);
-	return FmodFloor(LWCGetTile(Lhs) * ModTileSize + Lhs.Offset, Rhs);
+	return LWCToFloat(LWCSubtract(Lhs, LWCMultiply(LWCFloor(LWCDivide(Lhs, Rhs)), Rhs)));
+	//FFloatType ModTileSize = FmodFloor(UE_LWC_RENDER_TILE_SIZE, Rhs);
+	//return FmodFloor(LWCGetTile(Lhs) * ModTileSize + Lhs.Offset, Rhs);
 }
-FFloatType LWCFmodFloorPI(FLWCType Lhs)
+FFloatType LWCFmodFloorPI(FLWCType V)
 {
-	return FmodFloor(LWCGetTile(Lhs) * UE_LWC_RENDER_TILE_SIZE_FMOD_PI + Lhs.Offset, PI);
+	return LWCFmodFloor(V, PI);
+	//return LWCGetTile(V) * UE_LWC_RENDER_TILE_SIZE_FMOD_PI + V.Offset;
 }
-FFloatType LWCFmodFloor2PI(FLWCType Lhs)
+FFloatType LWCFmodFloor2PI(FLWCType V)
 {
-	return FmodFloor(LWCGetTile(Lhs) * UE_LWC_RENDER_TILE_SIZE_FMOD_2PI + Lhs.Offset, 2.0f*PI);
+	return LWCFmodFloor(V, 2.0f * PI);
+	//return LWCGetTile(V) * UE_LWC_RENDER_TILE_SIZE_FMOD_2PI + V.Offset;
 }
 
 FFloatType LWCSin(FLWCType V) { return sin(LWCFmodFloor2PI(V)); }
