Fix for GpuSkinVertexFactory compilation

diff --git a/Engine/Shaders/Private/GpuSkinVertexFactory.ush b/Engine/Shaders/Private/GpuSkinVertexFactory.ush
--- a/Engine/Shaders/Private/GpuSkinVertexFactory.ush
+++ b/Engine/Shaders/Private/GpuSkinVertexFactory.ush
@@ -938,14 +938,14 @@ float4 VertexFactoryGetPreviousWorldPosition(FVertexFactoryInput Input, FVertexF
 {
 	float4x4 PreviousLocalToWorldTranslated = LWCMultiplyTranslation(Intermediates.PrevLocalToWorld, ResolvedView.PrevPreViewTranslation);
 	float3 PrevSkinnedPosition = SkinPreviousPosition(Input, Intermediates);
-#if GPUSKIN_APEX_CLOTH
+#if GPUSKIN_APEX_CLOTH_PREVIOUS
 	if (IsSimulatedVertex(Intermediates.ClothVertexInfluences[0]))
 	{
 		// Transform simulated position from cloth space into local space and blend with skinned position
 		float4 PrevTransformedSimulatedPos = mul(float4(Intermediates.PreviousSimulatedPosition.xyz, 1), PreviousClothToLocal);
 		PrevSkinnedPosition = lerp(PrevSkinnedPosition, PrevTransformedSimulatedPos.xyz, ClothBlendWeight * Intermediates.PreviousSimulatedPosition.w);
 	}
-#endif
+#endif // GPUSKIN_APEX_CLOTH_PREVIOUS
 	return mul(float4(PrevSkinnedPosition,1), PreviousLocalToWorldTranslated);
 }
 
