Remove an extra unneeded frac() in LWCFrac, LWCApplyAddressModeWrap

diff --git a/Engine/Shaders/Private/LWCOperations.ush b/Engine/Shaders/Private/LWCOperations.ush
--- a/Engine/Shaders/Private/LWCOperations.ush
+++ b/Engine/Shaders/Private/LWCOperations.ush
@@ -121,13 +121,9 @@ FLWCType LWCNegate(FLWCType V) { return LWCConstructor(-LWCGetTile(V), -V.Offset
 
 FFloatType LWCFrac(FLWCType V)
 {
-	// Inner frac() avoids overflow if tile is large
-	// Outer frac() avoids helps avoid precision loss in the 'FracTile + V.Offset' addition
-	FFloatType FracTile = frac(frac(LWCGetTile(V)) * UE_LWC_RENDER_TILE_SIZE);
-
+	FFloatType FracTile = frac(LWCGetTile(V) * UE_LWC_RENDER_TILE_SIZE);
 	return frac(FracTile + V.Offset);
 }
-//FFloatType LWCFrac(FLWCType V) { FLWCType VN = LWCMakeIntTile(V); return frac(VN.Offset); }
 
 FLWCType LWCFloor(FLWCType V) { FLWCType VN = LWCMakeIntTile(V); return LWCConstructor(LWCGetTile(VN), floor(VN.Offset)); }
 FLWCType LWCCeil(FLWCType V) { FLWCType VN = LWCMakeIntTile(V); return LWCConstructor(LWCGetTile(VN), ceil(VN.Offset)); }
diff --git a/Engine/Shaders/Private/MaterialTemplate.ush b/Engine/Shaders/Private/MaterialTemplate.ush
--- a/Engine/Shaders/Private/MaterialTemplate.ush
+++ b/Engine/Shaders/Private/MaterialTemplate.ush
@@ -1848,7 +1848,7 @@ float LWCApplyAddressModeWrap(FLWCScalar V)
 {
 	// Compute the fractional part of the tile, then add to the offset
 	// Let the texture unit apply the final coordinate wrapping, which will allow derivatives to work correctly unless we cross a tile boundary
-	const float FracTile = frac(frac(LWCGetTile(V)) * UE_LWC_RENDER_TILE_SIZE);
+	const float FracTile = frac(LWCGetTile(V) * UE_LWC_RENDER_TILE_SIZE);
 	return FracTile + V.Offset;
 }
 
