Trivial wrapping of GPUSKIN_APEX_CLOTH macro just be extra safe

diff --git a/Engine/Shaders/Private/GpuSkinVertexFactory.ush b/Engine/Shaders/Private/GpuSkinVertexFactory.ush
--- a/Engine/Shaders/Private/GpuSkinVertexFactory.ush
+++ b/Engine/Shaders/Private/GpuSkinVertexFactory.ush
@@ -938,6 +938,7 @@ float4 VertexFactoryGetPreviousWorldPosition(FVertexFactoryInput Input, FVertexF
 {
 	float4x4 PreviousLocalToWorldTranslated = LWCMultiplyTranslation(Intermediates.PrevLocalToWorld, ResolvedView.PrevPreViewTranslation);
 	float3 PrevSkinnedPosition = SkinPreviousPosition(Input, Intermediates);
+#if GPUSKIN_APEX_CLOTH
 #if GPUSKIN_APEX_CLOTH_PREVIOUS
 	if (IsSimulatedVertex(Intermediates.ClothVertexInfluences[0]))
 	{
@@ -946,6 +947,7 @@ float4 VertexFactoryGetPreviousWorldPosition(FVertexFactoryInput Input, FVertexF
 		PrevSkinnedPosition = lerp(PrevSkinnedPosition, PrevTransformedSimulatedPos.xyz, ClothBlendWeight * Intermediates.PreviousSimulatedPosition.w);
 	}
 #endif // GPUSKIN_APEX_CLOTH_PREVIOUS
+#endif // GPUSKIN_APEX_CLOTH
 	return mul(float4(PrevSkinnedPosition,1), PreviousLocalToWorldTranslated);
 }
 
