Optimize LWC comparison operations

diff --git a/Engine/Shaders/Private/LWCOperations.ush b/Engine/Shaders/Private/LWCOperations.ush
--- a/Engine/Shaders/Private/LWCOperations.ush
+++ b/Engine/Shaders/Private/LWCOperations.ush
@@ -46,15 +46,11 @@ FFloatType LWCRcp(FLWCType V) { return LWCRcpUnscaled(V); }
 
 #define LWC_COMPARE_OP(Name, Op, FLWCType, FFloatType, FBoolType) \
 	FBoolType Name(FLWCType Lhs, FLWCType Rhs) { \
-		FLWCType LhsN = LWCNormalizeTile(Lhs); \
-		FLWCType RhsN = LWCNormalizeTile(Rhs); \
-		return (LWCGetTile(LhsN) != LWCGetTile(RhsN)) ? (LWCGetTile(LhsN) Op LWCGetTile(RhsN)) : (LhsN.Offset Op RhsN.Offset); } \
+		return (LWCGetTile(Lhs) - LWCGetTile(Rhs)) * UE_LWC_RENDER_TILE_SIZE Op Rhs.Offset - Lhs.Offset; } \
 	FBoolType Name(FFloatType Lhs, FLWCType Rhs) { \
-		FLWCType RhsN = LWCNormalizeTile(Rhs); \
-		return ((FFloatType)0 != LWCGetTile(RhsN)) ? ((FFloatType)0 Op LWCGetTile(RhsN)) : (Lhs Op RhsN.Offset); } \
+		return -LWCGetTile(Rhs) * UE_LWC_RENDER_TILE_SIZE Op Rhs.Offset - Lhs; } \
 	FBoolType Name(FLWCType Lhs, FFloatType Rhs) { \
-		FLWCType LhsN = LWCNormalizeTile(Lhs); \
-		return (LWCGetTile(LhsN) != (FFloatType)0) ? (LWCGetTile(LhsN) Op (FFloatType)0) : (LhsN.Offset Op Rhs); }
+		return LWCGetTile(Lhs) * UE_LWC_RENDER_TILE_SIZE Op Rhs - Lhs.Offset; }
 
 LWC_COMPARE_OP(LWCGreater, >, FLWCType, FFloatType, FBoolType)
 LWC_COMPARE_OP(LWCGreaterEqual, >=, FLWCType, FFloatType, FBoolType)
@@ -82,35 +78,27 @@ FLWCType LWCSubtract(FLWCType Lhs, FFloatType Rhs) { return LWCConstructor(LWCGe
 
 FBoolType LWCEquals(FLWCType Lhs, FLWCType Rhs)
 {
-	FLWCType LhsN = LWCNormalizeTile(Lhs);
-	FLWCType RhsN = LWCNormalizeTile(Rhs);
-	return LWCGetTile(LhsN) == LWCGetTile(RhsN) && LhsN.Offset == RhsN.Offset;
+	return (LWCGetTile(Lhs) - LWCGetTile(Rhs)) * UE_LWC_RENDER_TILE_SIZE == Rhs.Offset - Lhs.Offset;
 }
 FBoolType LWCEquals(FFloatType Lhs, FLWCType Rhs)
 {
-	FLWCType RhsN = LWCNormalizeTile(Rhs); 
-	return LWCGetTile(RhsN) == 0 && Lhs == RhsN.Offset;
+	return -LWCGetTile(Rhs) * UE_LWC_RENDER_TILE_SIZE == Rhs.Offset - Lhs;
 } 
 FBoolType LWCEquals(FLWCType Lhs, FFloatType Rhs)
 {
-	FLWCType LhsN = LWCNormalizeTile(Lhs); 
-	return LWCGetTile(LhsN) == 0 && LhsN.Offset == Rhs;
+	return LWCGetTile(Lhs) * UE_LWC_RENDER_TILE_SIZE == Rhs - Lhs.Offset;
 }
 FBoolType LWCEqualsApprox(FLWCType Lhs, FLWCType Rhs, float Threshold)
 {
-	FLWCType LhsN = LWCNormalizeTile(Lhs);
-	FLWCType RhsN = LWCNormalizeTile(Rhs);
-	return LWCGetTile(LhsN) == LWCGetTile(RhsN) && abs(LhsN.Offset - RhsN.Offset) < (FFloatType)Threshold;
+	return abs((LWCGetTile(Lhs) - LWCGetTile(Rhs)) * UE_LWC_RENDER_TILE_SIZE + (Lhs.Offset - Rhs.Offset)) < (FFloatType)Threshold;
 }
 FBoolType LWCEqualsApprox(FFloatType Lhs, FLWCType Rhs, float Threshold)
 {
-	FLWCType RhsN = LWCNormalizeTile(Rhs);
-	return LWCGetTile(RhsN) == 0 && abs(Lhs - RhsN.Offset) < (FFloatType)Threshold;
+	return abs(-LWCGetTile(Rhs) * UE_LWC_RENDER_TILE_SIZE + (Lhs - Rhs.Offset)) < (FFloatType)Threshold;
 }
 FBoolType LWCEqualsApprox(FLWCType Lhs, FFloatType Rhs, float Threshold)
 {
-	FLWCType LhsN = LWCNormalizeTile(Lhs);
-	return LWCGetTile(LhsN) == 0 && abs(LhsN.Offset - Rhs) < (FFloatType)Threshold;
+	return abs(LWCGetTile(Lhs) * UE_LWC_RENDER_TILE_SIZE + (Lhs.Offset - Rhs)) < (FFloatType)Threshold;
 }
 
 FLWCType LWCSelect(FBoolType S, FLWCType Lhs, FLWCType Rhs) { return LWCConstructor(S ? LWCGetTile(Lhs) : LWCGetTile(Rhs), S ? Lhs.Offset : Rhs.Offset); }
