Change hair transmittance mask & hair shadow mask to be computed in translated world space to fix/improve precision with LWC.

diff --git a/Engine/Shaders/Private/HairStrands/HairStrandsDeepShadowAllocation.usf b/Engine/Shaders/Private/HairStrands/HairStrandsDeepShadowAllocation.usf
--- a/Engine/Shaders/Private/HairStrands/HairStrandsDeepShadowAllocation.usf
+++ b/Engine/Shaders/Private/HairStrands/HairStrandsDeepShadowAllocation.usf
@@ -13,15 +13,15 @@
 
 ////////////////////////////////////////////////////////////////////////////////
 
-float4x4 ComputeWorldToLight(
-	const FHairAABB Aabb,
+float4x4 ComputeTranslatedWorldToLight(
+	const FHairAABB TranslatedAABB,
 	const float3 LightDirection,
-	const float3 LightPosition,
+	const float3 TranslatedLightPosition,
 	const bool bIsDirectional)
 {
-	const float3 Extents = GetExtents(Aabb);
-	const float3 Center = GetCenter(Aabb);
-	const float  Radius = length(Extents);
+	const float3 Extents = GetExtents(TranslatedAABB);
+	const float3 Center  = GetCenter(TranslatedAABB);
+	const float  Radius  = length(Extents);
 
 	if (bIsDirectional) // (LightType == LightType_Directional)
 	{
@@ -29,20 +29,16 @@ float4x4 ComputeWorldToLight(
 	}
 	else // if (LightType == LightType_Spot || LightType == LightType_Point || LightType == LightType_Rect)
 	{
-		return LookAtMatrix(LightPosition, Center, float3(0, 0, 1));
+		return LookAtMatrix(TranslatedLightPosition, Center, float3(0, 0, 1));
 	}
 }
 
 ////////////////////////////////////////////////////////////////////////////////
 
 float4				LightDirections[MAX_SLOT_COUNT];
-float4				LightPositions[MAX_SLOT_COUNT];
+float4				TranslatedLightPositions[MAX_SLOT_COUNT];
 uint4				MacroGroupIndices[MAX_SLOT_COUNT];
 
-// For testing parity with CPU version
-float3				CPU_MinAABB;
-uint				CPU_bUseCPUData;
-float3				CPU_MaxAABB;
 float				RasterizationScale;
 float				AABBScale;
 float				MaxHafFovInRad;
@@ -53,7 +49,7 @@ uint				MacroGroupCount;
 
 Buffer<int>								MacroGroupAABBBuffer;
 RWStructuredBuffer<FDeepShadowViewInfo>	OutShadowViewInfoBuffer;
-RWStructuredBuffer<FDeepShadowTransform>OutShadowWorldToLightTransformBuffer;
+RWStructuredBuffer<FDeepShadowTransform>OutShadowTranslatedWorldToLightTransformBuffer;
 
 
 float DegreesToRadians(float InDeg)
@@ -70,23 +66,23 @@ float ComputeMinStrandRadiusAtDepth1(const int2 Resolution, const float InFOVInR
 }
 
 // This function is similar to the CPU version in HairStrandsDeepShadow.cpp
-void ComputeWorldToLightClip(
-	inout float4x4	OutWorldToClipTransform,
+void ComputeTranslatedWorldToLightClip(
+	inout float4x4	OutTranslatedWorldToClipTransform,
 	inout float		OutMinStrandRadiusAtDepth1,
 	inout float3	OutLightDirection,
-	const FHairAABB World_AABB,
+	const FHairAABB TranslatedWorld_AABB,
 	const uint		SlotIndex)
 {
-	const float3 LightDirection = LightDirections[SlotIndex].xyz;
-	const float3 LightPosition  = LightPositions[SlotIndex].xyz;
-	const bool bIsDirectional	= LightPositions[SlotIndex].w == 0;
+	const float3 LightDirection			  = LightDirections[SlotIndex].xyz;
+	const float3 TranslatedLightPosition  = TranslatedLightPositions[SlotIndex].xyz;
+	const bool bIsDirectional			  = TranslatedLightPositions[SlotIndex].w == 0;
 
-	float3 Center = GetCenter(World_AABB);
-	OutLightDirection = bIsDirectional ? LightDirection : -normalize(LightPosition - Center);
+	float3 Center = GetCenter(TranslatedWorld_AABB);
+	OutLightDirection = bIsDirectional ? LightDirection : -normalize(TranslatedLightPosition - Center);
 
-	const float4x4 Coarse_WorldToLight = ComputeWorldToLight(World_AABB, LightDirection, LightPosition, bIsDirectional);
+	const float4x4 Coarse_TranslatedWorldToLight = ComputeTranslatedWorldToLight(TranslatedWorld_AABB, LightDirection, TranslatedLightPosition, bIsDirectional);
 
-	const FHairAABB Light_AABB = Transform(World_AABB, Coarse_WorldToLight);
+	const FHairAABB Light_AABB = Transform(TranslatedWorld_AABB, Coarse_TranslatedWorldToLight);
 	const float3 Light_Extents = GetExtents(Light_AABB);
 
 	const float Radius = length(Light_Extents.xy) * AABBScale;
@@ -98,30 +94,30 @@ void ComputeWorldToLightClip(
 
 	const float StrandHairStableRasterizationScale = max(RasterizationScale, 1.0f);
 	OutMinStrandRadiusAtDepth1 = 1;
-	OutWorldToClipTransform = 0;
+	OutTranslatedWorldToClipTransform = 0;
 
 	if (bIsDirectional) // (LightType == LightType_Directional)
 	{
-		const float4x4 WorldToLight	= LookAtMatrix(Center - LightDirection * abs(MinZ), Center, float3(0, 0, 1));
-		const float4x4 ProjMatrix	= ReversedZOrthoMatrix(Radius, Radius, 1.f / (MaxZ-MinZ), 0);
-		OutWorldToClipTransform		= mul(WorldToLight, ProjMatrix);
+		const float4x4 TranslatedWorldToLight	= LookAtMatrix(Center - LightDirection * abs(MinZ), Center, float3(0, 0, 1));
+		const float4x4 ProjMatrix				= ReversedZOrthoMatrix(Radius, Radius, 1.f / (MaxZ-MinZ), 0);
+		OutTranslatedWorldToClipTransform		= mul(TranslatedWorldToLight, ProjMatrix);
 
-		const float RadiusAtDepth1	= Radius / min(SlotResolution.x, SlotResolution.y);
-		OutMinStrandRadiusAtDepth1	= RadiusAtDepth1 * RasterizationScale;
+		const float RadiusAtDepth1				= Radius / min(SlotResolution.x, SlotResolution.y);
+		OutMinStrandRadiusAtDepth1				= RadiusAtDepth1 * RasterizationScale;
 	}
 	else // if (LightType == LightType_Spot || LightType == LightType_Point || LightType == LightType_Rect)
 	{
-		const float LightDistanceToCenter = length(LightPosition - Center);
+		const float LightDistanceToCenter = length(TranslatedLightPosition - Center);
 		MaxZ = max(0.2f, LightDistanceToCenter) + MaxZ;
 		MinZ = max(0.1f, LightDistanceToCenter) + MinZ;
 		MinZ = max(1.0f, MinZ);
 
-		const float SphereDistance	= length(LightPosition - Center);
-		const float HalfFov			= min(MaxHafFovInRad, asin(Radius / SphereDistance));
+		const float SphereDistance			= length(TranslatedLightPosition - Center);
+		const float HalfFov					= min(MaxHafFovInRad, asin(Radius / SphereDistance));
 
-		const float4x4 WorldToLight = LookAtMatrix(LightPosition, Center, float3(0, 0, 1));
-		const float4x4 ProjMatrix	= ReversedZPerspectiveMatrix(HalfFov, 1, 1, MinZ, MaxZ);
-		OutWorldToClipTransform		= mul(WorldToLight, ProjMatrix);
+		const float4x4 TranslatedWorldToLight	= LookAtMatrix(TranslatedLightPosition, Center, float3(0, 0, 1));
+		const float4x4 ProjMatrix				= ReversedZPerspectiveMatrix(HalfFov, 1, 1, MinZ, MaxZ);
+		OutTranslatedWorldToClipTransform		= mul(TranslatedWorldToLight, ProjMatrix);
 
 		OutMinStrandRadiusAtDepth1 = ComputeMinStrandRadiusAtDepth1(SlotResolution, 2 * HalfFov, RasterizationScale);
 	}
@@ -138,19 +134,23 @@ void CreateViewInfo(uint2 DispatchThreadId : SV_DispatchThreadID)
 
 	const uint MacroGroupId = MacroGroupIndices[SlotIndex].x;
 
-	FHairAABB Bound = InitHairAABB();
+	FHairAABB TranslatedBound = InitHairAABB();
 	if (MacroGroupId < MacroGroupCount)
 	{
-		Bound = ReadHairAABB(MacroGroupId, MacroGroupAABBBuffer);
+		const float3 TranslatedWorldOffset = LWCHackToFloat(PrimaryView.PreViewTranslation);
+
+		TranslatedBound = ReadHairAABB(MacroGroupId, MacroGroupAABBBuffer);
+		TranslatedBound.Min += TranslatedWorldOffset; // HAIR_LWC
+		TranslatedBound.Max += TranslatedWorldOffset; // HAIR_LWC
 	}
 
 	FDeepShadowViewInfo ViewInfo;
-	ComputeWorldToLightClip(ViewInfo.WorldToClip, ViewInfo.MinRadiusAtDepth1, ViewInfo.ViewForward, Bound, SlotIndex);
+	ComputeTranslatedWorldToLightClip(ViewInfo.TranslatedWorldToClip, ViewInfo.MinRadiusAtDepth1, ViewInfo.ViewForward, TranslatedBound, SlotIndex);
 
 	FDeepShadowTransform Transform;
-	Transform.WorldToClip = ViewInfo.WorldToClip;
+	Transform.TranslatedWorldToClip = ViewInfo.TranslatedWorldToClip;
 
 	OutShadowViewInfoBuffer[SlotIndex] = ViewInfo;
-	OutShadowWorldToLightTransformBuffer[SlotIndex] = Transform;
+	OutShadowTranslatedWorldToLightTransformBuffer[SlotIndex] = Transform;
 }
 #endif
diff --git a/Engine/Shaders/Private/HairStrands/HairStrandsDeepShadowCommonStruct.ush b/Engine/Shaders/Private/HairStrands/HairStrandsDeepShadowCommonStruct.ush
--- a/Engine/Shaders/Private/HairStrands/HairStrandsDeepShadowCommonStruct.ush
+++ b/Engine/Shaders/Private/HairStrands/HairStrandsDeepShadowCommonStruct.ush
@@ -4,12 +4,12 @@
 
 struct FDeepShadowViewInfo
 {
-	float4x4 WorldToClip;
+	float4x4 TranslatedWorldToClip;
 	float3   ViewForward;
 	float    MinRadiusAtDepth1;
 };
 
 struct FDeepShadowTransform
 {
-	float4x4 WorldToClip; // Store the transform into a struct to avoid the matrix to be loaded as transposed
+	float4x4 TranslatedWorldToClip; // Store the transform into a struct to avoid the matrix to be loaded as transposed
 };
\ No newline at end of file
diff --git a/Engine/Shaders/Private/HairStrands/HairStrandsDeepShadowDebug.usf b/Engine/Shaders/Private/HairStrands/HairStrandsDeepShadowDebug.usf
--- a/Engine/Shaders/Private/HairStrands/HairStrandsDeepShadowDebug.usf
+++ b/Engine/Shaders/Private/HairStrands/HairStrandsDeepShadowDebug.usf
@@ -136,7 +136,7 @@ uint AllocatedSlotCount;
 uint MacroGroupCount;
 
 Buffer<int>	MacroGroupAABBBuffer;
-StructuredBuffer<FDeepShadowTransform> ShadowWorldToLightTransformBuffer;
+StructuredBuffer<FDeepShadowTransform> ShadowTranslatedWorldToLightTransformBuffer;
 
 [numthreads(1, 1, 1)]
 void MainCS(uint3 DispatchThreadId : SV_DispatchThreadID)
@@ -148,15 +148,17 @@ void MainCS(uint3 DispatchThreadId : SV_DispatchThreadID)
 
 	for (uint MacroGroupIt = 0; MacroGroupIt < MacroGroupCount; ++MacroGroupIt)
 	{
-		const FHairAABB Bound = ReadHairAABB(MacroGroupIt, MacroGroupAABBBuffer);
+		const FHairAABB Bound = ReadHairAABB(MacroGroupIt, MacroGroupAABBBuffer); // HAIR_LWC
 		AddAABB(Bound.Min, Bound.Max, float4(1, 0.5f, 0, 1));
 	}
 
+	const float3 TranslatedWorldOffset = LWCHackToFloat(PrimaryView.PreViewTranslation); // HAIR_LWC
+
 	const float4 FrustumColor = float4(0.5f, 1, 0, 1);
 	for (uint SlotIt = 0; SlotIt < AllocatedSlotCount; ++SlotIt)
 	{
-		const float4x4 WorldToClip = ShadowWorldToLightTransformBuffer[SlotIt].WorldToClip;
-		const float4x4 ClipToWorld = Inverse(WorldToClip);
+		const float4x4 TranslatedWorldToClip = ShadowTranslatedWorldToLightTransformBuffer[SlotIt].TranslatedWorldToClip;
+		const float4x4 ClipToTranslatedWorld = Inverse(TranslatedWorldToClip);
 
 		float4 P01 = float4(-1, -1, 0, 1);
 		float4 P11 = float4(-1,  1, 0, 1);
@@ -168,15 +170,15 @@ void MainCS(uint3 DispatchThreadId : SV_DispatchThreadID)
 		float4 P20 = float4( 1,  1, 1, 1);
 		float4 P30 = float4( 1, -1, 1, 1);
 
-		P00 = mul(P00, ClipToWorld); P00 /= P00.w; 
-		P10 = mul(P10, ClipToWorld); P10 /= P10.w; 
-		P20 = mul(P20, ClipToWorld); P20 /= P20.w; 
-		P30 = mul(P30, ClipToWorld); P30 /= P30.w; 
+		P00 = mul(P00, ClipToTranslatedWorld); P00 /= P00.w; P00.xyz += TranslatedWorldOffset;
+		P10 = mul(P10, ClipToTranslatedWorld); P10 /= P10.w; P10.xyz += TranslatedWorldOffset;
+		P20 = mul(P20, ClipToTranslatedWorld); P20 /= P20.w; P20.xyz += TranslatedWorldOffset;
+		P30 = mul(P30, ClipToTranslatedWorld); P30 /= P30.w; P30.xyz += TranslatedWorldOffset;
 
-		P01 = mul(P01, ClipToWorld); P01 /= P01.w;
-		P11 = mul(P11, ClipToWorld); P11 /= P11.w;
-		P21 = mul(P21, ClipToWorld); P21 /= P21.w;
-		P31 = mul(P31, ClipToWorld); P31 /= P31.w;
+		P01 = mul(P01, ClipToTranslatedWorld); P01 /= P01.w; P01.xyz += TranslatedWorldOffset;
+		P11 = mul(P11, ClipToTranslatedWorld); P11 /= P11.w; P11.xyz += TranslatedWorldOffset;
+		P21 = mul(P21, ClipToTranslatedWorld); P21 /= P21.w; P21.xyz += TranslatedWorldOffset;
+		P31 = mul(P31, ClipToTranslatedWorld); P31 /= P31.w; P31.xyz += TranslatedWorldOffset;
 
 		AddLine(P00, P10, FrustumColor, FrustumColor);
 		AddLine(P10, P20, FrustumColor, FrustumColor);
diff --git a/Engine/Shaders/Private/HairStrands/HairStrandsDeepShadowMask.usf b/Engine/Shaders/Private/HairStrands/HairStrandsDeepShadowMask.usf
--- a/Engine/Shaders/Private/HairStrands/HairStrandsDeepShadowMask.usf
+++ b/Engine/Shaders/Private/HairStrands/HairStrandsDeepShadowMask.usf
@@ -26,21 +26,19 @@ float4 WriteShadowMaskOutput(const float InFadedShadow)
 #define KERNEL_TYPE_Gaussian16 3
 #define KERNEL_TYPE_GaussianTransmission8 4
 
-float3 TotWorldPosition(float2 UV, float Depth)
+float3 GetTranslatedWorldPosition(float2 UV, float SceneDepth)
 {
-	float2 ScreenPosition = (UV - View.ScreenPositionScaleBias.wz) / View.ScreenPositionScaleBias.xy;
-	float4 HomogeneousWorldPosition = mul(float4(ScreenPosition * Depth, Depth, 1), LWCHackToFloat(PrimaryView.ScreenToWorld));
-	float3 WorldPosition = HomogeneousWorldPosition.xyz / HomogeneousWorldPosition.w;
-
-	return WorldPosition;
+	const float2 ScreenPosition = (UV - View.ScreenPositionScaleBias.wz) / View.ScreenPositionScaleBias.xy;
+	return mul(float4(ScreenPosition * SceneDepth, SceneDepth, 1), View.ScreenToTranslatedWorld).xyz;
 }
 
+
 #if SHADER_SHADOWMASK_DEEPSHADOW
 
-StructuredBuffer<FDeepShadowTransform> DeepShadow_WorldToLightTransformBuffer;
+StructuredBuffer<FDeepShadowTransform> DeepShadow_TranslatedWorldToLightTransformBuffer;
 uint				DeepShadow_bIsGPUDriven;
 
-float4x4			DeepShadow_CPUWorldToLightTransform;
+float4x4			DeepShadow_CPUTranslatedWorldToLightTransform;
 int2				DeepShadow_SlotResolution;
 int2				DeepShadow_SlotOffset;
 uint				DeepShadow_SlotIndex;
@@ -59,19 +57,19 @@ SamplerState 		ShadowSampler;
 // XY: UV of the shadow space
 //  Z: Normalized depth value in light clip space
 //  W: Positive if the position is valid, negative otherwise
-float4 ToLightPosition(float3 WorldPosition)
+float4 ToLightPosition(float3 InTranslatedWorldPosition)
 {
-	float4x4 WorldToLightTransform = 0;
+	float4x4 TranslatedWorldToLightTransform = 0;
 	if (DeepShadow_bIsGPUDriven > 0)
 	{
-		WorldToLightTransform = DeepShadow_WorldToLightTransformBuffer[DeepShadow_SlotIndex].WorldToClip;
+		TranslatedWorldToLightTransform = DeepShadow_TranslatedWorldToLightTransformBuffer[DeepShadow_SlotIndex].TranslatedWorldToClip;
 	}
 	else
 	{
-		WorldToLightTransform = DeepShadow_CPUWorldToLightTransform;
+		TranslatedWorldToLightTransform = DeepShadow_CPUTranslatedWorldToLightTransform;
 	}
 
-	float4 LightPos = mul(float4(WorldPosition,1), WorldToLightTransform);
+	float4 LightPos = mul(float4(InTranslatedWorldPosition,1), TranslatedWorldToLightTransform);
 	LightPos.xyz/= LightPos.w;
 	const float2 LightUV = (LightPos.xy + float(1).xx) * 0.5f;
 	return float4(LightUV.x, 1-LightUV.y, LightPos.z, sign(LightPos.w));
@@ -339,8 +337,8 @@ void MainPS(
 	const float SceneDepth = ConvertFromDeviceZ(SceneDepthTexture.Load(uint3(PixelCoord, 0)).x);
 	const bool bIsHairPixel = HairStrands.HairCoverageTexture.Load(uint3(PixelCoord, 0)) >= 1;
 
-	const float3 WorldPosition		= TotWorldPosition(UV, SceneDepth);
-	const float4 LightSpacePosition = ToLightPosition(WorldPosition);
+	const float3 TranslatedWorldPosition = GetTranslatedWorldPosition(UV, SceneDepth);
+	const float4 LightSpacePosition		 = ToLightPosition(TranslatedWorldPosition);
 
 	// PCF or simple point sampler
 	float Visibility = 1;
@@ -395,8 +393,7 @@ void MainPS(
 #include "HairStrandsVoxelPageTraversal.ush"
 #endif
 
-float4				Voxel_LightPosition;
-float3				Voxel_LightDirection;
+float4				Voxel_TranslatedLightPosition_LightDirection;
 uint				Voxel_MacroGroupId;
 uint				Voxel_RandomType;
 Texture2D<float4>	RayMarchMaskTexture;
@@ -444,16 +441,18 @@ void MainPS(
 	const float HairPixelCoverage = HairStrands.HairCoverageTexture.Load(uint3(PixelCoord, 0));
 	const bool bIsHairPixel = HairPixelCoverage >= 1;
 
-	float3 WorldPosition = TotWorldPosition(UV, SceneDepth);
+	float3 TranslatedWorldPosition = GetTranslatedWorldPosition(UV, SceneDepth);
 
 	// PCF or simple point sampler
 	float Visibility = 1;
 	if (!bIsHairPixel)
 	{
-		const float VoxelShadowMaxDistance = 100000.0; // 1 km shadow distance
-		const bool bIsDirectionalLight = Voxel_LightPosition.w == 0;
-		const float3 LightPosition  = bIsDirectionalLight ? -Voxel_LightDirection*VoxelShadowMaxDistance + WorldPosition : Voxel_LightPosition.xyz;
-		const float3 LightDirection = bIsDirectionalLight ? -Voxel_LightDirection : normalize(Voxel_LightPosition.xyz - WorldPosition);
+		const float3 TranslatedWorldOffset = LWCHackToFloat(PrimaryView.PreViewTranslation); // HAIR_LWC
+
+		const float VoxelShadowMaxDistance  = 100000.0; // 1 km shadow distance
+		const bool bIsDirectional			= IsLightDirectional(Voxel_TranslatedLightPosition_LightDirection);
+		const float3 TranslatedLightPosition= GetTranslatedLightPosition(Voxel_TranslatedLightPosition_LightDirection, TranslatedWorldPosition);
+		const float3 LightDirection			= GetTranslatedLightDirection(Voxel_TranslatedLightPosition_LightDirection, TranslatedWorldPosition);
 
 		// Depth bias
 		// Origin is shifted 2 voxels away towards the light + a constant bias of the size of the voxel
@@ -461,7 +460,7 @@ void MainPS(
 		{
 			const float PositionBiasScale = 0.5f;
 			const float3 DepthBias = VirtualVoxel.VoxelWorldSize * (VirtualVoxel.DepthBiasScale_Shadow * LightDirection + PositionBiasScale * (Random * 2 - 1));
-			WorldPosition += DepthBias;
+			TranslatedWorldPosition += DepthBias;
 		}
 
 		FHairTraversalResult Result = InitHairTraversalResult();
@@ -476,7 +475,7 @@ void MainPS(
 		CommonDesc.PageResolution = VirtualVoxel.PageResolution;
 
 		const FPackedVirtualVoxelNodeDesc PackedNode = VirtualVoxel.NodeDescBuffer[Voxel_MacroGroupId];
-		const FVirtualVoxelNodeDesc NodeDesc = UnpackVoxelNode(PackedNode, VirtualVoxel.PageResolution);
+		const FVirtualVoxelNodeDesc NodeDesc = UnpackVoxelNode(PackedNode, VirtualVoxel.PageResolution, TranslatedWorldOffset); // HAIR_LWC
 		const bool bIsValid = all(NodeDesc.PageIndexResolution != 0);
 
 		#if VOXEL_TRAVERSAL_DEBUG
@@ -498,8 +497,8 @@ void MainPS(
 		if (bIsValid && RayMarchMaskTexture.Load(uint3(PixelCoord, 0)).r > 0.0f && HairPixelCoverage < CoverageThreshold)
 		{
 			Result = ComputeHairCountVirtualVoxel(
-				WorldPosition,
-				LightPosition,
+				TranslatedWorldPosition,
+				TranslatedLightPosition,
 				CommonDesc,
 				NodeDesc,
 				VirtualVoxel.PageIndexBuffer,
diff --git a/Engine/Shaders/Private/HairStrands/HairStrandsDeepShadowVS.usf b/Engine/Shaders/Private/HairStrands/HairStrandsDeepShadowVS.usf
--- a/Engine/Shaders/Private/HairStrands/HairStrandsDeepShadowVS.usf
+++ b/Engine/Shaders/Private/HairStrands/HairStrandsDeepShadowVS.usf
@@ -36,7 +36,7 @@ void Main(
 	FVertexFactoryIntermediates VFIntermediates = GetVertexFactoryIntermediates(Input);
 
 #if MESH_RENDER_MODE == MESH_RENDER_DEPTH || MESH_RENDER_MODE == MESH_RENDER_DOM
-	float4x4 WorldToClipMatrix = DeepRasterPass.CPU_WorldToClipMatrix;
+	float4x4 TranslatedWorldToClipMatrix = DeepRasterPass.CPU_TranslatedWorldToClipMatrix;
 #endif
 
 #ifdef HAIR_STRAND_MESH_FACTORY	
@@ -53,19 +53,18 @@ void Main(
 	{
 		const uint AtlasSlotIndex = DeepRasterPass.AtlasSlotIndex;
 		FDeepShadowViewInfo DeepShadowViewInfo = DeepRasterPass.DeepShadowViewInfoBuffer[AtlasSlotIndex];
-		WorldToClipMatrix			= DeepShadowViewInfo.WorldToClip;
+		TranslatedWorldToClipMatrix = DeepShadowViewInfo.TranslatedWorldToClip;
 		HairViewInfo.ViewForward	= DeepShadowViewInfo.ViewForward;
 		HairViewInfo.RadiusAtDepth1 = DeepShadowViewInfo.MinRadiusAtDepth1;
 	}
 	#endif
 
-	float4 WorldPosition = VertexFactoryGetWorldPosition(Input, VFIntermediates, HairViewInfo);
+	float4 TranslatedWorldPosition = VertexFactoryGetWorldPosition(Input, VFIntermediates, HairViewInfo);
 #else
-	float4 WorldPosition = VertexFactoryGetWorldPosition(Input, VFIntermediates);
+	float4 TranslatedWorldPosition = VertexFactoryGetWorldPosition(Input, VFIntermediates);
 #endif
-	WorldPosition.xyz -= LWCHackToFloat(ResolvedView.PreViewTranslation); // Remove the camera relative position (todo move this into matrix transform)
 
-	Output.Position = mul(WorldPosition, WorldToClipMatrix);
+	Output.Position = mul(TranslatedWorldPosition, TranslatedWorldToClipMatrix);
 
 	float HairCoverage = 1;
   #ifdef HAIR_STRAND_MESH_FACTORY	
diff --git a/Engine/Shaders/Private/HairStrands/HairStrandsDeepTransmittanceCommon.ush b/Engine/Shaders/Private/HairStrands/HairStrandsDeepTransmittanceCommon.ush
--- a/Engine/Shaders/Private/HairStrands/HairStrandsDeepTransmittanceCommon.ush
+++ b/Engine/Shaders/Private/HairStrands/HairStrandsDeepTransmittanceCommon.ush
@@ -44,21 +44,23 @@ FHairTransmittanceMask UnpackTransmittanceMask(uint In)
 	return Out;
 }
 
-////////////////////////////////////////////////////////////////////////////////////////////////////
+bool IsLightDirectional(float4 InTranslatedLightPosition_LightDirection)
+{
+	return InTranslatedLightPosition_LightDirection.w == 0;
+}
 
-struct FTransmittanceSettings
+// Return the direction from the world position to the light
+float3 GetTranslatedLightDirection(float4 InTranslatedLightPosition_LightDirection, float3 InTranslatedWorldPosition)
 {
-	int2	 DeepShadowResolution;
-	int2	 DeepShadowOffset;
-	float4x4 WorldToLightTransform;
-	float3   LightDirection;
-	float4   LightPosition;
+	const bool bIsDirectional = InTranslatedLightPosition_LightDirection.w == 0;
+	return bIsDirectional ? -InTranslatedLightPosition_LightDirection.xyz : normalize(InTranslatedLightPosition_LightDirection.xyz - InTranslatedWorldPosition);
+}
 
-	float    DeepShadowDepthBiasScale;
-	float    DeepShadowDensityScale;
-	float	 TransmittanceKernelApertureInDegree;
-	float3   Random;
-	uint	 DebugMode;
+float3 GetTranslatedLightPosition(float4 InTranslatedLightPosition_LightDirection, float3 InTranslatedWorldPosition)
+{
+	const float VoxelShadowMaxDistance = 100000.0; // 1 km shadow distance	
+	const float3 TranslatedLightDirection = GetTranslatedLightDirection(InTranslatedLightPosition_LightDirection, InTranslatedWorldPosition);
 
-	uint TransmittanceKernelType;
-};
\ No newline at end of file
+	const bool bIsDirectional = InTranslatedLightPosition_LightDirection.w == 0;
+	return bIsDirectional ? InTranslatedWorldPosition + TranslatedLightDirection * VoxelShadowMaxDistance : InTranslatedLightPosition_LightDirection.xyz;
+}
diff --git a/Engine/Shaders/Private/HairStrands/HairStrandsDeepTransmittanceMask.usf b/Engine/Shaders/Private/HairStrands/HairStrandsDeepTransmittanceMask.usf
--- a/Engine/Shaders/Private/HairStrands/HairStrandsDeepTransmittanceMask.usf
+++ b/Engine/Shaders/Private/HairStrands/HairStrandsDeepTransmittanceMask.usf
@@ -20,6 +20,42 @@
 #error MAX_HAIR_MACROGROUP_COUNT needs to be defined
 #endif
 
+////////////////////////////////////////////////////////////////////////////////////////////////////
+// Utils function & structs
+struct FTransmittanceSettings
+{
+	int2	 DeepShadowResolution;
+	int2	 DeepShadowOffset;
+	float4x4 TranslatedWorldToLightTransform;
+	float3   LightDirection; // Direction from world position to the light
+	float3   TranslatedLightPosition;
+	bool	 bIsDirectional;
+
+	float    DeepShadowDepthBiasScale;
+	float    DeepShadowDensityScale;
+	float	 TransmittanceKernelApertureInDegree;
+	float3   Random;
+	uint	 DebugMode;
+
+	uint TransmittanceKernelType;
+};
+
+float3 GetViewVector(float3 InTranslatedWorldPosition)
+{
+	float3 CameraToPixel = normalize(InTranslatedWorldPosition - View.TranslatedWorldCameraOrigin);
+	float3 V = -CameraToPixel;
+	return V;
+}
+
+float3 GetTranslatedWorldPosition(float2 UV, float SceneDepth)
+{
+	const float2 ScreenPosition = (UV - View.ScreenPositionScaleBias.wz) / View.ScreenPositionScaleBias.xy;
+	return mul(float4(ScreenPosition * SceneDepth, SceneDepth, 1), View.ScreenToTranslatedWorld).xyz;
+}
+
+////////////////////////////////////////////////////////////////////////////////////////////////
+// Defines
+
 // Need to match FDeepShadowData::MaxMacroGroupCount in HairRendering.h
 #if SHADER_TRANSMITTANCE_DEEPSHADOW || SHADER_TRANSMITTANCE_VOXEL 
 	#if PERMUTATION_GROUP_SIZE == 64
@@ -31,29 +67,20 @@
 	#endif
 #endif
 
-float3 ToWorldPosition(float2 UV, float Depth)
-{
-	float2 ScreenPosition = (UV - View.ScreenPositionScaleBias.wz) / View.ScreenPositionScaleBias.xy;
-	float4 HomogeneousWorldPosition = mul(float4(ScreenPosition * Depth, Depth, 1), LWCHackToFloat(PrimaryView.ScreenToWorld));
-	float3 WorldPosition = HomogeneousWorldPosition.xyz / HomogeneousWorldPosition.w;
-
-	return WorldPosition;
-}
-
 ////////////////////////////////////////////////////////////////////////////////////////////////
 // Deep shadow
 
 #if SHADER_TRANSMITTANCE_DEEPSHADOW
 
-StructuredBuffer<FDeepShadowTransform> DeepShadow_WorldToLightTransformBuffer;
+StructuredBuffer<FDeepShadowTransform> DeepShadow_TranslatedWorldToLightTransformBuffer;
 uint				DeepShadow_bIsGPUDriven;
 
 uint2				DeepShadow_Resolution;
-float4x4			DeepShadow_CPUWorldToLightTransforms[MAX_HAIR_MACROGROUP_COUNT];
-float3				LightDirection;
+float4x4			DeepShadow_CPUTranslatedWorldToLightTransforms[MAX_HAIR_MACROGROUP_COUNT];
+//float3				LightDirection;
 uint				LightChannelMask;
 float4				ShadowChannelMask;
-float4				LightPosition;
+float4				TranslatedLightPosition_LightDirection;
 float				DeepShadow_DepthBiasScale;
 float				DeepShadow_DensityScale;
 float				LightRadius;
@@ -68,87 +95,8 @@ Texture2D<float4>	DeepShadow_DomTexture;
 Buffer<uint4>		IndirectArgsBuffer;
 RWBuffer<uint>		OutTransmittanceMask;
 
-#if 0
-// Sample deep shadow with PCF kernel and doing transmittance evaluation for each sample
-FHairTransmittanceMask SampleDOM_PCF_Accurate(
-	float3 LightSpacePosition,
-	float DepthBias,
-	const float4 InLayerDepths,
-	Texture2D<float> FrontDepthTexture,
-	Texture2D<float4> DomTexture,
-	FGBufferData GBuffer, float SinLightAngle, Texture3D<float4> InHairScatteringLUTTexture, SamplerState InHairLUTSampler)
-{
-	const float DeepShadowDensityScale = 1;
-	// Explicit 1 ring (5x5 = 3x3 of 2x2 PCF) with weighting
-	float HairCount = 0;
-	float w0 = 1.f;
-	float w1 = 0.5f;
-	float w2 = 0.25f;
-	float w = w0 + 4 * w1 + 4 * w2;
-	
-	float HairCount0 = SampleDOM_PCF2x2(LightSpacePosition + float3(-2, -2, 0), DepthBias, InLayerDepths, FrontDepthTexture, DomTexture);
-	float HairCount1 = SampleDOM_PCF2x2(LightSpacePosition + float3(0, -2, 0), DepthBias, InLayerDepths, FrontDepthTexture, DomTexture);
-	float HairCount2 = SampleDOM_PCF2x2(LightSpacePosition + float3(2, -2, 0), DepthBias, InLayerDepths, FrontDepthTexture, DomTexture);
-	float HairCount3 = SampleDOM_PCF2x2(LightSpacePosition + float3(-2, 0, 0), DepthBias, InLayerDepths, FrontDepthTexture, DomTexture);
-	float HairCount4 = SampleDOM_PCF2x2(LightSpacePosition + float3(0, 0, 0), DepthBias, InLayerDepths, FrontDepthTexture, DomTexture);
-	float HairCount5 = SampleDOM_PCF2x2(LightSpacePosition + float3(2, 0, 0), DepthBias, InLayerDepths, FrontDepthTexture, DomTexture);
-	float HairCount6 = SampleDOM_PCF2x2(LightSpacePosition + float3(-2, 2, 0), DepthBias, InLayerDepths, FrontDepthTexture, DomTexture);
-	float HairCount7 = SampleDOM_PCF2x2(LightSpacePosition + float3(0, 2, 0), DepthBias, InLayerDepths, FrontDepthTexture, DomTexture);
-	float HairCount8 = SampleDOM_PCF2x2(LightSpacePosition + float3(2, 2, 0), DepthBias, InLayerDepths, FrontDepthTexture, DomTexture);
-
-	FHairTransmittanceMask Out8 = ComputeTransmittanceFromHairCount(HairCount0 * DeepShadowDensityScale, 1, GBuffer, SinLightAngle, InHairScatteringLUTTexture, InHairLUTSampler);
-	FHairTransmittanceMask Out7 = ComputeTransmittanceFromHairCount(HairCount1 * DeepShadowDensityScale, 1, GBuffer, SinLightAngle, InHairScatteringLUTTexture, InHairLUTSampler);
-	FHairTransmittanceMask Out6 = ComputeTransmittanceFromHairCount(HairCount2 * DeepShadowDensityScale, 1, GBuffer, SinLightAngle, InHairScatteringLUTTexture, InHairLUTSampler);
-	FHairTransmittanceMask Out5 = ComputeTransmittanceFromHairCount(HairCount3 * DeepShadowDensityScale, 1, GBuffer, SinLightAngle, InHairScatteringLUTTexture, InHairLUTSampler);
-	FHairTransmittanceMask Out0 = ComputeTransmittanceFromHairCount(HairCount4 * DeepShadowDensityScale, 1, GBuffer, SinLightAngle, InHairScatteringLUTTexture, InHairLUTSampler);
-	FHairTransmittanceMask Out1 = ComputeTransmittanceFromHairCount(HairCount5 * DeepShadowDensityScale, 1, GBuffer, SinLightAngle, InHairScatteringLUTTexture, InHairLUTSampler);
-	FHairTransmittanceMask Out2 = ComputeTransmittanceFromHairCount(HairCount6 * DeepShadowDensityScale, 1, GBuffer, SinLightAngle, InHairScatteringLUTTexture, InHairLUTSampler);
-	FHairTransmittanceMask Out3 = ComputeTransmittanceFromHairCount(HairCount7 * DeepShadowDensityScale, 1, GBuffer, SinLightAngle, InHairScatteringLUTTexture, InHairLUTSampler);
-	FHairTransmittanceMask Out4 = ComputeTransmittanceFromHairCount(HairCount8 * DeepShadowDensityScale, 1, GBuffer, SinLightAngle, InHairScatteringLUTTexture, InHairLUTSampler);
-
-	FHairTransmittanceMask Out;
-	Out.HairCount =
-		w2 * Out0.HairCount +
-		w1 * Out1.HairCount +
-		w2 * Out2.HairCount +
-		w1 * Out3.HairCount +
-		w0 * Out4.HairCount +
-		w1 * Out5.HairCount +
-		w2 * Out6.HairCount +
-		w1 * Out7.HairCount +
-		w2 * Out8.HairCount;
-	Out.HairCount = Out.HairCount / w;
-
-	Out.Transmittance =
-		w2 * Out0.Transmittance +
-		w1 * Out1.Transmittance +
-		w2 * Out2.Transmittance +
-		w1 * Out3.Transmittance +
-		w0 * Out4.Transmittance +
-		w1 * Out5.Transmittance +
-		w2 * Out6.Transmittance +
-		w1 * Out7.Transmittance +
-		w2 * Out8.Transmittance;
-	Out.Transmittance = Out.Transmittance / w;
-
-	Out.Visibility =
-		w2 * Out0.Visibility +
-		w1 * Out1.Visibility +
-		w2 * Out2.Visibility +
-		w1 * Out3.Visibility +
-		w0 * Out4.Visibility +
-		w1 * Out5.Visibility +
-		w2 * Out6.Visibility +
-		w1 * Out7.Visibility +
-		w2 * Out8.Visibility;
-	Out.Visibility = Out.Visibility / w;
-
-	return Out;
-}
-#endif 
-
 FHairTransmittanceMask ComputeTransmittance(
-	float3 WorldPosition,
+	float3 TranslatedWorldPosition,
 	FTransmittanceSettings Settings,
 	float4				InLayerDepths,
 	Texture2D<float>	FrontDepthTexture,
@@ -156,11 +104,8 @@ FHairTransmittanceMask ComputeTransmittance(
 {
 	FHairTransmittanceMask Out = InitHairTransmittanceMask();
 	{
-		const bool bIsDirectionalLight = Settings.LightPosition.w == 0;
-
-		// LightSpacePosition is 'WorldPosition' in LightSpace (TexelCoord.XY / NDC Depth.Z)
-		const float3 LightSpacePosition = ToLightPosition(WorldPosition, Settings.WorldToLightTransform) * float3(Settings.DeepShadowResolution, 1) + float3(Settings.DeepShadowOffset, 0); 
-		const float3 L = bIsDirectionalLight ? Settings.LightDirection : normalize(Settings.LightPosition.xyz - WorldPosition);
+		// LightSpacePosition is 'TranslatedWorldPosition' in LightSpace (TexelCoord.XY / NDC Depth.Z)
+		const float3 LightSpacePosition = ToLightPosition(TranslatedWorldPosition, Settings.TranslatedWorldToLightTransform) * float3(Settings.DeepShadowResolution, 1) + float3(Settings.DeepShadowOffset, 0);
 
 		float DepthBias = 0;
 		if (Settings.DeepShadowDepthBiasScale > 0)
@@ -240,22 +185,26 @@ void MainCS(uint2 DispatchThreadId : SV_DispatchThreadID)
 
 		const uint MacroGroupIndex = clamp(Sample.MacroGroupId, 0, MAX_HAIR_MACROGROUP_COUNT - 1);
 		const int3 SlotOffset_SlotIndex = DeepShadow_AtlasSlotOffsets_AtlasSlotIndex[MacroGroupIndex].xyz;
-		float4x4 WorldToLightTransform = 0;
+		float4x4 TranslatedWorldToLightTransform = 0;
 		if (DeepShadow_bIsGPUDriven > 0)
 		{
-			WorldToLightTransform = DeepShadow_WorldToLightTransformBuffer[SlotOffset_SlotIndex.z].WorldToClip;
+			TranslatedWorldToLightTransform = DeepShadow_TranslatedWorldToLightTransformBuffer[SlotOffset_SlotIndex.z].TranslatedWorldToClip;
 		}	
 		else
 		{
-			WorldToLightTransform = DeepShadow_CPUWorldToLightTransforms[MacroGroupIndex];
+			TranslatedWorldToLightTransform = DeepShadow_CPUTranslatedWorldToLightTransforms[MacroGroupIndex];
 		}
 		
+		const float SceneDepth = ConvertFromDeviceZ(Sample.Depth);
+		const float3 TranslatedWorldPosition = GetTranslatedWorldPosition(UV, SceneDepth);
+
 		FTransmittanceSettings Settings;
 		Settings.DeepShadowOffset = SlotOffset_SlotIndex.xy;
 		Settings.DeepShadowResolution = DeepShadow_Resolution;
-		Settings.WorldToLightTransform = WorldToLightTransform;
-		Settings.LightDirection = LightDirection;
-		Settings.LightPosition = LightPosition;
+		Settings.TranslatedWorldToLightTransform = TranslatedWorldToLightTransform;
+		Settings.bIsDirectional = IsLightDirectional(TranslatedLightPosition_LightDirection);
+		Settings.LightDirection = GetTranslatedLightDirection(TranslatedLightPosition_LightDirection, TranslatedWorldPosition);
+		Settings.TranslatedLightPosition = GetTranslatedLightPosition(TranslatedLightPosition_LightDirection, TranslatedWorldPosition);
 		Settings.DeepShadowDepthBiasScale = DeepShadow_DepthBiasScale;
 		Settings.DeepShadowDensityScale = DeepShadow_DensityScale;
 		Settings.TransmittanceKernelApertureInDegree = DeepShadow_KernelAperture;
@@ -263,9 +212,7 @@ void MainCS(uint2 DispatchThreadId : SV_DispatchThreadID)
 		Settings.Random = TransmittanceRandom;
 		Settings.DebugMode = DeepShadow_DebugMode;
 
-		const float SceneDepth = ConvertFromDeviceZ(Sample.Depth);
-		const float3 WorldPosition = ToWorldPosition(UV, SceneDepth);
-		const FHairTransmittanceMask TransmittanceMask = ComputeTransmittance(WorldPosition, Settings, DeepShadow_LayerDepths, DeepShadow_FrontDepthTexture, DeepShadow_DomTexture);
+		const FHairTransmittanceMask TransmittanceMask = ComputeTransmittance(TranslatedWorldPosition, Settings, DeepShadow_LayerDepths, DeepShadow_FrontDepthTexture, DeepShadow_DomTexture);
 
 		OutTransmittanceMask[SampleIndex] = PackTransmittanceMask(TransmittanceMask);
 	}
@@ -297,34 +244,13 @@ void PrintInfo(uint Index, uint2 PixelCoord, FHairTransmittanceMask Transmittanc
 	Origin.y += (Index % 8) * 2 * ShaderPrintUniform.FontSize.w;
 
 	float2 Pos = Origin;
-
-	Pos = ShaderPrintSymbol(Pos, _H_);
-	Pos = ShaderPrintSymbol(Pos, _A_);
-	Pos = ShaderPrintSymbol(Pos, _I_);
-	Pos = ShaderPrintSymbol(Pos, _R_);
-	Pos = ShaderPrintSymbol(Pos, _SPC_);
-	Pos = ShaderPrintSymbol(Pos, _C_);
-	Pos = ShaderPrintSymbol(Pos, _O_);
-	Pos = ShaderPrintSymbol(Pos, _U_);
-	Pos = ShaderPrintSymbol(Pos, _N_);
-	Pos = ShaderPrintSymbol(Pos, _T_);
-	Pos = ShaderPrintSymbol(Pos, _SPC_);
+	Pos = ShaderPrintText(Pos, TEXT("Hair Count "));
 	Pos = ShaderPrint(Pos, TransmittanceMask.HairCount);
 
 	Pos = ShaderPrintNewline(Pos);
 	Pos.x = Origin.x;
 
-	Pos = ShaderPrintSymbol(Pos, _V_);
-	Pos = ShaderPrintSymbol(Pos, _I_);
-	Pos = ShaderPrintSymbol(Pos, _S_);
-	Pos = ShaderPrintSymbol(Pos, _I_);
-	Pos = ShaderPrintSymbol(Pos, _B_);
-	Pos = ShaderPrintSymbol(Pos, _I_);
-	Pos = ShaderPrintSymbol(Pos, _L_);
-	Pos = ShaderPrintSymbol(Pos, _I_);
-	Pos = ShaderPrintSymbol(Pos, _T_);
-	Pos = ShaderPrintSymbol(Pos, _Y_);
-	Pos = ShaderPrintSymbol(Pos, _SPC_);
+	Pos = ShaderPrintText(Pos, TEXT("Visibility "));
 	Pos = ShaderPrint(Pos, TransmittanceMask.Visibility);
 
 	Pos = ShaderPrintNewline(Pos);
@@ -332,9 +258,8 @@ void PrintInfo(uint Index, uint2 PixelCoord, FHairTransmittanceMask Transmittanc
 }
 #endif // VOXEL_TRAVERSAL_DEBUG
 
-float3				LightDirection;
 uint				LightChannelMask;
-float4				LightPosition;
+float4				TranslatedLightPosition_LightDirection;
 float				LightRadius;
 #if PERMUTATION_ONE_PASS
 Texture2D<uint> 	ShadowMaskBitsTexture;
@@ -348,17 +273,15 @@ RWBuffer<uint>		OutTransmittanceMask;
 FHairTransmittanceMask ComputeTransmittanceVoxel(
 	const uint2 PixelCoord,
 	const float PixelRadius,
-	float3 WorldPosition,
+	float3 TranslatedWorldPosition,
 	uint MacroGroupId,
 	FTransmittanceSettings Settings,
 	bool bDebugEnabled)
 {
+	// Cone opening - Area light transmission is disabled as it is too noisy for now
+	const float TanLightAngle = 0; // Settings.bIsDirectional ? 0 : LightRadius / length(Settings.TranslatedLightPosition - TranslatedWorldPosition);
 
-	const float VoxelShadowMaxDistance = 100000.0; // 1 km shadow distance	
-	const bool bIsDirectionalLight = Settings.LightPosition.w == 0;	
-	const float3 LightPosition  = bIsDirectionalLight ? -Settings.LightDirection*VoxelShadowMaxDistance + WorldPosition : Settings.LightPosition.xyz;
-	const float3 LightDirection = bIsDirectionalLight ? -Settings.LightDirection : normalize(Settings.LightPosition.xyz - WorldPosition);
-	const float TanLightAngle = bIsDirectionalLight ? 0 : LightRadius / length(Settings.LightPosition.xyz - WorldPosition);
+	const float3 TranslatedWorldOffset = LWCHackToFloat(PrimaryView.PreViewTranslation); // HAIR_LWC
 
 	// Compute the number of hair count between light & shading point
 	const float DistanceThreshold = 100000;
@@ -378,8 +301,8 @@ FHairTransmittanceMask ComputeTransmittanceVoxel(
 		#endif
 		
 		const float PositionBiasScale = 0.5f;
-		const float3 DepthBias = VirtualVoxel.VoxelWorldSize * (VirtualVoxel.DepthBiasScale_Transmittance*LightDirection + PositionBiasScale*(SampleRandom*2-1));
-		const float3 SampleWorldPosition = WorldPosition + DepthBias;
+		const float3 DepthBias = VirtualVoxel.VoxelWorldSize * (VirtualVoxel.DepthBiasScale_Transmittance * Settings.LightDirection + PositionBiasScale*(SampleRandom*2-1));
+		const float3 SampleTranslatedWorldPosition = TranslatedWorldPosition + DepthBias;
 
 		FVirtualVoxelCommonDesc CommonDesc;
 		CommonDesc.PageCountResolution	= VirtualVoxel.PageCountResolution;
@@ -388,7 +311,7 @@ FHairTransmittanceMask ComputeTransmittanceVoxel(
 		CommonDesc.PageResolution		= VirtualVoxel.PageResolution;
 
 		const FPackedVirtualVoxelNodeDesc PackedNode = VirtualVoxel.NodeDescBuffer[MacroGroupId];
-		const FVirtualVoxelNodeDesc NodeDesc = UnpackVoxelNode(PackedNode, VirtualVoxel.PageResolution);
+		const FVirtualVoxelNodeDesc NodeDesc = UnpackVoxelNode(PackedNode, VirtualVoxel.PageResolution, TranslatedWorldOffset); // HAIR_LWC
 
 		FHairTraversalSettings TraversalSettings = InitHairTraversalSettings();
 		TraversalSettings.DensityScale = VirtualVoxel.DensityScale_Transmittance;
@@ -397,12 +320,12 @@ FHairTransmittanceMask ComputeTransmittanceVoxel(
 		TraversalSettings.bDebugEnabled = bDebugEnabled;
 		TraversalSettings.SteppingScale = VirtualVoxel.SteppingScale_Transmittance;
 		TraversalSettings.Random = SampleRandom;
-		TraversalSettings.TanConeAngle = 0; // TanLightAngle; // Area light transmission is disabled as it is too noisy for now
+		TraversalSettings.TanConeAngle = TanLightAngle;
 		TraversalSettings.PixelRadius = PixelRadius;
 
 		const FHairTraversalResult Result = ComputeHairCountVirtualVoxel(
-			SampleWorldPosition, 
-			LightPosition,
+			SampleTranslatedWorldPosition,
+			Settings.TranslatedLightPosition,
 			CommonDesc,
 			NodeDesc,
 			VirtualVoxel.PageIndexBuffer,
@@ -437,13 +360,13 @@ void MainCS(uint2 DispatchThreadId : SV_DispatchThreadID)
 	{
 		return;
 	}
-
+	const float3 TranslatedWorldOffset = LWCHackToFloat(PrimaryView.PreViewTranslation); // HAIR_LWC
 	const FHairSample Sample	= UnpackHairSample(HairStrands.HairSampleData[SampleIndex]);
 	const uint MacroGroupIndex	= clamp(Sample.MacroGroupId, 0, MAX_HAIR_MACROGROUP_COUNT - 1);
 	const float SceneDepth		= ConvertFromDeviceZ(Sample.Depth);
 	const uint2 PixelCoord		= HairStrands.HairSampleCoords[SampleIndex];
 	const float2 UV				= (PixelCoord.xy + float2(0.5f, 0.5f)) / float2(View.BufferSizeAndInvSize.xy);
-	const float3 WorldPosition	= ToWorldPosition(UV, SceneDepth);
+	const float3 TranslatedWorldPosition = GetTranslatedWorldPosition(UV, SceneDepth);
 	const float PixelRadius		= SceneDepth * VirtualVoxel.HairCoveragePixelRadiusAtDepth1;
 
 	const uint4 PackedShadowMask = ~ShadowMaskBitsTexture[PixelCoord];
@@ -481,14 +404,13 @@ void MainCS(uint2 DispatchThreadId : SV_DispatchThreadID)
 		#endif	
 
 		// The Virtual Shadow Remap stores directional lights first
-		const FLocalLightData LightData = GetLocalLightData(CulledLightGridData.DataStartIndex + LightIndex, EyeIndex);
-		const float3 LightPosition		= LightData.LightPositionAndInvRadius.xyz - LWCHackToFloat(PrimaryView.PreViewTranslation);
-		const float3 LightDirection		= LightData.LightDirectionAndShadowMask.xyz; // Direction light only
-		const bool bIsLocalLight		= true;
+		const FLocalLightData LightData 	 = GetLocalLightData(CulledLightGridData.DataStartIndex + LightIndex, EyeIndex);
+		const float3 TranslatedLightPosition = LightData.LightPositionAndInvRadius.xyz; // LightPositionAndInvRadius is already in translated world space
 
 		FTransmittanceSettings Settings;
-		Settings.LightDirection				= normalize(LightPosition - WorldPosition);
-		Settings.LightPosition				= float4(LightPosition, bIsLocalLight ? 1 : 0);
+		Settings.bIsDirectional				= false;
+		Settings.LightDirection				= normalize(TranslatedLightPosition - TranslatedWorldPosition);
+		Settings.TranslatedLightPosition	= TranslatedLightPosition;
 		Settings.DeepShadowDensityScale		= 0;	// Use virtual voxel struct value
 		Settings.DeepShadowDepthBiasScale	= 0;	// Use virtual voxel struct value
 		Settings.Random						= Random;
@@ -501,7 +423,7 @@ void MainCS(uint2 DispatchThreadId : SV_DispatchThreadID)
 		const FHairTransmittanceMask Out = ComputeTransmittanceVoxel(
 			PixelCoord, 
 			PixelRadius, 
-			WorldPosition, 
+			TranslatedWorldPosition,
 			MacroGroupIndex, 
 			Settings, 
 			bDebugEnabled);
@@ -541,20 +463,21 @@ void MainCS(uint2 DispatchThreadId : SV_DispatchThreadID)
 			return;
 		}
 
+		const float SceneDepth = ConvertFromDeviceZ(Sample.Depth);
+		const float3 TranslatedWorldPosition = GetTranslatedWorldPosition(UV, SceneDepth);
+
 		const uint MacroGroupIndex = clamp(Sample.MacroGroupId, 0, MAX_HAIR_MACROGROUP_COUNT - 1);
 
 		FTransmittanceSettings Settings;
-		Settings.LightDirection = LightDirection;
-		Settings.LightPosition = LightPosition;
+		Settings.bIsDirectional = IsLightDirectional(TranslatedLightPosition_LightDirection);
+		Settings.LightDirection = GetTranslatedLightDirection(TranslatedLightPosition_LightDirection, TranslatedWorldPosition);
+		Settings.TranslatedLightPosition = GetTranslatedLightPosition(TranslatedLightPosition_LightDirection, TranslatedWorldPosition);
 		Settings.DeepShadowDensityScale = 0;	// Use virtual voxel struct value
 		Settings.DeepShadowDepthBiasScale = 0;	// Use virtual voxel struct value
 		Settings.Random = GetHairVoxelJitter(PixelCoord.xy, View.StateFrameIndexMod8, VirtualVoxel.JitterMode);
 
-		const float SceneDepth = ConvertFromDeviceZ(Sample.Depth);
-		const float3 WorldPosition = ToWorldPosition(UV, SceneDepth);
-
 		const float PixelRadius = SceneDepth * VirtualVoxel.HairCoveragePixelRadiusAtDepth1;	
-		const FHairTransmittanceMask Out = ComputeTransmittanceVoxel(PixelCoord, PixelRadius, WorldPosition, MacroGroupIndex, Settings, bDebugEnabled);
+		const FHairTransmittanceMask Out = ComputeTransmittanceVoxel(PixelCoord, PixelRadius, TranslatedWorldPosition, MacroGroupIndex, Settings, bDebugEnabled);
 
 		#if VOXEL_TRAVERSAL_DEBUG
 		if (bDebugEnabled)
diff --git a/Engine/Shaders/Private/HairStrands/HairStrandsVoxelDepthInjection.usf b/Engine/Shaders/Private/HairStrands/HairStrandsVoxelDepthInjection.usf
--- a/Engine/Shaders/Private/HairStrands/HairStrandsVoxelDepthInjection.usf
+++ b/Engine/Shaders/Private/HairStrands/HairStrandsVoxelDepthInjection.usf
@@ -12,7 +12,7 @@
 #include "HairStrandsVoxelPageCommon.ush"
 
 #if SHADER_DEPTH_INJECTION
-float4x4 CPU_WorldToClip;
+float4x4 CPU_TranslatedWorldToClip;
 
 float2 OutputResolution;
 uint MacroGroupId;
@@ -20,7 +20,7 @@ uint AtlasSlotIndex;
 
 float3 LightDirection;
 uint bIsDirectional;
-float3 LightPosition;
+float3 TranslatedLightPosition;
 uint bIsGPUDriven;
 
 StructuredBuffer<FDeepShadowViewInfo> DeepShadowViewInfoBuffer;
@@ -28,14 +28,16 @@ StructuredBuffer<FDeepShadowViewInfo> DeepShadowViewInfoBuffer;
 void MainVS(
 	uint VertexId : SV_VertexID,
 	out float4 OutPosition : SV_POSITION,
-	out float3 OutWorldPosition : WORLD_POSITION)
+	out float3 OutTranslatedWorldPosition : WORLD_POSITION)
 {
 	const FPackedVirtualVoxelNodeDesc PackedNode = VirtualVoxel.NodeDescBuffer[MacroGroupId];
 	const FVirtualVoxelNodeDesc NodeDesc = UnpackVoxelNode(PackedNode, VirtualVoxel.PageResolution);
 
 	// Move this to an actual vertex/index buffer
-	const float3 Min = NodeDesc.MinAABB;
-	const float3 Max = NodeDesc.MaxAABB;
+	const float3 TranslatedWorldOffset = LWCHackToFloat(PrimaryView.PreViewTranslation); // HAIR_LWC
+
+	const float3 Min = NodeDesc.MinAABB + TranslatedWorldOffset;
+	const float3 Max = NodeDesc.MaxAABB + TranslatedWorldOffset;
 
 	const float3 Center = (Min + Max) * 0.5f;
 	const float3 Extent = (Max - Min) * 0.5f;
@@ -49,61 +51,61 @@ void MainVS(
 	const float3 Position6 = Center + float3(+Extent.x, +Extent.y, +Extent.z);
 	const float3 Position7 = Center + float3(-Extent.x, +Extent.y, +Extent.z);
 
-	float3 Position = 0;
+	float3 TranslatedWorldPosition = 0;
 	switch (VertexId)
 	{
-	case 0: Position = Position0; break;
-	case 1: Position = Position1; break;
-	case 2: Position = Position2; break;
-	case 3: Position = Position0; break;
-	case 4: Position = Position2; break;
-	case 5: Position = Position3; break;
-
-	case 6: Position = Position4; break;
-	case 7: Position = Position5; break;
-	case 8: Position = Position6; break;
-	case 9: Position = Position4; break;
-	case 10: Position = Position6; break;
-	case 11: Position = Position7; break;
-
-	case 12: Position = Position0; break;
-	case 13: Position = Position1; break;
-	case 14: Position = Position5; break;
-	case 15: Position = Position0; break;
-	case 16: Position = Position5; break;
-	case 17: Position = Position4; break;
-
-	case 18: Position = Position2; break;
-	case 19: Position = Position3; break;
-	case 20: Position = Position7; break;
-	case 21: Position = Position2; break;
-	case 22: Position = Position7; break;
-	case 23: Position = Position6; break;
-
-	case 24: Position = Position1; break;
-	case 25: Position = Position2; break;
-	case 26: Position = Position6; break;
-	case 27: Position = Position1; break;
-	case 28: Position = Position6; break;
-	case 29: Position = Position5; break;
-
-	case 30: Position = Position3; break;
-	case 31: Position = Position0; break;
-	case 32: Position = Position4; break;
-	case 33: Position = Position3; break;
-	case 34: Position = Position4; break;
-	case 35: Position = Position7; break;
+	case  0: TranslatedWorldPosition = Position0; break;
+	case  1: TranslatedWorldPosition = Position1; break;
+	case  2: TranslatedWorldPosition = Position2; break;
+	case  3: TranslatedWorldPosition = Position0; break;
+	case  4: TranslatedWorldPosition = Position2; break;
+	case  5: TranslatedWorldPosition = Position3; break;
+
+	case  6: TranslatedWorldPosition = Position4; break;
+	case  7: TranslatedWorldPosition = Position5; break;
+	case  8: TranslatedWorldPosition = Position6; break;
+	case  9: TranslatedWorldPosition = Position4; break;
+	case 10: TranslatedWorldPosition = Position6; break;
+	case 11: TranslatedWorldPosition = Position7; break;
+
+	case 12: TranslatedWorldPosition = Position0; break;
+	case 13: TranslatedWorldPosition = Position1; break;
+	case 14: TranslatedWorldPosition = Position5; break;
+	case 15: TranslatedWorldPosition = Position0; break;
+	case 16: TranslatedWorldPosition = Position5; break;
+	case 17: TranslatedWorldPosition = Position4; break;
+
+	case 18: TranslatedWorldPosition = Position2; break;
+	case 19: TranslatedWorldPosition = Position3; break;
+	case 20: TranslatedWorldPosition = Position7; break;
+	case 21: TranslatedWorldPosition = Position2; break;
+	case 22: TranslatedWorldPosition = Position7; break;
+	case 23: TranslatedWorldPosition = Position6; break;
+
+	case 24: TranslatedWorldPosition = Position1; break;
+	case 25: TranslatedWorldPosition = Position2; break;
+	case 26: TranslatedWorldPosition = Position6; break;
+	case 27: TranslatedWorldPosition = Position1; break;
+	case 28: TranslatedWorldPosition = Position6; break;
+	case 29: TranslatedWorldPosition = Position5; break;
+
+	case 30: TranslatedWorldPosition = Position3; break;
+	case 31: TranslatedWorldPosition = Position0; break;
+	case 32: TranslatedWorldPosition = Position4; break;
+	case 33: TranslatedWorldPosition = Position3; break;
+	case 34: TranslatedWorldPosition = Position4; break;
+	case 35: TranslatedWorldPosition = Position7; break;
 	}
 
-	float4x4 WorldToClipMatrix = CPU_WorldToClip;
+	float4x4 TranslatedWorldToClipMatrix = CPU_TranslatedWorldToClip;
 	if (bIsGPUDriven)
 	{
 		FDeepShadowViewInfo DeepShadowViewInfo = DeepShadowViewInfoBuffer[AtlasSlotIndex];
-		WorldToClipMatrix = DeepShadowViewInfo.WorldToClip;
+		TranslatedWorldToClipMatrix = DeepShadowViewInfo.TranslatedWorldToClip;
 	}
 
-	OutWorldPosition = Position;
-	OutPosition = mul(float4(Position, 1), WorldToClipMatrix);
+	OutTranslatedWorldPosition = TranslatedWorldPosition;
+	OutPosition = mul(float4(TranslatedWorldPosition, 1), TranslatedWorldToClipMatrix);
 }
 
 #define VOXEL_TRAVERSAL_DEBUG_PAGEOCCUPANCY 0
@@ -113,7 +115,7 @@ void MainVS(
 
 void MainPS(
 	in float4 InPosition : SV_POSITION,
-	in float3 InWorldPosition : WORLD_POSITION,
+	in float3 InTranslatedWorldPosition : WORLD_POSITION,
 	out float OutDepth : SV_DEPTH)
 {
 	const float2 PixelCoord = InPosition.xy;
@@ -126,9 +128,9 @@ void MainPS(
 	const bool bDebugEnabled = false;
 	const float3 SampleRandom = GetHairVoxelJitter(PixelCoord, View.StateFrameIndexMod8, VirtualVoxel.JitterMode);
 
-	const float3 TracingDirection = bIsDirectional ? LightDirection : normalize(InWorldPosition - LightPosition);
-	const float3 WP0 = InWorldPosition;
-	const float3 WP1 = InWorldPosition + TracingDirection * DistanceThreshold;
+	const float3 TracingDirection = bIsDirectional ? LightDirection : normalize(InTranslatedWorldPosition - TranslatedLightPosition);
+	const float3 TranslatedWP0 = InTranslatedWorldPosition;
+	const float3 TranslatedWP1 = InTranslatedWorldPosition + TracingDirection * DistanceThreshold;
 
 	FVirtualVoxelCommonDesc CommonDesc;
 	CommonDesc.PageCountResolution = VirtualVoxel.PageCountResolution;
@@ -136,8 +138,9 @@ void MainPS(
 	CommonDesc.PageTextureResolution = VirtualVoxel.PageTextureResolution;
 	CommonDesc.PageResolution = VirtualVoxel.PageResolution;
 
+	const float3 TranslatedWorldOffset = LWCHackToFloat(PrimaryView.PreViewTranslation); // HAIR_LWC
 	const FPackedVirtualVoxelNodeDesc PackedNode = VirtualVoxel.NodeDescBuffer[MacroGroupId];
-	const FVirtualVoxelNodeDesc NodeDesc = UnpackVoxelNode(PackedNode, VirtualVoxel.PageResolution);
+	const FVirtualVoxelNodeDesc NodeDesc = UnpackVoxelNode(PackedNode, VirtualVoxel.PageResolution, TranslatedWorldOffset); // HAIR_LWC
 
 	FHairTraversalSettings TraversalSettings = InitHairTraversalSettings();
 	TraversalSettings.DensityScale = VirtualVoxel.DensityScale;
@@ -153,8 +156,8 @@ void MainPS(
 	TraversalSettings.ForcedMip = -1;
 
 	const FHairTraversalResult TraversalResult = ComputeHairCountVirtualVoxel(
-		WP0,
-		WP1,
+		TranslatedWP0,
+		TranslatedWP1,
 		CommonDesc,
 		NodeDesc,
 		VirtualVoxel.PageIndexBuffer,
@@ -165,14 +168,14 @@ void MainPS(
 
 	if (bIsValid)
 	{
-		float4x4 WorldToClipMatrix = CPU_WorldToClip;
+		float4x4 TranslatedWorldToClipMatrix = CPU_TranslatedWorldToClip;
 		if (bIsGPUDriven)
 		{
 			FDeepShadowViewInfo DeepShadowViewInfo = DeepShadowViewInfoBuffer[AtlasSlotIndex];
-			WorldToClipMatrix = DeepShadowViewInfo.WorldToClip;
+			TranslatedWorldToClipMatrix = DeepShadowViewInfo.TranslatedWorldToClip;
 		}
-		const float3 HitP = WP0 + normalize(WP1 - WP0) *  TraversalResult.HitT;
-		float4 ClipP = mul(float4(HitP, 1), WorldToClipMatrix);
+		const float3 HitP = TranslatedWP0 + normalize(TranslatedWP1 - TranslatedWP0) * TraversalResult.HitT;
+		float4 ClipP = mul(float4(HitP, 1), TranslatedWorldToClipMatrix);
 		OutDepth = ClipP.z /= ClipP.w;
 	}
 	else
diff --git a/Engine/Shaders/Private/HairStrands/HairStrandsVoxelPageCommon.ush b/Engine/Shaders/Private/HairStrands/HairStrandsVoxelPageCommon.ush
--- a/Engine/Shaders/Private/HairStrands/HairStrandsVoxelPageCommon.ush
+++ b/Engine/Shaders/Private/HairStrands/HairStrandsVoxelPageCommon.ush
@@ -76,19 +76,19 @@ struct FVirtualVoxelCommonDesc
 
 struct FVirtualVoxelNodeDesc
 {
-	float3	MinAABB;
-	float3	MaxAABB;
+	float3	MinAABB; // In translated world space
+	float3	MaxAABB; // In translated world space
 	uint3	PageIndexResolution;
 	uint3   VirtualResolution;
 	uint    PageIndexOffset;
 	bool    bIsValid;
 };
 
-FVirtualVoxelNodeDesc UnpackVoxelNode(FPackedVirtualVoxelNodeDesc In, uint InPageResolution)
+FVirtualVoxelNodeDesc UnpackVoxelNode(FPackedVirtualVoxelNodeDesc In, uint InPageResolution, float3 InTranslatedWorldOffset = float3(0,0,0)/*HAIR_LWC*/)
 {
 	FVirtualVoxelNodeDesc Out;
-	Out.MaxAABB = In.MaxAABB;
-	Out.MinAABB = In.MinAABB;
+	Out.MaxAABB = In.MaxAABB + InTranslatedWorldOffset;
+	Out.MinAABB = In.MinAABB + InTranslatedWorldOffset;
 	Out.PageIndexResolution = uint3(In.PackedPageIndexResolution & 0xFF, (In.PackedPageIndexResolution >> 8) & 0xFF, (In.PackedPageIndexResolution >> 16) & 0xFF);
 	Out.VirtualResolution = Out.PageIndexResolution * InPageResolution;
 	Out.PageIndexOffset = In.PageIndexOffset;
diff --git a/Engine/Shaders/Private/HairStrands/HairStrandsVoxelPageTraversal.ush b/Engine/Shaders/Private/HairStrands/HairStrandsVoxelPageTraversal.ush
--- a/Engine/Shaders/Private/HairStrands/HairStrandsVoxelPageTraversal.ush
+++ b/Engine/Shaders/Private/HairStrands/HairStrandsVoxelPageTraversal.ush
@@ -160,8 +160,8 @@ FHairTraversalSettings InitHairTraversalSettings()
 /////////////////////////////////////////////////////////////////////////////////////////////////////////////////////////////////
 // Sparse voxel traversal routines
 FHairTraversalResult ComputeHairCountVirtualVoxel(
-	float3 WorldPosition0,
-	float3 WorldPosition1,
+	float3 TranslatedWorldPosition0,
+	float3 TranslatedWorldPosition1,
 	FVirtualVoxelCommonDesc InCommonDesc,
 	FVirtualVoxelNodeDesc InNodeDesc,
 	Buffer<uint> InPageIndexBuffer,
@@ -188,14 +188,14 @@ FHairTraversalResult ComputeHairCountVirtualVoxel(
 	bool bIsPageValid = false;
 	uint3 PageCoord = 0;
 
-	const float2 HitT = LineBoxIntersect(WorldPosition0, WorldPosition1, InNodeDesc.MinAABB, InNodeDesc.MaxAABB);
+	const float2 HitT = LineBoxIntersect(TranslatedWorldPosition0, TranslatedWorldPosition1, InNodeDesc.MinAABB, InNodeDesc.MaxAABB);
 	if (HitT.x < HitT.y)
 	{
 		// Count the number of fibers which are within a cylinder defined by the voxel size, 
 		// and the distance between the origin and the extent of the volume
 		// This assumes that the voxel volume is cubic (i.e. equal dimensions on all sides)
-		const float3 O = lerp(WorldPosition0, WorldPosition1, HitT.xxx);
-		const float3 E = lerp(WorldPosition0, WorldPosition1, HitT.yyy);
+		const float3 O = lerp(TranslatedWorldPosition0, TranslatedWorldPosition1, HitT.xxx);
+		const float3 E = lerp(TranslatedWorldPosition0, TranslatedWorldPosition1, HitT.yyy);
 		const float OELength = min(length(E - O), InSettings.DistanceThreshold);
 		#if VOXEL_TRAVERSAL_DEBUG
 		if (InSettings.bDebugEnabled)
@@ -296,13 +296,13 @@ FHairTraversalResult ComputeHairCountVirtualVoxel(
 
 					#if VOXEL_TRAVERSAL_DEBUG_PAGEOCCUPANCY
 					Out.HairCount = 100;
-					Out.HitT = length(HitP - WorldPosition0);
+					Out.HitT = length(HitP - TranslatedWorldPosition0);
 					return Out;
 					#else
 					if (ShouldStopTraversal(Out, InSettings.CountThreshold, InSettings.bUseOpaqueVisibility))
 					{
 						//Out.HitT = length(HitP - O);
-						Out.HitT = length(HitP - WorldPosition0);
+						Out.HitT = length(HitP - TranslatedWorldPosition0);
 						return Out;
 					}
 					#endif	
@@ -337,22 +337,22 @@ FHairTraversalResult ComputeHairCountVirtualVoxel(
 
 	#if VOXEL_TRAVERSAL_DEBUG
 	AddMacroGroupAABB(InSettings.bDebugEnabled, InNodeDesc.MinAABB, InNodeDesc.MaxAABB);
-	AddLine(InSettings.bDebugEnabled, WorldPosition0, WorldPosition1);
+	AddLine(InSettings.bDebugEnabled, TranslatedWorldPosition0, TranslatedWorldPosition1);
 	#endif
 
-	const float2 HitT = LineBoxIntersect(WorldPosition0, WorldPosition1, InNodeDesc.MinAABB, InNodeDesc.MaxAABB);
+	const float2 HitT = LineBoxIntersect(TranslatedWorldPosition0, TranslatedWorldPosition1, InNodeDesc.MinAABB, InNodeDesc.MaxAABB);
 	if (HitT.x < HitT.y)
 	{
 		// Count the number of fibers which are within a cylinder defined by the voxel size, 
 		// and the distance between the origin and the extent of the volume
 		// This assumes that the voxel volume is cubic (i.e. equal dimensions on all sides)
-		const float3 O = WorldPosition0; // lerp(WorldPosition, IntersectEndPoint, HitT.xxx);
-		const float3 E = WorldPosition1; // lerp(WorldPosition, IntersectEndPoint, HitT.yyy);
+		const float3 O = TranslatedWorldPosition0; // lerp(TranslatedWorldPosition, IntersectEndPoint, HitT.xxx);
+		const float3 E = TranslatedWorldPosition1; // lerp(TranslatedWorldPosition, IntersectEndPoint, HitT.yyy);
 		const float OELength = min(length(E - O), InSettings.DistanceThreshold);
 
 	
 		// Init to 1 or -1 depending of the orientation of stepping
-		const float3 UNormD = WorldPosition1 - WorldPosition0;
+		const float3 UNormD = TranslatedWorldPosition1 - TranslatedWorldPosition0;
 		const int3 Step = sign(UNormD);
 
 		// Step according to voxel size
diff --git a/Engine/Source/Runtime/Renderer/Private/HairStrands/HairStrandsData.h b/Engine/Source/Runtime/Renderer/Private/HairStrands/HairStrandsData.h
--- a/Engine/Source/Runtime/Renderer/Private/HairStrands/HairStrandsData.h
+++ b/Engine/Source/Runtime/Renderer/Private/HairStrands/HairStrandsData.h
@@ -219,7 +219,7 @@ struct FHairStrandsDeepShadowData
 {
 	static const uint32 MaxMacroGroupCount = 16u;
 
-	FMatrix CPU_WorldToLightTransform;
+	FMatrix CPU_TranslatedWorldToLightTransform;
 	FMinHairRadiusAtDepth1 CPU_MinStrandRadiusAtDepth1;
 	FIntRect AtlasRect;
 	uint32 MacroGroupId = ~0;
@@ -229,7 +229,7 @@ struct FHairStrandsDeepShadowData
 	uint32 LightId = ~0;
 	bool bIsLightDirectional = false;
 	FVector3f LightDirection;
-	FVector4f LightPosition;
+	FVector3f TranslatedLightPosition;
 	FLinearColor LightLuminance;
 	float LayerDistribution;
 
@@ -250,7 +250,7 @@ struct FHairStrandsDeepShadowResources
 
 	FRDGTextureRef DepthAtlasTexture = nullptr;
 	FRDGTextureRef LayersAtlasTexture = nullptr;
-	FRDGBufferRef DeepShadowWorldToLightTransforms = nullptr;
+	FRDGBufferRef DeepShadowTranslatedWorldToLightTransforms = nullptr;
 };
 
 ////////////////////////////////////////////////////////////////////////////////////
diff --git a/Engine/Source/Runtime/Renderer/Private/HairStrands/HairStrandsDebug.cpp b/Engine/Source/Runtime/Renderer/Private/HairStrands/HairStrandsDebug.cpp
--- a/Engine/Source/Runtime/Renderer/Private/HairStrands/HairStrandsDebug.cpp
+++ b/Engine/Source/Runtime/Renderer/Private/HairStrands/HairStrandsDebug.cpp
@@ -562,7 +562,7 @@ class FDeepShadowInfoCS : public FGlobalShader
 		SHADER_PARAMETER(uint32, MacroGroupCount)
 		SHADER_PARAMETER_STRUCT_REF(FViewUniformShaderParameters, ViewUniformBuffer)
 		SHADER_PARAMETER_RDG_BUFFER_SRV(Buffer, MacroGroupAABBBuffer)
-		SHADER_PARAMETER_RDG_BUFFER_SRV(StructuredBuffer, ShadowWorldToLightTransformBuffer)
+		SHADER_PARAMETER_RDG_BUFFER_SRV(StructuredBuffer, ShadowTranslatedWorldToLightTransformBuffer)
 		SHADER_PARAMETER_RDG_TEXTURE_UAV(RWTexture2D, OutputTexture)
 		END_SHADER_PARAMETER_STRUCT()
 
@@ -599,7 +599,7 @@ static void AddDeepShadowInfoPass(
 	Parameters->MacroGroupCount = MacroGroupResources.MacroGroupCount;
 	Parameters->SceneTextures = SceneTextures;
 	Parameters->MacroGroupAABBBuffer = GraphBuilder.CreateSRV(MacroGroupResources.MacroGroupAABBsBuffer, PF_R32_SINT);
-	Parameters->ShadowWorldToLightTransformBuffer = GraphBuilder.CreateSRV(DeepShadowResources.DeepShadowWorldToLightTransforms);
+	Parameters->ShadowTranslatedWorldToLightTransformBuffer = GraphBuilder.CreateSRV(DeepShadowResources.DeepShadowTranslatedWorldToLightTransforms);
 	ShaderDrawDebug::SetParameters(GraphBuilder, View.ShaderDrawData, Parameters->ShaderDrawParameters);
 	ShaderPrint::SetParameters(GraphBuilder, View, Parameters->ShaderPrintParameters);
 	Parameters->OutputTexture = GraphBuilder.CreateUAV(OutputTexture);
@@ -1274,7 +1274,7 @@ static void InternalRenderHairStrandsDebugInfo(
 				{
 					for (const FHairStrandsDeepShadowData& DomData : MacroGroupData.DeepShadowDatas)
 					{
-						DrawFrustumWireframe(&ShadowFrustumPDI, DomData.CPU_WorldToLightTransform.Inverse(), FColor::Emerald, 0);
+						DrawFrustumWireframe(&ShadowFrustumPDI, DomData.CPU_TranslatedWorldToLightTransform.Inverse(), FColor::Emerald, 0);
 						DrawWireBox(&ShadowFrustumPDI, DomData.Bounds.GetBox(), FColor::Yellow, 0);
 					}
 				}
@@ -1293,7 +1293,7 @@ static void InternalRenderHairStrandsDebugInfo(
 					{
 						if (DomData.AtlasSlotIndex == CurrentIndex)
 						{
-							DrawFrustumWireframe(&ShadowFrustumPDI, DomData.CPU_WorldToLightTransform.Inverse(), FColor::Emerald, 0);
+							DrawFrustumWireframe(&ShadowFrustumPDI, DomData.CPU_TranslatedWorldToLightTransform.Inverse(), FColor::Emerald, 0);
 							DrawWireBox(&ShadowFrustumPDI, DomData.Bounds.GetBox(), FColor::Yellow, 0);
 						}
 					}
diff --git a/Engine/Source/Runtime/Renderer/Private/HairStrands/HairStrandsDeepShadow.cpp b/Engine/Source/Runtime/Renderer/Private/HairStrands/HairStrandsDeepShadow.cpp
--- a/Engine/Source/Runtime/Renderer/Private/HairStrands/HairStrandsDeepShadow.cpp
+++ b/Engine/Source/Runtime/Renderer/Private/HairStrands/HairStrandsDeepShadow.cpp
@@ -22,7 +22,7 @@ static FAutoConsoleVariableRef CVarDeepShadowInjectVoxelDepth(TEXT("r.HairStrand
 // Inject voxel structure into shadow map to amortize the tracing, and rely on look up kernel to 
 // filter limited resolution
 BEGIN_SHADER_PARAMETER_STRUCT(FHairStransShadowDepthInjectionParameters, )
-	SHADER_PARAMETER(FMatrix44f, CPU_WorldToClip)
+	SHADER_PARAMETER(FMatrix44f, CPU_TranslatedWorldToClip)
 
 	SHADER_PARAMETER(FVector2f, OutputResolution)
 	SHADER_PARAMETER(uint32, AtlasSlotIndex)
@@ -31,7 +31,7 @@ BEGIN_SHADER_PARAMETER_STRUCT(FHairStransShadowDepthInjectionParameters, )
 	SHADER_PARAMETER(FVector3f, LightDirection)
 	SHADER_PARAMETER(uint32, MacroGroupId)
 
-	SHADER_PARAMETER(FVector3f, LightPosition)
+	SHADER_PARAMETER(FVector3f, TranslatedLightPosition)
 	SHADER_PARAMETER(uint32, bIsDirectional)
 	SHADER_PARAMETER_RDG_BUFFER_SRV(StructuredBuffer<FDeepShadowViewInfo>, DeepShadowViewInfoBuffer)
 	SHADER_PARAMETER_STRUCT_REF(FViewUniformShaderParameters, ViewUniformBuffer)
@@ -84,7 +84,7 @@ void AddInjectHairVoxelShadowCaster(
 	const FViewInfo& View,
 	const bool bClear,
 	const FHairStrandsDeepShadowData& DomData,
-	FMatrix CPU_WorldToClipMatrix,
+	FMatrix CPU_TranslatedWorldToClipMatrix,
 	FIntRect AtlasRect,
 	uint32 AtlasSlotIndex,
 	FIntPoint AtlasSlotResolution,
@@ -94,12 +94,12 @@ void AddInjectHairVoxelShadowCaster(
 {
 	FHairStransShadowDepthInjectionParameters* Parameters = GraphBuilder.AllocParameters<FHairStransShadowDepthInjectionParameters>();
 	Parameters->OutputResolution = AtlasSlotResolution;
-	Parameters->CPU_WorldToClip = CPU_WorldToClipMatrix;
+	Parameters->CPU_TranslatedWorldToClip = CPU_TranslatedWorldToClipMatrix;
 	Parameters->ViewUniformBuffer = View.ViewUniformBuffer;
 	Parameters->RenderTargets.DepthStencil = FDepthStencilBinding(OutDepthTexture, bClear ? ERenderTargetLoadAction::EClear : ERenderTargetLoadAction::ELoad, ERenderTargetLoadAction::ENoAction, FExclusiveDepthStencil::DepthWrite_StencilNop);
 	Parameters->VirtualVoxel = VoxelResources.UniformBuffer;
 	Parameters->LightDirection = DomData.LightDirection;
-	Parameters->LightPosition = DomData.LightPosition;
+	Parameters->TranslatedLightPosition = DomData.TranslatedLightPosition;
 	Parameters->bIsDirectional = DomData.bIsLightDirectional ? 1 : 0;
 	Parameters->MacroGroupId = DomData.MacroGroupId;
 	Parameters->DeepShadowViewInfoBuffer = DeepShadowViewInfoBufferSRV;
@@ -179,13 +179,10 @@ class FDeepShadowCreateViewInfoCS : public FGlobalShader
 
 	BEGIN_SHADER_PARAMETER_STRUCT(FParameters, )
 
-		SHADER_PARAMETER_ARRAY(FVector4f,	LightDirections,	[FHairStrandsDeepShadowData::MaxMacroGroupCount])
-		SHADER_PARAMETER_ARRAY(FVector4f,	LightPositions,		[FHairStrandsDeepShadowData::MaxMacroGroupCount])
-		SHADER_PARAMETER_ARRAY(FIntVector4,	MacroGroupIndices,	[FHairStrandsDeepShadowData::MaxMacroGroupCount])
+		SHADER_PARAMETER_ARRAY(FVector4f,	LightDirections,			[FHairStrandsDeepShadowData::MaxMacroGroupCount])
+		SHADER_PARAMETER_ARRAY(FVector4f,	TranslatedLightPositions,	[FHairStrandsDeepShadowData::MaxMacroGroupCount])
+		SHADER_PARAMETER_ARRAY(FIntVector4,	MacroGroupIndices,			[FHairStrandsDeepShadowData::MaxMacroGroupCount])
 
-		SHADER_PARAMETER(FVector3f, CPU_MinAABB)
-		SHADER_PARAMETER(uint32, CPU_bUseCPUData)
-		SHADER_PARAMETER(FVector3f, CPU_MaxAABB)
 		SHADER_PARAMETER(float, RasterizationScale)
 
 		SHADER_PARAMETER(FIntPoint, SlotResolution)
@@ -197,7 +194,9 @@ class FDeepShadowCreateViewInfoCS : public FGlobalShader
 
 		SHADER_PARAMETER_RDG_BUFFER_SRV(Buffer<int>, MacroGroupAABBBuffer)
 		SHADER_PARAMETER_RDG_BUFFER_UAV(RWStructuredBuffer<FDeepShadowViewInfo>, OutShadowViewInfoBuffer)
-		SHADER_PARAMETER_RDG_BUFFER_UAV(RWStructuredBuffer<float4x4>, OutShadowWorldToLightTransformBuffer)
+		SHADER_PARAMETER_RDG_BUFFER_UAV(RWStructuredBuffer<float4x4>, OutShadowTranslatedWorldToLightTransformBuffer)
+		SHADER_PARAMETER_STRUCT_REF(FViewUniformShaderParameters, ViewUniformBuffer)
+
 	END_SHADER_PARAMETER_STRUCT()
 
 public:
@@ -287,6 +286,8 @@ void RenderHairStrandsDeepShadows(
 		FRDGTextureRef FrontDepthAtlasTexture = GraphBuilder.CreateTexture(FRDGTextureDesc::Create2D(AtlasResolution, PF_DepthStencil, FClearValueBinding::DepthFar, TexCreate_DepthStencilTargetable | TexCreate_ShaderResource), TEXT("Hair.ShadowDepth"));
 		FRDGTextureRef DeepShadowLayersAtlasTexture = GraphBuilder.CreateTexture(FRDGTextureDesc::Create2D(AtlasResolution, PF_FloatRGBA, FClearValueBinding::Transparent, TexCreate_RenderTargetable | TexCreate_ShaderResource), TEXT("Hair.DeepShadowLayers"));
 
+		const FVector& TranslatedWorldOffset = View.ViewMatrices.GetPreViewTranslation();
+
 		// TODO add support for multiple view: need to deduplicate light which are visible accross several views
 		// Allocate atlas CPU slot
 		uint32 TotalAtlasSlotIndex = 0;
@@ -317,9 +318,9 @@ void RenderHairStrandsDeepShadows(
 
 				// Note: LightPosition.W is used in the transmittance mask shader to differentiate between directional and local lights.
 				FHairStrandsDeepShadowData& DomData = MacroGroup.DeepShadowDatas.AddZeroed_GetRef();
-				ComputeWorldToLightClip(DomData.CPU_WorldToLightTransform, MinStrandRadiusAtDepth1, MacroGroupBounds, *LightProxy, LightType, AtlasSlotResolution);
+				ComputeTranslatedWorldToLightClip(TranslatedWorldOffset, DomData.CPU_TranslatedWorldToLightTransform, MinStrandRadiusAtDepth1, MacroGroupBounds, *LightProxy, LightType, AtlasSlotResolution);
 				DomData.LightDirection = LightProxy->GetDirection();
-				DomData.LightPosition = FVector4f(FVector(LightProxy->GetPosition()), bIsDirectional ? 0 : 1);
+				DomData.TranslatedLightPosition = FVector4f(FVector(LightProxy->GetPosition() + TranslatedWorldOffset), bIsDirectional ? 0 : 1);
 				DomData.LightLuminance = LightProxy->GetColor();
 				DomData.LayerDistribution = LightProxy->GetDeepShadowLayerDistribution();
 				DomData.bIsLightDirectional = bIsDirectional;
@@ -341,7 +342,7 @@ void RenderHairStrandsDeepShadows(
 		DeepShadowResources.AtlasSlotResolution = AtlasSlotResolution;
 
 		FRDGBufferRef DeepShadowViewInfoBuffer = GraphBuilder.CreateBuffer(FRDGBufferDesc::CreateStructuredDesc(20 * sizeof(float), FMath::Max(1u, TotalAtlasSlotIndex)), TEXT("Hair.DeepShadowViewInfo"));
-		FRDGBufferRef DeepShadowWorldToLightBuffer = GraphBuilder.CreateBuffer(FRDGBufferDesc::CreateStructuredDesc(16 * sizeof(float), FMath::Max(1u, TotalAtlasSlotIndex)), TEXT("Hair.DeepShadowWorldToLightTransform"));
+		FRDGBufferRef DeepShadowTranslatedWorldToLightBuffer = GraphBuilder.CreateBuffer(FRDGBufferDesc::CreateStructuredDesc(16 * sizeof(float), FMath::Max(1u, TotalAtlasSlotIndex)), TEXT("Hair.DeepShadowTranslatedWorldToLightTransform"));
 		FRDGBufferSRVRef DeepShadowViewInfoBufferSRV = GraphBuilder.CreateSRV(DeepShadowViewInfoBuffer);
 
 		DeepShadowResources.bIsGPUDriven = GDeepShadowGPUDriven > 0;
@@ -357,9 +358,9 @@ void RenderHairStrandsDeepShadows(
 			{
 				for (FHairStrandsDeepShadowData& DomData : MacroGroup.DeepShadowDatas)
 				{				
-					Parameters->LightDirections[DomData.AtlasSlotIndex]		= FVector4f(DomData.LightDirection.X, DomData.LightDirection.Y, DomData.LightDirection.Z, 0);
-					Parameters->LightPositions[DomData.AtlasSlotIndex]		= FVector4f(DomData.LightPosition.X, DomData.LightPosition.Y, DomData.LightPosition.Z, DomData.bIsLightDirectional ? 0 : 1);
-					Parameters->MacroGroupIndices[DomData.AtlasSlotIndex]	= FIntVector4(DomData.MacroGroupId, 0,0,0);
+					Parameters->LightDirections[DomData.AtlasSlotIndex]				= FVector4f(DomData.LightDirection.X, DomData.LightDirection.Y, DomData.LightDirection.Z, 0);
+					Parameters->TranslatedLightPositions[DomData.AtlasSlotIndex]	= FVector4f(DomData.TranslatedLightPosition.X, DomData.TranslatedLightPosition.Y, DomData.TranslatedLightPosition.Z, DomData.bIsLightDirectional ? 0 : 1);
+					Parameters->MacroGroupIndices[DomData.AtlasSlotIndex]			= FIntVector4(DomData.MacroGroupId, 0,0,0);
 				}
 			}
 
@@ -368,14 +369,12 @@ void RenderHairStrandsDeepShadows(
 			Parameters->MacroGroupCount = MacroGroupDatas.Num();
 			Parameters->MacroGroupAABBBuffer = GraphBuilder.CreateSRV(MacroGroupResources.MacroGroupAABBsBuffer, PF_R32_SINT);
 			Parameters->OutShadowViewInfoBuffer = GraphBuilder.CreateUAV(DeepShadowViewInfoBuffer);
-			Parameters->OutShadowWorldToLightTransformBuffer = GraphBuilder.CreateUAV(DeepShadowWorldToLightBuffer);
+			Parameters->OutShadowTranslatedWorldToLightTransformBuffer = GraphBuilder.CreateUAV(DeepShadowTranslatedWorldToLightBuffer);
 
 			Parameters->MaxHafFovInRad = 0.5f * FMath::DegreesToRadians(GetDeepShadowMaxFovAngle());
 			Parameters->AABBScale = GetDeepShadowAABBScale();
 			Parameters->RasterizationScale = GetDeepShadowRasterizationScale();
-			Parameters->CPU_bUseCPUData	= 0;
-			Parameters->CPU_MinAABB		= FVector::ZeroVector;
-			Parameters->CPU_MaxAABB		= FVector::ZeroVector;
+			Parameters->ViewUniformBuffer = View.ViewUniformBuffer;
 
 			// Currently support only 32 instance group at max
 			TShaderMapRef<FDeepShadowCreateViewInfoCS> ComputeShader(View.ShaderMap);
@@ -409,7 +408,7 @@ void RenderHairStrandsDeepShadows(
 						View,
 						bClear,
 						DomData,
-						DomData.CPU_WorldToLightTransform,
+						DomData.CPU_TranslatedWorldToLightTransform,
 						DomData.AtlasRect,
 						DomData.AtlasSlotIndex,
 						AtlasSlotResolution,
@@ -438,7 +437,7 @@ void RenderHairStrandsDeepShadows(
 					{
 						FHairDeepShadowRasterUniformParameters* UniformParameters = GraphBuilder.AllocParameters<FHairDeepShadowRasterUniformParameters>();
 
-						UniformParameters->CPU_WorldToClipMatrix = DomData.CPU_WorldToLightTransform;;
+						UniformParameters->CPU_TranslatedWorldToClipMatrix = DomData.CPU_TranslatedWorldToLightTransform;
 						UniformParameters->SliceValue = FVector4f(1, 1, 1, 1);
 						UniformParameters->AtlasRect = DomData.AtlasRect;
 						UniformParameters->AtlasSlotIndex = DomData.AtlasSlotIndex;
@@ -478,7 +477,7 @@ void RenderHairStrandsDeepShadows(
 					{
 						FHairDeepShadowRasterUniformParameters* UniformParameters = GraphBuilder.AllocParameters<FHairDeepShadowRasterUniformParameters>();
 
-						UniformParameters->CPU_WorldToClipMatrix = DomData.CPU_WorldToLightTransform;;
+						UniformParameters->CPU_TranslatedWorldToClipMatrix = DomData.CPU_TranslatedWorldToLightTransform;;
 						UniformParameters->SliceValue = FVector4f(1, 1, 1, 1);
 						UniformParameters->AtlasRect = DomData.AtlasRect;
 						UniformParameters->AtlasSlotIndex = DomData.AtlasSlotIndex;
@@ -510,6 +509,6 @@ void RenderHairStrandsDeepShadows(
 		}
 		DeepShadowResources.DepthAtlasTexture = FrontDepthAtlasTexture;
 		DeepShadowResources.LayersAtlasTexture = DeepShadowLayersAtlasTexture;
-		DeepShadowResources.DeepShadowWorldToLightTransforms = DeepShadowWorldToLightBuffer;
+		DeepShadowResources.DeepShadowTranslatedWorldToLightTransforms = DeepShadowTranslatedWorldToLightBuffer;
 	}
 }
diff --git a/Engine/Source/Runtime/Renderer/Private/HairStrands/HairStrandsRasterCommon.h b/Engine/Source/Runtime/Renderer/Private/HairStrands/HairStrandsRasterCommon.h
--- a/Engine/Source/Runtime/Renderer/Private/HairStrands/HairStrandsRasterCommon.h
+++ b/Engine/Source/Runtime/Renderer/Private/HairStrands/HairStrandsRasterCommon.h
@@ -22,7 +22,7 @@ enum class EHairStrandsRasterPassType : uint8
 // Deep shadow raster pass
 
 BEGIN_GLOBAL_SHADER_PARAMETER_STRUCT(FHairDeepShadowRasterUniformParameters, RENDERER_API)
-	SHADER_PARAMETER(FMatrix44f, CPU_WorldToClipMatrix)
+	SHADER_PARAMETER(FMatrix44f, CPU_TranslatedWorldToClipMatrix)
 	SHADER_PARAMETER(FVector4f, SliceValue)
 	SHADER_PARAMETER(FIntRect, AtlasRect)
 	SHADER_PARAMETER(FIntPoint, ViewportResolution)
diff --git a/Engine/Source/Runtime/Renderer/Private/HairStrands/HairStrandsTransmittance.cpp b/Engine/Source/Runtime/Renderer/Private/HairStrands/HairStrandsTransmittance.cpp
--- a/Engine/Source/Runtime/Renderer/Private/HairStrands/HairStrandsTransmittance.cpp
+++ b/Engine/Source/Runtime/Renderer/Private/HairStrands/HairStrandsTransmittance.cpp
@@ -89,10 +89,22 @@ FVector4f ComputeDeepShadowLayerDepths(float LayerDistribution)
 	return Depths;
 }
 
+static FVector4f GetLightTranslatedWorldPositionAndDirection(const FViewInfo& View, const FLightSceneInfo* LightSceneInfo)
+{
+	const FVector& TranslatedWorldOffset = View.ViewMatrices.GetPreViewTranslation();
+	if (LightSceneInfo->Proxy->GetLightType() == ELightComponentType::LightType_Directional)
+	{
+		return FVector4f(LightSceneInfo->Proxy->GetDirection(), 0.f);
+	}
+	else
+	{
+		return FVector4f(LightSceneInfo->Proxy->GetPosition() + TranslatedWorldOffset, 1.0f);
+	}
+}
+
 struct FHairStrandsTransmittanceLightParams
 {
-	FVector3f LightDirection = FVector3f::ZeroVector;
-	FVector4f LightPosition = FVector4f(0, 0, 0, 0);
+	FVector4f TranslatedLightPosition_LightDirection = FVector4f(0, 0, 0, 0);
 	uint32 LightChannelMask = 0;
 	uint32 ShadowChannelMask = 0;
 	float LightRadius = 0;
@@ -186,11 +198,10 @@ class FHairStrandsVoxelTransmittanceMaskCS : public FGlobalShader
 		SHADER_PARAMETER_RDG_UNIFORM_BUFFER(FVirtualVoxelParameters, VirtualVoxel)
 		SHADER_PARAMETER_STRUCT_INCLUDE(FVirtualShadowMapSamplingParameters, VirtualShadowMap)
 
-		SHADER_PARAMETER(float, LightRadius)
-		SHADER_PARAMETER(FVector3f, LightDirection)
-		SHADER_PARAMETER(FVector4f, LightPosition)
-		SHADER_PARAMETER(uint32, LightChannelMask)
+		SHADER_PARAMETER(FVector4f, TranslatedLightPosition_LightDirection)
 		SHADER_PARAMETER(FVector4f, ShadowChannelMask)
+		SHADER_PARAMETER(uint32, LightChannelMask)
+		SHADER_PARAMETER(float, LightRadius)
 
 		SHADER_PARAMETER_RDG_TEXTURE(Texture2D<uint>, ShadowMaskBitsTexture)
 		SHADER_PARAMETER_RDG_TEXTURE(Texture2D, RayMarchMaskTexture)
@@ -251,8 +262,7 @@ static FRDGBufferRef AddHairStrandsVoxelTransmittanceMaskPass(
 	}
 	else
 	{
-		Parameters->LightDirection = Params.LightDirection;
-		Parameters->LightPosition = Params.LightPosition;
+		Parameters->TranslatedLightPosition_LightDirection = Params.TranslatedLightPosition_LightDirection;
 		Parameters->LightRadius = Params.LightRadius;
 		Parameters->RayMarchMaskTexture = ShadowMaskTexture ? ShadowMaskTexture : GSystemTextures.GetWhiteDummy(GraphBuilder);
 		Parameters->ShadowMaskBitsTexture = nullptr;
@@ -305,11 +315,10 @@ class FHairStrandsDeepShadowTransmittanceMaskCS : public FGlobalShader
 		SHADER_PARAMETER_STRUCT_REF(FViewUniformShaderParameters, ViewUniformBuffer)
 
 		SHADER_PARAMETER_ARRAY(FIntVector4, DeepShadow_AtlasSlotOffsets_AtlasSlotIndex, [FHairStrandsDeepShadowData::MaxMacroGroupCount])
-		SHADER_PARAMETER_ARRAY(FMatrix44f, DeepShadow_CPUWorldToLightTransforms, [FHairStrandsDeepShadowData::MaxMacroGroupCount])
+		SHADER_PARAMETER_ARRAY(FMatrix44f, DeepShadow_CPUTranslatedWorldToLightTransforms, [FHairStrandsDeepShadowData::MaxMacroGroupCount])
 		SHADER_PARAMETER(FIntPoint, DeepShadow_Resolution)
 		SHADER_PARAMETER(float, LightRadius)
-		SHADER_PARAMETER(FVector3f, LightDirection)
-		SHADER_PARAMETER(FVector4f, LightPosition)
+		SHADER_PARAMETER(FVector4f, TranslatedLightPosition_LightDirection)
 		SHADER_PARAMETER(uint32, LightChannelMask)
 		SHADER_PARAMETER(FVector4f, ShadowChannelMask)
 		SHADER_PARAMETER(FVector4f, DeepShadow_LayerDepths)
@@ -324,7 +333,7 @@ class FHairStrandsDeepShadowTransmittanceMaskCS : public FGlobalShader
 		SHADER_PARAMETER_RDG_TEXTURE(Texture2D, RayMarchMaskTexture)
 		SHADER_PARAMETER_RDG_TEXTURE(Texture2D, DeepShadow_FrontDepthTexture)
 		SHADER_PARAMETER_RDG_TEXTURE(Texture2D, DeepShadow_DomTexture)
-		SHADER_PARAMETER_RDG_BUFFER_SRV(StructuredBuffer, DeepShadow_WorldToLightTransformBuffer)
+		SHADER_PARAMETER_RDG_BUFFER_SRV(StructuredBuffer, DeepShadow_TranslatedWorldToLightTransformBuffer)
 		SHADER_PARAMETER_RDG_BUFFER_UAV(RWBuffer, OutTransmittanceMask)
 
 		RDG_BUFFER_ACCESS(IndirectArgsBuffer, ERHIAccess::IndirectArgs)
@@ -349,8 +358,8 @@ IMPLEMENT_GLOBAL_SHADER(FHairStrandsDeepShadowTransmittanceMaskCS, "/Engine/Priv
 struct FHairStrandsDeepShadowTransmittanceLightParams : FHairStrandsTransmittanceLightParams
 {
 	FIntVector4 DeepShadow_AtlasSlotOffsets_AtlasSlotIndex[FHairStrandsDeepShadowData::MaxMacroGroupCount];
-	FMatrix DeepShadow_CPUWorldToLightTransforms[FHairStrandsDeepShadowData::MaxMacroGroupCount];
-	FRDGBufferSRVRef DeepShadow_WorldToLightTransformBuffer = nullptr;
+	FMatrix DeepShadow_CPUTranslatedWorldToLightTransforms[FHairStrandsDeepShadowData::MaxMacroGroupCount];
+	FRDGBufferSRVRef DeepShadow_TranslatedWorldToLightTransformBuffer = nullptr;
 	FIntPoint DeepShadow_Resolution = FIntPoint(0, 0);
 	bool DeepShadow_bIsGPUDriven = false;
 	FVector4f DeepShadow_LayerDepths = FVector4f(0, 0, 0, 0);
@@ -384,8 +393,7 @@ static FRDGBufferRef AddHairStrandsDeepShadowTransmittanceMaskPass(
 	Parameters->OutTransmittanceMask = GraphBuilder.CreateUAV(OutBuffer, FHairStrandsTransmittanceMaskData::Format);
 	Parameters->LightChannelMask = Params.LightChannelMask;
 	Parameters->DeepShadow_Resolution = Params.DeepShadow_Resolution;
-	Parameters->LightDirection = Params.LightDirection;
-	Parameters->LightPosition = Params.LightPosition;
+	Parameters->TranslatedLightPosition_LightDirection = Params.TranslatedLightPosition_LightDirection;
 	Parameters->LightRadius = Params.LightRadius;
 	Parameters->DeepShadow_DepthBiasScale = Params.DeepShadow_DepthBiasScale;
 	Parameters->DeepShadow_DensityScale = Params.DeepShadow_DensityScale;
@@ -397,7 +405,7 @@ static FRDGBufferRef AddHairStrandsDeepShadowTransmittanceMaskPass(
 	Parameters->IndirectArgsBuffer = IndirectArgsBuffer;
 	Parameters->HairStrands = HairStrands::BindHairStrandsViewUniformParameters(View);
 	Parameters->DeepShadow_bIsGPUDriven = Params.DeepShadow_bIsGPUDriven ? 1 : 0;;
-	Parameters->DeepShadow_WorldToLightTransformBuffer = Params.DeepShadow_WorldToLightTransformBuffer;
+	Parameters->DeepShadow_TranslatedWorldToLightTransformBuffer = Params.DeepShadow_TranslatedWorldToLightTransformBuffer;
 	Parameters->RayMarchMaskTexture = ScreenShadowMaskSubPixelTexture ? ScreenShadowMaskSubPixelTexture : GraphBuilder.RegisterExternalTexture(GSystemTextures.WhiteDummy);
 	Parameters->ShadowChannelMask = FVector4f(0, 0, 0, 0);
 	Parameters->ShadowChannelMask[FMath::Clamp(Params.ShadowChannelMask, 0, 3)] = 1.0f;
@@ -405,7 +413,7 @@ static FRDGBufferRef AddHairStrandsDeepShadowTransmittanceMaskPass(
 	for (uint32 SlotIndex=0;SlotIndex< FHairStrandsDeepShadowData::MaxMacroGroupCount;++SlotIndex)
 	{
 		Parameters->DeepShadow_AtlasSlotOffsets_AtlasSlotIndex[SlotIndex] = Params.DeepShadow_AtlasSlotOffsets_AtlasSlotIndex[SlotIndex];
-		Parameters->DeepShadow_CPUWorldToLightTransforms[SlotIndex] = Params.DeepShadow_CPUWorldToLightTransforms[SlotIndex];
+		Parameters->DeepShadow_CPUTranslatedWorldToLightTransforms[SlotIndex] = Params.DeepShadow_CPUTranslatedWorldToLightTransforms[SlotIndex];
 	}
 	
 	check(NodeGroupSize == 64 || NodeGroupSize == 32);
@@ -435,8 +443,7 @@ class FHairStrandsVoxelShadowMaskPS : public FGlobalShader
 	BEGIN_SHADER_PARAMETER_STRUCT(FParameters, )
 		SHADER_PARAMETER_STRUCT_INCLUDE(ShaderDrawDebug::FShaderParameters, ShaderDrawParameters)
 		
-		SHADER_PARAMETER(FVector4f, Voxel_LightPosition)
-		SHADER_PARAMETER(FVector3f, Voxel_LightDirection)
+		SHADER_PARAMETER(FVector4f, Voxel_TranslatedLightPosition_LightDirection)
 		SHADER_PARAMETER(uint32, Voxel_MacroGroupId)
 		SHADER_PARAMETER(uint32, Voxel_RandomType)
 		SHADER_PARAMETER(uint32, bIsWholeSceneLight)
@@ -466,11 +473,10 @@ IMPLEMENT_GLOBAL_SHADER(FHairStrandsVoxelShadowMaskPS, "/Engine/Private/HairStra
 
 struct FHairStrandsVoxelShadowParams
 {
-	bool			bIsWholeSceneLight = false;
+	FVector4f		Voxel_TranslatedLightPosition_LightDirection = FVector4f(0, 0, 0, 0);
 	uint32			OutputChannel = ~0;
-	FVector3f		Voxel_LightDirection = FVector3f::ZeroVector;
-	FVector4f		Voxel_LightPosition = FVector4f(0, 0, 0, 0);
 	uint32			Voxel_MacroGroupId;
+	bool			bIsWholeSceneLight = false;
 };
 
 // Opaque mask from voxel
@@ -492,8 +498,7 @@ static void AddHairStrandsVoxelShadowMaskPass(
 	Parameters->LinearSampler = TStaticSamplerState<SF_Bilinear, AM_Clamp, AM_Clamp, AM_Clamp>::GetRHI();
 	Parameters->ShadowSampler = TStaticSamplerState<SF_Point, AM_Clamp, AM_Clamp, AM_Clamp>::GetRHI();
 	Parameters->bIsWholeSceneLight = Params.bIsWholeSceneLight ? 1 : 0;
-	Parameters->Voxel_LightPosition		= Params.Voxel_LightPosition;
-	Parameters->Voxel_LightDirection	= Params.Voxel_LightDirection;
+	Parameters->Voxel_TranslatedLightPosition_LightDirection = Params.Voxel_TranslatedLightPosition_LightDirection;
 	Parameters->Voxel_MacroGroupId		= Params.Voxel_MacroGroupId;
 	Parameters->Voxel_RandomType = FMath::Clamp(GHairStrandsShadowRandomTraversalType, 0, 2);	
 	Parameters->RenderTargets[0] = FRenderTargetBinding(OutShadowMask, ERenderTargetLoadAction::ELoad);
@@ -586,7 +591,7 @@ class FHairStrandsDeepShadowMaskPS : public FGlobalShader
 		SHADER_PARAMETER(FIntPoint, DeepShadow_SlotOffset)
 		SHADER_PARAMETER(uint32, DeepShadow_SlotIndex)
 		SHADER_PARAMETER(FIntPoint, DeepShadow_SlotResolution)
-		SHADER_PARAMETER(FMatrix44f, DeepShadow_CPUWorldToLightTransform)
+		SHADER_PARAMETER(FMatrix44f, DeepShadow_CPUTranslatedWorldToLightTransform)
 		SHADER_PARAMETER(uint32, bIsWholeSceneLight)
 		SHADER_PARAMETER(float, DeepShadow_DepthBiasScale)
 		SHADER_PARAMETER(float, DeepShadow_DensityScale)
@@ -595,7 +600,7 @@ class FHairStrandsDeepShadowMaskPS : public FGlobalShader
 
 		SHADER_PARAMETER_RDG_TEXTURE(Texture2D, RayMarchMaskTexture)
 
-		SHADER_PARAMETER_RDG_BUFFER_SRV(StructuredBuffer, DeepShadow_WorldToLightTransformBuffer)
+		SHADER_PARAMETER_RDG_BUFFER_SRV(StructuredBuffer, DeepShadow_TranslatedWorldToLightTransformBuffer)
 		SHADER_PARAMETER_RDG_TEXTURE(Texture2D, DeepShadow_FrontDepthTexture)
 		SHADER_PARAMETER_RDG_TEXTURE(Texture2D, DeepShadow_DomTexture)
 		SHADER_PARAMETER_RDG_TEXTURE(Texture2D, SceneDepthTexture)
@@ -624,8 +629,8 @@ struct FHairStrandsDeepShadowParams
 {
 	bool			bIsWholeSceneLight = false;
 	uint32			OutputChannel = ~0;
-	FRDGBufferSRVRef DeepShadow_WorldToLightTransformBuffer = nullptr;
-	FMatrix			DeepShadow_CPUWorldToLightTransform;
+	FRDGBufferSRVRef DeepShadow_TranslatedWorldToLightTransformBuffer = nullptr;
+	FMatrix			DeepShadow_CPUTranslatedWorldToLightTransform;
 	FIntRect		DeepShadow_AtlasRect;
 	FRDGTextureRef	DeepShadow_FrontDepthTexture = nullptr;
 	FRDGTextureRef	DeepShadow_LayerTexture = nullptr;
@@ -649,7 +654,7 @@ static void AddHairStrandsDeepShadowMaskPass(
 	FHairStrandsDeepShadowMaskPS::FParameters* Parameters = GraphBuilder.AllocParameters<FHairStrandsDeepShadowMaskPS::FParameters>();
 	Parameters->ViewUniformBuffer = View.ViewUniformBuffer;
 	Parameters->SceneDepthTexture = SceneDepthTexture;
-	Parameters->DeepShadow_CPUWorldToLightTransform = Params.DeepShadow_CPUWorldToLightTransform;
+	Parameters->DeepShadow_CPUTranslatedWorldToLightTransform = Params.DeepShadow_CPUTranslatedWorldToLightTransform;
 	Parameters->DeepShadow_FrontDepthTexture = Params.DeepShadow_FrontDepthTexture;
 	Parameters->DeepShadow_DomTexture = Params.DeepShadow_LayerTexture;
 	Parameters->LinearSampler = TStaticSamplerState<SF_Bilinear, AM_Clamp, AM_Clamp, AM_Clamp>::GetRHI();
@@ -663,7 +668,7 @@ static void AddHairStrandsDeepShadowMaskPass(
 	Parameters->DeepShadow_SlotIndex = Params.DeepShadow_AtlasSlotIndex;
 	Parameters->DeepShadow_SlotOffset = FIntPoint(Params.DeepShadow_AtlasRect.Min.X, Params.DeepShadow_AtlasRect.Min.Y);
 	Parameters->DeepShadow_SlotResolution = FIntPoint(Params.DeepShadow_AtlasRect.Max.X - Params.DeepShadow_AtlasRect.Min.X, Params.DeepShadow_AtlasRect.Max.Y - Params.DeepShadow_AtlasRect.Min.Y);
-	Parameters->DeepShadow_WorldToLightTransformBuffer = Params.DeepShadow_WorldToLightTransformBuffer;
+	Parameters->DeepShadow_TranslatedWorldToLightTransformBuffer = Params.DeepShadow_TranslatedWorldToLightTransformBuffer;
 	Parameters->DeepShadow_bIsGPUDriven = Params.DeepShadow_bIsGPUDriven ? 1 : 0;;
 	Parameters->HairStrands = HairStrands::BindHairStrandsViewUniformParameters(View);
 	if (ShaderDrawDebug::IsEnabled(View))
@@ -770,33 +775,32 @@ static FHairStrandsTransmittanceMaskData InternalRenderHairStrandsTransmittanceM
 		Params.DeepShadow_DensityScale = GetDeepShadowDensityScale();
 		Params.DeepShadow_DepthBiasScale = GetDeepShadowDepthBiasScale();
 		memset(Params.DeepShadow_AtlasSlotOffsets_AtlasSlotIndex, 0, sizeof(Params.DeepShadow_AtlasSlotOffsets_AtlasSlotIndex));
-		memset(Params.DeepShadow_CPUWorldToLightTransforms, 0, sizeof(Params.DeepShadow_CPUWorldToLightTransforms));
+		memset(Params.DeepShadow_CPUTranslatedWorldToLightTransforms, 0, sizeof(Params.DeepShadow_CPUTranslatedWorldToLightTransforms));
 
-		FRDGBufferSRVRef DeepShadow_WorldToLightTransformBufferSRV = nullptr;
+		FRDGBufferSRVRef DeepShadow_TranslatedWorldToLightTransformBufferSRV = nullptr;
 		for (const FHairStrandsMacroGroupData& MacroGroupData : MacroGroupDatas)
 		{
 			for (const FHairStrandsDeepShadowData& DeepShadowData : MacroGroupData.DeepShadowDatas)
 			{
 				if (DeepShadowData.LightId == LightSceneInfo->Id)
 				{
-					if (DeepShadow_WorldToLightTransformBufferSRV == nullptr)
+					if (DeepShadow_TranslatedWorldToLightTransformBufferSRV == nullptr)
 					{
-						DeepShadow_WorldToLightTransformBufferSRV = GraphBuilder.CreateSRV(DeepShadowResources.DeepShadowWorldToLightTransforms);
+						DeepShadow_TranslatedWorldToLightTransformBufferSRV = GraphBuilder.CreateSRV(DeepShadowResources.DeepShadowTranslatedWorldToLightTransforms);
 					}
 
 					bHasFoundLight = true;
+					Params.TranslatedLightPosition_LightDirection = GetLightTranslatedWorldPositionAndDirection(View, LightSceneInfo);
 					Params.DeepShadow_FrontDepthTexture = DeepShadowResources.DepthAtlasTexture;
 					Params.DeepShadow_DomTexture = DeepShadowResources.LayersAtlasTexture;
 					Params.DeepShadow_Resolution = DeepShadowData.ShadowResolution;
-					Params.LightDirection = DeepShadowData.LightDirection;
-					Params.LightPosition = DeepShadowData.LightPosition;
 					Params.LightRadius = 0;
 					Params.LightChannelMask = LightSceneInfo->Proxy->GetLightingChannelMask();
 					Params.ShadowChannelMask = bProjectingForForwardShading ? LightSceneInfo->GetDynamicShadowMapChannel() : 0;
 					Params.DeepShadow_LayerDepths = ComputeDeepShadowLayerDepths(DeepShadowData.LayerDistribution);
 					Params.DeepShadow_AtlasSlotOffsets_AtlasSlotIndex[DeepShadowData.MacroGroupId] = FIntVector4(DeepShadowData.AtlasRect.Min.X, DeepShadowData.AtlasRect.Min.Y, DeepShadowData.AtlasSlotIndex, 0);
-					Params.DeepShadow_CPUWorldToLightTransforms[DeepShadowData.MacroGroupId] = DeepShadowData.CPU_WorldToLightTransform;
-					Params.DeepShadow_WorldToLightTransformBuffer = DeepShadow_WorldToLightTransformBufferSRV;
+					Params.DeepShadow_CPUTranslatedWorldToLightTransforms[DeepShadowData.MacroGroupId] = DeepShadowData.CPU_TranslatedWorldToLightTransform;
+					Params.DeepShadow_TranslatedWorldToLightTransformBuffer = DeepShadow_TranslatedWorldToLightTransformBufferSRV;
 					Params.DeepShadow_bIsGPUDriven = DeepShadowResources.bIsGPUDriven;
 				}
 			}
@@ -823,8 +827,7 @@ static FHairStrandsTransmittanceMaskData InternalRenderHairStrandsTransmittanceM
 		LightSceneInfo->Proxy->GetLightShaderParameters(LightParameters);
 
 		FHairStrandsTransmittanceLightParams Params;
-		Params.LightDirection = LightSceneInfo->Proxy->GetDirection();
-		Params.LightPosition = FVector4f(FVector(LightSceneInfo->Proxy->GetPosition()), LightSceneInfo->Proxy->GetLightType() == ELightComponentType::LightType_Directional ? 0 : 1);
+		Params.TranslatedLightPosition_LightDirection = GetLightTranslatedWorldPositionAndDirection(View, LightSceneInfo);
 		Params.LightChannelMask = LightSceneInfo->Proxy->GetLightingChannelMask();
 		Params.ShadowChannelMask = bProjectingForForwardShading ? LightSceneInfo->GetDynamicShadowMapChannel() : 0;
 		Params.LightRadius = FMath::Max(LightParameters.SourceLength, LightParameters.SourceRadius);
@@ -928,8 +931,8 @@ static void InternalRenderHairStrandsShadowMask(
 	bool bHasDeepShadow = false;
 	if (!IsHairStrandsForVoxelTransmittanceAndShadowEnable())
 	{
-		FRDGBufferRef DeepShadow_WorldToLightTransformBuffer = nullptr;
-		FRDGBufferSRVRef DeepShadow_WorldToLightTransformBufferSRV = nullptr;
+		FRDGBufferRef DeepShadow_TranslatedWorldToLightTransformBuffer = nullptr;
+		FRDGBufferSRVRef DeepShadow_TranslatedWorldToLightTransformBufferSRV = nullptr;
 
 		for (const FHairStrandsMacroGroupData& MacroGroupData : InMacroGroupDatas)
 		{
@@ -938,10 +941,10 @@ static void InternalRenderHairStrandsShadowMask(
 				if (DomData.LightId != LightSceneInfo->Id)
 					continue;
 
-				if (DeepShadow_WorldToLightTransformBuffer == nullptr)
+				if (DeepShadow_TranslatedWorldToLightTransformBuffer == nullptr)
 				{
-					DeepShadow_WorldToLightTransformBuffer = DeepShadowResources.DeepShadowWorldToLightTransforms;
-					DeepShadow_WorldToLightTransformBufferSRV = GraphBuilder.CreateSRV(DeepShadow_WorldToLightTransformBuffer);
+					DeepShadow_TranslatedWorldToLightTransformBuffer = DeepShadowResources.DeepShadowTranslatedWorldToLightTransforms;
+					DeepShadow_TranslatedWorldToLightTransformBufferSRV = GraphBuilder.CreateSRV(DeepShadow_TranslatedWorldToLightTransformBuffer);
 				}
 
 				bHasDeepShadow = true;
@@ -949,9 +952,9 @@ static void InternalRenderHairStrandsShadowMask(
 
 				FHairStrandsDeepShadowParams Params;
 				Params.DeepShadow_AtlasSlotIndex = DomData.AtlasSlotIndex;
-				Params.DeepShadow_WorldToLightTransformBuffer = DeepShadow_WorldToLightTransformBufferSRV;
+				Params.DeepShadow_TranslatedWorldToLightTransformBuffer = DeepShadow_TranslatedWorldToLightTransformBufferSRV;
 				Params.DeepShadow_bIsGPUDriven = DeepShadowResources.bIsGPUDriven ? 1 : 0;
-				Params.DeepShadow_CPUWorldToLightTransform = DomData.CPU_WorldToLightTransform;
+				Params.DeepShadow_CPUTranslatedWorldToLightTransform = DomData.CPU_TranslatedWorldToLightTransform;
 				Params.DeepShadow_AtlasRect = DomData.AtlasRect;
 				Params.DeepShadow_FrontDepthTexture = DeepShadowResources.DepthAtlasTexture;
 				Params.DeepShadow_LayerTexture = DeepShadowResources.LayersAtlasTexture;
@@ -980,8 +983,7 @@ static void InternalRenderHairStrandsShadowMask(
 			const bool bIsWholeSceneLight = LightSceneInfo->Proxy->GetLightType() == LightType_Directional;
 
 			FHairStrandsVoxelShadowParams Params;
-			Params.Voxel_LightDirection = LightSceneInfo->Proxy->GetDirection();
-			Params.Voxel_LightPosition = FVector4f(FVector(LightSceneInfo->Proxy->GetPosition()), LightSceneInfo->Proxy->GetLightType() == ELightComponentType::LightType_Directional ? 0 : 1);
+			Params.Voxel_TranslatedLightPosition_LightDirection = GetLightTranslatedWorldPositionAndDirection(View, LightSceneInfo);
 			Params.Voxel_MacroGroupId = MacroGroupData.MacroGroupId;
 			Params.bIsWholeSceneLight = bIsWholeSceneLight ? 1 : 0;
 			Params.OutputChannel = bProjectingForForwardShading ? LightSceneInfo->GetDynamicShadowMapChannel() : ~0;
diff --git a/Engine/Source/Runtime/Renderer/Private/HairStrands/HairStrandsUtils.cpp b/Engine/Source/Runtime/Renderer/Private/HairStrands/HairStrandsUtils.cpp
--- a/Engine/Source/Runtime/Renderer/Private/HairStrands/HairStrandsUtils.cpp
+++ b/Engine/Source/Runtime/Renderer/Private/HairStrands/HairStrandsUtils.cpp
@@ -147,59 +147,62 @@ FMinHairRadiusAtDepth1 ComputeMinStrandRadiusAtDepth1(
 	return Out;
 }
 
-void ComputeWorldToLightClip(
-	FMatrix& WorldToClipTransform,
-	FMinHairRadiusAtDepth1& MinStrandRadiusAtDepth1,
+void ComputeTranslatedWorldToLightClip(
+	const FVector& TranslatedWorldOffset,
+	FMatrix& OutTranslatedWorldToClipTransform,
+	FMinHairRadiusAtDepth1& OutMinStrandRadiusAtDepth1,
 	const FBoxSphereBounds& PrimitivesBounds,
 	const FLightSceneProxy& LightProxy,
 	const ELightComponentType LightType,
 	const FIntPoint& ShadowResolution)
 {
-	const FSphere SphereBound = PrimitivesBounds.GetSphere();
-	const float SphereRadius = SphereBound.W * GetDeepShadowAABBScale();
-	const FVector LightPosition = LightProxy.GetPosition();
+	// Translated SphereBound & translated light position
+	const FSphere TranslatedSphereBound = FSphere(PrimitivesBounds.GetSphere().Center + TranslatedWorldOffset, PrimitivesBounds.GetSphere().W);
+	const float SphereRadius = TranslatedSphereBound.W * GetDeepShadowAABBScale();
+	const FVector3f TranslatedLightPosition = FVector3f(LightProxy.GetPosition() + TranslatedWorldOffset);
+
 	const float MinNear = 1.0f; // 1cm, lower value than this cause precision issue. Similar value than in HairStrandsDeepShadowAllocation.usf
-	const float MinZ = FMath::Max(MinNear, FMath::Max(0.1f, FVector::Distance(LightPosition, SphereBound.Center)) - SphereBound.W);
-	const float MaxZ = FMath::Max(0.2f, FVector::Distance(LightPosition, SphereBound.Center)) + SphereBound.W;
+	const float MinZ = FMath::Max(MinNear, FMath::Max(0.1f, FVector::Distance(TranslatedLightPosition, TranslatedSphereBound.Center)) - TranslatedSphereBound.W);
+	const float MaxZ = FMath::Max(0.2f, FVector::Distance(TranslatedLightPosition, TranslatedSphereBound.Center)) + TranslatedSphereBound.W;
 	const float MaxDeepShadowFrustumHalfAngleInRad = 0.5f * FMath::DegreesToRadians(GetDeepShadowMaxFovAngle());
 
 	const float StrandHairRasterizationScale = GetDeepShadowRasterizationScale();
 	const float StrandHairStableRasterizationScale = FMath::Max(GStrandHairStableRasterizationScale, 1.0f);
-	MinStrandRadiusAtDepth1 = FMinHairRadiusAtDepth1();
-	WorldToClipTransform = FMatrix::Identity;
+	OutMinStrandRadiusAtDepth1 = FMinHairRadiusAtDepth1();
+	OutTranslatedWorldToClipTransform = FMatrix::Identity;
 	if (LightType == LightType_Directional)
 	{
 		const FVector3f LightDirection = LightProxy.GetDirection();
 		FReversedZOrthoMatrix OrthoMatrix(SphereRadius, SphereRadius, 1.f / (2 * SphereRadius), 0);
-		FLookAtMatrix LookAt(SphereBound.Center - LightDirection * SphereRadius, SphereBound.Center, FVector(0, 0, 1));
-		WorldToClipTransform = LookAt * OrthoMatrix;
+		FLookAtMatrix LookAt(TranslatedSphereBound.Center - LightDirection * SphereRadius, TranslatedSphereBound.Center, FVector(0, 0, 1));
+		OutTranslatedWorldToClipTransform = LookAt * OrthoMatrix;
 
 		const float RadiusAtDepth1 = SphereRadius / FMath::Min(ShadowResolution.X, ShadowResolution.Y);
-		MinStrandRadiusAtDepth1.Stable = RadiusAtDepth1 * StrandHairStableRasterizationScale;
-		MinStrandRadiusAtDepth1.Primary = RadiusAtDepth1 * StrandHairRasterizationScale;
-		MinStrandRadiusAtDepth1.Velocity = MinStrandRadiusAtDepth1.Primary;
+		OutMinStrandRadiusAtDepth1.Stable = RadiusAtDepth1 * StrandHairStableRasterizationScale;
+		OutMinStrandRadiusAtDepth1.Primary = RadiusAtDepth1 * StrandHairRasterizationScale;
+		OutMinStrandRadiusAtDepth1.Velocity = OutMinStrandRadiusAtDepth1.Primary;
 	}
 	else if (LightType == LightType_Spot || LightType == LightType_Point)
 	{
-		const float SphereDistance = FVector::Distance(LightPosition, SphereBound.Center);
+		const float SphereDistance = FVector3f::Distance(TranslatedLightPosition, TranslatedSphereBound.Center);
 		float HalfFov = asin(SphereRadius / SphereDistance);
 		HalfFov = FMath::Min(HalfFov, MaxDeepShadowFrustumHalfAngleInRad);
 
 		FReversedZPerspectiveMatrix ProjMatrix(HalfFov, 1, 1, MinZ, MaxZ);
-		FLookAtMatrix WorldToLight(LightPosition, SphereBound.Center, FVector(0, 0, 1));
-		WorldToClipTransform = WorldToLight * ProjMatrix;
-		MinStrandRadiusAtDepth1 = ComputeMinStrandRadiusAtDepth1(ShadowResolution, 2 * HalfFov, 1, StrandHairRasterizationScale);
+		FLookAtMatrix TranslatedWorldToLight(TranslatedLightPosition, TranslatedSphereBound.Center, FVector(0, 0, 1));
+		OutTranslatedWorldToClipTransform = TranslatedWorldToLight * ProjMatrix;
+		OutMinStrandRadiusAtDepth1 = ComputeMinStrandRadiusAtDepth1(ShadowResolution, 2 * HalfFov, 1, StrandHairRasterizationScale);
 	}
 	else if (LightType == LightType_Rect)
 	{
-		const float SphereDistance = FVector::Distance(LightPosition, SphereBound.Center);
+		const float SphereDistance = FVector3f::Distance(TranslatedLightPosition, TranslatedSphereBound.Center);
 		float HalfFov = asin(SphereRadius / SphereDistance);
 		HalfFov = FMath::Min(HalfFov, MaxDeepShadowFrustumHalfAngleInRad);
 
 		FReversedZPerspectiveMatrix ProjMatrix(HalfFov, 1, 1, MinZ, MaxZ);
-		FLookAtMatrix WorldToLight(LightPosition, SphereBound.Center, FVector(0, 0, 1));
-		WorldToClipTransform = WorldToLight * ProjMatrix;
-		MinStrandRadiusAtDepth1 = ComputeMinStrandRadiusAtDepth1(ShadowResolution, 2 * HalfFov, 1, StrandHairRasterizationScale);
+		FLookAtMatrix TranslatedWorldToLight(TranslatedLightPosition, TranslatedSphereBound.Center, FVector(0, 0, 1));
+		OutTranslatedWorldToClipTransform = TranslatedWorldToLight * ProjMatrix;
+		OutMinStrandRadiusAtDepth1 = ComputeMinStrandRadiusAtDepth1(ShadowResolution, 2 * HalfFov, 1, StrandHairRasterizationScale);
 	}
 }
 
diff --git a/Engine/Source/Runtime/Renderer/Private/HairStrands/HairStrandsUtils.h b/Engine/Source/Runtime/Renderer/Private/HairStrands/HairStrandsUtils.h
--- a/Engine/Source/Runtime/Renderer/Private/HairStrands/HairStrandsUtils.h
+++ b/Engine/Source/Runtime/Renderer/Private/HairStrands/HairStrandsUtils.h
@@ -22,9 +22,10 @@ FMinHairRadiusAtDepth1 ComputeMinStrandRadiusAtDepth1(
 
 FIntRect ComputeProjectedScreenRect(const FBox& B, const FViewInfo& View);
 
-void ComputeWorldToLightClip(
-	FMatrix& WorldToClipTransform,
-	FMinHairRadiusAtDepth1& MinStrandRadiusAtDepth1,
+void ComputeTranslatedWorldToLightClip(
+	const FVector& TranslatedWorldOffset,
+	FMatrix& OutTranslatedWorldToClipTransform,
+	FMinHairRadiusAtDepth1& OutMinStrandRadiusAtDepth1,
 	const FBoxSphereBounds& PrimitivesBounds,
 	const class FLightSceneProxy& LightProxy,
 	const ELightComponentType LightType,
@@ -74,4 +75,4 @@ FVector4f PackHairRenderInfo(
 
 uint32 PackHairRenderInfoBits(
 	bool bIsOrtho,
-	bool bIsGPUDriven);
\ No newline at end of file
+	bool bIsGPUDriven);
diff --git a/Engine/Source/Runtime/Renderer/Private/HairStrands/HairStrandsVisibility.cpp b/Engine/Source/Runtime/Renderer/Private/HairStrands/HairStrandsVisibility.cpp
--- a/Engine/Source/Runtime/Renderer/Private/HairStrands/HairStrandsVisibility.cpp
+++ b/Engine/Source/Runtime/Renderer/Private/HairStrands/HairStrandsVisibility.cpp
@@ -2601,7 +2601,7 @@ static void AddHairAuxilaryPass(
 	const bool bUseTile = TileData.IsValid();
 
 	GraphBuilder.AddPass(
-		RDG_EVENT_NAME("ComputeHairAuxilary(%s)(%s)", Method, bUseTile ? TEXT("Tile") : TEXT("Screen")),
+		RDG_EVENT_NAME("HairStrands::AuxilaryPass(%s)(%s)", Method, bUseTile ? TEXT("Tile") : TEXT("Screen")),
 		Parameters,
 		ERDGPassFlags::Raster,
 		[Parameters, ScreenVertexShader, TileVertexShader, PixelShader, Viewport, Resolution, bUseTile, bDepthTested, TileType](FRHICommandList& RHICmdList)
