Fix a bright artifacts bug of clear coat on mobile vulkan.

diff --git a/Engine/Shaders/Private/MobileShadingModels.ush b/Engine/Shaders/Private/MobileShadingModels.ush
--- a/Engine/Shaders/Private/MobileShadingModels.ush
+++ b/Engine/Shaders/Private/MobileShadingModels.ush
@@ -306,9 +306,7 @@ FMobileDirectLighting MobileClearCoatBxDF(FMobileShadingModelContext ShadingMode
 
 	// Vis_SmithJointApprox
 	half a = ClearCoatRoughness * ClearCoatRoughness;
-	half Vis_SmithV = NoL * (ShadingModelContext.NoV * (1 - a) + a);
-	half Vis_SmithL = ShadingModelContext.NoV * (NoL * (1 - a) + a);
-	float Vis = 0.5 * rcp(Vis_SmithV + Vis_SmithL);
+	float Vis = Vis_SmithJointApprox_Mobile(a, NoL, ShadingModelContext.NoV);
 
 	Lighting.Specular = NoL * ShadingModelContext.ClearCoat * F * Vis * GGX_Mobile(ClearCoatRoughness, NoH);
 
