Local merge of landscape edit layers : complete refactor of the merge algorithm that doesn't require allocating several render targets the size of the entire loaded landscape components. Instead, each landscape component is merged individually and heightmaps/weightmaps are then resolved in several small batches. This allows working on large and/or partially loaded landscapes without overtaxing GPU memory.
* The new merge algorithm also speeds up the merging process by performing the merge operation in a single shader pass, which avoids intermediate clear/copies, even more so when dealing with several edit layers (O(n) -> O(1))
* Implemented via FRDGBuilder to benefit from transient memory allocation and optimal resource transitions.
* This is activated by a CVar (landscape.EditLayersLocalMerge.Enable, disabled by default) and disabled for landscapes containing BP brushes, since customers (similar to what the WaterBrushManager did) might have made the assumption that the merge is global and we don't want to break the existing landscapes.

For heightmaps:
* The algorithm is divided into several "component resolve batches" that render a list of neighboring components and merge all their edit layers into a single slice of a scratch texture array (first, copy all the edit layers texture subregions to a texture array then merge them in a single shader pass). Then those rendered components' borders are stitched (in order to mimic the behavior of global merge) so that there are no seams in between components even if they were not properly authored (i.e. one component authored while its neighbor was not loaded). Then normals are generated in another pass.
* In order to limit the amount of temporary copies of heightmaps, batches only render a max number of landscape components (including neighbor components), controlled by CVar landscape.EditLayersLocalMerge.MaxComponentsPerHeightmapResolveBatch. This means that a given component can effectively be rendered (i.e. merged) multiple times if it's needed for the computation of different batches but batches are built in a way to minimize this situation.
* Finally, mips are generated and the resulting heightmaps are copied to staging textures the same way they used to.

For weightmaps:
* The algorithm is divided into several "weightmap resolve batches" that render a list of components that participate to the output of a list of given weightmaps, processing each individual channel into a single slice of a scratch texture array (first, copy all the edit layers texture channels to a texture array then merge them in a single shader pass). Then, the weightmaps to resolve within the batch are finalized (i.e. individual, single-channel, weightmaps are packed into the final multi-channel weightmap), their mips are generated and the resulting weightmaps are copied to staging textures the same way they used to.
* In order to limit the amount of temporary copies of weightmaps, batches only render a max number of individual weightmaps, controlled by CVar landscape.EditLayersLocalMerge.MaxComponentsPerWeightmapResolveBatch. This means that a given weightmap can effectively be rendered (i.e. merged) multiple times if it's needed for the computation of different batches but batches are built in a way to minimize this situation.

Misc:
* Fixed RHI validation for GlobalMerge of weightmaps
* Moved landscape render resources to their own file to de-clutter LandscapeEditLayers.cpp a bit
* landscape.RenderCaptureLayersNextHeightmapDraws and landscape.RenderCaptureLayersNextWeightmapDraws now support capturing several draw calls
* Moved all intermediate containers (LandscapeComponentHeightmapsToRender/ToResolve, MapHelper, etc.) to a context object in order to centralize the local/global merge common structures. Note : the context is mutable and can be partially refreshed since weightmaps can be potentially added by the merge algorithms and therefore those structures were not up to date and we could miss reading back from newly-added weightmap textures
* Removed unused WeightmapLayersBlendSubstractive in global merge algorithm
* CVar commands landscape.Dirty and landscape.FixSplines now process all loaded game worlds

The next step is to re-work our internal brushes and provide a new API for them and user-implemented ones to work in the same fashion. Then the old BP brush API will be progressively phased out.

#tests full QA pass (UEENGQA-60460)

diff --git a/Engine/Shaders/Private/Landscape/LandscapeLayersHeightmapsPS.usf b/Engine/Shaders/Private/Landscape/LandscapeLayersHeightmapsPS.usf
new file mode 100644
--- /dev/null
+++ b/Engine/Shaders/Private/Landscape/LandscapeLayersHeightmapsPS.usf
@@ -0,0 +1,494 @@
+// Copyright Epic Games, Inc. All Rights Reserved.
+
+#include "../Common.ush"
+
+// ----------------------------------------------------------------------------------
+
+#if defined (__INTELLISENSE__)
+// Uncomment the appropriate define for enabling syntax highlighting with HLSL Tools for Visual Studio : 
+//#define MERGE_EDIT_LAYERS 1
+#define STITCH_HEIGHTMAP 1
+#define FINALIZE_HEIGHTMAP 1
+//#define GENERATE_MIPS 1
+#endif // defined (__INTELLISENSE__)
+
+#define INDEX_NONE -1
+
+// ----------------------------------------------------------------------------------
+
+#if MERGE_EDIT_LAYERS
+// MergeEditLayers inputs/outputs :
+
+/** EEditLayerHeightmapBlendMode enum */
+#define EEDITLAYERHEIGHTMAPBLENDMODE_ADDITIVE 0
+#define EEDITLAYERHEIGHTMAPBLENDMODE_ALPHABLEND 1
+
+// Contains the information related to a given edit layers and how it should merge into the final heightmap
+struct FEditLayerHeightmapMergeInfo
+{
+	int4 TextureSubregion; // Subregion of the source (edit layer) texture to use
+	
+	uint BlendMode; // See EEditLayerHeightmapBlendMode
+	float Alpha; // Alpha value to be used in the blend
+	float Padding0; // Align to next float4
+	float Padding1;
+};
+
+uint InNumEditLayers; // Number of layers to merge into the destination texture
+Texture2DArray<float4> InEditLayersTextures; // Source texture for each individual edit layer (effective size == InNumEditLayers)
+
+StructuredBuffer<FEditLayerHeightmapMergeInfo> InEditLayersMergeInfos; // Describes how to merge each individual edit layer (effective size == InNumEditLayers)
+#endif // MERGE_EDIT_LAYERS
+
+#if STITCH_HEIGHTMAP || FINALIZE_HEIGHTMAP
+// StitchHeightmap/FinalizeHeightmaps shared inputs/outputs :
+
+uint2 InSourceTextureSize; // Texture size we read from (XY texture size of InSourceHeightmaps)
+
+// Texture array containing all the source heightmaps that can participate to the stitching operation (i.e. the component's heightmap, its 8 neighbors and potentially many other component 
+//  that are needed for this batch). 
+Texture2DArray<float2> InSourceHeightmaps;
+
+// The index of the input heightmap to process and its 8 neighbors arranged in a 3x3 array (linear : starting from the top-left neighbor [index 0] and ending with the bottom-right [index 8])
+//   (index relative to InSourceHeightmaps)
+DECLARE_SCALAR_ARRAY(uint, InNeighborHeightmapIndices, 9);
+#endif // STITCH_HEIGHTMAP || FINALIZE_HEIGHTMAP
+
+#if FINALIZE_HEIGHTMAP
+// FinalizeHeightmap inputs/outputs :
+
+uint4 InDestinationTextureSubregion; // Subregion of the destination texture to write to
+float3 InLandscapeGridScale; // x == LS Actor DrawScale.X, y == LS Actor DrawScale.y, z == LS Actor DrawScale.z / 128.0f (ZSCALE)
+#endif // FINALIZE_HEIGHTMAP
+
+
+#if GENERATE_MIPS
+// GenerateMips inputs/outputs :
+uint2 InCurrentMipSubregionSize; // Size of the the texture subregion that corresponds to a single landscape component (texture sharing)
+
+Texture2D<float4> InSourceHeightmap;
+#endif // GENERATE_MIPS
+
+
+// ----------------------------------------------------------------------------------
+// Util functions : 
+
+float2 PackHeight(float InExtractedHeight)
+{
+	InExtractedHeight = clamp(InExtractedHeight, 0.0f, 65535.0f);
+	int iHeight = (int) InExtractedHeight;
+	return float2((float) ((iHeight - (iHeight & 255)) >> 8) / 255.0, (float) (iHeight & 255) / 255.0);
+}
+
+float UnpackHeight(float2 InPackedHeight)
+{
+	return float(((int) round(InPackedHeight.r * 255.0) << 8) | (int) round(InPackedHeight.g * 255.0));
+}
+
+float ExtractAlpha(float2 InPackedAlpha)
+{
+	int AlphaWithFlags = ((int) round(InPackedAlpha.r * 255.0f) << 8) | (int) round(InPackedAlpha.g * 255.0f);
+	int Alpha = AlphaWithFlags & 0xFFFC; // Last 2 bits are used to store lower/raise flags
+	return float(Alpha) / 65532.0f; // 0xFFFC in float
+}
+
+#if MERGE_EDIT_LAYERS
+float4 SampleEditLayerHeightmap(uint2 InRelativeTextureCoordinates, uint InIndex)
+{
+	FEditLayerHeightmapMergeInfo LayerMergeInfo = InEditLayersMergeInfos[InIndex];
+	int4 LoadCoordinates = int4(InRelativeTextureCoordinates, InIndex, 0); // xy = relative coordinates, z = index in texture array, w = mip level
+	LoadCoordinates.xy += LayerMergeInfo.TextureSubregion.xy;
+	return InEditLayersTextures.Load(LoadCoordinates);
+}
+#endif // #if MERGE_EDIT_LAYERS
+
+#if STITCH_HEIGHTMAP || FINALIZE_HEIGHTMAP
+// InOffset between (-1,-1) and (1,1) : output between 0 (TL) and 8 (BR) :
+uint NeighborOffsetToNeighborLinearIndex(int2 InOffset)
+{
+	int2 PositiveOffset = InOffset + 1;
+	return PositiveOffset.y * 3 + PositiveOffset.x;
+}
+
+// InLinearIndex between 0 (TL) and 8 (BR) : output between between (-1,-1) and (1,1) :
+int2 NeighborLinearIndexToOffset(uint InLinearIndex)
+{
+	int2 NeighborOffsets[9] = { int2(-1, -1), int2(+0, -1), int2(+1, -1), int2(-1, +0), int2(+0, +0), int2(+1, +0), int2(-1, +1), int2(+0, +1), int2(+1, +1) };
+	return NeighborOffsets[InLinearIndex];
+}
+
+// Load the appropriate heightmap at the given location and returns the position and packed height:
+//  (handles the edge cases when we need to sample from neighboring heightmaps) :
+//  bSkipDuplicateRowColumn allows to skip the first/last row/column on the borders, since they're duplicates of the center texture :
+//  Returns the unpacked height (.x) and whether the sample is valid or not (.y = 1 if valid, 0 otherwise)
+float2 LoadNeighborHeightmapUnpackedHeight(bool2 bInIsMinComponentBorder, bool2 bInIsMaxComponentBorder, uint2 InRelativeTextureCoordinates, int2 InNeighborOffset, bool bSkipDuplicateRowColumn)
+{
+	// Which of the 9 input textures should be sampled (as a texture index offset between (-1, -1) and (1, 1)):
+	// In the default case (not on a texture border), we sample only the center texture so no texture offset needed:
+	int2 TextureIndexOffset = InNeighborOffset;
+	TextureIndexOffset.x *= (!bInIsMinComponentBorder.x && (TextureIndexOffset.x == -1)) ? 0 : 1;
+	TextureIndexOffset.y *= (!bInIsMinComponentBorder.y && (TextureIndexOffset.y == -1)) ? 0 : 1;
+	TextureIndexOffset.x *= (!bInIsMaxComponentBorder.x && (TextureIndexOffset.x == +1)) ? 0 : 1;
+	TextureIndexOffset.y *= (!bInIsMaxComponentBorder.y && (TextureIndexOffset.y == +1)) ? 0 : 1;
+
+	// Retrieve the linear index between 0 (TL) and 8 (BR) that corresponds to this neighbor:
+	uint NeighborLinearIndex = NeighborOffsetToNeighborLinearIndex(TextureIndexOffset);
+	uint NeighborHeightmapIndex = GET_SCALAR_ARRAY_ELEMENT(InNeighborHeightmapIndices, NeighborLinearIndex);
+	bool bIsValidNeighbor = (NeighborHeightmapIndex != INDEX_NONE);
+	float Height = 0.0f;
+	if (bIsValidNeighbor)
+	{
+		int2 SampleCoordinates = InRelativeTextureCoordinates + InNeighborOffset;
+		if (TextureIndexOffset.x == -1)
+		{
+			// If we use the textures from the left, we need to skip the rightmost column as it's a duplicate of our the center texture's leftmost column :
+			SampleCoordinates.x = InSourceTextureSize.x - 1 - (bSkipDuplicateRowColumn ? 1 : 0);
+		}
+		if (TextureIndexOffset.x == +1)
+		{
+			// If we use the textures from the right, we need to skip the leftmost column as it's a duplicate of our the center texture's rightmost column :
+			SampleCoordinates.x = bSkipDuplicateRowColumn ? +1 : 0;
+		}
+		if (TextureIndexOffset.y == -1)
+		{
+			// If we use the textures from the top, we need to skip the downmost row as it's a duplicate of our the center texture's upmost row :
+			SampleCoordinates.y = InSourceTextureSize.y - 1 - (bSkipDuplicateRowColumn ? 1 : 0);
+		}
+		if (TextureIndexOffset.y == +1)
+		{
+			// If we use the textures from the bottom, we need to skip the upmost row as it's a duplicate of our the center texture's downmost row :
+			SampleCoordinates.y = bSkipDuplicateRowColumn ? +1 : 0;
+		}
+
+		int4 LoadCoordinates = int4(SampleCoordinates, NeighborHeightmapIndex, 0); // xy = relative coordinates, z = index in texture array, w = mip level
+		float2 PackedHeight = InSourceHeightmaps.Load(LoadCoordinates);
+		Height = UnpackHeight(PackedHeight);
+	}
+	return float2(Height, bIsValidNeighbor ? 1.0f : 0.0f);
+}
+
+// Load from the main heightmap and returns the unpacked height (.x) and whether the sample is valid or not (.y = 1 since the main heightmap is always valid):
+float2 LoadHeightmapUnpackedHeight(uint2 InRelativeTextureCoordinates)
+{
+	// Load from the main input/output heightmap (linear neighbor index 4 == CC == center point) :
+	uint HeightmapIndex = GET_SCALAR_ARRAY_ELEMENT(InNeighborHeightmapIndices, 4);
+	int4 LoadCoordinates = int4(InRelativeTextureCoordinates, HeightmapIndex, 0); // xy = relative coordinates, z = index in texture array, w = mip level
+	float2 PackedHeight = InSourceHeightmaps.Load(LoadCoordinates);
+	return float2(UnpackHeight(PackedHeight), 1.0f);
+}
+#endif // STITCH_HEIGHTMAP || FINALIZE_HEIGHTMAP
+
+#if FINALIZE_HEIGHTMAP
+float ScaleHeight(float InHeight)
+{
+	return (InHeight - 32768.0f) * InLandscapeGridScale.z * TERRAIN_ZSCALE;
+}
+
+float UnscaleHeight(float InScaledHeight)
+{
+	return (InScaledHeight / InLandscapeGridScale.z / TERRAIN_ZSCALE) + 32768.0f;
+}
+
+// Load from the main heightmap and returns the position (.xyz) and whether the sample is valid or not (.w = 1 because the main heightmap is always valid)
+float4 LoadHeightmapScaledPosition(uint2 InRelativeTextureCoordinates)
+{
+	float Height = LoadHeightmapUnpackedHeight(InRelativeTextureCoordinates);
+	return float4(InRelativeTextureCoordinates * InLandscapeGridScale.xy, ScaleHeight(Height), 1.0f);
+}
+
+//  Returns the position (.xyz) and whether the sample is valid or not (.w = 1 if valid, 0 otherwise)
+float4 LoadNeighborHeightmapScaledPosition(bool2 bInIsMinComponentBorder, bool2 bInIsMaxComponentBorder, uint2 InRelativeTextureCoordinates, int2 InNeighborOffset, bool bSkipDuplicateRowColumn)
+{
+	float2 NeighborHeight = LoadNeighborHeightmapUnpackedHeight(bInIsMinComponentBorder, bInIsMaxComponentBorder, InRelativeTextureCoordinates, InNeighborOffset, bSkipDuplicateRowColumn);
+	return float4(((int2) InRelativeTextureCoordinates + InNeighborOffset) * InLandscapeGridScale.xy, ScaleHeight(NeighborHeight.x), NeighborHeight.y);
+}
+
+// Compute the normal of the triangle formed by the 3 points (in winding order). Additionally, the .w component of each point indicates its validity and the result will be 0 if either of the points is invalid :
+float3 ComputeNullableTriangleNormal(float4 InPoint0, float4 InPoint1, float4 InPoint2)
+{
+	float3 Normal = normalize(cross(InPoint0.xyz - InPoint1.xyz, InPoint1.xyz - InPoint2.xyz));
+	return Normal * InPoint0.w * InPoint1.w * InPoint2.w;
+}
+#endif // FINALIZE_HEIGHTMAP
+
+
+// ----------------------------------------------------------------------------------
+// Compute shaders : 
+
+#if MERGE_EDIT_LAYERS
+// Take all source edit layers heightmap regions and merge them all into a (almost) final heightmap
+//  Note : reads from a potentially shared heightmaps (i.e. from subregions of the heightmaps) and writes to R8G8 non-shared scratch texture
+void MergeEditLayers(in float4 SVPos : SV_POSITION, out float2 OutColor : SV_Target0)
+{
+	uint2 RelativeTextureCoordinates = floor(SVPos.xy);
+		
+	// Take the base height straight from the first layer : 
+	float4 EditLayerSample = SampleEditLayerHeightmap(RelativeTextureCoordinates, 0);
+	float Height = UnpackHeight(EditLayerSample.xy) - 32768.0f;
+		
+	// Then merge the subsequent layers according to their merge info : 
+	LOOP
+	for (uint i = 1; i < InNumEditLayers; ++i)
+	{
+		FEditLayerHeightmapMergeInfo LayerMergeInfo = InEditLayersMergeInfos[i];
+		EditLayerSample = SampleEditLayerHeightmap(RelativeTextureCoordinates, i);
+		float LayerHeight = UnpackHeight(EditLayerSample.xy) - 32768.0f;
+		LayerHeight *= LayerMergeInfo.Alpha;
+			
+		if (LayerMergeInfo.BlendMode == EEDITLAYERHEIGHTMAPBLENDMODE_ADDITIVE)
+		{
+			Height += LayerHeight;
+		}
+		else if (LayerMergeInfo.BlendMode == EEDITLAYERHEIGHTMAPBLENDMODE_ALPHABLEND)
+		{
+			float LayerHeightAlphaBlend = ExtractAlpha(EditLayerSample.ba);
+			float NewHeight = LayerHeight + (LayerHeightAlphaBlend * Height);
+			uint RaiseLower = floor(EditLayerSample.a * 255.f + 0.5f);
+			bool bLowerTerrain = (RaiseLower & 1);
+			bool bRaiseTerrain = (RaiseLower & 2);
+			if ((bRaiseTerrain && NewHeight > Height) || (bLowerTerrain && NewHeight < Height))
+			{
+				Height = NewHeight;
+			}
+			else
+			{
+				Height = Height;
+			}
+		}
+	}
+		
+	Height = clamp(Height, -32768.0, 32767.0);
+	Height += 32768.0f;
+	OutColor = PackHeight(Height);
+}
+#endif // MERGE_EDIT_LAYERS
+
+
+// ----------------------------------------------------------------------------------
+
+#if STITCH_HEIGHTMAP
+
+void StitchHeightmap(in float4 SVPos : SV_POSITION, out float2 OutColor : SV_Target0)
+{
+	uint2 RelativeTextureCoordinates = floor(SVPos.xy);
+
+	bool2 bIsMinComponentBorder = (RelativeTextureCoordinates == 0);
+	bool2 bIsMaxComponentBorder = (RelativeTextureCoordinates == (InSourceTextureSize - 1));
+
+	// On borders, average the neighbors' unpacked height so that in order to repair potential discrepancies between 2 neighboring rows/columns from different textures (effectively stitching them together) :
+	//  They should mostly be identical, but because it's possible that one neighbor component is not loaded while one is being recomputed, we might end up with different values there :
+	float2 TL = LoadNeighborHeightmapUnpackedHeight(bIsMinComponentBorder, bIsMaxComponentBorder, RelativeTextureCoordinates, /*NeighborOffset = */int2(-1, -1), /*bSkipDuplicateRowColumn = */false);
+	float2 TT = LoadNeighborHeightmapUnpackedHeight(bIsMinComponentBorder, bIsMaxComponentBorder, RelativeTextureCoordinates, /*NeighborOffset = */int2(+0, -1), /*bSkipDuplicateRowColumn = */false);
+	float2 TR = LoadNeighborHeightmapUnpackedHeight(bIsMinComponentBorder, bIsMaxComponentBorder, RelativeTextureCoordinates, /*NeighborOffset = */int2(+1, -1), /*bSkipDuplicateRowColumn = */false);
+	float2 LL = LoadNeighborHeightmapUnpackedHeight(bIsMinComponentBorder, bIsMaxComponentBorder, RelativeTextureCoordinates, /*NeighborOffset = */int2(-1, +0), /*bSkipDuplicateRowColumn = */false);
+	float2 CC = LoadHeightmapUnpackedHeight(RelativeTextureCoordinates);
+	float2 RR = LoadNeighborHeightmapUnpackedHeight(bIsMinComponentBorder, bIsMaxComponentBorder, RelativeTextureCoordinates, /*NeighborOffset = */int2(+1, +0), /*bSkipDuplicateRowColumn = */false);
+	float2 BL = LoadNeighborHeightmapUnpackedHeight(bIsMinComponentBorder, bIsMaxComponentBorder, RelativeTextureCoordinates, /*NeighborOffset = */int2(-1, +1), /*bSkipDuplicateRowColumn = */false);
+	float2 BB = LoadNeighborHeightmapUnpackedHeight(bIsMinComponentBorder, bIsMaxComponentBorder, RelativeTextureCoordinates, /*NeighborOffset = */int2(+0, +1), /*bSkipDuplicateRowColumn = */false);
+	float2 BR = LoadNeighborHeightmapUnpackedHeight(bIsMinComponentBorder, bIsMaxComponentBorder, RelativeTextureCoordinates, /*NeighborOffset = */int2(+1, +1), /*bSkipDuplicateRowColumn = */false);
+	
+	// .x = scaled height, .y = whether the sample is valid or not :
+	float2 ValidScaledHeights[9] = { TL, TT, TR, LL, CC, RR, BL, BB, BR };
+
+	// 9 possible cases : only pixels on the border (A..H) need stitching
+	//    bIsMinComponentBorder.x = true
+	//   |                bIsMaxComponentBorder.x = true
+	//   |               |
+	//   v               v
+	// +---+   +---+   +---+
+	// | A |...| B |...| C | <-- bIsMinComponentBorder.y = true
+	// +---+   +---+   +---+
+	//  ...     ...     ...
+	// +---+   +---+   +---+
+	// | H |...| I |...| D | 
+	// +---+   +---+   +---+
+	//  ...     ...     ...
+	// +---+   +---+   +---+
+	// | G |...| F |...| E | <-- bIsMaxComponentBorder.y = true
+	// +---+   +---+   +---+
+
+	// Truth table (x == True) that indicates which neigbor sample (TL, TT, ...) is needed for each of the 9 possible cases (whether the pixel is on a border or not) : 
+	//                         | A | B | C | D | E | F | G | H | I |
+	//                         +---+---+---+---+---+---+---+---+---+
+	//                      TL | x |   |   |   |   |   |   |   |   |
+	//                      TT | x | x | x |   |   |   |   |   |   |
+	//                      TR |   |   | x |   |   |   |   |   |   |
+	//                      LL | x |   |   |   |   |   | x | x |   |
+	//                      CC | x | x | x | x | x | x | x | x | x |
+	//                      RR |   |   | x | x | x |   |   |   |   |
+	//                      BL |   |   |   |   |   |   | x |   |   |
+	//                      BB |   |   |   |   | x | x | x |   |   |
+	//                      BR |   |   |   |   | x |   |   |   |   |
+	//                         +---+---+---+---+---+---+---+---+---+
+	// bIsMinComponentBorder.x | x |   |   |   |   |   | x | x |   |
+	// bIsMinComponentBorder.y | x | x | x |   |   |   |   |   |   |
+	// bIsMaxComponentBorder.x |   |   | x | x | x |   |   |   |   |
+	// bIsMaxComponentBorder.y |   |   |   |   | x | x | x |   |   |
+	//                         +---+---+---+---+---+---+---+---+---+
+
+	bool bA = bIsMinComponentBorder.x && bIsMinComponentBorder.y;
+	bool bB = bIsMinComponentBorder.y && !bIsMinComponentBorder.x && !bIsMaxComponentBorder.x;
+	bool bC = bIsMaxComponentBorder.x && bIsMinComponentBorder.y;
+	bool bD = bIsMaxComponentBorder.x && !bIsMinComponentBorder.y && !bIsMaxComponentBorder.y;
+	bool bE = bIsMaxComponentBorder.x && bIsMaxComponentBorder.y;
+	bool bF = bIsMaxComponentBorder.y && !bIsMinComponentBorder.x && !bIsMaxComponentBorder.x;
+	bool bG = bIsMinComponentBorder.x && bIsMaxComponentBorder.y;
+	bool bH = bIsMinComponentBorder.x && !bIsMinComponentBorder.y && !bIsMaxComponentBorder.y;
+
+	// Describes whether to consider each neighbor sample depending on whether we're a pixel on the border or not :
+	bool bValidHeights[9] =
+	{
+		/* TL = */bA,
+		/* TT = */bA || bB || bC,
+		/* TR = */bC,
+		/* LL = */bA || bG || bH,
+		/* CC = */true, 
+		/* RR = */bC || bD || bE, 
+		/* BL = */bG,
+		/* BB = */bE || bF || bG,
+		/* BR = */bE, 
+	};
+	
+	float FinalHeight = 0.0;
+	float WeightSum = 0.0f;
+	UNROLL
+	for (int Index = 0; Index < 9; ++Index)
+	{
+		float Height = ValidScaledHeights[Index].x;
+		float Weight = bValidHeights[Index] ? 1.0f : 0.0f;
+		Weight *= ValidScaledHeights[Index].y; // .y of ValidScaledHeights indicates whether the sample is valid or not (i.e. whether there's a neighboring heighmap or not)
+		FinalHeight += Height * Weight;
+		WeightSum += Weight;
+	}
+	FinalHeight /= WeightSum;
+
+	OutColor = PackHeight(FinalHeight);
+}
+
+#endif // STITCH_HEIGHTMAP
+
+// ----------------------------------------------------------------------------------
+
+#if FINALIZE_HEIGHTMAP
+
+void FinalizeHeightmap(in float4 SVPos : SV_POSITION, out float4 OutColor : SV_Target0)
+{
+	uint2 DestinationTextureSize = InDestinationTextureSubregion.zw - InDestinationTextureSubregion.xy;
+	uint2 RelativeTextureCoordinates = floor(SVPos.xy) - InDestinationTextureSubregion.xy;
+
+	bool2 bIsMinComponentBorder = (RelativeTextureCoordinates == 0);
+	bool2 bIsMaxComponentBorder = (RelativeTextureCoordinates == (DestinationTextureSize - 1));
+		
+	// The triangle topology is the following (where C = center, T = top, B = bottom, L = left, R = right and Nx the normals we need to interpolate):
+	// TL ------ TT
+	// | \       | \
+	// |  \      |  \
+	// |   \     |   \
+	// | N0 \ N1 | N3 \
+	// |     \   |     \
+	// |      \  |      \
+	// |       \ |       \
+	// LL ------ CC ------ RR
+	//   \       | \       |
+	//    \      |  \      |
+	//     \     |   \     |
+	//      \ N2 | N4 \ N5 |
+	//       \   |     \   |
+	//        \  |      \  |
+	//         \ |       \ |
+	//           BB ------ BR
+
+	float4 CC = LoadHeightmapScaledPosition(RelativeTextureCoordinates);
+	float4 TL = LoadNeighborHeightmapScaledPosition(bIsMinComponentBorder, bIsMaxComponentBorder, RelativeTextureCoordinates, /*NeighborOffset = */int2(-1, -1), /*bSkipDuplicateRowColumn = */true);
+	float4 TT = LoadNeighborHeightmapScaledPosition(bIsMinComponentBorder, bIsMaxComponentBorder, RelativeTextureCoordinates, /*NeighborOffset = */int2(+0, -1), /*bSkipDuplicateRowColumn = */true);
+	float4 LL = LoadNeighborHeightmapScaledPosition(bIsMinComponentBorder, bIsMaxComponentBorder, RelativeTextureCoordinates, /*NeighborOffset = */int2(-1, +0), /*bSkipDuplicateRowColumn = */true);
+	float4 RR = LoadNeighborHeightmapScaledPosition(bIsMinComponentBorder, bIsMaxComponentBorder, RelativeTextureCoordinates, /*NeighborOffset = */int2(+1, +0), /*bSkipDuplicateRowColumn = */true);
+	float4 BR = LoadNeighborHeightmapScaledPosition(bIsMinComponentBorder, bIsMaxComponentBorder, RelativeTextureCoordinates, /*NeighborOffset = */int2(+1, +1), /*bSkipDuplicateRowColumn = */true);
+	float4 BB = LoadNeighborHeightmapScaledPosition(bIsMinComponentBorder, bIsMaxComponentBorder, RelativeTextureCoordinates, /*NeighborOffset = */int2(+0, +1), /*bSkipDuplicateRowColumn = */true);
+
+	float3 N0 = ComputeNullableTriangleNormal(CC, LL, TL);
+	float3 N1 = ComputeNullableTriangleNormal(TL, TT, CC);
+	float3 N2 = ComputeNullableTriangleNormal(LL, CC, BB);
+	float3 N3 = ComputeNullableTriangleNormal(RR, CC, TT);
+	float3 N4 = ComputeNullableTriangleNormal(BR, BB, CC);
+	float3 N5 = ComputeNullableTriangleNormal(CC, RR, BR);
+	
+	float3 Normal = normalize(N0 + N1 + N2 + N3 + N4 + N5);
+	// Scale back to be 0-1 normal
+	Normal = 0.5f * Normal + 0.5f;
+	
+	float FinalHeight = UnscaleHeight(CC.z);
+	OutColor = float4(PackHeight(FinalHeight), Normal.xy);
+}
+#endif // FINALIZE_HEIGHTMAP
+
+
+// ----------------------------------------------------------------------------------
+#if GENERATE_MIPS
+void GenerateMips(in float4 SVPos : SV_POSITION, out float4 OutColor : SV_Target0)
+{
+	uint2 TextureCoordinates = floor(SVPos.xy);
+	
+	float4 SourceSamples[4] =
+	{
+		InSourceHeightmap.Load(int3(2 * TextureCoordinates + int2(+0, +0), 0)),
+		InSourceHeightmap.Load(int3(2 * TextureCoordinates + int2(+1, +0), 0)),
+		InSourceHeightmap.Load(int3(2 * TextureCoordinates + int2(+0, +1), 0)),
+		InSourceHeightmap.Load(int3(2 * TextureCoordinates + int2(+1, +1), 0)),
+	};
+		
+	// Because the borders of each landscape components are shared between neighbors, we must ensure that the parent mip's inner row/column of pixels don't contribute : 
+	// 9 possible cases (only the samples marked with a * must be kept): 
+	//      bIsMinBorder.x = true
+	//     |                        bIsMaxBorder.x = true
+	//     |                       |
+	//     v                       v
+	// +-------+   +-------+   +-------+
+	// | * :   |   | * : * |   |   : * |
+	// | - + - |...| - + - |...| - + - | <-- bIsMinBorder.y = true
+	// |   :   |   |   :   |   |   :   |
+	// +-------+   +-------+   +-------+
+	//    ...         ...	      ...
+	// +-------+   +-------+   +-------+
+	// | * :   |   | * : * |   |   : * |
+	// | - + - |...| - + - |...| - + - |
+	// | * :   |   | * : * |   |   : * |
+	// +-------+   +-------+   +-------+
+	//    ...         ...	      ...
+	// +-------+   +-------+   +-------+
+	// |   :   |   |   :   |   |   :   |
+	// | - + - |...| - + - |...| - + - | <-- bIsMaxBorder.y = true
+	// | * :   |   | * : * |   |   : * | 
+	// +-------+   +-------+   +-------+
+
+	uint2 RelativeTextureCoordinates = TextureCoordinates % InCurrentMipSubregionSize;
+	bool2 bIsMinBorder = (RelativeTextureCoordinates == 0);
+	bool2 bIsMaxBorder = (RelativeTextureCoordinates == (InCurrentMipSubregionSize - 1));
+	bool bIsLastMip = all(InCurrentMipSubregionSize == 1);
+		
+	float SampleWeights[4] =
+	{ 
+		// On the last mip, it's ok to keep all 4 samples : all neighbors components share them : 
+		((bIsMaxBorder.x || bIsMaxBorder.y) && !bIsLastMip) ? 0.0f : 1.0f,
+		((bIsMinBorder.x || bIsMaxBorder.y) && !bIsLastMip) ? 0.0f : 1.0f,
+		((bIsMaxBorder.x || bIsMinBorder.y) && !bIsLastMip) ? 0.0f : 1.0f,
+		((bIsMinBorder.x || bIsMinBorder.y) && !bIsLastMip) ? 0.0f : 1.0f,
+	};
+	
+	float TotalSampleWeight = 0.0f;
+	OutColor = 0.0f;
+	
+	float3 HeightAndNormal = 0.0f;
+	
+	UNROLL
+	for (int i = 0; i < 4; ++i)
+	{
+		HeightAndNormal += float3(UnpackHeight(SourceSamples[i].xy), SourceSamples[i].zw) * SampleWeights[i];
+		TotalSampleWeight += SampleWeights[i];
+	}
+	HeightAndNormal /= TotalSampleWeight;
+	
+	OutColor.xy = PackHeight(HeightAndNormal.x);
+	OutColor.zw = HeightAndNormal.yz;
+}
+#endif // GENERATE_MIPS
\ No newline at end of file
diff --git a/Engine/Shaders/Private/Landscape/LandscapeLayersWeightmapsPS.usf b/Engine/Shaders/Private/Landscape/LandscapeLayersWeightmapsPS.usf
new file mode 100644
--- /dev/null
+++ b/Engine/Shaders/Private/Landscape/LandscapeLayersWeightmapsPS.usf
@@ -0,0 +1,287 @@
+// Copyright Epic Games, Inc. All Rights Reserved.
+
+#include "../Common.ush"
+
+// ----------------------------------------------------------------------------------
+
+#if defined (__INTELLISENSE__)
+// Uncomment the appropriate define for enabling syntax highlighting with HLSL Tools for Visual Studio : 
+//#define MERGE_EDIT_LAYERS 1
+#define FINALIZE_WEIGHTMAP 1
+//#define GENERATE_MIPS 1
+#endif // defined (__INTELLISENSE__)
+
+
+// ----------------------------------------------------------------------------------
+
+#if MERGE_EDIT_LAYERS
+// MergeEditLayers inputs/outputs :
+/** EEditLayerWeightmapBlendMode enum */
+#define EEDITLAYERWEIGHTMAPBLENDMODE_ADDITIVE 0
+#define EEDITLAYERWEIGHTMAPBLENDMODE_SUBSTRACTIVE 1
+
+// Contains the information related to a given edit layers and how it should merge into the final weightmap
+struct FEditLayerWeightmapMergeInfo
+{
+	uint SourceWeightmapTextureIndex; // The index in InPackedWeightmaps of the texture to read from for this layer
+	uint SourceWeightmapTextureChannel; // The channel of the texture to read from for this layer
+	uint BlendMode; // See EEditLayerWeightmapBlendMode
+	float Alpha; // Alpha value to be used in the blend
+};
+
+uint InNumEditLayers; // Number of layers to merge into the destination texture
+uint InStartIndexInEditLayersMergeInfos; // Start index for edit layers merge infos in the big buffer InEditLayersMergeInfos
+Texture2DArray<float4> InPackedWeightmaps; // Source -packed- RGBA weightmaps
+StructuredBuffer<FEditLayerWeightmapMergeInfo> InEditLayersMergeInfos; // Describes how to merge each individual edit layer (effective size == sum of InNumEditLayers of all passes)
+#endif // MERGE_EDIT_LAYERS
+
+
+#if FINALIZE_WEIGHTMAP
+/** EWeightmapPaintLayerFlags enum */
+#define EWEIGHTMAPPAINTLAYERFLAGS_ISVISIBILITYLAYER (1 << 0)
+#define EWEIGHTMAPPAINTLAYERFLAGS_ISWEIGHTBLENDED (1 << 1)
+
+// Contains the information related to a given edit layers and how it should merge into the final weightmap
+struct FWeightmapPaintLayerInfo
+{
+	uint Flags; // See EWeightmapPaintLayerFlags
+};
+
+// FinalizeWeightmap inputs/outputs :
+uint InValidTextureChannelsMask; // 4 bits-wide bitmask that indicates which of the 4 channels need to be processed 
+
+// Contains (per channel : .r for the first channel, .g for the 2nd, etc.) the index of the paint layer being processed for this channel. Note : this is an index for the channel's InPerChannelPaintLayerWeightmaps 
+//  (which is a subset of all paint layers), not InPaintLayerInfos (i.e. all paint layers)
+uint4 InPerChannelPaintLayerIndexInWeightmaps; 
+// Contains (per channel : .r for the first channel, .g for the 2nd, etc.) the start index for addressing the indices of actual paint layers contained in the channel's InPerChannelPaintLayerWeightmaps
+uint4 InPerChannelStartPaintLayerIndex; 
+// Contains (per channel : .r for the first channel, .g for the 2nd, etc.) the number of actual paint layers contained in the channel's InPerChannelPaintLayerWeightmaps
+uint4 InPerChannelNumPaintLayers; 
+
+// 4 input texture arrays (R8) (one per channel) we're about to pack on the output weightmap (RGBA), each one contains a list of all paint layers (after having been edit layer-merged) for a given component, 
+//  since we need to weight-blend some of them. For each channel, NumElements = InPerChannelNumPaintLayers[ChannelIndex]
+Texture2DArray<float> InPerChannelPaintLayerWeightmaps_0; 
+Texture2DArray<float> InPerChannelPaintLayerWeightmaps_1; 
+Texture2DArray<float> InPerChannelPaintLayerWeightmaps_2; 
+Texture2DArray<float> InPerChannelPaintLayerWeightmaps_3; 
+
+// Input uint buffers that contains all the indices of the paint layer in InPaintLayerInfos. This buffer is shared between all texture resolve passes and is therefore adressed using 
+//  InPerChannelStartPaintLayerIndex and InPerChannelNumPaintLayers (one entry per channel each)
+StructuredBuffer<uint> InPaintLayerInfoIndices;
+
+// Array that contains per-paint layer information (e.g. paint layer's blend mode)
+StructuredBuffer<FWeightmapPaintLayerInfo> InPaintLayerInfos; 
+#endif // FINALIZE_WEIGHTMAP
+
+
+#if GENERATE_MIPS
+// GenerateMips inputs/outputs :
+uint2 InCurrentMipSize; // Size of the texture mip currently being generated
+
+Texture2D<float4> InSourceWeightmap;
+#endif // GENERATE_MIPS
+
+
+// ----------------------------------------------------------------------------------
+// Util functions : 
+
+#if MERGE_EDIT_LAYERS
+// For a given edit layer, samples the appropriate texture in the input texture array and returns the appropriate channel, attenuated by the layer's alpha value : 
+float SampleEditLayerWeightmap(uint2 InSampleCoordinates, uint InIndex)
+{
+	FEditLayerWeightmapMergeInfo LayerMergeInfo = InEditLayersMergeInfos[InStartIndexInEditLayersMergeInfos + InIndex];
+	int4 LoadCoordinates = int4(InSampleCoordinates, LayerMergeInfo.SourceWeightmapTextureIndex, 0); // xy = relative coordinates, z = index in texture array, w = mip level
+	float4 Packed = InPackedWeightmaps.Load(LoadCoordinates);
+	float Weight = Packed[LayerMergeInfo.SourceWeightmapTextureChannel];
+	return Weight * LayerMergeInfo.Alpha;
+}
+#endif // #if MERGE_EDIT_LAYERS
+
+#if FINALIZE_WEIGHTMAP
+float LoadPaintLayerWeightmapForChannel(uint InChannelIndex, uint InPaintLayerIndex, uint2 InSampleCoordinates)
+{
+	int4 LoadCoordinates = int4(InSampleCoordinates, InPaintLayerIndex, 0); // xy = coordinates, z = index in texture array, w = mip level
+	// We have to use a switch case because we cannot sample a static array of textures without a literal expression:
+	switch (InChannelIndex)
+	{
+		case 0: return InPerChannelPaintLayerWeightmaps_0.Load(LoadCoordinates);
+		case 1: return InPerChannelPaintLayerWeightmaps_1.Load(LoadCoordinates);
+		case 2: return InPerChannelPaintLayerWeightmaps_2.Load(LoadCoordinates);
+		default: return InPerChannelPaintLayerWeightmaps_3.Load(LoadCoordinates);
+	}
+}
+
+uint LoadPaintLayerInfoIndexForChannel(uint InChannelIndex, uint InPaintLayerIndex)
+{
+	uint PaintLayerInfoIndexInBigBuffer = InPerChannelStartPaintLayerIndex[InChannelIndex] + InPaintLayerIndex;
+	return InPaintLayerInfoIndices[PaintLayerInfoIndexInBigBuffer];
+}
+
+bool IsWeightBlendedPaintLayer(FWeightmapPaintLayerInfo InPaintLayerInfo)
+{
+	// For the visibility layer, deactivate weight blending altogether : 
+	return (((InPaintLayerInfo.Flags & EWEIGHTMAPPAINTLAYERFLAGS_ISVISIBILITYLAYER) == 0)
+		&& (InPaintLayerInfo.Flags & EWEIGHTMAPPAINTLAYERFLAGS_ISWEIGHTBLENDED));
+}
+
+float FinalizeWeightmapChannel(uint2 InSampleCoordinates, uint InChannelIndex)
+{
+	float Result = 0.0f;
+	
+	if (InValidTextureChannelsMask & (1 << InChannelIndex))
+	{
+		uint PaintLayerIndexInWeightmaps = InPerChannelPaintLayerIndexInWeightmaps[InChannelIndex];
+		uint NumPaintLayers = InPerChannelNumPaintLayers[InChannelIndex];
+		
+		uint ActivePaintLayerInfoIndex = LoadPaintLayerInfoIndexForChannel(InChannelIndex, PaintLayerIndexInWeightmaps);
+		FWeightmapPaintLayerInfo ActivePaintLayerInfo = InPaintLayerInfos[ActivePaintLayerInfoIndex];
+
+		if (!IsWeightBlendedPaintLayer(ActivePaintLayerInfo))
+		{
+			return LoadPaintLayerWeightmapForChannel(InChannelIndex, PaintLayerIndexInWeightmaps, InSampleCoordinates);
+		}
+		
+		float ActivePaintLayerWeight = 0.0f;
+		float BlendedWeightsSum = 0.0f;
+
+		LOOP
+		for (int i = 0; i < NumPaintLayers; ++i)
+		{
+			bool bIsOutputPaintLayer = (i == PaintLayerIndexInWeightmaps);
+			uint PaintLayerInfoIndex = LoadPaintLayerInfoIndexForChannel(InChannelIndex, i);
+			FWeightmapPaintLayerInfo PaintLayerInfo = InPaintLayerInfos[PaintLayerInfoIndex];
+
+			// Only weight blended (and non-visibility) paint layers participate to weight blending : 
+			if (IsWeightBlendedPaintLayer(PaintLayerInfo)) 
+			{
+				float Weight = LoadPaintLayerWeightmapForChannel(InChannelIndex, i, InSampleCoordinates);
+				if (bIsOutputPaintLayer)
+				{
+					ActivePaintLayerWeight = Weight;
+				}
+				BlendedWeightsSum += Weight;
+			}
+		}
+
+		Result = (BlendedWeightsSum > 0.0f) ? saturate(ActivePaintLayerWeight / BlendedWeightsSum) : ActivePaintLayerWeight;
+	}
+
+	return Result;
+}
+
+#endif // FINALIZE_WEIGHTMAP
+
+
+// ----------------------------------------------------------------------------------
+// Compute shaders : 
+
+#if MERGE_EDIT_LAYERS
+void MergeEditLayers(in float4 SVPos : SV_POSITION, out float OutColor : SV_Target0)
+{
+	uint2 TextureCoordinates = floor(SVPos.xy);
+		
+	// Take the base weight straight from the first layer : 
+	float Weight = SampleEditLayerWeightmap(TextureCoordinates, 0);
+		
+	// Then merge the subsequent layers according to their merge info : 
+	LOOP
+	for (uint i = 1; i < InNumEditLayers; ++i)
+	{
+		FEditLayerWeightmapMergeInfo LayerMergeInfo = InEditLayersMergeInfos[i];
+		float LayerWeight = SampleEditLayerWeightmap(TextureCoordinates, i);
+			
+		if (LayerMergeInfo.BlendMode == EEDITLAYERWEIGHTMAPBLENDMODE_ADDITIVE)
+		{
+			Weight += LayerWeight;
+		}
+		else if (LayerMergeInfo.BlendMode == EEDITLAYERWEIGHTMAPBLENDMODE_SUBSTRACTIVE)
+		{
+			Weight -= LayerWeight;
+		}
+	}
+	
+	OutColor = Weight;
+}
+#endif // MERGE_EDIT_LAYERS
+
+
+// ----------------------------------------------------------------------------------
+
+#if FINALIZE_WEIGHTMAP
+void FinalizeWeightmap(in float4 SVPos : SV_POSITION, out float4 OutColor : SV_Target0)
+{
+	uint2 TextureCoordinates = floor(SVPos.xy);
+	OutColor = 0.0f;
+	
+	UNROLL
+	for (uint i = 0; i < 4; ++i)
+	{
+		OutColor[i] = FinalizeWeightmapChannel(TextureCoordinates, i);
+	}
+}
+#endif // FINALIZE_WEIGHTMAP
+
+
+// ----------------------------------------------------------------------------------
+#if GENERATE_MIPS
+void GenerateMips(in float4 SVPos : SV_POSITION, out float4 OutColor : SV_Target0)
+{
+	uint2 TextureCoordinates = floor(SVPos.xy);
+	
+	float4 SourceSamples[4] =
+	{
+		InSourceWeightmap.Load(int3(2 * TextureCoordinates + int2(+0, +0), 0)),
+		InSourceWeightmap.Load(int3(2 * TextureCoordinates + int2(+1, +0), 0)),
+		InSourceWeightmap.Load(int3(2 * TextureCoordinates + int2(+0, +1), 0)),
+		InSourceWeightmap.Load(int3(2 * TextureCoordinates + int2(+1, +1), 0)),
+	};
+		
+	// Because the borders of each landscape components are shared between neighbors, we must ensure that the parent mip's inner row/column of pixels don't contribute : 
+	// 9 possible cases (only the samples marked with a * must be kept): 
+	//      bIsMinBorder.x = true
+	//     |                        bIsMaxBorder.x = true
+	//     |                       |
+	//     v                       v
+	// +-------+   +-------+   +-------+
+	// | * :   |   | * : * |   |   : * |
+	// | - + - |...| - + - |...| - + - | <-- bIsMinBorder.y = true
+	// |   :   |   |   :   |   |   :   |
+	// +-------+   +-------+   +-------+
+	//    ...         ...	      ...
+	// +-------+   +-------+   +-------+
+	// | * :   |   | * : * |   |   : * |
+	// | - + - |...| - + - |...| - + - |
+	// | * :   |   | * : * |   |   : * |
+	// +-------+   +-------+   +-------+
+	//    ...         ...	      ...
+	// +-------+   +-------+   +-------+
+	// |   :   |   |   :   |   |   :   |
+	// | - + - |...| - + - |...| - + - |
+	// | * :   |   | * : * |   |   : * | <-- bIsMaxBorder.y = true
+	// +-------+   +-------+   +-------+
+
+	bool2 bIsMinBorder = (TextureCoordinates == 0);
+	bool2 bIsMaxBorder = (TextureCoordinates == (InCurrentMipSize - 1));
+	bool bIsLastMip = all(InCurrentMipSize == 1);
+		
+	float SampleWeights[4] =
+	{ 
+		// On the last mip, it's ok to keep all 4 samples : all neighbors components share them : 
+		((bIsMaxBorder.x || bIsMaxBorder.y) && !bIsLastMip) ? 0.0f : 1.0f,
+		((bIsMinBorder.x || bIsMaxBorder.y) && !bIsLastMip) ? 0.0f : 1.0f,
+		((bIsMaxBorder.x || bIsMinBorder.y) && !bIsLastMip) ? 0.0f : 1.0f,
+		((bIsMinBorder.x || bIsMinBorder.y) && !bIsLastMip) ? 0.0f : 1.0f,
+	};
+	
+	float TotalSampleWeight = 0.0f;
+	OutColor = 0.0f;
+	
+	UNROLL
+	for (int i = 0; i < 4; ++i)
+	{
+		OutColor += SourceSamples[i] * SampleWeights[i];
+		TotalSampleWeight += SampleWeights[i];
+	}
+	OutColor /= TotalSampleWeight;
+}
+#endif // GENERATE_MIPS
\ No newline at end of file
diff --git a/Engine/Source/Runtime/Landscape/Classes/Landscape.h b/Engine/Source/Runtime/Landscape/Classes/Landscape.h
--- a/Engine/Source/Runtime/Landscape/Classes/Landscape.h
+++ b/Engine/Source/Runtime/Landscape/Classes/Landscape.h
@@ -18,12 +18,25 @@ class FMaterialResource;
 struct FLandscapeEditLayerReadbackResult;
 struct FNotificationInfo;
 struct FTextureToComponentHelper;
+struct FUpdateLayersContentContext;
+struct FEditLayersHeightmapMergeParams;
+struct FEditLayersWeightmapMergeParams;
 
 namespace ELandscapeToolTargetType
 {
 	enum Type : int8;
 };
 
+namespace EditLayersHeightmapLocalMerge_RenderThread
+{
+	struct FMergeInfo;
+}
+
+namespace EditLayersWeightmapLocalMerge_RenderThread
+{
+	struct FMergeInfo;
+}
+
 #if WITH_EDITOR
 extern LANDSCAPE_API TAutoConsoleVariable<int32> CVarLandscapeSplineFalloffModulation;
 #endif
@@ -334,15 +347,21 @@ public:
 #endif // WITH_EDITOR
 
 private:
+	bool SupportsEditLayersLocalMerge();
 	void CreateLayersRenderingResource();
-	void GetLandscapeComponentNeighborsToRender(ULandscapeComponent* LandscapeComponent, TSet<ULandscapeComponent*>& NeighborComponents) const;
-	void GetLandscapeComponentWeightmapsToRender(ULandscapeComponent* LandscapeComponent, TSet<ULandscapeComponent*>& WeightmapComponents) const;
+	void PrepareEditLayersLocalMergeResources();
 	void UpdateLayersContent(bool bInWaitForStreaming = false, bool bInSkipMonitorLandscapeEdModeChanges = false, bool bIntermediateRender = false, bool bFlushRender = false);
 	void MonitorShaderCompilation();
 	void MonitorLandscapeEdModeChanges();
-	int32 RegenerateLayersHeightmaps(FTextureToComponentHelper const& MapHelper, const TArray<ULandscapeComponent*>& InLandscapeComponents, const TArray<ULandscapeComponent*>& InLandscapeComponentsToResolve);
-	int32 RegenerateLayersWeightmaps(FTextureToComponentHelper const& MapHelper, const TArray<ULandscapeComponent*>& InLandscapeComponents, const TArray<ULandscapeComponent*>& InLandscapeComponentsToResolve);
+	
+	int32 RegenerateLayersHeightmaps(const FUpdateLayersContentContext& InUpdateLayersContentContext);
+	int32 PerformLayersHeightmapsLocalMerge(const FUpdateLayersContentContext& InUpdateLayersContentContext, const FEditLayersHeightmapMergeParams& InMergeParams);
+	int32 PerformLayersHeightmapsGlobalMerge(const FUpdateLayersContentContext& InUpdateLayersContentContext, const FEditLayersHeightmapMergeParams& InMergeParams);
 	void ResolveLayersHeightmapTexture(FTextureToComponentHelper const& MapHelper, TSet<UTexture2D*> const& HeightmapsToResolve, bool bIntermediateRender, bool bFlushRender, TMap<ULandscapeComponent*, FLandscapeEditLayerReadbackResult>& InOutComponents);
+
+	int32 RegenerateLayersWeightmaps(FUpdateLayersContentContext& InUpdateLayersContentContext);
+	int32 PerformLayersWeightmapsLocalMerge(FUpdateLayersContentContext& InUpdateLayersContentContext, const FEditLayersWeightmapMergeParams& InMergeParams);
+	int32 PerformLayersWeightmapsGlobalMerge(FUpdateLayersContentContext& InUpdateLayersContentContext, const FEditLayersWeightmapMergeParams& InMergeParams);
 	void ResolveLayersWeightmapTexture(FTextureToComponentHelper const& MapHelper, TSet<UTexture2D*> const& WeightmapsToResolve, bool bIntermediateRender, bool bFlushRender, TMap<ULandscapeComponent*, FLandscapeEditLayerReadbackResult>& InOutComponents);
 
 	using FDirtyDelegate = TFunctionRef<void(UTexture2D const*, FColor const*, FColor const*)>;
@@ -360,6 +379,8 @@ private:
 	bool PrepareLayersTextureResources(const TArray<FLandscapeLayer>& InLayers, bool bInWaitForStreaming);
 	bool PrepareLayersBrushResources(bool bInWaitForStreaming);
 	void InvalidateRVTForTextures(const TSet<TObjectPtr<UTexture2D>>& InTextures);
+	void PrepareLayersHeightmapsLocalMergeRenderThreadData(const FUpdateLayersContentContext& InUpdateLayersContentContext, const FEditLayersHeightmapMergeParams& InMergeParams, EditLayersHeightmapLocalMerge_RenderThread::FMergeInfo& OutRenderThreadData);
+	void PrepareLayersWeightmapsLocalMergeRenderThreadData(const FUpdateLayersContentContext& InUpdateLayersContentContext, const FEditLayersWeightmapMergeParams& InMergeParams, EditLayersWeightmapLocalMerge_RenderThread::FMergeInfo& OutRenderThreadData);
 
 	void UpdateLayersMaterialInstances(const TArray<ULandscapeComponent*>& InLandscapeComponents);
 
@@ -367,7 +388,7 @@ private:
 														  TArray<struct FLandscapeLayerWeightmapExtractMaterialLayersComponentData>& OutComponentData, TMap<ULandscapeLayerInfoObject*, int32>& OutLayerInfoObjects);
 	void PrepareComponentDataToPackMaterialLayersCS(int32 InCurrentWeightmapToProcessIndex, const FIntPoint& InLandscapeBase, const TArray<ULandscapeComponent*>& InAllLandscapeComponents, TArray<UTexture2D*>& InOutProcessedWeightmaps,
 													TArray<FLandscapeEditLayerReadback*>& OutProcessedCPUReadBacks, TArray<struct FLandscapeLayerWeightmapPackMaterialLayersComponentData>& OutComponentData);
-	void ReallocateLayersWeightmaps(const TArray<ULandscapeComponent*>& InLandscapeComponents, const TArray<ULandscapeLayerInfoObject*>& InBrushRequiredAllocations);
+	void ReallocateLayersWeightmaps(FUpdateLayersContentContext& InUpdateLayersContentContext, const TArray<ULandscapeLayerInfoObject*>& InBrushRequiredAllocations);
 	void InitializeLayersWeightmapResources();
 	bool GenerateZeroAllocationPerComponents(const TArray<ALandscapeProxy*>& InAllLandscape, const TMap<ULandscapeLayerInfoObject*, bool>& InWeightmapLayersBlendSubstractive);
 
@@ -473,6 +494,9 @@ private:
 	UPROPERTY(Transient)
 	bool bLandscapeLayersAreInitialized;
 	
+	UPROPERTY(Transient)
+	bool bLandscapeLayersAreUsingLocalMerge;
+
 	UPROPERTY(Transient)
 	bool WasCompilingShaders;
 
diff --git a/Engine/Source/Runtime/Landscape/Classes/LandscapeComponent.h b/Engine/Source/Runtime/Landscape/Classes/LandscapeComponent.h
--- a/Engine/Source/Runtime/Landscape/Classes/LandscapeComponent.h
+++ b/Engine/Source/Runtime/Landscape/Classes/LandscapeComponent.h
@@ -1106,6 +1106,12 @@ public:
 	LANDSCAPE_API void RequestDeferredClientUpdate();
 	LANDSCAPE_API uint32 GetLayerUpdateFlagPerMode() const { return LayerUpdateFlagPerMode; }
 	LANDSCAPE_API uint32 ComputeWeightmapsHash();
+
+	void GetUsedPaintLayers(const FGuid& InLayerGuid, TArray<ULandscapeLayerInfoObject*>& OutUsedLayerInfos) const;
+
+	void GetLandscapeComponentNeighborsToRender(TSet<ULandscapeComponent*>& NeighborComponents) const;
+	void GetLandscapeComponentWeightmapsToRender(TSet<ULandscapeComponent*>& WeightmapComponents) const;
+	void GetLandscapeComponentNeighbors3x3(TStaticArray<ULandscapeComponent*, 9>& OutNeighborComponents) const;
 #endif
 
 	friend class FLandscapeComponentSceneProxy;
diff --git a/Engine/Source/Runtime/Landscape/Classes/LandscapeWeightmapUsage.h b/Engine/Source/Runtime/Landscape/Classes/LandscapeWeightmapUsage.h
--- a/Engine/Source/Runtime/Landscape/Classes/LandscapeWeightmapUsage.h
+++ b/Engine/Source/Runtime/Landscape/Classes/LandscapeWeightmapUsage.h
@@ -58,4 +58,17 @@ public:
 	{
 		return FreeChannelCount() == NumChannels;
 	}
+
+	TArray<ULandscapeComponent*, TInlineAllocator<4>> GetUniqueValidComponents() const
+	{
+		TArray<ULandscapeComponent*, TInlineAllocator<4>> UniqueComponents;
+		for (TObjectPtr<ULandscapeComponent> Component : ChannelUsage)
+		{
+			if (Component != nullptr)
+			{
+				UniqueComponents.AddUnique(Component);
+			}
+		}
+		return UniqueComponents;
+	}
 };
diff --git a/Engine/Source/Runtime/Landscape/Private/Landscape.cpp b/Engine/Source/Runtime/Landscape/Private/Landscape.cpp
--- a/Engine/Source/Runtime/Landscape/Private/Landscape.cpp
+++ b/Engine/Source/Runtime/Landscape/Private/Landscape.cpp
@@ -1128,6 +1128,7 @@ ALandscape::ALandscape(const FObjectInitializer& ObjectInitializer)
 	WeightmapScratchExtractLayerTextureResource = nullptr;
 	WeightmapScratchPackLayerTextureResource = nullptr;
 	bLandscapeLayersAreInitialized = false;
+	bLandscapeLayersAreUsingLocalMerge = false;
 	LandscapeEdMode = nullptr;
 	bGrassUpdateEnabled = true;
 	bIsSpatiallyLoaded = false;
diff --git a/Engine/Source/Runtime/Landscape/Private/LandscapeEdit.cpp b/Engine/Source/Runtime/Landscape/Private/LandscapeEdit.cpp
--- a/Engine/Source/Runtime/Landscape/Private/LandscapeEdit.cpp
+++ b/Engine/Source/Runtime/Landscape/Private/LandscapeEdit.cpp
@@ -3958,14 +3958,7 @@ void ULandscapeInfo::GetUsedPaintLayers(const FGuid& InLayerGuid, TArray<ULandsc
 	{
 		for (ULandscapeComponent* Component : Proxy->LandscapeComponents)
 		{
-			const TArray<FWeightmapLayerAllocationInfo>& AllocInfos = Component->GetWeightmapLayerAllocations(InLayerGuid);
-			for (const FWeightmapLayerAllocationInfo& AllocInfo : AllocInfos)
-			{
-				if (AllocInfo.LayerInfo != nullptr)
-				{
-					OutUsedLayerInfos.AddUnique(AllocInfo.LayerInfo);
-				}
-			}
+			Component->GetUsedPaintLayers(InLayerGuid, OutUsedLayerInfos);
 		}
 	});
 }
