Fixed shader compilation errors

diff --git a/Engine/Shaders/Private/DiaphragmDOF/DOFReduce.usf b/Engine/Shaders/Private/DiaphragmDOF/DOFReduce.usf
--- a/Engine/Shaders/Private/DiaphragmDOF/DOFReduce.usf
+++ b/Engine/Shaders/Private/DiaphragmDOF/DOFReduce.usf
@@ -558,10 +558,10 @@ void ReduceCS(
 		for (uint i = 0; i < (DIM_REDUCE_MIP_COUNT - 1); i++)
 		{
 			// Multiplier that needs to be applied on Coc computation to be resolution independent.
-			DownsampleParams.CocRadiusMultiplier = 0.5 / float(1 << (i));
+			DownsampleParams.CocRadiusMultiplier = 0.5 / float(1u << (i));
 
 			// Scale of the reduction within the 8x8 tile.
-			const uint ReductionScale = 1 << i;
+			const uint ReductionScale = 1u << i;
 			
 			// Chooses output CoC across lanes.
 			float OutputCocRadius = DownsampleCoc_Wave8x8(CocRadius, ReductionScale);
@@ -585,7 +585,7 @@ void ReduceCS(
 			// Output mip map.
 			OutputPixelId = OutputPixelId >> 1;
 			BRANCH
-			if (((DispatchThreadId.x | DispatchThreadId.y) & ((2 << i) - 1)) == 0)
+			if (((DispatchThreadId.x | DispatchThreadId.y) & ((2u << i) - 1)) == 0)
 			{
 				OutputMipLevel(/* MipLevel = */ i + 1, OutputPixelId, GatheredColor, CocRadius);
 			}
diff --git a/Engine/Shaders/Private/TemporalSuperResolution/TSRHoleFillVelocity.usf b/Engine/Shaders/Private/TemporalSuperResolution/TSRHoleFillVelocity.usf
--- a/Engine/Shaders/Private/TemporalSuperResolution/TSRHoleFillVelocity.usf
+++ b/Engine/Shaders/Private/TemporalSuperResolution/TSRHoleFillVelocity.usf
@@ -207,8 +207,8 @@ void MainCS(
 				NeighborPixelVelocity0 = WaveBroadcast(BroadcastSettings[0], PixelVelocityNeighborhood[FetchIndex[0]]);
 				NeighborPixelVelocity1 = WaveBroadcast(BroadcastSettings[1], PixelVelocityNeighborhood[FetchIndex[1]]);
 			
-				NeighborPixelVelocityLength0 = WaveBroadcast(BroadcastSettings[0], PixelVelocityLengthNeighborhood[FetchIndex[0]]);
-				NeighborPixelVelocityLength1 = WaveBroadcast(BroadcastSettings[1], PixelVelocityLengthNeighborhood[FetchIndex[1]]);
+				NeighborPixelVelocityLength0 = taa_half(WaveBroadcast(BroadcastSettings[0], PixelVelocityLengthNeighborhood[FetchIndex[0]]));
+				NeighborPixelVelocityLength1 = taa_half(WaveBroadcast(BroadcastSettings[1], PixelVelocityLengthNeighborhood[FetchIndex[1]]));
 				
 				bNeighborIsParallaxRejected0 = WaveBroadcast(BroadcastSettings[0], ParallaxRejectionMaskNeighborhood[FetchIndex[0]]) < taa_half(PARALLAX_REJECTION_MASK_THRESHOLD);
 				bNeighborIsParallaxRejected1 = WaveBroadcast(BroadcastSettings[1], ParallaxRejectionMaskNeighborhood[FetchIndex[1]]) < taa_half(PARALLAX_REJECTION_MASK_THRESHOLD);
