Fix DeepShadow GPU Allocation.

This CL removes a temporary fix which was adding an extra world offset to compensate 'in-conversion' code.

diff --git a/Engine/Shaders/Private/HairStrands/HairStrandsDeepShadowAllocation.usf b/Engine/Shaders/Private/HairStrands/HairStrandsDeepShadowAllocation.usf
--- a/Engine/Shaders/Private/HairStrands/HairStrandsDeepShadowAllocation.usf
+++ b/Engine/Shaders/Private/HairStrands/HairStrandsDeepShadowAllocation.usf
@@ -137,11 +137,7 @@ void CreateViewInfo(uint2 DispatchThreadId : SV_DispatchThreadID)
 	FHairAABB TranslatedBound = InitHairAABB();
 	if (MacroGroupId < MacroGroupCount)
 	{
-		const float3 TranslatedWorldOffset = LWCHackToFloat(PrimaryView.PreViewTranslation);
-
 		TranslatedBound = ReadHairAABB(MacroGroupId, MacroGroupAABBBuffer);
-		TranslatedBound.Min += TranslatedWorldOffset; // HAIR_LWC
-		TranslatedBound.Max += TranslatedWorldOffset; // HAIR_LWC
 	}
 
 	FDeepShadowViewInfo ViewInfo;
