Add LWC support for shader draw debug.

Draw function have now a WorldSpace (WS) and a TranslatedWorldSpace(TWS) variant.
Legacy functions are rerouted to the WS variant.

All engine code is converted into the new functions.

diff --git a/Engine/Plugins/Experimental/VirtualHeightfieldMesh/Shaders/Private/VirtualHeightfieldMesh.usf b/Engine/Plugins/Experimental/VirtualHeightfieldMesh/Shaders/Private/VirtualHeightfieldMesh.usf
--- a/Engine/Plugins/Experimental/VirtualHeightfieldMesh/Shaders/Private/VirtualHeightfieldMesh.usf
+++ b/Engine/Plugins/Experimental/VirtualHeightfieldMesh/Shaders/Private/VirtualHeightfieldMesh.usf
@@ -161,12 +161,12 @@ void DebugDrawUVBox(float3 InUVMin, float3 InUVMax, float4x4 InTransform, float4
 	WorldPos[6] = mul(float4(InUVMin.x, InUVMax.y, InUVMax.z, 1), InTransform);
 	WorldPos[7] = mul(float4(InUVMax.x, InUVMax.y, InUVMax.z, 1), InTransform);
 
-	AddQuad(WorldPos[0], WorldPos[2], WorldPos[3], WorldPos[1], InColor);
-	AddQuad(WorldPos[4], WorldPos[6], WorldPos[7], WorldPos[5], InColor);
-	AddLine(WorldPos[0], WorldPos[4], InColor, InColor);
-	AddLine(WorldPos[1], WorldPos[5], InColor, InColor);
-	AddLine(WorldPos[2], WorldPos[6], InColor, InColor);
-	AddLine(WorldPos[3], WorldPos[7], InColor, InColor);
+	AddQuadWS(WorldPos[0], WorldPos[2], WorldPos[3], WorldPos[1], InColor);
+	AddQuadWS(WorldPos[4], WorldPos[6], WorldPos[7], WorldPos[5], InColor);
+	AddLineWS(WorldPos[0], WorldPos[4], InColor, InColor);
+	AddLineWS(WorldPos[1], WorldPos[5], InColor, InColor);
+	AddLineWS(WorldPos[2], WorldPos[6], InColor, InColor);
+	AddLineWS(WorldPos[3], WorldPos[7], InColor, InColor);
 #endif
 }
 
diff --git a/Engine/Shaders/Private/GPUSceneDebugRender.usf b/Engine/Shaders/Private/GPUSceneDebugRender.usf
--- a/Engine/Shaders/Private/GPUSceneDebugRender.usf
+++ b/Engine/Shaders/Private/GPUSceneDebugRender.usf
@@ -157,8 +157,8 @@ void GPUSceneDebugRenderCS(uint InstanceID : SV_DispatchThreadID)
 	}
 	if (bIsMouseOver || bIsSelected || bDrawAll != 0)
 	{
-		AddOBB(InstanceData.LocalBoundsCenter - InstanceData.LocalBoundsExtent, InstanceData.LocalBoundsCenter + InstanceData.LocalBoundsExtent, InstanceBoundColor, LWCHackToFloat(InstanceData.LocalToWorld));
-		AddReferential(LWCHackToFloat(InstanceData.LocalToWorld), 1);
+		AddOBBWS(InstanceData.LocalBoundsCenter - InstanceData.LocalBoundsExtent, InstanceData.LocalBoundsCenter + InstanceData.LocalBoundsExtent, InstanceBoundColor, LWCHackToFloat(InstanceData.LocalToWorld));
+		AddReferentialWS(LWCHackToFloat(InstanceData.LocalToWorld), 1);
 	}
 
 	uint StartOffsetX = 10;
diff --git a/Engine/Shaders/Private/HairStrands/HairCardsDebug.usf b/Engine/Shaders/Private/HairStrands/HairCardsDebug.usf
--- a/Engine/Shaders/Private/HairStrands/HairCardsDebug.usf
+++ b/Engine/Shaders/Private/HairStrands/HairCardsDebug.usf
@@ -167,7 +167,7 @@ void MainCS(uint3 DispatchThreadId : SV_DispatchThreadID)
 
 	if (P1.Type != HAIR_CONTROLPOINT_START)
 	{
-		AddLine(mul(float4(P0.Position, 1), LocalToWorld), mul(float4(P1.Position,1), LocalToWorld), Color, Color);
+		AddLineWS(mul(float4(P0.Position, 1), LocalToWorld), mul(float4(P1.Position,1), LocalToWorld), Color, Color);
 	}
 }
 
diff --git a/Engine/Shaders/Private/HairStrands/HairCardsGeneration.usf b/Engine/Shaders/Private/HairStrands/HairCardsGeneration.usf
--- a/Engine/Shaders/Private/HairStrands/HairCardsGeneration.usf
+++ b/Engine/Shaders/Private/HairStrands/HairCardsGeneration.usf
@@ -204,7 +204,7 @@ void MainPS(
 	if (bDebugEnable)
 	{
 		const float4 Color = float4(1, 1, 0, 1);
-		AddAABB(Raster_MinBound, Raster_MaxBound, Color);
+		AddAABBWS(Raster_MinBound, Raster_MaxBound, Color);
 	}
 	#endif 
 
@@ -216,7 +216,7 @@ void MainPS(
 	if (bDebugEnable)
 	{
 		const float4 Color = float4(1, 0, 0, 1);
-		AddLine(RayO, RayO + RayD * MaxDistance, Color, Color);
+		AddLineWS(RayO, RayO + RayD * MaxDistance, Color, Color);
 	}
 	#endif
 
@@ -249,7 +249,7 @@ void MainPS(
 		if (bDebugEnable)
 		{
 			const float4 Color = float4(0, 1, 0, 1);
-			AddLine(CP0, CP1, Color, Color);
+			AddLineWS(CP0, CP1, Color, Color);
 		}
 		#endif
 
@@ -275,7 +275,7 @@ void MainPS(
 			if (bDebugEnable)
 			{
 				const float4 Color = float4(1, 0, 0, 1);
-				AddLine(RayP0, RayP1, Color, Color);
+				AddLineWS(RayP0, RayP1, Color, Color);
 			}
 			#endif
 			// Scale radius in order to avoid having 0 coverage when hair are too thin
diff --git a/Engine/Shaders/Private/HairStrands/HairCardsVoxel.usf b/Engine/Shaders/Private/HairStrands/HairCardsVoxel.usf
--- a/Engine/Shaders/Private/HairStrands/HairCardsVoxel.usf
+++ b/Engine/Shaders/Private/HairStrands/HairCardsVoxel.usf
@@ -106,7 +106,7 @@ void MainCS(uint3 DispatchThreadId : SV_DispatchThreadID)
 
 	if (uint(PixelCoord.x) == 0 && uint(PixelCoord.y) == 0)
 	{
-		AddAABB(Voxel_MinBound, Voxel_MaxBound, float4(1, 0.5f, 0, 1));
+		AddAABBWS(Voxel_MinBound, Voxel_MaxBound, float4(1, 0.5f, 0, 1));
 	}
 
 	bool bIsValid = false;
@@ -196,10 +196,10 @@ void MainCS(uint3 DispatchThreadId : SV_DispatchThreadID)
 					//HitColor = float(Density) / 10.f;
 					if (bDebugEnabled)
 					{
-						//AddLine(HitP, HitP + UnoccludedDirection * 3, float4(1, 0, 0, 1), float4(1, 1, 0, 1));
+						//AddLineWS(HitP, HitP + UnoccludedDirection * 3, float4(1, 0, 0, 1), float4(1, 1, 0, 1));
 
-						AddLine(HitP, HitP + Tangent * 3, float4(0, 0, 1, 1), float4(0, 0, 1, 1));
-						AddLine(HitP, HitP + Normal * 3, float4(1, 0, 0, 1), float4(1, 0, 0, 1));
+						AddLineWS(HitP, HitP + Tangent * 3, float4(0, 0, 1, 1), float4(0, 0, 1, 1));
+						AddLineWS(HitP, HitP + Normal * 3, float4(1, 0, 0, 1), float4(1, 0, 0, 1));
 					}
 
 					break;
diff --git a/Engine/Shaders/Private/HairStrands/HairStrandsClusterCulling.usf b/Engine/Shaders/Private/HairStrands/HairStrandsClusterCulling.usf
--- a/Engine/Shaders/Private/HairStrands/HairStrandsClusterCulling.usf
+++ b/Engine/Shaders/Private/HairStrands/HairStrandsClusterCulling.usf
@@ -234,7 +234,7 @@ void MainDrawDebugAABBCS(uint3 DispatchThreadId : SV_DispatchThreadID)
 		{
 			const float4 GroupColor = float4(1.0f, 0.5f, 0.0f, 1.0f);
 			const FHairAABB Bound = ReadHairAABB(0, GroupAABBBuffer);
-			AddAABB(Bound.Min, Bound.Max, GroupColor);
+			AddAABBTWS(Bound.Min, Bound.Max, GroupColor);
 		}
 
 		const uint ClusterId = DispatchThreadId;
@@ -278,7 +278,7 @@ void MainDrawDebugAABBCS(uint3 DispatchThreadId : SV_DispatchThreadID)
 		}
 		#endif
 
-		AddAABB(Min, Max, Color);
+		AddAABBTWS(Min, Max, Color);
 	}
 }
 
diff --git a/Engine/Shaders/Private/HairStrands/HairStrandsDebugPrint.usf b/Engine/Shaders/Private/HairStrands/HairStrandsDebugPrint.usf
--- a/Engine/Shaders/Private/HairStrands/HairStrandsDebugPrint.usf
+++ b/Engine/Shaders/Private/HairStrands/HairStrandsDebugPrint.usf
@@ -6,21 +6,13 @@
 #include "../DeferredShadingCommon.ush"
 #include "../ShaderPrintCommon.ush"
 #include "../ShaderDrawDebug.ush"
+#include "../PositionReconstructionCommon.ush"
 #include "HairStrandsVisibilityCommon.ush"
 #include "HairStrandsAABBCommon.ush" 
 
 
 ////////////////////////////////////////////////////////////////////////////////////////////////
 
-float3 ReconstructWorldPositionFromDepth(float2 UV, float Depth)
-{
-	float2 ScreenPosition = (UV - View.ScreenPositionScaleBias.wz) / View.ScreenPositionScaleBias.xy;
-	float4 HomogeneousWorldPosition = mul(float4(ScreenPosition * Depth, Depth, 1), LWCHackToFloat(PrimaryView.ScreenToWorld));
-	float3 WorldPosition = HomogeneousWorldPosition.xyz / HomogeneousWorldPosition.w;
-
-	return WorldPosition;
-}
-
 void AddNewLine(float2 OriginalPos, inout float2 Pos)
 {
 	Pos = ShaderPrintNewline(Pos); Pos.x = OriginalPos.x;
@@ -66,7 +58,7 @@ void MainCS(uint GroupIndex : SV_GroupIndex, uint3 DispatchThreadId : SV_Dispatc
 		{
 			const float4 InstanceBoundColor = float4(1.f, 0.5f, 0.5f, 1.0f);
 			FHairAABB Bound = ReadHairAABB(GroupIndex, HairMacroGroupAABBBuffer);
-			AddAABB(Bound.Min, Bound.Max, InstanceBoundColor);
+			AddAABBTWS(Bound.Min, Bound.Max, InstanceBoundColor);
 		}
 
 		const uint DispatchX = HairVisibilityIndirectArgsBuffer[0];
@@ -306,8 +298,8 @@ FTangentDesc AddTangentPlot(uint2 PixelCoord, float2 UV)
 		if (NodeDesc.Count > 0)
 		{
 			const float SceneDepth = ConvertFromDeviceZ(Closest_Depth);
-			const float3 WorldPosition = ReconstructWorldPositionFromDepth(UV, SceneDepth);
-			AddLine(WorldPosition, WorldPosition + Closest_Tangent, ColorRed, ColorYellow);
+			const float3 WorldPosition = ReconstructTranslatedWorldPositionFromDepth(UV, SceneDepth);
+			AddLineTWS(WorldPosition, WorldPosition + Closest_Tangent, ColorRed, ColorYellow);
 			Out.Type = 1;
 			Out.Tangent = Closest_Tangent;
 		}
@@ -318,11 +310,11 @@ FTangentDesc AddTangentPlot(uint2 PixelCoord, float2 UV)
 	{
 		FScreenSpaceData ScreenSpaceData = GetScreenSpaceData(UV);
 		const float SceneDepth = CalcSceneDepth(UV);
-		const float3 WorldPosition = ReconstructWorldPositionFromDepth(UV, SceneDepth);
+		const float3 WorldPosition = ReconstructTranslatedWorldPositionFromDepth(UV, SceneDepth);
 		if (ScreenSpaceData.GBuffer.ShadingModelID == SHADINGMODELID_HAIR)
 		{
 			const float3 Tangent = ScreenSpaceData.GBuffer.WorldNormal;
-			AddLine(WorldPosition, WorldPosition + Tangent, ColorPurple, ColorYellow);
+			AddLineTWS(WorldPosition, WorldPosition + Tangent, ColorPurple, ColorYellow);
 			Out.Type = 2;
 			Out.Tangent = Tangent;
 		}
diff --git a/Engine/Shaders/Private/HairStrands/HairStrandsDeepShadowDebug.usf b/Engine/Shaders/Private/HairStrands/HairStrandsDeepShadowDebug.usf
--- a/Engine/Shaders/Private/HairStrands/HairStrandsDeepShadowDebug.usf
+++ b/Engine/Shaders/Private/HairStrands/HairStrandsDeepShadowDebug.usf
@@ -148,12 +148,10 @@ void MainCS(uint3 DispatchThreadId : SV_DispatchThreadID)
 
 	for (uint MacroGroupIt = 0; MacroGroupIt < MacroGroupCount; ++MacroGroupIt)
 	{
-		const FHairAABB Bound = ReadHairAABB(MacroGroupIt, MacroGroupAABBBuffer); // HAIR_LWC
-		AddAABB(Bound.Min, Bound.Max, float4(1, 0.5f, 0, 1));
+		const FHairAABB Bound = ReadHairAABB(MacroGroupIt, MacroGroupAABBBuffer);
+		AddAABBTWS(Bound.Min, Bound.Max, float4(1, 0.5f, 0, 1));
 	}
 
-	const float3 TranslatedWorldOffset = LWCHackToFloat(PrimaryView.PreViewTranslation); // HAIR_LWC
-
 	const float4 FrustumColor = float4(0.5f, 1, 0, 1);
 	for (uint SlotIt = 0; SlotIt < AllocatedSlotCount; ++SlotIt)
 	{
@@ -170,30 +168,30 @@ void MainCS(uint3 DispatchThreadId : SV_DispatchThreadID)
 		float4 P20 = float4( 1,  1, 1, 1);
 		float4 P30 = float4( 1, -1, 1, 1);
 
-		P00 = mul(P00, ClipToTranslatedWorld); P00 /= P00.w; P00.xyz += TranslatedWorldOffset;
-		P10 = mul(P10, ClipToTranslatedWorld); P10 /= P10.w; P10.xyz += TranslatedWorldOffset;
-		P20 = mul(P20, ClipToTranslatedWorld); P20 /= P20.w; P20.xyz += TranslatedWorldOffset;
-		P30 = mul(P30, ClipToTranslatedWorld); P30 /= P30.w; P30.xyz += TranslatedWorldOffset;
-
-		P01 = mul(P01, ClipToTranslatedWorld); P01 /= P01.w; P01.xyz += TranslatedWorldOffset;
-		P11 = mul(P11, ClipToTranslatedWorld); P11 /= P11.w; P11.xyz += TranslatedWorldOffset;
-		P21 = mul(P21, ClipToTranslatedWorld); P21 /= P21.w; P21.xyz += TranslatedWorldOffset;
-		P31 = mul(P31, ClipToTranslatedWorld); P31 /= P31.w; P31.xyz += TranslatedWorldOffset;
-
-		AddLine(P00, P10, FrustumColor, FrustumColor);
-		AddLine(P10, P20, FrustumColor, FrustumColor);
-		AddLine(P20, P30, FrustumColor, FrustumColor);
-		AddLine(P30, P00, FrustumColor, FrustumColor);
-
-		AddLine(P00, P01, FrustumColor, FrustumColor);
-		AddLine(P10, P11, FrustumColor, FrustumColor);
-		AddLine(P20, P21, FrustumColor, FrustumColor);
-		AddLine(P30, P31, FrustumColor, FrustumColor);
-
-		AddLine(P01, P11, FrustumColor, FrustumColor);
-		AddLine(P11, P21, FrustumColor, FrustumColor);
-		AddLine(P21, P31, FrustumColor, FrustumColor);
-		AddLine(P31, P01, FrustumColor, FrustumColor);
+		P00 = mul(P00, ClipToTranslatedWorld); P00 /= P00.w;
+		P10 = mul(P10, ClipToTranslatedWorld); P10 /= P10.w;
+		P20 = mul(P20, ClipToTranslatedWorld); P20 /= P20.w;
+		P30 = mul(P30, ClipToTranslatedWorld); P30 /= P30.w;
+
+		P01 = mul(P01, ClipToTranslatedWorld); P01 /= P01.w;
+		P11 = mul(P11, ClipToTranslatedWorld); P11 /= P11.w;
+		P21 = mul(P21, ClipToTranslatedWorld); P21 /= P21.w;
+		P31 = mul(P31, ClipToTranslatedWorld); P31 /= P31.w;
+
+		AddLineTWS(P00, P10, FrustumColor, FrustumColor);
+		AddLineTWS(P10, P20, FrustumColor, FrustumColor);
+		AddLineTWS(P20, P30, FrustumColor, FrustumColor);
+		AddLineTWS(P30, P00, FrustumColor, FrustumColor);
+
+		AddLineTWS(P00, P01, FrustumColor, FrustumColor);
+		AddLineTWS(P10, P11, FrustumColor, FrustumColor);
+		AddLineTWS(P20, P21, FrustumColor, FrustumColor);
+		AddLineTWS(P30, P31, FrustumColor, FrustumColor);
+
+		AddLineTWS(P01, P11, FrustumColor, FrustumColor);
+		AddLineTWS(P11, P21, FrustumColor, FrustumColor);
+		AddLineTWS(P21, P31, FrustumColor, FrustumColor);
+		AddLineTWS(P31, P01, FrustumColor, FrustumColor);
 	}
 }
 #endif // SHADER_DOMINFO
diff --git a/Engine/Shaders/Private/HairStrands/HairStrandsDeepTransmittanceMask.usf b/Engine/Shaders/Private/HairStrands/HairStrandsDeepTransmittanceMask.usf
--- a/Engine/Shaders/Private/HairStrands/HairStrandsDeepTransmittanceMask.usf
+++ b/Engine/Shaders/Private/HairStrands/HairStrandsDeepTransmittanceMask.usf
@@ -358,7 +358,6 @@ void MainCS(uint2 DispatchThreadId : SV_DispatchThreadID)
 	{
 		return;
 	}
-	const float3 TranslatedWorldOffset = LWCHackToFloat(PrimaryView.PreViewTranslation); // HAIR_LWC
 	const FHairSample Sample	= UnpackHairSample(HairStrands.HairSampleData[SampleIndex]);
 	const uint MacroGroupIndex	= clamp(Sample.MacroGroupId, 0, MAX_HAIR_MACROGROUP_COUNT - 1);
 	const float SceneDepth		= ConvertFromDeviceZ(Sample.Depth);
diff --git a/Engine/Shaders/Private/HairStrands/HairStrandsEnvironmentLighting.usf b/Engine/Shaders/Private/HairStrands/HairStrandsEnvironmentLighting.usf
--- a/Engine/Shaders/Private/HairStrands/HairStrandsEnvironmentLighting.usf
+++ b/Engine/Shaders/Private/HairStrands/HairStrandsEnvironmentLighting.usf
@@ -308,7 +308,7 @@ float3 ReflectionEnvironment(FGBufferData GBuffer, float2 ScreenPosition, float2
 		#if VOXEL_TRAVERSAL_DEBUG
 		if (bDebugEnabled)
 		{
-			AddLine(TranslatedWorldPosition, TranslatedWorldPosition + UnoccludedN * 50, float4(1, 0, 1, 1), float4(1, 0, 1, 1));
+			AddLineTWS(TranslatedWorldPosition, TranslatedWorldPosition + UnoccludedN * 50, float4(1, 0, 1, 1), float4(1, 0, 1, 1));
 		}
 		#endif
 
@@ -613,7 +613,7 @@ float3 ReflectionEnvironmentUniformSampling(
 		#if PERMUTATION_DEBUG
 		if (bDebugEnabled)
 		{
-			AddLine(SampleTranslatedWorldPosition, SampleTranslatedLightPosition, float4(1, 0, 1, 1), float4(1, 0, 1, 1));
+			AddLineTWS(SampleTranslatedWorldPosition, SampleTranslatedLightPosition, float4(1, 0, 1, 1), float4(1, 0, 1, 1));
 		}
 		#endif
 	}
diff --git a/Engine/Shaders/Private/HairStrands/HairStrandsInterpolation.usf b/Engine/Shaders/Private/HairStrands/HairStrandsInterpolation.usf
--- a/Engine/Shaders/Private/HairStrands/HairStrandsInterpolation.usf
+++ b/Engine/Shaders/Private/HairStrands/HairStrandsInterpolation.usf
@@ -343,14 +343,14 @@ void MainCS(uint2 DispatchThreadId : SV_DispatchThreadID)
 				const float4 DeformedColor = float4(1, 0, 0, 1);
 				if (VertexIndex==0)
 				{	
-					AddCross(InRenderHairPositionOffset, 1.2, RestColor);
-					AddCross(OutRenHairPositionOffset, 1, DeformedColor);
+					AddCrossWS(InRenderHairPositionOffset, 1.2, RestColor);
+					AddCrossWS(OutRenHairPositionOffset, 1, DeformedColor);
 				}
 			
 				if (VertexIndex >= 0 && VertexIndex < 64)
 				{
-					AddReferential(RestControlPoint, float3(0, 0, 1));
-					AddReferential(WorldVertexPosition, float3(0, 0, 1));
+					AddReferentialWS(RestControlPoint, float3(0, 0, 1));
+					AddReferentialWS(WorldVertexPosition, float3(0, 0, 1));
 				}
 
 				{
@@ -358,8 +358,8 @@ void MainCS(uint2 DispatchThreadId : SV_DispatchThreadID)
 					const bool bIsRoot = ControlPointType == 1;
 					if (bIsRoot)
 					{
-						AddLineTriangle(RestTriangle.P0, RestTriangle.P1, RestTriangle.P2, RestColor);
-						AddLineTriangle(DeformedTriangle.P0, DeformedTriangle.P1, DeformedTriangle.P2, DeformedColor);
+						AddLineTriangleWS(RestTriangle.P0, RestTriangle.P1, RestTriangle.P2, RestColor);
+						AddLineTriangleWS(DeformedTriangle.P0, DeformedTriangle.P1, DeformedTriangle.P2, DeformedColor);
 					}
 				}
 			}
@@ -697,8 +697,8 @@ void ClusterAABBEvaluationCS(uint3 DispatchThreadId : SV_DispatchThreadID, uint3
 	if (ThreadIndex == 0)
 	{
 		const float4 DebugColor = float4(ColorMapViridis(ClusterId / float(TotalClusterCount)),1);
-		AddAABB(SharedClusterAABBMin[0], SharedClusterAABBMax[0], DebugColor);
-		AddReferential(LocalToWorldMatrix, 10);
+		AddAABBTWS(SharedClusterAABBMin[0], SharedClusterAABBMax[0], DebugColor);
+		AddReferentialWS(LocalToWorldMatrix, 10);
 	}
 	#endif
 
diff --git a/Engine/Shaders/Private/HairStrands/HairStrandsRaytracingGeometry.usf b/Engine/Shaders/Private/HairStrands/HairStrandsRaytracingGeometry.usf
--- a/Engine/Shaders/Private/HairStrands/HairStrandsRaytracingGeometry.usf
+++ b/Engine/Shaders/Private/HairStrands/HairStrandsRaytracingGeometry.usf
@@ -302,7 +302,7 @@ void MainCS(uint2 DispatchThreadId : SV_DispatchThreadID, uint LinearIndex : SV_
 				if (Offset < 2)
 				{
 					const FHairAABB Bound = ReadHairAABB(0, GroupAABBBuffer);
-					AddAABB(Bound.Min, Bound.Max, ColorRed);
+					AddAABBTWS(Bound.Min, Bound.Max, ColorRed);
 				}
 			}
 			s_InvalidMask[LinearIndex] = bIsInvalid ? 1u : 0u;
diff --git a/Engine/Shaders/Private/HairStrands/HairStrandsTexturesGeneration.usf b/Engine/Shaders/Private/HairStrands/HairStrandsTexturesGeneration.usf
--- a/Engine/Shaders/Private/HairStrands/HairStrandsTexturesGeneration.usf
+++ b/Engine/Shaders/Private/HairStrands/HairStrandsTexturesGeneration.usf
@@ -151,7 +151,7 @@ bool Trace(
 	{
 		const float4 Color0 = float4(0, 1, 0, 1);
 		const float4 Color1 = float4(1, 0, 0, 1);
-		AddLine(WP0, WP1, Color0, Color1);
+		AddLineWS(WP0, WP1, Color0, Color1);
 	}
 	#endif	
 
@@ -163,7 +163,7 @@ bool Trace(
 		if (bDebugEnable)
 		{
 			//AddLine(true, O, E);
-			AddAABB(Voxel_MinBound, Voxel_MaxBound, float4(1, 1, 1, 1));
+			AddAABBWS(Voxel_MinBound, Voxel_MaxBound, float4(1, 1, 1, 1));
 		}
 		#endif	
 
@@ -241,7 +241,7 @@ bool Trace(
 			{
 				const float3 MinVoxelAABB = Voxel_MinBound + PageIndexCoord * Voxel_Size;
 				const float3 MaxVoxelAABB = Voxel_MinBound + (PageIndexCoord + 1) * Voxel_Size;
-				AddAABB(MinVoxelAABB, MaxVoxelAABB, float4(0, 1, 1, 1));
+				AddAABBWS(MinVoxelAABB, MaxVoxelAABB, float4(0, 1, 1, 1));
 			}
 			#endif	
 
@@ -311,9 +311,9 @@ bool Trace(
 							if (bDebugEnable)
 							{
 								const float4 Color = float4(1, 0, 0, 1);
-								AddLine(CP_Curr.Position, CP_Next.Position, Color, Color);
-								AddCross(CP_Curr.Position, 0.05f, Color);
-								AddCross(CP_Next.Position, 0.05f, Color);
+								AddLineWS(CP_Curr.Position, CP_Next.Position, Color, Color);
+								AddCrossWS(CP_Curr.Position, 0.05f, Color);
+								AddCrossWS(CP_Next.Position, 0.05f, Color);
 							}
 							#endif							
 						}
diff --git a/Engine/Shaders/Private/HairStrands/HairStrandsVoxelPageRayMarching.usf b/Engine/Shaders/Private/HairStrands/HairStrandsVoxelPageRayMarching.usf
--- a/Engine/Shaders/Private/HairStrands/HairStrandsVoxelPageRayMarching.usf
+++ b/Engine/Shaders/Private/HairStrands/HairStrandsVoxelPageRayMarching.usf
@@ -76,7 +76,7 @@ void MainCS(uint3 DispatchThreadId : SV_DispatchThreadID)
 			TotalAllocatedPageCount += LocalNodeDesc.PageIndexResolution.x * LocalNodeDesc.PageIndexResolution.y * LocalNodeDesc.PageIndexResolution.z;
 		}
 
-		AddAABB(NodeDesc.TranslatedWorldMinAABB, NodeDesc.TranslatedWorldMaxAABB, ColorOrange);
+		AddAABBTWS(NodeDesc.TranslatedWorldMinAABB, NodeDesc.TranslatedWorldMaxAABB, ColorOrange);
 
 		// General information
 		// General
diff --git a/Engine/Shaders/Private/HairStrands/HairStrandsVoxelPageTraversal.ush b/Engine/Shaders/Private/HairStrands/HairStrandsVoxelPageTraversal.ush
--- a/Engine/Shaders/Private/HairStrands/HairStrandsVoxelPageTraversal.ush
+++ b/Engine/Shaders/Private/HairStrands/HairStrandsVoxelPageTraversal.ush
@@ -38,14 +38,14 @@ void AddLine(bool bEnabled, float3 P0, float3 P1)
 	if (!bEnabled)
 		return;
 	const float4 DebugColor = float4(1, 1, 0, 1);
-	AddLine(P0, P1, DebugColor, DebugColor);
+	AddLineTWS(P0, P1, DebugColor, DebugColor);
 }
 
 void AddMacroGroupAABB(bool bEnabled, float3 P0, float3 P1)
 {
 	if (!bEnabled)
 		return;
-	AddAABB(P0, P1, float4(1, 0, 1, 1));
+	AddAABBTWS(P0, P1, float4(1, 0, 1, 1));
 }
 
 void AddPageAABB(bool bEnabled, float3 P0, float3 P1, bool bIsPageValid)
@@ -53,7 +53,7 @@ void AddPageAABB(bool bEnabled, float3 P0, float3 P1, bool bIsPageValid)
 	if (!bEnabled)
 		return;
 	const float4 DebugColor = bIsPageValid ? float4(0, 1, 1, 1) : float4(1, 0, 0, 1);
-	AddAABB(P0, P1, DebugColor);
+	AddAABBTWS(P0, P1, DebugColor);
 }
 
 void AddStepAABB(bool bEnabled, float3 P0, float3 P1, float3 Color)
@@ -61,7 +61,7 @@ void AddStepAABB(bool bEnabled, float3 P0, float3 P1, float3 Color)
 	if (!bEnabled)
 		return;
 	const float4 DebugColor = float4(Color, 1);
-	AddAABB(P0, P1, DebugColor);
+	AddAABBTWS(P0, P1, DebugColor);
 }
 
 void AddAABBVolume(bool bEnabled, float3 P0, float3 P1)
@@ -69,7 +69,7 @@ void AddAABBVolume(bool bEnabled, float3 P0, float3 P1)
 	if (!bEnabled)
 		return;
 	const float4 DebugColor = float4(1,1,0, 1);
-	AddAABB(P0, P1, DebugColor);
+	AddAABBTWS(P0, P1, DebugColor);
 }
 
 #endif
diff --git a/Engine/Shaders/Private/OIT/OITSorting.usf b/Engine/Shaders/Private/OIT/OITSorting.usf
--- a/Engine/Shaders/Private/OIT/OITSorting.usf
+++ b/Engine/Shaders/Private/OIT/OITSorting.usf
@@ -210,7 +210,7 @@ void MainCS(uint3 DispatchThreadId : SV_DispatchThreadID, uint GroupThreadId : S
 			const float3 WP0 = mul(float4(P0, 1), LocalToWorld).xyz;
 			const float3 WP1 = mul(float4(P1, 1), LocalToWorld).xyz;
 			const float3 WP2 = mul(float4(P2, 1), LocalToWorld).xyz;
-			AddLineTriangle(WP0, WP1, WP2, TriColor);
+			AddLineTriangleWS(WP0, WP1, WP2, TriColor);
 		}
 		#endif
 
@@ -220,8 +220,8 @@ void MainCS(uint3 DispatchThreadId : SV_DispatchThreadID, uint GroupThreadId : S
 	#if PERMUTATION_DEBUG
 	if (all(DispatchThreadId == 0))
 	{
-		AddAABB(WorldBound_Min, WorldBound_Max, ColorYellow);
-		AddOBB(ViewBound_Min, ViewBound_Max, ColorCyan, ViewToWorld);
+		AddAABBWS(WorldBound_Min, WorldBound_Max, ColorYellow);
+		AddOBBWS(ViewBound_Min, ViewBound_Max, ColorCyan, ViewToWorld);
 
 		FOITDebugData DebugData;
 		DebugData.PrimitiveCount= NumPrimitives;
diff --git a/Engine/Shaders/Private/PhysicsFieldVisualizer.usf b/Engine/Shaders/Private/PhysicsFieldVisualizer.usf
--- a/Engine/Shaders/Private/PhysicsFieldVisualizer.usf
+++ b/Engine/Shaders/Private/PhysicsFieldVisualizer.usf
@@ -255,14 +255,14 @@ void MainCS(uint3 DispatchThreadId : SV_DispatchThreadID,
 			{
 				FieldVector /= VectorLength;
 				const float4 FieldColor = float4((FieldVector + float3(1,1,1))/2.0, min(1.0, VectorLength));
-				AddLine(VoxelCenter, VoxelCenter + FieldVector * VoxelLength , FieldColor, FieldColor );
+				AddLineWS(VoxelCenter, VoxelCenter + FieldVector * VoxelLength , FieldColor, FieldColor );
 			}
 		}
 		if(DispatchThreadId.x == 0 && DispatchThreadId.y == 0 && DispatchThreadId.z == 0) 
 		{
 			for(int BoundsIndex = BoundsOffset, BoundsEnd = BoundsOffset + BoundsSize; BoundsIndex < BoundsEnd; ++BoundsIndex)
 			{
-				AddAABB(BoundsMin[BoundsIndex].xyz, BoundsMax[BoundsIndex].xyz, float4(1, 0.5f, 0, 1));
+				AddAABBWS(BoundsMin[BoundsIndex].xyz, BoundsMax[BoundsIndex].xyz, float4(1, 0.5f, 0, 1));
 			}
 			PrintVectorField();
 		}
@@ -643,7 +643,7 @@ void MainCS(uint3 DispatchThreadId : SV_DispatchThreadID,
 		{
 			for(int BoundsIndex = BoundsOffset, BoundsEnd = BoundsOffset + BoundsSize; BoundsIndex < BoundsEnd; ++BoundsIndex)
 			{
-				AddAABB(BoundsMin[BoundsIndex].xyz, BoundsMax[BoundsIndex].xyz, float4(1, 0.5f, 0, 1));
+				AddAABBWS(BoundsMin[BoundsIndex].xyz, BoundsMax[BoundsIndex].xyz, float4(1, 0.5f, 0, 1));
 			}
 #if PERMUTATION_FIELD == 2
 			PrintIntegerField();
diff --git a/Engine/Shaders/Private/ShaderDrawDebug.usf b/Engine/Shaders/Private/ShaderDrawDebug.usf
--- a/Engine/Shaders/Private/ShaderDrawDebug.usf
+++ b/Engine/Shaders/Private/ShaderDrawDebug.usf
@@ -23,10 +23,10 @@ void ShaderDrawDebugVS(
 	FPackedShaderDrawElement PackedItem = ShaderDrawDebugPrimitive[InstanceId];
 	FShaderDrawElement Item = UnpackShaderElement(PackedItem);
 
-	float4 VertexPosition = float4(VertexId == 0 ? Item.Pos0 : Item.Pos1, 1.0f);
+	float4 VertexPositionTWS = float4(VertexId == 0 ? Item.Pos0 : Item.Pos1, 1.0f);
 	float4 VertexColor = VertexId == 0 ? Item.Color0 : Item.Color1;
 
-	Position = mul(VertexPosition, LWCHackToFloat(PrimaryView.WorldToClip));
+	Position = mul(VertexPositionTWS, PrimaryView.TranslatedWorldToClip);
 	ColorAlpha = VertexColor;
 }
 #endif
diff --git a/Engine/Shaders/Private/ShaderDrawDebug.ush b/Engine/Shaders/Private/ShaderDrawDebug.ush
--- a/Engine/Shaders/Private/ShaderDrawDebug.ush
+++ b/Engine/Shaders/Private/ShaderDrawDebug.ush
@@ -6,7 +6,9 @@
 #define SHADER_DRAW_DEBUG 1
 #endif
 
+///////////////////////////////////////////////////////////////////////////////////////////////////
 // Set of predefined colors
+
 #define ColorWhite		float4(1, 1, 1, 1)
 #define ColorBlack		float4(0, 0, 0, 1)
 #define ColorRed		float4(1, 0, 0, 1)
@@ -22,6 +24,9 @@
 #define ColorEmerald	float4( 46.f / 255.f, 204.f / 255.f, 113.f / 255.f, 1)
 #define ColorLightGreen	float4(0.5f, 0.75f, 0.5f, 1.0f)
 
+///////////////////////////////////////////////////////////////////////////////////////////////////
+// Struct and core functions
+
 struct FShaderDrawElement 
 {
 	float3 Pos0;
@@ -69,13 +74,20 @@ RWBuffer<uint>								OutputShaderDrawIndirect;
 
 int2 ShaderDrawCursorPos;
 uint ShaderDrawMaxElementCount;
+float3 ShaderDrawTranslatedWorldOffset;
+
+///////////////////////////////////////////////////////////////////////////////////////////////////
+// Utils functions
 
 int2 GetCursorPos()
 {
 	return ShaderDrawCursorPos;
 }
 
-void AddLine(float3 Pos0, float3 Pos1, float4 Color0, float4 Color1)
+///////////////////////////////////////////////////////////////////////////////////////////////////
+// Line
+
+void AddLineTWS(float3 Pos0, float3 Pos1, float4 Color0, float4 Color1)
 {
 	if (ShaderDrawMaxElementCount == 0)
 	{
@@ -100,12 +112,16 @@ void AddLine(float3 Pos0, float3 Pos1, float4 Color0, float4 Color1)
 	}
 }
 
-void AddLine(float3 Pos0, float3 Pos1, float4 Color)
-{
-	AddLine(Pos0, Pos1, Color, Color);
-}
+void AddLineTWS(float3 Pos0, float3 Pos1, float4 Color)					{ AddLineTWS(Pos0, Pos1, Color, Color); }
+void AddLineWS(float3 Pos0, float3 Pos1, float4 Color)					{ AddLineTWS(Pos0+ShaderDrawTranslatedWorldOffset, Pos1+ShaderDrawTranslatedWorldOffset, Color); }
+void AddLineWS(float3 Pos0, float3 Pos1, float4 Color0, float4 Color1)	{ AddLineTWS(Pos0+ShaderDrawTranslatedWorldOffset, Pos1+ShaderDrawTranslatedWorldOffset, Color0, Color1); }
+void AddLine(float3 Pos0, float3 Pos1, float4 Color)					{ AddLineTWS(Pos0+ShaderDrawTranslatedWorldOffset, Pos1+ShaderDrawTranslatedWorldOffset, Color); }
+void AddLine(float3 Pos0, float3 Pos1, float4 Color0, float4 Color1)	{ AddLineTWS(Pos0+ShaderDrawTranslatedWorldOffset, Pos1+ShaderDrawTranslatedWorldOffset, Color0, Color1); }
+
+///////////////////////////////////////////////////////////////////////////////////////////////////
+// Quad
 
-void AddQuad(float3 Pos0, float3 Pos1, float3 Pos2, float3 Pos3, float4 Color)
+void AddQuadTWS(float3 Pos0, float3 Pos1, float3 Pos2, float3 Pos3, float4 Color)
 {
 	if (ShaderDrawMaxElementCount == 0)
 	{
@@ -137,7 +153,13 @@ void AddQuad(float3 Pos0, float3 Pos1, float3 Pos2, float3 Pos3, float4 Color)
 	}
 }
 
-void AddAABB(float3 Min, float3 Max, float4 Color)
+void AddQuadWS(float3 Pos0, float3 Pos1, float3 Pos2, float3 Pos3, float4 Color) { AddQuadTWS(Pos0+ShaderDrawTranslatedWorldOffset, Pos1+ShaderDrawTranslatedWorldOffset, Pos2+ShaderDrawTranslatedWorldOffset, Pos3+ShaderDrawTranslatedWorldOffset, Color); }
+void AddQuad(float3 Pos0, float3 Pos1, float3 Pos2, float3 Pos3, float4 Color)	 { AddQuadTWS(Pos0+ShaderDrawTranslatedWorldOffset, Pos1+ShaderDrawTranslatedWorldOffset, Pos2+ShaderDrawTranslatedWorldOffset, Pos3+ShaderDrawTranslatedWorldOffset, Color); }
+
+///////////////////////////////////////////////////////////////////////////////////////////////////
+// AABB
+
+void AddAABBTWS(float3 Min, float3 Max, float4 Color)
 {
 	if (ShaderDrawMaxElementCount == 0)
 	{
@@ -181,21 +203,33 @@ void AddAABB(float3 Min, float3 Max, float4 Color)
 	}
 }
 
-void AddCross(float3 Pos, float Size, float4 Color)
+void AddAABBWS(float3 Min, float3 Max, float4 Color)	{ AddAABBTWS(Min + ShaderDrawTranslatedWorldOffset, Max + ShaderDrawTranslatedWorldOffset, Color); }
+void AddAABB(float3 Min, float3 Max, float4 Color)		{ AddAABBTWS(Min + ShaderDrawTranslatedWorldOffset, Max + ShaderDrawTranslatedWorldOffset, Color); }
+
+///////////////////////////////////////////////////////////////////////////////////////////////////
+// Cross
+
+void AddCrossTWS(float3 Pos, float Size, float4 Color)
 {
-	AddLine(Pos - float3(Size,0,0), Pos + float3(Size,0,0), Color, Color);
-	AddLine(Pos - float3(0,Size,0), Pos + float3(0,Size,0), Color, Color);
-	AddLine(Pos - float3(0,0,Size), Pos + float3(0,0,Size), Color, Color);
+	AddLineTWS(Pos - float3(Size,0,0), Pos + float3(Size,0,0), Color, Color);
+	AddLineTWS(Pos - float3(0,Size,0), Pos + float3(0,Size,0), Color, Color);
+	AddLineTWS(Pos - float3(0,0,Size), Pos + float3(0,0,Size), Color, Color);
 }
 
-void AddReferential(float3 Pos, float3 T, float3 B, float3 N, float Scale=1)
+void AddCrossWS(float3 Pos, float Size, float4 Color)	{ AddAABBTWS(Pos + ShaderDrawTranslatedWorldOffset, Size, Color); }
+void AddCross(float3 Pos, float Size, float4 Color)		{ AddAABBTWS(Pos + ShaderDrawTranslatedWorldOffset, Size, Color); }
+
+///////////////////////////////////////////////////////////////////////////////////////////////////
+// Referential
+
+void AddReferentialTWS(float3 Pos, float3 T, float3 B, float3 N, float Scale=1)
 {
-	AddLine(Pos, Pos + T*Scale, ColorRed);
-	AddLine(Pos, Pos + B*Scale, ColorGreen);
-	AddLine(Pos, Pos + N*Scale, ColorBlue);
+	AddLineTWS(Pos, Pos + T*Scale, ColorRed);
+	AddLineTWS(Pos, Pos + B*Scale, ColorGreen);
+	AddLineTWS(Pos, Pos + N*Scale, ColorBlue);
 }
 
-void AddReferential(float3 Pos, float3 TangentZ, float Scale=1)
+void AddReferentialTWS(float3 Pos, float3 TangentZ, float Scale=1)
 {
 	const float Sign = TangentZ.z >= 0 ? 1 : -1;
 	const float a = -rcp(Sign + TangentZ.z);
@@ -204,22 +238,38 @@ void AddReferential(float3 Pos, float3 TangentZ, float Scale=1)
 	const float3 TangentX = { 1 + Sign * a * Pow2(TangentZ.x), Sign * b, -Sign * TangentZ.x };
 	const float3 TangentY = { b,  Sign + a * Pow2(TangentZ.y), -TangentZ.y };
 
-	AddReferential(Pos, TangentX, TangentY, TangentZ, Scale);
+	AddReferentialTWS(Pos, TangentX, TangentY, TangentZ, Scale);
 }
 
-void AddReferential(float4x4 InM, float Scale = 1)
+void AddReferentialTWS(float4x4 InM, float Scale = 1)
 {
-	AddReferential(InM[3].xyz, InM[0].xyz, InM[1].xyz, InM[2].xyz, Scale);
+	AddReferentialTWS(InM[3].xyz, InM[0].xyz, InM[1].xyz, InM[2].xyz, Scale);
 }
 
-void AddLineTriangle(float3 P0, float3 P1, float3 P2, float4 Color)
+void AddReferentialWS(float4x4 InM, float Scale = 1)							{ AddReferentialTWS(InM[3].xyz, InM[0].xyz, InM[1].xyz, InM[2].xyz + ShaderDrawTranslatedWorldOffset, Scale); }
+void AddReferentialWS(float3 Pos, float3 TangentZ, float Scale = 1)				{ AddReferentialTWS(Pos + ShaderDrawTranslatedWorldOffset, TangentZ, Scale); }
+void AddReferentialWS(float3 Pos, float3 T, float3 B, float3 N, float Scale = 1){ AddReferentialTWS(Pos + ShaderDrawTranslatedWorldOffset, T, B, N, Scale); }
+void AddReferential(float4x4 InM, float Scale = 1)								{ AddReferentialTWS(InM[3].xyz, InM[0].xyz, InM[1].xyz, InM[2].xyz + ShaderDrawTranslatedWorldOffset, Scale); }
+void AddReferential(float3 Pos, float3 TangentZ, float Scale = 1)				{ AddReferentialTWS(Pos + ShaderDrawTranslatedWorldOffset, TangentZ, Scale); }
+void AddReferential(float3 Pos, float3 T, float3 B, float3 N, float Scale = 1)	{ AddReferentialTWS(Pos + ShaderDrawTranslatedWorldOffset, T, B, N, Scale); }
+
+///////////////////////////////////////////////////////////////////////////////////////////////////
+// Triangle
+
+void AddLineTriangleTWS(float3 P0, float3 P1, float3 P2, float4 Color)
 {
-	AddLine(P0, P1, Color, Color);
-	AddLine(P1, P2, Color, Color);
-	AddLine(P2, P0, Color, Color);
+	AddLineTWS(P0, P1, Color, Color);
+	AddLineTWS(P1, P2, Color, Color);
+	AddLineTWS(P2, P0, Color, Color);
 }
 
-void AddOBB(float3 Min, float3 Max, float4 Color, float4x4 Transform)
+void AddLineTriangleWS(float3 P0, float3 P1, float3 P2, float4 Color)	{ AddLineTriangleTWS(P0 + ShaderDrawTranslatedWorldOffset, P1 + ShaderDrawTranslatedWorldOffset, P2 + ShaderDrawTranslatedWorldOffset, Color); }
+void AddLineTriangle(float3 P0, float3 P1, float3 P2, float4 Color)		{ AddLineTriangleTWS(P0 + ShaderDrawTranslatedWorldOffset, P1 + ShaderDrawTranslatedWorldOffset, P2 + ShaderDrawTranslatedWorldOffset, Color); }
+
+///////////////////////////////////////////////////////////////////////////////////////////////////
+// OBB
+
+void AddOBBTWS(float3 Min, float3 Max, float4 Color, float4x4 Transform)
 {
 	if (ShaderDrawMaxElementCount == 0)
 	{
@@ -263,4 +313,7 @@ void AddOBB(float3 Min, float3 Max, float4 Color, float4x4 Transform)
 	}
 }
 
+void AddOBBWS(float3 Min, float3 Max, float4 Color, float4x4 Transform) { Transform[2].xyz += ShaderDrawTranslatedWorldOffset; AddOBBTWS(Min, Max, Color, Transform); }
+void AddOBB(float3 Min, float3 Max, float4 Color, float4x4 Transform)	{ Transform[2].xyz += ShaderDrawTranslatedWorldOffset; AddOBBTWS(Min, Max, Color, Transform); }
+
 #endif
diff --git a/Engine/Shaders/Private/ShadingFurnaceTest.usf b/Engine/Shaders/Private/ShadingFurnaceTest.usf
--- a/Engine/Shaders/Private/ShadingFurnaceTest.usf
+++ b/Engine/Shaders/Private/ShadingFurnaceTest.usf
@@ -91,8 +91,8 @@ void MainPS(
 
 	const float4 PixelLightAttenuation = 1.0f;
 	const float3 InRadiance = 0.5f;
-	const float3 WorldPosition = mul(float4(ScreenPosition * SceneDepth, SceneDepth, 1), LWCHackToFloat(PrimaryView.ScreenToWorld)).xyz;
-	const float3 CameraVector = normalize(WorldPosition - LWCHackToFloat(PrimaryView.WorldCameraOrigin));
+	const float3 WorldPosition = mul(float4(ScreenPosition * SceneDepth, SceneDepth, 1), PrimaryView.ScreenToTranslatedWorld).xyz;
+	const float3 CameraVector = normalize(WorldPosition - PrimaryView.TranslatedWorldCameraOrigin);
 	const float3 TranslatedWorldPosition = WorldPosition;
 	const float3 V = -CameraVector;
 
@@ -102,8 +102,8 @@ void MainPS(
 	const bool bDebugEnabled = all(CursorCoord == uint2(PixelPos));
 	if (bDebugEnabled)
 	{
-		AddReferential(WorldPosition, TangentBasis[0], TangentBasis[1], TangentBasis[2], 3.f);
-		AddLine(WorldPosition, WorldPosition + V * 3.0f, ColorPurple);
+		AddReferentialTWS(WorldPosition, TangentBasis[0], TangentBasis[1], TangentBasis[2], 3.f);
+		AddLineTWS(WorldPosition, WorldPosition + V * 3.0f, ColorPurple);
 	}
 #endif
 
@@ -174,7 +174,7 @@ void MainPS(
 				{
 					const float Scale = Set == 0 ? 1.f : 2.0f;
 					const float4 DebugColor = Set == 0 ? ColorOrange : ColorEmerald;
-					AddLine(WorldPosition, WorldPosition + L * Scale, DebugColor);
+					AddLineTWS(WorldPosition, WorldPosition + L * Scale, DebugColor);
 				}
 				#endif
 
diff --git a/Engine/Shaders/Private/Strata/StrataVisualize.usf b/Engine/Shaders/Private/Strata/StrataVisualize.usf
--- a/Engine/Shaders/Private/Strata/StrataVisualize.usf
+++ b/Engine/Shaders/Private/Strata/StrataVisualize.usf
@@ -12,7 +12,7 @@
 #include "/Engine/Private/ShaderDrawDebug.ush"
 
 
-void DrawReferential(float3 P, float3 X, float3 Y, float3 N, float3 InColor)
+void DrawReferential(float3 P /*In Translated World Space*/, float3 X, float3 Y, float3 N, float3 InColor)
 {
 	const float Size = 10.f;
 	const float SizeWithTip = 12.0f;
@@ -22,13 +22,13 @@ void DrawReferential(float3 P, float3 X, float3 Y, float3 N, float3 InColor)
 	const float4 ColorZ = float4(0, 0, 1, 1);
 
 	// Core
-	AddLine(P, P + X * Size, Color, Color);
-	AddLine(P, P + Y * Size, Color, Color);
-	AddLine(P, P + N * Size, Color, Color);
+	AddLineTWS(P, P + X * Size, Color, Color);
+	AddLineTWS(P, P + Y * Size, Color, Color);
+	AddLineTWS(P, P + N * Size, Color, Color);
 	// Tips
-	AddLine(P + X * Size, P + X * SizeWithTip, ColorX, ColorX);
-	AddLine(P + Y * Size, P + Y * SizeWithTip, ColorY, ColorY);
-	AddLine(P + N * Size, P + N * SizeWithTip, ColorZ, ColorZ);
+	AddLineTWS(P + X * Size, P + X * SizeWithTip, ColorX, ColorX);
+	AddLineTWS(P + Y * Size, P + Y * SizeWithTip, ColorY, ColorY);
+	AddLineTWS(P + N * Size, P + N * SizeWithTip, ColorZ, ColorZ);
 }
 
 float3 DrawPlot(
@@ -88,8 +88,8 @@ void VisualizeMaterialPS(
 	const float2 ScreenMaterialPosition = SvPositionToScreenPosition(float4(MaterialPos, 0.5f, 1)).xy;
 	float DeviceZ = SampleDeviceZFromSceneTextures(BufferUV);
 	float SceneDepth = ConvertFromDeviceZ(DeviceZ);
-	const float3 WorldPosition = mul(float4(ScreenMaterialPosition * SceneDepth, SceneDepth, 1), LWCHackToFloat(PrimaryView.ScreenToWorld)).xyz;
-	const float3 V = -normalize(WorldPosition - LWCHackToFloat(PrimaryView.WorldCameraOrigin));
+	const float3 WorldPosition = mul(float4(ScreenMaterialPosition * SceneDepth, SceneDepth, 1), PrimaryView.ScreenToTranslatedWorld).xyz;
+	const float3 V = -normalize(WorldPosition - PrimaryView.TranslatedWorldCameraOrigin);
 
 #if 0
 	OutColor = 0;
@@ -154,7 +154,7 @@ void VisualizeMaterialPS(
 			if (all(PixelPos == MaterialPos))
 			{
 				FStrataBSDFContext BSDFContext = StrataCreateBSDFContext(Header, BSDF, StrataAddressing, V, float3(0,0,1));
-				DrawReferential(WorldPosition, BSDFContext.X, BSDFContext.Y, BSDFContext.N, float3(1,1,0));
+				DrawReferentialTWS(WorldPosition, BSDFContext.X, BSDFContext.Y, BSDFContext.N, float3(1,1,0));
 			}
 
 			// Plot
diff --git a/Engine/Shaders/Private/VirtualShadowMaps/CacheManagement.usf b/Engine/Shaders/Private/VirtualShadowMaps/CacheManagement.usf
--- a/Engine/Shaders/Private/VirtualShadowMaps/CacheManagement.usf
+++ b/Engine/Shaders/Private/VirtualShadowMaps/CacheManagement.usf
@@ -204,7 +204,7 @@ void VirtualSmInvalidateInstancePagesCS(
 		if (bDrawBounds && PageInvalidationCount > 0U)
 		{
 			float3 Color = float3(0.3f, 0.3f, 0.3f) + ColorMapTurbo(min(1.0f, float(PageInvalidationCount) / 100.0f)) * 0.7f;
-			AddOBB(InstanceSceneData.LocalBoundsCenter - InstanceSceneData.LocalBoundsExtent, InstanceSceneData.LocalBoundsCenter + InstanceSceneData.LocalBoundsExtent, float4(Color, 1.0f), LWCHackToFloat(InstanceSceneData.LocalToWorld));
+			AddOBBWS(InstanceSceneData.LocalBoundsCenter - InstanceSceneData.LocalBoundsExtent, InstanceSceneData.LocalBoundsCenter + InstanceSceneData.LocalBoundsExtent, float4(Color, 1.0f), LWCHackToFloat(InstanceSceneData.LocalToWorld));
 		}
 #endif
 	}
diff --git a/Engine/Shaders/Private/VolumetricCloud.usf b/Engine/Shaders/Private/VolumetricCloud.usf
--- a/Engine/Shaders/Private/VolumetricCloud.usf
+++ b/Engine/Shaders/Private/VolumetricCloud.usf
@@ -1737,7 +1737,7 @@ void MainDrawDebugShadowCS(uint3 DispatchThreadId : SV_DispatchThreadID)
 
 		if (OpticalDepth > 0.0f)
 		{
-			AddQuad(NearClipPlaneWorldPos00, NearClipPlaneWorldPos01, NearClipPlaneWorldPos11, NearClipPlaneWorldPos10, DebugColor);
+			AddQuadWS(NearClipPlaneWorldPos00, NearClipPlaneWorldPos01, NearClipPlaneWorldPos11, NearClipPlaneWorldPos10, DebugColor);
 		}
 	}
 }
diff --git a/Engine/Source/Runtime/Renderer/Private/HairStrands/HairStrandsDebug.cpp b/Engine/Source/Runtime/Renderer/Private/HairStrands/HairStrandsDebug.cpp
--- a/Engine/Source/Runtime/Renderer/Private/HairStrands/HairStrandsDebug.cpp
+++ b/Engine/Source/Runtime/Renderer/Private/HairStrands/HairStrandsDebug.cpp
@@ -1265,6 +1265,8 @@ static void InternalRenderHairStrandsDebugInfo(
 
 		FViewElementPDI ShadowFrustumPDI(&View, nullptr, nullptr);
 
+		// Do not draw the CPU boud. Only GPU bound are drawn using ShaderDrawDebug shader functions.
+		#if 0
 		// All DOMs
 		if (HairDebugMode == EHairDebugMode::LightBounds)
 		{
@@ -1314,6 +1316,7 @@ static void InternalRenderHairStrandsDebugInfo(
 				}
 			}
 		}
+		#endif
 	}
 	
 	const bool bRunDebugPass =
diff --git a/Engine/Source/Runtime/Renderer/Private/ShaderDebug.cpp b/Engine/Source/Runtime/Renderer/Private/ShaderDebug.cpp
--- a/Engine/Source/Runtime/Renderer/Private/ShaderDebug.cpp
+++ b/Engine/Source/Runtime/Renderer/Private/ShaderDebug.cpp
@@ -281,6 +281,7 @@ namespace ShaderDrawDebug
 		View.ShaderDrawData.Buffer = DataBuffer;
 		View.ShaderDrawData.IndirectBuffer = IndirectBuffer;
 		View.ShaderDrawData.CursorPosition = View.CursorPos;
+		View.ShaderDrawData.ShaderDrawTranslatedWorldOffset = View.ViewMatrices.GetPreViewTranslation();
 
 		if (IsShaderDrawLocked() && ViewState && !ViewState->ShaderDrawDebugStateData.bIsLocked)
 		{
@@ -341,10 +342,11 @@ namespace ShaderDrawDebug
 		if (Data.Buffer == nullptr)
 			return;
 
-		OutParameters.ShaderDrawCursorPos		= Data.CursorPosition;
-		OutParameters.ShaderDrawMaxElementCount = Data.MaxElementCount;
-		OutParameters.OutShaderDrawPrimitive	= GraphBuilder.CreateUAV(Data.Buffer);
-		OutParameters.OutputShaderDrawIndirect	= GraphBuilder.CreateUAV(Data.IndirectBuffer);
+		OutParameters.ShaderDrawCursorPos				= Data.CursorPosition;
+		OutParameters.ShaderDrawMaxElementCount			= Data.MaxElementCount;
+		OutParameters.ShaderDrawTranslatedWorldOffset	= Data.ShaderDrawTranslatedWorldOffset;
+		OutParameters.OutShaderDrawPrimitive			= GraphBuilder.CreateUAV(Data.Buffer);
+		OutParameters.OutputShaderDrawIndirect			= GraphBuilder.CreateUAV(Data.IndirectBuffer);
 	}
 
 	// Returns true if the default view exists and has shader debug rendering enabled (this needs to be checked before using a permutation that requires the shader draw parameters)
diff --git a/Engine/Source/Runtime/Renderer/Public/ShaderDebug.h b/Engine/Source/Runtime/Renderer/Public/ShaderDebug.h
--- a/Engine/Source/Runtime/Renderer/Public/ShaderDebug.h
+++ b/Engine/Source/Runtime/Renderer/Public/ShaderDebug.h
@@ -12,6 +12,7 @@ struct FShaderDrawDebugData
 {
 	uint32 MaxElementCount = 0u;
 	FIntPoint CursorPosition = FIntPoint(-1,-1);
+	FVector ShaderDrawTranslatedWorldOffset;
 	FRDGBufferRef Buffer = nullptr;
 	FRDGBufferRef IndirectBuffer = nullptr;
 	bool IsEnabled() const { return MaxElementCount > 0; }
@@ -50,6 +51,7 @@ namespace ShaderDrawDebug
 	BEGIN_SHADER_PARAMETER_STRUCT(FShaderParameters, )
 		SHADER_PARAMETER(FIntPoint, ShaderDrawCursorPos)
 		SHADER_PARAMETER(int32, ShaderDrawMaxElementCount)
+		SHADER_PARAMETER(FVector3f, ShaderDrawTranslatedWorldOffset)
 		SHADER_PARAMETER_RDG_BUFFER_UAV(RWStructuredBuffer, OutShaderDrawPrimitive)
 		SHADER_PARAMETER_RDG_BUFFER_UAV(RWBuffer<uint>, OutputShaderDrawIndirect)
 	END_SHADER_PARAMETER_STRUCT()
