Fix a few more cases of subtracting PreViewTranslation rather than adding it

diff --git a/Engine/Shaders/Private/TranslucentLightInjectionShaders.usf b/Engine/Shaders/Private/TranslucentLightInjectionShaders.usf
--- a/Engine/Shaders/Private/TranslucentLightInjectionShaders.usf
+++ b/Engine/Shaders/Private/TranslucentLightInjectionShaders.usf
@@ -137,7 +137,7 @@ void InjectMainPS(
 
 		WorldPosition = float3(View.TranslucencyLightingVolumeMin[VolumeCascadeIndex].xy + Input.Vertex.UV / View.TranslucencyLightingVolumeInvSize[VolumeCascadeIndex].xy - .5f * View.TranslucencyLightingVolumeInvSize[VolumeCascadeIndex].w, ZPosition);
 	}
-	float3 TranslatedWorldPosition = WorldPosition - LWCHackToFloat(PrimaryView.PreViewTranslation);
+	float3 TranslatedWorldPosition = WorldPosition + LWCHackToFloat(PrimaryView.PreViewTranslation);
 
 	// compute UVW of the position we shade in the volume
 	float3 VolumeUVs = float3(Input.Vertex.UV, (Input.LayerIndex + .5f) * View.TranslucencyLightingVolumeMin[VolumeCascadeIndex].w);
diff --git a/Engine/Shaders/Private/VolumetricFog.usf b/Engine/Shaders/Private/VolumetricFog.usf
--- a/Engine/Shaders/Private/VolumetricFog.usf
+++ b/Engine/Shaders/Private/VolumetricFog.usf
@@ -219,7 +219,7 @@ float GetLocalLightFunction(float3 WorldPosition, bool IsPointLight)
 	if (IsPointLight)
 	{
 		// Get vector from light to cell world pos.
-		float TranslatedWorldPosition = WorldPosition - LWCHackToFloat(PrimaryView.PreViewTranslation);
+		float TranslatedWorldPosition = WorldPosition + LWCHackToFloat(PrimaryView.PreViewTranslation);
 		float3 LightVec = normalize(TranslatedWorldPosition - DeferredLightUniforms.TranslatedWorldPosition);
 		// Transform and swizzle
 		LightVec = mul(float4(LightVec.xyz, 0), LocalLightFunctionMatrix).zyx; 
