Change shader draw debug use a single UAV when bound.

Shader draw debug now uses a single buffer for tracking primitive count and primitive data.

diff --git a/Engine/Shaders/Private/ShaderDrawDebug.usf b/Engine/Shaders/Private/ShaderDrawDebug.usf
--- a/Engine/Shaders/Private/ShaderDrawDebug.usf
+++ b/Engine/Shaders/Private/ShaderDrawDebug.usf
@@ -6,7 +6,7 @@
 #include "ShaderDrawDebug.ush"
 
 // The shader used to render the debug element on view
-StructuredBuffer<FPackedShaderDrawElement>	ShaderDrawDebugPrimitive;
+StructuredBuffer<uint> ShaderDrawDebugPrimitive;
 
 Texture2D<float> DepthTexture;
 int2 DepthTextureResolution;
@@ -20,8 +20,7 @@ void ShaderDrawDebugVS(
 	out float4 Position : SV_POSITION,
 	out float4 ColorAlpha : TEXCOORD0)
 {
-	FPackedShaderDrawElement PackedItem = ShaderDrawDebugPrimitive[InstanceId];
-	FShaderDrawElement Item = UnpackShaderElement(PackedItem);
+	FShaderDrawElement Item = UnpackShaderElement(ShaderDrawDebugPrimitive, InstanceId);
 
 	float4 VertexPositionTWS = float4(VertexId == 0 ? Item.Pos0 : Item.Pos1, 1.0f);
 	float4 VertexColor = VertexId == 0 ? Item.Color0 : Item.Color1;
@@ -52,15 +51,33 @@ void ShaderDrawDebugPS(
 
 #if GPU_DEBUG_RENDERING_CLEAR_CS
 
-RWStructuredBuffer<FPackedShaderDrawElement> ElementBuffer; // Referenced so that the resource is marked as used
-RWBuffer<uint> IndirectBuffer;
+RWStructuredBuffer<uint> RWElementBuffer;
 
 [numthreads(1, 1, 1)]
 void ShaderDrawDebugClearCS(uint3 DispatchThreadId : SV_DispatchThreadID)
 {
-	IndirectBuffer[0] = 2;
-	IndirectBuffer[1] = 0;
-	IndirectBuffer[2] = 0;
-	IndirectBuffer[3] = 0;
+	RWElementBuffer[0] = 0; // Only element which really need to be cleared
+	RWElementBuffer[1] = 0;
+	RWElementBuffer[2] = 0;
+	RWElementBuffer[3] = 0;
+	RWElementBuffer[4] = 0;
+	RWElementBuffer[5] = 0;
+	RWElementBuffer[6] = 0;
+	RWElementBuffer[7] = 0;
 }
 #endif
+
+#if GPU_DEBUG_RENDERING_COPY_CS
+
+StructuredBuffer<uint> ElementBuffer;
+RWBuffer<uint> RWIndirectArgs;
+
+[numthreads(1, 1, 1)]
+void ShaderDrawDebugCopyCS(uint3 DispatchThreadId : SV_DispatchThreadID)
+{
+	RWIndirectArgs[0] = 2; // Line primitive
+	RWIndirectArgs[1] = ElementBuffer[0];
+	RWIndirectArgs[2] = 0;
+	RWIndirectArgs[3] = 0;
+}
+#endif
\ No newline at end of file
diff --git a/Engine/Shaders/Private/ShaderDrawDebug.ush b/Engine/Shaders/Private/ShaderDrawDebug.ush
--- a/Engine/Shaders/Private/ShaderDrawDebug.ush
+++ b/Engine/Shaders/Private/ShaderDrawDebug.ush
@@ -35,50 +35,72 @@ struct FShaderDrawElement
 	float4 Color1;
 };
 
-// Note: Unaligned structures used for structured buffers is an unsupported and/or sparsely
-//         supported feature in VK (VK_EXT_scalar_block_layout) and Metal. Consequently, we do
-//         manual packing in order to accommodate.
-struct FPackedShaderDrawElement 
+FShaderDrawElement UnpackShaderElement(StructuredBuffer<uint> InPrimitiveBuffer, uint InIndex)
 {
-	float4 Pos0_ColorX;		// float3 Pos0 + Color.x
-	float4 Pos1_ColorY;		// float3 Pos1 + Color.y
-};
+	const uint Index8 = (1 + InIndex) * 8;
 
-FPackedShaderDrawElement PackShaderElement(FShaderDrawElement E)
-{
-	uint4 PackedC0 = uint4(255.0f * saturate(E.Color0));
-	uint4 PackedC1 = uint4(255.0f * saturate(E.Color1));
+	FShaderDrawElement Out = (FShaderDrawElement)0;
+	{
+		Out.Pos0.x = asfloat(InPrimitiveBuffer[Index8 + 0]);
+		Out.Pos0.y = asfloat(InPrimitiveBuffer[Index8 + 1]);
+		Out.Pos0.z = asfloat(InPrimitiveBuffer[Index8 + 2]);
 
-	FPackedShaderDrawElement PackedE;
-	PackedE.Pos0_ColorX = float4(E.Pos0, asfloat((PackedC0.x << 24) | (PackedC0.y << 16) | (PackedC0.z << 8) | (PackedC0.w)));
-	PackedE.Pos1_ColorY = float4(E.Pos1, asfloat((PackedC1.x << 24) | (PackedC1.y << 16) | (PackedC1.z << 8) | (PackedC1.w)));
-	return PackedE;
-}
+		const uint PackedColor0 = InPrimitiveBuffer[Index8 + 3];
+		Out.Color0 = float4((PackedColor0.x >> 24) & 0xFF, (PackedColor0.x >> 16) & 0xFF, (PackedColor0.x >> 8) & 0xFF, (PackedColor0.x) & 0xFF) / 255.0f;
+	}
+	{
+		Out.Pos1.x = asfloat(InPrimitiveBuffer[Index8 + 4]);
+		Out.Pos1.y = asfloat(InPrimitiveBuffer[Index8 + 5]);
+		Out.Pos1.z = asfloat(InPrimitiveBuffer[Index8 + 6]);
 
-FShaderDrawElement UnpackShaderElement(FPackedShaderDrawElement PackedE)
-{
-	uint2 PackedColor01 = uint2(asuint(PackedE.Pos0_ColorX.w), asuint(PackedE.Pos1_ColorY.w));
-
-	FShaderDrawElement E;
-	E.Pos0 = PackedE.Pos0_ColorX.xyz;
-	E.Pos1 = PackedE.Pos1_ColorY.xyz;
-	E.Color0 = float4((PackedColor01.x >> 24) & 0xFF, (PackedColor01.x >> 16) & 0xFF, (PackedColor01.x >> 8) & 0xFF, (PackedColor01.x) & 0xFF) / 255.0f;
-	E.Color1 = float4((PackedColor01.y >> 24) & 0xFF, (PackedColor01.y >> 16) & 0xFF, (PackedColor01.y >> 8) & 0xFF, (PackedColor01.y) & 0xFF) / 255.0f;
-	return E;
+		const uint PackedColor1 = InPrimitiveBuffer[Index8 + 7];
+		Out.Color1 = float4((PackedColor1.x >> 24) & 0xFF, (PackedColor1.x >> 16) & 0xFF, (PackedColor1.x >> 8) & 0xFF, (PackedColor1.x) & 0xFF) / 255.0f;
+	}
+	return Out;
 }
 
-#if SHADER_DRAW_DEBUG == 1
+///////////////////////////////////////////////////////////////////////////////////////////////////
+// Parameters
 
-RWStructuredBuffer<FPackedShaderDrawElement>OutShaderDrawPrimitive;
-RWBuffer<uint>								OutputShaderDrawIndirect;
+#if SHADER_DRAW_DEBUG == 1
 
-int2 ShaderDrawCursorPos;
-uint ShaderDrawMaxElementCount;
-float3 ShaderDrawTranslatedWorldOffset;
+RWStructuredBuffer<uint>	OutShaderDrawPrimitive;
+int2						ShaderDrawCursorPos;
+uint						ShaderDrawMaxElementCount;
+float3						ShaderDrawTranslatedWorldOffset;
 
 ///////////////////////////////////////////////////////////////////////////////////////////////////
 // Utils functions
 
+bool AllocateShaderDrawElement(uint Count, inout uint OutIndex)
+{
+	OutIndex = 0;
+	if (ShaderDrawMaxElementCount == 0)
+	{
+		return false;
+	}
+
+	InterlockedAdd(OutShaderDrawPrimitive[0], Count, OutIndex);
+	return (1+OutIndex) < ShaderDrawMaxElementCount;
+}
+
+void AddShaderElement(FShaderDrawElement E, uint Index)
+{
+	const uint4 PackedC0 = uint4(255.0f * saturate(E.Color0));
+	const uint4 PackedC1 = uint4(255.0f * saturate(E.Color1));
+
+	const uint Index8 = (1 + Index) * 8;
+	OutShaderDrawPrimitive[Index8 + 0] = asuint(E.Pos0.x);
+	OutShaderDrawPrimitive[Index8 + 1] = asuint(E.Pos0.y);
+	OutShaderDrawPrimitive[Index8 + 2] = asuint(E.Pos0.z);
+	OutShaderDrawPrimitive[Index8 + 3] = (PackedC0.x << 24) | (PackedC0.y << 16) | (PackedC0.z << 8) | (PackedC0.w);
+
+	OutShaderDrawPrimitive[Index8 + 4] = asuint(E.Pos1.x);
+	OutShaderDrawPrimitive[Index8 + 5] = asuint(E.Pos1.y);
+	OutShaderDrawPrimitive[Index8 + 6] = asuint(E.Pos1.z);
+	OutShaderDrawPrimitive[Index8 + 7] = (PackedC1.x << 24) | (PackedC1.y << 16) | (PackedC1.z << 8) | (PackedC1.w);
+}
+
 int2 GetCursorPos()
 {
 	return ShaderDrawCursorPos;
@@ -89,26 +111,15 @@ int2 GetCursorPos()
 
 void AddLineTWS(float3 Pos0, float3 Pos1, float4 Color0, float4 Color1)
 {
-	if (ShaderDrawMaxElementCount == 0)
-	{
-		return;
-	}
-
-	uint PreviousPrimitiveCount = 0;
-	InterlockedAdd(OutputShaderDrawIndirect[1], 1, PreviousPrimitiveCount);
-	if ((PreviousPrimitiveCount+1) <= ShaderDrawMaxElementCount)
+	uint Offset = 0;
+	if (AllocateShaderDrawElement(1, Offset))
 	{
 		FShaderDrawElement Element;
 		Element.Pos0 = Pos0;
 		Element.Pos1 = Pos1;
 		Element.Color0 = Color0;
 		Element.Color1 = Color1;
-
-		OutShaderDrawPrimitive[PreviousPrimitiveCount] = PackShaderElement(Element);
-	}
-	else
-	{
-		InterlockedAdd(OutputShaderDrawIndirect[1], -1, PreviousPrimitiveCount);// Make sure we stay in the safe zone for indirect draw
+		AddShaderElement(Element, Offset);
 	}
 }
 
@@ -123,33 +134,15 @@ void AddLine(float3 Pos0, float3 Pos1, float4 Color0, float4 Color1)	{ AddLineTW
 
 void AddQuadTWS(float3 Pos0, float3 Pos1, float3 Pos2, float3 Pos3, float4 Color)
 {
-	if (ShaderDrawMaxElementCount == 0)
-	{
-		return;
-	}
-
-	uint PreviousPrimitiveCount = 0;
-	InterlockedAdd(OutputShaderDrawIndirect[1], 4, PreviousPrimitiveCount);
-	if ((PreviousPrimitiveCount+4) <= ShaderDrawMaxElementCount)
+	uint Offset = 0;
+	if (AllocateShaderDrawElement(4, Offset))
 	{
 		FShaderDrawElement Element;
 		Element.Color0 = Element.Color1 = Color;
-
-		Element.Pos0 = Pos0; Element.Pos1 = Pos1;
-		OutShaderDrawPrimitive[PreviousPrimitiveCount + 0] = PackShaderElement(Element);
-
-		Element.Pos0 = Pos1; Element.Pos1 = Pos2;
-		OutShaderDrawPrimitive[PreviousPrimitiveCount + 1] = PackShaderElement(Element);
-
-		Element.Pos0 = Pos2; Element.Pos1 = Pos3;
-		OutShaderDrawPrimitive[PreviousPrimitiveCount + 2] = PackShaderElement(Element);
-
-		Element.Pos0 = Pos3; Element.Pos1 = Pos0;
-		OutShaderDrawPrimitive[PreviousPrimitiveCount + 3] = PackShaderElement(Element);
-	}
-	else
-	{
-		InterlockedAdd(OutputShaderDrawIndirect[1], -4, PreviousPrimitiveCount);// Make sure we stay in the safe zone for indirect draw
+		Element.Pos0 = Pos0; Element.Pos1 = Pos1; AddShaderElement(Element, Offset + 0);
+		Element.Pos0 = Pos1; Element.Pos1 = Pos2; AddShaderElement(Element, Offset + 1);
+		Element.Pos0 = Pos2; Element.Pos1 = Pos3; AddShaderElement(Element, Offset + 2);
+		Element.Pos0 = Pos3; Element.Pos1 = Pos0; AddShaderElement(Element, Offset + 3);
 	}
 }
 
@@ -161,14 +154,8 @@ void AddQuad(float3 Pos0, float3 Pos1, float3 Pos2, float3 Pos3, float4 Color)
 
 void AddAABBTWS(float3 Min, float3 Max, float4 Color)
 {
-	if (ShaderDrawMaxElementCount == 0)
-	{
-		return;
-	}
-
-	uint PreviousPrimitiveCount = 0;
-	InterlockedAdd(OutputShaderDrawIndirect[1], 12, PreviousPrimitiveCount);
-	if ((PreviousPrimitiveCount + 12) <= ShaderDrawMaxElementCount)
+	uint Offset = 0;
+	if (AllocateShaderDrawElement(12, Offset))
 	{
 		float3 P0 = float3(Min.x, Min.y, Min.z);
 		float3 P1 = float3(Max.x, Min.y, Min.z);
@@ -182,24 +169,20 @@ void AddAABBTWS(float3 Min, float3 Max, float4 Color)
 		FShaderDrawElement Element;
 		Element.Color0 = Element.Color1 = Color;
 
-		Element.Pos0 = P0; Element.Pos1 = P1; OutShaderDrawPrimitive[PreviousPrimitiveCount + 0] = PackShaderElement(Element);
-		Element.Pos0 = P1; Element.Pos1 = P2; OutShaderDrawPrimitive[PreviousPrimitiveCount + 1] = PackShaderElement(Element);
-		Element.Pos0 = P2; Element.Pos1 = P3; OutShaderDrawPrimitive[PreviousPrimitiveCount + 2] = PackShaderElement(Element);
-		Element.Pos0 = P3; Element.Pos1 = P0; OutShaderDrawPrimitive[PreviousPrimitiveCount + 3] = PackShaderElement(Element);
+		Element.Pos0 = P0; Element.Pos1 = P1; AddShaderElement(Element, Offset + 0);
+		Element.Pos0 = P1; Element.Pos1 = P2; AddShaderElement(Element, Offset + 1);
+		Element.Pos0 = P2; Element.Pos1 = P3; AddShaderElement(Element, Offset + 2);
+		Element.Pos0 = P3; Element.Pos1 = P0; AddShaderElement(Element, Offset + 3);
 
-		Element.Pos0 = P4; Element.Pos1 = P5; OutShaderDrawPrimitive[PreviousPrimitiveCount + 4] = PackShaderElement(Element);
-		Element.Pos0 = P5; Element.Pos1 = P6; OutShaderDrawPrimitive[PreviousPrimitiveCount + 5] = PackShaderElement(Element);
-		Element.Pos0 = P6; Element.Pos1 = P7; OutShaderDrawPrimitive[PreviousPrimitiveCount + 6] = PackShaderElement(Element);
-		Element.Pos0 = P7; Element.Pos1 = P4; OutShaderDrawPrimitive[PreviousPrimitiveCount + 7] = PackShaderElement(Element);
+		Element.Pos0 = P4; Element.Pos1 = P5; AddShaderElement(Element, Offset + 4);
+		Element.Pos0 = P5; Element.Pos1 = P6; AddShaderElement(Element, Offset + 5);
+		Element.Pos0 = P6; Element.Pos1 = P7; AddShaderElement(Element, Offset + 6);
+		Element.Pos0 = P7; Element.Pos1 = P4; AddShaderElement(Element, Offset + 7);
 
-		Element.Pos0 = P0; Element.Pos1 = P4; OutShaderDrawPrimitive[PreviousPrimitiveCount + 8] = PackShaderElement(Element);
-		Element.Pos0 = P1; Element.Pos1 = P5; OutShaderDrawPrimitive[PreviousPrimitiveCount + 9] = PackShaderElement(Element);
-		Element.Pos0 = P2; Element.Pos1 = P6; OutShaderDrawPrimitive[PreviousPrimitiveCount +10] = PackShaderElement(Element);
-		Element.Pos0 = P3; Element.Pos1 = P7; OutShaderDrawPrimitive[PreviousPrimitiveCount +11] = PackShaderElement(Element);
-	}
-	else
-	{
-		InterlockedAdd(OutputShaderDrawIndirect[1], -12, PreviousPrimitiveCount);// Make sure we stay in the safe zone for indirect draw
+		Element.Pos0 = P0; Element.Pos1 = P4; AddShaderElement(Element, Offset + 8);
+		Element.Pos0 = P1; Element.Pos1 = P5; AddShaderElement(Element, Offset + 9);
+		Element.Pos0 = P2; Element.Pos1 = P6; AddShaderElement(Element, Offset +10);
+		Element.Pos0 = P3; Element.Pos1 = P7; AddShaderElement(Element, Offset +11);
 	}
 }
 
@@ -276,9 +259,8 @@ void AddOBBTWS(float3 Min, float3 Max, float4 Color, float4x4 Transform)
 		return;
 	}
 
-	uint PreviousPrimitiveCount = 0;
-	InterlockedAdd(OutputShaderDrawIndirect[1], 12, PreviousPrimitiveCount);
-	if ((PreviousPrimitiveCount + 12) <= ShaderDrawMaxElementCount)
+	uint Offset = 0;
+	if (AllocateShaderDrawElement(12, Offset))
 	{
 		float3 P0 = mul(float4(Min.x, Min.y, Min.z, 1.0f), Transform).xyz;
 		float3 P1 = mul(float4(Max.x, Min.y, Min.z, 1.0f), Transform).xyz;
@@ -292,24 +274,20 @@ void AddOBBTWS(float3 Min, float3 Max, float4 Color, float4x4 Transform)
 		FShaderDrawElement Element;
 		Element.Color0 = Element.Color1 = Color;
 
-		Element.Pos0 = P0; Element.Pos1 = P1; OutShaderDrawPrimitive[PreviousPrimitiveCount + 0] = PackShaderElement(Element);
-		Element.Pos0 = P1; Element.Pos1 = P2; OutShaderDrawPrimitive[PreviousPrimitiveCount + 1] = PackShaderElement(Element);
-		Element.Pos0 = P2; Element.Pos1 = P3; OutShaderDrawPrimitive[PreviousPrimitiveCount + 2] = PackShaderElement(Element);
-		Element.Pos0 = P3; Element.Pos1 = P0; OutShaderDrawPrimitive[PreviousPrimitiveCount + 3] = PackShaderElement(Element);
+		Element.Pos0 = P0; Element.Pos1 = P1; AddShaderElement(Element, Offset + 0);
+		Element.Pos0 = P1; Element.Pos1 = P2; AddShaderElement(Element, Offset + 1);
+		Element.Pos0 = P2; Element.Pos1 = P3; AddShaderElement(Element, Offset + 2);
+		Element.Pos0 = P3; Element.Pos1 = P0; AddShaderElement(Element, Offset + 3);
 
-		Element.Pos0 = P4; Element.Pos1 = P5; OutShaderDrawPrimitive[PreviousPrimitiveCount + 4] = PackShaderElement(Element);
-		Element.Pos0 = P5; Element.Pos1 = P6; OutShaderDrawPrimitive[PreviousPrimitiveCount + 5] = PackShaderElement(Element);
-		Element.Pos0 = P6; Element.Pos1 = P7; OutShaderDrawPrimitive[PreviousPrimitiveCount + 6] = PackShaderElement(Element);
-		Element.Pos0 = P7; Element.Pos1 = P4; OutShaderDrawPrimitive[PreviousPrimitiveCount + 7] = PackShaderElement(Element);
+		Element.Pos0 = P4; Element.Pos1 = P5; AddShaderElement(Element, Offset + 4);
+		Element.Pos0 = P5; Element.Pos1 = P6; AddShaderElement(Element, Offset + 5);
+		Element.Pos0 = P6; Element.Pos1 = P7; AddShaderElement(Element, Offset + 6);
+		Element.Pos0 = P7; Element.Pos1 = P4; AddShaderElement(Element, Offset + 7);
 
-		Element.Pos0 = P0; Element.Pos1 = P4; OutShaderDrawPrimitive[PreviousPrimitiveCount + 8] = PackShaderElement(Element);
-		Element.Pos0 = P1; Element.Pos1 = P5; OutShaderDrawPrimitive[PreviousPrimitiveCount + 9] = PackShaderElement(Element);
-		Element.Pos0 = P2; Element.Pos1 = P6; OutShaderDrawPrimitive[PreviousPrimitiveCount + 10] = PackShaderElement(Element);
-		Element.Pos0 = P3; Element.Pos1 = P7; OutShaderDrawPrimitive[PreviousPrimitiveCount + 11] = PackShaderElement(Element);
-	}
-	else
-	{
-		InterlockedAdd(OutputShaderDrawIndirect[1], -12, PreviousPrimitiveCount);// Make sure we stay in the safe zone for indirect draw
+		Element.Pos0 = P0; Element.Pos1 = P4; AddShaderElement(Element, Offset + 8);
+		Element.Pos0 = P1; Element.Pos1 = P5; AddShaderElement(Element, Offset + 9);
+		Element.Pos0 = P2; Element.Pos1 = P6; AddShaderElement(Element, Offset + 10);
+		Element.Pos0 = P3; Element.Pos1 = P7; AddShaderElement(Element, Offset + 11);
 	}
 }
 
diff --git a/Engine/Source/Runtime/Renderer/Private/ScenePrivate.h b/Engine/Source/Runtime/Renderer/Private/ScenePrivate.h
--- a/Engine/Source/Runtime/Renderer/Private/ScenePrivate.h
+++ b/Engine/Source/Runtime/Renderer/Private/ScenePrivate.h
@@ -692,13 +692,11 @@ struct FHLODSceneNodeVisibilityState
 struct FShaderDrawDebugStateData
 {
 	TRefCountPtr<FRDGPooledBuffer> Buffer;
-	TRefCountPtr<FRDGPooledBuffer> IndirectBuffer;
 	bool bIsLocked = false;
 
 	void Release()
 	{
 		Buffer = nullptr;
-		IndirectBuffer = nullptr;
 		bIsLocked = false;
 	}
 };
diff --git a/Engine/Source/Runtime/Renderer/Private/ShaderDebug.cpp b/Engine/Source/Runtime/Renderer/Private/ShaderDebug.cpp
--- a/Engine/Source/Runtime/Renderer/Private/ShaderDebug.cpp
+++ b/Engine/Source/Runtime/Renderer/Private/ShaderDebug.cpp
@@ -71,16 +71,6 @@ namespace ShaderDrawDebug
 		return IsEnabled() && IsSupported(View.GetShaderPlatform());
 	}
 
-	// Note: Unaligned structures used for structured buffers is an unsupported and/or sparsely
-	//         supported feature in VK (VK_EXT_scalar_block_layout) and Metal. Consequently, we
-	//         do manual packing in order to accommodate.
-	struct FPackedShaderDrawElement
-	{
-		// This is not packed as fp16 to be able to debug large scale data while preserving accuracy at short range.
-		float Pos0_ColorX[4];		// float3 pos0 + packed color0
-		float Pos1_ColorY[4];		// float3 pos1 + packed color1
-	};
-
 	//////////////////////////////////////////////////////////////////////////
 
 	class FShaderDrawDebugClearCS : public FGlobalShader
@@ -89,8 +79,7 @@ namespace ShaderDrawDebug
 		SHADER_USE_PARAMETER_STRUCT(FShaderDrawDebugClearCS, FGlobalShader);
 
 		BEGIN_SHADER_PARAMETER_STRUCT(FParameters, )
-		SHADER_PARAMETER_RDG_BUFFER_UAV(RWStructuredBuffer, ElementBuffer)
-		SHADER_PARAMETER_RDG_BUFFER_UAV(RWBuffer, IndirectBuffer)
+			SHADER_PARAMETER_RDG_BUFFER_UAV(RWStructuredBuffer, RWElementBuffer)
 		END_SHADER_PARAMETER_STRUCT()
 
 		static bool ShouldCompilePermutation(const FGlobalShaderPermutationParameters& Parameters)
@@ -107,6 +96,33 @@ namespace ShaderDrawDebug
 	};
 
 	IMPLEMENT_GLOBAL_SHADER(FShaderDrawDebugClearCS, "/Engine/Private/ShaderDrawDebug.usf", "ShaderDrawDebugClearCS", SF_Compute);
+	
+	//////////////////////////////////////////////////////////////////////////
+
+	class FShaderDrawDebugCopyCS : public FGlobalShader
+	{
+		DECLARE_GLOBAL_SHADER(FShaderDrawDebugCopyCS);
+		SHADER_USE_PARAMETER_STRUCT(FShaderDrawDebugCopyCS, FGlobalShader);
+
+		BEGIN_SHADER_PARAMETER_STRUCT(FParameters, )
+			SHADER_PARAMETER_RDG_BUFFER_SRV(StructuredBuffer, ElementBuffer)
+			SHADER_PARAMETER_RDG_BUFFER_UAV(RWBuffer, RWIndirectArgs)
+		END_SHADER_PARAMETER_STRUCT()
+
+		static bool ShouldCompilePermutation(const FGlobalShaderPermutationParameters& Parameters)
+		{
+			return IsSupported(Parameters.Platform);
+		}
+
+		static void ModifyCompilationEnvironment(const FGlobalShaderPermutationParameters& Parameters, FShaderCompilerEnvironment& OutEnvironment)
+		{
+			FGlobalShader::ModifyCompilationEnvironment(Parameters, OutEnvironment);
+			OutEnvironment.SetDefine(TEXT("GPU_DEBUG_RENDERING"), 1);
+			OutEnvironment.SetDefine(TEXT("GPU_DEBUG_RENDERING_COPY_CS"), 1);
+		}
+	};
+
+	IMPLEMENT_GLOBAL_SHADER(FShaderDrawDebugCopyCS, "/Engine/Private/ShaderDrawDebug.usf", "ShaderDrawDebugCopyCS", SF_Compute);
 
 	//////////////////////////////////////////////////////////////////////////
 
@@ -177,15 +193,52 @@ namespace ShaderDrawDebug
 	END_SHADER_PARAMETER_STRUCT()
 
 	//////////////////////////////////////////////////////////////////////////
+		
+	static void AddShaderDrawDebugClearPass(FRDGBuilder& GraphBuilder, const FViewInfo& View, FRDGBufferRef& DataBuffer)
+	{
+		FShaderDrawDebugClearCS::FParameters* Parameters = GraphBuilder.AllocParameters<FShaderDrawDebugClearCS::FParameters>();
+		Parameters->RWElementBuffer = GraphBuilder.CreateUAV(DataBuffer);
+
+		TShaderMapRef<FShaderDrawDebugClearCS> ComputeShader(View.ShaderMap);
+		ClearUnusedGraphResources(ComputeShader, Parameters);
+		GraphBuilder.AddPass(
+			RDG_EVENT_NAME("ShaderDebug::Clear"),
+			Parameters,
+			ERDGPassFlags::Compute,
+		[Parameters, ComputeShader](FRHICommandList& RHICmdList)
+		{
+			FComputeShaderUtils::Dispatch(RHICmdList, ComputeShader, *Parameters, FIntVector(1, 1, 1));
+		});
+	}
+
+	static void AddShaderDrawDebugCopyPass(FRDGBuilder& GraphBuilder, const FViewInfo& View, FRDGBufferRef& DataBuffer, FRDGBufferRef& IndirectBuffer)
+	{
+		FShaderDrawDebugCopyCS::FParameters* Parameters = GraphBuilder.AllocParameters<FShaderDrawDebugCopyCS::FParameters>();
+		Parameters->ElementBuffer = GraphBuilder.CreateSRV(DataBuffer);
+		Parameters->RWIndirectArgs = GraphBuilder.CreateUAV(IndirectBuffer, PF_R32_UINT);
+
+		TShaderMapRef<FShaderDrawDebugCopyCS> ComputeShader(View.ShaderMap);
+		ClearUnusedGraphResources(ComputeShader, Parameters);
+		GraphBuilder.AddPass(
+			RDG_EVENT_NAME("ShaderDebug::CopyArgs"),
+			Parameters,
+			ERDGPassFlags::Compute,
+		[Parameters, ComputeShader](FRHICommandList& RHICmdList)
+		{
+			FComputeShaderUtils::Dispatch(RHICmdList, ComputeShader, *Parameters, FIntVector(1, 1, 1));
+		});
+	}
 
 	static void InternalDrawView(
 		FRDGBuilder& GraphBuilder,
 		const FViewInfo& View,
 		FRDGBufferRef DataBuffer,
-		FRDGBufferRef IndirectBuffer,
 		FRDGTextureRef OutputTexture,
 		FRDGTextureRef DepthTexture)
 	{
+		FRDGBufferRef IndirectBuffer = GraphBuilder.CreateBuffer(FRDGBufferDesc::CreateIndirectDesc<FRHIDrawIndirectParameters>(1), TEXT("ShaderDraw.IndirectBuffer"), ERDGBufferFlags::None);
+		AddShaderDrawDebugCopyPass(GraphBuilder, View, DataBuffer, IndirectBuffer);
+
 		TShaderMapRef<FShaderDrawDebugVS> VertexShader(View.ShaderMap);
 		TShaderMapRef<FShaderDrawDebugPS> PixelShader(View.ShaderMap);
 
@@ -235,33 +288,13 @@ namespace ShaderDrawDebug
 			});
 	}
 
-	static void AddShaderDrawDebugClearPass(FRDGBuilder& GraphBuilder, FViewInfo& View, FRDGBufferRef& DataBuffer, FRDGBufferRef& IndirectBuffer)
-	{
-		FShaderDrawDebugClearCS::FParameters* Parameters = GraphBuilder.AllocParameters<FShaderDrawDebugClearCS::FParameters>();
-		Parameters->ElementBuffer	= GraphBuilder.CreateUAV(DataBuffer);
-		Parameters->IndirectBuffer	= GraphBuilder.CreateUAV(IndirectBuffer);
-
-		TShaderMapRef<FShaderDrawDebugClearCS> ComputeShader(View.ShaderMap);
-
-		// Note: we do not call ClearUnusedGraphResources here as we want to for the allocation of DataBuffer
-		GraphBuilder.AddPass(
-			RDG_EVENT_NAME("ShaderDebug::Clear"),
-			Parameters,
-			ERDGPassFlags::Compute,
-			[Parameters, ComputeShader](FRHICommandList& RHICmdList)
-		{
-			FComputeShaderUtils::Dispatch(RHICmdList, ComputeShader, *Parameters, FIntVector(1,1,1));
-		});
-	}
-
 	void BeginView(FRDGBuilder& GraphBuilder, FViewInfo& View)
 	{
 		View.ShaderDrawData = FShaderDrawDebugData();
 		if (!IsEnabled(View))
 		{
 			// Default resources
-			View.ShaderDrawData.Buffer			= GraphBuilder.CreateBuffer(FRDGBufferDesc::CreateStructuredDesc(sizeof(FPackedShaderDrawElement), 1), TEXT("ShaderDraw.DataBuffer(Dummy)"), ERDGBufferFlags::None);
-			View.ShaderDrawData.IndirectBuffer	= GraphBuilder.CreateBuffer(FRDGBufferDesc::CreateIndirectDesc<FRHIDrawIndirectParameters>(1), TEXT("ShaderDraw.IndirectBuffer(Dummy)"), ERDGBufferFlags::None);
+			View.ShaderDrawData.Buffer = GraphBuilder.CreateBuffer(FRDGBufferDesc::CreateStructuredDesc(4, 8), TEXT("ShaderDraw.DataBuffer(Dummy)"), ERDGBufferFlags::None);
 			return;
 		}
 
@@ -274,26 +307,22 @@ namespace ShaderDrawDebug
 		const bool bLockBufferThisFrame = IsShaderDrawLocked() && ViewState && !ViewState->ShaderDrawDebugStateData.bIsLocked;
 		ERDGBufferFlags Flags = bLockBufferThisFrame ? ERDGBufferFlags::MultiFrame : ERDGBufferFlags::None;
 
-		FRDGBufferRef DataBuffer = GraphBuilder.CreateBuffer(FRDGBufferDesc::CreateStructuredDesc(sizeof(FPackedShaderDrawElement), View.ShaderDrawData.MaxElementCount), TEXT("ShaderDraw.DataBuffer"), Flags);
-		FRDGBufferRef IndirectBuffer = GraphBuilder.CreateBuffer(FRDGBufferDesc::CreateIndirectDesc<FRHIDrawIndirectParameters>(1), TEXT("ShaderDraw.IndirectBuffer"), Flags);
-		AddShaderDrawDebugClearPass(GraphBuilder, View, DataBuffer, IndirectBuffer);
+		FRDGBufferRef DataBuffer = GraphBuilder.CreateBuffer(FRDGBufferDesc::CreateStructuredDesc(4, 8*View.ShaderDrawData.MaxElementCount), TEXT("ShaderDraw.DataBuffer"), Flags);
+		AddShaderDrawDebugClearPass(GraphBuilder, View, DataBuffer);
 
 		View.ShaderDrawData.Buffer = DataBuffer;
-		View.ShaderDrawData.IndirectBuffer = IndirectBuffer;
 		View.ShaderDrawData.CursorPosition = View.CursorPos;
 		View.ShaderDrawData.ShaderDrawTranslatedWorldOffset = View.ViewMatrices.GetPreViewTranslation();
 
 		if (IsShaderDrawLocked() && ViewState && !ViewState->ShaderDrawDebugStateData.bIsLocked)
 		{
 			ViewState->ShaderDrawDebugStateData.Buffer = GraphBuilder.ConvertToExternalBuffer(View.ShaderDrawData.Buffer);
-			ViewState->ShaderDrawDebugStateData.IndirectBuffer = GraphBuilder.ConvertToExternalBuffer(View.ShaderDrawData.IndirectBuffer);
 			ViewState->ShaderDrawDebugStateData.bIsLocked = true;
 		}
 
 		if (!IsShaderDrawLocked() && ViewState && ViewState->ShaderDrawDebugStateData.bIsLocked)
 		{
 			ViewState->ShaderDrawDebugStateData.Buffer = nullptr;
-			ViewState->ShaderDrawDebugStateData.IndirectBuffer = nullptr;
 			ViewState->ShaderDrawDebugStateData.bIsLocked = false;
 		}
 
@@ -314,15 +343,13 @@ namespace ShaderDrawDebug
 
 		{
 			FRDGBufferRef DataBuffer = View.ShaderDrawData.Buffer;
-			FRDGBufferRef IndirectBuffer = View.ShaderDrawData.IndirectBuffer;
-			InternalDrawView(GraphBuilder, View, DataBuffer, IndirectBuffer, OutputTexture, DepthTexture);
+			InternalDrawView(GraphBuilder, View, DataBuffer, OutputTexture, DepthTexture);
 		}
 
 		if (View.ViewState && View.ViewState->ShaderDrawDebugStateData.bIsLocked)
 		{
 			FRDGBufferRef DataBuffer = GraphBuilder.RegisterExternalBuffer(View.ViewState->ShaderDrawDebugStateData.Buffer);
-			FRDGBufferRef IndirectBuffer = GraphBuilder.RegisterExternalBuffer(View.ViewState->ShaderDrawDebugStateData.IndirectBuffer);
-			InternalDrawView(GraphBuilder, View, DataBuffer, IndirectBuffer, OutputTexture, DepthTexture);
+			InternalDrawView(GraphBuilder, View, DataBuffer, OutputTexture, DepthTexture);
 		}
 	}
 
@@ -346,7 +373,6 @@ namespace ShaderDrawDebug
 		OutParameters.ShaderDrawMaxElementCount			= Data.MaxElementCount;
 		OutParameters.ShaderDrawTranslatedWorldOffset	= Data.ShaderDrawTranslatedWorldOffset;
 		OutParameters.OutShaderDrawPrimitive			= GraphBuilder.CreateUAV(Data.Buffer);
-		OutParameters.OutputShaderDrawIndirect			= GraphBuilder.CreateUAV(Data.IndirectBuffer);
 	}
 
 	// Returns true if the default view exists and has shader debug rendering enabled (this needs to be checked before using a permutation that requires the shader draw parameters)
diff --git a/Engine/Source/Runtime/Renderer/Public/ShaderDebug.h b/Engine/Source/Runtime/Renderer/Public/ShaderDebug.h
--- a/Engine/Source/Runtime/Renderer/Public/ShaderDebug.h
+++ b/Engine/Source/Runtime/Renderer/Public/ShaderDebug.h
@@ -14,9 +14,8 @@ struct FShaderDrawDebugData
 	FIntPoint CursorPosition = FIntPoint(-1,-1);
 	FVector ShaderDrawTranslatedWorldOffset;
 	FRDGBufferRef Buffer = nullptr;
-	FRDGBufferRef IndirectBuffer = nullptr;
 	bool IsEnabled() const { return MaxElementCount > 0; }
-	bool IsValid() const { return Buffer != nullptr && IndirectBuffer != nullptr; }
+	bool IsValid() const { return Buffer != nullptr; }
 };
 
 namespace ShaderDrawDebug 
@@ -53,7 +52,6 @@ namespace ShaderDrawDebug
 		SHADER_PARAMETER(int32, ShaderDrawMaxElementCount)
 		SHADER_PARAMETER(FVector3f, ShaderDrawTranslatedWorldOffset)
 		SHADER_PARAMETER_RDG_BUFFER_UAV(RWStructuredBuffer, OutShaderDrawPrimitive)
-		SHADER_PARAMETER_RDG_BUFFER_UAV(RWBuffer<uint>, OutputShaderDrawIndirect)
 	END_SHADER_PARAMETER_STRUCT()
 
 	// Call this to fill the FShaderDrawParameters
