GlobalDistanceField - fixed Landscape holes on clipmap borders caused by negative coordinates stored in an uint3 (one texel page border can become negative there).

diff --git a/Engine/Shaders/Private/GlobalDistanceField.usf b/Engine/Shaders/Private/GlobalDistanceField.usf
--- a/Engine/Shaders/Private/GlobalDistanceField.usf
+++ b/Engine/Shaders/Private/GlobalDistanceField.usf
@@ -569,7 +569,7 @@ void ComposeObjectsIntoPagesCS(
 
 	uint3 TexelCoordInPage = PageGroupOffset * 4 + GroupThreadId;
 
-	uint3 PageCoord = PageGridCoord * GLOBAL_DISTANCE_FIELD_PAGE_RESOLUTION + TexelCoordInPage - 1;
+	int3 PageCoord = PageGridCoord * GLOBAL_DISTANCE_FIELD_PAGE_RESOLUTION + TexelCoordInPage - 1;
 	float3 ComposeTileCenterPageCoord = PageGridCoord * GLOBAL_DISTANCE_FIELD_PAGE_RESOLUTION + PageGroupOffset * 4 + 1.5f - 1.0f;
 	uint3 CullGridCoord = PageGridCoord;
 
diff --git a/Engine/Shaders/Private/GlobalDistanceFieldHeightfields.usf b/Engine/Shaders/Private/GlobalDistanceFieldHeightfields.usf
--- a/Engine/Shaders/Private/GlobalDistanceFieldHeightfields.usf
+++ b/Engine/Shaders/Private/GlobalDistanceFieldHeightfields.usf
@@ -244,7 +244,7 @@ void ComposeHeightfieldsIntoPagesCS(
 
 	uint3 TexelCoordInPage = uint3(PageGroupOffset.xy * 8 + GroupThreadId.xy, 0);
 
-	uint3 PageCoord = PageGridCoord * GLOBAL_DISTANCE_FIELD_PAGE_RESOLUTION + TexelCoordInPage - 1;
+	int3 PageCoord = PageGridCoord * GLOBAL_DISTANCE_FIELD_PAGE_RESOLUTION + TexelCoordInPage - 1;
 	float3 VoxelWorldCenter = PageCoord * PageCoordToVoxelCenterScale + PageCoordToVoxelCenterBias;
 
 	uint3 PageTableTextureCoord = PageGridCoordToPageTableTextureCoord(PageGridCoord);
@@ -284,7 +284,7 @@ void ComposeHeightfieldsIntoPagesCS(
 			LOOP
 			for (uint ZIndex = 0; ZIndex < GLOBAL_DISTANCE_FIELD_PAGE_RESOLUTION_IN_ATLAS; ++ZIndex)
 			{
-				float3 WorldPosition = VoxelWorldCenter.xyz + float3(0.0f, 0.0f, ZIndex * 2.0f * ClipmapVoxelExtent);
+				float3 WorldPosition = VoxelWorldCenter.xyz + float3(0.0f, 0.0f, ZIndex * PageCoordToVoxelCenterScale.z);
 
 				// Project the vertical height vector onto the normal of the heightfield directly below the point we are computing the distance field for, use the perpendicular distance
 				float DistanceToHeightfieldPlane = dot(WorldHeightfieldNormal, WorldPosition - WorldHeightfieldShadingPosition);
