Fix missing groom shadow when RT is enabled.

This fixes a typo introduced when adding LWC support.

diff --git a/Engine/Shaders/Private/HairStrands/HairStrandsRaytracing.ush b/Engine/Shaders/Private/HairStrands/HairStrandsRaytracing.ush
--- a/Engine/Shaders/Private/HairStrands/HairStrandsRaytracing.ush
+++ b/Engine/Shaders/Private/HairStrands/HairStrandsRaytracing.ush
@@ -2,7 +2,9 @@
 
 #pragma once
 
+#ifndef VOXEL_TRAVERSAL_DEBUG
 #define VOXEL_TRAVERSAL_DEBUG 0
+#endif
 #define VOXEL_TRAVERSAL_TYPE VOXEL_TRAVERSAL_LINEAR_MIPMAP
 //#define VOXEL_TRAVERSAL_TYPE VOXEL_TRAVERSAL_LINEAR
 
@@ -15,7 +17,7 @@ float InternalTraverseHair(uint2 PixelCoord, float3 RayOrigin, float3 RayDirecti
 	float HitT = InMinT;
 
 	// Note: VirtualVoxel.TranslatedWorldOffset should be equivalent to LWCHackToFloat(PrimaryView.WorldCameraOrigin), but insures the value is consistent
-	RayOrigin = RayOrigin - VirtualVoxel.TranslatedWorldOffset; // HAIR_LWC
+	RayOrigin = RayOrigin + VirtualVoxel.TranslatedWorldOffset; // HAIR_LWC
 
 	// Depth bias
 	// Origin is shifted voxels away towards the light + a constant bias of the size of the voxel
