Add static lighting support for hair strands geometry with default lighting integrator.

diff --git a/Engine/Shaders/Private/HairStrands/HairStrandsEnvironmentLighting.usf b/Engine/Shaders/Private/HairStrands/HairStrandsEnvironmentLighting.usf
--- a/Engine/Shaders/Private/HairStrands/HairStrandsEnvironmentLighting.usf
+++ b/Engine/Shaders/Private/HairStrands/HairStrandsEnvironmentLighting.usf
@@ -4,6 +4,42 @@
 	ReflectionEnvironmentComputeShaders - functionality to apply local cubemaps.
 =============================================================================*/
 
+////////////////////////////////////////////////////////////////////////////////////////////////////////////////////////////////////////////////////////////////
+// Note: 
+// * This version of environement lighting is very much in WIP state, with a lot of option to tune tweak to explore 
+//   what kind of approximation works the best for certains type of grooms and lighting condition
+// * There is a few fudges factor as well (roughness mapping, TT density, reflected vector orientations, ...)
+
+// Sky lighting
+// ============
+// Note: Strands & Cards are always dynamic primitive
+// 
+//          | Static    | Stationnary  | Dynamic
+// ---------|-----------|--------------|--------
+// Cards    | BasePass  | BasePass     | EnvPass
+// Strands  | EnvPass   | EnvPass      | EnvPass
+
+////////////////////////////////////////////////////////////////////////////////////////////////////////////////////////////////////////////////////////////////
+// Includes
+
+// Override ALLOW_STATIC_LIGHTING based on shader permutation
+// Need to be define prior to other includes
+#undef ALLOW_STATIC_LIGHTING
+#if PERMUTATION_STATIC_LIGHTING
+	#define ALLOW_STATIC_LIGHTING 1
+	#define REFLECTION_COMPOSITE_SUPPORT_SKYLIGHT_BLEND 1
+	#define PRECOMPUTED_IRRADIANCE_VOLUME_LIGHTING 1
+	#define APPLY_SKY_SHADOWING 1
+#else
+	#define ALLOW_STATIC_LIGHTING 0
+	#define REFLECTION_COMPOSITE_SUPPORT_SKYLIGHT_BLEND 0
+	#define PRECOMPUTED_IRRADIANCE_VOLUME_LIGHTING 0
+#endif
+#define REFLECTION_COMPOSITE_USE_BLENDED_REFLECTION_CAPTURES 1
+#define REFLECTION_COMPOSITE_HAS_BOX_CAPTURES 1
+#define REFLECTION_COMPOSITE_HAS_SPHERE_CAPTURES 1
+#define ENABLE_SKY_LIGHT 1
+
 #include "../Common.ush"
 #include "HairStrandsCommon.ush"
 #include "../DeferredShadingCommon.ush"  
@@ -15,24 +51,17 @@
 #include "../ShadingModels.ush"
 #include "../LightGridCommon.ush"
 #include "../SceneTextureParameters.ush"
-
+#include "../ReflectionEnvironmentComposite.ush"
+#include "../SHCommon.ush"
+#include "../ReflectionEnvironmentShared.ush"
+#include "../SkyLightingShared.ush"
+#include "../VolumetricLightmapShared.ush"
+#include "../HairBsdf.ush"
+#include "../ShaderDrawDebug.ush"
 #if PERMUTATION_LUMEN
 #include "../Lumen/LumenRadianceCacheCommon.ush"
 #endif
 
-// Disable static lighting, as hair don't have any diffused lighting baked, 
-// and rely solely on specular information (i.e. reflection capture)
-#undef ALLOW_STATIC_LIGHTING
-#define ALLOW_STATIC_LIGHTING 0
-
-#define ENABLE_SKY_LIGHT 1
-#define REFLECTION_COMPOSITE_USE_BLENDED_REFLECTION_CAPTURES 1
-#define REFLECTION_COMPOSITE_SUPPORT_SKYLIGHT_BLEND 0
-#define REFLECTION_COMPOSITE_HAS_BOX_CAPTURES 1
-#define REFLECTION_COMPOSITE_HAS_SPHERE_CAPTURES 1
-#include "../ReflectionEnvironmentComposite.ush"
-
-#include "../HairBsdf.ush"
 #include "HairStrandsVisibilityCommon.ush"
 #include "HairStrandsVisibilityUtils.ush"
 #include "HairStrandsEnvironmentLightingCommon.ush"
@@ -41,28 +70,77 @@
 #include "HairStrandsDeepTransmittanceCommon.ush"
 #include "HairStrandsDeepTransmittanceDualScattering.ush"
 
-#include "../ShaderDrawDebug.ush"
-
 #define VOXEL_TRAVERSAL_TYPE VOXEL_TRAVERSAL_LINEAR_MIPMAP
 //#define VOXEL_TRAVERSAL_TYPE VOXEL_TRAVERSAL_LINEAR
 #define VOXEL_TRAVERSAL_DEBUG 0
 #include "HairStrandsVoxelPageTraversal.ush"
 
+////////////////////////////////////////////////////////////////////////////////////////////////////////////////////////////////////////////////////////////////
+// Defines
+
 #define INTEGRATION_SCENECOLOR 0
 #define INTEGRATION_ADHOC 1
 #define INTEGRATION_UNIFORM 2
 #define INTEGRATION_SH 3
 
-// Note: 
-// * This version of environement lighting is very much in WIP state, with a lot of option to tune tweak to explore 
-//   what kind of approximation works the best for certains type of grooms and lighting condition
-// * There is a few fudges factor as well (roughness mapping, TT density, reflected vector orientations, ...)
+////////////////////////////////////////////////////////////////////////////////////////////////////////////////////////////////////////////////////////////////
+// Parameters
+
+uint				bHasStaticLighting;
+Texture3D<float4>	HairEnergyLUTTexture;
+
+////////////////////////////////////////////////////////////////////////////////////////////////////////////////////////////////////////////////////////////////
+// Static lighting evaluation (diffuse)
+#if PERMUTATION_STATIC_LIGHTING
 
-float GetPreExposure()
+/** Calculates indirect lighting contribution on this object from precomputed data. */
+float3 GetPrecomputedIndirectLighting(float3 AbsoluteWorldPosition, float3 DiffuseDir)
 {
-	return View.PreExposure;
+	float3 OutDiffuseLighting = 0;
+
+	#if PRECOMPUTED_IRRADIANCE_VOLUME_LIGHTING
+	const float3 VolumetricLightmapBrickTextureUVs = ComputeVolumetricLightmapBrickTextureUVs(LWCHackToFloat(AbsoluteWorldPosition));
+	FThreeBandSHVectorRGB IrradianceSH = GetVolumetricLightmapSH3(VolumetricLightmapBrickTextureUVs);
+	
+	// Diffuse convolution
+	FThreeBandSHVector DiffuseTransferSH = CalcDiffuseTransferSH3(DiffuseDir, 1);
+	OutDiffuseLighting = max(float3(0,0,0), DotSH3(IrradianceSH, DiffuseTransferSH)) / PI;
+	#endif
+
+	// Apply indirect lighting scale while we have only accumulated lightmaps
+	OutDiffuseLighting *= View.PrecomputedIndirectLightingColorScale;
+
+	// Note: pre-exposition is applied by the caller
+	return OutDiffuseLighting;
 }
+#endif // PERMUTATION_STATIC_LIGHTING
+
+float3 GetStaticLighting(in FGBufferData GBuffer, float3 CameraDirection, float3 AbsoluteWorldPosition)
+{
+	float3 OutStaticLighting = 0;
 
+	#if PERMUTATION_STATIC_LIGHTING
+	// Evaluation of the static lighting
+	if (bHasStaticLighting)
+	{
+		FHairTransmittanceData TransmittanceData = InitHairTransmittanceData(true);
+		float3 N = GBuffer.WorldNormal;
+		float3 V = CameraDirection;
+		float3 L = normalize(V - N * dot(V, N));
+		const float3 DiffuseColorForIndirect = 2 * PI * HairShading(GBuffer, L, V, N, 1, TransmittanceData, 0, 0.2, uint2(0, 0));
+		const float3 DiffuseIndirectLighting = GetPrecomputedIndirectLighting(AbsoluteWorldPosition, L);
+
+		const float MaterialAO = 1;
+		const float DiffOcclusion = MaterialAO;
+		OutStaticLighting = (DiffuseIndirectLighting * DiffuseColorForIndirect) * AOMultiBounce(GBuffer.BaseColor, DiffOcclusion);
+	}
+	#endif
+	return OutStaticLighting * View.PreExposure;
+}
+
+
+////////////////////////////////////////////////////////////////////////////////////////////////////////////////////////////////////////////////////////////////
+// Reflection probe evaluation
 float3 GatherRadiance(float3 RayDirection, float Roughness, uint2 PixelPosition, float PixelSceneDepth, float3 TranslatedWorldPosition)
 {
 #if PERMUTATION_LUMEN
@@ -106,8 +184,6 @@ float3 GatherRadiance(float3 RayDirection, float Roughness, uint2 PixelPosition,
 #endif // PERMUTATION_LUMEN
 }
 
-Texture3D<float4>			HairEnergyLUTTexture;
-
 //////////////////////////////////////////////////////////////////////////////////////////////////////////////////////////////
 // Debug info
 
@@ -162,8 +238,6 @@ Texture2D<float> HairCountTexture;
 uint  MultipleScatterSampleCount;
 float TransmissionDensityScaleFactor;
 
-float Voxel_TanConeAngle;
-
 FHairTraversalResult InternalComputeHairCountVoxel(
 	uint MacroGroupId,
 	float3 TranslatedWorldPosition, 
@@ -221,7 +295,6 @@ float MaxComponent3(float3 In)
 
 float3 ReflectionEnvironment(FGBufferData GBuffer, float2 ScreenPosition, float2 SvPosition, uint PrimitiveId, uint MacroGroupId)
 {
-	const float PreExposure = GetPreExposure();
 	const float3 TranslatedWorldPosition = mul(float4(ScreenPosition * GBuffer.Depth, GBuffer.Depth, 1), PrimaryView.ScreenToTranslatedWorld).xyz;
 	const float3 CameraToPixel = normalize(TranslatedWorldPosition - PrimaryView.TranslatedWorldCameraOrigin);
 	const uint2  PixelPosition = SvPosition - ResolvedView.ViewRectMin.xy;
@@ -316,8 +389,8 @@ float3 ReflectionEnvironment(FGBufferData GBuffer, float2 ScreenPosition, float2
 		//const float UnoccludedVarianceR = NormR * (3 - NormR2) / (1 - NormR2);
 		//const float3 ColorDir = (UnoccludedN+ float3(1, 1, 1)) * 0.5f;
 		//return ColorDir.rgb;
-		//return PreExposure * GetSkySHDiffuse(UnoccludedN) * View.SkyLightColor.rgb;
-		//return PreExposure * GatherRadiance(UnoccludedN, GBuffer.Roughness);
+		//return View.PreExposure * GetSkySHDiffuse(UnoccludedN) * View.SkyLightColor.rgb;
+		//return View.PreExposure * GatherRadiance(UnoccludedN, GBuffer.Roughness);
 		//return AvgDualScattering;
 	}
 	#endif
@@ -348,8 +421,8 @@ float3 ReflectionEnvironment(FGBufferData GBuffer, float2 ScreenPosition, float2
 		// Select the TT term for representing the MS roughness lookup as TT lobe carries the most energy + fudge scale as the MS lobe should be wider than a TT lobe
 		const float FudgeScale = 2;
 		const float GGXRoughness = HairToGGXRoughness_TT(GBuffer.Roughness) * FudgeScale;
-		ColorMS.rgb += PreExposure * GatherRadiance(R, GGXRoughness, PixelPosition, PixelSceneDepth, TranslatedWorldPosition);
-		//ColorMS.rgb += PreExposure * GetSkySHDiffuse(R) * View.SkyLightColor.rgb;
+		ColorMS.rgb += View.PreExposure * GatherRadiance(R, GGXRoughness, PixelPosition, PixelSceneDepth, TranslatedWorldPosition);
+		//ColorMS.rgb += View.PreExposure * GetSkySHDiffuse(R) * View.SkyLightColor.rgb;
 		ColorMS.rgb *= AvgDualScattering;
 
 		// Directional albedo
@@ -394,8 +467,8 @@ float3 ReflectionEnvironment(FGBufferData GBuffer, float2 ScreenPosition, float2
 		// Select the TRT for remapping the hair roughness onto the GGX based lobe, as it is the 
 		// widest lobe between the R & TRT components
 		const float GGXRoughness = HairToGGXRoughness_TRT(GBuffer.Roughness);
-		ColorR.rgb += PreExposure * GatherRadiance(R, GGXRoughness, PixelPosition, PixelSceneDepth, TranslatedWorldPosition);
-		//ColorR.rgb += PreExposure * GetSkySHDiffuse(R) * View.SkyLightColor.rgb * PI;
+		ColorR.rgb += View.PreExposure * GatherRadiance(R, GGXRoughness, PixelPosition, PixelSceneDepth, TranslatedWorldPosition);
+		//ColorR.rgb += View.PreExposure * GetSkySHDiffuse(R) * View.SkyLightColor.rgb * PI;
 
 		const float3 L = R;
 		const float SinLightAngle = dot(L, T);
@@ -454,8 +527,8 @@ float3 ReflectionEnvironment(FGBufferData GBuffer, float2 ScreenPosition, float2
 		const float SinViewAngle = dot(V, T);
 
 		const float GGXRoughness = HairToGGXRoughness_TT(GBuffer.Roughness);
-		ColorTT.rgb += PreExposure * GatherRadiance(L, GGXRoughness, PixelPosition, PixelSceneDepth, TranslatedWorldPosition);
-		//ColorTT.rgb += PreExposure * GetSkySHDiffuse(L) * View.SkyLightColor.rgb;
+		ColorTT.rgb += View.PreExposure * GatherRadiance(L, GGXRoughness, PixelPosition, PixelSceneDepth, TranslatedWorldPosition);
+		//ColorTT.rgb += View.PreExposure * GetSkySHDiffuse(L) * View.SkyLightColor.rgb;
 
 		// Compute the transmittance
 		float3 Transmittance = 1;
@@ -537,7 +610,6 @@ float3 ReflectionEnvironmentUniformSampling(
 	uint MacroGroupId,
 	bool bDebugEnabled)
 {
-	const float PreExposure = GetPreExposure();
 	const float3 TranslatedWorldPosition = mul(float4(ScreenPosition * GBuffer.Depth, GBuffer.Depth, 1), PrimaryView.ScreenToTranslatedWorld).xyz;
 	const float3 CameraToPixel = normalize(TranslatedWorldPosition - PrimaryView.TranslatedWorldCameraOrigin);
 	const uint2  PixelPosition = SvPosition - ResolvedView.ViewRectMin.xy;
@@ -601,13 +673,12 @@ float3 ReflectionEnvironmentUniformSampling(
 			View.HairScatteringLUTSampler,
 			View.HairComponents);	
 
-		const float Area = 0;
+		const float Area = 0;	
 		const float Backlit = 1;
 		const float RoughnessBias = 0.1f;
 		AccLighting += 
 			HairShading(GBuffer, SampleL, V, T, 1, TransmittanceData, Backlit, Area, uint2(0, 0)) * 
 			Result.Visibility * 
-			PreExposure * 
 			GatherRadiance(SampleL, RoughnessBias, PixelPosition, PixelSceneDepth, TranslatedWorldPosition);
 
 		#if PERMUTATION_DEBUG
@@ -618,10 +689,12 @@ float3 ReflectionEnvironmentUniformSampling(
 		#endif
 	}
 	const float SphereInt = 4 * PI;
-	const float3 Color = SphereInt * AccLighting / SampleCount;
+	float3 FinalColor = SphereInt * AccLighting / SampleCount;
 	
+	FinalColor *= View.PreExposure;
+
 	// Transform NaNs to black, transform negative colors to black.
-	return -min(-Color, 0.0);
+	return -min(-FinalColor, 0.0);
 }
 #endif // #if PERMUTATION_INTEGRATION_TYPE == INTEGRATION_UNIFORM
 
@@ -643,8 +716,18 @@ float3 ReflectionEnvironmentSH(
 	const float3 DiffuseColor = GBuffer.BaseColor; // not used
 	const float2 BufferUV = SvPosition * View.BufferSizeAndInvSize.zw; // not used
 
-	// Note this code is enabled only if sky lighting is dynamic
-	const float3 Lighting = SkyLightDiffuse(GBuffer, AmbientOcclusion, BufferUV, ScreenPosition, BentNormal, DiffuseColor);
+	const float3 TranslatedWorldPosition = mul(float4(ScreenPosition * GBuffer.Depth, GBuffer.Depth, 1), PrimaryView.ScreenToTranslatedWorld).xyz;
+	const float3 AbsoluteWorldPosition = TranslatedWorldPosition + LWCHackToFloat(PrimaryView.WorldCameraOrigin);
+	const float3 CameraToPixel = normalize(TranslatedWorldPosition - PrimaryView.TranslatedWorldCameraOrigin);
+	const float3 V = -CameraToPixel;
+
+	float3 Lighting = 0;
+
+	// Sky lighting (already pre-exposed)
+	Lighting += SkyLightDiffuse(GBuffer, AmbientOcclusion, BufferUV, ScreenPosition, BentNormal, DiffuseColor);
+
+	// Static lighting (already pre-exposed)
+	Lighting += GetStaticLighting(GBuffer, V, AbsoluteWorldPosition);
 
 	// Transform NaNs to black, transform negative colors to black.
 	return -min(-Lighting, 0.0);
@@ -676,9 +759,11 @@ float3 ComputeThinTransmission(FGBufferData GBuffer, float2 ScreenPosition, floa
 	// The pre-exposure is not applied onto the output as the SceneColorTexture value are already preexposed.
 	const float SinViewAngle = dot(T, V);
 	const FHairAverageScattering Energy = SampleHairLUT(View.HairScatteringLUTTexture, View.HairScatteringLUTSampler, GBuffer.BaseColor, GBuffer.Roughness, SinViewAngle);
-	float3 Color = Energy.A_front * Luminance.xyz * PixelCoverage;
+	float3 FinalColor = Energy.A_front * Luminance.xyz * PixelCoverage;
+
+	// Note: No need for pre-exposition, as we fetch & weight SceneColor which is already pre-exposed.
 
-	return -min(-Color, 0.0);
+	return -min(-FinalColor, 0.0);
 }
 
 #endif
diff --git a/Engine/Shaders/Private/SkyLightingDiffuseShared.ush b/Engine/Shaders/Private/SkyLightingDiffuseShared.ush
--- a/Engine/Shaders/Private/SkyLightingDiffuseShared.ush
+++ b/Engine/Shaders/Private/SkyLightingDiffuseShared.ush
@@ -116,7 +116,7 @@ float3 SkyLightDiffuse(FGBufferData GBuffer, float AmbientOcclusion, float2 Buff
 		float3 L = normalize(V - N * dot(V, N));
 		SkyVisData.SkyDiffuseLookUpNormal = L;
 		FHairTransmittanceData TransmittanceData = InitHairTransmittanceData(true);
-		DiffuseColor = PI * HairShading(GBuffer, L, V, N, 1, TransmittanceData, 0, 0.2, uint2(0, 0));
+		DiffuseColor = 2 * PI * HairShading(GBuffer, L, V, N, 1, TransmittanceData, 0, 0.2, uint2(0, 0));
 		DiffuseWeight = 1.0f;
 	}
 
diff --git a/Engine/Source/Runtime/Renderer/Private/BasePassRendering.cpp b/Engine/Source/Runtime/Renderer/Private/BasePassRendering.cpp
--- a/Engine/Source/Runtime/Renderer/Private/BasePassRendering.cpp
+++ b/Engine/Source/Runtime/Renderer/Private/BasePassRendering.cpp
@@ -1590,6 +1590,12 @@ void FBasePassMeshProcessor::AddMeshBatch(const FMeshBatch& RESTRICT MeshBatch,
 	}
 }
 
+bool AllowStaticLighting()
+{
+	static const auto AllowStaticLightingVar = IConsoleManager::Get().FindTConsoleVariableDataInt(TEXT("r.AllowStaticLighting"));
+	return (!AllowStaticLightingVar || AllowStaticLightingVar->GetValueOnRenderThread() != 0);
+}
+
 bool FBasePassMeshProcessor::TryAddMeshBatch(const FMeshBatch& RESTRICT MeshBatch, uint64 BatchElementMask, const FPrimitiveSceneProxy* RESTRICT PrimitiveSceneProxy, int32 StaticMeshId, const FMaterialRenderProxy& MaterialRenderProxy, const FMaterial& Material)
 {
 	// Determine the mesh's material and blend mode.
@@ -1645,8 +1651,7 @@ bool FBasePassMeshProcessor::TryAddMeshBatch(const FMeshBatch& RESTRICT MeshBatc
 	{
 		// Check for a cached light-map.
 		const bool bIsLitMaterial = ShadingModels.IsLit();
-		static const auto AllowStaticLightingVar = IConsoleManager::Get().FindTConsoleVariableDataInt(TEXT("r.AllowStaticLighting"));
-		const bool bAllowStaticLighting = (!AllowStaticLightingVar || AllowStaticLightingVar->GetValueOnRenderThread() != 0);
+		const bool bAllowStaticLighting = AllowStaticLighting();
 
 		const FLightMapInteraction LightMapInteraction = (bAllowStaticLighting && MeshBatch.LCI && bIsLitMaterial)
 			? MeshBatch.LCI->GetLightMapInteraction(FeatureLevel)
@@ -1836,8 +1841,7 @@ bool FBasePassMeshProcessor::TryAddMeshBatch(const FMeshBatch& RESTRICT MeshBatc
 			default:
 				if (bIsLitMaterial
 					&& bAllowStaticLighting
-					&& Scene
-					&& Scene->VolumetricLightmapSceneData.HasData()
+					&& bUseVolumetricLightmap
 					&& PrimitiveSceneProxy
 					&& (PrimitiveSceneProxy->IsMovable()
 						|| PrimitiveSceneProxy->NeedsUnbuiltPreviewLighting()
diff --git a/Engine/Source/Runtime/Renderer/Private/DeferredShadingRenderer.h b/Engine/Source/Runtime/Renderer/Private/DeferredShadingRenderer.h
--- a/Engine/Source/Runtime/Renderer/Private/DeferredShadingRenderer.h
+++ b/Engine/Source/Runtime/Renderer/Private/DeferredShadingRenderer.h
@@ -753,8 +753,6 @@ private:
 
 	bool HasDeferredPlanarReflections(const FViewInfo& View) const;
 	void RenderDeferredPlanarReflections(FRDGBuilder& GraphBuilder, const FSceneTextureParameters& SceneTextures, const FViewInfo& View, FRDGTextureRef& ReflectionsOutput);
-
-	bool ShouldDoReflectionEnvironment() const;
 	
 	bool ShouldRenderDistanceFieldAO() const;
 
diff --git a/Engine/Source/Runtime/Renderer/Private/DistanceFieldAmbientOcclusion.cpp b/Engine/Source/Runtime/Renderer/Private/DistanceFieldAmbientOcclusion.cpp
--- a/Engine/Source/Runtime/Renderer/Private/DistanceFieldAmbientOcclusion.cpp
+++ b/Engine/Source/Runtime/Renderer/Private/DistanceFieldAmbientOcclusion.cpp
@@ -681,6 +681,14 @@ bool ShouldRenderDeferredDynamicSkyLight(const FScene* Scene, const FSceneViewFa
 		&& !ShouldRenderRayTracingSkyLight(Scene->SkyLight); // Disable diffuse sky contribution if evaluated by RT Sky;
 }
 
+bool ShouldDoReflectionEnvironment(const FScene* Scene, const FSceneViewFamily& ViewFamily)
+{
+	const ERHIFeatureLevel::Type SceneFeatureLevel = Scene->GetFeatureLevel();
+
+	return IsReflectionEnvironmentAvailable(SceneFeatureLevel)
+		&& Scene->ReflectionSceneData.RegisteredReflectionCaptures.Num()
+		&& ViewFamily.EngineShowFlags.ReflectionEnvironment;
+}
 
 bool FSceneRenderer::ShouldPrepareForDistanceFieldAO() const
 {
diff --git a/Engine/Source/Runtime/Renderer/Private/DistanceFieldAmbientOcclusion.h b/Engine/Source/Runtime/Renderer/Private/DistanceFieldAmbientOcclusion.h
--- a/Engine/Source/Runtime/Renderer/Private/DistanceFieldAmbientOcclusion.h
+++ b/Engine/Source/Runtime/Renderer/Private/DistanceFieldAmbientOcclusion.h
@@ -151,6 +151,7 @@ protected:
 extern void TrackGPUProgress(FRHICommandListImmediate& RHICmdList, uint32 DebugId);
 
 extern bool ShouldRenderDeferredDynamicSkyLight(const FScene* Scene, const FSceneViewFamily& ViewFamily);
+extern bool ShouldDoReflectionEnvironment(const FScene* Scene, const FSceneViewFamily& ViewFamily);
 
 class FDistanceFieldCulledObjectBufferParameters;
 
diff --git a/Engine/Source/Runtime/Renderer/Private/HairStrands/HairStrandsEnvironment.cpp b/Engine/Source/Runtime/Renderer/Private/HairStrands/HairStrandsEnvironment.cpp
--- a/Engine/Source/Runtime/Renderer/Private/HairStrands/HairStrandsEnvironment.cpp
+++ b/Engine/Source/Runtime/Renderer/Private/HairStrands/HairStrandsEnvironment.cpp
@@ -22,6 +22,8 @@
 #include "Lumen/LumenRadianceCache.h"
 #include "Lumen/LumenScreenProbeGather.h"
 #include "IndirectLightRendering.h"
+#include "DistanceFieldAmbientOcclusion.h"
+#include "RenderGraphEvent.h"
 
 ///////////////////////////////////////////////////////////////////////////////////////////////////
 
@@ -62,6 +64,10 @@ static FAutoConsoleVariableRef CVarHairStrandsSkyLighting_DebugSample(TEXT("r.Ha
 
 ///////////////////////////////////////////////////////////////////////////////////////////////////
 
+DECLARE_GPU_STAT_NAMED(HairSkyLighting, TEXT("Hair Sky lighting"));
+
+///////////////////////////////////////////////////////////////////////////////////////////////////
+
 enum class EHairLightingSourceType : uint8
 {
 	SceneColor = 0,
@@ -202,7 +208,8 @@ public:
 	class FIntegrationType	: SHADER_PERMUTATION_INT("PERMUTATION_INTEGRATION_TYPE", uint32(EHairLightingIntegrationType::Count));
 	class FDebug			: SHADER_PERMUTATION_BOOL("PERMUTATION_DEBUG"); 
 	class FLumen			: SHADER_PERMUTATION_BOOL("PERMUTATION_LUMEN"); 
-	using FPermutationDomain = TShaderPermutationDomain<FIntegrationType, FLumen, FDebug>;
+	class FStaticLighting	: SHADER_PERMUTATION_BOOL("PERMUTATION_STATIC_LIGHTING"); 
+	using FPermutationDomain = TShaderPermutationDomain<FIntegrationType, FLumen, FDebug, FStaticLighting>;
 
 	static void ModifyCompilationEnvironment(const FGlobalShaderPermutationParameters& Parameters, FShaderCompilerEnvironment& OutEnvironment)
 	{
@@ -217,6 +224,19 @@ public:
 		{
 			return false;
 		}
+
+		// If using Lumen, no static lighting
+		if (PermutationVector.Get<FStaticLighting>() && PermutationVector.Get<FLumen>())
+		{
+			return false;
+		}
+
+		// If using another integrator than SH integrator, no static lighting
+		if (PermutationVector.Get<FStaticLighting>() && PermutationVector.Get<FIntegrationType>() != int32(EHairLightingIntegrationType::SH))
+		{
+			return false;
+		}
+
 		return IsHairStrandsSupported(EHairStrandsShaderType::Strands, Parameters.Platform);
 	}
 
@@ -230,11 +250,11 @@ public:
 	}
 
 	BEGIN_SHADER_PARAMETER_STRUCT(FParameters, )
-		SHADER_PARAMETER(float, Voxel_TanConeAngle)
+		SHADER_PARAMETER(uint32, bHasStaticLighting)
 		SHADER_PARAMETER(uint32, MultipleScatterSampleCount)
 		SHADER_PARAMETER(float,  HairDualScatteringRoughnessOverride)
-		SHADER_PARAMETER(float, TransmissionDensityScaleFactor)
-		SHADER_PARAMETER(float, HairDistanceThreshold)
+		SHADER_PARAMETER(float,  TransmissionDensityScaleFactor)
+		SHADER_PARAMETER(float,  HairDistanceThreshold)
 		SHADER_PARAMETER(uint32, bHairUseViewHairCount)
 
 		SHADER_PARAMETER_TEXTURE(Texture2D, PreIntegratedGF)
@@ -309,6 +329,8 @@ class FHairEnvironmentLightingPS : public FGlobalShader
 IMPLEMENT_GLOBAL_SHADER(FHairEnvironmentLightingPS, "/Engine/Private/HairStrands/HairStrandsEnvironmentLighting.usf", "MainPS", SF_Pixel);
 IMPLEMENT_GLOBAL_SHADER(FHairEnvironmentLightingVS, "/Engine/Private/HairStrands/HairStrandsEnvironmentLighting.usf", "MainVS", SF_Vertex);
 
+bool AllowStaticLighting();
+
 static void AddHairStrandsEnvironmentLightingPassPS(
 	FRDGBuilder& GraphBuilder,
 	const FScene* Scene,
@@ -319,24 +341,12 @@ static void AddHairStrandsEnvironmentLightingPassPS(
 	const EHairLightingSourceType LightingType,
 	const FHairStrandsDebugData::Data* DebugData)
 {
-	FSceneTextureParameters SceneTextures = GetSceneTextureParameters(GraphBuilder);
-
-	check(VirtualVoxelResources.IsValid());
-
-	// Render the reflection environment with tiled deferred culling
-	const bool bHasBoxCaptures = (View.NumBoxReflectionCaptures > 0);
-	const bool bHasSphereCaptures = (View.NumSphereReflectionCaptures > 0);
-	FHairEnvironmentLightingPS::FParameters* ParametersPS = GraphBuilder.AllocParameters<FHairEnvironmentLightingPS::FParameters>();
-	FHairEnvironmentLighting::FParameters* PassParameters = &ParametersPS->Common;
-
-	PassParameters->HairEnergyLUTTexture = GetHairLUT(GraphBuilder, View, HairLUTType_MeanEnergy);
 
 	EHairLightingIntegrationType IntegrationType = EHairLightingIntegrationType::AdHoc;
 	if (LightingType == EHairLightingSourceType::SceneColor)
 	{
 		check(SceneColorTexture != nullptr);
 		IntegrationType = EHairLightingIntegrationType::SceneColor;
-		PassParameters->SceneColorTexture = SceneColorTexture;
 	}
 	else
 	{
@@ -347,9 +357,38 @@ static void AddHairStrandsEnvironmentLightingPassPS(
 			case 2: IntegrationType = EHairLightingIntegrationType::SH; break;
 		}
 	}
-	
+
+	// The specular sky light contribution is also needed by RT Reflections as a fallback.
+	const bool bSkyLight = Scene->SkyLight && Scene->SkyLight->ProcessedTexture && !Scene->SkyLight->bHasStaticLighting;
+
+	const bool bDynamicSkyLight = ShouldRenderDeferredDynamicSkyLight(Scene, *View.Family);
+	const bool bReflectionEnv = ShouldDoReflectionEnvironment(Scene, *View.Family);
+
+	// Only support static lighting with SH integrator at the moment
+	const bool bUseVolumetricLightmap = Scene && Scene->VolumetricLightmapSceneData.HasData();
+	const bool bHasStaticLighting = AllowStaticLighting() && bUseVolumetricLightmap && View.FinalPostProcessSettings.DynamicGlobalIlluminationMethod != EDynamicGlobalIlluminationMethod::Lumen && IntegrationType == EHairLightingIntegrationType::SH;
+
+	// Sanity check
+	if (bHasStaticLighting) { check(LightingType != EHairLightingSourceType::Lumen); }
+
+	// Early out if there is no static lighting, no sky lighting, nor reflection probes
+	if (!bSkyLight && !bDynamicSkyLight && !bHasStaticLighting && !bReflectionEnv)
+		return;
+
+	FSceneTextureParameters SceneTextures = GetSceneTextureParameters(GraphBuilder);
+
+	check(VirtualVoxelResources.IsValid());
+
+	// Render the reflection environment with tiled deferred culling
+	const bool bHasBoxCaptures = (View.NumBoxReflectionCaptures > 0);
+	const bool bHasSphereCaptures = (View.NumSphereReflectionCaptures > 0);
+	FHairEnvironmentLightingPS::FParameters* ParametersPS = GraphBuilder.AllocParameters<FHairEnvironmentLightingPS::FParameters>();
+	FHairEnvironmentLighting::FParameters* PassParameters = &ParametersPS->Common;
+
+	PassParameters->SceneColorTexture = SceneColorTexture;
+	PassParameters->HairEnergyLUTTexture = GetHairLUT(GraphBuilder, View, HairLUTType_MeanEnergy);
+	PassParameters->bHasStaticLighting = bHasStaticLighting ? 1 : 0;
 	PassParameters->SkyDiffuseLighting = GetSkyDiffuseLightingParameters(Scene->SkyLight, 1.0f /*DynamicBentNormalAO*/);
-	PassParameters->Voxel_TanConeAngle = FMath::Tan(FMath::DegreesToRadians(GetHairStrandsSkyLightingConeAngle()));
 	PassParameters->HairDistanceThreshold = FMath::Max(GHairStrandsSkyLightingDistanceThreshold, 1.f);
 	PassParameters->bHairUseViewHairCount = VisibilityData.ViewHairCountTexture && GHairStrandsSkyLightingUseHairCountTexture ? 1 : 0;
 	PassParameters->MultipleScatterSampleCount = FMath::Max(uint32(GHairStrandsSkyLightingSampleCount), 1u);
@@ -391,6 +430,7 @@ static void AddHairStrandsEnvironmentLightingPassPS(
 	PermutationVector.Set<FHairEnvironmentLighting::FIntegrationType>(uint32(IntegrationType));
 	PermutationVector.Set<FHairEnvironmentLighting::FLumen>(LightingType == EHairLightingSourceType::Lumen);
 	PermutationVector.Set<FHairEnvironmentLighting::FDebug>(DebugData != nullptr);
+	PermutationVector.Set<FHairEnvironmentLighting::FStaticLighting>(bHasStaticLighting);
 	PermutationVector = FHairEnvironmentLighting::RemapPermutation(PermutationVector);
 
 	FIntPoint ViewportResolution = VisibilityData.SampleLightingViewportResolution;
@@ -502,6 +542,7 @@ void RenderHairStrandsEnvironmentLighting(
 	const FScene* Scene,
 	const FViewInfo& View)
 {
+	RDG_GPU_STAT_SCOPE(GraphBuilder, HairSkyLighting);
 	InternalRenderHairStrandsEnvironmentLighting(GraphBuilder, Scene, View, EHairLightingSourceType::ReflectionProbe);
 }
 
diff --git a/Engine/Source/Runtime/Renderer/Private/IndirectLightRendering.cpp b/Engine/Source/Runtime/Renderer/Private/IndirectLightRendering.cpp
--- a/Engine/Source/Runtime/Renderer/Private/IndirectLightRendering.cpp
+++ b/Engine/Source/Runtime/Renderer/Private/IndirectLightRendering.cpp
@@ -89,7 +89,6 @@ static TAutoConsoleVariable<int32> CVarProbeSamplePerPixel(
 
 DECLARE_GPU_STAT_NAMED(ReflectionEnvironment, TEXT("Reflection Environment"));
 DECLARE_GPU_STAT_NAMED(RayTracingReflections, TEXT("Ray Tracing Reflections"));
-DECLARE_GPU_STAT_NAMED(HairSkyLighting, TEXT("Hair Sky lighting"));
 DECLARE_GPU_STAT(SkyLightDiffuse);
 
 int GetReflectionEnvironmentCVar();
@@ -538,15 +537,6 @@ TUniformBufferRef<FReflectionUniformParameters> CreateReflectionUniformBuffer(co
 	return CreateUniformBufferImmediate(ReflectionStruct, Usage);
 }
 
-bool FDeferredShadingSceneRenderer::ShouldDoReflectionEnvironment() const
-{
-	const ERHIFeatureLevel::Type SceneFeatureLevel = Scene->GetFeatureLevel();
-
-	return IsReflectionEnvironmentAvailable(SceneFeatureLevel)
-		&& Scene->ReflectionSceneData.RegisteredReflectionCaptures.Num()
-		&& ViewFamily.EngineShowFlags.ReflectionEnvironment;
-}
-
 static const FVector SampleArray4x4x6[96] = {
 	FVector(0.72084325551986694, -0.44043412804603577, -0.53516626358032227),
 	FVector(-0.51286971569061279, 0.57541996240615845, 0.63706874847412109),
@@ -1411,7 +1401,7 @@ void FDeferredShadingSceneRenderer::RenderDeferredReflectionsAndSkyLighting(
 
 	RDG_EVENT_SCOPE(GraphBuilder, "ReflectionIndirect");
 
-	const bool bReflectionEnv = ShouldDoReflectionEnvironment();
+	const bool bReflectionEnv = ShouldDoReflectionEnvironment(Scene, ViewFamily);
 
 	FSceneTextureParameters SceneTextureParameters = GetSceneTextureParameters(GraphBuilder);
 	const auto& SceneColorTexture = SceneTextures.Color;
@@ -1638,10 +1628,8 @@ void FDeferredShadingSceneRenderer::RenderDeferredReflectionsAndSkyLighting(
 			}
 		}
 
-		const bool bIsHairSkyLightingEnabled = HairStrands::HasViewHairStrandsData(View) && (bSkyLight || bDynamicSkyLight || bReflectionEnv);
-		if (bIsHairSkyLightingEnabled)
+		if (HairStrands::HasViewHairStrandsData(View))
 		{
-			RDG_GPU_STAT_SCOPE(GraphBuilder, HairSkyLighting);
 			RenderHairStrandsEnvironmentLighting(GraphBuilder, Scene, View);
 		}
 	}
@@ -1664,15 +1652,7 @@ void FDeferredShadingSceneRenderer::RenderDeferredReflectionsAndSkyLightingHair(
 			continue;
 		}
 
-		// The specular sky light contribution is also needed by RT Reflections as a fallback.
-		const bool bSkyLight = Scene->SkyLight
-			&& Scene->SkyLight->ProcessedTexture
-			&& !Scene->SkyLight->bHasStaticLighting;
-
-		const bool bDynamicSkyLight = ShouldRenderDeferredDynamicSkyLight(Scene, ViewFamily);
-		const bool bReflectionEnv = ShouldDoReflectionEnvironment();
-		const bool bIsHairSkyLightingEnabled = HairStrands::HasViewHairStrandsData(View) && (bSkyLight || bDynamicSkyLight || bReflectionEnv);
-		if (bIsHairSkyLightingEnabled)
+		if (HairStrands::HasViewHairStrandsData(View))
 		{
 			RenderHairStrandsEnvironmentLighting(GraphBuilder, Scene, View);
 		}
