Fix missing micro-occlusion on rect. light specular component.

Indirectly this fixes the lack of specular cancelling on rect. light when setting ShowFlag.Specular is set to 0.

diff --git a/Engine/Shaders/Private/RectLight.ush b/Engine/Shaders/Private/RectLight.ush
--- a/Engine/Shaders/Private/RectLight.ush
+++ b/Engine/Shaders/Private/RectLight.ush
@@ -434,6 +434,10 @@ float3 PolygonIrradiance( float3 Poly[4] )
 // [ Heitz et al. 2016, "Real-Time Polygonal-Light Shading with Linearly Transformed Cosines" ]
 float3 RectGGXApproxLTC( float Roughness, float3 SpecularColor, half3 N, float3 V, FRect Rect, FRectTexture RectTexture )
 {
+	// F90 is derived as, anything less than 2% is physically impossible and is instead considered to be shadowing
+	const float3 F0  = SpecularColor;
+	const float3 F90 = saturate(50.0 * SpecularColor);
+
 	// No visibile rect light due to barn door occlusion
 	if (Rect.Extent.x == 0 || Rect.Extent.y == 0) return 0;
 
@@ -504,14 +508,14 @@ float3 RectGGXApproxLTC( float Roughness, float3 SpecularColor, half3 N, float3
 	Irradiance = Irradiance * ( 1 - a2*a2 ) + a2 * Falloff * NoL;
 #endif
 	
-	SpecularColor = LTCAmp.y + ( LTCAmp.x - LTCAmp.y ) * SpecularColor;
+	float3 IrradianceScale = F90 * LTCAmp.y + ( LTCAmp.x - LTCAmp.y ) * F0;
 	
 	// Transform to world space
 	L = mul( InvLTC, L );
 
 	float3 LightColor = SampleSourceTexture( L, Rect, RectTexture );
 	
-	return LightColor * Irradiance * SpecularColor;
+	return LightColor * Irradiance * IrradianceScale;
 }
 
 
