Fixed strata compilation post LWC fixes for 5.0 and main.

diff --git a/Engine/Shaders/Private/Strata/StrataVisualize.usf b/Engine/Shaders/Private/Strata/StrataVisualize.usf
--- a/Engine/Shaders/Private/Strata/StrataVisualize.usf
+++ b/Engine/Shaders/Private/Strata/StrataVisualize.usf
@@ -12,7 +12,7 @@
 #include "/Engine/Private/ShaderDrawDebug.ush"
 
 
-void DrawReferential(float3 P /*In Translated World Space*/, float3 X, float3 Y, float3 N, float3 InColor)
+void DrawReferentialTWS(float3 P /*In Translated World Space*/, float3 X, float3 Y, float3 N, float3 InColor)
 {
 	const float Size = 10.f;
 	const float SizeWithTip = 12.0f;
