Shadercompile fix for Nanite + lightmap VT path
Added new TextureLoadVirtualPageTableGrad overload variation that takes two page tables and doesn't do feedback

diff --git a/Engine/Shaders/Private/LightmapCommon.ush b/Engine/Shaders/Private/LightmapCommon.ush
--- a/Engine/Shaders/Private/LightmapCommon.ush
+++ b/Engine/Shaders/Private/LightmapCommon.ush
@@ -15,21 +15,22 @@
 
 #if IS_NANITE_PASS
 #define LightmapUVType FloatDeriv2
-LightmapUVType ScaleLightmapUV(LightmapUVType UV, float2 Scale)
-{
-	UV.Value *= Scale;
-	UV.Ddx *= Scale;
-	UV.Ddy *= Scale;
-	return UV;
-}
 #else
 #define LightmapUVType float2
-LightmapUVType ScaleLightmapUV(LightmapUVType UV, float2 Scale)
+#endif
+
+float2 ScaleLightmapUV(float2 UV, float2 Scale)
 {
 	return UV * Scale;
 }
-#endif
 
+FloatDeriv2 ScaleLightmapUV(FloatDeriv2 UV, float2 Scale)
+{
+	UV.Value *= Scale;
+	UV.Ddx *= Scale;
+	UV.Ddy *= Scale;
+	return UV;
+}
 
 #if HQ_TEXTURE_LIGHTMAP || LQ_TEXTURE_LIGHTMAP
  
@@ -49,7 +50,7 @@ LightmapUVType ScaleLightmapUV(LightmapUVType UV, float2 Scale)
 #if LIGHTMAP_VT_ENABLED
 VTPageTableResult LightmapGetVTSampleInfo(float2 UV, uint LightmapDataIndex, float2 SvPositionXY)
 {
-	UV *= float2(1.0f, 2.0f); // Undo transform used to pack 2 lightmap coeffs in 1 texture for the non-VT default case
+	UV = ScaleLightmapUV(UV, float2(1.0f, 2.0f)); // Undo transform used to pack 2 lightmap coeffs in 1 texture for the non-VT default case
 
 	return TextureLoadVirtualPageTable(LightmapResourceCluster.LightmapVirtualTexturePageTable0, LightmapResourceCluster.LightmapVirtualTexturePageTable1,
 		VTPageTableUniform_Unpack(GetLightmapData(LightmapDataIndex).LightmapVTPackedPageTableUniform[0], GetLightmapData(LightmapDataIndex).LightmapVTPackedPageTableUniform[1]),
@@ -58,11 +59,11 @@ VTPageTableResult LightmapGetVTSampleInfo(float2 UV, uint LightmapDataIndex, flo
 
 VTPageTableResult LightmapGetVTSampleInfo(FloatDeriv2 UV, uint LightmapDataIndex, float2 SvPositionXY)
 {
-	ScaleLightmapUV(UV, float2(1.0f, 2.0f)); // Undo transform used to pack 2 lightmap coeffs in 1 texture for the non-VT default case
+	UV = ScaleLightmapUV(UV, float2(1.0f, 2.0f)); // Undo transform used to pack 2 lightmap coeffs in 1 texture for the non-VT default case
 	
 	return TextureLoadVirtualPageTableGrad(LightmapResourceCluster.LightmapVirtualTexturePageTable0, LightmapResourceCluster.LightmapVirtualTexturePageTable1,
 		VTPageTableUniform_Unpack(GetLightmapData(LightmapDataIndex).LightmapVTPackedPageTableUniform[0], GetLightmapData(LightmapDataIndex).LightmapVTPackedPageTableUniform[1]),
-		UV.Value, LIGHTMAP_VTADDRESSMODE, LIGHTMAP_VTADDRESSMODE, UV.Ddx, UV.Ddy, 0, SvPositionXY);
+		UV.Value, LIGHTMAP_VTADDRESSMODE, LIGHTMAP_VTADDRESSMODE, UV.Ddx, UV.Ddy, SvPositionXY);
 }
 
 half4 SampleLightmapVT(VTPageTableResult LightmapVTPageTableResult, uint LayerIndex, uint LightmapDataIndex, Texture2D CacheTexture, SamplerState CacheSampler)
diff --git a/Engine/Shaders/Private/VirtualTextureCommon.ush b/Engine/Shaders/Private/VirtualTextureCommon.ush
--- a/Engine/Shaders/Private/VirtualTextureCommon.ush
+++ b/Engine/Shaders/Private/VirtualTextureCommon.ush
@@ -477,6 +477,21 @@ VTPageTableResult TextureLoadVirtualPageTableGrad(
 	return Result;
 }
 
+// LoadPageTable: 2 Page tables, SampleGrad, no feedback
+VTPageTableResult TextureLoadVirtualPageTableGrad(
+	Texture2D<uint4> PageTable0, Texture2D<uint4> PageTable1,
+	VTPageTableUniform PageTableUniform,
+	float2 UV, uint AddressU, uint AddressV,
+	MaterialFloat2 dUVdx, MaterialFloat2 dUVdy, float2 SvPositionXY)
+{
+	VTPageTableResult Result = (VTPageTableResult)0.0f;
+	int vLevel = TextureComputeVirtualMipLevel(Result, dUVdx * PageTableUniform.UVScale, dUVdy * PageTableUniform.UVScale, 0, SvPositionXY, PageTableUniform);
+	UV = UV * PageTableUniform.UVScale;
+	UV = ApplyAddressMode(UV, AddressU, AddressV);
+	TextureLoadVirtualPageTableInternal(Result, PageTable0, PageTable1, PageTableUniform, UV, vLevel);
+	return Result;
+}
+
 // LoadPageTable: 1 Page table, SampleLevel
 VTPageTableResult TextureLoadVirtualPageTableLevel(
 	Texture2D<uint4> PageTable0,
