Change the tonemapper's output pixel format to match the view family's pixel format to avoid bitdepth losses in editor

diff --git a/Engine/Shaders/Private/Fxaa3_11.ush b/Engine/Shaders/Private/Fxaa3_11.ush
--- a/Engine/Shaders/Private/Fxaa3_11.ush
+++ b/Engine/Shaders/Private/Fxaa3_11.ush
@@ -307,18 +307,19 @@ A. Or use FXAA_GREEN_AS_LUMA.
     // 1 = API supports gather4 on alpha channel.
     // 0 = API does not support gather4 on alpha channel.
     //
-    #if (FXAA_HLSL_5 == 1)
-        #define FXAA_GATHER4_ALPHA 1
-    #endif
-    #ifdef GL_ARB_gpu_shader5
-        #define FXAA_GATHER4_ALPHA 1
-    #endif
-    #ifdef GL_NV_gpu_shader5
-        #define FXAA_GATHER4_ALPHA 1
-    #endif
-    #ifndef FXAA_GATHER4_ALPHA
-        #define FXAA_GATHER4_ALPHA 0
-    #endif
+    //#if (FXAA_HLSL_5 == 1)
+    //    #define FXAA_GATHER4_ALPHA 1
+    //#endif
+    //#ifdef GL_ARB_gpu_shader5
+    //    #define FXAA_GATHER4_ALPHA 1
+    //#endif
+    //#ifdef GL_NV_gpu_shader5
+    //    #define FXAA_GATHER4_ALPHA 1
+    //#endif
+    //#ifndef FXAA_GATHER4_ALPHA
+    //    #define FXAA_GATHER4_ALPHA 0
+    //#endif
+	#define FXAA_GATHER4_ALPHA 0
 #endif
 
 /*============================================================================
@@ -705,9 +706,8 @@ NOTE the other tuning knobs are now in the shader function inputs!
 ============================================================================*/
 #if (FXAA_GREEN_AS_LUMA == 0)
 
-	#if POST_PROCESS_ALPHA == 2
-		// POST_PROCESS_ALPHA == 2 allows the tonemapper to not output Luma in alpha channel,
-		// therefore need to compute it manually.
+	#if 1
+		// allows the tonemapper to not output Luma in alpha channel to use a R10G10BA2
 		FxaaFloat FxaaLuma(FxaaFloat4 rgba)
 		{
 			return dot(rgba.rgb, FxaaFloat3(0.299f, 0.587f, 0.114f));
@@ -898,7 +898,7 @@ FxaaFloat4 FxaaPixelShader(
         #if (FXAA_DISCARD == 0)
             FxaaFloat4 rgbyM = FxaaTexTop(tex, posM);
             #if (FXAA_GREEN_AS_LUMA == 0)
-                #define lumaM rgbyM.w
+                #define lumaM FxaaLuma(rgbyM)
             #else
                 #define lumaM rgbyM.y
             #endif
@@ -922,7 +922,7 @@ FxaaFloat4 FxaaPixelShader(
     #else
         FxaaFloat4 rgbyM = FxaaTexTop(tex, posM);
         #if (FXAA_GREEN_AS_LUMA == 0)
-            #define lumaM rgbyM.w
+            #define lumaM FxaaLuma(rgbyM)
         #else
             #define lumaM rgbyM.y
         #endif
@@ -1244,7 +1244,7 @@ FxaaFloat4 FxaaPixelShader(
     FxaaFloat pixelOffsetSubpix = max(pixelOffsetGood, subpixH);
     if(!horzSpan) posM.x += pixelOffsetSubpix * lengthSign;
     if( horzSpan) posM.y += pixelOffsetSubpix * lengthSign;
-    #if ((FXAA_DISCARD == 1) || (POST_PROCESS_ALPHA == 2))
+    #if ((FXAA_DISCARD == 1) || 1)
         return FxaaTexTop(tex, posM);
     #else
         return FxaaFloat4(FxaaTexTop(tex, posM).xyz, lumaM);
@@ -1296,7 +1296,7 @@ FxaaFloat4 FxaaPixelShader(
 /*--------------------------------------------------------------------------*/
     FxaaFloat4 rgbyM = FxaaTexTop(tex, pos.xy);
     #if (FXAA_GREEN_AS_LUMA == 0)
-        FxaaFloat lumaM = rgbyM.w;
+        FxaaFloat lumaM = FxaaLuma(rgbyM);
     #else
         FxaaFloat lumaM = rgbyM.y;
     #endif
@@ -1339,7 +1339,7 @@ FxaaFloat4 FxaaPixelShader(
     FxaaFloat4 rgbyB = ((rgbyN2 + rgbyP2) * 0.25) + (rgbyA * 0.25);
 /*--------------------------------------------------------------------------*/
     #if (FXAA_GREEN_AS_LUMA == 0)
-        FxaaBool twoTap = (rgbyB.w < lumaMin) || (rgbyB.w > lumaMax);
+        FxaaBool twoTap = (FxaaLuma(rgbyB) < lumaMin) || (FxaaLuma(rgbyB) > lumaMax);
     #else
         FxaaBool twoTap = (rgbyB.y < lumaMin) || (rgbyB.y > lumaMax);
     #endif
diff --git a/Engine/Shaders/Private/PostProcessTonemap.usf b/Engine/Shaders/Private/PostProcessTonemap.usf
--- a/Engine/Shaders/Private/PostProcessTonemap.usf
+++ b/Engine/Shaders/Private/PostProcessTonemap.usf
@@ -649,8 +649,6 @@ float4 TonemapCommonPS(
 
 		half3 OutDeviceColor = ColorLookupTable(LinearColor);
 
-		half LuminanceForPostProcessAA  = dot(OutDeviceColor, half3(0.299f, 0.587f, 0.114f));
-
 		#if USE_GRAIN_QUANTIZATION
 			half Grain = GrainFromUV(GrainUV.zw);
 
@@ -672,7 +670,7 @@ float4 TonemapCommonPS(
 		#endif
 
 		// RETURN_COLOR not needed unless writing to SceneColor
-		OutColor = float4(OutDeviceColor, saturate(LuminanceForPostProcessAA));
+		OutColor.rgb = OutDeviceColor;
 
 		#if USE_LINEAR_FLOAT_RT
 			OutColor.rgb = sRGBToLinear( OutColor.rgb );
diff --git a/Engine/Source/Runtime/Renderer/Private/PostProcess/PostProcessTonemap.cpp b/Engine/Source/Runtime/Renderer/Private/PostProcess/PostProcessTonemap.cpp
--- a/Engine/Source/Runtime/Renderer/Private/PostProcess/PostProcessTonemap.cpp
+++ b/Engine/Source/Runtime/Renderer/Private/PostProcess/PostProcessTonemap.cpp
@@ -606,7 +606,6 @@ FScreenPassTexture AddTonemapPass(FRDGBuilder& GraphBuilder, const FViewInfo& Vi
 		OutputDesc.Flags |= View.bUseComputePasses ? TexCreate_UAV : TexCreate_RenderTargetable;
 		OutputDesc.Flags |= GFastVRamConfig.Tonemap;
 		// RGB is the color in LDR, A is the luminance for PostprocessAA
-		OutputDesc.Format = Inputs.bOutputInHDR ? GRHIHDRDisplayOutputFormat : PF_B8G8R8A8;
 		OutputDesc.ClearValue = FClearValueBinding(FLinearColor(0, 0, 0, 0));
 
 		const FTonemapperOutputDeviceParameters OutputDeviceParameters = GetTonemapperOutputDeviceParameters(*View.Family);
@@ -616,10 +615,29 @@ FScreenPassTexture AddTonemapPass(FRDGBuilder& GraphBuilder, const FViewInfo& Vi
 		{
 			OutputDesc.Format = PF_A32B32G32R32F;
 		}
-		if (OutputDevice == ETonemapperOutputDevice::LinearNoToneCurve || OutputDevice == ETonemapperOutputDevice::LinearWithToneCurve)
+		else if (OutputDevice == ETonemapperOutputDevice::LinearNoToneCurve || OutputDevice == ETonemapperOutputDevice::LinearWithToneCurve)
 		{
 			OutputDesc.Format = PF_FloatRGBA;
 		}
+		else if (Inputs.bOutputInHDR)
+		{
+			OutputDesc.Format = GRHIHDRDisplayOutputFormat;
+		}
+		else if (View.Family->RenderTarget && View.Family->RenderTarget->GetRenderTargetTexture())
+		{
+			// Render into a pixel format that do not loose bit depth precision for the view family.
+			OutputDesc.Format = View.Family->RenderTarget->GetRenderTargetTexture()->GetFormat();
+		}
+		else if (IsPostProcessingWithAlphaChannelSupported())
+		{
+			// Make sure there is no loss for a 10bit bit-depth using the 10bit of mantissa of halfs
+			OutputDesc.Format = PF_FloatRGBA;
+		}
+		else 
+		{
+			// Make sure there is no loss for a 10bit bit-depth
+			OutputDesc.Format = PF_A2B10G10R10;
+		}
 
 		Output = FScreenPassRenderTarget(
 			GraphBuilder.CreateTexture(OutputDesc, TEXT("Tonemap")),
