Refactored Distance Field culling
* Removed culled DF object copies during culling. Instead now only indices to culled objects are stored
* Refactored DF heightfield object loads into FHeightfieldObjectBounds and FHeightfieldObjectData

This is a step towards optimizing DF culling and reusing this code for Lumen Landscape culling

Perf Reverb on 2080:
* CullMeshSDFObjectsToFrustum 0.04ms->0.03ms (removed DF object copies)
* Other passes didn't change

diff --git a/Engine/Shaders/Private/DistanceFieldLightingShared.ush b/Engine/Shaders/Private/DistanceFieldLightingShared.ush
--- a/Engine/Shaders/Private/DistanceFieldLightingShared.ush
+++ b/Engine/Shaders/Private/DistanceFieldLightingShared.ush
@@ -131,6 +131,9 @@ FDFAssetData LoadDFAssetDataHighestResolution(uint AssetIndex)
 RWBuffer<uint> RWObjectIndirectArguments;
 Buffer<uint> ObjectIndirectArguments;
 
+RWStructuredBuffer<uint> RWCulledObjectIndices;
+StructuredBuffer<uint> CulledObjectIndices;
+
 uint GetCulledNumObjects()
 {
 	// IndexCount, NumInstances, StartIndex, BaseVertexIndex, FirstInstance
@@ -138,15 +141,10 @@ uint GetCulledNumObjects()
 }
 
 // In float4's.  Must match equivalent C++ variables.
-#if (DISTANCEFIELD_PRIMITIVE_TYPE == DFPT_SignedDistanceField)
-#define OBJECT_BOUNDS_STRIDE 2
-#define OBJECT_DATA_STRIDE 9
-#elif (DISTANCEFIELD_PRIMITIVE_TYPE == DFPT_HeightField)
-#define OBJECT_BOUNDS_STRIDE 2
-#define OBJECT_DATA_STRIDE 6
-#else
-#error Unknown distance field primitive type
-#endif
+#define DF_OBJECT_BOUNDS_STRIDE 2
+#define DF_OBJECT_DATA_STRIDE 9
+#define HEIGHTFIELD_OBJECT_BOUNDS_STRIDE 2
+#define HEIGHTFIELD_OBJECT_DATA_STRIDE 6
 
 struct FDFObjectBounds
 {
@@ -161,17 +159,21 @@ struct FDFObjectBounds
 
 uint NumSceneObjects;
 StructuredBuffer<float4> SceneObjectBounds;
+
+StructuredBuffer<float4> SceneHeightfieldObjectBounds;
+uint NumSceneHeightfieldObjects;
+
 SamplerState DistanceFieldSampler;
 
 FDFObjectBounds LoadDFObjectBounds(uint ObjectIndex)
 {
 	FDFObjectBounds Bounds;
 
-	float4 Vector0 = SceneObjectBounds[ObjectIndex * OBJECT_BOUNDS_STRIDE + 0];
+	float4 Vector0 = SceneObjectBounds[ObjectIndex * DF_OBJECT_BOUNDS_STRIDE + 0];
 	Bounds.Center = Vector0.xyz;
 	Bounds.SphereRadius = Vector0.w;
 
-	float4 Vector1 = SceneObjectBounds[ObjectIndex * OBJECT_BOUNDS_STRIDE + 1];
+	float4 Vector1 = SceneObjectBounds[ObjectIndex * DF_OBJECT_BOUNDS_STRIDE + 1];
 	Bounds.BoxExtent = Vector1.xyz;
 
 	uint Flags = asuint(Vector1.w);
@@ -183,6 +185,25 @@ FDFObjectBounds LoadDFObjectBounds(uint ObjectIndex)
 	return Bounds;
 }
 
+struct FHeightfieldObjectBounds
+{
+	float3 BoxOrigin;	// World space AABB center
+	float3 BoxExtent;	// World space AABB extent
+};
+
+FHeightfieldObjectBounds LoadHeightfieldObjectBounds(uint ObjectIndex)
+{
+	FHeightfieldObjectBounds Bounds;
+
+	float4 Vector0 = SceneHeightfieldObjectBounds[ObjectIndex * HEIGHTFIELD_OBJECT_BOUNDS_STRIDE + 0];
+	Bounds.BoxOrigin = Vector0.xyz;
+
+	float4 Vector1 = SceneHeightfieldObjectBounds[ObjectIndex * HEIGHTFIELD_OBJECT_BOUNDS_STRIDE + 1];
+	Bounds.BoxExtent = Vector1.xyz;
+
+	return Bounds;
+}
+
 struct FDFObjectData
 {
 	// Extent of the mesh bounds in volume space
@@ -220,27 +241,27 @@ FDFObjectData LoadDFObjectDataFromBuffer(StructuredBuffer<float4> SourceBuffer,
 {
 	FDFObjectData Data;
 
-	float4 V0 = SourceBuffer[ObjectIndex * OBJECT_DATA_STRIDE + 0];
-	float4 V1 = SourceBuffer[ObjectIndex * OBJECT_DATA_STRIDE + 1];
-	float4 V2 = SourceBuffer[ObjectIndex * OBJECT_DATA_STRIDE + 2];
+	float4 V0 = SourceBuffer[ObjectIndex * DF_OBJECT_DATA_STRIDE + 0];
+	float4 V1 = SourceBuffer[ObjectIndex * DF_OBJECT_DATA_STRIDE + 1];
+	float4 V2 = SourceBuffer[ObjectIndex * DF_OBJECT_DATA_STRIDE + 2];
 	Data.WorldToVolume = transpose(float4x4(V0, V1, V2, float4(0.0f, 0.0f, 0.0f, 1.0f)));
 
-	float4 Vector3 = SourceBuffer[ObjectIndex * OBJECT_DATA_STRIDE + 3];
+	float4 Vector3 = SourceBuffer[ObjectIndex * DF_OBJECT_DATA_STRIDE + 3];
 	Data.VolumePositionExtent = Vector3.xyz;
 	Data.VolumeSurfaceBias = abs(Vector3.w);
 	Data.bMostlyTwoSided = Vector3.w < 0.0f;
 
-	float4 Vector4 = SourceBuffer[ObjectIndex * OBJECT_DATA_STRIDE + 4];
+	float4 Vector4 = SourceBuffer[ObjectIndex * DF_OBJECT_DATA_STRIDE + 4];
 	Data.MinMaxDrawDistance2 = Vector4.xy;
 	Data.SelfShadowBias = Vector4.z;
 	Data.GPUSceneInstanceIndex = asuint(Vector4.w);
 	
-	V0 = SourceBuffer[ObjectIndex * OBJECT_DATA_STRIDE + 5];
-	V1 = SourceBuffer[ObjectIndex * OBJECT_DATA_STRIDE + 6];
-	V2 = SourceBuffer[ObjectIndex * OBJECT_DATA_STRIDE + 7];
+	V0 = SourceBuffer[ObjectIndex * DF_OBJECT_DATA_STRIDE + 5];
+	V1 = SourceBuffer[ObjectIndex * DF_OBJECT_DATA_STRIDE + 6];
+	V2 = SourceBuffer[ObjectIndex * DF_OBJECT_DATA_STRIDE + 7];
 	Data.VolumeToWorld = transpose(float4x4(V0, V1, V2, float4(0.0f, 0.0f, 0.0f, 1.0f)));
 
-	float4 Vector8 = SourceBuffer[ObjectIndex * OBJECT_DATA_STRIDE + 8];
+	float4 Vector8 = SourceBuffer[ObjectIndex * DF_OBJECT_DATA_STRIDE + 8];
 	Data.VolumeToWorldScale = Vector8.xyz;
 	Data.VolumeScale = min(Data.VolumeToWorldScale.x, min(Data.VolumeToWorldScale.y, Data.VolumeToWorldScale.z));
 	Data.AssetIndex = asuint(Vector8.w);
@@ -248,120 +269,41 @@ FDFObjectData LoadDFObjectDataFromBuffer(StructuredBuffer<float4> SourceBuffer,
 	return Data;
 }
 
-StructuredBuffer<float4> SceneObjectData;
-
-FDFObjectData LoadDFObjectData(uint ObjectIndex)
+struct FHeightfieldObjectData
 {
-	return LoadDFObjectDataFromBuffer(SceneObjectData, ObjectIndex);
-}
-
-// In float4's.  Must match equivalent C++ variables.
-#if (DISTANCEFIELD_PRIMITIVE_TYPE == DFPT_SignedDistanceField)
-#define CULLED_OBJECT_DATA_STRIDE OBJECT_DATA_STRIDE
-#define CULLED_OBJECT_BOX_BOUNDS_STRIDE 5
-#elif (DISTANCEFIELD_PRIMITIVE_TYPE == DFPT_HeightField)
-#define CULLED_OBJECT_DATA_STRIDE 6
-#define CULLED_OBJECT_BOX_BOUNDS_STRIDE 5
-#else
-#error Unknown distance field primitive type
-#endif
-
-float4 LoadObjectPositionAndRadiusFromBuffer(uint ObjectIndex, StructuredBuffer<float4> InBuffer)
-{
-	return InBuffer[ObjectIndex];
-}
-
-void LoadHeightFieldObjectBoxBoundsFromBuffer(uint ObjectIndex, StructuredBuffer<float4> InBuffer, out float3 OutBoxOrigin, out float3 OutBoxExtent)
-{
-	OutBoxOrigin = InBuffer[ObjectIndex * OBJECT_BOUNDS_STRIDE + 0].xyz;
-	OutBoxExtent = InBuffer[ObjectIndex * OBJECT_BOUNDS_STRIDE + 1].xyz;
-}
-
-float4x4 LoadHeightFieldObjectWorldToLocalFromBuffer(uint ObjectIndex, StructuredBuffer<float4> InBuffer)
-{
-	float4 M0 = InBuffer[ObjectIndex * CULLED_OBJECT_DATA_STRIDE + 0];
-	float4 M1 = InBuffer[ObjectIndex * CULLED_OBJECT_DATA_STRIDE + 1];
-	float4 M2 = InBuffer[ObjectIndex * CULLED_OBJECT_DATA_STRIDE + 2];
-
-	return transpose(float4x4(M0, M1, M2, float4(0, 0, 0, 1)));
-}
-
-float4 LoadHeightFieldObjectSizeScaleFromBuffer(uint ObjectIndex, StructuredBuffer<float4> InBuffer)
-{
-	return InBuffer[ObjectIndex * CULLED_OBJECT_DATA_STRIDE + 3];
-}
-
-float4 LoadHeightFieldObjectAtlasUVScaleBiasFromBuffer(uint ObjectIndex, StructuredBuffer<float4> InBuffer)
-{
-	return InBuffer[ObjectIndex * CULLED_OBJECT_DATA_STRIDE + 4];
-}
-
-float4 LoadHeightFieldObjectVisibilityAtlasUVScaleBiasFromBuffer(uint ObjectIndex, StructuredBuffer<float4> InBuffer)
-{
-	return InBuffer[ObjectIndex * CULLED_OBJECT_DATA_STRIDE + 5];
-}
-
-// These are structured buffers so they can be scalar memory loads on GCN
-StructuredBuffer<float4> CulledObjectBounds;
-StructuredBuffer<float4> CulledObjectData;
-StructuredBuffer<float4> CulledObjectBoxBounds;
-
-RWStructuredBuffer<float4> RWCulledObjectBounds;
-RWStructuredBuffer<float4> RWCulledObjectData;
-RWStructuredBuffer<float4> RWCulledObjectBoxBounds;
-
-FDFObjectData LoadCulledDFObjectData(uint ObjectIndex)
-{
-	return LoadDFObjectDataFromBuffer(CulledObjectData, ObjectIndex);
-}
-
-float4 LoadObjectPositionAndRadius(uint ObjectIndex)
-{
-	return LoadObjectPositionAndRadiusFromBuffer(ObjectIndex, CulledObjectBounds);
-}
+	float4x4 WorldToLocal;
+	float4 SizeScale;
+	float4 AtlasUVScaleBias;
+	float4 VisibilityAtlasUVScaleBias;
+};
 
-void LoadHeightFieldObjectBoxBounds(uint ObjectIndex, out float3 OutBoxOrigin, out float3 OutBoxExtent)
+FHeightfieldObjectData LoadHeightfieldObjectDataFromBuffer(StructuredBuffer<float4> SourceBuffer, uint ObjectIndex)
 {
-	LoadHeightFieldObjectBoxBoundsFromBuffer(ObjectIndex, CulledObjectBounds, OutBoxOrigin, OutBoxExtent);
-}
+	FHeightfieldObjectData Data;
 
-float4x4 LoadHeightFieldObjectWorldToLocal(uint ObjectIndex)
-{
-	return LoadHeightFieldObjectWorldToLocalFromBuffer(ObjectIndex, CulledObjectData);
-}
+	float4 V0 = SourceBuffer[ObjectIndex * HEIGHTFIELD_OBJECT_DATA_STRIDE + 0];
+	float4 V1 = SourceBuffer[ObjectIndex * HEIGHTFIELD_OBJECT_DATA_STRIDE + 1];
+	float4 V2 = SourceBuffer[ObjectIndex * HEIGHTFIELD_OBJECT_DATA_STRIDE + 2];
+	Data.WorldToLocal = transpose(float4x4(V0, V1, V2, float4(0.0f, 0.0f, 0.0f, 1.0f)));
 
-float4 LoadHeightFieldObjectAtlasUVScaleBias(uint ObjectIndex)
-{
-	return LoadHeightFieldObjectAtlasUVScaleBiasFromBuffer(ObjectIndex, CulledObjectData);
-}
+	Data.SizeScale = SourceBuffer[ObjectIndex * HEIGHTFIELD_OBJECT_DATA_STRIDE + 3];
+	Data.AtlasUVScaleBias = SourceBuffer[ObjectIndex * HEIGHTFIELD_OBJECT_DATA_STRIDE + 4];
+	Data.VisibilityAtlasUVScaleBias = SourceBuffer[ObjectIndex * HEIGHTFIELD_OBJECT_DATA_STRIDE + 5];
 
-float4 LoadHeightFieldObjectSizeScale(uint ObjectIndex)
-{
-	return LoadHeightFieldObjectSizeScaleFromBuffer(ObjectIndex, CulledObjectData);
-}
-
-float4 LoadHeightFieldObjectVisibilityAtlasUVScaleBias(uint ObjectIndex)
-{
-	return LoadHeightFieldObjectVisibilityAtlasUVScaleBiasFromBuffer(ObjectIndex, CulledObjectData);
+	return Data;
 }
 
-// x = Offset in global buffer, y = NumLOD0, z = NumSurfels (all LODs), W = instance index
-uint4 LoadObjectSurfelCoordinate(uint ObjectIndex)
-{
-	return uint4(0, 0, 0, 0);
-}
+StructuredBuffer<float4> SceneObjectData;
+StructuredBuffer<float4> SceneHeightfieldObjectData;
 
-void LoadObjectViewSpaceBox(uint ObjectIndex, out float3 ObjectViewSpaceMin, out float3 ObjectViewSpaceMax)
+FDFObjectData LoadDFObjectData(uint ObjectIndex)
 {
-	ObjectViewSpaceMin = CulledObjectBoxBounds[ObjectIndex * CULLED_OBJECT_BOX_BOUNDS_STRIDE + 0].xyz;
-	ObjectViewSpaceMax = CulledObjectBoxBounds[ObjectIndex * CULLED_OBJECT_BOX_BOUNDS_STRIDE + 1].xyz;
+	return LoadDFObjectDataFromBuffer(SceneObjectData, ObjectIndex);
 }
 
-void LoadObjectAxes(uint ObjectIndex, out float3 ObjectAxisX, out float3 ObjectAxisY, out float3 ObjectAxisZ)
+FHeightfieldObjectData LoadHeightfieldObjectData(uint ObjectIndex)
 {
-	ObjectAxisX = CulledObjectBoxBounds[ObjectIndex * CULLED_OBJECT_BOX_BOUNDS_STRIDE + 2].xyz;
-	ObjectAxisY = CulledObjectBoxBounds[ObjectIndex * CULLED_OBJECT_BOX_BOUNDS_STRIDE + 3].xyz;
-	ObjectAxisZ = CulledObjectBoxBounds[ObjectIndex * CULLED_OBJECT_BOX_BOUNDS_STRIDE + 4].xyz;
+	return LoadHeightfieldObjectDataFromBuffer(SceneHeightfieldObjectData, ObjectIndex);
 }
 
 // VPLs generated by raytracing from the light
diff --git a/Engine/Shaders/Private/DistanceFieldObjectCulling.usf b/Engine/Shaders/Private/DistanceFieldObjectCulling.usf
--- a/Engine/Shaders/Private/DistanceFieldObjectCulling.usf
+++ b/Engine/Shaders/Private/DistanceFieldObjectCulling.usf
@@ -17,20 +17,6 @@ groupshared uint NumGroupObjects;
 groupshared uint GroupBaseIndex;
 groupshared uint GroupObjectIndices[UPDATEOBJECTS_THREADGROUP_SIZE];
 
-void CopyCulledObjectData(uint DestIndex, uint SourceIndex)
-{
-	RWCulledObjectBounds[DestIndex] = SceneObjectBounds[SourceIndex * OBJECT_BOUNDS_STRIDE + 0];
-
-	UNROLL
-	for (uint VectorIndex = 0; VectorIndex < CULLED_OBJECT_DATA_STRIDE; VectorIndex++)
-	{
-		float4 Data = SceneObjectData[SourceIndex * OBJECT_DATA_STRIDE + VectorIndex];
-
-		// Note: only copying the first CULLED_OBJECT_DATA_STRIDE of the original object data
-		RWCulledObjectData[DestIndex * CULLED_OBJECT_DATA_STRIDE + VectorIndex] = Data;
-	}
-}
-
 [numthreads(UPDATEOBJECTS_THREADGROUP_SIZE, 1, 1)]
 void CullObjectsForViewCS(
 	uint3 GroupId : SV_GroupID,
@@ -92,7 +78,7 @@ void CullObjectsForViewCS(
 	{
 		uint SourceIndex = GroupObjectIndices[GroupThreadId.x];
 		uint DestIndex = GroupBaseIndex + GroupThreadId.x;
-		CopyCulledObjectData(DestIndex, SourceIndex);
+		RWCulledObjectIndices[DestIndex] = SourceIndex;
 	}
 
 #else
@@ -111,7 +97,7 @@ void CullObjectsForViewCS(
 		uint SourceIndex = ObjectIndex;
 		uint DestIndex = ObjectIndex;
 
-		CopyCulledObjectData(DestIndex, SourceIndex);
+		RWCulledObjectIndices[DestIndex] = SourceIndex;
 	}
 
 #endif
@@ -254,18 +240,21 @@ float ConservativeRadiusScale;
 /** Used when culling objects into screenspace tile lists */
 void ObjectCullVS(
 	float4 InPosition : ATTRIBUTE0,
-	uint ObjectIndex : SV_InstanceID,
+	uint InstanceIndex : SV_InstanceID,
 	out FObjectCullVertexOutput Output,
 	out float4 OutPosition : SV_POSITION
 	)
 {
+	const uint ObjectIndex = CulledObjectIndices[InstanceIndex];
+
 	//@todo - implement ConservativelyBoundSphere
-	float4 ObjectPositionAndRadius = LoadObjectPositionAndRadius(ObjectIndex);
+	FDFObjectBounds ObjectBounds = LoadDFObjectBounds(ObjectIndex);
 	//@todo - expand to handle conservative rasterization
-	float EffectiveRadius = (ObjectPositionAndRadius.w + AOObjectMaxDistance) * ConservativeRadiusScale;
-	float3 WorldPosition = InPosition.xyz * EffectiveRadius + ObjectPositionAndRadius.xyz;
+	float EffectiveRadius = (ObjectBounds.SphereRadius + AOObjectMaxDistance) * ConservativeRadiusScale;
+	float3 WorldPosition = InPosition.xyz * EffectiveRadius + ObjectBounds.Center;
 	OutPosition = mul(float4(WorldPosition, 1), LWCHackToFloat(PrimaryView.WorldToClip));
-	Output.PositionAndRadius = ObjectPositionAndRadius;
+	Output.PositionAndRadius.xyz = ObjectBounds.Center;
+	Output.PositionAndRadius.w = ObjectBounds.SphereRadius;
 	Output.ObjectIndex = ObjectIndex;
 } 
 
@@ -284,7 +273,7 @@ bool IntersectObjectWithConeDepthRange(
 	{
 #define USE_DISTANCE_FIELD_FOR_OBJECT_CULLING 1
 #if USE_DISTANCE_FIELD_FOR_OBJECT_CULLING
-		FDFObjectData DFObjectData = LoadCulledDFObjectData(ObjectIndex);
+		FDFObjectData DFObjectData = LoadDFObjectData(ObjectIndex);
 
 		// Use the position halfway between the depth ranges as the center for the bounding sphere of this tile depth range
 		float3 ViewTileBoundingSphereCenter = TileConeVertex + TileConeAxis * (.5f * (ConeDepthRange.x + ConeDepthRange.y));
diff --git a/Engine/Shaders/Private/DistanceFieldScreenGridLighting.usf b/Engine/Shaders/Private/DistanceFieldScreenGridLighting.usf
--- a/Engine/Shaders/Private/DistanceFieldScreenGridLighting.usf
+++ b/Engine/Shaders/Private/DistanceFieldScreenGridLighting.usf
@@ -75,14 +75,14 @@ void HemisphereConeTraceAgainstTileCulledObjects(uint ObjectIndex, uint OutputBa
 	InvMaxOcclusionDistance = 1.0f / AOGlobalMaxOcclusionDistance;
 #endif
 
-	float4 ObjectPositionAndRadius = LoadObjectPositionAndRadius(ObjectIndex);
-	float ObjectDistanceSq = dot(ObjectPositionAndRadius.xyz - WorldShadingPosition, ObjectPositionAndRadius.xyz - WorldShadingPosition);
+	FDFObjectBounds ObjectBounds = LoadDFObjectBounds(ObjectIndex);
+	float ObjectDistanceSq = dot(ObjectBounds.Center - WorldShadingPosition, ObjectBounds.Center - WorldShadingPosition);
 
 	BRANCH
 	// Skip tracing objects with a small projected angle 
-	if (ObjectPositionAndRadius.w * ObjectPositionAndRadius.w / ObjectDistanceSq > Square(.25f))
+	if (ObjectBounds.SphereRadius * ObjectBounds.SphereRadius / ObjectDistanceSq > Square(.25f))
 	{
-		FDFObjectData DFObjectData = LoadCulledDFObjectData(ObjectIndex);
+		FDFObjectData DFObjectData = LoadDFObjectData(ObjectIndex);
 
 		float3 VolumeShadingPosition = mul(float4(WorldShadingPosition, 1), DFObjectData.WorldToVolume).xyz;
 		float BoxDistance = ComputeDistanceFromBoxToPoint(-DFObjectData.VolumePositionExtent, DFObjectData.VolumePositionExtent, VolumeShadingPosition) * DFObjectData.VolumeScale;
diff --git a/Engine/Shaders/Private/DistanceFieldShadowing.usf b/Engine/Shaders/Private/DistanceFieldShadowing.usf
--- a/Engine/Shaders/Private/DistanceFieldShadowing.usf
+++ b/Engine/Shaders/Private/DistanceFieldShadowing.usf
@@ -25,63 +25,6 @@
 uint ObjectBoundingGeometryIndexCount;
 float ObjectExpandScale;
 
-void CopyCulledObjectData(uint DestIndex, uint SourceIndex)
-{
-#if (DISTANCEFIELD_PRIMITIVE_TYPE == DFPT_SignedDistanceField)
-	RWCulledObjectBounds[DestIndex] = SceneObjectBounds[SourceIndex * OBJECT_BOUNDS_STRIDE + 0];
-#else
-	float4 BoxOriginAndSphereRadius = SceneObjectBounds[SourceIndex * OBJECT_BOUNDS_STRIDE + 0];
-	float4 BoxExtent = SceneObjectBounds[SourceIndex * OBJECT_BOUNDS_STRIDE + 1];
-	BoxExtent.z += MAX_TRACE_SPHERE_RADIUS;
-
-	RWCulledObjectBounds[DestIndex * OBJECT_BOUNDS_STRIDE] = BoxOriginAndSphereRadius;
-	RWCulledObjectBounds[DestIndex * OBJECT_BOUNDS_STRIDE + 1] = BoxExtent;
-#endif
-
-	UNROLL
-	for (uint VectorIndex = 0; VectorIndex < CULLED_OBJECT_DATA_STRIDE; VectorIndex++)
-	{
-		// Note: only copying the first CULLED_OBJECT_DATA_STRIDE of the original object data
-		RWCulledObjectData[DestIndex * CULLED_OBJECT_DATA_STRIDE + VectorIndex] = SceneObjectData[SourceIndex * OBJECT_DATA_STRIDE + VectorIndex];
-	}
-
-	float3 MinViewSpacePosition = float3(2000000, 2000000, 2000000);
-	float3 MaxViewSpacePosition = float3(-2000000, -2000000, -2000000);
-	float3 ViewSpaceBoundsVertices[8];
-
-	for (uint i = 0; i < 8; i++)
-	{
-		float3 UnitBoxVertex;
-		UnitBoxVertex.x = i & 0x1 ? 1.0f : -1.0f;
-		UnitBoxVertex.y = i & 0x2 ? 1.0f : -1.0f;
-		UnitBoxVertex.z = i & 0x4 ? 1.0f : -1.0f;
-
-#if (DISTANCEFIELD_PRIMITIVE_TYPE == DFPT_SignedDistanceField)
-		FDFObjectData DFObjectData = LoadDFObjectData(SourceIndex);
-		float3 WorldBoundsPosition = mul(float4(UnitBoxVertex * DFObjectData.VolumePositionExtent, 1.0f), DFObjectData.VolumeToWorld).xyz;
-#else
-		float3 WorldBoundsPosition = UnitBoxVertex * BoxExtent.xyz + BoxOriginAndSphereRadius.xyz;
-#endif
-
-		float3 ViewSpacePosition = mul(float4(WorldBoundsPosition, 1.0f), WorldToShadow).xyz;
-		MinViewSpacePosition = min(MinViewSpacePosition, ViewSpacePosition);
-		MaxViewSpacePosition = max(MaxViewSpacePosition, ViewSpacePosition);
-		ViewSpaceBoundsVertices[i] = ViewSpacePosition;
-	}
-
-	float3 ObjectXAxis = (ViewSpaceBoundsVertices[1] - ViewSpaceBoundsVertices[0]) / 2.0f;
-	float3 ObjectYAxis = (ViewSpaceBoundsVertices[2] - ViewSpaceBoundsVertices[0]) / 2.0f;
-	float3 ObjectZAxis = (ViewSpaceBoundsVertices[4] - ViewSpaceBoundsVertices[0]) / 2.0f;
-
-	MinViewSpacePosition.xy -= ObjectExpandScale * MAX_TRACE_SPHERE_RADIUS;
-	MaxViewSpacePosition.xy += ObjectExpandScale * MAX_TRACE_SPHERE_RADIUS;
-	RWCulledObjectBoxBounds[DestIndex * CULLED_OBJECT_BOX_BOUNDS_STRIDE + 0] = float4(MinViewSpacePosition, 0);
-	RWCulledObjectBoxBounds[DestIndex * CULLED_OBJECT_BOX_BOUNDS_STRIDE + 1] = float4(MaxViewSpacePosition, 0);
-	RWCulledObjectBoxBounds[DestIndex * CULLED_OBJECT_BOX_BOUNDS_STRIDE + 2] = float4(ObjectXAxis / max(dot(ObjectXAxis, ObjectXAxis), .0001f), 0);
-	RWCulledObjectBoxBounds[DestIndex * CULLED_OBJECT_BOX_BOUNDS_STRIDE + 3] = float4(ObjectYAxis / max(dot(ObjectYAxis, ObjectYAxis), .0001f), 0);
-	RWCulledObjectBoxBounds[DestIndex * CULLED_OBJECT_BOX_BOUNDS_STRIDE + 4] = float4(ObjectZAxis / max(dot(ObjectZAxis, ObjectZAxis), .0001f), 0);
-}
-
 float4 ShadowConvexHull[12];
 float4 ShadowBoundingSphere;
 uint NumShadowHullPlanes;
@@ -122,6 +65,15 @@ bool ShadowConvexHullIntersectBox(float3 BoxOrigin, float3 BoxExtent)
 	return true;
 }
 
+uint GetNumSceneObjects()
+{
+#if (DISTANCEFIELD_PRIMITIVE_TYPE == DFPT_SignedDistanceField)
+	return NumSceneObjects;
+#else
+	return NumSceneHeightfieldObjects;
+#endif
+}
+
 [numthreads(UPDATEOBJECTS_THREADGROUP_SIZE, 1, 1)]
 void CullObjectsForShadowCS(
 	uint3 GroupId : SV_GroupID,
@@ -142,9 +94,7 @@ void CullObjectsForShadowCS(
 
 	GroupMemoryBarrierWithGroupSync();
 
-	uint SourceIndex = ObjectIndex;
-
-	if (ObjectIndex < NumSceneObjects)
+	if (ObjectIndex < GetNumSceneObjects())
 	{
 #if (DISTANCEFIELD_PRIMITIVE_TYPE == DFPT_SignedDistanceField)
 		FDFObjectBounds DFObjectBounds = LoadDFObjectBounds(ObjectIndex);
@@ -166,18 +116,16 @@ void CullObjectsForShadowCS(
             {
                 uint DestIndex;
                 InterlockedAdd(RWObjectIndirectArguments[1], 1U, DestIndex);
-                CopyCulledObjectData(DestIndex, SourceIndex);
+                RWCulledObjectIndices[DestIndex] = ObjectIndex;
             }
         }
 #else
-		float3 BoxOrigin = SceneObjectBounds[SourceIndex * OBJECT_BOUNDS_STRIDE + 0].xyz;
-		float3 BoxExtent = SceneObjectBounds[SourceIndex * OBJECT_BOUNDS_STRIDE + 1].xyz;
-
-		if (ShadowConvexHullIntersectBox(BoxOrigin, BoxExtent))
+		FHeightfieldObjectBounds Bounds = LoadHeightfieldObjectBounds(ObjectIndex);
+		if (ShadowConvexHullIntersectBox(Bounds.BoxOrigin, Bounds.BoxExtent))
 		{
 			uint DestIndex;
 			InterlockedAdd(RWObjectIndirectArguments[1], 1U, DestIndex);
-			CopyCulledObjectData(DestIndex, SourceIndex);
+			RWCulledObjectIndices[DestIndex] = ObjectIndex;
 		}
 #endif
 	}
@@ -188,15 +136,12 @@ void CullObjectsForShadowCS(
 	{
 		// IndexCount, NumInstances, StartIndex, BaseVertexIndex, FirstInstance
 		RWObjectIndirectArguments[0] = ObjectBoundingGeometryIndexCount;
-		RWObjectIndirectArguments[1] = NumSceneObjects;
+		RWObjectIndirectArguments[1] = GetNumSceneObjects();
 	}
 
-	uint SourceIndex = ObjectIndex;
-	uint DestIndex = ObjectIndex;
-
-	if (ObjectIndex < NumSceneObjects)
+	if (ObjectIndex < GetNumSceneObjects())
 	{
-		CopyCulledObjectData(DestIndex, SourceIndex);
+		RWCulledObjectIndices[ObjectIndex] = ObjectIndex;
 	}
 #endif
 }
@@ -265,13 +210,15 @@ float MinExpandRadius;
 /** Used when culling objects into screenspace tile lists */
 void ShadowObjectCullVS(
 	float4 InPosition : ATTRIBUTE0,
-	uint ObjectIndex : SV_InstanceID,
+	uint InstanceIndex : SV_InstanceID,
 	out FShadowObjectCullVertexOutput Output,
 	out float4 OutPosition : SV_POSITION
 )
 {
+	uint ObjectIndex = CulledObjectIndices[InstanceIndex];
+
 #if (DISTANCEFIELD_PRIMITIVE_TYPE == DFPT_SignedDistanceField)
-	FDFObjectData DFObjectData = LoadCulledDFObjectData(ObjectIndex);
+	FDFObjectData DFObjectData = LoadDFObjectData(ObjectIndex);
 	float4x4 ObjectOBBToWorld = DFObjectData.VolumeToWorld;
 
 	float ObjectScaleX = length(ObjectOBBToWorld[0].xyz);
@@ -284,9 +231,10 @@ void ShadowObjectCullVS(
 
 	float3 WorldPosition = mul(float4(InPosition.xyz, 1.0f), ObjectOBBToWorld).xyz;
 #else
-	float3 BoxOrigin;
-	float3 BoxExtent;
-	LoadHeightFieldObjectBoxBounds(ObjectIndex, BoxOrigin, BoxExtent);
+
+	FHeightfieldObjectBounds Bounds = LoadHeightfieldObjectBounds(ObjectIndex);
+	float3 BoxOrigin = Bounds.BoxOrigin;
+	float3 BoxExtent = Bounds.BoxExtent;
 
 	BoxExtent += MinExpandRadius;
 
@@ -323,6 +271,58 @@ void HandleShadowTileObjectIntersection(uint TileIndex, uint ObjectIndex)
 	#endif
 }
 
+struct FObjectViewSpaceBox
+{
+	float3 Min;
+	float3 Max;
+	float3 XAxis;
+	float3 YAxis;
+	float3 ZAxis;
+};
+
+FObjectViewSpaceBox GetObjectViewSpaceBox(uint ObjectIndex)
+{
+	float3 MinViewSpacePosition = float3(2000000, 2000000, 2000000);
+	float3 MaxViewSpacePosition = float3(-2000000, -2000000, -2000000);
+	float3 ViewSpaceBoundsVertices[8];
+
+	for (uint i = 0; i < 8; i++)
+	{
+		float3 UnitBoxVertex;
+		UnitBoxVertex.x = i & 0x1 ? 1.0f : -1.0f;
+		UnitBoxVertex.y = i & 0x2 ? 1.0f : -1.0f;
+		UnitBoxVertex.z = i & 0x4 ? 1.0f : -1.0f;
+
+#if (DISTANCEFIELD_PRIMITIVE_TYPE == DFPT_SignedDistanceField)
+		FDFObjectData DFObjectData = LoadDFObjectData(ObjectIndex);
+		float3 WorldBoundsPosition = mul(float4(UnitBoxVertex * DFObjectData.VolumePositionExtent, 1.0f), DFObjectData.VolumeToWorld).xyz;
+#else
+		FHeightfieldObjectBounds HeightfieldObjectBounds = LoadHeightfieldObjectBounds(ObjectIndex);
+		float3 WorldBoundsPosition = UnitBoxVertex * HeightfieldObjectBounds.BoxExtent + HeightfieldObjectBounds.BoxOrigin;
+#endif
+
+		float3 ViewSpacePosition = mul(float4(WorldBoundsPosition, 1.0f), WorldToShadow).xyz;
+		MinViewSpacePosition = min(MinViewSpacePosition, ViewSpacePosition);
+		MaxViewSpacePosition = max(MaxViewSpacePosition, ViewSpacePosition);
+		ViewSpaceBoundsVertices[i] = ViewSpacePosition;
+	}
+
+	float3 ObjectXAxis = (ViewSpaceBoundsVertices[1] - ViewSpaceBoundsVertices[0]) / 2.0f;
+	float3 ObjectYAxis = (ViewSpaceBoundsVertices[2] - ViewSpaceBoundsVertices[0]) / 2.0f;
+	float3 ObjectZAxis = (ViewSpaceBoundsVertices[4] - ViewSpaceBoundsVertices[0]) / 2.0f;
+
+	MinViewSpacePosition.xy -= ObjectExpandScale * MAX_TRACE_SPHERE_RADIUS;
+	MaxViewSpacePosition.xy += ObjectExpandScale * MAX_TRACE_SPHERE_RADIUS;
+
+	FObjectViewSpaceBox ViewBox;
+	ViewBox.Min = MinViewSpacePosition;
+	ViewBox.Max = MaxViewSpacePosition;
+	ViewBox.XAxis = ObjectXAxis / max(dot(ObjectXAxis, ObjectXAxis), .0001f);
+	ViewBox.YAxis = ObjectYAxis / max(dot(ObjectYAxis, ObjectYAxis), .0001f);
+	ViewBox.ZAxis = ObjectZAxis / max(dot(ObjectZAxis, ObjectZAxis), .0001f);
+	return ViewBox;
+}
+
 /** Intersects a single object with the tile and adds to the intersection list if needed. */
 void ShadowObjectCullPS(
 	FShadowObjectCullVertexOutput Input, 
@@ -347,16 +347,14 @@ void ShadowObjectCullPS(
 	ShadowTileMin.z = 0;
 	ShadowTileMax.z = 1000;
 
-	float3 ObjectViewSpaceMin;
-	float3 ObjectViewSpaceMax;
-	LoadObjectViewSpaceBox(Input.ObjectIndex, ObjectViewSpaceMin, ObjectViewSpaceMax);
+	FObjectViewSpaceBox ObjectViewSpaceBox = GetObjectViewSpaceBox(Input.ObjectIndex);
 
 	BRANCH
 	// Separating axis test on the AABB
 	// Note: don't clip by near plane, objects closer to the light can still cast into the frustum
-	if (all(ObjectViewSpaceMax > ShadowTileMin) && all(ObjectViewSpaceMin.xy < ShadowTileMax.xy))
+	if (all(ObjectViewSpaceBox.Max > ShadowTileMin) && all(ObjectViewSpaceBox.Min.xy < ShadowTileMax.xy))
 	{
-		float3 ObjectCenter = .5f * (ObjectViewSpaceMin + ObjectViewSpaceMax);
+		float3 ObjectCenter = .5f * (ObjectViewSpaceBox.Min + ObjectViewSpaceBox.Max);
 		float3 MinProjections = 500000;
 		float3 MaxProjections = -500000;
 
@@ -371,10 +369,9 @@ void ShadowObjectCullPS(
 			Corners[6] = float3(ShadowTileMin.x, ShadowTileMax.y, ShadowTileMax.z);
 			Corners[7] = float3(ShadowTileMax.x, ShadowTileMax.y, ShadowTileMax.z);
 
-			float3 ObjectAxisX;
-			float3 ObjectAxisY; 
-			float3 ObjectAxisZ;
-			LoadObjectAxes(Input.ObjectIndex, ObjectAxisX, ObjectAxisY, ObjectAxisZ);
+			float3 ObjectAxisX = ObjectViewSpaceBox.XAxis;
+			float3 ObjectAxisY = ObjectViewSpaceBox.YAxis; 
+			float3 ObjectAxisZ = ObjectViewSpaceBox.ZAxis;
 
 			UNROLL
 			for (int i = 0; i < 8; i++)
@@ -547,12 +544,13 @@ void CullObjectsToTileWithGather(
 		// Compute per-tile lists of affecting objects through bounds culling
 		// Each thread now operates on a sample instead of a pixel
 		LOOP
-		for (uint ObjectIndex = ThreadIndex; ObjectIndex < NumCulledObjects; ObjectIndex += THREADGROUP_TOTALSIZE)
+		for (uint IndexInCulledList = ThreadIndex; IndexInCulledList < NumCulledObjects; IndexInCulledList += THREADGROUP_TOTALSIZE)
 		{
-			float4 SphereCenterAndRadius = LoadObjectPositionAndRadius(ObjectIndex);
+			const uint ObjectIndex = CulledObjectIndices[IndexInCulledList];
+			FDFObjectBounds Bounds = LoadDFObjectBounds(ObjectIndex);
 
 			BRANCH
-			if (RaySegmentHitSphere(WorldGroup0Center, LightDirection0, RayLength0, SphereCenterAndRadius.xyz, SphereCenterAndRadius.w + Group0BoundingRadius))
+			if (RaySegmentHitSphere(WorldGroup0Center, LightDirection0, RayLength0, Bounds.Center, Bounds.SphereRadius + Group0BoundingRadius))
 			{
 				uint ListIndex;
 				InterlockedAdd(TileNumObjects0, 1U, ListIndex);
@@ -562,7 +560,7 @@ void CullObjectsToTileWithGather(
 			}
 
 			BRANCH
-			if (RaySegmentHitSphere(WorldGroup1Center, LightDirection1, RayLength1, SphereCenterAndRadius.xyz, SphereCenterAndRadius.w + Group1BoundingRadius))
+			if (RaySegmentHitSphere(WorldGroup1Center, LightDirection1, RayLength1, Bounds.Center, Bounds.SphereRadius + Group1BoundingRadius))
 			{
 				uint ListIndex;
 				InterlockedAdd(TileNumObjects1, 1U, ListIndex);
diff --git a/Engine/Shaders/Private/DistanceFieldShadowingShared.ush b/Engine/Shaders/Private/DistanceFieldShadowingShared.ush
--- a/Engine/Shaders/Private/DistanceFieldShadowingShared.ush
+++ b/Engine/Shaders/Private/DistanceFieldShadowingShared.ush
@@ -5,7 +5,7 @@
 =============================================================================*/
 
 #ifndef COMPACT_CULLED_SHADOW_OBJECTS
-#define COMPACT_CULLED_SHADOW_OBJECTS 1
+	#define COMPACT_CULLED_SHADOW_OBJECTS 1
 #endif
 
 Buffer<uint> ShadowTileNumCulledObjects;
@@ -91,7 +91,7 @@ float ShadowRayTraceThroughCulledObjects(
 		}
 
 		{
-			FDFObjectData DFObjectData = LoadCulledDFObjectData(ObjectIndex);
+			FDFObjectData DFObjectData = LoadDFObjectData(ObjectIndex);
 
 			float3 VolumeRayStart = mul(float4(WorldRayStart, 1), DFObjectData.WorldToVolume).xyz;
 			float3 VolumeRayEnd = mul(float4(WorldRayEnd, 1), DFObjectData.WorldToVolume).xyz;
@@ -327,10 +327,10 @@ float ShadowRayTraceThroughCulledHeightFieldObjects(
 			ObjectIndex = ListObjectIndex;
 		}
 
-		float4x4 WorldToLocal = LoadHeightFieldObjectWorldToLocal(ObjectIndex);
-		float4 SizeScale = LoadHeightFieldObjectSizeScale(ObjectIndex);
-		float2 HeightFieldSize = SizeScale.xy;
-		float WorldToLocalScale = SizeScale.z;
+		FHeightfieldObjectData HeightfieldObject = LoadHeightfieldObjectData(ObjectIndex);
+		float4x4 WorldToLocal = HeightfieldObject.WorldToLocal;
+		float2 HeightFieldSize = HeightfieldObject.SizeScale.xy;
+		float WorldToLocalScale = HeightfieldObject.SizeScale.z;
 
 		float3 LocalRayStart = mul(float4(WorldRayStart, 1.0), WorldToLocal).xyz;
 		float3 LocalRayEnd = mul(float4(WorldRayEnd, 1.0), WorldToLocal).xyz;
@@ -343,9 +343,9 @@ float ShadowRayTraceThroughCulledHeightFieldObjects(
 
 		if (bValidIntersection)
 		{
-			float4 AtlasUVScaleBias = LoadHeightFieldObjectAtlasUVScaleBias(ObjectIndex);
+			float4 AtlasUVScaleBias = HeightfieldObject.AtlasUVScaleBias;
 #if DF_SHADOW_QUALITY == 2
-			float4 VisUVScaleBias = LoadHeightFieldObjectVisibilityAtlasUVScaleBias(ObjectIndex);
+			float4 VisUVScaleBias = HeightfieldObject.VisibilityAtlasUVScaleBias;
 #else
 			float4 VisUVScaleBias = float4(0, 0, 0, 0);
 #endif
diff --git a/Engine/Shaders/Private/DistanceFieldVisualization.usf b/Engine/Shaders/Private/DistanceFieldVisualization.usf
--- a/Engine/Shaders/Private/DistanceFieldVisualization.usf
+++ b/Engine/Shaders/Private/DistanceFieldVisualization.usf
@@ -31,7 +31,7 @@ void RayTraceThroughTileCulledDistanceFields(float3 WorldRayStart, float3 WorldR
 	for (uint ListObjectIndex = 0; ListObjectIndex < min(NumIntersectingObjects, (uint)MAX_INTERSECTING_OBJECTS); ListObjectIndex++)
 	{
 		uint ObjectIndex = IntersectingObjectIndices[ListObjectIndex];
-		FDFObjectData DFObjectData = LoadCulledDFObjectData(ObjectIndex);
+		FDFObjectData DFObjectData = LoadDFObjectData(ObjectIndex);
 
 		float3 VolumeRayStart = mul(float4(WorldRayStart, 1), DFObjectData.WorldToVolume).xyz;
 		float3 WorldRayEnd = WorldRayStart + WorldRayDirection * MaxRayTime;
@@ -185,9 +185,11 @@ void VisualizeMeshDistanceFieldCS(
 	uint NumCulledObjects = GetCulledNumObjects();
 
 	LOOP
-	for (uint ObjectIndex = ThreadIndex; ObjectIndex < NumCulledObjects; ObjectIndex += THREADGROUP_TOTALSIZE)
+	for (uint IndexInCulledList = ThreadIndex; IndexInCulledList < NumCulledObjects; IndexInCulledList += THREADGROUP_TOTALSIZE)
 	{
-		float4 SphereCenterAndRadius = LoadObjectPositionAndRadius(ObjectIndex);
+		const uint ObjectIndex = CulledObjectIndices[IndexInCulledList];
+		FDFObjectBounds ObjectBounds = LoadDFObjectBounds(ObjectIndex);
+		float4 SphereCenterAndRadius = float4(ObjectBounds.Center, ObjectBounds.SphereRadius);
 
 		BRANCH
 		if (SphereIntersectCone(SphereCenterAndRadius, TileConeVertex, TileConeAxis, TileConeAngleCos, TileConeAngleSin))
@@ -217,7 +219,7 @@ void VisualizeMeshDistanceFieldCS(
 
 	if (HitCulledObjectIndex != 0xFFFFFFFF)
 	{
-		FDFObjectData DFObjectData = LoadCulledDFObjectData(HitCulledObjectIndex);
+		FDFObjectData DFObjectData = LoadDFObjectData(HitCulledObjectIndex);
 
 		float3 VolumeHitPosition = mul(float4(WorldRayEnd, 1), DFObjectData.WorldToVolume).xyz;
 		//VolumeHitPosition = clamp(VolumeHitPosition, -DFObjectData.VolumePositionExtent, DFObjectData.VolumePositionExtent);
diff --git a/Engine/Source/Runtime/Renderer/Private/DistanceFieldAmbientOcclusion.cpp b/Engine/Source/Runtime/Renderer/Private/DistanceFieldAmbientOcclusion.cpp
--- a/Engine/Source/Runtime/Renderer/Private/DistanceFieldAmbientOcclusion.cpp
+++ b/Engine/Source/Runtime/Renderer/Private/DistanceFieldAmbientOcclusion.cpp
@@ -830,9 +830,7 @@ void FDeferredShadingSceneRenderer::RenderDistanceFieldLighting(
 	{
 		AllocateDistanceFieldCulledObjectBuffers(
 			GraphBuilder,
-			false,
 			MaxSceneObjects,
-			DFPT_SignedDistanceField,
 			ObjectIndirectArguments,
 			CulledObjectBufferParameters);
 
diff --git a/Engine/Source/Runtime/Renderer/Private/DistanceFieldLightingShared.h b/Engine/Source/Runtime/Renderer/Private/DistanceFieldLightingShared.h
--- a/Engine/Source/Runtime/Renderer/Private/DistanceFieldLightingShared.h
+++ b/Engine/Source/Runtime/Renderer/Private/DistanceFieldLightingShared.h
@@ -46,13 +46,9 @@ enum EDistanceFieldPrimitiveType
 // Must match equivalent shader defines
 static const int32 GDistanceFieldObjectDataStride = 9;
 static const int32 GDistanceFieldObjectBoundsStride = 2;
-static const int32 GDistanceFieldCulledObjectDataStride = GDistanceFieldObjectDataStride;
-static const int32 GDistanceFieldCulledObjectBoxBoundsStride = 5;
 
 static const int32 GHeightFieldObjectDataStride = 6;
 static const int32 GHeightFieldObjectBoundsStride = 2;
-static const int32 GHeightFieldCulledObjectDataStride = GHeightFieldObjectDataStride;
-static const int32 GHeightFieldCulledObjectBoxBoundsStride = 5;
 
 template <EDistanceFieldPrimitiveType PrimitiveType>
 class TDistanceFieldObjectBuffers
@@ -90,6 +86,9 @@ BEGIN_SHADER_PARAMETER_STRUCT(FDistanceFieldObjectBufferParameters, )
 	SHADER_PARAMETER_SRV(StructuredBuffer<float4>, SceneObjectBounds)
 	SHADER_PARAMETER_SRV(StructuredBuffer<float4>, SceneObjectData)
 	SHADER_PARAMETER(uint32, NumSceneObjects)
+	SHADER_PARAMETER_SRV(StructuredBuffer<float4>, SceneHeightfieldObjectBounds)
+	SHADER_PARAMETER_SRV(StructuredBuffer<float4>, SceneHeightfieldObjectData)
+	SHADER_PARAMETER(uint32, NumSceneHeightfieldObjects)
 END_SHADER_PARAMETER_STRUCT()
 
 BEGIN_SHADER_PARAMETER_STRUCT(FDistanceFieldAtlasParameters, )
@@ -133,20 +132,14 @@ namespace DistanceField
 
 BEGIN_SHADER_PARAMETER_STRUCT(FDistanceFieldCulledObjectBufferParameters, )
 	SHADER_PARAMETER_RDG_BUFFER_UAV(RWBuffer<uint>, RWObjectIndirectArguments)
-	SHADER_PARAMETER_RDG_BUFFER_UAV(RWBuffer<uint>, RWCulledObjectBounds)
-	SHADER_PARAMETER_RDG_BUFFER_UAV(RWBuffer<uint>, RWCulledObjectData)
-	SHADER_PARAMETER_RDG_BUFFER_UAV(RWBuffer<uint>, RWCulledObjectBoxBounds)
+	SHADER_PARAMETER_RDG_BUFFER_UAV(RWBuffer<uint>, RWCulledObjectIndices)
 	SHADER_PARAMETER_RDG_BUFFER_SRV(Buffer<uint>, ObjectIndirectArguments)
-	SHADER_PARAMETER_RDG_BUFFER_SRV(Buffer<uint>, CulledObjectBounds)
-	SHADER_PARAMETER_RDG_BUFFER_SRV(Buffer<uint>, CulledObjectData)
-	SHADER_PARAMETER_RDG_BUFFER_SRV(Buffer<uint>, CulledObjectBoxBounds)
+	SHADER_PARAMETER_RDG_BUFFER_SRV(Buffer<uint>, CulledObjectIndices)
 END_SHADER_PARAMETER_STRUCT()
 
 extern void AllocateDistanceFieldCulledObjectBuffers(
 	FRDGBuilder& GraphBuilder,
-	bool bWantBoxBounds,
 	uint32 MaxObjects,
-	EDistanceFieldPrimitiveType PrimitiveType,
 	FRDGBufferRef& OutObjectIndirectArguments,
 	FDistanceFieldCulledObjectBufferParameters& OutParameters);
 
diff --git a/Engine/Source/Runtime/Renderer/Private/DistanceFieldObjectCulling.cpp b/Engine/Source/Runtime/Renderer/Private/DistanceFieldObjectCulling.cpp
--- a/Engine/Source/Runtime/Renderer/Private/DistanceFieldObjectCulling.cpp
+++ b/Engine/Source/Runtime/Renderer/Private/DistanceFieldObjectCulling.cpp
@@ -150,6 +150,7 @@ public:
 
 	BEGIN_SHADER_PARAMETER_STRUCT(FParameters, )
 		SHADER_PARAMETER_STRUCT_REF(FViewUniformShaderParameters, View)
+		SHADER_PARAMETER_STRUCT_INCLUDE(FDistanceFieldObjectBufferParameters, DistanceFieldObjectBuffers)
 		SHADER_PARAMETER_STRUCT_INCLUDE(FDistanceFieldCulledObjectBufferParameters, DistanceFieldCulledObjectBuffers)
 		SHADER_PARAMETER_STRUCT_INCLUDE(FDistanceFieldAtlasParameters, DistanceFieldAtlas)
 		SHADER_PARAMETER_STRUCT_INCLUDE(FAOParameters, AOParameters)
@@ -173,6 +174,7 @@ public:
 	BEGIN_SHADER_PARAMETER_STRUCT(FParameters, )
 		SHADER_PARAMETER_STRUCT_REF(FViewUniformShaderParameters, View)
 		SHADER_PARAMETER_STRUCT_INCLUDE(FTileIntersectionParameters, TileIntersectionParameters)
+		SHADER_PARAMETER_STRUCT_INCLUDE(FDistanceFieldObjectBufferParameters, DistanceFieldObjectBuffers)
 		SHADER_PARAMETER_STRUCT_INCLUDE(FDistanceFieldCulledObjectBufferParameters, DistanceFieldCulledObjectBuffers)
 		SHADER_PARAMETER_STRUCT_INCLUDE(FDistanceFieldAtlasParameters, DistanceFieldAtlas)
 		SHADER_PARAMETER_STRUCT_INCLUDE(FAOParameters, AOParameters)
@@ -252,6 +254,8 @@ void ScatterTilesToObjects(
 	const FTileIntersectionParameters& TileIntersectionParameters,
 	TRDGUniformBufferRef<FSceneTextureUniformParameters> SceneTexturesUniformBuffer)
 {
+	FDistanceFieldObjectBufferParameters DistanceFieldObjectBuffers = DistanceField::SetupObjectBufferParameters(DistanceFieldSceneData);
+
 	FObjectCullPS::FPermutationDomain PermutationVector;
 	PermutationVector.Set<FObjectCullPS::FCountingPass>(bCountingPass);
 
@@ -260,6 +264,7 @@ void ScatterTilesToObjects(
 
 	auto* PassParameters = GraphBuilder.AllocParameters<FObjectCullParameters>();
 	PassParameters->VS.View = GetShaderBinding(View.ViewUniformBuffer);
+	PassParameters->VS.DistanceFieldObjectBuffers = DistanceFieldObjectBuffers;
 	PassParameters->VS.DistanceFieldCulledObjectBuffers = CulledObjectBufferParameters;
 	PassParameters->VS.DistanceFieldAtlas = DistanceField::SetupAtlasParameters(DistanceFieldSceneData);
 	PassParameters->VS.AOParameters = DistanceField::SetupAOShaderParameters(Parameters);
@@ -274,6 +279,7 @@ void ScatterTilesToObjects(
 
 	PassParameters->PS.View = GetShaderBinding(View.ViewUniformBuffer);
 	PassParameters->PS.TileIntersectionParameters = TileIntersectionParameters;
+	PassParameters->PS.DistanceFieldObjectBuffers = DistanceFieldObjectBuffers;
 	PassParameters->PS.DistanceFieldCulledObjectBuffers = CulledObjectBufferParameters;
 	PassParameters->PS.DistanceFieldAtlas = DistanceField::SetupAtlasParameters(DistanceFieldSceneData);
 	PassParameters->PS.AOParameters = DistanceField::SetupAOShaderParameters(Parameters);
diff --git a/Engine/Source/Runtime/Renderer/Private/DistanceFieldObjectManagement.cpp b/Engine/Source/Runtime/Renderer/Private/DistanceFieldObjectManagement.cpp
--- a/Engine/Source/Runtime/Renderer/Private/DistanceFieldObjectManagement.cpp
+++ b/Engine/Source/Runtime/Renderer/Private/DistanceFieldObjectManagement.cpp
@@ -102,6 +102,7 @@ FDistanceFieldObjectBufferParameters DistanceField::SetupObjectBufferParameters(
 	FDistanceFieldObjectBufferParameters ObjectBufferParameters;
 
 	ObjectBufferParameters.NumSceneObjects = DistanceFieldSceneData.NumObjectsInBuffer;
+	ObjectBufferParameters.NumSceneHeightfieldObjects = DistanceFieldSceneData.NumHeightFieldObjectsInBuffer;
 
 	if (DistanceFieldSceneData.NumObjectsInBuffer > 0)
 	{
@@ -115,6 +116,18 @@ FDistanceFieldObjectBufferParameters DistanceField::SetupObjectBufferParameters(
 		ObjectBufferParameters.SceneObjectData = nullptr;
 	}
 
+	if (DistanceFieldSceneData.NumHeightFieldObjectsInBuffer > 0)
+	{
+		check(DistanceFieldSceneData.GetHeightFieldObjectBuffers());
+		ObjectBufferParameters.SceneHeightfieldObjectBounds = DistanceFieldSceneData.GetHeightFieldObjectBuffers()->Bounds.SRV;
+		ObjectBufferParameters.SceneHeightfieldObjectData = DistanceFieldSceneData.GetHeightFieldObjectBuffers()->Data.SRV;
+	}
+	else
+	{
+		ObjectBufferParameters.SceneHeightfieldObjectBounds = nullptr;
+		ObjectBufferParameters.SceneHeightfieldObjectData = nullptr;
+	}
+
 	return ObjectBufferParameters;
 }
 
diff --git a/Engine/Source/Runtime/Renderer/Private/DistanceFieldScreenGridLighting.cpp b/Engine/Source/Runtime/Renderer/Private/DistanceFieldScreenGridLighting.cpp
--- a/Engine/Source/Runtime/Renderer/Private/DistanceFieldScreenGridLighting.cpp
+++ b/Engine/Source/Runtime/Renderer/Private/DistanceFieldScreenGridLighting.cpp
@@ -135,6 +135,7 @@ public:
 
 	BEGIN_SHADER_PARAMETER_STRUCT(FParameters, )
 		SHADER_PARAMETER_STRUCT_REF(FViewUniformShaderParameters, View)
+		SHADER_PARAMETER_STRUCT_INCLUDE(FDistanceFieldObjectBufferParameters, DistanceFieldObjectBuffers)
 		SHADER_PARAMETER_STRUCT_INCLUDE(FDistanceFieldCulledObjectBufferParameters, DistanceFieldCulledObjectBuffers)
 		SHADER_PARAMETER_STRUCT_INCLUDE(FDistanceFieldAtlasParameters, DistanceFieldAtlas)
 		SHADER_PARAMETER_STRUCT_INCLUDE(FTileIntersectionParameters, TileIntersectionParameters)
@@ -350,6 +351,7 @@ void FDeferredShadingSceneRenderer::RenderDistanceFieldAOScreenGrid(
 
 		auto* PassParameters = GraphBuilder.AllocParameters<FConeTraceScreenGridObjectOcclusionCS::FParameters>();
 		PassParameters->View = View.ViewUniformBuffer;
+		PassParameters->DistanceFieldObjectBuffers = DistanceField::SetupObjectBufferParameters(DistanceFieldSceneData);
 		PassParameters->DistanceFieldCulledObjectBuffers = CulledObjectBufferParameters;
 		PassParameters->DistanceFieldAtlas = DistanceField::SetupAtlasParameters(DistanceFieldSceneData);
 		PassParameters->TileIntersectionParameters = TileIntersectionParameters;
diff --git a/Engine/Source/Runtime/Renderer/Private/DistanceFieldShadowing.cpp b/Engine/Source/Runtime/Renderer/Private/DistanceFieldShadowing.cpp
--- a/Engine/Source/Runtime/Renderer/Private/DistanceFieldShadowing.cpp
+++ b/Engine/Source/Runtime/Renderer/Private/DistanceFieldShadowing.cpp
@@ -162,7 +162,6 @@ class FCullObjectsForShadowCS : public FGlobalShader
 		SHADER_PARAMETER_STRUCT_INCLUDE(FDistanceFieldCulledObjectBufferParameters, CulledObjectBufferParameters)
 		SHADER_PARAMETER(uint32, ObjectBoundingGeometryIndexCount)
 		SHADER_PARAMETER(FMatrix44f, WorldToShadow)
-		SHADER_PARAMETER(float, ObjectExpandScale)
 		SHADER_PARAMETER(uint32, NumShadowHullPlanes)
 		SHADER_PARAMETER(uint32, bDrawNaniteMeshes)
 		SHADER_PARAMETER(FVector4f, ShadowBoundingSphere)
@@ -187,9 +186,8 @@ class FCullObjectsForShadowCS : public FGlobalShader
 
 IMPLEMENT_GLOBAL_SHADER(FCullObjectsForShadowCS, "/Engine/Private/DistanceFieldShadowing.usf", "CullObjectsForShadowCS", SF_Compute);
 
-/**  */
 class FShadowObjectCullVS : public FGlobalShader
-	{
+{
 	DECLARE_GLOBAL_SHADER(FShadowObjectCullVS);
 	SHADER_USE_PARAMETER_STRUCT(FShadowObjectCullVS, FGlobalShader);
 
@@ -212,7 +210,7 @@ class FShadowObjectCullVS : public FGlobalShader
 IMPLEMENT_GLOBAL_SHADER(FShadowObjectCullVS, "/Engine/Private/DistanceFieldShadowing.usf", "ShadowObjectCullVS", SF_Vertex);
 
 class FShadowObjectCullPS : public FGlobalShader
-	{
+{
 	DECLARE_GLOBAL_SHADER(FShadowObjectCullPS);
 	SHADER_USE_PARAMETER_STRUCT(FShadowObjectCullPS, FGlobalShader);
 
@@ -220,6 +218,8 @@ class FShadowObjectCullPS : public FGlobalShader
 		SHADER_PARAMETER_STRUCT_INCLUDE(FDistanceFieldObjectBufferParameters, ObjectBufferParameters)
 		SHADER_PARAMETER_STRUCT_INCLUDE(FDistanceFieldCulledObjectBufferParameters, CulledObjectBufferParameters)
 		SHADER_PARAMETER_STRUCT_INCLUDE(FLightTileIntersectionParameters, LightTileIntersectionParameters)
+		SHADER_PARAMETER(FMatrix44f, WorldToShadow)
+		SHADER_PARAMETER(float, ObjectExpandScale)
 	END_SHADER_PARAMETER_STRUCT()
 
 	class FPrimitiveType : SHADER_PERMUTATION_INT("DISTANCEFIELD_PRIMITIVE_TYPE", 2);
@@ -415,6 +415,8 @@ void ScatterObjectsToShadowTiles(
 		PassParameters->PS.ObjectBufferParameters = ObjectBufferParameters;
 		PassParameters->PS.CulledObjectBufferParameters = CulledObjectBufferParameters;
 		PassParameters->PS.LightTileIntersectionParameters = LightTileIntersectionParameters;
+		PassParameters->PS.WorldToShadow = FMatrix44f(WorldToShadowValue);
+		PassParameters->PS.ObjectExpandScale = PrimitiveType == DFPT_HeightField ? 0.f : WorldToShadowValue.GetMaximumAxisScale();
 
 		PassParameters->MeshSDFIndirectArgs = ObjectIndirectArguments;
 
@@ -471,49 +473,19 @@ void ScatterObjectsToShadowTiles(
 
 void AllocateDistanceFieldCulledObjectBuffers(
 	FRDGBuilder& GraphBuilder, 
-	bool bWantBoxBounds, 
 	uint32 MaxObjects, 
-	EDistanceFieldPrimitiveType PrimitiveType,
 	FRDGBufferRef& OutObjectIndirectArguments,
 	FDistanceFieldCulledObjectBufferParameters& OutParameters)
 {
 	check(MaxObjects > 0);
-	OutObjectIndirectArguments = GraphBuilder.CreateBuffer(FRDGBufferDesc::CreateIndirectDesc<FRHIDrawIndexedIndirectParameters>(), TEXT("FDistanceFieldCulledObjectBuffers_ObjectIndirectArguments"));
-
-	uint32 NumBoundsElementsScale;
-	uint32 ObjectDataStride;
-	uint32 ObjectBoxBoundsStride;
-
-	if (PrimitiveType == DFPT_SignedDistanceField)
-	{
-		NumBoundsElementsScale = 1;
-		ObjectDataStride = GDistanceFieldCulledObjectDataStride;
-		ObjectBoxBoundsStride = GDistanceFieldCulledObjectBoxBoundsStride;
-	}
-	else
-	{
-		NumBoundsElementsScale = 2;
-		ObjectDataStride = GHeightFieldCulledObjectDataStride;
-		ObjectBoxBoundsStride = GHeightFieldCulledObjectBoxBoundsStride;
-	}
-
-	FRDGBufferRef Bounds = GraphBuilder.CreateBuffer(FRDGBufferDesc::CreateStructuredDesc(sizeof(FVector4f), MaxObjects * NumBoundsElementsScale), TEXT("FDistanceFieldCulledObjectBuffers_Bounds"));
-	FRDGBufferRef Data = GraphBuilder.CreateBuffer(FRDGBufferDesc::CreateStructuredDesc(sizeof(FVector4f), MaxObjects * ObjectDataStride), TEXT("FDistanceFieldCulledObjectBuffers_Data"));
+	OutObjectIndirectArguments = GraphBuilder.CreateBuffer(FRDGBufferDesc::CreateIndirectDesc<FRHIDrawIndexedIndirectParameters>(), TEXT("DistanceField.ObjectIndirectArguments"));
+	FRDGBufferRef ObjectIndices = GraphBuilder.CreateBuffer(FRDGBufferDesc::CreateStructuredDesc(sizeof(uint32), MaxObjects), TEXT("DistanceField.ObjectIndices"));
 
 	OutParameters.RWObjectIndirectArguments = GraphBuilder.CreateUAV(OutObjectIndirectArguments, PF_R32_UINT);
-	OutParameters.RWCulledObjectBounds = GraphBuilder.CreateUAV(Bounds);
-	OutParameters.RWCulledObjectData = GraphBuilder.CreateUAV(Data);
+	OutParameters.RWCulledObjectIndices = GraphBuilder.CreateUAV(ObjectIndices);
 
 	OutParameters.ObjectIndirectArguments = GraphBuilder.CreateSRV(OutObjectIndirectArguments, PF_R32_UINT);
-	OutParameters.CulledObjectBounds = GraphBuilder.CreateSRV(Bounds);
-	OutParameters.CulledObjectData = GraphBuilder.CreateSRV(Data);
-
-	if (bWantBoxBounds)
-	{
-		FRDGBufferRef BoxBounds = GraphBuilder.CreateBuffer(FRDGBufferDesc::CreateStructuredDesc(sizeof(FVector4f), MaxObjects * ObjectBoxBoundsStride), TEXT("FDistanceFieldCulledObjectBuffers_BoxBounds"));
-		OutParameters.RWCulledObjectBoxBounds = GraphBuilder.CreateUAV(BoxBounds);
-		OutParameters.CulledObjectBoxBounds = GraphBuilder.CreateSRV(BoxBounds);
-	}
+	OutParameters.CulledObjectIndices = GraphBuilder.CreateSRV(ObjectIndices);
 }
 
 void CullDistanceFieldObjectsForLight(
@@ -541,9 +513,7 @@ void CullDistanceFieldObjectsForLight(
 
 	AllocateDistanceFieldCulledObjectBuffers(
 		GraphBuilder, 
-		true, 
 		FMath::DivideAndRoundUp(NumObjectsInBuffer, 256) * 256, 
-		PrimitiveType,
 		ObjectIndirectArguments,
 		CulledObjectBufferParameters);
 
@@ -557,7 +527,6 @@ void CullDistanceFieldObjectsForLight(
 		PassParameters->CulledObjectBufferParameters = CulledObjectBufferParameters;
 		PassParameters->ObjectBoundingGeometryIndexCount = UE_ARRAY_COUNT(GCubeIndices);
 		PassParameters->WorldToShadow = FMatrix44f(WorldToShadowValue);	// LWC_TODO: Precision loss
-		PassParameters->ObjectExpandScale = bIsHeightfield ? 0.f : WorldToShadowValue.GetMaximumAxisScale();
 		PassParameters->NumShadowHullPlanes = NumPlanes;
 		PassParameters->ShadowBoundingSphere = (FVector4f)ShadowBoundingSphereValue; // LWC_TODO: Precision loss
 		// Disable Nanite meshes for directional lights that use VSM since they draw into the VSM unconditionally (and would get double shadow)
@@ -926,10 +895,7 @@ void FProjectedShadowInfo::BeginRenderRayTracedDistanceFieldProjection(
 		const FVector4 ShadowBoundingSphereValue(0.f, 0.f, 0.f, 0.f);
 		const FMatrix WorldToShadowValue = FTranslationMatrix(PreShadowTranslation) * FMatrix(TranslatedWorldToClipInnerMatrix);
 
-		FDistanceFieldObjectBufferParameters ObjectBufferParameters;
-		ObjectBufferParameters.SceneObjectBounds = Scene->DistanceFieldSceneData.GetHeightFieldObjectBuffers()->Bounds.SRV;
-		ObjectBufferParameters.SceneObjectData = Scene->DistanceFieldSceneData.GetHeightFieldObjectBuffers()->Data.SRV;
-		ObjectBufferParameters.NumSceneObjects = Scene->DistanceFieldSceneData.NumHeightFieldObjectsInBuffer;
+		FDistanceFieldObjectBufferParameters ObjectBufferParameters = DistanceField::SetupObjectBufferParameters(Scene->DistanceFieldSceneData);
 
 		FLightTileIntersectionParameters LightTileIntersectionParameters;
 		FDistanceFieldCulledObjectBufferParameters CulledObjectBufferParameters;
diff --git a/Engine/Source/Runtime/Renderer/Private/DistanceFieldVisualization.cpp b/Engine/Source/Runtime/Renderer/Private/DistanceFieldVisualization.cpp
--- a/Engine/Source/Runtime/Renderer/Private/DistanceFieldVisualization.cpp
+++ b/Engine/Source/Runtime/Renderer/Private/DistanceFieldVisualization.cpp
@@ -30,6 +30,7 @@ public:
 		SHADER_PARAMETER_STRUCT_REF(FViewUniformShaderParameters, View)
 		SHADER_PARAMETER_RDG_UNIFORM_BUFFER(FForwardLightData, ForwardLightData)
 		SHADER_PARAMETER_STRUCT_REF(FReflectionUniformParameters, ReflectionStruct)
+		SHADER_PARAMETER_STRUCT_INCLUDE(FDistanceFieldObjectBufferParameters, DistanceFieldObjectBuffers)
 		SHADER_PARAMETER_STRUCT_INCLUDE(FDistanceFieldCulledObjectBufferParameters, DistanceFieldCulledObjectBuffers)
 		SHADER_PARAMETER_STRUCT_INCLUDE(FDistanceFieldAtlasParameters, DistanceFieldAtlas)
 		SHADER_PARAMETER_STRUCT_INCLUDE(FAOParameters, AOParameters)
@@ -131,9 +132,7 @@ void FDeferredShadingSceneRenderer::RenderMeshDistanceFieldVisualization(
 		{
 			AllocateDistanceFieldCulledObjectBuffers(
 				GraphBuilder,
-				false,
 				FMath::DivideAndRoundUp(Scene->DistanceFieldSceneData.NumObjectsInBuffer, 256) * 256,
-				DFPT_SignedDistanceField,
 				ObjectIndirectArguments,
 				CulledObjectBufferParameters);
 
@@ -149,6 +148,7 @@ void FDeferredShadingSceneRenderer::RenderMeshDistanceFieldVisualization(
 		PassParameters->View = View.ViewUniformBuffer;
 		PassParameters->ForwardLightData = View.ForwardLightingResources.ForwardLightUniformBuffer;
 		PassParameters->ReflectionStruct = CreateReflectionUniformBuffer(View, UniformBuffer_MultiFrame);
+		PassParameters->DistanceFieldObjectBuffers = DistanceField::SetupObjectBufferParameters(Scene->DistanceFieldSceneData);
 		PassParameters->DistanceFieldCulledObjectBuffers = CulledObjectBufferParameters;
 		PassParameters->DistanceFieldAtlas = DistanceField::SetupAtlasParameters(Scene->DistanceFieldSceneData);
 		PassParameters->AOParameters = DistanceField::SetupAOShaderParameters(Parameters);
