Fixes s/taa_/tsr_/g in TSR shaders

diff --git a/Engine/Shaders/Private/TemporalSuperResolution/TSRCommon.ush b/Engine/Shaders/Private/TemporalSuperResolution/TSRCommon.ush
--- a/Engine/Shaders/Private/TemporalSuperResolution/TSRCommon.ush
+++ b/Engine/Shaders/Private/TemporalSuperResolution/TSRCommon.ush
@@ -141,87 +141,87 @@ static const uint kPlusIndexes3x3[5] = { 4, 1, 3, 7, 5 };
 
 
 #if CONFIG_COMPILE_FP16
-	#define taa_half  half
-	#define taa_half2 half2
-	#define taa_half3 half3
-	#define taa_half4 half4
+	#define tsr_half  half
+	#define tsr_half2 half2
+	#define tsr_half3 half3
+	#define tsr_half4 half4
 
-	#define taa_short  int16_t
-	#define taa_short2 int16_t2
-	#define taa_short3 int16_t3
-	#define taa_short4 int16_t4
+	#define tsr_short  int16_t
+	#define tsr_short2 int16_t2
+	#define tsr_short3 int16_t3
+	#define tsr_short4 int16_t4
 
-	#define taa_ushort  uint16_t
-	#define taa_ushort2 uint16_t2
-	#define taa_ushort3 uint16_t3
-	#define taa_ushort4 uint16_t4
+	#define tsr_ushort  uint16_t
+	#define tsr_ushort2 uint16_t2
+	#define tsr_ushort3 uint16_t3
+	#define tsr_ushort4 uint16_t4
 
-	#define taa_half2x2 half2x2
-	#define taa_half3x2 half3x2
-	#define taa_half4x2 half4x2
+	#define tsr_half2x2 half2x2
+	#define tsr_half3x2 half3x2
+	#define tsr_half4x2 half4x2
 
-	#define taa_short2x2 int16_t2x2
-	#define taa_short3x2 int16_t3x2
-	#define taa_short4x2 int16_t4x2
+	#define tsr_short2x2 int16_t2x2
+	#define tsr_short3x2 int16_t3x2
+	#define tsr_short4x2 int16_t4x2
 
-	#define taa_ushort2x2 uint16_t2x2
-	#define taa_ushort3x2 uint16_t3x2
-	#define taa_ushort4x2 uint16_t4x2
+	#define tsr_ushort2x2 uint16_t2x2
+	#define tsr_ushort3x2 uint16_t3x2
+	#define tsr_ushort4x2 uint16_t4x2
 
 #else
-	#define taa_half  float
-	#define taa_half2 float2
-	#define taa_half3 float3
-	#define taa_half4 float4
+	#define tsr_half  float
+	#define tsr_half2 float2
+	#define tsr_half3 float3
+	#define tsr_half4 float4
 
-	#define taa_short  int
-	#define taa_short2 int2
-	#define taa_short3 int3
-	#define taa_short4 int4
+	#define tsr_short  int
+	#define tsr_short2 int2
+	#define tsr_short3 int3
+	#define tsr_short4 int4
 
-	#define taa_ushort  uint
-	#define taa_ushort2 uint2
-	#define taa_ushort3 uint3
-	#define taa_ushort4 uint4
+	#define tsr_ushort  uint
+	#define tsr_ushort2 uint2
+	#define tsr_ushort3 uint3
+	#define tsr_ushort4 uint4
 
-	#define taa_half2x2 float2x2
-	#define taa_half3x2 float3x2
-	#define taa_half4x2 float4x2
+	#define tsr_half2x2 float2x2
+	#define tsr_half3x2 float3x2
+	#define tsr_half4x2 float4x2
 
-	#define taa_short2x2 int2x2
-	#define taa_short3x2 int3x2
-	#define taa_short4x2 int4x2
+	#define tsr_short2x2 int2x2
+	#define tsr_short3x2 int3x2
+	#define tsr_short4x2 int4x2
 
-	#define taa_ushort2x2 uint2x2
-	#define taa_ushort3x2 uint3x2
-	#define taa_ushort4x2 uint4x2
+	#define tsr_ushort2x2 uint2x2
+	#define tsr_ushort3x2 uint3x2
+	#define tsr_ushort4x2 uint4x2
 
 #endif
 
 #if CONFIG_SCENE_COLOR_APLHA
-	#define taa_halfC taa_half4
-	#define taa_halfCx2 taa_half4x2
+	#define tsr_halfC tsr_half4
+	#define tsr_halfCx2 tsr_half4x2
 
 #else
-	#define taa_halfC taa_half3
-	#define taa_halfCx2 taa_half3x2
+	#define tsr_halfC tsr_half3
+	#define tsr_halfCx2 tsr_half3x2
 
 #endif
 
 
-#define taa_subpixel_details uint
-#define taa_subpixel_details2 uint2
-#define taa_subpixel_payload uint // TODO: uint16_t
-#define taa_subpixel_payload2 uint2 // TODO: uint16_t
+#define tsr_subpixel_details uint
+#define tsr_subpixel_details2 uint2
+#define tsr_subpixel_payload uint // TODO: uint16_t
+#define tsr_subpixel_payload2 uint2 // TODO: uint16_t
 
 // Largest encodable normal number in a half used on console.
-static const taa_half LargestNormalNumber = taa_half(65504.0);
+static const tsr_half LargestNormalNumber = tsr_half(65504.0);
 
 // Largest encodable scene color value to not overflaw in YCoCg with halfs on console
-static const taa_half3 LargestSceneColorRGB = (LargestNormalNumber * taa_half(0.25)).xxx;
-static const taa_half4 LargestSceneColorRGBA = (LargestNormalNumber * taa_half(0.25)).xxxx;
-static const taa_half3 LargestSceneColorYCoCg = LargestNormalNumber.xxx;
-static const taa_half4 LargestSceneColorYCoCgA = LargestNormalNumber.xxxx;
+static const tsr_half3 LargestSceneColorRGB = (LargestNormalNumber * tsr_half(0.25)).xxx;
+static const tsr_half4 LargestSceneColorRGBA = (LargestNormalNumber * tsr_half(0.25)).xxxx;
+static const tsr_half3 LargestSceneColorYCoCg = LargestNormalNumber.xxx;
+static const tsr_half4 LargestSceneColorYCoCgA = LargestNormalNumber.xxxx;
 
 
 //------------------------------------------------------- PARAMETERS
@@ -364,18 +364,18 @@ uint GetGroupThreadIndex(uint GroupThreadIndex, uint GroupWaveIndex)
 
 // Clamp the offset to be shared across multiple samples
 CALL_SITE_DEBUGLOC
-taa_short2x2 ClampPixelOffset(
-	taa_short2x2 KernelCenterPixelPos,
-	taa_short2x2 Offset,
-	const taa_short2 OffsetDirection,
+tsr_short2x2 ClampPixelOffset(
+	tsr_short2x2 KernelCenterPixelPos,
+	tsr_short2x2 Offset,
+	const tsr_short2 OffsetDirection,
 	int2 MinPixelPos, int2 MaxPixelPos)
 {
-	taa_short2x2 Min = dpv_sub(taa_short2(MinPixelPos), KernelCenterPixelPos);
-	taa_short2x2 Max = dpv_sub(taa_short2(MaxPixelPos), KernelCenterPixelPos);
+	tsr_short2x2 Min = dpv_sub(tsr_short2(MinPixelPos), KernelCenterPixelPos);
+	tsr_short2x2 Max = dpv_sub(tsr_short2(MaxPixelPos), KernelCenterPixelPos);
 
 	// Only do clamp based on the compile time known direction of the offset.
 	// This turns only 1 v_pk_max_i16 and v_pk_min_i16 for a 3x3 kernel.
-	taa_short2x2 ClampedOffset = 0;
+	tsr_short2x2 ClampedOffset = 0;
 
 	if (OffsetDirection.x > 0)
 	{
@@ -399,18 +399,18 @@ taa_short2x2 ClampPixelOffset(
 }
 
 CALL_SITE_DEBUGLOC
-taa_short2 ClampPixelOffset(
-	taa_short2 KernelCenterPixelPos,
-	taa_short2 Offset,
-	const taa_short2 OffsetDirection,
+tsr_short2 ClampPixelOffset(
+	tsr_short2 KernelCenterPixelPos,
+	tsr_short2 Offset,
+	const tsr_short2 OffsetDirection,
 	int2 MinPixelPos, int2 MaxPixelPos)
 {
-	taa_short2 Min = taa_short2(MinPixelPos) - KernelCenterPixelPos;
-	taa_short2 Max = taa_short2(MaxPixelPos) - KernelCenterPixelPos;
+	tsr_short2 Min = tsr_short2(MinPixelPos) - KernelCenterPixelPos;
+	tsr_short2 Max = tsr_short2(MaxPixelPos) - KernelCenterPixelPos;
 
 	// Only do clamp based on the compile time known direction of the offset.
 	// This turns only 1 v_pk_max_i16 and v_pk_min_i16 for a 3x3 kernel.
-	taa_short2 ClampedOffset = 0;
+	tsr_short2 ClampedOffset = 0;
 
 	if (OffsetDirection.x > 0)
 		ClampedOffset.x = min(Offset.x, Max.x);
@@ -426,114 +426,114 @@ taa_short2 ClampPixelOffset(
 }
 
 CALL_SITE_DEBUGLOC
-taa_short2x2 ClampPixelOffset(
-	taa_short2x2 SamplePixelPos,
+tsr_short2x2 ClampPixelOffset(
+	tsr_short2x2 SamplePixelPos,
 	int2 MinPixelPos, int2 MaxPixelPos)
 {
-	taa_short2 SamplePixelPos0 = dpv_lo(SamplePixelPos);
-	taa_short2 SamplePixelPos1 = dpv_hi(SamplePixelPos);
+	tsr_short2 SamplePixelPos0 = dpv_lo(SamplePixelPos);
+	tsr_short2 SamplePixelPos1 = dpv_hi(SamplePixelPos);
 
-	SamplePixelPos0.x = min(SamplePixelPos0.x, taa_short(MaxPixelPos.x));
-	SamplePixelPos0.x = max(SamplePixelPos0.x, taa_short(MinPixelPos.x));
+	SamplePixelPos0.x = min(SamplePixelPos0.x, tsr_short(MaxPixelPos.x));
+	SamplePixelPos0.x = max(SamplePixelPos0.x, tsr_short(MinPixelPos.x));
 		
-	SamplePixelPos0.y = min(SamplePixelPos0.y, taa_short(MaxPixelPos.y));
-	SamplePixelPos0.y = max(SamplePixelPos0.y, taa_short(MinPixelPos.y));
+	SamplePixelPos0.y = min(SamplePixelPos0.y, tsr_short(MaxPixelPos.y));
+	SamplePixelPos0.y = max(SamplePixelPos0.y, tsr_short(MinPixelPos.y));
 	
-	SamplePixelPos1.x = min(SamplePixelPos1.x, taa_short(MaxPixelPos.x));
-	SamplePixelPos1.x = max(SamplePixelPos1.x, taa_short(MinPixelPos.x));
+	SamplePixelPos1.x = min(SamplePixelPos1.x, tsr_short(MaxPixelPos.x));
+	SamplePixelPos1.x = max(SamplePixelPos1.x, tsr_short(MinPixelPos.x));
 		
-	SamplePixelPos1.y = min(SamplePixelPos1.y, taa_short(MaxPixelPos.y));
-	SamplePixelPos1.y = max(SamplePixelPos1.y, taa_short(MinPixelPos.y));
+	SamplePixelPos1.y = min(SamplePixelPos1.y, tsr_short(MaxPixelPos.y));
+	SamplePixelPos1.y = max(SamplePixelPos1.y, tsr_short(MinPixelPos.y));
 
 	return dpv_interleave_registers(SamplePixelPos0, SamplePixelPos1);
 }
 
 CALL_SITE_DEBUGLOC
-taa_short2 ClampPixelOffset(
-	taa_short2 SamplePixelPos,
+tsr_short2 ClampPixelOffset(
+	tsr_short2 SamplePixelPos,
 	int2 MinPixelPos, int2 MaxPixelPos)
 {
-	SamplePixelPos.x = min(SamplePixelPos.x, taa_short(MaxPixelPos.x));
-	SamplePixelPos.x = max(SamplePixelPos.x, taa_short(MinPixelPos.x));
+	SamplePixelPos.x = min(SamplePixelPos.x, tsr_short(MaxPixelPos.x));
+	SamplePixelPos.x = max(SamplePixelPos.x, tsr_short(MinPixelPos.x));
 		
-	SamplePixelPos.y = min(SamplePixelPos.y, taa_short(MaxPixelPos.y));
-	SamplePixelPos.y = max(SamplePixelPos.y, taa_short(MinPixelPos.y));
+	SamplePixelPos.y = min(SamplePixelPos.y, tsr_short(MaxPixelPos.y));
+	SamplePixelPos.y = max(SamplePixelPos.y, tsr_short(MinPixelPos.y));
 
 	return SamplePixelPos;
 }
 
 CALL_SITE_DEBUGLOC
-taa_short2x2 AddAndClampPixelOffset(
-	taa_short2x2 KernelCenterPixelPos,
-	taa_short2x2 Offset,
-	const taa_short2 OffsetDirection,
+tsr_short2x2 AddAndClampPixelOffset(
+	tsr_short2x2 KernelCenterPixelPos,
+	tsr_short2x2 Offset,
+	const tsr_short2 OffsetDirection,
 	int2 MinPixelPos, int2 MaxPixelPos)
 {
-	taa_short2x2 SamplePixelPos = KernelCenterPixelPos + Offset;
+	tsr_short2x2 SamplePixelPos = KernelCenterPixelPos + Offset;
 	
-	taa_short2 SamplePixelPos0 = dpv_lo(SamplePixelPos);
-	taa_short2 SamplePixelPos1 = dpv_hi(SamplePixelPos);
+	tsr_short2 SamplePixelPos0 = dpv_lo(SamplePixelPos);
+	tsr_short2 SamplePixelPos1 = dpv_hi(SamplePixelPos);
 
 	// Only do clamp based on the compile time known direction of the offset.
 	// This turns only 1 v_pk_max_i16 and v_pk_min_i16 for a 3x3 kernel.
 	if (OffsetDirection.x > 0)
-		SamplePixelPos0.x = min(SamplePixelPos0.x, taa_short(MaxPixelPos.x));
+		SamplePixelPos0.x = min(SamplePixelPos0.x, tsr_short(MaxPixelPos.x));
 	else if (OffsetDirection.x < 0)
-		SamplePixelPos0.x = max(SamplePixelPos0.x, taa_short(MinPixelPos.x));
+		SamplePixelPos0.x = max(SamplePixelPos0.x, tsr_short(MinPixelPos.x));
 		
 	if (OffsetDirection.y > 0)
-		SamplePixelPos0.y = min(SamplePixelPos0.y, taa_short(MaxPixelPos.y));
+		SamplePixelPos0.y = min(SamplePixelPos0.y, tsr_short(MaxPixelPos.y));
 	else if (OffsetDirection.y < 0)
-		SamplePixelPos0.y = max(SamplePixelPos0.y, taa_short(MinPixelPos.y));
+		SamplePixelPos0.y = max(SamplePixelPos0.y, tsr_short(MinPixelPos.y));
 	
 	if (OffsetDirection.x > 0)
-		SamplePixelPos1.x = min(SamplePixelPos1.x, taa_short(MaxPixelPos.x));
+		SamplePixelPos1.x = min(SamplePixelPos1.x, tsr_short(MaxPixelPos.x));
 	else if (OffsetDirection.x < 0)
-		SamplePixelPos1.x = max(SamplePixelPos1.x, taa_short(MinPixelPos.x));
+		SamplePixelPos1.x = max(SamplePixelPos1.x, tsr_short(MinPixelPos.x));
 		
 	if (OffsetDirection.y > 0)
-		SamplePixelPos1.y = min(SamplePixelPos1.y, taa_short(MaxPixelPos.y));
+		SamplePixelPos1.y = min(SamplePixelPos1.y, tsr_short(MaxPixelPos.y));
 	else if (OffsetDirection.y < 0)
-		SamplePixelPos1.y = max(SamplePixelPos1.y, taa_short(MinPixelPos.y));
+		SamplePixelPos1.y = max(SamplePixelPos1.y, tsr_short(MinPixelPos.y));
 	
 	return dpv_interleave_registers(SamplePixelPos0, SamplePixelPos1);
 }
 
 CALL_SITE_DEBUGLOC
-taa_short2 AddAndClampPixelOffset(
-	taa_short2 KernelCenterPixelPos,
-	taa_short2 Offset,
-	const taa_short2 OffsetDirection,
+tsr_short2 AddAndClampPixelOffset(
+	tsr_short2 KernelCenterPixelPos,
+	tsr_short2 Offset,
+	const tsr_short2 OffsetDirection,
 	int2 MinPixelPos, int2 MaxPixelPos)
 {
-	taa_short2 SamplePixelPos = KernelCenterPixelPos + Offset;
+	tsr_short2 SamplePixelPos = KernelCenterPixelPos + Offset;
 
 	// Only do clamp based on the compile time known direction of the offset.
 	// This turns only 1 v_pk_max_i16 and v_pk_min_i16 for a 3x3 kernel.
 	if (OffsetDirection.x > 0)
-		SamplePixelPos.x = min(SamplePixelPos.x, taa_short(MaxPixelPos.x));
+		SamplePixelPos.x = min(SamplePixelPos.x, tsr_short(MaxPixelPos.x));
 	else if (OffsetDirection.x < 0)
-		SamplePixelPos.x = max(SamplePixelPos.x, taa_short(MinPixelPos.x));
+		SamplePixelPos.x = max(SamplePixelPos.x, tsr_short(MinPixelPos.x));
 		
 	if (OffsetDirection.y > 0)
-		SamplePixelPos.y = min(SamplePixelPos.y, taa_short(MaxPixelPos.y));
+		SamplePixelPos.y = min(SamplePixelPos.y, tsr_short(MaxPixelPos.y));
 	else if (OffsetDirection.y < 0)
-		SamplePixelPos.y = max(SamplePixelPos.y, taa_short(MinPixelPos.y));
+		SamplePixelPos.y = max(SamplePixelPos.y, tsr_short(MinPixelPos.y));
 
 	return SamplePixelPos;
 }
 
 CALL_SITE_DEBUGLOC
-taa_short2x2 InvalidateOutputPixelPos(taa_short2x2 PixelPos, uint2 ViewportMax)
+tsr_short2x2 InvalidateOutputPixelPos(tsr_short2x2 PixelPos, uint2 ViewportMax)
 #if 1
 {
-	taa_short2x2 Subtract = dpv_sub(taa_short2(ViewportMax - 1), PixelPos);
-	taa_ushort2 Override = taa_ushort2(Subtract[0] | Subtract[1]);
+	tsr_short2x2 Subtract = dpv_sub(tsr_short2(ViewportMax - 1), PixelPos);
+	tsr_ushort2 Override = tsr_ushort2(Subtract[0] | Subtract[1]);
 
 	#if CONFIG_COMPILE_FP16
-		PixelPos[0] |= -taa_short2((Override & uint16_t(0x8000)) >> 15);
+		PixelPos[0] |= -tsr_short2((Override & uint16_t(0x8000)) >> 15);
 	#else
-		PixelPos[0] |= -taa_short2((Override & uint(0x80000000)) >> 31);
+		PixelPos[0] |= -tsr_short2((Override & uint(0x80000000)) >> 31);
 	#endif
 
 	return PixelPos;
@@ -541,22 +541,22 @@ taa_short2x2 InvalidateOutputPixelPos(taa_short2x2 PixelPos, uint2 ViewportMax)
 #else
 {
 	bool bIsValidPixel = all(PixelPos < ViewportMax);
-	PixelPos.x = bIsValidPixel ? PixelPos.x : ~taa_short(0);
+	PixelPos.x = bIsValidPixel ? PixelPos.x : ~tsr_short(0);
 	return PixelPos;
 }
 #endif
 
 CALL_SITE_DEBUGLOC
-taa_short2 InvalidateOutputPixelPos(taa_short2 PixelPos, uint2 ViewportMax)
+tsr_short2 InvalidateOutputPixelPos(tsr_short2 PixelPos, uint2 ViewportMax)
 #if 1
 {
-	taa_short2 Subtract = taa_short2(ViewportMax - 1) - PixelPos;
-	taa_ushort Override = taa_ushort(Subtract.x | Subtract.y);
+	tsr_short2 Subtract = tsr_short2(ViewportMax - 1) - PixelPos;
+	tsr_ushort Override = tsr_ushort(Subtract.x | Subtract.y);
 
 	#if CONFIG_COMPILE_FP16
-		PixelPos.x |= -taa_short((Override & uint16_t(0x8000)) >> 15);
+		PixelPos.x |= -tsr_short((Override & uint16_t(0x8000)) >> 15);
 	#else
-		PixelPos.x |= -taa_short((Override & uint(0x80000000)) >> 31);
+		PixelPos.x |= -tsr_short((Override & uint(0x80000000)) >> 31);
 	#endif
 
 	return PixelPos;
@@ -564,27 +564,27 @@ taa_short2 InvalidateOutputPixelPos(taa_short2 PixelPos, uint2 ViewportMax)
 #else
 {
 	bool bIsValidPixel = all(PixelPos < ViewportMax);
-	PixelPos.x = bIsValidPixel ? PixelPos.x : ~taa_short(0);
+	PixelPos.x = bIsValidPixel ? PixelPos.x : ~tsr_short(0);
 	return PixelPos;
 }
 #endif
 
 CALL_SITE_DEBUGLOC
-taa_ushort2 Map8x8Tile2x2Lane(uint GroupThreadIndex)
+tsr_ushort2 Map8x8Tile2x2Lane(uint GroupThreadIndex)
 {
-	taa_ushort2 T = taa_ushort(GroupThreadIndex).xx;
-	taa_ushort2 GroupId = (T >> taa_ushort2(0, 1) & taa_ushort2(0x01, 0x01)) | ((T >> taa_ushort2(2 - 1, 4 - 1)) & taa_ushort2(0x03 << 1, 0x03 << 1));
+	tsr_ushort2 T = tsr_ushort(GroupThreadIndex).xx;
+	tsr_ushort2 GroupId = (T >> tsr_ushort2(0, 1) & tsr_ushort2(0x01, 0x01)) | ((T >> tsr_ushort2(2 - 1, 4 - 1)) & tsr_ushort2(0x03 << 1, 0x03 << 1));
 	return GroupId;
 }
 
 CALL_SITE_DEBUGLOC
-taa_ushort2 Map16x16Tile2x2Lane(uint GroupThreadIndex)
+tsr_ushort2 Map16x16Tile2x2Lane(uint GroupThreadIndex)
 {
-	taa_ushort T = taa_ushort(GroupThreadIndex);
+	tsr_ushort T = tsr_ushort(GroupThreadIndex);
 
-	taa_ushort2 GroupId;
-	GroupId.x = ((T >> taa_ushort(0)) & taa_ushort(0x01)) | ((T >> taa_ushort(2 - 1)) & taa_ushort(0x07 << 1));
-	GroupId.y = ((T >> taa_ushort(1)) & taa_ushort(0x01)) | ((T >> taa_ushort(5 - 1)) & taa_ushort(0x07 << 1));
+	tsr_ushort2 GroupId;
+	GroupId.x = ((T >> tsr_ushort(0)) & tsr_ushort(0x01)) | ((T >> tsr_ushort(2 - 1)) & tsr_ushort(0x07 << 1));
+	GroupId.y = ((T >> tsr_ushort(1)) & tsr_ushort(0x01)) | ((T >> tsr_ushort(5 - 1)) & tsr_ushort(0x07 << 1));
 	
 	return GroupId;
 }
@@ -620,80 +620,80 @@ half2 SafeRcp(half2 x)
 #endif
 
 CALL_SITE_DEBUGLOC
-taa_half2 ComputeLowFrequencyContribution(taa_half2 HistoryValidity)
+tsr_half2 ComputeLowFrequencyContribution(tsr_half2 HistoryValidity)
 {
-	return saturate((HistoryValidity - taa_half(1.0 * rcp(MAX_SAMPLE_COUNT))) * taa_half(6.0));
+	return saturate((HistoryValidity - tsr_half(1.0 * rcp(MAX_SAMPLE_COUNT))) * tsr_half(6.0));
 }
 
 CALL_SITE_DEBUGLOC
-taa_half ComputeLowFrequencyContribution(taa_half HistoryValidity)
+tsr_half ComputeLowFrequencyContribution(tsr_half HistoryValidity)
 {
 	return dpv_lo(ComputeLowFrequencyContribution(dpv_interleave_mono_registers(HistoryValidity)));
 }
 
 CALL_SITE_DEBUGLOC
-taa_half3x2 RGBToYCoCg(taa_half3x2 RGB)
+tsr_half3x2 RGBToYCoCg(tsr_half3x2 RGB)
 {
-    taa_half2 R = RGB[0];
-    taa_half2 G = RGB[1];
-    taa_half2 B = RGB[2];       
+    tsr_half2 R = RGB[0];
+    tsr_half2 G = RGB[1];
+    tsr_half2 B = RGB[2];       
 
-    taa_half2 RplusB = R + B;
-    taa_half2 Y = 2.0 * G + RplusB;
-    taa_half2 Cg = 2.0 * G - RplusB;
-    taa_half2 Co = 2.0 * (R - B);
-    taa_half3x2 YCoCg = taa_half3x2(Y, Co, Cg);
+    tsr_half2 RplusB = R + B;
+    tsr_half2 Y = 2.0 * G + RplusB;
+    tsr_half2 Cg = 2.0 * G - RplusB;
+    tsr_half2 Co = 2.0 * (R - B);
+    tsr_half3x2 YCoCg = tsr_half3x2(Y, Co, Cg);
 
     return YCoCg;
 }
 
 CALL_SITE_DEBUGLOC
-taa_half3x2 YCoCgToRGB(taa_half3x2 YCoCg)
+tsr_half3x2 YCoCgToRGB(tsr_half3x2 YCoCg)
 {
-	taa_half2 Y =  YCoCg[0] * taa_half(0.25);
-	taa_half2 Co = YCoCg[1] * taa_half(0.25);
-	taa_half2 Cg = YCoCg[2] * taa_half(0.25);
+	tsr_half2 Y =  YCoCg[0] * tsr_half(0.25);
+	tsr_half2 Co = YCoCg[1] * tsr_half(0.25);
+	tsr_half2 Cg = YCoCg[2] * tsr_half(0.25);
 
-	taa_half2 R = Y + Co - Cg;
-	taa_half2 G = Y + Cg;
-	taa_half2 B = Y - Co - Cg;
+	tsr_half2 R = Y + Co - Cg;
+	tsr_half2 G = Y + Cg;
+	tsr_half2 B = Y - Co - Cg;
 
-	taa_half3x2 RGB = taa_half3x2(R, G, B);
+	tsr_half3x2 RGB = tsr_half3x2(R, G, B);
 	return RGB;
 }
 
 CALL_SITE_DEBUGLOC
-taa_half4x2 RGBToYCoCg(taa_half4x2 RGBA)
+tsr_half4x2 RGBToYCoCg(tsr_half4x2 RGBA)
 {
-	return taa_half4x2(RGBToYCoCg(taa_half3x2(RGBA[0], RGBA[1], RGBA[2])), RGBA[3]);
+	return tsr_half4x2(RGBToYCoCg(tsr_half3x2(RGBA[0], RGBA[1], RGBA[2])), RGBA[3]);
 }
 
 CALL_SITE_DEBUGLOC
-taa_half4x2 YCoCgToRGB(taa_half4x2 YCoCgA)
+tsr_half4x2 YCoCgToRGB(tsr_half4x2 YCoCgA)
 {
-	return taa_half4x2(YCoCgToRGB(taa_half3x2(YCoCgA[0], YCoCgA[1], YCoCgA[2])), YCoCgA[3]);
+	return tsr_half4x2(YCoCgToRGB(tsr_half3x2(YCoCgA[0], YCoCgA[1], YCoCgA[2])), YCoCgA[3]);
 }
 
 CALL_SITE_DEBUGLOC
-taa_half3 RGBToYCoCg(taa_half3 RGB)
+tsr_half3 RGBToYCoCg(tsr_half3 RGB)
 {
 	return dpv_lo(RGBToYCoCg(dpv_interleave_mono_registers(RGB)));
 }
 
 CALL_SITE_DEBUGLOC
-taa_half3 YCoCgToRGB(taa_half3 YCoCg)
+tsr_half3 YCoCgToRGB(tsr_half3 YCoCg)
 {
 	return dpv_lo(YCoCgToRGB(dpv_interleave_mono_registers(YCoCg)));
 }
 
 CALL_SITE_DEBUGLOC
-taa_half4 RGBToYCoCg(taa_half4 RGBA)
+tsr_half4 RGBToYCoCg(tsr_half4 RGBA)
 {
 	return dpv_lo(RGBToYCoCg(dpv_interleave_mono_registers(RGBA)));
 }
 
 CALL_SITE_DEBUGLOC
-taa_half4 YCoCgToRGB(taa_half4 YCoCgA)
+tsr_half4 YCoCgToRGB(tsr_half4 YCoCgA)
 {
 	return dpv_lo(YCoCgToRGB(dpv_interleave_mono_registers(YCoCgA)));
 }
@@ -705,76 +705,76 @@ taa_half4 YCoCgToRGB(taa_half4 YCoCgA)
 // Faster but less accurate luma computation. 
 // Luma includes a scaling by 4.
 CALL_SITE_DEBUGLOC
-taa_half2 Luma4(taa_half3x2 Color)
+tsr_half2 Luma4(tsr_half3x2 Color)
 {
-	return (Color[1] * taa_half(2.0)) + (Color[0] + Color[2]);
+	return (Color[1] * tsr_half(2.0)) + (Color[0] + Color[2]);
 }
 
 CALL_SITE_DEBUGLOC
-taa_half2 Luma4(taa_half4x2 Color)
+tsr_half2 Luma4(tsr_half4x2 Color)
 {
-	return Luma4(taa_half3x2(Color[0], Color[1], Color[2]));
+	return Luma4(tsr_half3x2(Color[0], Color[1], Color[2]));
 }
 
 CALL_SITE_DEBUGLOC
-taa_half Luma4(taa_half3 Color)
+tsr_half Luma4(tsr_half3 Color)
 {
 	return dpv_lo(Luma4(dpv_interleave_mono_registers(Color)));
 }
 
 CALL_SITE_DEBUGLOC
-taa_half Luma4(taa_half4 Color)
+tsr_half Luma4(tsr_half4 Color)
 {
 	return Luma4(Color.rgb);
 }
 
 CALL_SITE_DEBUGLOC
-taa_half2 HdrWeightY(taa_half2 Luma)
+tsr_half2 HdrWeightY(tsr_half2 Luma)
 {
-	taa_half Exposure = taa_half(1.0);
+	tsr_half Exposure = tsr_half(1.0);
 
-	return max(taa_half(HDR_WEIGHT_SAFE_MIN_VALUE), rcp(Luma * Exposure + taa_half(4.0)));
+	return max(tsr_half(HDR_WEIGHT_SAFE_MIN_VALUE), rcp(Luma * Exposure + tsr_half(4.0)));
 }
 
 CALL_SITE_DEBUGLOC
-taa_half HdrWeightY(taa_half Luma)
+tsr_half HdrWeightY(tsr_half Luma)
 {
 	return dpv_lo(HdrWeightY(dpv_interleave_mono_registers(Luma)));
 }
 
 CALL_SITE_DEBUGLOC
-taa_half2 HdrWeightInvY(taa_half2 LDRLuma) 
+tsr_half2 HdrWeightInvY(tsr_half2 LDRLuma) 
 {
-	return taa_half(4.0) * rcp(taa_half(1.0) - LDRLuma);
+	return tsr_half(4.0) * rcp(tsr_half(1.0) - LDRLuma);
 }
 
 CALL_SITE_DEBUGLOC
-taa_half HdrWeightInvY(taa_half LDRLuma) 
+tsr_half HdrWeightInvY(tsr_half LDRLuma) 
 {
 	return dpv_lo(HdrWeightInvY(dpv_interleave_mono_registers(LDRLuma)));
 }
 
 // Optimized HDR weighting function.
 CALL_SITE_DEBUGLOC
-taa_half2 HdrWeight4(taa_half3x2 Color)
+tsr_half2 HdrWeight4(tsr_half3x2 Color)
 {
 	return HdrWeightY(Luma4(Color));
 }
 
 CALL_SITE_DEBUGLOC
-taa_half2 HdrWeight4(taa_half4x2 Color)
+tsr_half2 HdrWeight4(tsr_half4x2 Color)
 {
 	return HdrWeightY(Luma4(Color));
 }
 
 CALL_SITE_DEBUGLOC
-taa_half HdrWeight4(taa_half3 Color)
+tsr_half HdrWeight4(tsr_half3 Color)
 {
 	return HdrWeightY(Luma4(Color));
 }
 
 CALL_SITE_DEBUGLOC
-taa_half HdrWeight4(taa_half4 Color)
+tsr_half HdrWeight4(tsr_half4 Color)
 {
 	return HdrWeightY(Luma4(Color));
 }
@@ -782,18 +782,18 @@ taa_half HdrWeight4(taa_half4 Color)
 
 // Returns the weight of a pixels at a coordinate <PixelDelta> from the PDF highest point.
 CALL_SITE_DEBUGLOC
-taa_half2 ComputeSampleWeigth(taa_half2 UpscaleFactor, taa_half2x2 PixelDelta, const float MinimalContribution)
+tsr_half2 ComputeSampleWeigth(tsr_half2 UpscaleFactor, tsr_half2x2 PixelDelta, const float MinimalContribution)
 {
-	taa_half2 u2 = UpscaleFactor * UpscaleFactor;
+	tsr_half2 u2 = UpscaleFactor * UpscaleFactor;
 
 	// 1 - 1.9 * x^2 + 0.9 * x^4
-	taa_half2 x2 = saturate(u2 * dpv_length2(PixelDelta));
-	//return taa_half(((float(0.9) + MinimalContribution) * x2 - float(1.9)) * x2 + float(1.0));
-	return saturate((taa_half(0.9) * x2 - taa_half(1.9)) * x2 + taa_half(1.0 + MinimalContribution));
+	tsr_half2 x2 = saturate(u2 * dpv_length2(PixelDelta));
+	//return tsr_half(((float(0.9) + MinimalContribution) * x2 - float(1.9)) * x2 + float(1.0));
+	return saturate((tsr_half(0.9) * x2 - tsr_half(1.9)) * x2 + tsr_half(1.0 + MinimalContribution));
 }
 
 CALL_SITE_DEBUGLOC
-taa_half ComputeSampleWeigth(taa_half UpscaleFactor, taa_half2 PixelDelta, const float MinimalContribution)
+tsr_half ComputeSampleWeigth(tsr_half UpscaleFactor, tsr_half2 PixelDelta, const float MinimalContribution)
 {
 	return dpv_lo(ComputeSampleWeigth(
 		dpv_interleave_mono_registers(UpscaleFactor),
@@ -802,18 +802,18 @@ taa_half ComputeSampleWeigth(taa_half UpscaleFactor, taa_half2 PixelDelta, const
 }
 
 CALL_SITE_DEBUGLOC
-taa_half2x2 WeightedLerpFactors(taa_half2 WeightA, taa_half2 WeightB, taa_half2 Blend)
+tsr_half2x2 WeightedLerpFactors(tsr_half2 WeightA, tsr_half2 WeightB, tsr_half2 Blend)
 {
-	taa_half2 BlendA = (taa_half(1.0) - Blend) * WeightA;
-	taa_half2 BlendB = Blend * WeightB;
-	taa_half2 RcpBlend = SafeRcp(BlendA + BlendB);
+	tsr_half2 BlendA = (tsr_half(1.0) - Blend) * WeightA;
+	tsr_half2 BlendB = Blend * WeightB;
+	tsr_half2 RcpBlend = SafeRcp(BlendA + BlendB);
 	BlendA *= RcpBlend;
 	BlendB *= RcpBlend;
-	return taa_half2x2(BlendA, BlendB);
+	return tsr_half2x2(BlendA, BlendB);
 }
 
 CALL_SITE_DEBUGLOC
-taa_half2 WeightedLerpFactors(taa_half WeightA, taa_half WeightB, taa_half Blend)
+tsr_half2 WeightedLerpFactors(tsr_half WeightA, tsr_half WeightB, tsr_half Blend)
 {
 	return dpv_lo(WeightedLerpFactors(
 		dpv_interleave_mono_registers(WeightA),
@@ -837,67 +837,67 @@ float2 ComputeStaticVelocity(float2 ScreenPos, float DeviceZ)
 	return PosN.xy - PrevScreen;
 }
 
-taa_half ComputePredictionCompleteness(taa_half SampleHistoryValidity)
+tsr_half ComputePredictionCompleteness(tsr_half SampleHistoryValidity)
 {
-	return saturate(SampleHistoryValidity * taa_half(MAX_SAMPLE_COUNT) - taa_half(0.2));
+	return saturate(SampleHistoryValidity * tsr_half(MAX_SAMPLE_COUNT) - tsr_half(0.2));
 }
 
 CALL_SITE_DEBUGLOC
-bool IsOffScreen(uint bCameraCut, float2 ScreenPos, taa_half ParallaxRejectionMask)
+bool IsOffScreen(uint bCameraCut, float2 ScreenPos, tsr_half ParallaxRejectionMask)
 {
 	bool bIsCameraCut = bCameraCut != 0;
 	bool bIsOutOfBounds = max(abs(ScreenPos.x), abs(ScreenPos.y)) >= 1.0;
-	bool bIsParallaxRejected = ParallaxRejectionMask < taa_half(PARALLAX_REJECTION_MASK_THRESHOLD);
+	bool bIsParallaxRejected = ParallaxRejectionMask < tsr_half(PARALLAX_REJECTION_MASK_THRESHOLD);
 
 	return (bIsCameraCut || bIsOutOfBounds || bIsParallaxRejected);
 }
 
-taa_half MeasureBackbufferLDRQuantizationError()
+tsr_half MeasureBackbufferLDRQuantizationError()
 {
 	// Assume the backbuffer is 10bit per channels
-	return taa_half(0.5 / 1024.0);
+	return tsr_half(0.5 / 1024.0);
 }
 
-taa_half2 MeasureBackbufferLDRQuantizationErrorFromHDRLuma(taa_half2 Luma)
+tsr_half2 MeasureBackbufferLDRQuantizationErrorFromHDRLuma(tsr_half2 Luma)
 {
-	taa_half2 Tonemap = Luma * HdrWeightY(Luma) + MeasureBackbufferLDRQuantizationError().xx;
-	taa_half2 PixelColorQuantizationError = abs(Tonemap * HdrWeightInvY(Tonemap) - Luma);
+	tsr_half2 Tonemap = Luma * HdrWeightY(Luma) + MeasureBackbufferLDRQuantizationError().xx;
+	tsr_half2 PixelColorQuantizationError = abs(Tonemap * HdrWeightInvY(Tonemap) - Luma);
 
 	return PixelColorQuantizationError;
 }
 
-taa_half MeasureBackbufferLDRQuantizationErrorFromHDRLuma(taa_half Luma)
+tsr_half MeasureBackbufferLDRQuantizationErrorFromHDRLuma(tsr_half Luma)
 {
 	return dpv_lo(MeasureBackbufferLDRQuantizationErrorFromHDRLuma(dpv_interleave_mono_registers(Luma)));
 }
 
-taa_half2 MeasureRejectionFactor(
-	taa_halfCx2 PrevYCoCg,
-	taa_halfCx2 ClampedPrevYCoCg,
-	taa_halfCx2 InputCenterYCoCg,
-	taa_halfCx2 InputMinYCoCg, taa_halfCx2 InputMaxYCoCg,
-	taa_half2 BackbufferQuantizationError)
+tsr_half2 MeasureRejectionFactor(
+	tsr_halfCx2 PrevYCoCg,
+	tsr_halfCx2 ClampedPrevYCoCg,
+	tsr_halfCx2 InputCenterYCoCg,
+	tsr_halfCx2 InputMinYCoCg, tsr_halfCx2 InputMaxYCoCg,
+	tsr_half2 BackbufferQuantizationError)
 {
-	const taa_halfC Null = taa_half(0.0);
+	const tsr_halfC Null = tsr_half(0.0);
 	
 	#if CONFIG_SCENE_COLOR_APLHA
-		taa_halfCx2 BackbufferQuantizationErrorVector = taa_half4x2(BackbufferQuantizationError, BackbufferQuantizationError, BackbufferQuantizationError, BackbufferQuantizationError);
+		tsr_halfCx2 BackbufferQuantizationErrorVector = tsr_half4x2(BackbufferQuantizationError, BackbufferQuantizationError, BackbufferQuantizationError, BackbufferQuantizationError);
 	#else
-		taa_halfCx2 BackbufferQuantizationErrorVector = taa_half3x2(BackbufferQuantizationError, BackbufferQuantizationError, BackbufferQuantizationError);
+		tsr_halfCx2 BackbufferQuantizationErrorVector = tsr_half3x2(BackbufferQuantizationError, BackbufferQuantizationError, BackbufferQuantizationError);
 	#endif
 
-	taa_halfCx2 BoxSize = InputMaxYCoCg - InputMinYCoCg;
+	tsr_halfCx2 BoxSize = InputMaxYCoCg - InputMinYCoCg;
 
-	taa_halfCx2 ClampedEnergy = abs(ClampedPrevYCoCg - PrevYCoCg);
-	taa_halfCx2 Delta = abs(InputCenterYCoCg - PrevYCoCg);
+	tsr_halfCx2 ClampedEnergy = abs(ClampedPrevYCoCg - PrevYCoCg);
+	tsr_halfCx2 Delta = abs(InputCenterYCoCg - PrevYCoCg);
 
-	taa_halfCx2 ClampError = max(
-		BoxSize * taa_half(0.25),
+	tsr_halfCx2 ClampError = max(
+		BoxSize * tsr_half(0.25),
 		BackbufferQuantizationErrorVector);
 
-	taa_halfCx2 Factor = taa_half(1.0) - saturate(max(ClampedEnergy - ClampError, dpv_interleave_mono_registers(Null)) * rcp(max(Delta, BoxSize)));
+	tsr_halfCx2 Factor = tsr_half(1.0) - saturate(max(ClampedEnergy - ClampError, dpv_interleave_mono_registers(Null)) * rcp(max(Delta, BoxSize)));
 
-	taa_half2 Rejection = min3(Factor[0], Factor[1], Factor[2]);
+	tsr_half2 Rejection = min3(Factor[0], Factor[1], Factor[2]);
 	
 	#if CONFIG_SCENE_COLOR_APLHA
 		Rejection = min(Rejection, Factor[3]);
@@ -906,12 +906,12 @@ taa_half2 MeasureRejectionFactor(
 	return Rejection;
 }
 
-taa_half MeasureRejectionFactor(
-	taa_halfC PrevYCoCg,
-	taa_halfC ClampedPrevYCoCg,
-	taa_halfC InputCenterYCoCg,
-	taa_halfC InputMinYCoCg, taa_halfC InputMaxYCoCg,
-	taa_half BackbufferQuantizationError)
+tsr_half MeasureRejectionFactor(
+	tsr_halfC PrevYCoCg,
+	tsr_halfC ClampedPrevYCoCg,
+	tsr_halfC InputCenterYCoCg,
+	tsr_halfC InputMinYCoCg, tsr_halfC InputMaxYCoCg,
+	tsr_half BackbufferQuantizationError)
 {
 	return dpv_lo(MeasureRejectionFactor(
 		dpv_interleave_mono_registers(PrevYCoCg), 
@@ -925,7 +925,7 @@ taa_half MeasureRejectionFactor(
 struct FSubpixelNeighborhood
 {
 	// 2x2 quad of subpixels grid
-	taa_subpixel_details SampleArray[4];
+	tsr_subpixel_details SampleArray[4];
 
 	// Offset of subpixel
 	uint2 SubpixelOffset;
@@ -963,14 +963,14 @@ FSubpixelNeighborhood GatherPrevSubpixelNeighborhood(Texture2D<uint> PrevSubPixe
 	}
 	#endif
 
-	taa_half2 Interp = taa_half2(PixelCoord - TopLeftPixelCoord);
+	tsr_half2 Interp = tsr_half2(PixelCoord - TopLeftPixelCoord);
 	Neighborhood.SubpixelOffset = uint2(round(Interp * 2.0));
 	Neighborhood.ClosestSubpixelOffset = Interp > 0.5;
 
 	return Neighborhood;
 }
 
-taa_subpixel_payload GetSubpixelPayload(FSubpixelNeighborhood Neighborhood, uint SubpixelId)
+tsr_subpixel_payload GetSubpixelPayload(FSubpixelNeighborhood Neighborhood, uint SubpixelId)
 {
 	const uint2 SubpixelCoord = uint2(SubpixelId % SUB_PIXEL_GRID_SIZE, SubpixelId / SUB_PIXEL_GRID_SIZE);
 
@@ -980,48 +980,48 @@ taa_subpixel_payload GetSubpixelPayload(FSubpixelNeighborhood Neighborhood, uint
 
 
 	uint QuadSampleId = dot(InputTexelCoord, uint2(1, 2));
-	taa_subpixel_details InputSample = Neighborhood.SampleArray[QuadSampleId];
+	tsr_subpixel_details InputSample = Neighborhood.SampleArray[QuadSampleId];
 
 	uint BitShift = SUB_PIXEL_BIT_COUNT * dot(QuadSubpixelCoord, uint2(1, SUB_PIXEL_GRID_SIZE));
 
-	taa_subpixel_payload SubpixelPayload = taa_subpixel_payload(InputSample >> BitShift) & SUB_PIXEL_BIT_MASK;
+	tsr_subpixel_payload SubpixelPayload = tsr_subpixel_payload(InputSample >> BitShift) & SUB_PIXEL_BIT_MASK;
 
 	return SubpixelPayload;
 }
 
-taa_subpixel_payload GetClosestSubpixelPayload(FSubpixelNeighborhood Neighborhood)
+tsr_subpixel_payload GetClosestSubpixelPayload(FSubpixelNeighborhood Neighborhood)
 {
 	uint QuadSampleId = dot(uint2(Neighborhood.ClosestSubpixelOffset), uint2(1, 2));
-	taa_subpixel_details InputSample = Neighborhood.SampleArray[QuadSampleId];
+	tsr_subpixel_details InputSample = Neighborhood.SampleArray[QuadSampleId];
 
 	uint BitShift = SUB_PIXEL_BIT_COUNT * (3 - QuadSampleId);
 
-	taa_subpixel_payload SubpixelPayload = taa_subpixel_payload(InputSample >> BitShift) & SUB_PIXEL_BIT_MASK;
+	tsr_subpixel_payload SubpixelPayload = tsr_subpixel_payload(InputSample >> BitShift) & SUB_PIXEL_BIT_MASK;
 
 	return SubpixelPayload;
 }
 
 CALL_SITE_DEBUGLOC
-taa_subpixel_payload CompressParallaxFactor(taa_half ParallaxFactor)
+tsr_subpixel_payload CompressParallaxFactor(tsr_half ParallaxFactor)
 {
 	return clamp(
-		taa_subpixel_payload((ParallaxFactor - taa_half(1.0)) * taa_half(float(SUB_PIXEL_PARALLAX_FACTOR_BIT_MASK) / MAX_PARALLAX_FACTOR)),
-		taa_subpixel_payload(0), taa_subpixel_payload(SUB_PIXEL_PARALLAX_FACTOR_BIT_MASK));
+		tsr_subpixel_payload((ParallaxFactor - tsr_half(1.0)) * tsr_half(float(SUB_PIXEL_PARALLAX_FACTOR_BIT_MASK) / MAX_PARALLAX_FACTOR)),
+		tsr_subpixel_payload(0), tsr_subpixel_payload(SUB_PIXEL_PARALLAX_FACTOR_BIT_MASK));
 }
 
 CALL_SITE_DEBUGLOC
-taa_half UncompressParallaxFactor(taa_subpixel_payload CompressedParallaxFactor)
+tsr_half UncompressParallaxFactor(tsr_subpixel_payload CompressedParallaxFactor)
 {
 	return 
-		taa_half(CompressedParallaxFactor & SUB_PIXEL_PARALLAX_FACTOR_BIT_MASK) *
-		taa_half(MAX_PARALLAX_FACTOR / float(SUB_PIXEL_PARALLAX_FACTOR_BIT_MASK)) +
-		taa_half(1.0);
+		tsr_half(CompressedParallaxFactor & SUB_PIXEL_PARALLAX_FACTOR_BIT_MASK) *
+		tsr_half(MAX_PARALLAX_FACTOR / float(SUB_PIXEL_PARALLAX_FACTOR_BIT_MASK)) +
+		tsr_half(1.0);
 }
 
 CALL_SITE_DEBUGLOC
-taa_half3x2 QuantizeForFloatRenderTarget(taa_half3x2 Color, taa_half E, const float3 QuantizationError)
+tsr_half3x2 QuantizeForFloatRenderTarget(tsr_half3x2 Color, tsr_half E, const float3 QuantizationError)
 {
-	taa_half3x2 Error = dpv_mul(Color, taa_half3(QuantizationError));
+	tsr_half3x2 Error = dpv_mul(Color, tsr_half3(QuantizationError));
 
 	#if 0
 		// NOP
@@ -1043,7 +1043,7 @@ taa_half3x2 QuantizeForFloatRenderTarget(taa_half3x2 Color, taa_half E, const fl
 }
 
 CALL_SITE_DEBUGLOC
-taa_half3 QuantizeForFloatRenderTarget(taa_half3 Color, taa_half E, const float3 QuantizationError)
+tsr_half3 QuantizeForFloatRenderTarget(tsr_half3 Color, tsr_half E, const float3 QuantizationError)
 {
 	return dpv_lo(QuantizeForFloatRenderTarget(dpv_interleave_mono_registers(Color), E, QuantizationError));
 }
@@ -1064,12 +1064,12 @@ void ComputeCompressedUseCountPixelCoordinates(uint2 PixelPos, out uint2 Compres
 #endif
 
 CALL_SITE_DEBUGLOC
-uint2 EncodeSpatialAntiAliasingOffset(taa_half2x2 TexelOffset)
+uint2 EncodeSpatialAntiAliasingOffset(tsr_half2x2 TexelOffset)
 {
 	const uint Mask = (1 << SPATIAL_ANTI_ALIASING_OFFSET_BITDEPTH_PER_CHANNEL) - 1;
 	const uint Zero = Mask / 2;
 
-	taa_half Multiply = taa_half(float(Zero) / SPATIAL_ANTI_ALIASING_OFFSET_RANGE);
+	tsr_half Multiply = tsr_half(float(Zero) / SPATIAL_ANTI_ALIASING_OFFSET_RANGE);
 
 	uint2 EncodedTexelOffset = 0;
 	EncodedTexelOffset |= clamp(uint2(int2(round(TexelOffset[0] * Multiply)) + int(1 + Zero)), uint(1), Mask) << 0;
@@ -1078,28 +1078,28 @@ uint2 EncodeSpatialAntiAliasingOffset(taa_half2x2 TexelOffset)
 }
 
 CALL_SITE_DEBUGLOC
-uint EncodeSpatialAntiAliasingOffset(taa_half2 TexelOffset)
+uint EncodeSpatialAntiAliasingOffset(tsr_half2 TexelOffset)
 {
 	return dpv_lo(EncodeSpatialAntiAliasingOffset(dpv_interleave_mono_registers(TexelOffset)));
 }
 
 CALL_SITE_DEBUGLOC
-taa_half2x2 DecodeSpatialAntiAliasingOffset(taa_ushort2 EncodedInputTexelOffset)
+tsr_half2x2 DecodeSpatialAntiAliasingOffset(tsr_ushort2 EncodedInputTexelOffset)
 {
-	const taa_ushort Mask = taa_ushort(1 << SPATIAL_ANTI_ALIASING_OFFSET_BITDEPTH_PER_CHANNEL) - taa_ushort(1);
-	const taa_ushort Zero = Mask / taa_ushort(2);
+	const tsr_ushort Mask = tsr_ushort(1 << SPATIAL_ANTI_ALIASING_OFFSET_BITDEPTH_PER_CHANNEL) - tsr_ushort(1);
+	const tsr_ushort Zero = Mask / tsr_ushort(2);
 	
-	taa_half Multiply = taa_half(SPATIAL_ANTI_ALIASING_OFFSET_RANGE / float(Zero));
+	tsr_half Multiply = tsr_half(SPATIAL_ANTI_ALIASING_OFFSET_RANGE / float(Zero));
 
-	taa_half2x2 TexelOffset;
-	TexelOffset[0] = taa_half2((EncodedInputTexelOffset >> taa_ushort(0                                                )) & Mask) * Multiply - (taa_half(1 + Zero) * Multiply);
-	TexelOffset[1] = taa_half2((EncodedInputTexelOffset >> taa_ushort(SPATIAL_ANTI_ALIASING_OFFSET_BITDEPTH_PER_CHANNEL)) & Mask) * Multiply - (taa_half(1 + Zero) * Multiply);
+	tsr_half2x2 TexelOffset;
+	TexelOffset[0] = tsr_half2((EncodedInputTexelOffset >> tsr_ushort(0                                                )) & Mask) * Multiply - (tsr_half(1 + Zero) * Multiply);
+	TexelOffset[1] = tsr_half2((EncodedInputTexelOffset >> tsr_ushort(SPATIAL_ANTI_ALIASING_OFFSET_BITDEPTH_PER_CHANNEL)) & Mask) * Multiply - (tsr_half(1 + Zero) * Multiply);
 
 	return TexelOffset;
 }
 
 CALL_SITE_DEBUGLOC
-taa_half2 DecodeSpatialAntiAliasingOffset(taa_ushort EncodedInputTexelOffset)
+tsr_half2 DecodeSpatialAntiAliasingOffset(tsr_ushort EncodedInputTexelOffset)
 {
 	return dpv_lo(DecodeSpatialAntiAliasingOffset(dpv_interleave_mono_registers(EncodedInputTexelOffset)));
 }
@@ -1151,16 +1151,16 @@ void DecodeHoleFillingVelocity(uint EncodedHoleFillingVelocity, out float2 Pixel
 }
 
 CALL_SITE_DEBUGLOC
-taa_short2 GetLaneOffsetSign()
+tsr_short2 GetLaneOffsetSign()
 #if CONFIG_BUTTERFLY_KERNEL
 {
 	uint LaneIndex = WaveGetLaneIndex();
 
-	return taa_short(-1) + taa_short2(taa_ushort2(taa_ushort(LaneIndex) << taa_ushort(1), LaneIndex) & taa_ushort(0x2));
+	return tsr_short(-1) + tsr_short2(tsr_ushort2(tsr_ushort(LaneIndex) << tsr_ushort(1), LaneIndex) & tsr_ushort(0x2));
 }
 #else
 {
-	return taa_short(1).xx;
+	return tsr_short(1).xx;
 }
 #endif
 
diff --git a/Engine/Shaders/Private/TemporalSuperResolution/TSRCompareHistory.usf b/Engine/Shaders/Private/TemporalSuperResolution/TSRCompareHistory.usf
--- a/Engine/Shaders/Private/TemporalSuperResolution/TSRCompareHistory.usf
+++ b/Engine/Shaders/Private/TemporalSuperResolution/TSRCompareHistory.usf
@@ -12,21 +12,21 @@
 
 //------------------------------------------------------- PARAMETERS
 
-Texture2D<taa_half> ParallaxRejectionMaskTexture;
-Texture2D<taa_halfC> FilteredInputTexture;
-Texture2D<taa_halfC> FilteredPredictionSceneColorTexture;
+Texture2D<tsr_half> ParallaxRejectionMaskTexture;
+Texture2D<tsr_halfC> FilteredInputTexture;
+Texture2D<tsr_halfC> FilteredPredictionSceneColorTexture;
 
-RWTexture2D<taa_half> HistoryRejectionOutput;
+RWTexture2D<tsr_half> HistoryRejectionOutput;
 
 
 //------------------------------------------------------- LDS
 
-groupshared taa_half SharedRejection[TILE_SIZE * TILE_SIZE];
+groupshared tsr_half SharedRejection[TILE_SIZE * TILE_SIZE];
 
 
 //------------------------------------------------------- FUNCTIONS
 
-taa_halfC TransformColorForClampingBox(taa_halfC Color)
+tsr_halfC TransformColorForClampingBox(tsr_halfC Color)
 {
 	return RGBToYCoCg(Color);
 }
@@ -41,15 +41,15 @@ void MainCS(
 {
 	float4 Debug = 0.0;
 
-	taa_short2 InputPixelPos = (
-		taa_short2(InputPixelPosMin) +
-		taa_short2(GroupId) * taa_short2(TILE_SIZE, TILE_SIZE) +
+	tsr_short2 InputPixelPos = (
+		tsr_short2(InputPixelPosMin) +
+		tsr_short2(GroupId) * tsr_short2(TILE_SIZE, TILE_SIZE) +
 		Map16x16Tile2x2Lane(GroupThreadIndex));
 
-	taa_halfC InputMinColor;
-	taa_halfC InputMaxColor;
-	taa_halfC InputCenterColor;
-	taa_half ParallaxRejectionMask;
+	tsr_halfC InputMinColor;
+	tsr_halfC InputMaxColor;
+	tsr_halfC InputCenterColor;
+	tsr_half ParallaxRejectionMask;
 
 	{
 		InputCenterColor = TransformColorForClampingBox(FilteredInputTexture[InputPixelPos]);
@@ -59,7 +59,7 @@ void MainCS(
 		InputMaxColor = InputCenterColor;
 	}
 	
-	taa_half NeigborMaxInterferenceSeedGradient = 0.0;
+	tsr_half NeigborMaxInterferenceSeedGradient = 0.0;
 	
 	UNROLL
 	for (uint i = 0; i < 9; i++)
@@ -69,24 +69,24 @@ void MainCS(
 			continue;
 		}
 		
-		taa_short2 SampleInputPixelPos = InputPixelPos + ClampPixelOffset(
+		tsr_short2 SampleInputPixelPos = InputPixelPos + ClampPixelOffset(
 			InputPixelPos,
-			taa_short2(kOffsets3x3[i]), taa_short2(kOffsets3x3[i]),
+			tsr_short2(kOffsets3x3[i]), tsr_short2(kOffsets3x3[i]),
 			InputPixelPosMin, InputPixelPosMax);
 
-		taa_halfC Color = TransformColorForClampingBox(FilteredInputTexture[SampleInputPixelPos]);
+		tsr_halfC Color = TransformColorForClampingBox(FilteredInputTexture[SampleInputPixelPos]);
 
 		InputMinColor = min(InputMinColor, Color);
 		InputMaxColor = max(InputMaxColor, Color);
 	}
 	
 
-	taa_halfC PrevColor = TransformColorForClampingBox(FilteredPredictionSceneColorTexture[InputPixelPos]);
+	tsr_halfC PrevColor = TransformColorForClampingBox(FilteredPredictionSceneColorTexture[InputPixelPos]);
 
-	taa_half Rejection = 1.0;
+	tsr_half Rejection = 1.0;
 	#if 1
 	{
-		taa_halfC ClampedPrevColor = clamp(PrevColor, InputMinColor, InputMaxColor);
+		tsr_halfC ClampedPrevColor = clamp(PrevColor, InputMinColor, InputMaxColor);
 
 		Rejection = MeasureRejectionFactor(
 			PrevColor, ClampedPrevColor,
@@ -97,9 +97,9 @@ void MainCS(
 	}
 	#endif
 
-	if (ParallaxRejectionMask < taa_half(PARALLAX_REJECTION_MASK_THRESHOLD))
+	if (ParallaxRejectionMask < tsr_half(PARALLAX_REJECTION_MASK_THRESHOLD))
 	{
-		Rejection = taa_half(1.0);
+		Rejection = tsr_half(1.0);
 	}
 
 	#if CONFIG_OUTPUT_HALF_RES
@@ -116,19 +116,19 @@ void MainCS(
 			Rejection = min(Rejection, SharedRejection[GroupThreadIndex ^ ButterflySwap]);
 		}
 
-		bool bIsValid = all(InputPixelPos <= taa_short2(InputPixelPosMax)) && all((taa_ushort2(InputPixelPos - taa_short2(InputPixelPosMin)) % taa_ushort(2)) == taa_ushort(0));
+		bool bIsValid = all(InputPixelPos <= tsr_short2(InputPixelPosMax)) && all((tsr_ushort2(InputPixelPos - tsr_short2(InputPixelPosMin)) % tsr_ushort(2)) == tsr_ushort(0));
 
-		taa_ushort2 OutputPixelPos = taa_ushort2(InputPixelPos - taa_short2(InputPixelPosMin)) / taa_ushort(2);
-		OutputPixelPos.x = bIsValid ? OutputPixelPos.x : ~taa_ushort(0);
+		tsr_ushort2 OutputPixelPos = tsr_ushort2(InputPixelPos - tsr_short2(InputPixelPosMin)) / tsr_ushort(2);
+		OutputPixelPos.x = bIsValid ? OutputPixelPos.x : ~tsr_ushort(0);
 
 		HistoryRejectionOutput[OutputPixelPos] = Rejection;
 	}
 	#else
 	{
-		bool bIsValid = all(InputPixelPos <= taa_short2(InputPixelPosMax));
+		bool bIsValid = all(InputPixelPos <= tsr_short2(InputPixelPosMax));
 
-		taa_ushort2 OutputPixelPos = taa_ushort2(InputPixelPos - taa_short2(InputPixelPosMin));
-		OutputPixelPos.x = bIsValid ? OutputPixelPos.x : ~taa_ushort(0);
+		tsr_ushort2 OutputPixelPos = tsr_ushort2(InputPixelPos - tsr_short2(InputPixelPosMin));
+		OutputPixelPos.x = bIsValid ? OutputPixelPos.x : ~tsr_ushort(0);
 
 		HistoryRejectionOutput[OutputPixelPos] = Rejection;
 	}
@@ -136,7 +136,7 @@ void MainCS(
 
 	#if DEBUG_OUTPUT
 	{
-		taa_short2 OutputPixelPos = InvalidateOutputPixelPos(taa_ushort2(InputPixelPos), InputPixelPosMax);
+		tsr_short2 OutputPixelPos = InvalidateOutputPixelPos(tsr_ushort2(InputPixelPos), InputPixelPosMax);
 		DebugOutput[OutputPixelPos] = Debug;
 	}
 	#endif
diff --git a/Engine/Shaders/Private/TemporalSuperResolution/TSRCompareTranslucency.usf b/Engine/Shaders/Private/TemporalSuperResolution/TSRCompareTranslucency.usf
--- a/Engine/Shaders/Private/TemporalSuperResolution/TSRCompareTranslucency.usf
+++ b/Engine/Shaders/Private/TemporalSuperResolution/TSRCompareTranslucency.usf
@@ -53,39 +53,39 @@ float TranslucencyHighlightLuminance;
 
 FScreenTransform ScreenPosToPrevTranslucencyTextureUV;
 
-Texture2D<taa_half4> TranslucencyTexture;
+Texture2D<tsr_half4> TranslucencyTexture;
 Texture2D<float2> DilatedVelocityTexture;
-Texture2D<taa_half4> PrevTranslucencyTexture;
+Texture2D<tsr_half4> PrevTranslucencyTexture;
 
-RWTexture2D<taa_half> TranslucencyRejectionOutput;
+RWTexture2D<tsr_half> TranslucencyRejectionOutput;
 
 
 //------------------------------------------------------- ENTRY POINT
 
-taa_half4 TransformToComparisonColorSpace(taa_half4 Color)
+tsr_half4 TransformToComparisonColorSpace(tsr_half4 Color)
 {
-	return taa_half4(RGBToYCoCg(Color.rgb), Color.a);
+	return tsr_half4(RGBToYCoCg(Color.rgb), Color.a);
 }
 
-taa_half MeasureTranslucencyRejectionFactor(
-	taa_half4 PrevTranslucency,
-	taa_half4 Translucency,
-	taa_half4 TranslucencyMin,
-	taa_half4 TranslucencyMax)
+tsr_half MeasureTranslucencyRejectionFactor(
+	tsr_half4 PrevTranslucency,
+	tsr_half4 Translucency,
+	tsr_half4 TranslucencyMin,
+	tsr_half4 TranslucencyMax)
 {
-	taa_half IsHighlight = saturate(PrevTranslucency.x - taa_half(TranslucencyHighlightLuminance));
+	tsr_half IsHighlight = saturate(PrevTranslucency.x - tsr_half(TranslucencyHighlightLuminance));
 
-	taa_half4 ClampedPrevTranslucency = clamp(PrevTranslucency, TranslucencyMin, TranslucencyMax);
+	tsr_half4 ClampedPrevTranslucency = clamp(PrevTranslucency, TranslucencyMin, TranslucencyMax);
 	
-	taa_half BackbufferQuantizationError = 0.0; //MeasureBackbufferLDRQuantizationErrorFromHDRLuma(max(PrevTranslucency[0], Translucency[0]));
+	tsr_half BackbufferQuantizationError = 0.0; //MeasureBackbufferLDRQuantizationErrorFromHDRLuma(max(PrevTranslucency[0], Translucency[0]));
 	
 	#if CONFIG_SCENE_COLOR_APLHA
-		taa_half Rejection = MeasureRejectionFactor(
+		tsr_half Rejection = MeasureRejectionFactor(
 			PrevTranslucency, ClampedPrevTranslucency,
 			Translucency, TranslucencyMin, TranslucencyMax,
 			BackbufferQuantizationError);
 	#else
-		taa_half Rejection = MeasureRejectionFactor(
+		tsr_half Rejection = MeasureRejectionFactor(
 			PrevTranslucency.xyz, ClampedPrevTranslucency.xyz,
 			Translucency.xyz, TranslucencyMin.xyz, TranslucencyMax.xyz,
 			BackbufferQuantizationError);
@@ -105,7 +105,7 @@ void MainCS(
 {
 	float4 Debug = 0;
 	
-	taa_short2 InputPixelPos = taa_short2(
+	tsr_short2 InputPixelPos = tsr_short2(
 		ZOrder2D(GroupThreadIndex, uint(log2(float(TILE_SIZE)))) +
 		GroupId * uint2(TILE_SIZE, TILE_SIZE));
 
@@ -113,19 +113,19 @@ void MainCS(
 	float2 EncodedVelocity = DilatedVelocityTexture[InputPixelPos];
 	
 	// Sample current frame translucency's
-	taa_half4 Translucency;
-	taa_half4 TranslucencyMin;
-	taa_half4 TranslucencyMax;
+	tsr_half4 Translucency;
+	tsr_half4 TranslucencyMin;
+	tsr_half4 TranslucencyMax;
 	{
 		UNROLL_N(CONFIG_SAMPLES)
 		for (uint SampleId = 0; SampleId < CONFIG_SAMPLES; SampleId++)
 		{
-			taa_short2 SampleInputPixelPos;
+			tsr_short2 SampleInputPixelPos;
 			
 			// TODO: handle SeparateTranslucencyScreenPercentage
 			#if CONFIG_SAMPLES == 9
 			{
-				taa_short2 iPixelOffset = taa_short2(kOffsets3x3[kSquareIndexes3x3[SampleId]]);
+				tsr_short2 iPixelOffset = tsr_short2(kOffsets3x3[kSquareIndexes3x3[SampleId]]);
 				
 				SampleInputPixelPos = AddAndClampPixelOffset(
 					InputPixelPos,
@@ -134,7 +134,7 @@ void MainCS(
 			}
 			#elif CONFIG_SAMPLES == 5
 			{
-				taa_short2 iPixelOffset = taa_short2(kOffsets3x3[kPlusIndexes3x3[SampleId]]);
+				tsr_short2 iPixelOffset = tsr_short2(kOffsets3x3[kPlusIndexes3x3[SampleId]]);
 					
 				SampleInputPixelPos = AddAndClampPixelOffset(
 					InputPixelPos,
@@ -145,7 +145,7 @@ void MainCS(
 				#error Unknown sample count
 			#endif
 
-			taa_half4 TranslucencyColor = TransformToComparisonColorSpace(TranslucencyTexture[SampleInputPixelPos]);
+			tsr_half4 TranslucencyColor = TransformToComparisonColorSpace(TranslucencyTexture[SampleInputPixelPos]);
 
 			if (SampleId == 0)
 			{
@@ -165,14 +165,14 @@ void MainCS(
 		}
 
 		#if CONFIG_KERNEL
-			Translucency *= rcp(taa_half(CONFIG_SAMPLES));
+			Translucency *= rcp(tsr_half(CONFIG_SAMPLES));
 		#endif
 	}
 
 	// Sample previous frame's translucency
-	taa_half4 PrevTranslucency;
-	taa_half4 PrevTranslucencyMin;
-	taa_half4 PrevTranslucencyMax;
+	tsr_half4 PrevTranslucency;
+	tsr_half4 PrevTranslucencyMin;
+	tsr_half4 PrevTranslucencyMax;
 	{
 		float2 PrevScreenPos = ScreenPos - DecodeVelocityFromTexture(float4(EncodedVelocity, 0.0, 0.0)).xy;
 		float2 PrevTranslucencyTextureUV = ApplyScreenTransform(PrevScreenPos, ScreenPosToPrevTranslucencyTextureUV);
@@ -180,18 +180,18 @@ void MainCS(
 		UNROLL_N(CONFIG_SAMPLES)
 		for (uint SampleId = 0; SampleId < CONFIG_SAMPLES; SampleId++)
 		{
-			taa_half2 PixelOffset;
+			tsr_half2 PixelOffset;
 			
 			// TODO: handle r.SeparateTranslucencyScreenPercentage
 			#if CONFIG_SAMPLES == 9
 			{
-				taa_short2 iPixelOffset = taa_short2(kOffsets3x3[kSquareIndexes3x3[SampleId]]);
-				PixelOffset = taa_half2(iPixelOffset);
+				tsr_short2 iPixelOffset = tsr_short2(kOffsets3x3[kSquareIndexes3x3[SampleId]]);
+				PixelOffset = tsr_half2(iPixelOffset);
 			}
 			#elif CONFIG_SAMPLES == 5
 			{
-				taa_short2 iPixelOffset = taa_short2(kOffsets3x3[kPlusIndexes3x3[SampleId]]);
-				PixelOffset = taa_half2(iPixelOffset);
+				tsr_short2 iPixelOffset = tsr_short2(kOffsets3x3[kPlusIndexes3x3[SampleId]]);
+				PixelOffset = tsr_half2(iPixelOffset);
 			}
 			#else
 				#error Unknown sample count
@@ -202,7 +202,7 @@ void MainCS(
 				PrevTranslucencyInfo_UVViewportBilinearMin,
 				PrevTranslucencyInfo_UVViewportBilinearMax);
 		
-			taa_half4 PrevTranslucencyColor = TransformToComparisonColorSpace(
+			tsr_half4 PrevTranslucencyColor = TransformToComparisonColorSpace(
 				PrevTranslucencyTexture.SampleLevel(GlobalPointClampedSampler, SampleTextureUV, 0));
 
 			if (SampleId == 0)
@@ -223,25 +223,25 @@ void MainCS(
 		}
 
 		#if CONFIG_KERNEL
-			PrevTranslucency *= rcp(taa_half(CONFIG_SAMPLES));
+			PrevTranslucency *= rcp(tsr_half(CONFIG_SAMPLES));
 		#endif
 
-		PrevTranslucency    *= taa_half(PrevTranslucencyPreExposureCorrection);
-		PrevTranslucencyMin *= taa_half(PrevTranslucencyPreExposureCorrection);
-		PrevTranslucencyMax *= taa_half(PrevTranslucencyPreExposureCorrection);
+		PrevTranslucency    *= tsr_half(PrevTranslucencyPreExposureCorrection);
+		PrevTranslucencyMin *= tsr_half(PrevTranslucencyPreExposureCorrection);
+		PrevTranslucencyMax *= tsr_half(PrevTranslucencyPreExposureCorrection);
 	}
 
-	taa_half Rejection = 1.0;
+	tsr_half Rejection = 1.0;
 	#if 1
 	{
 		// TODO: reject with alpha too? Might not need since translucency is premultiplied RGBA 
-		taa_half ForwardRejection = MeasureTranslucencyRejectionFactor(
+		tsr_half ForwardRejection = MeasureTranslucencyRejectionFactor(
 			PrevTranslucency,
 			Translucency,
 			TranslucencyMin,
 			TranslucencyMax);
 		
-		taa_half BackwardRejection = MeasureTranslucencyRejectionFactor(
+		tsr_half BackwardRejection = MeasureTranslucencyRejectionFactor(
 			Translucency,
 			PrevTranslucency,
 			PrevTranslucencyMin,
@@ -251,7 +251,7 @@ void MainCS(
 	}
 	#endif
 
-	if (all(InputPixelPos < taa_short2(InputInfo_ViewportMax)))
+	if (all(InputPixelPos < tsr_short2(InputInfo_ViewportMax)))
 	{
 		TranslucencyRejectionOutput[InputPixelPos] = Rejection;
 
diff --git a/Engine/Shaders/Private/TemporalSuperResolution/TSRDebugHistory.usf b/Engine/Shaders/Private/TemporalSuperResolution/TSRDebugHistory.usf
--- a/Engine/Shaders/Private/TemporalSuperResolution/TSRDebugHistory.usf
+++ b/Engine/Shaders/Private/TemporalSuperResolution/TSRDebugHistory.usf
@@ -11,15 +11,15 @@
 
 //------------------------------------------------------- PARAMETERS
 
-Texture2D<taa_half3> History_Textures_0;
-Texture2D<taa_half3> History_Textures_1;
-Texture2D<taa_half2> History_Textures_2;
-Texture2D<taa_subpixel_details> History_Textures_3;
+Texture2D<tsr_half3> History_Textures_0;
+Texture2D<tsr_half3> History_Textures_1;
+Texture2D<tsr_half2> History_Textures_2;
+Texture2D<tsr_subpixel_details> History_Textures_3;
 
-Texture2D<taa_half3> PrevHistory_Textures_0;
-Texture2D<taa_half3> PrevHistory_Textures_1;
-Texture2D<taa_half2> PrevHistory_Textures_2;
-Texture2D<taa_subpixel_details> PrevHistory_Textures_3;
+Texture2D<tsr_half3> PrevHistory_Textures_0;
+Texture2D<tsr_half3> PrevHistory_Textures_1;
+Texture2D<tsr_half2> PrevHistory_Textures_2;
+Texture2D<tsr_subpixel_details> PrevHistory_Textures_3;
 
 
 //------------------------------------------------------- FUNCTION
@@ -45,7 +45,7 @@ float Median(float A, float B, float C)
 }
 
 
-taa_subpixel_payload ReadSubpixelDetails(uint2 SubpixelPos)
+tsr_subpixel_payload ReadSubpixelDetails(uint2 SubpixelPos)
 {
 	uint2 HistoryPixelPos = SubpixelPos / SUB_PIXEL_GRID_SIZE;
 
@@ -53,9 +53,9 @@ taa_subpixel_payload ReadSubpixelDetails(uint2 SubpixelPos)
 	uint SubpixelId = dot(SubpixelRelPos, uint2(1, SUB_PIXEL_GRID_SIZE));
 	uint SubpixelShift = SubpixelId * SUB_PIXEL_BIT_COUNT;
 	
-	taa_subpixel_details RawHistory3 = History_Textures_3[HistoryPixelPos];
+	tsr_subpixel_details RawHistory3 = History_Textures_3[HistoryPixelPos];
 
-	return taa_subpixel_payload(RawHistory3 >> SubpixelShift) & SUB_PIXEL_BIT_MASK;
+	return tsr_subpixel_payload(RawHistory3 >> SubpixelShift) & SUB_PIXEL_BIT_MASK;
 }
 
 
@@ -82,18 +82,18 @@ void MainCS(
 	
 	{
 		#if 1
-			taa_half3 RawHistory0 = History_Textures_0[HistoryPixelPos];
-			taa_half3 RawHistory1 = History_Textures_1[HistoryPixelPos];
-			taa_half2 RawHistory2 = History_Textures_2[HistoryPixelPos];
-			taa_subpixel_details RawHistory3 = History_Textures_3[HistoryPixelPos];
+			tsr_half3 RawHistory0 = History_Textures_0[HistoryPixelPos];
+			tsr_half3 RawHistory1 = History_Textures_1[HistoryPixelPos];
+			tsr_half2 RawHistory2 = History_Textures_2[HistoryPixelPos];
+			tsr_subpixel_details RawHistory3 = History_Textures_3[HistoryPixelPos];
 		#else
-			taa_half3 RawHistory0 = PrevHistory_Textures_0[HistoryPixelPos];
-			taa_half3 RawHistory1 = PrevHistory_Textures_1[HistoryPixelPos];
-			taa_half2 RawHistory2 = PrevHistory_Textures_2[HistoryPixelPos];
-			taa_subpixel_details RawHistory3 = PrevHistory_Textures_3[HistoryPixelPos];
+			tsr_half3 RawHistory0 = PrevHistory_Textures_0[HistoryPixelPos];
+			tsr_half3 RawHistory1 = PrevHistory_Textures_1[HistoryPixelPos];
+			tsr_half2 RawHistory2 = PrevHistory_Textures_2[HistoryPixelPos];
+			tsr_subpixel_details RawHistory3 = PrevHistory_Textures_3[HistoryPixelPos];
 		#endif
 
-		taa_subpixel_payload SubpixelDetails = taa_subpixel_payload(RawHistory3 >> SubpixelShift) & SUB_PIXEL_BIT_MASK;
+		tsr_subpixel_payload SubpixelDetails = tsr_subpixel_payload(RawHistory3 >> SubpixelShift) & SUB_PIXEL_BIT_MASK;
 
 		Debug.r = UncompressParallaxFactor(SubpixelDetails);
 		Debug.b = float(SubpixelId);
diff --git a/Engine/Shaders/Private/TemporalSuperResolution/TSRDecimateHistory.usf b/Engine/Shaders/Private/TemporalSuperResolution/TSRDecimateHistory.usf
--- a/Engine/Shaders/Private/TemporalSuperResolution/TSRDecimateHistory.usf
+++ b/Engine/Shaders/Private/TemporalSuperResolution/TSRDecimateHistory.usf
@@ -17,42 +17,42 @@
 float4x4 RotationalClipToPrevClip;
 float3 OutputQuantizationError;
 
-Texture2D<taa_halfC> InputSceneColorTexture;
+Texture2D<tsr_halfC> InputSceneColorTexture;
 Texture2D<float2> DilatedVelocityTexture;
 Texture2D<float> ClosestDepthTexture;
 Texture2D<uint> PrevUseCountTexture;
 Texture2D<uint> PrevClosestDepthTexture;
 Texture2D<uint> ParallaxFactorTexture;
 
-Texture2D<taa_halfC> PrevHistory_LowFrequency;
-Texture2D<taa_halfC> PrevHistory_HighFrequency;
-Texture2D<taa_half2> PrevHistory_Metadata;
-Texture2D<taa_half4> PrevHistory_Translucency;
-Texture2D<taa_subpixel_details> PrevHistory_SubpixelDetails;
+Texture2D<tsr_halfC> PrevHistory_LowFrequency;
+Texture2D<tsr_halfC> PrevHistory_HighFrequency;
+Texture2D<tsr_half2> PrevHistory_Metadata;
+Texture2D<tsr_half4> PrevHistory_Translucency;
+Texture2D<tsr_subpixel_details> PrevHistory_SubpixelDetails;
 
-RWTexture2D<taa_halfC> HalfResSceneColorOutput;
-RWTexture2D<taa_halfC> HalfResPredictionSceneColorOutput;
-RWTexture2D<taa_half> HalfResParallaxRejectionMaskOutput;
+RWTexture2D<tsr_halfC> HalfResSceneColorOutput;
+RWTexture2D<tsr_halfC> HalfResPredictionSceneColorOutput;
+RWTexture2D<tsr_half> HalfResParallaxRejectionMaskOutput;
 
-RWTexture2D<taa_halfC> PredictionSceneColorOutput;
-RWTexture2D<taa_half> ParallaxRejectionMaskOutput;
+RWTexture2D<tsr_halfC> PredictionSceneColorOutput;
+RWTexture2D<tsr_half> ParallaxRejectionMaskOutput;
 
 
 //------------------------------------------------------- LDS
 
-groupshared taa_halfC SharedInputColor[TILE_SIZE * TILE_SIZE];
+groupshared tsr_halfC SharedInputColor[TILE_SIZE * TILE_SIZE];
 
 #if CONFIG_SCENE_COLOR_APLHA
-groupshared taa_halfC SharedHistoryColor[TILE_SIZE * TILE_SIZE];
-groupshared taa_half  SharedMask[TILE_SIZE * TILE_SIZE];
+groupshared tsr_halfC SharedHistoryColor[TILE_SIZE * TILE_SIZE];
+groupshared tsr_half  SharedMask[TILE_SIZE * TILE_SIZE];
 #else
-groupshared taa_half4 SharedHistoryColorAndMask[TILE_SIZE * TILE_SIZE];
+groupshared tsr_half4 SharedHistoryColorAndMask[TILE_SIZE * TILE_SIZE];
 #endif
 
 
 //------------------------------------------------------- ENTRY POINT
 
-taa_half3 TransformToComparisonColorSpace(taa_half3 Color)
+tsr_half3 TransformToComparisonColorSpace(tsr_half3 Color)
 {
 	return Color;
 	//return log2(Color);
@@ -81,13 +81,13 @@ void MainCS(
 	
 	float2 EncodedVelocity = DilatedVelocityTexture[InputPixelPos];
 	float DeviceZ = ClosestDepthTexture[InputPixelPos];
-	taa_half ParallaxFactor = UncompressParallaxFactor((ParallaxFactorTexture[InputPixelPos] >> 4) & 0xF);
-	taa_halfC InputSceneColor = InputSceneColorTexture[InputPixelPos];
+	tsr_half ParallaxFactor = UncompressParallaxFactor((ParallaxFactorTexture[InputPixelPos] >> 4) & 0xF);
+	tsr_halfC InputSceneColor = InputSceneColorTexture[InputPixelPos];
 
 	float WorldDepth = ConvertFromDeviceZ(DeviceZ);
 	
 	// Compute the parralax rejection mask
-	taa_half ParallaxRejectionMask = 0.0;
+	tsr_half ParallaxRejectionMask = 0.0;
 	{
 		float2 PrevScreenPos = ScreenPos - DecodeVelocityFromTexture(float4(EncodedVelocity, 0.0, 0.0)).xy;
 		float2 PrevInputBufferUV = (InputInfo_ScreenPosToViewportScale * PrevScreenPos + InputInfo_ScreenPosToViewportBias) * InputInfo_ExtentInverse;
@@ -115,7 +115,7 @@ void MainCS(
 			float HistoryClosestDeviceZ = f16tof32(PrevClosestDepthTexture[PixelPos] >> VELOCITY_HOLE_FILLING_BITS);
 			float HistoryClosestWorldDepth = ConvertFromDeviceZ(HistoryClosestDeviceZ);
 
-			taa_half DepthRejection;
+			tsr_half DepthRejection;
 			{
 				const float WorldDepthToPixelWorldRadius = GetTanHalfFieldOfView().x * InputInfo_ViewportSizeInverse.x;
 				const float PixelDepthError = 3.0;
@@ -124,14 +124,14 @@ void MainCS(
 				float WorldDepthEpsilon = (WorldDepthToPixelWorldRadius * RadiusToDiamater * PixelDepthError) * HistoryClosestWorldDepth; 
 				float DeltaDepth = abs(HistoryClosestWorldDepth - WorldDepth);
 
-				DepthRejection = saturate(taa_half(2.0) - taa_half(DeltaDepth / WorldDepthEpsilon));
+				DepthRejection = saturate(tsr_half(2.0) - tsr_half(DeltaDepth / WorldDepthEpsilon));
 			}
 
-			taa_half HistoryGhostingRejection = saturate(taa_half(1.0) - taa_half(3.0) * abs(taa_half(HistoryUseCount) * rcp(taa_half(PREV_USE_COUNT_QUANTIZATION)) - taa_half(1.0)));
+			tsr_half HistoryGhostingRejection = saturate(tsr_half(1.0) - tsr_half(3.0) * abs(tsr_half(HistoryUseCount) * rcp(tsr_half(PREV_USE_COUNT_QUANTIZATION)) - tsr_half(1.0)));
 
 			FLATTEN
 			if (bValidPixelPos)
-				ParallaxRejectionMask += taa_half(BilinearWeight) * max(DepthRejection, HistoryGhostingRejection);
+				ParallaxRejectionMask += tsr_half(BilinearWeight) * max(DepthRejection, HistoryGhostingRejection);
 		}
 	}
 
@@ -151,16 +151,16 @@ void MainCS(
 	float2 HistoryBufferUV = HistoryPPCo * PrevHistoryInfo_ExtentInverse;
 
 	// Vector in pixel between pixel K -> O.
-	taa_half2 dKO = taa_half2(HistoryPPCo - HistoryPPCk);
+	tsr_half2 dKO = tsr_half2(HistoryPPCo - HistoryPPCk);
 	
 	// Sample the history
-	taa_halfC DecimatedColor;
-	taa_half HistoryMaxRejectionFactor;
+	tsr_halfC DecimatedColor;
+	tsr_half HistoryMaxRejectionFactor;
 
-	taa_half ClosestInterferenceSeed;
-	taa_half IntegratedInterferenceSeed;
-	taa_half MinInterferenceSeed;
-	taa_half MaxInterferenceSeed;
+	tsr_half ClosestInterferenceSeed;
+	tsr_half IntegratedInterferenceSeed;
+	tsr_half MinInterferenceSeed;
+	tsr_half MaxInterferenceSeed;
 	{
 		float2 KernelHistoryBufferUV = clamp(
 			HistoryPPCk * PrevHistoryInfo_ExtentInverse,
@@ -187,7 +187,7 @@ void MainCS(
 					PrevHistoryInfo_UVViewportBilinearMin,
 					PrevHistoryInfo_UVViewportBilinearMax);
 				
-				taa_halfC HistorySample = PrevHistory_LowFrequency.SampleLevel(GlobalBilinearClampedSampler, SampleBufferUV, 0);
+				tsr_halfC HistorySample = PrevHistory_LowFrequency.SampleLevel(GlobalBilinearClampedSampler, SampleBufferUV, 0);
 
 				DecimatedColor += HistorySample;
 			}
@@ -204,8 +204,8 @@ void MainCS(
 			UNROLL_N(SUB_PIXEL_COUNT)
 			for (uint SubpixelId = 0; SubpixelId < SUB_PIXEL_COUNT; SubpixelId++)
 			{
-				taa_subpixel_payload SubpixelPayload = GetSubpixelPayload(SubpixelNeigborhood, SubpixelId);
-				taa_subpixel_payload EncodedRejectionFactor = (SubpixelPayload >> SUB_PIXEL_PARALLAX_FACTOR_BIT_OFFSET) & SUB_PIXEL_PARALLAX_FACTOR_BIT_MASK;
+				tsr_subpixel_payload SubpixelPayload = GetSubpixelPayload(SubpixelNeigborhood, SubpixelId);
+				tsr_subpixel_payload EncodedRejectionFactor = (SubpixelPayload >> SUB_PIXEL_PARALLAX_FACTOR_BIT_OFFSET) & SUB_PIXEL_PARALLAX_FACTOR_BIT_MASK;
 
 				EncodedHistoryMaxRejectionFactor = max(EncodedHistoryMaxRejectionFactor, EncodedRejectionFactor);
 			}
@@ -216,7 +216,7 @@ void MainCS(
 	
 	// Reject subpixel details that might be in history but not input just based on their parrallax
 	{
-		taa_half HistorySubdetailRejectionFactor = max(HistoryMaxRejectionFactor / ParallaxFactor, taa_half(1.0));
+		tsr_half HistorySubdetailRejectionFactor = max(HistoryMaxRejectionFactor / ParallaxFactor, tsr_half(1.0));
 
 		float4 ThisClip = float4(ScreenPos, DeviceZ, 1);
 		float4 PrevClip = mul(ThisClip, View.ClipToPrevClip);
@@ -225,13 +225,13 @@ void MainCS(
 		float4 RotationalPrevClip = mul(ThisClip, RotationalClipToPrevClip);
 		float2 RotationalPrevScreen = RotationalPrevClip.xy / RotationalPrevClip.w;
 			
-		taa_half PixelParallaxVelocity = taa_half(0.5) * length(taa_half2(RotationalPrevScreen - PrevScreen) * taa_half2(InputInfo_ViewportSize));
-		taa_half HistorySubdetailAdditionalParralax = (HistorySubdetailRejectionFactor - taa_half(1.0)) * PixelParallaxVelocity;
+		tsr_half PixelParallaxVelocity = tsr_half(0.5) * length(tsr_half2(RotationalPrevScreen - PrevScreen) * tsr_half2(InputInfo_ViewportSize));
+		tsr_half HistorySubdetailAdditionalParralax = (HistorySubdetailRejectionFactor - tsr_half(1.0)) * PixelParallaxVelocity;
 
 		// PixelParallaxVelocity can be +inf and generate NaN if HistorySubdetailRejectionFactor = 1.0;
-		HistorySubdetailAdditionalParralax = -min(-HistorySubdetailAdditionalParralax, taa_half(0.0));
+		HistorySubdetailAdditionalParralax = -min(-HistorySubdetailAdditionalParralax, tsr_half(0.0));
 
-		taa_half HistoryParallaxRejectionMask = saturate(taa_half(1.0) - HistorySubdetailAdditionalParralax);
+		tsr_half HistoryParallaxRejectionMask = saturate(tsr_half(1.0) - HistorySubdetailAdditionalParralax);
 
 		ParallaxRejectionMask = min(ParallaxRejectionMask, HistoryParallaxRejectionMask);
 	}
@@ -244,7 +244,7 @@ void MainCS(
 		(PrevOutputScreenPos.y > 1.0))
 	//if (IsOffScreen2(bCameraCut, PrevOutputScreenPos, /* ParallaxRejectionMask = */ 1.0))
 	{
-		ParallaxRejectionMask = taa_half(0.0);
+		ParallaxRejectionMask = tsr_half(0.0);
 	}
 
 	
@@ -255,16 +255,16 @@ void MainCS(
 			SharedHistoryColor[GroupThreadIndex] = DecimatedColor;
 			SharedMask[GroupThreadIndex] = ParallaxRejectionMask;
 		#else
-			SharedHistoryColorAndMask[GroupThreadIndex] = taa_half4(DecimatedColor, ParallaxRejectionMask);
+			SharedHistoryColorAndMask[GroupThreadIndex] = tsr_half4(DecimatedColor, ParallaxRejectionMask);
 		#endif
 
 		GroupMemoryBarrierWithGroupSync();
 
 		bool bCheckerboard = (((InputPixelPos.x ^ InputPixelPos.y) >> 1) & 0x1) != 0;
 
-		taa_halfC HalfResInputSceneColor = InputSceneColor;
-		taa_halfC HalfResDecimatedColor = DecimatedColor;
-		taa_half HalfResParallaxRejectionMask = ParallaxRejectionMask;
+		tsr_halfC HalfResInputSceneColor = InputSceneColor;
+		tsr_halfC HalfResDecimatedColor = DecimatedColor;
+		tsr_half HalfResParallaxRejectionMask = ParallaxRejectionMask;
 
 		UNROLL
 		for (uint i = 1; i < 4; i++)
@@ -289,13 +289,13 @@ void MainCS(
 			#endif
 		}
 
-		HalfResDecimatedColor *= taa_half(0.25);
-		HalfResParallaxRejectionMask *= taa_half(0.25);
+		HalfResDecimatedColor *= tsr_half(0.25);
+		HalfResParallaxRejectionMask *= tsr_half(0.25);
 
 		#if CONFIG_ENABLE_STOCASTIC_QUANTIZATION
 		{
 			uint2 Random = Rand3DPCG16(int3((InputPixelPos - InputInfo_ViewportMin) / 2, View.StateFrameIndexMod8)).xy;
-			taa_half E = taa_half(Hammersley16(0, 1, Random).x);
+			tsr_half E = tsr_half(Hammersley16(0, 1, Random).x);
 
 			HalfResInputSceneColor = QuantizeForFloatRenderTarget(HalfResInputSceneColor, E, OutputQuantizationError);
 			HalfResDecimatedColor = QuantizeForFloatRenderTarget(HalfResDecimatedColor, E, OutputQuantizationError);
@@ -314,7 +314,7 @@ void MainCS(
 	#if CONFIG_ENABLE_STOCASTIC_QUANTIZATION
 	{
 		uint2 Random = Rand3DPCG16(int3(InputPixelPos - InputInfo_ViewportMin, View.StateFrameIndexMod8)).xy;
-		taa_half E = taa_half(Hammersley16(0, 1, Random).x);
+		tsr_half E = tsr_half(Hammersley16(0, 1, Random).x);
 
 		DecimatedColor = QuantizeForFloatRenderTarget(DecimatedColor, E, OutputQuantizationError);
 	}
diff --git a/Engine/Shaders/Private/TemporalSuperResolution/TSRDilateRejection.usf b/Engine/Shaders/Private/TemporalSuperResolution/TSRDilateRejection.usf
--- a/Engine/Shaders/Private/TemporalSuperResolution/TSRDilateRejection.usf
+++ b/Engine/Shaders/Private/TemporalSuperResolution/TSRDilateRejection.usf
@@ -11,8 +11,8 @@
 
 //------------------------------------------------------- PARAMETERS
 
-Texture2D<taa_half> HistoryRejectionTexture;
-RWTexture2D<taa_half> DilatedHistoryRejectionOutput;
+Texture2D<tsr_half> HistoryRejectionTexture;
+RWTexture2D<tsr_half> DilatedHistoryRejectionOutput;
 
 
 //------------------------------------------------------- ENTRY POINT
@@ -24,12 +24,12 @@ void MainCS(
 {
 	float4 Debug = 0.0;
 	
-	taa_short2 RejectionPixelPos = (
-		taa_short2(GroupId) * taa_short2(TILE_SIZE, TILE_SIZE) +
+	tsr_short2 RejectionPixelPos = (
+		tsr_short2(GroupId) * tsr_short2(TILE_SIZE, TILE_SIZE) +
 		Map8x8Tile2x2Lane(GroupThreadIndex));
 
 	// Fetch center of the kernel
-	taa_half HistoryRejectionInfo = HistoryRejectionTexture[RejectionPixelPos];
+	tsr_half HistoryRejectionInfo = HistoryRejectionTexture[RejectionPixelPos];
 
 	const int KernelSize = 2;
 
@@ -40,26 +40,26 @@ void MainCS(
 		UNROLL
 		for (int y = -KernelSize; y <= KernelSize; y++)
 		{
-			taa_short2 Offset = taa_short2(x, y);
+			tsr_short2 Offset = tsr_short2(x, y);
 		
 			if (all(Offset == 0))
 			{
 				continue;
 			}
 			
-			taa_short2 SampleInputPixelPos = RejectionPixelPos + ClampPixelOffset(
+			tsr_short2 SampleInputPixelPos = RejectionPixelPos + ClampPixelOffset(
 				RejectionPixelPos,
 				Offset, Offset,
 				RejectionInfo_ViewportMin, RejectionInfo_ViewportMax - 1);
 
-			taa_half SampleRejectionInfo = HistoryRejectionTexture[SampleInputPixelPos];
+			tsr_half SampleRejectionInfo = HistoryRejectionTexture[SampleInputPixelPos];
 
 			HistoryRejectionInfo = min(HistoryRejectionInfo, SampleRejectionInfo);
 		}
 	}
 	
-	bool bIsValid = all(RejectionPixelPos < taa_short2(RejectionInfo_ViewportMax));
-	RejectionPixelPos.x = bIsValid ? RejectionPixelPos.x : ~taa_ushort(0);
+	bool bIsValid = all(RejectionPixelPos < tsr_short2(RejectionInfo_ViewportMax));
+	RejectionPixelPos.x = bIsValid ? RejectionPixelPos.x : ~tsr_ushort(0);
 
 	DilatedHistoryRejectionOutput[RejectionPixelPos] = HistoryRejectionInfo;
 
diff --git a/Engine/Shaders/Private/TemporalSuperResolution/TSRDilateVelocity.usf b/Engine/Shaders/Private/TemporalSuperResolution/TSRDilateVelocity.usf
--- a/Engine/Shaders/Private/TemporalSuperResolution/TSRDilateVelocity.usf
+++ b/Engine/Shaders/Private/TemporalSuperResolution/TSRDilateVelocity.usf
@@ -267,7 +267,7 @@ void MainCS(
 		//DilatedVelocityOutput[OutputPixelPos] = EncodedVelocity.xy; //EncodeVelocityToTexture(float3(FinalScreenVelocity, 0.0)).xy;
 		DilatedVelocityOutput[OutputPixelPos] = EncodeVelocityToTexture(float3(FinalScreenVelocity, 0.0)).xy;
 		ClosestDepthOutput[OutputPixelPos] = ClosestDeviceZ;
-		ParallaxFactorOutput[OutputPixelPos] = CompressParallaxFactor(taa_half(ParallaxFactorForHistory)) | (CompressParallaxFactor(taa_half(ParallaxFactorForHistoryRejection)) << 4);
+		ParallaxFactorOutput[OutputPixelPos] = CompressParallaxFactor(tsr_half(ParallaxFactorForHistory)) | (CompressParallaxFactor(tsr_half(ParallaxFactorForHistoryRejection)) << 4);
 
 		// Output motion blur velocity flatten
 		#if DIM_MOTION_BLUR_DIRECTIONS > 0
diff --git a/Engine/Shaders/Private/TemporalSuperResolution/TSRFilterAntiAliasing.usf b/Engine/Shaders/Private/TemporalSuperResolution/TSRFilterAntiAliasing.usf
--- a/Engine/Shaders/Private/TemporalSuperResolution/TSRFilterAntiAliasing.usf
+++ b/Engine/Shaders/Private/TemporalSuperResolution/TSRFilterAntiAliasing.usf
@@ -11,19 +11,19 @@
 
 //------------------------------------------------------- PARAMETERS
 
-Texture2D<taa_ushort> AntiAliasingTexture;
+Texture2D<tsr_ushort> AntiAliasingTexture;
 RWTexture2D<uint> AntiAliasingOutput;
 
 
 //------------------------------------------------------- FUNCTIONS
 
-taa_half2 SampleAntiAliasing(taa_short2 KernelCenter, taa_short2 Offset)
+tsr_half2 SampleAntiAliasing(tsr_short2 KernelCenter, tsr_short2 Offset)
 {
-	taa_short2 SampleInputPixelPos = ClampPixelOffset(
+	tsr_short2 SampleInputPixelPos = ClampPixelOffset(
 		KernelCenter + Offset,
 		InputPixelPosMin, InputPixelPosMax);
 	
-	taa_ushort EncodedInputTexelOffset = AntiAliasingTexture[SampleInputPixelPos];
+	tsr_ushort EncodedInputTexelOffset = AntiAliasingTexture[SampleInputPixelPos];
 
 	return DecodeSpatialAntiAliasingOffset(EncodedInputTexelOffset);
 }
@@ -39,24 +39,24 @@ void MainCS(
 
 	float4 Debug = 0.0;
 	
-	taa_short2 InputPixelPos = (
-		taa_short2(InputPixelPosMin) +
-		taa_short2(GroupId) * taa_short2(TILE_SIZE, TILE_SIZE) +
+	tsr_short2 InputPixelPos = (
+		tsr_short2(InputPixelPosMin) +
+		tsr_short2(GroupId) * tsr_short2(TILE_SIZE, TILE_SIZE) +
 		Map8x8Tile2x2Lane(GroupThreadIndex));
 	
-	taa_half2 TexelOffsetC = SampleAntiAliasing(InputPixelPos, taa_short2( 0,  0));
+	tsr_half2 TexelOffsetC = SampleAntiAliasing(InputPixelPos, tsr_short2( 0,  0));
 	
 	Debug.rg = abs(TexelOffsetC);
 
-	taa_half2 TexelOffsetN = SampleAntiAliasing(InputPixelPos, taa_short2( 0, -1));
-	taa_half2 TexelOffsetS = SampleAntiAliasing(InputPixelPos, taa_short2( 0, +1));
-	taa_half2 TexelOffsetE = SampleAntiAliasing(InputPixelPos, taa_short2(+1,  0));
-	taa_half2 TexelOffsetW = SampleAntiAliasing(InputPixelPos, taa_short2(-1,  0));
+	tsr_half2 TexelOffsetN = SampleAntiAliasing(InputPixelPos, tsr_short2( 0, -1));
+	tsr_half2 TexelOffsetS = SampleAntiAliasing(InputPixelPos, tsr_short2( 0, +1));
+	tsr_half2 TexelOffsetE = SampleAntiAliasing(InputPixelPos, tsr_short2(+1,  0));
+	tsr_half2 TexelOffsetW = SampleAntiAliasing(InputPixelPos, tsr_short2(-1,  0));
 	
-	taa_half2 TexelOffsetNE = SampleAntiAliasing(InputPixelPos, taa_short2(+1, -1));
-	taa_half2 TexelOffsetNW = SampleAntiAliasing(InputPixelPos, taa_short2(-1, -1));
-	taa_half2 TexelOffsetSE = SampleAntiAliasing(InputPixelPos, taa_short2(+1, +1));
-	taa_half2 TexelOffsetSW = SampleAntiAliasing(InputPixelPos, taa_short2(-1, +1));
+	tsr_half2 TexelOffsetNE = SampleAntiAliasing(InputPixelPos, tsr_short2(+1, -1));
+	tsr_half2 TexelOffsetNW = SampleAntiAliasing(InputPixelPos, tsr_short2(-1, -1));
+	tsr_half2 TexelOffsetSE = SampleAntiAliasing(InputPixelPos, tsr_short2(+1, +1));
+	tsr_half2 TexelOffsetSW = SampleAntiAliasing(InputPixelPos, tsr_short2(-1, +1));
 	
 	FLATTEN
 	if (TexelOffsetC.x > 0.0)
@@ -101,13 +101,13 @@ void MainCS(
 #endif
 
 	{
-		taa_short2 OutputPixelPos;
+		tsr_short2 OutputPixelPos;
 		#if 0
 			ISOLATE
 			{
 				OutputPixelPos = dpv_add(
-					taa_short2(InputPixelPosMin) + taa_short2(GroupId) * taa_short2(TILE_SIZE, TILE_SIZE),
-					taa_short2x2(Map8x8Tile2x2LaneDPV(GroupThreadIndex)));
+					tsr_short2(InputPixelPosMin) + tsr_short2(GroupId) * tsr_short2(TILE_SIZE, TILE_SIZE),
+					tsr_short2x2(Map8x8Tile2x2LaneDPV(GroupThreadIndex)));
 			}
 		#else
 			OutputPixelPos = InputPixelPos;
diff --git a/Engine/Shaders/Private/TemporalSuperResolution/TSRFilterFrequencies.usf b/Engine/Shaders/Private/TemporalSuperResolution/TSRFilterFrequencies.usf
--- a/Engine/Shaders/Private/TemporalSuperResolution/TSRFilterFrequencies.usf
+++ b/Engine/Shaders/Private/TemporalSuperResolution/TSRFilterFrequencies.usf
@@ -32,31 +32,31 @@
 
 float3 OutputQuantizationError;
 
-Texture2D<taa_halfC> InputTexture;
-Texture2D<taa_halfC> PredictionSceneColorTexture;
-Texture2D<taa_half> ParallaxRejectionMaskTexture;
+Texture2D<tsr_halfC> InputTexture;
+Texture2D<tsr_halfC> PredictionSceneColorTexture;
+Texture2D<tsr_half> ParallaxRejectionMaskTexture;
 
-RWTexture2D<taa_halfC> FilteredInputOutput;
-RWTexture2D<taa_halfC> FilteredPredictionSceneColorOutput;
-RWTexture2D<taa_half>  InputSceneColorLdrLumaOutput;
+RWTexture2D<tsr_halfC> FilteredInputOutput;
+RWTexture2D<tsr_halfC> FilteredPredictionSceneColorOutput;
+RWTexture2D<tsr_half>  InputSceneColorLdrLumaOutput;
 
 
 //------------------------------------------------------- FUNCTIONS
 
-taa_half ComputeColorWeight(taa_halfC Color)
+tsr_half ComputeColorWeight(tsr_halfC Color)
 #if CONFIG_ACCUMULATE_KARIS_WEIGHTING
 {
 	return HdrWeight4(Color);
 }
 #else
 {
-	return taa_half(1);
+	return tsr_half(1);
 }
 #endif
 
-taa_half ParallaxRejectionMaskToHistoryValidity(taa_half ParallaxRejectionMask)
+tsr_half ParallaxRejectionMaskToHistoryValidity(tsr_half ParallaxRejectionMask)
 {
-	return (ParallaxRejectionMask < taa_half(PARALLAX_REJECTION_MASK_THRESHOLD)) ? taa_half(0.0) : taa_half(1.0);
+	return (ParallaxRejectionMask < tsr_half(PARALLAX_REJECTION_MASK_THRESHOLD)) ? tsr_half(0.0) : tsr_half(1.0);
 }
 
 //------------------------------------------------------- ENTRY POINT
@@ -70,28 +70,28 @@ void MainCS(
 
 	float4 Debug = 0.0;
 	
-	taa_short2 InputPixelPos = (
-		taa_short2(InputPixelPosMin) +
-		taa_short2(GroupId) * taa_short2(TILE_SIZE, TILE_SIZE) +
+	tsr_short2 InputPixelPos = (
+		tsr_short2(InputPixelPosMin) +
+		tsr_short2(GroupId) * tsr_short2(TILE_SIZE, TILE_SIZE) +
 		Map16x16Tile2x2Lane(GroupThreadIndex));
 
 	// Fetch center of the kernel
-	taa_halfC CenterInput;
-	taa_halfC FilteredInput;
-	taa_half FilteredInputWeight;
-	taa_halfC FilteredHistory;
-	taa_half FilteredHistoryWeight;
+	tsr_halfC CenterInput;
+	tsr_halfC FilteredInput;
+	tsr_half FilteredInputWeight;
+	tsr_halfC FilteredHistory;
+	tsr_half FilteredHistoryWeight;
 
 	{
-		taa_half SampleSpatialWeight = 1.0;
+		tsr_half SampleSpatialWeight = 1.0;
 
-		taa_halfC InputSample = InputTexture[InputPixelPos];
-		taa_halfC HistorySample = PredictionSceneColorTexture[InputPixelPos];
-		taa_half HistorySampleValidity = ParallaxRejectionMaskToHistoryValidity(
+		tsr_halfC InputSample = InputTexture[InputPixelPos];
+		tsr_halfC HistorySample = PredictionSceneColorTexture[InputPixelPos];
+		tsr_half HistorySampleValidity = ParallaxRejectionMaskToHistoryValidity(
 			ParallaxRejectionMaskTexture[InputPixelPos]);
 		
-		taa_half InputWeight = ComputeColorWeight(InputSample);
-		taa_half HistoryWeight = HistorySampleValidity * ComputeColorWeight(HistorySample);
+		tsr_half InputWeight = ComputeColorWeight(InputSample);
+		tsr_half HistoryWeight = HistorySampleValidity * ComputeColorWeight(HistorySample);
 
 		CenterInput = InputSample;
 
@@ -121,24 +121,24 @@ void MainCS(
 	#endif
 
 	#if CONFIG_USE_DYNAMIC_SPREAD_FACTOR
-		taa_short2 DynamicSpreadFactor = (
-			taa_short(CONFIG_KERNEL_SPREAD_FACTOR) +
-			taa_short2((taa_ushort2(InputPixelPos) % taa_ushort(2)) ^ (taa_ushort2(InputPixelPos.yx) % taa_ushort(2))));
+		tsr_short2 DynamicSpreadFactor = (
+			tsr_short(CONFIG_KERNEL_SPREAD_FACTOR) +
+			tsr_short2((tsr_ushort2(InputPixelPos) % tsr_ushort(2)) ^ (tsr_ushort2(InputPixelPos.yx) % tsr_ushort(2))));
 	#else
-		taa_short2 DynamicSpreadFactor = taa_short(CONFIG_KERNEL_SPREAD_FACTOR);
+		tsr_short2 DynamicSpreadFactor = tsr_short(CONFIG_KERNEL_SPREAD_FACTOR);
 	#endif
 
 	// Fetch neighborhood of the kernel
 	UNROLL
 	for (uint SampleId = StartSampleId; SampleId < SampleCount; SampleId++)
 	{
-		taa_short2 Offset;
-		taa_half SampleSpatialWeight;
+		tsr_short2 Offset;
+		tsr_half SampleSpatialWeight;
 		
 		#if CONFIG_KERNEL == KERNEL_1X1
 		{
 			Offset = 0;
-			SampleSpatialWeight = taa_half(1.0);
+			SampleSpatialWeight = tsr_half(1.0);
 			continue;
 		}
 		#elif CONFIG_KERNEL == KERNEL_3X3 || CONFIG_KERNEL == KERNEL_3X3_PLUS
@@ -147,7 +147,7 @@ void MainCS(
 			{
 				continue;
 			}
-			Offset = taa_short2(kOffsets3x3[SampleId]);
+			Offset = tsr_short2(kOffsets3x3[SampleId]);
 			
 			#if CONFIG_KERNEL == KERNEL_3X3_PLUS
 				if (abs(Offset.x) + abs(Offset.y) > 1)
@@ -157,25 +157,25 @@ void MainCS(
 			#endif
 			
 			#if 0
-				SampleSpatialWeight = taa_half(1.0);
+				SampleSpatialWeight = tsr_half(1.0);
 			#elif 1
-				SampleSpatialWeight = taa_half(rcp(abs(Offset.x) + 1.0) * rcp(abs(Offset.y) + 1.0));
+				SampleSpatialWeight = tsr_half(rcp(abs(Offset.x) + 1.0) * rcp(abs(Offset.y) + 1.0));
 			#else
-				SampleSpatialWeight = taa_half(rcp(abs(Offset.x) + abs(Offset.y)));
+				SampleSpatialWeight = tsr_half(rcp(abs(Offset.x) + abs(Offset.y)));
 			#endif
 		}
 		#else
 			#error Unknown kernel
 		#endif
 		
-		taa_short2 SampleInputPixelPos = InputPixelPos + ClampPixelOffset(
+		tsr_short2 SampleInputPixelPos = InputPixelPos + ClampPixelOffset(
 			InputPixelPos,
 			DynamicSpreadFactor * Offset, Offset,
 			InputPixelPosMin, InputPixelPosMax);
 
-		taa_halfC InputSample = InputTexture[SampleInputPixelPos];
-		taa_halfC HistorySample = PredictionSceneColorTexture[SampleInputPixelPos];
-		taa_half HistorySampleValidity = ParallaxRejectionMaskToHistoryValidity(
+		tsr_halfC InputSample = InputTexture[SampleInputPixelPos];
+		tsr_halfC HistorySample = PredictionSceneColorTexture[SampleInputPixelPos];
+		tsr_half HistorySampleValidity = ParallaxRejectionMaskToHistoryValidity(
 			ParallaxRejectionMaskTexture[SampleInputPixelPos]);
 
 		#if CONFIG_SCENE_COLOR_OVERFLOW
@@ -186,8 +186,8 @@ void MainCS(
 			#endif
 		#endif
 
-		taa_half InputWeight = HistorySampleValidity * ComputeColorWeight(InputSample);
-		taa_half HistoryWeight = HistorySampleValidity * ComputeColorWeight(HistorySample);
+		tsr_half InputWeight = HistorySampleValidity * ComputeColorWeight(InputSample);
+		tsr_half HistoryWeight = HistorySampleValidity * ComputeColorWeight(HistorySample);
 
 		FilteredInput       += (SampleSpatialWeight * InputWeight) * InputSample;
 		FilteredInputWeight += (SampleSpatialWeight * InputWeight);
@@ -203,30 +203,30 @@ void MainCS(
 	
 	// Compute the LDR luma used for spatial anti-aliasing's computation of the TexelOffset
 	#if DIM_OUTPUT_ANTI_ALIASING_LUMA
-	taa_half CenterLumaLDR;
+	tsr_half CenterLumaLDR;
 	{
-		const taa_half SpatialAAExposure = taa_half(0.5);
+		const tsr_half SpatialAAExposure = tsr_half(0.5);
 
 		#if 0
-			taa_half3 CenterInputLDR = CenterInput / (SpatialAAExposure + CenterInput);
+			tsr_half3 CenterInputLDR = CenterInput / (SpatialAAExposure + CenterInput);
 			CenterLumaLDR = dot(CenterInputLDR, half3(0.299f, 0.587f, 0.114f));
 		#elif 0
-			taa_half CenterLuma = max(CenterInput.r, max(CenterInput.g, CenterInput.b));
+			tsr_half CenterLuma = max(CenterInput.r, max(CenterInput.g, CenterInput.b));
 			CenterLumaLDR = CenterLuma / (SpatialAAExposure + CenterLuma);
 		#else
-			taa_half CenterLuma = dot(CenterInput.rgb, half3(0.299f, 0.587f, 0.114f));
+			tsr_half CenterLuma = dot(CenterInput.rgb, half3(0.299f, 0.587f, 0.114f));
 			CenterLumaLDR = CenterLuma / (SpatialAAExposure + CenterLuma);
 		#endif
 	}
 	#endif // DIM_OUTPUT_ANTI_ALIASING_LUMA
 	
-	taa_short2 OutputPixelPos;
+	tsr_short2 OutputPixelPos;
 	#if 1
 		ISOLATE
 		{
 			OutputPixelPos = (
-				taa_short2(InputPixelPosMin) +
-				taa_short2(GroupId) * taa_short2(TILE_SIZE, TILE_SIZE) +
+				tsr_short2(InputPixelPosMin) +
+				tsr_short2(GroupId) * tsr_short2(TILE_SIZE, TILE_SIZE) +
 				Map16x16Tile2x2Lane(GetGroupThreadIndex(GroupThreadIndex, GroupWaveIndex)));
 		}
 	#else
@@ -236,14 +236,14 @@ void MainCS(
 	#if CONFIG_ENABLE_STOCASTIC_QUANTIZATION
 	{
 		uint2 Random = Rand3DPCG16(int3(OutputPixelPos, View.StateFrameIndexMod8)).xy;
-		taa_half E = taa_half(Hammersley16(0, 1, Random).x);
+		tsr_half E = tsr_half(Hammersley16(0, 1, Random).x);
 		
 		FilteredInput = QuantizeForFloatRenderTarget(FilteredInput, E.x, OutputQuantizationError);
 		FilteredHistory = QuantizeForFloatRenderTarget(FilteredHistory, E.x, OutputQuantizationError);
 	}
 	#endif
 
-	OutputPixelPos.x = all(OutputPixelPos < taa_short2(LowFrequencyInfo_ViewportMax)) ? OutputPixelPos.x : ~taa_short(0);
+	OutputPixelPos.x = all(OutputPixelPos < tsr_short2(LowFrequencyInfo_ViewportMax)) ? OutputPixelPos.x : ~tsr_short(0);
 
 	FilteredInputOutput[OutputPixelPos] = FilteredInput;
 	FilteredPredictionSceneColorOutput[OutputPixelPos] = FilteredHistory;
diff --git a/Engine/Shaders/Private/TemporalSuperResolution/TSRHoleFillVelocity.usf b/Engine/Shaders/Private/TemporalSuperResolution/TSRHoleFillVelocity.usf
--- a/Engine/Shaders/Private/TemporalSuperResolution/TSRHoleFillVelocity.usf
+++ b/Engine/Shaders/Private/TemporalSuperResolution/TSRHoleFillVelocity.usf
@@ -27,10 +27,10 @@ float MaxHollFillPixelVelocity;
 
 Texture2D<uint> PrevClosestDepthTexture;
 Texture2D<float2> DilatedVelocityTexture;
-Texture2D<taa_half> ParallaxRejectionMaskTexture;
+Texture2D<tsr_half> ParallaxRejectionMaskTexture;
 
 RWTexture2D<float2> HoleFilledVelocityOutput;
-RWTexture2D<taa_half2> HoleFilledVelocityMaskOutput;
+RWTexture2D<tsr_half2> HoleFilledVelocityMaskOutput;
 
 
 //------------------------------------------------------- ENTRY POINT
@@ -45,33 +45,33 @@ void MainCS(
 	uint2 DispatchThreadId = (
 		ZOrder2D(GroupThreadIndex, uint(log2(float(TILE_SIZE)))) +
 		GroupId * uint2(TILE_SIZE, TILE_SIZE));
-	taa_short2 LaneOffsetSign = GetLaneOffsetSign();
+	tsr_short2 LaneOffsetSign = GetLaneOffsetSign();
 
-	taa_short2 InputPixelPos = taa_short2(InputInfo_ViewportMin + DispatchThreadId);
+	tsr_short2 InputPixelPos = tsr_short2(InputInfo_ViewportMin + DispatchThreadId);
 	
 	float2 ScreenPos = ApplyScreenTransform(float2(InputPixelPos), InputPixelPosToScreenPos);
 
 	float2 EncodedVelocityNeighborhood[CONFIG_PRE_FETCH_SAMPLES];
-	taa_half ParallaxRejectionMaskNeighborhood[CONFIG_PRE_FETCH_SAMPLES];
+	tsr_half ParallaxRejectionMaskNeighborhood[CONFIG_PRE_FETCH_SAMPLES];
 	ISOLATE
 	{
 		UNROLL_N(CONFIG_PRE_FETCH_SAMPLES)
 		for (uint i = 0; i < CONFIG_PRE_FETCH_SAMPLES; i++)
 		{
-			taa_short2 SampleInputPixelPos;
+			tsr_short2 SampleInputPixelPos;
 			#if CONFIG_BUTTERFLY_KERNEL
 			{
-				taa_short2 Offset = taa_short2(Offsets2x2[i]) * LaneOffsetSign;
+				tsr_short2 Offset = tsr_short2(Offsets2x2[i]) * LaneOffsetSign;
 				
 				SampleInputPixelPos = InputPixelPos + Offset;
 			}
 			#else
 			{
-				taa_short2 Offset;
+				tsr_short2 Offset;
 				#if CONFIG_SAMPLES == 9
-					Offset = taa_short2(kOffsets3x3[kSquareIndexes3x3[i]]);
+					Offset = tsr_short2(kOffsets3x3[kSquareIndexes3x3[i]]);
 				#elif CONFIG_SAMPLES == 5
-					Offset = taa_short2(kOffsets3x3[kPlusIndexes3x3[i]]);
+					Offset = tsr_short2(kOffsets3x3[kPlusIndexes3x3[i]]);
 				#endif
 					
 				SampleInputPixelPos = InputPixelPos + ClampPixelOffset(
@@ -144,11 +144,11 @@ void MainCS(
 	
 	float2 EndodedScreenVelocity = EncodedVelocityNeighborhood[0];
 
-	bool bIsParallaxRejected = ParallaxRejectionMaskNeighborhood[0] < taa_half(PARALLAX_REJECTION_MASK_THRESHOLD);
+	bool bIsParallaxRejected = ParallaxRejectionMaskNeighborhood[0] < tsr_half(PARALLAX_REJECTION_MASK_THRESHOLD);
 
-	taa_half HoleFilledVelocityMask = taa_half(0.0);
-	taa_half RejectHighFrequency = taa_half(0.0);
-	//taa_half IsVelocityEdge = taa_half(0.0);
+	tsr_half HoleFilledVelocityMask = tsr_half(0.0);
+	tsr_half RejectHighFrequency = tsr_half(0.0);
+	//tsr_half IsVelocityEdge = tsr_half(0.0);
 
 	{
 		float2 ScreenPixelVelocity = ScreenVelocity * ScreenVelocityToInputPixelVelocity;
@@ -158,8 +158,8 @@ void MainCS(
 		for (uint i = 0; i < CONFIG_SAMPLE_PAIRS; i++)
 		ISOLATE
 		{
-			taa_short2 Offset0;
-			taa_short2 Offset1;
+			tsr_short2 Offset0;
+			tsr_short2 Offset1;
 			
 			float2 NeighborEncodedVelocity0;
 			float2 NeighborEncodedVelocity1;
@@ -189,13 +189,13 @@ void MainCS(
 				{
 					FetchIndex[0] = 2;
 					FetchIndex[1] = 1;
-					Offset0 = taa_short2(-1, 1);
+					Offset0 = tsr_short2(-1, 1);
 				}
 				else
 				{
 					FetchIndex[0] = PairIndex + 1;
 					FetchIndex[1] = PairIndex + 1;
-					Offset0 = taa_short2(Offsets2x2[PairIndex + 1]);
+					Offset0 = tsr_short2(Offsets2x2[PairIndex + 1]);
 				}
 
 				Offset0 *= LaneOffsetSign;
@@ -207,11 +207,11 @@ void MainCS(
 				NeighborPixelVelocity0 = WaveBroadcast(BroadcastSettings[0], PixelVelocityNeighborhood[FetchIndex[0]]);
 				NeighborPixelVelocity1 = WaveBroadcast(BroadcastSettings[1], PixelVelocityNeighborhood[FetchIndex[1]]);
 			
-				NeighborPixelVelocityLength0 = taa_half(WaveBroadcast(BroadcastSettings[0], PixelVelocityLengthNeighborhood[FetchIndex[0]]));
-				NeighborPixelVelocityLength1 = taa_half(WaveBroadcast(BroadcastSettings[1], PixelVelocityLengthNeighborhood[FetchIndex[1]]));
+				NeighborPixelVelocityLength0 = tsr_half(WaveBroadcast(BroadcastSettings[0], PixelVelocityLengthNeighborhood[FetchIndex[0]]));
+				NeighborPixelVelocityLength1 = tsr_half(WaveBroadcast(BroadcastSettings[1], PixelVelocityLengthNeighborhood[FetchIndex[1]]));
 				
-				bNeighborIsParallaxRejected0 = WaveBroadcast(BroadcastSettings[0], ParallaxRejectionMaskNeighborhood[FetchIndex[0]]) < taa_half(PARALLAX_REJECTION_MASK_THRESHOLD);
-				bNeighborIsParallaxRejected1 = WaveBroadcast(BroadcastSettings[1], ParallaxRejectionMaskNeighborhood[FetchIndex[1]]) < taa_half(PARALLAX_REJECTION_MASK_THRESHOLD);
+				bNeighborIsParallaxRejected0 = WaveBroadcast(BroadcastSettings[0], ParallaxRejectionMaskNeighborhood[FetchIndex[0]]) < tsr_half(PARALLAX_REJECTION_MASK_THRESHOLD);
+				bNeighborIsParallaxRejected1 = WaveBroadcast(BroadcastSettings[1], ParallaxRejectionMaskNeighborhood[FetchIndex[1]]) < tsr_half(PARALLAX_REJECTION_MASK_THRESHOLD);
 			}
 			#else
 			{
@@ -219,11 +219,11 @@ void MainCS(
 				const uint NeightbordId1 = i + 1 + CONFIG_SAMPLE_PAIRS;
 			
 				#if CONFIG_SAMPLES == 9
-					Offset0 = taa_short2(kOffsets3x3[kSquareIndexes3x3[NeightbordId0]]);
-					Offset1 = taa_short2(kOffsets3x3[kSquareIndexes3x3[NeightbordId1]]);
+					Offset0 = tsr_short2(kOffsets3x3[kSquareIndexes3x3[NeightbordId0]]);
+					Offset1 = tsr_short2(kOffsets3x3[kSquareIndexes3x3[NeightbordId1]]);
 				#elif CONFIG_SAMPLES == 5
-					Offset0 = taa_short2(kOffsets3x3[kPlusIndexes3x3[NeightbordId0]]);
-					Offset1 = taa_short2(kOffsets3x3[kPlusIndexes3x3[NeightbordId1]]);
+					Offset0 = tsr_short2(kOffsets3x3[kPlusIndexes3x3[NeightbordId0]]);
+					Offset1 = tsr_short2(kOffsets3x3[kPlusIndexes3x3[NeightbordId1]]);
 				#endif
 					
 				NeighborEncodedVelocity0 = EncodedVelocityNeighborhood[NeightbordId0];
@@ -235,17 +235,17 @@ void MainCS(
 				NeighborPixelVelocityLength0 = dot(NeighborPixelVelocity0, NeighborPixelVelocity0);
 				NeighborPixelVelocityLength1 = dot(NeighborPixelVelocity1, NeighborPixelVelocity1);
 				
-				bNeighborIsParallaxRejected0 = ParallaxRejectionMaskNeighborhood[NeightbordId0] < taa_half(PARALLAX_REJECTION_MASK_THRESHOLD);
-				bNeighborIsParallaxRejected1 = ParallaxRejectionMaskNeighborhood[NeightbordId1] < taa_half(PARALLAX_REJECTION_MASK_THRESHOLD);
+				bNeighborIsParallaxRejected0 = ParallaxRejectionMaskNeighborhood[NeightbordId0] < tsr_half(PARALLAX_REJECTION_MASK_THRESHOLD);
+				bNeighborIsParallaxRejected1 = ParallaxRejectionMaskNeighborhood[NeightbordId1] < tsr_half(PARALLAX_REJECTION_MASK_THRESHOLD);
 			}
 			#endif
 
-			taa_half2 NeighborDeltaPixelVelocity0 = taa_half2(ScreenPixelVelocity - NeighborPixelVelocity0);
-			taa_half2 NeighborDeltaPixelVelocity1 = taa_half2(ScreenPixelVelocity - NeighborPixelVelocity1);
+			tsr_half2 NeighborDeltaPixelVelocity0 = tsr_half2(ScreenPixelVelocity - NeighborPixelVelocity0);
+			tsr_half2 NeighborDeltaPixelVelocity1 = tsr_half2(ScreenPixelVelocity - NeighborPixelVelocity1);
 			
 			// Detect edge of the velocity
 			//{
-			//	taa_half MaxDeltaVelocityLength = saturate(max(NeighborPixelVelocityLength0, NeighborPixelVelocityLength1));
+			//	tsr_half MaxDeltaVelocityLength = saturate(max(NeighborPixelVelocityLength0, NeighborPixelVelocityLength1));
 			//
 			//	if (i == 0)
 			//		IsVelocityEdge = MaxDeltaVelocityLength;
@@ -257,28 +257,28 @@ void MainCS(
 			{
 				bool bIsPixelThinRejection = bIsParallaxRejected && !bNeighborIsParallaxRejected0 && !bNeighborIsParallaxRejected1;
 
-				taa_half IsSomeWhatIdentical0 = saturate(taa_half(1.0) - dot(NeighborDeltaPixelVelocity0, NeighborDeltaPixelVelocity0));
-				taa_half IsSomeWhatIdentical1 = saturate(taa_half(1.0) - dot(NeighborDeltaPixelVelocity1, NeighborDeltaPixelVelocity1));
+				tsr_half IsSomeWhatIdentical0 = saturate(tsr_half(1.0) - dot(NeighborDeltaPixelVelocity0, NeighborDeltaPixelVelocity0));
+				tsr_half IsSomeWhatIdentical1 = saturate(tsr_half(1.0) - dot(NeighborDeltaPixelVelocity1, NeighborDeltaPixelVelocity1));
 
 				FLATTEN
 				if (bIsPixelThinRejection && NeighborPixelVelocityLength0 > ScreenPixelVelocityLength && IsSomeWhatIdentical1 > 0.0)
 				{
 					EndodedScreenVelocity = NeighborEncodedVelocity0;
-					HoleFilledVelocityMask = taa_half(1.0);
+					HoleFilledVelocityMask = tsr_half(1.0);
 				}
 				else if (bIsPixelThinRejection && NeighborPixelVelocityLength1 > ScreenPixelVelocityLength && IsSomeWhatIdentical0 > 0.0)
 				{
 					EndodedScreenVelocity = NeighborEncodedVelocity1;
-					HoleFilledVelocityMask = taa_half(1.0);
+					HoleFilledVelocityMask = tsr_half(1.0);
 				}
 			}
 
 			// Reject high frequency on velocity edge discontinuity.
 			{
-				taa_half IsOverlapingEdge0 = taa_half(1.0) - dot(saturate(NeighborDeltaPixelVelocity0 * Offset0), taa_half(1.0).xx); // * saturate(ScreenPixelVelocityLength - NeighborPixelVelocityLength0);
-				taa_half IsOverlapingEdge1 = taa_half(1.0) - dot(saturate(NeighborDeltaPixelVelocity1 * Offset1), taa_half(1.0).xx); // * saturate(ScreenPixelVelocityLength - NeighborPixelVelocityLength1);
+				tsr_half IsOverlapingEdge0 = tsr_half(1.0) - dot(saturate(NeighborDeltaPixelVelocity0 * Offset0), tsr_half(1.0).xx); // * saturate(ScreenPixelVelocityLength - NeighborPixelVelocityLength0);
+				tsr_half IsOverlapingEdge1 = tsr_half(1.0) - dot(saturate(NeighborDeltaPixelVelocity1 * Offset1), tsr_half(1.0).xx); // * saturate(ScreenPixelVelocityLength - NeighborPixelVelocityLength1);
 				
-				taa_half IsOverlapingEdge = min(IsOverlapingEdge0, IsOverlapingEdge1);
+				tsr_half IsOverlapingEdge = min(IsOverlapingEdge0, IsOverlapingEdge1);
 
 				if (i == 0)
 					RejectHighFrequency = IsOverlapingEdge;
@@ -289,27 +289,27 @@ void MainCS(
 
 		// If the hole filling velocity has better result that the neighborhood, use that.
 		FLATTEN
-		if (taa_half(HoleFillingOutOfRange) > HoleFilledVelocityMask && bIsParallaxRejected)
+		if (tsr_half(HoleFillingOutOfRange) > HoleFilledVelocityMask && bIsParallaxRejected)
 		{
 			float VelocityDelta = length(ScreenPixelVelocity - HoleFillingPixelVelocity);
 
 			float2 HoleFillingVelocity = HoleFillingPixelVelocity * InputPixelVelocityToScreenVelocity;
 
 			EndodedScreenVelocity = EncodeVelocityToTexture(float3(HoleFillingVelocity, 0.0)).xy;
-			HoleFilledVelocityMask = taa_half(HoleFillingOutOfRange) * taa_half(saturate(MaxHollFillPixelVelocity - VelocityDelta));
+			HoleFilledVelocityMask = tsr_half(HoleFillingOutOfRange) * tsr_half(saturate(MaxHollFillPixelVelocity - VelocityDelta));
 		}
 
 		// Remove noise from high frequency rejections.
-		RejectHighFrequency = saturate(taa_half(1.0) - RejectHighFrequency * taa_half(1.1));
+		RejectHighFrequency = saturate(tsr_half(1.0) - RejectHighFrequency * tsr_half(1.1));
 
 		// Do not reject high frequency additionally where hole filling happen
 		RejectHighFrequency = max(RejectHighFrequency, HoleFilledVelocityMask);
 	}
 
-	if (all(InputPixelPos < taa_short2(InputInfo_ViewportMax)))
+	if (all(InputPixelPos < tsr_short2(InputInfo_ViewportMax)))
 	{
 		HoleFilledVelocityOutput[InputPixelPos] = EndodedScreenVelocity;
-		HoleFilledVelocityMaskOutput[InputPixelPos] = taa_half2(HoleFilledVelocityMask, RejectHighFrequency);
+		HoleFilledVelocityMaskOutput[InputPixelPos] = tsr_half2(HoleFilledVelocityMask, RejectHighFrequency);
 
 		#if DEBUG_OUTPUT
 		{
diff --git a/Engine/Shaders/Private/TemporalSuperResolution/TSRPostfilterRejection.usf b/Engine/Shaders/Private/TemporalSuperResolution/TSRPostfilterRejection.usf
--- a/Engine/Shaders/Private/TemporalSuperResolution/TSRPostfilterRejection.usf
+++ b/Engine/Shaders/Private/TemporalSuperResolution/TSRPostfilterRejection.usf
@@ -14,47 +14,47 @@
 
 uint4 HistoryRejectionViewport;
 
-Texture2D<taa_half> HistoryRejectionTexture;
-RWTexture2D<taa_half> HistoryRejectionOutput;
+Texture2D<tsr_half> HistoryRejectionTexture;
+RWTexture2D<tsr_half> HistoryRejectionOutput;
 
 
 //------------------------------------------------------- LDS
 
-groupshared taa_half SharedRejection[TILE_SIZE * TILE_SIZE];
+groupshared tsr_half SharedRejection[TILE_SIZE * TILE_SIZE];
 
 
 //------------------------------------------------------- FUNCTIONS
 
 // Operator that output, lowest, median and highest values from 3 input values.
-void LMHOperator(taa_half A, taa_half B, taa_half C, out taa_half L, out taa_half M, out taa_half H)
+void LMHOperator(tsr_half A, tsr_half B, tsr_half C, out tsr_half L, out tsr_half M, out tsr_half H)
 {
 	// TODO: min3, max3, med3.
-	taa_half X = min(B, C);
-	taa_half Y = max(B, C);
+	tsr_half X = min(B, C);
+	tsr_half Y = max(B, C);
 	
 	L = min(A, X);
-	taa_half Z = max(A, X);
+	tsr_half Z = max(A, X);
 
 	M = min(Z, Y);
 	H = max(Z, Y);
 }
 
 // 3 samples median.
-taa_half Median(taa_half A, taa_half B, taa_half C)
+tsr_half Median(tsr_half A, tsr_half B, tsr_half C)
 {
 	// TODO: med3.
-	taa_half L, M, H;
+	tsr_half L, M, H;
 	LMHOperator(A, B, C, L, M, H);
 	return M;
 }
 
 // 9 samples median using Smith1996.
-taa_half Median9(taa_half In[9])
+tsr_half Median9(tsr_half In[9])
 {
 	// First layer.
-	taa_half L0[3];
-	taa_half M0[3];
-	taa_half H0[3];
+	tsr_half L0[3];
+	tsr_half M0[3];
+	tsr_half H0[3];
 	for (uint j = 0; j < 3; j++)
 	{
 		LMHOperator(In[j * 3 + 0], In[j * 3 + 1], In[j * 3 + 2], L0[j], M0[j], H0[j]);
@@ -62,7 +62,7 @@ taa_half Median9(taa_half In[9])
 	
 	// Second layer.
 	// TODO: min3, max3.
-	taa_half M1[3];
+	tsr_half M1[3];
 	M1[0] = max(max(L0[0], L0[1]), L0[2]);
 	M1[1] = Median(M0[0], M0[1], M0[2]);
 	M1[2] = min(min(H0[0], H0[1]), H0[2]);
@@ -81,26 +81,26 @@ void MainCS(
 {
 	float4 Debug = 0.0;
 
-	taa_short2 RejectionPixelPos = (
-		taa_short2(GroupId) * taa_short2(TILE_SIZE, TILE_SIZE) +
+	tsr_short2 RejectionPixelPos = (
+		tsr_short2(GroupId) * tsr_short2(TILE_SIZE, TILE_SIZE) +
 		Map8x8Tile2x2Lane(GroupThreadIndex));
 
 	RejectionPixelPos = ClampPixelOffset(RejectionPixelPos, HistoryRejectionViewport.xy, HistoryRejectionViewport.zw - 1);
 
 	// Fetch neighborhood of the kernel
-	taa_half InSamples[9];
+	tsr_half InSamples[9];
 	UNROLL
 	for (uint i = 0; i < 9; i++)
 	{
-		taa_short2 SampleInputPixelPos = AddAndClampPixelOffset(
+		tsr_short2 SampleInputPixelPos = AddAndClampPixelOffset(
 			RejectionPixelPos,
-			taa_short2(kOffsets3x3[i]), taa_short2(kOffsets3x3[i]),
+			tsr_short2(kOffsets3x3[i]), tsr_short2(kOffsets3x3[i]),
 			HistoryRejectionViewport.xy, HistoryRejectionViewport.zw - 1);
 
 		InSamples[i] = HistoryRejectionTexture[SampleInputPixelPos];
 	}
 
-	taa_half Rejection = Median9(InSamples);
+	tsr_half Rejection = Median9(InSamples);
 
 	#if CONFIG_OUTPUT_HALF_RES
 	{
@@ -116,23 +116,23 @@ void MainCS(
 			Rejection = min(Rejection, SharedRejection[GroupThreadIndex ^ ButterflySwap]);
 		}
 		
-		bool bIsValid = all(RejectionPixelPos <= taa_short2(HistoryRejectionViewport.zw)) && all((taa_ushort2(RejectionPixelPos) % taa_ushort(2)) == taa_ushort(0));
+		bool bIsValid = all(RejectionPixelPos <= tsr_short2(HistoryRejectionViewport.zw)) && all((tsr_ushort2(RejectionPixelPos) % tsr_ushort(2)) == tsr_ushort(0));
 
-		taa_ushort2 OutputPixelPos = taa_ushort2(RejectionPixelPos) / taa_ushort(2);
-		OutputPixelPos.x = bIsValid ? OutputPixelPos.x : ~taa_ushort(0);
+		tsr_ushort2 OutputPixelPos = tsr_ushort2(RejectionPixelPos) / tsr_ushort(2);
+		OutputPixelPos.x = bIsValid ? OutputPixelPos.x : ~tsr_ushort(0);
 
 		HistoryRejectionOutput[OutputPixelPos] = Rejection;
 	}
 	#else
 	{
-		taa_short2 OutputPixelPos = InvalidateOutputPixelPos(RejectionPixelPos, HistoryRejectionViewport.zw);
+		tsr_short2 OutputPixelPos = InvalidateOutputPixelPos(RejectionPixelPos, HistoryRejectionViewport.zw);
 		HistoryRejectionOutput[OutputPixelPos] = Rejection;
 	}
 	#endif
 	
 	#if DEBUG_OUTPUT
 	{
-		taa_short2 OutputPixelPos = InvalidateOutputPixelPos(RejectionPixelPos, HistoryRejectionViewport.zw);
+		tsr_short2 OutputPixelPos = InvalidateOutputPixelPos(RejectionPixelPos, HistoryRejectionViewport.zw);
 		DebugOutput[OutputPixelPos] = Debug;
 	}
 	#endif
diff --git a/Engine/Shaders/Private/TemporalSuperResolution/TSRResolveHistory.usf b/Engine/Shaders/Private/TemporalSuperResolution/TSRResolveHistory.usf
--- a/Engine/Shaders/Private/TemporalSuperResolution/TSRResolveHistory.usf
+++ b/Engine/Shaders/Private/TemporalSuperResolution/TSRResolveHistory.usf
@@ -28,7 +28,7 @@ RWTexture2D<uint> PrevClosestDepthOutput;
 
 //------------------------------------------------------- LDS
 
-groupshared taa_halfC SharedArray0[TILE_SIZE * TILE_SIZE];
+groupshared tsr_halfC SharedArray0[TILE_SIZE * TILE_SIZE];
 
 
 //------------------------------------------------------- FUNCTIONS
@@ -83,14 +83,14 @@ uint2 OutputViewRectMax;
 uint bGenerateOutputMip1;
 float HistoryValidityMultiply;
 
-Texture2D<taa_halfC> History_LowFrequency;
-Texture2D<taa_halfC> History_HighFrequency;
-Texture2D<taa_half2> History_Metadata;
-Texture2D<taa_half3> History_Translucency;
-Texture2D<taa_half>  History_TranslucencyAlpha;
+Texture2D<tsr_halfC> History_LowFrequency;
+Texture2D<tsr_halfC> History_HighFrequency;
+Texture2D<tsr_half2> History_Metadata;
+Texture2D<tsr_half3> History_Translucency;
+Texture2D<tsr_half>  History_TranslucencyAlpha;
 
-RWTexture2D<taa_halfC> SceneColorOutputMip0;
-RWTexture2D<taa_halfC> SceneColorOutputMip1;
+RWTexture2D<tsr_halfC> SceneColorOutputMip0;
+RWTexture2D<tsr_halfC> SceneColorOutputMip1;
 
 
 //------------------------------------------------------- ENTRY POINT
@@ -106,9 +106,9 @@ void MainCS(
 		ZOrder2D(GroupThreadIndex, uint(log2(float(TILE_SIZE)))) +
 		GroupId * uint2(TILE_SIZE, TILE_SIZE));
 
-	taa_halfC LowFrequencyColor = 0.0;
-	taa_halfC HighFrequencyColor = 0.0;
-	taa_half  HistoryValidity = 0.0;
+	tsr_halfC LowFrequencyColor = 0.0;
+	tsr_halfC HighFrequencyColor = 0.0;
+	tsr_half  HistoryValidity = 0.0;
 	
 	// Downsample the history
 	{
@@ -119,9 +119,9 @@ void MainCS(
 		float2 OutputPixelPos = ApplyScreenTransform(float2(DispatchThreadId), DispatchThreadToHistoryPixelPos);
 		float2 TopLeftKernel = floor(OutputPixelPos + (0.5 - 0.5 * KERNEL_SIZE)) + 0.5;
 		
-		taa_half LowFrequencyWeight = 0.0;
-		taa_half HighFrequencyWeight = 0.0;
-		taa_half HistoryValidityWeight = 0.0;
+		tsr_half LowFrequencyWeight = 0.0;
+		tsr_half HighFrequencyWeight = 0.0;
+		tsr_half HistoryValidityWeight = 0.0;
 
 		UNROLL_N(KERNEL_SIZE)
 		for (uint x = 0; x < KERNEL_SIZE; x++)
@@ -135,15 +135,15 @@ void MainCS(
 
 				SampleUV = clamp(SampleUV, HistoryInfo_UVViewportBilinearMin, HistoryInfo_UVViewportBilinearMax);
 
-				taa_halfC SampleLowFrequencyColor  = History_LowFrequency.SampleLevel(GlobalBilinearClampedSampler, SampleUV, 0);
-				taa_halfC SampleHighFrequencyColor = History_HighFrequency.SampleLevel(GlobalBilinearClampedSampler, SampleUV, 0);
-				taa_half2 SampleMetadata   	       = History_Metadata.SampleLevel(GlobalBilinearClampedSampler, SampleUV, 0);
+				tsr_halfC SampleLowFrequencyColor  = History_LowFrequency.SampleLevel(GlobalBilinearClampedSampler, SampleUV, 0);
+				tsr_halfC SampleHighFrequencyColor = History_HighFrequency.SampleLevel(GlobalBilinearClampedSampler, SampleUV, 0);
+				tsr_half2 SampleMetadata   	       = History_Metadata.SampleLevel(GlobalBilinearClampedSampler, SampleUV, 0);
 				
-				taa_half3 SampleTranslucencyColor = History_Translucency.SampleLevel(GlobalBilinearClampedSampler, SampleUV, 0);
-				taa_half  SampleTranslucencyAlpha = History_TranslucencyAlpha.SampleLevel(GlobalBilinearClampedSampler, SampleUV, 0);
+				tsr_half3 SampleTranslucencyColor = History_Translucency.SampleLevel(GlobalBilinearClampedSampler, SampleUV, 0);
+				tsr_half  SampleTranslucencyAlpha = History_TranslucencyAlpha.SampleLevel(GlobalBilinearClampedSampler, SampleUV, 0);
 				
-				taa_half SampleHistoryValidity = SampleMetadata.y;
-				taa_half InvSampleHistoryValidity = SafeRcp(SampleHistoryValidity);
+				tsr_half SampleHistoryValidity = SampleMetadata.y;
+				tsr_half InvSampleHistoryValidity = SafeRcp(SampleHistoryValidity);
 
 				#if 1
 				{
@@ -151,18 +151,18 @@ void MainCS(
 					SampleHighFrequencyColor.rgb = SampleHighFrequencyColor.rgb * SampleTranslucencyAlpha + SampleTranslucencyColor.rgb * SampleHistoryValidity;
 				
 					#if CONFIG_SCENE_COLOR_APLHA
-						SampleLowFrequencyColor.a = saturate(taa_half(1.0) - saturate(taa_half(1.0) - SampleLowFrequencyColor.a) * SampleTranslucencyAlpha);
-						SampleHighFrequencyColor.a = saturate(taa_half(1.0) - saturate(taa_half(1.0) - SampleHighFrequencyColor.a * InvSampleHistoryValidity) * SampleTranslucencyAlpha) * SampleHistoryValidity;
+						SampleLowFrequencyColor.a = saturate(tsr_half(1.0) - saturate(tsr_half(1.0) - SampleLowFrequencyColor.a) * SampleTranslucencyAlpha);
+						SampleHighFrequencyColor.a = saturate(tsr_half(1.0) - saturate(tsr_half(1.0) - SampleHighFrequencyColor.a * InvSampleHistoryValidity) * SampleTranslucencyAlpha) * SampleHistoryValidity;
 					#endif
 				}
 				#endif
 
 				float2 PixelOffset = abs(SamplePixelPos - OutputPixelPos);
 
-				taa_half kernelWeight = taa_half(MitchellNetravali(PixelOffset.x, B, C) * MitchellNetravali(PixelOffset.y, B, C));
+				tsr_half kernelWeight = tsr_half(MitchellNetravali(PixelOffset.x, B, C) * MitchellNetravali(PixelOffset.y, B, C));
 				
-				taa_half LowFrequencyHdrWeight  = HdrWeight4(SampleLowFrequencyColor);
-				taa_half HighFrequencyHdrWeight = HdrWeight4(SampleHighFrequencyColor * InvSampleHistoryValidity);
+				tsr_half LowFrequencyHdrWeight  = HdrWeight4(SampleLowFrequencyColor);
+				tsr_half HighFrequencyHdrWeight = HdrWeight4(SampleHighFrequencyColor * InvSampleHistoryValidity);
 
 				LowFrequencyColor   += (kernelWeight * LowFrequencyHdrWeight)  * SampleLowFrequencyColor;
 				LowFrequencyWeight  += (kernelWeight * LowFrequencyHdrWeight);
@@ -170,30 +170,30 @@ void MainCS(
 				HighFrequencyColor  += (kernelWeight * HighFrequencyHdrWeight) * SampleHighFrequencyColor;
 				HighFrequencyWeight += (kernelWeight * HighFrequencyHdrWeight) * SampleHistoryValidity;
 
-				HistoryValidity       += max(kernelWeight, taa_half(0.0)) * SampleHistoryValidity;
-				HistoryValidityWeight += max(kernelWeight, taa_half(0.0));
+				HistoryValidity       += max(kernelWeight, tsr_half(0.0)) * SampleHistoryValidity;
+				HistoryValidityWeight += max(kernelWeight, tsr_half(0.0));
 			}
 		}
 
 		LowFrequencyColor  *= rcp(LowFrequencyWeight);
 		HighFrequencyColor *= SafeRcp(HighFrequencyWeight);
-		HistoryValidity *= taa_half(HistoryValidityMultiply) / HistoryValidityWeight;
+		HistoryValidity *= tsr_half(HistoryValidityMultiply) / HistoryValidityWeight;
 	}
 	
 	// Compute final output
-	taa_halfC FinalOutputColor;
+	tsr_halfC FinalOutputColor;
 	{
-		taa_half OutputBlend = ComputeLowFrequencyContribution(HistoryValidity);
+		tsr_half OutputBlend = ComputeLowFrequencyContribution(HistoryValidity);
 		
-		taa_half LowFrequencyWeight = HdrWeight4(LowFrequencyColor);
-		taa_half Moment1Weight = HdrWeight4(HighFrequencyColor);
+		tsr_half LowFrequencyWeight = HdrWeight4(LowFrequencyColor);
+		tsr_half Moment1Weight = HdrWeight4(HighFrequencyColor);
 
-		taa_half2 Weights = WeightedLerpFactors(LowFrequencyWeight, Moment1Weight, OutputBlend);
+		tsr_half2 Weights = WeightedLerpFactors(LowFrequencyWeight, Moment1Weight, OutputBlend);
 		
 		FinalOutputColor = LowFrequencyColor * Weights[0] + HighFrequencyColor * Weights[1];
 	}
 
-	taa_short2 OutputPixelPosition = InvalidateOutputPixelPos(taa_short2(OutputViewRectMin + DispatchThreadId), OutputViewRectMax);
+	tsr_short2 OutputPixelPosition = InvalidateOutputPixelPos(tsr_short2(OutputViewRectMin + DispatchThreadId), OutputViewRectMax);
 	
 	// Output final scene color Mip0
 	{
@@ -204,10 +204,10 @@ void MainCS(
 	{
 		uint LocalGroupThreadIndex = GroupThreadIndex;
 
-		taa_halfC HalfResOutput = FinalOutputColor * taa_half(0.25);
-		taa_short2 HalfResOutputPixelPos;
-		HalfResOutputPixelPos[0] = (OutputPixelPosition[0] >> taa_short(1)) | (((OutputPixelPosition[0] & taa_short(0x1)) | taa_short(!bGenerateOutputMip1)) * taa_short(~0));
-		HalfResOutputPixelPos[1] = (OutputPixelPosition[1] >> taa_short(1)) | (((OutputPixelPosition[1] & taa_short(0x1)) | taa_short(!bGenerateOutputMip1)) * taa_short(~0));
+		tsr_halfC HalfResOutput = FinalOutputColor * tsr_half(0.25);
+		tsr_short2 HalfResOutputPixelPos;
+		HalfResOutputPixelPos[0] = (OutputPixelPosition[0] >> tsr_short(1)) | (((OutputPixelPosition[0] & tsr_short(0x1)) | tsr_short(!bGenerateOutputMip1)) * tsr_short(~0));
+		HalfResOutputPixelPos[1] = (OutputPixelPosition[1] >> tsr_short(1)) | (((OutputPixelPosition[1] & tsr_short(0x1)) | tsr_short(!bGenerateOutputMip1)) * tsr_short(~0));
 			
 		#if PLATFORM_SUPPORTS_WAVE_BROADCAST && 0 // Support WaveBroadcast with halfs
 		BRANCH
diff --git a/Engine/Shaders/Private/TemporalSuperResolution/TSRSpatialAntiAliasing.usf b/Engine/Shaders/Private/TemporalSuperResolution/TSRSpatialAntiAliasing.usf
--- a/Engine/Shaders/Private/TemporalSuperResolution/TSRSpatialAntiAliasing.usf
+++ b/Engine/Shaders/Private/TemporalSuperResolution/TSRSpatialAntiAliasing.usf
@@ -34,41 +34,41 @@
 
 //------------------------------------------------------- PARAMETERS
 
-Texture2D<taa_half> InputSceneColorLdrLumaTexture;
+Texture2D<tsr_half> InputSceneColorLdrLumaTexture;
 
 RWTexture2D<uint> AntiAliasingOutput;
-RWTexture2D<taa_half> NoiseFilteringOutput;
+RWTexture2D<tsr_half> NoiseFilteringOutput;
 
 
 //------------------------------------------------------- FUNCTIONS
 
-taa_ushort2x2 Map8x8Tile2x2LaneDPV(uint GroupThreadIndex)
+tsr_ushort2x2 Map8x8Tile2x2LaneDPV(uint GroupThreadIndex)
 #if CONFIG_ENABLE_DUAL_PIXEL_VECTORIZATION
 {
-	taa_ushort T = taa_ushort(GroupThreadIndex);
+	tsr_ushort T = tsr_ushort(GroupThreadIndex);
 
-	taa_ushort2 GroupId;
-	GroupId.x =                                             ((T >> taa_ushort(1 - 1)) & taa_ushort(0x03 << 1));
-	GroupId.y = ((T >> taa_ushort(0)) & taa_ushort(0x01)) | ((T >> taa_ushort(3 - 1)) & taa_ushort(0x03 << 1));
+	tsr_ushort2 GroupId;
+	GroupId.x =                                             ((T >> tsr_ushort(1 - 1)) & tsr_ushort(0x03 << 1));
+	GroupId.y = ((T >> tsr_ushort(0)) & tsr_ushort(0x01)) | ((T >> tsr_ushort(3 - 1)) & tsr_ushort(0x03 << 1));
 	
-	return dpv_interleave_registers(GroupId, GroupId + taa_ushort2(1, 0));
+	return dpv_interleave_registers(GroupId, GroupId + tsr_ushort2(1, 0));
 }
 #else
 {
-	taa_ushort2 GroupId = Map8x8Tile2x2Lane(GroupThreadIndex);
+	tsr_ushort2 GroupId = Map8x8Tile2x2Lane(GroupThreadIndex);
 	return dpv_interleave_mono_registers(GroupId);
 }
 #endif
 
-taa_half2 InputToLDRLuma(taa_half3x2 Color)
+tsr_half2 InputToLDRLuma(tsr_half3x2 Color)
 {
-	taa_half2 Luma = Luma4(Color);
+	tsr_half2 Luma = Luma4(Color);
 	return Luma * HdrWeightY(Luma);
 }
 
-taa_half SampleInputLDRLuma(taa_short2 KernelCenter, taa_short2 Offset)
+tsr_half SampleInputLDRLuma(tsr_short2 KernelCenter, tsr_short2 Offset)
 {
-	taa_short2 SampleInputPixelPos = KernelCenter + ClampPixelOffset(
+	tsr_short2 SampleInputPixelPos = KernelCenter + ClampPixelOffset(
 		KernelCenter,
 		Offset, Offset,
 		InputPixelPosMin, InputPixelPosMax);
@@ -76,34 +76,34 @@ taa_half SampleInputLDRLuma(taa_short2 KernelCenter, taa_short2 Offset)
 	return InputSceneColorLdrLumaTexture[SampleInputPixelPos];
 }
 
-taa_half2 SampleInputLDRLuma(taa_short2x2 KernelCenter, taa_short2 Offset)
+tsr_half2 SampleInputLDRLuma(tsr_short2x2 KernelCenter, tsr_short2 Offset)
 {
-	taa_short2x2 SampleInputPixelPos = KernelCenter + ClampPixelOffset(
+	tsr_short2x2 SampleInputPixelPos = KernelCenter + ClampPixelOffset(
 		KernelCenter,
 		dpv_interleave_mono_registers(Offset), Offset,
 		InputPixelPosMin, InputPixelPosMax);
 	
-	taa_half2 InputLuma = dpv_interleave_registers(
+	tsr_half2 InputLuma = dpv_interleave_registers(
 		InputSceneColorLdrLumaTexture[dpv_lo(SampleInputPixelPos)],
 		InputSceneColorLdrLumaTexture[dpv_hi(SampleInputPixelPos)]);
 
 	return InputLuma;
 }
 
-taa_half2 SampleInputLDRLuma(taa_short2x2 KernelCenter, taa_short2x2 Offset)
+tsr_half2 SampleInputLDRLuma(tsr_short2x2 KernelCenter, tsr_short2x2 Offset)
 {
-	taa_short2x2 SampleInputPixelPos = ClampPixelOffset(
+	tsr_short2x2 SampleInputPixelPos = ClampPixelOffset(
 		KernelCenter + Offset,
 		InputPixelPosMin, InputPixelPosMax);
 	
-	taa_half2 InputLuma = dpv_interleave_registers(
+	tsr_half2 InputLuma = dpv_interleave_registers(
 		InputSceneColorLdrLumaTexture[dpv_lo(SampleInputPixelPos)],
 		InputSceneColorLdrLumaTexture[dpv_hi(SampleInputPixelPos)]);
 
 	return InputLuma;
 }
 
-taa_half2 BilinearSampleInputLDRLuma(float2x2 KernelCenterUV, float2x2 UVOffset, const bool bIsNegativeDirection)
+tsr_half2 BilinearSampleInputLDRLuma(float2x2 KernelCenterUV, float2x2 UVOffset, const bool bIsNegativeDirection)
 {
 	float2x2 SampleUV = KernelCenterUV + UVOffset;
 
@@ -112,7 +112,7 @@ taa_half2 BilinearSampleInputLDRLuma(float2x2 KernelCenterUV, float2x2 UVOffset,
 	else
 		SampleUV = min(SampleUV, dpv_interleave_mono_registers(InputInfo_UVViewportBilinearMax));
 
-	taa_half2 InputLuma = dpv_interleave_registers(
+	tsr_half2 InputLuma = dpv_interleave_registers(
 		InputSceneColorLdrLumaTexture.SampleLevel(GlobalBilinearClampedSampler, dpv_lo(SampleUV), 0),
 		InputSceneColorLdrLumaTexture.SampleLevel(GlobalBilinearClampedSampler, dpv_hi(SampleUV), 0));
 
@@ -120,12 +120,12 @@ taa_half2 BilinearSampleInputLDRLuma(float2x2 KernelCenterUV, float2x2 UVOffset,
 }
 
 void AccumulateEdgeLength(
-	taa_half2 StartLuma, taa_half2 StartEdgeLuma,
-	taa_half2 SampleLuma, taa_half2 SampleEdgeLuma,
-	taa_half2 LumaDelta,
+	tsr_half2 StartLuma, tsr_half2 StartEdgeLuma,
+	tsr_half2 SampleLuma, tsr_half2 SampleEdgeLuma,
+	tsr_half2 LumaDelta,
 	inout bool2 bEdgeStopedByIncrement,
 	inout bool2 bEdgeStopedByDecrement,
-	inout taa_ushort2 EdgeLength)
+	inout tsr_ushort2 EdgeLength)
 {
 	bool2 bStopNowIncrement = abs(StartLuma - SampleLuma) > LumaDelta;
 	bool2 bStopNowDecrement = abs(StartEdgeLuma - SampleEdgeLuma) > LumaDelta;
@@ -133,24 +133,24 @@ void AccumulateEdgeLength(
 	bEdgeStopedByIncrement = bEdgeStopedByIncrement || bStopNowIncrement;
 	bEdgeStopedByDecrement = bEdgeStopedByDecrement || bStopNowDecrement;
 
-	EdgeLength += (bEdgeStopedByIncrement || bEdgeStopedByDecrement) ? taa_ushort(0) : taa_ushort(1);
+	EdgeLength += (bEdgeStopedByIncrement || bEdgeStopedByDecrement) ? tsr_ushort(0) : tsr_ushort(1);
 }
 
-taa_half ComputeDistanceToEdge(
+tsr_half ComputeDistanceToEdge(
 	bool bEdgeStopedByIncrementN,
 	bool bEdgeStopedByIncrementP,
 	bool bEdgeStopedByDecrementN,
 	bool bEdgeStopedByDecrementP,
-	taa_ushort EdgeLengthN,
-	taa_ushort EdgeLengthP)
+	tsr_ushort EdgeLengthN,
+	tsr_ushort EdgeLengthP)
 {	
-	taa_half fEdgeLengthN = taa_half(EdgeLengthN);
-	taa_half fEdgeLengthP = taa_half(EdgeLengthP);
-	taa_half fEdgeLength = taa_half(1.0) + fEdgeLengthN + fEdgeLengthP;
-	taa_half fEdgeInvLength = rcp(fEdgeLength);
+	tsr_half fEdgeLengthN = tsr_half(EdgeLengthN);
+	tsr_half fEdgeLengthP = tsr_half(EdgeLengthP);
+	tsr_half fEdgeLength = tsr_half(1.0) + fEdgeLengthN + fEdgeLengthP;
+	tsr_half fEdgeInvLength = rcp(fEdgeLength);
 
 	
-	taa_half TexelCenterOffsetToEdge = 0.0;
+	tsr_half TexelCenterOffsetToEdge = 0.0;
 
 	const float Min = -0.5;
 
@@ -158,12 +158,12 @@ taa_half ComputeDistanceToEdge(
 	if (!bEdgeStopedByIncrementN && !bEdgeStopedByIncrementP && !bEdgeStopedByDecrementP && !bEdgeStopedByDecrementN)
 	{
 		// No aliasing in the neighborhood.
-		TexelCenterOffsetToEdge = taa_half(0.0);
+		TexelCenterOffsetToEdge = tsr_half(0.0);
 	}
 	else if (bEdgeStopedByIncrementN && bEdgeStopedByIncrementP)
 	{
 		// Looks like InputC is concave detail
-		TexelCenterOffsetToEdge = taa_half(0.5);
+		TexelCenterOffsetToEdge = tsr_half(0.5);
 	}
 	else if (bEdgeStopedByDecrementN && bEdgeStopedByDecrementP)
 	{
@@ -173,12 +173,12 @@ taa_half ComputeDistanceToEdge(
 	else if (bEdgeStopedByIncrementN) // && bEdgeStopedByDecrementP)
 	{
 		// Looks like staircasing from - to +
-		TexelCenterOffsetToEdge = taa_half(0.5) - (EdgeLengthN + taa_half(0.5)) * fEdgeInvLength;
+		TexelCenterOffsetToEdge = tsr_half(0.5) - (EdgeLengthN + tsr_half(0.5)) * fEdgeInvLength;
 	}
 	else if (bEdgeStopedByIncrementP) // && bEdgeStopedByDecrementN)
 	{
 		// Looks like staircasing from + to -
-		TexelCenterOffsetToEdge = taa_half(0.5) - (EdgeLengthP + taa_half(0.5)) * fEdgeInvLength;
+		TexelCenterOffsetToEdge = tsr_half(0.5) - (EdgeLengthP + tsr_half(0.5)) * fEdgeInvLength;
 	}
 
 	return TexelCenterOffsetToEdge;
@@ -200,17 +200,17 @@ void MainCS(
 
 	float4x2 Debug = 0.0;
 	
-	taa_short2x2 InputPixelPos = dpv_add(
-		taa_short2(InputPixelPosMin) + taa_short2(GroupId) * taa_short2(TILE_SIZE, TILE_SIZE),
-		taa_short2x2(Map8x8Tile2x2LaneDPV(GroupThreadIndex)));
+	tsr_short2x2 InputPixelPos = dpv_add(
+		tsr_short2(InputPixelPosMin) + tsr_short2(GroupId) * tsr_short2(TILE_SIZE, TILE_SIZE),
+		tsr_short2x2(Map8x8Tile2x2LaneDPV(GroupThreadIndex)));
 	
 	//Debug[0] = 1.0;
 
-	taa_half2x2 TexelOffset = 0.0;
-	taa_half2 NoiseFiltering = 0.0;
+	tsr_half2x2 TexelOffset = 0.0;
+	tsr_half2 NoiseFiltering = 0.0;
 	{
 		#if CONFIG_ENABLE_DUAL_PIXEL_VECTORIZATION && 1
-			taa_half InputCache[3 * 4];
+			tsr_half InputCache[3 * 4];
 			UNROLL_N(3)
 			for (int cy = 0; cy < 3; cy++)
 			{
@@ -218,65 +218,65 @@ void MainCS(
 				for (int cx = 0; cx < 4; cx++)
 				{
 					const uint ci = cx + 4 * cy;
-					InputCache[ci] = SampleInputLDRLuma(dpv_lo(InputPixelPos), taa_short2(cx - 1, cy - 1));
+					InputCache[ci] = SampleInputLDRLuma(dpv_lo(InputPixelPos), tsr_short2(cx - 1, cy - 1));
 				}
 			}
 
-			#define ReadCache(x, y) taa_half2(InputCache[(x + 1) + 4 * (y + 1)], InputCache[(x + 2) + 4 * (y + 1)])
+			#define ReadCache(x, y) tsr_half2(InputCache[(x + 1) + 4 * (y + 1)], InputCache[(x + 2) + 4 * (y + 1)])
 
-			taa_half2 InputC = ReadCache( 0,  0);
+			tsr_half2 InputC = ReadCache( 0,  0);
 
-			taa_half2 InputN = ReadCache( 0, -1);
-			taa_half2 InputS = ReadCache( 0, +1);
-			taa_half2 InputE = ReadCache(+1,  0);
-			taa_half2 InputW = ReadCache(-1,  0);
+			tsr_half2 InputN = ReadCache( 0, -1);
+			tsr_half2 InputS = ReadCache( 0, +1);
+			tsr_half2 InputE = ReadCache(+1,  0);
+			tsr_half2 InputW = ReadCache(-1,  0);
 	
-			taa_half2 InputNE = ReadCache(+1, -1);
-			taa_half2 InputNW = ReadCache(-1, -1);
-			taa_half2 InputSE = ReadCache(+1, +1);
-			taa_half2 InputSW = ReadCache(-1, +1);
+			tsr_half2 InputNE = ReadCache(+1, -1);
+			tsr_half2 InputNW = ReadCache(-1, -1);
+			tsr_half2 InputSE = ReadCache(+1, +1);
+			tsr_half2 InputSW = ReadCache(-1, +1);
 		#else
-			taa_half2 InputC = SampleInputLDRLuma(InputPixelPos, taa_short2( 0,  0));
+			tsr_half2 InputC = SampleInputLDRLuma(InputPixelPos, tsr_short2( 0,  0));
 
-			taa_half2 InputN = SampleInputLDRLuma(InputPixelPos, taa_short2( 0, -1));
-			taa_half2 InputS = SampleInputLDRLuma(InputPixelPos, taa_short2( 0, +1));
-			taa_half2 InputE = SampleInputLDRLuma(InputPixelPos, taa_short2(+1,  0));
-			taa_half2 InputW = SampleInputLDRLuma(InputPixelPos, taa_short2(-1,  0));
+			tsr_half2 InputN = SampleInputLDRLuma(InputPixelPos, tsr_short2( 0, -1));
+			tsr_half2 InputS = SampleInputLDRLuma(InputPixelPos, tsr_short2( 0, +1));
+			tsr_half2 InputE = SampleInputLDRLuma(InputPixelPos, tsr_short2(+1,  0));
+			tsr_half2 InputW = SampleInputLDRLuma(InputPixelPos, tsr_short2(-1,  0));
 	
-			taa_half2 InputNE = SampleInputLDRLuma(InputPixelPos, taa_short2(+1, -1));
-			taa_half2 InputNW = SampleInputLDRLuma(InputPixelPos, taa_short2(-1, -1));
-			taa_half2 InputSE = SampleInputLDRLuma(InputPixelPos, taa_short2(+1, +1));
-			taa_half2 InputSW = SampleInputLDRLuma(InputPixelPos, taa_short2(-1, +1));
+			tsr_half2 InputNE = SampleInputLDRLuma(InputPixelPos, tsr_short2(+1, -1));
+			tsr_half2 InputNW = SampleInputLDRLuma(InputPixelPos, tsr_short2(-1, -1));
+			tsr_half2 InputSE = SampleInputLDRLuma(InputPixelPos, tsr_short2(+1, +1));
+			tsr_half2 InputSW = SampleInputLDRLuma(InputPixelPos, tsr_short2(-1, +1));
 		#endif
 
 		bool2 bBrowseVerticaly;
-		taa_short2x2 BrowseDirection;
-		taa_short2x2 EdgeSide = 0;
-		taa_half2 EdgeLuma = 0;
+		tsr_short2x2 BrowseDirection;
+		tsr_short2x2 EdgeSide = 0;
+		tsr_half2 EdgeLuma = 0;
 
-		taa_half2 TotalVarHN = abs(0.5 * (InputNE + InputNW) - InputN);
-		taa_half2 TotalVarH  = abs(0.5 * (InputE  + InputW ) - InputC);
-		taa_half2 TotalVarHS = abs(0.5 * (InputSE + InputSW) - InputS);
+		tsr_half2 TotalVarHN = abs(0.5 * (InputNE + InputNW) - InputN);
+		tsr_half2 TotalVarH  = abs(0.5 * (InputE  + InputW ) - InputC);
+		tsr_half2 TotalVarHS = abs(0.5 * (InputSE + InputSW) - InputS);
 			
-		taa_half2 TotalVarVE = abs(0.5 * (InputNE + InputSE) - InputE);
-		taa_half2 TotalVarV  = abs(0.5 * (InputN  + InputS ) - InputC);
-		taa_half2 TotalVarVW = abs(0.5 * (InputNW + InputSW) - InputW);
+		tsr_half2 TotalVarVE = abs(0.5 * (InputNE + InputSE) - InputE);
+		tsr_half2 TotalVarV  = abs(0.5 * (InputN  + InputS ) - InputC);
+		tsr_half2 TotalVarVW = abs(0.5 * (InputNW + InputSW) - InputW);
 		
-		taa_half2 DiffN = abs(InputN - InputC);
-		taa_half2 DiffS = abs(InputS - InputC);
-		taa_half2 DiffE = abs(InputE - InputC);
-		taa_half2 DiffW = abs(InputW - InputC);
+		tsr_half2 DiffN = abs(InputN - InputC);
+		tsr_half2 DiffS = abs(InputS - InputC);
+		tsr_half2 DiffE = abs(InputE - InputC);
+		tsr_half2 DiffW = abs(InputW - InputC);
 
 		#if CONFIG_COMPUTE_NOISE_FILTERING
 		{
-			NoiseFiltering = max(NoiseFiltering, taa_half(2.0) * TotalVarH - max(DiffE, DiffW));
-			NoiseFiltering = max(NoiseFiltering, taa_half(2.0) * TotalVarV - max(DiffN, DiffS));
+			NoiseFiltering = max(NoiseFiltering, tsr_half(2.0) * TotalVarH - max(DiffE, DiffW));
+			NoiseFiltering = max(NoiseFiltering, tsr_half(2.0) * TotalVarV - max(DiffN, DiffS));
 		}
 		#endif
 
 		bBrowseVerticaly = (TotalVarHN + TotalVarH + TotalVarHS) > (TotalVarVE + TotalVarV + TotalVarVW);
-		BrowseDirection[0] = bBrowseVerticaly ? taa_short(0) : taa_short(1);
-		BrowseDirection[1] = bBrowseVerticaly ? taa_short(1) : taa_short(0);
+		BrowseDirection[0] = bBrowseVerticaly ? tsr_short(0) : tsr_short(1);
+		BrowseDirection[1] = bBrowseVerticaly ? tsr_short(1) : tsr_short(0);
 
 		FLATTEN
 		if (bBrowseVerticaly[0])
@@ -302,14 +302,14 @@ void MainCS(
 			EdgeLuma[1] = DiffN[1] > DiffS[1] ? InputN[1] : InputS[1];
 		}
 
-		taa_half2 LumaDelta = abs(EdgeLuma - InputC) * 0.5;
+		tsr_half2 LumaDelta = abs(EdgeLuma - InputC) * 0.5;
 		
 		bool2 WorthBrowsing = LumaDelta > CONFIG_MIN_LUMA_ALIASING;
 
-		taa_half2 BrowseDelta = saturate(LumaDelta * rcp(taa_half(CONFIG_MIN_LUMA_ALIASING)));
+		tsr_half2 BrowseDelta = saturate(LumaDelta * rcp(tsr_half(CONFIG_MIN_LUMA_ALIASING)));
 		
-		taa_ushort2 EdgeLengthP = 0;
-		taa_ushort2 EdgeLengthN = 0;
+		tsr_ushort2 EdgeLengthP = 0;
+		tsr_ushort2 EdgeLengthN = 0;
 		
 		/**
 		 * Increment:
@@ -327,33 +327,33 @@ void MainCS(
 		bool2 bEdgeStopedByDecrementN = false;
 		
 #if 1 // Optimized for RDNA 2
-		const taa_half BilinearInterp = 0.5;
+		const tsr_half BilinearInterp = 0.5;
 
 		float2x2 fBrowseDirection = float2x2(BrowseDirection) * dpv_interleave_mono_registers(InputInfo_ExtentInverse);
 		float2x2 KernelUV = (float2x2(InputPixelPos) + 0.5 + float2x2(EdgeSide) * BilinearInterp) * dpv_interleave_mono_registers(InputInfo_ExtentInverse);
 
-		taa_half2 MergedEdgeLuma = BilinearInterp * (EdgeLuma + InputC);
+		tsr_half2 MergedEdgeLuma = BilinearInterp * (EdgeLuma + InputC);
 		
-		taa_half2 MinLuma = MergedEdgeLuma - taa_half(BilinearInterp * 0.5) * abs(EdgeLuma - InputC);
-		taa_half2 MaxLuma = MergedEdgeLuma + taa_half(BilinearInterp * 0.5) * abs(EdgeLuma - InputC);
+		tsr_half2 MinLuma = MergedEdgeLuma - tsr_half(BilinearInterp * 0.5) * abs(EdgeLuma - InputC);
+		tsr_half2 MaxLuma = MergedEdgeLuma + tsr_half(BilinearInterp * 0.5) * abs(EdgeLuma - InputC);
 
 		bool2 bEdgeMoreLuminous = EdgeLuma > InputC;
 		
-		taa_ushort2 EdgeStopedByMinLumaP = 0;
-		taa_ushort2 EdgeStopedByMinLumaN = 0;
+		tsr_ushort2 EdgeStopedByMinLumaP = 0;
+		tsr_ushort2 EdgeStopedByMinLumaN = 0;
 		
-		taa_ushort2 EdgeStopedByMaxLumaP = 0;
-		taa_ushort2 EdgeStopedByMaxLumaN = 0;
+		tsr_ushort2 EdgeStopedByMaxLumaP = 0;
+		tsr_ushort2 EdgeStopedByMaxLumaN = 0;
 		
 		bool bStop = false;
 		
 		#if CONFIG_COMPILE_FP16
 			#define asuint(x) asuint16(x)
-			const taa_ushort SignShift = 15;
-			const taa_ushort SignMask = 0x8000;
+			const tsr_ushort SignShift = 15;
+			const tsr_ushort SignMask = 0x8000;
 		#else
-			const taa_ushort SignShift = 31;
-			const taa_ushort SignMask = 0x80000000;
+			const tsr_ushort SignShift = 31;
+			const tsr_ushort SignMask = 0x80000000;
 		#endif
 			
 		EdgeLengthP = CONFIG_ITERATIONS;
@@ -364,15 +364,15 @@ void MainCS(
 		{
 			float2x2 NeighborOffset = fBrowseDirection * float(i + 1);
 
-			taa_half2 SampleLumaP = BilinearSampleInputLDRLuma(KernelUV,  NeighborOffset, /* bIsNegativeDirection = */ false);
-			taa_half2 SampleLumaN = BilinearSampleInputLDRLuma(KernelUV, -NeighborOffset, /* bIsNegativeDirection = */ true);
+			tsr_half2 SampleLumaP = BilinearSampleInputLDRLuma(KernelUV,  NeighborOffset, /* bIsNegativeDirection = */ false);
+			tsr_half2 SampleLumaN = BilinearSampleInputLDRLuma(KernelUV, -NeighborOffset, /* bIsNegativeDirection = */ true);
 				
 			//bEdgeStopedByMinLuma = bEdgeStopedByMinLuma || bool2(SampleLuma < MinLuma) && !bEdgeStopedByMaxLuma;
 			//bEdgeStopedByMaxLuma = bEdgeStopedByMaxLuma || bool2(SampleLuma > MaxLuma) && !bEdgeStopedByMinLuma;
-			taa_ushort2 StopByMinLumaP = asuint(SampleLumaP - MinLuma) & ~EdgeStopedByMaxLumaP;
-			taa_ushort2 StopByMaxLumaP = asuint(MaxLuma - SampleLumaP) & ~EdgeStopedByMinLumaP;
-			taa_ushort2 StopByMinLumaN = asuint(SampleLumaN - MinLuma) & ~EdgeStopedByMaxLumaN;
-			taa_ushort2 StopByMaxLumaN = asuint(MaxLuma - SampleLumaN) & ~EdgeStopedByMinLumaN;
+			tsr_ushort2 StopByMinLumaP = asuint(SampleLumaP - MinLuma) & ~EdgeStopedByMaxLumaP;
+			tsr_ushort2 StopByMaxLumaP = asuint(MaxLuma - SampleLumaP) & ~EdgeStopedByMinLumaP;
+			tsr_ushort2 StopByMinLumaN = asuint(SampleLumaN - MinLuma) & ~EdgeStopedByMaxLumaN;
+			tsr_ushort2 StopByMaxLumaN = asuint(MaxLuma - SampleLumaN) & ~EdgeStopedByMinLumaN;
 
 			EdgeStopedByMinLumaP |= StopByMinLumaP;
 			EdgeStopedByMaxLumaP |= StopByMaxLumaP;
@@ -405,13 +405,13 @@ void MainCS(
 			BRANCH
 			if (!bIsStoped.x)
 			{
-				taa_short2x2 NeighborOffset = BrowseDirection * taa_short(i + 1);
+				tsr_short2x2 NeighborOffset = BrowseDirection * tsr_short(i + 1);
 
-				taa_half2 NeighborP = SampleInputLDRLuma(InputPixelPos, NeighborOffset);
-				taa_half2 EdgeP     = SampleInputLDRLuma(InputPixelPos, NeighborOffset + EdgeSide);
+				tsr_half2 NeighborP = SampleInputLDRLuma(InputPixelPos, NeighborOffset);
+				tsr_half2 EdgeP     = SampleInputLDRLuma(InputPixelPos, NeighborOffset + EdgeSide);
 			
-				taa_half2 NeighborN = SampleInputLDRLuma(InputPixelPos, -NeighborOffset);
-				taa_half2 EdgeN     = SampleInputLDRLuma(InputPixelPos, -NeighborOffset + EdgeSide);
+				tsr_half2 NeighborN = SampleInputLDRLuma(InputPixelPos, -NeighborOffset);
+				tsr_half2 EdgeN     = SampleInputLDRLuma(InputPixelPos, -NeighborOffset + EdgeSide);
 
 				AccumulateEdgeLength(
 					InputC, EdgeLuma,
@@ -432,7 +432,7 @@ void MainCS(
 		}
 #endif
 
-		taa_half TexelOffsetLo = ComputeDistanceToEdge(
+		tsr_half TexelOffsetLo = ComputeDistanceToEdge(
 			bEdgeStopedByIncrementN[0],
 			bEdgeStopedByIncrementP[0],
 			bEdgeStopedByDecrementN[0],
@@ -440,7 +440,7 @@ void MainCS(
 			dpv_lo(EdgeLengthN),
 			dpv_lo(EdgeLengthP));
 		
-		taa_half TexelOffsetHi = ComputeDistanceToEdge(
+		tsr_half TexelOffsetHi = ComputeDistanceToEdge(
 			bEdgeStopedByIncrementN[1],
 			bEdgeStopedByIncrementP[1],
 			bEdgeStopedByDecrementN[1],
@@ -456,7 +456,7 @@ void MainCS(
 		if (!WorthBrowsing[1])
 			TexelOffsetHi = 0.0;
 
-		TexelOffset = dpv_scale(taa_half2x2(EdgeSide), BrowseDelta * dpv_interleave_registers(TexelOffsetLo, TexelOffsetHi));
+		TexelOffset = dpv_scale(tsr_half2x2(EdgeSide), BrowseDelta * dpv_interleave_registers(TexelOffsetLo, TexelOffsetHi));
 		
 		#if 0
 			if (WorthBrowsing.x)
@@ -467,8 +467,8 @@ void MainCS(
 				//Debug[3] = abs(DecodeSpatialAntiAliasingOffset(EncodeSpatialAntiAliasingOffset(TexelOffset))[1]);
 				//Debug[3] = bBrowseVerticaly ? 1 : 0;
 
-				//Debug[0] = EdgeLengthP * taa_ushort2(WorthBrowsing && bBrowseVerticaly);
-				//Debug[1] = EdgeLengthN * taa_ushort2(WorthBrowsing && bBrowseVerticaly);
+				//Debug[0] = EdgeLengthP * tsr_ushort2(WorthBrowsing && bBrowseVerticaly);
+				//Debug[1] = EdgeLengthN * tsr_ushort2(WorthBrowsing && bBrowseVerticaly);
 				//Debug[0] = (1 + EdgeLengthP + EdgeLengthN);
 				//Debug[0] = EdgeLengthN;
 		
@@ -496,7 +496,7 @@ void MainCS(
 	uint2 EncodedTexelOffset = EncodeSpatialAntiAliasingOffset(TexelOffset);
 
 	{
-		taa_short2x2 OutputPixelPos = InputPixelPos;
+		tsr_short2x2 OutputPixelPos = InputPixelPos;
 		
 		OutputPixelPos = InvalidateOutputPixelPos(OutputPixelPos, InputInfo_ViewportMax);
 		
diff --git a/Engine/Shaders/Private/TemporalSuperResolution/TSRUpdateHistory.usf b/Engine/Shaders/Private/TemporalSuperResolution/TSRUpdateHistory.usf
--- a/Engine/Shaders/Private/TemporalSuperResolution/TSRUpdateHistory.usf
+++ b/Engine/Shaders/Private/TemporalSuperResolution/TSRUpdateHistory.usf
@@ -83,19 +83,19 @@
 
 //------------------------------------------------------- PARAMETERS
 
-Texture2D<taa_halfC> InputSceneColorTexture;
+Texture2D<tsr_halfC> InputSceneColorTexture;
 Texture2D<uint2> InputSceneStencilTexture;
-Texture2D<taa_half4> InputSceneTranslucencyTexture;
+Texture2D<tsr_half4> InputSceneTranslucencyTexture;
 
-Texture2D<taa_half> HistoryRejectionTexture;
-Texture2D<taa_half> TranslucencyRejectionTexture;
+Texture2D<tsr_half> HistoryRejectionTexture;
+Texture2D<tsr_half> TranslucencyRejectionTexture;
 
 Texture2D<float2> DilatedVelocityTexture;
 Texture2D<uint> ParallaxFactorTexture;
-Texture2D<taa_half> ParallaxRejectionMaskTexture;
-Texture2D<taa_ushort> AntiAliasingTexture;
-Texture2D<taa_half> NoiseFilteringTexture;
-Texture2D<taa_half2> HoleFilledVelocityMaskTexture;
+Texture2D<tsr_half> ParallaxRejectionMaskTexture;
+Texture2D<tsr_ushort> AntiAliasingTexture;
+Texture2D<tsr_half> NoiseFilteringTexture;
+Texture2D<tsr_half2> HoleFilledVelocityMaskTexture;
 
 float4 HistoryPixelPosToScreenPos;
 float4 HistoryPixelPosToPPCo;
@@ -106,22 +106,22 @@ float InputToHistoryFactor;
 uint ResponsiveStencilMask;
 uint bGenerateOutputMip1;
 
-Texture2D<taa_halfC> PrevHistory_LowFrequency;
-Texture2D<taa_halfC> PrevHistory_HighFrequency;
-Texture2D<taa_half2> PrevHistory_Metadata;
-Texture2D<taa_half3> PrevHistory_Translucency;
-Texture2D<taa_half>  PrevHistory_TranslucencyAlpha;
-Texture2D<taa_subpixel_details> PrevHistory_SubpixelDetails;
+Texture2D<tsr_halfC> PrevHistory_LowFrequency;
+Texture2D<tsr_halfC> PrevHistory_HighFrequency;
+Texture2D<tsr_half2> PrevHistory_Metadata;
+Texture2D<tsr_half3> PrevHistory_Translucency;
+Texture2D<tsr_half>  PrevHistory_TranslucencyAlpha;
+Texture2D<tsr_subpixel_details> PrevHistory_SubpixelDetails;
 
-RWTexture2D<taa_halfC> HistoryOutput_LowFrequency;
-RWTexture2D<taa_halfC> HistoryOutput_HighFrequency;
-RWTexture2D<taa_half2> HistoryOutput_Metadata;
-RWTexture2D<taa_half3> HistoryOutput_Translucency;
-RWTexture2D<taa_half>  HistoryOutput_TranslucencyAlpha;
-RWTexture2D<taa_subpixel_details> HistoryOutput_SubpixelDetails;
+RWTexture2D<tsr_halfC> HistoryOutput_LowFrequency;
+RWTexture2D<tsr_halfC> HistoryOutput_HighFrequency;
+RWTexture2D<tsr_half2> HistoryOutput_Metadata;
+RWTexture2D<tsr_half3> HistoryOutput_Translucency;
+RWTexture2D<tsr_half>  HistoryOutput_TranslucencyAlpha;
+RWTexture2D<tsr_subpixel_details> HistoryOutput_SubpixelDetails;
 
-RWTexture2D<taa_halfC> SceneColorOutputMip0;
-RWTexture2D<taa_halfC> SceneColorOutputMip1;
+RWTexture2D<tsr_halfC> SceneColorOutputMip0;
+RWTexture2D<tsr_halfC> SceneColorOutputMip1;
 
 
 //------------------------------------------------------- LDS
@@ -132,44 +132,44 @@ RWTexture2D<taa_halfC> SceneColorOutputMip1;
 #define GROUP_SHARED_SIZE (TILE_SIZE * TILE_SIZE)
 #endif
 
-groupshared taa_half4 SharedArray0[GROUP_SHARED_SIZE];
-groupshared taa_half4 SharedArray1[GROUP_SHARED_SIZE];
-groupshared taa_half4 SharedArray2[GROUP_SHARED_SIZE];
-groupshared taa_half4 SharedArray3[GROUP_SHARED_SIZE];
-groupshared taa_half4 SharedArray4[GROUP_SHARED_SIZE];
-groupshared taa_half4 SharedArray5[GROUP_SHARED_SIZE];
+groupshared tsr_half4 SharedArray0[GROUP_SHARED_SIZE];
+groupshared tsr_half4 SharedArray1[GROUP_SHARED_SIZE];
+groupshared tsr_half4 SharedArray2[GROUP_SHARED_SIZE];
+groupshared tsr_half4 SharedArray3[GROUP_SHARED_SIZE];
+groupshared tsr_half4 SharedArray4[GROUP_SHARED_SIZE];
+groupshared tsr_half4 SharedArray5[GROUP_SHARED_SIZE];
 
-groupshared taa_half4 SharedArray6[GROUP_SHARED_SIZE];
-groupshared taa_half4 SharedArray7[GROUP_SHARED_SIZE];
-groupshared taa_half4 SharedArray8[GROUP_SHARED_SIZE];
-groupshared taa_half4 SharedArray9[GROUP_SHARED_SIZE];
-groupshared taa_half4 SharedArray10[GROUP_SHARED_SIZE];
-groupshared taa_half4 SharedArray11[GROUP_SHARED_SIZE];
+groupshared tsr_half4 SharedArray6[GROUP_SHARED_SIZE];
+groupshared tsr_half4 SharedArray7[GROUP_SHARED_SIZE];
+groupshared tsr_half4 SharedArray8[GROUP_SHARED_SIZE];
+groupshared tsr_half4 SharedArray9[GROUP_SHARED_SIZE];
+groupshared tsr_half4 SharedArray10[GROUP_SHARED_SIZE];
+groupshared tsr_half4 SharedArray11[GROUP_SHARED_SIZE];
 
-groupshared taa_subpixel_payload2 SharedParallaxFactorBits[GROUP_SHARED_SIZE];
+groupshared tsr_subpixel_payload2 SharedParallaxFactorBits[GROUP_SHARED_SIZE];
 
 
 //------------------------------------------------------- FUNCTIONS
 
 CALL_SITE_DEBUGLOC 
-taa_ushort2x2 Map8x8Tile2x2LaneDPV(uint GroupThreadIndex)
+tsr_ushort2x2 Map8x8Tile2x2LaneDPV(uint GroupThreadIndex)
 {
-	taa_ushort2 GroupId = Map8x8Tile2x2Lane(GroupThreadIndex);
+	tsr_ushort2 GroupId = Map8x8Tile2x2Lane(GroupThreadIndex);
 #if CONFIG_ENABLE_DUAL_PIXEL_VECTORIZATION
-	return dpv_interleave_registers(GroupId, GroupId + taa_ushort2(0,  8 / 2));
+	return dpv_interleave_registers(GroupId, GroupId + tsr_ushort2(0,  8 / 2));
 #else
 	return dpv_interleave_mono_registers(GroupId);
 #endif
 }
 
 CALL_SITE_DEBUGLOC 
-taa_halfCx2 TransformColorForClampingBox(taa_halfCx2 Color)
+tsr_halfCx2 TransformColorForClampingBox(tsr_halfCx2 Color)
 {
 	return RGBToYCoCg(Color);
 }
 
 CALL_SITE_DEBUGLOC 
-taa_half3x2 GetColorChannel(taa_half4x2 Color)
+tsr_half3x2 GetColorChannel(tsr_half4x2 Color)
 {
 	return dpv_interleave_registers(
 		dpv_lo(Color).rgb,
@@ -177,17 +177,17 @@ taa_half3x2 GetColorChannel(taa_half4x2 Color)
 }
 
 void ComputeInputKernelSamplePosition(
-	taa_short2x2 InputPixelPos,
-	taa_half2x2 dKO,
+	tsr_short2x2 InputPixelPos,
+	tsr_half2x2 dKO,
 	uint SampleId,
-	out taa_short2x2 SampleInputPixelPos,
-	out taa_half2x2 PixelOffset)
+	out tsr_short2x2 SampleInputPixelPos,
+	out tsr_half2x2 PixelOffset)
 {
 	
 	#if CONFIG_SAMPLES == 9
 	{
-		taa_short2 iPixelOffset = taa_short2(kOffsets3x3[kSquareIndexes3x3[SampleId]]);
-		PixelOffset = dpv_interleave_registers(taa_half2(iPixelOffset), taa_half2(iPixelOffset));
+		tsr_short2 iPixelOffset = tsr_short2(kOffsets3x3[kSquareIndexes3x3[SampleId]]);
+		PixelOffset = dpv_interleave_registers(tsr_half2(iPixelOffset), tsr_half2(iPixelOffset));
 				
 		SampleInputPixelPos = InputPixelPos + ClampPixelOffset(
 			InputPixelPos,
@@ -199,7 +199,7 @@ void ComputeInputKernelSamplePosition(
 		#if CONFIG_SAMPLES == 6
 		if (SampleId == 5)
 		{
-			taa_short2x2 iPixelOffset;
+			tsr_short2x2 iPixelOffset;
 			#if CONFIG_COMPILE_FP16
 				iPixelOffset[0] = int16_t(1) - int16_t2((asuint16(dKO[0]) & uint16_t(0x8000)) >> uint16_t(14));
 				iPixelOffset[1] = int16_t(1) - int16_t2((asuint16(dKO[1]) & uint16_t(0x8000)) >> uint16_t(14));
@@ -217,8 +217,8 @@ void ComputeInputKernelSamplePosition(
 		else
 		#endif
 		{
-			taa_short2 iPixelOffset = taa_short2(kOffsets3x3[kPlusIndexes3x3[SampleId]]);
-			PixelOffset = dpv_interleave_registers(taa_half2(iPixelOffset), taa_half2(iPixelOffset));
+			tsr_short2 iPixelOffset = tsr_short2(kOffsets3x3[kPlusIndexes3x3[SampleId]]);
+			PixelOffset = dpv_interleave_registers(tsr_half2(iPixelOffset), tsr_half2(iPixelOffset));
 					
 			SampleInputPixelPos = InputPixelPos + ClampPixelOffset(
 				InputPixelPos,
@@ -255,9 +255,9 @@ void MainCS(
 
 	float4x2 Debug = 0.0;
 
-	taa_short2x2 HistoryPixelPos = dpv_add(
-		taa_short2(GroupId) * taa_short2(TILE_SIZE, TILE_SIZE),
-		taa_short2x2(Map8x8Tile2x2LaneDPV(GroupThreadIndex)));
+	tsr_short2x2 HistoryPixelPos = dpv_add(
+		tsr_short2(GroupId) * tsr_short2(TILE_SIZE, TILE_SIZE),
+		tsr_short2x2(Map8x8Tile2x2LaneDPV(GroupThreadIndex)));
 
 	float2x2 ScreenPos = ApplyScreenTransform(float2x2(HistoryPixelPos), HistoryPixelPosToScreenPos);
 	
@@ -267,19 +267,19 @@ void MainCS(
 	// Pixel coordinate of the center of the nearest input pixel K.
 	float2x2 PPCk = floor(PPCo) + 0.5;
 	
-	taa_short2x2 InputPixelPos = ClampPixelOffset(
-		taa_short2x2(PPCo),
+	tsr_short2x2 InputPixelPos = ClampPixelOffset(
+		tsr_short2x2(PPCo),
 		InputPixelPosMin, InputPixelPosMax);
 
 	// Fetch reprojections related informations.
 	float2x2 PrevScreenPos = ScreenPos;
-	taa_half2 ParallaxRejectionMask = taa_half(1.0).xx;
-	taa_half2 LowFrequencyRejection = taa_half(1.0).xx;
-	taa_half2 TranslucencyRejection = taa_half(1.0).xx;
-	taa_half2 OutputPixelVelocity = taa_half(0.0).xx;
-	taa_half2 NoiseFiltering = taa_half(0.0).xx;
-	taa_half2x2 HoleFilledVelocityMask = dpv_interleave_mono_registers(taa_half(1.0).xx);
-	taa_subpixel_payload2 ParallaxFactorBits = taa_subpixel_payload(0x0).xx;
+	tsr_half2 ParallaxRejectionMask = tsr_half(1.0).xx;
+	tsr_half2 LowFrequencyRejection = tsr_half(1.0).xx;
+	tsr_half2 TranslucencyRejection = tsr_half(1.0).xx;
+	tsr_half2 OutputPixelVelocity = tsr_half(0.0).xx;
+	tsr_half2 NoiseFiltering = tsr_half(0.0).xx;
+	tsr_half2x2 HoleFilledVelocityMask = dpv_interleave_mono_registers(tsr_half(1.0).xx);
+	tsr_subpixel_payload2 ParallaxFactorBits = tsr_subpixel_payload(0x0).xx;
 	bool2 bIsResponsiveAAPixel = false;
 
 	PLATFORM_SPECIFIC_ISOLATE
@@ -290,30 +290,30 @@ void MainCS(
 			float2 RawEncodedVelocity[DPV_PIXEL_PER_LANE];
 		#endif
 
-		taa_half RawParallaxRejectionMask[DPV_PIXEL_PER_LANE];
+		tsr_half RawParallaxRejectionMask[DPV_PIXEL_PER_LANE];
 		uint RawParallaxFactorBits[DPV_PIXEL_PER_LANE];
-		taa_half RawLowFrequencyRejection[DPV_PIXEL_PER_LANE];
-		taa_half RawTranslucencyRejection[DPV_PIXEL_PER_LANE];
-		taa_half RawTranslucencyRejectionNeighborhood[CONFIG_SAMPLES - 1][DPV_PIXEL_PER_LANE];
+		tsr_half RawLowFrequencyRejection[DPV_PIXEL_PER_LANE];
+		tsr_half RawTranslucencyRejection[DPV_PIXEL_PER_LANE];
+		tsr_half RawTranslucencyRejectionNeighborhood[CONFIG_SAMPLES - 1][DPV_PIXEL_PER_LANE];
 		
 		uint RawSceneStencilRef[DPV_PIXEL_PER_LANE];
 		uint RawSceneStencilRefNeighborhood[CONFIG_SAMPLES - 1][DPV_PIXEL_PER_LANE];
 		
 		#if CONFIG_REJECTION_ANTI_ALIASING
-			taa_ushort RawEncodedInputTexelOffset[DPV_PIXEL_PER_LANE];
-			taa_half RawNoiseFiltering[DPV_PIXEL_PER_LANE]; 
-			taa_half2 RawHoleFilledVelocityMask[DPV_PIXEL_PER_LANE];
+			tsr_ushort RawEncodedInputTexelOffset[DPV_PIXEL_PER_LANE];
+			tsr_half RawNoiseFiltering[DPV_PIXEL_PER_LANE]; 
+			tsr_half2 RawHoleFilledVelocityMask[DPV_PIXEL_PER_LANE];
 		#endif
 
 		// Issue overlapped texture fetches
 		PLATFORM_SPECIFIC_ISOLATE
 		{
-			taa_ushort2x2 RejectionPixelPos = dpv_sub(taa_ushort2x2(InputPixelPos), taa_ushort2(InputPixelPosMin)) / 2;
+			tsr_ushort2x2 RejectionPixelPos = dpv_sub(tsr_ushort2x2(InputPixelPos), tsr_ushort2(InputPixelPosMin)) / 2;
 
 			UNROLL_N(DPV_PIXEL_PER_LANE)
 			for (uint PixelId = 0; PixelId < DPV_PIXEL_PER_LANE; PixelId++)
 			{
-				taa_ushort2 LocalInputPixelPos = dpv_access_pixel(InputPixelPos, PixelId);
+				tsr_ushort2 LocalInputPixelPos = dpv_access_pixel(InputPixelPos, PixelId);
 				
 				#if !CONFIG_INTERPOLATE_VELOCITY
 					RawEncodedVelocity[PixelId]   = DilatedVelocityTexture[LocalInputPixelPos];
@@ -335,13 +335,13 @@ void MainCS(
 			}
 
 			// Vector in pixel between pixel K -> O.
-			taa_half2x2 dKO = taa_half2x2(PPCo - PPCk);
+			tsr_half2x2 dKO = tsr_half2x2(PPCo - PPCk);
 		
 			UNROLL_N(CONFIG_SAMPLES - 1)
 			for (uint NeighborId = 0; NeighborId < (CONFIG_SAMPLES - 1); NeighborId++)
 			{
-				taa_short2x2 SampleInputPixelPos;
-				taa_half2x2 PixelOffset;
+				tsr_short2x2 SampleInputPixelPos;
+				tsr_half2x2 PixelOffset;
 				ComputeInputKernelSamplePosition(
 					InputPixelPos, dKO, /* SampleId = */ NeighborId + 1,
 					/* out */ SampleInputPixelPos,
@@ -360,14 +360,14 @@ void MainCS(
 			UNROLL_N(DPV_PIXEL_PER_LANE)
 			for (uint PixelId2 = 0; PixelId2 < DPV_PIXEL_PER_LANE; PixelId2++)
 			{
-				taa_short2 LocalInputPixelPos0 = dpv_access_pixel(InputPixelPos, PixelId2);
-				taa_short2 LocalInputPixelPos1 = LocalInputPixelPos0 + taa_short2(sign(dpv_access_pixel(dKO, PixelId2)));
+				tsr_short2 LocalInputPixelPos0 = dpv_access_pixel(InputPixelPos, PixelId2);
+				tsr_short2 LocalInputPixelPos1 = LocalInputPixelPos0 + tsr_short2(sign(dpv_access_pixel(dKO, PixelId2)));
 				LocalInputPixelPos1 = ClampPixelOffset(LocalInputPixelPos1, InputPixelPosMin, InputPixelPosMax);
 
-				RawEncodedVelocityNeighborhood[0][PixelId2] = DilatedVelocityTexture[taa_short2(LocalInputPixelPos0.x, LocalInputPixelPos0.y)];
-				RawEncodedVelocityNeighborhood[1][PixelId2] = DilatedVelocityTexture[taa_short2(LocalInputPixelPos1.x, LocalInputPixelPos0.y)];
-				RawEncodedVelocityNeighborhood[2][PixelId2] = DilatedVelocityTexture[taa_short2(LocalInputPixelPos0.x, LocalInputPixelPos1.y)];
-				RawEncodedVelocityNeighborhood[3][PixelId2] = DilatedVelocityTexture[taa_short2(LocalInputPixelPos1.x, LocalInputPixelPos1.y)];
+				RawEncodedVelocityNeighborhood[0][PixelId2] = DilatedVelocityTexture[tsr_short2(LocalInputPixelPos0.x, LocalInputPixelPos0.y)];
+				RawEncodedVelocityNeighborhood[1][PixelId2] = DilatedVelocityTexture[tsr_short2(LocalInputPixelPos1.x, LocalInputPixelPos0.y)];
+				RawEncodedVelocityNeighborhood[2][PixelId2] = DilatedVelocityTexture[tsr_short2(LocalInputPixelPos0.x, LocalInputPixelPos1.y)];
+				RawEncodedVelocityNeighborhood[3][PixelId2] = DilatedVelocityTexture[tsr_short2(LocalInputPixelPos1.x, LocalInputPixelPos1.y)];
 			}
 			#endif
 		}
@@ -378,7 +378,7 @@ void MainCS(
 		PLATFORM_SPECIFIC_ISOLATE
 		{
 			// Vector in pixel between pixel K -> O.
-			taa_half2x2 dKO = taa_half2x2(PPCo - PPCk);
+			tsr_half2x2 dKO = tsr_half2x2(PPCo - PPCk);
 		
 			// Process velocity
 			{
@@ -392,7 +392,7 @@ void MainCS(
 					float2 V01 = DecodeVelocityFromTexture(float4(RawEncodedVelocityNeighborhood[2][PixelId], 0.0, 0.0)).xy;
 					float2 V11 = DecodeVelocityFromTexture(float4(RawEncodedVelocityNeighborhood[3][PixelId], 0.0, 0.0)).xy;
 
-					taa_half2 BilinearInterp = abs(dpv_access_pixel(dKO, PixelId));
+					tsr_half2 BilinearInterp = abs(dpv_access_pixel(dKO, PixelId));
 					
 					float Bilinear00 = saturate(1.0 - BilinearInterp.x) * saturate(1.0 - BilinearInterp.y);
 					float Bilinear10 = saturate(0.0 + BilinearInterp.x) * saturate(1.0 - BilinearInterp.y);
@@ -420,12 +420,12 @@ void MainCS(
 				#endif
 
 				PrevScreenPos = ScreenPos - ScreenVelocity;
-				OutputPixelVelocity = taa_half2(dpv_length(dpv_mul(ScreenVelocity, HistoryInfo_ViewportSize)));
+				OutputPixelVelocity = tsr_half2(dpv_length(dpv_mul(ScreenVelocity, HistoryInfo_ViewportSize)));
 			}
 
 			// Process input texel, forcing to pack register.
 			ParallaxRejectionMask = dpv_force_interleave_registers_array(RawParallaxRejectionMask);
-			ParallaxFactorBits    = dpv_force_interleave_registers_array(RawParallaxFactorBits) & taa_subpixel_payload(SUB_PIXEL_PARALLAX_FACTOR_BIT_MASK);
+			ParallaxFactorBits    = dpv_force_interleave_registers_array(RawParallaxFactorBits) & tsr_subpixel_payload(SUB_PIXEL_PARALLAX_FACTOR_BIT_MASK);
 			LowFrequencyRejection = dpv_force_interleave_registers_array(RawLowFrequencyRejection);
 			TranslucencyRejection = dpv_force_interleave_registers_array(RawTranslucencyRejection);
 			SceneStencilRef       = dpv_force_interleave_registers_array(RawSceneStencilRef);
@@ -447,22 +447,22 @@ void MainCS(
 				IsOffScreen(bCameraCut, dpv_lo(PrevScreenPos), dpv_lo(ParallaxRejectionMask)),
 				IsOffScreen(bCameraCut, dpv_hi(PrevScreenPos), dpv_hi(ParallaxRejectionMask)));
 			
-			taa_half2 TotalRejection = taa_half2(
+			tsr_half2 TotalRejection = tsr_half2(
 				bOffScreen[0] ? 0.0 : saturate(LowFrequencyRejection[0] * 4.0),
 				bOffScreen[1] ? 0.0 : saturate(LowFrequencyRejection[1] * 4.0));
 
-			taa_ushort2 EncodedInputTexelOffset = dpv_force_interleave_registers_array(RawEncodedInputTexelOffset);
+			tsr_ushort2 EncodedInputTexelOffset = dpv_force_interleave_registers_array(RawEncodedInputTexelOffset);
 			
 			NoiseFiltering = dpv_force_interleave_registers_array(RawNoiseFiltering);
-			HoleFilledVelocityMask = dpv_interleave_mono_registers(taa_half(1.0).xx) - dpv_force_interleave_registers_array(RawHoleFilledVelocityMask);
+			HoleFilledVelocityMask = dpv_interleave_mono_registers(tsr_half(1.0).xx) - dpv_force_interleave_registers_array(RawHoleFilledVelocityMask);
 
-			taa_half2x2 TexelOffset = DecodeSpatialAntiAliasingOffset(EncodedInputTexelOffset);
+			tsr_half2x2 TexelOffset = DecodeSpatialAntiAliasingOffset(EncodedInputTexelOffset);
 			
 			{
-				PPCo += dpv_scale(TexelOffset, taa_half(1.0) - TotalRejection);
+				PPCo += dpv_scale(TexelOffset, tsr_half(1.0) - TotalRejection);
 				PPCk = floor(PPCo) + 0.5;
 				InputPixelPos = ClampPixelOffset(
-					taa_short2x2(PPCo),
+					tsr_short2x2(PPCo),
 					InputPixelPosMin, InputPixelPosMax);
 			}
 		}
@@ -470,7 +470,7 @@ void MainCS(
 		
 		// Final post processing.
 		{
-			TranslucencyRejection = max(TranslucencyRejection, taa_half(MinTranslucencyRejection));
+			TranslucencyRejection = max(TranslucencyRejection, tsr_half(MinTranslucencyRejection));
 		
 			// Fetch wehther the pixel is responsive AA or not.
 			#if CONFIG_RESPONSIVE_STENCIL
@@ -478,16 +478,16 @@ void MainCS(
 				bIsResponsiveAAPixel = (SceneStencilRef & ResponsiveStencilMask) != 0;
 
 				#if CONFIG_REJECT_TRAMSLUCENCY_ON_RESPONSIVE_AA_ONLY
-					TranslucencyRejection = max(TranslucencyRejection, bIsResponsiveAAPixel ? taa_half(0.0) : taa_half(1.0));
+					TranslucencyRejection = max(TranslucencyRejection, bIsResponsiveAAPixel ? tsr_half(0.0) : tsr_half(1.0));
 				#endif
 			}
 			#endif
 
 			#if !CONFIG_CLAMP
 			{
-				ParallaxRejectionMask = taa_half(1.0).xx;
-				LowFrequencyRejection = taa_half(1.0).xx;
-				TranslucencyRejection = taa_half(1.0).xx;
+				ParallaxRejectionMask = tsr_half(1.0).xx;
+				LowFrequencyRejection = tsr_half(1.0).xx;
+				TranslucencyRejection = tsr_half(1.0).xx;
 			}
 			#endif
 		}
@@ -503,22 +503,22 @@ void MainCS(
 		IsOffScreen(bCameraCut, dpv_lo(PrevScreenPos), dpv_lo(ParallaxRejectionMask)),
 		IsOffScreen(bCameraCut, dpv_hi(PrevScreenPos), dpv_hi(ParallaxRejectionMask)));
 	
-	taa_half2 TotalRejection = taa_half2(
+	tsr_half2 TotalRejection = tsr_half2(
 		bOffScreen[0] ? 0.0 : saturate(LowFrequencyRejection[0] * 4.0),
 		bOffScreen[1] ? 0.0 : saturate(LowFrequencyRejection[1] * 4.0));
 
 
 	// Filter input scene color at predictor frequency.
-	taa_halfCx2 FilteredInputColor;
-	taa_halfCx2 InputMinColor;
-	taa_halfCx2 InputMaxColor;
+	tsr_halfCx2 FilteredInputColor;
+	tsr_halfCx2 InputMinColor;
+	tsr_halfCx2 InputMaxColor;
 	
-	taa_half4x2 FilteredTranslucencyColor;
-	taa_half4x2 TranslucencyMinColor;
-	taa_half4x2 TranslucencyMaxColor;
+	tsr_half4x2 FilteredTranslucencyColor;
+	tsr_half4x2 TranslucencyMinColor;
+	tsr_half4x2 TranslucencyMaxColor;
 
-	taa_half2 InputPixelAlignement;
-	taa_half2 ClosestInputLuma4;
+	tsr_half2 InputPixelAlignement;
+	tsr_half2 ClosestInputLuma4;
 	
 	
 	bool2 bUpdateSubpixelDetails;
@@ -527,7 +527,7 @@ void MainCS(
 	PLATFORM_SPECIFIC_ISOLATE
 	{
 		// Vector in pixel between pixel K -> O.
-		taa_half2x2 dKO = taa_half2x2(PPCo - PPCk);
+		tsr_half2x2 dKO = tsr_half2x2(PPCo - PPCk);
 
 		// Compresses the sign of dKO into SGPR to keep VGPR pressure low.
 		{
@@ -540,44 +540,44 @@ void MainCS(
 			bKOSign[1] = dKO[1] < 0.0;
 		}
 		
-		taa_halfC RawInputColorArray[CONFIG_SAMPLES][DPV_PIXEL_PER_LANE];
-		taa_half4 RawTranslucencyColorArray[CONFIG_SAMPLES][DPV_PIXEL_PER_LANE];
+		tsr_halfC RawInputColorArray[CONFIG_SAMPLES][DPV_PIXEL_PER_LANE];
+		tsr_half4 RawTranslucencyColorArray[CONFIG_SAMPLES][DPV_PIXEL_PER_LANE];
 
-		taa_half2x2 RawdPPArray[CONFIG_SAMPLES];
-		taa_half2   RawSampleSpatialWeightArray[CONFIG_SAMPLES];
+		tsr_half2x2 RawdPPArray[CONFIG_SAMPLES];
+		tsr_half2   RawSampleSpatialWeightArray[CONFIG_SAMPLES];
 		
 		// Issues overlapped texture fetches
 		PLATFORM_SPECIFIC_ISOLATE
 		{
 			// Soften the spatial kernel when there is high noise in the input to make parrallax disocclusion of detailed geometry less distracting.
-			taa_half2 OffScreenInputToHistoryFactor = taa_half(1.0) - taa_half(0.5) * NoiseFiltering;
+			tsr_half2 OffScreenInputToHistoryFactor = tsr_half(1.0) - tsr_half(0.5) * NoiseFiltering;
 
-			taa_half2 KernelInputToHistoryFactor = taa_half(InputToHistoryFactor).xx;
+			tsr_half2 KernelInputToHistoryFactor = tsr_half(InputToHistoryFactor).xx;
 			
 			FLATTEN
 			if (bOffScreen.x)
 			{
-				KernelInputToHistoryFactor.x = taa_half(OffScreenInputToHistoryFactor.x);
+				KernelInputToHistoryFactor.x = tsr_half(OffScreenInputToHistoryFactor.x);
 			}
 			
 			FLATTEN
 			if (bOffScreen.y)
 			{
-				KernelInputToHistoryFactor.y = taa_half(OffScreenInputToHistoryFactor.y);
+				KernelInputToHistoryFactor.y = tsr_half(OffScreenInputToHistoryFactor.y);
 			}
 
 			UNROLL_N(CONFIG_SAMPLES)
 			for (uint SampleId = 0; SampleId < CONFIG_SAMPLES; SampleId++)
 			{
-				taa_short2x2 SampleInputPixelPos;
-				taa_half2x2 PixelOffset;
+				tsr_short2x2 SampleInputPixelPos;
+				tsr_half2x2 PixelOffset;
 				ComputeInputKernelSamplePosition(
 					InputPixelPos, dKO, SampleId,
 					/* out */ SampleInputPixelPos,
 					/* out */ PixelOffset);
 				
-				taa_half2x2 dPP = PixelOffset - dKO;
-				taa_half2 SampleSpatialWeight = ComputeSampleWeigth(KernelInputToHistoryFactor, dPP, /* MinimalContribution = */ float(0.005));
+				tsr_half2x2 dPP = PixelOffset - dKO;
+				tsr_half2 SampleSpatialWeight = ComputeSampleWeigth(KernelInputToHistoryFactor, dPP, /* MinimalContribution = */ float(0.005));
 				
 				RawdPPArray[SampleId] = dPP; 
 				RawSampleSpatialWeightArray[SampleId] = SampleSpatialWeight;
@@ -593,9 +593,9 @@ void MainCS(
 			} // for (uint SampleId = 0; SampleId < CONFIG_SAMPLES; SampleId++)
 		}
 			
-		taa_half2 FilteredInputColorWeight = taa_half(0.0);
+		tsr_half2 FilteredInputColorWeight = tsr_half(0.0);
 		
-		FilteredInputColor = taa_half(0.0);
+		FilteredInputColor = tsr_half(0.0);
 		#if CONFIG_SCENE_COLOR_APLHA
 			InputMinColor = dpv_interleave_mono_registers(+LargestSceneColorYCoCgA);
 			InputMaxColor = dpv_interleave_mono_registers(-LargestSceneColorYCoCgA);
@@ -609,9 +609,9 @@ void MainCS(
 		UNROLL_N(CONFIG_SAMPLES)
 		for (uint SampleId = 0; SampleId < CONFIG_SAMPLES; SampleId++)
 		{
-			taa_half2 SampleSpatialWeight = RawSampleSpatialWeightArray[SampleId];
-			taa_halfCx2 InputColor = dpv_force_interleave_registers_array(RawInputColorArray[SampleId]);
-			taa_half4x2 TranslucencyColor = dpv_force_interleave_registers_array(RawTranslucencyColorArray[SampleId]);
+			tsr_half2 SampleSpatialWeight = RawSampleSpatialWeightArray[SampleId];
+			tsr_halfCx2 InputColor = dpv_force_interleave_registers_array(RawInputColorArray[SampleId]);
+			tsr_half4x2 TranslucencyColor = dpv_force_interleave_registers_array(RawTranslucencyColorArray[SampleId]);
 
 			#if CONFIG_SCENE_COLOR_OVERFLOW
 				#if CONFIG_SCENE_COLOR_APLHA
@@ -621,13 +621,13 @@ void MainCS(
 				#endif
 			#endif
 
-			taa_half2 ToneWeight = HdrWeight4(InputColor);
+			tsr_half2 ToneWeight = HdrWeight4(InputColor);
 
 			#if CONFIG_VFX_SURVIVAL
 			if (SampleId == 0)
 			{
-				taa_half2x2 dPP =  RawdPPArray[SampleId];
-				ToneWeight = max(ToneWeight, ToneWeight > taa_half(0.2) ? taa_half(0.0).xx : ComputeSampleWeigth(taa_half(1.0), dPP, /* MinimalContribution = */ 0.0));
+				tsr_half2x2 dPP =  RawdPPArray[SampleId];
+				ToneWeight = max(ToneWeight, ToneWeight > tsr_half(0.2) ? tsr_half(0.0).xx : ComputeSampleWeigth(tsr_half(1.0), dPP, /* MinimalContribution = */ 0.0));
 			}
 			#endif
 
@@ -657,7 +657,7 @@ void MainCS(
 		
 		FilteredInputColor = dpv_scale(FilteredInputColor, rcp(FilteredInputColorWeight));
 
-		InputPixelAlignement = ComputeSampleWeigth(taa_half(InputToHistoryFactor), dKO, /* MinimalContribution = */ float(0.0));
+		InputPixelAlignement = ComputeSampleWeigth(tsr_half(InputToHistoryFactor), dKO, /* MinimalContribution = */ float(0.0));
 		
 		#if CONFIG_SCENE_COLOR_OVERFLOW
 		{
@@ -679,55 +679,55 @@ void MainCS(
 		uint LocalGroupThreadIndex = GetGroupThreadIndex(GroupThreadIndex, GroupWaveIndex);
 
 		#if CONFIG_ENABLE_DUAL_PIXEL_VECTORIZATION
-			SharedArray0[LocalGroupThreadIndex] = taa_half4(FilteredInputColor[0], FilteredInputColor[1]);
-			SharedArray1[LocalGroupThreadIndex] = taa_half4(FilteredInputColor[2], LowFrequencyRejection);
-			SharedArray2[LocalGroupThreadIndex] = taa_half4(InputMinColor[0], InputMinColor[1]);
-			SharedArray3[LocalGroupThreadIndex] = taa_half4(InputMinColor[2], InputPixelAlignement);
-			SharedArray4[LocalGroupThreadIndex] = taa_half4(InputMaxColor[0], InputMaxColor[1]);
-			SharedArray5[LocalGroupThreadIndex] = taa_half4(InputMaxColor[2], OutputPixelVelocity);
+			SharedArray0[LocalGroupThreadIndex] = tsr_half4(FilteredInputColor[0], FilteredInputColor[1]);
+			SharedArray1[LocalGroupThreadIndex] = tsr_half4(FilteredInputColor[2], LowFrequencyRejection);
+			SharedArray2[LocalGroupThreadIndex] = tsr_half4(InputMinColor[0], InputMinColor[1]);
+			SharedArray3[LocalGroupThreadIndex] = tsr_half4(InputMinColor[2], InputPixelAlignement);
+			SharedArray4[LocalGroupThreadIndex] = tsr_half4(InputMaxColor[0], InputMaxColor[1]);
+			SharedArray5[LocalGroupThreadIndex] = tsr_half4(InputMaxColor[2], OutputPixelVelocity);
 
 			#if CONFIG_ACCUMULATE_TRANSLUCENCY_SEPARATLY
-				SharedArray6[LocalGroupThreadIndex] = taa_half4(FilteredTranslucencyColor[0], FilteredTranslucencyColor[1]);
-				SharedArray7[LocalGroupThreadIndex] = taa_half4(FilteredTranslucencyColor[2], FilteredTranslucencyColor[3]);
-				SharedArray8[LocalGroupThreadIndex] = taa_half4(TranslucencyMinColor[0], TranslucencyMinColor[1]);
-				SharedArray9[LocalGroupThreadIndex] = taa_half4(TranslucencyMinColor[2], TranslucencyMinColor[3]);
-				SharedArray10[LocalGroupThreadIndex] = taa_half4(TranslucencyMaxColor[0], TranslucencyMaxColor[1]);
-				SharedArray11[LocalGroupThreadIndex] = taa_half4(TranslucencyMaxColor[2], TranslucencyMaxColor[3]);
+				SharedArray6[LocalGroupThreadIndex] = tsr_half4(FilteredTranslucencyColor[0], FilteredTranslucencyColor[1]);
+				SharedArray7[LocalGroupThreadIndex] = tsr_half4(FilteredTranslucencyColor[2], FilteredTranslucencyColor[3]);
+				SharedArray8[LocalGroupThreadIndex] = tsr_half4(TranslucencyMinColor[0], TranslucencyMinColor[1]);
+				SharedArray9[LocalGroupThreadIndex] = tsr_half4(TranslucencyMinColor[2], TranslucencyMinColor[3]);
+				SharedArray10[LocalGroupThreadIndex] = tsr_half4(TranslucencyMaxColor[0], TranslucencyMaxColor[1]);
+				SharedArray11[LocalGroupThreadIndex] = tsr_half4(TranslucencyMaxColor[2], TranslucencyMaxColor[3]);
 			#endif
 
 			SharedParallaxFactorBits[LocalGroupThreadIndex] = ParallaxFactorBits;
 		#else
-			SharedArray0[LocalGroupThreadIndex] = taa_half4(dpv_lo(FilteredInputColor), dpv_lo(LowFrequencyRejection));
-			SharedArray1[LocalGroupThreadIndex] = taa_half4(dpv_lo(InputMinColor), dpv_lo(InputPixelAlignement));
-			SharedArray2[LocalGroupThreadIndex] = taa_half4(dpv_lo(InputMaxColor), dpv_lo(OutputPixelVelocity));
+			SharedArray0[LocalGroupThreadIndex] = tsr_half4(dpv_lo(FilteredInputColor), dpv_lo(LowFrequencyRejection));
+			SharedArray1[LocalGroupThreadIndex] = tsr_half4(dpv_lo(InputMinColor), dpv_lo(InputPixelAlignement));
+			SharedArray2[LocalGroupThreadIndex] = tsr_half4(dpv_lo(InputMaxColor), dpv_lo(OutputPixelVelocity));
 		#endif
 	}
 	#endif
 	
 	// Reproject history
-	taa_halfCx2 PrevLowFrequencyColor;
-	taa_half2 PrevLowFrequencyWeight;
+	tsr_halfCx2 PrevLowFrequencyColor;
+	tsr_half2 PrevLowFrequencyWeight;
 	
-	taa_halfCx2 PrevHighFrequencyColor;
-	taa_half2 PrevHistoryValidity;
+	tsr_halfCx2 PrevHighFrequencyColor;
+	tsr_half2 PrevHistoryValidity;
 	
-	taa_halfCx2 PrevHistoryMommentMin;
-	taa_halfCx2 PrevHistoryMommentMax;
+	tsr_halfCx2 PrevHistoryMommentMin;
+	tsr_halfCx2 PrevHistoryMommentMax;
 
-	taa_half4x2 PrevHistoryTranslucency;
+	tsr_half4x2 PrevHistoryTranslucency;
 	
-	taa_subpixel_details2 PrevSubpixelDetails;
+	tsr_subpixel_details2 PrevSubpixelDetails;
 
 	PLATFORM_SPECIFIC_ISOLATE
 	{
 		// Reproject history
-		taa_halfCx2 AccumulateLowFrequency = taa_half(0);
-		taa_halfCx2 AccumulateHighFrequency = taa_half(0);
-		taa_half2x2 AccumulateMetadata = taa_half(0);
-		taa_half4x2 AccumulateTranslucency = taa_half(0);
+		tsr_halfCx2 AccumulateLowFrequency = tsr_half(0);
+		tsr_halfCx2 AccumulateHighFrequency = tsr_half(0);
+		tsr_half2x2 AccumulateMetadata = tsr_half(0);
+		tsr_half4x2 AccumulateTranslucency = tsr_half(0);
 
-		taa_halfCx2 HighFrequencyMin = INFINITE_FLOAT;
-		taa_halfCx2 HighFrequencyMax = -INFINITE_FLOAT;
+		tsr_halfCx2 HighFrequencyMin = INFINITE_FLOAT;
+		tsr_halfCx2 HighFrequencyMax = -INFINITE_FLOAT;
 
 		// Sample raw history
 		{
@@ -740,13 +740,13 @@ void MainCS(
 			FCatmullRomSamples Samples1 = GetBicubic2DCatmullRomSamples(
 				dpv_hi(PrevHistoryBufferUV), PrevHistoryInfo_Extent, PrevHistoryInfo_ExtentInverse);
 			
-			taa_halfC RawLowFrequency[BICUBIC_CATMULL_ROM_SAMPLES][DPV_PIXEL_PER_LANE];
-			taa_halfC RawHighFrequency[BICUBIC_CATMULL_ROM_SAMPLES][DPV_PIXEL_PER_LANE];
-			taa_half2 RawMetadata[BICUBIC_CATMULL_ROM_SAMPLES][DPV_PIXEL_PER_LANE];
-			taa_half3 RawTranslucencyColor[BICUBIC_CATMULL_ROM_SAMPLES][DPV_PIXEL_PER_LANE];
-			taa_half  RawTranslucencyAlpha[BICUBIC_CATMULL_ROM_SAMPLES][DPV_PIXEL_PER_LANE];
+			tsr_halfC RawLowFrequency[BICUBIC_CATMULL_ROM_SAMPLES][DPV_PIXEL_PER_LANE];
+			tsr_halfC RawHighFrequency[BICUBIC_CATMULL_ROM_SAMPLES][DPV_PIXEL_PER_LANE];
+			tsr_half2 RawMetadata[BICUBIC_CATMULL_ROM_SAMPLES][DPV_PIXEL_PER_LANE];
+			tsr_half3 RawTranslucencyColor[BICUBIC_CATMULL_ROM_SAMPLES][DPV_PIXEL_PER_LANE];
+			tsr_half  RawTranslucencyAlpha[BICUBIC_CATMULL_ROM_SAMPLES][DPV_PIXEL_PER_LANE];
 
-			taa_half2 RawKernelWeight[BICUBIC_CATMULL_ROM_SAMPLES];
+			tsr_half2 RawKernelWeight[BICUBIC_CATMULL_ROM_SAMPLES];
 
 			// Issues texture fetches.
 			PLATFORM_SPECIFIC_ISOLATE
@@ -757,7 +757,7 @@ void MainCS(
 					float2 SampleUV0 = clamp(Samples0.UV[i], PrevHistoryInfo_UVViewportBilinearMin, PrevHistoryInfo_UVViewportBilinearMax);
 					float2 SampleUV1 = clamp(Samples1.UV[i], PrevHistoryInfo_UVViewportBilinearMin, PrevHistoryInfo_UVViewportBilinearMax);
 					
-					taa_half2 KernelWeight = dpv_force_interleave_registers(taa_half(Samples0.Weight[i]), taa_half(Samples1.Weight[i]));
+					tsr_half2 KernelWeight = dpv_force_interleave_registers(tsr_half(Samples0.Weight[i]), tsr_half(Samples1.Weight[i]));
 
 					RawKernelWeight[i] = KernelWeight;
 
@@ -785,16 +785,16 @@ void MainCS(
 				UNROLL_N(BICUBIC_CATMULL_ROM_SAMPLES)
 				for (uint i = 0; i < BICUBIC_CATMULL_ROM_SAMPLES; i++)
 				{
-					taa_halfCx2 SampleLowFrequency = dpv_force_interleave_registers_array(RawLowFrequency[i]);
-					taa_halfCx2 SampleHighFrequency = dpv_force_interleave_registers_array(RawHighFrequency[i]);
-					taa_half2x2 SampleMetadata = dpv_force_interleave_registers_array(RawMetadata[i]);
-					taa_half4x2 SampleTranslucency = dpv_force_interleave_registers(
-						taa_half4(RawTranslucencyColor[i][0], RawTranslucencyAlpha[i][0]),
-						taa_half4(RawTranslucencyColor[i][1], RawTranslucencyAlpha[i][1]));
+					tsr_halfCx2 SampleLowFrequency = dpv_force_interleave_registers_array(RawLowFrequency[i]);
+					tsr_halfCx2 SampleHighFrequency = dpv_force_interleave_registers_array(RawHighFrequency[i]);
+					tsr_half2x2 SampleMetadata = dpv_force_interleave_registers_array(RawMetadata[i]);
+					tsr_half4x2 SampleTranslucency = dpv_force_interleave_registers(
+						tsr_half4(RawTranslucencyColor[i][0], RawTranslucencyAlpha[i][0]),
+						tsr_half4(RawTranslucencyColor[i][1], RawTranslucencyAlpha[i][1]));
 					
-					taa_half2 KernelWeight = RawKernelWeight[i];
+					tsr_half2 KernelWeight = RawKernelWeight[i];
 
-					taa_half2 InvWeight = SafeRcp(SampleMetadata[1]);
+					tsr_half2 InvWeight = SafeRcp(SampleMetadata[1]);
 
 					HighFrequencyMin = min(HighFrequencyMin, dpv_scale(SampleHighFrequency, InvWeight));
 					HighFrequencyMax = max(HighFrequencyMax, dpv_scale(SampleHighFrequency, InvWeight));
@@ -808,7 +808,7 @@ void MainCS(
 					#endif
 				}
 				
-				taa_half2 FinalMultiplier = dpv_force_interleave_registers(taa_half(Samples0.FinalMultiplier), taa_half(Samples1.FinalMultiplier));
+				tsr_half2 FinalMultiplier = dpv_force_interleave_registers(tsr_half(Samples0.FinalMultiplier), tsr_half(Samples1.FinalMultiplier));
 
 				AccumulateLowFrequency  = dpv_scale(AccumulateLowFrequency, FinalMultiplier);
 				AccumulateHighFrequency = dpv_scale(AccumulateHighFrequency, FinalMultiplier);
@@ -823,22 +823,22 @@ void MainCS(
 				UNROLL_N(SUB_PIXEL_COUNT)
 				for (uint SubpixelId = 0; SubpixelId < SUB_PIXEL_COUNT; SubpixelId++)
 				{
-					taa_subpixel_payload SubpixelPayload0 = GetSubpixelPayload(SubpixelNeighborhood0, SubpixelId);
-					taa_subpixel_payload SubpixelPayload1 = GetSubpixelPayload(SubpixelNeighborhood1, SubpixelId);
-					PrevSubpixelDetails[0] |= SubpixelPayload0 << taa_subpixel_payload(SUB_PIXEL_BIT_COUNT * SubpixelId);
-					PrevSubpixelDetails[1] |= SubpixelPayload1 << taa_subpixel_payload(SUB_PIXEL_BIT_COUNT * SubpixelId);
+					tsr_subpixel_payload SubpixelPayload0 = GetSubpixelPayload(SubpixelNeighborhood0, SubpixelId);
+					tsr_subpixel_payload SubpixelPayload1 = GetSubpixelPayload(SubpixelNeighborhood1, SubpixelId);
+					PrevSubpixelDetails[0] |= SubpixelPayload0 << tsr_subpixel_payload(SUB_PIXEL_BIT_COUNT * SubpixelId);
+					PrevSubpixelDetails[1] |= SubpixelPayload1 << tsr_subpixel_payload(SUB_PIXEL_BIT_COUNT * SubpixelId);
 				}
 			}
 			
 			#if CONFIG_SCENE_COLOR_APLHA
-				AccumulateLowFrequency  = -dpv_min(-AccumulateLowFrequency, taa_half(0.0).xxxx);
-				AccumulateHighFrequency = -dpv_min(-AccumulateHighFrequency, taa_half(0.0).xxxx);
+				AccumulateLowFrequency  = -dpv_min(-AccumulateLowFrequency, tsr_half(0.0).xxxx);
+				AccumulateHighFrequency = -dpv_min(-AccumulateHighFrequency, tsr_half(0.0).xxxx);
 			#else
-				AccumulateLowFrequency  = -dpv_min(-AccumulateLowFrequency, taa_half(0.0).xxx);
-				AccumulateHighFrequency = -dpv_min(-AccumulateHighFrequency, taa_half(0.0).xxx);
+				AccumulateLowFrequency  = -dpv_min(-AccumulateLowFrequency, tsr_half(0.0).xxx);
+				AccumulateHighFrequency = -dpv_min(-AccumulateHighFrequency, tsr_half(0.0).xxx);
 			#endif
-			AccumulateMetadata      = -dpv_min(-AccumulateMetadata, taa_half(0.0).xx);
-			AccumulateTranslucency  = -dpv_min(-AccumulateTranslucency, taa_half(0.0).xxxx);
+			AccumulateMetadata      = -dpv_min(-AccumulateMetadata, tsr_half(0.0).xx);
+			AccumulateTranslucency  = -dpv_min(-AccumulateTranslucency, tsr_half(0.0).xxxx);
 		}
 		
 		// Unpack history.
@@ -857,13 +857,13 @@ void MainCS(
 
 		// Correct history
 		{
-			PrevHistoryMommentMin *= taa_half(HistoryPreExposureCorrection);
-			PrevHistoryMommentMax *= taa_half(HistoryPreExposureCorrection);
-			PrevHighFrequencyColor *= taa_half(HistoryPreExposureCorrection);
-			PrevLowFrequencyColor *= taa_half(HistoryPreExposureCorrection);
+			PrevHistoryMommentMin *= tsr_half(HistoryPreExposureCorrection);
+			PrevHistoryMommentMax *= tsr_half(HistoryPreExposureCorrection);
+			PrevHighFrequencyColor *= tsr_half(HistoryPreExposureCorrection);
+			PrevLowFrequencyColor *= tsr_half(HistoryPreExposureCorrection);
 
 			// Fixes negative validity that can happen due to negative lobes of Catrom.
-			PrevHistoryValidity = max(PrevHistoryValidity, taa_half(0.0).xx);
+			PrevHistoryValidity = max(PrevHistoryValidity, tsr_half(0.0).xx);
 		}
 	}
 		
@@ -873,41 +873,41 @@ void MainCS(
 		uint LocalGroupThreadIndex = GetGroupThreadIndex(GroupThreadIndex, GroupWaveIndex);
 	
 		#if CONFIG_ENABLE_DUAL_PIXEL_VECTORIZATION
-			taa_half4 RawLDS0 = SharedArray0[LocalGroupThreadIndex];
-			taa_half4 RawLDS1 = SharedArray1[LocalGroupThreadIndex];
-			taa_half4 RawLDS2 = SharedArray2[LocalGroupThreadIndex];
-			taa_half4 RawLDS3 = SharedArray3[LocalGroupThreadIndex];
-			taa_half4 RawLDS4 = SharedArray4[LocalGroupThreadIndex];
-			taa_half4 RawLDS5 = SharedArray5[LocalGroupThreadIndex];
+			tsr_half4 RawLDS0 = SharedArray0[LocalGroupThreadIndex];
+			tsr_half4 RawLDS1 = SharedArray1[LocalGroupThreadIndex];
+			tsr_half4 RawLDS2 = SharedArray2[LocalGroupThreadIndex];
+			tsr_half4 RawLDS3 = SharedArray3[LocalGroupThreadIndex];
+			tsr_half4 RawLDS4 = SharedArray4[LocalGroupThreadIndex];
+			tsr_half4 RawLDS5 = SharedArray5[LocalGroupThreadIndex];
 			
 			#if CONFIG_ACCUMULATE_TRANSLUCENCY_SEPARATLY
-				taa_half4 RawLDS6 = SharedArray6[LocalGroupThreadIndex];
-				taa_half4 RawLDS7 = SharedArray7[LocalGroupThreadIndex];
-				taa_half4 RawLDS8 = SharedArray8[LocalGroupThreadIndex];
-				taa_half4 RawLDS9 = SharedArray9[LocalGroupThreadIndex];
-				taa_half4 RawLDS10 = SharedArray10[LocalGroupThreadIndex];
-				taa_half4 RawLDS11 = SharedArray11[LocalGroupThreadIndex];
+				tsr_half4 RawLDS6 = SharedArray6[LocalGroupThreadIndex];
+				tsr_half4 RawLDS7 = SharedArray7[LocalGroupThreadIndex];
+				tsr_half4 RawLDS8 = SharedArray8[LocalGroupThreadIndex];
+				tsr_half4 RawLDS9 = SharedArray9[LocalGroupThreadIndex];
+				tsr_half4 RawLDS10 = SharedArray10[LocalGroupThreadIndex];
+				tsr_half4 RawLDS11 = SharedArray11[LocalGroupThreadIndex];
 			#endif
 
 			ParallaxFactorBits = SharedParallaxFactorBits[LocalGroupThreadIndex];
 
-			FilteredInputColor = taa_half3x2(RawLDS0.rg, RawLDS0.ba, RawLDS1.rg);
-			InputMinColor = taa_half3x2(RawLDS2.rg, RawLDS2.ba, RawLDS3.rg);
-			InputMaxColor = taa_half3x2(RawLDS4.rg, RawLDS4.ba, RawLDS5.rg);
+			FilteredInputColor = tsr_half3x2(RawLDS0.rg, RawLDS0.ba, RawLDS1.rg);
+			InputMinColor = tsr_half3x2(RawLDS2.rg, RawLDS2.ba, RawLDS3.rg);
+			InputMaxColor = tsr_half3x2(RawLDS4.rg, RawLDS4.ba, RawLDS5.rg);
 		
 			LowFrequencyRejection = RawLDS1.ba;
 			InputPixelAlignement = RawLDS3.ba;
 			OutputPixelVelocity = RawLDS5.ba;
 			
 			#if CONFIG_ACCUMULATE_TRANSLUCENCY_SEPARATLY
-				FilteredTranslucencyColor = taa_half4x2(RawLDS6.rg, RawLDS6.ba, RawLDS7.rg, RawLDS7.ba);
-				TranslucencyMinColor = taa_half4x2(RawLDS8.rg, RawLDS8.ba, RawLDS9.rg, RawLDS9.ba);
-				TranslucencyMaxColor = taa_half4x2(RawLDS10.rg, RawLDS10.ba, RawLDS11.rg, RawLDS11.ba);
+				FilteredTranslucencyColor = tsr_half4x2(RawLDS6.rg, RawLDS6.ba, RawLDS7.rg, RawLDS7.ba);
+				TranslucencyMinColor = tsr_half4x2(RawLDS8.rg, RawLDS8.ba, RawLDS9.rg, RawLDS9.ba);
+				TranslucencyMaxColor = tsr_half4x2(RawLDS10.rg, RawLDS10.ba, RawLDS11.rg, RawLDS11.ba);
 			#endif
 		#else
-			taa_half4 RawLDS0 = SharedArray0[LocalGroupThreadIndex];
-			taa_half4 RawLDS1 = SharedArray1[LocalGroupThreadIndex];
-			taa_half4 RawLDS2 = SharedArray2[LocalGroupThreadIndex];
+			tsr_half4 RawLDS0 = SharedArray0[LocalGroupThreadIndex];
+			tsr_half4 RawLDS1 = SharedArray1[LocalGroupThreadIndex];
+			tsr_half4 RawLDS2 = SharedArray2[LocalGroupThreadIndex];
 
 			FilteredInputColor = dpv_interleave_mono_registers(RawLDS0.rgb);
 			InputMinColor = dpv_interleave_mono_registers(RawLDS1.rgb);
@@ -923,13 +923,13 @@ void MainCS(
 	// Rejects high frequency details if their low frequencies drift from the low frequency history
 	#if CONFIG_LOW_FREQUENCY_DRIFT_REJECTION
 	{
-		taa_halfCx2 NormalizedPrevHighFrequencyColor = dpv_scale(PrevHighFrequencyColor, SafeRcp(PrevHistoryValidity));
+		tsr_halfCx2 NormalizedPrevHighFrequencyColor = dpv_scale(PrevHighFrequencyColor, SafeRcp(PrevHistoryValidity));
 
-		taa_halfCx2 PrevHighFrequencyYCoCg = TransformColorForClampingBox(NormalizedPrevHighFrequencyColor);
-		taa_halfCx2 PrevYCoCg = TransformColorForClampingBox(PrevLowFrequencyColor);
-		taa_halfCx2 ClampedPrevYCoCg = TransformColorForClampingBox(clamp(PrevLowFrequencyColor, PrevHistoryMommentMin, PrevHistoryMommentMax));
+		tsr_halfCx2 PrevHighFrequencyYCoCg = TransformColorForClampingBox(NormalizedPrevHighFrequencyColor);
+		tsr_halfCx2 PrevYCoCg = TransformColorForClampingBox(PrevLowFrequencyColor);
+		tsr_halfCx2 ClampedPrevYCoCg = TransformColorForClampingBox(clamp(PrevLowFrequencyColor, PrevHistoryMommentMin, PrevHistoryMommentMax));
 
-		taa_half2 HighFrequencyRejection = MeasureRejectionFactor(
+		tsr_half2 HighFrequencyRejection = MeasureRejectionFactor(
 			PrevYCoCg, ClampedPrevYCoCg,
 			PrevHighFrequencyYCoCg, InputMinColor, InputMaxColor,
 			MeasureBackbufferLDRQuantizationErrorFromHDRLuma(max(PrevYCoCg[0], ClampedPrevYCoCg[0])));
@@ -941,21 +941,21 @@ void MainCS(
 	#endif
 
 	// Contribute current frame input into the predictor for next frame.
-	const taa_half Histeresis = rcp(taa_half(MAX_SAMPLE_COUNT));
-	const taa_half PredictionOnlyValidity = Histeresis * taa_half(2.0);
+	const tsr_half Histeresis = rcp(tsr_half(MAX_SAMPLE_COUNT));
+	const tsr_half PredictionOnlyValidity = Histeresis * tsr_half(2.0);
 	
 	// Clamp low frequency
-	taa_half2 LumaMin;
-	taa_half2 LumaMax;
-	taa_halfCx2 ClampedLowFrequencyColor;
-	taa_half2 MidFrequencyRejection;
+	tsr_half2 LumaMin;
+	tsr_half2 LumaMax;
+	tsr_halfCx2 ClampedLowFrequencyColor;
+	tsr_half2 MidFrequencyRejection;
 	{
 		LumaMin = InputMinColor[0];
 		LumaMax = InputMaxColor[0];
 
-		taa_halfCx2 PrevYCoCg = TransformColorForClampingBox(PrevLowFrequencyColor);
-		taa_halfCx2 ClampedPrevYCoCg = dpv_clamp(PrevYCoCg, InputMinColor, InputMaxColor);
-		taa_halfCx2 InputCenterYCoCg = TransformColorForClampingBox(FilteredInputColor);
+		tsr_halfCx2 PrevYCoCg = TransformColorForClampingBox(PrevLowFrequencyColor);
+		tsr_halfCx2 ClampedPrevYCoCg = dpv_clamp(PrevYCoCg, InputMinColor, InputMaxColor);
+		tsr_halfCx2 InputCenterYCoCg = TransformColorForClampingBox(FilteredInputColor);
 
 		ClampedLowFrequencyColor = YCoCgToRGB(ClampedPrevYCoCg);
 		
@@ -967,7 +967,7 @@ void MainCS(
 		#if !CONFIG_CLAMP
 		{
 			ClampedLowFrequencyColor = PrevLowFrequencyColor;
-			MidFrequencyRejection = taa_half(1.0);
+			MidFrequencyRejection = tsr_half(1.0);
 		}
 		#endif
 	}
@@ -979,13 +979,13 @@ void MainCS(
 	}
 	#endif
 	
-	taa_halfCx2 FinalHighFrequencyColor;
-	taa_half2 FinalHistoryValidity;
+	tsr_halfCx2 FinalHighFrequencyColor;
+	tsr_half2 FinalHistoryValidity;
 	{
-		taa_half2 InvPrevHistoryValidity = SafeRcp(PrevHistoryValidity);
+		tsr_half2 InvPrevHistoryValidity = SafeRcp(PrevHistoryValidity);
 
 		// Compute how much the history needs to be rejected, based on the completeness.
-		taa_half2 PrevHistoryRejectionWeight = LowFrequencyRejection * HoleFilledVelocityMask[1];
+		tsr_half2 PrevHistoryRejectionWeight = LowFrequencyRejection * HoleFilledVelocityMask[1];
 			
 		//#if !CONFIG_CLAMP
 		//	PrevHistoryRejectionWeight = 1.0;
@@ -994,33 +994,33 @@ void MainCS(
 		FLATTEN
 		if (bOffScreen[0])
 		{
-			PrevHistoryRejectionWeight[0] = taa_half(0.0);
+			PrevHistoryRejectionWeight[0] = tsr_half(0.0);
 		}
 
 		FLATTEN
 		if (bOffScreen[1])
 		{
-			PrevHistoryRejectionWeight[1] = taa_half(0.0);
+			PrevHistoryRejectionWeight[1] = tsr_half(0.0);
 		}
 
-		taa_half2 DesiredCurrentContribution = max(Histeresis * InputPixelAlignement, taa_half(0.0));
+		tsr_half2 DesiredCurrentContribution = max(Histeresis * InputPixelAlignement, tsr_half(0.0));
 
 		// Determine whether the the prediction based rejection was confident enough.
-		taa_half2 RejectionConfidentEnough = taa_half(1); // saturate(RejectionValidity * MAX_SAMPLE_COUNT - 3.0);
+		tsr_half2 RejectionConfidentEnough = tsr_half(1); // saturate(RejectionValidity * MAX_SAMPLE_COUNT - 3.0);
 		
 		// Clamp the validity due to motion to maintain better sharpness in history reprojection under motion.
-		taa_half2 ClampedPrevHistoryValidity;
+		tsr_half2 ClampedPrevHistoryValidity;
 		#if 1
 		{
-			taa_half2 MaxValidity = taa_half(1.0) - (taa_half(1.0) - PredictionOnlyValidity) * saturate(OutputPixelVelocity * taa_half(InvWeightClampingPixelSpeed));
+			tsr_half2 MaxValidity = tsr_half(1.0) - (tsr_half(1.0) - PredictionOnlyValidity) * saturate(OutputPixelVelocity * tsr_half(InvWeightClampingPixelSpeed));
 
 			// Clamp up the max validity to favorise stability under motion on high contrast edges.
 			#if 1
 			{
-				taa_half2 PrevHistoryLuma = Luma4(PrevHighFrequencyColor) * InvPrevHistoryValidity;
-				taa_half2 FilteredLuma = Luma4(FilteredInputColor);
+				tsr_half2 PrevHistoryLuma = Luma4(PrevHighFrequencyColor) * InvPrevHistoryValidity;
+				tsr_half2 FilteredLuma = Luma4(FilteredInputColor);
 
-				taa_half2 MinValidityForStability = abs(FilteredLuma - PrevHistoryLuma) / max(FilteredLuma, PrevHistoryLuma);
+				tsr_half2 MinValidityForStability = abs(FilteredLuma - PrevHistoryLuma) / max(FilteredLuma, PrevHistoryLuma);
 				
 				MaxValidity = max(MaxValidity, MinValidityForStability);
 			}
@@ -1035,43 +1035,43 @@ void MainCS(
 		#endif
 
 		// Compute the newly rejected validity
-		taa_half2 RejectedValidity = (
+		tsr_half2 RejectedValidity = (
 			min(ClampedPrevHistoryValidity, PredictionOnlyValidity - DesiredCurrentContribution) +
-			max(ClampedPrevHistoryValidity - PredictionOnlyValidity + DesiredCurrentContribution, taa_half(0.0)) * PrevHistoryRejectionWeight);
+			max(ClampedPrevHistoryValidity - PredictionOnlyValidity + DesiredCurrentContribution, tsr_half(0.0)) * PrevHistoryRejectionWeight);
 
 		RejectedValidity = ClampedPrevHistoryValidity * PrevHistoryRejectionWeight;
 
 		// Compute the maximum output validity.
-		taa_half2 OutputValidity = (
-			clamp(RejectedValidity + DesiredCurrentContribution, taa_half(0.0), PredictionOnlyValidity) +
-			clamp(RejectedValidity + DesiredCurrentContribution * PrevHistoryRejectionWeight * RejectionConfidentEnough - PredictionOnlyValidity, taa_half(0.0), taa_half(1.0) - PredictionOnlyValidity));
+		tsr_half2 OutputValidity = (
+			clamp(RejectedValidity + DesiredCurrentContribution, tsr_half(0.0), PredictionOnlyValidity) +
+			clamp(RejectedValidity + DesiredCurrentContribution * PrevHistoryRejectionWeight * RejectionConfidentEnough - PredictionOnlyValidity, tsr_half(0.0), tsr_half(1.0) - PredictionOnlyValidity));
 
 		#if !CONFIG_ACCUMULATE_TRANSLUCENCY_SEPARATLY
 			FLATTEN
 			if (bIsResponsiveAAPixel[0])
 			{
-				OutputValidity[0] = taa_half(0.0);
+				OutputValidity[0] = tsr_half(0.0);
 			}
 
 			FLATTEN
 			if (bIsResponsiveAAPixel[1])
 			{
-				OutputValidity[1] = taa_half(0.0);
+				OutputValidity[1] = tsr_half(0.0);
 			}
 		#endif
 		
-		taa_half2 PrevMomentWeight = max(OutputValidity - DesiredCurrentContribution, taa_half(0.0));
-		taa_half2 CurrentMomentWeight = min(DesiredCurrentContribution, OutputValidity);
+		tsr_half2 PrevMomentWeight = max(OutputValidity - DesiredCurrentContribution, tsr_half(0.0));
+		tsr_half2 CurrentMomentWeight = min(DesiredCurrentContribution, OutputValidity);
 		
 		#if 1
 		{
-			taa_half2 PrevHistoryToneWeight = HdrWeightY(Luma4(PrevHighFrequencyColor) * InvPrevHistoryValidity);
-			taa_half2 FilteredInputToneWeight = HdrWeight4(FilteredInputColor);
+			tsr_half2 PrevHistoryToneWeight = HdrWeightY(Luma4(PrevHighFrequencyColor) * InvPrevHistoryValidity);
+			tsr_half2 FilteredInputToneWeight = HdrWeight4(FilteredInputColor);
 			
-			taa_half2 BlendPrevHistory = PrevMomentWeight * PrevHistoryToneWeight;
-			taa_half2 BlendFilteredInput = CurrentMomentWeight * FilteredInputToneWeight;
+			tsr_half2 BlendPrevHistory = PrevMomentWeight * PrevHistoryToneWeight;
+			tsr_half2 BlendFilteredInput = CurrentMomentWeight * FilteredInputToneWeight;
 
-			taa_half2 CommonWeight = SafeRcp(BlendPrevHistory + BlendFilteredInput);
+			tsr_half2 CommonWeight = SafeRcp(BlendPrevHistory + BlendFilteredInput);
 
 			FinalHighFrequencyColor = (
 				dpv_scale(PrevHighFrequencyColor, CommonWeight * BlendPrevHistory * InvPrevHistoryValidity) +
@@ -1087,7 +1087,7 @@ void MainCS(
 		#if 0
 			FinalHistoryValidity = OutputValidity;
 		#else
-			FinalHistoryValidity = ceil(taa_half(255.0) * OutputValidity) * rcp(taa_half(255.0));
+			FinalHistoryValidity = ceil(tsr_half(255.0) * OutputValidity) * rcp(tsr_half(255.0));
 		#endif
 
 		#if 0 // DEBUG
@@ -1099,46 +1099,46 @@ void MainCS(
 	}
 
 	// Compute the low frequency history
-	taa_halfCx2 FinalLowFrequencyColor;
-	taa_half2 FinalLowFrequencyWeight;
+	tsr_halfCx2 FinalLowFrequencyColor;
+	tsr_half2 FinalLowFrequencyWeight;
 	{
-		const taa_half TargetHesteresisCurrentFrameWeight = rcp(taa_half(MAX_FALLBACK_SAMPLE_COUNT));
+		const tsr_half TargetHesteresisCurrentFrameWeight = rcp(tsr_half(MAX_FALLBACK_SAMPLE_COUNT));
 
-		taa_half2 LumaHistory = Luma4(PrevLowFrequencyColor);
-		taa_half2 LumaFiltered = Luma4(FilteredInputColor);
+		tsr_half2 LumaHistory = Luma4(PrevLowFrequencyColor);
+		tsr_half2 LumaFiltered = Luma4(FilteredInputColor);
 
 		{
-			taa_half2 OutputBlend = ComputeLowFrequencyContribution(FinalHistoryValidity);
+			tsr_half2 OutputBlend = ComputeLowFrequencyContribution(FinalHistoryValidity);
 
 			//ClampedLowFrequencyColor = lerp(ClampedLowFrequencyColor, PrevLowFrequencyColor, OutputBlend);
 		}
 
-		taa_half2 BlendFinal;
+		tsr_half2 BlendFinal;
 		#if 1
 		{
-			taa_half2 CurrentFrameSampleCount = max(InputPixelAlignement, taa_half(0.005));
+			tsr_half2 CurrentFrameSampleCount = max(InputPixelAlignement, tsr_half(0.005));
 			
 			// Use only one sample count to be extremely fast to recover in history rejection. But then stabilize imediately so sub pixel
 			// frequency can be used as soon as possible.
-			taa_half2 PrevLowFrequencySampleCount;
+			tsr_half2 PrevLowFrequencySampleCount;
 			FLATTEN
-			if (PrevLowFrequencyWeight[0] < taa_half(1.0))
+			if (PrevLowFrequencyWeight[0] < tsr_half(1.0))
 			{
 				PrevLowFrequencySampleCount[0] = PrevLowFrequencyWeight[0];
 			}
 			else
 			{
-				PrevLowFrequencySampleCount[0] = taa_half(MAX_FALLBACK_SAMPLE_COUNT);
+				PrevLowFrequencySampleCount[0] = tsr_half(MAX_FALLBACK_SAMPLE_COUNT);
 			}
 
 			FLATTEN
-			if (PrevLowFrequencyWeight[1] < taa_half(1.0))
+			if (PrevLowFrequencyWeight[1] < tsr_half(1.0))
 			{
 				PrevLowFrequencySampleCount[1] = PrevLowFrequencyWeight[1];
 			}
 			else
 			{
-				PrevLowFrequencySampleCount[1] = taa_half(MAX_FALLBACK_SAMPLE_COUNT);
+				PrevLowFrequencySampleCount[1] = tsr_half(MAX_FALLBACK_SAMPLE_COUNT);
 			}
 
 			// Reject the history based on the low frequencies.
@@ -1147,7 +1147,7 @@ void MainCS(
 				// float InvertedBlendFinal = PrevLowFrequencySampleCount / (CurrentFrameSampleCount + PrevLowFrequencySampleCount);
 				// float PrevLowFrequencyRejectionFactor = saturate(LowFrequencyRejection / InvertedBlendFinal);
 				
-				taa_half2 PrevLowFrequencyRejectionFactor = saturate(LowFrequencyRejection * (CurrentFrameSampleCount + PrevLowFrequencySampleCount) * rcp(PrevLowFrequencySampleCount));
+				tsr_half2 PrevLowFrequencyRejectionFactor = saturate(LowFrequencyRejection * (CurrentFrameSampleCount + PrevLowFrequencySampleCount) * rcp(PrevLowFrequencySampleCount));
 
 				PrevLowFrequencySampleCount *= PrevLowFrequencyRejectionFactor;
 			}
@@ -1158,34 +1158,34 @@ void MainCS(
 			// Increase the blending on motion
 			#if 1
 			{
-				BlendFinal = lerp(BlendFinal, max(taa_half(0.2), BlendFinal), saturate(OutputPixelVelocity * rcp(taa_half(40.0))));
+				BlendFinal = lerp(BlendFinal, max(tsr_half(0.2), BlendFinal), saturate(OutputPixelVelocity * rcp(tsr_half(40.0))));
 			}
 			#endif
 
 			// Antiflicker
 			#if 1
 			{
-				taa_half2 DistToClamp = min( abs(LumaHistory - LumaMin), abs(LumaHistory - LumaMax) ) * rcp(max3( LumaHistory, LumaFiltered, taa_half(1e-4) ));
-				BlendFinal *= taa_half(0.2) + taa_half(0.8) * saturate(taa_half(0.5) * DistToClamp);
+				tsr_half2 DistToClamp = min( abs(LumaHistory - LumaMin), abs(LumaHistory - LumaMax) ) * rcp(max3( LumaHistory, LumaFiltered, tsr_half(1e-4) ));
+				BlendFinal *= tsr_half(0.2) + tsr_half(0.8) * saturate(tsr_half(0.5) * DistToClamp);
 			}
 			#endif
 			
 			// Make sure to have at least some small contribution
 			#if 1
 			{
-				BlendFinal = max(BlendFinal, saturate(taa_half(0.01) * LumaHistory * rcp(abs(LumaFiltered - LumaHistory))));
+				BlendFinal = max(BlendFinal, saturate(tsr_half(0.01) * LumaHistory * rcp(abs(LumaFiltered - LumaHistory))));
 			}
 			#endif
 
 			#if !CONFIG_ACCUMULATE_TRANSLUCENCY_SEPARATLY
 				// Responsive forces 1/4 of new frame.
-				BlendFinal[0] = bIsResponsiveAAPixel[0] ? taa_half(1.0/4.0) : BlendFinal[0];
-				BlendFinal[1] = bIsResponsiveAAPixel[1] ? taa_half(1.0/4.0) : BlendFinal[1];
+				BlendFinal[0] = bIsResponsiveAAPixel[0] ? tsr_half(1.0/4.0) : BlendFinal[0];
+				BlendFinal[1] = bIsResponsiveAAPixel[1] ? tsr_half(1.0/4.0) : BlendFinal[1];
 			#endif
 
 			// Completely discard the history
 			{
-				taa_half2 TotalRejectionBlending = max(BlendFinal, HoleFilledVelocityMask[0]);
+				tsr_half2 TotalRejectionBlending = max(BlendFinal, HoleFilledVelocityMask[0]);
 
 				PrevLowFrequencySampleCount *= TotalRejection;
 				BlendFinal = lerp(TotalRejectionBlending, BlendFinal, TotalRejection);
@@ -1194,90 +1194,90 @@ void MainCS(
 			FinalLowFrequencyWeight = saturate(CurrentFrameSampleCount + PrevLowFrequencySampleCount);
 			
 			#if 1
-				FinalLowFrequencyWeight = saturate(floor(taa_half(255.0) * (CurrentFrameSampleCount + PrevLowFrequencySampleCount)) * rcp(taa_half(255.0)));
+				FinalLowFrequencyWeight = saturate(floor(tsr_half(255.0) * (CurrentFrameSampleCount + PrevLowFrequencySampleCount)) * rcp(tsr_half(255.0)));
 			#endif
 		}
 		#endif
 
 		{
-			taa_half2 FilterWeight = HdrWeight4(FilteredInputColor);
-			taa_half2 ClampedHistoryWeight = HdrWeight4(ClampedLowFrequencyColor);
+			tsr_half2 FilterWeight = HdrWeight4(FilteredInputColor);
+			tsr_half2 ClampedHistoryWeight = HdrWeight4(ClampedLowFrequencyColor);
 
-			taa_half2x2 Weights = WeightedLerpFactors(ClampedHistoryWeight, FilterWeight, BlendFinal);
+			tsr_half2x2 Weights = WeightedLerpFactors(ClampedHistoryWeight, FilterWeight, BlendFinal);
 
 			FinalLowFrequencyColor = dpv_scale(ClampedLowFrequencyColor, Weights[0]) + dpv_scale(FilteredInputColor, Weights[1]);
 		}
 	}
 	
 	// Compute the low frequency history
-	taa_half4x2 FinalTranslucencyColor;
+	tsr_half4x2 FinalTranslucencyColor;
 	#if CONFIG_ACCUMULATE_TRANSLUCENCY_SEPARATLY
 	{
-		taa_half2 TranslucencyLumaMin = Luma4(GetColorChannel(TranslucencyMinColor));
-		taa_half2 TranslucencyLumaMax = Luma4(GetColorChannel(TranslucencyMaxColor));
+		tsr_half2 TranslucencyLumaMin = Luma4(GetColorChannel(TranslucencyMinColor));
+		tsr_half2 TranslucencyLumaMax = Luma4(GetColorChannel(TranslucencyMaxColor));
 
-		taa_half2 TranslucencyLumaHistory  = Luma4(GetColorChannel(PrevHistoryTranslucency));
-		taa_half2 TranslucencyLumaFiltered = Luma4(GetColorChannel(FilteredTranslucencyColor));
+		tsr_half2 TranslucencyLumaHistory  = Luma4(GetColorChannel(PrevHistoryTranslucency));
+		tsr_half2 TranslucencyLumaFiltered = Luma4(GetColorChannel(FilteredTranslucencyColor));
 
-		taa_half4x2 ClampedPrevHistoryTranslucency = clamp(PrevHistoryTranslucency, TranslucencyMinColor, TranslucencyMaxColor);
+		tsr_half4x2 ClampedPrevHistoryTranslucency = clamp(PrevHistoryTranslucency, TranslucencyMinColor, TranslucencyMaxColor);
 		
 		#if CONFIG_SCENE_COLOR_APLHA
-			taa_half2 HighFrequencyTranslucencyRejection = MeasureRejectionFactor(
+			tsr_half2 HighFrequencyTranslucencyRejection = MeasureRejectionFactor(
 				PrevHistoryTranslucency, ClampedPrevHistoryTranslucency,
 				FilteredTranslucencyColor, TranslucencyMinColor, TranslucencyMaxColor,
-				taa_half(0.0));
+				tsr_half(0.0));
 		#else
-			taa_half2 HighFrequencyTranslucencyRejection = MeasureRejectionFactor(
+			tsr_half2 HighFrequencyTranslucencyRejection = MeasureRejectionFactor(
 				GetColorChannel(PrevHistoryTranslucency), GetColorChannel(ClampedPrevHistoryTranslucency),
 				GetColorChannel(FilteredTranslucencyColor), GetColorChannel(TranslucencyMinColor), GetColorChannel(TranslucencyMaxColor),
-				taa_half(0.0));
+				tsr_half(0.0));
 		#endif
 
-		taa_half2 BlendFinal = taa_half(0.2) * InputPixelAlignement;
+		tsr_half2 BlendFinal = tsr_half(0.2) * InputPixelAlignement;
 
 		// Antiflicker
 		#if 1
 		{
-			taa_half2 DistToClamp = min( abs(TranslucencyLumaHistory - TranslucencyLumaMin), abs(TranslucencyLumaHistory - TranslucencyLumaMax) ) / max3( TranslucencyLumaHistory, TranslucencyLumaFiltered, taa_half(1e-4) );
-			BlendFinal *= taa_half(0.2) + taa_half(0.8) * saturate(taa_half(0.5) * DistToClamp);
+			tsr_half2 DistToClamp = min( abs(TranslucencyLumaHistory - TranslucencyLumaMin), abs(TranslucencyLumaHistory - TranslucencyLumaMax) ) / max3( TranslucencyLumaHistory, TranslucencyLumaFiltered, tsr_half(1e-4) );
+			BlendFinal *= tsr_half(0.2) + tsr_half(0.8) * saturate(tsr_half(0.5) * DistToClamp);
 		}
 		#endif
 			
 		// Make sure to have at least some small contribution
 		#if 0
 		{
-			BlendFinal = max( BlendFinal, saturate( taa_half(0.01) * TranslucencyLumaHistory * rcp(abs(TranslucencyLumaFiltered - TranslucencyLumaHistory))));
+			BlendFinal = max( BlendFinal, saturate( tsr_half(0.01) * TranslucencyLumaHistory * rcp(abs(TranslucencyLumaFiltered - TranslucencyLumaHistory))));
 		}
 		#endif
 		
-		BlendFinal = max(BlendFinal, (taa_half(1.0) - TranslucencyRejection));
+		BlendFinal = max(BlendFinal, (tsr_half(1.0) - TranslucencyRejection));
 		
 		FLATTEN
 		if (bCameraCut)
 		{
-			BlendFinal = taa_half(1.0);
+			BlendFinal = tsr_half(1.0);
 		}
 		
 		#if 0
 		{
-			taa_half2 OneMinusBlendFinal = taa_half(1.0) - BlendFinal;
+			tsr_half2 OneMinusBlendFinal = tsr_half(1.0) - BlendFinal;
 		
 			FinalTranslucencyColor = dpv_scale(ClampedPrevHistoryTranslucency, OneMinusBlendFinal) + dpv_scale(FilteredTranslucencyColor, BlendFinal);
 		}
 		#else
 		{
-			taa_half2 FilterWeight = HdrWeight4(GetColorChannel(FilteredTranslucencyColor));
-			taa_half2 ClampedHistoryWeight = HdrWeight4(GetColorChannel(ClampedPrevHistoryTranslucency));
+			tsr_half2 FilterWeight = HdrWeight4(GetColorChannel(FilteredTranslucencyColor));
+			tsr_half2 ClampedHistoryWeight = HdrWeight4(GetColorChannel(ClampedPrevHistoryTranslucency));
 
 			
-			taa_half2 BlendA = saturate(taa_half(1.0) - 1.05 * BlendFinal) * ClampedHistoryWeight;
-			taa_half2 BlendB = BlendFinal * FilterWeight;
-			taa_half2 RcpBlend = SafeRcp(BlendA + BlendB);
+			tsr_half2 BlendA = saturate(tsr_half(1.0) - 1.05 * BlendFinal) * ClampedHistoryWeight;
+			tsr_half2 BlendB = BlendFinal * FilterWeight;
+			tsr_half2 RcpBlend = SafeRcp(BlendA + BlendB);
 			BlendA *= RcpBlend;
 			BlendB *= RcpBlend;
-			taa_half2x2 Weights =  taa_half2x2(BlendA, BlendB);
+			tsr_half2x2 Weights =  tsr_half2x2(BlendA, BlendB);
 
-			//taa_half2x2 Weights = WeightedLerpFactors(ClampedHistoryWeight, FilterWeight, BlendFinal);
+			//tsr_half2x2 Weights = WeightedLerpFactors(ClampedHistoryWeight, FilterWeight, BlendFinal);
 
 			FinalTranslucencyColor = dpv_scale(ClampedPrevHistoryTranslucency, Weights[0]) + dpv_scale(FilteredTranslucencyColor, Weights[1]);
 		}
@@ -1286,13 +1286,13 @@ void MainCS(
 	#endif
 
 	// Update subpixel detail
-	taa_subpixel_details2 FinalSubpixelDetails;
+	tsr_subpixel_details2 FinalSubpixelDetails;
 	{
 		FinalSubpixelDetails = PrevSubpixelDetails;
 
 
 		{
-			const taa_subpixel_details ParallaxFactorMask = taa_subpixel_details(
+			const tsr_subpixel_details ParallaxFactorMask = tsr_subpixel_details(
 				(SUB_PIXEL_PARALLAX_FACTOR_BIT_MASK << (SUB_PIXEL_PARALLAX_FACTOR_BIT_OFFSET + 0 * SUB_PIXEL_BIT_COUNT)) | 
 				(SUB_PIXEL_PARALLAX_FACTOR_BIT_MASK << (SUB_PIXEL_PARALLAX_FACTOR_BIT_OFFSET + 1 * SUB_PIXEL_BIT_COUNT)) | 
 				(SUB_PIXEL_PARALLAX_FACTOR_BIT_MASK << (SUB_PIXEL_PARALLAX_FACTOR_BIT_OFFSET + 2 * SUB_PIXEL_BIT_COUNT)) | 
@@ -1315,11 +1315,11 @@ void MainCS(
 		
 		uint2 SubpixelId = dpv_dot(uint2x2(bKOSign), uint2(1, SUB_PIXEL_GRID_SIZE));
 
-		taa_subpixel_payload2 SubpixelShift = taa_subpixel_payload2(SubpixelId) * taa_subpixel_payload(SUB_PIXEL_BIT_COUNT);
+		tsr_subpixel_payload2 SubpixelShift = tsr_subpixel_payload2(SubpixelId) * tsr_subpixel_payload(SUB_PIXEL_BIT_COUNT);
 
-		taa_subpixel_payload2 SubpixelPayload = (ParallaxFactorBits << taa_subpixel_payload(SUB_PIXEL_PARALLAX_FACTOR_BIT_OFFSET));
+		tsr_subpixel_payload2 SubpixelPayload = (ParallaxFactorBits << tsr_subpixel_payload(SUB_PIXEL_PARALLAX_FACTOR_BIT_OFFSET));
 
-		taa_subpixel_details2 UpdatedFinalSubpixelDetails = (FinalSubpixelDetails & (~(taa_subpixel_payload(SUB_PIXEL_BIT_MASK) << SubpixelShift))) | (SubpixelPayload << SubpixelShift);
+		tsr_subpixel_details2 UpdatedFinalSubpixelDetails = (FinalSubpixelDetails & (~(tsr_subpixel_payload(SUB_PIXEL_BIT_MASK) << SubpixelShift))) | (SubpixelPayload << SubpixelShift);
 
 		FLATTEN
 		if (bUpdateSubpixelDetails[0])
@@ -1335,9 +1335,9 @@ void MainCS(
 	}
 
 	// Compute final output
-	taa_halfCx2 FinalOutputColor;
+	tsr_halfCx2 FinalOutputColor;
 	{
-		taa_half2 OutputBlend = ComputeLowFrequencyContribution(FinalHistoryValidity);
+		tsr_half2 OutputBlend = ComputeLowFrequencyContribution(FinalHistoryValidity);
 
 		#if 0 // DEBUG
 		{
@@ -1356,18 +1356,18 @@ void MainCS(
 		}
 		#endif
 
-		taa_half2 LowFrequencyWeight = HdrWeight4(FinalLowFrequencyColor);
-		taa_half2 Moment1Weight = HdrWeight4(FinalHighFrequencyColor);
+		tsr_half2 LowFrequencyWeight = HdrWeight4(FinalLowFrequencyColor);
+		tsr_half2 Moment1Weight = HdrWeight4(FinalHighFrequencyColor);
 
-		taa_half2x2 Weights = WeightedLerpFactors(LowFrequencyWeight, Moment1Weight, OutputBlend);
+		tsr_half2x2 Weights = WeightedLerpFactors(LowFrequencyWeight, Moment1Weight, OutputBlend);
 
 		#if DEBUG_FALLBACK_BLENDING
-			taa_half3 LowFrequencyColor = taa_half3(1, 0.25, 0.25);
-			taa_half3 HighFrequencyColor = taa_half3(0.25, 1, 0.25);
+			tsr_half3 LowFrequencyColor = tsr_half3(1, 0.25, 0.25);
+			tsr_half3 HighFrequencyColor = tsr_half3(0.25, 1, 0.25);
 
 			FinalOutputColor = dpv_scale(FinalLowFrequencyColor * dpv_interleave_mono_registers(LowFrequencyColor), Weights[0]) + dpv_scale(FinalHighFrequencyColor * dpv_interleave_mono_registers(HighFrequencyColor), Weights[1]);
 		#elif DEBUG_LOW_FREQUENCY_REJECTION
-			taa_half3 DebugColor = lerp(taa_half3(1, 0.5, 0.5), taa_half3(0.5, 1, 0.5), LowFrequencyRejection);
+			tsr_half3 DebugColor = lerp(tsr_half3(1, 0.5, 0.5), tsr_half3(0.5, 1, 0.5), LowFrequencyRejection);
 			
 			FinalOutputColor = FinalLowFrequencyColor * Weights.x * DebugColor + FinalHighFrequencyColor * Weights.y * DebugColor;
 		#else
@@ -1379,7 +1379,7 @@ void MainCS(
 			FinalOutputColor[1] = FinalOutputColor[1] * FinalTranslucencyColor[3] + FinalTranslucencyColor[1];
 			FinalOutputColor[2] = FinalOutputColor[2] * FinalTranslucencyColor[3] + FinalTranslucencyColor[2];
 			#if CONFIG_SCENE_COLOR_APLHA
-				FinalOutputColor[3] = saturate(taa_half(1.0) - (taa_half(1.0) - FinalOutputColor[3]) * FinalTranslucencyColor[3]);
+				FinalOutputColor[3] = saturate(tsr_half(1.0) - (tsr_half(1.0) - FinalOutputColor[3]) * FinalTranslucencyColor[3]);
 			#endif
 		#endif
 	}
@@ -1388,34 +1388,34 @@ void MainCS(
 	PLATFORM_SPECIFIC_ISOLATE
 	{
 		#if CONFIG_SCENE_COLOR_APLHA
-			const taa_halfC ColorNull = taa_half(0.0).xxxx;
-			const taa_halfC ColorMax10BitsFloat = taa_half(Max10BitsFloat).xxxx;
+			const tsr_halfC ColorNull = tsr_half(0.0).xxxx;
+			const tsr_halfC ColorMax10BitsFloat = tsr_half(Max10BitsFloat).xxxx;
 		#else
-			const taa_halfC ColorNull = taa_half(0.0).xxx;
-			const taa_halfC ColorMax10BitsFloat = taa_half(Max10BitsFloat).xxx;
+			const tsr_halfC ColorNull = tsr_half(0.0).xxx;
+			const tsr_halfC ColorMax10BitsFloat = tsr_half(Max10BitsFloat).xxx;
 		#endif
 
 		uint LocalGroupThreadIndex = GetGroupThreadIndex(GroupThreadIndex, GroupWaveIndex);
 
 		#if 1
-			taa_short2x2 LocalHistoryPixelPos = dpv_add(
-				taa_short2(GroupId) * taa_short2(TILE_SIZE, TILE_SIZE),
-				taa_short2x2(Map8x8Tile2x2LaneDPV(LocalGroupThreadIndex)));
+			tsr_short2x2 LocalHistoryPixelPos = dpv_add(
+				tsr_short2(GroupId) * tsr_short2(TILE_SIZE, TILE_SIZE),
+				tsr_short2x2(Map8x8Tile2x2LaneDPV(LocalGroupThreadIndex)));
 		#else
-			taa_short2x2 LocalHistoryPixelPos = HistoryPixelPos;
+			tsr_short2x2 LocalHistoryPixelPos = HistoryPixelPos;
 		#endif
 			
 		LocalHistoryPixelPos = InvalidateOutputPixelPos(LocalHistoryPixelPos, HistoryInfo_ViewportMax);
 		
 		// Premultiply FinalHighFrequencyColor to have validity weighted average in next frame history reprojection.
-		taa_halfCx2 PremultipliedHighFrequencyColor = dpv_scale(FinalHighFrequencyColor, FinalHistoryValidity);
+		tsr_halfCx2 PremultipliedHighFrequencyColor = dpv_scale(FinalHighFrequencyColor, FinalHistoryValidity);
 
 		// Stocastically round up or down using the hardware RWTexture2D truncation unit to take into precision
 		// loss due to pixel format encoding.
 		#if CONFIG_ENABLE_STOCASTIC_QUANTIZATION
 		{
 			uint2 Random = Rand3DPCG16(int3(dpv_lo(LocalHistoryPixelPos), View.StateFrameIndexMod8)).xy;
-			taa_half E = taa_half(Hammersley16(0, 1, Random).x);
+			tsr_half E = tsr_half(Hammersley16(0, 1, Random).x);
 
 			PremultipliedHighFrequencyColor = QuantizeForFloatRenderTarget(PremultipliedHighFrequencyColor, E, HistoryQuantizationError);
 			FinalLowFrequencyColor = QuantizeForFloatRenderTarget(FinalLowFrequencyColor, E, HistoryQuantizationError);
@@ -1431,8 +1431,8 @@ void MainCS(
 			PremultipliedHighFrequencyColor = dpv_min(PremultipliedHighFrequencyColor, ColorMax10BitsFloat);	
 
 			#if CONFIG_ACCUMULATE_TRANSLUCENCY_SEPARATLY
-				FinalTranslucencyColor = -dpv_min(-FinalTranslucencyColor, taa_half(0.0).xxxx);
-				FinalTranslucencyColor = dpv_min(FinalTranslucencyColor, taa_half(Max10BitsFloat).xxxx);
+				FinalTranslucencyColor = -dpv_min(-FinalTranslucencyColor, tsr_half(0.0).xxxx);
+				FinalTranslucencyColor = dpv_min(FinalTranslucencyColor, tsr_half(Max10BitsFloat).xxxx);
 			#endif
 		}
 
@@ -1440,7 +1440,7 @@ void MainCS(
 		{
 			HistoryOutput_LowFrequency[dpv_lo(LocalHistoryPixelPos)]  = dpv_lo(FinalLowFrequencyColor);
 			HistoryOutput_HighFrequency[dpv_lo(LocalHistoryPixelPos)] = dpv_lo(PremultipliedHighFrequencyColor);
-			HistoryOutput_Metadata[dpv_lo(LocalHistoryPixelPos)]      = taa_half2(dpv_lo(FinalLowFrequencyWeight), dpv_lo(FinalHistoryValidity));
+			HistoryOutput_Metadata[dpv_lo(LocalHistoryPixelPos)]      = tsr_half2(dpv_lo(FinalLowFrequencyWeight), dpv_lo(FinalHistoryValidity));
 			#if CONFIG_ACCUMULATE_TRANSLUCENCY_SEPARATLY
 				HistoryOutput_Translucency[dpv_lo(LocalHistoryPixelPos)] = dpv_lo(FinalTranslucencyColor).rgb;
 				HistoryOutput_TranslucencyAlpha[dpv_lo(LocalHistoryPixelPos)] = dpv_lo(FinalTranslucencyColor).a;
@@ -1453,7 +1453,7 @@ void MainCS(
 		{
 			HistoryOutput_LowFrequency[dpv_hi(LocalHistoryPixelPos)]  = dpv_hi(FinalLowFrequencyColor);
 			HistoryOutput_HighFrequency[dpv_hi(LocalHistoryPixelPos)] = dpv_hi(PremultipliedHighFrequencyColor);
-			HistoryOutput_Metadata[dpv_hi(LocalHistoryPixelPos)]      = taa_half2(dpv_hi(FinalLowFrequencyWeight), dpv_hi(FinalHistoryValidity));
+			HistoryOutput_Metadata[dpv_hi(LocalHistoryPixelPos)]      = tsr_half2(dpv_hi(FinalLowFrequencyWeight), dpv_hi(FinalHistoryValidity));
 			#if CONFIG_ACCUMULATE_TRANSLUCENCY_SEPARATLY
 				HistoryOutput_Translucency[dpv_hi(LocalHistoryPixelPos)] = dpv_hi(FinalTranslucencyColor).rgb;
 				HistoryOutput_TranslucencyAlpha[dpv_hi(LocalHistoryPixelPos)] = dpv_hi(FinalTranslucencyColor).a;
@@ -1463,7 +1463,7 @@ void MainCS(
 		#endif
 
 		// Output final scene color Mip0
-		taa_halfCx2 OutputColor;
+		tsr_halfCx2 OutputColor;
 		{
 			OutputColor = FinalOutputColor;
 
@@ -1480,11 +1480,11 @@ void MainCS(
 		// Output final scene color Mip1
 		#if !DEBUG_OUTPUT // running out of UAVs
 		{
-			taa_halfCx2 HalfResOutput = OutputColor * taa_half(0.25);
-			//taa_short2x2 HalfResOutputPixelPos = (LocalHistoryPixelPos >> taa_short(1)) | (((LocalHistoryPixelPos & taa_short(0x1)) | !bGenerateOutputMip1) * taa_short(~0));
-			taa_short2x2 HalfResOutputPixelPos;
-			HalfResOutputPixelPos[0] = (LocalHistoryPixelPos[0] >> taa_short(1)) | (((LocalHistoryPixelPos[0] & taa_short(0x1)) | taa_short(!bGenerateOutputMip1)) * taa_short(~0));
-			HalfResOutputPixelPos[1] = (LocalHistoryPixelPos[1] >> taa_short(1)) | (((LocalHistoryPixelPos[1] & taa_short(0x1)) | taa_short(!bGenerateOutputMip1)) * taa_short(~0));
+			tsr_halfCx2 HalfResOutput = OutputColor * tsr_half(0.25);
+			//tsr_short2x2 HalfResOutputPixelPos = (LocalHistoryPixelPos >> tsr_short(1)) | (((LocalHistoryPixelPos & tsr_short(0x1)) | !bGenerateOutputMip1) * tsr_short(~0));
+			tsr_short2x2 HalfResOutputPixelPos;
+			HalfResOutputPixelPos[0] = (LocalHistoryPixelPos[0] >> tsr_short(1)) | (((LocalHistoryPixelPos[0] & tsr_short(0x1)) | tsr_short(!bGenerateOutputMip1)) * tsr_short(~0));
+			HalfResOutputPixelPos[1] = (LocalHistoryPixelPos[1] >> tsr_short(1)) | (((LocalHistoryPixelPos[1] & tsr_short(0x1)) | tsr_short(!bGenerateOutputMip1)) * tsr_short(~0));
 			
 			#if PLATFORM_SUPPORTS_WAVE_BROADCAST && 0 // Support WaveBroadcast with halfs
 			BRANCH
@@ -1506,9 +1506,9 @@ void MainCS(
 						SharedArray1[LocalGroupThreadIndex] = dpv_hi(HalfResOutput);
 					#endif
 				#else
-					SharedArray0[LocalGroupThreadIndex] = taa_half4(dpv_lo(HalfResOutput), 0.0);
+					SharedArray0[LocalGroupThreadIndex] = tsr_half4(dpv_lo(HalfResOutput), 0.0);
 					#if CONFIG_ENABLE_DUAL_PIXEL_VECTORIZATION
-						SharedArray1[LocalGroupThreadIndex] = taa_half4(dpv_hi(HalfResOutput), 0.0);
+						SharedArray1[LocalGroupThreadIndex] = tsr_half4(dpv_hi(HalfResOutput), 0.0);
 					#endif
 				#endif
 
