Fixes for volumetric fog after moving some lighting to translated world space

diff --git a/Engine/Shaders/Private/VolumetricFog.usf b/Engine/Shaders/Private/VolumetricFog.usf
--- a/Engine/Shaders/Private/VolumetricFog.usf
+++ b/Engine/Shaders/Private/VolumetricFog.usf
@@ -324,14 +324,15 @@ void InjectShadowedLocalLightPS(
 
 			float SceneDepth;
 			float3 WorldPosition = ComputeCellWorldPosition(GridCoordinate, CellOffset, SceneDepth);
-			float3 CameraVector = normalize(WorldPosition - LWCHackToFloat(PrimaryView.WorldCameraOrigin));
+			float3 TranslatedWorldPosition = WorldPosition + LWCHackToFloat(PrimaryView.PreViewTranslation);
+			float3 CameraVector = normalize(TranslatedWorldPosition - PrimaryView.TranslatedWorldCameraOrigin);
 
 			float CellRadius = length(WorldPosition - ComputeCellWorldPosition(GridCoordinate + uint3(1, 1, 1), CellOffset));
 			// Bias the inverse squared light falloff based on voxel size to prevent aliasing near the light source
 			float DistanceBias = max(CellRadius * InverseSquaredLightDistanceBiasScale, 1);
 
 			float3 LightColor = DeferredLightUniforms.Color;
-			float LightMask = GetLocalLightAttenuation(WorldPosition, LightData, ToLight, L);
+			float LightMask = GetLocalLightAttenuation(TranslatedWorldPosition, LightData, ToLight, L);
 
 			float Lighting;
 			if( LightData.bRectLight )
@@ -357,7 +358,7 @@ void InjectShadowedLocalLightPS(
 			if (CombinedAttenuation > 0)
 			{			
 				bool bUnused = false;
-				ShadowFactor = ComputeVolumeShadowing(WorldPosition, IsPointLight, LightData.bSpotLight, bUnused);
+				ShadowFactor = ComputeVolumeShadowing(TranslatedWorldPosition, IsPointLight, LightData.bSpotLight, bUnused);
 
 #if VIRTUAL_SHADOW_MAP
 				FVirtualShadowMapSampleResult VirtualShadowMapSample = SampleVirtualShadowMap(VirtualShadowMapId, WorldPosition);
@@ -574,8 +575,8 @@ void LightScatteringCS(
 		float SceneDepth;
 		float3 WorldPosition = ComputeCellWorldPosition(GridCoordinate, CellOffset, SceneDepth);
 		float3 TranslatedWorldPosition = WorldPosition + LWCHackToFloat(PrimaryView.PreViewTranslation);
-		float CameraVectorLength = length(WorldPosition - LWCHackToFloat(PrimaryView.WorldCameraOrigin));
-		float3 CameraVector = (WorldPosition - LWCHackToFloat(PrimaryView.WorldCameraOrigin)) / CameraVectorLength;
+		float CameraVectorLength = length(TranslatedWorldPosition - PrimaryView.TranslatedWorldCameraOrigin);
+		float3 CameraVector = (TranslatedWorldPosition - PrimaryView.TranslatedWorldCameraOrigin) / CameraVectorLength;
 
 		BRANCH
 		if (ForwardLightData.HasDirectionalLight)
@@ -718,7 +719,7 @@ void LightScatteringCS(
 
 				float3 L = 0;
 				float3 ToLight = 0;
-				float LightMask = GetLocalLightAttenuation(WorldPosition, LightData, ToLight, L);
+				float LightMask = GetLocalLightAttenuation(TranslatedWorldPosition, LightData, ToLight, L);
 
 				float Lighting;
 				if( LightData.bRectLight )
