Volumetric fog computed in translated world space, fixes LWC scale

diff --git a/Engine/Shaders/Private/VolumetricFog.usf b/Engine/Shaders/Private/VolumetricFog.usf
--- a/Engine/Shaders/Private/VolumetricFog.usf
+++ b/Engine/Shaders/Private/VolumetricFog.usf
@@ -45,9 +45,9 @@ float ComputeDepthFromZSlice(float ZSlice)
 }
 
 float4x4 UnjitteredClipToTranslatedWorld;    
-float4x4 UnjitteredPrevWorldToClip;
+float4x4 UnjitteredPrevTranslatedWorldToClip;
 
-float3 ComputeCellWorldPosition(uint3 GridCoordinate, float3 CellOffset, out float SceneDepth)
+float3 ComputeCellTranslatedWorldPosition(uint3 GridCoordinate, float3 CellOffset, out float SceneDepth)
 {
 	float2 VolumeUV = (GridCoordinate.xy + CellOffset.xy) / VolumetricFog.GridSize.xy;
 	float2 VolumeNDC = (VolumeUV * 2 - 1) * float2(1, -1);
@@ -56,7 +56,18 @@ float3 ComputeCellWorldPosition(uint3 GridCoordinate, float3 CellOffset, out flo
 
 	float TileDeviceZ = ConvertToDeviceZ(SceneDepth);
 	float4 CenterPosition = mul(float4(VolumeNDC, TileDeviceZ, 1), UnjitteredClipToTranslatedWorld);
-	return CenterPosition.xyz / CenterPosition.w - LWCHackToFloat(PrimaryView.PreViewTranslation);
+	return CenterPosition.xyz / CenterPosition.w;
+}
+
+float3 ComputeCellTranslatedWorldPosition(uint3 GridCoordinate, float3 CellOffset)
+{
+	float Unused;
+	return ComputeCellTranslatedWorldPosition(GridCoordinate, CellOffset, Unused);
+}
+
+float3 ComputeCellWorldPosition(uint3 GridCoordinate, float3 CellOffset, out float SceneDepth)
+{
+	return ComputeCellTranslatedWorldPosition(GridCoordinate, CellOffset, SceneDepth) - LWCHackToFloat(PrimaryView.PreViewTranslation);
 }
 
 float3 ComputeCellWorldPosition(uint3 GridCoordinate, float3 CellOffset)
@@ -213,13 +224,12 @@ float4 LightFunctionAtlasTileMinMaxUvBound;
 Texture2D LightFunctionAtlasTexture;
 SamplerState LightFunctionAtlasSampler;
 
-float GetLocalLightFunction(float3 WorldPosition, bool IsPointLight)
+float GetLocalLightFunction(float3 TranslatedWorldPosition, bool IsPointLight)
 {
 	float2 LightFunctionUV;
 	if (IsPointLight)
 	{
 		// Get vector from light to cell world pos.
-		float TranslatedWorldPosition = WorldPosition + LWCHackToFloat(PrimaryView.PreViewTranslation);
 		float3 LightVec = normalize(TranslatedWorldPosition - DeferredLightUniforms.TranslatedWorldPosition);
 		// Transform and swizzle
 		LightVec = mul(float4(LightVec.xyz, 0), LocalLightFunctionMatrix).zyx; 
@@ -228,7 +238,7 @@ float GetLocalLightFunction(float3 WorldPosition, bool IsPointLight)
 	}
 	else
 	{
-		float4 HomogeneousShadowPosition = mul(float4(WorldPosition, 1), LocalLightFunctionMatrix);
+		float4 HomogeneousShadowPosition = mul(float4(TranslatedWorldPosition, 1), LocalLightFunctionMatrix);
 		HomogeneousShadowPosition.xyz /= HomogeneousShadowPosition.w;
 		LightFunctionUV = HomogeneousShadowPosition.xy;
 		LightFunctionUV.x = 1 - LightFunctionUV.x;
@@ -304,7 +314,7 @@ void InjectShadowedLocalLightPS(
 
 #if USE_TEMPORAL_REPROJECTION
 
-		float3 HistoryUV = ComputeVolumeUV_DEPRECATED(ComputeCellWorldPosition(GridCoordinate, .5f), UnjitteredPrevWorldToClip);
+		float3 HistoryUV = ComputeVolumeUV_DEPRECATED(ComputeCellTranslatedWorldPosition(GridCoordinate, .5f), UnjitteredPrevTranslatedWorldToClip);
 		float HistoryAlpha = HistoryWeight;
 
 		FLATTEN
@@ -323,11 +333,11 @@ void InjectShadowedLocalLightPS(
 			//float CellOffset = .5f;
 
 			float SceneDepth;
-			float3 WorldPosition = ComputeCellWorldPosition(GridCoordinate, CellOffset, SceneDepth);
-			float3 TranslatedWorldPosition = WorldPosition + LWCHackToFloat(PrimaryView.PreViewTranslation);
+			float3 TranslatedWorldPosition = ComputeCellTranslatedWorldPosition(GridCoordinate, CellOffset, SceneDepth);
+			float3 WorldPosition = TranslatedWorldPosition - LWCHackToFloat(PrimaryView.PreViewTranslation);
 			float3 CameraVector = normalize(TranslatedWorldPosition - PrimaryView.TranslatedWorldCameraOrigin);
 
-			float CellRadius = length(WorldPosition - ComputeCellWorldPosition(GridCoordinate + uint3(1, 1, 1), CellOffset));
+			float CellRadius = length(TranslatedWorldPosition - ComputeCellTranslatedWorldPosition(GridCoordinate + uint3(1, 1, 1), CellOffset));
 			// Bias the inverse squared light falloff based on voxel size to prevent aliasing near the light source
 			float DistanceBias = max(CellRadius * InverseSquaredLightDistanceBiasScale, 1);
 
@@ -368,7 +378,7 @@ void InjectShadowedLocalLightPS(
 #endif
 	
 #if	USE_LIGHT_FUNCTION
-			ShadowFactor *= GetLocalLightFunction(WorldPosition, IsPointLight);
+			ShadowFactor *= GetLocalLightFunction(TranslatedWorldPosition, IsPointLight);
 #endif
 
 			OutScattering.rgb += LightColor * (PhaseFunction(PhaseG, dot(L, -CameraVector)) * CombinedAttenuation * ShadowFactor * VolumetricScatteringIntensity);
@@ -492,13 +502,13 @@ float ComputeSkyVisibility(float3 WorldPosition, float3 BrickTextureUVs)
 	return Visibility;
 }
 
-float4x4 DirectionalLightFunctionWorldToShadow;
+float4x4 DirectionalLightFunctionTranslatedWorldToShadow;
 Texture2D LightFunctionTexture;
 SamplerState LightFunctionSampler;
 
-float GetLightFunction(float3 WorldPosition)
+float GetLightFunction(float3 TranslatedWorldPosition)
 {
-	float4 HomogeneousShadowPosition = mul(float4(WorldPosition, 1), DirectionalLightFunctionWorldToShadow);
+	float4 HomogeneousShadowPosition = mul(float4(TranslatedWorldPosition, 1), DirectionalLightFunctionTranslatedWorldToShadow);
 	float2 LightFunctionUV = HomogeneousShadowPosition.xy * .5f + .5f;
 	LightFunctionUV.y = 1 - LightFunctionUV.y;
 
@@ -526,9 +536,9 @@ void LightScatteringCS(
 	if (UseConservativeDepthTexture > 0)
 	{
 		const float3 FarDepthOffset = float3(0.5f, 0.5f, -1.0f);
-		float3 CellWorldPosition = ComputeCellWorldPosition(GridCoordinate, FarDepthOffset);
+		float3 CellTranslatedWorldPosition = ComputeCellTranslatedWorldPosition(GridCoordinate, FarDepthOffset);
 
-		float4 CellNDCPosition = mul(float4(CellWorldPosition, 1), LWCHackToFloat(PrimaryView.WorldToClip));
+		float4 CellNDCPosition = mul(float4(CellTranslatedWorldPosition, 1), PrimaryView.TranslatedWorldToClip);
 		CellNDCPosition.xyz /= CellNDCPosition.w;
 		float2 UVs = CellNDCPosition.xy * float2(0.5f, -0.5f) + 0.5f;
 
@@ -540,7 +550,7 @@ void LightScatteringCS(
 			return;
 		}
 
-		float4 PrevCellNDCPosition = mul(float4(CellWorldPosition, 1), UnjitteredPrevWorldToClip);
+		float4 PrevCellNDCPosition = mul(float4(CellTranslatedWorldPosition, 1), UnjitteredPrevTranslatedWorldToClip);
 		PrevCellNDCPosition.xyz /= PrevCellNDCPosition.w;
 		float2 PrevUVs = PrevCellNDCPosition.xy * float2(0.5f, -0.5f) + 0.5f;
 		const float PrevSceneDepth = PrevConservativeDepthTexture.Load(uint3(PrevConservativeDepthTextureSize * PrevUVs, 0));
@@ -551,7 +561,7 @@ void LightScatteringCS(
 	}
 
 #if USE_TEMPORAL_REPROJECTION
-	float3 HistoryUV = ComputeVolumeUV_DEPRECATED(ComputeCellWorldPosition(GridCoordinate, .5f), UnjitteredPrevWorldToClip);
+	float3 HistoryUV = ComputeVolumeUV_DEPRECATED(ComputeCellTranslatedWorldPosition(GridCoordinate, .5f), UnjitteredPrevTranslatedWorldToClip);
 	float HistoryAlpha = HistoryWeight;
 
 	FLATTEN
@@ -573,8 +583,8 @@ void LightScatteringCS(
 		//CellOffset = 0.5f;
 
 		float SceneDepth;
-		float3 WorldPosition = ComputeCellWorldPosition(GridCoordinate, CellOffset, SceneDepth);
-		float3 TranslatedWorldPosition = WorldPosition + LWCHackToFloat(PrimaryView.PreViewTranslation);
+		float3 TranslatedWorldPosition = ComputeCellTranslatedWorldPosition(GridCoordinate, CellOffset, SceneDepth);
+		float3 WorldPosition = TranslatedWorldPosition - LWCHackToFloat(PrimaryView.PreViewTranslation);
 		float CameraVectorLength = length(TranslatedWorldPosition - PrimaryView.TranslatedWorldCameraOrigin);
 		float3 CameraVector = (TranslatedWorldPosition - PrimaryView.TranslatedWorldCameraOrigin) / CameraVectorLength;
 
@@ -600,7 +610,7 @@ void LightScatteringCS(
 #endif
 			}
 
-			ShadowFactor *= GetLightFunction(WorldPosition);
+			ShadowFactor *= GetLightFunction(TranslatedWorldPosition);
 
 			float3 DirectionalLightColor = ForwardLightData.DirectionalLightColor;
 
@@ -681,7 +691,7 @@ void LightScatteringCS(
 		uint GridIndex = ComputeLightGridCellIndex(GridCoordinate.xy * VolumetricFog.FogGridToPixelXY, SceneDepth, 0);
 		const FCulledLightsGridData CulledLightsGrid = GetCulledLightsGrid(GridIndex, 0);
 
-		float CellRadius = length(WorldPosition - ComputeCellWorldPosition(GridCoordinate + uint3(1, 1, 1), CellOffset));
+		float CellRadius = length(TranslatedWorldPosition - ComputeCellTranslatedWorldPosition(GridCoordinate + uint3(1, 1, 1), CellOffset));
 		// Bias the inverse squared light falloff based on voxel size to prevent aliasing near the light source
 		float DistanceBiasSqr = max(CellRadius * InverseSquaredLightDistanceBiasScale, 1);
 		DistanceBiasSqr *= DistanceBiasSqr;
@@ -791,16 +801,16 @@ void FinalIntegrationCS(
 
 	float3 AccumulatedLighting = 0;
 	float AccumulatedTransmittance = 1.0f;
-	float3 PreviousSliceWorldPosition = LWCHackToFloat(PrimaryView.WorldCameraOrigin);
+	float3 PreviousSliceTranslatedWorldPosition = PrimaryView.TranslatedWorldCameraOrigin;
 
 	for (int LayerIndex = 0; LayerIndex < VolumetricFog.GridSizeInt.z; LayerIndex++)
 	{
 		uint3 LayerCoordinate = uint3(GridCoordinate.xy, LayerIndex);
 		float4 ScatteringAndExtinction = DecodeHDR(LightScattering[LayerCoordinate]);
 
-		float3 LayerWorldPosition = ComputeCellWorldPosition(LayerCoordinate, .5f);
-		float StepLength = length(LayerWorldPosition - PreviousSliceWorldPosition);
-		PreviousSliceWorldPosition = LayerWorldPosition;
+		float3 LayerTranslatedWorldPosition = ComputeCellTranslatedWorldPosition(LayerCoordinate, .5f);
+		float StepLength = length(LayerTranslatedWorldPosition - PreviousSliceTranslatedWorldPosition);
+		PreviousSliceTranslatedWorldPosition = LayerTranslatedWorldPosition;
 
 		float Transmittance = exp(-ScatteringAndExtinction.w * StepLength);
 
diff --git a/Engine/Shaders/Private/VolumetricFogLightFunction.usf b/Engine/Shaders/Private/VolumetricFogLightFunction.usf
--- a/Engine/Shaders/Private/VolumetricFogLightFunction.usf
+++ b/Engine/Shaders/Private/VolumetricFogLightFunction.usf
@@ -6,11 +6,11 @@
 #include "/Engine/Generated/Material.ush"
 #include "LightFunctionCommon.ush"
 
-float4x4 ShadowToWorld;
+float4x4 ShadowToTranslatedWorld;
 
 /** Fade distance in x, disabled brightness in y, output for preview shadows mask in z. */
 float3 LightFunctionParameters2;
-float3 LightWorldPosition;
+float3 LightTranslatedWorldPosition;
 float2 LightFunctionTexelSize;
 
 void Main(
@@ -19,8 +19,7 @@ void Main(
 	out float4 OutColor : SV_Target0
 	)
 {
-	float4 AbsoluteWorldPosition = float4(0,0,0,1);
-	float3 TranslatedWorldPosition = float3(0,0,0);
+	float4 TranslatedWorldPosition = float4(0,0,0,0);
 	float3 LightVector = float3(0,0,0);
 
 #if LIGHT_TYPE == 0 
@@ -31,8 +30,7 @@ void Main(
 	float Y = sin(Theta) * sin(Phi);
 	float Z = cos(Phi);
 	LightVector = float3(Z, Y, X);
-	AbsoluteWorldPosition = float4(LightWorldPosition + LightVector.zyx, 1);
-	TranslatedWorldPosition = AbsoluteWorldPosition.xyz + LWCHackToFloat(PrimaryView.PreViewTranslation);
+	TranslatedWorldPosition = float4(LightTranslatedWorldPosition + LightVector.zyx, 1);
 
 #else
 
@@ -50,19 +48,18 @@ void Main(
 #endif
 
 	// make SvPosition appear to be rasterized with the depth from the depth buffer
-	AbsoluteWorldPosition = mul(float4(ScreenPosition, Depth, 1), ShadowToWorld);
-	AbsoluteWorldPosition.xyz /= AbsoluteWorldPosition.w;
-	TranslatedWorldPosition = AbsoluteWorldPosition.xyz + LWCHackToFloat(PrimaryView.PreViewTranslation);
+	TranslatedWorldPosition = mul(float4(ScreenPosition, Depth, 1), ShadowToTranslatedWorld);
+	TranslatedWorldPosition.xyz /= TranslatedWorldPosition.w;
 
-	float4 Hom = mul(float4(TranslatedWorldPosition, 1), LightFunctionTranslatedWorldToLight);
+	float4 Hom = mul(float4(TranslatedWorldPosition.xyz, 1), LightFunctionTranslatedWorldToLight);
 	LightVector = Hom.xyz / Hom.w;
 #endif
 	
-	float3 LightFunction = GetLightFunctionColor(LightVector, TranslatedWorldPosition);
+	float3 LightFunction = GetLightFunctionColor(LightVector, TranslatedWorldPosition.xyz);
 	float GreyScale = dot(LightFunction, .3333f).x;
 
 	// Calculate radial view distance for stable fading
-	float ViewDistance = length(PrimaryView.TranslatedWorldCameraOrigin - TranslatedWorldPosition);
+	float ViewDistance = length(PrimaryView.TranslatedWorldCameraOrigin - TranslatedWorldPosition.xyz);
 
 	float DistanceFadeAlpha = saturate((LightFunctionParameters2.x - ViewDistance) / (LightFunctionParameters2.x * .2f));
 	// Fade to disabled based on LightFunctionFadeDistance
diff --git a/Engine/Source/Runtime/Renderer/Private/DeferredShadingRenderer.h b/Engine/Source/Runtime/Renderer/Private/DeferredShadingRenderer.h
--- a/Engine/Source/Runtime/Renderer/Private/DeferredShadingRenderer.h
+++ b/Engine/Source/Runtime/Renderer/Private/DeferredShadingRenderer.h
@@ -730,7 +730,7 @@ private:
 		const FSceneTextures& SceneTextures,
 		FIntVector VolumetricFogGridSize,
 		float VolumetricFogMaxDistance,
-		FMatrix& OutLightFunctionWorldToShadow,
+		FMatrix44f& OutLightFunctionTranslatedWorldToShadow,
 		FRDGTexture*& OutLightFunctionTexture,
 		bool& bOutUseDirectionalLightShadowing);
 
diff --git a/Engine/Source/Runtime/Renderer/Private/SceneRendering.h b/Engine/Source/Runtime/Renderer/Private/SceneRendering.h
--- a/Engine/Source/Runtime/Renderer/Private/SceneRendering.h
+++ b/Engine/Source/Runtime/Renderer/Private/SceneRendering.h
@@ -772,7 +772,7 @@ struct FTransientLightFunctionTextureAtlasTile
 struct FVolumetricFogLocalLightFunctionInfo
 {
 	FTransientLightFunctionTextureAtlasTile AtlasTile;
-	FMatrix LightFunctionMatrix;
+	FMatrix44f LightFunctionTranslatedWorldToLightMatrix;
 };
 
 class FVolumetricFogViewResources
diff --git a/Engine/Source/Runtime/Renderer/Private/VolumetricFog.cpp b/Engine/Source/Runtime/Renderer/Private/VolumetricFog.cpp
--- a/Engine/Source/Runtime/Renderer/Private/VolumetricFog.cpp
+++ b/Engine/Source/Runtime/Renderer/Private/VolumetricFog.cpp
@@ -165,11 +165,12 @@ void SetupVolumetricFogIntegrationParameters(
 {
 	Out.VolumetricFog = View.VolumetricFogResources.VolumetricFogGlobalData;
 
-	FMatrix44f UnjitteredInvTranslatedViewProjectionMatrix = FMatrix44f(View.ViewMatrices.ComputeInvProjectionNoAAMatrix() * View.ViewMatrices.GetTranslatedViewMatrix().GetTransposed());		// LWC_TODO: Precision loss?
+	FMatrix44f UnjitteredInvTranslatedViewProjectionMatrix = FMatrix44f(View.ViewMatrices.ComputeInvProjectionNoAAMatrix() * View.ViewMatrices.GetTranslatedViewMatrix().GetTransposed());
 	Out.UnjitteredClipToTranslatedWorld = UnjitteredInvTranslatedViewProjectionMatrix;
 
-	FMatrix44f UnjitteredViewProjectionMatrix = FMatrix44f(View.PrevViewInfo.ViewMatrices.GetViewMatrix() * View.PrevViewInfo.ViewMatrices.ComputeProjectionNoAAMatrix());
-	Out.UnjitteredPrevWorldToClip = UnjitteredViewProjectionMatrix;
+	FMatrix TranslatedWorldToWorld = FTranslationMatrix(-View.ViewMatrices.GetPreViewTranslation());
+	FMatrix44f UnjitteredTranslatedViewProjectionMatrix = FMatrix44f(TranslatedWorldToWorld * View.PrevViewInfo.ViewMatrices.GetViewMatrix() * View.PrevViewInfo.ViewMatrices.ComputeProjectionNoAAMatrix());
+	Out.UnjitteredPrevTranslatedWorldToClip = UnjitteredTranslatedViewProjectionMatrix;
 
 	int32 OffsetCount = IntegrationData.FrameJitterOffsetValues.Num();
 	for (int32 i = 0; i < OffsetCount; ++i)
@@ -520,7 +521,7 @@ void FDeferredShadingSceneRenderer::RenderLocalLightsForVolumetricFog(
 						continue;
 					}
 
-					PassParameters->LocalLightFunctionMatrix = FMatrix44f(LightFunctionData->LightFunctionMatrix);
+					PassParameters->LocalLightFunctionMatrix = LightFunctionData->LightFunctionTranslatedWorldToLightMatrix;
 					PassParameters->LightFunctionAtlasTexture = LightFunctionData->AtlasTile.Texture;
 					PassParameters->LightFunctionAtlasTileMinMaxUvBound = LightFunctionData->AtlasTile.MinMaxUvBound;
 				}
@@ -657,7 +658,7 @@ class FVolumetricFogLightScatteringCS : public FGlobalShader
 		SHADER_PARAMETER_STRUCT_INCLUDE(FGlobalDistanceFieldParameters2, GlobalDistanceFieldParameters)
 		SHADER_PARAMETER_RDG_TEXTURE_UAV(RWTexture2D, RWLightScattering)
 		SHADER_PARAMETER_ARRAY(FVector4f, SkySH, [3])
-		SHADER_PARAMETER(FMatrix44f, DirectionalLightFunctionWorldToShadow)
+		SHADER_PARAMETER(FMatrix44f, DirectionalLightFunctionTranslatedWorldToShadow)
 		SHADER_PARAMETER(FMatrix44f, CloudShadowmapWorldToLightClipMatrix)
 		SHADER_PARAMETER(FVector2f, PrevConservativeDepthTextureSize)
 		SHADER_PARAMETER(FVector2f, UseHeightFogColors)
@@ -955,7 +956,7 @@ void FDeferredShadingSceneRenderer::ComputeVolumetricFog(FRDGBuilder& GraphBuild
 			&& ViewFamily.bRealtimeUpdate
 			&& View.ViewState->LightScatteringHistory;
 
-		FMatrix DirectionalLightFunctionWorldToShadow;
+		FMatrix44f DirectionalLightFunctionTranslatedWorldToShadow;
 
 		RDG_EVENT_SCOPE(GraphBuilder, "VolumetricFog");
 
@@ -1017,7 +1018,7 @@ void FDeferredShadingSceneRenderer::ComputeVolumetricFog(FRDGBuilder& GraphBuild
 			SceneTextures,
 			VolumetricFogGridSize,
 			FogInfo.VolumetricFogDistance,
-			DirectionalLightFunctionWorldToShadow,
+			DirectionalLightFunctionTranslatedWorldToShadow,
 			DirectionalLightFunctionTexture,
 			bUseDirectionalLightShadowing);
 			
@@ -1119,7 +1120,7 @@ void FDeferredShadingSceneRenderer::ComputeVolumetricFog(FRDGBuilder& GraphBuild
 				PassParameters->PrevConservativeDepthTextureSize = FVector2D(1, 1);
 			}
 
-			PassParameters->DirectionalLightFunctionWorldToShadow = FMatrix44f(DirectionalLightFunctionWorldToShadow);		// LWC_TODO: Precision loss?
+			PassParameters->DirectionalLightFunctionTranslatedWorldToShadow = DirectionalLightFunctionTranslatedWorldToShadow;
 			PassParameters->LightFunctionTexture = DirectionalLightFunctionTexture;
 			PassParameters->LightFunctionSampler = TStaticSamplerState<SF_Bilinear, AM_Clamp, AM_Clamp, AM_Clamp>::GetRHI();
 
diff --git a/Engine/Source/Runtime/Renderer/Private/VolumetricFogLightFunction.cpp b/Engine/Source/Runtime/Renderer/Private/VolumetricFogLightFunction.cpp
--- a/Engine/Source/Runtime/Renderer/Private/VolumetricFogLightFunction.cpp
+++ b/Engine/Source/Runtime/Renderer/Private/VolumetricFogLightFunction.cpp
@@ -58,10 +58,10 @@ class FVolumetricFogLightFunctionPS : public FMaterialShader
 
 	BEGIN_SHADER_PARAMETER_STRUCT(FParameters, )
 		SHADER_PARAMETER(FMatrix44f, LightFunctionTranslatedWorldToLight)
-		SHADER_PARAMETER(FMatrix44f, ShadowToWorld)
+		SHADER_PARAMETER(FMatrix44f, ShadowToTranslatedWorld)
 		SHADER_PARAMETER(FVector4f, LightFunctionParameters)
 		SHADER_PARAMETER(FVector3f, LightFunctionParameters2)
-		SHADER_PARAMETER(FVector3f, LightWorldPosition)
+		SHADER_PARAMETER(FVector3f, LightTranslatedWorldPosition)
 		SHADER_PARAMETER(FVector2f, LightFunctionTexelSize)
 	END_SHADER_PARAMETER_STRUCT()
 
@@ -90,7 +90,7 @@ class FVolumetricFogLightFunctionPS : public FMaterialShader
 		const FLightSceneInfo* LightSceneInfo, 
 		const FMaterialRenderProxy* MaterialProxy,
 		FVector2D LightFunctionTexelSizeValue,
-		const FMatrix& ShadowToWorldValue)
+		const FMatrix44f& ShadowToTranslatedWorldValue)
 	{
 		FRHIPixelShader* ShaderRHI = RHICmdList.GetBoundPixelShader();
 
@@ -115,8 +115,8 @@ class FVolumetricFogLightFunctionPS : public FMaterialShader
 		}
 
 		PS.LightFunctionTexelSize = FVector2f(LightFunctionTexelSizeValue);
-		PS.ShadowToWorld = FMatrix44f(ShadowToWorldValue);		// LWC_TODO: Precision loss
-		PS.LightWorldPosition = FVector4f(LightSceneInfo->Proxy->GetPosition());
+		PS.ShadowToTranslatedWorld = ShadowToTranslatedWorldValue;
+		PS.LightTranslatedWorldPosition = FVector4f(LightSceneInfo->Proxy->GetPosition() + View.ViewMatrices.GetPreViewTranslation());
 	}
 };
 
@@ -135,7 +135,7 @@ void FDeferredShadingSceneRenderer::RenderLightFunctionForVolumetricFog(
 	const FSceneTextures& SceneTextures,
 	FIntVector VolumetricFogGridSize,
 	float VolumetricFogMaxDistance,
-	FMatrix& OutDirectionalLightFunctionWorldToShadow,
+	FMatrix44f& OutDirectionalLightFunctionTranslatedWorldToShadow,
 	FRDGTexture*& OutDirectionalLightFunctionTexture,
 	bool& bOutUseDirectionalLightShadowing)
 {
@@ -146,7 +146,7 @@ void FDeferredShadingSceneRenderer::RenderLightFunctionForVolumetricFog(
 	// The only directional light we can accept in the volumetric fog because we use the forward lighting data in the Scattering compute shader.
 	const FLightSceneProxy* SelectedForwardDirectionalLightProxy = View.ForwardLightingResources.SelectedForwardDirectionalLightProxy;
 	// Default directional light properties
-	OutDirectionalLightFunctionWorldToShadow = FMatrix::Identity;
+	OutDirectionalLightFunctionTranslatedWorldToShadow = FMatrix44f::Identity;
 	bOutUseDirectionalLightShadowing = false;
 	FLightSceneInfo* DirectionalLightSceneInfo = NULL;
 
@@ -271,7 +271,7 @@ void FDeferredShadingSceneRenderer::RenderLightFunctionForVolumetricFog(
 				0
 			);
 
-			OutDirectionalLightFunctionWorldToShadow = FTranslationMatrix(ProjectedShadowInfo.PreShadowTranslation) * FMatrix(ProjectedShadowInfo.TranslatedWorldToClipInnerMatrix);
+			OutDirectionalLightFunctionTranslatedWorldToShadow = ProjectedShadowInfo.TranslatedWorldToClipInnerMatrix;
 		}
 
 		// Now render the texture
@@ -289,12 +289,12 @@ void FDeferredShadingSceneRenderer::RenderLightFunctionForVolumetricFog(
 			PassParameters->RenderTargets[0] = FRenderTargetBinding(LightFunctionTexture, ERenderTargetLoadAction::ENoAction);
 			PassParameters->SceneTextures = SceneTextures.UniformBuffer;
 
-			FMatrix LightFunctionWorldToShadowMatrix = OutDirectionalLightFunctionWorldToShadow;
+			FMatrix44f LightFunctionTranslatedWorldToShadowMatrix = OutDirectionalLightFunctionTranslatedWorldToShadow;
 			GraphBuilder.AddPass(
 				RDG_EVENT_NAME("LightFunction %ux%u Material=%s", LightFunctionResolution.X, LightFunctionResolution.Y, *(Material.GetFriendlyName())),
 				PassParameters,
 				ERDGPassFlags::Raster,
-				[PassParameters, &View, MaterialProxyForRendering, &Material, LightFunctionResolution, DirectionalLightSceneInfo, LightFunctionWorldToShadowMatrix](FRHICommandList& RHICmdList)
+				[PassParameters, &View, MaterialProxyForRendering, &Material, LightFunctionResolution, DirectionalLightSceneInfo, LightFunctionTranslatedWorldToShadowMatrix](FRHICommandList& RHICmdList)
 				{
 					RHICmdList.SetViewport(0.f, 0.f, 0.f, LightFunctionResolution.X, LightFunctionResolution.Y, 1.f);
 
@@ -321,7 +321,7 @@ void FDeferredShadingSceneRenderer::RenderLightFunctionForVolumetricFog(
 
 					FVolumetricFogLightFunctionPS::FParameters PS;
 					PixelShader->SetParameters(RHICmdList, PS, View, DirectionalLightSceneInfo, MaterialProxyForRendering, 
-						FVector2D(1.0f / LightFunctionResolution.X, 1.0f / LightFunctionResolution.Y), LightFunctionWorldToShadowMatrix.Inverse());
+						FVector2D(1.0f / LightFunctionResolution.X, 1.0f / LightFunctionResolution.Y), LightFunctionTranslatedWorldToShadowMatrix.Inverse());
 					ClearUnusedGraphResources(PixelShader, &PS);
 					SetShaderParameters(RHICmdList, PixelShader, PixelShader.GetPixelShader(), PS);
 
@@ -348,6 +348,7 @@ void FDeferredShadingSceneRenderer::RenderLightFunctionForVolumetricFog(
 	View.VolumetricFogResources.TransientLightFunctionTextureAtlas = new (GraphBuilder.Alloc(sizeof(FTransientLightFunctionTextureAtlas), alignof(FTransientLightFunctionTextureAtlas))) FTransientLightFunctionTextureAtlas(GraphBuilder);
 
 	// Conpute all the data required by the lights
+	const FMatrix TranslatedWorldToWorld = FTranslationMatrix(-View.ViewMatrices.GetPreViewTranslation());
 	for (FLightSceneInfo* LightSceneInfo : LocalLightsToEvaluate)
 	{
 		const bool bIsPointLight = LightSceneInfo->Proxy->GetLightType() == LightType_Point;
@@ -381,8 +382,9 @@ void FDeferredShadingSceneRenderer::RenderLightFunctionForVolumetricFog(
 		FVolumetricFogLocalLightFunctionInfo& LightFunctionData = LocalLightFunctionData.Add(LightSceneInfo);
 		LightFunctionData.AtlasTile = View.VolumetricFogResources.TransientLightFunctionTextureAtlas->AllocateAtlasTile();
 
-		// Pass the World To Light when light type is Point because we need it instead of the World To Shadow. 
-		LightFunctionData.LightFunctionMatrix = (bIsPointLight) ? LightSceneInfo->Proxy->GetWorldToLight() : WorldToShadowMatrix;
+		// Pass the World To Light when light type is Point because we need it instead of the World To Shadow.
+		const FMatrix WorldToLight = bIsPointLight ? LightSceneInfo->Proxy->GetWorldToLight() : WorldToShadowMatrix;
+		LightFunctionData.LightFunctionTranslatedWorldToLightMatrix = FMatrix44f(TranslatedWorldToWorld * WorldToLight);
 	}
 
 	// And now process all the local lights, allocated atlas tile and generate the tile as needed
@@ -442,7 +444,7 @@ void FDeferredShadingSceneRenderer::RenderLightFunctionForVolumetricFog(
 
 					FVolumetricFogLightFunctionPS::FParameters PS;
 					PixelShader->SetParameters(RHICmdList, PS, View, LightSceneInfo, MaterialProxyForRendering,
-						FVector2D(1.0f / LightFunctionResolution, 1.0f / LightFunctionResolution), LightFunctionInfo.LightFunctionMatrix.Inverse());
+						FVector2D(1.0f / LightFunctionResolution, 1.0f / LightFunctionResolution), LightFunctionInfo.LightFunctionTranslatedWorldToLightMatrix.Inverse());
 					ClearUnusedGraphResources(PixelShader, &PS);
 					SetShaderParameters(RHICmdList, PixelShader, PixelShader.GetPixelShader(), PS);
 
diff --git a/Engine/Source/Runtime/Renderer/Private/VolumetricFogShared.h b/Engine/Source/Runtime/Renderer/Private/VolumetricFogShared.h
--- a/Engine/Source/Runtime/Renderer/Private/VolumetricFogShared.h
+++ b/Engine/Source/Runtime/Renderer/Private/VolumetricFogShared.h
@@ -39,7 +39,7 @@ struct FVolumetricFogIntegrationParameterData
 BEGIN_SHADER_PARAMETER_STRUCT(FVolumetricFogIntegrationParameters, )
 	SHADER_PARAMETER_STRUCT_REF(FVolumetricFogGlobalData, VolumetricFog)
 	SHADER_PARAMETER(FMatrix44f, UnjitteredClipToTranslatedWorld)
-	SHADER_PARAMETER(FMatrix44f, UnjitteredPrevWorldToClip)
+	SHADER_PARAMETER(FMatrix44f, UnjitteredPrevTranslatedWorldToClip)
 	SHADER_PARAMETER_ARRAY(FVector4f, FrameJitterOffsets, [16])
 	SHADER_PARAMETER(float, HistoryWeight)
 	SHADER_PARAMETER(uint32, HistoryMissSuperSampleCount)
