- Change texcoord semantic to avoid trampling other parameters

#rnx

diff --git a/Engine/Shaders/Private/ParticleBeamTrailVertexFactory.ush b/Engine/Shaders/Private/ParticleBeamTrailVertexFactory.ush
--- a/Engine/Shaders/Private/ParticleBeamTrailVertexFactory.ush
+++ b/Engine/Shaders/Private/ParticleBeamTrailVertexFactory.ush
@@ -23,7 +23,7 @@ struct FVertexFactoryInterpolantsVSToPS
 	TANGENTTOWORLD_INTERPOLATOR_BLOCK
 
 #if NUM_TEX_COORD_INTERPOLATORS
-	float4	TexCoords[(NUM_TEX_COORD_INTERPOLATORS+1)/2] : TEXCOORD2;
+	float4	TexCoords[(NUM_TEX_COORD_INTERPOLATORS+1)/2] : PACKED_TEXCOORDS;
 #endif
 
 #if NEEDS_PARTICLE_COLOR
