Move some utility functions from SkyAtmosphere.usf to SkyAtmosphereCommon.ush so they can be shared more easily

diff --git a/Engine/Shaders/Private/SkyAtmosphere.usf b/Engine/Shaders/Private/SkyAtmosphere.usf
--- a/Engine/Shaders/Private/SkyAtmosphere.usf
+++ b/Engine/Shaders/Private/SkyAtmosphere.usf
@@ -198,90 +198,6 @@ bool MoveToTopAtmosphere(inout float3 WorldPos, in float3 WorldDir, in float Atm
 	return true; // ok to start tracing
 }
 
-
-
-////////////////////////////////////////////////////////////
-// Participating medium properties
-////////////////////////////////////////////////////////////
-
-float RayleighPhase(float CosTheta)
-{
-	float Factor = 3.0f / (16.0f * PI);
-	return Factor * (1.0f + CosTheta * CosTheta);
-}
-
-float HgPhase(float G, float CosTheta)
-{
-	// Reference implementation (i.e. not schlick approximation). 
-	// See http://www.pbr-book.org/3ed-2018/Volume_Scattering/Phase_Functions.html
-	float Numer = 1.0f - G * G;
-	float Denom = 1.0f + G * G + 2.0f * G * CosTheta;
-	return Numer / (4.0f * PI * Denom * sqrt(Denom));
-}
-
-float3 GetAlbedo(float3 Scattering, float3 Extinction)
-{
-	return Scattering / max(0.001f, Extinction);
-}
-
-struct MediumSampleRGB
-{
-	float3 Scattering;
-	float3 Absorption;
-	float3 Extinction;
-
-	float3 ScatteringMie;
-	float3 AbsorptionMie;
-	float3 ExtinctionMie;
-
-	float3 ScatteringRay;
-	float3 AbsorptionRay;
-	float3 ExtinctionRay;
-
-	float3 ScatteringOzo;
-	float3 AbsorptionOzo;
-	float3 ExtinctionOzo;
-
-	float3 Albedo;
-};
-
-// If this is changed, please also update USkyAtmosphereComponent::GetTransmittance 
-MediumSampleRGB SampleMediumRGB(in float3 WorldPos)
-{
-	const float SampleHeight = max(0.0, (length(WorldPos) - Atmosphere.BottomRadiusKm));
-
-	const float DensityMie = exp(Atmosphere.MieDensityExpScale * SampleHeight);
-
-	const float DensityRay = exp(Atmosphere.RayleighDensityExpScale * SampleHeight);
-
-	const float DensityOzo = SampleHeight < Atmosphere.AbsorptionDensity0LayerWidth ?
-		saturate(Atmosphere.AbsorptionDensity0LinearTerm * SampleHeight + Atmosphere.AbsorptionDensity0ConstantTerm) :	// We use saturate to allow the user to create plateau, and it is free on GCN.
-		saturate(Atmosphere.AbsorptionDensity1LinearTerm * SampleHeight + Atmosphere.AbsorptionDensity1ConstantTerm);
-
-	MediumSampleRGB s;
-
-	s.ScatteringMie = DensityMie * Atmosphere.MieScattering.rgb;
-	s.AbsorptionMie = DensityMie * Atmosphere.MieAbsorption.rgb;
-	s.ExtinctionMie = DensityMie * Atmosphere.MieExtinction.rgb;
-
-	s.ScatteringRay = DensityRay * Atmosphere.RayleighScattering.rgb;
-	s.AbsorptionRay = 0.0f;
-	s.ExtinctionRay = s.ScatteringRay + s.AbsorptionRay;
-
-	s.ScatteringOzo = 0.0f;
-	s.AbsorptionOzo = DensityOzo * Atmosphere.AbsorptionExtinction.rgb;
-	s.ExtinctionOzo = s.ScatteringOzo + s.AbsorptionOzo;
-
-	s.Scattering = s.ScatteringMie + s.ScatteringRay + s.ScatteringOzo;
-	s.Absorption = s.AbsorptionMie + s.AbsorptionRay + s.AbsorptionOzo;
-	s.Extinction = s.ExtinctionMie + s.ExtinctionRay + s.ExtinctionOzo;
-	s.Albedo = GetAlbedo(s.Scattering, s.Extinction);
-
-	return s;
-}
-
-
-
 ////////////////////////////////////////////////////////////
 // LUT functions
 ////////////////////////////////////////////////////////////
@@ -294,18 +210,7 @@ MediumSampleRGB SampleMediumRGB(in float3 WorldPos)
 void UvToLutTransmittanceParams(out float ViewHeight, out float ViewZenithCosAngle, in float2 UV)
 {
 	//UV = FromSubUvsToUnit(UV, SkyAtmosphere.TransmittanceLutSizeAndInvSize); // No real impact so off
-	float Xmu = UV.x;
-	float Xr = UV.y;
-
-	float H = sqrt(Atmosphere.TopRadiusKm * Atmosphere.TopRadiusKm - Atmosphere.BottomRadiusKm * Atmosphere.BottomRadiusKm);
-	float Rho = H * Xr;
-	ViewHeight = sqrt(Rho * Rho + Atmosphere.BottomRadiusKm * Atmosphere.BottomRadiusKm);
-
-	float Dmin = Atmosphere.TopRadiusKm - ViewHeight;
-	float Dmax = Rho + H;
-	float D = Dmin + Xmu * (Dmax - Dmin);
-	ViewZenithCosAngle = D == 0.0f ? 1.0f : (H * H - Rho * Rho - D * D) / (2.0f * ViewHeight * D);
-	ViewZenithCosAngle = clamp(ViewZenithCosAngle, -1.0f, 1.0f);
+	fromTransmittanceLutUVs(ViewHeight, ViewZenithCosAngle, Atmosphere.BottomRadiusKm, Atmosphere.TopRadiusKm, UV);
 }
 
 void LutTransmittanceParamsToUv(in float ViewHeight, in float ViewZenithCosAngle, out float2 UV)
@@ -1446,7 +1351,7 @@ void RenderCameraAerialPerspectiveVolumeCS(uint3 ThreadId : SV_DispatchThreadID)
 		float3 VoxelGroundToRayStart = CamPos - VoxProjOnGround;
 		if (BelowHorizon && dot(normalize(VoxelGroundToRayStart), VoxelWorldPosNorm) < 0.0001f)
 		{
-			// We are behing the sphere and the psehre normal is pointing away from V: we are below the horizon.
+			// We are behind the sphere and the sphere normal is pointing away from V: we are below the horizon.
 			float3 MiddlePoint = 0.5f * (CamProjOnGround + VoxProjOnGround);
 			float MiddlePointHeight = length(MiddlePoint);
 
diff --git a/Engine/Shaders/Private/SkyAtmosphereCommon.ush b/Engine/Shaders/Private/SkyAtmosphereCommon.ush
--- a/Engine/Shaders/Private/SkyAtmosphereCommon.ush
+++ b/Engine/Shaders/Private/SkyAtmosphereCommon.ush
@@ -5,8 +5,7 @@
 	e.g. shared with mesh rendering or path tracing.
 =============================================================================*/
 
-#ifndef D_SKY_ATMOSPHERE_COMMON
-#define D_SKY_ATMOSPHERE_COMMON
+#pragma once
 
 // Disabled for all mobile platforms, see USkyLightComponent::IsRealTimeCaptureEnabled
 #define AP_SUPPORT_REALTIME_REFLECTION_CAPTURE (FEATURE_LEVEL > FEATURE_LEVEL_ES3_1)
@@ -129,6 +128,25 @@ float4 GetAerialPerspectiveLuminanceTransmittanceWithFogOver(
 	return FinalFog;
 }
 
+void fromTransmittanceLutUVs(
+	out float ViewHeight, out float ViewZenithCosAngle,
+	in float BottomRadius, in float TopRadius, in float2 UV)
+{
+	float Xmu = UV.x;
+	float Xr = UV.y;
+
+	float H = sqrt(TopRadius * TopRadius - BottomRadius * BottomRadius);
+	float Rho = H * Xr;
+	ViewHeight = sqrt(Rho * Rho + BottomRadius * BottomRadius);
+
+	float Dmin = TopRadius - ViewHeight;
+	float Dmax = Rho + H;
+	float D = Dmin + Xmu * (Dmax - Dmin);
+	ViewZenithCosAngle = D == 0.0f ? 1.0f : (H * H - Rho * Rho - D * D) / (2.0f * ViewHeight * D);
+	ViewZenithCosAngle = clamp(ViewZenithCosAngle, -1.0f, 1.0f);
+}
+
+
 void getTransmittanceLutUvs(
 	in float viewHeight, in float viewZenithCosAngle, in float BottomRadius, in float TopRadius,
 	out float2 UV)
@@ -230,8 +248,82 @@ float3 GetLightDiskLuminance(
 	return 0.0f;
 }
 
+////////////////////////////////////////////////////////////
+// Participating medium properties
+////////////////////////////////////////////////////////////
+
+float RayleighPhase(float CosTheta)
+{
+	float Factor = 3.0f / (16.0f * PI);
+	return Factor * (1.0f + CosTheta * CosTheta);
+}
+
+float HgPhase(float G, float CosTheta)
+{
+	// Reference implementation (i.e. not schlick approximation). 
+	// See http://www.pbr-book.org/3ed-2018/Volume_Scattering/Phase_Functions.html
+	float Numer = 1.0f - G * G;
+	float Denom = 1.0f + G * G + 2.0f * G * CosTheta;
+	return Numer / (4.0f * PI * Denom * sqrt(Denom));
+}
+
+float3 GetAlbedo(float3 Scattering, float3 Extinction)
+{
+	return Scattering / max(0.001f, Extinction);
+}
+
+struct MediumSampleRGB
+{
+	float3 Scattering;
+	float3 Absorption;
+	float3 Extinction;
+
+	float3 ScatteringMie;
+	float3 AbsorptionMie;
+	float3 ExtinctionMie;
 
+	float3 ScatteringRay;
+	float3 AbsorptionRay;
+	float3 ExtinctionRay;
+
+	float3 ScatteringOzo;
+	float3 AbsorptionOzo;
+	float3 ExtinctionOzo;
+
+	float3 Albedo;
+};
+
+// If this is changed, please also update USkyAtmosphereComponent::GetTransmittance 
+MediumSampleRGB SampleMediumRGB(in float3 WorldPos)
+{
+	const float SampleHeight = max(0.0, (length(WorldPos) - Atmosphere.BottomRadiusKm));
 
-#endif // D_SKY_ATMOSPHERE_COMMON
+	const float DensityMie = exp(Atmosphere.MieDensityExpScale * SampleHeight);
 
+	const float DensityRay = exp(Atmosphere.RayleighDensityExpScale * SampleHeight);
 
+	const float DensityOzo = SampleHeight < Atmosphere.AbsorptionDensity0LayerWidth ?
+		saturate(Atmosphere.AbsorptionDensity0LinearTerm * SampleHeight + Atmosphere.AbsorptionDensity0ConstantTerm) :	// We use saturate to allow the user to create plateau, and it is free on GCN.
+		saturate(Atmosphere.AbsorptionDensity1LinearTerm * SampleHeight + Atmosphere.AbsorptionDensity1ConstantTerm);
+
+	MediumSampleRGB s;
+
+	s.ScatteringMie = DensityMie * Atmosphere.MieScattering.rgb;
+	s.AbsorptionMie = DensityMie * Atmosphere.MieAbsorption.rgb;
+	s.ExtinctionMie = DensityMie * Atmosphere.MieExtinction.rgb;
+
+	s.ScatteringRay = DensityRay * Atmosphere.RayleighScattering.rgb;
+	s.AbsorptionRay = 0.0f;
+	s.ExtinctionRay = s.ScatteringRay + s.AbsorptionRay;
+
+	s.ScatteringOzo = 0.0f;
+	s.AbsorptionOzo = DensityOzo * Atmosphere.AbsorptionExtinction.rgb;
+	s.ExtinctionOzo = s.ScatteringOzo + s.AbsorptionOzo;
+
+	s.Scattering = s.ScatteringMie + s.ScatteringRay + s.ScatteringOzo;
+	s.Absorption = s.AbsorptionMie + s.AbsorptionRay + s.AbsorptionOzo;
+	s.Extinction = s.ExtinctionMie + s.ExtinctionRay + s.ExtinctionOzo;
+	s.Albedo = GetAlbedo(s.Scattering, s.Extinction);
+
+	return s;
+}
