Reimplement clipmap panning to reduce full-level cache invalidations during camera movement

diff --git a/Engine/Shaders/Private/VirtualShadowMaps/PageCacheCommon.ush b/Engine/Shaders/Private/VirtualShadowMaps/PageCacheCommon.ush
--- a/Engine/Shaders/Private/VirtualShadowMaps/PageCacheCommon.ush
+++ b/Engine/Shaders/Private/VirtualShadowMaps/PageCacheCommon.ush
@@ -17,3 +17,4 @@ StructuredBuffer<uint> PrevPageFlags;
 StructuredBuffer<uint> PrevPageTable;
 StructuredBuffer<FPhysicalPageMetaData> PrevPhysicalPageMetaData;
 StructuredBuffer<uint> PrevDynamicCasterPageFlags;
+ByteAddressBuffer PrevProjectionData;
diff --git a/Engine/Shaders/Private/VirtualShadowMaps/PageManagement.usf b/Engine/Shaders/Private/VirtualShadowMaps/PageManagement.usf
--- a/Engine/Shaders/Private/VirtualShadowMaps/PageManagement.usf
+++ b/Engine/Shaders/Private/VirtualShadowMaps/PageManagement.usf
@@ -381,6 +381,27 @@ void InitPhysicalPageMetaData(uint3 Index : SV_DispatchThreadID)
 	}
 }
 
+// Utility to map from PageIndex -> PageX, PageY, MipLevel
+void CalcPageAddressFromIndex(uint Index, inout uint MipLevel, inout uint2 PageAddress)
+{
+	PageAddress = uint2(0xFFFFFFFF, 0xFFFFFFFF);
+
+	// TODO: There is probably some clever math we can use for this instead;
+	// See CalcLevelOffsets for a start
+	UNROLL
+	for (MipLevel = 0; MipLevel < VSM_MAX_MIP_LEVELS - 1; ++MipLevel)
+	{
+		if (Index < CalcLevelOffsets(MipLevel + 1))
+		{
+			break;
+		}
+	}
+	const uint Level0RowMask = ((1U << VSM_LOG2_LEVEL0_DIM_PAGES_XY) - 1U);
+	const uint OffsetInLevel = Index - CalcLevelOffsets(MipLevel);
+	PageAddress.x = OffsetInLevel & (Level0RowMask >> MipLevel);
+	PageAddress.y = OffsetInLevel >> (VSM_LOG2_LEVEL0_DIM_PAGES_XY - MipLevel);
+}
+
 int bDynamicPageInvalidation;
 
 /**
@@ -394,9 +415,9 @@ void CreateCachedPageMappings(uint2 ThreadId : SV_DispatchThreadID)
 		return;
 	}
 
-	const uint PageOffsetInShadowMap = ThreadId.x;
+	const uint PageIndex = ThreadId.x;
 	const uint VirtualShadowMapId = ThreadId.y;
-	const uint GlobalPageOffset = VirtualShadowMapId * VSM_PAGE_TABLE_SIZE + PageOffsetInShadowMap;
+	const uint GlobalPageOffset = VirtualShadowMapId * VSM_PAGE_TABLE_SIZE + PageIndex;
 
 	uint ResultPageTable = VSM_PHYSICAL_PAGE_INVALID;
 	uint ResultPageFlags = 0;
@@ -410,65 +431,84 @@ void CreateCachedPageMappings(uint2 ThreadId : SV_DispatchThreadID)
 		const uint PrevVirtualShadowMapId = ShadowMapCacheData[VirtualShadowMapId].PrevVirtualShadowMapId;
 		if (PrevVirtualShadowMapId != INDEX_NONE)
 		{
-			uint PrevGlobalOffset = PrevVirtualShadowMapId * VSM_PAGE_TABLE_SIZE + PageOffsetInShadowMap;
-			// NOTE: We only care about the raw page flags from the previous frame, not the hierarchical
-			// portion as we will regenerate that later anyways depending on this frame's allocations.
-			const uint PrevFlags = PrevPageFlags[PrevGlobalOffset] & VSM_PAGE_FLAGS_BITS_MASK;
+			const FVirtualShadowMapProjectionShaderData Projection = GetVirtualShadowMapProjectionData(VirtualShadowMapId);
 
-			// True if the request matches the cached page in terms of whether it includes nanite geometry
-			bool bMatchesNonNanite = (RequestFlags & VSM_NON_NANITE_FLAG) == (PrevFlags & VSM_NON_NANITE_FLAG);
+			uint MipLevel;
+			uint2 PageAddress;
+			CalcPageAddressFromIndex(PageIndex, MipLevel, PageAddress);
 
-			if ((PrevFlags & VSM_ALLOCATED_FLAG) != 0 && bMatchesNonNanite)
+			int2 PrevPageAddress = int2(PageAddress);
+			if (Projection.LightType == LIGHT_TYPE_DIRECTIONAL)
 			{
-				uint2 PhysicalAddress = ShadowDecodePageTable(PrevPageTable[PrevGlobalOffset]).PhysicalAddress;
-				uint PhysicalPageIndex = VSMPhysicalPageAddressToIndex(PhysicalAddress);
-				FPhysicalPageMetaData PrevMetaData = PrevPhysicalPageMetaData[PhysicalPageIndex];
-				// This should probably always be true but just in case...
-				if ((PrevMetaData.Flags & VSM_ALLOCATED_FLAG) != 0)
-				{
-					// Check if moving geometry has invalidated the cached page(s)
-					uint DynamicCasterPageFlags = bDynamicPageInvalidation ? PrevDynamicCasterPageFlags[PrevGlobalOffset] : 0;
-				
-					bool bDynamicValid = bCacheStaticSeparately ? 
-						(DynamicCasterPageFlags & VSM_DYNAMIC_UNCACHED_FLAG) == 0 :
-						(DynamicCasterPageFlags == 0);
-					bool bStaticValid = bCacheStaticSeparately ? 
-						((DynamicCasterPageFlags & VSM_STATIC_UNCACHED_FLAG) == 0) :
-						(DynamicCasterPageFlags == 0);
-
-					// Stats
-					if (!bStaticValid)
-					{
-						StatsBufferInterlockedInc(VSM_STAT_STATIC_INVALIDATED_BY_DYNAMIC_PAGES);
-					}
-					else if (!bDynamicValid)
-					{
-						StatsBufferInterlockedInc(VSM_STAT_DYNAMIC_INVALIDATED_BY_DYNAMIC_PAGES);
-					}
+				// Clipmap panning
+				const FVirtualShadowMapProjectionShaderData PrevProjection = DecodeVirtualShadowMapProjectionData(PrevProjectionData, PrevVirtualShadowMapId);
+				const int OffsetScale = (VSM_LEVEL0_DIM_PAGES_XY >> 2);
+				PrevPageAddress += OffsetScale * (PrevProjection.ClipmapCornerOffset - Projection.ClipmapCornerOffset);
+			}
 
-					// If the static page is invalidated we can't use the cache at all
-					if (bStaticValid)
+			if (all(PrevPageAddress >= 0) && all(PrevPageAddress < (VSM_LEVEL0_DIM_PAGES_XY >> MipLevel)))
+			{
+				uint PrevGlobalOffset = CalcPageOffset(PrevVirtualShadowMapId, MipLevel, uint2(PrevPageAddress));
+
+				// NOTE: We only care about the raw page flags from the previous frame, not the hierarchical
+				// portion as we will regenerate that later anyways depending on this frame's allocations.
+				const uint PrevFlags = PrevPageFlags[PrevGlobalOffset] & VSM_PAGE_FLAGS_BITS_MASK;
+
+				// True if the request matches the cached page in terms of whether it includes nanite geometry
+				bool bMatchesNonNanite = (RequestFlags & VSM_NON_NANITE_FLAG) == (PrevFlags & VSM_NON_NANITE_FLAG);
+
+				if ((PrevFlags & VSM_ALLOCATED_FLAG) != 0 && bMatchesNonNanite)
+				{
+					uint2 PhysicalAddress = ShadowDecodePageTable(PrevPageTable[PrevGlobalOffset]).PhysicalAddress;
+					uint PhysicalPageIndex = VSMPhysicalPageAddressToIndex(PhysicalAddress);
+					FPhysicalPageMetaData PrevMetaData = PrevPhysicalPageMetaData[PhysicalPageIndex];
+					// This should probably always be true but just in case...
+					if ((PrevMetaData.Flags & VSM_ALLOCATED_FLAG) != 0)
 					{
-						StatsBufferInterlockedInc(VSM_STAT_ALLOCATED_PAGES);			// Total pages
-						StatsBufferInterlockedInc(VSM_STAT_STATIC_CACHED_PAGES);		// Static cached pages
-						if (bDynamicValid)
+						// Check if moving geometry has invalidated the cached page(s)
+						uint DynamicCasterPageFlags = bDynamicPageInvalidation ? PrevDynamicCasterPageFlags[PrevGlobalOffset] : 0;
+				
+						bool bDynamicValid = bCacheStaticSeparately ? 
+							(DynamicCasterPageFlags & VSM_DYNAMIC_UNCACHED_FLAG) == 0 :
+							(DynamicCasterPageFlags == 0);
+						bool bStaticValid = bCacheStaticSeparately ? 
+							((DynamicCasterPageFlags & VSM_STATIC_UNCACHED_FLAG) == 0) :
+							(DynamicCasterPageFlags == 0);
+
+						// Stats
+						if (!bStaticValid)
 						{
-							StatsBufferInterlockedInc(VSM_STAT_DYNAMIC_CACHED_PAGES);	// Dynamic cached pages
+							StatsBufferInterlockedInc(VSM_STAT_STATIC_INVALIDATED_BY_DYNAMIC_PAGES);
+						}
+						else if (!bDynamicValid)
+						{
+							StatsBufferInterlockedInc(VSM_STAT_DYNAMIC_INVALIDATED_BY_DYNAMIC_PAGES);
 						}
 
-						// Keep the same cached physical page
-						// Note that we explicitly do not retain the propogated mip data as that could change this
-						// frame, so we always "repropogate" after mapping.
-						ResultPageTable = ShadowEncodePageTable(PhysicalAddress);
-
-						// Update relevant cache flags and maintain the others
-						ResultPageFlags = PrevFlags & ~(VSM_DYNAMIC_UNCACHED_FLAG | VSM_STATIC_UNCACHED_FLAG);
-						ResultPageFlags = ResultPageFlags | (bDynamicValid ? 0 : VSM_DYNAMIC_UNCACHED_FLAG);
+						// If the static page is invalidated we can't use the cache at all
+						if (bStaticValid)
+						{
+							StatsBufferInterlockedInc(VSM_STAT_ALLOCATED_PAGES);			// Total pages
+							StatsBufferInterlockedInc(VSM_STAT_STATIC_CACHED_PAGES);		// Static cached pages
+							if (bDynamicValid)
+							{
+								StatsBufferInterlockedInc(VSM_STAT_DYNAMIC_CACHED_PAGES);	// Dynamic cached pages
+							}
+
+							// Keep the same cached physical page
+							// Note that we explicitly do not retain the propogated mip data as that could change this
+							// frame, so we always "repropogate" after mapping.
+							ResultPageTable = ShadowEncodePageTable(PhysicalAddress);
+
+							// Update relevant cache flags and maintain the others
+							ResultPageFlags = PrevFlags & ~(VSM_DYNAMIC_UNCACHED_FLAG | VSM_STATIC_UNCACHED_FLAG);
+							ResultPageFlags = ResultPageFlags | (bDynamicValid ? 0 : VSM_DYNAMIC_UNCACHED_FLAG);
 						
-						// Update metadata to mark this page as cached
-						// TODO: Separate age for static/dynamic pages probably
-						OutPhysicalPageMetaData[PhysicalPageIndex].Age = bDynamicValid ? PrevMetaData.Age + 1 : 0;
-						OutPhysicalPageMetaData[PhysicalPageIndex].Flags = ResultPageFlags;
+							// Update metadata to mark this page as cached
+							// TODO: Separate age for static/dynamic pages probably
+							OutPhysicalPageMetaData[PhysicalPageIndex].Age = bDynamicValid ? PrevMetaData.Age + 1 : 0;
+							OutPhysicalPageMetaData[PhysicalPageIndex].Flags = ResultPageFlags;
+						}
 					}
 				}
 			}
@@ -510,9 +550,9 @@ void AllocateNewPageMappings(uint2 ThreadId : SV_DispatchThreadID, uint2 GroupId
 		return;
 	}
 
-	const uint PageOffsetInShadowMap = ThreadId.x;
+	const uint PageIndex = ThreadId.x;
 	const uint VirtualShadowMapId = ThreadId.y;
-	const uint GlobalPageOffset = VirtualShadowMapId * VSM_PAGE_TABLE_SIZE + PageOffsetInShadowMap;
+	const uint GlobalPageOffset = VirtualShadowMapId * VSM_PAGE_TABLE_SIZE + PageIndex;
 
 	uint ResultPageTable = VSM_PHYSICAL_PAGE_INVALID;
 	uint ResultPageFlags = 0;
@@ -591,29 +631,16 @@ void GenerateHierarchicalPageFlags(uint2 ThreadId : SV_DispatchThreadID, uint2 G
 #endif
 	if (Flag != 0)
 	{
-		// find the mip-level for this thread.
-		uint MipLevel = 0;
-		UNROLL
-		for (MipLevel = 0; MipLevel < VSM_MAX_MIP_LEVELS - 1; ++MipLevel)
-		{
-			if (PageTableEntryIndex < CalcLevelOffsets( MipLevel + 1))
-			{
-				break;
-			}
-		}
-		uint Level0RowMask = ((1U << VSM_LOG2_LEVEL0_DIM_PAGES_XY) - 1U);
-		// Offset within level
-		uint PageTableLevelOffset = PageTableEntryIndex - CalcLevelOffsets(MipLevel);
-		// Coordinate within Mip Level
-		uint PX = PageTableLevelOffset & (Level0RowMask >> MipLevel);
-		uint PY = PageTableLevelOffset >> (VSM_LOG2_LEVEL0_DIM_PAGES_XY - MipLevel);
+		uint MipLevel;
+		uint2 PageAddress;
+		CalcPageAddressFromIndex(PageTableEntryIndex, MipLevel, PageAddress);
 
 		// Compute the min/max rect of active pages
 		uint PageBoundIndex = ShadowMapID * VSM_MAX_MIP_LEVELS + MipLevel;
-		InterlockedMin(OutPageRectBounds[PageBoundIndex].x, PX);
-		InterlockedMin(OutPageRectBounds[PageBoundIndex].y, PY);
-		InterlockedMax(OutPageRectBounds[PageBoundIndex].z, PX);
-		InterlockedMax(OutPageRectBounds[PageBoundIndex].w, PY);
+		InterlockedMin(OutPageRectBounds[PageBoundIndex].x, PageAddress.x);
+		InterlockedMin(OutPageRectBounds[PageBoundIndex].y, PageAddress.y);
+		InterlockedMax(OutPageRectBounds[PageBoundIndex].z, PageAddress.x);
+		InterlockedMax(OutPageRectBounds[PageBoundIndex].w, PageAddress.y);
 
 		// Loop over H flag levels, this builds a mip pyramid over _each_ mip level in the page table
 		// the 0-th level in this hiearchy is the page table mip level itself.
@@ -621,8 +648,7 @@ void GenerateHierarchicalPageFlags(uint2 ThreadId : SV_DispatchThreadID, uint2 G
 		// Note: starting from 1 as level 0 is the ordinary flag mip level
 		for (uint HMipLevel = 1U; HMipLevel < MaxHLevel; ++HMipLevel)
 		{
-			PX >>= 1U;
-			PY >>= 1U;
+			PageAddress.xy >>= 1U;
 
 			uint HMipBitShift = VSM_PAGE_FLAGS_BITS_PER_HMIP * HMipLevel;
 			uint HMipLevelFlagMask = VSM_PAGE_FLAGS_BITS_MASK << HMipBitShift;
@@ -630,7 +656,7 @@ void GenerateHierarchicalPageFlags(uint2 ThreadId : SV_DispatchThreadID, uint2 G
 
 			uint PreviousValue = 0;
 			uint MipToSample = MipLevel + HMipLevel;
-			uint HPageFlagOffset = CalcPageOffset(ShadowMapID, MipToSample, uint2(PX, PY));
+			uint HPageFlagOffset = CalcPageOffset(ShadowMapID, MipToSample, PageAddress);
 			InterlockedOr(OutPageFlags[HPageFlagOffset], HMipLevelFlag, PreviousValue);
 			// If this was already the value for this HMip, then whoever did that will continue up the hierarhcy.
 			// TODO: We could probably get fancier here and let a thread carry through HMip values from multiple
diff --git a/Engine/Shaders/Private/VirtualShadowMaps/ProjectionStructs.ush b/Engine/Shaders/Private/VirtualShadowMaps/ProjectionStructs.ush
--- a/Engine/Shaders/Private/VirtualShadowMaps/ProjectionStructs.ush
+++ b/Engine/Shaders/Private/VirtualShadowMaps/ProjectionStructs.ush
@@ -27,7 +27,7 @@ struct FVirtualShadowMapProjectionShaderData
 	int VirtualShadowMapId;
 };
 
-FVirtualShadowMapProjectionShaderData GetVirtualShadowMapProjectionData(int VirtualShadowMapId)
+FVirtualShadowMapProjectionShaderData DecodeVirtualShadowMapProjectionData(ByteAddressBuffer ProjectionData, int VirtualShadowMapId)
 {
 	FVirtualShadowMapProjectionShaderData Result;
 	Result.VirtualShadowMapId = VirtualShadowMapId;
@@ -36,37 +36,42 @@ FVirtualShadowMapProjectionShaderData GetVirtualShadowMapProjectionData(int Virt
 	const uint Stride = 16 * 20;
 	const uint Base = VirtualShadowMapId * Stride;
 
-	Result.TranslatedWorldToShadowViewMatrix[0]		= asfloat(VirtualShadowMap.ProjectionData.Load4(Base + 16*0));
-	Result.TranslatedWorldToShadowViewMatrix[1]		= asfloat(VirtualShadowMap.ProjectionData.Load4(Base + 16*1));
-	Result.TranslatedWorldToShadowViewMatrix[2]		= asfloat(VirtualShadowMap.ProjectionData.Load4(Base + 16*2));
-	Result.TranslatedWorldToShadowViewMatrix[3]		= asfloat(VirtualShadowMap.ProjectionData.Load4(Base + 16*3));
-	Result.ShadowViewToClipMatrix[0]				= asfloat(VirtualShadowMap.ProjectionData.Load4(Base + 16*4));
-	Result.ShadowViewToClipMatrix[1]				= asfloat(VirtualShadowMap.ProjectionData.Load4(Base + 16*5));
-	Result.ShadowViewToClipMatrix[2]				= asfloat(VirtualShadowMap.ProjectionData.Load4(Base + 16*6));
-	Result.ShadowViewToClipMatrix[3]				= asfloat(VirtualShadowMap.ProjectionData.Load4(Base + 16*7));
-	Result.TranslatedWorldToShadowUVMatrix[0]		= asfloat(VirtualShadowMap.ProjectionData.Load4(Base + 16*8));
-	Result.TranslatedWorldToShadowUVMatrix[1]		= asfloat(VirtualShadowMap.ProjectionData.Load4(Base + 16*9));
-	Result.TranslatedWorldToShadowUVMatrix[2]		= asfloat(VirtualShadowMap.ProjectionData.Load4(Base + 16*10));
-	Result.TranslatedWorldToShadowUVMatrix[3]		= asfloat(VirtualShadowMap.ProjectionData.Load4(Base + 16*11));
-	Result.TranslatedWorldToShadowUVNormalMatrix[0]	= asfloat(VirtualShadowMap.ProjectionData.Load4(Base + 16*12));
-	Result.TranslatedWorldToShadowUVNormalMatrix[1]	= asfloat(VirtualShadowMap.ProjectionData.Load4(Base + 16*13));
-	Result.TranslatedWorldToShadowUVNormalMatrix[2]	= asfloat(VirtualShadowMap.ProjectionData.Load4(Base + 16*14));
-	Result.TranslatedWorldToShadowUVNormalMatrix[3]	= asfloat(VirtualShadowMap.ProjectionData.Load4(Base + 16*15));
+	Result.TranslatedWorldToShadowViewMatrix[0]		= asfloat(ProjectionData.Load4(Base + 16*0));
+	Result.TranslatedWorldToShadowViewMatrix[1]		= asfloat(ProjectionData.Load4(Base + 16*1));
+	Result.TranslatedWorldToShadowViewMatrix[2]		= asfloat(ProjectionData.Load4(Base + 16*2));
+	Result.TranslatedWorldToShadowViewMatrix[3]		= asfloat(ProjectionData.Load4(Base + 16*3));
+	Result.ShadowViewToClipMatrix[0]				= asfloat(ProjectionData.Load4(Base + 16*4));
+	Result.ShadowViewToClipMatrix[1]				= asfloat(ProjectionData.Load4(Base + 16*5));
+	Result.ShadowViewToClipMatrix[2]				= asfloat(ProjectionData.Load4(Base + 16*6));
+	Result.ShadowViewToClipMatrix[3]				= asfloat(ProjectionData.Load4(Base + 16*7));
+	Result.TranslatedWorldToShadowUVMatrix[0]		= asfloat(ProjectionData.Load4(Base + 16*8));
+	Result.TranslatedWorldToShadowUVMatrix[1]		= asfloat(ProjectionData.Load4(Base + 16*9));
+	Result.TranslatedWorldToShadowUVMatrix[2]		= asfloat(ProjectionData.Load4(Base + 16*10));
+	Result.TranslatedWorldToShadowUVMatrix[3]		= asfloat(ProjectionData.Load4(Base + 16*11));
+	Result.TranslatedWorldToShadowUVNormalMatrix[0]	= asfloat(ProjectionData.Load4(Base + 16*12));
+	Result.TranslatedWorldToShadowUVNormalMatrix[1]	= asfloat(ProjectionData.Load4(Base + 16*13));
+	Result.TranslatedWorldToShadowUVNormalMatrix[2]	= asfloat(ProjectionData.Load4(Base + 16*14));
+	Result.TranslatedWorldToShadowUVNormalMatrix[3]	= asfloat(ProjectionData.Load4(Base + 16*15));
 
 	// NOTE: Stick with struct-element-sized loads for the moment since we may only be using subsets
 	// of the data in the calling code and we want to ensure the compiler has simple DCE options.
-	Result.ShadowPreViewTranslation			= asfloat(VirtualShadowMap.ProjectionData.Load3(Base + (16*16 + 4*0)));
-	Result.LightType						=		 (VirtualShadowMap.ProjectionData.Load (Base + (16*16 + 4*3)));
+	Result.ShadowPreViewTranslation			= asfloat(ProjectionData.Load3(Base + (16*16 + 4*0)));
+	Result.LightType						=		 (ProjectionData.Load (Base + (16*16 + 4*3)));
 
-	Result.ClipmapWorldOrigin				= asfloat(VirtualShadowMap.ProjectionData.Load3(Base + (16*17 + 4*0)));
-	Result.LightSourceRadius				= asfloat(VirtualShadowMap.ProjectionData.Load (Base + (16*17 + 4*3)));
+	Result.ClipmapWorldOrigin				= asfloat(ProjectionData.Load3(Base + (16*17 + 4*0)));
+	Result.LightSourceRadius				= asfloat(ProjectionData.Load (Base + (16*17 + 4*3)));
 
-	Result.ClipmapCornerOffset				= asint  (VirtualShadowMap.ProjectionData.Load2(Base + (16*18 + 4*0)));
-	Result.ClipmapIndex						= asint  (VirtualShadowMap.ProjectionData.Load (Base + (16*18 + 4*2)));
-	Result.ClipmapLevel						= asint  (VirtualShadowMap.ProjectionData.Load (Base + (16*18 + 4*3)));
+	Result.ClipmapCornerOffset				= asint  (ProjectionData.Load2(Base + (16*18 + 4*0)));
+	Result.ClipmapIndex						= asint  (ProjectionData.Load (Base + (16*18 + 4*2)));
+	Result.ClipmapLevel						= asint  (ProjectionData.Load (Base + (16*18 + 4*3)));
 
-	Result.ClipmapLevelCount 				= asint  (VirtualShadowMap.ProjectionData.Load (Base + (16*19 + 4*0)));
-	Result.ClipmapResolutionLodBias			= asfloat(VirtualShadowMap.ProjectionData.Load (Base + (16*19 + 4*1)));
+	Result.ClipmapLevelCount 				= asint  (ProjectionData.Load (Base + (16*19 + 4*0)));
+	Result.ClipmapResolutionLodBias			= asfloat(ProjectionData.Load (Base + (16*19 + 4*1)));
 
 	return Result;
 }
+
+FVirtualShadowMapProjectionShaderData GetVirtualShadowMapProjectionData(int VirtualShadowMapId)
+{
+	return DecodeVirtualShadowMapProjectionData(VirtualShadowMap.ProjectionData, VirtualShadowMapId);
+}
diff --git a/Engine/Source/Runtime/Renderer/Private/VirtualShadowMaps/VirtualShadowMapArray.cpp b/Engine/Source/Runtime/Renderer/Private/VirtualShadowMaps/VirtualShadowMapArray.cpp
--- a/Engine/Source/Runtime/Renderer/Private/VirtualShadowMaps/VirtualShadowMapArray.cpp
+++ b/Engine/Source/Runtime/Renderer/Private/VirtualShadowMaps/VirtualShadowMapArray.cpp
@@ -293,6 +293,7 @@ BEGIN_SHADER_PARAMETER_STRUCT(FCacheDataParameters, )
 	SHADER_PARAMETER_RDG_BUFFER_SRV( StructuredBuffer< uint >,					PrevPageTable )
 	SHADER_PARAMETER_RDG_BUFFER_SRV( StructuredBuffer< FPhysicalPageMetaData >,	PrevPhysicalPageMetaData)
 	SHADER_PARAMETER_RDG_BUFFER_SRV( StructuredBuffer< uint >,					PrevDynamicCasterPageFlags)
+	SHADER_PARAMETER_RDG_BUFFER_SRV( StructuredBuffer< uint >,					PrevProjectionData)
 END_SHADER_PARAMETER_STRUCT()
 
 static void SetCacheDataShaderParameters(FRDGBuilder& GraphBuilder, const TArray<FVirtualShadowMap*, SceneRenderingAllocator> &ShadowMaps, FVirtualShadowMapArrayCacheManager* CacheManager, FCacheDataParameters &CacheDataParameters)
@@ -316,6 +317,7 @@ static void SetCacheDataShaderParameters(FRDGBuilder& GraphBuilder, const TArray
 	CacheDataParameters.PrevPageTable = GraphBuilder.CreateSRV(GraphBuilder.RegisterExternalBuffer(CacheManager->PrevBuffers.PageTable, TEXT("Shadow.Virtual.PrevPageTable")));
 	CacheDataParameters.PrevPhysicalPageMetaData = GraphBuilder.CreateSRV(GraphBuilder.RegisterExternalBuffer(CacheManager->PrevBuffers.PhysicalPageMetaData, TEXT("Shadow.Virtual.PrevPhysicalPageMetaData")));
 	CacheDataParameters.PrevDynamicCasterPageFlags = GraphBuilder.CreateSRV(GraphBuilder.RegisterExternalBuffer(CacheManager->PrevBuffers.DynamicCasterPageFlags, TEXT("Shadow.Virtual.PrevDynamicCasterPageFlags")));
+	CacheDataParameters.PrevProjectionData = GraphBuilder.CreateSRV(GraphBuilder.RegisterExternalBuffer(CacheManager->PrevBuffers.ProjectionData, TEXT("Shadow.Virtual.PrevProjectionData")));
 }
 
 static FRDGBufferRef CreateProjectionDataBuffer(
@@ -1049,8 +1051,8 @@ void FVirtualShadowMapArray::BuildPageAllocations(
 #endif //!UE_BUILD_SHIPPING
 		
 	// Store shadow map projection data for each virtual shadow map
-	TArray<FVirtualShadowMapProjectionShaderData, SceneRenderingAllocator> ShadowMapProjectionData;
-	ShadowMapProjectionData.AddDefaulted(ShadowMaps.Num());
+	TArray<FVirtualShadowMapProjectionShaderData, SceneRenderingAllocator> ProjectionData;
+	ProjectionData.AddDefaulted(ShadowMaps.Num());
 
 	// Gather directional light virtual shadow maps
 	TArray<int32, SceneRenderingAllocator> DirectionalLightIds;
@@ -1063,7 +1065,7 @@ void FVirtualShadowMapArray::BuildPageAllocations(
 			DirectionalLightIds.Add(ClipmapID);
 			for (int32 ClipmapLevel = 0; ClipmapLevel < Clipmap->GetLevelCount(); ++ClipmapLevel)
 			{
-				ShadowMapProjectionData[ClipmapID + ClipmapLevel] = Clipmap->GetProjectionShaderData(ClipmapLevel);
+				ProjectionData[ClipmapID + ClipmapLevel] = Clipmap->GetProjectionShaderData(ClipmapLevel);
 			}
 
 			if (bDebugOutputEnabled)
@@ -1098,7 +1100,7 @@ void FVirtualShadowMapArray::BuildPageAllocations(
 
 					FViewMatrices ViewMatrices = ProjectedShadowInfo->GetShadowDepthRenderingViewMatrices( i, true );
 
-					FVirtualShadowMapProjectionShaderData& Data = ShadowMapProjectionData[ ID ];
+					FVirtualShadowMapProjectionShaderData& Data = ProjectionData[ ID ];
 					Data.TranslatedWorldToShadowViewMatrix		= FMatrix44f(ViewMatrices.GetTranslatedViewMatrix());	// LWC_TODO: Precision loss?
 					Data.ShadowViewToClipMatrix					= FMatrix44f(ViewMatrices.GetProjectionMatrix());
 					Data.TranslatedWorldToShadowUVMatrix		= FMatrix44f(CalcTranslatedWorldToShadowUVMatrix( ViewMatrices.GetTranslatedViewMatrix(), ViewMatrices.GetProjectionMatrix() ));
@@ -1119,9 +1121,9 @@ void FVirtualShadowMapArray::BuildPageAllocations(
 	UniformParameters.NumShadowMaps = ShadowMaps.Num();
 	UniformParameters.NumDirectionalLights = DirectionalLightIds.Num();
 
-	ShadowMapProjectionDataRDG = CreateProjectionDataBuffer(GraphBuilder, TEXT("Shadow.Virtual.ProjectionData"), ShadowMapProjectionData);
+	ProjectionDataRDG = CreateProjectionDataBuffer(GraphBuilder, TEXT("Shadow.Virtual.ProjectionData"), ProjectionData);
 
-	UniformParameters.ProjectionData = GraphBuilder.CreateSRV(ShadowMapProjectionDataRDG);
+	UniformParameters.ProjectionData = GraphBuilder.CreateSRV(ProjectionDataRDG);
 
 	if (CVarShowStats.GetValueOnRenderThread() || CacheManager->IsAccumulatingStats())
 	{
diff --git a/Engine/Source/Runtime/Renderer/Private/VirtualShadowMaps/VirtualShadowMapArray.h b/Engine/Source/Runtime/Renderer/Private/VirtualShadowMaps/VirtualShadowMapArray.h
--- a/Engine/Source/Runtime/Renderer/Private/VirtualShadowMaps/VirtualShadowMapArray.h
+++ b/Engine/Source/Runtime/Renderer/Private/VirtualShadowMaps/VirtualShadowMapArray.h
@@ -89,7 +89,7 @@ struct FVirtualShadowMapProjectionShaderData
 	FVector3f ClipmapWorldOrigin;
 	float LightSourceRadius;				// This should live in shared light structure...
 	
-	FIntPoint ClipmapCornerOffset;
+	FIntPoint ClipmapCornerOffset = FIntPoint(0, 0);
 	int32 ClipmapIndex = 0;					// 0 .. ClipmapLevelCount-1
 	int32 ClipmapLevel = 0;					// "Absolute" level, can be negative
 	int32 ClipmapLevelCount = 0;
@@ -311,7 +311,7 @@ public:
 	// Used to clip the rect size of clusters during culling.
 	FRDGBufferRef PageRectBoundsRDG = nullptr;
 	FRDGBufferRef AllocatedPageRectBoundsRDG = nullptr;
-	FRDGBufferRef ShadowMapProjectionDataRDG = nullptr;
+	FRDGBufferRef ProjectionDataRDG = nullptr;
 
 	// HZB generated for the *current* frame's physical page pool
 	// We use the *previous* frame's HZB (from VirtualShadowMapCacheManager) for culling the current frame
diff --git a/Engine/Source/Runtime/Renderer/Private/VirtualShadowMaps/VirtualShadowMapCacheManager.cpp b/Engine/Source/Runtime/Renderer/Private/VirtualShadowMaps/VirtualShadowMapCacheManager.cpp
--- a/Engine/Source/Runtime/Renderer/Private/VirtualShadowMaps/VirtualShadowMapCacheManager.cpp
+++ b/Engine/Source/Runtime/Renderer/Private/VirtualShadowMaps/VirtualShadowMapCacheManager.cpp
@@ -45,6 +45,14 @@ static TAutoConsoleVariable<int32> CVarCacheVsmUseHzb(
 	TEXT("Enables testing HZB for Virtual Shadow Map invalidations."),
 	ECVF_RenderThreadSafe);
 
+int32 GClipmapPanning = 1;
+FAutoConsoleVariableRef CVarEnableClipmapPanning(
+	TEXT("r.Shadow.Virtual.Cache.ClipmapPanning"),
+	GClipmapPanning,
+	TEXT("Enable support for panning cached clipmap pages for directional lights."),
+	ECVF_RenderThreadSafe
+);
+
 void FVirtualShadowMapCacheEntry::UpdateClipmap(
 	int32 VirtualShadowMapId,
 	const FMatrix &WorldToLight,
@@ -62,7 +70,7 @@ void FVirtualShadowMapCacheEntry::UpdateClipmap(
 		//UE_LOG(LogRenderer, Display, TEXT("Invalidated clipmap level (VSM %d) due to light movement"), VirtualShadowMapId);
 	}
 
-	if (bCacheValid)
+	if (bCacheValid && GClipmapPanning == 0)
 	{
 		if (PageSpaceLocation.X != PrevPageSpaceLocation.X ||
 			PageSpaceLocation.Y != PrevPageSpaceLocation.Y)
@@ -415,7 +423,7 @@ void FVirtualShadowMapArrayCacheManager::ExtractFrameData(
 		
 			GraphBuilder.QueueBufferExtraction(VirtualShadowMapArray.PhysicalPageMetaDataRDG, &PrevBuffers.PhysicalPageMetaData);
 			GraphBuilder.QueueBufferExtraction(VirtualShadowMapArray.DynamicCasterPageFlagsRDG, &PrevBuffers.DynamicCasterPageFlags);
-			GraphBuilder.QueueBufferExtraction(VirtualShadowMapArray.ShadowMapProjectionDataRDG, &PrevBuffers.ShadowMapProjectionDataBuffer);
+			GraphBuilder.QueueBufferExtraction(VirtualShadowMapArray.ProjectionDataRDG, &PrevBuffers.ProjectionData);
 			GraphBuilder.QueueBufferExtraction(VirtualShadowMapArray.PageRectBoundsRDG, &PrevBuffers.PageRectBounds);
 
 			GraphBuilder.QueueBufferExtraction(VirtualShadowMapArray.InvalidatingInstancesRDG, &PrevBuffers.InvalidatingInstancesBuffer);
@@ -770,7 +778,7 @@ static void SetupCommonParameters(FRDGBuilder& GraphBuilder, FVirtualShadowMapAr
 	FVirtualShadowMapArrayFrameData &PrevBuffers = CacheManager->PrevBuffers;
 
 	// Update references in our last frame uniform buffer with reimported resources for this frame
-	CacheManager->PrevUniformParameters.ProjectionData = RegExtCreateSrv(PrevBuffers.ShadowMapProjectionDataBuffer, TEXT("Shadow.Virtual.PrevProjectionData"));
+	CacheManager->PrevUniformParameters.ProjectionData = RegExtCreateSrv(PrevBuffers.ProjectionData, TEXT("Shadow.Virtual.PrevProjectionData"));
 	CacheManager->PrevUniformParameters.PageTable = RegExtCreateSrv(PrevBuffers.PageTable, TEXT("Shadow.Virtual.PrevPageTable"));
 	CacheManager->PrevUniformParameters.PageFlags = RegExtCreateSrv(PrevBuffers.PageFlags, TEXT("Shadow.Virtual.PrevPageFlags"));
 	CacheManager->PrevUniformParameters.PageRectBounds = RegExtCreateSrv(PrevBuffers.PageRectBounds, TEXT("Shadow.Virtual.PrevPageRectBounds"));
diff --git a/Engine/Source/Runtime/Renderer/Private/VirtualShadowMaps/VirtualShadowMapCacheManager.h b/Engine/Source/Runtime/Renderer/Private/VirtualShadowMaps/VirtualShadowMapCacheManager.h
--- a/Engine/Source/Runtime/Renderer/Private/VirtualShadowMaps/VirtualShadowMapCacheManager.h
+++ b/Engine/Source/Runtime/Renderer/Private/VirtualShadowMaps/VirtualShadowMapCacheManager.h
@@ -82,7 +82,7 @@ struct FVirtualShadowMapArrayFrameData
 	TRefCountPtr<FRDGPooledBuffer>				PageTable;
 	TRefCountPtr<FRDGPooledBuffer>				PageFlags;
 
-	TRefCountPtr<FRDGPooledBuffer>				ShadowMapProjectionDataBuffer;
+	TRefCountPtr<FRDGPooledBuffer>				ProjectionData;
 	TRefCountPtr<FRDGPooledBuffer>				PageRectBounds;
 
 	TRefCountPtr<FRDGPooledBuffer>				DynamicCasterPageFlags;
