Add support voxel jitter mode for raytracing passes.

This allows to have stable raytracing voxel for EngineTests.

diff --git a/Engine/Shaders/Private/HairStrands/HairStrandsRaytracing.ush b/Engine/Shaders/Private/HairStrands/HairStrandsRaytracing.ush
--- a/Engine/Shaders/Private/HairStrands/HairStrandsRaytracing.ush
+++ b/Engine/Shaders/Private/HairStrands/HairStrandsRaytracing.ush
@@ -94,9 +94,23 @@ float TraverseHair(uint2 PixelCoord, inout RandomSequence RandSequence, float3 R
 {
 	float HitT = InMinT;
 
-	const float4 Rnd = RandomSequence_GenerateSample4D(RandSequence);
-	const float3 VoxelRandom = Rnd.xyz;
-	const float RayRandom = Rnd.w;
+	// 0: No jitter, 1:Time-Random jitter, 2:Time-Constant jitter
+	float3 VoxelRandom = 0.5f;
+	float  RayRandom   = 1.0f;
+	if (VirtualVoxel.JitterMode == 1)
+	{
+		const float4 Rnd = RandomSequence_GenerateSample4D(RandSequence);
+		VoxelRandom = Rnd.xyz;
+		RayRandom = Rnd.w;
+	}
+	else
+	{
+		const uint2  Seed0 = Rand3DPCG16(int3(PixelCoord, 0)).xy;
+		const uint2  Seed1 = Rand3DPCG16(int3(PixelCoord + 17, 0)).xy;
+		const float4 Rnd  = float4(Hammersley16(0, 8, Seed0), Hammersley16(0, 8, Seed1));
+		VoxelRandom = Rnd.xyz;
+		RayRandom = Rnd.w;
+	}
 
 	return InternalTraverseHair(PixelCoord, RayOrigin, RayDirection, InMinT, InOcclusionThreshold, VoxelRandom, RayRandom);
 }
\ No newline at end of file
