Fix invalid shader draw primitives when converting from WS to TWS with matrix arguments.

diff --git a/Engine/Shaders/Private/GPUSceneDebugRender.usf b/Engine/Shaders/Private/GPUSceneDebugRender.usf
--- a/Engine/Shaders/Private/GPUSceneDebugRender.usf
+++ b/Engine/Shaders/Private/GPUSceneDebugRender.usf
@@ -158,7 +158,7 @@ void GPUSceneDebugRenderCS(uint InstanceID : SV_DispatchThreadID)
 	if (bIsMouseOver || bIsSelected || bDrawAll != 0)
 	{
 		AddOBBWS(InstanceData.LocalBoundsCenter - InstanceData.LocalBoundsExtent, InstanceData.LocalBoundsCenter + InstanceData.LocalBoundsExtent, InstanceBoundColor, LWCHackToFloat(InstanceData.LocalToWorld));
-		AddReferentialWS(LWCHackToFloat(InstanceData.LocalToWorld), 1);
+		AddReferentialWS(LWCHackToFloat(InstanceData.LocalToWorld), 50.f);
 	}
 
 	uint StartOffsetX = 10;
diff --git a/Engine/Shaders/Private/ShaderDrawDebug.ush b/Engine/Shaders/Private/ShaderDrawDebug.ush
--- a/Engine/Shaders/Private/ShaderDrawDebug.ush
+++ b/Engine/Shaders/Private/ShaderDrawDebug.ush
@@ -207,9 +207,9 @@ void AddCross(float3 Pos, float Size, float4 Color)		{ AddAABBTWS(Pos + ShaderDr
 
 void AddReferentialTWS(float3 Pos, float3 T, float3 B, float3 N, float Scale=1)
 {
-	AddLineTWS(Pos, Pos + T*Scale, ColorRed);
-	AddLineTWS(Pos, Pos + B*Scale, ColorGreen);
-	AddLineTWS(Pos, Pos + N*Scale, ColorBlue);
+	AddLineTWS(Pos, Pos + normalize(T)*Scale, ColorRed);
+	AddLineTWS(Pos, Pos + normalize(B)*Scale, ColorGreen);
+	AddLineTWS(Pos, Pos + normalize(N)*Scale, ColorBlue);
 }
 
 void AddReferentialTWS(float3 Pos, float3 TangentZ, float Scale=1)
@@ -229,10 +229,10 @@ void AddReferentialTWS(float4x4 InM, float Scale = 1)
 	AddReferentialTWS(InM[3].xyz, InM[0].xyz, InM[1].xyz, InM[2].xyz, Scale);
 }
 
-void AddReferentialWS(float4x4 InM, float Scale = 1)							{ AddReferentialTWS(InM[3].xyz, InM[0].xyz, InM[1].xyz, InM[2].xyz + ShaderDrawTranslatedWorldOffset, Scale); }
+void AddReferentialWS(float4x4 InM, float Scale = 1)							{ AddReferentialTWS(InM[3].xyz + ShaderDrawTranslatedWorldOffset, InM[0].xyz, InM[1].xyz, InM[2].xyz, Scale); }
 void AddReferentialWS(float3 Pos, float3 TangentZ, float Scale = 1)				{ AddReferentialTWS(Pos + ShaderDrawTranslatedWorldOffset, TangentZ, Scale); }
 void AddReferentialWS(float3 Pos, float3 T, float3 B, float3 N, float Scale = 1){ AddReferentialTWS(Pos + ShaderDrawTranslatedWorldOffset, T, B, N, Scale); }
-void AddReferential(float4x4 InM, float Scale = 1)								{ AddReferentialTWS(InM[3].xyz, InM[0].xyz, InM[1].xyz, InM[2].xyz + ShaderDrawTranslatedWorldOffset, Scale); }
+void AddReferential(float4x4 InM, float Scale = 1)								{ AddReferentialTWS(InM[3].xyz + ShaderDrawTranslatedWorldOffset, InM[0].xyz, InM[1].xyz, InM[2].xyz, Scale); }
 void AddReferential(float3 Pos, float3 TangentZ, float Scale = 1)				{ AddReferentialTWS(Pos + ShaderDrawTranslatedWorldOffset, TangentZ, Scale); }
 void AddReferential(float3 Pos, float3 T, float3 B, float3 N, float Scale = 1)	{ AddReferentialTWS(Pos + ShaderDrawTranslatedWorldOffset, T, B, N, Scale); }
 
@@ -286,7 +286,7 @@ void AddOBBTWS(float3 Min, float3 Max, float4 Color, float4x4 Transform)
 	}
 }
 
-void AddOBBWS(float3 Min, float3 Max, float4 Color, float4x4 Transform) { Transform[2].xyz += ShaderDrawTranslatedWorldOffset; AddOBBTWS(Min, Max, Color, Transform); }
-void AddOBB(float3 Min, float3 Max, float4 Color, float4x4 Transform)	{ Transform[2].xyz += ShaderDrawTranslatedWorldOffset; AddOBBTWS(Min, Max, Color, Transform); }
+void AddOBBWS(float3 Min, float3 Max, float4 Color, float4x4 Transform) { Transform[3].xyz += ShaderDrawTranslatedWorldOffset; AddOBBTWS(Min, Max, Color, Transform); }
+void AddOBB(float3 Min, float3 Max, float4 Color, float4x4 Transform)	{ Transform[3].xyz += ShaderDrawTranslatedWorldOffset; AddOBBTWS(Min, Max, Color, Transform); }
 
 #endif
