Remove unused function

diff --git a/Engine/Shaders/Private/MonteCarlo.ush b/Engine/Shaders/Private/MonteCarlo.ush
--- a/Engine/Shaders/Private/MonteCarlo.ush
+++ b/Engine/Shaders/Private/MonteCarlo.ush
@@ -338,25 +338,6 @@ float4 UniformSampleConeConcentricRobust(float2 E, float SinThetaMax2)
 	return float4(L, PDF);
 }
 
-float4 ImportanceSampleBlinn( float2 E, float a2 )
-{
-	float n = 2 / a2 - 2;
-
-	float Phi = 2 * PI * E.x;
-	float CosTheta = ClampedPow( E.y, 1 / (n + 1) );
-	float SinTheta = sqrt( 1 - CosTheta * CosTheta );
-
-	float3 H;
-	H.x = SinTheta * cos( Phi );
-	H.y = SinTheta * sin( Phi );
-	H.z = CosTheta;
-
-	float D = (n+2) / (2*PI) * ClampedPow( CosTheta, n );
-	float PDF = D * CosTheta;
-
-	return float4( H, PDF );
-}
-
 // PDF = D * NoH / (4 * VoH)
 float4 ImportanceSampleGGX( float2 E, float a2 )
 {
