- Fix Radix sort reading off end of groupshared

#rnx

diff --git a/Engine/Shaders/Private/RadixSortShaders.usf b/Engine/Shaders/Private/RadixSortShaders.usf
--- a/Engine/Shaders/Private/RadixSortShaders.usf
+++ b/Engine/Shaders/Private/RadixSortShaders.usf
@@ -458,7 +458,7 @@ void RadixSort_Spine(
 	GroupMemoryBarrierWithGroupSync();
 
 	uint DigitSum = 0;
-	if(ThreadId < RAKING_THREAD_COUNT)
+	if(ThreadId < DIGIT_COUNT)
 	{
 		DigitSum = LocalRakingTotals[RAKING_COUNTERS_PER_DIGIT * (ThreadId+1) - 1];
 	}
@@ -635,7 +635,7 @@ void ScanKeys(
 	GroupMemoryBarrierWithGroupSync();
 
 	uint DigitSum = 0;
-	if(ThreadId < RAKING_THREAD_COUNT)
+	if(ThreadId < DIGIT_COUNT)
 	{
 		DigitSum = LocalRakingTotals[RAKING_COUNTERS_PER_DIGIT * (ThreadId+1) - 1];
 	}
