Add clipmap address space visualization and hook up to (advanced) visualization stuff.
Move visualize enum/defines to shared header file.
Rename a few cvars for consistency.

diff --git a/Engine/Shaders/Private/VirtualShadowMaps/Debug.usf b/Engine/Shaders/Private/VirtualShadowMaps/Debug.usf
--- a/Engine/Shaders/Private/VirtualShadowMaps/Debug.usf
+++ b/Engine/Shaders/Private/VirtualShadowMaps/Debug.usf
@@ -6,183 +6,153 @@
 
 #include "../Common.ush"
 #include "../SceneTexturesCommon.ush"
-//#include "../Nanite/DebugVisualize.ush"
 #include "PageOverlap.ush"
 #include "PageCacheCommon.ush"
+#include "ProjectionCommon.ush"
+#include "../Visualization.ush"
+#include "/Engine/Shared/VirtualShadowMapDefinitions.h"
 
 uint DebugTargetWidth;
 uint DebugTargetHeight;
 uint BorderWidth;
-uint ZoomScaleFactor;
-uint DebugMethod;
+uint VisualizeModeId;	// VIRTUAL_SHADOW_MAP_VISUALIZE_*
+int VirtualShadowMapId;
 
-Texture2D< float >			HZBPhysical;
+RWTexture2D<float4> OutVisualize;
 
-float3 IntToColor(uint Index)
-{
-	const float3 ColorTable[9] = {
-		float3(0.0f, 0.0f, 0.0f),
-		float3(1.0f, 0.0f, 0.0f),
-		float3(0.0f, 1.0f, 0.0f),
-		float3(0.0f, 0.0f, 1.0f),
-
-		float3(1.0f, 1.0f, 0.0f),
-		float3(0.0f, 1.0f, 1.0f),
-		float3(1.0f, 0.0f, 1.0f),
-
-		float3(0.5f, 0.7f, 0.1f),
-		float3(1.0f, 1.0f, 1.0f)
-	};
-	return ColorTable[Index % 9];
-}
+static const float3 BackGroundColor = float3(0.1f, 0.1f, 0.1f);
+static const float3 BlackColor = float3(0.0f, 0.0f, 0.0f);
 
-RWTexture2D<float4>		DebugOutput;
-/**
- * 
- */
-[numthreads(VSM_DEFAULT_CS_GROUP_XY, VSM_DEFAULT_CS_GROUP_XY, 1)]
-void DebugVisualizeVirtualSmCS(uint2 OutputPixel : SV_DispatchThreadID)
+
+float3 DebugVisualizeClipmap(float2 PixelUV, FVirtualShadowMapProjectionShaderData ProjectionData0)
 {
-	if (any(OutputPixel >= uint2(DebugTargetWidth, DebugTargetHeight)))
-	{
-		return;
-	}
+	float UVScale = 0.0f;
+	int LastValidLevelIndex = -1;
 
-	if( DebugMethod > 5 )
+#if 1
+	// First figure out a scale factor to zoom in on non-empty pages.
+	for (int Index = ProjectionData0.ClipmapLevelCount - 1; Index >= 0; --Index)
 	{
-		float3 Color = 0;
-
-		uint ShadowMapID = 0;
-		uint2 Texel = OutputPixel * ( VSM_VIRTUAL_MAX_RESOLUTION_XY / 2048 );
+		int LevelSmId = ProjectionData0.VirtualShadowMapId + Index;
+		uint4 AllocatedBounds = VirtualShadowMap.PageRectBounds[LevelSmId * VSM_MAX_MIP_LEVELS];
 
-		float Depth = 0;
-		float HZBDepth = 0;
 
-		uint Level = 0;
-		for( ; Level < VSM_MAX_MIP_LEVELS; Level++ )
+		if (all(AllocatedBounds.xy <= AllocatedBounds.zw))
 		{
-			uint2 vAddress = Texel >> Level;
-			uint2 vPage = vAddress >> VSM_LOG2_PAGE_SIZE;
-			FShadowPhysicalPage pPage = ShadowGetPhysicalPage( CalcPageOffset( ShadowMapID, Level, vPage ) );
-
-			if( pPage.bThisLODValid )
+			FVirtualShadowMapProjectionShaderData Level = GetVirtualShadowMapProjectionData(LevelSmId);
+
+			float4 ShadowUVz = mul(float4(Level.ClipmapWorldOrigin + Level.ShadowPreViewTranslation, 1.0f), Level.TranslatedWorldToShadowUVMatrix);
+			const float2 PageCenter = ShadowUVz * float(VSM_LEVEL0_DIM_PAGES_XY);
+			float2 AllocatedExtentPages = max(abs(float2(AllocatedBounds.xy) - PageCenter), abs(float2(AllocatedBounds.zw) - PageCenter));
+			
+			float NewUVScale = (2.0f + max(AllocatedExtentPages.x, AllocatedExtentPages.y)) / float(VSM_LEVEL0_DIM_PAGES_XY / 2);
+		
+			if (LastValidLevelIndex < 0)
+			{
+				UVScale = NewUVScale;
+				LastValidLevelIndex = Index;
+			}
+			else
 			{
-				uint2 pAddress = pPage.PhysicalAddress * VSM_PAGE_SIZE + ( vAddress & VSM_PAGE_SIZE_MASK );
-				Depth = asfloat( VirtualShadowMap.PhysicalPagePool[ pAddress ] );
-				break;
+				UVScale = max(UVScale, NewUVScale / float(1 << (LastValidLevelIndex - Index)));
 			}
 		}
+	}
+#endif
+
+	if (LastValidLevelIndex < 0)
+	{
+		return BlackColor;
+	}
 
-		Level = DebugMethod - 6;
-		for( ; Level < VSM_MAX_MIP_LEVELS; Level++ )
+	const float2 CenteredPixelUV = PixelUV - 0.5f;
+
+	int PageOverlayCount = 0;
+
+	float3 ResultColor = BlackColor;
+	for (int Index = LastValidLevelIndex; Index >= 0; --Index)
+	{
+		int LevelSmId = ProjectionData0.VirtualShadowMapId  + Index;
+
+		FVirtualShadowMapProjectionShaderData Level = GetVirtualShadowMapProjectionData(LevelSmId);
+
+		float4 ShadowUVz = mul(float4(Level.ClipmapWorldOrigin + Level.ShadowPreViewTranslation, 1.0f), Level.TranslatedWorldToShadowUVMatrix);
+
+		int2 vPage = int2((CenteredPixelUV * UVScale + ShadowUVz) * float(VSM_LEVEL0_DIM_PAGES_XY));
+		UVScale *= 2.0f;
+
+		if (any(vPage < 0) || any(vPage >= VSM_LEVEL0_DIM_PAGES_XY))
 		{
-			uint2 vAddress = Texel >> Level;
-			uint2 vPage = vAddress >> VSM_LOG2_PAGE_SIZE;
-			//uint2 pPage = HZBPageTable[ CalcPageOffset( ShadowMapID, Level, vPage ) ];
-			FShadowPhysicalPage pPage = ShadowGetPhysicalPage( CalcPageOffset( ShadowMapID, Level, vPage ) );
+			continue;
+		}
 
-			if( pPage.bThisLODValid )
-			{
-				vAddress >>= 1;
+		uint PageOffset = CalcPageOffset(LevelSmId, 0, vPage);
+		FShadowPhysicalPage pPage = ShadowGetPhysicalPage(PageOffset);
 
-				uint2 pAddress = pPage.PhysicalAddress * HZB_PAGE_SIZE + ( vAddress & (HZB_PAGE_SIZE - 1) );
-				HZBDepth = HZBPhysical[ pAddress ];
-				break;
-			}
+		if (pPage.bThisLODValid)
+		{
+			PageOverlayCount += 1;
+			ResultColor = lerp(IntToColor(vPage.x + (vPage.y << 10U)), IntToColor(Level.ClipmapLevel), 0.75f);
 		}
+	}
 
-		//Color = Depth - HZBDepth;
-		//Color.r = Depth;
-		Color = Depth;
-		//Color = HZBDepth;
+	if (PageOverlayCount > 0)
+	{
+		//ResultColor = GreenToRedTurbo(float(PageOverlayCount - 1) / float(8));
+	}
 
-		DebugOutput[OutputPixel] = float4(Color, 1.0f);
-		return;
+	return ResultColor;
+}
+
+float3 DebugVisualizeClipmap(uint2 OutputPixel, int VirtualShadowMapId)
+{
+	// Figure out the square area do draw in (could be any sub-rect)
+	uint SquareDomainDim = min(DebugTargetWidth, DebugTargetHeight) - BorderWidth * 2U;
+
+	// fill background pixels
+	if (any(OutputPixel < BorderWidth || OutputPixel >(SquareDomainDim + BorderWidth)))
+	{
+		return BackGroundColor;
 	}
 
-	uint2 InputPixel = OutputPixel / ZoomScaleFactor;
+	// The pixel we are drawing in 0-1 range UV within the domain
+	float2 PixelUV = float2(OutputPixel - BorderWidth) / float(SquareDomainDim);
 
-	uint VirtualShadowMapId = max(0, int(InputPixel.y) - 2) / (CalcLevelDimsPages(0) + BorderWidth);
-	uint MipLevel = 0;
-	uint Offset = 0;
+	FVirtualShadowMapProjectionShaderData ProjectionData = GetVirtualShadowMapProjectionData(VirtualShadowMapId);
+	if (ProjectionData.LightType != LIGHT_TYPE_DIRECTIONAL)
+	{
+		return BackGroundColor;
+	}
 
-	while (InputPixel.x >= Offset + CalcLevelDimsPages(MipLevel) + BorderWidth)
+	return DebugVisualizeClipmap(PixelUV, ProjectionData);
+}
+
+
+[numthreads(VSM_DEFAULT_CS_GROUP_XY, VSM_DEFAULT_CS_GROUP_XY, 1)]
+void DebugVisualizeVirtualSmCS(uint2 PixelPos : SV_DispatchThreadID)
+{
+	if (any(PixelPos >= uint2(DebugTargetWidth, DebugTargetHeight)))
 	{
-		Offset += CalcLevelDimsPages(MipLevel) + BorderWidth;
-		MipLevel += 1;
+		return;
 	}
 
-	uint2 PageAddress = uint2(InputPixel.x - Offset, InputPixel.y % (CalcLevelDimsPages(0) + BorderWidth));
+	// NOTE: Important to *not* write output if it isn't a recognized mode, as a different
+	// pass may write that output instead.
+	float3 Output = float3(1, 0, 1);
+	bool bWriteOutput = true;
 
-	if (all(PageAddress - BorderWidth < CalcLevelDimsPages(MipLevel)) && VirtualShadowMapId < VirtualShadowMap.NumShadowMaps)
+	if (VisualizeModeId == VIRTUAL_SHADOW_MAP_VISUALIZE_CLIPMAP_VIRTUAL_SPACE)
 	{
-		PageAddress -= BorderWidth;
-		if (DebugMethod >= 1 && DebugMethod <= 2)
-		{
-			if (MipLevel == 0 && DebugMethod == 2)
-			{
-				int OverAlloc = 0;
-				//uint PageFlag = PageFlags[CalcPageOffset(VirtualShadowMapId, MipLevel, PageAddress)];
-				for (uint Level = 0; Level < VSM_MAX_MIP_LEVELS; ++Level)
-				{
-					if (VirtualShadowMap.PageFlags[CalcPageOffset(VirtualShadowMapId, Level, PageAddress >> Level)] & VSM_ALLOCATED_FLAG)
-					{
-						OverAlloc += 1;
-					}
-				}
-				float3 Color = IntToColor(OverAlloc);
-				DebugOutput[OutputPixel] = float4(Color, 1.0f);
-			}
-			else
-			{
-				float3 Color = IntToColor(MipLevel + 1);
-				uint PageFlag = VirtualShadowMap.PageFlags[CalcPageOffset(VirtualShadowMapId, MipLevel, PageAddress)];
-				DebugOutput[OutputPixel] = float4((PageFlag & VSM_ALLOCATED_FLAG) ? Color : float3(0.0f, 0.0f, 0.0f), 1.0f);
-			}
-		}
-		else
-		{
-#if HAS_CACHE_DATA
-			if (DebugMethod == 3)
-			{
-				uint PrevPageFlag = 0;
-				uint PrevDynFlag = 0;
-
-				const uint PrevVirtualShadowMapId = ShadowMapCacheData[VirtualShadowMapId].PrevVirtualShadowMapId;
-				if (PrevVirtualShadowMapId != INDEX_NONE)
-				{
-					uint PrevOffset = CalcPageOffset(PrevVirtualShadowMapId, MipLevel, PageAddress);
-					PrevPageFlag = PrevPageFlags[PrevOffset];
-					PrevDynFlag = PrevDynamicCasterPageFlags[PrevOffset];
-				}
-
-				uint Offset = CalcPageOffset(VirtualShadowMapId, MipLevel, PageAddress);
-				uint PageFlag = VirtualShadowMap.PageFlags[Offset];
-
-				bool bPrevAllocated = (PrevPageFlag & VSM_ALLOCATED_FLAG) != 0;
-				bool bPageAllocated = (PageFlag & VSM_ALLOCATED_FLAG) != 0;
-
-				float3 Color = float3(!bPrevAllocated && bPageAllocated, bPrevAllocated && bPageAllocated, bPrevAllocated && !bPageAllocated);
-				if (bPageAllocated && PrevDynFlag)
-				{
-					Color = float3(1.0f, 0.1f, 0.7f);
-				}
-				DebugOutput[OutputPixel] = float4(Color, 1.0f);
-			}
-			else
-#endif // HAS_CACHE_DATA
-			{
-				uint Offset = CalcPageOffset(VirtualShadowMapId, MipLevel, PageAddress);
-				uint PageFlag = VirtualShadowMap.PageFlags[Offset];
-				bool PageValid = ShadowGetPhysicalPage(Offset).bThisLODValid;
-				float3 Color = float3(0, !PageValid && ((PageFlag & VSM_ALLOCATED_FLAG) != 0), 0);
-				DebugOutput[OutputPixel] = float4(Color, 1.0f);
-			}
-		}
+		Output = DebugVisualizeClipmap(PixelPos, VirtualShadowMapId);
 	}
 	else
 	{
-		DebugOutput[OutputPixel] = float4(float3(0.1f, 0.1f, 0.1f), 1.0f);
+		bWriteOutput = false;
+	}
+
+	if (bWriteOutput)
+	{
+		OutVisualize[PixelPos] = float4(Output, 1.0f);
 	}
 }
 
diff --git a/Engine/Shaders/Private/VirtualShadowMaps/VirtualShadowMapProjection.usf b/Engine/Shaders/Private/VirtualShadowMaps/VirtualShadowMapProjection.usf
--- a/Engine/Shaders/Private/VirtualShadowMaps/VirtualShadowMapProjection.usf
+++ b/Engine/Shaders/Private/VirtualShadowMaps/VirtualShadowMapProjection.usf
@@ -18,7 +18,7 @@
 #include "ProjectionDirectional.ush"
 #include "ProjectionSpot.ush"
 #include "TransmissionCommon.ush"
-#include "Visualize.ush"
+#include "/Engine/Shared/VirtualShadowMapDefinitions.h"
 
 #if HAS_HAIR_STRANDS
 #define VOXEL_TRAVERSAL_TYPE VOXEL_TRAVERSAL_LINEAR_MIPMAP
@@ -358,6 +358,8 @@ FVirtualShadowMapSampleResult ProjectLight(
 
 void OutputVisualize( uint2 PixelPos, FVirtualShadowMapSampleResult Result )
 {
+	// NOTE: Important to *not* write output if it isn't a recognized mode, as a different
+	// pass may write that output instead.
 	float3 Output = float3( 1, 0, 1 );
 	bool bWriteOutput = false;
 
@@ -375,26 +377,26 @@ void OutputVisualize( uint2 PixelPos, FVirtualShadowMapSampleResult Result )
 			pPageMeta = PhysicalPageMetaData[VSMPhysicalPageAddressToIndex(pPage)];
 		}
 
-		if ( VisualizeModeId == VISUALIZE_SHADOW_FACTOR )
+		if ( VisualizeModeId == VIRTUAL_SHADOW_MAP_VISUALIZE_SHADOW_FACTOR )
 		{
 			Output = Result.ShadowFactor.xxx;
 		}
-		else if ( VisualizeModeId == VISUALIZE_CLIPMAP_OR_MIP )
+		else if ( VisualizeModeId == VIRTUAL_SHADOW_MAP_VISUALIZE_CLIPMAP_OR_MIP )
 		{
 			float3 Color = IntToColor( Result.ClipmapIndexOrMipLevel );
 			Output = 0.8f*Color + 0.2f*Result.ShadowFactor.xxx;
 		}
-		else if ( VisualizeModeId == VISUALIZE_VIRTUAL_PAGE && bValidPageAddress )
+		else if ( VisualizeModeId == VIRTUAL_SHADOW_MAP_VISUALIZE_VIRTUAL_PAGE && bValidPageAddress )
 		{
 			float3 PageColor = IntToColor( vPage.x + ( vPage.y << 10U ) );
 			float3 MipColor = IntToColor( Result.ClipmapIndexOrMipLevel );
 			Output = 0.4f*PageColor + 0.4f*MipColor + 0.2f*Result.ShadowFactor.xxx;
 		}		
-		else if ( VisualizeModeId == VISUALIZE_SMRT_RAY_COUNT )
+		else if ( VisualizeModeId == VIRTUAL_SHADOW_MAP_VISUALIZE_SMRT_RAY_COUNT )
 		{
 			Output = GreenToRedTurbo( float( Result.RayCount ) / float( SMRTRayCount ) );
 		}
-		else if ( VisualizeModeId == VISUALIZE_CACHED_PAGE && bValidPageAddress )
+		else if ( VisualizeModeId == VIRTUAL_SHADOW_MAP_VISUALIZE_CACHED_PAGE && bValidPageAddress )
 		{
 			bool bUncachedDynamic = (pPageMeta.Flags & VSM_DYNAMIC_UNCACHED_FLAG) != 0;
 			bool bUncachedStatic = (pPageMeta.Flags & VSM_STATIC_UNCACHED_FLAG) != 0;
@@ -408,7 +410,7 @@ void OutputVisualize( uint2 PixelPos, FVirtualShadowMapSampleResult Result )
 			Output = 0.55f*CacheColor + 0.25f*PageColor + 0.2f*Result.ShadowFactor.xxx;
 		}
 		/*
-		else if ( VisualizeModeId == VISUALIZE_CACHED_PAGE_AGE && bValidPageAddress )
+		else if ( VisualizeModeId == VIRTUAL_SHADOW_MAP_VISUALIZE_CACHED_PAGE_AGE && bValidPageAddress )
 		{
 			float3 Color = GreenToRedTurbo( 1.0f - 0.1f * float(pPageMeta.Age) );
 			Output = lerp( Color, Result.ShadowFactor.xxx, 0.15f );
diff --git a/Engine/Shaders/Private/VirtualShadowMaps/Visualize.ush b/Engine/Shaders/Private/VirtualShadowMaps/Visualize.ush
deleted file mode 100644
diff --git a/Engine/Shaders/Shared/VirtualShadowMapDefinitions.h b/Engine/Shaders/Shared/VirtualShadowMapDefinitions.h
new file mode 100644
--- /dev/null
+++ b/Engine/Shaders/Shared/VirtualShadowMapDefinitions.h
@@ -0,0 +1,16 @@
+// Copyright Epic Games, Inc. All Rights Reserved.
+
+/*================================================================================================
+VirtualShadowMapDefinitions.h: used in virtual shadow map shaders and C++ code to define common constants
+!!! Changing this file requires recompilation of the engine !!!
+=================================================================================================*/
+
+#pragma once
+
+#define VIRTUAL_SHADOW_MAP_VISUALIZE_NONE						0
+#define VIRTUAL_SHADOW_MAP_VISUALIZE_SHADOW_FACTOR				(1 << 0)
+#define VIRTUAL_SHADOW_MAP_VISUALIZE_CLIPMAP_OR_MIP				(1 << 1)
+#define VIRTUAL_SHADOW_MAP_VISUALIZE_VIRTUAL_PAGE				(1 << 2)
+#define VIRTUAL_SHADOW_MAP_VISUALIZE_CACHED_PAGE				(1 << 3)
+#define VIRTUAL_SHADOW_MAP_VISUALIZE_SMRT_RAY_COUNT				(1 << 4)
+#define VIRTUAL_SHADOW_MAP_VISUALIZE_CLIPMAP_VIRTUAL_SPACE		(1 << 5)
diff --git a/Engine/Source/Editor/CommonMenuExtensions/Private/VirtualShadowMapVisualizationMenuCommands.cpp b/Engine/Source/Editor/CommonMenuExtensions/Private/VirtualShadowMapVisualizationMenuCommands.cpp
--- a/Engine/Source/Editor/CommonMenuExtensions/Private/VirtualShadowMapVisualizationMenuCommands.cpp
+++ b/Engine/Source/Editor/CommonMenuExtensions/Private/VirtualShadowMapVisualizationMenuCommands.cpp
@@ -52,17 +52,13 @@ void FVirtualShadowMapVisualizationMenuCommands::BuildCommandMap()
 
 		switch (Entry.ModeType)
 		{
-		default:
-		case FVirtualShadowMapVisualizationData::FModeType::ProjectionStandard:
-			Record.Type = FVirtualShadowMapVisualizationType::ProjectionStandard;
-			break;
-
-		case FVirtualShadowMapVisualizationData::FModeType::ProjectionAdvanced:
-			Record.Type = FVirtualShadowMapVisualizationType::ProjectionAdvanced;
+		case FVirtualShadowMapVisualizationData::FModeType::Standard:
+			Record.Type = FVirtualShadowMapVisualizationType::Standard;
 			break;
 
-		case FVirtualShadowMapVisualizationData::FModeType::DebugStandard:
-			Record.Type = FVirtualShadowMapVisualizationType::DebugStandard;
+		default:
+		case FVirtualShadowMapVisualizationData::FModeType::Advanced:
+			Record.Type = FVirtualShadowMapVisualizationType::Advanced;
 			break;
 		}
 	}
@@ -77,21 +73,17 @@ void FVirtualShadowMapVisualizationMenuCommands::BuildVisualisationSubMenu(FMenu
 	{
 		Menu.BeginSection("LevelViewportVirtualShadowMapVisualizationMode", LOCTEXT("VirtualShadowMapVisualizationHeader", "Visualization Mode"));
 
-		Commands.AddCommandTypeToMenu(Menu, FVirtualShadowMapVisualizationType::ProjectionStandard);
+		Commands.AddCommandTypeToMenu(Menu, FVirtualShadowMapVisualizationType::Standard, false);
 		if (bShowAdvanced)
 		{
-			Menu.AddMenuSeparator();
-			Commands.AddCommandTypeToMenu(Menu, FVirtualShadowMapVisualizationType::ProjectionAdvanced);
+			Commands.AddCommandTypeToMenu(Menu, FVirtualShadowMapVisualizationType::Advanced, true);
 		}
 
-		//Menu.AddMenuSeparator();
-		//Commands.AddCommandTypeToMenu(Menu, FVirtualShadowMapVisualizationType::DebugStandard);
-
 		Menu.EndSection();
 	}
 }
 
-bool FVirtualShadowMapVisualizationMenuCommands::AddCommandTypeToMenu(FMenuBuilder& Menu, const FVirtualShadowMapVisualizationType Type) const
+bool FVirtualShadowMapVisualizationMenuCommands::AddCommandTypeToMenu(FMenuBuilder& Menu, const FVirtualShadowMapVisualizationType Type, bool bSeparatorBefore) const
 {
 	bool bAddedCommands = false;
 
@@ -101,6 +93,10 @@ bool FVirtualShadowMapVisualizationMenuCommands::AddCommandTypeToMenu(FMenuBuild
 		const FVirtualShadowMapVisualizationRecord& Record = It.Value();
 		if (Record.Type == Type)
 		{
+			if (!bAddedCommands && bSeparatorBefore)
+			{
+				Menu.AddMenuSeparator();
+			}
 			Menu.AddMenuEntry(Record.Command, NAME_None, Record.Command->GetLabel());
 			bAddedCommands = true;
 		}
diff --git a/Engine/Source/Editor/CommonMenuExtensions/Public/VirtualShadowMapVisualizationMenuCommands.h b/Engine/Source/Editor/CommonMenuExtensions/Public/VirtualShadowMapVisualizationMenuCommands.h
--- a/Engine/Source/Editor/CommonMenuExtensions/Public/VirtualShadowMapVisualizationMenuCommands.h
+++ b/Engine/Source/Editor/CommonMenuExtensions/Public/VirtualShadowMapVisualizationMenuCommands.h
@@ -17,9 +17,8 @@ public:
 	// Re-sort into categories for the menu
 	enum class FVirtualShadowMapVisualizationType : uint8
 	{
-		ProjectionStandard,
-		ProjectionAdvanced,
-		DebugStandard,
+		Standard,
+		Advanced,
 	};
 
 	struct FVirtualShadowMapVisualizationRecord
@@ -31,7 +30,7 @@ public:
 		FVirtualShadowMapVisualizationRecord()
 		: Name()
 		, Command()
-		, Type(FVirtualShadowMapVisualizationType::ProjectionStandard)
+		, Type(FVirtualShadowMapVisualizationType::Standard)
 		{
 		}
 	};
@@ -56,7 +55,7 @@ public:
 
 private:
 	void BuildCommandMap();
-	bool AddCommandTypeToMenu(FMenuBuilder& Menu, const FVirtualShadowMapVisualizationType Type) const;
+	bool AddCommandTypeToMenu(FMenuBuilder& Menu, const FVirtualShadowMapVisualizationType Type, bool bSeparatorBefore = false) const;
 
 	static void ChangeVirtualShadowMapVisualizationMode(const TSharedPtr<FEditorViewportClient>& Client, FName InName);
 	static bool IsVirtualShadowMapVisualizationModeSelected(const TSharedPtr<FEditorViewportClient>& Client, FName InName);
diff --git a/Engine/Source/Runtime/Engine/Private/VirtualShadowMapVisualizationData.cpp b/Engine/Source/Runtime/Engine/Private/VirtualShadowMapVisualizationData.cpp
--- a/Engine/Source/Runtime/Engine/Private/VirtualShadowMapVisualizationData.cpp
+++ b/Engine/Source/Runtime/Engine/Private/VirtualShadowMapVisualizationData.cpp
@@ -1,6 +1,7 @@
 // Copyright Epic Games, Inc. All Rights Reserved.
 
 #include "VirtualShadowMapVisualizationData.h"
+#include "VirtualShadowMapDefinitions.h"
 #include "HAL/IConsoleManager.h"
 #include "Materials/Material.h"
 #include "Misc/ConfigCacheIni.h"
@@ -9,15 +10,6 @@
 
 static FVirtualShadowMapVisualizationData GVirtualShadowMapVisualizationData;
 
-// Must match values in Shadows/VirtualShadowMaps/Visualize.ush
-#define VISUALIZE_NONE						0
-#define VISUALIZE_SHADOW_FACTOR				(1 << 0)
-#define VISUALIZE_CLIPMAP_OR_MIP			(1 << 1)
-#define VISUALIZE_VIRTUAL_PAGE				(1 << 2)
-#define VISUALIZE_CACHED_PAGE				(1 << 3)
-#define VISUALIZE_SMRT_RAY_COUNT			(1 << 4)
-
-
 void FVirtualShadowMapVisualizationData::Initialize()
 {
 	if (!bIsInitialized)
@@ -27,36 +19,43 @@ void FVirtualShadowMapVisualizationData::Initialize()
 			TEXT("mask"),
 			LOCTEXT("ShadowMask", "Shadow Mask"),
 			LOCTEXT("ShadowMaskDesc", "The final shadow mask that is used by shading"),
-			FModeType::ProjectionStandard,
-			VISUALIZE_SHADOW_FACTOR);
+			FModeType::Standard,
+			VIRTUAL_SHADOW_MAP_VISUALIZE_SHADOW_FACTOR);
 
 		AddVisualizationMode(
 			TEXT("mip"),
 			LOCTEXT("ClipmapOrMip", "Clipmap/Mip Level"),
 			LOCTEXT("ClipmapOrMipDesc", "The chosen clipmap (for directional lights) or mip (for local lights) level"),
-			FModeType::ProjectionStandard,
-			VISUALIZE_CLIPMAP_OR_MIP);
+			FModeType::Standard,
+			VIRTUAL_SHADOW_MAP_VISUALIZE_CLIPMAP_OR_MIP);
 
 		AddVisualizationMode(
 			TEXT("vpage"),
 			LOCTEXT("VirtualPage", "Virtual Page"),
 			LOCTEXT("VirtualPageDesc", "Visualization of the virtual page address"),
-			FModeType::ProjectionStandard,
-			VISUALIZE_VIRTUAL_PAGE);
+			FModeType::Standard,
+			VIRTUAL_SHADOW_MAP_VISUALIZE_VIRTUAL_PAGE);
 
 		AddVisualizationMode(
 			TEXT("cache"),
 			LOCTEXT("CachedPage", "Cached Page"),
 			LOCTEXT("CachedPageDesc", "Cached pages are tinted green, uncached are red. Pages where only the static page is cached (dynamic uncached) are blue."),
-			FModeType::ProjectionStandard,
-			VISUALIZE_CACHED_PAGE);
+			FModeType::Standard,
+			VIRTUAL_SHADOW_MAP_VISUALIZE_CACHED_PAGE);
 
 		AddVisualizationMode(
 			TEXT("raycount"),
 			LOCTEXT("SMRTRayCount", "SMRT Ray Count"),
 			LOCTEXT("SMRTRayCountDesc", "Rays evaluated per pixel: red is more, green is fewer. Penumbra regions require more rays and are more expensive."),
-			FModeType::ProjectionAdvanced,
-			VISUALIZE_SMRT_RAY_COUNT);
+			FModeType::Advanced,
+			VIRTUAL_SHADOW_MAP_VISUALIZE_SMRT_RAY_COUNT);
+
+		AddVisualizationMode(
+			TEXT("clipmapvirtual"),
+			LOCTEXT("ClipmapVirtualSpace", "Clipmap Virtual Address Space"),
+			LOCTEXT("ClipmapVirtualSpaceDesc", "Visualization of the clipmap virtual address space and mapped pages"),
+			FModeType::Advanced,
+			VIRTUAL_SHADOW_MAP_VISUALIZE_CLIPMAP_VIRTUAL_SPACE);
 
 		ConfigureConsoleCommand();
 
diff --git a/Engine/Source/Runtime/Engine/Public/VirtualShadowMapVisualizationData.h b/Engine/Source/Runtime/Engine/Public/VirtualShadowMapVisualizationData.h
--- a/Engine/Source/Runtime/Engine/Public/VirtualShadowMapVisualizationData.h
+++ b/Engine/Source/Runtime/Engine/Public/VirtualShadowMapVisualizationData.h
@@ -9,9 +9,8 @@ class FVirtualShadowMapVisualizationData
 public:
 	enum class FModeType : uint8
 	{
-		ProjectionStandard,
-		ProjectionAdvanced,
-		DebugStandard,
+		Standard,
+		Advanced,
 	};
 
 	/** Describes a single available visualization mode. */
@@ -26,7 +25,7 @@ public:
 	};
 
 	/** Mapping of FName to a visualization mode record. */
-	typedef TMultiMap<FName, FModeRecord> TModeMap;
+	typedef TMap<FName, FModeRecord> TModeMap;
 
 public:
 	FVirtualShadowMapVisualizationData()
diff --git a/Engine/Source/Runtime/Renderer/Private/VirtualShadowMaps/VirtualShadowMapArray.cpp b/Engine/Source/Runtime/Renderer/Private/VirtualShadowMaps/VirtualShadowMapArray.cpp
--- a/Engine/Source/Runtime/Renderer/Private/VirtualShadowMaps/VirtualShadowMapArray.cpp
+++ b/Engine/Source/Runtime/Renderer/Private/VirtualShadowMaps/VirtualShadowMapArray.cpp
@@ -64,14 +64,6 @@ TAutoConsoleVariable<int32> CVarCacheStaticSeparate(
 	ECVF_RenderThreadSafe
 );
 
-static TAutoConsoleVariable<int32> CVarDebugVisualizeVirtualSms(
-	TEXT("r.Shadow.Virtual.DebugVisualize"),
-	0,
-	TEXT("Set Debug Visualization method for virtual shadow maps, default is off (0).\n")
-	TEXT("  To display the result also use the command 'vis VirtSmDebug'"),
-	ECVF_RenderThreadSafe
-);
-
 static TAutoConsoleVariable<int32> CVarShowStats(
 	TEXT("r.Shadow.Virtual.ShowStats"),
 	0,
@@ -234,10 +226,9 @@ TAutoConsoleVariable<int32> CVarDebugSkipMergePhysical(
 );
 
 TAutoConsoleVariable<int32> CVarDebugSkipDynamicPageInvalidation(
-	TEXT("r.Shadow.Virtual.DebugSkipDynamicPageInvalidation"),
+	TEXT("r.Shadow.Virtual.Cache.DebugSkipDynamicPageInvalidation"),
 	0,
-	TEXT("Invalidate cached pages when geometry moves.\n")
-	TEXT("This should be left enabled except for targeted profiling, as disabling it will produce artifacts with moving geometry."),
+	TEXT("Skip invalidation of cached pages when geometry moves for debugging purposes. This will create obvious visual artifacts when disabled.\n"),
 	ECVF_RenderThreadSafe
 );
 #endif // !UE_BUILD_SHIPPING
@@ -258,7 +249,7 @@ TAutoConsoleVariable<int32> CVarVSMHZB32Bit(
 
 
 static TAutoConsoleVariable<float> CVarMaxMaterialPositionInvalidationRange(
-	TEXT("r.Shadow.Virtual.MaxMaterialPositionInvalidationRange"),
+	TEXT("r.Shadow.Virtual.Cache.MaxMaterialPositionInvalidationRange"),
 	-1.0f,
 	TEXT("Beyond this distance in world units, material position effects (e.g., WPO or PDO) cease to cause VSM invalidations.\n")
 	TEXT(" This can be used to tune performance by reducing re-draw overhead, but causes some artifacts.\n")
@@ -984,9 +975,9 @@ void FVirtualShadowMapVisualizeLightSearch::CheckLight(const FLightSceneProxy* C
 #endif
 }
 
-static FRDGTextureRef CreateDebugOutputTexture(FRDGBuilder& GraphBuilder, FIntPoint Extent)
+static FRDGTextureRef CreateDebugVisualizationTexture(FRDGBuilder& GraphBuilder, FIntPoint Extent)
 {
-	const FLinearColor ClearColor(1.0f, 0.0f, 1.0f, 0.0f);
+	const FLinearColor ClearColor(0.0f, 0.0f, 0.0f, 0.0f);
 
 	FRDGTextureDesc Desc = FRDGTextureDesc::Create2D(
 		Extent,
@@ -1045,7 +1036,7 @@ void FVirtualShadowMapArray::BuildPageAllocations(
 		if (VisualizationData.IsActive() && EngineShowFlags.VisualizeVirtualShadowMap)
 		{
 			bDebugOutputEnabled = true;
-			DebugVisualizationOutput = CreateDebugOutputTexture(GraphBuilder, Extent);
+			DebugVisualizationOutput = CreateDebugVisualizationTexture(GraphBuilder, Extent);
 		}
 	}
 #endif //!UE_BUILD_SHIPPING
@@ -1535,24 +1526,15 @@ class FDebugVisualizeVirtualSmCS : public FVirtualPageManagementShader
 	DECLARE_GLOBAL_SHADER(FDebugVisualizeVirtualSmCS);
 	SHADER_USE_PARAMETER_STRUCT(FDebugVisualizeVirtualSmCS, FVirtualPageManagementShader)
 
-	class FHasCacheDataDim : SHADER_PERMUTATION_BOOL("HAS_CACHE_DATA");
-	using FPermutationDomain = TShaderPermutationDomain<FHasCacheDataDim>;
-
 	BEGIN_SHADER_PARAMETER_STRUCT(FParameters, )
 		SHADER_PARAMETER_STRUCT_INCLUDE(FVirtualShadowMapSamplingParameters, ProjectionParameters)
 		SHADER_PARAMETER(uint32, DebugTargetWidth)
 		SHADER_PARAMETER(uint32, DebugTargetHeight)
 		SHADER_PARAMETER(uint32, BorderWidth)
-		SHADER_PARAMETER(uint32, ZoomScaleFactor)
-		SHADER_PARAMETER(uint32, DebugMethod)
-		SHADER_PARAMETER_RDG_BUFFER_SRV(StructuredBuffer< uint >, PageFlags)
-		
-		SHADER_PARAMETER_RDG_TEXTURE(Texture2D< float >,			HZBPhysical )
-		SHADER_PARAMETER_RDG_BUFFER_SRV(StructuredBuffer< uint >,	HZBPageTable )
-
-		SHADER_PARAMETER_STRUCT_INCLUDE(FCacheDataParameters, CacheDataParameters)
+		SHADER_PARAMETER(uint32, VisualizeModeId)
+		SHADER_PARAMETER(int32, VirtualShadowMapId)
 
-		SHADER_PARAMETER_RDG_TEXTURE_UAV(RWTexture2D<float4>, DebugOutput)
+		SHADER_PARAMETER_RDG_TEXTURE_UAV(RWTexture2D<float4>, OutVisualize)
 	END_SHADER_PARAMETER_STRUCT()
 };
 IMPLEMENT_GLOBAL_SHADER(FDebugVisualizeVirtualSmCS, "/Engine/Private/VirtualShadowMaps/Debug.usf", "DebugVisualizeVirtualSmCS", SF_Compute);
@@ -1561,71 +1543,38 @@ IMPLEMENT_GLOBAL_SHADER(FDebugVisualizeVirtualSmCS, "/Engine/Private/VirtualShad
 void FVirtualShadowMapArray::RenderDebugInfo(FRDGBuilder& GraphBuilder)
 {
 	check(IsEnabled());
-	int32 DebugMethod = CVarDebugVisualizeVirtualSms.GetValueOnRenderThread();
-	if (ShadowMaps.Num() && DebugMethod > 0)
+			
+	if (DebugVisualizationOutput == nullptr || !VisualizeLight.IsValid())
 	{
-		int32 ZoomScaleFactor = 1;
-		int32 BorderWidth = 2;
-		// Make debug target wide enough to show a mip-chain
-		int32 DebugTargetWidth = ZoomScaleFactor * (FVirtualShadowMap::Level0DimPagesXY * 2 + BorderWidth * (FVirtualShadowMap::MaxMipLevels));
-		// Enough rows for all the shadow maps to show
-		int32 DebugTargetHeight = ZoomScaleFactor * (FVirtualShadowMap::Level0DimPagesXY + BorderWidth * 2) * ShadowMaps.Num();
-
-		if( DebugMethod > 5 )
-		{
-			DebugTargetWidth = 2048;
-			DebugTargetHeight = 2048;
-		}
+		return;
+	}
 
-		FRDGTextureDesc DebugOutputDesc = FRDGTextureDesc::Create2D(
-			FIntPoint(DebugTargetWidth, DebugTargetHeight),
-			PF_A32B32G32R32F,
-			FClearValueBinding::None,
-			TexCreate_ShaderResource | TexCreate_UAV);
-
-		FRDGTextureRef DebugOutput = GraphBuilder.CreateTexture(
-			DebugOutputDesc,
-			TEXT("Shadow.Virtual.Debug"));
-
-		FDebugVisualizeVirtualSmCS::FParameters* PassParameters = GraphBuilder.AllocParameters<FDebugVisualizeVirtualSmCS::FParameters>();
-		PassParameters->ProjectionParameters = GetSamplingParameters(GraphBuilder);
-		PassParameters->PageFlags = GraphBuilder.CreateSRV( PageFlagsRDG );
-
-		// TODO: unclear if it's preferable to debug the HZB we generated "this frame" here, or the previous frame (that was used for culling)?
-		// We'll stick with the previous frame logic that was there, but it's cleaner to just reference the current frame one
-		TRefCountPtr<IPooledRenderTarget> PrevHZBPhysical = CacheManager->PrevBuffers.HZBPhysical;
-		TRefCountPtr<FRDGPooledBuffer>    PrevPageTable = CacheManager->PrevBuffers.PageTable;
-		PassParameters->HZBPhysical	 = RegisterExternalTextureWithFallback( GraphBuilder, PrevHZBPhysical, GSystemTextures.BlackDummy );
-		PassParameters->HZBPageTable = GraphBuilder.CreateSRV( PrevPageTable ? GraphBuilder.RegisterExternalBuffer( PrevPageTable ) : PageTableRDG );
-
-		PassParameters->DebugTargetWidth = DebugTargetWidth;
-		PassParameters->DebugTargetHeight = DebugTargetHeight;
-		PassParameters->BorderWidth = BorderWidth;
-		PassParameters->ZoomScaleFactor = ZoomScaleFactor;
-		PassParameters->DebugMethod = DebugMethod;
-
-		bool bCacheDataAvailable = CacheManager->IsValid();
-		if (bCacheDataAvailable)
-		{
-			SetCacheDataShaderParameters(GraphBuilder, ShadowMaps, CacheManager, PassParameters->CacheDataParameters);
-		}
-		PassParameters->DebugOutput = GraphBuilder.CreateUAV(DebugOutput);
+	const FVirtualShadowMapVisualizationData& VisualizationData = GetVirtualShadowMapVisualizationData();
 
-		FDebugVisualizeVirtualSmCS::FPermutationDomain PermutationVector;
-		PermutationVector.Set<FDebugVisualizeVirtualSmCS::FHasCacheDataDim>(bCacheDataAvailable);
-		auto ComputeShader = GetGlobalShaderMap(GMaxRHIFeatureLevel)->GetShader<FDebugVisualizeVirtualSmCS>(PermutationVector);
+	int32 BorderWidth = 2;
 
-		FComputeShaderUtils::AddPass(
-			GraphBuilder,
-			RDG_EVENT_NAME("DebugVisualizeVirtualSmCS"),
-			ComputeShader,
-			PassParameters,
-			FComputeShaderUtils::GetGroupCount(FIntPoint(DebugTargetWidth, DebugTargetHeight), FVirtualPageManagementShader::DefaultCSGroupXY)
-		);
+	FIntPoint DebugTargetExtent = DebugVisualizationOutput->Desc.Extent;
 
-		// TODO!
-		// DebugVisualizationOutput = GraphBuilder.ConvertToExternalTexture(DebugOutput);		
-	}
+	FDebugVisualizeVirtualSmCS::FParameters* PassParameters = GraphBuilder.AllocParameters<FDebugVisualizeVirtualSmCS::FParameters>();
+	PassParameters->ProjectionParameters = GetSamplingParameters(GraphBuilder);
+	
+	PassParameters->DebugTargetWidth = DebugTargetExtent.X;
+	PassParameters->DebugTargetHeight = DebugTargetExtent.Y;
+	PassParameters->BorderWidth = BorderWidth;
+	PassParameters->VisualizeModeId = VisualizationData.GetActiveModeID();
+	PassParameters->VirtualShadowMapId = VisualizeLight.GetVirtualShadowMapId();
+
+	PassParameters->OutVisualize = GraphBuilder.CreateUAV(DebugVisualizationOutput);
+
+	auto ComputeShader = GetGlobalShaderMap(GMaxRHIFeatureLevel)->GetShader<FDebugVisualizeVirtualSmCS>();
+
+	FComputeShaderUtils::AddPass(
+		GraphBuilder,
+		RDG_EVENT_NAME("DebugVisualizeVirtualShadowMap"),
+		ComputeShader,
+		PassParameters,
+		FComputeShaderUtils::GetGroupCount(DebugTargetExtent, FVirtualPageManagementShader::DefaultCSGroupXY)
+	);
 }
 
 
diff --git a/Engine/Source/Runtime/Renderer/Private/VirtualShadowMaps/VirtualShadowMapArray.h b/Engine/Source/Runtime/Renderer/Private/VirtualShadowMaps/VirtualShadowMapArray.h
--- a/Engine/Source/Runtime/Renderer/Private/VirtualShadowMaps/VirtualShadowMapArray.h
+++ b/Engine/Source/Runtime/Renderer/Private/VirtualShadowMaps/VirtualShadowMapArray.h
@@ -249,9 +249,8 @@ public:
 	 */
 	void RenderVirtualShadowMapsNonNanite(FRDGBuilder& GraphBuilder, const TArray<FProjectedShadowInfo*, SceneRenderingAllocator>& VirtualSmMeshCommandPasses, FScene& Scene, TArrayView<FViewInfo> Views);
 
-	// Draw debug info into render target 'VSMDebug' of screen-size, the mode is controlled by 'r.Shadow.Virtual.DebugVisualize'.
 	void RenderDebugInfo(FRDGBuilder& GraphBuilder);
-	// 
+	
 	void PrintStats(FRDGBuilder& GraphBuilder, const FViewInfo& View);
 
 	TRDGUniformBufferRef<FVirtualShadowMapUniformParameters> GetUniformBuffer(FRDGBuilder& GraphBuilder) const;
diff --git a/Engine/Source/Runtime/Renderer/Private/VirtualShadowMaps/VirtualShadowMapProjection.cpp b/Engine/Source/Runtime/Renderer/Private/VirtualShadowMaps/VirtualShadowMapProjection.cpp
--- a/Engine/Source/Runtime/Renderer/Private/VirtualShadowMaps/VirtualShadowMapProjection.cpp
+++ b/Engine/Source/Runtime/Renderer/Private/VirtualShadowMaps/VirtualShadowMapProjection.cpp
@@ -331,7 +331,7 @@ static void RenderVirtualShadowMapProjectionCommon(
 	
 	bool bDebugOutput = false;
 #if !UE_BUILD_SHIPPING
-	if ( VirtualShadowMapArray.DebugVisualizationOutput && InputType == EVirtualShadowMapProjectionInputType::GBuffer )
+	if ( VirtualShadowMapArray.DebugVisualizationOutput && InputType == EVirtualShadowMapProjectionInputType::GBuffer && VirtualShadowMapArray.VisualizeLight.IsValid() )
 	{
 		const FVirtualShadowMapVisualizationData& VisualizationData = GetVirtualShadowMapVisualizationData();
 
