Make sure fog is used for water shading on mobile platforms

diff --git a/Engine/Shaders/Private/MobileBasePassCommon.ush b/Engine/Shaders/Private/MobileBasePassCommon.ush
--- a/Engine/Shaders/Private/MobileBasePassCommon.ush
+++ b/Engine/Shaders/Private/MobileBasePassCommon.ush
@@ -18,7 +18,7 @@
 
 // Use vertex fogging for translucency and opaque if project does not use per-pixel fog
 // And always use vertex fogging on sky for perf
-#define USE_VERTEX_FOG	(!PROJECT_MOBILE_DISABLE_VERTEX_FOG || (MATERIAL_ENABLE_TRANSLUCENCY_FOGGING && MATERIALBLENDING_ANY_TRANSLUCENT) || MATERIAL_IS_SKY)
+#define USE_VERTEX_FOG	(!PROJECT_MOBILE_DISABLE_VERTEX_FOG || MATERIAL_IS_SKY || (MATERIAL_ENABLE_TRANSLUCENCY_FOGGING && (MATERIALBLENDING_ANY_TRANSLUCENT || MATERIAL_SHADINGMODEL_SINGLELAYERWATER)))
 
 #define PACK_INTERPOLANTS (USE_VERTEX_FOG && NUM_VF_PACKED_INTERPOLANTS > 0 && (ES3_1_PROFILE))
 #define LANDSCAPE_BUG_WORKAROUND (IOS && IS_MOBILE_BASEPASS_VERTEX_SHADER && PACK_INTERPOLANTS)
