FPixelShaderUtils::FRasterizeToRectsVS compile fix

#rnx

diff --git a/Engine/Shaders/Private/RenderGraphUtilities.usf b/Engine/Shaders/Private/RenderGraphUtilities.usf
--- a/Engine/Shaders/Private/RenderGraphUtilities.usf
+++ b/Engine/Shaders/Private/RenderGraphUtilities.usf
@@ -42,6 +42,7 @@ void RasterizeToRectsVS(
 
 	OutPosition = float4(float2(VertexCoord) * InvViewSize * float2(2.0f, -2.0f) + float2(-1.0, 1.0f), 0.0f, 1.0f);
 	OutUV = VertexUV;
-	OutRectUV = VertMax(VertexId) ? 1.0f : 0.0f;
+	OutRectUV.x = VertMax(VertexId).x ? 1.0f : 0.0f;
+	OutRectUV.y = VertMax(VertexId).y ? 1.0f : 0.0f;
 	OutRectIndex = InstanceId;
 }
\ No newline at end of file
