Add DBuffer support for Strata.

The application of the decals onto Strata surface is currently disabled in shader code, as further optimization are needed for reduce occupancy reduction on simple shaders.

diff --git a/Engine/Shaders/Private/BasePassPixelShader.usf b/Engine/Shaders/Private/BasePassPixelShader.usf
--- a/Engine/Shaders/Private/BasePassPixelShader.usf
+++ b/Engine/Shaders/Private/BasePassPixelShader.usf
@@ -38,6 +38,12 @@
 	#endif
 #endif
 
+// Enable Strata. This define & include need to be defined before certains includes (i.e., DBufferDecalShared which uses them internally)
+#define STRATA_ENABLED (PROJECT_STRATA && MATERIAL_IS_STRATA)
+#if STRATA_ENABLED
+#include "/Engine/Private/Strata/Strata.ush"
+#endif
+
 #include "SHCommon.ush"
 #include "/Engine/Generated/Material.ush"
 #include "BasePassCommon.ush"
@@ -86,7 +92,6 @@
 #include "Lumen/LumenTranslucencyVolumeShared.ush"
 #endif
 
-#define STRATA_ENABLED					(PROJECT_STRATA && MATERIAL_IS_STRATA)
 
 // We detect the isolated unlit BSDF node case to make sure lots of code is optimised out, when used, there can only be one.
 // When unlit is selected as part of dynamic shading model, the SLAB model is used instead.
@@ -157,7 +162,6 @@ float NormalCurvatureToRoughness(float3 WorldNormal)
 #endif
 
 #if STRATA_TRANSLUCENT_FORWARD || STRATA_FORWARD_SHADING
-#define STRATA_USE_ENV_BRDF_APPROX (MATERIAL_USE_PREINTEGRATED_GF == 0)
 #include "/Engine/Private/Strata/StrataEvaluation.ush"
 #include "/Engine/Private/Strata/StrataForwardLighting.ush"
 #endif
@@ -1085,10 +1089,17 @@ void FPixelShaderInOut_MainPS(
 			if (ValidDBufferTargetMask)
 			{
 				float2 BufferUV = SvPositionToBufferUV(In.SvPosition);
+			#if STRATA_ENABLED 
+			   #if STRATA_INLINE_SHADING && !STRATA_OPTIMIZED_UNLIT
+				const FStrataDBuffer StrataBufferData = GetDBufferData(BufferUV, ValidDBufferTargetMask);
+				ApplyDBufferData(StrataBufferData, StrataPixelHeader, StrataData);
+				DBufferOpacity = StrataBufferData.Coverage;
+			   #endif
+			#else
 				FDBufferData DBufferData = GetDBufferData(BufferUV, ValidDBufferTargetMask);
-
 				ApplyDBufferData(DBufferData, MaterialParameters.WorldNormal, SubsurfaceColor, Roughness, BaseColor, Metallic, Specular);
 				DBufferOpacity = (DBufferData.ColorOpacity + DBufferData.NormalOpacity + DBufferData.RoughnessOpacity) * (1.0f / 3.0f);
+			#endif // STRATA_ENABLED
 			}
 		}
 #endif
diff --git a/Engine/Shaders/Private/DBufferDecalShared.ush b/Engine/Shaders/Private/DBufferDecalShared.ush
--- a/Engine/Shaders/Private/DBufferDecalShared.ush
+++ b/Engine/Shaders/Private/DBufferDecalShared.ush
@@ -6,7 +6,249 @@
 
 #pragma once
 
+
+#if SUPPORTS_INDEPENDENT_SAMPLERS
+// Share samplers to give artists more
+#define OpaqueBasePassDBufferATextureSampler OpaqueBasePass.DBufferATextureSampler
+#define OpaqueBasePassDBufferBTextureSampler OpaqueBasePass.DBufferATextureSampler
+#define OpaqueBasePassDBufferCTextureSampler OpaqueBasePass.DBufferATextureSampler
+#else
+#define OpaqueBasePassDBufferATextureSampler OpaqueBasePass.DBufferATextureSampler
+#define OpaqueBasePassDBufferBTextureSampler OpaqueBasePass.DBufferBTextureSampler
+#define OpaqueBasePassDBufferCTextureSampler OpaqueBasePass.DBufferCTextureSampler
+#endif
+
+// Returns a bitmask of DBuffer targets that contain valid data for this pixel.
+// @param BufferUV - UV space in the DBuffer textures
+uint GetDBufferTargetMask(uint2 PixelPos)
+{
+#if USE_DBUFFER && ENABLE_DBUFFER_TEXTURES && !RAYHITGROUPSHADER && !LUMEN_CARD_CAPTURE
+
+#if PLATFORM_SUPPORTS_RENDERTARGET_WRITE_MASK
+	return DecodeRTWriteMask(PixelPos, OpaqueBasePass.DBufferRenderMask, 3);
+#elif PLATFORM_SUPPORTS_PER_PIXEL_DBUFFER_MASK
+	uint Mask = OpaqueBasePass.DBufferRenderMask.Load(uint3(PixelPos, 0));
+	return Mask > 0 ? 0x07 : 0x00;
+#else
+	return 0x07;
+#endif
+
+#else
+	return 0;
+#endif
+}
+
 // all values that are output by the forward rendering pass
+#if STRATA_ENABLED
+
+void ConvertFromMetalness(
+	in float3 BaseColor,
+	in float Specular,
+	in float Metallic,
+	inout float3 DiffuseAlbedo,
+	inout float3 F0)
+{
+	DiffuseAlbedo = lerp(BaseColor, 0, Metallic);
+	F0			  = lerp(DielectricSpecularToF0(Specular), BaseColor, Metallic);
+}
+
+void ConvertToMetalness(
+	in float3 DiffuseAlbedo,
+	in float3 F0,
+	inout float3 BaseColor,
+	inout float Specular,
+	inout float Metallic)
+{
+	Metallic  = F0RGBToMetallic(F0);
+	Specular  = F0ToDielectricSpecular(F0.x);
+	BaseColor = lerp(DiffuseAlbedo, F0, Metallic);
+}
+
+// @param BufferUV - UV space in the DBuffer textures
+FStrataDBuffer GetDBufferData(float2 BufferUV, uint RTMaskBit)
+{
+	// Setup default values, which mean that no decals are present in DBuffer
+	float4 DBufferA = float4(0.0, 0.0, 0.0, 1.0);
+	float4 DBufferB = float4(128.0f / 255.0f, 128.f / 255.5f, 128.f / 255.5f, 1.0);
+	float4 DBufferC = float4(0.0, 0.0, 0.0, 1.0);
+
+#if USE_DBUFFER && ENABLE_DBUFFER_TEXTURES && !RAYHITGROUPSHADER && !LUMEN_CARD_CAPTURE
+	BRANCH
+	if (RTMaskBit & 0x1) { DBufferA = Texture2DSampleLevel(OpaqueBasePass.DBufferATexture, OpaqueBasePassDBufferATextureSampler, BufferUV, 0);	}
+	BRANCH
+	if (RTMaskBit & 0x2) { DBufferB = Texture2DSampleLevel(OpaqueBasePass.DBufferBTexture, OpaqueBasePassDBufferBTextureSampler, BufferUV, 0);	}
+	BRANCH
+	if (RTMaskBit & 0x4) { DBufferC = Texture2DSampleLevel(OpaqueBasePass.DBufferCTexture, OpaqueBasePassDBufferCTextureSampler, BufferUV, 0);	}
+#endif
+	return StrataUnpackDBuffer(DBufferA, DBufferB, DBufferC);
+}
+
+
+// Define how DBuffer is applied onto existing layers:
+// * 0: Apply DBuffer as per-component parameter blending (legacy mode). Default. 
+// * 1: Apply DBuffer as a coating layer (using parameter blending for efficiency purpose). Not used for now
+#define STRATA_DBUFFER_LEGACY 0
+#define STRATA_DBUFFER_COATING 1
+#define STRATA_DBUFFER_BLENDING_MODE STRATA_DBUFFER_LEGACY
+
+// Define which DBuffers are valid to read from, for this material
+#define STRATA_DBUFFER_RESPONSE_BASECOLOR (MATERIALDECALRESPONSEMASK & 0x1)
+#define STRATA_DBUFFER_RESPONSE_NORMAL    (MATERIALDECALRESPONSEMASK & 0x2)
+#define STRATA_DBUFFER_RESPONSE_ROUGHNESS (MATERIALDECALRESPONSEMASK & 0x4)
+
+void ApplyDBufferData(in FStrataDBuffer StrataDBufferData, inout FStrataPixelHeader StrataPixelHeader, inout FStrataData StrataData)
+{
+  // Strata - Disabling temporary DBuffer prior to optimization
+  #if 0
+
+	// Early out if DBuffer has no coverage
+	if (StrataDBufferData.Coverage <= 0)
+		return;
+
+	// For normal, apply decal data onto all basis, without consideration of top layers
+	#if STRATA_DBUFFER_RESPONSE_NORMAL && STRATA_INLINE_SHADING
+	{
+		UNROLL
+		for (uint i = 0; i < StrataPixelHeader.SharedLocalBases.Count; ++i)
+		{
+			// Note:
+			// * Assume Normal is already been normalized
+			// * After modification we normalize the normal to get smoother visual result (it helps to avoid having D_GGX explodes toward infinity)
+			// 
+			// STRATA_TODO: 
+			// * If the normal is shared by different layers on which decals would be applied differently (e.g., one affected, not the other), we 
+			//   would need to add a new normal entry and modify it. 
+			// * If the basis has a tangent, the average operation should be changed into a basis sler/rotation
+			StrataPixelHeader.SharedLocalBases.Normals[i] = normalize(StrataPixelHeader.SharedLocalBases.Normals[i] * StrataDBufferData.OneMinusCoverage_WorldNormal + StrataDBufferData.WorldNormal);
+			//if (StrataGetSharedLocalBasisType(StrataPixelHeader.SharedLocalBases.Types, i) == STRATA_BASIS_TYPE_TANGENT)
+			//{
+			//	StrataPixelHeader.SharedLocalBases.Tangents[i] = normalize(StrataPixelHeader.SharedLocalBases.Tangents[i]);
+			//}
+		}
+	}
+	#endif
+
+	// For now apply decal data onto all layers, without consideration of top layers
+	STRATA_UNROLL_N(STRATA_CLAMPED_LAYER_COUNT)
+	for (int l = 0; l < StrataData.LayerCount; ++l)
+	{			
+		STRATA_UNROLL_N(STRATA_CLAMPED_MAX_BSDF_COUNT_PER_LAYER)
+		for (int i = 0; i < StrataData.Layers[l].BSDFCount; ++i)
+		{
+
+			switch (BSDF_GETTYPE(StrataData.Layers[l].BSDFs[i]))
+			{
+				case STRATA_BSDF_TYPE_SLAB:
+				{	
+					// Note:
+					// * For Slab, since the DBuffer is stored with a Metalness parameterization for storage-space reason, 
+					//   we convert back&forth the Slab data between DiffuseAlbedo/F0 and Metalness. This induces some shift 
+					//   since the conversion is lossy. In the future, we could change the DBuffer storage layout (would require 
+					//   more space) or do all the blending in DiffuseAlbedo/F0 space
+					
+					// Apply DBuffer as per-component parameter blending (legacy mode)
+					#if STRATA_DBUFFER_BLENDING_MODE == STRATA_DBUFFER_LEGACY
+					{
+						float3 BaseColor = 0;
+						float Metallic = 0;
+						float Specular = 0;
+						ConvertToMetalness(SLAB_DIFFUSEALBEDO(StrataData.Layers[l].BSDFs[i]), SLAB_F0(StrataData.Layers[l].BSDFs[i]), BaseColor, Specular, Metallic);
+						float Roughness = SLAB_ROUGHNESS(StrataData.Layers[l].BSDFs[i]);
+
+					#if STRATA_DBUFFER_RESPONSE_BASECOLOR
+						BaseColor = BaseColor * StrataDBufferData.OneMinusCoverage_BaseColor + StrataDBufferData.BaseColor;
+						if (BSDF_GETHASFUZZ(StrataData.Layers[l].BSDFs[i]))
+						{
+							SLAB_FUZZ_COLOR(StrataData.Layers[l].BSDFs[i]) = SLAB_FUZZ_COLOR(StrataData.Layers[l].BSDFs[i]) * StrataDBufferData.OneMinusCoverage_BaseColor + StrataDBufferData.BaseColor;
+						}
+
+						// Disable SSS when the coverage reach one to preserve the decal color
+						SLAB_SSSDMFP(StrataData.Layers[l].BSDFs[i]) = lerp(0.f, SLAB_SSSDMFP(StrataData.Layers[l].BSDFs[i]), StrataDBufferData.OneMinusCoverage_BaseColor);
+						if (all(SLAB_SSSDMFP(StrataData.Layers[l].BSDFs[i]) == 0.0f))
+						{
+							BSDF_SETHASDMFP(StrataData.Layers[l].BSDFs[i], 0);
+							BSDF_SETHASSSS(StrataData.Layers[l].BSDFs[i], 0);
+						}
+					#endif
+
+					#if STRATA_DBUFFER_RESPONSE_ROUGHNESS
+						Metallic  = Metallic  * StrataDBufferData.OneMinusCoverage_Roughness + StrataDBufferData.Metallic;
+						Specular  = Specular  * StrataDBufferData.OneMinusCoverage_Roughness + StrataDBufferData.Specular;
+						Roughness = Roughness * StrataDBufferData.OneMinusCoverage_Roughness + StrataDBufferData.Roughness;
+					#endif
+
+						float3 DiffuseAlbedo = 0;
+						float3 F0 = 0;
+						ConvertFromMetalness(BaseColor, Specular, Metallic, SLAB_DIFFUSEALBEDO(StrataData.Layers[l].BSDFs[i]), SLAB_F0(StrataData.Layers[l].BSDFs[i]));
+						SLAB_ROUGHNESS(StrataData.Layers[l].BSDFs[i]) = Roughness;
+					}
+					// Apply DBuffer as a coating layer (using parameter blending for efficiency purpose)
+					#elif STRATA_DBUFFER_BLENDING_MODE == STRATA_DBUFFER_COATING
+					{
+						// Since all DBuffer parameters are not necessary initialized, take default value from the existing BDSF
+						FStrataDBuffer LocalStrataDBufferData = StrataDBufferData;
+						{
+							float3 SlabBaseColor = 0;
+							float  SlabMetallic = 0;
+							float  SlabSpecular = 0;
+							ConvertToMetalness(SLAB_DIFFUSEALBEDO(StrataData.Layers[l].BSDFs[i]), SLAB_F0(StrataData.Layers[l].BSDFs[i]), SlabBaseColor, SlabSpecular, SlabMetallic);
+						#if !STRATA_DBUFFER_RESPONSE_BASECOLOR
+							LocalStrataDBufferData.BaseColor					= LocalStrataDBufferData.Coverage * SlabBaseColor;
+							LocalStrataDBufferData.OneMinusCoverage_BaseColor	= LocalStrataDBufferData.OneMinusCoverage;
+						#endif
+						#if !STRATA_DBUFFER_RESPONSE_ROUGHNESS
+							LocalStrataDBufferData.Metallic						= LocalStrataDBufferData.Coverage * SlabMetallic;
+							LocalStrataDBufferData.Specular						= LocalStrataDBufferData.Coverage * SlabSpecular;
+							LocalStrataDBufferData.Roughness					= LocalStrataDBufferData.Coverage * SLAB_ROUGHNESS(StrataData.Layers[l].BSDFs[i]);
+							LocalStrataDBufferData.OneMinusCoverage_Roughness	= LocalStrataDBufferData.OneMinusCoverage;
+						#endif
+						}
+
+						FStrataPixelHeader StrataDecalHeader;
+						FStrataData StrataDecalData;
+						StrataConvertFromDBuffer(LocalStrataDBufferData, StrataDecalHeader, StrataDecalData);
+
+						// Parameters blend attributes
+						const float Coverage = StrataDecalData.Layers[l].BSDFs[i].Coverage;
+						const float DummyNoV = 1;
+						StrataData.Layers[l].BSDFs[i] = StrataVerticalLayeringParameterBlending(StrataDecalData.Layers[l].BSDFs[i], StrataData.Layers[l].BSDFs[i], BasisIndex, DummyNoV);
+					}
+					#endif // STRATA_DBUFFER_BLENDING_MODE
+					break;
+				}
+
+				case STRATA_BSDF_TYPE_HAIR:
+				{
+				#if STRATA_DBUFFER_RESPONSE_BASECOLOR
+					HAIR_BASECOLOR(StrataData.Layers[l].BSDFs[i]) = HAIR_BASECOLOR(StrataData.Layers[l].BSDFs[i]) * StrataDBufferData.OneMinusCoverage_BaseColor + StrataDBufferData.BaseColor;
+				#endif
+
+				#if STRATA_DBUFFER_RESPONSE_ROUGHNESS
+					HAIR_ROUGHNESS(StrataData.Layers[l].BSDFs[i]) = HAIR_ROUGHNESS(StrataData.Layers[l].BSDFs[i]) * StrataDBufferData.OneMinusCoverage_Roughness + StrataDBufferData.Roughness;
+					HAIR_SPECULAR(StrataData.Layers[l].BSDFs[i])  = HAIR_SPECULAR(StrataData.Layers[l].BSDFs[i])  * StrataDBufferData.OneMinusCoverage_Roughness + StrataDBufferData.Specular;
+				#endif
+					break;
+				}
+
+				case STRATA_BSDF_TYPE_SINGLELAYERWATER:
+				{
+				#if STRATA_DBUFFER_RESPONSE_BASECOLOR
+					SLW_BASECOLOR(StrataData.Layers[l].BSDFs[i]) = SLW_BASECOLOR(StrataData.Layers[l].BSDFs[i]) * StrataDBufferData.OneMinusCoverage_BaseColor + StrataDBufferData.BaseColor;
+				#endif
+
+				#if STRATA_DBUFFER_RESPONSE_ROUGHNESS
+					SLW_ROUGHNESS(StrataData.Layers[l].BSDFs[i]) = SLW_ROUGHNESS(StrataData.Layers[l].BSDFs[i]) * StrataDBufferData.OneMinusCoverage_Roughness + StrataDBufferData.Roughness;
+					SLW_SPECULAR(StrataData.Layers[l].BSDFs[i])  = SLW_SPECULAR(StrataData.Layers[l].BSDFs[i])  * StrataDBufferData.OneMinusCoverage_Roughness + StrataDBufferData.Specular;
+					SLW_METALLIC(StrataData.Layers[l].BSDFs[i])  = SLW_METALLIC(StrataData.Layers[l].BSDFs[i])  * StrataDBufferData.OneMinusCoverage_Roughness + StrataDBufferData.Metallic;
+				#endif
+					break;
+				}
+			} // switch
+		}
+	}
+  #endif // Strata - Disabling temporary DBuffer prior to optimization
+}
+#else // STRATA_ENABLED
 struct FDBufferData
 {
 	// 0..1, premultiplied with ColorOpacity
@@ -71,37 +313,6 @@ FDBufferData DecodeDBufferData(
 	return ret;
 }
 
-#if SUPPORTS_INDEPENDENT_SAMPLERS
-	// Share samplers to give artists more
-	#define OpaqueBasePassDBufferATextureSampler OpaqueBasePass.DBufferATextureSampler
-	#define OpaqueBasePassDBufferBTextureSampler OpaqueBasePass.DBufferATextureSampler
-	#define OpaqueBasePassDBufferCTextureSampler OpaqueBasePass.DBufferATextureSampler
-#else
-	#define OpaqueBasePassDBufferATextureSampler OpaqueBasePass.DBufferATextureSampler
-	#define OpaqueBasePassDBufferBTextureSampler OpaqueBasePass.DBufferBTextureSampler
-	#define OpaqueBasePassDBufferCTextureSampler OpaqueBasePass.DBufferCTextureSampler
-#endif
-
-// Returns a bitmask of DBuffer targets that contain valid data for this pixel.
-// @param BufferUV - UV space in the DBuffer textures
-uint GetDBufferTargetMask(uint2 PixelPos)
-{
-#if USE_DBUFFER && ENABLE_DBUFFER_TEXTURES && !RAYHITGROUPSHADER && !LUMEN_CARD_CAPTURE
-
-	#if PLATFORM_SUPPORTS_RENDERTARGET_WRITE_MASK
-		return DecodeRTWriteMask(PixelPos, OpaqueBasePass.DBufferRenderMask, 3);
-	#elif PLATFORM_SUPPORTS_PER_PIXEL_DBUFFER_MASK
-		uint Mask = OpaqueBasePass.DBufferRenderMask.Load(uint3(PixelPos, 0));
-		return Mask > 0 ? 0x07 : 0x00;
-	#else
-		return 0x07;
-	#endif
-
-#else
-	return 0;
-#endif
-}
-
 // @param BufferUV - UV space in the DBuffer textures
 FDBufferData GetDBufferData(float2 BufferUV, uint RTMaskBit)
 {
@@ -154,3 +365,4 @@ void ApplyDBufferData(
 	Specular = Specular * DBufferData.RoughnessOpacity + DBufferData.PreMulSpecular;
 #endif
 }
+#endif // STRATA_ENABLED
\ No newline at end of file
diff --git a/Engine/Shaders/Private/DecalCommon.ush b/Engine/Shaders/Private/DecalCommon.ush
--- a/Engine/Shaders/Private/DecalCommon.ush
+++ b/Engine/Shaders/Private/DecalCommon.ush
@@ -9,19 +9,67 @@
 // DECAL_RENDERTARGET_COUNT is set by C++
 
 #include "Common.ush"
+#include "DeferredShadingCommon.ush"
+#include "DBufferDecalShared.ush"
+
+#if STRATA_ENABLED  
+float DecalCommonOutput(
+	inout FPixelShaderOut Out, 
+	in FStrataPixelHeader InHeader,
+	in FStrataData InData,
+	in float CoverageWeight)
+{
+	FStrataDBuffer DBuffer = StrataConvertToDBuffer(InHeader, InData);
+	DBuffer.Coverage *= CoverageWeight;
+
+	float3 Color = DBuffer.Emissive;
+
+#if DECAL_RENDERSTAGE == DECAL_RENDERSTAGE_MOBILE
+	// Cheap mobile forward approximation uses both emissive and color.
+	Color += DBuffer.BaseColor;
+#if MATERIALBLENDING_MODULATE
+	Color *= DBuffer.Coverage;
+#endif
+#elif DECAL_RENDERSTAGE == DECAL_RENDERSTAGE_AO
+	// Replace color with AO for AO decals.
+	Color = lerp(1, InHeader.MaterialAO, DBuffer.Coverage);
+#endif
 
-#if SHADING_PATH_MOBILE
-	// Reroute MobileSceneTextures uniform buffer references to the base pass uniform buffer 
-	#define MobileSceneTextures MobileBasePass.SceneTextures
-	#define EyeAdaptationStruct MobileBasePass
-#else
-	#define SceneTexturesStruct DecalPass.SceneTextures
-	#define EyeAdaptationStruct DecalPass
+#if DECAL_RENDERTARGETMODE == DECAL_RENDERTARGETMODE_SCENECOLOR
+	Color *= View.PreExposure;
 #endif
 
-#include "DeferredShadingCommon.ush"
-#include "DBufferDecalShared.ush"
+	Out.MRT[0] = float4(Color, DBuffer.Coverage);
+
+#if DECAL_RENDERTARGETMODE == DECAL_RENDERTARGETMODE_DBUFFER
+
+	//StrataPackDBuffer(DBuffer, Out.MRT[0], Out.MRT[1], Out.MRT[2], Out.MRT[3]);
+	float4 PreMulEmissive_Coverage = 0;
+	StrataPackDBuffer(DBuffer, PreMulEmissive_Coverage, Out.MRT[0], Out.MRT[1], Out.MRT[2]);
+
+	#if DECAL_RENDERTARGET_COUNT == 4
+		Out.MRT[3] = 1.0f; // DBufferMask
+	#endif
+
+#endif
+	// Allow shader compiler to cull outputs that we won't use.
+	// Blend state also prevents writing unused outputs, but we want to be sure that we also avoid shader cost.
+#if !DECAL_OUT_MRT0
+	Out.MRT[0] = 0;
+#endif
+#if !DECAL_OUT_MRT1
+	Out.MRT[1] = 0;
+#endif
+#if !DECAL_OUT_MRT2
+	Out.MRT[2] = 0;
+#endif
+#if !DECAL_OUT_MRT3
+	Out.MRT[3] = 0;
+#endif
 
+	return DBuffer.Coverage;
+}
+#else // STRATA_ENABLED
 void DecalCommonOutput(inout FPixelShaderIn In, inout FPixelShaderOut Out, float3 Color, float Opacity, FGBufferData Data)
 {
 #if DECAL_RENDERSTAGE == DECAL_RENDERSTAGE_MOBILE
@@ -53,6 +101,7 @@ void DecalCommonOutput(inout FPixelShaderIn In, inout FPixelShaderOut Out, float
 	}
 #endif
 
+// This is not support by Strata
 #if DECAL_RENDERTARGETMODE == DECAL_RENDERTARGETMODE_GBUFFER || DECAL_RENDERTARGETMODE == DECAL_RENDERTARGETMODE_GBUFFER_NONORMAL
 	
 	// Dummy outputs from EncodeGBuffer()
@@ -101,3 +150,4 @@ void DecalCommonOutput(inout FPixelShaderIn In, inout FPixelShaderOut Out, float
 	Out.MRT[3] = 0;
 #endif
 }
+#endif // STRATA_ENABLED
diff --git a/Engine/Shaders/Private/DeferredDecal.usf b/Engine/Shaders/Private/DeferredDecal.usf
--- a/Engine/Shaders/Private/DeferredDecal.usf
+++ b/Engine/Shaders/Private/DeferredDecal.usf
@@ -4,8 +4,29 @@
 	DeferredDecal.usf: Pixel shader for computing a deferred decal.
 =============================================================================*/
 
+#include "Common.ush"
+
+#if SHADING_PATH_MOBILE
+// Reroute MobileSceneTextures uniform buffer references to the base pass uniform buffer 
+#define MobileSceneTextures MobileBasePass.SceneTextures
+#define EyeAdaptationStruct MobileBasePass
+#else
+#define SceneTexturesStruct DecalPass.SceneTextures
+#define EyeAdaptationStruct DecalPass
+#endif
+
+#if PIXELSHADER
+
+#if STRATA_ENABLED 
+#define STRATA_INLINE_SHADING 1
+#define STRATA_DEFERRED_SHADING 1
+#include "Strata/Strata.ush"
+#endif
+
 #include "DecalCommon.ush"
 
+#endif // PIXELSHADER
+
 // large world coordinate tile position
 float4 DecalTilePosition;
 // from SvPosition to decal space (for position), used like this: mul(float4(SvPosition.xyz, 1), SvPositionToDecal);
@@ -30,7 +51,7 @@ void MainVS(
 	OutPosition = mul(InPosition, FrustumComponentToClip);
 }
 
-#if IS_DECAL
+#if IS_DECAL && PIXELSHADER
 
 // DECAL_PRIMITIVE informs material templates which functions to expose when rendering decals.
 #define DECAL_PRIMITIVE 1
@@ -118,9 +139,18 @@ void FPixelShaderInOut_MainPS(inout FPixelShaderIn In, inout FPixelShaderOut Out
 	FPixelMaterialInputs PixelMaterialInputs;
 	CalcPixelMaterialInputs(MaterialParameters, PixelMaterialInputs);
 	
+	const float DecalFading = saturate(4 - 4 * abs(SwizzlePos.z * 2 - 1)) * DecalParams.x;
+
+#if STRATA_ENABLED    
+	FStrataPixelHeader StrataHeader = InitialiseStrataPixelHeader();
+	StrataHeader.BSDFCount = 1;
+	StrataHeader.MaterialAO = GetMaterialAmbientOcclusion(PixelMaterialInputs);
+	StrataHeader.SharedLocalBases = MaterialParameters.SharedLocalBases;
+	const float Coverage = DecalCommonOutput(Out, StrataHeader, PixelMaterialInputs.FrontMaterial, DecalFading);
+	const float Opacity = Coverage;
+#else // STRATA_ENABLED
 	float3 Color = GetMaterialEmissive(PixelMaterialInputs);
 	
-	float DecalFading = saturate(4 - 4 * abs(SwizzlePos.z * 2 - 1)) * DecalParams.x;
 	float Opacity = GetMaterialOpacity(PixelMaterialInputs) * DecalFading;
 
 	FGBufferData GBufferData;
@@ -140,6 +170,7 @@ void FPixelShaderInOut_MainPS(inout FPixelShaderIn In, inout FPixelShaderOut Out
 	GBufferData.Velocity = 0;
 
 	DecalCommonOutput(In, Out, Color, Opacity, GBufferData);
+#endif // STRATA_ENABLED
 
 #if NUM_VIRTUALTEXTURE_SAMPLES
 	FinalizeVirtualTextureFeedback(
diff --git a/Engine/Shaders/Private/MaterialTemplate.ush b/Engine/Shaders/Private/MaterialTemplate.ush
--- a/Engine/Shaders/Private/MaterialTemplate.ush
+++ b/Engine/Shaders/Private/MaterialTemplate.ush
@@ -19,7 +19,6 @@
 #include "/Engine/Private/HairStrands/HairCardsAttributeCommon.ush"
 #include "/Engine/Private/HairStrands/HairStrandsAttributeCommon.ush"
 #include "/Engine/Private/DeferredShadingCommon.ush"
-#include "/Engine/Private/DBufferDecalShared.ush"
 
 // Update this GUID to force all material shaders to recompile
 // Merge conflicts on this line should be resolved by generating a new GUID
@@ -159,6 +158,8 @@ struct FStrataData
 FStrataData GetInitialisedStrataData() { return (FStrataData)0; }
 #endif
 
+#include "/Engine/Private/DBufferDecalShared.ush"
+
 /**
  * Parameters used by vertex and pixel shaders to access particle properties.
  */
@@ -2386,6 +2387,16 @@ float4 MaterialExpressionDBufferTextureLookup(float2 BufferUV, int DBufferTextur
 	uint2 PixelPos = uint2(BufferUV * View.BufferSizeAndInvSize.xy);
 	uint ValidDBufferTargetMask = GetDBufferTargetMask(PixelPos) & (1u << DBufferTextureIndex);
 
+#if STRATA_ENABLED
+	const FStrataDBuffer DBufferData = GetDBufferData(BufferUV, DBufferTextureIndex);
+	switch (DBufferTextureIndex)
+	{
+	// All data (BaseColor, WorldNormal, Roughness) are pre-multiplied by their coverage
+	case 0:	return float4(DBufferData.BaseColor,	DBufferData.OneMinusCoverage_BaseColor);
+	case 1:	return float4(DBufferData.WorldNormal,	DBufferData.OneMinusCoverage_WorldNormal);
+	case 2:	return float4(DBufferData.Roughness,	DBufferData.Metallic, DBufferData.Specular, DBufferData.OneMinusCoverage_Roughness);
+	}
+#else
 	FDBufferData DBufferData = GetDBufferData(BufferUV, ValidDBufferTargetMask);
 	switch (DBufferTextureIndex)
 	{
@@ -2393,6 +2404,7 @@ float4 MaterialExpressionDBufferTextureLookup(float2 BufferUV, int DBufferTextur
 	case 1:	return float4(DBufferData.PreMulWorldNormal, DBufferData.NormalOpacity);
 	case 2:	return float4(DBufferData.PreMulRoughness, DBufferData.PreMulMetallic, DBufferData.PreMulSpecular, DBufferData.RoughnessOpacity);
 	}
+#endif
 
 	return float4(0, 0, 0, 1);
 }
diff --git a/Engine/Shaders/Private/MeshDecals.usf b/Engine/Shaders/Private/MeshDecals.usf
--- a/Engine/Shaders/Private/MeshDecals.usf
+++ b/Engine/Shaders/Private/MeshDecals.usf
@@ -4,6 +4,23 @@
 	MeshDecals.usf: Vertex/Hull/Domain/Pixel shader for computing a mesh decal.
 =============================================================================*/
 
+#include "Common.ush"
+
+#if SHADING_PATH_MOBILE
+	// Reroute MobileSceneTextures uniform buffer references to the base pass uniform buffer 
+	#define MobileSceneTextures MobileBasePass.SceneTextures
+	#define EyeAdaptationStruct MobileBasePass
+#else
+	#define SceneTexturesStruct DecalPass.SceneTextures
+	#define EyeAdaptationStruct DecalPass
+#endif
+
+#if STRATA_ENABLED 
+#define STRATA_INLINE_SHADING 1
+#define STRATA_DEFERRED_SHADING 1
+#include "Strata/Strata.ush"
+#endif
+
 #include "DecalCommon.ush"
 #include "/Engine/Generated/Material.ush"
 #include "/Engine/Generated/VertexFactory.ush"
@@ -120,6 +137,14 @@ void FPixelShaderInOut_MainPS(
 	}
 #endif
 
+#if STRATA_ENABLED
+	FStrataPixelHeader StrataHeader = InitialiseStrataPixelHeader();
+	StrataHeader.BSDFCount = 1;
+	StrataHeader.MaterialAO = GetMaterialAmbientOcclusion(PixelMaterialInputs);
+	StrataHeader.SharedLocalBases = MaterialParameters.SharedLocalBases;
+	const float Coverage = DecalCommonOutput(Out, StrataHeader, PixelMaterialInputs.FrontMaterial, 1.f);
+	const float Opacity = Coverage;
+#else // STRATA_ENABLED
 	float3 Color = GetMaterialEmissive(PixelMaterialInputs);
 
 	float Opacity = GetMaterialOpacity(PixelMaterialInputs);
@@ -141,6 +166,7 @@ void FPixelShaderInOut_MainPS(
 	GBufferData.Velocity = 0;
 
 	DecalCommonOutput(In, Out, Color, Opacity, GBufferData);
+#endif // STRATA_ENABLED
 
 #if NUM_VIRTUALTEXTURE_SAMPLES
 	FinalizeVirtualTextureFeedback(
diff --git a/Engine/Shaders/Private/Strata/Strata.ush b/Engine/Shaders/Private/Strata/Strata.ush
--- a/Engine/Shaders/Private/Strata/Strata.ush
+++ b/Engine/Shaders/Private/Strata/Strata.ush
@@ -1320,12 +1320,9 @@ float VerticalLayeringParameterBlendingBSDFCoverageToNormalMix(FStrataData TopSt
 }
 
 // ==> NOTE: Always pair with the compiler behavior in StrataCompilationInfoVerticalLayeringParamBlend
-FStrataData StrataVerticalLayeringParameterBlending(FStrataData Top, FStrataData Base, uint NewNormalIndex, float NoV)
+FStrataBSDF StrataVerticalLayeringParameterBlending(in FStrataBSDF TopBSDF, in FStrataBSDF BaseBSDF, uint NewNormalIndex, float NoV)
 {
-	FStrataData Result = GetInitialisedStrataDataWithReservedFirstSlot();
-#define ResultBSDF		Result.Layers[0].BSDFs[0]
-#define TopBSDF			Top.Layers[0].BSDFs[0]
-#define BaseBSDF		Base.Layers[0].BSDFs[0]
+	FStrataBSDF ResultBSDF = (FStrataBSDF)0;
 
 #if STRATA_INLINE_SHADING
 	const float SafeTopCoverage = saturate(TopBSDF.Coverage);
@@ -1397,14 +1394,16 @@ FStrataData StrataVerticalLayeringParameterBlending(FStrataData Top, FStrataData
 
 	ResultBSDF.Coverage =				1.0f - (1.0f - SafeTopCoverage) * (1.0f - SafeBaseCoverage);
 #endif
+	return ResultBSDF;
+}
 
-#undef ResultBSDF
-#undef TopBSDF
-#undef BaseBSDF
+FStrataData StrataVerticalLayeringParameterBlending(FStrataData Top, FStrataData Base, uint NewNormalIndex, float NoV)
+{
+	FStrataData Result = GetInitialisedStrataDataWithReservedFirstSlot();
+	Result.Layers[0].BSDFs[0] = StrataVerticalLayeringParameterBlending(Top.Layers[0].BSDFs[0], Base.Layers[0].BSDFs[0], NewNormalIndex, NoV);
 	return Result;
 }
 
-
 ///////////////////////////////////////////////////////////////////////////////
 // Functions used to sanitize BSDF before they are used for lighting (when forward) or stored to memory (base pass).
 // We do not normalize normals, this is done before the lighting step.
@@ -1994,6 +1993,147 @@ void StrataLoad_ColorGamma2ToLinearAlphaLinear(FStrataMaterialContainer StrataBu
 	A = Data.a;
 }
 
+///////////////////////////////////////////////////////////////////////////////
+// Strata decal
+
+// Dbuffer support only a single layer, with minimal set of feature (i.e., no anisotropy, no thinfilm, ...);
+// This special structure is for easing data processing
+struct FStrataDBuffer
+{
+	float3 WorldNormal;
+	float  Roughness;
+	float3 BaseColor;
+	float  Specular;
+	float3 Emissive;
+	float  Metallic;
+	float  Coverage;
+
+	// Per-component coverage 
+	// Since DBuffer properties are blended independely, they can end-up with varying coverage per-property. 
+	// These values are only valid when *reading* DBuffer Value
+	float  OneMinusCoverage;
+	float  OneMinusCoverage_BaseColor;
+	float  OneMinusCoverage_WorldNormal;
+	float  OneMinusCoverage_Roughness;
+};
+
+void StrataPackDBuffer(
+	in FStrataDBuffer In,
+	out float4 Out0,
+	out float4 Out1,
+	out float4 Out2,
+	out float4 Out3)
+{
+	// DBuffer packing needs to be blendable
+	Out0 = float4(In.Emissive,  In.Coverage);
+	Out1 = float4(In.BaseColor, In.Coverage);
+	Out2 = float4(In.WorldNormal * 0.5f + 128.0f / 255.0f, In.Coverage);
+	Out3 = float4(In.Metallic, In.Specular, In.Roughness, In.Coverage);
+}
+
+FStrataDBuffer StrataUnpackDBuffer(
+	in float4 In1,
+	in float4 In2,
+	in float4 In3)
+{
+	FStrataDBuffer Out = (FStrataDBuffer)0;
+
+	// Per-component coverage
+	Out.OneMinusCoverage			 = (In1.w + In2.w + In3.w) / 3.0f;
+	Out.OneMinusCoverage_BaseColor   = In1.w;
+	Out.OneMinusCoverage_WorldNormal = In2.w;
+	Out.OneMinusCoverage_Roughness   = In3.w;
+
+	// Unpack data
+	Out.BaseColor   = In1.xyz;
+	Out.Coverage    = 1.f - Out.OneMinusCoverage;
+	Out.WorldNormal = In2.rgb * 2 - (256.0 / 255.0);
+	Out.Metallic    = In3.x;
+	Out.Specular    = In3.y;
+	Out.Roughness   = In3.z;
+	Out.Emissive    = 0.f;
+
+
+	return Out;
+}
+
+FStrataDBuffer StrataConvertToDBuffer(
+	in FStrataPixelHeader InHeader,
+	in FStrataData InData)
+{
+	FStrataDBuffer Out = (FStrataDBuffer)0;
+	const float3 DiffuseColor = SLAB_DIFFUSEALBEDO(InData.Layers[0].BSDFs[0]);
+	const float3 F0 = SLAB_F0(InData.Layers[0].BSDFs[0]);
+
+	// STRATA_TODO : Can we avoid this by having a lerpable packing?
+	Out.Roughness = SLAB_ROUGHNESS(InData.Layers[0].BSDFs[0]);
+	Out.Metallic = F0RGBToMetallic(F0);
+	Out.Specular = F0ToDielectricSpecular(F0.x);
+	Out.BaseColor = lerp(DiffuseColor, F0, Out.Metallic);
+	Out.Emissive = BSDF_GETEMISSIVE(InData.Layers[0].BSDFs[0]);
+#if STRATA_INLINE_SHADING
+	Out.WorldNormal = InHeader.SharedLocalBases.Normals[BSDF_GETSHAREDLOCALBASISID(InData.Layers[0].BSDFs[0])];
+	Out.Coverage = InData.Layers[0].BSDFs[0].Coverage;
+#else
+	Out.WorldNormal = float3(0,0,1);
+	Out.Coverage = 1.f;
+#endif
+	return Out;
+}
+
+void StrataConvertFromDBuffer(
+	in FStrataDBuffer In,
+	inout FStrataPixelHeader OutHeader,
+	inout FStrataData OutData)
+{
+	// Header
+	const uint SharedLocalBasisIndex = 0;
+	uint SharedLocalBasisTypes = 0;
+	OutHeader.BSDFCount = 1;
+#if STRATA_INLINE_SHADING
+	OutHeader.SharedLocalBases = StrataInitialiseSharedLocalBases();
+	OutHeader.SharedLocalBases.Normals[SharedLocalBasisIndex] = In.WorldNormal;
+	OutHeader.SharedLocalBases.Tangents[SharedLocalBasisIndex] = 0.f;
+#endif
+	OutHeader.MaterialAO = 1.0f;
+	OutHeader.State = 0;
+	AddShadingModels(OutHeader, StrataShadingModelBit(STRATA_BSDF_TYPE_SLAB));
+
+	// Data
+	// Decal only support a single Slab BSDF with no-anisotropy
+	const float DefaultThickness = STRATA_LAYER_DEFAULT_THICKNESS_CM;
+	OutData = GetStrataSlabBSDF(
+		true,				// UseMetalness
+		In.BaseColor,		// BaseColor		-  Metalness workflow
+		1.f,				// EdgeColor		-  Metalness workflow
+		In.Specular,		// Specular			-  Metalness workflow
+		In.Metallic,		// Metallic			-  Metalness workflow
+		0.f,				// DiffuseAlbedo	- !Metalness workflow
+		0.f,				// F0				- !Metalness workflow
+		0.f,				// F90				- !Metalness workflow
+		In.Roughness,		// Roughness
+		0.f,				// Anisotropy
+		0.f,				// SSSProfileID
+		0.f,				// SSSDMFP
+		0.f,				// SSSDMFPScale
+		0.f,				// Emissive
+		0.f,				// Haziness
+		0.f,				// ThinFilmThickness
+		0.f,				// FuzzAmount
+		0.f,				// FuzzColor
+		DefaultThickness,
+		SharedLocalBasisIndex,
+		SharedLocalBasisTypes);
+#if STRATA_INLINE_SHADING
+	OutData.Layers[0].BSDFs[0].Emissive = In.Emissive;
+	OutData.Layers[0].BSDFs[0].Coverage = In.Coverage;
+
+	// DBuffer data are alreay pre-multiplied when read from the DBuffers
+	OutData.Layers[0].BSDFs[0].bPremultipliedCoverage = true;
+#endif
+}
+
+ 
 ///////////////////////////////////////////////////////////////////////////////
 // Strata deferred base-pass output
 
diff --git a/Engine/Shaders/Private/Strata/StrataEvaluation.ush b/Engine/Shaders/Private/Strata/StrataEvaluation.ush
--- a/Engine/Shaders/Private/Strata/StrataEvaluation.ush
+++ b/Engine/Shaders/Private/Strata/StrataEvaluation.ush
@@ -19,10 +19,6 @@
 #define STRATA_FASTPATH 0
 #endif
 
-#ifndef STRATA_USE_ENV_BRDF_APPROX
-#define STRATA_USE_ENV_BRDF_APPROX 0
-#endif
-
 #ifndef STRATA_TRANSLUCENT_ENABLED
 #define STRATA_TRANSLUCENT_ENABLED 0
 #endif
diff --git a/Engine/Source/Runtime/Engine/Classes/Engine/EngineTypes.h b/Engine/Source/Runtime/Engine/Classes/Engine/EngineTypes.h
--- a/Engine/Source/Runtime/Engine/Classes/Engine/EngineTypes.h
+++ b/Engine/Source/Runtime/Engine/Classes/Engine/EngineTypes.h
@@ -680,6 +680,11 @@ public:
 	void SetShadingModelFromExpression(bool bIn) { bHasShadingModelFromExpression = bIn ? 1u : 0u; }
 	bool HasShadingModelFromExpression() const { return bHasShadingModelFromExpression > 0u; }
 
+	uint32 GetPropertyConnected() const { return ConnectedProperties; }
+	void AddPropertyConnected(uint32 In) { ConnectedProperties |= (1 << In); }
+	bool HasPropertyConnected(uint32 In) const { return !!(ConnectedProperties & (1 << In)); }
+	static bool HasPropertyConnected(uint32 InConnectedProperties, uint32 In) { return !!(InConnectedProperties & (1 << In)); }
+
 	bool IsValid() const { return (ShadingModelField > 0) && (ShadingModelField < (1 << SSM_NUM)); }
 
 	bool operator==(const FStrataMaterialInfo& Other) const { return ShadingModelField == Other.GetShadingModelField(); }
@@ -693,6 +698,10 @@ private:
 	UPROPERTY()
 	uint8 bHasShadingModelFromExpression = 0;
 
+	/* Indicates which (legacy) inputs are connected */
+	UPROPERTY()
+	uint32 ConnectedProperties = 0;
+	
 	UPROPERTY()
 	TArray<TObjectPtr<USubsurfaceProfile>> SubsurfaceProfiles;
 };
diff --git a/Engine/Source/Runtime/Engine/Classes/Materials/Material.h b/Engine/Source/Runtime/Engine/Classes/Materials/Material.h
--- a/Engine/Source/Runtime/Engine/Classes/Materials/Material.h
+++ b/Engine/Source/Runtime/Engine/Classes/Materials/Material.h
@@ -399,6 +399,12 @@ class UMaterial : public UMaterialInterface
 	UPROPERTY(EditAnywhere, BlueprintReadOnly, Category=Material, AdvancedDisplay, meta=(DisplayName = "Decal Response (DBuffer)"), AssetRegistrySearchable)
 	TEnumAsByte<enum EMaterialDecalResponse> MaterialDecalResponse;
 
+	/**
+	 * Cached connected inputs
+	 */
+	UPROPERTY()
+	uint32 CachedConnectedInputs;
+
 private:
 	/** Determines how inputs are combined to create the material's final color. */
 	UPROPERTY(EditAnywhere, Category=Material, AssetRegistrySearchable)
diff --git a/Engine/Source/Runtime/Engine/Private/Materials/Material.cpp b/Engine/Source/Runtime/Engine/Private/Materials/Material.cpp
--- a/Engine/Source/Runtime/Engine/Private/Materials/Material.cpp
+++ b/Engine/Source/Runtime/Engine/Private/Materials/Material.cpp
@@ -3000,6 +3000,39 @@ void UMaterial::ConvertMaterialToStrataMaterial()
 			AddStrataShadingModelFromMaterialShadingModel(ConvertNode->ConvertedStrataMaterialInfo, ShadingModels);
 			check(ConvertNode->ConvertedStrataMaterialInfo.CountShadingModels() == 1);
 
+			FrontMaterial.Connect(0, ConvertNode);
+			bInvalidateShader = true;
+		}
+		else if (MaterialDomain == MD_DeferredDecal)
+		{
+			// Some decal materials don't have their shading mode set correctly to DefaultLit. Since only DefaultLit is supported, forcing it here.
+			ShadingModel = MSM_DefaultLit;
+			ShadingModels.ClearShadingModels();
+			ShadingModels.AddShadingModel(MSM_DefaultLit);
+
+			UMaterialExpressionStrataLegacyConversion* ConvertNode = NewObject<UMaterialExpressionStrataLegacyConversion>(this);
+			MoveConnectionTo(BaseColor, ConvertNode, 0);
+			MoveConnectionTo(Metallic, ConvertNode, 1);
+			MoveConnectionTo(Specular, ConvertNode, 2);
+			MoveConnectionTo(Roughness, ConvertNode, 3);
+			MoveConnectionTo(Anisotropy, ConvertNode, 4);
+			MoveConnectionTo(EmissiveColor, ConvertNode, 5);
+			CopyConnectionTo(Normal, ConvertNode, 6);
+			MoveConnectionTo(Tangent, ConvertNode, 7);
+			MoveConnectionTo(SubsurfaceColor, ConvertNode, 8);
+			MoveConnectionTo(ClearCoat, ConvertNode, 9);
+			MoveConnectionTo(ClearCoatRoughness, ConvertNode, 10);
+			MoveConnectionTo(Opacity, ConvertNode, 11);
+
+			// Add constant for the Unlit shading model
+			UMaterialExpressionConstant* ShadingModelNode = NewObject<UMaterialExpressionConstant>(this);
+			ShadingModelNode->SetParameterName(FName(TEXT("ConstantShadingModel")));
+			ShadingModelNode->R = ShadingModel;
+			ConvertNode->ShadingModel.Connect(0, ShadingModelNode);
+
+			AddStrataShadingModelFromMaterialShadingModel(ConvertNode->ConvertedStrataMaterialInfo, ShadingModels);
+			check(ConvertNode->ConvertedStrataMaterialInfo.CountShadingModels() == 1);
+
 			FrontMaterial.Connect(0, ConvertNode);
 			bInvalidateShader = true;
 		}
@@ -3918,6 +3951,7 @@ void UMaterial::RebuildShadingModelField()
 		if (this->FrontMaterial.Expression->IsResultStrataMaterial(this->FrontMaterial.OutputIndex))
 		{
 			this->FrontMaterial.Expression->GatherStrataMaterialInfo(StrataMaterialInfo, this->FrontMaterial.OutputIndex);
+			this->CachedConnectedInputs = StrataMaterialInfo.GetPropertyConnected();
 		}
 
 		bool bSanitizeMaterial = false;
@@ -4037,7 +4071,10 @@ void UMaterial::RebuildShadingModelField()
 			}
 			else if (StrataMaterialInfo.HasOnlyShadingModel(SSM_DefaultLit))
 			{
-				MaterialDomain = EMaterialDomain::MD_Surface;
+				if (MaterialDomain != EMaterialDomain::MD_Surface && MaterialDomain != EMaterialDomain::MD_DeferredDecal)
+				{
+					MaterialDomain = EMaterialDomain::MD_Surface;
+				}
 				ShadingModel = MSM_DefaultLit;
 				if (BlendMode != EBlendMode::BLEND_Opaque && BlendMode != EBlendMode::BLEND_Masked)
 				{
@@ -5356,7 +5393,11 @@ static bool IsPropertyActive_Internal(EMaterialProperty InProperty,
 	}
 	else if (Domain == MD_DeferredDecal)
 	{
-		if (InProperty >= MP_CustomizedUVs0 && InProperty <= MP_CustomizedUVs7)
+		if (bStrataEnabled)
+		{
+			return InProperty == MP_FrontMaterial;
+		}
+		else if (InProperty >= MP_CustomizedUVs0 && InProperty <= MP_CustomizedUVs7)
 		{
 			return true;
 		}
diff --git a/Engine/Source/Runtime/Engine/Private/Materials/MaterialExpressions.cpp b/Engine/Source/Runtime/Engine/Private/Materials/MaterialExpressions.cpp
--- a/Engine/Source/Runtime/Engine/Private/Materials/MaterialExpressions.cpp
+++ b/Engine/Source/Runtime/Engine/Private/Materials/MaterialExpressions.cpp
@@ -20616,6 +20616,21 @@ bool UMaterialExpressionStrataLegacyConversion::IsResultStrataMaterial(int32 Out
 
 void UMaterialExpressionStrataLegacyConversion::GatherStrataMaterialInfo(FStrataMaterialInfo& StrataMaterialInfo, int32 OutputIndex)
 {	
+	// Track connected input
+	if (BaseColor.IsConnected())			{ StrataMaterialInfo.AddPropertyConnected(MP_BaseColor); }
+	if (Metallic.IsConnected())				{ StrataMaterialInfo.AddPropertyConnected(MP_Metallic); }
+	if (Specular.IsConnected())				{ StrataMaterialInfo.AddPropertyConnected(MP_Specular); }
+	if (Roughness.IsConnected())			{ StrataMaterialInfo.AddPropertyConnected(MP_Roughness); }
+	if (Anisotropy.IsConnected())			{ StrataMaterialInfo.AddPropertyConnected(MP_Anisotropy); }
+	if (EmissiveColor.IsConnected())		{ StrataMaterialInfo.AddPropertyConnected(MP_EmissiveColor); }
+	if (Normal.IsConnected())				{ StrataMaterialInfo.AddPropertyConnected(MP_Normal); }
+	if (Tangent.IsConnected())				{ StrataMaterialInfo.AddPropertyConnected(MP_Tangent); }
+	if (SubSurfaceColor.IsConnected())		{ StrataMaterialInfo.AddPropertyConnected(MP_SubsurfaceColor); }
+	if (ClearCoat.IsConnected())			{ StrataMaterialInfo.AddPropertyConnected(MP_CustomData0); }
+	if (ClearCoatRoughness.IsConnected())	{ StrataMaterialInfo.AddPropertyConnected(MP_CustomData1); }
+	if (Opacity.IsConnected())				{ StrataMaterialInfo.AddPropertyConnected(MP_Opacity); }
+	if (ShadingModel.IsConnected())			{ StrataMaterialInfo.AddPropertyConnected(MP_ShadingModel); }
+
 	if (ConvertedStrataMaterialInfo.HasShadingModel(SSM_Unlit))					{ StrataMaterialInfo.AddShadingModel(SSM_Unlit); }
 	if (ConvertedStrataMaterialInfo.HasShadingModel(SSM_DefaultLit))			{ StrataMaterialInfo.AddShadingModel(SSM_DefaultLit); }
 	if (ConvertedStrataMaterialInfo.HasShadingModel(SSM_SubsurfaceLit))			{ StrataMaterialInfo.AddShadingModel(SSM_SubsurfaceLit); }
@@ -21111,6 +21126,25 @@ bool UMaterialExpressionStrataSlabBSDF::IsResultStrataMaterial(int32 OutputIndex
 
 void UMaterialExpressionStrataSlabBSDF::GatherStrataMaterialInfo(FStrataMaterialInfo& StrataMaterialInfo, int32 OutputIndex)
 {
+	// Track connected inputs
+	if (bUseMetalness)
+	{
+		if (BaseColor.IsConnected())		{ StrataMaterialInfo.AddPropertyConnected(MP_BaseColor); }
+		if (Metallic.IsConnected())			{ StrataMaterialInfo.AddPropertyConnected(MP_Metallic); }
+		if (Specular.IsConnected())			{ StrataMaterialInfo.AddPropertyConnected(MP_Specular); }
+	}
+	else
+	{
+		if (DiffuseAlbedo.IsConnected())	{ StrataMaterialInfo.AddPropertyConnected(MP_DiffuseColor); }
+		if (F0.IsConnected())				{ StrataMaterialInfo.AddPropertyConnected(MP_SpecularColor); }
+	}
+	if (Roughness.IsConnected())			{ StrataMaterialInfo.AddPropertyConnected(MP_Roughness); }
+	if (Anisotropy.IsConnected())			{ StrataMaterialInfo.AddPropertyConnected(MP_Anisotropy); }
+	if (EmissiveColor.IsConnected())		{ StrataMaterialInfo.AddPropertyConnected(MP_EmissiveColor); }
+	if (Normal.IsConnected())				{ StrataMaterialInfo.AddPropertyConnected(MP_Normal); }
+	if (Tangent.IsConnected())				{ StrataMaterialInfo.AddPropertyConnected(MP_Tangent); }
+	if (SSSDMFP.IsConnected())				{ StrataMaterialInfo.AddPropertyConnected(MP_SubsurfaceColor); }
+
 	if (HasSSS())
 	{
 		// We still do not know if this is going to be a real SSS node because it is only possible for BSDF at the bottom of the stack. Nevertheless, we take the worst case into account.
@@ -21301,6 +21335,7 @@ bool UMaterialExpressionStrataUnlitBSDF::IsResultStrataMaterial(int32 OutputInde
 
 void UMaterialExpressionStrataUnlitBSDF::GatherStrataMaterialInfo(FStrataMaterialInfo& StrataMaterialInfo, int32 OutputIndex)
 {
+	if (EmissiveColor.IsConnected()) { StrataMaterialInfo.AddPropertyConnected(MP_EmissiveColor); }
 	StrataMaterialInfo.AddShadingModel(SSM_Unlit);
 }
 #endif // WITH_EDITOR
@@ -21391,6 +21426,13 @@ bool UMaterialExpressionStrataHairBSDF::IsResultStrataMaterial(int32 OutputIndex
 
 void UMaterialExpressionStrataHairBSDF::GatherStrataMaterialInfo(FStrataMaterialInfo& StrataMaterialInfo, int32 OutputIndex)
 {
+	// Track connected inputs
+	if (BaseColor.IsConnected())		{ StrataMaterialInfo.AddPropertyConnected(MP_BaseColor); }
+	if (Specular.IsConnected())			{ StrataMaterialInfo.AddPropertyConnected(MP_Specular); }
+	if (Roughness.IsConnected())		{ StrataMaterialInfo.AddPropertyConnected(MP_Roughness); }
+	if (EmissiveColor.IsConnected())	{ StrataMaterialInfo.AddPropertyConnected(MP_EmissiveColor); }
+	if (Tangent.IsConnected())			{ StrataMaterialInfo.AddPropertyConnected(MP_Tangent); }
+
 	StrataMaterialInfo.AddShadingModel(SSM_Hair);
 }
 #endif // WITH_EDITOR
@@ -21496,6 +21538,14 @@ bool UMaterialExpressionStrataSingleLayerWaterBSDF::IsResultStrataMaterial(int32
 
 void UMaterialExpressionStrataSingleLayerWaterBSDF::GatherStrataMaterialInfo(FStrataMaterialInfo& StrataMaterialInfo, int32 OutputIndex)
 {
+	// Track connected inputs
+	if (BaseColor.IsConnected())		{ StrataMaterialInfo.AddPropertyConnected(MP_BaseColor); }
+	if (Metallic.IsConnected())			{ StrataMaterialInfo.AddPropertyConnected(MP_Metallic); }
+	if (Specular.IsConnected())			{ StrataMaterialInfo.AddPropertyConnected(MP_Specular); }
+	if (Roughness.IsConnected())		{ StrataMaterialInfo.AddPropertyConnected(MP_Roughness); }
+	if (EmissiveColor.IsConnected())	{ StrataMaterialInfo.AddPropertyConnected(MP_EmissiveColor); }
+	if (Normal.IsConnected())			{ StrataMaterialInfo.AddPropertyConnected(MP_Normal); }
+
 	StrataMaterialInfo.AddShadingModel(SSM_SingleLayerWater);
 }
 #endif // WITH_EDITOR
diff --git a/Engine/Source/Runtime/Engine/Private/Materials/MaterialShared.cpp b/Engine/Source/Runtime/Engine/Private/Materials/MaterialShared.cpp
--- a/Engine/Source/Runtime/Engine/Private/Materials/MaterialShared.cpp
+++ b/Engine/Source/Runtime/Engine/Private/Materials/MaterialShared.cpp
@@ -1703,6 +1703,52 @@ bool FMaterialResource::IsStrataMaterial() const
 	return bStrataEnabled;
 }
 
+bool Engine_IsStrataEnabled();
+
+bool FMaterialResource::HasMaterialPropertyConnected(EMaterialProperty In) const
+{
+	// STRATA_TODO: temporary validation until we have converted all domains
+	const bool bIsStrataSupportedDomain = 
+		Material->MaterialDomain == MD_PostProcess || 
+		Material->MaterialDomain == MD_DeferredDecal || 
+		Material->MaterialDomain == MD_Surface || 
+		Material->MaterialDomain == MD_Volume;
+
+	if (Engine_IsStrataEnabled() && bIsStrataSupportedDomain)
+	{
+		// Strata material traversal is cached as this is an expensive operation
+		#if WITH_EDITOR
+		check(Material->HasStrataFrontMaterialConnected());
+		if (!CachedStrataMaterialInfo.IsValid())
+		{
+			if (Material->FrontMaterial.Expression->IsResultStrataMaterial(Material->FrontMaterial.OutputIndex))
+			{
+				Material->FrontMaterial.Expression->GatherStrataMaterialInfo(CachedStrataMaterialInfo, Material->FrontMaterial.OutputIndex);
+			}
+		}
+		return CachedStrataMaterialInfo.HasPropertyConnected(In);
+		#else
+		return FStrataMaterialInfo::HasPropertyConnected(Material->CachedConnectedInputs, In);
+		#endif
+	}
+	else
+	{
+		switch (In)
+		{
+		case MP_EmissiveColor: 		return Material->HasEmissiveColorConnected();
+		case MP_Opacity: 			return Material->HasEmissiveColorConnected();
+		case MP_BaseColor: 			return Material->HasBaseColorConnected();
+		case MP_Normal: 			return Material->HasNormalConnected();
+		case MP_Roughness: 			return Material->HasRoughnessConnected();
+		case MP_Specular: 			return Material->HasSpecularConnected();
+		case MP_Metallic: 			return HasMaterialAttributesConnected() || Material->HasMetallicConnected();
+		case MP_Anisotropy: 		return Material->HasAnisotropyConnected();
+		case MP_AmbientOcclusion: 	return Material->HasAmbientOcclusionConnected();
+		}
+	}
+	return false;
+}
+
 bool FMaterialResource::RequiresSynchronousCompilation() const
 {
 	return Material->IsDefaultMaterial();
@@ -5231,14 +5277,30 @@ FMaterialShaderParameters::FMaterialShaderParameters(const FMaterial* InMaterial
 	bShouldCastDynamicShadows = InMaterial->ShouldCastDynamicShadows();
 	bWritesEveryPixel = InMaterial->WritesEveryPixel(false);
 	bWritesEveryPixelShadowPass = InMaterial->WritesEveryPixel(true);
-	bHasBaseColorConnected = InMaterial->HasBaseColorConnected();
-	bHasNormalConnected = InMaterial->HasNormalConnected();
-	bHasRoughnessConnected = InMaterial->HasRoughnessConnected();
-	bHasSpecularConnected = InMaterial->HasSpecularConnected();
-	bHasMetallicConnected = InMaterial->HasMetallicConnected();
-	bHasEmissiveColorConnected = InMaterial->HasEmissiveColorConnected();
-	bHasAmbientOcclusionConnected = InMaterial->HasAmbientOcclusionConnected();
-	bHasAnisotropyConnected = InMaterial->HasAnisotropyConnected();
+	if (Engine_IsStrataEnabled())
+	{
+		bHasDiffuseAlbedoConnected  = InMaterial->HasMaterialPropertyConnected(MP_DiffuseColor);
+		bHasF0Connected = InMaterial->HasMaterialPropertyConnected(MP_SpecularColor);
+		bHasBaseColorConnected = InMaterial->HasMaterialPropertyConnected(MP_BaseColor);
+		bHasNormalConnected = InMaterial->HasMaterialPropertyConnected(MP_Normal);
+		bHasRoughnessConnected = InMaterial->HasMaterialPropertyConnected(MP_Roughness);
+		bHasSpecularConnected = InMaterial->HasMaterialPropertyConnected(MP_Specular);
+		bHasMetallicConnected = InMaterial->HasMaterialPropertyConnected(MP_Metallic);
+		bHasEmissiveColorConnected = InMaterial->HasMaterialPropertyConnected(MP_EmissiveColor);
+		bHasAmbientOcclusionConnected = InMaterial->HasMaterialPropertyConnected(MP_AmbientOcclusion);
+		bHasAnisotropyConnected = InMaterial->HasMaterialPropertyConnected(MP_Anisotropy);
+	}
+	else
+	{
+		bHasBaseColorConnected = InMaterial->HasBaseColorConnected();
+		bHasNormalConnected = InMaterial->HasNormalConnected();
+		bHasRoughnessConnected = InMaterial->HasRoughnessConnected();
+		bHasSpecularConnected = InMaterial->HasSpecularConnected();
+		bHasMetallicConnected = InMaterial->HasMetallicConnected();
+		bHasEmissiveColorConnected = InMaterial->HasEmissiveColorConnected();
+		bHasAmbientOcclusionConnected = InMaterial->HasAmbientOcclusionConnected();
+		bHasAnisotropyConnected = InMaterial->HasAnisotropyConnected();
+	}
 	bHasVertexPositionOffsetConnected = InMaterial->HasVertexPositionOffsetConnected();
 	bHasPixelDepthOffsetConnected = InMaterial->HasPixelDepthOffsetConnected();
 	bMaterialMayModifyMeshPosition = InMaterial->MaterialMayModifyMeshPosition();
@@ -5274,4 +5336,4 @@ FMaterialShaderParameters::FMaterialShaderParameters(const FMaterial* InMaterial
 	bShouldDisableDepthTest = InMaterial->ShouldDisableDepthTest();
 }
 
-#undef LOCTEXT_NAMESPACE
\ No newline at end of file
+#undef LOCTEXT_NAMESPACE
diff --git a/Engine/Source/Runtime/Engine/Public/MaterialShared.h b/Engine/Source/Runtime/Engine/Public/MaterialShared.h
--- a/Engine/Source/Runtime/Engine/Public/MaterialShared.h
+++ b/Engine/Source/Runtime/Engine/Public/MaterialShared.h
@@ -1822,6 +1822,7 @@ public:
 	virtual bool HasEmissiveColorConnected() const { return false; }
 	virtual bool HasAnisotropyConnected() const { return false; }
 	virtual bool HasAmbientOcclusionConnected() const { return false; }
+	virtual bool HasMaterialPropertyConnected(EMaterialProperty In) const { return false; }
 	virtual bool IsStrataMaterial() const { return false; }
 	virtual bool RequiresSynchronousCompilation() const { return false; };
 	virtual bool IsDefaultMaterial() const { return false; };
@@ -2722,6 +2723,7 @@ public:
 	ENGINE_API virtual bool HasAnisotropyConnected() const override;
 	ENGINE_API virtual bool HasAmbientOcclusionConnected() const override;
 	ENGINE_API virtual bool IsStrataMaterial() const override;
+	ENGINE_API virtual bool HasMaterialPropertyConnected(EMaterialProperty In) const override;
 	ENGINE_API virtual FMaterialShadingModelField GetShadingModels() const override;
 	ENGINE_API virtual bool IsShadingModelFromMaterialExpression() const override;
 	ENGINE_API virtual enum ETranslucencyLightingMode GetTranslucencyLightingMode() const override;
@@ -2799,6 +2801,7 @@ public:
 protected:
 	UMaterial* Material;
 	UMaterialInstance* MaterialInstance;
+	mutable FStrataMaterialInfo CachedStrataMaterialInfo;
 
 	/** Entry point for compiling a specific material property.  This must call SetMaterialProperty. */
 	ENGINE_API virtual int32 CompilePropertyAndSetMaterialProperty(EMaterialProperty Property, class FMaterialCompiler* Compiler, EShaderFrequency OverrideShaderFrequency, bool bUsePreviousFrameTime) const override;
@@ -3303,6 +3306,8 @@ struct FMaterialShaderParameters
 			uint64 bShouldCastDynamicShadows : 1;
 			uint64 bWritesEveryPixel : 1;
 			uint64 bWritesEveryPixelShadowPass : 1;
+			uint64 bHasDiffuseAlbedoConnected : 1;
+			uint64 bHasF0Connected : 1;
 			uint64 bHasBaseColorConnected : 1;
 			uint64 bHasNormalConnected : 1;
 			uint64 bHasRoughnessConnected : 1;
diff --git a/Engine/Source/Runtime/Engine/Public/SceneTypes.h b/Engine/Source/Runtime/Engine/Public/SceneTypes.h
--- a/Engine/Source/Runtime/Engine/Public/SceneTypes.h
+++ b/Engine/Source/Runtime/Engine/Public/SceneTypes.h
@@ -161,8 +161,8 @@ enum EMaterialProperty
 	MP_EmissiveColor = 0 UMETA(DisplayName = "Emissive"),
 	MP_Opacity UMETA(DisplayName = "Opacity"),
 	MP_OpacityMask UMETA(DisplayName = "Opacity Mask"),
-	MP_DiffuseColor UMETA(Hidden),			// used in Lightmass, not exposed to user, computed from: BaseColor, Metallic
-	MP_SpecularColor UMETA(Hidden),			// used in Lightmass, not exposed to user, derived from: SpecularColor, Metallic, Specular
+	MP_DiffuseColor UMETA(Hidden),			// used in Lightmass, not exposed to user, computed from: BaseColor, Metallic				Also used in Strata which uses Albedo/F0 parameterrization
+	MP_SpecularColor UMETA(Hidden),			// used in Lightmass, not exposed to user, derived from: SpecularColor, Metallic, Specular	Also used in Strata which uses Albedo/F0 parameterrization
 	MP_BaseColor UMETA(DisplayName = "Diffuse"),
 	MP_Metallic UMETA(DisplayName = "Metallic"),
 	MP_Specular UMETA(DisplayName = "Specular"),
diff --git a/Engine/Source/Runtime/Renderer/Private/CompositionLighting/PostProcessDeferredDecals.cpp b/Engine/Source/Runtime/Renderer/Private/CompositionLighting/PostProcessDeferredDecals.cpp
--- a/Engine/Source/Runtime/Renderer/Private/CompositionLighting/PostProcessDeferredDecals.cpp
+++ b/Engine/Source/Runtime/Renderer/Private/CompositionLighting/PostProcessDeferredDecals.cpp
@@ -456,6 +456,12 @@ void AddDeferredDecalPass(
 
 	const auto RenderDecals = [&](uint32 DecalIndexBegin, uint32 DecalIndexEnd, EDecalRenderTargetMode RenderTargetMode)
 	{
+		// Sanity check - Strata only support DBuffer, SceneColor, or AO decals
+		if (Strata::IsStrataEnabled())
+		{
+			check(RenderTargetMode == EDecalRenderTargetMode::DBuffer || RenderTargetMode == EDecalRenderTargetMode::SceneColor || RenderTargetMode == EDecalRenderTargetMode::AmbientOcclusion);
+		}
+
 		auto* PassParameters = GraphBuilder.AllocParameters<FDeferredDecalPassParameters>();
 		GetDeferredDecalPassParameters(GraphBuilder, View, PassTextures, RenderTargetMode, *PassParameters);
 
diff --git a/Engine/Source/Runtime/Renderer/Private/DecalRenderingCommon.cpp b/Engine/Source/Runtime/Renderer/Private/DecalRenderingCommon.cpp
--- a/Engine/Source/Runtime/Renderer/Private/DecalRenderingCommon.cpp
+++ b/Engine/Source/Runtime/Renderer/Private/DecalRenderingCommon.cpp
@@ -2,6 +2,7 @@
 
 #include "DecalRenderingCommon.h"
 #include "RenderUtils.h"
+#include "Strata/Strata.h"
 
 namespace DecalRendering
 {
@@ -74,12 +75,36 @@ namespace DecalRendering
 	FDecalBlendDesc ComputeDecalBlendDesc(EShaderPlatform Platform, FMaterial const* Material)
 	{
 		FDecalBlendDesc Desc;
-		Desc.BlendMode = Material->GetBlendMode();
-		Desc.bWriteBaseColor = Material->HasBaseColorConnected();
-		Desc.bWriteNormal = Material->HasNormalConnected();
-		Desc.bWriteRoughnessSpecularMetallic = Material->HasRoughnessConnected() || Material->HasSpecularConnected() || Material->HasMetallicConnected();
-		Desc.bWriteEmissive = Material->HasEmissiveColorConnected();
-		Desc.bWriteAmbientOcclusion = Material->HasAmbientOcclusionConnected();
+		if (Strata::IsStrataEnabled())
+		{
+			check(Material->IsStrataMaterial());
+
+			const bool bUseDiffuseAlbedoAndF0 =
+				Material->HasMaterialPropertyConnected(EMaterialProperty::MP_DiffuseColor) ||	// This is used for Strata Slab using (DiffuseAlbedo | F0) parameterization
+				Material->HasMaterialPropertyConnected(EMaterialProperty::MP_SpecularColor);	// This is used for Strata Slab using (DiffuseAlbedo | F0) parameterization
+
+			Desc.BlendMode = Material->GetBlendMode();
+			Desc.bWriteBaseColor = Material->HasMaterialPropertyConnected(EMaterialProperty::MP_BaseColor) || bUseDiffuseAlbedoAndF0;
+			Desc.bWriteNormal = Material->HasMaterialPropertyConnected(EMaterialProperty::MP_Normal);
+			Desc.bWriteRoughnessSpecularMetallic =
+				bUseDiffuseAlbedoAndF0 ||
+				Material->HasMaterialPropertyConnected(EMaterialProperty::MP_Metallic) ||
+				Material->HasMaterialPropertyConnected(EMaterialProperty::MP_Specular) ||
+				Material->HasMaterialPropertyConnected(EMaterialProperty::MP_Roughness);
+			Desc.bWriteEmissive=
+				Material->HasMaterialPropertyConnected(EMaterialProperty::MP_EmissiveColor);
+			Desc.bWriteAmbientOcclusion =
+				Material->HasMaterialPropertyConnected(EMaterialProperty::MP_AmbientOcclusion);
+		}
+		else
+		{
+			Desc.BlendMode = Material->GetBlendMode();
+			Desc.bWriteBaseColor = Material->HasBaseColorConnected();
+			Desc.bWriteNormal = Material->HasNormalConnected();
+			Desc.bWriteRoughnessSpecularMetallic = Material->HasRoughnessConnected() || Material->HasSpecularConnected() || Material->HasMetallicConnected();
+			Desc.bWriteEmissive = Material->HasEmissiveColorConnected();
+			Desc.bWriteAmbientOcclusion = Material->HasAmbientOcclusionConnected();
+		}
 		FinalizeBlendDesc(Platform, Desc);
 		return Desc;
 	}
@@ -87,12 +112,32 @@ namespace DecalRendering
 	FDecalBlendDesc ComputeDecalBlendDesc(EShaderPlatform Platform, FMaterialShaderParameters const& MaterialShaderParameters)
 	{
 		FDecalBlendDesc Desc;
-		Desc.BlendMode = MaterialShaderParameters.BlendMode;
-		Desc.bWriteBaseColor = MaterialShaderParameters.bHasBaseColorConnected;
-		Desc.bWriteNormal = MaterialShaderParameters.bHasNormalConnected;
-		Desc.bWriteRoughnessSpecularMetallic = MaterialShaderParameters.bHasRoughnessConnected || MaterialShaderParameters.bHasSpecularConnected || MaterialShaderParameters.bHasMetallicConnected;
-		Desc.bWriteEmissive = MaterialShaderParameters.bHasEmissiveColorConnected;
-		Desc.bWriteAmbientOcclusion = MaterialShaderParameters.bHasAmbientOcclusionConnected;
+		if (Strata::IsStrataEnabled())
+		{
+			const bool bUseDiffuseAlbedoAndF0 = 
+				MaterialShaderParameters.bHasDiffuseAlbedoConnected || 
+				MaterialShaderParameters.bHasF0Connected;
+			
+			Desc.BlendMode = MaterialShaderParameters.BlendMode;
+			Desc.bWriteBaseColor = MaterialShaderParameters.bHasBaseColorConnected || bUseDiffuseAlbedoAndF0;
+			Desc.bWriteNormal = MaterialShaderParameters.bHasNormalConnected;
+			Desc.bWriteRoughnessSpecularMetallic = 
+				bUseDiffuseAlbedoAndF0 ||
+				MaterialShaderParameters.bHasRoughnessConnected || 
+				MaterialShaderParameters.bHasSpecularConnected || 
+				MaterialShaderParameters.bHasMetallicConnected;
+			Desc.bWriteEmissive = MaterialShaderParameters.bHasEmissiveColorConnected;
+			Desc.bWriteAmbientOcclusion = MaterialShaderParameters.bHasAmbientOcclusionConnected;
+		}
+		else
+		{
+			Desc.BlendMode = MaterialShaderParameters.BlendMode;
+			Desc.bWriteBaseColor = MaterialShaderParameters.bHasBaseColorConnected;
+			Desc.bWriteNormal = MaterialShaderParameters.bHasNormalConnected;
+			Desc.bWriteRoughnessSpecularMetallic = MaterialShaderParameters.bHasRoughnessConnected || MaterialShaderParameters.bHasSpecularConnected || MaterialShaderParameters.bHasMetallicConnected;
+			Desc.bWriteEmissive = MaterialShaderParameters.bHasEmissiveColorConnected;
+			Desc.bWriteAmbientOcclusion = MaterialShaderParameters.bHasAmbientOcclusionConnected;
+		}
 		FinalizeBlendDesc(Platform, Desc);
 		return Desc;
 	}
diff --git a/Engine/Source/Runtime/Renderer/Private/DecalRenderingShared.cpp b/Engine/Source/Runtime/Renderer/Private/DecalRenderingShared.cpp
--- a/Engine/Source/Runtime/Renderer/Private/DecalRenderingShared.cpp
+++ b/Engine/Source/Runtime/Renderer/Private/DecalRenderingShared.cpp
@@ -66,6 +66,7 @@ public:
 	{
 		FMaterialShader::ModifyCompilationEnvironment(Parameters, OutEnvironment);
 		DecalRendering::ModifyCompilationEnvironment(DecalRendering::ComputeDecalBlendDesc(Parameters.Platform, Parameters.MaterialParameters), EDecalRenderStage::None, OutEnvironment);
+		OutEnvironment.SetDefine(TEXT("STRATA_ENABLED"), Strata::IsStrataEnabled() ? 1u : 0u);
 	}
 
 	FDeferredDecalPS() {}
@@ -370,6 +371,7 @@ namespace DecalRendering
 		{
 			PixelShader = MaterialShaderMap->GetShader<FDeferredDecalPS>();
 		}
+		check(!PixelShader.IsNull());
 
 		TShaderMapRef<FDeferredDecalVS> VertexShader(View.ShaderMap);
 
