Fix bad resolve error in DeferredLightPixelShaders.usf

diff --git a/Engine/Shaders/Private/DeferredLightPixelShaders.usf b/Engine/Shaders/Private/DeferredLightPixelShaders.usf
--- a/Engine/Shaders/Private/DeferredLightPixelShaders.usf
+++ b/Engine/Shaders/Private/DeferredLightPixelShaders.usf
@@ -272,7 +272,6 @@ void DeferredLightPixelMain(
 				ToLight,
 				LightMask,
 				ShadowTerms,
-				RectTexture,
 				Strata.MaterialTextureArray,
 				StrataAddressing,
 				StrataPixelHeader);
