Reduced volumetric cloud screenspace triangle vertex shader to barebone minimum to avoid compilation issue on some tv platform.

diff --git a/Engine/Shaders/Private/VolumetricCloud.usf b/Engine/Shaders/Private/VolumetricCloud.usf
--- a/Engine/Shaders/Private/VolumetricCloud.usf
+++ b/Engine/Shaders/Private/VolumetricCloud.usf
@@ -39,12 +39,6 @@
 
 #if !SHADER_RENDERVIEW_CS
 #include "/Engine/Generated/VertexFactory.ush"
-
-struct RenderVolumetricCloudVSToPS
-{
-	FVertexFactoryInterpolantsVSToPS FactoryInterpolants;
-	float4 Position : SV_POSITION;
-};
 #endif
 
 #include "ReflectionEnvironmentShared.ush"
@@ -172,20 +166,13 @@ bool ViewHierarchicalZBufferMinDepth(float4 ScreenUVMinMax, const float MaxZDept
 #ifdef SHADER_MAINVS
 
 void MainVS(
-	in FVertexFactoryInput FactoryInput,
-	out RenderVolumetricCloudVSToPS Interpolants
+	float3 InPosition : ATTRIBUTE0,
+	out float4 Position : SV_POSITION
 )
 {
 	ResolvedView = ResolveView();
 
-	// Full screen triangle
-	Interpolants.Position = FactoryInput.Position;
-
-	const float3 WorldPosition = 0.0f;
-	FVertexFactoryIntermediates VFIntermediates = GetVertexFactoryIntermediates(FactoryInput);
-	float3x3 TangentToLocal = VertexFactoryGetTangentToLocal(FactoryInput, VFIntermediates);
-	FMaterialVertexParameters VertexParameters = GetMaterialVertexParameters(FactoryInput, VFIntermediates, WorldPosition, TangentToLocal);
-	Interpolants.FactoryInterpolants = VertexFactoryGetInterpolantsVSToPS(FactoryInput, VFIntermediates, VertexParameters);
+	Position = float4(InPosition, 1.0f);
 }
 
 #endif // SHADER_MAINVS
@@ -1312,12 +1299,11 @@ void MainCommon(in FMaterialPixelParameters MaterialParameters, in float4 SvPosi
 #if SHADER_RENDERVIEW_PS
 
 void MainPS(
-	  in FVertexFactoryInterpolantsVSToPS FactoryInterpolants
-	, in float4 SvPosition : SV_Position
+	in float4 SvPosition : SV_Position
 	, out float4 OutColor0 : SV_Target0
 	, out float2 OutColor1 : SV_Target1)
 {
-	FMaterialPixelParameters MaterialParameters = GetMaterialPixelParameters(FactoryInterpolants, SvPosition);
+	FMaterialPixelParameters MaterialParameters = MakeInitializedMaterialPixelParameters();
 
 	MainCommon(MaterialParameters, SvPosition, OutColor0, OutColor1);
 }
