Fix for opaque casting shadows on cloud using non VSM shadowmaps.
- Shader now sends correct translated world position
- Fixed an autotest to use far shadow cascades and have the "Far Shadow" enabled on the cube mesh.
- Updated a reference image differing a bit.

diff --git a/Engine/Shaders/Private/VolumetricCloud.usf b/Engine/Shaders/Private/VolumetricCloud.usf
--- a/Engine/Shaders/Private/VolumetricCloud.usf
+++ b/Engine/Shaders/Private/VolumetricCloud.usf
@@ -780,7 +780,7 @@ void MainCommon(in FMaterialPixelParameters MaterialParameters, in float4 SvPosi
 				// B- sample shadow from opaque if enabled. Only light0 supported.
 				//
 				bool bUnused;
-				float OpaqueShadow = ComputeLight0VolumeShadowing(SampleWorldPosition - LWCHackToFloat(PrimaryView.PreViewTranslation), false, false, bUnused);
+				float OpaqueShadow = ComputeLight0VolumeShadowing(SampleTranslatedWorldPosition, false, false, bUnused);
 
 #if VIRTUAL_SHADOW_MAP 
 				if (RenderVolumetricCloudParameters.VirtualShadowMapId0 != INDEX_NONE)
