Fix issue in eye and clear coat shading models when introduce shading energy conservation/preservation framework.

diff --git a/Engine/Shaders/Private/ShadingModels.ush b/Engine/Shaders/Private/ShadingModels.ush
--- a/Engine/Shaders/Private/ShadingModels.ush
+++ b/Engine/Shaders/Private/ShadingModels.ush
@@ -410,7 +410,12 @@ FDirectLighting ClearCoatBxDF( FGBufferData GBuffer, half3 N, half3 V, half3 L,
 
 	// Incoming and exiting Fresnel terms are identical to incoming Fresnel term (VoH == HoL)
 	// float FresnelCoeff = (1.0 - F1) * (1.0 - F2);
+	#if USE_ENERGY_CONSERVATION
 	float3 FresnelCoeff = ComputeEnergyPreservation(EnergyTermsCoat); // 1.0 - F;
+	#else
+	// Preserve old behavior when energy conservation is disabled
+	float FresnelCoeff = 1.0 - F;
+	#endif
 	FresnelCoeff *= FresnelCoeff;
 
 	//////////////////////////////
@@ -781,8 +786,15 @@ FDirectLighting EyeBxDF( FGBufferData GBuffer, half3 N, half3 V, half3 L, float
 	
 	Lighting.Specular *= ComputeEnergyConservation(EnergyTerms);
 
+#if USE_ENERGY_CONSERVATION
+	const float EnergyPreservation = ComputeEnergyPreservation(EnergyTerms);
+#else
+	// Preserve old behavior when energy conservation is disabled
+	const float EnergyPreservation = 1.0f - F;
+#endif
+
 	Lighting.Diffuse = 0;
-	Lighting.Transmission = AreaLight.FalloffColor * ( Falloff * lerp( Sclera, Iris, IrisMask ) * ComputeEnergyPreservation(EnergyTerms) ) * Diffuse_Lambert( GBuffer.DiffuseColor );
+	Lighting.Transmission = AreaLight.FalloffColor * ( Falloff * lerp( Sclera, Iris, IrisMask ) * EnergyPreservation ) * Diffuse_Lambert( GBuffer.DiffuseColor );
 	return Lighting;
 }
 
