Fixe (dx12) crash due to vs/ps mismatch in CloudShadowPS.
Originated from 18861501, this pixelshader spot was missed.

diff --git a/Engine/Shaders/Private/VolumetricCloud.usf b/Engine/Shaders/Private/VolumetricCloud.usf
--- a/Engine/Shaders/Private/VolumetricCloud.usf
+++ b/Engine/Shaders/Private/VolumetricCloud.usf
@@ -1384,8 +1384,7 @@ struct CloudShadowTraceContext
 };
 
 void MainPS(
-	in FVertexFactoryInterpolantsVSToPS FactoryInterpolants
-	, in float4 SvPosition : SV_Position
+	in float4 SvPosition : SV_Position
 	, out float3 OutColor0 : SV_Target0
 )
 {
@@ -1420,7 +1419,7 @@ void MainPS(
 		TraceContext.TraceDir					= RenderVolumetricCloudParameters.CloudSkyAOTraceDir;
 	}
 
-	FMaterialPixelParameters MaterialParameters = GetMaterialPixelParameters(FactoryInterpolants, SvPosition);
+	FMaterialPixelParameters MaterialParameters = MakeInitializedMaterialPixelParameters();
 	FPixelMaterialInputs PixelMaterialInputs;
 	CalcMaterialParameters(MaterialParameters, PixelMaterialInputs, SvPosition, true);
 
