Fixed compilation issue when energy conservation is enabled.

diff --git a/Engine/Shaders/Private/ShadingModels.ush b/Engine/Shaders/Private/ShadingModels.ush
--- a/Engine/Shaders/Private/ShadingModels.ush
+++ b/Engine/Shaders/Private/ShadingModels.ush
@@ -787,7 +787,7 @@ FDirectLighting EyeBxDF( FGBufferData GBuffer, half3 N, half3 V, half3 L, float
 	Lighting.Specular *= ComputeEnergyConservation(EnergyTerms);
 
 #if USE_ENERGY_CONSERVATION
-	const float EnergyPreservation = ComputeEnergyPreservation(EnergyTerms);
+	const float3 EnergyPreservation = ComputeEnergyPreservation(EnergyTerms);
 #else
 	// Preserve old behavior when energy conservation is disabled
 	const float EnergyPreservation = 1.0f - F;
