Fix energy error for clear coat material when illuminated with sky lighting.

diff --git a/Engine/Shaders/Private/ReflectionEnvironmentPixelShader.usf b/Engine/Shaders/Private/ReflectionEnvironmentPixelShader.usf
--- a/Engine/Shaders/Private/ReflectionEnvironmentPixelShader.usf
+++ b/Engine/Shaders/Private/ReflectionEnvironmentPixelShader.usf
@@ -183,7 +183,10 @@ float3 ReflectionEnvironment(FGBufferData GBuffer, float AmbientOcclusion, float
 		#if USE_ENERGY_CONSERVATION
 		Color.rgb *= EnergyTerms.E * (1 - ClearCoat);
 		#else
-		Color.rgb *= EnvBRDF(SpecularColor, GBuffer.Roughness, NoV) * (1 - ClearCoat);
+		// Hack: Ensures when clear coat is >0, grazing angle does not get too much energy, 
+		//       but preserve response at normal incidence 
+		float2 AB = PreIntegratedGF.SampleLevel(PreIntegratedGFSampler, float2(NoV, GBuffer.Roughness), 0).rg;
+		Color.rgb *= SpecularColor * AB.x + AB.y * saturate(50 * SpecularColor.g) * (1 - ClearCoat);
 		#endif
 
 		// F_Schlick
