Remove signed integer & div. computation from voxel traversal loop.

Small improvement in perf 1.20ms -> 1.12ms

diff --git a/Engine/Shaders/Private/HairStrands/HairStrandsVoxelPageTraversal.ush b/Engine/Shaders/Private/HairStrands/HairStrandsVoxelPageTraversal.ush
--- a/Engine/Shaders/Private/HairStrands/HairStrandsVoxelPageTraversal.ush
+++ b/Engine/Shaders/Private/HairStrands/HairStrandsVoxelPageTraversal.ush
@@ -184,7 +184,7 @@ FHairTraversalResult ComputeHairCountVirtualVoxel(
 		return Out;
 	}
 
-	int3 CurrentPageIndexCoord = -1;
+	uint3 CurrentPageIndexCoord = 9999;
 	bool bIsPageValid = false;
 	uint3 PageCoord = 0;
 
@@ -219,18 +219,18 @@ FHairTraversalResult ComputeHairCountVirtualVoxel(
 		{
 			const float SteppingWorldSize = max(DeltaWorld * StepScale, InSettings.TanConeAngle * StepIt * InCommonDesc.VoxelWorldSize);
 			const float3 HitP = O + StepIt * D + InSettings.JitterScale * RandomStepJitter * SteppingWorldSize * 0.5f;
-
+			
+			const float3 NormalizedPos  = saturate((HitP - InNodeDesc.TranslatedWorldMinAABB) / (InNodeDesc.TranslatedWorldMaxAABB - InNodeDesc.TranslatedWorldMinAABB));
+			const uint3  VolumeCoord    = min(NormalizedPos * InNodeDesc.VirtualResolution, InNodeDesc.VirtualResolution-1u);
+			const uint3  PageIndexCoord = min(NormalizedPos * InNodeDesc.PageIndexResolution, InNodeDesc.PageIndexResolution-1u);
+			
 			#if VOXEL_TRAVERSAL_DEBUG
 			AddStepAABB(InSettings.bDebugEnabled, PreviousP, HitP, float3(0, 1, 0));
 			#endif
 
-			const int3 VolumeCoord = PositionToCoord(HitP, InNodeDesc.TranslatedWorldMinAABB, InNodeDesc.TranslatedWorldMaxAABB, InNodeDesc.VirtualResolution);
-			const int3 PageIndexCoord = VolumeCoord / InCommonDesc.PageResolution;
-
 			// Update page index only when needed
-			if (PageIndexCoord.x != CurrentPageIndexCoord.x ||
-				PageIndexCoord.y != CurrentPageIndexCoord.y ||
-				PageIndexCoord.z != CurrentPageIndexCoord.z)
+			const bool bIsNewPage = any(PageIndexCoord != CurrentPageIndexCoord);
+			if (bIsNewPage)
 			{
 				CurrentPageIndexCoord = PageIndexCoord;
 				const uint LinearPageIndexCoord = CoordToIndex(PageIndexCoord, InNodeDesc.PageIndexResolution, InNodeDesc.PageIndexOffset);
@@ -253,9 +253,9 @@ FHairTraversalResult ComputeHairCountVirtualVoxel(
 
 			if (bIsPageValid)
 			{
-				const int3 VoxelPageBase = PageCoord * InCommonDesc.PageResolution;
-				const int3 VoxelPageOffset = VolumeCoord - PageIndexCoord * InCommonDesc.PageResolution;
-				const int3 VoxelPageCoord = VoxelPageBase + VoxelPageOffset;
+				const uint3 VoxelPageBase   = PageCoord * InCommonDesc.PageResolution;
+				const uint3 VoxelPageOffset = VolumeCoord - PageIndexCoord * InCommonDesc.PageResolution;
+				const uint3 VoxelPageCoord  = VoxelPageBase + VoxelPageOffset;
 
 				#if VOXEL_TRAVERSAL_PAGEOCCUPANCY
 				const bool bIsSubBrickValid = InSettings.bUseOccupancy ? DoesVoxelContainData(PageOccupancy, VoxelPageOffset) : true;
