Restructure cluster shader (extract sub-functions) to ease porting Strata lighting code.

diff --git a/Engine/Shaders/Private/ClusteredDeferredShadingPixelShader.usf b/Engine/Shaders/Private/ClusteredDeferredShadingPixelShader.usf
--- a/Engine/Shaders/Private/ClusteredDeferredShadingPixelShader.usf
+++ b/Engine/Shaders/Private/ClusteredDeferredShadingPixelShader.usf
@@ -1,4 +1,3 @@
-// Copyright Epic Games, Inc. All Rights Reserved.
 
 /*=============================================================================
 	ClusteredDeferredShadingPixelShader - applies all local lights in the
@@ -39,6 +38,9 @@
 // Extracts out the simple lights into their own loop that is executed first, independently of shading model.
 #define ENABLE_SEPARATE_SIMPLE_LIGHT_LOOP 0
 
+// Enable cluster debug visualization
+#define GRID_DEBUG_ENABLE 0
+
 Texture2D< uint4 > ShadowMaskBits;
 
 #if USE_HAIR_LIGHTING
@@ -46,19 +48,169 @@ uint				HairTransmittanceBufferMaxCount;
 Buffer<uint>		HairTransmittanceBuffer;
 #endif
 
+/**
+ * Debug function for visualizing clusters
+ */ 
+bool GetClusterDebugColor(float2 LocalPosition, float SceneDepth, FCulledLightsGridData GridData, inout float4 OutColor)
+{
+	OutColor = 0.0f;
+#if GRID_DEBUG_ENABLE
+	uint ZSlice = (uint)(max(0, log2(SceneDepth * GridData.LightGridZParams.x + GridData.LightGridZParams.y) * GridData.LightGridZParams.z));
+	ZSlice = min(ZSlice, (uint)(GridData.CulledGridSize.z - 1));
+	uint3 GridCoordinate = uint3(uint2(LocalPosition.x, LocalPosition.y) >> GridData.LightGridPixelSizeShift, ZSlice);
+
+	if (ZSlice % 2)
+	{
+		OutColor = float(ZSlice) / float(GridData.CulledGridSize.z - 1);
+	}
+	else
+	{
+		OutColor = 1.0f - float(ZSlice) / float(GridData.CulledGridSize.z - 1);
+	}
+
+	if (GridCoordinate.x % 2)
+	{
+		OutColor.xyz = 1.0f - OutColor.xyz;
+	}
+	if (GridCoordinate.y % 2)
+	{
+		OutColor.xyz = 1.0f - OutColor.xyz;
+	}
+	OutColor.w = 0.0f;
+	return true;
+#else
+	return false;
+#endif
+}
+
+/**
+ * Return true is the light is visible, false otherwise
+ */
+bool IsLightVisible(const FLocalLightData In, float3 TranslatedWorldPosition)
+{
+	// extra-early out since we know light grid is sloppy and all lights in list are radial (have a range)
+	// appears useless
+	const float InvLightRadiusSq = In.LightPositionAndInvRadius.w * In.LightPositionAndInvRadius.w;
+	return length2(TranslatedWorldPosition - In.LightPositionAndInvRadius.xyz) * InvLightRadiusSq <= 1.0f;
+}
+
+/**
+ * Convert a FLocalLightData into a FDeferredLightData
+ */
+FDeferredLightData ConvertLocalToDeferredLight(const FLocalLightData In)
+{
+	// TODO: Not part of LightGrid ligth data, does it need to be added?
+	float SpecularScale = 1.0f;
+	const uint LightTypeAndPackedShadowMapChannelMask = asuint(In.LightDirectionAndShadowMask.w);
+
+	// bits [17:16] really
+	const uint LightType = LightTypeAndPackedShadowMapChannelMask >> 16;
+
+	FDeferredLightData Out = (FDeferredLightData)0;
+	Out.TranslatedWorldPosition = In.LightPositionAndInvRadius.xyz;
+	Out.InvRadius = In.LightPositionAndInvRadius.w;
+	Out.Color = In.LightColorAndFalloffExponent.xyz;
+	//Out.Color.xy = 0.0f;
+	Out.FalloffExponent = In.LightColorAndFalloffExponent.w;
+	Out.Direction = In.LightDirectionAndShadowMask.xyz;
+	Out.SpotAngles = In.SpotAnglesAndSourceRadiusPacked.xy;
+	Out.SourceRadius = In.SpotAnglesAndSourceRadiusPacked.z;
+	Out.SourceLength = f16tof32(asuint(In.SpotAnglesAndSourceRadiusPacked.w));
+	Out.Tangent = In.LightTangentAndSoftSourceRadius.xyz;
+	Out.SoftSourceRadius = In.LightTangentAndSoftSourceRadius.w;
+	Out.bInverseSquared = Out.FalloffExponent == 0;
+	// Only directional is not radial, and this loop is only non-directional lights (LIGHT_SOURCE_SHAPE > 0 in trad) 
+	Out.bRadialLight = true;
+	// Setting it to true follows how it is setup in standard deferred, where only the directional lights are not doing this (LIGHT_SOURCE_SHAPE > 0)
+	//Out.bSpotLight = true;
+	//Out.bSpotLight = Out.SpotAngles.x > -2.0f;
+	Out.SpecularScale = SpecularScale;
+	// TODO: not sure if this check is even needed, standard deferred always seems to set it to true?
+	Out.bSpotLight = (LightType == LIGHT_TYPE_SPOT);
+	// NOTE: supporting the rect-light code path adds ~30% overhead on NVIDIA (RTX 2060) even when not used! 
+	// Presumed register pressure increase again, not profiled in detail.
+	// Rect lights are now disabled by the sorting, so the shader should not handle them, for now.
+	Out.bRectLight = false;// (LightType == LIGHT_TYPE_RECT);
+	// bits [7:0]
+	Out.ShadowedBits = (LightTypeAndPackedShadowMapChannelMask & 0xFF) != 0 ? 1 : 0;
+	// Static shadowing uses ShadowMapChannel, dynamic shadows are packed into light attenuation using PreviewShadowMapChannel
+	Out.ShadowMapChannelMask = UnpackShadowMapChannelMask(LightTypeAndPackedShadowMapChannelMask);
+	Out.ShadowedBits = 1;
+	Out.HairTransmittance = InitHairTransmittanceData();
+	return Out;
+}
+
+/**
+ * Adds local lighting of simple lights and returns the index at which non-simple light should starts
+ */
+float4 GetSimpleLightLighting(
+	FScreenSpaceData ScreenSpaceData,
+	const FCulledLightsGridData CulledLightGridData,
+	float3 TranslatedWorldPosition,
+	float3 CameraVector,
+	uint EyeIndex,
+	inout uint FirstNonSimpleLightIndex)
+{
+	float4 CompositedLighting = 0;
+	FirstNonSimpleLightIndex = 0;
+
+#if ENABLE_SEPARATE_SIMPLE_LIGHT_LOOP
+
+	// Limit max to ForwardLightData.NumLocalLights.
+	// This prevents GPU hangs when the PS tries to read from uninitialized NumCulledLightsGrid buffer
+	const uint NumLocalLights = min(CulledLightGridData.NumLocalLights, GetNumLocalLights(EyeIndex));
+
+	uint LocalLightListIndex = 0;
+	LOOP
+	for (; LocalLightListIndex < NumLocalLights; LocalLightListIndex++)
+	{
+		const FLocalLightData LocalLight = GetLocalLightData(CulledLightGridData.DataStartIndex + LocalLightListIndex, EyeIndex);
+		if (!LocalLight.bIsSimpleLight)
+		{
+			break;
+		}
+
+		if (!IsLightVisible(LocalLight, TranslatedWorldPosition))
+		{
+			continue;
+		}
+
+		FSimpleDeferredLightData SimpleLightData = (FSimpleDeferredLightData)0;
+		SimpleLightData.TranslatedWorldPosition = LocalLight.LightPositionAndInvRadius.xyz;
+		SimpleLightData.InvRadius = LocalLight.LightPositionAndInvRadius.w;
+		SimpleLightData.Color = LocalLight.LightColorAndFalloffExponent.xyz;
+		SimpleLightData.FalloffExponent = LocalLight.LightColorAndFalloffExponent.w;
+		SimpleLightData.bInverseSquared = SimpleLightData.FalloffExponent == 0;
+		CompositedLighting.xyz += GetSimpleDynamicLighting(TranslatedWorldPosition, CameraVector, ScreenSpaceData.GBuffer.WorldNormal, ScreenSpaceData.AmbientOcclusion, ScreenSpaceData.GBuffer.DiffuseColor, ScreenSpaceData.GBuffer.SpecularColor, ScreenSpaceData.GBuffer.Roughness, SimpleLightData);
+	}
+	FirstNonSimpleLightIndex = LocalLightListIndex;
+
+#endif // ENABLE_SEPARATE_SIMPLE_LIGHT_LOOP
+
+	return CompositedLighting;
+}
+
 /**
  * Adds local lighting using the light grid, does not apply directional lights, as they are done elsewhere.
  * Does not support dynamic shadows, as these require the per-light shadow mask.
  */
-float4 GetLightGridLocalLighting(const FCulledLightsGridData InLightGridData, float3 TranslatedWorldPosition, float3 CameraVector, FScreenSpaceData ScreenSpaceData, float2 ScreenUV, uint EyeIndex, float Dither, uint FirstNonSimpleLightIndex, uint SampleIndex=0, uint TotalSampleCount=0, uint2 PixelCoord=0)
+float4 GetLightGridLocalLighting(
+	FScreenSpaceData ScreenSpaceData, 
+	const FCulledLightsGridData InLightGridData, 
+	float3 TranslatedWorldPosition, 
+	float3 CameraVector, 
+	float2 ScreenUV, 
+	uint EyeIndex, 
+	float Dither, 
+	uint FirstNonSimpleLightIndex, 
+	uint SampleIndex=0, 
+	uint TotalSampleCount=0, 
+	uint2 PixelCoord=0)
 {
 	float3 WorldPosition = TranslatedWorldPosition - LWCHackToFloat(PrimaryView.PreViewTranslation);
 
 	float4 DirectLighting = 0;
 
-	// TODO: Not part of LightGrid ligth data, does it need to be added?
-	float SpecularScale = 1.0f;
-
 	// Limit max to ForwardLightData.NumLocalLights.
 	// This prevents GPU hangs when the PS tries to read from uninitialized NumCulledLightsGrid buffer
 	const uint NumLocalLights = min(InLightGridData.NumLocalLights, GetNumLocalLights(EyeIndex));
@@ -79,56 +231,14 @@ float4 GetLightGridLocalLighting(const FCulledLightsGridData InLightGridData, fl
 			break;
 		}
 
-		FDeferredLightData LightData = (FDeferredLightData)0;
-		LightData.TranslatedWorldPosition = LocalLight.LightPositionAndInvRadius.xyz;
-		LightData.InvRadius = LocalLight.LightPositionAndInvRadius.w;
-
-		// extra-early out since we know light grid is sloppy and all lights in list are radial (have a range)
-		// appears useless
-		float invLightRadiusSq = LightData.InvRadius*LightData.InvRadius;
-		if (length2(TranslatedWorldPosition - LightData.TranslatedWorldPosition) * invLightRadiusSq > 1.0f)
+		if (!IsLightVisible(LocalLight, TranslatedWorldPosition))
 		{
 			continue;
 		}
 
-		LightData.Color = LocalLight.LightColorAndFalloffExponent.xyz;
-		//LightData.Color.xy = 0.0f;
-		LightData.FalloffExponent = LocalLight.LightColorAndFalloffExponent.w;
-		LightData.Direction = LocalLight.LightDirectionAndShadowMask.xyz;
-		LightData.SpotAngles = LocalLight.SpotAnglesAndSourceRadiusPacked.xy;
-		LightData.SourceRadius = LocalLight.SpotAnglesAndSourceRadiusPacked.z;
-		LightData.SourceLength = f16tof32(asuint(LocalLight.SpotAnglesAndSourceRadiusPacked.w));
-		LightData.Tangent = LocalLight.LightTangentAndSoftSourceRadius.xyz;
-		LightData.SoftSourceRadius = LocalLight.LightTangentAndSoftSourceRadius.w;
-		LightData.bInverseSquared = LightData.FalloffExponent == 0;
-		// Only directional is not radial, and this loop is only non-directional lights (LIGHT_SOURCE_SHAPE > 0 in trad) 
-		LightData.bRadialLight = true;
-		// Setting it to true follows how it is setup in standard deferred, where only the directional lights are not doing this (LIGHT_SOURCE_SHAPE > 0)
-		//LightData.bSpotLight = true;
-		//LightData.bSpotLight = LightData.SpotAngles.x > -2.0f;
-		LightData.SpecularScale = SpecularScale;
-
-		uint LightTypeAndPackedShadowMapChannelMask = asuint(LocalLight.LightDirectionAndShadowMask.w);
-
-		// bits [17:16] really
-		uint LightType = LightTypeAndPackedShadowMapChannelMask >> 16;
-
-		// TODO: not sure if this check is even needed, standard deferred always seems to set it to true?
-		LightData.bSpotLight = (LightType == LIGHT_TYPE_SPOT);
-
-		// NOTE: supporting the rect-light code path adds ~30% overhead on NVIDIA (RTX 2060) even when not used! 
-		// Presumed register pressure increase again, not profiled in detail.
-		// Rect lights are now disabled by the sorting, so the shader should not handle them, for now.
-		LightData.bRectLight = false;// (LightType == LIGHT_TYPE_RECT);
-
-		// bits [7:0]
-		LightData.ShadowedBits = (LightTypeAndPackedShadowMapChannelMask & 0xFF) != 0 ? 1 : 0;
-		// Static shadowing uses ShadowMapChannel, dynamic shadows are packed into light attenuation using PreviewShadowMapChannel
-		LightData.ShadowMapChannelMask = UnpackShadowMapChannelMask(LightTypeAndPackedShadowMapChannelMask);
+		FDeferredLightData LightData = ConvertLocalToDeferredLight(LocalLight);
 
 		float4 LightAttenuation = float4(1, 1, 1, 1);
-
-		LightData.ShadowedBits = 1;
 		if (LocalLightListIndex < GetPackedShadowMaskMaxLightCount())
 		{
 			LightAttenuation = UnpackShadowMask(ShadowMask, LocalLightListIndex, Dither).xxxx;
@@ -154,18 +264,11 @@ float4 GetLightGridLocalLighting(const FCulledLightsGridData InLightGridData, fl
 			LightAttenuation.yw = SSSTransmission;
 		}
 
-		LightData.HairTransmittance = InitHairTransmittanceData();
-
 		#if USE_HAIR_LIGHTING
 		LightAttenuation = 1.0f; // TODO: use shadow mask bits (e.g., subpixel shadow mask)
 		if (LightAttenuation.x > 0)
 		{
-			float3 L = float3(0, 0, 1);
-			if (LightData.bRadialLight)
-				L = normalize(LightData.TranslatedWorldPosition - TranslatedWorldPosition);
-			else
-				L = LightData.Direction;
-
+			const float3 L = normalize(LightData.TranslatedWorldPosition - TranslatedWorldPosition);
 			const float3 V = normalize(-CameraVector);
 
 			// Fetch precomputed transmittance
@@ -195,13 +298,12 @@ float4 GetLightGridLocalLighting(const FCulledLightsGridData InLightGridData, fl
 		}
 		#endif
 
-#if 1
+	#if 1
 		// NOTE: uses AO data like tiled deferred (as opposed to forward path)
 		// TODO: Passing in 0,0 as SVPos, what is meaning of this? Only used if 'REFERENCE_QUALITY' is on.
 		float SurfaceShadow = 1.0f;
 		DirectLighting += GetDynamicLighting(TranslatedWorldPosition, CameraVector, ScreenSpaceData.GBuffer, ScreenSpaceData.AmbientOcclusion, ScreenSpaceData.GBuffer.ShadingModelID, LightData,LightAttenuation, Dither, uint2(0, 0), SurfaceShadow);
-
-#else // 1
+	#else // 1
 		FSimpleDeferredLightData SimpleLightData = (FSimpleDeferredLightData)0;
 		SimpleLightData.TranslatedWorldPosition = LightData.TranslatedWorldPosition;
 		SimpleLightData.InvRadius = LightData.InvRadius;
@@ -209,8 +311,7 @@ float4 GetLightGridLocalLighting(const FCulledLightsGridData InLightGridData, fl
 		SimpleLightData.FalloffExponent = LightData.FalloffExponent;
 		SimpleLightData.bInverseSquared = LightData.bInverseSquared;
 		DirectLighting += GetSimpleDynamicLighting(TranslatedWorldPosition, CameraVector, ScreenSpaceData.GBuffer.WorldNormal, ScreenSpaceData.AmbientOcclusion, ScreenSpaceData.GBuffer.DiffuseColor, ScreenSpaceData.GBuffer.SpecularColor, ScreenSpaceData.GBuffer.Roughness, SimpleLightData);
-#endif //
-
+	#endif // STRATA_ENABLED
 	}
 
 	// For debugging
@@ -234,7 +335,7 @@ float4 GetLightGridLocalLighting(const FCulledLightsGridData InLightGridData, fl
 		float3 TranslatedWorldPosition = mul(float4(ScreenPosition * SceneDepth, SceneDepth, 1), PrimaryView.ScreenToTranslatedWorld).xyz;	\
 		float3 CameraVector = normalize(TranslatedWorldPosition - PrimaryView.TranslatedWorldCameraOrigin);	\
 		ScreenSpaceData.GBuffer.ShadingModelID = ShadingModelId;	\
-		CompositedLighting += GetLightGridLocalLighting(LightGridData, TranslatedWorldPosition, CameraVector, ScreenSpaceData, ScreenUV, 0, Dither, FirstNonSimpleLightIndex); \
+		CompositedLighting += GetLightGridLocalLighting(ScreenSpaceData, LightGridData, TranslatedWorldPosition, CameraVector, ScreenUV, 0, Dither, FirstNonSimpleLightIndex); \
 	}	\
 }
 
@@ -247,6 +348,7 @@ void ClusteredShadingPixelShader(
 	out float4 OutColor : SV_Target0)
 {
 	float2 ScreenUV = UVAndScreenPos.xy;
+	float2 ScreenPosition = UVAndScreenPos.zw;
 
 	uint PixelShadingModelID = GetScreenSpaceData(ScreenUV).GBuffer.ShadingModelID;
 
@@ -259,32 +361,12 @@ void ClusteredShadingPixelShader(
 	float2 LocalPosition = SvPosition.xy - View.ViewRectMin.xy;
 	uint GridIndex = ComputeLightGridCellIndex(uint2(LocalPosition.x, LocalPosition.y), SceneDepth);
 	const FCulledLightsGridData CulledLightGridData = GetCulledLightsGrid(GridIndex, EyeIndex);
-#if 0
-	const LightGridData GridData = GetLightGridData(EyeIndex);
-	uint ZSlice = (uint)(max(0, log2(SceneDepth * GridData.LightGridZParams.x + GridData.LightGridZParams.y) * GridData.LightGridZParams.z));
-	ZSlice = min(ZSlice, (uint)(GridData.CulledGridSize.z - 1));
-	uint3 GridCoordinate = uint3(uint2(LocalPosition.x, LocalPosition.y) >> GridData.LightGridPixelSizeShift, ZSlice);
 
-	if (ZSlice % 2)
-	{
-		OutColor = float(ZSlice) / float(GridData.CulledGridSize.z - 1);
-	}
-	else
-	{
-		OutColor = 1.0f - float(ZSlice) / float(GridData.CulledGridSize.z - 1);
-	}
+	// Debug output
+	#if GRID_DEBUG_ENABLE
+	if (GetClusterDebugColor(LocalPosition, SceneDepth, CulledLightGridData, OutColor)) return;
+	#endif
 
-	if (GridCoordinate.x % 2)
-	{
-		OutColor.xyz = 1.0f - OutColor.xyz;
-	}
-	if (GridCoordinate.y % 2)
-	{
-		OutColor.xyz = 1.0f - OutColor.xyz;
-	}
-	OutColor.w = 0.0f;
-	return;
-#endif // 0
 	// NOTE: this early out helps with the case where there are no lights in the grid cell.
 	if (CulledLightGridData.NumLocalLights == 0 || PixelShadingModelID == SHADINGMODELID_UNLIT)
 	{
@@ -293,69 +375,29 @@ void ClusteredShadingPixelShader(
 	}
 
 	float4 CompositedLighting = 0;
-
-#if ENABLE_SEPARATE_SIMPLE_LIGHT_LOOP
-
+	float3 TranslatedWorldPosition = mul(float4(ScreenPosition * SceneDepth, SceneDepth, 1), PrimaryView.ScreenToTranslatedWorld).xyz;
+	float3 CameraVector = normalize(TranslatedWorldPosition - PrimaryView.TranslatedWorldCameraOrigin);
 	uint FirstNonSimpleLightIndex = 0;
-	{
-		float2 ScreenPosition = UVAndScreenPos.zw;
-		FScreenSpaceData ScreenSpaceData = GetScreenSpaceData(ScreenUV);
-		float3 TranslatedWorldPosition = mul(float4(ScreenPosition * SceneDepth, SceneDepth, 1), PrimaryView.ScreenToTranslatedWorld).xyz;
-		float3 CameraVector = normalize(TranslatedWorldPosition - PrimaryView.TranslatedWorldCameraOrigin);
-
-		// Limit max to ForwardLightData.NumLocalLights.
-		// This prevents GPU hangs when the PS tries to read from uninitialized NumCulledLightsGrid buffer
-		const uint NumLocalLights = min(CulledLightGridData.NumLocalLights, GetNumLocalLights(EyeIndex));
-
-		uint LocalLightListIndex = 0;
 
-		LOOP
-		for (; LocalLightListIndex < NumLocalLights; LocalLightListIndex++)
-		{
-			const FLocalLightData LocalLight = GetLocalLightData(CulledLightGridData.DataStartIndex + LocalLightListIndex, EyeIndex);
-			if (!LocalLight.bIsSimpleLight)
-			{
-				break;
-			}
-
-			FSimpleDeferredLightData SimpleLightData = (FSimpleDeferredLightData)0;
-			SimpleLightData.TranslatedWorldPosition = LocalLight.LightPositionAndInvRadius.xyz;
-			SimpleLightData.InvRadius = LocalLight.LightPositionAndInvRadius.w;
-			SimpleLightData.Color = LocalLight.LightColorAndFalloffExponent.xyz;
-			SimpleLightData.FalloffExponent = LocalLight.LightColorAndFalloffExponent.w;
-			SimpleLightData.bInverseSquared = SimpleLightData.FalloffExponent == 0;
-
-			CompositedLighting += GetSimpleDynamicLighting(TranslatedWorldPosition, CameraVector, ScreenSpaceData.GBuffer.WorldNormal, ScreenSpaceData.AmbientOcclusion, ScreenSpaceData.GBuffer.DiffuseColor, ScreenSpaceData.GBuffer.SpecularColor, ScreenSpaceData.GBuffer.Roughness, SimpleLightData);
-		}
-		FirstNonSimpleLightIndex = LocalLightListIndex;
-	}
-#else // !ENABLE_SEPARATE_SIMPLE_LIGHT_LOOP
-	uint FirstNonSimpleLightIndex = 0;
+	// Simple lights
+	CompositedLighting += GetSimpleLightLighting(GetScreenSpaceData(ScreenUV), CulledLightGridData, TranslatedWorldPosition, CameraVector, EyeIndex, FirstNonSimpleLightIndex);
 
-#endif // ENABLE_SEPARATE_SIMPLE_LIGHT_LOOP
+	// Regular lights
 #if USE_PASS_PER_SHADING_MODEL
 
-	GET_LIGHT_GRID_LOCAL_LIGHTING_SINGLE_SM(SHADINGMODELID_DEFAULT_LIT, PixelShadingModelID, CompositedLighting, ScreenUV, CulledLightGridData, Dither, FirstNonSimpleLightIndex);
-	GET_LIGHT_GRID_LOCAL_LIGHTING_SINGLE_SM(SHADINGMODELID_SUBSURFACE, PixelShadingModelID, CompositedLighting, ScreenUV, CulledLightGridData, Dither, FirstNonSimpleLightIndex);
-	GET_LIGHT_GRID_LOCAL_LIGHTING_SINGLE_SM(SHADINGMODELID_PREINTEGRATED_SKIN, PixelShadingModelID, CompositedLighting, ScreenUV, CulledLightGridData, Dither, FirstNonSimpleLightIndex);
-	GET_LIGHT_GRID_LOCAL_LIGHTING_SINGLE_SM(SHADINGMODELID_CLEAR_COAT, PixelShadingModelID, CompositedLighting, ScreenUV, CulledLightGridData, Dither, FirstNonSimpleLightIndex);
-	GET_LIGHT_GRID_LOCAL_LIGHTING_SINGLE_SM(SHADINGMODELID_SUBSURFACE_PROFILE, PixelShadingModelID, CompositedLighting, ScreenUV, CulledLightGridData, Dither, FirstNonSimpleLightIndex);
-	GET_LIGHT_GRID_LOCAL_LIGHTING_SINGLE_SM(SHADINGMODELID_TWOSIDED_FOLIAGE, PixelShadingModelID, CompositedLighting, ScreenUV, CulledLightGridData, Dither, FirstNonSimpleLightIndex);
-	GET_LIGHT_GRID_LOCAL_LIGHTING_SINGLE_SM(SHADINGMODELID_HAIR, PixelShadingModelID, CompositedLighting, ScreenUV, CulledLightGridData, Dither, FirstNonSimpleLightIndex);
-	GET_LIGHT_GRID_LOCAL_LIGHTING_SINGLE_SM(SHADINGMODELID_CLOTH, PixelShadingModelID, CompositedLighting, ScreenUV, CulledLightGridData, Dither, FirstNonSimpleLightIndex);
-	GET_LIGHT_GRID_LOCAL_LIGHTING_SINGLE_SM(SHADINGMODELID_EYE, PixelShadingModelID, CompositedLighting, ScreenUV, CulledLightGridData, Dither, FirstNonSimpleLightIndex);
-	GET_LIGHT_GRID_LOCAL_LIGHTING_SINGLE_SM(SHADINGMODELID_SINGLELAYERWATER, PixelShadingModelID, CompositedLighting, ScreenUV, CulledLightGridData, Dither, FirstNonSimpleLightIndex);
+	GET_LIGHT_GRID_LOCAL_LIGHTING_SINGLE_SM(SHADINGMODELID_DEFAULT_LIT,			PixelShadingModelID, CompositedLighting, ScreenUV, CulledLightGridData, Dither, FirstNonSimpleLightIndex);
+	GET_LIGHT_GRID_LOCAL_LIGHTING_SINGLE_SM(SHADINGMODELID_SUBSURFACE,			PixelShadingModelID, CompositedLighting, ScreenUV, CulledLightGridData, Dither, FirstNonSimpleLightIndex);
+	GET_LIGHT_GRID_LOCAL_LIGHTING_SINGLE_SM(SHADINGMODELID_PREINTEGRATED_SKIN,	PixelShadingModelID, CompositedLighting, ScreenUV, CulledLightGridData, Dither, FirstNonSimpleLightIndex);
+	GET_LIGHT_GRID_LOCAL_LIGHTING_SINGLE_SM(SHADINGMODELID_CLEAR_COAT,			PixelShadingModelID, CompositedLighting, ScreenUV, CulledLightGridData, Dither, FirstNonSimpleLightIndex);
+	GET_LIGHT_GRID_LOCAL_LIGHTING_SINGLE_SM(SHADINGMODELID_SUBSURFACE_PROFILE,	PixelShadingModelID, CompositedLighting, ScreenUV, CulledLightGridData, Dither, FirstNonSimpleLightIndex);
+	GET_LIGHT_GRID_LOCAL_LIGHTING_SINGLE_SM(SHADINGMODELID_TWOSIDED_FOLIAGE,	PixelShadingModelID, CompositedLighting, ScreenUV, CulledLightGridData, Dither, FirstNonSimpleLightIndex);
+	GET_LIGHT_GRID_LOCAL_LIGHTING_SINGLE_SM(SHADINGMODELID_HAIR,				PixelShadingModelID, CompositedLighting, ScreenUV, CulledLightGridData, Dither, FirstNonSimpleLightIndex);
+	GET_LIGHT_GRID_LOCAL_LIGHTING_SINGLE_SM(SHADINGMODELID_CLOTH,				PixelShadingModelID, CompositedLighting, ScreenUV, CulledLightGridData, Dither, FirstNonSimpleLightIndex);
+	GET_LIGHT_GRID_LOCAL_LIGHTING_SINGLE_SM(SHADINGMODELID_EYE,					PixelShadingModelID, CompositedLighting, ScreenUV, CulledLightGridData, Dither, FirstNonSimpleLightIndex);
+	GET_LIGHT_GRID_LOCAL_LIGHTING_SINGLE_SM(SHADINGMODELID_SINGLELAYERWATER,	PixelShadingModelID, CompositedLighting, ScreenUV, CulledLightGridData, Dither, FirstNonSimpleLightIndex);
 	// SHADINGMODELID_THIN_TRANSLUCENT - skipping because it can not be opaque
 #else // !USE_PASS_PER_SHADING_MODEL
-
-	float2 ScreenPosition = UVAndScreenPos.zw;
-	FScreenSpaceData ScreenSpaceData = GetScreenSpaceData(ScreenUV);
-	float3 TranslatedWorldPosition = mul(float4(ScreenPosition * SceneDepth, SceneDepth, 1), PrimaryView.ScreenToTranslatedWorld).xyz;
-	float3 CameraVector = normalize(TranslatedWorldPosition - PrimaryView.TranslatedWorldCameraOrigin);
-
-	uint NumCulledEntryIndex = (ForwardLightData.NumGridCells + GridIndex) * NUM_CULLED_LIGHTS_GRID_STRIDE;
-
-	CompositedLighting += GetLightGridLocalLighting(CulledLightGridData, TranslatedWorldPosition, CameraVector, ScreenSpaceData, ScreenUV, 0, Dither, FirstNonSimpleLightIndex);
+	CompositedLighting += GetLightGridLocalLighting(GetScreenSpaceData(ScreenUV), CulledLightGridData, TranslatedWorldPosition, CameraVector, ScreenUV, 0, Dither, FirstNonSimpleLightIndex);
 #endif // USE_PASS_PER_SHADING_MODEL
 	
 #if !VISUALIZE_LIGHT_CULLING
@@ -412,7 +454,7 @@ void ClusteredShadingPixelShader(
 
 	float4 CompositedLighting = 0;
 	uint FirstNonSimpleLightIndex = 0;
-	CompositedLighting += GetLightGridLocalLighting(CulledLightGridData, TranslatedWorldPosition, CameraVector, ScreenSpaceData, ScreenUV, 0, Dither, FirstNonSimpleLightIndex, SampleIndex, TotalSampleCount, PixelCoord);
+	CompositedLighting += GetLightGridLocalLighting(ScreenSpaceData, CulledLightGridData, TranslatedWorldPosition, CameraVector, ScreenUV, 0, Dither, FirstNonSimpleLightIndex, SampleIndex, TotalSampleCount, PixelCoord);
 #if !VISUALIZE_LIGHT_CULLING
 	CompositedLighting *= View.PreExposure;
 #endif
