Add cluster shading support for Strata.

diff --git a/Engine/Shaders/Private/ClusteredDeferredShadingPixelShader.usf b/Engine/Shaders/Private/ClusteredDeferredShadingPixelShader.usf
--- a/Engine/Shaders/Private/ClusteredDeferredShadingPixelShader.usf
+++ b/Engine/Shaders/Private/ClusteredDeferredShadingPixelShader.usf
@@ -16,6 +16,13 @@
 
 #include "HairStrands/HairStrandsVisibilityCommonStruct.ush" 
 #include "Common.ush"
+
+#if STRATA_ENABLED
+#define STRATA_INLINE_SHADING 0
+#define STRATA_SSS_TRANSMISSION USE_TRANSMISSION
+#include "/Engine/Private/Strata/Strata.ush"
+#endif
+
 #include "DeferredShadingCommon.ush"  
 #include "BRDF.ush"
 #include "ShadingModels.ush"
@@ -23,6 +30,11 @@
 #include "DeferredLightingCommon.ush"
 #include "VirtualShadowMaps/TransmissionCommon.ush"
 
+#if STRATA_ENABLED
+#include "/Engine/Private/Strata/StrataEvaluation.ush"
+#include "/Engine/Private/Strata/StrataDeferredLighting.ush"
+#endif
+
 #if USE_HAIR_LIGHTING == 1
 	#include "HairStrands/HairStrandsCommon.ush"
 	#include "HairStrands/HairStrandsVisibilityCommon.ush"
@@ -33,7 +45,11 @@
 
 // Causes the shader to run one pass over the light list per shading model, which cuts down peak register pressure.
 // Moves some overhead out of the inner loop (testing for shading model) which ought to be a good thing also.
-#define USE_PASS_PER_SHADING_MODEL 1
+#if STRATA_ENABLED
+	#define USE_PASS_PER_SHADING_MODEL 0
+#else
+	#define USE_PASS_PER_SHADING_MODEL 1
+#endif
 
 // Extracts out the simple lights into their own loop that is executed first, independently of shading model.
 #define ENABLE_SEPARATE_SIMPLE_LIGHT_LOOP 0
@@ -144,7 +160,12 @@ FDeferredLightData ConvertLocalToDeferredLight(const FLocalLightData In)
  * Adds local lighting of simple lights and returns the index at which non-simple light should starts
  */
 float4 GetSimpleLightLighting(
+#if STRATA_ENABLED && USE_HAIR_LIGHTING == 0
+	FStrataAddressing StrataAddressing,
+	FStrataPixelHeader StrataPixelHeader,
+#else
 	FScreenSpaceData ScreenSpaceData,
+#endif 
 	const FCulledLightsGridData CulledLightGridData,
 	float3 TranslatedWorldPosition,
 	float3 CameraVector,
@@ -175,16 +196,36 @@ float4 GetSimpleLightLighting(
 			continue;
 		}
 
+	#if STRATA_ENABLED && USE_HAIR_LIGHTING == 0
+		const FDeferredLightData LightData = ConvertLocalToDeferredLight(LocalLight);
+		const FShadowTerms ShadowTerms = { StrataPixelHeader.MaterialAO, 1.0, 1.0, InitHairTransmittanceData() };
+
+		const float3 ToLight = LightData.TranslatedWorldPosition - TranslatedWorldPosition;
+		const float3 L = normalize(ToLight);
+		const float3 V = -CameraVector;
+		const float LightMask = GetLocalLightAttenuation(TranslatedWorldPosition, LightData, ToLight, L);
+
+		CompositedLighting += StrataDeferredLighting(
+			LightData,
+			V,
+			L,
+			ToLight,
+			LightMask,
+			ShadowTerms,
+			Strata.MaterialTextureArray,
+			StrataAddressing,
+			StrataPixelHeader);
+	#else
 		FSimpleDeferredLightData SimpleLightData = (FSimpleDeferredLightData)0;
 		SimpleLightData.TranslatedWorldPosition = LocalLight.LightPositionAndInvRadius.xyz;
 		SimpleLightData.InvRadius = LocalLight.LightPositionAndInvRadius.w;
 		SimpleLightData.Color = LocalLight.LightColorAndFalloffExponent.xyz;
 		SimpleLightData.FalloffExponent = LocalLight.LightColorAndFalloffExponent.w;
 		SimpleLightData.bInverseSquared = SimpleLightData.FalloffExponent == 0;
-		CompositedLighting.xyz += GetSimpleDynamicLighting(TranslatedWorldPosition, CameraVector, ScreenSpaceData.GBuffer.WorldNormal, ScreenSpaceData.AmbientOcclusion, ScreenSpaceData.GBuffer.DiffuseColor, ScreenSpaceData.GBuffer.SpecularColor, ScreenSpaceData.GBuffer.Roughness, SimpleLightData);
+		CompositedLighting += GetSimpleDynamicLighting(TranslatedWorldPosition, CameraVector, ScreenSpaceData.GBuffer.WorldNormal, ScreenSpaceData.AmbientOcclusion, ScreenSpaceData.GBuffer.DiffuseColor, ScreenSpaceData.GBuffer.SpecularColor, ScreenSpaceData.GBuffer.Roughness, SimpleLightData);
+	#endif
 	}
 	FirstNonSimpleLightIndex = LocalLightListIndex;
-
 #endif // ENABLE_SEPARATE_SIMPLE_LIGHT_LOOP
 
 	return CompositedLighting;
@@ -195,7 +236,12 @@ float4 GetSimpleLightLighting(
  * Does not support dynamic shadows, as these require the per-light shadow mask.
  */
 float4 GetLightGridLocalLighting(
+#if STRATA_ENABLED && USE_HAIR_LIGHTING == 0
+	FStrataAddressing StrataAddressing,
+	FStrataPixelHeader StrataPixelHeader,
+#else 
 	FScreenSpaceData ScreenSpaceData, 
+#endif
 	const FCulledLightsGridData InLightGridData, 
 	float3 TranslatedWorldPosition, 
 	float3 CameraVector, 
@@ -298,7 +344,30 @@ float4 GetLightGridLocalLighting(
 		}
 		#endif
 
-	#if 1
+	#if STRATA_ENABLED && USE_HAIR_LIGHTING == 0
+		FShadowTerms ShadowTerms = { StrataPixelHeader.MaterialAO, 1.0, 1.0, InitHairTransmittanceData() };
+
+		const uint FakeShadingModelID = 0;
+		const float FakeContactShadowOpacity = 1.0f;
+		const float SceneDepth = CalcSceneDepth(ScreenUV);
+		const float3 ToLight = LightData.TranslatedWorldPosition - TranslatedWorldPosition;
+		const float3 L = normalize(ToLight);
+		const float3 V = -CameraVector;
+		float4 PrecomputedShadowFactors = StrataReadPrecomputedShadowFactors(StrataPixelHeader, PixelPos, SceneTexturesStruct.GBufferETexture);
+		GetShadowTerms(SceneDepth, PrecomputedShadowFactors, FakeShadingModelID, FakeContactShadowOpacity, LightData, TranslatedWorldPosition, L, LightAttenuation, Dither, ShadowTerms);
+
+		const float LightMask = 1;
+		DirectLighting += StrataDeferredLighting(
+			LightData,
+			V,
+			L,
+			ToLight,
+			LightMask,
+			ShadowTerms,
+			Strata.MaterialTextureArray,
+			StrataAddressing,
+			StrataPixelHeader);
+	#elif 1
 		// NOTE: uses AO data like tiled deferred (as opposed to forward path)
 		// TODO: Passing in 0,0 as SVPos, what is meaning of this? Only used if 'REFERENCE_QUALITY' is on.
 		float SurfaceShadow = 1.0f;
@@ -341,7 +410,54 @@ float4 GetLightGridLocalLighting(
 
 #endif // USE_PASS_PER_SHADING_MODEL
 
-#if !USE_HAIR_LIGHTING
+#if STRATA_ENABLED && USE_HAIR_LIGHTING == 0
+void ClusteredShadingPixelShader(
+	in noperspective float4 UVAndScreenPos : TEXCOORD0,
+	in float4 SvPosition : SV_Position,
+	out float4 OutColor : SV_Target0)
+{
+	uint2 PixelPos = SvPosition.xy;
+	float2 ScreenUV = UVAndScreenPos.xy;
+
+	const float Dither = InterleavedGradientNoise(SvPosition.xy, View.StateFrameIndexMod8);
+	const uint EyeIndex = 0;
+	const float SceneDepth = CalcSceneDepth(ScreenUV);		
+	const float2 LocalPosition = SvPosition.xy - View.ViewRectMin.xy;
+
+	const uint GridIndex = ComputeLightGridCellIndex(uint2(LocalPosition.x, LocalPosition.y), SceneDepth);
+	const FCulledLightsGridData CulledLightGridData = GetCulledLightsGrid(GridIndex, EyeIndex);
+
+	// Debug output
+	#if GRID_DEBUG_ENABLE
+	if (GetClusterDebugColor(LocalPosition, SceneDepth, CulledLightGridData, OutColor) return;
+	#endif
+
+	FStrataAddressing StrataAddressing = GetStrataPixelDataByteOffset(PixelPos, uint2(View.BufferSizeAndInvSize.xy), Strata.MaxBytesPerPixel);
+	FStrataPixelHeader StrataPixelHeader = UnpackStrataHeaderIn(Strata.MaterialTextureArray, StrataAddressing, Strata.TopLayerTexture);
+	
+	float4 CompositedLighting = 0;
+
+	BRANCH
+	if (StrataPixelHeader.BSDFCount > 0 || CulledLightGridData.NumLocalLights == 0)
+	{
+		const float2 ScreenPosition = UVAndScreenPos.zw;
+		const float3 TranslatedWorldPosition = mul(float4(ScreenPosition * SceneDepth, SceneDepth, 1), PrimaryView.ScreenToTranslatedWorld).xyz;
+		const float3 CameraVector = normalize(TranslatedWorldPosition - PrimaryView.TranslatedWorldCameraOrigin);
+		uint FirstNonSimpleLightIndex = 0;
+
+		// Simple lights
+		CompositedLighting += GetSimpleLightLighting(StrataAddressing, StrataPixelHeader, CulledLightGridData, TranslatedWorldPosition, CameraVector, EyeIndex, FirstNonSimpleLightIndex);
+
+		// Regular lights
+		CompositedLighting += GetLightGridLocalLighting(StrataAddressing, StrataPixelHeader, CulledLightGridData, TranslatedWorldPosition, CameraVector, ScreenUV, EyeIndex, Dither, FirstNonSimpleLightIndex);
+	}
+	
+#if !VISUALIZE_LIGHT_CULLING
+	CompositedLighting *= View.PreExposure;
+#endif
+	OutColor = CompositedLighting;
+}
+#elif !USE_HAIR_LIGHTING
 void ClusteredShadingPixelShader(
 	in noperspective float4 UVAndScreenPos : TEXCOORD0,
 	in float4 SvPosition : SV_Position,
diff --git a/Engine/Shaders/Private/DeferredLightPixelShaders.usf b/Engine/Shaders/Private/DeferredLightPixelShaders.usf
--- a/Engine/Shaders/Private/DeferredLightPixelShaders.usf
+++ b/Engine/Shaders/Private/DeferredLightPixelShaders.usf
@@ -261,7 +261,7 @@ void DeferredLightPixelMain(
 			float Dither = InterleavedGradientNoise(InputParams.PixelPos, View.StateFrameIndexMod8);
 			const uint FakeShadingModelID = 0;
 			const float FakeContactShadowOpacity = 1.0f;
-			float4 PrecomputedShadowFactors = StrataReadPrecomputedShadowFactors(StrataPixelHeader, PixelPos);
+			float4 PrecomputedShadowFactors = StrataReadPrecomputedShadowFactors(StrataPixelHeader, PixelPos, SceneTexturesStruct.GBufferETexture);
 			GetShadowTerms(SceneDepth, PrecomputedShadowFactors, FakeShadingModelID, FakeContactShadowOpacity,
 				LightData, DerivedParams.TranslatedWorldPosition, L, LightAttenuation, Dither, ShadowTerms);
 
diff --git a/Engine/Shaders/Private/Strata/StrataDeferredLighting.ush b/Engine/Shaders/Private/Strata/StrataDeferredLighting.ush
--- a/Engine/Shaders/Private/Strata/StrataDeferredLighting.ush
+++ b/Engine/Shaders/Private/Strata/StrataDeferredLighting.ush
@@ -12,6 +12,20 @@
 #define STRATA_SINGLEPATH 0
 #endif
 
+float4 StrataReadPrecomputedShadowFactors(FStrataPixelHeader StrataPixelHeader, int2 PixelPos, Texture2D PrecomputedShadowTexture)
+{
+	if (HasPrecShadowMask(StrataPixelHeader))
+	{
+	#if ALLOW_STATIC_LIGHTING
+		float4 GBufferE = PrecomputedShadowTexture.Load(int3(PixelPos, 0));
+	#else
+		float4 GBufferE = 1;
+	#endif
+		return GBufferE;
+	}
+	return HasZeroPrecShadowMask(StrataPixelHeader) ? 0.0f : 1.0f;
+}
+
 // Analytical lighting evaluation for strata material.
 // Unpack BSDF on-the-fly
 float4 StrataDeferredLighting(
diff --git a/Engine/Source/Runtime/Renderer/Private/ClusteredDeferredShadingPass.cpp b/Engine/Source/Runtime/Renderer/Private/ClusteredDeferredShadingPass.cpp
--- a/Engine/Source/Runtime/Renderer/Private/ClusteredDeferredShadingPass.cpp
+++ b/Engine/Source/Runtime/Renderer/Private/ClusteredDeferredShadingPass.cpp
@@ -102,6 +102,7 @@ class FClusteredShadingPS : public FGlobalShader
 		SHADER_PARAMETER_STRUCT_INCLUDE(ShaderDrawDebug::FShaderParameters, ShaderDrawParameters)
 		SHADER_PARAMETER_STRUCT_INCLUDE(ShaderPrint::FShaderParameters, ShaderPrintUniformBuffer)
 		SHADER_PARAMETER_RDG_UNIFORM_BUFFER(FHairStrandsViewUniformParameters, HairStrands)
+		SHADER_PARAMETER_RDG_UNIFORM_BUFFER(FStrataGlobalUniformParameters, Strata)
 		SHADER_PARAMETER_RDG_UNIFORM_BUFFER(FSceneTextureUniformParameters, SceneTextures)
 		SHADER_PARAMETER_RDG_TEXTURE(Texture2D, ShadowMaskBits)
 		SHADER_PARAMETER_STRUCT_INCLUDE(FVirtualShadowMapSamplingParameters, VirtualShadowMapSamplingParameters)
@@ -121,6 +122,7 @@ class FClusteredShadingPS : public FGlobalShader
 		FGlobalShader::ModifyCompilationEnvironment(Parameters, OutEnvironment);
 		FForwardLightingParameters::ModifyCompilationEnvironment(Parameters.Platform, OutEnvironment);
 		FVirtualShadowMapArray::SetShaderDefines(OutEnvironment);
+		OutEnvironment.SetDefine(TEXT("STRATA_ENABLED"), Strata::IsStrataEnabled() ? 1u : 0u);
 	}
 };
 
@@ -154,7 +156,7 @@ static void InternalAddClusteredDeferredShadingPass(
 	PassParameters->ShadowMaskBits = ShadowMaskBits;
 	PassParameters->VirtualShadowMapSamplingParameters = VirtualShadowMapArray.GetSamplingParameters(GraphBuilder);
 	PassParameters->HairTransmittanceBuffer = HairTransmittanceBuffer;
-
+	PassParameters->Strata = Strata::BindStrataGlobalUniformParameters(View.StrataSceneData);
 	ShaderDrawDebug::SetParameters(GraphBuilder, View.ShaderDrawData, PassParameters->ShaderDrawParameters);
 	ShaderPrint::SetParameters(GraphBuilder, View, PassParameters->ShaderPrintUniformBuffer);
 
@@ -164,8 +166,9 @@ static void InternalAddClusteredDeferredShadingPass(
 		PassParameters->RenderTargets[0] = FRenderTargetBinding(View.HairStrandsViewData.VisibilityData.SampleLightingTexture, ERenderTargetLoadAction::ELoad);
 	}
 
+	// STRATA_TODO: use classification tile
 	GraphBuilder.AddPass(
-		RDG_EVENT_NAME("ClusteredDeferredShading(%s), #Lights: %d", bHairStrands ? TEXT("HairStrands") : TEXT("GBuffer"), SortedLightsSet.ClusteredSupportedEnd),
+		RDG_EVENT_NAME("Light::ClusteredDeferredShading(%s,#Lights:%d)", bHairStrands ? TEXT("HairStrands") : (Strata::IsStrataEnabled() ? TEXT("Strata") : TEXT("GBuffer")), SortedLightsSet.ClusteredSupportedEnd),
 		PassParameters,
 		ERDGPassFlags::Raster,
 		[PassParameters, &View, SceneTextureExtent, bHairStrands](FRHICommandListImmediate& InRHICmdList)
