Add Strata support for shadow & ambient occlusion.

diff --git a/Engine/Shaders/Private/RayTracing/RayTracingAmbientOcclusionRGS.usf b/Engine/Shaders/Private/RayTracing/RayTracingAmbientOcclusionRGS.usf
--- a/Engine/Shaders/Private/RayTracing/RayTracingAmbientOcclusionRGS.usf
+++ b/Engine/Shaders/Private/RayTracing/RayTracingAmbientOcclusionRGS.usf
@@ -6,6 +6,10 @@
 #include "../SceneTextureParameters.ush"
 #include "../PathTracing/Utilities/PathTracingRandomSequence.ush"
 #include "../SphericalGaussian.ush"
+#if STRATA_ENABLED
+#define STRATA_INLINE_SHADING 0
+#include "../Strata/Strata.ush"
+#endif
 #include "RayTracingCommon.ush"
 #include "RayTracingDeferredShadingCommon.ush"
 
@@ -85,31 +89,32 @@ float3 SampleNeightborPosition(uint2 PixelCoord, int2 Offset)
 
 RAY_TRACING_ENTRY_RAYGEN(AmbientOcclusionRGS)
 {
-	uint2 PixelCoord = DispatchRaysIndex().xy + View.ViewRectMin;
+	const uint2 PixelCoord = DispatchRaysIndex().xy + View.ViewRectMin.xy;
 
-	// Get G-Buffer surface data
-	float2 InvBufferSize = View.BufferSizeAndInvSize.zw;
-	float2 UV = (float2(PixelCoord) + 0.5) * InvBufferSize;
+	#if STRATA_ENABLED
+	const FStrataTopLayerData TopLayerData = StrataUnpackTopLayerData(Strata.TopLayerTexture.Load(uint3(PixelCoord, 0)));
+	const bool bIsValid = TopLayerData.bIsValid;
+	const bool bReconstructGeometryNormal = true; // STRATA_TODO: legacy avoids this on foliage, how do we do that?
+	const float3 WorldNormal = TopLayerData.WorldNormal;
+	#else
+	const FGBufferData GBufferData = GetGBufferDataFromSceneTexturesLoad(PixelCoord);
+	const bool bIsValid = GBufferData.ShadingModelID != SHADINGMODELID_UNLIT;
+	const bool bReconstructGeometryNormal = GBufferData.ShadingModelID != SHADINGMODELID_TWOSIDED_FOLIAGE;
+	const float3 WorldNormal = GBufferData.WorldNormal;
+	#endif
 
-#if 0
-	FGBufferData GBufferData = GetGBufferDataFromSceneTextures(UV);
-#else
-	//#dxr-todo: workaround for flickering. UE-87281
-	FGBufferData GBufferData = GetGBufferDataFromSceneTexturesLoad(PixelCoord);
-#endif
+	const float DeviceZ = SceneDepthTexture.Load(int3(PixelCoord, 0)).r;
+	const float SceneDepth = ConvertFromDeviceZ(DeviceZ);
 
-	float Depth = GBufferData.Depth;
-	float DeviceZ = SceneDepthTexture.Load(int3(PixelCoord, 0)).r;
 	float3 WorldPosition;
 	float3 CameraDirection;
 	ReconstructWorldPositionAndCameraDirectionFromDeviceZ(PixelCoord, DeviceZ, WorldPosition, CameraDirection);
-	float3 WorldNormal = GBufferData.WorldNormal;
 
 	float3 GeometricNormal;
 	float ShadingDotGeometric;
-	if (GBufferData.ShadingModelID != SHADINGMODELID_TWOSIDED_FOLIAGE && CONFIG_SHOOT_WITH_GEOMETRIC_NORMAL)
+	if (bReconstructGeometryNormal && CONFIG_SHOOT_WITH_GEOMETRIC_NORMAL)
 	{
-		float WorldPixelRadius = ComputeWorldBluringRadiusCausedByPixelSize(Depth);
+		float WorldPixelRadius = ComputeWorldBluringRadiusCausedByPixelSize(SceneDepth);
 
 		float3 E = SampleNeightborPosition(PixelCoord, int2( 1,  0)) - WorldPosition;
 		float3 W = SampleNeightborPosition(PixelCoord, int2(-1,  0)) - WorldPosition;
@@ -160,7 +165,7 @@ RAY_TRACING_ENTRY_RAYGEN(AmbientOcclusionRGS)
 	
 	// Mask out depth values beyond far plane
 	bool IsFiniteDepth = DeviceZ > 0.0;
-	bool bTraceRay = (IsFiniteDepth && GBufferData.ShadingModelID != SHADINGMODELID_UNLIT);
+	bool bTraceRay = (IsFiniteDepth && bIsValid);
 	if (!bTraceRay)
 	{
 		Visibility = 1.0;
diff --git a/Engine/Shaders/Private/RayTracing/RayTracingOcclusionRGS.usf b/Engine/Shaders/Private/RayTracing/RayTracingOcclusionRGS.usf
--- a/Engine/Shaders/Private/RayTracing/RayTracingOcclusionRGS.usf
+++ b/Engine/Shaders/Private/RayTracing/RayTracingOcclusionRGS.usf
@@ -11,6 +11,10 @@
 #include "../ScreenSpaceDenoise/SSDPublic.ush"
 #include "../TransmissionCommon.ush"
 #include "../HairStrands/HairStrandsVisibilityCommon.ush"
+#if STRATA_ENABLED
+#define STRATA_INLINE_SHADING 0
+#include "../Strata/Strata.ush"
+#endif
 #include "RayTracingCommon.ush"
 #include "RayTracingDeferredShadingCommon.ush"
 #include "RayTracingDirectionalLight.ush"
@@ -305,12 +309,21 @@ RayDesc Invert(in RayDesc Ray)
 	return InvertedRay;
 }
 
+struct FOcclusionShadingParameters
+{
+	bool bIsValid;
+	bool bBackfaceCulling;
+	bool bIsHair;
+	bool bIsSubsurface;
+	bool bHasSubsurfaceProfile;
+};
+
 FOcclusionResult ComputeOcclusion(
 	const uint2 PixelCoord,
-	const uint ShadingModelID,
 	const uint RaytracingMask,
 	const float DeviceZ,
 	float3 WorldNormal,
+	const FOcclusionShadingParameters ShadingDesc,
 	const FLightShaderParameters LightParameters,
 	const FTransmissionProfileParams TransmissionProfileParams,
 	const uint LocalSamplesPerPixel)
@@ -331,8 +344,7 @@ FOcclusionResult ComputeOcclusion(
 	// toward the view point (this create issue for transmitted light)
 	// It seems having proper transmission is more important, thus the light 
 	// version is enabled by default
-	bool bApplyNormalCulling = true;
-	if (ShadingModelID == SHADINGMODELID_HAIR)
+	if (ShadingDesc.bIsHair)
 	{
 	#if LIGHT_TYPE == LIGHT_TYPE_DIRECTIONAL
 		const float3 LightDirection = LightParameters.Direction;
@@ -340,13 +352,10 @@ FOcclusionResult ComputeOcclusion(
 		const float3 LightDirection = normalize(LightParameters.Position - WorldPosition);
 	#endif
 		WorldNormal = LightDirection;
-		bApplyNormalCulling = false;
-	}
-	else if (ShadingModelID == SHADINGMODELID_EYE || ShadingModelID == SHADINGMODELID_TWOSIDED_FOLIAGE || ShadingModelID == SHADINGMODELID_SUBSURFACE)
-	{
-		bApplyNormalCulling = false;
 	}
 
+	bool bApplyNormalCulling = ShadingDesc.bBackfaceCulling;
+
 	// Disable normal culling for RectLights
 	if (LIGHT_TYPE == LIGHT_TYPE_RECT)
 	{
@@ -429,7 +438,7 @@ FOcclusionResult ComputeOcclusion(
 		}
 
 	#if USE_HAIR_LIGHTING
-		if (ShadingModelID == SHADINGMODELID_HAIR)
+		if (ShadingDesc.bIsHair)
 		{
 			// Denoiser mode 0 doesn't use depth, so take first hit
 			RayFlags |= RAY_FLAG_SKIP_CLOSEST_HIT_SHADER | RAY_FLAG_ACCEPT_FIRST_HIT_AND_END_SEARCH;
@@ -446,7 +455,7 @@ FOcclusionResult ComputeOcclusion(
 			Ray);
 
 		#if USE_HAIR_LIGHTING
-		if (ShadingModelID != SHADINGMODELID_HAIR && bUseHairVoxel)
+		if (!ShadingDesc.bIsHair && bUseHairVoxel)
 		{
 			MinimalPayload.HitT = TraverseHair(PixelCoord, RandSequence, Ray.Origin, Ray.Direction, MinimalPayload.HitT, VirtualVoxel.Raytracing_ShadowOcclusionThreshold);
 		}
@@ -463,7 +472,7 @@ FOcclusionResult ComputeOcclusion(
 			Out.SumRayDistance += HitT;
 			Out.HitCount += 1.0;
 
-			if (ShadingModelID == SHADINGMODELID_SUBSURFACE || ShadingModelID == SHADINGMODELID_HAIR)
+			if (ShadingDesc.bIsSubsurface || ShadingDesc.bIsHair)
 			{
 				// Reverse the ray to support sub-surface casts
 				RayDesc InvertedRay = Invert(Ray);
@@ -480,7 +489,7 @@ FOcclusionResult ComputeOcclusion(
 					float3 HitPosition = InvertedRay.Origin + InvertedRay.Direction * SubsurfacePayload.HitT;
 					float Opacity = TransmissionProfileParams.ExtinctionScale;
 					float Density = -.05f * log(1 - min(Opacity, .999f));
-					if (ShadingModelID == SHADINGMODELID_HAIR)
+					if (ShadingDesc.bIsHair)
 					{
 						Opacity = 1;
 						Density = 1;
@@ -513,7 +522,7 @@ FOcclusionResult ComputeOcclusion(
 	while (0);
 #endif // ENABLE_MULTIPLE_SAMPLES_PER_PIXEL
 
-	if (ENABLE_TRANSMISSION && LocalSamplesPerPixel > 0 && ShadingModelID == SHADINGMODELID_SUBSURFACE_PROFILE)
+	if (ENABLE_TRANSMISSION && LocalSamplesPerPixel > 0 && ShadingDesc.bHasSubsurfaceProfile)
 	{
 		RandomSequence RandSequence;
 		RandomSequence_Initialize(RandSequence, PixelCoord, 0, TimeSeed, 1);
@@ -536,7 +545,7 @@ FOcclusionResult ComputeOcclusion(
 		// Conversion to raster model: 1.0 - (distance / SSSS_MAX_TRANSMISSION_PROFILE_DISTANCE)
 		Out.TransmissionDistance = saturate(EncodeThickness(Out.TransmissionDistance / (SSSS_MAX_TRANSMISSION_PROFILE_DISTANCE)));
 	}
-	else if (ShadingModelID == SHADINGMODELID_SUBSURFACE || ShadingModelID == SHADINGMODELID_HAIR)
+	else if (ShadingDesc.bIsSubsurface || ShadingDesc.bIsHair)
 	{
 		float Depth = length(WorldPosition - LWCHackToFloat(PrimaryView.WorldCameraOrigin));
 		float Range = LODTransitionEnd - LODTransitionStart;
@@ -547,7 +556,7 @@ FOcclusionResult ComputeOcclusion(
 		}
 		Out.TransmissionDistance = (LocalSamplesPerPixel > 0) ? Out.TransmissionDistance / LocalSamplesPerPixel : Out.TransmissionDistance;
 	}
-	else // if ((ShadingModelID == SHADINGMODELID_EYE) || (ShadingModelID == SHADINGMODELID_TWOSIDED_FOLIAGE))
+	else // Eye/Foliage (SHADINGMODELID_EYE, SHADINGMODELID_TWOSIDED_FOLIAGE)
 	{
 		Out.TransmissionDistance = (LocalSamplesPerPixel > 0) ? Out.Visibility / LocalSamplesPerPixel : Out.Visibility;
 	}
@@ -556,7 +565,7 @@ FOcclusionResult ComputeOcclusion(
 
 RAY_TRACING_ENTRY_RAYGEN(OcclusionRGS)
 {
-	uint2 PixelCoord = DispatchRaysIndex().xy + View.ViewRectMin.xy + PixelOffset;
+	const uint2 PixelCoord = DispatchRaysIndex().xy + View.ViewRectMin.xy + PixelOffset;
 
 	FOcclusionResult Occlusion = InitOcclusionResult();
 	FOcclusionResult HairOcclusion = InitOcclusionResult();
@@ -572,42 +581,65 @@ RAY_TRACING_ENTRY_RAYGEN(OcclusionRGS)
 
 		FLightShaderParameters LightParameters = GetRootLightShaderParameters(PrimaryView.PreViewTranslation);
 
-		// Get G-Buffer surface data
-		float2 InvBufferSize = View.BufferSizeAndInvSize.zw;
-		float2 BufferUV = (float2(PixelCoord) + 0.5) * InvBufferSize;
+		// Read material data
 		float3 WorldNormal = 0;
-		uint ShadingModelID = SHADINGMODELID_UNLIT;
-		FTransmissionProfileParams TransmissionProfileParams = (FTransmissionProfileParams)0;
+		FOcclusionShadingParameters ShadingParameters = (FOcclusionShadingParameters)0;
+		FTransmissionProfileParams TransmissionParameters = (FTransmissionProfileParams)0;
+		#if STRATA_ENABLED
+		{
+			FStrataAddressing StrataAddressing = GetStrataPixelDataByteOffset(PixelCoord, uint2(View.BufferSizeAndInvSize.xy), Strata.MaxBytesPerPixel);
+			const FStrataPixelHeader StrataPixelHeader = UnpackStrataHeaderIn(Strata.MaterialTextureArray, StrataAddressing, Strata.TopLayerTexture);
+			const FStrataTopLayerData TopLayerData = StrataUnpackTopLayerData(Strata.TopLayerTexture.Load(uint3(PixelCoord, 0)));
+			const FStrataSubsurfaceData SSSData = StrataUnpackSSSData(Strata.SSSTexture.Load(uint3(PixelCoord, 0)));
+
+			ShadingParameters.bIsValid = TopLayerData.bIsValid;
+			ShadingParameters.bIsHair = StrataHasShadingModel(StrataPixelHeader, STRATA_BSDF_TYPE_HAIR);
+			ShadingParameters.bIsSubsurface = SSSData.bIsValid && !SSSData.bIsProfile;
+			ShadingParameters.bHasSubsurfaceProfile = SSSData.bIsValid && SSSData.bIsProfile;
+			ShadingParameters.bBackfaceCulling = any(SSSData.DMFP > 0); // STRATA_TODO
+
+			WorldNormal = TopLayerData.WorldNormal;
+
+			if (ENABLE_TRANSMISSION && ShadingParameters.bHasSubsurfaceProfile)
+			{
+				TransmissionParameters = GetTransmissionProfileParams(SSSData.ProfileId);
+			}
+
+			if (ShadingParameters.bIsSubsurface)
+			{
+				TransmissionParameters.ExtinctionScale = StrataGetOpacityFromSubSurface(SSSData);
+			}
+		}	
+		#else
 		{
-#if 0
-			FGBufferData GBufferData = GetGBufferDataFromSceneTextures(UV);
-#else
-			//#dxr-todo: workaround for flickering. UE-87281
 			FGBufferData GBufferData = GetGBufferDataFromSceneTexturesLoad(PixelCoord);
-#endif
-			ShadingModelID = GBufferData.ShadingModelID;
-			WorldNormal = GBufferData.WorldNormal;
+			ShadingParameters.bIsValid = GBufferData.ShadingModelID != SHADINGMODELID_UNLIT;
+			ShadingParameters.bIsHair = GBufferData.ShadingModelID == SHADINGMODELID_HAIR;
+			ShadingParameters.bIsSubsurface = GBufferData.ShadingModelID == SHADINGMODELID_SUBSURFACE;
+			ShadingParameters.bHasSubsurfaceProfile = GBufferData.ShadingModelID == SHADINGMODELID_SUBSURFACE_PROFILE;
+			ShadingParameters.bBackfaceCulling =
+				GBufferData.ShadingModelID == SHADINGMODELID_EYE ||
+				GBufferData.ShadingModelID == SHADINGMODELID_TWOSIDED_FOLIAGE ||
+				GBufferData.ShadingModelID == SHADINGMODELID_SUBSURFACE;
+
+			WorldNormal	= GBufferData.WorldNormal;
+
 			if (ENABLE_TRANSMISSION)
 			{
-				TransmissionProfileParams = GetTransmissionProfileParams(GBufferData);
+				TransmissionParameters = GetTransmissionProfileParams(GBufferData);
 			}
 
-			if (ShadingModelID == SHADINGMODELID_SUBSURFACE)
+			if (ShadingParameters.bIsSubsurface)
 			{
-				TransmissionProfileParams.ExtinctionScale = GBufferData.CustomData.a;
+				TransmissionParameters.ExtinctionScale = GBufferData.CustomData.a;
 			}
 		}
+		#endif
 
 		// Mask out depth values that are infinitely far away
-
-#if 0
-		float DeviceZ = SceneDepthTexture.SampleLevel(SceneDepthTextureSampler, BufferUV, 0).r;
-#else
-		//#dxr-todo: workaround for flickering. UE-87281
 		float DeviceZ = SceneDepthTexture.Load(int3(PixelCoord, 0)).r;
-#endif
 		const bool bIsDepthValid = SceneDepthTexture.Load(int3(PixelCoord, 0)).r > 0.0;
-		const bool bIsValidPixel = ShadingModelID != SHADINGMODELID_UNLIT && bIsDepthValid;
+		const bool bIsValidPixel = ShadingParameters.bIsValid && bIsDepthValid;
 		const uint LightChannel = GetSceneLightingChannel(PixelCoord);
 		const bool bTraceRay = bIsValidPixel && (LightChannel & LightingChannelMask) != 0;
 		if (!bTraceRay)
@@ -623,25 +655,26 @@ RAY_TRACING_ENTRY_RAYGEN(OcclusionRGS)
 
 		Occlusion = ComputeOcclusion(
 			PixelCoord,
-			ShadingModelID,
 			RAY_TRACING_MASK_SHADOW | RAY_TRACING_MASK_THIN_SHADOW,
 			DeviceZ,
 			WorldNormal,
+			ShadingParameters,
 			LightParameters, 
-			TransmissionProfileParams,
+			TransmissionParameters,
 			LocalSamplesPerPixel);
 
 		#if USE_HAIR_LIGHTING
 		if (bTraceHairRay)
 		{
+			ShadingParameters.bIsHair = true;
 			HairOcclusion = ComputeOcclusion(
 				PixelCoord,
-				SHADINGMODELID_HAIR,
 				RAY_TRACING_MASK_SHADOW,
 				HairDeviceZ,
 				float3(1,0,0),
+				ShadingParameters,
 				LightParameters,
-				TransmissionProfileParams,
+				TransmissionParameters,
 				RequestedSamplePerPixel);
 		}
 		#endif
diff --git a/Engine/Source/Runtime/Renderer/Private/RayTracing/RayTracingAmbientOcclusion.cpp b/Engine/Source/Runtime/Renderer/Private/RayTracing/RayTracingAmbientOcclusion.cpp
--- a/Engine/Source/Runtime/Renderer/Private/RayTracing/RayTracingAmbientOcclusion.cpp
+++ b/Engine/Source/Runtime/Renderer/Private/RayTracing/RayTracingAmbientOcclusion.cpp
@@ -88,6 +88,12 @@ class FRayTracingAmbientOcclusionRGS : public FGlobalShader
 		return ShouldCompileRayTracingShadersForProject(Parameters.Platform);
 	}
 
+	static void ModifyCompilationEnvironment(const FGlobalShaderPermutationParameters& Parameters, FShaderCompilerEnvironment& OutEnvironment)
+	{
+		FGlobalShader::ModifyCompilationEnvironment(Parameters, OutEnvironment);
+		OutEnvironment.SetDefine(TEXT("STRATA_ENABLED"), Strata::IsStrataEnabled());
+	}
+
 	BEGIN_SHADER_PARAMETER_STRUCT(FParameters, )
 		SHADER_PARAMETER(int, SamplesPerPixel)
 		SHADER_PARAMETER(float, MaxRayDistance)
@@ -96,9 +102,9 @@ class FRayTracingAmbientOcclusionRGS : public FGlobalShader
 		SHADER_PARAMETER_SRV(RaytracingAccelerationStructure, TLAS)
 		SHADER_PARAMETER_RDG_TEXTURE_UAV(RWTexture2D<float>, RWAmbientOcclusionMaskUAV)
 		SHADER_PARAMETER_RDG_TEXTURE_UAV(RWTexture2D<float>, RWAmbientOcclusionHitDistanceUAV)
-
 		SHADER_PARAMETER_STRUCT_REF(FViewUniformShaderParameters, ViewUniformBuffer)
 		SHADER_PARAMETER_STRUCT_INCLUDE(FSceneTextureParameters, SceneTextures)
+		SHADER_PARAMETER_RDG_UNIFORM_BUFFER(FStrataGlobalUniformParameters, Strata)
 	END_SHADER_PARAMETER_STRUCT()
 };
 
@@ -159,6 +165,7 @@ void FDeferredShadingSceneRenderer::RenderRayTracingAmbientOcclusion(
 	PassParameters->Intensity = View.FinalPostProcessSettings.RayTracingAOIntensity;
 	PassParameters->MaxNormalBias = GetRaytracingMaxNormalBias();
 	PassParameters->TLAS = View.GetRayTracingSceneViewChecked();
+	PassParameters->Strata = Strata::BindStrataGlobalUniformParameters(View.StrataSceneData);
 	PassParameters->RWAmbientOcclusionMaskUAV = GraphBuilder.CreateUAV(DenoiserInputs.Mask);
 	PassParameters->RWAmbientOcclusionHitDistanceUAV = GraphBuilder.CreateUAV(DenoiserInputs.RayHitDistance);
 	PassParameters->ViewUniformBuffer = View.ViewUniformBuffer;
diff --git a/Engine/Source/Runtime/Renderer/Private/RayTracing/RayTracingShadows.cpp b/Engine/Source/Runtime/Renderer/Private/RayTracing/RayTracingShadows.cpp
--- a/Engine/Source/Runtime/Renderer/Private/RayTracing/RayTracingShadows.cpp
+++ b/Engine/Source/Runtime/Renderer/Private/RayTracing/RayTracingShadows.cpp
@@ -133,6 +133,7 @@ class FOcclusionRGS : public FGlobalShader
 		{
 			OutEnvironment.SetDefine(TEXT("UE_RAY_TRACING_COHERENT_RAYS"), 0);
 		}
+		OutEnvironment.SetDefine(TEXT("STRATA_ENABLED"), Strata::IsStrataEnabled());
 	}
 
 	BEGIN_SHADER_PARAMETER_STRUCT(FParameters, )
@@ -165,6 +166,7 @@ class FOcclusionRGS : public FGlobalShader
 		SHADER_PARAMETER_STRUCT_REF(FViewUniformShaderParameters, ViewUniformBuffer)
 		SHADER_PARAMETER_RDG_UNIFORM_BUFFER(FHairStrandsViewUniformParameters, HairStrands)
 		SHADER_PARAMETER_RDG_UNIFORM_BUFFER(FVirtualVoxelParameters, VirtualVoxel)
+		SHADER_PARAMETER_RDG_UNIFORM_BUFFER(FStrataGlobalUniformParameters, Strata)
 	END_SHADER_PARAMETER_STRUCT()
 };
 
@@ -295,6 +297,7 @@ void FDeferredShadingSceneRenderer::RenderRayTracingShadows(
 		PassParameters->bTransmissionSamplingDistanceCulling = CVarRayTracingTransmissionSamplingDistanceCulling.GetValueOnRenderThread();
 		PassParameters->TransmissionSamplingTechnique = CVarRayTracingTransmissionSamplingTechnique.GetValueOnRenderThread();
 		PassParameters->RejectionSamplingTrials = CVarRayTracingTransmissionRejectionSamplingTrials.GetValueOnRenderThread();
+		PassParameters->Strata = Strata::BindStrataGlobalUniformParameters(View.StrataSceneData);
 		if (bUseHairLighting)
 		{
 			const bool bUseHairVoxel = CVarRayTracingShadowsEnableHairVoxel.GetValueOnRenderThread() > 0;
