* Add initial Strata support for sub-surface profile roughnesses.
 * Move SSS profile roughness transformation to a common place

diff --git a/Engine/Shaders/Private/ShadingModels.ush b/Engine/Shaders/Private/ShadingModels.ush
--- a/Engine/Shaders/Private/ShadingModels.ush
+++ b/Engine/Shaders/Private/ShadingModels.ush
@@ -562,16 +562,12 @@ float ApproximateHG(float cosJ, float g)
 	return (ISO_PHASE_FUNC_Normalized * (1.0f - g2) / max( 1e-5, gcos2));
 }
 
-void GetProfileDualSpecular(FGBufferData GBuffer, out float AverageToRoughness0, out float AverageToRoughness1, out float LobeMix)
+void GetProfileDualSpecular(FGBufferData GBuffer, inout float AverageToRoughness0, inout float AverageToRoughness1, inout float LobeMix)
 {
 #if !FORWARD_SHADING
-	// 0..255, which SubSurface profile to pick
-	uint SubsurfaceProfileInt = ExtractSubsurfaceProfileInt(GBuffer);
-
-	float4 Data = View.SSProfilesTexture.Load(int3(SSSS_DUAL_SPECULAR_OFFSET, SubsurfaceProfileInt, 0));
-	AverageToRoughness0 = Data.x * SSSS_MAX_DUAL_SPECULAR_ROUGHNESS;
-	AverageToRoughness1 = Data.y * SSSS_MAX_DUAL_SPECULAR_ROUGHNESS;
-	LobeMix = Data.z;
+	float DummyOpacity = 0;
+	float DummyMaterialRoughnessToAverage = 0;
+	GetSubsurfaceProfileDualSpecular(ExtractSubsurfaceProfileInt(GBuffer), DummyOpacity, AverageToRoughness0, AverageToRoughness1, LobeMix, DummyMaterialRoughnessToAverage);
 #else
 	// Disable dual lobe, as subsurface profile doesn't work with forward.
 	AverageToRoughness0 = 1.f;
@@ -596,11 +592,8 @@ FDirectLighting SubsurfaceProfileBxDF( FGBufferData GBuffer, half3 N, half3 V, h
 	float Lobe0Roughness = max(saturate(AverageRoughness * AverageToRoughness0), 0.02f);
 	float Lobe1Roughness = saturate(AverageRoughness * AverageToRoughness1);
 
-	// Smoothly lerp to default single GGX lobe as Opacity approaches 0, before reverting to SHADINGMODELID_DEFAULT_LIT.
-	// See SUBSURFACE_PROFILE_OPACITY_THRESHOLD in ShadingModelsMaterial.ush.
-	float Opacity = GBuffer.CustomData.a;
-	Lobe0Roughness = lerp(1.0f, Lobe0Roughness, saturate(Opacity * 10.0f));
-	Lobe1Roughness = lerp(1.0f, Lobe1Roughness, saturate(Opacity * 10.0f));
+	const float Opacity = GBuffer.CustomData.a;
+	GetSubsurfaceProfileRoughnesses(Lobe0Roughness, Lobe1Roughness, Opacity);
 
 	// Take the average roughness instead of compute a separate energy term for each roughness
 	const FBxDFEnergyTerms EnergyTerms = ComputeGGXSpecEnergyTerms(AverageRoughness, Context.NoV, GBuffer.SpecularColor);
diff --git a/Engine/Shaders/Private/ShadingModelsMaterial.ush b/Engine/Shaders/Private/ShadingModelsMaterial.ush
--- a/Engine/Shaders/Private/ShadingModelsMaterial.ush
+++ b/Engine/Shaders/Private/ShadingModelsMaterial.ush
@@ -65,11 +65,12 @@ void SetGBufferForShadingModel(
 
 			// Average roughness for dual specular.
 			uint SubsurfaceProfileUInt = uint(SubsurfaceProfile * 255.0f + 0.5f);
-			float MaterialRoughnessToAverage = GetSubsurfaceProfileTexture(SSSS_DUAL_SPECULAR_OFFSET, SubsurfaceProfileUInt).w * SSSS_MAX_DUAL_SPECULAR_ROUGHNESS;
 
-			// Smooth blend out dual specular when opacity is low, we have the extra SSSS_OPACITY_THRESHOLD_EPS so that we fade out by the time we
-			// get to 0.01, as opposed to 0.0.
-			MaterialRoughnessToAverage = lerp(1.0f, MaterialRoughnessToAverage, saturate((Opacity - SSSS_OPACITY_THRESHOLD_EPS) * 10.0f));
+			float AverageToRoughness0 = 0;
+			float AverageToRoughness1 = 0;
+			float LobeMix = 0;
+			float MaterialRoughnessToAverage = 0;
+			GetSubsurfaceProfileDualSpecular(SubsurfaceProfileUInt, Opacity, AverageToRoughness0, AverageToRoughness1, LobeMix, MaterialRoughnessToAverage);
 
 			GBuffer.Roughness = saturate(GBuffer.Roughness * MaterialRoughnessToAverage);
 		}
diff --git a/Engine/Shaders/Private/Strata/Strata.ush b/Engine/Shaders/Private/Strata/Strata.ush
--- a/Engine/Shaders/Private/Strata/Strata.ush
+++ b/Engine/Shaders/Private/Strata/Strata.ush
@@ -11,6 +11,7 @@
 #include "../ThinFilmBSDF.ush"
 #include "../MortonCode.ush"
 #include "../OctahedralCommon.ush"
+#include "../SubsurfaceProfileCommon.ush"
 #include "/Engine/Shared/StrataDefinitions.h"     
 
 
@@ -584,13 +585,6 @@ struct FStrataData
 	FStrataLayer Layers[STRATA_MAX_LAYER_COUNT];	// Layer 0 is the top layer
 };
 
-uint GetStrataSubsurfaceProfile(float NormalizedId)
-{
-	// can be optimized
-	return uint(NormalizedId * 255.0f + 0.5f);
-
-}
-
 FStrataData GetInitialisedStrataData()
 {
 	FStrataData StrataData = (FStrataData)0;
@@ -622,6 +616,39 @@ FStrataData GetStrataSlabBSDF(
 {
 	FStrataData StrataData = GetInitialisedStrataDataWithReservedFirstSlot();
 
+	// STRATA_TODO apply SSS as a function of pixel footprint w.r.t. MFP
+	SSSDMFPScale = saturate(SSSDMFPScale);
+	const bool bSSSProfilePlugged = SSSProfileID > 0.0f;
+	const bool bDMFPPlugged = any(SSSDMFP > 0);
+	const bool bIsThin = bDMFPPlugged && Thickness < STRATA_LAYER_ISTHIN_THICKNESS_THRESHOLD_CM;
+	const bool bHasSSSProfile = !bIsThin && bSSSProfilePlugged && SSSDMFPScale > 0.f;	// When no profile is provide, we use Burley. STRATA_TODO we should clarify as burley is also an option of the profile.
+	const bool bHasSSS = bDMFPPlugged || bHasSSSProfile;
+
+	// Derive roughness and haziness from profile data
+	if (bSSSProfilePlugged)
+	{
+		// Average roughness for dual specular.
+		const uint SubsurfaceProfileUInt = StrataSubsurfaceProfileIdTo8bits(SSSProfileID);
+
+		float AverageToRoughness0 = 0;
+		float AverageToRoughness1 = 0;
+		float LobeMix = 0;
+		float MaterialRoughnessToAverage = 0;
+		GetSubsurfaceProfileDualSpecular(SubsurfaceProfileUInt, SSSDMFPScale, AverageToRoughness0, AverageToRoughness1, LobeMix, MaterialRoughnessToAverage);
+
+		float AverageRoughness = saturate(Roughness * MaterialRoughnessToAverage); // This is the roughness originallys stored into the gbuffer
+		float Lobe0Roughness   = saturate(AverageRoughness * AverageToRoughness0);
+		float Lobe1Roughness   = saturate(AverageRoughness * AverageToRoughness1);
+
+		GetSubsurfaceProfileRoughnesses(Lobe0Roughness, Lobe1Roughness, SSSDMFPScale);
+
+		// Convertion from roughness0/1 to roughness+haziness
+		// * This encoding is based on how roughness0 + haziness is translated to roughness1, see ComputeHazyLobeRoughness
+		// * The mix factor is hardoded to 0.5 i.e., 0.5 * Fs(roughness0) + 0.5 * Fs(roughness1)
+		Haziness = (Lobe1Roughness - Lobe0Roughness) / (1.f - min(Lobe0Roughness, 0.99f));
+		Roughness = Lobe0Roughness;
+	}
+
 	const uint bHasHaziness = Haziness > 0.0f ? 1 : 0;
 	const uint bHasThinFilm = ThinFilmThickness > 0.0f ? 1 : 0;
 	const uint bHasAnisotropy = Anisotropy != 0 ? 1 : 0;
@@ -631,11 +658,6 @@ FStrataData GetStrataSlabBSDF(
 		StrataRequestSharedLocalBasisTangent(SharedLocalBasisTypes, SharedLocalBasisIndex);
 	}
 
-	// STRATA_TODO apply SSS as a function of pixel footprint w.r.t. MFP
-	const bool bDMFPPlugged = any(SSSDMFP > 0);
-	const bool bIsThin = bDMFPPlugged && Thickness < STRATA_LAYER_ISTHIN_THICKNESS_THRESHOLD_CM;
-	const bool bHasSSSProfile = !bIsThin && SSSProfileID  > 0.0f;	// When no profile is provide, we use Burley. STRATA_TODO we should clarify as burley is also an option of the profile.
-	const bool bHasSSS = bDMFPPlugged || bHasSSSProfile;
 
 	// About the thin, a.k.a. two sided lighting, model:
 	// 	- Thin lighting model is used when thickness becomes less than STRATA_LAYER_ISTHIN_THICKNESS_THRESHOLD_CM
@@ -686,11 +708,11 @@ FStrataData GetStrataSlabBSDF(
 	if (bHasSSSProfile)
 	{
 		SLAB_SSSPROFILEID	(StrataData.Layers[0].BSDFs[0])		= SSSProfileID;
-		SLAB_SSSPROFILERADIUSSCALE(StrataData.Layers[0].BSDFs[0])= saturate(SSSDMFPScale);
+		SLAB_SSSPROFILERADIUSSCALE(StrataData.Layers[0].BSDFs[0])= SSSDMFPScale;
 	}
 	else
 	{
-		SLAB_SSSDMFP	(StrataData.Layers[0].BSDFs[0])			= SSSDMFP * saturate(SSSDMFPScale);
+		SLAB_SSSDMFP	(StrataData.Layers[0].BSDFs[0])			= SSSDMFP * SSSDMFPScale;
 	}
 
 	SLAB_FUZZ_AMOUNT		(StrataData.Layers[0].BSDFs[0])		= FuzzAmount;
@@ -698,7 +720,7 @@ FStrataData GetStrataSlabBSDF(
 
 #if STRATA_INLINE_SHADING
 	// At this stage we do not know if this strata BSDF is going to be using simple volume lighting so always store the DMFP on the side to not override SSSPROFILE data.
-	StrataData.Layers[0].BSDFs[0].TmpDMFP = SSSDMFP * saturate(SSSDMFPScale);
+	StrataData.Layers[0].BSDFs[0].TmpDMFP = SSSDMFP * SSSDMFPScale;
 
 	StrataData.Layers[0].BSDFs[0].Coverage = 1.0f;
 #endif
diff --git a/Engine/Shaders/Private/SubsurfaceProfileCommon.ush b/Engine/Shaders/Private/SubsurfaceProfileCommon.ush
--- a/Engine/Shaders/Private/SubsurfaceProfileCommon.ush
+++ b/Engine/Shaders/Private/SubsurfaceProfileCommon.ush
@@ -48,3 +48,27 @@ half4 GetSubsurfaceProfileTexture(Texture2D InSSProfilesTexture, SamplerState In
 	return InSSProfilesTexture.Load(int3(SampleIndex, SubsurfaceProfileInt, 0));
 #endif 
 }
+
+void GetSubsurfaceProfileDualSpecular(uint SubsurfaceProfileInt, float Opacity, inout float AverageToRoughness0, inout float AverageToRoughness1, inout float LobeMix, inout float MaterialRoughnessToAverage)
+{
+#if COMPILER_GLSL_ES3_1 // Force to use a point sampler for Texture2D.Load on OpenGLES platform
+	float4 Data = View.SSProfilesTexture.SampleLevel(View.SSProfilesSampler, (uint2(SSSS_DUAL_SPECULAR_OFFSET, SubsurfaceProfileInt) + float2(0.5f, 0.5f)) * View.SSProfilesTextureSizeAndInvSize.zw, 0);
+#else
+	float4 Data = View.SSProfilesTexture.Load(int3(SSSS_DUAL_SPECULAR_OFFSET, SubsurfaceProfileInt, 0));
+#endif 
+	AverageToRoughness0 = Data.x * SSSS_MAX_DUAL_SPECULAR_ROUGHNESS;
+	AverageToRoughness1 = Data.y * SSSS_MAX_DUAL_SPECULAR_ROUGHNESS;
+	LobeMix = Data.z;
+
+	// Smooth blend out dual specular when opacity is low, we have the extra SSSS_OPACITY_THRESHOLD_EPS so that we fade out by the time we
+	// get to 0.01, as opposed to 0.0.
+	MaterialRoughnessToAverage = lerp(1.0f, Data.w * SSSS_MAX_DUAL_SPECULAR_ROUGHNESS, saturate((Opacity - SSSS_OPACITY_THRESHOLD_EPS) * 10.0f));
+}
+
+void GetSubsurfaceProfileRoughnesses(inout float Lobe0Roughness, inout float Lobe1Roughness, float Opacity)
+{
+	// Smoothly lerp to default single GGX lobe as Opacity approaches 0, before reverting to SHADINGMODELID_DEFAULT_LIT.
+	// See SUBSURFACE_PROFILE_OPACITY_THRESHOLD in ShadingModelsMaterial.ush.
+	Lobe0Roughness = lerp(1.0f, Lobe0Roughness, saturate(Opacity * 10.0f));
+	Lobe1Roughness = lerp(1.0f, Lobe1Roughness, saturate(Opacity * 10.0f));
+}
\ No newline at end of file
