Officially adds support for alpha channel through TSR

diff --git a/Engine/Shaders/Private/TemporalSuperResolution/TSRUpdateHistory.usf b/Engine/Shaders/Private/TemporalSuperResolution/TSRUpdateHistory.usf
--- a/Engine/Shaders/Private/TemporalSuperResolution/TSRUpdateHistory.usf
+++ b/Engine/Shaders/Private/TemporalSuperResolution/TSRUpdateHistory.usf
@@ -857,10 +857,18 @@ void MainCS(
 
 		// Correct history
 		{
-			PrevHistoryMommentMin *= tsr_half(HistoryPreExposureCorrection);
-			PrevHistoryMommentMax *= tsr_half(HistoryPreExposureCorrection);
-			PrevHighFrequencyColor *= tsr_half(HistoryPreExposureCorrection);
-			PrevLowFrequencyColor *= tsr_half(HistoryPreExposureCorrection);
+			PrevHistoryMommentMin[0] *= tsr_half(HistoryPreExposureCorrection);
+			PrevHistoryMommentMin[1] *= tsr_half(HistoryPreExposureCorrection);
+			PrevHistoryMommentMin[2] *= tsr_half(HistoryPreExposureCorrection);
+			PrevHistoryMommentMax[0] *= tsr_half(HistoryPreExposureCorrection);
+			PrevHistoryMommentMax[1] *= tsr_half(HistoryPreExposureCorrection);
+			PrevHistoryMommentMax[2] *= tsr_half(HistoryPreExposureCorrection);
+			PrevHighFrequencyColor[0] *= tsr_half(HistoryPreExposureCorrection);
+			PrevHighFrequencyColor[1] *= tsr_half(HistoryPreExposureCorrection);
+			PrevHighFrequencyColor[2] *= tsr_half(HistoryPreExposureCorrection);
+			PrevLowFrequencyColor[0] *= tsr_half(HistoryPreExposureCorrection);
+			PrevLowFrequencyColor[1] *= tsr_half(HistoryPreExposureCorrection);
+			PrevLowFrequencyColor[2] *= tsr_half(HistoryPreExposureCorrection);
 
 			// Fixes negative validity that can happen due to negative lobes of Catrom.
 			PrevHistoryValidity = max(PrevHistoryValidity, tsr_half(0.0).xx);
@@ -1382,6 +1390,14 @@ void MainCS(
 				FinalOutputColor[3] = saturate(tsr_half(1.0) - (tsr_half(1.0) - FinalOutputColor[3]) * FinalTranslucencyColor[3]);
 			#endif
 		#endif
+		
+		// Ensure that alpha values that are expected to be opaque (but are only close to opaque) are forced to be opaque.
+		// (0.995 chosen to accomodate handling of 254/255)
+		#if CONFIG_SCENE_COLOR_APLHA
+		{
+			FinalOutputColor[3] = FinalOutputColor[3] > 0.995 ? tsr_half(1.0).xx : FinalOutputColor[3];
+		}
+		#endif
 	}
 
 
diff --git a/Engine/Source/Runtime/Renderer/Private/PostProcess/TemporalAA.cpp b/Engine/Source/Runtime/Renderer/Private/PostProcess/TemporalAA.cpp
--- a/Engine/Source/Runtime/Renderer/Private/PostProcess/TemporalAA.cpp
+++ b/Engine/Source/Runtime/Renderer/Private/PostProcess/TemporalAA.cpp
@@ -990,9 +990,8 @@ EMainTAAPassConfig ITemporalUpscaler::GetMainTAAPassConfig(const FViewInfo& View
 	{
 		return EMainTAAPassConfig::ThirdParty;
 	}
-	else if (View.AntiAliasingMethod == AAM_TSR && !IsPostProcessingWithAlphaChannelSupported())
+	else if (View.AntiAliasingMethod == AAM_TSR)
 	{
-		// TODO(TSR): alpha channel is not supported yet
 		check(SupportsTSR(View.GetShaderPlatform()));
 		return EMainTAAPassConfig::TSR;
 	}
