Made more phase functions common.
Fixed a bad name

diff --git a/Engine/Shaders/Private/ParticipatingMediaCommon.ush b/Engine/Shaders/Private/ParticipatingMediaCommon.ush
--- a/Engine/Shaders/Private/ParticipatingMediaCommon.ush
+++ b/Engine/Shaders/Private/ParticipatingMediaCommon.ush
@@ -83,12 +83,17 @@ float IsotropicPhase()
 // Follows PBRT convention http://www.pbr-book.org/3ed-2018/Volume_Scattering/Phase_Functions.html#PhaseHG
 float HenyeyGreensteinPhase(float G, float CosTheta)
 {
-	return (1.0f - G * G) / (4.0f * PI * pow(1.0f + G * G + 2.0f * G * CosTheta, 1.5f));
+	// Reference implementation (i.e. not schlick approximation). 
+	// See http://www.pbr-book.org/3ed-2018/Volume_Scattering/Phase_Functions.html
+	float Numer = 1.0f - G * G;
+	float Denom = 1.0f + G * G + 2.0f * G * CosTheta;
+	return Numer / (4.0f * PI * Denom * sqrt(Denom));
 }
 
-float RaleighPhase(float CosTheta)
+float RayleighPhase(float CosTheta)
 {
-	return 3.0f * (1.0f + CosTheta * CosTheta) / (16.0f * PI);
+	float Factor = 3.0f / (16.0f * PI);
+	return Factor * (1.0f + CosTheta * CosTheta);
 }
 
 // Schlick phase function approximating henyey-greenstein
diff --git a/Engine/Shaders/Private/SkyAtmosphere.usf b/Engine/Shaders/Private/SkyAtmosphere.usf
--- a/Engine/Shaders/Private/SkyAtmosphere.usf
+++ b/Engine/Shaders/Private/SkyAtmosphere.usf
@@ -490,11 +490,11 @@ SingleScatteringResult IntegrateSingleScatteredLuminance(
 	const float3 wi = Light0Dir;
 	const float3 wo = WorldDir;
 	float cosTheta = dot(wi, wo);
-	float MiePhaseValueLight0 = HgPhase(Atmosphere.MiePhaseG, -cosTheta);	// negate cosTheta because due to WorldDir being a "in" direction. 
+	float MiePhaseValueLight0 = HenyeyGreensteinPhase(Atmosphere.MiePhaseG, -cosTheta);	// negate cosTheta because due to WorldDir being a "in" direction. 
 	float RayleighPhaseValueLight0 = RayleighPhase(cosTheta);
 #if SECOND_ATMOSPHERE_LIGHT_ENABLED
 	cosTheta = dot(Light1Dir, wo);
-	float MiePhaseValueLight1 = HgPhase(Atmosphere.MiePhaseG, -cosTheta);	// negate cosTheta because due to WorldDir being a "in" direction. 
+	float MiePhaseValueLight1 = HenyeyGreensteinPhase(Atmosphere.MiePhaseG, -cosTheta);	// negate cosTheta because due to WorldDir being a "in" direction. 
 	float RayleighPhaseValueLight1 = RayleighPhase(cosTheta);
 #endif
 
diff --git a/Engine/Shaders/Private/SkyAtmosphereCommon.ush b/Engine/Shaders/Private/SkyAtmosphereCommon.ush
--- a/Engine/Shaders/Private/SkyAtmosphereCommon.ush
+++ b/Engine/Shaders/Private/SkyAtmosphereCommon.ush
@@ -7,6 +7,8 @@
 
 #pragma once
 
+#include "ParticipatingMediaCommon.ush"
+
 // Disabled for all mobile platforms, see USkyLightComponent::IsRealTimeCaptureEnabled
 #define AP_SUPPORT_REALTIME_REFLECTION_CAPTURE (FEATURE_LEVEL > FEATURE_LEVEL_ES3_1)
 
@@ -252,21 +254,6 @@ float3 GetLightDiskLuminance(
 // Participating medium properties
 ////////////////////////////////////////////////////////////
 
-float RayleighPhase(float CosTheta)
-{
-	float Factor = 3.0f / (16.0f * PI);
-	return Factor * (1.0f + CosTheta * CosTheta);
-}
-
-float HgPhase(float G, float CosTheta)
-{
-	// Reference implementation (i.e. not schlick approximation). 
-	// See http://www.pbr-book.org/3ed-2018/Volume_Scattering/Phase_Functions.html
-	float Numer = 1.0f - G * G;
-	float Denom = 1.0f + G * G + 2.0f * G * CosTheta;
-	return Numer / (4.0f * PI * Denom * sqrt(Denom));
-}
-
 float3 GetAlbedo(float3 Scattering, float3 Extinction)
 {
 	return Scattering / max(0.001f, Extinction);
diff --git a/Engine/Shaders/Private/VolumetricCloud.usf b/Engine/Shaders/Private/VolumetricCloud.usf
--- a/Engine/Shaders/Private/VolumetricCloud.usf
+++ b/Engine/Shaders/Private/VolumetricCloud.usf
@@ -277,8 +277,8 @@ float SamplePhaseFunction(in float PhaseCosTheta, in float PhaseG, in float Phas
 	PhaseG = clamp(PhaseG, -0.999f, 0.999f);
 	PhaseG2 = clamp(PhaseG2, -0.999f, 0.999f);
 	PhaseBlend = clamp(PhaseBlend, 0.0f, 1.0f);
-	float MiePhaseValueLight0 = HgPhase(PhaseG, -PhaseCosTheta);	// negate cosTheta because due to WorldDir being a "in" direction. 
-	float MiePhaseValueLight1 = HgPhase(PhaseG2, -PhaseCosTheta);
+	float MiePhaseValueLight0 = HenyeyGreensteinPhase(PhaseG, -PhaseCosTheta);	// negate cosTheta because due to WorldDir being a "in" direction. 
+	float MiePhaseValueLight1 = HenyeyGreensteinPhase(PhaseG2, -PhaseCosTheta);
 	const float Phase = MiePhaseValueLight0 + PhaseBlend * (MiePhaseValueLight1 - MiePhaseValueLight0);
 	return Phase;
 }
