PostProcessMaterials should always sample velocity through GBuffer SceneVelocity texture.
Uses  SceneTextureParameters.GBufferVelocityTexture which is always filled in even if velocity is after base pass. (It also contains default black texture if there is no velocity at all rendered for frame.)
Fixes case where PPM failed to sample velocity when velocity pass is during depth.
#rnx

diff --git a/Engine/Shaders/Private/MaterialTemplate.ush b/Engine/Shaders/Private/MaterialTemplate.ush
--- a/Engine/Shaders/Private/MaterialTemplate.ush
+++ b/Engine/Shaders/Private/MaterialTemplate.ush
@@ -2206,11 +2206,7 @@ MaterialFloat4 MobileSceneTextureLookup(inout FMaterialPixelParameters Parameter
 /** Samples the screen-space velocity for the specified UV coordinates. */
 float2 PostProcessVelocityLookup(float Depth, float2 UV)
 {
-#if GBUFFER_HAS_VELOCITY
 	float4 EncodedVelocity = Texture2DSampleLevel(SceneTexturesStruct.GBufferVelocityTexture, SceneTexturesStruct_GBufferVelocityTextureSampler, UV, 0);
-#else
-	float4 EncodedVelocity = Texture2DSample(PostProcessInput_4_Texture, PostProcessInput_4_SharedSampler, UV);
-#endif
 
 	float2 Velocity;
 	if( EncodedVelocity.x > 0.0 )
