Fix artifacts in Global Distance Field.
- frac(...) sometimes returns 1.0f which was causing sampling outside the page table.

- Use GLOBAL_DISTANCE_FIELD_INVALID_PAGE_ID instead of 0xFFFF.

diff --git a/Engine/Shaders/Private/DistanceFieldScreenGridLighting.usf b/Engine/Shaders/Private/DistanceFieldScreenGridLighting.usf
--- a/Engine/Shaders/Private/DistanceFieldScreenGridLighting.usf
+++ b/Engine/Shaders/Private/DistanceFieldScreenGridLighting.usf
@@ -256,7 +256,6 @@ void HemisphereConeTraceAgainstGlobalDistanceFieldClipmap(
 	// Cover AO leaking by comparing initial step SDF value against initial step size.
 	{
 		float3 WorldSamplePosition = WorldShadingPosition + WorldNormal * InitialOffset;
-		float3 StepVolumeUV = ComputeGlobalUV(WorldSamplePosition, ClipmapIndex);
 		float DistanceToOccluder = SampleGlobalDistanceField(WorldSamplePosition, AOGlobalMaxOcclusionDistance, ClipmapIndex);
 		ConeTraceLeakFill = saturate(Pow2(DistanceToOccluder / InitialOffset)) * 0.4f + 0.6f;
 	}
@@ -275,7 +274,6 @@ void HemisphereConeTraceAgainstGlobalDistanceFieldClipmap(
 		for (uint StepIndex = 0; StepIndex < NUM_CONE_STEPS && WorldStepOffset < AOGlobalMaxOcclusionDistance; StepIndex++)
 		{
 			float3 WorldSamplePosition = WorldShadingPosition + RotatedConeDirection * WorldStepOffset;
-			float3 StepVolumeUV = ComputeGlobalUV(WorldSamplePosition, ClipmapIndex);
 			float DistanceToOccluder = SampleGlobalDistanceField(WorldSamplePosition, AOGlobalMaxOcclusionDistance, ClipmapIndex);
 			float SphereRadius = WorldStepOffset * TanConeHalfAngle;
 			float InvSphereRadius = rcpFast(SphereRadius);
diff --git a/Engine/Shaders/Private/GlobalDistanceFieldMip.usf b/Engine/Shaders/Private/GlobalDistanceFieldMip.usf
--- a/Engine/Shaders/Private/GlobalDistanceFieldMip.usf
+++ b/Engine/Shaders/Private/GlobalDistanceFieldMip.usf
@@ -51,13 +51,14 @@ float LoadPrevDistanceFieldValue(int3 MipCoord, int3 Offset)
 #if READ_PAGES
 	// Reverse clipmap scroll offset as coarse clipmap is always centered around camera
 	float3 ScrolledClipmapUV = (MipCoord + 0.5f) / ClipmapMipResolution;
-	ScrolledClipmapUV = frac(ScrolledClipmapUV + ClipmapUVScrollOffset);
+	ScrolledClipmapUV = frac(ScrolledClipmapUV + ClipmapUVScrollOffset); // wraparound addressing
+	ScrolledClipmapUV = frac(ScrolledClipmapUV); // apply frac twice to prevent UV == 1.0f because frac(-0.00...001f) = 1.0f
 
 	int3 PageTableTextureCoord = saturate(ScrolledClipmapUV) * GlobalDistanceFieldClipmapSizeInPages + int3(0, 0, ClipmapIndex * GlobalDistanceFieldClipmapSizeInPages);
 	uint PageIndex = PageTableTexture.Load(int4(PageTableTextureCoord, 0));
 
 	float DistanceFieldValue = 1.0f;
-	if (PageIndex < 0xFFFF)
+	if (PageIndex < GLOBAL_DISTANCE_FIELD_INVALID_PAGE_ID)
 	{
 		float3 PageUV = ComputeGlobalDistanceFieldPageUV(ScrolledClipmapUV, PageIndex);
 		DistanceFieldValue = Texture3DSampleLevel(PageAtlasTexture, GlobalTrilinearWrappedSampler, PageUV, 0).x;
diff --git a/Engine/Shaders/Private/GlobalDistanceFieldShared.ush b/Engine/Shaders/Private/GlobalDistanceFieldShared.ush
--- a/Engine/Shaders/Private/GlobalDistanceFieldShared.ush
+++ b/Engine/Shaders/Private/GlobalDistanceFieldShared.ush
@@ -115,7 +115,9 @@ float3 ComputeGlobalUV(float3 WorldPosition, uint ClipmapIndex)
 {
 	//return ((WorldPosition - GlobalVolumeCenterAndExtent[ClipmapIndex].xyz + GlobalVolumeScollOffset[ClipmapIndex].xyz) / (GlobalVolumeCenterAndExtent[ClipmapIndex].w * 2) + .5f);
 	float4 WorldToUVAddAndMul = GlobalVolumeWorldToUVAddAndMul[ClipmapIndex];
-	return frac(WorldPosition * WorldToUVAddAndMul.www + WorldToUVAddAndMul.xyz);
+	float3 UV = frac(WorldPosition * WorldToUVAddAndMul.www + WorldToUVAddAndMul.xyz); // wraparound addressing
+	UV = frac(UV); // apply frac twice to prevent UV == 1.0f because frac(-0.00...001f) = 1.0f
+	return UV;
 }
 
 float3 ComputeGlobalMipUV(float3 WorldPosition, uint ClipmapIndex)
@@ -144,7 +146,7 @@ float GetDistanceToNearestSurfaceGlobal(float3 WorldPosition)
 
 		if (DistanceFromClipmap > ClipmapVoxelExtent)
 		{
-			float3 ClipmapVolumeUV = frac(ComputeGlobalUV(WorldPosition, ClipmapIndex));
+			float3 ClipmapVolumeUV = ComputeGlobalUV(WorldPosition, ClipmapIndex);
 
 			uint PageIndex = GetGlobalDistanceFieldPage(ClipmapVolumeUV, ClipmapIndex);
 
@@ -174,7 +176,7 @@ float SampleGlobalDistanceField(float3 WorldPosition, float MaxDistance, int Min
 	{
 		const float ClipmapInfluenceRange = GLOBAL_DISTANCE_FIELD_INFLUENCE_RANGE_IN_VOXELS * 2.0f * GlobalVolumeCenterAndExtent[ClipmapIndex].w * GlobalVolumeTexelSize;
 		
-		float3 ClipmapVolumeUV = frac(ComputeGlobalUV(WorldPosition, ClipmapIndex));
+		float3 ClipmapVolumeUV = ComputeGlobalUV(WorldPosition, ClipmapIndex);
 		uint PageIndex = GetGlobalDistanceFieldPage(ClipmapVolumeUV, ClipmapIndex);
 
 		if (PageIndex < GLOBAL_DISTANCE_FIELD_INVALID_PAGE_ID)
@@ -222,7 +224,7 @@ float3 GetDistanceFieldGradientGlobal(float3 WorldPosition)
 
 		if (DistanceFromClipmap >= ClipmapVoxelExtent)
 		{
-			const float3 ClipmapVolumeUV = frac(ComputeGlobalUV(WorldPosition, ClipmapIndex));
+			const float3 ClipmapVolumeUV = ComputeGlobalUV(WorldPosition, ClipmapIndex);
 
 			uint PageIndex = GetGlobalDistanceFieldPage(ClipmapVolumeUV, ClipmapIndex);
 
diff --git a/Engine/Shaders/Private/GlobalDistanceFieldUtils.ush b/Engine/Shaders/Private/GlobalDistanceFieldUtils.ush
--- a/Engine/Shaders/Private/GlobalDistanceFieldUtils.ush
+++ b/Engine/Shaders/Private/GlobalDistanceFieldUtils.ush
@@ -114,7 +114,7 @@ FGlobalSDFTraceResult RayTraceGlobalDistanceField(FGlobalSDFTraceInput TraceInpu
 				float DistanceFieldMipValue = Texture3DSampleLevel(GlobalDistanceFieldMipTexture, GlobalTrilinearClampedSampler, MipUV, 0).x;
 				float DistanceField = DecodeGlobalDistanceFieldPageDistance(DistanceFieldMipValue, GlobalDistanceFieldMipFactor * ClipmapInfluenceRange);
 				uint PageIndex = GetGlobalDistanceFieldPage(ClipmapVolumeUV, ClipmapIndex);
-				if (PageIndex < 0xFFFF && DistanceFieldMipValue < GlobalDistanceFieldMipTransition)
+				if (PageIndex < GLOBAL_DISTANCE_FIELD_INVALID_PAGE_ID && DistanceFieldMipValue < GlobalDistanceFieldMipTransition)
 				{
 					float3 PageUV = ComputeGlobalDistanceFieldPageUV(ClipmapVolumeUV, PageIndex);
 					float DistanceFieldValue = Texture3DSampleLevel(GlobalDistanceFieldPageAtlasTexture, GlobalTrilinearWrappedSampler, PageUV, 0).x;
diff --git a/Engine/Shaders/Private/VolumetricFog.usf b/Engine/Shaders/Private/VolumetricFog.usf
--- a/Engine/Shaders/Private/VolumetricFog.usf
+++ b/Engine/Shaders/Private/VolumetricFog.usf
@@ -434,7 +434,6 @@ float HemisphereConeTraceAgainstGlobalDistanceFieldClipmap(
 	for (uint StepIndex = 0; StepIndex < NUM_CONE_STEPS && WorldStepOffset < AOGlobalMaxOcclusionDistance; StepIndex++)
 	{
 		float3 WorldSamplePosition = WorldShadingPosition + ConeDirection * WorldStepOffset;
-		float3 StepVolumeUV = ComputeGlobalUV(WorldSamplePosition, ClipmapIndex);
 		float DistanceToOccluder = SampleGlobalDistanceField(WorldSamplePosition, AOGlobalMaxOcclusionDistance, ClipmapIndex);
 		float SphereRadius = WorldStepOffset * TanConeHalfAngle;
 		float InvSphereRadius = rcpFast(SphereRadius);
