Big Strata core update!
- Completely changed the way we evaluate and represent the material topology.
- A topology is now represented as a tree to allow any horizontal/vertical operations (almost like a kdtree).
- This allows for correct translucent material coverage/transmittance evaluation, as well as correct BSDFs luminance weight (finally).
- All the operations to be able to do that are expended in HLSLTranslator instead of letting the compiler handle it (otherwise was too slow, leading to bugs and could not even debug shaders). This means no conditional BSDF packing.
- There is no more Material/layer/BSDFs array: only a Tree or operators, some of them being BSDF at leaves.
- Parameter blending has also been updated: when enabled on a node, it will be forced on all the oprator child / in the sub tree.
- Parameter blending is applied on an inlinedbsdf on the StrataData. When it is the root of the parameter blending subtree, it is promoted to a fully fledged BSDF operator node to allow other regular operators.
- The inlinedbsdf is not longer in premultiplied mode. Horizontal/Vertical/Add parameter blending operators have been updated to reflect that. This allows for simpler/unified/clearer code.
- Vertical blending now uses the new simple "layering of two slabs with uncorelated coverage" math I have RnD'd. See GetVerticalLayeringInfo in ShadingCommon.ush.

Follow up to this CL:
- CompileStrataBlendFunction should use parameter blending.
- Add more compilation debug output
- Update material vizualisation
- Fix Rough refraction
- Fix decals
- Fix debug probes
- Pack FStrataBSDF
- Move all StrataCompilationInfoCreateSingleBSDFMaterial into into the strata tree and rework cost evaluation and material LOD
- Expand graph visit such as UpdateBSDFWeightAfterOperatorVisit to be specific to operators for the compiler to have to do less search to inline what we know already
- Fix STRATA_TODO: operation using parameter blending are not actually discarded so they still occupy a pot in the opeartion arrays in the compiler and in the shader also.

diff --git a/Engine/Shaders/Private/BasePassPixelShader.usf b/Engine/Shaders/Private/BasePassPixelShader.usf
--- a/Engine/Shaders/Private/BasePassPixelShader.usf
+++ b/Engine/Shaders/Private/BasePassPixelShader.usf
@@ -1043,11 +1043,8 @@ void FPixelShaderInOut_MainPS(
 
 	// Initialise a Strata header with normal in registers
 	FStrataPixelHeader StrataPixelHeader = InitialiseStrataPixelHeader();
-	UNROLL
-	for (int l = 0; l < StrataData.LayerCount; ++l)
-	{
-		StrataPixelHeader.BSDFCount += StrataData.Layers[l].BSDFCount;
-	}
+	StrataPixelHeader.StrataTree = MaterialParameters.StrataTree;
+	StrataPixelHeader.BSDFCount = MaterialParameters.StrataTree.BSDFCount;
 	StrataPixelHeader.SharedLocalBases = MaterialParameters.SharedLocalBases;
 	StrataPixelHeader.MaterialAO = MaterialAO;
 	SetCastContactShadow(StrataPixelHeader, GetPrimitiveData(MaterialParameters).Flags & PRIMITIVE_SCENE_DATA_FLAG_HAS_CAST_CONTACT_SHADOW);
@@ -1059,8 +1056,8 @@ void FPixelShaderInOut_MainPS(
 
 	// Override GBuffer data with Strata SLW water BSDF to run forward shadfing code.
 	// STRATA_TODO: run the shading through a strata path? (by adding a special BSDF?)
-	SanitizeStrataSingleLayerWater(PixelMaterialInputs.FrontMaterial.Layers[0].BSDFs[0]);
-	FStrataBSDF SLWBSDF = PixelMaterialInputs.FrontMaterial.Layers[0].BSDFs[0];
+	SanitizeStrataSingleLayerWater(PixelMaterialInputs.FrontMaterial.InlinedBSDF);
+	FStrataBSDF SLWBSDF = PixelMaterialInputs.FrontMaterial.InlinedBSDF;
 
 	BaseColor = SLW_BASECOLOR(SLWBSDF);
 	Metallic = SLW_METALLIC(SLWBSDF);
@@ -1673,7 +1670,7 @@ void FPixelShaderInOut_MainPS(
 
 #if (STRATA_TRANSLUCENT_FORWARD || STRATA_FORWARD_SHADING) && !STRATA_INLINE_SINGLELAYERWATER && !STRATA_OPTIMIZED_UNLIT
 	//FORWARD_SHADING
-	if (StrataData.LayerCount > 0)
+	if (StrataPixelHeader.BSDFCount > 0)
 	{
 		float2 ScreenUV = ScreenPositionToBufferUV(MaterialParameters.ScreenPosition);
 
@@ -1714,7 +1711,7 @@ void FPixelShaderInOut_MainPS(
 #elif STRATA_ENABLED && STRATA_OPTIMIZED_UNLIT
 
 	// Unlit forces a single BSDF
-	FStrataBSDF UnlitBSDF = PixelMaterialInputs.FrontMaterial.Layers[0].BSDFs[0];
+	FStrataBSDF UnlitBSDF = PixelMaterialInputs.FrontMaterial.InlinedBSDF;
 
 #if STRATA_OPAQUE_DEFERRED
 	Color = BSDF_GETEMISSIVE(UnlitBSDF);
diff --git a/Engine/Shaders/Private/DebugProbes.usf b/Engine/Shaders/Private/DebugProbes.usf
--- a/Engine/Shaders/Private/DebugProbes.usf
+++ b/Engine/Shaders/Private/DebugProbes.usf
@@ -64,6 +64,17 @@ void MainPS(
 	, out float  OutDepth		: SV_DEPTH
 )
 {
+	OutColor = 0;
+#if PERMUTATION_PASS == RENDER_BASEPASS
+	OutGBufferA = 0;
+	OutGBufferB = 0;
+	OutGBufferC = 0;
+	OutGBufferD = 0;
+	OutGBufferE = 0;
+#endif
+	OutDepth = 0;
+
+#if 0
 	ResolvedView = ResolveView();
 	const uint2 PixelCoord = uint2(SvPosition.xy);
 
@@ -243,6 +254,8 @@ void MainPS(
 	{
 		FStrataPixelHeader StrataPixelHeader = InitialiseStrataPixelHeader();
 
+		// STRATA_TODO: produce the tree
+
 		uint SharedLocalBasisIndex = 0;
 		StrataPixelHeader.SharedLocalBases.Normals[SharedLocalBasisIndex] = normalize(Normal);
 		StrataPixelHeader.SharedLocalBases.Count = 1;
@@ -322,4 +335,5 @@ void MainPS(
 	OutDepth = DeviceZ;
 #endif
 
+#endif // TODO
 }
\ No newline at end of file
diff --git a/Engine/Shaders/Private/DistortAccumulatePS.usf b/Engine/Shaders/Private/DistortAccumulatePS.usf
--- a/Engine/Shaders/Private/DistortAccumulatePS.usf
+++ b/Engine/Shaders/Private/DistortAccumulatePS.usf
@@ -88,6 +88,8 @@ void Main(
 
 #if STRATA_ENABLED
 
+	// STRATA_TODO Fix me
+#if 0
 
 	FStrataData StrataData = PixelMaterialInputs.FrontMaterial;
 
@@ -101,6 +103,8 @@ void Main(
 		StrataPixelHeader.BSDFCount += StrataData.Layers[l].BSDFCount;
 	}
 	StrataPixelHeader.SharedLocalBases = MaterialParameters.SharedLocalBases;
+	StrataPixelHeader.StrataTree = MaterialParameters.StrataTree;
+	// STRATA_TODO: process the tree
 
 	const float3 V = MaterialParameters.CameraVector;
 	
@@ -224,7 +228,11 @@ void Main(
 
 	// Sample depth at distortion offset
 	float2 DistortBufferUV = ScreenUV + BufferUVDistortion;
+#endif
 
+	OutVariance = 0.0f;
+	float2 BufferUVDistortion = 0.0f;
+	float2 DistortBufferUV = ScreenUV;
 
 #else // STRATA_ENABLED
 
diff --git a/Engine/Shaders/Private/HairStrands/HairStrandsMaterialCommon.ush b/Engine/Shaders/Private/HairStrands/HairStrandsMaterialCommon.ush
--- a/Engine/Shaders/Private/HairStrands/HairStrandsMaterialCommon.ush
+++ b/Engine/Shaders/Private/HairStrands/HairStrandsMaterialCommon.ush
@@ -203,7 +203,7 @@ FEvaluateOutput Evaluate(
 		CalcMaterialParametersEx(MaterialParameters, PixelMaterialInputs, SvPosition, LocalScreenPosition, bIsFrontFace, TranslatedWorldPosition, TranslatedWorldPosition);
 
 #if PROJECT_STRATA && MATERIAL_IS_STRATA && MATERIAL_SHADINGMODEL_HAIR
-		FStrataBSDF HairBSDF = PixelMaterialInputs.FrontMaterial.Layers[0].BSDFs[0];
+		FStrataBSDF HairBSDF = PixelMaterialInputs.FrontMaterial.InlinedBSDF;
 
 		OutSample.BaseColor = HAIR_BASECOLOR(HairBSDF);
 		OutSample.Roughness = HAIR_ROUGHNESS(HairBSDF);
diff --git a/Engine/Shaders/Private/HairStrands/HairStrandsVisibilityPS.usf b/Engine/Shaders/Private/HairStrands/HairStrandsVisibilityPS.usf
--- a/Engine/Shaders/Private/HairStrands/HairStrandsVisibilityPS.usf
+++ b/Engine/Shaders/Private/HairStrands/HairStrandsVisibilityPS.usf
@@ -125,7 +125,7 @@ void MainVisibility(
 		CalcMaterialParametersEx(MaterialParameters, PixelMaterialInputs, SvPosition, ScreenPosition, bIsFrontFace, TranslatedWorldPosition, TranslatedWorldPosition);
 
 #if PROJECT_STRATA && MATERIAL_IS_STRATA
-		FStrataBSDF HairBSDF = PixelMaterialInputs.FrontMaterial.Layers[0].BSDFs[0];
+		FStrataBSDF HairBSDF = PixelMaterialInputs.FrontMaterial.InlinedBSDF;
 
 		BaseColor		 = HAIR_BASECOLOR(HairBSDF);
 		Roughness		 = HAIR_ROUGHNESS(HairBSDF);
diff --git a/Engine/Shaders/Private/LightFunctionCommon.ush b/Engine/Shaders/Private/LightFunctionCommon.ush
--- a/Engine/Shaders/Private/LightFunctionCommon.ush
+++ b/Engine/Shaders/Private/LightFunctionCommon.ush
@@ -59,7 +59,7 @@ float3 GetLightFunctionColor(float3 LightVector, float3 TranslatedWorldSpace)
   #if COMPUTESHADER
 	return 0.f; // STRATA_TODO: This path doesn't seem to be used
   #else
-	return StrataGetBSDFEmissive(PixelMaterialInputs.FrontMaterial.Layers[0].BSDFs[0]);
+	return StrataGetBSDFEmissive(PixelMaterialInputs.FrontMaterial.InlinedBSDF);
   #endif
 
 #else // STRATA_ENABLED
diff --git a/Engine/Shaders/Private/MaterialTemplate.ush b/Engine/Shaders/Private/MaterialTemplate.ush
--- a/Engine/Shaders/Private/MaterialTemplate.ush
+++ b/Engine/Shaders/Private/MaterialTemplate.ush
@@ -149,6 +149,7 @@
 
 #if PROJECT_STRATA
 #include "/Engine/Private/Strata/Strata.ush"
+#include "/Engine/Private/Strata/StrataTree.ush"
 #include "/Engine/Private/Strata/StrataLegacyConversion.ush"
 #else
 struct FStrataData
@@ -445,6 +446,7 @@ struct FMaterialPixelParameters
 
 #if TEMPLATE_USES_STRATA
 	FSharedLocalBases SharedLocalBases;
+	FStrataTree StrataTree;
 #endif
 };
 
@@ -2879,6 +2881,7 @@ void CalcPixelMaterialInputs(in out FMaterialPixelParameters Parameters, in out
 
 #if TEMPLATE_USES_STRATA
 	Parameters.SharedLocalBases = StrataInitialiseSharedLocalBases();
+	Parameters.StrataTree = GetInitialisedStrataTree();
 #endif
 
 	// Note that here MaterialNormal can be in world space or tangent space
@@ -2938,6 +2941,7 @@ void CalcPixelMaterialInputsAnalyticDerivatives(in out FMaterialPixelParameters
 
 #if TEMPLATE_USES_STRATA
 	Parameters.SharedLocalBases = StrataInitialiseSharedLocalBases();
+	Parameters.StrataTree = GetInitialisedStrataTree();
 #endif
 
 	// Note that here MaterialNormal can be in world space or tangent space
diff --git a/Engine/Shaders/Private/PostProcessMaterialShaders.usf b/Engine/Shaders/Private/PostProcessMaterialShaders.usf
--- a/Engine/Shaders/Private/PostProcessMaterialShaders.usf
+++ b/Engine/Shaders/Private/PostProcessMaterialShaders.usf
@@ -208,7 +208,7 @@ void MainPS(
 	float Alpha = 1.0f;
 #if STRATA_ENABLED   
 	// Unlit forces a single BSDF
-	const FStrataBSDF UnlitBSDF = PixelMaterialInputs.FrontMaterial.Layers[0].BSDFs[0];
+	const FStrataBSDF UnlitBSDF = PixelMaterialInputs.FrontMaterial.InlinedBSDF;
 
 	// Grab emissive colour as output
   #if MATERIAL_OUTPUT_OPACITY_AS_ALPHA
diff --git a/Engine/Shaders/Private/ShadingCommon.ush b/Engine/Shaders/Private/ShadingCommon.ush
--- a/Engine/Shaders/Private/ShadingCommon.ush
+++ b/Engine/Shaders/Private/ShadingCommon.ush
@@ -4,6 +4,11 @@
 
 #include "Common.ush"
 
+
+///////////////////////////////////////////////////////////////////////////////
+// Shading models
+
+
 // SHADINGMODELID_* occupy the 4 low bits of an 8bit channel and SKIP_* occupy the 4 high bits
 #define SHADINGMODELID_UNLIT				0
 #define SHADINGMODELID_DEFAULT_LIT			1
@@ -89,6 +94,8 @@ bool GetShadingModelRequiresBackfaceLighting(uint ShadingModelID)
 	return ShadingModelID == SHADINGMODELID_TWOSIDED_FOLIAGE;
 }
 
+///////////////////////////////////////////////////////////////////////////////
+// Shading parameterisation
 
 float F0ToDielectricSpecular(float F0)
 {
@@ -153,3 +160,41 @@ float F0RGBToMetallic(float3 F0)
 {
 	return F0ToMetallic(max(F0.r, max(F0.g, F0.b)));
 }
+
+
+///////////////////////////////////////////////////////////////////////////////
+// Layering, coverage and transmittance
+
+struct FVerticalLayeringInfo
+{
+	float TransmittanceTopAndBottom;	// Ratio of transmittance of both top and bottom surface only
+	float TransmittanceOnlyBottom;		// Ratio of transmittance of bottom surface only 
+	float TransmittanceOnlyTop;			// Ratio of transmittance of top surface only 
+
+	float SurfaceBottom;				// Ratio of bottom surface visible
+	float SurfaceTop;					// Ratio of top surface visible
+
+	float Coverage;						// Ratio of any surface is visible
+	float NoSurface;					// Ratio of no surface is visible
+};
+
+// This function returns surface information after a layering of two slab of matter having difference and uncorelated coverage.
+// See https://www.desmos.com/calculator/qg5yc3nnr4
+FVerticalLayeringInfo GetVerticalLayeringInfo(const float TopCoverage, const float BottomCoverage)
+{
+	FVerticalLayeringInfo Info;
+
+	Info.TransmittanceTopAndBottom	= TopCoverage * BottomCoverage;
+	Info.TransmittanceOnlyBottom	= (1.0f - TopCoverage) * BottomCoverage;
+	Info.TransmittanceOnlyTop		= (1.0f - BottomCoverage) * TopCoverage;
+
+	Info.SurfaceBottom				= Info.TransmittanceOnlyBottom;
+	Info.SurfaceTop					= TopCoverage;								// == Info.TransmittanceOnlyTop + Info.TransmittanceTopAndBottom;
+
+	Info.Coverage					= Info.SurfaceTop + Info.SurfaceBottom;		// == Info.TransmittanceTopAndBottom + Info.TransmittanceOnlyBottom + Info.TransmittanceOnlyTop;// == TopCoverage + Info.TransmittanceOnlyBottom
+	Info.NoSurface					= 1.0f - Info.Coverage;
+
+	return Info;
+}
+
+
diff --git a/Engine/Shaders/Private/SingleLayerWaterShading.ush b/Engine/Shaders/Private/SingleLayerWaterShading.ush
--- a/Engine/Shaders/Private/SingleLayerWaterShading.ush
+++ b/Engine/Shaders/Private/SingleLayerWaterShading.ush
@@ -104,7 +104,7 @@ WaterVolumeLightingOutput EvaluateWaterVolumeLighting(
 		const float BehindWaterDeltaDepth = CameraIsUnderWater ? WaterDepth : max(0.0f, SceneDepth - WaterDepth);
 
 #if STRATA_INLINE_SINGLELAYERWATER
-		FStrataBSDF SLWBSDF = PixelMaterialInputs.FrontMaterial.Layers[0].BSDFs[0];
+		FStrataBSDF SLWBSDF = PixelMaterialInputs.FrontMaterial.InlinedBSDF;
 		const float3 Albedo = SLW_WATERALBEDO(SLWBSDF);
 		const float3 ExtinctionCoeff = SLW_WATEREXTINCTION(SLWBSDF);
 		const float3 ScatteringCoeff = Albedo * ExtinctionCoeff;
diff --git a/Engine/Shaders/Private/Strata/Strata.ush b/Engine/Shaders/Private/Strata/Strata.ush
--- a/Engine/Shaders/Private/Strata/Strata.ush
+++ b/Engine/Shaders/Private/Strata/Strata.ush
@@ -90,6 +90,14 @@ FRWStrataMaterialContainer InitialiseRWStrataMaterialContainer(RWTexture2DArray<
 #define MATERIAL_FULLY_ROUGH 0
 #endif
 
+#ifndef STRATA_CLAMPED_BSDF_COUNT
+#define STRATA_CLAMPED_BSDF_COUNT 0
+#endif
+
+#ifndef STRATA_USES_CONVERSION_FROM_LEGACY
+#define STRATA_USES_CONVERSION_FROM_LEGACY 0
+#endif
+
 #define STRATA_INLINE_SINGLELAYERWATER (PROJECT_STRATA && MATERIAL_IS_STRATA && STRATA_INLINE_SHADING && MATERIAL_SHADINGMODEL_SINGLELAYERWATER)
 
 // Behind the scene, we do all the simple volumetric lighting & transmitance computations considering a slab of 100 centimeters = 1 meter.
@@ -103,12 +111,16 @@ struct FStrataBSDF
 {
 	uint				State;		// Extra informations and other enabled BSDF features
 
-	float3				FinalWeight;// Contains the throughput towards the view accounting for BSDF visible above. It also contains coverage if it has not been baked into the BSDF parameters when using operation with parameter blending.
+	// STRATA_TODO pack the following
+	int					OperatorIndex;
+	float3				LuminanceWeight;
+	int					bIsBottom;		// STRATA TODO use thet state one or set only when inline
+	int					bIsTop;			// STRATA TODO use thet state one or set only when inline
+
 #if STRATA_INLINE_SHADING
-	float				Coverage;	// When creating a material, each BSDF/slab can have its coverage updated. 
-									// When this happens, we also update its parameters to account for that, similary to premultiplied alpha blending.
-									// Coverage is kept around because it is needed to evalauted the contribution to throughput.
-	bool				bPremultipliedCoverage;// True when coverage is already weighting all parameters. This is true when parameter blending is used. This allows to use parameter blending or not in the same graph.
+
+	float				Coverage;	// The coverage of the material, as ratio of visible matter in this pixel.
+
 	float3				Emissive;	// All bsdf can potientially have emissive. This is never written out to the strata gbuffer but instead accumulated in the emissive buffer during the base pass.
 	float				ThicknessCm;// This must be kept in order to be able to normalize the mean free path to comply with our volumetric material math all done in a normalised slab of medium.
 	float3				TmpDMFP;
@@ -563,26 +575,65 @@ bool StrataHasAreaLightIntegrator(in FStrataBSDF BSDF)
 	return BSDFType == STRATA_BSDF_TYPE_SLAB;
 }
 
+bool StrataIsBSDFVisible(in FStrataBSDF BSDF)
+{
+	return any(BSDF.LuminanceWeight > 0.0f);
+}
+
+
 bool StrataHasGreyScaleWeight(in float3 Weight)
 {
 	return Weight.x == Weight.y && Weight.y == Weight.z;
 }
 
 ///////////////////////////////////////////////////////////////////////////////
-// Layer & BSDF description
+// Material & BSDF description
+
+struct FStrataOperator
+{
+	// STRATA_TODO: pack in a single uint
+	uint TopIndex;
+	uint MaxDistanceFromLeaves;
+	uint Type;
+	int	AIndex;
+	int BIndex;
+	float Weight;
+
+	float  Coverage;
+	float3 Transmittance;	// This is the transmittance of the matter. If coverage is 0.25, then 25% of the light will be affected by the transmittance and 75% of light won't be.
+
+	float VerticalTopCoverage;
+	float3 VerticalTopTransmittance;
+};
 
-struct FStrataLayer
+FStrataOperator GetInitialisedStrataOperator()
+{
+	FStrataOperator StrataOp = (FStrataOperator)0;
+	StrataOp.TopIndex = -1;
+	return StrataOp;
+}
+
+// Representes a strata material as BSDFs organised in a tree repreesnting its topology 
+struct FStrataTree
 {
 	int BSDFCount;
+	uint OperatorCount;
 
-	FStrataBSDF BSDFs[STRATA_MAX_BSDF_COUNT_PER_LAYER];
+	FStrataBSDF BSDFs[STRATA_MAX_BSDF_COUNT];
+	FStrataOperator Operators[STRATA_MAX_OPERATOR_COUNT]; // STRATA_TODO assert when too many ops on the compiler side. Shoud already be done for BSDF count.
 };
 
+FStrataTree GetInitialisedStrataTree()
+{
+	FStrataTree StrataTree = (FStrataTree)0;
+	return StrataTree;
+}
+
 struct FStrataData
 {
-	int LayerCount;
+	FStrataBSDF InlinedBSDF;// Used for parameter blending or special BSDF such as hair, cloud or water.
 
-	FStrataLayer Layers[STRATA_MAX_LAYER_COUNT];	// Layer 0 is the top layer
+	int  OperatorIndex;		// Used as entry point ot the material graph with multiple BSDF and operators.
 };
 
 FStrataData GetInitialisedStrataData()
@@ -590,13 +641,6 @@ FStrataData GetInitialisedStrataData()
 	FStrataData StrataData = (FStrataData)0;
 	return StrataData;
 }
-FStrataData GetInitialisedStrataDataWithReservedFirstSlot()
-{
-	FStrataData StrataData = GetInitialisedStrataData();
-	StrataData.LayerCount++;
-	StrataData.Layers[0].BSDFCount++;
-	return StrataData;
-}
 
 ///////////////////////////////////////////////////////////////////////////////
 // BSDF
@@ -611,10 +655,10 @@ FStrataData GetStrataSlabBSDF(
 	float Haziness, 
 	float ThinFilmThickness,
 	float FuzzAmount, float3 FuzzColor,
-	float Thickness, 
+	float Thickness,
 	uint SharedLocalBasisIndex, inout uint SharedLocalBasisTypes)
 {
-	FStrataData StrataData = GetInitialisedStrataDataWithReservedFirstSlot();
+	FStrataData StrataData = GetInitialisedStrataData();
 
 	// STRATA_TODO apply SSS as a function of pixel footprint w.r.t. MFP
 	SSSDMFPScale = saturate(SSSDMFPScale);
@@ -667,101 +711,106 @@ FStrataData GetStrataSlabBSDF(
 	//  - We also only enabled thin lighting is the MFP input pin is plugged in.
 	//  - It will only be ran for the bottom layer only. Otherwise, the simple volume should be used. As such, Two-sided lighting can also be disabled later if part of the bottom layer.
 
-	BSDF_SETTYPE			(StrataData.Layers[0].BSDFs[0],		STRATA_BSDF_TYPE_SLAB);
-	BSDF_SETSHAREDLOCALBASISID(StrataData.Layers[0].BSDFs[0],	SharedLocalBasisIndex);
-	BSDF_SETEMISSIVE		(StrataData.Layers[0].BSDFs[0],		Emissive);
-	BSDF_SETHASANISOTROPY	(StrataData.Layers[0].BSDFs[0],		bHasAnisotropy);
-	BSDF_SETISTOPLAYER		(StrataData.Layers[0].BSDFs[0],		0);
-	BSDF_SETHASSSS			(StrataData.Layers[0].BSDFs[0],		bHasSSS ? 1 : 0);
-	BSDF_SETHASSSSPROFILE	(StrataData.Layers[0].BSDFs[0],		bHasSSSProfile ? 1 : 0);
-	BSDF_SETISTHIN			(StrataData.Layers[0].BSDFs[0],		bIsThin ? 1 : 0);
-	BSDF_SETHASDMFP			(StrataData.Layers[0].BSDFs[0],		bDMFPPlugged ? 1 : 0);
-	BSDF_SETHASHAZINESS		(StrataData.Layers[0].BSDFs[0],		bHasHaziness);
-	BSDF_SETHASTHINFILM		(StrataData.Layers[0].BSDFs[0],		bHasThinFilm);
-	BSDF_SETTHICKNESSCM		(StrataData.Layers[0].BSDFs[0],		Thickness);
-	BSDF_SETHASFUZZ			(StrataData.Layers[0].BSDFs[0],		bHasFuzz);
+	BSDF_SETTYPE			(StrataData.InlinedBSDF,		STRATA_BSDF_TYPE_SLAB);
+	BSDF_SETSHAREDLOCALBASISID(StrataData.InlinedBSDF,		SharedLocalBasisIndex);
+	BSDF_SETEMISSIVE		(StrataData.InlinedBSDF,		Emissive);
+	BSDF_SETHASANISOTROPY	(StrataData.InlinedBSDF,		bHasAnisotropy);
+	BSDF_SETISTOPLAYER		(StrataData.InlinedBSDF,		0);
+	BSDF_SETHASSSS			(StrataData.InlinedBSDF,		bHasSSS ? 1 : 0);
+	BSDF_SETHASSSSPROFILE	(StrataData.InlinedBSDF,		bHasSSSProfile ? 1 : 0);
+	BSDF_SETISTHIN			(StrataData.InlinedBSDF,		bIsThin ? 1 : 0);
+	BSDF_SETHASDMFP			(StrataData.InlinedBSDF,		bDMFPPlugged ? 1 : 0);
+	BSDF_SETHASHAZINESS		(StrataData.InlinedBSDF,		bHasHaziness);
+	BSDF_SETHASTHINFILM		(StrataData.InlinedBSDF,		bHasThinFilm);
+	BSDF_SETTHICKNESSCM		(StrataData.InlinedBSDF,		Thickness);
+	BSDF_SETHASFUZZ			(StrataData.InlinedBSDF,		bHasFuzz);
 
 	if (UseMetalness > 0.0f)
 	{
-		BSDF_SETHASF90(StrataData.Layers[0].BSDFs[0], any(EdgeColor < 1.0f));
+		BSDF_SETHASF90(StrataData.InlinedBSDF, any(EdgeColor < 1.0f));
 
 		const float SafeMetallic = saturate(Metallic);
-		SLAB_DIFFUSEALBEDO(StrataData.Layers[0].BSDFs[0])	= BaseColor - BaseColor * SafeMetallic;
-		SLAB_F0(StrataData.Layers[0].BSDFs[0])				= ComputeF0(Specular, BaseColor, SafeMetallic);
-		SLAB_F90(StrataData.Layers[0].BSDFs[0])				= ComputeF90(F0, EdgeColor, SafeMetallic);
+		SLAB_DIFFUSEALBEDO(StrataData.InlinedBSDF)			= BaseColor - BaseColor * SafeMetallic;
+		SLAB_F0(StrataData.InlinedBSDF)						= ComputeF0(Specular, BaseColor, SafeMetallic);
+		SLAB_F90(StrataData.InlinedBSDF)					= ComputeF90(F0, EdgeColor, SafeMetallic);
 	}
 	else
 	{
 		// See how UMaterialExpressionStrataSlabBSDF node interface is changed
-		BSDF_SETHASF90(StrataData.Layers[0].BSDFs[0], any(F90 < 1.0f));
+		BSDF_SETHASF90(StrataData.InlinedBSDF, any(F90 < 1.0f));
 
-		SLAB_DIFFUSEALBEDO(StrataData.Layers[0].BSDFs[0])	= DiffuseAlbedo;
-		SLAB_F0(StrataData.Layers[0].BSDFs[0])				= F0;
-		SLAB_F90(StrataData.Layers[0].BSDFs[0])				= F90;
+		SLAB_DIFFUSEALBEDO(StrataData.InlinedBSDF)			= DiffuseAlbedo;
+		SLAB_F0(StrataData.InlinedBSDF)						= F0;
+		SLAB_F90(StrataData.InlinedBSDF)					= F90;
 	}
 
-	SLAB_ROUGHNESS			(StrataData.Layers[0].BSDFs[0])		= Roughness;
-	SLAB_ANISOTROPY			(StrataData.Layers[0].BSDFs[0])		= Anisotropy;
-	SLAB_HAZINESS			(StrataData.Layers[0].BSDFs[0])		= Haziness;
-	SLAB_THINFILM_THICKNESS	(StrataData.Layers[0].BSDFs[0])		= ThinFilmThickness;
+	SLAB_ROUGHNESS			(StrataData.InlinedBSDF)		= Roughness;
+	SLAB_ANISOTROPY			(StrataData.InlinedBSDF)		= Anisotropy;
+	SLAB_HAZINESS			(StrataData.InlinedBSDF)		= Haziness;
+	SLAB_THINFILM_THICKNESS	(StrataData.InlinedBSDF)		= ThinFilmThickness;
 
 	if (bHasSSSProfile)
 	{
-		SLAB_SSSPROFILEID	(StrataData.Layers[0].BSDFs[0])		= SSSProfileID;
-		SLAB_SSSPROFILERADIUSSCALE(StrataData.Layers[0].BSDFs[0])= SSSDMFPScale;
+		SLAB_SSSPROFILEID	(StrataData.InlinedBSDF)		= SSSProfileID;
+		SLAB_SSSPROFILERADIUSSCALE(StrataData.InlinedBSDF)= SSSDMFPScale;
 	}
 	else
 	{
-		SLAB_SSSDMFP	(StrataData.Layers[0].BSDFs[0])			= SSSDMFP * SSSDMFPScale;
+		SLAB_SSSDMFP	(StrataData.InlinedBSDF)			= SSSDMFP * SSSDMFPScale;
 	}
 
-	SLAB_FUZZ_AMOUNT		(StrataData.Layers[0].BSDFs[0])		= FuzzAmount;
-	SLAB_FUZZ_COLOR			(StrataData.Layers[0].BSDFs[0])		= FuzzColor;
+	SLAB_FUZZ_AMOUNT		(StrataData.InlinedBSDF)		= FuzzAmount;
+	SLAB_FUZZ_COLOR			(StrataData.InlinedBSDF)		= FuzzColor;
 
 #if STRATA_INLINE_SHADING
 	// At this stage we do not know if this strata BSDF is going to be using simple volume lighting so always store the DMFP on the side to not override SSSPROFILE data.
-	StrataData.Layers[0].BSDFs[0].TmpDMFP = SSSDMFP * SSSDMFPScale;
+	StrataData.InlinedBSDF.TmpDMFP = SSSDMFP * SSSDMFPScale;
 
-	StrataData.Layers[0].BSDFs[0].Coverage = 1.0f;
+	StrataData.InlinedBSDF.Coverage = 1.0f;
 #endif
+
 	return StrataData;
 }
 
 FStrataData GetStrataVolumeFogCloudBSDF(float3 Albedo, float3 Extinction, float3 Emissive, float AmbientOcclusion)
 {
-	FStrataData StrataData = GetInitialisedStrataDataWithReservedFirstSlot();
+	FStrataData StrataData = GetInitialisedStrataData();
 
-	BSDF_SETTYPE					(StrataData.Layers[0].BSDFs[0], STRATA_BSDF_TYPE_VOLUMETRICFOGCLOUD);
-	BSDF_SETEMISSIVE				(StrataData.Layers[0].BSDFs[0], Emissive);
-	VOLUMETRICFOGCLOUD_ALBEDO		(StrataData.Layers[0].BSDFs[0]) = Albedo;
-	VOLUMETRICFOGCLOUD_EXTINCTION	(StrataData.Layers[0].BSDFs[0]) = Extinction;
-	VOLUMETRICFOGCLOUD_AO			(StrataData.Layers[0].BSDFs[0]) = AmbientOcclusion;
+	BSDF_SETTYPE					(StrataData.InlinedBSDF, STRATA_BSDF_TYPE_VOLUMETRICFOGCLOUD);
+	BSDF_SETEMISSIVE				(StrataData.InlinedBSDF, Emissive);
+	VOLUMETRICFOGCLOUD_ALBEDO		(StrataData.InlinedBSDF) = Albedo;
+	VOLUMETRICFOGCLOUD_EXTINCTION	(StrataData.InlinedBSDF) = Extinction;
+	VOLUMETRICFOGCLOUD_AO			(StrataData.InlinedBSDF) = AmbientOcclusion;
 
 #if STRATA_INLINE_SHADING
-	StrataData.Layers[0].BSDFs[0].Coverage = 1.0f;
+	StrataData.InlinedBSDF.Coverage = 1.0f;
 #endif
 
+	// STRATA_TODO BSDF OP
+
 	return StrataData;
 }
 
 FStrataData GetStrataUnlitBSDF(float3 Emissive, float3 TransmittanceColor)
 {
-	FStrataData StrataData = GetInitialisedStrataDataWithReservedFirstSlot();
+	FStrataData StrataData = GetInitialisedStrataData();
 
-	BSDF_SETTYPE					(StrataData.Layers[0].BSDFs[0], STRATA_BSDF_TYPE_UNLIT);
-	BSDF_SETEMISSIVE				(StrataData.Layers[0].BSDFs[0], Emissive);
-	UNLIT_TRANSMITTANCE				(StrataData.Layers[0].BSDFs[0]) = TransmittanceColor;
+	BSDF_SETTYPE					(StrataData.InlinedBSDF, STRATA_BSDF_TYPE_UNLIT);
+	BSDF_SETEMISSIVE				(StrataData.InlinedBSDF, Emissive);
+	UNLIT_TRANSMITTANCE				(StrataData.InlinedBSDF) = TransmittanceColor;
 
 #if STRATA_INLINE_SHADING
-	StrataData.Layers[0].BSDFs[0].Coverage = 1.0f;
+	StrataData.InlinedBSDF.Coverage = 1.0f;
 #endif
 
+	// STRATA_TODO BSDF OP
+
 	return StrataData;
 }
 
 FStrataData GetStrataHairBSDF(float3 BaseColor, float Scatter, float Specular, float Roughness, float Backlit, float3 Emissive, uint SharedLocalBasisIndex)
 {
-	FStrataData StrataData = GetInitialisedStrataDataWithReservedFirstSlot();
+	FStrataData StrataData = GetInitialisedStrataData();
 
 	// Enable complex transmittance only for hair using cards/strands vertex factories
 	#ifndef USE_HAIR_COMPLEX_TRANSMITTANCE
@@ -772,20 +821,22 @@ FStrataData GetStrataHairBSDF(float3 BaseColor, float Scatter, float Specular, f
 	  #endif
 	#endif
 
-	BSDF_SETTYPE					(StrataData.Layers[0].BSDFs[0], STRATA_BSDF_TYPE_HAIR);
-	BSDF_SETSHAREDLOCALBASISID		(StrataData.Layers[0].BSDFs[0], SharedLocalBasisIndex);
-	BSDF_SETEMISSIVE				(StrataData.Layers[0].BSDFs[0], Emissive);
-	HAIR_BASECOLOR					(StrataData.Layers[0].BSDFs[0]) = BaseColor;
-	HAIR_SCATTER					(StrataData.Layers[0].BSDFs[0]) = Scatter;
-	HAIR_ROUGHNESS					(StrataData.Layers[0].BSDFs[0]) = Roughness;
-	HAIR_SPECULAR					(StrataData.Layers[0].BSDFs[0]) = Specular;
-	HAIR_BACKLIT					(StrataData.Layers[0].BSDFs[0]) = Backlit;
-	HAIR_COMPLEXTRANSMITTANCE		(StrataData.Layers[0].BSDFs[0]) = USE_HAIR_COMPLEX_TRANSMITTANCE;
+	BSDF_SETTYPE					(StrataData.InlinedBSDF, STRATA_BSDF_TYPE_HAIR);
+	BSDF_SETSHAREDLOCALBASISID		(StrataData.InlinedBSDF, SharedLocalBasisIndex);
+	BSDF_SETEMISSIVE				(StrataData.InlinedBSDF, Emissive);
+	HAIR_BASECOLOR					(StrataData.InlinedBSDF) = BaseColor;
+	HAIR_SCATTER					(StrataData.InlinedBSDF) = Scatter;
+	HAIR_ROUGHNESS					(StrataData.InlinedBSDF) = Roughness;
+	HAIR_SPECULAR					(StrataData.InlinedBSDF) = Specular;
+	HAIR_BACKLIT					(StrataData.InlinedBSDF) = Backlit;
+	HAIR_COMPLEXTRANSMITTANCE		(StrataData.InlinedBSDF) = USE_HAIR_COMPLEX_TRANSMITTANCE;
 
 #if STRATA_INLINE_SHADING
-	StrataData.Layers[0].BSDFs[0].Coverage = 1.0f;
+	StrataData.InlinedBSDF.Coverage = 1.0f;
 #endif
 
+	// STRATA_TODO BSDF OP
+
 	return StrataData;
 }
 
@@ -793,27 +844,29 @@ FStrataData GetStrataSingleLayerWaterBSDF(
 	float3 BaseColor, float Metallic, float Specular, float Roughness, float3 Emissive, float TopMaterialOpacity,
 	float3 WaterAlbedo, float3 WaterExtinction, float WaterPhaseG, float3 ColorScaleBehindWater, uint SharedLocalBasisIndex)
 {
-	FStrataData StrataData = GetInitialisedStrataDataWithReservedFirstSlot();
+	FStrataData StrataData = GetInitialisedStrataData();
 
-	BSDF_SETTYPE					(StrataData.Layers[0].BSDFs[0], STRATA_BSDF_TYPE_SINGLELAYERWATER);
-	BSDF_SETSHAREDLOCALBASISID		(StrataData.Layers[0].BSDFs[0], SharedLocalBasisIndex);
-	BSDF_SETEMISSIVE				(StrataData.Layers[0].BSDFs[0], Emissive);
-	SLW_BASECOLOR					(StrataData.Layers[0].BSDFs[0]) = BaseColor;
-	SLW_METALLIC					(StrataData.Layers[0].BSDFs[0]) = Metallic;
-	SLW_SPECULAR					(StrataData.Layers[0].BSDFs[0]) = Specular;
-	SLW_ROUGHNESS					(StrataData.Layers[0].BSDFs[0]) = Roughness;
-	SLW_TOPMATERIALOPACITY			(StrataData.Layers[0].BSDFs[0]) = TopMaterialOpacity;
+	BSDF_SETTYPE					(StrataData.InlinedBSDF, STRATA_BSDF_TYPE_SINGLELAYERWATER);
+	BSDF_SETSHAREDLOCALBASISID		(StrataData.InlinedBSDF, SharedLocalBasisIndex);
+	BSDF_SETEMISSIVE				(StrataData.InlinedBSDF, Emissive);
+	SLW_BASECOLOR					(StrataData.InlinedBSDF) = BaseColor;
+	SLW_METALLIC					(StrataData.InlinedBSDF) = Metallic;
+	SLW_SPECULAR					(StrataData.InlinedBSDF) = Specular;
+	SLW_ROUGHNESS					(StrataData.InlinedBSDF) = Roughness;
+	SLW_TOPMATERIALOPACITY			(StrataData.InlinedBSDF) = TopMaterialOpacity;
 #if STRATA_INLINE_SINGLELAYERWATER
-	SLW_WATERALBEDO					(StrataData.Layers[0].BSDFs[0]) = WaterAlbedo;
-	SLW_WATEREXTINCTION				(StrataData.Layers[0].BSDFs[0]) = WaterExtinction;
-	SLW_WATERPHASEG					(StrataData.Layers[0].BSDFs[0]) = WaterPhaseG;
-	SLW_COLORSCALEBEHINDWATER		(StrataData.Layers[0].BSDFs[0]) = ColorScaleBehindWater;
+	SLW_WATERALBEDO					(StrataData.InlinedBSDF) = WaterAlbedo;
+	SLW_WATEREXTINCTION				(StrataData.InlinedBSDF) = WaterExtinction;
+	SLW_WATERPHASEG					(StrataData.InlinedBSDF) = WaterPhaseG;
+	SLW_COLORSCALEBEHINDWATER		(StrataData.InlinedBSDF) = ColorScaleBehindWater;
 #endif
 
 #if STRATA_INLINE_SHADING
-	StrataData.Layers[0].BSDFs[0].Coverage = 1.0f;
+	StrataData.InlinedBSDF.Coverage = 1.0f;
 #endif
 
+	// STRATA_TODO BSDF OP
+
 	return StrataData;
 }
 
@@ -924,6 +977,8 @@ struct FStrataPixelHeader
 #if STRATA_INLINE_SHADING
 	// All the shared local bases in VGPRS. This is only used in basepass and during forward rendering.
 	FSharedLocalBases SharedLocalBases;
+	// The strata material BSDF topology organised as a tree
+	FStrataTree StrataTree;
 #endif
 #if STRATA_DEFERRED_SHADING
 	FStrataMaterialContainer StrataBuffer;
@@ -969,6 +1024,7 @@ FStrataPixelHeader InitialiseStrataPixelHeader()
 	Out.BSDFCount = 0;
 #if STRATA_INLINE_SHADING
 	Out.SharedLocalBases = (FSharedLocalBases)0;
+	Out.StrataTree = (FStrataTree)0;
 #endif
 #if STRATA_DEFERRED_SHADING
 	Out.SharedLocalBasesIndexOffset = 0;
@@ -1038,100 +1094,125 @@ float3x3 StrataGetBSDFSharedBasis(in FStrataPixelHeader StrataPixelHeader, in FS
 ///////////////////////////////////////////////////////////////////////////////
 // Operator nodes
 
-FStrataData StrataAdd(FStrataData A, FStrataData B)
+FStrataData StrataAdd(FStrataData A, FStrataData B, inout FStrataTree StrataTree, int OperatorIndex, uint MaxDistanceFromLeaves)
 {
-	FStrataData StrataData = A;
+	FStrataData StrataData = GetInitialisedStrataData();
 
-	// Append each BSDF from B to A, with same layer position
-	STRATA_UNROLL_N(STRATA_CLAMPED_LAYER_COUNT)
-	for (int Layer = 0; Layer < B.LayerCount; ++Layer)
-	{
-		FStrataLayer ALayer = A.Layers[Layer];
-		FStrataLayer BLayer = B.Layers[Layer];
-		STRATA_UNROLL
-		for (int BSDF = 0; BSDF < BLayer.BSDFCount; BSDF++)
-		{
-			StrataData.Layers[Layer].BSDFs[ALayer.BSDFCount + BSDF] = BLayer.BSDFs[BSDF];
-		}
+	StrataTree.OperatorCount++;
+	StrataData.OperatorIndex = OperatorIndex;
 
-		StrataData.Layers[Layer].BSDFCount = ALayer.BSDFCount + BLayer.BSDFCount;
-	}
+	StrataTree.Operators[OperatorIndex]							= GetInitialisedStrataOperator();
+	StrataTree.Operators[OperatorIndex].MaxDistanceFromLeaves	= MaxDistanceFromLeaves;
+	StrataTree.Operators[OperatorIndex].Type					= STRATA_OPERATOR_ADD;
+	StrataTree.Operators[OperatorIndex].AIndex					= A.OperatorIndex;
+	StrataTree.Operators[OperatorIndex].BIndex					= B.OperatorIndex;
 
-	StrataData.LayerCount = max(A.LayerCount, B.LayerCount);
+	StrataTree.Operators[A.OperatorIndex].TopIndex				= OperatorIndex;
+	StrataTree.Operators[B.OperatorIndex].TopIndex				= OperatorIndex;
 
 	return StrataData;
 }
 
-FStrataData StrataWeight(FStrataData A, float Weight)
+FStrataData StrataWeight(FStrataData A, float Weight, inout FStrataTree StrataTree, int OperatorIndex, uint MaxDistanceFromLeaves)
 {
 	const float SafeWeight = saturate(Weight);
-	FStrataData StrataData = A;
 
-	// Apply the weight to each BSDF
-	STRATA_UNROLL_N(STRATA_CLAMPED_LAYER_COUNT)
-	for (int Layer = 0; Layer < StrataData.LayerCount; ++Layer)
-	{
-		STRATA_UNROLL
-		for (int BSDF = 0; BSDF < StrataData.Layers[Layer].BSDFCount; ++BSDF)
-		{
-#if STRATA_INLINE_SHADING
-			StrataData.Layers[Layer].BSDFs[BSDF].Coverage *= SafeWeight;
+	FStrataData StrataData = GetInitialisedStrataData();
 
-			if (StrataData.Layers[Layer].BSDFs[BSDF].bPremultipliedCoverage) 
-			{
-				// Only slab can follow parameter blending rule: we can encounter pre weighted lighting contribution according to coverage.
-				SLAB_DIFFUSEALBEDO(StrataData.Layers[Layer].BSDFs[BSDF]) *= SafeWeight;
-				SLAB_F0(StrataData.Layers[Layer].BSDFs[BSDF]) *= SafeWeight;
-				SLAB_F90(StrataData.Layers[Layer].BSDFs[BSDF]) *= SafeWeight;
-				SLAB_FUZZ_COLOR(StrataData.Layers[Layer].BSDFs[BSDF]) *= SafeWeight;
-			}
-#endif
-		}
-	}
+	StrataTree.OperatorCount++;
+	StrataData.OperatorIndex = OperatorIndex;
+
+	StrataTree.Operators[OperatorIndex]							= GetInitialisedStrataOperator();
+	StrataTree.Operators[OperatorIndex].MaxDistanceFromLeaves	= MaxDistanceFromLeaves;
+	StrataTree.Operators[OperatorIndex].Type					= STRATA_OPERATOR_WEIGHT;
+	StrataTree.Operators[OperatorIndex].Weight					= Weight;
+	StrataTree.Operators[OperatorIndex].AIndex					= A.OperatorIndex;
+
+	StrataTree.Operators[A.OperatorIndex].TopIndex				= OperatorIndex;
 
 	return StrataData;
 }
 
-FStrataData StrataHorizontalMixing(FStrataData Background, FStrataData Foreground, float Mix)
+FStrataData StrataHorizontalMixing(FStrataData Background, FStrataData Foreground, float Mix, inout FStrataTree StrataTree, int OperatorIndex, uint MaxDistanceFromLeaves)
 {
 	const float ForegroundMixFactor = saturate(Mix);
 	const float BackgroundMixFactor = 1.0 - ForegroundMixFactor;
 
-	FStrataData A = StrataWeight(Foreground, ForegroundMixFactor);
-	FStrataData B = StrataWeight(Background, BackgroundMixFactor);
-	return StrataAdd(A, B);
+	FStrataData StrataData = GetInitialisedStrataData();
+
+	StrataTree.OperatorCount++;
+	StrataData.OperatorIndex = OperatorIndex;
+	
+	StrataTree.Operators[OperatorIndex]							= GetInitialisedStrataOperator();
+	StrataTree.Operators[OperatorIndex].MaxDistanceFromLeaves	= MaxDistanceFromLeaves;
+	StrataTree.Operators[OperatorIndex].Type					= STRATA_OPERATOR_HORIZONTAL;
+	StrataTree.Operators[OperatorIndex].Weight					= Mix;
+	StrataTree.Operators[OperatorIndex].AIndex					= Background.OperatorIndex;
+	StrataTree.Operators[OperatorIndex].BIndex					= Foreground.OperatorIndex;
+
+	StrataTree.Operators[Background.OperatorIndex].TopIndex		= OperatorIndex;
+	StrataTree.Operators[Foreground.OperatorIndex].TopIndex		= OperatorIndex;
+
+	return StrataData;
 }
 
-FStrataData StrataVerticalLayering(FStrataData Top, FStrataData Base)
+FStrataData StrataVerticalLayering(FStrataData Top, FStrataData Base, inout FStrataTree StrataTree, int OperatorIndex, uint MaxDistanceFromLeaves)
 {
-	FStrataData StrataData = Top;
+	FStrataData StrataData = GetInitialisedStrataData();
 
-	// Add each layer from Base under Top
-	int TopLayerCount = Top.LayerCount;
-	STRATA_UNROLL_N(STRATA_CLAMPED_LAYER_COUNT)
-	for (int Layer = 0; Layer < Base.LayerCount; ++Layer)
-	{
-		StrataData.Layers[TopLayerCount + Layer] = Base.Layers[Layer];
-	}
+	StrataTree.OperatorCount++;
+	StrataData.OperatorIndex = OperatorIndex;
+	
+	StrataTree.Operators[OperatorIndex]							= GetInitialisedStrataOperator();
+	StrataTree.Operators[OperatorIndex].MaxDistanceFromLeaves	= MaxDistanceFromLeaves;
+	StrataTree.Operators[OperatorIndex].Type					= STRATA_OPERATOR_VERTICAL;
+	StrataTree.Operators[OperatorIndex].AIndex					= Top.OperatorIndex;
+	StrataTree.Operators[OperatorIndex].BIndex					= Base.OperatorIndex;
 
-	StrataData.LayerCount += Base.LayerCount;
+	StrataTree.Operators[Top.OperatorIndex].TopIndex			= OperatorIndex;
+	StrataTree.Operators[Base.OperatorIndex].TopIndex			= OperatorIndex;
 
 	return StrataData;
 }
 
+// Take the inlined BSDF resulting from parameter blending (or BSDF creation) and register it to the StrataTree
+FStrataData PromoteParameterBlendedBSDFToOperator(FStrataData StrataData, inout FStrataTree StrataTree, int OperatorIndex, int BSDFIndex, int bIsBottom, int bIsTop)
+{
+	StrataTree.BSDFCount++;
+	StrataTree.OperatorCount++;
+
+	StrataTree.BSDFs[BSDFIndex] = StrataData.InlinedBSDF;
+	
+	StrataTree.Operators[OperatorIndex]							= GetInitialisedStrataOperator();
+	StrataTree.Operators[OperatorIndex].MaxDistanceFromLeaves	= 0;
+	StrataTree.Operators[OperatorIndex].Type					= STRATA_OPERATOR_BSDF;
+	StrataTree.Operators[OperatorIndex].AIndex					= BSDFIndex;
+
+	StrataTree.BSDFs[BSDFIndex].OperatorIndex					= OperatorIndex;
+	StrataTree.BSDFs[BSDFIndex].bIsBottom						= bIsBottom;
+	StrataTree.BSDFs[BSDFIndex].bIsTop							= bIsTop;
+
+	StrataData.OperatorIndex = OperatorIndex;
+	return StrataData;
+}
+
 float HorizontalMixingParameterBlendingBSDFCoverageToNormalMix(FStrataData Background, FStrataData Foreground, float HorizontalMixValue)
 {
 	// This is a normalised mix so even though both BSDF have weights, the sum of weights used to combined them should always be 1.
+	// This also account for relative BSDF coverage.
 #if STRATA_INLINE_SHADING
-	const float SafeBackCoverage = saturate(Background.Layers[0].BSDFs[0].Coverage);
-	const float SafeForeCoverage = saturate(Foreground.Layers[0].BSDFs[0].Coverage);
+	const float SafeBackCoverage = saturate(Background.InlinedBSDF.Coverage);
+	const float SafeForeCoverage = saturate(Foreground.InlinedBSDF.Coverage);
+	const float ForegroundMixFactor = saturate(HorizontalMixValue);
+	const float BackgroundMixFactor = 1.0 - ForegroundMixFactor;
 
-	return HorizontalMixValue;
+	return (SafeForeCoverage * ForegroundMixFactor) / max(1e-5, SafeForeCoverage * ForegroundMixFactor + SafeBackCoverage * BackgroundMixFactor);
 #else
 	return 0.5f;
 #endif
 }
 
+
 // Note about parameter blending
 //  - We can only parameter blend Slab nodes. This error is handled in the compiler.
 //  - We can only parameter blend one BSDF into another single BSDF. More complex topology are not handled and result in a compiler error. This could be extended to a single BSDF per layer.
@@ -1141,22 +1222,13 @@ float HorizontalMixingParameterBlendingBSDFCoverageToNormalMix(FStrataData Backg
 // ==> NOTE: Always pair with the compiler behavior in StrataCompilationInfoHorizontalMixingParamBlend
 FStrataData StrataHorizontalMixingParameterBlending(FStrataData Background, FStrataData Foreground, float HorizontalMixValue, float NormalMixCodeChunk, uint NewNormalIndex)
 {
-	const float ForegroundLightingMixFactor = saturate(HorizontalMixValue);	// Only the horizontal mixing value
 	const float ForegroundOtherMixFactor = saturate(NormalMixCodeChunk);	// Horizontal mixing value combined with coverage
+	const float ForegroundLightingMixFactor = saturate(HorizontalMixValue);
 
-	FStrataData Result = GetInitialisedStrataDataWithReservedFirstSlot();
-#define ResultBSDF		Result.Layers[0].BSDFs[0]
-#define ForegroundBSDF	Foreground.Layers[0].BSDFs[0]
-#define BackgroundBSDF	Background.Layers[0].BSDFs[0]
-
-#if STRATA_INLINE_SHADING
-	const float ForeLightingFactor = ForegroundBSDF.bPremultipliedCoverage ? 1.0f : saturate(ForegroundBSDF.Coverage);
-	const float BackLightingFactor = BackgroundBSDF.bPremultipliedCoverage ? 1.0f : saturate(BackgroundBSDF.Coverage);
-	ResultBSDF.bPremultipliedCoverage = true;
-#else
-	const float ForeLightingFactor = 1.0f;
-	const float BackLightingFactor = 1.0f;
-#endif
+	FStrataData Result = GetInitialisedStrataData();
+#define ResultBSDF		Result.InlinedBSDF
+#define ForegroundBSDF	Foreground.InlinedBSDF
+#define BackgroundBSDF	Background.InlinedBSDF
 
 	// We can only parameter blend the slab 
 	BSDF_SETTYPE(ResultBSDF,			STRATA_BSDF_TYPE_SLAB);
@@ -1173,18 +1245,18 @@ FStrataData StrataHorizontalMixingParameterBlending(FStrataData Background, FStr
 	BSDF_SETHASTHINFILM(ResultBSDF,		BSDF_GETHASTHINFILM(ForegroundBSDF)				| BSDF_GETHASTHINFILM(BackgroundBSDF));
 	BSDF_SETHASFUZZ(ResultBSDF,			BSDF_GETHASFUZZ(ForegroundBSDF)					| BSDF_GETHASFUZZ(BackgroundBSDF));
 
-	BSDF_SETEMISSIVE(ResultBSDF,		lerp(BackLightingFactor * BSDF_GETEMISSIVE(BackgroundBSDF),			ForeLightingFactor * BSDF_GETEMISSIVE(ForegroundBSDF),			ForegroundLightingMixFactor));
-	SLAB_DIFFUSEALBEDO(ResultBSDF) =	lerp(BackLightingFactor * SLAB_DIFFUSEALBEDO(BackgroundBSDF),		ForeLightingFactor * SLAB_DIFFUSEALBEDO(ForegroundBSDF),		ForegroundLightingMixFactor);
-	SLAB_F0(ResultBSDF) =				lerp(BackLightingFactor * SLAB_F0(BackgroundBSDF),					ForeLightingFactor * SLAB_F0(ForegroundBSDF),					ForegroundLightingMixFactor);
-	SLAB_F90(ResultBSDF) =				lerp(BackLightingFactor * SLAB_F90(BackgroundBSDF),					ForeLightingFactor * SLAB_F90(ForegroundBSDF),					ForegroundLightingMixFactor);
-	SLAB_FUZZ_COLOR(ResultBSDF)	=		lerp(BackLightingFactor * SLAB_FUZZ_COLOR(BackgroundBSDF),			ForeLightingFactor * SLAB_FUZZ_COLOR(ForegroundBSDF),			ForegroundLightingMixFactor);
+	BSDF_SETEMISSIVE(ResultBSDF,		lerp(BSDF_GETEMISSIVE(BackgroundBSDF),			BSDF_GETEMISSIVE(ForegroundBSDF),			ForegroundOtherMixFactor));
+	SLAB_DIFFUSEALBEDO(ResultBSDF) =	lerp(SLAB_DIFFUSEALBEDO(BackgroundBSDF),		SLAB_DIFFUSEALBEDO(ForegroundBSDF),			ForegroundOtherMixFactor);
+	SLAB_F0(ResultBSDF) =				lerp(SLAB_F0(BackgroundBSDF),					SLAB_F0(ForegroundBSDF),					ForegroundOtherMixFactor);
+	SLAB_F90(ResultBSDF) =				lerp(SLAB_F90(BackgroundBSDF),					SLAB_F90(ForegroundBSDF),					ForegroundOtherMixFactor);
+	SLAB_FUZZ_COLOR(ResultBSDF)	=		lerp(SLAB_FUZZ_COLOR(BackgroundBSDF),			SLAB_FUZZ_COLOR(ForegroundBSDF),			ForegroundOtherMixFactor);
 
-	BSDF_SETTHICKNESSCM(ResultBSDF,		lerp(BSDF_GETTHICKNESSCM(BackgroundBSDF),							BSDF_GETTHICKNESSCM(ForegroundBSDF),							ForegroundOtherMixFactor));
-	SLAB_ROUGHNESS(ResultBSDF) =		lerp(SLAB_ROUGHNESS(BackgroundBSDF),								SLAB_ROUGHNESS(ForegroundBSDF),									ForegroundOtherMixFactor);
-	SLAB_ANISOTROPY(ResultBSDF) =		lerp(SLAB_ANISOTROPY(BackgroundBSDF),								SLAB_ANISOTROPY(ForegroundBSDF),								ForegroundOtherMixFactor);
-	SLAB_HAZINESS(ResultBSDF) =			lerp(SLAB_HAZINESS(BackgroundBSDF),									SLAB_HAZINESS(ForegroundBSDF),									ForegroundOtherMixFactor);
-	SLAB_THINFILM_THICKNESS(ResultBSDF)=lerp(SLAB_THINFILM_THICKNESS(BackgroundBSDF),						SLAB_THINFILM_THICKNESS(ForegroundBSDF),						ForegroundOtherMixFactor);
-	SLAB_FUZZ_AMOUNT(ResultBSDF) =		lerp(SLAB_FUZZ_AMOUNT(BackgroundBSDF),								SLAB_FUZZ_AMOUNT(ForegroundBSDF),								ForegroundOtherMixFactor);
+	BSDF_SETTHICKNESSCM(ResultBSDF,		lerp(BSDF_GETTHICKNESSCM(BackgroundBSDF),		BSDF_GETTHICKNESSCM(ForegroundBSDF),		ForegroundOtherMixFactor));
+	SLAB_ROUGHNESS(ResultBSDF) =		lerp(SLAB_ROUGHNESS(BackgroundBSDF),			SLAB_ROUGHNESS(ForegroundBSDF),				ForegroundOtherMixFactor);
+	SLAB_ANISOTROPY(ResultBSDF) =		lerp(SLAB_ANISOTROPY(BackgroundBSDF),			SLAB_ANISOTROPY(ForegroundBSDF),			ForegroundOtherMixFactor);
+	SLAB_HAZINESS(ResultBSDF) =			lerp(SLAB_HAZINESS(BackgroundBSDF),				SLAB_HAZINESS(ForegroundBSDF),				ForegroundOtherMixFactor);
+	SLAB_THINFILM_THICKNESS(ResultBSDF)=lerp(SLAB_THINFILM_THICKNESS(BackgroundBSDF),	SLAB_THINFILM_THICKNESS(ForegroundBSDF),	ForegroundOtherMixFactor);
+	SLAB_FUZZ_AMOUNT(ResultBSDF) =		lerp(SLAB_FUZZ_AMOUNT(BackgroundBSDF),			SLAB_FUZZ_AMOUNT(ForegroundBSDF),			ForegroundOtherMixFactor);
 
 	if (BSDF_GETHASSSSPROFILE(ResultBSDF))
 	{
@@ -1209,7 +1281,7 @@ FStrataData StrataHorizontalMixingParameterBlending(FStrataData Background, FStr
 	}
 	else
 	{
-		SLAB_SSSDMFP(ResultBSDF) =		lerp(SLAB_SSSDMFP(BackgroundBSDF),								SLAB_SSSDMFP(ForegroundBSDF),										ForegroundOtherMixFactor);
+		SLAB_SSSDMFP(ResultBSDF) =		lerp(SLAB_SSSDMFP(BackgroundBSDF),				SLAB_SSSDMFP(ForegroundBSDF),				ForegroundOtherMixFactor);
 
 		// We still try to disable SSS if possible: for performance and also to preserve the surface color (especially specular reflection color).
 		if (all(SLAB_SSSDMFP(ResultBSDF) == 0.0f))
@@ -1221,10 +1293,10 @@ FStrataData StrataHorizontalMixingParameterBlending(FStrataData Background, FStr
 
 #if STRATA_INLINE_SHADING
 	// At this stage we do not know if this strata BSDF is going to be using simple volume lighting so always store the DMFP on the side to not override SSSPROFILE data.
-	ResultBSDF.TmpDMFP =				lerp(BackgroundBSDF.TmpDMFP,									ForegroundBSDF.TmpDMFP,												ForegroundOtherMixFactor);
+	ResultBSDF.TmpDMFP =				lerp(BackgroundBSDF.TmpDMFP,					ForegroundBSDF.TmpDMFP,						ForegroundOtherMixFactor);
 
 	// We blend the parameter above, and we also need to lerp the Coverage to match the total energy.
-	ResultBSDF.Coverage =				lerp(BackgroundBSDF.Coverage,									ForegroundBSDF.Coverage,											ForegroundLightingMixFactor);
+	ResultBSDF.Coverage =				lerp(BackgroundBSDF.Coverage,					ForegroundBSDF.Coverage,					HorizontalMixValue);
 #endif
 
 #undef ResultBSDF
@@ -1237,8 +1309,8 @@ float AddParameterBlendingBSDFCoverageToNormalMix(FStrataData AStrata, FStrataDa
 {
 	// This is a normalised mix so even though both BSDF have weights, the sum of weights used to combined them should always be 1.
 #if STRATA_INLINE_SHADING
-	const float SafeABSDFCoverage = saturate(AStrata.Layers[0].BSDFs[0].Coverage);
-	const float SafeBBSDFCoverage = saturate(BStrata.Layers[0].BSDFs[0].Coverage);
+	const float SafeABSDFCoverage = saturate(AStrata.InlinedBSDF.Coverage);
+	const float SafeBBSDFCoverage = saturate(BStrata.InlinedBSDF.Coverage);
 	const float AMixFactor = SafeABSDFCoverage / (1e-10 + SafeABSDFCoverage + SafeBBSDFCoverage);
 	return AMixFactor;
 #else
@@ -1249,50 +1321,41 @@ float AddParameterBlendingBSDFCoverageToNormalMix(FStrataData AStrata, FStrataDa
 // ==> NOTE: Always pair with the compiler behavior in StrataCompilationInfoAddParamBlend
 FStrataData StrataAddParameterBlending(FStrataData A, FStrataData B, float AMixFactor, uint NewNormalIndex)
 {
-	FStrataData Result = GetInitialisedStrataDataWithReservedFirstSlot();
-#define ResultBSDF		Result.Layers[0].BSDFs[0]
-#define ABSDF			A.Layers[0].BSDFs[0]
-#define BBSDF			B.Layers[0].BSDFs[0]
-
-#if STRATA_INLINE_SHADING
-	const float ALightingFactor = ABSDF.bPremultipliedCoverage ? 1.0f : saturate(ABSDF.Coverage);
-	const float BLightingFactor = BBSDF.bPremultipliedCoverage ? 1.0f : saturate(BBSDF.Coverage);
-	ResultBSDF.bPremultipliedCoverage = true;
-#else
-	const float ALightingFactor = 1.0f;
-	const float BLightingFactor = 1.0f;
-#endif
+	FStrataData Result = GetInitialisedStrataData();
+#define ResultBSDF		Result.InlinedBSDF
+#define ABSDF			A.InlinedBSDF
+#define BBSDF			B.InlinedBSDF
 
 	// We can only parameter blend the slab 
 	BSDF_SETTYPE(ResultBSDF,			STRATA_BSDF_TYPE_SLAB);
 	BSDF_SETSHAREDLOCALBASISID(ResultBSDF,NewNormalIndex);
 
 	BSDF_SETISTOPLAYER(ResultBSDF,		0);
-	BSDF_SETHASANISOTROPY(ResultBSDF,	BSDF_GETHASANISOTROPY(BBSDF)							| BSDF_GETHASANISOTROPY(ABSDF));
-	BSDF_SETHASSSS(ResultBSDF,			BSDF_GETHASSSS(BBSDF)									| BSDF_GETHASSSS(ABSDF));
-	BSDF_SETHASSSSPROFILE(ResultBSDF,	BSDF_GETHASSSSPROFILE(BBSDF)							| BSDF_GETHASSSSPROFILE(ABSDF));	// Problem: cannot blend SSS profiles...
-	BSDF_SETISTHIN(ResultBSDF,			BSDF_GETISTHIN(BBSDF)									| BSDF_GETISTHIN(ABSDF));			// STRATA_TODO fix the MPF/DMFP situation for that to really work
-	BSDF_SETHASDMFP(ResultBSDF,			BSDF_GETHASDMFP(BBSDF)									| BSDF_GETHASDMFP(ABSDF));
-	BSDF_SETHASHAZINESS(ResultBSDF,		BSDF_GETHASHAZINESS(BBSDF)								| BSDF_GETHASHAZINESS(ABSDF));
-	BSDF_SETHASF90(ResultBSDF,			BSDF_GETHASF90(BBSDF)									| BSDF_GETHASF90(ABSDF));
-	BSDF_SETHASTHINFILM(ResultBSDF,		BSDF_GETHASTHINFILM(BBSDF)								| BSDF_GETHASTHINFILM(ABSDF));
-	BSDF_SETHASFUZZ(ResultBSDF,			BSDF_GETHASFUZZ(BBSDF)									| BSDF_GETHASFUZZ(ABSDF));
+	BSDF_SETHASANISOTROPY(ResultBSDF,	BSDF_GETHASANISOTROPY(BBSDF)					| BSDF_GETHASANISOTROPY(ABSDF));
+	BSDF_SETHASSSS(ResultBSDF,			BSDF_GETHASSSS(BBSDF)							| BSDF_GETHASSSS(ABSDF));
+	BSDF_SETHASSSSPROFILE(ResultBSDF,	BSDF_GETHASSSSPROFILE(BBSDF)					| BSDF_GETHASSSSPROFILE(ABSDF));	// Problem: cannot blend SSS profiles...
+	BSDF_SETISTHIN(ResultBSDF,			BSDF_GETISTHIN(BBSDF)							| BSDF_GETISTHIN(ABSDF));			// STRATA_TODO fix the MPF/DMFP situation for that to really work
+	BSDF_SETHASDMFP(ResultBSDF,			BSDF_GETHASDMFP(BBSDF)							| BSDF_GETHASDMFP(ABSDF));
+	BSDF_SETHASHAZINESS(ResultBSDF,		BSDF_GETHASHAZINESS(BBSDF)						| BSDF_GETHASHAZINESS(ABSDF));
+	BSDF_SETHASF90(ResultBSDF,			BSDF_GETHASF90(BBSDF)							| BSDF_GETHASF90(ABSDF));
+	BSDF_SETHASTHINFILM(ResultBSDF,		BSDF_GETHASTHINFILM(BBSDF)						| BSDF_GETHASTHINFILM(ABSDF));
+	BSDF_SETHASFUZZ(ResultBSDF,			BSDF_GETHASFUZZ(BBSDF)							| BSDF_GETHASFUZZ(ABSDF));
 
 
 	// Some parameters will contribute to added luminance: in this case we add them together and saturate to not go out of the safe range
-	BSDF_SETEMISSIVE(ResultBSDF,		BLightingFactor * BSDF_GETEMISSIVE(BBSDF)				+ ALightingFactor * BSDF_GETEMISSIVE(ABSDF));
-	SLAB_DIFFUSEALBEDO(ResultBSDF) =	saturate(BLightingFactor * SLAB_DIFFUSEALBEDO(BBSDF)	+ ALightingFactor * SLAB_DIFFUSEALBEDO(ABSDF));
-	SLAB_F0(ResultBSDF) =				saturate(BLightingFactor * SLAB_F0(BBSDF)				+ ALightingFactor * SLAB_F0(ABSDF));
-	SLAB_F90(ResultBSDF) =				saturate(BLightingFactor * SLAB_F90(BBSDF)				+ ALightingFactor * SLAB_F90(ABSDF));
-	SLAB_FUZZ_COLOR(ResultBSDF) =		saturate(BLightingFactor * SLAB_FUZZ_COLOR(BBSDF)		+ ALightingFactor * SLAB_FUZZ_COLOR(ABSDF));
+	BSDF_SETEMISSIVE(ResultBSDF,		BSDF_GETEMISSIVE(ABSDF)								+ BSDF_GETEMISSIVE(ABSDF));
+	SLAB_DIFFUSEALBEDO(ResultBSDF) =	saturate(SLAB_DIFFUSEALBEDO(BBSDF)					+ SLAB_DIFFUSEALBEDO(ABSDF));
+	SLAB_F0(ResultBSDF) =				saturate(SLAB_F0(BBSDF)								+ SLAB_F0(ABSDF));
+	SLAB_F90(ResultBSDF) =				saturate(SLAB_F90(BBSDF)							+ SLAB_F90(ABSDF));
+	SLAB_FUZZ_COLOR(ResultBSDF) =		saturate(SLAB_FUZZ_COLOR(BBSDF)						+ SLAB_FUZZ_COLOR(ABSDF));
 
 	// Some parameters are not contributing to added luminance: in this case we simply lerp them
-	BSDF_SETTHICKNESSCM(ResultBSDF,		lerp(BSDF_GETTHICKNESSCM(BBSDF),						BSDF_GETTHICKNESSCM(ABSDF),					AMixFactor));
-	SLAB_ROUGHNESS(ResultBSDF) =		lerp(SLAB_ROUGHNESS(BBSDF),								SLAB_ROUGHNESS(ABSDF),						AMixFactor);
-	SLAB_ANISOTROPY(ResultBSDF) =		lerp(SLAB_ANISOTROPY(BBSDF),							SLAB_ANISOTROPY(ABSDF),						AMixFactor);
-	SLAB_HAZINESS(ResultBSDF) =			lerp(SLAB_HAZINESS(BBSDF),								SLAB_HAZINESS(ABSDF),						AMixFactor);
+	BSDF_SETTHICKNESSCM(ResultBSDF,		lerp(BSDF_GETTHICKNESSCM(BBSDF),				BSDF_GETTHICKNESSCM(ABSDF),					AMixFactor));
+	SLAB_ROUGHNESS(ResultBSDF) =		lerp(SLAB_ROUGHNESS(BBSDF),						SLAB_ROUGHNESS(ABSDF),						AMixFactor);
+	SLAB_ANISOTROPY(ResultBSDF) =		lerp(SLAB_ANISOTROPY(BBSDF),					SLAB_ANISOTROPY(ABSDF),						AMixFactor);
+	SLAB_HAZINESS(ResultBSDF) =			lerp(SLAB_HAZINESS(BBSDF),						SLAB_HAZINESS(ABSDF),						AMixFactor);
 	SLAB_THINFILM_THICKNESS(ResultBSDF) = 0;// This does not support coverage applied on pararmeters so we simply disable ThinFilm. lerp(SLAB_THINFILM_THICKNESS(BBSDF), SLAB_THINFILM_THICKNESS(ABSDF), AMixFactor);
-	SLAB_FUZZ_AMOUNT(ResultBSDF) =		lerp(SLAB_FUZZ_AMOUNT(BBSDF),							SLAB_FUZZ_AMOUNT(ABSDF),					AMixFactor);
+	SLAB_FUZZ_AMOUNT(ResultBSDF) =		lerp(SLAB_FUZZ_AMOUNT(BBSDF),					SLAB_FUZZ_AMOUNT(ABSDF),					AMixFactor);
 
 	if (BSDF_GETHASSSSPROFILE(ResultBSDF))
 	{
@@ -1310,17 +1373,17 @@ FStrataData StrataAddParameterBlending(FStrataData A, FStrataData B, float AMixF
 	}
 	else
 	{
-		SLAB_SSSDMFP(ResultBSDF) =		lerp(SLAB_SSSDMFP(BBSDF),								SLAB_SSSDMFP(ABSDF),						AMixFactor);
+		SLAB_SSSDMFP(ResultBSDF) =		lerp(SLAB_SSSDMFP(BBSDF),						SLAB_SSSDMFP(ABSDF),						AMixFactor);
 	}
 
 
 #if STRATA_INLINE_SHADING
 	// At this stage we do not know if this strata BSDF is going to be using simple volume lighting so always store the DMFP on the side to not override SSSPROFILE data.
-	ResultBSDF.TmpDMFP =				lerp(BBSDF.TmpDMFP * BLightingFactor,					ABSDF.TmpDMFP * ALightingFactor,			AMixFactor);
+	ResultBSDF.TmpDMFP =				lerp(BBSDF.TmpDMFP,								ABSDF.TmpDMFP,								AMixFactor);
 
 	// When it comes to adding the visual contribution, we chose to add the coverage without clamping as currently done in the StrataAdd function.
 	// Operations using Coverage are currently making sure they are in a safe range.
-	ResultBSDF.Coverage =	 BBSDF.Coverage + ABSDF.Coverage;
+	ResultBSDF.Coverage =				saturate(BBSDF.Coverage + ABSDF.Coverage);
 #endif
 
 #undef ResultBSDF
@@ -1333,8 +1396,9 @@ float VerticalLayeringParameterBlendingBSDFCoverageToNormalMix(FStrataData TopSt
 {
 	// This is a normalised mix so even though both BSDF have weights, the sum of weights used to combined them should always be 1.
 #if STRATA_INLINE_SHADING
-	const float SafeTopBSDFCoverage = saturate(TopStrata.Layers[0].BSDFs[0].Coverage);
+	const float SafeTopBSDFCoverage = saturate(TopStrata.InlinedBSDF.Coverage);
 	const float TopNormalContributionForWeightOfOne = 0.5f;
+	// STRATA_TODO take into account transmittance when computing the normal weight
 	return SafeTopBSDFCoverage * TopNormalContributionForWeightOfOne;
 #else
 	return 0.5f;
@@ -1345,18 +1409,12 @@ float VerticalLayeringParameterBlendingBSDFCoverageToNormalMix(FStrataData TopSt
 FStrataBSDF StrataVerticalLayeringParameterBlending(in FStrataBSDF TopBSDF, in FStrataBSDF BaseBSDF, uint NewNormalIndex, float NoV)
 {
 	FStrataBSDF ResultBSDF = (FStrataBSDF)0;
-
 #if STRATA_INLINE_SHADING
 	const float SafeTopCoverage = saturate(TopBSDF.Coverage);
 	const float SafeBaseCoverage = saturate(BaseBSDF.Coverage);
-	const float TopLightingFactor  = TopBSDF.bPremultipliedCoverage  ? 1.0f : SafeTopCoverage;
-	const float BaseLightingFactor = BaseBSDF.bPremultipliedCoverage ? 1.0f : SafeBaseCoverage;
-	ResultBSDF.bPremultipliedCoverage = true;
 #else
-	const float SafeTopCoverage = 1.0f;
-	const float SafeBaseCoverage = 1.0f;
-	const float TopLightingFactor = 1.0f;
-	const float BaseLightingFactor = 1.0f;
+	const float SafeTopCoverage = 0.5f;
+	const float SafeBaseCoverage = 0.5f;
 #endif
 
 	// We can only parameter blend the slab 
@@ -1368,7 +1426,7 @@ FStrataBSDF StrataVerticalLayeringParameterBlending(in FStrataBSDF TopBSDF, in F
 	BSDF_SETHASSSS(ResultBSDF,			BSDF_GETHASSSS(BaseBSDF));															// We only keep SSS if the bottom layer has it
 	BSDF_SETHASSSSPROFILE(ResultBSDF,	0);																					// Problem: cannot blend SSS profiles so we simply disable SSSProfile
 	BSDF_SETISTHIN(ResultBSDF,			BSDF_GETISTHIN(BaseBSDF));															// STRATA_TODO fix the MPF/DMFP situation for that to really work
-	BSDF_SETHASDMFP(ResultBSDF,			BSDF_GETHASDMFP(BaseBSDF));															// We only keep SSS if the bottom layer has it
+	BSDF_SETHASDMFP(ResultBSDF,			BSDF_GETHASDMFP(BaseBSDF)						| BSDF_GETHASDMFP(TopBSDF));		// We keep DMFP of both to make sure bottom layer on translucent can laos be effective
 	BSDF_SETHASHAZINESS(ResultBSDF,		BSDF_GETHASHAZINESS(TopBSDF));														// We only keep haziness if the top layer has it
 	BSDF_SETHASF90(ResultBSDF,			BSDF_GETHASF90(BaseBSDF)						| BSDF_GETHASF90(TopBSDF));			// F90: keep union of both even though it will be hard to get a match
 	BSDF_SETHASTHINFILM(ResultBSDF,		BSDF_GETHASTHINFILM(TopBSDF));														// We only keep thinfilme if the top layer has it
@@ -1377,29 +1435,35 @@ FStrataBSDF StrataVerticalLayeringParameterBlending(in FStrataBSDF TopBSDF, in F
 	// Compute the top slab transmittance that will be used to blend the top and base material parameters.
 	// We do not check BSDF_GETISSIMPLEVOLUME(BSDFContext.BSDF) because that is set when we know the BSDF is not at the bottom later. And here we already know it is not the case.
 	EnableSlabBSDFSimpleVolumetric(TopBSDF);
+	EnableSlabBSDFSimpleVolumetric(BaseBSDF);
 	const float3 UnusedDiffuseColor		= 0;	// Unused when only transmittance is needed
 	FParticipatingMedia PM				= StrataSlabCreateParticipatingMedia(UnusedDiffuseColor, SLAB_SSSDMFP(TopBSDF));
 	const float3 SlabTransmittance		= IsotropicMediumSlabTransmittance(PM, STRATA_SIMPLEVOLUME_THICKNESS_M, NoV);
-	// Compute the colored throughput, taking into account the weight.
-	const float3 ColorThroughput		= saturate((1.0f - SafeTopCoverage) + SafeTopCoverage * SlabTransmittance);
-	const float GreyThroughput			= dot(ColorThroughput, 1.0f / 3.0f);
+
+	// STRATA_TODO account for energy conservation here.
+
+	FVerticalLayeringInfo Info			= GetVerticalLayeringInfo(SafeTopCoverage, SafeBaseCoverage);
+	// Compute some ways to intepoldate and accumulate parameters, but always normalise because coverage is applied later in UpdateSingleBSDFOperatorCoverageTransmittance.
+	const float  TopToBottomLerp		= saturate((Info.TransmittanceTopAndBottom * dot(SlabTransmittance,0.33.xxx) + Info.SurfaceBottom) / max(1e-5, Info.Coverage));
+	const float3 BottomColorFactor		= (Info.TransmittanceTopAndBottom * SlabTransmittance + Info.SurfaceBottom) / max(1e-5, Info.Coverage);
+	const float  TopFactor				= Info.SurfaceTop / max(1e-5, Info.Coverage);
 	
 	// Add both layers thickness
-	BSDF_SETTHICKNESSCM(ResultBSDF,		BSDF_GETTHICKNESSCM(BaseBSDF)												+ TopLightingFactor * BSDF_GETTHICKNESSCM(TopBSDF));
+//	BSDF_SETTHICKNESSCM(ResultBSDF,		BSDF_GETTHICKNESSCM(BaseBSDF)												+ SafeTopCoverage * BSDF_GETTHICKNESSCM(TopBSDF));
 
 	// Now add bottom layer component weights by throughput to the top layer.
 	// Those parameters will contribute to added luminance, in this case we add them together and saturate to not go out of the safe range
 	// Some parameters are not contributing to added luminance: in this case we simply lerp them
-	BSDF_SETEMISSIVE(ResultBSDF,		ColorThroughput * BaseLightingFactor * BSDF_GETEMISSIVE(BaseBSDF)			+ TopLightingFactor * BSDF_GETEMISSIVE(TopBSDF));
+	BSDF_SETEMISSIVE(ResultBSDF,		BottomColorFactor * BSDF_GETEMISSIVE(BaseBSDF)								+ TopFactor * BSDF_GETEMISSIVE(TopBSDF));
 
-	SLAB_DIFFUSEALBEDO(ResultBSDF) =	saturate(ColorThroughput * BaseLightingFactor * SLAB_DIFFUSEALBEDO(BaseBSDF)+ TopLightingFactor * SLAB_DIFFUSEALBEDO(TopBSDF));
-	SLAB_F0(ResultBSDF) =				saturate(ColorThroughput * BaseLightingFactor * SLAB_F0(BaseBSDF)			+ TopLightingFactor * SLAB_F0(TopBSDF));
-	SLAB_F90(ResultBSDF) =				saturate(ColorThroughput * BaseLightingFactor * SLAB_F90(BaseBSDF)			+ TopLightingFactor * SLAB_F90(TopBSDF));
-	SLAB_FUZZ_COLOR(ResultBSDF) =		saturate(ColorThroughput * BaseLightingFactor * SLAB_FUZZ_COLOR(BaseBSDF)	+ TopLightingFactor * SLAB_FUZZ_COLOR(TopBSDF));
+	SLAB_DIFFUSEALBEDO(ResultBSDF) =	saturate(BottomColorFactor * SLAB_DIFFUSEALBEDO(BaseBSDF)					+ TopFactor * SLAB_DIFFUSEALBEDO(TopBSDF));
+	SLAB_F0(ResultBSDF) =				saturate(BottomColorFactor * SLAB_F0(BaseBSDF)								+ TopFactor * SLAB_F0(TopBSDF));
+	SLAB_F90(ResultBSDF) =				saturate(BottomColorFactor * SLAB_F90(BaseBSDF)								+ TopFactor * SLAB_F90(TopBSDF));
+	SLAB_FUZZ_COLOR(ResultBSDF) =		saturate(BottomColorFactor * SLAB_FUZZ_COLOR(BaseBSDF)						+ TopFactor * SLAB_FUZZ_COLOR(TopBSDF));
 		
-	SLAB_ROUGHNESS(ResultBSDF) =		lerp(SLAB_ROUGHNESS(TopBSDF),					SLAB_ROUGHNESS(BaseBSDF),					GreyThroughput);
-	SLAB_HAZINESS(ResultBSDF) =			lerp(SLAB_HAZINESS(TopBSDF),					SLAB_HAZINESS(BaseBSDF),					GreyThroughput);
-	SLAB_FUZZ_AMOUNT(ResultBSDF) =		lerp(SLAB_FUZZ_AMOUNT(TopBSDF),					SLAB_FUZZ_AMOUNT(BaseBSDF),					GreyThroughput);
+	SLAB_ROUGHNESS(ResultBSDF) =		lerp(SLAB_ROUGHNESS(TopBSDF),					SLAB_ROUGHNESS(BaseBSDF),	TopToBottomLerp);
+	SLAB_HAZINESS(ResultBSDF) =			lerp(SLAB_HAZINESS(TopBSDF),					SLAB_HAZINESS(BaseBSDF),	TopToBottomLerp);
+	SLAB_FUZZ_AMOUNT(ResultBSDF) =		lerp(SLAB_FUZZ_AMOUNT(TopBSDF),					SLAB_FUZZ_AMOUNT(BaseBSDF), TopToBottomLerp);
 
 	// Anisotropy is disabled with vertical blending. 
 	SLAB_ANISOTROPY(ResultBSDF) =		0.f;
@@ -1408,24 +1472,66 @@ FStrataBSDF StrataVerticalLayeringParameterBlending(in FStrataBSDF TopBSDF, in F
 	SLAB_THINFILM_THICKNESS(ResultBSDF)= (SafeTopCoverage >= 1.0f) ? SLAB_THINFILM_THICKNESS(TopBSDF) : 0.0f;
 
 	// Keep the DMFP form the bottom layer only (top layer simple volume is backed in parameters).
-	SLAB_SSSDMFP(ResultBSDF) =			SLAB_SSSDMFP(BaseBSDF);
+	SLAB_SSSDMFP(ResultBSDF) =			lerp(SLAB_SSSDMFP(TopBSDF),						SLAB_SSSDMFP(BaseBSDF),		TopToBottomLerp); // SLAB_SSSDMFP(BaseBSDF);
 
 #if STRATA_INLINE_SHADING
+#if STRATA_OPAQUE_MATERIAL
 	// Only the bottom layer SSS / Thin model is used so we only keep the bottom layer data there. Top layer is converted to throughput.
 	ResultBSDF.TmpDMFP =				BaseBSDF.TmpDMFP;
+	ResultBSDF.Coverage =				Info.Coverage;
+#else
+
+	// We compute the total transmittance of the medium over the max coverage
+	const float MaxCoverage = max(TopBSDF.Coverage, BaseBSDF.Coverage);
+	FVerticalLayeringInfo Info2 = GetVerticalLayeringInfo(TopBSDF.Coverage / MaxCoverage, BaseBSDF.Coverage / MaxCoverage);
+#if 0
+	// Using transmittance as a way to interpolate should be better but there are issues with large values
+	const float3 TopExtinction			= 1.0f / max(0.0, TopBSDF.TmpDMFP);
+	const float3 BaseExtinction			= 1.0f / max(0.0, BaseBSDF.TmpDMFP);
+	const float3 TopTransmitance		= exp(-TopExtinction);	// extinction = optical depth, we assume here d = 1 meter
+	const float3 BaseTransmitance		= exp(-BaseExtinction);
+	const float3 NewTransmittance		= Info2.TransmittanceOnlyBottom * BaseTransmitance + Info2.TransmittanceOnlyTop * TopTransmitance + Info2.TransmittanceTopAndBottom * TopTransmitance * BaseTransmitance;
+	const float3 NewExtinction			= -log(saturate(max(0.0, NewTransmittance)));
+	const float3 NewMFP					= 1.0f / NewExtinction;
+#else
+	// MFP looks more prceptually linear even though it is not perfect
+	const float3 NewMFP = Info2.TransmittanceOnlyBottom * BaseBSDF.TmpDMFP + Info2.TransmittanceOnlyTop * TopBSDF.TmpDMFP + Info2.TransmittanceTopAndBottom * min(TopBSDF.TmpDMFP, BaseBSDF.TmpDMFP);
+#endif
+
+	BSDF_SETTHICKNESSCM(ResultBSDF, CENTIMETER_TO_METER);
+
+	//ResultBSDF.TmpDMFP =				lerp(TopBSDF.TmpDMFP,							BaseBSDF.TmpDMFP,			TopToBottomLerp);
+	ResultBSDF.TmpDMFP					= NewMFP;
+	SLAB_SSSDMFP(ResultBSDF)			= NewMFP;
+
+	ResultBSDF.Coverage				= MaxCoverage;
+#endif
 
-	ResultBSDF.Coverage =				1.0f - (1.0f - SafeTopCoverage) * (1.0f - SafeBaseCoverage);
 #endif
 	return ResultBSDF;
 }
 
 FStrataData StrataVerticalLayeringParameterBlending(FStrataData Top, FStrataData Base, uint NewNormalIndex, float NoV)
 {
-	FStrataData Result = GetInitialisedStrataDataWithReservedFirstSlot();
-	Result.Layers[0].BSDFs[0] = StrataVerticalLayeringParameterBlending(Top.Layers[0].BSDFs[0], Base.Layers[0].BSDFs[0], NewNormalIndex, NoV);
+	FStrataData Result = GetInitialisedStrataData();
+	Result.InlinedBSDF = StrataVerticalLayeringParameterBlending(Top.InlinedBSDF, Base.InlinedBSDF, NewNormalIndex, NoV);
 	return Result;
 }
 
+
+FStrataData StrataWeightParameterBlending(FStrataData A, float Weight)
+{
+	const float SafeWeight = saturate(Weight);
+	FStrataData StrataData = A;
+
+	// Apply the weight to coverage
+#if STRATA_INLINE_SHADING
+	StrataData.InlinedBSDF.Coverage *= SafeWeight;
+#endif
+
+	return StrataData;
+}
+
 ///////////////////////////////////////////////////////////////////////////////
 // Functions used to sanitize BSDF before they are used for lighting (when forward) or stored to memory (base pass).
 // We do not normalize normals, this is done before the lighting step.
@@ -2096,18 +2202,18 @@ FStrataDBuffer StrataConvertToDBuffer(
 	in FStrataData InData)
 {
 	FStrataDBuffer Out = (FStrataDBuffer)0;
-	const float3 DiffuseColor = SLAB_DIFFUSEALBEDO(InData.Layers[0].BSDFs[0]);
-	const float3 F0 = SLAB_F0(InData.Layers[0].BSDFs[0]);
+	const float3 DiffuseColor = SLAB_DIFFUSEALBEDO(InData.InlinedBSDF);
+	const float3 F0 = SLAB_F0(InData.InlinedBSDF);
 
 	// STRATA_TODO : Can we avoid this by having a lerpable packing?
-	Out.Roughness = SLAB_ROUGHNESS(InData.Layers[0].BSDFs[0]);
+	Out.Roughness = SLAB_ROUGHNESS(InData.InlinedBSDF);
 	Out.Metallic = F0RGBToMetallic(F0);
 	Out.Specular = F0ToDielectricSpecular(F0.x);
 	Out.BaseColor = lerp(DiffuseColor, F0, Out.Metallic);
-	Out.Emissive = BSDF_GETEMISSIVE(InData.Layers[0].BSDFs[0]);
+	Out.Emissive = BSDF_GETEMISSIVE(InData.InlinedBSDF);
 #if STRATA_INLINE_SHADING
-	Out.WorldNormal = InHeader.SharedLocalBases.Normals[BSDF_GETSHAREDLOCALBASISID(InData.Layers[0].BSDFs[0])];
-	Out.Coverage = InData.Layers[0].BSDFs[0].Coverage;
+	Out.WorldNormal = InHeader.SharedLocalBases.Normals[BSDF_GETSHAREDLOCALBASISID(InData.InlinedBSDF)];
+	Out.Coverage = InData.InlinedBSDF.Coverage;
 #else
 	Out.WorldNormal = float3(0,0,1);
 	Out.Coverage = 1.f;
@@ -2159,11 +2265,8 @@ void StrataConvertFromDBuffer(
 		SharedLocalBasisIndex,
 		SharedLocalBasisTypes);
 #if STRATA_INLINE_SHADING
-	OutData.Layers[0].BSDFs[0].Emissive = In.Emissive;
-	OutData.Layers[0].BSDFs[0].Coverage = In.Coverage;
-
-	// DBuffer data are alreay pre-multiplied when read from the DBuffers
-	OutData.Layers[0].BSDFs[0].bPremultipliedCoverage = true;
+	OutData.InlinedBSDF.Emissive = In.Emissive;
+	OutData.InlinedBSDF.Coverage = In.Coverage;
 #endif
 }
 
@@ -2171,6 +2274,19 @@ void StrataConvertFromDBuffer(
 ///////////////////////////////////////////////////////////////////////////////
 // Strata deferred base-pass output
 
+#if STRATA_INLINE_SHADING && STRATA_CLAMPED_BSDF_COUNT > 0
+
+// Forward declaration
+#if STRATA_USES_CONVERSION_FROM_LEGACY==0
+void UpdateAllBSDFsOperatorCoverageTransmittance(inout FStrataPixelHeader StrataPixelHeader, inout FStrataTree StrataTree, bool bRoughDiffuseEnabled, FStrataAddressing NullStrataAddressing, float3 V, float3 L);
+void UpdateAllOperatorsCoverageTransmittance(inout FStrataTree StrataTree);
+void UpdateAllBSDFWeightAfterOperatorVisit(inout FStrataTree StrataTree);
+#else
+void LegacyUpdateBSDFsOperators(inout FStrataPixelHeader StrataPixelHeader, inout FStrataTree StrataTree, bool bRoughDiffuseEnabled, FStrataAddressing NullStrataAddressing, float3 V, float3 L);
+#endif
+
+#endif
+
 uint PackStrataHeader(uint BSDFCount, FStrataPixelHeader In);
 
 // If this is changed, please update the compiler side material size evaluation in StrataMaterial.cpp
@@ -2187,6 +2303,9 @@ void PackStrataOut(
 	inout FStrataSubsurfaceData SSSData,
 	inout FStrataTopLayerData TopLayerData)
 {
+	// This only exists with inline shading and when we are going to write out BSDFs (UpdateAllXXX functions needs to be defined)
+#if STRATA_INLINE_SHADING && STRATA_CLAMPED_BSDF_COUNT > 0
+
 	bStrataSubsurfaceEnable = false;
 	EmissiveColor = 0.0f;
 
@@ -2202,325 +2321,227 @@ void PackStrataOut(
 	SetIsSimpleMaterial(StrataPixelHeader, true); // non strata material or sky pixels are considered as simple (so that (strata simple pixels) + (non-strata pixels) can build a simple tile).
 	SetIsSingleMaterial(StrataPixelHeader, true); // non strata material or sky pixels are considered as simple (so that (strata simple pixels) + (non-strata pixels) can build a simple tile).
 
-	int SimpleMaterial_LayerIndex = 0;
-	int SimpleMaterial_BSDFIndex = 0;
-	uint BSDFCount = 0;
 	uint BSDFVisibleCount = 0;
-	float3 Throughput = 1.0f;
-
-	// We clamp the layer count to the maximum one according to our byte per pixel budget (removing the bottom layer first)
-	Strata.LayerCount = min(Strata.LayerCount, STRATA_CLAMPED_LAYER_COUNT);
+	
 
-	// Evaluate the bottom layer index programmatically according to coverage.
-	// This is needed for instance when doing legacy to strata material conversion: a static topology with 2 layers is used for the compiler to
-	// successfully unroll the loops below. Typically the BSDF is set as the top layer. 
-	// SSS would not work if not the bottom layer so that is why we recover the bottom layer index this way.
-	int BottomLayerIndex = Strata.LayerCount - 1;
-#if STRATA_INLINE_SHADING
+	const bool bRoughDiffuseEnabled = false;	// STARTA_TODO 
+	FStrataTree StrataTree = StrataPixelHeader.StrataTree;
+	if (StrataTree.BSDFCount > 0)
 	{
-		STRATA_UNROLL_N(STRATA_CLAMPED_LAYER_COUNT)
-		for (int l = BottomLayerIndex; l > 0; --l)	// Skipping layer 0 to process one less layer when there are two.
+		FStrataAddressing NullStrataAddressing = (FStrataAddressing)0;	// Fake unused in StrataCreateBSDFContext when using Forward inline shading
+		const float3 NullLightDir = float3(0.0, 0.0, 1.0);
+#if STRATA_USES_CONVERSION_FROM_LEGACY==0
+		//
+		// Update the coverage/transmittance of each leaves (==BSDFs) of the strata tree.
+		//
+		UpdateAllBSDFsOperatorCoverageTransmittance(StrataPixelHeader, StrataTree, bRoughDiffuseEnabled, NullStrataAddressing, V, NullLightDir);
+		
+		//
+		// Propagate up the coverage/transmittance of each node in the strata tree.
+		//
+		UpdateAllOperatorsCoverageTransmittance(StrataTree);
+		
+		//
+		// Update the luminance weight of each BSDF according to the operators it has to traverse up to the strata tree root node.
+		//
+		UpdateAllBSDFWeightAfterOperatorVisit(StrataTree);
+
+#else
+		LegacyUpdateBSDFsOperators(StrataPixelHeader, StrataTree, bRoughDiffuseEnabled, NullStrataAddressing, V, NullLightDir);
+#endif
+		
+		int SimpleMaterial_BSDFIndex = 0;
+		float TopLayerCoverage = 0.0f;
+		
 		{
-			float TotalWeight = 0.0f;
-			STRATA_UNROLL_N(STRATA_CLAMPED_MAX_BSDF_COUNT_PER_LAYER)
-			for (int i = 0; i < Strata.Layers[l].BSDFCount; ++i)
+			STRATA_UNROLL_N(STRATA_CLAMPED_BSDF_COUNT)
+			for (int BSDFIdx = 0; BSDFIdx < StrataTree.BSDFCount; ++BSDFIdx)
 			{
-				TotalWeight += /*Throughput * */(Strata.Layers[l].BSDFs[i].bPremultipliedCoverage ? 1.0 : Strata.Layers[l].BSDFs[i].Coverage);
-			}
-			if (TotalWeight == 0.0f)
-			{
-				BottomLayerIndex--;
-			}
-			else
-			{
-				break;
-			}
-		}
-	}
-#endif
+			#define BSDF StrataTree.BSDFs[BSDFIdx]
 
-	// For each of the BSDF, we evaluated their combined coverage/throughput according to the material operation performed (layering, mixing, etc.).
-	{
-		STRATA_UNROLL_N(STRATA_CLAMPED_LAYER_COUNT)
-		for (int l = 0; l < Strata.LayerCount; ++l)
-		{
-			float3 LayerThroughput = 0.0f;
-			float3 LayerCoverage = 0.0f;
+				const bool bIsVisible = StrataIsBSDFVisible(BSDF);
 
-			const bool bTopLayer = l == 0;
-			const bool bBottomLayer = l == BottomLayerIndex;
+				if (bIsVisible)
+				{
+					BSDFVisibleCount++;
 
-			STRATA_UNROLL_N(STRATA_CLAMPED_MAX_BSDF_COUNT_PER_LAYER)
-			for (int i = 0; i < Strata.Layers[l].BSDFCount; ++i)
-			{
-				// Dither BSDF data to hide more aggressive bit reduction.
-				// Sanitize happening after will clamp values that could have gone above and below valid range.
-				StrataDitherBSDF(Strata.Layers[l].BSDFs[i], Dither);
+					// Dither BSDF data to hide more aggressive bit reduction.
+					// Sanitize happening after will clamp values that could have gone above and below valid range.
+					StrataDitherBSDF(BSDF, Dither);
 
-				// Sanitize BSDF in place before it is used and packed out to memory
-				StrataSanitizeBSDF(Strata.Layers[l].BSDFs[i]);
+					// Sanitize BSDF in place before it is used and packed out to memory
+					StrataSanitizeBSDF(BSDF);
 
-				FStrataBSDF BSDF = Strata.Layers[l].BSDFs[i];
 
-			#if STRATA_INLINE_SHADING
-				const float BSDFCoverage = saturate(BSDF.Coverage);
-			#else
-				const float BSDFCoverage = 1.0f;
-			#endif
-				BSDFCount++;
+					bool bIsSimpleMaterial = true;
+					bool bIsSingleMaterial = true;
+					switch (BSDF_GETTYPE(BSDF))
+					{
+					case STRATA_BSDF_TYPE_SLAB:
+					{
+						EmissiveColor += BSDF_GETEMISSIVE(BSDF) * BSDF.LuminanceWeight;
 
-				// STRATA_TODO apply the same normal transform as for the material
-				float3x3 TangentBasis = StrataGetBSDFSharedBasis_InlineShading(StrataPixelHeader, BSDF, StrataAddressing);
-				float3 N = TangentBasis[2];
-				//N = normalize(N);
+						// Simple and single bsdf do not allow weights other than 1
+						const bool LuminanceWeightEqualOne = all(BSDF.LuminanceWeight == 1.0f);
 
-				BxDFContext ContextNoLight = (BxDFContext)0;
-				const float3 FakeL = float3(0, 0, 1);
-				Init(ContextNoLight, N, V, FakeL);
+						// Update simple material compatibility
+						bIsSimpleMaterial = IsStrataSlabFastPathCompatible(BSDF) && LuminanceWeightEqualOne;
 
-				bool bIsSimpleMaterial = true;
-				bool bIsSingleMaterial = true;
-				switch (BSDF_GETTYPE(BSDF))
-				{
-				case STRATA_BSDF_TYPE_SLAB:
-				{
-					// Update simple material compatibility
-					bIsSimpleMaterial = IsStrataSlabFastPathCompatible(BSDF);
+						// Update single material compatibility. For now, single materials don't support anisotropy, 
+						// they use the complex pass, as the toplayer data does not contain the full frame basis, only the top normal.
+						bIsSingleMaterial = BSDF_GETHASANISOTROPY(BSDF) == 0 && LuminanceWeightEqualOne;
 
-					// Update single material compatibility. For now, single materials don't support anisotropy, 
-					// they use the complex pass, as the toplayer data does not contain the full frame basis, only the top normal.
-					bIsSingleMaterial = BSDF_GETHASANISOTROPY(BSDF) == 0;
+						// Two sided lighting can only be evalauted on the bottom layer, together with option SSS. Otherwise, the simple volume will be used.
+						const bool bIsThin = BSDF.bIsBottom && BSDF_GETISTHIN(BSDF);
 
-					EmissiveColor += BSDF_GETEMISSIVE(BSDF) * Throughput;
+						// We want simple volumetric only if the DMFP input is pluged in, otherwise we avoid storing the DFMP
+						const bool bIsSimpleVolume = !BSDF.bIsBottom && BSDF_GETHASDMFP(BSDF);
 
-					// Two sided lighting can only be evalauted on the bottom layer, together with option SSS. Otherwise, the simple volume will be used.
-					const bool bIsThin = bBottomLayer && BSDF_GETISTHIN(Strata.Layers[l].BSDFs[i]);
+						bStrataSubsurfaceEnable = bStrataSubsurfaceEnable || (BSDF_GETHASSSS(BSDF) && BSDF.Coverage > 0.0f && !bIsSimpleVolume);
 
-					// We want simple volumetric only if the DMFP input is pluged in, otherwise we avoid storing the DFMP
-					const bool bIsSimpleVolume = !bBottomLayer && BSDF_GETHASDMFP(Strata.Layers[l].BSDFs[i]);
+						// By default, without coverage, the diffuse BSDF is completely opaque
+						float3 BSDFThroughput = OpaqueBSDFThroughput;
 
-					bStrataSubsurfaceEnable = bStrataSubsurfaceEnable || (BSDF_GETHASSSS(BSDF) && BSDFCoverage > 0 && !bIsSimpleVolume);
+						if (bIsSimpleVolume)
+						{
+							EnableSlabBSDFSimpleVolumetric(BSDF);
 
-					// By default, without coverage, the diffuse BSDF is completely opaque
-					float3 BSDFThroughput = OpaqueBSDFThroughput;
+							// Mark the BSDF has not having any SSS through post process.
+							bStrataSubsurfaceEnable = false;
+						}
 
-					if (bIsSimpleVolume)
-					{
-						float3 ApproxFresnelForTransmittance = 0;
-						// STRATA_TODO Fresnel throughput, Use a LUT, this is currently in progress and will be coming soon.
-						// In the meantime, we use a simple approach to feed Fresnel function and compute an approximation of the general transmition/througput.
-						const float SatNoV = saturate(ContextNoLight.NoV);
-						const float VoH = SatNoV;
-						const float NoL = SatNoV;
-						const float3 F0 = SLAB_F0(BSDF);
-						const float3 F90 = SLAB_F90(BSDF);
-						const bool bHasThinFilm = BSDF_GETHASTHINFILM(BSDF);
-						BRANCH
-						if (bHasThinFilm)
+						if (bIsThin)
 						{
-							float ThinFilmIOR = SLAB_THINFILM_IOR(BSDF);
-							float ThinFilmThickness = SLAB_THINFILM_THICKNESS(BSDF);
-							ApproxFresnelForTransmittance = F_ThinFilm(SatNoV, NoL, VoH, F0, F90, ThinFilmIOR, ThinFilmThickness);
+							EnableSlabBSDFIsThin(BSDF);
 						}
 						else
 						{
-							ApproxFresnelForTransmittance = F_Schlick(F0, F90, VoH);
+							// If not selected or not possible to have two sided lighting enabled (not bottom layer or SSSProfile is enabled), let's make sure we disable the feature.
+							BSDF_SETISTHIN(BSDF, 0);
 						}
 
-						EnableSlabBSDFSimpleVolumetric(Strata.Layers[l].BSDFs[i]);
-
-						const float3 UnusedDiffuseColor = 0; // Unused when only trasmittance is needed
-						FParticipatingMedia PM = StrataSlabCreateParticipatingMedia(UnusedDiffuseColor, SLAB_SSSDMFP(Strata.Layers[l].BSDFs[i]));
-						const float DiffuseToVolumeBlend = StrataSlabDiffuseToVolumeBlend(PM);
-						const float3 SlabTransmittance = IsotropicMediumSlabTransmittance(PM, STRATA_SIMPLEVOLUME_THICKNESS_M, ContextNoLight.NoV);
+						if (BSDF.bIsTop)
+						{
+							TopLayerCoverage += BSDF.bIsTop ? BSDF.Coverage : 0.0f;
+							TopLayerData.Roughness += BSDF.Coverage * SLAB_ROUGHNESS(BSDF);
+						}
 
-						// Update the see-through throughput accoding to the medium transmittance
-						BSDFThroughput = lerp(OpaqueBSDFThroughput, (1- ApproxFresnelForTransmittance) * SlabTransmittance, DiffuseToVolumeBlend);
+						const bool bSSSMask = !bIsThin && BSDF.bIsBottom;	// SSS and Thin can only be on the bottom layer, so ignoring bIsSimpleVolume here.
+						SetHasSubsurface(StrataPixelHeader, (bSSSMask && BSDF_GETHASSSS(BSDF)));
 
-						// Mark the BSDF has not having any SSS through post process.
-						bStrataSubsurfaceEnable = false;
+						break;
 					}
 
-					if (bIsThin)
-					{
-						EnableSlabBSDFIsThin(Strata.Layers[l].BSDFs[i]);
-					}
-					else
-					{
-						// If not selected or not possible to have two sided lighting enabled (not bottom layer or SSSProfile is enabled), let's make sure we disable the feature.
-						BSDF_SETISTHIN(Strata.Layers[l].BSDFs[i], 0);
-					}
-
-					if (bTopLayer)
+					case STRATA_BSDF_TYPE_HAIR:
 					{
-						TopLayerData.Roughness += BSDFCoverage * SLAB_ROUGHNESS(BSDF);
-					}
+						bIsSimpleMaterial = false;
 
-					const bool bSSSMask = !bIsThin && bBottomLayer;	// SSS and Thin can only be on the bottom layer, so ignoring bIsSimpleVolume here.
-					SetHasSubsurface(StrataPixelHeader,	(bSSSMask && BSDF_GETHASSSS(BSDF)));
-
-					// STRATA_TODO	Should we have the Fuzz influence the next layter visibility accoding to fuzz amount? Should it also occlude diffuse?
-					//				If this is done, we should also update ComputePostProcessBSDFData.
-
-					LayerCoverage	+= BSDFCoverage;
-					LayerThroughput += BSDFCoverage * BSDFThroughput;
-					break;
-				}
-
-				case STRATA_BSDF_TYPE_HAIR:
-				{
-					bIsSimpleMaterial = false;
+						TopLayerCoverage += BSDF.bIsTop ? BSDF.Coverage : 0.0f;
+						TopLayerData.Roughness += BSDF.bIsTop ? BSDF.Coverage * FullyRough : 0.0f;
 
-					if (bTopLayer)
-					{
-						TopLayerData.Roughness += BSDFCoverage * FullyRough;
+						break;
 					}
 
-					// Hair are considered opaque
-					LayerCoverage	+= BSDFCoverage;
-					LayerThroughput += BSDFCoverage * OpaqueBSDFThroughput;
-					break;
-				}
-
-				case STRATA_BSDF_TYPE_SINGLELAYERWATER:
-				{
-					bIsSimpleMaterial = false;
+					case STRATA_BSDF_TYPE_SINGLELAYERWATER:
+					{
+						bIsSimpleMaterial = false;
+						EmissiveColor += BSDF_GETEMISSIVE(BSDF) * BSDF.LuminanceWeight;
 
-					EmissiveColor += BSDF_GETEMISSIVE(BSDF) * Throughput;
+						TopLayerCoverage += BSDF.bIsTop ? BSDF.Coverage : 0.0f;
+						TopLayerData.Roughness += BSDF.bIsTop ? BSDF.Coverage * SLW_ROUGHNESS(BSDF) : 0.0f;
 
-					if (bTopLayer)
-					{
-						TopLayerData.Roughness += BSDFCoverage * SLW_ROUGHNESS(BSDF);
+						break;
+					}
 					}
 
-					// Water is considered opaque, it has a custom process to simulate translucency
-					LayerCoverage += BSDFCoverage;
-					LayerThroughput += BSDFCoverage * OpaqueBSDFThroughput;
-					break;
-				}
-
-				} // switch
-
-			#if STRATA_INLINE_SHADING
-				// Update the BSDF final weight with throughput towards the view from layers above and coverage if it has not been baked into the parameters.
-				Strata.Layers[l].BSDFs[i].FinalWeight = Throughput * (Strata.Layers[l].BSDFs[i].bPremultipliedCoverage ? 1.0 : Strata.Layers[l].BSDFs[i].Coverage);
-			#else
-				Strata.Layers[l].BSDFs[i].FinalWeight = Throughput;
-			#endif
-			
-				// Cull non visible BSDF
-				const bool bIsVisible = any(Strata.Layers[l].BSDFs[i].FinalWeight > 0);
-				if (bIsVisible)
-				{
-					++BSDFVisibleCount;
+					float3x3 TangentBasis = StrataGetBSDFSharedBasis_InlineShading(StrataPixelHeader, BSDF, StrataAddressing);
+					float3 N = TangentBasis[2];
+					TopLayerData.WorldNormal += BSDF.bIsTop ? N * max(1e-5, BSDF.Coverage) : 0.0;
 
 					AddShadingModels(StrataPixelHeader, StrataShadingModelBit(BSDF_GETTYPE(BSDF)));
 					AndIsSimpleMaterial(StrataPixelHeader, bIsSimpleMaterial);
 					AndIsSingleMaterial(StrataPixelHeader, bIsSingleMaterial);
 
-					if (bTopLayer)
-					{
-						// together with a threshold to avoid 0-length normal after normalization in case influence is fadeout using the Weight node.
-						TopLayerData.WorldNormal += N * max(1e-5, BSDFCoverage);
-					}
+					SimpleMaterial_BSDFIndex = BSDFIdx;
 
-					// if a material is simple/single, it is necessary the only one
-					{
-						SimpleMaterial_LayerIndex = l;
-						SimpleMaterial_BSDFIndex = i;
-					}
+					// Notify that the BSDF is at the top for SSR to only affect reflection there and not on the lower layers
+					BSDF_SETISTOPLAYER(BSDF, BSDF.bIsTop ? 1 : 0);
 				}
-
-				// Notify that the BSDF is at the top for SSR to only affect reflection there and not on the lower layers
-				BSDF_SETISTOPLAYER(Strata.Layers[l].BSDFs[i], bTopLayer ? 1 : 0);
-			}
-
-			// Average the TopLayer data once we have read all the data
-			if (bTopLayer)
-			{
-				// Top layer roughness & normal have been weighted with BSDF Coverage only (since the accumulated throughput is 1.f for the top-layer)
-				const float BSDFWeightAvg = dot(LayerCoverage, 1.0f / 3.0f);
-
-				TopLayerData.WorldNormal = normalize(TopLayerData.WorldNormal);
-				TopLayerData.Roughness = BSDFWeightAvg > 0.0f ? TopLayerData.Roughness / BSDFWeightAvg : 0.0f;
+		
+		
+			#undef BSDF
 			}
-
-			// Take into account the non covered part for this layer and update throughput accordingly
-			LayerThroughput += (1.0 - saturate(LayerCoverage)) * FullThroughput;
-
-			// And combine this layer throughtput with he material throughput
-			Throughput *= LayerThroughput;
 		}
-	}
 
-	// Set storage layout as either: fast(0), single(1), or complex(2)
-	AndIsSimpleMaterial(StrataPixelHeader, BSDFVisibleCount == 1);
-	AndIsSingleMaterial(StrataPixelHeader, BSDFVisibleCount == 1);
+		// Finalize top layer data
+		TopLayerData.WorldNormal = normalize(TopLayerData.WorldNormal);
+		TopLayerData.Roughness = TopLayerCoverage > 0.0f ? TopLayerData.Roughness / TopLayerCoverage : 0.0f;
 
-	// Now write out Strata data
+		// Set storage layout as either: fast(0), single(1), or complex(2)
+		AndIsSimpleMaterial(StrataPixelHeader, BSDFVisibleCount == 1);
+		AndIsSingleMaterial(StrataPixelHeader, BSDFVisibleCount == 1);
 
-	///////////////////////////////////////////////////////////////////////////
-	// 3 types of encodings
-	// * A. Layout0: Simple encoding	(use top layer normal)
-	// * B. Layout1: Single encoding	(use top layer normal)
-	// * C. Layout2: Complex encoding	(use basis)
-	const bool bHasFastEncoding   = IsSimpleMaterial(StrataPixelHeader);
-	const bool bHasSingleEncoding = IsSingleMaterial(StrataPixelHeader);
+		// Now write out Strata data
 
-	// (Layout2)
-	if (!bHasFastEncoding && !bHasSingleEncoding)
-	{
-		// 1. the header (Regular/Complex encoding)
-		const uint PackedHeader = PackStrataHeader(BSDFVisibleCount, StrataPixelHeader);
-		STRATA_STORE_UINT1(PackedHeader);
-
-		// 1.1 tangent basis (Regular/Complex encoding)
-	#if STRATA_INLINE_SHADING
-		UNROLL
-		for (uint i = 0; i < StrataPixelHeader.SharedLocalBases.Count; ++i)
+		///////////////////////////////////////////////////////////////////////////
+		// 3 types of encodings
+		// * A. Layout0: Simple encoding	(use top layer normal)
+		// * B. Layout1: Single encoding	(use top layer normal)
+		// * C. Layout2: Complex encoding	(use basis)
+		const bool bHasFastEncoding = IsSimpleMaterial(StrataPixelHeader);
+		const bool bHasSingleEncoding = IsSingleMaterial(StrataPixelHeader);
+
+		// (Layout2)
+		if (!bHasFastEncoding && !bHasSingleEncoding)
 		{
-			const uint BasisType = StrataGetSharedLocalBasisType(StrataPixelHeader.SharedLocalBases.Types, i);
-			if (BasisType == STRATA_BASIS_TYPE_NORMAL)
-			{
-				STRATA_STORE_UINT1(StrataPackNormal(StrataPixelHeader.SharedLocalBases.Normals[i]));
-			}
-			else // if (BasisType == STRATA_BASIS_TYPE_TANGENT)
+			// 1. the header (Regular/Complex encoding)
+			const uint PackedHeader = PackStrataHeader(BSDFVisibleCount, StrataPixelHeader);
+			STRATA_STORE_UINT1(PackedHeader);
+
+			// 1.1 tangent basis (Regular/Complex encoding)
+#if STRATA_INLINE_SHADING
+			UNROLL
+			for (uint i = 0; i < StrataPixelHeader.SharedLocalBases.Count; ++i)
 			{
-				STRATA_STORE_UINT1(StrataPackNormalAndTangent(StrataPixelHeader.SharedLocalBases.Normals[i], StrataPixelHeader.SharedLocalBases.Tangents[i]));
+				const uint BasisType = StrataGetSharedLocalBasisType(StrataPixelHeader.SharedLocalBases.Types, i);
+				if (BasisType == STRATA_BASIS_TYPE_NORMAL)
+				{
+					STRATA_STORE_UINT1(StrataPackNormal(StrataPixelHeader.SharedLocalBases.Normals[i]));
+				}
+				else // if (BasisType == STRATA_BASIS_TYPE_TANGENT)
+				{
+					STRATA_STORE_UINT1(StrataPackNormalAndTangent(StrataPixelHeader.SharedLocalBases.Normals[i], StrataPixelHeader.SharedLocalBases.Tangents[i]));
+				}
 			}
+#endif
 		}
-	#endif
-	}
 
-	// 2. the list of BSDFs
-	BSDFCount = 0;
-	{
-		STRATA_UNROLL_N(STRATA_CLAMPED_LAYER_COUNT)
-		for (int l = 0; l < Strata.LayerCount; ++l)
-		{
-			FStrataLayer Layer = Strata.Layers[l];
 
-			STRATA_UNROLL_N(STRATA_CLAMPED_MAX_BSDF_COUNT_PER_LAYER)
-			for (int i = 0; i < Layer.BSDFCount; ++i)
+		{
+			int BSDFCount = 0;
+			STRATA_UNROLL_N(STRATA_CLAMPED_BSDF_COUNT)
+			for (int BSDFIdx = 0; BSDFIdx < StrataTree.BSDFCount; ++BSDFIdx)
 			{
-				FStrataBSDF BSDF = Layer.BSDFs[i];
+			#define BSDF StrataTree.BSDFs[BSDFIdx]
+
+				const bool bIsVisible = StrataIsBSDFVisible(BSDF);
 
 				// Only write visible BSDF
 				BRANCH
-				if (any(BSDF.FinalWeight > 0))
+				if (bIsVisible)
 				{
-					const uint GreyScaleThroughput = StrataHasGreyScaleWeight(BSDF.FinalWeight) ? 1 : 0;
+					const uint GreyScaleThroughput = StrataHasGreyScaleWeight(BSDF.LuminanceWeight) ? 1 : 0;
 					BSDF_SETHASGREYWEIGHT(BSDF, GreyScaleThroughput);
 
 					// A. Layout0 - Simple encoding (aka. fast-path): store header & data
-					if (bHasFastEncoding && l == SimpleMaterial_LayerIndex && i == SimpleMaterial_BSDFIndex)
-					{		
+					if (bHasFastEncoding && BSDFIdx == SimpleMaterial_BSDFIndex)
+					{
 						// Now we pack diffuse and F0 in a special way: both encoded as R7G7B6 with gamma 2.0. The low bit cound will be hidden by dithering and TAA.
 						// We pack 32 bits in the first uint and the remaining bits are pack is the lowest significant bit of the second uint.
-						const uint PackedDiffuse20Bits	= PackR7G7B6Gamma2(SLAB_DIFFUSEALBEDO(BSDF));
-						const uint PackedDiffuse12Bits	= PackedDiffuse20Bits & 0xFFF;
-						const uint PackedDiffuse8Bits	= (PackedDiffuse20Bits>>12) & 0xFF;
-						const uint PackedF020Bits		= PackR7G7B6Gamma2(SLAB_F0(BSDF));
+						const uint PackedDiffuse20Bits = PackR7G7B6Gamma2(SLAB_DIFFUSEALBEDO(BSDF));
+						const uint PackedDiffuse12Bits = PackedDiffuse20Bits & 0xFFF;
+						const uint PackedDiffuse8Bits = (PackedDiffuse20Bits >> 12) & 0xFF;
+						const uint PackedF020Bits = PackR7G7B6Gamma2(SLAB_F0(BSDF));
 						const uint PackedRoughness8bits = PackR8(SLAB_ROUGHNESS(BSDF));
 
 						// Data0 (Header_State|Header_AO|Roughness|Diffuse8bits)
@@ -2533,9 +2554,9 @@ void PackStrataOut(
 							Out = (Out & HEADER_FASTENCODING_MASK) | (Packed16 << HEADER_FASTENCODING_BIT_COUNT);
 							STRATA_STORE_UINT1(Out);
 
-							#if (HEADER_FASTENCODING_BIT_COUNT + 16) > 32
-							#error Strata fast path header is > 32bits
-							#endif
+						#if (HEADER_FASTENCODING_BIT_COUNT + 16) > 32
+						#error Strata fast path header is > 32bits
+						#endif
 						}
 						// Data1 (F0|Diffuse12bits)
 						{
@@ -2546,7 +2567,7 @@ void PackStrataOut(
 						BSDF.State = 0;
 					}
 					// B. Layout1 - Simple encoding: single BSDF, whose header & BSDF state are merged
-					else if (bHasSingleEncoding && l == SimpleMaterial_LayerIndex && i == SimpleMaterial_BSDFIndex)
+					else if (bHasSingleEncoding && BSDFIdx == SimpleMaterial_BSDFIndex)
 					{
 						// Data0 (Header_State|Header_AO|Header_BSDFTypes|BSDF_State)
 						{
@@ -2554,25 +2575,25 @@ void PackStrataOut(
 							HEADER_SETSTATE(Out, StrataPixelHeader.State);
 							HEADER_SETAO(Out, PackR8(StrataPixelHeader.MaterialAO));
 							HEADER_SETSHADINGMODELS(Out, GetShadingModels(StrataPixelHeader));
-							Out = Out & HEADER_SINGLEENCODING_MASK;							
+							Out = Out & HEADER_SINGLEENCODING_MASK;
 							Out = Out | ((BSDF.State & STATE_SINGLEENCODING_MASK) << HEADER_SINGLEENCODING_BIT_COUNT);
 							STRATA_STORE_UINT1(Out);
 
-							#if (HEADER_SIMPLEENCODING_BIT_COUNT) > 32
-							#error Strata fast path header is > 32bits
-							#endif
+						#if (HEADER_SIMPLEENCODING_BIT_COUNT) > 32
+						#error Strata fast path header is > 32bits
+						#endif
 						}
 					}
 					// C. Layout2 - Weight for Regular/Complex path
 					else if (GreyScaleThroughput > 0)
 					{
-						BSDF_SETWEIGHT10F(BSDF, Pack10F(BSDF.FinalWeight.x));
+						BSDF_SETWEIGHT10F(BSDF, Pack10F(BSDF.LuminanceWeight.x));
 						STRATA_STORE_UINT1(BSDF.State);
 					}
 					else
 					{
 						STRATA_STORE_UINT1(BSDF.State);
-						STRATA_STORE_UINT1(PackR11G11B10F(BSDF.FinalWeight));
+						STRATA_STORE_UINT1(PackR11G11B10F(BSDF.LuminanceWeight));
 					}
 
 					// Layout1 & Layout2
@@ -2583,13 +2604,13 @@ void PackStrataOut(
 					{
 						// Now we pack diffuse and F0 in a special way: both encoded as R7G7B6 with gamma 2.0. The low bit cound will be hidden by dithering and TAA.
 						// We pack 32 bits in the first uint and the remaining bits are pack is the lowest significant bit of the second uint.
-						const uint PackedDiffuse20Bits	= PackR7G7B6Gamma2(SLAB_DIFFUSEALBEDO(BSDF));
-						const uint PackedF020Bits		= PackR7G7B6Gamma2(SLAB_F0(BSDF));
-						const uint PackedData32Bits		= ((PackedDiffuse20Bits << 12) & 0xFFFFF000) | (PackedF020Bits & 0xFFF);
-						const uint PackedData8Bits		= (PackedF020Bits >> 12) & 0xFF; 
+						const uint PackedDiffuse20Bits = PackR7G7B6Gamma2(SLAB_DIFFUSEALBEDO(BSDF));
+						const uint PackedF020Bits = PackR7G7B6Gamma2(SLAB_F0(BSDF));
+						const uint PackedData32Bits = ((PackedDiffuse20Bits << 12) & 0xFFFFF000) | (PackedF020Bits & 0xFFF);
+						const uint PackedData8Bits = (PackedF020Bits >> 12) & 0xFF;
 
 						STRATA_STORE_UINT1(PackedData32Bits);
-						STRATA_STORE_UINT1(PackedData8Bits | PackRGBA8(float4(0.0f , SLAB_ROUGHNESS(BSDF), (SLAB_ANISOTROPY(BSDF) + 1.f) * 0.5f, 0.f)));
+						STRATA_STORE_UINT1(PackedData8Bits | PackRGBA8(float4(0.0f, SLAB_ROUGHNESS(BSDF), (SLAB_ANISOTROPY(BSDF) + 1.f) * 0.5f, 0.f)));
 
 						if (BSDF_GETHASF90(BSDF) || BSDF_GETHASTHINFILM(BSDF) || BSDF_GETHASHAZINESS(BSDF))
 						{
@@ -2610,13 +2631,13 @@ void PackStrataOut(
 							if (!bIsThin && BSDF_GETHASSSS(BSDF))
 							{
 								SSSData.BaseColor = SLAB_DIFFUSEALBEDO(BSDF);
-								SSSData.Specular  = F0ToDielectricSpecular(SLAB_F0(BSDF).x);
+								SSSData.Specular = F0ToDielectricSpecular(SLAB_F0(BSDF).x);
 
 								// Profile ID: we only keep the latest ID, and profile will overtake any non-profile based SSS
-								SSSData.bIsValid 			= true;
-								SSSData.bIsProfile 			= SSSData.bIsProfile || BSDF_GETHASSSSPROFILE(BSDF);
-								SSSData.ProfileId 			= StrataSubsurfaceProfileIdTo8bits(SLAB_SSSPROFILEID(BSDF));
-								SSSData.ProfileRadiusScale 	= SLAB_SSSPROFILERADIUSSCALE(BSDF);
+								SSSData.bIsValid = true;
+								SSSData.bIsProfile = SSSData.bIsProfile || BSDF_GETHASSSSPROFILE(BSDF);
+								SSSData.ProfileId = StrataSubsurfaceProfileIdTo8bits(SLAB_SSSPROFILEID(BSDF));
+								SSSData.ProfileRadiusScale = SLAB_SSSPROFILERADIUSSCALE(BSDF);
 
 								// STRATA_TODO This should be accumulated according to the BSDF weight when we run this in the base pass.
 								// We cannot do that now because BSDFWeight contains weight*coverage*trhougput all mixed together.
@@ -2661,13 +2682,18 @@ void PackStrataOut(
 					// Stop writing if above budget.
 					BSDFCount++;
 				}
+
+			#undef BSDF
 			}
+
+			checkSlow(BSDFCount == BSDFVisibleCount);
 		}
-	}
 
-	FinalizeWrites(StrataBuffer, StrataAddressing);
+		FinalizeWrites(StrataBuffer, StrataAddressing);
+	}
+	
+#endif // STRATA_INLINE_SHADING && STRATA_CLAMPED_BSDF_COUNT > 0
 
-	checkSlow(BSDFCount == BSDFVisibleCount);
 }
 
 // Pack only strata header
@@ -2778,7 +2804,7 @@ FStrataBSDF UnpackFastPathStrataBSDFIn(FStrataMaterialContainer StrataBuffer, in
 	uint Data0 = StrataPixelHeader.SharedLocalBasesTypes_PackedHeader;
 	uint Data1 = StrataLoadUint1(StrataBuffer, StrataAddressing);
 	OutBSDF.State = 0;
-	OutBSDF.FinalWeight = 1.0f;
+	OutBSDF.LuminanceWeight = 1.0f;
 
 	const uint PackedRoughness8bits = 0xFF    & (Data0 >>  HEADER_FASTENCODING_BIT_COUNT);
 	const uint PackedData8Bits		= 0xFF    & (Data0 >> (HEADER_FASTENCODING_BIT_COUNT + 8));
@@ -2835,7 +2861,7 @@ FStrataBSDF UnpackStrataBSDFIn(FStrataMaterialContainer StrataBuffer, inout FStr
 		Data = StrataHeader.SharedLocalBasesTypes_PackedHeader;
 	#endif
 		OutBSDF.State = (Data>>HEADER_SINGLEENCODING_BIT_COUNT) & STATE_SINGLEENCODING_MASK;
-		OutBSDF.FinalWeight = 1.0f;
+		OutBSDF.LuminanceWeight = 1.0f;
 		const uint BSDFType = firstbitlow(GetShadingModels(StrataHeader));
 		BSDF_SETTYPE(OutBSDF, BSDFType);
 
@@ -2847,11 +2873,11 @@ FStrataBSDF UnpackStrataBSDFIn(FStrataMaterialContainer StrataBuffer, inout FStr
 		OutBSDF.State = StrataLoadUint1(StrataBuffer, StrataAddressing);
 		if (BSDF_GETHASGREYWEIGHT(OutBSDF))
 		{
-			OutBSDF.FinalWeight = Unpack10F(BSDF_GETWEIGHT10F(OutBSDF));
+			OutBSDF.LuminanceWeight = Unpack10F(BSDF_GETWEIGHT10F(OutBSDF));
 		}
 		else
 		{
-			StrataLoad_R11G11B10F(StrataBuffer, StrataAddressing, OutBSDF.FinalWeight);
+			StrataLoad_R11G11B10F(StrataBuffer, StrataAddressing, OutBSDF.LuminanceWeight);
 		}
 	}
 
diff --git a/Engine/Shaders/Private/Strata/StrataDeferredLighting.ush b/Engine/Shaders/Private/Strata/StrataDeferredLighting.ush
--- a/Engine/Shaders/Private/Strata/StrataDeferredLighting.ush
+++ b/Engine/Shaders/Private/Strata/StrataDeferredLighting.ush
@@ -86,7 +86,7 @@ float4 StrataDeferredLighting(
 		float3 SpecularLuminance = BSDFEvaluate.IntegratedSpecularValue;
 
 		const bool bNeedsSeparateSubsurfaceLightAccumulation = BSDFEvaluate.bSubsurface;
-		const float3 CommonMultiplier = LightData.Color * LightMask * BSDF.FinalWeight;
+		const float3 CommonMultiplier = LightData.Color * LightMask * BSDF.LuminanceWeight;
 		LightAccumulator_AddSplit(Out, DiffuseLuminance, SpecularLuminance, DiffuseLuminance, CommonMultiplier, bNeedsSeparateSubsurfaceLightAccumulation);
 	}
 
diff --git a/Engine/Shaders/Private/Strata/StrataEnvironmentLighting.ush b/Engine/Shaders/Private/Strata/StrataEnvironmentLighting.ush
--- a/Engine/Shaders/Private/Strata/StrataEnvironmentLighting.ush
+++ b/Engine/Shaders/Private/Strata/StrataEnvironmentLighting.ush
@@ -61,7 +61,7 @@ float4 StrataEnvironmentLighting(
 	#else
 		FStrataBSDF BSDF = UnpackStrataBSDFIn(MaterialBuffer, StrataAddressing, StrataPixelHeader);
 	#endif
-		const float3 BSDFThroughput = BSDF.FinalWeight;
+		const float3 BSDFThroughput = BSDF.LuminanceWeight;
 
 		// Create the BSDF context
 		FStrataBSDFContext StrataBSDFContext = StrataCreateBSDFContext(StrataPixelHeader, BSDF, StrataAddressing, V);
diff --git a/Engine/Shaders/Private/Strata/StrataForwardLighting.ush b/Engine/Shaders/Private/Strata/StrataForwardLighting.ush
--- a/Engine/Shaders/Private/Strata/StrataForwardLighting.ush
+++ b/Engine/Shaders/Private/Strata/StrataForwardLighting.ush
@@ -5,7 +5,6 @@
 
 #define FORWARD_PER_PIXEL_SHADING (FORWARD_SHADING || TRANSLUCENCY_LIGHTING_SURFACE_FORWARDSHADING || MATERIAL_SHADINGMODEL_SINGLELAYERWATER)
 
-
 // Forward declarations
 void GetVolumeLightingNonDirectional(float4 AmbientLightingVector, float3 DiffuseColor, inout float3 InterpolatedLighting, out float4 VolumeLighting);
 void GetVolumeLightingDirectional(float4 AmbientLightingVector, float3 DirectionalLightingVector, float3 WorldNormal, float3 DiffuseColor, inout float3 InterpolatedLighting, out float4 VolumeLighting);
@@ -176,12 +175,8 @@ float3 StrataForwardLighting(
 
 	FStrataAddressing NullStrataAddressing = (FStrataAddressing)0;	// Fake unused in StrataCreateBSDFContext when using Forward inline shading
 
-	const float OpaqueBSDFThroughput = 0.0f;
-	const float FullThroughput = 1.0f;
-
 	float3 Color = 0;
 	OutCoverage = 0.0f;
-	float3 TransmittancePostCoverage = 1.0f;
 	OutTransmittancePreCoverage = 1.0f;
 
 	float4 DynamicShadowFactors = 1;
@@ -253,48 +248,59 @@ float3 StrataForwardLighting(
 	GridIndex = ComputeLightGridCellIndex((uint2)((MaterialParameters.SvPosition.xy - ResolvedView.ViewRectMin.xy) * View.LightProbeSizeRatioAndInvSizeRatio.zw), MaterialParameters.SvPosition.w, EyeIndex);
 #endif
 
-	STRATA_UNROLL_N(STRATA_CLAMPED_LAYER_COUNT)
-	for (int l = 0; l < StrataData.LayerCount; ++l)
-	{
-		float3 LayerThroughput = 0.0f;
-		float3 LayerCoverage = 0.0f;
-
-		const bool bTopLayer = l == 0;
-		const bool bBottomLayer = l == (StrataData.LayerCount - 1);
 
-		STRATA_UNROLL_N(STRATA_CLAMPED_MAX_BSDF_COUNT_PER_LAYER)
-		for (int i = 0; i < StrataData.Layers[l].BSDFCount; ++i)
-		{
-			FStrataBSDF BSDF = StrataData.Layers[l].BSDFs[i];
 
-			// Sanitize BSDF before it is used for forward shading
-			StrataSanitizeBSDF(BSDF);
+	FStrataTree StrataTree = StrataPixelHeader.StrataTree;
+	if (StrataTree.BSDFCount > 0)
+	{
+#if STRATA_USES_CONVERSION_FROM_LEGACY==0
+		//
+		// Update the coverage/transmittance of each leaves (==BSDFs) of the strata tree.
+		//
+		UpdateAllBSDFsOperatorCoverageTransmittance(StrataPixelHeader, StrataTree, bRoughDiffuseEnabled, NullStrataAddressing, V, L);
+
+		//
+		// Propagate up the coverage/transmittance of each node in the strata tree.
+		//
+		UpdateAllOperatorsCoverageTransmittance(StrataTree);
+
+		//
+		// Update the luminance weight of each BSDF according to the operators it has to traverse up to the strata tree root node.
+		//
+		UpdateAllBSDFWeightAfterOperatorVisit(StrataTree);
+#else
+		LegacyUpdateBSDFsOperators(StrataPixelHeader, StrataTree, bRoughDiffuseEnabled, NullStrataAddressing, V, L);
+#endif
 
-			BRANCH
-			if(BSDF_GETTYPE(BSDF) == STRATA_BSDF_TYPE_SLAB)
-			{
-				// We want simple volumetric only if the DMFP input is pluged in, otherwise we avoid dealing with simple volumetric.
-				// We also do not want to use the simple volumetric path if we are dealing with an opaque material and this slab is not at the bottom.
-				// And when opaque, the bottom layer must be opaque and thus will ahve a solid opaque diffuse color.
-				const bool bIsSimpleVolume = BSDF_GETHASDMFP(BSDF) && (!STRATA_OPAQUE_MATERIAL || (STRATA_OPAQUE_MATERIAL && !bBottomLayer));
-				if (bIsSimpleVolume)
-				{
-					EnableSlabBSDFSimpleVolumetric(BSDF);
-				}
-			}
+		//
+		// TODO the top root node now contains the coverage and transmittance of the entire material (used for blending and fog on transparent_.
+		//
+		OutCoverage = saturate(StrataTree.Operators[StrataData.OperatorIndex].Coverage);
+		OutTransmittancePreCoverage = saturate(StrataTree.Operators[StrataData.OperatorIndex].Transmittance);
+
+		//
+		// Forward lighting
+		//
+		STRATA_UNROLL_N(STRATA_CLAMPED_BSDF_COUNT)
+		for (int BSDFIdx = 0; BSDFIdx < StrataTree.BSDFCount; ++BSDFIdx)
+		{
+			#define CurrentBSDF StrataTree.BSDFs[BSDFIdx]
 
-			const float3 BSDFCoverage = BSDF.Coverage;
-			const float3 BSDFColoredVisibility = TransmittancePostCoverage * BSDFCoverage;
+			// This causes some flickering.
+			//if (!StrataIsBSDFVisible(CurrentBSDF))
+			//{
+			//	continue;
+			//}
 
-			// Create the BSDF context
-			FStrataBSDFContext StrataBSDFContext = StrataCreateBSDFContext(StrataPixelHeader, BSDF, NullStrataAddressing, V, L);
+			// Compute current BSDF transmittance
+			FStrataBSDFContext StrataBSDFContext = StrataCreateBSDFContext(StrataPixelHeader, CurrentBSDF, NullStrataAddressing, V, L);
 			
 			#if FORCE_FULLY_ROUGH
 			const bool bForceFullyRough = true;
 			#else
 			const bool bForceFullyRough = View.RenderingReflectionCaptureMask > 0;
 			#endif
-
+			
 			////
 			//// Evaluate environment lighting
 			////
@@ -302,7 +308,7 @@ float3 StrataForwardLighting(
 			if (UseBasePassSkylight > 0)
 			{
 				const bool bEnableSpecular = ReflectionStruct.SkyLightParameters.y > 0.0f;
-				float Roughness = StrataGetBSDFRoughness(BSDF);
+				float Roughness = StrataGetBSDFRoughness(CurrentBSDF);
 				float IndirectOcclusion = 1.0f;
 				float IndirectIrradiance = 0.0f; 
 
@@ -356,7 +362,7 @@ float3 StrataForwardLighting(
 					#endif
 
 						Color += (DiffuseIndirectLighting * DiffuseColorForIndirect + SubsurfaceIndirectLighting * SubsurfaceColorForIndirect) * 
-							BSDFColoredVisibility * AOMultiBounce(StrataEnvLight.DiffuseColor, DiffOcclusion);
+							CurrentBSDF.LuminanceWeight * AOMultiBounce(StrataEnvLight.DiffuseColor, DiffOcclusion);
 					}
 				}
 
@@ -371,7 +377,7 @@ float3 StrataForwardLighting(
 					{
 						int SingleCaptureIndex = GetPrimitiveData(MaterialParameters).SingleCaptureIndex;
 						Color += GetImageBasedReflectionSpecular(MaterialParameters, StrataEnvLight.SpecularDirection, Roughness, IndirectIrradiance, GridIndex, SingleCaptureIndex, EyeIndex) *
-							 StrataEnvLight.SpecularWeight * IndirectOcclusion * BSDFColoredVisibility * AOMultiBounce(StrataEnvLight.SpecularColor, SpecOcclusion);
+							 StrataEnvLight.SpecularWeight * IndirectOcclusion * CurrentBSDF.LuminanceWeight * AOMultiBounce(StrataEnvLight.SpecularColor, SpecOcclusion);
 					}
 					// We skip hazy lob evaluation in forward.
 				}
@@ -404,7 +410,6 @@ float3 StrataForwardLighting(
 			ComputeVolumeUVs(TranslatedWorldPosition, LightingPositionOffset, InnerVolumeUVs, OuterVolumeUVs, FinalLerpFactor);
 
 			FStrataEvaluateResult BSDFEvaluate = StrataEvaluateBSDF(StrataBSDFContext, bRoughDiffuseEnabled);
-			LayerThroughput += BSDFCoverage * BSDFEvaluate.Throughput;
 
 		#if TRANSLUCENCY_PERVERTEX_FORWARD_SHADING
 
@@ -471,7 +476,7 @@ float3 StrataForwardLighting(
 			}
 		#endif
 
-			Color += (SurfaceLighting + BSDFEvaluate.EmissivePathValue) * BSDFColoredVisibility;
+			Color += (SurfaceLighting + BSDFEvaluate.EmissivePathValue) * CurrentBSDF.LuminanceWeight;
 
 		#elif FORWARD_PER_PIXEL_SHADING
 
@@ -513,21 +518,11 @@ float3 StrataForwardLighting(
 						AbsoluteWorldPosition,
 						TranslatedWorldPosition,
 						SceneDepth,
-						BSDFColoredVisibility,
+						CurrentBSDF.LuminanceWeight,
 						StrataPixelHeader,
 						StrataBSDFContext,
 						bEvaluateHasBeenExecuted,
 						BSDFEvaluate) * DirectionalLightCloudShadow;
-
-					LayerThroughput += BSDFCoverage * BSDFEvaluate.Throughput;
-				}
-
-				BRANCH
-				if(!bEvaluateHasBeenExecuted)
-				{
-					// Only update the throughput
-					FStrataEvaluateResult BSDFEvaluate = StrataEvaluateBSDF(StrataBSDFContext, bRoughDiffuseEnabled);
-					LayerThroughput += BSDFCoverage * BSDFEvaluate.Throughput;
 				}
 			}
 
@@ -608,7 +603,7 @@ float3 StrataForwardLighting(
 					AbsoluteWorldPosition,
 					TranslatedWorldPosition,
 					SceneDepth,
-					BSDFColoredVisibility,
+					CurrentBSDF.LuminanceWeight,
 					StrataPixelHeader,
 					StrataBSDFContext,
 					bEvaluateHasBeenExecuted,
@@ -619,38 +614,11 @@ float3 StrataForwardLighting(
 
 		#endif // Lighting technique
 
-			LayerCoverage	+= BSDFCoverage;
+		#undef CurrentBSDF
 		}
-
-
-
-		// Now we are going to compute TransmittancePreCoverage and Coverage by combining the top layer0 and bottom layer1.
-		const float3 Layer0Meant = OutTransmittancePreCoverage;
-		const float3 Layer1Meant = LayerThroughput / max(1e-5, LayerCoverage);
-
-		const float Layer0Coverage = OutCoverage;
-		const float Layer1Coverage = LayerCoverage;
-		const float LayerMaxCoverage = max(Layer0Coverage, Layer1Coverage);
-		const float LayerMaxCoverageInv = LayerMaxCoverage <= 0.0f ? 1.0f : 1.0f / LayerMaxCoverage;
-
-		// If a layer has a lower coverage than the other, we account for that when compuing the TransmittancePreCoverage.
-		const float TransmittanceOne = 1.0f;
-		const float3 LayertT0 = lerp(1.0, Layer0Meant, Layer0Coverage * LayerMaxCoverageInv); // Layer0Meant * (Layer0Coverage * LayerMaxCoverageInv) + TransmittanceOne * saturate(1.0f - Layer0Coverage * LayerMaxCoverageInv);
-		const float3 LayertT1 = lerp(1.0, Layer1Meant, Layer1Coverage * LayerMaxCoverageInv); // Layer1Meant * (Layer1Coverage * LayerMaxCoverageInv) + TransmittanceOne * saturate(1.0f - Layer1Coverage * LayerMaxCoverageInv);
-
-		// Now we update the material transmittance pre coverage
-		OutTransmittancePreCoverage = LayertT0 * LayertT1;
-		OutCoverage = LayerMaxCoverage;
-
-
-
-		// Take into account the non covered part for this layer and update throughput accordingly
-		LayerThroughput += (1.0 - saturate(LayerCoverage)) * FullThroughput;
-
-		// And combine this layer throughtput with the full material throughput
-		TransmittancePostCoverage *= LayerThroughput;
 	}
 
+
 	return Color;
 }
 
diff --git a/Engine/Shaders/Private/Strata/StrataLegacyConversion.ush b/Engine/Shaders/Private/Strata/StrataLegacyConversion.ush
--- a/Engine/Shaders/Private/Strata/StrataLegacyConversion.ush
+++ b/Engine/Shaders/Private/Strata/StrataLegacyConversion.ush
@@ -25,6 +25,9 @@ struct FStrataLegacyParameters
 	float  Thickness;
 	uint   SharedLocalBasisIndex;
 
+	int bIsTop;
+	int bIsBottom;
+
 	float Weight;
 };
 
@@ -52,37 +55,48 @@ FStrataLegacyParameters InitStrataLegacyParameters(uint InSharedLocalBasisIndex,
 	Out.FuzzColor = Zeros;
 	Out.Thickness = STRATA_LAYER_DEFAULT_THICKNESS_CM;
 	Out.SharedLocalBasisIndex = InSharedLocalBasisIndex;
+	Out.bIsTop = 1;
+	Out.bIsBottom = 1;
 	Out.Weight = InWeight;
 	return Out;
 }
 
-FStrataData CreateLegacySlab(FStrataLegacyParameters In, inout uint SharedLocalBasisTypes)
+FStrataData CreateLegacySlab(
+	FStrataLegacyParameters In, inout uint SharedLocalBasisTypes,
+	inout FStrataTree StrataTree,
+	int OperatorIndex, int BSDFIndex)
 {
 	const float3 Zeros = float3(0, 0, 0);
 	const float3 Ones  = float3(1, 1, 1);
 
 	// Fixed layer structure for helping compiler to unroll and optimize shader
-	return GetStrataSlabBSDF(
-		In.UseMetalness,	// UseMetalness
-		In.BaseColor,		// BaseColor		-  Metalness workflow
-		Ones,				// EdgeColor		-  Metalness workflow
-		In.Specular,		// Specular			-  Metalness workflow
-		In.Metallic,		// Metallic			-  Metalness workflow
-		In.DiffuseAlbedo,	// DiffuseAlbedo	- !Metalness workflow
-		In.F0,				// F0				- !Metalness workflow
-		In.F90,				// F90				- !Metalness workflow
-		In.Roughness,		// Roughness
-		In.Anisotropy,		// Anisotropy
-		In.SSSProfileID,	// SSSProfileID
-		In.SSSDMFP,			// SSSDMFP
-		In.SSSDMFPScale,	// SSSDMFPScale
-		In.Emissive,		// Emissive
-		0.f,				// Haziness
-		0.f,				// ThinFilmThickness
-		In.FuzzAmount,		// FuzzAmount
-		In.FuzzColor,		// FuzzColor
-		In.Thickness,
-		In.SharedLocalBasisIndex, SharedLocalBasisTypes);
+	return PromoteParameterBlendedBSDFToOperator(
+		GetStrataSlabBSDF(
+			In.UseMetalness,	// UseMetalness
+			In.BaseColor,		// BaseColor		-  Metalness workflow
+			Ones,				// EdgeColor		-  Metalness workflow
+			In.Specular,		// Specular			-  Metalness workflow
+			In.Metallic,		// Metallic			-  Metalness workflow
+			In.DiffuseAlbedo,	// DiffuseAlbedo	- !Metalness workflow
+			In.F0,				// F0				- !Metalness workflow
+			In.F90,				// F90				- !Metalness workflow
+			In.Roughness,		// Roughness
+			In.Anisotropy,		// Anisotropy
+			In.SSSProfileID,	// SSSProfileID
+			In.SSSDMFP,			// SSSDMFP
+			In.SSSDMFPScale,	// SSSDMFPScale
+			In.Emissive,		// Emissive
+			0.f,				// Haziness
+			0.f,				// ThinFilmThickness
+			In.FuzzAmount,		// FuzzAmount
+			In.FuzzColor,		// FuzzColor
+			In.Thickness,
+			In.SharedLocalBasisIndex, SharedLocalBasisTypes),
+		StrataTree,
+		OperatorIndex,
+		BSDFIndex,
+		In.bIsBottom,
+		In.bIsTop);
 }
 
 // Convert legacy shading models - Dynamic
@@ -102,7 +116,8 @@ FStrataData StrataConvertLegacyMaterialDynamic(
 	uint ShadingModel,
 	uint SharedLocalBasisIndex,
 	uint ClearCoatBottomNormal_SharedLocalBasisIndex,
-	inout uint SharedLocalBasisTypes)
+	inout uint SharedLocalBasisTypes,
+	inout FStrataTree StrataTree)
 {
 	const float DefaultThickness = STRATA_LAYER_DEFAULT_THICKNESS_CM;
 	const float3 Zeros = float3(0, 0, 0);
@@ -235,6 +250,8 @@ FStrataData StrataConvertLegacyMaterialDynamic(
 		Bot.SSSDMFPScale = 0.f;
 		Bot.Emissive = Zeros;
 		Bot.SharedLocalBasisIndex = ClearCoatBottomNormal_SharedLocalBasisIndex;
+		Bot.bIsTop = 0;
+		Bot.bIsBottom = 1;
 		Bot.Weight = 1.f;
 		// === TOP layer === 
 		 																								   
@@ -260,6 +277,8 @@ FStrataData StrataConvertLegacyMaterialDynamic(
 		Top.Emissive = Emissive;
 		Top.Thickness = Thickness;
 		Top.SharedLocalBasisIndex = SharedLocalBasisIndex;
+		Top.bIsTop = 1;
+		Top.bIsBottom = 0;
 		Top.Weight = ClearCoat;
 
 		FinalWeight = Opacity;
@@ -322,6 +341,8 @@ FStrataData StrataConvertLegacyMaterialDynamic(
 		Top.SSSDMFPScale = Opacity;
 		Top.Emissive = Emissive;
 		Top.SharedLocalBasisIndex = SharedLocalBasisIndex;
+		Top.bIsTop = 1;
+		Top.bIsBottom = 0;
 		Top.Weight = 1.f;
 
 		// Iris
@@ -335,6 +356,8 @@ FStrataData StrataConvertLegacyMaterialDynamic(
 		Bot.SSSDMFPScale = Opacity;
 		Bot.Emissive = Emissive;
 		Bot.SharedLocalBasisIndex = ClearCoatBottomNormal_SharedLocalBasisIndex;
+		Bot.bIsTop = 0;
+		Bot.bIsBottom = 1;
 		Bot.Weight = 1.f;
 
 		// No need to take into account opacity weighting for translucent material, as the coverage is assumed to be 1
@@ -382,7 +405,7 @@ FStrataData StrataConvertLegacyMaterialDynamic(
 		Top.SharedLocalBasisIndex = SharedLocalBasisIndex;
 
 		Top.Weight  = 1.0f;
-		Bot.Weight  = 0.f;
+		Bot.Weight  = 0.0f;
 		FinalWeight = 1.0f;
 
 		// Reference two layers approach
@@ -423,14 +446,17 @@ FStrataData StrataConvertLegacyMaterialDynamic(
 	}
 
 	// Fixed layer structure for helping compiler to unroll and optimize shader
-	FStrataData Top_Slab = CreateLegacySlab(Top, SharedLocalBasisTypes);
-	FStrataData Bot_Slab = CreateLegacySlab(Bot, SharedLocalBasisTypes);
-	FStrataData Out = StrataVerticalLayering(StrataWeight(Top_Slab, Top.Weight), StrataWeight(Bot_Slab, Bot.Weight));
+	// This must match what is done in PackStrataOut
+	FStrataData Top_Slab = CreateLegacySlab(Top, SharedLocalBasisTypes, StrataTree,	/*OperatorIndex*/0, /*BSDFIndex*/0);
+	FStrataData Bot_Slab = CreateLegacySlab(Bot, SharedLocalBasisTypes, StrataTree,	/*OperatorIndex*/1, /*BSDFIndex*/1);
+	FStrataData WTop_Slab= StrataWeight(Top_Slab, Top.Weight, StrataTree,			/*OperatorIndex*/2, /*MaxDistanceFromLeaves*/1);
+	FStrataData WBot_Slab= StrataWeight(Bot_Slab, Bot.Weight, StrataTree,			/*OperatorIndex*/3, /*MaxDistanceFromLeaves*/1);
+	FStrataData Out		 = StrataVerticalLayering(WTop_Slab, WBot_Slab, StrataTree,	/*OperatorIndex*/4, /*MaxDistanceFromLeaves*/2);
 
 	// Take in account the legacy opacity weighting for translucent material
 	#if !MATERIALBLENDING_MASKED && !MATERIALBLENDING_SOLID
 	{
-		Out = StrataWeight(Out, FinalWeight);
+		Out = StrataWeight(Out, FinalWeight, StrataTree,	/*OperatorIndex*/5, /*MaxDistanceFromLeaves*/3);
 	}
 	#endif
 
@@ -453,7 +479,8 @@ FStrataData StrataConvertLegacyMaterialStatic(
 	uint ShadingModel,
 	uint SharedLocalBasisIndex,
 	uint ClearCoatBottomNormal_SharedLocalBasisIndex,
-	inout uint SharedLocalBasisTypes)
+	inout uint SharedLocalBasisTypes,
+	inout FStrataTree StrataTree)
 {
 	FStrataData Out = GetInitialisedStrataData();
 
@@ -464,7 +491,7 @@ FStrataData StrataConvertLegacyMaterialStatic(
 	if (ShadingModel == SHADINGMODELID_UNLIT)
 	{
 		Out = GetStrataUnlitBSDF(Emissive, Zeros);
-		Out = StrataWeight(Out, Opacity);
+//		Out = StrataWeight(Out, Opacity);				// STRATA_TODO
 	}
 	else if (ShadingModel == SHADINGMODELID_HAIR)
 	{
@@ -513,7 +540,8 @@ FStrataData StrataConvertLegacyMaterialStatic(
 			ShadingModel,
 			SharedLocalBasisIndex,
 			ClearCoatBottomNormal_SharedLocalBasisIndex,
-			SharedLocalBasisTypes);
+			SharedLocalBasisTypes,
+			StrataTree);
 	}
 
 	return Out;
diff --git a/Engine/Shaders/Private/Strata/StrataTree.ush b/Engine/Shaders/Private/Strata/StrataTree.ush
new file mode 100644
--- /dev/null
+++ b/Engine/Shaders/Private/Strata/StrataTree.ush
@@ -0,0 +1,209 @@
+
+
+#include "/Engine/Private/Strata/StrataEvaluation.ush"
+#include "/Engine/Private/ShadingCommon.ush"
+
+#ifndef STRATA_OPAQUE_MATERIAL
+#define STRATA_OPAQUE_MATERIAL 0
+#endif
+
+void UpdateSingleBSDFOperatorCoverageTransmittance(
+	FStrataPixelHeader StrataPixelHeader,
+	inout FStrataTree StrataTree, 
+	inout FStrataBSDF CurrentBSDF,
+	bool bRoughDiffuseEnabled,
+	FStrataAddressing NullStrataAddressing,
+	float3 V,
+	float3 L
+)
+{
+	// Sanitize BSDF before it is used for forward shading
+	StrataSanitizeBSDF(CurrentBSDF);
+
+	BRANCH
+	if (BSDF_GETTYPE(CurrentBSDF) == STRATA_BSDF_TYPE_SLAB)
+	{
+		// We want simple volumetric only if the DMFP input is pluged in, otherwise we avoid dealing with simple volumetric.
+		// We also do not want to use the simple volumetric path if we are dealing with an opaque material and this slab is not at the bottom.
+		// And when opaque, the bottom layer must be opaque and thus will ahve a solid opaque diffuse color.
+		const bool bIsSimpleVolume = BSDF_GETHASDMFP(CurrentBSDF) && (!STRATA_OPAQUE_MATERIAL || (STRATA_OPAQUE_MATERIAL && !CurrentBSDF.bIsBottom));
+
+		if (bIsSimpleVolume)
+		{
+			EnableSlabBSDFSimpleVolumetric(CurrentBSDF);
+		}
+	}
+
+	// Compute current BSDF transmittance
+	FStrataBSDFContext StrataBSDFContext = StrataCreateBSDFContext(StrataPixelHeader, CurrentBSDF, NullStrataAddressing, V, L);
+	FStrataEvaluateResult BSDFEvaluate = StrataEvaluateBSDF(StrataBSDFContext, bRoughDiffuseEnabled);
+
+	float CurrentBSDFCoverage = CurrentBSDF.Coverage;
+	float3 CurrentBSDFTransmittance = BSDFEvaluate.Throughput;
+
+	// We multiply this weight with coverage:
+	//  - Because when coming from parameter blending, we need to scale the output down according to coverage resulting from parameter blending.
+	//  - That will be applied on emissive and reflected light, so we do not ned to apply that anywhere else.
+	//  - When coming from non parameter blending, this is equal identity 1 and the operator parsing will continue to correctly weight the BSDF according to the tree.
+	CurrentBSDF.LuminanceWeight = float3(1.0f, 1.0f, 1.0f) * CurrentBSDFCoverage;
+
+	StrataTree.Operators[CurrentBSDF.OperatorIndex].Coverage = CurrentBSDFCoverage;
+	StrataTree.Operators[CurrentBSDF.OperatorIndex].Transmittance = CurrentBSDFTransmittance;
+}
+
+
+void UpdateSingleOperatorCoverageTransmittance(inout FStrataTree StrataTree, int OpIndex)
+{
+#define Op StrataTree.Operators[OpIndex]
+#define OpA StrataTree.Operators[StrataTree.Operators[OpIndex].AIndex]
+#define OpB StrataTree.Operators[StrataTree.Operators[OpIndex].BIndex]
+
+	switch (Op.Type)
+	{
+
+	case STRATA_OPERATOR_WEIGHT:
+	{
+		const float Weight = saturate(Op.Weight);
+		Op.Coverage = Weight * OpA.Coverage;
+		Op.Transmittance = OpA.Transmittance;
+		break;
+	}
+
+	case STRATA_OPERATOR_VERTICAL:
+	{
+		const float3 TopTransmittance = OpA.Transmittance;
+		const float3 BotTransmittance = OpB.Transmittance;
+		const float TopCoverage = OpA.Coverage;
+		const float BotCoverage = OpB.Coverage;
+
+#if 0
+		// This is coverage assuming correlation
+		const float LayerMaxCoverage = max(TopCoverage, BotCoverage);
+		const float LayerMaxCoverageInv = LayerMaxCoverage <= 0.0f ? 1.0f : 1.0f / LayerMaxCoverage;
+
+		// If a layer has a lower coverage than the other, we account for that when compuing the Transmittance.
+		const float TransmittanceOne = 1.0f;
+		const float3 TopTransmittanceAdjusted = lerp(TransmittanceOne, TopTransmittance, TopCoverage * LayerMaxCoverageInv);
+		const float3 BotTransmittanceAdjusted = lerp(TransmittanceOne, BotTransmittance, BotCoverage * LayerMaxCoverageInv);
+
+		// Now we update the material transmittance and coverage
+		Op.Coverage = LayerMaxCoverage;
+		Op.Transmittance = TopTransmittanceAdjusted * BotTransmittanceAdjusted;
+#else
+		FVerticalLayeringInfo Info = GetVerticalLayeringInfo(TopCoverage, BotCoverage);
+
+		Op.Coverage					= Info.Coverage;
+		Op.Transmittance			= Info.TransmittanceOnlyTop * TopTransmittance + Info.TransmittanceOnlyBottom * BotTransmittance + Info.TransmittanceTopAndBottom * TopTransmittance * BotTransmittance;
+#endif
+
+		Op.VerticalTopCoverage		= OpA.Coverage;
+		Op.VerticalTopTransmittance = OpA.Transmittance;
+		break;
+	}
+
+	case STRATA_OPERATOR_HORIZONTAL:
+	{
+		const float Mix = saturate(Op.Weight);
+		const float AMix = 1.0 - Mix;
+		const float BMix = Mix;
+
+		Op.Coverage = AMix * OpA.Coverage + BMix * OpB.Coverage;
+		Op.Transmittance = (AMix * OpA.Coverage * OpA.Transmittance + BMix * OpB.Coverage * OpB.Transmittance) / max(1e-5, AMix * OpA.Coverage + BMix * OpB.Coverage);
+		break;
+	}
+
+	case STRATA_OPERATOR_ADD:
+	{
+		Op.Coverage = saturate(OpA.Coverage + OpB.Coverage);
+		Op.Transmittance = saturate(OpA.Transmittance + OpB.Transmittance);
+		break;
+	}
+
+	}
+
+#undef Op
+#undef OpA
+#undef OpB
+}
+
+
+
+void UpdateBSDFWeightAfterOperatorVisit(
+	inout FStrataTree StrataTree,
+	inout FStrataBSDF CurrentBSDF, 
+	int OpIndex,
+	int PreviousIsInputA)
+{
+#define Op StrataTree.Operators[OpIndex]
+
+	switch (Op.Type)
+	{
+
+	case STRATA_OPERATOR_WEIGHT:
+	{
+		const float Weight = saturate(Op.Weight);
+		CurrentBSDF.LuminanceWeight *= Weight;
+#if STRATA_INLINE_SHADING
+		CurrentBSDF.Coverage *= Weight;
+#endif
+		break;
+	}
+
+	case STRATA_OPERATOR_VERTICAL:
+	{
+		// PreviousIsInputA > 0 means it is a BSDF affecting the top layer so we do not affect it by anything for this vertical layering.
+		// Otherise, the BSDF comes from the bottom part so we affect it weights from the gathered top layer coverage/transmittance
+		CurrentBSDF.LuminanceWeight *= PreviousIsInputA > 0 ? 1.0f : saturate(1.0f - Op.VerticalTopCoverage) + saturate(Op.VerticalTopCoverage) * Op.VerticalTopTransmittance;
+		break;
+	}
+
+	case STRATA_OPERATOR_HORIZONTAL:
+	{
+		const float Mix = saturate(Op.Weight);
+		const float AMix = 1.0 - Mix;
+		const float BMix = Mix;
+		const float Weight = PreviousIsInputA > 0 ? AMix : BMix;
+		CurrentBSDF.LuminanceWeight *= Weight;
+#if STRATA_INLINE_SHADING
+		CurrentBSDF.Coverage *= Weight;
+#endif
+		break;
+	}
+
+	// case STRATA_OPERATOR_ADD: NOP
+
+	}
+
+#undef Op
+}
+
+
+void LegacyUpdateBSDFsOperators(inout FStrataPixelHeader StrataPixelHeader, inout FStrataTree StrataTree, bool bRoughDiffuseEnabled, FStrataAddressing NullStrataAddressing, float3 V, float3 NullLightDir)
+{
+	// This must match what is done in StrataConvertLegacyMaterialDynamic
+	UpdateSingleBSDFOperatorCoverageTransmittance(StrataPixelHeader, StrataTree, StrataTree.BSDFs[0], bRoughDiffuseEnabled, NullStrataAddressing, V, NullLightDir);
+	UpdateSingleBSDFOperatorCoverageTransmittance(StrataPixelHeader, StrataTree, StrataTree.BSDFs[1], bRoughDiffuseEnabled, NullStrataAddressing, V, NullLightDir);
+
+	// MaxDistanceFromLeaves = 1 
+	UpdateSingleOperatorCoverageTransmittance(StrataTree, /*OperatorIndex*/2);
+	UpdateSingleOperatorCoverageTransmittance(StrataTree, /*OperatorIndex*/3);
+	// MaxDistanceFromLeaves = 2 
+	UpdateSingleOperatorCoverageTransmittance(StrataTree, /*OperatorIndex*/4);
+#if !MATERIALBLENDING_MASKED && !MATERIALBLENDING_SOLID
+	// MaxDistanceFromLeaves = 3 
+	UpdateSingleOperatorCoverageTransmittance(StrataTree, /*OperatorIndex*/5);
+#endif
+
+	UpdateBSDFWeightAfterOperatorVisit(StrataTree, StrataTree.BSDFs[0], /*OperatorIndex*/2, /*PreviousIsInputA*/1); // WEIGHT
+	UpdateBSDFWeightAfterOperatorVisit(StrataTree, StrataTree.BSDFs[0], /*OperatorIndex*/4, /*PreviousIsInputA*/1); // VERTICAL
+#if !MATERIALBLENDING_MASKED && !MATERIALBLENDING_SOLID
+	UpdateBSDFWeightAfterOperatorVisit(StrataTree, StrataTree.BSDFs[0], /*OperatorIndex*/5, /*PreviousIsInputA*/1); // WEIGHT
+#endif
+	UpdateBSDFWeightAfterOperatorVisit(StrataTree, StrataTree.BSDFs[1], /*OperatorIndex*/3, /*PreviousIsInputA*/1); // WEIGHT
+	UpdateBSDFWeightAfterOperatorVisit(StrataTree, StrataTree.BSDFs[1], /*OperatorIndex*/4, /*PreviousIsInputA*/0); // VERTICAL
+#if !MATERIALBLENDING_MASKED && !MATERIALBLENDING_SOLID
+	UpdateBSDFWeightAfterOperatorVisit(StrataTree, StrataTree.BSDFs[1], /*OperatorIndex*/5, /*PreviousIsInputA*/1); // WEIGHT
+#endif
+}
+
+
diff --git a/Engine/Shaders/Private/Strata/StrataVisualize.usf b/Engine/Shaders/Private/Strata/StrataVisualize.usf
--- a/Engine/Shaders/Private/Strata/StrataVisualize.usf
+++ b/Engine/Shaders/Private/Strata/StrataVisualize.usf
@@ -202,7 +202,7 @@ void VisualizeMaterialPS(
 #define PRINTBSDFCHAR(X) PrintCharacter(PixelPos, OutColor.xyz, BSDFFontColor, TopLeft, X);
 #define PRINTBSDFCHARGREY(X) PrintCharacter(PixelPos, OutColor.xyz, BSDFFontColorGrey, TopLeft, X);
 
-#define PRINTBSDFWEIGHT()	PRINTBSDFCHAR(_SPC_); PRINTBSDFCHAR(_MINUS_); PRINTBSDFCHAR(_W_); PRINTBSDFCHAR(_E_); PRINTBSDFCHAR(_I_); PRINTBSDFCHAR(_G_); PRINTBSDFCHAR(_H_);PRINTBSDFCHAR(_T_);  PRINTBSDFCHAR(_SPC_); PrintFloat(PixelPos, OutColor.xyz, FontColor, TopLeft, BSDF.FinalWeight.x); TopLeft.x+=8*7; PRINTBSDFCHAR(_SPC_); PrintFloat(PixelPos, OutColor.xyz, FontColor, TopLeft, BSDF.FinalWeight.y); TopLeft.x+=8*7; PRINTBSDFCHAR(_SPC_); PrintFloat(PixelPos, OutColor.xyz, FontColor, TopLeft, BSDF.FinalWeight.z); TopLeft.x+=8*7; 
+#define PRINTBSDFWEIGHT()	PRINTBSDFCHAR(_SPC_); PRINTBSDFCHAR(_MINUS_); PRINTBSDFCHAR(_W_); PRINTBSDFCHAR(_E_); PRINTBSDFCHAR(_I_); PRINTBSDFCHAR(_G_); PRINTBSDFCHAR(_H_);PRINTBSDFCHAR(_T_);  PRINTBSDFCHAR(_SPC_); PrintFloat(PixelPos, OutColor.xyz, FontColor, TopLeft, BSDF.LuminanceWeight.x); TopLeft.x+=8*7; PRINTBSDFCHAR(_SPC_); PrintFloat(PixelPos, OutColor.xyz, FontColor, TopLeft, BSDF.LuminanceWeight.y); TopLeft.x+=8*7; PRINTBSDFCHAR(_SPC_); PrintFloat(PixelPos, OutColor.xyz, FontColor, TopLeft, BSDF.LuminanceWeight.z); TopLeft.x+=8*7; 
 #define PRINTBSDFGREYWEIGHT()PRINTBSDFCHAR(_SPC_);PRINTBSDFCHARGREY(_MINUS_); PRINTBSDFCHARGREY(_W_); PRINTBSDFCHARGREY(_E_); PRINTBSDFCHARGREY(_I_); PRINTBSDFCHARGREY(_G_); PRINTBSDFCHARGREY(_H_);PRINTBSDFCHARGREY(_T_);  PRINTBSDFCHAR(_SPC_); PrintFloat(PixelPos, OutColor.xyz, FontColor, TopLeft, Unpack10F(BSDF_GETWEIGHT10F(BSDF))); TopLeft.x+=8*23; 
 #define PRINTBSDFSHAREDLOCALBASISID()	PRINTBSDFCHAR(_SPC_); PRINTBSDFCHAR(_MINUS_); PRINTBSDFCHAR(_N_); PRINTBSDFCHAR(_R_); PRINTBSDFCHAR(_M_); PRINTBSDFCHAR(_I_); PRINTBSDFCHAR(_D_); PRINTBSDFCHAR(_SPC_); PrintFloatNoFraction(PixelPos, OutColor.xyz, FontColor, TopLeft, float(BSDF_GETSHAREDLOCALBASISID(BSDF)), 1); TopLeft.x+=8*1;
 #define PRINTBSDFANISO()	PRINTBSDFCHAR(_SPC_); PRINTBSDFCHAR(_MINUS_); PRINTBSDFCHAR(_A_); PRINTBSDFCHAR(_N_); PRINTBSDFCHAR(_I_); PRINTBSDFCHAR(_S_); PRINTBSDFCHAR(_O_); PRINTBSDFCHAR(_SPC_); PrintFloatNoFraction(PixelPos, OutColor.xyz, FontColor, TopLeft, float(BSDF_GETHASANISOTROPY(BSDF)), 1); TopLeft.x+=8*1;
diff --git a/Engine/Shaders/Private/VolumetricCloud.usf b/Engine/Shaders/Private/VolumetricCloud.usf
--- a/Engine/Shaders/Private/VolumetricCloud.usf
+++ b/Engine/Shaders/Private/VolumetricCloud.usf
@@ -51,7 +51,7 @@ float3 SampleExtinctionCoefficients(in FPixelMaterialInputs PixelMaterialInputs)
 {
 	float3 Extinction = 0.0f;
 #if STRATA_ENABLED
-	FStrataBSDF BSDF = PixelMaterialInputs.FrontMaterial.Layers[0].BSDFs[0];
+	FStrataBSDF BSDF = PixelMaterialInputs.FrontMaterial.InlinedBSDF;
 	Extinction = VOLUMETRICFOGCLOUD_EXTINCTION(BSDF).r;
 #else
 	#if !MATERIAL_SHADINGMODEL_UNLIT
@@ -65,7 +65,7 @@ float3 SampleEmissive(in FPixelMaterialInputs PixelMaterialInputs)
 {
 	float3 EmissiveColor = 0.0f;
 #if STRATA_ENABLED
-	FStrataBSDF BSDF = PixelMaterialInputs.FrontMaterial.Layers[0].BSDFs[0];
+	FStrataBSDF BSDF = PixelMaterialInputs.FrontMaterial.InlinedBSDF;
 	EmissiveColor = BSDF_GETEMISSIVE(BSDF);
 #else
 	EmissiveColor = GetMaterialEmissiveRaw(PixelMaterialInputs);
@@ -77,7 +77,7 @@ float3 SampleAlbedo(in FPixelMaterialInputs PixelMaterialInputs)
 {
 	float3 Albedo = 0.0f;
 #if STRATA_ENABLED
-	FStrataBSDF BSDF = PixelMaterialInputs.FrontMaterial.Layers[0].BSDFs[0];
+	FStrataBSDF BSDF = PixelMaterialInputs.FrontMaterial.InlinedBSDF;
 	Albedo = VOLUMETRICFOGCLOUD_ALBEDO(BSDF);
 #else
 	#if !MATERIAL_SHADINGMODEL_UNLIT
@@ -91,7 +91,7 @@ float SampleAmbientOcclusion(in FPixelMaterialInputs PixelMaterialInputs)
 {
 	float AO = 0.0f;
 #if STRATA_ENABLED
-	FStrataBSDF BSDF = PixelMaterialInputs.FrontMaterial.Layers[0].BSDFs[0];
+	FStrataBSDF BSDF = PixelMaterialInputs.FrontMaterial.InlinedBSDF;
 	AO = VOLUMETRICFOGCLOUD_AO(BSDF);
 #else
 	AO = GetMaterialAmbientOcclusion(PixelMaterialInputs);
diff --git a/Engine/Shaders/Private/VolumetricFogVoxelization.usf b/Engine/Shaders/Private/VolumetricFogVoxelization.usf
--- a/Engine/Shaders/Private/VolumetricFogVoxelization.usf
+++ b/Engine/Shaders/Private/VolumetricFogVoxelization.usf
@@ -27,7 +27,7 @@ float SampleExtinctionCoefficients(in FPixelMaterialInputs PixelMaterialInputs)
 {
 	float Extinction = 0.0f;
 #if STRATA_ENABLED
-	FStrataBSDF BSDF = PixelMaterialInputs.FrontMaterial.Layers[0].BSDFs[0];
+	FStrataBSDF BSDF = PixelMaterialInputs.FrontMaterial.InlinedBSDF;
 	Extinction = VOLUMETRICFOGCLOUD_EXTINCTION(BSDF).r;
 #else
 	#if !MATERIAL_SHADINGMODEL_UNLIT
@@ -41,7 +41,7 @@ float3 SampleEmissive(in FPixelMaterialInputs PixelMaterialInputs)
 {
 	float3 EmissiveColor = 0.0f;
 #if STRATA_ENABLED
-	FStrataBSDF BSDF = PixelMaterialInputs.FrontMaterial.Layers[0].BSDFs[0];
+	FStrataBSDF BSDF = PixelMaterialInputs.FrontMaterial.InlinedBSDF;
 	EmissiveColor = BSDF_GETEMISSIVE(BSDF);
 #else
 	EmissiveColor = GetMaterialEmissiveRaw(PixelMaterialInputs);
@@ -53,7 +53,7 @@ float3 SampleAlbedo(in FPixelMaterialInputs PixelMaterialInputs, in ViewState In
 {
 	float3 Albedo = 0.0f;
 #if STRATA_ENABLED
-	FStrataBSDF BSDF = PixelMaterialInputs.FrontMaterial.Layers[0].BSDFs[0];
+	FStrataBSDF BSDF = PixelMaterialInputs.FrontMaterial.InlinedBSDF;
 	Albedo = VOLUMETRICFOGCLOUD_ALBEDO(BSDF);
 #else
 	#if !MATERIAL_SHADINGMODEL_UNLIT
diff --git a/Engine/Shaders/Shared/StrataDefinitions.h b/Engine/Shaders/Shared/StrataDefinitions.h
--- a/Engine/Shaders/Shared/StrataDefinitions.h
+++ b/Engine/Shaders/Shared/StrataDefinitions.h
@@ -15,6 +15,7 @@
 #define STRATA_MAX_BSDF_COUNT_PER_LAYER		4
 #define STRATA_MAX_LAYER_COUNT				4
 #define STRATA_MAX_BSDF_COUNT				15
+#define STRATA_MAX_OPERATOR_COUNT			15
 
 // It this is changed, STATE_BIT_COUNT_SHAREDLOCALBASESID and HEADER_BIT_COUNT_SHAREDLOCALBASES_COUNT also needs to be updated
 #define STRATA_MAX_SHAREDLOCALBASES_REGISTERS		4
@@ -43,3 +44,10 @@
 #define STRATA_LAYER_ISTHIN_THICKNESS_THRESHOLD_CM 0.01f
 
 #define STRATA_BASE_PASS_MRT_OUTPUT_COUNT	2
+
+#define STRATA_OPERATOR_WEIGHT			0
+#define STRATA_OPERATOR_VERTICAL		1
+#define STRATA_OPERATOR_HORIZONTAL		2
+#define STRATA_OPERATOR_ADD				3
+#define STRATA_OPERATOR_BSDF			4
+#define STRATA_OPERATOR_BSDF_LEGACY		5
diff --git a/Engine/Source/Runtime/Engine/Classes/Materials/MaterialExpression.h b/Engine/Source/Runtime/Engine/Classes/Materials/MaterialExpression.h
--- a/Engine/Source/Runtime/Engine/Classes/Materials/MaterialExpression.h
+++ b/Engine/Source/Runtime/Engine/Classes/Materials/MaterialExpression.h
@@ -17,6 +17,7 @@ class UMaterial;
 class UTexture;
 struct FPropertyChangedEvent;
 struct FMaterialShadingModelField;
+struct FStrataOperator;
 
 class FMaterialHLSLGenerator;
 enum class EMaterialNewScopeFlag : uint8;
@@ -422,6 +423,13 @@ class ENGINE_API UMaterialExpression : public UObject
 	 */
 	virtual void GatherStrataMaterialInfo(FStrataMaterialInfo& StrataMaterialInfo, int32 OutputIndex) { }
 
+	/**
+	 * A starta material is a tree with FrontMateiral being its root and BSDF being leaves, with operators in the middle.
+	 * This recursively parse nodes outputing strata material in order to gather the maximum distance to any leaves. 
+	 * This is used to drive the bottom up order processing of those nodes.
+	 */
+	virtual FStrataOperator* StrataGenerateMaterialTopologyTree(class FMaterialCompiler* Compiler, class UMaterialExpression* Parent, int32 OutputIndex) { return nullptr; }
+
 	/**
 	 * If true, discards the output index when caching this expression which allows more cases to re-use the output instead of adding a separate instruction
 	 */
diff --git a/Engine/Source/Runtime/Engine/Classes/Materials/MaterialExpressionBlendMaterialAttributes.h b/Engine/Source/Runtime/Engine/Classes/Materials/MaterialExpressionBlendMaterialAttributes.h
--- a/Engine/Source/Runtime/Engine/Classes/Materials/MaterialExpressionBlendMaterialAttributes.h
+++ b/Engine/Source/Runtime/Engine/Classes/Materials/MaterialExpressionBlendMaterialAttributes.h
@@ -52,6 +52,7 @@ class UMaterialExpressionBlendMaterialAttributes : public UMaterialExpression
 	virtual bool IsResultMaterialAttributes(int32 OutputIndex) override {return true;}
 	virtual bool IsResultStrataMaterial(int32 OutputIndex) override;
 	virtual void GatherStrataMaterialInfo(FStrataMaterialInfo& StrataMaterialInfo, int32 OutputIndex) override;
+	virtual FStrataOperator* StrataGenerateMaterialTopologyTree(class FMaterialCompiler* Compiler, class UMaterialExpression* Parent, int32 OutputIndex) override;
 	virtual uint32 GetInputType(int32 InputIndex) override {return InputIndex == 2 ? MCT_Float1 : MCT_MaterialAttributes;}
 #endif
 	//~ End UMaterialExpression Interface
diff --git a/Engine/Source/Runtime/Engine/Classes/Materials/MaterialExpressionFunctionOutput.h b/Engine/Source/Runtime/Engine/Classes/Materials/MaterialExpressionFunctionOutput.h
--- a/Engine/Source/Runtime/Engine/Classes/Materials/MaterialExpressionFunctionOutput.h
+++ b/Engine/Source/Runtime/Engine/Classes/Materials/MaterialExpressionFunctionOutput.h
@@ -65,6 +65,7 @@ class UMaterialExpressionFunctionOutput : public UMaterialExpression
 	virtual bool IsResultMaterialAttributes(int32 OutputIndex) override;
 	virtual bool IsResultStrataMaterial(int32 OutputIndex) override;
 	virtual void GatherStrataMaterialInfo(FStrataMaterialInfo& StrataMaterialInfo, int32 OutputIndex) override;
+	virtual FStrataOperator* StrataGenerateMaterialTopologyTree(class FMaterialCompiler* Compiler, class UMaterialExpression* Parent, int32 OutputIndex) override;
 	bool GenerateHLSLExpression(FMaterialHLSLGenerator& Generator, UE::HLSLTree::FScope& Scope, int32 OutputIndex, UE::HLSLTree::FExpression*& OutExpression);
 #endif // WITH_EDITOR
 	//~ End UMaterialExpression Interface
diff --git a/Engine/Source/Runtime/Engine/Classes/Materials/MaterialExpressionGetMaterialAttributes.h b/Engine/Source/Runtime/Engine/Classes/Materials/MaterialExpressionGetMaterialAttributes.h
--- a/Engine/Source/Runtime/Engine/Classes/Materials/MaterialExpressionGetMaterialAttributes.h
+++ b/Engine/Source/Runtime/Engine/Classes/Materials/MaterialExpressionGetMaterialAttributes.h
@@ -48,6 +48,7 @@ class UMaterialExpressionGetMaterialAttributes : public UMaterialExpression
 	virtual bool IsResultMaterialAttributes(int32 OutputIndex) override { return OutputIndex == 0; }
 	virtual bool IsResultStrataMaterial(int32 OutputIndex) override;
 	virtual void GatherStrataMaterialInfo(FStrataMaterialInfo& StrataMaterialInfo, int32 OutputIndex) override;
+	virtual FStrataOperator* StrataGenerateMaterialTopologyTree(class FMaterialCompiler* Compiler, class UMaterialExpression* Parent, int32 OutputIndex) override;
 #endif
 	//~ End UMaterialExpression Interface
 };
diff --git a/Engine/Source/Runtime/Engine/Classes/Materials/MaterialExpressionMaterialFunctionCall.h b/Engine/Source/Runtime/Engine/Classes/Materials/MaterialExpressionMaterialFunctionCall.h
--- a/Engine/Source/Runtime/Engine/Classes/Materials/MaterialExpressionMaterialFunctionCall.h
+++ b/Engine/Source/Runtime/Engine/Classes/Materials/MaterialExpressionMaterialFunctionCall.h
@@ -130,6 +130,7 @@ class UMaterialExpressionMaterialFunctionCall : public UMaterialExpression
 	virtual bool IsResultMaterialAttributes(int32 OutputIndex) override;
 	virtual bool IsResultStrataMaterial(int32 OutputIndex) override;
 	virtual void GatherStrataMaterialInfo(FStrataMaterialInfo& StrataMaterialInfo, int32 OutputIndex) override;
+	virtual FStrataOperator* StrataGenerateMaterialTopologyTree(class FMaterialCompiler* Compiler, class UMaterialExpression* Parent, int32 OutputIndex) override;
 	virtual uint32 GetInputType(int32 InputIndex) override;
 	//~ End UMaterialExpression Interface
 
diff --git a/Engine/Source/Runtime/Engine/Classes/Materials/MaterialExpressionStrata.h b/Engine/Source/Runtime/Engine/Classes/Materials/MaterialExpressionStrata.h
--- a/Engine/Source/Runtime/Engine/Classes/Materials/MaterialExpressionStrata.h
+++ b/Engine/Source/Runtime/Engine/Classes/Materials/MaterialExpressionStrata.h
@@ -170,6 +170,7 @@ class UMaterialExpressionStrataLegacyConversion : public UMaterialExpressionStra
 	virtual uint32 GetInputType(int32 InputIndex) override;
 	virtual bool IsResultStrataMaterial(int32 OutputIndex) override;
 	virtual void GatherStrataMaterialInfo(FStrataMaterialInfo& StrataMaterialInfo, int32 OutputIndex) override;
+	virtual FStrataOperator* StrataGenerateMaterialTopologyTree(class FMaterialCompiler* Compiler, class UMaterialExpression* Parent, int32 OutputIndex) override;
 	virtual FName GetInputName(int32 InputIndex) const override;
 	virtual void GetConnectorToolTip(int32 InputIndex, int32 OutputIndex, TArray<FString>& OutToolTip) override;
 	virtual void PostEditChangeProperty(FPropertyChangedEvent& PropertyChangedEvent) override;
@@ -317,6 +318,7 @@ class UMaterialExpressionStrataSlabBSDF : public UMaterialExpressionStrataBSDF
 	virtual uint32 GetInputType(int32 InputIndex) override;
 	virtual bool IsResultStrataMaterial(int32 OutputIndex) override;
 	virtual void GatherStrataMaterialInfo(FStrataMaterialInfo& StrataMaterialInfo, int32 OutputIndex) override;
+	virtual FStrataOperator* StrataGenerateMaterialTopologyTree(class FMaterialCompiler* Compiler, class UMaterialExpression* Parent, int32 OutputIndex) override;
 	virtual FName GetInputName(int32 InputIndex) const override;
 	virtual void GetConnectorToolTip(int32 InputIndex, int32 OutputIndex, TArray<FString>& OutToolTip) override;
 	virtual void PostEditChangeProperty(FPropertyChangedEvent& PropertyChangedEvent) override;
@@ -371,6 +373,7 @@ class UMaterialExpressionStrataVolumetricFogCloudBSDF : public UMaterialExpressi
 	virtual uint32 GetInputType(int32 InputIndex) override;
 	virtual bool IsResultStrataMaterial(int32 OutputIndex) override;
 	virtual void GatherStrataMaterialInfo(FStrataMaterialInfo& StrataMaterialInfo, int32 OutputIndex) override;
+	virtual FStrataOperator* StrataGenerateMaterialTopologyTree(class FMaterialCompiler* Compiler, class UMaterialExpression* Parent, int32 OutputIndex) override;
 #endif
 	//~ End UMaterialExpression Interface
 };
@@ -400,6 +403,7 @@ class UMaterialExpressionStrataUnlitBSDF : public UMaterialExpressionStrataBSDF
 	virtual uint32 GetInputType(int32 InputIndex) override;
 	virtual bool IsResultStrataMaterial(int32 OutputIndex) override;
 	virtual void GatherStrataMaterialInfo(FStrataMaterialInfo& StrataMaterialInfo, int32 OutputIndex) override;
+	virtual FStrataOperator* StrataGenerateMaterialTopologyTree(class FMaterialCompiler* Compiler, class UMaterialExpression* Parent, int32 OutputIndex) override;
 #endif
 	//~ End UMaterialExpression Interface
 };
@@ -459,6 +463,7 @@ class UMaterialExpressionStrataHairBSDF : public UMaterialExpressionStrataBSDF
 	virtual uint32 GetInputType(int32 InputIndex) override;
 	virtual bool IsResultStrataMaterial(int32 OutputIndex) override;
 	virtual void GatherStrataMaterialInfo(FStrataMaterialInfo& StrataMaterialInfo, int32 OutputIndex) override;
+	virtual FStrataOperator* StrataGenerateMaterialTopologyTree(class FMaterialCompiler* Compiler, class UMaterialExpression* Parent, int32 OutputIndex) override;
 #endif
 	//~ End UMaterialExpression Interface
 };
@@ -542,6 +547,7 @@ class UMaterialExpressionStrataSingleLayerWaterBSDF : public UMaterialExpression
 	virtual uint32 GetInputType(int32 InputIndex) override;
 	virtual bool IsResultStrataMaterial(int32 OutputIndex) override;
 	virtual void GatherStrataMaterialInfo(FStrataMaterialInfo& StrataMaterialInfo, int32 OutputIndex) override;
+	virtual FStrataOperator* StrataGenerateMaterialTopologyTree(class FMaterialCompiler* Compiler, class UMaterialExpression* Parent, int32 OutputIndex) override;
 #endif
 	//~ End UMaterialExpression Interface
 };
@@ -585,6 +591,7 @@ class UMaterialExpressionStrataHorizontalMixing : public UMaterialExpressionStra
 	virtual uint32 GetInputType(int32 InputIndex) override;
 	virtual bool IsResultStrataMaterial(int32 OutputIndex) override;
 	virtual void GatherStrataMaterialInfo(FStrataMaterialInfo& StrataMaterialInfo, int32 OutputIndex) override;
+	virtual FStrataOperator* StrataGenerateMaterialTopologyTree(class FMaterialCompiler* Compiler, class UMaterialExpression* Parent, int32 OutputIndex) override;
 #endif
 	//~ End UMaterialExpression Interface
 };
@@ -617,6 +624,7 @@ class UMaterialExpressionStrataVerticalLayering : public UMaterialExpressionStra
 	virtual uint32 GetInputType(int32 InputIndex) override;
 	virtual bool IsResultStrataMaterial(int32 OutputIndex) override;
 	virtual void GatherStrataMaterialInfo(FStrataMaterialInfo& StrataMaterialInfo, int32 OutputIndex) override;
+	virtual FStrataOperator* StrataGenerateMaterialTopologyTree(class FMaterialCompiler* Compiler, UMaterialExpression* Parent, int32 OutputIndex) override;
 #endif
 	//~ End UMaterialExpression Interface
 };
@@ -649,6 +657,7 @@ class UMaterialExpressionStrataAdd : public UMaterialExpressionStrataBSDF
 	virtual uint32 GetInputType(int32 InputIndex) override;
 	virtual bool IsResultStrataMaterial(int32 OutputIndex) override;
 	virtual void GatherStrataMaterialInfo(FStrataMaterialInfo& StrataMaterialInfo, int32 OutputIndex) override;
+	virtual FStrataOperator* StrataGenerateMaterialTopologyTree(class FMaterialCompiler* Compiler, class UMaterialExpression* Parent, int32 OutputIndex) override;
 #endif
 	//~ End UMaterialExpression Interface
 };
@@ -678,6 +687,7 @@ class UMaterialExpressionStrataWeight : public UMaterialExpressionStrataBSDF
 	virtual uint32 GetInputType(int32 InputIndex) override;
 	virtual bool IsResultStrataMaterial(int32 OutputIndex) override;
 	virtual void GatherStrataMaterialInfo(FStrataMaterialInfo& StrataMaterialInfo, int32 OutputIndex) override;
+	virtual FStrataOperator* StrataGenerateMaterialTopologyTree(class FMaterialCompiler* Compiler, class UMaterialExpression* Parent, int32 OutputIndex) override;
 #endif
 	//~ End UMaterialExpression Interface
 };
diff --git a/Engine/Source/Runtime/Engine/Private/Materials/HLSLMaterialTranslator.cpp b/Engine/Source/Runtime/Engine/Private/Materials/HLSLMaterialTranslator.cpp
--- a/Engine/Source/Runtime/Engine/Private/Materials/HLSLMaterialTranslator.cpp
+++ b/Engine/Source/Runtime/Engine/Private/Materials/HLSLMaterialTranslator.cpp
@@ -258,6 +258,12 @@ FHLSLMaterialTranslator::FHLSLMaterialTranslator(FMaterial* InMaterial,
 			TargetPlatform = TPM->GetRunningTargetPlatform();
 		}
 	}
+
+	StrataMaterialRootOperator = nullptr;
+	StrataMaterialExpressionRegisteredOperators.Reset(STRATA_MAX_OPERATOR_COUNT);
+	StrataMaterialExpressionToOperatorIndex.Reserve(STRATA_MAX_OPERATOR_COUNT);
+	StrataMaterialBSDFCount = 0;
+	bStrataUsesConversionFromLegacy = false;
 }
 
 FHLSLMaterialTranslator::~FHLSLMaterialTranslator()
@@ -654,6 +660,20 @@ bool FHLSLMaterialTranslator::Translate()
 
 		bCompileForComputeShader = Material->IsLightFunction();
 
+		const bool bStrataEnabled = Engine_IsStrataEnabled();
+		FStrataMaterialInput* FrontMaterialInput = Material->GetMaterialInterface() ? &Material->GetMaterialInterface()->GetMaterial()->FrontMaterial : nullptr;
+		UMaterialExpression* FrontMaterialExpr = FrontMaterialInput ? FrontMaterialInput->GetTracedInput().Expression : nullptr;
+		if (bStrataEnabled && FrontMaterialExpr)
+		{
+			if (Material->GetAssetName() == FString("ColorCoat_Strata"))
+			{
+				int i = 0;
+				i = 2;
+			}
+			FrontMaterialExpr->StrataGenerateMaterialTopologyTree(this, nullptr, 0);
+			StrataGenerateDerivedMaterialOperatorData();
+		}
+
 		// Generate code:
 		// Normally one would expect the generator to emit something like
 		//		float Local0 = ...
@@ -1047,9 +1067,8 @@ bool FHLSLMaterialTranslator::Translate()
 				Errorf(TEXT("SingleLayerWater materials cannot be combined with other shading models.")); // Simply untested for now
 			}
 
-			const bool bStrata = Engine_IsStrataEnabled();
 			if (Material->GetMaterialInterface() && !Material->GetMaterialInterface()->GetMaterial()->HasAnyExpressionsInMaterialAndFunctionsOfType<UMaterialExpressionSingleLayerWaterMaterialOutput>()
-				&& !bStrata
+				&& !bStrataEnabled
 				)
 			{
 				Errorf(TEXT("SingleLayerWater materials requires the use of SingleLayerWaterMaterial output node."));
@@ -1337,11 +1356,123 @@ bool FHLSLMaterialTranslator::Translate()
 					Error(*ErrorMsg);
 				}
 			}
+
+			if (!bStrataUsesConversionFromLegacy)
+			{
+				// Now implement the functions needed to process the material topology
+
+				// Update the coverage/transmittance of each leaves (==BSDFs) of the strata tree.
+				{
+					ResourcesString += "void UpdateAllBSDFsOperatorCoverageTransmittance(inout FStrataPixelHeader StrataPixelHeader, inout FStrataTree StrataTree, bool bRoughDiffuseEnabled, FStrataAddressing NullStrataAddressing, float3 V, float3 L)\n";
+					ResourcesString += "{\n";
+					for (uint32 BSDFIndex = 0; BSDFIndex < StrataMaterialBSDFCount; ++BSDFIndex)
+					{
+						ResourcesString += FString::Printf(TEXT("\t UpdateSingleBSDFOperatorCoverageTransmittance(StrataPixelHeader, StrataTree, StrataTree.BSDFs[%d], bRoughDiffuseEnabled, NullStrataAddressing, V, L);\n"), BSDFIndex);
+					}
+					ResourcesString += "}\n";
+				}
+
+				// Propagate up the coverage/transmittance of each node in the strata tree.
+				{
+					check(StrataMaterialRootOperator);
+					int32 RootMaximumDistanceToLeaves = StrataMaterialRootOperator->MaxDistanceFromLeaves;
+
+					ResourcesString += "void UpdateAllOperatorsCoverageTransmittance(inout FStrataTree StrataTree)\n";
+					ResourcesString += "{\n";
+					for (int32 DistanceToLeaves = 1; DistanceToLeaves <= RootMaximumDistanceToLeaves; ++DistanceToLeaves)
+					{
+						ResourcesString += FString::Printf(TEXT("\t// MaxDistanceFromLeaves = %d \n"), DistanceToLeaves);
+						for (auto& It : StrataMaterialExpressionRegisteredOperators)
+						{
+							if (!It.IsDiscarded() && It.MaxDistanceFromLeaves == DistanceToLeaves)
+							{
+								ResourcesString += FString::Printf(TEXT("\t UpdateSingleOperatorCoverageTransmittance(StrataTree, %d /*operator index*/);\n"), It.Index);
+							}
+						}
+					}
+					ResourcesString += "}\n";
+				}
+
+				// Update the luminance weight of each BSDF according to the operators it has to traverse up to the strata tree root node.
+				{
+					// Update the coverage/transmittance of each node in the graph
+					check(StrataMaterialRootOperator);
+					int32 RootMaximumDistanceToLeaves = StrataMaterialRootOperator->MaxDistanceFromLeaves;
+
+					ResourcesString += "void UpdateAllBSDFWeightAfterOperatorVisit(inout FStrataTree StrataTree)\n";
+					ResourcesString += "{\n";
+					for (uint32 BSDFIndex = 0; BSDFIndex < StrataMaterialBSDFCount; ++BSDFIndex)
+					{
+						for (auto& It : StrataMaterialExpressionRegisteredOperators)
+						{
+							if (!It.IsDiscarded() && It.BSDFIndex == BSDFIndex)
+							{
+								// Walk up the graph to the root node and apply weight factors
+								std::function<void(FStrataOperator&, int32)> WalkOperatorsUp = [&](FStrataOperator& CurrentOperator, int32 PreviousOperatorIndex) -> void
+								{
+									switch (CurrentOperator.OperatorType)
+									{
+									case STRATA_OPERATOR_VERTICAL:
+									{
+										ResourcesString += FString::Printf(TEXT("\t UpdateBSDFWeightAfterOperatorVisit(StrataTree, StrataTree.BSDFs[%d], %d /*Op index*/, %d /*PreviousIsInputA*/); // VERTICAL\n"), BSDFIndex, CurrentOperator.Index, CurrentOperator.AIndex == PreviousOperatorIndex ? 1 : 0);
+										break;
+									}
+									case STRATA_OPERATOR_HORIZONTAL:
+									{
+										ResourcesString += FString::Printf(TEXT("\t UpdateBSDFWeightAfterOperatorVisit(StrataTree, StrataTree.BSDFs[%d], %d /*Op index*/, %d /*PreviousIsInputA*/); // HORIZONTAL\n"), BSDFIndex, CurrentOperator.Index, CurrentOperator.AIndex == PreviousOperatorIndex ? 1 : 0);
+										break;
+									}
+
+									case STRATA_OPERATOR_WEIGHT:
+									{
+										ResourcesString += FString::Printf(TEXT("\t UpdateBSDFWeightAfterOperatorVisit(StrataTree, StrataTree.BSDFs[%d], %d /*Op index*/, %d /*PreviousIsInputA*/); // WEIGHT\n"), BSDFIndex, CurrentOperator.Index, 1);
+										break;
+									}
+
+									case STRATA_OPERATOR_ADD:
+									{
+										break; // NOP
+									}
+
+									default:
+									case STRATA_OPERATOR_BSDF:
+									{
+										check(false);
+									}
+									}
+
+									if (CurrentOperator.TopIndex != INDEX_NONE)
+									{
+										WalkOperatorsUp(StrataMaterialExpressionRegisteredOperators[CurrentOperator.TopIndex], CurrentOperator.Index);
+									}
+								};
+
+								const int32 BSDFOperatorIndex = It.Index;
+								FStrataOperator& BSDFOperator = StrataMaterialExpressionRegisteredOperators[BSDFOperatorIndex];
+
+								// Start visiting node up from the BSDF leaf only if it has a parent.
+								if (BSDFOperator.TopIndex != INDEX_NONE)
+								{
+									WalkOperatorsUp(StrataMaterialExpressionRegisteredOperators[BSDFOperator.TopIndex], BSDFOperator.Index);
+								}
+							}
+						}
+					}
+					ResourcesString += "}\n";
+				}
+			}
 		}
 		else
 		{
 			MaterialCompilationOutput.StrataMaterialDescription = "";
-			ResourcesString += "// No Strata material provided \r\n";
+			ResourcesString += "// No Strata material provided, or using legacy material conversion \r\n";
+			
+			// Adde default strata functions
+			ResourcesString += "#if TEMPLATE_USES_STRATA\n;";
+			ResourcesString += "void UpdateAllBSDFsOperatorCoverageTransmittance(inout FStrataPixelHeader StrataPixelHeader, inout FStrataTree StrataTree, bool bRoughDiffuseEnabled, FStrataAddressing NullStrataAddressing, float3 V, float3 L) {}\n";
+			ResourcesString += "void UpdateAllOperatorsCoverageTransmittance(inout FStrataTree StrataTree) {}\n";
+			ResourcesString += "void UpdateAllBSDFWeightAfterOperatorVisit(inout FStrataTree StrataTree) {}\n";
+			ResourcesString += "#endif\n;";
 		}
 
 		LoadShaderSourceFileChecked(TEXT("/Engine/Private/MaterialTemplate.ush"), GetShaderPlatform(), MaterialTemplate);
@@ -1716,6 +1847,11 @@ void FHLSLMaterialTranslator::GetMaterialEnvironment(EShaderPlatform InPlatform,
 		FinalUsedSharedLocalBasesCount = 0;
 		FStrataRegisteredSharedLocalBasis UsedSharedLocalBasesInfo[STRATA_MAX_SHAREDLOCALBASES_REGISTERS];
 
+		check(StrataMaterialRootOperator);
+		uint32 RootMaximumDistanceToLeaves = StrataMaterialRootOperator->MaxDistanceFromLeaves;
+		OutEnvironment.SetDefine(TEXT("MATERIAL_TOPOLOGY_TREE_MAX_DEPTH"), RootMaximumDistanceToLeaves);
+
+		OutEnvironment.SetDefine(TEXT("STRATA_USES_CONVERSION_FROM_LEGACY"), bStrataUsesConversionFromLegacy ? 1 : 0);
 
 		/*
 		* The final output code/workflow for shared tangent basis should look like
@@ -1834,6 +1970,8 @@ void FHLSLMaterialTranslator::GetMaterialEnvironment(EShaderPlatform InPlatform,
 				OutEnvironment.SetDefine(TEXT("STRATA_CLAMPED_MAX_BSDF_COUNT_PER_LAYER"), StrataMaterialAnalysis.RequestedMaxBSDFCountPerLayer);
 			}
 
+			OutEnvironment.SetDefine(TEXT("STRATA_CLAMPED_BSDF_COUNT"), StrataMaterialBSDFCount);
+
 			FString StrataMaterialDescription;
 
 			StrataMaterialDescription += FString::Printf(TEXT("----- STRATA -----\r\n"));
@@ -1852,6 +1990,10 @@ void FHLSLMaterialTranslator::GetMaterialEnvironment(EShaderPlatform InPlatform,
 				}
 			}
 
+			StrataMaterialDescription += FString::Printf(TEXT("------------------\r\n"));
+			StrataMaterialDescription += FString::Printf(TEXT("Graph maximum distance to leaves %u\r\n"), RootMaximumDistanceToLeaves);
+			StrataMaterialDescription += FString::Printf(TEXT("------------------\r\n"));
+
 			StrataMaterialDescription += FString::Printf(TEXT("Byte Per Pixel Budget      %u\r\n"), StrataBytePerPixel);
 			StrataMaterialDescription += FString::Printf(TEXT("Result.bFitInMemoryBudget  %s\r\n"), StrataMaterialAnalysis.bFitInMemoryBudget ? TEXT("YES") : TEXT("NO"));
 			StrataMaterialDescription += FString::Printf(TEXT("Result.RequestedLayerCount %u\r\n"), StrataMaterialAnalysis.RequestedLayerCount);
@@ -1859,7 +2001,7 @@ void FHLSLMaterialTranslator::GetMaterialEnvironment(EShaderPlatform InPlatform,
 			StrataMaterialDescription += FString::Printf(TEXT("Result.RequestedByteCount  %u\r\n"), StrataMaterialAnalysis.RequestedByteCount);
 			if (!StrataMaterialAnalysis.bFitInMemoryBudget)
 			{
-				StrataMaterialDescription += FString::Printf(TEXT("Result.ClampedLayerCount   %u\r\n"), StrataMaterialAnalysis.ClampedLayerCount);
+				StrataMaterialDescription += FString::Printf(TEXT("Result.ClampedLayerCount   %u\r\n"), StrataMaterialAnalysis.ClampedLayerCount);	// STRATA_TODO this is no longer needed
 				StrataMaterialDescription += FString::Printf(TEXT("Result.ClampedBSDFCount    %u\r\n"), StrataMaterialAnalysis.ClampedBSDFCount);
 				StrataMaterialDescription += FString::Printf(TEXT("Result.UsedByteCount       %u\r\n"), StrataMaterialAnalysis.UsedByteCount);
 			}
@@ -9400,6 +9542,281 @@ int32 FHLSLMaterialTranslator::GetLocal(const FName& LocalName)
 	return AddInlinedCodeChunk(MCT_Float1, TEXT("%s"), *Entry->Name);
 }
 
+FStrataOperator& FHLSLMaterialTranslator::StrataCompilationRegisterOperator(int32 OperatorType, UMaterialExpression* Expression, UMaterialExpression* Parent, bool bUseParameterBlending)
+{
+	if (OperatorType == STRATA_OPERATOR_BSDF_LEGACY)
+	{
+		// We register the fact that a legacy material conversion is used and register a simple BSDF
+		bStrataUsesConversionFromLegacy = true;
+		OperatorType = STRATA_OPERATOR_BSDF;
+	}
+
+	if (StrataMaterialExpressionToOperatorIndex.Find(Expression))
+	{
+		UE_LOG(LogMaterial, Error, TEXT("Material %s tries to register a Strata operator twice (asset: %s).\r\n"), *Material->GetDebugName(), *Material->GetAssetPath().ToString());
+		return StrataMaterialExpressionRegisteredOperators[0];
+	}
+
+	const uint32 NewOperatorIndex = StrataMaterialExpressionToOperatorIndex.Num();
+	if (NewOperatorIndex >= STRATA_MAX_OPERATOR_COUNT)
+	{
+		UE_LOG(LogMaterial, Error, TEXT("Material %s have too many Strata Operators (asset: %s).\r\n"), *Material->GetDebugName(), *Material->GetAssetPath().ToString());
+		return StrataMaterialExpressionRegisteredOperators[0];
+	}
+
+	int32* ParentOperatorIndex = StrataMaterialExpressionToOperatorIndex.Find(Parent);
+	if (Parent!=nullptr && ParentOperatorIndex == nullptr)
+	{
+		UE_LOG(LogMaterial, Error, TEXT("Material %s tries to register unknown operator parents (asset: %s).\r\n"), *Material->GetDebugName(), *Material->GetAssetPath().ToString());
+		return StrataMaterialExpressionRegisteredOperators[0];
+	}
+
+	StrataMaterialExpressionToOperatorIndex.Add(Expression, NewOperatorIndex);
+
+	FStrataOperator& NewOperator = StrataMaterialExpressionRegisteredOperators.AddDefaulted_GetRef();
+
+	NewOperator.OperatorType = OperatorType;
+	NewOperator.bNodeRequestParameterBlending = bUseParameterBlending;
+	NewOperator.Index = NewOperatorIndex;
+	NewOperator.TopIndex = ParentOperatorIndex != nullptr ? *ParentOperatorIndex : INDEX_NONE;
+	NewOperator.AIndex   = INDEX_NONE;
+	NewOperator.BIndex   = INDEX_NONE;
+
+	NewOperator.BSDFIndex = INDEX_NONE;	// Allocated later to be ableto account for inline
+
+	NewOperator.MaxDistanceFromLeaves = 0;
+	NewOperator.bIsBottom = false;
+	NewOperator.bIsTop = false;
+
+	return NewOperator;
+}
+
+FStrataOperator& FHLSLMaterialTranslator::StrataCompilationGetOperator(UMaterialExpression* Expression)
+{
+	auto* OperatorIndex = StrataMaterialExpressionToOperatorIndex.Find(Expression);
+	check(OperatorIndex && *OperatorIndex >= 0 && *OperatorIndex < STRATA_MAX_OPERATOR_COUNT);
+	return StrataMaterialExpressionRegisteredOperators[*OperatorIndex];
+}
+
+void FHLSLMaterialTranslator::StrataGenerateDerivedMaterialOperatorData()
+{
+	//
+	// Evaluate the one and only root node
+	//
+	int32 RootIndex = INDEX_NONE;
+	for (auto& It : StrataMaterialExpressionRegisteredOperators)
+	{
+		if (It.TopIndex == INDEX_NONE)
+		{
+			check(RootIndex == INDEX_NONE);	// There can only be one
+			RootIndex = It.Index;
+		}
+	}
+	StrataMaterialRootOperator = &StrataMaterialExpressionRegisteredOperators[RootIndex];
+	check(StrataMaterialRootOperator);
+
+	//
+	// Parse the tree and mark nodes that are the root of a subtree using parameter blending, while other nodes in that tree are forced to use parameter blending.
+	// Allocate BSDFIndex at the same time.
+	//
+	{
+		int VOpTopBranchCountTaken = 0;
+		int VOpBottomBranchCountTaken = 0;
+
+		std::function<void(FStrataOperator&, bool)> WalkOperators = [&](FStrataOperator& CurrentOperator, bool bInsideParameterBlendingSubTree) -> void
+		{
+			const bool bCurrentOpRequestParameterBlending	= CurrentOperator.bNodeRequestParameterBlending;
+			const bool bRootOfParameterBlendingSubTree		= bCurrentOpRequestParameterBlending && !bInsideParameterBlendingSubTree;
+			const bool bUseParameterBlending				= bCurrentOpRequestParameterBlending || bInsideParameterBlendingSubTree;
+
+			CurrentOperator.bUseParameterBlending			= bUseParameterBlending;
+			CurrentOperator.bRootOfParameterBlendingSubTree	= bRootOfParameterBlendingSubTree;
+
+			switch (CurrentOperator.OperatorType)
+			{
+			case STRATA_OPERATOR_VERTICAL:
+			{
+				WalkOperators(StrataMaterialExpressionRegisteredOperators[CurrentOperator.AIndex], bUseParameterBlending);
+				WalkOperators(StrataMaterialExpressionRegisteredOperators[CurrentOperator.BIndex], bUseParameterBlending);
+				break;
+			}
+			case STRATA_OPERATOR_HORIZONTAL:
+			case STRATA_OPERATOR_ADD:
+			{
+				WalkOperators(StrataMaterialExpressionRegisteredOperators[CurrentOperator.AIndex], bUseParameterBlending);
+				WalkOperators(StrataMaterialExpressionRegisteredOperators[CurrentOperator.BIndex], bUseParameterBlending);
+				break;
+			}
+			case STRATA_OPERATOR_WEIGHT:
+			{
+				WalkOperators(StrataMaterialExpressionRegisteredOperators[CurrentOperator.AIndex], bUseParameterBlending);
+				break;
+			}
+			case STRATA_OPERATOR_BSDF:
+			{
+				if (!bInsideParameterBlendingSubTree)
+				{
+					CurrentOperator.BSDFIndex = StrataMaterialBSDFCount++;
+				}
+				break;
+			}
+			}
+
+			// We mark the top of a parameter blending tree as a BSDF now to allocate a slot for it that can then be used next for non-parameter blending operations.
+			// Intermediate parameter blending BSDF and operation will be done inline and stored in FStrataData.
+			if (bRootOfParameterBlendingSubTree)
+			{
+				CurrentOperator.OperatorType = STRATA_OPERATOR_BSDF;
+				CurrentOperator.BSDFIndex = StrataMaterialBSDFCount++;
+				CurrentOperator.AIndex = INDEX_NONE;
+				CurrentOperator.BIndex = INDEX_NONE;
+			}
+		};
+
+		WalkOperators(*StrataMaterialRootOperator, false);
+	}
+
+	// STRATA_TODO: operation using parameter blending are not actually discarded so they still occupy a pot in the opeartion arrays in the compiler and in the shader also.
+	// Even thought the compiler will remove unused operators from the shader code, we will still be limited to STRATA_MAX_OPERATOR_COUNT eve nwith parameter blending.
+	// This can be fixed by remapping all operation index and make sure this is the index that is allways used in the compiler and specified to the shader.
+
+	//
+	// Make sure all the types have valid children operator indices
+	//
+	for (auto& It : StrataMaterialExpressionRegisteredOperators)
+	{
+		if (It.IsDiscarded())
+		{
+			continue; // ignore discarded operations in sub tree using parameter blending
+		}
+
+		check(It.Index != INDEX_NONE);
+
+		switch (It.OperatorType)
+		{
+		// Operators without any child
+		case STRATA_OPERATOR_BSDF:
+		{
+			check(It.AIndex == INDEX_NONE && It.BIndex == INDEX_NONE && It.BSDFIndex != INDEX_NONE);
+			break;
+		}
+
+		// Operators with two children
+		case STRATA_OPERATOR_HORIZONTAL:
+		case STRATA_OPERATOR_VERTICAL:
+		case STRATA_OPERATOR_ADD:
+		{
+			check(It.BIndex != INDEX_NONE);
+		}
+		// Fallthrough
+
+		// Operators with a single child
+		case STRATA_OPERATOR_WEIGHT:
+		{
+			check(It.AIndex != INDEX_NONE);
+		}
+		// Fallthrough
+		}
+	}
+
+	//
+	// Compute the maximum depth from the BSDF node for each operator
+	//
+	{
+		std::function<void(FStrataOperator&)>  WalkOperatorsToRoot = [&](FStrataOperator& CurrentOperator) -> void
+		{
+			switch (CurrentOperator.OperatorType)
+			{
+			case STRATA_OPERATOR_WEIGHT:
+			{
+				CurrentOperator.MaxDistanceFromLeaves = StrataMaterialExpressionRegisteredOperators[CurrentOperator.AIndex].MaxDistanceFromLeaves + 1;
+				break;
+			}
+			case STRATA_OPERATOR_VERTICAL:
+			case STRATA_OPERATOR_HORIZONTAL:
+			case STRATA_OPERATOR_ADD:
+			{
+				CurrentOperator.MaxDistanceFromLeaves = FMath::Max(
+					StrataMaterialExpressionRegisteredOperators[CurrentOperator.AIndex].MaxDistanceFromLeaves,
+					StrataMaterialExpressionRegisteredOperators[CurrentOperator.BIndex].MaxDistanceFromLeaves) + 1;
+				break;
+			}
+			case STRATA_OPERATOR_BSDF:
+			{
+				CurrentOperator.MaxDistanceFromLeaves = 0;
+				break;
+			}
+			}
+
+			if (CurrentOperator.TopIndex != INDEX_NONE)
+			{
+				WalkOperatorsToRoot(StrataMaterialExpressionRegisteredOperators[CurrentOperator.TopIndex]);
+			}
+		};
+
+		for (auto& It : StrataMaterialExpressionRegisteredOperators)
+		{
+			if (It.IsDiscarded())
+			{
+				continue; // ignore discarded operations in sub tree using parameter blending
+			}
+
+			if (It.OperatorType == STRATA_OPERATOR_BSDF)
+			{
+				// Recursively parse all node from BSDF to the root node and update the necessary properties.
+				WalkOperatorsToRoot(It);
+			}
+		}
+	}
+
+	//
+	// Compute IsTop or IsBottom layer using a depth first tree visit while counting vertical right and left branches taken.
+	//
+	{
+		int VOpTopBranchCountTaken = 0;
+		int VOpBottomBranchCountTaken = 0;
+
+		std::function<void(FStrataOperator&)> WalkOperators = [&](FStrataOperator& CurrentOperator) -> void
+		{
+			switch (CurrentOperator.OperatorType)
+			{
+			case STRATA_OPERATOR_VERTICAL:
+			{
+				VOpTopBranchCountTaken++;
+				WalkOperators(StrataMaterialExpressionRegisteredOperators[CurrentOperator.AIndex]);
+				VOpTopBranchCountTaken--;
+				VOpBottomBranchCountTaken++;
+				WalkOperators(StrataMaterialExpressionRegisteredOperators[CurrentOperator.BIndex]);
+				VOpBottomBranchCountTaken--;
+				break;
+			}
+			case STRATA_OPERATOR_HORIZONTAL:
+			case STRATA_OPERATOR_ADD:
+			{
+				WalkOperators(StrataMaterialExpressionRegisteredOperators[CurrentOperator.AIndex]);
+				WalkOperators(StrataMaterialExpressionRegisteredOperators[CurrentOperator.BIndex]);
+				break;
+			}
+			case STRATA_OPERATOR_WEIGHT:
+			{
+				WalkOperators(StrataMaterialExpressionRegisteredOperators[CurrentOperator.AIndex]);
+				break;
+			}
+			case STRATA_OPERATOR_BSDF:
+			{
+				const int32 VopCount = VOpTopBranchCountTaken + VOpBottomBranchCountTaken;
+				CurrentOperator.LayerDepth = VOpBottomBranchCountTaken;
+				CurrentOperator.bIsTop = VopCount == 0 || VOpTopBranchCountTaken == VopCount;
+				CurrentOperator.bIsBottom = VopCount == 0 || VOpBottomBranchCountTaken == VopCount;
+				break;
+			}
+			}
+		};
+
+		WalkOperators(*StrataMaterialRootOperator);
+	}
+}
+
 void FHLSLMaterialTranslator::StrataCompilationInfoRegisterCodeChunk(int32 CodeChunk, FStrataMaterialCompilationInfo& StrataMaterialCompilationInfo)
 {
 	check(CodeChunk != INDEX_NONE);
@@ -9543,10 +9960,40 @@ int32 FHLSLMaterialTranslator::StrataSlabBSDF(
 	int32 ThinFilmThickness, 
 	int32 FuzzAmount, int32 FuzzColor,
 	int32 Thickness, 
-	int32 Normal, int32 Tangent, const FString& SharedLocalBasisIndexMacro)
+	int32 Normal, int32 Tangent, const FString& SharedLocalBasisIndexMacro,
+	FStrataOperator* PromoteToOperator)
 {
 	const FString NormalCode = GetParameterCode(Normal);
 	const FString TangentCode = Tangent != INDEX_NONE ? *GetParameterCode(Tangent) : TEXT("NONE");
+
+
+	if (PromoteToOperator)
+	{
+		check(PromoteToOperator->Index != INDEX_NONE);
+		check(PromoteToOperator->BSDFIndex != INDEX_NONE);
+		return AddCodeChunk(
+			MCT_Strata, TEXT("PromoteParameterBlendedBSDFToOperator(GetStrataSlabBSDF(%s, %s, %s, %s, %s, %s, %s, %s, %s, %s, %s, %s, %s, %s, %s, %s, %s, %s, %s, %s, Parameters.SharedLocalBases.Types) /* Normal = %s ; Tangent = %s */, Parameters.StrataTree, %u, %u, %u, %u)"),
+			*GetParameterCode(UseMetalness),
+			*GetParameterCode(BaseColor), *GetParameterCode(EdgeColor), *GetParameterCode(Specular), *GetParameterCode(Metallic),
+			*GetParameterCode(DiffuseAlbedo), *GetParameterCode(F0), *GetParameterCode(F90),
+			*GetParameterCode(Roughness), *GetParameterCode(Anisotropy),
+			*GetParameterCode(SSSProfileId),
+			*GetParameterCode(SSSDMFP),
+			*GetParameterCode(SSSDMFPScale),
+			*GetParameterCode(EmissiveColor),
+			*GetParameterCode(Haziness),
+			*GetParameterCode(ThinFilmThickness),
+			*GetParameterCode(FuzzAmount), *GetParameterCode(FuzzColor),
+			*GetParameterCode(Thickness),
+			*SharedLocalBasisIndexMacro,
+			*NormalCode,
+			*TangentCode,
+			PromoteToOperator->Index,
+			PromoteToOperator->BSDFIndex,
+			PromoteToOperator->bIsBottom ? 1 : 0,
+			PromoteToOperator->bIsTop ? 1 : 0);
+	}
+
 	return AddCodeChunk(
 		MCT_Strata, TEXT("GetStrataSlabBSDF(%s, %s, %s, %s, %s, %s, %s, %s, %s, %s, %s, %s, %s, %s, %s, %s, %s, %s, %s, %s, Parameters.SharedLocalBases.Types) /* Normal = %s ; Tangent = %s */"),
 		* GetParameterCode(UseMetalness),
@@ -9588,7 +10035,7 @@ int32 FHLSLMaterialTranslator::StrataConversionFromLegacy(
 	const FString ClearCoat_TangentCode = Tangent != INDEX_NONE ? *GetParameterCode(ClearCoat_Tangent) : TEXT("NONE");
 
 	return AddCodeChunk(
-		MCT_Strata, TEXT("StrataConvertLegacyMaterial%s(%s, %s, %s, %s, %s, %s, %s, %s, %s, %s, %s, %s, %s, %s, %s, %s, %s, %s, %s, Parameters.SharedLocalBases.Types) /* Normal = %s ; Tangent = %s ; ClearCoat_Normal = %s ; ClearCoat_Tangent = %s */"),
+		MCT_Strata, TEXT("StrataConvertLegacyMaterial%s(%s, %s, %s, %s, %s, %s, %s, %s, %s, %s, %s, %s, %s, %s, %s, %s, %s, %s, %s, Parameters.SharedLocalBases.Types, Parameters.StrataTree) /* Normal = %s ; Tangent = %s ; ClearCoat_Normal = %s ; ClearCoat_Tangent = %s */"),
 		bHasDynamicShadingModels ? TEXT("Dynamic") : TEXT("Static"),
 		*GetParameterCode(BaseColor), *GetParameterCode(Specular), *GetParameterCode(Metallic),
 		*GetParameterCode(Roughness), *GetParameterCode(Anisotropy),
@@ -9666,26 +10113,48 @@ int32 FHLSLMaterialTranslator::StrataSingleLayerWaterBSDF(
 	);
 }
 
-int32 FHLSLMaterialTranslator::StrataHorizontalMixing(int32 Background, int32 Foreground, int32 Mix)
+int32 FHLSLMaterialTranslator::StrataHorizontalMixing(int32 Background, int32 Foreground, int32 Mix, int OperatorIndex, uint32 MaxDistanceFromLeaves)
 {
 	if (Foreground == INDEX_NONE || Background == INDEX_NONE || Mix == INDEX_NONE)
 	{
 		return INDEX_NONE;
 	}
 	return AddCodeChunk(
-		MCT_Strata, TEXT("StrataHorizontalMixing(%s, %s, %s)"),
+		MCT_Strata, TEXT("StrataHorizontalMixing(%s, %s, %s, Parameters.StrataTree, %u, %u)"),
 		*GetParameterCode(Background),
 		*GetParameterCode(Foreground),
-		*GetParameterCode(Mix)
+		*GetParameterCode(Mix),
+		OperatorIndex,
+		MaxDistanceFromLeaves
 	);
 }
 
-int32 FHLSLMaterialTranslator::StrataHorizontalMixingParameterBlending(int32 Background, int32 Foreground, int32 HorizontalMixCodeChunk, int32 NormalMixCodeChunk, const FString& SharedLocalBasisIndexMacro)
+int32 FHLSLMaterialTranslator::StrataHorizontalMixingParameterBlending(
+	int32 Background, int32 Foreground, int32 HorizontalMixCodeChunk, int32 NormalMixCodeChunk, const FString& SharedLocalBasisIndexMacro, FStrataOperator* PromoteToOperator)
 {
 	if (Foreground == INDEX_NONE || Background == INDEX_NONE || HorizontalMixCodeChunk == INDEX_NONE || NormalMixCodeChunk == INDEX_NONE)
 	{
 		return INDEX_NONE;
 	}
+
+	if (PromoteToOperator)
+	{
+		check(PromoteToOperator->Index != INDEX_NONE);
+		check(PromoteToOperator->BSDFIndex != INDEX_NONE);
+		return AddCodeChunk(
+			MCT_Strata, TEXT("PromoteParameterBlendedBSDFToOperator(StrataHorizontalMixingParameterBlending(%s, %s, %s, %s, %s), Parameters.StrataTree, %u, %u, %u, %u)"),
+			*GetParameterCode(Background),
+			*GetParameterCode(Foreground),
+			*GetParameterCode(HorizontalMixCodeChunk),
+			*GetParameterCode(NormalMixCodeChunk),
+			*SharedLocalBasisIndexMacro,
+			PromoteToOperator->Index,
+			PromoteToOperator->BSDFIndex,
+			PromoteToOperator->bIsBottom,
+			PromoteToOperator->bIsTop
+		);
+	}
+
 	return AddCodeChunk(
 		MCT_Strata, TEXT("StrataHorizontalMixingParameterBlending(%s, %s, %s, %s, %s)"),
 		*GetParameterCode(Background),
@@ -9696,25 +10165,46 @@ int32 FHLSLMaterialTranslator::StrataHorizontalMixingParameterBlending(int32 Bac
 	);
 }
 
-int32 FHLSLMaterialTranslator::StrataVerticalLayering(int32 Top, int32 Base)
+int32 FHLSLMaterialTranslator::StrataVerticalLayering(int32 Top, int32 Base, int OperatorIndex, uint32 MaxDistanceFromLeaves)
 {
 	if (Top == INDEX_NONE || Base == INDEX_NONE)
 	{
 		return INDEX_NONE;
 	}
 	return AddCodeChunk(
-		MCT_Strata, TEXT("StrataVerticalLayering(%s, %s)"),
+		MCT_Strata, TEXT("StrataVerticalLayering(%s, %s, Parameters.StrataTree, %u, %u)"),
 		*GetParameterCode(Top),
-		*GetParameterCode(Base)
+		*GetParameterCode(Base),
+		OperatorIndex,
+		MaxDistanceFromLeaves
 	);
 }
 
-int32 FHLSLMaterialTranslator::StrataVerticalLayeringParameterBlending(int32 Top, int32 Base, const FString& SharedLocalBasisIndexMacro, int32 TopBSDFNormalCodeChunk)
+int32 FHLSLMaterialTranslator::StrataVerticalLayeringParameterBlending(int32 Top, int32 Base, const FString& SharedLocalBasisIndexMacro, int32 TopBSDFNormalCodeChunk, FStrataOperator* PromoteToOperator)
 {
 	if (Top == INDEX_NONE || Base == INDEX_NONE)
 	{
 		return INDEX_NONE;
 	}
+
+	if (PromoteToOperator)
+	{
+		check(PromoteToOperator->Index != INDEX_NONE);
+		check(PromoteToOperator->BSDFIndex != INDEX_NONE);
+		return AddCodeChunk(
+			MCT_Strata, TEXT("PromoteParameterBlendedBSDFToOperator(StrataVerticalLayeringParameterBlending(%s, %s, %s, dot(%s, %s)), Parameters.StrataTree, %u, %u, %u, %u)"),
+			*GetParameterCode(Top),
+			*GetParameterCode(Base),
+			*SharedLocalBasisIndexMacro,
+			*GetParameterCode(TopBSDFNormalCodeChunk),
+			*GetParameterCode(CameraVector()),
+			PromoteToOperator->Index,
+			PromoteToOperator->BSDFIndex,
+			PromoteToOperator->bIsBottom,
+			PromoteToOperator->bIsTop
+		);
+	}
+
 	return AddCodeChunk(
 		MCT_Strata, TEXT("StrataVerticalLayeringParameterBlending(%s, %s, %s, dot(%s, %s))"),
 		*GetParameterCode(Top),
@@ -9725,25 +10215,45 @@ int32 FHLSLMaterialTranslator::StrataVerticalLayeringParameterBlending(int32 Top
 	);
 }
 
-int32 FHLSLMaterialTranslator::StrataAdd(int32 A, int32 B)
+int32 FHLSLMaterialTranslator::StrataAdd(int32 A, int32 B, int OperatorIndex, uint32 MaxDistanceFromLeaves)
 {
 	if (A == INDEX_NONE || B == INDEX_NONE)
 	{
 		return INDEX_NONE;
 	}
 	return AddCodeChunk(
-		MCT_Strata, TEXT("StrataAdd(%s, %s)"),
+		MCT_Strata, TEXT("StrataAdd(%s, %s, Parameters.StrataTree, %u, %u)"),
 		*GetParameterCode(A),
-		*GetParameterCode(B)
+		*GetParameterCode(B),
+		OperatorIndex,
+		MaxDistanceFromLeaves
 	);
 }
 
-int32 FHLSLMaterialTranslator::StrataAddParameterBlending(int32 A, int32 B, int32 AMixWeight, const FString& SharedLocalBasisIndexMacro)
+int32 FHLSLMaterialTranslator::StrataAddParameterBlending(int32 A, int32 B, int32 AMixWeight, const FString& SharedLocalBasisIndexMacro, FStrataOperator* PromoteToOperator)
 {
 	if (A == INDEX_NONE || B == INDEX_NONE)
 	{
 		return INDEX_NONE;
 	}
+
+	if (PromoteToOperator)
+	{
+		check(PromoteToOperator->Index != INDEX_NONE);
+		check(PromoteToOperator->BSDFIndex != INDEX_NONE);
+		return AddCodeChunk(
+			MCT_Strata, TEXT("PromoteParameterBlendedBSDFToOperator(StrataAddParameterBlending(%s, %s, %s, %s), Parameters.StrataTree, %u, %u, %u, %u)"),
+			*GetParameterCode(A),
+			*GetParameterCode(B),
+			*GetParameterCode(AMixWeight),
+			*SharedLocalBasisIndexMacro,
+			PromoteToOperator->Index,
+			PromoteToOperator->BSDFIndex,
+			PromoteToOperator->bIsBottom,
+			PromoteToOperator->bIsTop
+		);
+	}
+
 	return AddCodeChunk(
 		MCT_Strata, TEXT("StrataAddParameterBlending(%s, %s, %s, %s)"),
 		*GetParameterCode(A),
@@ -9753,14 +10263,45 @@ int32 FHLSLMaterialTranslator::StrataAddParameterBlending(int32 A, int32 B, int3
 	);
 }
 
-int32 FHLSLMaterialTranslator::StrataWeight(int32 A, int32 Weight)
+int32 FHLSLMaterialTranslator::StrataWeight(int32 A, int32 Weight, int OperatorIndex, uint32 MaxDistanceFromLeaves)
+{
+	if (A == INDEX_NONE || Weight == INDEX_NONE)
+	{
+		return INDEX_NONE;
+	}
+	return AddCodeChunk(
+		MCT_Strata, TEXT("StrataWeight(%s, %s, Parameters.StrataTree, %u, %u)"),
+		*GetParameterCode(A),
+		*GetParameterCode(Weight),
+		OperatorIndex,
+		MaxDistanceFromLeaves
+	);
+}
+
+int32 FHLSLMaterialTranslator::StrataWeightParameterBlending(int32 A, int32 Weight, FStrataOperator* PromoteToOperator)
 {
 	if (A == INDEX_NONE || Weight == INDEX_NONE)
 	{
 		return INDEX_NONE;
 	}
+
+	if (PromoteToOperator)
+	{
+		check(PromoteToOperator->Index != INDEX_NONE);
+		check(PromoteToOperator->BSDFIndex != INDEX_NONE);
+		return AddCodeChunk(
+			MCT_Strata, TEXT("PromoteParameterBlendedBSDFToOperator(StrataWeightParameterBlending(%s, %s), Parameters.StrataTree, %u, %u, %u, %u)"),
+			*GetParameterCode(A),
+			*GetParameterCode(Weight),
+			PromoteToOperator->Index,
+			PromoteToOperator->BSDFIndex,
+			PromoteToOperator->bIsBottom,
+			PromoteToOperator->bIsTop
+		);
+	}
+
 	return AddCodeChunk(
-		MCT_Strata, TEXT("StrataWeight(%s, %s)"),
+		MCT_Strata, TEXT("StrataWeightParameterBlending(%s, %s)"),
 		*GetParameterCode(A),
 		*GetParameterCode(Weight)
 	);
diff --git a/Engine/Source/Runtime/Engine/Private/Materials/HLSLMaterialTranslator.h b/Engine/Source/Runtime/Engine/Private/Materials/HLSLMaterialTranslator.h
--- a/Engine/Source/Runtime/Engine/Private/Materials/HLSLMaterialTranslator.h
+++ b/Engine/Source/Runtime/Engine/Private/Materials/HLSLMaterialTranslator.h
@@ -427,6 +427,12 @@ protected:
 	 * A normal code chunk hash can point to multiple shared info in case it is paired with different tangents. */
 	TMultiMap<uint64, FStrataSharedLocalBasesInfo> CodeChunkToStrataSharedLocalBasis;
 
+	TMap<UMaterialExpression*, int32> StrataMaterialExpressionToOperatorIndex;
+	TArray<FStrataOperator> StrataMaterialExpressionRegisteredOperators;
+	FStrataOperator* StrataMaterialRootOperator;
+	uint32 StrataMaterialBSDFCount;
+	bool bStrataUsesConversionFromLegacy;
+
 	/** Tracks the total number of vt samples in the shader. */
 	uint32 NumVtSamples;
 
@@ -637,6 +643,8 @@ protected:
 
 	int32 GenericSwitch(const TCHAR* Function, int32 IfTrue, int32 IfFalse);
 
+	void StrataGenerateDerivedMaterialOperatorData();
+
 	// FMaterialCompiler interface.
 
 	/** 
@@ -1012,7 +1020,8 @@ protected:
 		int32 ThinFilmThickness, 
 		int32 FuzzAmount, int32 FuzzColor, 
 		int32 Thickness,
-		int32 Normal, int32 Tangent, const FString& SharedLocalBasisIndexMacro) override;
+		int32 Normal, int32 Tangent, const FString& SharedLocalBasisIndexMacro,
+		FStrataOperator* PromoteToOperator) override;
 	virtual int32 StrataConversionFromLegacy(
 		bool bHasDynamicShadingModel,
 		int32 BaseColor, int32 Specular, int32 Metallic,
@@ -1033,15 +1042,19 @@ protected:
 		int32 BaseColor, int32 Metallic, int32 Specular, int32 Roughness, 
 		int32 EmissiveColor, int32 TopMaterialOpacity, int32 WaterAlbedo, int32 WaterExtinction, int32 WaterPhaseG, 
 		int32 ColorScaleBehindWater, int32 Normal, const FString& SharedLocalBasisIndexMacro) override;
-	virtual int32 StrataHorizontalMixing(int32 Background, int32 Foreground, int32 Mix) override;
-	virtual int32 StrataHorizontalMixingParameterBlending(int32 Background, int32 Foreground, int32 HorizontalMixCodeChunk, int32 NormalMixCodeChunk, const FString& SharedLocalBasisIndexMacro) override;
-	virtual int32 StrataVerticalLayering(int32 Top, int32 Base) override;
-	virtual int32 StrataVerticalLayeringParameterBlending(int32 Top, int32 Base, const FString& SharedLocalBasisIndexMacro, int32 TopBSDFNormalCodeChunk) override;
-	virtual int32 StrataAdd(int32 A, int32 B) override;
-	virtual int32 StrataAddParameterBlending(int32 A, int32 B, int32 AMixWeight, const FString& SharedLocalBasisIndexMacro) override;
-	virtual int32 StrataWeight(int32 A, int32 Weight) override;
+	virtual int32 StrataHorizontalMixing(int32 Background, int32 Foreground, int32 Mix, int OperatorIndex, uint32 MaxDistanceFromLeaves) override;
+	virtual int32 StrataHorizontalMixingParameterBlending(int32 Background, int32 Foreground, int32 HorizontalMixCodeChunk, int32 NormalMixCodeChunk, const FString& SharedLocalBasisIndexMacro, FStrataOperator* PromoteToOperator) override;
+	virtual int32 StrataVerticalLayering(int32 Top, int32 Base, int OperatorIndex, uint32 MaxDistanceFromLeaves) override;
+	virtual int32 StrataVerticalLayeringParameterBlending(int32 Top, int32 Base, const FString& SharedLocalBasisIndexMacro, int32 TopBSDFNormalCodeChunk, FStrataOperator* PromoteToOperator) override;
+	virtual int32 StrataAdd(int32 A, int32 B, int OperatorIndex, uint32 MaxDistanceFromLeaves) override;
+	virtual int32 StrataAddParameterBlending(int32 A, int32 B, int32 AMixWeight, const FString& SharedLocalBasisIndexMacro, FStrataOperator* PromoteToOperator) override;
+	virtual int32 StrataWeight(int32 A, int32 Weight, int OperatorIndex, uint32 MaxDistanceFromLeaves) override;
+	virtual int32 StrataWeightParameterBlending(int32 A, int32 Weight, FStrataOperator* PromoteToOperator) override;
 	virtual int32 StrataTransmittanceToMFP(int32 TransmittanceColor, int32 DesiredThickness, int32 OutputIndex) override;
 
+	virtual FStrataOperator& StrataCompilationRegisterOperator(int32 OperatorType, UMaterialExpression* Expression, UMaterialExpression* Parent, bool bUseParameterBlending = false) override;
+	virtual FStrataOperator& StrataCompilationGetOperator(UMaterialExpression* Expression) override;
+
 	virtual void StrataCompilationInfoRegisterCodeChunk(int32 CodeChunk, FStrataMaterialCompilationInfo& StrataMaterialCompilationInfo) override;
 	virtual bool StrataCompilationInfoContainsCodeChunk(int32 CodeChunk) override;
 	virtual const FStrataMaterialCompilationInfo& GetStrataCompilationInfo(int32 CodeChunk) override;
diff --git a/Engine/Source/Runtime/Engine/Private/Materials/MaterialExpressions.cpp b/Engine/Source/Runtime/Engine/Private/Materials/MaterialExpressions.cpp
--- a/Engine/Source/Runtime/Engine/Private/Materials/MaterialExpressions.cpp
+++ b/Engine/Source/Runtime/Engine/Private/Materials/MaterialExpressions.cpp
@@ -649,7 +649,9 @@ int32 CompileStrataBlendFunction(FMaterialCompiler* Compiler, const int32 A, con
 	{
 		return INDEX_NONE;
 	}
-	int32 OutputCodeChunk = Compiler->StrataHorizontalMixing(A, B, Alpha); // A is background (fully visible when Alpha=0) and B is foreground (fully visible when Alpha=1)
+	uint32 MaxDistanceToLeaves = 0; // STRATA_TODO: fix me
+
+	int32 OutputCodeChunk = Compiler->StrataHorizontalMixing(A, B, Alpha, 0, MaxDistanceToLeaves); // A is background (fully visible when Alpha=0) and B is foreground (fully visible when Alpha=1)
 
 #if WITH_EDITOR
 	// Also register the horizontal mixing operation with the compiler.
@@ -6365,6 +6367,15 @@ void UMaterialExpressionGetMaterialAttributes::GatherStrataMaterialInfo(FStrataM
 	}
 }
 
+FStrataOperator* UMaterialExpressionGetMaterialAttributes::StrataGenerateMaterialTopologyTree(class FMaterialCompiler* Compiler, class UMaterialExpression* Parent, int32 OutputIndex)
+{
+	if (MaterialAttributes.Expression)
+	{
+		return MaterialAttributes.Expression->StrataGenerateMaterialTopologyTree(Compiler, Parent, 0);
+	}
+	return nullptr;
+}
+
 void UMaterialExpressionGetMaterialAttributes::PreEditChange(FProperty* PropertyAboutToChange)
 {
 	// Backup attribute array so we can re-connect pins
@@ -6870,6 +6881,26 @@ void UMaterialExpressionBlendMaterialAttributes::GatherStrataMaterialInfo(FStrat
 		B.GetTracedInput().Expression->GatherStrataMaterialInfo(StrataMaterialInfo, B.OutputIndex);
 	}
 }
+
+FStrataOperator* UMaterialExpressionBlendMaterialAttributes::StrataGenerateMaterialTopologyTree(class FMaterialCompiler* Compiler, class UMaterialExpression* Parent, int32 OutputIndex)
+{
+	// STRATA_TODO: this likely no longer work. We would need to stop parsing and always do parameter blending at this stage.
+	const bool bUseParameterBlending = false;
+	FStrataOperator& StrataOperator = Compiler->StrataCompilationRegisterOperator(STRATA_OPERATOR_HORIZONTAL, this, Parent, bUseParameterBlending);
+
+	UMaterialExpression* ChildAExpression = A.GetTracedInput().Expression;
+	UMaterialExpression* ChildBExpression = B.GetTracedInput().Expression;
+	if (ChildAExpression)
+	{
+		StrataOperator.AIndex = ChildAExpression->StrataGenerateMaterialTopologyTree(Compiler, this, A.OutputIndex)->Index;
+	}
+	if (ChildBExpression)
+	{
+		StrataOperator.BIndex = ChildBExpression->StrataGenerateMaterialTopologyTree(Compiler, this, B.OutputIndex)->Index;
+	}
+
+	return &StrataOperator;
+}
 #endif // WITH_EDITOR
 
 //
@@ -14832,6 +14863,15 @@ void UMaterialExpressionMaterialFunctionCall::GatherStrataMaterialInfo(FStrataMa
 	}
 }
 
+FStrataOperator* UMaterialExpressionMaterialFunctionCall::StrataGenerateMaterialTopologyTree(class FMaterialCompiler* Compiler, class UMaterialExpression* Parent, int32 OutputIndex)
+{
+	if (OutputIndex >= 0 && OutputIndex < FunctionOutputs.Num() && FunctionOutputs[OutputIndex].ExpressionOutput)
+	{
+		return FunctionOutputs[OutputIndex].ExpressionOutput->StrataGenerateMaterialTopologyTree(Compiler, Parent, 0);
+	}
+	return nullptr;
+}
+
 uint32 UMaterialExpressionMaterialFunctionCall::GetInputType(int32 InputIndex)
 {
 	if (InputIndex < FunctionInputs.Num())
@@ -15375,6 +15415,15 @@ void UMaterialExpressionFunctionOutput::GatherStrataMaterialInfo(FStrataMaterial
 	}
 }
 
+FStrataOperator* UMaterialExpressionFunctionOutput::StrataGenerateMaterialTopologyTree(class FMaterialCompiler* Compiler, class UMaterialExpression* Parent, int32 OutputIndex)
+{
+	if (A.GetTracedInput().Expression && !A.Expression->ContainsInputLoop())
+	{
+		return A.Expression->StrataGenerateMaterialTopologyTree(Compiler, Parent, A.OutputIndex);
+	}
+	return nullptr;
+}
+
 #endif // WITH_EDITOR
 
 ///////////////////////////////////////////////////////////////////////////////
@@ -20641,6 +20690,20 @@ void UMaterialExpressionStrataLegacyConversion::GatherStrataMaterialInfo(FStrata
 	if (ConvertedStrataMaterialInfo.HasShadingModelFromExpression())			{ StrataMaterialInfo.SetShadingModelFromExpression(true); }
 }
 
+FStrataOperator* UMaterialExpressionStrataLegacyConversion::StrataGenerateMaterialTopologyTree(class FMaterialCompiler* Compiler, class UMaterialExpression* Parent, int32 OutputIndex)
+{
+	FStrataOperator& StrataOperator = Compiler->StrataCompilationRegisterOperator(STRATA_OPERATOR_VERTICAL, this, Parent);
+
+	// The legacy node can only be added once and we know the worst shape it can take so we hard code it here, 
+	// with fake addresses for BSDF defined in shader code.
+	UMaterialExpression* ChildAExpression = reinterpret_cast<UMaterialExpression*>(0);
+	UMaterialExpression* ChildBExpression = reinterpret_cast<UMaterialExpression*>(1);
+	StrataOperator.AIndex = Compiler->StrataCompilationRegisterOperator(STRATA_OPERATOR_BSDF_LEGACY, ChildAExpression, this).Index;
+	StrataOperator.BIndex = Compiler->StrataCompilationRegisterOperator(STRATA_OPERATOR_BSDF_LEGACY, ChildBExpression, this).Index;
+
+	return &StrataOperator;
+}
+
 bool UMaterialExpressionStrataLegacyConversion::HasSSS() const
 {
 	return SubsurfaceProfile != nullptr;
@@ -20718,6 +20781,8 @@ int32 UMaterialExpressionStrataSlabBSDF::Compile(class FMaterialCompiler* Compil
 	const float DefaultSpecular = 0.5f;
 	const float DefaultF0 = DielectricSpecularToF0(DefaultSpecular);
 
+	FStrataOperator& StrataOperator = Compiler->StrataCompilationGetOperator(this);
+
 	int32 OutputCodeChunk = Compiler->StrataSlabBSDF(
 		bUseMetalness ? Compiler->Constant(1.0f) : Compiler->Constant(0.0f),
 
@@ -20745,7 +20810,8 @@ int32 UMaterialExpressionStrataSlabBSDF::Compile(class FMaterialCompiler* Compil
 		CompileWithDefaultFloat1(Compiler, Thickness, STRATA_LAYER_DEFAULT_THICKNESS_CM),
 		NormalCodeChunk,
 		TangentCodeChunk,
-		Compiler->GetStrataSharedLocalBasisIndexMacro(NewRegisteredSharedLocalBasis));
+		Compiler->GetStrataSharedLocalBasisIndexMacro(NewRegisteredSharedLocalBasis),
+		!StrataOperator.bUseParameterBlending || (StrataOperator.bUseParameterBlending && StrataOperator.bRootOfParameterBlendingSubTree) ? &StrataOperator : nullptr);
 	StrataCompilationInfoCreateSingleBSDFMaterial(Compiler, OutputCodeChunk, NewRegisteredSharedLocalBasis, STRATA_BSDF_TYPE_SLAB, bHasSSS, bHasDMFPPluggedIn, bHasEdgeColor, bHasThinFilm, bHasFuzz, bHasHaziness);
 
 	return OutputCodeChunk;
@@ -21160,6 +21226,11 @@ void UMaterialExpressionStrataSlabBSDF::GatherStrataMaterialInfo(FStrataMaterial
 	}
 }
 
+FStrataOperator* UMaterialExpressionStrataSlabBSDF::StrataGenerateMaterialTopologyTree(class FMaterialCompiler* Compiler, class UMaterialExpression* Parent, int32 OutputIndex)
+{
+	return &Compiler->StrataCompilationRegisterOperator(STRATA_OPERATOR_BSDF, this, Parent);
+}
+
 bool UMaterialExpressionStrataSlabBSDF::HasSSS() const
 {
 	return SubsurfaceProfile != nullptr || SSSDMFP.IsConnected();
@@ -21272,6 +21343,11 @@ void UMaterialExpressionStrataVolumetricFogCloudBSDF::GatherStrataMaterialInfo(F
 {
 	StrataMaterialInfo.AddShadingModel(SSM_VolumetricFogCloud);
 }
+
+FStrataOperator* UMaterialExpressionStrataVolumetricFogCloudBSDF::StrataGenerateMaterialTopologyTree(class FMaterialCompiler* Compiler, class UMaterialExpression* Parent, int32 OutputIndex)
+{
+	return &Compiler->StrataCompilationRegisterOperator(STRATA_OPERATOR_BSDF, this, Parent);
+}
 #endif // WITH_EDITOR
 
 
@@ -21338,6 +21414,11 @@ void UMaterialExpressionStrataUnlitBSDF::GatherStrataMaterialInfo(FStrataMateria
 	if (EmissiveColor.IsConnected()) { StrataMaterialInfo.AddPropertyConnected(MP_EmissiveColor); }
 	StrataMaterialInfo.AddShadingModel(SSM_Unlit);
 }
+
+FStrataOperator* UMaterialExpressionStrataUnlitBSDF::StrataGenerateMaterialTopologyTree(class FMaterialCompiler* Compiler, class UMaterialExpression* Parent, int32 OutputIndex)
+{
+	return &Compiler->StrataCompilationRegisterOperator(STRATA_OPERATOR_BSDF, this, Parent);
+}
 #endif // WITH_EDITOR
 
 
@@ -21435,6 +21516,11 @@ void UMaterialExpressionStrataHairBSDF::GatherStrataMaterialInfo(FStrataMaterial
 
 	StrataMaterialInfo.AddShadingModel(SSM_Hair);
 }
+
+FStrataOperator* UMaterialExpressionStrataHairBSDF::StrataGenerateMaterialTopologyTree(class FMaterialCompiler* Compiler, class UMaterialExpression* Parent, int32 OutputIndex)
+{
+	return &Compiler->StrataCompilationRegisterOperator(STRATA_OPERATOR_BSDF, this, Parent);
+}
 #endif // WITH_EDITOR
 
 
@@ -21548,6 +21634,11 @@ void UMaterialExpressionStrataSingleLayerWaterBSDF::GatherStrataMaterialInfo(FSt
 
 	StrataMaterialInfo.AddShadingModel(SSM_SingleLayerWater);
 }
+
+FStrataOperator* UMaterialExpressionStrataSingleLayerWaterBSDF::StrataGenerateMaterialTopologyTree(class FMaterialCompiler* Compiler, class UMaterialExpression* Parent, int32 OutputIndex)
+{
+	return &Compiler->StrataCompilationRegisterOperator(STRATA_OPERATOR_BSDF, this, Parent);
+}
 #endif // WITH_EDITOR
 
 
@@ -21590,8 +21681,10 @@ int32 UMaterialExpressionStrataHorizontalMixing::Compile(class FMaterialCompiler
 
 	const int32 HorizontalMixCodeChunk = CompileWithDefaultFloat1(Compiler, Mix, 0.5f);
 
+
 	int32 OutputCodeChunk = INDEX_NONE;
-	if (bUseParameterBlending)
+	FStrataOperator& StrataOperator = Compiler->StrataCompilationGetOperator(this);
+	if (StrataOperator.bUseParameterBlending)
 	{
 		const int32 NormalMixCodeChunk = Compiler->StrataHorizontalMixingParameterBlendingBSDFCoverageToNormalMixCodeChunk(BackgroundCodeChunk, ForegroundCodeChunk, HorizontalMixCodeChunk);
 
@@ -21633,7 +21726,9 @@ int32 UMaterialExpressionStrataHorizontalMixing::Compile(class FMaterialCompiler
 		}
 		const FStrataRegisteredSharedLocalBasis NewRegisteredSharedLocalBasis = StrataCompilationInfoCreateSharedLocalBasis(Compiler, NewNormalCodeChunk, NewTangentCodeChunk);
 
-		OutputCodeChunk = Compiler->StrataHorizontalMixingParameterBlending(BackgroundCodeChunk, ForegroundCodeChunk, HorizontalMixCodeChunk, NormalMixCodeChunk, Compiler->GetStrataSharedLocalBasisIndexMacro(NewRegisteredSharedLocalBasis));
+		OutputCodeChunk = Compiler->StrataHorizontalMixingParameterBlending(
+			BackgroundCodeChunk, ForegroundCodeChunk, HorizontalMixCodeChunk, NormalMixCodeChunk, Compiler->GetStrataSharedLocalBasisIndexMacro(NewRegisteredSharedLocalBasis),
+			StrataOperator.bRootOfParameterBlendingSubTree ? &StrataOperator : nullptr);
 
 		StrataCompilationInfoCreateSingleBSDFMaterial(Compiler, OutputCodeChunk, NewRegisteredSharedLocalBasis, STRATA_BSDF_TYPE_SLAB);
 		FStrataMaterialCompilationInfo StrataInfo = StrataCompilationInfoHorizontalMixingParamBlend(Compiler, BackgroundStrataData, ForegroundStrataData, NewRegisteredSharedLocalBasis);
@@ -21645,7 +21740,9 @@ int32 UMaterialExpressionStrataHorizontalMixing::Compile(class FMaterialCompiler
 		OutputCodeChunk = Compiler->StrataHorizontalMixing(
 			BackgroundCodeChunk,
 			ForegroundCodeChunk,
-			HorizontalMixCodeChunk);
+			HorizontalMixCodeChunk,
+			StrataOperator.Index,
+			StrataOperator.MaxDistanceFromLeaves);
 
 		FStrataMaterialCompilationInfo StrataInfo = StrataCompilationInfoHorizontalMixing(Compiler, BackgroundStrataData, ForegroundStrataData);
 
@@ -21693,6 +21790,24 @@ void UMaterialExpressionStrataHorizontalMixing::GatherStrataMaterialInfo(FStrata
 		Background.GetTracedInput().Expression->GatherStrataMaterialInfo(StrataMaterialInfo, Background.OutputIndex);
 	}
 }
+
+FStrataOperator* UMaterialExpressionStrataHorizontalMixing::StrataGenerateMaterialTopologyTree(class FMaterialCompiler* Compiler, class UMaterialExpression* Parent, int32 OutputIndex)
+{
+	FStrataOperator& StrataOperator = Compiler->StrataCompilationRegisterOperator(STRATA_OPERATOR_HORIZONTAL, this, Parent, bUseParameterBlending);
+
+	UMaterialExpression* ChildAExpression = Background.GetTracedInput().Expression;
+	UMaterialExpression* ChildBExpression = Foreground.GetTracedInput().Expression;
+	if (ChildAExpression)
+	{
+		StrataOperator.AIndex = ChildAExpression->StrataGenerateMaterialTopologyTree(Compiler, this, Background.OutputIndex)->Index;
+	}
+	if (ChildBExpression)
+	{
+		StrataOperator.BIndex = ChildBExpression->StrataGenerateMaterialTopologyTree(Compiler, this, Foreground.OutputIndex)->Index;
+	}
+
+	return &StrataOperator;
+}
 #endif // WITH_EDITOR
 
 
@@ -21736,7 +21851,8 @@ int32 UMaterialExpressionStrataVerticalLayering::Compile(class FMaterialCompiler
 	const int32 TopNormalMixCodeChunk = Compiler->StrataVerticalLayeringParameterBlendingBSDFCoverageToNormalMixCodeChunk(TopCodeChunk);
 
 	int32 OutputCodeChunk = INDEX_NONE;
-	if (bUseParameterBlending)
+	FStrataOperator& StrataOperator = Compiler->StrataCompilationGetOperator(this);
+	if (StrataOperator.bUseParameterBlending)
 	{
 		if (TopStrataData.TotalBSDFCount != 1)
 		{
@@ -21776,7 +21892,9 @@ int32 UMaterialExpressionStrataVerticalLayering::Compile(class FMaterialCompiler
 		}
 		const FStrataRegisteredSharedLocalBasis NewRegisteredSharedLocalBasis = StrataCompilationInfoCreateSharedLocalBasis(Compiler, NewNormalCodeChunk, NewTangentCodeChunk);
 
-		OutputCodeChunk = Compiler->StrataVerticalLayeringParameterBlending(TopCodeChunk, BaseCodeChunk, Compiler->GetStrataSharedLocalBasisIndexMacro(NewRegisteredSharedLocalBasis), TopBSDF.RegisteredSharedLocalBasis.NormalCodeChunk);
+		OutputCodeChunk = Compiler->StrataVerticalLayeringParameterBlending(
+			TopCodeChunk, BaseCodeChunk, Compiler->GetStrataSharedLocalBasisIndexMacro(NewRegisteredSharedLocalBasis), TopBSDF.RegisteredSharedLocalBasis.NormalCodeChunk,
+			StrataOperator.bRootOfParameterBlendingSubTree ? &StrataOperator : nullptr);
 
 		StrataCompilationInfoCreateSingleBSDFMaterial(Compiler, OutputCodeChunk, NewRegisteredSharedLocalBasis, STRATA_BSDF_TYPE_SLAB);
 		FStrataMaterialCompilationInfo StrataInfo = StrataCompilationInfoVerticalLayeringParamBlend(Compiler, TopStrataData, BaseStrataData, NewRegisteredSharedLocalBasis);
@@ -21785,7 +21903,7 @@ int32 UMaterialExpressionStrataVerticalLayering::Compile(class FMaterialCompiler
 	}
 	else
 	{
-		OutputCodeChunk = Compiler->StrataVerticalLayering(TopCodeChunk, BaseCodeChunk);
+		OutputCodeChunk = Compiler->StrataVerticalLayering(TopCodeChunk, BaseCodeChunk, StrataOperator.Index, StrataOperator.MaxDistanceFromLeaves);
 
 		FStrataMaterialCompilationInfo StrataInfo = StrataCompilationInfoVerticalLayering(Compiler, Compiler->GetStrataCompilationInfo(TopCodeChunk), Compiler->GetStrataCompilationInfo(BaseCodeChunk));
 
@@ -21834,6 +21952,24 @@ void UMaterialExpressionStrataVerticalLayering::GatherStrataMaterialInfo(FStrata
 		Base.GetTracedInput().Expression->GatherStrataMaterialInfo(StrataMaterialInfo, Base.OutputIndex);
 	}
 }
+
+FStrataOperator* UMaterialExpressionStrataVerticalLayering::StrataGenerateMaterialTopologyTree(class FMaterialCompiler* Compiler, class UMaterialExpression* Parent, int32 OutputIndex)
+{
+	FStrataOperator& StrataOperator = Compiler->StrataCompilationRegisterOperator(STRATA_OPERATOR_VERTICAL, this, Parent, bUseParameterBlending);
+
+	UMaterialExpression* ChildAExpression = Top.GetTracedInput().Expression;
+	UMaterialExpression* ChildBExpression = Base.GetTracedInput().Expression;
+	if (ChildAExpression)
+	{
+		StrataOperator.AIndex = ChildAExpression->StrataGenerateMaterialTopologyTree(Compiler, this, Top.OutputIndex)->Index;
+	}
+	if (ChildBExpression)
+	{
+		StrataOperator.BIndex = ChildBExpression->StrataGenerateMaterialTopologyTree(Compiler, this, Base.OutputIndex)->Index;
+	}
+
+	return &StrataOperator;
+}
 #endif // WITH_EDITOR
 
 
@@ -21877,7 +22013,8 @@ int32 UMaterialExpressionStrataAdd::Compile(class FMaterialCompiler* Compiler, i
 	const int32 ANormalMixCodeChunk = Compiler->StrataAddParameterBlendingBSDFCoverageToNormalMixCodeChunk(ACodeChunk, BCodeChunk);
 
 	int32 OutputCodeChunk = INDEX_NONE;
-	if (bUseParameterBlending)
+	FStrataOperator& StrataOperator = Compiler->StrataCompilationGetOperator(this);
+	if (StrataOperator.bUseParameterBlending)
 	{
 		if (AStrataData.TotalBSDFCount != 1)
 		{
@@ -21917,7 +22054,9 @@ int32 UMaterialExpressionStrataAdd::Compile(class FMaterialCompiler* Compiler, i
 		}
 		const FStrataRegisteredSharedLocalBasis NewRegisteredSharedLocalBasis = StrataCompilationInfoCreateSharedLocalBasis(Compiler, NewNormalCodeChunk, NewTangentCodeChunk);
 
-		OutputCodeChunk = Compiler->StrataAddParameterBlending(ACodeChunk, BCodeChunk, ANormalMixCodeChunk, Compiler->GetStrataSharedLocalBasisIndexMacro(NewRegisteredSharedLocalBasis));
+		OutputCodeChunk = Compiler->StrataAddParameterBlending(
+			ACodeChunk, BCodeChunk, ANormalMixCodeChunk, Compiler->GetStrataSharedLocalBasisIndexMacro(NewRegisteredSharedLocalBasis),
+			StrataOperator.bRootOfParameterBlendingSubTree ? &StrataOperator : nullptr);
 
 		StrataCompilationInfoCreateSingleBSDFMaterial(Compiler, OutputCodeChunk, NewRegisteredSharedLocalBasis, STRATA_BSDF_TYPE_SLAB);
 		FStrataMaterialCompilationInfo StrataInfo = StrataCompilationInfoAddParamBlend(Compiler, AStrataData, BStrataData, NewRegisteredSharedLocalBasis);
@@ -21926,7 +22065,7 @@ int32 UMaterialExpressionStrataAdd::Compile(class FMaterialCompiler* Compiler, i
 	}
 	else
 	{
-		OutputCodeChunk = Compiler->StrataAdd(ACodeChunk, BCodeChunk);
+		OutputCodeChunk = Compiler->StrataAdd(ACodeChunk, BCodeChunk, StrataOperator.Index, StrataOperator.MaxDistanceFromLeaves);
 
 		FStrataMaterialCompilationInfo StrataInfo = StrataCompilationInfoAdd(Compiler, Compiler->GetStrataCompilationInfo(ACodeChunk), Compiler->GetStrataCompilationInfo(BCodeChunk));
 
@@ -21974,6 +22113,24 @@ void UMaterialExpressionStrataAdd::GatherStrataMaterialInfo(FStrataMaterialInfo&
 		B.GetTracedInput().Expression->GatherStrataMaterialInfo(StrataMaterialInfo, B.OutputIndex);
 	}
 }
+
+FStrataOperator* UMaterialExpressionStrataAdd::StrataGenerateMaterialTopologyTree(class FMaterialCompiler* Compiler, class UMaterialExpression* Parent, int32 OutputIndex)
+{
+	FStrataOperator& StrataOperator = Compiler->StrataCompilationRegisterOperator(STRATA_OPERATOR_ADD, this, Parent, bUseParameterBlending);
+
+	UMaterialExpression* ChildAExpression = A.GetTracedInput().Expression;
+	UMaterialExpression* ChildBExpression = B.GetTracedInput().Expression;
+	if (ChildAExpression)
+	{
+		StrataOperator.AIndex = ChildAExpression->StrataGenerateMaterialTopologyTree(Compiler, this, A.OutputIndex)->Index;
+	}
+	if (ChildBExpression)
+	{
+		StrataOperator.BIndex = ChildBExpression->StrataGenerateMaterialTopologyTree(Compiler, this, B.OutputIndex)->Index;
+	}
+
+	return &StrataOperator;
+}
 #endif // WITH_EDITOR
 
 
@@ -22003,14 +22160,29 @@ int32 UMaterialExpressionStrataWeight::Compile(class FMaterialCompiler* Compiler
 	int32 ACodeChunk = A.Compile(Compiler);
 	int32 WeightCodeChunk = Weight.GetTracedInput().Expression ? Weight.Compile(Compiler) : Compiler->Constant(1.0f);
 
-	int32 OutputCodeChunk = Compiler->StrataWeight(ACodeChunk, WeightCodeChunk);
+	int32 OutputCodeChunk = INDEX_NONE;
+	FStrataOperator& StrataOperator = Compiler->StrataCompilationGetOperator(this);
+	if (StrataOperator.bUseParameterBlending)
+	{
+		OutputCodeChunk = Compiler->StrataWeightParameterBlending(
+			ACodeChunk, WeightCodeChunk, 
+			StrataOperator.bRootOfParameterBlendingSubTree ? &StrataOperator : nullptr);
+
+		FStrataMaterialCompilationInfo StrataInfo = StrataCompilationInfoWeightParamBlend(Compiler, Compiler->GetStrataCompilationInfo(ACodeChunk));
 
-	if (!Compiler->StrataCompilationInfoContainsCodeChunk(ACodeChunk))
+		Compiler->StrataCompilationInfoRegisterCodeChunk(OutputCodeChunk, StrataInfo);
+	}
+	else
 	{
-		return Compiler->Errorf(TEXT("Could not find ACodeChunk to multiply"));
+		OutputCodeChunk = Compiler->StrataWeight(ACodeChunk, WeightCodeChunk, StrataOperator.Index, StrataOperator.MaxDistanceFromLeaves);
+
+		if (!Compiler->StrataCompilationInfoContainsCodeChunk(ACodeChunk))
+		{
+			return Compiler->Errorf(TEXT("Could not find ACodeChunk to multiply"));
+		}
+		FStrataMaterialCompilationInfo StrataInfo = StrataCompilationInfoWeight(Compiler, Compiler->GetStrataCompilationInfo(ACodeChunk));
+		Compiler->StrataCompilationInfoRegisterCodeChunk(OutputCodeChunk, StrataInfo);
 	}
-	FStrataMaterialCompilationInfo StrataInfo = StrataCompilationInfoWeight(Compiler, Compiler->GetStrataCompilationInfo(ACodeChunk));
-	Compiler->StrataCompilationInfoRegisterCodeChunk(OutputCodeChunk, StrataInfo);
 
 	return OutputCodeChunk;
 }
@@ -22046,6 +22218,19 @@ void UMaterialExpressionStrataWeight::GatherStrataMaterialInfo(FStrataMaterialIn
 		A.GetTracedInput().Expression->GatherStrataMaterialInfo(StrataMaterialInfo, A.OutputIndex);
 	}
 }
+
+FStrataOperator* UMaterialExpressionStrataWeight::StrataGenerateMaterialTopologyTree(class FMaterialCompiler* Compiler, class UMaterialExpression* Parent, int32 OutputIndex)
+{
+	FStrataOperator& StrataOperator = Compiler->StrataCompilationRegisterOperator(STRATA_OPERATOR_WEIGHT, this, Parent);
+
+	UMaterialExpression* ChildAExpression = A.GetTracedInput().Expression;
+	if (ChildAExpression)
+	{
+		StrataOperator.AIndex = ChildAExpression->StrataGenerateMaterialTopologyTree(Compiler, this, A.OutputIndex)->Index;
+	}
+
+	return &StrataOperator;
+}
 #endif // WITH_EDITOR
 
 
diff --git a/Engine/Source/Runtime/Engine/Private/Materials/StrataMaterial.cpp b/Engine/Source/Runtime/Engine/Private/Materials/StrataMaterial.cpp
--- a/Engine/Source/Runtime/Engine/Private/Materials/StrataMaterial.cpp
+++ b/Engine/Source/Runtime/Engine/Private/Materials/StrataMaterial.cpp
@@ -97,6 +97,12 @@ FStrataMaterialCompilationInfo StrataCompilationInfoWeight(FMaterialCompiler* Co
 	return StrataInfo;
 }
 
+FStrataMaterialCompilationInfo StrataCompilationInfoWeightParamBlend(FMaterialCompiler* Compiler, const FStrataMaterialCompilationInfo& A)
+{
+	FStrataMaterialCompilationInfo StrataInfo = A;
+	return StrataInfo;
+}
+
 
 FStrataMaterialCompilationInfo StrataCompilationInfoAdd(FMaterialCompiler* Compiler, const FStrataMaterialCompilationInfo& A, const FStrataMaterialCompilationInfo& B)
 {
diff --git a/Engine/Source/Runtime/Engine/Private/Materials/StrataMaterial.h b/Engine/Source/Runtime/Engine/Private/Materials/StrataMaterial.h
--- a/Engine/Source/Runtime/Engine/Private/Materials/StrataMaterial.h
+++ b/Engine/Source/Runtime/Engine/Private/Materials/StrataMaterial.h
@@ -69,6 +69,7 @@ void  StrataCompilationInfoCreateSingleBSDFMaterial(FMaterialCompiler* Compiler,
 	uint8 BSDFType, bool bHasSSS = false, bool bHasDMFPPluggedIn = false, bool bHasEdgeColor = false, bool bHasThinFilm = false, bool bHasFuzz = false, bool bHasHaziness = false);
 
 FStrataMaterialCompilationInfo StrataCompilationInfoWeight(FMaterialCompiler* Compiler, const FStrataMaterialCompilationInfo& A);
+FStrataMaterialCompilationInfo StrataCompilationInfoWeightParamBlend(FMaterialCompiler* Compiler, const FStrataMaterialCompilationInfo& A);
 FStrataMaterialCompilationInfo StrataCompilationInfoAdd(FMaterialCompiler* Compiler, const FStrataMaterialCompilationInfo& A, const FStrataMaterialCompilationInfo& B);
 FStrataMaterialCompilationInfo StrataCompilationInfoAddParamBlend(FMaterialCompiler* Compiler, const FStrataMaterialCompilationInfo& A, const FStrataMaterialCompilationInfo& B, const FStrataRegisteredSharedLocalBasis& RegisteredSharedLocalBasis);
 FStrataMaterialCompilationInfo StrataCompilationInfoHorizontalMixing(FMaterialCompiler* Compiler, const FStrataMaterialCompilationInfo& Background, const FStrataMaterialCompilationInfo& Foreground);
diff --git a/Engine/Source/Runtime/Engine/Public/MaterialCompiler.h b/Engine/Source/Runtime/Engine/Public/MaterialCompiler.h
--- a/Engine/Source/Runtime/Engine/Public/MaterialCompiler.h
+++ b/Engine/Source/Runtime/Engine/Public/MaterialCompiler.h
@@ -68,6 +68,49 @@ struct FStrataRegisteredSharedLocalBasis
 	uint8 GraphSharedLocalBasisIndex;
 };
 
+struct FStrataOperator
+{
+	int32 OperatorType;
+	bool bNodeRequestParameterBlending;
+
+	int32 Index;		// Index into the array of operators
+	int32 TopIndex;		// Parent operator index
+	int32 AIndex;		// Left child operator index
+	int32 BIndex;		// Right child operator index
+
+	int32 BSDFIndex;	// Index in the array of BSDF if a BSDF operator
+
+	// Data derived after the tree has been built.
+	int32 MaxDistanceFromLeaves;
+	int32 LayerDepth;
+	bool bIsTop;
+	bool bIsBottom;
+	bool bUseParameterBlending;				// True when part of a sub tree where parameter blending is in use
+	bool bRootOfParameterBlendingSubTree;	// True when the root of a sub tree where parameter blending is in use. Only this node will register a BSDF
+
+	FStrataOperator()
+	{
+		OperatorType = INDEX_NONE;
+		bNodeRequestParameterBlending = false;
+		Index = INDEX_NONE;
+		TopIndex = INDEX_NONE;
+		AIndex = INDEX_NONE;
+		BIndex = INDEX_NONE;
+		BSDFIndex = INDEX_NONE;
+		MaxDistanceFromLeaves = 0;
+		LayerDepth = 0;
+		bIsTop = false;
+		bIsBottom = false;
+		bUseParameterBlending = false;
+		bRootOfParameterBlendingSubTree = false;
+	}
+
+	bool IsDiscarded()
+	{
+		return bUseParameterBlending && !bRootOfParameterBlendingSubTree;
+	}
+};
+
 /** 
  * The interface used to translate material expressions into executable code. 
  * Note: Most member functions should be pure virtual to force a FProxyMaterialCompiler override!
@@ -454,7 +497,8 @@ public:
 		int32 ThinFilmThickness, 
 		int32 FuzzAmount, int32 FuzzColor,
 		int32 Thickness,
-		int32 Normal, int32 Tangent, const FString& SharedLocalBasisIndexMacro) = 0;
+		int32 Normal, int32 Tangent, const FString& SharedLocalBasisIndexMacro,
+		FStrataOperator* PromoteToOperator) = 0;
 	virtual int32 StrataConversionFromLegacy(
 		bool bHasDynamicShadingModels,
 		int32 BaseColor, int32 Specular, int32 Metallic,
@@ -474,15 +518,25 @@ public:
 	virtual int32 StrataSingleLayerWaterBSDF(
 		int32 BaseColor, int32 Metallic, int32 Specular, int32 Roughness, int32 EmissiveColor, int32 TopMaterialOpacity, 
 		int32 WaterAlbedo, int32 WaterExtinction, int32 WaterPhaseG, int32 ColorScaleBehindWater, int32 Normal, const FString& SharedLocalBasisIndexMacro) = 0;
-	virtual int32 StrataHorizontalMixing(int32 Background, int32 Foreground, int32 Mix) = 0;
-	virtual int32 StrataHorizontalMixingParameterBlending(int32 Background, int32 Foreground, int32 HorizontalMixCodeChunk, int32 NormalMixCodeChunk, const FString& SharedLocalBasisIndexMacro) = 0;
-	virtual int32 StrataVerticalLayering(int32 Top, int32 Base) = 0;
-	virtual int32 StrataVerticalLayeringParameterBlending(int32 Top, int32 Base, const FString& SharedLocalBasisIndexMacro, int32 TopBSDFNormalCodeChunk) = 0;
-	virtual int32 StrataAdd(int32 A, int32 B) = 0;
-	virtual int32 StrataAddParameterBlending(int32 A, int32 B, int32 AMixWeight, const FString& SharedLocalBasisIndexMacro) = 0;
-	virtual int32 StrataWeight(int32 A, int32 Weight) = 0;
+	virtual int32 StrataHorizontalMixing(int32 Background, int32 Foreground, int32 Mix, int OperatorIndex, uint32 MaxDistanceFromLeaves) = 0;
+	virtual int32 StrataHorizontalMixingParameterBlending(int32 Background, int32 Foreground, int32 HorizontalMixCodeChunk, int32 NormalMixCodeChunk, const FString& SharedLocalBasisIndexMacro, FStrataOperator* PromoteToOperator) = 0;
+	virtual int32 StrataVerticalLayering(int32 Top, int32 Base, int OperatorIndex, uint32 MaxDistanceFromLeaves) = 0;
+	virtual int32 StrataVerticalLayeringParameterBlending(int32 Top, int32 Base, const FString& SharedLocalBasisIndexMacro, int32 TopBSDFNormalCodeChunk, FStrataOperator* PromoteToOperator) = 0;
+	virtual int32 StrataAdd(int32 A, int32 B, int OperatorIndex, uint32 MaxDistanceFromLeaves) = 0;
+	virtual int32 StrataAddParameterBlending(int32 A, int32 B, int32 AMixWeight, const FString& SharedLocalBasisIndexMacro, FStrataOperator* PromoteToOperator) = 0;
+	virtual int32 StrataWeight(int32 A, int32 Weight, int OperatorIndex, uint32 MaxDistanceFromLeaves) = 0;
+	virtual int32 StrataWeightParameterBlending(int32 A, int32 Weight, FStrataOperator* PromoteToOperator) = 0;
 	virtual int32 StrataTransmittanceToMFP(int32 TransmittanceColor, int32 DesiredThickness, int32 OutputIndex) = 0;
 
+	/**
+	 * Register an operator of the tree representation the starta material and its topology.
+	 */
+	virtual FStrataOperator& StrataCompilationRegisterOperator(int32 OperatorType, UMaterialExpression* Expression, UMaterialExpression* Parent, bool bUseParameterBlending = false) = 0;
+	/**
+	 * Return the operator information for a given expression.
+	 */
+	virtual FStrataOperator& StrataCompilationGetOperator(UMaterialExpression* Expression) = 0;
+
 	virtual void StrataCompilationInfoRegisterCodeChunk(int32 CodeChunk, FStrataMaterialCompilationInfo& StrataMaterialCompilationInfo) = 0;
 	virtual bool StrataCompilationInfoContainsCodeChunk(int32 CodeChunk) = 0;
 	virtual const FStrataMaterialCompilationInfo& GetStrataCompilationInfo(int32 CodeChunk) = 0;
@@ -982,7 +1036,8 @@ public:
 		int32 ThinFilmThickness, 
 		int32 FuzzAmount, int32 FuzzColor,
 		int32 Thickness,
-		int32 Normal, int32 Tangent, const FString& SharedLocalBasisIndexMacro) override
+		int32 Normal, int32 Tangent, const FString& SharedLocalBasisIndexMacro,
+		FStrataOperator* PromoteToOperator) override
 	{
 		return Compiler->StrataSlabBSDF(
 			UseMetalness,
@@ -995,7 +1050,8 @@ public:
 			ThinFilmThickness,
 			FuzzAmount, FuzzColor,
 			Thickness,
-			Normal, Tangent, SharedLocalBasisIndexMacro);
+			Normal, Tangent, SharedLocalBasisIndexMacro,
+			PromoteToOperator);
 	}
 
 	virtual int32 StrataConversionFromLegacy(
@@ -1051,39 +1107,44 @@ public:
 			WaterAlbedo, WaterExtinction, WaterPhaseG, ColorScaleBehindWater, Normal, SharedLocalBasisIndexMacro);
 	}
 
-	virtual int32 StrataHorizontalMixing(int32 Background, int32 Foreground, int32 Mix) override
+	virtual int32 StrataHorizontalMixing(int32 Background, int32 Foreground, int32 Mix, int OperatorIndex, uint32 MaxDistanceFromLeaves) override
+	{
+		return Compiler->StrataHorizontalMixing(Background, Foreground, Mix, OperatorIndex, MaxDistanceFromLeaves);
+	}
+
+	virtual int32 StrataHorizontalMixingParameterBlending(int32 Background, int32 Foreground, int32 HorizontalMixCodeChunk, int32 NormalMixCodeChunk, const FString& SharedLocalBasisIndexMacro, FStrataOperator* PromoteToOperator) override
 	{
-		return Compiler->StrataHorizontalMixing(Background, Foreground, Mix);
+		return Compiler->StrataHorizontalMixingParameterBlending(Background, Foreground, HorizontalMixCodeChunk, NormalMixCodeChunk, SharedLocalBasisIndexMacro, PromoteToOperator);
 	}
 
-	virtual int32 StrataHorizontalMixingParameterBlending(int32 Background, int32 Foreground, int32 HorizontalMixCodeChunk, int32 NormalMixCodeChunk, const FString& SharedLocalBasisIndexMacro) override
+	virtual int32 StrataVerticalLayering(int32 Top, int32 Base, int OperatorIndex, uint32 MaxDistanceFromLeaves) override
 	{
-		return Compiler->StrataHorizontalMixingParameterBlending(Background, Foreground, HorizontalMixCodeChunk, NormalMixCodeChunk, SharedLocalBasisIndexMacro);
+		return Compiler->StrataVerticalLayering(Top, Base, OperatorIndex, MaxDistanceFromLeaves);
 	}
 
-	virtual int32 StrataVerticalLayering(int32 Top, int32 Base) override
+	virtual int32 StrataVerticalLayeringParameterBlending(int32 Top, int32 Base, const FString& SharedLocalBasisIndexMacro, int32 TopBSDFNormalCodeChunk, FStrataOperator* PromoteToOperator) override
 	{
-		return Compiler->StrataVerticalLayering(Top, Base);
+		return Compiler->StrataVerticalLayeringParameterBlending(Top, Base, SharedLocalBasisIndexMacro, TopBSDFNormalCodeChunk, PromoteToOperator);
 	}
 
-	virtual int32 StrataVerticalLayeringParameterBlending(int32 Top, int32 Base, const FString& SharedLocalBasisIndexMacro, int32 TopBSDFNormalCodeChunk) override
+	virtual int32 StrataAdd(int32 A, int32 B, int OperatorIndex, uint32 MaxDistanceFromLeaves) override
 	{
-		return Compiler->StrataVerticalLayeringParameterBlending(Top, Base, SharedLocalBasisIndexMacro, TopBSDFNormalCodeChunk);
+		return Compiler->StrataAdd(A, B, OperatorIndex, MaxDistanceFromLeaves);
 	}
 
-	virtual int32 StrataAdd(int32 A, int32 B) override
+	virtual int32 StrataAddParameterBlending(int32 A, int32 B, int32 AMixWeight, const FString& SharedLocalBasisIndexMacro, FStrataOperator* PromoteToOperator) override
 	{
-		return Compiler->StrataAdd(A, B);
+		return Compiler->StrataAddParameterBlending(A, B, AMixWeight, SharedLocalBasisIndexMacro, PromoteToOperator);
 	}
 
-	virtual int32 StrataAddParameterBlending(int32 A, int32 B, int32 AMixWeight, const FString& SharedLocalBasisIndexMacro) override
+	virtual int32 StrataWeight(int32 A, int32 Weight, int OperatorIndex, uint32 MaxDistanceFromLeaves) override
 	{
-		return Compiler->StrataAddParameterBlending(A, B, AMixWeight, SharedLocalBasisIndexMacro);
+		return Compiler->StrataWeight(A, Weight, OperatorIndex, MaxDistanceFromLeaves);
 	}
 
-	virtual int32 StrataWeight(int32 A, int32 Weight) override
+	virtual int32 StrataWeightParameterBlending(int32 A, int32 Weight, FStrataOperator* PromoteToOperator) override
 	{
-		return Compiler->StrataWeight(A, Weight);
+		return Compiler->StrataWeightParameterBlending(A, Weight, PromoteToOperator);
 	}
 
 	virtual int32 StrataTransmittanceToMFP(int32 TransmittanceColor, int32 DesiredThickness, int32 OutputIndex) override
@@ -1096,6 +1157,16 @@ public:
 		Compiler->StrataCompilationInfoRegisterCodeChunk(CodeChunk, StrataMaterialCompilationInfo);
 	}
 
+	virtual FStrataOperator& StrataCompilationRegisterOperator(int32 OperatorType, UMaterialExpression* Expression, UMaterialExpression* Parent, bool bUseParameterBlending = false) override
+	{
+		return Compiler->StrataCompilationRegisterOperator(OperatorType, Expression, Parent, bUseParameterBlending);
+	}
+
+	virtual FStrataOperator& StrataCompilationGetOperator(UMaterialExpression* Expression) override
+	{
+		return Compiler->StrataCompilationGetOperator(Expression);
+	}
+
 	virtual bool StrataCompilationInfoContainsCodeChunk(int32 CodeChunk) override
 	{
 		return Compiler->StrataCompilationInfoContainsCodeChunk(CodeChunk);
