Normalise strata epsilon usage.

diff --git a/Engine/Shaders/Private/Strata/Strata.ush b/Engine/Shaders/Private/Strata/Strata.ush
--- a/Engine/Shaders/Private/Strata/Strata.ush
+++ b/Engine/Shaders/Private/Strata/Strata.ush
@@ -105,6 +105,8 @@ FRWStrataMaterialContainer InitialiseRWStrataMaterialContainer(RWTexture2DArray<
 #define STRATA_SIMPLEVOLUME_THICKNESS_CM	100.0f
 #define STRATA_SIMPLEVOLUME_THICKNESS_M		(STRATA_SIMPLEVOLUME_THICKNESS_CM * CENTIMETER_TO_METER)
 
+#define STRATA_EPSILON						1e-10f
+
 ///////////////////////////////////////////////////////////////////////////////
 // BSDFs representation
 struct FStrataBSDF
@@ -1200,7 +1202,7 @@ float HorizontalMixingParameterBlendingBSDFCoverageToNormalMix(FStrataData Backg
 	const float ForegroundMixFactor = saturate(HorizontalMixValue);
 	const float BackgroundMixFactor = 1.0 - ForegroundMixFactor;
 
-	return (SafeForeCoverage * ForegroundMixFactor) / max(1e-5, SafeForeCoverage * ForegroundMixFactor + SafeBackCoverage * BackgroundMixFactor);
+	return (SafeForeCoverage * ForegroundMixFactor) / max(STRATA_EPSILON, SafeForeCoverage * ForegroundMixFactor + SafeBackCoverage * BackgroundMixFactor);
 #else
 	return 0.5f;
 #endif
@@ -1305,7 +1307,7 @@ float AddParameterBlendingBSDFCoverageToNormalMix(FStrataData AStrata, FStrataDa
 #if STRATA_INLINE_SHADING
 	const float SafeABSDFCoverage = saturate(AStrata.InlinedBSDF.Coverage);
 	const float SafeBBSDFCoverage = saturate(BStrata.InlinedBSDF.Coverage);
-	const float AMixFactor = SafeABSDFCoverage / (1e-10 + SafeABSDFCoverage + SafeBBSDFCoverage);
+	const float AMixFactor = SafeABSDFCoverage / max(STRATA_EPSILON, SafeABSDFCoverage + SafeBBSDFCoverage);
 	return AMixFactor;
 #else
 	return 0.5f;
@@ -1438,9 +1440,9 @@ FStrataBSDF StrataVerticalLayeringParameterBlending(in FStrataBSDF TopBSDF, in F
 
 	FVerticalLayeringInfo Info			= GetVerticalLayeringInfo(SafeTopCoverage, SafeBaseCoverage);
 	// Compute some ways to intepoldate and accumulate parameters, but always normalise because coverage is applied later in UpdateSingleBSDFOperatorCoverageTransmittance.
-	const float  TopToBottomLerp		= saturate((Info.TransmittanceTopAndBottom * dot(SlabTransmittance,0.33.xxx) + Info.SurfaceBottom) / max(1e-5, Info.Coverage));
-	const float3 BottomColorFactor		= (Info.TransmittanceTopAndBottom * SlabTransmittance + Info.SurfaceBottom) / max(1e-5, Info.Coverage);
-	const float  TopFactor				= Info.SurfaceTop / max(1e-5, Info.Coverage);
+	const float  TopToBottomLerp		= saturate((Info.TransmittanceTopAndBottom * dot(SlabTransmittance,0.33.xxx) + Info.SurfaceBottom) / max(STRATA_EPSILON, Info.Coverage));
+	const float3 BottomColorFactor		= (Info.TransmittanceTopAndBottom * SlabTransmittance + Info.SurfaceBottom) / max(STRATA_EPSILON, Info.Coverage);
+	const float  TopFactor				= Info.SurfaceTop / max(STRATA_EPSILON, Info.Coverage);
 	
 	// Add both layers thickness
 	BSDF_SETTHICKNESSCM(ResultBSDF,		BSDF_GETTHICKNESSCM(BaseBSDF)												+ SafeTopCoverage * BSDF_GETTHICKNESSCM(TopBSDF));
@@ -2442,7 +2444,7 @@ void PackStrataOut(
 
 					float3x3 TangentBasis = StrataGetBSDFSharedBasis_InlineShading(StrataPixelHeader, BSDF, StrataAddressing);
 					float3 N = TangentBasis[2];
-					TopLayerData.WorldNormal += BSDF.bIsTop ? N * max(1e-5, BSDF.Coverage) : 0.0;
+					TopLayerData.WorldNormal += BSDF.bIsTop ? N * max(STRATA_EPSILON, BSDF.Coverage) : 0.0;
 
 					AddShadingModels(StrataPixelHeader, StrataShadingModelBit(BSDF_GETTYPE(BSDF)));
 					AndIsSimpleMaterial(StrataPixelHeader, bIsSimpleMaterial);
diff --git a/Engine/Shaders/Private/Strata/StrataEvaluation.ush b/Engine/Shaders/Private/Strata/StrataEvaluation.ush
--- a/Engine/Shaders/Private/Strata/StrataEvaluation.ush
+++ b/Engine/Shaders/Private/Strata/StrataEvaluation.ush
@@ -220,7 +220,7 @@ FStrataEvaluateResult StrataEvaluateBSDFCommon(FStrataBSDFContext BSDFContext, F
 
 				SphereMaxNoH(BSDFContext.Context, AreaLightContext.AreaLight.SphereSinAlpha, true);
 			}
-			BSDFContext.Context.NoV = saturate(abs(BSDFContext.Context.NoV) + 1e-5);
+			BSDFContext.Context.NoV = saturate(max(abs(BSDFContext.Context.NoV), STRATA_EPSILON));
 
 			////
 			//// Evaluate the diffuse component.
@@ -273,7 +273,7 @@ FStrataEvaluateResult StrataEvaluateBSDFCommon(FStrataBSDFContext BSDFContext, F
 				{
 					const float ChanF0 = 0.04;
 					FBxDFEnergyTerms EnergyTerms = ComputeGGXSpecEnergyTerms(SafeRoughness, BSDFContext.Context.NoV, ChanF0, 1.0f);
-					const float3 SpecularTransmitionChanF0ToMaterialF0 = DirectionalAlbedo_SpecularTransmission / max(1e-8f, ComputeEnergyPreservation(EnergyTerms));
+					const float3 SpecularTransmitionChanF0ToMaterialF0 = DirectionalAlbedo_SpecularTransmission / max(STRATA_EPSILON, ComputeEnergyPreservation(EnergyTerms));
 					Sample.DiffusePathValue *= SpecularTransmitionChanF0ToMaterialF0;
 					Sample.IntegratedDiffuseValue *= SpecularTransmitionChanF0ToMaterialF0;
 				}
diff --git a/Engine/Shaders/Private/Strata/StrataTree.ush b/Engine/Shaders/Private/Strata/StrataTree.ush
--- a/Engine/Shaders/Private/Strata/StrataTree.ush
+++ b/Engine/Shaders/Private/Strata/StrataTree.ush
@@ -108,7 +108,7 @@ void UpdateSingleOperatorCoverageTransmittance(inout FStrataTree StrataTree, int
 		const float BMix = Mix;
 
 		Op.Coverage = AMix * OpA.Coverage + BMix * OpB.Coverage;
-		Op.Transmittance = (AMix * OpA.Coverage * OpA.Transmittance + BMix * OpB.Coverage * OpB.Transmittance) / max(1e-5, AMix * OpA.Coverage + BMix * OpB.Coverage);
+		Op.Transmittance = (AMix * OpA.Coverage * OpA.Transmittance + BMix * OpB.Coverage * OpB.Transmittance) / max(STRATA_EPSILON, AMix * OpA.Coverage + BMix * OpB.Coverage);
 		break;
 	}
 
@@ -116,7 +116,7 @@ void UpdateSingleOperatorCoverageTransmittance(inout FStrataTree StrataTree, int
 	{
 		const float SafeABSDFCoverage = saturate(OpA.Coverage);
 		const float SafeBBSDFCoverage = saturate(OpB.Coverage);
-		const float AMixFactor = SafeABSDFCoverage / (1e-10 + SafeABSDFCoverage + SafeBBSDFCoverage);
+		const float AMixFactor = SafeABSDFCoverage / max(STRATA_EPSILON, SafeABSDFCoverage + SafeBBSDFCoverage);
 
 		Op.Coverage = saturate(OpA.Coverage + OpB.Coverage);
 		Op.Transmittance = lerp(OpB.Transmittance, OpA.Transmittance, AMixFactor);
