Strata - re enabled rough refraction after last strata core change.

diff --git a/Engine/Shaders/Private/DistortAccumulatePS.usf b/Engine/Shaders/Private/DistortAccumulatePS.usf
--- a/Engine/Shaders/Private/DistortAccumulatePS.usf
+++ b/Engine/Shaders/Private/DistortAccumulatePS.usf
@@ -88,31 +88,18 @@ void Main(
 
 #if STRATA_ENABLED
 
-	// STRATA_TODO Fix me
-#if 0
-
 	FStrataData StrataData = PixelMaterialInputs.FrontMaterial;
 
 	FStrataAddressing NullStrataAddressing = (FStrataAddressing)0;	// Fake unused in StrataCreateBSDFContext when using Forward inline shading
 
 	FStrataPixelHeader StrataPixelHeader = InitialiseStrataPixelHeader();
-	int l;
-	UNROLL
-	for (l = 0; l < StrataData.LayerCount; ++l)
-	{
-		StrataPixelHeader.BSDFCount += StrataData.Layers[l].BSDFCount;
-	}
-	StrataPixelHeader.SharedLocalBases = MaterialParameters.SharedLocalBases;
 	StrataPixelHeader.StrataTree = MaterialParameters.StrataTree;
-	// STRATA_TODO: process the tree
+	StrataPixelHeader.BSDFCount = MaterialParameters.StrataTree.BSDFCount;
+	StrataPixelHeader.SharedLocalBases = MaterialParameters.SharedLocalBases;
 
 	const float3 V = MaterialParameters.CameraVector;
-	
-	const float OpaqueBSDFThroughput = 0.0f;
-	const float FullThroughput = 1.0f;
+	const float3 DummyL = float3(1, 0, 0);
 
-	float3 Color = 0;
-	float OutThroughput = 1.0f;
 
 	// The current way rough refraction are evaluated is
 	//  1- start with the dirac view Lobe
@@ -120,28 +107,43 @@ void Main(
 	//  3- Use that new Lobe as input (pointing away from the surface) for the next layer until done.
 	StrataLobeStatistic WiLobe = StrataGetViewDiracLobe(V);
 
-	UNROLL
-	for (l = 0; l < StrataData.LayerCount; ++l)
-	{
-		float3 LayerThroughput = 0.0f;
-		float  LayerCoverage = 0.0f;
-
-		const bool bTopLayer = l == 0;
-		const bool bBottomLayer = l == (StrataData.LayerCount - 1);
+	StrataLobeStatistic AccumulatedLayerLobes = StrataGetNullLobe();
+	float AccumulatedLayerLobesWeights = 0.0f;
 
-		StrataLobeStatistic AccumulatedLayerLobes = StrataGetNullLobe();
-		float AccumulatedLayerLobesWeights = 0.0f;
-
-		UNROLL
-		for (int i = 0; i < StrataData.Layers[l].BSDFCount; ++i)
+	FStrataTree StrataTree = StrataPixelHeader.StrataTree;
+	if (StrataTree.BSDFCount > 0)
+	{
+		const bool bRoughDiffuseEnabled = false;
+	#if STRATA_USES_CONVERSION_FROM_LEGACY==0
+		//
+		// Update the coverage/transmittance of each leaves (==BSDFs) of the strata tree.
+		//
+		UpdateAllBSDFsOperatorCoverageTransmittance(StrataPixelHeader, StrataTree, bRoughDiffuseEnabled, NullStrataAddressing, V, DummyL);
+
+		//
+		// Propagate up the coverage/transmittance of each node in the strata tree.
+		//
+		UpdateAllOperatorsCoverageTransmittance(StrataTree);
+
+		//
+		// Update the luminance weight of each BSDF according to the operators it has to traverse up to the strata tree root node.
+		//
+		UpdateAllBSDFWeightAfterOperatorVisit(StrataTree);
+	#else
+		LegacyUpdateBSDFsOperators(StrataPixelHeader, StrataTree, bRoughDiffuseEnabled, NullStrataAddressing, V, DummyL);
+	#endif
+
+		
+		STRATA_UNROLL_N(STRATA_CLAMPED_BSDF_COUNT)
+		for (int BSDFIdx = 0; BSDFIdx < StrataTree.BSDFCount; ++BSDFIdx)
 		{
-			FStrataBSDF BSDF = StrataData.Layers[l].BSDFs[i];
+		#define BSDF StrataTree.BSDFs[BSDFIdx]
 
 			// Sanitize BSDF before it is used for forward shading
 			StrataSanitizeBSDF(BSDF);
 
 			BRANCH
-			if(BSDF_GETTYPE(BSDF) == STRATA_BSDF_TYPE_SLAB)
+			if (BSDF_GETTYPE(BSDF) == STRATA_BSDF_TYPE_SLAB)
 			{
 				// We want simple volumetric only if the DMFP input is pluged in, otherwise we avoid dealing with simple volumetric
 				const bool bIsSimpleVolume = BSDF_GETHASDMFP(BSDF);
@@ -154,10 +156,9 @@ void Main(
 			const float BSDFCoverage = BSDF.Coverage;
 
 			// Create the BSDF context
-			const float3 FakeL = float3(0, 0, 1);
-			FStrataBSDFContext StrataBSDFContext = StrataCreateBSDFContext(StrataPixelHeader, BSDF, NullStrataAddressing, V, FakeL);
+			FStrataBSDFContext StrataBSDFContext = StrataCreateBSDFContext(StrataPixelHeader, BSDF, NullStrataAddressing, V, DummyL);
 
-			FStrataEvaluateResult StrataEvaluate = StrataEvaluateBSDF(StrataBSDFContext, true); // Works to call because FakeL is nto used when computing StrataEvaluate.Throughput
+			FStrataEvaluateResult StrataEvaluate = StrataEvaluateBSDF(StrataBSDFContext, true); // Works to call because DummyL is nto used when computing StrataEvaluate.Throughput
 
 			// We only support rough refraction from Slab material type for now.
 			if (BSDF_GETTYPE(BSDF) == STRATA_BSDF_TYPE_SLAB)
@@ -166,7 +167,7 @@ void Main(
 				const FBxDFEnergyTerms EnergyTerms = ComputeGGXSpecEnergyTerms(SLAB_ROUGHNESS(BSDF), StrataBSDFContext.SatNoV, SLAB_F0(BSDF), SLAB_F90(BSDF));
 
 				// The top layer we consider a transition air to water. For layers below, we assume they all have similar index of refraction.
-				const float InterfaceEta12 = i == 0 ? 1.0f / 1.33f : 1.0f; // STRATA_TODO Use Refraction.x in IOR mode?
+				const float InterfaceEta12 = BSDF.bIsTop ? 1.0f / 1.33f : 1.0f; // STRATA_TODO Use Refraction.x in IOR mode? Or the material IOR corresponding to F0?
 
 				// Transform input Lobe into local space w.r.t. normal
 				const float3 SavedWorldSpaceWiLobeMu = WiLobe.Mu;
@@ -174,9 +175,10 @@ void Main(
 
 				// Compute the refracted Lobe according to roughness . 
 				// STRATA_TODO: take into account scattering and transmittance. Colored transmittance is taken into account already but not scattering influence on variance.
-				StrataLobeStatistic NewLobe = StrataGetRefractedLobe(WiLobe, EnergyTerms.E * OutThroughput * BSDFCoverage, SLAB_ROUGHNESS(BSDF), InterfaceEta12);
+				StrataLobeStatistic NewLobe = StrataGetRefractedLobe(WiLobe, EnergyTerms.E * BSDF.LuminanceWeight, SLAB_ROUGHNESS(BSDF), InterfaceEta12);
 
-				// Transform the new Lobe direction to world space and accumulate
+				// Transform the new Lobe direction to world space and weight/accumulate
+				// STRATA_TODO it should be better to weight lobes horizontally weighted and chain lobes that are vertically layered. For that we would need another set of UpdateAllXXX function as seen above.
 				NewLobe.Mu = NewLobe.Mu.x * StrataBSDFContext.X + NewLobe.Mu.y * StrataBSDFContext.Y + NewLobe.Mu.z * StrataBSDFContext.N;
 				StrataAccumulateWeightedLobe(AccumulatedLayerLobes, AccumulatedLayerLobesWeights, NewLobe);
 
@@ -187,26 +189,21 @@ void Main(
 			{
 				// Full roughness to make it visible as something to fix in case a type is not handled.
 				WiLobe.Sigma = 1000.0;
-			} 
+			}
 
-			LayerCoverage	+= BSDFCoverage;
-			LayerThroughput += BSDFCoverage * StrataEvaluate.Throughput;
+		#undef BSDF
 		}
+	}
 
-		if (AccumulatedLayerLobesWeights > 0.0f)
-		{
-			// Now compute the new input Wi Lobe from the accumulated refracted Lobe of the previous layer
-			WiLobe = StrataNormaliseLobe(AccumulatedLayerLobes, AccumulatedLayerLobesWeights);
-			// And negate the direction to make sure the previously refracted direction is the new incoming
-			WiLobe.Mu = -WiLobe.Mu;
-		}
+	if (AccumulatedLayerLobesWeights > 0.0f)
+	{
+		// Now compute the new input Wi Lobe from the accumulated refracted Lobe of the previous layer
+		WiLobe = StrataNormaliseLobe(AccumulatedLayerLobes, AccumulatedLayerLobesWeights);
+		// And negate the direction to make sure the previously refracted direction is the new incoming
+		WiLobe.Mu = -WiLobe.Mu;
+	}
 
-		// Take into account the non covered part for this layer and update throughput accordingly
-		LayerThroughput += (1.0 - saturate(LayerCoverage)) * FullThroughput;
 
-		// And combine this layer throughtput with he material throughput
-		OutThroughput *= LayerThroughput;
-	}
 
 	// Not write out variance (as this is more linear to accumulate than roughness)
 	OutVariance = float4(WiLobe.Sigma, 0.0f, 0.0f, 0.0f);
@@ -228,11 +225,6 @@ void Main(
 
 	// Sample depth at distortion offset
 	float2 DistortBufferUV = ScreenUV + BufferUVDistortion;
-#endif
-
-	OutVariance = 0.0f;
-	float2 BufferUVDistortion = 0.0f;
-	float2 DistortBufferUV = ScreenUV;
 
 #else // STRATA_ENABLED
 
