Fixed tile artifacts when rendering Nanite materials that have pixel depth offset enabled

diff --git a/Engine/Shaders/Private/MaterialTemplate.ush b/Engine/Shaders/Private/MaterialTemplate.ush
--- a/Engine/Shaders/Private/MaterialTemplate.ush
+++ b/Engine/Shaders/Private/MaterialTemplate.ush
@@ -3294,7 +3294,7 @@ half3x3 AssembleTangentToWorld( half3 TangentToWorld0, half4 TangentToWorld2 )
 }
 
 // Whether the material shader should output pixel depth offset
-#define OUTPUT_PIXEL_DEPTH_OFFSET (WANT_PIXEL_DEPTH_OFFSET && ((MATERIALBLENDING_SOLID || MATERIALBLENDING_MASKED) || (TRANSLUCENT_WRITING_VELOCITY)))
+#define OUTPUT_PIXEL_DEPTH_OFFSET (WANT_PIXEL_DEPTH_OFFSET && !IS_NANITE_PASS && ((MATERIALBLENDING_SOLID || MATERIALBLENDING_MASKED) || (TRANSLUCENT_WRITING_VELOCITY)))
 
 // Whether to use the hidden d3d11 feature that supports depth writes with ZCull by only pushing into the screen
 //@todo - use for other SM5 platforms
