Update GPUSceneWriter with better support for LWC
Add ability to add instance local bounds to dynamic primitive instances

diff --git a/Engine/Plugins/FX/Niagara/Shaders/Private/NiagaraUpdateMeshGPUSceneInstances.usf b/Engine/Plugins/FX/Niagara/Shaders/Private/NiagaraUpdateMeshGPUSceneInstances.usf
--- a/Engine/Plugins/FX/Niagara/Shaders/Private/NiagaraUpdateMeshGPUSceneInstances.usf
+++ b/Engine/Plugins/FX/Niagara/Shaders/Private/NiagaraUpdateMeshGPUSceneInstances.usf
@@ -180,28 +180,18 @@ void UpdateMeshInstancesCS(uint3 GroupId : SV_GroupID, uint GroupIndex : SV_Grou
 		}
 	}
 
-	// The transforms solved are currently relative to the system LWC tile. We need to re-base it to the 
-	// Primitive's current tile
-	FLWCMatrix LocalToWorld = Transforms.LocalToWorld;
-	FLWCMatrix PrevLocalToWorld = Transforms.PrevLocalToWorld;
-	FLWCVector3 NewTileOrigin = MakeLWCVector3(PrimitiveData.TilePosition, (float3)0);
-	FLWCVector3 PrimitiveOffset = LWCSubtract(LWCGetOrigin(Transforms.LocalToWorld), NewTileOrigin);
-	FLWCVector3 PrevPrimitiveOffset = LWCSubtract(LWCGetOrigin(Transforms.PrevLocalToWorld), NewTileOrigin);
-	LWCSetOrigin(LocalToWorld, MakeLWCVector3(PrimitiveData.TilePosition, LWCToFloat(PrimitiveOffset)));
-	LWCSetOrigin(PrevLocalToWorld, MakeLWCVector3(PrimitiveData.TilePosition, LWCToFloat(PrevPrimitiveOffset)));
-
 	const uint PayloadDataFlags = GetPayloadDataFlags();
-	InitializeInstanceSceneData(
+	InitializeInstanceSceneDataWS(
 		InstanceId,
 		PrimitiveId,
 		LocalInstanceId,
  		PayloadDataFlags,
  		1,
  		0.0f,
-		LocalToWorld.M
+		Transforms.LocalToWorld
 	);
 	
 	FInstancePayloadDataOffsets PayloadOffsets = GetInstancePayloadDataOffsets(PrimitiveId, PayloadDataFlags, LocalInstanceId);
-	WriteInstanceDynamicData(PayloadOffsets, PrevLocalToWorld.M);
+	WriteInstanceDynamicDataWS(PayloadOffsets, PrimitiveId, Transforms.PrevLocalToWorld);
 	WriteInstanceCustomData(PayloadOffsets, 0, float4(asfloat(ParticleIndex), 0.0f, 0.0f, 0.0f));
 }
\ No newline at end of file
diff --git a/Engine/Shaders/Private/GPUScene/GPUSceneWriter.ush b/Engine/Shaders/Private/GPUScene/GPUSceneWriter.ush
--- a/Engine/Shaders/Private/GPUScene/GPUSceneWriter.ush
+++ b/Engine/Shaders/Private/GPUScene/GPUSceneWriter.ush
@@ -37,6 +37,11 @@ void StoreInstancePayloadDataElement(uint PayloadDataOffset, uint ElementIndex,
 	GPUSceneInstancePayloadDataRW[PayloadDataOffset + ElementIndex] = DataElement;
 }
 
+void StoreInstancePayloadDataElementZW(uint PayloadDataOffset, uint ElementIndex, float2 DataElement)
+{
+	GPUSceneInstancePayloadDataRW[PayloadDataOffset + ElementIndex].zw = DataElement;
+}
+
 /**
  * Helpers to load and store the primitive ID and flags for the specified instance, which are packed together
  */
@@ -64,21 +69,24 @@ void WriteInstancePrimitiveId(uint InstanceId, uint NewPrimitiveId)
 
 /**
  * Store the contents of InstanceData.LocalToWorld to the GPUScene buffers. Optionally, also calculates the determinant
- * of the transform's rotation to set or unset the determinant sign flag on the instance flags
+ * of the transform's rotation to set or unset the determinant sign flag on the instance flags.
+ *
+ * NOTE: With Large World Coordinates, the transform provided is expected to be relative to the primitive's tile position.
+ * Use the WriteInstanceLocalToWorld helper if you don't already have a relative transform.
  **/
-void WriteInstanceLocalToWorld(uint InstanceId, float4x4 LocalToWorld, bool StoreDeterminantSign)
+void WriteInstanceLocalToRelativeWorld(uint InstanceId, float4x4 LocalToRelativeWorld, bool StoreDeterminantSign)
 {
 #if INSTANCE_SCENE_DATA_COMPRESSED_TRANSFORMS
 	uint4 RotationScale;
 	float3 Translation;
-	EncodeTransform(LocalToWorld, RotationScale, Translation);
+	EncodeTransform(LocalToRelativeWorld, RotationScale, Translation);
 	StoreInstanceSceneDataElement(InstanceId, 1, asfloat(RotationScale));
 	StoreInstanceSceneDataElement(InstanceId, 2, float4(Translation, 0.0f));
 #else
-	LocalToWorld = transpose(LocalToWorld);
-	StoreInstanceSceneDataElement(InstanceId, 1, LocalToWorld[0]);
-	StoreInstanceSceneDataElement(InstanceId, 2, LocalToWorld[1]);
-	StoreInstanceSceneDataElement(InstanceId, 3, LocalToWorld[2]);
+	LocalToRelativeWorld = transpose(LocalToRelativeWorld);
+	StoreInstanceSceneDataElement(InstanceId, 1, LocalToRelativeWorld[0]);
+	StoreInstanceSceneDataElement(InstanceId, 2, LocalToRelativeWorld[1]);
+	StoreInstanceSceneDataElement(InstanceId, 3, LocalToRelativeWorld[2]);
 #endif
 
 	// check to update the determinant sign flag on the instance flags as well
@@ -87,13 +95,22 @@ void WriteInstanceLocalToWorld(uint InstanceId, float4x4 LocalToWorld, bool Stor
 		// Check to update the determinant sign flag on the instance flags
 		uint PrimitiveId, InstanceFlags;
 		LoadInstancePrimitiveIdAndFlags(InstanceId, PrimitiveId, InstanceFlags);
-		SetInstanceDeterminantSignFlag(determinant((float3x3)LocalToWorld), InstanceFlags);
+		SetInstanceDeterminantSignFlag(determinant((float3x3)LocalToRelativeWorld), InstanceFlags);
 		WriteInstancePrimitiveIdAndFlags(InstanceId, PrimitiveId, InstanceFlags);
 	}
 }
 
+void WriteInstanceLocalToWorld(uint InstanceId, uint PrimitiveId, FLWCMatrix LocalToWorld, bool StoreDeterminantSign)
+{
+	// Ensure the LocalToWorld matrix is relative to the primitive tile position
+	LocalToWorld = LWCMakeRelativeToTile(LocalToWorld, GetPrimitiveData(PrimitiveId).TilePosition);
+	WriteInstanceLocalToRelativeWorld(InstanceId, LocalToWorld.M, StoreDeterminantSign);
+}
+
 /**
- * Initializes the instance scene data of the specified instance data slot.
+ * Initializes the instance scene data of the specified instance data slot using a LocalToWorld matrix that is already relative to
+ * its primitive's render tile position (When LWC. Otherwise, the matrix is simply LocalToWorld). Use InitializeInstanceSceneDataWS
+ * or InitializeInstanceSceneDataLS if you don't already have a LocalToRelativeWorldMatrix
  * NOTE: This function must match FInstanceSceneShaderData::BuildInternal
  **/
  void InitializeInstanceSceneData(
@@ -103,9 +120,9 @@ void WriteInstanceLocalToWorld(uint InstanceId, float4x4 LocalToWorld, bool Stor
  	uint InstanceFlags, // combination of INSTANCE_SCENE_DATA_FLAG_
  	uint CustomDataCount,
  	float RandomId,
- 	float4x4 LocalToWorld)
+ 	float4x4 LocalToRelativeWorld)
  {
- 	SetInstanceDeterminantSignFlag(determinant((float3x3)LocalToWorld), InstanceFlags);
+ 	SetInstanceDeterminantSignFlag(determinant((float3x3)LocalToRelativeWorld), InstanceFlags);
  	
  	float4 Element0 =
  	{
@@ -115,9 +132,39 @@ void WriteInstanceLocalToWorld(uint InstanceId, float4x4 LocalToWorld, bool Stor
  		RandomId
 	};
  	StoreInstanceSceneDataElement(InstanceId, 0, Element0);
- 	WriteInstanceLocalToWorld(InstanceId, LocalToWorld, false);
+ 	WriteInstanceLocalToRelativeWorld(InstanceId, LocalToRelativeWorld, false);
  }
 
+/** This version initializes the instance with a world-space transform */
+void InitializeInstanceSceneDataWS(
+ 	uint InstanceId,
+ 	uint PrimitiveId,
+ 	uint InstanceRelativeId,
+ 	uint PayloadDataFlags,
+ 	uint CustomDataCount,
+ 	float RandomId,
+ 	FLWCMatrix LocalToWorld)
+ {
+ 	// Ensure the WS matrix is relative to the primitive's render tile, then the internal matrix is LocalToRelativeWorld
+ 	LocalToWorld = LWCMakeRelativeToTile(LocalToWorld, GetPrimitiveData(PrimitiveId).TilePosition);
+ 	InitializeInstanceSceneData(InstanceId, PrimitiveId, InstanceRelativeId, PayloadDataFlags, CustomDataCount, RandomId, LocalToWorld.M);
+ }
+
+/** This version initializes the instance with a transform that is relative to its primitive's local space */
+void InitializeInstanceSceneDataLS(
+ 	uint InstanceId,
+ 	uint PrimitiveId,
+ 	uint InstanceRelativeId,
+ 	uint PayloadDataFlags,
+ 	uint CustomDataCount,
+ 	float RandomId,
+ 	float4x4 LocalToPrimitive)
+{
+	// Multiply the LS by the primitive's LocalToWorld, then the internal matrix is guaranteed to be LocalToRelativeWorld
+	FLWCMatrix LocalToWorld = LWCMultiply(LocalToPrimitive, GetPrimitiveData(PrimitiveId).LocalToWorld);
+	InitializeInstanceSceneData(InstanceId, PrimitiveId, InstanceRelativeId, PayloadDataFlags, CustomDataCount, RandomId, LocalToWorld.M);
+}
+
 /**
  * Helper for loading instance payload offsets for GPU writers (assumes global parameter for SOA stride)
  **/
@@ -130,9 +177,31 @@ void WriteInstanceLocalToWorld(uint InstanceId, float4x4 LocalToWorld, bool Stor
 }
 
 /**
- * Store the data for dynamic instances to the GPUScene sidecar payload buffer.
+ * Store the local instance bounds to the GPU Scene sidecar payload buffer.
  **/
-void WriteInstanceDynamicData(FInstancePayloadDataOffsets PayloadOffsets, float4x4 PrevLocalToWorld)
+void WriteInstanceLocalBounds(FInstancePayloadDataOffsets PayloadOffsets, float3 LocalCenter, float3 LocalExtent)
+{
+	BRANCH
+	if (PayloadOffsets.LocalBounds == INVALID_INSTANCE_PAYLOAD_OFFSET)
+	{
+		return;
+	}
+
+	StoreInstancePayloadDataElementZW(PayloadOffsets.LocalBounds, 0, LocalCenter.xy);
+	StoreInstancePayloadDataElement(PayloadOffsets.LocalBounds, 1, float4(LocalCenter.z, LocalExtent.x, LocalExtent.y, LocalExtent.z));
+}
+
+void WriteInstanceLocalBounds(uint InstanceId, float3 LocalCenter, float3 LocalExtent)
+{
+	FInstancePayloadDataOffsets PayloadOffsets = GetInstancePayloadDataOffsets(InstanceId);
+	WriteInstanceLocalBounds(PayloadOffsets, LocalCenter, LocalExtent);
+}
+
+/**
+ * Store the data for dynamic instances to the GPUScene sidecar payload buffer. Use WriteInstanceDynamicDataWS or 
+ * WriteInstanceDynamicDataLS if you don't have a relative transform.
+ **/
+void WriteInstanceDynamicData(FInstancePayloadDataOffsets PayloadOffsets, float4x4 PrevLocalToRelativeWorld)
 {	
 	BRANCH
 	if (PayloadOffsets.DynamicData == INVALID_INSTANCE_PAYLOAD_OFFSET)
@@ -144,21 +213,44 @@ void WriteInstanceDynamicData(FInstancePayloadDataOffsets PayloadOffsets, float4
 #if INSTANCE_SCENE_DATA_COMPRESSED_TRANSFORMS
 	uint4 RotationScale;
 	float3 Translation;
-	EncodeTransform(PrevLocalToWorld, RotationScale, Translation);
+	EncodeTransform(PrevLocalToRelativeWorld, RotationScale, Translation);
 	StoreInstancePayloadDataElement(PayloadOffsets.DynamicData, 0, asfloat(RotationScale));
 	StoreInstancePayloadDataElement(PayloadOffsets.DynamicData, 1, float4(Translation, 0.0f));
 #else
-	PrevLocalToWorld = transpose(PrevLocalToWorld);
-	StoreInstancePayloadDataElement(PayloadOffsets.DynamicData, 0, PrevLocalToWorld[0]);
-	StoreInstancePayloadDataElement(PayloadOffsets.DynamicData, 1, PrevLocalToWorld[1]);
-	StoreInstancePayloadDataElement(PayloadOffsets.DynamicData, 2, PrevLocalToWorld[2]);
+	PrevLocalToRelativeWorld = transpose(PrevLocalToRelativeWorld);
+	StoreInstancePayloadDataElement(PayloadOffsets.DynamicData, 0, PrevLocalToRelativeWorld[0]);
+	StoreInstancePayloadDataElement(PayloadOffsets.DynamicData, 1, PrevLocalToRelativeWorld[1]);
+	StoreInstancePayloadDataElement(PayloadOffsets.DynamicData, 2, PrevLocalToRelativeWorld[2]);
 #endif
 }
 
-void WriteInstanceDynamicData(uint InstanceId, float4x4 PrevLocalToWorld)
+void WriteInstanceDynamicData(uint InstanceId, float4x4 PrevLocalToRelativeWorld)
 {
-	FInstancePayloadDataOffsets PayloadOffsets = GetInstancePayloadDataOffsets(InstanceId);
-	WriteInstanceDynamicData(PayloadOffsets, PrevLocalToWorld);
+	WriteInstanceDynamicData(GetInstancePayloadDataOffsets(InstanceId), PrevLocalToRelativeWorld);
+}
+
+/** This version computes the relative transform from the specified world transform */
+void WriteInstanceDynamicDataWS(FInstancePayloadDataOffsets PayloadOffsets, uint PrimitiveId, FLWCMatrix PrevLocalToWorld)
+{
+	PrevLocalToWorld = LWCMakeRelativeToTile(PrevLocalToWorld, GetPrimitiveData(PrimitiveId).TilePosition);
+	WriteInstanceDynamicData(PayloadOffsets, PrevLocalToWorld.M);
+}
+
+void WriteInstanceDynamicDataWS(uint InstanceId, uint PrimitiveId, FLWCMatrix PrevLocalToWorld)
+{
+	WriteInstanceDynamicDataWS(GetInstancePayloadDataOffsets(InstanceId), PrimitiveId, PrevLocalToWorld);
+}
+
+/** This version computes the relative transform from the specified local-space transform */
+void WriteInstanceDynamicDataLS(FInstancePayloadDataOffsets PayloadOffsets, uint PrimitiveId, float4x4 PrevLocalToPrimitive)
+{
+	FLWCMatrix PrevLocalToWorld = LWCMultiply(PrevLocalToPrimitive, GetPrimitiveData(PrimitiveId).PreviousLocalToWorld);
+	WriteInstanceDynamicData(PayloadOffsets, PrevLocalToWorld.M);
+}
+
+void WriteInstanceDynamicDataLS(uint InstanceId, uint PrimitiveId, float4x4 PrevLocalToPrimitive)
+{
+	WriteInstanceDynamicDataLS(GetInstancePayloadDataOffsets(InstanceId), PrimitiveId, PrevLocalToPrimitive);
 }
 
 /**
diff --git a/Engine/Shaders/Private/LargeWorldCoordinates.ush b/Engine/Shaders/Private/LargeWorldCoordinates.ush
--- a/Engine/Shaders/Private/LargeWorldCoordinates.ush
+++ b/Engine/Shaders/Private/LargeWorldCoordinates.ush
@@ -388,6 +388,17 @@ void LWCSetOrigin(inout FLWCMatrix InOutMatrix, FLWCVector3 Origin)
 #undef FBoolType
 #undef LWCConstructor
 
+// Creates a coordinate with the same value, but relative to the specified tile, while keeping precision loss as low as possible
+FLWCScalar LWCMakeRelativeToTile(FLWCScalar V, float NewTile) { return MakeLWCScalar(NewTile, LWCToFloat(LWCSubtract(V, MakeLWCScalar(NewTile, (float)0.0f)))); }
+FLWCVector2 LWCMakeRelativeToTile(FLWCVector2 V, float2 NewTile) { return MakeLWCVector2(NewTile, LWCToFloat(LWCSubtract(V, MakeLWCVector2(NewTile, (float2)0.0f)))); }
+FLWCVector3 LWCMakeRelativeToTile(FLWCVector3 V, float3 NewTile) { return MakeLWCVector3(NewTile, LWCToFloat(LWCSubtract(V, MakeLWCVector3(NewTile, (float3)0.0f)))); }
+FLWCVector4 LWCMakeRelativeToTile(FLWCVector4 V, float4 NewTile) { return MakeLWCVector4(NewTile, LWCToFloat(LWCSubtract(V, MakeLWCVector4(NewTile, (float4)0.0f)))); }
+FLWCMatrix LWCMakeRelativeToTile(FLWCMatrix M, float3 NewTile)
+{
+	LWCSetOrigin(M, LWCMakeRelativeToTile(LWCGetOrigin(M), NewTile));
+	return M;
+}
+
 FLWCScalar LWCDot(FLWCScalar Lhs, FLWCScalar Rhs)
 {
 	return LWCMultiply(Lhs, Rhs);
diff --git a/Engine/Source/Runtime/Engine/Public/MeshBatch.h b/Engine/Source/Runtime/Engine/Public/MeshBatch.h
--- a/Engine/Source/Runtime/Engine/Public/MeshBatch.h
+++ b/Engine/Source/Runtime/Engine/Public/MeshBatch.h
@@ -71,6 +71,7 @@ struct FMeshBatchDynamicPrimitiveData
 {
 	TConstArrayView<FPrimitiveInstance> InstanceSceneData;
 	TConstArrayView<FPrimitiveInstanceDynamicData> InstanceDynamicData;
+	TConstArrayView<FRenderBounds> InstanceLocalBounds;
 	TConstArrayView<float> InstanceCustomData;
 	FGPUSceneWriteDelegate DataWriterGPU;		
 	EGPUSceneGPUWritePass DataWriterGPUPass = EGPUSceneGPUWritePass::None;
@@ -90,15 +91,30 @@ struct FMeshBatchDynamicPrimitiveData
 	}
 	
 	FORCEINLINE void EnableInstanceDynamicData(bool bEnable) { SetPayloadDataFlags(INSTANCE_SCENE_DATA_FLAG_HAS_DYNAMIC_DATA, bEnable); }
+	FORCEINLINE void EnableInstanceLocalBounds(bool bEnable) { SetPayloadDataFlags(INSTANCE_SCENE_DATA_FLAG_HAS_LOCAL_BOUNDS, bEnable); }
 	FORCEINLINE void SetNumInstanceCustomDataFloats(uint32 NumFloats)
 	{
 		SetPayloadDataFlags(INSTANCE_SCENE_DATA_FLAG_HAS_CUSTOM_DATA, NumFloats > 0);
 		NumInstanceCustomDataFloats = NumFloats;
 	}
 
+	/**
+	 * Computes the full float4 stride of the instance's payload data.
+	 * NOTE: Needs to align with GetInstancePayloadDataOffsets in SceneData.ush
+	 **/
 	FORCEINLINE uint32 GetPayloadFloat4Stride() const
 	{
 		uint32 Total = 0;
+		
+		if (PayloadDataFlags & INSTANCE_SCENE_DATA_FLAG_HAS_LOCAL_BOUNDS)
+		{
+			Total += 2;
+		}
+		else if (PayloadDataFlags & (INSTANCE_SCENE_DATA_FLAG_HAS_HIERARCHY_OFFSET | INSTANCE_SCENE_DATA_FLAG_HAS_EDITOR_DATA))
+		{
+			Total += 1;
+		}
+
 		if (PayloadDataFlags & INSTANCE_SCENE_DATA_FLAG_HAS_DYNAMIC_DATA)
 		{
 #if INSTANCE_SCENE_DATA_COMPRESSED_TRANSFORMS
@@ -107,6 +123,12 @@ struct FMeshBatchDynamicPrimitiveData
 			Total += 3;
 #endif
 		}
+
+		if (PayloadDataFlags & INSTANCE_SCENE_DATA_FLAG_HAS_LIGHTSHADOW_UV_BIAS)
+		{
+			Total += 1;
+		}
+		
 		if (PayloadDataFlags & INSTANCE_SCENE_DATA_FLAG_HAS_CUSTOM_DATA)
 		{
 			Total += FMath::DivideAndRoundUp(NumInstanceCustomDataFloats, 4u);
diff --git a/Engine/Source/Runtime/Renderer/Private/GPUScene.cpp b/Engine/Source/Runtime/Renderer/Private/GPUScene.cpp
--- a/Engine/Source/Runtime/Renderer/Private/GPUScene.cpp
+++ b/Engine/Source/Runtime/Renderer/Private/GPUScene.cpp
@@ -1606,6 +1606,7 @@ struct FUploadDataSourceAdapterDynamicPrimitives
 			InstanceUploadInfo.InstanceFlags                = PrimData.SourceData.PayloadDataFlags;
 			InstanceUploadInfo.PrimitiveInstances 			= PrimData.SourceData.InstanceSceneData;
 			InstanceUploadInfo.InstanceDynamicData 			= PrimData.SourceData.InstanceDynamicData;
+			InstanceUploadInfo.InstanceLocalBounds 			= PrimData.SourceData.InstanceLocalBounds;
 			InstanceUploadInfo.InstanceCustomData 			= PrimData.SourceData.InstanceCustomData;
 			InstanceUploadInfo.InstanceRandomID				= TConstArrayView<float>();
 			InstanceUploadInfo.InstanceHierarchyOffset		= TConstArrayView<uint32>();
