* Move SSS profile function accessor from BurleyNormalizedSSSCommon.ush to SubsurfaceProfileCommon.ush
* Use common sampling function for mobile compatibility.

diff --git a/Engine/Shaders/Private/BurleyNormalizedSSSCommon.ush b/Engine/Shaders/Private/BurleyNormalizedSSSCommon.ush
--- a/Engine/Shaders/Private/BurleyNormalizedSSSCommon.ush
+++ b/Engine/Shaders/Private/BurleyNormalizedSSSCommon.ush
@@ -7,63 +7,6 @@
 
 #include "SubsurfaceProfileCommon.ush"
 
-#define BURLEY_MM_2_CM		0.1f
-#define BURLEY_CM_2_MM      10.0f
-
-// exactly one of these should be true
-#define LIGHT_PERPENDICULAR 0
-#define LIGHT_DIFFUSESURFACE 0
-#define LIGHT_PERPENDICULAR_DIFFUSE_SURFACE 1
-
-//------------------------------------------------------------------------------------------
-// Consistent in BurleyNormalizedSSSCommon.ush and SubsurfaceProfile.cpp
-
-#define SSS_TYPE_BURLEY	    0
-#define SSS_TYPE_SSSS		1
-
-// Make sure UIMax|ClampMax of WorldUNitScale * ENC_WORLDUNITSCALE_IN_CM_TO_UNIT <= 1
-#define ENC_WORLDUNITSCALE_IN_CM_TO_UNIT 0.02f
-#define DEC_UNIT_TO_WORLDUNITSCALE_IN_CM 1/ENC_WORLDUNITSCALE_IN_CM_TO_UNIT
-
-// Make sure UIMax|ClampMax of DiffuseMeanFreePath * ENC_DIFFUSEMEANFREEPATH_IN_MM_TO_UNIT <= 1
-#define ENC_DIFFUSEMEANFREEPATH_IN_MM_TO_UNIT (0.01f*0.2f)
-#define DEC_UNIT_TO_DIFFUSEMEANFREEPATH_IN_MM 1/ENC_DIFFUSEMEANFREEPATH_IN_MM_TO_UNIT
-//------------------------------------------------------------------------------------------
-
-float DecodeWorldUnitScale(float EncodedWorldUnitScale)
-{
-	return EncodedWorldUnitScale * DEC_UNIT_TO_WORLDUNITSCALE_IN_CM;
-}
-
-float4 DecodeDiffuseMeanFreePath(float4 EncodedDiffuseMeanFreePath)
-{
-	return EncodedDiffuseMeanFreePath * DEC_UNIT_TO_DIFFUSEMEANFREEPATH_IN_MM;
-}
-
-bool GetSubsurfaceTransmittanceProfileUseBurley(uint SubsurfaceProfileInt)
-{
-	// 0..255, which SubSurface profile to pick
-	float Type = View.SSProfilesTexture.Load(int3(SSSS_BOUNDARY_COLOR_BLEED_OFFSET, SubsurfaceProfileInt, 0)).a;
-
-	return abs(Type - SSS_TYPE_BURLEY) < 0.01f;
-}
-
-bool GetSubsurfaceProfileUseBurley(uint SubsurfaceProfileInt)
-{
-	// 0..255, which SubSurface profile to pick
-	float Type = View.SSProfilesTexture.Load(int3(SSSS_BOUNDARY_COLOR_BLEED_OFFSET, SubsurfaceProfileInt, 0)).a;
-
-	return abs(Type - SSS_TYPE_BURLEY) < 0.01f;
-}
-
-bool GetSubsurfaceProfileUseSeparable(uint SubsurfaceProfileInt)
-{
-	// 0..255, which SubSurface profile to pick
-	float Type = View.SSProfilesTexture.Load(int3(SSSS_BOUNDARY_COLOR_BLEED_OFFSET, SubsurfaceProfileInt, 0)).a;
-
-	return abs(Type - SSS_TYPE_SSSS) < 0.01f;
-}
-
 inline float3 Burley_Profile(float Radius, float3 Albedo, float3 S3D, float L)
 {   //R(r)r
 	float3 D = 1 / S3D;
@@ -248,20 +191,4 @@ struct FBurleyParameter
 	float4 DiffuseMeanFreePath;
 	float  WorldUnitScale;
 	float  SurfaceOpacity; // 1.0 means full Burley, 0.0 means full Default Lit
-};
-
-//Surface albedo and mean free path length
-float4 GetSubsurfaceProfileSurfaceAlbedo(uint SubsurfaceProfileInt)
-{
-	return View.SSProfilesTexture.Load(int3(BSSS_SURFACEALBEDO_OFFSET, SubsurfaceProfileInt, 0));
-}
-
-float4 GetSubsurfaceProfileDiffuseMeanFreePath(uint SubsurfaceProfileInt)
-{
-	return View.SSProfilesTexture.Load(int3(BSSS_DMFP_OFFSET, SubsurfaceProfileInt, 0));
-}
-
-float GetSubsurfaceProfileWorldUnitScale(uint SubsurfaceProfileInt)
-{
-	return View.SSProfilesTexture.Load(int3(SSSS_TINT_SCALE_OFFSET, SubsurfaceProfileInt, 0)).a;
-}
+};
\ No newline at end of file
diff --git a/Engine/Shaders/Private/SubsurfaceProfileCommon.ush b/Engine/Shaders/Private/SubsurfaceProfileCommon.ush
--- a/Engine/Shaders/Private/SubsurfaceProfileCommon.ush
+++ b/Engine/Shaders/Private/SubsurfaceProfileCommon.ush
@@ -6,6 +6,8 @@
 
 #pragma once
 
+///////////////////////////////////////////////////////////////////////////////////////////////
+// Describe SSS profile texture (data offsets & size)
 // NOTE: Changing offsets below requires updating all instances of #SSSS_CONSTANTS
 // TODO: This needs to be defined in a single place and shared between C++ and shaders!
 #define SSSS_TINT_SCALE_OFFSET					0
@@ -29,6 +31,31 @@
 // Threshold value at which model switches from SSS to default lit
 #define SSSS_OPACITY_THRESHOLD_EPS 0.10
 
+///////////////////////////////////////////////////////////////////////////////////////////////
+// Burley constants
+#define BURLEY_MM_2_CM		0.1f
+#define BURLEY_CM_2_MM      10.0f
+
+// exactly one of these should be true
+#define LIGHT_PERPENDICULAR 0
+#define LIGHT_DIFFUSESURFACE 0
+#define LIGHT_PERPENDICULAR_DIFFUSE_SURFACE 1
+
+// Consistent in BurleyNormalizedSSSCommon.ush and SubsurfaceProfile.cpp
+#define SSS_TYPE_BURLEY	    0
+#define SSS_TYPE_SSSS		1
+
+// Make sure UIMax|ClampMax of WorldUNitScale * ENC_WORLDUNITSCALE_IN_CM_TO_UNIT <= 1
+#define ENC_WORLDUNITSCALE_IN_CM_TO_UNIT 0.02f
+#define DEC_UNIT_TO_WORLDUNITSCALE_IN_CM 1/ENC_WORLDUNITSCALE_IN_CM_TO_UNIT
+
+// Make sure UIMax|ClampMax of DiffuseMeanFreePath * ENC_DIFFUSEMEANFREEPATH_IN_MM_TO_UNIT <= 1
+#define ENC_DIFFUSEMEANFREEPATH_IN_MM_TO_UNIT (0.01f*0.2f)
+#define DEC_UNIT_TO_DIFFUSEMEANFREEPATH_IN_MM 1/ENC_DIFFUSEMEANFREEPATH_IN_MM_TO_UNIT
+
+///////////////////////////////////////////////////////////////////////////////////////////////
+// Accessors
+
 half4 GetSubsurfaceProfileTexture(uint SampleIndex, uint SubsurfaceProfileInt)
 {
 	// One profile per line, encoded using constants above. See FSubsurfaceProfileTexture::CreateTexture() in SubsurfaceProfile.cpp.
@@ -51,11 +78,7 @@ half4 GetSubsurfaceProfileTexture(Texture2D InSSProfilesTexture, SamplerState In
 
 void GetSubsurfaceProfileDualSpecular(uint SubsurfaceProfileInt, float Opacity, inout float AverageToRoughness0, inout float AverageToRoughness1, inout float LobeMix, inout float MaterialRoughnessToAverage)
 {
-#if COMPILER_GLSL_ES3_1 // Force to use a point sampler for Texture2D.Load on OpenGLES platform
-	float4 Data = View.SSProfilesTexture.SampleLevel(View.SSProfilesSampler, (uint2(SSSS_DUAL_SPECULAR_OFFSET, SubsurfaceProfileInt) + float2(0.5f, 0.5f)) * View.SSProfilesTextureSizeAndInvSize.zw, 0);
-#else
-	float4 Data = View.SSProfilesTexture.Load(int3(SSSS_DUAL_SPECULAR_OFFSET, SubsurfaceProfileInt, 0));
-#endif 
+	const half4 Data = GetSubsurfaceProfileTexture(SSSS_DUAL_SPECULAR_OFFSET, SubsurfaceProfileInt);
 	AverageToRoughness0 = Data.x * SSSS_MAX_DUAL_SPECULAR_ROUGHNESS;
 	AverageToRoughness1 = Data.y * SSSS_MAX_DUAL_SPECULAR_ROUGHNESS;
 	LobeMix = Data.z;
@@ -71,4 +94,55 @@ void GetSubsurfaceProfileRoughnesses(inout float Lobe0Roughness, inout float Lob
 	// See SUBSURFACE_PROFILE_OPACITY_THRESHOLD in ShadingModelsMaterial.ush.
 	Lobe0Roughness = lerp(1.0f, Lobe0Roughness, saturate(Opacity * 10.0f));
 	Lobe1Roughness = lerp(1.0f, Lobe1Roughness, saturate(Opacity * 10.0f));
+}
+
+// Surface albedo and mean free path length
+float4 GetSubsurfaceProfileSurfaceAlbedo(uint SubsurfaceProfileInt)
+{
+	return GetSubsurfaceProfileTexture(BSSS_SURFACEALBEDO_OFFSET, SubsurfaceProfileInt);
+}
+
+float4 GetSubsurfaceProfileDiffuseMeanFreePath(uint SubsurfaceProfileInt)
+{
+	return GetSubsurfaceProfileTexture(BSSS_DMFP_OFFSET, SubsurfaceProfileInt);
+}
+
+float GetSubsurfaceProfileWorldUnitScale(uint SubsurfaceProfileInt)
+{
+	return GetSubsurfaceProfileTexture(SSSS_TINT_SCALE_OFFSET, SubsurfaceProfileInt).a;
+}
+
+float DecodeWorldUnitScale(float EncodedWorldUnitScale)
+{
+	return EncodedWorldUnitScale * DEC_UNIT_TO_WORLDUNITSCALE_IN_CM;
+}
+
+float4 DecodeDiffuseMeanFreePath(float4 EncodedDiffuseMeanFreePath)
+{
+	return EncodedDiffuseMeanFreePath * DEC_UNIT_TO_DIFFUSEMEANFREEPATH_IN_MM;
+}
+
+bool GetSubsurfaceTransmittanceProfileUseBurley(uint SubsurfaceProfileInt)
+{
+	half Type = GetSubsurfaceProfileTexture(SSSS_BOUNDARY_COLOR_BLEED_OFFSET, SubsurfaceProfileInt).a;
+	return abs(Type - SSS_TYPE_BURLEY) < 0.01f;
+}
+
+bool GetSubsurfaceProfileUseBurley(uint SubsurfaceProfileInt)
+{
+	half Type = GetSubsurfaceProfileTexture(SSSS_BOUNDARY_COLOR_BLEED_OFFSET, SubsurfaceProfileInt).a;
+	return abs(Type - SSS_TYPE_BURLEY) < 0.01f;
+}
+
+bool GetSubsurfaceProfileUseSeparable(uint SubsurfaceProfileInt)
+{
+	half Type = GetSubsurfaceProfileTexture(SSSS_BOUNDARY_COLOR_BLEED_OFFSET, SubsurfaceProfileInt).a;
+	return abs(Type - SSS_TYPE_SSSS) < 0.01f;
+}
+
+float4 GetSubsurfaceProfileDMFPInCm(int SubsurfaceProfileInt)
+{
+	const float4 DiffuseMeanFreePath = DecodeDiffuseMeanFreePath(GetSubsurfaceProfileDiffuseMeanFreePath(SubsurfaceProfileInt));
+	const float WorldUnitScale = DecodeWorldUnitScale(GetSubsurfaceProfileWorldUnitScale(SubsurfaceProfileInt));
+	return DiffuseMeanFreePath * WorldUnitScale; // In cm
 }
\ No newline at end of file
