Add pixel footprint computation for strata material, and remove SSS contribution when MFP is smaller than pixel size.

diff --git a/Engine/Shaders/Private/DebugProbes.usf b/Engine/Shaders/Private/DebugProbes.usf
--- a/Engine/Shaders/Private/DebugProbes.usf
+++ b/Engine/Shaders/Private/DebugProbes.usf
@@ -255,12 +255,13 @@ void MainPS(
 		FStrataPixelHeader StrataPixelHeader = InitialiseStrataPixelHeader();
 
 		// STRATA_TODO: produce the tree
-
+		FStrataPixelFootprint PixelFootprint = (FStrataPixelFootprint)0;
 		uint SharedLocalBasisIndex = 0;
 		StrataPixelHeader.SharedLocalBases.Normals[SharedLocalBasisIndex] = normalize(Normal);
 		StrataPixelHeader.SharedLocalBases.Count = 1;
 
 		FStrataData StrataData = GetStrataSlabBSDF(
+			/**FStrataPixelFootprint PixelFootprint */PixelFootprint,
 			/*float UseMetalness					*/true,
 			/*float3 BaseColor						*/BaseColor,
 			/*float3 EdgeColor						*/1.0f,
diff --git a/Engine/Shaders/Private/MaterialTemplate.ush b/Engine/Shaders/Private/MaterialTemplate.ush
--- a/Engine/Shaders/Private/MaterialTemplate.ush
+++ b/Engine/Shaders/Private/MaterialTemplate.ush
@@ -447,6 +447,7 @@ struct FMaterialPixelParameters
 #if TEMPLATE_USES_STRATA
 	FSharedLocalBases SharedLocalBases;
 	FStrataTree StrataTree;
+	FStrataPixelFootprint StrataPixelFootprint;
 #endif
 };
 
@@ -2880,6 +2881,7 @@ void CalcPixelMaterialInputs(in out FMaterialPixelParameters Parameters, in out
 %s
 
 #if TEMPLATE_USES_STRATA
+	Parameters.StrataPixelFootprint = StrataGetPixelFootprint(Parameters.WorldPosition_CamRelative);
 	Parameters.SharedLocalBases = StrataInitialiseSharedLocalBases();
 	Parameters.StrataTree = GetInitialisedStrataTree();
 #endif
@@ -2940,6 +2942,7 @@ void CalcPixelMaterialInputsAnalyticDerivatives(in out FMaterialPixelParameters
 %s
 
 #if TEMPLATE_USES_STRATA
+	Parameters.StrataPixelFootprint = StrataGetPixelFootprint(Parameters.WorldPosition_CamRelative);
 	Parameters.SharedLocalBases = StrataInitialiseSharedLocalBases();
 	Parameters.StrataTree = GetInitialisedStrataTree();
 #endif
diff --git a/Engine/Shaders/Private/Strata/Strata.ush b/Engine/Shaders/Private/Strata/Strata.ush
--- a/Engine/Shaders/Private/Strata/Strata.ush
+++ b/Engine/Shaders/Private/Strata/Strata.ush
@@ -12,6 +12,7 @@
 #include "../MortonCode.ush"
 #include "../OctahedralCommon.ush"
 #include "../SubsurfaceProfileCommon.ush"
+#include "../BurleyNormalizedSSSCommon.ush"
 #include "/Engine/Shared/StrataDefinitions.h"     
 
 
@@ -646,10 +647,26 @@ FStrataData GetInitialisedStrataData()
 	return StrataData;
 }
 
+struct FStrataPixelFootprint
+{
+	float PixelRadiusInWorldSpace; // In cm
+};
+
+FStrataPixelFootprint StrataGetPixelFootprint(float3 InWorldPosition)
+{
+	const float dX = lengthFast(ddx(InWorldPosition));
+	const float dY = lengthFast(ddy(InWorldPosition));
+
+	FStrataPixelFootprint Out;
+	Out.PixelRadiusInWorldSpace = min(dX, dY) * 0.5f;
+	return Out;
+}
+
 ///////////////////////////////////////////////////////////////////////////////
 // BSDF
 
 FStrataData GetStrataSlabBSDF(
+	FStrataPixelFootprint InPixelFootprint,
 	float UseMetalness,
 	float3 BaseColor, float3 EdgeColor, float Specular, float Metallic,	//  Metalness workflow
 	float3 DiffuseAlbedo, float3 F0, float3 F90,						// !Metalness workflow
@@ -668,9 +685,9 @@ FStrataData GetStrataSlabBSDF(
 	SSSDMFPScale = saturate(SSSDMFPScale);
 	const bool bSSSProfilePlugged = SSSProfileID > 0.0f;
 	const bool bDMFPPlugged = any(SSSDMFP > 0);
-	const bool bIsThin = bDMFPPlugged && Thickness < STRATA_LAYER_ISTHIN_THICKNESS_THRESHOLD_CM;
-	const bool bHasSSSProfile = !bIsThin && bSSSProfilePlugged && SSSDMFPScale > 0.f;	// When no profile is provide, we use Burley. STRATA_TODO we should clarify as burley is also an option of the profile.
-	const bool bHasSSS = bDMFPPlugged || bHasSSSProfile;
+	bool bIsThin = bDMFPPlugged && Thickness < STRATA_LAYER_ISTHIN_THICKNESS_THRESHOLD_CM;
+	bool bHasSSSProfile = !bIsThin && bSSSProfilePlugged && SSSDMFPScale > 0.f;	// When no profile is provide, we use Burley. STRATA_TODO we should clarify as burley is also an option of the profile.
+	bool bHasSSS = bDMFPPlugged || bHasSSSProfile;
 
 	// Derive roughness and haziness from profile data
 	if (bSSSProfilePlugged)
@@ -695,6 +712,20 @@ FStrataData GetStrataSlabBSDF(
 		// * The mix factor is hardoded to 0.5 i.e., 0.5 * Fs(roughness0) + 0.5 * Fs(roughness1)
 		Haziness = (Lobe1Roughness - Lobe0Roughness) / (1.f - min(Lobe0Roughness, 0.99f));
 		Roughness = Lobe0Roughness;
+
+		SSSDMFP = GetSubsurfaceProfileDMFPInCm(SubsurfaceProfileUInt).xyz * SSSDMFPScale;
+	}
+
+	// If MFP is smaller than the pixel's footprint, remove SSS attributes
+	// The scaling factor 'SSSFootprintScale' is based on empiric tests
+	const float SSSFootprintScale = 0.9f;
+	if (all(SSSDMFP * SSSFootprintScale < InPixelFootprint.PixelRadiusInWorldSpace))
+	{
+		bHasSSSProfile = false;
+		bHasSSS = false;
+
+		// DMFPPlugged is not cleared as it is used for thin surface and thin volumetric
+		// bDMFPPlugged = false;
 	}
 
 	const uint bHasHaziness = Haziness > 0.0f ? 1 : 0;
@@ -706,7 +737,6 @@ FStrataData GetStrataSlabBSDF(
 		StrataRequestSharedLocalBasisTangent(SharedLocalBasisTypes, SharedLocalBasisIndex);
 	}
 
-
 	// About the thin, a.k.a. two sided lighting, model:
 	// 	- Thin lighting model is used when thickness becomes less than STRATA_LAYER_ISTHIN_THICKNESS_THRESHOLD_CM
 	//  - It cannot work with a SSS Profile because the profile id is stored in the register used for the MFP (e.g. see SLAB_SSSPROFILEID).
@@ -756,7 +786,7 @@ FStrataData GetStrataSlabBSDF(
 	if (bHasSSSProfile)
 	{
 		SLAB_SSSPROFILEID	(StrataData.InlinedBSDF)		= SSSProfileID;
-		SLAB_SSSPROFILERADIUSSCALE(StrataData.InlinedBSDF)= SSSDMFPScale;
+		SLAB_SSSPROFILERADIUSSCALE(StrataData.InlinedBSDF)	= SSSDMFPScale;
 	}
 	else
 	{
@@ -2234,8 +2264,10 @@ void StrataConvertFromDBuffer(
 
 	// Data
 	// Decal only support a single Slab BSDF with no-anisotropy
+	const FStrataPixelFootprint PixelFootprint = (FStrataPixelFootprint)0;
 	const float DefaultThickness = STRATA_LAYER_DEFAULT_THICKNESS_CM;
 	OutData = GetStrataSlabBSDF(
+		PixelFootprint,		// PixelFootprint
 		true,				// UseMetalness
 		In.BaseColor,		// BaseColor		-  Metalness workflow
 		1.f,				// EdgeColor		-  Metalness workflow
diff --git a/Engine/Shaders/Private/Strata/StrataLegacyConversion.ush b/Engine/Shaders/Private/Strata/StrataLegacyConversion.ush
--- a/Engine/Shaders/Private/Strata/StrataLegacyConversion.ush
+++ b/Engine/Shaders/Private/Strata/StrataLegacyConversion.ush
@@ -6,6 +6,7 @@
 
 struct FStrataLegacyParameters
 {
+	FStrataPixelFootprint PixelFootprint;
 	bool UseMetalness;
 	float3 DiffuseAlbedo;
 	float3 F0;
@@ -31,7 +32,7 @@ struct FStrataLegacyParameters
 	float Weight;
 };
 
-FStrataLegacyParameters InitStrataLegacyParameters(uint InSharedLocalBasisIndex, float InWeight)
+FStrataLegacyParameters InitStrataLegacyParameters(FStrataPixelFootprint InPixelFootprint, uint InSharedLocalBasisIndex, float InWeight)
 {
 	const float3 Zeros = float3(0, 0, 0);
 
@@ -58,6 +59,7 @@ FStrataLegacyParameters InitStrataLegacyParameters(uint InSharedLocalBasisIndex,
 	Out.bIsTop = 1;
 	Out.bIsBottom = 1;
 	Out.Weight = InWeight;
+	Out.PixelFootprint = InPixelFootprint;
 	return Out;
 }
 
@@ -72,6 +74,7 @@ FStrataData CreateLegacySlab(
 	// Fixed layer structure for helping compiler to unroll and optimize shader
 	return PromoteParameterBlendedBSDFToOperator(
 		GetStrataSlabBSDF(
+			In.PixelFootprint,
 			In.UseMetalness,	// UseMetalness
 			In.BaseColor,		// BaseColor		-  Metalness workflow
 			Ones,				// EdgeColor		-  Metalness workflow
@@ -105,6 +108,7 @@ FStrataData CreateLegacySlab(
 // helping the compiler to unroll Strata BSDF traversing and packing. In most cases, the bottom slab is 
 // weighted by 0 and it will be removed once the data are written-out/packed.
 FStrataData StrataConvertLegacyMaterialDynamic(
+	FStrataPixelFootprint PixelFootprint,
 	float3 BaseColor, float Specular, float Metallic,
 	float Roughness, float Anisotropy,
 	float3 SubSurfaceColor, float SubSurfaceProfileId,
@@ -165,8 +169,8 @@ FStrataData StrataConvertLegacyMaterialDynamic(
 	#endif
 
 	// Top/Bottom
-	FStrataLegacyParameters Top = InitStrataLegacyParameters(SharedLocalBasisIndex, 1.0f);
-	FStrataLegacyParameters Bot = InitStrataLegacyParameters(SharedLocalBasisIndex, 1.0f);
+	FStrataLegacyParameters Top = InitStrataLegacyParameters(PixelFootprint, SharedLocalBasisIndex, 1.0f);
+	FStrataLegacyParameters Bot = InitStrataLegacyParameters(PixelFootprint, SharedLocalBasisIndex, 1.0f);
 
 	float FinalWeight = 1.0f;
 
@@ -468,6 +472,7 @@ FStrataData StrataConvertLegacyMaterialDynamic(
 // It reuses the dynamic version for most part, but for special node like Unlit/Hair/Water, 
 // we use the dedicated node
 FStrataData StrataConvertLegacyMaterialStatic(
+	FStrataPixelFootprint PixelFootprint,
 	float3 BaseColor, float Specular, float Metallic,
 	float Roughness, float Anisotropy,
 	float3 SubSurfaceColor, float SubSurfaceProfileId,
@@ -529,6 +534,7 @@ FStrataData StrataConvertLegacyMaterialStatic(
 	else
 	{
 		Out = StrataConvertLegacyMaterialDynamic(
+			PixelFootprint,
 			BaseColor, Specular, Metallic,
 			Roughness, Anisotropy,
 			SubSurfaceColor, SubSurfaceProfileId,
diff --git a/Engine/Shaders/Private/Strata/StrataVisualize.usf b/Engine/Shaders/Private/Strata/StrataVisualize.usf
--- a/Engine/Shaders/Private/Strata/StrataVisualize.usf
+++ b/Engine/Shaders/Private/Strata/StrataVisualize.usf
@@ -71,6 +71,49 @@ float3 DrawPlot(
 	return Out;
 }
 
+void AddDrawPixelFootprint(float3 P, float3 dPdx, float3 dPdy, float2 Scale, bool bNormalize, float4 Color)
+{
+	const float3 T = (bNormalize ? normalize(dPdx) : dPdx) * Scale.x;
+	const float3 B = (bNormalize ? normalize(dPdy) : dPdy) * Scale.y;
+	const float3 N = normalize(cross(T, B));
+
+	const float3 WP0 = P - T - B;
+	const float3 WP1 = P + T - B;
+	const float3 WP2 = P + T + B;
+	const float3 WP3 = P - T + B;
+
+	AddLineTWS(WP0, WP1, Color);
+	AddLineTWS(WP1, WP2, Color);
+	AddLineTWS(WP2, WP3, Color);
+	AddLineTWS(WP3, WP0, Color);
+}
+
+void DrawPixelFootprint(float3 P, float3 dPdx, float3 dPdy, uint2 PixelCoord)
+{
+	const FStrataSubsurfaceData SSSData = StrataUnpackSSSData(Strata.SSSTexture.Load(uint3(PixelCoord, 0)));
+	if (SSSData.bIsValid)
+	{
+		float3 DMFP = 0;
+		if (SSSData.bIsProfile)
+		{
+			DMFP = GetSubsurfaceProfileDMFPInCm(SSSData.ProfileId).xyz * SSSData.ProfileRadiusScale;
+		}
+		else
+		{
+			DMFP = SSSData.DMFP;
+		}
+
+		FStrataPixelFootprint Footprint = StrataGetPixelFootprint(P);
+		AddDrawPixelFootprint(P, dPdx, dPdy, 0.5f, false, ColorRed);
+		AddDrawPixelFootprint(P, dPdx, dPdy, Footprint.PixelRadiusInWorldSpace, true, ColorOrange);
+		AddDrawPixelFootprint(P, dPdx, dPdy, max3(DMFP.x, DMFP.y, DMFP.z), true, ColorCyan);
+
+		// Normal
+		const float3 N = normalize(cross(normalize(dPdx), normalize(dPdy)));
+		AddLineTWS(P, P + N * 10.f, ColorBlue);
+	}
+}
+
 void VisualizeMaterialPS(
 	float4 SVPos : SV_POSITION,
 	out float4 OutColor : SV_Target0)
@@ -80,10 +123,6 @@ void VisualizeMaterialPS(
 	float2 BufferUV = SvPositionToBufferUV(float4(MaterialPos, SVPos.zw));
 	const int  InspectSquare = 10;
 
-	FStrataAddressing StrataAddressing = GetStrataPixelDataByteOffset(MaterialPos, uint2(View.BufferSizeAndInvSize.xy), Strata.MaxBytesPerPixel);
-	FStrataPixelHeader Header = UnpackStrataHeaderIn(Strata.MaterialTextureArray, StrataAddressing, Strata.TopLayerTexture);
-	bool bStrataMaterial = Header.BSDFCount > 0;
-
 	const float2 ScreenPosition = SvPositionToScreenPosition(SVPos).xy;
 	const float2 ScreenMaterialPosition = SvPositionToScreenPosition(float4(MaterialPos, 0.5f, 1)).xy;
 	float DeviceZ = SampleDeviceZFromSceneTextures(BufferUV);
@@ -91,6 +130,10 @@ void VisualizeMaterialPS(
 	const float3 WorldPosition = mul(float4(ScreenMaterialPosition * SceneDepth, SceneDepth, 1), PrimaryView.ScreenToTranslatedWorld).xyz;
 	const float3 V = -normalize(WorldPosition - PrimaryView.TranslatedWorldCameraOrigin);
 
+	FStrataAddressing StrataAddressing = GetStrataPixelDataByteOffset(MaterialPos, uint2(View.BufferSizeAndInvSize.xy), Strata.MaxBytesPerPixel);
+	FStrataPixelHeader Header = UnpackStrataHeaderIn(Strata.MaterialTextureArray, StrataAddressing, Strata.TopLayerTexture);
+	bool bStrataMaterial = Header.BSDFCount > 0;
+
 #if 0
 	OutColor = 0;
 	if (Header.BSD)
@@ -136,6 +179,15 @@ void VisualizeMaterialPS(
 	}
 #endif
 
+	// Print pixel footprint & SSS footprint
+	const float3 P = WorldPosition;
+	const float3 dPdx = ddx(P);
+	const float3 dPdy = ddy(P);
+	if (bStrataMaterial && all(PixelPos == MaterialPos))
+	{
+		DrawPixelFootprint(P, dPdx, dPdy, PixelPos);
+	}
+
 	OutColor = float4(0.0f, 0.0f, 0.0f, 1.0f);
 
 	// Plot BSDFs
diff --git a/Engine/Source/Runtime/Engine/Private/Materials/HLSLMaterialTranslator.cpp b/Engine/Source/Runtime/Engine/Private/Materials/HLSLMaterialTranslator.cpp
--- a/Engine/Source/Runtime/Engine/Private/Materials/HLSLMaterialTranslator.cpp
+++ b/Engine/Source/Runtime/Engine/Private/Materials/HLSLMaterialTranslator.cpp
@@ -9967,7 +9967,7 @@ int32 FHLSLMaterialTranslator::StrataSlabBSDF(
 		check(PromoteToOperator->Index != INDEX_NONE);
 		check(PromoteToOperator->BSDFIndex != INDEX_NONE);
 		return AddCodeChunk(
-			MCT_Strata, TEXT("PromoteParameterBlendedBSDFToOperator(GetStrataSlabBSDF(%s, %s, %s, %s, %s, %s, %s, %s, %s, %s, %s, %s, %s, %s, %s, %s, %s, %s, %s, %s, Parameters.SharedLocalBases.Types) /* Normal = %s ; Tangent = %s */, Parameters.StrataTree, %u, %u, %u, %u)"),
+			MCT_Strata, TEXT("PromoteParameterBlendedBSDFToOperator(GetStrataSlabBSDF(Parameters.StrataPixelFootprint, %s, %s, %s, %s, %s, %s, %s, %s, %s, %s, %s, %s, %s, %s, %s, %s, %s, %s, %s, %s, Parameters.SharedLocalBases.Types) /* Normal = %s ; Tangent = %s */, Parameters.StrataTree, %u, %u, %u, %u)"),
 			*GetParameterCode(UseMetalness),
 			*GetParameterCode(BaseColor), *GetParameterCode(EdgeColor), *GetParameterCode(Specular), *GetParameterCode(Metallic),
 			*GetParameterCode(DiffuseAlbedo), *GetParameterCode(F0), *GetParameterCode(F90),
@@ -9990,7 +9990,7 @@ int32 FHLSLMaterialTranslator::StrataSlabBSDF(
 	}
 
 	return AddCodeChunk(
-		MCT_Strata, TEXT("GetStrataSlabBSDF(%s, %s, %s, %s, %s, %s, %s, %s, %s, %s, %s, %s, %s, %s, %s, %s, %s, %s, %s, %s, Parameters.SharedLocalBases.Types) /* Normal = %s ; Tangent = %s */"),
+		MCT_Strata, TEXT("GetStrataSlabBSDF(Parameters.StrataPixelFootprint, %s, %s, %s, %s, %s, %s, %s, %s, %s, %s, %s, %s, %s, %s, %s, %s, %s, %s, %s, %s, Parameters.SharedLocalBases.Types) /* Normal = %s ; Tangent = %s */"),
 		* GetParameterCode(UseMetalness),
 		*GetParameterCode(BaseColor),		*GetParameterCode(EdgeColor),	*GetParameterCode(Specular), *GetParameterCode(Metallic),
 		*GetParameterCode(DiffuseAlbedo),	*GetParameterCode(F0),			*GetParameterCode(F90),
@@ -10030,7 +10030,7 @@ int32 FHLSLMaterialTranslator::StrataConversionFromLegacy(
 	const FString ClearCoat_TangentCode = Tangent != INDEX_NONE ? *GetParameterCode(ClearCoat_Tangent) : TEXT("NONE");
 
 	return AddCodeChunk(
-		MCT_Strata, TEXT("StrataConvertLegacyMaterial%s(%s, %s, %s, %s, %s, %s, %s, %s, %s, %s, %s, %s, %s, %s, %s, %s, %s, %s, %s, Parameters.SharedLocalBases.Types, Parameters.StrataTree) /* Normal = %s ; Tangent = %s ; ClearCoat_Normal = %s ; ClearCoat_Tangent = %s */"),
+		MCT_Strata, TEXT("StrataConvertLegacyMaterial%s(Parameters.StrataPixelFootprint, %s, %s, %s, %s, %s, %s, %s, %s, %s, %s, %s, %s, %s, %s, %s, %s, %s, %s, %s, Parameters.SharedLocalBases.Types, Parameters.StrataTree) /* Normal = %s ; Tangent = %s ; ClearCoat_Normal = %s ; ClearCoat_Tangent = %s */"),
 		bHasDynamicShadingModels ? TEXT("Dynamic") : TEXT("Static"),
 		*GetParameterCode(BaseColor), *GetParameterCode(Specular), *GetParameterCode(Metallic),
 		*GetParameterCode(Roughness), *GetParameterCode(Anisotropy),
