* Change haziness parameterization to be more perceptually linear.
* Fix haziness terminator issue

diff --git a/Engine/Shaders/Private/Strata/Strata.ush b/Engine/Shaders/Private/Strata/Strata.ush
--- a/Engine/Shaders/Private/Strata/Strata.ush
+++ b/Engine/Shaders/Private/Strata/Strata.ush
@@ -589,6 +589,56 @@ bool StrataHasGreyScaleWeight(in float3 Weight)
 	return Weight.x == Weight.y && Weight.y == Weight.z;
 }
 
+void StrataComputeHazeRoughnessAndWeight(
+	float Roughness0,
+	float Haziness,
+	inout float Roughness1,
+	inout float MixLobe)
+{	
+	// Fade out Haziness as the roughness reaches 0.5f, as otherwise it create
+	//       27 |\
+	// Extent   | \
+	//          |  \
+	//        3 |    ----
+	//          0  0.25  1
+	//           Rougness
+
+	// Lerp factor for adapting haziness parameters
+	const float S2 = saturate(2 * Roughness0);
+	const float S4 = saturate(4 * Roughness0);
+
+	// For low roughness, the haze extent is larger, and the core highlight is lower. 
+	// This allows to emphase the haze effect at low roughness, and fade out haze as roughness increase.
+	const float HazeExtent = lerp(27.0f, 3.0f, S4);
+	const float HazeCore   = lerp(0.01f, 0.5f, S2);
+
+	// Reference "A Composite BRDF Model for Hazy Gloss"
+	 
+	// Compute wide roughness
+	float alpha_n = Square(Roughness0); 
+	float lambda_h = HazeExtent;
+	const float alpha_w = alpha_n * (1 + lambda_h); // Eq.8
+	Roughness1 = sqrtFast(alpha_w);
+
+	// Compute mix weights
+	const float beta_h = Haziness;
+	const float r_c = HazeCore;
+	const float w = alpha_w;
+
+	const float p = 1.0 / Square(1 + lambda_h); // Table 2, GGX
+	const float b = 2 * (r_c * (1 - w) + w * p);
+	const float D = Square(b) - 4 * (b - 1) * r_c;
+	const float u = (b - sqrtFast(D)) / (2 * b - 2); // Eq. 12, 13, 14
+	const float u2 = u * u;
+	const float w2 = 2 * w;
+	const float k_h = ((u - u2) * w2 * beta_h) / (Square(1 - u) + (u - u2) * w2 + u2); // Eq. 15
+	const float r = r_c + (1 - p) * k_h; // Eq. 5
+	MixLobe = k_h / r; // Eq. 4
+}
+
+float StrataComputeHazeRoughness(float Roughness0) { float Roughness1=0; float MixLobe=0; StrataComputeHazeRoughnessAndWeight(Roughness0, 1.0f, Roughness1, MixLobe); return Roughness1; }
+float StrataComputeHazeWeight(float Roughness0, float Haziness)    { float Roughness1=0; float MixLobe=0; StrataComputeHazeRoughnessAndWeight(Roughness0, Haziness, Roughness1, MixLobe); return MixLobe; }
+
 ///////////////////////////////////////////////////////////////////////////////
 // Material & BSDF description
 
@@ -715,6 +765,13 @@ FStrataData GetStrataSlabBSDF(
 
 		SSSDMFP = GetSubsurfaceProfileDMFPInCm(SubsurfaceProfileUInt).xyz * SSSDMFPScale;
 	}
+	else
+	{
+		// Convert Haziness into a mix weight
+		// Haze roughness is inexpensive to compute, unlike mix factor. This is why we comput its value here, and 
+		// store it in place of the original haziness value. This saves computation time at runtime.
+		Haziness = Haziness > 0 ? StrataComputeHazeWeight(Roughness, Haziness) : 0.f;
+	}
 
 	// If MFP is smaller than the pixel's footprint, remove SSS attributes
 	// The scaling factor 'SSSFootprintScale' is based on empiric tests
diff --git a/Engine/Shaders/Private/Strata/StrataEvaluation.ush b/Engine/Shaders/Private/Strata/StrataEvaluation.ush
--- a/Engine/Shaders/Private/Strata/StrataEvaluation.ush
+++ b/Engine/Shaders/Private/Strata/StrataEvaluation.ush
@@ -248,6 +248,7 @@ FStrataEvaluateResult StrataEvaluateBSDFCommon(FStrataBSDFContext BSDFContext, F
 			float3 MSScale = 1;
 			float3 HazeMSScale = 1;
 			float3 DirectionalAlbedo_SpecularTransmission = 1.0f;
+			float HazeWeight = 0.f;
 			{
 				{
 					FBxDFEnergyTerms EnergyTerms = ComputeGGXSpecEnergyTerms(SafeRoughness, BSDFContext.Context.NoV, F0, F90);
@@ -258,7 +259,9 @@ FStrataEvaluateResult StrataEvaluateBSDFCommon(FStrataBSDFContext BSDFContext, F
 			#if STRATA_FASTPATH==0
 				if (bHaziness)
 				{
-					HazeSafeRoughness = ComputeHazyLobeRoughness(SafeRoughness, SLAB_HAZINESS(BSDFContext.BSDF));
+					HazeWeight = SLAB_HAZINESS(BSDFContext.BSDF);
+					HazeSafeRoughness = MakeRoughnessSafe(StrataComputeHazeRoughness(SafeRoughness));
+
 					FBxDFEnergyTerms EnergyTerms = ComputeGGXSpecEnergyTerms(HazeSafeRoughness, BSDFContext.Context.NoV, F0, F90);
 					HazeMSScale = ComputeEnergyConservation(EnergyTerms);
 				}
@@ -363,7 +366,7 @@ FStrataEvaluateResult StrataEvaluateBSDFCommon(FStrataBSDFContext BSDFContext, F
 						}
 						else
 						{
-							float Energy = EnergyNormalization(Alpha2Spec, BSDFContext.Context.VoH, AreaLightContext.AreaLight);
+							const float Energy = EnergyNormalization(Alpha2Spec, BSDFContext.Context.VoH, AreaLightContext.AreaLight);
 							D = D_GGX(Alpha2Spec, BSDFContext.Context.NoH) * Energy;
 						}
 						// STRATA_TODO Should we use Vis_SmithJointApprox here?
@@ -378,12 +381,12 @@ FStrataEvaluateResult StrataEvaluateBSDFCommon(FStrataBSDFContext BSDFContext, F
 						float Alpha2SpecHaze = Pow4(HazeSafeRoughness);
 						if (IntegrationType == INTEGRATION_PUNCTUAL_LIGHT)
 						{
-							D = D_GGX(Alpha2SpecHaze, BSDFContext.Context.NoH);
+							HazeD = D_GGX(Alpha2SpecHaze, BSDFContext.Context.NoH);
 						}
 						else
 						{
-							float Energy = EnergyNormalization(Alpha2SpecHaze, BSDFContext.Context.VoH, AreaLightContext.AreaLight);
-							D = D_GGX(Alpha2SpecHaze, BSDFContext.Context.NoH) * Energy;
+							const float Energy = EnergyNormalization(Alpha2SpecHaze, BSDFContext.Context.VoH, AreaLightContext.AreaLight);
+							HazeD = D_GGX(Alpha2SpecHaze, BSDFContext.Context.NoH) * Energy;
 						}
 						// STRATA_TODO Should we use Vis_SmithJointApprox here?
 						HazeVis = Vis_SmithJoint(Alpha2SpecHaze, BSDFContext.Context.NoV, AreaLightContext.NoL);
@@ -412,8 +415,8 @@ FStrataEvaluateResult StrataEvaluateBSDFCommon(FStrataBSDFContext BSDFContext, F
 				FresnelTerm = SpecularMicroOcclusion * F_Schlick(F0, F90, BSDFContext.Context.VoH);
 			}
 
-			Sample.SpecularPathProbability		= bHaziness	? (1.0f - SPECULAR_HAZINESS_WEIGHT) : 1.0f;
-			Sample.SpecularHazePathProbability	= bHaziness	? SPECULAR_HAZINESS_WEIGHT			: 0.0f;
+			Sample.SpecularPathProbability		= bHaziness	? (1.0f - HazeWeight) : 1.0f;
+			Sample.SpecularHazePathProbability	= bHaziness	? HazeWeight : 0.0f;
 
 			Sample.SpecularPDF				= PDF * Sample.SpecularPathProbability;
 			Sample.SpecularHazePDF			= HazePDF * Sample.SpecularHazePathProbability;
@@ -863,21 +866,20 @@ FStrataEnvLightResult StrataEvaluateForEnvLight(FStrataBSDFContext BSDFContext,
 				const bool bHasHaziness = BSDF_GETHASHAZINESS(BSDFContext.BSDF);
 				if (bHasHaziness)
 				{
+					const float SafeRoughnessHaze = MakeRoughnessSafe(StrataComputeHazeRoughness(SafeRoughness));
+					const float HazeWeight = SLAB_HAZINESS(BSDFContext.BSDF);
+
 					// Smoothly fade in haziness while fading out SSR to make sure we do not add energy and avoid popping.
-					const float Haziness = SLAB_HAZINESS(BSDFContext.BSDF);
-					const float HazinessFadeIn = saturate(Haziness / 0.1f);
-					StrataEnvLightResult.SSRReduction = lerp(0.0f, SPECULAR_HAZINESS_WEIGHT, HazinessFadeIn);
+					const float HazinessFadeIn = saturate(HazeWeight / 0.1f);
+					StrataEnvLightResult.SSRReduction = lerp(0.0f, HazeWeight, HazinessFadeIn);
 
 					// Apply blend factor on the sharp specular contribution
-					StrataEnvLightResult.SpecularWeight *= lerp(1.0f, (1.0f - SPECULAR_HAZINESS_WEIGHT), HazinessFadeIn);
-
-					// Compute the second specular weight and roughness
-					float SafeRoughnessHaze = ComputeHazyLobeRoughness(SafeRoughness, Haziness);
+					StrataEnvLightResult.SpecularWeight *= lerp(1.0f, (1.0f - HazeWeight), HazinessFadeIn);
 
 					// Compute the directional albedo for hazy specular
 					FBxDFEnergyTerms EnergyTerms = ComputeGGXSpecEnergyTerms(SafeRoughnessHaze, BSDFContext.Context.NoV, F0, F90);
 					const float3 EvalEnvBRDFHaze = ThinFilmFresnelWeight * EnergyTerms.E;
-					StrataEnvLightResult.SpecularHazeWeight = lerp(0.0f, SPECULAR_HAZINESS_WEIGHT, HazinessFadeIn) * EvalEnvBRDFHaze * SpecularMicroOcclusion;
+					StrataEnvLightResult.SpecularHazeWeight = lerp(0.0f, HazeWeight, HazinessFadeIn) * EvalEnvBRDFHaze * SpecularMicroOcclusion;
 					StrataEnvLightResult.SpecularHazeSafeRoughness = SafeRoughnessHaze;
 				}
 				#endif
