Fix for debug lighting on probes to work with Strata

diff --git a/Engine/Shaders/Private/DebugProbes.usf b/Engine/Shaders/Private/DebugProbes.usf
--- a/Engine/Shaders/Private/DebugProbes.usf
+++ b/Engine/Shaders/Private/DebugProbes.usf
@@ -4,10 +4,20 @@
 #include "VelocityCommon.ush"
 
 #if PROJECT_STRATA
+
+// Mention that we are going to use a single slab
+#define STRATA_CLAMPED_BSDF_COUNT 1
+
 #include "SceneTexturesCommon.ush"
 #include "Strata/Strata.ush"
 RWTexture2DArray<uint> MaterialTextureArrayUAV;
 uint MaxBytesPerPixel;
+
+// Default functions for PackStrataOut to function properly
+void UpdateAllBSDFsOperatorCoverageTransmittance(inout FStrataPixelHeader StrataPixelHeader, inout FStrataTree StrataTree, bool bRoughDiffuseEnabled, FStrataAddressing NullStrataAddressing, float3 V, float3 L) {}
+void UpdateAllOperatorsCoverageTransmittance(inout FStrataTree StrataTree) {}
+void UpdateAllBSDFWeightAfterOperatorVisit(inout FStrataTree StrataTree) {}
+
 #endif
 int DebugProbesMode;
 
@@ -55,26 +65,23 @@ void MainPS(
 	  in  float4 SvPosition		: SV_Position
 	, out float4 OutColor		: SV_Target0
 #if PERMUTATION_PASS == RENDER_BASEPASS
+	#if PROJECT_STRATA
+	, out float4 OutGBufferA : SV_Target1
+	, out uint OutGBufferB : SV_Target2
+	, out uint OutGBufferC : SV_Target3
+	, out uint OutGBufferD : SV_Target4
+	, out uint OutGBufferE : SV_Target5
+	#else
 	, out float4 OutGBufferA	: SV_Target1
 	, out float4 OutGBufferB	: SV_Target2
 	, out float4 OutGBufferC	: SV_Target3
 	, out float4 OutGBufferD	: SV_Target4
 	, out float4 OutGBufferE	: SV_Target5
+	#endif
 #endif
 	, out float  OutDepth		: SV_DEPTH
 )
 {
-	OutColor = 0;
-#if PERMUTATION_PASS == RENDER_BASEPASS
-	OutGBufferA = 0;
-	OutGBufferB = 0;
-	OutGBufferC = 0;
-	OutGBufferD = 0;
-	OutGBufferE = 0;
-#endif
-	OutDepth = 0;
-
-#if 0
 	ResolvedView = ResolveView();
 	const uint2 PixelCoord = uint2(SvPosition.xy);
 
@@ -253,14 +260,16 @@ void MainPS(
 	if (DepthBufferDeviceZ < (DeviceZ+0.000001))
 	{
 		FStrataPixelHeader StrataPixelHeader = InitialiseStrataPixelHeader();
+		StrataPixelHeader.StrataTree = GetInitialisedStrataTree();
+		StrataPixelHeader.MaterialAO = 1.0f - IndirectIrradianceAndAO;
+		SetCastContactShadow(StrataPixelHeader, 0);
 
-		// STRATA_TODO: produce the tree
-		FStrataPixelFootprint PixelFootprint = (FStrataPixelFootprint)0;
 		uint SharedLocalBasisIndex = 0;
 		StrataPixelHeader.SharedLocalBases.Normals[SharedLocalBasisIndex] = normalize(Normal);
 		StrataPixelHeader.SharedLocalBases.Count = 1;
 
-		FStrataData StrataData = GetStrataSlabBSDF(
+		FStrataPixelFootprint PixelFootprint = (FStrataPixelFootprint)0;
+		FStrataData StrataSlabData = GetStrataSlabBSDF(
 			/**FStrataPixelFootprint PixelFootprint */PixelFootprint,
 			/*float UseMetalness					*/true,
 			/*float3 BaseColor						*/BaseColor,
@@ -284,8 +293,8 @@ void MainPS(
 			/*uint SharedLocalBasisIndex			*/SharedLocalBasisIndex,
 			/*inout uint SharedLocalBasisTypes		*/StrataPixelHeader.SharedLocalBases.Types);
 
-		StrataPixelHeader.MaterialAO = 1.0f - IndirectIrradianceAndAO;
-		StrataPixelHeader.BSDFCount = 1;
+		FStrataData StrataData = PromoteParameterBlendedBSDFToOperator(StrataSlabData, StrataPixelHeader.StrataTree, /*int OperatorIndex*/0, /*int BSDFIndex*/0, /*int bIsBottom*/1, /*int bIsTop*/1);
+		StrataPixelHeader.StrataTree.BSDFs[0].LuminanceWeight = 1.0f;
 
 		FStrataSubsurfaceData SSSData = (FStrataSubsurfaceData)0;
 		FStrataTopLayerData TopLayerData = (FStrataTopLayerData)0;
@@ -302,7 +311,19 @@ void MainPS(
 			StrataAddressing,
 			StrataPixelHeader, StrataData, -RayDir, bStrataSubsurfaceEnable, EmissiveLuminance,
 			SSSData, TopLayerData);
-		// TODO write out the two MRT uint?
+
+		// Write out MRTs. Note: this only works for deferred with basepassvelocity off
+		OutGBufferA = 0;
+		OutGBufferB = RWStrataMaterialContainer.MaterialRenderTargets[0];
+		OutGBufferC = RWStrataMaterialContainer.MaterialRenderTargets[1];
+		OutGBufferD = StrataPackTopLayerData(TopLayerData);
+		OutGBufferE = 0;
+	}
+	else
+	{
+		// If we are behind the front depth buffer, we must also discard to avoid writting to output MRTs.
+		discard;
+		return;
 	}
 	
 #else // PROJECT_STRATA
@@ -336,5 +357,4 @@ void MainPS(
 	OutDepth = DeviceZ;
 #endif
 
-#endif // TODO
 }
\ No newline at end of file
diff --git a/Engine/Source/Runtime/Renderer/Private/DebugProbeRendering.cpp b/Engine/Source/Runtime/Renderer/Private/DebugProbeRendering.cpp
--- a/Engine/Source/Runtime/Renderer/Private/DebugProbeRendering.cpp
+++ b/Engine/Source/Runtime/Renderer/Private/DebugProbeRendering.cpp
@@ -8,6 +8,9 @@
 #include "SceneTextureParameters.h"
 #include "Strata/Strata.h"
 
+//
+//	Deferred probes are only stamped in deferred mode.
+//
 
 // Changing this causes a full shader recompile
 static TAutoConsoleVariable<int32> CVarVisualizeLightingOnProbes(
@@ -76,7 +79,7 @@ static void CommonStampDeferredDebugProbeDrawCall(
 	int32 RenderPass)
 {
 	PassParameters->ViewUniformBuffer = View.ViewUniformBuffer;
-	PassParameters->MaterialTextureArrayUAV = View.StrataSceneData->MaterialTextureArrayUAV;
+	PassParameters->MaterialTextureArrayUAV = View.StrataSceneData->MaterialTextureArrayUAVWithoutRTs;
 	PassParameters->MaxBytesPerPixel = View.StrataSceneData->MaxBytesPerPixel;
 	PassParameters->DebugProbesMode = View.Family->EngineShowFlags.VisualizeLightingOnProbes ? 3 : FMath::Clamp(CVarVisualizeLightingOnProbes.GetValueOnRenderThread(), 0, 3);
 		
