Brought back to life Strata Decals.

diff --git a/Engine/Shaders/Private/DBufferDecalShared.ush b/Engine/Shaders/Private/DBufferDecalShared.ush
--- a/Engine/Shaders/Private/DBufferDecalShared.ush
+++ b/Engine/Shaders/Private/DBufferDecalShared.ush
@@ -129,122 +129,121 @@ void ApplyDBufferData(in FStrataDBuffer StrataDBufferData, inout FStrataPixelHea
 	#endif
 
 	// For now apply decal data onto all layers, without consideration of top layers
-	STRATA_UNROLL_N(STRATA_CLAMPED_LAYER_COUNT)
-	for (int l = 0; l < StrataData.LayerCount; ++l)
-	{			
-		STRATA_UNROLL_N(STRATA_CLAMPED_MAX_BSDF_COUNT_PER_LAYER)
-		for (int i = 0; i < StrataData.Layers[l].BSDFCount; ++i)
-		{
+	STRATA_UNROLL_N(STRATA_CLAMPED_BSDF_COUNT)
+	for (int BSDFIdx = 0; BSDFIdx < StrataPixelHeader.StrataTree.BSDFCount; ++BSDFIdx)
+	{
+	#define BSDF StrataPixelHeader.StrataTree.BSDFs[BSDFIdx]
 
-			switch (BSDF_GETTYPE(StrataData.Layers[l].BSDFs[i]))
-			{
-				case STRATA_BSDF_TYPE_SLAB:
-				{	
-					// Note:
-					// * For Slab, since the DBuffer is stored with a Metalness parameterization for storage-space reason, 
-					//   we convert back&forth the Slab data between DiffuseAlbedo/F0 and Metalness. This induces some shift 
-					//   since the conversion is lossy. In the future, we could change the DBuffer storage layout (would require 
-					//   more space) or do all the blending in DiffuseAlbedo/F0 space
+		switch (BSDF_GETTYPE(BSDF))
+		{
+			case STRATA_BSDF_TYPE_SLAB:
+			{	
+				// Note:
+				// * For Slab, since the DBuffer is stored with a Metalness parameterization for storage-space reason, 
+				//   we convert back&forth the Slab data between DiffuseAlbedo/F0 and Metalness. This induces some shift 
+				//   since the conversion is lossy. In the future, we could change the DBuffer storage layout (would require 
+				//   more space) or do all the blending in DiffuseAlbedo/F0 space
 					
-					// Apply DBuffer as per-component parameter blending (legacy mode)
-					#if STRATA_DBUFFER_BLENDING_MODE == STRATA_DBUFFER_LEGACY
-					{
-						float3 BaseColor = 0;
-						float Metallic = 0;
-						float Specular = 0;
-						ConvertToMetalness(SLAB_DIFFUSEALBEDO(StrataData.Layers[l].BSDFs[i]), SLAB_F0(StrataData.Layers[l].BSDFs[i]), BaseColor, Specular, Metallic);
-						float Roughness = SLAB_ROUGHNESS(StrataData.Layers[l].BSDFs[i]);
-
-					#if STRATA_DBUFFER_RESPONSE_BASECOLOR
-						BaseColor = BaseColor * StrataDBufferData.OneMinusCoverage_BaseColor + StrataDBufferData.BaseColor;
-						if (BSDF_GETHASFUZZ(StrataData.Layers[l].BSDFs[i]))
-						{
-							SLAB_FUZZ_COLOR(StrataData.Layers[l].BSDFs[i]) = SLAB_FUZZ_COLOR(StrataData.Layers[l].BSDFs[i]) * StrataDBufferData.OneMinusCoverage_BaseColor + StrataDBufferData.BaseColor;
-						}
-
-						// Disable SSS when the coverage reach one to preserve the decal color
-						SLAB_SSSDMFP(StrataData.Layers[l].BSDFs[i]) = lerp(0.f, SLAB_SSSDMFP(StrataData.Layers[l].BSDFs[i]), StrataDBufferData.OneMinusCoverage_BaseColor);
-						if (all(SLAB_SSSDMFP(StrataData.Layers[l].BSDFs[i]) == 0.0f))
-						{
-							BSDF_SETHASDMFP(StrataData.Layers[l].BSDFs[i], 0);
-							BSDF_SETHASSSS(StrataData.Layers[l].BSDFs[i], 0);
-						}
-					#endif
-
-					#if STRATA_DBUFFER_RESPONSE_ROUGHNESS
-						Metallic  = Metallic  * StrataDBufferData.OneMinusCoverage_Roughness + StrataDBufferData.Metallic;
-						Specular  = Specular  * StrataDBufferData.OneMinusCoverage_Roughness + StrataDBufferData.Specular;
-						Roughness = Roughness * StrataDBufferData.OneMinusCoverage_Roughness + StrataDBufferData.Roughness;
-					#endif
+				// Apply DBuffer as per-component parameter blending (legacy mode)
+				#if STRATA_DBUFFER_BLENDING_MODE == STRATA_DBUFFER_LEGACY
+				{
+					float3 BaseColor = 0;
+					float Metallic = 0;
+					float Specular = 0;
+					ConvertToMetalness(SLAB_DIFFUSEALBEDO(BSDF), SLAB_F0(BSDF), BaseColor, Specular, Metallic);
+					float Roughness = SLAB_ROUGHNESS(BSDF);
 
-						float3 DiffuseAlbedo = 0;
-						float3 F0 = 0;
-						ConvertFromMetalness(BaseColor, Specular, Metallic, SLAB_DIFFUSEALBEDO(StrataData.Layers[l].BSDFs[i]), SLAB_F0(StrataData.Layers[l].BSDFs[i]));
-						SLAB_ROUGHNESS(StrataData.Layers[l].BSDFs[i]) = Roughness;
-					}
-					// Apply DBuffer as a coating layer (using parameter blending for efficiency purpose)
-					#elif STRATA_DBUFFER_BLENDING_MODE == STRATA_DBUFFER_COATING
+				#if STRATA_DBUFFER_RESPONSE_BASECOLOR
+					BaseColor = BaseColor * StrataDBufferData.OneMinusCoverage_BaseColor + StrataDBufferData.BaseColor;
+					if (BSDF_GETHASFUZZ(BSDF))
 					{
-						// Since all DBuffer parameters are not necessary initialized, take default value from the existing BDSF
-						FStrataDBuffer LocalStrataDBufferData = StrataDBufferData;
-						{
-							float3 SlabBaseColor = 0;
-							float  SlabMetallic = 0;
-							float  SlabSpecular = 0;
-							ConvertToMetalness(SLAB_DIFFUSEALBEDO(StrataData.Layers[l].BSDFs[i]), SLAB_F0(StrataData.Layers[l].BSDFs[i]), SlabBaseColor, SlabSpecular, SlabMetallic);
-						#if !STRATA_DBUFFER_RESPONSE_BASECOLOR
-							LocalStrataDBufferData.BaseColor					= LocalStrataDBufferData.Coverage * SlabBaseColor;
-							LocalStrataDBufferData.OneMinusCoverage_BaseColor	= LocalStrataDBufferData.OneMinusCoverage;
-						#endif
-						#if !STRATA_DBUFFER_RESPONSE_ROUGHNESS
-							LocalStrataDBufferData.Metallic						= LocalStrataDBufferData.Coverage * SlabMetallic;
-							LocalStrataDBufferData.Specular						= LocalStrataDBufferData.Coverage * SlabSpecular;
-							LocalStrataDBufferData.Roughness					= LocalStrataDBufferData.Coverage * SLAB_ROUGHNESS(StrataData.Layers[l].BSDFs[i]);
-							LocalStrataDBufferData.OneMinusCoverage_Roughness	= LocalStrataDBufferData.OneMinusCoverage;
-						#endif
-						}
-
-						FStrataPixelHeader StrataDecalHeader;
-						FStrataData StrataDecalData;
-						StrataConvertFromDBuffer(LocalStrataDBufferData, StrataDecalHeader, StrataDecalData);
-
-						// Parameters blend attributes
-						const float Coverage = StrataDecalData.Layers[l].BSDFs[i].Coverage;
-						const float DummyNoV = 1;
-						StrataData.Layers[l].BSDFs[i] = StrataVerticalLayeringParameterBlending(StrataDecalData.Layers[l].BSDFs[i], StrataData.Layers[l].BSDFs[i], BasisIndex, DummyNoV);
+						SLAB_FUZZ_COLOR(BSDF) = SLAB_FUZZ_COLOR(BSDF) * StrataDBufferData.OneMinusCoverage_BaseColor + StrataDBufferData.BaseColor;
 					}
-					#endif // STRATA_DBUFFER_BLENDING_MODE
-					break;
-				}
 
-				case STRATA_BSDF_TYPE_HAIR:
-				{
-				#if STRATA_DBUFFER_RESPONSE_BASECOLOR
-					HAIR_BASECOLOR(StrataData.Layers[l].BSDFs[i]) = HAIR_BASECOLOR(StrataData.Layers[l].BSDFs[i]) * StrataDBufferData.OneMinusCoverage_BaseColor + StrataDBufferData.BaseColor;
+					// Disable SSS when the coverage reach one to preserve the decal color
+					SLAB_SSSDMFP(BSDF) = lerp(0.f, SLAB_SSSDMFP(BSDF), StrataDBufferData.OneMinusCoverage_BaseColor);
+					if (all(SLAB_SSSDMFP(BSDF) == 0.0f))
+					{
+						BSDF_SETHASDMFP(BSDF, 0);
+						BSDF_SETHASSSS(BSDF, 0);
+					}
 				#endif
 
 				#if STRATA_DBUFFER_RESPONSE_ROUGHNESS
-					HAIR_ROUGHNESS(StrataData.Layers[l].BSDFs[i]) = HAIR_ROUGHNESS(StrataData.Layers[l].BSDFs[i]) * StrataDBufferData.OneMinusCoverage_Roughness + StrataDBufferData.Roughness;
-					HAIR_SPECULAR(StrataData.Layers[l].BSDFs[i])  = HAIR_SPECULAR(StrataData.Layers[l].BSDFs[i])  * StrataDBufferData.OneMinusCoverage_Roughness + StrataDBufferData.Specular;
+					Metallic  = Metallic  * StrataDBufferData.OneMinusCoverage_Roughness + StrataDBufferData.Metallic;
+					Specular  = Specular  * StrataDBufferData.OneMinusCoverage_Roughness + StrataDBufferData.Specular;
+					Roughness = Roughness * StrataDBufferData.OneMinusCoverage_Roughness + StrataDBufferData.Roughness;
 				#endif
-					break;
-				}
 
-				case STRATA_BSDF_TYPE_SINGLELAYERWATER:
+					float3 DiffuseAlbedo = 0;
+					float3 F0 = 0;
+					ConvertFromMetalness(BaseColor, Specular, Metallic, SLAB_DIFFUSEALBEDO(BSDF), SLAB_F0(BSDF));
+					SLAB_ROUGHNESS(BSDF) = Roughness;
+				}
+				// Apply DBuffer as a coating layer (using parameter blending for efficiency purpose)
+				#elif STRATA_DBUFFER_BLENDING_MODE == STRATA_DBUFFER_COATING
 				{
-				#if STRATA_DBUFFER_RESPONSE_BASECOLOR
-					SLW_BASECOLOR(StrataData.Layers[l].BSDFs[i]) = SLW_BASECOLOR(StrataData.Layers[l].BSDFs[i]) * StrataDBufferData.OneMinusCoverage_BaseColor + StrataDBufferData.BaseColor;
-				#endif
+					// Since all DBuffer parameters are not necessary initialized, take default value from the existing BDSF
+					FStrataDBuffer LocalStrataDBufferData = StrataDBufferData;
+					{
+						float3 SlabBaseColor = 0;
+						float  SlabMetallic = 0;
+						float  SlabSpecular = 0;
+						ConvertToMetalness(SLAB_DIFFUSEALBEDO(BSDF), SLAB_F0(BSDF), SlabBaseColor, SlabSpecular, SlabMetallic);
+					#if !STRATA_DBUFFER_RESPONSE_BASECOLOR
+						LocalStrataDBufferData.BaseColor					= LocalStrataDBufferData.Coverage * SlabBaseColor;
+						LocalStrataDBufferData.OneMinusCoverage_BaseColor	= LocalStrataDBufferData.OneMinusCoverage;
+					#endif
+					#if !STRATA_DBUFFER_RESPONSE_ROUGHNESS
+						LocalStrataDBufferData.Metallic						= LocalStrataDBufferData.Coverage * SlabMetallic;
+						LocalStrataDBufferData.Specular						= LocalStrataDBufferData.Coverage * SlabSpecular;
+						LocalStrataDBufferData.Roughness					= LocalStrataDBufferData.Coverage * SLAB_ROUGHNESS(BSDF);
+						LocalStrataDBufferData.OneMinusCoverage_Roughness	= LocalStrataDBufferData.OneMinusCoverage;
+					#endif
+					}
 
-				#if STRATA_DBUFFER_RESPONSE_ROUGHNESS
-					SLW_ROUGHNESS(StrataData.Layers[l].BSDFs[i]) = SLW_ROUGHNESS(StrataData.Layers[l].BSDFs[i]) * StrataDBufferData.OneMinusCoverage_Roughness + StrataDBufferData.Roughness;
-					SLW_SPECULAR(StrataData.Layers[l].BSDFs[i])  = SLW_SPECULAR(StrataData.Layers[l].BSDFs[i])  * StrataDBufferData.OneMinusCoverage_Roughness + StrataDBufferData.Specular;
-					SLW_METALLIC(StrataData.Layers[l].BSDFs[i])  = SLW_METALLIC(StrataData.Layers[l].BSDFs[i])  * StrataDBufferData.OneMinusCoverage_Roughness + StrataDBufferData.Metallic;
-				#endif
-					break;
+					FStrataPixelHeader StrataDecalHeader;
+					FStrataData StrataDecalData;
+					StrataConvertFromDBuffer(LocalStrataDBufferData, StrataDecalHeader, StrataDecalData);
+
+					// Parameters blend attributes
+					const float Coverage = StrataDecalData.Layers[l].BSDFs[i].Coverage;
+					const float DummyNoV = 1;
+					BSDF = StrataVerticalLayeringParameterBlending(StrataDecalData.Layers[l].BSDFs[i], BSDF, BasisIndex, DummyNoV);
 				}
-			} // switch
-		}
+				#endif // STRATA_DBUFFER_BLENDING_MODE
+				break;
+			}
+
+			case STRATA_BSDF_TYPE_HAIR:
+			{
+			#if STRATA_DBUFFER_RESPONSE_BASECOLOR
+				HAIR_BASECOLOR(BSDF) = HAIR_BASECOLOR(BSDF) * StrataDBufferData.OneMinusCoverage_BaseColor + StrataDBufferData.BaseColor;
+			#endif
+
+			#if STRATA_DBUFFER_RESPONSE_ROUGHNESS
+				HAIR_ROUGHNESS(BSDF) = HAIR_ROUGHNESS(BSDF) * StrataDBufferData.OneMinusCoverage_Roughness + StrataDBufferData.Roughness;
+				HAIR_SPECULAR(BSDF)  = HAIR_SPECULAR(BSDF)  * StrataDBufferData.OneMinusCoverage_Roughness + StrataDBufferData.Specular;
+			#endif
+				break;
+			}
+
+			case STRATA_BSDF_TYPE_SINGLELAYERWATER:
+			{
+			#if STRATA_DBUFFER_RESPONSE_BASECOLOR
+				SLW_BASECOLOR(BSDF) = SLW_BASECOLOR(BSDF) * StrataDBufferData.OneMinusCoverage_BaseColor + StrataDBufferData.BaseColor;
+			#endif
+
+			#if STRATA_DBUFFER_RESPONSE_ROUGHNESS
+				SLW_ROUGHNESS(BSDF) = SLW_ROUGHNESS(BSDF) * StrataDBufferData.OneMinusCoverage_Roughness + StrataDBufferData.Roughness;
+				SLW_SPECULAR(BSDF)  = SLW_SPECULAR(BSDF)  * StrataDBufferData.OneMinusCoverage_Roughness + StrataDBufferData.Specular;
+				SLW_METALLIC(BSDF)  = SLW_METALLIC(BSDF)  * StrataDBufferData.OneMinusCoverage_Roughness + StrataDBufferData.Metallic;
+			#endif
+				break;
+			}
+		} // switch
+
+	#undef BSDF
 	}
   #endif // Strata - Disabling temporary DBuffer prior to optimization
 }
diff --git a/Engine/Shaders/Private/DecalCommon.ush b/Engine/Shaders/Private/DecalCommon.ush
--- a/Engine/Shaders/Private/DecalCommon.ush
+++ b/Engine/Shaders/Private/DecalCommon.ush
@@ -16,10 +16,9 @@
 float DecalCommonOutput(
 	inout FPixelShaderOut Out, 
 	in FStrataPixelHeader InHeader,
-	in FStrataData InData,
 	in float CoverageWeight)
 {
-	FStrataDBuffer DBuffer = StrataConvertToDBuffer(InHeader, InData);
+	FStrataDBuffer DBuffer = StrataConvertToDBuffer(InHeader);
 	DBuffer.Coverage *= CoverageWeight;
 
 	float3 Color = DBuffer.Emissive;
diff --git a/Engine/Shaders/Private/DeferredDecal.usf b/Engine/Shaders/Private/DeferredDecal.usf
--- a/Engine/Shaders/Private/DeferredDecal.usf
+++ b/Engine/Shaders/Private/DeferredDecal.usf
@@ -143,10 +143,11 @@ void FPixelShaderInOut_MainPS(inout FPixelShaderIn In, inout FPixelShaderOut Out
 
 #if STRATA_ENABLED    
 	FStrataPixelHeader StrataHeader = InitialiseStrataPixelHeader();
+	StrataHeader.StrataTree = MaterialParameters.StrataTree;
 	StrataHeader.BSDFCount = 1;
 	StrataHeader.MaterialAO = GetMaterialAmbientOcclusion(PixelMaterialInputs);
 	StrataHeader.SharedLocalBases = MaterialParameters.SharedLocalBases;
-	const float Coverage = DecalCommonOutput(Out, StrataHeader, PixelMaterialInputs.FrontMaterial, DecalFading);
+	const float Coverage = DecalCommonOutput(Out, StrataHeader, DecalFading);
 	const float Opacity = Coverage;
 #else // STRATA_ENABLED
 	float3 Color = GetMaterialEmissive(PixelMaterialInputs);
diff --git a/Engine/Shaders/Private/Strata/Strata.ush b/Engine/Shaders/Private/Strata/Strata.ush
--- a/Engine/Shaders/Private/Strata/Strata.ush
+++ b/Engine/Shaders/Private/Strata/Strata.ush
@@ -2278,22 +2278,26 @@ FStrataDBuffer StrataUnpackDBuffer(
 }
 
 FStrataDBuffer StrataConvertToDBuffer(
-	in FStrataPixelHeader InHeader,
-	in FStrataData InData)
+	in FStrataPixelHeader InHeader)
 {
 	FStrataDBuffer Out = (FStrataDBuffer)0;
-	const float3 DiffuseColor = SLAB_DIFFUSEALBEDO(InData.InlinedBSDF);
-	const float3 F0 = SLAB_F0(InData.InlinedBSDF);
+
+#if STRATA_INLINE_SHADING
+
+	FStrataBSDF BSDF = InHeader.StrataTree.BSDFs[0];
+
+	const float3 DiffuseColor = SLAB_DIFFUSEALBEDO(BSDF);
+	const float3 F0 = SLAB_F0(BSDF);
 
 	// STRATA_TODO : Can we avoid this by having a lerpable packing?
-	Out.Roughness = SLAB_ROUGHNESS(InData.InlinedBSDF);
+	Out.Roughness = SLAB_ROUGHNESS(BSDF);
 	Out.Metallic = F0RGBToMetallic(F0);
 	Out.Specular = F0ToDielectricSpecular(F0.x);
 	Out.BaseColor = lerp(DiffuseColor, F0, Out.Metallic);
-	Out.Emissive = BSDF_GETEMISSIVE(InData.InlinedBSDF);
-#if STRATA_INLINE_SHADING
-	Out.WorldNormal = InHeader.SharedLocalBases.Normals[BSDF_GETSHAREDLOCALBASISID(InData.InlinedBSDF)];
-	Out.Coverage = InData.InlinedBSDF.Coverage;
+	Out.Emissive = BSDF_GETEMISSIVE(BSDF);
+	Out.WorldNormal = InHeader.SharedLocalBases.Normals[BSDF_GETSHAREDLOCALBASISID(BSDF)];
+	Out.Coverage = BSDF.Coverage;
+
 #else
 	Out.WorldNormal = float3(0,0,1);
 	Out.Coverage = 1.f;
