Fixed Strata material visualisation byte count.

diff --git a/Engine/Shaders/Private/Strata/Strata.ush b/Engine/Shaders/Private/Strata/Strata.ush
--- a/Engine/Shaders/Private/Strata/Strata.ush
+++ b/Engine/Shaders/Private/Strata/Strata.ush
@@ -1015,6 +1015,7 @@ uint StrataLoadUint1(FStrataMaterialContainer StrataBuffer, inout FStrataAddress
 {
 	uint Data = StrataBuffer[uint3(StrataAddressing.PixelCoords, StrataAddressing.CurrentIndex)];
 	StrataAddressing.CurrentIndex++;
+	StrataAddressing.ReadBytes += 4;
 	return Data;
 }
 
@@ -2829,8 +2830,10 @@ FStrataPixelHeader UnpackStrataHeaderIn(FStrataMaterialContainer StrataBuffer, i
 	#if STRATA_DEFERRED_SHADING
 		Out.StrataBuffer = StrataBuffer;
 		Out.SharedLocalBasesTypes_PackedHeader = PackedHeader;
-		Out.PackedTopLayerData = InStrataTopLayerTexture.Load(uint3(StrataAddressing.PixelCoords, 0));
 		Out.SharedLocalBasesIndexOffset = 0; // Unused
+
+		Out.PackedTopLayerData = InStrataTopLayerTexture.Load(uint3(StrataAddressing.PixelCoords, 0));
+		StrataAddressing.ReadBytes += 4;
 	  #endif
 	}
 	else
