Implement Sky Atmosphere and Exponential Height Fog support in the Path Tracer

This is the start of volumetric support in the path tracer, so a lot of basic infrastructure had to be put into place, making this changelist fairly large. Some shuffling of light parameters had to be done to make room for the volumetric scattering multiplier.

The integration strategy is to use null tracking (with Spectral MIS) for choosing a random scatter point along the ray. This point is chosen similarly to transparent hits, so surface and volume shading are unified. However, these volume hits are chosen proportionally to transmittance times scattering which is not optimal for lights embedded in the volume. To handle the latter, we also allow the main trace call to return a volume segment over which we can compute the direct lighting from local light sources. To simplify the handling of overlapping media and inter-mixed transparent hits, we stochastically select a ray segment. From this point on, we can evaluate direct lighting using equi-angular sampling. The MIS combination of equi-angular sampling and null tracked spectral density sampling was prototyped but found not to bring any improvement for the currently implemented volume types. This will likely have to be revisited in the future.

To improve quality and reduce the amount of ray-marching required, the volume API allows for analytic implementations of transmittance for cases where this is possible to do more efficiently than by ray-marching.

Implementation details for Atmosphere: This volume type is a planet sized, spherically symetric model. Because the default units in UE are centimeters, objects like the planet that are kilometer sized will run into all sorts of numerical precision artifacts. To solve this, an implementation of double-word arithmetic was added which allows enough decimal digits to robustly intersect the planet, as well as cary out the lookups required. The Transmittance is cached in a lookup table indexed by height above the ground and viewing angle cosine. This is similar to the LUT used by the realtime version but with a different parameterization which covers the full range of angles/heights with high precision. This lookup table is automatically baked on demand when atmosphere parameters change. The volumetric sky model is only used when "reference atmosphere" is enabled in the post process volume. This is because the existing approached (cached into a skylight) is generally a bit faster and supports clouds. This toggle may be removed as the support for volumes matures.

Implementation details for ExponentialHeightFog: This volume is represented as a finite slab centered around the camera. Transmittance is easily computed analytically for this volume. We only add this volume when the volumetric fog checkbox is enabled, as the default parameters are not fully physically based. We limit the fog to be present only within a certain radius of the camera, to prevent rays from scattering forever.

diff --git a/Engine/Plugins/Experimental/GPULightmass/Source/GPULightmass/Private/PathTracingLightParameters.inl b/Engine/Plugins/Experimental/GPULightmass/Source/GPULightmass/Private/PathTracingLightParameters.inl
--- a/Engine/Plugins/Experimental/GPULightmass/Source/GPULightmass/Private/PathTracingLightParameters.inl
+++ b/Engine/Plugins/Experimental/GPULightmass/Source/GPULightmass/Private/PathTracingLightParameters.inl
@@ -53,9 +53,8 @@ void SetupPathTracingLightParameters(
 
 		DestLight.Normal = (FVector3f)-Light.Direction;
 		DestLight.Color = FVector3f(Light.Color);
-		DestLight.Dimensions = FVector3f(
+		DestLight.Dimensions = FVector2f(
 			FMath::Sin(0.5f * FMath::DegreesToRadians(Light.LightSourceAngle)),
-			FMath::Sin(0.5f * FMath::DegreesToRadians(Light.LightSourceSoftAngle)),
 			0.0f);
 		DestLight.Attenuation = 1.0;
 		DestLight.IESTextureSlice = -1;
@@ -81,7 +80,7 @@ void SetupPathTracingLightParameters(
 		DestLight.dPdu = (FVector3f)FVector::CrossProduct(Light.Tangent, Light.Direction);
 		DestLight.dPdv = (FVector3f)Light.Tangent;
 
-		DestLight.Dimensions = FVector3f(Light.SourceRadius, Light.SourceSoftRadius, Light.SourceLength);
+		DestLight.Dimensions = FVector2f(Light.SourceRadius, Light.SourceLength);
 		DestLight.Attenuation = 1.0f / Light.AttenuationRadius;
 		DestLight.FalloffExponent = Light.FalloffExponent;
 
@@ -118,7 +117,7 @@ void SetupPathTracingLightParameters(
 		DestLight.dPdu = (FVector3f)FVector::CrossProduct(Light.Tangent, Light.Direction);
 		DestLight.dPdv = (FVector3f)Light.Tangent;
 		DestLight.Color = FVector3f(Light.Color);
-		DestLight.Dimensions = FVector3f(Light.SourceRadius, Light.SourceSoftRadius, Light.SourceLength);
+		DestLight.Dimensions = FVector2f(Light.SourceRadius, Light.SourceLength);
 		DestLight.Shaping = FVector2f(Light.SpotAngles);
 		DestLight.Attenuation = 1.0f / Light.AttenuationRadius;
 		DestLight.FalloffExponent = Light.FalloffExponent;
@@ -174,7 +173,7 @@ void SetupPathTracingLightParameters(
 		LightColor /= 0.5f * Light.SourceWidth * Light.SourceHeight;
 		DestLight.Color = FVector3f(LightColor);
 
-		DestLight.Dimensions = FVector3f(Light.SourceWidth, Light.SourceHeight, 0.0f);
+		DestLight.Dimensions = FVector2f(Light.SourceWidth, Light.SourceHeight);
 		DestLight.Attenuation = 1.0f / Light.AttenuationRadius;
 		DestLight.Shaping = FVector2f(FMath::Cos(FMath::DegreesToRadians(Light.BarnDoorAngle)), Light.BarnDoorLength);
 
diff --git a/Engine/Shaders/Private/DoubleWordMath.ush b/Engine/Shaders/Private/DoubleWordMath.ush
new file mode 100644
--- /dev/null
+++ b/Engine/Shaders/Private/DoubleWordMath.ush
@@ -0,0 +1,196 @@
+// Copyright Epic Games, Inc. All Rights Reserved.
+
+#pragma once
+
+// Double-word arithmetic using algorithms from:
+//    "Tight and rigourous error bounds for basic building blocks of double-word arithmetic"
+//       Mioara Joldes, Jean - Michel Muller, Valentina Popescu 
+//    ACM Transactions on Mathematical Software [2017]
+
+// NOTE: the publication above is mainly about proving tight error bounds, the technique itself is much older even (see [Knuth 1974] for instance).
+
+// The basic idea is to represent floating point results using an un-evaluated sum of two floats. For example for addition of a + b, we have:
+//    a + b == c.h + c.l
+//    c.h: closest float to a + b
+//    c.l: remainder (usually a small value)
+
+// The functions below implement a minimum set of arithmetic operations on such numbers. The suffix _f or _d is used to
+// distinguish the type of the arguments (plain float vs. "double" float).
+// TODO: Make use of operator overloading (when available)?
+
+// Important note: the HLSL keyword 'precise' is required to prevent the compiler from optimizing away certain constructions.
+// Indeed while some expressions below may appear to cancel out mathematically, they do not cancel out _numerically_ which is
+// how the algorithms are able to capture additional precision.
+
+// The implementation below assume that the underlying hardware has properly working FMA instruction (mad). Without this, a different
+// implementation of FastTwoMult would be required (see papers above for details on the litterature for this).
+
+struct DWFloat
+{
+	float h;
+	float l;
+
+	float GetFloat()
+	{
+		// merge the hi/lo portions back together to make an ordinary float
+		return h + l;
+	}
+};
+
+DWFloat DWFloatCreate(float h, float l = 0)
+{
+	DWFloat Result;
+	Result.h = h;
+	Result.l = l;
+	return Result;
+}
+
+DWFloat DWNeg(DWFloat a)
+{
+	return DWFloatCreate(-a.h, -a.l);
+}
+
+DWFloat FastTwoSum(float a, float b)
+{
+	const float s = a + b;
+	const precise float e = b - (s - a);
+	return DWFloatCreate(s, e);
+}
+
+DWFloat TwoSum(float a, float b)
+{
+	const float s = a + b;
+	const precise float ap = s - b;
+	const precise float bp = s - ap;
+	const precise float da = a - ap;
+	const precise float db = b - bp;
+	return DWFloatCreate(s, da + db);
+}
+
+DWFloat FastTwoMult(float a, float b)
+{
+	const float p = a * b;
+	const precise float e = mad(a, b, -p);
+	return DWFloatCreate(p, e);
+}
+
+DWFloat DWAdd_ff(float a, float b)
+{
+	return TwoSum(a, b);
+}
+
+DWFloat DWSub_ff(float a, float b)
+{
+	return TwoSum(a, -b);
+}
+
+DWFloat DWAdd_dd(DWFloat a, DWFloat b)
+{
+	DWFloat s = TwoSum(a.h, b.h);
+	DWFloat t = TwoSum(a.l, b.l);
+	s = FastTwoSum(s.h, s.l + t.h);
+	s = FastTwoSum(s.h, s.l + t.l);
+	return s;
+}
+
+DWFloat DWSub_dd(DWFloat a, DWFloat b)
+{
+	return DWAdd_dd(a, DWNeg(b));
+}
+
+DWFloat DWAdd_df(DWFloat a, float b)
+{
+	const DWFloat s = TwoSum(a.h, b);
+	return FastTwoSum(s.h, a.l + s.l);
+}
+
+DWFloat DWSub_df(DWFloat a, float b)
+{
+	const DWFloat s = TwoSum(a.h, -b);
+	return FastTwoSum(s.h, a.l + s.l);
+}
+
+DWFloat DWSub_fd(float a, DWFloat b)
+{
+	return DWNeg(DWSub_df(b, a));
+}
+
+DWFloat DWMul_dd(DWFloat a, DWFloat b)
+{
+	DWFloat p = FastTwoMult(a.h, b.h);
+	p.l = mad(a.h, b.l, p.l);
+	p.l = mad(a.l, b.h, p.l);
+	//p.l = mad(a.l, b.l, p.l); // doesn't seem to add much precision in practice
+	return FastTwoSum(p.h, p.l);
+}
+
+DWFloat DWMul_ff(float a, float b)
+{
+	return FastTwoMult(a, b);
+}
+
+DWFloat DWMul_df(DWFloat a, float b)
+{
+	DWFloat p = FastTwoMult(a.h, b);
+	p.l = mad(a.l, b, p.l);
+	return FastTwoSum(p.h, p.l);
+}
+
+DWFloat DWSqr_d(DWFloat a)
+{
+	return DWMul_dd(a, a);
+}
+
+DWFloat DWSqr_f(float a)
+{
+	return FastTwoMult(a, a);
+}
+
+DWFloat DWSqrt_d(DWFloat a)
+{
+	// "Accurate calculation of Euclidean Norms using Double-word arithmetic"
+	// Vincent Lefevre, Nicolas Louvet, Jean-Michel Muller, Joris Picot, Laurence Rideau - [2021]
+
+	// Algorithm 8
+	const float sh = sqrt(a.h);
+	const precise float p1 = mad(-sh, sh, a.h);
+	const float p2 = a.l + p1;
+	const float sl = p2 / (2.0 * sh);
+	return FastTwoSum(sh, sl);
+}
+
+DWFloat DWDiv_dd(DWFloat b, DWFloat a)
+{
+	const float th = rcp(a.h);
+	const precise float rh = mad(-a.h, th, 1.0);
+	const float rl = -a.l * th;
+	const DWFloat e = FastTwoSum(rh, rl);
+	const DWFloat d = DWMul_df(e, th);
+	const DWFloat m = DWAdd_df(d, th);
+	return DWMul_dd(b, m);
+}
+
+DWFloat DWLengthSquared_fff(float X, float Y, float Z)
+{
+	// "Accurate calculation of Euclidean Norms using Double-word arithmetic"
+	// Vincent Lefevre, Nicolas Louvet, Jean-Michel Muller, Joris Picot, Laurence Rideau - [2021]
+
+	// Algorithm 10
+	const DWFloat y0 = FastTwoMult(X, X);
+	const DWFloat y1 = FastTwoMult(Y, Y);
+	const DWFloat y2 = FastTwoMult(Z, Z);
+	const DWFloat r = DWAdd_df(TwoSum(y0.h, y1.h), y2.h);
+	const float e = (y0.l + y1.l) + y2.l;
+	return DWAdd_df(r, e);
+}
+
+DWFloat DWLengthSquared_ddd(DWFloat X, DWFloat Y, DWFloat Z)
+{
+	return DWAdd_dd(DWAdd_dd(DWSqr_d(X), DWSqr_d(Y)), DWSqr_d(Z));
+}
+
+
+bool DWGreaterThan(DWFloat a, float b)
+{
+	return a.h > b || (a.h == b && a.l > b);
+}
diff --git a/Engine/Shaders/Private/ParticipatingMediaCommon.ush b/Engine/Shaders/Private/ParticipatingMediaCommon.ush
--- a/Engine/Shaders/Private/ParticipatingMediaCommon.ush
+++ b/Engine/Shaders/Private/ParticipatingMediaCommon.ush
@@ -112,26 +112,24 @@ float SchlickPhase(float G, float CosTheta)
 // Follows PBRT convention http://www.pbr-book.org/3ed-2018/Light_Transport_II_Volume_Rendering/Sampling_Volume_Scattering.html#SamplingPhaseFunctions
 float HenyeyGreensteinPhaseInvertCDF(float E, float G)
 {
+	float R = 1.0 - 2.0 * E;
+	if (abs(G) < 0.001f)
+	{
+		// Careful: OneOver2G undefined for g close to 0
+		return R;
+	}
 	float OnePlusG2  = 1.0f + G * G;
 	float OneMinusG2 = 1.0f - G * G;
 	float OneOver2G  = 0.5f / G;
-	float t = OneMinusG2 / (1.0f + G - 2.0f * G * E);
-	return -OneOver2G * (OnePlusG2 - t * t);	// Careful: OneOver2G undefined for g close to 0
+	float t = OneMinusG2 / (1.0f + G * R);
+	return OneOver2G * (t * t - OnePlusG2);
 }
 
 // Follows PBRT convention http://www.pbr-book.org/3ed-2018/Light_Transport_II_Volume_Rendering/Sampling_Volume_Scattering.html#SamplingPhaseFunctions
 float4 ImportanceSampleHenyeyGreensteinPhase(float2 E, float G)
 {
 	float Phi = 2.0f * PI * E.x;
-	float CosTheta = 0;
-	if (abs(G) < 0.001f)
-	{
-		CosTheta = 1.0f - 2.0f * E.y;
-	}
-	else
-	{
-		CosTheta = HenyeyGreensteinPhaseInvertCDF(E.y, G);
-	}
+	float CosTheta = HenyeyGreensteinPhaseInvertCDF(E.y, G);
 	float SinTheta = sqrt(max(0.0f, 1.0f - CosTheta * CosTheta));
 
 	float3 H = float3(SinTheta * sin(Phi), SinTheta * cos(Phi), CosTheta);
@@ -139,6 +137,25 @@ float4 ImportanceSampleHenyeyGreensteinPhase(float2 E, float G)
 	return float4(H, HenyeyGreensteinPhase(G, CosTheta));
 }
 
+// Returns CosTheta given a random value in [0,1)
+float RayleighPhaseInvertCdf(float E)
+{
+	// [Frisvad, 2011] "Importance sampling the Rayleigh phase function"
+	// Optical Society of America. Journal A: Optics, Image Science, and Vision
+	float Z = E * 4.0 - 2.0;
+	float InvZ = sqrt(Z * Z + 1.0);
+	float u = pow(Z + InvZ, 1.0 / 3.0); // cbrt
+	return u - rcp(u);
+}
+
+float4 ImportanceSampleRayleigh(float2 E)
+{
+	float Phi = 2.0f * PI * E.x;
+	float CosTheta = RayleighPhaseInvertCdf(E.y);
+	float SinTheta = sqrt(max(0.0f, 1.0f - CosTheta * CosTheta));
+	float3 H = float3(SinTheta * sin(Phi), SinTheta * cos(Phi), CosTheta);
+	return float4(H, RayleighPhase(CosTheta));
+}
 
 
 //---------------------------------------------
diff --git a/Engine/Shaders/Private/PathTracing/Light/PathTracingDirectionalLight.ush b/Engine/Shaders/Private/PathTracing/Light/PathTracingDirectionalLight.ush
--- a/Engine/Shaders/Private/PathTracing/Light/PathTracingDirectionalLight.ush
+++ b/Engine/Shaders/Private/PathTracing/Light/PathTracingDirectionalLight.ush
@@ -60,3 +60,20 @@ float DirectionalLight_EstimateLight(
 	float LightPower = Luminance(GetColor(LightId));
 	return LightPower * NoL;
 }
+
+// NOTE: unused
+FVolumeLightSampleSetup DirectionalLight_PrepareLightVolumeSample(
+	int LightId,
+	float3 RayOrigin,
+	float3 RayDirection,
+	float TMin,
+	float TMax
+)
+{
+	float3 LightColor = GetColor(LightId);
+	float SinThetaMax = max(0.001, GetRadius(LightId)); // Sin(Angle / 2)
+	float SinThetaMax2 = SinThetaMax * SinThetaMax;
+	float OneMinusCosThetaMax = SinThetaMax2 < 0.01 ? SinThetaMax2 * (0.5 + 0.125 * SinThetaMax2) : 1 - sqrt(1 - SinThetaMax2);
+	float SolidAngle = 2 * PI * OneMinusCosThetaMax;
+	return CreateUniformSampler(max3(LightColor.x, LightColor.y, LightColor.z) * rcp(SolidAngle), TMin, TMax);
+}
diff --git a/Engine/Shaders/Private/PathTracing/Light/PathTracingLightCommon.ush b/Engine/Shaders/Private/PathTracing/Light/PathTracingLightCommon.ush
--- a/Engine/Shaders/Private/PathTracing/Light/PathTracingLightCommon.ush
+++ b/Engine/Shaders/Private/PathTracing/Light/PathTracingLightCommon.ush
@@ -12,7 +12,8 @@ LightCommon.usf: Common utility functions for light sampling
 uint SceneLightCount;
 StructuredBuffer<FPathTracingLight> SceneLights;
 
-struct FLightHit {
+struct FLightHit
+{
 	float3 Radiance;
 	float Pdf;
 	float HitT;
@@ -21,13 +22,103 @@ struct FLightHit {
 	bool IsHit() { return HitT > 0; }
 };
 
-struct FLightSample {
+struct FLightSample
+{
 	float3 RadianceOverPdf;
 	float Pdf;
 	float3 Direction;
 	float Distance;
 };
 
+
+struct FVolumeLightSampleSetup
+{
+	float VolumeTMin;      // range of 't' values along the ray that overlap the light (and the volume)
+	float VolumeTMax;
+
+	float LightImportance; // estimate of the light radiance for this ray (could be 0 if light does not support equi-angular sampling)
+
+	// Equi-Angular case
+	float T0;			   // distance from ray origin to projected light center
+	float Height;		   // distance from light center to ray (when negative, use uniform sampler)
+	float ThetaMin;        // angle to start of ray interval
+	float ThetaMax;        // angle to end of ray interval
+
+	float K;               // normalization constant for the pdf
+
+	bool IsValid() { return VolumeTMin < VolumeTMax; } // NOTE: we assume the range is a subset of the ray's TMin/TMax
+
+	float2 SampleDistance(float RandSample)
+	{
+		if (Height > 0)
+		{
+			// equi-angular pdf
+			const float TanTheta = tan(lerp(ThetaMin, ThetaMax, RandSample));
+			const float T = Height * TanTheta + T0;
+			const float Pdf = rcp((1 + TanTheta * TanTheta) * K);
+			return float2(T, Pdf);
+		}
+		else
+		{
+			// uniform pdf
+			const float T = lerp(VolumeTMin, VolumeTMax, RandSample);
+			const float Pdf = rcp(K);
+			return float2(T, Pdf);
+		}
+	}
+
+	float Pdf(float T)
+	{
+		if (T > VolumeTMin && T < VolumeTMax)
+		{
+			if (Height > 0)
+			{
+				const float TanTheta = (T - T0) / Height;
+				return rcp((1 + TanTheta * TanTheta) * K);
+			}
+			else
+			{
+				return rcp(K);
+			}
+		}
+		return 0.0;
+	}
+};
+
+FVolumeLightSampleSetup CreateEquiAngularSampler(float LightImportance, float3 Center, float3 Ro, float3 Rd, float TMin, float TMax)
+{
+#if 0
+	float T0 = dot(Center - Ro, Rd);
+	float3 P0 = Ro + T0 * Rd;
+	Center = Center - normalize(Center - P0);
+#endif
+	const float3 V = Center - Ro;
+	FVolumeLightSampleSetup Result = (FVolumeLightSampleSetup)0;
+	Result.VolumeTMin = TMin;
+	Result.VolumeTMax = TMax;
+	Result.T0 = dot(V, Rd);
+	Result.Height = length(cross(V, Rd));
+	Result.ThetaMin = Result.Height > 0 ? atan((TMin - Result.T0) / Result.Height) : 0.0;
+	Result.ThetaMax = Result.Height > 0 ? atan((TMax - Result.T0) / Result.Height) : 0.0;
+	Result.K = (Result.ThetaMax - Result.ThetaMin) * Result.Height;
+	Result.LightImportance = LightImportance * (TMax - TMin);
+	return Result;
+}
+
+FVolumeLightSampleSetup CreateUniformSampler(float LightImportance, float TMin, float TMax)
+{
+	FVolumeLightSampleSetup Result = (FVolumeLightSampleSetup)0;
+	Result.VolumeTMin = TMin;
+	Result.VolumeTMax = TMax;
+	Result.LightImportance = LightImportance * (TMax - TMin);
+	Result.Height = -1.0; // mark that we don't want equi-angular sampling here
+	Result.K = TMax - TMin;
+	return Result;
+}
+
+
+
+
 FLightHit NullLightHit()
 {
 	return (FLightHit)0;
@@ -57,6 +148,11 @@ FLightSample CreateLightSample(float3 RadianceOverPdf, float Pdf, float3 Directi
 	return Sample;
 }
 
+FVolumeLightSampleSetup NullVolumeLightSetup()
+{
+	return (FVolumeLightSampleSetup)0;
+}
+
 bool IsEnvironmentLight(int LightId)
 {
 	return (SceneLights[LightId].Flags & PATHTRACER_FLAG_TYPE_MASK) == PATHTRACING_LIGHT_SKY;
@@ -141,14 +237,9 @@ float GetRadius(int LightId)
 	return SceneLights[LightId].Dimensions.x;
 }
 
-float GetSoftRadius(int LightId)
-{
-	return SceneLights[LightId].Dimensions.y;
-}
-
 float GetSourceLength(int LightId)
 {
-	return SceneLights[LightId].Dimensions.z;
+	return SceneLights[LightId].Dimensions.y;
 }
 
 float GetAttenuation(int LightId)
@@ -263,3 +354,18 @@ bool CastsShadow(int LightId)
 {
 	return (SceneLights[LightId].Flags & PATHTRACER_FLAG_CAST_SHADOW_MASK) != 0;
 }
+
+bool CastsVolumeShadow(int LightId)
+{
+	return (SceneLights[LightId].Flags & PATHTRACER_FLAG_CAST_VOL_SHADOW_MASK) != 0;
+}
+
+bool HasRectTexture(int LightId)
+{
+	return (SceneLights[LightId].Flags & PATHTRACER_FLAG_HAS_RECT_TEXTURE_MASK) != 0;
+}
+
+float GetVolumetricScatteringIntensity(int LightId)
+{
+	return SceneLights[LightId].VolumetricScatteringIntensity;
+}
\ No newline at end of file
diff --git a/Engine/Shaders/Private/PathTracing/Light/PathTracingLightSampling.ush b/Engine/Shaders/Private/PathTracing/Light/PathTracingLightSampling.ush
--- a/Engine/Shaders/Private/PathTracing/Light/PathTracingLightSampling.ush
+++ b/Engine/Shaders/Private/PathTracing/Light/PathTracingLightSampling.ush
@@ -76,6 +76,26 @@ float EstimateLight(
 	}
 }
 
+FVolumeLightSampleSetup PrepareLightVolumeSample(
+	int LightId,
+	float3 RayOrigin,
+	float3 RayDirection,
+	float TMin,
+	float TMax
+)
+{
+	// NOTE: distant lights are excluded from this because they work better with density sampling
+	switch (SceneLights[LightId].Flags & PATHTRACER_FLAG_TYPE_MASK)
+	{
+		//case PATHTRACING_LIGHT_SKY:         return SkyLight_PrepareLightVolumeSample(LightId, RayOrigin, RayDirection, TMin, TMax);
+		case PATHTRACING_LIGHT_POINT:       return PointLight_PrepareLightVolumeSample(LightId, RayOrigin, RayDirection, TMin, TMax);
+		//case PATHTRACING_LIGHT_DIRECTIONAL: return DirectionalLight_PrepareLightVolumeSample(LightId, RayOrigin, RayDirection, TMin, TMax);
+		case PATHTRACING_LIGHT_RECT:        return RectLight_PrepareLightVolumeSample(LightId, RayOrigin, RayDirection, TMin, TMax);
+		case PATHTRACING_LIGHT_SPOT:        return SpotLight_PrepareLightVolumeSample(LightId, RayOrigin, RayDirection, TMin, TMax);
+		default:                            return NullVolumeLightSetup();
+	}
+}
+
 // #dxr_todo: passing LightPickingCdf as anything but inout causes performance to drop significantly -- why ?
 float InitLightPickingCdf(
 	float3 TranslatedWorldPos,
diff --git a/Engine/Shaders/Private/PathTracing/Light/PathTracingPointLight.ush b/Engine/Shaders/Private/PathTracing/Light/PathTracingPointLight.ush
--- a/Engine/Shaders/Private/PathTracing/Light/PathTracingPointLight.ush
+++ b/Engine/Shaders/Private/PathTracing/Light/PathTracingPointLight.ush
@@ -119,3 +119,29 @@ float PointLight_EstimateLight(
 	float Falloff = ComputeAttenuationFalloff(LightDistanceSquared, LightId);
 	return LightPower * Falloff * NoL / LightDistanceSquared;
 }
+
+FVolumeLightSampleSetup PointLight_PrepareLightVolumeSample(
+	int LightId,
+	float3 RayOrigin,
+	float3 RayDirection,
+	float TMin,
+	float TMax
+)
+{
+	float3 Center = GetTranslatedPosition(LightId);
+	float AttenuationRadius = rcp(GetAttenuation(LightId));
+	float AttenuationRadius2 = Pow2(AttenuationRadius);
+	float3 oc = RayOrigin - Center;
+	float b = dot(oc, RayDirection);
+	float h = AttenuationRadius2 - length2(oc - b * RayDirection);
+	if (h > 0)
+	{
+		float2 t = -b + float2(-1, +1) * sqrt(h);
+		TMin = max(t.x, TMin);
+		TMax = min(t.y, TMax);
+		float3 C = GetColor(LightId);
+		float LightImportance = max3(C.x, C.y, C.z);
+		return CreateEquiAngularSampler(LightImportance, Center, RayOrigin, RayDirection, TMin, TMax);
+	}
+	return NullVolumeLightSetup();
+}
diff --git a/Engine/Shaders/Private/PathTracing/Light/PathTracingRectLight.ush b/Engine/Shaders/Private/PathTracing/Light/PathTracingRectLight.ush
--- a/Engine/Shaders/Private/PathTracing/Light/PathTracingRectLight.ush
+++ b/Engine/Shaders/Private/PathTracing/Light/PathTracingRectLight.ush
@@ -13,14 +13,17 @@
 
 float3 EvaluateTexture(int LightId, float2 UV)
 {
-	// TODO: could use path roughness / ray cone to lower the mip level here
-	UV.y = 1 - UV.y; // match orientation with raster path
-
-	// Transform local UV into rect. light atlas UV
-	const float2 UVScale  = GetRectLightAtlasUVScale(LightId);
-	const float2 UVOffset = GetRectLightAtlasUVOffset(LightId);
-	UV = UV * UVScale + UVOffset;
-	return View.RectLightAtlasTexture.SampleLevel(View.RectLightAtlasSampler, UV, 0).xyz;
+	if (HasRectTexture(LightId))
+	{
+		// TODO: could use path roughness / ray cone to lower the mip level here
+		UV.y = 1 - UV.y; // match orientation with raster path
+		// Transform local UV into rect. light atlas UV
+		const float2 UVScale = GetRectLightAtlasUVScale(LightId);
+		const float2 UVOffset = GetRectLightAtlasUVOffset(LightId);
+		UV = UV * UVScale + UVOffset;
+		return View.RectLightAtlasTexture.SampleLevel(View.RectLightAtlasSampler, UV, 0).xyz;
+	}
+	return 1.0;
 }
 
 #endif
@@ -111,7 +114,7 @@ FLightHit RectLight_TraceLight(RayDesc Ray, int LightId)
 #if USE_RECT_LIGHT_TEXTURES
 				Radiance *= EvaluateTexture(LightId, (UV + LightExtent) / (LightExtent * 2));
 #endif
-				return CreateLightHit(Radiance, 1.0 / SolidAngle, t);
+				return CreateLightHit(Radiance, SolidAngle > 0 ? rcp(SolidAngle) : 0.0, t);
 			}
 		}
 	}
@@ -163,7 +166,11 @@ FLightSample RectLight_SampleLight(
 	RadianceOverPdf *= EvaluateTexture(LightId, UV);
 #endif
 
-	return CreateLightSample(RadianceOverPdf, 1.0 / SphericalRect.SolidAngle, normalize(Direction), length(Direction));
+	if (SphericalRect.SolidAngle > 0)
+	{
+		return CreateLightSample(RadianceOverPdf, 1.0 / SphericalRect.SolidAngle, normalize(Direction), length(Direction));
+	}
+	return NullLightSample();
 }
 
 float RectLight_EstimateLight(
@@ -197,3 +204,44 @@ float RectLight_EstimateLight(
 	float Falloff = ComputeAttenuationFalloff(LightDistanceSquared, LightId);
 	return LightPower * Falloff * Area * NoL * LNoL / LightDistanceSquared;
 }
+
+FVolumeLightSampleSetup RectLight_PrepareLightVolumeSample(
+	int LightId,
+	float3 RayOrigin,
+	float3 RayDirection,
+	float TMin,
+	float TMax
+)
+{
+	float3 Center = GetTranslatedPosition(LightId);
+	float AttenuationRadius = rcp(GetAttenuation(LightId));
+	float AttenuationRadius2 = Pow2(AttenuationRadius);
+	float3 oc = RayOrigin - Center;
+	float b = dot(oc, RayDirection);
+	float h = AttenuationRadius2 - length2(oc - b * RayDirection);
+	if (h > 0)
+	{
+		float2 t = -b + float2(-1, +1) * sqrt(h);
+		TMin = max(t.x, TMin);
+		TMax = min(t.y, TMax);
+
+		// now clip the sphere of influence against the plane of the quad and adjust the near/far distance depending on the sign
+		float3 N = GetNormal(LightId);
+		float DoN = dot(RayDirection, N);
+		float TPlane = dot(Center - RayOrigin, N) * rcp(DoN);
+		if (DoN > 0.0)
+			TMin = max(TMin, TPlane);
+		else
+			TMax = min(TMax, TPlane);
+
+		// TODO: clip against barndoors if the ray intersects them?
+
+		float3 C = GetColor(LightId);
+		float Width = GetWidth(LightId);
+		float Height = GetHeight(LightId);
+		float Area = Width * Height;
+		float LightImportance = max3(C.x, C.y, C.z) * Area;
+		return CreateEquiAngularSampler(LightImportance, Center, RayOrigin, RayDirection, TMin, TMax);
+	}
+	return NullVolumeLightSetup();
+}
diff --git a/Engine/Shaders/Private/PathTracing/Light/PathTracingSkyLight.ush b/Engine/Shaders/Private/PathTracing/Light/PathTracingSkyLight.ush
--- a/Engine/Shaders/Private/PathTracing/Light/PathTracingSkyLight.ush
+++ b/Engine/Shaders/Private/PathTracing/Light/PathTracingSkyLight.ush
@@ -87,3 +87,15 @@ FLightSample SkyLight_SampleLight(
 	return CreateLightSample(SkyLightRadiance / Pdf, Pdf, Direction, POSITIVE_INFINITY);
 #endif
 }
+
+// NOTE: unused
+FVolumeLightSampleSetup SkyLight_PrepareLightVolumeSample(
+	int LightId,
+	float3 RayOrigin,
+	float3 RayDirection,
+	float TMin,
+	float TMax
+)
+{
+	return CreateUniformSampler(4 * PI, TMin, TMax);
+}
diff --git a/Engine/Shaders/Private/PathTracing/Light/PathTracingSpotLight.ush b/Engine/Shaders/Private/PathTracing/Light/PathTracingSpotLight.ush
--- a/Engine/Shaders/Private/PathTracing/Light/PathTracingSpotLight.ush
+++ b/Engine/Shaders/Private/PathTracing/Light/PathTracingSpotLight.ush
@@ -85,3 +85,91 @@ float SpotLight_EstimateLight(
 	return OutIrradiance;
 }
 
+float2 IntersectInfiniteConeVolume(
+	float3  ro, float3  rd,
+	float3  apex, float3  normal,
+	float radius)
+{
+	// Adapted from: https://www.iquilezles.org/www/articles/intersectors/intersectors.htm
+	const float3 ba = normal;
+	const float3 oa = ro - apex;
+
+	const float m0 = dot(ba, ba);
+	const float m1 = dot(oa, ba);
+	const float m3 = dot(rd, ba);
+
+	const float m4 = dot(rd, oa);
+	const float m5 = dot(oa, oa);
+	const float hy = m0 + radius * radius;
+
+	const float k2 = m0 * m0 - m3 * m3 * hy;
+	const float k1 = m0 * m0 * m4 - m1 * m3 * hy;
+	const float k0 = m0 * m0 * m5 - m1 * m1 * hy;
+
+	const float h = k1 * k1 - k2 * k0;
+	if (h > 0)
+	{
+		const float s = sign(k2);
+		const float2 t = (-k1 + float2(-s, s) * sqrt(h)) / k2;
+		const float2 y = m1 + t * m3;
+
+		// check which hits are valid:
+		const int c = ((y.x > 0) ? 1 : 0) + (y.y > 0 ? 2 : 0);
+		if (c == 1)
+		{
+			// only smaller root is valid, far root is invalid
+			return float2(0.0, t.x);
+		}
+		if (c == 2)
+		{
+			// only larger root is valid, must have started outside the cone
+			return float2(t.y, POSITIVE_INFINITY);
+		}
+		if (c == 3)
+		{
+			// both roots are valid
+			return t;
+		}
+	}
+	return -1.0;
+}
+
+
+FVolumeLightSampleSetup SpotLight_PrepareLightVolumeSample(
+	int LightId,
+	float3 RayOrigin,
+	float3 RayDirection,
+	float TMin,
+	float TMax
+)
+{
+	float3 Center = GetTranslatedPosition(LightId);
+	float AttenuationRadius = rcp(GetAttenuation(LightId));
+	float AttenuationRadius2 = Pow2(AttenuationRadius);
+	float3 oc = RayOrigin - Center;
+	float b = dot(oc, RayDirection);
+	float h = AttenuationRadius2 - length2(oc - b * RayDirection);
+	if (h > 0)
+	{
+		float2 t = -b + float2(-1, +1) * sqrt(h);
+		TMin = max(t.x, TMin);
+		TMax = min(t.y, TMax);
+
+		float3 N = GetNormal(LightId);
+		// now clip against the cone of the spot
+		const float CosOuterConeAngle = GetCosConeAngles(LightId).x;
+		const float SinOuterConeAngle = sqrt(1 - CosOuterConeAngle * CosOuterConeAngle);
+
+		float2 TCone = IntersectInfiniteConeVolume(RayOrigin, RayDirection, Center, N * CosOuterConeAngle, SinOuterConeAngle);
+		if (TCone.x < TCone.y)
+		{
+			TMin = max(TCone.x, TMin);
+			TMax = min(TCone.y, TMax);
+			float3 C = GetColor(LightId);
+			float LightImportance = max3(C.x, C.y, C.z);
+			return CreateEquiAngularSampler(LightImportance, Center, RayOrigin, RayDirection, TMin, TMax);
+		}
+	}
+	return NullVolumeLightSetup();
+
+}
diff --git a/Engine/Shaders/Private/PathTracing/Material/PathTracingMaterialSampling.ush b/Engine/Shaders/Private/PathTracing/Material/PathTracingMaterialSampling.ush
--- a/Engine/Shaders/Private/PathTracing/Material/PathTracingMaterialSampling.ush
+++ b/Engine/Shaders/Private/PathTracing/Material/PathTracingMaterialSampling.ush
@@ -22,6 +22,7 @@ PathTracingMaterialSampling.usf: Material sampling functions
 #include "PathTracingSubsurfaceProfile.ush"
 #include "PathTracingEye.ush"
 #include "PathTracingHair.ush"
+#include "PathTracingMedium.ush"
 #include "../../BurleyNormalizedSSSCommon.ush"
 
 // Simplified material shaders for GPU Lightmass, supporting Lambert, Two Sided Foliage and Thin Translucent
@@ -259,6 +260,25 @@ void AccumulateAlbedoNormal(inout FPathTracingPayload Payload, float3 PathThroug
 #endif
 }
 
+void AccumulateAlbedo(float3 SigmaS, float3 PathThroughput, float PathRoughness, inout float3 Albedo)
+{
+	const float MaxRoughness = 0.15; // don't bother past this level
+	const float MinRoughness = 0.10; // capture anything below this level
+	if (PathRoughness >= MaxRoughness)
+	{
+		// no point tracking past a certain roughness
+		return;
+	}
+	// figure out how much the current path should contribute to the denoiser aovs
+	PathThroughput *= 1 - saturate((PathRoughness - MinRoughness) / (MaxRoughness - MinRoughness));
+	// figure out how much the specular lobes should contribute to the total (relative to the current PathRoughness)
+	// if we are above the MinRoughness, use the PathRoughness as are lower bound instead so that multiple bounces don't accumulate too much albedo
+	const float VolumeRoughness = 1.0; // assume volume scattering is diffuse
+	const float RoughnessWeight = saturate((VolumeRoughness - max(MinRoughness, PathRoughness)) / (MaxRoughness - max(MinRoughness, PathRoughness)));
+
+	Albedo += RoughnessWeight * SigmaS * PathThroughput;
+}
+
 void AdjustMaterialRoughness(inout FPathTracingPayload Payload, float PathRoughness)
 {
 #if !SIMPLIFIED_MATERIAL_SHADER
@@ -540,6 +560,9 @@ FMaterialSample SampleMaterial(
 		case SHADINGMODELID_HAIR:
 			Result = Hair_SampleMaterial(RayDirection, Payload, RandSample);
 			break;
+		case SHADINGMODELID_MEDIUM:
+			Result = Medium_SampleMaterial(RayDirection, Payload, RandSample);
+			break;
 		default:
 			Result = Lambert_SampleMaterial(Payload, RandSample);
 			break;
@@ -639,6 +662,9 @@ FMaterialEval EvalMaterial(
 		case SHADINGMODELID_HAIR:
 			Result = Hair_EvalMaterial(IncomingDirection, OutgoingDirection, Payload);
 			break;
+		case SHADINGMODELID_MEDIUM:
+			Result = Medium_EvalMaterial(IncomingDirection, OutgoingDirection, Payload);
+			break;
 		default:
 			Result = Lambert_EvalMaterial(OutgoingDirection, Payload);
 			break;
diff --git a/Engine/Shaders/Private/PathTracing/Material/PathTracingMedium.ush b/Engine/Shaders/Private/PathTracing/Material/PathTracingMedium.ush
new file mode 100644
--- /dev/null
+++ b/Engine/Shaders/Private/PathTracing/Material/PathTracingMedium.ush
@@ -0,0 +1,66 @@
+// Copyright Epic Games, Inc. All Rights Reserved.
+
+/*=============================================================================================
+PathTracingMedium.usf: Volumetric phase function
+   To support atmosphere, we represent a blend of two components: rayleigh scattering and an HG lobe
+   In the future, we may add additional HG lobe for cases like cloud rendering
+===============================================================================================*/
+#pragma once
+
+#include "../../ParticipatingMediaCommon.ush"
+
+#define USE_UNIFORM_PHASE_FUNCTION 0
+
+FMaterialSample Medium_SampleMaterial(
+	float3 IncomingDirection,
+	FPathTracingPayload Payload,
+	float4 RandSample)
+{
+	const float3 RayleighWeight = Payload.BaseColor;
+	const float3 HGWeight = Payload.CustomData.xyz;
+	const float G = Payload.CustomData.w;
+
+#if USE_UNIFORM_PHASE_FUNCTION
+	const float4 Result = UniformSampleSphere(RandSample.xy);
+	// uniform scattering
+	return CreateMaterialSample(Result.xyz, RayleighWeight + HGWeight, Result.w, 1.0, 1.0);
+#else
+	const float LobeProb = LobeSelectionProb(RayleighWeight, HGWeight);
+
+	const float CosTheta = RandSample.x < LobeProb ?
+		RayleighPhaseInvertCdf(RandSample.x / LobeProb) :
+		HenyeyGreensteinPhaseInvertCDF((RandSample.x - LobeProb) / (1.0 - LobeProb), G);
+	const float SinTheta = sqrt(saturate(1.0 - CosTheta * CosTheta));
+	const float Phi = RandSample.y * (2.0 * PI);
+
+	FMaterialSample Result = NullMaterialSample();
+	Result.Direction = TangentToWorld(float3(cos(Phi) * SinTheta, sin(Phi) * SinTheta, CosTheta), -IncomingDirection);
+	Result.Roughness = 1.0;
+	AddLobeWithMIS(Result.Weight, Result.Pdf, RayleighWeight, RayleighPhase(CosTheta), LobeProb);
+	AddLobeWithMIS(Result.Weight, Result.Pdf, HGWeight      , HenyeyGreensteinPhase(G, CosTheta )  , 1.0 - LobeProb);
+	return Result;
+#endif
+}
+
+FMaterialEval Medium_EvalMaterial(
+	float3 IncomingDirection,
+	float3 OutgoingDirection,
+	FPathTracingPayload Payload
+)
+{
+	const float3 RayleighWeight = Payload.BaseColor;
+	const float3 HGWeight = Payload.CustomData.xyz;
+	const float G = Payload.CustomData.w;
+#if USE_UNIFORM_PHASE_FUNCTION
+	// simple omni-directional evaluation
+	return CreateMaterialEval(RayleighWeight + HGWeight, 1.0 / (4 * PI));
+#else
+	const float LobeProb = LobeSelectionProb(RayleighWeight, HGWeight);
+
+	const float CosTheta = -dot(IncomingDirection, OutgoingDirection);
+	FMaterialEval Result = NullMaterialEval();
+	AddLobeWithMIS(Result.Weight, Result.Pdf, RayleighWeight, RayleighPhase(CosTheta), LobeProb);
+	AddLobeWithMIS(Result.Weight, Result.Pdf, HGWeight      , HenyeyGreensteinPhase(G, CosTheta), 1.0 - LobeProb);
+	return Result;
+#endif
+}
diff --git a/Engine/Shaders/Private/PathTracing/PathTracingBuildAtmosphereLUT.usf b/Engine/Shaders/Private/PathTracing/PathTracingBuildAtmosphereLUT.usf
new file mode 100644
--- /dev/null
+++ b/Engine/Shaders/Private/PathTracing/PathTracingBuildAtmosphereLUT.usf
@@ -0,0 +1,86 @@
+// Copyright Epic Games, Inc. All Rights Reserved.
+
+#include "../Common.ush"
+#include "../SkyAtmosphereCommon.ush"
+#include "./Volume/PathTracingAtmosphereCommon.ush"
+
+uint NumSamples;
+uint Resolution;
+RWTexture2D<float3> AtmosphereOpticalDepthLUT;
+
+[numthreads(THREADGROUPSIZE_X, THREADGROUPSIZE_Y, 1)]
+void PathTracingBuildAtmosphereOpticalDepthLUTCS(uint2 DispatchThreadId : SV_DispatchThreadID)
+{
+	if (any(DispatchThreadId >= Resolution))
+	{
+		return;
+	}
+
+	const float2 UV = (DispatchThreadId + 0.5) / float(Resolution);
+
+	const float T = Atmosphere.TopRadiusKm;
+	const float B = Atmosphere.BottomRadiusKm;
+
+	const float2 Result = FromAtmosphereOpticalDepthLUTUV(UV, B, T);
+
+	const float h = Result.x;
+	const float c = Result.y;
+
+	const float Rb = B / h;
+	const float Rt = T / h;
+	const float deltaT = Rt * Rt + (c * c - 1.0);
+	const float deltaB = Rb * Rb + (c * c - 1.0);
+	const float rsT = (-c + sqrt(max(deltaT, 0.0))); // distance to top atmosphere (farthest)
+	const float rsB0 = (-c - sqrt(max(deltaB, 0.0))); // distance to ground (nearest)
+	const float rsB1 = (-c + sqrt(max(deltaB, 0.0))); // distance to ground (farthest)
+
+	float3 Tau = 0.0;
+	if (deltaB > 0.0 && c < 0.0)
+	{
+		// we see the ground, integrate the variable portions numerically, and the part through the planet center analytically (since it is homogeneous)
+		// this greatly increases the precision of the integral, reducing artifacts at rendertime
+		const int NS0 = NumSamples / 2;
+		const int NS1 = NumSamples - NS0;
+		{
+			// part above ground
+			const float drs = rsB0 / float(NS0);
+			const float ds = drs * h;
+			for (int i = 0; i < NS0; i++)
+			{
+				const float rs = (i + 0.5) * drs;
+				const float ViewHeightS = h * sqrt(1.0 + rs * rs + 2 * c * rs);
+				Tau += SampleAtmosphereMediumRGB(float3(0, 0, ViewHeightS)).Extinction * ds;
+			}
+		}
+		// planet core: homogeneous, so optical depth is simple
+		Tau += h * (rsB1 - rsB0) * SampleAtmosphereMediumRGB(float3(0, 0, 0)).Extinction;
+		{
+			// part on the other side of the planet
+			const float drs = (rsT - rsB1) / float(NS1);
+			const float ds = drs * h;
+			for (int i = 0; i < NS1; i++)
+			{
+				const float rs = rsB1 + (i + 0.5) * drs;
+				const float ViewHeightS = h * sqrt(1.0 + rs * rs + 2 * c * rs);
+				Tau += SampleAtmosphereMediumRGB(float3(0, 0, ViewHeightS)).Extinction * ds;
+			}
+		}
+	}
+	else
+	{
+		// we see the sky only, single integral
+		const int NS = NumSamples;
+		const float drs = rsT / float(NS);
+		const float ds = drs * h;
+		for (int i = 0; i < NS; i++)
+		{
+			const float rs = (i + 0.5) * drs;
+			const float ViewHeightS = h * sqrt(1.0 + rs * rs + 2 * c * rs);
+			Tau += SampleAtmosphereMediumRGB(float3(0, 0, ViewHeightS)).Extinction * ds;
+		}
+	}
+
+	// store the integrated density (also known as optical depth)
+	// note that this is a unitless quantity. We store this rather than transmittance so we can have the most numerical precision to perform subtractions along the ray
+	AtmosphereOpticalDepthLUT[DispatchThreadId] = Tau;
+}
diff --git a/Engine/Shaders/Private/PathTracing/PathTracingCommon.ush b/Engine/Shaders/Private/PathTracing/PathTracingCommon.ush
--- a/Engine/Shaders/Private/PathTracing/PathTracingCommon.ush
+++ b/Engine/Shaders/Private/PathTracing/PathTracingCommon.ush
@@ -18,6 +18,10 @@
 
 #define PATH_TRACING_PAYLOAD_NORMAL_BITS  24
 
+// shading model extensions beyond the default count
+#define SHADINGMODELID_MEDIUM  (SHADINGMODELID_NUM + 1)
+
+
 uint PayloadEncodeUnitVector(float3 Normal)
 {
 	const int BITS = PATH_TRACING_PAYLOAD_NORMAL_BITS;
@@ -51,6 +55,7 @@ struct FPackedPathTracingPayload : FMinimalPayload
 	uint2 GetPixelCoord() { return uint2(PackedPixelCoord & 0xFFFF, PackedPixelCoord >> 16); }
 
 	FRayCone GetRayCone() { return UnpackRayCone(PackedRayCone); }
+	void SetRayCone(FRayCone RayCone) { PackedRayCone = PackRayCone(RayCone); }
 
 	uint GetFlags() { return (PackedData[7] >> 23) & 0x3F; }
 	void SetFlag(uint Flag) { PackedData[7] |= Flag << 23; }
@@ -173,6 +178,7 @@ struct FPathTracingPayload : FMinimalPayload
 		return ShadingModelID == SHADINGMODELID_TWOSIDED_FOLIAGE ||
 			   ShadingModelID == SHADINGMODELID_THIN_TRANSLUCENT || 
 			   ShadingModelID == SHADINGMODELID_HAIR ||
+			   ShadingModelID == SHADINGMODELID_MEDIUM ||
 			   IsMaterialSolidGlass();
 	}
 
diff --git a/Engine/Shaders/Private/PathTracing/PathTracingCore.ush b/Engine/Shaders/Private/PathTracing/PathTracingCore.ush
--- a/Engine/Shaders/Private/PathTracing/PathTracingCore.ush
+++ b/Engine/Shaders/Private/PathTracing/PathTracingCore.ush
@@ -17,9 +17,12 @@
 
 #include "../ShadingModels.ush"
 #include "./Utilities/PathTracingRandomSequence.ush"
+#include "./Utilities/PathTracingRIS.ush"
 #include "./Light/PathTracingLightSampling.ush"
 #include "./Light/PathTracingLightGrid.ush"
 #include "./Material/PathTracingMaterialSampling.ush"
+#include "./Volume/PathTracingVolume.ush"
+#include "./Volume/PathTracingVolumeSampling.ush"
 
 float BlendFactor;
 uint Iteration;
@@ -33,6 +36,10 @@ uint MaxSSSBounces;
 // 2: both Material and Light
 uint MISMode;
 
+// 0: only Density sampling
+// 1: Light sampling
+uint VolumeMISMode;
+
 uint ApproximateCaustics;
 uint EnableCameraBackfaceCulling;
 uint EnableDirectLighting;
@@ -97,11 +104,12 @@ void AccumulateRadiance(inout float3 TotalRadiance, float3 PathRadiance, bool bI
 	TotalRadiance += PathRadiance;
 }
 
-FPathTracingPayload TraceTransparentRay(inout FPathState PathState, bool bIsCameraRay, bool bLastBounce, bool bIncludeEmission, uint NumVisibleLights)
+FPathTracingPayload TraceTransparentRay(inout FPathState PathState, bool bIsCameraRay, bool bLastBounce, bool bIncludeEmission, uint NumVisibleLights, inout FVolumeSegment VolumeSegment)
 {
 	const uint RayFlags = bIsCameraRay && EnableCameraBackfaceCulling ? RAY_FLAG_CULL_BACK_FACING_TRIANGLES : 0;
 	const uint MissShaderIndex = 0;
-	float SelectionWeightSum = 0;
+	FRISContext HitSample = InitRISContext();
+	FRISContext VolSample = InitRISContext();
 	float3 PayloadThroughput = PathState.PathThroughput;
 	FPathTracingPayload Payload;
 	if (!bIncludeEmission && bLastBounce)
@@ -110,6 +118,13 @@ FPathTracingPayload TraceTransparentRay(inout FPathState PathState, bool bIsCame
 		PathState.PathThroughput = 0;
 		return Payload;
 	}
+	FVolumeIntersectionList VolumeIsectList = VolumeIntersect(PathState.Ray.Origin, PathState.Ray.Direction, PathState.Ray.TMin, PathState.Ray.TMax);
+	if (VolumeIsectList.HitBlocker())
+	{
+		// We will hit the volume's blocker, 
+		// Clip the ray just before the planet hit so that we can intersect all other geometry but this one
+		PathState.Ray.TMax = asfloat(asuint(VolumeIsectList.BlockerHitT) - 1);
+	}
 	for (;;)
 	{
 		// avoid double counting what was captured by the skylight
@@ -128,16 +143,168 @@ FPathTracingPayload TraceTransparentRay(inout FPathState PathState, bool bIsCame
 			PathState.Ray,
 			PackedPayload);
 
+		if (PackedPayload.IsMiss() && VolumeIsectList.HitBlocker())
+		{
+			// we didn't hit any real geometry, but we did hit the volume's blocker geometry
+			// create a virtual hit with it here
+			PackedPayload = VolumeGetBlockerHit(VolumeIsectList.BlockerID, PathState.Ray.Origin, PathState.Ray.Direction, VolumeIsectList.BlockerHitT);
+			PackedPayload.SetRayCone(PropagateRayCone(PathState.RayCone, 0.0, VolumeIsectList.BlockerHitT));
+		}
+
 		// Loop over lights to capture their contribution
 		// #dxr_todo: if we have lots of lights, having some hierarchical structure would be better ....
 		for (uint LightId = 0; LightId < NumVisibleLights; ++LightId)
 		{
 			RayDesc LightRay = PathState.Ray;
 			LightRay.TMax = PackedPayload.IsMiss() ? PathState.Ray.TMax : PackedPayload.HitT;
-			float3 LightRadiance = TraceLight(LightRay, LightId).Radiance;
-			AccumulateRadiance(PathState.Radiance, PathState.PathThroughput * LightRadiance, bIsCameraRay);
+			FLightHit LightHit = TraceLight(LightRay, LightId);
+			if (LightHit.IsHit())
+			{
+				// create a virtual surface hit so we only need to worry about volume transmission in one place
+				// NOTE: returning only a single light hit here causes lights at infinity to occlude each other
+				//       this is most likely what artists would want (imagine a scene with two suns for example)
+				//       but is not consistent with how reflections behave and the sorting order will be arbitrary
+				FPathTracingPayload LightHitPayload = (FPathTracingPayload)0; // clear all fields
+				LightHitPayload.HitT = LightHit.HitT;
+				LightHitPayload.RayCone = PropagateRayCone(PathState.RayCone, 0.0, LightHit.HitT);
+				LightHitPayload.Radiance = LightHit.Radiance;
+				LightHitPayload.ShadingModelID = SHADINGMODELID_UNLIT;
+				// alpha composite assumes premult, which allows emissive hits that are fully transparent
+				LightHitPayload.BlendingMode = RAY_TRACING_BLEND_MODE_ALPHA_COMPOSITE;
+				LightHitPayload.PrimitiveLightingChannelMask = 7;
+				LightHitPayload.SetFrontFace();
+				PackedPayload = PackPathTracingPayload(LightHitPayload);
+			}
 		}
 
+		// Volume Transmittance + Scatter Segment/Point selection
+		// This loop splits the potential volumetric segments into sub-intervals that do not overlap.
+		while (VolumeIsectList.HitVolume())
+		{
+			// extract the nearest interval from the list of segments
+			FVolumeIntersectionInterval Interval = VolumeIsectList.GetCurrentInterval();
+
+			if (PackedPayload.IsHit())
+			{
+				if (PackedPayload.HitT <= Interval.VolumeTMin)
+				{
+					// surface hit is in-front of the current interval, we are done
+					break;
+				}
+				// clip current volume segment to the next surface hit
+				Interval.VolumeTMax = min(PackedPayload.HitT, Interval.VolumeTMax);
+			}
+
+			// if this isn't our last bounce, potentially keep this segment for later so we can integrate local lighting
+			if (!bLastBounce)
+			{
+				// TODO: can we improve the selection probability for each ray segment somehow?
+				// NOTE: this only matters if we are tracing through several transparent hits or if there are multiple overlapping volumes
+				float3 Contrib = PathState.PathThroughput;
+				float SelectionWeight = max3(Contrib.x, Contrib.y, Contrib.z);
+				if (VolSample.Accept(SelectionWeight, RandomSequence_GenerateSample1D(PathState.RandSequence)))
+				{
+					// store this volume segment for later
+					VolumeSegment.Throughput = PathState.PathThroughput / SelectionWeight;
+					VolumeSegment.Interval = Interval;
+				}
+			}
+
+#if PATH_TRACING_USE_TRANSMITTANCE_FUNC
+			// NOTE: when using analytic raymarching, we can skip all raymarching on the last bounce because we know we don't need a scattering point
+			if (!bLastBounce)
+#endif
+			{
+				// remember the start of our interval before ray marching messes with it
+				float StartVolumeTMin = Interval.VolumeTMin;
+				// Ray marching loop
+				float3 PathThroughput = PathState.PathThroughput;
+				// Limit number of steps to prevent timeouts // FIXME: This biases the result!
+				for (int Step = 0; Step < MaxRaymarchSteps; Step++)
+				{
+					// figure out optimal bounds for the current ray segment
+					float3 SigmaBar = VolumeGetDensityBounds(PathState.Ray.Origin, PathState.Ray.Direction, Interval).SigmaMax;
+#if PATH_TRACING_USE_TRANSMITTANCE_FUNC == 0
+					// if we are not using analytical transmittance, a tight majorant could prevent us from "seeing" hits that match the density exactly
+					// leading to heavy noise on bright objects embedded in the volume
+					// however if we can track transmittance analytically, this workaround is not needed as we will re-compute a precise answer after ray marching
+					SigmaBar *= bIsCameraRay ? 2 : 1;
+#endif
+
+					// there is some atmosphere present along the current ray, ray march until we get to the current hit (or exit the atmosphere)
+					float3 ColorChannelPdf = PathThroughput;
+
+					// Sample the distance to the next interaction
+					float2 RandValue = RandomSequence_GenerateSample2D(PathState.RandSequence);
+					float DeltaT = SampleSpectralTransmittance(RandValue.x, SigmaBar, ColorChannelPdf);
+					if (DeltaT < 0.0)
+					{
+						// no more energy left in the path
+						break;
+					}
+
+					if (Interval.VolumeTMin + DeltaT < Interval.VolumeTMax)
+					{
+						Interval.VolumeTMin += DeltaT;
+						// our ray marching step stayed inside the atmo and is still in front of the next hit
+						PathThroughput *= EvaluateSpectralTransmittanceHit(DeltaT, SigmaBar, ColorChannelPdf).xyz;
+
+						// find out how much volume exists at the current point
+						float3 Ro = PathState.Ray.Origin;
+						float3 Rd = PathState.Ray.Direction;
+						float3 TranslatedWorldPos = Ro + Interval.VolumeTMin * Rd;
+						FVolumeShadedResult Result = VolumeGetDensity(TranslatedWorldPos, Interval);
+
+						// clamp to make sure we never exceed the majorant (should not be the case, but need to avoid any possible numerical issues)
+						float3 SigmaT = min(Result.SigmaT, SigmaBar);
+						float3 SigmaN = SigmaBar - SigmaT;
+
+#if PATH_TRACING_USE_TRANSMITTANCE_FUNC == 0
+						if (!bLastBounce) // check now if this loop is also computing transmittance
+#endif
+						{
+							float3 SigmaS = min(Result.SigmaSRayleigh + Result.SigmaSHG, SigmaT);
+							// accumulate a signal for the denoiser
+							AccumulateAlbedo(SigmaS, PathThroughput, PathState.PathRoughness, PathState.Albedo);
+							float3 Contrib = PathThroughput * SigmaS;
+							float SelectionWeight = max3(Contrib.x, Contrib.y, Contrib.z);
+							if (HitSample.Accept(SelectionWeight, RandValue.y))
+							{
+								// stash this hit for next time
+								Payload = CreateMediumHitPayload(Interval.VolumeTMin, TranslatedWorldPos, Result);
+								Payload.RayCone = PropagateRayCone(PathState.RayCone, 0.0, Interval.VolumeTMin);
+								PayloadThroughput = Contrib / SelectionWeight;
+							}
+						}
+
+						// keep tracing through the volume
+						PathThroughput *= SigmaN;
+					}
+					else
+					{
+						// update the path throughput, knowing that we escaped the medium
+						PathThroughput *= EvaluateSpectralTransmittanceMiss(Interval.VolumeTMax - Interval.VolumeTMin, SigmaBar, ColorChannelPdf).xyz;
+						// exit the ray marching loop
+						break;
+					}
+				}
+				// restore interval start now that we are finished with ray marching
+				Interval.VolumeTMin = StartVolumeTMin;
+#if PATH_TRACING_USE_TRANSMITTANCE_FUNC == 0
+				// record the change in transmittance since we already computed it
+				PathState.PathThroughput = PathThroughput;
+#endif
+			}
+#if PATH_TRACING_USE_TRANSMITTANCE_FUNC
+			// analytically handle the transmittance to the next surface (which will be higher quality than the transmittance implicitly computed above and prevent noise with emissive surfaces in volumes)
+			PathState.PathThroughput = VolumeGetTransmittance(PathState.PathThroughput, PathState.Ray.Origin, PathState.Ray.Direction, Interval, PathState.RandSequence);
+#endif
+			// Update our current volume intersection list to reflect the fact that we have made it to the end of this interval
+			// This will either clip the finished portions of active segments, or remove the segments we have fully marched through
+			VolumeIsectList = VolumeIsectList.Update(Interval.VolumeTMax);
+		}
+		// proceed to handling the surface hit (if any)
+
 		if (PackedPayload.IsMiss())
 		{
 			// Ray didn't hit any real geometry
@@ -168,6 +335,7 @@ FPathTracingPayload TraceTransparentRay(inout FPathState PathState, bool bIsCame
 		float3 Transparency = GetMaterialTransparency(HitPayload, PathState.Ray.Direction);
 
 		// account for Beer's law through the currently active medium
+		// TODO: merge this with the volume raymarching?
 		PathState.PathThroughput *= exp(-PathState.SigmaT * (HitPayload.HitT - PathState.Ray.TMin));
 
 		if (!bLastBounce)
@@ -178,27 +346,7 @@ FPathTracingPayload TraceTransparentRay(inout FPathState PathState, bool bIsCame
 			AccumulateAlbedoNormal(HitPayload, PathState.PathThroughput, PathState.PathRoughness, PathState.Albedo, PathState.Normal);
 
 			float SelectionWeight = max3(Contrib.x, Contrib.y, Contrib.z);
-			SelectionWeightSum += SelectionWeight;
-			bool AcceptHit = false;
-
-			// weighted reservoir sampling
-			if (SelectionWeight > 0)
-			{
-				if (SelectionWeight < SelectionWeightSum)
-				{
-					// the acceptance probability is not 1.0
-					// generate a random number to see if we should accept this hit
-					float RandValue = RandomSequence_GenerateSample1D(PathState.RandSequence);
-					AcceptHit = RandValue * SelectionWeightSum < SelectionWeight;
-				}
-				else
-				{
-					// accept automatically on the first hit
-					AcceptHit = true;
-				}
-			}
-
-			if (AcceptHit)
+			if (HitSample.Accept(SelectionWeight, RandomSequence_GenerateSample1D(PathState.RandSequence)))
 			{
 				// stash this hit for next time
 				Payload = HitPayload;
@@ -206,8 +354,15 @@ FPathTracingPayload TraceTransparentRay(inout FPathState PathState, bool bIsCame
 			}
 		}
 
+		if (HitPayload.HitT == POSITIVE_INFINITY)
+		{
+			// if our hit was against an infinite light, exit now
+			break;
+		}
+
 		// account for local transparency change
 		PathState.PathThroughput *= Transparency;
+
 		// prepare next step around the loop
 		// retrace the exact same ray with TMin one ulp past the hit we just found
 		PathState.Ray.TMin = asfloat(asuint(HitPayload.HitT) + 1);
@@ -218,10 +373,13 @@ FPathTracingPayload TraceTransparentRay(inout FPathState PathState, bool bIsCame
 		}
 	}
 
-	if (SelectionWeightSum > 0)
+	// normalization will be 0 if we didn't pick any volume segment
+	VolumeSegment.Throughput *= VolSample.GetNormalization();
+
+	if (HitSample.HasSample())
 	{
 		// if we stored a valid hit in the payload, reset the path throughput to this point
-		PathState.PathThroughput = PayloadThroughput * SelectionWeightSum;
+		PathState.PathThroughput = PayloadThroughput * HitSample.GetNormalization();
 	}
 	else
 	{
@@ -231,7 +389,7 @@ FPathTracingPayload TraceTransparentRay(inout FPathState PathState, bool bIsCame
 	return Payload;
 }
 
-float3 TraceTransparentVisibilityRay(RayDesc Ray, uint2 PixelCoord, float PathRoughness)
+float3 TraceTransparentVisibilityRay(RayDesc Ray, uint2 PixelCoord, float PathRoughness, inout RandomSequence RandSequence)
 {
 	const uint RayFlags = RAY_FLAG_ACCEPT_FIRST_HIT_AND_END_SEARCH | RAY_FLAG_SKIP_CLOSEST_HIT_SHADER;
 	const uint InstanceInclusionMask = RAY_TRACING_MASK_SHADOW | RAY_TRACING_MASK_THIN_SHADOW;
@@ -254,12 +412,25 @@ float3 TraceTransparentVisibilityRay(RayDesc Ray, uint2 PixelCoord, float PathRo
 
 	if (PackedPayload.IsMiss())
 	{
+		FVolumeIntersectionList VolumeIsectList = VolumeIntersect(Ray.Origin, Ray.Direction, Ray.TMin, Ray.TMax);
+		if (VolumeIsectList.HitBlocker())
+		{
+			// blockers are opaque -- done!
+			return 0.0;
+		}
+		float3 Throughput = PackedPayload.GetRayThroughput();
+		if (VolumeIsectList.HitVolume())
+		{
+			Throughput = VolumeGetTransmittance(Throughput, Ray.Origin, Ray.Direction, VolumeIsectList, RandSequence);
+		}
+
 		// we didn't hit any geo - just return whatever throughput is left
-		return PackedPayload.GetRayThroughput() * exp(-max(PackedPayload.GetTau() * AbsorptionScale, 0.0));
+		float3 Tau = PackedPayload.GetTau() * AbsorptionScale;
+		return Throughput * exp(-max(Tau, 0.0));
 	}
 
 	// we registered a hit -- assumed to be fully opaque
-	return 0;
+	return 0.0;
 }
 
 struct FProbeResult
@@ -417,35 +588,13 @@ bool ProcessSubsurfaceRandomWalk(inout FPathTracingPayload Payload, inout float3
 	{
 		// TODO: add Dwivedi sampling for faster escape in low albedo media ...
 
-		// Pick a color channel for transmittance sampling:
-		float3 ColorChannelWeight = PathThroughput * Albedo;
-		float3 ColorChannelCdf = float3(ColorChannelWeight.x,
-			ColorChannelWeight.x + ColorChannelWeight.y,
-			ColorChannelWeight.x + ColorChannelWeight.y + ColorChannelWeight.z);
-		if (ColorChannelCdf.z == 0)
+		float3 ColorChannelPdf = PathThroughput * Albedo;
+		Ray.TMax = SampleSpectralTransmittance(RandSample.z, SigmaT, ColorChannelPdf);
+		if (Ray.TMax < 0.0)
 		{
 			// no more energy left in the path
 			break;
 		}
-		ColorChannelCdf.xy /= ColorChannelCdf.z;
-		float3 ColorChannelPdf = float3(ColorChannelCdf.x, ColorChannelCdf.y - ColorChannelCdf.x, 1 - ColorChannelCdf.y);
-		// pick a maximum distance to scatter
-		const float OneMinusEpsilon = 0.99999994; // 32-bit float just before 1.0
-		if (RandSample.z < ColorChannelCdf.x)
-		{
-			float RescaleRand = min(RandSample.z / ColorChannelPdf.x, OneMinusEpsilon);
-			Ray.TMax = -log(1 - RescaleRand) * SSSRadius.x;
-		}
-		else if (RandSample.z < ColorChannelCdf.y)
-		{
-			float RescaleRand = min((RandSample.z - ColorChannelCdf.x) / ColorChannelPdf.y, OneMinusEpsilon);
-			Ray.TMax = -log(1 - RescaleRand) * SSSRadius.y;
-		}
-		else
-		{
-			float RescaleRand = min((RandSample.z - ColorChannelCdf.y) / ColorChannelPdf.z, OneMinusEpsilon);
-			Ray.TMax = -log(1 - RescaleRand) * SSSRadius.z;
-		}
 		for (;;)
 		{
 			FProbeResult ProbeResult = TraceProbeRay(Ray, PixelCoord);
@@ -457,12 +606,7 @@ bool ProcessSubsurfaceRandomWalk(inout FPathTracingPayload Payload, inout float3
 				Ray.Direction = UniformSampleSphere(RandSample.xy).xyz;
 				Ray.TMin = 0;
 				// account for transmittance and scattering up to this point
-				float3 Transmittance = exp(-Ray.TMax * SigmaT);
-				PathThroughput *= SigmaS * Transmittance;
-				// account for transmittance pdf (which is the balance heuristic combination of transmittance sampling in each color channel)
-				float3 TransmittancePdf = SigmaT * Transmittance;
-				PathThroughput /= dot(TransmittancePdf, ColorChannelPdf);
-
+				PathThroughput *= SigmaS * EvaluateSpectralTransmittanceHit(Ray.TMax, SigmaT, ColorChannelPdf).xyz;
 				// NOTE: phase function cancels out since it is being perfectly importance sampled
 
 				// exit the Interface counting loop and keep scattering
@@ -479,12 +623,8 @@ bool ProcessSubsurfaceRandomWalk(inout FPathTracingPayload Payload, inout float3
 			else
 #endif
 			{
-				// account for transmittance and scattering up to this point
-				float3 Transmittance = exp(-ProbeResult.HitT * SigmaT);
-				PathThroughput *= Transmittance;
-				// account for transmittance pdf (which is the balance heuristic combination of transmittance sampling in each color channel)
-				float3 TransmittancePdf = Transmittance; // integral of pdf from HitT to infinity
-				PathThroughput /= dot(Transmittance, ColorChannelPdf);
+				// account for transmittance to the boundary
+				PathThroughput *= EvaluateSpectralTransmittanceMiss(ProbeResult.HitT, SigmaT, ColorChannelPdf).xyz;
 
 				// our short ray hit the geomtry - decide if we should exit or not
 #if SSS_USE_TIR
@@ -691,12 +831,135 @@ bool PathTracingKernel(inout FPathState PathState, int Bounce)
 	// because NEE handled it for us
 	uint NumVisibleLights = bIsCameraRay ? SceneVisibleLightCount : 0;
 
+	FVolumeSegment VolumeSegment = CreateEmptyVolumeSegment();
 	FPathTracingPayload Payload = TraceTransparentRay(
 		PathState,
 		bIsCameraRay,
 		bIsLastBounce,
 		bIncludeEmissive,
-		NumVisibleLights);
+		NumVisibleLights,
+		VolumeSegment);
+
+	// process the returned volume segment if we got one
+	if (VolumeSegment.IsValid() && VolumeMISMode != 0)
+	{
+		const float3 Ro = PathState.Ray.Origin;
+		const float3 Rd = PathState.Ray.Direction;
+		const float VTMin = VolumeSegment.Interval.VolumeTMin;
+		const float VTMax = VolumeSegment.Interval.VolumeTMax;
+
+		float LightPickingCdfSum = 0.0;
+		for (int LightId = 0; LightId < SceneLightCount; LightId++)
+		{
+			FVolumeLightSampleSetup LightSetup = PrepareLightVolumeSample(LightId, Ro, Rd, VTMin, VTMax);
+			if (LightSetup.IsValid())
+			{
+				float LightProb = LightSetup.LightImportance * GetVolumetricScatteringIntensity(LightId);
+				LightPickingCdfSum += LightProb;
+			}
+			LightPickingCdf[LightId] = LightPickingCdfSum;
+		}
+
+		if (LightPickingCdfSum > 0.0)
+		{
+			// at least one light is overlapping with our ray, so we have a chance to sample it
+			float3 RandSample = RandomSequence_GenerateSample3D(PathState.RandSequence);
+			int LightSampleLightId = 0;
+			float LightPickPdf = 0.0;
+			SelectLight(RandSample.x * LightPickingCdfSum, SceneLightCount, LightPickingCdf, LightSampleLightId, LightPickPdf);
+			LightPickPdf /= LightPickingCdfSum;
+
+			// picked a light! now use the equi-angular sampler to pick a position along the ray and store it for later
+			// so that we can compute the light during the ray-marching loop, when the path prefix pdf will be known
+			FVolumeLightSampleSetup LightSetup = PrepareLightVolumeSample(LightSampleLightId, Ro, Rd, VTMin, VTMax);
+			// Should be safe to assume LightSetup.IsValid() is true because otherwise the pdf would have been 0
+			float2 SampleResult = LightSetup.SampleDistance(RandSample.y);
+			float EquiAngularT = SampleResult.x;
+			float EquiAngularPathPdf = SampleResult.y;
+
+			// find out how much volume exists at the current point
+			const float3 TranslatedWorldPos = Ro + EquiAngularT * Rd;
+			const FVolumeShadedResult Result = VolumeGetDensity(TranslatedWorldPos, VolumeSegment.Interval);
+			const float3 SigmaT = Result.SigmaT;
+			const float3 SigmaS = min(Result.SigmaSRayleigh + Result.SigmaSHG, SigmaT);
+
+			if (any(SigmaS > 0))
+			{
+				// Account for the transmittance up to the current point within the current slice of volume
+				VolumeSegment.Interval.VolumeTMax = EquiAngularT;
+				VolumeSegment.Throughput = VolumeGetTransmittance(VolumeSegment.Throughput, Ro, Rd, VolumeSegment.Interval, PathState.RandSequence);
+
+				float3 Contrib = VolumeSegment.Throughput * SigmaS;
+
+				// account for probability of the path prefix
+				Contrib /= EquiAngularPathPdf * LightPickPdf;
+
+				// prepare a minimal payload that describes the hit we need to shade
+				FPathTracingPayload VolPayload = CreateMediumHitPayload(EquiAngularT, TranslatedWorldPos, Result);
+				float4 LightRandValue = RandomSequence_GenerateSample4D(PathState.RandSequence);
+
+				// compute direct light sampling?
+				if (MISMode != 0)
+				{
+					FLightSample LightSample = SampleLight(LightSampleLightId, LightRandValue.xy, TranslatedWorldPos, float3(0, 0, 0));
+					if (LightSample.Pdf > 0)
+					{
+						if (CastsVolumeShadow(LightSampleLightId))
+						{
+							RayDesc LightRay;
+							LightRay.Origin = TranslatedWorldPos;
+							LightRay.TMin = 0;
+							LightRay.Direction = LightSample.Direction;
+							LightRay.TMax = LightSample.Distance;
+							LightSample.RadianceOverPdf *= TraceTransparentVisibilityRay(LightRay, PathState.PixelIndex, 1.0, PathState.RandSequence);
+						}
+						// #dxr_todo: Is it cheaper to fire the ray first? Or eval the material first?
+						if (any(LightSample.RadianceOverPdf > 0))
+						{
+							// Evaluate material
+							FMaterialEval MaterialEval = Medium_EvalMaterial(Rd, LightSample.Direction, VolPayload);
+
+							// Record the contribution
+							float3 LightContrib = Contrib * LightSample.RadianceOverPdf * GetVolumetricScatteringIntensity(LightSampleLightId) * MaterialEval.Weight * MaterialEval.Pdf;
+							if (MISMode == 2)
+							{
+								LightContrib *= MISWeightRobust(LightSample.Pdf, MaterialEval.Pdf);
+							}
+							AccumulateRadiance(PathState.Radiance, LightContrib, bIsCameraRay);
+						}
+					}
+				}
+				// now compute again with phase function MIS
+				if (MISMode != 1)
+				{
+					FMaterialSample MaterialSample = Medium_SampleMaterial(Rd, VolPayload, LightRandValue);
+					if (MaterialSample.Pdf > 0)
+					{
+						RayDesc MaterialRay;
+						MaterialRay.Origin = TranslatedWorldPos;
+						MaterialRay.Direction = MaterialSample.Direction;
+						MaterialRay.TMin = 0.0;
+						MaterialRay.TMax = POSITIVE_INFINITY;
+						FLightHit LightResult = TraceLight(MaterialRay, LightSampleLightId);
+						if (LightResult.IsHit())
+						{
+							float3 LightContrib = Contrib * MaterialSample.Weight * LightResult.Radiance * GetVolumetricScatteringIntensity(LightSampleLightId);
+							if (MISMode == 2)
+							{
+								LightContrib *= MISWeightRobust(MaterialSample.Pdf, LightResult.Pdf);
+							}
+							MaterialRay.TMax = LightResult.HitT;
+							if (CastsVolumeShadow(LightSampleLightId))
+							{
+								LightContrib *= TraceTransparentVisibilityRay(MaterialRay, PathState.PixelIndex, MaterialSample.Roughness, PathState.RandSequence);
+							}
+							AccumulateRadiance(PathState.Radiance, LightContrib, bIsCameraRay);
+						}
+					}
+				}
+			}
+		}
+	}
 
 	if (Payload.IsMiss())
 	{
@@ -705,23 +968,15 @@ bool PathTracingKernel(inout FPathState PathState, int Bounce)
 	}
 
 	// fix dark edges in reflections from shading normals
-	Payload.WorldNormal = AdjustShadingNormal(Payload.WorldNormal, Payload.WorldGeoNormal, PathState.Ray.Direction);
+	Payload.WorldNormal       = AdjustShadingNormal(Payload.WorldNormal      , Payload.WorldGeoNormal, PathState.Ray.Direction);
 	Payload.WorldSmoothNormal = AdjustShadingNormal(Payload.WorldSmoothNormal, Payload.WorldGeoNormal, PathState.Ray.Direction);
 
+
 #if 0
 	PathState.Radiance = 0.5 * Payload.WorldGeoNormal + 0.5;
 	return false;
 #endif
 
-#if 0
-	if (Payload.ShadingModelID == SHADINGMODELID_EYE)
-	{
-		//Payload.WorldNormal = Payload.CustomVector;
-		Payload.ShadingModelID = SHADINGMODELID_DEFAULT_LIT;
-		Payload.Anisotropy = 0;
-	}
-#endif
-
 	// As soon as the path is blurry enough, we can get away with diffuse sampling only
 	const bool bSimplifySSS = PathState.PathRoughness >= 0.15;
 
@@ -742,6 +997,12 @@ bool PathTracingKernel(inout FPathState PathState, int Bounce)
 	}
 
 	FLightLoopCount LightLoopCount = LightGridLookup(Payload.TranslatedWorldPos);
+	if (Payload.ShadingModelID == SHADINGMODELID_MEDIUM && VolumeMISMode != 0)
+	{
+		// if we are using the volume segment for local lights, exclude them from the light loop below
+		LightLoopCount.NumLights = SceneInfiniteLightCount;
+		LightLoopCount.NumMISLights = SceneInfiniteLightCount;
+	}
 
 	// visualize the # of lights
 	if (VisualizeLightGrid)
@@ -760,6 +1021,8 @@ bool PathTracingKernel(inout FPathState PathState, int Bounce)
 
 	const bool bDoLightLoop = EnableDirectLighting != 0 || Bounce > 0;
 
+	const bool bIsVolumeSample = Payload.ShadingModelID == SHADINGMODELID_MEDIUM;
+
 	float LightPickingCdfSum = 0;
 
 	// If we are using Light sampling and the material can use it ...
@@ -770,12 +1033,17 @@ bool PathTracingKernel(inout FPathState PathState, int Bounce)
 		float3 WorldNormal = Payload.WorldNormal;
 		uint PrimitiveLightingChannelMask = Payload.PrimitiveLightingChannelMask;
 
-		bool IsTransmissiveMaterial = ENABLE_TRANSMISSION && Payload.IsMaterialTransmissive();
+		const bool bIsTransmissiveMaterial = ENABLE_TRANSMISSION && Payload.IsMaterialTransmissive();
 
 		for (uint Index = 0, Num = LightLoopCount.NumLights; Index < Num; ++Index)
 		{
 			uint LightIndex = GetLightId(Index, LightLoopCount);
-			LightPickingCdfSum += EstimateLight(LightIndex, TranslatedWorldPos, WorldNormal, PrimitiveLightingChannelMask, IsTransmissiveMaterial);
+			float LightEstimate = EstimateLight(LightIndex, TranslatedWorldPos, WorldNormal, PrimitiveLightingChannelMask, bIsTransmissiveMaterial);
+			if (bIsVolumeSample)
+			{
+				LightEstimate *= GetVolumetricScatteringIntensity(LightIndex);
+			}
+			LightPickingCdfSum += LightEstimate;
 			LightPickingCdf[Index] = LightPickingCdfSum;
 		}
 
@@ -798,8 +1066,8 @@ bool PathTracingKernel(inout FPathState PathState, int Bounce)
 				if (CastsShadow(LightId))
 				{
 					// for transmissive materials, bias the position to the other side of the surface if the light is coming from behind
-					const float SignedPositionBias = IsTransmissiveMaterial ? sign(dot(Payload.WorldGeoNormal, LightSample.Direction)) : 1.0;
-					const bool bEnableShadowTermFix = Payload.ShadingModelID != SHADINGMODELID_HAIR;
+					const float SignedPositionBias = bIsTransmissiveMaterial ? sign(dot(Payload.WorldGeoNormal, LightSample.Direction)) : 1.0;
+					const bool bEnableShadowTermFix = Payload.ShadingModelID != SHADINGMODELID_HAIR && Payload.ShadingModelID != SHADINGMODELID_MEDIUM;
 					RayDesc LightRay;
 					LightRay.Origin = TranslatedWorldPos;
 					LightRay.TMin = 0;
@@ -809,7 +1077,12 @@ bool PathTracingKernel(inout FPathState PathState, int Bounce)
 
 					float AvgRoughness = ApproximateCaustics ? GetAverageRoughness(Payload) : 0.0;
 
-					LightSample.RadianceOverPdf *= TraceTransparentVisibilityRay(LightRay, PathState.PixelIndex, AvgRoughness);
+					LightSample.RadianceOverPdf *= TraceTransparentVisibilityRay(LightRay, PathState.PixelIndex, AvgRoughness, PathState.RandSequence);
+				}
+
+				if (bIsVolumeSample)
+				{
+					LightSample.RadianceOverPdf *= GetVolumetricScatteringIntensity(LightId);
 				}
 
 				// #dxr_todo: Is it cheaper to fire the ray first? Or eval the material first?
@@ -861,7 +1134,7 @@ bool PathTracingKernel(inout FPathState PathState, int Bounce)
 	//   The second tweak is to add a sqrt() around the probability to soften the termination probability (paths will last
 	//  a little longer). This allows paths to go a bit deeper than the naive heuristic while still allowing them to terminate
 	//  early. This makes RR effective from the very first bounce without needing to delay it.
-	float ContinuationProb = sqrt(saturate(max(NextPathThroughput.x, max(NextPathThroughput.y, NextPathThroughput.z)) / max(PathState.PathThroughput.x, max(PathState.PathThroughput.y, PathState.PathThroughput.z))));
+	float ContinuationProb = sqrt(saturate(max3(NextPathThroughput.x, NextPathThroughput.y, NextPathThroughput.z) / max3(PathState.PathThroughput.x, PathState.PathThroughput.y, PathState.PathThroughput.z)));
 	if (ContinuationProb < 1)
 	{
 		// If there is some chance we should terminate the ray, draw an extra random value
@@ -933,6 +1206,11 @@ bool PathTracingKernel(inout FPathState PathState, int Bounce)
 			}
 
 			float3 LightContrib = PathState.PathThroughput * LightResult.Radiance;
+			if (bIsVolumeSample)
+			{
+				LightContrib *= GetVolumetricScatteringIntensity(LightId);
+			}
+
 			if (bUseMIS)
 			{
 				float PreviousCdfValue = 0.0;
@@ -951,7 +1229,7 @@ bool PathTracingKernel(inout FPathState PathState, int Bounce)
 				{
 					RayDesc LightRay = PathState.Ray;
 					LightRay.TMax = LightResult.HitT;
-					LightContrib *= TraceTransparentVisibilityRay(LightRay, PathState.PixelIndex, ShadowRayRoughness);
+					LightContrib *= TraceTransparentVisibilityRay(LightRay, PathState.PixelIndex, ShadowRayRoughness, PathState.RandSequence);
 				}
 				// the light made some contribution, and there was nothing along the shadow ray
 				AccumulateRadiance(PathState.Radiance, LightContrib, bIsCameraRay);
diff --git a/Engine/Shaders/Private/PathTracing/Utilities/PathTracingRIS.ush b/Engine/Shaders/Private/PathTracing/Utilities/PathTracingRIS.ush
new file mode 100644
--- /dev/null
+++ b/Engine/Shaders/Private/PathTracing/Utilities/PathTracingRIS.ush
@@ -0,0 +1,32 @@
+// Copyright Epic Games, Inc. All Rights Reserved.
+
+#pragma once
+
+// Implementation of the core logic of Weighted Resampled Importance Sampling
+// The idea is to simulate sampling a discrete pdf built of individually weighted samples
+
+struct FRISContext
+{
+	float WeightSum;
+
+	bool Accept(float Weight, float RandValue)
+	{
+		WeightSum += Weight;
+		return Weight > 0 ? (Weight < WeightSum ? RandValue * WeightSum < Weight : true) : false;
+	}
+
+	bool HasSample()
+	{
+		return WeightSum > 0;
+	}
+
+	float GetNormalization()
+	{
+		return WeightSum;
+	}
+};
+
+FRISContext InitRISContext()
+{
+	return (FRISContext)0;
+}
diff --git a/Engine/Shaders/Private/PathTracing/Volume/PathTracingAtmosphere.ush b/Engine/Shaders/Private/PathTracing/Volume/PathTracingAtmosphere.ush
new file mode 100644
--- /dev/null
+++ b/Engine/Shaders/Private/PathTracing/Volume/PathTracingAtmosphere.ush
@@ -0,0 +1,301 @@
+// Copyright Epic Games, Inc. All Rights Reserved.
+
+#pragma once
+
+#include "./PathTracingAtmosphereCommon.ush"
+#include "../../DoubleWordMath.ush"
+#include "../../SkyAtmosphereCommon.ush"
+
+
+#define PLANET_ISECT_USES_DOUBLE_WORD_MATH 1
+
+// Shader parameters
+// NOTE: Atmosphere is a global struct (not directly defined here)
+float3 PlanetCenterTranslatedWorldHi;
+float3 PlanetCenterTranslatedWorldLo;
+Texture2D<float3> AtmosphereOpticalDepthLUT;
+SamplerState      AtmosphereOpticalDepthLUTSampler;
+
+
+#if PLANET_ISECT_USES_DOUBLE_WORD_MATH
+
+// intersect the planet using high precision
+// returns: planet hit T, atmo hit min/max (or -1.0 for any if missed)
+float3 RayPlanetIntersectPrecise(float3 Ro, float3 Rd)
+{
+	float3 Result = -1.0;
+
+	DWFloat CenterX = DWFloatCreate(PlanetCenterTranslatedWorldHi.x, PlanetCenterTranslatedWorldLo.x);
+	DWFloat CenterY = DWFloatCreate(PlanetCenterTranslatedWorldHi.y, PlanetCenterTranslatedWorldLo.y);
+	DWFloat CenterZ = DWFloatCreate(PlanetCenterTranslatedWorldHi.z, PlanetCenterTranslatedWorldLo.z);
+	DWFloat PlanetRadius = DWMul_ff(View.SkyAtmosphereBottomRadiusKm, SKY_UNIT_TO_CM);
+	DWFloat AtmoRadius   = DWMul_ff(View.SkyAtmosphereTopRadiusKm   , SKY_UNIT_TO_CM);
+
+	const DWFloat CoX = DWSub_df(CenterX, Ro.x);
+	const DWFloat CoY = DWSub_df(CenterY, Ro.y);
+	const DWFloat CoZ = DWSub_df(CenterZ, Ro.z);
+	const DWFloat b = DWAdd_dd(DWAdd_dd(
+		DWMul_df(CoX, Rd.x),
+		DWMul_df(CoY, Rd.y)),
+		DWMul_df(CoZ, Rd.z));
+	const DWFloat Hx = DWSub_dd(CoX, DWMul_df(b, Rd.x));
+	const DWFloat Hy = DWSub_dd(CoY, DWMul_df(b, Rd.y));
+	const DWFloat Hz = DWSub_dd(CoZ, DWMul_df(b, Rd.z));
+	const DWFloat hd = DWLengthSquared_ddd(Hx, Hy, Hz);
+	const DWFloat AR2 = DWSqr_d(AtmoRadius);
+	DWFloat ha = DWSub_dd(AR2, hd);
+	if (DWGreaterThan(ha, 0.0))
+	{
+		ha = DWSqrt_d(ha);
+		DWFloat q;
+		if (DWGreaterThan(b, 0.0))
+			q = DWAdd_dd(b, ha);
+		else
+			q = DWSub_dd(b, ha);
+		const DWFloat ta = q;
+		const DWFloat tb = DWDiv_dd(DWSub_dd(DWLengthSquared_ddd(CoX, CoY, CoZ), AR2), q);
+		const float t0 = ta.GetFloat();
+		const float t1 = tb.GetFloat();
+		Result.y = min(t0, t1);
+		Result.z = max(t0, t1);
+	}
+	const DWFloat PR2 = DWSqr_d(PlanetRadius);
+	DWFloat hp = DWSub_dd(PR2, hd);
+	if (DWGreaterThan(hp, 0.0))
+	{
+		hp = DWSqrt_d(hp);
+		DWFloat q;
+		if (DWGreaterThan(b, 0.0))
+			q = DWAdd_dd(b, hp);
+		else
+			q = DWSub_dd(b, hp);
+		const DWFloat ta = q;
+		const DWFloat tb = DWDiv_dd(DWSub_dd(DWLengthSquared_ddd(CoX, CoY, CoZ), PR2), q);
+		float t0 = ta.GetFloat();
+		float t1 = tb.GetFloat();
+		Result.x = min(t0, t1);
+		// clip volume layer against the back of the planet in case we started inside
+		if (Result.x < 0.0)
+			Result.y = max(t0, t1);
+	}
+	return Result;
+}
+
+#else
+
+// robust intersector with plain floats, still runs into trouble beyond 1000km or so
+float3 RayPlanetIntersectPrecise(float3 Ro, float3 Rd)
+{
+	float3 Result = -1.0;
+	float3 Center      = PlanetCenterTranslatedWorldHi + PlanetCenterTranslatedWorldLo;
+	float PlanetRadius = View.SkyAtmosphereBottomRadiusKm * SKY_UNIT_TO_CM;
+	float AtmoRadius   = View.SkyAtmosphereTopRadiusKm    * SKY_UNIT_TO_CM;
+
+	float3 co = Center - Ro;
+	float b = dot(co, Rd);
+	// "Precision Improvements for Ray / Sphere Intersection" - Ray Tracing Gems (2019)
+	// https://link.springer.com/content/pdf/10.1007%2F978-1-4842-4427-2_7.pdf
+	// NOTE: we assume the incoming ray direction is normalized
+	float hd = length2(co - b * Rd);
+	float ha = AtmoRadius * AtmoRadius - hd;
+	if (ha > 0)
+	{
+		ha = sqrt(ha);
+		float q = b > 0.0 ? b + ha : b - ha;
+		float ta = q;
+		float tb = (length2(co) - AtmoRadius * AtmoRadius) / q;
+		Result.y = min(ta, tb);
+		Result.z = max(ta, tb);
+	}
+	float hp = PlanetRadius * PlanetRadius - hd;
+	if (hp > 0)
+	{
+		hp = sqrt(hp);
+		float q = b > 0.0 ? b + hp : b - hp;
+		float ta = q;
+		float tb = (length2(co) - PlanetRadius * PlanetRadius) / q;
+		Result.x = min(ta, tb);
+		// clip volume layer against the back of the planet in case we started inside
+		if (Result.x < 0.0)
+			Result.y = max(ta, tb);
+	}
+
+	return Result;
+}
+
+#endif
+
+FVolumeIntersection AtmosphereIntersect(float3 Origin, float3 Direction, float TMin, float TMax)
+{
+	// NOTE: only this call needs to be guarded at the moment since the others can't be reached if this returns nothing
+	if (!EnableAtmosphere)
+	{
+		return CreateEmptyVolumeIntersection();
+	}
+	float3 Result = RayPlanetIntersectPrecise(Origin, Direction);
+	// clip to provided ray interval
+	return CreateVolumeIntersection(
+		max(TMin, Result.y),
+		min(TMax, Result.z),
+		Result.x > TMin && Result.x < TMax ? Result.x : -1.0);
+}
+
+FPackedPathTracingPayload AtmosphereGetBlockerHit(float3 Origin, float3 Direction, float HitT)
+{
+	float3 Center = PlanetCenterTranslatedWorldHi + PlanetCenterTranslatedWorldLo;
+	float3 WorldP = Origin + HitT * Direction;
+	float3 WorldN = normalize(WorldP - Center);
+
+	FPathTracingPayload Payload = (FPathTracingPayload)0;
+	Payload.BaseColor = Atmosphere.GroundAlbedo.rgb;
+	Payload.Opacity = 1.0;
+	Payload.WorldNormal = Payload.WorldGeoNormal = Payload.WorldSmoothNormal = WorldN;
+	Payload.ShadingModelID = SHADINGMODELID_NUM; // defaults to diffuse
+	Payload.HitT = HitT;
+	Payload.PrimitiveLightingChannelMask = 7; // visible to all lights
+	Payload.SetFrontFace();
+	return PackPathTracingPayload(Payload);
+}
+
+struct FSphericalPlanetParameters {
+	float HeightKm;         // Distance from planet center (in Km)
+	float ZenithCosAngle;   // angle between the ray and the planet normal (+1 goes up towards space, -1 points down toward the ground)
+
+	FSphericalPlanetParameters GetAt(float TKm)
+	{
+		// work out the ViewHeight / CosAngle of the specified point along the ray (in SkyUnits)
+		const float TOverH = TKm / HeightKm;
+		const float RelativeHeight = sqrt(1 + TOverH * TOverH + 2 * TOverH * ZenithCosAngle);
+		FSphericalPlanetParameters Result;
+		Result.HeightKm = HeightKm * RelativeHeight;
+		Result.ZenithCosAngle = (TOverH + ZenithCosAngle) / RelativeHeight;
+		return Result;
+	}
+
+	float GetHorizonCosine()
+	{
+		return ComputeHorizonCosine(HeightKm, Atmosphere.BottomRadiusKm);
+	}
+};
+
+FSphericalPlanetParameters GetSphericalPlanetParameters(float3 Origin, float3 Direction)
+{
+	FSphericalPlanetParameters Result = (FSphericalPlanetParameters)0;
+#if PLANET_ISECT_USES_DOUBLE_WORD_MATH
+	// Compute the origin's height/angle with high precision since it involves distances to the planet center
+	const DWFloat CenterX = DWFloatCreate(PlanetCenterTranslatedWorldHi.x, PlanetCenterTranslatedWorldLo.x);
+	const DWFloat CenterY = DWFloatCreate(PlanetCenterTranslatedWorldHi.y, PlanetCenterTranslatedWorldLo.y);
+	const DWFloat CenterZ = DWFloatCreate(PlanetCenterTranslatedWorldHi.z, PlanetCenterTranslatedWorldLo.z);
+	const DWFloat OCx = DWMul_df(DWSub_fd(Origin.x, CenterX), CM_TO_SKY_UNIT);
+	const DWFloat OCy = DWMul_df(DWSub_fd(Origin.y, CenterY), CM_TO_SKY_UNIT);
+	const DWFloat OCz = DWMul_df(DWSub_fd(Origin.z, CenterZ), CM_TO_SKY_UNIT);
+	const DWFloat DWViewHeight = DWSqrt_d(DWLengthSquared_ddd(OCx, OCy, OCz));
+	const DWFloat DWViewZenithCosAngle = DWDiv_dd(DWAdd_dd(DWAdd_dd(
+		DWMul_df(OCx, Direction.x),
+		DWMul_df(OCy, Direction.y)),
+		DWMul_df(OCz, Direction.z)),
+		DWViewHeight);
+
+	Result.HeightKm = DWViewHeight.GetFloat();
+	Result.ZenithCosAngle = DWViewZenithCosAngle.GetFloat();
+#else
+	// Get the ViewHeight / CosAngle of the origin
+	const float3 PosFromCenter = Origin - (PlanetCenterTranslatedWorldHi + PlanetCenterTranslatedWorldLo);
+	Result.HeightKm = length(PosFromCenter * CM_TO_SKY_UNIT);
+	Result.ZenithCosAngle = dot(normalize(PosFromCenter), Direction);
+#endif
+	return Result;
+}
+
+FVolumeDensityBounds AtmosphereGetDensityBounds(float3 Origin, float3 Direction, float TMin, float TMax)
+{
+#if 0
+	// conservative global bound for validation
+	const float ta = Atmosphere.AbsorptionDensity0LayerWidth;
+	const float tv = Atmosphere.AbsorptionDensity0LinearTerm * ta + Atmosphere.AbsorptionDensity0ConstantTerm;
+	float3 SigmaMax = (Atmosphere.MieExtinction.rgb + Atmosphere.RayleighScattering.rgb + tv * Atmosphere.AbsorptionExtinction.rgb) * CM_TO_SKY_UNIT;
+	return CreateVolumeDensityBound(0.0, SigmaMax);
+#else
+	// optimized bounds for the current ray
+	FSphericalPlanetParameters Params = GetSphericalPlanetParameters(Origin, Direction);
+
+	const float Ha = Params.GetAt(TMin * CM_TO_SKY_UNIT).HeightKm;
+	const float Hb = Params.GetAt(TMax * CM_TO_SKY_UNIT).HeightKm;
+	// one of the two end-points has to be the closest
+	float HLo = min(Ha, Hb);
+	float HHi = max(Ha, Hb);
+	// compute the closest point to the center along the ray
+	float TClosest = -Params.HeightKm * Params.ZenithCosAngle * SKY_UNIT_TO_CM;
+	// see if it falls within our ray segment
+	if (TClosest > TMin && TClosest < TMax)
+	{
+		// if so, we can reduce our lower bound a bit more
+		HLo = min(HLo, Params.HeightKm * sqrt(saturate(1 - Params.ZenithCosAngle * Params.ZenithCosAngle)));
+	}
+
+	// now compute density bounds for each component of the atmosphere
+	// exponential portion is easy since this is a stricly decreasing function of altitude
+	const float2 SampleHeight = max(0.0, float2(HLo, HHi) - Atmosphere.BottomRadiusKm);
+	const float2 DensityMie = exp(Atmosphere.MieDensityExpScale      * SampleHeight.yx); // flip height lo/hi because function is decreasing
+	const float2 DensityRay = exp(Atmosphere.RayleighDensityExpScale * SampleHeight.yx);
+
+	// Ozone is a bit trickier -- first eval both ends
+	float2 DensityOzo = float2(
+		SampleHeight.x < Atmosphere.AbsorptionDensity0LayerWidth ?
+			saturate(Atmosphere.AbsorptionDensity0LinearTerm * SampleHeight.x + Atmosphere.AbsorptionDensity0ConstantTerm) :
+			saturate(Atmosphere.AbsorptionDensity1LinearTerm * SampleHeight.x + Atmosphere.AbsorptionDensity1ConstantTerm),
+		SampleHeight.y < Atmosphere.AbsorptionDensity0LayerWidth ?
+			saturate(Atmosphere.AbsorptionDensity0LinearTerm * SampleHeight.y + Atmosphere.AbsorptionDensity0ConstantTerm) :
+			saturate(Atmosphere.AbsorptionDensity1LinearTerm * SampleHeight.y + Atmosphere.AbsorptionDensity1ConstantTerm)
+	);
+	if (DensityOzo.x > DensityOzo.y)
+		DensityOzo = DensityOzo.yx; // flip if wrong way around
+	if (SampleHeight.x < Atmosphere.AbsorptionDensity0LayerWidth && Atmosphere.AbsorptionDensity0LayerWidth < SampleHeight.y)
+	{
+		// height range includes the peak -- bump max to peak
+		const float ta = Atmosphere.AbsorptionDensity0LayerWidth;
+		const float tv = Atmosphere.AbsorptionDensity0LinearTerm * ta + Atmosphere.AbsorptionDensity0ConstantTerm;
+		DensityOzo.y = tv;
+	}
+	return CreateVolumeDensityBound(
+		(DensityMie.x * Atmosphere.MieExtinction.rgb + DensityRay.x * Atmosphere.RayleighScattering.rgb + DensityOzo.x * Atmosphere.AbsorptionExtinction.rgb) * CM_TO_SKY_UNIT,
+		(DensityMie.y * Atmosphere.MieExtinction.rgb + DensityRay.y * Atmosphere.RayleighScattering.rgb + DensityOzo.y * Atmosphere.AbsorptionExtinction.rgb) * CM_TO_SKY_UNIT
+	);
+#endif
+}
+
+FVolumeShadedResult AtmosphereGetDensity(float3 TranslatedWorldPos)
+{
+#if 0 //PLANET_ISECT_USES_DOUBLE_WORD_MATH
+	// Extra precision here doesn't seem to matter too much -- probably because density function is very smooth, so the exact position doesn't matter too much
+	const float H = GetSphericalPlanetParameters(TranslatedWorldPos, float3(0, 0, 1)).HeightKm;
+	MediumSampleRGB Medium = SampleAtmosphereMediumRGB(float3(0.0, 0.0, H));
+#else
+	float3 PosFromCenter = TranslatedWorldPos - (PlanetCenterTranslatedWorldHi + PlanetCenterTranslatedWorldLo);
+	MediumSampleRGB Medium = SampleAtmosphereMediumRGB(PosFromCenter * CM_TO_SKY_UNIT);
+#endif
+	FVolumeShadedResult Result = (FVolumeShadedResult)0;
+	Result.SigmaT         = Medium.Extinction    * CM_TO_SKY_UNIT;
+	Result.SigmaSRayleigh = Medium.ScatteringRay * CM_TO_SKY_UNIT;
+	Result.SigmaSHG       = Medium.ScatteringMie * CM_TO_SKY_UNIT;
+	Result.PhaseG         = Atmosphere.MiePhaseG;
+	return Result;
+}
+
+float3 AtmosphereGetOpticalDensity(FSphericalPlanetParameters Params)
+{
+	float2 UV = ToAtmosphereOpticalDepthLUTUV(Params.HeightKm, Params.ZenithCosAngle, Atmosphere.BottomRadiusKm, Atmosphere.TopRadiusKm);
+	return AtmosphereOpticalDepthLUT.SampleLevel(AtmosphereOpticalDepthLUTSampler, UV, 0);
+}
+
+float3 AtmosphereGetTransmittance(float3 Origin, float3 Direction, float TMin, float TMax)
+{
+	FSphericalPlanetParameters Params = GetSphericalPlanetParameters(Origin, Direction);
+	// now lookup the optical density LUT from each point (integral up to the farthest possible point)
+	float3 DensityA = AtmosphereGetOpticalDensity(Params.GetAt(TMin * CM_TO_SKY_UNIT));
+	float3 DensityB = AtmosphereGetOpticalDensity(Params.GetAt(TMax * CM_TO_SKY_UNIT));
+	// get density for just this segment by removing double-counted density
+	float3 Tau = max(DensityA - DensityB, 0.0);
+	// compute transmittance
+	return exp(-Tau);
+}
diff --git a/Engine/Shaders/Private/PathTracing/Volume/PathTracingAtmosphereCommon.ush b/Engine/Shaders/Private/PathTracing/Volume/PathTracingAtmosphereCommon.ush
new file mode 100644
--- /dev/null
+++ b/Engine/Shaders/Private/PathTracing/Volume/PathTracingAtmosphereCommon.ush
@@ -0,0 +1,53 @@
+// Copyright Epic Games, Inc. All Rights Reserved.
+
+#pragma once
+
+// Atmosphere function that are independent of any shader parameters
+
+float Gain(float x, float k)
+{
+	float a = 1 - 2 * x;
+	return 0.5 - 0.5 * sign(a) * pow(abs(a), k);
+}
+
+float ComputeHorizonCosine(float ViewHeight, float BotRadius)
+{
+	const float R = BotRadius / ViewHeight;
+	return -sqrt(saturate(1.0 - R * R));
+}
+
+float2 ToAtmosphereOpticalDepthLUTUV(float Height, float Cosine, float BotRadius, float TopRadius)
+{
+	const float HorizonCosine = ComputeHorizonCosine(Height, BotRadius);
+
+	const float Q = 0.5; // line up horizon line to u=Q
+	const float2 ab = Cosine >= HorizonCosine
+		? float2((HorizonCosine - 1) / (Q - 1), (Q - HorizonCosine) / (Q - 1)) // above horizon
+		: float2((1 + HorizonCosine) / Q, -1);                                 // below horizon
+	
+	float2 UV = float2((Cosine - ab.y) / ab.x, sqrt(saturate((Height - BotRadius) / (TopRadius - BotRadius))));
+	// un-squish
+	const float Sq = 2.0;
+	UV.x = Gain(UV.x, 1.0 / Sq);
+	return UV;
+}
+
+float2 FromAtmosphereOpticalDepthLUTUV(float2 UV, float BotRadius, float TopRadius)
+{
+	const float Height = lerp(BotRadius, TopRadius, Pow2(UV.y)); // more resolution toward the ground
+
+	// from this height, at what cosine is the horizon line?
+	const float HorizonCosine = ComputeHorizonCosine(Height, BotRadius);
+
+	// squish around 0.5 to increase resolution near the horizon
+	const float Sq = 2.0;
+	const float U = Gain(UV.x, Sq);
+
+	const float Q = 0.5; // line up horizon line to u=Q
+	const float2 ab = U > Q
+		? float2((HorizonCosine - 1) / (Q - 1), (Q - HorizonCosine) / (Q - 1)) // above horizon
+		: float2((1 + HorizonCosine) / Q, -1);                                 // below horizon
+	const float Cosine = ab.x * U + ab.y;
+
+	return float2(Height, Cosine);
+}
diff --git a/Engine/Shaders/Private/PathTracing/Volume/PathTracingFog.ush b/Engine/Shaders/Private/PathTracing/Volume/PathTracingFog.ush
new file mode 100644
--- /dev/null
+++ b/Engine/Shaders/Private/PathTracing/Volume/PathTracingFog.ush
@@ -0,0 +1,103 @@
+// Copyright Epic Games, Inc. All Rights Reserved.
+
+#pragma once
+
+// FogParameters
+float2 FogDensity;
+float2 FogHeight;
+float2 FogFalloff;
+float4 FogAlbedo;
+float FogPhaseG;
+float2 FogCenter;
+float FogMinZ;
+float FogMaxZ;
+float FogRadius;
+
+FVolumeIntersection FogIntersect(float3 Origin, float3 Direction, float TMin, float TMax)
+{
+	// adapted from https://www.iquilezles.org/www/articles/intersectors/intersectors.htm
+	// extended/simplified because we only need the entrance/exit distances
+	float3 oc = Origin - float3(FogCenter, FogMinZ);
+
+	float ba = FogMaxZ - FogMinZ;
+	float baba = ba * ba;
+	float bard = ba * Direction.z;
+	float baoc = ba * oc.z;
+
+	float k2 = baba - bard * bard;
+	float k1 = baba * dot(oc, Direction) - baoc * bard;
+	float k0 = baba * dot(oc, oc) - baoc * baoc - FogRadius * FogRadius * baba;
+
+	float h = k1 * k1 - k2 * k0;
+	if (h > 0.0)
+	{
+		float s = sign(k2);
+		float2 TCyl = (-k1 + float2(-s, s) * sqrt(h)) / k2;
+		float2 TCap = ((bard > 0.0 ? float2(0, baba) : float2(baba, 0)) - baoc) / bard;
+
+		const float THitMin = max3(TCap.x, TCyl.x, TMin);
+		const float THitMax = min3(TCap.y, TCyl.y, TMax);
+		return CreateVolumeIntersection(THitMin, THitMax);
+	}
+	return CreateEmptyVolumeIntersection();
+}
+
+FVolumeDensityBounds FogGetDensityBounds(float3 Origin, float3 Direction, float TMin, float TMax)
+{
+#if 0
+	// conservative default (reference)
+	return CreateVolumeDensityBound(0, FogDensity.x + FogDensity.y);
+#else
+	// get bounds at each endpoint of the ray (density is a monotonic function, so just eval at both endpoints and sort)
+	float2 DensityA = FogDensity * exp2(-max(Origin.z + TMin * Direction.z - FogHeight, 0.0) * FogFalloff);
+	float2 DensityB = FogDensity * exp2(-max(Origin.z + TMax * Direction.z - FogHeight, 0.0) * FogFalloff);
+
+	float2 DensityBounds = float2(
+		DensityA.x + DensityA.y,
+		DensityB.x + DensityB.y
+	);
+	// flip if in the wrong order
+	if (DensityBounds.x > DensityBounds.y)
+		DensityBounds = DensityBounds.yx;
+	return CreateVolumeDensityBound(DensityBounds.x, DensityBounds.y);
+#endif
+}
+
+// Given a point in world space, return the amount of volume and its scattering properties
+FVolumeShadedResult FogGetDensity(float3 TranslatedWorldPos)
+{
+	FVolumeShadedResult Result = (FVolumeShadedResult)0;
+
+	float2 Density = FogDensity * exp2(-max(TranslatedWorldPos.z - FogHeight, 0.0) * FogFalloff);
+	Result.SigmaT = Density.x + Density.y;
+	Result.SigmaSHG = FogAlbedo.xyz * Result.SigmaT;
+	Result.PhaseG = FogPhaseG;
+	return Result;
+}
+
+// Return the transmittance along a ray segment
+float3 FogGetTransmittance(float3 Origin, float3 Direction, float TMin, float TMax)
+{
+	// Solve for intersection with Z == FogHeight (or nearest point within slab)
+	float2 OH = Origin.z - FogHeight;
+	float Dz = Direction.z;
+
+	float2 Tau = 0.0;
+	if (abs(Dz) < 1e-3)
+	{
+		// perfectly horizontal ray -- either go through all uniform or all varying, but either way density will be constant
+		Tau = (TMax - TMin) * exp2(-FogFalloff * max(OH, 0.0));
+	}
+	else
+	{
+		// figure out optical density through the uniform portion and varying portion respectively
+		float2 TMid = clamp(-OH / Dz, TMin, TMax);
+		Tau = Dz > 0
+			? exp2(-FogFalloff * max(0.0, OH + Dz * TMid)) - exp2(-FogFalloff * max(0.0, OH + Dz * TMax))
+			: exp2(-FogFalloff * max(0.0, OH + Dz * TMin)) - exp2(-FogFalloff * max(0.0, OH + Dz * TMid));
+		Tau *= rcp(Dz * FogFalloff * log(2.0));
+		Tau += Dz > 0 ? TMid - TMin : TMax - TMid;
+	}
+
+	return exp(-dot(FogDensity, Tau));
+}
diff --git a/Engine/Shaders/Private/PathTracing/Volume/PathTracingVolume.ush b/Engine/Shaders/Private/PathTracing/Volume/PathTracingVolume.ush
new file mode 100644
--- /dev/null
+++ b/Engine/Shaders/Private/PathTracing/Volume/PathTracingVolume.ush
@@ -0,0 +1,195 @@
+// Copyright Epic Games, Inc. All Rights Reserved.
+
+#pragma once
+
+#include "PathTracingVolumeCommon.ush"
+#include "PathTracingAtmosphere.ush"
+#include "PathTracingFog.ush"
+#include "PathTracingVolumeSampling.ush"
+
+#define PATH_TRACING_USE_TRANSMITTANCE_FUNC 1 // 0: null tracking transmittance (reference) 1: analytic transmittance
+
+uint EnableAtmosphere;
+uint EnableFog;
+int MaxRaymarchSteps;
+
+
+// Given an input ray, figure out where the provided volume intersects it
+FVolumeIntersectionList VolumeIntersect(float3 Origin, float3 Direction, float TMin, float TMax)
+{
+	FVolumeIntersectionList Result = CreateEmptyVolumeIntersectionList();
+	if (EnableAtmosphere)
+	{
+		Result.Add(VOLUMEID_ATMOSPHERE, AtmosphereIntersect(Origin, Direction, TMin, TMax));
+	}
+	if (EnableFog)
+	{
+		Result.Add(VOLUMEID_FOG, FogIntersect(Origin, Direction, TMin, TMax));
+	}
+	return Result;
+}
+
+FPackedPathTracingPayload VolumeGetBlockerHit(uint VolumeID, float3 Origin, float3 Direction, float HitT)
+{
+	// NOTE: only atmosphere supports blockers, so if this gets called it is safe to assume its for the atmosphere volume
+	return AtmosphereGetBlockerHit(Origin, Direction, HitT);
+}
+
+// Given an input ray, figure out the bounds on density. The T interval may be smaller than the one returned by VolumeIntersect
+FVolumeDensityBounds VolumeGetDensityBounds(float3 Origin, float3 Direction, FVolumeIntersectionInterval Interval)
+{
+	FVolumeDensityBounds Result = CreateVolumeDensityBound(0, 0);
+	for (int Index = 0; Index < Interval.Num; Index++)
+	{
+		switch (Interval.VolumeID[Index])
+		{
+			case VOLUMEID_ATMOSPHERE: MergeVolumeDensityBounds(Result, AtmosphereGetDensityBounds(Origin, Direction, Interval.VolumeTMin, Interval.VolumeTMax)); break;
+			case VOLUMEID_FOG:        MergeVolumeDensityBounds(Result, FogGetDensityBounds(Origin, Direction, Interval.VolumeTMin, Interval.VolumeTMax)); break;
+		}
+	}
+	return Result;
+}
+
+// Given a point in world space, return the amount of volume and its scattering properties
+FVolumeShadedResult VolumeGetDensity(float3 TranslatedWorldPos, FVolumeIntersectionInterval Interval)
+{
+	FVolumeShadedResult Result = (FVolumeShadedResult)0;
+	for (int Index = 0; Index < Interval.Num; Index++)
+	{
+		switch (Interval.VolumeID[Index])
+		{
+			case VOLUMEID_ATMOSPHERE: MergeVolumeShadedResult(Result, AtmosphereGetDensity(TranslatedWorldPos)); break;
+			case VOLUMEID_FOG:        MergeVolumeShadedResult(Result, FogGetDensity(TranslatedWorldPos)); break;
+		}
+	}
+	return Result;
+}
+
+#if PATH_TRACING_USE_TRANSMITTANCE_FUNC
+
+// Return the transmittance along a ray segment
+float3 VolumeGetTransmittance(float3 StartThroughput, float3 Origin, float3 Direction, FVolumeIntersectionInterval Interval, inout RandomSequence RandSequence)
+{
+	float3 Throughput = StartThroughput;
+	for (int Index = 0; Index < Interval.Num; Index++)
+	{
+		switch (Interval.VolumeID[Index])
+		{
+			case VOLUMEID_ATMOSPHERE: Throughput *= AtmosphereGetTransmittance(Origin, Direction, Interval.VolumeTMin, Interval.VolumeTMax); break;
+			case VOLUMEID_FOG:        Throughput *= FogGetTransmittance(Origin, Direction, Interval.VolumeTMin, Interval.VolumeTMax); break;
+		}
+	}
+	return Throughput;
+}
+
+// Return the transmittance along a whole intersection list
+float3 VolumeGetTransmittance(float3 StartThroughput, float3 Origin, float3 Direction, FVolumeIntersectionList IntersectionList, inout RandomSequence RandSequence)
+{
+	float3 Throughput = StartThroughput;
+	for (int Index = 0; Index < IntersectionList.Num; Index++)
+	{
+		float TMin = IntersectionList.VolumeTMin[Index];
+		float TMax = IntersectionList.VolumeTMax[Index];
+		switch (IntersectionList.VolumeID[Index])
+		{
+			case VOLUMEID_ATMOSPHERE: Throughput *= AtmosphereGetTransmittance(Origin, Direction, TMin, TMax); break;
+			case VOLUMEID_FOG:        Throughput *= FogGetTransmittance(Origin, Direction, TMin, TMax); break;
+		}
+	}
+	return Throughput;
+}
+
+#else
+
+// Reference implementation of transmittance that uses only GetDensity and GetDensityBounds
+// express the inner loop here as a macro so we can stamp down a separate copy per volume ID and avoid a nested loop during ray marching
+#define PATH_TRACER_REFERENCE_TRANSMITTANCE_LOOP(GetDensityBoundsFunc, GetDensityFunc)				\
+	/* Limit number of steps to prevent timeouts // FIXME: This biases the result! */				\
+	for (int Step = 0; Step < MaxRaymarchSteps; Step++)												\
+	{																								\
+		float3 SigmaBar = GetDensityBoundsFunc(Origin, Direction, TMin, TMax).SigmaMax;				\
+		/* take stochastic steps along the ray to estimate transmittance (null scattering) */		\
+		float3 ColorChannelPdf = Throughput;														\
+		/* Sample the distance to the next interaction */											\
+		float2 RandValue = RandomSequence_GenerateSample2D(RandSequence);							\
+		float DeltaT = SampleSpectralTransmittance(RandValue.x, SigmaBar, ColorChannelPdf);			\
+		if (DeltaT < 0.0)																			\
+		{																							\
+			/* no more energy left in the path */													\
+			break;																					\
+		}																							\
+		if (TMin + DeltaT < TMax)																	\
+		{																							\
+			TMin += DeltaT;																			\
+			/* our ray marching step stayed inside the atmo and is still in front of the next hit */\
+			/* Compute transmittance through the bounding homogeneous medium (both real and fictitious particles) */	\
+			Throughput *= EvaluateSpectralTransmittanceHit(DeltaT, SigmaBar, ColorChannelPdf).xyz;	\
+			float3 WorldPos = Origin + TMin * Direction;											\
+			/* clamp to make sure we never exceed the majorant (should not be the case, but need to avoid any possible numerical issues) */	\
+			float3 SigmaT = min(GetDensityFunc(WorldPos).SigmaT, SigmaBar);							\
+			float3 SigmaN = SigmaBar - SigmaT;														\
+			/* keep tracing through the volume */													\
+			Throughput *= SigmaN;																	\
+		}																							\
+		else                                                                                        \
+		{																							\
+			/* update the path throughput, knowing that we escaped the medium*/						\
+			Throughput *= EvaluateSpectralTransmittanceMiss(TMax - TMin, SigmaBar, ColorChannelPdf).xyz;	\
+			/* exit the ray marching loop */														\
+			break;																					\
+		}																							\
+	}																								\
+
+
+// The returned value factors in the initial throughput. This is used to avoid sampling too much if throughput reaches 0.
+float3 VolumeGetTransmittance(float3 StartThroughput, float3 Origin, float3 Direction, FVolumeIntersectionInterval Interval, inout RandomSequence RandSequence)
+{
+	float3 Throughput = StartThroughput;
+	for (int Index = 0; Index < Interval.Num; Index++)
+	{
+		float TMin = Interval.VolumeTMin;
+		float TMax = Interval.VolumeTMax;
+		switch (Interval.VolumeID[Index])
+		{
+			case VOLUMEID_ATMOSPHERE:
+			{
+				PATH_TRACER_REFERENCE_TRANSMITTANCE_LOOP(AtmosphereGetDensityBounds, AtmosphereGetDensity);
+				break;
+			}
+			case VOLUMEID_FOG:
+			{
+				PATH_TRACER_REFERENCE_TRANSMITTANCE_LOOP(FogGetDensityBounds, FogGetDensity);
+				break;
+			}
+		}
+	}
+	return Throughput;
+
+}
+float3 VolumeGetTransmittance(float3 StartThroughput, float3 Origin, float3 Direction, FVolumeIntersectionList IntersectionList, inout RandomSequence RandSequence)
+{
+	float3 Throughput = StartThroughput;
+	for (int Index = 0; Index < IntersectionList.Num; Index++)
+	{
+		float TMin = IntersectionList.VolumeTMin[Index];
+		float TMax = IntersectionList.VolumeTMax[Index];
+		switch (IntersectionList.VolumeID[Index])
+		{
+			case VOLUMEID_ATMOSPHERE:
+			{
+				PATH_TRACER_REFERENCE_TRANSMITTANCE_LOOP(AtmosphereGetDensityBounds, AtmosphereGetDensity);
+				break;
+			}
+			case VOLUMEID_FOG:
+			{
+				PATH_TRACER_REFERENCE_TRANSMITTANCE_LOOP(FogGetDensityBounds, FogGetDensity);
+				break;
+			}
+		}
+	}
+	return Throughput;
+
+#undef PATH_TRACER_REFERENCE_TRANSMITTANCE_LOOP
+}
+
+#endif
diff --git a/Engine/Shaders/Private/PathTracing/Volume/PathTracingVolumeCommon.ush b/Engine/Shaders/Private/PathTracing/Volume/PathTracingVolumeCommon.ush
new file mode 100644
--- /dev/null
+++ b/Engine/Shaders/Private/PathTracing/Volume/PathTracingVolumeCommon.ush
@@ -0,0 +1,255 @@
+// Copyright Epic Games, Inc. All Rights Reserved.
+
+#pragma once
+
+struct FVolumeIntersection
+{
+	float VolumeTMin;
+	float VolumeTMax;
+	float BlockerHitT;
+
+	bool HitVolume()
+	{
+		return VolumeTMin < VolumeTMax;
+	}
+
+	bool HitBlocker()
+	{
+		return BlockerHitT > 0.0;
+	}
+};
+
+FVolumeIntersection CreateVolumeIntersection(float TMin, float TMax, float BlockerT = -1.0)
+{
+	FVolumeIntersection Result = (FVolumeIntersection)0;
+	Result.VolumeTMin = TMin;
+	Result.VolumeTMax = TMax;
+	Result.BlockerHitT = BlockerT;
+	return Result;
+}
+
+FVolumeIntersection CreateEmptyVolumeIntersection()
+{
+	return (FVolumeIntersection)0;
+}
+
+#define VOLUMEID_ATMOSPHERE		0
+#define VOLUMEID_FOG			1
+#define PATH_TRACER_MAX_VOLUMES 2
+
+// representes an interval along the ray over which the specified volumes are known to exist
+struct FVolumeIntersectionInterval
+{
+	float VolumeTMin;
+	float VolumeTMax;
+	uint VolumeID[PATH_TRACER_MAX_VOLUMES];
+	int Num;
+};
+
+struct FVolumeIntersectionList
+{
+	float VolumeTMin[PATH_TRACER_MAX_VOLUMES];
+	float VolumeTMax[PATH_TRACER_MAX_VOLUMES];
+	uint  VolumeID[PATH_TRACER_MAX_VOLUMES];
+	int Num;
+
+	// only need to track one blocker (assume they are opaque)
+	float BlockerHitT;
+	uint  BlockerID;
+
+	bool HitVolume()
+	{
+		return Num > 0;
+	}
+
+	bool HitBlocker()
+	{
+		return BlockerHitT > 0;
+	}
+
+	void Add(uint ID, FVolumeIntersection VolumeIntersection)
+	{
+		// merge the provided hit into the list
+		if (VolumeIntersection.HitVolume())
+		{
+			VolumeTMin[Num] = VolumeIntersection.VolumeTMin;
+			VolumeTMax[Num] = VolumeIntersection.VolumeTMax;
+			VolumeID[Num] = ID;
+			Num++;
+		}
+		if (VolumeIntersection.HitBlocker())
+		{
+			if (!HitBlocker() || BlockerHitT > VolumeIntersection.BlockerHitT)
+			{
+				BlockerHitT = VolumeIntersection.BlockerHitT;
+				BlockerID = ID;
+			}
+		}
+	}
+
+	// Return the interval of TValues closest to the front of the list
+	FVolumeIntersectionInterval GetCurrentInterval()
+	{
+		FVolumeIntersectionInterval Result = (FVolumeIntersectionInterval)0;
+		Result.Num = 0;
+		Result.VolumeTMin = VolumeTMin[0];
+		Result.VolumeTMax = VolumeTMax[0];
+		for (int Index = 1; Index < Num; Index++)
+		{
+			// NOTE: we only enter the loop body if Result was a valid interval
+			float TMin = VolumeTMin[Index];
+			float TMax = VolumeTMax[Index];
+			if (Result.VolumeTMin == TMin)
+			{
+				// both segments start together, figure out which ends first
+				Result.VolumeTMax = min(Result.VolumeTMax, TMax);
+			}
+			else if (Result.VolumeTMin < TMin)
+			{
+				// what we have so far _starts_ before the current segment
+				// so either the two segments are disjoint (Result.y is closer)
+				// or the two overlap, and we need to stop at the next transition line
+				Result.VolumeTMax = min(Result.VolumeTMax, TMin);
+			}
+			else
+			{
+				// segments could be disjoint (TMax is closer) or overlapping (TMax is closer)
+				Result.VolumeTMax = min(TMax, Result.VolumeTMin);
+				Result.VolumeTMin = TMin; // we just established this one is closer
+			}
+		}
+		// now that we established the distance, figure out which IDs overlap with the resolved segment
+		for (int Index2 = 0; Index2 < Num; Index2++)
+		{
+			if (Result.VolumeTMax > VolumeTMin[Index2] &&
+				Result.VolumeTMin < VolumeTMax[Index2])
+			{
+				Result.VolumeID[Result.Num] = VolumeID[Index2];
+				Result.Num++;
+			}
+		}
+		return Result;
+	}
+
+	// Update all active intervals, knowing we have processed everything up to distance 'T'
+	FVolumeIntersectionList Update(float T)
+	{
+		FVolumeIntersectionList Result = (FVolumeIntersectionList)0;
+		Result.Num = 0;
+		// clip the current intervals, knowning that we are "T" away
+		for (int Index = 0; Index < Num; Index++)
+		{
+			float TMax = VolumeTMax[Index];
+			if (T < TMax)
+			{
+				// still valid? then add it back in
+				Result.VolumeTMin[Result.Num] = max(T, VolumeTMin[Index]);
+				Result.VolumeTMax[Result.Num] = TMax;
+				Result.VolumeID[Result.Num] = VolumeID[Index];
+				Result.Num++;
+			}
+		}
+		Result.BlockerHitT = BlockerHitT;
+		Result.BlockerID = BlockerID;
+		return Result;
+	}
+};
+
+FVolumeIntersectionList CreateEmptyVolumeIntersectionList()
+{
+	return (FVolumeIntersectionList)0;
+}
+
+
+// Represent a strict lower and upper bound on the values that can be returned by VolumeGetDensity along the given ray
+struct FVolumeDensityBounds
+{
+	float3 SigmaMin; // minorant (control extinction)
+	float3 SigmaMax; // majorant (upper bound)
+};
+
+FVolumeDensityBounds CreateVolumeDensityBound(float3 SigmaMin, float3 SigmaMax)
+{
+	FVolumeDensityBounds Result = (FVolumeDensityBounds)0;
+	Result.SigmaMin = SigmaMin;
+	Result.SigmaMax = SigmaMax;
+	return Result;
+}
+
+void MergeVolumeDensityBounds(inout FVolumeDensityBounds Merged, FVolumeDensityBounds Other)
+{
+	Merged.SigmaMin += Other.SigmaMin;
+	Merged.SigmaMax += Other.SigmaMax;
+}
+
+
+// The result of sampling a volume at a given point
+struct FVolumeShadedResult
+{
+	float3 SigmaT;         // extinction
+	float3 SigmaSRayleigh; // Rayleight scattering coefficient
+	float3 SigmaSHG;       // Henyey-Greenstein scattering coefficient
+	float PhaseG;          // 'g' term for HG lobe
+
+	// TODO: second HG lobe?
+	// TODO: emission?
+};
+
+FPathTracingPayload CreateMediumHitPayload(float HitT, float3 TranslatedWorldPos, FVolumeShadedResult ShadedResult)
+{
+	FPathTracingPayload Result = (FPathTracingPayload)0; // clear all fields
+	Result.HitT = HitT;
+	Result.TranslatedWorldPos = TranslatedWorldPos;
+	Result.ShadingModelID = SHADINGMODELID_MEDIUM;
+	Result.BlendingMode = RAY_TRACING_BLEND_MODE_OPAQUE;
+	Result.Opacity = 1.0;
+	Result.PrimitiveLightingChannelMask = 7;
+	Result.SetFrontFace();
+	float3 SigmaS = ShadedResult.SigmaSRayleigh + ShadedResult.SigmaSHG;
+	Result.BaseColor = SigmaS > 0 ? min(ShadedResult.SigmaSRayleigh, SigmaS) / SigmaS : 0.0;
+	Result.CustomData.xyz = SigmaS > 0 ? min(ShadedResult.SigmaSHG, SigmaS) / SigmaS : 0.0;
+	Result.CustomData.w = ShadedResult.PhaseG;
+	return Result;
+}
+
+void MergeVolumeShadedResult(inout FVolumeShadedResult Merged, FVolumeShadedResult Other)
+{
+	Merged.SigmaT += Other.SigmaT;
+	Merged.SigmaSRayleigh += Other.SigmaSRayleigh;
+	Merged.SigmaSHG += Other.SigmaSHG;
+	// blend phase function anisotropy based on amount of each medium
+	float Qa = Other.SigmaSHG.x + Other.SigmaSHG.y + Other.SigmaSHG.z;
+	float Qs = Merged.SigmaSHG.x + Merged.SigmaSHG.y + Merged.SigmaSHG.z;
+	Merged.PhaseG = lerp(Merged.PhaseG, Other.PhaseG, Qs > 0 ? Qa / Qs : 0.0);
+}
+
+struct FVolumeSampleResult
+{
+	float3 SurfaceThroughput;  // the path throughput beyond the volume (for the surface hit behind usually)
+	float3 VolumeThroughput;   // the path throughput up to the volume sample (if we sampled one)
+	float  VolumeT;            // distance to volume sample along the ray (will be <= 0 if none was sampled)
+	FVolumeShadedResult VolumeSample; // actual shaded volume hit (if we have one)
+
+	bool VolumeHit() { return VolumeT > 0.0; }
+};
+
+FVolumeSampleResult CreateEmptyVolumeSampleResult()
+{
+	FVolumeSampleResult Result = (FVolumeSampleResult)0;
+	Result.VolumeT = -1.0;
+	return Result;
+}
+
+// represent a stochastically chosen volume element that is guaranteed to lie between 
+struct FVolumeSegment
+{
+	float3 Throughput; // path contribution from the start of the volume segment
+	FVolumeIntersectionInterval Interval;
+
+	bool IsValid() { return Interval.VolumeTMin < Interval.VolumeTMax; }
+};
+
+FVolumeSegment CreateEmptyVolumeSegment()
+{
+	return (FVolumeSegment)0;
+}
diff --git a/Engine/Shaders/Private/PathTracing/Volume/PathTracingVolumeSampling.ush b/Engine/Shaders/Private/PathTracing/Volume/PathTracingVolumeSampling.ush
new file mode 100644
--- /dev/null
+++ b/Engine/Shaders/Private/PathTracing/Volume/PathTracingVolumeSampling.ush
@@ -0,0 +1,56 @@
+// Copyright Epic Games, Inc. All Rights Reserved.
+
+#pragma once
+
+
+// Returns a distance sampled proportionally to transmittance with an extinction of Sigma, using the provided color channel pdf for weighting the probabilities
+float SampleSpectralTransmittance(float RandValue, float3 Sigma, float3 ColorChannelPdf)
+{
+	float3 ColorChannelCdf = float3(
+		ColorChannelPdf.x,
+		ColorChannelPdf.x + ColorChannelPdf.y,
+		ColorChannelPdf.x + ColorChannelPdf.y + ColorChannelPdf.z);
+	if (ColorChannelCdf.z > 0)
+	{
+		const float OneMinusEpsilon = 0.99999994; // 32-bit float just before 1.0
+		float q = RandValue * ColorChannelCdf.z;
+		if (q < ColorChannelCdf.x)
+		{
+			float RescaleRand = min(q / ColorChannelCdf.x, OneMinusEpsilon);
+			return -log(1 - RescaleRand) / Sigma.x;
+		}
+		else if (q < ColorChannelCdf.y)
+		{
+			float RescaleRand = min((q - ColorChannelCdf.x) / (ColorChannelCdf.y - ColorChannelCdf.x), OneMinusEpsilon);
+			return -log(1 - RescaleRand) / Sigma.y;
+		}
+		else
+		{
+			float RescaleRand = min((q - ColorChannelCdf.y) / (ColorChannelCdf.z - ColorChannelCdf.y), OneMinusEpsilon);
+			return -log(1 - RescaleRand) / Sigma.z;
+		}
+	}
+	// all channels have 0 probability
+	return -1.0;
+}
+
+
+float4 EvaluateSpectralTransmittanceHit(float SampledT, float3 Sigma, float3 ColorChannelPdf)
+{
+	// normalize the pdf (to match code above)
+	ColorChannelPdf *= rcp(ColorChannelPdf.x + ColorChannelPdf.y + ColorChannelPdf.z);
+	float3 Transmittance = exp(-SampledT * Sigma);
+	float3 TransmittancePdf = Sigma * Transmittance; // probability of reaching the sampled point
+	float MisPdf = dot(ColorChannelPdf, TransmittancePdf);
+	return MisPdf > 0 ? float4(Transmittance / MisPdf, MisPdf) : 0.0;
+}
+
+float4 EvaluateSpectralTransmittanceMiss(float MaxT, float3 Sigma, float3 ColorChannelPdf)
+{
+	// normalize the pdf (to match code above)
+	ColorChannelPdf *= rcp(ColorChannelPdf.x + ColorChannelPdf.y + ColorChannelPdf.z);
+	float3 Transmittance = exp(-MaxT * Sigma);
+	float3 TransmittancePdf = Transmittance; // probability of going past MaxT (integral of the pdf from MaxT to infinity)
+	float MisPdf = dot(ColorChannelPdf, TransmittancePdf);
+	return MisPdf > 0 ? float4(Transmittance / MisPdf, MisPdf) : 0.0;
+}
diff --git a/Engine/Shaders/Private/SkyAtmosphere.usf b/Engine/Shaders/Private/SkyAtmosphere.usf
--- a/Engine/Shaders/Private/SkyAtmosphere.usf
+++ b/Engine/Shaders/Private/SkyAtmosphere.usf
@@ -551,7 +551,7 @@ SingleScatteringResult IntegrateSingleScatteredLuminance(
 		float PHeight = length(P);
 
 		// Sample the medium
-		MediumSampleRGB Medium = SampleMediumRGB(P);
+		MediumSampleRGB Medium = SampleAtmosphereMediumRGB(P);
 		const float3 SampleOpticalDepth = Medium.Extinction * dt * AerialPespectiveViewDistanceScale;
 		const float3 SampleTransmittance = exp(-SampleOpticalDepth);
 		OpticalDepth += SampleOpticalDepth;
@@ -1554,7 +1554,7 @@ void RenderSkyAtmosphereDebugPS(
 
 
 		float3 SampleWorldPos = float3(0.0f, 0.0f, Atmosphere.BottomRadiusKm + UV.x * (Atmosphere.TopRadiusKm - Atmosphere.BottomRadiusKm));
-		MediumSampleRGB Sample = SampleMediumRGB(SampleWorldPos);
+		MediumSampleRGB Sample = SampleAtmosphereMediumRGB(SampleWorldPos);
 
 		// Curves represent the maximum contribution as extinction (no differenciation between scattering or absorption)
 		float3 Curves = float3(
diff --git a/Engine/Shaders/Private/SkyAtmosphereCommon.ush b/Engine/Shaders/Private/SkyAtmosphereCommon.ush
--- a/Engine/Shaders/Private/SkyAtmosphereCommon.ush
+++ b/Engine/Shaders/Private/SkyAtmosphereCommon.ush
@@ -281,7 +281,7 @@ struct MediumSampleRGB
 };
 
 // If this is changed, please also update USkyAtmosphereComponent::GetTransmittance 
-MediumSampleRGB SampleMediumRGB(in float3 WorldPos)
+MediumSampleRGB SampleAtmosphereMediumRGB(in float3 WorldPos)
 {
 	const float SampleHeight = max(0.0, (length(WorldPos) - Atmosphere.BottomRadiusKm));
 
diff --git a/Engine/Shaders/Shared/PathTracingDefinitions.h b/Engine/Shaders/Shared/PathTracingDefinitions.h
--- a/Engine/Shaders/Shared/PathTracingDefinitions.h
+++ b/Engine/Shaders/Shared/PathTracingDefinitions.h
@@ -24,6 +24,8 @@
 #define PATHTRACING_LIGHT_SPOT								(3 << 6)
 #define PATHTRACING_LIGHT_RECT								(4 << 6)
 #define PATHTRACER_FLAG_CAST_SHADOW_MASK 					(1 << 9)
+#define PATHTRACER_FLAG_CAST_VOL_SHADOW_MASK 				(1 << 10)
+#define PATHTRACER_FLAG_HAS_RECT_TEXTURE_MASK				(1 << 11)
 
 // Constants for the path tracer light grid
 #define PATHTRACER_LIGHT_GRID_SINGULAR_MASK					0x80000000u
diff --git a/Engine/Shaders/Shared/RayTracingTypes.h b/Engine/Shaders/Shared/RayTracingTypes.h
--- a/Engine/Shaders/Shared/RayTracingTypes.h
+++ b/Engine/Shaders/Shared/RayTracingTypes.h
@@ -19,10 +19,11 @@ struct FPathTracingLight {
 	FVector3f dPdu;
 	FVector3f dPdv;
 	FVector3f Color;
-	FVector3f Dimensions; // Radius,SoftRadius,Length or RectWidth,RectHeight or Sin(Angle/2),Sin(SoftAngle/2) depending on light type
+	FVector2f Dimensions; // Radius,Length or RectWidth,RectHeight or Sin(Angle/2),0 depending on light type
 	FVector2f Shaping;  // Barndoor controls for RectLights, Cone angles for spots lights
 	float   Attenuation;
 	float   FalloffExponent; // for non-inverse square decay lights only
+	float   VolumetricScatteringIntensity;  // scale for volume contributions
 	int32   IESTextureSlice;
 	uint32  Flags; // see defines PATHTRACER_FLAG_*
 	FVector3f TranslatedBoundMin;
@@ -64,10 +65,11 @@ struct FPathTracingLight {
 	float3  dPdu;
 	float3  dPdv;
 	float3  Color;
-	float3  Dimensions;
+	float2  Dimensions;
 	float2  Shaping;
 	float   Attenuation;
 	float   FalloffExponent;
+	float   VolumetricScatteringIntensity;  // scale for volume contributions
 	int     IESTextureSlice;
 	uint    Flags;
 	float3  TranslatedBoundMin;
diff --git a/Engine/Source/Runtime/Engine/Classes/Engine/Scene.h b/Engine/Source/Runtime/Engine/Classes/Engine/Scene.h
--- a/Engine/Source/Runtime/Engine/Classes/Engine/Scene.h
+++ b/Engine/Source/Runtime/Engine/Classes/Engine/Scene.h
@@ -1273,6 +1273,9 @@ struct FPostProcessSettings
 	UPROPERTY(EditAnywhere, BlueprintReadWrite, Category = Overrides, meta = (PinHiddenByDefault, InlineEditConditionToggle))
 	uint32 bOverride_PathTracingEnableReferenceDOF : 1;
 
+	UPROPERTY(EditAnywhere, BlueprintReadWrite, Category = Overrides, meta = (PinHiddenByDefault, InlineEditConditionToggle))
+	uint32 bOverride_PathTracingEnableReferenceAtmosphere : 1;
+
 	UPROPERTY(EditAnywhere, BlueprintReadWrite, Category = Overrides, meta = (PinHiddenByDefault, InlineEditConditionToggle))
 	uint32 bOverride_PathTracingEnableDenoiser : 1;
 
@@ -2104,6 +2107,11 @@ struct FPostProcessSettings
 	UPROPERTY(EditAnywhere, BlueprintReadWrite, Category = "Rendering Features|PathTracing", meta = (editcondition = "bOverride_PathTracingEnableReferenceDOF", DisplayName = "Reference Depth Of Field"))
 	uint32 PathTracingEnableReferenceDOF : 1;
 
+	/** Enables path tracing in the atmosphere instead of baking the sky atmosphere contribution into a skylight. Any skylight present in the scene will be automatically ignored when this is enabled. */
+	UPROPERTY(EditAnywhere, BlueprintReadWrite, Category = "Rendering Features|PathTracing", meta = (editcondition = "bOverride_PathTracingEnableReferenceAtmosphere", DisplayName = "Reference Atmosphere"))
+	uint32 PathTracingEnableReferenceAtmosphere : 1;
+
+
 	/** Run the currently loaded denoiser plugin on the last sample to remove noise from the output. Has no effect if a plug-in is not loaded. */
 	UPROPERTY(EditAnywhere, BlueprintReadWrite, Category = "Rendering Features|PathTracing", meta = (editcondition = "bOverride_PathTracingEnableDenoiser", DisplayName = "Denoiser"))
 	uint32 PathTracingEnableDenoiser : 1;
diff --git a/Engine/Source/Runtime/Engine/Private/Scene.cpp b/Engine/Source/Runtime/Engine/Private/Scene.cpp
--- a/Engine/Source/Runtime/Engine/Private/Scene.cpp
+++ b/Engine/Source/Runtime/Engine/Private/Scene.cpp
@@ -598,6 +598,7 @@ FPostProcessSettings::FPostProcessSettings()
 	PathTracingEnableEmissive = 1;
 	PathTracingMaxPathExposure = 30.0f;
 	PathTracingEnableReferenceDOF = 0;
+	PathTracingEnableReferenceAtmosphere = 0;
 	PathTracingEnableDenoiser = 1;
 	
 	bMobileHQGaussian = false;
@@ -797,6 +798,7 @@ FPostProcessSettings::FPostProcessSettings(const FPostProcessSettings& Settings)
 	, bOverride_PathTracingEnableEmissive(Settings.bOverride_PathTracingEnableEmissive)
 	, bOverride_PathTracingMaxPathExposure(Settings.bOverride_PathTracingMaxPathExposure)
 	, bOverride_PathTracingEnableReferenceDOF(Settings.bOverride_PathTracingEnableReferenceDOF)
+	, bOverride_PathTracingEnableReferenceAtmosphere(Settings.bOverride_PathTracingEnableReferenceAtmosphere)
 	, bOverride_PathTracingEnableDenoiser(Settings.bOverride_PathTracingEnableDenoiser)
 
 	, bMobileHQGaussian(Settings.bMobileHQGaussian)
@@ -983,6 +985,7 @@ FPostProcessSettings::FPostProcessSettings(const FPostProcessSettings& Settings)
 	, PathTracingEnableEmissive(Settings.PathTracingEnableEmissive)
 	, PathTracingMaxPathExposure(Settings.PathTracingMaxPathExposure)
 	, PathTracingEnableReferenceDOF(Settings.PathTracingEnableReferenceDOF)
+	, PathTracingEnableReferenceAtmosphere(Settings.PathTracingEnableReferenceAtmosphere)
 	, PathTracingEnableDenoiser(Settings.PathTracingEnableDenoiser)
 
 	, ScreenPercentage_DEPRECATED(Settings.ScreenPercentage_DEPRECATED)
diff --git a/Engine/Source/Runtime/Engine/Private/SceneView.cpp b/Engine/Source/Runtime/Engine/Private/SceneView.cpp
--- a/Engine/Source/Runtime/Engine/Private/SceneView.cpp
+++ b/Engine/Source/Runtime/Engine/Private/SceneView.cpp
@@ -1674,6 +1674,11 @@ void FSceneView::OverridePostProcessSettings(const FPostProcessSettings& Src, fl
 			Dest.PathTracingEnableReferenceDOF = Src.PathTracingEnableReferenceDOF;
 		}
 
+		if (Src.bOverride_PathTracingEnableReferenceAtmosphere)
+		{
+			Dest.PathTracingEnableReferenceAtmosphere = Src.PathTracingEnableReferenceAtmosphere;
+		}
+
 		if (Src.bOverride_PathTracingEnableDenoiser)
 		{
 			Dest.PathTracingEnableDenoiser = Src.PathTracingEnableDenoiser;
diff --git a/Engine/Source/Runtime/Renderer/Private/DeferredShadingRenderer.cpp b/Engine/Source/Runtime/Renderer/Private/DeferredShadingRenderer.cpp
--- a/Engine/Source/Runtime/Renderer/Private/DeferredShadingRenderer.cpp
+++ b/Engine/Source/Runtime/Renderer/Private/DeferredShadingRenderer.cpp
@@ -1985,7 +1985,8 @@ void FDeferredShadingSceneRenderer::Render(FRDGBuilder& GraphBuilder)
 	// Setups the final FViewInfo::ViewRect.
 	PrepareViewRectsForRendering(GraphBuilder.RHICmdList);
 
-	if (ShouldRenderSkyAtmosphere(Scene, ViewFamily.EngineShowFlags))
+	const bool bPathTracedAtmosphere = ViewFamily.EngineShowFlags.PathTracing && Views.Num() > 0 && Views[0].FinalPostProcessSettings.PathTracingEnableReferenceAtmosphere;
+	if (ShouldRenderSkyAtmosphere(Scene, ViewFamily.EngineShowFlags) && !bPathTracedAtmosphere)
 	{
 		for (int32 LightIndex = 0; LightIndex < NUM_ATMOSPHERE_LIGHTS; ++LightIndex)
 		{
diff --git a/Engine/Source/Runtime/Renderer/Private/PathTracing.cpp b/Engine/Source/Runtime/Renderer/Private/PathTracing.cpp
--- a/Engine/Source/Runtime/Renderer/Private/PathTracing.cpp
+++ b/Engine/Source/Runtime/Renderer/Private/PathTracing.cpp
@@ -97,6 +97,22 @@ TAutoConsoleVariable<int32> CVarPathTracingMISMode(
 	ECVF_RenderThreadSafe
 );
 
+TAutoConsoleVariable<int32> CVarPathTracingVolumeMISMode(
+	TEXT("r.PathTracing.VolumeMISMode"),
+	1,
+	TEXT("Selects the sampling technique for volumetric integration of local lighting (default = 1)\n")
+	TEXT("0: Density sampling\n")
+	TEXT("1: Light sampling (default)\n"),
+	ECVF_RenderThreadSafe
+);
+
+TAutoConsoleVariable<int32> CVarPathTracingMaxRaymarchSteps(
+	TEXT("r.PathTracing.MaxRaymarchSteps"),
+	256,
+	TEXT("Upper limit on the number of ray marching steps in volumes. This limit should not be hit in most cases, but raising it can reduce bias in case it is. (default = 256)."),
+	ECVF_RenderThreadSafe
+);
+
 TAutoConsoleVariable<int32> CVarPathTracingMISCompensation(
 	TEXT("r.PathTracing.MISCompensation"),
 	1,
@@ -158,6 +174,20 @@ TAutoConsoleVariable<int32> CVarPathTracingEnableCameraBackfaceCulling(
 	ECVF_RenderThreadSafe
 );
 
+TAutoConsoleVariable<int32> CVarPathTracingAtmosphereOpticalDepthLutResolution(
+	TEXT("r.PathTracing.AtmosphereOpticalDepthLUTResolution"),
+	512,
+	TEXT("Size of the square lookup texture used for transmittance calculations by the path tracer in reference atmosphere mode.  (default = 512)"),
+	ECVF_RenderThreadSafe
+);
+
+TAutoConsoleVariable<int32> CVarPathTracingAtmosphereOpticalDepthLutNumSamples(
+	TEXT("r.PathTracing.AtmosphereOpticalDepthLUTNumSamples"),
+	16384,
+	TEXT("Number of ray marching samples used when building the transmittance lookup texture used for transmittance calculations by the path tracer in reference atmosphere mode.  (default = 16384)"),
+	ECVF_RenderThreadSafe
+);
+
 TAutoConsoleVariable<int32> CVarPathTracingFrameIndependentTemporalSeed(
 	TEXT("r.PathTracing.FrameIndependentTemporalSeed"),
 	1,
@@ -248,12 +278,16 @@ BEGIN_SHADER_PARAMETER_STRUCT(FPathTracingData, )
 	SHADER_PARAMETER(uint32, MaxBounces)
 	SHADER_PARAMETER(uint32, MaxSSSBounces)
 	SHADER_PARAMETER(uint32, MISMode)
+	SHADER_PARAMETER(uint32, VolumeMISMode)
 	SHADER_PARAMETER(uint32, ApproximateCaustics)
 	SHADER_PARAMETER(uint32, EnableCameraBackfaceCulling)
 	SHADER_PARAMETER(uint32, EnableDirectLighting)
 	SHADER_PARAMETER(uint32, EnableEmissive)
 	SHADER_PARAMETER(uint32, SamplerType)
 	SHADER_PARAMETER(uint32, VisualizeLightGrid)
+	SHADER_PARAMETER(uint32, EnableAtmosphere)
+	SHADER_PARAMETER(uint32, EnableFog)
+	SHADER_PARAMETER(int32, MaxRaymarchSteps)
 	SHADER_PARAMETER(float, MaxPathIntensity)
 	SHADER_PARAMETER(float, MaxNormalBias)
 	SHADER_PARAMETER(float, FilterWidth)
@@ -284,6 +318,7 @@ struct FPathTracingConfig
 			PathTracingData.MaxBounces != Other.PathTracingData.MaxBounces ||
 			PathTracingData.MaxSSSBounces != Other.PathTracingData.MaxSSSBounces ||
 			PathTracingData.MISMode != Other.PathTracingData.MISMode ||
+			PathTracingData.VolumeMISMode != Other.PathTracingData.VolumeMISMode ||
 			PathTracingData.SamplerType != Other.PathTracingData.SamplerType ||
 			PathTracingData.ApproximateCaustics != Other.PathTracingData.ApproximateCaustics ||
 			PathTracingData.EnableCameraBackfaceCulling != Other.PathTracingData.EnableCameraBackfaceCulling ||
@@ -295,6 +330,9 @@ struct FPathTracingConfig
 			PathTracingData.AbsorptionScale != Other.PathTracingData.AbsorptionScale ||
 			PathTracingData.CameraFocusDistance != Other.PathTracingData.CameraFocusDistance ||
 			PathTracingData.CameraLensRadius != Other.PathTracingData.CameraLensRadius ||
+			PathTracingData.EnableAtmosphere != Other.PathTracingData.EnableAtmosphere ||
+			PathTracingData.EnableFog != Other.PathTracingData.EnableFog ||
+			PathTracingData.MaxRaymarchSteps != Other.PathTracingData.MaxRaymarchSteps ||
 			ViewRect != Other.ViewRect ||
 			LightShowFlags != Other.LightShowFlags ||
 			LightGridResolution != Other.LightGridResolution ||
@@ -305,6 +343,44 @@ struct FPathTracingConfig
 	}
 };
 
+struct FAtmosphereConfig
+{
+	FAtmosphereConfig() = default;
+	FAtmosphereConfig(const FAtmosphereUniformShaderParameters& Parameters) :
+		AtmoParameters(Parameters),
+		NumSamples(CVarPathTracingAtmosphereOpticalDepthLutNumSamples.GetValueOnRenderThread()),
+		Resolution(CVarPathTracingAtmosphereOpticalDepthLutResolution.GetValueOnRenderThread()) {}
+
+	bool IsDifferent(const FAtmosphereConfig& Other) const
+	{
+		// Compare only those parameters which impact the LUT construction
+		return
+			AtmoParameters.BottomRadiusKm != Other.AtmoParameters.BottomRadiusKm ||
+			AtmoParameters.TopRadiusKm != Other.AtmoParameters.TopRadiusKm ||
+			AtmoParameters.RayleighDensityExpScale != Other.AtmoParameters.RayleighDensityExpScale ||
+			AtmoParameters.RayleighScattering != Other.AtmoParameters.RayleighScattering ||
+			AtmoParameters.MieScattering != Other.AtmoParameters.MieScattering ||
+			AtmoParameters.MieDensityExpScale != Other.AtmoParameters.MieDensityExpScale ||
+			AtmoParameters.MieExtinction != Other.AtmoParameters.MieExtinction ||
+			AtmoParameters.MieAbsorption != Other.AtmoParameters.MieAbsorption ||
+			AtmoParameters.AbsorptionDensity0LayerWidth != Other.AtmoParameters.AbsorptionDensity0LayerWidth ||
+			AtmoParameters.AbsorptionDensity0ConstantTerm != Other.AtmoParameters.AbsorptionDensity0ConstantTerm ||
+			AtmoParameters.AbsorptionDensity0LinearTerm != Other.AtmoParameters.AbsorptionDensity0LinearTerm ||
+			AtmoParameters.AbsorptionDensity1ConstantTerm != Other.AtmoParameters.AbsorptionDensity1ConstantTerm ||
+			AtmoParameters.AbsorptionDensity1LinearTerm != Other.AtmoParameters.AbsorptionDensity1LinearTerm ||
+			AtmoParameters.AbsorptionExtinction != Other.AtmoParameters.AbsorptionExtinction ||
+			NumSamples != Other.NumSamples ||
+			Resolution != Other.Resolution;
+	}
+
+	// hold a copy of the parameters that influence LUT construction so we can detect when they change
+	FAtmosphereUniformShaderParameters AtmoParameters;
+
+	// parameters for the LUT itself
+	uint32 NumSamples;
+	uint32 Resolution;
+};
+
 struct FPathTracingState {
 	FPathTracingConfig LastConfig;
 	// Textures holding onto the accumulated frame data
@@ -313,6 +389,10 @@ struct FPathTracingState {
 	TRefCountPtr<IPooledRenderTarget> NormalRT;
 	TRefCountPtr<IPooledRenderTarget> RadianceDenoisedRT;
 
+	// Texture holding onto the precomputed atmosphere data
+	TRefCountPtr<IPooledRenderTarget> AtmosphereOpticalDepthLUT;
+	FAtmosphereConfig LastAtmosphereConfig;
+
 	// Current sample index to be rendered by the path tracer - this gets incremented each time the path tracer accumulates a frame of samples
 	uint32 SampleIndex = 0;
 
@@ -322,7 +402,7 @@ struct FPathTracingState {
 };
 
 // This function prepares the portion of shader arguments that may involve invalidating the path traced state
-static void PrepareShaderArgs(const FViewInfo& View, FPathTracingData& PathTracingData)
+static void PreparePathTracingData(const FScene* Scene, const FViewInfo& View, FPathTracingData& PathTracingData)
 {
 	PathTracingData.EnableDirectLighting = true;
 	int32 MaxBounces = CVarPathTracingMaxBounces.GetValueOnRenderThread();
@@ -356,6 +436,7 @@ static void PrepareShaderArgs(const FViewInfo& View, FPathTracingData& PathTraci
 	PathTracingData.MaxSSSBounces = CVarPathTracingMaxSSSBounces.GetValueOnRenderThread();
 	PathTracingData.MaxNormalBias = GetRaytracingMaxNormalBias();
 	PathTracingData.MISMode = CVarPathTracingMISMode.GetValueOnRenderThread();
+	PathTracingData.VolumeMISMode = CVarPathTracingVolumeMISMode.GetValueOnRenderThread();
 	PathTracingData.MaxPathIntensity = CVarPathTracingMaxPathIntensity.GetValueOnRenderThread();
 	if (PathTracingData.MaxPathIntensity <= 0)
 	{
@@ -386,6 +467,18 @@ static void PrepareShaderArgs(const FViewInfo& View, FPathTracingData& PathTraci
 		PathTracingData.CameraFocusDistance = View.FinalPostProcessSettings.DepthOfFieldFocalDistance;
 		PathTracingData.CameraLensRadius = 0.5f * FocalLengthInCM / View.FinalPostProcessSettings.DepthOfFieldFstop;
 	}
+	PathTracingData.EnableAtmosphere =
+		ShouldRenderSkyAtmosphere(Scene, View.Family->EngineShowFlags) && 
+		View.SkyAtmosphereUniformShaderParameters &&
+		View.FinalPostProcessSettings.PathTracingEnableReferenceAtmosphere;
+	PathTracingData.EnableFog = ShouldRenderFog(*View.Family)
+		&& Scene->ExponentialFogs.Num() > 0
+		&& Scene->ExponentialFogs[0].bEnableVolumetricFog
+		&& Scene->ExponentialFogs[0].VolumetricFogDistance > 0
+		&& Scene->ExponentialFogs[0].VolumetricFogExtinctionScale > 0
+		&& (Scene->ExponentialFogs[0].FogData[0].Density > 0 ||
+			Scene->ExponentialFogs[0].FogData[1].Density > 0);
+	PathTracingData.MaxRaymarchSteps = CVarPathTracingMaxRaymarchSteps.GetValueOnRenderThread();
 }
 
 static bool ShouldCompilePathTracingShadersForProject(EShaderPlatform ShaderPlatform)
@@ -495,6 +588,55 @@ class FPathTracingBuildLightGridCS : public FGlobalShader
 };
 IMPLEMENT_SHADER_TYPE(, FPathTracingBuildLightGridCS, TEXT("/Engine/Private/PathTracing/PathTracingBuildLightGrid.usf"), TEXT("PathTracingBuildLightGridCS"), SF_Compute);
 
+// make a small custom struct to represent fog, because we need a more physical approach than the rest of the engine
+BEGIN_SHADER_PARAMETER_STRUCT(FPathTracingFogParameters, )
+	SHADER_PARAMETER(FVector2f, FogDensity)
+	SHADER_PARAMETER(FVector2f, FogHeight)
+	SHADER_PARAMETER(FVector2f, FogFalloff)
+	SHADER_PARAMETER(FLinearColor, FogAlbedo)
+	SHADER_PARAMETER(float, FogPhaseG)
+	SHADER_PARAMETER(FVector2f, FogCenter)
+	SHADER_PARAMETER(float, FogMinZ)
+	SHADER_PARAMETER(float, FogMaxZ)
+	SHADER_PARAMETER(float, FogRadius)
+END_SHADER_PARAMETER_STRUCT()
+
+static FPathTracingFogParameters PrepareFogParameters(const FViewInfo& View, const FExponentialHeightFogSceneInfo& FogInfo)
+{
+	static_assert(FExponentialHeightFogSceneInfo::NumFogs == 2, "Path tracing code assumes a fixed number of fogs");
+	FPathTracingFogParameters Parameters = {};
+
+	const FVector PreViewTranslation = View.ViewMatrices.GetPreViewTranslation();
+
+	Parameters.FogDensity.X = FogInfo.FogData[0].Density * FogInfo.VolumetricFogExtinctionScale;
+	Parameters.FogDensity.Y = FogInfo.FogData[1].Density * FogInfo.VolumetricFogExtinctionScale;
+	Parameters.FogHeight.X = FogInfo.FogData[0].Height + PreViewTranslation.Z;
+	Parameters.FogHeight.Y = FogInfo.FogData[1].Height + PreViewTranslation.Z;
+	Parameters.FogFalloff.X = FogInfo.FogData[0].HeightFalloff;
+	Parameters.FogFalloff.Y = FogInfo.FogData[1].HeightFalloff;
+	Parameters.FogAlbedo = FogInfo.VolumetricFogAlbedo;
+	Parameters.FogPhaseG = FogInfo.VolumetricFogScatteringDistribution;
+
+	const float DensityEpsilon = 1e-6f;
+	const float Radius = FogInfo.VolumetricFogDistance;
+	// compute the value of Z at which the density becomes negligible (but don't go beyond the radius)
+	const float ZMax0 = Parameters.FogHeight.X + FMath::Min(Radius, FMath::Log2(FMath::Max(Parameters.FogDensity.X, DensityEpsilon) / DensityEpsilon) / Parameters.FogFalloff.X);
+	const float ZMax1 = Parameters.FogHeight.Y + FMath::Min(Radius, FMath::Log2(FMath::Max(Parameters.FogDensity.Y, DensityEpsilon) / DensityEpsilon) / Parameters.FogFalloff.Y);
+	// lowest point is just defined by the radius (fog is homogeneous below the height)
+	const float ZMin0 = Parameters.FogHeight.X - Radius;
+	const float ZMin1 = Parameters.FogHeight.Y - Radius;
+
+	// center X,Y around the current view point
+	// NOTE: this can lead to "sliding" when the view distance is low, would it be better to just use the component center instead?
+	// NOTE: the component position is not available here, would need to be added to FogInfo ...
+	const FVector O = View.ViewMatrices.GetViewOrigin() + PreViewTranslation;
+	Parameters.FogCenter = FVector2f(O.X, O.Y);
+	Parameters.FogMinZ = FMath::Min(ZMin0, ZMin1);
+	Parameters.FogMaxZ = FMath::Max(ZMax0, ZMax1);
+	Parameters.FogRadius = Radius;
+	return Parameters;
+}
+
 class FPathTracingRG : public FGlobalShader
 {
 	DECLARE_GLOBAL_SHADER(FPathTracingRG)
@@ -534,6 +676,16 @@ class FPathTracingRG : public FGlobalShader
 		// Skylight
 		SHADER_PARAMETER_STRUCT_INCLUDE(FPathTracingSkylight, SkylightParameters)
 
+		// sky atmosphere
+		SHADER_PARAMETER_STRUCT_REF(FAtmosphereUniformShaderParameters, Atmosphere)
+		SHADER_PARAMETER_RDG_TEXTURE(Texture2D, AtmosphereOpticalDepthLUT)
+		SHADER_PARAMETER_SAMPLER(SamplerState, AtmosphereOpticalDepthLUTSampler)
+		SHADER_PARAMETER(FVector3f, PlanetCenterTranslatedWorldHi)
+		SHADER_PARAMETER(FVector3f, PlanetCenterTranslatedWorldLo)
+
+		// exponential height fog
+		SHADER_PARAMETER_STRUCT_INCLUDE(FPathTracingFogParameters, FogParameters)
+
 		// IES Profiles
 		SHADER_PARAMETER_RDG_TEXTURE(Texture2DArray, IESTexture)
 		SHADER_PARAMETER_SAMPLER(SamplerState, IESTextureSampler) // Shared sampler for all IES profiles
@@ -578,6 +730,34 @@ class FPathTracingIESAtlasCS : public FGlobalShader
 };
 IMPLEMENT_SHADER_TYPE(, FPathTracingIESAtlasCS, TEXT("/Engine/Private/PathTracing/PathTracingIESAtlas.usf"), TEXT("PathTracingIESAtlasCS"), SF_Compute);
 
+class FPathTracingBuildAtmosphereOpticalDepthLUTCS : public FGlobalShader
+{
+	DECLARE_GLOBAL_SHADER(FPathTracingBuildAtmosphereOpticalDepthLUTCS)
+	SHADER_USE_PARAMETER_STRUCT(FPathTracingBuildAtmosphereOpticalDepthLUTCS, FGlobalShader)
+
+	static bool ShouldCompilePermutation(const FGlobalShaderPermutationParameters& Parameters)
+	{
+		return ShouldCompileRayTracingShadersForProject(Parameters.Platform);
+	}
+
+	static void ModifyCompilationEnvironment(const FGlobalShaderPermutationParameters& Parameters, FShaderCompilerEnvironment& OutEnvironment)
+	{
+		//OutEnvironment.CompilerFlags.Add(CFLAG_WarningsAsErrors);
+		OutEnvironment.SetDefine(TEXT("THREADGROUPSIZE_X"), FComputeShaderUtils::kGolden2DGroupSize);
+		OutEnvironment.SetDefine(TEXT("THREADGROUPSIZE_Y"), FComputeShaderUtils::kGolden2DGroupSize);
+	}
+
+	BEGIN_SHADER_PARAMETER_STRUCT(FParameters, )
+		SHADER_PARAMETER(uint32, NumSamples)
+		SHADER_PARAMETER(uint32, Resolution)
+		SHADER_PARAMETER_STRUCT_REF(FAtmosphereUniformShaderParameters, Atmosphere)
+		SHADER_PARAMETER_RDG_TEXTURE_UAV(RWTexture2D, AtmosphereOpticalDepthLUT)
+	END_SHADER_PARAMETER_STRUCT()
+};
+IMPLEMENT_SHADER_TYPE(, FPathTracingBuildAtmosphereOpticalDepthLUTCS, TEXT("/Engine/Private/PathTracing/PathTracingBuildAtmosphereLUT.usf"), TEXT("PathTracingBuildAtmosphereOpticalDepthLUTCS"), SF_Compute);
+
+
+
 template<bool UseAnyHitShader, bool UseIntersectionShader>
 class TPathTracingMaterial : public FMeshMaterialShader
 {
@@ -1024,7 +1204,10 @@ void SetLightParameters(FRDGBuilder& GraphBuilder, FPathTracingRG::FParameters*
 
 	// Prepend SkyLight to light buffer since it is not part of the regular light list
 	const float Inf = std::numeric_limits<float>::infinity();
-	if (PrepareSkyTexture(GraphBuilder, Scene, View, true, UseMISCompensation, &PassParameters->SkylightParameters))
+	// skylight should be excluded if we are using the reference atmosphere calculation (don't bother checking again if an atmosphere is present)
+	const bool bUseAtmosphere = PassParameters->PathTracingData.EnableAtmosphere != 0;
+	const bool bEnableSkydome = !bUseAtmosphere;
+	if (PrepareSkyTexture(GraphBuilder, Scene, View, bEnableSkydome, UseMISCompensation, &PassParameters->SkylightParameters))
 	{
 		check(Scene->SkyLight != nullptr);
 		FPathTracingLight& DestLight = Lights[NumLights++];
@@ -1033,6 +1216,8 @@ void SetLightParameters(FRDGBuilder& GraphBuilder, FPathTracingRG::FParameters*
 		DestLight.Flags |= PATHTRACER_FLAG_LIGHTING_CHANNEL_MASK;
 		DestLight.Flags |= PATHTRACING_LIGHT_SKY;
 		DestLight.Flags |= Scene->SkyLight->bCastShadows ? PATHTRACER_FLAG_CAST_SHADOW_MASK : 0;
+		DestLight.Flags |= Scene->SkyLight->bCastVolumetricShadow ? PATHTRACER_FLAG_CAST_VOL_SHADOW_MASK : 0;
+		DestLight.VolumetricScatteringIntensity = Scene->SkyLight->VolumetricScatteringIntensity;
 		DestLight.IESTextureSlice = -1;
 		DestLight.TranslatedBoundMin = FVector3f(-Inf, -Inf, -Inf);
 		DestLight.TranslatedBoundMax = FVector3f( Inf,  Inf,  Inf);
@@ -1072,6 +1257,7 @@ void SetLightParameters(FRDGBuilder& GraphBuilder, FPathTracingRG::FParameters*
 			DestLight.Flags = Transmission ? PATHTRACER_FLAG_TRANSMISSION_MASK : 0;
 			DestLight.Flags |= LightingChannelMask & PATHTRACER_FLAG_LIGHTING_CHANNEL_MASK;
 			DestLight.Flags |= Light.LightSceneInfo->Proxy->CastsDynamicShadow() ? PATHTRACER_FLAG_CAST_SHADOW_MASK : 0;
+			DestLight.Flags |= Light.LightSceneInfo->Proxy->CastsVolumetricShadow() ? PATHTRACER_FLAG_CAST_VOL_SHADOW_MASK : 0;
 			DestLight.IESTextureSlice = -1;
 
 			// these mean roughly the same thing across all light types
@@ -1082,11 +1268,12 @@ void SetLightParameters(FRDGBuilder& GraphBuilder, FPathTracingRG::FParameters*
 			DestLight.dPdv = LightParameters.Tangent;
 			DestLight.Attenuation = LightParameters.InvRadius;
 			DestLight.FalloffExponent = 0;
+			DestLight.VolumetricScatteringIntensity = Light.LightSceneInfo->Proxy->GetVolumetricScatteringIntensity();
 			DestLight.RectLightAtlasUVOffset = 0;
 			DestLight.RectLightAtlasUVScale = 0;
 
 			DestLight.Normal = LightParameters.Direction;
-			DestLight.Dimensions = FVector3f(LightParameters.SourceRadius, LightParameters.SoftSourceRadius, 0.0f);
+			DestLight.Dimensions = FVector2f(LightParameters.SourceRadius, 0.0f);
 			DestLight.Flags |= PATHTRACING_LIGHT_DIRECTIONAL;
 
 			DestLight.TranslatedBoundMin = FVector3f(-Inf, -Inf, -Inf);
@@ -1094,6 +1281,13 @@ void SetLightParameters(FRDGBuilder& GraphBuilder, FPathTracingRG::FParameters*
 		}
 	}
 
+	if (bUseAtmosphere)
+	{
+		// show directional lights when atmosphere is enabled
+		// NOTE: there cannot be any skydome in this case
+		PassParameters->SceneVisibleLightCount = NumLights;
+	}
+
 	uint32 NumInfiniteLights = NumLights;
 
 	int32 NextRectTextureIndex = 0;
@@ -1128,6 +1322,7 @@ void SetLightParameters(FRDGBuilder& GraphBuilder, FPathTracingRG::FParameters*
 		DestLight.Flags = Transmission ? PATHTRACER_FLAG_TRANSMISSION_MASK : 0;
 		DestLight.Flags |= LightingChannelMask & PATHTRACER_FLAG_LIGHTING_CHANNEL_MASK;
 		DestLight.Flags |= Light.LightSceneInfo->Proxy->CastsDynamicShadow() ? PATHTRACER_FLAG_CAST_SHADOW_MASK : 0;
+		DestLight.Flags |= Light.LightSceneInfo->Proxy->CastsVolumetricShadow() ? PATHTRACER_FLAG_CAST_VOL_SHADOW_MASK : 0;
 		DestLight.IESTextureSlice = -1;
 
 		if (View.Family->EngineShowFlags.TexturedLightProfiles)
@@ -1148,6 +1343,7 @@ void SetLightParameters(FRDGBuilder& GraphBuilder, FPathTracingRG::FParameters*
 		DestLight.dPdv = LightParameters.Tangent;
 		DestLight.Attenuation = LightParameters.InvRadius;
 		DestLight.FalloffExponent = 0;
+		DestLight.VolumetricScatteringIntensity = Light.LightSceneInfo->Proxy->GetVolumetricScatteringIntensity();
 		DestLight.RectLightAtlasUVOffset = 0;
 		DestLight.RectLightAtlasUVScale = 0;
 
@@ -1155,7 +1351,7 @@ void SetLightParameters(FRDGBuilder& GraphBuilder, FPathTracingRG::FParameters*
 		{
 			case LightType_Rect:
 			{
-				DestLight.Dimensions = FVector3f(2.0f * LightParameters.SourceRadius, 2.0f * LightParameters.SourceLength, 0.0f);
+				DestLight.Dimensions = FVector2f(2.0f * LightParameters.SourceRadius, 2.0f * LightParameters.SourceLength);
 				DestLight.Shaping = FVector2f(LightParameters.RectLightBarnCosAngle, LightParameters.RectLightBarnLength);
 				DestLight.FalloffExponent = LightParameters.FalloffExponent;
 				DestLight.Flags |= Light.LightSceneInfo->Proxy->IsInverseSquared() ? 0 : PATHTRACER_FLAG_NON_INVERSE_SQUARE_FALLOFF_MASK;
@@ -1178,11 +1374,15 @@ void SetLightParameters(FRDGBuilder& GraphBuilder, FPathTracingRG::FParameters*
 				// Rect light atlas UV transformation
 				DestLight.RectLightAtlasUVOffset = EncodeToF16x2(LightParameters.RectLightAtlasUVOffset);
 				DestLight.RectLightAtlasUVScale  = EncodeToF16x2(LightParameters.RectLightAtlasUVScale);
+				if (LightParameters.RectLightAtlasMaxLevel < 16)
+				{
+					DestLight.Flags |= PATHTRACER_FLAG_HAS_RECT_TEXTURE_MASK;
+				}
 				break;
 			}
 			case LightType_Spot:
 			{
-				DestLight.Dimensions = FVector3f(LightParameters.SourceRadius, LightParameters.SoftSourceRadius, LightParameters.SourceLength);
+				DestLight.Dimensions = FVector2f(LightParameters.SourceRadius, LightParameters.SourceLength);
 				DestLight.Shaping = LightParameters.SpotAngles;
 				DestLight.FalloffExponent = LightParameters.FalloffExponent;
 				DestLight.Flags |= Light.LightSceneInfo->Proxy->IsInverseSquared() ? 0 : PATHTRACER_FLAG_NON_INVERSE_SQUARE_FALLOFF_MASK;
@@ -1211,7 +1411,7 @@ void SetLightParameters(FRDGBuilder& GraphBuilder, FPathTracingRG::FParameters*
 			}
 			case LightType_Point:
 			{
-				DestLight.Dimensions = FVector3f(LightParameters.SourceRadius, LightParameters.SoftSourceRadius, LightParameters.SourceLength);
+				DestLight.Dimensions = FVector2f(LightParameters.SourceRadius, LightParameters.SourceLength);
 				DestLight.FalloffExponent = LightParameters.FalloffExponent;
 				DestLight.Flags |= Light.LightSceneInfo->Proxy->IsInverseSquared() ? 0 : PATHTRACER_FLAG_NON_INVERSE_SQUARE_FALLOFF_MASK;
 				DestLight.Flags |= PATHTRACING_LIGHT_POINT;
@@ -1325,6 +1525,14 @@ uint32 FSceneViewState::GetPathTracingSampleCount() const {
 	return State ? State->LastConfig.PathTracingData.MaxSamples : 0;
 }
 
+static void SplitDouble(double x, float* hi, float* lo)
+{
+	const double SPLIT = 134217729.0; // 2^27+1
+	double temp = SPLIT * x;
+	*hi = static_cast<float>(temp - (temp - x));
+	*lo = static_cast<float>(x - *hi);
+}
+
 
 BEGIN_SHADER_PARAMETER_STRUCT(FDenoiseTextureParameters, )
 	RDG_TEXTURE_ACCESS(InputTexture, ERHIAccess::CopySrc)
@@ -1365,7 +1573,6 @@ void FDeferredShadingSceneRenderer::RenderPathTracing(
 	MaxSPP = FMath::Max(MaxSPP, 1u);
 	Config.LockedSamplingPattern = CVarPathTracingFrameIndependentTemporalSeed.GetValueOnRenderThread() == 0;
 
-
 	// compute an integer code of what show flags related to lights are currently enabled so we can detect changes
 	Config.LightShowFlags = 0;
 	Config.LightShowFlags |= View.Family->EngineShowFlags.SkyLighting           ? 1 << 0 : 0;
@@ -1375,7 +1582,7 @@ void FDeferredShadingSceneRenderer::RenderPathTracing(
 	Config.LightShowFlags |= View.Family->EngineShowFlags.PointLights           ? 1 << 4 : 0;
 	Config.LightShowFlags |= View.Family->EngineShowFlags.TexturedLightProfiles ? 1 << 5 : 0;
 
-	PrepareShaderArgs(View, Config.PathTracingData);
+	PreparePathTracingData(Scene, View, Config.PathTracingData);
 
 	Config.VisibleLights = CVarPathTracingVisibleLights.GetValueOnRenderThread() != 0;
 	Config.UseMISCompensation = Config.PathTracingData.MISMode == 2 && CVarPathTracingMISCompensation.GetValueOnRenderThread() != 0;
@@ -1413,6 +1620,52 @@ void FDeferredShadingSceneRenderer::RenderPathTracing(
 		View.ViewState->PathTracingInvalidate();
 	}
 
+	// prepare atmosphere optical depth lookup texture (if needed)
+	FRDGTexture* AtmosphereOpticalDepthLUT = nullptr;
+	if (Config.PathTracingData.EnableAtmosphere)
+	{
+		check(Scene->GetSkyAtmosphereSceneInfo() != nullptr);
+		check(Scene->GetSkyAtmosphereSceneInfo()->GetAtmosphereShaderParameters() != nullptr);
+		FAtmosphereConfig AtmoConfig(*Scene->GetSkyAtmosphereSceneInfo()->GetAtmosphereShaderParameters());
+		if (!PathTracingState->AtmosphereOpticalDepthLUT.IsValid() || PathTracingState->LastAtmosphereConfig.IsDifferent(AtmoConfig))
+		{
+			PathTracingState->LastAtmosphereConfig = AtmoConfig;
+			// need to create a new LUT
+			FRDGTextureDesc Desc = FRDGTextureDesc::Create2D(
+				FIntPoint(AtmoConfig.Resolution, AtmoConfig.Resolution),
+				PF_A32B32G32R32F,
+				FClearValueBinding::None,
+				TexCreate_ShaderResource | TexCreate_UAV);
+			AtmosphereOpticalDepthLUT = GraphBuilder.CreateTexture(Desc, TEXT("PathTracer.AtmosphereOpticalDepthLUT"), ERDGTextureFlags::MultiFrame);
+			FPathTracingBuildAtmosphereOpticalDepthLUTCS::FParameters* PassParameters = GraphBuilder.AllocParameters<FPathTracingBuildAtmosphereOpticalDepthLUTCS::FParameters>();
+			PassParameters->NumSamples = AtmoConfig.NumSamples;
+			PassParameters->Resolution = AtmoConfig.Resolution;
+			PassParameters->Atmosphere = Scene->GetSkyAtmosphereSceneInfo()->GetAtmosphereUniformBuffer();
+			PassParameters->AtmosphereOpticalDepthLUT = GraphBuilder.CreateUAV(AtmosphereOpticalDepthLUT);
+			TShaderMapRef<FPathTracingBuildAtmosphereOpticalDepthLUTCS> ComputeShader(GetGlobalShaderMap(GMaxRHIFeatureLevel));
+			FComputeShaderUtils::AddPass(
+				GraphBuilder,
+				RDG_EVENT_NAME("Path Tracing Atmosphere Optical Depth LUT (Resolution=%u, NumSamples=%u)", AtmoConfig.Resolution, AtmoConfig.NumSamples),
+				ComputeShader,
+				PassParameters,
+				FComputeShaderUtils::GetGroupCount(
+					FIntPoint(AtmoConfig.Resolution, AtmoConfig.Resolution),
+					FIntPoint(FComputeShaderUtils::kGolden2DGroupSize, FComputeShaderUtils::kGolden2DGroupSize))
+			);
+			GraphBuilder.QueueTextureExtraction(AtmosphereOpticalDepthLUT, &PathTracingState->AtmosphereOpticalDepthLUT);
+		}
+		else
+		{
+			AtmosphereOpticalDepthLUT = GraphBuilder.RegisterExternalTexture(PathTracingState->AtmosphereOpticalDepthLUT, TEXT("PathTracer.AtmosphereOpticalDepthLUT"));
+		}
+	}
+	else
+	{
+		AtmosphereOpticalDepthLUT = GraphBuilder.RegisterExternalTexture(GSystemTextures.BlackDummy);
+	}
+
+
+
 	// If the scene has changed in some way (camera move, object movement, etc ...)
 	// we must invalidate the ViewState to start over from scratch
 	// NOTE: only check things like hair position changes for interactive viewports, for offline renders we don't want any chance of mid-render invalidation
@@ -1549,6 +1802,33 @@ void FDeferredShadingSceneRenderer::RenderPathTracing(
 					PassParameters->AlbedoTexture   = GraphBuilder.CreateUAV(AlbedoTexture);
 					PassParameters->NormalTexture   = GraphBuilder.CreateUAV(NormalTexture);
 
+					if (Config.PathTracingData.EnableAtmosphere)
+					{
+						PassParameters->Atmosphere = Scene->GetSkyAtmosphereSceneInfo()->GetAtmosphereUniformBuffer();
+						FVector PlanetCenterTranslatedWorld = Scene->GetSkyAtmosphereSceneInfo()->GetSkyAtmosphereSceneProxy().GetAtmosphereSetup().PlanetCenterKm * double(FAtmosphereSetup::SkyUnitToCm) + View.ViewMatrices.GetPreViewTranslation();
+						SplitDouble(PlanetCenterTranslatedWorld.X, &PassParameters->PlanetCenterTranslatedWorldHi.X, &PassParameters->PlanetCenterTranslatedWorldLo.X);
+						SplitDouble(PlanetCenterTranslatedWorld.Y, &PassParameters->PlanetCenterTranslatedWorldHi.Y, &PassParameters->PlanetCenterTranslatedWorldLo.Y);
+						SplitDouble(PlanetCenterTranslatedWorld.Z, &PassParameters->PlanetCenterTranslatedWorldHi.Z, &PassParameters->PlanetCenterTranslatedWorldLo.Z);
+					}
+					else
+					{
+						FAtmosphereUniformShaderParameters AtmosphereParams = {};
+						PassParameters->Atmosphere = CreateUniformBufferImmediate(AtmosphereParams, EUniformBufferUsage::UniformBuffer_SingleFrame);
+						PassParameters->PlanetCenterTranslatedWorldHi = FVector3f(0);
+						PassParameters->PlanetCenterTranslatedWorldLo = FVector3f(0);
+					}
+					PassParameters->AtmosphereOpticalDepthLUT = AtmosphereOpticalDepthLUT;
+					PassParameters->AtmosphereOpticalDepthLUTSampler = TStaticSamplerState<SF_Bilinear, AM_Clamp, AM_Clamp, AM_Clamp>::GetRHI();
+
+					if (Config.PathTracingData.EnableFog)
+					{
+						PassParameters->FogParameters = PrepareFogParameters(View, Scene->ExponentialFogs[0]);
+					}
+					else
+					{
+						PassParameters->FogParameters = {};
+					}
+
 					// TODO: in multi-gpu case, assign different tiles to different GPUs
 					PassParameters->TileOffset.X = TileX;
 					PassParameters->TileOffset.Y = TileY;
diff --git a/Engine/Source/Runtime/Renderer/Private/RayTracing/RayTracingGlobalIllumination.cpp b/Engine/Source/Runtime/Renderer/Private/RayTracing/RayTracingGlobalIllumination.cpp
--- a/Engine/Source/Runtime/Renderer/Private/RayTracing/RayTracingGlobalIllumination.cpp
+++ b/Engine/Source/Runtime/Renderer/Private/RayTracing/RayTracingGlobalIllumination.cpp
@@ -319,7 +319,7 @@ static void SetupLightParameters(
 			DestLight.dPdu = FVector3f::CrossProduct(LightShaderParameters.Direction, LightShaderParameters.Tangent);
 			DestLight.dPdv = LightShaderParameters.Tangent;
 			DestLight.Color = FVector3f(LightShaderParameters.Color);
-			DestLight.Dimensions = FVector3f(2.0f * LightShaderParameters.SourceRadius, 2.0f * LightShaderParameters.SourceLength, 0.0f);
+			DestLight.Dimensions = FVector2f(2.0f * LightShaderParameters.SourceRadius, 2.0f * LightShaderParameters.SourceLength);
 			DestLight.Shaping = FVector2f(LightShaderParameters.RectLightBarnCosAngle, LightShaderParameters.RectLightBarnLength);
 			DestLight.Flags |= PATHTRACING_LIGHT_RECT;
 			break;
@@ -333,7 +333,7 @@ static void SetupLightParameters(
 			// #dxr_todo: UE-72556 define these differences from Lit..
 			DestLight.Color = FVector3f(LightShaderParameters.Color);
 			float SourceRadius = 0.0; // LightShaderParameters.SourceRadius causes too much noise for little pay off at this time
-			DestLight.Dimensions = FVector3f(SourceRadius, 0.0, 0.0);
+			DestLight.Dimensions = FVector2f(SourceRadius, 0.0);
 			DestLight.Flags |= PATHTRACING_LIGHT_POINT;
 			break;
 		}
@@ -346,7 +346,7 @@ static void SetupLightParameters(
 			// #dxr_todo: UE-72556 define these differences from Lit..
 			DestLight.Color = FVector3f(LightShaderParameters.Color);
 			float SourceRadius = 0.0; // LightShaderParameters.SourceRadius causes too much noise for little pay off at this time
-			DestLight.Dimensions = FVector3f(SourceRadius, 0.0, 0.0);
+			DestLight.Dimensions = FVector2f(SourceRadius, 0.0);
 			DestLight.Shaping = LightShaderParameters.SpotAngles;
 			DestLight.Flags |= PATHTRACING_LIGHT_SPOT;
 			break;
