Strata- made UpdateBSDFWeightAfterOperatorVisit explicit for each oprator to reduce relying on compiler.

diff --git a/Engine/Shaders/Private/Strata/StrataTree.ush b/Engine/Shaders/Private/Strata/StrataTree.ush
--- a/Engine/Shaders/Private/Strata/StrataTree.ush
+++ b/Engine/Shaders/Private/Strata/StrataTree.ush
@@ -131,55 +131,50 @@ void UpdateSingleOperatorCoverageTransmittance(inout FStrataTree StrataTree, int
 }
 
 
+#define Op StrataTree.Operators[OpIndex]
 
-void UpdateBSDFWeightAfterOperatorVisit(
+void UpdateBSDFWeightAfterVerticalOperatorVisit(
 	inout FStrataTree StrataTree,
-	inout FStrataBSDF CurrentBSDF, 
+	inout FStrataBSDF CurrentBSDF,
 	int OpIndex,
 	int PreviousIsInputA)
 {
-#define Op StrataTree.Operators[OpIndex]
-
-	switch (Op.Type)
-	{
+	// PreviousIsInputA > 0 means it is a BSDF affecting the top layer so we do not affect it by anything for this vertical layering.
+	// Otherise, the BSDF comes from the bottom part so we affect it weights from the gathered top layer coverage/transmittance
+	CurrentBSDF.LuminanceWeight *= 
+		PreviousIsInputA > 0 ? 1.0f : saturate(1.0f - Op.VerticalTopCoverage) + saturate(Op.VerticalTopCoverage) * Op.VerticalTopTransmittance;
+}
 
-	case STRATA_OPERATOR_WEIGHT:
-	{
-		const float Weight = saturate(Op.Weight);
-		CurrentBSDF.LuminanceWeight *= Weight;
+void UpdateBSDFWeightAfterHorizontalOperatorVisit(
+	inout FStrataTree StrataTree,
+	inout FStrataBSDF CurrentBSDF,
+	int OpIndex,
+	int PreviousIsInputA)
+{
+	const float Mix = saturate(Op.Weight);
+	const float AMix = 1.0 - Mix;
+	const float BMix = Mix;
+	const float Weight = PreviousIsInputA > 0 ? AMix : BMix;
+	CurrentBSDF.LuminanceWeight *= Weight;
 #if STRATA_INLINE_SHADING
-		CurrentBSDF.Coverage *= Weight;
+	CurrentBSDF.Coverage *= Weight;
 #endif
-		break;
-	}
-
-	case STRATA_OPERATOR_VERTICAL:
-	{
-		// PreviousIsInputA > 0 means it is a BSDF affecting the top layer so we do not affect it by anything for this vertical layering.
-		// Otherise, the BSDF comes from the bottom part so we affect it weights from the gathered top layer coverage/transmittance
-		CurrentBSDF.LuminanceWeight *= PreviousIsInputA > 0 ? 1.0f : saturate(1.0f - Op.VerticalTopCoverage) + saturate(Op.VerticalTopCoverage) * Op.VerticalTopTransmittance;
-		break;
-	}
+}
 
-	case STRATA_OPERATOR_HORIZONTAL:
-	{
-		const float Mix = saturate(Op.Weight);
-		const float AMix = 1.0 - Mix;
-		const float BMix = Mix;
-		const float Weight = PreviousIsInputA > 0 ? AMix : BMix;
-		CurrentBSDF.LuminanceWeight *= Weight;
+void UpdateBSDFWeightAfterWeightOperatorVisit(
+	inout FStrataTree StrataTree,
+	inout FStrataBSDF CurrentBSDF,
+	int OpIndex,
+	int PreviousIsInputA)
+{
+	const float Weight = saturate(Op.Weight);
+	CurrentBSDF.LuminanceWeight *= Weight;
 #if STRATA_INLINE_SHADING
-		CurrentBSDF.Coverage *= Weight;
+	CurrentBSDF.Coverage *= Weight;
 #endif
-		break;
-	}
-
-	// case STRATA_OPERATOR_ADD: NOP
-
-	}
+}
 
 #undef Op
-}
 
 
 void LegacyUpdateBSDFsOperators(inout FStrataPixelHeader StrataPixelHeader, inout FStrataTree StrataTree, bool bRoughDiffuseEnabled, FStrataAddressing NullStrataAddressing, float3 V, float3 NullLightDir)
@@ -198,15 +193,15 @@ void LegacyUpdateBSDFsOperators(inout FStrataPixelHeader StrataPixelHeader, inou
 	UpdateSingleOperatorCoverageTransmittance(StrataTree, /*OperatorIndex*/5);
 #endif
 
-	UpdateBSDFWeightAfterOperatorVisit(StrataTree, StrataTree.BSDFs[0], /*OperatorIndex*/2, /*PreviousIsInputA*/1); // WEIGHT
-	UpdateBSDFWeightAfterOperatorVisit(StrataTree, StrataTree.BSDFs[0], /*OperatorIndex*/4, /*PreviousIsInputA*/1); // VERTICAL
+	UpdateBSDFWeightAfterWeightOperatorVisit	(StrataTree, StrataTree.BSDFs[0], /*OperatorIndex*/2, /*PreviousIsInputA*/1); // WEIGHT
+	UpdateBSDFWeightAfterVerticalOperatorVisit	(StrataTree, StrataTree.BSDFs[0], /*OperatorIndex*/4, /*PreviousIsInputA*/1); // VERTICAL
 #if !MATERIALBLENDING_MASKED && !MATERIALBLENDING_SOLID
-	UpdateBSDFWeightAfterOperatorVisit(StrataTree, StrataTree.BSDFs[0], /*OperatorIndex*/5, /*PreviousIsInputA*/1); // WEIGHT
+	UpdateBSDFWeightAfterWeightOperatorVisit	(StrataTree, StrataTree.BSDFs[0], /*OperatorIndex*/5, /*PreviousIsInputA*/1); // WEIGHT
 #endif
-	UpdateBSDFWeightAfterOperatorVisit(StrataTree, StrataTree.BSDFs[1], /*OperatorIndex*/3, /*PreviousIsInputA*/1); // WEIGHT
-	UpdateBSDFWeightAfterOperatorVisit(StrataTree, StrataTree.BSDFs[1], /*OperatorIndex*/4, /*PreviousIsInputA*/0); // VERTICAL
+	UpdateBSDFWeightAfterWeightOperatorVisit	(StrataTree, StrataTree.BSDFs[1], /*OperatorIndex*/3, /*PreviousIsInputA*/1); // WEIGHT
+	UpdateBSDFWeightAfterVerticalOperatorVisit	(StrataTree, StrataTree.BSDFs[1], /*OperatorIndex*/4, /*PreviousIsInputA*/0); // VERTICAL
 #if !MATERIALBLENDING_MASKED && !MATERIALBLENDING_SOLID
-	UpdateBSDFWeightAfterOperatorVisit(StrataTree, StrataTree.BSDFs[1], /*OperatorIndex*/5, /*PreviousIsInputA*/1); // WEIGHT
+	UpdateBSDFWeightAfterWeightOperatorVisit	(StrataTree, StrataTree.BSDFs[1], /*OperatorIndex*/5, /*PreviousIsInputA*/1); // WEIGHT
 #endif
 }
 
diff --git a/Engine/Source/Runtime/Engine/Private/Materials/HLSLMaterialTranslator.cpp b/Engine/Source/Runtime/Engine/Private/Materials/HLSLMaterialTranslator.cpp
--- a/Engine/Source/Runtime/Engine/Private/Materials/HLSLMaterialTranslator.cpp
+++ b/Engine/Source/Runtime/Engine/Private/Materials/HLSLMaterialTranslator.cpp
@@ -1409,18 +1409,18 @@ bool FHLSLMaterialTranslator::Translate()
 									{
 									case STRATA_OPERATOR_VERTICAL:
 									{
-										ResourcesString += FString::Printf(TEXT("\t UpdateBSDFWeightAfterOperatorVisit(StrataTree, StrataTree.BSDFs[%d], %d /*Op index*/, %d /*PreviousIsInputA*/); // VERTICAL\n"), BSDFIndex, CurrentOperator.Index, CurrentOperator.LeftIndex == PreviousOperatorIndex ? 1 : 0);
+										ResourcesString += FString::Printf(TEXT("\t UpdateBSDFWeightAfterVerticalOperatorVisit(StrataTree, StrataTree.BSDFs[%d], %d /*Op index*/, %d /*PreviousIsInputA*/);\n"), BSDFIndex, CurrentOperator.Index, CurrentOperator.LeftIndex == PreviousOperatorIndex ? 1 : 0);
 										break;
 									}
 									case STRATA_OPERATOR_HORIZONTAL:
 									{
-										ResourcesString += FString::Printf(TEXT("\t UpdateBSDFWeightAfterOperatorVisit(StrataTree, StrataTree.BSDFs[%d], %d /*Op index*/, %d /*PreviousIsInputA*/); // HORIZONTAL\n"), BSDFIndex, CurrentOperator.Index, CurrentOperator.LeftIndex == PreviousOperatorIndex ? 1 : 0);
+										ResourcesString += FString::Printf(TEXT("\t UpdateBSDFWeightAfterHorizontalOperatorVisit(StrataTree, StrataTree.BSDFs[%d], %d /*Op index*/, %d /*PreviousIsInputA*/);\n"), BSDFIndex, CurrentOperator.Index, CurrentOperator.LeftIndex == PreviousOperatorIndex ? 1 : 0);
 										break;
 									}
 
 									case STRATA_OPERATOR_WEIGHT:
 									{
-										ResourcesString += FString::Printf(TEXT("\t UpdateBSDFWeightAfterOperatorVisit(StrataTree, StrataTree.BSDFs[%d], %d /*Op index*/, %d /*PreviousIsInputA*/); // WEIGHT\n"), BSDFIndex, CurrentOperator.Index, 1);
+										ResourcesString += FString::Printf(TEXT("\t UpdateBSDFWeightAfterWeightOperatorVisit(StrataTree, StrataTree.BSDFs[%d], %d /*Op index*/, %d /*PreviousIsInputA*/);\n"), BSDFIndex, CurrentOperator.Index, 1);
 										break;
 									}
 
