Fix VirtualShadowMap projection with Strata.

diff --git a/Engine/Shaders/Private/VirtualShadowMaps/VirtualShadowMapProjection.usf b/Engine/Shaders/Private/VirtualShadowMaps/VirtualShadowMapProjection.usf
--- a/Engine/Shaders/Private/VirtualShadowMaps/VirtualShadowMapProjection.usf
+++ b/Engine/Shaders/Private/VirtualShadowMaps/VirtualShadowMapProjection.usf
@@ -41,6 +41,38 @@
 
 #define CONTACT_SHADOW_SAMPLES 8
 
+struct FProjectionShadingInfo
+{
+	bool bIsValid;
+	bool bIsHair;
+	bool bIsSubsurface;
+
+	half3 WorldNormal;
+};
+
+FProjectionShadingInfo GetProjectionShadingInfo(uint2 PixelPos)
+{
+	FProjectionShadingInfo Out;
+#if STRATA_ENABLED
+	FStrataAddressing StrataAddressing = GetStrataPixelDataByteOffset(PixelPos, uint2(View.BufferSizeAndInvSize.xy), Strata.MaxBytesPerPixel);
+	const FStrataPixelHeader StrataPixelHeader = UnpackStrataHeaderIn(Strata.MaterialTextureArray, StrataAddressing, Strata.TopLayerTexture);
+	const FStrataTopLayerData TopLayerData = StrataUnpackTopLayerData(Strata.TopLayerTexture.Load(uint3(PixelPos, 0)));
+	const FStrataSubsurfaceData SSSData = StrataUnpackSSSData(Strata.SSSTexture.Load(uint3(PixelPos, 0)));
+
+	Out.bIsValid = TopLayerData.bIsValid;
+	Out.bIsHair = StrataHasShadingModel(StrataPixelHeader, STRATA_BSDF_TYPE_HAIR);
+	Out.bIsSubsurface = SSSData.bIsValid && !SSSData.bIsProfile;
+	Out.WorldNormal = TopLayerData.WorldNormal;
+#else
+	const FGBufferData GBufferData = GetGBufferDataUint(PixelPos, true);
+	Out.bIsValid = GBufferData.ShadingModelID != SHADINGMODELID_UNLIT;
+	Out.bIsHair = GBufferData.ShadingModelID == SHADINGMODELID_HAIR;
+	Out.bIsSubsurface = IsSubsurfaceModel(GBufferData.ShadingModelID);
+	Out.WorldNormal = GBufferData.WorldNormal;
+#endif
+	return Out;
+}
+
 float4 EncodeLightAttenuationFromMask(float ShadowMask, float SSSTransmission)
 {
 	const float ShadowFadeFraction = 1;
@@ -85,7 +117,7 @@ float VirtualShadowMapScreenRayCast(
 	return (bHitCastContactShadow) ? HitDistance : -1.0f;
 }
 
-float3 GetEstimatedGeoWorldNormal(float3 TranslatedWorldPosition, float2 ScreenUV)
+float3 GetEstimatedGeoWorldNormal(float3 TranslatedWorldPosition, FProjectionShadingInfo ShadingInfo)
 {
 	// Figure out slope, we do world space since that is the space where we might override using the shading normal...
 	float3 TranslatedWorldPositionDDX = DDX(TranslatedWorldPosition);
@@ -99,11 +131,10 @@ float3 GetEstimatedGeoWorldNormal(float3 TranslatedWorldPosition, float2 ScreenU
 #if 1
 	// NOTE: Gbuffer normal is not the surface normal for hair; hair lighting takes a different path but possibly
 	// necessary to do some sort of special casing here (disable normal offset and biasing entirely?).
-	FGBufferData GBufferData = GetGBufferData(ScreenUV);
 	// If the estimated geo normal is too far out we assume it's broken (derivative includes other surfaces or background) and fall back to the shading normal
-	if (dot(GBufferData.WorldNormal, EstimatedGeoWorldNormal) < 0.25f)
+	if (dot(ShadingInfo.WorldNormal, EstimatedGeoWorldNormal) < 0.25f)
 	{
-		EstimatedGeoWorldNormal = GBufferData.WorldNormal;
+		EstimatedGeoWorldNormal = ShadingInfo.WorldNormal;
 	}
 #endif
 
@@ -169,14 +200,13 @@ uint bCullBackfacingPixels;
 FVirtualShadowMapSampleResult ProjectLight(
 	int VirtualShadowMapId,
 	FLightShaderParameters Light,
+	FProjectionShadingInfo ShadingInfo,
 	uint2 PixelPos,
 	float SceneDepth,
 	float ContactShadowLengthWorld,
 	float3 TranslatedWorldPosition,
-	half3 GBufferNormal,
 	const half Noise)
 {
-	const uint ShadingModelID = GetGBufferDataUint(PixelPos).ShadingModelID;
 	const bool bIsHairInput = InputType == INPUT_TYPE_HAIRSTRANDS;
 
 	// NOTE: If PrimaryView.PreViewTranslation is guaranteed to be the camera position we can simplify
@@ -199,14 +229,14 @@ FVirtualShadowMapSampleResult ProjectLight(
 
 	FVirtualShadowMapSampleResult Result = InitVirtualShadowMapSampleResult();
 
-	const bool bValidPixel = bIsHairInput || ShadingModelID != SHADINGMODELID_UNLIT;
+	const bool bValidPixel = bIsHairInput || ShadingInfo.bIsValid;
 
 	BRANCH
 	if (bInLightRegion && bValidPixel)
 	{
-		const bool bBackfaceCull = (bCullBackfacingPixels != 0) && !bIsHairInput && !IsSubsurfaceModel(ShadingModelID);
+		const bool bBackfaceCull = (bCullBackfacingPixels != 0) && !bIsHairInput && !ShadingInfo.bIsSubsurface;
 
-		const half3 WorldNormal = (bIsHairInput || ShadingModelID == SHADINGMODELID_HAIR) ? L : GBufferNormal;
+		const half3 WorldNormal = (bIsHairInput || ShadingInfo.bIsHair) ? L : ShadingInfo.WorldNormal;
 		TranslatedWorldPosition += WorldNormal * NormalBiasLength;
 
 		if (SMRTRayCount > 0)
@@ -471,7 +501,7 @@ void VirtualShadowMapProjection(
 	const float ContactShadowLengthWorld = ContactShadowLength * View.ClipToView[1][1] * SceneDepth;
 	const half Noise = InterleavedGradientNoise( SvPosition.xy, View.StateFrameIndexMod8 );
 
-	const half3 GBufferNormal = GetGBufferDataUint( PixelPos, true ).WorldNormal;
+	const FProjectionShadingInfo ShadingInfo = GetProjectionShadingInfo(PixelPos);
 	const half SubsurfaceOpacity = bIsHairInput ? 1.0f : GetSubsurfaceOpacityFromGbuffer( PixelPos );
 
 	FVirtualShadowMapSampleResult VisualizeResult = InitVirtualShadowMapSampleResult();
@@ -499,11 +529,11 @@ void VirtualShadowMapProjection(
 			FVirtualShadowMapSampleResult Result = ProjectLight(
 				VirtualShadowMapId,
 				Light,
+				ShadingInfo,
 				PixelPos,
 				SceneDepth,
 				ContactShadowLengthWorld,
 				TranslatedWorldPosition,
-				GBufferNormal,
 				Noise);
 
 			// NOTE: Subsurface opacity is handled in the clustered shading pass in the one pass projection path
@@ -527,11 +557,11 @@ void VirtualShadowMapProjection(
 		FVirtualShadowMapSampleResult Result = ProjectLight(
 			VirtualShadowMapId,
 			Light,
+			ShadingInfo,
 			PixelPos,
 			SceneDepth,
 			ContactShadowLengthWorld,
 			TranslatedWorldPosition,
-			GBufferNormal,
 			Noise);
 		
 		half SSSTransmission = Result.ShadowFactor;
