Rewrite ApplyBentNormal, GetBentNormal, and GetDiffuseIndirectSampleOcclusion to clarify input/output change.

diff --git a/Engine/Shaders/Private/BasePassPixelShader.usf b/Engine/Shaders/Private/BasePassPixelShader.usf
--- a/Engine/Shaders/Private/BasePassPixelShader.usf
+++ b/Engine/Shaders/Private/BasePassPixelShader.usf
@@ -161,6 +161,13 @@ float NormalCurvatureToRoughness(float3 WorldNormal)
 #include "ForwardLightingCommon.ush"
 #endif
 
+struct FShadingOcclusion
+{
+	float  DiffOcclusion;
+	float  SpecOcclusion;
+	float3 BentNormal;
+};
+
 #if STRATA_TRANSLUCENT_FORWARD || STRATA_FORWARD_SHADING
 #include "/Engine/Private/Strata/StrataEvaluation.ush"
 #include "/Engine/Private/Strata/StrataForwardLighting.ush"
@@ -784,57 +791,70 @@ float3 GetBentNormalZero(in FMaterialPixelParameters MaterialParameters)
 #endif
 }
 
-void ApplyBentNormal(
-	in float3 CameraVector, in float3 WorldNormal, in float3x3 TangentToWorld, in float3 BentNormal0, 
-	in float Roughness, inout float3 BentNormal, inout float DiffOcclusion, inout float SpecOcclusion)
+FShadingOcclusion ApplyBentNormal(
+	in float3 CameraVector, 
+	in float3 WorldNormal, 
+	in float3x3 TangentToWorld, 
+	in float3 BentNormal0, 
+	in float Roughness,
+	in float MaterialAO)
 {
+	FShadingOcclusion Out;
+	Out.DiffOcclusion = MaterialAO;
+	Out.SpecOcclusion = MaterialAO;
+	Out.BentNormal    = WorldNormal;
+
 #if NUM_MATERIAL_OUTPUTS_GETBENTNORMAL > 0
 	#if MATERIAL_TANGENTSPACENORMAL
-		BentNormal = normalize( TransformTangentVectorToWorld(TangentToWorld, BentNormal0) );
+	Out.BentNormal = normalize( TransformTangentVectorToWorld(TangentToWorld, BentNormal0) );
 	#else
-		BentNormal = BentNormal0;
+	Out.BentNormal = BentNormal0;
 	#endif
 
 	FSphericalGaussian HemisphereSG = Hemisphere_ToSphericalGaussian(WorldNormal);
 	FSphericalGaussian NormalSG = ClampedCosine_ToSphericalGaussian(WorldNormal);
-	FSphericalGaussian VisibleSG = BentNormalAO_ToSphericalGaussian( BentNormal, DiffOcclusion );
+	FSphericalGaussian VisibleSG = BentNormalAO_ToSphericalGaussian(Out.BentNormal, Out.DiffOcclusion );
 	FSphericalGaussian DiffuseSG = Mul( NormalSG, VisibleSG );
 	
-	float VisibleCosAngle = sqrt( 1 - DiffOcclusion );
+	float VisibleCosAngle = sqrt( 1 - Out.DiffOcclusion );
 
 #if 1	// Mix full resolution normal with low res bent normal
-	BentNormal = DiffuseSG.Axis;
+	Out.BentNormal = DiffuseSG.Axis;
 	//DiffOcclusion = saturate( Integral( DiffuseSG ) / Dot( NormalSG, HemisphereSG ) );
-	DiffOcclusion = saturate( Integral( DiffuseSG ) * 0.42276995 );
+	Out.DiffOcclusion = saturate( Integral( DiffuseSG ) * 0.42276995 );
 #endif
 
 	float3 N = WorldNormal;
 	float3 V = CameraVector;
 
-	SpecOcclusion  = DotSpecularSG( Roughness, N, V, VisibleSG );
-	SpecOcclusion /= DotSpecularSG( Roughness, N, V, HemisphereSG );
+	Out.SpecOcclusion  = DotSpecularSG( Roughness, N, V, VisibleSG );
+	Out.SpecOcclusion /= DotSpecularSG( Roughness, N, V, HemisphereSG );
+
+	Out.SpecOcclusion = saturate(Out.SpecOcclusion );
+#endif
+	return Out;
+}
 
-	SpecOcclusion = saturate( SpecOcclusion );
+float3 GetBentNormal(in FMaterialPixelParameters MaterialParameters)
+{
+#if NUM_MATERIAL_OUTPUTS_GETBENTNORMAL > 0
+	#if MATERIAL_TANGENTSPACENORMAL
+	return normalize(TransformTangentVectorToWorld(MaterialParameters.TangentToWorld, GetBentNormalZero(MaterialParameters)));
+	#else
+	return GetBentNormalZero(MaterialParameters);
+	#endif
+#else
+	return MaterialParameters.WorldNormal;
 #endif
 }
 
-uint GetDiffuseIndirectSampleOcclusion(FGBufferData GBuffer, FMaterialPixelParameters MaterialParameters, float2 SvPosition, float MaterialAO)
+uint GetDiffuseIndirectSampleOcclusion(FGBufferData GBuffer, float3 V, float3 WorldNormal, float3 BentNormal, float2 SvPosition, float MaterialAO)
 {
 	uint DiffuseIndirectSampleOcclusion = 0;
 #if GBUFFER_HAS_DIFFUSE_SAMPLE_OCCLUSION && !MATERIAL_SHADINGMODEL_UNLIT
 	if (MaterialAO < 1.0f)
 	{
-		#if NUM_MATERIAL_OUTPUTS_GETBENTNORMAL > 0
-			#if MATERIAL_TANGENTSPACENORMAL
-				float3 BentNormal = normalize(TransformTangentVectorToWorld(MaterialParameters.TangentToWorld, GetBentNormalZero(MaterialParameters)));
-			#else
-				float3 BentNormal = GetBentNormalZero(MaterialParameters);
-			#endif
-		#else
-			float3 BentNormal = MaterialParameters.WorldNormal;
-		#endif
-
-		FSphericalGaussian HemisphereSG = Hemisphere_ToSphericalGaussian(MaterialParameters.WorldNormal);
+		FSphericalGaussian HemisphereSG = Hemisphere_ToSphericalGaussian(WorldNormal);
 		FSphericalGaussian VisibleSG = BentNormalAO_ToSphericalGaussian(BentNormal, MaterialAO);
 		float VisibilityThreshold = InterleavedGradientNoise(SvPosition, View.StateFrameIndexMod8);
 
@@ -846,7 +866,7 @@ uint GetDiffuseIndirectSampleOcclusion(FGBufferData GBuffer, FMaterialPixelParam
 
 			float4 E = ComputeIndirectLightingSampleE(SvPosition, TracingRayIndex, INDIRECT_SAMPLE_COUNT);
 
-			FBxDFSample BxDFSample = SampleBxDF(TermMask, GBuffer, MaterialParameters.CameraVector, E);
+			FBxDFSample BxDFSample = SampleBxDF(TermMask, GBuffer, V, E);
 
 			// Integrate what is visible over the maximum visibility for the normal.
 			float LVisibility = saturate(Evaluate(VisibleSG, BxDFSample.L) / Evaluate(HemisphereSG, BxDFSample.L));
@@ -1251,22 +1271,20 @@ void FPixelShaderInOut_MainPS(
 		// When rendering reflection captures, GBuffer.Roughness is already forced to 1 using RoughnessOverrideParameter in GetMaterialRoughness.
 	}
 
-	float3 BentNormal = MaterialParameters.WorldNormal;
+	float3 InputBentNormal = MaterialParameters.WorldNormal;
 			
 	// Clear Coat Bottom Normal
 	BRANCH if( GBuffer.ShadingModelID == SHADINGMODELID_CLEAR_COAT && CLEAR_COAT_BOTTOM_NORMAL)
 	{
 		const float2 oct1 = ((float2(GBuffer.CustomData.a, GBuffer.CustomData.z) * 4) - (512.0/255.0)) + UnitVectorToOctahedron(GBuffer.WorldNormal);
-		BentNormal = OctahedronToUnitVector(oct1);			
+		InputBentNormal = OctahedronToUnitVector(oct1);
 	}
 	
-	float DiffOcclusion = MaterialAO;
-	float SpecOcclusion = MaterialAO;
-	ApplyBentNormal(MaterialParameters.CameraVector, MaterialParameters.WorldNormal, MaterialParameters.TangentToWorld, GetBentNormalZero(MaterialParameters), GBuffer.Roughness, BentNormal, DiffOcclusion, SpecOcclusion);
-	// FIXME: ALLOW_STATIC_LIGHTING == 0 expects this to be AO
-	GBuffer.GBufferAO = AOMultiBounce( Luminance( GBuffer.SpecularColor ), SpecOcclusion ).g;
+	const FShadingOcclusion ShadingOcclusion = ApplyBentNormal(MaterialParameters.CameraVector, InputBentNormal, MaterialParameters.TangentToWorld, GetBentNormalZero(MaterialParameters), GBuffer.Roughness, MaterialAO);
 
-	GBuffer.DiffuseIndirectSampleOcclusion = GetDiffuseIndirectSampleOcclusion(GBuffer, MaterialParameters, In.SvPosition.xy, MaterialAO);
+	// FIXME: ALLOW_STATIC_LIGHTING == 0 expects this to be AO
+	GBuffer.GBufferAO = AOMultiBounce( Luminance( GBuffer.SpecularColor ), ShadingOcclusion.SpecOcclusion ).g;
+	GBuffer.DiffuseIndirectSampleOcclusion = GetDiffuseIndirectSampleOcclusion(GBuffer, MaterialParameters.CameraVector, MaterialParameters.WorldNormal, GetBentNormal(MaterialParameters), In.SvPosition.xy, MaterialAO);
 
 #endif // !STRATA_ENABLED
 
@@ -1284,7 +1302,7 @@ void FPixelShaderInOut_MainPS(
 
 	#if !MATERIAL_SHADINGMODEL_UNLIT
 
-		float3 DiffuseDir = BentNormal;
+		float3 DiffuseDir = ShadingOcclusion.BentNormal;
 		float3 DiffuseColorForIndirect = GBuffer.DiffuseColor;
 
 		#if MATERIAL_SHADINGMODEL_SUBSURFACE || MATERIAL_SHADINGMODEL_PREINTEGRATED_SKIN
@@ -1337,7 +1355,7 @@ void FPixelShaderInOut_MainPS(
 			IndirectIrradiance *= IndirectOcclusion;
 		#endif
 
-		DiffuseColor += (DiffuseIndirectLighting * DiffuseColorForIndirect + SubsurfaceIndirectLighting * SubsurfaceColor) * AOMultiBounce( GBuffer.BaseColor, DiffOcclusion );
+		DiffuseColor += (DiffuseIndirectLighting * DiffuseColorForIndirect + SubsurfaceIndirectLighting * SubsurfaceColor) * AOMultiBounce( GBuffer.BaseColor, ShadingOcclusion.DiffOcclusion );
 
 #if MATERIAL_SHADINGMODEL_SINGLELAYERWATER
 		// Fade out diffuse as this will be handled by the single scattering lighting in water material.
@@ -1382,7 +1400,7 @@ void FPixelShaderInOut_MainPS(
 
 					half3 ReflectionColor = GetImageBasedReflectionLighting(MaterialParameters, GBuffer.Roughness, GBuffer.SpecularColor, IndirectIrradiance, GridIndex, SingleCaptureIndex, EyeIndex) 
 						* IndirectOcclusion
-						* AOMultiBounce(GBuffer.SpecularColor, SpecOcclusion);
+						* AOMultiBounce(GBuffer.SpecularColor, ShadingOcclusion.SpecOcclusion);
 
 					#if MATERIAL_SHADINGMODEL_THIN_TRANSLUCENT
 						ColorSeparateSpecular += ReflectionColor;
diff --git a/Engine/Shaders/Private/Strata/StrataForwardLighting.ush b/Engine/Shaders/Private/Strata/StrataForwardLighting.ush
--- a/Engine/Shaders/Private/Strata/StrataForwardLighting.ush
+++ b/Engine/Shaders/Private/Strata/StrataForwardLighting.ush
@@ -9,9 +9,14 @@
 void GetVolumeLightingNonDirectional(float4 AmbientLightingVector, float3 DiffuseColor, inout float3 InterpolatedLighting, out float4 VolumeLighting);
 void GetVolumeLightingDirectional(float4 AmbientLightingVector, float3 DirectionalLightingVector, float3 WorldNormal, float3 DiffuseColor, inout float3 InterpolatedLighting, out float4 VolumeLighting);
 float3 GetBentNormalZero(in FMaterialPixelParameters MaterialParameters);
-void ApplyBentNormal(
-	in float3 CameraVector, in float3 WorldNormal, in float3x3 TangentToWorld, in float3 BentNormal0,
-	in float Roughness, inout float3 BentNormal, inout float DiffOcclusion, inout float SpecOcclusion);
+
+FShadingOcclusion ApplyBentNormal(
+	in float3 CameraVector,
+	in float3 WorldNormal,
+	in float3x3 TangentToWorld,
+	in float3 BentNormal0,
+	in float Roughness,
+	in float MaterialAO);
 
 void GetPrecomputedIndirectLightingAndSkyLight(
 	FMaterialPixelParameters MaterialParameters,
@@ -312,19 +317,14 @@ float3 StrataForwardLighting(
 				float IndirectOcclusion = 1.0f;
 				float IndirectIrradiance = 0.0f; 
 
-				const float MaterialAO = StrataPixelHeader.MaterialAO;
-				float DiffOcclusion = MaterialAO;
-				float SpecOcclusion = MaterialAO;
+				// Compute the bent normal from the BSDF normal
+				const FShadingOcclusion ShadingOcclusion = ApplyBentNormal(MaterialParameters.CameraVector, StrataBSDFContext.N, MaterialParameters.TangentToWorld, GetBentNormalZero(MaterialParameters), Roughness, StrataPixelHeader.MaterialAO);
 
 				// 1. Diffuse Sky evaluation using bent normal
 				{
-					// Compute the bent normal from the BSDF normal
-					float3 BentNormal = StrataBSDFContext.N;
-					ApplyBentNormal(MaterialParameters.CameraVector, MaterialParameters.WorldNormal, MaterialParameters.TangentToWorld, GetBentNormalZero(MaterialParameters), Roughness, BentNormal, DiffOcclusion, SpecOcclusion);
-
 					// Update a context specific for environment lighting
 					FStrataBSDFContext EnvBSDFContext = StrataBSDFContext;
-					EnvBSDFContext.N = BentNormal;
+					EnvBSDFContext.N = ShadingOcclusion.BentNormal;
 					StrataUpdateBSDFContext(EnvBSDFContext, EnvBSDFContext.L);
 
 					// And evaluate diffuse parameters
@@ -362,7 +362,7 @@ float3 StrataForwardLighting(
 					#endif
 
 						Color += (DiffuseIndirectLighting * DiffuseColorForIndirect + SubsurfaceIndirectLighting * SubsurfaceColorForIndirect) * 
-							CurrentBSDF.LuminanceWeight * AOMultiBounce(StrataEnvLight.DiffuseColor, DiffOcclusion);
+							CurrentBSDF.LuminanceWeight * AOMultiBounce(StrataEnvLight.DiffuseColor, ShadingOcclusion.DiffOcclusion);
 					}
 				}
 
@@ -377,7 +377,7 @@ float3 StrataForwardLighting(
 					{
 						int SingleCaptureIndex = GetPrimitiveData(MaterialParameters).SingleCaptureIndex;
 						Color += GetImageBasedReflectionSpecular(MaterialParameters, StrataEnvLight.SpecularDirection, Roughness, IndirectIrradiance, GridIndex, SingleCaptureIndex, EyeIndex) *
-							 StrataEnvLight.SpecularWeight * IndirectOcclusion * CurrentBSDF.LuminanceWeight * AOMultiBounce(StrataEnvLight.SpecularColor, SpecOcclusion);
+							 StrataEnvLight.SpecularWeight * IndirectOcclusion * CurrentBSDF.LuminanceWeight * AOMultiBounce(StrataEnvLight.SpecularColor, ShadingOcclusion.SpecOcclusion);
 					}
 					// We skip hazy lob evaluation in forward.
 				}
