Strata forward lighting optimisation: skip BSDF lighting evaluation if LuminanceWeight is 0.
Using the "continue" keyword in this case was causing lots of flickering. But an encapsulating if section works fine seems.

diff --git a/Engine/Shaders/Private/Strata/StrataForwardLighting.ush b/Engine/Shaders/Private/Strata/StrataForwardLighting.ush
--- a/Engine/Shaders/Private/Strata/StrataForwardLighting.ush
+++ b/Engine/Shaders/Private/Strata/StrataForwardLighting.ush
@@ -289,229 +289,311 @@ float3 StrataForwardLighting(
 		STRATA_UNROLL_N(STRATA_CLAMPED_BSDF_COUNT)
 		for (int BSDFIdx = 0; BSDFIdx < StrataTree.BSDFCount; ++BSDFIdx)
 		{
-			#define CurrentBSDF StrataTree.BSDFs[BSDFIdx]
+		#define CurrentBSDF StrataTree.BSDFs[BSDFIdx]
 
-			// This causes some flickering.
-			//if (!StrataIsBSDFVisible(CurrentBSDF))
-			//{
-			//	continue;
-			//}
+			if (StrataIsBSDFVisible(CurrentBSDF))
+			{
 
-			// Compute current BSDF transmittance
-			FStrataBSDFContext StrataBSDFContext = StrataCreateBSDFContext(StrataPixelHeader, CurrentBSDF, NullStrataAddressing, V, L);
+				// Compute current BSDF transmittance
+				FStrataBSDFContext StrataBSDFContext = StrataCreateBSDFContext(StrataPixelHeader, CurrentBSDF, NullStrataAddressing, V, L);
 			
-			#if FORCE_FULLY_ROUGH
-			const bool bForceFullyRough = true;
-			#else
-			const bool bForceFullyRough = View.RenderingReflectionCaptureMask > 0;
-			#endif
+				#if FORCE_FULLY_ROUGH
+				const bool bForceFullyRough = true;
+				#else
+				const bool bForceFullyRough = View.RenderingReflectionCaptureMask > 0;
+				#endif
 			
-			////
-			//// Evaluate environment lighting
-			////
+				////
+				//// Evaluate environment lighting
+				////
 
-			if (UseBasePassSkylight > 0)
-			{
-				const bool bEnableSpecular = ReflectionStruct.SkyLightParameters.y > 0.0f;
-				float Roughness = StrataGetBSDFRoughness(CurrentBSDF);
-				float IndirectOcclusion = 1.0f;
-				float IndirectIrradiance = 0.0f; 
+				if (UseBasePassSkylight > 0)
+				{
+					const bool bEnableSpecular = ReflectionStruct.SkyLightParameters.y > 0.0f;
+					float Roughness = StrataGetBSDFRoughness(CurrentBSDF);
+					float IndirectOcclusion = 1.0f;
+					float IndirectIrradiance = 0.0f; 
 
-				// Compute the bent normal from the BSDF normal
-				const FShadingOcclusion ShadingOcclusion = ApplyBentNormal(MaterialParameters.CameraVector, StrataBSDFContext.N, MaterialParameters.TangentToWorld, GetBentNormalZero(MaterialParameters), Roughness, StrataPixelHeader.MaterialAO);
+					// Compute the bent normal from the BSDF normal
+					const FShadingOcclusion ShadingOcclusion = ApplyBentNormal(MaterialParameters.CameraVector, StrataBSDFContext.N, MaterialParameters.TangentToWorld, GetBentNormalZero(MaterialParameters), Roughness, StrataPixelHeader.MaterialAO);
 
-				// 1. Diffuse Sky evaluation using bent normal
-				{
-					// Update a context specific for environment lighting
-					FStrataBSDFContext EnvBSDFContext = StrataBSDFContext;
-					EnvBSDFContext.N = ShadingOcclusion.BentNormal;
-					StrataUpdateBSDFContext(EnvBSDFContext, EnvBSDFContext.L);
-
-					// And evaluate diffuse parameters
-					FStrataEnvLightResult StrataEnvLight = StrataEvaluateForEnvLight(EnvBSDFContext, false/*bEnableSpecular*/, bForceFullyRough);
-					float3 DiffuseNormal = StrataEnvLight.DiffuseNormal;
-					float3 DiffuseColorForIndirect = StrataEnvLight.DiffuseWeight;
-					float3 SubsurfaceColorForIndirect = StrataEnvLight.DiffuseBackFaceWeight;
-
-					// Evaluate diffuse lighting
-					if (any((DiffuseColorForIndirect + SubsurfaceColorForIndirect) > 0.0f))
+					// 1. Diffuse Sky evaluation using bent normal
 					{
-						const bool bEvaluateBackface = any(SubsurfaceColorForIndirect > 0.0);
-
-						float3 DiffuseIndirectLighting = 0.0f;
-						float3 SubsurfaceIndirectLighting = 0.0f;
-						GetPrecomputedIndirectLightingAndSkyLight(
-							MaterialParameters,
-							Interpolants,
-							BasePassInterpolants,
-							LightmapVTPageTableResult,
-							bEvaluateBackface,
-							DiffuseNormal,
-							VolumetricLightmapBrickTextureUVs,
-							DiffuseIndirectLighting,
-							SubsurfaceIndirectLighting,
-							IndirectIrradiance);
-
-					#if FORWARD_SHADING && (MATERIALBLENDING_SOLID || MATERIALBLENDING_MASKED)
-						float2 NearestResolvedDepthScreenUV = CalculateNearestResolvedDepthScreenUV(ScreenUV, MaterialParameters.ScreenPosition.w);
-
-						IndirectOcclusion = GetIndirectOcclusion(NearestResolvedDepthScreenUV, HasDynamicIndirectShadowCasterRepresentation(StrataPixelHeader));
-						DiffuseIndirectLighting *= IndirectOcclusion;
-						SubsurfaceIndirectLighting *= IndirectOcclusion;
-						IndirectIrradiance *= IndirectOcclusion;
-					#endif
-
-						Color += (DiffuseIndirectLighting * DiffuseColorForIndirect + SubsurfaceIndirectLighting * SubsurfaceColorForIndirect) * 
-							CurrentBSDF.LuminanceWeight * AOMultiBounce(StrataEnvLight.DiffuseColor, ShadingOcclusion.DiffOcclusion);
+						// Update a context specific for environment lighting
+						FStrataBSDFContext EnvBSDFContext = StrataBSDFContext;
+						EnvBSDFContext.N = ShadingOcclusion.BentNormal;
+						StrataUpdateBSDFContext(EnvBSDFContext, EnvBSDFContext.L);
+
+						// And evaluate diffuse parameters
+						FStrataEnvLightResult StrataEnvLight = StrataEvaluateForEnvLight(EnvBSDFContext, false/*bEnableSpecular*/, bForceFullyRough);
+						float3 DiffuseNormal = StrataEnvLight.DiffuseNormal;
+						float3 DiffuseColorForIndirect = StrataEnvLight.DiffuseWeight;
+						float3 SubsurfaceColorForIndirect = StrataEnvLight.DiffuseBackFaceWeight;
+
+						// Evaluate diffuse lighting
+						if (any((DiffuseColorForIndirect + SubsurfaceColorForIndirect) > 0.0f))
+						{
+							const bool bEvaluateBackface = any(SubsurfaceColorForIndirect > 0.0);
+
+							float3 DiffuseIndirectLighting = 0.0f;
+							float3 SubsurfaceIndirectLighting = 0.0f;
+							GetPrecomputedIndirectLightingAndSkyLight(
+								MaterialParameters,
+								Interpolants,
+								BasePassInterpolants,
+								LightmapVTPageTableResult,
+								bEvaluateBackface,
+								DiffuseNormal,
+								VolumetricLightmapBrickTextureUVs,
+								DiffuseIndirectLighting,
+								SubsurfaceIndirectLighting,
+								IndirectIrradiance);
+
+						#if FORWARD_SHADING && (MATERIALBLENDING_SOLID || MATERIALBLENDING_MASKED)
+							float2 NearestResolvedDepthScreenUV = CalculateNearestResolvedDepthScreenUV(ScreenUV, MaterialParameters.ScreenPosition.w);
+
+							IndirectOcclusion = GetIndirectOcclusion(NearestResolvedDepthScreenUV, HasDynamicIndirectShadowCasterRepresentation(StrataPixelHeader));
+							DiffuseIndirectLighting *= IndirectOcclusion;
+							SubsurfaceIndirectLighting *= IndirectOcclusion;
+							IndirectIrradiance *= IndirectOcclusion;
+						#endif
+
+							Color += (DiffuseIndirectLighting * DiffuseColorForIndirect + SubsurfaceIndirectLighting * SubsurfaceColorForIndirect) * 
+								CurrentBSDF.LuminanceWeight * AOMultiBounce(StrataEnvLight.DiffuseColor, ShadingOcclusion.DiffOcclusion);
+						}
 					}
-				}
 
-				// 2. Specular evaluation (using regular normal)
-			#if FORWARD_PER_PIXEL_SHADING || TRANSLUCENCY_LIGHTING_SURFACE_LIGHTINGVOLUME
-				{
-					// Evaluate specular parameters without bent normal
-					FStrataEnvLightResult StrataEnvLight = StrataEvaluateForEnvLight(StrataBSDFContext, bEnableSpecular, bForceFullyRough);
-
-					// Evaluate specular lighting
-					if (any(StrataEnvLight.SpecularWeight > 0.0f))
+					// 2. Specular evaluation (using regular normal)
+				#if FORWARD_PER_PIXEL_SHADING || TRANSLUCENCY_LIGHTING_SURFACE_LIGHTINGVOLUME
 					{
-						int SingleCaptureIndex = GetPrimitiveData(MaterialParameters).SingleCaptureIndex;
-						Color += GetImageBasedReflectionSpecular(MaterialParameters, StrataEnvLight.SpecularDirection, Roughness, IndirectIrradiance, GridIndex, SingleCaptureIndex, EyeIndex) *
-							 StrataEnvLight.SpecularWeight * IndirectOcclusion * CurrentBSDF.LuminanceWeight * AOMultiBounce(StrataEnvLight.SpecularColor, ShadingOcclusion.SpecOcclusion);
+						// Evaluate specular parameters without bent normal
+						FStrataEnvLightResult StrataEnvLight = StrataEvaluateForEnvLight(StrataBSDFContext, bEnableSpecular, bForceFullyRough);
+
+						// Evaluate specular lighting
+						if (any(StrataEnvLight.SpecularWeight > 0.0f))
+						{
+							int SingleCaptureIndex = GetPrimitiveData(MaterialParameters).SingleCaptureIndex;
+							Color += GetImageBasedReflectionSpecular(MaterialParameters, StrataEnvLight.SpecularDirection, Roughness, IndirectIrradiance, GridIndex, SingleCaptureIndex, EyeIndex) *
+								 StrataEnvLight.SpecularWeight * IndirectOcclusion * CurrentBSDF.LuminanceWeight * AOMultiBounce(StrataEnvLight.SpecularColor, ShadingOcclusion.SpecOcclusion);
+						}
+						// We skip hazy lob evaluation in forward.
 					}
-					// We skip hazy lob evaluation in forward.
+				#endif
 				}
-			#endif
-			}
 
-			////
-			//// Evaluate translucent lighting volume / vertex lighting.
-			////
+				////
+				//// Evaluate translucent lighting volume / vertex lighting.
+				////
 
-			// In deferred shading, the translucenyt volume is used by 
-			//	- TRANSLUCENCY_LIGHTING_VOLUMETRIC_PERVERTEX_NONDIRECTIONAL
-			//	- TRANSLUCENCY_LIGHTING_VOLUMETRIC_PERVERTEX_DIRECTIONAL
-			//	- TRANSLUCENCY_LIGHTING_VOLUMETRIC_NONDIRECTIONAL
-			//	- TRANSLUCENCY_LIGHTING_VOLUMETRIC_DIRECTIONAL
-			//	- TRANSLUCENCY_LIGHTING_SURFACE_LIGHTINGVOLUME
-			// In forward, things change (for some reasons...)
-			//  - TRANSLUCENCY_PERVERTEX_FORWARD_SHADING make sure TRANSLUCENCY_LIGHTING_VOLUMETRIC_PERVERTEX_NONDIRECTIONAL and TRANSLUCENCY_LIGHTING_VOLUMETRIC_PERVERTEX_DIRECTIONAL works the same way.
-			//  - TRANSLUCENCY_LIGHTING_VOLUMETRIC_NONDIRECTIONAL will follow the StrataForwardLightingCommon path and directionality is removed using NON_DIRECTIONAL_DIRECT_LIGHTING
-			//  - TRANSLUCENCY_LIGHTING_VOLUMETRIC_DIRECTIONAL will be same as TRANSLUCENCY_LIGHTING_SURFACE_FORWARDSHADING but without the environment specular contribution
-			//	- TRANSLUCENCY_LIGHTING_SURFACE_LIGHTINGVOLUME becomes the same as TRANSLUCENCY_LIGHTING_SURFACE_FORWARDSHADING
-			// STRATA_TODO simplify all that and have forward map to the same technique as in deferred.
+				// In deferred shading, the translucenyt volume is used by 
+				//	- TRANSLUCENCY_LIGHTING_VOLUMETRIC_PERVERTEX_NONDIRECTIONAL
+				//	- TRANSLUCENCY_LIGHTING_VOLUMETRIC_PERVERTEX_DIRECTIONAL
+				//	- TRANSLUCENCY_LIGHTING_VOLUMETRIC_NONDIRECTIONAL
+				//	- TRANSLUCENCY_LIGHTING_VOLUMETRIC_DIRECTIONAL
+				//	- TRANSLUCENCY_LIGHTING_SURFACE_LIGHTINGVOLUME
+				// In forward, things change (for some reasons...)
+				//  - TRANSLUCENCY_PERVERTEX_FORWARD_SHADING make sure TRANSLUCENCY_LIGHTING_VOLUMETRIC_PERVERTEX_NONDIRECTIONAL and TRANSLUCENCY_LIGHTING_VOLUMETRIC_PERVERTEX_DIRECTIONAL works the same way.
+				//  - TRANSLUCENCY_LIGHTING_VOLUMETRIC_NONDIRECTIONAL will follow the StrataForwardLightingCommon path and directionality is removed using NON_DIRECTIONAL_DIRECT_LIGHTING
+				//  - TRANSLUCENCY_LIGHTING_VOLUMETRIC_DIRECTIONAL will be same as TRANSLUCENCY_LIGHTING_SURFACE_FORWARDSHADING but without the environment specular contribution
+				//	- TRANSLUCENCY_LIGHTING_SURFACE_LIGHTINGVOLUME becomes the same as TRANSLUCENCY_LIGHTING_SURFACE_FORWARDSHADING
+				// STRATA_TODO simplify all that and have forward map to the same technique as in deferred.
 
-		#if TRANSLUCENCY_PERVERTEX_FORWARD_SHADING || TRANSLUCENCY_LIGHTING_VOLUMETRIC_PERVERTEX_NONDIRECTIONAL || TRANSLUCENCY_LIGHTING_VOLUMETRIC_PERVERTEX_DIRECTIONAL || (!FORWARD_SHADING && (TRANSLUCENCY_LIGHTING_VOLUMETRIC_NONDIRECTIONAL || TRANSLUCENCY_LIGHTING_VOLUMETRIC_DIRECTIONAL || TRANSLUCENCY_LIGHTING_SURFACE_LIGHTINGVOLUME))
-			float4 VolumeLighting;
-			float3 SurfaceLighting = 0;
-			float3 InnerVolumeUVs;
-			float3 OuterVolumeUVs;
-			float FinalLerpFactor;
-			ComputeVolumeUVs(TranslatedWorldPosition, LightingPositionOffset, InnerVolumeUVs, OuterVolumeUVs, FinalLerpFactor);
+			#if TRANSLUCENCY_PERVERTEX_FORWARD_SHADING || TRANSLUCENCY_LIGHTING_VOLUMETRIC_PERVERTEX_NONDIRECTIONAL || TRANSLUCENCY_LIGHTING_VOLUMETRIC_PERVERTEX_DIRECTIONAL || (!FORWARD_SHADING && (TRANSLUCENCY_LIGHTING_VOLUMETRIC_NONDIRECTIONAL || TRANSLUCENCY_LIGHTING_VOLUMETRIC_DIRECTIONAL || TRANSLUCENCY_LIGHTING_SURFACE_LIGHTINGVOLUME))
+				float4 VolumeLighting;
+				float3 SurfaceLighting = 0;
+				float3 InnerVolumeUVs;
+				float3 OuterVolumeUVs;
+				float FinalLerpFactor;
+				ComputeVolumeUVs(TranslatedWorldPosition, LightingPositionOffset, InnerVolumeUVs, OuterVolumeUVs, FinalLerpFactor);
 
-			FStrataEvaluateResult BSDFEvaluate = StrataEvaluateBSDF(StrataBSDFContext, bRoughDiffuseEnabled);
+				FStrataEvaluateResult BSDFEvaluate = StrataEvaluateBSDF(StrataBSDFContext, bRoughDiffuseEnabled);
 
-		#if TRANSLUCENCY_PERVERTEX_FORWARD_SHADING
+			#if TRANSLUCENCY_PERVERTEX_FORWARD_SHADING
 
-			SurfaceLighting = BasePassInterpolants.VertexDiffuseLighting * BSDFEvaluate.DiffuseColor;
+				SurfaceLighting = BasePassInterpolants.VertexDiffuseLighting * BSDFEvaluate.DiffuseColor;
 
-		#elif TRANSLUCENCY_LIGHTING_VOLUMETRIC_PERVERTEX_NONDIRECTIONAL
+			#elif TRANSLUCENCY_LIGHTING_VOLUMETRIC_PERVERTEX_NONDIRECTIONAL
 
-			// ForwardShadingNote: when forward shading is enabled, this will follow the special TRANSLUCENCY_PERVERTEX_FORWARD_SHADING path above
-			GetVolumeLightingNonDirectional(float4(BasePassInterpolants.AmbientLightingVector, 1), BSDFEvaluate.DiffuseColor, SurfaceLighting, VolumeLighting);
+				// ForwardShadingNote: when forward shading is enabled, this will follow the special TRANSLUCENCY_PERVERTEX_FORWARD_SHADING path above
+				GetVolumeLightingNonDirectional(float4(BasePassInterpolants.AmbientLightingVector, 1), BSDFEvaluate.DiffuseColor, SurfaceLighting, VolumeLighting);
 
-		#elif TRANSLUCENCY_LIGHTING_VOLUMETRIC_PERVERTEX_DIRECTIONAL
+			#elif TRANSLUCENCY_LIGHTING_VOLUMETRIC_PERVERTEX_DIRECTIONAL
 
-			// ForwardShadingNote: when forward shading is enabled, this will follow the special TRANSLUCENCY_PERVERTEX_FORWARD_SHADING path above
-			GetVolumeLightingDirectional(float4(BasePassInterpolants.AmbientLightingVector, 1), BasePassInterpolants.DirectionalLightingVector, StrataBSDFContext.N, BSDFEvaluate.DiffuseColor, SurfaceLighting, VolumeLighting);
+				// ForwardShadingNote: when forward shading is enabled, this will follow the special TRANSLUCENCY_PERVERTEX_FORWARD_SHADING path above
+				GetVolumeLightingDirectional(float4(BasePassInterpolants.AmbientLightingVector, 1), BasePassInterpolants.DirectionalLightingVector, StrataBSDFContext.N, BSDFEvaluate.DiffuseColor, SurfaceLighting, VolumeLighting);
 
-		#elif TRANSLUCENCY_LIGHTING_VOLUMETRIC_NONDIRECTIONAL && !FORWARD_SHADING
+			#elif TRANSLUCENCY_LIGHTING_VOLUMETRIC_NONDIRECTIONAL && !FORWARD_SHADING
 
-			//ForwardShadingNote: when forward shading is enabled, StrataForwardLightingCommon path will be used 
-			float4 AmbientLightingVector = GetAmbientLightingVectorFromTranslucentLightingVolume(InnerVolumeUVs, OuterVolumeUVs, FinalLerpFactor);
-			GetVolumeLightingNonDirectional(AmbientLightingVector, BSDFEvaluate.DiffuseColor, SurfaceLighting, VolumeLighting);
+				//ForwardShadingNote: when forward shading is enabled, StrataForwardLightingCommon path will be used 
+				float4 AmbientLightingVector = GetAmbientLightingVectorFromTranslucentLightingVolume(InnerVolumeUVs, OuterVolumeUVs, FinalLerpFactor);
+				GetVolumeLightingNonDirectional(AmbientLightingVector, BSDFEvaluate.DiffuseColor, SurfaceLighting, VolumeLighting);
 
-		#elif (TRANSLUCENCY_LIGHTING_VOLUMETRIC_DIRECTIONAL || TRANSLUCENCY_LIGHTING_SURFACE_LIGHTINGVOLUME) && !FORWARD_SHADING
+			#elif (TRANSLUCENCY_LIGHTING_VOLUMETRIC_DIRECTIONAL || TRANSLUCENCY_LIGHTING_SURFACE_LIGHTINGVOLUME) && !FORWARD_SHADING
 
-			//ForwardShadingNote: when forward shading is enabled, StrataForwardLightingCommon path will be used 
-			float4 AmbientLightingVector = GetAmbientLightingVectorFromTranslucentLightingVolume(InnerVolumeUVs, OuterVolumeUVs, FinalLerpFactor);
-			float3 DirectionalLightingVector = GetDirectionalLightingVectorFromTranslucentLightingVolume(InnerVolumeUVs, OuterVolumeUVs, FinalLerpFactor);
-			GetVolumeLightingDirectional(AmbientLightingVector, DirectionalLightingVector, StrataBSDFContext.N, BSDFEvaluate.DiffuseColor, SurfaceLighting, VolumeLighting);
+				//ForwardShadingNote: when forward shading is enabled, StrataForwardLightingCommon path will be used 
+				float4 AmbientLightingVector = GetAmbientLightingVectorFromTranslucentLightingVolume(InnerVolumeUVs, OuterVolumeUVs, FinalLerpFactor);
+				float3 DirectionalLightingVector = GetDirectionalLightingVectorFromTranslucentLightingVolume(InnerVolumeUVs, OuterVolumeUVs, FinalLerpFactor);
+				GetVolumeLightingDirectional(AmbientLightingVector, DirectionalLightingVector, StrataBSDFContext.N, BSDFEvaluate.DiffuseColor, SurfaceLighting, VolumeLighting);
 
-		#endif // TRANSLUCENCY_
+			#endif // TRANSLUCENCY_
 
-			// STRATA_TODO: we should simply remove this self shadow code path TRANSLUCENT_SELF_SHADOWING and only use FOM to shadow the ? Translucent shadow should .
-		#if (TRANSLUCENCY_LIGHTING_VOLUMETRIC_DIRECTIONAL || TRANSLUCENCY_LIGHTING_VOLUMETRIC_NONDIRECTIONAL || TRANSLUCENCY_LIGHTING_SURFACE_LIGHTINGVOLUME) && TRANSLUCENT_SELF_SHADOWING
-			// Only apply self shadowing if the shadow hasn't faded out completely
-			if (TranslucentSelfShadow.DirectionalLightColor.a > 0)
-			{
-				// Determine the shadow space position
-				// Apply a stable offset to the world position used for shadowing, which blurs out high frequency details in the shadowmap with many layers
-				float4 HomogeneousShadowPosition = mul(float4(LWCHackToFloat(AbsoluteWorldPosition) + LightingPositionOffset, 1), TranslucentSelfShadow.WorldToShadowMatrix);
-				float2 ShadowUVs = HomogeneousShadowPosition.xy / HomogeneousShadowPosition.w;
-				// Lookup the shadow density at the point being shaded
-				float3 ShadowDensity = CalculateTranslucencyShadowingDensity(ShadowUVs, HomogeneousShadowPosition.z) / GetMaterialTranslucentMultipleScatteringExtinction();
-				// Compute colored transmission based on the density that the light ray passed through
-				float3 SelfShadowing = saturate(exp(-ShadowDensity * GetMaterialTranslucentSelfShadowDensityScale()));
-				// Compute a second shadow gradient to add interesting information in the shadowed area of the first
-				// This is a stop gap for not having self shadowing from other light sources
-				float3 SelfShadowing2 = lerp(float3(1, 1, 1), saturate(exp(-ShadowDensity * GetMaterialTranslucentSelfShadowSecondDensityScale())), GetMaterialTranslucentSelfShadowSecondOpacity());
-				SelfShadowing = SelfShadowing * SelfShadowing2;
-
-				// Force unshadowed if we read outside the valid area of the shadowmap atlas
-				// This can happen if the particle system's bounds don't match its visible area
-				FLATTEN
-				if (any(ShadowUVs < TranslucentSelfShadow.ShadowUVMinMax.xy || ShadowUVs > TranslucentSelfShadow.ShadowUVMinMax.zw))
+				// STRATA_TODO: we should simply remove this self shadow code path TRANSLUCENT_SELF_SHADOWING and only use FOM to shadow the ? Translucent shadow should .
+			#if (TRANSLUCENCY_LIGHTING_VOLUMETRIC_DIRECTIONAL || TRANSLUCENCY_LIGHTING_VOLUMETRIC_NONDIRECTIONAL || TRANSLUCENCY_LIGHTING_SURFACE_LIGHTINGVOLUME) && TRANSLUCENT_SELF_SHADOWING
+				// Only apply self shadowing if the shadow hasn't faded out completely
+				if (TranslucentSelfShadow.DirectionalLightColor.a > 0)
 				{
-					SelfShadowing = 1;
-				}
+					// Determine the shadow space position
+					// Apply a stable offset to the world position used for shadowing, which blurs out high frequency details in the shadowmap with many layers
+					float4 HomogeneousShadowPosition = mul(float4(LWCHackToFloat(AbsoluteWorldPosition) + LightingPositionOffset, 1), TranslucentSelfShadow.WorldToShadowMatrix);
+					float2 ShadowUVs = HomogeneousShadowPosition.xy / HomogeneousShadowPosition.w;
+					// Lookup the shadow density at the point being shaded
+					float3 ShadowDensity = CalculateTranslucencyShadowingDensity(ShadowUVs, HomogeneousShadowPosition.z) / GetMaterialTranslucentMultipleScatteringExtinction();
+					// Compute colored transmission based on the density that the light ray passed through
+					float3 SelfShadowing = saturate(exp(-ShadowDensity * GetMaterialTranslucentSelfShadowDensityScale()));
+					// Compute a second shadow gradient to add interesting information in the shadowed area of the first
+					// This is a stop gap for not having self shadowing from other light sources
+					float3 SelfShadowing2 = lerp(float3(1, 1, 1), saturate(exp(-ShadowDensity * GetMaterialTranslucentSelfShadowSecondDensityScale())), GetMaterialTranslucentSelfShadowSecondOpacity());
+					SelfShadowing = SelfShadowing * SelfShadowing2;
+
+					// Force unshadowed if we read outside the valid area of the shadowmap atlas
+					// This can happen if the particle system's bounds don't match its visible area
+					FLATTEN
+					if (any(ShadowUVs < TranslucentSelfShadow.ShadowUVMinMax.xy || ShadowUVs > TranslucentSelfShadow.ShadowUVMinMax.zw))
+					{
+						SelfShadowing = 1;
+					}
 
-				// The volume lighting already contains the contribution of the directional light, 
-				// So calculate the amount of light to remove from the volume lighting in order to apply per-pixel self shadowing
-				// VolumeLighting.a stores all attenuation and opaque shadow factors
-				float3 SelfShadowingCorrection = TranslucentSelfShadow.DirectionalLightColor.rgb * VolumeLighting.a * (1 - SelfShadowing);
+					// The volume lighting already contains the contribution of the directional light, 
+					// So calculate the amount of light to remove from the volume lighting in order to apply per-pixel self shadowing
+					// VolumeLighting.a stores all attenuation and opaque shadow factors
+					float3 SelfShadowingCorrection = TranslucentSelfShadow.DirectionalLightColor.rgb * VolumeLighting.a * (1 - SelfShadowing);
 
-				// Combine backscattering and directional light self shadowing
-				SurfaceLighting = (BSDFEvaluate.DiffuseColor * max(VolumeLighting.rgb - SelfShadowingCorrection, 0));
-			}
-		#endif
+					// Combine backscattering and directional light self shadowing
+					SurfaceLighting = (BSDFEvaluate.DiffuseColor * max(VolumeLighting.rgb - SelfShadowingCorrection, 0));
+				}
+			#endif
 
-			Color += (SurfaceLighting + BSDFEvaluate.EmissivePathValue) * CurrentBSDF.LuminanceWeight;
+				Color += (SurfaceLighting + BSDFEvaluate.EmissivePathValue) * CurrentBSDF.LuminanceWeight;
 
-		#elif FORWARD_PER_PIXEL_SHADING
+			#elif FORWARD_PER_PIXEL_SHADING
 
-			////
-			//// Evaluate Directional light only
-			////
-			{
-				float3 DirLightL = DirLightData.Direction;	// Already normalized
-				float3 ToDirLight = DirLightL;
-				float DirLightMask = 1;
-				if (DirLightData.bRadialLight)
+				////
+				//// Evaluate Directional light only
+				////
 				{
-					DirLightMask = GetLocalLightAttenuation(AbsoluteWorldPosition, DirLightData, ToDirLight, DirLightL);
+					float3 DirLightL = DirLightData.Direction;	// Already normalized
+					float3 ToDirLight = DirLightL;
+					float DirLightMask = 1;
+					if (DirLightData.bRadialLight)
+					{
+						DirLightMask = GetLocalLightAttenuation(AbsoluteWorldPosition, DirLightData, ToDirLight, DirLightL);
+					}
+
+					bool bEvaluateHasBeenExecuted = false;
+					BRANCH
+					if (DirectionalLightData.HasDirectionalLight)
+					{
+						uint LightChannelMask = DirectionalLightData.DirectionalLightShadowMapChannelMask >> 8;
+
+						float DirectionalLightCloudShadow = 1.0f;
+					#if NEEDS_BASEPASS_CLOUD_SHADOW_INTERPOLATOR
+						DirectionalLightCloudShadow = BasePassInterpolants.VertexCloudShadow;
+					#endif
+
+						FStrataEvaluateResult BSDFEvaluate = (FStrataEvaluateResult)0;
+						Color += StrataForwardLightingCommon(
+							Dither,
+							bForceFullyRough,
+							bRoughDiffuseEnabled,
+							DirLightData,
+							ToDirLight,
+							DirLightMask,
+							DirLightAttenuation,
+							RectTexture,
+							LightChannelMask,
+							PrimitiveLightingChannelMask,
+							AbsoluteWorldPosition,
+							TranslatedWorldPosition,
+							SceneDepth,
+							CurrentBSDF.LuminanceWeight,
+							StrataPixelHeader,
+							StrataBSDFContext,
+							bEvaluateHasBeenExecuted,
+							BSDFEvaluate) * DirectionalLightCloudShadow;
+					}
 				}
 
-				bool bEvaluateHasBeenExecuted = false;
-				BRANCH
-				if (DirectionalLightData.HasDirectionalLight)
+				////
+				//// Evaluate local lights
+				////
+
+			#if !DISABLE_FORWARD_LOCAL_LIGHTS && (FORWARD_SHADING || TRANSLUCENCY_LIGHTING_SURFACE_FORWARDSHADING) // Forbidden for Single Layer Water
+				// Basic implementation for FORWARD_LOCAL_LIGHTS where local lights are reloaded per BSDF layer. It is fast by default when layer==1.
+				const FCulledLightsGridData CulledLightsGrid = GetCulledLightsGrid(GridIndex, EyeIndex);
+
+				// Limit max to ForwardLightData.NumLocalLights.
+				// This prevents GPU hangs when the PS tries to read from uninitialized NumCulledLightsGrid buffer
+				const uint NumLocalLights = min(CulledLightsGrid.NumLocalLights, GetNumLocalLights(EyeIndex));
+
+				LOOP
+				for (uint LocalLightListIndex = 0; LocalLightListIndex < NumLocalLights; LocalLightListIndex++)
 				{
-					uint LightChannelMask = DirectionalLightData.DirectionalLightShadowMapChannelMask >> 8;
+					const FLocalLightData LocalLight = GetLocalLightData(CulledLightsGrid.DataStartIndex + LocalLightListIndex, EyeIndex);
+
+					FDeferredLightData LightData = (FDeferredLightData)0;
+					LightData.TranslatedWorldPosition = LocalLight.LightPositionAndInvRadius.xyz;
+					LightData.InvRadius = LocalLight.LightPositionAndInvRadius.w;
+					LightData.Color = LocalLight.LightColorAndFalloffExponent.xyz;
+					LightData.FalloffExponent = LocalLight.LightColorAndFalloffExponent.w;
+					LightData.Direction = LocalLight.LightDirectionAndShadowMask.xyz;
+					LightData.SpotAngles = LocalLight.SpotAnglesAndSourceRadiusPacked.xy;
+					LightData.SourceRadius = LocalLight.SpotAnglesAndSourceRadiusPacked.z;
+					LightData.SourceLength = f16tof32(asuint(LocalLight.SpotAnglesAndSourceRadiusPacked.w));
+					LightData.Tangent = LocalLight.LightTangentAndSoftSourceRadius.xyz;
+					LightData.SoftSourceRadius = LocalLight.LightTangentAndSoftSourceRadius.w;
+					LightData.bInverseSquared = LightData.FalloffExponent == 0;
+					LightData.bRadialLight = true;
+					LightData.bSpotLight = LightData.SpotAngles.x > -2.0f;
+					LightData.SpecularScale = SpecularScale;
+					LightData.HairTransmittance = InitHairTransmittanceData();
+
+					// LightType=bits[17:16], LightingChannelMask=[15:8], DynShadowMask=[7:4] ShadowMapChannelMask=[3:0]
+					uint LightTypeAndPackedShadowMapChannelMask = asuint(LocalLight.LightDirectionAndShadowMask.w);
+
+					// TODO: currently LightType is ignored and the code always sets 'LightData.bRectLight' to false (initialization)
+					//       but the light grid adds rect lights as well so this probably leads to artifacts.
+					uint LightType = LightTypeAndPackedShadowMapChannelMask >> 16;
+
+					LightData.ShadowedBits = (LightTypeAndPackedShadowMapChannelMask & 0xFF) != 0 ? 1 : 0;
+					// Static shadowing uses ShadowMapChannel, dynamic shadows are packed into light attenuation using PreviewShadowMapChannel
+					LightData.ShadowMapChannelMask = UnpackShadowMapChannelMask(LightTypeAndPackedShadowMapChannelMask);
+					float4 PreviewShadowMapChannelMask = UnpackShadowMapChannelMask(LightTypeAndPackedShadowMapChannelMask >> 4);
+					float DynamicShadowing = dot(PreviewShadowMapChannelMask, DynamicShadowFactors);
+					float4 LightAttenuation = float4(1, 1, DynamicShadowing.x, DynamicShadowing.x);
+
+					float3 L = LightData.Direction;	// Already normalized
+					float3 ToLight = L;
+					float LightMask = 1;
+					if (LightData.bRadialLight)
+					{
+						LightMask = GetLocalLightAttenuation(AbsoluteWorldPosition, LightData, ToLight, L);
+					}
 
-					float DirectionalLightCloudShadow = 1.0f;
-				#if NEEDS_BASEPASS_CLOUD_SHADOW_INTERPOLATOR
-					DirectionalLightCloudShadow = BasePassInterpolants.VertexCloudShadow;
-				#endif
+					// Update the strata BSDF context accordgin to the new L
+					StrataUpdateBSDFContext(StrataBSDFContext, ToLight);
+
+					uint LightChannelMask = LightTypeAndPackedShadowMapChannelMask >> 8;
 
+					bool bEvaluateHasBeenExecuted = false;
 					FStrataEvaluateResult BSDFEvaluate = (FStrataEvaluateResult)0;
 					Color += StrataForwardLightingCommon(
 						Dither,
 						bForceFullyRough,
 						bRoughDiffuseEnabled,
-						DirLightData,
-						ToDirLight,
-						DirLightMask,
-						DirLightAttenuation,
+						LightData,
+						ToLight,
+						LightMask,
+						LightAttenuation,
 						RectTexture,
 						LightChannelMask,
 						PrimitiveLightingChannelMask,
@@ -522,98 +604,15 @@ float3 StrataForwardLighting(
 						StrataPixelHeader,
 						StrataBSDFContext,
 						bEvaluateHasBeenExecuted,
-						BSDFEvaluate) * DirectionalLightCloudShadow;
+						BSDFEvaluate);
 				}
-			}
-
-			////
-			//// Evaluate local lights
-			////
-
-		#if !DISABLE_FORWARD_LOCAL_LIGHTS && (FORWARD_SHADING || TRANSLUCENCY_LIGHTING_SURFACE_FORWARDSHADING) // Forbidden for Single Layer Water
-			// Basic implementation for FORWARD_LOCAL_LIGHTS where local lights are reloaded per BSDF layer. It is fast by default when layer==1.
-			const FCulledLightsGridData CulledLightsGrid = GetCulledLightsGrid(GridIndex, EyeIndex);
-
-			// Limit max to ForwardLightData.NumLocalLights.
-			// This prevents GPU hangs when the PS tries to read from uninitialized NumCulledLightsGrid buffer
-			const uint NumLocalLights = min(CulledLightsGrid.NumLocalLights, GetNumLocalLights(EyeIndex));
 
-			LOOP
-			for (uint LocalLightListIndex = 0; LocalLightListIndex < NumLocalLights; LocalLightListIndex++)
-			{
-				const FLocalLightData LocalLight = GetLocalLightData(CulledLightsGrid.DataStartIndex + LocalLightListIndex, EyeIndex);
-
-				FDeferredLightData LightData = (FDeferredLightData)0;
-				LightData.TranslatedWorldPosition = LocalLight.LightPositionAndInvRadius.xyz;
-				LightData.InvRadius = LocalLight.LightPositionAndInvRadius.w;
-				LightData.Color = LocalLight.LightColorAndFalloffExponent.xyz;
-				LightData.FalloffExponent = LocalLight.LightColorAndFalloffExponent.w;
-				LightData.Direction = LocalLight.LightDirectionAndShadowMask.xyz;
-				LightData.SpotAngles = LocalLight.SpotAnglesAndSourceRadiusPacked.xy;
-				LightData.SourceRadius = LocalLight.SpotAnglesAndSourceRadiusPacked.z;
-				LightData.SourceLength = f16tof32(asuint(LocalLight.SpotAnglesAndSourceRadiusPacked.w));
-				LightData.Tangent = LocalLight.LightTangentAndSoftSourceRadius.xyz;
-				LightData.SoftSourceRadius = LocalLight.LightTangentAndSoftSourceRadius.w;
-				LightData.bInverseSquared = LightData.FalloffExponent == 0;
-				LightData.bRadialLight = true;
-				LightData.bSpotLight = LightData.SpotAngles.x > -2.0f;
-				LightData.SpecularScale = SpecularScale;
-				LightData.HairTransmittance = InitHairTransmittanceData();
-
-				// LightType=bits[17:16], LightingChannelMask=[15:8], DynShadowMask=[7:4] ShadowMapChannelMask=[3:0]
-				uint LightTypeAndPackedShadowMapChannelMask = asuint(LocalLight.LightDirectionAndShadowMask.w);
-
-				// TODO: currently LightType is ignored and the code always sets 'LightData.bRectLight' to false (initialization)
-				//       but the light grid adds rect lights as well so this probably leads to artifacts.
-				uint LightType = LightTypeAndPackedShadowMapChannelMask >> 16;
-
-				LightData.ShadowedBits = (LightTypeAndPackedShadowMapChannelMask & 0xFF) != 0 ? 1 : 0;
-				// Static shadowing uses ShadowMapChannel, dynamic shadows are packed into light attenuation using PreviewShadowMapChannel
-				LightData.ShadowMapChannelMask = UnpackShadowMapChannelMask(LightTypeAndPackedShadowMapChannelMask);
-				float4 PreviewShadowMapChannelMask = UnpackShadowMapChannelMask(LightTypeAndPackedShadowMapChannelMask >> 4);
-				float DynamicShadowing = dot(PreviewShadowMapChannelMask, DynamicShadowFactors);
-				float4 LightAttenuation = float4(1, 1, DynamicShadowing.x, DynamicShadowing.x);
-
-				float3 L = LightData.Direction;	// Already normalized
-				float3 ToLight = L;
-				float LightMask = 1;
-				if (LightData.bRadialLight)
-				{
-					LightMask = GetLocalLightAttenuation(AbsoluteWorldPosition, LightData, ToLight, L);
-				}
-
-				// Update the strata BSDF context accordgin to the new L
-				StrataUpdateBSDFContext(StrataBSDFContext, ToLight);
-
-				uint LightChannelMask = LightTypeAndPackedShadowMapChannelMask >> 8;
-
-				bool bEvaluateHasBeenExecuted = false;
-				FStrataEvaluateResult BSDFEvaluate = (FStrataEvaluateResult)0;
-				Color += StrataForwardLightingCommon(
-					Dither,
-					bForceFullyRough,
-					bRoughDiffuseEnabled,
-					LightData,
-					ToLight,
-					LightMask,
-					LightAttenuation,
-					RectTexture,
-					LightChannelMask,
-					PrimitiveLightingChannelMask,
-					AbsoluteWorldPosition,
-					TranslatedWorldPosition,
-					SceneDepth,
-					CurrentBSDF.LuminanceWeight,
-					StrataPixelHeader,
-					StrataBSDFContext,
-					bEvaluateHasBeenExecuted,
-					BSDFEvaluate);
-			}
+			#endif // !DISABLE_FORWARD_LOCAL_LIGHTS
 
-		#endif // !DISABLE_FORWARD_LOCAL_LIGHTS
+			#endif // Lighting technique
 
-		#endif // Lighting technique
 
+			}// StrataIsBSDFVisible
 		#undef CurrentBSDF
 		}
 	}
