Add Strata support for irradiance / irradiance occlusion data

diff --git a/Engine/Shaders/Private/BasePassPixelShader.usf b/Engine/Shaders/Private/BasePassPixelShader.usf
--- a/Engine/Shaders/Private/BasePassPixelShader.usf
+++ b/Engine/Shaders/Private/BasePassPixelShader.usf
@@ -1066,7 +1066,8 @@ void FPixelShaderInOut_MainPS(
 	StrataPixelHeader.StrataTree = MaterialParameters.StrataTree;
 	StrataPixelHeader.BSDFCount = MaterialParameters.StrataTree.BSDFCount;
 	StrataPixelHeader.SharedLocalBases = MaterialParameters.SharedLocalBases;
-	StrataPixelHeader.MaterialAO = MaterialAO;
+	StrataPixelHeader.IrradianceAO = InitIrradianceAndOcclusion();
+	StrataPixelHeader.IrradianceAO.MaterialAO = MaterialAO;
 	SetCastContactShadow(StrataPixelHeader, GetPrimitiveData(MaterialParameters).Flags & PRIMITIVE_SCENE_DATA_FLAG_HAS_CAST_CONTACT_SHADOW);
 	SetDynamicIndirectShadowCasterRepresentation(StrataPixelHeader, GetPrimitiveData(MaterialParameters).Flags & PRIMITIVE_SCENE_DATA_FLAG_HAS_CAPSULE_REPRESENTATION);
 #endif
@@ -1645,7 +1646,13 @@ void FPixelShaderInOut_MainPS(
 #endif
 
 	// STRATA_TODO GBUFFER_HAS_DIFFUSE_SAMPLE_OCCLUSION and ALLOW_STATIC_LIGHTING cases
-	StrataPixelHeader.MaterialAO = GBuffer.GBufferAO;
+	// const FShadingOcclusion ShadingOcclusion = ApplyBentNormal(MaterialParameters.CameraVector, MaterialParameters.WorldNormal, MaterialParameters.TangentToWorld, GetBentNormalZero(MaterialParameters), GBuffer.Roughness, MaterialAO);
+	{
+		StrataPixelHeader.IrradianceAO									= InitIrradianceAndOcclusion();
+		StrataPixelHeader.IrradianceAO.MaterialAO						= GBuffer.GBufferAO;	// STRATA_TODO: AOMultiBounce(Luminance(GBuffer.SpecularColor), ShadingOcclusion.SpecOcclusion).g; // FIXME: ALLOW_STATIC_LIGHTING == 0 expects this to be AO
+		StrataPixelHeader.IrradianceAO.IndirectIrradiance				= IndirectIrradiance;	// STRATA_TODO
+		StrataPixelHeader.IrradianceAO.DiffuseIndirectSampleOcclusion	= 0xFFu;				// STRATA_TODO: GetDiffuseIndirectSampleOcclusion(GBuffer, MaterialParameters.CameraVector, MaterialParameters.WorldNormal, GetBentNormal(MaterialParameters), In.SvPosition.xy, MaterialAO);
+	}
 
 	// We only rely on this GBuffer structure for the write out to legacy render target such as velocity or precomputed shadow factors
 	const float4 PrecomputedShadowFactors = GBuffer.PrecomputedShadowFactors;
diff --git a/Engine/Shaders/Private/ClusteredDeferredShadingPixelShader.usf b/Engine/Shaders/Private/ClusteredDeferredShadingPixelShader.usf
--- a/Engine/Shaders/Private/ClusteredDeferredShadingPixelShader.usf
+++ b/Engine/Shaders/Private/ClusteredDeferredShadingPixelShader.usf
@@ -207,7 +207,7 @@ float4 GetSimpleLightLighting(
 
 	#if STRATA_ENABLED && USE_HAIR_LIGHTING == 0
 		const FDeferredLightData LightData = ConvertLocalToDeferredLight(LocalLight);
-		const FShadowTerms ShadowTerms = { StrataPixelHeader.MaterialAO, 1.0, 1.0, InitHairTransmittanceData() };
+		const FShadowTerms ShadowTerms = { StrataGetAO(StrataPixelHeader), 1.0, 1.0, InitHairTransmittanceData() };
 
 		const float3 ToLight = LightData.TranslatedWorldPosition - TranslatedWorldPosition;
 		const float3 L = normalize(ToLight);
@@ -352,7 +352,7 @@ float4 GetLightGridLocalLighting(
 		#endif
 
 	#if STRATA_ENABLED && USE_HAIR_LIGHTING == 0
-		FShadowTerms ShadowTerms = { StrataPixelHeader.MaterialAO, 1.0, 1.0, InitHairTransmittanceData() };
+		FShadowTerms ShadowTerms = { StrataGetAO(StrataPixelHeader), 1.0, 1.0, InitHairTransmittanceData() };
 
 		const uint FakeShadingModelID = 0;
 		const float FakeContactShadowOpacity = 1.0f;
diff --git a/Engine/Shaders/Private/DebugProbes.usf b/Engine/Shaders/Private/DebugProbes.usf
--- a/Engine/Shaders/Private/DebugProbes.usf
+++ b/Engine/Shaders/Private/DebugProbes.usf
@@ -261,7 +261,8 @@ void MainPS(
 	{
 		FStrataPixelHeader StrataPixelHeader = InitialiseStrataPixelHeader();
 		StrataPixelHeader.StrataTree = GetInitialisedStrataTree();
-		StrataPixelHeader.MaterialAO = 1.0f - IndirectIrradianceAndAO;
+		StrataPixelHeader.IrradianceAO = InitIrradianceAndOcclusion();
+		StrataPixelHeader.IrradianceAO.MaterialAO = 1.0f - IndirectIrradianceAndAO;
 		SetCastContactShadow(StrataPixelHeader, 0);
 
 		uint SharedLocalBasisIndex = 0;
diff --git a/Engine/Shaders/Private/DecalCommon.ush b/Engine/Shaders/Private/DecalCommon.ush
--- a/Engine/Shaders/Private/DecalCommon.ush
+++ b/Engine/Shaders/Private/DecalCommon.ush
@@ -31,7 +31,7 @@ float DecalCommonOutput(
 #endif
 #elif DECAL_RENDERSTAGE == DECAL_RENDERSTAGE_AO
 	// Replace color with AO for AO decals.
-	Color = lerp(1, InHeader.MaterialAO, DBuffer.Coverage);
+	Color = lerp(1, StrataGetAO(InHeader), DBuffer.Coverage);
 #endif
 
 #if DECAL_RENDERTARGETMODE == DECAL_RENDERTARGETMODE_SCENECOLOR
diff --git a/Engine/Shaders/Private/DeferredDecal.usf b/Engine/Shaders/Private/DeferredDecal.usf
--- a/Engine/Shaders/Private/DeferredDecal.usf
+++ b/Engine/Shaders/Private/DeferredDecal.usf
@@ -145,7 +145,8 @@ void FPixelShaderInOut_MainPS(inout FPixelShaderIn In, inout FPixelShaderOut Out
 	FStrataPixelHeader StrataHeader = InitialiseStrataPixelHeader();
 	StrataHeader.StrataTree = MaterialParameters.StrataTree;
 	StrataHeader.BSDFCount = 1;
-	StrataHeader.MaterialAO = GetMaterialAmbientOcclusion(PixelMaterialInputs);
+	StrataHeader.IrradianceAO = InitIrradianceAndOcclusion();
+	StrataHeader.IrradianceAO.MaterialAO = GetMaterialAmbientOcclusion(PixelMaterialInputs);
 	StrataHeader.SharedLocalBases = MaterialParameters.SharedLocalBases;
 	const float Coverage = DecalCommonOutput(Out, StrataHeader, DecalFading);
 	const float Opacity = Coverage;
diff --git a/Engine/Shaders/Private/DeferredLightPixelShaders.usf b/Engine/Shaders/Private/DeferredLightPixelShaders.usf
--- a/Engine/Shaders/Private/DeferredLightPixelShaders.usf
+++ b/Engine/Shaders/Private/DeferredLightPixelShaders.usf
@@ -255,7 +255,7 @@ void DeferredLightPixelMain(
 		float4 Radiance = 0;
 		if (LightMask > 0)
 		{
-			FShadowTerms ShadowTerms = { StrataPixelHeader.MaterialAO, 1.0, 1.0, InitHairTransmittanceData() };
+			FShadowTerms ShadowTerms = { StrataGetAO(StrataPixelHeader), 1.0, 1.0, InitHairTransmittanceData() };
 			float4 LightAttenuation = GetPerPixelLightAttenuation(InputParams.ScreenUV);
 
 			float Dither = InterleavedGradientNoise(InputParams.PixelPos, View.StateFrameIndexMod8);
diff --git a/Engine/Shaders/Private/MeshDecals.usf b/Engine/Shaders/Private/MeshDecals.usf
--- a/Engine/Shaders/Private/MeshDecals.usf
+++ b/Engine/Shaders/Private/MeshDecals.usf
@@ -140,7 +140,8 @@ void FPixelShaderInOut_MainPS(
 #if STRATA_ENABLED
 	FStrataPixelHeader StrataHeader = InitialiseStrataPixelHeader();
 	StrataHeader.BSDFCount = 1;
-	StrataHeader.MaterialAO = GetMaterialAmbientOcclusion(PixelMaterialInputs);
+	StrataHeader.IrradianceAO = InitIrradianceAndOcclusion();
+	StrataHeader.IrrdianceAO.MaterialAO = GetMaterialAmbientOcclusion(PixelMaterialInputs);
 	StrataHeader.SharedLocalBases = MaterialParameters.SharedLocalBases;
 	const float Coverage = DecalCommonOutput(Out, StrataHeader, PixelMaterialInputs.FrontMaterial, 1.f);
 	const float Opacity = Coverage;
diff --git a/Engine/Shaders/Private/Strata/Strata.ush b/Engine/Shaders/Private/Strata/Strata.ush
--- a/Engine/Shaders/Private/Strata/Strata.ush
+++ b/Engine/Shaders/Private/Strata/Strata.ush
@@ -210,7 +210,7 @@ struct FStrataBSDF
 
 // For HEADER_BIT_BSDF_COUNT and  HEADER_BIT_SHAREDLOCALBASES_X, please refer to values in StrataDefinitions.h
 #define HEADER_BIT_COUNT_STATE						HEADER_MASK_BIT_COUNT
-#define HEADER_BIT_COUNT_AO							8
+#define HEADER_BIT_COUNT_IRRADIANCE_AO				8
 #define HEADER_BIT_COUNT_SHADINGMODELS				4
 #define HEADER_BIT_COUNT_BSDF_COUNT					4
 #define HEADER_BIT_COUNT_SHAREDLOCALBASES_COUNT		3
@@ -218,8 +218,8 @@ struct FStrataBSDF
 #define HEADER_BIT_COUNT___UNUSED___				1
 
 #define HEADER_BIT_OFFSET_STATE						0
-#define HEADER_BIT_OFFSET_AO						(HEADER_BIT_OFFSET_STATE					+ HEADER_BIT_COUNT_STATE)
-#define HEADER_BIT_OFFSET_SHADINGMODELS				(HEADER_BIT_OFFSET_AO						+ HEADER_BIT_COUNT_AO)
+#define HEADER_BIT_OFFSET_IRRADIANCE_AO				(HEADER_BIT_OFFSET_STATE					+ HEADER_BIT_COUNT_STATE)
+#define HEADER_BIT_OFFSET_SHADINGMODELS				(HEADER_BIT_OFFSET_IRRADIANCE_AO			+ HEADER_BIT_COUNT_IRRADIANCE_AO)
 #define HEADER_BIT_OFFSET_BSDF_COUNT				(HEADER_BIT_OFFSET_SHADINGMODELS			+ HEADER_BIT_COUNT_SHADINGMODELS)
 #define HEADER_BIT_OFFSET_SHAREDLOCALBASES_COUNT	(HEADER_BIT_OFFSET_BSDF_COUNT				+ HEADER_BIT_COUNT_BSDF_COUNT)
 #define HEADER_BIT_OFFSET_SHAREDLOCALBASES_TYPE		(HEADER_BIT_OFFSET_SHAREDLOCALBASES_COUNT	+ HEADER_BIT_COUNT_SHAREDLOCALBASES_COUNT)
@@ -229,10 +229,10 @@ struct FStrataBSDF
 // The different layout 'aliases' differently the material header
 // 
 // Layout 0 (Fast)
-#define HEADER_FASTENCODING_BIT_COUNT				(HEADER_BIT_COUNT_STATE + HEADER_BIT_COUNT_AO)
+#define HEADER_FASTENCODING_BIT_COUNT				(HEADER_BIT_COUNT_STATE + HEADER_BIT_COUNT_IRRADIANCE_AO)
 #define HEADER_FASTENCODING_MASK					((1u<<HEADER_FASTENCODING_BIT_COUNT)-1u)
 // Layout 1 (Single BSDF)
-#define HEADER_SINGLEENCODING_BIT_COUNT				(HEADER_BIT_COUNT_STATE + HEADER_BIT_COUNT_AO + HEADER_BIT_COUNT_SHADINGMODELS)
+#define HEADER_SINGLEENCODING_BIT_COUNT				(HEADER_BIT_COUNT_STATE + HEADER_BIT_COUNT_IRRADIANCE_AO + HEADER_BIT_COUNT_SHADINGMODELS)
 #define HEADER_SINGLEENCODING_MASK					((1u<<HEADER_SINGLEENCODING_BIT_COUNT)-1u)
 // Layout 2 (Complex)
 #define HEADER_COMPLEXENCODING_BIT_COUNT			(32u)
@@ -252,8 +252,8 @@ struct FStrataBSDF
 #define HEADER_GETSHAREDLOCALBASESTYPE(X)			READ_BITS( X, HEADER_BIT_COUNT_SHAREDLOCALBASES_TYPE, HEADER_BIT_OFFSET_SHAREDLOCALBASES_TYPE)
 #define HEADER_SETSHAREDLOCALBASESTYPE(X, NewValue)	WRITE_BITS(X, HEADER_BIT_COUNT_SHAREDLOCALBASES_TYPE, HEADER_BIT_OFFSET_SHAREDLOCALBASES_TYPE, NewValue)
 
-#define HEADER_GETAO(X)								READ_BITS( X, HEADER_BIT_COUNT_AO, HEADER_BIT_OFFSET_AO)
-#define HEADER_SETAO(X, NewValue)					WRITE_BITS(X, HEADER_BIT_COUNT_AO, HEADER_BIT_OFFSET_AO, NewValue)
+#define HEADER_GETIRRADIANCE_AO(X)					READ_BITS( X, HEADER_BIT_COUNT_IRRADIANCE_AO, HEADER_BIT_OFFSET_IRRADIANCE_AO)
+#define HEADER_SETIRRADIANCE_AO(X, NewValue)		WRITE_BITS(X, HEADER_BIT_COUNT_IRRADIANCE_AO, HEADER_BIT_OFFSET_IRRADIANCE_AO, NewValue)
 
 #define HEADER_GETSTATE(X)							READ_BITS( X, HEADER_BIT_COUNT_STATE, HEADER_BIT_OFFSET_STATE)
 #define HEADER_SETSTATE(X, NewValue)				WRITE_BITS(X, HEADER_BIT_COUNT_STATE, HEADER_BIT_OFFSET_STATE, NewValue)
@@ -474,6 +474,55 @@ uint   StrataPackNormal24(in float3 Normal);
 float3 StrataUnpackNormal24(uint PackedNormal);
 bool   IsStrataSlabFastPathCompatible(in FStrataBSDF BSDF);
 
+///////////////////////////////////////////////////////////////////////////////
+// Irradiance and occlusion
+struct FStrataIrradianceAndOcclusion
+{
+	float MaterialAO;
+	float IndirectIrradiance;
+	uint  DiffuseIndirectSampleOcclusion;
+};
+
+FStrataIrradianceAndOcclusion InitIrradianceAndOcclusion(float InAO=1.0f)
+{
+	FStrataIrradianceAndOcclusion Out;
+	Out.MaterialAO = 1.0f;
+	Out.IndirectIrradiance = 1.0f;
+	Out.DiffuseIndirectSampleOcclusion = 0xFFu;
+	return Out;
+}
+
+uint StrataPackIrradianceAndOcclusion(FStrataIrradianceAndOcclusion In, float QuantizationBias=0)
+{
+#if GBUFFER_HAS_DIFFUSE_SAMPLE_OCCLUSION
+	return In.DiffuseIndirectSampleOcclusion;
+#elif ALLOW_STATIC_LIGHTING
+	// No space for AO. Multiply IndirectIrradiance by AO instead of storing.
+	return min(0xFFu, (EncodeIndirectIrradiance(In.IndirectIrradiance * In.MaterialAO) + QuantizationBias * (1.0 / 255.0)) * 0xFFu);
+#else
+	return saturate(In.MaterialAO) * 0xFFu;
+#endif
+}
+
+FStrataIrradianceAndOcclusion StrataUnpackIrradianceAndOcclusion(uint In /*8bits*/)
+{
+	FStrataIrradianceAndOcclusion Out;
+#if GBUFFER_HAS_DIFFUSE_SAMPLE_OCCLUSION
+	Out.DiffuseIndirectSampleOcclusion = In;
+	Out.MaterialAO = saturate(1.0 - float(countbits(Out.DiffuseIndirectSampleOcclusion)) * rcp(float(INDIRECT_SAMPLE_COUNT)));
+	Out.IndirectIrradiance = 1.f;
+#elif ALLOW_STATIC_LIGHTING
+	Out.MaterialAO = 1.f;
+	Out.DiffuseIndirectSampleOcclusion = 0x0u;
+	Out.IndirectIrradiance = DecodeIndirectIrradiance(In * rcp(255));
+#else
+	Out.MaterialAO = In * rcp(255);
+	Out.DiffuseIndirectSampleOcclusion = 0x0u;
+	Out.IndirectIrradiance = 1.f;
+#endif
+	return Out;
+}
+
 ///////////////////////////////////////////////////////////////////////////////
 // Top layer
 struct FStrataTopLayerData
@@ -1063,6 +1112,7 @@ struct FStrataPixelHeader
 	FSharedLocalBases SharedLocalBases;
 	// The strata material BSDF topology organised as a tree
 	FStrataTree StrataTree;
+	FStrataIrradianceAndOcclusion IrradianceAO;
 #endif
 #if STRATA_DEFERRED_SHADING
 	FStrataMaterialContainer StrataBuffer;
@@ -1071,8 +1121,6 @@ struct FStrataPixelHeader
 	uint PackedTopLayerData;
 #endif
 
-	float MaterialAO;
-
 	uint State;
 };
 
@@ -1109,13 +1157,13 @@ FStrataPixelHeader InitialiseStrataPixelHeader()
 #if STRATA_INLINE_SHADING
 	Out.SharedLocalBases = (FSharedLocalBases)0;
 	Out.StrataTree = (FStrataTree)0;
+	Out.IrradianceAO = InitIrradianceAndOcclusion();
 #endif
 #if STRATA_DEFERRED_SHADING
 	Out.SharedLocalBasesIndexOffset = 0;
 	Out.SharedLocalBasesTypes_PackedHeader = 0;
 	Out.PackedTopLayerData = 0;
 #endif
-	Out.MaterialAO = 1.0f;
 	Out.State = 0;
 	return Out;
 }
@@ -1903,7 +1951,12 @@ float3 StrataGetWorldTangent(in FStrataPixelHeader Header, in FStrataBSDF BSDF,
 
 float StrataGetAO(in FStrataPixelHeader Header)
 {
-	return HEADER_GETAO(Header.State);
+	return StrataUnpackIrradianceAndOcclusion(HEADER_GETIRRADIANCE_AO(Header.State)).MaterialAO;
+}
+
+FStrataIrradianceAndOcclusion StrataGetIrradianceAndAO(in FStrataPixelHeader Header)
+{
+	return StrataUnpackIrradianceAndOcclusion(HEADER_GETIRRADIANCE_AO(Header.State));
 }
 
 float StrataGetLegacyShadingModels(in FStrataBSDF BSDF)
@@ -2319,8 +2372,8 @@ void StrataConvertFromDBuffer(
 	OutHeader.SharedLocalBases = StrataInitialiseSharedLocalBases();
 	OutHeader.SharedLocalBases.Normals[SharedLocalBasisIndex] = In.WorldNormal;
 	OutHeader.SharedLocalBases.Tangents[SharedLocalBasisIndex] = 0.f;
+	OutHeader.IrradianceAO = InitIrradianceAndOcclusion();
 #endif
-	OutHeader.MaterialAO = 1.0f;
 	OutHeader.State = 0;
 	AddShadingModels(OutHeader, StrataShadingModelBit(STRATA_BSDF_TYPE_SLAB));
 
@@ -2563,6 +2616,11 @@ void PackStrataOut(
 		AndIsSimpleMaterial(StrataPixelHeader, BSDFVisibleCount == 1);
 		AndIsSingleMaterial(StrataPixelHeader, BSDFVisibleCount == 1);
 
+		// Store irradiance and AO into the state bits of the header
+	#if STRATA_INLINE_SHADING
+		HEADER_SETIRRADIANCE_AO(StrataPixelHeader.State, StrataPackIrradianceAndOcclusion(StrataPixelHeader.IrradianceAO, 0));
+	#endif
+
 		// Now write out Strata data
 
 		///////////////////////////////////////////////////////////////////////////
@@ -2598,7 +2656,6 @@ void PackStrataOut(
 #endif
 		}
 
-
 		{
 			int BSDFCount = 0;
 			STRATA_UNROLL_N(STRATA_CLAMPED_BSDF_COUNT)
@@ -2630,7 +2687,6 @@ void PackStrataOut(
 						{
 							uint Out = 0;
 							HEADER_SETSTATE(Out, StrataPixelHeader.State);
-							HEADER_SETAO(Out, PackR8(StrataPixelHeader.MaterialAO));
 
 							uint Packed16 = (PackedDiffuse8Bits << 8) | PackedRoughness8bits;
 							Out = (Out & HEADER_FASTENCODING_MASK) | (Packed16 << HEADER_FASTENCODING_BIT_COUNT);
@@ -2655,7 +2711,6 @@ void PackStrataOut(
 						{
 							uint Out = 0;
 							HEADER_SETSTATE(Out, StrataPixelHeader.State);
-							HEADER_SETAO(Out, PackR8(StrataPixelHeader.MaterialAO));
 							HEADER_SETSHADINGMODELS(Out, GetShadingModels(StrataPixelHeader));
 							Out = Out & HEADER_SINGLEENCODING_MASK;
 							Out = Out | ((BSDF.State & STATE_SINGLEENCODING_MASK) << HEADER_SINGLEENCODING_BIT_COUNT);
@@ -2789,7 +2844,6 @@ uint PackStrataHeader(uint InBSDFCount, FStrataPixelHeader InHeader)
 	HEADER_SETSHAREDLOCALBASESTYPE(Out, InHeader.SharedLocalBases.Types);
 #endif
 
-	HEADER_SETAO(Out, PackR8(InHeader.MaterialAO));
 	HEADER_SETSTATE(Out, InHeader.State);
 
 	return Out;
@@ -2803,7 +2857,6 @@ FStrataPixelHeader UnpackStrataHeaderIn(FStrataMaterialContainer StrataBuffer, i
 
 	uint PackedHeader = StrataLoadUint1(StrataBuffer, StrataAddressing);
 	Out.State		= HEADER_GETSTATE(PackedHeader);
-	Out.MaterialAO	= UnpackR8(HEADER_GETAO(PackedHeader));
 	Out.BSDFCount   = HEADER_GETBSDFCOUNT(PackedHeader);
 
 	const bool bIsSimpleMaterial = (PackedHeader & HEADER_MASK_ISSIMPLEMATERIAL) > 0;
diff --git a/Engine/Shaders/Private/Strata/StrataEnvironmentLighting.ush b/Engine/Shaders/Private/Strata/StrataEnvironmentLighting.ush
--- a/Engine/Shaders/Private/Strata/StrataEnvironmentLighting.ush
+++ b/Engine/Shaders/Private/Strata/StrataEnvironmentLighting.ush
@@ -37,7 +37,7 @@ float4 StrataEnvironmentLighting(
 	FUpsampleDFAOOutput UpsampleDFAOOutput = UpsampleDFAO(BufferUV, SceneDepth);
 #endif
 
-	float CombinedScreenAndMaterialAO = StrataPixelHeader.MaterialAO * AmbientOcclusion;
+	float CombinedScreenAndMaterialAO = StrataGetAO(StrataPixelHeader) * AmbientOcclusion;
 
 	
 #if FEATURE_LEVEL >= FEATURE_LEVEL_SM5
@@ -85,7 +85,7 @@ float4 StrataEnvironmentLighting(
 		if (bProcessFrontFaceDiffuse || bProcessBackFaceDiffuse)
 		{
 			// Compute the common sky visibility factors
-			FSkyLightVisibilityData SkyVisData = GetSkyLightVisibilityData(StrataBSDFContext.N, StrataBSDFContext.N, StrataPixelHeader.MaterialAO, AmbientOcclusion, BentNormal);
+			FSkyLightVisibilityData SkyVisData = GetSkyLightVisibilityData(StrataBSDFContext.N, StrataBSDFContext.N, StrataGetAO(StrataPixelHeader), AmbientOcclusion, BentNormal);
 
 			if (bProcessFrontFaceDiffuse)
 			{
diff --git a/Engine/Shaders/Private/Strata/StrataForwardLighting.ush b/Engine/Shaders/Private/Strata/StrataForwardLighting.ush
--- a/Engine/Shaders/Private/Strata/StrataForwardLighting.ush
+++ b/Engine/Shaders/Private/Strata/StrataForwardLighting.ush
@@ -57,7 +57,7 @@ float3 StrataForwardLightingCommon(
 	if (LightMask > 0.0)
 	{
 		// Evaluate the ShadowTerm that can then be used when integrating the lighting
-		FShadowTerms ShadowTerms = { StrataPixelHeader.MaterialAO, 1.0, 1.0, InitHairTransmittanceData() };
+		FShadowTerms ShadowTerms = { StrataGetAO(StrataPixelHeader), 1.0, 1.0, InitHairTransmittanceData() };
 
 		float4 PrecomputedShadowFactors = HasZeroPrecShadowMask(StrataPixelHeader) ? 0.0f : 1.0f;
 	#if TRANSLUCENCY_LIGHTING_SURFACE_FORWARDSHADING
@@ -315,7 +315,7 @@ float3 StrataForwardLighting(
 					float IndirectIrradiance = 0.0f; 
 
 					// Compute the bent normal from the BSDF normal
-					const FShadingOcclusion ShadingOcclusion = ApplyBentNormal(MaterialParameters.CameraVector, StrataBSDFContext.N, MaterialParameters.TangentToWorld, GetBentNormalZero(MaterialParameters), Roughness, StrataPixelHeader.MaterialAO);
+					const FShadingOcclusion ShadingOcclusion = ApplyBentNormal(MaterialParameters.CameraVector, StrataBSDFContext.N, MaterialParameters.TangentToWorld, GetBentNormalZero(MaterialParameters), Roughness,  StrataGetAO(StrataPixelHeader));
 
 					// 1. Diffuse Sky evaluation using bent normal
 					{
diff --git a/Engine/Shaders/Private/Strata/StrataVisualize.usf b/Engine/Shaders/Private/Strata/StrataVisualize.usf
--- a/Engine/Shaders/Private/Strata/StrataVisualize.usf
+++ b/Engine/Shaders/Private/Strata/StrataVisualize.usf
@@ -268,7 +268,7 @@ void VisualizeMaterialPS(
 #define SIGNOFFLOAT1(Value) if(Value < 0.0f) { PrintCharacter(PixelPos, OutColor.xyz, PARAMFontColor, TopLeft, _MINUS_); } /*else { PrintCharacter(PixelPos, OutColor.xyz, PARAMFontColor, TopLeft, _SPC_);}*/
 
 #define PRINTBSDFCOUNT()	PRINTBSDFCHAR(min(_0_+Header.BSDFCount,_F_)); PRINTBSDFCHAR(_SPC_); PRINTBSDFCHAR(_B_); PRINTBSDFCHAR(_S_); PRINTBSDFCHAR(_D_); PRINTBSDFCHAR(_F_);
-#define PRINTMATERIALAO()	PRINTBSDFCHAR(_A_); PRINTBSDFCHAR(_O_); PRINTBSDFCHAR(_SPC_); PRINTFLOAT1(Header.MaterialAO);
+#define PRINTMATERIALAO()	PRINTBSDFCHAR(_A_); PRINTBSDFCHAR(_O_); PRINTBSDFCHAR(_SPC_); PRINTFLOAT1(StrataGetAO(Header));
 #define PRINTMATERIALSKIPPRECSHADOWMASK()	PRINTBSDFCHAR(_H_); PRINTBSDFCHAR(_A_); PRINTBSDFCHAR(_S_); PRINTBSDFCHAR(_SPC_); PRINTBSDFCHAR(_P_); PRINTBSDFCHAR(_R_); PRINTBSDFCHAR(_E_); PRINTBSDFCHAR(_S_); PRINTBSDFCHAR(_H_); PRINTBSDFCHAR(_A_); PRINTBSDFCHAR(_D_); PRINTBSDFCHAR(_SPC_); PRINTBSDFCHAR(_0_+(HasPrecShadowMask(Header)?1:0));
 #define PRINTMATERIALZEROPRECSHADOWMASK()	PRINTBSDFCHAR(_Z_); PRINTBSDFCHAR(_E_); PRINTBSDFCHAR(_R_); PRINTBSDFCHAR(_O_); PRINTBSDFCHAR(_P_); PRINTBSDFCHAR(_R_); PRINTBSDFCHAR(_E_); PRINTBSDFCHAR(_S_); PRINTBSDFCHAR(_H_); PRINTBSDFCHAR(_A_); PRINTBSDFCHAR(_D_); PRINTBSDFCHAR(_SPC_); PRINTBSDFCHAR(_0_+(HasZeroPrecShadowMask(Header)?1:0));
 #define PRINTMATERIALSLW()	PRINTBSDFCHAR(_S_); PRINTBSDFCHAR(_L_); PRINTBSDFCHAR(_W_); PRINTBSDFCHAR(_SPC_); PRINTBSDFCHAR(_0_+(IsSingleLayerWater(Header)?1:0));
