Fixed volumetric cloud material for loops (ddx/ddy is not allowed in loops).

diff --git a/Engine/Shaders/Private/Strata/Strata.ush b/Engine/Shaders/Private/Strata/Strata.ush
--- a/Engine/Shaders/Private/Strata/Strata.ush
+++ b/Engine/Shaders/Private/Strata/Strata.ush
@@ -753,11 +753,15 @@ struct FStrataPixelFootprint
 
 FStrataPixelFootprint StrataGetPixelFootprint(float3 InWorldPosition)
 {
+#if CLOUD_LAYER_PIXEL_SHADER==1
+	// Clouds materials feature dynamic loops and thus cannot support derivatives.
+	FStrataPixelFootprint Out = (FStrataPixelFootprint)0;
+#else
 	const float dX = lengthFast(ddx(InWorldPosition));
 	const float dY = lengthFast(ddy(InWorldPosition));
-
 	FStrataPixelFootprint Out;
 	Out.PixelRadiusInWorldSpace = min(dX, dY) * 0.5f;
+#endif
 	return Out;
 }
 
