Add Strata support for GetDiffuseIndirectSampleOcclusion()

diff --git a/Engine/Shaders/Private/BasePassPixelShader.usf b/Engine/Shaders/Private/BasePassPixelShader.usf
--- a/Engine/Shaders/Private/BasePassPixelShader.usf
+++ b/Engine/Shaders/Private/BasePassPixelShader.usf
@@ -848,6 +848,44 @@ float3 GetBentNormal(in FMaterialPixelParameters MaterialParameters)
 #endif
 }
 
+#if STRATA_ENABLED
+uint GetDiffuseIndirectSampleOcclusion(FSharedLocalBases SharedLocalBases, float3 V, float2 SvPosition, float MaterialAO)
+{
+	uint DiffuseIndirectSampleOcclusion = 0;
+#if STRATA_INLINE_SHADING && GBUFFER_HAS_DIFFUSE_SAMPLE_OCCLUSION && !MATERIAL_SHADINGMODEL_UNLIT
+	if (MaterialAO < 1.0f)
+	{
+		uint OcclusionMask = 0x0;
+		// We must normalize each normal and tangent to avoid non normalised vectors due to per vertex interpolation or texture filtering,
+		// for the deferred (our packing relies on normalized normal) and forward (normals are going to be used as-is from registers) paths.
+		UNROLL
+		for (uint i = 0; i < SharedLocalBases.Count; ++i)
+		{
+			const float3 WorldNormal = SharedLocalBases.Normals[i];
+			const float3 BenNormal = WorldNormal; // STRATA_TODO: bent normal support - GetBentNormal(MaterialParameters)
+			const FSphericalGaussian HemisphereSG = Hemisphere_ToSphericalGaussian(WorldNormal);
+			const FSphericalGaussian VisibleSG = BentNormalAO_ToSphericalGaussian(BentNormal, MaterialAO);
+			const float VisibilityThreshold = InterleavedGradientNoise(SvPosition, View.StateFrameIndexMod8);
+
+			for (uint TracingRayIndex = 0; TracingRayIndex < INDIRECT_SAMPLE_COUNT; TracingRayIndex++)
+			{
+				const float4 E = ComputeIndirectLightingSampleE(SvPosition, TracingRayIndex, INDIRECT_SAMPLE_COUNT);
+				const FBxDFSample BxDFSample = SampleDiffuseBxDF(WorldNormal, E);
+
+				// Integrate what is visible over the maximum visibility for the normal.
+				float LVisibility = saturate(Evaluate(VisibleSG, BxDFSample.L) / Evaluate(HemisphereSG, BxDFSample.L));
+				bool bIsBentNormalOccluded = LVisibility < VisibilityThreshold;
+
+				OcclusionMask |= bIsBentNormalOccluded ? (1 << TracingRayIndex) : 0;
+			}
+		}
+
+		DiffuseIndirectSampleOcclusion = OcclusionMask;
+	}
+#endif
+	return DiffuseIndirectSampleOcclusion;
+}
+#else // STRATA_ENABLED
 uint GetDiffuseIndirectSampleOcclusion(FGBufferData GBuffer, float3 V, float3 WorldNormal, float3 BentNormal, float2 SvPosition, float MaterialAO)
 {
 	uint DiffuseIndirectSampleOcclusion = 0;
@@ -880,6 +918,8 @@ uint GetDiffuseIndirectSampleOcclusion(FGBufferData GBuffer, float3 V, float3 Wo
 #endif
 	return DiffuseIndirectSampleOcclusion;
 }
+#endif // STRATA_ENABLED
+
 
 #if USES_GBUFFER
 
@@ -1285,8 +1325,9 @@ void FPixelShaderInOut_MainPS(
 
 	// FIXME: ALLOW_STATIC_LIGHTING == 0 expects this to be AO
 	GBuffer.GBufferAO = AOMultiBounce( Luminance( GBuffer.SpecularColor ), ShadingOcclusion.SpecOcclusion ).g;
+	#if !STRATA_INLINE_SINGLELAYERWATER
 	GBuffer.DiffuseIndirectSampleOcclusion = GetDiffuseIndirectSampleOcclusion(GBuffer, MaterialParameters.CameraVector, MaterialParameters.WorldNormal, GetBentNormal(MaterialParameters), In.SvPosition.xy, MaterialAO);
-
+	#endif
 #endif // !STRATA_ENABLED
 
 	half3 DiffuseColor = 0;
@@ -1651,7 +1692,7 @@ void FPixelShaderInOut_MainPS(
 		StrataPixelHeader.IrradianceAO									= InitIrradianceAndOcclusion();
 		StrataPixelHeader.IrradianceAO.MaterialAO						= GBuffer.GBufferAO;	// STRATA_TODO: AOMultiBounce(Luminance(GBuffer.SpecularColor), ShadingOcclusion.SpecOcclusion).g; // FIXME: ALLOW_STATIC_LIGHTING == 0 expects this to be AO
 		StrataPixelHeader.IrradianceAO.IndirectIrradiance				= IndirectIrradiance;	// STRATA_TODO
-		StrataPixelHeader.IrradianceAO.DiffuseIndirectSampleOcclusion	= 0xFFu;				// STRATA_TODO: GetDiffuseIndirectSampleOcclusion(GBuffer, MaterialParameters.CameraVector, MaterialParameters.WorldNormal, GetBentNormal(MaterialParameters), In.SvPosition.xy, MaterialAO);
+		StrataPixelHeader.IrradianceAO.DiffuseIndirectSampleOcclusion	= GetDiffuseIndirectSampleOcclusion(StrataPixelHeader.SharedLocalBases, MaterialParameters.CameraVector, In.SvPosition.xy, StrataPixelHeader.IrradianceAO.MaterialAO);
 	}
 
 	// We only rely on this GBuffer structure for the write out to legacy render target such as velocity or precomputed shadow factors
diff --git a/Engine/Shaders/Private/ShadingModelsSampling.ush b/Engine/Shaders/Private/ShadingModelsSampling.ush
--- a/Engine/Shaders/Private/ShadingModelsSampling.ush
+++ b/Engine/Shaders/Private/ShadingModelsSampling.ush
@@ -51,9 +51,19 @@ struct FBxDFSample
 	uint Term;
 };
 
+FBxDFSample SampleDiffuseBxDF(float N, float4 E)
+{
+	float2 DiskE = UniformSampleDiskConcentric(E.xy);
+	float3x3 TangentBasis = GetTangentBasis(N);
+	FBxDFSample BxDFSample = (FBxDFSample)0;
+	const float TangentZ = sqrt(1 - length2(DiskE));
 
-
-
+	BxDFSample.L = mul(float3(DiskE, TangentZ), TangentBasis);
+	BxDFSample.PDF = TangentZ * rcp(PI);
+	BxDFSample.Weight = 1.0;
+	BxDFSample.Term = SHADING_TERM_DIFFUSE;
+	return BxDFSample;
+}
 
 FBxDFSample SampleDefaultLitBxDF(uint TermMask, FGBufferData GBuffer, float3 V, float4 E)
 {
