Avoid DFAO self shadowing when using r.AOQuality >=2
- This was particularly an issue on terrain.
- Increase InitialOffset based on VoxelExtent (previously it was based only on how far mesh SDF tracing went).
- Set MinStepSize based on VoxelExtent instead of hardcoded scales.

diff --git a/Engine/Shaders/Private/DistanceFieldScreenGridLighting.usf b/Engine/Shaders/Private/DistanceFieldScreenGridLighting.usf
--- a/Engine/Shaders/Private/DistanceFieldScreenGridLighting.usf
+++ b/Engine/Shaders/Private/DistanceFieldScreenGridLighting.usf
@@ -244,15 +244,16 @@ void HemisphereConeTraceAgainstGlobalDistanceFieldClipmap(
 	float3 TangentX,
 	float3 TangentY)
 {
-	float MinStepSize = GlobalVolumeCenterAndExtent[ClipmapIndex].w * 2 / 300.0f;
-	float InvAOGlobalMaxOcclusionDistance = 1.0f / AOGlobalMaxOcclusionDistance;
-	float InitialOffset = GetStepOffset(NUM_CONE_STEPS);
+	const float VoxelExtent = GlobalVolumeCenterAndExtent[ClipmapIndex].w * GlobalVolumeTexelSize;
+	const float MinStepSize = VoxelExtent;
+	const float InvAOGlobalMaxOcclusionDistance = 1.0f / AOGlobalMaxOcclusionDistance;
 
+	float InitialOffset = VoxelExtent * 2.0f;
 	float ConeTraceLeakFill = 1.0f;
-#if	!CONE_TRACE_OBJECTS
-	// Without the object cone trace there is a lot of AO leaking. Constant initial offset looks worse, but helps with leaking.
-	InitialOffset = GlobalVolumeCenterAndExtent[ClipmapIndex].w * GlobalVolumeTexelSize * 2.0f;
 
+#if CONE_TRACE_OBJECTS
+	InitialOffset = max(InitialOffset, GetStepOffset(NUM_CONE_STEPS));
+#else
 	// Cover AO leaking by comparing initial step SDF value against initial step size.
 	{
 		float3 WorldSamplePosition = WorldShadingPosition + WorldNormal * InitialOffset;
