Fix typos.

diff --git a/Engine/Shaders/Private/DeferredShadingCommon.ush b/Engine/Shaders/Private/DeferredShadingCommon.ush
--- a/Engine/Shaders/Private/DeferredShadingCommon.ush
+++ b/Engine/Shaders/Private/DeferredShadingCommon.ush
@@ -730,9 +730,6 @@ uint ExtractSubsurfaceProfileInt(FGBufferData BufferData)
 	}
 #endif
 
-
-#if SHADING_PATH_DEFERRED
-
 // @param UV - UV space in the GBuffer textures (BufferSize resolution)
 // TOOD: need to find a way to make this more convenient.
 FGBufferData GetGBufferDataFromSceneTextures(float2 UV, bool bGetNormalizedNormal = true)
@@ -770,8 +767,6 @@ uint GetSceneLightingChannel(uint2 PixelCoord)
 	return ~0;
 }
 
-#endif
-
 // @param UV - UV space in the GBuffer textures (BufferSize resolution)
 FGBufferData GetGBufferData(float2 UV, bool bGetNormalizedNormal = true)
 {
diff --git a/Engine/Shaders/Private/ShadingModelsSampling.ush b/Engine/Shaders/Private/ShadingModelsSampling.ush
--- a/Engine/Shaders/Private/ShadingModelsSampling.ush
+++ b/Engine/Shaders/Private/ShadingModelsSampling.ush
@@ -51,7 +51,7 @@ struct FBxDFSample
 	uint Term;
 };
 
-FBxDFSample SampleDiffuseBxDF(float N, float4 E)
+FBxDFSample SampleDiffuseBxDF(float3 N, float4 E)
 {
 	float2 DiskE = UniformSampleDiskConcentric(E.xy);
 	float3x3 TangentBasis = GetTangentBasis(N);
