Change SSS MFP / pixel footprint threshold as current threshold was too aggressive.

diff --git a/Engine/Shaders/Private/Strata/Strata.ush b/Engine/Shaders/Private/Strata/Strata.ush
--- a/Engine/Shaders/Private/Strata/Strata.ush
+++ b/Engine/Shaders/Private/Strata/Strata.ush
@@ -828,7 +828,7 @@ FStrataData GetStrataSlabBSDF(
 
 	// If MFP is smaller than the pixel's footprint, remove SSS attributes
 	// The scaling factor 'SSSFootprintScale' is based on empiric tests
-	const float SSSFootprintScale = 0.9f;
+	const float SSSFootprintScale = 3.f;
 	if (all(SSSDMFP * SSSFootprintScale < InPixelFootprint.PixelRadiusInWorldSpace))
 	{
 		bHasSSSProfile = false;
