Add Strata for/unpack macro to reduce boilerplate over BSDFs

diff --git a/Engine/Shaders/Private/Strata/StrataDeferredLighting.ush b/Engine/Shaders/Private/Strata/StrataDeferredLighting.ush
--- a/Engine/Shaders/Private/Strata/StrataDeferredLighting.ush
+++ b/Engine/Shaders/Private/Strata/StrataDeferredLighting.ush
@@ -4,14 +4,6 @@
 
 #include "../AreaLightCommon.ush"
 
-#ifndef STRATA_FASTPATH 
-#define STRATA_FASTPATH 0
-#endif
-
-#ifndef STRATA_SINGLEPATH 
-#define STRATA_SINGLEPATH 0
-#endif
-
 float4 StrataReadPrecomputedShadowFactors(FStrataPixelHeader StrataPixelHeader, int2 PixelPos, Texture2D PrecomputedShadowTexture)
 {
 	if (HasPrecShadowMask(StrataPixelHeader))
@@ -44,17 +36,11 @@ float4 StrataDeferredLighting(
 
 	FRectTexture RectTexture = InitRectTexture(LightData);
 
-#if STRATA_FASTPATH == 0 && STRATA_SINGLEPATH == 0
-	LOOP
-	for (uint BSDFIndex = 0; BSDFIndex < StrataPixelHeader.BSDFCount; ++BSDFIndex)
-#endif
+	Strata_for (uint BSDFIndex = 0; BSDFIndex < StrataPixelHeader.BSDFCount; ++BSDFIndex)
 	{
 		// Unpack BSDF data
-#if STRATA_FASTPATH
-		FStrataBSDF BSDF = UnpackFastPathStrataBSDFIn(MaterialBuffer, StrataAddressing, StrataPixelHeader);
-#else
-		FStrataBSDF BSDF = UnpackStrataBSDFIn(MaterialBuffer, StrataAddressing, StrataPixelHeader);
-#endif
+		FStrataBSDF BSDF = UnpackStrataBSDF(MaterialBuffer, StrataAddressing, StrataPixelHeader);
+
 		FStrataBSDFContext BSDFContext = StrataCreateBSDFContext(StrataPixelHeader, BSDF, StrataAddressing, V, L);
 
 		float Roughness = StrataGetBSDFRoughness(BSDFContext.BSDF);
diff --git a/Engine/Shaders/Private/Strata/StrataEnvironmentLighting.ush b/Engine/Shaders/Private/Strata/StrataEnvironmentLighting.ush
--- a/Engine/Shaders/Private/Strata/StrataEnvironmentLighting.ush
+++ b/Engine/Shaders/Private/Strata/StrataEnvironmentLighting.ush
@@ -2,14 +2,6 @@
 
 #pragma once
 
-#ifndef STRATA_FASTPATH 
-#define STRATA_FASTPATH 0
-#endif
-
-#ifndef STRATA_SINGLEPATH 
-#define STRATA_SINGLEPATH 0
-#endif
-
 // Env. lighting evaluation for strata material.
 // Unpack BSDF on-the-fly
 float4 StrataEnvironmentLighting(
@@ -52,15 +44,10 @@ float4 StrataEnvironmentLighting(
 	uint NumCulledReflectionCaptures = 0;
 #endif
 
-#if STRATA_FASTPATH == 0 && STRATA_SINGLEPATH == 0
-	for (uint BSDFIndex = 0; BSDFIndex < StrataPixelHeader.BSDFCount; ++BSDFIndex)
-#endif
+	Strata_for (uint BSDFIndex = 0; BSDFIndex < StrataPixelHeader.BSDFCount; ++BSDFIndex)
 	{
-	#if STRATA_FASTPATH
-		FStrataBSDF BSDF = UnpackFastPathStrataBSDFIn(MaterialBuffer, StrataAddressing, StrataPixelHeader);
-	#else
-		FStrataBSDF BSDF = UnpackStrataBSDFIn(MaterialBuffer, StrataAddressing, StrataPixelHeader);
-	#endif
+		FStrataBSDF BSDF = UnpackStrataBSDF(MaterialBuffer, StrataAddressing, StrataPixelHeader);
+
 		const float3 BSDFThroughput = BSDF.LuminanceWeight;
 
 		// Create the BSDF context
diff --git a/Engine/Shaders/Private/Strata/StrataEvaluation.ush b/Engine/Shaders/Private/Strata/StrataEvaluation.ush
--- a/Engine/Shaders/Private/Strata/StrataEvaluation.ush
+++ b/Engine/Shaders/Private/Strata/StrataEvaluation.ush
@@ -19,6 +19,10 @@
 #define STRATA_FASTPATH 0
 #endif
 
+#ifndef STRATA_SINGLE 
+#define STRATA_SINGLE 0
+#endif
+
 #ifndef STRATA_TRANSLUCENT_ENABLED
 #define STRATA_TRANSLUCENT_ENABLED 0
 #endif
@@ -35,6 +39,28 @@
 #define STRATA_SSS_TRANSMISSION 0
 #endif
 
+///////////////////////////////////////////////////////////////////////////////
+// Strata Macros for reducing permutation boiler plate
+// 
+// This allows to write a loop over all BSDFs as follow and support permutations
+//
+// Strata_for(uint BSDFIndex = 0; BSDFIndex < StrataPixelHeader.BSDFCount; ++BSDFIndex)
+// {
+//		FStrataBSDF BSDF = UnpackStrataBSDF(MaterialBuffer, StrataAddressing, StrataPixelHeader);
+//		...
+// }
+#if STRATA_FASTPATH == 0 && STRATA_SINGLEPATH == 0
+	#define Strata_for(X) for(X)
+#else
+	#define Strata_for(X) 
+#endif
+
+#if STRATA_FASTPATH
+	#define UnpackStrataBSDF(X, Y, Z) UnpackFastPathStrataBSDFIn(X, Y, Z)
+#else
+	#define UnpackStrataBSDF(X, Y, Z) UnpackStrataBSDFIn(X, Y, Z)
+#endif
+
 // STRATA_TODO put in a common file
 // Point lobe in off-specular peak direction
 float3 StrataGetOffSpecularPeakReflectionDir(float3 Normal, float3 ReflectionVector, float Roughness)
@@ -47,7 +73,6 @@ float3 StrataGetOffSpecularPeakReflectionDir(float3 Normal, float3 ReflectionVec
 #define SPECULAR_HAZINESS_WEIGHT 0.5f
 
 
-
 ///////////////////////////////////////////////////////////////////////////////
 // BSDF evaluate and sampling
 
