Fix missing support for dithered opacity mask in velocity only shaders.

diff --git a/Engine/Shaders/Private/VelocityShader.usf b/Engine/Shaders/Private/VelocityShader.usf
--- a/Engine/Shaders/Private/VelocityShader.usf
+++ b/Engine/Shaders/Private/VelocityShader.usf
@@ -175,6 +175,8 @@ void MainPixelShader(
 
 #if MATERIALBLENDING_MASKED_USING_COVERAGE
 	OutCoverage = DiscardMaterialWithPixelCoverage(MaterialParameters, PixelMaterialInputs);
+#else
+	GetMaterialCoverageAndClipping(MaterialParameters, PixelMaterialInputs);
 #endif
 
 	// 3d velocity, includes camera an object motion
