* Simplify shader print API: Context/Print/Newline.
* Added small inline documentation for shader print
* Fix incorrect slider/checkbox text alignement.

diff --git a/Engine/Shaders/Private/HairStrands/HairStrandsVoxelPageRayMarching.usf b/Engine/Shaders/Private/HairStrands/HairStrandsVoxelPageRayMarching.usf
--- a/Engine/Shaders/Private/HairStrands/HairStrandsVoxelPageRayMarching.usf
+++ b/Engine/Shaders/Private/HairStrands/HairStrandsVoxelPageRayMarching.usf
@@ -61,8 +61,11 @@ void MainCS(uint3 DispatchThreadId : SV_DispatchThreadID)
 	const FVirtualVoxelNodeDesc NodeDesc = UnpackVoxelNode(VirtualVoxel.NodeDescBuffer[MacroGroupId], VirtualVoxel.PageResolution);
 	const bool bIsNodeValid = all(NodeDesc.PageIndexResolution != 0);
 
+
+	const bool bPrint = all(uint2(PixelCoord) == 0);
+
 	// Draw macro group AABB & debug infos
-	if (uint(PixelCoord.x) == 0 && uint(PixelCoord.y) == 0)
+	if (bPrint)
 	{
 		const uint TotalPageIndexCount = MaxTotalPageIndexCount;
 		const uint TotalPageCount = VirtualVoxel.PageCountResolution.x * VirtualVoxel.PageCountResolution.y * VirtualVoxel.PageCountResolution.z;
@@ -93,85 +96,84 @@ void MainCS(uint3 DispatchThreadId : SV_DispatchThreadID)
 		//  Total Page Occupancy          :
 		if (MacroGroupId == 0)
 		{
-			float2 StartPos = float2(50.0f, 100.0f ) / OutputResolution;
-			float2 Pos = StartPos;
+			FShaderPrintContext Context = InitShaderPrintContext(true, uint2(50, 100));
 
 			// Adaptive voxels
-			Pos = ShaderPrintText(Pos, TEXT("Alloc. Feedback  "), FontOrange);
+			Print(Context, TEXT("Alloc. Feedback  "), FontOrange);
 			if (VirtualVoxel.AllocationFeedbackEnable)
 			{
-				Pos = ShaderPrintText(Pos, TEXT("True"), FontGreen);
+				Print(Context, TEXT("True"), FontGreen);
 			}
 			else
 			{
-				Pos = ShaderPrintText(Pos, TEXT("False"), FontRed);
+				Print(Context, TEXT("False"), FontRed);
 			}
-			Pos = ShaderPrintNewline(Pos); Pos.x = StartPos.x;
+			Newline(Context);
 
 			// Actual voxel/pages size & resolution
-			Pos = ShaderPrintText(Pos, TEXT("GPU Curr. Voxel Size        "), FontOrange);
-			Pos = ShaderPrint(Pos, VirtualVoxel.CurrGPUMinVoxelSize[0], FontEmerald);
-			Pos = ShaderPrintNewline(Pos); Pos.x = StartPos.x;
+			Print(Context, TEXT("GPU Curr. Voxel Size        "), FontOrange);
+			Print(Context, VirtualVoxel.CurrGPUMinVoxelSize[0], FontEmerald);
+			Newline(Context);
 
-			Pos = ShaderPrintText(Pos, TEXT("GPU Next. Voxel Size        "));
-			Pos = ShaderPrint(Pos, VirtualVoxel.NextGPUMinVoxelSize[0]);
-			Pos = ShaderPrintNewline(Pos); Pos.x = StartPos.x;
+			Print(Context, TEXT("GPU Next. Voxel Size        "));
+			Print(Context, VirtualVoxel.NextGPUMinVoxelSize[0]);
+			Newline(Context);
 
-			Pos = ShaderPrintText(Pos, TEXT("CPU Min. Voxel Size         "));
-			Pos = ShaderPrint(Pos, VirtualVoxel.CPUMinVoxelWorldSize);
-			Pos = ShaderPrintNewline(Pos); Pos.x = StartPos.x;
+			Print(Context, TEXT("CPU Min. Voxel Size         "));
+			Print(Context, VirtualVoxel.CPUMinVoxelWorldSize);
+			Newline(Context);
 			
-			Pos = ShaderPrintText(Pos, TEXT("Page Res                    "));
-			Pos = ShaderPrint(Pos, VirtualVoxel.PageResolution);
-			Pos = ShaderPrintNewline(Pos); Pos.x = StartPos.x;
-
-			Pos = ShaderPrintText(Pos, TEXT("Page Count Res              "));
-			Pos = ShaderPrint(Pos, VirtualVoxel.PageCountResolution.x);
-			Pos = ShaderPrint(Pos, VirtualVoxel.PageCountResolution.y);
-			Pos = ShaderPrint(Pos, VirtualVoxel.PageCountResolution.z);
-			Pos = ShaderPrintNewline(Pos); Pos.x = StartPos.x;
-
-			Pos = ShaderPrintText(Pos, TEXT("Page Texture Res            "), FontEmerald);
-			Pos = ShaderPrint(Pos, VirtualVoxel.PageTextureResolution.x, FontEmerald);
-			Pos = ShaderPrint(Pos, VirtualVoxel.PageTextureResolution.y, FontEmerald);
-			Pos = ShaderPrint(Pos, VirtualVoxel.PageTextureResolution.z, FontEmerald);
-			Pos = ShaderPrintNewline(Pos); Pos.x = StartPos.x;
+			Print(Context, TEXT("Page Res                    "));
+			Print(Context, VirtualVoxel.PageResolution);
+			Newline(Context);
+
+			Print(Context, TEXT("Page Count Res              "));
+			Print(Context, VirtualVoxel.PageCountResolution.x);
+			Print(Context, VirtualVoxel.PageCountResolution.y);
+			Print(Context, VirtualVoxel.PageCountResolution.z);
+			Newline(Context);
+
+			Print(Context, TEXT("Page Texture Res            "), FontEmerald);
+			Print(Context, VirtualVoxel.PageTextureResolution.x, FontEmerald);
+			Print(Context, VirtualVoxel.PageTextureResolution.y, FontEmerald);
+			Print(Context, VirtualVoxel.PageTextureResolution.z, FontEmerald);
+			Newline(Context);
 
 			// Total page indices stats
-			Pos = ShaderPrintNewline(Pos); Pos.x = StartPos.x;
+			Newline(Context);
 
-			Pos = ShaderPrintText(Pos, TEXT("Total Allocated Page Index  "));
-			Pos = ShaderPrint(Pos, TotalAllocatedPageCount);
-			Pos = ShaderPrintNewline(Pos); Pos.x = StartPos.x;
+			Print(Context, TEXT("Total Allocated Page Index  "));
+			Print(Context, TotalAllocatedPageCount);
+			Newline(Context);
 
-			Pos = ShaderPrintText(Pos, TEXT("Total Page Index            "));
-			Pos = ShaderPrint(Pos, TotalPageIndexCount);
-			Pos = ShaderPrintNewline(Pos); Pos.x = StartPos.x;
+			Print(Context, TEXT("Total Page Index            "));
+			Print(Context, TotalPageIndexCount);
+			Newline(Context);
 
 			const float PageIndexOccupancy = TotalAllocatedPageCount / float(TotalPageIndexCount);
-			Pos = ShaderPrintText(Pos, TEXT("Page Index Occupancy        "), FontOrange);
-			Pos = ShaderPrint(Pos, PageIndexOccupancy * 100.f, GetOccupancyColor(PageIndexOccupancy));
-			Pos = ShaderPrintNewline(Pos); Pos.x = StartPos.x;
+			Print(Context, TEXT("Page Index Occupancy        "), FontOrange);
+			Print(Context, PageIndexOccupancy * 100.f, GetOccupancyColor(PageIndexOccupancy));
+			Newline(Context);
 
 			// Total pages stats
-			Pos = ShaderPrintNewline(Pos); Pos.x = StartPos.x;
+			Newline(Context);
 
-			Pos = ShaderPrintText(Pos, TEXT("Total Allocated Valid Page  "));
-			Pos = ShaderPrint(Pos, TotalAllocatedValidPageCount);
-			Pos = ShaderPrintNewline(Pos); Pos.x = StartPos.x;
+			Print(Context, TEXT("Total Allocated Valid Page  "));
+			Print(Context, TotalAllocatedValidPageCount);
+			Newline(Context);
 
-			Pos = ShaderPrintText(Pos, TEXT("Total Page                  "));
-			Pos = ShaderPrint(Pos, TotalPageCount);
-			Pos = ShaderPrintNewline(Pos); Pos.x = StartPos.x;
+			Print(Context, TEXT("Total Page                  "));
+			Print(Context, TotalPageCount);
+			Newline(Context);
 
 			const float PageOccupancy = TotalAllocatedValidPageCount / float(TotalPageCount);
-			Pos = ShaderPrintText(Pos, TEXT("Page Occupancy              "), FontOrange);
-			Pos = ShaderPrint(Pos, PageOccupancy * 100.f, GetOccupancyColor(PageOccupancy));
-			Pos = ShaderPrintNewline(Pos); Pos.x = StartPos.x;
+			Print(Context, TEXT("Page Occupancy              "), FontOrange);
+			Print(Context, PageOccupancy * 100.f, GetOccupancyColor(PageOccupancy));
+			Newline(Context);
 
-			Pos = ShaderPrintText(Pos, TEXT("Translated world offset     "), FontOrange);
-			Pos = ShaderPrint(Pos, VirtualVoxel.TranslatedWorldOffset);
-			Pos = ShaderPrintNewline(Pos); Pos.x = StartPos.x;
+			Print(Context, TEXT("Translated world offset     "), FontOrange);
+			Print(Context, VirtualVoxel.TranslatedWorldOffset);
+			Newline(Context);
 		}
 
 		// MacroGroupId XX
@@ -180,76 +182,74 @@ void MainCS(uint3 DispatchThreadId : SV_DispatchThreadID)
 		//  Page Index Offset    : O
 		if (MacroGroupId == 0)
 		{
-			float2 StartPos = float2(50.0f, 420.0f) / OutputResolution;
-			float2 Pos = StartPos;
+			FShaderPrintContext Context = InitShaderPrintContext(true, uint2(50, 420));
 
-			Pos = ShaderPrintText(Pos, TEXT("Macro Group        "), FontOrange);
-			Pos = ShaderPrintNewline(Pos); Pos.x = StartPos.x;
+			Print(Context, TEXT("Macro Group        "), FontOrange);
+			Newline(Context);
 
-			Pos = ShaderPrintText(Pos, TEXT(" Virtual Res X     "), FontEmerald);
-			Pos = ShaderPrintNewline(Pos); Pos.x = StartPos.x;
-			Pos = ShaderPrintText(Pos, TEXT(" Virtual Res Y     "), FontEmerald);
-			Pos = ShaderPrintNewline(Pos); Pos.x = StartPos.x;
-			Pos = ShaderPrintText(Pos, TEXT(" Virtual Res Z     "), FontEmerald);
-			Pos = ShaderPrintNewline(Pos); Pos.x = StartPos.x;
+			Print(Context, TEXT(" Virtual Res X     "), FontEmerald);
+			Newline(Context);
+			Print(Context, TEXT(" Virtual Res Y     "), FontEmerald);
+			Newline(Context);
+			Print(Context, TEXT(" Virtual Res Z     "), FontEmerald);
+			Newline(Context);
 
-			Pos = ShaderPrintNewline(Pos); Pos.x = StartPos.x;
+			Newline(Context);
 
-			Pos = ShaderPrintText(Pos, TEXT(" Page Index Res X  "));
-			Pos = ShaderPrintNewline(Pos); Pos.x = StartPos.x;
-			Pos = ShaderPrintText(Pos, TEXT(" Page Index Res Y  "));
-			Pos = ShaderPrintNewline(Pos); Pos.x = StartPos.x;
-			Pos = ShaderPrintText(Pos, TEXT(" Page Index Res Z  "));
-			Pos = ShaderPrintNewline(Pos); Pos.x = StartPos.x;
+			Print(Context, TEXT(" Page Index Res X  "));
+			Newline(Context);
+			Print(Context, TEXT(" Page Index Res Y  "));
+			Newline(Context);
+			Print(Context, TEXT(" Page Index Res Z  "));
+			Newline(Context);
 
-			Pos = ShaderPrintNewline(Pos); Pos.x = StartPos.x;
+			Newline(Context);
 
-			Pos = ShaderPrintText(Pos, TEXT(" Page Index Offset "));
-			Pos = ShaderPrintNewline(Pos); Pos.x = StartPos.x;
+			Print(Context, TEXT(" Page Index Offset "));
+			Newline(Context);
 
-			Pos = ShaderPrintNewline(Pos); Pos.x = StartPos.x;
+			Newline(Context);
 
-			Pos = ShaderPrintText(Pos, TEXT(" Voxel World Size  "));
-			Pos = ShaderPrintNewline(Pos); Pos.x = StartPos.x;
+			Print(Context, TEXT(" Voxel World Size  "));
+			Newline(Context);
 		}
 
 		{
-			float2 StartPos = float2(250.0f + float(MacroGroupId) * 35.0f, 420.0f) / OutputResolution;
-			float2 Pos = StartPos;
+			FShaderPrintContext Context = InitShaderPrintContext(true, uint2(250 + MacroGroupId * 35, 420));
 
-			Pos = ShaderPrint(Pos, MacroGroupId, FontOrange);
-			Pos = ShaderPrintNewline(Pos); Pos.x = StartPos.x;
+			Print(Context, MacroGroupId, FontOrange);
+			Newline(Context);
 
-			Pos = ShaderPrint(Pos, NodeDesc.VirtualResolution.x, FontEmerald);
-			Pos = ShaderPrintNewline(Pos); Pos.x = StartPos.x;
-			Pos = ShaderPrint(Pos, NodeDesc.VirtualResolution.y, FontEmerald);
-			Pos = ShaderPrintNewline(Pos); Pos.x = StartPos.x;
-			Pos = ShaderPrint(Pos, NodeDesc.VirtualResolution.z, FontEmerald);
-			Pos = ShaderPrintNewline(Pos); Pos.x = StartPos.x;
+			Print(Context, NodeDesc.VirtualResolution.x, FontEmerald);
+			Newline(Context);
+			Print(Context, NodeDesc.VirtualResolution.y, FontEmerald);
+			Newline(Context);
+			Print(Context, NodeDesc.VirtualResolution.z, FontEmerald);
+			Newline(Context);
 
-			Pos = ShaderPrintNewline(Pos); Pos.x = StartPos.x;
+			Newline(Context);
 
-			Pos = ShaderPrint(Pos, NodeDesc.PageIndexResolution.x, FontYellow);
-			Pos = ShaderPrintNewline(Pos); Pos.x = StartPos.x;
-			Pos = ShaderPrint(Pos, NodeDesc.PageIndexResolution.y, FontYellow);
-			Pos = ShaderPrintNewline(Pos); Pos.x = StartPos.x;
-			Pos = ShaderPrint(Pos, NodeDesc.PageIndexResolution.z, FontYellow);
-			Pos = ShaderPrintNewline(Pos); Pos.x = StartPos.x;
+			Print(Context, NodeDesc.PageIndexResolution.x, FontYellow);
+			Newline(Context);
+			Print(Context, NodeDesc.PageIndexResolution.y, FontYellow);
+			Newline(Context);
+			Print(Context, NodeDesc.PageIndexResolution.z, FontYellow);
+			Newline(Context);
 
-			Pos = ShaderPrintNewline(Pos); Pos.x = StartPos.x;
+			Newline(Context);
 
 			// Offset of 1 line for every other voxel size, to avoid super-position with large float number
 			const bool bAddLine = (MacroGroupId & 1) == 1;
 
-			if (bAddLine) { Pos = ShaderPrintNewline(Pos); Pos.x = StartPos.x; }
+			if (bAddLine) { Newline(Context); }
 
-			Pos = ShaderPrint(Pos, NodeDesc.PageIndexOffset, FontSilver);
-			Pos = ShaderPrintNewline(Pos); Pos.x = StartPos.x;
+			Print(Context, NodeDesc.PageIndexOffset, FontSilver);
+			Newline(Context);
 
-			Pos = ShaderPrintNewline(Pos); Pos.x = StartPos.x;
+			Newline(Context);
 
-			Pos = ShaderPrint(Pos, NodeDesc.VoxelWorldSize, FontSilver);
-			Pos = ShaderPrintNewline(Pos); Pos.x = StartPos.x;
+			Print(Context, NodeDesc.VoxelWorldSize, FontSilver);
+			Newline(Context);
 		}
 	}
 
@@ -360,20 +360,19 @@ void MainCS(uint3 DispatchThreadId : SV_DispatchThreadID)
 		const bool bIsPixelEnable = all(int2(PixelCoord) == GetCursorPos());
 		if (bIsPixelEnable && bIsValid)
 		{
-			float2 StartPos = float2(50.0f, 50.0f) / OutputResolution;
-			float2 Pos = StartPos;
+			FShaderPrintContext Context = InitShaderPrintContext(true, uint2(50, 50));
 
-			Pos = ShaderPrintText(Pos, TEXT("X, Y "));
-			Pos = ShaderPrint(Pos, GetCursorPos());
-			Pos = ShaderPrintNewline(Pos); Pos.x = StartPos.x;
+			Print(Context, TEXT("X, Y "));
+			Print(Context, GetCursorPos());
+			Newline(Context);
 
-			Pos = ShaderPrintText(Pos, TEXT("T.Count "));
-			Pos = ShaderPrint(Pos, TraversalResult.HairCount);
-			Pos = ShaderPrintNewline(Pos); Pos.x = StartPos.x;
+			Print(Context, TEXT("T.Count "));
+			Print(Context, TraversalResult.HairCount);
+			Newline(Context);
 			
-			Pos = ShaderPrintText(Pos, TEXT("T.Visibility "));
-			Pos = ShaderPrint(Pos, TraversalResult.Visibility);
-			Pos = ShaderPrintNewline(Pos); Pos.x = StartPos.x;
+			Print(Context, TEXT("T.Visibility "));
+			Print(Context, TraversalResult.Visibility);
+			Newline(Context);
 		}
 	}
 
diff --git a/Engine/Shaders/Private/ShaderPrintCommon.ush b/Engine/Shaders/Private/ShaderPrintCommon.ush
--- a/Engine/Shaders/Private/ShaderPrintCommon.ush
+++ b/Engine/Shaders/Private/ShaderPrintCommon.ush
@@ -8,7 +8,81 @@
 // Include the MiniFont symbol definitions
 #include "MiniFontCommon.ush"
 
-// --------------------------------------------------------------------------
+///////////////////////////////////////////////////////////////////////////////////////////////////
+// Description
+// Shader Print allows to print value & text from shaders. 
+// 
+// 0. Contex-based state API
+// -------------------------
+// FShaderPrintContext Context = InitShaderPrintContext(all(SvPosition.xy == float2(100, 100), uint2(50,50))
+//                             ------------------------------------   ------------
+//                                 Pixel/thread to filter             Coord at which 
+//                                                                       to draw
+// 
+// 1. How to print value?
+// ----------------------
+// First create a context, then print scalar/vector/matrix/text...
+// 
+//   FShaderPrintContext Context = InitShaderPrintContext(all(PixelPos.xy == uint2(100, 100), uint2(50,50))
+// 
+//   Print(Context, MyFloat);
+//   Newline(Context);
+// 
+//   Print(Context, TEXT("My string for labbelling a float :"), FontOrange);
+//   Print(Context, MyFloatInYellow, FontYellow);
+//   Newline(Context);
+// 
+// 
+// 2. How to add widget?
+// ---------------------
+// Widget can be added from shader. 
+// For slider:
+// 
+//   const float MySliderValue  = AddSlider(Context, DefaultValue, TEXT("My Slider"));
+// 
+// For checkbox:
+// 
+//   const bool MyCheckboxValue = AddCheckbox(Context, DefaultValue, TEXT("My Checkbox"));
+//
+// The value of these widgets can be retrieved/evaluated from *any* thread/pixels, but the 
+// drawing need to be done from a single thread/pixel. This is done through the ShaderPrintContextcontext, 
+// which setup which pixel/thread is active/filtered. Example with a compute shader:
+// 
+//    [numthreads(8, 8, 1)]
+//    void MainCS(uint2 DispatchThreadId : SV_DispatchThreadID)
+//    {
+//       // Create context which 'add' the widget only on thread (0,0)
+//       FShaderPrintContext Context = InitShaderPrintContext(all(DispatchThreadId == uint2(0, 0), uint2(50,50))
+// 
+//		 // Any threads can get the slider value with this call
+//       const float ScaleValue = AddSlider(Context, 1.f, TEXT("Scale"));
+//       ...
+//    }
+// 
+// 
+// 4. Global-based state API
+// -------------------------
+// Another API can be used with a more 'global' context accross the shader execution.
+// 
+// This API is used to restrict ShaderPrint to a subset of active threads
+// The filter functions can be called once at the start of the shader to filter all later printing
+//
+// To debug a single pixel:
+//
+//   ShaderPrintFilter(all(SvPosition.xy == float2(100, 100)));
+// 
+// To debug a single compute shader thread:
+// 
+//   ShaderPrintFilter(all(DispatchThreadId == uint3(0, 0, 0)));
+// 
+// To debug a single random active thread
+// 
+//   ShaderPrintFilterOneThread()
+// 
+// Note that ShaderPrintFilterOneThread() only works for one shader invocation in a view since it relies on a global memory location 
+//
+
+///////////////////////////////////////////////////////////////////////////////////////////////////
 // Uniforms and structures
 
 // Needs to match C++ code in ShaderPrintUniform.cpp
@@ -32,7 +106,7 @@ struct FPackedShaderPrintItem
 StructuredBuffer<uint> StateBuffer;
 RWStructuredBuffer<FPackedShaderPrintItem> RWValuesBuffer;
 
-// --------------------------------------------------------------------------
+///////////////////////////////////////////////////////////////////////////////////////////////////
 // Input types
 
 // Content of FShaderPrintItem.Type defines what to cast FShaderPrintItem.Value as
@@ -48,7 +122,7 @@ RWStructuredBuffer<FPackedShaderPrintItem> RWValuesBuffer;
 #define SHADER_PRINT_COUNT_OFFSET 0
 #define SHADER_PRINT_VALUE_OFFSET 1
 
-// --------------------------------------------------------------------------
+///////////////////////////////////////////////////////////////////////////////////////////////////
 // State 
 #define SHADER_PRINT_STATE_STRIDE 3
 #define SHADER_PRINT_STATE_INDEX_METADATA 0
@@ -58,7 +132,7 @@ RWStructuredBuffer<FPackedShaderPrintItem> RWValuesBuffer;
 #define SHADER_PRINT_STATE_INVALID_INDEX 0xFF
 #define SHADER_PRINT_STATE_HASH_MASK 0xFFFFFF
 
-// --------------------------------------------------------------------------
+///////////////////////////////////////////////////////////////////////////////////////////////////
 // Font Color
 
 struct FFontColor
@@ -102,7 +176,7 @@ static const FFontColor FontSilver		= InitFontColor(189.f / 255.f, 195.f / 255.f
 static const FFontColor FontEmerald		= InitFontColor( 46.f / 255.f, 204.f / 255.f, 113.f / 255.f);
 #endif
 
-// --------------------------------------------------------------------------
+///////////////////////////////////////////////////////////////////////////////////////////////////
 // Util pack/unpac functions
 
 FPackedShaderPrintItem PackShaderPrintItem(FShaderPrintItem In)
@@ -133,7 +207,7 @@ FShaderPrintItem UnpackShaderPrintItem(FPackedShaderPrintItem In)
 	return Out;
 }
 
-// --------------------------------------------------------------------------
+///////////////////////////////////////////////////////////////////////////////////////////////////
 // Global variables
 
 #include "/Engine/Generated/UniformBuffers/ShaderPrintUniform.ush"
@@ -145,7 +219,33 @@ static const float ShaderPrintTabCount = 12.f;
 static float2 ShaderPrintCursorPos = ShaderPrintCursorStart;
 static bool ShaderPrintFilterEnable = true;
 
-// --------------------------------------------------------------------------
+///////////////////////////////////////////////////////////////////////////////////////////////////
+// Print API (Experimental)
+
+struct FShaderPrintContext
+{
+	bool bIsActive;
+	float2 StartPos;
+	float2 Pos;
+};
+
+FShaderPrintContext InitShaderPrintContext(bool bActive, float2 InStartPos)
+{
+	FShaderPrintContext Out;
+	Out.bIsActive = bActive;
+	Out.StartPos = InStartPos;
+	Out.Pos = InStartPos;
+	return Out;
+}
+FShaderPrintContext InitShaderPrintContext(bool bActive, uint2 InStartCoord) { return InitShaderPrintContext(bActive, float2(InStartCoord) / float2(ShaderPrintUniform.Resolution)); }
+
+void Newline(inout FShaderPrintContext InContext)
+{
+	InContext.Pos.x = InContext.StartPos.x;
+	InContext.Pos.y += ShaderPrintUniform.FontSpacing.y;
+}
+
+///////////////////////////////////////////////////////////////////////////////////////////////////
 // Internal helpers
 
 void ShaderPrint_Internal(in FShaderPrintItem Item)
@@ -182,7 +282,7 @@ void ShaderPrint_Internal(in float2 ScreenPos, in int Value, in FFontColor FontC
 	ShaderPrint_Internal(Item);
 }
 
-// --------------------------------------------------------------------------
+///////////////////////////////////////////////////////////////////////////////////////////////////
 // Symbol printing
 
 float2 ShaderPrintSymbol(in float2 ScreenPos, in int Symbol, in FFontColor Color)
@@ -210,53 +310,63 @@ void ShaderPrintSymbol(in int symbol)
 	ShaderPrintCursorPos = ShaderPrintSymbol(ShaderPrintCursorPos, symbol, GetDefaultFontColor());
 }
 
-// Function for reading global TEXT string. 
-// The data and function are generated by the shader compiler. These functions are used to access the text data:
-// - uint ShaderPrintGetChar(uint InIndex);
-// - uint ShaderPrintGetOffset(uint InTextEntry);
-// - uint ShaderPrintGetHash(uint InTextEntry);
-GENERATED_SHADER_PRINT
+///////////////////////////////////////////////////////////////////////////////////////////////////
+// Text & hash
 
+struct FShaderPrintText
+{
+	uint Index;
+};
+
+FShaderPrintText InitShaderPrintText(uint InIndex)
+{
+	FShaderPrintText Out;
+	Out.Index = InIndex;
+	return Out;
+}
+
+// Forward declarations for text accessors
+uint ShaderPrintGetChar(uint InIndex);
+uint ShaderPrintGetOffset(FShaderPrintText InTextEntry);
+uint ShaderPrintGetHash(FShaderPrintText InTextEntry);
+ 
 // Function for reading global TEXT string
-float2 ShaderPrintText(float2 Pos, uint InTextEntry, FFontColor InColor)
+float2 ShaderPrintText(float2 Pos, FShaderPrintText InTextEntry, FFontColor InColor)
 {
 	uint Begin = ShaderPrintGetOffset(InTextEntry);
-	uint End = ShaderPrintGetOffset(InTextEntry + 1);
+	uint End = ShaderPrintGetOffset(InitShaderPrintText(InTextEntry.Index + 1));
 	for (uint i = Begin; i < End; ++i)
 	{
 		Pos = ShaderPrintSymbol(Pos, ShaderPrintGetChar(i), InColor);
 	}
 	return Pos;
 }
-float2 ShaderPrintText(float2 Pos, uint InTextEntry) { return ShaderPrintText(Pos, InTextEntry, GetDefaultFontColor()); }
+float2 ShaderPrintText(float2 Pos, FShaderPrintText InTextEntry) { return ShaderPrintText(Pos, InTextEntry, GetDefaultFontColor()); }
 
-float2 ShaderPrintHelloWorld(in float2 ScreenPos)
+void ShaderPrintText(inout FShaderPrintContext InContext, FShaderPrintText InTextEntry, FFontColor InColor)
 {
-	if (ShaderPrintFilterEnable)
+	if (InContext.bIsActive)
 	{
-		ScreenPos = ShaderPrintSymbol(ScreenPos, _H_);
-		ScreenPos = ShaderPrintSymbol(ScreenPos, _E_);
-		ScreenPos = ShaderPrintSymbol(ScreenPos, _L_);
-		ScreenPos = ShaderPrintSymbol(ScreenPos, _L_);
-		ScreenPos = ShaderPrintSymbol(ScreenPos, _O_);
-		ScreenPos = ShaderPrintSymbol(ScreenPos, _SPC_);
-		ScreenPos = ShaderPrintSymbol(ScreenPos, _W_);
-		ScreenPos = ShaderPrintSymbol(ScreenPos, _O_);
-		ScreenPos = ShaderPrintSymbol(ScreenPos, _R_);
-		ScreenPos = ShaderPrintSymbol(ScreenPos, _L_);
-		ScreenPos = ShaderPrintSymbol(ScreenPos, _D_);
-		ScreenPos = ShaderPrintSymbol(ScreenPos, _SPC_);
+		uint Begin = ShaderPrintGetOffset(InTextEntry);
+		uint End = ShaderPrintGetOffset(InitShaderPrintText(InTextEntry.Index + 1));
+		for (uint i = Begin; i < End; ++i)
+		{
+			InContext.Pos = ShaderPrintSymbol(InContext.Pos, ShaderPrintGetChar(i), InColor);
+		}	
 	}
-	return ScreenPos;
 }
+void ShaderPrintText(inout FShaderPrintContext InContext, FShaderPrintText InTextEntry) { ShaderPrintText(InContext, InTextEntry, GetDefaultFontColor()); }
 
-void ShaderPrintHelloWorld()
-{
-	ShaderPrintCursorPos = ShaderPrintHelloWorld(ShaderPrintCursorPos);
-}
+float2 Print(float2 Pos, FShaderPrintText InTextEntry, FFontColor InColor)							{ return ShaderPrintText(Pos, InTextEntry, InColor); }
+float2 Print(float2 Pos, FShaderPrintText InTextEntry)												{ return ShaderPrintText(Pos, InTextEntry, GetDefaultFontColor()); }
+void   Print(inout FShaderPrintContext InContext, FShaderPrintText InTextEntry, FFontColor InColor)	{ ShaderPrintText(InContext, InTextEntry, InColor); }
+void   Print(inout FShaderPrintContext InContext, FShaderPrintText InTextEntry)						{ ShaderPrintText(InContext, InTextEntry, GetDefaultFontColor()); }
 
+// Function for reading global TEXT string. 
+// The data and function are generated by the shader compiler.
+GENERATED_SHADER_PRINT
 
-// --------------------------------------------------------------------------
+///////////////////////////////////////////////////////////////////////////////////////////////////
 // Value printing (common value printing)
 
 // float
@@ -291,11 +401,11 @@ float2 ShaderPrintValue(in float2 ScreenPos, in bool Value, in FFontColor Color)
 	return ScreenPos;
 }
 
-// --------------------------------------------------------------------------
+///////////////////////////////////////////////////////////////////////////////////////////////////
 // Scalar type printing
 
-#define SHADER_PRINT_OVERLOAD_1(InNumComponent, InType) \
-float2 ShaderPrint(in float2 ScreenPos, in InType Value, in FFontColor Color) \
+#define SHADER_PRINT_OVERLOAD_1(InPrefix, InNumComponent, InType) \
+float2 InPrefix##Print(in float2 ScreenPos, in InType Value, in FFontColor Color) \
 { \
 	if (ShaderPrintFilterEnable) \
 	{ \
@@ -303,34 +413,47 @@ float2 ShaderPrint(in float2 ScreenPos, in InType Value, in FFontColor Color) \
 	} \
 	return ScreenPos; \
 } \
-float2 ShaderPrint(in float2 ScreenPos, in InType Value) \
+float2 InPrefix##Print(in float2 ScreenPos, in InType Value) \
 { \
 	if (ShaderPrintFilterEnable) \
 	{ \
-		ScreenPos = ShaderPrint(ScreenPos, Value, GetDefaultFontColor()); \
+		ScreenPos = ShaderPrintValue(ScreenPos, Value, GetDefaultFontColor()); \
 	} \
 	return ScreenPos; \
 } \
-void ShaderPrint(in InType Value, in FFontColor Color) \
+void InPrefix##Print(in InType Value, in FFontColor Color) \
 { \
 	if (ShaderPrintFilterEnable) \
 	{ \
-		ShaderPrintCursorPos = ShaderPrint(ShaderPrintCursorPos, Value, Color); \
+		ShaderPrintCursorPos = ShaderPrintValue(ShaderPrintCursorPos, Value, Color); \
 	} \
 } \
-void ShaderPrint(in InType Value) \
+void InPrefix##Print(in InType Value) \
 { \
 	if (ShaderPrintFilterEnable) \
 	{ \
-		ShaderPrintCursorPos = ShaderPrint(ShaderPrintCursorPos, Value, GetDefaultFontColor()); \
+		ShaderPrintCursorPos = ShaderPrintValue(ShaderPrintCursorPos, Value, GetDefaultFontColor()); \
+	} \
+}\
+void InPrefix##Print(inout FShaderPrintContext InContext, in InType Value, in FFontColor Color) \
+{ \
+	if (InContext.bIsActive) \
+	{ \
+		InContext.Pos = ShaderPrintValue(InContext.Pos, Value, Color); \
+	} \
+} \
+void InPrefix##Print(inout FShaderPrintContext InContext, in InType Value) \
+{ \
+	if (InContext.bIsActive) \
+	{ \
+		InContext.Pos = ShaderPrintValue(InContext.Pos, Value, GetDefaultFontColor()); \
 	} \
 }
-
-// --------------------------------------------------------------------------
+///////////////////////////////////////////////////////////////////////////////////////////////////
 // Vector type printing
 
-#define SHADER_PRINT_OVERLOAD_N(InNumComponent, InType) \
-float2 ShaderPrint(in float2 ScreenPos, in InType Value, in FFontColor Color) \
+#define SHADER_PRINT_OVERLOAD_N(InPrefix, InNumComponent, InType) \
+float2 InPrefix##Print(in float2 ScreenPos, in InType Value, in FFontColor Color) \
 { \
 	if (ShaderPrintFilterEnable) \
 	{ \
@@ -342,118 +465,189 @@ float2 ShaderPrint(in float2 ScreenPos, in InType Value, in FFontColor Color) \
 	} \
 	return ScreenPos; \
 } \
-float2 ShaderPrint(in float2 ScreenPos, in InType Value) \
+float2 InPrefix##Print(in float2 ScreenPos, in InType Value) \
 { \
 	if (ShaderPrintFilterEnable) \
 	{ \
 		UNROLL \
 		for (uint CompIt=0;CompIt<InNumComponent;++CompIt) \
 		{ \
-			ScreenPos = ShaderPrint(ScreenPos, Value[CompIt], GetDefaultFontColor()); \
+			ScreenPos = ShaderPrintValue(ScreenPos, Value[CompIt], GetDefaultFontColor()); \
 		} \
 	} \
 	return ScreenPos; \
 } \
-void ShaderPrint(in InType Value, in FFontColor Color) \
+void InPrefix##Print(in InType Value, in FFontColor Color) \
 { \
 	if (ShaderPrintFilterEnable) \
 	{ \
 		UNROLL \
 		for (uint CompIt=0;CompIt<InNumComponent;++CompIt) \
 		{ \
-			ShaderPrintCursorPos = ShaderPrint(ShaderPrintCursorPos, Value[CompIt], Color); \
+			ShaderPrintCursorPos = ShaderPrintValue(ShaderPrintCursorPos, Value[CompIt], Color); \
 		} \
 	} \
 } \
-void ShaderPrint(in InType Value) \
+void InPrefix##Print(in InType Value) \
 { \
 	if (ShaderPrintFilterEnable) \
 	{ \
 		UNROLL \
 		for (uint CompIt=0;CompIt<InNumComponent;++CompIt) \
 		{ \
-			ShaderPrintCursorPos = ShaderPrint(ShaderPrintCursorPos, Value[CompIt], GetDefaultFontColor()); \
+			ShaderPrintCursorPos = ShaderPrintValue(ShaderPrintCursorPos, Value[CompIt], GetDefaultFontColor()); \
+		} \
+	} \
+} \
+void InPrefix##Print(inout FShaderPrintContext InContext, in InType Value, in FFontColor Color) \
+{ \
+	if (InContext.bIsActive) \
+	{ \
+		UNROLL \
+		for (uint CompIt=0;CompIt<InNumComponent;++CompIt) \
+		{ \
+			InContext.Pos = ShaderPrintValue(InContext.Pos, Value[CompIt], Color); \
+		} \
+	} \
+} \
+void InPrefix##Print(inout FShaderPrintContext InContext, in InType Value) \
+{ \
+	if (InContext.bIsActive) \
+	{ \
+		UNROLL \
+		for (uint CompIt=0;CompIt<InNumComponent;++CompIt) \
+		{ \
+			InContext.Pos = ShaderPrintValue(InContext.Pos, Value[CompIt], GetDefaultFontColor()); \
 		} \
 	} \
 }
 
-SHADER_PRINT_OVERLOAD_1(1, float)
-SHADER_PRINT_OVERLOAD_N(2, float2)
-SHADER_PRINT_OVERLOAD_N(3, float3)
-SHADER_PRINT_OVERLOAD_N(4, float4)
+// ShaderPrint(X) version
+SHADER_PRINT_OVERLOAD_1(Shader, 1, float)
+SHADER_PRINT_OVERLOAD_N(Shader, 2, float2)
+SHADER_PRINT_OVERLOAD_N(Shader, 3, float3)
+SHADER_PRINT_OVERLOAD_N(Shader, 4, float4)
+
+SHADER_PRINT_OVERLOAD_1(Shader, 1, uint)
+SHADER_PRINT_OVERLOAD_N(Shader, 2, uint2)
+SHADER_PRINT_OVERLOAD_N(Shader, 3, uint3)
+SHADER_PRINT_OVERLOAD_N(Shader, 4, uint4)
+
+SHADER_PRINT_OVERLOAD_1(Shader, 1, int)
+SHADER_PRINT_OVERLOAD_N(Shader, 2, int2)
+SHADER_PRINT_OVERLOAD_N(Shader, 3, int3)
+SHADER_PRINT_OVERLOAD_N(Shader, 4, int4)
 
-SHADER_PRINT_OVERLOAD_1(1, uint)
-SHADER_PRINT_OVERLOAD_N(2, uint2)
-SHADER_PRINT_OVERLOAD_N(3, uint3)
-SHADER_PRINT_OVERLOAD_N(4, uint4)
+SHADER_PRINT_OVERLOAD_1(Shader, 1, bool)
 
-SHADER_PRINT_OVERLOAD_1(1, int)
-SHADER_PRINT_OVERLOAD_N(2, int2)
-SHADER_PRINT_OVERLOAD_N(3, int3)
-SHADER_PRINT_OVERLOAD_N(4, int4)
+// Print(X) version
+SHADER_PRINT_OVERLOAD_1(, 1, float)
+SHADER_PRINT_OVERLOAD_N(, 2, float2)
+SHADER_PRINT_OVERLOAD_N(, 3, float3)
+SHADER_PRINT_OVERLOAD_N(, 4, float4)
 
-SHADER_PRINT_OVERLOAD_1(1, bool)
+SHADER_PRINT_OVERLOAD_1(, 1, uint)
+SHADER_PRINT_OVERLOAD_N(, 2, uint2)
+SHADER_PRINT_OVERLOAD_N(, 3, uint3)
+SHADER_PRINT_OVERLOAD_N(, 4, uint4)
 
-// --------------------------------------------------------------------------
+SHADER_PRINT_OVERLOAD_1(, 1, int)
+SHADER_PRINT_OVERLOAD_N(, 2, int2)
+SHADER_PRINT_OVERLOAD_N(, 3, int3)
+SHADER_PRINT_OVERLOAD_N(, 4, int4)
+
+SHADER_PRINT_OVERLOAD_1(, 1, bool)
+
+///////////////////////////////////////////////////////////////////////////////////////////////////
 // Matrix type printing
 
-#define SHADER_PRINT_OVERLOAD_NN(InNumComponentX, InNumComponentY, InType) \
-float2 ShaderPrint(in float2 ScreenPos, in InType Value, in FFontColor Color) \
+#define SHADER_PRINT_OVERLOAD_NN(InPrefix, InNumComponentX, InNumComponentY, InType) \
+float2 InPrefix##Print(in float2 ScreenPos, in InType Value, in FFontColor Color) \
 { \
 	if (ShaderPrintFilterEnable) \
 	{ \
 		UNROLL \
 		for (uint CompY=0;CompY<InNumComponentY;++CompY) \
 		{ \
-			ScreenPos = ShaderPrint(ScreenPos, Value[CompY], Color); \
+			ScreenPos = InPrefix##Print(ScreenPos, Value[CompY], Color); \
 			ScreenPos += float2(0.f, ShaderPrintUniform.FontSpacing.y); \
 		} \
 	} \
 	return ScreenPos; \
 } \
-float2 ShaderPrint(in float2 ScreenPos, in InType Value) \
+float2 InPrefix##Print(in float2 ScreenPos, in InType Value) \
 { \
 	if (ShaderPrintFilterEnable) \
 	{ \
 		UNROLL \
 		for (uint CompY = 0; CompY < InNumComponentY; ++CompY) \
 		{ \
-			ScreenPos = ShaderPrint(ScreenPos, Value[CompY], GetDefaultFontColor()); \
+			ScreenPos = InPrefix##Print(ScreenPos, Value[CompY], GetDefaultFontColor()); \
 			ScreenPos += float2(0.f, ShaderPrintUniform.FontSpacing.y); \
 		} \
 	} \
 	return ScreenPos; \
 } \
-void ShaderPrint(in InType Value, in FFontColor Color) \
+void InPrefix##Print(in InType Value, in FFontColor Color) \
 { \
 	if (ShaderPrintFilterEnable) \
 	{ \
 		UNROLL \
 		for (uint CompY = 0; CompY < InNumComponentY; ++CompY) \
 		{ \
-			ShaderPrintCursorPos = ShaderPrint(ShaderPrintCursorPos, Value[CompY], Color); \
+			ShaderPrintCursorPos = InPrefix##Print(ShaderPrintCursorPos, Value[CompY], Color); \
 			ShaderPrintCursorPos += float2(0.f, ShaderPrintUniform.FontSpacing.y); \
 		} \
 	} \
 } \
-void ShaderPrint(in InType Value) \
+void InPrefix##Print(in InType Value) \
 { \
 	if (ShaderPrintFilterEnable) \
 	{ \
 		UNROLL \
 		for (uint CompY = 0; CompY < InNumComponentY; ++CompY) \
 		{ \
-			ShaderPrintCursorPos = ShaderPrint(ShaderPrintCursorPos, Value[CompY], GetDefaultFontColor()); \
+			ShaderPrintCursorPos = InPrefix##Print(ShaderPrintCursorPos, Value[CompY], GetDefaultFontColor()); \
 			ShaderPrintCursorPos += float2(0.f, ShaderPrintUniform.FontSpacing.y); \
 		} \
 	} \
-	}
+} \
+void InPrefix##Print(inout FShaderPrintContext InContext, in InType Value, in FFontColor Color) \
+{ \
+	if (InContext.bIsActive) \
+	{ \
+		UNROLL \
+		for (uint CompY = 0; CompY < InNumComponentY; ++CompY) \
+		{ \
+			InContext.Pos = InPrefix##Print(InContext.Pos, Value[CompY], Color); \
+			InContext.Pos += float2(0.f, ShaderPrintUniform.FontSpacing.y); \
+		} \
+	} \
+} \
+void InPrefix##Print(inout FShaderPrintContext InContext, in InType Value) \
+{ \
+	if (InContext.bIsActive) \
+	{ \
+		UNROLL \
+		for (uint CompY = 0; CompY < InNumComponentY; ++CompY) \
+		{ \
+			InContext.Pos = InPrefix##Print(InContext.Pos, Value[CompY], GetDefaultFontColor()); \
+			InContext.Pos += float2(0.f, ShaderPrintUniform.FontSpacing.y); \
+		} \
+	} \
+}
+
+// ShaderPrint(X) version
+SHADER_PRINT_OVERLOAD_NN(Shader, 3, 3, float3x3)
+SHADER_PRINT_OVERLOAD_NN(Shader, 4, 3, float4x3)
+SHADER_PRINT_OVERLOAD_NN(Shader, 4, 4, float4x4)
 
-SHADER_PRINT_OVERLOAD_NN(3, 3, float3x3)
-SHADER_PRINT_OVERLOAD_NN(4, 3, float4x3)
-SHADER_PRINT_OVERLOAD_NN(4, 4, float4x4)
+// Print(X) version
+SHADER_PRINT_OVERLOAD_NN(, 3, 3, float3x3)
+SHADER_PRINT_OVERLOAD_NN(, 4, 3, float4x3)
+SHADER_PRINT_OVERLOAD_NN(, 4, 4, float4x4)
 
-// --------------------------------------------------------------------------
+///////////////////////////////////////////////////////////////////////////////////////////////////
 // Formating helpers
 
 float2 ShaderPrintSpace(in float2 ScreenPos, in float Count = 1)
@@ -503,22 +697,8 @@ float2 ShaderPrintGetPos(in uint2 PixelCoord)
 	return float2(PixelCoord) / float2(ShaderPrintUniform.Resolution);
 }
 
-// --------------------------------------------------------------------------
+///////////////////////////////////////////////////////////////////////////////////////////////////
 // Filter helpers
-// Use these to restrict ShaderPrint to a subset of active threads
-// The filter functions can be called once at the start of the shader to filter all later printing
-//
-// Example use cases would be:
-//
-// ShaderPrintFilter(all(SvPosition.xy == float2(100, 100)));
-// to debug a single pixel
-//
-// ShaderPrintFilter(all(DispatchThreadId == uint3(0, 0, 0)));
-// to debug a single compute shader thread
-//
-// ShaderPrintFilterOneThread()
-// to debug a single random active thread
-// Note that ShaderPrintFilterOneThread() only works for one shader invocation in a view since it relies on a global memory location
 
 void ShaderPrintFilter(bool bFilter)
 {
@@ -546,29 +726,6 @@ uint2 GetSliderWidgetSize()
 	return uint2(10,1) * uint(ShaderPrintUniform.FontSpacing.x * ShaderPrintUniform.Resolution.x);
 }
 
-struct FShaderPrintContext
-{
-	bool bIsActive;
-	float2 StartPos;
-	float2 Pos;
-};
-
-FShaderPrintContext InitShaderPrintContext(bool bActive, float2 InStartPos)
-{
-	FShaderPrintContext Out;
-	Out.bIsActive = bActive;
-	Out.StartPos = InStartPos;
-	Out.Pos = InStartPos;
-	return Out;
-}
-FShaderPrintContext InitShaderPrintContext(bool bActive, uint2 InStartCoord) { return InitShaderPrintContext(bActive, float2(InStartCoord) / float2(ShaderPrintUniform.Resolution)); }
-
-void Newline(inout FShaderPrintContext InContext)
-{
-	InContext.Pos.x  = InContext.StartPos.x;
-	InContext.Pos.y += ShaderPrintUniform.FontSpacing.y;
-}
-
 struct FShaderPrintMetadata
 {
 	uint Hash;
@@ -612,7 +769,7 @@ bool ShaderPrintGetStateValue(in uint InHash, inout uint OutIndex, inout uint Ou
 }
 
 // Add a checkbox widget 
-bool AddCheckbox(inout FShaderPrintContext InContext, uint InTextEntry, bool bDefault, in FFontColor InColor)
+bool AddCheckbox(inout FShaderPrintContext InContext, FShaderPrintText InTextEntry, bool bDefault, in FFontColor InColor)
 {
 	uint Index = 0;
 	uint RawValue = 0;
@@ -654,10 +811,10 @@ bool AddCheckbox(inout FShaderPrintContext InContext, uint InTextEntry, bool bDe
 		E.Metadata = ShaderPrintPackMetadata(Hash, Index);
 		ShaderPrint_Internal(E);
 
-		InContext.Pos.x += (WidgetSize.x * ShaderPrintUniform.Resolution.x) + ShaderPrintUniform.FontSpacing.x;
+		InContext.Pos.x += float(WidgetSize.x) / float(ShaderPrintUniform.Resolution.x) + ShaderPrintUniform.FontSpacing.x;
 
 		// Text
-		float2 Pos = ShaderPrintText(InContext.Pos, InTextEntry, InColor);
+		InContext.Pos = ShaderPrintText(InContext.Pos, InTextEntry, InColor);
 		InContext.Pos.x += ShaderPrintUniform.FontSpacing.x;
 	}
 
@@ -665,7 +822,7 @@ bool AddCheckbox(inout FShaderPrintContext InContext, uint InTextEntry, bool bDe
 }
 
 // Add a slider widget  
-float AddSlider(inout FShaderPrintContext InContext, uint InTextEntry, float bDefault, in FFontColor InColor, float InMin = 0.f, float InMax = 1.f)
+float AddSlider(inout FShaderPrintContext InContext, FShaderPrintText InTextEntry, float bDefault, in FFontColor InColor, float InMin = 0.f, float InMax = 1.f)
 {
 	// The value is stored as normalized, for displaying slide value correctly. This is not idealy
 	// as it implies some precision lost due to back&forth conversion between normalized/non-normalized value.
@@ -702,10 +859,10 @@ float AddSlider(inout FShaderPrintContext InContext, uint InTextEntry, float bDe
 		E.Metadata = ShaderPrintPackMetadata(Hash, Index);
 		ShaderPrint_Internal(E);
 
-		InContext.Pos.x += (WidgetSize.x * ShaderPrintUniform.Resolution.x) + ShaderPrintUniform.FontSpacing.x;
+		InContext.Pos.x += float(WidgetSize.x) / float(ShaderPrintUniform.Resolution.x) + ShaderPrintUniform.FontSpacing.x;
 
 		// Text
-		float2 Pos = ShaderPrintText(InContext.Pos, InTextEntry, InColor);
+		InContext.Pos = ShaderPrintText(InContext.Pos, InTextEntry, InColor);
 		InContext.Pos.x += ShaderPrintUniform.FontSpacing.x;
 	}
 
diff --git a/Engine/Source/Developer/ShaderCompilerCommon/Private/ShaderCompilerCommon.cpp b/Engine/Source/Developer/ShaderCompilerCommon/Private/ShaderCompilerCommon.cpp
--- a/Engine/Source/Developer/ShaderCompilerCommon/Private/ShaderCompilerCommon.cpp
+++ b/Engine/Source/Developer/ShaderCompilerCommon/Private/ShaderCompilerCommon.cpp
@@ -1466,6 +1466,9 @@ void TransformStringIntoCharacterArray(FString& PreprocessedShaderSource)
 	// 3. Replace TEXT by its entry number
 	uint32 GlobalCount = 0;
 	{
+		const FString InitHashBegin(TEXT("InitShaderPrintText("));
+		const FString InitHashEnd(TEXT(")"));
+
 		const TCHAR* TextIdentifier = TEXT("TEXT(");
 		const TCHAR* SearchPtr = FCString::Strstr(&PreprocessedShaderSource[0], TextIdentifier);
 		while (SearchPtr)
@@ -1495,7 +1498,9 @@ void TransformStringIntoCharacterArray(FString& PreprocessedShaderSource)
 				const uint32 StartIndex = SearchPtr - StartPtr;
 				const uint32 CharCount = (EndPtr - SearchPtr) + 1;
 				PreprocessedShaderSource.RemoveAt(StartIndex, CharCount);
-				PreprocessedShaderSource.InsertAt(StartIndex, FString::FromInt(EntryIndex));
+
+				const FString HashText = InitHashBegin + FString::FromInt(EntryIndex) + InitHashEnd;
+				PreprocessedShaderSource.InsertAt(StartIndex, HashText);
 
 				// Update SearchPtr, as PreprocessedShaderSource has been modified, and its memory could have been reallocated, causing SearchPtr to be invalid.
 				SearchPtr = &PreprocessedShaderSource[0] + StartIndex;
@@ -1536,15 +1541,15 @@ void TransformStringIntoCharacterArray(FString& PreprocessedShaderSource)
 		}
 		TextChars += TEXT("};\n\n");
 
-		TextChars += TEXT("uint ShaderPrintGetChar(uint InIndex)       { return TEXT_CHARS[InIndex]; }\n");
-		TextChars += TEXT("uint ShaderPrintGetOffset(uint InTextEntry) { return TEXT_OFFSETS[InTextEntry]; }\n");
-		TextChars += TEXT("uint ShaderPrintGetHash(uint InTextEntry)   { return TEXT_HASHES[InTextEntry]; }\n");
+		TextChars += TEXT("uint ShaderPrintGetChar(uint InIndex)              { return TEXT_CHARS[InIndex]; }\n");
+		TextChars += TEXT("uint ShaderPrintGetOffset(FShaderPrintText InText) { return TEXT_OFFSETS[InText.Index]; }\n");
+		TextChars += TEXT("uint ShaderPrintGetHash(FShaderPrintText InText)   { return TEXT_HASHES[InText.Index]; }\n");
 	}
 	else
 	{	
-		TextChars += TEXT("uint ShaderPrintGetChar(uint Index) { return 0; }\n");
-		TextChars += TEXT("uint ShaderPrintGetOffset(uint InTextEntry) { return 0; }\n");
-		TextChars += TEXT("uint ShaderPrintGetHash(uint InTextEntry) { return 0; }\n");
+		TextChars += TEXT("uint ShaderPrintGetChar(uint Index)                { return 0; }\n");
+		TextChars += TEXT("uint ShaderPrintGetOffset(FShaderPrintText InText) { return 0; }\n");
+		TextChars += TEXT("uint ShaderPrintGetHash(FShaderPrintText InText)   { return 0; }\n");
 	}
 	
 	// 6. Insert global struct data + print function
