Use Basic Eye Adaptation on mobile by default. Saves 0.7ms on S10
r.EyeAdaptationQuality now controls which method can be used, previously it was doing nothing

diff --git a/Engine/Config/Android/AndroidScalability.ini b/Engine/Config/Android/AndroidScalability.ini
--- a/Engine/Config/Android/AndroidScalability.ini
+++ b/Engine/Config/Android/AndroidScalability.ini
@@ -82,6 +82,7 @@ r.Tonemapper.Quality=1
 r.BloomQuality=0
 r.DepthOfFieldQuality=0
 r.LightShaftQuality=0
+r.EyeAdaptationQuality=1
 r.Mobile.AmbientOcclusionQuality=0
 r.Mobile.PixelProjectedReflectionQuality=0
 
@@ -90,6 +91,7 @@ r.Tonemapper.Quality=2
 r.BloomQuality=1
 r.DepthOfFieldQuality=1
 r.LightShaftQuality=1
+r.EyeAdaptationQuality=1
 r.Mobile.AmbientOcclusionQuality=0
 r.Mobile.PixelProjectedReflectionQuality=0
 
@@ -98,6 +100,7 @@ r.Tonemapper.Quality=3
 r.BloomQuality=1
 r.DepthOfFieldQuality=1
 r.LightShaftQuality=1
+r.EyeAdaptationQuality=1
 r.Mobile.AmbientOcclusionQuality=1
 r.Mobile.PixelProjectedReflectionQuality=1
 
@@ -106,6 +109,7 @@ r.Tonemapper.Quality=3
 r.BloomQuality=1
 r.DepthOfFieldQuality=1
 r.LightShaftQuality=1
+r.EyeAdaptationQuality=1
 r.Mobile.AmbientOcclusionQuality=1
 r.Mobile.PixelProjectedReflectionQuality=1
 
diff --git a/Engine/Config/IOS/IOSScalability.ini b/Engine/Config/IOS/IOSScalability.ini
--- a/Engine/Config/IOS/IOSScalability.ini
+++ b/Engine/Config/IOS/IOSScalability.ini
@@ -82,6 +82,7 @@ r.Tonemapper.Quality=1
 r.BloomQuality=0
 r.DepthOfFieldQuality=0
 r.LightShaftQuality=0
+r.EyeAdaptationQuality=1
 r.Mobile.AmbientOcclusionQuality=0
 r.Mobile.PixelProjectedReflectionQuality=0
 
@@ -90,6 +91,7 @@ r.Tonemapper.Quality=2
 r.BloomQuality=1
 r.DepthOfFieldQuality=0
 r.LightShaftQuality=0
+r.EyeAdaptationQuality=1
 r.Mobile.AmbientOcclusionQuality=0
 r.Mobile.PixelProjectedReflectionQuality=0
 
@@ -98,6 +100,7 @@ r.Tonemapper.Quality=3
 r.BloomQuality=1
 r.DepthOfFieldQuality=0
 r.LightShaftQuality=0
+r.EyeAdaptationQuality=1
 r.Mobile.AmbientOcclusionQuality=0
 r.Mobile.PixelProjectedReflectionQuality=0
 
@@ -106,6 +109,7 @@ r.Tonemapper.Quality=3
 r.BloomQuality=1
 r.DepthOfFieldQuality=1
 r.LightShaftQuality=1
+r.EyeAdaptationQuality=1
 r.Mobile.AmbientOcclusionQuality=1
 r.Mobile.PixelProjectedReflectionQuality=1
 
diff --git a/Engine/Shaders/Private/PostProcessMobile.usf b/Engine/Shaders/Private/PostProcessMobile.usf
--- a/Engine/Shaders/Private/PostProcessMobile.usf
+++ b/Engine/Shaders/Private/PostProcessMobile.usf
@@ -361,10 +361,14 @@ void BloomPS_Mobile(
 #endif
 
 #if MOBILE_USEEYEADAPTATION
-
-	const half Intensity = max(dot(AverageColor, half3(1.0f, 1.0f, 1.0f) / 3.0f), EyeAdaptation_LuminanceMin);
+	half Intensity = dot(AverageColor, half3(1.0f, 1.0f, 1.0f) / 3.0f);
+	if (Intensity == 0.0)
+	{
+		// Override intensity for black pixels to 0.1 This behaviour is controlled by the cvar "r.EyeAdaptation.BlackHistogramBucketInfluence"
+		Intensity = (0.1 - EyeAdaptation_BlackHistogramBucketInfluence);
+	}
+	Intensity = max(Intensity, EyeAdaptation_LuminanceMin);
 	const float LogIntensity = clamp(log2(Intensity), -10.0f, 20.0f);
-
 	// Store log intensity in the alpha channel: scale to 0,1 range.
 	OutEyeAdaptation = EyeAdaptation_HistogramScale * LogIntensity + EyeAdaptation_HistogramBias;
 #endif
@@ -1371,10 +1375,11 @@ void Histogram_MainCS(
 				// case is so the black background pixels in the editor have no effect. But if we have cases where
 				// pixel values can actually be black, we want to set EyeAdaptation_LastHistogramBucketInfluence=1.0.
 				// This value is controlled by the cvar "r.EyeAdaptation.BlackHistogramBucketInfluence"
-				if (Bucket0 == 0)
-				{
-					Weight0 *= EyeAdaptation_BlackHistogramBucketInfluence;
-				}
+				// THIS IS DONE in a Bloom setup, black pixels will set intencity to 0.1
+				//if (Bucket0 == 0)
+				//{
+				//	Weight0 *= EyeAdaptation_BlackHistogramBucketInfluence;
+				//}
 
 				// Accumulate the weight to the nearby history buckets.
 #if IOS // The IOS A8 and lower devices seems don't support using a float4 array as a two dimension array, separate the operations to two steps.
diff --git a/Engine/Source/Runtime/Renderer/Private/PostProcess/PostProcessEyeAdaptation.cpp b/Engine/Source/Runtime/Renderer/Private/PostProcess/PostProcessEyeAdaptation.cpp
--- a/Engine/Source/Runtime/Renderer/Private/PostProcess/PostProcessEyeAdaptation.cpp
+++ b/Engine/Source/Runtime/Renderer/Private/PostProcess/PostProcessEyeAdaptation.cpp
@@ -93,6 +93,23 @@ bool IsExtendLuminanceRangeEnabled()
 	return VarDefaultAutoExposureExtendDefaultLuminanceRange->GetValueOnRenderThread() == 1;
 }
 
+static EAutoExposureMethod ApplyEyeAdaptationQuality(EAutoExposureMethod AutoExposureMethod)
+{
+	static const auto CVarEyeAdaptationQuality = IConsoleManager::Get().FindTConsoleVariableDataInt(TEXT("r.EyeAdaptationQuality"));
+	const int32 EyeAdaptationQuality = CVarEyeAdaptationQuality->GetValueOnRenderThread();
+	
+	if (AutoExposureMethod != EAutoExposureMethod::AEM_Manual)
+	{
+		if (EyeAdaptationQuality == 1)
+		{
+			// Clamp current method to AEM_Basic
+			AutoExposureMethod = EAutoExposureMethod::AEM_Basic;
+		}
+	}
+
+	return AutoExposureMethod;
+}
+
 float LuminanceMaxFromLensAttenuation()
 {
 	const bool bExtendedLuminanceRange = IsExtendLuminanceRangeEnabled();
@@ -119,6 +136,9 @@ EAutoExposureMethod GetAutoExposureMethod(const FViewInfo& View)
 		AutoExposureMethod = IsAutoExposureMethodSupported(View.GetFeatureLevel(), EAutoExposureMethod::AEM_Basic) ? EAutoExposureMethod::AEM_Basic : EAutoExposureMethod::AEM_Manual;
 	}
 
+	// Apply quality settings
+	AutoExposureMethod = ApplyEyeAdaptationQuality(AutoExposureMethod);
+
 	const int32 EyeOverride = CVarEyeAdaptationMethodOverride.GetValueOnRenderThread();
 
 	EAutoExposureMethod OverrideAutoExposureMethod = AutoExposureMethod;
