Add MaxDigit and MaxDecimal formating for shader Print()

diff --git a/Engine/Shaders/Private/HairStrands/HairStrandsDebugPrint.usf b/Engine/Shaders/Private/HairStrands/HairStrandsDebugPrint.usf
--- a/Engine/Shaders/Private/HairStrands/HairStrandsDebugPrint.usf
+++ b/Engine/Shaders/Private/HairStrands/HairStrandsDebugPrint.usf
@@ -354,21 +354,6 @@ float2 PrintInstanceName(float2 InPos, uint PrimitiveID, FFontColor InColor)
 	return InPos;
 }
 
-
-float2 Print(float2 Pos, uint X, uint DigitCount, FFontColor Color)
-{
-	Pos = ShaderPrint(Pos, X, Color);
-	Pos.x -= ShaderPrintUniform.FontSpacing.x * (ShaderPrintTabCount - min(ShaderPrintTabCount, DigitCount));
-	return Pos;
-}
-
-float2 Print(float2 Pos, float X, uint DigitCount, FFontColor Color)
-{
-	Pos = ShaderPrint(Pos, X, Color);
-	Pos.x -= ShaderPrintUniform.FontSpacing.x * (ShaderPrintTabCount - min(ShaderPrintTabCount, DigitCount));
-	return Pos;
-}
-
 [numthreads(1, 1, 1)]
 void MainCS(uint GroupIndex : SV_GroupIndex, uint3 DispatchThreadId : SV_DispatchThreadID)
 #if PERMUTATION_OUTPUT_TYPE == 0
@@ -396,21 +381,21 @@ void MainCS(uint GroupIndex : SV_GroupIndex, uint3 DispatchThreadId : SV_Dispatc
 	}
 
 	Pos = ShaderPrintText(Pos, TEXT("Total Groups        : "), FontWhite);
-	Pos = Print(Pos, TotalGroup, 8, FontYellow);
+	Pos = Print(Pos, TotalGroup, FontYellow, 8);
 	AddNewLine(OriginalPos, Pos);
 
 	Pos = ShaderPrintText(Pos, TEXT("Total Strand Curves : "), FontWhite);
-	Pos = Print(Pos, TotalStrandsCount, 8, FontOrange);
+	Pos = Print(Pos, TotalStrandsCount, FontOrange, 8);
 	AddNewLine(OriginalPos, Pos);
 
 	Pos = ShaderPrintText(Pos, TEXT("Total Strand Vertex : "), FontWhite);
-	Pos = Print(Pos, TotalVertexCount, 8, FontSilver);
+	Pos = Print(Pos, TotalVertexCount, FontSilver, 8);
 	AddNewLine(OriginalPos, Pos);
 
 	AddNewLine(OriginalPos, Pos);
 
 	// Draw instance details
-	Pos = ShaderPrintText(Pos, TEXT("   Group  LOD        Geometry Binding Sim   RBF   Clip    Strands Vertex  Name"));
+	Pos = ShaderPrintText(Pos, TEXT("   Group  LOD     Geometry Binding Sim   RBF   Clip    Strands Vertex  Name"));
 	AddNewLine(OriginalPos, Pos);
 
 	for (uint InstanceIndex = 0; InstanceIndex < InstanceCount; ++InstanceIndex)
@@ -431,19 +416,19 @@ void MainCS(uint GroupIndex : SV_GroupIndex, uint3 DispatchThreadId : SV_Dispatc
 		const uint VertexCount		= InstanceData.w;
 
 		// Index
-		Pos = Print(Pos, InstanceIndex, 2, FontSilver);
+		Pos = Print(Pos, InstanceIndex, FontSilver, 2);
 		Pos = ShaderPrintText(Pos, TEXT(" "));
 
 		// Group
-		Pos = Print(Pos, GroupIndex, 1, FontYellow);
+		Pos = Print(Pos, GroupIndex, FontYellow, 1);
 		Pos = ShaderPrintText(Pos, TEXT("/"));
-		Pos = Print(Pos, GroupCount, 1, FontYellow);
+		Pos = Print(Pos, GroupCount, FontYellow, 1);
 		Pos = ShaderPrintText(Pos, TEXT("    "));
 
 		// LOD
-		Pos = Print(Pos, LOD, 7, FontYellow);	
+		Pos = Print(Pos, LOD, FontYellow, 4, 2);	
 		Pos = ShaderPrintText(Pos, TEXT("/"));
-		Pos = Print(Pos, LODCount, 1, FontYellow);
+		Pos = Print(Pos, LODCount, FontYellow, 1);
 		Pos = ShaderPrintText(Pos, TEXT("  "));
 
 		// Geometry
@@ -473,14 +458,14 @@ void MainCS(uint GroupIndex : SV_GroupIndex, uint3 DispatchThreadId : SV_Dispatc
 
 		// Clip
 		if (ClipValue >= 0)
-			Pos = Print(Pos, ClipValue, 7, FontGreen);
+			Pos = Print(Pos, ClipValue, FontGreen, 7);
 		else 
-			Pos = Print(Pos, 0u, 7, FontSilver);
+			Pos = Print(Pos, 0u, FontSilver, 7);
 		Pos = ShaderPrintText(Pos, TEXT(" "));
 
 		// Count
-		Pos = Print(Pos, StrandsCount, 8, FontOrange);
-		Pos = Print(Pos, VertexCount,  8, FontSilver);
+		Pos = Print(Pos, StrandsCount, FontOrange, 8);
+		Pos = Print(Pos, VertexCount,  FontSilver, 8);
 
 		// Name
 		Pos = PrintInstanceName(Pos, InstanceIndex, FontWhite);
diff --git a/Engine/Shaders/Private/ShaderPrintCommon.ush b/Engine/Shaders/Private/ShaderPrintCommon.ush
--- a/Engine/Shaders/Private/ShaderPrintCommon.ush
+++ b/Engine/Shaders/Private/ShaderPrintCommon.ush
@@ -32,6 +32,23 @@
 //   Print(Context, MyFloatInYellow, FontYellow);
 //   Newline(Context);
 // 
+// Numerical values can be further formatted, by configuring the number of digits spanned by a number 
+// (i.e. the number of written&non-written digits), and the number of wanted decimal for floating value.
+// 
+//   Print(Context, 123, FontYellow, 5); // Spanning is show as: .....
+//   Print(Context, 456, FontYellow, 4); // Spanning is show as: xxxx
+//   Print(Context, 789, FontYellow, 3); // Spanning is show as: ---
+//
+// Will be displayed as follow:
+//   123   456  789
+//   ..... xxxx ---
+// 
+// Print(Context, 12.3456, FontYellow, 5, 2); // Spanning is show as: ...
+// Print(Context, 12.3456, FontYellow, 5, 1); // Spanning is show as: xxx
+// 
+// Will be displayed as follow:
+//   12.34 12.3
+//   ..... xxxxx
 // 
 // 2. How to add widget?
 // ---------------------
@@ -143,10 +160,9 @@ struct FFontColor
 FFontColor InitFontColor(float InX, float InY, float InZ)	{ FFontColor Out; Out.Color = float3(InX, InY, InZ); return Out; }
 FFontColor InitFontColor(float3 In)							{ FFontColor Out; Out.Color = In; return Out; }
 FFontColor GetDefaultFontColor()							{ FFontColor Out; Out.Color = float3(1,1,1); return Out; }
+#define DEFAULTFONTCOLOR ((FFontColor)1)
 
-// Certain shader compiler does not support static initialization with complex type. 
-// In the meantime, using define to provide default color initialization
-#if 1 
+// Predefined colors
 #define FontWhite		InitFontColor(1, 1, 1)
 #define FontBlack		InitFontColor(0, 0, 0)
 #define FontRed			InitFontColor(1, 0, 0)
@@ -160,21 +176,6 @@ FFontColor GetDefaultFontColor()							{ FFontColor Out; Out.Color = float3(1,1,
 #define FontTurquoise	InitFontColor(26.f / 255.f, 188.f / 255.f, 156.f / 255.f)
 #define FontSilver		InitFontColor(189.f / 255.f, 195.f / 255.f, 199.f / 255.f)
 #define FontEmerald		InitFontColor(46.f / 255.f, 204.f / 255.f, 113.f / 255.f)
-#else
-static const FFontColor FontWhite		= InitFontColor(1, 1, 1);
-static const FFontColor FontBlack		= InitFontColor(0, 0, 0);
-static const FFontColor FontRed			= InitFontColor(1, 0, 0);
-static const FFontColor FontGreen		= InitFontColor(0, 1, 0);
-static const FFontColor FontBlue		= InitFontColor(0, 0, 1);
-static const FFontColor FontYellow		= InitFontColor(1, 1, 0);
-static const FFontColor FontCyan		= InitFontColor(0, 1, 1);
-static const FFontColor FontMagenta		= InitFontColor(1, 0, 1);
-static const FFontColor FontOrange		= InitFontColor(243.f / 255.f, 156.f / 255.f,  18.f / 255.f);
-static const FFontColor FontPurple		= InitFontColor(169.f /255.f,    7.f / 255.f, 228.f / 255.f);
-static const FFontColor FontTurquoise	= InitFontColor( 26.f / 255.f, 188.f / 255.f, 156.f / 255.f);
-static const FFontColor FontSilver		= InitFontColor(189.f / 255.f, 195.f / 255.f, 199.f / 255.f);
-static const FFontColor FontEmerald		= InitFontColor( 46.f / 255.f, 204.f / 255.f, 113.f / 255.f);
-#endif
 
 ///////////////////////////////////////////////////////////////////////////////////////////////////
 // Util pack/unpac functions
@@ -213,7 +214,6 @@ FShaderPrintItem UnpackShaderPrintItem(FPackedShaderPrintItem In)
 #include "/Engine/Generated/UniformBuffers/ShaderPrintUniform.ush"
 
 static const float2 ShaderPrintCursorStart = float2(0.05f, 0.05f);
-static const float ShaderPrintTabCount = 12.f;
 
 // 'Global' values for tracking printing state (per thread)
 static float2 ShaderPrintCursorPos = ShaderPrintCursorStart;
@@ -271,45 +271,35 @@ void ShaderPrint_Internal(in FShaderPrintItem Item)
 	RWValuesBuffer[IndexToStore + SHADER_PRINT_VALUE_OFFSET] = PackShaderPrintItem(Item);
 }
 
-void ShaderPrint_Internal(in float2 ScreenPos, in int Value, in FFontColor FontColor, in int Type)
+void ShaderPrint_Internal(in float2 ScreenPos, in int Value, in FFontColor FontColor, in uint Metadata, in int Type)
 {
 	FShaderPrintItem Item;
 	Item.ScreenPos = ScreenPos;
 	Item.Value = Value;
 	Item.Type = Type;
 	Item.Color = FontColor.Color;
-	Item.Metadata = 0u;
+	Item.Metadata = Metadata;
 	ShaderPrint_Internal(Item);
 }
 
 ///////////////////////////////////////////////////////////////////////////////////////////////////
 // Symbol printing
 
-float2 ShaderPrintSymbol(in float2 ScreenPos, in int Symbol, in FFontColor Color)
+float2 ShaderPrintSymbol(in float2 ScreenPos, in int Symbol, in FFontColor Color = DEFAULTFONTCOLOR)
 {
 	if (ShaderPrintFilterEnable)
 	{
-		ShaderPrint_Internal(ScreenPos, Symbol, Color, SHADER_PRINT_TYPE_SYMBOL);
+		ShaderPrint_Internal(ScreenPos, Symbol, Color, 0u, SHADER_PRINT_TYPE_SYMBOL);
 		ScreenPos.x += ShaderPrintUniform.FontSpacing.x;
 	}
 	return ScreenPos;
 }
 
-float2 ShaderPrintSymbol(in float2 ScreenPos, in int Symbol)
-{
-	return ShaderPrintSymbol(ScreenPos, Symbol, GetDefaultFontColor());
-}
-
-void ShaderPrintSymbol(in int symbol, in FFontColor Color)
+void ShaderPrintSymbol(in int symbol, in FFontColor Color = DEFAULTFONTCOLOR)
 {
 	ShaderPrintCursorPos = ShaderPrintSymbol(ShaderPrintCursorPos, symbol, Color);
 }
 
-void ShaderPrintSymbol(in int symbol)
-{
-	ShaderPrintCursorPos = ShaderPrintSymbol(ShaderPrintCursorPos, symbol, GetDefaultFontColor());
-}
-
 ///////////////////////////////////////////////////////////////////////////////////////////////////
 // Text & hash
 
@@ -331,7 +321,7 @@ uint ShaderPrintGetOffset(FShaderPrintText InTextEntry);
 uint ShaderPrintGetHash(FShaderPrintText InTextEntry);
  
 // Function for reading global TEXT string
-float2 ShaderPrintText(float2 Pos, FShaderPrintText InTextEntry, FFontColor InColor)
+float2 ShaderPrintText(float2 Pos, FShaderPrintText InTextEntry, FFontColor InColor = DEFAULTFONTCOLOR)
 {
 	uint Begin = ShaderPrintGetOffset(InTextEntry);
 	uint End = ShaderPrintGetOffset(InitShaderPrintText(InTextEntry.Index + 1));
@@ -341,9 +331,8 @@ float2 ShaderPrintText(float2 Pos, FShaderPrintText InTextEntry, FFontColor InCo
 	}
 	return Pos;
 }
-float2 ShaderPrintText(float2 Pos, FShaderPrintText InTextEntry) { return ShaderPrintText(Pos, InTextEntry, GetDefaultFontColor()); }
 
-void ShaderPrintText(inout FShaderPrintContext InContext, FShaderPrintText InTextEntry, FFontColor InColor)
+void ShaderPrintText(inout FShaderPrintContext InContext, FShaderPrintText InTextEntry, FFontColor InColor = DEFAULTFONTCOLOR)
 {
 	if (InContext.bIsActive)
 	{
@@ -355,12 +344,9 @@ void ShaderPrintText(inout FShaderPrintContext InContext, FShaderPrintText InTex
 		}	
 	}
 }
-void ShaderPrintText(inout FShaderPrintContext InContext, FShaderPrintText InTextEntry) { ShaderPrintText(InContext, InTextEntry, GetDefaultFontColor()); }
 
-float2 Print(float2 Pos, FShaderPrintText InTextEntry, FFontColor InColor)							{ return ShaderPrintText(Pos, InTextEntry, InColor); }
-float2 Print(float2 Pos, FShaderPrintText InTextEntry)												{ return ShaderPrintText(Pos, InTextEntry, GetDefaultFontColor()); }
-void   Print(inout FShaderPrintContext InContext, FShaderPrintText InTextEntry, FFontColor InColor)	{ ShaderPrintText(InContext, InTextEntry, InColor); }
-void   Print(inout FShaderPrintContext InContext, FShaderPrintText InTextEntry)						{ ShaderPrintText(InContext, InTextEntry, GetDefaultFontColor()); }
+float2 Print(float2 Pos, FShaderPrintText InTextEntry, FFontColor InColor=DEFAULTFONTCOLOR)                         { return ShaderPrintText(Pos, InTextEntry, InColor); }
+void   Print(inout FShaderPrintContext InCtx, FShaderPrintText InTextEntry, FFontColor InColor = DEFAULTFONTCOLOR)	{ ShaderPrintText(InCtx, InTextEntry, InColor); }
 
 // Function for reading global TEXT string. 
 // The data and function are generated by the shader compiler.
@@ -368,161 +354,31 @@ GENERATED_SHADER_PRINT
 
 ///////////////////////////////////////////////////////////////////////////////////////////////////
 // Value printing (common value printing)
+#define MAX_DECIMAL_COUNT 5u
+#define MAX_DIGIT_COUNT 12u
 
-// float
-float2 ShaderPrintValue(in float2 ScreenPos, in float Value, in FFontColor Color)
-{
-	ShaderPrint_Internal(ScreenPos, asint(Value), Color, SHADER_PRINT_TYPE_FLOAT);
-	ScreenPos.x += ShaderPrintUniform.FontSpacing.x * ShaderPrintTabCount;
-	return ScreenPos;
-}
-
-// int
-float2 ShaderPrintValue(in float2 ScreenPos, in int Value, in FFontColor Color)
-{
-	ShaderPrint_Internal(ScreenPos, Value, Color, SHADER_PRINT_TYPE_INT);
-	ScreenPos.x += ShaderPrintUniform.FontSpacing.x * ShaderPrintTabCount;
-	return ScreenPos;
-}
-
-// uint
-float2 ShaderPrintValue(in float2 ScreenPos, in uint Value, in FFontColor Color)
-{
-	ShaderPrint_Internal(ScreenPos, asint(Value), Color, SHADER_PRINT_TYPE_UINT);
-	ScreenPos.x += ShaderPrintUniform.FontSpacing.x * ShaderPrintTabCount;
-	return ScreenPos;
-}
-
-// bool
-float2 ShaderPrintValue(in float2 ScreenPos, in bool Value, in FFontColor Color)
-{
-	ShaderPrint_Internal(ScreenPos, asint(Value ? 1u : 0u), Color, SHADER_PRINT_TYPE_UINT);
-	ScreenPos.x += ShaderPrintUniform.FontSpacing.x * ShaderPrintTabCount;
-	return ScreenPos;
-}
+float2 ShaderPrintValue(float2 ScreenPos, float Value, FFontColor Color, uint MaxDigit, uint MaxDecimal)  { ShaderPrint_Internal(ScreenPos, asint(Value),           Color, MaxDecimal, SHADER_PRINT_TYPE_FLOAT); ScreenPos.x += ShaderPrintUniform.FontSpacing.x * MaxDigit; return ScreenPos; }
+float2 ShaderPrintValue(float2 ScreenPos, int   Value, FFontColor Color, uint MaxDigit, uint MaxDecimal)  { ShaderPrint_Internal(ScreenPos, Value,                  Color, MaxDecimal, SHADER_PRINT_TYPE_INT);   ScreenPos.x += ShaderPrintUniform.FontSpacing.x * MaxDigit; return ScreenPos; }
+float2 ShaderPrintValue(float2 ScreenPos, uint  Value, FFontColor Color, uint MaxDigit, uint MaxDecimal)  { ShaderPrint_Internal(ScreenPos, asint(Value),           Color, MaxDecimal, SHADER_PRINT_TYPE_UINT);  ScreenPos.x += ShaderPrintUniform.FontSpacing.x * MaxDigit; return ScreenPos; }
+float2 ShaderPrintValue(float2 ScreenPos, bool  Value, FFontColor Color, uint MaxDigit, uint MaxDecimal)  { ShaderPrint_Internal(ScreenPos, asint(Value ? 1u : 0u), Color, MaxDecimal, SHADER_PRINT_TYPE_UINT);  ScreenPos.x += ShaderPrintUniform.FontSpacing.x * MaxDigit; return ScreenPos; }
 
 ///////////////////////////////////////////////////////////////////////////////////////////////////
 // Scalar type printing
 
 #define SHADER_PRINT_OVERLOAD_1(InPrefix, InNumComponent, InType) \
-float2 InPrefix##Print(in float2 ScreenPos, in InType Value, in FFontColor Color) \
-{ \
-	if (ShaderPrintFilterEnable) \
-	{ \
-		ScreenPos = ShaderPrintValue(ScreenPos, Value, Color); \
-	} \
-	return ScreenPos; \
-} \
-float2 InPrefix##Print(in float2 ScreenPos, in InType Value) \
-{ \
-	if (ShaderPrintFilterEnable) \
-	{ \
-		ScreenPos = ShaderPrintValue(ScreenPos, Value, GetDefaultFontColor()); \
-	} \
-	return ScreenPos; \
-} \
-void InPrefix##Print(in InType Value, in FFontColor Color) \
-{ \
-	if (ShaderPrintFilterEnable) \
-	{ \
-		ShaderPrintCursorPos = ShaderPrintValue(ShaderPrintCursorPos, Value, Color); \
-	} \
-} \
-void InPrefix##Print(in InType Value) \
-{ \
-	if (ShaderPrintFilterEnable) \
-	{ \
-		ShaderPrintCursorPos = ShaderPrintValue(ShaderPrintCursorPos, Value, GetDefaultFontColor()); \
-	} \
-}\
-void InPrefix##Print(inout FShaderPrintContext InContext, in InType Value, in FFontColor Color) \
-{ \
-	if (InContext.bIsActive) \
-	{ \
-		InContext.Pos = ShaderPrintValue(InContext.Pos, Value, Color); \
-	} \
-} \
-void InPrefix##Print(inout FShaderPrintContext InContext, in InType Value) \
-{ \
-	if (InContext.bIsActive) \
-	{ \
-		InContext.Pos = ShaderPrintValue(InContext.Pos, Value, GetDefaultFontColor()); \
-	} \
-}
+void   InPrefix##Print(InType Value, FFontColor Color=DEFAULTFONTCOLOR, uint MaxDigit=MAX_DIGIT_COUNT, uint MaxDecimal=MAX_DECIMAL_COUNT)                                  { if (ShaderPrintFilterEnable) { ShaderPrintCursorPos = ShaderPrintValue(ShaderPrintCursorPos, Value, Color, MaxDigit, MaxDecimal); } } \
+float2 InPrefix##Print(float2 ScreenPos, InType Value, FFontColor Color=DEFAULTFONTCOLOR, uint MaxDigit=MAX_DIGIT_COUNT, uint MaxDecimal=MAX_DECIMAL_COUNT)                { if (ShaderPrintFilterEnable) { ScreenPos            = ShaderPrintValue(ScreenPos,            Value, Color, MaxDigit, MaxDecimal); } return ScreenPos; } \
+void   InPrefix##Print(inout FShaderPrintContext InCtx, InType Value, FFontColor Color=DEFAULTFONTCOLOR, uint MaxDigit=MAX_DIGIT_COUNT, uint MaxDecimal=MAX_DECIMAL_COUNT) { if (InCtx.bIsActive)         { InCtx.Pos            = ShaderPrintValue(InCtx.Pos,            Value, Color, MaxDigit, MaxDecimal); } }
+
 ///////////////////////////////////////////////////////////////////////////////////////////////////
 // Vector type printing
 
 #define SHADER_PRINT_OVERLOAD_N(InPrefix, InNumComponent, InType) \
-float2 InPrefix##Print(in float2 ScreenPos, in InType Value, in FFontColor Color) \
-{ \
-	if (ShaderPrintFilterEnable) \
-	{ \
-		UNROLL \
-		for (uint CompIt = 0; CompIt < InNumComponent; ++CompIt) \
-		{ \
-			ScreenPos = ShaderPrintValue(ScreenPos, Value[CompIt], Color); \
-		} \
-	} \
-	return ScreenPos; \
-} \
-float2 InPrefix##Print(in float2 ScreenPos, in InType Value) \
-{ \
-	if (ShaderPrintFilterEnable) \
-	{ \
-		UNROLL \
-		for (uint CompIt=0;CompIt<InNumComponent;++CompIt) \
-		{ \
-			ScreenPos = ShaderPrintValue(ScreenPos, Value[CompIt], GetDefaultFontColor()); \
-		} \
-	} \
-	return ScreenPos; \
-} \
-void InPrefix##Print(in InType Value, in FFontColor Color) \
-{ \
-	if (ShaderPrintFilterEnable) \
-	{ \
-		UNROLL \
-		for (uint CompIt=0;CompIt<InNumComponent;++CompIt) \
-		{ \
-			ShaderPrintCursorPos = ShaderPrintValue(ShaderPrintCursorPos, Value[CompIt], Color); \
-		} \
-	} \
-} \
-void InPrefix##Print(in InType Value) \
-{ \
-	if (ShaderPrintFilterEnable) \
-	{ \
-		UNROLL \
-		for (uint CompIt=0;CompIt<InNumComponent;++CompIt) \
-		{ \
-			ShaderPrintCursorPos = ShaderPrintValue(ShaderPrintCursorPos, Value[CompIt], GetDefaultFontColor()); \
-		} \
-	} \
-} \
-void InPrefix##Print(inout FShaderPrintContext InContext, in InType Value, in FFontColor Color) \
-{ \
-	if (InContext.bIsActive) \
-	{ \
-		UNROLL \
-		for (uint CompIt=0;CompIt<InNumComponent;++CompIt) \
-		{ \
-			InContext.Pos = ShaderPrintValue(InContext.Pos, Value[CompIt], Color); \
-		} \
-	} \
-} \
-void InPrefix##Print(inout FShaderPrintContext InContext, in InType Value) \
-{ \
-	if (InContext.bIsActive) \
-	{ \
-		UNROLL \
-		for (uint CompIt=0;CompIt<InNumComponent;++CompIt) \
-		{ \
-			InContext.Pos = ShaderPrintValue(InContext.Pos, Value[CompIt], GetDefaultFontColor()); \
-		} \
-	} \
-}
+void   InPrefix##Print(InType Value, FFontColor Color=DEFAULTFONTCOLOR, uint MaxDigit=MAX_DIGIT_COUNT, uint MaxDecimal=MAX_DECIMAL_COUNT)                                  { if (ShaderPrintFilterEnable) { UNROLL for (uint CompIt=0;CompIt<InNumComponent;++CompIt) { ShaderPrintCursorPos = ShaderPrintValue(ShaderPrintCursorPos, Value[CompIt], Color, MaxDigit, MaxDecimal); } } } \
+float2 InPrefix##Print(float2 ScreenPos, InType Value, FFontColor Color=DEFAULTFONTCOLOR, uint MaxDigit=MAX_DIGIT_COUNT, uint MaxDecimal=MAX_DECIMAL_COUNT)                { if (ShaderPrintFilterEnable) { UNROLL for (uint CompIt=0;CompIt<InNumComponent;++CompIt) { ScreenPos            = ShaderPrintValue(ScreenPos,            Value[CompIt], Color, MaxDigit, MaxDecimal); } } return ScreenPos; } \
+void   InPrefix##Print(inout FShaderPrintContext InCtx, InType Value, FFontColor Color=DEFAULTFONTCOLOR, uint MaxDigit=MAX_DIGIT_COUNT, uint MaxDecimal=MAX_DECIMAL_COUNT) { if (InCtx.bIsActive)         { UNROLL for (uint CompIt=0;CompIt<InNumComponent;++CompIt) { InCtx.Pos            = ShaderPrintValue(InCtx.Pos,            Value[CompIt], Color, MaxDigit, MaxDecimal); } } }
 
-// ShaderPrint(X) version
+// ShaderPrint(X) version (legacy)
 SHADER_PRINT_OVERLOAD_1(Shader, 1, float)
 SHADER_PRINT_OVERLOAD_N(Shader, 2, float2)
 SHADER_PRINT_OVERLOAD_N(Shader, 3, float3)
@@ -562,80 +418,9 @@ SHADER_PRINT_OVERLOAD_1(, 1, bool)
 // Matrix type printing
 
 #define SHADER_PRINT_OVERLOAD_NN(InPrefix, InNumComponentX, InNumComponentY, InType) \
-float2 InPrefix##Print(in float2 ScreenPos, in InType Value, in FFontColor Color) \
-{ \
-	if (ShaderPrintFilterEnable) \
-	{ \
-		UNROLL \
-		for (uint CompY=0;CompY<InNumComponentY;++CompY) \
-		{ \
-			ScreenPos = InPrefix##Print(ScreenPos, Value[CompY], Color); \
-			ScreenPos += float2(0.f, ShaderPrintUniform.FontSpacing.y); \
-		} \
-	} \
-	return ScreenPos; \
-} \
-float2 InPrefix##Print(in float2 ScreenPos, in InType Value) \
-{ \
-	if (ShaderPrintFilterEnable) \
-	{ \
-		UNROLL \
-		for (uint CompY = 0; CompY < InNumComponentY; ++CompY) \
-		{ \
-			ScreenPos = InPrefix##Print(ScreenPos, Value[CompY], GetDefaultFontColor()); \
-			ScreenPos += float2(0.f, ShaderPrintUniform.FontSpacing.y); \
-		} \
-	} \
-	return ScreenPos; \
-} \
-void InPrefix##Print(in InType Value, in FFontColor Color) \
-{ \
-	if (ShaderPrintFilterEnable) \
-	{ \
-		UNROLL \
-		for (uint CompY = 0; CompY < InNumComponentY; ++CompY) \
-		{ \
-			ShaderPrintCursorPos = InPrefix##Print(ShaderPrintCursorPos, Value[CompY], Color); \
-			ShaderPrintCursorPos += float2(0.f, ShaderPrintUniform.FontSpacing.y); \
-		} \
-	} \
-} \
-void InPrefix##Print(in InType Value) \
-{ \
-	if (ShaderPrintFilterEnable) \
-	{ \
-		UNROLL \
-		for (uint CompY = 0; CompY < InNumComponentY; ++CompY) \
-		{ \
-			ShaderPrintCursorPos = InPrefix##Print(ShaderPrintCursorPos, Value[CompY], GetDefaultFontColor()); \
-			ShaderPrintCursorPos += float2(0.f, ShaderPrintUniform.FontSpacing.y); \
-		} \
-	} \
-} \
-void InPrefix##Print(inout FShaderPrintContext InContext, in InType Value, in FFontColor Color) \
-{ \
-	if (InContext.bIsActive) \
-	{ \
-		UNROLL \
-		for (uint CompY = 0; CompY < InNumComponentY; ++CompY) \
-		{ \
-			InContext.Pos = InPrefix##Print(InContext.Pos, Value[CompY], Color); \
-			InContext.Pos += float2(0.f, ShaderPrintUniform.FontSpacing.y); \
-		} \
-	} \
-} \
-void InPrefix##Print(inout FShaderPrintContext InContext, in InType Value) \
-{ \
-	if (InContext.bIsActive) \
-	{ \
-		UNROLL \
-		for (uint CompY = 0; CompY < InNumComponentY; ++CompY) \
-		{ \
-			InContext.Pos = InPrefix##Print(InContext.Pos, Value[CompY], GetDefaultFontColor()); \
-			InContext.Pos += float2(0.f, ShaderPrintUniform.FontSpacing.y); \
-		} \
-	} \
-}
+void   InPrefix##Print(InType Value, FFontColor Color=DEFAULTFONTCOLOR, uint MaxDigit=MAX_DIGIT_COUNT, uint MaxDecimal=MAX_DECIMAL_COUNT)                                   { if (ShaderPrintFilterEnable) { UNROLL for (uint CompY=0;CompY<InNumComponentY;++CompY) { ShaderPrintCursorPos = InPrefix##Print(ShaderPrintCursorPos, Value[CompY], Color, MaxDigit, MaxDecimal); ShaderPrintCursorPos += float2(0.f, ShaderPrintUniform.FontSpacing.y); } } } \
+float2 InPrefix##Print(float2 ScreenPos, InType Value, FFontColor Color=DEFAULTFONTCOLOR, uint MaxDigit=MAX_DIGIT_COUNT, uint MaxDecimal=MAX_DECIMAL_COUNT)                 { if (ShaderPrintFilterEnable) { UNROLL for (uint CompY=0;CompY<InNumComponentY;++CompY) { ScreenPos            = InPrefix##Print(ScreenPos,            Value[CompY], Color, MaxDigit, MaxDecimal); ScreenPos += float2(0.f, ShaderPrintUniform.FontSpacing.y); } } return ScreenPos; } \
+void   InPrefix##Print(inout FShaderPrintContext InCtx, InType Value, FFontColor Color=DEFAULTFONTCOLOR, uint MaxDigit=MAX_DIGIT_COUNT, uint MaxDecimal=MAX_DECIMAL_COUNT)  { if (InCtx.bIsActive)         { UNROLL for (uint CompY=0;CompY<InNumComponentY;++CompY) { InCtx.Pos            = InPrefix##Print(InCtx.Pos,            Value[CompY], Color, MaxDigit, MaxDecimal); InCtx.Pos += float2(0.f, ShaderPrintUniform.FontSpacing.y); } } }
 
 // ShaderPrint(X) version
 SHADER_PRINT_OVERLOAD_NN(Shader, 3, 3, float3x3)
diff --git a/Engine/Shaders/Private/ShaderPrintDraw.usf b/Engine/Shaders/Private/ShaderPrintDraw.usf
--- a/Engine/Shaders/Private/ShaderPrintDraw.usf
+++ b/Engine/Shaders/Private/ShaderPrintDraw.usf
@@ -154,7 +154,7 @@ void AddIntSymbols(inout float2 ScreenPos, in int Value, in float3 Color)
 	AddUIntSymbols(ScreenPos, (uint)Value, Color);
 }
 
-void AddFloatSymbols(inout float2 ScreenPos, in float Value, in float3 Color)
+void AddFloatSymbols(inout float2 ScreenPos, in float Value, in float3 Color, in uint MaxDecimalCount)
 {
 	if (isinf(Value))
 	{
@@ -181,7 +181,8 @@ void AddFloatSymbols(inout float2 ScreenPos, in float Value, in float3 Color)
 	AddIntSymbols(ScreenPos, (int)Value, Color);
 	AddSymbol(ScreenPos, _DOT_, Color);
 
-	for (int i = 5; i > 0; --i)
+	MaxDecimalCount = min(MaxDecimalCount, MAX_DECIMAL_COUNT);
+	for (uint i = MaxDecimalCount; i > 0; --i)
 	{
 		Value = frac(Value);
 		if (Value > 0.f)
@@ -218,7 +219,7 @@ void BuildSymbolBufferCS(uint3 DispatchThreadId : SV_DispatchThreadID)
 	if (Value.Type == SHADER_PRINT_TYPE_SYMBOL)
 		AddSymbol(Value.ScreenPos, Value.Value, Value.Color);
 	else if (Value.Type == SHADER_PRINT_TYPE_FLOAT)
-		AddFloatSymbols(Value.ScreenPos, asfloat(Value.Value), Value.Color);
+		AddFloatSymbols(Value.ScreenPos, asfloat(Value.Value), Value.Color, Value.Metadata);
 	else if (Value.Type == SHADER_PRINT_TYPE_INT)
 		AddIntSymbols(Value.ScreenPos, Value.Value, Value.Color);
 	else if (Value.Type == SHADER_PRINT_TYPE_UINT)
