Strata blending clatification

diff --git a/Engine/Shaders/Private/Strata/StrataLegacyConversion.ush b/Engine/Shaders/Private/Strata/StrataLegacyConversion.ush
--- a/Engine/Shaders/Private/Strata/StrataLegacyConversion.ush
+++ b/Engine/Shaders/Private/Strata/StrataLegacyConversion.ush
@@ -497,16 +497,16 @@ FStrataData StrataConvertLegacyMaterialStatic(
 	{
 		const float GreyTransmittance = saturate(1.0f - Opacity);
 	#if MATERIALBLENDING_OPAQUE || MATERIALBLENDING_MASKED
-		// Uses opaque blending
+		// Opaque materials only write emissive.
 		Out = GetStrataUnlitBSDF(Emissive, 0.0f);
 	#elif MATERIALBLENDING_ALPHACOMPOSITE
-		// Uses (Luminance, grey transmittance) blending
+		// Uses (Luminance, black transmittance) for premultiplied alpha blending. Opacity is controled using coverage
 		Out = GetStrataUnlitBSDF(Emissive, 0.0f);
 		#if STRATA_INLINE_SHADING
 			Out.InlinedBSDF.Coverage = Opacity;
 		#endif
 	#elif MATERIALBLENDING_TRANSLUCENT
-		// Uses (Luminance, grey transmittance) blending
+		// Uses (Luminance, black transmittance) for translucency blending. Opacity is controled using coverage
 		Out = GetStrataUnlitBSDF(Emissive * Opacity, 0.0f);
 		#if STRATA_INLINE_SHADING
 			Out.InlinedBSDF.Coverage = Opacity;
@@ -514,11 +514,13 @@ FStrataData StrataConvertLegacyMaterialStatic(
 	#elif MATERIALBLENDING_ALPHAHOLDOUT
 		Out = GetStrataUnlitBSDF(0.0f, 1.0f);
 	#elif MATERIALBLENDING_ADDITIVE
+		// Emissive is weighted by opacity in the legacy material.
 		Out = GetStrataUnlitBSDF(Emissive * Opacity, 1.0f);
 	#elif MATERIALBLENDING_MODULATE
+		// Setting up emissive as the transmittance color. It is not clamped in case brightening is required.
 		Out = GetStrataUnlitBSDF(0.0f, Emissive);
 	#else
-		// Opaque and masked
+		// Required default for some materials such as Editor UI
 		Out = GetStrataUnlitBSDF(Emissive, 0.0f);
 	#endif
 	}
