Strata - conversion from legacy materials with premultiplied alpha now works for lit translucent materials.

diff --git a/Engine/Shaders/Private/BasePassPixelShader.usf b/Engine/Shaders/Private/BasePassPixelShader.usf
--- a/Engine/Shaders/Private/BasePassPixelShader.usf
+++ b/Engine/Shaders/Private/BasePassPixelShader.usf
@@ -1834,10 +1834,12 @@ void FPixelShaderInOut_MainPS(
 			Out.MRT[0] = half4(0.0f.xxx, Opacity);
 			Out.MRT[0] = RETURN_COLOR(Out.MRT[0]);
 		#else
+			// Pre multipled alpha blending
+			float AdjustedAlpha = saturate(1 - dot(AdjustedDualBlendMul, float3(1.0f, 1.0f, 1.0f) / 3.0f));
 			#if STRATA_USES_CONVERSION_FROM_LEGACY && MATERIALBLENDING_ADDITIVE
-				float AdjustedAlpha = 0.0f;
-			#else
-				float AdjustedAlpha = saturate(1 - dot(AdjustedDualBlendMul, float3(1.0f, 1.0f, 1.0f) / 3.0f));
+				AdjustedAlpha = 0.0f;
+			#elif STRATA_PREMULT_ALPHA_OVERRIDE && !STRATA_OPTIMIZED_UNLIT
+				AdjustedAlpha = StrataData.PreMultipliedAlphaOverrideCoverage >= 0.0f ? saturate(StrataData.PreMultipliedAlphaOverrideCoverage) : AdjustedAlpha;
 			#endif
 			// Pre multipled alpha blending
 			Out.MRT[0] = half4(AdjustedDualBlendAdd, AdjustedAlpha);
diff --git a/Engine/Shaders/Private/Strata/Strata.ush b/Engine/Shaders/Private/Strata/Strata.ush
--- a/Engine/Shaders/Private/Strata/Strata.ush
+++ b/Engine/Shaders/Private/Strata/Strata.ush
@@ -99,6 +99,13 @@ FRWStrataMaterialContainer InitialiseRWStrataMaterialContainer(RWTexture2DArray<
 #define STRATA_USES_CONVERSION_FROM_LEGACY 0
 #endif
 
+// When converting legacy materials to Strata, we need to override transmittance/coverage to be independent from the material and directly linked to opacity. 
+#if STRATA_USES_CONVERSION_FROM_LEGACY && MATERIALBLENDING_ALPHACOMPOSITE
+#define STRATA_PREMULT_ALPHA_OVERRIDE 1
+#else
+#define STRATA_PREMULT_ALPHA_OVERRIDE 0
+#endif
+
 #define STRATA_INLINE_SINGLELAYERWATER (PROJECT_STRATA && MATERIAL_IS_STRATA && STRATA_INLINE_SHADING && MATERIAL_SHADINGMODEL_SINGLELAYERWATER)
 
 // Behind the scene, we do all the simple volumetric lighting & transmitance computations considering a slab of 100 centimeters = 1 meter.
@@ -738,11 +745,18 @@ struct FStrataData
 	FStrataBSDF InlinedBSDF;// Used for parameter blending or special BSDF such as hair, cloud or water.
 
 	int  OperatorIndex;		// Used as entry point ot the material graph with multiple BSDF and operators.
+
+#if STRATA_PREMULT_ALPHA_OVERRIDE
+	float PreMultipliedAlphaOverrideCoverage;
+#endif
 };
 
 FStrataData GetInitialisedStrataData()
 {
 	FStrataData StrataData = (FStrataData)0;
+#if STRATA_PREMULT_ALPHA_OVERRIDE
+	StrataData.PreMultipliedAlphaOverrideCoverage = -1.0f;
+#endif
 	return StrataData;
 }
 
diff --git a/Engine/Shaders/Private/Strata/StrataLegacyConversion.ush b/Engine/Shaders/Private/Strata/StrataLegacyConversion.ush
--- a/Engine/Shaders/Private/Strata/StrataLegacyConversion.ush
+++ b/Engine/Shaders/Private/Strata/StrataLegacyConversion.ush
@@ -458,12 +458,16 @@ FStrataData StrataConvertLegacyMaterialDynamic(
 	FStrataData Out		 = StrataVerticalLayering(WTop_Slab, WBot_Slab, StrataTree,	/*OperatorIndex*/4, /*MaxDistanceFromLeaves*/2);
 
 	// Take in account the legacy opacity weighting for translucent material
-	#if !MATERIALBLENDING_MASKED && !MATERIALBLENDING_SOLID
+	#if !MATERIALBLENDING_MASKED && !MATERIALBLENDING_SOLID && STRATA_PREMULT_ALPHA_OVERRIDE==0
 	{
 		Out = StrataWeight(Out, FinalWeight, StrataTree,	/*OperatorIndex*/5, /*MaxDistanceFromLeaves*/3);
 	}
 	#endif
 
+	#if STRATA_PREMULT_ALPHA_OVERRIDE
+	Out.PreMultipliedAlphaOverrideCoverage = Opacity;
+	#endif
+
 	return Out;
 }
 
diff --git a/Engine/Shaders/Private/Strata/StrataTree.ush b/Engine/Shaders/Private/Strata/StrataTree.ush
--- a/Engine/Shaders/Private/Strata/StrataTree.ush
+++ b/Engine/Shaders/Private/Strata/StrataTree.ush
@@ -188,19 +188,19 @@ void LegacyUpdateBSDFsOperators(inout FStrataPixelHeader StrataPixelHeader, inou
 	UpdateSingleOperatorCoverageTransmittance(StrataTree, /*OperatorIndex*/3);
 	// MaxDistanceFromLeaves = 2 
 	UpdateSingleOperatorCoverageTransmittance(StrataTree, /*OperatorIndex*/4);
-#if !MATERIALBLENDING_MASKED && !MATERIALBLENDING_SOLID
+#if !MATERIALBLENDING_MASKED && !MATERIALBLENDING_SOLID && STRATA_PREMULT_ALPHA_OVERRIDE==0
 	// MaxDistanceFromLeaves = 3 
 	UpdateSingleOperatorCoverageTransmittance(StrataTree, /*OperatorIndex*/5);
 #endif
 
 	UpdateBSDFWeightAfterWeightOperatorVisit	(StrataTree, StrataTree.BSDFs[0], /*OperatorIndex*/2, /*PreviousIsInputA*/1); // WEIGHT
 	UpdateBSDFWeightAfterVerticalOperatorVisit	(StrataTree, StrataTree.BSDFs[0], /*OperatorIndex*/4, /*PreviousIsInputA*/1); // VERTICAL
-#if !MATERIALBLENDING_MASKED && !MATERIALBLENDING_SOLID
+#if !MATERIALBLENDING_MASKED && !MATERIALBLENDING_SOLID && STRATA_PREMULT_ALPHA_OVERRIDE==0
 	UpdateBSDFWeightAfterWeightOperatorVisit	(StrataTree, StrataTree.BSDFs[0], /*OperatorIndex*/5, /*PreviousIsInputA*/1); // WEIGHT
 #endif
 	UpdateBSDFWeightAfterWeightOperatorVisit	(StrataTree, StrataTree.BSDFs[1], /*OperatorIndex*/3, /*PreviousIsInputA*/1); // WEIGHT
 	UpdateBSDFWeightAfterVerticalOperatorVisit	(StrataTree, StrataTree.BSDFs[1], /*OperatorIndex*/4, /*PreviousIsInputA*/0); // VERTICAL
-#if !MATERIALBLENDING_MASKED && !MATERIALBLENDING_SOLID
+#if !MATERIALBLENDING_MASKED && !MATERIALBLENDING_SOLID && STRATA_PREMULT_ALPHA_OVERRIDE==0
 	UpdateBSDFWeightAfterWeightOperatorVisit	(StrataTree, StrataTree.BSDFs[1], /*OperatorIndex*/5, /*PreviousIsInputA*/1); // WEIGHT
 #endif
 }
