Strata - fixed fog application on legacy premultiple alpha material.

diff --git a/Engine/Shaders/Private/BasePassPixelShader.usf b/Engine/Shaders/Private/BasePassPixelShader.usf
--- a/Engine/Shaders/Private/BasePassPixelShader.usf
+++ b/Engine/Shaders/Private/BasePassPixelShader.usf
@@ -1820,6 +1820,12 @@ void FPixelShaderInOut_MainPS(
 		#endif
 
 	#elif STRATA_TRANSLUCENT_MATERIAL
+			
+		#if STRATA_LEGACY_PREMULT_ALPHA_OVERRIDE && !STRATA_OPTIMIZED_UNLIT
+			// This data patching is working because DualBlendSurfaceCoverage is always 1 when converting material from legacy to Strata with premultiplied alpha blending.
+			DualBlendSurfaceTransmittancePreCoverage= StrataData.PreMultipliedAlphaOverrideCoverage >= 0.0f ? 0.0f											: DualBlendSurfaceTransmittancePreCoverage;
+			DualBlendSurfaceCoverage				= StrataData.PreMultipliedAlphaOverrideCoverage >= 0.0f ? StrataData.PreMultipliedAlphaOverrideCoverage	: DualBlendSurfaceCoverage;
+		#endif
 
 			// Add fog luminance according to surfacecoverage and reduce surface luminance according to fog coverage.
 			float3 AdjustedDualBlendAdd = DualBlendSurfaceCoverage * Fogging.rgb + Fogging.a * DualBlendSurfaceLuminancePostCoverage;
@@ -1838,8 +1844,6 @@ void FPixelShaderInOut_MainPS(
 			float AdjustedAlpha = saturate(1 - dot(AdjustedDualBlendMul, float3(1.0f, 1.0f, 1.0f) / 3.0f));
 			#if STRATA_USES_CONVERSION_FROM_LEGACY && MATERIALBLENDING_ADDITIVE
 				AdjustedAlpha = 0.0f;
-			#elif STRATA_PREMULT_ALPHA_OVERRIDE && !STRATA_OPTIMIZED_UNLIT
-				AdjustedAlpha = StrataData.PreMultipliedAlphaOverrideCoverage >= 0.0f ? saturate(StrataData.PreMultipliedAlphaOverrideCoverage) : AdjustedAlpha;
 			#endif
 			// Pre multipled alpha blending
 			Out.MRT[0] = half4(AdjustedDualBlendAdd, AdjustedAlpha);
diff --git a/Engine/Shaders/Private/Strata/Strata.ush b/Engine/Shaders/Private/Strata/Strata.ush
--- a/Engine/Shaders/Private/Strata/Strata.ush
+++ b/Engine/Shaders/Private/Strata/Strata.ush
@@ -99,12 +99,10 @@ FRWStrataMaterialContainer InitialiseRWStrataMaterialContainer(RWTexture2DArray<
 #define STRATA_USES_CONVERSION_FROM_LEGACY 0
 #endif
 
-// When converting legacy materials to Strata, we need to override transmittance/coverage to be independent from the material and directly linked to opacity. 
-#if STRATA_USES_CONVERSION_FROM_LEGACY && MATERIALBLENDING_ALPHACOMPOSITE
-#define STRATA_PREMULT_ALPHA_OVERRIDE 1
-#else
-#define STRATA_PREMULT_ALPHA_OVERRIDE 0
-#endif
+// When converting legacy materials to Strata, we need to override transmittance/coverage to be independent from the material and directly linked to opacity
+#define STRATA_LEGACY_PREMULT_ALPHA_OVERRIDE (STRATA_USES_CONVERSION_FROM_LEGACY && MATERIALBLENDING_ALPHACOMPOSITE)
+
+#define STRATA_LEGACY_MATERIAL_APPLIES_FINAL_WEIGHT (!MATERIALBLENDING_MASKED && !MATERIALBLENDING_SOLID && STRATA_LEGACY_PREMULT_ALPHA_OVERRIDE==0)
 
 #define STRATA_INLINE_SINGLELAYERWATER (PROJECT_STRATA && MATERIAL_IS_STRATA && STRATA_INLINE_SHADING && MATERIAL_SHADINGMODEL_SINGLELAYERWATER)
 
@@ -746,7 +744,7 @@ struct FStrataData
 
 	int  OperatorIndex;		// Used as entry point ot the material graph with multiple BSDF and operators.
 
-#if STRATA_PREMULT_ALPHA_OVERRIDE
+#if STRATA_LEGACY_PREMULT_ALPHA_OVERRIDE
 	float PreMultipliedAlphaOverrideCoverage;
 #endif
 };
@@ -754,7 +752,7 @@ struct FStrataData
 FStrataData GetInitialisedStrataData()
 {
 	FStrataData StrataData = (FStrataData)0;
-#if STRATA_PREMULT_ALPHA_OVERRIDE
+#if STRATA_LEGACY_PREMULT_ALPHA_OVERRIDE
 	StrataData.PreMultipliedAlphaOverrideCoverage = -1.0f;
 #endif
 	return StrataData;
diff --git a/Engine/Shaders/Private/Strata/StrataLegacyConversion.ush b/Engine/Shaders/Private/Strata/StrataLegacyConversion.ush
--- a/Engine/Shaders/Private/Strata/StrataLegacyConversion.ush
+++ b/Engine/Shaders/Private/Strata/StrataLegacyConversion.ush
@@ -458,13 +458,13 @@ FStrataData StrataConvertLegacyMaterialDynamic(
 	FStrataData Out		 = StrataVerticalLayering(WTop_Slab, WBot_Slab, StrataTree,	/*OperatorIndex*/4, /*MaxDistanceFromLeaves*/2);
 
 	// Take in account the legacy opacity weighting for translucent material
-	#if !MATERIALBLENDING_MASKED && !MATERIALBLENDING_SOLID && STRATA_PREMULT_ALPHA_OVERRIDE==0
+	#if STRATA_LEGACY_MATERIAL_APPLIES_FINAL_WEIGHT
 	{
 		Out = StrataWeight(Out, FinalWeight, StrataTree,	/*OperatorIndex*/5, /*MaxDistanceFromLeaves*/3);
 	}
 	#endif
 
-	#if STRATA_PREMULT_ALPHA_OVERRIDE
+	#if STRATA_LEGACY_PREMULT_ALPHA_OVERRIDE
 	Out.PreMultipliedAlphaOverrideCoverage = Opacity;
 	#endif
 
diff --git a/Engine/Shaders/Private/Strata/StrataTree.ush b/Engine/Shaders/Private/Strata/StrataTree.ush
--- a/Engine/Shaders/Private/Strata/StrataTree.ush
+++ b/Engine/Shaders/Private/Strata/StrataTree.ush
@@ -188,19 +188,19 @@ void LegacyUpdateBSDFsOperators(inout FStrataPixelHeader StrataPixelHeader, inou
 	UpdateSingleOperatorCoverageTransmittance(StrataTree, /*OperatorIndex*/3);
 	// MaxDistanceFromLeaves = 2 
 	UpdateSingleOperatorCoverageTransmittance(StrataTree, /*OperatorIndex*/4);
-#if !MATERIALBLENDING_MASKED && !MATERIALBLENDING_SOLID && STRATA_PREMULT_ALPHA_OVERRIDE==0
+#if STRATA_LEGACY_MATERIAL_APPLIES_FINAL_WEIGHT
 	// MaxDistanceFromLeaves = 3 
 	UpdateSingleOperatorCoverageTransmittance(StrataTree, /*OperatorIndex*/5);
 #endif
 
 	UpdateBSDFWeightAfterWeightOperatorVisit	(StrataTree, StrataTree.BSDFs[0], /*OperatorIndex*/2, /*PreviousIsInputA*/1); // WEIGHT
 	UpdateBSDFWeightAfterVerticalOperatorVisit	(StrataTree, StrataTree.BSDFs[0], /*OperatorIndex*/4, /*PreviousIsInputA*/1); // VERTICAL
-#if !MATERIALBLENDING_MASKED && !MATERIALBLENDING_SOLID && STRATA_PREMULT_ALPHA_OVERRIDE==0
+#if STRATA_LEGACY_MATERIAL_APPLIES_FINAL_WEIGHT
 	UpdateBSDFWeightAfterWeightOperatorVisit	(StrataTree, StrataTree.BSDFs[0], /*OperatorIndex*/5, /*PreviousIsInputA*/1); // WEIGHT
 #endif
 	UpdateBSDFWeightAfterWeightOperatorVisit	(StrataTree, StrataTree.BSDFs[1], /*OperatorIndex*/3, /*PreviousIsInputA*/1); // WEIGHT
 	UpdateBSDFWeightAfterVerticalOperatorVisit	(StrataTree, StrataTree.BSDFs[1], /*OperatorIndex*/4, /*PreviousIsInputA*/0); // VERTICAL
-#if !MATERIALBLENDING_MASKED && !MATERIALBLENDING_SOLID && STRATA_PREMULT_ALPHA_OVERRIDE==0
+#if STRATA_LEGACY_MATERIAL_APPLIES_FINAL_WEIGHT
 	UpdateBSDFWeightAfterWeightOperatorVisit	(StrataTree, StrataTree.BSDFs[1], /*OperatorIndex*/5, /*PreviousIsInputA*/1); // WEIGHT
 #endif
 }
