Fixes black sky pixels in sky of AncientGame at sg.AntiAliasingQuality=4 due to incorrect handling of the history validity with Mitchel-Netravali's negative lobes.

diff --git a/Engine/Shaders/Private/TemporalSuperResolution/TSRResolveHistory.usf b/Engine/Shaders/Private/TemporalSuperResolution/TSRResolveHistory.usf
--- a/Engine/Shaders/Private/TemporalSuperResolution/TSRResolveHistory.usf
+++ b/Engine/Shaders/Private/TemporalSuperResolution/TSRResolveHistory.usf
@@ -170,8 +170,8 @@ void MainCS(
 				HighFrequencyColor  += (kernelWeight * HighFrequencyHdrWeight) * SampleHighFrequencyColor;
 				HighFrequencyWeight += (kernelWeight * HighFrequencyHdrWeight) * SampleHistoryValidity;
 
-				HistoryValidity       += max(kernelWeight, tsr_half(0.0)) * SampleHistoryValidity;
-				HistoryValidityWeight += max(kernelWeight, tsr_half(0.0));
+				HistoryValidity       += kernelWeight * SampleHistoryValidity;
+				HistoryValidityWeight += kernelWeight;
 			}
 		}
 
