Strata - legacy converted materials now works with multiple shading material attribute.
Next, foliage do need attention as it is missing the back face color.

diff --git a/Engine/Shaders/Private/Strata/StrataLegacyConversion.ush b/Engine/Shaders/Private/Strata/StrataLegacyConversion.ush
--- a/Engine/Shaders/Private/Strata/StrataLegacyConversion.ush
+++ b/Engine/Shaders/Private/Strata/StrataLegacyConversion.ush
@@ -185,7 +185,11 @@ FStrataData StrataConvertLegacyMaterialDynamic(
 		Top.Emissive = Emissive;
 
 		Bot.Weight = 0.f;
+
+		FinalWeight = 1.0f;
+#if STRATA_TRANSLUCENT_MATERIAL
 		FinalWeight = Opacity;
+#endif
 	}
 	else if (ShadingModel == SHADINGMODELID_DEFAULT_LIT)
 	{
@@ -197,7 +201,11 @@ FStrataData StrataConvertLegacyMaterialDynamic(
 		Top.Emissive = Emissive;
 
 		Bot.Weight = 0.f;
+
+		FinalWeight = 1.0f;
+#if STRATA_TRANSLUCENT_MATERIAL
 		FinalWeight = Opacity;
+#endif
 	}
 	else if (ShadingModel == SHADINGMODELID_SUBSURFACE || ShadingModel == SHADINGMODELID_PREINTEGRATED_SKIN)
 	{
@@ -224,7 +232,22 @@ FStrataData StrataConvertLegacyMaterialDynamic(
 
 		// No need to take into account opacity weighting for translucent material, as the coverage is assumed to be 1
 		Bot.Weight = 0.f;
-		FinalWeight = 1.f;
+
+		FinalWeight = 1.0f;
+
+#if STRATA_TRANSLUCENT_MATERIAL
+		// STRATA_TODO 
+		// If we try to reduce the diffusion according to transmittance that will increase mfp and make the material disapear (make it optically thin).
+		// Instead for now we use a mfp of 0 and make the material vanish according to opacity.
+		// All this will be fixed when this material will be converted to LUT.
+		const float Transmittance = saturate(1.f - Opacity);
+		const float Thickness = DefaultThickness;
+		Top.SSSDMFP = 0;// TransmittanceToMeanFreePath(Transmittance, Thickness* CENTIMETER_TO_METER)* METER_TO_CENTIMETER;
+		Top.SSSDMFPScale = 1.0f;
+		Top.Thickness = Thickness;
+
+		FinalWeight = Opacity;
+#endif
 	}
 	else if (ShadingModel == SHADINGMODELID_SUBSURFACE_PROFILE)
 	{
@@ -239,7 +262,11 @@ FStrataData StrataConvertLegacyMaterialDynamic(
 
 		// No need to take into account opacity weighting for translucent material, as the coverage is assumed to be 1
 		Bot.Weight = 0.f;
-		FinalWeight = 1.f;
+
+		FinalWeight = 1.0f;
+#if STRATA_TRANSLUCENT_MATERIAL
+		FinalWeight = Opacity;
+#endif
 	}
 	else if (ShadingModel == SHADINGMODELID_CLEAR_COAT)
 	{
@@ -285,7 +312,10 @@ FStrataData StrataConvertLegacyMaterialDynamic(
 		Top.bIsBottom = 0;
 		Top.Weight = ClearCoat;
 
+		FinalWeight = 1.0f;
+#if STRATA_TRANSLUCENT_MATERIAL
 		FinalWeight = Opacity;
+#endif
 	}
 	else if (ShadingModel == SHADINGMODELID_TWOSIDED_FOLIAGE)
 	{
@@ -306,7 +336,11 @@ FStrataData StrataConvertLegacyMaterialDynamic(
 
 		// Take in account the legacy opacity weighting for translucent material
 		Bot.Weight = 0.f;
+
+		FinalWeight = 1.0f;
+#if STRATA_TRANSLUCENT_MATERIAL
 		FinalWeight = Opacity;
+#endif
 	}
 	else if (ShadingModel == SHADINGMODELID_CLOTH)
 	{
@@ -323,7 +357,11 @@ FStrataData StrataConvertLegacyMaterialDynamic(
 		Top.FuzzColor = SubSurfaceColor;
 
 		Bot.Weight = 0.f;
+
+		FinalWeight = 1.0f;
+#if STRATA_TRANSLUCENT_MATERIAL
 		FinalWeight = Opacity;
+#endif
 	}
 	else if (ShadingModel == SHADINGMODELID_EYE)
 	{
@@ -410,7 +448,6 @@ FStrataData StrataConvertLegacyMaterialDynamic(
 
 		Top.Weight  = 1.0f;
 		Bot.Weight  = 0.0f;
-		FinalWeight = 1.0f;
 
 		// Reference two layers approach
 	#if 0
@@ -447,6 +484,9 @@ FStrataData StrataConvertLegacyMaterialDynamic(
 		// Now weight the top base material by opacity
 		Top.Weight = Opacity;
 	#endif
+
+		// Thin translucency model always have a coverage of 1. It also canont be used in a multiple shading model context.
+		FinalWeight = 1.0f;
 	}
 
 	// Fixed layer structure for helping compiler to unroll and optimize shader
