Rename ShaderDrawDebug.ush into ShaderPrintDrawPrimitive.usf to be consistent with ShaderPrint/ShaderDrawDebug merge.

diff --git a/Engine/Shaders/Private/ShaderDrawDebug.usf b/Engine/Shaders/Private/ShaderDrawDebug.usf
deleted file mode 100644
diff --git a/Engine/Shaders/Private/ShaderPrintDrawPrimitive.usf b/Engine/Shaders/Private/ShaderPrintDrawPrimitive.usf
new file mode 100644
--- /dev/null
+++ b/Engine/Shaders/Private/ShaderPrintDrawPrimitive.usf
@@ -0,0 +1,94 @@
+// Copyright Epic Games, Inc. All Rights Reserved.
+
+
+#include "Common.ush"
+
+#include "ShaderDrawDebug.ush"
+
+float3 TranslatedWorldOffsetConversion;
+
+// The shader used to render the debug element on view
+StructuredBuffer<uint> ShaderDrawDebugPrimitive;
+
+Texture2D<float> DepthTexture;
+float2 OutputInvResolution; 
+float2 OriginalViewRectMin;
+float2 OriginalViewSize;
+float2 OriginalBufferInvSize;
+
+SamplerState DepthSampler;
+
+#if GPU_DEBUG_RENDERING_VS 
+void ShaderDrawDebugVS(
+	in uint InstanceId : SV_InstanceID,
+	in uint VertexId : SV_VertexID,
+	out float4 Position : SV_POSITION,
+	out float4 ColorAlpha : TEXCOORD0)
+{
+	FShaderDrawElement Item = UnpackShaderElement(ShaderDrawDebugPrimitive, InstanceId);
+
+	float4 VertexPositionTWS = float4(VertexId == 0 ? Item.Pos0 : Item.Pos1, 1.0f);
+	float4 VertexColor = VertexId == 0 ? Item.Color0 : Item.Color1;
+
+	// Optional world offset. For 'locked' debug drawing, this offset translates the world 
+	// position recorded for a given translated world offset, into the current translated 
+	// world offset
+	VertexPositionTWS.xyz += TranslatedWorldOffsetConversion;
+
+	Position = mul(VertexPositionTWS, PrimaryView.TranslatedWorldToClip);
+	ColorAlpha = VertexColor;
+}
+#endif
+
+#if GPU_DEBUG_RENDERING_PS
+void ShaderDrawDebugPS(
+	in float4 Position : SV_POSITION,
+	in float4 ColorAlpha : TEXCOORD0,
+	out float4 OutSceneColor : SV_Target0)
+{
+	uint2 PixelCoord = Position.xy;
+
+	const float2 ViewportUV = float2(PixelCoord) * OutputInvResolution;
+	const float2 DepthUV    = (ViewportUV * OriginalViewSize.xy + OriginalViewRectMin) * OriginalBufferInvSize;
+	const float Depth = DepthTexture.SampleLevel(DepthSampler, DepthUV, 0);
+	const bool bIsHidden = Position.z < Depth; // Reverse-Z
+
+	const float ColorScale = bIsHidden ? 0.4f : 1;
+	const float CheckerboardScale = ((PixelCoord.x%2) ^ (PixelCoord.y%2)) > 0 ? 1 : 0.0;
+	const float OcclusionScale = bIsHidden ? CheckerboardScale : 1;
+	OutSceneColor = OcclusionScale * ColorAlpha * float4(ColorAlpha.aaa * ColorScale, 1.0); // Pre multiplied alpha with color scale
+}
+#endif
+
+#if GPU_DEBUG_RENDERING_CLEAR_CS
+
+RWStructuredBuffer<uint> RWElementBuffer;
+
+[numthreads(1, 1, 1)]
+void ShaderDrawDebugClearCS(uint3 DispatchThreadId : SV_DispatchThreadID)
+{
+	RWElementBuffer[0] = 0; // Only element which really need to be cleared
+	RWElementBuffer[1] = 0;
+	RWElementBuffer[2] = 0;
+	RWElementBuffer[3] = 0;
+	RWElementBuffer[4] = 0;
+	RWElementBuffer[5] = 0;
+	RWElementBuffer[6] = 0;
+	RWElementBuffer[7] = 0;
+}
+#endif
+
+#if GPU_DEBUG_RENDERING_COPY_CS
+
+StructuredBuffer<uint> ElementBuffer;
+RWBuffer<uint> RWIndirectArgs;
+
+[numthreads(1, 1, 1)]
+void ShaderDrawDebugCopyCS(uint3 DispatchThreadId : SV_DispatchThreadID)
+{
+	RWIndirectArgs[0] = 2; // Line primitive
+	RWIndirectArgs[1] = ElementBuffer[0];
+	RWIndirectArgs[2] = 0;
+	RWIndirectArgs[3] = 0;
+}
+#endif
\ No newline at end of file
diff --git a/Engine/Source/Runtime/Renderer/Private/ShaderPrint.cpp b/Engine/Source/Runtime/Renderer/Private/ShaderPrint.cpp
--- a/Engine/Source/Runtime/Renderer/Private/ShaderPrint.cpp
+++ b/Engine/Source/Runtime/Renderer/Private/ShaderPrint.cpp
@@ -412,7 +412,7 @@ namespace ShaderPrint
 		}
 	};
 
-	IMPLEMENT_GLOBAL_SHADER(FShaderDrawDebugClearCS, "/Engine/Private/ShaderDrawDebug.usf", "ShaderDrawDebugClearCS", SF_Compute);
+	IMPLEMENT_GLOBAL_SHADER(FShaderDrawDebugClearCS, "/Engine/Private/ShaderPrintDrawPrimitive.usf", "ShaderDrawDebugClearCS", SF_Compute);
 	
 	class FShaderDrawDebugCopyCS : public FGlobalShader
 	{
@@ -437,7 +437,7 @@ namespace ShaderPrint
 		}
 	};
 
-	IMPLEMENT_GLOBAL_SHADER(FShaderDrawDebugCopyCS, "/Engine/Private/ShaderDrawDebug.usf", "ShaderDrawDebugCopyCS", SF_Compute);
+	IMPLEMENT_GLOBAL_SHADER(FShaderDrawDebugCopyCS, "/Engine/Private/ShaderPrintDrawPrimitive.usf", "ShaderDrawDebugCopyCS", SF_Compute);
 
 	class FShaderDrawDebugVS : public FGlobalShader
 	{
@@ -466,7 +466,7 @@ namespace ShaderPrint
 		}
 	};
 
-	IMPLEMENT_GLOBAL_SHADER(FShaderDrawDebugVS, "/Engine/Private/ShaderDrawDebug.usf", "ShaderDrawDebugVS", SF_Vertex);
+	IMPLEMENT_GLOBAL_SHADER(FShaderDrawDebugVS, "/Engine/Private/ShaderPrintDrawPrimitive.usf", "ShaderDrawDebugVS", SF_Vertex);
 
 
 	class FShaderDrawDebugPS : public FGlobalShader
@@ -500,7 +500,7 @@ namespace ShaderPrint
 		}
 	};
 
-	IMPLEMENT_GLOBAL_SHADER(FShaderDrawDebugPS, "/Engine/Private/ShaderDrawDebug.usf", "ShaderDrawDebugPS", SF_Pixel);
+	IMPLEMENT_GLOBAL_SHADER(FShaderDrawDebugPS, "/Engine/Private/ShaderPrintDrawPrimitive.usf", "ShaderDrawDebugPS", SF_Pixel);
 
 	BEGIN_SHADER_PARAMETER_STRUCT(FShaderDrawVSPSParameters , )
 		SHADER_PARAMETER_STRUCT_INCLUDE(FShaderDrawDebugVS::FParameters, VS)
