Move ShaderDebugDraw.ush into ShaderPrint.ush.

This is the final step for merging ShaderDrawDebug and ShaderPrint.

diff --git a/Engine/Plugins/Experimental/VirtualHeightfieldMesh/Shaders/Private/VirtualHeightfieldMesh.usf b/Engine/Plugins/Experimental/VirtualHeightfieldMesh/Shaders/Private/VirtualHeightfieldMesh.usf
--- a/Engine/Plugins/Experimental/VirtualHeightfieldMesh/Shaders/Private/VirtualHeightfieldMesh.usf
+++ b/Engine/Plugins/Experimental/VirtualHeightfieldMesh/Shaders/Private/VirtualHeightfieldMesh.usf
@@ -3,7 +3,6 @@
 #include "/Engine/Private/Common.ush"	
 #include "/Engine/Private/MortonCode.ush"	
 //#include "/Engine/Private/ShaderPrint.ush"
-//#include "/Engine/Private/ShaderDrawDebug.ush"	
 #include "VirtualHeightfieldMesh.ush"	
 
 groupshared uint NumGroupTasks;
diff --git a/Engine/Shaders/Private/GPUSceneDebugRender.usf b/Engine/Shaders/Private/GPUSceneDebugRender.usf
--- a/Engine/Shaders/Private/GPUSceneDebugRender.usf
+++ b/Engine/Shaders/Private/GPUSceneDebugRender.usf
@@ -2,7 +2,6 @@
 
 #include "Common.ush"
 #include "SceneData.ush"
-#include "ShaderDrawDebug.ush"
 #include "ShaderPrint.ush"
 
 int NumInstances;
diff --git a/Engine/Shaders/Private/HairStrands/HairCardsDebug.usf b/Engine/Shaders/Private/HairStrands/HairCardsDebug.usf
--- a/Engine/Shaders/Private/HairStrands/HairCardsDebug.usf
+++ b/Engine/Shaders/Private/HairStrands/HairCardsDebug.usf
@@ -1,15 +1,10 @@
 // Copyright Epic Games, Inc. All Rights Reserved.
 
 #include "../Common.ush"
-#include "../ShaderDrawDebug.ush"
+#include "../ShaderPrint.ush"
 
 #if SHADER_ATLAS
 
-#define PRINT_ATLAS_RESOLUTION 0
-#if PRINT_ATLAS_RESOLUTION
-#include "../ShaderPrint.ush"
-#endif
-
 #define DEBUG_MODE_DEPTH 1
 #define DEBUG_MODE_COVERAGE 2
 #define DEBUG_MODE_TANGENT 3
diff --git a/Engine/Shaders/Private/HairStrands/HairCardsDeformation.usf b/Engine/Shaders/Private/HairStrands/HairCardsDeformation.usf
--- a/Engine/Shaders/Private/HairStrands/HairCardsDeformation.usf
+++ b/Engine/Shaders/Private/HairStrands/HairCardsDeformation.usf
@@ -7,7 +7,7 @@
 #define DEBUG_ENABLE 0
 
 #if DEBUG_ENABLE
-#include "../ShaderDrawDebug.ush"?
+#include "../ShaderPrint.ush"?
 #endif
 
 float4x4 LocalToWorld;
diff --git a/Engine/Shaders/Private/HairStrands/HairCardsGeneration.usf b/Engine/Shaders/Private/HairStrands/HairCardsGeneration.usf
--- a/Engine/Shaders/Private/HairStrands/HairCardsGeneration.usf
+++ b/Engine/Shaders/Private/HairStrands/HairCardsGeneration.usf
@@ -1,7 +1,7 @@
 // Copyright Epic Games, Inc. All Rights Reserved.
 
 #include "../Common.ush"
-#include "../ShaderDrawDebug.ush"
+#include "../ShaderPrint.ush"
 #include "../MonteCarlo.ush"
 
 #include "HairStrandsVisibilityCommon.ush"
diff --git a/Engine/Shaders/Private/HairStrands/HairCardsVoxel.usf b/Engine/Shaders/Private/HairStrands/HairCardsVoxel.usf
--- a/Engine/Shaders/Private/HairStrands/HairCardsVoxel.usf
+++ b/Engine/Shaders/Private/HairStrands/HairCardsVoxel.usf
@@ -4,7 +4,7 @@
 #include "../CommonViewUniformBuffer.ush"
 #include "../SceneTextureParameters.ush"
 #include "../DeferredShadingCommon.ush"
-#include "../ShaderDrawDebug.ush"
+#include "../ShaderPrint.ush"
 #include "../PositionReconstructionCommon.ush" 
 
 ///////////////////////////////////////////////////////////////////////////////////////////////////
diff --git a/Engine/Shaders/Private/HairStrands/HairStrandsClusterCulling.usf b/Engine/Shaders/Private/HairStrands/HairStrandsClusterCulling.usf
--- a/Engine/Shaders/Private/HairStrands/HairStrandsClusterCulling.usf
+++ b/Engine/Shaders/Private/HairStrands/HairStrandsClusterCulling.usf
@@ -81,7 +81,6 @@ void MainClearClusterAABBCS(uint2 DispatchThreadId : SV_DispatchThreadID)
 #ifdef SHADER_DRAWDEBUGAABB
 
 #include "../Common.ush"
-#include "../ShaderDrawDebug.ush"
 #include "../ShaderPrint.ush"
 
 #define MODE_CLUSTER	1
diff --git a/Engine/Shaders/Private/HairStrands/HairStrandsDebugPrint.usf b/Engine/Shaders/Private/HairStrands/HairStrandsDebugPrint.usf
--- a/Engine/Shaders/Private/HairStrands/HairStrandsDebugPrint.usf
+++ b/Engine/Shaders/Private/HairStrands/HairStrandsDebugPrint.usf
@@ -5,7 +5,6 @@
 #include "../SceneTextureParameters.ush"
 #include "../DeferredShadingCommon.ush"
 #include "../ShaderPrint.ush"
-#include "../ShaderDrawDebug.ush"
 #include "../PositionReconstructionCommon.ush"
 #include "HairStrandsVisibilityCommon.ush"
 #include "HairStrandsAABBCommon.ush" 
diff --git a/Engine/Shaders/Private/HairStrands/HairStrandsDeepShadowDebug.usf b/Engine/Shaders/Private/HairStrands/HairStrandsDeepShadowDebug.usf
--- a/Engine/Shaders/Private/HairStrands/HairStrandsDeepShadowDebug.usf
+++ b/Engine/Shaders/Private/HairStrands/HairStrandsDeepShadowDebug.usf
@@ -127,7 +127,6 @@ void VisualizeDomPS(
 #if SHADER_DOMINFO
 
 #include "../ShaderPrint.ush"
-#include "../ShaderDrawDebug.ush"
 #include "../Matrices.ush"
 #include "HairStrandsDeepShadowCommonStruct.ush"
 #include "HairStrandsAABBCommon.ush"
diff --git a/Engine/Shaders/Private/HairStrands/HairStrandsDeepTransmittanceMask.usf b/Engine/Shaders/Private/HairStrands/HairStrandsDeepTransmittanceMask.usf
--- a/Engine/Shaders/Private/HairStrands/HairStrandsDeepTransmittanceMask.usf
+++ b/Engine/Shaders/Private/HairStrands/HairStrandsDeepTransmittanceMask.usf
@@ -12,7 +12,6 @@
 
 #define DEBUG_ENABLE 0
 #if DEBUG_ENABLE
-	#include "../ShaderDrawDebug.ush"
 	#include "../ShaderPrint.ush"
 #endif //DEBUG_ENABLE
 
diff --git a/Engine/Shaders/Private/HairStrands/HairStrandsEnvironmentAO.usf b/Engine/Shaders/Private/HairStrands/HairStrandsEnvironmentAO.usf
--- a/Engine/Shaders/Private/HairStrands/HairStrandsEnvironmentAO.usf
+++ b/Engine/Shaders/Private/HairStrands/HairStrandsEnvironmentAO.usf
@@ -10,7 +10,7 @@
 #include "HairStrandsVoxelPageCommon.ush"
 #include "HairStrandsEnvironmentLightingCommon.ush"
 
-#include "../ShaderDrawDebug.ush"
+#include "../ShaderPrint.ush"
 
 #define VOXEL_TRAVERSAL_TYPE VOXEL_TRAVERSAL_LINEAR_MIPMAP
 #define VOXEL_TRAVERSAL_DEBUG 0 
diff --git a/Engine/Shaders/Private/HairStrands/HairStrandsEnvironmentLighting.usf b/Engine/Shaders/Private/HairStrands/HairStrandsEnvironmentLighting.usf
--- a/Engine/Shaders/Private/HairStrands/HairStrandsEnvironmentLighting.usf
+++ b/Engine/Shaders/Private/HairStrands/HairStrandsEnvironmentLighting.usf
@@ -57,7 +57,7 @@
 #include "../SkyLightingShared.ush"
 #include "../VolumetricLightmapShared.ush"
 #include "../HairBsdf.ush"
-#include "../ShaderDrawDebug.ush"
+#include "../ShaderPrint.ush"
 #if PERMUTATION_LUMEN
 #include "../Lumen/LumenRadianceCacheCommon.ush"
 #endif
diff --git a/Engine/Shaders/Private/HairStrands/HairStrandsInterpolation.usf b/Engine/Shaders/Private/HairStrands/HairStrandsInterpolation.usf
--- a/Engine/Shaders/Private/HairStrands/HairStrandsInterpolation.usf
+++ b/Engine/Shaders/Private/HairStrands/HairStrandsInterpolation.usf
@@ -23,7 +23,7 @@ float UnpackHairControlPointCoordU(uint InPackedDataW)
 #define DEBUG_ENABLE 0
 
 #if DEBUG_ENABLE && ALLOW_DEBUG_RENDERING
-#include "../ShaderDrawDebug.ush"
+#include "../ShaderPrint.ush"
 #include "../ColorMap.ush"
 #endif
 
diff --git a/Engine/Shaders/Private/HairStrands/HairStrandsRaytracingGeometry.usf b/Engine/Shaders/Private/HairStrands/HairStrandsRaytracingGeometry.usf
--- a/Engine/Shaders/Private/HairStrands/HairStrandsRaytracingGeometry.usf
+++ b/Engine/Shaders/Private/HairStrands/HairStrandsRaytracingGeometry.usf
@@ -245,7 +245,6 @@ void MainCS(uint2 DispatchThreadId : SV_DispatchThreadID)
 #if SHADER_POSITION_CHANGED
 
 #include "../ShaderPrint.ush"
-#include "../ShaderDrawDebug.ush"
 #include "HairStrandsAABBCommon.ush"
 
 uint VertexCount;
diff --git a/Engine/Shaders/Private/HairStrands/HairStrandsTexturesGeneration.usf b/Engine/Shaders/Private/HairStrands/HairStrandsTexturesGeneration.usf
--- a/Engine/Shaders/Private/HairStrands/HairStrandsTexturesGeneration.usf
+++ b/Engine/Shaders/Private/HairStrands/HairStrandsTexturesGeneration.usf
@@ -8,7 +8,7 @@
 
 #define DEBUG_ENABLE 0
 #if DEBUG_ENABLE
-#include "../ShaderDrawDebug.ush"
+#include "../ShaderPrint.ush"
 #endif
 
 /////////////////////////////////////////////////////////////////////////////////////////////////////////////////////////////
diff --git a/Engine/Shaders/Private/HairStrands/HairStrandsVoxelDepthInjection.usf b/Engine/Shaders/Private/HairStrands/HairStrandsVoxelDepthInjection.usf
--- a/Engine/Shaders/Private/HairStrands/HairStrandsVoxelDepthInjection.usf
+++ b/Engine/Shaders/Private/HairStrands/HairStrandsVoxelDepthInjection.usf
@@ -3,7 +3,6 @@
 #include "../Common.ush"
 #include "../SceneTextureParameters.ush"
 #include "../DeferredShadingCommon.ush"
-#include "../ShaderDrawDebug.ush"
 #include "../ShaderPrint.ush"
 #include "HairStrandsVisibilityCommon.ush"
 #include "HairStrandsVoxelPageCommon.ush"
diff --git a/Engine/Shaders/Private/HairStrands/HairStrandsVoxelPageTraversal.ush b/Engine/Shaders/Private/HairStrands/HairStrandsVoxelPageTraversal.ush
--- a/Engine/Shaders/Private/HairStrands/HairStrandsVoxelPageTraversal.ush
+++ b/Engine/Shaders/Private/HairStrands/HairStrandsVoxelPageTraversal.ush
@@ -21,7 +21,7 @@
 #ifndef VOXEL_TRAVERSAL_DEBUG
 #define VOXEL_TRAVERSAL_DEBUG 0
 #else
-#include "../ShaderDrawDebug.ush"
+#include "../ShaderPrint.ush"
 #endif
 
 #ifndef VOXEL_TRAVERSAL_FORCE_MIP_ENABLE
diff --git a/Engine/Shaders/Private/OIT/OITSorting.usf b/Engine/Shaders/Private/OIT/OITSorting.usf
--- a/Engine/Shaders/Private/OIT/OITSorting.usf
+++ b/Engine/Shaders/Private/OIT/OITSorting.usf
@@ -24,7 +24,6 @@
 #endif
 
 #if PERMUTATION_DEBUG || SHADER_DEBUG
-#include "../ShaderDrawDebug.ush"
 #include "../ShaderPrint.ush"
 #include "../ColorMap.ush"
 #endif
diff --git a/Engine/Shaders/Private/PhysicsFieldVisualizer.usf b/Engine/Shaders/Private/PhysicsFieldVisualizer.usf
--- a/Engine/Shaders/Private/PhysicsFieldVisualizer.usf
+++ b/Engine/Shaders/Private/PhysicsFieldVisualizer.usf
@@ -6,7 +6,6 @@
 #include "CommonViewUniformBuffer.ush"
 #include "SceneTextureParameters.ush"
 #include "DeferredShadingCommon.ush"
-#include "ShaderDrawDebug.ush"
 #include "ShaderPrint.ush" 
 #include "PositionReconstructionCommon.ush" 
  
diff --git a/Engine/Shaders/Private/RectLightAtlas.usf b/Engine/Shaders/Private/RectLightAtlas.usf
--- a/Engine/Shaders/Private/RectLightAtlas.usf
+++ b/Engine/Shaders/Private/RectLightAtlas.usf
@@ -2,7 +2,6 @@
 
 #include "Common.ush"
 #include "ShaderPrint.ush"
-#include "ShaderDrawDebug.ush"
 #include "ColorMap.ush"
 
 ////////////////////////////////////////////////////////////////////////////////////////////////
diff --git a/Engine/Shaders/Private/ShaderDrawDebug.ush b/Engine/Shaders/Private/ShaderDrawDebug.ush
deleted file mode 100644
diff --git a/Engine/Shaders/Private/ShaderPrint.ush b/Engine/Shaders/Private/ShaderPrint.ush
--- a/Engine/Shaders/Private/ShaderPrint.ush
+++ b/Engine/Shaders/Private/ShaderPrint.ush
@@ -177,6 +177,25 @@ FFontColor GetDefaultFontColor()							{ FFontColor Out; Out.Color = float3(1,1,
 #define FontSilver		InitFontColor(189.f / 255.f, 195.f / 255.f, 199.f / 255.f)
 #define FontEmerald		InitFontColor(46.f / 255.f, 204.f / 255.f, 113.f / 255.f)
 
+///////////////////////////////////////////////////////////////////////////////////////////////////
+// Primitive colors
+
+#define ColorWhite		float4(1, 1, 1, 1)
+#define ColorBlack		float4(0, 0, 0, 1)
+#define ColorRed		float4(1, 0, 0, 1)
+#define ColorGreen		float4(0, 1, 0, 1)
+#define ColorBlue		float4(0, 0, 1, 1)
+#define ColorYellow		float4(1, 1, 0, 1)
+#define ColorCyan		float4(0, 1, 1, 1)
+#define ColorMagenta	float4(1, 0, 1, 1)
+#define ColorOrange		float4(243.f / 255.f, 156.f / 255.f,  18.f / 255.f, 1)
+#define ColorPurple		float4(169.f / 255.f,   7.f / 255.f, 228.f / 255.f, 1)
+#define ColorTurquoise	float4( 26.f / 255.f, 188.f / 255.f, 156.f / 255.f, 1)
+#define ColorSilver		float4(189.f / 255.f, 195.f / 255.f, 199.f / 255.f, 1)
+#define ColorEmerald	float4( 46.f / 255.f, 204.f / 255.f, 113.f / 255.f, 1)
+#define ColorLightGreen	float4(0.5f, 0.75f, 0.5f, 1.0f)
+#define ColorDarkGreen	float4(0.25f,0.50f, 0.25f, 1.0f)
+
 ///////////////////////////////////////////////////////////////////////////////////////////////////
 // Util pack/unpac functions
 
@@ -652,4 +671,258 @@ float AddSlider(inout FShaderPrintContext InContext, FShaderPrintText InTextEntr
 	}
 
 	return Value;
-}
\ No newline at end of file
+}
+
+///////////////////////////////////////////////////////////////////////////////////////////////////
+// Primitives (Line, AABB, OBB, Referentials, ...)
+
+struct FShaderDrawElement 
+{
+	float3 Pos0;
+	float3 Pos1;
+	float4 Color0;
+	float4 Color1;
+};
+
+FShaderDrawElement UnpackShaderElement(StructuredBuffer<uint> InPrimitiveBuffer, uint InIndex)
+{
+	const uint Index8 = (1 + InIndex) * 8;
+
+	FShaderDrawElement Out = (FShaderDrawElement)0;
+	{
+		Out.Pos0.x = asfloat(InPrimitiveBuffer[Index8 + 0]);
+		Out.Pos0.y = asfloat(InPrimitiveBuffer[Index8 + 1]);
+		Out.Pos0.z = asfloat(InPrimitiveBuffer[Index8 + 2]);
+
+		const uint PackedColor0 = InPrimitiveBuffer[Index8 + 3];
+		Out.Color0 = float4((PackedColor0.x >> 24) & 0xFF, (PackedColor0.x >> 16) & 0xFF, (PackedColor0.x >> 8) & 0xFF, (PackedColor0.x) & 0xFF) / 255.0f;
+	}
+	{
+		Out.Pos1.x = asfloat(InPrimitiveBuffer[Index8 + 4]);
+		Out.Pos1.y = asfloat(InPrimitiveBuffer[Index8 + 5]);
+		Out.Pos1.z = asfloat(InPrimitiveBuffer[Index8 + 6]);
+
+		const uint PackedColor1 = InPrimitiveBuffer[Index8 + 7];
+		Out.Color1 = float4((PackedColor1.x >> 24) & 0xFF, (PackedColor1.x >> 16) & 0xFF, (PackedColor1.x >> 8) & 0xFF, (PackedColor1.x) & 0xFF) / 255.0f;
+	}
+	return Out;
+}
+
+RWStructuredBuffer<uint> ShaderPrint_RWLinesBuffer;
+
+bool AllocateShaderDrawElement(uint Count, inout uint OutIndex)
+{
+	OutIndex = 0;
+	if (ShaderPrint.MaxLineCount == 0)
+	{
+		return false;
+	}
+
+	InterlockedAdd(ShaderPrint_RWLinesBuffer[0], Count, OutIndex);
+	return (1+OutIndex) < ShaderPrint.MaxLineCount;
+}
+
+void AddShaderElement(FShaderDrawElement E, uint Index)
+{
+	const uint4 PackedC0 = uint4(255.0f * saturate(E.Color0));
+	const uint4 PackedC1 = uint4(255.0f * saturate(E.Color1));
+
+	const uint Index8 = (1 + Index) * 8;
+	ShaderPrint_RWLinesBuffer[Index8 + 0] = asuint(E.Pos0.x);
+	ShaderPrint_RWLinesBuffer[Index8 + 1] = asuint(E.Pos0.y);
+	ShaderPrint_RWLinesBuffer[Index8 + 2] = asuint(E.Pos0.z);
+	ShaderPrint_RWLinesBuffer[Index8 + 3] = (PackedC0.x << 24) | (PackedC0.y << 16) | (PackedC0.z << 8) | (PackedC0.w);
+
+	ShaderPrint_RWLinesBuffer[Index8 + 4] = asuint(E.Pos1.x);
+	ShaderPrint_RWLinesBuffer[Index8 + 5] = asuint(E.Pos1.y);
+	ShaderPrint_RWLinesBuffer[Index8 + 6] = asuint(E.Pos1.z);
+	ShaderPrint_RWLinesBuffer[Index8 + 7] = (PackedC1.x << 24) | (PackedC1.y << 16) | (PackedC1.z << 8) | (PackedC1.w);
+}
+
+int2 GetCursorPos()
+{
+	return ShaderPrint.CursorCoord;
+}
+
+///////////////////////////////////////////////////////////////////////////////////////////////////
+// Line
+
+void AddLineTWS(float3 Pos0, float3 Pos1, float4 Color0, float4 Color1)
+{
+	uint Offset = 0;
+	if (AllocateShaderDrawElement(1, Offset))
+	{
+		FShaderDrawElement Element;
+		Element.Pos0 = Pos0;
+		Element.Pos1 = Pos1;
+		Element.Color0 = Color0;
+		Element.Color1 = Color1;
+		AddShaderElement(Element, Offset);
+	}
+}
+
+void AddLineTWS(float3 Pos0, float3 Pos1, float4 Color)					{ AddLineTWS(Pos0, Pos1, Color, Color); }
+void AddLineWS(float3 Pos0, float3 Pos1, float4 Color)					{ AddLineTWS(Pos0+ShaderPrint.TranslatedWorldOffset, Pos1+ShaderPrint.TranslatedWorldOffset, Color); }
+void AddLineWS(float3 Pos0, float3 Pos1, float4 Color0, float4 Color1)	{ AddLineTWS(Pos0+ShaderPrint.TranslatedWorldOffset, Pos1+ShaderPrint.TranslatedWorldOffset, Color0, Color1); }
+void AddLine(float3 Pos0, float3 Pos1, float4 Color)					{ AddLineTWS(Pos0+ShaderPrint.TranslatedWorldOffset, Pos1+ShaderPrint.TranslatedWorldOffset, Color); }
+void AddLine(float3 Pos0, float3 Pos1, float4 Color0, float4 Color1)	{ AddLineTWS(Pos0+ShaderPrint.TranslatedWorldOffset, Pos1+ShaderPrint.TranslatedWorldOffset, Color0, Color1); }
+
+///////////////////////////////////////////////////////////////////////////////////////////////////
+// Quad
+
+void AddQuadTWS(float3 Pos0, float3 Pos1, float3 Pos2, float3 Pos3, float4 Color)
+{
+	uint Offset = 0;
+	if (AllocateShaderDrawElement(4, Offset))
+	{
+		FShaderDrawElement Element;
+		Element.Color0 = Element.Color1 = Color;
+		Element.Pos0 = Pos0; Element.Pos1 = Pos1; AddShaderElement(Element, Offset + 0);
+		Element.Pos0 = Pos1; Element.Pos1 = Pos2; AddShaderElement(Element, Offset + 1);
+		Element.Pos0 = Pos2; Element.Pos1 = Pos3; AddShaderElement(Element, Offset + 2);
+		Element.Pos0 = Pos3; Element.Pos1 = Pos0; AddShaderElement(Element, Offset + 3);
+	}
+}
+
+void AddQuadWS(float3 Pos0, float3 Pos1, float3 Pos2, float3 Pos3, float4 Color) { AddQuadTWS(Pos0+ShaderPrint.TranslatedWorldOffset, Pos1+ShaderPrint.TranslatedWorldOffset, Pos2+ShaderPrint.TranslatedWorldOffset, Pos3+ShaderPrint.TranslatedWorldOffset, Color); }
+void AddQuad(float3 Pos0, float3 Pos1, float3 Pos2, float3 Pos3, float4 Color)	 { AddQuadTWS(Pos0+ShaderPrint.TranslatedWorldOffset, Pos1+ShaderPrint.TranslatedWorldOffset, Pos2+ShaderPrint.TranslatedWorldOffset, Pos3+ShaderPrint.TranslatedWorldOffset, Color); }
+
+///////////////////////////////////////////////////////////////////////////////////////////////////
+// AABB
+
+void AddAABBTWS(float3 Min, float3 Max, float4 Color)
+{
+	uint Offset = 0;
+	if (AllocateShaderDrawElement(12, Offset))
+	{
+		float3 P0 = float3(Min.x, Min.y, Min.z);
+		float3 P1 = float3(Max.x, Min.y, Min.z);
+		float3 P2 = float3(Max.x, Max.y, Min.z);
+		float3 P3 = float3(Min.x, Max.y, Min.z);
+		float3 P4 = float3(Min.x, Min.y, Max.z);
+		float3 P5 = float3(Max.x, Min.y, Max.z);
+		float3 P6 = float3(Max.x, Max.y, Max.z);
+		float3 P7 = float3(Min.x, Max.y, Max.z);
+
+		FShaderDrawElement Element;
+		Element.Color0 = Element.Color1 = Color;
+
+		Element.Pos0 = P0; Element.Pos1 = P1; AddShaderElement(Element, Offset + 0);
+		Element.Pos0 = P1; Element.Pos1 = P2; AddShaderElement(Element, Offset + 1);
+		Element.Pos0 = P2; Element.Pos1 = P3; AddShaderElement(Element, Offset + 2);
+		Element.Pos0 = P3; Element.Pos1 = P0; AddShaderElement(Element, Offset + 3);
+
+		Element.Pos0 = P4; Element.Pos1 = P5; AddShaderElement(Element, Offset + 4);
+		Element.Pos0 = P5; Element.Pos1 = P6; AddShaderElement(Element, Offset + 5);
+		Element.Pos0 = P6; Element.Pos1 = P7; AddShaderElement(Element, Offset + 6);
+		Element.Pos0 = P7; Element.Pos1 = P4; AddShaderElement(Element, Offset + 7);
+
+		Element.Pos0 = P0; Element.Pos1 = P4; AddShaderElement(Element, Offset + 8);
+		Element.Pos0 = P1; Element.Pos1 = P5; AddShaderElement(Element, Offset + 9);
+		Element.Pos0 = P2; Element.Pos1 = P6; AddShaderElement(Element, Offset +10);
+		Element.Pos0 = P3; Element.Pos1 = P7; AddShaderElement(Element, Offset +11);
+	}
+}
+
+void AddAABBWS(float3 Min, float3 Max, float4 Color)	{ AddAABBTWS(Min + ShaderPrint.TranslatedWorldOffset, Max + ShaderPrint.TranslatedWorldOffset, Color); }
+void AddAABB(float3 Min, float3 Max, float4 Color)		{ AddAABBTWS(Min + ShaderPrint.TranslatedWorldOffset, Max + ShaderPrint.TranslatedWorldOffset, Color); }
+
+///////////////////////////////////////////////////////////////////////////////////////////////////
+// Cross
+
+void AddCrossTWS(float3 Pos, float Size, float4 Color)
+{
+	AddLineTWS(Pos - float3(Size,0,0), Pos + float3(Size,0,0), Color, Color);
+	AddLineTWS(Pos - float3(0,Size,0), Pos + float3(0,Size,0), Color, Color);
+	AddLineTWS(Pos - float3(0,0,Size), Pos + float3(0,0,Size), Color, Color);
+}
+
+void AddCrossWS(float3 Pos, float Size, float4 Color)	{ AddAABBTWS(Pos + ShaderPrint.TranslatedWorldOffset, Size, Color); }
+void AddCross(float3 Pos, float Size, float4 Color)		{ AddAABBTWS(Pos + ShaderPrint.TranslatedWorldOffset, Size, Color); }
+
+///////////////////////////////////////////////////////////////////////////////////////////////////
+// Referential
+
+void AddReferentialTWS(float3 Pos, float3 T, float3 B, float3 N, float Scale=1)
+{
+	AddLineTWS(Pos, Pos + normalize(T)*Scale, ColorRed);
+	AddLineTWS(Pos, Pos + normalize(B)*Scale, ColorGreen);
+	AddLineTWS(Pos, Pos + normalize(N)*Scale, ColorBlue);
+}
+
+void AddReferentialTWS(float3 Pos, float3 TangentZ, float Scale=1)
+{
+	const float Sign = TangentZ.z >= 0 ? 1 : -1;
+	const float a = -rcp(Sign + TangentZ.z);
+	const float b = TangentZ.x * TangentZ.y * a;
+
+	const float3 TangentX = { 1 + Sign * a * Pow2(TangentZ.x), Sign * b, -Sign * TangentZ.x };
+	const float3 TangentY = { b,  Sign + a * Pow2(TangentZ.y), -TangentZ.y };
+
+	AddReferentialTWS(Pos, TangentX, TangentY, TangentZ, Scale);
+}
+
+void AddReferentialTWS(float4x4 InM, float Scale = 1)
+{
+	AddReferentialTWS(InM[3].xyz, InM[0].xyz, InM[1].xyz, InM[2].xyz, Scale);
+}
+
+void AddReferentialWS(float4x4 InM, float Scale = 1)							{ AddReferentialTWS(InM[3].xyz + ShaderPrint.TranslatedWorldOffset, InM[0].xyz, InM[1].xyz, InM[2].xyz, Scale); }
+void AddReferentialWS(float3 Pos, float3 TangentZ, float Scale = 1)				{ AddReferentialTWS(Pos + ShaderPrint.TranslatedWorldOffset, TangentZ, Scale); }
+void AddReferentialWS(float3 Pos, float3 T, float3 B, float3 N, float Scale = 1){ AddReferentialTWS(Pos + ShaderPrint.TranslatedWorldOffset, T, B, N, Scale); }
+void AddReferential(float4x4 InM, float Scale = 1)								{ AddReferentialTWS(InM[3].xyz + ShaderPrint.TranslatedWorldOffset, InM[0].xyz, InM[1].xyz, InM[2].xyz, Scale); }
+void AddReferential(float3 Pos, float3 TangentZ, float Scale = 1)				{ AddReferentialTWS(Pos + ShaderPrint.TranslatedWorldOffset, TangentZ, Scale); }
+void AddReferential(float3 Pos, float3 T, float3 B, float3 N, float Scale = 1)	{ AddReferentialTWS(Pos + ShaderPrint.TranslatedWorldOffset, T, B, N, Scale); }
+
+///////////////////////////////////////////////////////////////////////////////////////////////////
+// Triangle
+
+void AddLineTriangleTWS(float3 P0, float3 P1, float3 P2, float4 Color)
+{
+	AddLineTWS(P0, P1, Color, Color);
+	AddLineTWS(P1, P2, Color, Color);
+	AddLineTWS(P2, P0, Color, Color);
+}
+
+void AddLineTriangleWS(float3 P0, float3 P1, float3 P2, float4 Color)	{ AddLineTriangleTWS(P0 + ShaderPrint.TranslatedWorldOffset, P1 + ShaderPrint.TranslatedWorldOffset, P2 + ShaderPrint.TranslatedWorldOffset, Color); }
+void AddLineTriangle(float3 P0, float3 P1, float3 P2, float4 Color)		{ AddLineTriangleTWS(P0 + ShaderPrint.TranslatedWorldOffset, P1 + ShaderPrint.TranslatedWorldOffset, P2 + ShaderPrint.TranslatedWorldOffset, Color); }
+
+///////////////////////////////////////////////////////////////////////////////////////////////////
+// OBB
+
+void AddOBBTWS(float3 Min, float3 Max, float4 Color, float4x4 Transform)
+{
+	uint Offset = 0;
+	if (AllocateShaderDrawElement(12, Offset))
+	{
+		float3 P0 = mul(float4(Min.x, Min.y, Min.z, 1.0f), Transform).xyz;
+		float3 P1 = mul(float4(Max.x, Min.y, Min.z, 1.0f), Transform).xyz;
+		float3 P2 = mul(float4(Max.x, Max.y, Min.z, 1.0f), Transform).xyz;
+		float3 P3 = mul(float4(Min.x, Max.y, Min.z, 1.0f), Transform).xyz;
+		float3 P4 = mul(float4(Min.x, Min.y, Max.z, 1.0f), Transform).xyz;
+		float3 P5 = mul(float4(Max.x, Min.y, Max.z, 1.0f), Transform).xyz;
+		float3 P6 = mul(float4(Max.x, Max.y, Max.z, 1.0f), Transform).xyz;
+		float3 P7 = mul(float4(Min.x, Max.y, Max.z, 1.0f), Transform).xyz;
+
+		FShaderDrawElement Element;
+		Element.Color0 = Element.Color1 = Color;
+
+		Element.Pos0 = P0; Element.Pos1 = P1; AddShaderElement(Element, Offset + 0);
+		Element.Pos0 = P1; Element.Pos1 = P2; AddShaderElement(Element, Offset + 1);
+		Element.Pos0 = P2; Element.Pos1 = P3; AddShaderElement(Element, Offset + 2);
+		Element.Pos0 = P3; Element.Pos1 = P0; AddShaderElement(Element, Offset + 3);
+
+		Element.Pos0 = P4; Element.Pos1 = P5; AddShaderElement(Element, Offset + 4);
+		Element.Pos0 = P5; Element.Pos1 = P6; AddShaderElement(Element, Offset + 5);
+		Element.Pos0 = P6; Element.Pos1 = P7; AddShaderElement(Element, Offset + 6);
+		Element.Pos0 = P7; Element.Pos1 = P4; AddShaderElement(Element, Offset + 7);
+
+		Element.Pos0 = P0; Element.Pos1 = P4; AddShaderElement(Element, Offset + 8);
+		Element.Pos0 = P1; Element.Pos1 = P5; AddShaderElement(Element, Offset + 9);
+		Element.Pos0 = P2; Element.Pos1 = P6; AddShaderElement(Element, Offset + 10);
+		Element.Pos0 = P3; Element.Pos1 = P7; AddShaderElement(Element, Offset + 11);
+	}
+}
+
+void AddOBBWS(float3 Min, float3 Max, float4 Color, float4x4 Transform) { Transform[3].xyz += ShaderPrint.TranslatedWorldOffset; AddOBBTWS(Min, Max, Color, Transform); }
+void AddOBB(float3 Min, float3 Max, float4 Color, float4x4 Transform)	{ Transform[3].xyz += ShaderPrint.TranslatedWorldOffset; AddOBBTWS(Min, Max, Color, Transform); }
diff --git a/Engine/Shaders/Private/ShaderPrintDrawPrimitive.usf b/Engine/Shaders/Private/ShaderPrintDrawPrimitive.usf
--- a/Engine/Shaders/Private/ShaderPrintDrawPrimitive.usf
+++ b/Engine/Shaders/Private/ShaderPrintDrawPrimitive.usf
@@ -3,7 +3,7 @@
 
 #include "Common.ush"
 
-#include "ShaderDrawDebug.ush"
+#include "ShaderPrint.ush"
 
 float3 TranslatedWorldOffsetConversion;
 
diff --git a/Engine/Shaders/Private/ShadingFurnaceTest.usf b/Engine/Shaders/Private/ShadingFurnaceTest.usf
--- a/Engine/Shaders/Private/ShadingFurnaceTest.usf
+++ b/Engine/Shaders/Private/ShadingFurnaceTest.usf
@@ -21,7 +21,6 @@
 
 #if DEBUG_ENABLE
 #include "ShaderPrint.ush"
-#include "ShaderDrawDebug.ush"
 #include "ColorMap.ush"
 #endif 
 
diff --git a/Engine/Shaders/Private/Strata/StrataVisualize.usf b/Engine/Shaders/Private/Strata/StrataVisualize.usf
--- a/Engine/Shaders/Private/Strata/StrataVisualize.usf
+++ b/Engine/Shaders/Private/Strata/StrataVisualize.usf
@@ -9,7 +9,7 @@
 #include "/Engine/Private/Strata/StrataEvaluation.ush"
 
 #include "/Engine/Private/MiniFontCommon.ush"
-#include "/Engine/Private/ShaderDrawDebug.ush"
+#include "/Engine/Private/ShaderPrint.ush"
 
 
 void DrawReferentialTWS(float3 P /*In Translated World Space*/, float3 X, float3 Y, float3 N, float3 InColor)
diff --git a/Engine/Shaders/Private/VirtualShadowMaps/VirtualShadowMapCacheManagement.usf b/Engine/Shaders/Private/VirtualShadowMaps/VirtualShadowMapCacheManagement.usf
--- a/Engine/Shaders/Private/VirtualShadowMaps/VirtualShadowMapCacheManagement.usf
+++ b/Engine/Shaders/Private/VirtualShadowMaps/VirtualShadowMapCacheManagement.usf
@@ -11,7 +11,7 @@
  #include "../Nanite/NaniteHZBCull.ush"
 
 #if ENABLE_DEBUG_MODE
-#include "../ShaderDrawDebug.ush"
+#include "../ShaderPrint.ush"
 #include "../ColorMap.ush"
 #endif
 
diff --git a/Engine/Shaders/Private/VolumetricCloud.usf b/Engine/Shaders/Private/VolumetricCloud.usf
--- a/Engine/Shaders/Private/VolumetricCloud.usf
+++ b/Engine/Shaders/Private/VolumetricCloud.usf
@@ -1545,7 +1545,7 @@ void MainPS(
 
 #if SHADER_SHADOW_FILTER_CS
 
-#include "ShaderDrawDebug.ush"
+#include "ShaderPrint.ush"
 #include "MonteCarlo.ush"
 
 SamplerState BilinearSampler;
@@ -1680,7 +1680,7 @@ void MainShadowFilterCS(uint3 DispatchThreadId : SV_DispatchThreadID)
 
 #if SHADER_DEBUG_SHADOW_CS
 
-#include "ShaderDrawDebug.ush"
+#include "ShaderPrint.ush"
 
 Texture2D<float2> CloudTracedTexture;
 float4 CloudTextureSizeInvSize;
