Fix cooking error due to ambiguous int/uint call.

diff --git a/Engine/Shaders/Private/ShaderPrintDraw.usf b/Engine/Shaders/Private/ShaderPrintDraw.usf
--- a/Engine/Shaders/Private/ShaderPrintDraw.usf
+++ b/Engine/Shaders/Private/ShaderPrintDraw.usf
@@ -46,7 +46,7 @@ void InitValueBufferCS()
 [numthreads(1, 1, 1)]
 void BuildIndirectDispatchArgsCS()
 {
-	uint ValueCount = min(ShaderPrint.MaxValueCount, ValuesBuffer[SHADER_PRINT_COUNT_OFFSET].Value);
+	uint ValueCount = min(ShaderPrint.MaxValueCount, uint(ValuesBuffer[SHADER_PRINT_COUNT_OFFSET].Value));
 
 	RWIndirectDispatchArgsBuffer[0] = (ValueCount + 63)/64;
 	RWIndirectDispatchArgsBuffer[1] = 1;
@@ -75,7 +75,7 @@ void AddSymbol(inout float2 ScreenPos, in int Symbol, in float3 Color, uint Type
 	
 		// Prevent writing off the buffer
 		// Note that counter still increases so need to clamp when reading it in later passes
-		if (IndexToStore >= ShaderPrint.MaxSymbolCount)
+		if (uint(IndexToStore) >= ShaderPrint.MaxSymbolCount)
 		{
 			return;
 		}
@@ -211,7 +211,7 @@ void AddHexSymbols(inout float2 ScreenPos, uint Value, in float3 Color)
 [numthreads(64, 1, 1)]
 void BuildSymbolBufferCS(uint3 DispatchThreadId : SV_DispatchThreadID)
 {
-	uint ValueCount = min(ShaderPrint.MaxValueCount, ValuesBuffer[SHADER_PRINT_COUNT_OFFSET].Value);
+	uint ValueCount = min(ShaderPrint.MaxValueCount, uint(ValuesBuffer[SHADER_PRINT_COUNT_OFFSET].Value));
 	if (DispatchThreadId.x >= ValueCount)
 	{
 		return;
@@ -279,7 +279,7 @@ void CompactStateBufferCS(uint3 DispatchThreadId : SV_DispatchThreadID)
 [numthreads(1, 1, 1)]
 void BuildIndirectDrawArgsCS()
 {
-	uint SymbolCount = min(ShaderPrint.MaxSymbolCount, SymbolsBuffer[SHADER_PRINT_COUNT_OFFSET].Value);
+	uint SymbolCount = min(ShaderPrint.MaxSymbolCount, uint(SymbolsBuffer[SHADER_PRINT_COUNT_OFFSET].Value));
 
 	RWIndirectDrawArgsBuffer[0] = 6;
 	RWIndirectDrawArgsBuffer[1] = SymbolCount * 2; // Main glyph + drop shadow
