Strata
- Fixed thin translucent not working correctly in forward
- Unified the way we prepare slab for lighting and sanitize it
- This also removes a few duplicated operations on bsdf.

diff --git a/Engine/Shaders/Private/DistortAccumulatePS.usf b/Engine/Shaders/Private/DistortAccumulatePS.usf
--- a/Engine/Shaders/Private/DistortAccumulatePS.usf
+++ b/Engine/Shaders/Private/DistortAccumulatePS.usf
@@ -142,20 +142,6 @@ void Main(
 		{
 		#define BSDF StrataTree.BSDFs[BSDFIdx]
 
-			// Sanitize BSDF before it is used for forward shading
-			StrataSanitizeBSDF(BSDF);
-
-			BRANCH
-			if (BSDF_GETTYPE(BSDF) == STRATA_BSDF_TYPE_SLAB)
-			{
-				// We want simple volumetric only if the DMFP input is pluged in, otherwise we avoid dealing with simple volumetric
-				const bool bIsSimpleVolume = BSDF_GETHASDMFP(BSDF);
-				if (bIsSimpleVolume)
-				{
-					EnableSlabBSDFSimpleVolumetric(BSDF);
-				}
-			}
-
 			const float BSDFCoverage = BSDF.Coverage;
 
 			// Create the BSDF context
diff --git a/Engine/Shaders/Private/Strata/Strata.ush b/Engine/Shaders/Private/Strata/Strata.ush
--- a/Engine/Shaders/Private/Strata/Strata.ush
+++ b/Engine/Shaders/Private/Strata/Strata.ush
@@ -1857,6 +1857,8 @@ void EnableSlabBSDFSimpleVolumetric(inout FStrataBSDF BSDF)
 
 void EnableSlabBSDFIsThin(inout FStrataBSDF BSDF)
 {
+	// BSDF_SETISTHIN is already set when creating the Slab.
+	BSDF_SETISSIMPLEVOLUME(BSDF, 0);
 	EnableSlabMFPComputationSpace(BSDF);
 }
 
@@ -2524,10 +2526,6 @@ void PackStrataOut(
 					// Sanitize happening after will clamp values that could have gone above and below valid range.
 					StrataDitherBSDF(BSDF, Dither);
 
-					// Sanitize BSDF in place before it is used and packed out to memory
-					StrataSanitizeBSDF(BSDF);
-
-
 					bool bIsSimpleMaterial = true;
 					bool bIsSingleMaterial = true;
 					switch (BSDF_GETTYPE(BSDF))
@@ -2546,38 +2544,16 @@ void PackStrataOut(
 						// they use the complex pass, as the toplayer data does not contain the full frame basis, only the top normal.
 						bIsSingleMaterial = BSDF_GETHASANISOTROPY(BSDF) == 0 && LuminanceWeightEqualOne;
 
-						// Two sided lighting can only be evalauted on the bottom layer, together with option SSS. Otherwise, the simple volume will be used.
-						const bool bIsThin = BSDF.bIsBottom && BSDF_GETISTHIN(BSDF);
-
-						// We want simple volumetric only if the DMFP input is pluged in, otherwise we avoid storing the DFMP
-						const bool bIsSimpleVolume = !BSDF.bIsBottom && BSDF_GETHASDMFP(BSDF);
-
-						bStrataSubsurfaceEnable = bStrataSubsurfaceEnable || (BSDF_GETHASSSS(BSDF) && BSDF.Coverage > 0.0f && !bIsSimpleVolume);
-
-						if (bIsSimpleVolume)
-						{
-							EnableSlabBSDFSimpleVolumetric(BSDF);
-
-							// Mark the BSDF has not having any SSS through post process.
-							bStrataSubsurfaceEnable = false;
-						}
-
-						if (bIsThin)
-						{
-							EnableSlabBSDFIsThin(BSDF);
-						}
-						else
-						{
-							// If not selected or not possible to have two sided lighting enabled (not bottom layer or SSSProfile is enabled), let's make sure we disable the feature.
-							BSDF_SETISTHIN(BSDF, 0);
-						}
-
 						if (BSDF.bIsTop)
 						{
 							TopLayerCoverage += BSDF.bIsTop ? BSDF.Coverage : 0.0f;
 							TopLayerData.Roughness += BSDF.Coverage * SLAB_ROUGHNESS(BSDF);
 						}
 
+						const bool bIsThin = BSDF_GETISTHIN(BSDF);
+						const bool bIsSimpleVolume = BSDF_GETISSIMPLEVOLUME(BSDF);
+						bStrataSubsurfaceEnable = bStrataSubsurfaceEnable || (BSDF_GETHASSSS(BSDF) && BSDF.Coverage > 0.0f && !bIsSimpleVolume && !bIsThin); // Should it be coverage>0 or any(LuminanceWeight>0) ?
+
 						const bool bSSSMask = !bIsThin && BSDF.bIsBottom;	// SSS and Thin can only be on the bottom layer, so ignoring bIsSimpleVolume here.
 						SetHasSubsurface(StrataPixelHeader, (bSSSMask && BSDF_GETHASSSS(BSDF)));
 
diff --git a/Engine/Shaders/Private/Strata/StrataTree.ush b/Engine/Shaders/Private/Strata/StrataTree.ush
--- a/Engine/Shaders/Private/Strata/StrataTree.ush
+++ b/Engine/Shaders/Private/Strata/StrataTree.ush
@@ -1,5 +1,4 @@
 
-
 #include "/Engine/Private/Strata/StrataEvaluation.ush"
 #include "/Engine/Private/ShadingCommon.ush"
 
@@ -7,6 +6,57 @@
 #define STRATA_OPAQUE_MATERIAL 0
 #endif
 
+
+
+///////////////////////////////////////////////////////////////////////////////
+// Utilities
+
+void PostProcessSlabBeforeLighting(inout FStrataBSDF BSDF)
+{
+	// Two sided lighting can only be evaluated on the bottom layer, together with option SSS. Otherwise, the simple volume will be used.
+	const bool bIsThin = BSDF.bIsBottom && BSDF_GETISTHIN(BSDF);
+
+	// We want simple volumetric only if the DMFP input is pluged in, otherwise we avoid dealing with simple volumetric.
+	// We also do not want to use the simple volumetric path if we are dealing with an opaque material and this slab is not at the bottom.
+	// And when opaque, the bottom layer must be opaque and thus will ahve a solid opaque diffuse color.
+	const bool bIsSimpleVolume = !bIsThin && BSDF_GETHASDMFP(BSDF) && (!STRATA_OPAQUE_MATERIAL || (STRATA_OPAQUE_MATERIAL && !BSDF.bIsBottom));
+
+	if (bIsSimpleVolume)
+	{
+		EnableSlabBSDFSimpleVolumetric(BSDF);
+	}
+
+	if (bIsThin)
+	{
+		EnableSlabBSDFIsThin(BSDF);
+	}
+	else
+	{
+		// If not selected or not possible to have two sided lighting enabled (not bottom layer or SSSProfile is enabled), let's make sure we disable the feature.
+		BSDF_SETISTHIN(BSDF, 0);
+	}
+}
+
+void PostProcessBSDFBeforeLighting(inout FStrataBSDF BSDF)
+{
+	// Sanitize BSDF before it is used for forward shading
+	StrataSanitizeBSDF(BSDF);
+
+	switch (BSDF_GETTYPE(BSDF))
+	{
+	case STRATA_BSDF_TYPE_SLAB:
+	{
+		PostProcessSlabBeforeLighting(BSDF);
+		break;
+	}
+	}
+}
+
+
+
+///////////////////////////////////////////////////////////////////////////////
+// Tree processing functions
+
 void UpdateSingleBSDFOperatorCoverageTransmittance(
 	FStrataPixelHeader StrataPixelHeader,
 	inout FStrataTree StrataTree, 
@@ -17,22 +67,9 @@ void UpdateSingleBSDFOperatorCoverageTransmittance(
 	float3 L
 )
 {
-	// Sanitize BSDF before it is used for forward shading
-	StrataSanitizeBSDF(CurrentBSDF);
+#if STRATA_INLINE_SHADING
 
-	BRANCH
-	if (BSDF_GETTYPE(CurrentBSDF) == STRATA_BSDF_TYPE_SLAB)
-	{
-		// We want simple volumetric only if the DMFP input is pluged in, otherwise we avoid dealing with simple volumetric.
-		// We also do not want to use the simple volumetric path if we are dealing with an opaque material and this slab is not at the bottom.
-		// And when opaque, the bottom layer must be opaque and thus will ahve a solid opaque diffuse color.
-		const bool bIsSimpleVolume = BSDF_GETHASDMFP(CurrentBSDF) && (!STRATA_OPAQUE_MATERIAL || (STRATA_OPAQUE_MATERIAL && !CurrentBSDF.bIsBottom));
-
-		if (bIsSimpleVolume)
-		{
-			EnableSlabBSDFSimpleVolumetric(CurrentBSDF);
-		}
-	}
+	PostProcessBSDFBeforeLighting(CurrentBSDF);
 
 	// Compute current BSDF transmittance
 	FStrataBSDFContext StrataBSDFContext = StrataCreateBSDFContext(StrataPixelHeader, CurrentBSDF, NullStrataAddressing, V, L);
@@ -49,11 +86,15 @@ void UpdateSingleBSDFOperatorCoverageTransmittance(
 
 	StrataTree.Operators[CurrentBSDF.OperatorIndex].Coverage = CurrentBSDFCoverage;
 	StrataTree.Operators[CurrentBSDF.OperatorIndex].Transmittance = CurrentBSDFTransmittance;
+
+#endif // STRATA_INLINE_SHADING
 }
 
 
 void UpdateSingleOperatorCoverageTransmittance(inout FStrataTree StrataTree, int OpIndex)
 {
+#if STRATA_INLINE_SHADING
+
 #define Op StrataTree.Operators[OpIndex]
 #define OpA StrataTree.Operators[StrataTree.Operators[OpIndex].LeftIndex]
 #define OpB StrataTree.Operators[StrataTree.Operators[OpIndex].RightIndex]
@@ -128,6 +169,8 @@ void UpdateSingleOperatorCoverageTransmittance(inout FStrataTree StrataTree, int
 #undef Op
 #undef OpA
 #undef OpB
+
+#endif // STRATA_INLINE_SHADING
 }
 
 
@@ -139,10 +182,14 @@ void UpdateBSDFWeightAfterVerticalOperatorVisit(
 	int OpIndex,
 	int PreviousIsInputA)
 {
+#if STRATA_INLINE_SHADING
+
 	// PreviousIsInputA > 0 means it is a BSDF affecting the top layer so we do not affect it by anything for this vertical layering.
 	// Otherise, the BSDF comes from the bottom part so we affect it weights from the gathered top layer coverage/transmittance
 	CurrentBSDF.LuminanceWeight *= 
 		PreviousIsInputA > 0 ? 1.0f : saturate(1.0f - Op.VerticalTopCoverage) + saturate(Op.VerticalTopCoverage) * Op.VerticalTopTransmittance;
+
+#endif // STRATA_INLINE_SHADING
 }
 
 void UpdateBSDFWeightAfterHorizontalOperatorVisit(
@@ -151,14 +198,16 @@ void UpdateBSDFWeightAfterHorizontalOperatorVisit(
 	int OpIndex,
 	int PreviousIsInputA)
 {
+#if STRATA_INLINE_SHADING
+
 	const float Mix = saturate(Op.Weight);
 	const float AMix = 1.0 - Mix;
 	const float BMix = Mix;
 	const float Weight = PreviousIsInputA > 0 ? AMix : BMix;
 	CurrentBSDF.LuminanceWeight *= Weight;
-#if STRATA_INLINE_SHADING
 	CurrentBSDF.Coverage *= Weight;
-#endif
+
+#endif // STRATA_INLINE_SHADING
 }
 
 void UpdateBSDFWeightAfterWeightOperatorVisit(
@@ -167,16 +216,22 @@ void UpdateBSDFWeightAfterWeightOperatorVisit(
 	int OpIndex,
 	int PreviousIsInputA)
 {
+#if STRATA_INLINE_SHADING
+
 	const float Weight = saturate(Op.Weight);
 	CurrentBSDF.LuminanceWeight *= Weight;
-#if STRATA_INLINE_SHADING
 	CurrentBSDF.Coverage *= Weight;
-#endif
+
+#endif // STRATA_INLINE_SHADING
 }
 
 #undef Op
 
 
+
+///////////////////////////////////////////////////////////////////////////////
+// Legacy materials conversion
+
 void LegacyUpdateBSDFsOperators(inout FStrataPixelHeader StrataPixelHeader, inout FStrataTree StrataTree, bool bRoughDiffuseEnabled, FStrataAddressing NullStrataAddressing, float3 V, float3 NullLightDir)
 {
 	// This must match what is done in StrataConvertLegacyMaterialDynamic
