Strata - made safe the MFP compuation from thickness from EnableSlabMFPComputationSpace.

diff --git a/Engine/Shaders/Private/Strata/Strata.ush b/Engine/Shaders/Private/Strata/Strata.ush
--- a/Engine/Shaders/Private/Strata/Strata.ush
+++ b/Engine/Shaders/Private/Strata/Strata.ush
@@ -1839,7 +1839,7 @@ void EnableSlabMFPComputationSpace(inout FStrataBSDF BSDF)
 
 	// Rescale the BSDF mean free path so that we can drop the thickness parameter and run volumetric material math on a unit slab of medium
 	// For instance, if thickness is larger => we need to reduce the MFP for the medium to appear visually thicker.
-	SLAB_SSSDMFP(BSDF) = clamp(SLAB_SSSDMFP(BSDF) * (STRATA_SIMPLEVOLUME_THICKNESS_CM / BSDF_GETTHICKNESSCM(BSDF)), 0.0f, Max111110BitsFloat3);
+	SLAB_SSSDMFP(BSDF) = clamp(SLAB_SSSDMFP(BSDF) * (STRATA_SIMPLEVOLUME_THICKNESS_CM / max(STRATA_EPSILON, BSDF_GETTHICKNESSCM(BSDF))), 0.0f, Max111110BitsFloat3);
 }
 
 void EnableSlabBSDFSimpleVolumetric(inout FStrataBSDF BSDF)
