Strata - added new way to compute top layer weight taking into account the dynamic topology coverage.
I went the simple way and did not move the accumulation site for for now.

diff --git a/Engine/Shaders/Private/Strata/Strata.ush b/Engine/Shaders/Private/Strata/Strata.ush
--- a/Engine/Shaders/Private/Strata/Strata.ush
+++ b/Engine/Shaders/Private/Strata/Strata.ush
@@ -116,6 +116,7 @@ FRWStrataMaterialContainer InitialiseRWStrataMaterialContainer(RWTexture2DArray<
 
 ///////////////////////////////////////////////////////////////////////////////
 // BSDFs representation
+
 struct FStrataBSDF
 {
 	uint				State;		// Extra informations and other enabled BSDF features
@@ -133,6 +134,8 @@ struct FStrataBSDF
 	float3				Emissive;	// All bsdf can potientially have emissive. This is never written out to the strata gbuffer but instead accumulated in the emissive buffer during the base pass.
 	float				ThicknessCm;// This must be kept in order to be able to normalize the mean free path to comply with our volumetric material math all done in a normalised slab of medium.
 	float3				TmpDMFP;
+
+	float				TopLayerWeight;// The weight this BSDF contributes to the top layer (normal and roughness for SSR and SSAO for instance)
 #endif
 
 	// Water has too many parameters to be stores in the extra VGPRs below. So we extend this structure for the case where it is evaluated inline only.
@@ -2508,7 +2511,7 @@ void PackStrataOut(
 #endif
 		
 		int SimpleMaterial_BSDFIndex = 0;
-		float TopLayerCoverage = 0.0f;
+		float TopLayerTotalWeight = 0.0f;
 		
 		{
 			STRATA_UNROLL_N(STRATA_CLAMPED_BSDF_COUNT)
@@ -2544,11 +2547,8 @@ void PackStrataOut(
 						// they use the complex pass, as the toplayer data does not contain the full frame basis, only the top normal.
 						bIsSingleMaterial = BSDF_GETHASANISOTROPY(BSDF) == 0 && LuminanceWeightEqualOne;
 
-						if (BSDF.bIsTop)
-						{
-							TopLayerCoverage += BSDF.bIsTop ? BSDF.Coverage : 0.0f;
-							TopLayerData.Roughness += BSDF.Coverage * SLAB_ROUGHNESS(BSDF);
-						}
+						TopLayerTotalWeight += BSDF.TopLayerWeight;
+						TopLayerData.Roughness += BSDF.TopLayerWeight * SLAB_ROUGHNESS(BSDF);
 
 						const bool bIsThin = BSDF_GETISTHIN(BSDF);
 						const bool bIsSimpleVolume = BSDF_GETISSIMPLEVOLUME(BSDF);
@@ -2564,8 +2564,8 @@ void PackStrataOut(
 					{
 						bIsSimpleMaterial = false;
 
-						TopLayerCoverage += BSDF.bIsTop ? BSDF.Coverage : 0.0f;
-						TopLayerData.Roughness += BSDF.bIsTop ? BSDF.Coverage * FullyRough : 0.0f;
+						TopLayerTotalWeight += BSDF.TopLayerWeight;
+						TopLayerData.Roughness += BSDF.TopLayerWeight * FullyRough;
 
 						break;
 					}
@@ -2575,8 +2575,8 @@ void PackStrataOut(
 						bIsSimpleMaterial = false;
 						EmissiveColor += BSDF_GETEMISSIVE(BSDF) * BSDF.LuminanceWeight;
 
-						TopLayerCoverage += BSDF.bIsTop ? BSDF.Coverage : 0.0f;
-						TopLayerData.Roughness += BSDF.bIsTop ? BSDF.Coverage * SLW_ROUGHNESS(BSDF) : 0.0f;
+						TopLayerTotalWeight += BSDF.TopLayerWeight;
+						TopLayerData.Roughness += BSDF.TopLayerWeight * SLW_ROUGHNESS(BSDF);
 
 						break;
 					}
@@ -2584,7 +2584,7 @@ void PackStrataOut(
 
 					float3x3 TangentBasis = StrataGetBSDFSharedBasis_InlineShading(StrataPixelHeader, BSDF, StrataAddressing);
 					float3 N = TangentBasis[2];
-					TopLayerData.WorldNormal += BSDF.bIsTop ? N * max(STRATA_EPSILON, BSDF.Coverage) : 0.0;
+					TopLayerData.WorldNormal += N * BSDF.TopLayerWeight;
 
 					AddShadingModels(StrataPixelHeader, StrataShadingModelBit(BSDF_GETTYPE(BSDF)));
 					AndIsSimpleMaterial(StrataPixelHeader, bIsSimpleMaterial);
@@ -2602,8 +2602,8 @@ void PackStrataOut(
 		}
 
 		// Finalize top layer data
-		TopLayerData.WorldNormal = normalize(TopLayerData.WorldNormal);
-		TopLayerData.Roughness = TopLayerCoverage > 0.0f ? TopLayerData.Roughness / TopLayerCoverage : 0.0f;
+		TopLayerData.WorldNormal= TopLayerTotalWeight > 0.0f	? normalize(TopLayerData.WorldNormal)			: 0.0f;
+		TopLayerData.Roughness	= TopLayerTotalWeight > 0.0f	? TopLayerData.Roughness / TopLayerTotalWeight	: 0.0f;
 
 		// Set storage layout as either: fast(0), single(1), or complex(2)
 		AndIsSimpleMaterial(StrataPixelHeader, BSDFVisibleCount == 1);
diff --git a/Engine/Shaders/Private/Strata/StrataTree.ush b/Engine/Shaders/Private/Strata/StrataTree.ush
--- a/Engine/Shaders/Private/Strata/StrataTree.ush
+++ b/Engine/Shaders/Private/Strata/StrataTree.ush
@@ -84,6 +84,9 @@ void UpdateSingleBSDFOperatorCoverageTransmittance(
 	//  - When coming from non parameter blending, this is equal identity 1 and the operator parsing will continue to correctly weight the BSDF according to the tree.
 	CurrentBSDF.LuminanceWeight = float3(1.0f, 1.0f, 1.0f) * CurrentBSDFCoverage;
 
+	// All BSDF are considered for contribution as the top BSDF can have their coverage dip to 0 and then another BSDF need to be considered for the top layer representation.
+	CurrentBSDF.TopLayerWeight = 1.0f;
+
 	StrataTree.Operators[CurrentBSDF.OperatorIndex].Coverage = CurrentBSDFCoverage;
 	StrataTree.Operators[CurrentBSDF.OperatorIndex].Transmittance = CurrentBSDFTransmittance;
 
@@ -174,6 +177,7 @@ void UpdateSingleOperatorCoverageTransmittance(inout FStrataTree StrataTree, int
 }
 
 
+
 #define Op StrataTree.Operators[OpIndex]
 
 void UpdateBSDFWeightAfterVerticalOperatorVisit(
@@ -189,6 +193,9 @@ void UpdateBSDFWeightAfterVerticalOperatorVisit(
 	CurrentBSDF.LuminanceWeight *= 
 		PreviousIsInputA > 0 ? 1.0f : saturate(1.0f - Op.VerticalTopCoverage) + saturate(Op.VerticalTopCoverage) * Op.VerticalTopTransmittance;
 
+	// If the BSDF is from the bottom, we reduce its contribution according to the top layer coverage only.
+	CurrentBSDF.TopLayerWeight *= PreviousIsInputA ? 1.0 : saturate(1.0f - Op.VerticalTopCoverage);
+
 #endif // STRATA_INLINE_SHADING
 }
 
@@ -207,6 +214,8 @@ void UpdateBSDFWeightAfterHorizontalOperatorVisit(
 	CurrentBSDF.LuminanceWeight *= Weight;
 	CurrentBSDF.Coverage *= Weight;
 
+	CurrentBSDF.TopLayerWeight *= Weight;
+
 #endif // STRATA_INLINE_SHADING
 }
 
@@ -222,6 +231,8 @@ void UpdateBSDFWeightAfterWeightOperatorVisit(
 	CurrentBSDF.LuminanceWeight *= Weight;
 	CurrentBSDF.Coverage *= Weight;
 
+	CurrentBSDF.TopLayerWeight *= Weight;
+
 #endif // STRATA_INLINE_SHADING
 }
 
