Move #if STRATA_ENABLED within the strata shader files to simplify strata includes.

diff --git a/Engine/Shaders/Private/BasePassPixelShader.usf b/Engine/Shaders/Private/BasePassPixelShader.usf
--- a/Engine/Shaders/Private/BasePassPixelShader.usf
+++ b/Engine/Shaders/Private/BasePassPixelShader.usf
@@ -45,10 +45,7 @@
 #endif
 
 // For the base pass STRATA_INLINE_SHADING is defined in BasePassRendering.h
-#if STRATA_ENABLED
-#include "/Engine/Private/Strata/Strata.ush"
-#endif
-
+#include "Strata/Strata.ush"
 #include "SHCommon.ush"
 #include "/Engine/Generated/Material.ush"
 #include "BasePassCommon.ush"
diff --git a/Engine/Shaders/Private/CapsuleShadowShaders.usf b/Engine/Shaders/Private/CapsuleShadowShaders.usf
--- a/Engine/Shaders/Private/CapsuleShadowShaders.usf
+++ b/Engine/Shaders/Private/CapsuleShadowShaders.usf
@@ -11,11 +11,7 @@
 #include "SHCommon.ush"
 #include "VolumetricLightmapShared.ush"
 #include "ReflectionEnvironmentShared.ush"
-
-#if STRATA_ENABLED
-#define STRATA_INLINE_SHADING 0
-#include "/Engine/Private/Strata/Strata.ush"
-#endif
+#include "Strata/Strata.ush"
 
 #ifndef THREADGROUP_SIZEX
 #define THREADGROUP_SIZEX 1
diff --git a/Engine/Shaders/Private/ClusteredDeferredShadingPixelShader.usf b/Engine/Shaders/Private/ClusteredDeferredShadingPixelShader.usf
--- a/Engine/Shaders/Private/ClusteredDeferredShadingPixelShader.usf
+++ b/Engine/Shaders/Private/ClusteredDeferredShadingPixelShader.usf
@@ -9,6 +9,7 @@
 #define SUPPORT_CONTACT_SHADOWS 0
 #define USE_IES_PROFILE 0
 #define NON_DIRECTIONAL_DIRECT_LIGHTING 0
+#define STRATA_SSS_TRANSMISSION USE_TRANSMISSION
 
 // Used in RectLight.ush, and requires a per-light texture which we cannot support at present.
 // If the textures are atlased or something like that in the future it should be ok to turn on.
@@ -16,13 +17,7 @@
 
 #include "HairStrands/HairStrandsVisibilityCommonStruct.ush" 
 #include "Common.ush"
-
-#if STRATA_ENABLED 
-#define STRATA_INLINE_SHADING 0
-#define STRATA_SSS_TRANSMISSION USE_TRANSMISSION
-#include "/Engine/Private/Strata/Strata.ush"
-#endif
-
+#include "Strata/Strata.ush"
 #include "DeferredShadingCommon.ush"  
 #include "BRDF.ush"
 #include "ShadingModels.ush"
@@ -40,8 +35,8 @@
   #else
 	#error Strata tile type non-implemented
   #endif
-#include "/Engine/Private/Strata/StrataEvaluation.ush"
-#include "/Engine/Private/Strata/StrataDeferredLighting.ush"
+#include "Strata/StrataEvaluation.ush"
+#include "Strata/StrataDeferredLighting.ush"
 #endif
 
 #if USE_HAIR_LIGHTING == 1
diff --git a/Engine/Shaders/Private/DebugProbes.usf b/Engine/Shaders/Private/DebugProbes.usf
--- a/Engine/Shaders/Private/DebugProbes.usf
+++ b/Engine/Shaders/Private/DebugProbes.usf
@@ -2,13 +2,10 @@
 
 #include "Common.ush"
 
-
-#if STRATA_ENABLED
 // Mention that we are going to use a single slab
 #define STRATA_CLAMPED_BSDF_COUNT 1
 #define STRATA_INLINE_SHADING 1
 #include "Strata/Strata.ush"
-#endif // STRATA_ENABLED
 
 #include "DeferredShadingCommon.ush"
 #include "VelocityCommon.ush"
diff --git a/Engine/Shaders/Private/DeferredDecal.usf b/Engine/Shaders/Private/DeferredDecal.usf
--- a/Engine/Shaders/Private/DeferredDecal.usf
+++ b/Engine/Shaders/Private/DeferredDecal.usf
@@ -17,12 +17,10 @@
 
 #if PIXELSHADER
 
-#if STRATA_ENABLED 
+// Decals require both inline & deferred shading path for reading & writing Strata data
 #define STRATA_INLINE_SHADING 1
 #define STRATA_DEFERRED_SHADING 1
 #include "Strata/Strata.ush"
-#endif
-
 #include "DecalCommon.ush"
 
 #endif // PIXELSHADER
diff --git a/Engine/Shaders/Private/DeferredLightPixelShaders.usf b/Engine/Shaders/Private/DeferredLightPixelShaders.usf
--- a/Engine/Shaders/Private/DeferredLightPixelShaders.usf
+++ b/Engine/Shaders/Private/DeferredLightPixelShaders.usf
@@ -5,18 +5,14 @@
 =============================================================================*/
 
 #define SUPPORT_CONTACT_SHADOWS 1
+#define STRATA_SSS_TRANSMISSION USE_TRANSMISSION
 
 #include "HairStrands/HairStrandsVisibilityCommonStruct.ush" 
 #include "Common.ush"
 
 #include "LargeWorldCoordinates.ush"
 
-#if STRATA_ENABLED
-#define STRATA_INLINE_SHADING 0
-#define STRATA_SSS_TRANSMISSION USE_TRANSMISSION
-#include "/Engine/Private/Strata/Strata.ush"
-#endif
-
+#include "Strata/Strata.ush"
 #include "DeferredShadingCommon.ush"
 #include "DeferredLightingCommon.ush"
 
@@ -46,8 +42,8 @@
   #else
 	#error Strata tile type non-implemented
   #endif
-#include "/Engine/Private/Strata/StrataEvaluation.ush"
-#include "/Engine/Private/Strata/StrataDeferredLighting.ush"
+#include "Strata/StrataEvaluation.ush"
+#include "Strata/StrataDeferredLighting.ush"
 #endif
 
 struct FInputParams
diff --git a/Engine/Shaders/Private/DeferredLightingCommon.ush b/Engine/Shaders/Private/DeferredLightingCommon.ush
--- a/Engine/Shaders/Private/DeferredLightingCommon.ush
+++ b/Engine/Shaders/Private/DeferredLightingCommon.ush
@@ -11,20 +11,18 @@
 #include "IESLightProfilesCommon.ush"
 #include "CapsuleLightIntegrate.ush"
 #include "RectLightIntegrate.ush"
+#include "Strata/Strata.ush"
 
 // This Strata include assumes that if inline shading is required, it needs to 
 // be defined prior to DeferredLightingCommong.ush
-#if STRATA_ENABLED 
-	#include "Strata/Strata.ush"
-	// StrataStruct.MaterialTextureArray is only available to translucent materials or global shaders computing contact shadows.
-	// It is not defined for opaque materials, only the UAV is available to write the strata buffer.
-	#define STRATA_RAYCAST_ENABLED (!MATERIAL_IS_STRATA || STRATA_TRANSLUCENT_ENABLED)
-	#if STRATA_RAYCAST_ENABLED
-		#if MATERIAL_IS_STRATA
-			#define StrataRayCast StrataStruct
-		#else
-			#define StrataRayCast Strata
-		#endif
+// StrataStruct.MaterialTextureArray is only available to translucent materials or global shaders computing contact shadows.
+// It is not defined for opaque materials, only the UAV is available to write the strata buffer.
+#define STRATA_RAYCAST_ENABLED (!MATERIAL_IS_STRATA || STRATA_TRANSLUCENT_ENABLED)
+#if STRATA_RAYCAST_ENABLED
+	#if MATERIAL_IS_STRATA
+		#define StrataRayCast StrataStruct
+	#else
+		#define StrataRayCast Strata
 	#endif
 #endif
 
diff --git a/Engine/Shaders/Private/DistanceFieldScreenGridLighting.usf b/Engine/Shaders/Private/DistanceFieldScreenGridLighting.usf
--- a/Engine/Shaders/Private/DistanceFieldScreenGridLighting.usf
+++ b/Engine/Shaders/Private/DistanceFieldScreenGridLighting.usf
@@ -9,11 +9,7 @@
 #include "DistanceFieldLightingShared.ush"
 #include "DistanceFieldAOShared.ush"
 #include "GlobalDistanceFieldShared.ush"
-
-#if STRATA_ENABLED
-#define STRATA_INLINE_SHADING 0
-#include "/Engine/Private/Strata/Strata.ush"
-#endif
+#include "Strata/Strata.ush"
 
 // Pass a debug value through the pipeline instead of a bent normal
 #define PASS_THROUGH_DEBUG_VALUE 0
diff --git a/Engine/Shaders/Private/DistanceFieldShadowing.usf b/Engine/Shaders/Private/DistanceFieldShadowing.usf
--- a/Engine/Shaders/Private/DistanceFieldShadowing.usf
+++ b/Engine/Shaders/Private/DistanceFieldShadowing.usf
@@ -8,10 +8,7 @@
 #include "DeferredShadingCommon.ush"
 #include "DistanceFieldLightingShared.ush"
 #include "DistanceFieldShadowingShared.ush"
-
-#if STRATA_ENABLED
-#include "/Engine/Private/Strata/Strata.ush"
-#endif
+#include "Strata/Strata.ush"
 
 #define SCATTER_TILE_CULLING (CULLING_TYPE == 0)
 #define POINT_LIGHT (CULLING_TYPE == 2)
diff --git a/Engine/Shaders/Private/DistortAccumulatePS.usf b/Engine/Shaders/Private/DistortAccumulatePS.usf
--- a/Engine/Shaders/Private/DistortAccumulatePS.usf
+++ b/Engine/Shaders/Private/DistortAccumulatePS.usf
@@ -27,13 +27,10 @@
 #endif
 
 #if STRATA_ENABLED
-#define STRATA_USE_ENV_BRDF_APPROX 1
-#include "/Engine/Private/Strata/StrataEvaluation.ush"
-
-#include "/Engine/Private/Strata/StrataStatisticalOperators.ush"
+#include "Strata/StrataEvaluation.ush"
+#include "Strata/StrataStatisticalOperators.ush"
 #endif
 
-
 void ClipOccludedPixel(float PixelZ, float SceneZ)
 {
 #if	HAS_INVERTED_Z_BUFFER
diff --git a/Engine/Shaders/Private/LightFunctionCommon.ush b/Engine/Shaders/Private/LightFunctionCommon.ush
--- a/Engine/Shaders/Private/LightFunctionCommon.ush
+++ b/Engine/Shaders/Private/LightFunctionCommon.ush
@@ -4,9 +4,7 @@
 	LightFunctionCommon.usf: Utility functions for light functions.
 =============================================================================*/
 
-#if STRATA_ENABLED
 #include "/Engine/Private/Strata/Strata.ush" 
-#endif
 
 /** Tan of spotlight cone outer angle in x, ShadowFadeFraction in y, IsSpotLight in z, IsPointLight in w. */
 float4 LightFunctionParameters;
diff --git a/Engine/Shaders/Private/MeshDecals.usf b/Engine/Shaders/Private/MeshDecals.usf
--- a/Engine/Shaders/Private/MeshDecals.usf
+++ b/Engine/Shaders/Private/MeshDecals.usf
@@ -15,10 +15,7 @@
 	#define EyeAdaptationStruct DecalPass
 #endif
 
-#if STRATA_ENABLED
 #include "Strata/Strata.ush"
-#endif
-
 #include "DecalCommon.ush"
 #include "/Engine/Generated/Material.ush"
 #include "/Engine/Generated/VertexFactory.ush"
diff --git a/Engine/Shaders/Private/PostProcessAmbientOcclusion.usf b/Engine/Shaders/Private/PostProcessAmbientOcclusion.usf
--- a/Engine/Shaders/Private/PostProcessAmbientOcclusion.usf
+++ b/Engine/Shaders/Private/PostProcessAmbientOcclusion.usf
@@ -8,11 +8,7 @@
 #include "ScreenPass.ush"
 #include "PostProcessCommon.ush"				
 #include "DeferredShadingCommon.ush"
-
-#if STRATA_ENABLED
-#define STRATA_INLINE_SHADING 0
-#include "/Engine/Private/Strata/Strata.ush"
-#endif
+#include "Strata/Strata.ush"
 
 // set by C++:
 //
diff --git a/Engine/Shaders/Private/PostProcessSubsurface.usf b/Engine/Shaders/Private/PostProcessSubsurface.usf
--- a/Engine/Shaders/Private/PostProcessSubsurface.usf
+++ b/Engine/Shaders/Private/PostProcessSubsurface.usf
@@ -5,11 +5,7 @@
 =============================================================================*/
 
 #include "Common.ush"
-
-#ifdef STRATA_ENABLED
 #include "Strata/Strata.ush"
-#endif
-
 #include "PostProcessCommon.ush"
 #include "DeferredShadingCommon.ush"
 #include "SubsurfaceProfileCommon.ush"
diff --git a/Engine/Shaders/Private/RayTracing/RayTracingAmbientOcclusionRGS.usf b/Engine/Shaders/Private/RayTracing/RayTracingAmbientOcclusionRGS.usf
--- a/Engine/Shaders/Private/RayTracing/RayTracingAmbientOcclusionRGS.usf
+++ b/Engine/Shaders/Private/RayTracing/RayTracingAmbientOcclusionRGS.usf
@@ -6,10 +6,7 @@
 #include "../SceneTextureParameters.ush"
 #include "../PathTracing/Utilities/PathTracingRandomSequence.ush"
 #include "../SphericalGaussian.ush"
-#if STRATA_ENABLED
-#define STRATA_INLINE_SHADING 0
 #include "../Strata/Strata.ush"
-#endif
 #include "RayTracingCommon.ush"
 #include "RayTracingDeferredShadingCommon.ush"
 
diff --git a/Engine/Shaders/Private/RayTracing/RayTracingOcclusionRGS.usf b/Engine/Shaders/Private/RayTracing/RayTracingOcclusionRGS.usf
--- a/Engine/Shaders/Private/RayTracing/RayTracingOcclusionRGS.usf
+++ b/Engine/Shaders/Private/RayTracing/RayTracingOcclusionRGS.usf
@@ -11,10 +11,7 @@
 #include "../ScreenSpaceDenoise/SSDPublic.ush"
 #include "../TransmissionCommon.ush"
 #include "../HairStrands/HairStrandsVisibilityCommon.ush"
-#if STRATA_ENABLED
-#define STRATA_INLINE_SHADING 0
 #include "../Strata/Strata.ush"
-#endif
 #include "RayTracingCommon.ush"
 #include "RayTracingDeferredShadingCommon.ush"
 #include "RayTracingDirectionalLight.ush"
diff --git a/Engine/Shaders/Private/SSRT/SSRTReflections.usf b/Engine/Shaders/Private/SSRT/SSRTReflections.usf
--- a/Engine/Shaders/Private/SSRT/SSRTReflections.usf
+++ b/Engine/Shaders/Private/SSRT/SSRTReflections.usf
@@ -6,11 +6,7 @@
 #include "../MonteCarlo.ush"
 #include "../SceneTextureParameters.ush"
 #include "../ScreenSpaceDenoise/SSDPublic.ush"
-
-#if STRATA_ENABLED
-#define STRATA_INLINE_SHADING 0
-#include "/Engine/Private/Strata/Strata.ush"
-#endif
+#include "../Strata/Strata.ush"
 
 // Generate vector truncation warnings to errors.
 #pragma warning(error: 3206)
diff --git a/Engine/Shaders/Private/ShadingFurnaceTest.usf b/Engine/Shaders/Private/ShadingFurnaceTest.usf
--- a/Engine/Shaders/Private/ShadingFurnaceTest.usf
+++ b/Engine/Shaders/Private/ShadingFurnaceTest.usf
@@ -8,13 +8,8 @@
 #define SHADING_PATH_DEFERRED 1
 #define USE_SOURCE_TEXTURE 0
 #define DEBUG_ENABLE 0
-
-#if STRATA_ENABLED
-#define STRATA_INLINE_SHADING 0
 #define STRATA_SSS_TRANSMISSION USE_TRANSMISSION
-#include "/Engine/Private/Strata/Strata.ush"
-#endif
-
+#include "Strata/Strata.ush"
 #include "DeferredShadingCommon.ush"
 #include "DeferredLightingCommon.ush"
 #include "MonteCarlo.ush" 
@@ -24,10 +19,8 @@
 #include "ColorMap.ush"
 #endif 
 
-#if STRATA_ENABLED
-#include "/Engine/Private/Strata/StrataEvaluation.ush"
-#include "/Engine/Private/Strata/StrataDeferredLighting.ush"
-#endif
+#include "Strata/StrataEvaluation.ush"
+#include "Strata/StrataDeferredLighting.ush"
 
 #if SHADER_FURNACE_ANALYTIC
 
diff --git a/Engine/Shaders/Private/ShadowProjectionPixelShader.usf b/Engine/Shaders/Private/ShadowProjectionPixelShader.usf
--- a/Engine/Shaders/Private/ShadowProjectionPixelShader.usf
+++ b/Engine/Shaders/Private/ShadowProjectionPixelShader.usf
@@ -89,10 +89,7 @@ float3 ModulatedShadowBlendOp(float3 Source)
 float2 ShadowNearAndFarDepth;
 #endif
 
-#if STRATA_ENABLED
-#define STRATA_INLINE_SHADING 0
-#include "/Engine/Private/Strata/Strata.ush"
-#endif
+#include "Strata/Strata.ush"
 
 /**
 * Entry point for uniform manual PCF that supports lights using normal shadows.
diff --git a/Engine/Shaders/Private/SingleLayerWaterComposite.usf b/Engine/Shaders/Private/SingleLayerWaterComposite.usf
--- a/Engine/Shaders/Private/SingleLayerWaterComposite.usf
+++ b/Engine/Shaders/Private/SingleLayerWaterComposite.usf
@@ -21,11 +21,7 @@
 #define REFLECTION_COMPOSITE_USE_BLENDED_REFLECTION_CAPTURES (FEATURE_LEVEL >= FEATURE_LEVEL_SM5)
 #define REFLECTION_COMPOSITE_SUPPORT_SKYLIGHT_BLEND 0
 #include "ReflectionEnvironmentComposite.ush"
-
-#if STRATA_ENABLED
-#define STRATA_INLINE_SHADING 0
 #include "Strata/StrataEvaluation.ush"
-#endif
 
 Texture2D ScreenSpaceReflectionsTexture;
 SamplerState ScreenSpaceReflectionsSampler;
diff --git a/Engine/Shaders/Private/Strata/Strata.ush b/Engine/Shaders/Private/Strata/Strata.ush
--- a/Engine/Shaders/Private/Strata/Strata.ush
+++ b/Engine/Shaders/Private/Strata/Strata.ush
@@ -2,6 +2,14 @@
 
 #pragma once
 
+// Sanity guard. 
+#ifndef STRATA_ENABLED
+	#define STRATA_ENABLED 1
+	#error STRATA_ENABLED needs to be defined
+#endif
+
+#if STRATA_ENABLED
+
 #include "../ShadingEnergyConservation.ush"
 #include "../BRDF.ush"
 #include "../ColorSpace.ush"
@@ -3107,3 +3115,5 @@ FStrataBSDF UnpackStrataBSDFIn(FStrataMaterialContainer StrataBuffer, inout FStr
 
 	return OutBSDF;
 }
+
+#endif // STRATA_ENABLED
\ No newline at end of file
diff --git a/Engine/Shaders/Private/Strata/StrataDeferredLighting.ush b/Engine/Shaders/Private/Strata/StrataDeferredLighting.ush
--- a/Engine/Shaders/Private/Strata/StrataDeferredLighting.ush
+++ b/Engine/Shaders/Private/Strata/StrataDeferredLighting.ush
@@ -4,6 +4,14 @@
 
 #include "../AreaLightCommon.ush"
 
+// Sanity guard. 
+#ifndef STRATA_ENABLED
+	#define STRATA_ENABLED 1
+	#error STRATA_ENABLED needs to be defined
+#endif
+
+#if STRATA_ENABLED
+
 float4 StrataReadPrecomputedShadowFactors(FStrataPixelHeader StrataPixelHeader, int2 PixelPos, Texture2D PrecomputedShadowTexture)
 {
 	if (HasPrecShadowMask(StrataPixelHeader))
@@ -80,3 +88,5 @@ float4 StrataDeferredLighting(
 	const float4 OutRadiance = SplitLighting.SpecularLighting + SplitLighting.DiffuseLighting;
 	return OutRadiance;
 }
+
+#endif // STRATA_ENABLED
\ No newline at end of file
diff --git a/Engine/Shaders/Private/Strata/StrataEnvironmentLighting.ush b/Engine/Shaders/Private/Strata/StrataEnvironmentLighting.ush
--- a/Engine/Shaders/Private/Strata/StrataEnvironmentLighting.ush
+++ b/Engine/Shaders/Private/Strata/StrataEnvironmentLighting.ush
@@ -2,6 +2,14 @@
 
 #pragma once
 
+// Sanity guard. 
+#ifndef STRATA_ENABLED
+	#define STRATA_ENABLED 1
+	#error STRATA_ENABLED needs to be defined
+#endif
+
+#if STRATA_ENABLED
+
 // Env. lighting evaluation for strata material.
 // Unpack BSDF on-the-fly
 float4 StrataEnvironmentLighting(
@@ -182,3 +190,4 @@ float4 StrataEnvironmentLighting(
 	return LightAccumulator_GetResult(Out);
 }
 
+#endif // STRATA_ENABLED
\ No newline at end of file
diff --git a/Engine/Shaders/Private/Strata/StrataEvaluation.ush b/Engine/Shaders/Private/Strata/StrataEvaluation.ush
--- a/Engine/Shaders/Private/Strata/StrataEvaluation.ush
+++ b/Engine/Shaders/Private/Strata/StrataEvaluation.ush
@@ -2,6 +2,14 @@
 
 #pragma once
 
+// Sanity guard. 
+#ifndef STRATA_ENABLED
+	#define STRATA_ENABLED 1
+	#error STRATA_ENABLED needs to be defined
+#endif
+
+#if STRATA_ENABLED
+
 #include "Strata.ush"
 
 #include "../ParticipatingMediaCommon.ush"
@@ -1053,4 +1061,4 @@ FThreeBandSHVector StrataBSDFToSH(FStrataBSDFContext BSDFContext)
 	return SHVector;
 }
 
-
+#endif // STRATA_ENABLED
\ No newline at end of file
diff --git a/Engine/Shaders/Private/SubsurfaceBurleyNormalized.ush b/Engine/Shaders/Private/SubsurfaceBurleyNormalized.ush
--- a/Engine/Shaders/Private/SubsurfaceBurleyNormalized.ush
+++ b/Engine/Shaders/Private/SubsurfaceBurleyNormalized.ush
@@ -8,10 +8,7 @@
 #include "Random.ush"
 #include "DeferredShadingCommon.ush"
 #include "MonteCarlo.ush"
-
-#if STRATA_ENABLED
 #include "Strata/Strata.ush"
-#endif
 
 // Setup the max number of Burley samples
 #define BURLEY_NUM_SAMPLES	64
diff --git a/Engine/Shaders/Private/VirtualShadowMaps/VirtualShadowMapPageManagement.usf b/Engine/Shaders/Private/VirtualShadowMaps/VirtualShadowMapPageManagement.usf
--- a/Engine/Shaders/Private/VirtualShadowMaps/VirtualShadowMapPageManagement.usf
+++ b/Engine/Shaders/Private/VirtualShadowMaps/VirtualShadowMapPageManagement.usf
@@ -15,6 +15,7 @@
 #include "../HairStrands/HairStrandsVisibilityCommon.ush"
 #include "../HairStrands/HairStrandsTileCommon.ush"
 #include "../ReductionCommon.ush"
+#include "../Strata/Strata.ush"
 #include "VirtualShadowMapPageAccessCommon.ush"
 #include "VirtualShadowMapPageCacheCommon.ush"
 #include "VirtualShadowMapPageOverlap.ush"
@@ -32,11 +33,6 @@
 #define INPUT_TYPE_GBUFFER 0
 #define INPUT_TYPE_HAIRSTRANDS 1
 
-#if STRATA_ENABLED 
-#define STRATA_INLINE_SHADING 0
-#include "/Engine/Private/Strata/Strata.ush"
-#endif
-
 // Flags generated by per-pixel pass to determine which pages are required to provide shadow for the visible geometry
 RWStructuredBuffer<uint> OutPageRequestFlags;
 
diff --git a/Engine/Shaders/Private/VirtualShadowMaps/VirtualShadowMapProjection.usf b/Engine/Shaders/Private/VirtualShadowMaps/VirtualShadowMapProjection.usf
--- a/Engine/Shaders/Private/VirtualShadowMaps/VirtualShadowMapProjection.usf
+++ b/Engine/Shaders/Private/VirtualShadowMaps/VirtualShadowMapProjection.usf
@@ -7,10 +7,7 @@
 #define SUPPORT_CONTACT_SHADOWS 1
 #include "../Common.ush"
 #include "../SceneTexturesCommon.ush"
-#if STRATA_ENABLED 
-#define STRATA_INLINE_SHADING 0
-#include "/Engine/Private/Strata/Strata.ush"
-#endif
+#include "../Strata/Strata.ush"
 #include "../DeferredShadingCommon.ush"
 #include "../DeferredLightingCommon.ush"
 #include "../LightShaderParameters.ush"
diff --git a/Engine/Shaders/Private/VirtualShadowMaps/VirtualShadowMapTransmissionCommon.ush b/Engine/Shaders/Private/VirtualShadowMaps/VirtualShadowMapTransmissionCommon.ush
--- a/Engine/Shaders/Private/VirtualShadowMaps/VirtualShadowMapTransmissionCommon.ush
+++ b/Engine/Shaders/Private/VirtualShadowMaps/VirtualShadowMapTransmissionCommon.ush
@@ -10,10 +10,6 @@ VirtualShadowMaps/VirtualShadowMapTransmissionCommon.ush:
 #include "../DeferredShadingCommon.ush"
 #include "VirtualShadowMapProjectionCommon.ush"
 
-#ifndef STRATA_ENABLED
-#error STRATA_ENABLED needs to be defined
-#endif
-
 half GetSubsurfaceOpacityFromGbuffer(uint2 PixelPos)
 {
 	half Opacity = 1.0;
