Strata minor code clean up in preparation for more statistical operator.

diff --git a/Engine/Shaders/Private/DebugProbes.usf b/Engine/Shaders/Private/DebugProbes.usf
--- a/Engine/Shaders/Private/DebugProbes.usf
+++ b/Engine/Shaders/Private/DebugProbes.usf
@@ -18,7 +18,7 @@ uint MaxBytesPerPixel;
 // Default functions for PackStrataOut to function properly
 void UpdateAllBSDFsOperatorCoverageTransmittance(inout FStrataPixelHeader StrataPixelHeader, inout FStrataTree StrataTree, bool bRoughDiffuseEnabled, FStrataAddressing NullStrataAddressing, float3 V, float3 L) {}
 void UpdateAllOperatorsCoverageTransmittance(inout FStrataTree StrataTree) {}
-void UpdateAllBSDFWeightAfterOperatorVisit(inout FStrataTree StrataTree) {}
+void UpdateAllBSDFWithBottomUpOperatorVisit(inout FStrataTree StrataTree) {}
 
 #endif // STRATA_ENABLED
 
diff --git a/Engine/Shaders/Private/DistortAccumulatePS.usf b/Engine/Shaders/Private/DistortAccumulatePS.usf
--- a/Engine/Shaders/Private/DistortAccumulatePS.usf
+++ b/Engine/Shaders/Private/DistortAccumulatePS.usf
@@ -28,7 +28,6 @@
 
 #if STRATA_ENABLED
 #include "Strata/StrataEvaluation.ush"
-#include "Strata/StrataStatisticalOperators.ush"
 #endif
 
 void ClipOccludedPixel(float PixelZ, float SceneZ)
@@ -128,7 +127,7 @@ void Main(
 		//
 		// Update the luminance weight of each BSDF according to the operators it has to traverse up to the strata tree root node.
 		//
-		UpdateAllBSDFWeightAfterOperatorVisit(StrataTree);
+		UpdateAllBSDFWithBottomUpOperatorVisit(StrataTree);
 	#else
 		LegacyUpdateBSDFsOperators(StrataPixelHeader, StrataTree, bRoughDiffuseEnabled, NullStrataAddressing, V, DummyL);
 	#endif
diff --git a/Engine/Shaders/Private/Strata/Strata.ush b/Engine/Shaders/Private/Strata/Strata.ush
--- a/Engine/Shaders/Private/Strata/Strata.ush
+++ b/Engine/Shaders/Private/Strata/Strata.ush
@@ -21,7 +21,8 @@
 #include "../OctahedralCommon.ush"
 #include "../SubsurfaceProfileCommon.ush"
 #include "../BurleyNormalizedSSSCommon.ush"
-#include "/Engine/Shared/StrataDefinitions.h"     
+#include "/Engine/Shared/StrataDefinitions.h"
+#include "/Engine/Private/Strata/StrataStatisticalOperators.ush"
 
 
 // DXC has issue with unrolling certain part of the code. To address temporarly the warning-as-error issues, 
@@ -127,7 +128,7 @@ FRWStrataMaterialContainer InitialiseRWStrataMaterialContainer(RWTexture2DArray<
 
 struct FStrataBSDF
 {
-	uint				State;		// Extra informations and other enabled BSDF features
+	uint				State;			// Extra informations and other enabled BSDF features
 
 	// STRATA_TODO pack the following
 	int					OperatorIndex;
@@ -137,13 +138,13 @@ struct FStrataBSDF
 
 #if STRATA_INLINE_SHADING
 
-	float				Coverage;	// The coverage of the material, as ratio of visible matter in this pixel.
+	float				Coverage;		// The coverage of the material, as ratio of visible matter in this pixel.
 
-	float3				Emissive;	// All bsdf can potientially have emissive. This is never written out to the strata gbuffer but instead accumulated in the emissive buffer during the base pass.
-	float				ThicknessCm;// This must be kept in order to be able to normalize the mean free path to comply with our volumetric material math all done in a normalised slab of medium.
+	float3				Emissive;		// All bsdf can potientially have emissive. This is never written out to the strata gbuffer but instead accumulated in the emissive buffer during the base pass.
+	float				ThicknessCm;	// This must be kept in order to be able to normalize the mean free path to comply with our volumetric material math all done in a normalised slab of medium.
 	float3				TmpDMFP;
 
-	float				TopLayerWeight;// The weight this BSDF contributes to the top layer (normal and roughness for SSR and SSAO for instance)
+	float				TopLayerDataWeight;// The weight this BSDF contributes to the top layer (normal and roughness for SSR and SSAO for instance)
 #endif
 
 	// Water has too many parameters to be stores in the extra VGPRs below. So we extend this structure for the case where it is evaluated inline only.
@@ -2450,7 +2451,7 @@ void StrataConvertFromDBuffer(
 #if STRATA_USES_CONVERSION_FROM_LEGACY==0
 void UpdateAllBSDFsOperatorCoverageTransmittance(inout FStrataPixelHeader StrataPixelHeader, inout FStrataTree StrataTree, bool bRoughDiffuseEnabled, FStrataAddressing NullStrataAddressing, float3 V, float3 L);
 void UpdateAllOperatorsCoverageTransmittance(inout FStrataTree StrataTree);
-void UpdateAllBSDFWeightAfterOperatorVisit(inout FStrataTree StrataTree);
+void UpdateAllBSDFWithBottomUpOperatorVisit(inout FStrataTree StrataTree);
 #else
 void LegacyUpdateBSDFsOperators(inout FStrataPixelHeader StrataPixelHeader, inout FStrataTree StrataTree, bool bRoughDiffuseEnabled, FStrataAddressing NullStrataAddressing, float3 V, float3 L);
 #endif
@@ -2512,7 +2513,7 @@ void PackStrataOut(
 		//
 		// Update the luminance weight of each BSDF according to the operators it has to traverse up to the strata tree root node.
 		//
-		UpdateAllBSDFWeightAfterOperatorVisit(StrataTree);
+		UpdateAllBSDFWithBottomUpOperatorVisit(StrataTree);
 
 #else
 		LegacyUpdateBSDFsOperators(StrataPixelHeader, StrataTree, bRoughDiffuseEnabled, NullStrataAddressing, V, NullLightDir);
@@ -2555,8 +2556,8 @@ void PackStrataOut(
 						// they use the complex pass, as the toplayer data does not contain the full frame basis, only the top normal.
 						bIsSingleMaterial = BSDF_GETHASANISOTROPY(BSDF) == 0 && LuminanceWeightEqualOne;
 
-						TopLayerTotalWeight += BSDF.TopLayerWeight;
-						TopLayerData.Roughness += BSDF.TopLayerWeight * SLAB_ROUGHNESS(BSDF);
+						TopLayerTotalWeight += BSDF.TopLayerDataWeight;
+						TopLayerData.Roughness += BSDF.TopLayerDataWeight * SLAB_ROUGHNESS(BSDF);
 
 						const bool bIsThin = BSDF_GETISTHIN(BSDF);
 						const bool bIsSimpleVolume = BSDF_GETISSIMPLEVOLUME(BSDF);
@@ -2572,8 +2573,8 @@ void PackStrataOut(
 					{
 						bIsSimpleMaterial = false;
 
-						TopLayerTotalWeight += BSDF.TopLayerWeight;
-						TopLayerData.Roughness += BSDF.TopLayerWeight * FullyRough;
+						TopLayerTotalWeight += BSDF.TopLayerDataWeight;
+						TopLayerData.Roughness += BSDF.TopLayerDataWeight * FullyRough;
 
 						break;
 					}
@@ -2583,8 +2584,8 @@ void PackStrataOut(
 						bIsSimpleMaterial = false;
 						EmissiveColor += BSDF_GETEMISSIVE(BSDF) * BSDF.LuminanceWeight;
 
-						TopLayerTotalWeight += BSDF.TopLayerWeight;
-						TopLayerData.Roughness += BSDF.TopLayerWeight * SLW_ROUGHNESS(BSDF);
+						TopLayerTotalWeight += BSDF.TopLayerDataWeight;
+						TopLayerData.Roughness += BSDF.TopLayerDataWeight * SLW_ROUGHNESS(BSDF);
 
 						break;
 					}
@@ -2592,7 +2593,7 @@ void PackStrataOut(
 
 					float3x3 TangentBasis = StrataGetBSDFSharedBasis_InlineShading(StrataPixelHeader, BSDF, StrataAddressing);
 					float3 N = TangentBasis[2];
-					TopLayerData.WorldNormal += N * BSDF.TopLayerWeight;
+					TopLayerData.WorldNormal += N * BSDF.TopLayerDataWeight;
 
 					AddShadingModels(StrataPixelHeader, StrataShadingModelBit(BSDF_GETTYPE(BSDF)));
 					AndIsSimpleMaterial(StrataPixelHeader, bIsSimpleMaterial);
diff --git a/Engine/Shaders/Private/Strata/StrataForwardLighting.ush b/Engine/Shaders/Private/Strata/StrataForwardLighting.ush
--- a/Engine/Shaders/Private/Strata/StrataForwardLighting.ush
+++ b/Engine/Shaders/Private/Strata/StrataForwardLighting.ush
@@ -272,7 +272,7 @@ float3 StrataForwardLighting(
 		//
 		// Update the luminance weight of each BSDF according to the operators it has to traverse up to the strata tree root node.
 		//
-		UpdateAllBSDFWeightAfterOperatorVisit(StrataTree);
+		UpdateAllBSDFWithBottomUpOperatorVisit(StrataTree);
 #else
 		LegacyUpdateBSDFsOperators(StrataPixelHeader, StrataTree, bRoughDiffuseEnabled, NullStrataAddressing, V, L);
 #endif
diff --git a/Engine/Shaders/Private/Strata/StrataTree.ush b/Engine/Shaders/Private/Strata/StrataTree.ush
--- a/Engine/Shaders/Private/Strata/StrataTree.ush
+++ b/Engine/Shaders/Private/Strata/StrataTree.ush
@@ -85,7 +85,7 @@ void UpdateSingleBSDFOperatorCoverageTransmittance(
 	CurrentBSDF.LuminanceWeight = float3(1.0f, 1.0f, 1.0f) * CurrentBSDFCoverage;
 
 	// All BSDF are considered for contribution as the top BSDF can have their coverage dip to 0 and then another BSDF need to be considered for the top layer representation.
-	CurrentBSDF.TopLayerWeight = 1.0f;
+	CurrentBSDF.TopLayerDataWeight = 1.0f;
 
 	StrataTree.Operators[CurrentBSDF.OperatorIndex].Coverage = CurrentBSDFCoverage;
 	StrataTree.Operators[CurrentBSDF.OperatorIndex].Transmittance = CurrentBSDFTransmittance;
@@ -180,7 +180,8 @@ void UpdateSingleOperatorCoverageTransmittance(inout FStrataTree StrataTree, int
 
 #define Op StrataTree.Operators[OpIndex]
 
-void UpdateBSDFWeightAfterVerticalOperatorVisit(
+
+void UpdateAllBSDFWithBottomUpOperatorVisit_Weight(
 	inout FStrataTree StrataTree,
 	inout FStrataBSDF CurrentBSDF,
 	int OpIndex,
@@ -188,18 +189,16 @@ void UpdateBSDFWeightAfterVerticalOperatorVisit(
 {
 #if STRATA_INLINE_SHADING
 
-	// PreviousIsInputA > 0 means it is a BSDF affecting the top layer so we do not affect it by anything for this vertical layering.
-	// Otherise, the BSDF comes from the bottom part so we affect it weights from the gathered top layer coverage/transmittance
-	CurrentBSDF.LuminanceWeight *= 
-		PreviousIsInputA > 0 ? 1.0f : saturate(1.0f - Op.VerticalTopCoverage) + saturate(Op.VerticalTopCoverage) * Op.VerticalTopTransmittance;
+	const float Weight = saturate(Op.Weight);
+	CurrentBSDF.LuminanceWeight *= Weight;
+	CurrentBSDF.Coverage *= Weight;
 
-	// If the BSDF is from the bottom, we reduce its contribution according to the top layer coverage only.
-	CurrentBSDF.TopLayerWeight *= PreviousIsInputA ? 1.0 : saturate(1.0f - Op.VerticalTopCoverage);
+	CurrentBSDF.TopLayerDataWeight *= Weight;
 
 #endif // STRATA_INLINE_SHADING
 }
 
-void UpdateBSDFWeightAfterHorizontalOperatorVisit(
+void UpdateAllBSDFWithBottomUpOperatorVisit_Horizontal(
 	inout FStrataTree StrataTree,
 	inout FStrataBSDF CurrentBSDF,
 	int OpIndex,
@@ -214,12 +213,12 @@ void UpdateBSDFWeightAfterHorizontalOperatorVisit(
 	CurrentBSDF.LuminanceWeight *= Weight;
 	CurrentBSDF.Coverage *= Weight;
 
-	CurrentBSDF.TopLayerWeight *= Weight;
+	CurrentBSDF.TopLayerDataWeight *= Weight;
 
 #endif // STRATA_INLINE_SHADING
 }
 
-void UpdateBSDFWeightAfterWeightOperatorVisit(
+void UpdateAllBSDFWithBottomUpOperatorVisit_Vertical(
 	inout FStrataTree StrataTree,
 	inout FStrataBSDF CurrentBSDF,
 	int OpIndex,
@@ -227,11 +226,13 @@ void UpdateBSDFWeightAfterWeightOperatorVisit(
 {
 #if STRATA_INLINE_SHADING
 
-	const float Weight = saturate(Op.Weight);
-	CurrentBSDF.LuminanceWeight *= Weight;
-	CurrentBSDF.Coverage *= Weight;
+	// PreviousIsInputA > 0 means it is a BSDF affecting the top layer so we do not affect it by anything for this vertical layering.
+	// Otherise, the BSDF comes from the bottom part so we affect it weights from the gathered top layer coverage/transmittance
+	CurrentBSDF.LuminanceWeight *=
+		PreviousIsInputA > 0 ? 1.0f : saturate(1.0f - Op.VerticalTopCoverage) + saturate(Op.VerticalTopCoverage) * Op.VerticalTopTransmittance;
 
-	CurrentBSDF.TopLayerWeight *= Weight;
+	// If the BSDF is from the bottom, we reduce its contribution according to the top layer coverage only.
+	CurrentBSDF.TopLayerDataWeight *= PreviousIsInputA ? 1.0 : saturate(1.0f - Op.VerticalTopCoverage);
 
 #endif // STRATA_INLINE_SHADING
 }
@@ -259,15 +260,15 @@ void LegacyUpdateBSDFsOperators(inout FStrataPixelHeader StrataPixelHeader, inou
 	UpdateSingleOperatorCoverageTransmittance(StrataTree, /*OperatorIndex*/5);
 #endif
 
-	UpdateBSDFWeightAfterWeightOperatorVisit	(StrataTree, StrataTree.BSDFs[0], /*OperatorIndex*/2, /*PreviousIsInputA*/1); // WEIGHT
-	UpdateBSDFWeightAfterVerticalOperatorVisit	(StrataTree, StrataTree.BSDFs[0], /*OperatorIndex*/4, /*PreviousIsInputA*/1); // VERTICAL
+	UpdateAllBSDFWithBottomUpOperatorVisit_Weight	(StrataTree, StrataTree.BSDFs[0], /*OperatorIndex*/2, /*PreviousIsInputA*/1);
+	UpdateAllBSDFWithBottomUpOperatorVisit_Vertical	(StrataTree, StrataTree.BSDFs[0], /*OperatorIndex*/4, /*PreviousIsInputA*/1);
 #if STRATA_LEGACY_MATERIAL_APPLIES_FINAL_WEIGHT
-	UpdateBSDFWeightAfterWeightOperatorVisit	(StrataTree, StrataTree.BSDFs[0], /*OperatorIndex*/5, /*PreviousIsInputA*/1); // WEIGHT
+	UpdateAllBSDFWithBottomUpOperatorVisit_Weight	(StrataTree, StrataTree.BSDFs[0], /*OperatorIndex*/5, /*PreviousIsInputA*/1);
 #endif
-	UpdateBSDFWeightAfterWeightOperatorVisit	(StrataTree, StrataTree.BSDFs[1], /*OperatorIndex*/3, /*PreviousIsInputA*/1); // WEIGHT
-	UpdateBSDFWeightAfterVerticalOperatorVisit	(StrataTree, StrataTree.BSDFs[1], /*OperatorIndex*/4, /*PreviousIsInputA*/0); // VERTICAL
+	UpdateAllBSDFWithBottomUpOperatorVisit_Weight	(StrataTree, StrataTree.BSDFs[1], /*OperatorIndex*/3, /*PreviousIsInputA*/1);
+	UpdateAllBSDFWithBottomUpOperatorVisit_Vertical	(StrataTree, StrataTree.BSDFs[1], /*OperatorIndex*/4, /*PreviousIsInputA*/0);
 #if STRATA_LEGACY_MATERIAL_APPLIES_FINAL_WEIGHT
-	UpdateBSDFWeightAfterWeightOperatorVisit	(StrataTree, StrataTree.BSDFs[1], /*OperatorIndex*/5, /*PreviousIsInputA*/1); // WEIGHT
+	UpdateAllBSDFWithBottomUpOperatorVisit_Weight	(StrataTree, StrataTree.BSDFs[1], /*OperatorIndex*/5, /*PreviousIsInputA*/1);
 #endif
 }
 
diff --git a/Engine/Source/Runtime/Engine/Private/Materials/HLSLMaterialTranslator.cpp b/Engine/Source/Runtime/Engine/Private/Materials/HLSLMaterialTranslator.cpp
--- a/Engine/Source/Runtime/Engine/Private/Materials/HLSLMaterialTranslator.cpp
+++ b/Engine/Source/Runtime/Engine/Private/Materials/HLSLMaterialTranslator.cpp
@@ -1402,6 +1402,7 @@ bool FHLSLMaterialTranslator::Translate()
 				}
 
 				// Propagate up the coverage/transmittance of each node in the strata tree.
+				// For that we visit all the operator according to their distance from the Strata tree leaves, from small to large.
 				{
 					check(StrataMaterialRootOperator);
 					int32 RootMaximumDistanceToLeaves = StrataMaterialRootOperator->MaxDistanceFromLeaves;
@@ -1422,13 +1423,13 @@ bool FHLSLMaterialTranslator::Translate()
 					ResourcesString += "}\n";
 				}
 
-				// Update the luminance weight of each BSDF according to the operators it has to traverse up to the strata tree root node.
+				// Update the luminance weight of each BSDF according to the operators it has to traverse bottom-up to the strata tree root node.
 				{
 					// Update the coverage/transmittance of each node in the graph
 					check(StrataMaterialRootOperator);
 					int32 RootMaximumDistanceToLeaves = StrataMaterialRootOperator->MaxDistanceFromLeaves;
 
-					ResourcesString += "void UpdateAllBSDFWeightAfterOperatorVisit(inout FStrataTree StrataTree)\n";
+					ResourcesString += "void UpdateAllBSDFWithBottomUpOperatorVisit(inout FStrataTree StrataTree)\n";
 					ResourcesString += "{\n";
 					for (uint32 BSDFIndex = 0; BSDFIndex < StrataMaterialBSDFCount; ++BSDFIndex)
 					{
@@ -1441,20 +1442,20 @@ bool FHLSLMaterialTranslator::Translate()
 								{
 									switch (CurrentOperator.OperatorType)
 									{
-									case STRATA_OPERATOR_VERTICAL:
+									case STRATA_OPERATOR_WEIGHT:
 									{
-										ResourcesString += FString::Printf(TEXT("\t UpdateBSDFWeightAfterVerticalOperatorVisit(StrataTree, StrataTree.BSDFs[%d], %d /*Op index*/, %d /*PreviousIsInputA*/);\n"), BSDFIndex, CurrentOperator.Index, CurrentOperator.LeftIndex == PreviousOperatorIndex ? 1 : 0);
+										ResourcesString += FString::Printf(TEXT("\t UpdateAllBSDFWithBottomUpOperatorVisit_Weight(StrataTree, StrataTree.BSDFs[%d], %d /*Op index*/, %d /*PreviousIsInputA*/);\n"), BSDFIndex, CurrentOperator.Index, 1);
 										break;
 									}
 									case STRATA_OPERATOR_HORIZONTAL:
 									{
-										ResourcesString += FString::Printf(TEXT("\t UpdateBSDFWeightAfterHorizontalOperatorVisit(StrataTree, StrataTree.BSDFs[%d], %d /*Op index*/, %d /*PreviousIsInputA*/);\n"), BSDFIndex, CurrentOperator.Index, CurrentOperator.LeftIndex == PreviousOperatorIndex ? 1 : 0);
+										ResourcesString += FString::Printf(TEXT("\t UpdateAllBSDFWithBottomUpOperatorVisit_Horizontal(StrataTree, StrataTree.BSDFs[%d], %d /*Op index*/, %d /*PreviousIsInputA*/);\n"), BSDFIndex, CurrentOperator.Index, CurrentOperator.LeftIndex == PreviousOperatorIndex ? 1 : 0);
 										break;
 									}
 
-									case STRATA_OPERATOR_WEIGHT:
+									case STRATA_OPERATOR_VERTICAL:
 									{
-										ResourcesString += FString::Printf(TEXT("\t UpdateBSDFWeightAfterWeightOperatorVisit(StrataTree, StrataTree.BSDFs[%d], %d /*Op index*/, %d /*PreviousIsInputA*/);\n"), BSDFIndex, CurrentOperator.Index, 1);
+										ResourcesString += FString::Printf(TEXT("\t UpdateAllBSDFWithBottomUpOperatorVisit_Vertical(StrataTree, StrataTree.BSDFs[%d], %d /*Op index*/, %d /*PreviousIsInputA*/);\n"), BSDFIndex, CurrentOperator.Index, CurrentOperator.LeftIndex == PreviousOperatorIndex ? 1 : 0);
 										break;
 									}
 
