Strata - moved statistical lob evaluation to be part of bsdf/operators update functions.

diff --git a/Engine/Shaders/Private/DistortAccumulatePS.usf b/Engine/Shaders/Private/DistortAccumulatePS.usf
--- a/Engine/Shaders/Private/DistortAccumulatePS.usf
+++ b/Engine/Shaders/Private/DistortAccumulatePS.usf
@@ -132,81 +132,20 @@ void Main(
 		LegacyUpdateBSDFsOperators(StrataPixelHeader, StrataTree, bRoughDiffuseEnabled, NullStrataAddressing, V, DummyL);
 	#endif
 
-		
-		STRATA_UNROLL_N(STRATA_CLAMPED_BSDF_COUNT)
-		for (int BSDFIdx = 0; BSDFIdx < StrataTree.BSDFCount; ++BSDFIdx)
-		{
-		#define BSDF StrataTree.BSDFs[BSDFIdx]
-
-			const float BSDFCoverage = BSDF.Coverage;
-
-			// Create the BSDF context
-			FStrataBSDFContext StrataBSDFContext = StrataCreateBSDFContext(StrataPixelHeader, BSDF, NullStrataAddressing, V, DummyL);
-
-			FStrataEvaluateResult StrataEvaluate = StrataEvaluateBSDF(StrataBSDFContext, true); // Works to call because DummyL is nto used when computing StrataEvaluate.Throughput
-
-			// We only support rough refraction from Slab material type for now.
-			if (BSDF_GETTYPE(BSDF) == STRATA_BSDF_TYPE_SLAB)
-			{
-				// Surface directional albedo
-				const FBxDFEnergyTerms EnergyTerms = ComputeGGXSpecEnergyTerms(SLAB_ROUGHNESS(BSDF), StrataBSDFContext.SatNoV, SLAB_F0(BSDF), SLAB_F90(BSDF));
-
-				// The top layer we consider a transition air to water. For layers below, we assume they all have similar index of refraction.
-				const float InterfaceEta12 = BSDF.bIsTop ? 1.0f / 1.33f : 1.0f; // STRATA_TODO Use Refraction.x in IOR mode? Or the material IOR corresponding to F0?
-
-				// Transform input Lobe into local space w.r.t. normal
-				const float3 SavedWorldSpaceWiLobeMu = WiLobe.Mu;
-				WiLobe.Mu = float3(dot(WiLobe.Mu, StrataBSDFContext.X), dot(WiLobe.Mu, StrataBSDFContext.Y), dot(WiLobe.Mu, StrataBSDFContext.N));
-
-				// Compute the refracted Lobe according to roughness . 
-				// STRATA_TODO: take into account scattering and transmittance. Colored transmittance is taken into account already but not scattering influence on variance.
-				StrataLobeStatistic NewLobe = StrataGetRefractedLobe(WiLobe, EnergyTerms.E * BSDF.LuminanceWeight, SLAB_ROUGHNESS(BSDF), InterfaceEta12);
-
-				// Transform the new Lobe direction to world space and weight/accumulate
-				// STRATA_TODO it should be better to weight lobes horizontally weighted and chain lobes that are vertically layered. For that we would need another set of UpdateAllXXX function as seen above.
-				NewLobe.Mu = NewLobe.Mu.x * StrataBSDFContext.X + NewLobe.Mu.y * StrataBSDFContext.Y + NewLobe.Mu.z * StrataBSDFContext.N;
-				StrataAccumulateWeightedLobe(AccumulatedLayerLobes, AccumulatedLayerLobesWeights, NewLobe);
-
-				// Restore current layer input Lobe direction to world space
-				WiLobe.Mu = SavedWorldSpaceWiLobeMu;
-			}
-			else
-			{
-				// Full roughness to make it visible as something to fix in case a type is not handled.
-				WiLobe.Sigma = 1000.0;
-			}
-
-		#undef BSDF
-		}
+		WiLobe = StrataTree.Operators[StrataData.OperatorIndex].LobeStat;
+		//half3 LobeWorldDir = normalize(TransformWorldVectorToView(WiLobe.Mu));
 	}
 
-	if (AccumulatedLayerLobesWeights > 0.0f)
-	{
-		// Now compute the new input Wi Lobe from the accumulated refracted Lobe of the previous layer
-		WiLobe = StrataNormaliseLobe(AccumulatedLayerLobes, AccumulatedLayerLobesWeights);
-		// And negate the direction to make sure the previously refracted direction is the new incoming
-		WiLobe.Mu = -WiLobe.Mu;
-	}
-
-
-
 	// Not write out variance (as this is more linear to accumulate than roughness)
 	OutVariance = float4(WiLobe.Sigma, 0.0f, 0.0f, 0.0f);
 
 	// Only allow clipping if roughness is close to 0.
 	const bool TryToClip = OutVariance.x < 0.01f;
 
-	// Finally, negate the direction to make sure the new incombing Lobe is the refracted Lobe we want to take into account
-	WiLobe.Mu = -WiLobe.Mu;
-
-	// FixUp to reuse the same distortion code we have for now. STRATA_TODO implement a new model?
-	half3 LobeWorldDir = normalize(TransformWorldVectorToView(WiLobe.Mu));
-	Normal = LobeWorldDir;
-
 	// Compute UV distortion
 	float2 BufferUVDistortion = ComputeBufferUVDistortion(
 		MaterialParameters, PixelMaterialInputs, ResolvedView,
-		Normal, DistortionParams, ScreenUV, Refraction, TryToClip, EyeIndex);
+		MaterialParameters.WorldNormal, DistortionParams, ScreenUV, Refraction, TryToClip, EyeIndex);
 
 	// Sample depth at distortion offset
 	float2 DistortBufferUV = ScreenUV + BufferUVDistortion;
diff --git a/Engine/Shaders/Private/Strata/Strata.ush b/Engine/Shaders/Private/Strata/Strata.ush
--- a/Engine/Shaders/Private/Strata/Strata.ush
+++ b/Engine/Shaders/Private/Strata/Strata.ush
@@ -719,11 +719,16 @@ struct FStrataOperator
 	int RightIndex;
 	float Weight;
 
+	// Coverage / Transmittance
+
 	float  Coverage;					// The summarised coverage of all the matter (from the sub tree) represented by that operator
 	float3 Transmittance;				// This is the transmittance of the matter itself. If coverage is 0.25, then 25% of the light will be affected by the transmittance and 75% of light won't be.
 
 	float VerticalTopCoverage;			// Only used by vertical layering operators. It represent the coverage of all the matter of the top layer (from the left sub tree).
 	float3 VerticalTopTransmittance;	// Only used by vertical layering operators. It represent the transmittance of matter of the top layer (from the left sub tree).
+
+	// Rough refraction
+	StrataLobeStatistic LobeStat;
 };
 
 FStrataOperator GetInitialisedStrataOperator()
@@ -1910,6 +1915,17 @@ float3 StrataGetBSDFSpecularColor(in FStrataBSDF BSDF)
 	return 0.f;
 }
 
+float3 StrataGetBSDFSpecularF90(in FStrataBSDF BSDF)
+{
+	switch (BSDF_GETTYPE(BSDF))
+	{
+	case STRATA_BSDF_TYPE_SLAB: 			return SLAB_F90(BSDF);
+	case STRATA_BSDF_TYPE_HAIR: 			return 1.0f;
+	case STRATA_BSDF_TYPE_SINGLELAYERWATER: return 1.0f;
+	}
+	return 0.f;
+}
+
 float3 StrataGetBSDFSubSurfaceColor(in FStrataBSDF BSDF)
 {
 	switch (BSDF_GETTYPE(BSDF))
diff --git a/Engine/Shaders/Private/Strata/StrataStatisticalOperators.ush b/Engine/Shaders/Private/Strata/StrataStatisticalOperators.ush
--- a/Engine/Shaders/Private/Strata/StrataStatisticalOperators.ush
+++ b/Engine/Shaders/Private/Strata/StrataStatisticalOperators.ush
@@ -1,5 +1,7 @@
 // Copyright Epic Games, Inc. All Rights Reserved.
 
+#pragma once
+
 #include "/Engine/Private/Common.ush"
 
 // Statistical operators representing BSDF as a sum of lobes.
@@ -8,9 +10,8 @@
 struct StrataLobeStatistic
 {
 	// Mean
-	// xy is the 2d projection of the main direction vector onto the plane defined by the surface normal. 
-	// z is along the normal.
-	float3 Mu;
+	// xy is the 2d projection of the lobe main direction onto the plane defined by the surface normal. 
+	float2 Mu;
 
 	// Energy
 	float3 E;
@@ -52,7 +53,7 @@ StrataLobeStatistic StrataGetViewDiracLobe(float3 Wi)
 {
 	StrataLobeStatistic WiLobe;
 	WiLobe.E = 1.0f;
-	WiLobe.Mu = Wi;
+	WiLobe.Mu = Wi.xy;
 	WiLobe.Sigma = 0.0f;
 	return WiLobe;
 }
@@ -63,7 +64,7 @@ StrataLobeStatistic StrataGetReflectedLobe(StrataLobeStatistic WiLobe, float3 In
 
 	WoLobe.E = WiLobe.E * InterfaceFDG;
 
-	WoLobe.Mu = float3(-WiLobe.Mu.xy, WiLobe.Mu.z);
+	WoLobe.Mu = WiLobe.Mu;
 
 	WoLobe.Sigma = WiLobe.Sigma + StrataLobeRoughnessToVariance(InterfaceRoughness);
 
@@ -76,9 +77,7 @@ StrataLobeStatistic StrataGetRefractedLobe(StrataLobeStatistic WiLobe, float3 In
 
 	WoLobe.E = WiLobe.E * (1.0f - InterfaceFDG);
 
-	WoLobe.Mu.xy = -WiLobe.Mu.xy * InterfaceEta12;
-	const float SquaredProjectedLength = WiLobe.Mu.x * WiLobe.Mu.x + WiLobe.Mu.y * WiLobe.Mu.y;
-	WoLobe.Mu.z = -sqrt(1.0f - SquaredProjectedLength);
+	WoLobe.Mu.xy = WiLobe.Mu.xy * InterfaceEta12;
 
 	//	const float S = 0.5f * (1.0f + InterfaceEta12 * (WiLobe.Mu.z / WoLobe.Mu.z));				// This respect eq.10 but it goes crazy and does not respect the roughness of a single front layer
 	//	const float S = 0.5f * (1.0f + InterfaceEta12 * max(0.0, WiLobe.Mu.z / WoLobe.Mu.z));		// This respect the roughness range better but looks incorrect
@@ -94,48 +93,20 @@ StrataLobeStatistic StrataGetTransmittedLobe(StrataLobeStatistic WiLobe, float3
 
 	WoLobe.E = WiLobe.E * exp(-OpticalDepth);
 
-	// No refraction in the paper here
-	WoLobe.Mu = -WiLobe.Mu;
+	// Continue straight forward
+	WoLobe.Mu = WiLobe.Mu;
 
 	WoLobe.Sigma = WiLobe.Sigma;
 
 	return WoLobe;
 }
 
-StrataLobeStatistic StrataGetScatteredLobe(StrataLobeStatistic WiLobe, float3 InScattering)
+StrataLobeStatistic StrataHorizontalMixLobe(StrataLobeStatistic A, StrataLobeStatistic B, float Mix)
 {
 	StrataLobeStatistic WoLobe;
-
-	WoLobe.E = WiLobe.E * InScattering;
-
-	// No refraction in the paper here
-	WoLobe.Mu = -WiLobe.Mu;
-
-	// In our case, we use an isotropic phase function and the paper do not support that (only HG phase function with g > 0.7).
-	// So we simply assume a roughness of 1 for the Lobe over which incomling light to scatter is integrated.
-	const float IsotropicPhaseLobeVariance = StrataLobeRoughnessToVariance(1.0f);
-	WoLobe.Sigma = WiLobe.Sigma + IsotropicPhaseLobeVariance;
-
+	WoLobe.E = lerp(A.E, B.E, Mix);
+	WoLobe.Mu = lerp(A.Mu, B.Mu, Mix);
+	WoLobe.Sigma = lerp(A.Sigma, B.Sigma, Mix);
 	return WoLobe;
 }
 
-
-
-// This accumulation function can be used to merge Lobes while taking into account their relative importance due to E and a coverage
-void StrataAccumulateWeightedLobe(inout StrataLobeStatistic LobeAccum, inout float WeightAccum, in StrataLobeStatistic NewLobe)
-{
-	const float NewLobeWeight = dot(NewLobe.E, (1.0f / 3.0f).xxx);
-	LobeAccum.E		+= NewLobe.E * NewLobeWeight;
-	LobeAccum.Mu	+= NewLobe.Mu * NewLobeWeight;
-	LobeAccum.Sigma	+= NewLobe.Sigma * NewLobeWeight;
-	WeightAccum		+= NewLobeWeight;
-}
-
-StrataLobeStatistic StrataNormaliseLobe(inout StrataLobeStatistic LobeAccum, inout float WeightAccum)
-{
-	StrataLobeStatistic OutputLobe;
-	OutputLobe.E	= LobeAccum.E / WeightAccum;
-	OutputLobe.Mu	= normalize(LobeAccum.Mu);
-	OutputLobe.Sigma= LobeAccum.Sigma / WeightAccum;
-	return OutputLobe;
-}
\ No newline at end of file
diff --git a/Engine/Shaders/Private/Strata/StrataTree.ush b/Engine/Shaders/Private/Strata/StrataTree.ush
--- a/Engine/Shaders/Private/Strata/StrataTree.ush
+++ b/Engine/Shaders/Private/Strata/StrataTree.ush
@@ -90,6 +90,16 @@ void UpdateSingleBSDFOperatorCoverageTransmittance(
 	StrataTree.Operators[CurrentBSDF.OperatorIndex].Coverage = CurrentBSDFCoverage;
 	StrataTree.Operators[CurrentBSDF.OperatorIndex].Transmittance = CurrentBSDFTransmittance;
 
+	// The top layer we consider a transition air to water. For layers below, we assume they all have similar index of refraction.
+	const float Roughness = StrataGetBSDFRoughness(CurrentBSDF);
+	const FBxDFEnergyTerms EnergyTerms = ComputeGGXSpecEnergyTerms(Roughness, StrataBSDFContext.SatNoV, StrataGetBSDFSpecularColor(CurrentBSDF), StrataGetBSDFSpecularF90(CurrentBSDF));
+	const float InterfaceEta12 = CurrentBSDF.bIsTop ? 1.0f / 1.33f : 1.0f; // STRATA_TODO Use Refraction.x in IOR mode? Or even the material IOR corresponding to F0?
+	StrataLobeStatistic RefractedLobe = StrataGetRefractedLobe(StrataGetViewDiracLobe(-V), EnergyTerms.E, Roughness, InterfaceEta12);
+	// Account for transmittance of the medium (but ignoring scattering for now)
+	RefractedLobe.E *= CurrentBSDFTransmittance;
+
+	StrataTree.Operators[CurrentBSDF.OperatorIndex].LobeStat = RefractedLobe;
+
 #endif // STRATA_INLINE_SHADING
 }
 
@@ -110,6 +120,8 @@ void UpdateSingleOperatorCoverageTransmittance(inout FStrataTree StrataTree, int
 		const float Weight = saturate(Op.Weight);
 		Op.Coverage = Weight * OpA.Coverage;
 		Op.Transmittance = OpA.Transmittance;
+
+		Op.LobeStat.E *= Weight;
 		break;
 	}
 
@@ -142,6 +154,10 @@ void UpdateSingleOperatorCoverageTransmittance(inout FStrataTree StrataTree, int
 
 		Op.VerticalTopCoverage		= OpA.Coverage;
 		Op.VerticalTopTransmittance = OpA.Transmittance;
+
+		// Project top lob into bottom layer
+		Op.LobeStat = StrataGetRefractedLobe(OpA.LobeStat, /*InterfaceFDG*/1.0f, StrataLobeVarianceToRoughness(OpA.LobeStat.Sigma), /*InterfaceEta12*/1.0f);
+		Op.LobeStat.E *= OpB.LobeStat.E; // Combine with botton layer medium and specular transmittance
 		break;
 	}
 
@@ -153,6 +169,8 @@ void UpdateSingleOperatorCoverageTransmittance(inout FStrataTree StrataTree, int
 
 		Op.Coverage = AMix * OpA.Coverage + BMix * OpB.Coverage;
 		Op.Transmittance = (AMix * OpA.Coverage * OpA.Transmittance + BMix * OpB.Coverage * OpB.Transmittance) / max(STRATA_EPSILON, AMix * OpA.Coverage + BMix * OpB.Coverage);
+
+		Op.LobeStat = StrataHorizontalMixLobe(OpA.LobeStat, OpB.LobeStat, BMix);
 		break;
 	}
 
@@ -164,6 +182,8 @@ void UpdateSingleOperatorCoverageTransmittance(inout FStrataTree StrataTree, int
 
 		Op.Coverage = saturate(OpA.Coverage + OpB.Coverage);
 		Op.Transmittance = lerp(OpB.Transmittance, OpA.Transmittance, AMixFactor);
+
+		Op.LobeStat = StrataHorizontalMixLobe(OpA.LobeStat, OpB.LobeStat, 0.5f);
 		break;
 	}
 
