Fix shader compilation with default font color on HLSCC platform.

diff --git a/Engine/Shaders/Private/ShaderPrint.ush b/Engine/Shaders/Private/ShaderPrint.ush
--- a/Engine/Shaders/Private/ShaderPrint.ush
+++ b/Engine/Shaders/Private/ShaderPrint.ush
@@ -160,7 +160,6 @@ struct FFontColor
 FFontColor InitFontColor(float InX, float InY, float InZ)	{ FFontColor Out; Out.Color = float3(InX, InY, InZ); return Out; }
 FFontColor InitFontColor(float3 In)							{ FFontColor Out; Out.Color = In; return Out; }
 FFontColor GetDefaultFontColor()							{ FFontColor Out; Out.Color = float3(1,1,1); return Out; }
-#define DEFAULTFONTCOLOR ((FFontColor)1)
 
 // Predefined colors
 #define FontWhite		InitFontColor(1, 1, 1)
@@ -281,7 +280,7 @@ void ShaderPrint_Internal(in FShaderPrintItem Item)
 
 	// Prevent writing off the buffer
 	// Note that counter still increases so need clamp when reading it in later passes
-	if (IndexToStore >= ShaderPrint.MaxValueCount)
+	if (uint(IndexToStore) >= ShaderPrint.MaxValueCount)
 	{
 		return;
 	}
@@ -304,7 +303,7 @@ void ShaderPrint_Internal(in float2 ScreenPos, in int Value, in FFontColor FontC
 ///////////////////////////////////////////////////////////////////////////////////////////////////
 // Symbol printing
 
-float2 ShaderPrintSymbol(in float2 ScreenPos, in int Symbol, in FFontColor Color = DEFAULTFONTCOLOR)
+float2 ShaderPrintSymbol(in float2 ScreenPos, in int Symbol, in FFontColor Color)
 {
 	if (ShaderPrintFilterEnable)
 	{
@@ -313,11 +312,13 @@ float2 ShaderPrintSymbol(in float2 ScreenPos, in int Symbol, in FFontColor Color
 	}
 	return ScreenPos;
 }
+float2 ShaderPrintSymbol(in float2 ScreenPos, in int Symbol) { return ShaderPrintSymbol(ScreenPos, Symbol, GetDefaultFontColor()); }
 
-void ShaderPrintSymbol(in int symbol, in FFontColor Color = DEFAULTFONTCOLOR)
+void ShaderPrintSymbol(in int symbol, in FFontColor Color)
 {
 	ShaderPrintCursorPos = ShaderPrintSymbol(ShaderPrintCursorPos, symbol, Color);
 }
+void ShaderPrintSymbol(in int symbol) { ShaderPrintSymbol(symbol, GetDefaultFontColor()); }
 
 ///////////////////////////////////////////////////////////////////////////////////////////////////
 // Text & hash
@@ -340,7 +341,7 @@ uint ShaderPrintGetOffset(FShaderPrintText InTextEntry);
 uint ShaderPrintGetHash(FShaderPrintText InTextEntry);
  
 // Function for reading global TEXT string
-float2 ShaderPrintText(float2 Pos, FShaderPrintText InTextEntry, FFontColor InColor = DEFAULTFONTCOLOR)
+float2 ShaderPrintText(float2 Pos, FShaderPrintText InTextEntry, FFontColor InColor)
 {
 	uint Begin = ShaderPrintGetOffset(InTextEntry);
 	uint End = ShaderPrintGetOffset(InitShaderPrintText(InTextEntry.Index + 1));
@@ -350,8 +351,9 @@ float2 ShaderPrintText(float2 Pos, FShaderPrintText InTextEntry, FFontColor InCo
 	}
 	return Pos;
 }
+float2 ShaderPrintText(float2 Pos, FShaderPrintText InTextEntry) { return ShaderPrintText(Pos, InTextEntry, GetDefaultFontColor()); }
 
-void ShaderPrintText(inout FShaderPrintContext InContext, FShaderPrintText InTextEntry, FFontColor InColor = DEFAULTFONTCOLOR)
+void ShaderPrintText(inout FShaderPrintContext InContext, FShaderPrintText InTextEntry, FFontColor InColor)
 {
 	if (InContext.bIsActive)
 	{
@@ -363,9 +365,12 @@ void ShaderPrintText(inout FShaderPrintContext InContext, FShaderPrintText InTex
 		}	
 	}
 }
+void ShaderPrintText(inout FShaderPrintContext InContext, FShaderPrintText InTextEntry) { ShaderPrintText(InContext, InTextEntry, GetDefaultFontColor()); }
 
-float2 Print(float2 Pos, FShaderPrintText InTextEntry, FFontColor InColor=DEFAULTFONTCOLOR)                         { return ShaderPrintText(Pos, InTextEntry, InColor); }
-void   Print(inout FShaderPrintContext InCtx, FShaderPrintText InTextEntry, FFontColor InColor = DEFAULTFONTCOLOR)	{ ShaderPrintText(InCtx, InTextEntry, InColor); }
+float2 Print(float2 Pos, FShaderPrintText InTextEntry, FFontColor InColor)                      { return ShaderPrintText(Pos, InTextEntry, InColor); }
+float2 Print(float2 Pos, FShaderPrintText InTextEntry)											{ return ShaderPrintText(Pos, InTextEntry, GetDefaultFontColor()); }
+void   Print(inout FShaderPrintContext InCtx, FShaderPrintText InTextEntry, FFontColor InColor)	{ ShaderPrintText(InCtx, InTextEntry, InColor); }
+void   Print(inout FShaderPrintContext InCtx, FShaderPrintText InTextEntry)						{ ShaderPrintText(InCtx, InTextEntry, GetDefaultFontColor()); }
 
 // Function for reading global TEXT string. 
 // The data and function are generated by the shader compiler.
@@ -385,17 +390,23 @@ float2 ShaderPrintValue(float2 ScreenPos, bool  Value, FFontColor Color, uint Ma
 // Scalar type printing
 
 #define SHADER_PRINT_OVERLOAD_1(InPrefix, InNumComponent, InType) \
-void   InPrefix##Print(InType Value, FFontColor Color=DEFAULTFONTCOLOR, uint MaxDigit=MAX_DIGIT_COUNT, uint MaxDecimal=MAX_DECIMAL_COUNT)                                  { if (ShaderPrintFilterEnable) { ShaderPrintCursorPos = ShaderPrintValue(ShaderPrintCursorPos, Value, Color, MaxDigit, MaxDecimal); } } \
-float2 InPrefix##Print(float2 ScreenPos, InType Value, FFontColor Color=DEFAULTFONTCOLOR, uint MaxDigit=MAX_DIGIT_COUNT, uint MaxDecimal=MAX_DECIMAL_COUNT)                { if (ShaderPrintFilterEnable) { ScreenPos            = ShaderPrintValue(ScreenPos,            Value, Color, MaxDigit, MaxDecimal); } return ScreenPos; } \
-void   InPrefix##Print(inout FShaderPrintContext InCtx, InType Value, FFontColor Color=DEFAULTFONTCOLOR, uint MaxDigit=MAX_DIGIT_COUNT, uint MaxDecimal=MAX_DECIMAL_COUNT) { if (InCtx.bIsActive)         { InCtx.Pos            = ShaderPrintValue(InCtx.Pos,            Value, Color, MaxDigit, MaxDecimal); } }
+void   InPrefix##Print(InType Value, FFontColor Color, uint MaxDigit=MAX_DIGIT_COUNT, uint MaxDecimal=MAX_DECIMAL_COUNT)									{ if (ShaderPrintFilterEnable) { ShaderPrintCursorPos = ShaderPrintValue(ShaderPrintCursorPos, Value, Color, MaxDigit, MaxDecimal); } } \
+float2 InPrefix##Print(float2 ScreenPos, InType Value, FFontColor Color, uint MaxDigit=MAX_DIGIT_COUNT, uint MaxDecimal=MAX_DECIMAL_COUNT)					{ if (ShaderPrintFilterEnable) { ScreenPos            = ShaderPrintValue(ScreenPos,            Value, Color, MaxDigit, MaxDecimal); } return ScreenPos; } \
+void   InPrefix##Print(inout FShaderPrintContext InCtx, InType Value, FFontColor Color, uint MaxDigit=MAX_DIGIT_COUNT, uint MaxDecimal=MAX_DECIMAL_COUNT)	{ if (InCtx.bIsActive)         { InCtx.Pos            = ShaderPrintValue(InCtx.Pos,            Value, Color, MaxDigit, MaxDecimal); } } \
+void   InPrefix##Print(InType Value)																														{ if (ShaderPrintFilterEnable) { ShaderPrintCursorPos = ShaderPrintValue(ShaderPrintCursorPos, Value, GetDefaultFontColor(), MAX_DIGIT_COUNT, MAX_DECIMAL_COUNT); } } \
+float2 InPrefix##Print(float2 ScreenPos, InType Value)																										{ if (ShaderPrintFilterEnable) { ScreenPos            = ShaderPrintValue(ScreenPos,            Value, GetDefaultFontColor(), MAX_DIGIT_COUNT, MAX_DECIMAL_COUNT); } return ScreenPos; } \
+void   InPrefix##Print(inout FShaderPrintContext InCtx, InType Value)																						{ if (InCtx.bIsActive)         { InCtx.Pos            = ShaderPrintValue(InCtx.Pos,            Value, GetDefaultFontColor(), MAX_DIGIT_COUNT, MAX_DECIMAL_COUNT); } }
 
 ///////////////////////////////////////////////////////////////////////////////////////////////////
 // Vector type printing
 
 #define SHADER_PRINT_OVERLOAD_N(InPrefix, InNumComponent, InType) \
-void   InPrefix##Print(InType Value, FFontColor Color=DEFAULTFONTCOLOR, uint MaxDigit=MAX_DIGIT_COUNT, uint MaxDecimal=MAX_DECIMAL_COUNT)                                  { if (ShaderPrintFilterEnable) { UNROLL for (uint CompIt=0;CompIt<InNumComponent;++CompIt) { ShaderPrintCursorPos = ShaderPrintValue(ShaderPrintCursorPos, Value[CompIt], Color, MaxDigit, MaxDecimal); } } } \
-float2 InPrefix##Print(float2 ScreenPos, InType Value, FFontColor Color=DEFAULTFONTCOLOR, uint MaxDigit=MAX_DIGIT_COUNT, uint MaxDecimal=MAX_DECIMAL_COUNT)                { if (ShaderPrintFilterEnable) { UNROLL for (uint CompIt=0;CompIt<InNumComponent;++CompIt) { ScreenPos            = ShaderPrintValue(ScreenPos,            Value[CompIt], Color, MaxDigit, MaxDecimal); } } return ScreenPos; } \
-void   InPrefix##Print(inout FShaderPrintContext InCtx, InType Value, FFontColor Color=DEFAULTFONTCOLOR, uint MaxDigit=MAX_DIGIT_COUNT, uint MaxDecimal=MAX_DECIMAL_COUNT) { if (InCtx.bIsActive)         { UNROLL for (uint CompIt=0;CompIt<InNumComponent;++CompIt) { InCtx.Pos            = ShaderPrintValue(InCtx.Pos,            Value[CompIt], Color, MaxDigit, MaxDecimal); } } }
+void   InPrefix##Print(InType Value, FFontColor Color, uint MaxDigit=MAX_DIGIT_COUNT, uint MaxDecimal=MAX_DECIMAL_COUNT)									{ if (ShaderPrintFilterEnable) { UNROLL for (uint CompIt=0;CompIt<InNumComponent;++CompIt) { ShaderPrintCursorPos = ShaderPrintValue(ShaderPrintCursorPos, Value[CompIt], Color, MaxDigit, MaxDecimal); } } } \
+float2 InPrefix##Print(float2 ScreenPos, InType Value, FFontColor Color, uint MaxDigit=MAX_DIGIT_COUNT, uint MaxDecimal=MAX_DECIMAL_COUNT)					{ if (ShaderPrintFilterEnable) { UNROLL for (uint CompIt=0;CompIt<InNumComponent;++CompIt) { ScreenPos            = ShaderPrintValue(ScreenPos,            Value[CompIt], Color, MaxDigit, MaxDecimal); } } return ScreenPos; } \
+void   InPrefix##Print(inout FShaderPrintContext InCtx, InType Value, FFontColor Color, uint MaxDigit=MAX_DIGIT_COUNT, uint MaxDecimal=MAX_DECIMAL_COUNT)	{ if (InCtx.bIsActive)         { UNROLL for (uint CompIt=0;CompIt<InNumComponent;++CompIt) { InCtx.Pos            = ShaderPrintValue(InCtx.Pos,            Value[CompIt], Color, MaxDigit, MaxDecimal); } } } \
+void   InPrefix##Print(InType Value)																														{ if (ShaderPrintFilterEnable) { UNROLL for (uint CompIt=0;CompIt<InNumComponent;++CompIt) { ShaderPrintCursorPos = ShaderPrintValue(ShaderPrintCursorPos, Value[CompIt], GetDefaultFontColor(), MAX_DIGIT_COUNT, MAX_DECIMAL_COUNT); } } } \
+float2 InPrefix##Print(float2 ScreenPos, InType Value)																										{ if (ShaderPrintFilterEnable) { UNROLL for (uint CompIt=0;CompIt<InNumComponent;++CompIt) { ScreenPos            = ShaderPrintValue(ScreenPos,            Value[CompIt], GetDefaultFontColor(), MAX_DIGIT_COUNT, MAX_DECIMAL_COUNT); } } return ScreenPos; } \
+void   InPrefix##Print(inout FShaderPrintContext InCtx, InType Value)																						{ if (InCtx.bIsActive)         { UNROLL for (uint CompIt=0;CompIt<InNumComponent;++CompIt) { InCtx.Pos            = ShaderPrintValue(InCtx.Pos,            Value[CompIt], GetDefaultFontColor(), MAX_DIGIT_COUNT, MAX_DECIMAL_COUNT); } } }
 
 // ShaderPrint(X) version (legacy)
 SHADER_PRINT_OVERLOAD_1(Shader, 1, float)
@@ -437,9 +448,12 @@ SHADER_PRINT_OVERLOAD_1(, 1, bool)
 // Matrix type printing
 
 #define SHADER_PRINT_OVERLOAD_NN(InPrefix, InNumComponentX, InNumComponentY, InType) \
-void   InPrefix##Print(InType Value, FFontColor Color=DEFAULTFONTCOLOR, uint MaxDigit=MAX_DIGIT_COUNT, uint MaxDecimal=MAX_DECIMAL_COUNT)                                   { if (ShaderPrintFilterEnable) { UNROLL for (uint CompY=0;CompY<InNumComponentY;++CompY) { ShaderPrintCursorPos = InPrefix##Print(ShaderPrintCursorPos, Value[CompY], Color, MaxDigit, MaxDecimal); ShaderPrintCursorPos += float2(0.f, ShaderPrint.FontSpacing.y); } } } \
-float2 InPrefix##Print(float2 ScreenPos, InType Value, FFontColor Color=DEFAULTFONTCOLOR, uint MaxDigit=MAX_DIGIT_COUNT, uint MaxDecimal=MAX_DECIMAL_COUNT)                 { if (ShaderPrintFilterEnable) { UNROLL for (uint CompY=0;CompY<InNumComponentY;++CompY) { ScreenPos            = InPrefix##Print(ScreenPos,            Value[CompY], Color, MaxDigit, MaxDecimal); ScreenPos += float2(0.f, ShaderPrint.FontSpacing.y); } } return ScreenPos; } \
-void   InPrefix##Print(inout FShaderPrintContext InCtx, InType Value, FFontColor Color=DEFAULTFONTCOLOR, uint MaxDigit=MAX_DIGIT_COUNT, uint MaxDecimal=MAX_DECIMAL_COUNT)  { if (InCtx.bIsActive)         { UNROLL for (uint CompY=0;CompY<InNumComponentY;++CompY) { InCtx.Pos            = InPrefix##Print(InCtx.Pos,            Value[CompY], Color, MaxDigit, MaxDecimal); InCtx.Pos += float2(0.f, ShaderPrint.FontSpacing.y); } } }
+void   InPrefix##Print(InType Value, FFontColor Color, uint MaxDigit=MAX_DIGIT_COUNT, uint MaxDecimal=MAX_DECIMAL_COUNT)									{ if (ShaderPrintFilterEnable) { UNROLL for (uint CompY=0;CompY<InNumComponentY;++CompY) { ShaderPrintCursorPos = InPrefix##Print(ShaderPrintCursorPos, Value[CompY], Color, MaxDigit, MaxDecimal); ShaderPrintCursorPos += float2(0.f, ShaderPrint.FontSpacing.y); } } } \
+float2 InPrefix##Print(float2 ScreenPos, InType Value, FFontColor Color, uint MaxDigit=MAX_DIGIT_COUNT, uint MaxDecimal=MAX_DECIMAL_COUNT)					{ if (ShaderPrintFilterEnable) { UNROLL for (uint CompY=0;CompY<InNumComponentY;++CompY) { ScreenPos            = InPrefix##Print(ScreenPos,            Value[CompY], Color, MaxDigit, MaxDecimal); ScreenPos += float2(0.f, ShaderPrint.FontSpacing.y); } } return ScreenPos; } \
+void   InPrefix##Print(inout FShaderPrintContext InCtx, InType Value, FFontColor Color, uint MaxDigit=MAX_DIGIT_COUNT, uint MaxDecimal=MAX_DECIMAL_COUNT)	{ if (InCtx.bIsActive)         { UNROLL for (uint CompY=0;CompY<InNumComponentY;++CompY) { InCtx.Pos            = InPrefix##Print(InCtx.Pos,            Value[CompY], Color, MaxDigit, MaxDecimal); InCtx.Pos += float2(0.f, ShaderPrint.FontSpacing.y); } } } \
+void   InPrefix##Print(InType Value)																														{ if (ShaderPrintFilterEnable) { UNROLL for (uint CompY=0;CompY<InNumComponentY;++CompY) { ShaderPrintCursorPos = InPrefix##Print(ShaderPrintCursorPos, Value[CompY], GetDefaultFontColor(), MAX_DIGIT_COUNT, MAX_DECIMAL_COUNT); ShaderPrintCursorPos += float2(0.f, ShaderPrint.FontSpacing.y); } } } \
+float2 InPrefix##Print(float2 ScreenPos, InType Value)																										{ if (ShaderPrintFilterEnable) { UNROLL for (uint CompY=0;CompY<InNumComponentY;++CompY) { ScreenPos            = InPrefix##Print(ScreenPos,            Value[CompY], GetDefaultFontColor(), MAX_DIGIT_COUNT, MAX_DECIMAL_COUNT); ScreenPos += float2(0.f, ShaderPrint.FontSpacing.y); } } return ScreenPos; } \
+void   InPrefix##Print(inout FShaderPrintContext InCtx, InType Value)																						{ if (InCtx.bIsActive)         { UNROLL for (uint CompY=0;CompY<InNumComponentY;++CompY) { InCtx.Pos            = InPrefix##Print(InCtx.Pos,            Value[CompY], GetDefaultFontColor(), MAX_DIGIT_COUNT, MAX_DECIMAL_COUNT); InCtx.Pos += float2(0.f, ShaderPrint.FontSpacing.y); } } }
 
 // ShaderPrint(X) version
 SHADER_PRINT_OVERLOAD_NN(Shader, 3, 3, float3x3)
