Make sure Bloom's FindKernelCenter have all thread are active on AllMemoryBarrier()

diff --git a/Engine/Shaders/Private/Bloom/BloomFindKernelCenter.usf b/Engine/Shaders/Private/Bloom/BloomFindKernelCenter.usf
--- a/Engine/Shaders/Private/Bloom/BloomFindKernelCenter.usf
+++ b/Engine/Shaders/Private/Bloom/BloomFindKernelCenter.usf
@@ -26,6 +26,9 @@ RWCoherentStructuredBuffer(uint) KernelCenterCoordOutput;
 
 groupshared uint MaxLuminance;
 groupshared uint2 MaxLuminancePixelPos;
+groupshared uint CurrentMaxLuminance;
+groupshared uint bDone;
+groupshared uint PrevLock;
 
 
 //------------------------------------------------------- ENTRY POINT
@@ -39,6 +42,7 @@ void MainCS(
 	if (GroupThreadIndex == 0)
 	{
 		MaxLuminance = 0;
+		bDone = 0;
 	}
 
 	GroupMemoryBarrierWithGroupSync();
@@ -70,42 +74,62 @@ void MainCS(
 	
 	GroupMemoryBarrierWithGroupSync();
 
-	BRANCH
-	if (GroupThreadIndex == 0 && MaxLuminance > 0)
+	if (MaxLuminance == 0)
 	{
-		bool bDone = false;
-		while (!bDone)
+		return;
+	}
+
+	while (bDone == 0)
+	{
+		if (GroupThreadIndex == 0)
 		{
-			if (MaxLuminance > KernelCenterCoordOutput[2])
+			CurrentMaxLuminance = KernelCenterCoordOutput[2];
+		}
+		
+		GroupMemoryBarrierWithGroupSync();
+
+		if (MaxLuminance > CurrentMaxLuminance)
+		{
+			if (GroupThreadIndex == 0)
 			{
-				uint PrevLock;
+				uint LocalPrevLock;
 				InterlockedCompareExchange(
 					KernelCenterCoordOutput[3],
 					uint(0),
 					uint(1),
-					/* out */ PrevLock);
+					/* out */ LocalPrevLock);
 
-				if (PrevLock == 0)
-				{
-					AllMemoryBarrier();
+				PrevLock = LocalPrevLock;
+			}
+			
+			GroupMemoryBarrierWithGroupSync();
 
-					if (MaxLuminance > KernelCenterCoordOutput[2])
-					{
-						KernelCenterCoordOutput[0] = MaxLuminancePixelPos.x;
-						KernelCenterCoordOutput[1] = MaxLuminancePixelPos.y;
-						KernelCenterCoordOutput[2] = MaxLuminance;
-					}
+			if (PrevLock == 0)
+			{
+				AllMemoryBarrier();
 
-					AllMemoryBarrier();
+				if (GroupThreadIndex == 0 && MaxLuminance > KernelCenterCoordOutput[2])
+				{
+					KernelCenterCoordOutput[0] = MaxLuminancePixelPos.x;
+					KernelCenterCoordOutput[1] = MaxLuminancePixelPos.y;
+					KernelCenterCoordOutput[2] = MaxLuminance;
+				}
 
-					InterlockedExchange(KernelCenterCoordOutput[3], 0, PrevLock);
-					bDone = true;
+				AllMemoryBarrier();
+				
+				if (GroupThreadIndex == 0)
+				{
+					uint Unused;
+					InterlockedExchange(KernelCenterCoordOutput[3], 0, Unused);
 				}
+				bDone += 1;
 			}
-			else
-			{
-				bDone = true;
-			}
 		}
+		else
+		{
+			bDone += 1;
+		}
+		
+		GroupMemoryBarrierWithGroupSync();
 	}
 }
