Merge all Strata tree update code.

diff --git a/Engine/Shaders/Private/DistortAccumulatePS.usf b/Engine/Shaders/Private/DistortAccumulatePS.usf
--- a/Engine/Shaders/Private/DistortAccumulatePS.usf
+++ b/Engine/Shaders/Private/DistortAccumulatePS.usf
@@ -112,25 +112,8 @@ void Main(
 	FStrataTree StrataTree = StrataPixelHeader.StrataTree;
 	if (StrataTree.BSDFCount > 0)
 	{
-		const bool bRoughDiffuseEnabled = false;
-	#if STRATA_USES_CONVERSION_FROM_LEGACY==0
-		//
-		// Update the coverage/transmittance of each leaves (==BSDFs) of the strata tree.
-		//
-		UpdateAllBSDFsOperatorCoverageTransmittance(StrataPixelHeader, StrataTree, bRoughDiffuseEnabled, NullStrataAddressing, V, DummyL);
-
-		//
-		// Propagate up the coverage/transmittance of each node in the strata tree.
-		//
-		UpdateAllOperatorsCoverageTransmittance(StrataTree);
-
-		//
-		// Update the luminance weight of each BSDF according to the operators it has to traverse up to the strata tree root node.
-		//
-		UpdateAllBSDFWithBottomUpOperatorVisit(StrataTree);
-	#else
-		LegacyUpdateBSDFsOperators(StrataPixelHeader, StrataTree, bRoughDiffuseEnabled, NullStrataAddressing, V, DummyL);
-	#endif
+		// Update tree (coverage/transmittance/luminace weights)
+		StrataUpdateTree(NullStrataAddressing, StrataPixelHeader, StrataTree, V, DummyL, false /*bRoughDiffuseEnabled*/);
 
 		WiLobe = StrataTree.Operators[StrataData.OperatorIndex].LobeStat;
 		//half3 LobeWorldDir = normalize(TransformWorldVectorToView(WiLobe.Mu));
diff --git a/Engine/Shaders/Private/Strata/Strata.ush b/Engine/Shaders/Private/Strata/Strata.ush
--- a/Engine/Shaders/Private/Strata/Strata.ush
+++ b/Engine/Shaders/Private/Strata/Strata.ush
@@ -2472,7 +2472,59 @@ void UpdateAllBSDFWithBottomUpOperatorVisit(inout FStrataTree StrataTree);
 void LegacyUpdateBSDFsOperators(inout FStrataPixelHeader StrataPixelHeader, inout FStrataTree StrataTree, bool bRoughDiffuseEnabled, FStrataAddressing NullStrataAddressing, float3 V, float3 L);
 #endif
 
+// Update all Coverage/Transmittance/Weights for all BSDFs based on material topology
+void StrataUpdateTree(
+	inout FStrataAddressing StrataAddressing,
+	inout FStrataPixelHeader StrataPixelHeader,
+	inout FStrataTree StrataTree,
+	float3 V,
+	float3 L,
+	bool bRoughDiffuseEnabled)
+{
+#if STRATA_USES_CONVERSION_FROM_LEGACY==0
+	//
+	// Update the coverage/transmittance of each leaves (==BSDFs) of the strata tree.
+	//
+	UpdateAllBSDFsOperatorCoverageTransmittance(StrataPixelHeader, StrataTree, bRoughDiffuseEnabled, StrataAddressing, V, L);
+
+	//
+	// Propagate up the coverage/transmittance of each node in the strata tree.
+	//
+	UpdateAllOperatorsCoverageTransmittance(StrataTree);
+
+	//
+	// Update the luminance weight of each BSDF according to the operators it has to traverse up to the strata tree root node.
+	//
+	UpdateAllBSDFWithBottomUpOperatorVisit(StrataTree);
+#else
+	LegacyUpdateBSDFsOperators(StrataPixelHeader, StrataTree, bRoughDiffuseEnabled, StrataAddressing, V, L);
 #endif
+}
+
+void StrataUpdateTree(
+	inout FStrataAddressing StrataAddressing,
+	inout FStrataPixelHeader StrataPixelHeader,
+	inout FStrataTree StrataTree,
+	inout FStrataData StrataData,
+	float3 V,
+	float3 L,
+	bool bRoughDiffuseEnabled,
+	inout float OutCoverage,
+	inout float3 OutTransmittancePreCoverage)
+{
+	StrataUpdateTree(
+		StrataAddressing,
+		StrataPixelHeader,
+		StrataTree,
+		V,
+		L,
+		bRoughDiffuseEnabled);
+
+	OutCoverage = saturate(StrataTree.Operators[StrataData.OperatorIndex].Coverage);
+	OutTransmittancePreCoverage = saturate(StrataTree.Operators[StrataData.OperatorIndex].Transmittance);
+}
+
+#endif // STRATA_INLINE_SHADING && STRATA_CLAMPED_BSDF_COUNT > 0
 
 uint PackStrataHeader(uint BSDFCount, FStrataPixelHeader In);
 
@@ -2513,27 +2565,10 @@ void PackStrataOut(
 	FStrataTree StrataTree = StrataPixelHeader.StrataTree;
 	if (StrataTree.BSDFCount > 0)
 	{
+		// Update tree (coverage/transmittance/luminace weights)
 		FStrataAddressing NullStrataAddressing = (FStrataAddressing)0;	// Fake unused in StrataCreateBSDFContext when using Forward inline shading
 		const float3 NullLightDir = float3(0.0, 0.0, 1.0);
-#if STRATA_USES_CONVERSION_FROM_LEGACY==0
-		//
-		// Update the coverage/transmittance of each leaves (==BSDFs) of the strata tree.
-		//
-		UpdateAllBSDFsOperatorCoverageTransmittance(StrataPixelHeader, StrataTree, bRoughDiffuseEnabled, NullStrataAddressing, V, NullLightDir);
-		
-		//
-		// Propagate up the coverage/transmittance of each node in the strata tree.
-		//
-		UpdateAllOperatorsCoverageTransmittance(StrataTree);
-		
-		//
-		// Update the luminance weight of each BSDF according to the operators it has to traverse up to the strata tree root node.
-		//
-		UpdateAllBSDFWithBottomUpOperatorVisit(StrataTree);
-
-#else
-		LegacyUpdateBSDFsOperators(StrataPixelHeader, StrataTree, bRoughDiffuseEnabled, NullStrataAddressing, V, NullLightDir);
-#endif
+		StrataUpdateTree(NullStrataAddressing, StrataPixelHeader, StrataTree, V, NullLightDir, false /*bRoughDiffuseEnabled*/);
 		
 		int SimpleMaterial_BSDFIndex = 0;
 		float TopLayerTotalWeight = 0.0f;
diff --git a/Engine/Shaders/Private/Strata/StrataForwardLighting.ush b/Engine/Shaders/Private/Strata/StrataForwardLighting.ush
--- a/Engine/Shaders/Private/Strata/StrataForwardLighting.ush
+++ b/Engine/Shaders/Private/Strata/StrataForwardLighting.ush
@@ -258,30 +258,8 @@ float3 StrataForwardLighting(
 	FStrataTree StrataTree = StrataPixelHeader.StrataTree;
 	if (StrataTree.BSDFCount > 0)
 	{
-#if STRATA_USES_CONVERSION_FROM_LEGACY==0
-		//
-		// Update the coverage/transmittance of each leaves (==BSDFs) of the strata tree.
-		//
-		UpdateAllBSDFsOperatorCoverageTransmittance(StrataPixelHeader, StrataTree, bRoughDiffuseEnabled, NullStrataAddressing, V, L);
-
-		//
-		// Propagate up the coverage/transmittance of each node in the strata tree.
-		//
-		UpdateAllOperatorsCoverageTransmittance(StrataTree);
-
-		//
-		// Update the luminance weight of each BSDF according to the operators it has to traverse up to the strata tree root node.
-		//
-		UpdateAllBSDFWithBottomUpOperatorVisit(StrataTree);
-#else
-		LegacyUpdateBSDFsOperators(StrataPixelHeader, StrataTree, bRoughDiffuseEnabled, NullStrataAddressing, V, L);
-#endif
-
-		//
-		// TODO the top root node now contains the coverage and transmittance of the entire material (used for blending and fog on transparent_.
-		//
-		OutCoverage = saturate(StrataTree.Operators[StrataData.OperatorIndex].Coverage);
-		OutTransmittancePreCoverage = saturate(StrataTree.Operators[StrataData.OperatorIndex].Transmittance);
+		// Update tree (coverage/transmittance/luminace weights)
+		StrataUpdateTree(NullStrataAddressing, StrataPixelHeader, StrataTree, StrataData, V, L, bRoughDiffuseEnabled, OutCoverage, OutTransmittancePreCoverage);
 
 		//
 		// Forward lighting
