Strata - fixed vertical layering and weighting when computing rough refraction.

diff --git a/Engine/Shaders/Private/DistortAccumulatePS.usf b/Engine/Shaders/Private/DistortAccumulatePS.usf
--- a/Engine/Shaders/Private/DistortAccumulatePS.usf
+++ b/Engine/Shaders/Private/DistortAccumulatePS.usf
@@ -96,31 +96,25 @@ void Main(
 	StrataPixelHeader.BSDFCount = MaterialParameters.StrataTree.BSDFCount;
 	StrataPixelHeader.SharedLocalBases = MaterialParameters.SharedLocalBases;
 
-	const float3 V = MaterialParameters.CameraVector;
-	const float3 DummyL = float3(1, 0, 0);
-
-
-	// The current way rough refraction are evaluated is
-	//  1- start with the dirac view Lobe
-	//	2- refract/scatter/transmit that Lobe though all BSDF of a layer. Compute the mean of all the resulting Lobes.
-	//  3- Use that new Lobe as input (pointing away from the surface) for the next layer until done.
-	StrataLobeStatistic WiLobe = StrataGetViewDiracLobe(V);
-
-	StrataLobeStatistic AccumulatedLayerLobes = StrataGetNullLobe();
-	float AccumulatedLayerLobesWeights = 0.0f;
-
+	// Retrieve rough refraction variance
+	float Sigma = 0.0f;
 	FStrataTree StrataTree = StrataPixelHeader.StrataTree;
 	if (StrataTree.BSDFCount > 0)
 	{
+		const float3 V = MaterialParameters.CameraVector;
+		const float3 DummyL = float3(1, 0, 0);
+
 		// Update tree (coverage/transmittance/luminace weights)
 		StrataUpdateTree(NullStrataAddressing, StrataPixelHeader, StrataTree, V, DummyL, false /*bRoughDiffuseEnabled*/);
 
-		WiLobe = StrataTree.Operators[StrataData.OperatorIndex].LobeStat;
+		StrataLobeStatistic WiLobe = StrataTree.Operators[StrataData.OperatorIndex].LobeStat;
 		//half3 LobeWorldDir = normalize(TransformWorldVectorToView(WiLobe.Mu));
+
+		Sigma = WiLobe.Sigma;
 	}
 
-	// Not write out variance (as this is more linear to accumulate than roughness)
-	OutVariance = float4(WiLobe.Sigma, 0.0f, 0.0f, 0.0f);
+	// Write out variance (as this is more linear to accumulate than roughness)
+	OutVariance = float4(Sigma, 0.0f, 0.0f, 0.0f);
 
 	// Only allow clipping if roughness is close to 0.
 	const bool TryToClip = OutVariance.x < 0.01f;
diff --git a/Engine/Shaders/Private/Strata/StrataStatisticalOperators.ush b/Engine/Shaders/Private/Strata/StrataStatisticalOperators.ush
--- a/Engine/Shaders/Private/Strata/StrataStatisticalOperators.ush
+++ b/Engine/Shaders/Private/Strata/StrataStatisticalOperators.ush
@@ -101,7 +101,20 @@ StrataLobeStatistic StrataGetTransmittedLobe(StrataLobeStatistic WiLobe, float3
 	return WoLobe;
 }
 
-StrataLobeStatistic StrataHorizontalMixLobe(StrataLobeStatistic A, StrataLobeStatistic B, float Mix)
+StrataLobeStatistic StrataWeightLobe(StrataLobeStatistic A, float Weight)
+{
+	// We simple weight the lob properties so that they progressively disapear as a function of coverage.
+	// This is not entirely correct though because if matter coverage is reduced there should be two lobe: 
+	//	- this lobe with visibility = coverage
+	//	- a dirac lobe with visibility = 1 - coverage
+	StrataLobeStatistic WoLobe = A;
+	WoLobe.E *= Weight;
+	WoLobe.Mu *= Weight;
+	WoLobe.Sigma *= Weight;
+	return WoLobe;
+}
+
+StrataLobeStatistic StrataHorizontalMixLobes(StrataLobeStatistic A, StrataLobeStatistic B, float Mix)
 {
 	StrataLobeStatistic WoLobe;
 	WoLobe.E = lerp(A.E, B.E, Mix);
diff --git a/Engine/Shaders/Private/Strata/StrataTree.ush b/Engine/Shaders/Private/Strata/StrataTree.ush
--- a/Engine/Shaders/Private/Strata/StrataTree.ush
+++ b/Engine/Shaders/Private/Strata/StrataTree.ush
@@ -97,6 +97,9 @@ void UpdateSingleBSDFOperatorCoverageTransmittance(
 	StrataLobeStatistic RefractedLobe = StrataGetRefractedLobe(StrataGetViewDiracLobe(-V), EnergyTerms.E, Roughness, InterfaceEta12);
 	// Account for transmittance of the medium (but ignoring scattering for now)
 	RefractedLobe.E *= CurrentBSDFTransmittance;
+	// Account for coverate
+	RefractedLobe = StrataWeightLobe(RefractedLobe, CurrentBSDFCoverage);
+	
 
 	StrataTree.Operators[CurrentBSDF.OperatorIndex].LobeStat = RefractedLobe;
 
@@ -121,7 +124,7 @@ void UpdateSingleOperatorCoverageTransmittance(inout FStrataTree StrataTree, int
 		Op.Coverage = Weight * OpA.Coverage;
 		Op.Transmittance = OpA.Transmittance;
 
-		Op.LobeStat.E *= Weight;
+		Op.LobeStat = StrataWeightLobe(OpA.LobeStat, Weight);
 		break;
 	}
 
@@ -156,8 +159,9 @@ void UpdateSingleOperatorCoverageTransmittance(inout FStrataTree StrataTree, int
 		Op.VerticalTopTransmittance = OpA.Transmittance;
 
 		// Project top lob into bottom layer
-		Op.LobeStat = StrataGetRefractedLobe(OpA.LobeStat, /*InterfaceFDG*/1.0f, StrataLobeVarianceToRoughness(OpA.LobeStat.Sigma), /*InterfaceEta12*/1.0f);
-		Op.LobeStat.E *= OpB.LobeStat.E; // Combine with botton layer medium and specular transmittance
+		const float InterfaceEta12 = 1.0f; // STRATA_TODO We assume layers all have the same IOR. But we should take into account the material later.
+		Op.LobeStat = StrataGetRefractedLobe(OpA.LobeStat, /*InterfaceFDG*/0.0f, StrataLobeVarianceToRoughness(OpB.LobeStat.Sigma), InterfaceEta12);
+		Op.LobeStat.E *= OpB.LobeStat.E * saturate(Op.VerticalTopTransmittance + (1.0f - Op.VerticalTopCoverage)); // Combine with botton layer medium and specular transmittance
 		break;
 	}
 
@@ -170,7 +174,7 @@ void UpdateSingleOperatorCoverageTransmittance(inout FStrataTree StrataTree, int
 		Op.Coverage = AMix * OpA.Coverage + BMix * OpB.Coverage;
 		Op.Transmittance = (AMix * OpA.Coverage * OpA.Transmittance + BMix * OpB.Coverage * OpB.Transmittance) / max(STRATA_EPSILON, AMix * OpA.Coverage + BMix * OpB.Coverage);
 
-		Op.LobeStat = StrataHorizontalMixLobe(OpA.LobeStat, OpB.LobeStat, BMix);
+		Op.LobeStat = StrataHorizontalMixLobes(OpA.LobeStat, OpB.LobeStat, BMix);
 		break;
 	}
 
@@ -183,7 +187,7 @@ void UpdateSingleOperatorCoverageTransmittance(inout FStrataTree StrataTree, int
 		Op.Coverage = saturate(OpA.Coverage + OpB.Coverage);
 		Op.Transmittance = lerp(OpB.Transmittance, OpA.Transmittance, AMixFactor);
 
-		Op.LobeStat = StrataHorizontalMixLobe(OpA.LobeStat, OpB.LobeStat, 0.5f);
+		Op.LobeStat = StrataHorizontalMixLobes(OpA.LobeStat, OpB.LobeStat, 0.5f);
 		break;
 	}
 
