Fixed rough refraction normal to work wihtout using th one plugged on the root node.

diff --git a/Engine/Shaders/Private/DistortAccumulatePS.usf b/Engine/Shaders/Private/DistortAccumulatePS.usf
--- a/Engine/Shaders/Private/DistortAccumulatePS.usf
+++ b/Engine/Shaders/Private/DistortAccumulatePS.usf
@@ -98,6 +98,7 @@ void Main(
 
 	// Retrieve rough refraction variance
 	float Sigma = 0.0f;
+	float3 RefractionWorldNormal = MaterialParameters.CameraVector;
 	FStrataTree StrataTree = StrataPixelHeader.StrataTree;
 	if (StrataTree.BSDFCount > 0)
 	{
@@ -107,10 +108,11 @@ void Main(
 		// Update tree (coverage/transmittance/luminace weights)
 		StrataUpdateTree(NullStrataAddressing, StrataPixelHeader, StrataTree, V, DummyL, false /*bRoughDiffuseEnabled*/);
 
-		StrataLobeStatistic WiLobe = StrataTree.Operators[StrataData.OperatorIndex].LobeStat;
-		//half3 LobeWorldDir = normalize(TransformWorldVectorToView(WiLobe.Mu));
+		RefractionWorldNormal = StrataTree.Operators[StrataData.OperatorIndex].RefractionWorldNormal;
 
-		Sigma = WiLobe.Sigma;
+		StrataLobeStatistic RefractionLobeStat = StrataTree.Operators[StrataData.OperatorIndex].RefractionLobeStat;
+
+		Sigma = RefractionLobeStat.Sigma;
 	}
 
 	// Write out variance (as this is more linear to accumulate than roughness)
@@ -122,7 +124,7 @@ void Main(
 	// Compute UV distortion
 	float2 BufferUVDistortion = ComputeBufferUVDistortion(
 		MaterialParameters, PixelMaterialInputs, ResolvedView,
-		MaterialParameters.WorldNormal, DistortionParams, ScreenUV, Refraction, TryToClip, EyeIndex);
+		RefractionWorldNormal, DistortionParams, ScreenUV, Refraction, TryToClip, EyeIndex);
 
 	// Sample depth at distortion offset
 	float2 DistortBufferUV = ScreenUV + BufferUVDistortion;
diff --git a/Engine/Shaders/Private/Strata/Strata.ush b/Engine/Shaders/Private/Strata/Strata.ush
--- a/Engine/Shaders/Private/Strata/Strata.ush
+++ b/Engine/Shaders/Private/Strata/Strata.ush
@@ -728,7 +728,8 @@ struct FStrataOperator
 	float3 VerticalTopTransmittance;	// Only used by vertical layering operators. It represent the transmittance of matter of the top layer (from the left sub tree).
 
 	// Rough refraction
-	StrataLobeStatistic LobeStat;
+	float3				RefractionWorldNormal;
+	StrataLobeStatistic RefractionLobeStat;
 };
 
 FStrataOperator GetInitialisedStrataOperator()
diff --git a/Engine/Shaders/Private/Strata/StrataTree.ush b/Engine/Shaders/Private/Strata/StrataTree.ush
--- a/Engine/Shaders/Private/Strata/StrataTree.ush
+++ b/Engine/Shaders/Private/Strata/StrataTree.ush
@@ -99,9 +99,10 @@ void UpdateSingleBSDFOperatorCoverageTransmittance(
 	RefractedLobe.E *= CurrentBSDFTransmittance;
 	// Account for coverate
 	RefractedLobe = StrataWeightLobe(RefractedLobe, CurrentBSDFCoverage);
+	StrataTree.Operators[CurrentBSDF.OperatorIndex].RefractionLobeStat = RefractedLobe;
 	
-
-	StrataTree.Operators[CurrentBSDF.OperatorIndex].LobeStat = RefractedLobe;
+	float3x3 TangentBasis = StrataGetBSDFSharedBasis_InlineShading(StrataPixelHeader, CurrentBSDF, NullStrataAddressing);
+	StrataTree.Operators[CurrentBSDF.OperatorIndex].RefractionWorldNormal = TangentBasis[2];
 
 #endif // STRATA_INLINE_SHADING
 }
@@ -124,7 +125,8 @@ void UpdateSingleOperatorCoverageTransmittance(inout FStrataTree StrataTree, int
 		Op.Coverage = Weight * OpA.Coverage;
 		Op.Transmittance = OpA.Transmittance;
 
-		Op.LobeStat = StrataWeightLobe(OpA.LobeStat, Weight);
+		Op.RefractionLobeStat = StrataWeightLobe(OpA.RefractionLobeStat, Weight);
+		Op.RefractionWorldNormal = OpA.RefractionWorldNormal * Weight;
 		break;
 	}
 
@@ -160,8 +162,12 @@ void UpdateSingleOperatorCoverageTransmittance(inout FStrataTree StrataTree, int
 
 		// Project top lob into bottom layer
 		const float InterfaceEta12 = 1.0f; // STRATA_TODO We assume layers all have the same IOR. But we should take into account the material later.
-		Op.LobeStat = StrataGetRefractedLobe(OpA.LobeStat, /*InterfaceFDG*/0.0f, StrataLobeVarianceToRoughness(OpB.LobeStat.Sigma), InterfaceEta12);
-		Op.LobeStat.E *= OpB.LobeStat.E * saturate(Op.VerticalTopTransmittance + (1.0f - Op.VerticalTopCoverage)); // Combine with botton layer medium and specular transmittance
+		Op.RefractionLobeStat = StrataGetRefractedLobe(OpA.RefractionLobeStat, /*InterfaceFDG*/0.0f, StrataLobeVarianceToRoughness(OpB.RefractionLobeStat.Sigma), InterfaceEta12);
+		Op.RefractionLobeStat.E *= OpB.RefractionLobeStat.E * saturate(Op.VerticalTopTransmittance + (1.0f - Op.VerticalTopCoverage)); // Combine with botton layer medium and specular transmittance
+
+		// The ratio of bottomSurfaceVisbile and visible through top layer (with a scale of 0.5 because at maximum both surfaces will be visible) as a fonction of the total coverage.
+		const float RefractionNormalMix = saturate((Info.SurfaceBottom + 0.5f * Info.TransmittanceTopAndBottom * dot(TopTransmittance, 0.33.xxx)) / max(STRATA_EPSILON, Info.Coverage));
+		Op.RefractionWorldNormal = lerp(OpA.RefractionWorldNormal, OpB.RefractionWorldNormal, RefractionNormalMix);
 		break;
 	}
 
@@ -174,7 +180,8 @@ void UpdateSingleOperatorCoverageTransmittance(inout FStrataTree StrataTree, int
 		Op.Coverage = AMix * OpA.Coverage + BMix * OpB.Coverage;
 		Op.Transmittance = (AMix * OpA.Coverage * OpA.Transmittance + BMix * OpB.Coverage * OpB.Transmittance) / max(STRATA_EPSILON, AMix * OpA.Coverage + BMix * OpB.Coverage);
 
-		Op.LobeStat = StrataHorizontalMixLobes(OpA.LobeStat, OpB.LobeStat, BMix);
+		Op.RefractionLobeStat = StrataHorizontalMixLobes(OpA.RefractionLobeStat, OpB.RefractionLobeStat, BMix);
+		Op.RefractionWorldNormal = lerp(OpA.RefractionWorldNormal, OpB.RefractionWorldNormal, BMix);
 		break;
 	}
 
@@ -187,7 +194,8 @@ void UpdateSingleOperatorCoverageTransmittance(inout FStrataTree StrataTree, int
 		Op.Coverage = saturate(OpA.Coverage + OpB.Coverage);
 		Op.Transmittance = lerp(OpB.Transmittance, OpA.Transmittance, AMixFactor);
 
-		Op.LobeStat = StrataHorizontalMixLobes(OpA.LobeStat, OpB.LobeStat, 0.5f);
+		Op.RefractionLobeStat = StrataHorizontalMixLobes(OpA.RefractionLobeStat, OpB.RefractionLobeStat, 0.5f);
+		Op.RefractionWorldNormal = lerp(OpA.RefractionWorldNormal, OpB.RefractionWorldNormal, 0.5f);
 		break;
 	}
 
