Change Strata BSDF sampling routine to support 'multiple-lobe' sampling with MIS weighting.

diff --git a/Engine/Shaders/Private/Strata/StrataEvaluation.ush b/Engine/Shaders/Private/Strata/StrataEvaluation.ush
--- a/Engine/Shaders/Private/Strata/StrataEvaluation.ush
+++ b/Engine/Shaders/Private/Strata/StrataEvaluation.ush
@@ -16,6 +16,7 @@
 #include "../MonteCarlo.ush"
 #include "../SHCommon.ush"
 #include "../ShadingModels.ush"
+#include "../ShadingModelsSampling.ush"
 #include "../ThinFilmBSDF.ush"
 #include "../AreaLightCommon.ush"
 
@@ -165,6 +166,9 @@ void StrataUpdateBSDFContext(inout FStrataBSDFContext BSDFContext, float3 NewL)
 	BSDFContext.SatNoV = saturate(BSDFContext.Context.NoV);
 }
 
+///////////////////////////////////////////////////////////////////////////////////////////////////
+// BSDF evaluation (punctual/area)
+
 struct FStrataEvaluateResult
 {
 	float3 IntegratedDiffuseValue;
@@ -673,6 +677,7 @@ FStrataEvaluateResult StrataEvaluateBSDFCommon(FStrataBSDFContext BSDFContext, F
 	return Sample;
 }
 
+// Punctual light evaluation
 FStrataEvaluateResult StrataEvaluateBSDF(FStrataBSDFContext BSDFContext, bool bRoughDiffuseEnabled)
 {
 	FShadowTerms IdentityShadow = { 1, 1, 1, InitHairTransmittanceData() };
@@ -684,6 +689,7 @@ FStrataEvaluateResult StrataEvaluateBSDF(FStrataBSDFContext BSDFContext, bool bR
 	return StrataEvaluateBSDFCommon(BSDFContext, IdentityShadow, DummyAreaLightContext, INTEGRATION_PUNCTUAL_LIGHT, bForceFullyRough, bRoughDiffuseEnabled);
 }
 
+// Area light evaluation
 FStrataEvaluateResult StrataIntegrateBSDF(FStrataBSDFContext BSDFContext, FShadowTerms Shadow, FAreaLightIntegrateContext AreaLightContext, bool bForceFullyRough, bool bRoughDiffuseEnabled)
 {
 	BRANCH
@@ -703,14 +709,57 @@ FStrataEvaluateResult StrataIntegrateBSDF(FStrataBSDFContext BSDFContext, FShado
 	return StrataEvaluateBSDFCommon(BSDFContext, Shadow, AreaLightContext, INTEGRATION_PUNCTUAL_LIGHT, bForceFullyRough, bRoughDiffuseEnabled);
 }
 
-struct FStrataImportanceSampleResult
+///////////////////////////////////////////////////////////////////////////////////////////////////
+// BSDF Importance sampling
+ 
+// STRATA_TODO: merge this with Path-Tracing
+
+// Given two lobes that will roughly contribute colors A and B to the total (estimated for example using directional albedo)
+// return the probability of choosing lobe A
+float StrataLobeSelectionProb(float3 A, float3 B)
 {
-	float3 SpecularDirection;
-	float  SpecularPDF;
+	const float SumA = A.x + A.y + A.z;
+	const float SumB = B.x + B.y + B.z;
+	return SumA / (SumA + SumB + 1e-6);
+}
+
+// Takes as input the sample weight and pdf for a certain lobe of a mixed model, together with the probability of picking that lobe
+// This function then updates a running total Weight and Pdf value that represents the overall contribution of the BxDF
+// This function should be called when a BxDF is made up of multiple lobes combined with a sum to correctly account for the probability
+// of sampling directions via all lobes.
+// NOTE: this function also contains special logic to handle cases with infinite pdfs cleanly
+void StrataAddLobeWithMIS(inout float3 Weight, inout float Pdf, float3 LobeWeight, float LobePdf, float LobeProb)
+{
+	const float MinLobeProb = 1.1754943508e-38; // smallest normal float
+	if (LobeProb > MinLobeProb)
+	{
+		LobePdf *= LobeProb;
+		LobeWeight *= 1 / LobeProb;
+
+		// See discussion in MISWeightRobust for why this is better than LobePdf / (Pdf + LobePdf)
+		float MISWeight;
+		if (Pdf < LobePdf)
+			MISWeight = 1 / (1 + Pdf / LobePdf);
+		else if (LobePdf < Pdf)
+			MISWeight = 1 - 1 / (1 + LobePdf / Pdf);
+		else
+			MISWeight = 0.5f; // avoid (rare) inf/inf
+
+		Weight = lerp(Weight, LobeWeight, MISWeight);
+		Pdf += LobePdf;
+	}
+}
+
+// Select a term based based on a random value and a PDF, and returns the selected term
+// The random value is rescale/updated to reusable after evaluation
+#define StrataSelectLobe_Type(TDataType) TDataType StrataSelectLobe(inout float InE, float InPDF, TDataType InA, TDataType InB)\
+{\
+	if (InE < InPDF){ InE /= InPDF;							return InA; }\
+	else			{ InE -= InPDF; InE /= (1.0 - InPDF);	return InB; }\
+}
+StrataSelectLobe_Type(uint)
+StrataSelectLobe_Type(float)
 
-	float3 DiffuseDirection;
-	float  DiffusePDF;
-};
 
 /**
  * Importance sample a Strata BSDF
@@ -718,13 +767,11 @@ struct FStrataImportanceSampleResult
  * E: two random numbers
  * CameraVector: vector from the camera to the considered direction (not V)
  */
-FStrataImportanceSampleResult StrataImportanceSampleBSDF(FStrataBSDFContext BSDFContext, FStrataBSDF BSDF, float2 E)
+#define FStrataImportanceSampleResult FBxDFSample
+FStrataImportanceSampleResult StrataImportanceSampleBSDF(FStrataBSDFContext BSDFContext, FStrataBSDF BSDF, float2 E, uint InTermMask, bool bRoughDiffuseEnabled)
 {
-	FStrataImportanceSampleResult Sample = (FStrataImportanceSampleResult)0;
-
-	float PDF = 0.0f;
-	const uint BSDFType = BSDF_GETTYPE(BSDF);
-	switch (BSDFType)
+	FStrataImportanceSampleResult Out = (FStrataImportanceSampleResult)0;
+	switch (BSDF_GETTYPE(BSDF))
 	{
 		case STRATA_BSDF_TYPE_SLAB:
 		{
@@ -732,38 +779,89 @@ FStrataImportanceSampleResult StrataImportanceSampleBSDF(FStrataBSDFContext BSDF
 			const float3 F0 			= SLAB_F0(BSDFContext.BSDF);
 			const float3 F90 			= SLAB_F90(BSDFContext.BSDF);
 
-			// We assume a cosine hemisphere sampling is enough for any roughness.
-			// We also assume that this will work just as well for volumetric (considering it is an isotropic phase function and we are interesting in sampling mostly the hemisphere and not the layers below)
-			float4 ImportanceSample = CosineSampleHemisphere(E);
-			Sample.DiffuseDirection = mul(ImportanceSample.xyz, BSDFContext.TangentBasis);
-			Sample.DiffusePDF = ImportanceSample.w;
-			PDF = Sample.DiffusePDF;
+			const bool bHasAnisotropy	= BSDF_GETHASANISOTROPY(BSDFContext.BSDF);
+			const bool bHaziness		= BSDF_GETHASHAZINESS(BSDFContext.BSDF);
+			const float SafeRoughness   = MakeRoughnessSafe(SLAB_ROUGHNESS(BSDFContext.BSDF));
 
-			const float SafeRoughness = MakeRoughnessSafe(SLAB_ROUGHNESS(BSDFContext.BSDF));
-			const bool bHasAnisotropy = BSDF_GETHASANISOTROPY(BSDFContext.BSDF);
-			float4 TangentH = 0;
-			if (bHasAnisotropy)
+			// Pick diffuse/specular lobe
+			float TermPDF = Out.Term == SHADING_TERM_DIFFUSE ? 1 : 0;
+			Out.Term = InTermMask;
+			if (Out.Term == (SHADING_TERM_DIFFUSE | SHADING_TERM_SPECULAR))
 			{
-				float2 Alpha = 0;
-				GetAnisotropicRoughness(SafeRoughness, SLAB_ANISOTROPY(BSDFContext.BSDF), Alpha.x, Alpha.y);
-				TangentH = ImportanceSampleVisibleGGX_aniso(UniformSampleDisk(E), Alpha, BSDFContext.TangentV);
+				// Use (main) specular lobe as a heuristic to pick select between specular/diffuse lobe
+				const FBxDFEnergyTerms SpecEnergyTerms = ComputeGGXSpecEnergyTerms(SafeRoughness, BSDFContext.Context.NoV, F0, F90);
+
+				// Probability of picking diffuse lobe vs. specular lobe
+				TermPDF  = StrataLobeSelectionProb(DiffuseColor * (1 - SpecEnergyTerms.E), SpecEnergyTerms.E);
+				Out.Term = StrataSelectLobe(E.x, TermPDF, SHADING_TERM_DIFFUSE, SHADING_TERM_SPECULAR);
 			}
-			else
+
+			// Pick specular lobe
+			float SpecularLobePDF = 1.f;
+			float SpecularRoughness = SafeRoughness;
+			if (Out.Term == SHADING_TERM_SPECULAR && bHaziness)
 			{
-				TangentH = ImportanceSampleVisibleGGX(UniformSampleDisk(E), Pow4(SafeRoughness), BSDFContext.TangentV);
+				const float HazeWeight = SLAB_HAZINESS(BSDFContext.BSDF);
+				const float HazeRoughness = MakeRoughnessSafe(StrataComputeHazeRoughness(SafeRoughness));
+
+				SpecularLobePDF   = 1.f - HazeWeight;
+				SpecularRoughness = StrataSelectLobe(E.x, SpecularLobePDF, SafeRoughness, HazeRoughness);
 			}
 
-			// STRATA_TODO take into account Haziness
-			// STRATA_TODO take into account Fuzz: Per "Production Friendly Microfacet Sheen BRDF", uniform hemispherical sampling gives good result as the roughness is usually high and avoid issue at grazing angle.
-			// STRATA_TODO take into account Two sided material
+			// Sample Diffuse lobe
+			if (Out.Term == SHADING_TERM_DIFFUSE)
+			{
+				float4 ImportanceSample = CosineSampleHemisphere(E);
+				Out.L   = mul(ImportanceSample.xyz, BSDFContext.TangentBasis);
+				Out.PDF = ImportanceSample.w;
+			}
 
-			const float HPDF = TangentH.w;
-			const float3 H = mul(TangentH.xyz, BSDFContext.TangentBasis);
-			const float VoH = saturate(dot(BSDFContext.V, H));
+			// Sample Specular lobe
+			if (Out.Term == SHADING_TERM_SPECULAR)
+			{
+				float4 TangentH = 0;
+				if (bHasAnisotropy)
+				{
+					float2 Alpha = 0;
+					GetAnisotropicRoughness(SpecularRoughness, SLAB_ANISOTROPY(BSDFContext.BSDF), Alpha.x, Alpha.y);
+					TangentH = ImportanceSampleVisibleGGX_aniso(UniformSampleDisk(E), Alpha, BSDFContext.TangentV);
+				}
+				else
+				{
+					TangentH = ImportanceSampleVisibleGGX(UniformSampleDisk(E), Pow4(SpecularRoughness), BSDFContext.TangentV);
+				}
+
+				const float HPDF = TangentH.w;
+				const float3 H = mul(TangentH.xyz, BSDFContext.TangentBasis);
+				const float VoH = saturate(dot(BSDFContext.V, H));
+
+				Out.L   = 2 * dot(BSDFContext.V, H) * H - BSDFContext.V;
+				Out.PDF = RayPDFToReflectionRayPDF(VoH, HPDF);
+			}
+
+			// Evaluate lobes
+			StrataUpdateBSDFContext(BSDFContext, Out.L);
+			const FStrataEvaluateResult Evaluate = StrataEvaluateBSDF(BSDFContext, bRoughDiffuseEnabled);
+
+			// Add lobes mixtures
+			Out.PDF = 0;
+			Out.Weight = 0;
+			if (InTermMask & SHADING_TERM_DIFFUSE)
+			{
+				StrataAddLobeWithMIS(Out.Weight, Out.PDF, Evaluate.DiffusePathValue / Evaluate.DiffusePDF, Evaluate.DiffusePDF, TermPDF);
+			}
 
-			Sample.SpecularDirection = 2 * dot(BSDFContext.V, H) * H - BSDFContext.V;
-			Sample.SpecularPDF = RayPDFToReflectionRayPDF(VoH, HPDF);
-			PDF = Sample.SpecularPDF;
+			if (InTermMask & SHADING_TERM_SPECULAR)
+			{
+				StrataAddLobeWithMIS(Out.Weight, Out.PDF, Evaluate.SpecularPathValue / Evaluate.SpecularPDF, Evaluate.SpecularPDF, (1-TermPDF) * SpecularLobePDF);
+				if (bHaziness)
+				{
+					StrataAddLobeWithMIS(Out.Weight, Out.PDF, Evaluate.SpecularHazePathValue / Evaluate.SpecularHazePDF, Evaluate.SpecularPDF, (1 - TermPDF) * (1-SpecularLobePDF));
+				}
+			}
+
+			// STRATA_TODO take into account Two sided material
+			// STRATA_TODO take into account Transmission for path-tracing
 			break;
 		}
 
@@ -771,8 +869,10 @@ FStrataImportanceSampleResult StrataImportanceSampleBSDF(FStrataBSDFContext BSDF
 		{
 			// STRATA_TODO do something better when we get there with Lumen, and evaluate the different researched solution (e.g. Importance Sampling for Physically-Based Hair Fiber Models)
 			float4 ImportanceSample = CosineSampleHemisphere(E);
-			Sample.DiffuseDirection = mul(ImportanceSample.xyz, BSDFContext.TangentBasis);
-			PDF = ImportanceSample.w;
+			Out.L    = mul(ImportanceSample.xyz, BSDFContext.TangentBasis);
+			Out.PDF  = ImportanceSample.w;
+			Out.Term = SHADING_TERM_SPECULAR;
+			Out.Weight = 1.f;
 			break;
 		}
 
@@ -782,22 +882,7 @@ FStrataImportanceSampleResult StrataImportanceSampleBSDF(FStrataBSDFContext BSDF
 		//Nothing to do in this case because these BSDF are evaluated in other specialised passes.
 		//break;
 	}
-
-	//StrataUpdateBSDFContext(BSDFContext, Sample.Direction);
-	//Sample.Eval = StrataEvaluateBSDF(BSDFContext);
-	//Sample.Eval.PDF = PDF;	// We override the pdf to make sure the compiler skips all pdf code from StrataEvaluateBSDF because GGXSample.w si the pdf already for instance.
-
-	return Sample;
-}
-
-FStrataImportanceSampleResult StrataImportanceSampleBSDF(FStrataPixelHeader StrataPixelHeader, FStrataBSDF BSDF, const FStrataAddressing StrataAddressing, float2 E, float3 CameraVector)
-{
-	FStrataImportanceSampleResult Sample = (FStrataImportanceSampleResult)0;
-
-	// Create a context from the information we have
-	const float3 V = -CameraVector;
-	FStrataBSDFContext BSDFContext = StrataCreateBSDFContext(StrataPixelHeader, BSDF, StrataAddressing, V);
-	return StrataImportanceSampleBSDF(BSDFContext, BSDF, E);
+	return Out;
 }
 
 struct FStrataEnvLightResult
