Fixed more volumetric cloud + LWC visual when 5Trillion away from origin .

diff --git a/Engine/Shaders/Private/VolumetricCloud.usf b/Engine/Shaders/Private/VolumetricCloud.usf
--- a/Engine/Shaders/Private/VolumetricCloud.usf
+++ b/Engine/Shaders/Private/VolumetricCloud.usf
@@ -396,11 +396,12 @@ void MainCommon(in FMaterialPixelParameters MaterialParameters, in float4 SvPosi
 	FPixelMaterialInputs PixelMaterialInputs	= (FPixelMaterialInputs)0;
 	CalcMaterialParameters(MaterialParameters, PixelMaterialInputs, SvPosition, true);
 
-	float3 RayOrigin			= LWCHackToFloat(ResolvedView.WorldCameraOrigin);
-	float3 RayOriginKm			= RayOrigin * CENTIMETER_TO_KILOMETER;
-	float3 Raydir				= -MaterialParameters.CameraVector;
-	float TMin					= -999999999.0f;
-	float TMax					= -999999999.0f;
+	float3 RayWorldOriginKm				= LWCHackToFloat(ResolvedView.WorldCameraOrigin) * CENTIMETER_TO_KILOMETER;
+	float3 RayTranslatedWorldOrigin		= View.TranslatedWorldCameraOrigin;
+	float3 RayTranslatedWorldOriginKm	= RayTranslatedWorldOrigin * CENTIMETER_TO_KILOMETER;
+	float3 Raydir						= -MaterialParameters.CameraVector;
+	float TMin							= -999999999.0f;
+	float TMax							= -999999999.0f;
 
 	float3 Luminance			= 0.0f;
 	MATVEC TransmittanceToView	= 1.0f;
@@ -414,12 +415,11 @@ void MainCommon(in FMaterialPixelParameters MaterialParameters, in float4 SvPosi
 		//
 		// Check tracing start and end position within the cloud layer
 		//
-
 		float2 tTop2 = 0.0f;
-		if (RayIntersectSphereSolution(RayOriginKm, Raydir, float4(RenderVolumetricCloudParameters.CloudLayerCenterKm, RenderVolumetricCloudParameters.TopRadiusKm), tTop2))
+		if (RayIntersectSphereSolution(RayWorldOriginKm, Raydir, float4(RenderVolumetricCloudParameters.CloudLayerCenterKm, RenderVolumetricCloudParameters.TopRadiusKm), tTop2))
 		{
 			float2 tBottom2 = 0.0f;
-			if (RayIntersectSphereSolution(RayOriginKm, Raydir, float4(RenderVolumetricCloudParameters.CloudLayerCenterKm, RenderVolumetricCloudParameters.BottomRadiusKm), tBottom2))
+			if (RayIntersectSphereSolution(RayWorldOriginKm, Raydir, float4(RenderVolumetricCloudParameters.CloudLayerCenterKm, RenderVolumetricCloudParameters.BottomRadiusKm), tBottom2))
 			{
 				// If we see both intersection in front of us, keep the min/closest, otherwise the max/furthest
 				float TempTop = all(tTop2 > 0.0f) ? min(tTop2.x, tTop2.y) : max(tTop2.x, tTop2.y);
@@ -474,8 +474,8 @@ void MainCommon(in FMaterialPixelParameters MaterialParameters, in float4 SvPosi
 #if HAS_INVERTED_Z_BUFFER
 		DeviceZ = max(0.000000000001, DeviceZ);
 #endif
-		float3 DepthBufferWorldPos = LWCHackToFloat(SvPositionToWorld(float4(SvPosition.xy, DeviceZ, 1.0)));
-		float TDepthBuffer = length(DepthBufferWorldPos - RayOrigin);
+		float3 DepthBufferTranslatedWorldPos = SvPositionToTranslatedWorld(float4(SvPosition.xy, DeviceZ, 1.0));
+		float TDepthBuffer = length(DepthBufferTranslatedWorldPos - RayTranslatedWorldOrigin);
 		OutColor1.y = TDepthBuffer * CENTIMETER_TO_KILOMETER;
 
 		//
@@ -624,9 +624,8 @@ void MainCommon(in FMaterialPixelParameters MaterialParameters, in float4 SvPosi
 		for (uint i = 0; i < IStepCount; ++i)
 		{
 			int ms;
-			float3 SampleWorldPosition = RayOrigin + t * Raydir;
-			float3 SampleTranslatedWorldPosition = SampleWorldPosition + LWCHackToFloat(PrimaryView.PreViewTranslation); // TODO_ keep world pos as LWCVEctor
-
+			float3	  SampleTranslatedWorldPosition = RayTranslatedWorldOriginKm + t * Raydir;
+			FLWCVector3 SampleWorldPosition = LWCSubtract(LWCPromote(SampleTranslatedWorldPosition), PrimaryView.PreViewTranslation);
 
 			//////////////////////////////
 			// Update FMaterialPixelParameters according to current sample.
@@ -671,7 +670,7 @@ void MainCommon(in FMaterialPixelParameters MaterialParameters, in float4 SvPosi
 			ParticipatingMediaPhaseContext PMPC = SetupParticipatingMediaPhaseContext(Phase0, Phase1, MsPhaseFactor);
 #endif
 
-			const float3 PlanetCenterToWorldPos = (SampleWorldPosition - ResolvedView.SkyPlanetTranslatedWorldCenterAndViewHeight.xyz) * CM_TO_SKY_UNIT;
+			const float3 PlanetCenterToWorldPos = (SampleTranslatedWorldPosition - ResolvedView.SkyPlanetTranslatedWorldCenterAndViewHeight.xyz) * CM_TO_SKY_UNIT;
 #if CLOUD_PER_SAMPLE_ATMOSPHERE_TRANSMITTANCE
 			// Apply per sample change of transmittance due to atmosphere. More expenssive but higher quality, (and required for space view)
 			const float3 AtmosphereTransmittanceToLight0 = GetAtmosphereTransmittance(
@@ -741,7 +740,7 @@ void MainCommon(in FMaterialPixelParameters MaterialParameters, in float4 SvPosi
 						float CurrentNormT = ShadowT * ShadowT; // CurrentNormT is the end of the considered segment to integrate, PreviousNormT is its beginning.
 						const float DetlaNormT = CurrentNormT - PreviousNormT;
 						const float ShadowSampleDistance = ShadowLengthTest * (CurrentNormT - 0.5 * DetlaNormT);
-						UpdateMaterialCloudParam(MaterialParameters, SampleWorldPosition + float3(0.0f, 0.0f, -1.0f) * ShadowSampleDistance,
+						UpdateMaterialCloudParam(MaterialParameters, LWCAdd(SampleWorldPosition, LWCPromote(float3(0.0f, 0.0f, -1.0f) * ShadowSampleDistance)),
 							ResolvedView, CloudLayerParams, ShadowSampleDistance);
 						PreviousNormT = CurrentNormT;
 #if MATERIAL_VOLUMETRIC_ADVANCED_CONSERVATIVE_DENSITY
@@ -761,7 +760,7 @@ void MainCommon(in FMaterialPixelParameters MaterialParameters, in float4 SvPosi
 					for (float ShadowT = 0.5; ShadowT < ShadowStepCount; ShadowT += 1.0f)
 					{
 						const float ShadowSampleDistance = ShadowLengthTest * (ShadowT * InvShadowStepCount);
-						UpdateMaterialCloudParam(MaterialParameters, SampleWorldPosition + float3(0.0f, 0.0f, -1.0f) * ShadowSampleDistance,
+						UpdateMaterialCloudParam(MaterialParameters, LWCAdd(SampleWorldPosition, LWCPromote(float3(0.0f, 0.0f, -1.0f) * ShadowSampleDistance)),
 							ResolvedView, CloudLayerParams, ShadowSampleDistance);
 #if MATERIAL_VOLUMETRIC_ADVANCED_CONSERVATIVE_DENSITY
 						if (MaterialParameters.VolumeSampleConservativeDensity.x <= 0.0f)
@@ -832,7 +831,7 @@ void MainCommon(in FMaterialPixelParameters MaterialParameters, in float4 SvPosi
 				for (float ShadowT = ShadowJitterNorm; ShadowT < ShadowStepCount; ShadowT += 1.0f)
 				{
 					const float ShadowSampleDistance = ShadowLengthTest * (ShadowT * InvShadowStepCount);
-					UpdateMaterialCloudParam(MaterialParameters, SampleWorldPosition + Light0Direction * ShadowSampleDistance,
+					UpdateMaterialCloudParam(MaterialParameters, LWCAdd(SampleWorldPosition, LWCPromote(Light0Direction * ShadowSampleDistance)),
 						ResolvedView, CloudLayerParams, ShadowSampleDistance);
 					const float ExtinctionFactor = 1.0f;
 #else
@@ -845,7 +844,7 @@ void MainCommon(in FMaterialPixelParameters MaterialParameters, in float4 SvPosi
 					const float DetlaNormT = CurrentNormT - PreviousNormT;
 					const float ExtinctionFactor = DetlaNormT;
 					const float ShadowSampleDistance = ShadowLengthTest * (PreviousNormT + DetlaNormT * ShadowJitterNorm);
-					UpdateMaterialCloudParam(MaterialParameters, SampleWorldPosition + Light0Direction * ShadowSampleDistance,
+					UpdateMaterialCloudParam(MaterialParameters, LWCAdd(SampleWorldPosition, LWCPromote(Light0Direction * ShadowSampleDistance)),
 						ResolvedView, CloudLayerParams, ShadowSampleDistance);
 					PreviousNormT = CurrentNormT;
 #endif
@@ -912,7 +911,7 @@ void MainCommon(in FMaterialPixelParameters MaterialParameters, in float4 SvPosi
 						const float DetlaNormT = CurrentNormT - PreviousNormT;
 						const float ExtinctionFactor = DetlaNormT;
 						const float ShadowSampleDistance = ShadowLengthTest * (PreviousNormT + DetlaNormT * ShadowJitterNorm);
-						UpdateMaterialCloudParam(MaterialParameters, SampleWorldPosition + Light1Direction * ShadowSampleDistance,
+						UpdateMaterialCloudParam(MaterialParameters, LWCAdd(SampleWorldPosition, LWCPromote(Light1Direction * ShadowSampleDistance)),
 							ResolvedView, CloudLayerParams, ShadowSampleDistance);
 						PreviousNormT = CurrentNormT;
 
@@ -980,7 +979,7 @@ void MainCommon(in FMaterialPixelParameters MaterialParameters, in float4 SvPosi
 			{ 
 				const float2 PixelPos = SvPosition.xy * RenderVolumetricCloudParameters.TracingCoordToFullResPixelCoordScaleBias.xy - ResolvedView.ViewRectMin.xy;// *RenderVolumetricCloudParameters.TracingCoordToFullResPixelCoordScaleBias.xy; // We do not take into account the bias so this owuld not work for split screens / VR.
 
-				float4 ClipPos = mul(float4(SampleWorldPosition, 1.0f), LWCHackToFloat(ResolvedView.WorldToClip));
+				float4 ClipPos = mul(float4(SampleTranslatedWorldPosition, 1.0f), PrimaryView.TranslatedWorldToClip);
 				ClipPos /= ClipPos.w;
 				const float SceneDepth = ConvertFromDeviceZ(ClipPos.z);
 				 
@@ -1018,7 +1017,7 @@ void MainCommon(in FMaterialPixelParameters MaterialParameters, in float4 SvPosi
 
 						float3 L = 0;
 						float3 ToLight = 0;
-						float LightMask = GetLocalLightAttenuation(SampleWorldPosition, LightData, ToLight, L);
+						float LightMask = GetLocalLightAttenuation(SampleTranslatedWorldPosition, LightData, ToLight, L);
 
 						float Lighting;
 						if( LightData.bRectLight )
@@ -1060,7 +1059,7 @@ void MainCommon(in FMaterialPixelParameters MaterialParameters, in float4 SvPosi
 							// Non-linear shadow sample distribution, hardcoded to x^2
 							for (float ShadowT = ShadowJitterNorm; ShadowT < ShadowStepCount; ShadowT += 1.0f)
 							{
-								UpdateMaterialCloudParam(MaterialParameters, SampleWorldPosition + SampleToLightNorm * ShadowLengthTest * (ShadowT * InvShadowStepCount), 
+								UpdateMaterialCloudParam(MaterialParameters, LWCAdd(SampleWorldPosition, LWCPromote(SampleToLightNorm * ShadowLengthTest * (ShadowT * InvShadowStepCount))),
 									ResolvedView, CloudLayerParams, TRACING_SHADOW_DISTANCE_OFF);
 								const float ExtinctionFactor = 1.0f;
 
@@ -1230,19 +1229,20 @@ void MainCommon(in FMaterialPixelParameters MaterialParameters, in float4 SvPosi
 	// It is better to limit depth to a close distance instead of max float to smooth out cloud edges when intersecting opaque meshes.
 	// No visual issues have been noticed with reprojection+TAA so far.
 	const float NoCloudDepth = TMax;
-
 	const float tAP = tAPWeightsSum==0.0f ? NoCloudDepth : tAPWeightedSum / max(0.0000000001f, tAPWeightsSum);
-	float3 AbsoluteWorldPosition = RayOrigin + tAP * Raydir;
-	float4 AbsoluteWorldPositionNDC = mul(float4(AbsoluteWorldPosition, 1.0f), LWCHackToFloat(ResolvedView.WorldToClip));
 	float MeanTransmittance = dot(TransmittanceToView.MATSUF, 1.0f / 3.0f);
 
+	float3 FogSampleTranslatedWorldPosition = RayTranslatedWorldOriginKm + tAP * Raydir;
+	FLWCVector3 FogSampleWorldPosition = LWCSubtract(LWCPromote(FogSampleTranslatedWorldPosition), PrimaryView.PreViewTranslation);
+	float4 ClipPos = mul(float4(FogSampleTranslatedWorldPosition, 1.0f), PrimaryView.TranslatedWorldToClip);
+
 	if (RenderVolumetricCloudParameters.EnableAerialPerspectiveSampling && tAPWeightsSum > 0.0f)
 	{
 		// Apply AP only once according to the mean position within the participating media weighted by transmittance/visibility.
 		// This allows to apply AP only once per pixel instead of per sample.
 		float4 AerialPerspective = GetAerialPerspectiveLuminanceTransmittance(
 			ResolvedView.RealTimeReflectionCapture, ResolvedView.SkyAtmosphereCameraAerialPerspectiveVolumeSizeAndInvSize,
-			AbsoluteWorldPositionNDC, AbsoluteWorldPosition * CM_TO_SKY_UNIT, LWCHackToFloat(ResolvedView.WorldCameraOrigin) * CM_TO_SKY_UNIT,
+			ClipPos, LWCHackToFloat(FogSampleWorldPosition) * CM_TO_SKY_UNIT, LWCHackToFloat(ResolvedView.WorldCameraOrigin) * CM_TO_SKY_UNIT,
 			View.CameraAerialPerspectiveVolume, View.CameraAerialPerspectiveVolumeSampler,
 			ResolvedView.SkyAtmosphereCameraAerialPerspectiveVolumeDepthResolutionInv,
 			ResolvedView.SkyAtmosphereCameraAerialPerspectiveVolumeDepthResolution,
@@ -1263,11 +1263,11 @@ void MainCommon(in FMaterialPixelParameters MaterialParameters, in float4 SvPosi
 	//
 	if (RenderVolumetricCloudParameters.EnableHeightFog && tAPWeightsSum > 0.0f)
 	{
-		float4 HeightFogInscatteringAndTransmittance = CalculateHeightFog(AbsoluteWorldPosition - LWCHackToFloat(ResolvedView.WorldCameraOrigin));
+		float4 HeightFogInscatteringAndTransmittance = CalculateHeightFog(FogSampleTranslatedWorldPosition - ResolvedView.TranslatedWorldCameraOrigin);
 		
 		if (FogStruct.ApplyVolumetricFog > 0)
 		{
-			float3 VolumeUV = ComputeVolumeUVFromNDC(AbsoluteWorldPositionNDC);
+			float3 VolumeUV = ComputeVolumeUVFromNDC(ClipPos);
 			const uint EyeIndex = 0;
 			HeightFogInscatteringAndTransmittance = CombineVolumetricFog(HeightFogInscatteringAndTransmittance, VolumeUV, EyeIndex);
 		}
@@ -1507,7 +1507,7 @@ void MainPS(
 		for (float ShadowT = 0.5; ShadowT < ShadowStepCount; ShadowT += 1.0f)
 		{
 			const float SampleT = ShadowLengthTest * (ShadowT * InvShadowStepCount);
-			UpdateMaterialCloudParam(MaterialParameters, WorldPosOnLayer + LightDirection * SampleT, 
+			UpdateMaterialCloudParam(MaterialParameters, LWCPromote(WorldPosOnLayer + LightDirection * SampleT),			// LWC TODO FIX ME
 				ResolvedView, CloudLayerParams, SampleT);
 
 		#if MATERIAL_VOLUMETRIC_ADVANCED_CONSERVATIVE_DENSITY
diff --git a/Engine/Shaders/Private/VolumetricCloudMaterialPixelCommon.ush b/Engine/Shaders/Private/VolumetricCloudMaterialPixelCommon.ush
--- a/Engine/Shaders/Private/VolumetricCloudMaterialPixelCommon.ush
+++ b/Engine/Shaders/Private/VolumetricCloudMaterialPixelCommon.ush
@@ -41,11 +41,10 @@ FCloudLayerParameters GetCloudLayerParams(
 
 // Function to update the material parameter structure for an evaluation
 void UpdateMaterialCloudParam(inout FMaterialPixelParameters MaterialParameters,
-	float3 AbsoluteWorldPosition, in ViewState InputView, in FCloudLayerParameters CloudLayerParams,
+	FLWCVector3 AbsoluteWorldPosition, in ViewState InputView, in FCloudLayerParameters CloudLayerParams,
 	in float ShadowSampleDistance)
 {
-	// TODO LWC
-	MaterialParameters.AbsoluteWorldPosition = LWCPromote(AbsoluteWorldPosition);
+	MaterialParameters.AbsoluteWorldPosition = AbsoluteWorldPosition;
 	MaterialParameters.WorldPosition_CamRelative = LWCToFloat(LWCAdd(MaterialParameters.AbsoluteWorldPosition, InputView.PreViewTranslation));
 	//MaterialParameters.WorldPosition_NoOffsets;				// TODO
 	//MaterialParameters.WorldPosition_NoOffsets_CamRelative;	// TODO
diff --git a/Engine/Shaders/Private/VolumetricFogVoxelization.usf b/Engine/Shaders/Private/VolumetricFogVoxelization.usf
--- a/Engine/Shaders/Private/VolumetricFogVoxelization.usf
+++ b/Engine/Shaders/Private/VolumetricFogVoxelization.usf
@@ -311,9 +311,7 @@ void VoxelizePS(
 		VoxelizeVolumePass.RenderVolumetricCloudParametersCloudLayerCenterKm, VoxelizeVolumePass.RenderVolumetricCloudParametersPlanetRadiusKm,
 		VoxelizeVolumePass.RenderVolumetricCloudParametersBottomRadiusKm, VoxelizeVolumePass.RenderVolumetricCloudParametersTopRadiusKm);
 	const float NotShadowSample = 0.0f;
-	float4 ScreenPosition = SvPositionToResolvedScreenPosition(SvPosition);
-	float3 TranslatedWorldPosition = SvPositionToResolvedTranslatedWorld(SvPosition);
-	float3 AbsoluteWorldPosition = TranslatedWorldPosition.xyz - LWCHackToFloat(ResolvedView.PreViewTranslation);
+	FLWCVector3 AbsoluteWorldPosition = SvPositionToWorld(SvPosition);
 	UpdateMaterialCloudParam(MaterialParameters, AbsoluteWorldPosition, ResolvedView, CloudLayerParams, NotShadowSample);
 #endif
 
