Volumetric Cloud minor clean up

diff --git a/Engine/Shaders/Private/VolumetricCloud.usf b/Engine/Shaders/Private/VolumetricCloud.usf
--- a/Engine/Shaders/Private/VolumetricCloud.usf
+++ b/Engine/Shaders/Private/VolumetricCloud.usf
@@ -397,8 +397,7 @@ void MainCommon(in FMaterialPixelParameters MaterialParameters, in float4 SvPosi
 	CalcMaterialParameters(MaterialParameters, PixelMaterialInputs, SvPosition, true);
 
 	float3 RayWorldOriginKm				= LWCHackToFloat(ResolvedView.WorldCameraOrigin) * CENTIMETER_TO_KILOMETER;
-	float3 RayTranslatedWorldOrigin		= View.TranslatedWorldCameraOrigin;
-	float3 RayTranslatedWorldOriginKm	= RayTranslatedWorldOrigin * CENTIMETER_TO_KILOMETER;
+	float3 RayTranslatedWorldOriginKm	= View.TranslatedWorldCameraOrigin * CENTIMETER_TO_KILOMETER;
 	float3 Raydir						= -MaterialParameters.CameraVector;
 	float TMin							= -999999999.0f;
 	float TMax							= -999999999.0f;
@@ -475,8 +474,8 @@ void MainCommon(in FMaterialPixelParameters MaterialParameters, in float4 SvPosi
 		DeviceZ = max(0.000000000001, DeviceZ);
 #endif
 		float3 DepthBufferTranslatedWorldPos = SvPositionToTranslatedWorld(float4(SvPosition.xy, DeviceZ, 1.0));
-		float TDepthBuffer = length(DepthBufferTranslatedWorldPos - RayTranslatedWorldOrigin);
-		OutColor1.y = TDepthBuffer * CENTIMETER_TO_KILOMETER;
+		const float TDepthBufferKm = length(DepthBufferTranslatedWorldPos * CENTIMETER_TO_KILOMETER - RayTranslatedWorldOriginKm);
+		OutColor1.y = TDepthBufferKm;
 
 		//
 		// Skip tracing if the depth buffer is in front of the cloud layer front interface. If inside the layer, clamp the tracing to the depth buffer
@@ -484,7 +483,7 @@ void MainCommon(in FMaterialPixelParameters MaterialParameters, in float4 SvPosi
 
 		const bool bDoDepthInteresectionAndTest = RenderVolumetricCloudParameters.OpaqueIntersectionMode >= 2;
 
-		if (bDoDepthInteresectionAndTest && TDepthBuffer < TMin)
+		if (bDoDepthInteresectionAndTest && (TDepthBufferKm * KILOMETER_TO_CENTIMETER) < TMin)
 		{
 			//OutColor0 = float4(1.0f, 1.0f, 0.0f, 1.0f);
 			return;
@@ -493,7 +492,7 @@ void MainCommon(in FMaterialPixelParameters MaterialParameters, in float4 SvPosi
 		if (bDoDepthInteresectionAndTest)
 		{
 			// Only trace up to the closest between the cloud layer far distance and the depth buffer 
-			TMax = min(TMax, TDepthBuffer);
+			TMax = min(TMax, TDepthBufferKm * KILOMETER_TO_CENTIMETER);
 		}
 
 
