LWC friendly GetAerialPerspectiveLuminanceTransmittance.

diff --git a/Engine/Shaders/Private/BasePassPixelShader.usf b/Engine/Shaders/Private/BasePassPixelShader.usf
--- a/Engine/Shaders/Private/BasePassPixelShader.usf
+++ b/Engine/Shaders/Private/BasePassPixelShader.usf
@@ -1514,7 +1514,7 @@ void FPixelShaderInOut_MainPS(
 		// Sample the aerial perspective (AP).
 		Fogging = GetAerialPerspectiveLuminanceTransmittanceWithFogOver(
 			ResolvedView.RealTimeReflectionCapture, ResolvedView.SkyAtmosphereCameraAerialPerspectiveVolumeSizeAndInvSize,
-			NDCPosition, LWCHackToFloat(MaterialParameters.AbsoluteWorldPosition)*CM_TO_SKY_UNIT, LWCHackToFloat(ResolvedView.WorldCameraOrigin)*CM_TO_SKY_UNIT,
+			NDCPosition, MaterialParameters.WorldPosition_CamRelative * CM_TO_SKY_UNIT,
 			View.CameraAerialPerspectiveVolume, View.CameraAerialPerspectiveVolumeSampler,
 			ResolvedView.SkyAtmosphereCameraAerialPerspectiveVolumeDepthResolutionInv,
 			ResolvedView.SkyAtmosphereCameraAerialPerspectiveVolumeDepthResolution,
diff --git a/Engine/Shaders/Private/BasePassVertexShader.usf b/Engine/Shaders/Private/BasePassVertexShader.usf
--- a/Engine/Shaders/Private/BasePassVertexShader.usf
+++ b/Engine/Shaders/Private/BasePassVertexShader.usf
@@ -117,7 +117,7 @@ void Main(
 		// Sample the aerial perspective (AP). It is also blended under the VertexFog parameter.
 		Output.BasePassInterpolants.VertexFog = GetAerialPerspectiveLuminanceTransmittanceWithFogOver(
 			ResolvedView.RealTimeReflectionCapture, ResolvedView.SkyAtmosphereCameraAerialPerspectiveVolumeSizeAndInvSize,
-			Output.Position, WorldPosition.xyz*CM_TO_SKY_UNIT, ResolvedView.TranslatedWorldCameraOrigin*CM_TO_SKY_UNIT,
+			Output.Position, (WorldPosition.xyz - ResolvedView.TranslatedWorldCameraOrigin) * CM_TO_SKY_UNIT,
 			View.CameraAerialPerspectiveVolume, View.CameraAerialPerspectiveVolumeSampler,
 			ResolvedView.SkyAtmosphereCameraAerialPerspectiveVolumeDepthResolutionInv,
 			ResolvedView.SkyAtmosphereCameraAerialPerspectiveVolumeDepthResolution,
diff --git a/Engine/Shaders/Private/HairStrands/HairStrandsVisibilityComposeSubPixelPS.usf b/Engine/Shaders/Private/HairStrands/HairStrandsVisibilityComposeSubPixelPS.usf
--- a/Engine/Shaders/Private/HairStrands/HairStrandsVisibilityComposeSubPixelPS.usf
+++ b/Engine/Shaders/Private/HairStrands/HairStrandsVisibilityComposeSubPixelPS.usf
@@ -59,8 +59,7 @@ float4 EvaluateVolumetric(float3 WorldPosition)
 			View.RealTimeReflectionCapture, 
 			View.SkyAtmosphereCameraAerialPerspectiveVolumeSizeAndInvSize,
 			NDCPosition, 
-			WorldPosition.xyz * CM_TO_SKY_UNIT, 
-			LWCHackToFloat(PrimaryView.WorldCameraOrigin).xyz * CM_TO_SKY_UNIT,
+			CameraRelative_WorldPosition * CM_TO_SKY_UNIT,
 			View.CameraAerialPerspectiveVolume, 
 			View.CameraAerialPerspectiveVolumeSampler,
 			View.SkyAtmosphereCameraAerialPerspectiveVolumeDepthResolutionInv,
diff --git a/Engine/Shaders/Private/HeightFogCommon.ush b/Engine/Shaders/Private/HeightFogCommon.ush
--- a/Engine/Shaders/Private/HeightFogCommon.ush
+++ b/Engine/Shaders/Private/HeightFogCommon.ush
@@ -208,8 +208,6 @@ half4 GetExponentialHeightFog(float3 WorldPositionRelativeToCamera, float Exclud
 
 half4 CalculateHeightFog(float3 WorldPositionRelativeToCamera)
 {
-	float3 WorldPosition = WorldPositionRelativeToCamera + ResolvedView.TranslatedWorldCameraOrigin - LWCHackToFloat(ResolvedView.PreViewTranslation);
-	
 	float ExcludeDistance = 0;
 
 #if PERMUTATION_SUPPORT_VOLUMETRIC_FOG
diff --git a/Engine/Shaders/Private/MaterialTemplate.ush b/Engine/Shaders/Private/MaterialTemplate.ush
--- a/Engine/Shaders/Private/MaterialTemplate.ush
+++ b/Engine/Shaders/Private/MaterialTemplate.ush
@@ -1567,7 +1567,7 @@ float4 MaterialExpressionSkyAtmosphereAerialPerspective(FMaterialPixelParameters
 
 	float4 AerialPerspective = GetAerialPerspectiveLuminanceTransmittance(
 		ResolvedView.RealTimeReflectionCapture, ResolvedView.SkyAtmosphereCameraAerialPerspectiveVolumeSizeAndInvSize,
-		NDCPosition, TranslatedWorldPosition, SkyCameraTranslatedWorldOrigin,
+		NDCPosition, TranslatedWorldPosition - SkyCameraTranslatedWorldOrigin,
 		View.CameraAerialPerspectiveVolume, View.CameraAerialPerspectiveVolumeSampler,
 		ResolvedView.SkyAtmosphereCameraAerialPerspectiveVolumeDepthResolutionInv,
 		ResolvedView.SkyAtmosphereCameraAerialPerspectiveVolumeDepthResolution,
diff --git a/Engine/Shaders/Private/MobileBasePassVertexShader.usf b/Engine/Shaders/Private/MobileBasePassVertexShader.usf
--- a/Engine/Shaders/Private/MobileBasePassVertexShader.usf
+++ b/Engine/Shaders/Private/MobileBasePassVertexShader.usf
@@ -106,7 +106,7 @@ void Main(
 		// Sample the aerial perspective (AP). It is also blended under the VertexFog parameter.
 		VertexFog = GetAerialPerspectiveLuminanceTransmittanceWithFogOver(
 			ResolvedView.RealTimeReflectionCapture, ResolvedView.SkyAtmosphereCameraAerialPerspectiveVolumeSizeAndInvSize,
-			Output.Position, WorldPosition.xyz*CM_TO_SKY_UNIT, ResolvedView.TranslatedWorldCameraOrigin*CM_TO_SKY_UNIT,
+			Output.Position, (WorldPosition.xyz - ResolvedView.TranslatedWorldCameraOrigin) * CM_TO_SKY_UNIT,
 			View.CameraAerialPerspectiveVolume, View.CameraAerialPerspectiveVolumeSampler,
 			ResolvedView.SkyAtmosphereCameraAerialPerspectiveVolumeDepthResolutionInv,
 			ResolvedView.SkyAtmosphereCameraAerialPerspectiveVolumeDepthResolution,
diff --git a/Engine/Shaders/Private/MobileFog.usf b/Engine/Shaders/Private/MobileFog.usf
--- a/Engine/Shaders/Private/MobileFog.usf
+++ b/Engine/Shaders/Private/MobileFog.usf
@@ -56,8 +56,7 @@ void MobileFogPS(
 			ResolvedView.RealTimeReflectionCapture, 
 			ResolvedView.SkyAtmosphereCameraAerialPerspectiveVolumeSizeAndInvSize,
 			NDCPosition, 
-			WorldPosition.xyz*CM_TO_SKY_UNIT, 
-			ResolvedView.TranslatedWorldCameraOrigin*CM_TO_SKY_UNIT,
+			(WorldPosition.xyz - ResolvedView.TranslatedWorldCameraOrigin) * CM_TO_SKY_UNIT,
 			View.CameraAerialPerspectiveVolume, 
 			View.CameraAerialPerspectiveVolumeSampler,
 			ResolvedView.SkyAtmosphereCameraAerialPerspectiveVolumeDepthResolutionInv,
diff --git a/Engine/Shaders/Private/SkyAtmosphere.usf b/Engine/Shaders/Private/SkyAtmosphere.usf
--- a/Engine/Shaders/Private/SkyAtmosphere.usf
+++ b/Engine/Shaders/Private/SkyAtmosphere.usf
@@ -871,7 +871,7 @@ void RenderSkyAtmosphereRayMarchingPS(
 
 	float4 AP = GetAerialPerspectiveLuminanceTransmittance(
 		ResolvedView.RealTimeReflectionCapture, ResolvedView.SkyAtmosphereCameraAerialPerspectiveVolumeSizeAndInvSize,
-		NDCPosition, DepthBufferTranslatedWorldPos *CM_TO_SKY_UNIT, GetCameraTranslatedWorldPos()*CM_TO_SKY_UNIT,
+		NDCPosition, (DepthBufferTranslatedWorldPos - GetCameraTranslatedWorldPos()) * CM_TO_SKY_UNIT,
 		CameraAerialPerspectiveVolumeTexture, CameraAerialPerspectiveVolumeTextureSampler,
 		SkyAtmosphere.CameraAerialPerspectiveVolumeDepthResolutionInv,
 		SkyAtmosphere.CameraAerialPerspectiveVolumeDepthResolution,
diff --git a/Engine/Shaders/Private/SkyAtmosphereCommon.ush b/Engine/Shaders/Private/SkyAtmosphereCommon.ush
--- a/Engine/Shaders/Private/SkyAtmosphereCommon.ush
+++ b/Engine/Shaders/Private/SkyAtmosphereCommon.ush
@@ -33,7 +33,7 @@ float2 FromSubUvsToUnit(float2 uv, float4 SizeAndInvSize) { return (uv - 0.5f *
 
 float4 GetAerialPerspectiveLuminanceTransmittance(
 	bool ViewIsRealTimeReflectionCapture, float4 CameraAerialPerspectiveVolumeSizeAndInvSize,
-	float4 NDC, float3 SampledWorldPos, float3 CameraWorldPos, 
+	float4 NDC, float3 WorldPositionRelativeToCamera,
 	Texture3D<float4> AerialPerspectiveVolumeTexture, SamplerState AerialPerspectiveVolumeTextureSampler,
 	float AerialPerspectiveVolumeDepthResolutionInv,
 	float AerialPerspectiveVolumeDepthResolution,
@@ -44,7 +44,7 @@ float4 GetAerialPerspectiveLuminanceTransmittance(
 {
 	float2 ScreenUv = (NDC.xy / NDC.ww) * float2(0.5f, -0.5f) + 0.5f;
 
-	float tDepth = max(0.0f, length(SampledWorldPos - CameraWorldPos) - AerialPerspectiveVolumeStartDepth);
+	float tDepth = max(0.0f, length(WorldPositionRelativeToCamera) - AerialPerspectiveVolumeStartDepth);
 
 	float LinearSlice = tDepth * AerialPerspectiveVolumeDepthSliceLengthKmInv;
 	float LinearW = LinearSlice * AerialPerspectiveVolumeDepthResolutionInv; // Depth slice coordinate in [0,1]
@@ -64,7 +64,7 @@ float4 GetAerialPerspectiveLuminanceTransmittance(
 	{
 		// We modify ScreenUv to sample the correct AP accordign to 360 reflection AP LUT.
 		// This is the inverse of what is in RenderCameraAerialPerspectiveVolumeCS.
-		float3 WorldDir = normalize(SampledWorldPos - CameraWorldPos);
+		float3 WorldDir = normalize(WorldPositionRelativeToCamera);
 
 		float SinPhi = WorldDir.z;
 		float CosPhi = sqrt(1.0f - SinPhi * SinPhi);
@@ -101,7 +101,7 @@ float4 GetAerialPerspectiveLuminanceTransmittance(
 
 float4 GetAerialPerspectiveLuminanceTransmittanceWithFogOver(
 	bool ViewIsRealTimeReflectionCapture, float4 CameraAerialPerspectiveVolumeSizeAndInvSize,
-	float4 NDC, float3 SampledWorldPos, float3 CameraWorldPos,
+	float4 NDC, float3 WorldPositionRelativeToCamera,
 	Texture3D<float4> AerialPerspectiveVolumeTexture, SamplerState AerialPerspectiveVolumeTextureSampler,
 	float AerialPerspectiveVolumeDepthResolutionInv,
 	float AerialPerspectiveVolumeDepthResolution,
@@ -112,7 +112,7 @@ float4 GetAerialPerspectiveLuminanceTransmittanceWithFogOver(
 {
 	float4 AP = GetAerialPerspectiveLuminanceTransmittance(
 		ViewIsRealTimeReflectionCapture, CameraAerialPerspectiveVolumeSizeAndInvSize,
-		NDC, SampledWorldPos, CameraWorldPos,
+		NDC, WorldPositionRelativeToCamera,
 		AerialPerspectiveVolumeTexture, AerialPerspectiveVolumeTextureSampler,
 		AerialPerspectiveVolumeDepthResolutionInv,
 		AerialPerspectiveVolumeDepthResolution,
diff --git a/Engine/Shaders/Private/VolumetricCloud.usf b/Engine/Shaders/Private/VolumetricCloud.usf
--- a/Engine/Shaders/Private/VolumetricCloud.usf
+++ b/Engine/Shaders/Private/VolumetricCloud.usf
@@ -1232,7 +1232,6 @@ void MainCommon(in FMaterialPixelParameters MaterialParameters, in float4 SvPosi
 	float MeanTransmittance = dot(TransmittanceToView.MATSUF, 1.0f / 3.0f);
 
 	float3 FogSampleTranslatedWorldPosition = RayTranslatedWorldOriginKm + tAP * Raydir;
-	FLWCVector3 FogSampleWorldPosition = LWCSubtract(LWCPromote(FogSampleTranslatedWorldPosition), PrimaryView.PreViewTranslation);
 	float4 ClipPos = mul(float4(FogSampleTranslatedWorldPosition, 1.0f), PrimaryView.TranslatedWorldToClip);
 
 	if (RenderVolumetricCloudParameters.EnableAerialPerspectiveSampling && tAPWeightsSum > 0.0f)
@@ -1241,7 +1240,7 @@ void MainCommon(in FMaterialPixelParameters MaterialParameters, in float4 SvPosi
 		// This allows to apply AP only once per pixel instead of per sample.
 		float4 AerialPerspective = GetAerialPerspectiveLuminanceTransmittance(
 			ResolvedView.RealTimeReflectionCapture, ResolvedView.SkyAtmosphereCameraAerialPerspectiveVolumeSizeAndInvSize,
-			ClipPos, LWCHackToFloat(FogSampleWorldPosition) * CM_TO_SKY_UNIT, LWCHackToFloat(ResolvedView.WorldCameraOrigin) * CM_TO_SKY_UNIT,
+			ClipPos, (FogSampleTranslatedWorldPosition - ResolvedView.TranslatedWorldCameraOrigin) * CM_TO_SKY_UNIT,
 			View.CameraAerialPerspectiveVolume, View.CameraAerialPerspectiveVolumeSampler,
 			ResolvedView.SkyAtmosphereCameraAerialPerspectiveVolumeDepthResolutionInv,
 			ResolvedView.SkyAtmosphereCameraAerialPerspectiveVolumeDepthResolution,
