Implement Nanite support for GetPerInstanceCustomData3Vector()

diff --git a/Engine/Shaders/Private/MaterialTemplate.ush b/Engine/Shaders/Private/MaterialTemplate.ush
--- a/Engine/Shaders/Private/MaterialTemplate.ush
+++ b/Engine/Shaders/Private/MaterialTemplate.ush
@@ -1071,6 +1071,19 @@ float GetPerInstanceCustomData(FMaterialVertexParameters Parameters, int Index,
 	return DefaultValue;
 }
 
+/** Per Instance Custom Data Getter (Pixel Shader Variant - Visibility Buffer) */
+/** Get the per-instance custom data when instancing */
+MaterialFloat3 GetPerInstanceCustomData3Vector(FMaterialPixelParameters Parameters, int Index, MaterialFloat3 DefaultValue)
+{
+#if IS_NANITE_PASS && USES_PER_INSTANCE_CUSTOM_DATA
+	return float3(GetPerInstanceCustomData(Parameters, Index + 0, DefaultValue.x),
+				  GetPerInstanceCustomData(Parameters, Index + 1, DefaultValue.y),
+				  GetPerInstanceCustomData(Parameters, Index + 2, DefaultValue.z));
+#endif
+
+	return DefaultValue;
+}
+
 // Per Instance Custom Data Getter (Vertex Shader Only)
 /** Get the per-instance custom data when instancing */
 MaterialFloat3 GetPerInstanceCustomData3Vector(FMaterialVertexParameters Parameters, int Index, MaterialFloat3 DefaultValue)
diff --git a/Engine/Source/Runtime/Engine/Private/Materials/HLSLMaterialTranslator.cpp b/Engine/Source/Runtime/Engine/Private/Materials/HLSLMaterialTranslator.cpp
--- a/Engine/Source/Runtime/Engine/Private/Materials/HLSLMaterialTranslator.cpp
+++ b/Engine/Source/Runtime/Engine/Private/Materials/HLSLMaterialTranslator.cpp
@@ -10985,7 +10985,10 @@ int32 FHLSLMaterialTranslator::PerInstanceCustomData3Vector(int32 DataIndex, int
 {
 	if (ShaderFrequency != SF_Vertex)
 	{
-		return NonVertexShaderExpressionError();
+		//return NonVertexShaderExpressionError();
+		// Only works with Nanite - TODO: Edge case the error
+		MaterialCompilationOutput.bUsesPerInstanceCustomData = true;
+		return AddInlinedCodeChunkZeroDeriv(MCT_Float3, TEXT("GetPerInstanceCustomData3Vector(Parameters, %d, %s)"), DataIndex, *GetParameterCode(DefaultValueIndex));
 	}
 	else
 	{
