Fix anisotropy tangent frame encoding.

diff --git a/Engine/Shaders/Private/Strata/Strata.ush b/Engine/Shaders/Private/Strata/Strata.ush
--- a/Engine/Shaders/Private/Strata/Strata.ush
+++ b/Engine/Shaders/Private/Strata/Strata.ush
@@ -2228,7 +2228,7 @@ uint StrataPackNormalAndTangent(in float3 Normal, in float3 Tangent)
 		{
 			Normal.xy = (1.0f - abs(Normal.yx)) * SignNotZero(Normal.xy);
 		}
-		//Normal = normalize(Normal);
+		Normal = normalize(Normal);
 
 		float3x3 LocalRef = GetTangentBasis(Normal);
 		float2 ProjXY = float2(dot(LocalRef[0], Tangent), dot(LocalRef[1], Tangent));
