Fix incorrect indexing when building mesh SDF culled tile data.
- Caused flickering and incorrect results in DFAO.
- Index NumCulledTilesArray and CulledTilesStartOffsetArray using packed culled instance index.

diff --git a/Engine/Shaders/Private/DistanceFieldObjectCulling.usf b/Engine/Shaders/Private/DistanceFieldObjectCulling.usf
--- a/Engine/Shaders/Private/DistanceFieldObjectCulling.usf
+++ b/Engine/Shaders/Private/DistanceFieldObjectCulling.usf
@@ -232,7 +232,7 @@ void BuildTileConesMain(
 struct FObjectCullVertexOutput
 {
 	nointerpolation float4 PositionAndRadius : TEXCOORD0;
-	nointerpolation uint ObjectIndex : TEXCOORD1;
+	nointerpolation uint2 ObjectIndexInstanceIndex : TEXCOORD1;
 };
  
 float ConservativeRadiusScale;
@@ -255,7 +255,7 @@ void ObjectCullVS(
 	OutPosition = mul(float4(WorldPosition, 1), LWCHackToFloat(PrimaryView.WorldToClip));
 	Output.PositionAndRadius.xyz = ObjectBounds.Center;
 	Output.PositionAndRadius.w = ObjectBounds.SphereRadius;
-	Output.ObjectIndex = ObjectIndex;
+	Output.ObjectIndexInstanceIndex = uint2(ObjectIndex, InstanceIndex);
 } 
 
 /** Used for object <-> tile culling */
@@ -344,27 +344,28 @@ void ObjectCullPS(
 		float ConeAxisDistance = dot(SphereCenterAndRadius.xyz - TileConeVertex, TileConeAxis);
 		float2 ConeAxisDistanceMinMax = float2(ConeAxisDistance + SphereCenterAndRadius.w, ConeAxisDistance - SphereCenterAndRadius.w);
 
-		uint TotalNumGroups = (uint)(NumGroups.x * NumGroups.y + .5f);
+		const uint ObjectIndex = Input.ObjectIndexInstanceIndex.x;
+		const uint InstanceIndex = Input.ObjectIndexInstanceIndex.y;
 
-		bool bTileIntersectsObject = IntersectObjectWithConeDepthRange(TileConeVertex, TileConeAxis, TileConeAngleCos, TileConeAngleSin, ConeAxisDepthRanges.xy, ConeAxisDistanceMinMax, Input.ObjectIndex);
+		bool bTileIntersectsObject = IntersectObjectWithConeDepthRange(TileConeVertex, TileConeAxis, TileConeAngleCos, TileConeAngleSin, ConeAxisDepthRanges.xy, ConeAxisDistanceMinMax, ObjectIndex);
 
 		if (!bTileIntersectsObject)
 		{
-			bTileIntersectsObject = IntersectObjectWithConeDepthRange(TileConeVertex, TileConeAxis, TileConeAngleCos, TileConeAngleSin, ConeAxisDepthRanges.zw, ConeAxisDistanceMinMax, Input.ObjectIndex);
+			bTileIntersectsObject = IntersectObjectWithConeDepthRange(TileConeVertex, TileConeAxis, TileConeAngleCos, TileConeAngleSin, ConeAxisDepthRanges.zw, ConeAxisDistanceMinMax, ObjectIndex);
 		}
 
 		if (bTileIntersectsObject)
 		{
 #if SCATTER_CULLING_COUNT_PASS
-			InterlockedAdd(RWNumCulledTilesArray[Input.ObjectIndex], 1);
+			InterlockedAdd(RWNumCulledTilesArray[InstanceIndex], 1);
 #else
 			uint CulledTileIndex;
-			InterlockedAdd(RWNumCulledTilesArray[Input.ObjectIndex], 1, CulledTileIndex);
+			InterlockedAdd(RWNumCulledTilesArray[InstanceIndex], 1, CulledTileIndex);
 
-			uint CulledTileDataStart = CulledTilesStartOffsetArray[Input.ObjectIndex];
+			uint CulledTileDataStart = CulledTilesStartOffsetArray[InstanceIndex];
 
 			RWCulledTileDataArray[(CulledTileDataStart + CulledTileIndex) * CULLED_TILE_DATA_STRIDE + 0] = TileIndex;
-			RWCulledTileDataArray[(CulledTileDataStart + CulledTileIndex) * CULLED_TILE_DATA_STRIDE + 1] = Input.ObjectIndex;
+			RWCulledTileDataArray[(CulledTileDataStart + CulledTileIndex) * CULLED_TILE_DATA_STRIDE + 1] = ObjectIndex;
 #endif
 		}
 	}
