Disables motionblur's CONFIG_SHUFFLE_HALF_RES

diff --git a/Engine/Shaders/Private/MotionBlur/MotionBlurApply.usf b/Engine/Shaders/Private/MotionBlur/MotionBlurApply.usf
--- a/Engine/Shaders/Private/MotionBlur/MotionBlurApply.usf
+++ b/Engine/Shaders/Private/MotionBlur/MotionBlurApply.usf
@@ -53,7 +53,7 @@
 #define SOFT_Z_EXTENT 1
 
 // Whether should interleave the half res motion blur with the full res to reduce the size of the grain on screen.
-#define CONFIG_SHUFFLE_HALF_RES 1
+#define CONFIG_SHUFFLE_HALF_RES 0
 
 // Whether should do post motion blur translucency.
 #define CONFIG_POST_MOTIONBLUR_TRANSLUCENCY 1
