Rename StrataLobeStatistic into FStrataLobeStatistic

diff --git a/Engine/Shaders/Private/DistortAccumulatePS.usf b/Engine/Shaders/Private/DistortAccumulatePS.usf
--- a/Engine/Shaders/Private/DistortAccumulatePS.usf
+++ b/Engine/Shaders/Private/DistortAccumulatePS.usf
@@ -110,7 +110,7 @@ void Main(
 
 		RefractionWorldNormal = StrataTree.Operators[StrataData.OperatorIndex].RefractionWorldNormal;
 
-		StrataLobeStatistic RefractionLobeStat = StrataTree.Operators[StrataData.OperatorIndex].RefractionLobeStat;
+		FStrataLobeStatistic RefractionLobeStat = StrataTree.Operators[StrataData.OperatorIndex].RefractionLobeStat;
 
 		Sigma = RefractionLobeStat.Sigma;
 	}
diff --git a/Engine/Shaders/Private/Strata/Strata.ush b/Engine/Shaders/Private/Strata/Strata.ush
--- a/Engine/Shaders/Private/Strata/Strata.ush
+++ b/Engine/Shaders/Private/Strata/Strata.ush
@@ -728,8 +728,8 @@ struct FStrataOperator
 	float3 VerticalTopTransmittance;	// Only used by vertical layering operators. It represent the transmittance of matter of the top layer (from the left sub tree).
 
 	// Rough refraction
-	float3				RefractionWorldNormal;
-	StrataLobeStatistic RefractionLobeStat;
+	float3				 RefractionWorldNormal;
+	FStrataLobeStatistic RefractionLobeStat;
 };
 
 FStrataOperator GetInitialisedStrataOperator()
diff --git a/Engine/Shaders/Private/Strata/StrataStatisticalOperators.ush b/Engine/Shaders/Private/Strata/StrataStatisticalOperators.ush
--- a/Engine/Shaders/Private/Strata/StrataStatisticalOperators.ush
+++ b/Engine/Shaders/Private/Strata/StrataStatisticalOperators.ush
@@ -7,7 +7,7 @@
 // Statistical operators representing BSDF as a sum of lobes.
 // This implementation is based on [Belcour 2018, "Efficient Rendering of Layered Materials using an Atomic Decomposition with Statistical Operators"]
 
-struct StrataLobeStatistic
+struct FStrataLobeStatistic
 {
 	// Mean
 	// xy is the 2d projection of the lobe main direction onto the plane defined by the surface normal. 
@@ -43,24 +43,24 @@ float StrataLobeVarianceToRoughness(float Variance)
 //  - OpticalDepth = the value impact light transmission (= sigma_t * depth)
 //  - Approximation to in-scattering from the back of the layer (eq.19)
 
-StrataLobeStatistic StrataGetNullLobe()
+FStrataLobeStatistic StrataGetNullLobe()
 {
-	StrataLobeStatistic NullLobe = (StrataLobeStatistic)0;
+	FStrataLobeStatistic NullLobe = (FStrataLobeStatistic)0;
 	return NullLobe;
 }
 
-StrataLobeStatistic StrataGetViewDiracLobe(float3 Wi)
+FStrataLobeStatistic StrataGetViewDiracLobe(float3 Wi)
 {
-	StrataLobeStatistic WiLobe;
+	FStrataLobeStatistic WiLobe;
 	WiLobe.E = 1.0f;
 	WiLobe.Mu = Wi.xy;
 	WiLobe.Sigma = 0.0f;
 	return WiLobe;
 }
 
-StrataLobeStatistic StrataGetReflectedLobe(StrataLobeStatistic WiLobe, float3 InterfaceFDG, float InterfaceRoughness)
+FStrataLobeStatistic StrataGetReflectedLobe(FStrataLobeStatistic WiLobe, float3 InterfaceFDG, float InterfaceRoughness)
 {
-	StrataLobeStatistic WoLobe;
+	FStrataLobeStatistic WoLobe;
 
 	WoLobe.E = WiLobe.E * InterfaceFDG;
 
@@ -71,9 +71,9 @@ StrataLobeStatistic StrataGetReflectedLobe(StrataLobeStatistic WiLobe, float3 In
 	return WoLobe;
 }
 
-StrataLobeStatistic StrataGetRefractedLobe(StrataLobeStatistic WiLobe, float3 InterfaceFDG, float InterfaceRoughness, float InterfaceEta12)
+FStrataLobeStatistic StrataGetRefractedLobe(FStrataLobeStatistic WiLobe, float3 InterfaceFDG, float InterfaceRoughness, float InterfaceEta12)
 {
-	StrataLobeStatistic WoLobe;
+	FStrataLobeStatistic WoLobe;
 
 	WoLobe.E = WiLobe.E * (1.0f - InterfaceFDG);
 
@@ -87,9 +87,9 @@ StrataLobeStatistic StrataGetRefractedLobe(StrataLobeStatistic WiLobe, float3 In
 	return WoLobe;
 }
 
-StrataLobeStatistic StrataGetTransmittedLobe(StrataLobeStatistic WiLobe, float3 OpticalDepth)
+FStrataLobeStatistic StrataGetTransmittedLobe(FStrataLobeStatistic WiLobe, float3 OpticalDepth)
 {
-	StrataLobeStatistic WoLobe;
+	FStrataLobeStatistic WoLobe;
 
 	WoLobe.E = WiLobe.E * exp(-OpticalDepth);
 
@@ -101,22 +101,22 @@ StrataLobeStatistic StrataGetTransmittedLobe(StrataLobeStatistic WiLobe, float3
 	return WoLobe;
 }
 
-StrataLobeStatistic StrataWeightLobe(StrataLobeStatistic A, float Weight)
+FStrataLobeStatistic StrataWeightLobe(FStrataLobeStatistic A, float Weight)
 {
 	// We simple weight the lob properties so that they progressively disapear as a function of coverage.
 	// This is not entirely correct though because if matter coverage is reduced there should be two lobe: 
 	//	- this lobe with visibility = coverage
 	//	- a dirac lobe with visibility = 1 - coverage
-	StrataLobeStatistic WoLobe = A;
+	FStrataLobeStatistic WoLobe = A;
 	WoLobe.E *= Weight;
 	WoLobe.Mu *= Weight;
 	WoLobe.Sigma *= Weight;
 	return WoLobe;
 }
 
-StrataLobeStatistic StrataHorizontalMixLobes(StrataLobeStatistic A, StrataLobeStatistic B, float Mix)
+FStrataLobeStatistic StrataHorizontalMixLobes(FStrataLobeStatistic A, FStrataLobeStatistic B, float Mix)
 {
-	StrataLobeStatistic WoLobe;
+	FStrataLobeStatistic WoLobe;
 	WoLobe.E = lerp(A.E, B.E, Mix);
 	WoLobe.Mu = lerp(A.Mu, B.Mu, Mix);
 	WoLobe.Sigma = lerp(A.Sigma, B.Sigma, Mix);
diff --git a/Engine/Shaders/Private/Strata/StrataTree.ush b/Engine/Shaders/Private/Strata/StrataTree.ush
--- a/Engine/Shaders/Private/Strata/StrataTree.ush
+++ b/Engine/Shaders/Private/Strata/StrataTree.ush
@@ -80,7 +80,7 @@ void UpdateSingleBSDFOperatorCoverageTransmittance(
 
 	// We multiply this weight with coverage:
 	//  - Because when coming from parameter blending, we need to scale the output down according to coverage resulting from parameter blending.
-	//  - That will be applied on emissive and reflected light, so we do not ned to apply that anywhere else.
+	//  - That will be applied on emissive and reflected light, so we do not need to apply that anywhere else.
 	//  - When coming from non parameter blending, this is equal identity 1 and the operator parsing will continue to correctly weight the BSDF according to the tree.
 	CurrentBSDF.LuminanceWeight = float3(1.0f, 1.0f, 1.0f) * CurrentBSDFCoverage;
 
@@ -94,7 +94,7 @@ void UpdateSingleBSDFOperatorCoverageTransmittance(
 	const float Roughness = StrataGetBSDFRoughness(CurrentBSDF);
 	const FBxDFEnergyTerms EnergyTerms = ComputeGGXSpecEnergyTerms(Roughness, StrataBSDFContext.SatNoV, StrataGetBSDFSpecularColor(CurrentBSDF), StrataGetBSDFSpecularF90(CurrentBSDF));
 	const float InterfaceEta12 = CurrentBSDF.bIsTop ? 1.0f / 1.33f : 1.0f; // STRATA_TODO Use Refraction.x in IOR mode? Or even the material IOR corresponding to F0?
-	StrataLobeStatistic RefractedLobe = StrataGetRefractedLobe(StrataGetViewDiracLobe(-V), EnergyTerms.E, Roughness, InterfaceEta12);
+	FStrataLobeStatistic RefractedLobe = StrataGetRefractedLobe(StrataGetViewDiracLobe(-V), EnergyTerms.E, Roughness, InterfaceEta12);
 	// Account for transmittance of the medium (but ignoring scattering for now)
 	RefractedLobe.E *= CurrentBSDFTransmittance;
 	// Account for coverate
