Fix missing lumen support when using default hair integrator.

diff --git a/Engine/Shaders/Private/HairStrands/HairStrandsEnvironmentLighting.usf b/Engine/Shaders/Private/HairStrands/HairStrandsEnvironmentLighting.usf
--- a/Engine/Shaders/Private/HairStrands/HairStrandsEnvironmentLighting.usf
+++ b/Engine/Shaders/Private/HairStrands/HairStrandsEnvironmentLighting.usf
@@ -712,16 +712,43 @@ float3 ReflectionEnvironmentSH(
 	uint MacroGroupId,
 	bool bDebugEnabled)
 {
-	const float AmbientOcclusion = 1.0f;
-	const float3 BentNormal = GBuffer.WorldNormal;
-	const float3 DiffuseColor = GBuffer.BaseColor; // not used
-	const float2 BufferUV = SvPosition * View.BufferSizeAndInvSize.zw; // not used
-
 	const float3 TranslatedWorldPosition = mul(float4(ScreenPosition * GBuffer.Depth, GBuffer.Depth, 1), PrimaryView.ScreenToTranslatedWorld).xyz;
 	const float3 AbsoluteWorldPosition = TranslatedWorldPosition + LWCHackToFloat(PrimaryView.WorldCameraOrigin);
 	const float3 CameraToPixel = normalize(TranslatedWorldPosition - PrimaryView.TranslatedWorldCameraOrigin);
 	const float3 V = -CameraToPixel;
 
+#if PERMUTATION_LUMEN
+	const float3 T = GBuffer.WorldNormal;
+	const float3 N = T;
+	const float3 L = normalize(V - N * dot(V, N));
+
+	FHairTransmittanceData TransmittanceData = GetHairTransmittance(
+		V,
+		L,
+		GBuffer,
+		InitHairTransmittanceMask(),
+		View.HairScatteringLUTTexture,
+		View.HairScatteringLUTSampler,
+		View.HairComponents);
+
+	bool bEvalMultiScatter = true;
+	float3 f_s = HairShading(GBuffer, L, V, N, 1, TransmittanceData, 1, 0.2, uint2(0, 0));
+
+	// This logic mimic the ReflecitonEnvironemtPixelShader.usf. Ideallyw we should share the same
+	// code/path, but the various permutation make the code a bit obscure
+	// Need to take into account DFAO?
+	float3 Lighting = 0;
+	{
+		// Probe lighting
+		const uint2  PixelPosition = SvPosition - ResolvedView.ViewRectMin.xy;
+		Lighting = View.PreExposure * GatherRadiance(L, GBuffer.Roughness, PixelPosition, GBuffer.Depth, TranslatedWorldPosition) * f_s;
+	}
+#else // PERMUTATION_LUMEN
+	const float AmbientOcclusion = 1.0f;
+	const float3 BentNormal = GBuffer.WorldNormal;
+	const float3 DiffuseColor = GBuffer.BaseColor; // not used
+	const float2 BufferUV = SvPosition * View.BufferSizeAndInvSize.zw; // not used
+
 	float3 Lighting = 0;
 
 	// Sky lighting (already pre-exposed)
@@ -729,6 +756,7 @@ float3 ReflectionEnvironmentSH(
 
 	// Static lighting (already pre-exposed)
 	Lighting += GetStaticLighting(GBuffer, V, AbsoluteWorldPosition);
+#endif // PERMUTATION_LUMEN
 
 	// Transform NaNs to black, transform negative colors to black.
 	return -min(-Lighting, 0.0);
