Strata - Propagate the top layer roughness to the bottom layer to simulate the fact that light rays are potentially heavily refracted before reaching the bottom layer.
Far field only for now and ior ignored.

diff --git a/Engine/Shaders/Private/Strata/Strata.ush b/Engine/Shaders/Private/Strata/Strata.ush
--- a/Engine/Shaders/Private/Strata/Strata.ush
+++ b/Engine/Shaders/Private/Strata/Strata.ush
@@ -728,8 +728,14 @@ struct FStrataOperator
 	float3 VerticalTopTransmittance;	// Only used by vertical layering operators. It represent the transmittance of matter of the top layer (from the left sub tree).
 
 	// Rough refraction
-	float3				 RefractionWorldNormal;
-	FStrataLobeStatistic RefractionLobeStat;
+	float3				 RefractionWorldNormal;			// The normal used for the single refracted direction
+	FStrataLobeStatistic RefractionLobeStat;			// The lobe used for the scene rough refraction
+
+	// Thickness
+	float				 ThicknessCm;					// The thickness of the layer
+
+	float				 VerticalTopThicknessCm;		// Only used by vertical layering operators. It represent the thickness that is used to drive the roughness propagation from top layers to sub layers.
+	FStrataLobeStatistic VerticalTopRefractionLobeStat;	// Only used by vertical layering operators. It represent the lobe statistics of the top layer at its interface.
 };
 
 FStrataOperator GetInitialisedStrataOperator()
@@ -1971,6 +1977,16 @@ float StrataGetBSDFRoughness(in FStrataBSDF BSDF)
 	return 0.f;
 }
 
+void StrataSetBSDFRoughness(inout FStrataBSDF BSDF, in float Roughness)
+{
+	switch (BSDF_GETTYPE(BSDF))
+	{
+	case STRATA_BSDF_TYPE_SLAB: 			SLAB_ROUGHNESS(BSDF) = Roughness; break;
+	case STRATA_BSDF_TYPE_HAIR: 			HAIR_ROUGHNESS(BSDF) = Roughness; break;
+	case STRATA_BSDF_TYPE_SINGLELAYERWATER: SLW_ROUGHNESS(BSDF)  = Roughness; break;
+	}
+}
+
 float StrataGetBSDFAnisotropy(in FStrataBSDF BSDF)
 {
 	switch (BSDF_GETTYPE(BSDF))
diff --git a/Engine/Shaders/Private/Strata/StrataTree.ush b/Engine/Shaders/Private/Strata/StrataTree.ush
--- a/Engine/Shaders/Private/Strata/StrataTree.ush
+++ b/Engine/Shaders/Private/Strata/StrataTree.ush
@@ -69,6 +69,9 @@ void UpdateSingleBSDFOperatorCoverageTransmittance(
 {
 #if STRATA_INLINE_SHADING
 
+	// Extract the true thickness of the BSDF before it is cast into the normalised Slab thickness we used.
+	StrataTree.Operators[CurrentBSDF.OperatorIndex].ThicknessCm = CurrentBSDF.ThicknessCm;
+
 	PostProcessBSDFBeforeLighting(CurrentBSDF);
 
 	// Compute current BSDF transmittance
@@ -127,6 +130,9 @@ void UpdateSingleOperatorCoverageTransmittance(inout FStrataTree StrataTree, int
 
 		Op.RefractionLobeStat = StrataWeightLobe(OpA.RefractionLobeStat, Weight);
 		Op.RefractionWorldNormal = OpA.RefractionWorldNormal * Weight;
+
+		// Coverage should not affect thickness
+		Op.ThicknessCm = OpA.ThicknessCm;
 		break;
 	}
 
@@ -168,6 +174,11 @@ void UpdateSingleOperatorCoverageTransmittance(inout FStrataTree StrataTree, int
 		// The ratio of bottomSurfaceVisbile and visible through top layer (with a scale of 0.5 because at maximum both surfaces will be visible) as a fonction of the total coverage.
 		const float RefractionNormalMix = saturate((Info.SurfaceBottom + 0.5f * Info.TransmittanceTopAndBottom * dot(TopTransmittance, 0.33.xxx)) / max(STRATA_EPSILON, Info.Coverage));
 		Op.RefractionWorldNormal = lerp(OpA.RefractionWorldNormal, OpB.RefractionWorldNormal, RefractionNormalMix);
+
+		// Vertical layering accumulates both operators thickness
+		Op.ThicknessCm = OpA.ThicknessCm + OpB.ThicknessCm;
+
+		Op.VerticalTopRefractionLobeStat = OpA.RefractionLobeStat;
 		break;
 	}
 
@@ -182,6 +193,9 @@ void UpdateSingleOperatorCoverageTransmittance(inout FStrataTree StrataTree, int
 
 		Op.RefractionLobeStat = StrataHorizontalMixLobes(OpA.RefractionLobeStat, OpB.RefractionLobeStat, BMix);
 		Op.RefractionWorldNormal = lerp(OpA.RefractionWorldNormal, OpB.RefractionWorldNormal, BMix);
+
+		// Horizontal mixes both operators thickness
+		Op.ThicknessCm = lerp(OpA.ThicknessCm, OpB.ThicknessCm, BMix);
 		break;
 	}
 
@@ -196,6 +210,8 @@ void UpdateSingleOperatorCoverageTransmittance(inout FStrataTree StrataTree, int
 
 		Op.RefractionLobeStat = StrataHorizontalMixLobes(OpA.RefractionLobeStat, OpB.RefractionLobeStat, 0.5f);
 		Op.RefractionWorldNormal = lerp(OpA.RefractionWorldNormal, OpB.RefractionWorldNormal, 0.5f);
+
+		Op.ThicknessCm = lerp(OpA.ThicknessCm, OpB.ThicknessCm, 0.5f);
 		break;
 	}
 
@@ -266,6 +282,19 @@ void UpdateAllBSDFWithBottomUpOperatorVisit_Vertical(
 	// If the BSDF is from the bottom, we reduce its contribution according to the top layer coverage only.
 	CurrentBSDF.TopLayerDataWeight *= PreviousIsInputA ? 1.0 : saturate(1.0f - Op.VerticalTopCoverage);
 
+	// We propagate the top layer roughness to the bottom layer to simulate the fact that light rays are potentially heavily refracted before reaching the bottom layer.
+	// This is done according to roughness.
+	// STRATA_TODO: handle InterfaceEta12 (according ot bIsTop) and Thickness (always assuming full far-field view for now).
+	const float3 DummyWi = float3(0, 0, 1);
+	const float DummyInterfaceDFG = 0.5f;
+	const float InterfaceEta12 = 1.0f;
+	FStrataLobeStatistic ReflectedLobe = StrataGetReflectedLobe(StrataGetViewDiracLobe(DummyWi), DummyInterfaceDFG, StrataGetBSDFRoughness(CurrentBSDF));
+	const float RoughnessAsVariance = Op.VerticalTopRefractionLobeStat.Sigma * Op.Coverage; // Simple weight of the top layer influence according to coverage. An approximation because we only track a single lob.
+	const float TopLayerSigma = StrataLobeVarianceToRoughness(RoughnessAsVariance);
+	FStrataLobeStatistic RefractedLobe = StrataGetRefractedLobe(ReflectedLobe, DummyInterfaceDFG, TopLayerSigma, DummyInterfaceDFG);
+	RefractedLobe = StrataWeightLobe(RefractedLobe, Op.Coverage);	// We should not weight with coverage but that is a good way to to handle coverage while only handling a single lob.
+	StrataSetBSDFRoughness(CurrentBSDF, StrataLobeVarianceToRoughness(Op.VerticalTopRefractionLobeStat.Sigma));
+
 #endif // STRATA_INLINE_SHADING
 }
 
