Strata opaque rough refraction - fixed self rough refraction.

diff --git a/Engine/Shaders/Private/Strata/StrataTree.ush b/Engine/Shaders/Private/Strata/StrataTree.ush
--- a/Engine/Shaders/Private/Strata/StrataTree.ush
+++ b/Engine/Shaders/Private/Strata/StrataTree.ush
@@ -282,18 +282,21 @@ void UpdateAllBSDFWithBottomUpOperatorVisit_Vertical(
 	// If the BSDF is from the bottom, we reduce its contribution according to the top layer coverage only.
 	CurrentBSDF.TopLayerDataWeight *= PreviousIsInputA ? 1.0 : saturate(1.0f - Op.VerticalTopCoverage);
 
-	// We propagate the top layer roughness to the bottom layer to simulate the fact that light rays are potentially heavily refracted before reaching the bottom layer.
-	// This is done according to roughness.
-	// STRATA_TODO: handle InterfaceEta12 (according ot bIsTop) and Thickness (always assuming full far-field view for now).
-	const float3 DummyWi = float3(0, 0, 1);
-	const float DummyInterfaceDFG = 0.5f;
-	const float InterfaceEta12 = 1.0f;
-	FStrataLobeStatistic ReflectedLobe = StrataGetReflectedLobe(StrataGetViewDiracLobe(DummyWi), DummyInterfaceDFG, StrataGetBSDFRoughness(CurrentBSDF));
-	const float RoughnessAsVariance = Op.VerticalTopRefractionLobeStat.Sigma * Op.Coverage; // Simple weight of the top layer influence according to coverage. An approximation because we only track a single lob.
-	const float TopLayerSigma = StrataLobeVarianceToRoughness(RoughnessAsVariance);
-	FStrataLobeStatistic RefractedLobe = StrataGetRefractedLobe(ReflectedLobe, DummyInterfaceDFG, TopLayerSigma, DummyInterfaceDFG);
-	RefractedLobe = StrataWeightLobe(RefractedLobe, Op.Coverage);	// We should not weight with coverage but that is a good way to to handle coverage while only handling a single lob.
-	StrataSetBSDFRoughness(CurrentBSDF, StrataLobeVarianceToRoughness(Op.VerticalTopRefractionLobeStat.Sigma));
+	if (PreviousIsInputA <= 0)	// Only apply to BSDF from the bottom layer
+	{
+		// We propagate the top layer roughness to the bottom layer to simulate the fact that light rays are potentially heavily refracted before reaching the bottom layer.
+		// This is done according to roughness.
+		// STRATA_TODO: handle InterfaceEta12 (according ot bIsTop) and Thickness (always assuming full far-field view for now).
+		const float3 DummyWi = float3(0, 0, 1);
+		const float DummyInterfaceDFG = 0.5f;
+		const float InterfaceEta12 = 1.0f;
+		FStrataLobeStatistic ReflectedLobe = StrataGetReflectedLobe(StrataGetViewDiracLobe(DummyWi), DummyInterfaceDFG, StrataGetBSDFRoughness(CurrentBSDF));
+		const float RoughnessAsVariance = Op.VerticalTopRefractionLobeStat.Sigma * Op.Coverage; // Simple weight of the top layer influence according to coverage. An approximation because we only track a single lob.
+		const float TopLayerSigma = StrataLobeVarianceToRoughness(RoughnessAsVariance);
+		FStrataLobeStatistic RefractedLobe = StrataGetRefractedLobe(ReflectedLobe, DummyInterfaceDFG, TopLayerSigma, DummyInterfaceDFG);
+		RefractedLobe = StrataWeightLobe(RefractedLobe, Op.Coverage);	// We should not weight with coverage but that is a good way to to handle coverage while only handling a single lob.
+		StrataSetBSDFRoughness(CurrentBSDF, StrataLobeVarianceToRoughness(Op.VerticalTopRefractionLobeStat.Sigma));
+	}
 
 #endif // STRATA_INLINE_SHADING
 }
