Strata - clarified opaque roughtrefraction, IOR useage and fixed a fiew issues.

diff --git a/Engine/Shaders/Private/Strata/Strata.ush b/Engine/Shaders/Private/Strata/Strata.ush
--- a/Engine/Shaders/Private/Strata/Strata.ush
+++ b/Engine/Shaders/Private/Strata/Strata.ush
@@ -733,8 +733,11 @@ struct FStrataOperator
 
 	// Thickness
 	float				 ThicknessCm;					// The thickness of the layer
-
 	float				 VerticalTopThicknessCm;		// Only used by vertical layering operators. It represent the thickness that is used to drive the roughness propagation from top layers to sub layers.
+
+	// Topology information
+	bool				 bIsTop;
+
 	FStrataLobeStatistic VerticalTopRefractionLobeStat;	// Only used by vertical layering operators. It represent the lobe statistics of the top layer at its interface.
 };
 
diff --git a/Engine/Shaders/Private/Strata/StrataTree.ush b/Engine/Shaders/Private/Strata/StrataTree.ush
--- a/Engine/Shaders/Private/Strata/StrataTree.ush
+++ b/Engine/Shaders/Private/Strata/StrataTree.ush
@@ -6,6 +6,9 @@
 #define STRATA_OPAQUE_MATERIAL 0
 #endif
 
+// STRATA_TODO We assume layers all have the same IOR. But we should take into account the material later.
+#define WATER_IOR		1.33f
+#define AIR_IOR			1.0f
 
 
 ///////////////////////////////////////////////////////////////////////////////
@@ -69,8 +72,10 @@ void UpdateSingleBSDFOperatorCoverageTransmittance(
 {
 #if STRATA_INLINE_SHADING
 
+#define Op StrataTree.Operators[CurrentBSDF.OperatorIndex]
+
 	// Extract the true thickness of the BSDF before it is cast into the normalised Slab thickness we used.
-	StrataTree.Operators[CurrentBSDF.OperatorIndex].ThicknessCm = CurrentBSDF.ThicknessCm;
+	Op.ThicknessCm = CurrentBSDF.ThicknessCm;
 
 	PostProcessBSDFBeforeLighting(CurrentBSDF);
 
@@ -90,22 +95,26 @@ void UpdateSingleBSDFOperatorCoverageTransmittance(
 	// All BSDF are considered for contribution as the top BSDF can have their coverage dip to 0 and then another BSDF need to be considered for the top layer representation.
 	CurrentBSDF.TopLayerDataWeight = 1.0f;
 
-	StrataTree.Operators[CurrentBSDF.OperatorIndex].Coverage = CurrentBSDFCoverage;
-	StrataTree.Operators[CurrentBSDF.OperatorIndex].Transmittance = CurrentBSDFTransmittance;
+	Op.Coverage = CurrentBSDFCoverage;
+	Op.Transmittance = CurrentBSDFTransmittance;
 
 	// The top layer we consider a transition air to water. For layers below, we assume they all have similar index of refraction.
 	const float Roughness = StrataGetBSDFRoughness(CurrentBSDF);
 	const FBxDFEnergyTerms EnergyTerms = ComputeGGXSpecEnergyTerms(Roughness, StrataBSDFContext.SatNoV, StrataGetBSDFSpecularColor(CurrentBSDF), StrataGetBSDFSpecularF90(CurrentBSDF));
-	const float InterfaceEta12 = CurrentBSDF.bIsTop ? 1.0f / 1.33f : 1.0f; // STRATA_TODO Use Refraction.x in IOR mode? Or even the material IOR corresponding to F0?
-	FStrataLobeStatistic RefractedLobe = StrataGetRefractedLobe(StrataGetViewDiracLobe(-V), EnergyTerms.E, Roughness, InterfaceEta12);
+	const float TopDownInterfaceEta = CurrentBSDF.bIsTop ? AIR_IOR / WATER_IOR : 1.0f;
+	FStrataLobeStatistic RefractedLobe = StrataGetRefractedLobe(StrataGetViewDiracLobe(-V), EnergyTerms.E, Roughness, TopDownInterfaceEta);
 	// Account for transmittance of the medium (but ignoring scattering for now)
 	RefractedLobe.E *= CurrentBSDFTransmittance;
 	// Account for coverate
 	RefractedLobe = StrataWeightLobe(RefractedLobe, CurrentBSDFCoverage);
-	StrataTree.Operators[CurrentBSDF.OperatorIndex].RefractionLobeStat = RefractedLobe;
+	Op.RefractionLobeStat = RefractedLobe;
 	
 	float3x3 TangentBasis = StrataGetBSDFSharedBasis_InlineShading(StrataPixelHeader, CurrentBSDF, NullStrataAddressing);
-	StrataTree.Operators[CurrentBSDF.OperatorIndex].RefractionWorldNormal = TangentBasis[2];
+	Op.RefractionWorldNormal = TangentBasis[2];
+
+	Op.bIsTop = CurrentBSDF.bIsTop;
+
+#undef Op
 
 #endif // STRATA_INLINE_SHADING
 }
@@ -119,6 +128,12 @@ void UpdateSingleOperatorCoverageTransmittance(inout FStrataTree StrataTree, int
 #define OpA StrataTree.Operators[StrataTree.Operators[OpIndex].LeftIndex]
 #define OpB StrataTree.Operators[StrataTree.Operators[OpIndex].RightIndex]
 
+	// Propagate bIsTop:
+	//  - horizontal, add, weight operations: Any of the operand bIsTop should be enough to determine if this operator is at the top of the material
+	//	- vertical: must take the top layer information (operand A)
+	// => we always take the operand A information.
+	Op.bIsTop = OpA.bIsTop;
+
 	switch (Op.Type)
 	{
 
@@ -166,12 +181,12 @@ void UpdateSingleOperatorCoverageTransmittance(inout FStrataTree StrataTree, int
 		Op.VerticalTopCoverage		= OpA.Coverage;
 		Op.VerticalTopTransmittance = OpA.Transmittance;
 
-		// Project top lob into bottom layer
-		const float InterfaceEta12 = 1.0f; // STRATA_TODO We assume layers all have the same IOR. But we should take into account the material later.
-		Op.RefractionLobeStat = StrataGetRefractedLobe(OpA.RefractionLobeStat, /*InterfaceFDG*/0.0f, StrataLobeVarianceToRoughness(OpB.RefractionLobeStat.Sigma), InterfaceEta12);
+		// Project top lob through bottom layer assuming medium have same IOR.
+		const float TopDownInterfaceEta = 1.0f;
+		Op.RefractionLobeStat = StrataGetRefractedLobe(OpA.RefractionLobeStat, /*InterfaceFDG*/0.0f, StrataLobeVarianceToRoughness(OpB.RefractionLobeStat.Sigma), TopDownInterfaceEta);
 		Op.RefractionLobeStat.E *= OpB.RefractionLobeStat.E * saturate(Op.VerticalTopTransmittance + (1.0f - Op.VerticalTopCoverage)); // Combine with botton layer medium and specular transmittance
 
-		// The ratio of bottomSurfaceVisbile and visible through top layer (with a scale of 0.5 because at maximum both surfaces will be visible) as a fonction of the total coverage.
+		// The ratio of BottomSurfaceOnlyVisibile and visible through top layer (with a scale of 0.5 because at maximum both surfaces will be visible) as a fonction of the total coverage.
 		const float RefractionNormalMix = saturate((Info.SurfaceBottom + 0.5f * Info.TransmittanceTopAndBottom * dot(TopTransmittance, 0.33.xxx)) / max(STRATA_EPSILON, Info.Coverage));
 		Op.RefractionWorldNormal = lerp(OpA.RefractionWorldNormal, OpB.RefractionWorldNormal, RefractionNormalMix);
 
@@ -284,18 +299,23 @@ void UpdateAllBSDFWithBottomUpOperatorVisit_Vertical(
 
 	if (PreviousIsInputA <= 0)	// Only apply to BSDF from the bottom layer
 	{
-		// We propagate the top layer roughness to the bottom layer to simulate the fact that light rays are potentially heavily refracted before reaching the bottom layer.
-		// This is done according to roughness.
-		// STRATA_TODO: handle InterfaceEta12 (according ot bIsTop) and Thickness (always assuming full far-field view for now).
+		// We propagate the bottom layer lob up after the top layer top layer simulate the fact that light rays are potentially heavily scattered/refracted by the top layer.
+		// The outcome is that we affect the bottom layer roughness to aggregate that fact.
+		// STRATA_TODO: take into account thickness (far-field for now).
+
+		// Evaluate the reflected lobe
 		const float3 DummyWi = float3(0, 0, 1);
 		const float DummyInterfaceDFG = 0.5f;
-		const float InterfaceEta12 = 1.0f;
 		FStrataLobeStatistic ReflectedLobe = StrataGetReflectedLobe(StrataGetViewDiracLobe(DummyWi), DummyInterfaceDFG, StrataGetBSDFRoughness(CurrentBSDF));
-		const float RoughnessAsVariance = Op.VerticalTopRefractionLobeStat.Sigma * Op.Coverage; // Simple weight of the top layer influence according to coverage. An approximation because we only track a single lob.
-		const float TopLayerSigma = StrataLobeVarianceToRoughness(RoughnessAsVariance);
-		FStrataLobeStatistic RefractedLobe = StrataGetRefractedLobe(ReflectedLobe, DummyInterfaceDFG, TopLayerSigma, DummyInterfaceDFG);
-		RefractedLobe = StrataWeightLobe(RefractedLobe, Op.Coverage);	// We should not weight with coverage but that is a good way to to handle coverage while only handling a single lob.
-		StrataSetBSDFRoughness(CurrentBSDF, StrataLobeVarianceToRoughness(Op.VerticalTopRefractionLobeStat.Sigma));
+
+		// Evaluate the reflected lobe after refraction through the top layer of the vertical operator.
+		// => This tells how light should be integrated when traveling up thorugh that layer.
+		const float TopLayerSigma = Op.VerticalTopRefractionLobeStat.Sigma;				// STRATA_TODO this sigma was evaluate as a top down lobe. We should evaluate it as a bottom up sigma for opaque rough refraction because this is refraction of the incoming light.
+		const float TopRoughnessAsVariance = TopLayerSigma * Op.VerticalTopCoverage;	// Weighting variance by coverage is a good way to handle rough refraction reduction according when relying on a single lob.
+		const float BottomUpInterfaceEta = Op.bIsTop ? WATER_IOR / AIR_IOR : 1.0f;
+		FStrataLobeStatistic RefractedLobe = StrataGetRefractedLobe(ReflectedLobe, DummyInterfaceDFG, StrataLobeVarianceToRoughness(TopRoughnessAsVariance), BottomUpInterfaceEta);
+
+		StrataSetBSDFRoughness(CurrentBSDF, StrataLobeVarianceToRoughness(RefractedLobe.Sigma));
 	}
 
 #endif // STRATA_INLINE_SHADING
