Support Capsule shadows in LWC
- Upload capsule and light position in translated world space.
- Perform culling and cone tracing in translated world space.

diff --git a/Engine/Shaders/Private/CapsuleShadowShaders.usf b/Engine/Shaders/Private/CapsuleShadowShaders.usf
--- a/Engine/Shaders/Private/CapsuleShadowShaders.usf
+++ b/Engine/Shaders/Private/CapsuleShadowShaders.usf
@@ -32,7 +32,7 @@
 
 struct FCapsuleShape
 {
-	float3 Center;
+	float3 TranslatedCenter;
 	float Radius;
 	float3 Orientation;
 	float Length;
@@ -201,7 +201,7 @@ bool SphereIntersectSphere(float4 SphereCenterAndRadius, float4 OtherSphereCente
 #if LIGHT_SOURCE_MODE == LIGHT_SOURCE_PUNCTUAL
 /** From point being shaded toward light, for directional lights. */
 float3 LightDirection;
-float4 LightPositionAndInvRadius;
+float4 LightTranslatedPositionAndInvRadius;
 float LightSourceRadius;
 float RayStartOffsetDepthScale;
 float3 LightAngleAndNormalThreshold;
@@ -220,7 +220,8 @@ groupshared uint IntegerTileMaxZ2;
 
 struct FTileCullingData
 {
-	float4 BoundingSphere;
+	float4 TranslatedBoundingSphere;
+	float4 BoundingSphere; // LWC_TODO
 #if LIGHT_SOURCE_MODE == LIGHT_SOURCE_PUNCTUAL
 	float3 ConeAxis;
 	float ConeAngleCos;
@@ -367,18 +368,22 @@ void SetupTileCullingData(
 	}
 
 	float3 ViewGroup0Center = (ViewTileMax + ViewTileMin) / 2;
-	TileCullingData0.BoundingSphere.xyz = mul(float4(ViewGroup0Center, 1), View.ViewToTranslatedWorld).xyz - LWCHackToFloat(PrimaryView.PreViewTranslation);
+	TileCullingData0.TranslatedBoundingSphere.xyz = mul(float4(ViewGroup0Center, 1), View.ViewToTranslatedWorld).xyz;
+	TileCullingData0.TranslatedBoundingSphere.w = length(ViewGroup0Center - ViewTileMax);
+	TileCullingData0.BoundingSphere.xyz = TileCullingData0.TranslatedBoundingSphere.xyz - LWCHackToFloat(PrimaryView.PreViewTranslation);
 	TileCullingData0.BoundingSphere.w = length(ViewGroup0Center - ViewTileMax);
 
 	float3 ViewGroup1Center = (ViewTileMax2 + ViewTileMin2) / 2;
-	TileCullingData1.BoundingSphere.xyz = mul(float4(ViewGroup1Center, 1), View.ViewToTranslatedWorld).xyz - LWCHackToFloat(PrimaryView.PreViewTranslation);
+	TileCullingData1.TranslatedBoundingSphere.xyz = mul(float4(ViewGroup1Center, 1), View.ViewToTranslatedWorld).xyz;
+	TileCullingData1.TranslatedBoundingSphere.w = length(ViewGroup1Center - ViewTileMax);
+	TileCullingData1.BoundingSphere.xyz = TileCullingData1.TranslatedBoundingSphere.xyz - LWCHackToFloat(PrimaryView.PreViewTranslation);
 	TileCullingData1.BoundingSphere.w = length(ViewGroup1Center - ViewTileMax2);
 
 #if LIGHT_SOURCE_MODE == LIGHT_SOURCE_PUNCTUAL
 	#if POINT_LIGHT
-		float3 LightVector0 = LightPositionAndInvRadius.xyz - TileCullingData0.BoundingSphere.xyz;
+		float3 LightVector0 = LightTranslatedPositionAndInvRadius.xyz - TileCullingData0.TranslatedBoundingSphere.xyz;
 		float LightVector0Length = length(LightVector0);
-		float3 LightVector1 = LightPositionAndInvRadius.xyz - TileCullingData1.BoundingSphere.xyz;
+		float3 LightVector1 = LightTranslatedPositionAndInvRadius.xyz - TileCullingData1.TranslatedBoundingSphere.xyz;
 		float LightVector1Length = length(LightVector1);
 		TileCullingData0.ConeAxis = LightVector0 / LightVector0Length;
 		TileCullingData1.ConeAxis = LightVector1 / LightVector1Length;;
@@ -392,8 +397,8 @@ void SetupTileCullingData(
 		TileCullingData1.ConeAngleSin = TileCullingData1.ConeAngleCos * TanLightAngle1;
 
 		// Don't operate on tiles completely outside of the light's influence
-		bool bTileInLightInfluenceBounds = LightVector0Length < 1.0f / LightPositionAndInvRadius.w + TileCullingData0.BoundingSphere.w
-			|| LightVector1Length < 1.0f / LightPositionAndInvRadius.w + TileCullingData1.BoundingSphere.w;
+		bool bTileInLightInfluenceBounds = LightVector0Length < 1.0f / LightTranslatedPositionAndInvRadius.w + TileCullingData0.TranslatedBoundingSphere.w
+			|| LightVector1Length < 1.0f / LightTranslatedPositionAndInvRadius.w + TileCullingData1.TranslatedBoundingSphere.w;
 
 		bTileShouldComputeShadowing = bTileShouldComputeShadowing && bTileInLightInfluenceBounds;
 
@@ -421,12 +426,12 @@ uint CullCapsuleShapesToTile(
 {
 #if LIGHT_SOURCE_MODE == LIGHT_SOURCE_PUNCTUAL
 
-	float3 ConeAxis0 = TileCullingData0.ConeAxis;
-	float ConeAngleCos0 = TileCullingData0.ConeAngleCos;
-	float ConeAngleSin0 = TileCullingData0.ConeAngleSin;
-	float3 ConeAxis1 = TileCullingData1.ConeAxis;
-	float ConeAngleCos1 = TileCullingData1.ConeAngleCos;
-	float ConeAngleSin1 = TileCullingData1.ConeAngleSin;
+	const float3 ConeAxis0 = TileCullingData0.ConeAxis;
+	const float ConeAngleCos0 = TileCullingData0.ConeAngleCos;
+	const float ConeAngleSin0 = TileCullingData0.ConeAngleSin;
+	const float3 ConeAxis1 = TileCullingData1.ConeAxis;
+	const float ConeAngleCos1 = TileCullingData1.ConeAngleCos;
+	const float ConeAngleSin1 = TileCullingData1.ConeAngleSin;
 
 #endif
 
@@ -451,10 +456,10 @@ uint CullCapsuleShapesToTile(
 		
 		FCapsuleShape CapsuleShape = ShadowCapsuleShapes[ShapeIndex];
 
-		float3 TransformedSphereCenter = CapsuleShape.Center;
+		float3 TransformedSphereTranslatedCenter = CapsuleShape.TranslatedCenter;
 		float TransformedSphereRadius = CapsuleShape.Radius;
-		float3 TransformedTileBoundingSphereCenter0 = TileCullingData0.BoundingSphere.xyz;
-		float3 TransformedTileBoundingSphereCenter1 = TileCullingData1.BoundingSphere.xyz;
+		float3 TransformedTileTranslatedBoundingSphereCenter0 = TileCullingData0.TranslatedBoundingSphere.xyz;
+		float3 TransformedTileTranslatedBoundingSphereCenter1 = TileCullingData1.TranslatedBoundingSphere.xyz;
 
 		#if LIGHT_SOURCE_MODE != LIGHT_SOURCE_FROM_RECEIVER
 			float3 TransformedConeAxis0 = ConeAxis0;
@@ -472,17 +477,17 @@ uint CullCapsuleShapesToTile(
 			CapsuleSpaceZ *= CapsuleZScale;
 
 			// The capsule is centered at 0 in the scaled sphere space
-			TransformedSphereCenter = 0;
+			TransformedSphereTranslatedCenter = 0;
 
 			// After scaling along the capsule axis it will become a sphere with the original radius
 			TransformedSphereRadius = CapsuleShape.Radius;
 
 			// Transform the sphere center and cone axis into the scaled sphere space
-			float3 CapsuleCenterToTileCenter0 = TileCullingData0.BoundingSphere.xyz - CapsuleShape.Center;
-			TransformedTileBoundingSphereCenter0 = float3(dot(CapsuleCenterToTileCenter0, CapsuleSpaceX), dot(CapsuleCenterToTileCenter0, CapsuleSpaceY), dot(CapsuleCenterToTileCenter0, CapsuleSpaceZ));
+			float3 CapsuleCenterToTileCenter0 = TileCullingData0.TranslatedBoundingSphere.xyz - CapsuleShape.TranslatedCenter;
+			TransformedTileTranslatedBoundingSphereCenter0 = float3(dot(CapsuleCenterToTileCenter0, CapsuleSpaceX), dot(CapsuleCenterToTileCenter0, CapsuleSpaceY), dot(CapsuleCenterToTileCenter0, CapsuleSpaceZ));
 			
-			float3 CapsuleCenterToTileCenter1 = TileCullingData1.BoundingSphere.xyz - CapsuleShape.Center;
-			TransformedTileBoundingSphereCenter1 = float3(dot(CapsuleCenterToTileCenter1, CapsuleSpaceX), dot(CapsuleCenterToTileCenter1, CapsuleSpaceY), dot(CapsuleCenterToTileCenter1, CapsuleSpaceZ));
+			float3 CapsuleCenterToTileCenter1 = TileCullingData1.TranslatedBoundingSphere.xyz - CapsuleShape.TranslatedCenter;
+			TransformedTileTranslatedBoundingSphereCenter1 = float3(dot(CapsuleCenterToTileCenter1, CapsuleSpaceX), dot(CapsuleCenterToTileCenter1, CapsuleSpaceY), dot(CapsuleCenterToTileCenter1, CapsuleSpaceZ));
 
 			#if LIGHT_SOURCE_MODE != LIGHT_SOURCE_FROM_RECEIVER
 				// Renormalize the cone axis as it went through a non-uniformly scaled transform
@@ -495,9 +500,9 @@ uint CullCapsuleShapesToTile(
 		#endif
 
 		BRANCH
-		if (SphereIntersectSphere(float4(TransformedSphereCenter, TransformedSphereRadius + MaxOcclusionDistance), float4(TransformedTileBoundingSphereCenter0, TileCullingData0.BoundingSphere.w))
+		if (SphereIntersectSphere(float4(TransformedSphereTranslatedCenter, TransformedSphereRadius + MaxOcclusionDistance), float4(TransformedTileTranslatedBoundingSphereCenter0, TileCullingData0.BoundingSphere.w))
 			#if LIGHT_SOURCE_MODE != LIGHT_SOURCE_FROM_RECEIVER
-				&& SphereIntersectConeWithMaxDistance(float4(TransformedSphereCenter, TransformedSphereRadius + TileCullingData0.BoundingSphere.w), TransformedTileBoundingSphereCenter0, TransformedConeAxis0, ConeAngleCos0, ConeAngleSin0, MaxOcclusionDistance)
+				&& SphereIntersectConeWithMaxDistance(float4(TransformedSphereTranslatedCenter, TransformedSphereRadius + TileCullingData0.BoundingSphere.w), TransformedTileTranslatedBoundingSphereCenter0, TransformedConeAxis0, ConeAngleCos0, ConeAngleSin0, MaxOcclusionDistance)
 			#endif
 			)
 		{
@@ -509,9 +514,9 @@ uint CullCapsuleShapesToTile(
 		}
 
 		BRANCH
-		if (SphereIntersectSphere(float4(TransformedSphereCenter, TransformedSphereRadius + MaxOcclusionDistance), float4(TransformedTileBoundingSphereCenter1, TileCullingData1.BoundingSphere.w))
+		if (SphereIntersectSphere(float4(TransformedSphereTranslatedCenter, TransformedSphereRadius + MaxOcclusionDistance), float4(TransformedTileTranslatedBoundingSphereCenter1, TileCullingData1.BoundingSphere.w))
 			#if LIGHT_SOURCE_MODE != LIGHT_SOURCE_FROM_RECEIVER
-				&& SphereIntersectConeWithMaxDistance(float4(TransformedSphereCenter, TransformedSphereRadius + TileCullingData1.BoundingSphere.w), TransformedTileBoundingSphereCenter1, TransformedConeAxis1, ConeAngleCos1, ConeAngleSin1, MaxOcclusionDistance)
+				&& SphereIntersectConeWithMaxDistance(float4(TransformedSphereTranslatedCenter, TransformedSphereRadius + TileCullingData1.BoundingSphere.w), TransformedTileTranslatedBoundingSphereCenter1, TransformedConeAxis1, ConeAngleCos1, ConeAngleSin1, MaxOcclusionDistance)
 			#endif
 			)
 		{
@@ -570,7 +575,7 @@ void ApplyFadeToCapsuleRadius(inout float CapsuleRadius, float3 LightDirection)
 }
 
 float ShadowConeTraceAgainstCulledCapsuleShapes(
-	float3 WorldRayStart, 
+	float3 TranslatedWorldRayStart, 
 	float3 UnitRayDirection,
 	float LightVectorLength,
 	float LightAngle,  
@@ -632,7 +637,7 @@ float ShadowConeTraceAgainstCulledCapsuleShapes(
 			float CapsuleZScale = CapsuleShape.Radius / (.5f * CapsuleShape.Length + CapsuleShape.Radius);
 			CapsuleSpaceZ *= CapsuleZScale;
 
-			float3 CapsuleCenterToRayStart = WorldRayStart - CapsuleShape.Center;
+			float3 CapsuleCenterToRayStart = TranslatedWorldRayStart - CapsuleShape.TranslatedCenter;
 			float3 CapsuleSpaceRayStart = float3(dot(CapsuleCenterToRayStart, CapsuleSpaceX), dot(CapsuleCenterToRayStart, CapsuleSpaceY), dot(CapsuleCenterToRayStart, CapsuleSpaceZ));
 
 			float3 CapsuleSpaceRayDirection = float3(dot(UnitRayDirection, CapsuleSpaceX), dot(UnitRayDirection, CapsuleSpaceY), dot(UnitRayDirection, CapsuleSpaceZ));
@@ -641,7 +646,7 @@ float ShadowConeTraceAgainstCulledCapsuleShapes(
 			UnitVectorToShadowSphere = -CapsuleSpaceRayStart / DistanceToShadowSphere;
 			UnitRayDirectionInCorrectSpace = normalize(CapsuleSpaceRayDirection);
 #else
-			float3 VectorToCapsuleCenter = CapsuleShape.Center - WorldRayStart;
+			float3 VectorToCapsuleCenter = CapsuleShape.TranslatedCenter - TranslatedWorldRayStart;
 
 			// Closest point on line segment to ray
 			float3 L01 = CapsuleShape.Orientation * CapsuleShape.Length;
@@ -671,8 +676,8 @@ float ShadowConeTraceAgainstCulledCapsuleShapes(
 		}
 		else
 		{
-			DistanceToShadowSphere = length(CapsuleShape.Center - WorldRayStart);
-			UnitVectorToShadowSphere = (CapsuleShape.Center - WorldRayStart) / DistanceToShadowSphere;
+			DistanceToShadowSphere = length(CapsuleShape.TranslatedCenter - TranslatedWorldRayStart);
+			UnitVectorToShadowSphere = (CapsuleShape.TranslatedCenter - TranslatedWorldRayStart) / DistanceToShadowSphere;
 		}
 
 		float AngleBetween = acosFast(dot(UnitVectorToShadowSphere, UnitRayDirectionInCorrectSpace));
@@ -1003,7 +1008,8 @@ void CapsuleShadowingCS(
 	float SceneDepth = CalcSceneDepth(ScreenUV);
 #endif
 	
-	float3 OpaqueWorldPosition = mul(float4(ScreenPosition * SceneDepth, SceneDepth, 1), LWCHackToFloat(PrimaryView.ScreenToWorld)).xyz;
+	const float3 OpaqueTranslatedWorldPosition = mul(float4(ScreenPosition * SceneDepth, SceneDepth, 1), PrimaryView.ScreenToTranslatedWorld).xyz;
+	const float3 OpaqueWorldPosition = OpaqueTranslatedWorldPosition - LWCHackToFloat(PrimaryView.PreViewTranslation);
 
 	uint CulledDataParameter = 0;
 	bool bTileShouldComputeShadowing = true;
@@ -1035,25 +1041,29 @@ void CapsuleShadowingCS(
 		#if LIGHT_SOURCE_MODE == LIGHT_SOURCE_PUNCTUAL
 			#if POINT_LIGHT
 
-				float3 LightVector = LightPositionAndInvRadius.xyz - OpaqueWorldPosition;
+				float3 LightVector = LightTranslatedPositionAndInvRadius.xyz - OpaqueTranslatedWorldPosition;
 				LightVectorLength = length(LightVector);
-				float3 WorldRayStart = OpaqueWorldPosition + LightVector / LightVectorLength * RayStartOffset;
-				float3 UnitRayDirection = (LightPositionAndInvRadius.xyz - OpaqueWorldPosition) / LightVectorLength;
+				float3 TranslatedWorldRayStart = OpaqueTranslatedWorldPosition + LightVector / LightVectorLength * RayStartOffset;
+				float3 WorldRayStart = OpaqueWorldPosition + LightVector / LightVectorLength * RayStartOffset; // LWC_TODO
+				float3 UnitRayDirection = (LightTranslatedPositionAndInvRadius.xyz - OpaqueTranslatedWorldPosition) / LightVectorLength;
 				float LightAngle = atanFastPos(LightSourceRadius / LightVectorLength);
 
 			#else
 
+				float3 TranslatedWorldRayStart = OpaqueTranslatedWorldPosition + LightDirection * RayStartOffset;
 				float3 WorldRayStart = OpaqueWorldPosition + LightDirection * RayStartOffset;
 				float3 UnitRayDirection = LightDirection;
 				float LightAngle = LightAngleAndNormalThreshold.x;
 
 			#endif
 		#elif LIGHT_SOURCE_MODE == LIGHT_SOURCE_FROM_RECEIVER
+			float3 TranslatedWorldRayStart = OpaqueTranslatedWorldPosition;
 			float3 WorldRayStart = OpaqueWorldPosition;
 			float BentNormalLength = length(ReceiverBentNormal);
 			float3 UnitRayDirection = ReceiverBentNormal / max(BentNormalLength, .00001f);
 			float LightAngle = max(BentNormalLength * .5f * PI, PI / 8);
 		#else
+			float3 TranslatedWorldRayStart = OpaqueTranslatedWorldPosition;
 			float3 WorldRayStart = OpaqueWorldPosition;
 			float3 UnitRayDirection = 0;
 			float LightAngle = 0;
@@ -1091,7 +1101,7 @@ void CapsuleShadowingCS(
 			
 		#if SUPPORT_CAPSULE_SHAPES
 			Visibility *= ShadowConeTraceAgainstCulledCapsuleShapes(
-				WorldRayStart, 
+				TranslatedWorldRayStart,
 				UnitRayDirection, 
 				LightVectorLength,
 				LightAngle, 
diff --git a/Engine/Source/Runtime/Engine/Private/SkeletalMesh.cpp b/Engine/Source/Runtime/Engine/Private/SkeletalMesh.cpp
--- a/Engine/Source/Runtime/Engine/Private/SkeletalMesh.cpp
+++ b/Engine/Source/Runtime/Engine/Private/SkeletalMesh.cpp
@@ -6829,15 +6829,15 @@ bool FSkeletalMeshSceneProxy::HasDynamicIndirectShadowCasterRepresentation() con
 	return CastsDynamicShadow() && CastsDynamicIndirectShadow();
 }
 
-void FSkeletalMeshSceneProxy::GetShadowShapes(TArray<FCapsuleShape3f>& CapsuleShapes) const 
+void FSkeletalMeshSceneProxy::GetShadowShapes(FVector PreViewTranslation, TArray<FCapsuleShape3f>& OutCapsuleShapes) const
 {
 	SCOPE_CYCLE_COUNTER(STAT_GetShadowShapes);
 
 	const TArray<FMatrix44f>& ReferenceToLocalMatrices = MeshObject->GetReferenceToLocalMatrices();
 	const FMatrix& ProxyLocalToWorld = GetLocalToWorld();
 
-	int32 CapsuleIndex = CapsuleShapes.Num();
-	CapsuleShapes.SetNum(CapsuleShapes.Num() + ShadowCapsuleData.Num(), false);
+	int32 CapsuleIndex = OutCapsuleShapes.Num();
+	OutCapsuleShapes.SetNum(OutCapsuleShapes.Num() + ShadowCapsuleData.Num(), false);
 
 	for(const TPair<int32, FCapsuleShape>& CapsuleData : ShadowCapsuleData)
 	{
@@ -6848,9 +6848,9 @@ void FSkeletalMeshSceneProxy::GetShadowShapes(TArray<FCapsuleShape3f>& CapsuleSh
 		}
 		const float MaxScale = ReferenceToWorld.GetScaleVector().GetMax();
 
-		FCapsuleShape3f& NewCapsule = CapsuleShapes[CapsuleIndex++];
+		FCapsuleShape3f& NewCapsule = OutCapsuleShapes[CapsuleIndex++];
 
-		NewCapsule.Center = (FVector4f)ReferenceToWorld.TransformPosition(CapsuleData.Value.Center);
+		NewCapsule.Center = (FVector4f)(ReferenceToWorld.TransformPosition(CapsuleData.Value.Center) + PreViewTranslation);
 		NewCapsule.Radius = CapsuleData.Value.Radius * MaxScale;
 		NewCapsule.Orientation = (FVector4f)ReferenceToWorld.TransformVector(CapsuleData.Value.Orientation).GetSafeNormal();
 		NewCapsule.Length = CapsuleData.Value.Length * MaxScale;
diff --git a/Engine/Source/Runtime/Engine/Public/PrimitiveSceneProxy.h b/Engine/Source/Runtime/Engine/Public/PrimitiveSceneProxy.h
--- a/Engine/Source/Runtime/Engine/Public/PrimitiveSceneProxy.h
+++ b/Engine/Source/Runtime/Engine/Public/PrimitiveSceneProxy.h
@@ -260,7 +260,7 @@ public:
 	virtual void GetMeshDescription(int32 LODIndex, TArray<FMeshBatch>& OutMeshElements) const {}
 
 	/** Gathers shadow shapes from this proxy. */
-	virtual void GetShadowShapes(TArray<FCapsuleShape3f>& CapsuleShapes) const {}		// LWC_TODO: Precision loss? Forcing float variant of FCapsuleShape3f.
+	virtual void GetShadowShapes(FVector PreViewTranslation, TArray<FCapsuleShape3f>& OutCapsuleShapes) const {}
 
 #if RHI_RAYTRACING
 	// TODO: remove these individual functions in favor of ERayTracingPrimitiveFlags
diff --git a/Engine/Source/Runtime/Engine/Public/SkeletalMeshTypes.h b/Engine/Source/Runtime/Engine/Public/SkeletalMeshTypes.h
--- a/Engine/Source/Runtime/Engine/Public/SkeletalMeshTypes.h
+++ b/Engine/Source/Runtime/Engine/Public/SkeletalMeshTypes.h
@@ -173,7 +173,7 @@ public:
 	virtual bool IsUsingDistanceCullFade() const override;
 	
 	virtual bool HasDynamicIndirectShadowCasterRepresentation() const override;
-	virtual void GetShadowShapes(TArray<FCapsuleShape3f>& CapsuleShapes) const override;
+	virtual void GetShadowShapes(FVector PreViewTranslation, TArray<FCapsuleShape3f>& OutCapsuleShapes) const override;
 
 	/** Return the bounds for the pre-skinned primitive in local space */
 	virtual void GetPreSkinnedLocalBounds(FBoxSphereBounds& OutBounds) const override { OutBounds = PreSkinnedLocalBounds; }
diff --git a/Engine/Source/Runtime/Renderer/Private/CapsuleShadowRendering.cpp b/Engine/Source/Runtime/Renderer/Private/CapsuleShadowRendering.cpp
--- a/Engine/Source/Runtime/Renderer/Private/CapsuleShadowRendering.cpp
+++ b/Engine/Source/Runtime/Renderer/Private/CapsuleShadowRendering.cpp
@@ -209,7 +209,7 @@ public:
 		SHADER_PARAMETER(FIntPoint, TileDimensions)
 		SHADER_PARAMETER(FVector2f, NumGroups)
 
-		SHADER_PARAMETER(FVector4f, LightPositionAndInvRadius)
+		SHADER_PARAMETER(FVector4f, LightTranslatedPositionAndInvRadius)
 		SHADER_PARAMETER(FVector3f, LightDirection)
 		SHADER_PARAMETER(float, LightSourceRadius)
 		SHADER_PARAMETER(FVector3f, LightAngleAndNormalThreshold)
@@ -430,9 +430,8 @@ void SetupCapsuleShadowingParameters(
 
 		Parameters.LightDirection = LightParameters.Direction;
 
-		// LWC_TODO
-		FVector4f LightPositionAndInvRadius((FVector3f)LightParameters.WorldPosition, LightParameters.InvRadius);
-		Parameters.LightPositionAndInvRadius = LightPositionAndInvRadius;
+		FVector4f LightTranslatedPositionAndInvRadius((FVector3f)(LightParameters.WorldPosition + View.ViewMatrices.GetPreViewTranslation()), LightParameters.InvRadius);
+		Parameters.LightTranslatedPositionAndInvRadius = LightTranslatedPositionAndInvRadius;
 
 		// Default light source radius of 0 gives poor results
 		Parameters.LightSourceRadius = LightParameters.SourceRadius == 0 ? 20 : FMath::Clamp(LightParameters.SourceRadius, .001f, 1.0f / (4 * LightParameters.InvRadius));
@@ -509,6 +508,8 @@ bool FDeferredShadingSceneRenderer::RenderCapsuleDirectShadows(
 	for (int32 ViewIndex = 0; ViewIndex < Views.Num(); ViewIndex++)
 	{
 		const FViewInfo& View = Views[ViewIndex];
+		const FVector PreViewTranslation = View.ViewMatrices.GetPreViewTranslation();
+
 		RDG_GPU_MASK_SCOPE(GraphBuilder, View.GPUMask);
 		RDG_EVENT_SCOPE(GraphBuilder, "CapsuleShadows");
 		RDG_GPU_STAT_SCOPE(GraphBuilder, CapsuleShadows);
@@ -530,7 +531,7 @@ bool FDeferredShadingSceneRenderer::RenderCapsuleDirectShadows(
 
 				if (PrimitiveSceneInfo->Proxy->CastsDynamicShadow())
 				{
-					PrimitiveSceneInfo->Proxy->GetShadowShapes(CapsuleShapeData);
+					PrimitiveSceneInfo->Proxy->GetShadowShapes(PreViewTranslation, CapsuleShapeData);
 				}
 			}
 
@@ -749,6 +750,8 @@ static IndirectCapsuleShadowsResources CreateIndirectCapsuleShadowsResources(
 	int32& NumMeshesWithCapsules, 
 	int32& NumMeshDistanceFieldCasters)
 {
+	const FVector PreViewTranslation = View.ViewMatrices.GetPreViewTranslation();
+
 	IndirectCapsuleShadowsResources Output;
 	Output.IndirectShadowCapsuleShapesSRV = nullptr;
 	Output.IndirectShadowMeshDistanceFieldCasterIndicesSRV = nullptr;
@@ -839,7 +842,7 @@ static IndirectCapsuleShadowsResources CreateIndirectCapsuleShadowsResources(
 
 				if (GroupPrimitiveSceneInfo->Proxy->CastsDynamicShadow())
 				{
-					GroupPrimitiveSceneInfo->Proxy->GetShadowShapes(CapsuleShapeData);
+					GroupPrimitiveSceneInfo->Proxy->GetShadowShapes(PreViewTranslation, CapsuleShapeData);
 					
 					if (GroupPrimitiveSceneInfo->Proxy->HasDistanceFieldRepresentation())
 					{
