Added bottom up refraction lobe to strata tree operators. That is a better match to evaluate light rough refraction of light rays from sub layers to outside.

diff --git a/Engine/Shaders/Private/DistortAccumulatePS.usf b/Engine/Shaders/Private/DistortAccumulatePS.usf
--- a/Engine/Shaders/Private/DistortAccumulatePS.usf
+++ b/Engine/Shaders/Private/DistortAccumulatePS.usf
@@ -108,9 +108,9 @@ void Main(
 		// Update tree (coverage/transmittance/luminace weights)
 		StrataUpdateTree(NullStrataAddressing, StrataPixelHeader, StrataTree, V, DummyL, false /*bRoughDiffuseEnabled*/);
 
-		RefractionWorldNormal = StrataTree.Operators[StrataData.OperatorIndex].RefractionWorldNormal;
+		RefractionWorldNormal = StrataTree.Operators[StrataData.OperatorIndex].TopDownRefractionWorldNormal;
 
-		FStrataLobeStatistic RefractionLobeStat = StrataTree.Operators[StrataData.OperatorIndex].RefractionLobeStat;
+		FStrataLobeStatistic RefractionLobeStat = StrataTree.Operators[StrataData.OperatorIndex].TopDownRefractionLobeStat;
 
 		Sigma = RefractionLobeStat.Sigma;
 	}
diff --git a/Engine/Shaders/Private/Strata/Strata.ush b/Engine/Shaders/Private/Strata/Strata.ush
--- a/Engine/Shaders/Private/Strata/Strata.ush
+++ b/Engine/Shaders/Private/Strata/Strata.ush
@@ -713,32 +713,34 @@ struct FStrataOperator
 {
 	// STRATA_TODO: pack in a single uint
 	uint ParentIndex;
-	uint MaxDistanceFromLeaves;			// The largest depth distance from any tree leaf representing a BSDF
-	uint Type;							// The type of operator
+	uint MaxDistanceFromLeaves;							// The largest depth distance from any tree leaf representing a BSDF
+	uint Type;											// The type of operator
 	int	LeftIndex;
 	int RightIndex;
 	float Weight;
 
 	// Coverage / Transmittance
 
-	float  Coverage;					// The summarised coverage of all the matter (from the sub tree) represented by that operator
-	float3 Transmittance;				// This is the transmittance of the matter itself. If coverage is 0.25, then 25% of the light will be affected by the transmittance and 75% of light won't be.
+	float  Coverage;									// The summarised coverage of all the matter (from the sub tree) represented by that operator
+	float3 Transmittance;								// This is the transmittance of the matter itself. If coverage is 0.25, then 25% of the light will be affected by the transmittance and 75% of light won't be.
 
-	float VerticalTopCoverage;			// Only used by vertical layering operators. It represent the coverage of all the matter of the top layer (from the left sub tree).
-	float3 VerticalTopTransmittance;	// Only used by vertical layering operators. It represent the transmittance of matter of the top layer (from the left sub tree).
+	float VerticalTop_Coverage;							// Only used by vertical layering operators. It represent the coverage of all the matter of the top layer (from the left sub tree).
+	float3 VerticalTop_Transmittance;					// Only used by vertical layering operators. It represent the transmittance of matter of the top layer (from the left sub tree).
 
 	// Rough refraction
-	float3				 RefractionWorldNormal;			// The normal used for the single refracted direction
-	FStrataLobeStatistic RefractionLobeStat;			// The lobe used for the scene rough refraction
+	float3				 TopDownRefractionWorldNormal;	// The normal used for the single refracted direction
+	FStrataLobeStatistic TopDownRefractionLobeStat;		// The lobe used for the scene rough refraction from front to back, e.g. used for translucent rough refraciton
+	FStrataLobeStatistic BottomUpRefractionLobeStat;	// The lobe used for the scene rough refraction from back to front, e.g. used for opaque rough refraction of incoming light within sub layers
 
 	// Thickness
 	float				 ThicknessCm;					// The thickness of the layer
-	float				 VerticalTopThicknessCm;		// Only used by vertical layering operators. It represent the thickness that is used to drive the roughness propagation from top layers to sub layers.
+	float				 VerticalTop_ThicknessCm;		// Only used by vertical layering operators. It represent the thickness that is used to drive the roughness propagation from top layers to sub layers.
 
 	// Topology information
 	bool				 bIsTop;
 
-	FStrataLobeStatistic VerticalTopRefractionLobeStat;	// Only used by vertical layering operators. It represent the lobe statistics of the top layer at its interface.
+	FStrataLobeStatistic VerticalTop_TopDownRefractionLobeStat;	// Only used by vertical layering operators. It represent the lobe statistics of the top layer at its interface.
+	FStrataLobeStatistic VerticalTop_BottomUpRefractionLobeStat;
 };
 
 FStrataOperator GetInitialisedStrataOperator()
diff --git a/Engine/Shaders/Private/Strata/StrataTree.ush b/Engine/Shaders/Private/Strata/StrataTree.ush
--- a/Engine/Shaders/Private/Strata/StrataTree.ush
+++ b/Engine/Shaders/Private/Strata/StrataTree.ush
@@ -98,19 +98,32 @@ void UpdateSingleBSDFOperatorCoverageTransmittance(
 	Op.Coverage = CurrentBSDFCoverage;
 	Op.Transmittance = CurrentBSDFTransmittance;
 
-	// The top layer we consider a transition air to water. For layers below, we assume they all have similar index of refraction.
+	// Evaluate refraction lobes.
 	const float Roughness = StrataGetBSDFRoughness(CurrentBSDF);
 	const FBxDFEnergyTerms EnergyTerms = ComputeGGXSpecEnergyTerms(Roughness, StrataBSDFContext.SatNoV, StrataGetBSDFSpecularColor(CurrentBSDF), StrataGetBSDFSpecularF90(CurrentBSDF));
-	const float TopDownInterfaceEta = CurrentBSDF.bIsTop ? AIR_IOR / WATER_IOR : 1.0f;
-	FStrataLobeStatistic RefractedLobe = StrataGetRefractedLobe(StrataGetViewDiracLobe(-V), EnergyTerms.E, Roughness, TopDownInterfaceEta);
-	// Account for transmittance of the medium (but ignoring scattering for now)
-	RefractedLobe.E *= CurrentBSDFTransmittance;
-	// Account for coverate
-	RefractedLobe = StrataWeightLobe(RefractedLobe, CurrentBSDFCoverage);
-	Op.RefractionLobeStat = RefractedLobe;
+	{
+		const float TopDownInterfaceEta = CurrentBSDF.bIsTop ? AIR_IOR / WATER_IOR : 1.0f;
+		FStrataLobeStatistic RefractedLobe = StrataGetRefractedLobe(StrataGetViewDiracLobe(-V), EnergyTerms.E, Roughness, TopDownInterfaceEta);
+		// Account for transmittance of the medium (but ignoring scattering for now)
+		RefractedLobe.E *= CurrentBSDFTransmittance;
+		// Reduce according to coverage her. It is a approximation to handle rough refraction reduction according to coverage when relying on a single lob.
+		RefractedLobe = StrataWeightLobe(RefractedLobe, CurrentBSDFCoverage);
+		Op.TopDownRefractionLobeStat = RefractedLobe;
+	}
+	// Same operation but for a bottom up lob
+	{
+		const float BottomUpInterfaceEta = CurrentBSDF.bIsTop ? WATER_IOR / AIR_IOR : 1.0f;
+		const float3 Up = float3(0, 0, 1);
+		FStrataLobeStatistic RefractedLobe = StrataGetRefractedLobe(StrataGetViewDiracLobe(Up), EnergyTerms.E, Roughness, BottomUpInterfaceEta);
+		// Account for transmittance of the medium (but ignoring scattering for now)
+		RefractedLobe.E *= CurrentBSDFTransmittance;
+		// Reduce according to coverage her. It is a approximation to handle rough refraction reduction according to coverage when relying on a single lob.
+		RefractedLobe = StrataWeightLobe(RefractedLobe, CurrentBSDFCoverage);
+		Op.BottomUpRefractionLobeStat = RefractedLobe;
+	}
 	
 	float3x3 TangentBasis = StrataGetBSDFSharedBasis_InlineShading(StrataPixelHeader, CurrentBSDF, NullStrataAddressing);
-	Op.RefractionWorldNormal = TangentBasis[2];
+	Op.TopDownRefractionWorldNormal = TangentBasis[2];
 
 	Op.bIsTop = CurrentBSDF.bIsTop;
 
@@ -143,8 +156,10 @@ void UpdateSingleOperatorCoverageTransmittance(inout FStrataTree StrataTree, int
 		Op.Coverage = Weight * OpA.Coverage;
 		Op.Transmittance = OpA.Transmittance;
 
-		Op.RefractionLobeStat = StrataWeightLobe(OpA.RefractionLobeStat, Weight);
-		Op.RefractionWorldNormal = OpA.RefractionWorldNormal * Weight;
+		Op.TopDownRefractionLobeStat = StrataWeightLobe(OpA.TopDownRefractionLobeStat, Weight);
+		Op.TopDownRefractionWorldNormal = OpA.TopDownRefractionWorldNormal * Weight;
+
+		Op.BottomUpRefractionLobeStat = StrataWeightLobe(OpA.BottomUpRefractionLobeStat, Weight);
 
 		// Coverage should not affect thickness
 		Op.ThicknessCm = OpA.ThicknessCm;
@@ -178,22 +193,32 @@ void UpdateSingleOperatorCoverageTransmittance(inout FStrataTree StrataTree, int
 		Op.Transmittance			= Info.TransmittanceOnlyTop * TopTransmittance + Info.TransmittanceOnlyBottom * BotTransmittance + Info.TransmittanceTopAndBottom * TopTransmittance * BotTransmittance;
 #endif
 
-		Op.VerticalTopCoverage		= OpA.Coverage;
-		Op.VerticalTopTransmittance = OpA.Transmittance;
+		Op.VerticalTop_Coverage		= OpA.Coverage;
+		Op.VerticalTop_Transmittance= OpA.Transmittance;
 
 		// Project top lob through bottom layer assuming medium have same IOR.
-		const float TopDownInterfaceEta = 1.0f;
-		Op.RefractionLobeStat = StrataGetRefractedLobe(OpA.RefractionLobeStat, /*InterfaceFDG*/0.0f, StrataLobeVarianceToRoughness(OpB.RefractionLobeStat.Sigma), TopDownInterfaceEta);
-		Op.RefractionLobeStat.E *= OpB.RefractionLobeStat.E * saturate(Op.VerticalTopTransmittance + (1.0f - Op.VerticalTopCoverage)); // Combine with botton layer medium and specular transmittance
-
-		// The ratio of BottomSurfaceOnlyVisibile and visible through top layer (with a scale of 0.5 because at maximum both surfaces will be visible) as a fonction of the total coverage.
-		const float RefractionNormalMix = saturate((Info.SurfaceBottom + 0.5f * Info.TransmittanceTopAndBottom * dot(TopTransmittance, 0.33.xxx)) / max(STRATA_EPSILON, Info.Coverage));
-		Op.RefractionWorldNormal = lerp(OpA.RefractionWorldNormal, OpB.RefractionWorldNormal, RefractionNormalMix);
+		{
+			const float TopDownInterfaceEta = 1.0f;
+			Op.TopDownRefractionLobeStat = StrataGetRefractedLobe(OpA.TopDownRefractionLobeStat, /*InterfaceFDG*/0.0f, StrataLobeVarianceToRoughness(OpB.TopDownRefractionLobeStat.Sigma), TopDownInterfaceEta);
+			Op.TopDownRefractionLobeStat.E *= saturate(OpB.Transmittance + (1.0f - OpB.Coverage)); // Combine with botton layer medium and specular transmittance
+
+			// The ratio of BottomSurfaceOnlyVisibile and visible through top layer (with a scale of 0.5 because at maximum both surfaces will be visible) as a fonction of the total coverage.
+			const float RefractionNormalMix = saturate((Info.SurfaceBottom + 0.5f * Info.TransmittanceTopAndBottom * dot(TopTransmittance, 0.33.xxx)) / max(STRATA_EPSILON, Info.Coverage));
+			Op.TopDownRefractionWorldNormal = lerp(OpA.TopDownRefractionWorldNormal, OpB.TopDownRefractionWorldNormal, RefractionNormalMix);
+		}
+
+		// Project bottom lob through top layer.
+		{
+			const float BottomUpInterfaceEta = Op.bIsTop ? WATER_IOR / AIR_IOR : 1.0f;
+			Op.BottomUpRefractionLobeStat = StrataGetRefractedLobe(OpB.BottomUpRefractionLobeStat, /*InterfaceFDG*/0.0f, StrataLobeVarianceToRoughness(OpA.BottomUpRefractionLobeStat.Sigma), BottomUpInterfaceEta);
+			Op.BottomUpRefractionLobeStat.E *= saturate(OpA.Transmittance + (1.0f - OpA.Coverage)); // Combine with top layer medium and specular transmittance
+		}
 
 		// Vertical layering accumulates both operators thickness
 		Op.ThicknessCm = OpA.ThicknessCm + OpB.ThicknessCm;
 
-		Op.VerticalTopRefractionLobeStat = OpA.RefractionLobeStat;
+		Op.VerticalTop_TopDownRefractionLobeStat = OpA.TopDownRefractionLobeStat;
+		Op.VerticalTop_BottomUpRefractionLobeStat = OpA.BottomUpRefractionLobeStat;
 		break;
 	}
 
@@ -206,8 +231,10 @@ void UpdateSingleOperatorCoverageTransmittance(inout FStrataTree StrataTree, int
 		Op.Coverage = AMix * OpA.Coverage + BMix * OpB.Coverage;
 		Op.Transmittance = (AMix * OpA.Coverage * OpA.Transmittance + BMix * OpB.Coverage * OpB.Transmittance) / max(STRATA_EPSILON, AMix * OpA.Coverage + BMix * OpB.Coverage);
 
-		Op.RefractionLobeStat = StrataHorizontalMixLobes(OpA.RefractionLobeStat, OpB.RefractionLobeStat, BMix);
-		Op.RefractionWorldNormal = lerp(OpA.RefractionWorldNormal, OpB.RefractionWorldNormal, BMix);
+		Op.TopDownRefractionLobeStat = StrataHorizontalMixLobes(OpA.TopDownRefractionLobeStat, OpB.TopDownRefractionLobeStat, BMix);
+		Op.TopDownRefractionWorldNormal = lerp(OpA.TopDownRefractionWorldNormal, OpB.TopDownRefractionWorldNormal, BMix);
+
+		Op.BottomUpRefractionLobeStat = StrataHorizontalMixLobes(OpA.BottomUpRefractionLobeStat, OpB.BottomUpRefractionLobeStat, BMix);
 
 		// Horizontal mixes both operators thickness
 		Op.ThicknessCm = lerp(OpA.ThicknessCm, OpB.ThicknessCm, BMix);
@@ -223,8 +250,10 @@ void UpdateSingleOperatorCoverageTransmittance(inout FStrataTree StrataTree, int
 		Op.Coverage = saturate(OpA.Coverage + OpB.Coverage);
 		Op.Transmittance = lerp(OpB.Transmittance, OpA.Transmittance, AMixFactor);
 
-		Op.RefractionLobeStat = StrataHorizontalMixLobes(OpA.RefractionLobeStat, OpB.RefractionLobeStat, 0.5f);
-		Op.RefractionWorldNormal = lerp(OpA.RefractionWorldNormal, OpB.RefractionWorldNormal, 0.5f);
+		Op.TopDownRefractionLobeStat = StrataHorizontalMixLobes(OpA.TopDownRefractionLobeStat, OpB.TopDownRefractionLobeStat, 0.5f);
+		Op.TopDownRefractionWorldNormal = lerp(OpA.TopDownRefractionWorldNormal, OpB.TopDownRefractionWorldNormal, 0.5f);
+
+		Op.BottomUpRefractionLobeStat = StrataHorizontalMixLobes(OpA.BottomUpRefractionLobeStat, OpB.BottomUpRefractionLobeStat, 0.5f);
 
 		Op.ThicknessCm = lerp(OpA.ThicknessCm, OpB.ThicknessCm, 0.5f);
 		break;
@@ -292,10 +321,10 @@ void UpdateAllBSDFWithBottomUpOperatorVisit_Vertical(
 	// PreviousIsInputA > 0 means it is a BSDF affecting the top layer so we do not affect it by anything for this vertical layering.
 	// Otherise, the BSDF comes from the bottom part so we affect it weights from the gathered top layer coverage/transmittance
 	CurrentBSDF.LuminanceWeight *=
-		PreviousIsInputA > 0 ? 1.0f : saturate(1.0f - Op.VerticalTopCoverage) + saturate(Op.VerticalTopCoverage) * Op.VerticalTopTransmittance;
+		PreviousIsInputA > 0 ? 1.0f : saturate(1.0f - Op.VerticalTop_Coverage) + saturate(Op.VerticalTop_Coverage) * Op.VerticalTop_Transmittance;
 
 	// If the BSDF is from the bottom, we reduce its contribution according to the top layer coverage only.
-	CurrentBSDF.TopLayerDataWeight *= PreviousIsInputA ? 1.0 : saturate(1.0f - Op.VerticalTopCoverage);
+	CurrentBSDF.TopLayerDataWeight *= PreviousIsInputA ? 1.0 : saturate(1.0f - Op.VerticalTop_Coverage);
 
 	if (PreviousIsInputA <= 0)	// Only apply to BSDF from the bottom layer
 	{
@@ -310,10 +339,9 @@ void UpdateAllBSDFWithBottomUpOperatorVisit_Vertical(
 
 		// Evaluate the reflected lobe after refraction through the top layer of the vertical operator.
 		// => This tells how light should be integrated when traveling up thorugh that layer.
-		const float TopLayerSigma = Op.VerticalTopRefractionLobeStat.Sigma;				// STRATA_TODO this sigma was evaluate as a top down lobe. We should evaluate it as a bottom up sigma for opaque rough refraction because this is refraction of the incoming light.
-		const float TopRoughnessAsVariance = TopLayerSigma * Op.VerticalTopCoverage;	// Weighting variance by coverage is a good way to handle rough refraction reduction according when relying on a single lob.
+		const float TopLayerSigma = Op.VerticalTop_BottomUpRefractionLobeStat.Sigma;
 		const float BottomUpInterfaceEta = Op.bIsTop ? WATER_IOR / AIR_IOR : 1.0f;
-		FStrataLobeStatistic RefractedLobe = StrataGetRefractedLobe(ReflectedLobe, DummyInterfaceDFG, StrataLobeVarianceToRoughness(TopRoughnessAsVariance), BottomUpInterfaceEta);
+		FStrataLobeStatistic RefractedLobe = StrataGetRefractedLobe(ReflectedLobe, DummyInterfaceDFG, StrataLobeVarianceToRoughness(TopLayerSigma), BottomUpInterfaceEta);
 
 		StrataSetBSDFRoughness(CurrentBSDF, StrataLobeVarianceToRoughness(RefractedLobe.Sigma));
 	}
