Use sample stretching for reservoir sampling in the path tracer

This speeds up the main trace loop by requiring fewer random numbers to be generated.

diff --git a/Engine/Shaders/Private/PathTracing/PathTracingCore.ush b/Engine/Shaders/Private/PathTracing/PathTracingCore.ush
--- a/Engine/Shaders/Private/PathTracing/PathTracingCore.ush
+++ b/Engine/Shaders/Private/PathTracing/PathTracingCore.ush
@@ -108,8 +108,10 @@ FPathTracingPayload TraceTransparentRay(inout FPathState PathState, bool bIsCame
 {
 	const uint RayFlags = bIsCameraRay && EnableCameraBackfaceCulling ? RAY_FLAG_CULL_BACK_FACING_TRIANGLES : 0;
 	const uint MissShaderIndex = 0;
-	FRISContext HitSample = InitRISContext();
-	FRISContext VolSample = InitRISContext();
+	float2 RISRandSample = RandomSequence_GenerateSample2D(PathState.RandSequence);
+	FRISContext HitSample = InitRISContext(RISRandSample.x);
+	FRISContext VolSample = InitRISContext(RISRandSample.y);
+
 	float3 PayloadThroughput = PathState.PathThroughput;
 	FPathTracingPayload Payload;
 	if (!bIncludeEmission && bLastBounce)
@@ -202,7 +204,7 @@ FPathTracingPayload TraceTransparentRay(inout FPathState PathState, bool bIsCame
 				// NOTE: this only matters if we are tracing through several transparent hits or if there are multiple overlapping volumes
 				float3 Contrib = PathState.PathThroughput;
 				float SelectionWeight = max3(Contrib.x, Contrib.y, Contrib.z);
-				if (VolSample.Accept(SelectionWeight, RandomSequence_GenerateSample1D(PathState.RandSequence)))
+				if (VolSample.Accept(SelectionWeight))
 				{
 					// store this volume segment for later
 					VolumeSegment.Throughput = PathState.PathThroughput / SelectionWeight;
@@ -235,8 +237,8 @@ FPathTracingPayload TraceTransparentRay(inout FPathState PathState, bool bIsCame
 					float3 ColorChannelPdf = PathThroughput;
 
 					// Sample the distance to the next interaction
-					float2 RandValue = RandomSequence_GenerateSample2D(PathState.RandSequence);
-					float DeltaT = SampleSpectralTransmittance(RandValue.x, SigmaBar, ColorChannelPdf);
+					float RandValue = RandomSequence_GenerateSample1D(PathState.RandSequence);
+					float DeltaT = SampleSpectralTransmittance(RandValue, SigmaBar, ColorChannelPdf);
 					if (DeltaT < 0.0)
 					{
 						// no more energy left in the path
@@ -268,7 +270,7 @@ FPathTracingPayload TraceTransparentRay(inout FPathState PathState, bool bIsCame
 							AccumulateAlbedo(SigmaS, PathThroughput, PathState.PathRoughness, PathState.Albedo);
 							float3 Contrib = PathThroughput * SigmaS;
 							float SelectionWeight = max3(Contrib.x, Contrib.y, Contrib.z);
-							if (HitSample.Accept(SelectionWeight, RandValue.y))
+							if (HitSample.Accept(SelectionWeight))
 							{
 								// stash this hit for next time
 								Payload = CreateMediumHitPayload(Interval.VolumeTMin, TranslatedWorldPos, Result);
@@ -346,7 +348,7 @@ FPathTracingPayload TraceTransparentRay(inout FPathState PathState, bool bIsCame
 			AccumulateAlbedoNormal(HitPayload, PathState.PathThroughput, PathState.PathRoughness, PathState.Albedo, PathState.Normal);
 
 			float SelectionWeight = max3(Contrib.x, Contrib.y, Contrib.z);
-			if (HitSample.Accept(SelectionWeight, RandomSequence_GenerateSample1D(PathState.RandSequence)))
+			if (HitSample.Accept(SelectionWeight))
 			{
 				// stash this hit for next time
 				Payload = HitPayload;
diff --git a/Engine/Shaders/Private/PathTracing/Utilities/PathTracingRIS.ush b/Engine/Shaders/Private/PathTracing/Utilities/PathTracingRIS.ush
--- a/Engine/Shaders/Private/PathTracing/Utilities/PathTracingRIS.ush
+++ b/Engine/Shaders/Private/PathTracing/Utilities/PathTracingRIS.ush
@@ -5,14 +5,33 @@
 // Implementation of the core logic of Weighted Resampled Importance Sampling
 // The idea is to simulate sampling a discrete pdf built of individually weighted samples
 
+// This implementation uses the sample-stretching trick presented in:
+//  Shinji Ogaki - "Vectorized Reservoir Sampling"
+//  Siggraph Asia 2021 Technical Communications
+
 struct FRISContext
 {
 	float WeightSum;
+	float RandSample;
 
-	bool Accept(float Weight, float RandValue)
+	bool Accept(float Weight)
 	{
 		WeightSum += Weight;
-		return Weight > 0 ? (Weight < WeightSum ? RandValue * WeightSum < Weight : true) : false;
+		if (Weight > 0)
+		{
+			const float OneMinusEpsilon = 0.99999994; // 32-bit float just before 1.0
+			RandSample *= WeightSum;
+			if (RandSample < Weight)
+			{
+				RandSample = min(RandSample / Weight, OneMinusEpsilon);
+				return true;
+			}
+			else
+			{
+				RandSample = min((RandSample - Weight) / (WeightSum - Weight), OneMinusEpsilon);
+			}
+		}
+		return false;
 	}
 
 	bool HasSample()
@@ -26,7 +45,9 @@ struct FRISContext
 	}
 };
 
-FRISContext InitRISContext()
+FRISContext InitRISContext(float RandSample)
 {
-	return (FRISContext)0;
+	FRISContext Context = (FRISContext)0;
+	Context.RandSample = RandSample;
+	return Context;
 }
