Fixed potential issue in pixel shaders reading strata mterial data in loops.
We now hardcode the mip level to 0, thus not requiring derivative from any pixel shaders.

diff --git a/Engine/Shaders/Private/Strata/Strata.ush b/Engine/Shaders/Private/Strata/Strata.ush
--- a/Engine/Shaders/Private/Strata/Strata.ush
+++ b/Engine/Shaders/Private/Strata/Strata.ush
@@ -1108,7 +1108,8 @@ void FinalizeWrites(inout FRWStrataMaterialContainer StrataBuffer, inout FStrata
 
 uint StrataLoadUint1(FStrataMaterialContainer StrataBuffer, inout FStrataAddressing StrataAddressing)
 {
-	uint Data = StrataBuffer[uint3(StrataAddressing.PixelCoords, StrataAddressing.CurrentIndex)];
+	const uint MipLevel = 0;
+	uint Data = StrataBuffer.mips[MipLevel][uint3(StrataAddressing.PixelCoords, StrataAddressing.CurrentIndex)];
 	StrataAddressing.CurrentIndex++;
 	StrataAddressing.ReadBytes += 4;
 	return Data;
