Add Line/AABB/OBB/... version with FShaderPrintContext.

This allows to create a single context for printing/drawing without relying on if (bIsDebugEnabledOnThisThread) DrawDebugThings();

diff --git a/Engine/Shaders/Private/ShaderPrint.ush b/Engine/Shaders/Private/ShaderPrint.ush
--- a/Engine/Shaders/Private/ShaderPrint.ush
+++ b/Engine/Shaders/Private/ShaderPrint.ush
@@ -775,11 +775,19 @@ void AddLineTWS(float3 Pos0, float3 Pos1, float4 Color0, float4 Color1)
 	}
 }
 
-void AddLineTWS(float3 Pos0, float3 Pos1, float4 Color)					{ AddLineTWS(Pos0, Pos1, Color, Color); }
-void AddLineWS(float3 Pos0, float3 Pos1, float4 Color)					{ AddLineTWS(Pos0+ShaderPrint.TranslatedWorldOffset, Pos1+ShaderPrint.TranslatedWorldOffset, Color); }
-void AddLineWS(float3 Pos0, float3 Pos1, float4 Color0, float4 Color1)	{ AddLineTWS(Pos0+ShaderPrint.TranslatedWorldOffset, Pos1+ShaderPrint.TranslatedWorldOffset, Color0, Color1); }
-void AddLine(float3 Pos0, float3 Pos1, float4 Color)					{ AddLineTWS(Pos0+ShaderPrint.TranslatedWorldOffset, Pos1+ShaderPrint.TranslatedWorldOffset, Color); }
-void AddLine(float3 Pos0, float3 Pos1, float4 Color0, float4 Color1)	{ AddLineTWS(Pos0+ShaderPrint.TranslatedWorldOffset, Pos1+ShaderPrint.TranslatedWorldOffset, Color0, Color1); }
+void AddLineTWS(float3 Pos0, float3 Pos1, float4 Color)												{ AddLineTWS(Pos0, Pos1, Color, Color); }
+void AddLineWS (float3 Pos0, float3 Pos1, float4 Color)												{ AddLineTWS(Pos0+ShaderPrint.TranslatedWorldOffset, Pos1+ShaderPrint.TranslatedWorldOffset, Color); }
+void AddLineWS (float3 Pos0, float3 Pos1, float4 Color0, float4 Color1)								{ AddLineTWS(Pos0+ShaderPrint.TranslatedWorldOffset, Pos1+ShaderPrint.TranslatedWorldOffset, Color0, Color1); }
+void AddLine   (float3 Pos0, float3 Pos1, float4 Color)												{ AddLineTWS(Pos0+ShaderPrint.TranslatedWorldOffset, Pos1+ShaderPrint.TranslatedWorldOffset, Color); }
+void AddLine   (float3 Pos0, float3 Pos1, float4 Color0, float4 Color1)								{ AddLineTWS(Pos0+ShaderPrint.TranslatedWorldOffset, Pos1+ShaderPrint.TranslatedWorldOffset, Color0, Color1); }
+
+void AddLineTWS(FShaderPrintContext Ctx, float3 Pos0, float3 Pos1, float4 Color0, float4 Color1)	{ if (Ctx.bIsActive) { AddLineTWS(Pos0, Pos1, Color0, Color1); } }
+void AddLineTWS(FShaderPrintContext Ctx, float3 Pos0, float3 Pos1, float4 Color)					{ if (Ctx.bIsActive) { AddLineTWS(Pos0, Pos1, Color, Color); } }
+void AddLineWS (FShaderPrintContext Ctx, float3 Pos0, float3 Pos1, float4 Color)					{ if (Ctx.bIsActive) { AddLineTWS(Pos0+ShaderPrint.TranslatedWorldOffset, Pos1+ShaderPrint.TranslatedWorldOffset, Color); } }
+void AddLineWS (FShaderPrintContext Ctx, float3 Pos0, float3 Pos1, float4 Color0, float4 Color1)	{ if (Ctx.bIsActive) { AddLineTWS(Pos0+ShaderPrint.TranslatedWorldOffset, Pos1+ShaderPrint.TranslatedWorldOffset, Color0, Color1); } }
+void AddLine   (FShaderPrintContext Ctx, float3 Pos0, float3 Pos1, float4 Color)					{ if (Ctx.bIsActive) { AddLineTWS(Pos0+ShaderPrint.TranslatedWorldOffset, Pos1+ShaderPrint.TranslatedWorldOffset, Color); } }
+void AddLine   (FShaderPrintContext Ctx, float3 Pos0, float3 Pos1, float4 Color0, float4 Color1)	{ if (Ctx.bIsActive) { AddLineTWS(Pos0+ShaderPrint.TranslatedWorldOffset, Pos1+ShaderPrint.TranslatedWorldOffset, Color0, Color1); } }
+
 
 ///////////////////////////////////////////////////////////////////////////////////////////////////
 // Quad
@@ -798,8 +806,12 @@ void AddQuadTWS(float3 Pos0, float3 Pos1, float3 Pos2, float3 Pos3, float4 Color
 	}
 }
 
-void AddQuadWS(float3 Pos0, float3 Pos1, float3 Pos2, float3 Pos3, float4 Color) { AddQuadTWS(Pos0+ShaderPrint.TranslatedWorldOffset, Pos1+ShaderPrint.TranslatedWorldOffset, Pos2+ShaderPrint.TranslatedWorldOffset, Pos3+ShaderPrint.TranslatedWorldOffset, Color); }
-void AddQuad(float3 Pos0, float3 Pos1, float3 Pos2, float3 Pos3, float4 Color)	 { AddQuadTWS(Pos0+ShaderPrint.TranslatedWorldOffset, Pos1+ShaderPrint.TranslatedWorldOffset, Pos2+ShaderPrint.TranslatedWorldOffset, Pos3+ShaderPrint.TranslatedWorldOffset, Color); }
+void AddQuadWS(float3 Pos0, float3 Pos1, float3 Pos2, float3 Pos3, float4 Color)							{ AddQuadTWS(Pos0+ShaderPrint.TranslatedWorldOffset, Pos1+ShaderPrint.TranslatedWorldOffset, Pos2+ShaderPrint.TranslatedWorldOffset, Pos3+ShaderPrint.TranslatedWorldOffset, Color); }
+void AddQuad  (float3 Pos0, float3 Pos1, float3 Pos2, float3 Pos3, float4 Color)							{ AddQuadTWS(Pos0+ShaderPrint.TranslatedWorldOffset, Pos1+ShaderPrint.TranslatedWorldOffset, Pos2+ShaderPrint.TranslatedWorldOffset, Pos3+ShaderPrint.TranslatedWorldOffset, Color); }
+
+void AddQuadTWS(FShaderPrintContext Ctx, float3 Pos0, float3 Pos1, float3 Pos2, float3 Pos3, float4 Color)	{ if (Ctx.bIsActive) { AddQuadTWS(Pos0, Pos1, Pos2, Pos3, Color); } }
+void AddQuadWS (FShaderPrintContext Ctx, float3 Pos0, float3 Pos1, float3 Pos2, float3 Pos3, float4 Color)	{ if (Ctx.bIsActive) { AddQuadTWS(Pos0+ShaderPrint.TranslatedWorldOffset, Pos1+ShaderPrint.TranslatedWorldOffset, Pos2+ShaderPrint.TranslatedWorldOffset, Pos3+ShaderPrint.TranslatedWorldOffset, Color); } }
+void AddQuad   (FShaderPrintContext Ctx, float3 Pos0, float3 Pos1, float3 Pos2, float3 Pos3, float4 Color)	{ if (Ctx.bIsActive) { AddQuadTWS(Pos0+ShaderPrint.TranslatedWorldOffset, Pos1+ShaderPrint.TranslatedWorldOffset, Pos2+ShaderPrint.TranslatedWorldOffset, Pos3+ShaderPrint.TranslatedWorldOffset, Color); } }
 
 ///////////////////////////////////////////////////////////////////////////////////////////////////
 // AABB
@@ -838,8 +850,12 @@ void AddAABBTWS(float3 Min, float3 Max, float4 Color)
 	}
 }
 
-void AddAABBWS(float3 Min, float3 Max, float4 Color)	{ AddAABBTWS(Min + ShaderPrint.TranslatedWorldOffset, Max + ShaderPrint.TranslatedWorldOffset, Color); }
-void AddAABB(float3 Min, float3 Max, float4 Color)		{ AddAABBTWS(Min + ShaderPrint.TranslatedWorldOffset, Max + ShaderPrint.TranslatedWorldOffset, Color); }
+void AddAABBWS(float3 Min, float3 Max, float4 Color)							{ AddAABBTWS(Min + ShaderPrint.TranslatedWorldOffset, Max + ShaderPrint.TranslatedWorldOffset, Color); }
+void AddAABB  (float3 Min, float3 Max, float4 Color)							{ AddAABBTWS(Min + ShaderPrint.TranslatedWorldOffset, Max + ShaderPrint.TranslatedWorldOffset, Color); }
+
+void AddAABBTWS(FShaderPrintContext Ctx, float3 Min, float3 Max, float4 Color)	{ if (Ctx.bIsActive) { AddAABBTWS(Min, Max, Color); } }
+void AddAABBWS (FShaderPrintContext Ctx, float3 Min, float3 Max, float4 Color)	{ if (Ctx.bIsActive) { AddAABBTWS(Min + ShaderPrint.TranslatedWorldOffset, Max + ShaderPrint.TranslatedWorldOffset, Color); } }
+void AddAABB   (FShaderPrintContext Ctx, float3 Min, float3 Max, float4 Color)	{ if (Ctx.bIsActive) { AddAABBTWS(Min + ShaderPrint.TranslatedWorldOffset, Max + ShaderPrint.TranslatedWorldOffset, Color); } }
 
 ///////////////////////////////////////////////////////////////////////////////////////////////////
 // Cross
@@ -851,8 +867,12 @@ void AddCrossTWS(float3 Pos, float Size, float4 Color)
 	AddLineTWS(Pos - float3(0,0,Size), Pos + float3(0,0,Size), Color, Color);
 }
 
-void AddCrossWS(float3 Pos, float Size, float4 Color)	{ AddAABBTWS(Pos + ShaderPrint.TranslatedWorldOffset, Size, Color); }
-void AddCross(float3 Pos, float Size, float4 Color)		{ AddAABBTWS(Pos + ShaderPrint.TranslatedWorldOffset, Size, Color); }
+void AddCrossWS(float3 Pos, float Size, float4 Color)							{ AddAABBTWS(Pos + ShaderPrint.TranslatedWorldOffset, Size, Color); }
+void AddCross  (float3 Pos, float Size, float4 Color)							{ AddAABBTWS(Pos + ShaderPrint.TranslatedWorldOffset, Size, Color); }
+
+void AddCrossTWS(FShaderPrintContext Ctx, float3 Pos, float Size, float4 Color) { if (Ctx.bIsActive) { AddCrossTWS(Pos, Size, Color); } }
+void AddCrossWS (FShaderPrintContext Ctx, float3 Pos, float Size, float4 Color)	{ if (Ctx.bIsActive) { AddAABBTWS(Pos + ShaderPrint.TranslatedWorldOffset, Size, Color); } }
+void AddCross   (FShaderPrintContext Ctx, float3 Pos, float Size, float4 Color)	{ if (Ctx.bIsActive) { AddAABBTWS(Pos + ShaderPrint.TranslatedWorldOffset, Size, Color); } }
 
 ///////////////////////////////////////////////////////////////////////////////////////////////////
 // Referential
@@ -881,13 +901,22 @@ void AddReferentialTWS(float4x4 InM, float Scale = 1)
 	AddReferentialTWS(InM[3].xyz, InM[0].xyz, InM[1].xyz, InM[2].xyz, Scale);
 }
 
-void AddReferentialWS(float4x4 InM, float Scale = 1)							{ AddReferentialTWS(InM[3].xyz + ShaderPrint.TranslatedWorldOffset, InM[0].xyz, InM[1].xyz, InM[2].xyz, Scale); }
-void AddReferentialWS(float3 Pos, float3 TangentZ, float Scale = 1)				{ AddReferentialTWS(Pos + ShaderPrint.TranslatedWorldOffset, TangentZ, Scale); }
-void AddReferentialWS(float3 Pos, float3 T, float3 B, float3 N, float Scale = 1){ AddReferentialTWS(Pos + ShaderPrint.TranslatedWorldOffset, T, B, N, Scale); }
-void AddReferential(float4x4 InM, float Scale = 1)								{ AddReferentialTWS(InM[3].xyz + ShaderPrint.TranslatedWorldOffset, InM[0].xyz, InM[1].xyz, InM[2].xyz, Scale); }
-void AddReferential(float3 Pos, float3 TangentZ, float Scale = 1)				{ AddReferentialTWS(Pos + ShaderPrint.TranslatedWorldOffset, TangentZ, Scale); }
-void AddReferential(float3 Pos, float3 T, float3 B, float3 N, float Scale = 1)	{ AddReferentialTWS(Pos + ShaderPrint.TranslatedWorldOffset, T, B, N, Scale); }
-
+void AddReferentialWS(float4x4 InM, float Scale = 1)								{ AddReferentialTWS(InM[3].xyz + ShaderPrint.TranslatedWorldOffset, InM[0].xyz, InM[1].xyz, InM[2].xyz, Scale); }
+void AddReferentialWS(float3 Pos, float3 TangentZ, float Scale = 1)					{ AddReferentialTWS(Pos + ShaderPrint.TranslatedWorldOffset, TangentZ, Scale); }
+void AddReferentialWS(float3 Pos, float3 T, float3 B, float3 N, float Scale = 1)	{ AddReferentialTWS(Pos + ShaderPrint.TranslatedWorldOffset, T, B, N, Scale); }
+void AddReferential  (float4x4 InM, float Scale = 1)								{ AddReferentialTWS(InM[3].xyz + ShaderPrint.TranslatedWorldOffset, InM[0].xyz, InM[1].xyz, InM[2].xyz, Scale); }
+void AddReferential  (float3 Pos, float3 TangentZ, float Scale = 1)					{ AddReferentialTWS(Pos + ShaderPrint.TranslatedWorldOffset, TangentZ, Scale); }
+void AddReferential  (float3 Pos, float3 T, float3 B, float3 N, float Scale = 1)	{ AddReferentialTWS(Pos + ShaderPrint.TranslatedWorldOffset, T, B, N, Scale); }
+
+void AddReferentialTWS(FShaderPrintContext Ctx, float3 Pos, float3 T, float3 B, float3 N, float Scale = 1)	{ if (Ctx.bIsActive) { AddReferentialTWS(Pos, T, B, N, Scale); } }
+void AddReferentialTWS(FShaderPrintContext Ctx, float3 Pos, float3 TangentZ, float Scale = 1)				{ if (Ctx.bIsActive) { AddReferentialTWS(Pos, TangentZ, Scale); }  }
+void AddReferentialTWS(FShaderPrintContext Ctx, float4x4 InM, float Scale = 1)								{ if (Ctx.bIsActive) { AddReferentialTWS(InM[3].xyz, InM[0].xyz, InM[1].xyz, InM[2].xyz, Scale); }  }
+void AddReferentialWS (FShaderPrintContext Ctx, float4x4 InM, float Scale = 1)								{ if (Ctx.bIsActive) { AddReferentialTWS(InM[3].xyz + ShaderPrint.TranslatedWorldOffset, InM[0].xyz, InM[1].xyz, InM[2].xyz, Scale); }  }
+void AddReferentialWS (FShaderPrintContext Ctx, float3 Pos, float3 TangentZ, float Scale = 1)				{ if (Ctx.bIsActive) { AddReferentialTWS(Pos + ShaderPrint.TranslatedWorldOffset, TangentZ, Scale); }  }
+void AddReferentialWS (FShaderPrintContext Ctx, float3 Pos, float3 T, float3 B, float3 N, float Scale = 1)	{ if (Ctx.bIsActive) { AddReferentialTWS(Pos + ShaderPrint.TranslatedWorldOffset, T, B, N, Scale); }  }
+void AddReferential   (FShaderPrintContext Ctx, float4x4 InM, float Scale = 1)								{ if (Ctx.bIsActive) { AddReferentialTWS(InM[3].xyz + ShaderPrint.TranslatedWorldOffset, InM[0].xyz, InM[1].xyz, InM[2].xyz, Scale); }  }
+void AddReferential   (FShaderPrintContext Ctx, float3 Pos, float3 TangentZ, float Scale = 1)				{ if (Ctx.bIsActive) { AddReferentialTWS(Pos + ShaderPrint.TranslatedWorldOffset, TangentZ, Scale); }  }
+void AddReferential   (FShaderPrintContext Ctx, float3 Pos, float3 T, float3 B, float3 N, float Scale = 1)	{ if (Ctx.bIsActive) { AddReferentialTWS(Pos + ShaderPrint.TranslatedWorldOffset, T, B, N, Scale); }  }
 ///////////////////////////////////////////////////////////////////////////////////////////////////
 // Triangle
 
@@ -898,8 +927,12 @@ void AddLineTriangleTWS(float3 P0, float3 P1, float3 P2, float4 Color)
 	AddLineTWS(P2, P0, Color, Color);
 }
 
-void AddLineTriangleWS(float3 P0, float3 P1, float3 P2, float4 Color)	{ AddLineTriangleTWS(P0 + ShaderPrint.TranslatedWorldOffset, P1 + ShaderPrint.TranslatedWorldOffset, P2 + ShaderPrint.TranslatedWorldOffset, Color); }
-void AddLineTriangle(float3 P0, float3 P1, float3 P2, float4 Color)		{ AddLineTriangleTWS(P0 + ShaderPrint.TranslatedWorldOffset, P1 + ShaderPrint.TranslatedWorldOffset, P2 + ShaderPrint.TranslatedWorldOffset, Color); }
+void AddLineTriangleWS(float3 P0, float3 P1, float3 P2, float4 Color)							{ AddLineTriangleTWS(P0 + ShaderPrint.TranslatedWorldOffset, P1 + ShaderPrint.TranslatedWorldOffset, P2 + ShaderPrint.TranslatedWorldOffset, Color); }
+void AddLineTriangle  (float3 P0, float3 P1, float3 P2, float4 Color)							{ AddLineTriangleTWS(P0 + ShaderPrint.TranslatedWorldOffset, P1 + ShaderPrint.TranslatedWorldOffset, P2 + ShaderPrint.TranslatedWorldOffset, Color); }
+
+void AddLineTriangleTWS(FShaderPrintContext Ctx, float3 P0, float3 P1, float3 P2, float4 Color)	{ if (Ctx.bIsActive) { AddLineTriangleTWS(P0, P1, P2, Color); } }
+void AddLineTriangleWS (FShaderPrintContext Ctx, float3 P0, float3 P1, float3 P2, float4 Color)	{ if (Ctx.bIsActive) { AddLineTriangleTWS(P0 + ShaderPrint.TranslatedWorldOffset, P1 + ShaderPrint.TranslatedWorldOffset, P2 + ShaderPrint.TranslatedWorldOffset, Color); } }
+void AddLineTriangle   (FShaderPrintContext Ctx, float3 P0, float3 P1, float3 P2, float4 Color)	{ if (Ctx.bIsActive) { AddLineTriangleTWS(P0 + ShaderPrint.TranslatedWorldOffset, P1 + ShaderPrint.TranslatedWorldOffset, P2 + ShaderPrint.TranslatedWorldOffset, Color); } }
 
 ///////////////////////////////////////////////////////////////////////////////////////////////////
 // OBB
@@ -938,5 +971,9 @@ void AddOBBTWS(float3 Min, float3 Max, float4 Color, float4x4 Transform)
 	}
 }
 
-void AddOBBWS(float3 Min, float3 Max, float4 Color, float4x4 Transform) { Transform[3].xyz += ShaderPrint.TranslatedWorldOffset; AddOBBTWS(Min, Max, Color, Transform); }
-void AddOBB(float3 Min, float3 Max, float4 Color, float4x4 Transform)	{ Transform[3].xyz += ShaderPrint.TranslatedWorldOffset; AddOBBTWS(Min, Max, Color, Transform); }
+void AddOBBWS(float3 Min, float3 Max, float4 Color, float4x4 Transform)								{ Transform[3].xyz += ShaderPrint.TranslatedWorldOffset; AddOBBTWS(Min, Max, Color, Transform); }
+void AddOBB  (float3 Min, float3 Max, float4 Color, float4x4 Transform)								{ Transform[3].xyz += ShaderPrint.TranslatedWorldOffset; AddOBBTWS(Min, Max, Color, Transform); }
+
+void AddOBBTWS(FShaderPrintContext Ctx, float3 Min, float3 Max, float4 Color, float4x4 Transform)	{ if (Ctx.bIsActive) { AddOBBTWS(Min, Max, Color, Transform); } }
+void AddOBBWS (FShaderPrintContext Ctx, float3 Min, float3 Max, float4 Color, float4x4 Transform)	{ if (Ctx.bIsActive) { Transform[3].xyz += ShaderPrint.TranslatedWorldOffset; AddOBBTWS(Min, Max, Color, Transform); } }
+void AddOBB   (FShaderPrintContext Ctx, float3 Min, float3 Max, float4 Color, float4x4 Transform)	{ if (Ctx.bIsActive) { Transform[3].xyz += ShaderPrint.TranslatedWorldOffset; AddOBBTWS(Min, Max, Color, Transform); } }
