Remove unused code from DistanceFieldLightingShared.ush

diff --git a/Engine/Shaders/Private/DistanceFieldLightingShared.ush b/Engine/Shaders/Private/DistanceFieldLightingShared.ush
--- a/Engine/Shaders/Private/DistanceFieldLightingShared.ush
+++ b/Engine/Shaders/Private/DistanceFieldLightingShared.ush
@@ -317,30 +317,6 @@ FHeightfieldObjectData LoadHeightfieldObjectData(uint ObjectIndex)
 	return LoadHeightfieldObjectDataFromBuffer(SceneHeightfieldObjectData, ObjectIndex);
 }
 
-// VPLs generated by raytracing from the light
-#define VPL_DATA_STRIDE 3
-
-#define FINAL_GATHER_THREADGROUP_SIZE 64
-
-// Must match C++
-#define NUM_VISIBILITY_STEPS 10
-
-// Must match C++
-#define RECORD_CONE_DATA_STRIDE NUM_VISIBILITY_STEPS
-
-// Must match C++
-#define BENT_NORMAL_STRIDE 1
-
-Buffer<float4> IrradianceCachePositionRadius;
-Buffer<float> IrradianceCacheOccluderRadius;
-Buffer<uint2> IrradianceCacheTileCoordinate;
-Buffer<float4> IrradianceCacheNormal;
-Buffer<float4> IrradianceCacheBentNormal;
-Buffer<float4> IrradianceCacheIrradiance;
-
-Buffer<uint> ScatterDrawParameters;
-Buffer<uint> SavedStartIndex;
-
 uint NumConvexHullPlanes;
 float4 ViewFrustumConvexHull[6];
 
@@ -378,31 +354,6 @@ bool ViewFrustumIntersectAABB(float3 BoxCenter, float3 BoxExtent)
 	return true;
 }
 
-float ApproximateConeConeIntersection(float ArcLength0, float ArcLength1, float AngleBetweenCones)
-{
-	float AngleDifference = abs(ArcLength0 - ArcLength1);
-
-	float Intersection = smoothstep(
-		0,
-		1.0,
-		1.0 - saturate((AngleBetweenCones - AngleDifference) / (ArcLength0 + ArcLength1 - AngleDifference)));
-
-	return Intersection;
-}
-
-float GetVPLOcclusion(float3 BentNormalAO, float3 NormalizedVectorToVPL, float VPLConeAngle, float VPLOcclusionStrength)
-{
-	float BentNormalLength = length(BentNormalAO);
-	float UnoccludedAngle = BentNormalLength * PI / VPLOcclusionStrength;
-	float AngleBetween = acos(dot(BentNormalAO, NormalizedVectorToVPL) / max(BentNormalLength, .0001f));
-	float Visibility = ApproximateConeConeIntersection(VPLConeAngle, UnoccludedAngle, AngleBetween);
-
-	// Can't rely on the direction of the bent normal when close to fully occluded, lerp to shadowed
-	Visibility = lerp(0, Visibility, saturate((UnoccludedAngle - .1f) / .2f));
-
-	return Visibility;
-}
-
 float SampleDistanceFieldBrickTexture(float3 UV)
 {
 #if SUPPORTS_INDEPENDENT_SAMPLERS
