Fix for decal remaining differences.
Also only have a single BSDF evaluated for now because that is enough for decals.

diff --git a/Engine/Shaders/Private/DBufferDecalShared.ush b/Engine/Shaders/Private/DBufferDecalShared.ush
--- a/Engine/Shaders/Private/DBufferDecalShared.ush
+++ b/Engine/Shaders/Private/DBufferDecalShared.ush
@@ -60,7 +60,7 @@ void ConvertToMetalness(
 	inout float Metallic)
 {
 	Metallic  = F0RGBToMetallic(F0);
-	Specular  = F0ToDielectricSpecular(F0.x);
+	Specular  = F0RGBToDielectricSpecular(F0);
 	BaseColor = lerp(DiffuseAlbedo, F0, Metallic);
 }
 
diff --git a/Engine/Shaders/Private/ShadingCommon.ush b/Engine/Shaders/Private/ShadingCommon.ush
--- a/Engine/Shaders/Private/ShadingCommon.ush
+++ b/Engine/Shaders/Private/ShadingCommon.ush
@@ -102,6 +102,11 @@ float F0ToDielectricSpecular(float F0)
 	return saturate(F0 / 0.08f);
 }
 
+float F0RGBToDielectricSpecular(float3 F0)
+{
+	return F0ToDielectricSpecular(max(F0.r, max(F0.g, F0.b)));
+}
+
 float DielectricSpecularToF0(float Specular)
 {
 	return 0.08f * Specular;
@@ -120,12 +125,17 @@ float DielectricIorToF0(float Ior)
 	return F0;
 }
 
-float ComputeSpecularMicroOcclusion(float F0)
+float F0ToMicroOcclusion(float F0)
 {
 	// Anything with Specular less than 2% is physically impossible and is instead considered to be shadowing.
 	return saturate(50.0 * F0);
 }
 
+float F0RGBToMicroOcclusion(float3 F0)
+{
+	return F0ToMicroOcclusion(max(F0.r, max(F0.g, F0.b)));
+}
+
 float3 ComputeF0(float Specular, float3 BaseColor, float Metallic)
 {
 	return lerp(DielectricSpecularToF0(Specular).xxx, BaseColor, Metallic.xxx);
diff --git a/Engine/Shaders/Private/Strata/Strata.ush b/Engine/Shaders/Private/Strata/Strata.ush
--- a/Engine/Shaders/Private/Strata/Strata.ush
+++ b/Engine/Shaders/Private/Strata/Strata.ush
@@ -2413,10 +2413,9 @@ FStrataDBuffer StrataConvertToDBuffer(
 	const float3 DiffuseColor = SLAB_DIFFUSEALBEDO(BSDF);
 	const float3 F0 = SLAB_F0(BSDF);
 
-	// STRATA_TODO : Can we avoid this by having a lerpable packing?
 	Out.Roughness = SLAB_ROUGHNESS(BSDF);
 	Out.Metallic = F0RGBToMetallic(F0);
-	Out.Specular = F0ToDielectricSpecular(F0.x);
+	Out.Specular = F0RGBToDielectricSpecular(F0);
 	Out.BaseColor = lerp(DiffuseColor, F0, Out.Metallic);
 	Out.Emissive = BSDF_GETEMISSIVE(BSDF);
 	Out.WorldNormal = InHeader.SharedLocalBases.Normals[BSDF_GETSHAREDLOCALBASISID(BSDF)];
@@ -2447,7 +2446,6 @@ void StrataConvertFromDBuffer(
 	OutHeader.State = 0;
 	AddShadingModels(OutHeader, StrataShadingModelBit(STRATA_BSDF_TYPE_SLAB));
 
-	// Data
 	// Decal only support a single Slab BSDF with no-anisotropy
 	const FStrataPixelFootprint PixelFootprint = (FStrataPixelFootprint)0;
 	const float DefaultThickness = STRATA_LAYER_DEFAULT_THICKNESS_CM;
diff --git a/Engine/Shaders/Private/Strata/StrataEvaluation.ush b/Engine/Shaders/Private/Strata/StrataEvaluation.ush
--- a/Engine/Shaders/Private/Strata/StrataEvaluation.ush
+++ b/Engine/Shaders/Private/Strata/StrataEvaluation.ush
@@ -435,7 +435,7 @@ FStrataEvaluateResult StrataEvaluateBSDFCommon(FStrataBSDFContext BSDFContext, F
 			}
 
 			float3 FresnelTerm = 0;
-			const float SpecularMicroOcclusion = ComputeSpecularMicroOcclusion(F0.g);
+			const float SpecularMicroOcclusion = F0RGBToMicroOcclusion(F0);
 			#if STRATA_FASTPATH==0
 			const bool bHasThinFilm = BSDF_GETHASTHINFILM(BSDFContext.BSDF);
 			BRANCH
@@ -945,7 +945,7 @@ FStrataEnvLightResult StrataEvaluateForEnvLight(FStrataBSDFContext BSDFContext,
 			BRANCH
 			if (bEnableSpecular)
 			{
-				const float SpecularMicroOcclusion = ComputeSpecularMicroOcclusion(F0.g);
+				const float SpecularMicroOcclusion = F0RGBToMicroOcclusion(F0);
 
 				#if STRATA_FASTPATH==0
 				const bool bHasThinFilm = BSDF_GETHASTHINFILM(BSDFContext.BSDF);
diff --git a/Engine/Shaders/Private/Strata/StrataLegacyConversion.ush b/Engine/Shaders/Private/Strata/StrataLegacyConversion.ush
--- a/Engine/Shaders/Private/Strata/StrataLegacyConversion.ush
+++ b/Engine/Shaders/Private/Strata/StrataLegacyConversion.ush
@@ -508,10 +508,15 @@ FStrataData StrataConvertLegacyMaterialDynamic(
 	Out.PreMultipliedAlphaOverrideCoverage = Opacity;
 	#endif
 
-	// We also evaluated the inlined BSDF in case it is used (for instance, for by decals)
-	const float DummyNoV = 1.0f;
-	uint NewNormalIndex = BSDF_GETSHAREDLOCALBASISID(Top_Slab.InlinedBSDF);
-	FStrataData ParamBlendVertical = StrataVerticalLayeringParameterBlending(StrataWeightParameterBlending(Top_Slab, Top.Weight), StrataWeightParameterBlending(Bot_Slab, Bot.Weight), NewNormalIndex, DummyNoV);
+	// We also evaluate the inlined BSDF resulting from parameter blending in case it is used later.
+	// For instance, this is used for deferred decals. We only evaluate the top material as this is enough for decals.
+	FStrataData ParamBlendVertical = StrataWeightParameterBlending(Top_Slab, Top.Weight);
+//	if(Bot.Weight > 0.0)
+//	{
+//		const float DummyNoV = 1.0f;
+//		uint NewNormalIndex = BSDF_GETSHAREDLOCALBASISID(Top_Slab.InlinedBSDF);
+//		ParamBlendVertical = StrataVerticalLayeringParameterBlending(ParamBlendVertical, StrataWeightParameterBlending(Bot_Slab, Bot.Weight), NewNormalIndex, DummyNoV);
+//	}
 #if STRATA_LEGACY_MATERIAL_APPLIES_FINAL_WEIGHT
 	Out.InlinedBSDF = StrataWeightParameterBlending(ParamBlendVertical, FinalWeight).InlinedBSDF;
 #else
