DistanceFieldObjectCulling in Translated World Space.

diff --git a/Engine/Shaders/Private/DistanceFieldObjectCulling.usf b/Engine/Shaders/Private/DistanceFieldObjectCulling.usf
--- a/Engine/Shaders/Private/DistanceFieldObjectCulling.usf
+++ b/Engine/Shaders/Private/DistanceFieldObjectCulling.usf
@@ -52,7 +52,7 @@ void CullObjectsForViewCS(
 		float DistanceToViewSq = dot(CenterToCamera, CenterToCamera);
 
 		if (DistanceToViewSq < Square(AOMaxViewDistance + DFObjectBounds.SphereRadius)
-			&& ViewFrustumIntersectSphere(LWCHackToFloat(DFObjectBounds.Center), DFObjectBounds.SphereRadius + AOObjectMaxDistance))
+			&& ViewFrustumIntersectSphere(TranslatedCenter, DFObjectBounds.SphereRadius + AOObjectMaxDistance))
 		{
 			FDFObjectData DFObjectData = LoadDFObjectData(SourceIndex);
 
@@ -232,7 +232,7 @@ void BuildTileConesMain(
 
 struct FObjectCullVertexOutput
 {
-	nointerpolation float4 PositionAndRadius : TEXCOORD0;
+	nointerpolation float4 TranslatedPositionAndRadius : TEXCOORD0;
 	nointerpolation uint2 ObjectIndexInstanceIndex : TEXCOORD1;
 };
  
@@ -255,8 +255,8 @@ void ObjectCullVS(
 	float EffectiveRadius = (ObjectBounds.SphereRadius + AOObjectMaxDistance) * ConservativeRadiusScale;
 	float3 TranslatedWorldPosition = InPosition.xyz * EffectiveRadius + TranslatedCenter;
 	OutPosition = mul(float4(TranslatedWorldPosition, 1), PrimaryView.TranslatedWorldToClip);
-	Output.PositionAndRadius.xyz = LWCHackToFloat(ObjectBounds.Center);
-	Output.PositionAndRadius.w = ObjectBounds.SphereRadius;
+	Output.TranslatedPositionAndRadius.xyz = TranslatedCenter;
+	Output.TranslatedPositionAndRadius.w = ObjectBounds.SphereRadius;
 	Output.ObjectIndexInstanceIndex = uint2(ObjectIndex, InstanceIndex);
 } 
 
@@ -276,16 +276,16 @@ bool IntersectObjectWithConeDepthRange(
 #define USE_DISTANCE_FIELD_FOR_OBJECT_CULLING 1
 #if USE_DISTANCE_FIELD_FOR_OBJECT_CULLING
 		FDFObjectData DFObjectData = LoadDFObjectData(ObjectIndex);
-		float4x4 WorldToVolume = LWCHackToFloat(DFObjectData.WorldToVolume);
+		float4x4 TranslatedWorldToVolume = LWCMultiplyTranslation(LWCNegate(PrimaryView.PreViewTranslation), DFObjectData.WorldToVolume);
 
 		// Use the position halfway between the depth ranges as the center for the bounding sphere of this tile depth range
 		float3 ViewTileBoundingSphereCenter = TileConeVertex + TileConeAxis * (.5f * (ConeDepthRange.x + ConeDepthRange.y));
-		float3 WorldTileBoundingSphereCenter = mul(float4(ViewTileBoundingSphereCenter.xyz, 1), View.ViewToTranslatedWorld).xyz - LWCHackToFloat(PrimaryView.PreViewTranslation);
+		float3 TranslatedWorldTileBoundingSphereCenter = mul(float4(ViewTileBoundingSphereCenter.xyz, 1), View.ViewToTranslatedWorld).xyz;
 		float DistanceAlongAxis = .5f * (ConeDepthRange.y - ConeDepthRange.x);
 		float FarDepthDistanceToEdgeOfCone = ConeDepthRange.y * TileConeAngleSin / TileConeAngleCos;
 		float TileBoundingSphereRadius = sqrt(DistanceAlongAxis * DistanceAlongAxis + FarDepthDistanceToEdgeOfCone * FarDepthDistanceToEdgeOfCone);
 
-		float3 VolumeTileBoundingSphereCenter = mul(float4(WorldTileBoundingSphereCenter, 1), WorldToVolume).xyz;
+		float3 VolumeTileBoundingSphereCenter = mul(float4(TranslatedWorldTileBoundingSphereCenter, 1), TranslatedWorldToVolume).xyz;
 		float BoxDistance = ComputeDistanceFromBoxToPoint(-DFObjectData.VolumePositionExtent, DFObjectData.VolumePositionExtent, VolumeTileBoundingSphereCenter) * DFObjectData.VolumeScale;
 
 		BRANCH
@@ -334,13 +334,14 @@ void ObjectCullPS(
 	float TileConeAngleCos = ConeAxisAndCos.w;
 	float TileConeAngleSin = sqrt(1 - TileConeAngleCos * TileConeAngleCos);
 
-	float4 WorldSphereCenterAndRadius = Input.PositionAndRadius;
-	float3 ViewSpaceSphereCenter = mul(float4(WorldSphereCenterAndRadius.xyz + LWCHackToFloat(PrimaryView.PreViewTranslation), 1), View.TranslatedWorldToView).xyz;
+	float3 TranslatedWorldSphereCenter = Input.TranslatedPositionAndRadius.xyz;
+	float SphereRadius = Input.TranslatedPositionAndRadius.w;
+	float3 ViewSpaceSphereCenter = mul(float4(TranslatedWorldSphereCenter, 1), View.TranslatedWorldToView).xyz;
 	
 	// A value of 1 is conservative, but has a huge impact on performance
 	float RadiusScale = .5f;
 
-	float4 SphereCenterAndRadius = float4(ViewSpaceSphereCenter, WorldSphereCenterAndRadius.w + RadiusScale * AOObjectMaxDistance);
+	float4 SphereCenterAndRadius = float4(ViewSpaceSphereCenter, SphereRadius + RadiusScale * AOObjectMaxDistance);
 
 	if (SphereIntersectCone(SphereCenterAndRadius, TileConeVertex, TileConeAxis, TileConeAngleCos, TileConeAngleSin))
 	{
diff --git a/Engine/Source/Runtime/Renderer/Private/DistanceFieldObjectCulling.cpp b/Engine/Source/Runtime/Renderer/Private/DistanceFieldObjectCulling.cpp
--- a/Engine/Source/Runtime/Renderer/Private/DistanceFieldObjectCulling.cpp
+++ b/Engine/Source/Runtime/Renderer/Private/DistanceFieldObjectCulling.cpp
@@ -84,7 +84,8 @@ void CullObjectsToView(FRDGBuilder& GraphBuilder, FScene* Scene, const FViewInfo
 
 		for (int32 i = 0; i < View.ViewFrustum.Planes.Num(); i++)
 		{
-			PassParameters->ViewFrustumConvexHull[i] = FVector4f((FVector3f)View.ViewFrustum.Planes[i], View.ViewFrustum.Planes[i].W);
+			const FPlane4f Plane(View.ViewFrustum.Planes[i].TranslateBy(View.ViewMatrices.GetPreViewTranslation()));
+			PassParameters->ViewFrustumConvexHull[i] = Plane;
 		}
 
 		PassParameters->ObjectBoundingGeometryIndexCount = StencilingGeometry::GLowPolyStencilSphereIndexBuffer.GetIndexCount();
