DistanceFieldVisualization in Translated World Space

diff --git a/Engine/Shaders/Private/DistanceFieldVisualization.usf b/Engine/Shaders/Private/DistanceFieldVisualization.usf
--- a/Engine/Shaders/Private/DistanceFieldVisualization.usf
+++ b/Engine/Shaders/Private/DistanceFieldVisualization.usf
@@ -21,7 +21,7 @@ groupshared uint NumIntersectingObjects;
 
 #define SDF_TRACING_TRAVERSE_MIPS 1
 
-void RayTraceThroughTileCulledDistanceFields(float3 WorldRayStart, float3 WorldRayDirection, float MaxRayTime, out float MinRayTime, out float TotalStepsTaken, out uint HitCulledObjectIndex)
+void RayTraceThroughTileCulledDistanceFields(float3 TranslatedWorldRayStart, float3 WorldRayDirection, float MaxRayTime, out float MinRayTime, out float TotalStepsTaken, out uint HitCulledObjectIndex)
 {
 	MinRayTime = MaxRayTime;
 	TotalStepsTaken = 0;
@@ -32,11 +32,11 @@ void RayTraceThroughTileCulledDistanceFields(float3 WorldRayStart, float3 WorldR
 	{
 		uint ObjectIndex = IntersectingObjectIndices[ListObjectIndex];
 		FDFObjectData DFObjectData = LoadDFObjectData(ObjectIndex);
-		float4x4 WorldToVolume = LWCHackToFloat(DFObjectData.WorldToVolume);
+		float4x4 TranslatedWorldToVolume = LWCMultiplyTranslation(LWCNegate(PrimaryView.PreViewTranslation), DFObjectData.WorldToVolume);
 
-		float3 VolumeRayStart = mul(float4(WorldRayStart, 1), WorldToVolume).xyz;
-		float3 WorldRayEnd = WorldRayStart + WorldRayDirection * MaxRayTime;
-		float3 VolumeRayEnd = mul(float4(WorldRayEnd, 1), WorldToVolume).xyz;
+		float3 VolumeRayStart = mul(float4(TranslatedWorldRayStart, 1.0f), TranslatedWorldToVolume).xyz;
+		float3 TranslatedWorldRayEnd = TranslatedWorldRayStart + WorldRayDirection * MaxRayTime;
+		float3 VolumeRayEnd = mul(float4(TranslatedWorldRayEnd, 1.0f), TranslatedWorldToVolume).xyz;
 		float3 VolumeRayDirection = VolumeRayEnd - VolumeRayStart;
 		float VolumeRayLength = length(VolumeRayDirection);
 		VolumeRayDirection /= VolumeRayLength;
@@ -140,6 +140,7 @@ void VisualizeMeshDistanceFieldCS(
 	float3 TranslatedWorldPosition = mul(float4(ScreenPosition * SceneDepth, SceneDepth, 1), View.ScreenToTranslatedWorld).xyz;
 
 	float TraceDistance = 100000.0f;
+	float3 TranslatedWorldRayStart = PrimaryView.TranslatedWorldCameraOrigin;
 	float3 WorldRayStart = LWCHackToFloat(PrimaryView.WorldCameraOrigin);
 	float3 WorldRayDirection = normalize(TranslatedWorldPosition - View.TranslatedWorldCameraOrigin);
 	float3 WorldHitNormal = float3(0, 0, 1);
@@ -175,7 +176,7 @@ void VisualizeMeshDistanceFieldCS(
 
 	GroupMemoryBarrierWithGroupSync();
 
-	float3 TileConeVertex;
+	float3 TileConeTranslatedVertex;
 	float3 TileConeAxis;
 	float TileConeAngleCos;
 	float TileConeAngleSin;
@@ -191,7 +192,7 @@ void VisualizeMeshDistanceFieldCS(
 		TileConeAxis = mul(ViewSpaceTileConeAxis, (float3x3)View.ViewToTranslatedWorld);
 		TileConeAngleCos = min(min(dot(ViewSpaceTileConeAxis, TileCorner00), dot(ViewSpaceTileConeAxis, TileCorner10)), min(dot(ViewSpaceTileConeAxis, TileCorner01), dot(ViewSpaceTileConeAxis, TileCorner11)));
 		TileConeAngleSin = sqrt(1 - TileConeAngleCos * TileConeAngleCos);
-		TileConeVertex = LWCHackToFloat(PrimaryView.WorldCameraOrigin);
+		TileConeTranslatedVertex = PrimaryView.TranslatedWorldCameraOrigin;
 	}
 
 	uint NumCulledObjects = GetCulledNumObjects();
@@ -201,10 +202,10 @@ void VisualizeMeshDistanceFieldCS(
 	{
 		const uint ObjectIndex = CulledObjectIndices[IndexInCulledList];
 		FDFObjectBounds ObjectBounds = LoadDFObjectBounds(ObjectIndex);
-		float4 SphereCenterAndRadius = float4(LWCHackToFloat(ObjectBounds.Center), ObjectBounds.SphereRadius);
+		float4 SphereTranslatedCenterAndRadius = float4(LWCToFloat(LWCAdd(ObjectBounds.Center, PrimaryView.PreViewTranslation)), ObjectBounds.SphereRadius);
 
 		BRANCH
-		if (SphereIntersectCone(SphereCenterAndRadius, TileConeVertex, TileConeAxis, TileConeAngleCos, TileConeAngleSin))
+		if (SphereIntersectCone(SphereTranslatedCenterAndRadius, TileConeTranslatedVertex, TileConeAxis, TileConeAngleCos, TileConeAngleSin))
 		{
 			uint ListIndex;
 			InterlockedAdd(NumIntersectingObjects, 1U, ListIndex);
@@ -223,15 +224,15 @@ void VisualizeMeshDistanceFieldCS(
 	uint HitCulledObjectIndex;
 
 	// Trace once to find the distance to first intersection
-	RayTraceThroughTileCulledDistanceFields(WorldRayStart, WorldRayDirection, TraceDistance, MinRayTime, TotalStepsTaken, HitCulledObjectIndex);
+	RayTraceThroughTileCulledDistanceFields(TranslatedWorldRayStart, WorldRayDirection, TraceDistance, MinRayTime, TotalStepsTaken, HitCulledObjectIndex);
 
 	if (HitCulledObjectIndex != 0xFFFFFFFF)
 	{
 		FDFObjectData DFObjectData = LoadDFObjectData(HitCulledObjectIndex);
-		float4x4 WorldToVolume = LWCHackToFloat(DFObjectData.WorldToVolume);
+		float4x4 TranslatedWorldToVolume = LWCMultiplyTranslation(LWCNegate(PrimaryView.PreViewTranslation), DFObjectData.WorldToVolume);
 
-		float3 WorldRayEnd = WorldRayStart + WorldRayDirection * MinRayTime;
-		float3 VolumeHitPosition = mul(float4(WorldRayEnd, 1), WorldToVolume).xyz;
+		float3 TranslatedWorldRayEnd = TranslatedWorldRayStart + WorldRayDirection * MinRayTime;
+		float3 VolumeHitPosition = mul(float4(TranslatedWorldRayEnd, 1.0f), TranslatedWorldToVolume).xyz;
 		//VolumeHitPosition = clamp(VolumeHitPosition, -DFObjectData.VolumePositionExtent, DFObjectData.VolumePositionExtent);
 		uint NumMips = LoadDFAssetData(DFObjectData.AssetIndex, 0).NumMips;
 		FDFAssetData DFAssetData = LoadDFAssetData(DFObjectData.AssetIndex, NumMips - 1);
