GeometryAwareUpsample in translated world space.

diff --git a/Engine/Shaders/Private/DistanceFieldLightingPost.usf b/Engine/Shaders/Private/DistanceFieldLightingPost.usf
--- a/Engine/Shaders/Private/DistanceFieldLightingPost.usf
+++ b/Engine/Shaders/Private/DistanceFieldLightingPost.usf
@@ -77,9 +77,9 @@ float4 GetNormalWeights(float2 Corner00UV, float2 LowResTexelSize, float3 WorldN
 float ComputeSampleWeightBasedOnPosition(float4 ReferencePlane, float2 SampleScreenUV, float SampleDepth)
 {
 	float2 SampleScreenPosition = (SampleScreenUV.xy - View.ScreenPositionScaleBias.wz) / View.ScreenPositionScaleBias.xy;
-	float3 SampleWorldPosition = mul(float4(SampleScreenPosition * SampleDepth, SampleDepth, 1), LWCHackToFloat(PrimaryView.ScreenToWorld)).xyz;
+	float3 SampleTranslatedWorldPosition = mul(float4(SampleScreenPosition * SampleDepth, SampleDepth, 1), PrimaryView.ScreenToTranslatedWorld).xyz;
 
-	float PlaneDistance = dot(ReferencePlane, float4(SampleWorldPosition, 1));
+	float PlaneDistance = dot(ReferencePlane, float4(SampleTranslatedWorldPosition, 1));
 
 	float Epsilon = .0001f;
 	float RelativeDistance = 1000 * abs(PlaneDistance) / SampleDepth;
@@ -92,9 +92,9 @@ void GeometryAwareUpsample(float4 UVAndScreenPos, out float4 OutBentNormal)
 	float SceneDepth;
 	GetDownsampledGBuffer(UVAndScreenPos.xy, WorldNormal, SceneDepth);
 
-	float3 WorldPosition = mul(float4(UVAndScreenPos.zw * SceneDepth, SceneDepth, 1), LWCHackToFloat(PrimaryView.ScreenToWorld)).xyz;
+	float3 TranslatedWorldPosition = mul(float4(UVAndScreenPos.zw * SceneDepth, SceneDepth, 1), PrimaryView.ScreenToTranslatedWorld).xyz;
 
-	float4 ReferencePlane = float4(WorldNormal, -dot(WorldPosition, WorldNormal));
+	float4 ReferencePlane = float4(WorldNormal, -dot(TranslatedWorldPosition, WorldNormal));
 
 	float2 LowResBufferSize = BentNormalBufferAndTexelSize.xy;
 	float2 LowResTexelSize = BentNormalBufferAndTexelSize.zw;
