Object based DFAO in translated world space.

diff --git a/Engine/Shaders/Private/DistanceFieldScreenGridLighting.usf b/Engine/Shaders/Private/DistanceFieldScreenGridLighting.usf
--- a/Engine/Shaders/Private/DistanceFieldScreenGridLighting.usf
+++ b/Engine/Shaders/Private/DistanceFieldScreenGridLighting.usf
@@ -62,7 +62,7 @@ float TanConeHalfAngle;
 RWBuffer<uint> RWScreenGridConeVisibility;
 RWBuffer<float> RWConeDepthVisibilityFunction;
 
-void HemisphereConeTraceAgainstTileCulledObjects(uint ObjectIndex, uint OutputBaseIndex, float3 WorldShadingPosition, float SceneDepth, float3 WorldNormal, float3 TangentX, float3 TangentY)
+void HemisphereConeTraceAgainstTileCulledObjects(uint ObjectIndex, uint OutputBaseIndex, FLWCVector3 WorldShadingPosition, float SceneDepth, float3 WorldNormal, float3 TangentX, float3 TangentY)
 {
 	float MaxWorldStepOffset = GetStepOffset(NUM_CONE_STEPS);
 	float InvMaxOcclusionDistance = 1.0f / AOObjectMaxDistance;
@@ -73,7 +73,7 @@ void HemisphereConeTraceAgainstTileCulledObjects(uint ObjectIndex, uint OutputBa
 
 	FDFObjectBounds ObjectBounds = LoadDFObjectBounds(ObjectIndex);
 
-	const float3 CenterToCamera = LWCHackToFloat(ObjectBounds.Center) - WorldShadingPosition;
+	const float3 CenterToCamera = LWCToFloat(LWCSubtract(ObjectBounds.Center, WorldShadingPosition));
 	float ObjectDistanceSq = dot(CenterToCamera, CenterToCamera);
 
 	BRANCH
@@ -81,9 +81,8 @@ void HemisphereConeTraceAgainstTileCulledObjects(uint ObjectIndex, uint OutputBa
 	if (ObjectBounds.SphereRadius * ObjectBounds.SphereRadius / ObjectDistanceSq > Square(.25f))
 	{
 		FDFObjectData DFObjectData = LoadDFObjectData(ObjectIndex);
-		float4x4 WorldToVolume = LWCHackToFloat(DFObjectData.WorldToVolume);
 
-		float3 VolumeShadingPosition = mul(float4(WorldShadingPosition, 1), WorldToVolume).xyz;
+		float3 VolumeShadingPosition = LWCMultiply(WorldShadingPosition, DFObjectData.WorldToVolume);
 		float BoxDistance = ComputeDistanceFromBoxToPoint(-DFObjectData.VolumePositionExtent, DFObjectData.VolumePositionExtent, VolumeShadingPosition) * DFObjectData.VolumeScale;
 
 		BRANCH
@@ -222,7 +221,7 @@ void ConeTraceObjectOcclusionCS(
 			float2 ScreenUV = GetScreenUVFromScreenGrid(OutputCoordinate);
 			float2 ScreenPosition = (ScreenUV.xy - View.ScreenPositionScaleBias.wz) / View.ScreenPositionScaleBias.xy;
 		
-			float3 WorldShadingPosition = mul(float4(ScreenPosition * SceneDepth, SceneDepth, 1), LWCHackToFloat(PrimaryView.ScreenToWorld)).xyz;
+			FLWCVector3 WorldShadingPosition = LWCMultiply(float3(ScreenPosition * SceneDepth, SceneDepth), PrimaryView.ScreenToWorld);
 
 			HemisphereConeTraceAgainstTileCulledObjects(ObjectIndex, OutputBaseIndex, WorldShadingPosition, SceneDepth, WorldNormal, TangentX, TangentY);
 		}
