Fix SSS flickering and black out when viewport size is less than 100%. The fix is to minimize texture clear. Clamp texels conservatively:
* level 1: 4-2 texel.
* level 2: 8-2 texel.
* level n: 2^(1+n)-2 texel.

diff --git a/Engine/Shaders/Private/SubsurfaceBurleyNormalized.ush b/Engine/Shaders/Private/SubsurfaceBurleyNormalized.ush
--- a/Engine/Shaders/Private/SubsurfaceBurleyNormalized.ush
+++ b/Engine/Shaders/Private/SubsurfaceBurleyNormalized.ush
@@ -716,7 +716,7 @@ float3 GetBilateralNormal(float2 ClampedUV, float2 Extent)
 float2 ClampUVLevel(float2 UV, float2 MinUV, float2 MaxUV, float MipLevel)
 {
 	uint CeilMipLevel = (uint)(ceil(MipLevel));
-	float2 MipUVCorrection = float( (1u << CeilMipLevel) - 1u ) * SubsurfaceInput0_ExtentInverse;
+	float2 MipUVCorrection = float((1u << (1u + CeilMipLevel)) - 2u) * SubsurfaceInput0_ExtentInverse;
 	
 	float2 MinMipUVCorrection = 
 		float2(SubsurfaceInput0_UVViewportMin.x == 0.0f ? 0.0f : MipUVCorrection.x,
