Strata - compilation fix when r.forwardshading=1

diff --git a/Engine/Shaders/Private/BasePassPixelShader.usf b/Engine/Shaders/Private/BasePassPixelShader.usf
--- a/Engine/Shaders/Private/BasePassPixelShader.usf
+++ b/Engine/Shaders/Private/BasePassPixelShader.usf
@@ -1584,7 +1584,7 @@ void FPixelShaderInOut_MainPS(
 	#endif
 #endif
 
-#if MATERIAL_WORKS_WITH_DUAL_SOURCE_COLOR_BLENDING || STRATA_TRANSLUCENT_MATERIAL
+#if MATERIAL_WORKS_WITH_DUAL_SOURCE_COLOR_BLENDING || STRATA_TRANSLUCENT_MATERIAL || STRATA_FORWARD_SHADING
 	float3 DualBlendSurfaceLuminancePostCoverage	= 0.0f;
 	float3 DualBlendSurfaceTransmittancePreCoverage = 1.0f;
 	float DualBlendSurfaceCoverage					= 1.0f;
