Fixed base pass pixel shader to use strata resources from the correct uniform buffer to avoid all sort of valid asserts and crashes.

diff --git a/Engine/Shaders/Private/BasePassPixelShader.usf b/Engine/Shaders/Private/BasePassPixelShader.usf
--- a/Engine/Shaders/Private/BasePassPixelShader.usf
+++ b/Engine/Shaders/Private/BasePassPixelShader.usf
@@ -1711,11 +1711,11 @@ void FPixelShaderInOut_MainPS(
 	FStrataTopLayerData TopLayerData = (FStrataTopLayerData)0;
 
 	// Generate the strata material data to write out
-	FStrataAddressing StrataAddressing = GetStrataPixelDataByteOffset(PixelPos, uint2(ResolvedView.BufferSizeAndInvSize.xy), OpaqueBasePass.Strata.MaxBytesPerPixel);
-	FRWStrataMaterialContainer RWStrataMaterialContainer = InitialiseRWStrataMaterialContainer(OpaqueBasePass.Strata.MaterialTextureArrayUAVWithoutRTs);
+	FStrataAddressing StrataAddressing = GetStrataPixelDataByteOffset(PixelPos, uint2(ResolvedView.BufferSizeAndInvSize.xy), StrataStruct.MaxBytesPerPixel);
+	FRWStrataMaterialContainer RWStrataMaterialContainer = InitialiseRWStrataMaterialContainer(StrataStruct.MaterialTextureArrayUAVWithoutRTs);
 	PackStrataOut(
 		RWStrataMaterialContainer,
-		OpaqueBasePass.Strata.MaterialTextureArrayUAVWithoutRTs,
+		StrataStruct.MaterialTextureArrayUAVWithoutRTs,
 		Dither,
 		StrataAddressing,
 		StrataPixelHeader, StrataData, MaterialParameters.CameraVector, bStrataSubsurfaceEnable, EmissiveLuminance,
@@ -1730,7 +1730,7 @@ void FPixelShaderInOut_MainPS(
 	BRANCH
 	if(SSSData.bIsValid)
 	{
-		OpaqueBasePass.Strata.SSSTextureUAV[PixelPos] = StrataPackSSSData(SSSData);
+		StrataStruct.SSSTextureUAV[PixelPos] = StrataPackSSSData(SSSData);
 	}
 
 	Color += EmissiveLuminance;
@@ -1749,7 +1749,7 @@ void FPixelShaderInOut_MainPS(
 		Color = StrataForwardLighting(
 			EyeIndex,
 			In.SvPosition,
-			TranslucentBasePass.Strata.bRoughDiffuse,
+			StrataStruct.bRoughDiffuse,
 			BasePassInterpolants,
 			Interpolants,
 			LightmapVTPageTableResult,
