Strata - fix for unlit  legacy when r.forwardshading=1

diff --git a/Engine/Shaders/Private/BasePassPixelShader.usf b/Engine/Shaders/Private/BasePassPixelShader.usf
--- a/Engine/Shaders/Private/BasePassPixelShader.usf
+++ b/Engine/Shaders/Private/BasePassPixelShader.usf
@@ -1781,7 +1781,7 @@ void FPixelShaderInOut_MainPS(
 	// Unlit forces a single BSDF
 	FStrataBSDF UnlitBSDF = PixelMaterialInputs.FrontMaterial.InlinedBSDF;
 
-#if STRATA_OPAQUE_DEFERRED
+#if STRATA_OPAQUE_DEFERRED || (STRATA_OPAQUE_MATERIAL && STRATA_FORWARD_SHADING)
 	Color = BSDF_GETEMISSIVE(UnlitBSDF);
 	Opacity = 1.0f;
 #else // STRATA_TRANSLUCENT_FORWARD || STRATA_FORWARD_SHADING
