Trace Indirect Shadows against Distance Field Casters in translated world space.

diff --git a/Engine/Shaders/Private/CapsuleShadowShaders.usf b/Engine/Shaders/Private/CapsuleShadowShaders.usf
--- a/Engine/Shaders/Private/CapsuleShadowShaders.usf
+++ b/Engine/Shaders/Private/CapsuleShadowShaders.usf
@@ -221,7 +221,6 @@ groupshared uint IntegerTileMaxZ2;
 struct FTileCullingData
 {
 	float4 TranslatedBoundingSphere;
-	float4 BoundingSphere; // LWC_TODO
 #if LIGHT_SOURCE_MODE == LIGHT_SOURCE_PUNCTUAL
 	float3 ConeAxis;
 	float ConeAngleCos;
@@ -370,14 +369,10 @@ void SetupTileCullingData(
 	float3 ViewGroup0Center = (ViewTileMax + ViewTileMin) / 2;
 	TileCullingData0.TranslatedBoundingSphere.xyz = mul(float4(ViewGroup0Center, 1), View.ViewToTranslatedWorld).xyz;
 	TileCullingData0.TranslatedBoundingSphere.w = length(ViewGroup0Center - ViewTileMax);
-	TileCullingData0.BoundingSphere.xyz = TileCullingData0.TranslatedBoundingSphere.xyz - LWCHackToFloat(PrimaryView.PreViewTranslation);
-	TileCullingData0.BoundingSphere.w = length(ViewGroup0Center - ViewTileMax);
 
 	float3 ViewGroup1Center = (ViewTileMax2 + ViewTileMin2) / 2;
 	TileCullingData1.TranslatedBoundingSphere.xyz = mul(float4(ViewGroup1Center, 1), View.ViewToTranslatedWorld).xyz;
 	TileCullingData1.TranslatedBoundingSphere.w = length(ViewGroup1Center - ViewTileMax);
-	TileCullingData1.BoundingSphere.xyz = TileCullingData1.TranslatedBoundingSphere.xyz - LWCHackToFloat(PrimaryView.PreViewTranslation);
-	TileCullingData1.BoundingSphere.w = length(ViewGroup1Center - ViewTileMax2);
 
 #if LIGHT_SOURCE_MODE == LIGHT_SOURCE_PUNCTUAL
 	#if POINT_LIGHT
@@ -500,9 +495,9 @@ uint CullCapsuleShapesToTile(
 		#endif
 
 		BRANCH
-		if (SphereIntersectSphere(float4(TransformedSphereTranslatedCenter, TransformedSphereRadius + MaxOcclusionDistance), float4(TransformedTileTranslatedBoundingSphereCenter0, TileCullingData0.BoundingSphere.w))
+		if (SphereIntersectSphere(float4(TransformedSphereTranslatedCenter, TransformedSphereRadius + MaxOcclusionDistance), float4(TransformedTileTranslatedBoundingSphereCenter0, TileCullingData0.TranslatedBoundingSphere.w))
 			#if LIGHT_SOURCE_MODE != LIGHT_SOURCE_FROM_RECEIVER
-				&& SphereIntersectConeWithMaxDistance(float4(TransformedSphereTranslatedCenter, TransformedSphereRadius + TileCullingData0.BoundingSphere.w), TransformedTileTranslatedBoundingSphereCenter0, TransformedConeAxis0, ConeAngleCos0, ConeAngleSin0, MaxOcclusionDistance)
+				&& SphereIntersectConeWithMaxDistance(float4(TransformedSphereTranslatedCenter, TransformedSphereRadius + TileCullingData0.TranslatedBoundingSphere.w), TransformedTileTranslatedBoundingSphereCenter0, TransformedConeAxis0, ConeAngleCos0, ConeAngleSin0, MaxOcclusionDistance)
 			#endif
 			)
 		{
@@ -514,9 +509,9 @@ uint CullCapsuleShapesToTile(
 		}
 
 		BRANCH
-		if (SphereIntersectSphere(float4(TransformedSphereTranslatedCenter, TransformedSphereRadius + MaxOcclusionDistance), float4(TransformedTileTranslatedBoundingSphereCenter1, TileCullingData1.BoundingSphere.w))
+		if (SphereIntersectSphere(float4(TransformedSphereTranslatedCenter, TransformedSphereRadius + MaxOcclusionDistance), float4(TransformedTileTranslatedBoundingSphereCenter1, TileCullingData1.TranslatedBoundingSphere.w))
 			#if LIGHT_SOURCE_MODE != LIGHT_SOURCE_FROM_RECEIVER
-				&& SphereIntersectConeWithMaxDistance(float4(TransformedSphereTranslatedCenter, TransformedSphereRadius + TileCullingData1.BoundingSphere.w), TransformedTileTranslatedBoundingSphereCenter1, TransformedConeAxis1, ConeAngleCos1, ConeAngleSin1, MaxOcclusionDistance)
+				&& SphereIntersectConeWithMaxDistance(float4(TransformedSphereTranslatedCenter, TransformedSphereRadius + TileCullingData1.TranslatedBoundingSphere.w), TransformedTileTranslatedBoundingSphereCenter1, TransformedConeAxis1, ConeAngleCos1, ConeAngleSin1, MaxOcclusionDistance)
 			#endif
 			)
 		{
@@ -709,14 +704,14 @@ float ShadowConeTraceAgainstCulledCapsuleShapes(
 bool TileBoundsIntersectDistanceFieldCaster(
 	uint ObjectIndex, 
 	FDFObjectBounds DFObjectBounds, 
-	float4 TileBoundingSphere, 
+	float4 TileTranslatedBoundingSphere, 
 	float3 TileConeAxis,
 	float TileConeAngleCos,
 	float TileConeAngleSin,
 	float MaxOcclusionDistance)
 {
 	BRANCH
-	if (SphereIntersectSphere(float4(LWCHackToFloat(DFObjectBounds.Center), DFObjectBounds.SphereRadius + MaxOcclusionDistance), TileBoundingSphere)
+	if (SphereIntersectSphere(float4(LWCToFloat(LWCAdd(DFObjectBounds.Center, PrimaryView.PreViewTranslation)), DFObjectBounds.SphereRadius + MaxOcclusionDistance), TileTranslatedBoundingSphere)
 #if LIGHT_SOURCE_MODE != LIGHT_SOURCE_FROM_RECEIVER
 		// This 'can object cast on tile bounds' test has to be disabled because the weight used to combine occlusion from multiple objects is computed based on distance to occluder only
 		//&& SphereIntersectConeWithMaxDistance(float4(SphereCenterAndRadius.xyz, SphereCenterAndRadius.w + TileBoundingSphere.w), TileBoundingSphere.xyz, TileConeAxis, TileConeAngleCos, TileConeAngleSin, MaxOcclusionDistance)
@@ -725,15 +720,15 @@ bool TileBoundsIntersectDistanceFieldCaster(
 	{
 #if USE_DISTANCE_FIELD_FOR_TILE_CULLING
 		FDFObjectData DFObjectData = LoadDFObjectData(ObjectIndex);
-		float4x4 WorldToVolume = LWCHackToFloat(DFObjectData.WorldToVolume);
-		float3 VolumeSamplePosition = mul(float4(TileBoundingSphere.xyz, 1), WorldToVolume).xyz;
+		float4x4 TranslatedWorldToVolume = LWCMultiplyTranslation(LWCNegate(PrimaryView.PreViewTranslation), DFObjectData.WorldToVolume);
+		float3 VolumeSamplePosition = mul(float4(TileTranslatedBoundingSphere.xyz, 1), TranslatedWorldToVolume).xyz;
 		float3 ClampedSamplePosition = clamp(VolumeSamplePosition, -DFObjectData.VolumePositionExtent, DFObjectData.VolumePositionExtent);
 		float DistanceToClamped = length(ClampedSamplePosition - VolumeSamplePosition);
 		float WorldDistanceToOccluder = (DistanceToMeshSurfaceStandalone(ClampedSamplePosition, DFObjectData) + DistanceToClamped) * DFObjectData.VolumeScale;
 
 		float ErrorTolerance = 1.1f;
 		// The tile can only be affected by the object's shadow if the closest part of the object is less than MaxOcclusionDistance from the tile bounds
-		return WorldDistanceToOccluder - TileBoundingSphere.w < MaxOcclusionDistance * ErrorTolerance;
+		return WorldDistanceToOccluder - TileTranslatedBoundingSphere.w < MaxOcclusionDistance * ErrorTolerance;
 #else
 		return true;
 #endif
@@ -793,7 +788,7 @@ uint CullDistanceFieldCastersToTile(
 		float EffectiveMaxOcclusionDistance = MaxOcclusionDistance + .5f * DFObjectBounds.SphereRadius;
 
 		BRANCH
-		if (TileBoundsIntersectDistanceFieldCaster(ObjectIndex, DFObjectBounds, TileCullingData0.BoundingSphere, ConeAxis0, ConeAngleCos0, ConeAngleSin0, EffectiveMaxOcclusionDistance))
+		if (TileBoundsIntersectDistanceFieldCaster(ObjectIndex, DFObjectBounds, TileCullingData0.TranslatedBoundingSphere, ConeAxis0, ConeAngleCos0, ConeAngleSin0, EffectiveMaxOcclusionDistance))
 		{
 			uint ListIndex;
 			InterlockedAdd(TileNumDistanceFields0, 1U, ListIndex);
@@ -803,7 +798,7 @@ uint CullDistanceFieldCastersToTile(
 		}
 
 		BRANCH
-		if (TileBoundsIntersectDistanceFieldCaster(ObjectIndex, DFObjectBounds, TileCullingData1.BoundingSphere, ConeAxis1, ConeAngleCos1, ConeAngleSin1, EffectiveMaxOcclusionDistance))
+		if (TileBoundsIntersectDistanceFieldCaster(ObjectIndex, DFObjectBounds, TileCullingData1.TranslatedBoundingSphere, ConeAxis1, ConeAngleCos1, ConeAngleSin1, EffectiveMaxOcclusionDistance))
 		{
 			uint ListIndex;
 			InterlockedAdd(TileNumDistanceFields1, 1U, ListIndex);
@@ -819,7 +814,7 @@ uint CullDistanceFieldCastersToTile(
 }
 
 float ShadowConeTraceAgainstCulledDistanceFieldCasters(
-	float3 WorldRayStart, 
+	float3 TranslatedWorldRayStart,
 	float3 UnitRayDirection,
 	float LightAngle,  
 	float MaxOcclusionDistance,
@@ -855,17 +850,17 @@ float ShadowConeTraceAgainstCulledDistanceFieldCasters(
 
 		FDFObjectBounds DFObjectBounds = LoadDFObjectBounds(ObjectIndex);
 		FDFObjectData DFObjectData = LoadDFObjectData(ObjectIndex);
-		float4x4 WorldToVolume = LWCHackToFloat(DFObjectData.WorldToVolume);
+		float4x4 TranslatedWorldToVolume = LWCMultiplyTranslation(LWCNegate(PrimaryView.PreViewTranslation), DFObjectData.WorldToVolume);
 
 		// Increase max occlusion distance based on object size for distance field casters
 		// This improves the solidness of the shadows, since the fadeout distance causes internal structure of objects to become visible
 		float EffectiveMaxOcclusionDistance = MaxOcclusionDistance + .5f * DFObjectBounds.SphereRadius;
 
 		float MaxSphereRadius = TanConeAngle * EffectiveMaxOcclusionDistance;
-		float3 WorldRayEnd = WorldRayStart + UnitRayDirection * EffectiveMaxOcclusionDistance;
+		float3 TranslatedWorldRayEnd = TranslatedWorldRayStart + UnitRayDirection * EffectiveMaxOcclusionDistance;
 
-		float3 VolumeRayStart = mul(float4(WorldRayStart, 1), WorldToVolume).xyz;
-		float3 VolumeRayEnd = mul(float4(WorldRayEnd, 1), WorldToVolume).xyz;
+		float3 VolumeRayStart = mul(float4(TranslatedWorldRayStart, 1), TranslatedWorldToVolume).xyz;
+		float3 VolumeRayEnd = mul(float4(TranslatedWorldRayEnd, 1), TranslatedWorldToVolume).xyz;
 		float3 VolumeRayDirection = VolumeRayEnd - VolumeRayStart;
 		float VolumeRayLength = length(VolumeRayDirection);
 		VolumeRayDirection /= VolumeRayLength;
@@ -1011,7 +1006,6 @@ void CapsuleShadowingCS(
 #endif
 	
 	const float3 OpaqueTranslatedWorldPosition = mul(float4(ScreenPosition * SceneDepth, SceneDepth, 1), PrimaryView.ScreenToTranslatedWorld).xyz;
-	const float3 OpaqueWorldPosition = OpaqueTranslatedWorldPosition - LWCHackToFloat(PrimaryView.PreViewTranslation);
 
 	uint CulledDataParameter = 0;
 	bool bTileShouldComputeShadowing = true;
@@ -1046,27 +1040,23 @@ void CapsuleShadowingCS(
 				float3 LightVector = LightTranslatedPositionAndInvRadius.xyz - OpaqueTranslatedWorldPosition;
 				LightVectorLength = length(LightVector);
 				float3 TranslatedWorldRayStart = OpaqueTranslatedWorldPosition + LightVector / LightVectorLength * RayStartOffset;
-				float3 WorldRayStart = OpaqueWorldPosition + LightVector / LightVectorLength * RayStartOffset; // LWC_TODO
 				float3 UnitRayDirection = (LightTranslatedPositionAndInvRadius.xyz - OpaqueTranslatedWorldPosition) / LightVectorLength;
 				float LightAngle = atanFastPos(LightSourceRadius / LightVectorLength);
 
 			#else
 
 				float3 TranslatedWorldRayStart = OpaqueTranslatedWorldPosition + LightDirection * RayStartOffset;
-				float3 WorldRayStart = OpaqueWorldPosition + LightDirection * RayStartOffset;
 				float3 UnitRayDirection = LightDirection;
 				float LightAngle = LightAngleAndNormalThreshold.x;
 
 			#endif
 		#elif LIGHT_SOURCE_MODE == LIGHT_SOURCE_FROM_RECEIVER
 			float3 TranslatedWorldRayStart = OpaqueTranslatedWorldPosition;
-			float3 WorldRayStart = OpaqueWorldPosition;
 			float BentNormalLength = length(ReceiverBentNormal);
 			float3 UnitRayDirection = ReceiverBentNormal / max(BentNormalLength, .00001f);
 			float LightAngle = max(BentNormalLength * .5f * PI, PI / 8);
 		#else
 			float3 TranslatedWorldRayStart = OpaqueTranslatedWorldPosition;
-			float3 WorldRayStart = OpaqueWorldPosition;
 			float3 UnitRayDirection = 0;
 			float LightAngle = 0;
 		#endif
@@ -1115,7 +1105,7 @@ void CapsuleShadowingCS(
 
 		#if SUPPORT_MESH_DISTANCE_FIELDS
 			Visibility *= ShadowConeTraceAgainstCulledDistanceFieldCasters(
-				WorldRayStart,
+				TranslatedWorldRayStart,
 				UnitRayDirection,
 				LightAngle,
 				MaxOcclusionDistance,
