Fix shader print widget.

diff --git a/Engine/Shaders/Private/ShaderPrint.ush b/Engine/Shaders/Private/ShaderPrint.ush
--- a/Engine/Shaders/Private/ShaderPrint.ush
+++ b/Engine/Shaders/Private/ShaderPrint.ush
@@ -141,6 +141,9 @@ RWStructuredBuffer<FPackedShaderPrintItem> ShaderPrint_RWValuesBuffer;
 
 ///////////////////////////////////////////////////////////////////////////////////////////////////
 // State 
+#define SHADER_PRINT_STATE_COUNT_OFFSET 0
+#define SHADER_PRINT_STATE_VALUE_OFFSET 1
+
 #define SHADER_PRINT_STATE_STRIDE 3
 #define SHADER_PRINT_STATE_INDEX_METADATA 0
 #define SHADER_PRINT_STATE_INDEX_VALUE 1
@@ -571,10 +574,10 @@ bool ShaderPrintGetStateValue(in uint InHash, inout uint OutIndex, inout uint Ou
 	OutValue = 0;
 
 	// Slow linear search among all widget states
-	const uint MaxCount = min(ShaderPrint_StateBuffer[0], ShaderPrint.MaxStateCount);
+	const uint MaxCount = min(ShaderPrint_StateBuffer[SHADER_PRINT_STATE_COUNT_OFFSET], ShaderPrint.MaxStateCount);
 	for (uint Index = 0; Index < MaxCount; ++Index)
 	{
-		const uint Index3 = Index * SHADER_PRINT_STATE_STRIDE + 1u;
+		const uint Index3 = Index * SHADER_PRINT_STATE_STRIDE + SHADER_PRINT_STATE_VALUE_OFFSET;
 		const FShaderPrintMetadata Metadata = ShaderPrintUnpackMetadata(ShaderPrint_StateBuffer[Index3 + SHADER_PRINT_STATE_INDEX_METADATA]);
 		if (Metadata.Hash == InHash)
 		{
diff --git a/Engine/Shaders/Private/ShaderPrintDraw.usf b/Engine/Shaders/Private/ShaderPrintDraw.usf
--- a/Engine/Shaders/Private/ShaderPrintDraw.usf
+++ b/Engine/Shaders/Private/ShaderPrintDraw.usf
@@ -34,8 +34,10 @@ void InitValueBufferCS()
 	// Clear .Type in the first buffer element which is (ab)used as the ShaderPrintFilterOneThread() atomic value
 
 	// Symbol/print Counter
+	RWValuesBuffer[SHADER_PRINT_COUNT_OFFSET].ScreenPos16bits = 0;
 	RWValuesBuffer[SHADER_PRINT_COUNT_OFFSET].Value = 0;
 	RWValuesBuffer[SHADER_PRINT_COUNT_OFFSET].TypeAndColor = 0;
+	RWValuesBuffer[SHADER_PRINT_COUNT_OFFSET].Metadata = 0;
 }
 
 
@@ -93,12 +95,12 @@ void AddState(in uint InValue, in uint InPackedMetadata)
 	if (Metadata.Index == SHADER_PRINT_STATE_INVALID_INDEX)
 	{
 		// Change buffer to Ruint x 2 x cound
-		InterlockedAdd(RWStateBuffer[0].x, 1, Metadata.Index);
+		InterlockedAdd(RWStateBuffer[SHADER_PRINT_STATE_COUNT_OFFSET], 1, Metadata.Index);
 	}
 
 	if (Metadata.Index < ShaderPrint.MaxStateCount)
 	{
-		const uint Index3 = Metadata.Index * SHADER_PRINT_STATE_STRIDE + 1;
+		const uint Index3 = Metadata.Index * SHADER_PRINT_STATE_STRIDE + SHADER_PRINT_STATE_VALUE_OFFSET;
 		RWStateBuffer[Index3 + SHADER_PRINT_STATE_INDEX_METADATA] = ShaderPrintPackMetadata(Metadata); // Re-encode with the new index
 		RWStateBuffer[Index3 + SHADER_PRINT_STATE_INDEX_VALUE] = InValue;
 		RWStateBuffer[Index3 + SHADER_PRINT_STATE_INDEX_FRAME] = FrameIndex;
@@ -247,18 +249,18 @@ void BuildSymbolBufferCS(uint3 DispatchThreadId : SV_DispatchThreadID)
 [numthreads(1, 1, 1)]
 void CompactStateBufferCS(uint3 DispatchThreadId : SV_DispatchThreadID)
 {
-	uint StateCount = min(RWStateBuffer[0], ShaderPrint.MaxStateCount);
+	uint StateCount = min(RWStateBuffer[SHADER_PRINT_STATE_COUNT_OFFSET], ShaderPrint.MaxStateCount);
 
 	if (StateCount > 0)
 	{
 		for (uint Index = 0; Index < StateCount;)
 		{
-			const uint Index3 = Index * SHADER_PRINT_STATE_STRIDE + 1;
+			const uint Index3 = Index * SHADER_PRINT_STATE_STRIDE + SHADER_PRINT_STATE_VALUE_OFFSET;
 			const uint FrameDiff = FrameIndex - RWStateBuffer[Index3 + SHADER_PRINT_STATE_INDEX_FRAME];
 			if (FrameDiff > FrameThreshold)
 			{
 				// Swap with last element
-				const uint SrcIndex3 = (StateCount-1) * SHADER_PRINT_STATE_STRIDE + 1;
+				const uint SrcIndex3 = (StateCount-1) * SHADER_PRINT_STATE_STRIDE + SHADER_PRINT_STATE_VALUE_OFFSET;
 				RWStateBuffer[Index3 + SHADER_PRINT_STATE_INDEX_METADATA] = RWStateBuffer[SrcIndex3 + SHADER_PRINT_STATE_INDEX_METADATA];
 				RWStateBuffer[Index3 + SHADER_PRINT_STATE_INDEX_VALUE]    = RWStateBuffer[SrcIndex3 + SHADER_PRINT_STATE_INDEX_VALUE];
 				RWStateBuffer[Index3 + SHADER_PRINT_STATE_INDEX_FRAME]    = RWStateBuffer[SrcIndex3 + SHADER_PRINT_STATE_INDEX_FRAME];	
diff --git a/Engine/Source/Runtime/Renderer/Private/ShaderPrint.cpp b/Engine/Source/Runtime/Renderer/Private/ShaderPrint.cpp
--- a/Engine/Source/Runtime/Renderer/Private/ShaderPrint.cpp
+++ b/Engine/Source/Runtime/Renderer/Private/ShaderPrint.cpp
@@ -663,7 +663,7 @@ namespace ShaderPrint
 		FRDGBufferRef IndirectDrawArgsBuffer = GraphBuilder.CreateBuffer(FRDGBufferDesc::CreateIndirectDesc(5), TEXT("ShaderPrint.IndirectDrawArgs"));
 
 		// Non graph managed resources
-		FRDGBufferSRVRef ValuesBuffer = GraphBuilder.CreateSRV(View.ShaderPrintData.ShaderPrintValueBuffer);
+		FRDGBufferSRVRef ValueBuffer = GraphBuilder.CreateSRV(View.ShaderPrintData.ShaderPrintValueBuffer);
 		FRDGBufferSRVRef StateBuffer = GraphBuilder.CreateSRV(View.ShaderPrintData.ShaderPrintStateBuffer);
 		FTextureRHIRef FontTexture = GSystemTextures.AsciiTexture->GetRenderTargetItem().ShaderResourceTexture;
 
@@ -677,7 +677,7 @@ namespace ShaderPrint
 
 			SHADER::FParameters* PassParameters = GraphBuilder.AllocParameters<SHADER::FParameters>();
 			PassParameters->Common = View.ShaderPrintData.UniformBuffer;
-			PassParameters->ValuesBuffer = ValuesBuffer;
+			PassParameters->ValuesBuffer = ValueBuffer;
 			PassParameters->RWSymbolsBuffer = GraphBuilder.CreateUAV(SymbolBuffer, EPixelFormat::PF_R32_UINT);
 			PassParameters->RWIndirectDispatchArgsBuffer = GraphBuilder.CreateUAV(IndirectDispatchArgsBuffer, EPixelFormat::PF_R32_UINT);
 
@@ -696,7 +696,7 @@ namespace ShaderPrint
 			SHADER::FParameters* PassParameters = GraphBuilder.AllocParameters<SHADER::FParameters>();
 			PassParameters->FrameIndex = View.Family ? View.Family->FrameNumber : 0u;
 			PassParameters->Common = View.ShaderPrintData.UniformBuffer;
-			PassParameters->ValuesBuffer = ValuesBuffer;
+			PassParameters->ValuesBuffer = ValueBuffer;
 			PassParameters->RWSymbolsBuffer = GraphBuilder.CreateUAV(SymbolBuffer, EPixelFormat::PF_R32_UINT);
 			PassParameters->RWStateBuffer = GraphBuilder.CreateUAV(View.ShaderPrintData.ShaderPrintStateBuffer, EPixelFormat::PF_R32_UINT);
 			PassParameters->IndirectDispatchArgsBuffer = IndirectDispatchArgsBuffer;
@@ -709,6 +709,7 @@ namespace ShaderPrint
 		}
 
 		// CompactStateBuffer
+		#if 0
 		{
 			typedef FShaderCompactStateBufferCS SHADER;
 			TShaderMapRef<SHADER> ComputeShader(GlobalShaderMap);
@@ -724,6 +725,7 @@ namespace ShaderPrint
 				RDG_EVENT_NAME("ShaderPrint::CompactStateBuffer"),
 				ComputeShader, PassParameters, FIntVector(1,1,1));
 		}
+		#endif
 
 		// BuildIndirectDrawArgs
 		{
