Remove per character atomic add when printing text string with ShaderPrint.

Now a single atomic is done for the string.

diff --git a/Engine/Shaders/Private/ShaderPrint.ush b/Engine/Shaders/Private/ShaderPrint.ush
--- a/Engine/Shaders/Private/ShaderPrint.ush
+++ b/Engine/Shaders/Private/ShaderPrint.ush
@@ -344,36 +344,50 @@ uint ShaderPrintGetOffset(FShaderPrintText InTextEntry);
 uint ShaderPrintGetHash(FShaderPrintText InTextEntry);
  
 // Function for reading global TEXT string
-float2 ShaderPrintText(float2 Pos, FShaderPrintText InTextEntry, FFontColor InColor)
+float2 ShaderPrintText_Internal(bool bIsActive, float2 InPos, FShaderPrintText InTextEntry, FFontColor InColor)
 {
-	uint Begin = ShaderPrintGetOffset(InTextEntry);
-	uint End = ShaderPrintGetOffset(InitShaderPrintText(InTextEntry.Index + 1));
-	for (uint i = Begin; i < End; ++i)
-	{
-		Pos = ShaderPrintSymbol(Pos, ShaderPrintGetChar(i), InColor);
-	}
-	return Pos;
-}
-float2 ShaderPrintText(float2 Pos, FShaderPrintText InTextEntry) { return ShaderPrintText(Pos, InTextEntry, GetDefaultFontColor()); }
-
-void ShaderPrintText(inout FShaderPrintContext InContext, FShaderPrintText InTextEntry, FFontColor InColor)
-{
-	if (InContext.bIsActive)
+	// If MaxValueCount is 0 then we don't reset the buffer counter so early out here
+	if (bIsActive && ShaderPrint.MaxValueCount > 0)
 	{
-		uint Begin = ShaderPrintGetOffset(InTextEntry);
-		uint End = ShaderPrintGetOffset(InitShaderPrintText(InTextEntry.Index + 1));
-		for (uint i = Begin; i < End; ++i)
+		const uint Begin = ShaderPrintGetOffset(InTextEntry);
+		const uint End   = ShaderPrintGetOffset(InitShaderPrintText(InTextEntry.Index + 1));
+		const uint Count = End - Begin;
+
+		// Buffer counter is stored in first element .Value
+		// Prevent writing off the buffer
+		// Note that counter still increases so need clamp when reading it in later passes
+		int IndexToStore = 0;
+		InterlockedAdd(ShaderPrint_RWValuesBuffer[SHADER_PRINT_COUNT_OFFSET].Value, Count, IndexToStore);
+		if (uint(IndexToStore + Count) < ShaderPrint.MaxValueCount)
 		{
-			InContext.Pos = ShaderPrintSymbol(InContext.Pos, ShaderPrintGetChar(i), InColor);
-		}	
+			for (uint i = Begin; i < End; ++i)
+			{
+				FShaderPrintItem Item;
+				Item.ScreenPos = InPos;
+				Item.Value = ShaderPrintGetChar(i);
+				Item.Type = SHADER_PRINT_TYPE_SYMBOL;
+				Item.Color = InColor.Color;
+				Item.Metadata = 0u;
+
+				ShaderPrint_RWValuesBuffer[IndexToStore + SHADER_PRINT_VALUE_OFFSET] = PackShaderPrintItem(Item);
+
+				++IndexToStore;
+				InPos.x += ShaderPrint.FontSpacing.x;
+			}	
+		}
 	}
+	return InPos;
 }
-void ShaderPrintText(inout FShaderPrintContext InContext, FShaderPrintText InTextEntry) { ShaderPrintText(InContext, InTextEntry, GetDefaultFontColor()); }
 
-float2 Print(float2 Pos, FShaderPrintText InTextEntry, FFontColor InColor)                      { return ShaderPrintText(Pos, InTextEntry, InColor); }
-float2 Print(float2 Pos, FShaderPrintText InTextEntry)											{ return ShaderPrintText(Pos, InTextEntry, GetDefaultFontColor()); }
-void   Print(inout FShaderPrintContext InCtx, FShaderPrintText InTextEntry, FFontColor InColor)	{ ShaderPrintText(InCtx, InTextEntry, InColor); }
-void   Print(inout FShaderPrintContext InCtx, FShaderPrintText InTextEntry)						{ ShaderPrintText(InCtx, InTextEntry, GetDefaultFontColor()); }
+float2 ShaderPrintText(float2 InPos, FShaderPrintText InTextEntry, FFontColor InColor)						{ return ShaderPrintText_Internal(ShaderPrintFilterEnable, InPos, InTextEntry, InColor); }
+float2 ShaderPrintText(float2 InPos, FShaderPrintText InTextEntry)											{ return ShaderPrintText_Internal(ShaderPrintFilterEnable, InPos, InTextEntry, GetDefaultFontColor()); }
+void ShaderPrintText(inout FShaderPrintContext InCtx, FShaderPrintText InTextEntry, FFontColor InColor)		{ InCtx.Pos = ShaderPrintText_Internal(InCtx.bIsActive, InCtx.Pos, InTextEntry, InColor); }
+void ShaderPrintText(inout FShaderPrintContext InCtx, FShaderPrintText InTextEntry)							{ InCtx.Pos = ShaderPrintText_Internal(InCtx.bIsActive, InCtx.Pos, InTextEntry, GetDefaultFontColor()); }
+
+float2 Print(float2 InPos, FShaderPrintText InTextEntry, FFontColor InColor)								{ return ShaderPrintText_Internal(ShaderPrintFilterEnable, InPos, InTextEntry, InColor); }
+float2 Print(float2 InPos, FShaderPrintText InTextEntry)													{ return ShaderPrintText_Internal(ShaderPrintFilterEnable, InPos, InTextEntry, GetDefaultFontColor()); }
+void   Print(inout FShaderPrintContext InCtx, FShaderPrintText InTextEntry, FFontColor InColor)				{ InCtx.Pos = ShaderPrintText_Internal(InCtx.bIsActive, InCtx.Pos, InTextEntry, InColor); }
+void   Print(inout FShaderPrintContext InCtx, FShaderPrintText InTextEntry)									{ InCtx.Pos = ShaderPrintText_Internal(InCtx.bIsActive, InCtx.Pos, InTextEntry, GetDefaultFontColor()); }
 
 // Function for reading global TEXT string. 
 // The data and function are generated by the shader compiler.
