Fixed: SceneCapture crash on Android Vulkan devices

diff --git a/Engine/Shaders/Private/SceneCapturePixelShader.usf b/Engine/Shaders/Private/SceneCapturePixelShader.usf
--- a/Engine/Shaders/Private/SceneCapturePixelShader.usf
+++ b/Engine/Shaders/Private/SceneCapturePixelShader.usf
@@ -4,6 +4,9 @@
 	SceneCapturePixelShader.usf: 
 =============================================================================*/
 
+// do not use framebuffer fetch
+#define FORCE_DEPTH_TEXTURE_READS 1
+
 #include "Common.ush"
 #include "DeferredShadingCommon.ush"
 
