Added Strata rough refraction rnd code

diff --git a/Engine/Shaders/Private/Strata/StrataRoughRefraction.usf b/Engine/Shaders/Private/Strata/StrataRoughRefraction.usf
new file mode 100644
--- /dev/null
+++ b/Engine/Shaders/Private/Strata/StrataRoughRefraction.usf
@@ -0,0 +1,205 @@
+// Copyright Epic Games, Inc. All Rights Reserved.
+
+#pragma once
+
+#include "../Common.ush"
+
+#ifdef STRATA_RND_SHADERS
+
+#include "/Engine/Private/MiniFontCommon.ush"
+#include "/Engine/Private/ShaderPrint.ush"
+#include "/Engine/Private/BRDF.ush"
+#include "/Engine/Private/MonteCarlo.ush"
+#include "/Engine/Private/PathTracing/Material/PathTracingFresnel.ush"
+
+// Updated from http://jcgt.org/published/0007/03/04/
+bool slabs(float3 p0, float3 p1, float3 rayOrigin, float3 invRaydir, out float outTMin, out float outTMax)
+{
+	float3 t0 = (p0 - rayOrigin) * invRaydir;
+	float3 t1 = (p1 - rayOrigin) * invRaydir;
+	float3 tmin = min(t0, t1), tmax = max(t0, t1);
+	float maxtmin = max(max(tmin.x, tmin.y), tmin.z);
+	float mintmax = min(min(tmax.x, tmax.y), tmax.z);
+	outTMin = maxtmin;
+	outTMax = mintmax;
+	return maxtmin <= mintmax;
+}
+
+// From https://en.wikipedia.org/wiki/MöllerTrumbore_intersection_algorithm
+bool RayIntersectsTriangle(float3 rayOrigin, float3 rayVector, float3 v0, float3 v1, float3 v2, out float3 outIntersectionPoint, out float2 uv)
+{
+	outIntersectionPoint = 0.0f;
+	uv = 0.0f;
+
+	const float EPSILON = 0.0000001f;
+	float3 edge1, edge2, tvec, pvec, qvec;
+	float det, invDet, u, v;
+
+	edge1 = v1 - v0;
+	edge2 = v2 - v0;
+
+	pvec = cross(rayVector, edge2);
+	det = dot(edge1, pvec);
+	if (det > -EPSILON && det < EPSILON)
+		return false;
+	invDet = 1.0f / det;
+
+	tvec = rayOrigin - v0;
+	u = dot(tvec, pvec) * invDet;
+	if (u < 0.0f || u > 1.0f)
+		return false;
+
+	qvec = cross(tvec, edge1);
+
+	v = dot(rayVector, qvec) * invDet;
+	if (v < 0.0f || (u + v) > 1.0f)
+		return false;
+
+	float t = dot(edge2, qvec) * invDet;
+
+	if (t > EPSILON)
+	{
+		uv = float2(u, v);
+		outIntersectionPoint = rayOrigin + rayVector * t;
+		return true;
+	}
+	else
+		return false; // Intersection not in the ray direction
+}
+
+float3 GetAnyPerpendicularVector(float3 Dir)
+{
+	float3 Perp = 0;
+	float3 Source = float3(1, 0, 0);
+	if (dot(Dir, Source) < 0.99)
+	{
+		Perp = normalize(cross(Dir, Source));
+	}
+	else
+	{
+		Perp = normalize(cross(Dir.yzx, Source));
+	}
+	return Perp;
+}
+
+float3 GetScreenWorldDir(in float4 SVPos)
+{
+	float2 ScreenPosition = SvPositionToScreenPosition(SVPos).xy;
+	const float Depth = 1000000.0f;
+	float4 WorldPos = mul(float4(ScreenPosition * Depth, Depth, 1), View.ScreenToTranslatedWorld);
+	return normalize(WorldPos.xyz - View.TranslatedWorldCameraOrigin);
+}
+
+
+bool GetRefractedDirection(float2 RandomUV, float a2, float3 V, inout float3 L )
+{ 
+	float4 H = ImportanceSampleVisibleGGX(UniformSampleDisk(RandomUV), a2, V);
+
+	L = 0;
+	float F = 0;
+	const float AirIOR = 1.0f;
+	const float WaterIOR = 1.33f;
+	const float EtaIn = WaterIOR / AirIOR;	// We always consider entering from air to water
+	const float RandSample = 1.0f;			// We only consider refracted rays
+	return SampleRefraction(-V, H.xyz, EtaIn, RandSample, L, F);
+}
+
+float TraceDomainSize;
+uint SlabInterfaceLineCount;
+
+void VisualizeRoughRefractionPS(
+	float4 SVPos : SV_POSITION,
+	out float4 OutColor : SV_Target0)
+{
+	OutColor = float4(0.0, 0.0f, 0.0, 0.5f);
+
+	float2 PixelPos = SVPos.xy;
+	float2 UvBuffer = PixelPos * View.BufferSizeAndInvSize.zw;
+	float3 WorldDir = GetScreenWorldDir(SVPos);
+	float3 CamWorldPos = LWCToFloat(PrimaryView.WorldCameraOrigin);
+
+	FShaderPrintContext Context = InitShaderPrintContext(all(uint2(PixelPos.xy) == uint2(100, 100)), uint2(100, 100));
+	const float IncidentAngle = AddSlider(Context, TEXT("IncidentAngle"), 0.0f, GetDefaultFontColor(), 0.0f, 0.5f * PI);
+	Newline(Context);
+	const float SlabThickness = AddSlider(Context, TEXT("Slab Thickness"), 5.0f, GetDefaultFontColor(), 0.0f, 20.0f);
+	Newline(Context);
+	const float Roughness = AddSlider(Context, TEXT("Roughness"), 0.5f, GetDefaultFontColor(), 0.0f, 1.0f);
+	Newline(Context);
+
+	float a2 = Pow4(Roughness);
+	const float3 N = float3(0, 0, 1);
+	const float3 V = float3(sin(IncidentAngle), 0, cos(IncidentAngle));
+
+	float3 SlabEdgeX = float3(60.0f, 0.0f, 0.0f);
+	float3 SlabEdgeY = float3(0.0f, 60.0f, 0.0f);
+	float3 SlabEdgeZ = float3(0.0f, 0.0f, SlabThickness);
+	float3 SlabCenter = CamWorldPos + View.ViewForward * 100.0f;
+	float3 RefractionP= SlabCenter + 0.5 * SlabEdgeZ;
+
+
+	AddLineWS(Context, RefractionP, RefractionP + N * 25.0, float4(0, 0, 1, 1));
+	AddLineWS(Context, RefractionP, RefractionP + V * 25.0, float4(0, 1, 1, 1));
+
+#if 0
+	float TMin = -1.0f;
+	float TMax = -1.0f;
+	float3 p0 = SlabCenter - 0.5 * (SlabEdgeX + SlabEdgeY + SlabEdgeZ);
+	float3 p1 = SlabCenter + 0.5 * (SlabEdgeX + SlabEdgeY + SlabEdgeZ);
+	slabs(p0, p1, WorldOrigin, 1.0f / WorldDir, TMin, TMax);
+#endif
+
+#if 0
+	const float SphereRadius = 10.0f;
+	float2 Sol = RayIntersectSphere(CamWorldPos, WorldDir, float4(RefractionP, SphereRadius));
+	if (Sol.x > 0.0f && Sol.y > 0.0f)
+	{
+		const float3 P = CamWorldPos + WorldDir * Sol.x;
+		const float3 OmegaOut = P - RefractionP;
+		const float OmegaOutLen = length(OmegaOut);
+		const float3 OmegaOutNorm = OmegaOut / OmegaOutLen;
+
+		const float NoH = dot(normalize(OmegaOutNorm + V), N);
+		float D = D_GGX(a2, NoH);
+		OutColor = float4(D, D, 0.0, 0.0);
+	}
+#endif
+
+	// Draw the top and bottom interfaces
+	float3 SlabOrigin = SlabCenter - 0.5 * (SlabEdgeX + SlabEdgeY + SlabEdgeZ);
+	for (uint i = 0; i <= SlabInterfaceLineCount; i++)
+	{
+		const float Offset = (float(i) / float(SlabInterfaceLineCount));
+
+		AddLineWS(Context, SlabOrigin + SlabEdgeX * Offset, SlabOrigin + SlabEdgeX * Offset + SlabEdgeY, float4(0.2, 0.2, 0.2, 1.0));
+		AddLineWS(Context, SlabOrigin + SlabEdgeY * Offset, SlabOrigin + SlabEdgeY * Offset + SlabEdgeX, float4(0.2, 0.2, 0.2, 1.0));
+
+		AddLineWS(Context, SlabOrigin + SlabEdgeX * Offset + SlabEdgeZ, SlabOrigin + SlabEdgeX * Offset + SlabEdgeY + SlabEdgeZ, float4(0.2, 0.2, 0.2, 1.0));
+		AddLineWS(Context, SlabOrigin + SlabEdgeY * Offset + SlabEdgeZ, SlabOrigin + SlabEdgeY * Offset + SlabEdgeX + SlabEdgeZ, float4(0.2, 0.2, 0.2, 1.0));
+	}
+
+	for (float u = 0.5f; u < TraceDomainSize; u++)
+	{
+		for (float v = 0.5f; v < TraceDomainSize; v++)
+		{
+			float2 RandomUV = float2(u, v) / TraceDomainSize;
+			float3 L = 0;
+			if (GetRefractedDirection(RandomUV, a2, V, L))
+			{
+				AddLineWS(Context, RefractionP, RefractionP + L * 15.0, float4(0, 1, 0, 0.3));
+			}
+			else
+			{
+				AddLineWS(Context, RefractionP, RefractionP + L * 15.0, float4(1, 0, 0, 0.3));
+			}
+		}
+	}
+
+}
+
+
+
+
+
+#endif
+
+
diff --git a/Engine/Source/Runtime/Renderer/Private/Strata/Strata.cpp b/Engine/Source/Runtime/Renderer/Private/Strata/Strata.cpp
--- a/Engine/Source/Runtime/Renderer/Private/Strata/Strata.cpp
+++ b/Engine/Source/Runtime/Renderer/Private/Strata/Strata.cpp
@@ -833,6 +833,7 @@ bool ShouldRenderStrataDebugPasses(const FViewInfo& View)
 		(
 		   (FVisualizeMaterialPS::CanRunStrataVizualizeMaterial(View.GetShaderPlatform()) && View.Family && View.Family->EngineShowFlags.VisualizeStrataMaterial)
 		|| (CVarStrataClassificationDebug.GetValueOnAnyThread() > 0)
+		|| ShouldRenderStrataRoughRefractionRnD()
 		);
 }
 
@@ -857,6 +858,8 @@ FScreenPassTexture AddStrataDebugPasses(FRDGBuilder& GraphBuilder, const FViewIn
 		AddStrataInternalClassificationTilePass(GraphBuilder, View, nullptr, &ScreenPassSceneColor.Texture, EStrataTileMaterialType::EComplex, bDebugPass);
 	}
 
+	StrataRoughRefractionRnD(GraphBuilder, View, ScreenPassSceneColor);
+
 	return MoveTemp(ScreenPassSceneColor);
 }
 
diff --git a/Engine/Source/Runtime/Renderer/Private/Strata/Strata.h b/Engine/Source/Runtime/Renderer/Private/Strata/Strata.h
--- a/Engine/Source/Runtime/Renderer/Private/Strata/Strata.h
+++ b/Engine/Source/Runtime/Renderer/Private/Strata/Strata.h
@@ -149,4 +149,7 @@ FStrataTileParameter SetTileParameters(FRDGBuilder& GraphBuilder, const FViewInf
 FStrataTilePassVS::FParameters SetTileParameters(FRDGBuilder& GraphBuilder, const FViewInfo& View, const EStrataTileMaterialType Type, EPrimitiveType& PrimitiveType);
 FStrataTilePassVS::FParameters SetTileParameters(const FViewInfo& View, const EStrataTileMaterialType Type, EPrimitiveType& PrimitiveType);
 
+bool ShouldRenderStrataRoughRefractionRnD();
+void StrataRoughRefractionRnD(FRDGBuilder& GraphBuilder, const FViewInfo& View, FScreenPassTexture& ScreenPassSceneColor);
+
 };
diff --git a/Engine/Source/Runtime/Renderer/Private/Strata/StrataRoughRefraction.cpp b/Engine/Source/Runtime/Renderer/Private/Strata/StrataRoughRefraction.cpp
new file mode 100644
--- /dev/null
+++ b/Engine/Source/Runtime/Renderer/Private/Strata/StrataRoughRefraction.cpp
@@ -0,0 +1,145 @@
+// Copyright Epic Games, Inc. All Rights Reserved.
+
+#include "Strata.h"
+#include "HAL/IConsoleManager.h"
+#include "PixelShaderUtils.h"
+#include "SceneView.h"
+#include "ScenePrivate.h"
+#include "SceneRendering.h"
+#include "RendererInterface.h"
+#include "UniformBuffer.h"
+#include "SceneTextureParameters.h"
+#include "ScreenPass.h"
+#include "ShaderCompiler.h"
+
+
+namespace Strata
+{
+
+//////////////////////////////////////////////////////////////////////////
+// RnD shaders only used when enabled locally
+
+// Keeping it simple: this should always be checked in with a value of 0
+#define STRATA_ROUGH_REFRACTION_RND 0
+
+#if STRATA_ROUGH_REFRACTION_RND
+
+static TAutoConsoleVariable<int32> CVarStrataRoughRefractionShadersShowRoughRefractionRnD(
+	TEXT("r.Strata.ShowRoughRefractionRnD"),
+	1,
+	TEXT("Enable strata rough refraction shaders."),
+	ECVF_RenderThreadSafe);
+
+
+class FVisualizeRoughRefractionPS : public FGlobalShader
+{
+	DECLARE_GLOBAL_SHADER(FVisualizeRoughRefractionPS);
+	SHADER_USE_PARAMETER_STRUCT(FVisualizeRoughRefractionPS, FGlobalShader);
+
+	using FPermutationDomain = TShaderPermutationDomain<>;
+
+	BEGIN_SHADER_PARAMETER_STRUCT(FParameters, )
+		SHADER_PARAMETER_STRUCT_REF(FViewUniformShaderParameters, ViewUniformBuffer)
+		SHADER_PARAMETER_TEXTURE(Texture2D, MiniFontTexture)
+		SHADER_PARAMETER_STRUCT_INCLUDE(ShaderPrint::FShaderParameters, ShaderPrintParameters)
+		SHADER_PARAMETER(float, TraceDomainSize)
+		SHADER_PARAMETER(uint32, SlabInterfaceLineCount)
+		RENDER_TARGET_BINDING_SLOTS()
+	END_SHADER_PARAMETER_STRUCT()
+
+	static FPermutationDomain RemapPermutation(FPermutationDomain PermutationVector)
+	{
+		return PermutationVector;
+	}
+
+	static bool ShouldCompilePermutation(const FGlobalShaderPermutationParameters& Parameters)
+	{
+		return true;
+	}
+
+	static void ModifyCompilationEnvironment(const FGlobalShaderPermutationParameters& Parameters, FShaderCompilerEnvironment& OutEnvironment)
+	{
+		FGlobalShader::ModifyCompilationEnvironment(Parameters, OutEnvironment);
+
+		// Stay debug and skip optimizations to reduce compilation time on this long shader.
+		OutEnvironment.CompilerFlags.Add(CFLAG_PrecompileWithDXC);
+
+		OutEnvironment.SetDefine(TEXT("STRATA_RND_SHADERS"), 1);
+	}
+};
+IMPLEMENT_GLOBAL_SHADER(FVisualizeRoughRefractionPS, "/Engine/Private/Strata/StrataRoughRefraction.usf", "VisualizeRoughRefractionPS", SF_Pixel);
+
+
+
+#endif // STRATA_ROUGH_REFRACTION_RND
+	
+
+
+bool ShouldRenderStrataRoughRefractionRnD()
+{
+#if STRATA_ROUGH_REFRACTION_RND
+	return CVarStrataRoughRefractionShadersShowRoughRefractionRnD.GetValueOnAnyThread() > 0;
+#else
+	return false;
+#endif
+}
+
+void StrataRoughRefractionRnD(FRDGBuilder& GraphBuilder, const FViewInfo& View, FScreenPassTexture& ScreenPassSceneColor)
+{
+#if STRATA_ROUGH_REFRACTION_RND
+	if (IsStrataEnabled() && ShouldRenderStrataRoughRefractionRnD())
+	{
+		if (!ShaderPrint::IsSupported(View.GetShaderPlatform()))
+		{
+			return;
+		}
+		check(ShaderPrint::IsEnabled(View));	// One must enable ShaderPrint beforehand using r.ShaderPrint=1
+
+		//////////////////////////////////////////////////////////////////////////
+		// Create resources
+		
+		// Texture to count
+		//const uint32 SampleCountTextureWidth = 64;
+		//const FIntPoint SampleCountTextureSize(SampleCountTextureWidth, SampleCountTextureWidth);
+		//FRDGTextureRef SampleCountTexture = GraphBuilder.CreateTexture(FRDGTextureDesc::Create2D(SampleCountTextureSize, PF_R32_UINT, FClearValueBinding::Black, TexCreate_UAV | TexCreate_ShaderResource), TEXT("Strata.RoughRefrac.SampleCount"));
+		//FRDGTextureUAVRef SampleCountTextureUAV = GraphBuilder.CreateUAV(FRDGTextureUAVDesc(SampleCountTexture));
+		//FRDGTextureSRVRef SampleCountTextureSRV = GraphBuilder.CreateSRV(FRDGTextureSRVDesc::Create(SampleCountTexture));
+		//
+		//FRDGBufferRef LobStatisticsBuffer = GraphBuilder.CreateBuffer(FRDGBufferDesc::CreateBufferDesc(sizeof(uint32) * 4, 64), TEXT("Strata.RoughRefrac.LobStat"));
+		//FRDGTextureUAVRef SampleCountTextureUAV = GraphBuilder.CreateUAV(LobStatisticsBuffer, PF_A32B32G32R32F);
+		//FRDGTextureSRVRef SampleCountTextureSRV = GraphBuilder.CreateSRV(LobStatisticsBuffer, PF_A32B32G32R32F);
+
+
+
+
+
+		//////////////////////////////////////////////////////////////////////////
+		// Debug print everything on screen
+		const float TraceDomainSize = 32.0f;
+		const float SlabInterfaceLineCount = 16.0f;
+
+		ShaderPrint::RequestSpaceForLines(TraceDomainSize* TraceDomainSize + SlabInterfaceLineCount * SlabInterfaceLineCount * 2);
+		ShaderPrint::RequestSpaceForCharacters(256);
+
+		FSceneTextureParameters SceneTextureParameters = GetSceneTextureParameters(GraphBuilder);
+		FVisualizeRoughRefractionPS::FParameters* PassParameters = GraphBuilder.AllocParameters<FVisualizeRoughRefractionPS::FParameters>();
+		PassParameters->ViewUniformBuffer = View.ViewUniformBuffer;
+		PassParameters->MiniFontTexture = GetMiniFontTexture();
+		PassParameters->RenderTargets[0] = FRenderTargetBinding(ScreenPassSceneColor.Texture, ERenderTargetLoadAction::ELoad);
+		ShaderPrint::SetParameters(GraphBuilder, View.ShaderPrintData, PassParameters->ShaderPrintParameters);
+		PassParameters->TraceDomainSize = TraceDomainSize;
+		PassParameters->SlabInterfaceLineCount = SlabInterfaceLineCount;
+
+		FVisualizeRoughRefractionPS::FPermutationDomain PermutationVector;
+		TShaderMapRef<FVisualizeRoughRefractionPS> PixelShader(View.ShaderMap, PermutationVector);
+
+		FRHIBlendState* PreMultipliedColorTransmittanceBlend = TStaticBlendState<CW_RGB, BO_Add, BF_One, BF_SourceAlpha, BO_Add, BF_Zero, BF_One>::GetRHI();
+
+		FPixelShaderUtils::AddFullscreenPass<FVisualizeRoughRefractionPS>(GraphBuilder, View.ShaderMap, RDG_EVENT_NAME("Strata::VisualizeRoughRefractionPS"),
+			PixelShader, PassParameters, View.ViewRect, PreMultipliedColorTransmittanceBlend);
+	}
+#endif
+}
+
+
+} // namespace Strata
\ No newline at end of file
