Added lob statistic estimation and debug draw.

diff --git a/Engine/Shaders/Private/Strata/StrataRoughRefraction.usf b/Engine/Shaders/Private/Strata/StrataRoughRefraction.usf
--- a/Engine/Shaders/Private/Strata/StrataRoughRefraction.usf
+++ b/Engine/Shaders/Private/Strata/StrataRoughRefraction.usf
@@ -13,6 +13,8 @@
 #include "/Engine/Private/MonteCarlo.ush"
 #include "/Engine/Private/PathTracing/Material/PathTracingFresnel.ush"
 
+#include "/Engine/Private/Strata/StrataStatisticalOperators.ush"
+
 // Updated from http://jcgt.org/published/0007/03/04/
 bool Slabs(float3 P0, float3 P1, float3 RayOrigin, float3 InvRaydir, out float OutTMin, out float OutTMax)
 {
@@ -49,7 +51,10 @@ bool GetRefractedDirection(float2 RandomUV, float a2, float3 V, inout float3 L )
 	return SampleRefraction(-V, H.xyz, EtaIn, RandSample, L, F);
 }
 
+
+
 #if defined(EVALUATE_ROUGH_REFRACTION_LOBE_CS) || defined(VISUALIZE_ROUGH_REFRACTION_PS)
+
 struct FUI
 {
 	float IncidentAngle;
@@ -62,7 +67,6 @@ struct FUI
 	float3 V;
 };
 
-
 FShaderPrintContext GetUIContext(bool bActive)
 {
 	return InitShaderPrintContext(bActive, uint2(100, 100));
@@ -85,16 +89,67 @@ FUI GetUI(inout FShaderPrintContext Context)
 
 	return UI;
 }
+
+void DrawCircleXY(inout FShaderPrintContext Context, float3 Center, float Radius, float4 ColorAlpha)
+{
+	const float TStep = 0.1f;
+	for (float t = 0; t <= 1.0; t += TStep)
+	{
+		float S0, C0;
+		float S1, C1;
+		sincos(t * 2 * PI, S0, C0);
+		sincos((t + TStep) * 2 * PI, S1, C1);
+		float3 P0 = Center + float3(C0, S0, 0.0) * Radius;
+		float3 P1 = Center + float3(C1, S1, 0.0) * Radius;
+
+		AddLineWS(Context, P0, P1, ColorAlpha);
+	}
+}
+
+struct FLobeStat
+{
+	float Count;
+	float2 Mean;
+	float2 M2;
+	float2 Variance;
+	float pad;
+};
+
+// See Welford's algorithm https://en.wikipedia.org/wiki/Algorithms_for_calculating_variance
+void LobStatUpdate(inout FLobeStat LobeStat, in float2 InPos)
+{
+	LobeStat.Count += 1;
+	float2 Delta = InPos - LobeStat.Mean;
+	LobeStat.Mean += Delta / LobeStat.Count;
+	float2 Delta2 = InPos - LobeStat.Mean;
+	LobeStat.M2 += Delta * Delta2;
+}
+
+void LobStatFinalize(inout FLobeStat LobeStat)
+{
+	if (LobeStat.Count < 2)
+	{
+		// nan
+	}
+	else
+	{
+		// Mean is already valid
+		LobeStat.Variance = LobeStat.M2 / LobeStat.Count;
+		//LobeStat.SampleVariance = LobeStat.M2 / (LobeStat.Count - 1);
+	}
+}
+
 #endif
 
 
+
 #ifdef EVALUATE_ROUGH_REFRACTION_LOBE_CS
 
 uint TraceSqrtSampleCount;
 
 RWTexture2D<uint> SampleCountTextureUAV;
 
-RWStructuredBuffer<FLobeMeanVariance> LobStatisticsBufferUAV;
+RWStructuredBuffer<FLobeStat> LobStatisticsBufferUAV;
 
 [numthreads(THREADGROUP_SIZE, THREADGROUP_SIZE, 1)]
 void EvaluateRoughRefractionLobeCS(uint3 ThreadId : SV_DispatchThreadID)
@@ -105,8 +160,7 @@ void EvaluateRoughRefractionLobeCS(uint3 ThreadId : SV_DispatchThreadID)
 		FUI UI = GetUI(Context);
 
 
-		FLobeMeanVariance Lobe = (FLobeMeanVariance)0;
-		LobStatisticsBufferUAV[0] = Lobe;
+		FLobeStat LobeState = (FLobeStat)0;
 
 
 		for (float u = 0.5f; u < float(TraceSqrtSampleCount); u++)
@@ -119,6 +173,8 @@ void EvaluateRoughRefractionLobeCS(uint3 ThreadId : SV_DispatchThreadID)
 				{
 					float2 SampleUV = L.xy * 0.5 + 0.5;
 					SampleCountTextureUAV[uint2(SampleUV * 64.0f)] += 1;
+
+					LobStatUpdate(LobeState, L.xy);
 				}
 				else
 				{
@@ -126,6 +182,9 @@ void EvaluateRoughRefractionLobeCS(uint3 ThreadId : SV_DispatchThreadID)
 				}
 			}
 		}
+
+		LobStatFinalize(LobeState);
+		LobStatisticsBufferUAV[0] = LobeState;
 	}
 }
 
@@ -140,7 +199,7 @@ uint SlabInterfaceLineCount;
 
 Texture2D<uint> SampleCountTexture;
 
-StructuredBuffer<FLobeMeanVariance> LobStatisticsBuffer;
+StructuredBuffer<FLobeStat> LobStatisticsBuffer;
 
 void VisualizeRoughRefractionPS(
 	float4 SVPos : SV_POSITION,
@@ -162,8 +221,10 @@ void VisualizeRoughRefractionPS(
 	float3 SlabEdgeZ = float3(0.0f, 0.0f, UI.SlabThickness);
 	float3 SlabCenter = CamWorldPos + View.ViewForward * 60.0f;
 	float3 RefractionP= SlabCenter + 0.5 * SlabEdgeZ;
+	const float SlabTopZ    = SlabCenter.z  + 0.5 * UI.SlabThickness;
+	const float SlabBottomZ = SlabCenter.z - 0.5 * UI.SlabThickness;
 
-	const float VizProjectedRays = 15.0f;
+	const float VisProjectedRayRadius = 15.0f;
 
 	AddLineWS(Context, RefractionP, RefractionP + UI.N * 25.0, float4(0, 0, 1, 1));
 	AddLineWS(Context, RefractionP, RefractionP + UI.V * 25.0, float4(0, 1, 1, 1));
@@ -175,14 +236,14 @@ void VisualizeRoughRefractionPS(
 	if(Slabs(SlabP0, SlabP1, CamWorldPos, 1.0f / WorldDir, SlabTMin, SlabTMax))
 	{
 		float3 P = CamWorldPos + WorldDir * SlabTMin;
-		if (P.z > -0.499 * UI.SlabThickness)
+		if (P.z > (SlabBottomZ + 0.01))
 		{
 			P = CamWorldPos + WorldDir * SlabTMax;
 		}
 
-		if (P.z < (SlabCenter -0.499 * SlabEdgeZ).z)
+		if (P.z < (SlabBottomZ+0.01))
 		{
-			const float2 UV = 0.5 + 0.5 * ((P - RefractionP) / (VizProjectedRays)).xy;
+			const float2 UV = 0.5 + 0.5 * ((P - RefractionP) / (VisProjectedRayRadius)).xy;
 			const int2 Coord = (UV * 64.0f + 0.5f);
 			if (all(Coord >= 0) && all(Coord < 64))
 			{
@@ -229,16 +290,28 @@ void VisualizeRoughRefractionPS(
 			float3 L = 0;
 			if (GetRefractedDirection(RandomUV, UI.a2, UI.V, L))
 			{
-				AddLineWS(Context, RefractionP, RefractionP + L * VizProjectedRays, float4(0, 1, 0, 0.3));
-				//AddLineWS(Context, RefractionP + L * VizProjectedRays, float3((RefractionP + L * VizProjectedRays).xy, SlabCenter.z - 0.5* UI.SlabThickness), float4(0, 1, 0, 0.3));
+				AddLineWS(Context, RefractionP, RefractionP + L * VisProjectedRayRadius, float4(0, 1, 0, 0.3));
+				//AddLineWS(Context, RefractionP + L * VisProjectedRayRadius, float3((RefractionP + L * VisProjectedRayRadius).xy, SlabCenter.z - 0.5* UI.SlabThickness), float4(0, 1, 0, 0.3));
 			}
 			else
 			{
-				AddLineWS(Context, RefractionP, RefractionP + L * VizProjectedRays, float4(1, 0, 0, 0.3));
+				AddLineWS(Context, RefractionP, RefractionP + L * VisProjectedRayRadius, float4(1, 0, 0, 0.3));
 			}
 		}
 	}
 
+	// Draw the lobe estimated by Belcour
+	const float AirIOR = 1.0f;
+	const float WaterIOR = 1.33f;
+	const float EtaIn = WaterIOR / AirIOR;	// We always consider entering from air to water
+	FStrataLobeStatistic RefractedLobe = StrataGetRefractedLobe(StrataGetViewDiracLobe(-UI.V), 1.0, UI.Roughness, EtaIn);
+	DrawCircleXY(Context, float3(SlabCenter.xy + RefractedLobe.Mu.xy * VisProjectedRayRadius, SlabBottomZ), RefractedLobe.Sigma * VisProjectedRayRadius, float4(0, 1, 1, 1));
+
+	// Our lobe
+	FLobeStat LobeState = LobStatisticsBuffer[0];
+	DrawCircleXY(Context, float3(SlabCenter.xy + LobeState.Mean * VisProjectedRayRadius, SlabBottomZ), LobeState.Variance.x * VisProjectedRayRadius, float4(1, 1, 1, 1));
+	DrawCircleXY(Context, float3(SlabCenter.xy + LobeState.Mean * VisProjectedRayRadius, SlabBottomZ+1.0), LobeState.Variance.y * VisProjectedRayRadius, float4(0, 1, 0, 1));
+
 }
 
 #endif // VISUALIZE_ROUGH_REFRACTION_PS
diff --git a/Engine/Source/Runtime/Renderer/Private/Strata/StrataRoughRefraction.cpp b/Engine/Source/Runtime/Renderer/Private/Strata/StrataRoughRefraction.cpp
--- a/Engine/Source/Runtime/Renderer/Private/Strata/StrataRoughRefraction.cpp
+++ b/Engine/Source/Runtime/Renderer/Private/Strata/StrataRoughRefraction.cpp
@@ -147,7 +147,7 @@ void StrataRoughRefractionRnD(FRDGBuilder& GraphBuilder, const FViewInfo& View,
 		FRDGTextureUAVRef SampleCountTextureUAV = GraphBuilder.CreateUAV(FRDGTextureUAVDesc(SampleCountTexture));
 		FRDGTextureSRVRef SampleCountTextureSRV = GraphBuilder.CreateSRV(FRDGTextureSRVDesc::Create(SampleCountTexture));
 
-		FRDGBufferRef LobStatisticsBuffer = GraphBuilder.CreateBuffer(FRDGBufferDesc::CreateStructuredDesc(sizeof(float) * 4, 16), TEXT("Strata.RoughRefrac.LobStat"));
+		FRDGBufferRef LobStatisticsBuffer = GraphBuilder.CreateBuffer(FRDGBufferDesc::CreateStructuredDesc(sizeof(float) * 8, 16), TEXT("Strata.RoughRefrac.LobStat"));
 		FRDGBufferUAVRef LobStatisticsBufferUAV = GraphBuilder.CreateUAV(LobStatisticsBuffer, PF_R32_UINT);
 		FRDGBufferSRVRef LobStatisticsBufferSRV = GraphBuilder.CreateSRV(LobStatisticsBuffer, PF_R32_UINT);
 
