Add missing shader file.

diff --git a/Engine/Shaders/Private/ShadingModelsSampling.ush b/Engine/Shaders/Private/ShadingModelsSampling.ush
--- a/Engine/Shaders/Private/ShadingModelsSampling.ush
+++ b/Engine/Shaders/Private/ShadingModelsSampling.ush
@@ -65,13 +65,13 @@ FBxDFSample SampleDiffuseBxDF(float3 N, float4 E)
 	return BxDFSample;
 }
 
-FBxDFSample SampleDefaultLitBxDF(uint TermMask, FGBufferData GBuffer, float3 V, float4 E)
+FBxDFSample SampleDefaultLitBxDF(uint TermMask, float3 WorldNormal, float Roughness, float3 V, float4 E)
 {
 	TermMask &= SHADING_TERM_DIFFUSE | SHADING_TERM_SPECULAR;
 
 	float2 DiskE = UniformSampleDiskConcentric(E.xy);
 
-	float3 N = GBuffer.WorldNormal;
+	float3 N = WorldNormal;
 	float3x3 TangentBasis = GetTangentBasis(N);
 
 	// TODO: stocastically choose SHADING_TERM_D or SHADING_TERM_DIFFUSE if TermMask == (SHADING_TERM_D | SHADING_TERM_SPECULAR)
@@ -88,7 +88,7 @@ FBxDFSample SampleDefaultLitBxDF(uint TermMask, FGBufferData GBuffer, float3 V,
 	}
 	else if (TermMask == SHADING_TERM_SPECULAR)
 	{
-		float a2 = Pow4(GBuffer.Roughness);
+		float a2 = Pow4(Roughness);
 
 		float3 TangentV = mul(TangentBasis, V);
 
@@ -125,6 +125,7 @@ FBxDFSample SampleDefaultLitBxDF(uint TermMask, FGBufferData GBuffer, float3 V,
 
 	return BxDFSample;
 }
+FBxDFSample SampleDefaultLitBxDF(uint TermMask, FGBufferData GBuffer, float3 V, float4 E) { return SampleDefaultLitBxDF(TermMask, GBuffer.WorldNormal, GBuffer.Roughness, V, E); }
 
 FBxDFSample SampleHairBxDF(uint TermMask, FGBufferData GBuffer, float3 V, float4 E)
 #if 1
@@ -140,10 +141,10 @@ FBxDFSample SampleHairBxDF(uint TermMask, FGBufferData GBuffer, float3 V, float4
 
 	FHairTransmittanceData TransmittanceData = InitHairTransmittanceData(true);
 	#if USE_HAIR_COMPLEX_TRANSMITTANCE
-	if (ShouldUseHairComplexTransmittance(GBuffer))
-	{
-		TransmittanceData = EvaluateDualScattering(GBuffer, V, L.xyz);
-	}
+	//if (ShouldUseHairComplexTransmittance(GBuffer))
+	//{
+	//	TransmittanceData = EvaluateDualScattering(GBuffer, V, L.xyz);
+	//}
 	#endif
 
 	FDirectLighting Lighting;
