Remove unnecessary codes in BloomPS_Mobile

diff --git a/Engine/Shaders/Private/PostProcessMobile.usf b/Engine/Shaders/Private/PostProcessMobile.usf
--- a/Engine/Shaders/Private/PostProcessMobile.usf
+++ b/Engine/Shaders/Private/PostProcessMobile.usf
@@ -285,11 +285,6 @@ void BloomPS_Mobile(
 	half4 C2 = SceneColorTexture.Sample(SceneColorSampler, InUVs[2]);
 	half4 C3 = SceneColorTexture.Sample(SceneColorSampler, InUVs[3]);
 
-	C0.rgb = C0;
-	C1.rgb = C1;
-	C2.rgb = C2;
-	C3.rgb = C3;
-
 	// Output color is average.
 	AverageColor.rgb = (C0.rgb * 0.25) + (C1.rgb * 0.25) + (C2.rgb * 0.25) + (C3.rgb * 0.25);
 
