Fix typo in roughness to variance formula.

diff --git a/Engine/Shaders/Private/Strata/StrataStatisticalOperators.ush b/Engine/Shaders/Private/Strata/StrataStatisticalOperators.ush
--- a/Engine/Shaders/Private/Strata/StrataStatisticalOperators.ush
+++ b/Engine/Shaders/Private/Strata/StrataStatisticalOperators.ush
@@ -24,12 +24,15 @@ struct FStrataLobeStatistic
 
 float StrataLobeRoughnessToVariance(float Roughness)
 {
+	// Eq. 6
 	const float SafeRoughness = clamp(Roughness, 0.0f, 0.999f);
-	return pow(SafeRoughness, 1.1f) / (1.0f - SafeRoughness);
+	const float a11 = pow(SafeRoughness, 1.1f);
+	return a11 / (1.0f - a11);
 }
 
 float StrataLobeVarianceToRoughness(float Variance)
 {
+	// Inverse of Eq. 6
 	return pow(Variance / (1.0f + Variance), 1.0f / 1.1f);
 }
 
