Disable IOR changes when propagating statistic as it causes roughness transform to be non-symetric when layers are culled with 0 coverage.

diff --git a/Engine/Shaders/Private/Strata/StrataTree.ush b/Engine/Shaders/Private/Strata/StrataTree.ush
--- a/Engine/Shaders/Private/Strata/StrataTree.ush
+++ b/Engine/Shaders/Private/Strata/StrataTree.ush
@@ -7,7 +7,8 @@
 #endif
 
 // STRATA_TODO We assume layers all have the same IOR. But we should take into account the material later.
-#define WATER_IOR		1.33f
+// For now, we suppose no IOR changes at interface until they are tracked properly
+#define WATER_IOR		1.0f 
 #define AIR_IOR			1.0f
 
 
