Fixed bad SSS (overblure, missing specular) when there are IES profile, cloud shadow or a sunset with per pixel transmittance.

diff --git a/Engine/Shaders/Private/DeferredLightPixelShaders.usf b/Engine/Shaders/Private/DeferredLightPixelShaders.usf
--- a/Engine/Shaders/Private/DeferredLightPixelShaders.usf
+++ b/Engine/Shaders/Private/DeferredLightPixelShaders.usf
@@ -193,6 +193,16 @@ float4 StrataReadPrecomputedShadowFactors(FStrataPixelHeader StrataPixelHeader,
 
 ////////////////////////////////////////////////////////////////////////////////////////////////////
 
+// Affect the light color by any other shadow/transmittance before the lighting is evaluated
+// in order to get a correct alpha channel containing important information for SSS.
+void UpdateLightDataColor(inout FDeferredLightData LightData, FInputParams InputParams, FDerivedParams DerivedParams)
+{
+	const float Attenuation = ComputeLightProfileMultiplier(DerivedParams.TranslatedWorldPosition, GetDeferredLightTranslatedWorldPosition(), -DeferredLightUniforms.Direction, DeferredLightUniforms.Tangent);
+	const float3 AttenuationRGB = Attenuation * GetAtmosphereAndCloudAttenuation(GetStableTranslatedLargeWorldPosition(InputParams.PixelPos, LookupDeviceZ(InputParams.ScreenUV)), InputParams.ScreenUV);
+
+	LightData.Color *= AttenuationRGB;
+}
+
 #if USE_HAIR_LIGHTING == 0
 
 void DeferredLightPixelMain(
@@ -238,6 +248,7 @@ void DeferredLightPixelMain(
 		const FDerivedParams DerivedParams = GetDerivedParams(InputParams, SceneDepth);
 
 		FDeferredLightData LightData = SetupLightDataForStandardDeferred();
+		UpdateLightDataColor(LightData, InputParams, DerivedParams);
 
 		float3 V =-DerivedParams.CameraVector;
 		float3 L = LightData.Direction;	// Already normalized
@@ -287,6 +298,8 @@ void DeferredLightPixelMain(
 		const FDerivedParams DerivedParams = GetDerivedParams(InputParams, SceneDepth);
 
 		FDeferredLightData LightData = SetupLightDataForStandardDeferred();
+		UpdateLightDataColor(LightData, InputParams, DerivedParams);
+
 	 #if USE_HAIR_COMPLEX_TRANSMITTANCE
 		if (ScreenSpaceData.GBuffer.ShadingModelID == SHADINGMODELID_HAIR && ShouldUseHairComplexTransmittance(ScreenSpaceData.GBuffer))
 		{
@@ -302,10 +315,7 @@ void DeferredLightPixelMain(
 
 #endif // STRATA_ENABLED
 
-		const float  Attenuation = ComputeLightProfileMultiplier(DerivedParams.TranslatedWorldPosition, GetDeferredLightTranslatedWorldPosition(), -DeferredLightUniforms.Direction, DeferredLightUniforms.Tangent);
-
-		OutColor += (Radiance * Attenuation);
-		OutColor.rgb *= GetAtmosphereAndCloudAttenuation(GetStableTranslatedLargeWorldPosition(InputParams.PixelPos, LookupDeviceZ(InputParams.ScreenUV)), InputParams.ScreenUV);
+		OutColor += Radiance;
 	}
 
 	// RGB:SceneColor Specular and Diffuse
