Strata stat operator, brought back mu as vector3 for refraction investigation.
MAde direction match the paper and clarified what is happening when you want to refract a lob that was previously reflected through a top layer.

diff --git a/Engine/Shaders/Private/Strata/StrataStatisticalOperators.ush b/Engine/Shaders/Private/Strata/StrataStatisticalOperators.ush
--- a/Engine/Shaders/Private/Strata/StrataStatisticalOperators.ush
+++ b/Engine/Shaders/Private/Strata/StrataStatisticalOperators.ush
@@ -11,7 +11,8 @@ struct FStrataLobeStatistic
 {
 	// Mean
 	// xy is the 2d projection of the lobe main direction onto the plane defined by the surface normal. 
-	float2 Mu;
+	// z is making the 3d vector a unit vector. This is useful to evaluate `s` for refraction operations (equation 10).
+	float3 Mu;
 
 	// Energy
 	float3 E;
@@ -52,22 +53,35 @@ FStrataLobeStatistic StrataGetNullLobe()
 	return NullLobe;
 }
 
-FStrataLobeStatistic StrataGetViewDiracLobe(float3 Wi)
+// Wi typically points outward the surface
+FStrataLobeStatistic StrataGetDiracLobe(float3 Wi)
 {
 	FStrataLobeStatistic WiLobe;
 	WiLobe.E = 1.0f;
-	WiLobe.Mu = Wi.xy;
+	WiLobe.Mu = Wi;
 	WiLobe.Sigma = 0.0f;
 	return WiLobe;
 }
 
+// This is used when a lob needs to be converted from OmegaO to OmegaI.
+// For instance: a lobe OmegaR is reflected from OmegaI on a surface. Then it needs to be refracted from a top interface.
+//				We thus need to convert OmegaR to be BottomUpOmegaI relatively to the top surface.
+//				This is achieved using BottomUpOmegaI = StrataOppositeLobe(OmegaR)
+FStrataLobeStatistic StrataOppositeLobe(FStrataLobeStatistic LobeIn)
+{
+	FStrataLobeStatistic LobeOut = LobeIn;
+	LobeOut.Mu = -LobeOut.Mu;
+	return LobeOut;
+}
+
 FStrataLobeStatistic StrataGetReflectedLobe(FStrataLobeStatistic WiLobe, float3 InterfaceFDG, float InterfaceRoughness)
 {
 	FStrataLobeStatistic WoLobe;
 
 	WoLobe.E = WiLobe.E * InterfaceFDG;
 
-	WoLobe.Mu = WiLobe.Mu;
+	// The lobe remains on the same side of the surface, mirrored over the normal.
+	WoLobe.Mu = float3(-WiLobe.Mu.xy, WiLobe.Mu.z);
 
 	WoLobe.Sigma = WiLobe.Sigma + StrataLobeRoughnessToVariance(InterfaceRoughness);
 
@@ -80,7 +94,9 @@ FStrataLobeStatistic StrataGetRefractedLobe(FStrataLobeStatistic WiLobe, float3
 
 	WoLobe.E = WiLobe.E * (1.0f - InterfaceFDG);
 
-	WoLobe.Mu.xy = WiLobe.Mu.xy * InterfaceEta12;
+	// The lobe is on the oposite side of the surface, and the direction account for change of medium IOR.
+	WoLobe.Mu.xy = -WiLobe.Mu.xy * InterfaceEta12;
+	WoLobe.Mu.z  = -sign(WiLobe.Mu.z) * sqrt(1.0 - dot(WoLobe.Mu.xy, WoLobe.Mu.xy)); // derive z from vector projected xy
 
 	//	const float S = 0.5f * (1.0f + InterfaceEta12 * (WiLobe.Mu.z / WoLobe.Mu.z));				// This respect eq.10 but it goes crazy and does not respect the roughness of a single front layer
 	//	const float S = 0.5f * (1.0f + InterfaceEta12 * max(0.0, WiLobe.Mu.z / WoLobe.Mu.z));		// This respect the roughness range better but looks incorrect
@@ -96,8 +112,8 @@ FStrataLobeStatistic StrataGetTransmittedLobe(FStrataLobeStatistic WiLobe, float
 
 	WoLobe.E = WiLobe.E * exp(-OpticalDepth);
 
-	// Continue straight forward
-	WoLobe.Mu = WiLobe.Mu;
+	// The lobe is on the oposite side of the surface along -Wi
+	WoLobe.Mu = -WiLobe.Mu.xyz;
 
 	WoLobe.Sigma = WiLobe.Sigma;
 
@@ -112,7 +128,8 @@ FStrataLobeStatistic StrataWeightLobe(FStrataLobeStatistic A, float Weight)
 	//	- a dirac lobe with visibility = 1 - coverage
 	FStrataLobeStatistic WoLobe = A;
 	WoLobe.E *= Weight;
-	WoLobe.Mu *= Weight;
+	WoLobe.Mu.xy = WoLobe.Mu.xy * Weight;
+	WoLobe.Mu.z = sign(WoLobe.Mu.z) * sqrt(1.0 - dot(WoLobe.Mu.xy, WoLobe.Mu.xy)); // derive z from vector projected xy
 	WoLobe.Sigma *= Weight;
 	return WoLobe;
 }
@@ -122,6 +139,7 @@ FStrataLobeStatistic StrataHorizontalMixLobes(FStrataLobeStatistic A, FStrataLob
 	FStrataLobeStatistic WoLobe;
 	WoLobe.E = lerp(A.E, B.E, Mix);
 	WoLobe.Mu = lerp(A.Mu, B.Mu, Mix);
+	WoLobe.Mu.z = sign(WoLobe.Mu.z) * sqrt(1.0 - dot(WoLobe.Mu.xy, WoLobe.Mu.xy)); // derive z from vector projected xy
 	WoLobe.Sigma = lerp(A.Sigma, B.Sigma, Mix);
 	return WoLobe;
 }
diff --git a/Engine/Shaders/Private/Strata/StrataTree.ush b/Engine/Shaders/Private/Strata/StrataTree.ush
--- a/Engine/Shaders/Private/Strata/StrataTree.ush
+++ b/Engine/Shaders/Private/Strata/StrataTree.ush
@@ -104,21 +104,22 @@ void UpdateSingleBSDFOperatorCoverageTransmittance(
 	const FBxDFEnergyTerms EnergyTerms = ComputeGGXSpecEnergyTerms(Roughness, StrataBSDFContext.SatNoV, StrataGetBSDFSpecularColor(CurrentBSDF), StrataGetBSDFSpecularF90(CurrentBSDF));
 	{
 		const float TopDownInterfaceEta = CurrentBSDF.bIsTop ? AIR_IOR / WATER_IOR : 1.0f;
-		FStrataLobeStatistic RefractedLobe = StrataGetRefractedLobe(StrataGetViewDiracLobe(-V), EnergyTerms.E, Roughness, TopDownInterfaceEta);
+		FStrataLobeStatistic RefractedLobe = StrataGetRefractedLobe(StrataGetDiracLobe(V), EnergyTerms.E, Roughness, TopDownInterfaceEta);
 		// Account for transmittance of the medium (but ignoring scattering for now)
 		RefractedLobe.E *= CurrentBSDFTransmittance;
-		// Reduce according to coverage her. It is a approximation to handle rough refraction reduction according to coverage when relying on a single lob.
+		// Reduce according to coverage. It is a approximation to handle rough refraction reduction according to coverage when relying on a single lob.
 		RefractedLobe = StrataWeightLobe(RefractedLobe, CurrentBSDFCoverage);
 		Op.TopDownRefractionLobeStat = RefractedLobe;
 	}
+
 	// Same operation but for a bottom up lob
 	{
 		const float BottomUpInterfaceEta = CurrentBSDF.bIsTop ? WATER_IOR / AIR_IOR : 1.0f;
-		const float3 Up = float3(0, 0, 1);
-		FStrataLobeStatistic RefractedLobe = StrataGetRefractedLobe(StrataGetViewDiracLobe(Up), EnergyTerms.E, Roughness, BottomUpInterfaceEta);
+		const float3 NormalDown = float3(0, 0, -1);
+		FStrataLobeStatistic RefractedLobe = StrataGetRefractedLobe(StrataGetDiracLobe(NormalDown), EnergyTerms.E, Roughness, BottomUpInterfaceEta);
 		// Account for transmittance of the medium (but ignoring scattering for now)
 		RefractedLobe.E *= CurrentBSDFTransmittance;
-		// Reduce according to coverage her. It is a approximation to handle rough refraction reduction according to coverage when relying on a single lob.
+		// Reduce according to coverage. It is a approximation to handle rough refraction reduction according to coverage when relying on a single lob.
 		RefractedLobe = StrataWeightLobe(RefractedLobe, CurrentBSDFCoverage);
 		Op.BottomUpRefractionLobeStat = RefractedLobe;
 	}
@@ -200,7 +201,8 @@ void UpdateSingleOperatorCoverageTransmittance(inout FStrataTree StrataTree, int
 		// Project top lob through bottom layer assuming medium have same IOR.
 		{
 			const float TopDownInterfaceEta = 1.0f;
-			Op.TopDownRefractionLobeStat = StrataGetRefractedLobe(OpA.TopDownRefractionLobeStat, /*InterfaceFDG*/0.0f, StrataLobeVarianceToRoughness(OpB.TopDownRefractionLobeStat.Sigma), TopDownInterfaceEta);
+			Op.TopDownRefractionLobeStat = StrataGetRefractedLobe(StrataOppositeLobe(OpA.TopDownRefractionLobeStat), 
+				/*InterfaceFDG*/0.0f, StrataLobeVarianceToRoughness(OpB.TopDownRefractionLobeStat.Sigma), TopDownInterfaceEta);
 			Op.TopDownRefractionLobeStat.E *= saturate(OpB.Transmittance + (1.0f - OpB.Coverage)); // Combine with botton layer medium and specular transmittance
 
 			// The ratio of BottomSurfaceOnlyVisibile and visible through top layer (with a scale of 0.5 because at maximum both surfaces will be visible) as a fonction of the total coverage.
@@ -211,7 +213,8 @@ void UpdateSingleOperatorCoverageTransmittance(inout FStrataTree StrataTree, int
 		// Project bottom lob through top layer.
 		{
 			const float BottomUpInterfaceEta = Op.bIsTop ? WATER_IOR / AIR_IOR : 1.0f;
-			Op.BottomUpRefractionLobeStat = StrataGetRefractedLobe(OpB.BottomUpRefractionLobeStat, /*InterfaceFDG*/0.0f, StrataLobeVarianceToRoughness(OpA.BottomUpRefractionLobeStat.Sigma), BottomUpInterfaceEta);
+			Op.BottomUpRefractionLobeStat = StrataGetRefractedLobe(StrataOppositeLobe(OpB.BottomUpRefractionLobeStat), 
+				/*InterfaceFDG*/0.0f, StrataLobeVarianceToRoughness(OpA.BottomUpRefractionLobeStat.Sigma), BottomUpInterfaceEta);
 			Op.BottomUpRefractionLobeStat.E *= saturate(OpA.Transmittance + (1.0f - OpA.Coverage)); // Combine with top layer medium and specular transmittance
 		}
 
@@ -333,16 +336,16 @@ void UpdateAllBSDFWithBottomUpOperatorVisit_Vertical(
 		// The outcome is that we affect the bottom layer roughness to aggregate that fact.
 		// STRATA_TODO: take into account thickness (far-field for now).
 
-		// Evaluate the reflected lobe
+		// Evaluate the reflected lobe, as an approximation we only evaluate along the normal.
 		const float3 DummyWi = float3(0, 0, 1);
 		const float DummyInterfaceDFG = 0.5f;
-		FStrataLobeStatistic ReflectedLobe = StrataGetReflectedLobe(StrataGetViewDiracLobe(DummyWi), DummyInterfaceDFG, StrataGetBSDFRoughness(CurrentBSDF));
+		FStrataLobeStatistic ReflectedLobe = StrataGetReflectedLobe(StrataGetDiracLobe(DummyWi), DummyInterfaceDFG, StrataGetBSDFRoughness(CurrentBSDF));
 
 		// Evaluate the reflected lobe after refraction through the top layer of the vertical operator.
 		// => This tells how light should be integrated when traveling up thorugh that layer.
 		const float TopLayerSigma = Op.VerticalTop_BottomUpRefractionLobeStat.Sigma;
 		const float BottomUpInterfaceEta = Op.bIsTop ? WATER_IOR / AIR_IOR : 1.0f;
-		FStrataLobeStatistic RefractedLobe = StrataGetRefractedLobe(ReflectedLobe, DummyInterfaceDFG, StrataLobeVarianceToRoughness(TopLayerSigma), BottomUpInterfaceEta);
+		FStrataLobeStatistic RefractedLobe = StrataGetRefractedLobe(StrataOppositeLobe(ReflectedLobe), DummyInterfaceDFG, StrataLobeVarianceToRoughness(TopLayerSigma), BottomUpInterfaceEta);/////////////////
 
 		StrataSetBSDFRoughness(CurrentBSDF, StrataLobeVarianceToRoughness(RefractedLobe.Sigma));
 	}
