* Add screen-space line drawing for 2D drawing
* Fix cross drawing
* Add Referential variant (pos + float3x3)

diff --git a/Engine/Shaders/Private/HairStrands/HairStrandsDeepShadowDebug.usf b/Engine/Shaders/Private/HairStrands/HairStrandsDeepShadowDebug.usf
--- a/Engine/Shaders/Private/HairStrands/HairStrandsDeepShadowDebug.usf
+++ b/Engine/Shaders/Private/HairStrands/HairStrandsDeepShadowDebug.usf
@@ -177,20 +177,20 @@ void MainCS(uint3 DispatchThreadId : SV_DispatchThreadID)
 		P21 = mul(P21, ClipToTranslatedWorld); P21 /= P21.w;
 		P31 = mul(P31, ClipToTranslatedWorld); P31 /= P31.w;
 
-		AddLineTWS(P00, P10, FrustumColor, FrustumColor);
-		AddLineTWS(P10, P20, FrustumColor, FrustumColor);
-		AddLineTWS(P20, P30, FrustumColor, FrustumColor);
-		AddLineTWS(P30, P00, FrustumColor, FrustumColor);
-
-		AddLineTWS(P00, P01, FrustumColor, FrustumColor);
-		AddLineTWS(P10, P11, FrustumColor, FrustumColor);
-		AddLineTWS(P20, P21, FrustumColor, FrustumColor);
-		AddLineTWS(P30, P31, FrustumColor, FrustumColor);
-
-		AddLineTWS(P01, P11, FrustumColor, FrustumColor);
-		AddLineTWS(P11, P21, FrustumColor, FrustumColor);
-		AddLineTWS(P21, P31, FrustumColor, FrustumColor);
-		AddLineTWS(P31, P01, FrustumColor, FrustumColor);
+		AddLineTWS(P00.xyz, P10.xyz, FrustumColor, FrustumColor);
+		AddLineTWS(P10.xyz, P20.xyz, FrustumColor, FrustumColor);
+		AddLineTWS(P20.xyz, P30.xyz, FrustumColor, FrustumColor);
+		AddLineTWS(P30.xyz, P00.xyz, FrustumColor, FrustumColor);
+
+		AddLineTWS(P00.xyz, P01.xyz, FrustumColor, FrustumColor);
+		AddLineTWS(P10.xyz, P11.xyz, FrustumColor, FrustumColor);
+		AddLineTWS(P20.xyz, P21.xyz, FrustumColor, FrustumColor);
+		AddLineTWS(P30.xyz, P31.xyz, FrustumColor, FrustumColor);
+
+		AddLineTWS(P01.xyz, P11.xyz, FrustumColor, FrustumColor);
+		AddLineTWS(P11.xyz, P21.xyz, FrustumColor, FrustumColor);
+		AddLineTWS(P21.xyz, P31.xyz, FrustumColor, FrustumColor);
+		AddLineTWS(P31.xyz, P01.xyz, FrustumColor, FrustumColor);
 	}
 }
 #endif // SHADER_DOMINFO
diff --git a/Engine/Shaders/Private/ShaderPrint.ush b/Engine/Shaders/Private/ShaderPrint.ush
--- a/Engine/Shaders/Private/ShaderPrint.ush
+++ b/Engine/Shaders/Private/ShaderPrint.ush
@@ -713,6 +713,7 @@ struct FShaderDrawElement
 	float3 Pos1;
 	float4 Color0;
 	float4 Color1;
+	bool bIsScreenSpace;
 };
 
 FShaderDrawElement UnpackShaderElement(StructuredBuffer<uint> InPrimitiveBuffer, uint InIndex)
@@ -726,7 +727,7 @@ FShaderDrawElement UnpackShaderElement(StructuredBuffer<uint> InPrimitiveBuffer,
 		Out.Pos0.z = asfloat(InPrimitiveBuffer[Index8 + 2]);
 
 		const uint PackedColor0 = InPrimitiveBuffer[Index8 + 3];
-		Out.Color0 = float4((PackedColor0.x >> 24) & 0xFF, (PackedColor0.x >> 16) & 0xFF, (PackedColor0.x >> 8) & 0xFF, (PackedColor0.x) & 0xFF) / 255.0f;
+		Out.Color0 = float4((PackedColor0 >> 24) & 0xFF, (PackedColor0 >> 16) & 0xFF, (PackedColor0 >> 8) & 0xFF, (PackedColor0) & 0xFF) / 255.0f;
 	}
 	{
 		Out.Pos1.x = asfloat(InPrimitiveBuffer[Index8 + 4]);
@@ -734,7 +735,10 @@ FShaderDrawElement UnpackShaderElement(StructuredBuffer<uint> InPrimitiveBuffer,
 		Out.Pos1.z = asfloat(InPrimitiveBuffer[Index8 + 6]);
 
 		const uint PackedColor1 = InPrimitiveBuffer[Index8 + 7];
-		Out.Color1 = float4((PackedColor1.x >> 24) & 0xFF, (PackedColor1.x >> 16) & 0xFF, (PackedColor1.x >> 8) & 0xFF, (PackedColor1.x) & 0xFF) / 255.0f;
+		Out.Color1 = float4((PackedColor1 >> 24) & 0xFF, (PackedColor1 >> 16) & 0xFF, (PackedColor1 >> 8) & 0xFF, (PackedColor1) & 0xFF) / 255.0f;
+
+		// The 'space' info is stored into the alpha's LSB
+		Out.bIsScreenSpace = (PackedColor1 & 0x1) > 0u;
 	}
 	return Out;
 }
@@ -755,8 +759,12 @@ bool AllocateShaderDrawElement(uint Count, inout uint OutIndex)
 
 void AddShaderElement(FShaderDrawElement E, uint Index)
 {
-	const uint4 PackedC0 = uint4(255.0f * saturate(E.Color0));
-	const uint4 PackedC1 = uint4(255.0f * saturate(E.Color1));
+	uint4 PackedC0 = uint4(255.0f * saturate(E.Color0));
+	uint4 PackedC1 = uint4(255.0f * saturate(E.Color1));
+
+	// Store 'space' info into the alpha's LSB
+	PackedC1 = PackedC1 & 0xFE;
+	PackedC1 = PackedC1 | (E.bIsScreenSpace ? 0x1 : 0x0);
 
 	const uint Index8 = (1 + Index) * 8;
 	ShaderPrint_RWLinesBuffer[Index8 + 0] = asuint(E.Pos0.x);
@@ -788,6 +796,7 @@ void AddLineTWS(float3 Pos0, float3 Pos1, float4 Color0, float4 Color1)
 		Element.Pos1 = Pos1;
 		Element.Color0 = Color0;
 		Element.Color1 = Color1;
+		Element.bIsScreenSpace = false;
 		AddShaderElement(Element, Offset);
 	}
 }
@@ -805,6 +814,32 @@ void AddLineWS (FShaderPrintContext Ctx, float3 Pos0, float3 Pos1, float4 Color0
 void AddLine   (FShaderPrintContext Ctx, float3 Pos0, float3 Pos1, float4 Color)					{ if (Ctx.bIsActive) { AddLineTWS(Pos0+ShaderPrint.TranslatedWorldOffset, Pos1+ShaderPrint.TranslatedWorldOffset, Color); } }
 void AddLine   (FShaderPrintContext Ctx, float3 Pos0, float3 Pos1, float4 Color0, float4 Color1)	{ if (Ctx.bIsActive) { AddLineTWS(Pos0+ShaderPrint.TranslatedWorldOffset, Pos1+ShaderPrint.TranslatedWorldOffset, Color0, Color1); } }
 
+// Draw 2D line. Input positions are expressed in pixel coordinate
+void AddLineSS(float2 InPosInPixel0, float2 InPosInPixel1, float4 Color0, float4 Color1)
+{
+	const float2 Pos0 = float2(InPosInPixel0) / float2(ShaderPrint.Resolution);
+	const float2 Pos1 = float2(InPosInPixel1) / float2(ShaderPrint.Resolution);
+
+	uint Offset = 0;
+	if (AllocateShaderDrawElement(1, Offset))
+	{
+		FShaderDrawElement Element;
+		Element.Pos0 = float3(Pos0, 0);
+		Element.Pos1 = float3(Pos1, 0);
+		Element.Color0 = Color0;
+		Element.Color1 = Color1;
+		Element.bIsScreenSpace = true;
+		AddShaderElement(Element, Offset);
+	}
+}
+void AddLineSS(float2 Pos0, float2 Pos1, float4 Color)											{ AddLineSS(Pos0, Pos1, Color, Color); }
+void AddLine(float2 Pos0, float2 Pos1, float4 Color)											{ AddLineSS(Pos0, Pos1, Color); }
+void AddLine(float2 Pos0, float2 Pos1, float4 Color0, float4 Color1)							{ AddLineSS(Pos0, Pos1, Color0, Color1); }
+
+void AddLineSS(FShaderPrintContext Ctx, float2 Pos0, float2 Pos1, float4 Color0, float4 Color1) { if (Ctx.bIsActive) { AddLineSS(Pos0, Pos1, Color0, Color1); } }
+void AddLineSS(FShaderPrintContext Ctx, float2 Pos0, float2 Pos1, float4 Color)					{ if (Ctx.bIsActive) { AddLineSS(Pos0, Pos1, Color, Color); } }
+void AddLine(FShaderPrintContext Ctx, float2 Pos0, float2 Pos1, float4 Color)					{ if (Ctx.bIsActive) { AddLineSS(Pos0, Pos1, Color); } }
+void AddLine(FShaderPrintContext Ctx, float2 Pos0, float2 Pos1, float4 Color0, float4 Color1)	{ if (Ctx.bIsActive) { AddLineSS(Pos0, Pos1, Color0, Color1); } }
 
 ///////////////////////////////////////////////////////////////////////////////////////////////////
 // Quad
@@ -815,6 +850,7 @@ void AddQuadTWS(float3 Pos0, float3 Pos1, float3 Pos2, float3 Pos3, float4 Color
 	if (AllocateShaderDrawElement(4, Offset))
 	{
 		FShaderDrawElement Element;
+		Element.bIsScreenSpace = false;
 		Element.Color0 = Element.Color1 = Color;
 		Element.Pos0 = Pos0; Element.Pos1 = Pos1; AddShaderElement(Element, Offset + 0);
 		Element.Pos0 = Pos1; Element.Pos1 = Pos2; AddShaderElement(Element, Offset + 1);
@@ -848,6 +884,7 @@ void AddAABBTWS(float3 Min, float3 Max, float4 Color)
 		float3 P7 = float3(Min.x, Max.y, Max.z);
 
 		FShaderDrawElement Element;
+		Element.bIsScreenSpace = false;
 		Element.Color0 = Element.Color1 = Color;
 
 		Element.Pos0 = P0; Element.Pos1 = P1; AddShaderElement(Element, Offset + 0);
@@ -884,12 +921,12 @@ void AddCrossTWS(float3 Pos, float Size, float4 Color)
 	AddLineTWS(Pos - float3(0,0,Size), Pos + float3(0,0,Size), Color, Color);
 }
 
-void AddCrossWS(float3 Pos, float Size, float4 Color)							{ AddAABBTWS(Pos + ShaderPrint.TranslatedWorldOffset, Size, Color); }
-void AddCross  (float3 Pos, float Size, float4 Color)							{ AddAABBTWS(Pos + ShaderPrint.TranslatedWorldOffset, Size, Color); }
+void AddCrossWS(float3 Pos, float Size, float4 Color)							{ AddCrossTWS(Pos + ShaderPrint.TranslatedWorldOffset, Size, Color); }
+void AddCross  (float3 Pos, float Size, float4 Color)							{ AddCrossTWS(Pos + ShaderPrint.TranslatedWorldOffset, Size, Color); }
 
 void AddCrossTWS(FShaderPrintContext Ctx, float3 Pos, float Size, float4 Color) { if (Ctx.bIsActive) { AddCrossTWS(Pos, Size, Color); } }
-void AddCrossWS (FShaderPrintContext Ctx, float3 Pos, float Size, float4 Color)	{ if (Ctx.bIsActive) { AddAABBTWS(Pos + ShaderPrint.TranslatedWorldOffset, Size, Color); } }
-void AddCross   (FShaderPrintContext Ctx, float3 Pos, float Size, float4 Color)	{ if (Ctx.bIsActive) { AddAABBTWS(Pos + ShaderPrint.TranslatedWorldOffset, Size, Color); } }
+void AddCrossWS (FShaderPrintContext Ctx, float3 Pos, float Size, float4 Color)	{ if (Ctx.bIsActive) { AddCrossTWS(Pos + ShaderPrint.TranslatedWorldOffset, Size, Color); } }
+void AddCross   (FShaderPrintContext Ctx, float3 Pos, float Size, float4 Color)	{ if (Ctx.bIsActive) { AddCrossTWS(Pos + ShaderPrint.TranslatedWorldOffset, Size, Color); } }
 
 ///////////////////////////////////////////////////////////////////////////////////////////////////
 // Referential
@@ -918,20 +955,25 @@ void AddReferentialTWS(float4x4 InM, float Scale = 1)
 	AddReferentialTWS(InM[3].xyz, InM[0].xyz, InM[1].xyz, InM[2].xyz, Scale);
 }
 
+void AddReferentialWS(float3 Pos, float3x3 InM, float Scale = 1)					{ AddReferentialTWS(Pos + ShaderPrint.TranslatedWorldOffset, InM[0].xyz, InM[1].xyz, InM[2].xyz, Scale); }
 void AddReferentialWS(float4x4 InM, float Scale = 1)								{ AddReferentialTWS(InM[3].xyz + ShaderPrint.TranslatedWorldOffset, InM[0].xyz, InM[1].xyz, InM[2].xyz, Scale); }
 void AddReferentialWS(float3 Pos, float3 TangentZ, float Scale = 1)					{ AddReferentialTWS(Pos + ShaderPrint.TranslatedWorldOffset, TangentZ, Scale); }
 void AddReferentialWS(float3 Pos, float3 T, float3 B, float3 N, float Scale = 1)	{ AddReferentialTWS(Pos + ShaderPrint.TranslatedWorldOffset, T, B, N, Scale); }
+void AddReferential  (float3 Pos, float3x3 InM, float Scale = 1)					{ AddReferentialTWS(Pos + ShaderPrint.TranslatedWorldOffset, InM[0].xyz, InM[1].xyz, InM[2].xyz, Scale); }
 void AddReferential  (float4x4 InM, float Scale = 1)								{ AddReferentialTWS(InM[3].xyz + ShaderPrint.TranslatedWorldOffset, InM[0].xyz, InM[1].xyz, InM[2].xyz, Scale); }
 void AddReferential  (float3 Pos, float3 TangentZ, float Scale = 1)					{ AddReferentialTWS(Pos + ShaderPrint.TranslatedWorldOffset, TangentZ, Scale); }
 void AddReferential  (float3 Pos, float3 T, float3 B, float3 N, float Scale = 1)	{ AddReferentialTWS(Pos + ShaderPrint.TranslatedWorldOffset, T, B, N, Scale); }
 
 void AddReferentialTWS(FShaderPrintContext Ctx, float3 Pos, float3 T, float3 B, float3 N, float Scale = 1)	{ if (Ctx.bIsActive) { AddReferentialTWS(Pos, T, B, N, Scale); } }
 void AddReferentialTWS(FShaderPrintContext Ctx, float3 Pos, float3 TangentZ, float Scale = 1)				{ if (Ctx.bIsActive) { AddReferentialTWS(Pos, TangentZ, Scale); }  }
+void AddReferentialTWS(FShaderPrintContext Ctx, float3 Pos, float3x3 InM, float Scale = 1)					{ if (Ctx.bIsActive) { AddReferentialTWS(Pos, InM[0].xyz, InM[1].xyz, InM[2].xyz, Scale); }  }
 void AddReferentialTWS(FShaderPrintContext Ctx, float4x4 InM, float Scale = 1)								{ if (Ctx.bIsActive) { AddReferentialTWS(InM[3].xyz, InM[0].xyz, InM[1].xyz, InM[2].xyz, Scale); }  }
 void AddReferentialWS (FShaderPrintContext Ctx, float4x4 InM, float Scale = 1)								{ if (Ctx.bIsActive) { AddReferentialTWS(InM[3].xyz + ShaderPrint.TranslatedWorldOffset, InM[0].xyz, InM[1].xyz, InM[2].xyz, Scale); }  }
+void AddReferentialWS (FShaderPrintContext Ctx, float3 Pos, float3x3 InM, float Scale = 1)					{ if (Ctx.bIsActive) { AddReferentialTWS(Pos + ShaderPrint.TranslatedWorldOffset, InM[0].xyz, InM[1].xyz, InM[2].xyz, Scale); }  }
 void AddReferentialWS (FShaderPrintContext Ctx, float3 Pos, float3 TangentZ, float Scale = 1)				{ if (Ctx.bIsActive) { AddReferentialTWS(Pos + ShaderPrint.TranslatedWorldOffset, TangentZ, Scale); }  }
 void AddReferentialWS (FShaderPrintContext Ctx, float3 Pos, float3 T, float3 B, float3 N, float Scale = 1)	{ if (Ctx.bIsActive) { AddReferentialTWS(Pos + ShaderPrint.TranslatedWorldOffset, T, B, N, Scale); }  }
 void AddReferential   (FShaderPrintContext Ctx, float4x4 InM, float Scale = 1)								{ if (Ctx.bIsActive) { AddReferentialTWS(InM[3].xyz + ShaderPrint.TranslatedWorldOffset, InM[0].xyz, InM[1].xyz, InM[2].xyz, Scale); }  }
+void AddReferential   (FShaderPrintContext Ctx, float3 Pos, float3x3 InM, float Scale = 1)					{ if (Ctx.bIsActive) { AddReferentialTWS(Pos + ShaderPrint.TranslatedWorldOffset, InM[0].xyz, InM[1].xyz, InM[2].xyz, Scale); }  }
 void AddReferential   (FShaderPrintContext Ctx, float3 Pos, float3 TangentZ, float Scale = 1)				{ if (Ctx.bIsActive) { AddReferentialTWS(Pos + ShaderPrint.TranslatedWorldOffset, TangentZ, Scale); }  }
 void AddReferential   (FShaderPrintContext Ctx, float3 Pos, float3 T, float3 B, float3 N, float Scale = 1)	{ if (Ctx.bIsActive) { AddReferentialTWS(Pos + ShaderPrint.TranslatedWorldOffset, T, B, N, Scale); }  }
 ///////////////////////////////////////////////////////////////////////////////////////////////////
@@ -969,6 +1011,7 @@ void AddOBBTWS(float3 Min, float3 Max, float4 Color, float4x4 Transform)
 		float3 P7 = mul(float4(Min.x, Max.y, Max.z, 1.0f), Transform).xyz;
 
 		FShaderDrawElement Element;
+		Element.bIsScreenSpace = false;
 		Element.Color0 = Element.Color1 = Color;
 
 		Element.Pos0 = P0; Element.Pos1 = P1; AddShaderElement(Element, Offset + 0);
diff --git a/Engine/Shaders/Private/ShaderPrintDrawPrimitive.usf b/Engine/Shaders/Private/ShaderPrintDrawPrimitive.usf
--- a/Engine/Shaders/Private/ShaderPrintDrawPrimitive.usf
+++ b/Engine/Shaders/Private/ShaderPrintDrawPrimitive.usf
@@ -25,17 +25,29 @@ void ShaderDrawDebugVS(
 	out float4 Position : SV_POSITION,
 	out float4 ColorAlpha : TEXCOORD0)
 {
-	FShaderDrawElement Item = UnpackShaderElement(ShaderDrawDebugPrimitive, InstanceId);
+	const FShaderDrawElement Item = UnpackShaderElement(ShaderDrawDebugPrimitive, InstanceId);
 
 	float4 VertexPositionTWS = float4(VertexId == 0 ? Item.Pos0 : Item.Pos1, 1.0f);
 	float4 VertexColor = VertexId == 0 ? Item.Color0 : Item.Color1;
 
-	// Optional world offset. For 'locked' debug drawing, this offset translates the world 
-	// position recorded for a given translated world offset, into the current translated 
-	// world offset
-	VertexPositionTWS.xyz += TranslatedWorldOffsetConversion;
-
-	Position = mul(VertexPositionTWS, PrimaryView.TranslatedWorldToClip);
+	if (Item.bIsScreenSpace)
+	{
+		Position = VertexPositionTWS;
+		Position.xy = Position.xy * 2 - 1;
+		Position.y = -Position.y;
+
+		Position.z = 1;
+		Position.w = 1;
+	}
+	else
+	{
+		// Optional world offset. For 'locked' debug drawing, this offset translates the world 
+		// position recorded for a given translated world offset, into the current translated 
+		// world offset
+		VertexPositionTWS.xyz += TranslatedWorldOffsetConversion;
+
+		Position = mul(VertexPositionTWS, PrimaryView.TranslatedWorldToClip);
+	}
 	ColorAlpha = VertexColor;
 }
 #endif
