Strata - rough refraction update

diff --git a/Engine/Shaders/Private/Strata/StrataRoughRefraction.usf b/Engine/Shaders/Private/Strata/StrataRoughRefraction.usf
--- a/Engine/Shaders/Private/Strata/StrataRoughRefraction.usf
+++ b/Engine/Shaders/Private/Strata/StrataRoughRefraction.usf
@@ -37,6 +37,14 @@ float3 GetScreenWorldDir(in float4 SVPos)
 	return normalize(WorldPos.xyz - View.TranslatedWorldCameraOrigin);
 }
 
+float GetEta12()
+{
+	const float AirIOR = 1.0f;
+	const float WaterIOR = 1.33f;
+	const float EtaIn = WaterIOR / AirIOR;	// We always consider entering from air to water
+
+	return EtaIn;
+}
 
 bool GetRefractedDirection(float2 RandomUV, float a2, float3 V, inout float3 L )
 { 
@@ -44,11 +52,8 @@ bool GetRefractedDirection(float2 RandomUV, float a2, float3 V, inout float3 L )
 
 	L = 0;
 	float F = 0;
-	const float AirIOR = 1.0f;
-	const float WaterIOR = 1.33f;
-	const float EtaIn = WaterIOR / AirIOR;	// We always consider entering from air to water
 	const float RandSample = 1.0f;			// We only consider refracted rays
-	return SampleRefraction(-V, H.xyz, EtaIn, RandSample, L, F);
+	return SampleRefraction(-V, H.xyz, GetEta12(), RandSample, L, F);
 }
 
 
@@ -90,7 +95,7 @@ FUI GetUI(inout FShaderPrintContext Context)
 	return UI;
 }
 
-void DrawCircleXY(inout FShaderPrintContext Context, float3 Center, float Radius, float4 ColorAlpha)
+void DrawCircleXY(inout FShaderPrintContext Context, float3 Center, float2 Radius, float4 ColorAlpha)
 {
 	const float TStep = 0.1f;
 	for (float t = 0; t <= 1.0; t += TStep)
@@ -99,8 +104,8 @@ void DrawCircleXY(inout FShaderPrintContext Context, float3 Center, float Radius
 		float S1, C1;
 		sincos(t * 2 * PI, S0, C0);
 		sincos((t + TStep) * 2 * PI, S1, C1);
-		float3 P0 = Center + float3(C0, S0, 0.0) * Radius;
-		float3 P1 = Center + float3(C1, S1, 0.0) * Radius;
+		float3 P0 = Center + float3(C0, S0, 0.0) * float3(Radius, 0.0);
+		float3 P1 = Center + float3(C1, S1, 0.0) * float3(Radius, 0.0);
 
 		AddLineWS(Context, P0, P1, ColorAlpha);
 	}
@@ -159,9 +164,11 @@ void EvaluateRoughRefractionLobeCS(uint3 ThreadId : SV_DispatchThreadID)
 		FShaderPrintContext Context = GetUIContext(true);
 		FUI UI = GetUI(Context);
 
-
 		FLobeStat LobeState = (FLobeStat)0;
 
+		float Count = 0;
+		float2 Mean = 0;
+		float2 Variance = 0;
 
 		for (float u = 0.5f; u < float(TraceSqrtSampleCount); u++)
 		{
@@ -175,14 +182,45 @@ void EvaluateRoughRefractionLobeCS(uint3 ThreadId : SV_DispatchThreadID)
 					SampleCountTextureUAV[uint2(SampleUV * 64.0f)] += 1;
 
 					LobStatUpdate(LobeState, L.xy);
+
+					Count++;
+					Mean += L.xy;
 				}
 				else
 				{
-					// Discard refelctions
+					// Discard reflections
 				}
 			}
 		}
 
+/*		Mean = Mean / Count;
+		Count = 0;
+		for (float u = 0.5f; u < float(TraceSqrtSampleCount); u++)
+		{
+			for (float v = 0.5f; v < float(TraceSqrtSampleCount); v++)
+			{
+				float2 RandomUV = float2(u, v) / float(TraceSqrtSampleCount);
+				float3 L = 0;
+				if (GetRefractedDirection(RandomUV, UI.a2, UI.V, L))
+				{
+					float2 SampleUV = L.xy * 0.5 + 0.5;
+					SampleCountTextureUAV[uint2(SampleUV * 64.0f)] += 1;
+
+					LobStatUpdate(LobeState, L.xy);
+
+					Count++;
+					Variance += (L.xy - Mean) * (L.xy - Mean);
+				}
+				else
+				{
+					// Discard reflections
+				}
+			}
+		}
+		Variance = Variance / Count;
+		LobeState.Mean = Mean;
+		LobeState.Variance = Variance;*/
+
 		LobStatFinalize(LobeState);
 		LobStatisticsBufferUAV[0] = LobeState;
 	}
@@ -201,6 +239,27 @@ Texture2D<uint> SampleCountTexture;
 
 StructuredBuffer<FLobeStat> LobStatisticsBuffer;
 
+FStrataLobeStatistic StrataGetRefractedLobe2(FStrataLobeStatistic WiLobe, float3 InterfaceFDG, float InterfaceRoughness, float InterfaceEta12)
+{
+	FStrataLobeStatistic WoLobe;
+
+	WoLobe.E = WiLobe.E * (1.0f - InterfaceFDG);
+
+	// The lobe is on the oposite side of the surface, and the direction account for change of medium IOR.
+	WoLobe.Mu.xy = -WiLobe.Mu.xy * InterfaceEta12;
+	WoLobe.Mu.z = -sign(WiLobe.Mu.z) * sqrt(1.0 - dot(WoLobe.Mu.xy, WoLobe.Mu.xy)); // derive z from vector projected xy
+
+	// Equation 10, we have also noticed variance explosion at grazing angle so we make the inner product be 1 for now.
+	const float OmegaIDotN = 1.0f; //WiLobe.Mu.z;
+	const float OmegaTDotN = 1.0f; //WoLobe.Mu.z;
+	//const float S = 0.5f * (1.0f + InterfaceEta12 * (abs(OmegaIDotN) / abs(OmegaTDotN)));						// equation 10 from paper
+	const float S = 0.5f * (abs(OmegaTDotN * InterfaceEta12 - OmegaIDotN) / abs(OmegaTDotN * InterfaceEta12));	// equation from paper source code
+
+	WoLobe.Sigma = (WiLobe.Sigma / InterfaceEta12) + StrataLobeRoughnessToVariance(InterfaceRoughness * S);
+
+	return WoLobe;
+}
+
 void VisualizeRoughRefractionPS(
 	float4 SVPos : SV_POSITION,
 	out float4 OutColor : SV_Target0)
@@ -219,7 +278,7 @@ void VisualizeRoughRefractionPS(
 	float3 SlabEdgeX = float3(60.0f, 0.0f, 0.0f);
 	float3 SlabEdgeY = float3(0.0f, 60.0f, 0.0f);
 	float3 SlabEdgeZ = float3(0.0f, 0.0f, UI.SlabThickness);
-	float3 SlabCenter = CamWorldPos + View.ViewForward * 60.0f;
+	float3 SlabCenter = CamWorldPos + View.ViewForward * 35.0f;
 	float3 RefractionP= SlabCenter + 0.5 * SlabEdgeZ;
 	const float SlabTopZ    = SlabCenter.z  + 0.5 * UI.SlabThickness;
 	const float SlabBottomZ = SlabCenter.z - 0.5 * UI.SlabThickness;
@@ -301,16 +360,25 @@ void VisualizeRoughRefractionPS(
 	}
 
 	// Draw the lobe estimated by Belcour
-	const float AirIOR = 1.0f;
-	const float WaterIOR = 1.33f;
-	const float EtaIn = WaterIOR / AirIOR;	// We always consider entering from air to water
-	FStrataLobeStatistic RefractedLobe = StrataGetRefractedLobe(StrataGetViewDiracLobe(-UI.V), 1.0, UI.Roughness, EtaIn);
-	DrawCircleXY(Context, float3(SlabCenter.xy + RefractedLobe.Mu.xy * VisProjectedRayRadius, SlabBottomZ), RefractedLobe.Sigma * VisProjectedRayRadius, float4(0, 1, 1, 1));
+	FStrataLobeStatistic RefractedLobe = StrataGetRefractedLobe2(StrataGetDiracLobe(UI.V), 1.0, UI.Roughness, GetEta12());
+	DrawCircleXY(Context, float3(SlabCenter.xy + RefractedLobe.Mu.xy * VisProjectedRayRadius, SlabBottomZ), RefractedLobe.Sigma * VisProjectedRayRadius, float4(1, 0, 1, 1));
+
+	Newline(Context);
+	Newline(Context);
+	Print(Context, TEXT("Belcour variance = "));
+	Print(Context, RefractedLobe.Sigma, FontWhite, 5, 5);
 
 	// Our lobe
 	FLobeStat LobeState = LobStatisticsBuffer[0];
-	DrawCircleXY(Context, float3(SlabCenter.xy + LobeState.Mean * VisProjectedRayRadius, SlabBottomZ), LobeState.Variance.x * VisProjectedRayRadius, float4(1, 1, 1, 1));
-	DrawCircleXY(Context, float3(SlabCenter.xy + LobeState.Mean * VisProjectedRayRadius, SlabBottomZ+1.0), LobeState.Variance.y * VisProjectedRayRadius, float4(0, 1, 0, 1));
+	DrawCircleXY(Context, float3(SlabCenter.xy + LobeState.Mean * VisProjectedRayRadius, SlabBottomZ + 0.05), LobeState.Variance * VisProjectedRayRadius, float4(1, 1, 0, 1));
+
+	Newline(Context);
+	Newline(Context);
+	Print(Context, TEXT("New variance X = "));
+	Print(Context, LobeState.Variance.x, FontWhite, 5, 5);
+	Newline(Context);
+	Print(Context, TEXT("New variance Y = "));
+	Print(Context, LobeState.Variance.y, FontWhite, 5, 5);
 
 }
 
