Fix a bug of reflection is too bright on mobile deferred with static lighting disabled.

diff --git a/Engine/Shaders/Private/MobileDeferredShading.usf b/Engine/Shaders/Private/MobileDeferredShading.usf
--- a/Engine/Shaders/Private/MobileDeferredShading.usf
+++ b/Engine/Shaders/Private/MobileDeferredShading.usf
@@ -341,7 +341,9 @@ void MobileDirectLightPS(
 #elif APPLY_SKY_REFLECTION
 	float SkyAverageBrightness = 1.0f;
 	float3 SpecularIBL = GetSkyLightReflection(ReflectionVector, GBuffer.Roughness, SkyAverageBrightness);
+#if ALLOW_STATIC_LIGHTING
 	SpecularIBL *= ComputeMixingWeight(GBuffer.IndirectIrradiance, SkyAverageBrightness, GBuffer.Roughness);
+#endif
 	MobileLightAccumulator_Add(MobileLightAccumulator, 0.0f, SpecularIBL * ShadingModelContext.EnvBrdf, 1.0f);
 #endif
 	half3 SkyDiffuseLighting = GetSkySHDiffuseSimple(GBuffer.WorldNormal);
