Fix errors compiling NaniteVertexFactory for use in Raytracing HitGroup.
- Use LightmapUVType to support both float2 and FloatDeriv2 permutations.

diff --git a/Engine/Shaders/Private/RayTracing/RayTracingMaterialHitShaders.usf b/Engine/Shaders/Private/RayTracing/RayTracingMaterialHitShaders.usf
--- a/Engine/Shaders/Private/RayTracing/RayTracingMaterialHitShaders.usf
+++ b/Engine/Shaders/Private/RayTracing/RayTracingMaterialHitShaders.usf
@@ -36,7 +36,7 @@ void GetSkyLighting(
 	FMaterialPixelParameters MaterialParameters,
 	VTPageTableResult LightmapVTPageTableResult,
 	float3 WorldNormal,
-	float2 LightmapUV,
+	LightmapUVType LightmapUV,
 	uint LightmapDataIndex,
 	float3 SkyOcclusionUV3D,
 	uint ShadingModel,
@@ -72,7 +72,7 @@ void GetSkyLighting(
 			#elif HQ_TEXTURE_LIGHTMAP
 
 				// Bent normal from precomputed texture
-				float4 WorldSkyBentNormalAndOcclusion = GetSkyBentNormalAndOcclusion(LightmapVTPageTableResult, LightmapUV * float2(1, 2), LightmapDataIndex, MaterialParameters.SvPosition.xy);
+				float4 WorldSkyBentNormalAndOcclusion = GetSkyBentNormalAndOcclusion(LightmapVTPageTableResult, ScaleLightmapUV(LightmapUV, float2(1.0f, 2.0f)), LightmapDataIndex, MaterialParameters.SvPosition.xy);
 				// Renormalize as vector was quantized and compressed
 				float3 NormalizedBentNormal = normalize(WorldSkyBentNormalAndOcclusion.xyz);
 				SkyVisibility = WorldSkyBentNormalAndOcclusion.w;
@@ -139,11 +139,11 @@ void GetPrecomputedIndirectLightingAndSkyLight(
 	OutIndirectIrradiance = 0;
 	OutDiffuseLighting = float3(0,0,0);
 	OutSubsurfaceLighting = 0;
-	float2 SkyOcclusionUV = 0;
+	LightmapUVType SkyOcclusionUV = (LightmapUVType)0;
 	uint SkyOcclusionDataIndex = 0u;
 
 #if HQ_TEXTURE_LIGHTMAP || LQ_TEXTURE_LIGHTMAP
-		float2 LightmapUV0, LightmapUV1;
+		LightmapUVType LightmapUV0, LightmapUV1;
 		uint LightmapDataIndex;
 		GetLightMapCoordinates(Interpolants, LightmapUV0, LightmapUV1, LightmapDataIndex);
 	#if LIGHTMAP_VT_ENABLED
