Strata - rough refraction, added refelction mode to test validity.

diff --git a/Engine/Shaders/Private/Strata/StrataRoughRefraction.usf b/Engine/Shaders/Private/Strata/StrataRoughRefraction.usf
--- a/Engine/Shaders/Private/Strata/StrataRoughRefraction.usf
+++ b/Engine/Shaders/Private/Strata/StrataRoughRefraction.usf
@@ -11,10 +11,14 @@
 #include "/Engine/Private/ColorMap.ush"
 #include "/Engine/Private/BRDF.ush"
 #include "/Engine/Private/MonteCarlo.ush"
+
 #include "/Engine/Private/PathTracing/Material/PathTracingFresnel.ush"
+#include "/Engine/Private/PathTracing/Material/PathTracingGlossy.ush"
 
 #include "/Engine/Private/Strata/StrataStatisticalOperators.ush"
 
+#define MODE_REFLECTION 1
+
 // Updated from http://jcgt.org/published/0007/03/04/
 bool Slabs(float3 P0, float3 P1, float3 RayOrigin, float3 InvRaydir, out float OutTMin, out float OutTMax)
 {
@@ -46,14 +50,32 @@ float GetEta12()
 	return EtaIn;
 }
 
-bool GetRefractedDirection(float2 RandomUV, float a2, float3 V, inout float3 L )
+#define WOLRD_NORMAL float3(0, 0, 1)
+
+// returns true the direction is valid for the current mode (refraction or reflection)
+bool GetGGXDirection(float2 RandomUV, float Roughness, float3 V, inout float3 L, inout float Value)
 { 
-	float4 H = ImportanceSampleVisibleGGX(UniformSampleDisk(RandomUV), a2, V);
+	float2 Alpha; 
+	GetGGXBasis(Roughness, WOLRD_NORMAL, Alpha);
+
+	float4 H = ImportanceSampleVisibleGGX_aniso(UniformSampleDisk(RandomUV), Alpha, V);
 
+	Value = 0;
 	L = 0;
 	float F = 0;
-	const float RandSample = 1.0f;			// We only consider refracted rays
-	return SampleRefraction(-V, H.xyz, GetEta12(), RandSample, L, F);
+	const float RandSample = MODE_REFLECTION ? -0.01f : 1.0f;
+	const bool bIsRefraction = SampleRefraction(-V, H.xyz, GetEta12(), RandSample, L, F);
+
+	if (bIsRefraction)
+	{
+		// Refraction
+		Value = GGXEvalRefraction(L, V, H.xyz, Alpha, GetEta12()).x;
+		return MODE_REFLECTION == 0;
+	}
+
+	// Reflection
+	Value = GGXEvalReflection(L, V, H.xyz, Alpha).x;
+	return MODE_REFLECTION == 1;
 }
 
 
@@ -68,7 +90,6 @@ struct FUI
 
 	// Derived
 	float a2;
-	float3 N;
 	float3 V;
 };
 
@@ -89,7 +110,6 @@ FUI GetUI(inout FShaderPrintContext Context)
 	Newline(Context);
 
 	UI.a2 = Pow4(UI.Roughness);
-	UI.N = float3(0, 0, 1);
 	UI.V = float3(sin(UI.IncidentAngle), 0, cos(UI.IncidentAngle));
 
 	return UI;
@@ -176,7 +196,8 @@ void EvaluateRoughRefractionLobeCS(uint3 ThreadId : SV_DispatchThreadID)
 			{
 				float2 RandomUV = float2(u, v) / float(TraceSqrtSampleCount);
 				float3 L = 0;
-				if (GetRefractedDirection(RandomUV, UI.a2, UI.V, L))
+				float Value = 0;
+				if (GetGGXDirection(RandomUV, UI.Roughness, UI.V, L, Value) && Value > 0.0f)
 				{
 					float2 SampleUV = L.xy * 0.5 + 0.5;
 					SampleCountTextureUAV[uint2(SampleUV * 64.0f)] += 1;
@@ -186,10 +207,6 @@ void EvaluateRoughRefractionLobeCS(uint3 ThreadId : SV_DispatchThreadID)
 					Count++;
 					Mean += L.xy;
 				}
-				else
-				{
-					// Discard reflections
-				}
 			}
 		}
 
@@ -201,7 +218,8 @@ void EvaluateRoughRefractionLobeCS(uint3 ThreadId : SV_DispatchThreadID)
 			{
 				float2 RandomUV = float2(u, v) / float(TraceSqrtSampleCount);
 				float3 L = 0;
-				if (GetRefractedDirection(RandomUV, UI.a2, UI.V, L))
+				float3 Value = 0;
+				if (GetGGXDirection(RandomUV, UI.Roughness, UI.V, L, Value) && Value > 0.0f)
 				{
 					float2 SampleUV = L.xy * 0.5 + 0.5;
 					SampleCountTextureUAV[uint2(SampleUV * 64.0f)] += 1;
@@ -211,10 +229,6 @@ void EvaluateRoughRefractionLobeCS(uint3 ThreadId : SV_DispatchThreadID)
 					Count++;
 					Variance += (L.xy - Mean) * (L.xy - Mean);
 				}
-				else
-				{
-					// Discard reflections
-				}
 			}
 		}
 		Variance = Variance / Count;
@@ -285,7 +299,7 @@ void VisualizeRoughRefractionPS(
 
 	const float VisProjectedRayRadius = 15.0f;
 
-	AddLineWS(Context, RefractionP, RefractionP + UI.N * 25.0, float4(0, 0, 1, 1));
+	AddLineWS(Context, RefractionP, RefractionP + WOLRD_NORMAL * 25.0, float4(0, 0, 1, 1));
 	AddLineWS(Context, RefractionP, RefractionP + UI.V * 25.0, float4(0, 1, 1, 1));
 
 	float SlabTMin = -1.0f;
@@ -322,7 +336,7 @@ void VisualizeRoughRefractionPS(
 		const float OmegaOutLen = length(OmegaOut);
 		const float3 OmegaOutNorm = OmegaOut / OmegaOutLen;
 
-		const float NoH = dot(normalize(OmegaOutNorm + UI.V), UI.N);
+		const float NoH = dot(normalize(OmegaOutNorm + UI.V), WOLRD_NORMAL);
 		float D = D_GGX(UI.a2, NoH);
 		OutColor = float4(D, D, 0.0, 0.0);
 	}
@@ -347,26 +361,31 @@ void VisualizeRoughRefractionPS(
 		{
 			float2 RandomUV = float2(u, v) / TraceDomainSize;
 			float3 L = 0;
-			if (GetRefractedDirection(RandomUV, UI.a2, UI.V, L))
+			float Value = 0;
+			if (GetGGXDirection(RandomUV, UI.Roughness, UI.V, L, Value))
 			{
-				AddLineWS(Context, RefractionP, RefractionP + L * VisProjectedRayRadius, float4(0, 1, 0, 0.3));
-				//AddLineWS(Context, RefractionP + L * VisProjectedRayRadius, float3((RefractionP + L * VisProjectedRayRadius).xy, SlabCenter.z - 0.5* UI.SlabThickness), float4(0, 1, 0, 0.3));
+				if(Value > 0.0f) AddLineWS(Context, RefractionP, RefractionP + L * VisProjectedRayRadius, float4(0, Value, 0, 0.3));
+				//if(Value > 0.0f) AddLineWS(Context, RefractionP + L * VisProjectedRayRadius, float3((RefractionP + L * VisProjectedRayRadius).xy, SlabCenter.z - 0.5* UI.SlabThickness), float4(0, 1, 0, 0.3));
 			}
 			else
 			{
-				AddLineWS(Context, RefractionP, RefractionP + L * VisProjectedRayRadius, float4(1, 0, 0, 0.3));
+				if (Value > 0.0f) AddLineWS(Context, RefractionP, RefractionP + L * VisProjectedRayRadius, float4(Value, 0, 0, 0.3));
 			}
 		}
 	}
 
 	// Draw the lobe estimated by Belcour
-	FStrataLobeStatistic RefractedLobe = StrataGetRefractedLobe2(StrataGetDiracLobe(UI.V), 1.0, UI.Roughness, GetEta12());
-	DrawCircleXY(Context, float3(SlabCenter.xy + RefractedLobe.Mu.xy * VisProjectedRayRadius, SlabBottomZ), RefractedLobe.Sigma * VisProjectedRayRadius, float4(1, 0, 1, 1));
+#if MODE_REFLECTION
+	FStrataLobeStatistic PaperLobe = StrataGetReflectedLobe(StrataGetDiracLobe(UI.V), 1.0, UI.Roughness);
+#else
+	FStrataLobeStatistic PaperLobe = StrataGetRefractedLobe2(StrataGetDiracLobe(UI.V), 1.0, UI.Roughness, GetEta12());
+#endif
+	DrawCircleXY(Context, float3(SlabCenter.xy + PaperLobe.Mu.xy * VisProjectedRayRadius, SlabBottomZ), PaperLobe.Sigma * VisProjectedRayRadius, float4(1, 0, 1, 1));
 
 	Newline(Context);
 	Newline(Context);
 	Print(Context, TEXT("Belcour variance = "));
-	Print(Context, RefractedLobe.Sigma, FontWhite, 5, 5);
+	Print(Context, PaperLobe.Sigma, FontWhite, 5, 5);
 
 	// Our lobe
 	FLobeStat LobeState = LobStatisticsBuffer[0];
