Minor update to strata rough refractions

diff --git a/Engine/Shaders/Private/Strata/StrataRoughRefraction.usf b/Engine/Shaders/Private/Strata/StrataRoughRefraction.usf
--- a/Engine/Shaders/Private/Strata/StrataRoughRefraction.usf
+++ b/Engine/Shaders/Private/Strata/StrataRoughRefraction.usf
@@ -266,8 +266,8 @@ FStrataLobeStatistic StrataGetRefractedLobe2(FStrataLobeStatistic WiLobe, float3
 	// Equation 10, we have also noticed variance explosion at grazing angle so we make the inner product be 1 for now.
 	const float OmegaIDotN = 1.0f; //WiLobe.Mu.z;
 	const float OmegaTDotN = 1.0f; //WoLobe.Mu.z;
-	//const float S = 0.5f * (1.0f + InterfaceEta12 * (abs(OmegaIDotN) / abs(OmegaTDotN)));						// equation 10 from paper
-	const float S = 0.5f * (abs(OmegaTDotN * InterfaceEta12 - OmegaIDotN) / abs(OmegaTDotN * InterfaceEta12));	// equation from paper source code
+	//const float S = 0.5f * (1.0f + InterfaceEta12 * (abs(OmegaIDotN) / abs(OmegaTDotN)));					// equation 10 from paper
+	const float S = 0.5f * abs((OmegaTDotN * InterfaceEta12 - OmegaIDotN) / (OmegaTDotN * InterfaceEta12));	// equation from paper source code
 
 	WoLobe.Sigma = (WiLobe.Sigma / InterfaceEta12) + StrataLobeRoughnessToVariance(InterfaceRoughness * S);
 
