Fix VSM issues in the right eye (UE-142707).

diff --git a/Engine/Shaders/Private/SceneTexturesCommon.ush b/Engine/Shaders/Private/SceneTexturesCommon.ush
--- a/Engine/Shaders/Private/SceneTexturesCommon.ush
+++ b/Engine/Shaders/Private/SceneTexturesCommon.ush
@@ -83,6 +83,16 @@ float LookupDeviceZ( float2 ScreenUV )
 #endif
 }
 
+/** Returns DeviceZ which is the z value stored in the depth buffer. */
+float LookupDeviceZ(uint2 PixelPos)
+{
+#if SCENE_TEXTURES_DISABLED
+	return SCENE_TEXTURES_DISABLED_SCENE_DEPTH_VALUE;
+#else
+	return SceneTexturesStruct.SceneDepthTexture.Load(int3(PixelPos, 0)).r;
+#endif
+}
+
 /** Returns clip space W, which is world space distance along the View Z axis. */
 float CalcSceneDepth(uint2 PixelPos)
 {
diff --git a/Engine/Shaders/Private/VirtualShadowMaps/VirtualShadowMapPageManagement.usf b/Engine/Shaders/Private/VirtualShadowMaps/VirtualShadowMapPageManagement.usf
--- a/Engine/Shaders/Private/VirtualShadowMaps/VirtualShadowMapPageManagement.usf
+++ b/Engine/Shaders/Private/VirtualShadowMaps/VirtualShadowMapPageManagement.usf
@@ -44,11 +44,6 @@ float PageDilationBorderSizeLocal;
 uint InputType;
 uint bCullBackfacingPixels;
 
-float2 CalcScreenUV(uint2 PixelPos)
-{
-	return (float2(PixelPos.xy) + View.ViewRectMin.xy + 0.5f) * View.BufferSizeAndInvSize.zw;
-}
-
 uint CalcMipLevelLocal(float Footprint)
 {
 	float MipLevelFloat = log2(Footprint) + ResolutionLodBiasLocal;
@@ -153,8 +148,9 @@ void GeneratePageFlagsFromPixels(
 	const uint2 GroupId = InGroupId.xy;
 #endif
 	// Morton order within a group so page access/atomics are more coherent and wave-swizzled gradients are possible.
-	const uint2 PixelPos = VSM_DEFAULT_CS_GROUP_XY * GroupId.xy + MortonDecode(GroupIndex);
-	if (any(PixelPos >= uint2(View.ViewSizeAndInvSize.xy)))
+	const uint2 PixelPos = uint2(View.ViewRectMin.xy) + VSM_DEFAULT_CS_GROUP_XY * GroupId.xy + MortonDecode(GroupIndex);
+
+	if (any(PixelPos >= uint2(View.ViewRectMin.xy + View.ViewSizeAndInvSize.xy)))
 	{
 		return;
 	}
@@ -171,7 +167,7 @@ void GeneratePageFlagsFromPixels(
 	}
 	#else // PERMUTATION_INPUT_TYPE == INPUT_TYPE_GBUFFER
 	{
-		DeviceZ = LookupDeviceZ(CalcScreenUV(PixelPos));
+		DeviceZ = LookupDeviceZ(PixelPos);
 	  #if STRATA_ENABLED
 		const FStrataTopLayerData TopLayerData = StrataUnpackTopLayerData(Strata.TopLayerTexture.Load(uint3(PixelPos, 0)));
 		WorldNormal = TopLayerData.WorldNormal;
