Implements experimental TSR permutation that only stores high frequency details in the history for 60hz

This adds a new experimental r.TSR.History.HighFrequencyOnly.
This remove r.TSR.Translucency.SeparateTemporalAccumulation because translucency is otherwise too ghosty anyway, avoiding an increase of shader permutation.

diff --git a/Engine/Shaders/Private/TemporalSuperResolution/TSRCommon.ush b/Engine/Shaders/Private/TemporalSuperResolution/TSRCommon.ush
--- a/Engine/Shaders/Private/TemporalSuperResolution/TSRCommon.ush
+++ b/Engine/Shaders/Private/TemporalSuperResolution/TSRCommon.ush
@@ -619,6 +619,8 @@ half2 SafeRcp(half2 x)
 
 #endif
 
+#define MIN_IMPLICIT_HISTORY_WEIGHT ceil(tsr_half(255.0) * rcp(tsr_half(MAX_SAMPLE_COUNT)) * tsr_half(0.25)) * rcp(tsr_half(255.0))
+
 CALL_SITE_DEBUGLOC
 tsr_half2 ComputeLowFrequencyContribution(tsr_half2 HistoryValidity)
 {
diff --git a/Engine/Shaders/Private/TemporalSuperResolution/TSRDecimateHistory.usf b/Engine/Shaders/Private/TemporalSuperResolution/TSRDecimateHistory.usf
--- a/Engine/Shaders/Private/TemporalSuperResolution/TSRDecimateHistory.usf
+++ b/Engine/Shaders/Private/TemporalSuperResolution/TSRDecimateHistory.usf
@@ -169,7 +169,12 @@ void MainCS(
 		
 		float2 PrevInputBufferUV = (InputInfo_ScreenPosToViewportScale * PrevOutputScreenPos + InputInfo_ScreenPosToViewportBias) * InputInfo_ExtentInverse;
 
-		#if 1
+		#if DIM_HIGH_FREQUENCY_ONLY
+		{
+			tsr_half2 Metadata = PrevHistory_Metadata.SampleLevel(GlobalBilinearClampedSampler, KernelHistoryBufferUV, 0);
+			DecimatedColor = PrevHistory_HighFrequency.SampleLevel(GlobalBilinearClampedSampler, KernelHistoryBufferUV, 0) * SafeRcp(Metadata.r);
+		}
+		#elif 1
 		{
 			DecimatedColor = PrevHistory_LowFrequency.SampleLevel(GlobalBilinearClampedSampler, KernelHistoryBufferUV, 0);
 		}
diff --git a/Engine/Shaders/Private/TemporalSuperResolution/TSRResolveHistory.usf b/Engine/Shaders/Private/TemporalSuperResolution/TSRResolveHistory.usf
--- a/Engine/Shaders/Private/TemporalSuperResolution/TSRResolveHistory.usf
+++ b/Engine/Shaders/Private/TemporalSuperResolution/TSRResolveHistory.usf
@@ -142,10 +142,10 @@ void MainCS(
 				tsr_half3 SampleTranslucencyColor = History_Translucency.SampleLevel(GlobalBilinearClampedSampler, SampleUV, 0);
 				tsr_half  SampleTranslucencyAlpha = History_TranslucencyAlpha.SampleLevel(GlobalBilinearClampedSampler, SampleUV, 0);
 				
-				tsr_half SampleHistoryValidity = SampleMetadata.y;
+				tsr_half SampleHistoryValidity = DIM_HIGH_FREQUENCY_ONLY ? SampleMetadata.x : SampleMetadata.y;
 				tsr_half InvSampleHistoryValidity = SafeRcp(SampleHistoryValidity);
 
-				#if DIM_SEPARATE_TRANSLUCENCY
+				#if 1
 				{
 					SampleLowFrequencyColor.rgb  = SampleLowFrequencyColor.rgb  * SampleTranslucencyAlpha + SampleTranslucencyColor.rgb;
 					SampleHighFrequencyColor.rgb = SampleHighFrequencyColor.rgb * SampleTranslucencyAlpha + SampleTranslucencyColor.rgb * SampleHistoryValidity;
@@ -182,6 +182,11 @@ void MainCS(
 	
 	// Compute final output
 	tsr_halfC FinalOutputColor;
+	#if DIM_HIGH_FREQUENCY_ONLY
+	{
+		FinalOutputColor = HighFrequencyColor;
+	}
+	#else
 	{
 		tsr_half OutputBlend = ComputeLowFrequencyContribution(HistoryValidity);
 		
@@ -192,6 +197,7 @@ void MainCS(
 		
 		FinalOutputColor = LowFrequencyColor * Weights[0] + HighFrequencyColor * Weights[1];
 	}
+	#endif
 
 	tsr_short2 OutputPixelPosition = InvalidateOutputPixelPos(tsr_short2(OutputViewRectMin + DispatchThreadId), OutputViewRectMax);
 	
diff --git a/Engine/Shaders/Private/TemporalSuperResolution/TSRUpdateHistory.usf b/Engine/Shaders/Private/TemporalSuperResolution/TSRUpdateHistory.usf
--- a/Engine/Shaders/Private/TemporalSuperResolution/TSRUpdateHistory.usf
+++ b/Engine/Shaders/Private/TemporalSuperResolution/TSRUpdateHistory.usf
@@ -37,13 +37,18 @@
 	#error Unknown history update quality
 #endif
 
-#if DIM_SEPARATE_TRANSLUCENCY == 1
-	#define CONFIG_ACCUMULATE_TRANSLUCENCY_SEPARATLY 1
-	#define CONFIG_VFX_SURVIVAL 0
-	#define CONFIG_REJECT_TRAMSLUCENCY_ON_RESPONSIVE_AA_ONLY 0
+#if DIM_HIGH_FREQUENCY_ONLY == 1
+	#define CONFIG_HIGH_FREQUENCY_ONLY 1
+	#define CONFIG_HIGH_FREQUENCY_YCOCG 0
+
 #else
-	#define CONFIG_ACCUMULATE_TRANSLUCENCY_SEPARATLY 0
-	#define CONFIG_VFX_SURVIVAL 0
+	#define CONFIG_HIGH_FREQUENCY_ONLY 0
+	#define CONFIG_HIGH_FREQUENCY_YCOCG 1
+
+#endif
+
+#if 1
+	#define CONFIG_ACCUMULATE_TRANSLUCENCY_SEPARATLY 1
 	#define CONFIG_REJECT_TRAMSLUCENCY_ON_RESPONSIVE_AA_ONLY 0
 #endif
 
@@ -107,6 +112,7 @@ float InvWeightClampingPixelSpeed;
 float InputToHistoryFactor;
 uint ResponsiveStencilMask;
 uint bGenerateOutputMip1;
+uint bHasSeparateTranslucency;
 
 Texture2D<tsr_halfC> PrevHistory_LowFrequency;
 Texture2D<tsr_halfC> PrevHistory_HighFrequency;
@@ -295,8 +301,8 @@ void MainCS(
 		tsr_half RawParallaxRejectionMask[DPV_PIXEL_PER_LANE];
 		uint RawParallaxFactorBits[DPV_PIXEL_PER_LANE];
 		tsr_half RawLowFrequencyRejection[DPV_PIXEL_PER_LANE];
-		tsr_half RawTranslucencyRejection[DPV_PIXEL_PER_LANE];
-		tsr_half RawTranslucencyRejectionNeighborhood[CONFIG_SAMPLES - 1][DPV_PIXEL_PER_LANE];
+		//tsr_half RawTranslucencyRejection[DPV_PIXEL_PER_LANE];
+		//tsr_half RawTranslucencyRejectionNeighborhood[CONFIG_SAMPLES - 1][DPV_PIXEL_PER_LANE];
 		
 		uint RawSceneStencilRef[DPV_PIXEL_PER_LANE];
 		uint RawSceneStencilRefNeighborhood[CONFIG_SAMPLES - 1][DPV_PIXEL_PER_LANE];
@@ -322,7 +328,7 @@ void MainCS(
 				#endif
 				RawParallaxRejectionMask[PixelId] = ParallaxRejectionMaskTexture[LocalInputPixelPos];
 				RawParallaxFactorBits[PixelId]    = ParallaxFactorTexture[LocalInputPixelPos];
-				RawTranslucencyRejection[PixelId] = TranslucencyRejectionTexture[LocalInputPixelPos];
+				//RawTranslucencyRejection[PixelId] = TranslucencyRejectionTexture[LocalInputPixelPos];
 				RawSceneStencilRef[PixelId]       = InputSceneStencilTexture.Load(int3(LocalInputPixelPos, 0)) STENCIL_COMPONENT_SWIZZLE;
 				
 				RawLowFrequencyRejection[PixelId] = HistoryRejectionTexture[dpv_access_pixel(RejectionPixelPos, PixelId)];
@@ -352,7 +358,7 @@ void MainCS(
 				UNROLL_N(DPV_PIXEL_PER_LANE)
 				for (uint PixelId = 0; PixelId < DPV_PIXEL_PER_LANE; PixelId++)
 				{
-					RawTranslucencyRejectionNeighborhood[NeighborId][PixelId] = TranslucencyRejectionTexture[dpv_access_pixel(SampleInputPixelPos, PixelId)];
+					//RawTranslucencyRejectionNeighborhood[NeighborId][PixelId] = TranslucencyRejectionTexture[dpv_access_pixel(SampleInputPixelPos, PixelId)];
 					RawSceneStencilRefNeighborhood[NeighborId][PixelId] = InputSceneStencilTexture.Load(int3(dpv_access_pixel(SampleInputPixelPos, PixelId), 0)) STENCIL_COMPONENT_SWIZZLE;
 				}
 			} // for (uint NeighborId = 0; NeighborId < (CONFIG_SAMPLES - 1); NeighborId++)
@@ -429,14 +435,14 @@ void MainCS(
 			ParallaxRejectionMask = dpv_force_interleave_registers_array(RawParallaxRejectionMask);
 			ParallaxFactorBits    = dpv_force_interleave_registers_array(RawParallaxFactorBits) & tsr_subpixel_payload(SUB_PIXEL_PARALLAX_FACTOR_BIT_MASK);
 			LowFrequencyRejection = dpv_force_interleave_registers_array(RawLowFrequencyRejection);
-			TranslucencyRejection = dpv_force_interleave_registers_array(RawTranslucencyRejection);
+			//TranslucencyRejection = dpv_force_interleave_registers_array(RawTranslucencyRejection);
 			SceneStencilRef       = dpv_force_interleave_registers_array(RawSceneStencilRef);
 			
 			// Process neighborhood
 			UNROLL_N(CONFIG_SAMPLES - 1)
 			for (uint NeighborId = 0; NeighborId < (CONFIG_SAMPLES - 1); NeighborId++)
 			{
-				TranslucencyRejection = min(TranslucencyRejection, dpv_force_interleave_registers_array(RawTranslucencyRejectionNeighborhood[NeighborId]));
+				//TranslucencyRejection = min(TranslucencyRejection, dpv_force_interleave_registers_array(RawTranslucencyRejectionNeighborhood[NeighborId]));
 			
 				SceneStencilRef |= dpv_interleave_registers_array(RawSceneStencilRefNeighborhood[NeighborId]);
 			}
@@ -625,14 +631,6 @@ void MainCS(
 
 			tsr_half2 ToneWeight = HdrWeight4(InputColor);
 
-			#if CONFIG_VFX_SURVIVAL
-			if (SampleId == 0)
-			{
-				tsr_half2x2 dPP =  RawdPPArray[SampleId];
-				ToneWeight = max(ToneWeight, ToneWeight > tsr_half(0.2) ? tsr_half(0.0).xx : ComputeSampleWeigth(tsr_half(1.0), dPP, /* MinimalContribution = */ 0.0));
-			}
-			#endif
-
 			FilteredInputColor       += dpv_scale((SampleSpatialWeight * ToneWeight), InputColor);
 			FilteredInputColorWeight += (SampleSpatialWeight * ToneWeight);
 
@@ -641,8 +639,13 @@ void MainCS(
 				FilteredTranslucencyColor = TranslucencyColor;
 			}
 
-			InputMinColor = min(InputMinColor, TransformColorForClampingBox(InputColor));
-			InputMaxColor = max(InputMaxColor, TransformColorForClampingBox(InputColor));
+			#if CONFIG_HIGH_FREQUENCY_YCOCG
+				InputMinColor = min(InputMinColor, TransformColorForClampingBox(InputColor));
+				InputMaxColor = max(InputMaxColor, TransformColorForClampingBox(InputColor));
+			#else
+				InputMinColor = min(InputMinColor, InputColor);
+				InputMaxColor = max(InputMaxColor, InputColor);
+			#endif
 
 			#if CONFIG_ACCUMULATE_TRANSLUCENCY_SEPARATLY
 			{
@@ -707,10 +710,15 @@ void MainCS(
 	#endif
 	
 	// Reproject history
-	tsr_halfCx2 PrevLowFrequencyColor;
-	tsr_half2 PrevLowFrequencyWeight;
-	
+	#if !CONFIG_HIGH_FREQUENCY_ONLY
+		tsr_halfCx2 PrevLowFrequencyColor;
+		tsr_half2 PrevLowFrequencyWeight;
+	#endif
+
 	tsr_halfCx2 PrevHighFrequencyColor;
+	#if CONFIG_HIGH_FREQUENCY_ONLY
+		tsr_half2 PrevHistoryWeight;
+	#endif
 	tsr_half2 PrevHistoryValidity;
 	
 	tsr_halfCx2 PrevHistoryMommentMin;
@@ -742,7 +750,9 @@ void MainCS(
 			FCatmullRomSamples Samples1 = GetBicubic2DCatmullRomSamples(
 				dpv_hi(PrevHistoryBufferUV), PrevHistoryInfo_Extent, PrevHistoryInfo_ExtentInverse);
 			
-			tsr_halfC RawLowFrequency[BICUBIC_CATMULL_ROM_SAMPLES][DPV_PIXEL_PER_LANE];
+			#if !CONFIG_HIGH_FREQUENCY_ONLY
+				tsr_halfC RawLowFrequency[BICUBIC_CATMULL_ROM_SAMPLES][DPV_PIXEL_PER_LANE];
+			#endif
 			tsr_halfC RawHighFrequency[BICUBIC_CATMULL_ROM_SAMPLES][DPV_PIXEL_PER_LANE];
 			tsr_half2 RawMetadata[BICUBIC_CATMULL_ROM_SAMPLES][DPV_PIXEL_PER_LANE];
 			tsr_half3 RawTranslucencyColor[BICUBIC_CATMULL_ROM_SAMPLES][DPV_PIXEL_PER_LANE];
@@ -768,7 +778,9 @@ void MainCS(
 					{
 						float2 SampleUV = PixelId == 0 ? SampleUV0 : SampleUV1;
 						
-						RawLowFrequency[i][PixelId]      = PrevHistory_LowFrequency.SampleLevel(GlobalBilinearClampedSampler, SampleUV, 0);
+						#if !CONFIG_HIGH_FREQUENCY_ONLY
+							RawLowFrequency[i][PixelId]      = PrevHistory_LowFrequency.SampleLevel(GlobalBilinearClampedSampler, SampleUV, 0);
+						#endif
 						RawHighFrequency[i][PixelId]     = PrevHistory_HighFrequency.SampleLevel(GlobalBilinearClampedSampler, SampleUV, 0);
 						RawMetadata[i][PixelId]	         = PrevHistory_Metadata.SampleLevel(GlobalBilinearClampedSampler, SampleUV, 0);
 						
@@ -787,7 +799,9 @@ void MainCS(
 				UNROLL_N(BICUBIC_CATMULL_ROM_SAMPLES)
 				for (uint i = 0; i < BICUBIC_CATMULL_ROM_SAMPLES; i++)
 				{
-					tsr_halfCx2 SampleLowFrequency = dpv_force_interleave_registers_array(RawLowFrequency[i]);
+					#if !CONFIG_HIGH_FREQUENCY_ONLY
+						tsr_halfCx2 SampleLowFrequency = dpv_force_interleave_registers_array(RawLowFrequency[i]);
+					#endif
 					tsr_halfCx2 SampleHighFrequency = dpv_force_interleave_registers_array(RawHighFrequency[i]);
 					tsr_half2x2 SampleMetadata = dpv_force_interleave_registers_array(RawMetadata[i]);
 					tsr_half4x2 SampleTranslucency = dpv_force_interleave_registers(
@@ -796,12 +810,18 @@ void MainCS(
 					
 					tsr_half2 KernelWeight = RawKernelWeight[i];
 
-					tsr_half2 InvWeight = SafeRcp(SampleMetadata[1]);
+					#if CONFIG_HIGH_FREQUENCY_ONLY
+						tsr_half2 InvWeight = SafeRcp(SampleMetadata[0]);
+					#else
+						tsr_half2 InvWeight = SafeRcp(SampleMetadata[1]);
+					#endif
 
 					HighFrequencyMin = min(HighFrequencyMin, dpv_scale(SampleHighFrequency, InvWeight));
 					HighFrequencyMax = max(HighFrequencyMax, dpv_scale(SampleHighFrequency, InvWeight));
 
-					AccumulateLowFrequency  += dpv_scale(SampleLowFrequency, KernelWeight);
+					#if !CONFIG_HIGH_FREQUENCY_ONLY
+						AccumulateLowFrequency  += dpv_scale(SampleLowFrequency, KernelWeight);
+					#endif
 					AccumulateHighFrequency += dpv_scale(SampleHighFrequency, KernelWeight);
 					AccumulateMetadata      += dpv_scale(SampleMetadata, KernelWeight);
 					
@@ -812,7 +832,9 @@ void MainCS(
 				
 				tsr_half2 FinalMultiplier = dpv_force_interleave_registers(tsr_half(Samples0.FinalMultiplier), tsr_half(Samples1.FinalMultiplier));
 
-				AccumulateLowFrequency  = dpv_scale(AccumulateLowFrequency, FinalMultiplier);
+				#if !CONFIG_HIGH_FREQUENCY_ONLY
+					AccumulateLowFrequency  = dpv_scale(AccumulateLowFrequency, FinalMultiplier);
+				#endif
 				AccumulateHighFrequency = dpv_scale(AccumulateHighFrequency, FinalMultiplier);
 				AccumulateMetadata      = dpv_scale(AccumulateMetadata, FinalMultiplier);
 				AccumulateTranslucency  = dpv_scale(AccumulateTranslucency, FinalMultiplier);
@@ -845,13 +867,18 @@ void MainCS(
 		
 		// Unpack history.
 		{
-			PrevLowFrequencyColor = AccumulateLowFrequency;
-			PrevLowFrequencyWeight = AccumulateMetadata[0];
-			
+			#if !CONFIG_HIGH_FREQUENCY_ONLY
+				PrevLowFrequencyColor = AccumulateLowFrequency;
+				PrevLowFrequencyWeight = AccumulateMetadata[0];
+			#endif
+
 			PrevHistoryMommentMin = HighFrequencyMin;
 			PrevHistoryMommentMax = HighFrequencyMax;
 
 			PrevHighFrequencyColor = AccumulateHighFrequency;
+			#if CONFIG_HIGH_FREQUENCY_ONLY
+				PrevHistoryWeight = AccumulateMetadata[0];
+			#endif
 			PrevHistoryValidity = AccumulateMetadata[1];
 
 			PrevHistoryTranslucency = AccumulateTranslucency;
@@ -868,9 +895,11 @@ void MainCS(
 			PrevHighFrequencyColor[0] *= tsr_half(HistoryPreExposureCorrection);
 			PrevHighFrequencyColor[1] *= tsr_half(HistoryPreExposureCorrection);
 			PrevHighFrequencyColor[2] *= tsr_half(HistoryPreExposureCorrection);
-			PrevLowFrequencyColor[0] *= tsr_half(HistoryPreExposureCorrection);
-			PrevLowFrequencyColor[1] *= tsr_half(HistoryPreExposureCorrection);
-			PrevLowFrequencyColor[2] *= tsr_half(HistoryPreExposureCorrection);
+			#if !CONFIG_HIGH_FREQUENCY_ONLY
+				PrevLowFrequencyColor[0] *= tsr_half(HistoryPreExposureCorrection);
+				PrevLowFrequencyColor[1] *= tsr_half(HistoryPreExposureCorrection);
+				PrevLowFrequencyColor[2] *= tsr_half(HistoryPreExposureCorrection);
+			#endif
 
 			// Fixes negative validity that can happen due to negative lobes of Catrom.
 			PrevHistoryValidity = max(PrevHistoryValidity, tsr_half(0.0).xx);
@@ -930,23 +959,35 @@ void MainCS(
 	}
 	#endif
 
+	tsr_half2 HighFrequencyRejection = tsr_half(1.0).xx;
 	// Rejects high frequency details if their low frequencies drift from the low frequency history
 	#if CONFIG_LOW_FREQUENCY_DRIFT_REJECTION
 	{
-		tsr_halfCx2 NormalizedPrevHighFrequencyColor = dpv_scale(PrevHighFrequencyColor, SafeRcp(PrevHistoryValidity));
+		#if CONFIG_HIGH_FREQUENCY_ONLY
+			tsr_halfCx2 NormalizedPrevHighFrequencyColor = dpv_scale(PrevHighFrequencyColor, SafeRcp(PrevHistoryWeight));
 
-		tsr_halfCx2 PrevHighFrequencyYCoCg = TransformColorForClampingBox(NormalizedPrevHighFrequencyColor);
-		tsr_halfCx2 PrevYCoCg = TransformColorForClampingBox(PrevLowFrequencyColor);
-		tsr_halfCx2 ClampedPrevYCoCg = TransformColorForClampingBox(clamp(PrevLowFrequencyColor, PrevHistoryMommentMin, PrevHistoryMommentMax));
+			tsr_halfCx2 PrevHighFrequencyYCoCg = NormalizedPrevHighFrequencyColor;
+			tsr_halfCx2 PrevYCoCg = FilteredInputColor;
+			tsr_halfCx2 ClampedPrevYCoCg = clamp(FilteredInputColor, PrevHistoryMommentMin, PrevHistoryMommentMax);
+		#else
+			tsr_halfCx2 NormalizedPrevHighFrequencyColor = dpv_scale(PrevHighFrequencyColor, SafeRcp(PrevHistoryValidity));
+
+			tsr_halfCx2 PrevHighFrequencyYCoCg = TransformColorForClampingBox(NormalizedPrevHighFrequencyColor);
+			tsr_halfCx2 PrevYCoCg = TransformColorForClampingBox(PrevLowFrequencyColor);
+			tsr_halfCx2 ClampedPrevYCoCg = TransformColorForClampingBox(clamp(PrevLowFrequencyColor, PrevHistoryMommentMin, PrevHistoryMommentMax));
+		#endif
 
-		tsr_half2 HighFrequencyRejection = MeasureRejectionFactor(
+		HighFrequencyRejection = MeasureRejectionFactor(
 			PrevYCoCg, ClampedPrevYCoCg,
 			PrevHighFrequencyYCoCg, InputMinColor, InputMaxColor,
 			MeasureBackbufferLDRQuantizationErrorFromHDRLuma(max(PrevYCoCg[0], ClampedPrevYCoCg[0])));
 
 		//PrevHighFrequencyColor = dpv_scale(PrevHighFrequencyColor, HighFrequencyRejection);
-		PrevHighFrequencyColor = dpv_scale(clamp(NormalizedPrevHighFrequencyColor, PrevHistoryMommentMin, PrevHistoryMommentMax), HighFrequencyRejection * PrevHistoryValidity);
-		PrevHistoryValidity *= HighFrequencyRejection;
+
+		#if !CONFIG_HIGH_FREQUENCY_ONLY
+			PrevHighFrequencyColor = dpv_scale(clamp(NormalizedPrevHighFrequencyColor, PrevHistoryMommentMin, PrevHistoryMommentMax), HighFrequencyRejection * PrevHistoryValidity);
+			PrevHistoryValidity *= HighFrequencyRejection;
+		#endif
 	}
 	#endif
 
@@ -955,6 +996,7 @@ void MainCS(
 	const tsr_half PredictionOnlyValidity = Histeresis * tsr_half(2.0);
 	
 	// Clamp low frequency
+	#if !CONFIG_HIGH_FREQUENCY_ONLY
 	tsr_half2 LumaMin;
 	tsr_half2 LumaMax;
 	tsr_halfCx2 ClampedLowFrequencyColor;
@@ -981,22 +1023,37 @@ void MainCS(
 		}
 		#endif
 	}
+	#endif // !CONFIG_HIGH_FREQUENCY_ONLY
 
-	// Also reject sub pixel detail low frequency if detect mid frequency changes.
-	#if 0
-	{
-		LowFrequencyRejection = min(LowFrequencyRejection, MidFrequencyRejection);
-	}
-	#endif
-	
 	tsr_halfCx2 FinalHighFrequencyColor;
 	tsr_half2 FinalHistoryValidity;
 	{
-		tsr_half2 InvPrevHistoryValidity = SafeRcp(PrevHistoryValidity);
+		#if CONFIG_HIGH_FREQUENCY_ONLY
+			#if CONFIG_HIGH_FREQUENCY_YCOCG
+				tsr_halfCx2 ClampedPrevHighFrequencyColor = YCoCgToRGB(clamp(TransformColorForClampingBox(PrevHighFrequencyColor), dpv_scale(InputMinColor, PrevHistoryWeight), dpv_scale(InputMaxColor, PrevHistoryWeight)));
+			#else
+				tsr_halfCx2 ClampedPrevHighFrequencyColor = clamp(PrevHighFrequencyColor, dpv_scale(InputMinColor, PrevHistoryWeight), dpv_scale(InputMaxColor, PrevHistoryWeight));
+			#endif
 
-		// Compute how much the history needs to be rejected, based on the completeness.
-		tsr_half2 PrevHistoryRejectionWeight = LowFrequencyRejection * HoleFilledVelocityMask[1];
-			
+			tsr_half2 BlendClamp = TotalRejection * HoleFilledVelocityMask[1] * HighFrequencyRejection;
+
+			PrevHighFrequencyColor[0] = lerp(ClampedPrevHighFrequencyColor[0], PrevHighFrequencyColor[0], BlendClamp);
+			PrevHighFrequencyColor[1] = lerp(ClampedPrevHighFrequencyColor[1], PrevHighFrequencyColor[1], BlendClamp);
+			PrevHighFrequencyColor[2] = lerp(ClampedPrevHighFrequencyColor[2], PrevHighFrequencyColor[2], BlendClamp);
+
+			tsr_half2 InvPrevHistoryWeight = SafeRcp(PrevHistoryWeight);
+			tsr_half2 InvPrevHistoryValidity = SafeRcp(PrevHistoryValidity);
+
+			// Compute how much the history needs to be rejected, based on the completeness.
+			tsr_half2 PrevHistoryRejectionWeight = LowFrequencyRejection; // * HoleFilledVelocityMask[1];
+		#else // !CONFIG_HIGH_FREQUENCY_ONLY
+			tsr_half2 InvPrevHistoryValidity = SafeRcp(PrevHistoryValidity);
+			tsr_half2 InvPrevHistoryWeight = InvPrevHistoryValidity;
+
+			// Compute how much the history needs to be rejected, based on the completeness.
+			tsr_half2 PrevHistoryRejectionWeight = LowFrequencyRejection * HoleFilledVelocityMask[1];
+		#endif
+	
 		//#if !CONFIG_CLAMP
 		//	PrevHistoryRejectionWeight = 1.0;
 		//#endif
@@ -1027,7 +1084,7 @@ void MainCS(
 			// Clamp up the max validity to favorise stability under motion on high contrast edges.
 			#if 1
 			{
-				tsr_half2 PrevHistoryLuma = Luma4(PrevHighFrequencyColor) * InvPrevHistoryValidity;
+				tsr_half2 PrevHistoryLuma = Luma4(PrevHighFrequencyColor) * InvPrevHistoryWeight;
 				tsr_half2 FilteredLuma = Luma4(FilteredInputColor);
 
 				tsr_half2 MinValidityForStability = abs(FilteredLuma - PrevHistoryLuma) / max(FilteredLuma, PrevHistoryLuma);
@@ -1070,12 +1127,18 @@ void MainCS(
 			}
 		#endif
 		
-		tsr_half2 PrevMomentWeight = max(OutputValidity - DesiredCurrentContribution, tsr_half(0.0));
-		tsr_half2 CurrentMomentWeight = min(DesiredCurrentContribution, OutputValidity);
-		
+		#if CONFIG_HIGH_FREQUENCY_ONLY
+			tsr_half2 OutputValidity2 = max(OutputValidity, MIN_IMPLICIT_HISTORY_WEIGHT);
+		#else
+			tsr_half2 OutputValidity2 = OutputValidity;
+		#endif
+
+		tsr_half2 PrevMomentWeight = max(OutputValidity2 - DesiredCurrentContribution, tsr_half(0.0));
+		tsr_half2 CurrentMomentWeight = min(DesiredCurrentContribution, OutputValidity2);
+
 		#if 1
 		{
-			tsr_half2 PrevHistoryToneWeight = HdrWeightY(Luma4(PrevHighFrequencyColor) * InvPrevHistoryValidity);
+			tsr_half2 PrevHistoryToneWeight = HdrWeightY(Luma4(PrevHighFrequencyColor) * InvPrevHistoryWeight);
 			tsr_half2 FilteredInputToneWeight = HdrWeight4(FilteredInputColor);
 			
 			tsr_half2 BlendPrevHistory = PrevMomentWeight * PrevHistoryToneWeight;
@@ -1084,7 +1147,7 @@ void MainCS(
 			tsr_half2 CommonWeight = SafeRcp(BlendPrevHistory + BlendFilteredInput);
 
 			FinalHighFrequencyColor = (
-				dpv_scale(PrevHighFrequencyColor, CommonWeight * BlendPrevHistory * InvPrevHistoryValidity) +
+				dpv_scale(PrevHighFrequencyColor, CommonWeight * BlendPrevHistory * InvPrevHistoryWeight) +
 				dpv_scale(FilteredInputColor, CommonWeight * BlendFilteredInput));
 		}
 		#else
@@ -1109,6 +1172,7 @@ void MainCS(
 	}
 
 	// Compute the low frequency history
+	#if !CONFIG_HIGH_FREQUENCY_ONLY
 	tsr_halfCx2 FinalLowFrequencyColor;
 	tsr_half2 FinalLowFrequencyWeight;
 	{
@@ -1218,6 +1282,7 @@ void MainCS(
 			FinalLowFrequencyColor = dpv_scale(ClampedLowFrequencyColor, Weights[0]) + dpv_scale(FilteredInputColor, Weights[1]);
 		}
 	}
+	#endif // !CONFIG_HIGH_FREQUENCY_ONLY
 	
 	// Compute the low frequency history
 	tsr_half4x2 FinalTranslucencyColor;
@@ -1300,6 +1365,12 @@ void MainCS(
 			FinalTranslucencyColor = dpv_scale(ClampedPrevHistoryTranslucency, Weights[0]) + dpv_scale(FilteredTranslucencyColor, Weights[1]);
 		}
 		#endif
+		
+		FLATTEN
+		if (!bHasSeparateTranslucency)
+		{
+			FinalTranslucencyColor = dpv_interleave_mono_registers(tsr_half4(0.0, 0.0, 0.0, 1.0));
+		}
 	}
 	#endif
 
@@ -1354,6 +1425,11 @@ void MainCS(
 
 	// Compute final output
 	tsr_halfCx2 FinalOutputColor;
+	#if CONFIG_HIGH_FREQUENCY_ONLY
+	{
+		FinalOutputColor = FinalHighFrequencyColor;
+	}
+	#else // !CONFIG_HIGH_FREQUENCY_ONLY
 	{
 		tsr_half2 OutputBlend = ComputeLowFrequencyContribution(FinalHistoryValidity);
 
@@ -1392,6 +1468,9 @@ void MainCS(
 			FinalOutputColor = dpv_scale(FinalLowFrequencyColor, Weights[0]) + dpv_scale(FinalHighFrequencyColor, Weights[1]);
 		#endif
 
+	}
+	#endif // !CONFIG_HIGH_FREQUENCY_ONLY
+	{
 		#if CONFIG_ACCUMULATE_TRANSLUCENCY_SEPARATLY
 			FinalOutputColor[0] = FinalOutputColor[0] * FinalTranslucencyColor[3] + FinalTranslucencyColor[0];
 			FinalOutputColor[1] = FinalOutputColor[1] * FinalTranslucencyColor[3] + FinalTranslucencyColor[1];
@@ -1434,7 +1513,12 @@ void MainCS(
 		LocalHistoryPixelPos = InvalidateOutputPixelPos(LocalHistoryPixelPos, HistoryInfo_ViewportMax);
 		
 		// Premultiply FinalHighFrequencyColor to have validity weighted average in next frame history reprojection.
-		tsr_halfCx2 PremultipliedHighFrequencyColor = dpv_scale(FinalHighFrequencyColor, FinalHistoryValidity);
+		#if CONFIG_HIGH_FREQUENCY_ONLY
+			tsr_half2 FinalHistoryWeight = max(FinalHistoryValidity, MIN_IMPLICIT_HISTORY_WEIGHT);
+		#else
+			tsr_half2 FinalHistoryWeight = FinalHistoryValidity;
+		#endif
+		tsr_halfCx2 PremultipliedHighFrequencyColor = dpv_scale(FinalHighFrequencyColor, FinalHistoryWeight);
 
 		// Stocastically round up or down using the hardware RWTexture2D truncation unit to take into precision
 		// loss due to pixel format encoding.
@@ -1444,15 +1528,11 @@ void MainCS(
 			tsr_half E = tsr_half(Hammersley16(0, 1, Random).x);
 
 			PremultipliedHighFrequencyColor = QuantizeForFloatRenderTarget(PremultipliedHighFrequencyColor, E, HistoryQuantizationError);
-			FinalLowFrequencyColor = QuantizeForFloatRenderTarget(FinalLowFrequencyColor, E, HistoryQuantizationError);
 		}
 		#endif
 
 		// Protext from NaN and +Inf when writing out the history.
 		{
-			FinalLowFrequencyColor = -dpv_min(-FinalLowFrequencyColor, ColorNull);
-			FinalLowFrequencyColor = dpv_min(FinalLowFrequencyColor, ColorMax10BitsFloat);
-
 			PremultipliedHighFrequencyColor = -dpv_min(-PremultipliedHighFrequencyColor, ColorNull);
 			PremultipliedHighFrequencyColor = dpv_min(PremultipliedHighFrequencyColor, ColorMax10BitsFloat);	
 
@@ -1464,9 +1544,14 @@ void MainCS(
 
 		// Output final history lo pixel.
 		{
-			HistoryOutput_LowFrequency[dpv_lo(LocalHistoryPixelPos)]  = dpv_lo(FinalLowFrequencyColor);
-			HistoryOutput_HighFrequency[dpv_lo(LocalHistoryPixelPos)] = dpv_lo(PremultipliedHighFrequencyColor);
-			HistoryOutput_Metadata[dpv_lo(LocalHistoryPixelPos)]      = tsr_half2(dpv_lo(FinalLowFrequencyWeight), dpv_lo(FinalHistoryValidity));
+			#if CONFIG_HIGH_FREQUENCY_ONLY
+				HistoryOutput_HighFrequency[dpv_lo(LocalHistoryPixelPos)] = dpv_lo(PremultipliedHighFrequencyColor);
+				HistoryOutput_Metadata[dpv_lo(LocalHistoryPixelPos)]      = tsr_half2(dpv_lo(FinalHistoryWeight), dpv_lo(FinalHistoryValidity));
+			#else
+				HistoryOutput_LowFrequency[dpv_lo(LocalHistoryPixelPos)]  = dpv_lo(FinalLowFrequencyColor);
+				HistoryOutput_HighFrequency[dpv_lo(LocalHistoryPixelPos)] = dpv_lo(PremultipliedHighFrequencyColor);
+				HistoryOutput_Metadata[dpv_lo(LocalHistoryPixelPos)]      = tsr_half2(dpv_lo(FinalLowFrequencyWeight), dpv_lo(FinalHistoryValidity));
+			#endif
 			#if CONFIG_ACCUMULATE_TRANSLUCENCY_SEPARATLY
 				HistoryOutput_Translucency[dpv_lo(LocalHistoryPixelPos)] = dpv_lo(FinalTranslucencyColor).rgb;
 				HistoryOutput_TranslucencyAlpha[dpv_lo(LocalHistoryPixelPos)] = dpv_lo(FinalTranslucencyColor).a;
@@ -1477,9 +1562,14 @@ void MainCS(
 		// Output final history hi pixel.
 		#if CONFIG_ENABLE_DUAL_PIXEL_VECTORIZATION
 		{
-			HistoryOutput_LowFrequency[dpv_hi(LocalHistoryPixelPos)]  = dpv_hi(FinalLowFrequencyColor);
-			HistoryOutput_HighFrequency[dpv_hi(LocalHistoryPixelPos)] = dpv_hi(PremultipliedHighFrequencyColor);
-			HistoryOutput_Metadata[dpv_hi(LocalHistoryPixelPos)]      = tsr_half2(dpv_hi(FinalLowFrequencyWeight), dpv_hi(FinalHistoryValidity));
+			#if CONFIG_HIGH_FREQUENCY_ONLY
+				HistoryOutput_HighFrequency[dpv_hi(LocalHistoryPixelPos)] = dpv_hi(PremultipliedHighFrequencyColor);
+				HistoryOutput_Metadata[dpv_hi(LocalHistoryPixelPos)]      = tsr_half2(dpv_hi(FinalHistoryWeight), dpv_hi(FinalHistoryValidity));
+			#else
+				HistoryOutput_LowFrequency[dpv_hi(LocalHistoryPixelPos)]  = dpv_hi(FinalLowFrequencyColor);
+				HistoryOutput_HighFrequency[dpv_hi(LocalHistoryPixelPos)] = dpv_hi(PremultipliedHighFrequencyColor);
+				HistoryOutput_Metadata[dpv_hi(LocalHistoryPixelPos)]      = tsr_half2(dpv_hi(FinalLowFrequencyWeight), dpv_hi(FinalHistoryValidity));
+			#endif
 			#if CONFIG_ACCUMULATE_TRANSLUCENCY_SEPARATLY
 				HistoryOutput_Translucency[dpv_hi(LocalHistoryPixelPos)] = dpv_hi(FinalTranslucencyColor).rgb;
 				HistoryOutput_TranslucencyAlpha[dpv_hi(LocalHistoryPixelPos)] = dpv_hi(FinalTranslucencyColor).a;
diff --git a/Engine/Source/Runtime/Renderer/Private/PostProcess/TemporalSuperResolution.cpp b/Engine/Source/Runtime/Renderer/Private/PostProcess/TemporalSuperResolution.cpp
--- a/Engine/Source/Runtime/Renderer/Private/PostProcess/TemporalSuperResolution.cpp
+++ b/Engine/Source/Runtime/Renderer/Private/PostProcess/TemporalSuperResolution.cpp
@@ -31,6 +31,11 @@ TAutoConsoleVariable<int32> CVarTSRHistoryUpdateQuality(
 	TEXT("Select the quality of the history update."),
 	ECVF_Scalability | ECVF_RenderThreadSafe);
 
+TAutoConsoleVariable<int32> CVarTSRHistoryHighFrequencyOnly(
+	TEXT("r.TSR.History.HighFrequencyOnly"), 0,
+	TEXT("Experimental store only high frequency in the history for performance saving (default = 0)."),
+	ECVF_Scalability | ECVF_RenderThreadSafe);
+
 TAutoConsoleVariable<int32> CVarTSRHalfResShadingRejection(
 	TEXT("r.TSR.ShadingRejection.HalfRes"), 0,
 	TEXT("Whether the shading rejection should be done at half res. Saves performance but may introduce back some flickering (default = 0)."),
@@ -59,11 +64,6 @@ TAutoConsoleVariable<float> CVarTSRTranslucencyHighlightLuminance(
 //	TEXT("Enable heuristic to reject Separate translucency based on previous frame translucency."),
 //	ECVF_RenderThreadSafe);
 
-TAutoConsoleVariable<int32> CVarTSRTranslucencySeparateTemporalAccumulation(
-	TEXT("r.TSR.Translucency.SeparateTemporalAccumulation"), 1,
-	TEXT("Accumulates separate translucency separatly (enabled by default)."),
-	ECVF_RenderThreadSafe);
-
 TAutoConsoleVariable<int32> CVarTSREnableResponiveAA(
 	TEXT("r.TSR.Translucency.EnableResponiveAA"), 1,
 	TEXT("Whether the responsive AA should keep history fully clamped."),
@@ -144,15 +144,15 @@ END_SHADER_PARAMETER_STRUCT()
 FTSRHistoryUAVs CreateUAVs(FRDGBuilder& GraphBuilder, const FTSRHistoryTextures& Textures)
 {
 	FTSRHistoryUAVs UAVs;
-	UAVs.LowFrequency = GraphBuilder.CreateUAV(Textures.LowFrequency);
+	if (Textures.LowFrequency)
+	{
+		UAVs.LowFrequency = GraphBuilder.CreateUAV(Textures.LowFrequency);
+	}
 	UAVs.HighFrequency = GraphBuilder.CreateUAV(Textures.HighFrequency);
 	UAVs.Metadata = GraphBuilder.CreateUAV(Textures.Metadata);
 	UAVs.SubpixelDetails = GraphBuilder.CreateUAV(Textures.SubpixelDetails);
-	if (Textures.Translucency)
-	{
-		UAVs.Translucency = GraphBuilder.CreateUAV(Textures.Translucency);
-		UAVs.TranslucencyAlpha = GraphBuilder.CreateUAV(Textures.TranslucencyAlpha);
-	}
+	UAVs.Translucency = GraphBuilder.CreateUAV(Textures.Translucency);
+	UAVs.TranslucencyAlpha = GraphBuilder.CreateUAV(Textures.TranslucencyAlpha); // TODO: move into metadata
 	return UAVs;
 }
 
@@ -178,6 +178,8 @@ public:
 	}
 }; // class FTemporalSuperResolutionShader
 
+class FTSRHighFrequencyOnlyDim : SHADER_PERMUTATION_BOOL("DIM_HIGH_FREQUENCY_ONLY");
+
 class FTSRClearPrevTexturesCS : public FTSRShader
 {
 	DECLARE_GLOBAL_SHADER(FTSRClearPrevTexturesCS);
@@ -232,7 +234,7 @@ class FTSRDecimateHistoryCS : public FTSRShader
 
 	class FOutputHalfRes : SHADER_PERMUTATION_BOOL("DIM_OUTPUT_HALF_RES");
 
-	using FPermutationDomain = TShaderPermutationDomain<FOutputHalfRes>;
+	using FPermutationDomain = TShaderPermutationDomain<FTSRHighFrequencyOnlyDim, FOutputHalfRes>;
 
 	BEGIN_SHADER_PARAMETER_STRUCT(FParameters, )
 		SHADER_PARAMETER_STRUCT_INCLUDE(FTSRCommonParameters, CommonParameters)
@@ -438,10 +440,9 @@ class FTSRUpdateHistoryCS : public FTSRShader
 		MAX
 	};
 
-	class FSeparateTranslucencyDim : SHADER_PERMUTATION_BOOL("DIM_SEPARATE_TRANSLUCENCY");
 	class FQualityDim : SHADER_PERMUTATION_ENUM_CLASS("DIM_UPDATE_QUALITY", EQuality);
 
-	using FPermutationDomain = TShaderPermutationDomain<FSeparateTranslucencyDim, FQualityDim>;
+	using FPermutationDomain = TShaderPermutationDomain<FTSRHighFrequencyOnlyDim, FQualityDim>;
 
 	BEGIN_SHADER_PARAMETER_STRUCT(FParameters, )
 		SHADER_PARAMETER_STRUCT_INCLUDE(FTSRCommonParameters, CommonParameters)
@@ -467,6 +468,7 @@ class FTSRUpdateHistoryCS : public FTSRShader
 		SHADER_PARAMETER(float, InputToHistoryFactor)
 		SHADER_PARAMETER(int32, ResponsiveStencilMask)
 		SHADER_PARAMETER(int32, bGenerateOutputMip1)
+		SHADER_PARAMETER(int32, bHasSeparateTranslucency)
 
 		SHADER_PARAMETER_STRUCT_INCLUDE(FTSRPrevHistoryParameters, PrevHistoryParameters)
 		SHADER_PARAMETER_STRUCT(FTSRHistoryTextures, PrevHistory)
@@ -489,7 +491,7 @@ class FTSRResolveHistoryCS : public FTSRShader
 	DECLARE_GLOBAL_SHADER(FTSRResolveHistoryCS);
 	SHADER_USE_PARAMETER_STRUCT(FTSRResolveHistoryCS, FTSRShader);
 
-	using FPermutationDomain = TShaderPermutationDomain<FTSRUpdateHistoryCS::FSeparateTranslucencyDim>;
+	using FPermutationDomain = TShaderPermutationDomain<FTSRHighFrequencyOnlyDim>;
 
 	BEGIN_SHADER_PARAMETER_STRUCT(FParameters, )
 		SHADER_PARAMETER_STRUCT_INCLUDE(FTSRCommonParameters, CommonParameters)
@@ -556,7 +558,7 @@ FVector ComputePixelFormatQuantizationError(EPixelFormat PixelFormat);
 
 bool ComposeSeparateTranslucencyInTSR(const FViewInfo& View)
 {
-	return CVarTSRTranslucencySeparateTemporalAccumulation.GetValueOnRenderThread() != 0;
+	return true;
 }
 
 static FRDGTextureUAVRef CreateDummyUAV(FRDGBuilder& GraphBuilder, EPixelFormat PixelFormat)
@@ -598,7 +600,7 @@ ITemporalUpscaler::FOutputs AddTemporalSuperResolutionPasses(
 
 	const bool bRejectSeparateTranslucency = false; // bIsSeperateTranslucyTexturesValid && CVarTSRTranslucencyPreviousFrameRejection.GetValueOnRenderThread() != 0;
 
-	bool bAccumulateSeparateTranslucency = bIsSeperateTranslucyTexturesValid && CVarTSRTranslucencySeparateTemporalAccumulation.GetValueOnRenderThread() != 0;
+	const bool bHistoryHighFrequencyOnly = CVarTSRHistoryHighFrequencyOnly.GetValueOnRenderThread() != 0;
 
 	EPixelFormat ColorFormat = bSupportsAlpha ? PF_FloatRGBA : PF_FloatR11G11B10;
 
@@ -692,11 +694,11 @@ ITemporalUpscaler::FOutputs AddTemporalSuperResolutionPasses(
 	FRDGTextureRef BlackAlphaOneDummy = GSystemTextures.GetBlackAlphaOneDummy(GraphBuilder);
 	FRDGTextureRef WhiteDummy = GSystemTextures.GetWhiteDummy(GraphBuilder);
 
-	FIntRect SeparateTranslucencyRect;
-	FRDGTextureRef SeparateTranslucencyTexture = nullptr;
-	if (bRejectSeparateTranslucency || bAccumulateSeparateTranslucency)
+	FIntRect SeparateTranslucencyRect = FIntRect(0, 0, 1, 1);
+	FRDGTextureRef SeparateTranslucencyTexture = BlackAlphaOneDummy;
+	bool bHasSeparateTranslucency = PassInputs.PostDOFTranslucencyResources.IsValid();
+	if (bHasSeparateTranslucency)
 	{
-		check(PassInputs.PostDOFTranslucencyResources.IsValid());
 		SeparateTranslucencyTexture = PassInputs.PostDOFTranslucencyResources.ColorTexture.Resolve;
 		SeparateTranslucencyRect = PassInputs.PostDOFTranslucencyResources.ViewRect;
 	}
@@ -884,7 +886,8 @@ ITemporalUpscaler::FOutputs AddTemporalSuperResolutionPasses(
 	if (InputHistory.IsValid())
 	{
 		// Register filterable history
-		PrevHistory.LowFrequency = GraphBuilder.RegisterExternalTexture(InputHistory.LowFrequency);
+
+		PrevHistory.LowFrequency = InputHistory.LowFrequency.IsValid() ? GraphBuilder.RegisterExternalTexture(InputHistory.LowFrequency) : DummyHistory.LowFrequency;
 		PrevHistory.HighFrequency = GraphBuilder.RegisterExternalTexture(InputHistory.HighFrequency);
 		PrevHistory.Metadata = GraphBuilder.RegisterExternalTexture(InputHistory.Metadata);
 		PrevHistory.Translucency = InputHistory.Translucency.IsValid() ? GraphBuilder.RegisterExternalTexture(InputHistory.Translucency) : DummyHistory.Translucency;
@@ -903,7 +906,7 @@ ITemporalUpscaler::FOutputs AddTemporalSuperResolutionPasses(
 	{
 
 		// Setup prev history parameters.
-		FScreenPassTextureViewport PrevHistoryViewport(PrevHistory.LowFrequency->Desc.Extent, InputHistory.OutputViewportRect);
+		FScreenPassTextureViewport PrevHistoryViewport(PrevHistory.HighFrequency->Desc.Extent, InputHistory.OutputViewportRect);
 		if (bCameraCut)
 		{
 			PrevHistoryViewport.Extent = FIntPoint(1, 1);
@@ -924,20 +927,20 @@ ITemporalUpscaler::FOutputs AddTemporalSuperResolutionPasses(
 			FClearValueBinding::None,
 			TexCreate_ShaderResource | TexCreate_UAV);
 
-		History.LowFrequency = GraphBuilder.CreateTexture(Desc, TEXT("TSR.History.LowFrequencies"));
+		if (!bHistoryHighFrequencyOnly)
+		{
+			History.LowFrequency = GraphBuilder.CreateTexture(Desc, TEXT("TSR.History.LowFrequencies"));
+		}
 		History.HighFrequency = GraphBuilder.CreateTexture(Desc, TEXT("TSR.History.HighFrequencies"));
 
 		Desc.Format = PF_R8G8;
 		History.Metadata = GraphBuilder.CreateTexture(Desc, TEXT("TSR.History.Metadata"));
 
-		if (bAccumulateSeparateTranslucency)
-		{
-			Desc.Format = (CVarTSRR11G11B10History.GetValueOnRenderThread() != 0 && !bSupportsAlpha) ? PF_FloatR11G11B10 : PF_FloatRGBA;
-			History.Translucency = GraphBuilder.CreateTexture(Desc, TEXT("TSR.History.Translucency"));
+		Desc.Format = (CVarTSRR11G11B10History.GetValueOnRenderThread() != 0 && !bSupportsAlpha) ? PF_FloatR11G11B10 : PF_FloatRGBA;
+		History.Translucency = GraphBuilder.CreateTexture(Desc, TEXT("TSR.History.Translucency"));
 
-			Desc.Format = PF_R8;
-			History.TranslucencyAlpha = GraphBuilder.CreateTexture(Desc, TEXT("TSR.History.TranslucencyAlpha"));
-		}
+		Desc.Format = PF_R8;
+		History.TranslucencyAlpha = GraphBuilder.CreateTexture(Desc, TEXT("TSR.History.TranslucencyAlpha"));
 
 		Desc.Format = PF_R16_UINT;
 		History.SubpixelDetails = GraphBuilder.CreateTexture(Desc, TEXT("TSR.History.SubpixelInfo"));
@@ -1024,13 +1027,15 @@ ITemporalUpscaler::FOutputs AddTemporalSuperResolutionPasses(
 		PassParameters->DebugOutput = CreateDebugUAV(LowFrequencyExtent, TEXT("Debug.TSR.DecimateHistory"));
 
 		FTSRDecimateHistoryCS::FPermutationDomain PermutationVector;
+		PermutationVector.Set<FTSRHighFrequencyOnlyDim>(bHistoryHighFrequencyOnly);
 		PermutationVector.Set<FTSRDecimateHistoryCS::FOutputHalfRes>(bHalfResLowFrequency);
 
 		TShaderMapRef<FTSRDecimateHistoryCS> ComputeShader(View.ShaderMap, PermutationVector);
 		FComputeShaderUtils::AddPass(
 			GraphBuilder,
-			RDG_EVENT_NAME("TSR DecimateHistory(%s) %dx%d",
-				bHalfResLowFrequency ? TEXT("HalfResShadingOutput") : TEXT(""),
+			RDG_EVENT_NAME("TSR DecimateHistory(%s%s) %dx%d",
+				bHistoryHighFrequencyOnly ? TEXT("HighFrequency") : TEXT("LowFrequency"),
+				bHalfResLowFrequency ? TEXT(" HalfResShadingOutput") : TEXT(""),
 				InputRect.Width(), InputRect.Height()),
 			ComputeShader,
 			PassParameters,
@@ -1368,7 +1373,7 @@ ITemporalUpscaler::FOutputs AddTemporalSuperResolutionPasses(
 
 	// Update temporal history.
 	{
-		ensure(SeparateTranslucencyRect.Size() == InputRect.Size() || SeparateTranslucencyTexture == nullptr);
+		ensure(SeparateTranslucencyRect.Size() == InputRect.Size() || !bHasSeparateTranslucency);
 
 		static const TCHAR* const kUpdateQualityNames[] = {
 			TEXT("Low"),
@@ -1397,12 +1402,13 @@ ITemporalUpscaler::FOutputs AddTemporalSuperResolutionPasses(
 		FScreenTransform HistoryPixelPosToViewportUV = (FScreenTransform::Identity + 0.5f) * CommonParameters.HistoryInfo.ViewportSizeInverse;
 		PassParameters->HistoryPixelPosToScreenPos = HistoryPixelPosToViewportUV * FScreenTransform::ViewportUVToScreenPos;
 		PassParameters->HistoryPixelPosToPPCo = HistoryPixelPosToViewportUV * CommonParameters.InputInfo.ViewportSize + CommonParameters.InputJitter + CommonParameters.InputPixelPosMin;
-		PassParameters->HistoryQuantizationError = (FVector3f)ComputePixelFormatQuantizationError(History.LowFrequency->Desc.Format);
+		PassParameters->HistoryQuantizationError = (FVector3f)ComputePixelFormatQuantizationError(History.HighFrequency->Desc.Format);
 		PassParameters->MinTranslucencyRejection = TranslucencyRejectionTexture == nullptr ? 1.0 : 0.0;
 		PassParameters->InvWeightClampingPixelSpeed = 1.0f / CVarTSRWeightClampingPixelSpeed.GetValueOnRenderThread();
 		PassParameters->InputToHistoryFactor = float(HistorySize.X) / float(InputRect.Width());
 		PassParameters->ResponsiveStencilMask = CVarTSREnableResponiveAA.GetValueOnRenderThread() ? (STENCIL_TEMPORAL_RESPONSIVE_AA_MASK) : 0;
 		PassParameters->bGenerateOutputMip1 = (PassInputs.bGenerateOutputMip1 && HistorySize == OutputRect.Size()) ? 1 : 0;
+		PassParameters->bHasSeparateTranslucency = bHasSeparateTranslucency;
 
 		PassParameters->PrevHistoryParameters = PrevHistoryParameters;
 		PassParameters->PrevHistory = PrevHistory;
@@ -1428,7 +1434,7 @@ ITemporalUpscaler::FOutputs AddTemporalSuperResolutionPasses(
 		PassParameters->DebugOutput = CreateDebugUAV(HistoryExtent, TEXT("Debug.TSR.UpdateHistory"));
 
 		FTSRUpdateHistoryCS::FPermutationDomain PermutationVector;
-		PermutationVector.Set<FTSRUpdateHistoryCS::FSeparateTranslucencyDim>(bAccumulateSeparateTranslucency);
+		PermutationVector.Set<FTSRHighFrequencyOnlyDim>(bHistoryHighFrequencyOnly);
 		PermutationVector.Set<FTSRUpdateHistoryCS::FQualityDim>(UpdateHistoryQuality);
 
 		TShaderMapRef<FTSRUpdateHistoryCS> ComputeShader(View.ShaderMap, PermutationVector);
@@ -1436,8 +1442,8 @@ ITemporalUpscaler::FOutputs AddTemporalSuperResolutionPasses(
 			GraphBuilder,
 			RDG_EVENT_NAME("TSR UpdateHistory(Quality=%s %s%s%s) %dx%d",
 				kUpdateQualityNames[int32(PermutationVector.Get<FTSRUpdateHistoryCS::FQualityDim>())],
-				History.LowFrequency->Desc.Format == PF_FloatR11G11B10 ? TEXT("R11G11B10") : TEXT(""),
-				PermutationVector.Get<FTSRUpdateHistoryCS::FSeparateTranslucencyDim>() ? TEXT(" SeparateTranslucency") : TEXT(""),
+				History.HighFrequency->Desc.Format == PF_FloatR11G11B10 ? TEXT("R11G11B10") : TEXT(""),
+				PermutationVector.Get<FTSRHighFrequencyOnlyDim>() ? TEXT("") : TEXT(" LowFrequency"),
 				PassInputs.bGenerateOutputMip1 ? TEXT(" OutputMip1") : TEXT(""),
 				HistorySize.X, HistorySize.Y),
 			ComputeShader,
@@ -1496,7 +1502,7 @@ ITemporalUpscaler::FOutputs AddTemporalSuperResolutionPasses(
 		PassParameters->DebugOutput = CreateDebugUAV(HistoryExtent, TEXT("Debug.TSR.ResolveHistory"));
 
 		FTSRResolveHistoryCS::FPermutationDomain PermutationVector;
-		PermutationVector.Set<FTSRUpdateHistoryCS::FSeparateTranslucencyDim>(bAccumulateSeparateTranslucency);
+		PermutationVector.Set<FTSRHighFrequencyOnlyDim>(bHistoryHighFrequencyOnly);
 
 		TShaderMapRef<FTSRResolveHistoryCS> ComputeShader(View.ShaderMap, PermutationVector);
 		FComputeShaderUtils::AddPass(
@@ -1514,14 +1520,14 @@ ITemporalUpscaler::FOutputs AddTemporalSuperResolutionPasses(
 		OutputHistory.OutputViewportRect = FIntRect(FIntPoint(0, 0), HistorySize);
 
 		// Extract filterable history
-		GraphBuilder.QueueTextureExtraction(History.LowFrequency, &OutputHistory.LowFrequency);
-		GraphBuilder.QueueTextureExtraction(History.HighFrequency, &OutputHistory.HighFrequency);
-		GraphBuilder.QueueTextureExtraction(History.Metadata, &OutputHistory.Metadata);
-		if (bAccumulateSeparateTranslucency)
+		if (!bHistoryHighFrequencyOnly)
 		{
-			GraphBuilder.QueueTextureExtraction(History.Translucency, &OutputHistory.Translucency);
-			GraphBuilder.QueueTextureExtraction(History.TranslucencyAlpha, &OutputHistory.TranslucencyAlpha);
+			GraphBuilder.QueueTextureExtraction(History.LowFrequency, &OutputHistory.LowFrequency);
 		}
+		GraphBuilder.QueueTextureExtraction(History.HighFrequency, &OutputHistory.HighFrequency);
+		GraphBuilder.QueueTextureExtraction(History.Metadata, &OutputHistory.Metadata);
+		GraphBuilder.QueueTextureExtraction(History.Translucency, &OutputHistory.Translucency);
+		GraphBuilder.QueueTextureExtraction(History.TranslucencyAlpha, &OutputHistory.TranslucencyAlpha);
 
 		// Extract non-filterable history
 		GraphBuilder.QueueTextureExtraction(History.SubpixelDetails, &OutputHistory.SubpixelDetails);
@@ -1535,7 +1541,6 @@ ITemporalUpscaler::FOutputs AddTemporalSuperResolutionPasses(
 		}
 
 		// Extract the output for next frame SSR so that separate translucency shows up in SSR.
-		if (bAccumulateSeparateTranslucency || HistorySize != OutputRect.Size())
 		{
 			GraphBuilder.QueueTextureExtraction(
 				SceneColorOutputTexture, &View.ViewState->PrevFrameViewInfo.CustomSSRInput.RT[0]);
diff --git a/Engine/Source/Runtime/Renderer/Private/SceneRendering.h b/Engine/Source/Runtime/Renderer/Private/SceneRendering.h
--- a/Engine/Source/Runtime/Renderer/Private/SceneRendering.h
+++ b/Engine/Source/Runtime/Renderer/Private/SceneRendering.h
@@ -900,7 +900,7 @@ struct FTSRHistory
 
 	bool IsValid() const
 	{
-		return LowFrequency.IsValid();
+		return HighFrequency.IsValid();
 	}
 };
 
