Fixes remaining short-circuit in the renderer

diff --git a/Engine/Plugins/Experimental/VirtualHeightfieldMesh/Shaders/Private/HeightfieldMinMaxRender.usf b/Engine/Plugins/Experimental/VirtualHeightfieldMesh/Shaders/Private/HeightfieldMinMaxRender.usf
--- a/Engine/Plugins/Experimental/VirtualHeightfieldMesh/Shaders/Private/HeightfieldMinMaxRender.usf
+++ b/Engine/Plugins/Experimental/VirtualHeightfieldMesh/Shaders/Private/HeightfieldMinMaxRender.usf
@@ -31,7 +31,7 @@ float2 UnPackMinMax(in float4 Packed)
 // Min downsample excludes zero values to avoid empty pixels affecting the calculation
 float Min4(in float4 Values, in float ReplaceZero)
 {
-	Values = Values == 0 ? ReplaceZero : Values;
+	Values = select(Values == 0, ReplaceZero, Values);
 	return min(min(min(Values.x, Values.y), Values.z), Values.w);
 }
 
diff --git a/Engine/Shaders/Private/Common.ush b/Engine/Shaders/Private/Common.ush
--- a/Engine/Shaders/Private/Common.ush
+++ b/Engine/Shaders/Private/Common.ush
@@ -814,109 +814,94 @@ float ConvertTangentUnormToSnorm8(float Input)
 {
 	int IntVal = int(round(Input * 255.0f));
 	//negate
-	IntVal = (IntVal > 127) ? (IntVal | 0xFFFFFF80) : IntVal;
+	IntVal = select(IntVal > 127, IntVal | 0xFFFFFF80, IntVal);
 	return clamp(IntVal / 127.0f, -1, 1);
 }
-
 float2 ConvertTangentUnormToSnorm8(float2 Input)
 {
 	int2 IntVal = int2(round(Input * 255.0f));
 	//negate
-	IntVal = (IntVal > 127) ? (IntVal | 0xFFFFFF80) : IntVal;
+	IntVal = select(IntVal > 127, IntVal | 0xFFFFFF80, IntVal);
 	return clamp(IntVal / 127.0f, -1, 1);
 }
-
 float3 ConvertTangentUnormToSnorm8(float3 Input)
 {
 	int3 IntVal = int3(round(Input * 255.0f));
-	IntVal = (IntVal > 127) ? (IntVal | 0xFFFFFF80) : IntVal;
+	IntVal = select(IntVal > 127, IntVal | 0xFFFFFF80, IntVal);
 	return clamp(IntVal / 127.0f, -1, 1);
 }
-
 float4 ConvertTangentUnormToSnorm8(float4 Input)
 {
 	int4 IntVal = int4(round(Input * 255.0f));
 	//negate
-	IntVal = (IntVal > 127) ? (IntVal | 0xFFFFFF80) : IntVal;
+	IntVal = select(IntVal > 127, IntVal | 0xFFFFFF80, IntVal);
 	return clamp(IntVal / 127.0f, -1, 1);
 }
-
 float ConvertTangentUnormToSnorm16(float Input)
 {
 	int IntVal = int(round(Input * 65535.0f));
 	//negate
-	IntVal = (IntVal > 32767) ? (IntVal | 0xFFFF8000) : IntVal;
+	IntVal = select(IntVal > 32767, IntVal | 0xFFFF8000, IntVal);
 	return clamp(IntVal / 32767.0f, -1, 1);
 }
-
 float2 ConvertTangentUnormToSnorm16(float2 Input)
 {
 	int2 IntVal = int2(round(Input * 65535.0f));
 	//negate
-	IntVal = (IntVal > 32767) ? (IntVal | 0xFFFFFF80) : IntVal;
+	IntVal = select(IntVal > 32767, IntVal | 0xFFFFFF80, IntVal);
 	return clamp(IntVal / 32767.0f, -1, 1);
 }
-
 float3 ConvertTangentUnormToSnorm16(float3 Input)
 {
 	int3 IntVal = int3(round(Input * 65535.0f));
-	IntVal = (IntVal > 32767) ? (IntVal | 0xFFFFFF80) : IntVal;
+	IntVal = select(IntVal > 32767, IntVal | 0xFFFFFF80, IntVal);
 	return clamp(IntVal / 32767.0f, -1, 1);
 }
-
 float4 ConvertTangentUnormToSnorm16(float4 Input)
 {
 	int4 IntVal = int4(round(Input * 65535.0f));
 	//negate
-	IntVal = (IntVal > 32767) ? (IntVal | 0xFFFFFF80) : IntVal;
+	IntVal = select(IntVal > 32767, IntVal | 0xFFFFFF80, IntVal);
 	return clamp(IntVal / 32767.0f, -1, 1);
 }
-
 float ConvertTangentSnormToUnorm8(float Input)
 {
-	float Res = Input >= 0.0f ? Input * 127 : ((Input + 1.0) * 127) + 128;
+	float Res = select(Input >= 0.0f, Input * 127, ((Input + 1.0) * 127) + 128);
 	return clamp(Res / 255, 0.0f, 0.99f);
 }
-
 float2 ConvertTangentSnormToUnorm8(float2 Input)
 {
-	float2 Res = Input >= 0.0f ? Input * 127 : ((Input + 1.0) * 127) + 128;
+	float2 Res = select(Input >= 0.0f, Input * 127, ((Input + 1.0) * 127) + 128);
 	return clamp(Res / 255, 0.0f, 0.99f);
 }
-
 float3 ConvertTangentSnormToUnorm8(float3 Input)
 {
-	float3 Res = Input >= 0.0f ? Input * 127 : ((Input + 1.0) * 127) + 128;
+	float3 Res = select(Input >= 0.0f, Input * 127, ((Input + 1.0) * 127) + 128);
 	return clamp(Res / 255, 0.0f, 0.99f);
 }
-
 float4 ConvertTangentSnormToUnorm8(float4 Input)
 {
-	float4 Res = Input >= 0.0f ? Input * 127 : ((Input + 1.0) * 127) + 128;
+	float4 Res = select(Input >= 0.0f, Input * 127, ((Input + 1.0) * 127) + 128);
 	return clamp(Res / 255, 0.0f, 0.99f);
 }
-
 float ConvertTangentSnormToUnorm16(float Input)
 {
-	float Res = Input >= 0.0f ? Input * 32767 : ((Input + 1.0) * 32767) + 32768;
+	float Res = select(Input >= 0.0f, Input * 32767, ((Input + 1.0) * 32767) + 32768);
 	return clamp(Res / 65535, 0.0f, 0.99f);
 }
-
 float2 ConvertTangentSnormToUnorm16(float2 Input)
 {
-	float2 Res = Input >= 0.0f ? Input * 32767 : ((Input + 1.0) * 32767) + 32768;
+	float2 Res = select(Input >= 0.0f, Input * 32767, ((Input + 1.0) * 32767) + 32768);
 	return clamp(Res / 65535, 0.0f, 0.99f);
 }
-
 float3 ConvertTangentSnormToUnorm16(float3 Input)
 {
-	float3 Res = Input >= 0.0f ? Input * 32767 : ((Input + 1.0) * 32767) + 32768;
+	float3 Res = select(Input >= 0.0f, Input * 32767, ((Input + 1.0) * 32767) + 32768);
 	return clamp(Res / 65535, 0.0f, 0.99f);
 }
-
 float4 ConvertTangentSnormToUnorm16(float4 Input)
 {
-	float4 Res = Input >= 0.0f ? Input * 32767 : ((Input + 1.0) * 32767) + 32768;
+	float4 Res = select(Input >= 0.0f, Input * 32767, ((Input + 1.0) * 32767) + 32768);
 	return clamp(Res / 65535, 0.0f, 0.99f);
 }
 
@@ -1594,8 +1579,8 @@ float2 LineBoxIntersect(float3 RayOrigin, float3 RayEnd, float3 BoxMin, float3 B
 /** Computes distance from an AABB to a point in space. */
 MaterialFloat ComputeDistanceFromBoxToPoint(MaterialFloat3 Mins, MaterialFloat3 Maxs, MaterialFloat3 InPoint)
 {
-	MaterialFloat3 DistancesToMin = InPoint < Mins ? abs(InPoint - Mins) : 0;
-	MaterialFloat3 DistancesToMax = InPoint > Maxs ? abs(InPoint - Maxs) : 0;
+	MaterialFloat3 DistancesToMin = select(InPoint < Mins, abs(InPoint - Mins), MaterialFloat(0.0));
+	MaterialFloat3 DistancesToMax = select(InPoint > Maxs, abs(InPoint - Maxs), MaterialFloat(0.0));
 
 	//@todo - this is actually incorrect, it gives manhattan distance
 	MaterialFloat Distance = dot(DistancesToMin, 1);
@@ -1899,11 +1884,11 @@ booltype IsPositiveFinite(type In) \
 }\
 type MakeFinite(type In) \
 {\
-    return !IsFinite(In)? 0 : In; \
+    return select(!IsFinite(In), 0.0, In); \
 }\
 type MakePositiveFinite(type In) \
 {\
-    return !IsPositiveFinite(In)? 0 : In; \
+    return select(!IsPositiveFinite(In), 0.0, In); \
 }
 
 IsAndMakeFinite_Def(float,  bool)
diff --git a/Engine/Shaders/Private/DeferredLightingCommon.ush b/Engine/Shaders/Private/DeferredLightingCommon.ush
--- a/Engine/Shaders/Private/DeferredLightingCommon.ush
+++ b/Engine/Shaders/Private/DeferredLightingCommon.ush
@@ -168,7 +168,7 @@ float ShadowRayCast(
 
 		// Off screen masking
 		float3 HitUVz = RayStartUVz + RayStepUVz * FirstHitTime;
-		bool bValidUV = all( 0.0 < HitUVz.xy && HitUVz.xy < 1.0 );
+		bool bValidUV = all(and(0.0 < HitUVz.xy, HitUVz.xy < 1.0));
 		HitDistance = bValidUV ? ( FirstHitTime * RayLength ) : -1.0;
 	}	
 	
diff --git a/Engine/Shaders/Private/ForwardShadowingCommon.ush b/Engine/Shaders/Private/ForwardShadowingCommon.ush
--- a/Engine/Shaders/Private/ForwardShadowingCommon.ush
+++ b/Engine/Shaders/Private/ForwardShadowingCommon.ush
@@ -18,7 +18,7 @@ float ComputeDirectionalLightStaticShadowing(float3 TranslatedWorldPosition)
 		float2 ShadowUVs = HomogeneousShadowPosition.xy / HomogeneousShadowPosition.w;
 
 		// Treat as unshadowed if the voxel is outside of the shadow map
-		if (all(ShadowUVs >= 0 && ShadowUVs <= 1))
+		if (all(and(ShadowUVs >= 0, ShadowUVs <= 1)))
 		{
 			#define FILTER_STATIC_SHADOWING 0
 			#if FILTER_STATIC_SHADOWING
@@ -73,7 +73,7 @@ float ComputeDirectionalLightDynamicShadowing(float3 TranslatedWorldPosition, fl
 			float4 ShadowmapMinMax = ForwardLightData.DirectionalLightShadowmapMinMax[CascadeIndex];
 
 			// Treat as unshadowed if the voxel is outside of the shadow map
-			if (all(ShadowUVs >= ShadowmapMinMax.xy && ShadowUVs <= ShadowmapMinMax.zw))
+			if (all(and(ShadowUVs >= ShadowmapMinMax.xy, ShadowUVs <= ShadowmapMinMax.zw)))
 			{
 			#if FILTER_DIRECTIONAL_LIGHT_SHADOWING
 				FPCFSamplerSettings Settings;
diff --git a/Engine/Shaders/Private/GammaCorrectionCommon.ush b/Engine/Shaders/Private/GammaCorrectionCommon.ush
--- a/Engine/Shaders/Private/GammaCorrectionCommon.ush
+++ b/Engine/Shaders/Private/GammaCorrectionCommon.ush
@@ -44,7 +44,7 @@ half3 LinearToSrgb(half3 lin)
 half3 sRGBToLinear( half3 Color ) 
 {
 	Color = max(6.10352e-5, Color); // minimum positive non-denormal (fixes black problem on DX11 AMD and NV)
-	return Color > 0.04045 ? pow( Color * (1.0 / 1.055) + 0.0521327, 2.4 ) : Color * (1.0 / 12.92);
+	return select(Color > 0.04045, pow( Color * (1.0 / 1.055) + 0.0521327, 2.4 ), Color * (1.0 / 12.92));
 }
 
 /**
diff --git a/Engine/Shaders/Private/HeightfieldLightingShared.ush b/Engine/Shaders/Private/HeightfieldLightingShared.ush
--- a/Engine/Shaders/Private/HeightfieldLightingShared.ush
+++ b/Engine/Shaders/Private/HeightfieldLightingShared.ush
@@ -76,7 +76,7 @@ float2 GetHeightfieldUV(uint HeightfieldIndex, float2 LocalPosition, out float4
 		// One row of heights in the heightfield are duplicated between subsections
 		// The only number of subsections supported is 2
 		// If we are in subsection 0, no adjustment is needed, but if we are in subsection 1 we need to shift up by one texel
-		HeightfieldUV += HeightfieldUV + .5f * HeightfieldScaleBias.xy > CenterUV ? HeightfieldScaleBias.xy : 0;
+		HeightfieldUV += select((HeightfieldUV + .5f * HeightfieldScaleBias.xy) > CenterUV, HeightfieldScaleBias.xy, 0.0);
 	}
 
 	return HeightfieldUV;
diff --git a/Engine/Shaders/Private/LWCOperations.ush b/Engine/Shaders/Private/LWCOperations.ush
--- a/Engine/Shaders/Private/LWCOperations.ush
+++ b/Engine/Shaders/Private/LWCOperations.ush
@@ -82,9 +82,9 @@ FBoolType LWCEqualsApprox(FLWCType Lhs, FFloatType Rhs, float Threshold)
 	return abs(LWCGetTile(Lhs) * UE_LWC_RENDER_TILE_SIZE + (Lhs.Offset - Rhs)) < (FFloatType)Threshold;
 }
 
-FLWCType LWCSelect(FBoolType S, FLWCType Lhs, FLWCType Rhs) { return LWCConstructor(S ? LWCGetTile(Lhs) : LWCGetTile(Rhs), S ? Lhs.Offset : Rhs.Offset); }
-FLWCType LWCSelect(FBoolType S, FFloatType Lhs, FLWCType Rhs) { return LWCConstructor(S ? (FFloatType)0 : LWCGetTile(Rhs), S ? Lhs : Rhs.Offset); }
-FLWCType LWCSelect(FBoolType S, FLWCType Lhs, FFloatType Rhs) { return LWCConstructor(S ? LWCGetTile(Lhs) : (FFloatType)0, S ? Lhs.Offset : Rhs); }
+FLWCType LWCSelect(FBoolType S, FLWCType Lhs, FLWCType Rhs) { return LWCConstructor(select(S, LWCGetTile(Lhs), LWCGetTile(Rhs)), select(S, Lhs.Offset, Rhs.Offset)); }
+FLWCType LWCSelect(FBoolType S, FFloatType Lhs, FLWCType Rhs) { return LWCConstructor(select(S, (FFloatType)0, LWCGetTile(Rhs)), select(S, Lhs, Rhs.Offset)); }
+FLWCType LWCSelect(FBoolType S, FLWCType Lhs, FFloatType Rhs) { return LWCConstructor(select(S, LWCGetTile(Lhs), (FFloatType)0), select(S, Lhs.Offset, Rhs)); }
 
 FLWCType LWCNegate(FLWCType V) { return LWCConstructor(-LWCGetTile(V), -V.Offset); }
 
@@ -166,9 +166,9 @@ FLWCType LWCMax(FLWCType Lhs, FFloatType Rhs) { return LWCSelect(LWCGreater(Lhs,
 
 FLWCType LWCAbs(FLWCType V) { return LWCSelect(LWCLess(V, (FFloatType)0), LWCNegate(V), V); }
 
-FFloatType LWCStep(FLWCType Lhs, FLWCType Rhs) { return LWCGreaterEqual(Rhs, Lhs) ? (FFloatType)1.0f : (FFloatType)0.0f; }
-FFloatType LWCStep(FLWCType Lhs, FFloatType Rhs) { return LWCGreaterEqual(Rhs, Lhs) ? (FFloatType)1.0f : (FFloatType)0.0f; }
-FFloatType LWCStep(FFloatType Lhs, FLWCType Rhs) { return LWCGreaterEqual(Rhs, Lhs) ? (FFloatType)1.0f : (FFloatType)0.0f; }
+FFloatType LWCStep(FLWCType Lhs, FLWCType Rhs)   { return select(LWCGreaterEqual(Rhs, Lhs), (FFloatType)1.0f, (FFloatType)0.0f); }
+FFloatType LWCStep(FLWCType Lhs, FFloatType Rhs) { return select(LWCGreaterEqual(Rhs, Lhs), (FFloatType)1.0f, (FFloatType)0.0f); }
+FFloatType LWCStep(FFloatType Lhs, FLWCType Rhs) { return select(LWCGreaterEqual(Rhs, Lhs), (FFloatType)1.0f, (FFloatType)0.0f); }
 
 // LWCSquareScaled returns the value squared, scaled by the LWC tile size
 FLWCType LWCSquareScaled(FLWCType V)
diff --git a/Engine/Shaders/Private/PathTracing/Volume/PathTracingVolumeCommon.ush b/Engine/Shaders/Private/PathTracing/Volume/PathTracingVolumeCommon.ush
--- a/Engine/Shaders/Private/PathTracing/Volume/PathTracingVolumeCommon.ush
+++ b/Engine/Shaders/Private/PathTracing/Volume/PathTracingVolumeCommon.ush
@@ -206,8 +206,8 @@ FPathTracingPayload CreateMediumHitPayload(float HitT, float3 TranslatedWorldPos
 	Result.PrimitiveLightingChannelMask = 7;
 	Result.SetFrontFace();
 	float3 SigmaS = ShadedResult.SigmaSRayleigh + ShadedResult.SigmaSHG;
-	Result.BaseColor = SigmaS > 0 ? min(ShadedResult.SigmaSRayleigh, SigmaS) / SigmaS : 0.0;
-	Result.CustomData.xyz = SigmaS > 0 ? min(ShadedResult.SigmaSHG, SigmaS) / SigmaS : 0.0;
+	Result.BaseColor = select(SigmaS > 0.0, min(ShadedResult.SigmaSRayleigh, SigmaS) / SigmaS, 0.0);
+	Result.CustomData.xyz = select(SigmaS > 0, min(ShadedResult.SigmaSHG, SigmaS) / SigmaS, 0.0);
 	Result.CustomData.w = ShadedResult.PhaseG;
 	return Result;
 }
diff --git a/Engine/Shaders/Private/RayTracing/GenerateCulledLightListCS.usf b/Engine/Shaders/Private/RayTracing/GenerateCulledLightListCS.usf
--- a/Engine/Shaders/Private/RayTracing/GenerateCulledLightListCS.usf
+++ b/Engine/Shaders/Private/RayTracing/GenerateCulledLightListCS.usf
@@ -35,11 +35,11 @@ void ComputeCoords(out float3 Min, out float3 Max, int3 Cell)
 	// compute bounds as logarithmically increasing from center (with minimum cell size of 2)
 	Min = abs(floor(Center));
 	Max = abs(ceil(Center));
-	Min = Min > 0.0 ? pow(2.0, Min + 1.0) : 0.0;
-	Max = Max > 0.0 ? pow(2.0, Max + 1.0) : 0.0;
+	Min = select(Min > 0.0, pow(2.0, Min + 1.0), 0.0);
+	Max = select(Max > 0.0, pow(2.0, Max + 1.0), 0.0);
 
 	// Extend range of last cell
-	Max = (Cell == 0 || Cell == CellCount - 1) ? 1.0e7 : Max;
+	Max = select(or(Cell == 0, Cell == CellCount - 1), 1.0e7, Max);
 
 	// return to sector
 	Min *= SignVal * CellScale;
diff --git a/Engine/Shaders/Private/RayTracing/RayTracingCreateGatherPointsRGS.usf b/Engine/Shaders/Private/RayTracing/RayTracingCreateGatherPointsRGS.usf
--- a/Engine/Shaders/Private/RayTracing/RayTracingCreateGatherPointsRGS.usf
+++ b/Engine/Shaders/Private/RayTracing/RayTracingCreateGatherPointsRGS.usf
@@ -442,7 +442,7 @@ RAY_TRACING_ENTRY_RAYGEN(RayTracingCreateGatherPointsRGS)
 
 						// Record the contribution
 						float3 ExitantRadianceSample = LightRadianceOverPdf * MaterialThroughput * SplitFactor;
-						Irradiance += isfinite(ExitantRadianceSample) ? ExitantRadianceSample : 0.0;
+						Irradiance += select(isfinite(ExitantRadianceSample), ExitantRadianceSample, 0.0);
 					}
 				}
 			}
diff --git a/Engine/Shaders/Private/RayTracing/RayTracingGlobalIlluminationRGS.usf b/Engine/Shaders/Private/RayTracing/RayTracingGlobalIlluminationRGS.usf
--- a/Engine/Shaders/Private/RayTracing/RayTracingGlobalIlluminationRGS.usf
+++ b/Engine/Shaders/Private/RayTracing/RayTracingGlobalIlluminationRGS.usf
@@ -303,7 +303,7 @@ RAY_TRACING_ENTRY_RAYGEN(GlobalIlluminationRGS)
 							ExitantRadianceSample *= rcp(1.0 + Luminance(ExitantRadianceSample));
 						}
 
-						Irradiance += isfinite(ExitantRadianceSample) ? ExitantRadianceSample : 0.0;
+						Irradiance += select(isfinite(ExitantRadianceSample), ExitantRadianceSample, 0.0);
 					}
 				}
 			}
diff --git a/Engine/Shaders/Private/RayTracing/RayTracingReflectionResolve.usf b/Engine/Shaders/Private/RayTracing/RayTracingReflectionResolve.usf
--- a/Engine/Shaders/Private/RayTracing/RayTracingReflectionResolve.usf
+++ b/Engine/Shaders/Private/RayTracing/RayTracingReflectionResolve.usf
@@ -104,8 +104,7 @@ bool DisocclusionHeuristic(float2 ScreenPos, float DeviceZ, float2 HistoryUV)
 	// Empirically chosen threshold to balance between aliasing/noise from false positives and ghosting from false negatives
 	float DisocclusionMaxDeviceZ = 0.001;
 
-	return (PrevDeviceZ - DeviceZ > DisocclusionMaxDeviceZ)
-		&& (CompensatedPrevDeviceZ - PrevExpectedDeviceZ > DisocclusionMaxDeviceZ);
+	return and(PrevDeviceZ - DeviceZ > DisocclusionMaxDeviceZ, CompensatedPrevDeviceZ - PrevExpectedDeviceZ > DisocclusionMaxDeviceZ);
 }
 
 float RgbLuma(float3 Color)
@@ -147,7 +146,7 @@ void RayTracingReflectionResolveCS(
 	// Temporal accumulation uses neighborhood colors in groupshared memory and clamps the lookup index, sometimes creating a discontinuity between 8x8 tiles.
 	// We can offset and wrap the thread ID every frame to move the tile border, hiding this discontinuity.
 	DispatchThreadId = (DispatchThreadId + ThreadIdOffset);
-	DispatchThreadId = DispatchThreadId < ViewSize ? DispatchThreadId : DispatchThreadId - ViewSize; // faster than mod
+	DispatchThreadId = select(DispatchThreadId < ViewSize, DispatchThreadId, DispatchThreadId - ViewSize); // faster than mod
 
 	const float2 BufferUV = (float2(DispatchThreadId + View.ViewRectMin.xy) + 0.5) * View.BufferSizeAndInvSize.zw;
 	const float2 ViewportUV = BufferUVToViewportUV(BufferUV);
diff --git a/Engine/Shaders/Private/RenderGraphUtilities.usf b/Engine/Shaders/Private/RenderGraphUtilities.usf
--- a/Engine/Shaders/Private/RenderGraphUtilities.usf
+++ b/Engine/Shaders/Private/RenderGraphUtilities.usf
@@ -30,7 +30,7 @@ void RasterizeToRectsVS(
 	out float OutRectIndex : RECT_INDEX)
 {
 	uint4 RectCoord = RectCoordBuffer[InstanceId] * DownsampleFactor;
-	uint2 VertexCoord = VertMax(VertexId) ? RectCoord.zw : RectCoord.xy;
+	uint2 VertexCoord = select(VertMax(VertexId), RectCoord.zw, RectCoord.xy);
 
 	float4 RectUV = RectCoord * InvViewSize.xyxy;
 	#if RECT_UV
@@ -38,7 +38,7 @@ void RasterizeToRectsVS(
 		RectUV = RectUVBuffer[InstanceId] * DownsampleFactor * InvTextureSize.xyxy;
 	}
 	#endif
-	float2 VertexUV = VertMax(VertexId) ? RectUV.zw : RectUV.xy;
+	float2 VertexUV = select(VertMax(VertexId), RectUV.zw, RectUV.xy);
 
 	OutPosition = float4(float2(VertexCoord) * InvViewSize * float2(2.0f, -2.0f) + float2(-1.0, 1.0f), 0.0f, 1.0f);
 	OutUV = VertexUV;
diff --git a/Engine/Shaders/Private/SubsurfaceBurleyNormalized.ush b/Engine/Shaders/Private/SubsurfaceBurleyNormalized.ush
--- a/Engine/Shaders/Private/SubsurfaceBurleyNormalized.ush
+++ b/Engine/Shaders/Private/SubsurfaceBurleyNormalized.ush
@@ -1101,7 +1101,7 @@ float4 BurleyNormalizedSS(float2 BufferUV, uint2 GroupThreadID)
 	RadianceAccumulated.xyz *= (ActiveNumOfSamples==0)? 0 :(1/ActiveNumOfSamples * 2 * PI) * EnergyNormalization;
 #else
 	// The added epsilon is used to avoid divid by zero.
-	RadianceAccumulated.xyz *= (WeightingFactor == 0)? 0:(1.0f /WeightingFactor*EnergyNormalization);
+	RadianceAccumulated.xyz *= select(WeightingFactor == 0, 0.0, 1.0f /WeightingFactor*EnergyNormalization);
 #endif
 
 	RadianceAccumulated.xyz *= BoundaryColorBleedAccum*InvNumOfSamples;
diff --git a/Engine/Shaders/Private/VertexFactoryCommon.ush b/Engine/Shaders/Private/VertexFactoryCommon.ush
--- a/Engine/Shaders/Private/VertexFactoryCommon.ush
+++ b/Engine/Shaders/Private/VertexFactoryCommon.ush
@@ -146,7 +146,7 @@ float2 UnitToOct( float3 N )
 	N.xy /= dot( 1, abs(N) );
 	if( N.z <= 0 )
 	{
-		N.xy = ( 1 - abs(N.yx) ) * ( N.xy >= 0 ? float2(1,1) : float2(-1,-1) );
+		N.xy = ( 1 - abs(N.yx) ) * select( N.xy >= 0, float2(1,1), float2(-1,-1) );
 	}
 	return N.xy;
 }
@@ -156,7 +156,7 @@ float3 OctToUnit( float2 Oct )
 	float3 N = float3( Oct, 1 - dot( 1, abs(Oct) ) );
 	if( N.z < 0 )
 	{
-		N.xy = ( 1 - abs(N.yx) ) * ( N.xy >= 0 ? float2(1,1) : float2(-1,-1) );
+		N.xy = ( 1 - abs(N.yx) ) * select( N.xy >= 0, float2(1,1), float2(-1,-1) );
 	}
 	return normalize(N);
 }
diff --git a/Engine/Shaders/Private/VolumeLightingCommonSampling.ush b/Engine/Shaders/Private/VolumeLightingCommonSampling.ush
--- a/Engine/Shaders/Private/VolumeLightingCommonSampling.ush
+++ b/Engine/Shaders/Private/VolumeLightingCommonSampling.ush
@@ -75,7 +75,7 @@ float ComputeVolumeShadowing(
 			float2 ShadowUVs = HomogeneousShadowPosition.xy / HomogeneousShadowPosition.w;
 
 			// Treat as unshadowed if the voxel is outside of the shadow map
-			if (all(ShadowUVs >= 0 && ShadowUVs <= 1))
+			if (all(and(ShadowUVs >= 0, ShadowUVs <= 1)))
 			{
 				FPCFSamplerSettings Settings;
 				Settings.ShadowDepthTexture = VSSPREFIX(StaticShadowDepthTexture);
