Fixes HLSL2021 short-circuit in TAA, MB, Denoiser, SSR

diff --git a/Engine/Shaders/Private/MotionBlur/MotionBlurApply.usf b/Engine/Shaders/Private/MotionBlur/MotionBlurApply.usf
--- a/Engine/Shaders/Private/MotionBlur/MotionBlurApply.usf
+++ b/Engine/Shaders/Private/MotionBlur/MotionBlurApply.usf
@@ -931,8 +931,8 @@ void MainCS(
 		mb_short2 OutputPixelCoord1Mip0 = (bIsValid) ? OutputPixelCoord1 : mb_short(~0).xx;
 		mb_short2 OutputPixelCoord2Mip0 = (bIsValid) ? OutputPixelCoord2 : mb_short(~0).xx;
 		mb_short2 OutputPixelCoord3Mip0 = (bIsValid) ? OutputPixelCoord3 : mb_short(~0).xx;
-		mb_short2 OutputPixelCoordMip1 = (bIsValid && (OutputMip1 != 0)) ? (OutputPixelCoord >> mb_short(1)) : mb_short(~0).xx;
-		mb_short2 OutputPixelCoordMip2 = (bIsValid && (OutputMip2 != 0) && (OutputPixelCoord & 0x3) == 0) ? (OutputPixelCoord >> mb_short(2)) : mb_short(~0).xx;
+		mb_short2 OutputPixelCoordMip1 = select(and(bIsValid, OutputMip1 != 0), (OutputPixelCoord >> mb_short(1)), mb_short(~0).xx);
+		mb_short2 OutputPixelCoordMip2 = select(and(bIsValid, and(OutputMip2 != 0, (OutputPixelCoord & 0x3) == 0)), (OutputPixelCoord >> mb_short(2)), mb_short(~0).xx);
 
 		SceneColorOutputMip0[OutputPixelCoord0Mip0] = OutputColor0;
 		SceneColorOutputMip0[OutputPixelCoord1Mip0] = OutputColor1;
@@ -1059,8 +1059,8 @@ void MainCS(
 
 		bool bIsValid = all(OutputPixelCoord < mb_short2(Color_ViewportMax));
 		mb_short2 OutputPixelCoordMip0 = (bIsValid) ? OutputPixelCoord : mb_short(~0).xx;
-		mb_short2 OutputPixelCoordMip1 = (bIsValid && (OutputMip1 != 0) && (OutputPixelCoordMip0 & 0x1) == 0) ? (OutputPixelCoord >> mb_short(1)) : mb_short(~0).xx;
-		mb_short2 OutputPixelCoordMip2 = (bIsValid && (OutputMip2 != 0) && (OutputPixelCoordMip0 & 0x3) == 0) ? (OutputPixelCoord >> mb_short(2)) : mb_short(~0).xx;
+		mb_short2 OutputPixelCoordMip1 = select(and(bIsValid, and(OutputMip1 != 0, (OutputPixelCoordMip0 & 0x1) == 0)), (OutputPixelCoord >> mb_short(1)), mb_short(~0).xx);
+		mb_short2 OutputPixelCoordMip2 = select(and(bIsValid, and(OutputMip2 != 0, (OutputPixelCoordMip0 & 0x3) == 0)), (OutputPixelCoord >> mb_short(2)), mb_short(~0).xx);
 
 		SceneColorOutputMip0[OutputPixelCoordMip0] = OutputColor;
 		SceneColorOutputMip1[OutputPixelCoordMip1] = HalfResOutput;
diff --git a/Engine/Shaders/Private/MotionBlur/MotionBlurTileGather.usf b/Engine/Shaders/Private/MotionBlur/MotionBlurTileGather.usf
--- a/Engine/Shaders/Private/MotionBlur/MotionBlurTileGather.usf
+++ b/Engine/Shaders/Private/MotionBlur/MotionBlurTileGather.usf
@@ -45,7 +45,7 @@ void MainCS(
 			int2 Offset = int2(x,y);
 			int2 SampleIndex = PixelPos + Offset;
 			
-			bool2 bInsideViewport = 0 <= SampleIndex && SampleIndex < (int2)VelocityTile_ViewportMax;
+			bool2 bInsideViewport = and(0 <= SampleIndex, SampleIndex < (int2)VelocityTile_ViewportMax);
 			if (!all(bInsideViewport))
 				continue;
 			
diff --git a/Engine/Shaders/Private/MotionBlur/MotionBlurVelocityFlatten.usf b/Engine/Shaders/Private/MotionBlur/MotionBlurVelocityFlatten.usf
--- a/Engine/Shaders/Private/MotionBlur/MotionBlurVelocityFlatten.usf
+++ b/Engine/Shaders/Private/MotionBlur/MotionBlurVelocityFlatten.usf
@@ -131,7 +131,7 @@ void MainCS(
 
 	// Output velcity flatten texture.
 	{
-		uint2 OutputPos = PixelPos.xy < Velocity_ViewportMax ? PixelPos : uint(~0).xx;
+		uint2 OutputPos = select(PixelPos.xy < Velocity_ViewportMax, PixelPos, uint(~0).xx);
 		OutVelocityFlatTexture[OutputPos] = EncodeVelocityFlatten(VelocityPolar, DeviceZ);
 	}
 
diff --git a/Engine/Shaders/Private/SSRT/SSRTRayCast.ush b/Engine/Shaders/Private/SSRT/SSRTRayCast.ush
--- a/Engine/Shaders/Private/SSRT/SSRTRayCast.ush
+++ b/Engine/Shaders/Private/SSRT/SSRTRayCast.ush
@@ -493,7 +493,7 @@ void CastScreenSpaceRay(
 		#elif IS_SSGI_SHADER
 		{
 			// If hit set to intersect time. If missed set to beyond end of ray
-            float4 HitTime = bMultipleSampleHit ? float4(1, 2, 3, 4) : 5;
+            float4 HitTime = select(bMultipleSampleHit, float4(1, 2, 3, 4), 5.0);
 
 			// Take closest hit
             float Time1 = float(i) + min(min3(HitTime.x, HitTime.y, HitTime.z), HitTime.w);
diff --git a/Engine/Shaders/Private/ScreenSpaceDenoise/SSDSpatialKernel.ush b/Engine/Shaders/Private/ScreenSpaceDenoise/SSDSpatialKernel.ush
--- a/Engine/Shaders/Private/ScreenSpaceDenoise/SSDSpatialKernel.ush
+++ b/Engine/Shaders/Private/ScreenSpaceDenoise/SSDSpatialKernel.ush
@@ -566,7 +566,7 @@ uint ComputeSignalBatchIdFromSignalMultiplexId(FSSDKernelConfig KernelConfig, co
 /** Returns whether this sample is outside the viewport. */
 bool IsOutsideViewport(FSSDKernelConfig KernelConfig, float2 SampleBufferUV)
 {
-	return any(SampleBufferUV < KernelConfig.BufferBilinearUVMinMax.xy || SampleBufferUV > KernelConfig.BufferBilinearUVMinMax.zw);
+	return any(or(SampleBufferUV < KernelConfig.BufferBilinearUVMinMax.xy, SampleBufferUV > KernelConfig.BufferBilinearUVMinMax.zw));
 }
 
 /** Sample multiplexed samples and their metadata for kernel use. */
diff --git a/Engine/Shaders/Private/TemporalAA.usf b/Engine/Shaders/Private/TemporalAA.usf
--- a/Engine/Shaders/Private/TemporalAA.usf
+++ b/Engine/Shaders/Private/TemporalAA.usf
@@ -617,7 +617,7 @@ float HistoryClip(float3 History, float3 Filtered, float3 NeighborMin, float3 Ne
 
 	float3 RayOrigin = History;
 	float3 RayDir = Filtered - History;
-	RayDir = abs( RayDir ) < (1.0/65536.0) ? (1.0/65536.0) : RayDir;
+	RayDir = select(abs( RayDir ) < (1.0/65536.0), (1.0/65536.0), RayDir);
 	float3 InvRayDir = rcp( RayDir );
 
 	float3 MinIntersect = (BoxMin - RayOrigin) * InvRayDir;
diff --git a/Engine/Shaders/Private/TonemapCommon.ush b/Engine/Shaders/Private/TonemapCommon.ush
--- a/Engine/Shaders/Private/TonemapCommon.ush
+++ b/Engine/Shaders/Private/TonemapCommon.ush
@@ -172,8 +172,8 @@ half3 FilmToneMap( half3 LinearColor )
 	half3 ToeColor		= (    -FilmBlackClip ) + (2 *      ToeScale) / ( 1 + exp( (-2 * FilmSlope /      ToeScale) * ( LogColor -      ToeMatch ) ) );
 	half3 ShoulderColor	= ( 1 + FilmWhiteClip ) - (2 * ShoulderScale) / ( 1 + exp( ( 2 * FilmSlope / ShoulderScale) * ( LogColor - ShoulderMatch ) ) );
 
-	ToeColor		= LogColor <      ToeMatch ?      ToeColor : StraightColor;
-	ShoulderColor	= LogColor > ShoulderMatch ? ShoulderColor : StraightColor;
+	ToeColor		= select(LogColor <      ToeMatch,      ToeColor, StraightColor);
+	ShoulderColor	= select(LogColor > ShoulderMatch, ShoulderColor, StraightColor);
 
 	half3 t = saturate( ( LogColor - ToeMatch ) / ( ShoulderMatch - ToeMatch ) );
 	t = ShoulderMatch < ToeMatch ? 1 - t : t;
