Fixes HLSL2021 short-circuit in TSR

diff --git a/Engine/Shaders/Private/TemporalSuperResolution/TSRDecimateHistory.usf b/Engine/Shaders/Private/TemporalSuperResolution/TSRDecimateHistory.usf
--- a/Engine/Shaders/Private/TemporalSuperResolution/TSRDecimateHistory.usf
+++ b/Engine/Shaders/Private/TemporalSuperResolution/TSRDecimateHistory.usf
@@ -100,7 +100,7 @@ void MainCS(
 			float BilinearWeight = GetSampleWeight(BilinearInter, i);
 			uint2 PixelPos = GetSamplePixelCoord(BilinearInter, i);
 
-			bool bValidPixelPos = all(PixelPos >= InputInfo_ViewportMin && PixelPos < InputInfo_ViewportMax);
+			bool bValidPixelPos = all(and(PixelPos >= InputInfo_ViewportMin, PixelPos < InputInfo_ViewportMax));
 
 			#if COMPRESS_PREV_USE_COUNT
 				uint2 CompressedUseCountPixelPos;
diff --git a/Engine/Shaders/Private/TemporalSuperResolution/TSRDilateVelocity.usf b/Engine/Shaders/Private/TemporalSuperResolution/TSRDilateVelocity.usf
--- a/Engine/Shaders/Private/TemporalSuperResolution/TSRDilateVelocity.usf
+++ b/Engine/Shaders/Private/TemporalSuperResolution/TSRDilateVelocity.usf
@@ -230,7 +230,7 @@ void MainCS(
 
 			uint PrevClosestDepth = (f32tof16(ClosestDeviceZ) << VELOCITY_HOLE_FILLING_BITS) | EncodedHoleFillingVelocity;
 			
-			bool bValidHistoryCoord = bValidOutputPixel && all(PrevInputBufferUV > PrevOutputBufferUVMin && PrevInputBufferUV < PrevOutputBufferUVMax);
+			bool bValidHistoryCoord = bValidOutputPixel && all(and(PrevInputBufferUV > PrevOutputBufferUVMin, PrevInputBufferUV < PrevOutputBufferUVMax));
 
 			FBilinearSampleInfos BilinearInter = GetBilinearSampleLevelInfos(PrevInputBufferUV, InputInfo_Extent, InputInfo_ExtentInverse);
 
diff --git a/Engine/Shaders/Private/TemporalSuperResolution/TSRHoleFillVelocity.usf b/Engine/Shaders/Private/TemporalSuperResolution/TSRHoleFillVelocity.usf
--- a/Engine/Shaders/Private/TemporalSuperResolution/TSRHoleFillVelocity.usf
+++ b/Engine/Shaders/Private/TemporalSuperResolution/TSRHoleFillVelocity.usf
@@ -104,7 +104,7 @@ void MainCS(
 		{
 			uint2 PixelPos = GetSamplePixelCoord(BilinearInter, i);
 
-			bool bValidPixelPos = all(PixelPos >= InputInfo_ViewportMin && PixelPos < InputInfo_ViewportMax);
+			bool bValidPixelPos = all(and(PixelPos >= InputInfo_ViewportMin, PixelPos < InputInfo_ViewportMax));
 
 			uint SampleHistoryClosestDeviceZ = PrevClosestDepthTexture[PixelPos];
 
diff --git a/Engine/Shaders/Private/TemporalSuperResolution/TSRSpatialAntiAliasing.usf b/Engine/Shaders/Private/TemporalSuperResolution/TSRSpatialAntiAliasing.usf
--- a/Engine/Shaders/Private/TemporalSuperResolution/TSRSpatialAntiAliasing.usf
+++ b/Engine/Shaders/Private/TemporalSuperResolution/TSRSpatialAntiAliasing.usf
@@ -129,11 +129,11 @@ void AccumulateEdgeLength(
 {
 	bool2 bStopNowIncrement = abs(StartLuma - SampleLuma) > LumaDelta;
 	bool2 bStopNowDecrement = abs(StartEdgeLuma - SampleEdgeLuma) > LumaDelta;
+	
+	bEdgeStopedByIncrement = or(bEdgeStopedByIncrement, bStopNowIncrement);
+	bEdgeStopedByDecrement = or(bEdgeStopedByDecrement, bStopNowDecrement);
 
-	bEdgeStopedByIncrement = bEdgeStopedByIncrement || bStopNowIncrement;
-	bEdgeStopedByDecrement = bEdgeStopedByDecrement || bStopNowDecrement;
-
-	EdgeLength += (bEdgeStopedByIncrement || bEdgeStopedByDecrement) ? tsr_ushort(0) : tsr_ushort(1);
+	EdgeLength += select(or(bEdgeStopedByIncrement, bEdgeStopedByDecrement), tsr_ushort(0), tsr_ushort(1));
 }
 
 tsr_half ComputeDistanceToEdge(
@@ -275,8 +275,8 @@ void MainCS(
 		#endif
 
 		bBrowseVerticaly = (TotalVarHN + TotalVarH + TotalVarHS) > (TotalVarVE + TotalVarV + TotalVarVW);
-		BrowseDirection[0] = bBrowseVerticaly ? tsr_short(0) : tsr_short(1);
-		BrowseDirection[1] = bBrowseVerticaly ? tsr_short(1) : tsr_short(0);
+		BrowseDirection[0] = select(bBrowseVerticaly, tsr_short(0), tsr_short(1));
+		BrowseDirection[1] = select(bBrowseVerticaly, tsr_short(1), tsr_short(0));
 
 		FLATTEN
 		if (bBrowseVerticaly[0])
@@ -390,11 +390,11 @@ void MainCS(
 		bool2 bEdgeStopedByMaxLumaP = (EdgeStopedByMaxLumaP & SignMask) != 0;
 		bool2 bEdgeStopedByMaxLumaN = (EdgeStopedByMaxLumaN & SignMask) != 0;
 		
-		bEdgeStopedByIncrementP = (bEdgeMoreLuminous && bEdgeStopedByMaxLumaP) || ((!bEdgeMoreLuminous) && bEdgeStopedByMinLumaP);
-		bEdgeStopedByIncrementN = (bEdgeMoreLuminous && bEdgeStopedByMaxLumaN) || ((!bEdgeMoreLuminous) && bEdgeStopedByMinLumaN);
+		bEdgeStopedByIncrementP = or(and(bEdgeMoreLuminous, bEdgeStopedByMaxLumaP), and(!bEdgeMoreLuminous, bEdgeStopedByMinLumaP));
+		bEdgeStopedByIncrementN = or(and(bEdgeMoreLuminous, bEdgeStopedByMaxLumaN), and(!bEdgeMoreLuminous, bEdgeStopedByMinLumaN));
 		
-		bEdgeStopedByDecrementP = (bEdgeMoreLuminous && bEdgeStopedByMinLumaP) || ((!bEdgeMoreLuminous) && bEdgeStopedByMaxLumaP);
-		bEdgeStopedByDecrementN = (bEdgeMoreLuminous && bEdgeStopedByMinLumaN) || ((!bEdgeMoreLuminous) && bEdgeStopedByMaxLumaN);
+		bEdgeStopedByDecrementP = or(and(bEdgeMoreLuminous, bEdgeStopedByMinLumaP), and(!bEdgeMoreLuminous, bEdgeStopedByMaxLumaP));
+		bEdgeStopedByDecrementN = or(and(bEdgeMoreLuminous, bEdgeStopedByMinLumaN), and(!bEdgeMoreLuminous, bEdgeStopedByMaxLumaN));
 
 #else // Reference implementation
 		UNROLL_N(CONFIG_ITERATIONS)
diff --git a/Engine/Shaders/Private/TemporalSuperResolution/TSRUpdateHistory.usf b/Engine/Shaders/Private/TemporalSuperResolution/TSRUpdateHistory.usf
--- a/Engine/Shaders/Private/TemporalSuperResolution/TSRUpdateHistory.usf
+++ b/Engine/Shaders/Private/TemporalSuperResolution/TSRUpdateHistory.usf
@@ -486,7 +486,7 @@ void MainCS(
 				bIsResponsiveAAPixel = (SceneStencilRef & ResponsiveStencilMask) != 0;
 
 				#if CONFIG_REJECT_TRAMSLUCENCY_ON_RESPONSIVE_AA_ONLY
-					TranslucencyRejection = max(TranslucencyRejection, bIsResponsiveAAPixel ? tsr_half(0.0) : tsr_half(1.0));
+					TranslucencyRejection = max(TranslucencyRejection, select(bIsResponsiveAAPixel, tsr_half(0.0).xx, tsr_half(1.0).xx));
 				#endif
 			}
 			#endif
