Fix for shader compile error on some platforms introduced in CL 19352018

diff --git a/Engine/Shaders/Private/MotionBlur/MotionBlurVelocityFlatten.usf b/Engine/Shaders/Private/MotionBlur/MotionBlurVelocityFlatten.usf
--- a/Engine/Shaders/Private/MotionBlur/MotionBlurVelocityFlatten.usf
+++ b/Engine/Shaders/Private/MotionBlur/MotionBlurVelocityFlatten.usf
@@ -131,7 +131,7 @@ void MainCS(
 
 	// Output velcity flatten texture.
 	{
-		uint2 OutputPos = select(PixelPos.xy < Velocity_ViewportMax, PixelPos, uint(~0).xx);
+		uint2 OutputPos = select(PixelPos.xy < Velocity_ViewportMax, uint2(PixelPos), uint(~0).xx);
 		OutVelocityFlatTexture[OutputPos] = EncodeVelocityFlatten(VelocityPolar, DeviceZ);
 	}
 
