Clarified Strata environment lighting in relation to SSS (making sure cloud shadow do not influence  luma for SSS correctly)

diff --git a/Engine/Shaders/Private/ReflectionEnvironmentPixelShader.usf b/Engine/Shaders/Private/ReflectionEnvironmentPixelShader.usf
--- a/Engine/Shaders/Private/ReflectionEnvironmentPixelShader.usf
+++ b/Engine/Shaders/Private/ReflectionEnvironmentPixelShader.usf
@@ -230,7 +230,7 @@ float3 ReflectionEnvironment(FGBufferData GBuffer, float AmbientOcclusion, float
 	return -min(-Color.rgb, 0.0);
 }
 
-void ApplyCloudVolumetricShadow(in float3 TranslatedWorldPosition, inout float3 SkyLighting)
+float GetCloudVolumetricAOShadow(in float3 TranslatedWorldPosition)
 {
 	// Ideally we would compute spatially how much of the sky+cloud is visible for each point in the world. But we evaluate the sky light only once at the sky light component position as of today.
 	// To add some spatial variation, we can affect the sky light diffuse contribution according to cloud occlusion from above.
@@ -240,14 +240,16 @@ void ApplyCloudVolumetricShadow(in float3 TranslatedWorldPosition, inout float3
 	if (CloudSkyAOEnabled)
 	{
 		float OutOpticalDepth = 0.0f;
-		SkyLighting *= GetCloudVolumetricShadow(TranslatedWorldPosition, CloudSkyAOWorldToLightClipMatrix, CloudSkyAOFarDepthKm, CloudSkyAOTexture, CloudSkyAOSampler, OutOpticalDepth);
+		return GetCloudVolumetricShadow(TranslatedWorldPosition, CloudSkyAOWorldToLightClipMatrix, CloudSkyAOFarDepthKm, CloudSkyAOTexture, CloudSkyAOSampler, OutOpticalDepth);
 	}
 #endif
+	return 1.0f;
 }
 
 void ReflectionEnvironmentSkyLighting(
 	in float4 SvPosition : SV_Position,
-	out float4 OutColor : SV_Target0)
+	out float4 OutColor : SV_Target0
+	)
 {
 	uint2 PixelPos = SvPosition.xy;
 	float2 BufferUV = SvPositionToBufferUV(SvPosition);
@@ -283,6 +285,8 @@ void ReflectionEnvironmentSkyLighting(
 		float SSRReductionFactor = 1.0f;
 		float3 TopLayerEnvBRDF = 0.0f;
 
+		const float CloudVolumetricAOShadow = GetCloudVolumetricAOShadow(TranslatedWorldPosition);
+
 		float4 StrataEnvLighting = StrataEnvironmentLighting(
 			SvPosition,
 			BufferUV,
@@ -296,15 +300,15 @@ void ReflectionEnvironmentSkyLighting(
 			Strata.MaterialTextureArray,
 			StrataAddressing,
 			StrataPixelHeader,
+			CloudVolumetricAOShadow,
 			SSRReductionFactor,
 			TopLayerEnvBRDF);
 
-		float3 EnvLighting = SSR.rgb * TopLayerEnvBRDF * SSRReductionFactor + StrataEnvLighting.rgb;
-
-		ApplyCloudVolumetricShadow(WorldPosition, EnvLighting);
+		// Screen space reflections added on top of environment lighting.
+		// Alpha is 0 to not affect the subsurface diffuse luma information.
+		const float4 SSRColor = float4(SSR.rgb * TopLayerEnvBRDF * SSRReductionFactor, 0.0);
 
-		OutColor.xyz = EnvLighting;
-		OutColor.a = StrataEnvLighting.a;
+		OutColor += StrataEnvLighting + SSRColor;
 	}
 
 #else // STRATA_ENABLED
@@ -334,14 +338,13 @@ void ReflectionEnvironmentSkyLighting(
 	BRANCH
 	if (!bUnlitMaterial) // Only light pixels marked as lit
 	{
-		float3 SkyLighting = SkyLightDiffuse(GBuffer, AmbientOcclusion, BufferUV, ScreenPosition, BentNormal, DiffuseColor);
-
 		float3 TranslatedWorldPosition = mul(float4(ScreenPosition * GBuffer.Depth, GBuffer.Depth, 1), View.ScreenToTranslatedWorld).xyz;
-		ApplyCloudVolumetricShadow(TranslatedWorldPosition, SkyLighting);
+		const float CloudVolumetricAOShadow = GetCloudVolumetricAOShadow(TranslatedWorldPosition);
+
+		float3 SkyLighting = CloudVolumetricAOShadow * SkyLightDiffuse(GBuffer, AmbientOcclusion, BufferUV, ScreenPosition, BentNormal, DiffuseColor);
 
 		FLightAccumulator LightAccumulator = (FLightAccumulator)0;
 		const bool bNeedsSeparateSubsurfaceLightAccumulation = UseSubsurfaceProfile(ShadingModelID);
-
 		LightAccumulator_Add(LightAccumulator, SkyLighting, SkyLighting, 1.0f, bNeedsSeparateSubsurfaceLightAccumulation);
 		OutColor = LightAccumulator_GetResult(LightAccumulator);
 	}
diff --git a/Engine/Shaders/Private/Strata/StrataEnvironmentLighting.ush b/Engine/Shaders/Private/Strata/StrataEnvironmentLighting.ush
--- a/Engine/Shaders/Private/Strata/StrataEnvironmentLighting.ush
+++ b/Engine/Shaders/Private/Strata/StrataEnvironmentLighting.ush
@@ -25,6 +25,7 @@ float4 StrataEnvironmentLighting(
 	FStrataMaterialContainer MaterialBuffer,
 	FStrataAddressing StrataAddressing,
 	FStrataPixelHeader StrataPixelHeader,
+	float CloudVolumetricAOShadow,
 	inout float SSRReductionFactor,
 	inout float3 TopLayerEnvBRDF)
 {
@@ -88,7 +89,7 @@ float4 StrataEnvironmentLighting(
 				float3 DiffuseLookup = GetSkySHDiffuse(BentNormal) * View.SkyLightColor.rgb;
 				// And accumulate
 				// Note: Use diffuse directional albedo (i.e., DiffuseWeight) as first order approximation for env. integration (STRATA_TODO instead compute SH coefficients for Chan)
-				DiffuseLighting = BSDFThroughput * (SkyVisData.SkyDiffuseLookUpMul * DiffuseLookup + SkyVisData.SkyDiffuseLookUpAdd) * StrataEnvLight.DiffuseWeight;
+				DiffuseLighting = CloudVolumetricAOShadow * BSDFThroughput * (SkyVisData.SkyDiffuseLookUpMul * DiffuseLookup + SkyVisData.SkyDiffuseLookUpAdd) * StrataEnvLight.DiffuseWeight;
 			}
 			if (bProcessBackFaceDiffuse)
 			{
@@ -99,7 +100,7 @@ float4 StrataEnvironmentLighting(
 				float3 DiffuseLookup = GetSkySHDiffuse(-BentNormal) * View.SkyLightColor.rgb;
 				// And accumulate
 				// Note: Use diffuse directional albedo (i.e., DiffuseWeight) as first order approximation for env. integration (STRATA_TODO instead compute SH coefficients for Chan)
-				DiffuseLighting += BSDFThroughput * (SkyVisData.SkyDiffuseLookUpMul * DiffuseLookup) * StrataEnvLight.DiffuseBackFaceWeight;
+				DiffuseLighting += CloudVolumetricAOShadow * BSDFThroughput * (SkyVisData.SkyDiffuseLookUpMul * DiffuseLookup) * StrataEnvLight.DiffuseBackFaceWeight;
 			}
 		}
 	#endif // ENABLE_DYNAMIC_SKY_LIGHT
@@ -179,10 +180,9 @@ float4 StrataEnvironmentLighting(
 			TopLayerEnvBRDF += BSDFThroughput * StrataEnvLight.SpecularWeight;
 		}
 
-		const bool bNeedsSeparateSubsurfaceLightAccumulation = StrataEnvLight.bSubsurface;
-		LightAccumulator_AddSplit(Out, DiffuseLighting, SpecularLighting, DiffuseLighting, View.PreExposure, bNeedsSeparateSubsurfaceLightAccumulation);
+		// STRATA_TODO ReflectionEnvironment(...) for static reflection volumes.
 
-		// STRATA_TODO: Missing probe reflection here
+		LightAccumulator_AddSplit(Out, DiffuseLighting, SpecularLighting, DiffuseLighting, View.PreExposure, StrataEnvLight.bSubsurface);
 	}
 
 	SSRReductionFactor = saturate(SSRReductionFactor);
