Update strata DMFP to MFP

diff --git a/Engine/Shaders/Private/BurleyNormalizedSSSCommon.ush b/Engine/Shaders/Private/BurleyNormalizedSSSCommon.ush
--- a/Engine/Shaders/Private/BurleyNormalizedSSSCommon.ush
+++ b/Engine/Shaders/Private/BurleyNormalizedSSSCommon.ush
@@ -94,19 +94,37 @@ float3 GetSearchLightDiffuseScalingFactor3D(float3 SurfaceAlbedo)
 	return 3.5 + 100 * Value * Value * Value * Value;
 }
 
+// TODO: find the source of this magic number to convert dmfp to mfp to match path tracer and rasterizer.
+// DMFP/SearchLightDiffuseScalingFactor goes to the normalized distance(d=l/s), then multiply the PerpendicularScalingFactor to get the mfp
+// float Dmfp2MfpMagicNumber = 0.6f;
+
+// Mathmatically matching based on diffusion coefficient instead of burley's approximation. However, it leads to incorrect result as 
+// we use burley's approximation (IOR=1.0) for screenspace diffuse scattering.
+//float3 Alpha = 1 - exp(-11.43 * SurfaceAlbedo + 15.38 * SurfaceAlbedo * SurfaceAlbedo - 13.91 * SurfaceAlbedo * SurfaceAlbedo * SurfaceAlbedo);
+//return DMFP * sqrt(3 * (1 - Alpha) / (2 - Alpha));
+float3 GetMFPFromDMFPCoeff(float3 DMFPSurfaceAlbedo, float3 MFPSurfaceAlbedo, float Dmfp2MfpMagicNumber = 0.6f)
+{
+	return Dmfp2MfpMagicNumber * GetPerpendicularScalingFactor3D(MFPSurfaceAlbedo) / GetSearchLightDiffuseScalingFactor3D(DMFPSurfaceAlbedo);
+}
+
 float3 GetMFPFromDMFPApprox(float3 SurfaceAlbedo, float3 TargetSurfaceAlbedo, float3 DMFP)
 {
-	// TODO: find the source of this magic number to convert dmfp to mfp to match path tracer and rasterizer.
-	// DMFP/SearchLightDiffuseScalingFactor goes to the normalized distance(d=l/s), then multiply the PerpendicularScalingFactor to get the mfp
+	return GetMFPFromDMFPCoeff(SurfaceAlbedo, TargetSurfaceAlbedo) * DMFP;
+}
 
-	const float Dmfp2MfpMagicNumber = 0.6f;
-	float3 DMFPToMFPCoeff = Dmfp2MfpMagicNumber * GetPerpendicularScalingFactor3D(TargetSurfaceAlbedo) / GetSearchLightDiffuseScalingFactor3D(SurfaceAlbedo);
-	return DMFPToMFPCoeff * DMFP;
-	
-	// Mathmatically matching based on diffusion coefficient instead of burley's approximation. However, it leads to incorrect result as 
-	// we use burley's approximation (IOR=1.0) for screenspace diffuse scattering.
-	//float3 Alpha = 1 - exp(-11.43 * SurfaceAlbedo + 15.38 * SurfaceAlbedo * SurfaceAlbedo - 13.91 * SurfaceAlbedo * SurfaceAlbedo * SurfaceAlbedo);
-	//return DMFP * sqrt(3 * (1 - Alpha) / (2 - Alpha));
+float3 GetDMFPFromMFPApprox(float3 SurfaceAlbedo, float3 MFP)
+{
+	float3 MFPFromDMFPCoeff = GetMFPFromDMFPCoeff(SurfaceAlbedo, SurfaceAlbedo);
+	return MFP / MFPFromDMFPCoeff;
+}
+
+// With world unit scale 
+float4 GetSubsurfaceProfileMFPInCm(int SubsurfaceProfileInt)
+{
+	float4 DMFP = GetSubsurfaceProfileDMFPInCm(SubsurfaceProfileInt);
+	float4 SurfaceAlbedo = GetSubsurfaceProfileSurfaceAlbedo(SubsurfaceProfileInt);
+
+	return float4(GetMFPFromDMFPApprox(SurfaceAlbedo.xyz, SurfaceAlbedo.xyz, DMFP.xyz),0.0f);
 }
 
 float GetScalingFactor(float A)
@@ -166,7 +184,7 @@ float3 InternalGetBurleyTransmissionProfile(
 	float ThicknessInCm)
 {
 	// Assuming that the volume albedo is the same to the surface albedo for transmission;
-	const float3 ScalingFactor = GetSearchLightDiffuseScalingFactor3D(SubsurfaceAlebdo);
+	const float3 ScalingFactor = GetPerpendicularScalingFactor3D(SubsurfaceAlebdo);
 	const float3 r = ThicknessInCm / MeanFreePathInCm.xyz;
 	const float3 Output = 0.25 * SubsurfaceAlebdo * (exp(-ScalingFactor * r) + 3 * exp(-ScalingFactor * r / 3));
 
diff --git a/Engine/Shaders/Private/DBufferDecalShared.ush b/Engine/Shaders/Private/DBufferDecalShared.ush
--- a/Engine/Shaders/Private/DBufferDecalShared.ush
+++ b/Engine/Shaders/Private/DBufferDecalShared.ush
@@ -161,10 +161,10 @@ void ApplyDBufferData(in FStrataDBuffer StrataDBufferData, inout FStrataPixelHea
 					}
 
 					// Disable SSS when the coverage reach one to preserve the decal color
-					SLAB_SSSDMFP(BSDF) = lerp(0.f, SLAB_SSSDMFP(BSDF), StrataDBufferData.OneMinusCoverage_BaseColor);
-					if (all(SLAB_SSSDMFP(BSDF) == 0.0f))
+					SLAB_SSSMFP(BSDF) = lerp(0.f, SLAB_SSSMFP(BSDF), StrataDBufferData.OneMinusCoverage_BaseColor);
+					if (all(SLAB_SSSMFP(BSDF) == 0.0f))
 					{
-						BSDF_SETHASDMFP(BSDF, 0);
+						BSDF_SETHASMFP(BSDF, 0);
 						BSDF_SETHASSSS(BSDF, 0);
 					}
 				#endif
diff --git a/Engine/Shaders/Private/DebugProbes.usf b/Engine/Shaders/Private/DebugProbes.usf
--- a/Engine/Shaders/Private/DebugProbes.usf
+++ b/Engine/Shaders/Private/DebugProbes.usf
@@ -294,8 +294,8 @@ void MainPS(
 			/*float Roughness						*/Roughness,
 			/*float Anisotropy						*/0.0f,
 			/*float SSSProfileID					*/0.0f,
-			/*float3 SSSDMFP						*/0.0f,
-			/*float SSSDMFPScale					*/0.0f,
+			/*float3 SSSMFP							*/0.0f,
+			/*float SSSMFPScale						*/0.0f,
 			/*float3 Emissive						*/0.0f,
 			/*float Haziness						*/0.0f,
 			/*float FuzzAmount						*/0.0f,
diff --git a/Engine/Shaders/Private/PostProcessSubsurface.usf b/Engine/Shaders/Private/PostProcessSubsurface.usf
--- a/Engine/Shaders/Private/PostProcessSubsurface.usf
+++ b/Engine/Shaders/Private/PostProcessSubsurface.usf
@@ -187,7 +187,7 @@ float4 GetSceneColor(float2 BufferUV)
 half4 GetSubsurfaceProfileKernel(int SampleIndex, float3 D, float ScatterRadiusInMM, float WorldUnitScale)
 {
 	// See notes in SubsurfaceProfile.cpp. Scatter radius is in mm instead of cm here.
-	ScatterRadiusInMM = max(ScatterRadiusInMM/2.229f,1.0f);
+	ScatterRadiusInMM = max(ScatterRadiusInMM,1.0f);
 	
 	const float Range = ((2 * SSSS_N_KERNELWEIGHTCOUNT - 1) > 20) ? 3.0f : 2.0f;
 
@@ -213,19 +213,20 @@ half4 GetSubsurfaceProfileKernel(int SampleIndex, float3 D, float ScatterRadiusI
 			// Need to multiply by 0.5f, and modify MaxCDF to match as close as offline LUT. Here favors integration 
 			// instead of discrete approximation.
 			Kernel.rgb = (GetCDF3D(D, ScatterRadiusInMM * x1) - GetCDF3D(D, ScatterRadiusInMM * x0)) / MaxCDF;
-			Kernel.a =  4.0f * x; // 4.0f is consistent to the scale in LUT
+			Kernel.a = 2.0 * x; // 2.0f is consistent to the scale in ComputeMirroredBSSSKernel of BurleyNormalizedSSS.cpp
 		}
 
 	// Scale to match offline encoding
-	const float TableMaxA = 3.0f;
+	// const float TableMaxA = SUBSURFACE_KERNEL_SIZE;
 	// const FLinearColor TableColorScale = FLinearColor(
 	//		1.0f / TableMaxRGB,
 	//		1.0f / TableMaxRGB,
 	//		1.0f / TableMaxRGB,
 	//		1.0f / TableMaxA);
 	// C *= TableColorScale
-	// C.W *= Data.ScatterRadius / SUBSURFACE_RADIUS_SCALE;     
-	Kernel.a *= (ScatterRadiusInMM * BURLEY_MM_2_CM) * WorldUnitScale * 2.229f * 0.707f / (TableMaxA * SUBSURFACE_RADIUS_SCALE);
+	// C.W *= Data.ScatterRadius / SUBSURFACE_RADIUS_SCALE;
+	// SUBSURFACE_KERNEL_SIZE and TableMaxA cancels out
+	Kernel.a *= (ScatterRadiusInMM * BURLEY_MM_2_CM) * WorldUnitScale /** SUBSURFACE_KERNEL_SIZE*/ /(/*TableMaxA **/ SUBSURFACE_RADIUS_SCALE);
 
 	return Kernel;
 }
@@ -249,15 +250,17 @@ FBurleyParameter GetBurleyParameters(FStrataSubsurfaceData SSSData)
 	}
 	else if (SSSData.bIsValid)
 	{
-		// Decode the DMFP stored as Float111110 in a uint over a RGBA8 unorm
-		const float3 DMFP = SSSData.DMFP * BURLEY_CM_2_MM; // Stored DMFP is in cm, converted into mm
+		// Decode the MFP stored as Float111110 in a uint over a RGBA8 unorm
+		const float3 MFP = SSSData.MFP * BURLEY_CM_2_MM; // Stored DMFP is in cm, converted into mm
 
-		const float MaxDMFP = max3(DMFP.x, DMFP.y, DMFP.z);
-		const uint IndexOfMaxDmfp = (DMFP.r == MaxDMFP) ? 0 : ((DMFP.g == MaxDMFP) ? 1 : 2);
+		const float MaxMFP = max3(MFP.x, MFP.y, MFP.z);
+		const uint IndexOfMaxMfp = (MFP.r == MaxMFP) ? 0 : ((MFP.g == MaxMFP) ? 1 : 2);
 
-		Out.SurfaceAlbedo = float4(SSSData.BaseColor, SSSData.BaseColor[IndexOfMaxDmfp]);
-		Out.DiffuseMeanFreePath = float4(DMFP, MaxDMFP);
-		Out.WorldUnitScale = BURLEY_CM_2_MM;
+		const float3 DMFP = GetDMFPFromMFPApprox(SSSData.BaseColor, MFP);
+
+		Out.SurfaceAlbedo = float4(SSSData.BaseColor, SSSData.BaseColor[IndexOfMaxMfp]);
+		Out.DiffuseMeanFreePath = float4(DMFP, DMFP[IndexOfMaxMfp]);
+		Out.WorldUnitScale = 1.0f;	// no world unit scale change.
 
 		// STRATA_TODO smooth transistion based on pixel footprint w.r.t. MFPRadius and pixel footprint
 		Out.SurfaceOpacity = 1.0f;
diff --git a/Engine/Shaders/Private/RayTracing/RayTracingOcclusionRGS.usf b/Engine/Shaders/Private/RayTracing/RayTracingOcclusionRGS.usf
--- a/Engine/Shaders/Private/RayTracing/RayTracingOcclusionRGS.usf
+++ b/Engine/Shaders/Private/RayTracing/RayTracingOcclusionRGS.usf
@@ -626,7 +626,7 @@ RAY_TRACING_ENTRY_RAYGEN(OcclusionRGS)
 			ShadingParameters.bIsHair = StrataHasShadingModel(StrataPixelHeader, STRATA_BSDF_TYPE_HAIR);
 			ShadingParameters.bIsSubsurface = SSSData.bIsValid && !SSSData.bIsProfile;
 			ShadingParameters.bHasSubsurfaceProfile = SSSData.bIsValid && SSSData.bIsProfile;
-			ShadingParameters.bBackfaceCulling = any(SSSData.DMFP > 0); // STRATA_TODO
+			ShadingParameters.bBackfaceCulling = any(SSSData.MFP > 0); // STRATA_TODO
 
 			WorldNormal = TopLayerData.WorldNormal;
 
diff --git a/Engine/Shaders/Private/Strata/Strata.ush b/Engine/Shaders/Private/Strata/Strata.ush
--- a/Engine/Shaders/Private/Strata/Strata.ush
+++ b/Engine/Shaders/Private/Strata/Strata.ush
@@ -142,7 +142,7 @@ struct FStrataBSDF
 
 	float3				Emissive;		// All bsdf can potientially have emissive. This is never written out to the strata gbuffer but instead accumulated in the emissive buffer during the base pass.
 	float				ThicknessCm;	// This must be kept in order to be able to normalize the mean free path to comply with our volumetric material math all done in a normalised slab of medium.
-	float3				TmpDMFP;
+	float3				TmpMFP;
 
 	float				TopLayerDataWeight;// The weight this BSDF contributes to the top layer (normal and roughness for SSR and SSAO for instance)
 #endif
@@ -292,7 +292,7 @@ struct FStrataBSDF
 #define STATE_BIT_COUNT_F90						1
 #define STATE_BIT_COUNT_THINFILM				1
 #define STATE_BIT_COUNT_SIMPLEVOLUME			1
-#define STATE_BIT_COUNT_DMFPPLUGGED				1
+#define STATE_BIT_COUNT_MFPPLUGGED				1
 #define STATE_BIT_COUNT_HASFUZZ					1
 #define STATE_BIT_COUNT_ISTHIN					1
 #define STATE_BIT_COUNT_BSDF_TYPE				3
@@ -310,8 +310,8 @@ struct FStrataBSDF
 #define STATE_BIT_OFFSET_F90					(STATE_BIT_OFFSET_HAZINESS			+ STATE_BIT_COUNT_HAZINESS)
 #define STATE_BIT_OFFSET_THINFILM				(STATE_BIT_OFFSET_F90				+ STATE_BIT_COUNT_F90)
 #define STATE_BIT_OFFSET_SIMPLEVOLUME			(STATE_BIT_OFFSET_THINFILM			+ STATE_BIT_COUNT_THINFILM)
-#define STATE_BIT_OFFSET_DMFPPLUGGED			(STATE_BIT_OFFSET_SIMPLEVOLUME		+ STATE_BIT_COUNT_SIMPLEVOLUME)
-#define STATE_BIT_OFFSET_HASFUZZ				(STATE_BIT_OFFSET_DMFPPLUGGED		+ STATE_BIT_COUNT_DMFPPLUGGED)
+#define STATE_BIT_OFFSET_MFPPLUGGED				(STATE_BIT_OFFSET_SIMPLEVOLUME		+ STATE_BIT_COUNT_SIMPLEVOLUME)
+#define STATE_BIT_OFFSET_HASFUZZ				(STATE_BIT_OFFSET_MFPPLUGGED		+ STATE_BIT_COUNT_MFPPLUGGED)
 #define STATE_BIT_OFFSET_ISTHIN					(STATE_BIT_OFFSET_HASFUZZ			+ STATE_BIT_COUNT_HASFUZZ)
 #define STATE_BIT_OFFSET_BSDF_TYPE				(STATE_BIT_OFFSET_ISTHIN			+ STATE_BIT_COUNT_ISTHIN)
 #define STATE_BIT_OFFSET_SHAREDLOCALBASESID		(STATE_BIT_OFFSET_BSDF_TYPE			+ STATE_BIT_COUNT_BSDF_TYPE)
@@ -383,8 +383,8 @@ struct FStrataBSDF
 #define BSDF_SETISSIMPLEVOLUME(X, Value) WRITE_BITS(X.State, STATE_BIT_COUNT_SIMPLEVOLUME, STATE_BIT_OFFSET_SIMPLEVOLUME, Value)
 
 // Indicates if the BSDF is part of the top layer
-#define BSDF_GETHASDMFP(X)				READ_BITS(X.State,  STATE_BIT_COUNT_DMFPPLUGGED, STATE_BIT_OFFSET_DMFPPLUGGED)
-#define BSDF_SETHASDMFP(X, Value)		WRITE_BITS(X.State, STATE_BIT_COUNT_DMFPPLUGGED, STATE_BIT_OFFSET_DMFPPLUGGED, Value)
+#define BSDF_GETHASMFP(X)				READ_BITS(X.State,  STATE_BIT_COUNT_MFPPLUGGED, STATE_BIT_OFFSET_MFPPLUGGED)
+#define BSDF_SETHASMFP(X, Value)		WRITE_BITS(X.State, STATE_BIT_COUNT_MFPPLUGGED, STATE_BIT_OFFSET_MFPPLUGGED, Value)
 
 // Indicates if the BSDF is part of the top layer
 #define BSDF_GETHASFUZZ(X)				READ_BITS(X.State,  STATE_BIT_COUNT_HASFUZZ, STATE_BIT_OFFSET_HASFUZZ)
@@ -428,8 +428,8 @@ struct FStrataBSDF
 // Either with SSS profile
 #define SLAB_SSSPROFILEID(X)			X.VGPRs[3].x
 #define SLAB_SSSPROFILERADIUSSCALE(X)	X.VGPRs[3].y
-// Or explicit DMFP in centimeter
-#define SLAB_SSSDMFP(X)					X.VGPRs[3].xyz
+// Or explicit MFP in centimeter
+#define SLAB_SSSMFP(X)					X.VGPRs[3].xyz
 // Thin film. IOR is harcoded to oil IOR in order to save VGPR
 #define SLAB_THINFILM_THICKNESS(X)		X.VGPRs[3].w
 #define SLAB_THINFILM_IOR(X)			1.44f
@@ -583,14 +583,14 @@ struct FStrataSubsurfaceData
 	uint   ProfileId;
 	float  ProfileRadiusScale;
 
-	float3 DMFP;				// When the profile is not used, a per pixel DMFP is used
+	float3 MFP;				// When the profile is not used, a per pixel MFP is used
 };
 
 uint2 StrataPackSSSData(FStrataSubsurfaceData In)
 {
 	uint2 Out;
 	Out.x = PackColorLinearToGamma2AlphaLinear(float4(In.BaseColor, In.bIsValid ? (In.bIsProfile ? 1.0f : 0.5f) : 0.0f));
-	Out.y = In.bIsProfile ? ((In.ProfileId & 0xFF) | PackRGBA8(float4(0.f, In.ProfileRadiusScale, In.Specular, 0.0f))) : PackR11G11B10F(In.DMFP);
+	Out.y = In.bIsProfile ? ((In.ProfileId & 0xFF) | PackRGBA8(float4(0.f, In.ProfileRadiusScale, In.Specular, 0.0f))) : PackR11G11B10F(In.MFP);
 	return Out;
 }
 
@@ -621,7 +621,7 @@ FStrataSubsurfaceData StrataUnpackSSSData(uint2 In)
 	{
 		Out.ProfileId = SSS_PROFILE_ID_PERPIXEL;
 		Out.ProfileRadiusScale = 1.0f;
-		Out.DMFP = UnpackR11G11B10F(In.y);
+		Out.MFP = UnpackR11G11B10F(In.y);
 	}
 	return Out;
 }
@@ -821,7 +821,7 @@ FStrataData GetStrataSlabBSDF(
 	float3 BaseColor, float3 EdgeColor, float Specular, float Metallic,	//  Metalness workflow
 	float3 DiffuseAlbedo, float3 F0, float3 F90,						// !Metalness workflow
 	float Roughness, float Anisotropy,
-	float SSSProfileID, float3 SSSDMFP, float SSSDMFPScale,
+	float SSSProfileID, float3 SSSMFP, float SSSMFPScale,
 	float3 Emissive, 
 	float Haziness, 
 	float FuzzAmount, float3 FuzzColor,
@@ -831,12 +831,12 @@ FStrataData GetStrataSlabBSDF(
 	FStrataData StrataData = GetInitialisedStrataData();
 
 	// STRATA_TODO apply SSS as a function of pixel footprint w.r.t. MFP
-	SSSDMFPScale = saturate(SSSDMFPScale);
+	SSSMFPScale = saturate(SSSMFPScale);
 	const bool bSSSProfilePlugged = SSSProfileID > 0.0f;
-	const bool bDMFPPlugged = any(SSSDMFP > 0);
-	bool bIsThin = bDMFPPlugged && Thickness < STRATA_LAYER_ISTHIN_THICKNESS_THRESHOLD_CM;
-	bool bHasSSSProfile = !bIsThin && bSSSProfilePlugged && SSSDMFPScale > 0.f;	// When no profile is provide, we use Burley. STRATA_TODO we should clarify as burley is also an option of the profile.
-	bool bHasSSS = bDMFPPlugged || bHasSSSProfile;
+	const bool bMFPPlugged = any(SSSMFP > 0);
+	bool bIsThin = bMFPPlugged && Thickness < STRATA_LAYER_ISTHIN_THICKNESS_THRESHOLD_CM;
+	bool bHasSSSProfile = !bIsThin && bSSSProfilePlugged && SSSMFPScale > 0.f;	// When no profile is provide, we use Burley. STRATA_TODO we should clarify as burley is also an option of the profile.
+	bool bHasSSS = bMFPPlugged || bHasSSSProfile;
 
 	// Derive roughness and haziness from profile data
 	if (bSSSProfilePlugged)
@@ -848,13 +848,13 @@ FStrataData GetStrataSlabBSDF(
 		float AverageToRoughness1 = 0;
 		float LobeMix = 0;
 		float MaterialRoughnessToAverage = 0;
-		GetSubsurfaceProfileDualSpecular(SubsurfaceProfileUInt, SSSDMFPScale, AverageToRoughness0, AverageToRoughness1, LobeMix, MaterialRoughnessToAverage);
+		GetSubsurfaceProfileDualSpecular(SubsurfaceProfileUInt, SSSMFPScale, AverageToRoughness0, AverageToRoughness1, LobeMix, MaterialRoughnessToAverage);
 
 		float AverageRoughness = saturate(Roughness * MaterialRoughnessToAverage); // This is the roughness originallys stored into the gbuffer
 		float Lobe0Roughness   = saturate(AverageRoughness * AverageToRoughness0);
 		float Lobe1Roughness   = saturate(AverageRoughness * AverageToRoughness1);
 
-		GetSubsurfaceProfileRoughnesses(Lobe0Roughness, Lobe1Roughness, SSSDMFPScale);
+		GetSubsurfaceProfileRoughnesses(Lobe0Roughness, Lobe1Roughness, SSSMFPScale);
 
 		// Convertion from roughness0/1 to roughness+haziness
 		// * This encoding is based on how roughness0 + haziness is translated to roughness1, see ComputeHazyLobeRoughness
@@ -862,7 +862,7 @@ FStrataData GetStrataSlabBSDF(
 		Haziness = (Lobe1Roughness - Lobe0Roughness) / (1.f - min(Lobe0Roughness, 0.99f));
 		Roughness = Lobe0Roughness;
 
-		SSSDMFP = GetSubsurfaceProfileDMFPInCm(SubsurfaceProfileUInt).xyz * SSSDMFPScale;
+		SSSMFP = GetSubsurfaceProfileMFPInCm(SubsurfaceProfileUInt).xyz * SSSMFPScale;
 	}
 	else
 	{
@@ -875,13 +875,13 @@ FStrataData GetStrataSlabBSDF(
 	// If MFP is smaller than the pixel's footprint, remove SSS attributes
 	// The scaling factor 'SSSFootprintScale' is based on empiric tests
 	const float SSSFootprintScale = 3.f;
-	if (all(SSSDMFP * SSSFootprintScale < InPixelFootprint.PixelRadiusInWorldSpace))
+	if (all(SSSMFP * SSSFootprintScale < InPixelFootprint.PixelRadiusInWorldSpace))
 	{
 		bHasSSSProfile = false;
 		bHasSSS = false;
 
-		// DMFPPlugged is not cleared as it is used for thin surface and thin volumetric
-		// bDMFPPlugged = false;
+		// MFPPlugged is not cleared as it is used for thin surface and thin volumetric
+		// bMFPPlugged = false;
 	}
 
 	const uint bHasHaziness = Haziness > 0.0f ? 1 : 0;
@@ -909,7 +909,7 @@ FStrataData GetStrataSlabBSDF(
 	BSDF_SETHASSSS			(StrataData.InlinedBSDF,		bHasSSS ? 1 : 0);
 	BSDF_SETHASSSSPROFILE	(StrataData.InlinedBSDF,		bHasSSSProfile ? 1 : 0);
 	BSDF_SETISTHIN			(StrataData.InlinedBSDF,		bIsThin ? 1 : 0);
-	BSDF_SETHASDMFP			(StrataData.InlinedBSDF,		bDMFPPlugged ? 1 : 0);
+	BSDF_SETHASMFP			(StrataData.InlinedBSDF,		bMFPPlugged ? 1 : 0);
 	BSDF_SETHASHAZINESS		(StrataData.InlinedBSDF,		bHasHaziness);
 	BSDF_SETHASTHINFILM		(StrataData.InlinedBSDF,		bHasThinFilm);
 	BSDF_SETTHICKNESSCM		(StrataData.InlinedBSDF,		Thickness);
@@ -942,19 +942,19 @@ FStrataData GetStrataSlabBSDF(
 	if (bHasSSSProfile)
 	{
 		SLAB_SSSPROFILEID	(StrataData.InlinedBSDF)		= SSSProfileID;
-		SLAB_SSSPROFILERADIUSSCALE(StrataData.InlinedBSDF)	= SSSDMFPScale;
+		SLAB_SSSPROFILERADIUSSCALE(StrataData.InlinedBSDF)	= SSSMFPScale;
 	}
 	else
 	{
-		SLAB_SSSDMFP	(StrataData.InlinedBSDF)			= SSSDMFP * SSSDMFPScale;
+		SLAB_SSSMFP	(StrataData.InlinedBSDF)			= SSSMFP * SSSMFPScale;
 	}
 
 	SLAB_FUZZ_AMOUNT		(StrataData.InlinedBSDF)		= FuzzAmount;
 	SLAB_FUZZ_COLOR			(StrataData.InlinedBSDF)		= FuzzColor;
 
 #if STRATA_INLINE_SHADING
-	// At this stage we do not know if this strata BSDF is going to be using simple volume lighting so always store the DMFP on the side to not override SSSPROFILE data.
-	StrataData.InlinedBSDF.TmpDMFP = SSSDMFP * SSSDMFPScale;
+	// At this stage we do not know if this strata BSDF is going to be using simple volume lighting so always store the MFP on the side to not override SSSPROFILE data.
+	StrataData.InlinedBSDF.TmpMFP = SSSMFP * SSSMFPScale;
 
 	StrataData.InlinedBSDF.Coverage = 1.0f;
 #endif
@@ -1444,8 +1444,8 @@ FStrataData StrataHorizontalMixingParameterBlending(FStrataData Background, FStr
 	BSDF_SETHASANISOTROPY(ResultBSDF,	BSDF_GETHASANISOTROPY(ForegroundBSDF)			| BSDF_GETHASANISOTROPY(BackgroundBSDF));
 	BSDF_SETHASSSS(ResultBSDF,			BSDF_GETHASSSS(ForegroundBSDF)					| BSDF_GETHASSSS(BackgroundBSDF));
 	BSDF_SETHASSSSPROFILE(ResultBSDF,	BSDF_GETHASSSSPROFILE(ForegroundBSDF)			| BSDF_GETHASSSSPROFILE(BackgroundBSDF));	// Problem: cannot blend SSS profiles...
-	BSDF_SETISTHIN(ResultBSDF,			BSDF_GETISTHIN(ForegroundBSDF)					| BSDF_GETISTHIN(BackgroundBSDF));			// STRATA_TODO fix the MPF/DMFP situation for that to really work
-	BSDF_SETHASDMFP(ResultBSDF,			BSDF_GETHASDMFP(ForegroundBSDF)					| BSDF_GETHASDMFP(BackgroundBSDF));
+	BSDF_SETISTHIN(ResultBSDF,			BSDF_GETISTHIN(ForegroundBSDF)					| BSDF_GETISTHIN(BackgroundBSDF));
+	BSDF_SETHASMFP(ResultBSDF,			BSDF_GETHASMFP(ForegroundBSDF)					| BSDF_GETHASMFP(BackgroundBSDF));
 	BSDF_SETHASHAZINESS(ResultBSDF,		BSDF_GETHASHAZINESS(ForegroundBSDF)				| BSDF_GETHASHAZINESS(BackgroundBSDF));
 	BSDF_SETHASF90(ResultBSDF,			BSDF_GETHASF90(ForegroundBSDF)					| BSDF_GETHASF90(BackgroundBSDF));
 	BSDF_SETHASTHINFILM(ResultBSDF,		BSDF_GETHASTHINFILM(ForegroundBSDF)				| BSDF_GETHASTHINFILM(BackgroundBSDF));
@@ -1487,10 +1487,10 @@ FStrataData StrataHorizontalMixingParameterBlending(FStrataData Background, FStr
 	}
 	else
 	{
-		SLAB_SSSDMFP(ResultBSDF) =		lerp(SLAB_SSSDMFP(BackgroundBSDF),				SLAB_SSSDMFP(ForegroundBSDF),				ForegroundOtherMixFactor);
+		SLAB_SSSMFP(ResultBSDF) =		lerp(SLAB_SSSMFP(BackgroundBSDF),				SLAB_SSSMFP(ForegroundBSDF),				ForegroundOtherMixFactor);
 
 		// We still try to disable SSS if possible: for performance and also to preserve the surface color (especially specular reflection color).
-		if (all(SLAB_SSSDMFP(ResultBSDF) == 0.0f))
+		if (all(SLAB_SSSMFP(ResultBSDF) == 0.0f))
 		{
 			BSDF_SETHASSSS(ResultBSDF, 0);
 		}
@@ -1498,8 +1498,8 @@ FStrataData StrataHorizontalMixingParameterBlending(FStrataData Background, FStr
 
 
 #if STRATA_INLINE_SHADING
-	// At this stage we do not know if this strata BSDF is going to be using simple volume lighting so always store the DMFP on the side to not override SSSPROFILE data.
-	ResultBSDF.TmpDMFP =				lerp(BackgroundBSDF.TmpDMFP,					ForegroundBSDF.TmpDMFP,						ForegroundOtherMixFactor);
+	// At this stage we do not know if this strata BSDF is going to be using simple volume lighting so always store the MFP on the side to not override SSSPROFILE data.
+	ResultBSDF.TmpMFP =				lerp(BackgroundBSDF.TmpMFP,					ForegroundBSDF.TmpMFP,						ForegroundOtherMixFactor);
 
 	// We blend the parameter above, and we also need to lerp the Coverage to match the total energy.
 	ResultBSDF.Coverage =				lerp(BackgroundBSDF.Coverage,					ForegroundBSDF.Coverage,					HorizontalMixValue);
@@ -1540,8 +1540,8 @@ FStrataData StrataAddParameterBlending(FStrataData A, FStrataData B, float AMixF
 	BSDF_SETHASANISOTROPY(ResultBSDF,	BSDF_GETHASANISOTROPY(BBSDF)					| BSDF_GETHASANISOTROPY(ABSDF));
 	BSDF_SETHASSSS(ResultBSDF,			BSDF_GETHASSSS(BBSDF)							| BSDF_GETHASSSS(ABSDF));
 	BSDF_SETHASSSSPROFILE(ResultBSDF,	BSDF_GETHASSSSPROFILE(BBSDF)					| BSDF_GETHASSSSPROFILE(ABSDF));	// Problem: cannot blend SSS profiles...
-	BSDF_SETISTHIN(ResultBSDF,			BSDF_GETISTHIN(BBSDF)							| BSDF_GETISTHIN(ABSDF));			// STRATA_TODO fix the MPF/DMFP situation for that to really work
-	BSDF_SETHASDMFP(ResultBSDF,			BSDF_GETHASDMFP(BBSDF)							| BSDF_GETHASDMFP(ABSDF));
+	BSDF_SETISTHIN(ResultBSDF,			BSDF_GETISTHIN(BBSDF)							| BSDF_GETISTHIN(ABSDF));			// STRATA_TODO fix the MPF/MFP situation for that to really work
+	BSDF_SETHASMFP(ResultBSDF,			BSDF_GETHASMFP(BBSDF)							| BSDF_GETHASMFP(ABSDF));
 	BSDF_SETHASHAZINESS(ResultBSDF,		BSDF_GETHASHAZINESS(BBSDF)						| BSDF_GETHASHAZINESS(ABSDF));
 	BSDF_SETHASF90(ResultBSDF,			BSDF_GETHASF90(BBSDF)							| BSDF_GETHASF90(ABSDF));
 	BSDF_SETHASTHINFILM(ResultBSDF,		BSDF_GETHASTHINFILM(BBSDF)						| BSDF_GETHASTHINFILM(ABSDF));
@@ -1579,13 +1579,13 @@ FStrataData StrataAddParameterBlending(FStrataData A, FStrataData B, float AMixF
 	}
 	else
 	{
-		SLAB_SSSDMFP(ResultBSDF) =		lerp(SLAB_SSSDMFP(BBSDF),						SLAB_SSSDMFP(ABSDF),						AMixFactor);
+		SLAB_SSSMFP(ResultBSDF) =		lerp(SLAB_SSSMFP(BBSDF),						SLAB_SSSMFP(ABSDF),						AMixFactor);
 	}
 
 
 #if STRATA_INLINE_SHADING
-	// At this stage we do not know if this strata BSDF is going to be using simple volume lighting so always store the DMFP on the side to not override SSSPROFILE data.
-	ResultBSDF.TmpDMFP =				lerp(BBSDF.TmpDMFP,								ABSDF.TmpDMFP,								AMixFactor);
+	// At this stage we do not know if this strata BSDF is going to be using simple volume lighting so always store the MFP on the side to not override SSSPROFILE data.
+	ResultBSDF.TmpMFP =				lerp(BBSDF.TmpMFP,								ABSDF.TmpMFP,								AMixFactor);
 
 	// When it comes to adding the visual contribution, we chose to add the coverage without clamping as currently done in the StrataAdd function.
 	// Operations using Coverage are currently making sure they are in a safe range.
@@ -1631,8 +1631,8 @@ FStrataBSDF StrataVerticalLayeringParameterBlending(in FStrataBSDF TopBSDF, in F
 	BSDF_SETHASANISOTROPY(ResultBSDF,	0);																					// Anisotropy is disabled because it changes a lot of things
 	BSDF_SETHASSSS(ResultBSDF,			BSDF_GETHASSSS(BaseBSDF));															// We only keep SSS if the bottom layer has it
 	BSDF_SETHASSSSPROFILE(ResultBSDF,	0);																					// Problem: cannot blend SSS profiles so we simply disable SSSProfile
-	BSDF_SETISTHIN(ResultBSDF,			BSDF_GETISTHIN(BaseBSDF));															// STRATA_TODO fix the MPF/DMFP situation for that to really work
-	BSDF_SETHASDMFP(ResultBSDF,			BSDF_GETHASDMFP(BaseBSDF)						| BSDF_GETHASDMFP(TopBSDF));		// We keep DMFP of both to make sure bottom layer on translucent can laos be effective
+	BSDF_SETISTHIN(ResultBSDF,			BSDF_GETISTHIN(BaseBSDF));															// STRATA_TODO fix the MPF/MFP situation for that to really work
+	BSDF_SETHASMFP(ResultBSDF,			BSDF_GETHASMFP(BaseBSDF)						| BSDF_GETHASMFP(TopBSDF));		// We keep MFP of both to make sure bottom layer on translucent can laos be effective
 	BSDF_SETHASHAZINESS(ResultBSDF,		BSDF_GETHASHAZINESS(TopBSDF));														// We only keep haziness if the top layer has it
 	BSDF_SETHASF90(ResultBSDF,			BSDF_GETHASF90(BaseBSDF)						| BSDF_GETHASF90(TopBSDF));			// F90: keep union of both even though it will be hard to get a match
 	BSDF_SETHASTHINFILM(ResultBSDF,		BSDF_GETHASTHINFILM(TopBSDF));														// We only keep thinfilme if the top layer has it
@@ -1643,7 +1643,7 @@ FStrataBSDF StrataVerticalLayeringParameterBlending(in FStrataBSDF TopBSDF, in F
 	EnableSlabBSDFSimpleVolumetric(TopBSDF);
 	EnableSlabBSDFSimpleVolumetric(BaseBSDF);
 	const float3 UnusedDiffuseColor		= 0;	// Unused when only transmittance is needed
-	FParticipatingMedia PM				= StrataSlabCreateParticipatingMedia(UnusedDiffuseColor, SLAB_SSSDMFP(TopBSDF));
+	FParticipatingMedia PM				= StrataSlabCreateParticipatingMedia(UnusedDiffuseColor, SLAB_SSSMFP(TopBSDF));
 	const float3 SlabTransmittance		= IsotropicMediumSlabTransmittance(PM, STRATA_SIMPLEVOLUME_THICKNESS_M, NoV);
 
 	// STRATA_TODO account for energy conservation here.
@@ -1677,13 +1677,13 @@ FStrataBSDF StrataVerticalLayeringParameterBlending(in FStrataBSDF TopBSDF, in F
 	// Only keep thin film for the top layer.
 	SLAB_THINFILM_THICKNESS(ResultBSDF)= (SafeTopCoverage >= 1.0f) ? SLAB_THINFILM_THICKNESS(TopBSDF) : 0.0f;
 
-	// Keep the DMFP form the bottom layer only (top layer simple volume is backed in parameters).
-	SLAB_SSSDMFP(ResultBSDF) =			lerp(SLAB_SSSDMFP(TopBSDF),						SLAB_SSSDMFP(BaseBSDF),		TopToBottomLerp); // SLAB_SSSDMFP(BaseBSDF);
+	// Keep the MFP form the bottom layer only (top layer simple volume is backed in parameters).
+	SLAB_SSSMFP(ResultBSDF) =			lerp(SLAB_SSSMFP(TopBSDF),						SLAB_SSSMFP(BaseBSDF),		TopToBottomLerp); // SLAB_SSSMFP(BaseBSDF);
 
 #if STRATA_INLINE_SHADING
 #if STRATA_OPAQUE_MATERIAL
 	// Only the bottom layer SSS / Thin model is used so we only keep the bottom layer data there. Top layer is converted to throughput.
-	ResultBSDF.TmpDMFP =				BaseBSDF.TmpDMFP;
+	ResultBSDF.TmpMFP =				BaseBSDF.TmpMFP;
 	ResultBSDF.Coverage =				Info.Coverage;
 #else
 
@@ -1692,8 +1692,8 @@ FStrataBSDF StrataVerticalLayeringParameterBlending(in FStrataBSDF TopBSDF, in F
 	FVerticalLayeringInfo Info2 = GetVerticalLayeringInfo(TopBSDF.Coverage / MaxCoverage, BaseBSDF.Coverage / MaxCoverage);
 #if 0
 	// Using transmittance as a way to interpolate should be better but there are issues with large values
-	const float3 TopExtinction			= 1.0f / max(0.0, TopBSDF.TmpDMFP);
-	const float3 BaseExtinction			= 1.0f / max(0.0, BaseBSDF.TmpDMFP);
+	const float3 TopExtinction			= 1.0f / max(0.0, TopBSDF.TmpMFP);
+	const float3 BaseExtinction			= 1.0f / max(0.0, BaseBSDF.TmpMFP);
 	const float3 TopTransmitance		= exp(-TopExtinction);	// extinction = optical depth, we assume here d = 1 meter
 	const float3 BaseTransmitance		= exp(-BaseExtinction);
 	const float3 NewTransmittance		= Info2.TransmittanceOnlyBottom * BaseTransmitance + Info2.TransmittanceOnlyTop * TopTransmitance + Info2.TransmittanceTopAndBottom * TopTransmitance * BaseTransmitance;
@@ -1701,11 +1701,11 @@ FStrataBSDF StrataVerticalLayeringParameterBlending(in FStrataBSDF TopBSDF, in F
 	const float3 NewMFP					= 1.0f / NewExtinction;
 #else
 	// MFP looks more prceptually linear even though it is not perfect
-	const float3 NewMFP = Info2.TransmittanceOnlyBottom * BaseBSDF.TmpDMFP + Info2.TransmittanceOnlyTop * TopBSDF.TmpDMFP + Info2.TransmittanceTopAndBottom * min(TopBSDF.TmpDMFP, BaseBSDF.TmpDMFP);
+	const float3 NewMFP = Info2.TransmittanceOnlyBottom * BaseBSDF.TmpMFP + Info2.TransmittanceOnlyTop * TopBSDF.TmpMFP + Info2.TransmittanceTopAndBottom * min(TopBSDF.TmpMFP, BaseBSDF.TmpMFP);
 #endif
 
-	ResultBSDF.TmpDMFP					= NewMFP;
-	SLAB_SSSDMFP(ResultBSDF)			= NewMFP;
+	ResultBSDF.TmpMFP					= NewMFP;
+	SLAB_SSSMFP(ResultBSDF)			= NewMFP;
 
 	ResultBSDF.Coverage				= MaxCoverage;
 #endif
@@ -1764,10 +1764,10 @@ void SanitizeStrataSlab(inout FStrataBSDF BSDF)
 	SLAB_HAZINESS(BSDF) = saturate(SLAB_HAZINESS(BSDF));
 	SLAB_THINFILM_THICKNESS(BSDF) = saturate(SLAB_THINFILM_THICKNESS(BSDF));
 #if STRATA_INLINE_SHADING
-	BSDF.TmpDMFP = clamp(BSDF.TmpDMFP, 0.0f, Max111110BitsFloat3);
+	BSDF.TmpMFP = clamp(BSDF.TmpMFP, 0.0f, Max111110BitsFloat3);
 #endif
 
-	// BSDF feature tests must be executed for profile path to not interact with DMFP, especially for forward rendered simple volumetrics.
+	// BSDF feature tests must be executed for profile path to not interact with MFP, especially for forward rendered simple volumetrics.
 	if (BSDF_GETHASSSS(BSDF))
 	{
 		if (BSDF_GETHASSSSPROFILE(BSDF))
@@ -1776,7 +1776,7 @@ void SanitizeStrataSlab(inout FStrataBSDF BSDF)
 		}
 		else
 		{
-			SLAB_SSSDMFP(BSDF) = clamp(SLAB_SSSDMFP(BSDF), 0.0f, Max111110BitsFloat3);
+			SLAB_SSSMFP(BSDF) = clamp(SLAB_SSSMFP(BSDF), 0.0f, Max111110BitsFloat3);
 		}
 	}
 
@@ -1864,7 +1864,7 @@ void StrataDitherBSDF(inout FStrataBSDF BSDF, in float Dither)
 ///////////////////////////////////////////////////////////////////////////////
 // Other tools
 
-// Those participating media funtion are use by the Slab node. In this case, the DMFP is considered as a simple MFP because DMFP is not directly inversible to medium properties.
+// Those participating media funtion are use by the Slab node. In this case, the MFP is considered as a simple MFP because MFP is not directly inversible to medium properties.
 
 FParticipatingMedia StrataSlabCreateParticipatingMedia(float3 DiffuseColor, float3 MeanFreePathCentimeters)
 {
@@ -1878,10 +1878,10 @@ float StrataSlabDiffuseToVolumeBlend(FParticipatingMedia PM)
 	// We start blending from diffuse toward a slab of participating media when the mean free path reaches 4 centimeters.
 	const float MaxMeanFreePathMeters = max(PM.MeanFreePath.x, max(PM.MeanFreePath.y, PM.MeanFreePath.z));
 	// Start blending volumetric after mean free path of 4 centimeters (also needed because threshold when recovering extinction from transmittance cause a small minimum value)
-	const float StartVolumeBlendMFP = 0.04f;
+	const float StartVolumeBlenMFP = 0.04f;
 	// And when the mean free path reaches 1/3.0f = 0.33 meter. That is extinction=1/0.33=3, transmittance at 1m (bottom of the slab) = exp(-3) = 0.05) we fully use the volumetric model.
-	const float Blend = saturate(max(0.0, MaxMeanFreePathMeters - StartVolumeBlendMFP) * 3.0f);
-	//const float Blend = max(0.0, MaxMeanFreePathMeters - StartVolumeBlendMFP) > 0.0 ? 1.0 : 0.0;
+	const float Blend = saturate(max(0.0, MaxMeanFreePathMeters - StartVolumeBlenMFP) * 3.0f);
+	//const float Blend = max(0.0, MaxMeanFreePathMeters - StartVolumeBlenMFP) > 0.0 ? 1.0 : 0.0;
 
 	// Using a non linear blend helps smoothing out transitions visually.
 	return Blend * Blend;
@@ -1891,12 +1891,12 @@ void EnableSlabMFPComputationSpace(inout FStrataBSDF BSDF)
 {
 #if STRATA_INLINE_SHADING
 	// Now we know that this is a simple volume so we can freely set the mean free path
-	SLAB_SSSDMFP(BSDF) = BSDF.TmpDMFP;
+	SLAB_SSSMFP(BSDF) = BSDF.TmpMFP;
 #endif
 
 	// Rescale the BSDF mean free path so that we can drop the thickness parameter and run volumetric material math on a unit slab of medium
 	// For instance, if thickness is larger => we need to reduce the MFP for the medium to appear visually thicker.
-	SLAB_SSSDMFP(BSDF) = clamp(SLAB_SSSDMFP(BSDF) * (STRATA_SIMPLEVOLUME_THICKNESS_CM / max(STRATA_EPSILON, BSDF_GETTHICKNESSCM(BSDF))), 0.0f, Max111110BitsFloat3);
+	SLAB_SSSMFP(BSDF) = clamp(SLAB_SSSMFP(BSDF) * (STRATA_SIMPLEVOLUME_THICKNESS_CM / max(STRATA_EPSILON, BSDF_GETTHICKNESSCM(BSDF))), 0.0f, Max111110BitsFloat3);
 }
 
 void EnableSlabBSDFSimpleVolumetric(inout FStrataBSDF BSDF)
@@ -1991,7 +1991,7 @@ float3 StrataGetBSDFSubSurfaceColor(in FStrataBSDF BSDF)
 {
 	switch (BSDF_GETTYPE(BSDF))
 	{
-	case STRATA_BSDF_TYPE_SLAB: 			return BSDF_GETHASSSS(BSDF) ? SLAB_DIFFUSEALBEDO(BSDF) : 0; // A more accurate version would be IsotropicMediumSlabPunctualDirectionalAlbedo(StrataSlabCreateParticipatingMedia(SLAB_DIFFUSEALBEDO(BSDF), SLAB_SSSDMFP(BSDFContext.BSDF)));
+	case STRATA_BSDF_TYPE_SLAB: 			return BSDF_GETHASSSS(BSDF) ? SLAB_DIFFUSEALBEDO(BSDF) : 0; // A more accurate version would be IsotropicMediumSlabPunctualDirectionalAlbedo(StrataSlabCreateParticipatingMedia(SLAB_DIFFUSEALBEDO(BSDF), SLAB_SSSMFP(BSDFContext.BSDF)));
 	case STRATA_BSDF_TYPE_HAIR: 			return 0;
 	case STRATA_BSDF_TYPE_SINGLELAYERWATER: return 0;
 	}
@@ -2100,7 +2100,7 @@ float StrataGetLegacyShadingModels(in FStrataBSDF BSDF)
 // Note: When SSS data uses DFMP data, the conversion to opacity is ad-hoc/hacky
 float StrataGetOpacityFromSubSurface(FStrataSubsurfaceData In)
 {
-	return In.bIsValid ? (In.bIsProfile ? In.ProfileRadiusScale : saturate(1.f - max3(In.DMFP.x, In.DMFP.y, In.DMFP.z))) : 1.f;
+	return In.bIsValid ? (In.bIsProfile ? In.ProfileRadiusScale : saturate(1.f - max3(In.MFP.x, In.MFP.y, In.MFP.z))) : 1.f;
 }
 
 ///////////////////////////////////////////////////////////////////////////////
@@ -2510,8 +2510,8 @@ void StrataConvertFromDBuffer(
 		In.Roughness,		// Roughness
 		0.f,				// Anisotropy
 		0.f,				// SSSProfileID
-		0.f,				// SSSDMFP
-		0.f,				// SSSDMFPScale
+		0.f,				// SSSMFP
+		0.f,				// SSSMFPScale
 		0.f,				// Emissive
 		0.f,				// Haziness
 		0.f,				// FuzzAmount
@@ -2905,7 +2905,7 @@ void PackStrataOut(
 
 								// STRATA_TODO This should be accumulated according to the BSDF weight when we run this in the base pass.
 								// We cannot do that now because BSDFWeight contains weight*coverage*trhougput all mixed together.
-								SSSData.DMFP = SLAB_SSSDMFP(BSDF);
+								SSSData.MFP = SLAB_SSSMFP(BSDF);
 							}
 
 							if (BSDF_GETHASSSSPROFILE(BSDF) && !bIsSimpleVolume)
@@ -2916,7 +2916,7 @@ void PackStrataOut(
 							else
 							{
 								// Path used for bottom most layer SSS, simple volume and two sided lighting.
-								STRATA_STORE_UINT1(PackR11G11B10F(SLAB_SSSDMFP(BSDF)));
+								STRATA_STORE_UINT1(PackR11G11B10F(SLAB_SSSMFP(BSDF)));
 							}
 						}
 
@@ -3190,7 +3190,7 @@ FStrataBSDF UnpackStrataBSDFIn(FStrataMaterialContainer StrataBuffer, inout FStr
 			}
 			else
 			{
-				StrataLoad_R11G11B10F(StrataBuffer, StrataAddressing, SLAB_SSSDMFP(OutBSDF));
+				StrataLoad_R11G11B10F(StrataBuffer, StrataAddressing, SLAB_SSSMFP(OutBSDF));
 			}
 
 			// When loading Strata BSDF which have SSS enable, BaseColor and (optionally) Specular value can be overriden 
diff --git a/Engine/Shaders/Private/Strata/StrataEvaluation.ush b/Engine/Shaders/Private/Strata/StrataEvaluation.ush
--- a/Engine/Shaders/Private/Strata/StrataEvaluation.ush
+++ b/Engine/Shaders/Private/Strata/StrataEvaluation.ush
@@ -553,7 +553,7 @@ FStrataEvaluateResult StrataEvaluateBSDFCommon(FStrataBSDFContext BSDFContext, F
 				else
 				{
 					const float3 SubsurfaceAlebdo = SLAB_DIFFUSEALBEDO(BSDFContext.BSDF);
-					const float3 MeanFreePathInCm = SLAB_SSSDMFP(BSDFContext.BSDF);
+					const float3 MeanFreePathInCm = SLAB_SSSMFP(BSDFContext.BSDF);
 					TransmissionThroughput	= GetBurleyTransmissionProfile(SubsurfaceAlebdo, MeanFreePathInCm, ThicknessInCm);
 					PhaseFunctionAnisotropy = 0.93f; // STRATA_TODO: expose phase function eccentricity. For now, forward scattering matching default SSS profile value
 					OneOverIOR				= 1.f / DielectricF0ToIor(F0.y);
@@ -578,7 +578,7 @@ FStrataEvaluateResult StrataEvaluateBSDFCommon(FStrataBSDFContext BSDFContext, F
 			BRANCH
 			if (BSDF_GETISTHIN(BSDFContext.BSDF))
 			{
-				FParticipatingMedia PM = StrataSlabCreateParticipatingMedia(DiffuseColor, SLAB_SSSDMFP(BSDFContext.BSDF));
+				FParticipatingMedia PM = StrataSlabCreateParticipatingMedia(DiffuseColor, SLAB_SSSMFP(BSDFContext.BSDF));
 				const float BackFaceNoL = -BSDFContext.Context.NoL;
 				const float BackFaceNoLSat = saturate(BackFaceNoL);
 
@@ -611,7 +611,7 @@ FStrataEvaluateResult StrataEvaluateBSDFCommon(FStrataBSDFContext BSDFContext, F
 			BRANCH
 			if (BSDF_GETISSIMPLEVOLUME(BSDFContext.BSDF))
 			{
-				FParticipatingMedia PM = StrataSlabCreateParticipatingMedia(DiffuseColor, SLAB_SSSDMFP(BSDFContext.BSDF));
+				FParticipatingMedia PM = StrataSlabCreateParticipatingMedia(DiffuseColor, SLAB_SSSMFP(BSDFContext.BSDF));
 				const float DiffuseToVolumeBlend = StrataSlabDiffuseToVolumeBlend(PM);
 				const float3 SlabDirectionalAlbedo = IsotropicMediumSlabPunctualDirectionalAlbedo(PM);
 				const float3 SlabTransmittance = IsotropicMediumSlabTransmittance(PM, STRATA_SIMPLEVOLUME_THICKNESS_M, BSDFContext.Context.NoV);
@@ -1007,7 +1007,7 @@ FStrataEnvLightResult StrataEvaluateForEnvLight(FStrataBSDFContext BSDFContext,
 			BRANCH
 			if (BSDF_GETISSIMPLEVOLUME(BSDFContext.BSDF))
 			{
-				FParticipatingMedia PM = StrataSlabCreateParticipatingMedia(DiffuseColor, SLAB_SSSDMFP(BSDFContext.BSDF));
+				FParticipatingMedia PM = StrataSlabCreateParticipatingMedia(DiffuseColor, SLAB_SSSMFP(BSDFContext.BSDF));
 				const float DiffuseToVolumeBlend = StrataSlabDiffuseToVolumeBlend(PM);
 
 				// The environment response of has been measured for an hemisphere having a uniform radiance of 1. 
@@ -1046,7 +1046,7 @@ FStrataEnvLightResult StrataEvaluateForEnvLight(FStrataBSDFContext BSDFContext,
 			BRANCH
 			if (BSDF_GETISTHIN(BSDFContext.BSDF))
 			{
-				FParticipatingMedia PM = StrataSlabCreateParticipatingMedia(DiffuseColor, SLAB_SSSDMFP(BSDFContext.BSDF));
+				FParticipatingMedia PM = StrataSlabCreateParticipatingMedia(DiffuseColor, SLAB_SSSMFP(BSDFContext.BSDF));
 
 				// We consider single scattering and perpendicular to the surface light transmittance.
 				// Using BackFaceNoL instead of 1 for TransmittanceNoL would be more correct but it gives some high frequency details looking weird without multiple scattering.
diff --git a/Engine/Shaders/Private/Strata/StrataLegacyConversion.ush b/Engine/Shaders/Private/Strata/StrataLegacyConversion.ush
--- a/Engine/Shaders/Private/Strata/StrataLegacyConversion.ush
+++ b/Engine/Shaders/Private/Strata/StrataLegacyConversion.ush
@@ -18,8 +18,8 @@ struct FStrataLegacyParameters
 	float  Roughness;
 	float  Anisotropy;
 	float  SSSProfileID;
-	float3 SSSDMFP;
-	float  SSSDMFPScale;
+	float3 SSSMFP;
+	float  SSSMFPScale;
 	float3 Emissive;
 	float  FuzzAmount;
 	float3 FuzzColor;
@@ -49,8 +49,8 @@ FStrataLegacyParameters InitStrataLegacyParameters(FStrataPixelFootprint InPixel
 	Out.Roughness = 0.5f;
 	Out.Anisotropy = 0.f;
 	Out.SSSProfileID = 0.f;
-	Out.SSSDMFP = Zeros;
-	Out.SSSDMFPScale = 1.f;
+	Out.SSSMFP = Zeros;
+	Out.SSSMFPScale = 1.f;
 	Out.Emissive = Zeros;
 	Out.FuzzAmount = 0.f;
 	Out.FuzzColor = Zeros;
@@ -86,8 +86,8 @@ FStrataData CreateLegacySlab(
 			In.Roughness,		// Roughness
 			In.Anisotropy,		// Anisotropy
 			In.SSSProfileID,	// SSSProfileID
-			In.SSSDMFP,			// SSSDMFP
-			In.SSSDMFPScale,	// SSSDMFPScale
+			In.SSSMFP,			// SSSMFP
+			In.SSSMFPScale,	// SSSMFPScale
 			In.Emissive,		// Emissive
 			0.f,				// Haziness
 			In.FuzzAmount,		// FuzzAmount
@@ -225,8 +225,8 @@ FStrataData StrataConvertLegacyMaterialDynamic(
 		Top.Metallic = Metallic;
 		Top.Roughness = Roughness;
 		Top.Anisotropy = Anisotropy;
-		Top.SSSDMFP = SSSMFP;
-		Top.SSSDMFPScale = Opacity;
+		Top.SSSMFP = SSSMFP;
+		Top.SSSMFPScale = Opacity;
 		Top.Emissive = Emissive;
 
 		// No need to take into account opacity weighting for translucent material, as the coverage is assumed to be 1
@@ -241,8 +241,8 @@ FStrataData StrataConvertLegacyMaterialDynamic(
 		// All this will be fixed when this material will be converted to LUT.
 		const float Transmittance = saturate(1.f - Opacity);
 		const float Thickness = DefaultThickness;
-		Top.SSSDMFP = 0;// TransmittanceToMeanFreePath(Transmittance, Thickness* CENTIMETER_TO_METER)* METER_TO_CENTIMETER;
-		Top.SSSDMFPScale = 1.0f;
+		Top.SSSMFP = 0;// TransmittanceToMeanFreePath(Transmittance, Thickness* CENTIMETER_TO_METER)* METER_TO_CENTIMETER;
+		Top.SSSMFPScale = 1.0f;
 		Top.Thickness = Thickness;
 
 		FinalWeight = Opacity;
@@ -256,7 +256,7 @@ FStrataData StrataConvertLegacyMaterialDynamic(
 		Top.Roughness = Roughness;
 		Top.Anisotropy = Anisotropy;
 		Top.SSSProfileID = SubSurfaceProfileId;
-		Top.SSSDMFPScale = Opacity;
+		Top.SSSMFPScale = Opacity;
 		Top.Emissive = Emissive;
 
 		// No need to take into account opacity weighting for translucent material, as the coverage is assumed to be 1
@@ -276,8 +276,8 @@ FStrataData StrataConvertLegacyMaterialDynamic(
 		Bot.Roughness = Roughness;
 		Bot.Anisotropy = Anisotropy;
 		Bot.SSSProfileID = 0.f;
-		Bot.SSSDMFP = Zeros;
-		Bot.SSSDMFPScale = 0.f;
+		Bot.SSSMFP = Zeros;
+		Bot.SSSMFPScale = 0.f;
 		Bot.Emissive = Zeros;
 		Bot.SharedLocalBasisIndex = ClearCoatBottomNormal_SharedLocalBasisIndex;
 		Bot.bIsTop = 0;
@@ -302,8 +302,8 @@ FStrataData StrataConvertLegacyMaterialDynamic(
 		Top.Roughness = ClearCoatRoughness;
 		Top.Anisotropy = Anisotropy;
 		Top.SSSProfileID = 0.f;
-		Top.SSSDMFP = MFP;
-		Top.SSSDMFPScale = 1.f;
+		Top.SSSMFP = MFP;
+		Top.SSSMFPScale = 1.f;
 		Top.Emissive = Emissive;
 		Top.Thickness = Thickness;
 		Top.SharedLocalBasisIndex = SharedLocalBasisIndex;
@@ -328,8 +328,8 @@ FStrataData StrataConvertLegacyMaterialDynamic(
 		Top.Roughness = Roughness;
 		Top.Anisotropy = Anisotropy;
 		Top.SSSProfileID = 0.f;
-		Top.SSSDMFP = MFP;
-		Top.SSSDMFPScale = 1.f;
+		Top.SSSMFP = MFP;
+		Top.SSSMFPScale = 1.f;
 		Top.Emissive = Emissive;
 		Top.Thickness = Thickness;
 
@@ -349,8 +349,8 @@ FStrataData StrataConvertLegacyMaterialDynamic(
 		Top.Roughness = Roughness;
 		Top.Anisotropy = Anisotropy;
 		Top.SSSProfileID = 0.f;
-		Top.SSSDMFP = 0.f;
-		Top.SSSDMFPScale = 0.f;
+		Top.SSSMFP = 0.f;
+		Top.SSSMFPScale = 0.f;
 		Top.Emissive = Emissive;
 		Top.FuzzAmount = ClearCoat;
 		Top.FuzzColor = SubSurfaceColor;
@@ -378,8 +378,8 @@ FStrataData StrataConvertLegacyMaterialDynamic(
 		Top.Roughness = Roughness;
 		Top.Anisotropy = 0.f;
 		Top.SSSProfileID = SubSurfaceProfileId;
-		Top.SSSDMFP = MFP;
-		Top.SSSDMFPScale = Opacity;
+		Top.SSSMFP = MFP;
+		Top.SSSMFPScale = Opacity;
 		Top.Emissive = Emissive;
 		Top.SharedLocalBasisIndex = SharedLocalBasisIndex;
 		Top.bIsTop = 1;
@@ -393,8 +393,8 @@ FStrataData StrataConvertLegacyMaterialDynamic(
 		Bot.Roughness = Roughness;
 		Bot.Anisotropy = 0.f;
 		Bot.SSSProfileID = SubSurfaceProfileId;
-		Bot.SSSDMFP = Zeros;
-		Bot.SSSDMFPScale = Opacity;
+		Bot.SSSMFP = Zeros;
+		Bot.SSSMFPScale = Opacity;
 		Bot.Emissive = Emissive;
 		Bot.SharedLocalBasisIndex = ClearCoatBottomNormal_SharedLocalBasisIndex;
 		Bot.bIsTop = 0;
@@ -433,14 +433,14 @@ FStrataData StrataConvertLegacyMaterialDynamic(
 		const float3 BotMFP    = MFP;
 
 		Top.DiffuseAlbedo	= lerp(BotAlbedo, TopAlbedo, Opacity);
-		Top.SSSDMFP			= lerp(BotMFP,    TopMFP,    Opacity);
+		Top.SSSMFP			= lerp(BotMFP,    TopMFP,    Opacity);
 		Top.F0				= lerp(BotF0,     TopF0,     Opacity);
 		Top.F90				= Ones;
 		Top.UseMetalness	= false;
 		Top.Roughness 		= Roughness;
 		Top.Anisotropy 		= 0.f;
 		Top.SSSProfileID 	= 0.f;
-		Top.SSSDMFPScale 	= 1.f;
+		Top.SSSMFPScale 	= 1.f;
 		Top.Emissive 		= Emissive;
 		Top.Thickness 		= Thickness;
 		Top.SharedLocalBasisIndex = SharedLocalBasisIndex;
@@ -457,8 +457,8 @@ FStrataData StrataConvertLegacyMaterialDynamic(
 		Top.Roughness = Roughness;
 		Top.Anisotropy = Anisotropy;
 		Top.SSSProfileID = 0.f;
-		Top.SSSDMFP = Zeros;
-		Top.SSSDMFPScale = 0.f;
+		Top.SSSMFP = Zeros;
+		Top.SSSMFPScale = 0.f;
 		Top.Emissive = Emissive;
 		Top.SharedLocalBasisIndex = SharedLocalBasisIndex;
 
@@ -473,8 +473,8 @@ FStrataData StrataConvertLegacyMaterialDynamic(
 		Bot.Roughness = Roughness;
 		Bot.Anisotropy = 0.f;
 		Bot.SSSProfileID = .0f;
-		Bot.SSSDMFP = MFP;
-		Bot.SSSDMFPScale = 1.f;
+		Bot.SSSMFP = MFP;
+		Bot.SSSMFPScale = 1.f;
 		Bot.Emissive = Zeros;
 		Bot.Thickness = Thickness;
 		Bot.SharedLocalBasisIndex = ClearCoatBottomNormal_SharedLocalBasisIndex;
diff --git a/Engine/Shaders/Private/Strata/StrataTree.ush b/Engine/Shaders/Private/Strata/StrataTree.ush
--- a/Engine/Shaders/Private/Strata/StrataTree.ush
+++ b/Engine/Shaders/Private/Strata/StrataTree.ush
@@ -20,10 +20,10 @@ void PostProcessSlabBeforeLighting(inout FStrataBSDF BSDF)
 	// Two sided lighting can only be evaluated on the bottom layer, together with option SSS. Otherwise, the simple volume will be used.
 	const bool bIsThin = BSDF.bIsBottom && BSDF_GETISTHIN(BSDF);
 
-	// We want simple volumetric only if the DMFP input is pluged in, otherwise we avoid dealing with simple volumetric.
+	// We want simple volumetric only if the MFP input is pluged in, otherwise we avoid dealing with simple volumetric.
 	// We also do not want to use the simple volumetric path if we are dealing with an opaque material and this slab is not at the bottom.
 	// And when opaque, the bottom layer must be opaque and thus will ahve a solid opaque diffuse color.
-	const bool bIsSimpleVolume = !bIsThin && BSDF_GETHASDMFP(BSDF) && (!STRATA_OPAQUE_MATERIAL || (STRATA_OPAQUE_MATERIAL && !BSDF.bIsBottom));
+	const bool bIsSimpleVolume = !bIsThin && BSDF_GETHASMFP(BSDF) && (!STRATA_OPAQUE_MATERIAL || (STRATA_OPAQUE_MATERIAL && !BSDF.bIsBottom));
 
 	if (bIsSimpleVolume)
 	{
diff --git a/Engine/Shaders/Private/Strata/StrataVisualize.usf b/Engine/Shaders/Private/Strata/StrataVisualize.usf
--- a/Engine/Shaders/Private/Strata/StrataVisualize.usf
+++ b/Engine/Shaders/Private/Strata/StrataVisualize.usf
@@ -93,20 +93,20 @@ void DrawPixelFootprint(float3 P, float3 dPdx, float3 dPdy, uint2 PixelCoord)
 	const FStrataSubsurfaceData SSSData = StrataUnpackSSSData(Strata.SSSTexture.Load(uint3(PixelCoord, 0)));
 	if (SSSData.bIsValid)
 	{
-		float3 DMFP = 0;
+		float3 MFP = 0;
 		if (SSSData.bIsProfile)
 		{
-			DMFP = GetSubsurfaceProfileDMFPInCm(SSSData.ProfileId).xyz * SSSData.ProfileRadiusScale;
+			MFP = GetSubsurfaceProfileMFPInCm(SSSData.ProfileId).xyz * SSSData.ProfileRadiusScale;
 		}
 		else
 		{
-			DMFP = SSSData.DMFP;
+			MFP = SSSData.MFP;
 		}
 
 		FStrataPixelFootprint Footprint = StrataGetPixelFootprint(P);
 		AddDrawPixelFootprint(P, dPdx, dPdy, 0.5f, false, ColorRed);
 		AddDrawPixelFootprint(P, dPdx, dPdy, Footprint.PixelRadiusInWorldSpace, true, ColorOrange);
-		AddDrawPixelFootprint(P, dPdx, dPdy, max3(DMFP.x, DMFP.y, DMFP.z), true, ColorCyan);
+		AddDrawPixelFootprint(P, dPdx, dPdy, max3(MFP.x, MFP.y, MFP.z), true, ColorCyan);
 
 		// Normal
 		const float3 N = normalize(cross(normalize(dPdx), normalize(dPdy)));
@@ -280,7 +280,7 @@ void VisualizeMaterialPS(
 #define PRINTCLASSIFISPROF()	PRINTBSDFCHAR(_C_); PRINTBSDFCHAR(_L_); PRINTBSDFCHAR(_A_); PRINTBSDFCHAR(_S_); PRINTBSDFCHAR(_DOT_); PRINTBSDFCHAR(_I_) PRINTBSDFCHAR(_S_) PRINTBSDFCHAR(_P_) PRINTBSDFCHAR(_R_) PRINTBSDFCHAR(_O_) PRINTBSDFCHAR(_F_);	PRINTBSDFCHAR(_SPC_); PRINTBSDFCHAR(SSSData.bIsProfile?_1_:_0_);
 #define PRINTCLASSIFPROID()		PRINTBSDFCHAR(_C_); PRINTBSDFCHAR(_L_); PRINTBSDFCHAR(_A_); PRINTBSDFCHAR(_S_); PRINTBSDFCHAR(_DOT_); PRINTBSDFCHAR(_P_) PRINTBSDFCHAR(_R_) PRINTBSDFCHAR(_O_) PRINTBSDFCHAR(_I_) PRINTBSDFCHAR(_D_) PRINTBSDFCHAR(_SPC_);	PRINTBSDFCHAR(_SPC_); PRINTFLOATNOFRAC(SSSData.ProfileId, 3);
 #define PRINTCLASSIFPROSCL()	PRINTBSDFCHAR(_C_); PRINTBSDFCHAR(_L_); PRINTBSDFCHAR(_A_); PRINTBSDFCHAR(_S_); PRINTBSDFCHAR(_DOT_); PRINTBSDFCHAR(_P_) PRINTBSDFCHAR(_R_) PRINTBSDFCHAR(_O_) PRINTBSDFCHAR(_S_) PRINTBSDFCHAR(_C_) PRINTBSDFCHAR(_L_);	PRINTBSDFCHAR(_SPC_); PRINTFLOAT1(SSSData.ProfileRadiusScale);
-#define PRINTCLASSIFSSDMFP()	PRINTBSDFCHAR(_C_); PRINTBSDFCHAR(_L_); PRINTBSDFCHAR(_A_); PRINTBSDFCHAR(_S_); PRINTBSDFCHAR(_DOT_); PRINTBSDFCHAR(_S_) PRINTBSDFCHAR(_S_) PRINTBSDFCHAR(_D_) PRINTBSDFCHAR(_M_) PRINTBSDFCHAR(_F_) PRINTBSDFCHAR(_P_);	PRINTBSDFCHAR(_SPC_); PRINTFLOAT3(SSSData.DMFP);
+#define PRINTCLASSIFSSMFP()	PRINTBSDFCHAR(_C_); PRINTBSDFCHAR(_L_); PRINTBSDFCHAR(_A_); PRINTBSDFCHAR(_S_); PRINTBSDFCHAR(_DOT_); PRINTBSDFCHAR(_S_) PRINTBSDFCHAR(_S_) PRINTBSDFCHAR(_M_) PRINTBSDFCHAR(_F_) PRINTBSDFCHAR(_P_);	PRINTBSDFCHAR(_SPC_); PRINTFLOAT3(SSSData.MFP);
 
 #define NEXTLINE  TopLeft.y += 8; RectMax.x = max(RectMax.x, TopLeft.x); RectMax.y = max(RectMax.y, TopLeft.y); TopLeft.x = TopLeftInit.x; 
 #define HORIZONTALSEPARATION  TopLeft.y += 4;
@@ -322,7 +322,7 @@ void VisualizeMaterialPS(
 			}
 			else
 			{
-				PRINTCLASSIFSSDMFP();
+				PRINTCLASSIFSSMFP();
 				NEXTLINE;
 				NEXTLINE;
 			}
@@ -444,8 +444,8 @@ void VisualizeMaterialPS(
 					}
 					else
 					{
-						HORIZONTALTAB; PRINTPARAMCHAR(_S_); PRINTPARAMCHAR(_S_); PRINTPARAMCHAR(_S_); PRINTPARAMCHAR(_D_); PRINTPARAMCHAR(_M_); PRINTPARAMCHAR(_F_); PRINTPARAMCHAR(_P_);
-						HORIZONTALTAB; PRINTFLOAT3(SLAB_SSSDMFP(BSDF));
+						HORIZONTALTAB; PRINTPARAMCHAR(_S_); PRINTPARAMCHAR(_S_); PRINTPARAMCHAR(_S_); PRINTPARAMCHAR(_M_); PRINTPARAMCHAR(_F_); PRINTPARAMCHAR(_P_);
+						HORIZONTALTAB; PRINTFLOAT3(SLAB_SSSMFP(BSDF));
 						NEXTLINE;
 					}
 				}
diff --git a/Engine/Source/Runtime/Engine/Classes/Materials/MaterialExpressionStrata.h b/Engine/Source/Runtime/Engine/Classes/Materials/MaterialExpressionStrata.h
--- a/Engine/Source/Runtime/Engine/Classes/Materials/MaterialExpressionStrata.h
+++ b/Engine/Source/Runtime/Engine/Classes/Materials/MaterialExpressionStrata.h
@@ -258,13 +258,13 @@ class UMaterialExpressionStrataSlabBSDF : public UMaterialExpressionStrataBSDF
 	 * Chromatic mean free path . Only used when there is not any sub-surface profile provided. (type = float3, unit = centimeters, default = 0)
 	 */
 	UPROPERTY()
-	FExpressionInput SSSDMFP;
+	FExpressionInput SSSMFP;
 
 	/**
 	 * Scale the mean free path radius of the SSS profile according to a value between 0 and 1. Always used, when a subsurface profile is provided or not. (type = float, unitless, defaults to 1)
 	 */
 	UPROPERTY()
-	FExpressionInput SSSDMFPScale;
+	FExpressionInput SSSMFPScale;
 
 	/**
 	 * Emissive color on top of the surface (type = float3, unit = luminance, default = 0)
@@ -323,7 +323,7 @@ class UMaterialExpressionStrataSlabBSDF : public UMaterialExpressionStrataBSDF
 	bool HasHaziness() const;
 	bool HasSSS() const;
 	bool HasSSSProfile() const;
-	bool HasDMFPPluggedIn() const;
+	bool HasMFPPluggedIn() const;
 	bool HasAnisotropy() const;
 #endif
 	//~ End UMaterialExpression Interface
diff --git a/Engine/Source/Runtime/Engine/Private/Materials/HLSLMaterialTranslator.cpp b/Engine/Source/Runtime/Engine/Private/Materials/HLSLMaterialTranslator.cpp
--- a/Engine/Source/Runtime/Engine/Private/Materials/HLSLMaterialTranslator.cpp
+++ b/Engine/Source/Runtime/Engine/Private/Materials/HLSLMaterialTranslator.cpp
@@ -10137,7 +10137,7 @@ int32 FHLSLMaterialTranslator::StrataSlabBSDF(
 	int32 BaseColor, int32 EdgeColor, int32 Specular, int32 Metallic,
 	int32 DiffuseAlbedo, int32 F0, int32 F90,
 	int32 Roughness, int32 Anisotropy,
-	int32 SSSProfileId, int32 SSSDMFP, int32 SSSDMFPScale, 
+	int32 SSSProfileId, int32 SSSMFP, int32 SSSMFPScale, 
 	int32 EmissiveColor,	
 	int32 Haziness, 
 	int32 FuzzAmount, int32 FuzzColor,
@@ -10160,8 +10160,8 @@ int32 FHLSLMaterialTranslator::StrataSlabBSDF(
 			*GetParameterCode(DiffuseAlbedo), *GetParameterCode(F0), *GetParameterCode(F90),
 			*GetParameterCode(Roughness), *GetParameterCode(Anisotropy),
 			*GetParameterCode(SSSProfileId),
-			*GetParameterCode(SSSDMFP),
-			*GetParameterCode(SSSDMFPScale),
+			*GetParameterCode(SSSMFP),
+			*GetParameterCode(SSSMFPScale),
 			*GetParameterCode(EmissiveColor),
 			*GetParameterCode(Haziness),
 			*GetParameterCode(FuzzAmount), *GetParameterCode(FuzzColor),
@@ -10182,8 +10182,8 @@ int32 FHLSLMaterialTranslator::StrataSlabBSDF(
 		*GetParameterCode(DiffuseAlbedo),	*GetParameterCode(F0),			*GetParameterCode(F90),
 		*GetParameterCode(Roughness),		*GetParameterCode(Anisotropy),
 		*GetParameterCode(SSSProfileId),
-		*GetParameterCode(SSSDMFP),
-		*GetParameterCode(SSSDMFPScale),
+		*GetParameterCode(SSSMFP),
+		*GetParameterCode(SSSMFPScale),
 		*GetParameterCode(EmissiveColor),
 		*GetParameterCode(Haziness),
 		*GetParameterCode(FuzzAmount),		*GetParameterCode(FuzzColor),
diff --git a/Engine/Source/Runtime/Engine/Private/Materials/HLSLMaterialTranslator.h b/Engine/Source/Runtime/Engine/Private/Materials/HLSLMaterialTranslator.h
--- a/Engine/Source/Runtime/Engine/Private/Materials/HLSLMaterialTranslator.h
+++ b/Engine/Source/Runtime/Engine/Private/Materials/HLSLMaterialTranslator.h
@@ -1015,7 +1015,7 @@ protected:
 		int32 BaseColor, int32 EdgeColor, int32 Specular, int32 Metallic,
 		int32 DiffuseAlbedo, int32 F0, int32 F90,
 		int32 Roughness, int32 Anisotropy,
-		int32 SSSProfileId, int32 SSSDMFP, int32 SSSDMFPScale,
+		int32 SSSProfileId, int32 SSSMFP, int32 SSSMFPScale,
 		int32 EmissiveColor, 
 		int32 Haziness, 
 		int32 FuzzAmount, int32 FuzzColor, 
diff --git a/Engine/Source/Runtime/Engine/Private/Materials/MaterialExpressions.cpp b/Engine/Source/Runtime/Engine/Private/Materials/MaterialExpressions.cpp
--- a/Engine/Source/Runtime/Engine/Private/Materials/MaterialExpressions.cpp
+++ b/Engine/Source/Runtime/Engine/Private/Materials/MaterialExpressions.cpp
@@ -20514,7 +20514,7 @@ int32 UMaterialExpressionStrataLegacyConversion::Compile(class FMaterialCompiler
 		StrataInfo.Layers[0].BSDFs[0].Type = STRATA_BSDF_TYPE_SLAB;
 		StrataInfo.Layers[0].BSDFs[0].RegisteredSharedLocalBasis = NewRegisteredSharedLocalBasis;
 		StrataInfo.Layers[0].BSDFs[0].bHasSSS = bSSS;
-		StrataInfo.Layers[0].BSDFs[0].bHasDMFPPluggedIn = true;
+		StrataInfo.Layers[0].BSDFs[0].bHasMFPPluggedIn = true;
 		StrataInfo.Layers[0].BSDFs[0].bHasEdgeColor = false;
 		StrataInfo.Layers[0].BSDFs[0].bHasThinFilm = false;
 		StrataInfo.Layers[0].BSDFs[0].bHasFuzz = false;
@@ -20524,7 +20524,7 @@ int32 UMaterialExpressionStrataLegacyConversion::Compile(class FMaterialCompiler
 		StrataInfo.Layers[1].BSDFs[0].Type = STRATA_BSDF_TYPE_SLAB;
 		StrataInfo.Layers[1].BSDFs[0].RegisteredSharedLocalBasis = ClearCoat_NewRegisteredSharedLocalBasis;
 		StrataInfo.Layers[1].BSDFs[0].bHasSSS = bSSS;
-		StrataInfo.Layers[1].BSDFs[0].bHasDMFPPluggedIn = false;
+		StrataInfo.Layers[1].BSDFs[0].bHasMFPPluggedIn = false;
 		StrataInfo.Layers[1].BSDFs[0].bHasEdgeColor = false;
 		StrataInfo.Layers[1].BSDFs[0].bHasThinFilm = false;
 		StrataInfo.Layers[1].BSDFs[0].bHasFuzz = true;
@@ -20547,7 +20547,7 @@ int32 UMaterialExpressionStrataLegacyConversion::Compile(class FMaterialCompiler
 			StrataCompilationInfoCreateSingleBSDFMaterial(Compiler, OutputCodeChunk, NewRegisteredSharedLocalBasis,
 				STRATA_BSDF_TYPE_UNLIT, 
 				false /*bHasSSS*/, 
-				false /*bHasDMFPPluggedIn*/, 
+				false /*bHasMFPPluggedIn*/, 
 				false /*bHasEdgeColor*/, 
 				false /*bHasFuzz*/, 
 				false /*bHasHaziness*/);
@@ -20565,7 +20565,7 @@ int32 UMaterialExpressionStrataLegacyConversion::Compile(class FMaterialCompiler
 			StrataCompilationInfoCreateSingleBSDFMaterial(Compiler, OutputCodeChunk, NewRegisteredSharedLocalBasis,
 				STRATA_BSDF_TYPE_HAIR,
 				false /*bHasSSS*/,
-				false /*bHasDMFPPluggedIn*/,
+				false /*bHasMFPPluggedIn*/,
 				false /*bHasEdgeColor*/,
 				false /*bHasFuzz*/,
 				false /*bHasHaziness*/);
@@ -20575,7 +20575,7 @@ int32 UMaterialExpressionStrataLegacyConversion::Compile(class FMaterialCompiler
 			StrataCompilationInfoCreateSingleBSDFMaterial(Compiler, OutputCodeChunk, NewRegisteredSharedLocalBasis,
 				STRATA_BSDF_TYPE_SINGLELAYERWATER,
 				false /*bHasSSS*/,
-				false /*bHasDMFPPluggedIn*/,
+				false /*bHasMFPPluggedIn*/,
 				false /*bHasEdgeColor*/,
 				false /*bHasFuzz*/,
 				false /*bHasHaziness*/);
@@ -20795,7 +20795,7 @@ int32 UMaterialExpressionStrataSlabBSDF::Compile(class FMaterialCompiler* Compil
 	const bool bHasEdgeColor = HasEdgeColor(); // This accounts for EdgeColor and also F90 when the non metalness worfklow is selected.
 	const bool bHasFuzz = HasFuzz();
 	const bool bHasHaziness = HasHaziness();
-	const bool bHasDMFPPluggedIn = HasDMFPPluggedIn();
+	const bool bHasMFPPluggedIn = HasMFPPluggedIn();
 
 	int32 SSSProfileCodeChunk = INDEX_NONE;
 	const bool bHasSSS = HasSSS();
@@ -20830,8 +20830,8 @@ int32 UMaterialExpressionStrataSlabBSDF::Compile(class FMaterialCompiler* Compil
 		RoughnessCodeChunk,
 		AnisotropyCodeChunk,
 		SSSProfileCodeChunk != INDEX_NONE ? SSSProfileCodeChunk : Compiler->Constant(0.0f),	
-		CompileWithDefaultFloat3(Compiler, SSSDMFP, 0.0f, 0.0f, 0.0f),
-		CompileWithDefaultFloat1(Compiler, SSSDMFPScale, 1.0f),
+		CompileWithDefaultFloat3(Compiler, SSSMFP, 0.0f, 0.0f, 0.0f),
+		CompileWithDefaultFloat1(Compiler, SSSMFPScale, 1.0f),
 		CompileWithDefaultFloat3(Compiler, EmissiveColor, 0.0f, 0.0f, 0.0f),
 		CompileWithDefaultFloat1(Compiler, Haziness, 0.0f),
 		CompileWithDefaultFloat1(Compiler, FuzzAmount, 0.0f),
@@ -20841,7 +20841,7 @@ int32 UMaterialExpressionStrataSlabBSDF::Compile(class FMaterialCompiler* Compil
 		TangentCodeChunk,
 		Compiler->GetStrataSharedLocalBasisIndexMacro(NewRegisteredSharedLocalBasis),
 		!StrataOperator.bUseParameterBlending || (StrataOperator.bUseParameterBlending && StrataOperator.bRootOfParameterBlendingSubTree) ? &StrataOperator : nullptr);
-	StrataCompilationInfoCreateSingleBSDFMaterial(Compiler, OutputCodeChunk, NewRegisteredSharedLocalBasis, STRATA_BSDF_TYPE_SLAB, bHasSSS, bHasDMFPPluggedIn, bHasEdgeColor, bHasFuzz, bHasHaziness);
+	StrataCompilationInfoCreateSingleBSDFMaterial(Compiler, OutputCodeChunk, NewRegisteredSharedLocalBasis, STRATA_BSDF_TYPE_SLAB, bHasSSS, bHasMFPPluggedIn, bHasEdgeColor, bHasFuzz, bHasHaziness);
 
 	return OutputCodeChunk;
 }
@@ -20867,8 +20867,8 @@ const TArray<FExpressionInput*> UMaterialExpressionStrataSlabBSDF::GetInputs()
 	Result.Add(&Anisotropy);
 	Result.Add(&Normal);
 	Result.Add(&Tangent);
-	Result.Add(&SSSDMFP);
-	Result.Add(&SSSDMFPScale);
+	Result.Add(&SSSMFP);
+	Result.Add(&SSSMFPScale);
 	Result.Add(&EmissiveColor);
 	Result.Add(&Haziness);
 	Result.Add(&Thickness);
@@ -20993,11 +20993,11 @@ uint32 UMaterialExpressionStrataSlabBSDF::GetInputType(int32 InputIndex)
 	}
 	else if (InputIndex == (8 + SkipDisabledInputOffset))
 	{
-		return MCT_Float3; // SSSDMFP
+		return MCT_Float3; // SSSMFP
 	}
 	else if (InputIndex == (9 + SkipDisabledInputOffset))
 	{
-		return MCT_Float1; // SSSDMFPScale
+		return MCT_Float1; // SSSMFPScale
 	}
 	else if (InputIndex == (10 + SkipDisabledInputOffset))
 	{
@@ -21081,11 +21081,11 @@ FName UMaterialExpressionStrataSlabBSDF::GetInputName(int32 InputIndex) const
 	}
 	else if (InputIndex == (8 + SkipDisabledInputOffset))
 	{
-		return TEXT("SSS Diffuse MFP");
+		return TEXT("SSS MFP");
 	}
 	else if (InputIndex == (9 + SkipDisabledInputOffset))
 	{
-		return TEXT("SSS Diffuse MFP Scale");
+		return TEXT("SSS MFP Scale");
 	}
 	else if (InputIndex == (10 + SkipDisabledInputOffset))
 	{
@@ -21229,7 +21229,7 @@ void UMaterialExpressionStrataSlabBSDF::GatherStrataMaterialInfo(FStrataMaterial
 	if (EmissiveColor.IsConnected())		{ StrataMaterialInfo.AddPropertyConnected(MP_EmissiveColor); }
 	if (Normal.IsConnected())				{ StrataMaterialInfo.AddPropertyConnected(MP_Normal); }
 	if (Tangent.IsConnected())				{ StrataMaterialInfo.AddPropertyConnected(MP_Tangent); }
-	if (SSSDMFP.IsConnected())				{ StrataMaterialInfo.AddPropertyConnected(MP_SubsurfaceColor); }
+	if (SSSMFP.IsConnected())				{ StrataMaterialInfo.AddPropertyConnected(MP_SubsurfaceColor); }
 
 	if (HasSSS())
 	{
@@ -21253,7 +21253,7 @@ FStrataOperator* UMaterialExpressionStrataSlabBSDF::StrataGenerateMaterialTopolo
 
 bool UMaterialExpressionStrataSlabBSDF::HasSSS() const
 {
-	return SubsurfaceProfile != nullptr || SSSDMFP.IsConnected();
+	return SubsurfaceProfile != nullptr || SSSMFP.IsConnected();
 }
 
 bool UMaterialExpressionStrataSlabBSDF::HasSSSProfile() const
@@ -21261,9 +21261,9 @@ bool UMaterialExpressionStrataSlabBSDF::HasSSSProfile() const
 	return SubsurfaceProfile != nullptr;
 }
 
-bool UMaterialExpressionStrataSlabBSDF::HasDMFPPluggedIn() const
+bool UMaterialExpressionStrataSlabBSDF::HasMFPPluggedIn() const
 {
-	return SSSDMFP.IsConnected();
+	return SSSMFP.IsConnected();
 }
 
 bool UMaterialExpressionStrataSlabBSDF::HasEdgeColor() const
diff --git a/Engine/Source/Runtime/Engine/Private/Materials/StrataMaterial.cpp b/Engine/Source/Runtime/Engine/Private/Materials/StrataMaterial.cpp
--- a/Engine/Source/Runtime/Engine/Private/Materials/StrataMaterial.cpp
+++ b/Engine/Source/Runtime/Engine/Private/Materials/StrataMaterial.cpp
@@ -65,7 +65,7 @@ FStrataRegisteredSharedLocalBasis StrataCompilationInfoCreateSharedLocalBasis(FM
 
 void StrataCompilationInfoCreateSingleBSDFMaterial(FMaterialCompiler* Compiler, int32 CodeChunk,
 	const FStrataRegisteredSharedLocalBasis& RegisteredSharedLocalBasis,
-	uint8 BSDFType, bool bHasSSS, bool bHasDMFPPluggedIn, bool bHasEdgeColor, bool bHasFuzz, bool bHasHaziness)
+	uint8 BSDFType, bool bHasSSS, bool bHasMFPPluggedIn, bool bHasEdgeColor, bool bHasFuzz, bool bHasHaziness)
 {
 	FStrataMaterialCompilationInfo StrataInfo;
 	StrataInfo.LayerCount = 1;
@@ -73,7 +73,7 @@ void StrataCompilationInfoCreateSingleBSDFMaterial(FMaterialCompiler* Compiler,
 	StrataInfo.Layers[0].BSDFs[0].Type = BSDFType;
 	StrataInfo.Layers[0].BSDFs[0].RegisteredSharedLocalBasis = RegisteredSharedLocalBasis;
 	StrataInfo.Layers[0].BSDFs[0].bHasSSS = bHasSSS;
-	StrataInfo.Layers[0].BSDFs[0].bHasDMFPPluggedIn = bHasDMFPPluggedIn;
+	StrataInfo.Layers[0].BSDFs[0].bHasMFPPluggedIn = bHasMFPPluggedIn;
 	StrataInfo.Layers[0].BSDFs[0].bHasEdgeColor = bHasEdgeColor;
 	StrataInfo.Layers[0].BSDFs[0].bHasThinFilm = false;
 	StrataInfo.Layers[0].BSDFs[0].bHasFuzz = bHasFuzz;
@@ -155,7 +155,7 @@ FStrataMaterialCompilationInfo StrataCompilationInfoAddParamBlend(FMaterialCompi
 
 	// When parameter blending is used, we take the union of all the features activated by input BSDFs.
 	NewBSDF.bHasSSS				|=	OtherBSDF.bHasSSS;
-	NewBSDF.bHasDMFPPluggedIn	|=	OtherBSDF.bHasDMFPPluggedIn;
+	NewBSDF.bHasMFPPluggedIn	|=	OtherBSDF.bHasMFPPluggedIn;
 	NewBSDF.bHasEdgeColor		|=	OtherBSDF.bHasEdgeColor;
 	NewBSDF.bHasThinFilm		|=	OtherBSDF.bHasThinFilm;
 	NewBSDF.bHasFuzz			|=	OtherBSDF.bHasFuzz;
@@ -185,7 +185,7 @@ FStrataMaterialCompilationInfo StrataCompilationInfoHorizontalMixingParamBlend(F
 
 	// When parameter blending is used, we take the union of all the features activated by input BSDFs.
 	NewBSDF.bHasSSS				|=	OtherBSDF.bHasSSS;
-	NewBSDF.bHasDMFPPluggedIn	|=	OtherBSDF.bHasDMFPPluggedIn;
+	NewBSDF.bHasMFPPluggedIn	|=	OtherBSDF.bHasMFPPluggedIn;
 	NewBSDF.bHasEdgeColor		|=	OtherBSDF.bHasEdgeColor;
 	NewBSDF.bHasThinFilm		|=	OtherBSDF.bHasThinFilm;
 	NewBSDF.bHasFuzz			|=	OtherBSDF.bHasFuzz;
@@ -231,7 +231,7 @@ FStrataMaterialCompilationInfo StrataCompilationInfoVerticalLayeringParamBlend(F
 	NewBSDF.RegisteredSharedLocalBasis = RegisteredSharedLocalBasis;
 
 //	NewBSDF.bHasSSS				= TopBSDF.bHasSSS;				// We keep SSS only if the base layer has it. Otherwise it will be simple volume and the throughput wll be applied on the parameters.
-//	NewBSDF.bHasDMFPPluggedIn	= TopBSDF.bHasDMFPPluggedIn;	// Idem
+//	NewBSDF.bHasMFPPluggedIn	= TopBSDF.bHasMFPPluggedIn;		// Idem
 	NewBSDF.bHasEdgeColor		|= TopBSDF.bHasEdgeColor;		// We keep the union of both, even though it will be hard to get a perfect match.
 	NewBSDF.bHasFuzz			|= TopBSDF.bHasFuzz;			// Idem
 	NewBSDF.bHasHaziness		|= TopBSDF.bHasHaziness;		// Idem
@@ -386,7 +386,7 @@ FStrataMaterialAnalysisResult StrataCompilationInfoMaterialAnalysis(FMaterialCom
 			case STRATA_BSDF_TYPE_SLAB:
 			{
 				// Compute values closer to the reality for HasSSS and IsSimpleVolume, now that we know that we know the topology of the material.
-				const bool bIsSimpleVolume = !bBottomLayer && BSDF.bHasDMFPPluggedIn;
+				const bool bIsSimpleVolume = !bBottomLayer && BSDF.bHasMFPPluggedIn;
 				const bool bHasSSS = bBottomLayer && BSDF.bHasSSS && !bIsSimpleVolume;
 
 				Result.RequestedByteCount += UintByteSize;
diff --git a/Engine/Source/Runtime/Engine/Private/Materials/StrataMaterial.h b/Engine/Source/Runtime/Engine/Private/Materials/StrataMaterial.h
--- a/Engine/Source/Runtime/Engine/Private/Materials/StrataMaterial.h
+++ b/Engine/Source/Runtime/Engine/Private/Materials/StrataMaterial.h
@@ -22,7 +22,7 @@ struct FStrataMaterialCompilationInfo
 
 		// Those properties are centered around the SLAB node. Maybe this can be abstracted at some point.
 		bool  bHasSSS;
-		bool  bHasDMFPPluggedIn;
+		bool  bHasMFPPluggedIn;
 		bool  bHasEdgeColor;
 		bool  bHasThinFilm;
 		bool  bHasFuzz;
@@ -66,7 +66,7 @@ FStrataRegisteredSharedLocalBasis StrataCompilationInfoCreateSharedLocalBasis(FM
 void  StrataCompilationInfoCreateNullBSDF(FMaterialCompiler* Compiler, int32 CodeChunk);
 void  StrataCompilationInfoCreateSingleBSDFMaterial(FMaterialCompiler* Compiler, int32 CodeChunk, 
 	const FStrataRegisteredSharedLocalBasis& RegisteredSharedLocalBasis,
-	uint8 BSDFType, bool bHasSSS = false, bool bHasDMFPPluggedIn = false, bool bHasEdgeColor = false, bool bHasFuzz = false, bool bHasHaziness = false);
+	uint8 BSDFType, bool bHasSSS = false, bool bHasMFPPluggedIn = false, bool bHasEdgeColor = false, bool bHasFuzz = false, bool bHasHaziness = false);
 
 FStrataMaterialCompilationInfo StrataCompilationInfoWeight(FMaterialCompiler* Compiler, const FStrataMaterialCompilationInfo& A);
 FStrataMaterialCompilationInfo StrataCompilationInfoWeightParamBlend(FMaterialCompiler* Compiler, const FStrataMaterialCompilationInfo& A);
diff --git a/Engine/Source/Runtime/Engine/Private/Rendering/BurleyNormalizedSSS.cpp b/Engine/Source/Runtime/Engine/Private/Rendering/BurleyNormalizedSSS.cpp
--- a/Engine/Source/Runtime/Engine/Private/Rendering/BurleyNormalizedSSS.cpp
+++ b/Engine/Source/Runtime/Engine/Private/Rendering/BurleyNormalizedSSS.cpp
@@ -169,6 +169,7 @@ void ComputeMirroredBSSSKernel(FLinearColor* TargetBuffer, uint32 TargetBufferSi
 		}
 
 		// We still need to do a small tweak to get the radius to visually match. Multiplying by 2.0 seems to fix it.
+		// Match that in GetSubsurfaceProfileKernel in PostProcessSubsurface.usf
 		const float StepScale = 2.0f;
 		for (int32 i = 0; i < nTotalSamples; i++)
 		{
diff --git a/Engine/Source/Runtime/Engine/Public/MaterialCompiler.h b/Engine/Source/Runtime/Engine/Public/MaterialCompiler.h
--- a/Engine/Source/Runtime/Engine/Public/MaterialCompiler.h
+++ b/Engine/Source/Runtime/Engine/Public/MaterialCompiler.h
@@ -492,7 +492,7 @@ public:
 		int32 BaseColor, int32 EdgeColor, int32 Specular, int32 Metallic,
 		int32 DiffuseAlbedo, int32 F0, int32 F90,
 		int32 Roughness, int32 Anisotropy,
-		int32 SSSProfileId, int32 SSSDMFP, int32 SSSDMFPScale,
+		int32 SSSProfileId, int32 SSSMFP, int32 SSSMFPScale,
 		int32 EmissiveColor, 
 		int32 Haziness, 
 		int32 FuzzAmount, int32 FuzzColor,
@@ -1032,7 +1032,7 @@ public:
 		int32 BaseColor, int32 EdgeColor, int32 Specular, int32 Metallic,
 		int32 DiffuseAlbedo, int32 F0, int32 F90,
 		int32 Roughness, int32 Anisotropy,
-		int32 SSSProfileId, int32 SSSDMFP, int32 SSSDMFPScale,
+		int32 SSSProfileId, int32 SSSMFP, int32 SSSMFPScale,
 		int32 EmissiveColor, 
 		int32 Haziness, 
 		int32 FuzzAmount, int32 FuzzColor,
@@ -1045,7 +1045,7 @@ public:
 			BaseColor, EdgeColor, Specular, Metallic,
 			DiffuseAlbedo, F0, F90,
 			Roughness, Anisotropy,
-			SSSProfileId, SSSDMFP, SSSDMFPScale,
+			SSSProfileId, SSSMFP, SSSMFPScale,
 			EmissiveColor, 
 			Haziness, 
 			FuzzAmount, FuzzColor,
