Fix implementations of TransformLocalPositionToPrevWorld and TransformLocalVectorToPrevWorld to use PrimitiveData.
This matches the implementations of TransformLocalPositionToWorld and TransformLocalVectorToWorld.
Previously one would use the (compressed) instance transform, and one would use the (uncompressed) primitive transform.
The small difference between the two can give large differences when calculating WPO velocity vectors.

diff --git a/Engine/Shaders/Private/MaterialTemplate.ush b/Engine/Shaders/Private/MaterialTemplate.ush
--- a/Engine/Shaders/Private/MaterialTemplate.ush
+++ b/Engine/Shaders/Private/MaterialTemplate.ush
@@ -605,6 +605,11 @@ MaterialFloat3x3 GetLocalToWorld3x3(FMaterialVertexParameters Parameters)
 	return LWCToFloat3x3(GetPrimitiveData(Parameters).LocalToWorld);
 }
 
+MaterialFloat3x3 GetPreviousLocalToWorld3x3(FMaterialVertexParameters Parameters)
+{
+	return LWCToFloat3x3(GetPrimitiveData(Parameters).LocalToWorld);
+}
+
 MaterialFloat3x3 GetLocalToWorld3x3(FMaterialPixelParameters Parameters)
 {
 	return LWCToFloat3x3(GetPrimitiveData(Parameters).LocalToWorld);
@@ -1156,6 +1161,12 @@ FLWCMatrix GetPrevLocalToWorld(FMaterialVertexParameters Parameters)
 	return Parameters.PrevFrameLocalToWorld;
 }
 
+/** Transforms a vector from local space to world space (PS version) */
+MaterialFloat3 TransformLocalVectorToWorld(FMaterialPixelParameters Parameters, MaterialFloat3 InLocalVector)
+{
+	return mul(InLocalVector, GetLocalToWorld3x3(Parameters));
+}
+
 /** Transforms a vector from local space to world space (VS version) */
 MaterialFloat3 TransformLocalVectorToWorld(FMaterialVertexParameters Parameters,MaterialFloat3 InLocalVector)
 {
@@ -1166,16 +1177,14 @@ MaterialFloat3 TransformLocalVectorToWorld(FMaterialVertexParameters Parameters,
 	#endif
 }
 
-/** Transforms a vector from local space to world space (PS version) */
-MaterialFloat3 TransformLocalVectorToWorld(FMaterialPixelParameters Parameters,MaterialFloat3 InLocalVector)
-{
-	return mul(InLocalVector, GetLocalToWorld3x3(Parameters));
-}
-
 /** Transforms a vector from local space to previous frame world space (VS version) */
-MaterialFloat3 TransformLocalVectorToPrevWorld(FMaterialVertexParameters Parameters,MaterialFloat3 InLocalVector)
+MaterialFloat3 TransformLocalVectorToPrevWorld(FMaterialVertexParameters Parameters, MaterialFloat3 InLocalVector)
 {
+#if USE_INSTANCING || USE_INSTANCE_CULLING || IS_MESHPARTICLE_FACTORY
 	return LWCMultiplyVector(InLocalVector, Parameters.PrevFrameLocalToWorld);
+#else
+	return mul(InLocalVector, GetPreviousLocalToWorld3x3(Parameters));
+#endif
 }
 
 #if HAS_PRIMITIVE_UNIFORM_BUFFER
@@ -1199,7 +1208,11 @@ FLWCVector3 TransformLocalPositionToWorld(FMaterialVertexParameters Parameters,f
 /** Transforms a position from local space to previous frame absolute world space */
 FLWCVector3 TransformLocalPositionToPrevWorld(FMaterialVertexParameters Parameters,float3 InLocalPosition)
 {
+#if USE_INSTANCING || USE_INSTANCE_CULLING || IS_MESHPARTICLE_FACTORY
 	return LWCMultiply(InLocalPosition, Parameters.PrevFrameLocalToWorld);
+#else
+	return LWCMultiply(InLocalPosition, GetPrimitiveData(Parameters).PreviousLocalToWorld);
+#endif
 }
 
 #endif
