Remove thin film evaluation.

Thin film support relies now only on inline evaluation (i.e. based into F0/F90)

diff --git a/Engine/Shaders/Private/ShadingCommon.ush b/Engine/Shaders/Private/ShadingCommon.ush
--- a/Engine/Shaders/Private/ShadingCommon.ush
+++ b/Engine/Shaders/Private/ShadingCommon.ush
@@ -125,11 +125,10 @@ float DielectricIorToF0(float Ior)
 	return F0;
 }
 
-float F0ToMicroOcclusion(float F0)
-{
-	// Anything with Specular less than 2% is physically impossible and is instead considered to be shadowing.
-	return saturate(50.0 * F0);
-}
+// Anything with Specular less than 2% is physically impossible and is instead considered to be shadowing.
+float  GetF0MicroOcclusionThreshold() { return 0.02f; }
+float  F0ToMicroOcclusion(float F0)   { return saturate(50.0 * F0); }
+float3 F0ToMicroOcclusion(float3 F0)  { return saturate(50.0 * F0); }
 
 float F0RGBToMicroOcclusion(float3 F0)
 {
diff --git a/Engine/Shaders/Private/Strata/Strata.ush b/Engine/Shaders/Private/Strata/Strata.ush
--- a/Engine/Shaders/Private/Strata/Strata.ush
+++ b/Engine/Shaders/Private/Strata/Strata.ush
@@ -181,16 +181,16 @@ struct FStrataBSDF
 //            BASIS(0)      32 |  Normal: Oct 16/16 | Tangent: 11/11/10
 //            BASIS(1)      32 |  Normal: Oct 16/16 | Tangent: 11/11/10
 //                             |
-//            STATE(0)      17 |  BSDF Type | Normal Id | Aniso | TopLayer | SSS | SSSProfile | GreyWeigh | Haze | F90 | ThinFilm | Simple vol. | MFP Plugged | Fuzz | Thin
-//                             |  3           2           1       1          1     1            1           1      1     1          1             1             1      1 
+//            STATE(0)      16 |  BSDF Type | Normal Id | Aniso | TopLayer | SSS | SSSProfile | GreyWeigh | Haze | F90 | Simple vol. | MFP Plugged | Fuzz | Thin
+//                             |  3           2           1       1          1     1            1           1      1     1             1             1      1 
 //            BSDF(0)       ?? |  Adaptive
-//            STATE(1)      17 |  BSDF Type | Normal Id | Aniso | TopLayer | SSS | SSSProfile | GreyWeigh | Haze | F90 | ThinFilm | Simple vol. | MFP Plugged | Fuzz | Thin
-//                             |  3           2           1       1          1     1            1           1      1     1          1             1             1      1 
+//            STATE(1)      16 |  BSDF Type | Normal Id | Aniso | TopLayer | SSS | SSSProfile | GreyWeigh | Haze | F90 | Simple vol. | MFP Plugged | Fuzz | Thin
+//                             |  3           2           1       1          1     1            1           1      1     1             1             1      1 
 //            BSDF(1)       ?? |  Adaptive
 //                             |
 // * LAYOUT 1 -------------------------------------------------------------------------------------------------------------------------------------------------------------------------
-//  (Single)  HEADER        31 |  State | AO | BSDF Types | Aniso | TopLayer | SSS | SSSProfile | GreyWeigh | Haze | F90 | ThinFilm | Simple vol. | MFP Plugged | Fuzz | Thin
-//                             |  8       8    4            1       1          1     1            1           1      1     1          1             1             1      1            
+//  (Single)  HEADER        30 |  State | AO | BSDF Types | Aniso | TopLayer | SSS | SSSProfile | GreyWeigh | Haze | F90 | Simple vol. | MFP Plugged | Fuzz | Thin
+//                             |  8       8    4            1       1          1     1            1           1      1     1             1             1      1            
 //            BASIS(0)       0 | (Use top layer normal)
 //            BSDF(0)       ?? |  Adaptive
 // * LAYOUT 0 -------------------------------------------------------------------------------------------------------------------------------------------------------------------------          
@@ -290,14 +290,13 @@ struct FStrataBSDF
 #define STATE_BIT_COUNT_GREYWEIGHT				1
 #define STATE_BIT_COUNT_HAZINESS				1
 #define STATE_BIT_COUNT_F90						1
-#define STATE_BIT_COUNT_THINFILM				1
 #define STATE_BIT_COUNT_SIMPLEVOLUME			1
 #define STATE_BIT_COUNT_MFPPLUGGED				1
 #define STATE_BIT_COUNT_HASFUZZ					1
 #define STATE_BIT_COUNT_ISTHIN					1
 #define STATE_BIT_COUNT_BSDF_TYPE				3
 #define STATE_BIT_COUNT_SHAREDLOCALBASESID		2
-#define STATE_BIT_COUNT___UNUSED___				5
+#define STATE_BIT_COUNT___UNUSED___				6
 #define STATE_BIT_COUNT_WEIGHT10F				10
 
 // Bit offset of each component store in BSDF state 
@@ -308,8 +307,7 @@ struct FStrataBSDF
 #define STATE_BIT_OFFSET_GREYWEIGHT				(STATE_BIT_OFFSET_HASSSSPROFILE		+ STATE_BIT_COUNT_HASSSSPROFILE)
 #define STATE_BIT_OFFSET_HAZINESS				(STATE_BIT_OFFSET_GREYWEIGHT		+ STATE_BIT_COUNT_GREYWEIGHT)
 #define STATE_BIT_OFFSET_F90					(STATE_BIT_OFFSET_HAZINESS			+ STATE_BIT_COUNT_HAZINESS)
-#define STATE_BIT_OFFSET_THINFILM				(STATE_BIT_OFFSET_F90				+ STATE_BIT_COUNT_F90)
-#define STATE_BIT_OFFSET_SIMPLEVOLUME			(STATE_BIT_OFFSET_THINFILM			+ STATE_BIT_COUNT_THINFILM)
+#define STATE_BIT_OFFSET_SIMPLEVOLUME			(STATE_BIT_OFFSET_F90				+ STATE_BIT_COUNT_F90)
 #define STATE_BIT_OFFSET_MFPPLUGGED				(STATE_BIT_OFFSET_SIMPLEVOLUME		+ STATE_BIT_COUNT_SIMPLEVOLUME)
 #define STATE_BIT_OFFSET_HASFUZZ				(STATE_BIT_OFFSET_MFPPLUGGED		+ STATE_BIT_COUNT_MFPPLUGGED)
 #define STATE_BIT_OFFSET_ISTHIN					(STATE_BIT_OFFSET_HASFUZZ			+ STATE_BIT_COUNT_HASFUZZ)
@@ -370,10 +368,6 @@ struct FStrataBSDF
 #define BSDF_GETHASF90(X)				READ_BITS(X.State,  STATE_BIT_COUNT_F90, STATE_BIT_OFFSET_F90)
 #define BSDF_SETHASF90(X, Col)			WRITE_BITS(X.State, STATE_BIT_COUNT_F90, STATE_BIT_OFFSET_F90, Col)
 
-// Indicates if the BSDF has thin film
-#define BSDF_GETHASTHINFILM(X)			READ_BITS(X.State,  STATE_BIT_COUNT_THINFILM, STATE_BIT_OFFSET_THINFILM)
-#define BSDF_SETHASTHINFILM(X, Thin)	WRITE_BITS(X.State, STATE_BIT_COUNT_THINFILM, STATE_BIT_OFFSET_THINFILM, Thin)
-
 // Accessors for the 10bits float weight used if is it detected being a grey scale
 #define BSDF_GETWEIGHT10F(X)			READ_BITS(X.State,  STATE_BIT_COUNT_WEIGHT10F, STATE_BIT_OFFSET_WEIGHT10F)
 #define BSDF_SETWEIGHT10F(X, Weight)	WRITE_BITS(X.State, STATE_BIT_COUNT_WEIGHT10F, STATE_BIT_OFFSET_WEIGHT10F, Weight)
@@ -412,11 +406,10 @@ struct FStrataBSDF
 
 ///////////////////////////////////////////////////////////////////////////////
 // Slab BSDF members
-#define SLAB_COMPLEX_PATH_MASK (TO_OFFSET_BITMASK(STATE_BIT_COUNT_ANISOTROPY, STATE_BIT_OFFSET_ANISOTROPY)		|  TO_OFFSET_BITMASK(STATE_BIT_COUNT_HASSSS, STATE_BIT_OFFSET_HASSSS) | \
-								TO_OFFSET_BITMASK(STATE_BIT_COUNT_HASSSSPROFILE, STATE_BIT_OFFSET_HASSSSPROFILE)|  TO_OFFSET_BITMASK(STATE_BIT_COUNT_HAZINESS, STATE_BIT_OFFSET_HAZINESS) | \
-								TO_OFFSET_BITMASK(STATE_BIT_COUNT_THINFILM, STATE_BIT_OFFSET_THINFILM)			|  TO_OFFSET_BITMASK(STATE_BIT_COUNT_F90, STATE_BIT_OFFSET_F90) | \
-								TO_OFFSET_BITMASK(STATE_BIT_COUNT_SIMPLEVOLUME, STATE_BIT_OFFSET_SIMPLEVOLUME)	|  TO_OFFSET_BITMASK(STATE_BIT_COUNT_HASFUZZ, STATE_BIT_OFFSET_HASFUZZ) | \
-								TO_OFFSET_BITMASK(STATE_BIT_COUNT_ISTHIN, STATE_BIT_OFFSET_ISTHIN) )
+#define SLAB_COMPLEX_PATH_MASK (TO_OFFSET_BITMASK(STATE_BIT_COUNT_ANISOTROPY, STATE_BIT_OFFSET_ANISOTROPY)		| TO_OFFSET_BITMASK(STATE_BIT_COUNT_HASSSS, STATE_BIT_OFFSET_HASSSS) | \
+								TO_OFFSET_BITMASK(STATE_BIT_COUNT_HASSSSPROFILE, STATE_BIT_OFFSET_HASSSSPROFILE)| TO_OFFSET_BITMASK(STATE_BIT_COUNT_HAZINESS, STATE_BIT_OFFSET_HAZINESS) | \
+								TO_OFFSET_BITMASK(STATE_BIT_COUNT_F90, STATE_BIT_OFFSET_F90)                    | TO_OFFSET_BITMASK(STATE_BIT_COUNT_SIMPLEVOLUME, STATE_BIT_OFFSET_SIMPLEVOLUME) | \
+								TO_OFFSET_BITMASK(STATE_BIT_COUNT_HASFUZZ, STATE_BIT_OFFSET_HASFUZZ)            | TO_OFFSET_BITMASK(STATE_BIT_COUNT_ISTHIN, STATE_BIT_OFFSET_ISTHIN) )
 
 #define SLAB_DIFFUSEALBEDO(X)			X.VGPRs[0].xyz
 #define SLAB_ROUGHNESS(X)				X.VGPRs[0].w
@@ -430,9 +423,7 @@ struct FStrataBSDF
 #define SLAB_SSSPROFILERADIUSSCALE(X)	X.VGPRs[3].y
 // Or explicit MFP in centimeter
 #define SLAB_SSSMFP(X)					X.VGPRs[3].xyz
-// Thin film. IOR is harcoded to oil IOR in order to save VGPR
-#define SLAB_THINFILM_THICKNESS(X)		X.VGPRs[3].w
-#define SLAB_THINFILM_IOR(X)			1.44f
+// Free slot:                           X.VGPRs[3].w
 // Fuzz to simulate cloth
 #define SLAB_FUZZ_COLOR(X)				X.VGPRs[4].xyz
 #define SLAB_FUZZ_AMOUNT(X)				X.VGPRs[4].w
@@ -763,7 +754,6 @@ struct FStrataOperator
 	// Thin film
 	float ThinFilmThickness;							// Thickness of thin-film coating layer. In micro-meters (1um == 1000nm).  Merge with regular ThicknessCm?
 	float ThinFilmIOR;									// IOR of the thin-film coating layer
-	bool  bThinFilmInline;								// If true the thin-film reflectance is evaluated during the base-pass and stored into F0/F90 to speed up lighting computation
 
 	// Topology information
 	bool				 bIsTop;
@@ -909,7 +899,6 @@ FStrataData GetStrataSlabBSDF(
 	}
 
 	const uint bHasHaziness = Haziness > 0.0f ? 1 : 0;
-	const uint bHasThinFilm = 0;
 	const uint bHasAnisotropy = Anisotropy != 0 ? 1 : 0;
 	const uint bHasFuzz = FuzzAmount > 0.0f ? 1 : 0;
 	if (bHasAnisotropy)
@@ -935,7 +924,6 @@ FStrataData GetStrataSlabBSDF(
 	BSDF_SETISTHIN			(StrataData.InlinedBSDF,		bIsThin ? 1 : 0);
 	BSDF_SETHASMFP			(StrataData.InlinedBSDF,		bMFPPlugged ? 1 : 0);
 	BSDF_SETHASHAZINESS		(StrataData.InlinedBSDF,		bHasHaziness);
-	BSDF_SETHASTHINFILM		(StrataData.InlinedBSDF,		bHasThinFilm);
 	BSDF_SETTHICKNESSCM		(StrataData.InlinedBSDF,		Thickness);
 	BSDF_SETHASFUZZ			(StrataData.InlinedBSDF,		bHasFuzz);
 
@@ -961,7 +949,6 @@ FStrataData GetStrataSlabBSDF(
 	SLAB_ROUGHNESS			(StrataData.InlinedBSDF)		= Roughness;
 	SLAB_ANISOTROPY			(StrataData.InlinedBSDF)		= Anisotropy;
 	SLAB_HAZINESS			(StrataData.InlinedBSDF)		= Haziness;
-	SLAB_THINFILM_THICKNESS	(StrataData.InlinedBSDF)		= 0.f;
 
 	if (bHasSSSProfile)
 	{
@@ -1340,7 +1327,7 @@ FStrataData StrataWeight(FStrataData A, float Weight, inout FStrataTree StrataTr
 	return StrataData;
 }
 
-FStrataData StrataThinFilm(FStrataData A, float Thickness, float IOR, bool bInlineEvaluation, inout FStrataTree StrataTree, int OperatorIndex, uint MaxDistanceFromLeaves)
+FStrataData StrataThinFilm(FStrataData A, float Thickness, float IOR, inout FStrataTree StrataTree, int OperatorIndex, uint MaxDistanceFromLeaves)
 {
 	const float SafeThickness = saturate(Thickness);
 	const float SafeIOR = clamp(IOR, 1.01f, 3.0f);
@@ -1355,7 +1342,6 @@ FStrataData StrataThinFilm(FStrataData A, float Thickness, float IOR, bool bInli
 	StrataTree.Operators[OperatorIndex].Type = STRATA_OPERATOR_THINFILM;
 	StrataTree.Operators[OperatorIndex].ThinFilmThickness = SafeThickness;
 	StrataTree.Operators[OperatorIndex].ThinFilmIOR = SafeIOR;
-	StrataTree.Operators[OperatorIndex].bThinFilmInline = bInlineEvaluation;
 	StrataTree.Operators[OperatorIndex].LeftIndex = A.OperatorIndex;
 
 	StrataTree.Operators[A.OperatorIndex].ParentIndex = OperatorIndex;
@@ -1472,7 +1458,6 @@ FStrataData StrataHorizontalMixingParameterBlending(FStrataData Background, FStr
 	BSDF_SETHASMFP(ResultBSDF,			BSDF_GETHASMFP(ForegroundBSDF)					| BSDF_GETHASMFP(BackgroundBSDF));
 	BSDF_SETHASHAZINESS(ResultBSDF,		BSDF_GETHASHAZINESS(ForegroundBSDF)				| BSDF_GETHASHAZINESS(BackgroundBSDF));
 	BSDF_SETHASF90(ResultBSDF,			BSDF_GETHASF90(ForegroundBSDF)					| BSDF_GETHASF90(BackgroundBSDF));
-	BSDF_SETHASTHINFILM(ResultBSDF,		BSDF_GETHASTHINFILM(ForegroundBSDF)				| BSDF_GETHASTHINFILM(BackgroundBSDF));
 	BSDF_SETHASFUZZ(ResultBSDF,			BSDF_GETHASFUZZ(ForegroundBSDF)					| BSDF_GETHASFUZZ(BackgroundBSDF));
 
 	BSDF_SETEMISSIVE(ResultBSDF,		lerp(BSDF_GETEMISSIVE(BackgroundBSDF),			BSDF_GETEMISSIVE(ForegroundBSDF),			ForegroundOtherMixFactor));
@@ -1485,7 +1470,6 @@ FStrataData StrataHorizontalMixingParameterBlending(FStrataData Background, FStr
 	SLAB_ROUGHNESS(ResultBSDF) =		lerp(SLAB_ROUGHNESS(BackgroundBSDF),			SLAB_ROUGHNESS(ForegroundBSDF),				ForegroundOtherMixFactor);
 	SLAB_ANISOTROPY(ResultBSDF) =		lerp(SLAB_ANISOTROPY(BackgroundBSDF),			SLAB_ANISOTROPY(ForegroundBSDF),			ForegroundOtherMixFactor);
 	SLAB_HAZINESS(ResultBSDF) =			lerp(SLAB_HAZINESS(BackgroundBSDF),				SLAB_HAZINESS(ForegroundBSDF),				ForegroundOtherMixFactor);
-	SLAB_THINFILM_THICKNESS(ResultBSDF)=lerp(SLAB_THINFILM_THICKNESS(BackgroundBSDF),	SLAB_THINFILM_THICKNESS(ForegroundBSDF),	ForegroundOtherMixFactor);
 	SLAB_FUZZ_AMOUNT(ResultBSDF) =		lerp(SLAB_FUZZ_AMOUNT(BackgroundBSDF),			SLAB_FUZZ_AMOUNT(ForegroundBSDF),			ForegroundOtherMixFactor);
 
 	if (BSDF_GETHASSSSPROFILE(ResultBSDF))
@@ -1568,7 +1552,6 @@ FStrataData StrataAddParameterBlending(FStrataData A, FStrataData B, float AMixF
 	BSDF_SETHASMFP(ResultBSDF,			BSDF_GETHASMFP(BBSDF)							| BSDF_GETHASMFP(ABSDF));
 	BSDF_SETHASHAZINESS(ResultBSDF,		BSDF_GETHASHAZINESS(BBSDF)						| BSDF_GETHASHAZINESS(ABSDF));
 	BSDF_SETHASF90(ResultBSDF,			BSDF_GETHASF90(BBSDF)							| BSDF_GETHASF90(ABSDF));
-	BSDF_SETHASTHINFILM(ResultBSDF,		BSDF_GETHASTHINFILM(BBSDF)						| BSDF_GETHASTHINFILM(ABSDF));
 	BSDF_SETHASFUZZ(ResultBSDF,			BSDF_GETHASFUZZ(BBSDF)							| BSDF_GETHASFUZZ(ABSDF));
 
 
@@ -1584,7 +1567,6 @@ FStrataData StrataAddParameterBlending(FStrataData A, FStrataData B, float AMixF
 	SLAB_ROUGHNESS(ResultBSDF) =		lerp(SLAB_ROUGHNESS(BBSDF),						SLAB_ROUGHNESS(ABSDF),						AMixFactor);
 	SLAB_ANISOTROPY(ResultBSDF) =		lerp(SLAB_ANISOTROPY(BBSDF),					SLAB_ANISOTROPY(ABSDF),						AMixFactor);
 	SLAB_HAZINESS(ResultBSDF) =			lerp(SLAB_HAZINESS(BBSDF),						SLAB_HAZINESS(ABSDF),						AMixFactor);
-	SLAB_THINFILM_THICKNESS(ResultBSDF) = 0;// This does not support coverage applied on pararmeters so we simply disable ThinFilm. lerp(SLAB_THINFILM_THICKNESS(BBSDF), SLAB_THINFILM_THICKNESS(ABSDF), AMixFactor);
 	SLAB_FUZZ_AMOUNT(ResultBSDF) =		lerp(SLAB_FUZZ_AMOUNT(BBSDF),					SLAB_FUZZ_AMOUNT(ABSDF),					AMixFactor);
 
 	if (BSDF_GETHASSSSPROFILE(ResultBSDF))
@@ -1659,7 +1641,6 @@ FStrataBSDF StrataVerticalLayeringParameterBlending(in FStrataBSDF TopBSDF, in F
 	BSDF_SETHASMFP(ResultBSDF,			BSDF_GETHASMFP(BaseBSDF)						| BSDF_GETHASMFP(TopBSDF));		// We keep MFP of both to make sure bottom layer on translucent can laos be effective
 	BSDF_SETHASHAZINESS(ResultBSDF,		BSDF_GETHASHAZINESS(TopBSDF));														// We only keep haziness if the top layer has it
 	BSDF_SETHASF90(ResultBSDF,			BSDF_GETHASF90(BaseBSDF)						| BSDF_GETHASF90(TopBSDF));			// F90: keep union of both even though it will be hard to get a match
-	BSDF_SETHASTHINFILM(ResultBSDF,		BSDF_GETHASTHINFILM(TopBSDF));														// We only keep thinfilme if the top layer has it
 	BSDF_SETHASFUZZ(ResultBSDF,			BSDF_GETHASFUZZ(BaseBSDF)						| BSDF_GETHASFUZZ(TopBSDF));		// Fuzz: keep union of both even though it will be hard to get a match
 
 	// Compute the top slab transmittance that will be used to blend the top and base material parameters.
@@ -1698,9 +1679,6 @@ FStrataBSDF StrataVerticalLayeringParameterBlending(in FStrataBSDF TopBSDF, in F
 	// Anisotropy is disabled with vertical blending. 
 	SLAB_ANISOTROPY(ResultBSDF) =		0.f;
 
-	// Only keep thin film for the top layer.
-	SLAB_THINFILM_THICKNESS(ResultBSDF)= (SafeTopCoverage >= 1.0f) ? SLAB_THINFILM_THICKNESS(TopBSDF) : 0.0f;
-
 	// Keep the MFP form the bottom layer only (top layer simple volume is backed in parameters).
 	SLAB_SSSMFP(ResultBSDF) =			lerp(SLAB_SSSMFP(TopBSDF),						SLAB_SSSMFP(BaseBSDF),		TopToBottomLerp); // SLAB_SSSMFP(BaseBSDF);
 
@@ -1786,7 +1764,6 @@ void SanitizeStrataSlab(inout FStrataBSDF BSDF)
 	SLAB_ROUGHNESS(BSDF) = SanitizeRoughness(SLAB_ROUGHNESS(BSDF));
 	SLAB_ANISOTROPY(BSDF) = clamp(SLAB_ANISOTROPY(BSDF), -1.f, 1.f);
 	SLAB_HAZINESS(BSDF) = saturate(SLAB_HAZINESS(BSDF));
-	SLAB_THINFILM_THICKNESS(BSDF) = saturate(SLAB_THINFILM_THICKNESS(BSDF));
 #if STRATA_INLINE_SHADING
 	BSDF.TmpMFP = clamp(BSDF.TmpMFP, 0.0f, Max111110BitsFloat3);
 #endif
@@ -1880,7 +1857,6 @@ void StrataDitherBSDF(inout FStrataBSDF BSDF, in float Dither)
 
 		SLAB_DIFFUSEALBEDO(BSDF) = DitherIfNonBlack(SLAB_DIFFUSEALBEDO(BSDF), Dither776);
 		SLAB_F0(BSDF) = DitherIfNonBlack(SLAB_F0(BSDF), Dither776);
-		SLAB_THINFILM_THICKNESS(BSDF) += BSDF_GETHASTHINFILM(BSDF) ? Dither8bits : 0.0f;
 	}
 }
 
@@ -2412,7 +2388,7 @@ void StrataLoad_ColorGamma2ToLinearAlphaLinear(FStrataMaterialContainer StrataBu
 ///////////////////////////////////////////////////////////////////////////////
 // Strata decal
 
-// Dbuffer support only a single layer, with minimal set of feature (i.e., no anisotropy, no thinfilm, ...);
+// Dbuffer support only a single layer, with minimal set of feature (i.e., no anisotropy, ...);
 // This special structure is for easing data processing
 struct FStrataDBuffer
 {
@@ -2910,7 +2886,7 @@ void PackStrataOut(
 						STRATA_STORE_UINT1(PackedData32Bits);
 						STRATA_STORE_UINT1(PackedData8Bits | PackRGBA8(float4(0.0f, SLAB_ROUGHNESS(BSDF), (SLAB_ANISOTROPY(BSDF) + 1.f) * 0.5f, 0.f)));
 
-						if (BSDF_GETHASF90(BSDF) || BSDF_GETHASTHINFILM(BSDF) || BSDF_GETHASHAZINESS(BSDF))
+						if (BSDF_GETHASF90(BSDF) || BSDF_GETHASHAZINESS(BSDF))
 						{
 							// What is important is to maintain the hue and saturation, so we scale the color by the maximum of its components
 							float3 F90 = saturate(SLAB_F90(BSDF));
@@ -2918,8 +2894,7 @@ void PackStrataOut(
 							F90 = Divisor > 0.0f ? F90 / Divisor : 1.0f;
 							float3 F90YCoCg = LinearRGB_2_NormalisedYCoCg(F90);
 
-							const float ThinFilmGamma2 = sqrt(SLAB_THINFILM_THICKNESS(BSDF)); // Using a gamma 2.0 encoding help with precision to be visually more uniform. This is then further improved via dithering.
-							STRATA_STORE_UINT1(PackRGBA8(float4(F90YCoCg.y, F90YCoCg.z, ThinFilmGamma2, SLAB_HAZINESS(BSDF))));
+							STRATA_STORE_UINT1(PackRGBA8(float4(F90YCoCg.y, F90YCoCg.z, 0, SLAB_HAZINESS(BSDF))));
 						}
 
 						const bool bIsSimpleVolume = BSDF_GETISSIMPLEVOLUME(BSDF);
@@ -3129,7 +3104,6 @@ FStrataBSDF UnpackFastPathStrataBSDFIn(FStrataMaterialContainer StrataBuffer, in
 	BSDF_SETHASSSS(OutBSDF, 0);
 	BSDF_SETHASSSSPROFILE(OutBSDF, 0);
 	BSDF_SETHASHAZINESS(OutBSDF, 0);
-	BSDF_SETHASTHINFILM(OutBSDF, 0);
 	BSDF_SETISSIMPLEVOLUME(OutBSDF, 0);
 	BSDF_SETHASFUZZ(OutBSDF, 0);
 
@@ -3190,9 +3164,8 @@ FStrataBSDF UnpackStrataBSDFIn(FStrataMaterialContainer StrataBuffer, inout FStr
 		UnpackFastPathSlabBSDF(OutBSDF, uint2(DataX, DataY));
 
 		SLAB_F90(OutBSDF) = 1.0f;
-		SLAB_THINFILM_THICKNESS(OutBSDF) = 0.0f;
 		const bool bHasF90 = BSDF_GETHASF90(OutBSDF);
-		if (bHasF90 || BSDF_GETHASTHINFILM(OutBSDF) || BSDF_GETHASHAZINESS(OutBSDF))
+		if (bHasF90 || BSDF_GETHASHAZINESS(OutBSDF))
 		{
 			float4 Data = UnpackRGBA8(StrataLoadUint1(StrataBuffer, StrataAddressing));
 
@@ -3208,7 +3181,6 @@ FStrataBSDF UnpackStrataBSDFIn(FStrataMaterialContainer StrataBuffer, inout FStr
 			F90 = Divisor > 0.0f ? F90 / Divisor : F90;
 			SLAB_F90(OutBSDF) = F90;
 
-			SLAB_THINFILM_THICKNESS(OutBSDF) = Data.z * Data.z;
 			SLAB_HAZINESS(OutBSDF) = Data.w;
 		}
 
diff --git a/Engine/Shaders/Private/Strata/StrataEvaluation.ush b/Engine/Shaders/Private/Strata/StrataEvaluation.ush
--- a/Engine/Shaders/Private/Strata/StrataEvaluation.ush
+++ b/Engine/Shaders/Private/Strata/StrataEvaluation.ush
@@ -17,7 +17,6 @@
 #include "../SHCommon.ush"
 #include "../ShadingModels.ush"
 #include "../ShadingModelsSampling.ush"
-#include "../ThinFilmBSDF.ush"
 #include "../AreaLightCommon.ush"
 
 #include "../HairStrands/HairStrandsCommon.ush"
@@ -327,7 +326,7 @@ FStrataEvaluateResult StrataEvaluateBSDFCommon(FStrataBSDFContext BSDFContext, F
 
 			////
 			//// Evaluate the specular component.
-			//// This takes into account multiple scattering, thin film interference, micro occlusion.
+			//// This takes into account multiple scattering, micro occlusion.
 			//// Note: anisotropy completely disables area integrations. Lights fall back to punctual.
 			////
 
@@ -434,23 +433,8 @@ FStrataEvaluateResult StrataEvaluateBSDFCommon(FStrataBSDFContext BSDFContext, F
 				}
 			}
 
-			float3 FresnelTerm = 0;
 			const float SpecularMicroOcclusion = F0RGBToMicroOcclusion(F0);
-			#if STRATA_FASTPATH==0
-			const bool bHasThinFilm = BSDF_GETHASTHINFILM(BSDFContext.BSDF);
-			BRANCH
-			if (bHasThinFilm)
-			{
-				float ThinFilmIOR = SLAB_THINFILM_IOR(BSDFContext.BSDF);
-				float ThinFilmThickness = SLAB_THINFILM_THICKNESS(BSDFContext.BSDF);
-				//  We do not use the MRP NoL here for thin interference.
-				FresnelTerm = SpecularMicroOcclusion * F_ThinFilm(BSDFContext.Context.NoV, BSDFContext.Context.NoL, BSDFContext.Context.VoH, F0, F90, ThinFilmIOR, ThinFilmThickness);
-			}
-			else
-			#endif
-			{
-				FresnelTerm = SpecularMicroOcclusion * F_Schlick(F0, F90, BSDFContext.Context.VoH);
-			}
+			const float3 FresnelTerm = SpecularMicroOcclusion * F_Schlick(F0, F90, BSDFContext.Context.VoH);
 
 			Sample.SpecularPathProbability		= bHaziness	? (1.0f - HazeWeight) : 1.0f;
 			Sample.SpecularHazePathProbability	= bHaziness	? HazeWeight : 0.0f;
@@ -460,9 +444,6 @@ FStrataEvaluateResult StrataEvaluateBSDFCommon(FStrataBSDFContext BSDFContext, F
 
 			if (bIsRectLight)
 			{
-				// Rect light should not evaluate the Fresnel term which is already part of the LTC integration function.
-				// So we simply disable the Fresnel factor in this case. That does mean that will lose thin film interference on Rect lights.
-
 				Sample.SpecularPathValue     = RectLightSpec     * MSScale     * SpecularMicroOcclusion;
 				Sample.SpecularHazePathValue = RectLightSpecHaze * HazeMSScale * SpecularMicroOcclusion;
 
@@ -929,7 +910,6 @@ FStrataEnvLightResult StrataEvaluateForEnvLight(FStrataBSDFContext BSDFContext,
 			StrataEnvLightResult.bSubsurface = BSDF_GETHASSSS(BSDFContext.BSDF) > 0;
 
 			float SafeRoughness = MakeRoughnessSafe(SLAB_ROUGHNESS(BSDFContext.BSDF));
-			float3 ThinFilmFresnelWeight = 1;
 
 			float3 EvalEnvBRDF = 0;
 			float3 DirectionalAlbedo_SpecularTransmission = 0;
@@ -947,19 +927,6 @@ FStrataEnvLightResult StrataEvaluateForEnvLight(FStrataBSDFContext BSDFContext,
 			{
 				const float SpecularMicroOcclusion = F0RGBToMicroOcclusion(F0);
 
-				#if STRATA_FASTPATH==0
-				const bool bHasThinFilm = BSDF_GETHASTHINFILM(BSDFContext.BSDF);
-				if (bHasThinFilm)
-				{
-					float ThinFilmIOR = SLAB_THINFILM_IOR(BSDFContext.BSDF);
-					float ThinFilmThickness = SLAB_THINFILM_THICKNESS(BSDFContext.BSDF);
-					ThinFilmFresnelWeight = F_ThinFilm(BSDFContext.Context.NoV, BSDFContext.Context.NoL, BSDFContext.Context.VoH, F0, F90, ThinFilmIOR, ThinFilmThickness);
-					F0 = 1.f;
-				}
-				#endif
-				EvalEnvBRDF *= ThinFilmFresnelWeight;
-
-
 				#if STRATA_FASTPATH==0
 				const bool bHasAnisotropy = BSDF_GETHASANISOTROPY(BSDFContext.BSDF);
 				if (bHasAnisotropy)
@@ -996,7 +963,7 @@ FStrataEnvLightResult StrataEvaluateForEnvLight(FStrataBSDFContext BSDFContext,
 
 					// Compute the directional albedo for hazy specular
 					FBxDFEnergyTerms EnergyTerms = ComputeGGXSpecEnergyTerms(SafeRoughnessHaze, BSDFContext.Context.NoV, F0, F90);
-					const float3 EvalEnvBRDFHaze = ThinFilmFresnelWeight * EnergyTerms.E;
+					const float3 EvalEnvBRDFHaze = EnergyTerms.E;
 					StrataEnvLightResult.SpecularHazeWeight = lerp(0.0f, HazeWeight, HazinessFadeIn) * EvalEnvBRDFHaze * SpecularMicroOcclusion;
 					StrataEnvLightResult.SpecularHazeSafeRoughness = SafeRoughnessHaze;
 				}
diff --git a/Engine/Shaders/Private/Strata/StrataTree.ush b/Engine/Shaders/Private/Strata/StrataTree.ush
--- a/Engine/Shaders/Private/Strata/StrataTree.ush
+++ b/Engine/Shaders/Private/Strata/StrataTree.ush
@@ -401,42 +401,31 @@ void PreUpdateAllBSDFWithBottomUpOperatorVisit_ThinFilm(
 
 	if (bCanReceiveThinFilm && BSDF_GETTYPE(CurrentBSDF) == STRATA_BSDF_TYPE_SLAB)
 	{
-		if (Op.bThinFilmInline)
-		{
-			const float3 F0  = SLAB_F0(CurrentBSDF);
-			const float3 F90 = SLAB_F90(CurrentBSDF);
+		const float3 F0  = SLAB_F0(CurrentBSDF);
+		const float3 F90 = SLAB_F90(CurrentBSDF);
 
-			const float3x3 TangentBasis = StrataGetBSDFSharedBasis_InlineShading(StrataPixelHeader, CurrentBSDF, NullStrataAddressing);
-			const float3 N = TangentBasis[2];
-			const float NoV = saturate(dot(N, V));
-			const float NoL = NoV;
-			const float VoH = NoV;
-			float3 RThinFilm = F_ThinFilm(NoV, NoL, VoH, F0, F90, Op.ThinFilmIOR, Op.ThinFilmThickness);
+		const float3x3 TangentBasis = StrataGetBSDFSharedBasis_InlineShading(StrataPixelHeader, CurrentBSDF, NullStrataAddressing);
+		const float3 N = TangentBasis[2];
+		const float NoV = saturate(dot(N, V));
+		const float NoL = NoV;
+		const float VoH = NoV;
+		float3 RThinFilm = F_ThinFilm(NoV, NoL, VoH, F0, F90, Op.ThinFilmIOR, Op.ThinFilmThickness);
 
-			// Compute a F0 which match evaluation Thin-Film reflectance for the current view angle using Schlick's Fresnel
-			const float FReference = min(0.99f, F_Schlick(0, 1.0f, NoV).x);
-			RThinFilm = max(RThinFilm, FReference);
+		// Compute a F0 which match evaluation Thin-Film reflectance for the current view angle using Schlick's Fresnel
+		const float FReference = min(0.99f, F_Schlick(0, 1.0f, NoV).x);
+		RThinFilm = max(RThinFilm, FReference);
 
-			float3 F0_Thin = (RThinFilm - F90 * FReference) / (1 - FReference);
+		float3 F0_Thin = (RThinFilm - F90 * FReference) / (1 - FReference);
 
-			// Anything less than 2% is physically impossible and is instead considered to be shadowing
-			// Ensure the generated F0 doesn't cause micro-oclusion
-			const float MicroOcclusionThreshold = 0.02f;
-			F0_Thin = max(MicroOcclusionThreshold, F0_Thin);
+		// Ensure the generated F0 doesn't cause micro-oclusion
+		F0_Thin = max(GetF0MicroOcclusionThreshold(), F0_Thin);
 
-			// If micro-occlusion was set, re-apply it.
-			const float3 MicroOcclusion = saturate(F0 * 50.f);
-			F0_Thin *= MicroOcclusion;
+		// If micro-occlusion was set, re-apply it.		
+		const float3 MicroOcclusion = F0RGBToMicroOcclusion(F0);
+		F0_Thin *= MicroOcclusion;
 
-			SLAB_F0(CurrentBSDF)  = F0_Thin;
-			SLAB_F90(CurrentBSDF) = max(F0_Thin, F90);
-		}
-		else
-		{
-			const uint bHasThinFilm = Op.ThinFilmThickness > 0 ? 1u : 0u;
-			BSDF_SETHASTHINFILM(CurrentBSDF, bHasThinFilm);
-			SLAB_THINFILM_THICKNESS(CurrentBSDF) += Op.ThinFilmThickness;
-		}
+		SLAB_F0(CurrentBSDF)  = F0_Thin;
+		SLAB_F90(CurrentBSDF) = max(F0_Thin, F90);
 	}
 
 #endif // STRATA_INLINE_SHADING
diff --git a/Engine/Shaders/Private/Strata/StrataVisualize.usf b/Engine/Shaders/Private/Strata/StrataVisualize.usf
--- a/Engine/Shaders/Private/Strata/StrataVisualize.usf
+++ b/Engine/Shaders/Private/Strata/StrataVisualize.usf
@@ -450,16 +450,6 @@ void VisualizeMaterialPS(
 					}
 				}
 
-				if (BSDF_GETHASTHINFILM(BSDF))
-				{
-					HORIZONTALTAB; PRINTPARAMCHAR(_T_); PRINTPARAMCHAR(_F_); PRINTPARAMCHAR(_SPC_); PRINTPARAMCHAR(_I_); PRINTPARAMCHAR(_O_); PRINTPARAMCHAR(_R_); PRINTPARAMCHAR(_SPC_);
-					HORIZONTALTAB; PRINTFLOAT1(SLAB_THINFILM_IOR(BSDF));
-					NEXTLINE;
-					HORIZONTALTAB; PRINTPARAMCHAR(_T_); PRINTPARAMCHAR(_F_); PRINTPARAMCHAR(_SPC_); PRINTPARAMCHAR(_T_); PRINTPARAMCHAR(_C_); PRINTPARAMCHAR(_K_); PRINTPARAMCHAR(_SPC_);
-					HORIZONTALTAB; PRINTFLOAT1(SLAB_THINFILM_THICKNESS(BSDF));
-					NEXTLINE;
-				}
-
 				if (BSDF_GETHASFUZZ(BSDF))
 				{
 					HORIZONTALTAB; PRINTPARAMCHAR(_F_); PRINTPARAMCHAR(_U_); PRINTPARAMCHAR(_Z_); PRINTPARAMCHAR(_Z_); PRINTPARAMCHAR(_SPC_); PRINTPARAMCHAR(_SPC_); PRINTPARAMCHAR(_SPC_);
@@ -505,12 +495,6 @@ void VisualizeMaterialPS(
 					}
 				}
 
-				if (BSDF_GETHASTHINFILM(BSDF))
-				{
-					NEXTLINE;
-					NEXTLINE;
-				}
-
 				if (BSDF_GETHASFUZZ(BSDF))
 				{
 					NEXTLINE;
diff --git a/Engine/Source/Runtime/Engine/Private/Materials/HLSLMaterialTranslator.cpp b/Engine/Source/Runtime/Engine/Private/Materials/HLSLMaterialTranslator.cpp
--- a/Engine/Source/Runtime/Engine/Private/Materials/HLSLMaterialTranslator.cpp
+++ b/Engine/Source/Runtime/Engine/Private/Materials/HLSLMaterialTranslator.cpp
@@ -10498,18 +10498,17 @@ int32 FHLSLMaterialTranslator::StrataWeightParameterBlending(int32 A, int32 Weig
 	);
 }
 
-int32 FHLSLMaterialTranslator::StrataThinFilm(int32 A, int32 Thickness, int32 IOR, bool bInlinedEvaluation, int OperatorIndex, uint32 MaxDistanceFromLeaves)
+int32 FHLSLMaterialTranslator::StrataThinFilm(int32 A, int32 Thickness, int32 IOR, int OperatorIndex, uint32 MaxDistanceFromLeaves)
 {
 	if (A == INDEX_NONE || Thickness == INDEX_NONE || IOR == INDEX_NONE)
 	{
 		return INDEX_NONE;
 	}
 	return AddCodeChunk(
-		MCT_Strata, TEXT("StrataThinFilm(%s, %s, %s, %d, Parameters.StrataTree, %u, %u)"),
+		MCT_Strata, TEXT("StrataThinFilm(%s, %s, %s, Parameters.StrataTree, %u, %u)"),
 		*GetParameterCode(A),
 		*GetParameterCode(Thickness),
 		*GetParameterCode(IOR),
-		bInlinedEvaluation ? 1u : 0u,
 		OperatorIndex,
 		MaxDistanceFromLeaves
 	);
diff --git a/Engine/Source/Runtime/Engine/Private/Materials/HLSLMaterialTranslator.h b/Engine/Source/Runtime/Engine/Private/Materials/HLSLMaterialTranslator.h
--- a/Engine/Source/Runtime/Engine/Private/Materials/HLSLMaterialTranslator.h
+++ b/Engine/Source/Runtime/Engine/Private/Materials/HLSLMaterialTranslator.h
@@ -1051,7 +1051,7 @@ protected:
 	virtual int32 StrataAddParameterBlending(int32 A, int32 B, int32 AMixWeight, const FString& SharedLocalBasisIndexMacro, FStrataOperator* PromoteToOperator) override;
 	virtual int32 StrataWeight(int32 A, int32 Weight, int OperatorIndex, uint32 MaxDistanceFromLeaves) override;
 	virtual int32 StrataWeightParameterBlending(int32 A, int32 Weight, FStrataOperator* PromoteToOperator) override;
-	virtual int32 StrataThinFilm(int32 A, int32 Thickness, int32 IOR, bool bInlinedEvaluation, int OperatorIndex, uint32 MaxDistanceFromLeaves) override;
+	virtual int32 StrataThinFilm(int32 A, int32 Thickness, int32 IOR, int OperatorIndex, uint32 MaxDistanceFromLeaves) override;
 	virtual int32 StrataTransmittanceToMFP(int32 TransmittanceColor, int32 DesiredThickness, int32 OutputIndex) override;
 	virtual int32 StrataMetalnessToDiffuseAlbedoF0(int32 BaseColor, int32 Specular, int32 Metallic, int32 OutputIndex) override;
 	
diff --git a/Engine/Source/Runtime/Engine/Private/Materials/MaterialExpressions.cpp b/Engine/Source/Runtime/Engine/Private/Materials/MaterialExpressions.cpp
--- a/Engine/Source/Runtime/Engine/Private/Materials/MaterialExpressions.cpp
+++ b/Engine/Source/Runtime/Engine/Private/Materials/MaterialExpressions.cpp
@@ -22514,7 +22514,7 @@ int32 UMaterialExpressionStrataThinFilm::Compile(class FMaterialCompiler* Compil
 	int32 OutputCodeChunk = INDEX_NONE;
 	FStrataOperator& StrataOperator = Compiler->StrataCompilationGetOperator(this);
 	{
-		OutputCodeChunk = Compiler->StrataThinFilm(ACodeChunk, ThicknessCodeChunk, IORCodeChunk, bUseParameterBlending, StrataOperator.Index, StrataOperator.MaxDistanceFromLeaves);
+		OutputCodeChunk = Compiler->StrataThinFilm(ACodeChunk, ThicknessCodeChunk, IORCodeChunk, StrataOperator.Index, StrataOperator.MaxDistanceFromLeaves);
 
 		if (!Compiler->StrataCompilationInfoContainsCodeChunk(ACodeChunk))
 		{
diff --git a/Engine/Source/Runtime/Engine/Public/MaterialCompiler.h b/Engine/Source/Runtime/Engine/Public/MaterialCompiler.h
--- a/Engine/Source/Runtime/Engine/Public/MaterialCompiler.h
+++ b/Engine/Source/Runtime/Engine/Public/MaterialCompiler.h
@@ -526,7 +526,7 @@ public:
 	virtual int32 StrataAddParameterBlending(int32 A, int32 B, int32 AMixWeight, const FString& SharedLocalBasisIndexMacro, FStrataOperator* PromoteToOperator) = 0;
 	virtual int32 StrataWeight(int32 A, int32 Weight, int OperatorIndex, uint32 MaxDistanceFromLeaves) = 0;
 	virtual int32 StrataWeightParameterBlending(int32 A, int32 Weight, FStrataOperator* PromoteToOperator) = 0;
-	virtual int32 StrataThinFilm(int32 A, int32 Thickness, int32 IOR, bool bInlinedEvaluation, int OperatorIndex, uint32 MaxDistanceFromLeaves) = 0;
+	virtual int32 StrataThinFilm(int32 A, int32 Thickness, int32 IOR, int OperatorIndex, uint32 MaxDistanceFromLeaves) = 0;
 	virtual int32 StrataTransmittanceToMFP(int32 TransmittanceColor, int32 DesiredThickness, int32 OutputIndex) = 0;
 	virtual int32 StrataMetalnessToDiffuseAlbedoF0(int32 BaseColor, int32 Specular, int32 Metallic, int32 OutputIndex) = 0;
 	
@@ -1143,9 +1143,9 @@ public:
 		return Compiler->StrataWeight(A, Weight, OperatorIndex, MaxDistanceFromLeaves);
 	}
 
-	virtual int32 StrataThinFilm(int32 A, int32 Thickness, int32 IOR, bool bInlinedEvaluation, int OperatorIndex, uint32 MaxDistanceFromLeaves) override
+	virtual int32 StrataThinFilm(int32 A, int32 Thickness, int32 IOR, int OperatorIndex, uint32 MaxDistanceFromLeaves) override
 	{
-		return Compiler->StrataThinFilm(A, Thickness, IOR, bInlinedEvaluation, OperatorIndex, MaxDistanceFromLeaves);
+		return Compiler->StrataThinFilm(A, Thickness, IOR, OperatorIndex, MaxDistanceFromLeaves);
 	}
 
 	virtual int32 StrataWeightParameterBlending(int32 A, int32 Weight, FStrataOperator* PromoteToOperator) override
