Improves the quality of r.TSR.History.HighFrequencyOnly=1

diff --git a/Engine/Shaders/Private/TemporalSuperResolution/TSRCommon.ush b/Engine/Shaders/Private/TemporalSuperResolution/TSRCommon.ush
--- a/Engine/Shaders/Private/TemporalSuperResolution/TSRCommon.ush
+++ b/Engine/Shaders/Private/TemporalSuperResolution/TSRCommon.ush
@@ -854,6 +854,17 @@ bool IsOffScreen(uint bCameraCut, float2 ScreenPos, tsr_half ParallaxRejectionMa
 	return (bIsCameraCut || bIsOutOfBounds || bIsParallaxRejected);
 }
 
+CALL_SITE_DEBUGLOC
+void IsOffScreenOrDisoccluded(uint bCameraCut, float2x2 ScreenPos, tsr_half2 ParallaxRejectionMask, out bool2 bIsOffScreen, out bool2 bIsDisoccluded)
+{
+	bool bIsCameraCut = bCameraCut != 0;
+	bool2 bIsOutOfBounds = max(abs(ScreenPos[0]), abs(ScreenPos[1])) >= 1.0;
+	bool2 bIsParallaxRejected = ParallaxRejectionMask < tsr_half(PARALLAX_REJECTION_MASK_THRESHOLD);
+
+	bIsOffScreen = bIsCameraCut || bIsOutOfBounds;
+	bIsDisoccluded = !bIsOffScreen && bIsParallaxRejected;
+}
+
 tsr_half MeasureBackbufferLDRQuantizationError()
 {
 	// Assume the backbuffer is 10bit per channels
diff --git a/Engine/Shaders/Private/TemporalSuperResolution/TSRUpdateHistory.usf b/Engine/Shaders/Private/TemporalSuperResolution/TSRUpdateHistory.usf
--- a/Engine/Shaders/Private/TemporalSuperResolution/TSRUpdateHistory.usf
+++ b/Engine/Shaders/Private/TemporalSuperResolution/TSRUpdateHistory.usf
@@ -281,7 +281,6 @@ void MainCS(
 
 	// Fetch reprojections related informations.
 	float2x2 PrevScreenPos = ScreenPos;
-	tsr_half2 ParallaxRejectionMask = tsr_half(1.0).xx;
 	tsr_half2 LowFrequencyRejection = tsr_half(1.0).xx;
 	tsr_half2 TranslucencyRejection = tsr_half(1.0).xx;
 	tsr_half2 OutputPixelVelocity = tsr_half(0.0).xx;
@@ -289,6 +288,8 @@ void MainCS(
 	tsr_half2x2 HoleFilledVelocityMask = dpv_interleave_mono_registers(tsr_half(1.0).xx);
 	tsr_subpixel_payload2 ParallaxFactorBits = tsr_subpixel_payload(0x0).xx;
 	bool2 bIsResponsiveAAPixel = false;
+	bool2 bIsOffScreen = false;
+	bool2 bIsDisoccluded = false;
 
 	PLATFORM_SPECIFIC_ISOLATE
 	{
@@ -430,9 +431,16 @@ void MainCS(
 				PrevScreenPos = ScreenPos - ScreenVelocity;
 				OutputPixelVelocity = tsr_half2(dpv_length(dpv_mul(ScreenVelocity, HistoryInfo_ViewportSize)));
 			}
+			
+			// Detect whether the history reprojection is successful
+			IsOffScreenOrDisoccluded(
+				bCameraCut,
+				PrevScreenPos,
+				dpv_interleave_registers_array(RawParallaxRejectionMask),
+				/* out */ bIsOffScreen,
+				/* out */ bIsDisoccluded);
 
 			// Process input texel, forcing to pack register.
-			ParallaxRejectionMask = dpv_force_interleave_registers_array(RawParallaxRejectionMask);
 			ParallaxFactorBits    = dpv_force_interleave_registers_array(RawParallaxFactorBits) & tsr_subpixel_payload(SUB_PIXEL_PARALLAX_FACTOR_BIT_MASK);
 			LowFrequencyRejection = dpv_force_interleave_registers_array(RawLowFrequencyRejection);
 			//TranslucencyRejection = dpv_force_interleave_registers_array(RawTranslucencyRejection);
@@ -451,13 +459,7 @@ void MainCS(
 		#if CONFIG_REJECTION_ANTI_ALIASING
 		{
 			// Detect if HistoryBufferUV would be outside of the viewport.
-			bool2 bOffScreen = bool2(
-				IsOffScreen(bCameraCut, dpv_lo(PrevScreenPos), dpv_lo(ParallaxRejectionMask)),
-				IsOffScreen(bCameraCut, dpv_hi(PrevScreenPos), dpv_hi(ParallaxRejectionMask)));
-			
-			tsr_half2 TotalRejection = tsr_half2(
-				bOffScreen[0] ? 0.0 : saturate(LowFrequencyRejection[0] * 4.0),
-				bOffScreen[1] ? 0.0 : saturate(LowFrequencyRejection[1] * 4.0));
+			tsr_half2 TotalRejection = select(bIsOffScreen || bIsDisoccluded, tsr_half(0.0).xx, saturate(LowFrequencyRejection * tsr_half(4.0)));
 
 			tsr_ushort2 EncodedInputTexelOffset = dpv_force_interleave_registers_array(RawEncodedInputTexelOffset);
 			
@@ -506,16 +508,6 @@ void MainCS(
 		LowFrequencyRejection = min(LowFrequencyRejection, TranslucencyRejection);
 	#endif
 	
-	// Detect if HistoryBufferUV would be outside of the viewport.
-	bool2 bOffScreen = bool2(
-		IsOffScreen(bCameraCut, dpv_lo(PrevScreenPos), dpv_lo(ParallaxRejectionMask)),
-		IsOffScreen(bCameraCut, dpv_hi(PrevScreenPos), dpv_hi(ParallaxRejectionMask)));
-	
-	tsr_half2 TotalRejection = tsr_half2(
-		bOffScreen[0] ? 0.0 : saturate(LowFrequencyRejection[0] * 4.0),
-		bOffScreen[1] ? 0.0 : saturate(LowFrequencyRejection[1] * 4.0));
-
-
 	// Filter input scene color at predictor frequency.
 	tsr_halfCx2 FilteredInputColor;
 	tsr_halfCx2 InputMinColor;
@@ -560,19 +552,7 @@ void MainCS(
 			// Soften the spatial kernel when there is high noise in the input to make parrallax disocclusion of detailed geometry less distracting.
 			tsr_half2 OffScreenInputToHistoryFactor = tsr_half(1.0) - tsr_half(0.5) * NoiseFiltering;
 
-			tsr_half2 KernelInputToHistoryFactor = tsr_half(InputToHistoryFactor).xx;
-			
-			FLATTEN
-			if (bOffScreen.x)
-			{
-				KernelInputToHistoryFactor.x = tsr_half(OffScreenInputToHistoryFactor.x);
-			}
-			
-			FLATTEN
-			if (bOffScreen.y)
-			{
-				KernelInputToHistoryFactor.y = tsr_half(OffScreenInputToHistoryFactor.y);
-			}
+			tsr_half2 KernelInputToHistoryFactor = select(bIsOffScreen || bIsDisoccluded, OffScreenInputToHistoryFactor, tsr_half(InputToHistoryFactor));
 
 			UNROLL_N(CONFIG_SAMPLES)
 			for (uint SampleId = 0; SampleId < CONFIG_SAMPLES; SampleId++)
@@ -1027,48 +1007,113 @@ void MainCS(
 
 	tsr_halfCx2 FinalHighFrequencyColor;
 	tsr_half2 FinalHistoryValidity;
+	#if CONFIG_HIGH_FREQUENCY_ONLY
 	{
-		#if CONFIG_HIGH_FREQUENCY_ONLY
-			#if CONFIG_HIGH_FREQUENCY_YCOCG
-				tsr_halfCx2 ClampedPrevHighFrequencyColor = YCoCgToRGB(clamp(TransformColorForClampingBox(PrevHighFrequencyColor), dpv_scale(InputMinColor, PrevHistoryWeight), dpv_scale(InputMaxColor, PrevHistoryWeight)));
-			#else
-				tsr_halfCx2 ClampedPrevHighFrequencyColor = clamp(PrevHighFrequencyColor, dpv_scale(InputMinColor, PrevHistoryWeight), dpv_scale(InputMaxColor, PrevHistoryWeight));
-			#endif
+		#if CONFIG_HIGH_FREQUENCY_YCOCG
+			tsr_halfCx2 ClampedPrevHighFrequencyColor = YCoCgToRGB(clamp(TransformColorForClampingBox(PrevHighFrequencyColor), dpv_scale(InputMinColor, PrevHistoryWeight), dpv_scale(InputMaxColor, PrevHistoryWeight)));
+		#else
+			tsr_halfCx2 ClampedPrevHighFrequencyColor = clamp(PrevHighFrequencyColor, dpv_scale(InputMinColor, PrevHistoryWeight), dpv_scale(InputMaxColor, PrevHistoryWeight));
+		#endif
+		
+		tsr_half2 TotalRejection = select(bIsOffScreen || bIsDisoccluded, tsr_half(0.0).xx, saturate(LowFrequencyRejection * tsr_half(4.0)));
 
-			tsr_half2 BlendClamp = TotalRejection * HoleFilledVelocityMask[1] * HighFrequencyRejection;
+		tsr_half2 BlendClamp = TotalRejection * HoleFilledVelocityMask[1] * HighFrequencyRejection;
 
-			PrevHighFrequencyColor[0] = lerp(ClampedPrevHighFrequencyColor[0], PrevHighFrequencyColor[0], BlendClamp);
-			PrevHighFrequencyColor[1] = lerp(ClampedPrevHighFrequencyColor[1], PrevHighFrequencyColor[1], BlendClamp);
-			PrevHighFrequencyColor[2] = lerp(ClampedPrevHighFrequencyColor[2], PrevHighFrequencyColor[2], BlendClamp);
+		PrevHighFrequencyColor[0] = lerp(ClampedPrevHighFrequencyColor[0], PrevHighFrequencyColor[0], BlendClamp);
+		PrevHighFrequencyColor[1] = lerp(ClampedPrevHighFrequencyColor[1], PrevHighFrequencyColor[1], BlendClamp);
+		PrevHighFrequencyColor[2] = lerp(ClampedPrevHighFrequencyColor[2], PrevHighFrequencyColor[2], BlendClamp);
 
-			tsr_half2 InvPrevHistoryWeight = SafeRcp(PrevHistoryWeight);
-			tsr_half2 InvPrevHistoryValidity = SafeRcp(PrevHistoryValidity);
+		tsr_half2 InvPrevHistoryWeight = SafeRcp(PrevHistoryWeight);
+		tsr_half2 InvPrevHistoryValidity = SafeRcp(PrevHistoryValidity);
 
-			// Compute how much the history needs to be rejected, based on the completeness.
-			tsr_half2 PrevHistoryRejectionWeight = LowFrequencyRejection; // * HoleFilledVelocityMask[1];
-		#else // !CONFIG_HIGH_FREQUENCY_ONLY
-			tsr_half2 InvPrevHistoryValidity = SafeRcp(PrevHistoryValidity);
-			tsr_half2 InvPrevHistoryWeight = InvPrevHistoryValidity;
+		// Compute how much the history needs to be rejected, based on the completeness.
+		tsr_half2 PrevHistoryRejectionWeight = select(bIsOffScreen, tsr_half(0.0), LowFrequencyRejection);
+	
+		tsr_half2 DesiredCurrentContribution = max(Histeresis * InputPixelAlignement, tsr_half(0.0));
 
-			// Compute how much the history needs to be rejected, based on the completeness.
-			tsr_half2 PrevHistoryRejectionWeight = LowFrequencyRejection * HoleFilledVelocityMask[1];
+		// Determine whether the the prediction based rejection was confident enough.
+		tsr_half2 RejectionConfidentEnough = tsr_half(1); // saturate(RejectionValidity * MAX_SAMPLE_COUNT - 3.0);
+		
+		// Clamp the validity due to motion to maintain better sharpness in history reprojection under motion.
+		tsr_half2 ClampedPrevHistoryValidity;
+		#if 1
+		{
+			tsr_half2 MaxValidity = tsr_half(1.0) - (tsr_half(1.0) - PredictionOnlyValidity) * saturate(OutputPixelVelocity * tsr_half(InvWeightClampingPixelSpeed));
+
+			// Clamp up the max validity to favorise stability under motion on high contrast edges.
+			#if 1
+			{
+				tsr_half2 PrevHistoryLuma = Luma4(PrevHighFrequencyColor) * InvPrevHistoryWeight;
+				tsr_half2 FilteredLuma = Luma4(FilteredInputColor);
+
+				tsr_half2 MinValidityForStability = abs(FilteredLuma - PrevHistoryLuma) / max(FilteredLuma, PrevHistoryLuma);
+				
+				MaxValidity = max(MaxValidity, MinValidityForStability);
+			}
+			#endif
+			
+			ClampedPrevHistoryValidity = min(PrevHistoryValidity, MaxValidity);
+		}
+		#else
+		{
+			ClampedPrevHistoryValidity = PrevHistoryValidity;
+		}
 		#endif
+
+		// Compute the newly rejected validity
+		tsr_half2 RejectedValidity = (
+			min(ClampedPrevHistoryValidity, PredictionOnlyValidity - DesiredCurrentContribution) +
+			max(ClampedPrevHistoryValidity - PredictionOnlyValidity + DesiredCurrentContribution, tsr_half(0.0)) * PrevHistoryRejectionWeight);
+
+		RejectedValidity = ClampedPrevHistoryValidity * PrevHistoryRejectionWeight;
+
+		// Compute the maximum output validity.
+		tsr_half2 OutputValidity = (
+			clamp(RejectedValidity + DesiredCurrentContribution, tsr_half(0.0), PredictionOnlyValidity) +
+			clamp(RejectedValidity + DesiredCurrentContribution * PrevHistoryRejectionWeight * RejectionConfidentEnough - PredictionOnlyValidity, tsr_half(0.0), tsr_half(1.0) - PredictionOnlyValidity));
+
+		tsr_half2 OutputValidity2 = select(bIsOffScreen, max(MIN_IMPLICIT_HISTORY_WEIGHT, OutputValidity), OutputValidity);
+		DesiredCurrentContribution = select(bIsOffScreen, max(MIN_IMPLICIT_HISTORY_WEIGHT, DesiredCurrentContribution), DesiredCurrentContribution);
+
+		tsr_half2 PrevMomentWeight = max(OutputValidity2 - DesiredCurrentContribution, tsr_half(0.0));
+		tsr_half2 CurrentMomentWeight = min(DesiredCurrentContribution, OutputValidity2);
+
+		tsr_half2 PrevHistoryToneWeight = HdrWeightY(Luma4(PrevHighFrequencyColor) * InvPrevHistoryWeight);
+		tsr_half2 FilteredInputToneWeight = HdrWeight4(FilteredInputColor);
+			
+		tsr_half2 BlendPrevHistory = PrevMomentWeight * PrevHistoryToneWeight;
+		tsr_half2 BlendFilteredInput = CurrentMomentWeight * FilteredInputToneWeight;
+
+		tsr_half2 CommonWeight = SafeRcp(BlendPrevHistory + BlendFilteredInput);
+
+		FinalHighFrequencyColor = (
+			dpv_scale(PrevHighFrequencyColor, CommonWeight * BlendPrevHistory * InvPrevHistoryWeight) +
+			dpv_scale(FilteredInputColor, CommonWeight * BlendFilteredInput));
+				
+				
+		//Debug[0] = CommonWeight * BlendFilteredInput;
+		Debug[0] = OutputValidity2;
+		//Debug[0] = CurrentMomentWeight;
+		//Debug[0] = FilteredInputColor[0];
+		//Debug[1] = FilteredInputColor[1];
+		//Debug[2] = FilteredInputColor[2];
+
+		// Quantize validity for the 8bit encoding to avoid numerical shift between color and validity.
+		FinalHistoryValidity = ceil(tsr_half(255.0) * OutputValidity) * rcp(tsr_half(255.0));
+	}
+	#else // !CONFIG_HIGH_FREQUENCY_ONLY
+	{
+		tsr_half2 InvPrevHistoryValidity = SafeRcp(PrevHistoryValidity);
+		tsr_half2 InvPrevHistoryWeight = InvPrevHistoryValidity;
+
+		// Compute how much the history needs to be rejected, based on the completeness.
+		tsr_half2 PrevHistoryRejectionWeight = LowFrequencyRejection * HoleFilledVelocityMask[1];
 	
 		//#if !CONFIG_CLAMP
 		//	PrevHistoryRejectionWeight = 1.0;
 		//#endif
-			
-		FLATTEN
-		if (bOffScreen[0])
-		{
-			PrevHistoryRejectionWeight[0] = tsr_half(0.0);
-		}
-
-		FLATTEN
-		if (bOffScreen[1])
-		{
-			PrevHistoryRejectionWeight[1] = tsr_half(0.0);
-		}
+		
+		PrevHistoryRejectionWeight = select(bIsOffScreen || bIsDisoccluded, tsr_half(0.0).xx, PrevHistoryRejectionWeight);
 
 		tsr_half2 DesiredCurrentContribution = max(Histeresis * InputPixelAlignement, tsr_half(0.0));
 
@@ -1127,14 +1172,8 @@ void MainCS(
 			}
 		#endif
 		
-		#if CONFIG_HIGH_FREQUENCY_ONLY
-			tsr_half2 OutputValidity2 = max(OutputValidity, MIN_IMPLICIT_HISTORY_WEIGHT);
-		#else
-			tsr_half2 OutputValidity2 = OutputValidity;
-		#endif
-
-		tsr_half2 PrevMomentWeight = max(OutputValidity2 - DesiredCurrentContribution, tsr_half(0.0));
-		tsr_half2 CurrentMomentWeight = min(DesiredCurrentContribution, OutputValidity2);
+		tsr_half2 PrevMomentWeight = max(OutputValidity - DesiredCurrentContribution, tsr_half(0.0));
+		tsr_half2 CurrentMomentWeight = min(DesiredCurrentContribution, OutputValidity);
 
 		#if 1
 		{
@@ -1170,6 +1209,7 @@ void MainCS(
 		}
 		#endif
 	}
+	#endif // !CONFIG_HIGH_FREQUENCY_ONLY
 
 	// Compute the low frequency history
 	#if !CONFIG_HIGH_FREQUENCY_ONLY
@@ -1259,6 +1299,8 @@ void MainCS(
 
 			// Completely discard the history
 			{
+				tsr_half2 TotalRejection = select(bIsOffScreen || bIsDisoccluded, tsr_half(0.0).xx, saturate(LowFrequencyRejection * tsr_half(4.0)));
+
 				tsr_half2 TotalRejectionBlending = max(BlendFinal, HoleFilledVelocityMask[0]);
 
 				PrevLowFrequencySampleCount *= TotalRejection;
@@ -1389,17 +1431,7 @@ void MainCS(
 				0x0);
 			
 			// Reset parallax factor
-			FLATTEN
-			if (bOffScreen.x)
-			{
-				FinalSubpixelDetails.x = FinalSubpixelDetails.x & ~ParallaxFactorMask;
-			}
-			
-			FLATTEN
-			if (bOffScreen.y)
-			{
-				FinalSubpixelDetails.y = FinalSubpixelDetails.y & ~ParallaxFactorMask;
-			}
+			FinalSubpixelDetails = select(bIsOffScreen || bIsDisoccluded, FinalSubpixelDetails & ~ParallaxFactorMask, FinalSubpixelDetails);
 		}
 		
 		uint2 SubpixelId = dpv_dot(uint2x2(bKOSign), uint2(1, SUB_PIXEL_GRID_SIZE));
