Add thin film paramter blending support.

diff --git a/Engine/Shaders/Private/Strata/Strata.ush b/Engine/Shaders/Private/Strata/Strata.ush
--- a/Engine/Shaders/Private/Strata/Strata.ush
+++ b/Engine/Shaders/Private/Strata/Strata.ush
@@ -1723,7 +1723,6 @@ FStrataData StrataVerticalLayeringParameterBlending(FStrataData Top, FStrataData
 	return Result;
 }
 
-
 FStrataData StrataWeightParameterBlending(FStrataData A, float Weight)
 {
 	const float SafeWeight = saturate(Weight);
@@ -1737,6 +1736,39 @@ FStrataData StrataWeightParameterBlending(FStrataData A, float Weight)
 	return StrataData;
 }
 
+void StrataGetThinFilmF0F90(float NoV, float ThinFilmThickness, float ThinFilmIOR, inout float3 OutF0, inout float3 OutF90)
+{
+	const float3 F0 = OutF0;
+	const float3 F90 = OutF90;
+
+	const float NoL = NoV;
+	const float VoH = NoV;
+	float3 RThinFilm = F_ThinFilm(NoV, NoL, VoH, F0, F90, ThinFilmIOR, ThinFilmThickness);
+
+	// Compute a F0 which match evaluation Thin-Film reflectance for the current view angle using Schlick's Fresnel
+	const float FReference = min(0.99f, F_Schlick(0, 1.0f, NoV).x);
+	RThinFilm = max(RThinFilm, FReference);
+
+	float3 F0_Thin = (RThinFilm - F90 * FReference) / (1 - FReference);
+
+	// Ensure the generated F0 doesn't cause micro-oclusion
+	F0_Thin = max(GetF0MicroOcclusionThreshold(), F0_Thin);
+
+	// If micro-occlusion was set, re-apply it.		
+	const float3 MicroOcclusion = F0RGBToMicroOcclusion(F0);
+	F0_Thin *= MicroOcclusion;
+
+	OutF0 = F0_Thin;
+	OutF90 = max(F0_Thin, F90);
+}
+
+FStrataData StrataThinFilmParameterBlending(in FStrataData A, float ThinFilmThickness, float ThinFilmIOR, float NoV)
+{
+	FStrataData StrataData = A;
+	StrataGetThinFilmF0F90(NoV, ThinFilmThickness, ThinFilmIOR, SLAB_F0(StrataData.InlinedBSDF), SLAB_F90(StrataData.InlinedBSDF));
+	return StrataData;
+}
+
 ///////////////////////////////////////////////////////////////////////////////
 // Functions used to sanitize BSDF before they are used for lighting (when forward) or stored to memory (base pass).
 // We do not normalize normals, this is done before the lighting step.
diff --git a/Engine/Shaders/Private/Strata/StrataTree.ush b/Engine/Shaders/Private/Strata/StrataTree.ush
--- a/Engine/Shaders/Private/Strata/StrataTree.ush
+++ b/Engine/Shaders/Private/Strata/StrataTree.ush
@@ -401,31 +401,10 @@ void PreUpdateAllBSDFWithBottomUpOperatorVisit_ThinFilm(
 
 	if (bCanReceiveThinFilm && BSDF_GETTYPE(CurrentBSDF) == STRATA_BSDF_TYPE_SLAB)
 	{
-		const float3 F0  = SLAB_F0(CurrentBSDF);
-		const float3 F90 = SLAB_F90(CurrentBSDF);
-
 		const float3x3 TangentBasis = StrataGetBSDFSharedBasis_InlineShading(StrataPixelHeader, CurrentBSDF, NullStrataAddressing);
 		const float3 N = TangentBasis[2];
 		const float NoV = saturate(dot(N, V));
-		const float NoL = NoV;
-		const float VoH = NoV;
-		float3 RThinFilm = F_ThinFilm(NoV, NoL, VoH, F0, F90, Op.ThinFilmIOR, Op.ThinFilmThickness);
-
-		// Compute a F0 which match evaluation Thin-Film reflectance for the current view angle using Schlick's Fresnel
-		const float FReference = min(0.99f, F_Schlick(0, 1.0f, NoV).x);
-		RThinFilm = max(RThinFilm, FReference);
-
-		float3 F0_Thin = (RThinFilm - F90 * FReference) / (1 - FReference);
-
-		// Ensure the generated F0 doesn't cause micro-oclusion
-		F0_Thin = max(GetF0MicroOcclusionThreshold(), F0_Thin);
-
-		// If micro-occlusion was set, re-apply it.		
-		const float3 MicroOcclusion = F0RGBToMicroOcclusion(F0);
-		F0_Thin *= MicroOcclusion;
-
-		SLAB_F0(CurrentBSDF)  = F0_Thin;
-		SLAB_F90(CurrentBSDF) = max(F0_Thin, F90);
+		StrataGetThinFilmF0F90(NoV, Op.ThinFilmThickness, Op.ThinFilmIOR, SLAB_F0(CurrentBSDF), SLAB_F90(CurrentBSDF));
 	}
 
 #endif // STRATA_INLINE_SHADING
diff --git a/Engine/Source/Runtime/Engine/Classes/Materials/MaterialExpressionStrata.h b/Engine/Source/Runtime/Engine/Classes/Materials/MaterialExpressionStrata.h
--- a/Engine/Source/Runtime/Engine/Classes/Materials/MaterialExpressionStrata.h
+++ b/Engine/Source/Runtime/Engine/Classes/Materials/MaterialExpressionStrata.h
@@ -789,9 +789,6 @@ class UMaterialExpressionStrataThinFilm : public UMaterialExpressionStrataBSDF
 	UPROPERTY()
 	FExpressionInput IOR;
 
-	UPROPERTY(EditAnywhere, Category = Mode)
-	uint32 bUseParameterBlending : 1;
-
 	//~ Begin UMaterialExpression Interface
 #if WITH_EDITOR
 	virtual int32 Compile(class FMaterialCompiler* Compiler, int32 OutputIndex) override;
diff --git a/Engine/Source/Runtime/Engine/Private/Materials/HLSLMaterialTranslator.cpp b/Engine/Source/Runtime/Engine/Private/Materials/HLSLMaterialTranslator.cpp
--- a/Engine/Source/Runtime/Engine/Private/Materials/HLSLMaterialTranslator.cpp
+++ b/Engine/Source/Runtime/Engine/Private/Materials/HLSLMaterialTranslator.cpp
@@ -10512,7 +10512,41 @@ int32 FHLSLMaterialTranslator::StrataThinFilm(int32 A, int32 Thickness, int32 IO
 		OperatorIndex,
 		MaxDistanceFromLeaves
 	);
+}
+
+int32 FHLSLMaterialTranslator::StrataThinFilmParameterBlending(int32 A, int32 Thickness, int32 IOR, int32 BSDFNormalCodeChunk, FStrataOperator* PromoteToOperator)
+{
+	if (A == INDEX_NONE || Thickness == INDEX_NONE || IOR == INDEX_NONE)
+	{
+		return INDEX_NONE;
+	}
+
+	if (PromoteToOperator)
+	{
+		check(PromoteToOperator->Index != INDEX_NONE);
+		check(PromoteToOperator->BSDFIndex != INDEX_NONE);
+		return AddCodeChunk(
+			MCT_Strata, TEXT("PromoteParameterBlendedBSDFToOperator(StrataThinFilmParameterBlending(%s, %s, %s, dot(%s, %s)), Parameters.StrataTree, %u, %u, %u, %u)"),
+			*GetParameterCode(A),
+			*GetParameterCode(Thickness),
+			*GetParameterCode(IOR),
+			*GetParameterCode(BSDFNormalCodeChunk),
+			*GetParameterCode(CameraVector()),
+			PromoteToOperator->Index,
+			PromoteToOperator->BSDFIndex,
+			PromoteToOperator->bIsBottom,
+			PromoteToOperator->bIsTop
+		);
+	}
 
+	return AddCodeChunk(
+		MCT_Strata, TEXT("StrataThinFilmParameterBlending(%s, %s, %s, dot(%s, %s))"),
+		*GetParameterCode(A),
+		*GetParameterCode(Thickness),
+		*GetParameterCode(IOR),
+		*GetParameterCode(BSDFNormalCodeChunk),
+		*GetParameterCode(CameraVector())
+	);
 }
 
 int32 FHLSLMaterialTranslator::StrataTransmittanceToMFP(int32 TransmittanceColor, int32 DesiredThickness, int32 OutputIndex)
diff --git a/Engine/Source/Runtime/Engine/Private/Materials/HLSLMaterialTranslator.h b/Engine/Source/Runtime/Engine/Private/Materials/HLSLMaterialTranslator.h
--- a/Engine/Source/Runtime/Engine/Private/Materials/HLSLMaterialTranslator.h
+++ b/Engine/Source/Runtime/Engine/Private/Materials/HLSLMaterialTranslator.h
@@ -1052,6 +1052,7 @@ protected:
 	virtual int32 StrataWeight(int32 A, int32 Weight, int OperatorIndex, uint32 MaxDistanceFromLeaves) override;
 	virtual int32 StrataWeightParameterBlending(int32 A, int32 Weight, FStrataOperator* PromoteToOperator) override;
 	virtual int32 StrataThinFilm(int32 A, int32 Thickness, int32 IOR, int OperatorIndex, uint32 MaxDistanceFromLeaves) override;
+	virtual int32 StrataThinFilmParameterBlending(int32 A, int32 Thickness, int32 IOR, int32 BSDFNormalCodeChunk, FStrataOperator* PromoteToOperator) override;
 	virtual int32 StrataTransmittanceToMFP(int32 TransmittanceColor, int32 DesiredThickness, int32 OutputIndex) override;
 	virtual int32 StrataMetalnessToDiffuseAlbedoF0(int32 BaseColor, int32 Specular, int32 Metallic, int32 OutputIndex) override;
 	
diff --git a/Engine/Source/Runtime/Engine/Private/Materials/MaterialExpressions.cpp b/Engine/Source/Runtime/Engine/Private/Materials/MaterialExpressions.cpp
--- a/Engine/Source/Runtime/Engine/Private/Materials/MaterialExpressions.cpp
+++ b/Engine/Source/Runtime/Engine/Private/Materials/MaterialExpressions.cpp
@@ -22513,6 +22513,19 @@ int32 UMaterialExpressionStrataThinFilm::Compile(class FMaterialCompiler* Compil
 
 	int32 OutputCodeChunk = INDEX_NONE;
 	FStrataOperator& StrataOperator = Compiler->StrataCompilationGetOperator(this);
+	if (StrataOperator.bUseParameterBlending)
+	{
+		// Extra normal
+		const FStrataMaterialCompilationInfo& AStrataData = Compiler->GetStrataCompilationInfo(ACodeChunk);
+		const FStrataMaterialCompilationInfo::FBSDF& ABSDF = AStrataData.Layers[0].BSDFs[0];
+		int32 NormalCodeChunk = ABSDF.RegisteredSharedLocalBasis.NormalCodeChunk;
+
+		OutputCodeChunk = Compiler->StrataThinFilmParameterBlending(ACodeChunk, ThicknessCodeChunk, IORCodeChunk, NormalCodeChunk, StrataOperator.bRootOfParameterBlendingSubTree ? &StrataOperator : nullptr);
+
+		FStrataMaterialCompilationInfo StrataInfo = StrataCompilationInfoThinFilmParamBlend(Compiler, Compiler->GetStrataCompilationInfo(ACodeChunk));
+		Compiler->StrataCompilationInfoRegisterCodeChunk(OutputCodeChunk, StrataInfo);
+	}
+	else
 	{
 		OutputCodeChunk = Compiler->StrataThinFilm(ACodeChunk, ThicknessCodeChunk, IORCodeChunk, StrataOperator.Index, StrataOperator.MaxDistanceFromLeaves);
 
diff --git a/Engine/Source/Runtime/Engine/Private/Materials/StrataMaterial.cpp b/Engine/Source/Runtime/Engine/Private/Materials/StrataMaterial.cpp
--- a/Engine/Source/Runtime/Engine/Private/Materials/StrataMaterial.cpp
+++ b/Engine/Source/Runtime/Engine/Private/Materials/StrataMaterial.cpp
@@ -93,9 +93,7 @@ void  StrataCompilationInfoCreateNullBSDF(FMaterialCompiler* Compiler, int32 Cod
 
 FStrataMaterialCompilationInfo StrataCompilationInfoWeight(FMaterialCompiler* Compiler, const FStrataMaterialCompilationInfo& A)
 {
-	// STRATA_TODO: This is not correct. FStrataMaterialCompilationInfo is meant to be removed and replaced with Strata tree / metadata
 	FStrataMaterialCompilationInfo StrataInfo = A;
-	StrataInfo.Layers[0].BSDFs[0].bHasThinFilm = true;
 	return StrataInfo;
 }
 
@@ -106,6 +104,14 @@ FStrataMaterialCompilationInfo StrataCompilationInfoWeightParamBlend(FMaterialCo
 }
 
 FStrataMaterialCompilationInfo StrataCompilationInfoThinFilm(FMaterialCompiler* Compiler, const FStrataMaterialCompilationInfo& A)
+{
+	// STRATA_TODO: This is not correct. FStrataMaterialCompilationInfo is meant to be removed and replaced with Strata tree / metadata
+	FStrataMaterialCompilationInfo StrataInfo = A;
+	StrataInfo.Layers[0].BSDFs[0].bHasThinFilm = true;
+	return StrataInfo;
+}
+
+FStrataMaterialCompilationInfo StrataCompilationInfoThinFilmParamBlend(FMaterialCompiler* Compiler, const FStrataMaterialCompilationInfo& A)
 {
 	FStrataMaterialCompilationInfo StrataInfo = A;
 	return StrataInfo;
diff --git a/Engine/Source/Runtime/Engine/Private/Materials/StrataMaterial.h b/Engine/Source/Runtime/Engine/Private/Materials/StrataMaterial.h
--- a/Engine/Source/Runtime/Engine/Private/Materials/StrataMaterial.h
+++ b/Engine/Source/Runtime/Engine/Private/Materials/StrataMaterial.h
@@ -71,6 +71,7 @@ void  StrataCompilationInfoCreateSingleBSDFMaterial(FMaterialCompiler* Compiler,
 FStrataMaterialCompilationInfo StrataCompilationInfoWeight(FMaterialCompiler* Compiler, const FStrataMaterialCompilationInfo& A);
 FStrataMaterialCompilationInfo StrataCompilationInfoWeightParamBlend(FMaterialCompiler* Compiler, const FStrataMaterialCompilationInfo& A);
 FStrataMaterialCompilationInfo StrataCompilationInfoThinFilm(FMaterialCompiler* Compiler, const FStrataMaterialCompilationInfo& A);
+FStrataMaterialCompilationInfo StrataCompilationInfoThinFilmParamBlend(FMaterialCompiler* Compiler, const FStrataMaterialCompilationInfo& A);
 FStrataMaterialCompilationInfo StrataCompilationInfoAdd(FMaterialCompiler* Compiler, const FStrataMaterialCompilationInfo& A, const FStrataMaterialCompilationInfo& B);
 FStrataMaterialCompilationInfo StrataCompilationInfoAddParamBlend(FMaterialCompiler* Compiler, const FStrataMaterialCompilationInfo& A, const FStrataMaterialCompilationInfo& B, const FStrataRegisteredSharedLocalBasis& RegisteredSharedLocalBasis);
 FStrataMaterialCompilationInfo StrataCompilationInfoHorizontalMixing(FMaterialCompiler* Compiler, const FStrataMaterialCompilationInfo& Background, const FStrataMaterialCompilationInfo& Foreground);
diff --git a/Engine/Source/Runtime/Engine/Public/MaterialCompiler.h b/Engine/Source/Runtime/Engine/Public/MaterialCompiler.h
--- a/Engine/Source/Runtime/Engine/Public/MaterialCompiler.h
+++ b/Engine/Source/Runtime/Engine/Public/MaterialCompiler.h
@@ -527,6 +527,7 @@ public:
 	virtual int32 StrataWeight(int32 A, int32 Weight, int OperatorIndex, uint32 MaxDistanceFromLeaves) = 0;
 	virtual int32 StrataWeightParameterBlending(int32 A, int32 Weight, FStrataOperator* PromoteToOperator) = 0;
 	virtual int32 StrataThinFilm(int32 A, int32 Thickness, int32 IOR, int OperatorIndex, uint32 MaxDistanceFromLeaves) = 0;
+	virtual int32 StrataThinFilmParameterBlending(int32 A, int32 Thickness, int32 IOR, int32 BSDFNormalCodeChunk, FStrataOperator* PromoteToOperator) = 0;
 	virtual int32 StrataTransmittanceToMFP(int32 TransmittanceColor, int32 DesiredThickness, int32 OutputIndex) = 0;
 	virtual int32 StrataMetalnessToDiffuseAlbedoF0(int32 BaseColor, int32 Specular, int32 Metallic, int32 OutputIndex) = 0;
 	
@@ -1148,6 +1149,11 @@ public:
 		return Compiler->StrataThinFilm(A, Thickness, IOR, OperatorIndex, MaxDistanceFromLeaves);
 	}
 
+	virtual int32 StrataThinFilmParameterBlending(int32 A, int32 Thickness, int32 IOR, int32 BSDFNormalCodeChunk, FStrataOperator* PromoteToOperator) override
+	{
+		return Compiler->StrataThinFilmParameterBlending(A, Thickness, IOR, BSDFNormalCodeChunk, PromoteToOperator);
+	}
+
 	virtual int32 StrataWeightParameterBlending(int32 A, int32 Weight, FStrataOperator* PromoteToOperator) override
 	{
 		return Compiler->StrataWeightParameterBlending(A, Weight, PromoteToOperator);
