Support applying local bounds scale and center translation to ray tracing instance transform.
- Useful for procedural raytracing geometries.

diff --git a/Engine/Shaders/Private/RayTracing/RayTracingInstanceBufferUtil.usf b/Engine/Shaders/Private/RayTracing/RayTracingInstanceBufferUtil.usf
--- a/Engine/Shaders/Private/RayTracing/RayTracingInstanceBufferUtil.usf
+++ b/Engine/Shaders/Private/RayTracing/RayTracingInstanceBufferUtil.usf
@@ -20,6 +20,7 @@ struct FRayTracingInstanceDescriptorInput
 	uint InstanceId;
 	uint InstanceMaskAndFlags;
 	uint InstanceContributionToHitGroupIndex;
+	uint bApplyLocalBoundsTransform;
 };
 
 uint NumInstances;
@@ -55,19 +56,29 @@ void RayTracingBuildInstanceBufferCS(uint3 DispatchThreadId : SV_DispatchThreadI
 		FInstanceSceneData InstanceSceneData = GetInstanceSceneData(InputDesc.GPUSceneInstanceOrTransformIndex, InstanceSceneDataSOAStride);
 
 		float4x4 LocalToTranslatedWorld4x4 = LWCMultiplyTranslation(InstanceSceneData.LocalToWorld, PreViewTranslation);
+
+		if (InstanceSceneData.Flags & INSTANCE_SCENE_DATA_FLAG_IS_RAYTRACING_FAR_FIELD)
+		{
+			LocalToTranslatedWorld4x4[3].xyz += FarFieldReferencePos;
+		}
+
+		if (InputDesc.bApplyLocalBoundsTransform)
+		{
+			float4x4 ProcT;
+			ProcT[0] = float4(InstanceSceneData.LocalBoundsExtent.x * 2, 0, 0, 0);
+			ProcT[1] = float4(0, InstanceSceneData.LocalBoundsExtent.y * 2, 0, 0);
+			ProcT[2] = float4(0, 0, InstanceSceneData.LocalBoundsExtent.z * 2, 0);
+			ProcT[3] = float4(InstanceSceneData.LocalBoundsCenter, 1);
+
+			LocalToTranslatedWorld4x4 = mul(ProcT, LocalToTranslatedWorld4x4);
+		}
+
 		LocalToTranslatedWorld4x4 = transpose(LocalToTranslatedWorld4x4);
 
 		float3x4 LocalToTranslatedWorld;
 		LocalToTranslatedWorld[0] = LocalToTranslatedWorld4x4[0];
 		LocalToTranslatedWorld[1] = LocalToTranslatedWorld4x4[1];
 		LocalToTranslatedWorld[2] = LocalToTranslatedWorld4x4[2];
-
-		if (InstanceSceneData.Flags & INSTANCE_SCENE_DATA_FLAG_IS_RAYTRACING_FAR_FIELD)
-		{
-			LocalToTranslatedWorld[0].w += FarFieldReferencePos.x;
-			LocalToTranslatedWorld[1].w += FarFieldReferencePos.y;
-			LocalToTranslatedWorld[2].w += FarFieldReferencePos.z;
-		}
 #else
 		float3x4 LocalToTranslatedWorld;
 		LocalToTranslatedWorld[0] = InstanceTransforms[InputDesc.GPUSceneInstanceOrTransformIndex * 3 + 0];
diff --git a/Engine/Source/Runtime/Engine/Public/RayTracingInstance.h b/Engine/Source/Runtime/Engine/Public/RayTracingInstance.h
--- a/Engine/Source/Runtime/Engine/Public/RayTracingInstance.h
+++ b/Engine/Source/Runtime/Engine/Public/RayTracingInstance.h
@@ -64,6 +64,9 @@ struct FRayTracingInstance
 	/** Whether ray hits should be registered for front and back faces. */
 	bool bDoubleSided = false;
 
+	/** Whether local bounds scale and center translation should be applied to the instance transform. */
+	bool bApplyLocalBoundsTransform = false;
+
 	/** Instance mask that can be used to exclude the instance from specific effects (eg. ray traced shadows). */
 	uint8 Mask = 0xFF;
 
diff --git a/Engine/Source/Runtime/RHI/Public/RHIResources.h b/Engine/Source/Runtime/RHI/Public/RHIResources.h
--- a/Engine/Source/Runtime/RHI/Public/RHIResources.h
+++ b/Engine/Source/Runtime/RHI/Public/RHIResources.h
@@ -2324,6 +2324,9 @@ struct FRayTracingGeometryInstance
 	// Each geometry copy can have one bit to make it individually deactivated (removed from TLAS while maintaining hit group indexing). Useful for culling.
 	TArrayView<const uint32> ActivationMask;
 
+	// Whether local bounds scale and center translation should be applied to the instance transform.
+	bool bApplyLocalBoundsTransform = false;
+
 	// Mask that will be tested against one provided to TraceRay() in shader code.
 	// If binary AND of instance mask with ray mask is zero, then the instance is considered not intersected / invisible.
 	uint8 Mask = 0xFF;
diff --git a/Engine/Source/Runtime/Renderer/Private/DeferredShadingRenderer.cpp b/Engine/Source/Runtime/Renderer/Private/DeferredShadingRenderer.cpp
--- a/Engine/Source/Runtime/Renderer/Private/DeferredShadingRenderer.cpp
+++ b/Engine/Source/Runtime/Renderer/Private/DeferredShadingRenderer.cpp
@@ -1008,6 +1008,7 @@ bool FDeferredShadingSceneRenderer::GatherRayTracingWorldInstancesForView(FRDGBu
 
 					RayTracingInstance.DefaultUserData = PrimitiveIndex;
 					RayTracingInstance.Mask = Instance.Mask;
+					RayTracingInstance.bApplyLocalBoundsTransform = Instance.bApplyLocalBoundsTransform;
 					if (Instance.bForceOpaque)
 					{
 						RayTracingInstance.Flags |= ERayTracingInstanceFlags::ForceOpaque;
diff --git a/Engine/Source/Runtime/Renderer/Private/PrimitiveSceneInfo.cpp b/Engine/Source/Runtime/Renderer/Private/PrimitiveSceneInfo.cpp
--- a/Engine/Source/Runtime/Renderer/Private/PrimitiveSceneInfo.cpp
+++ b/Engine/Source/Runtime/Renderer/Private/PrimitiveSceneInfo.cpp
@@ -945,6 +945,8 @@ void FPrimitiveSceneInfo::UpdateCachedRayTracingInstance(FPrimitiveSceneInfo* Sc
 		SceneInfo->CachedRayTracingInstance.DefaultUserData = (uint32)SceneInfo->GetIndex();
 		SceneInfo->CachedRayTracingInstance.Mask = CachedRayTracingInstance.Mask; // When no cached command is found, InstanceMask == 0 and the instance is effectively filtered out
 
+		SceneInfo->CachedRayTracingInstance.bApplyLocalBoundsTransform = CachedRayTracingInstance.bApplyLocalBoundsTransform;
+
 		if (CachedRayTracingInstance.bForceOpaque)
 		{
 			SceneInfo->CachedRayTracingInstance.Flags |= ERayTracingInstanceFlags::ForceOpaque;
diff --git a/Engine/Source/Runtime/Renderer/Private/RayTracing/RayTracingInstanceBufferUtil.cpp b/Engine/Source/Runtime/Renderer/Private/RayTracing/RayTracingInstanceBufferUtil.cpp
--- a/Engine/Source/Runtime/Renderer/Private/RayTracing/RayTracingInstanceBufferUtil.cpp
+++ b/Engine/Source/Runtime/Renderer/Private/RayTracing/RayTracingInstanceBufferUtil.cpp
@@ -205,6 +205,7 @@ void FillRayTracingInstanceUploadBuffer(
 				}
 
 				InstanceDesc.AccelerationStructureIndex = AccelerationStructureIndex;
+				InstanceDesc.bApplyLocalBoundsTransform = SceneInstance.bApplyLocalBoundsTransform;
 
 				if (!SceneInstance.ActivationMask.IsEmpty() && (SceneInstance.ActivationMask[TransformIndex / 32] & (1 << (TransformIndex % 32))) == 0)
 				{
diff --git a/Engine/Source/Runtime/Renderer/Public/RayTracingInstanceBufferUtil.h b/Engine/Source/Runtime/Renderer/Public/RayTracingInstanceBufferUtil.h
--- a/Engine/Source/Runtime/Renderer/Public/RayTracingInstanceBufferUtil.h
+++ b/Engine/Source/Runtime/Renderer/Public/RayTracingInstanceBufferUtil.h
@@ -30,6 +30,7 @@ struct FRayTracingInstanceDescriptorInput
 	uint32 InstanceId;
 	uint32 InstanceMaskAndFlags;
 	uint32 InstanceContributionToHitGroupIndex;
+	uint32 bApplyLocalBoundsTransform;
 };
 
 struct FRayTracingGPUInstance
