Changes FTSRClearPrevTexturesCS to clear more pixels with similar VGPR pressure

diff --git a/Engine/Shaders/Private/TemporalSuperResolution/TSRClearPrevTextures.usf b/Engine/Shaders/Private/TemporalSuperResolution/TSRClearPrevTextures.usf
--- a/Engine/Shaders/Private/TemporalSuperResolution/TSRClearPrevTextures.usf
+++ b/Engine/Shaders/Private/TemporalSuperResolution/TSRClearPrevTextures.usf
@@ -7,6 +7,8 @@
 
 #define TILE_SIZE 8
 
+#define TILE_ROWS 2
+
 
 //------------------------------------------------------- PARAMETERS
 
@@ -21,15 +23,22 @@ void MainCS(
 	uint2 GroupId : SV_GroupID,
 	uint GroupThreadIndex : SV_GroupIndex)
 {
-	uint2 DispatchThreadId = (
-		ZOrder2D(GroupThreadIndex, uint(log2(float(TILE_SIZE)))) +
-		GroupId * uint2(TILE_SIZE, TILE_SIZE));
+	uint2 GroupThreadId = uint2(GroupThreadIndex % TILE_SIZE, GroupThreadIndex / TILE_SIZE);
+	uint2 OutputPixelOrigin = (InputInfo_ViewportMin + GroupId * uint(TILE_SIZE * TILE_ROWS)) + GroupThreadId;
 
-	uint2 InputPixelPos = InputInfo_ViewportMin + DispatchThreadId;
-	
-	if (all(InputPixelPos < InputInfo_ViewportMax))
+	UNROLL_N(TILE_ROWS)
+	for (uint y = 0; y < TILE_ROWS; y++)
 	{
-		PrevUseCountOutput[InputPixelPos] = 0;
-		PrevClosestDepthOutput[InputPixelPos] = 0;
+		UNROLL_N(TILE_ROWS)
+		for (uint x = 0; x < TILE_ROWS; x++)
+		{
+			uint2 PixelOffset = uint2(x, y) * TILE_SIZE;
+			uint2 InputPixelPos = OutputPixelOrigin + PixelOffset;
+			
+			InputPixelPos.x = select(all(InputPixelPos < InputInfo_ViewportMax), InputPixelPos.x, uint(~0));
+
+			PrevUseCountOutput[InputPixelPos] = 0;
+			PrevClosestDepthOutput[InputPixelPos] = 0;
+		}
 	}
 }
diff --git a/Engine/Source/Runtime/Renderer/Private/PostProcess/TemporalSuperResolution.cpp b/Engine/Source/Runtime/Renderer/Private/PostProcess/TemporalSuperResolution.cpp
--- a/Engine/Source/Runtime/Renderer/Private/PostProcess/TemporalSuperResolution.cpp
+++ b/Engine/Source/Runtime/Renderer/Private/PostProcess/TemporalSuperResolution.cpp
@@ -767,7 +767,7 @@ ITemporalUpscaler::FOutputs AddTemporalSuperResolutionPasses(
 			RDG_EVENT_NAME("TSR ClearPrevTextures %dx%d", InputRect.Width(), InputRect.Height()),
 			ComputeShader,
 			PassParameters,
-			FComputeShaderUtils::GetGroupCount(InputRect.Size(), 8));
+			FComputeShaderUtils::GetGroupCount(InputRect.Size(), 8 * 2));
 	}
 
 	// Dilate the velocity texture & scatter reprojection into previous frame
