Strata - fixed issues with placeholder variance for rough refractions not propagating to tiles.

diff --git a/Engine/Shaders/Private/BasePassPixelShader.usf b/Engine/Shaders/Private/BasePassPixelShader.usf
--- a/Engine/Shaders/Private/BasePassPixelShader.usf
+++ b/Engine/Shaders/Private/BasePassPixelShader.usf
@@ -1711,9 +1711,6 @@ void FPixelShaderInOut_MainPS(
 	FStrataSubsurfaceData SSSData = (FStrataSubsurfaceData)0;
 	FStrataTopLayerData TopLayerData = (FStrataTopLayerData)0;
 	FStrataOpaqueRoughRefractionData OpaqueRoughRefractionData = (FStrataOpaqueRoughRefractionData)0;
-#if STRATA_OUTPUT_ROUGH_REFRACTION
-	OpaqueRoughRefractionData.OpaqueRoughRefractionEnabled = true;
-#endif
 
 	// Generate the strata material data to write out
 	FStrataAddressing StrataAddressing = GetStrataPixelDataByteOffset(PixelPos, uint2(ResolvedView.BufferSizeAndInvSize.xy), StrataStruct.MaxBytesPerPixel);
diff --git a/Engine/Shaders/Private/Strata/Strata.ush b/Engine/Shaders/Private/Strata/Strata.ush
--- a/Engine/Shaders/Private/Strata/Strata.ush
+++ b/Engine/Shaders/Private/Strata/Strata.ush
@@ -572,16 +572,16 @@ struct FStrataOpaqueRoughRefractionData
 
 float StrataPackOpaqueRoughRefractionData(FStrataOpaqueRoughRefractionData In)
 {
-	float Value = In.OpaqueRoughRefractionEnabled && In.Variance > 0.0f ? 128.0f : 0.0f;
-	Value += In.OpaqueRoughRefractionEnabled ? In.Variance * 127.0f : 0.0f;
+	float Value = In.Variance > 0.0f ? 128.0f : 0.0f;
+	Value += In.Variance * 127.0f;
 	return Value / 255.0f;
 }
 
 FStrataOpaqueRoughRefractionData StrataUnpackOpaqueRoughRefractionData(float In)
 {
 	FStrataOpaqueRoughRefractionData Out = (FStrataOpaqueRoughRefractionData)0;
-	Out.OpaqueRoughRefractionEnabled = In >= 128.0f;
-	Out.Variance = (In) * (1.0f / 255.0f) - (128.0f / 255.0f); 
+	Out.OpaqueRoughRefractionEnabled = (In * 255.0) >= 128.0f;
+	Out.Variance = In * (127.0f / 255.0f) + (- 128.0f / 255.0f);
 	return Out;
 }
 
@@ -2718,7 +2718,7 @@ void PackStrataOut(
 						TopLayerTotalWeight += BSDF.TopLayerDataWeight;
 						TopLayerData.Roughness += BSDF.TopLayerDataWeight * SLAB_ROUGHNESS(BSDF);
 
-						OpaqueRoughRefractionData.Variance += BSDF.TopLayerDataWeight * StrataLobeRoughnessToVariance(SLAB_ROUGHNESS(BSDF));
+						OpaqueRoughRefractionData.Variance += BSDF.TopLayerDataWeight * SLAB_ROUGHNESS(BSDF);
 
 						const bool bIsThin = BSDF_GETISTHIN(BSDF);
 						const bool bIsSimpleVolume = BSDF_GETISSIMPLEVOLUME(BSDF);
@@ -2737,7 +2737,7 @@ void PackStrataOut(
 						TopLayerTotalWeight += BSDF.TopLayerDataWeight;
 						TopLayerData.Roughness += BSDF.TopLayerDataWeight * FullyRough;
 
-						OpaqueRoughRefractionData.Variance += BSDF.TopLayerDataWeight * StrataLobeRoughnessToVariance(SLAB_ROUGHNESS(BSDF));
+						OpaqueRoughRefractionData.Variance += BSDF.TopLayerDataWeight * SLAB_ROUGHNESS(BSDF);
 
 						break;
 					}
@@ -2750,7 +2750,7 @@ void PackStrataOut(
 						TopLayerTotalWeight += BSDF.TopLayerDataWeight;
 						TopLayerData.Roughness += BSDF.TopLayerDataWeight * SLW_ROUGHNESS(BSDF);
 
-						OpaqueRoughRefractionData.Variance += BSDF.TopLayerDataWeight * StrataLobeRoughnessToVariance(SLAB_ROUGHNESS(BSDF));
+						OpaqueRoughRefractionData.Variance += BSDF.TopLayerDataWeight * SLAB_ROUGHNESS(BSDF);
 
 						break;
 					}
@@ -2779,7 +2779,10 @@ void PackStrataOut(
 		TopLayerData.WorldNormal= TopLayerTotalWeight > 0.0f	? normalize(TopLayerData.WorldNormal)			: 0.0f;
 		TopLayerData.Roughness	= TopLayerTotalWeight > 0.0f	? TopLayerData.Roughness / TopLayerTotalWeight	: 0.0f;
 
+#if STRATA_OUTPUT_ROUGH_REFRACTION
 		OpaqueRoughRefractionData.Variance = TopLayerTotalWeight > 0.0f ? OpaqueRoughRefractionData.Variance / TopLayerTotalWeight : 0.0f;
+		OpaqueRoughRefractionData.OpaqueRoughRefractionEnabled = OpaqueRoughRefractionData.Variance > 0.0f;
+#endif
 
 		// Set storage layout as either: fast(0), single(1), or complex(2)
 		AndIsSimpleMaterial(StrataPixelHeader, BSDFVisibleCount == 1);
