Added eta as parameter of rough refraction rnd.

diff --git a/Engine/Shaders/Private/Strata/StrataRoughRefraction.usf b/Engine/Shaders/Private/Strata/StrataRoughRefraction.usf
--- a/Engine/Shaders/Private/Strata/StrataRoughRefraction.usf
+++ b/Engine/Shaders/Private/Strata/StrataRoughRefraction.usf
@@ -64,7 +64,7 @@ void OpaqueRoughRefractionPS(
 
 #include "/Engine/Private/Strata/StrataStatisticalOperators.ush"
 
-#define MODE_REFLECTION 1
+#define MODE_REFLECTION 0
 
 // Updated from http://jcgt.org/published/0007/03/04/
 bool Slabs(float3 P0, float3 P1, float3 RayOrigin, float3 InvRaydir, out float OutTMin, out float OutTMax)
@@ -100,7 +100,7 @@ float GetEta12()
 #define WOLRD_NORMAL float3(0, 0, 1)
 
 // returns true the direction is valid for the current mode (refraction or reflection)
-bool GetGGXDirection(float2 RandomUV, float Roughness, float3 V, inout float3 L, inout float Value)
+bool GetGGXDirection(float2 RandomUV, float Roughness, float3 V, inout float3 L, inout float Value, float Eta12)
 { 
 	float2 Alpha; 
 	GetGGXBasis(Roughness, WOLRD_NORMAL, Alpha);
@@ -111,12 +111,11 @@ bool GetGGXDirection(float2 RandomUV, float Roughness, float3 V, inout float3 L,
 	L = 0;
 	float F = 0;
 	const float RandSample = MODE_REFLECTION ? -0.01f : 1.0f;
-	const bool bIsRefraction = SampleRefraction(-V, H.xyz, GetEta12(), RandSample, L, F);
-
-	if (bIsRefraction)
+	const bool bIsRefraction = SampleRefraction(-V, H.xyz, Eta12, RandSample, L, F);
+	const bool bTotalInternalReflection = F == 1.0f;
+	if (bIsRefraction && !bTotalInternalReflection)
 	{
-		// Refraction
-		Value = GGXEvalRefraction(L, V, H.xyz, Alpha, GetEta12()).x;
+		Value = GGXEvalRefraction(L, V, H.xyz, Alpha, Eta12).x;
 		return MODE_REFLECTION == 0;
 	}
 
@@ -134,6 +133,7 @@ struct FUI
 	float IncidentAngle;
 	float SlabThickness;
 	float Roughness;
+	float Eta12;
 
 	// Derived
 	float a2;
@@ -150,10 +150,25 @@ FUI GetUI(inout FShaderPrintContext Context)
 	FUI UI;
 
 	UI.IncidentAngle = AddSlider(Context, TEXT("IncidentAngle"), 0.0f, GetDefaultFontColor(), 0.0f, 0.5f * PI);
+
 	Newline(Context);
-	UI.SlabThickness = AddSlider(Context, TEXT("Slab Thickness"), 5.0f, GetDefaultFontColor(), 0.0f, 20.0f);
+
 	Newline(Context);
+
 	UI.Roughness = AddSlider(Context, TEXT("Roughness"), 0.5f, GetDefaultFontColor(), 0.0f, 1.0f);
+
+	Newline(Context);
+
+	Newline(Context);
+
+	UI.Eta12 = AddSlider(Context, TEXT("Eta12"), 1.3f, GetDefaultFontColor(), 0.5f, 1.5f);
+
+	Newline(Context);
+
+	Newline(Context);
+
+	UI.SlabThickness = AddSlider(Context, TEXT("Slab Thickness"), 5.0f, GetDefaultFontColor(), 0.0f, 20.0f);
+
 	Newline(Context);
 
 	UI.a2 = Pow4(UI.Roughness);
@@ -244,7 +259,7 @@ void EvaluateRoughRefractionLobeCS(uint3 ThreadId : SV_DispatchThreadID)
 				float2 RandomUV = float2(u, v) / float(TraceSqrtSampleCount);
 				float3 L = 0;
 				float Value = 0;
-				if (GetGGXDirection(RandomUV, UI.Roughness, UI.V, L, Value) && Value > 0.0f)
+				if (GetGGXDirection(RandomUV, UI.Roughness, UI.V, L, Value, UI.Eta12) && Value > 0.0f)
 				{
 					float2 SampleUV = L.xy * 0.5 + 0.5;
 					SampleCountTextureUAV[uint2(SampleUV * 64.0f)] += 1;
@@ -266,7 +281,7 @@ void EvaluateRoughRefractionLobeCS(uint3 ThreadId : SV_DispatchThreadID)
 				float2 RandomUV = float2(u, v) / float(TraceSqrtSampleCount);
 				float3 L = 0;
 				float3 Value = 0;
-				if (GetGGXDirection(RandomUV, UI.Roughness, UI.V, L, Value) && Value > 0.0f)
+				if (GetGGXDirection(RandomUV, UI.Roughness, UI.V, L, Value, UI.Eta12) && Value > 0.0f)
 				{
 					float2 SampleUV = L.xy * 0.5 + 0.5;
 					SampleCountTextureUAV[uint2(SampleUV * 64.0f)] += 1;
@@ -409,7 +424,7 @@ void VisualizeRoughRefractionPS(
 			float2 RandomUV = float2(u, v) / TraceDomainSize;
 			float3 L = 0;
 			float Value = 0;
-			if (GetGGXDirection(RandomUV, UI.Roughness, UI.V, L, Value))
+			if (GetGGXDirection(RandomUV, UI.Roughness, UI.V, L, Value, UI.Eta12))
 			{
 				if(Value > 0.0f) AddLineWS(Context, RefractionP, RefractionP + L * VisProjectedRayRadius, float4(0, Value, 0, 0.3));
 				//if(Value > 0.0f) AddLineWS(Context, RefractionP + L * VisProjectedRayRadius, float3((RefractionP + L * VisProjectedRayRadius).xy, SlabCenter.z - 0.5* UI.SlabThickness), float4(0, 1, 0, 0.3));
@@ -425,7 +440,7 @@ void VisualizeRoughRefractionPS(
 #if MODE_REFLECTION
 	FStrataLobeStatistic PaperLobe = StrataGetReflectedLobe(StrataGetDiracLobe(UI.V), 1.0, UI.Roughness);
 #else
-	FStrataLobeStatistic PaperLobe = StrataGetRefractedLobe2(StrataGetDiracLobe(UI.V), 1.0, UI.Roughness, GetEta12());
+	FStrataLobeStatistic PaperLobe = StrataGetRefractedLobe2(StrataGetDiracLobe(UI.V), 1.0, UI.Roughness, UI.Eta12);
 #endif
 	DrawCircleXY(Context, float3(SlabCenter.xy + PaperLobe.Mu.xy * VisProjectedRayRadius, SlabBottomZ), PaperLobe.Sigma * VisProjectedRayRadius, float4(1, 0, 1, 1));
 
