Strata - rough refraction now take into accounr thickness and top (only) layer roughness.
Similarly to the blurring method, the formulas evaluting VarianceCm is still a place holder pending RnD for PSF.

diff --git a/Engine/Shaders/Private/BasePassPixelShader.usf b/Engine/Shaders/Private/BasePassPixelShader.usf
--- a/Engine/Shaders/Private/BasePassPixelShader.usf
+++ b/Engine/Shaders/Private/BasePassPixelShader.usf
@@ -1728,6 +1728,7 @@ void FPixelShaderInOut_MainPS(
 	Out.StrataOutput[1] = RWStrataMaterialContainer.MaterialRenderTargets[1];
 	Out.StrataOutput[2] = StrataPackTopLayerData(TopLayerData);
 #if STRATA_OUTPUT_ROUGH_REFRACTION
+	OpaqueRoughRefractionData.OpaqueRoughRefractionEnabled = true; // As long as we output to rough refraction, we want it to be enabled for the tiling process.
 	StrataStruct.OpaqueRoughRefractionTextureUAV[PixelPos] = StrataPackOpaqueRoughRefractionData(OpaqueRoughRefractionData);
 #endif
 
diff --git a/Engine/Shaders/Private/Strata/Strata.ush b/Engine/Shaders/Private/Strata/Strata.ush
--- a/Engine/Shaders/Private/Strata/Strata.ush
+++ b/Engine/Shaders/Private/Strata/Strata.ush
@@ -133,8 +133,8 @@ struct FStrataBSDF
 	// STRATA_TODO pack the following
 	int					OperatorIndex;
 	float3				LuminanceWeight;
-	int					bIsBottom;		// STRATA TODO use thet state one or set only when inline
-	int					bIsTop;			// STRATA TODO use thet state one or set only when inline
+	int					bIsBottom;
+	int					bIsTop;
 
 #if STRATA_INLINE_SHADING
 
@@ -566,22 +566,22 @@ bool IsStrataMaterial(FStrataTopLayerData In)
 
 struct FStrataOpaqueRoughRefractionData
 {
-	bool	OpaqueRoughRefractionEnabled;
-	float	Variance;
+	float	OpaqueRoughRefractionEnabled;
+	float	Coverage;
+	float	VarianceCm;
 };
 
-float StrataPackOpaqueRoughRefractionData(FStrataOpaqueRoughRefractionData In)
+float3 StrataPackOpaqueRoughRefractionData(FStrataOpaqueRoughRefractionData In)
 {
-	float Value = In.Variance > 0.0f ? 128.0f : 0.0f;
-	Value += In.Variance * 127.0f;
-	return Value / 255.0f;
+	return float3(In.Coverage, In.VarianceCm, In.OpaqueRoughRefractionEnabled);
 }
 
-FStrataOpaqueRoughRefractionData StrataUnpackOpaqueRoughRefractionData(float In)
+FStrataOpaqueRoughRefractionData StrataUnpackOpaqueRoughRefractionData(float3 In)
 {
 	FStrataOpaqueRoughRefractionData Out = (FStrataOpaqueRoughRefractionData)0;
-	Out.OpaqueRoughRefractionEnabled = (In * 255.0) >= 128.0f;
-	Out.Variance = In * (255.0f / 127.0f) + (- 128.0f / 127.0f); // = (In * 255 - 128) / 127
+	Out.Coverage = In.r;
+	Out.VarianceCm = In.g;
+	Out.OpaqueRoughRefractionEnabled = In.b;
 	return Out;
 }
 
@@ -747,9 +747,9 @@ struct FStrataOperator
 	FStrataLobeStatistic TopDownRefractionLobeStat;		// The lobe used for the scene rough refraction from front to back, e.g. used for translucent rough refraciton
 	FStrataLobeStatistic BottomUpRefractionLobeStat;	// The lobe used for the scene rough refraction from back to front, e.g. used for opaque rough refraction of incoming light within sub layers
 
-	// Thickness
-	float				 ThicknessCm;					// The thickness of the layer
-	float				 VerticalTop_ThicknessCm;		// Only used by vertical layering operators. It represent the thickness that is used to drive the roughness propagation from top layers to sub layers.
+	// Opaque rough refraction
+	float				 OpaqueRoughRefractThicknessCm;	// The thickness of stack material from the top layer to the bottom layer.
+	float				 OpaqueRoughRefractTopRoughness;// The roughness of the top layer. This is enough because we currently assume all the bottom layers have a constant IOR, so eta will after that always be 1 and no refraction will occur after.
 
 	// Thin film
 	float ThinFilmThickness;							// Thickness of thin-film coating layer. In micro-meters (1um == 1000nm).  Merge with regular ThicknessCm?
@@ -2677,6 +2677,15 @@ void PackStrataOut(
 		FStrataAddressing NullStrataAddressing = (FStrataAddressing)0;	// Fake unused in StrataCreateBSDFContext when using Forward inline shading
 		const float3 NullLightDir = float3(0.0, 0.0, 1.0);
 		StrataUpdateTree(NullStrataAddressing, StrataPixelHeader, StrataTree, V, NullLightDir, false /*bRoughDiffuseEnabled*/);
+
+		// Process the rough refraction data (this code is only used, and not culled out, if effectively used in base pass for instance)
+		{
+			uint RootOperatorIndex = Strata.OperatorIndex;
+			const FStrataOperator RootOperator = StrataTree.Operators[Strata.OperatorIndex];
+			OpaqueRoughRefractionData.Coverage = RootOperator.Coverage;
+			// STRATA_TODO Fake variance place holder for now
+			OpaqueRoughRefractionData.VarianceCm = RootOperator.OpaqueRoughRefractTopRoughness * RootOperator.OpaqueRoughRefractThicknessCm;
+		}
 		
 		int SimpleMaterial_BSDFIndex = 0;
 		float TopLayerTotalWeight = 0.0f;
@@ -2718,8 +2727,6 @@ void PackStrataOut(
 						TopLayerTotalWeight += BSDF.TopLayerDataWeight;
 						TopLayerData.Roughness += BSDF.TopLayerDataWeight * SLAB_ROUGHNESS(BSDF);
 
-						OpaqueRoughRefractionData.Variance += BSDF.TopLayerDataWeight * SLAB_ROUGHNESS(BSDF);
-
 						const bool bIsThin = BSDF_GETISTHIN(BSDF);
 						const bool bIsSimpleVolume = BSDF_GETISSIMPLEVOLUME(BSDF);
 						bStrataSubsurfaceEnable = bStrataSubsurfaceEnable || (BSDF_GETHASSSS(BSDF) && BSDF.Coverage > 0.0f && !bIsSimpleVolume && !bIsThin); // Should it be coverage>0 or any(LuminanceWeight>0) ?
@@ -2737,8 +2744,6 @@ void PackStrataOut(
 						TopLayerTotalWeight += BSDF.TopLayerDataWeight;
 						TopLayerData.Roughness += BSDF.TopLayerDataWeight * FullyRough;
 
-						OpaqueRoughRefractionData.Variance += BSDF.TopLayerDataWeight * SLAB_ROUGHNESS(BSDF);
-
 						break;
 					}
 
@@ -2750,8 +2755,6 @@ void PackStrataOut(
 						TopLayerTotalWeight += BSDF.TopLayerDataWeight;
 						TopLayerData.Roughness += BSDF.TopLayerDataWeight * SLW_ROUGHNESS(BSDF);
 
-						OpaqueRoughRefractionData.Variance += BSDF.TopLayerDataWeight * SLAB_ROUGHNESS(BSDF);
-
 						break;
 					}
 					}
@@ -2779,11 +2782,6 @@ void PackStrataOut(
 		TopLayerData.WorldNormal= TopLayerTotalWeight > 0.0f	? normalize(TopLayerData.WorldNormal)			: 0.0f;
 		TopLayerData.Roughness	= TopLayerTotalWeight > 0.0f	? TopLayerData.Roughness / TopLayerTotalWeight	: 0.0f;
 
-#if STRATA_OUTPUT_ROUGH_REFRACTION
-		OpaqueRoughRefractionData.Variance = TopLayerTotalWeight > 0.0f ? OpaqueRoughRefractionData.Variance / TopLayerTotalWeight : 0.0f;
-		OpaqueRoughRefractionData.OpaqueRoughRefractionEnabled = OpaqueRoughRefractionData.Variance > 0.0f;
-#endif
-
 		// Set storage layout as either: fast(0), single(1), or complex(2)
 		AndIsSimpleMaterial(StrataPixelHeader, BSDFVisibleCount == 1);
 		AndIsSingleMaterial(StrataPixelHeader, BSDFVisibleCount == 1);
diff --git a/Engine/Shaders/Private/Strata/StrataMaterialClassification.usf b/Engine/Shaders/Private/Strata/StrataMaterialClassification.usf
--- a/Engine/Shaders/Private/Strata/StrataMaterialClassification.usf
+++ b/Engine/Shaders/Private/Strata/StrataMaterialClassification.usf
@@ -34,7 +34,7 @@ RWBuffer<uint> ComplexTileListDataBuffer;
 #if STRATA_OPAQUE_ROUGH_REFRACTION_ENABLED
 RWBuffer<uint> OpaqueRoughRefractionTileListDataBuffer;
 RWBuffer<uint> SSSWithoutOpaqueRoughRefractionTileListDataBuffer;
-Texture2D<float> OpaqueRoughRefractionTexture;
+Texture2D<float3> OpaqueRoughRefractionTexture;
 #endif
 
 #if !PERMUTATION_WAVE_OPS
@@ -84,7 +84,7 @@ void TileMainCS(uint2 DispatchThreadId : SV_DispatchThreadID, uint LinearIndex :
 
 #if STRATA_OPAQUE_ROUGH_REFRACTION_ENABLED
 		OpaqueRoughRefractionData = StrataUnpackOpaqueRoughRefractionData(OpaqueRoughRefractionTexture[PixelCoord]);
-		bContainsOpaqueRoughRefraction = OpaqueRoughRefractionData.OpaqueRoughRefractionEnabled;
+		bContainsOpaqueRoughRefraction = OpaqueRoughRefractionData.OpaqueRoughRefractionEnabled > 0.0f;
 #endif
 	}
 
diff --git a/Engine/Shaders/Private/Strata/StrataRoughRefraction.usf b/Engine/Shaders/Private/Strata/StrataRoughRefraction.usf
--- a/Engine/Shaders/Private/Strata/StrataRoughRefraction.usf
+++ b/Engine/Shaders/Private/Strata/StrataRoughRefraction.usf
@@ -30,7 +30,7 @@ void OpaqueRoughRefractionPS(
 		float SampleCount = 0.0f;
 		for (float u = -4.5f; u <= 4.5f; u += 1.0f)
 		{
-			SourceData += SeparatedOpaqueRoughRefractionSceneColor[PixelPos + BlurDirection * u * ORRData.Variance * 10.0];
+			SourceData += SeparatedOpaqueRoughRefractionSceneColor[PixelPos + BlurDirection * u * ORRData.VarianceCm * 10000.0];
 			SampleCount++;
 		}
 		OutColor = SourceData / SampleCount;
diff --git a/Engine/Shaders/Private/Strata/StrataTree.ush b/Engine/Shaders/Private/Strata/StrataTree.ush
--- a/Engine/Shaders/Private/Strata/StrataTree.ush
+++ b/Engine/Shaders/Private/Strata/StrataTree.ush
@@ -75,8 +75,12 @@ void UpdateSingleBSDFOperatorCoverageTransmittance(
 
 #define Op StrataTree.Operators[CurrentBSDF.OperatorIndex]
 
+	const float Roughness = StrataGetBSDFRoughness(CurrentBSDF);
+
 	// Extract the true thickness of the BSDF before it is cast into the normalised Slab thickness we used.
-	Op.ThicknessCm = CurrentBSDF.ThicknessCm;
+	// for opaque rough refraction, we are only interested in thickness of non bottom layers.
+	Op.OpaqueRoughRefractThicknessCm  = CurrentBSDF.bIsBottom ? 0.0f : CurrentBSDF.ThicknessCm;
+	Op.OpaqueRoughRefractTopRoughness = CurrentBSDF.bIsTop    ? Roughness : 0.0f;
 
 	PostProcessBSDFBeforeLighting(CurrentBSDF);
 
@@ -100,7 +104,6 @@ void UpdateSingleBSDFOperatorCoverageTransmittance(
 	Op.Transmittance = CurrentBSDFTransmittance;
 
 	// Evaluate refraction lobes.
-	const float Roughness = StrataGetBSDFRoughness(CurrentBSDF);
 	const FBxDFEnergyTerms EnergyTerms = ComputeGGXSpecEnergyTerms(Roughness, StrataBSDFContext.SatNoV, StrataGetBSDFSpecularColor(CurrentBSDF), StrataGetBSDFSpecularF90(CurrentBSDF));
 	{
 		const float TopDownInterfaceEta = CurrentBSDF.bIsTop ? AIR_IOR / WATER_IOR : 1.0f;
@@ -163,8 +166,9 @@ void UpdateSingleOperatorCoverageTransmittance(inout FStrataTree StrataTree, int
 
 		Op.BottomUpRefractionLobeStat = StrataWeightLobe(OpA.BottomUpRefractionLobeStat, Weight);
 
-		// Coverage should not affect thickness
-		Op.ThicknessCm = OpA.ThicknessCm;
+		// Coverage does not affect rough refraction data
+		Op.OpaqueRoughRefractThicknessCm = OpA.OpaqueRoughRefractThicknessCm;
+		Op.OpaqueRoughRefractTopRoughness = OpA.OpaqueRoughRefractTopRoughness;
 		break;
 	}
 
@@ -218,8 +222,9 @@ void UpdateSingleOperatorCoverageTransmittance(inout FStrataTree StrataTree, int
 			Op.BottomUpRefractionLobeStat.E *= saturate(OpA.Transmittance + (1.0f - OpA.Coverage)); // Combine with top layer medium and specular transmittance
 		}
 
-		// Vertical layering accumulates both operators thickness
-		Op.ThicknessCm = OpA.ThicknessCm + OpB.ThicknessCm;
+		// Vertical layering accumulates both operators thickness and only keep the top layer roughness, see OpaqueRoughRefractTopRoughness declaration for the details.
+		Op.OpaqueRoughRefractThicknessCm  = OpA.OpaqueRoughRefractThicknessCm + OpB.OpaqueRoughRefractThicknessCm;
+		Op.OpaqueRoughRefractTopRoughness = OpA.OpaqueRoughRefractTopRoughness;
 
 		Op.VerticalTop_TopDownRefractionLobeStat = OpA.TopDownRefractionLobeStat;
 		Op.VerticalTop_BottomUpRefractionLobeStat = OpA.BottomUpRefractionLobeStat;
@@ -241,7 +246,8 @@ void UpdateSingleOperatorCoverageTransmittance(inout FStrataTree StrataTree, int
 		Op.BottomUpRefractionLobeStat = StrataHorizontalMixLobes(OpA.BottomUpRefractionLobeStat, OpB.BottomUpRefractionLobeStat, BMix);
 
 		// Horizontal mixes both operators thickness
-		Op.ThicknessCm = lerp(OpA.ThicknessCm, OpB.ThicknessCm, BMix);
+		Op.OpaqueRoughRefractThicknessCm  = lerp(OpA.OpaqueRoughRefractThicknessCm, OpB.OpaqueRoughRefractThicknessCm, BMix);
+		Op.OpaqueRoughRefractTopRoughness = lerp(OpA.OpaqueRoughRefractTopRoughness, OpB.OpaqueRoughRefractTopRoughness, BMix);
 		break;
 	}
 
@@ -259,7 +265,8 @@ void UpdateSingleOperatorCoverageTransmittance(inout FStrataTree StrataTree, int
 
 		Op.BottomUpRefractionLobeStat = StrataHorizontalMixLobes(OpA.BottomUpRefractionLobeStat, OpB.BottomUpRefractionLobeStat, 0.5f);
 
-		Op.ThicknessCm = lerp(OpA.ThicknessCm, OpB.ThicknessCm, 0.5f);
+		Op.OpaqueRoughRefractThicknessCm  = lerp(OpA.OpaqueRoughRefractThicknessCm, OpB.OpaqueRoughRefractThicknessCm, 0.5f);
+		Op.OpaqueRoughRefractTopRoughness = lerp(OpA.OpaqueRoughRefractTopRoughness, OpB.OpaqueRoughRefractTopRoughness, 0.5f);
 		break;
 	}
 
@@ -271,7 +278,8 @@ void UpdateSingleOperatorCoverageTransmittance(inout FStrataTree StrataTree, int
 		Op.TopDownRefractionLobeStat = OpA.TopDownRefractionLobeStat;
 		Op.TopDownRefractionWorldNormal = OpA.TopDownRefractionWorldNormal;
 		Op.BottomUpRefractionLobeStat = OpA.BottomUpRefractionLobeStat;
-		Op.ThicknessCm = OpA.ThicknessCm;
+		Op.OpaqueRoughRefractThicknessCm = OpA.OpaqueRoughRefractThicknessCm;
+		Op.OpaqueRoughRefractTopRoughness = OpA.OpaqueRoughRefractTopRoughness;
 		break;
 	}
 
diff --git a/Engine/Source/Runtime/Renderer/Private/Strata/Strata.cpp b/Engine/Source/Runtime/Renderer/Private/Strata/Strata.cpp
--- a/Engine/Source/Runtime/Renderer/Private/Strata/Strata.cpp
+++ b/Engine/Source/Runtime/Renderer/Private/Strata/Strata.cpp
@@ -209,7 +209,7 @@ void InitialiseStrataFrameSceneData(FSceneRenderer& SceneRenderer, FRDGBuilder&
 			StrataSceneData.ClassificationTileListBufferUAV[EStrataTileType::ESSSWithoutOpaqueRoughRefraction] = GraphBuilder.CreateUAV(StrataSceneData.ClassificationTileListBuffer[EStrataTileType::ESSSWithoutOpaqueRoughRefraction], PF_R32_UINT);
 
 			StrataSceneData.OpaqueRoughRefractionTexture = GraphBuilder.CreateTexture(
-				FRDGTextureDesc::Create2D(OpaqueRoughRefractionSceneExtent, PF_R8, FClearValueBinding::Black, TexCreate_ShaderResource | TexCreate_UAV | TexCreate_RenderTargetable), TEXT("Strata.OpaqueRoughRefractionTexture"));
+				FRDGTextureDesc::Create2D(OpaqueRoughRefractionSceneExtent, PF_FloatR11G11B10, FClearValueBinding::Black, TexCreate_ShaderResource | TexCreate_UAV | TexCreate_RenderTargetable), TEXT("Strata.OpaqueRoughRefractionTexture"));
 			StrataSceneData.OpaqueRoughRefractionTextureUAV = GraphBuilder.CreateUAV(StrataSceneData.OpaqueRoughRefractionTexture);
 			
 			StrataSceneData.SeparatedSubSurfaceSceneColor			= GraphBuilder.CreateTexture(
@@ -479,7 +479,7 @@ class FStrataMaterialTileClassificationPassCS : public FGlobalShader
 		SHADER_PARAMETER_RDG_BUFFER_UAV(RWBuffer, OpaqueRoughRefractionTileListDataBuffer)
 		SHADER_PARAMETER_RDG_BUFFER_UAV(RWBuffer, SSSWithoutOpaqueRoughRefractionTileListDataBuffer)
 		SHADER_PARAMETER_RDG_TEXTURE_UAV(RWTexture2D<uint2>, SSSTextureUAV)
-		SHADER_PARAMETER_RDG_TEXTURE(Texture2D<float>, OpaqueRoughRefractionTexture)
+		SHADER_PARAMETER_RDG_TEXTURE(Texture2D<float3>, OpaqueRoughRefractionTexture)
 	END_SHADER_PARAMETER_STRUCT()
 
 	static bool ShouldCompilePermutation(const FGlobalShaderPermutationParameters& Parameters)
diff --git a/Engine/Source/Runtime/Renderer/Private/Strata/Strata.h b/Engine/Source/Runtime/Renderer/Private/Strata/Strata.h
--- a/Engine/Source/Runtime/Renderer/Private/Strata/Strata.h
+++ b/Engine/Source/Runtime/Renderer/Private/Strata/Strata.h
@@ -21,7 +21,7 @@ BEGIN_SHADER_PARAMETER_STRUCT(FStrataBasePassUniformParameters, )
 	SHADER_PARAMETER(uint32, bRoughDiffuse)
 	SHADER_PARAMETER_RDG_TEXTURE_UAV(RWTexture2DArray<uint>, MaterialTextureArrayUAVWithoutRTs)
 	SHADER_PARAMETER_RDG_TEXTURE_UAV(RWTexture2D<uint2>, SSSTextureUAV)
-	SHADER_PARAMETER_RDG_TEXTURE_UAV(RWTexture2D<float>, OpaqueRoughRefractionTextureUAV)
+	SHADER_PARAMETER_RDG_TEXTURE_UAV(RWTexture2D<float3>, OpaqueRoughRefractionTextureUAV)
 END_SHADER_PARAMETER_STRUCT()
 
 BEGIN_SHADER_PARAMETER_STRUCT(FStrataForwardPassUniformParameters, )
@@ -41,7 +41,7 @@ BEGIN_GLOBAL_SHADER_PARAMETER_STRUCT(FStrataGlobalUniformParameters, RENDERER_AP
 	SHADER_PARAMETER_RDG_TEXTURE(Texture2DArray<uint>, MaterialTextureArray)
 	SHADER_PARAMETER_RDG_TEXTURE(Texture2D<uint>, TopLayerTexture)
 	SHADER_PARAMETER_RDG_TEXTURE(Texture2D<uint2>, SSSTexture)
-	SHADER_PARAMETER_RDG_TEXTURE(Texture2D<float>, OpaqueRoughRefractionTexture)
+	SHADER_PARAMETER_RDG_TEXTURE(Texture2D<float3>, OpaqueRoughRefractionTexture)
 END_GLOBAL_SHADER_PARAMETER_STRUCT()
 
 // This must map to the STRATA_TILE_TYPE defines.
