Fixes low frequency drift on r.TSR.History.HighFrequencyOnly=0

diff --git a/Engine/Shaders/Private/TemporalSuperResolution/TSRUpdateHistory.usf b/Engine/Shaders/Private/TemporalSuperResolution/TSRUpdateHistory.usf
--- a/Engine/Shaders/Private/TemporalSuperResolution/TSRUpdateHistory.usf
+++ b/Engine/Shaders/Private/TemporalSuperResolution/TSRUpdateHistory.usf
@@ -302,8 +302,6 @@ void MainCS(
 		tsr_half RawParallaxRejectionMask[DPV_PIXEL_PER_LANE];
 		uint RawParallaxFactorBits[DPV_PIXEL_PER_LANE];
 		tsr_half RawLowFrequencyRejection[DPV_PIXEL_PER_LANE];
-		//tsr_half RawTranslucencyRejection[DPV_PIXEL_PER_LANE];
-		//tsr_half RawTranslucencyRejectionNeighborhood[CONFIG_SAMPLES - 1][DPV_PIXEL_PER_LANE];
 		
 		uint RawSceneStencilRef[DPV_PIXEL_PER_LANE];
 		uint RawSceneStencilRefNeighborhood[CONFIG_SAMPLES - 1][DPV_PIXEL_PER_LANE];
@@ -329,7 +327,6 @@ void MainCS(
 				#endif
 				RawParallaxRejectionMask[PixelId] = ParallaxRejectionMaskTexture[LocalInputPixelPos];
 				RawParallaxFactorBits[PixelId]    = ParallaxFactorTexture[LocalInputPixelPos];
-				//RawTranslucencyRejection[PixelId] = TranslucencyRejectionTexture[LocalInputPixelPos];
 				RawSceneStencilRef[PixelId]       = InputSceneStencilTexture.Load(int3(LocalInputPixelPos, 0)) STENCIL_COMPONENT_SWIZZLE;
 				
 				RawLowFrequencyRejection[PixelId] = HistoryRejectionTexture[dpv_access_pixel(RejectionPixelPos, PixelId)];
@@ -359,7 +356,6 @@ void MainCS(
 				UNROLL_N(DPV_PIXEL_PER_LANE)
 				for (uint PixelId = 0; PixelId < DPV_PIXEL_PER_LANE; PixelId++)
 				{
-					//RawTranslucencyRejectionNeighborhood[NeighborId][PixelId] = TranslucencyRejectionTexture[dpv_access_pixel(SampleInputPixelPos, PixelId)];
 					RawSceneStencilRefNeighborhood[NeighborId][PixelId] = InputSceneStencilTexture.Load(int3(dpv_access_pixel(SampleInputPixelPos, PixelId), 0)) STENCIL_COMPONENT_SWIZZLE;
 				}
 			} // for (uint NeighborId = 0; NeighborId < (CONFIG_SAMPLES - 1); NeighborId++)
@@ -443,7 +439,6 @@ void MainCS(
 			// Process input texel, forcing to pack register.
 			ParallaxFactorBits    = dpv_force_interleave_registers_array(RawParallaxFactorBits) & tsr_subpixel_payload(SUB_PIXEL_PARALLAX_FACTOR_BIT_MASK);
 			LowFrequencyRejection = dpv_force_interleave_registers_array(RawLowFrequencyRejection);
-			//TranslucencyRejection = dpv_force_interleave_registers_array(RawTranslucencyRejection);
 			SceneStencilRef       = dpv_force_interleave_registers_array(RawSceneStencilRef);
 			
 			// Process neighborhood
@@ -1104,7 +1099,6 @@ void MainCS(
 	#else // !CONFIG_HIGH_FREQUENCY_ONLY
 	{
 		tsr_half2 InvPrevHistoryValidity = SafeRcp(PrevHistoryValidity);
-		tsr_half2 InvPrevHistoryWeight = InvPrevHistoryValidity;
 
 		// Compute how much the history needs to be rejected, based on the completeness.
 		tsr_half2 PrevHistoryRejectionWeight = LowFrequencyRejection * HoleFilledVelocityMask[1];
@@ -1113,7 +1107,7 @@ void MainCS(
 		//	PrevHistoryRejectionWeight = 1.0;
 		//#endif
 		
-		PrevHistoryRejectionWeight = select(bIsOffScreen || bIsDisoccluded, tsr_half(0.0).xx, PrevHistoryRejectionWeight);
+		PrevHistoryRejectionWeight = select(bIsOffScreen || bIsDisoccluded, tsr_half(0.0), PrevHistoryRejectionWeight);
 
 		tsr_half2 DesiredCurrentContribution = max(Histeresis * InputPixelAlignement, tsr_half(0.0));
 
@@ -1129,7 +1123,7 @@ void MainCS(
 			// Clamp up the max validity to favorise stability under motion on high contrast edges.
 			#if 1
 			{
-				tsr_half2 PrevHistoryLuma = Luma4(PrevHighFrequencyColor) * InvPrevHistoryWeight;
+				tsr_half2 PrevHistoryLuma = Luma4(PrevHighFrequencyColor) * InvPrevHistoryValidity;
 				tsr_half2 FilteredLuma = Luma4(FilteredInputColor);
 
 				tsr_half2 MinValidityForStability = abs(FilteredLuma - PrevHistoryLuma) / max(FilteredLuma, PrevHistoryLuma);
@@ -1174,10 +1168,10 @@ void MainCS(
 		
 		tsr_half2 PrevMomentWeight = max(OutputValidity - DesiredCurrentContribution, tsr_half(0.0));
 		tsr_half2 CurrentMomentWeight = min(DesiredCurrentContribution, OutputValidity);
-
+		
 		#if 1
 		{
-			tsr_half2 PrevHistoryToneWeight = HdrWeightY(Luma4(PrevHighFrequencyColor) * InvPrevHistoryWeight);
+			tsr_half2 PrevHistoryToneWeight = HdrWeightY(Luma4(PrevHighFrequencyColor) * InvPrevHistoryValidity);
 			tsr_half2 FilteredInputToneWeight = HdrWeight4(FilteredInputColor);
 			
 			tsr_half2 BlendPrevHistory = PrevMomentWeight * PrevHistoryToneWeight;
@@ -1186,7 +1180,7 @@ void MainCS(
 			tsr_half2 CommonWeight = SafeRcp(BlendPrevHistory + BlendFilteredInput);
 
 			FinalHighFrequencyColor = (
-				dpv_scale(PrevHighFrequencyColor, CommonWeight * BlendPrevHistory * InvPrevHistoryWeight) +
+				dpv_scale(PrevHighFrequencyColor, CommonWeight * BlendPrevHistory * InvPrevHistoryValidity) +
 				dpv_scale(FilteredInputColor, CommonWeight * BlendFilteredInput));
 		}
 		#else
@@ -1299,7 +1293,7 @@ void MainCS(
 
 			// Completely discard the history
 			{
-				tsr_half2 TotalRejection = select(bIsOffScreen || bIsDisoccluded, tsr_half(0.0).xx, saturate(LowFrequencyRejection * tsr_half(4.0)));
+				tsr_half2 TotalRejection = select(bIsOffScreen || bIsDisoccluded, tsr_half(0.0), saturate(LowFrequencyRejection * tsr_half(4.0)));
 
 				tsr_half2 TotalRejectionBlending = max(BlendFinal, HoleFilledVelocityMask[0]);
 
@@ -1558,13 +1552,21 @@ void MainCS(
 		{
 			uint2 Random = Rand3DPCG16(int3(dpv_lo(LocalHistoryPixelPos), View.StateFrameIndexMod8)).xy;
 			tsr_half E = tsr_half(Hammersley16(0, 1, Random).x);
-
+			
+			#if !CONFIG_HIGH_FREQUENCY_ONLY
+				FinalLowFrequencyColor = QuantizeForFloatRenderTarget(FinalLowFrequencyColor, E, HistoryQuantizationError);
+			#endif
 			PremultipliedHighFrequencyColor = QuantizeForFloatRenderTarget(PremultipliedHighFrequencyColor, E, HistoryQuantizationError);
 		}
 		#endif
 
 		// Protext from NaN and +Inf when writing out the history.
 		{
+			#if !CONFIG_HIGH_FREQUENCY_ONLY
+				FinalLowFrequencyColor = -dpv_min(-FinalLowFrequencyColor, ColorNull);
+				FinalLowFrequencyColor = dpv_min(FinalLowFrequencyColor, ColorMax10BitsFloat);
+			#endif
+
 			PremultipliedHighFrequencyColor = -dpv_min(-PremultipliedHighFrequencyColor, ColorNull);
 			PremultipliedHighFrequencyColor = dpv_min(PremultipliedHighFrequencyColor, ColorMax10BitsFloat);	
 
