SMRT improvements:
- Add slope-based depth extrapolation which improves the quality of penumbras on angled receivers. Costs ~10% performance in some cases so maintaining a permutation/cvar (default on) for scalability.
- Change screen ray trace to be a simple "space skipping" ray that terminates as soon as it goes behind geometry and continue VSM trace from that distance. This avoids various contact-shadow-like artifacts and undesirable/inconsistent contact shadows from things that aren't in the VSM. In certain cases if regular contact shadows are desired on top of VSM the engine contact shadows can be enabled, as it is with CSMs.
- Remove a bunch of use of "halfs" in the shaders as they cause some extra ALU on some platforms and don't appear to really be helping with occupancy anymore
- Small bump to minimum normal bias clamp (only affects things very close to the camera)

diff --git a/Engine/Shaders/Private/VirtualShadowMaps/VirtualShadowMapProjection.usf b/Engine/Shaders/Private/VirtualShadowMaps/VirtualShadowMapProjection.usf
--- a/Engine/Shaders/Private/VirtualShadowMaps/VirtualShadowMapProjection.usf
+++ b/Engine/Shaders/Private/VirtualShadowMaps/VirtualShadowMapProjection.usf
@@ -36,7 +36,7 @@
 #define VISUALIZE_OUTPUT 0
 #endif
 
-#define CONTACT_SHADOW_SAMPLES 8
+#define SCREEN_RAY_SAMPLES 4
 
 struct FProjectionShadingInfo
 {
@@ -44,7 +44,7 @@ struct FProjectionShadingInfo
 	bool bIsHair;
 	bool bIsSubsurface;
 
-	half3 WorldNormal;
+	float3 WorldNormal;
 };
 
 FProjectionShadingInfo GetProjectionShadingInfo(uint2 PixelPos)
@@ -82,7 +82,7 @@ float4 EncodeLightAttenuationFromMask(float ShadowMask, float SSSTransmission)
 	float FadedSSSShadow = lerp(1.0f, SSSTransmission, ShadowFadeFraction);
 
 	// the channel assignment is documented in ShadowRendering.cpp (look for Light Attenuation channel assignment)
-	return EncodeLightAttenuation(half4(FadedShadow, FadedSSSShadow, FadedShadow, FadedSSSShadow));
+	return EncodeLightAttenuation(float4(FadedShadow, FadedSSSShadow, FadedShadow, FadedSSSShadow));
 }
 
 Texture2D<float2> InputShadowFactor;
@@ -99,19 +99,57 @@ void VirtualShadowMapCompositePS(
 SCHEDULER_MIN_PRESSURE
 MAX_OCCUPENCY
 
-// Returns hit distance or negative if miss
+// Screen space ray trace to attempt to skip over ambiguous regions near the receiver surface
+// Returns length at which to start the virtual shadow map ray; usually this is where the screen ray ended or went behind a surface
 float VirtualShadowMapScreenRayCast(
-	float3 RayOriginTranslatedWorld, float3 RayDirection, float RayLength,
-	int NumSteps, float Dither)
+	float3 RayOriginTranslatedWorld,
+	float3 RayDirection,
+	float RayLength,
+	float Dither)
 {
+	float4 RayStartClip	= mul(float4(RayOriginTranslatedWorld, 1), View.TranslatedWorldToClip);
+	float4 RayDirClip	= mul(float4(RayDirection * RayLength, 0), View.TranslatedWorldToClip);
+	float4 RayEndClip	= RayStartClip + RayDirClip;
+
+	float3 RayStartScreen = RayStartClip.xyz / RayStartClip.w;
+	float3 RayEndScreen = RayEndClip.xyz / RayEndClip.w;
+
+	float3 RayStepScreen = RayEndScreen - RayStartScreen;
+
+	float3 RayStartUVz = float3(RayStartScreen.xy * View.ScreenPositionScaleBias.xy + View.ScreenPositionScaleBias.wz, RayStartScreen.z);
+	float3 RayStepUVz  = float3(RayStepScreen.xy * View.ScreenPositionScaleBias.xy, RayStepScreen.z);
+
+	float4 RayDepthClip	= RayStartClip + mul(float4(0, 0, RayLength, 0), View.ViewToClip);
+	float3 RayDepthScreen = RayDepthClip.xyz / RayDepthClip.w;
+
+	const int Steps = SCREEN_RAY_SAMPLES;
 	float StepOffset = Dither - 0.5f;
+	const float Step = 1.0 / Steps;
+	float SampleTime = StepOffset * Step + Step;
 
-	bool bHitCastContactShadow = false;
-	float HitDistance = ShadowRayCast(
-		RayOriginTranslatedWorld, RayDirection, RayLength,
-		NumSteps, StepOffset, bHitCastContactShadow);
+	const float StartDepth = SceneTexturesStruct.SceneDepthTexture.SampleLevel(SceneTexturesStruct_SceneDepthTextureSampler, RayStartUVz.xy, 0).r;
 
-	return (bHitCastContactShadow) ? HitDistance : -1.0f;
+	UNROLL
+	for (int i = 0; i < Steps; i++)
+	{
+		float3 SampleUVz = RayStartUVz + RayStepUVz * SampleTime;
+		float SampleDepth = SceneTexturesStruct.SceneDepthTexture.SampleLevel(SceneTexturesStruct_SceneDepthTextureSampler, SampleUVz.xy, 0).r;
+
+		// Avoid self-intersection with the start pixel (exact comparison due to point sampling depth buffer)
+		if (SampleDepth != StartDepth)
+		{
+			if (SampleUVz.z < SampleDepth)
+			{
+				// Behind geometry. Back up a bit along the ray and do the VSM sample from there.
+				return RayLength * max(0.0, SampleTime - 1.5f * Step);
+			}
+		}
+
+		SampleTime += Step;
+	}
+
+	// Got to the end of the ray without going behind or hitting anything
+	return RayLength;
 }
 
 float3 GetEstimatedGeoWorldNormal(float3 TranslatedWorldPosition, FProjectionShadingInfo ShadingInfo)
@@ -166,7 +204,7 @@ float4 ComputeRandom4(uint2 PixelPosition)
 }
 
 int4 ProjectionRect;
-float ContactShadowLength;
+float ScreenRayLength;
 float NormalBias;
 int SMRTRayCount;					// 0 = off
 int SMRTSamplesPerRay;
@@ -175,7 +213,7 @@ float SMRTCotMaxRayAngleFromLight;	// Spot/point lights
 float SMRTTexelDitherScale;			// Currently only directional lights
 
 // One pass projection
-RWTexture2D< uint4 >OutShadowMaskBits;
+RWTexture2D< uint4 > OutShadowMaskBits;
 
 // Single light per pass
 // Light parameters loaded via GetRootLightShaderParameters();
@@ -200,15 +238,15 @@ FVirtualShadowMapSampleResult ProjectLight(
 	FProjectionShadingInfo ShadingInfo,
 	uint2 PixelPos,
 	float SceneDepth,
-	float ContactShadowLengthWorld,
+	float ScreenRayLengthWorld,
 	float3 TranslatedWorldPosition,
-	const half Noise)
+	const float Noise)
 {
 	const bool bIsHairInput = InputType == INPUT_TYPE_HAIRSTRANDS;
 
 	// NOTE: If PrimaryView.PreViewTranslation is guaranteed to be the camera position we can simplify
 	const float DistanceToCamera = length(TranslatedWorldPosition - View.TranslatedWorldCameraOrigin);
-	const float NormalBiasLength = max(0.01f, NormalBias * DistanceToCamera / View.ViewToClip._11);
+	const float NormalBiasLength = max(0.02f, NormalBias * DistanceToCamera / View.ViewToClip._11);
 
 #if DIRECTIONAL_LIGHT
 	float3 L = Light.Direction;
@@ -217,7 +255,7 @@ FVirtualShadowMapSampleResult ProjectLight(
 	float3 ToLight = Light.TranslatedWorldPosition - TranslatedWorldPosition;
 	float d2 = dot( ToLight, ToLight );
 	float InvDist = rsqrt( d2 );
-	half3 L = ToLight * InvDist;
+	float3 L = ToLight * InvDist;
 	
 	bool bInLightRadius = InvDist >= Light.InvRadius;
 	bool bInLightCone = dot( L, Light.Direction ) >= Light.SpotAngles.x;
@@ -233,9 +271,26 @@ FVirtualShadowMapSampleResult ProjectLight(
 	{
 		const bool bBackfaceCull = (bCullBackfacingPixels != 0) && !bIsHairInput && !ShadingInfo.bIsSubsurface;
 
-		const half3 WorldNormal = (bIsHairInput || ShadingInfo.bIsHair) ? L : ShadingInfo.WorldNormal;
+		const float3 WorldNormal = (bIsHairInput || ShadingInfo.bIsHair) ? L : ShadingInfo.WorldNormal;
 		TranslatedWorldPosition += WorldNormal * NormalBiasLength;
 
+		// Do not run contact shadow when computing the hair shadow mask (i.e. shadow mask applied on hair, has the scene 
+		// depth buffer contains fully opaque depth, which create false positive intersection resulting in wrong self shadowing)
+		float SMRTRayOffset = ScreenRayLengthWorld;
+		if (!bIsHairInput)
+		{
+			if (ScreenRayLengthWorld > 0.0f)
+			{
+				// Trace a screen space ray to try and get "away" from the receiver surface before
+				// we trace the SMRT ray to avoid mismatches/incorrect self-shadowing.
+				SMRTRayOffset = VirtualShadowMapScreenRayCast(
+					TranslatedWorldPosition,
+					L,
+					ScreenRayLengthWorld,
+					Noise);
+			}
+		}
+
 		if (SMRTRayCount > 0)
 		{
 		#if DIRECTIONAL_LIGHT
@@ -245,7 +300,7 @@ FVirtualShadowMapSampleResult ProjectLight(
 				PixelPos,
 				SceneDepth,
 				TranslatedWorldPosition,
-				ContactShadowLengthWorld,
+				SMRTRayOffset,
 				SMRTRayCount,
 				SMRTSamplesPerRay,
 				SMRTRayLengthScale,
@@ -260,7 +315,7 @@ FVirtualShadowMapSampleResult ProjectLight(
 				PixelPos,
 				SceneDepth,
 				TranslatedWorldPosition,
-				ContactShadowLengthWorld,
+				SMRTRayOffset,
 				SMRTRayCount,
 				SMRTSamplesPerRay,
 				SMRTCotMaxRayAngleFromLight,
@@ -274,24 +329,11 @@ FVirtualShadowMapSampleResult ProjectLight(
 			Result = SampleVirtualShadowMapTranslatedWorld(
 				VirtualShadowMapId,
 				TranslatedWorldPosition,
-				ContactShadowLengthWorld,
+				SMRTRayOffset,
 				WorldNormal);
 		}
 
-		// Do not run contact shadow when computing the hair shadow mask (i.e. shadow mask applied on hair, has the scene 
-		// depth buffer contains fully opaque depth, which create false positive intersection resulting in wrong self shadowing)
-		if (!bIsHairInput)
-		{
-			if (ContactShadowLengthWorld > 0.0f)
-			{
-				half ContactShadowHitDistance = VirtualShadowMapScreenRayCast(TranslatedWorldPosition, L, ContactShadowLengthWorld, CONTACT_SHADOW_SAMPLES, Noise);
-				if (ContactShadowHitDistance > 0.0f)
-				{
-					Result.ShadowFactor = 0.0f;
-					Result.OccluderDistance = ContactShadowHitDistance;
-				}
-			}
-		}
+		//Result.GeneralDebug = GreenToRedTurbo(1.0f - (SMRTRayOffset / ScreenRayLengthWorld));
 
 		// Hair strands voxel traversal to apply hair shadow on top of opaque geometry
 		#if HAS_HAIR_STRANDS
@@ -311,7 +353,7 @@ FVirtualShadowMapSampleResult ProjectLight(
 					Light,
 					PixelPos,
 					TranslatedWorldPosition,
-					ContactShadowLengthWorld,
+					SMRTRayOffset,		// TODO: Is this actually what we want here?
 					RayIndex,
 					SMRTRayCount,
 					RayStart,
@@ -439,6 +481,10 @@ void OutputVisualize( uint2 PixelPos, FVirtualShadowMapSampleResult Result )
 			float3 PageColor = IntToColor(vPage.x + (vPage.y << 10U));
 			Output = 0.55f*CacheColor + 0.25f*PageColor + 0.2f*Result.ShadowFactor.xxx;
 		}
+		else if ( VisualizeModeId == VIRTUAL_SHADOW_MAP_VISUALIZE_GENERAL_DEBUG )
+		{
+			Output = 0.8f*Result.GeneralDebug + 0.2f*Result.ShadowFactor.xxx;
+		}
 		/*
 		else if ( VisualizeModeId == VIRTUAL_SHADOW_MAP_VISUALIZE_CACHED_PAGE_AGE && bValidPageAddress )
 		{
@@ -495,11 +541,11 @@ void VirtualShadowMapProjection(
 	const float4 SvPosition = float4( float2( PixelPos ) + 0.5, DeviceZ, 1.0f );
 	const float3 TranslatedWorldPosition = SvPositionToTranslatedWorld( SvPosition );
 
-	const float ContactShadowLengthWorld = ContactShadowLength * View.ClipToView[1][1] * SceneDepth;
-	const half Noise = InterleavedGradientNoise( SvPosition.xy, View.StateFrameIndexMod8 );
+	const float ScreenRayLengthWorld = ScreenRayLength * View.ClipToView[1][1] * SceneDepth;
+	const float Noise = InterleavedGradientNoise( SvPosition.xy, View.StateFrameIndexMod8 );
 
 	const FProjectionShadingInfo ShadingInfo = GetProjectionShadingInfo(PixelPos);
-	const half SubsurfaceOpacity = bIsHairInput ? 1.0f : GetSubsurfaceOpacityFromGbuffer( PixelPos );
+	const float SubsurfaceOpacity = bIsHairInput ? 1.0f : GetSubsurfaceOpacityFromGbuffer( PixelPos );
 
 	FVirtualShadowMapSampleResult VisualizeResult = InitVirtualShadowMapSampleResult();
 
@@ -529,7 +575,7 @@ void VirtualShadowMapProjection(
 				ShadingInfo,
 				PixelPos,
 				SceneDepth,
-				ContactShadowLengthWorld,
+				ScreenRayLengthWorld,
 				TranslatedWorldPosition,
 				Noise);
 
@@ -557,11 +603,11 @@ void VirtualShadowMapProjection(
 			ShadingInfo,
 			PixelPos,
 			SceneDepth,
-			ContactShadowLengthWorld,
+			ScreenRayLengthWorld,
 			TranslatedWorldPosition,
 			Noise);
 		
-		half SSSTransmission = Result.ShadowFactor;
+		float SSSTransmission = Result.ShadowFactor;
 		if ( SubsurfaceOpacity < 1.0f )
 		{
 			SSSTransmission = ComputeSimpleSubsurfaceTransmissionFromVirtualShadowMap(
@@ -570,7 +616,7 @@ void VirtualShadowMapProjection(
 				SubsurfaceOpacity);
 		}
 
-		OutShadowFactor[ PixelPos ] = half2( Result.ShadowFactor, SSSTransmission );
+		OutShadowFactor[ PixelPos ] = float2( Result.ShadowFactor, SSSTransmission );
 
 		if (VisualizeVirtualShadowMapId == VirtualShadowMapId)
 		{
diff --git a/Engine/Shaders/Private/VirtualShadowMaps/VirtualShadowMapProjectionCommon.ush b/Engine/Shaders/Private/VirtualShadowMaps/VirtualShadowMapProjectionCommon.ush
--- a/Engine/Shaders/Private/VirtualShadowMaps/VirtualShadowMapProjectionCommon.ush
+++ b/Engine/Shaders/Private/VirtualShadowMaps/VirtualShadowMapProjectionCommon.ush
@@ -298,6 +298,7 @@ struct FVirtualShadowMapSampleResult
 	uint RayCount;
 	uint2 VirtualTexelAddress;
 	uint2 PhysicalTexelAddress;
+	float3 GeneralDebug;		// General purpose debug output during shader development
 };
 
 
@@ -311,6 +312,7 @@ FVirtualShadowMapSampleResult InitVirtualShadowMapSampleResult()
 	Result.VirtualTexelAddress = uint2(0xFFFFFFFF, 0xFFFFFFFF);
 	Result.PhysicalTexelAddress = uint2(0xFFFFFFFF, 0xFFFFFFFF);
 	Result.RayCount = 0;
+	Result.GeneralDebug = float3(0, 0, 0);
 	return Result;
 }
 
diff --git a/Engine/Shaders/Private/VirtualShadowMaps/VirtualShadowMapProjectionDirectional.ush b/Engine/Shaders/Private/VirtualShadowMaps/VirtualShadowMapProjectionDirectional.ush
--- a/Engine/Shaders/Private/VirtualShadowMaps/VirtualShadowMapProjectionDirectional.ush
+++ b/Engine/Shaders/Private/VirtualShadowMaps/VirtualShadowMapProjectionDirectional.ush
@@ -36,7 +36,7 @@ float ComputeOptimalSlopeBiasDirectional(
 	float2 DepthSlopeUV = -NormalPlaneUV.xy / NormalPlaneUV.z;
 
 	// Clamp to avoid excessive degenerate slope biases causing flickering lit pixels
-	float2 Clamp = 20.0f;
+	float2 Clamp = 0.1f;
 	DepthSlopeUV = bClamp ? clamp(DepthSlopeUV, -Clamp, Clamp) : DepthSlopeUV;
 		
 	float2 OffsetUV = TexelOffset / float(CalcLevelDimsTexels(0));
diff --git a/Engine/Shaders/Private/VirtualShadowMaps/VirtualShadowMapProjectionSpot.ush b/Engine/Shaders/Private/VirtualShadowMaps/VirtualShadowMapProjectionSpot.ush
--- a/Engine/Shaders/Private/VirtualShadowMaps/VirtualShadowMapProjectionSpot.ush
+++ b/Engine/Shaders/Private/VirtualShadowMaps/VirtualShadowMapProjectionSpot.ush
@@ -156,7 +156,7 @@ FSMRTResult ShadowRayCastSpotLight(
 	float3 RayStartShadowTranslatedWorld,
 	float3 RayEndShadowTranslatedWorld,
 	int NumSteps,
-	half StepOffset,
+	float StepOffset,
 	// To get debug data out
 	inout FSMRTSingleRayState OutRayState)
 {
@@ -169,14 +169,14 @@ FSMRTResult ShadowRayCastPointLight(
 	float3 RayStartShadowTranslatedWorld,
 	float3 RayEndShadowTranslatedWorld,
 	int NumSteps,
-	half StepOffset,
+	float StepOffset,
 	// To get debug data out
 	inout FSMRTSingleRayState OutRayState)
 {
 	uint RayStartFace = VirtualShadowMapGetCubeFace(RayStartShadowTranslatedWorld);
 	uint RayEndFace   = VirtualShadowMapGetCubeFace(RayEndShadowTranslatedWorld);
 		
-	FSMRTResult Result;	
+	FSMRTResult Result;
 	if (RayStartFace != RayEndFace)
 	{
 		FSMRTTwoCubeFaceRayState RayState = SMRTTwoCubeFaceRayInitialize(
@@ -233,11 +233,11 @@ float4 SobolSampler(uint SampleIndex, inout uint Seed, uint MaxBits)
 // Normal should be normalized
 bool IsBackfaceToLocalLight(
 	float3 ToLight,
-	half3 Normal,
-	half LightSourceRadius,
-	out half3 OutConeAxis,
-	out half OutConeSin,
-	out half OutDistToLight)
+	float3 Normal,
+	float LightSourceRadius,
+	out float3 OutConeAxis,
+	out float OutConeSin,
+	out float OutDistToLight)
 {
 	float DistSqr = dot(ToLight, ToLight);
 	float Falloff = rcp(DistSqr + 1);
@@ -247,19 +247,19 @@ bool IsBackfaceToLocalLight(
 	OutConeSin = LightSourceRadius * InvDist;
 	OutDistToLight = DistSqr * InvDist;
 
-	half SinAlphaSqr = saturate(Pow2(LightSourceRadius) * Falloff);
-	half SinAlpha = sqrt(SinAlphaSqr);
+	float SinAlphaSqr = saturate(Pow2(LightSourceRadius) * Falloff);
+	float SinAlpha = sqrt(SinAlphaSqr);
 	bool bBackface = dot(Normal, OutConeAxis) < -SinAlpha;
 
 	return bBackface;
 }
 
 // Version without outputs
-bool IsBackfaceToLocalLight(float3 ToLight, half3 Normal, half LightSourceRadius)
+bool IsBackfaceToLocalLight(float3 ToLight, float3 Normal, float LightSourceRadius)
 {
-	half3 Dummy1;
-	half Dummy2;
-	half Dummy3;
+	float3 Dummy1;
+	float Dummy2;
+	float Dummy3;
 	return IsBackfaceToLocalLight(ToLight, Normal, LightSourceRadius, Dummy1, Dummy2, Dummy3);
 }
 
@@ -273,16 +273,16 @@ FVirtualShadowMapSampleResult TraceLocalLight(
 	uint MaxRayCount,
 	uint SamplesPerRay,
 	float CotMaxRayAngleFromLight,
-	const half Noise,
+	const float Noise,
 	const bool bCullBackface,
-	half3 WorldNormal)
+	float3 WorldNormal)
 {
 	const uint LightType = Light.SpotAngles.x > -2.0f ? LIGHT_TYPE_SPOT : LIGHT_TYPE_POINT;
 	const float3 ToLight = Light.TranslatedWorldPosition - TranslatedWorldPosition;
 
-	half3 ConeAxis;
-	half ConeSin;
-	half DistToLight;
+	float3 ConeAxis;
+	float ConeSin;
+	float DistToLight;
 	bool bBackface = IsBackfaceToLocalLight(ToLight, WorldNormal, Light.SourceRadius, ConeAxis, ConeSin, DistToLight);
 
 	FVirtualShadowMapSampleResult Result = InitVirtualShadowMapSampleResult();
@@ -318,14 +318,14 @@ FVirtualShadowMapSampleResult TraceLocalLight(
 			float4 RandSample = SobolSampler( RandSequence.SampleIndex, RandSequence.SampleSeed, 8 );
 #endif
 
-			half2 LightUV = UniformSampleDiskConcentricApprox( RandSample.xy ).xy;
+			float2 LightUV = UniformSampleDiskConcentricApprox( RandSample.xy ).xy;
 			LightUV *= ConeSin;
 
-			half SinTheta2 = dot( LightUV, LightUV );
-			half SinTheta = sqrt( SinTheta2 );
-			half CosTheta = sqrt( 1 - SinTheta2 );
+			float SinTheta2 = dot( LightUV, LightUV );
+			float SinTheta = sqrt( SinTheta2 );
+			float CosTheta = sqrt( 1 - SinTheta2 );
 
-			half3 Dir = TangentToWorld( half3( LightUV, CosTheta ), ConeAxis );
+			float3 Dir = TangentToWorld( float3( LightUV, CosTheta ), ConeAxis );
 
 			// The further afield our rays go the poorer our approximation is as the "bend" due to our testing
 			// against a shadow map instead of along the ray increases. Thus we avoid going all the way to the light
@@ -423,7 +423,7 @@ bool GenerateRayLocalLight(
 	FLightShaderParameters Light,
 	uint2 PixelPos,
 	float3 TranslatedWorldPosition,
-	half3 WorldNormal,
+	float3 WorldNormal,
 	float CotMaxRayAngleFromLight,
 	uint RayIndex,
 	uint RayCount,
@@ -432,9 +432,9 @@ bool GenerateRayLocalLight(
 {
 	const float3 ToLight = Light.TranslatedWorldPosition - TranslatedWorldPosition;
 
-	half3 ConeAxis;
-	half ConeSin;
-	half DistToLight;
+	float3 ConeAxis;
+	float ConeSin;
+	float DistToLight;
 	bool bBackface = IsBackfaceToLocalLight(ToLight, WorldNormal, Light.SourceRadius, ConeAxis, ConeSin, DistToLight);
 	
 	if (!bBackface)
@@ -446,14 +446,14 @@ bool GenerateRayLocalLight(
 		RandSequence.SampleSeed = 0; // always use the same sequence
 
 		float2 E = SobolSampler(RandSequence.SampleIndex, RandSequence.SampleSeed, 8).xy;
-		half2 LightUV = UniformSampleDiskConcentricApprox(E).xy;
+		float2 LightUV = UniformSampleDiskConcentricApprox(E).xy;
 		LightUV *= ConeSin;
 
-		half SinTheta2 = dot(LightUV, LightUV);
-		half SinTheta = sqrt(SinTheta2);
-		half CosTheta = sqrt(1 - SinTheta2);
+		float SinTheta2 = dot(LightUV, LightUV);
+		float SinTheta = sqrt(SinTheta2);
+		float CosTheta = sqrt(1 - SinTheta2);
 
-		half3 Dir = TangentToWorld(half3(LightUV, CosTheta), ConeAxis);
+		float3 Dir = TangentToWorld(float3(LightUV, CosTheta), ConeAxis);
 
 		// The further afield our rays go the poorer our approximation is as the "bend" due to our testing
 		// against a shadow map instead of along the ray increases. Thus we avoid going all the way to the light
diff --git a/Engine/Shaders/Private/VirtualShadowMaps/VirtualShadowMapSMRTCommon.ush b/Engine/Shaders/Private/VirtualShadowMaps/VirtualShadowMapSMRTCommon.ush
--- a/Engine/Shaders/Private/VirtualShadowMaps/VirtualShadowMapSMRTCommon.ush
+++ b/Engine/Shaders/Private/VirtualShadowMaps/VirtualShadowMapSMRTCommon.ush
@@ -6,6 +6,10 @@ VirtualShadowMapSMRTCommon.ush:
 
 #pragma once
 
+#ifndef SMRT_EXTRAPOLATE_WITH_SLOPE
+#define SMRT_EXTRAPOLATE_WITH_SLOPE 1
+#endif
+
 struct FSMRTSample
 {
 	bool bValid;	
diff --git a/Engine/Shaders/Private/VirtualShadowMaps/VirtualShadowMapSMRTTemplate.ush b/Engine/Shaders/Private/VirtualShadowMaps/VirtualShadowMapSMRTTemplate.ush
--- a/Engine/Shaders/Private/VirtualShadowMaps/VirtualShadowMapSMRTTemplate.ush
+++ b/Engine/Shaders/Private/VirtualShadowMaps/VirtualShadowMapSMRTTemplate.ush
@@ -24,21 +24,24 @@ FSMRTResult SMRTRayCast(
 	float StepOffset)
 {
 	float DepthHistory = -1;
+	float DepthHistoryTime = -1;
+	float DepthSlope = 0;
+
 	int SampleMissCount = 0;
 	float OccluderDistance = 100000;
 
 	const float TimeScale = -1.0f / NumSteps;
 	const float TimeBias = 1.0f + ( 1.0 - StepOffset ) * TimeScale;
 
-	// Backward delta. Time[i] - Time[i-1], Half because CompareTolerance is doubled.
-	const float CompareScale = Pow2( TimeScale );
-	const float CompareBias = ( TimeBias - 0.5 * TimeScale ) * TimeScale;
+	// Backward delta. Time[i] - Time[i-1]
+	const float DeltaTimeScale = 2.0 * Pow2( TimeScale );
+	const float DeltaTimeBias = 2.0 * TimeScale * ( TimeBias - 0.5 * TimeScale );
 
-	bool bValidHit = false;
+	bool bValidHit =  false;
 	for (int i = 0; i <= NumSteps; i++)
 	{
-		const float SampleTime		= i == NumSteps ? 0 : Pow2( TimeScale * i + TimeBias );
-		const float SampleCompare	= i == NumSteps ? 0.5 * Pow2( TimeScale * StepOffset ) : CompareScale * i + CompareBias;
+		const float SampleTime	= i == NumSteps ? 0 : Pow2( TimeScale * i + TimeBias );
+		const float DeltaTime	= i == NumSteps ? Pow2( TimeScale * StepOffset ) : DeltaTimeScale * i + DeltaTimeBias;
 
 		FSMRTSample Sample = SMRTFindSample(RayState, SampleTime);
 		if (Sample.bValid)
@@ -49,6 +52,7 @@ FSMRTResult SMRTRayCast(
 			{
 				// First valid sample we've seen. Do a regular depth compare.
 				DepthHistory = SampleDepth;
+				DepthHistoryTime = SampleTime;
 				if (SampleDepth > ReferenceDepth)
 				{
 					OccluderDistance = SMRTComputeOccluderDistance(RayState, SampleDepth);
@@ -61,14 +65,39 @@ FSMRTResult SMRTRayCast(
 				// Add a small relative error to the comparison to avoid missing surfaces due to numeric precision issues
 				// Without this there are occasionally flickering fireflies in fully shadowed regions with SMRT
 				const float EpsScale = 1.02f;
-				const float CompareTolerance = abs(Sample.CompareToleranceScale * SampleCompare) * EpsScale;
+				const float CompareTolerance = abs(Sample.CompareToleranceScale * DeltaTime) * EpsScale;
+				const bool bBehind = (SampleDepth - ReferenceDepth) > CompareTolerance;
+
+				float DeltaHistoryTime = SampleTime - DepthHistoryTime;
+				float DepthForComparison = SampleDepth;
 
-				bool bBehind = (SampleDepth - ReferenceDepth) > (2.0f * CompareTolerance);
-				float DepthSample = bBehind ? DepthHistory : SampleDepth;
-				DepthHistory = DepthSample;
+				if (bBehind)
+				{
+					#if SMRT_EXTRAPOLATE_WITH_SLOPE
+						DepthForComparison = DepthSlope * DeltaHistoryTime + DepthHistory;
+					#else
+						DepthForComparison = DepthHistory;
+					#endif
+				}
+				else
+				{
+					// This must be the second valid sample we've seen, as the first will take the DepthHistory < 0 path
+					// Only update the slope if we sampled a new texel; otherwise it will just zero out the slope
+					// This can happen at high sample counts with rays that are nearly parallel to the light direction
+					if (SampleDepth != DepthHistory)	// Exact test is fine here; point sampled texel with be identical
+					{
+						// NOTE: DCE will remove all this if SMRT_EXTRAPOLATE_WITH_SLOPE is false
+						const float SlopeClamp = 1e-5f;
+						DepthSlope = (SampleDepth - DepthHistory) / DeltaHistoryTime;
+						DepthSlope = clamp(DepthSlope, -SlopeClamp, SlopeClamp);
+					}
+					DepthHistory = SampleDepth;
+					DepthHistoryTime = SampleTime;
+				}
 
-				float DepthDiff = ReferenceDepth - DepthSample;
-				bool bHit = abs(DepthDiff + CompareTolerance) < CompareTolerance;
+				float DepthDiff = ReferenceDepth - DepthForComparison;
+				float HalfCompareTolerance = 0.5 * CompareTolerance;
+				bool bHit = abs(DepthDiff + HalfCompareTolerance) < HalfCompareTolerance;
 				if (bHit)
 				{
 					OccluderDistance = SMRTComputeOccluderDistance(RayState, SampleDepth);
diff --git a/Engine/Shaders/Shared/VirtualShadowMapDefinitions.h b/Engine/Shaders/Shared/VirtualShadowMapDefinitions.h
--- a/Engine/Shaders/Shared/VirtualShadowMapDefinitions.h
+++ b/Engine/Shaders/Shared/VirtualShadowMapDefinitions.h
@@ -14,3 +14,4 @@ VirtualShadowMapDefinitions.h: used in virtual shadow map shaders and C++ code t
 #define VIRTUAL_SHADOW_MAP_VISUALIZE_CACHED_PAGE				(1 << 3)
 #define VIRTUAL_SHADOW_MAP_VISUALIZE_SMRT_RAY_COUNT				(1 << 4)
 #define VIRTUAL_SHADOW_MAP_VISUALIZE_CLIPMAP_VIRTUAL_SPACE		(1 << 5)
+#define VIRTUAL_SHADOW_MAP_VISUALIZE_GENERAL_DEBUG				(1 << 6)
diff --git a/Engine/Source/Runtime/Engine/Private/VirtualShadowMapVisualizationData.cpp b/Engine/Source/Runtime/Engine/Private/VirtualShadowMapVisualizationData.cpp
--- a/Engine/Source/Runtime/Engine/Private/VirtualShadowMapVisualizationData.cpp
+++ b/Engine/Source/Runtime/Engine/Private/VirtualShadowMapVisualizationData.cpp
@@ -50,6 +50,13 @@ void FVirtualShadowMapVisualizationData::Initialize()
 			FModeType::Advanced,
 			VIRTUAL_SHADOW_MAP_VISUALIZE_SMRT_RAY_COUNT);
 
+		AddVisualizationMode(
+			TEXT("debug"),
+			LOCTEXT("GeneralDebug", "General Debug"),
+			LOCTEXT("GeneralDebugDesc", "General-purpose debug for use during shader development"),
+			FModeType::Advanced,
+			VIRTUAL_SHADOW_MAP_VISUALIZE_GENERAL_DEBUG);
+
 		AddVisualizationMode(
 			TEXT("clipmapvirtual"),
 			LOCTEXT("ClipmapVirtualSpace", "Clipmap Virtual Address Space"),
diff --git a/Engine/Source/Runtime/Renderer/Private/VirtualShadowMaps/VirtualShadowMapProjection.cpp b/Engine/Source/Runtime/Renderer/Private/VirtualShadowMaps/VirtualShadowMapProjection.cpp
--- a/Engine/Source/Runtime/Renderer/Private/VirtualShadowMaps/VirtualShadowMapProjection.cpp
+++ b/Engine/Source/Runtime/Renderer/Private/VirtualShadowMaps/VirtualShadowMapProjection.cpp
@@ -26,10 +26,10 @@
 #include "VirtualShadowMapClipmap.h"
 #include "HairStrands/HairStrandsData.h"
 
-static TAutoConsoleVariable<float> CVarContactShadowLength(
-	TEXT( "r.Shadow.Virtual.ContactShadowLength" ),
-	0.02f,
-	TEXT( "Length of the screen space contact shadow trace (smart shadow bias) before the virtual shadow map lookup." ),
+static TAutoConsoleVariable<float> CVarScreenRayLength(
+	TEXT( "r.Shadow.Virtual.ScreenRayLength" ),
+	0.01f,
+	TEXT( "Length of the screen space shadow trace (smart shadow bias) before the virtual shadow map lookup." ),
 	ECVF_Scalability | ECVF_RenderThreadSafe
 );
 
@@ -109,6 +109,13 @@ static TAutoConsoleVariable<float> CVarSMRTTexelDitherScale(
 	ECVF_Scalability | ECVF_RenderThreadSafe
 );
 
+TAutoConsoleVariable<int32> CVarSMRTExtrapolateWithSlope(
+	TEXT("r.Shadow.Virtual.SMRT.ExtrapolateWithSlope"),
+	1,
+	TEXT("Use slope-based extrapolation behind occluders. This can increase the quality of shadow penumbra on surfaces aligned with the light direction."),
+	ECVF_Scalability | ECVF_RenderThreadSafe
+);
+
 static TAutoConsoleVariable<int32> CVarForcePerLightShadowMaskClear(
 	TEXT( "r.Shadow.Virtual.ForcePerLightShadowMaskClear" ),
 	0,
@@ -193,6 +200,7 @@ class FVirtualShadowMapProjectionCS : public FGlobalShader
 	
 	class FDirectionalLightDim		: SHADER_PERMUTATION_BOOL("DIRECTIONAL_LIGHT");
 	class FSMRTAdaptiveRayCountDim	: SHADER_PERMUTATION_BOOL("SMRT_ADAPTIVE_RAY_COUNT");
+	class FSMRTExtrapolateSlopeDim	: SHADER_PERMUTATION_BOOL("SMRT_EXTRAPOLATE_WITH_SLOPE");
 	class FOnePassProjectionDim		: SHADER_PERMUTATION_BOOL("ONE_PASS_PROJECTION");
 	class FHairStrandsDim			: SHADER_PERMUTATION_BOOL("HAS_HAIR_STRANDS");
 	class FVisualizeOutputDim		: SHADER_PERMUTATION_BOOL("VISUALIZE_OUTPUT");
@@ -202,7 +210,8 @@ class FVirtualShadowMapProjectionCS : public FGlobalShader
 		FOnePassProjectionDim,
 		FSMRTAdaptiveRayCountDim,
 		FHairStrandsDim,
-		FVisualizeOutputDim>;
+		FVisualizeOutputDim,
+		FSMRTExtrapolateSlopeDim>;
 
 	BEGIN_SHADER_PARAMETER_STRUCT(FParameters, )
 		SHADER_PARAMETER_STRUCT_INCLUDE(FVirtualShadowMapSamplingParameters, SamplingParameters)
@@ -212,7 +221,7 @@ class FVirtualShadowMapProjectionCS : public FGlobalShader
 		SHADER_PARAMETER_RDG_UNIFORM_BUFFER(FStrataGlobalUniformParameters, Strata)
 		SHADER_PARAMETER_STRUCT_REF(FViewUniformShaderParameters, View)
 		SHADER_PARAMETER(FIntVector4, ProjectionRect)
-		SHADER_PARAMETER(float, ContactShadowLength)
+		SHADER_PARAMETER(float, ScreenRayLength)
 		SHADER_PARAMETER(float, NormalBias)
 		SHADER_PARAMETER(uint32, SMRTRayCount)
 		SHADER_PARAMETER(uint32, SMRTSamplesPerRay)
@@ -293,7 +302,7 @@ static void RenderVirtualShadowMapProjectionCommon(
 	PassParameters->SceneTexturesStruct = SceneTextures.UniformBuffer;
 	PassParameters->View = View.ViewUniformBuffer;
 	PassParameters->ProjectionRect = FIntVector4(ProjectionRect.Min.X, ProjectionRect.Min.Y, ProjectionRect.Max.X, ProjectionRect.Max.Y);
-	PassParameters->ContactShadowLength = CVarContactShadowLength.GetValueOnRenderThread();
+	PassParameters->ScreenRayLength = CVarScreenRayLength.GetValueOnRenderThread();
 	PassParameters->NormalBias = GetNormalBiasForShader();
 	PassParameters->InputType = uint32(InputType);
 	PassParameters->bCullBackfacingPixels = VirtualShadowMapArray.ShouldCullBackfacingPixels() ? 1 : 0;
@@ -357,7 +366,8 @@ static void RenderVirtualShadowMapProjectionCommon(
 	PermutationVector.Set< FVirtualShadowMapProjectionCS::FDirectionalLightDim >( bDirectionalLight );
 	PermutationVector.Set< FVirtualShadowMapProjectionCS::FOnePassProjectionDim >( bOnePassProjection );
 	PermutationVector.Set< FVirtualShadowMapProjectionCS::FSMRTAdaptiveRayCountDim >( bAdaptiveRayCount );
-	PermutationVector.Set< FVirtualShadowMapProjectionCS::FHairStrandsDim >( bHasHairStrandsData ? 1 : 0 );
+	PermutationVector.Set< FVirtualShadowMapProjectionCS::FSMRTExtrapolateSlopeDim >( CVarSMRTExtrapolateWithSlope.GetValueOnRenderThread() != 0 );
+	PermutationVector.Set< FVirtualShadowMapProjectionCS::FHairStrandsDim >( bHasHairStrandsData );
 	PermutationVector.Set< FVirtualShadowMapProjectionCS::FVisualizeOutputDim >( bDebugOutput );
 
 	auto ComputeShader = View.ShaderMap->GetShader< FVirtualShadowMapProjectionCS >( PermutationVector );
