Change the presicion of length2 from half to float to fix artifacts on mobile vulkan.

diff --git a/Engine/Shaders/Private/Common.ush b/Engine/Shaders/Private/Common.ush
--- a/Engine/Shaders/Private/Common.ush
+++ b/Engine/Shaders/Private/Common.ush
@@ -589,15 +589,15 @@ MaterialFloat Luminance( MaterialFloat3 LinearColor )
 	return dot( LinearColor, MaterialFloat3( 0.3, 0.59, 0.11 ) );
 }
 
-MaterialFloat length2(MaterialFloat2 v)
+float length2(float2 v)
 {
 	return dot(v, v);
 }
-MaterialFloat length2(MaterialFloat3 v)
+float length2(float3 v)
 {
 	return dot(v, v);
 }
-MaterialFloat length2(MaterialFloat4 v)
+float length2(float4 v)
 {
 	return dot(v, v);
 }
