Implements CONFIG_MEDIAN_DECIMATE for r.TSR.History.HighFrequencyOnly=1

diff --git a/Engine/Shaders/Private/TemporalSuperResolution/TSRDecimateHistory.usf b/Engine/Shaders/Private/TemporalSuperResolution/TSRDecimateHistory.usf
--- a/Engine/Shaders/Private/TemporalSuperResolution/TSRDecimateHistory.usf
+++ b/Engine/Shaders/Private/TemporalSuperResolution/TSRDecimateHistory.usf
@@ -11,6 +11,8 @@
 
 #define CONFIG_OUTPUT_HALF_RES (DIM_OUTPUT_HALF_RES)
 
+#define CONFIG_MEDIAN_DECIMATE 1
+
 
 //------------------------------------------------------- PARAMETERS
 
@@ -52,6 +54,42 @@ groupshared tsr_half4 SharedHistoryColorAndMask[TILE_SIZE * TILE_SIZE];
 
 //------------------------------------------------------- ENTRY POINT
 
+void LHOperator2(tsr_halfC A, tsr_halfC B, out tsr_halfC L, out tsr_halfC H)
+{
+	L = min(A, B);
+	H = max(A, B);
+}
+
+void LMHOperator3(tsr_halfC A, tsr_halfC B, tsr_halfC C, out tsr_halfC L, out tsr_halfC M, out tsr_halfC H)
+{
+	// TODO: min3, max3, med3.
+	tsr_halfC X = min(B, C);
+	tsr_halfC Y = max(B, C);
+	
+	L = min(A, X);
+	tsr_halfC Z = max(A, X);
+
+	M = min(Z, Y);
+	H = max(Z, Y);
+}
+
+tsr_halfC Median3(tsr_halfC A, tsr_halfC B, tsr_halfC C)
+{
+	// TODO: med3.
+	tsr_halfC L, M, H;
+	LMHOperator3(A, B, C, L, M, H);
+	return M;
+}
+
+tsr_halfC Median5(tsr_halfC In[5])
+{
+	LMHOperator3(In[1], In[2], In[3], /* out */ In[1], /* out */ In[2], /* out */ In[3]);
+	LHOperator2(In[0], In[4], /* out */ In[0], /* out */ In[4]);
+	LHOperator2(In[0], In[1], /* out */ In[0], /* out */ In[1]);
+	LHOperator2(In[3], In[4], /* out */ In[3], /* out */ In[4]);
+	return Median3(In[1], In[2], In[3]);
+}
+
 tsr_half3 TransformToComparisonColorSpace(tsr_half3 Color)
 {
 	return Color;
@@ -168,36 +206,42 @@ void MainCS(
 			PrevHistoryInfo_UVViewportBilinearMax);
 		
 		float2 PrevInputBufferUV = (InputInfo_ScreenPosToViewportScale * PrevOutputScreenPos + InputInfo_ScreenPosToViewportBias) * InputInfo_ExtentInverse;
+		
+		#if DIM_HIGH_FREQUENCY_ONLY && CONFIG_MEDIAN_DECIMATE
+		{
+			const float kTopSecretSauce = 2.0;
+
+			float2 PrevUV[5];
+			PrevUV[0] = KernelHistoryBufferUV;
+			PrevUV[1] = KernelHistoryBufferUV + (kTopSecretSauce * float2(+1.0, +1.0)) * PrevHistoryInfo_ExtentInverse;
+			PrevUV[2] = KernelHistoryBufferUV + (kTopSecretSauce * float2(-1.0, +1.0)) * PrevHistoryInfo_ExtentInverse;
+			PrevUV[3] = KernelHistoryBufferUV + (kTopSecretSauce * float2(+1.0, -1.0)) * PrevHistoryInfo_ExtentInverse;
+			PrevUV[4] = KernelHistoryBufferUV + (kTopSecretSauce * float2(-1.0, -1.0)) * PrevHistoryInfo_ExtentInverse;
 
-		#if DIM_HIGH_FREQUENCY_ONLY
+			tsr_halfC PrevDecimatedColorArray[5];
+			
+			UNROLL_N(5)
+			for (uint i = 0; i < 5; i++)
+			{
+				float2 UV = KernelHistoryBufferUV = clamp(
+					PrevUV[i],
+					PrevHistoryInfo_UVViewportBilinearMin,
+					PrevHistoryInfo_UVViewportBilinearMax);
+					
+				tsr_half2 PrevMetadata = PrevHistory_Metadata.SampleLevel(GlobalBilinearClampedSampler, UV, 0);
+				PrevDecimatedColorArray[i] = PrevHistory_HighFrequency.SampleLevel(GlobalBilinearClampedSampler, UV, 0) * SafeRcp(PrevMetadata.r);
+			}
+
+			DecimatedColor = Median5(PrevDecimatedColorArray);
+		}
+		#elif DIM_HIGH_FREQUENCY_ONLY
 		{
 			tsr_half2 Metadata = PrevHistory_Metadata.SampleLevel(GlobalBilinearClampedSampler, KernelHistoryBufferUV, 0);
 			DecimatedColor = PrevHistory_HighFrequency.SampleLevel(GlobalBilinearClampedSampler, KernelHistoryBufferUV, 0) * SafeRcp(Metadata.r);
 		}
-		#elif 1
-		{
-			DecimatedColor = PrevHistory_LowFrequency.SampleLevel(GlobalBilinearClampedSampler, KernelHistoryBufferUV, 0);
-		}
 		#else
 		{
-			DecimatedColor = 0.0;
-			UNROLL_N(4)
-			for (uint i = 0; i < 4; i++)
-			{
-				const float2 PixelOffset = float2(Offsets2x2[i]) - 0.5;
-
-				float2 SampleBufferUV = KernelHistoryBufferUV + PixelOffset * PrevHistoryInfo_ExtentInverse; 
-				SampleBufferUV = clamp(
-					SampleBufferUV,
-					PrevHistoryInfo_UVViewportBilinearMin,
-					PrevHistoryInfo_UVViewportBilinearMax);
-				
-				tsr_halfC HistorySample = PrevHistory_LowFrequency.SampleLevel(GlobalBilinearClampedSampler, SampleBufferUV, 0);
-
-				DecimatedColor += HistorySample;
-			}
-
-			DecimatedColor *= rcp(4.0);
+			DecimatedColor = PrevHistory_LowFrequency.SampleLevel(GlobalBilinearClampedSampler, KernelHistoryBufferUV, 0);
 		}
 		#endif
 
