Strata - fixed legacy material issues with opaque rough refractions.

diff --git a/Engine/Shaders/Private/BasePassPixelShader.usf b/Engine/Shaders/Private/BasePassPixelShader.usf
--- a/Engine/Shaders/Private/BasePassPixelShader.usf
+++ b/Engine/Shaders/Private/BasePassPixelShader.usf
@@ -1729,6 +1729,17 @@ void FPixelShaderInOut_MainPS(
 	Out.StrataOutput[2] = StrataPackTopLayerData(TopLayerData);
 #if STRATA_OUTPUT_ROUGH_REFRACTION
 	OpaqueRoughRefractionData.OpaqueRoughRefractionEnabled = true; // As long as we output to rough refraction, we want it to be enabled for the tiling process.
+#if STRATA_USES_CONVERSION_FROM_LEGACY
+	// When running legacy support, we want to only enable rough refraction for multi layered materials (eye, clear coat) or SSS.
+	// This is because the compiler will always see a vertical layer. But we do not want to enable opaque rough refraction for all the legacy materials.
+	if (StrataPixelHeader.StrataTree.BSDFCount > 1)
+	{
+		// We only enable opaque rough refraction if the material is coated or has SSS.
+		// BSDFs and Operators indices must match StrataLegacyConversion.ush.
+		OpaqueRoughRefractionData.OpaqueRoughRefractionEnabled = (StrataPixelHeader.StrataTree.Operators[2].Weight > 0.0f && StrataPixelHeader.StrataTree.Operators[3].Weight > 0.0f) 
+			|| BSDF_GETHASSSS(StrataPixelHeader.StrataTree.BSDFs[0]) || BSDF_GETHASSSS(StrataPixelHeader.StrataTree.BSDFs[1]);
+	}
+#endif
 	StrataStruct.OpaqueRoughRefractionTextureUAV[PixelPos] = StrataPackOpaqueRoughRefractionData(OpaqueRoughRefractionData);
 #endif
 
diff --git a/Engine/Shaders/Private/Strata/StrataEvaluation.ush b/Engine/Shaders/Private/Strata/StrataEvaluation.ush
--- a/Engine/Shaders/Private/Strata/StrataEvaluation.ush
+++ b/Engine/Shaders/Private/Strata/StrataEvaluation.ush
@@ -577,7 +577,8 @@ FStrataEvaluateResult StrataEvaluateBSDFCommon(FStrataBSDFContext BSDFContext, F
 				Sample.ThinTransmissionPathValue = Diffuse_Lambert(BackFaceWhiteDiffuseColor) * SlabTransmittance * DirectionalAlbedo_SpecularTransmission;
 				Sample.ThinTransmissionPathPDF = BackFaceNoLSat / PI;
 
-				// At this stage, Sample.bSubsurface should have been setup correctly already if SSS is enabled with two-sided lighting.
+				// We do not want thin (a.k.a. two sided lighting) contribution to go the SSS post process path.
+				Sample.bSubsurface = false;
 
 				Sample.IntegratedDiffuseValue += ShadowTerms.TransmissionShadow * Sample.ThinTransmissionPathValue * BackFaceNoLSat * AreaLightContext.Falloff * AreaLightContext.AreaLight.FalloffColor;
 			}
@@ -1028,7 +1029,8 @@ FStrataEnvLightResult StrataEvaluateForEnvLight(FStrataBSDFContext BSDFContext,
 				// The diffuse is attenuated by the slab transmittance and the GGX interface SpecularTransmission.
 				StrataEnvLightResult.DiffuseBackFaceWeight = Diffuse_Lambert(BackFaceWhiteDiffuseColor) * SlabTransmittance * DirectionalAlbedo_SpecularTransmission;
 
-				// At this stage, StrataEnvLightResult.bSubsurface should have been setup correctly already if SSS is enabled to be applied on a thin surface.
+				// We do not want thin (a.k.a. two sided lighting) contribution to go the SSS post process path.
+				StrataEnvLightResult.bSubsurface = false;
 			}
 			#endif
 
diff --git a/Engine/Shaders/Shared/StrataDefinitions.h b/Engine/Shaders/Shared/StrataDefinitions.h
--- a/Engine/Shaders/Shared/StrataDefinitions.h
+++ b/Engine/Shaders/Shared/StrataDefinitions.h
@@ -8,7 +8,7 @@
 #pragma once
 
 // Change this to force recompilation of all strata dependent shaders (use https://www.random.org/cgi-bin/randbyte?nbytes=4&format=h)
-#define STRATA_SHADER_VERSION 0x04499eb2  
+#define STRATA_SHADER_VERSION 0xa8970d72
 
 
 
diff --git a/Engine/Source/Runtime/Engine/Private/Materials/HLSLMaterialTranslator.cpp b/Engine/Source/Runtime/Engine/Private/Materials/HLSLMaterialTranslator.cpp
--- a/Engine/Source/Runtime/Engine/Private/Materials/HLSLMaterialTranslator.cpp
+++ b/Engine/Source/Runtime/Engine/Private/Materials/HLSLMaterialTranslator.cpp
@@ -9801,7 +9801,7 @@ bool FHLSLMaterialTranslator::StrataGenerateDerivedMaterialOperatorData()
 
 		const EStrataBlendMode StrataBlendMode = Material->GetStrataBlendMode();
 		const bool bIsOpaqueOrMasked = StrataBlendMode == EStrataBlendMode::SBM_Opaque || StrataBlendMode == EStrataBlendMode::SBM_Masked;
-		bStrataOutputsOpaqueRoughRefractions = bStrataUsesVerticalLayering && !bStrataUsesConversionFromLegacy && bIsOpaqueOrMasked;
+		bStrataOutputsOpaqueRoughRefractions = bStrataUsesVerticalLayering && bIsOpaqueOrMasked;
 	}
 
 	// STRATA_TODO: operation using parameter blending are not actually discarded so they still occupy a spot in the operation array in the compiler and in the shader also.
