Fix pointlight rendering artifact when mobile deferred is enabled on Vulkan.

diff --git a/Engine/Shaders/Private/MobileDeferredShading.usf b/Engine/Shaders/Private/MobileDeferredShading.usf
--- a/Engine/Shaders/Private/MobileDeferredShading.usf
+++ b/Engine/Shaders/Private/MobileDeferredShading.usf
@@ -216,7 +216,8 @@ void GetLightGridLocalLighting(
 		// extra-early out since we know light grid is sloppy and all lights in list are radial (have a range)
 		// appears useless
 		float invLightRadiusSq = LightData.InvRadius*LightData.InvRadius;
-		if (length2(TranslatedWorldPosition - LightData.TranslatedWorldPosition) * invLightRadiusSq > 1.0f)
+		float3 LightToCamera = LightData.TranslatedWorldPosition - TranslatedWorldPosition;
+		if (dot(LightToCamera, LightToCamera) * invLightRadiusSq > 1.0f)
 		{
 			continue;
 		}
