Regroup common Strata tile code into StrataTile.ush

diff --git a/Engine/Shaders/Private/Strata/StrataMaterialClassification.usf b/Engine/Shaders/Private/Strata/StrataMaterialClassification.usf
--- a/Engine/Shaders/Private/Strata/StrataMaterialClassification.usf
+++ b/Engine/Shaders/Private/Strata/StrataMaterialClassification.usf
@@ -5,18 +5,12 @@
 #define STRATA_INLINE_SHADING 0
 #define STRATA_SSS_MATERIAL_OVERRIDE 0
 #include "/Engine/Private/Strata/Strata.ush"
+#include "StrataTile.ush"
 
-#define STRATA_TILE_SIZE 8
 #define GROUP_THREAD_COUNT (STRATA_TILE_SIZE * STRATA_TILE_SIZE)
 
 ////////////////////////////////////////////////////////////////////////////////////////////////////////////
 
-// DecodeTile is in StrataTiling.usf.
-uint EncodeTile(uint2 TileCoord)
-{
-	return TileCoord.x | (TileCoord.y << 16); // assumes 16bit is enough to represent a tiled resolution up to 65,535 :)
-}
-
 #if SHADER_TILE_CATEGORIZATION
 int bRectPrimitive;
 int2 ViewResolution;
@@ -155,7 +149,7 @@ void TileMainCS(uint2 DispatchThreadId : SV_DispatchThreadID, uint LinearIndex :
 
 	if (LinearIndex < 1 && bTileContainsStrata)
 	{
-		uint EncodedTile = EncodeTile(GroupId.xy);
+		uint EncodedTile = StrataPackTile(GroupId.xy);
 
 		if (bTileContainsComplex)
 		{
diff --git a/Engine/Shaders/Private/Strata/StrataTile.usf b/Engine/Shaders/Private/Strata/StrataTile.usf
new file mode 100644
--- /dev/null
+++ b/Engine/Shaders/Private/Strata/StrataTile.usf
@@ -0,0 +1,70 @@
+// Copyright Epic Games, Inc. All Rights Reserved.
+
+#include "/Engine/Private/Common.ush"
+#include "/Engine/Shared/StrataDefinitions.h"   
+#include "StrataTile.ush"
+
+
+////////////////////////////////////////////////////////////////////////////////////////////////////////////
+
+#if SHADER_TILE_VS
+
+float4 OutputViewSizeAndInvSize;
+float4 OutputBufferSizeAndInvSize;
+float4x4 ViewScreenToTranslatedWorld;
+
+Buffer<uint> TileListBuffer;
+
+void StrataTilePassVS(
+	in uint InVertexId : SV_VertexID,
+	in uint InInstanceId : SV_InstanceID,
+#if PERMUTATION_ENABLE_TEXCOORD_SCREENVECTOR
+	out float2 OutTexCoord : TEXCOORD0,		// Used by DeferredLightPixelMain directional light
+	out float3 OutScreenVector : TEXCOORD1,	// Used by DeferredLightPixelMain directional light
+#endif
+#if PERMUTATION_ENABLE_DEBUG
+	out uint2 OutTileCoord : TILE_COORD,	// Used by StencilMainPS below
+#endif
+	out float4 OutPosition : SV_POSITION)	// Use by StencilMainPS
+{
+	const uint2 TileCoord = StrataUnpackTile(TileListBuffer[InInstanceId]);
+
+	uint2 TileVertex = TileCoord * STRATA_TILE_SIZE;
+	TileVertex.x += InVertexId == 1 || InVertexId == 2 || InVertexId == 4 ? STRATA_TILE_SIZE : 0;
+	TileVertex.y += InVertexId == 2 || InVertexId == 4 || InVertexId == 5 ? STRATA_TILE_SIZE : 0;
+
+	OutPosition = float4(float2(TileVertex) * OutputViewSizeAndInvSize.zw * float2(2.0f, -2.0f) + float2(-1.0, 1.0f), 0.5f, 1.0f);
+
+#if PERMUTATION_ENABLE_TEXCOORD_SCREENVECTOR
+	OutTexCoord = float2(TileVertex) * OutputBufferSizeAndInvSize.zw;
+	OutScreenVector = mul(float4(OutPosition.xy, 1, 0), ViewScreenToTranslatedWorld).xyz;
+#endif
+
+#if PERMUTATION_ENABLE_DEBUG
+	OutTileCoord = TileCoord;
+#endif
+}
+
+#endif // SHADER_TILE_VS
+
+////////////////////////////////////////////////////////////////////////////////////////////////////////////
+
+#if SHADER_STENCIL_TAGGING_PS
+
+float4 DebugTileColor;
+
+void StencilTaggingMainPS(
+	in uint2 InTileCoord : TILE_COORD,
+	in float4 SVPos : SV_POSITION,
+	out float4 OutColor0 : SV_Target0)
+{
+	const bool bTileX = (InTileCoord.x & 1) == 0;
+	const bool bTileY = (InTileCoord.y & 1) == 0;
+	const bool bChecker = (bTileX && bTileY) || (!bTileX && !bTileY);
+	OutColor0 = DebugTileColor * float4(1.0f, 1.0f, 1.0f, bChecker ? 0.33f : 0.66f);
+	OutColor0.rgb *= OutColor0.a; // premultiplied alpha blending assuming alpha is opacity==coverage
+}
+
+#endif // SHADER_STENCIL_PS
+
+
diff --git a/Engine/Shaders/Private/Strata/StrataTile.ush b/Engine/Shaders/Private/Strata/StrataTile.ush
new file mode 100644
--- /dev/null
+++ b/Engine/Shaders/Private/Strata/StrataTile.ush
@@ -0,0 +1,19 @@
+// Copyright Epic Games, Inc. All Rights Reserved.
+
+#pragma once
+
+#include "/Engine/Private/Common.ush"
+#include "/Engine/Shared/StrataDefinitions.h"   
+
+////////////////////////////////////////////////////////////////////////////////////////////////////////////
+// Tile
+
+uint2 StrataUnpackTile(uint In)
+{
+	return uint2(In & 0xFFFF, In >> 16);
+}
+
+uint StrataPackTile(uint2 TileCoord)
+{
+	return TileCoord.x | (TileCoord.y << 16); // assumes 16bit is enough to represent a tiled resolution up to 65,535 :)
+}
diff --git a/Engine/Shaders/Private/Strata/StrataTiling.usf b/Engine/Shaders/Private/Strata/StrataTiling.usf
deleted file mode 100644
diff --git a/Engine/Source/Runtime/Renderer/Private/Strata/Strata.cpp b/Engine/Source/Runtime/Renderer/Private/Strata/Strata.cpp
--- a/Engine/Source/Runtime/Renderer/Private/Strata/Strata.cpp
+++ b/Engine/Source/Runtime/Renderer/Private/Strata/Strata.cpp
@@ -539,8 +539,8 @@ class FStrataMaterialStencilTaggingPassPS : public FGlobalShader
 	}
 };
 
-IMPLEMENT_GLOBAL_SHADER(FStrataTilePassVS, "/Engine/Private/Strata/StrataTiling.usf", "StrataTilePassVS", SF_Vertex);
-IMPLEMENT_GLOBAL_SHADER(FStrataMaterialStencilTaggingPassPS, "/Engine/Private/Strata/StrataTiling.usf", "StencilTaggingMainPS", SF_Pixel);
+IMPLEMENT_GLOBAL_SHADER(FStrataTilePassVS, "/Engine/Private/Strata/StrataTile.usf", "StrataTilePassVS", SF_Vertex);
+IMPLEMENT_GLOBAL_SHADER(FStrataMaterialStencilTaggingPassPS, "/Engine/Private/Strata/StrataTile.usf", "StencilTaggingMainPS", SF_Pixel);
 
 static FStrataTileParameter InternalSetTileParameters(FRDGBuilder* GraphBuilder, const FViewInfo& View, const EStrataTileType TileType)
 {
