Expose Strata BSDF type accessor.

diff --git a/Engine/Shaders/Private/Strata/Strata.ush b/Engine/Shaders/Private/Strata/Strata.ush
--- a/Engine/Shaders/Private/Strata/Strata.ush
+++ b/Engine/Shaders/Private/Strata/Strata.ush
@@ -1077,6 +1077,7 @@ struct FStrataAddressing
 
 	uint  ReadBytes;
 };
+
 FStrataAddressing GetStrataPixelDataByteOffset(uint2 PixelPos, uint2 ViewBufferSize, uint StrataMaxBytesPerPixel)
 {
 	FStrataAddressing StrataAddressing = (FStrataAddressing)0;
@@ -1954,6 +1955,11 @@ void EnableSlabBSDFIsThin(inout FStrataBSDF BSDF)
 ///////////////////////////////////////////////////////////////////////////////
 // Generic BSDF accessors
 
+uint StrataGetBSDFType(in FStrataBSDF BSDF)
+{
+	return BSDF_GETTYPE(BSDF);
+}
+
 float3 StrataGetBSDFBaseColor(in FStrataBSDF BSDF)
 {
 	switch (BSDF_GETTYPE(BSDF))
