Fix cooking issue.

diff --git a/Engine/Shaders/Private/Strata/StrataTile.ush b/Engine/Shaders/Private/Strata/StrataTile.ush
--- a/Engine/Shaders/Private/Strata/StrataTile.ush
+++ b/Engine/Shaders/Private/Strata/StrataTile.ush
@@ -6,6 +6,8 @@
 #include "/Engine/Shared/StrataDefinitions.h"   
 #include "Strata.ush"
 
+#if STRATA_ENABLED
+
 ////////////////////////////////////////////////////////////////////////////////////////////////////////////
 // Strata material type tile
 
@@ -126,3 +128,5 @@ void StrataSeekBSDF(inout FStrataAddressing Addressing, FStrataBSDFOffset Offset
 		Addressing.ReadBytes = Addressing.CurrentIndex * 4;
 	}
 }
+
+#endif // STRATA_ENABLED
\ No newline at end of file
