Fix another CIS cooking issue.

diff --git a/Engine/Shaders/Private/Strata/StrataTile.usf b/Engine/Shaders/Private/Strata/StrataTile.usf
--- a/Engine/Shaders/Private/Strata/StrataTile.usf
+++ b/Engine/Shaders/Private/Strata/StrataTile.usf
@@ -27,7 +27,12 @@ void StrataTilePassVS(
 #endif
 	out float4 OutPosition : SV_POSITION)	// Use by StencilMainPS
 {
+#if STRATA_ENABLED
 	const uint2 TileCoord = StrataUnpackTile(TileListBuffer[InInstanceId]);
+#else
+	// Not used, but this shader is compiled for non-strata path for setup reason, but not used
+	const uint2 TileCoord = 0;
+#endif
 
 	uint2 TileVertex = TileCoord * STRATA_TILE_SIZE;
 	TileVertex.x += InVertexId == 1 || InVertexId == 2 || InVertexId == 4 ? STRATA_TILE_SIZE : 0;
