Resolve shadowed name errors in FVolumeLightSampleSetup during shader compilation on some platforms

diff --git a/Engine/Shaders/Private/PathTracing/Light/PathTracingLightCommon.ush b/Engine/Shaders/Private/PathTracing/Light/PathTracingLightCommon.ush
--- a/Engine/Shaders/Private/PathTracing/Light/PathTracingLightCommon.ush
+++ b/Engine/Shaders/Private/PathTracing/Light/PathTracingLightCommon.ush
@@ -55,15 +55,15 @@ struct FVolumeLightSampleSetup
 			// equi-angular pdf
 			const float TanTheta = tan(lerp(ThetaMin, ThetaMax, RandSample));
 			const float T = Height * TanTheta + T0;
-			const float Pdf = rcp((1 + TanTheta * TanTheta) * K);
-			return float2(T, Pdf);
+			const float PdfLocal = rcp((1 + TanTheta * TanTheta) * K);
+			return float2(T, PdfLocal);
 		}
 		else
 		{
 			// uniform pdf
 			const float T = lerp(VolumeTMin, VolumeTMax, RandSample);
-			const float Pdf = rcp(K);
-			return float2(T, Pdf);
+			const float PdfLocal = rcp(K);
+			return float2(T, PdfLocal);
 		}
 	}
 
