PR #8969: Fix typo in RayTracingMaterialHitShaders.usf (Contributed by [REDACTED_PERSON])

diff --git a/Engine/Shaders/Private/RayTracing/RayTracingMaterialHitShaders.usf b/Engine/Shaders/Private/RayTracing/RayTracingMaterialHitShaders.usf
--- a/Engine/Shaders/Private/RayTracing/RayTracingMaterialHitShaders.usf
+++ b/Engine/Shaders/Private/RayTracing/RayTracingMaterialHitShaders.usf
@@ -389,22 +389,22 @@ RAY_TRACING_ENTRY_CLOSEST_HIT(MaterialCHS,
 	#endif
 #endif
 
-	// Override custom data if sub-surface material
-#if MATERIAL_SHADINGMODEL_SUBSURFACE || MATERIAL_SHADINGMODEL_TWOSIDED_FOLIAGE || MATERIAL_SHADINGMODEL_CLOTH || MATERIAL_SHADINGMODEL_CLOTH || MATERIAL_SHADINGMODEL_SUBSURFACE_PROFILE || MATERIAL_SHADINGMODEL_EYE
-	if (Payload.ShadingModelID == SHADINGMODELID_SUBSURFACE ||
-		Payload.ShadingModelID == SHADINGMODELID_TWOSIDED_FOLIAGE ||
-		Payload.ShadingModelID == SHADINGMODELID_CLOTH ||
-		Payload.ShadingModelID == SHADINGMODELID_SUBSURFACE_PROFILE |
-		Payload.ShadingModelID == SHADINGMODELID_EYE)
+#if MATERIAL_SHADINGMODEL_CLOTH
+	if (Payload.ShadingModelID == SHADINGMODELID_CLOTH)
 	{
 		Payload.CustomData = GetMaterialSubsurfaceData(PixelMaterialInputs);
+		Payload.CustomData.a = saturate(GetMaterialCustomData0(MaterialParameters));
 	}
 #endif
 
-#if MATERIAL_SHADINGMODEL_CLOTH
-	if (Payload.ShadingModelID == SHADINGMODELID_CLOTH)
+	// Override custom data if sub-surface material
+#if MATERIAL_SHADINGMODEL_SUBSURFACE || MATERIAL_SHADINGMODEL_TWOSIDED_FOLIAGE || MATERIAL_SHADINGMODEL_SUBSURFACE_PROFILE || MATERIAL_SHADINGMODEL_EYE
+	if (Payload.ShadingModelID == SHADINGMODELID_SUBSURFACE ||
+		Payload.ShadingModelID == SHADINGMODELID_TWOSIDED_FOLIAGE ||
+		Payload.ShadingModelID == SHADINGMODELID_SUBSURFACE_PROFILE ||
+		Payload.ShadingModelID == SHADINGMODELID_EYE)
 	{
-		Payload.CustomData.a = saturate(GetMaterialCustomData0(MaterialParameters));
+		Payload.CustomData = GetMaterialSubsurfaceData(PixelMaterialInputs);
 	}
 #endif
 
