Moves parts of the velocity holefilling into the TSR DecimateHistory pass

diff --git a/Engine/Shaders/Private/TemporalSuperResolution/TSRDecimateHistory.usf b/Engine/Shaders/Private/TemporalSuperResolution/TSRDecimateHistory.usf
--- a/Engine/Shaders/Private/TemporalSuperResolution/TSRDecimateHistory.usf
+++ b/Engine/Shaders/Private/TemporalSuperResolution/TSRDecimateHistory.usf
@@ -9,17 +9,33 @@
 
 #define CONFIG_KERNEL 0
 
-#define CONFIG_OUTPUT_HALF_RES (DIM_OUTPUT_HALF_RES)
+#define CONFIG_OUTPUT_HALF_RES 0 // TODO
 
 #define CONFIG_MEDIAN_DECIMATE 1
 
+#define CONFIG_MANUAL_LDS_SPILL 1
+
+#define CONFIG_VELOCITY_SAMPLES 9
+#if CONFIG_VELOCITY_SAMPLES != 1
+	#define CONFIG_VELOCITY_SAMPLE_PAIRS ((CONFIG_VELOCITY_SAMPLES - 1) / 2)
+
+	#if CONFIG_BUTTERFLY_KERNEL
+		#if CONFIG_VELOCITY_SAMPLES == 9
+			#define CONFIG_VELOCITY_PRE_FETCH_SAMPLES 4
+		#else
+			#define CONFIG_VELOCITY_PRE_FETCH_SAMPLES 3
+		#endif
+	#else // !CONFIG_BUTTERFLY_KERNEL
+		#define CONFIG_VELOCITY_PRE_FETCH_SAMPLES (CONFIG_VELOCITY_SAMPLES)
+	#endif
+#endif // CONFIG_VELOCITY_SAMPLES != 1
+
 
 //------------------------------------------------------- PARAMETERS
 
 float4x4 RotationalClipToPrevClip;
 float3 OutputQuantizationError;
 
-Texture2D<tsr_halfC> InputSceneColorTexture;
 Texture2D<float2> DilatedVelocityTexture;
 Texture2D<float> ClosestDepthTexture;
 Texture2D<uint> PrevUseCountTexture;
@@ -28,21 +44,28 @@ Texture2D<uint> ParallaxFactorTexture;
 
 Texture2D<tsr_halfC> PrevHistory_LowFrequency;
 Texture2D<tsr_halfC> PrevHistory_HighFrequency;
+Texture2D<tsr_halfC> PrevHistory_HighFrequencyHalfRes;
 Texture2D<tsr_half2> PrevHistory_Metadata;
 Texture2D<tsr_half4> PrevHistory_Translucency;
 Texture2D<tsr_subpixel_details> PrevHistory_SubpixelDetails;
 
-RWTexture2D<tsr_halfC> HalfResSceneColorOutput;
 RWTexture2D<tsr_halfC> HalfResPredictionSceneColorOutput;
 RWTexture2D<tsr_half> HalfResParallaxRejectionMaskOutput;
 
 RWTexture2D<tsr_halfC> PredictionSceneColorOutput;
+
+RWTexture2D<float2> HoleFilledVelocityOutput;
+RWTexture2D<tsr_half2> HoleFilledVelocityMaskOutput;
 RWTexture2D<tsr_half> ParallaxRejectionMaskOutput;
 
 
 //------------------------------------------------------- LDS
 
-groupshared tsr_halfC SharedInputColor[TILE_SIZE * TILE_SIZE];
+#if CONFIG_MANUAL_LDS_SPILL
+	groupshared float2 SharedArray0[TILE_SIZE * TILE_SIZE];
+	groupshared float  SharedArray1[TILE_SIZE * TILE_SIZE];
+#endif
+
 
 #if CONFIG_SCENE_COLOR_APLHA
 groupshared tsr_halfC SharedHistoryColor[TILE_SIZE * TILE_SIZE];
@@ -106,147 +129,356 @@ void MainCS(
 	uint GroupThreadIndex : SV_GroupIndex)
 {
 	float4 Debug = 0;
-	
-	uint2 DispatchThreadId = (
-		ZOrder2D(GroupThreadIndex, uint(log2(float(TILE_SIZE)))) +
-		GroupId * uint2(TILE_SIZE, TILE_SIZE));
 
-	uint2 InputPixelPos = InputInfo_ViewportMin + DispatchThreadId;
-	
-	float2 ViewportUV = (float2(DispatchThreadId) + 0.5f) * InputInfo_ViewportSizeInverse;
-	float2 ScreenPos = ViewportUVToScreenPos(ViewportUV);
-	float2 InputBufferUV = (InputPixelPos + 0.5) * InputInfo_ExtentInverse;
-	
-	float2 EncodedVelocity = DilatedVelocityTexture[InputPixelPos];
-	float DeviceZ = ClosestDepthTexture[InputPixelPos];
-	tsr_half ParallaxFactor = UncompressParallaxFactor((ParallaxFactorTexture[InputPixelPos] >> 4) & 0xF);
-	tsr_halfC InputSceneColor = InputSceneColorTexture[InputPixelPos];
-
-	float WorldDepth = ConvertFromDeviceZ(DeviceZ);
-	
-	// Compute the parralax rejection mask
-	tsr_half ParallaxRejectionMask = 0.0;
+	// Sample current frame informations.
+	float2 EncodedScreenVelocity;
+	float DeviceZ;
+	tsr_half ParallaxFactor;
+	tsr_half RejectHighFrequency = tsr_half(0.0);
+	ISOLATE
 	{
-		float2 PrevScreenPos = ScreenPos - DecodeVelocityFromTexture(float4(EncodedVelocity, 0.0, 0.0)).xy;
-		float2 PrevInputBufferUV = (InputInfo_ScreenPosToViewportScale * PrevScreenPos + InputInfo_ScreenPosToViewportBias) * InputInfo_ExtentInverse;
-		
-		FBilinearSampleInfos BilinearInter = GetBilinearSampleLevelInfos(PrevInputBufferUV, InputInfo_Extent, InputInfo_ExtentInverse);
+		tsr_ushort2 InputPixelPos = (
+			Map8x8Tile2x2Lane(GroupThreadIndex) +
+			(tsr_ushort2(InputInfo_ViewportMin) + tsr_ushort2(GroupId) * tsr_ushort2(TILE_SIZE, TILE_SIZE)));
+			
+		#if CONFIG_VELOCITY_SAMPLES != 1
+		float2 EncodedVelocityNeighborhood[CONFIG_VELOCITY_PRE_FETCH_SAMPLES];
+		#endif
 
-		UNROLL_N(4)
-		for (uint i = 0; i < 4; i++)
+		// Issue texture fetechs
+		ISOLATE
 		{
-			float BilinearWeight = GetSampleWeight(BilinearInter, i);
-			uint2 PixelPos = GetSamplePixelCoord(BilinearInter, i);
-
-			bool bValidPixelPos = all(and(PixelPos >= InputInfo_ViewportMin, PixelPos < InputInfo_ViewportMax));
-
-			#if COMPRESS_PREV_USE_COUNT
-				uint2 CompressedUseCountPixelPos;
-				uint Shift;
-				ComputeCompressedUseCountPixelCoordinates(PixelPos, /* out */ CompressedUseCountPixelPos, /* out */ Shift);
+			DeviceZ = ClosestDepthTexture[InputPixelPos];
+			ParallaxFactor = UncompressParallaxFactor((ParallaxFactorTexture[InputPixelPos] >> 4) & 0xF);
+
+			#if CONFIG_VELOCITY_SAMPLES == 1
+				EncodedScreenVelocity = DilatedVelocityTexture[InputPixelPos];
+			#else // CONFIG_VELOCITY_SAMPLES != 1
+				tsr_short2 LaneOffsetSign = GetLaneOffsetSign();
+
+				UNROLL_N(CONFIG_VELOCITY_PRE_FETCH_SAMPLES)
+				for (uint i = 0; i < CONFIG_VELOCITY_PRE_FETCH_SAMPLES; i++)
+				{
+					tsr_short2 SampleInputPixelPos;
+					#if CONFIG_BUTTERFLY_KERNEL
+					{
+						tsr_short2 Offset = tsr_short2(Offsets2x2[i]) * LaneOffsetSign;
+				
+						SampleInputPixelPos = InputPixelPos + Offset;
+					}
+					#else
+					{
+						tsr_short2 Offset;
+						#if CONFIG_VELOCITY_SAMPLES == 9
+							Offset = tsr_short2(kOffsets3x3[kSquareIndexes3x3[i]]);
+						#elif CONFIG_VELOCITY_SAMPLES == 5
+							Offset = tsr_short2(kOffsets3x3[kPlusIndexes3x3[i]]);
+						#endif
+					
+						SampleInputPixelPos = InputPixelPos + ClampPixelOffset(
+							InputPixelPos,
+							Offset, Offset,
+							InputPixelPosMin, InputPixelPosMax);
+					}
+					#endif
+
+					EncodedVelocityNeighborhood[i] = DilatedVelocityTexture[SampleInputPixelPos];
+				}
+			#endif // CONFIG_VELOCITY_SAMPLES != 1
+		}
 
-				uint HistoryUseCount = (PrevUseCountTexture[CompressedUseCountPixelPos] >> (8 * Shift)) & 0xFF;
+		// LDS spill to save VGPR on the history texture fetches.
+		#if CONFIG_MANUAL_LDS_SPILL
+		ISOLATE
+		{
+			#if CONFIG_VELOCITY_SAMPLES == 1
+				SharedArray0[GroupThreadIndex] = EncodedScreenVelocity;
 			#else
-				uint HistoryUseCount = PrevUseCountTexture[PixelPos];
+				SharedArray0[GroupThreadIndex] = EncodedVelocityNeighborhood[0];
 			#endif
+			SharedArray1[GroupThreadIndex] = DeviceZ;
+		}
+		#endif
+		
+		// Detect whether is on velocity edge to clamp the high frequencies
+		#if CONFIG_VELOCITY_SAMPLES != 1
+		ISOLATE
+		{
+			// Preprocessed the prefetched samples
+			float2 PixelVelocityNeighborhood[CONFIG_VELOCITY_PRE_FETCH_SAMPLES];
+			{
+				UNROLL_N(CONFIG_VELOCITY_PRE_FETCH_SAMPLES)
+				for (uint i = 0; i < CONFIG_VELOCITY_PRE_FETCH_SAMPLES; i++)
+				{
+					PixelVelocityNeighborhood[i] = DecodeVelocityFromTexture(float4(EncodedVelocityNeighborhood[i], 0.0, 0.0)).xy * ScreenVelocityToInputPixelVelocity;
+				}
+			}
 
-			float HistoryClosestDeviceZ = f16tof32(PrevClosestDepthTexture[PixelPos] >> VELOCITY_HOLE_FILLING_BITS);
-			float HistoryClosestWorldDepth = ConvertFromDeviceZ(HistoryClosestDeviceZ);
+			// Sample
+			float2 ScreenPixelVelocity = PixelVelocityNeighborhood[0];
+			tsr_short2 LaneOffsetSign = GetLaneOffsetSign();
 
-			tsr_half DepthRejection;
+			UNROLL_N(CONFIG_VELOCITY_SAMPLE_PAIRS)
+			for (uint i = 0; i < CONFIG_VELOCITY_SAMPLE_PAIRS; i++)
+			ISOLATE
 			{
-				const float WorldDepthToPixelWorldRadius = GetTanHalfFieldOfView().x * InputInfo_ViewportSizeInverse.x;
-				const float PixelDepthError = 3.0;
-				const float RadiusToDiamater = 2.0;
+				tsr_short2 Offset0;
+				tsr_short2 Offset1;
 			
-				float WorldDepthEpsilon = (WorldDepthToPixelWorldRadius * RadiusToDiamater * PixelDepthError) * HistoryClosestWorldDepth; 
-				float DeltaDepth = abs(HistoryClosestWorldDepth - WorldDepth);
+				float2 NeighborPixelVelocity0;
+				float2 NeighborPixelVelocity1;
+
+				#if CONFIG_BUTTERFLY_KERNEL
+				{
+					const uint PairIndex = i;
+				
+					const uint ButterflyArray0[] = { 0x0, 0x0, 0x0, 0x1 };
+					const uint ButterflyArray1[] = { 0x1, 0x2, 0x3, 0x2 };
+				
+					uint FetchIndex[2];
+					FWaveBroadcastSettings BroadcastSettings[2];
+					BroadcastSettings[0] = InitWaveXorButterfly(ButterflyArray0[PairIndex]);
+					BroadcastSettings[1] = InitWaveXorButterfly(ButterflyArray1[PairIndex]);
+				
+					if (PairIndex == 3)
+					{
+						FetchIndex[0] = 2;
+						FetchIndex[1] = 1;
+						Offset0 = tsr_short2(-1, 1);
+					}
+					else
+					{
+						FetchIndex[0] = PairIndex + 1;
+						FetchIndex[1] = PairIndex + 1;
+						Offset0 = tsr_short2(Offsets2x2[PairIndex + 1]);
+					}
+
+					Offset0 *= LaneOffsetSign;
+					Offset1 = -Offset0;
+
+					NeighborPixelVelocity0 = WaveBroadcast(BroadcastSettings[0], PixelVelocityNeighborhood[FetchIndex[0]]);
+					NeighborPixelVelocity1 = WaveBroadcast(BroadcastSettings[1], PixelVelocityNeighborhood[FetchIndex[1]]);
+				}
+				#else
+				{
+					const uint NeightbordId0 = i + 1;
+					const uint NeightbordId1 = i + 1 + CONFIG_VELOCITY_SAMPLE_PAIRS;
+			
+					#if CONFIG_VELOCITY_SAMPLES == 9
+						Offset0 = tsr_short2(kOffsets3x3[kSquareIndexes3x3[NeightbordId0]]);
+						Offset1 = tsr_short2(kOffsets3x3[kSquareIndexes3x3[NeightbordId1]]);
+					#elif CONFIG_VELOCITY_SAMPLES == 5
+						Offset0 = tsr_short2(kOffsets3x3[kPlusIndexes3x3[NeightbordId0]]);
+						Offset1 = tsr_short2(kOffsets3x3[kPlusIndexes3x3[NeightbordId1]]);
+					#endif
+					
+					NeighborPixelVelocity0 = PixelVelocityNeighborhood[NeightbordId0];
+					NeighborPixelVelocity1 = PixelVelocityNeighborhood[NeightbordId1];
+				}
+				#endif
 
-				DepthRejection = saturate(tsr_half(2.0) - tsr_half(DeltaDepth / WorldDepthEpsilon));
+				tsr_half2 NeighborDeltaPixelVelocity0 = tsr_half2(ScreenPixelVelocity - NeighborPixelVelocity0);
+				tsr_half2 NeighborDeltaPixelVelocity1 = tsr_half2(ScreenPixelVelocity - NeighborPixelVelocity1);
+			
+				// Reject high frequency on velocity edge discontinuity.
+				{
+					tsr_half IsOverlapingEdge0 = tsr_half(1.0) - dot(saturate(NeighborDeltaPixelVelocity0 * Offset0), tsr_half(1.0).xx); // * saturate(ScreenPixelVelocityLength - NeighborPixelVelocityLength0);
+					tsr_half IsOverlapingEdge1 = tsr_half(1.0) - dot(saturate(NeighborDeltaPixelVelocity1 * Offset1), tsr_half(1.0).xx); // * saturate(ScreenPixelVelocityLength - NeighborPixelVelocityLength1);
+				
+					tsr_half IsOverlapingEdge = min(IsOverlapingEdge0, IsOverlapingEdge1);
+
+					if (i == 0)
+						RejectHighFrequency = IsOverlapingEdge;
+					else
+						RejectHighFrequency = min(RejectHighFrequency, IsOverlapingEdge);
+				}
 			}
 
-			tsr_half HistoryGhostingRejection = saturate(tsr_half(1.0) - tsr_half(3.0) * abs(tsr_half(HistoryUseCount) * rcp(tsr_half(PREV_USE_COUNT_QUANTIZATION)) - tsr_half(1.0)));
+			// Remove noise from high frequency rejections.
+			RejectHighFrequency = saturate(tsr_half(1.0) - RejectHighFrequency * tsr_half(1.1));
 
-			FLATTEN
-			if (bValidPixelPos)
-				ParallaxRejectionMask += tsr_half(BilinearWeight) * max(DepthRejection, HistoryGhostingRejection);
+			// Do not reject high frequency additionally where hole filling happen
+			// RejectHighFrequency = max(RejectHighFrequency, HoleFilledVelocityMask); TODO
+
+			EncodedScreenVelocity = EncodedVelocityNeighborhood[0];
 		}
+		#endif // CONFIG_VELOCITY_SAMPLES != 1
 	}
+	
+	// Sample the previous frame information
+	uint HistoryUseCountSamples[4];
+	uint HistoryClosestDeviceZSamples[4];
+
+	#if DIM_HIGH_FREQUENCY_ONLY && CONFIG_MEDIAN_DECIMATE
+		tsr_halfC DecimatedColorSamples[5];
+		tsr_half2 DecimatedMetadataSamples[5];
+	#elif DIM_HIGH_FREQUENCY_ONLY
+		tsr_halfC DecimatedColorSamples[1];
+		tsr_half2 DecimatedMetadataSamples[1];
+	#else
+		tsr_halfC DecimatedColorSamples[1];
+	#endif
+		
+	FSubpixelNeighborhood SubpixelNeigborhood;
+	
+	bool bIsOffScreen;
 
-	// Pixel coordinate of the input pixel in the predictor viewport.
-	float2 InputPPCk = float2(InputPixelPos - InputInfo_ViewportMin) + 0.5f - InputJitter;
-	float2 OutputViewportUV = InputPPCk * InputInfo_ViewportSizeInverse;
-	float2 OutputScreenPos = ViewportUVToScreenPos(OutputViewportUV);
+	// Issues all the history texture fetches
+	ISOLATE
+	{
+		tsr_ushort2 InputPixelPos = (
+			Map8x8Tile2x2Lane(GroupThreadIndex) +
+			(tsr_ushort2(InputInfo_ViewportMin) + tsr_ushort2(GroupId) * tsr_ushort2(TILE_SIZE, TILE_SIZE)));
 
-	float2 PrevOutputScreenPos = OutputScreenPos - DecodeVelocityFromTexture(float4(EncodedVelocity, 0.0, 0.0)).xy;
+		{
+			float2 ScreenPos = ApplyScreenTransform(float2(InputPixelPos), InputPixelPosToScreenPos);
 
-	float2 PrevOutputViewportUV = ScreenPosToViewportUV(PrevOutputScreenPos);
+			float2 PrevScreenPos = ScreenPos - DecodeVelocityFromTexture(float4(EncodedScreenVelocity, 0.0, 0.0)).xy;
+			float2 PrevInputBufferUV = (InputInfo_ScreenPosToViewportScale * PrevScreenPos + InputInfo_ScreenPosToViewportBias) * InputInfo_ExtentInverse;
+		
+			FBilinearSampleInfos BilinearInter = GetBilinearSampleLevelInfos(PrevInputBufferUV, InputInfo_Extent, InputInfo_ExtentInverse);
+
+			UNROLL_N(4)
+			for (uint i = 0; i < 4; i++)
+			{
+				float BilinearWeight = GetSampleWeight(BilinearInter, i);
+				uint2 PixelPos = GetSamplePixelCoord(BilinearInter, i);
 
-	float2 HistoryPPCk = PrevOutputViewportUV * PrevHistoryInfo_ViewportSize;
+				bool bValidPixelPos = all(and(PixelPos >= InputInfo_ViewportMin, PixelPos < InputInfo_ViewportMax));
 
-	// Pixel coordinate of the center of the nearest predictor pixel O.
-	float2 HistoryPPCo = floor(HistoryPPCk) + 0.5;
-	float2 HistoryBufferUV = HistoryPPCo * PrevHistoryInfo_ExtentInverse;
+				#if COMPRESS_PREV_USE_COUNT
+					uint2 CompressedUseCountPixelPos;
+					uint Shift;
+					ComputeCompressedUseCountPixelCoordinates(PixelPos, /* out */ CompressedUseCountPixelPos, /* out */ Shift);
 
-	// Vector in pixel between pixel K -> O.
-	tsr_half2 dKO = tsr_half2(HistoryPPCo - HistoryPPCk);
+					uint HistoryUseCount = PrevUseCountTexture[CompressedUseCountPixelPos];
+				#else
+					uint HistoryUseCount = PrevUseCountTexture[PixelPos];
+				#endif
+			
+				HistoryUseCountSamples[i] = HistoryUseCount;
+				HistoryClosestDeviceZSamples[i] = PrevClosestDepthTexture[PixelPos];
+			}
+		}
+			
+		{
+			// Pixel coordinate of the input pixel in the predictor viewport.
+			float2 InputPPCk = float2(InputPixelPos - InputInfo_ViewportMin) + 0.5f - InputJitter;
+			float2 OutputViewportUV = InputPPCk * InputInfo_ViewportSizeInverse;
+			float2 OutputScreenPos = ViewportUVToScreenPos(OutputViewportUV);
 	
-	// Sample the history
-	tsr_halfC DecimatedColor;
-	tsr_half HistoryMaxRejectionFactor;
+			float2 PrevOutputScreenPos = OutputScreenPos - DecodeVelocityFromTexture(float4(EncodedScreenVelocity, 0.0, 0.0)).xy;
+			bIsOffScreen = (
+				(PrevOutputScreenPos.x < -1.0) ||
+				(PrevOutputScreenPos.y < -1.0) ||
+				(PrevOutputScreenPos.x > 1.0) ||
+				(PrevOutputScreenPos.y > 1.0));
 
-	tsr_half ClosestInterferenceSeed;
-	tsr_half IntegratedInterferenceSeed;
-	tsr_half MinInterferenceSeed;
-	tsr_half MaxInterferenceSeed;
-	{
-		float2 KernelHistoryBufferUV = clamp(
-			HistoryPPCk * PrevHistoryInfo_ExtentInverse,
-			PrevHistoryInfo_UVViewportBilinearMin,
-			PrevHistoryInfo_UVViewportBilinearMax);
-		
-		float2 PrevInputBufferUV = (InputInfo_ScreenPosToViewportScale * PrevOutputScreenPos + InputInfo_ScreenPosToViewportBias) * InputInfo_ExtentInverse;
+			float2 PrevOutputViewportUV = ScreenPosToViewportUV(PrevOutputScreenPos);
+
+			float2 HistoryPPCk = PrevOutputViewportUV * PrevHistoryInfo_ViewportSize;
+
+			float2 KernelHistoryBufferUV = clamp(
+				HistoryPPCk * PrevHistoryInfo_ExtentInverse,
+				PrevHistoryInfo_UVViewportBilinearMin,
+				PrevHistoryInfo_UVViewportBilinearMax);
 		
-		#if DIM_HIGH_FREQUENCY_ONLY && CONFIG_MEDIAN_DECIMATE
-		{
-			const float kTopSecretSauce = 2.0;
+			#if DIM_HIGH_FREQUENCY_ONLY && CONFIG_MEDIAN_DECIMATE
+			{
+				const float kTopSecretSauce = 2.0;
 
-			float2 PrevUV[5];
-			PrevUV[0] = KernelHistoryBufferUV;
-			PrevUV[1] = KernelHistoryBufferUV + (kTopSecretSauce * float2(+1.0, +1.0)) * PrevHistoryInfo_ExtentInverse;
-			PrevUV[2] = KernelHistoryBufferUV + (kTopSecretSauce * float2(-1.0, +1.0)) * PrevHistoryInfo_ExtentInverse;
-			PrevUV[3] = KernelHistoryBufferUV + (kTopSecretSauce * float2(+1.0, -1.0)) * PrevHistoryInfo_ExtentInverse;
-			PrevUV[4] = KernelHistoryBufferUV + (kTopSecretSauce * float2(-1.0, -1.0)) * PrevHistoryInfo_ExtentInverse;
+				float2 PrevUV[5];
+				PrevUV[0] = KernelHistoryBufferUV;
+				PrevUV[1] = KernelHistoryBufferUV + (kTopSecretSauce * float2(+1.0, +1.0)) * PrevHistoryInfo_ExtentInverse;
+				PrevUV[2] = KernelHistoryBufferUV + (kTopSecretSauce * float2(-1.0, +1.0)) * PrevHistoryInfo_ExtentInverse;
+				PrevUV[3] = KernelHistoryBufferUV + (kTopSecretSauce * float2(+1.0, -1.0)) * PrevHistoryInfo_ExtentInverse;
+				PrevUV[4] = KernelHistoryBufferUV + (kTopSecretSauce * float2(-1.0, -1.0)) * PrevHistoryInfo_ExtentInverse;
 
-			tsr_halfC PrevDecimatedColorArray[5];
+				tsr_halfC PrevDecimatedColorArray[5];
 			
+				UNROLL_N(5)
+				for (uint i = 0; i < 5; i++)
+				{
+					float2 UV = KernelHistoryBufferUV = clamp(
+						PrevUV[i],
+						PrevHistoryInfo_UVViewportBilinearMin,
+						PrevHistoryInfo_UVViewportBilinearMax);
+					#if 1 // saves 0.120ms
+						DecimatedMetadataSamples[i] = tsr_half(1.0).xx;
+						DecimatedColorSamples[i] = PrevHistory_HighFrequencyHalfRes[uint2(UV * PrevHistoryInfo_Extent * 0.5)];
+					#elif 1
+						DecimatedMetadataSamples[i] = PrevHistory_Metadata[uint2(UV * PrevHistoryInfo_Extent)];
+						DecimatedColorSamples[i] = PrevHistory_HighFrequency[uint2(UV * PrevHistoryInfo_Extent)];
+					#else
+						DecimatedMetadataSamples[i] = PrevHistory_Metadata.SampleLevel(GlobalBilinearClampedSampler, UV, 0);
+						DecimatedColorSamples[i] = PrevHistory_HighFrequency.SampleLevel(GlobalBilinearClampedSampler, UV, 0);
+					#endif
+				}
+			}
+			#elif DIM_HIGH_FREQUENCY_ONLY
+			{
+				DecimatedMetadataSamples[0] = PrevHistory_Metadata.SampleLevel(GlobalBilinearClampedSampler, KernelHistoryBufferUV, 0);
+				DecimatedColorSamples[0] = PrevHistory_HighFrequency.SampleLevel(GlobalBilinearClampedSampler, KernelHistoryBufferUV, 0);
+			}
+			#else
+			{
+				DecimatedColorSamples[0] = PrevHistory_LowFrequency.SampleLevel(GlobalBilinearClampedSampler, KernelHistoryBufferUV, 0);
+			}
+			#endif
+				
+			SubpixelNeigborhood = GatherPrevSubpixelNeighborhood(PrevHistory_SubpixelDetails, KernelHistoryBufferUV);
+		}
+	}
+	
+	// Decimate the history color for comparison with input immediatly to quickly reduce VGPRs.
+	tsr_halfC DecimatedColor;
+	ISOLATE
+	{
+		#if DIM_HIGH_FREQUENCY_ONLY && CONFIG_MEDIAN_DECIMATE
+		{
 			UNROLL_N(5)
 			for (uint i = 0; i < 5; i++)
 			{
-				float2 UV = KernelHistoryBufferUV = clamp(
-					PrevUV[i],
-					PrevHistoryInfo_UVViewportBilinearMin,
-					PrevHistoryInfo_UVViewportBilinearMax);
-					
-				tsr_half2 PrevMetadata = PrevHistory_Metadata.SampleLevel(GlobalBilinearClampedSampler, UV, 0);
-				PrevDecimatedColorArray[i] = PrevHistory_HighFrequency.SampleLevel(GlobalBilinearClampedSampler, UV, 0) * SafeRcp(PrevMetadata.r);
+				DecimatedColorSamples[i] *= SafeRcp(DecimatedMetadataSamples[i].x);
 			}
 
-			DecimatedColor = Median5(PrevDecimatedColorArray);
+			DecimatedColor = Median5(DecimatedColorSamples);
 		}
 		#elif DIM_HIGH_FREQUENCY_ONLY
 		{
-			tsr_half2 Metadata = PrevHistory_Metadata.SampleLevel(GlobalBilinearClampedSampler, KernelHistoryBufferUV, 0);
-			DecimatedColor = PrevHistory_HighFrequency.SampleLevel(GlobalBilinearClampedSampler, KernelHistoryBufferUV, 0) * SafeRcp(Metadata.r);
+			DecimatedColor = DecimatedColorSamples[0] * SafeRcp(DecimatedMetadataSamples[0].x);
 		}
 		#else
 		{
-			DecimatedColor = PrevHistory_LowFrequency.SampleLevel(GlobalBilinearClampedSampler, KernelHistoryBufferUV, 0);
+			DecimatedColor = DecimatedColorSamples[0];
+		}
+		#endif
+	}
+		
+
+	// Process the history samples
+	tsr_half ParallaxRejectionMask = 0.0;
+
+	ISOLATE
+	{
+		// LDS despill
+		#if CONFIG_MANUAL_LDS_SPILL
+		{
+			EncodedScreenVelocity = SharedArray0[GroupThreadIndex];
+			DeviceZ = SharedArray1[GroupThreadIndex];
 		}
 		#endif
+		
+		tsr_ushort2 InputPixelPos = (
+			Map8x8Tile2x2Lane(GroupThreadIndex) +
+			(tsr_ushort2(InputInfo_ViewportMin) + tsr_ushort2(GroupId) * tsr_ushort2(TILE_SIZE, TILE_SIZE)));
 
-		// Compute the parallax rejection factor of the details in the history.
-		FSubpixelNeighborhood SubpixelNeigborhood = GatherPrevSubpixelNeighborhood(PrevHistory_SubpixelDetails, KernelHistoryBufferUV);
+		float2 ScreenPos = ApplyScreenTransform(float2(InputPixelPos), InputPixelPosToScreenPos);
+		float2 ScreenVelocity = DecodeVelocityFromTexture(float4(EncodedScreenVelocity, 0.0, 0.0)).xy;
+		float2 PrevScreenPos = ScreenPos - ScreenVelocity;
+	
+		// Compute HistoryMaxRejectionFactor
+		tsr_half HistoryMaxRejectionFactor;
 		{
 			uint EncodedHistoryMaxRejectionFactor = 0;
 
@@ -261,45 +493,137 @@ void MainCS(
 
 			HistoryMaxRejectionFactor = UncompressParallaxFactor(EncodedHistoryMaxRejectionFactor);
 		}
-	}
+		
+		// Compute the parralax rejection mask
+		bool bIsValidHoleFillingPixelVelocity = false;
+		uint EncodedHoleFillingVelocity = 0;
+		{
+			float WorldDepth = ConvertFromDeviceZ(DeviceZ);
 	
-	// Reject subpixel details that might be in history but not input just based on their parrallax
-	{
-		tsr_half HistorySubdetailRejectionFactor = max(HistoryMaxRejectionFactor / ParallaxFactor, tsr_half(1.0));
+			float2 PrevInputBufferUV = (InputInfo_ScreenPosToViewportScale * PrevScreenPos + InputInfo_ScreenPosToViewportBias) * InputInfo_ExtentInverse;
+			FBilinearSampleInfos BilinearInter = GetBilinearSampleLevelInfos(PrevInputBufferUV, InputInfo_Extent, InputInfo_ExtentInverse);
 
-		float4 ThisClip = float4(ScreenPos, DeviceZ, 1);
-		float4 PrevClip = mul(ThisClip, View.ClipToPrevClip);
-		float2 PrevScreen = PrevClip.xy / PrevClip.w;
+			UNROLL_N(4)
+			for (uint i = 0; i < 4; i++)
+			{
+				float BilinearWeight = GetSampleWeight(BilinearInter, i);
+				uint2 PixelPos = GetSamplePixelCoord(BilinearInter, i);
+
+				bool bValidPixelPos = all(and(PixelPos >= InputInfo_ViewportMin, PixelPos < InputInfo_ViewportMax));
+				
+				#if COMPRESS_PREV_USE_COUNT
+					uint2 CompressedUseCountPixelPos;
+					uint Shift;
+					ComputeCompressedUseCountPixelCoordinates(PixelPos, /* out */ CompressedUseCountPixelPos, /* out */ Shift);
+
+					uint HistoryUseCount = (HistoryUseCountSamples[i] >> (8 * Shift)) & 0xFF;
+				#else
+					uint HistoryUseCount = HistoryUseCountSamples[i];
+				#endif
 			
-		float4 RotationalPrevClip = mul(ThisClip, RotationalClipToPrevClip);
-		float2 RotationalPrevScreen = RotationalPrevClip.xy / RotationalPrevClip.w;
+				uint SampleHistoryClosestDeviceZ = HistoryClosestDeviceZSamples[i];
+		
+				float HistoryClosestDeviceZ = f16tof32(SampleHistoryClosestDeviceZ >> VELOCITY_HOLE_FILLING_BITS);
+				float HistoryClosestWorldDepth = ConvertFromDeviceZ(HistoryClosestDeviceZ);
+
+				tsr_half DepthRejection;
+				{
+					const float WorldDepthToPixelWorldRadius = GetTanHalfFieldOfView().x * InputInfo_ViewportSizeInverse.x;
+					const float PixelDepthError = 3.0;
+					const float RadiusToDiamater = 2.0;
 			
-		tsr_half PixelParallaxVelocity = tsr_half(0.5) * length(tsr_half2(RotationalPrevScreen - PrevScreen) * tsr_half2(InputInfo_ViewportSize));
-		tsr_half HistorySubdetailAdditionalParralax = (HistorySubdetailRejectionFactor - tsr_half(1.0)) * PixelParallaxVelocity;
+					float WorldDepthEpsilon = (WorldDepthToPixelWorldRadius * RadiusToDiamater * PixelDepthError) * HistoryClosestWorldDepth; 
+					float DeltaDepth = abs(HistoryClosestWorldDepth - WorldDepth);
 
-		// PixelParallaxVelocity can be +inf and generate NaN if HistorySubdetailRejectionFactor = 1.0;
-		HistorySubdetailAdditionalParralax = -min(-HistorySubdetailAdditionalParralax, tsr_half(0.0));
+					DepthRejection = saturate(tsr_half(2.0) - tsr_half(DeltaDepth / WorldDepthEpsilon));
+				}
 
-		tsr_half HistoryParallaxRejectionMask = saturate(tsr_half(1.0) - HistorySubdetailAdditionalParralax);
+				tsr_half HistoryGhostingRejection = saturate(tsr_half(1.0) - tsr_half(3.0) * abs(tsr_half(HistoryUseCount) * rcp(tsr_half(PREV_USE_COUNT_QUANTIZATION)) - tsr_half(1.0)));
+			
+				bIsValidHoleFillingPixelVelocity = bIsValidHoleFillingPixelVelocity || bValidPixelPos;
 
-		ParallaxRejectionMask = min(ParallaxRejectionMask, HistoryParallaxRejectionMask);
-	}
+				EncodedHoleFillingVelocity = select(bValidPixelPos, max(EncodedHoleFillingVelocity, SampleHistoryClosestDeviceZ), EncodedHoleFillingVelocity);
+				ParallaxRejectionMask = select(bValidPixelPos, ParallaxRejectionMask + tsr_half(BilinearWeight) * max(DepthRejection, HistoryGhostingRejection), ParallaxRejectionMask);
+			}
+		}
 
-	// Immediately reject based on the rejection mask
-	FLATTEN
-	if ((PrevOutputScreenPos.x < -1.0) ||
-		(PrevOutputScreenPos.y < -1.0) ||
-		(PrevOutputScreenPos.x > 1.0) ||
-		(PrevOutputScreenPos.y > 1.0))
-	//if (IsOffScreen2(bCameraCut, PrevOutputScreenPos, /* ParallaxRejectionMask = */ 1.0))
-	{
-		ParallaxRejectionMask = tsr_half(0.0);
-	}
+		// Reject subpixel details that might be in history but not input just based on their parrallax
+		{
+			tsr_half HistorySubdetailRejectionFactor = max(HistoryMaxRejectionFactor / ParallaxFactor, tsr_half(1.0));
 
+			float4 ThisClip = float4(ScreenPos, DeviceZ, 1);
+			float4 PrevClip = mul(ThisClip, View.ClipToPrevClip);
+			float2 PrevScreen = PrevClip.xy / PrevClip.w;
+			
+			float4 RotationalPrevClip = mul(ThisClip, RotationalClipToPrevClip);
+			float2 RotationalPrevScreen = RotationalPrevClip.xy / RotationalPrevClip.w;
+			
+			tsr_half PixelParallaxVelocity = tsr_half(0.5) * length(tsr_half2(RotationalPrevScreen - PrevScreen) * tsr_half2(InputInfo_ViewportSize));
+			tsr_half HistorySubdetailAdditionalParralax = (HistorySubdetailRejectionFactor - tsr_half(1.0)) * PixelParallaxVelocity;
+
+			// PixelParallaxVelocity can be +inf and generate NaN if HistorySubdetailRejectionFactor = 1.0;
+			HistorySubdetailAdditionalParralax = -min(-HistorySubdetailAdditionalParralax, tsr_half(0.0));
+
+			tsr_half HistoryParallaxRejectionMask = saturate(tsr_half(1.0) - HistorySubdetailAdditionalParralax);
+
+			ParallaxRejectionMask = min(ParallaxRejectionMask, HistoryParallaxRejectionMask);
+		}
+
+		// Reject subpixel details that might be in history but not input just based on their parrallax
+		{
+			tsr_half HistorySubdetailRejectionFactor = max(HistoryMaxRejectionFactor / ParallaxFactor, tsr_half(1.0));
+
+			float4 ThisClip = float4(ScreenPos, DeviceZ, 1);
+			float4 PrevClip = mul(ThisClip, View.ClipToPrevClip);
+			float2 PrevScreen = PrevClip.xy / PrevClip.w;
+			
+			float4 RotationalPrevClip = mul(ThisClip, RotationalClipToPrevClip);
+			float2 RotationalPrevScreen = RotationalPrevClip.xy / RotationalPrevClip.w;
+			
+			tsr_half PixelParallaxVelocity = tsr_half(0.5) * length(tsr_half2(RotationalPrevScreen - PrevScreen) * tsr_half2(InputInfo_ViewportSize));
+			tsr_half HistorySubdetailAdditionalParralax = (HistorySubdetailRejectionFactor - tsr_half(1.0)) * PixelParallaxVelocity;
+
+			// PixelParallaxVelocity can be +inf and generate NaN if HistorySubdetailRejectionFactor = 1.0;
+			HistorySubdetailAdditionalParralax = -min(-HistorySubdetailAdditionalParralax, tsr_half(0.0));
+
+			tsr_half HistoryParallaxRejectionMask = saturate(tsr_half(1.0) - HistorySubdetailAdditionalParralax);
+
+			ParallaxRejectionMask = min(ParallaxRejectionMask, HistoryParallaxRejectionMask);
+		}
+
+		// Finish off the ParallaxRejectionMask
+		ParallaxRejectionMask = select(bIsOffScreen, tsr_half(0.0), ParallaxRejectionMask);
+
+		bool bIsOffscreenOrParallaxRejected = ParallaxRejectionMask < tsr_half(PARALLAX_REJECTION_MASK_THRESHOLD);
+	
+		// Hole fill the velocity.
+		{
+			float2 HoleFillingPixelVelocity;
+			float HoleFillingOutOfRange;
+			DecodeHoleFillingVelocity(EncodedHoleFillingVelocity, /* out */ HoleFillingPixelVelocity, /* out */ HoleFillingOutOfRange);
+			if (!bIsValidHoleFillingPixelVelocity)
+			{
+				HoleFillingOutOfRange = 0.0;
+			}
+
+			// If the hole filling velocity has better result that the neighborhood, use that.
+			FLATTEN
+			if (tsr_half(HoleFillingOutOfRange) > 0.0 && bIsOffscreenOrParallaxRejected)
+			{
+				float2 ScreenPixelVelocity = ScreenVelocity * ScreenVelocityToInputPixelVelocity;
+				float VelocityDelta = length(ScreenPixelVelocity - HoleFillingPixelVelocity);
+
+				float2 HoleFillingVelocity = HoleFillingPixelVelocity * InputPixelVelocityToScreenVelocity;
+
+				EncodedScreenVelocity = EncodeVelocityToTexture(float3(HoleFillingVelocity, 0.0)).xy;
+				//HoleFilledVelocityMask = tsr_half(HoleFillingOutOfRange) * tsr_half(saturate(MaxHollFillPixelVelocity - VelocityDelta));
+			}
+		}
+	}
 	
+	// Output half res.
 	#if CONFIG_OUTPUT_HALF_RES
 	{
-		SharedInputColor[GroupThreadIndex] = InputSceneColor;
 		#if CONFIG_SCENE_COLOR_APLHA
 			SharedHistoryColor[GroupThreadIndex] = DecimatedColor;
 			SharedMask[GroupThreadIndex] = ParallaxRejectionMask;
@@ -311,7 +635,6 @@ void MainCS(
 
 		bool bCheckerboard = (((InputPixelPos.x ^ InputPixelPos.y) >> 1) & 0x1) != 0;
 
-		tsr_halfC HalfResInputSceneColor = InputSceneColor;
 		tsr_halfC HalfResDecimatedColor = DecimatedColor;
 		tsr_half HalfResParallaxRejectionMask = ParallaxRejectionMask;
 
@@ -320,15 +643,6 @@ void MainCS(
 		{
 			uint ButterflySwap = i;
 			
-			if (bCheckerboard)
-			{
-				HalfResInputSceneColor = min(HalfResInputSceneColor, SharedInputColor[GroupThreadIndex ^ ButterflySwap]);
-			}
-			else
-			{
-				HalfResInputSceneColor = max(HalfResInputSceneColor, SharedInputColor[GroupThreadIndex ^ ButterflySwap]);
-			}
-			
 			#if CONFIG_SCENE_COLOR_APLHA
 				HalfResDecimatedColor += SharedHistoryColor[GroupThreadIndex ^ ButterflySwap];
 				HalfResParallaxRejectionMask += SharedMask[GroupThreadIndex ^ ButterflySwap];
@@ -346,42 +660,49 @@ void MainCS(
 			uint2 Random = Rand3DPCG16(int3((InputPixelPos - InputInfo_ViewportMin) / 2, View.StateFrameIndexMod8)).xy;
 			tsr_half E = tsr_half(Hammersley16(0, 1, Random).x);
 
-			HalfResInputSceneColor = QuantizeForFloatRenderTarget(HalfResInputSceneColor, E, OutputQuantizationError);
 			HalfResDecimatedColor = QuantizeForFloatRenderTarget(HalfResDecimatedColor, E, OutputQuantizationError);
 		}
 		#endif
 
 		if (all(InputPixelPos < InputInfo_ViewportMax) && all(((InputPixelPos - InputInfo_ViewportMin) % 2) == 0))
 		{
-			HalfResSceneColorOutput[InputPixelPos / 2] = HalfResInputSceneColor;
 			HalfResPredictionSceneColorOutput[InputPixelPos / 2] = HalfResDecimatedColor;
 			HalfResParallaxRejectionMaskOutput[InputPixelPos / 2] = HalfResParallaxRejectionMask;
 		}
 	}
 	#endif
 
-	#if CONFIG_ENABLE_STOCASTIC_QUANTIZATION
+	// Output full res
+	ISOLATE
 	{
-		uint2 Random = Rand3DPCG16(int3(InputPixelPos - InputInfo_ViewportMin, View.StateFrameIndexMod8)).xy;
-		tsr_half E = tsr_half(Hammersley16(0, 1, Random).x);
+		tsr_ushort2 InputPixelPos = (
+			Map8x8Tile2x2Lane(GroupThreadIndex) +
+			(tsr_ushort2(InputInfo_ViewportMin) + tsr_ushort2(GroupId) * tsr_ushort2(TILE_SIZE, TILE_SIZE)));
+			
+		#if CONFIG_ENABLE_STOCASTIC_QUANTIZATION
+		{
+			uint2 Random = Rand3DPCG16(int3(InputPixelPos - InputInfo_ViewportMin, View.StateFrameIndexMod8)).xy;
+			tsr_half E = tsr_half(Hammersley16(0, 1, Random).x);
 
-		DecimatedColor = QuantizeForFloatRenderTarget(DecimatedColor, E, OutputQuantizationError);
-	}
-	#endif
+			DecimatedColor = QuantizeForFloatRenderTarget(DecimatedColor, E, OutputQuantizationError);
+		}
+		#endif
+		
+		tsr_short2 OutputPixelPos = InvalidateOutputPixelPos(InputPixelPos, InputInfo_ViewportMax);
 
-	if (all(InputPixelPos < InputInfo_ViewportMax))
-	{
 		#if !CONFIG_OUTPUT_HALF_RES
 		{
-			PredictionSceneColorOutput[InputPixelPos] = DecimatedColor;
+			PredictionSceneColorOutput[OutputPixelPos] = DecimatedColor;
 		}
 		#endif
 		
-		ParallaxRejectionMaskOutput[InputPixelPos] = ParallaxRejectionMask;
+		HoleFilledVelocityOutput[OutputPixelPos] = EncodedScreenVelocity;
+		HoleFilledVelocityMaskOutput[InputPixelPos] = tsr_half2(0.0, RejectHighFrequency);
+		ParallaxRejectionMaskOutput[OutputPixelPos] = ParallaxRejectionMask;
 
 		#if DEBUG_OUTPUT
 		{
-			DebugOutput[InputPixelPos] = Debug;
+			DebugOutput[OutputPixelPos] = Debug;
 		}
 		#endif
 	}
diff --git a/Engine/Shaders/Private/TemporalSuperResolution/TSRHoleFillVelocity.usf b/Engine/Shaders/Private/TemporalSuperResolution/TSRHoleFillVelocity.usf
deleted file mode 100644
diff --git a/Engine/Shaders/Private/TemporalSuperResolution/TSRUpdateHistory.usf b/Engine/Shaders/Private/TemporalSuperResolution/TSRUpdateHistory.usf
--- a/Engine/Shaders/Private/TemporalSuperResolution/TSRUpdateHistory.usf
+++ b/Engine/Shaders/Private/TemporalSuperResolution/TSRUpdateHistory.usf
@@ -40,10 +40,12 @@
 #if DIM_HIGH_FREQUENCY_ONLY == 1
 	#define CONFIG_HIGH_FREQUENCY_ONLY 1
 	#define CONFIG_HIGH_FREQUENCY_YCOCG 0
+	#define CONFIG_OUTPUT_HALFRES_HISTORY 1
 
 #else
 	#define CONFIG_HIGH_FREQUENCY_ONLY 0
 	#define CONFIG_HIGH_FREQUENCY_YCOCG 1
+	#define CONFIG_OUTPUT_HALFRES_HISTORY 0
 
 #endif
 
@@ -116,6 +118,7 @@ uint bHasSeparateTranslucency;
 
 Texture2D<tsr_halfC> PrevHistory_LowFrequency;
 Texture2D<tsr_halfC> PrevHistory_HighFrequency;
+Texture2D<tsr_halfC> PrevHistory_HighFrequencyHalfRes;
 Texture2D<tsr_half2> PrevHistory_Metadata;
 Texture2D<tsr_half3> PrevHistory_Translucency;
 Texture2D<tsr_half>  PrevHistory_TranslucencyAlpha;
@@ -123,6 +126,7 @@ Texture2D<tsr_subpixel_details> PrevHistory_SubpixelDetails;
 
 RWTexture2D<tsr_halfC> HistoryOutput_LowFrequency;
 RWTexture2D<tsr_halfC> HistoryOutput_HighFrequency;
+RWTexture2D<tsr_halfC> HistoryOutput_HighFrequencyHalfRes;
 RWTexture2D<tsr_half2> HistoryOutput_Metadata;
 RWTexture2D<tsr_half3> HistoryOutput_Translucency;
 RWTexture2D<tsr_half>  HistoryOutput_TranslucencyAlpha;
@@ -1538,6 +1542,10 @@ void MainCS(
 			
 		LocalHistoryPixelPos = InvalidateOutputPixelPos(LocalHistoryPixelPos, HistoryInfo_ViewportMax);
 		
+		tsr_short2x2 HalfResOutputPixelPos;
+		HalfResOutputPixelPos[0] = (LocalHistoryPixelPos[0] >> tsr_short(1)) | (((LocalHistoryPixelPos[0] & tsr_short(0x1))) * tsr_short(~0));
+		HalfResOutputPixelPos[1] = (LocalHistoryPixelPos[1] >> tsr_short(1)) | (((LocalHistoryPixelPos[1] & tsr_short(0x1))) * tsr_short(~0));
+
 		// Premultiply FinalHighFrequencyColor to have validity weighted average in next frame history reprojection.
 		#if CONFIG_HIGH_FREQUENCY_ONLY
 			tsr_half2 FinalHistoryWeight = max(FinalHistoryValidity, MIN_IMPLICIT_HISTORY_WEIGHT);
@@ -1576,41 +1584,110 @@ void MainCS(
 			#endif
 		}
 
-		// Output final history lo pixel.
+		// Output full res history
 		{
-			#if CONFIG_HIGH_FREQUENCY_ONLY
-				HistoryOutput_HighFrequency[dpv_lo(LocalHistoryPixelPos)] = dpv_lo(PremultipliedHighFrequencyColor);
-				HistoryOutput_Metadata[dpv_lo(LocalHistoryPixelPos)]      = tsr_half2(dpv_lo(FinalHistoryWeight), dpv_lo(FinalHistoryValidity));
-			#else
-				HistoryOutput_LowFrequency[dpv_lo(LocalHistoryPixelPos)]  = dpv_lo(FinalLowFrequencyColor);
-				HistoryOutput_HighFrequency[dpv_lo(LocalHistoryPixelPos)] = dpv_lo(PremultipliedHighFrequencyColor);
-				HistoryOutput_Metadata[dpv_lo(LocalHistoryPixelPos)]      = tsr_half2(dpv_lo(FinalLowFrequencyWeight), dpv_lo(FinalHistoryValidity));
-			#endif
-			#if CONFIG_ACCUMULATE_TRANSLUCENCY_SEPARATLY
-				HistoryOutput_Translucency[dpv_lo(LocalHistoryPixelPos)] = dpv_lo(FinalTranslucencyColor).rgb;
-				HistoryOutput_TranslucencyAlpha[dpv_lo(LocalHistoryPixelPos)] = dpv_lo(FinalTranslucencyColor).a;
-			#endif
-			HistoryOutput_SubpixelDetails[dpv_lo(LocalHistoryPixelPos)] = dpv_lo(FinalSubpixelDetails);
+			// Output final history lo pixel.
+			{
+				#if CONFIG_HIGH_FREQUENCY_ONLY
+					HistoryOutput_HighFrequency[dpv_lo(LocalHistoryPixelPos)] = dpv_lo(PremultipliedHighFrequencyColor);
+					HistoryOutput_Metadata[dpv_lo(LocalHistoryPixelPos)]      = tsr_half2(dpv_lo(FinalHistoryWeight), dpv_lo(FinalHistoryValidity));
+				#else
+					HistoryOutput_LowFrequency[dpv_lo(LocalHistoryPixelPos)]  = dpv_lo(FinalLowFrequencyColor);
+					HistoryOutput_HighFrequency[dpv_lo(LocalHistoryPixelPos)] = dpv_lo(PremultipliedHighFrequencyColor);
+					HistoryOutput_Metadata[dpv_lo(LocalHistoryPixelPos)]      = tsr_half2(dpv_lo(FinalLowFrequencyWeight), dpv_lo(FinalHistoryValidity));
+				#endif
+				#if CONFIG_ACCUMULATE_TRANSLUCENCY_SEPARATLY
+					HistoryOutput_Translucency[dpv_lo(LocalHistoryPixelPos)] = dpv_lo(FinalTranslucencyColor).rgb;
+					HistoryOutput_TranslucencyAlpha[dpv_lo(LocalHistoryPixelPos)] = dpv_lo(FinalTranslucencyColor).a;
+				#endif
+				HistoryOutput_SubpixelDetails[dpv_lo(LocalHistoryPixelPos)] = dpv_lo(FinalSubpixelDetails);
+			}
+		
+			// Output final history hi pixel.
+			#if CONFIG_ENABLE_DUAL_PIXEL_VECTORIZATION
+			{
+				#if CONFIG_HIGH_FREQUENCY_ONLY
+					HistoryOutput_HighFrequency[dpv_hi(LocalHistoryPixelPos)] = dpv_hi(PremultipliedHighFrequencyColor);
+					HistoryOutput_Metadata[dpv_hi(LocalHistoryPixelPos)]      = tsr_half2(dpv_hi(FinalHistoryWeight), dpv_hi(FinalHistoryValidity));
+				#else
+					HistoryOutput_LowFrequency[dpv_hi(LocalHistoryPixelPos)]  = dpv_hi(FinalLowFrequencyColor);
+					HistoryOutput_HighFrequency[dpv_hi(LocalHistoryPixelPos)] = dpv_hi(PremultipliedHighFrequencyColor);
+					HistoryOutput_Metadata[dpv_hi(LocalHistoryPixelPos)]      = tsr_half2(dpv_hi(FinalLowFrequencyWeight), dpv_hi(FinalHistoryValidity));
+				#endif
+				#if CONFIG_ACCUMULATE_TRANSLUCENCY_SEPARATLY
+					HistoryOutput_Translucency[dpv_hi(LocalHistoryPixelPos)] = dpv_hi(FinalTranslucencyColor).rgb;
+					HistoryOutput_TranslucencyAlpha[dpv_hi(LocalHistoryPixelPos)] = dpv_hi(FinalTranslucencyColor).a;
+				#endif
+				HistoryOutput_SubpixelDetails[dpv_hi(LocalHistoryPixelPos)] = dpv_hi(FinalSubpixelDetails);
+			}
+			#endif // !CONFIG_ENABLE_DUAL_PIXEL_VECTORIZATION
 		}
 		
-		// Output final history hi pixel.
-		#if CONFIG_ENABLE_DUAL_PIXEL_VECTORIZATION
+		// Output half res history
+		#if CONFIG_OUTPUT_HALFRES_HISTORY
 		{
-			#if CONFIG_HIGH_FREQUENCY_ONLY
-				HistoryOutput_HighFrequency[dpv_hi(LocalHistoryPixelPos)] = dpv_hi(PremultipliedHighFrequencyColor);
-				HistoryOutput_Metadata[dpv_hi(LocalHistoryPixelPos)]      = tsr_half2(dpv_hi(FinalHistoryWeight), dpv_hi(FinalHistoryValidity));
+			#if CONFIG_SCENE_COLOR_APLHA
+				tsr_half4x2 HalfResFinalHighFrequencyColor = FinalHighFrequencyColor * tsr_half(0.25);
 			#else
-				HistoryOutput_LowFrequency[dpv_hi(LocalHistoryPixelPos)]  = dpv_hi(FinalLowFrequencyColor);
-				HistoryOutput_HighFrequency[dpv_hi(LocalHistoryPixelPos)] = dpv_hi(PremultipliedHighFrequencyColor);
-				HistoryOutput_Metadata[dpv_hi(LocalHistoryPixelPos)]      = tsr_half2(dpv_hi(FinalLowFrequencyWeight), dpv_hi(FinalHistoryValidity));
+				tsr_half4x2 HalfResFinalHighFrequencyColor = tsr_half4x2(FinalHighFrequencyColor[0], FinalHighFrequencyColor[1], FinalHighFrequencyColor[2], tsr_half(0.0).xx) * tsr_half(0.25);
 			#endif
-			#if CONFIG_ACCUMULATE_TRANSLUCENCY_SEPARATLY
-				HistoryOutput_Translucency[dpv_hi(LocalHistoryPixelPos)] = dpv_hi(FinalTranslucencyColor).rgb;
-				HistoryOutput_TranslucencyAlpha[dpv_hi(LocalHistoryPixelPos)] = dpv_hi(FinalTranslucencyColor).a;
+
+			#if PLATFORM_SUPPORTS_WAVE_BROADCAST && 0 // Support WaveBroadcast with halfs
+			{
+				FWaveBroadcastSettings Horizontal = InitWaveXorButterfly(/* XorButterFly = */ 0x1);
+				FWaveBroadcastSettings Vertical = InitWaveXorButterfly(/* XorButterFly = */ 0x2);
+
+				HalfResFinalHighFrequencyColor += WaveBroadcast(Horizontal, HalfResFinalHighFrequencyColor);
+				HalfResFinalHighFrequencyColor += WaveBroadcast(Vertical, HalfResFinalHighFrequencyColor);
+			}
+			#else
+			{
+				SharedArray0[LocalGroupThreadIndex] = dpv_lo(HalfResFinalHighFrequencyColor);
+				#if CONFIG_ENABLE_DUAL_PIXEL_VECTORIZATION
+					SharedArray1[LocalGroupThreadIndex] = dpv_hi(HalfResFinalHighFrequencyColor);
+				#endif
+
+				SharedArray0[LocalGroupThreadIndex] += SharedArray0[LocalGroupThreadIndex ^ 0x1];
+				#if CONFIG_ENABLE_DUAL_PIXEL_VECTORIZATION
+					SharedArray1[LocalGroupThreadIndex] += SharedArray1[LocalGroupThreadIndex ^ 0x1];
+				#endif
+					
+				SharedArray0[LocalGroupThreadIndex] += SharedArray0[LocalGroupThreadIndex ^ 0x2];
+				#if CONFIG_ENABLE_DUAL_PIXEL_VECTORIZATION
+					SharedArray1[LocalGroupThreadIndex] += SharedArray1[LocalGroupThreadIndex ^ 0x2];
+				#endif
+					
+				#if CONFIG_ENABLE_DUAL_PIXEL_VECTORIZATION
+					HalfResFinalHighFrequencyColor = dpv_interleave_registers(SharedArray0[LocalGroupThreadIndex], SharedArray1[LocalGroupThreadIndex]);
+				#else
+					HalfResFinalHighFrequencyColor = dpv_interleave_mono_registers(SharedArray0[LocalGroupThreadIndex]);
+				#endif
+
+			}
+			#endif
+			
+			//HalfResFinalHighFrequencyColor[0] = 1.0;
+			//HalfResFinalHighFrequencyColor[1] = 0.0;
+			//HalfResFinalHighFrequencyColor[2] = 0.0;
+
+			//HalfResOutputPixelPos = LocalHistoryPixelPos;
+
+			#if CONFIG_SCENE_COLOR_APLHA
+				HistoryOutput_HighFrequencyHalfRes[dpv_lo(HalfResOutputPixelPos)] = dpv_lo(HalfResFinalHighFrequencyColor);
+			
+				#if CONFIG_ENABLE_DUAL_PIXEL_VECTORIZATION
+					HistoryOutput_HighFrequencyHalfRes[dpv_hi(HalfResOutputPixelPos)] = dpv_hi(HalfResFinalHighFrequencyColor);
+				#endif
+			#else
+				HistoryOutput_HighFrequencyHalfRes[dpv_lo(HalfResOutputPixelPos)] = dpv_lo(HalfResFinalHighFrequencyColor).rgb;
+			
+				#if CONFIG_ENABLE_DUAL_PIXEL_VECTORIZATION
+					HistoryOutput_HighFrequencyHalfRes[dpv_hi(HalfResOutputPixelPos)] = dpv_hi(HalfResFinalHighFrequencyColor).rgb;
+				#endif
 			#endif
-			HistoryOutput_SubpixelDetails[dpv_hi(LocalHistoryPixelPos)] = dpv_hi(FinalSubpixelDetails);
 		}
-		#endif
+		#endif // !CONFIG_OUTPUT_HALFRES_HISTORY
+
 
 		// Output final scene color Mip0
 		tsr_halfCx2 OutputColor;
@@ -1631,10 +1708,6 @@ void MainCS(
 		#if !DEBUG_OUTPUT // running out of UAVs
 		{
 			tsr_halfCx2 HalfResOutput = OutputColor * tsr_half(0.25);
-			//tsr_short2x2 HalfResOutputPixelPos = (LocalHistoryPixelPos >> tsr_short(1)) | (((LocalHistoryPixelPos & tsr_short(0x1)) | !bGenerateOutputMip1) * tsr_short(~0));
-			tsr_short2x2 HalfResOutputPixelPos;
-			HalfResOutputPixelPos[0] = (LocalHistoryPixelPos[0] >> tsr_short(1)) | (((LocalHistoryPixelPos[0] & tsr_short(0x1)) | tsr_short(!bGenerateOutputMip1)) * tsr_short(~0));
-			HalfResOutputPixelPos[1] = (LocalHistoryPixelPos[1] >> tsr_short(1)) | (((LocalHistoryPixelPos[1] & tsr_short(0x1)) | tsr_short(!bGenerateOutputMip1)) * tsr_short(~0));
 			
 			#if PLATFORM_SUPPORTS_WAVE_BROADCAST && 0 // Support WaveBroadcast with halfs
 			BRANCH
@@ -1688,6 +1761,9 @@ void MainCS(
 			}
 			#endif
 
+			tsr_short2x2 SceneColorOutputMip1PixelPos = HalfResOutputPixelPos;
+			SceneColorOutputMip1PixelPos[0] = select(bGenerateOutputMip1 != 0, SceneColorOutputMip1PixelPos[0], tsr_short(~0));
+
 			SceneColorOutputMip1[dpv_lo(HalfResOutputPixelPos)] = dpv_lo(HalfResOutput);
 			
 			#if CONFIG_ENABLE_DUAL_PIXEL_VECTORIZATION
diff --git a/Engine/Source/Runtime/Renderer/Private/PostProcess/TemporalSuperResolution.cpp b/Engine/Source/Runtime/Renderer/Private/PostProcess/TemporalSuperResolution.cpp
--- a/Engine/Source/Runtime/Renderer/Private/PostProcess/TemporalSuperResolution.cpp
+++ b/Engine/Source/Runtime/Renderer/Private/PostProcess/TemporalSuperResolution.cpp
@@ -85,16 +85,6 @@ TAutoConsoleVariable<float> CVarTSRVelocityExtrapolation(
 	TEXT("Defines how much the velocity should be extrapolated on geometric discontinuities (Default = 1.0f)."),
 	ECVF_Scalability | ECVF_RenderThreadSafe);
 
-TAutoConsoleVariable<int32> CVarTSRVelocityHoleFill(
-	TEXT("r.TSR.Velocity.HoleFill"), 1,
-	TEXT("Whether to holl-fill the velocity buffer on parralax disocclusion to reduce boiling of disoccluded areas."),
-	ECVF_Scalability | ECVF_RenderThreadSafe);
-
-TAutoConsoleVariable<float> CVarTSRVelocityMaxHoleFillScatterVelocity(
-	TEXT("r.TSR.Velocity.HoleFill.MaxScatterVelocity"), 8.0f,
-	TEXT("Maximum output pixel velocity difference tolerated between the disoccluded and scattered disoccluding geometries."),
-	ECVF_RenderThreadSafe);
-
 #if COMPILE_TSR_DEBUG_PASSES
 
 TAutoConsoleVariable<int32> CVarTSRSetupDebugPasses(
@@ -125,6 +115,7 @@ END_SHADER_PARAMETER_STRUCT()
 BEGIN_SHADER_PARAMETER_STRUCT(FTSRHistoryTextures, )
 	SHADER_PARAMETER_RDG_TEXTURE(Texture2D, LowFrequency)
 	SHADER_PARAMETER_RDG_TEXTURE(Texture2D, HighFrequency)
+	SHADER_PARAMETER_RDG_TEXTURE(Texture2D, HighFrequencyHalfRes)
 	SHADER_PARAMETER_RDG_TEXTURE(Texture2D, Metadata)
 	SHADER_PARAMETER_RDG_TEXTURE(Texture2D, SubpixelDetails)
 	SHADER_PARAMETER_RDG_TEXTURE(Texture2D, Translucency)
@@ -134,6 +125,7 @@ END_SHADER_PARAMETER_STRUCT()
 BEGIN_SHADER_PARAMETER_STRUCT(FTSRHistoryUAVs, )
 	SHADER_PARAMETER_RDG_TEXTURE_UAV(RWTexture2D, LowFrequency)
 	SHADER_PARAMETER_RDG_TEXTURE_UAV(RWTexture2D, HighFrequency)
+	SHADER_PARAMETER_RDG_TEXTURE_UAV(RWTexture2D, HighFrequencyHalfRes)
 	SHADER_PARAMETER_RDG_TEXTURE_UAV(RWTexture2D, Metadata)
 	SHADER_PARAMETER_RDG_TEXTURE_UAV(RWTexture2D, SubpixelDetails)
 	SHADER_PARAMETER_RDG_TEXTURE_UAV(RWTexture2D, Translucency)
@@ -154,6 +146,10 @@ FTSRHistoryUAVs CreateUAVs(FRDGBuilder& GraphBuilder, const FTSRHistoryTextures&
 		UAVs.LowFrequency = GraphBuilder.CreateUAV(Textures.LowFrequency);
 	}
 	UAVs.HighFrequency = GraphBuilder.CreateUAV(Textures.HighFrequency);
+	if (Textures.HighFrequencyHalfRes)
+	{
+		UAVs.HighFrequencyHalfRes = GraphBuilder.CreateUAV(Textures.HighFrequencyHalfRes);
+	}
 	UAVs.Metadata = GraphBuilder.CreateUAV(Textures.Metadata);
 	UAVs.SubpixelDetails = GraphBuilder.CreateUAV(Textures.SubpixelDetails);
 	UAVs.Translucency = GraphBuilder.CreateUAV(Textures.Translucency);
@@ -260,8 +256,13 @@ class FTSRDecimateHistoryCS : public FTSRShader
 		SHADER_PARAMETER_RDG_TEXTURE_UAV(RWTexture2D, HalfResSceneColorOutput)
 		SHADER_PARAMETER_RDG_TEXTURE_UAV(RWTexture2D, HalfResPredictionSceneColorOutput)
 		SHADER_PARAMETER_RDG_TEXTURE_UAV(RWTexture2D, HalfResParallaxRejectionMaskOutput)
+
 		SHADER_PARAMETER_RDG_TEXTURE_UAV(RWTexture2D, PredictionSceneColorOutput)
+
+		SHADER_PARAMETER_RDG_TEXTURE_UAV(RWTexture2D, HoleFilledVelocityOutput)
+		SHADER_PARAMETER_RDG_TEXTURE_UAV(RWTexture2D, HoleFilledVelocityMaskOutput)
 		SHADER_PARAMETER_RDG_TEXTURE_UAV(RWTexture2D, ParallaxRejectionMaskOutput)
+
 		SHADER_PARAMETER_RDG_TEXTURE_UAV(RWTexture2D, DebugOutput)
 	END_SHADER_PARAMETER_STRUCT()
 }; // class FTSRDecimateHistoryCS
@@ -290,25 +291,6 @@ class FTSRCompareTranslucencyCS : public FTSRShader
 	END_SHADER_PARAMETER_STRUCT()
 }; // class FTSRCompareTranslucencyCS
 
-class FTSRHoleFillVelocityCS : public FTSRShader
-{
-	DECLARE_GLOBAL_SHADER(FTSRHoleFillVelocityCS);
-	SHADER_USE_PARAMETER_STRUCT(FTSRHoleFillVelocityCS, FTSRShader);
-
-	BEGIN_SHADER_PARAMETER_STRUCT(FParameters, )
-		SHADER_PARAMETER_STRUCT_INCLUDE(FTSRCommonParameters, CommonParameters)
-		SHADER_PARAMETER(float, MaxHollFillPixelVelocity)
-
-		SHADER_PARAMETER_RDG_TEXTURE(Texture2D, PrevClosestDepthTexture)
-		SHADER_PARAMETER_RDG_TEXTURE(Texture2D, DilatedVelocityTexture)
-		SHADER_PARAMETER_RDG_TEXTURE(Texture2D, ParallaxRejectionMaskTexture)
-
-		SHADER_PARAMETER_RDG_TEXTURE_UAV(RWTexture2D, HoleFilledVelocityOutput)
-		SHADER_PARAMETER_RDG_TEXTURE_UAV(RWTexture2D, HoleFilledVelocityMaskOutput)
-		SHADER_PARAMETER_RDG_TEXTURE_UAV(RWTexture2D, DebugOutput)
-	END_SHADER_PARAMETER_STRUCT()
-}; // class FTSRHoleFillVelocityCS
-
 class FTSRFilterFrequenciesCS : public FTSRShader
 {
 	DECLARE_GLOBAL_SHADER(FTSRFilterFrequenciesCS);
@@ -541,7 +523,6 @@ IMPLEMENT_GLOBAL_SHADER(FTSRClearPrevTexturesCS,     "/Engine/Private/TemporalSu
 IMPLEMENT_GLOBAL_SHADER(FTSRDilateVelocityCS,        "/Engine/Private/TemporalSuperResolution/TSRDilateVelocity.usf",        "MainCS", SF_Compute);
 IMPLEMENT_GLOBAL_SHADER(FTSRDecimateHistoryCS,       "/Engine/Private/TemporalSuperResolution/TSRDecimateHistory.usf",       "MainCS", SF_Compute);
 IMPLEMENT_GLOBAL_SHADER(FTSRCompareTranslucencyCS,   "/Engine/Private/TemporalSuperResolution/TSRCompareTranslucency.usf",   "MainCS", SF_Compute);
-IMPLEMENT_GLOBAL_SHADER(FTSRHoleFillVelocityCS,      "/Engine/Private/TemporalSuperResolution/TSRHoleFillVelocity.usf",      "MainCS", SF_Compute);
 IMPLEMENT_GLOBAL_SHADER(FTSRFilterFrequenciesCS,     "/Engine/Private/TemporalSuperResolution/TSRFilterFrequencies.usf",     "MainCS", SF_Compute);
 IMPLEMENT_GLOBAL_SHADER(FTSRCompareHistoryCS,        "/Engine/Private/TemporalSuperResolution/TSRCompareHistory.usf",        "MainCS", SF_Compute);
 IMPLEMENT_GLOBAL_SHADER(FTSRPostfilterRejectionCS,   "/Engine/Private/TemporalSuperResolution/TSRPostfilterRejection.usf",   "MainCS", SF_Compute);
@@ -610,6 +591,9 @@ ITemporalUpscaler::FOutputs AddTemporalSuperResolutionPasses(
 
 	const bool bHistoryHighFrequencyOnly = CVarTSRHistoryHighFrequencyOnly.GetValueOnRenderThread() != 0;
 
+	// Whether the UpdateHistory pass also output half res history to reduce memory bottleneck up the FTSRDecimateHistoryCS.
+	const bool bHalfResHistory = bHistoryHighFrequencyOnly;
+
 	EPixelFormat ColorFormat = bSupportsAlpha ? PF_FloatRGBA : PF_FloatR11G11B10;
 
 	int32 RejectionAntiAliasingQuality = FMath::Clamp(CVarTSRRejectionAntiAliasingQuality.GetValueOnRenderThread(), 1, 2);
@@ -618,8 +602,6 @@ ITemporalUpscaler::FOutputs AddTemporalSuperResolutionPasses(
 		RejectionAntiAliasingQuality = 0; 
 	}
 
-	bool bHoleFillVelocity = CVarTSRVelocityHoleFill.GetValueOnRenderThread() != 0 && RejectionAntiAliasingQuality > 0;
-
 	enum class ERejectionPostFilter : uint8
 	{
 		Disabled,
@@ -792,7 +774,7 @@ ITemporalUpscaler::FOutputs AddTemporalSuperResolutionPasses(
 				FClearValueBinding::None,
 				/* InFlags = */ TexCreate_ShaderResource | TexCreate_UAV);
 
-			DilatedVelocityTexture = GraphBuilder.CreateTexture(Desc, TEXT("TSR.DilatedVelocity"));
+			DilatedVelocityTexture = GraphBuilder.CreateTexture(Desc, TEXT("TSR.Velocity.Dilated"));
 
 			Desc.Format = PF_R16F;
 			ClosestDepthTexture = GraphBuilder.CreateTexture(Desc, TEXT("TSR.ClosestDepthTexture"));
@@ -884,6 +866,7 @@ ITemporalUpscaler::FOutputs AddTemporalSuperResolutionPasses(
 	{
 		DummyHistory.LowFrequency = BlackDummy;
 		DummyHistory.HighFrequency = BlackDummy;
+		DummyHistory.HighFrequencyHalfRes = BlackDummy;
 		DummyHistory.Metadata = BlackDummy;
 		DummyHistory.Translucency = BlackDummy;
 		DummyHistory.TranslucencyAlpha = WhiteDummy;
@@ -896,9 +879,9 @@ ITemporalUpscaler::FOutputs AddTemporalSuperResolutionPasses(
 	if (InputHistory.IsValid())
 	{
 		// Register filterable history
-
 		PrevHistory.LowFrequency = InputHistory.LowFrequency.IsValid() ? GraphBuilder.RegisterExternalTexture(InputHistory.LowFrequency) : DummyHistory.LowFrequency;
 		PrevHistory.HighFrequency = GraphBuilder.RegisterExternalTexture(InputHistory.HighFrequency);
+		PrevHistory.HighFrequencyHalfRes = InputHistory.HighFrequencyHalfRes.IsValid() ? GraphBuilder.RegisterExternalTexture(InputHistory.HighFrequencyHalfRes) : DummyHistory.HighFrequencyHalfRes;
 		PrevHistory.Metadata = GraphBuilder.RegisterExternalTexture(InputHistory.Metadata);
 		PrevHistory.Translucency = InputHistory.Translucency.IsValid() ? GraphBuilder.RegisterExternalTexture(InputHistory.Translucency) : DummyHistory.Translucency;
 		PrevHistory.TranslucencyAlpha = InputHistory.TranslucencyAlpha.IsValid() ? GraphBuilder.RegisterExternalTexture(InputHistory.TranslucencyAlpha) : DummyHistory.TranslucencyAlpha;
@@ -937,11 +920,19 @@ ITemporalUpscaler::FOutputs AddTemporalSuperResolutionPasses(
 			FClearValueBinding::None,
 			TexCreate_ShaderResource | TexCreate_UAV);
 
+		FRDGTextureDesc HalfResDesc = Desc;
+		HalfResDesc.Extent /= 2;
+
 		if (!bHistoryHighFrequencyOnly)
 		{
 			History.LowFrequency = GraphBuilder.CreateTexture(Desc, TEXT("TSR.History.LowFrequencies"));
 		}
+
 		History.HighFrequency = GraphBuilder.CreateTexture(Desc, TEXT("TSR.History.HighFrequencies"));
+		if (bHalfResHistory)
+		{
+			History.HighFrequencyHalfRes = GraphBuilder.CreateTexture(HalfResDesc, TEXT("TSR.History.HighFrequenciesHalfRes"));
+		}
 
 		Desc.Format = PF_R8G8;
 		History.Metadata = GraphBuilder.CreateTexture(Desc, TEXT("TSR.History.Metadata"));
@@ -962,7 +953,9 @@ ITemporalUpscaler::FOutputs AddTemporalSuperResolutionPasses(
 	FRDGTextureRef HalfResParallaxRejectionMaskTexture = nullptr;
 	FRDGTextureRef PredictionSceneColorTexture = nullptr;
 	FRDGTextureRef ParallaxRejectionMaskTexture = nullptr;
+	FRDGTextureRef HoleFilledVelocityMaskTexture = nullptr;
 	{
+		FRDGTextureRef HoleFilledVelocityTexture;
 		{
 			FRDGTextureDesc Desc = FRDGTextureDesc::Create2D(
 				InputExtent,
@@ -971,6 +964,12 @@ ITemporalUpscaler::FOutputs AddTemporalSuperResolutionPasses(
 				/* InFlags = */ TexCreate_ShaderResource | TexCreate_UAV);
 
 			ParallaxRejectionMaskTexture = GraphBuilder.CreateTexture(Desc, TEXT("TSR.ParallaxRejectionMask"));
+
+			Desc.Format = PF_G16R16;
+			HoleFilledVelocityTexture = GraphBuilder.CreateTexture(Desc, TEXT("TSR.Velocity.HoleFilled"));
+
+			Desc.Format = PF_R8G8;
+			HoleFilledVelocityMaskTexture = GraphBuilder.CreateTexture(Desc, TEXT("TSR.velocity.HoleFillMask"));
 		}
 
 		FTSRDecimateHistoryCS::FParameters* PassParameters = GraphBuilder.AllocParameters<FTSRDecimateHistoryCS::FParameters>();
@@ -1033,6 +1032,8 @@ ITemporalUpscaler::FOutputs AddTemporalSuperResolutionPasses(
 			PassParameters->PredictionSceneColorOutput = GraphBuilder.CreateUAV(PredictionSceneColorTexture);
 		}
 
+		PassParameters->HoleFilledVelocityOutput = GraphBuilder.CreateUAV(HoleFilledVelocityTexture);
+		PassParameters->HoleFilledVelocityMaskOutput = GraphBuilder.CreateUAV(HoleFilledVelocityMaskTexture);
 		PassParameters->ParallaxRejectionMaskOutput = GraphBuilder.CreateUAV(ParallaxRejectionMaskTexture);
 		PassParameters->DebugOutput = CreateDebugUAV(LowFrequencyExtent, TEXT("Debug.TSR.DecimateHistory"));
 
@@ -1051,6 +1052,8 @@ ITemporalUpscaler::FOutputs AddTemporalSuperResolutionPasses(
 			ComputeShader,
 			PassParameters,
 			FComputeShaderUtils::GetGroupCount(InputRect.Size(), 8));
+
+		DilatedVelocityTexture = HoleFilledVelocityTexture;
 	}
 
 	FRDGTextureRef TranslucencyRejectionTexture = nullptr;
@@ -1111,47 +1114,6 @@ ITemporalUpscaler::FOutputs AddTemporalSuperResolutionPasses(
 			FComputeShaderUtils::GetGroupCount(InputRect.Size(), 8));
 	}
 
-	FRDGTextureRef HoleFilledVelocityMaskTexture = BlackDummy;
-	if (bHoleFillVelocity)
-	{
-		FRDGTextureRef HoleFilledVelocityTexture;
-		{
-			FRDGTextureDesc Desc = FRDGTextureDesc::Create2D(
-				InputExtent,
-				PF_G16R16,
-				FClearValueBinding::None,
-				/* InFlags = */ TexCreate_ShaderResource | TexCreate_UAV);
-
-			HoleFilledVelocityTexture = GraphBuilder.CreateTexture(Desc, TEXT("TSR.Velocity.HoleFilled"));
-
-			Desc.Format = PF_R8G8;
-			HoleFilledVelocityMaskTexture = GraphBuilder.CreateTexture(Desc, TEXT("TSR.velocity.HoleFillMask"));
-		}
-
-		FTSRHoleFillVelocityCS::FParameters* PassParameters = GraphBuilder.AllocParameters<FTSRHoleFillVelocityCS::FParameters>();
-		PassParameters->CommonParameters = CommonParameters;
-
-		PassParameters->MaxHollFillPixelVelocity = CVarTSRVelocityMaxHoleFillScatterVelocity.GetValueOnRenderThread() * ScreenPercentage;
-		PassParameters->PrevClosestDepthTexture = PrevClosestDepthTexture;
-		PassParameters->DilatedVelocityTexture = DilatedVelocityTexture;
-		PassParameters->ParallaxRejectionMaskTexture = ParallaxRejectionMaskTexture;
-
-		PassParameters->HoleFilledVelocityOutput = GraphBuilder.CreateUAV(HoleFilledVelocityTexture);
-		PassParameters->HoleFilledVelocityMaskOutput = GraphBuilder.CreateUAV(HoleFilledVelocityMaskTexture);
-		PassParameters->DebugOutput = CreateDebugUAV(LowFrequencyExtent, TEXT("Debug.TSR.HoleFillVelocity"));
-
-		TShaderMapRef<FTSRHoleFillVelocityCS> ComputeShader(View.ShaderMap);
-		FComputeShaderUtils::AddPass(
-			GraphBuilder,
-			RDG_EVENT_NAME("TSR HoleFillVelocity %dx%d", InputRect.Width(), InputRect.Height()),
-			ComputePassFlags,
-			ComputeShader,
-			PassParameters,
-			FComputeShaderUtils::GetGroupCount(InputRect.Size(), 8));
-
-		DilatedVelocityTexture = HoleFilledVelocityTexture;
-	}
-
 	// Reject the history with frequency decomposition.
 	FRDGTextureRef HistoryRejectionTexture;
 	FRDGTextureRef InputSceneColorLdrLumaTexture = nullptr;
@@ -1554,6 +1516,12 @@ ITemporalUpscaler::FOutputs AddTemporalSuperResolutionPasses(
 		// Extract non-filterable history
 		GraphBuilder.QueueTextureExtraction(History.SubpixelDetails, &OutputHistory.SubpixelDetails);
 
+		// Extract half res history
+		if (bHalfResHistory)
+		{
+			GraphBuilder.QueueTextureExtraction(History.HighFrequencyHalfRes, &OutputHistory.HighFrequencyHalfRes);
+		}
+
 		// Extract the translucency buffer to compare it with next frame
 		if (bRejectSeparateTranslucency)
 		{
diff --git a/Engine/Source/Runtime/Renderer/Private/SceneRendering.h b/Engine/Source/Runtime/Renderer/Private/SceneRendering.h
--- a/Engine/Source/Runtime/Renderer/Private/SceneRendering.h
+++ b/Engine/Source/Runtime/Renderer/Private/SceneRendering.h
@@ -891,6 +891,9 @@ struct FTSRHistory
 	// Non-filterable output resolution
 	TRefCountPtr<IPooledRenderTarget> SubpixelDetails;
 
+	// Filterable half-res output resolution
+	TRefCountPtr<IPooledRenderTarget> HighFrequencyHalfRes;
+
 	FIntRect OutputViewportRect;
 
 
