Fix broken Mesh Texture Painting. Issue introduced in 17790257 from LWC changes.

#rnx

diff --git a/Engine/Shaders/Private/MeshPaintVertexShader.usf b/Engine/Shaders/Private/MeshPaintVertexShader.usf
--- a/Engine/Shaders/Private/MeshPaintVertexShader.usf
+++ b/Engine/Shaders/Private/MeshPaintVertexShader.usf
@@ -9,8 +9,8 @@
 float4x4 c_Transform;
 
 void Main( float4 InPosition : ATTRIBUTE0,
-		   float2 InCloneTextureCoordinates : ATTRIBUTE1,
-		   float3 InWorldSpaceVertexPosition : ATTRIBUTE2,
+		   float2 InCloneTextureCoordinates : ATTRIBUTE2,
+		   float3 InWorldSpaceVertexPosition : ATTRIBUTE3,
 		   out float4 OutPosition : SV_POSITION,
 		   out float2 OutCloneTextureCoordinates: TEXCOORD0,
 		   out float3 OutWorldSpaceVertexPosition : TEXCOORD1 )
diff --git a/Engine/Shaders/Private/meshpaintdilatevertexshader.usf b/Engine/Shaders/Private/meshpaintdilatevertexshader.usf
--- a/Engine/Shaders/Private/meshpaintdilatevertexshader.usf
+++ b/Engine/Shaders/Private/meshpaintdilatevertexshader.usf
@@ -9,8 +9,8 @@
 float4x4 c_Transform;
 
 void Main( in float4 InPosition : ATTRIBUTE0,
-		   in float2 InTextureCoordinates : ATTRIBUTE1,
-		   in float3 InWorldSpaceVertexPosition : ATTRIBUTE2,
+		   in float2 InTextureCoordinates : ATTRIBUTE2,
+		   in float3 InWorldSpaceVertexPosition : ATTRIBUTE3,
 		
 		   out float2 OutTextureCoordinates: TEXCOORD0,
 		   out float3 OutWorldSpaceVertexPosition : TEXCOORD1,
