Fix GeometryCollection velocity vector calculation.
VertexFactoryGetPreviousWorldPosition() was using current frame view transform.
#rnx

diff --git a/Engine/Shaders/Private/GeometryCollectionVertexFactory.ush b/Engine/Shaders/Private/GeometryCollectionVertexFactory.ush
--- a/Engine/Shaders/Private/GeometryCollectionVertexFactory.ush
+++ b/Engine/Shaders/Private/GeometryCollectionVertexFactory.ush
@@ -504,7 +504,9 @@ float3 VertexFactoryGetWorldNormal(FVertexFactoryInput Input, FVertexFactoryInte
 // @return previous translated world position
 float4 VertexFactoryGetPreviousWorldPosition(FVertexFactoryInput Input, FVertexFactoryIntermediates Intermediates)
 {
-	return TransformLocalToTranslatedWorld(Input.Position.xyz, Intermediates.PrevLocalToWorld);
+	FLWCVector3 WorldPosition = LWCMultiply(Input.Position.xyz, Intermediates.PrevLocalToWorld);
+	FLWCVector3 TranslatedWorldPosition = LWCAdd(WorldPosition, ResolvedView.PrevPreViewTranslation);
+	return float4(LWCToFloat(TranslatedWorldPosition), 1.0f);
 }
 
 float4 VertexFactoryGetTranslatedPrimitiveVolumeBounds(FVertexFactoryInterpolantsVSToPS Interpolants)
