Strata - more debug of projected ggx data.

diff --git a/Engine/Shaders/Private/Strata/StrataRoughRefraction.usf b/Engine/Shaders/Private/Strata/StrataRoughRefraction.usf
--- a/Engine/Shaders/Private/Strata/StrataRoughRefraction.usf
+++ b/Engine/Shaders/Private/Strata/StrataRoughRefraction.usf
@@ -167,7 +167,7 @@ FUI GetUI(inout FShaderPrintContext Context)
 
 	Newline(Context);
 
-	UI.SlabThickness = AddSlider(Context, TEXT("Slab Thickness"), 5.0f, GetDefaultFontColor(), 0.0f, 20.0f);
+	UI.SlabThickness = AddSlider(Context, TEXT("Slab Thickness"), 5.0f, GetDefaultFontColor(), 0.0f, 50.0f);
 
 	Newline(Context);
 
@@ -388,6 +388,38 @@ void VisualizeRoughRefractionPS(
 		}
 	}
 
+#if 0
+	// Draw projected distribution as a curve
+	float3 PrevP = -1.0f;
+	float PrevValue = -1.0f;
+	for (float u = 0.0; u < 1.0; u += 0.01)
+	{
+		float3 P = lerp(SlabP0, SlabP1, u);
+		P.y = 0.5 * (SlabP0.y + SlabP1.y);
+		P.z = SlabBottomZ;
+
+		const float2 UV = 0.5 + 0.5 * ((P - RefractionP) / (VisProjectedRayRadius)).xy;
+		const int2 Coord = (UV * 64.0f + 0.5f);
+		if (all(Coord >= 0) && all(Coord < 64))
+		{
+			const float SampleCount = float(SampleCountTexture[Coord].x);
+			const float Value = float(SampleCount) * 0.05;
+
+			if (PrevValue != -1.0f)
+			{
+				AddLineWS(Context, 
+					PrevP +  float3(0, 0, PrevValue), 
+					P +  float3(0, 0, Value),
+					float4(1, 1, 0, 0.5));
+				//float4(GetHSVDebugColor(SampleCount * 0.01f), 0.5));
+			}
+
+			PrevP = P;
+			PrevValue = Value;
+		}
+	}
+#endif
+
 #if 0
 	const float SphereRadius = 10.0f;
 	float2 Sol = RayIntersectSphere(CamWorldPos, WorldDir, float4(RefractionP, SphereRadius));
@@ -436,7 +468,9 @@ void VisualizeRoughRefractionPS(
 		}
 	}
 
-	// Draw the lobe estimated by Belcour
+#if 0
+	// Draw the mean / variance estimated by Belcour.
+	// However, those data cannot be interpreted as-is because they have been linearized for a add operation of variance to lead to a angular distribution of two convolution (w.r.t. 2 roughnesses).
 #if MODE_REFLECTION
 	FStrataLobeStatistic PaperLobe = StrataGetReflectedLobe(StrataGetDiracLobe(UI.V), 1.0, UI.Roughness);
 #else
@@ -448,8 +482,9 @@ void VisualizeRoughRefractionPS(
 	Newline(Context);
 	Print(Context, TEXT("Belcour variance = "));
 	Print(Context, PaperLobe.Sigma, FontWhite, 5, 5);
+#endif
 
-	// Our lobe
+	// Our mean / variance
 	FLobeStat LobeState = LobStatisticsBuffer[0];
 	DrawCircleXY(Context, float3(SlabCenter.xy + LobeState.Mean * VisProjectedRayRadius, SlabBottomZ + 0.05), LobeState.Variance * VisProjectedRayRadius, float4(1, 1, 0, 1));
 
