Strata - minor clean up of comments that are no longer relevant.

diff --git a/Engine/Shaders/Private/Strata/Strata.ush b/Engine/Shaders/Private/Strata/Strata.ush
--- a/Engine/Shaders/Private/Strata/Strata.ush
+++ b/Engine/Shaders/Private/Strata/Strata.ush
@@ -1549,7 +1549,7 @@ FStrataData StrataAddParameterBlending(FStrataData A, FStrataData B, float AMixF
 	BSDF_SETHASANISOTROPY(ResultBSDF,	BSDF_GETHASANISOTROPY(BBSDF)					| BSDF_GETHASANISOTROPY(ABSDF));
 	BSDF_SETHASSSS(ResultBSDF,			BSDF_GETHASSSS(BBSDF)							| BSDF_GETHASSSS(ABSDF));
 	BSDF_SETHASSSSPROFILE(ResultBSDF,	BSDF_GETHASSSSPROFILE(BBSDF)					| BSDF_GETHASSSSPROFILE(ABSDF));	// Problem: cannot blend SSS profiles...
-	BSDF_SETISTHIN(ResultBSDF,			BSDF_GETISTHIN(BBSDF)							| BSDF_GETISTHIN(ABSDF));			// STRATA_TODO fix the MPF/MFP situation for that to really work
+	BSDF_SETISTHIN(ResultBSDF,			BSDF_GETISTHIN(BBSDF)							| BSDF_GETISTHIN(ABSDF));
 	BSDF_SETHASMFP(ResultBSDF,			BSDF_GETHASMFP(BBSDF)							| BSDF_GETHASMFP(ABSDF));
 	BSDF_SETHASHAZINESS(ResultBSDF,		BSDF_GETHASHAZINESS(BBSDF)						| BSDF_GETHASHAZINESS(ABSDF));
 	BSDF_SETHASF90(ResultBSDF,			BSDF_GETHASF90(BBSDF)							| BSDF_GETHASF90(ABSDF));
@@ -1638,8 +1638,8 @@ FStrataBSDF StrataVerticalLayeringParameterBlending(in FStrataBSDF TopBSDF, in F
 	BSDF_SETHASANISOTROPY(ResultBSDF,	0);																					// Anisotropy is disabled because it changes a lot of things
 	BSDF_SETHASSSS(ResultBSDF,			BSDF_GETHASSSS(BaseBSDF));															// We only keep SSS if the bottom layer has it
 	BSDF_SETHASSSSPROFILE(ResultBSDF,	0);																					// Problem: cannot blend SSS profiles so we simply disable SSSProfile
-	BSDF_SETISTHIN(ResultBSDF,			BSDF_GETISTHIN(BaseBSDF));															// STRATA_TODO fix the MPF/MFP situation for that to really work
-	BSDF_SETHASMFP(ResultBSDF,			BSDF_GETHASMFP(BaseBSDF)						| BSDF_GETHASMFP(TopBSDF));		// We keep MFP of both to make sure bottom layer on translucent can laos be effective
+	BSDF_SETISTHIN(ResultBSDF,			BSDF_GETISTHIN(BaseBSDF));
+	BSDF_SETHASMFP(ResultBSDF,			BSDF_GETHASMFP(BaseBSDF)						| BSDF_GETHASMFP(TopBSDF));			// We keep MFP of both to make sure bottom layer on translucent can also be effective
 	BSDF_SETHASHAZINESS(ResultBSDF,		BSDF_GETHASHAZINESS(TopBSDF));														// We only keep haziness if the top layer has it
 	BSDF_SETHASF90(ResultBSDF,			BSDF_GETHASF90(BaseBSDF)						| BSDF_GETHASF90(TopBSDF));			// F90: keep union of both even though it will be hard to get a match
 	BSDF_SETHASFUZZ(ResultBSDF,			BSDF_GETHASFUZZ(BaseBSDF)						| BSDF_GETHASFUZZ(TopBSDF));		// Fuzz: keep union of both even though it will be hard to get a match
@@ -1681,7 +1681,7 @@ FStrataBSDF StrataVerticalLayeringParameterBlending(in FStrataBSDF TopBSDF, in F
 	SLAB_ANISOTROPY(ResultBSDF) =		0.f;
 
 	// Keep the MFP form the bottom layer only (top layer simple volume is backed in parameters).
-	SLAB_SSSMFP(ResultBSDF) =			lerp(SLAB_SSSMFP(TopBSDF),						SLAB_SSSMFP(BaseBSDF),		TopToBottomLerp); // SLAB_SSSMFP(BaseBSDF);
+	SLAB_SSSMFP(ResultBSDF) =			lerp(SLAB_SSSMFP(TopBSDF),						SLAB_SSSMFP(BaseBSDF),		TopToBottomLerp);
 
 #if STRATA_INLINE_SHADING
 #if STRATA_OPAQUE_MATERIAL
@@ -2951,8 +2951,6 @@ void PackStrataOut(
 								SSSData.ProfileId = StrataSubsurfaceProfileIdTo8bits(SLAB_SSSPROFILEID(BSDF));
 								SSSData.ProfileRadiusScale = SLAB_SSSPROFILERADIUSSCALE(BSDF);
 
-								// STRATA_TODO This should be accumulated according to the BSDF weight when we run this in the base pass.
-								// We cannot do that now because BSDFWeight contains weight*coverage*trhougput all mixed together.
 								SSSData.MFP = SLAB_SSSMFP(BSDF);
 							}
 
