Fixed half resolution depth downsample shader code.
This was fixed for the checkerboarded minmax half resolution depth so now the uv setup is just common and correct.

diff --git a/Engine/Shaders/Private/DownsampleDepthPixelShader.usf b/Engine/Shaders/Private/DownsampleDepthPixelShader.usf
--- a/Engine/Shaders/Private/DownsampleDepthPixelShader.usf
+++ b/Engine/Shaders/Private/DownsampleDepthPixelShader.usf
@@ -14,8 +14,7 @@
 #define DOWNSAMPLE_DEPTH_FILTER_CBMINMAX	2
 
 Texture2D<float> DepthTexture;
-float4 SourceTexelOffsets01;
-float4 SourceTexelOffsets23;
+float2 DestinationTexelSize;
 float2 SourceMaxUV;
 float2 DestinationResolution;
 uint DownsampleDepthFilter;
@@ -33,14 +32,35 @@ void Main(
 	OutColor = 0;
 	float Depth = 0;
 
+	// Lookup the four DeviceZ's of the full resolution pixels corresponding to this low resolution pixel
+	//
+	// Source texture texels:
+	//	|------||------|------(UV + DestinationTexelSize)
+	//	|      ||      |      |
+	//	|      ||      |      |
+	//	|      ||      |      |
+	//	|------||------|------|
+	//	|      ||      |      |
+	//	|  z0  ||  z2  |      |
+	//	|      ||      |      |
+	//	|------UV------|------(UV + float2(DestinationTexelSize.x, 0))
+	//	|      ||      |      |
+	//	|  z1  ||  z3  |      |
+	//	|      ||      |      |
+	//	|------||------|------|
+	// OffsetScale is 0.25 because:
+	//  - UV.xy is centered on a half resolution texel of the destination render target (in [0,1] uv space). 
+	//  - This texel includes 2x2 texels from the source texture (z0, z1, z2 and z3).
+	//  - Scaling the DestinationTexelSize by 0.5 will get us to the edge of the half resolution texel. 
+	//  - And to sample the center of each texel from the high resolution source texture we need to divide again the offset by half, leading to a final scale of 0.25.
+	const float OffsetScale = 0.25f;
+	float DeviceZ0 = GetDeviceZ(UV.xy + DestinationTexelSize * float2(-OffsetScale,-OffsetScale));
+	float DeviceZ1 = GetDeviceZ(UV.xy + DestinationTexelSize * float2(-OffsetScale, OffsetScale));
+	float DeviceZ2 = GetDeviceZ(UV.xy + DestinationTexelSize * float2( OffsetScale,-OffsetScale));
+	float DeviceZ3 = GetDeviceZ(UV.xy + DestinationTexelSize * float2( OffsetScale, OffsetScale));
+
 	if (DownsampleDepthFilter < DOWNSAMPLE_DEPTH_FILTER_CBMINMAX)
 	{
-		// Lookup the four DeviceZ's of the full resolution pixels corresponding to this low resolution pixel
-		float DeviceZ0 = GetDeviceZ(UV.xy + SourceTexelOffsets01.xy);
-		float DeviceZ1 = GetDeviceZ(UV.xy + SourceTexelOffsets01.zw);
-		float DeviceZ2 = GetDeviceZ(UV.xy + SourceTexelOffsets23.xy);
-		float DeviceZ3 = GetDeviceZ(UV.xy + SourceTexelOffsets23.zw);
-
 #if HAS_INVERTED_Z_BUFFER
 		float FarDepth = min(min(DeviceZ0, DeviceZ1), min(DeviceZ2, DeviceZ3));
 #else
@@ -53,18 +73,6 @@ void Main(
 	}
 	else // DownsampleDepthFilter == DOWNSAMPLE_DEPTH_FILTER_CBMINMAX
 	{
-		// Lookup the four DeviceZ's of the full resolution pixels corresponding to this low resolution pixel
-		const float OffsetScale = 0.25f;
-		// OffsetScale is 0.25 because:
-		//  - UV is centered on the low resolution render targetr in UV space
-		//  - SourceTexelOffsets23 in fact represent a texel size for the DESTINATION render target (TODO refactor that code once in main)
-		//  - So scale by 0.5 will get us to the edge of the half resolution texel. 
-		//  - And to sample the center of the high resolution source pixel in UV space we need to divide angain the offset by half ==> 0.25
-		float DeviceZ0 = GetDeviceZ(UV.xy + SourceTexelOffsets23.zw * OffsetScale * float2(-1,-1));
-		float DeviceZ1 = GetDeviceZ(UV.xy + SourceTexelOffsets23.zw * OffsetScale * float2(-1, 1));
-		float DeviceZ2 = GetDeviceZ(UV.xy + SourceTexelOffsets23.zw * OffsetScale * float2( 1,-1));
-		float DeviceZ3 = GetDeviceZ(UV.xy + SourceTexelOffsets23.zw * OffsetScale * float2( 1, 1));
-
 		const float MaxDeviceZ = max(max(DeviceZ0, DeviceZ1), max(DeviceZ2, DeviceZ3));
 		const float MinDeviceZ = min(min(DeviceZ0, DeviceZ1), min(DeviceZ2, DeviceZ3));
 
diff --git a/Engine/Source/Runtime/Renderer/Private/ScreenPass.cpp b/Engine/Source/Runtime/Renderer/Private/ScreenPass.cpp
--- a/Engine/Source/Runtime/Renderer/Private/ScreenPass.cpp
+++ b/Engine/Source/Runtime/Renderer/Private/ScreenPass.cpp
@@ -180,8 +180,7 @@ public:
 	BEGIN_SHADER_PARAMETER_STRUCT(FParameters, )
 		SHADER_PARAMETER_STRUCT_REF(FViewUniformShaderParameters, View)
 		SHADER_PARAMETER_RDG_TEXTURE(Texture2D, DepthTexture)
-		SHADER_PARAMETER(FVector4f, SourceTexelOffsets01)
-		SHADER_PARAMETER(FVector4f, SourceTexelOffsets23)
+		SHADER_PARAMETER(FVector2f, DestinationTexelSize)
 		SHADER_PARAMETER(FVector2f, SourceMaxUV)
 		SHADER_PARAMETER(FVector2f, DestinationResolution)
 		SHADER_PARAMETER(uint32, DownsampleDepthFilter)
@@ -212,8 +211,7 @@ void AddDownsampleDepthPass(
 	FDownsampleDepthPS::FParameters* PassParameters = GraphBuilder.AllocParameters<FDownsampleDepthPS::FParameters>();
 	PassParameters->View = View.ViewUniformBuffer;
 	PassParameters->DepthTexture = Input.Texture;
-	PassParameters->SourceTexelOffsets01 = FVector4f(0.0f, 0.0f, 1.0f / OutputViewport.Extent.X, 0.0f);
-	PassParameters->SourceTexelOffsets23 = FVector4f(0.0f, 1.0f / OutputViewport.Extent.Y, 1.0f / OutputViewport.Extent.X, 1.0f / OutputViewport.Extent.Y);
+	PassParameters->DestinationTexelSize = FVector2f(1.0f / OutputViewport.Extent.X, 1.0f / OutputViewport.Extent.X);
 	PassParameters->SourceMaxUV = FVector2f((View.ViewRect.Max.X - 0.5f) / InputViewport.Extent.X, (View.ViewRect.Max.Y - 0.5f) / InputViewport.Extent.Y);
 	PassParameters->DownsampleDepthFilter = (uint32)DownsampleDepthFilter;
 
