Updates TSR subpixel details at input resolution for the r.TSR.History.HighFrequencyOnly permutation

diff --git a/Engine/Shaders/Private/TemporalSuperResolution/TSRCommon.ush b/Engine/Shaders/Private/TemporalSuperResolution/TSRCommon.ush
--- a/Engine/Shaders/Private/TemporalSuperResolution/TSRCommon.ush
+++ b/Engine/Shaders/Private/TemporalSuperResolution/TSRCommon.ush
@@ -291,6 +291,8 @@ float2 PrevHistoryInfo_UVViewportSizeInverse;
 float2 PrevHistoryInfo_UVViewportBilinearMin;
 float2 PrevHistoryInfo_UVViewportBilinearMax;
 FScreenTransform ScreenPosToPrevHistoryBufferUV;
+FScreenTransform ScreenPosToPrevSubpixelDetails;
+float2 PrevSubpixelDetailsExtent;
 float HistoryPreExposureCorrection;
 
 uint bCameraCut;
@@ -947,9 +949,9 @@ struct FSubpixelNeighborhood
 	bool2 ClosestSubpixelOffset;
 };
 
-FSubpixelNeighborhood GatherPrevSubpixelNeighborhood(Texture2D<uint> PrevSubPixelDetailTexture, float2 PrevHistoryBufferUV)
+FSubpixelNeighborhood GatherPrevSubpixelNeighborhood(Texture2D<tsr_subpixel_details> PrevSubPixelDetailTexture, float2 PrevHistoryBufferUV)
 {
-	float2 PixelCoord = PrevHistoryBufferUV * PrevHistoryInfo_Extent;
+	float2 PixelCoord = PrevHistoryBufferUV * PrevSubpixelDetailsExtent;
 	float2 TopLeftPixelCoord = floor(PixelCoord - 0.5) + 0.5;
 	
 	FSubpixelNeighborhood Neighborhood;
diff --git a/Engine/Shaders/Private/TemporalSuperResolution/TSRDebugHistory.usf b/Engine/Shaders/Private/TemporalSuperResolution/TSRDebugHistory.usf
--- a/Engine/Shaders/Private/TemporalSuperResolution/TSRDebugHistory.usf
+++ b/Engine/Shaders/Private/TemporalSuperResolution/TSRDebugHistory.usf
@@ -11,15 +11,21 @@
 
 //------------------------------------------------------- PARAMETERS
 
-Texture2D<tsr_half3> History_Textures_0;
-Texture2D<tsr_half3> History_Textures_1;
-Texture2D<tsr_half2> History_Textures_2;
-Texture2D<tsr_subpixel_details> History_Textures_3;
-
-Texture2D<tsr_half3> PrevHistory_Textures_0;
-Texture2D<tsr_half3> PrevHistory_Textures_1;
-Texture2D<tsr_half2> PrevHistory_Textures_2;
-Texture2D<tsr_subpixel_details> PrevHistory_Textures_3;
+Texture2D<tsr_halfC> History_LowFrequency;
+Texture2D<tsr_halfC> History_HighFrequency;
+Texture2D<tsr_halfC> History_HighFrequencyHalfRes;
+Texture2D<tsr_half2> History_Metadata;
+Texture2D<tsr_half3> History_Translucency;
+Texture2D<tsr_half>  History_TranslucencyAlpha;
+Texture2D<tsr_subpixel_details> History_SubpixelDetails;
+
+Texture2D<tsr_halfC> PrevHistory_LowFrequency;
+Texture2D<tsr_halfC> PrevHistory_HighFrequency;
+Texture2D<tsr_halfC> PrevHistory_HighFrequencyHalfRes;
+Texture2D<tsr_half2> PrevHistory_Metadata;
+Texture2D<tsr_half3> PrevHistory_Translucency;
+Texture2D<tsr_half>  PrevHistory_TranslucencyAlpha;
+Texture2D<tsr_subpixel_details> PrevHistory_SubpixelDetails;
 
 
 //------------------------------------------------------- FUNCTION
@@ -53,7 +59,7 @@ tsr_subpixel_payload ReadSubpixelDetails(uint2 SubpixelPos)
 	uint SubpixelId = dot(SubpixelRelPos, uint2(1, SUB_PIXEL_GRID_SIZE));
 	uint SubpixelShift = SubpixelId * SUB_PIXEL_BIT_COUNT;
 	
-	tsr_subpixel_details RawHistory3 = History_Textures_3[HistoryPixelPos];
+	tsr_subpixel_details RawHistory3 = History_SubpixelDetails[HistoryPixelPos];
 
 	return tsr_subpixel_payload(RawHistory3 >> SubpixelShift) & SUB_PIXEL_BIT_MASK;
 }
@@ -82,24 +88,24 @@ void MainCS(
 	
 	{
 		#if 1
-			tsr_half3 RawHistory0 = History_Textures_0[HistoryPixelPos];
-			tsr_half3 RawHistory1 = History_Textures_1[HistoryPixelPos];
-			tsr_half2 RawHistory2 = History_Textures_2[HistoryPixelPos];
-			tsr_subpixel_details RawHistory3 = History_Textures_3[HistoryPixelPos];
+			//tsr_half3 RawHistory0 = History_Textures_0[HistoryPixelPos];
+			//tsr_half3 RawHistory1 = History_Textures_1[HistoryPixelPos];
+			// RawHistory2 = History_Textures_2[HistoryPixelPos];
+			tsr_subpixel_details RawHistory3 = History_SubpixelDetails[HistoryPixelPos];
 		#else
-			tsr_half3 RawHistory0 = PrevHistory_Textures_0[HistoryPixelPos];
-			tsr_half3 RawHistory1 = PrevHistory_Textures_1[HistoryPixelPos];
-			tsr_half2 RawHistory2 = PrevHistory_Textures_2[HistoryPixelPos];
-			tsr_subpixel_details RawHistory3 = PrevHistory_Textures_3[HistoryPixelPos];
+			//tsr_half3 RawHistory0 = PrevHistory_Textures_0[HistoryPixelPos];
+			//tsr_half3 RawHistory1 = PrevHistory_Textures_1[HistoryPixelPos];
+			//tsr_half2 RawHistory2 = PrevHistory_Textures_2[HistoryPixelPos];
+			tsr_subpixel_details RawHistory3 = PrevHistory_SubpixelDetails[HistoryPixelPos];
 		#endif
 
 		tsr_subpixel_payload SubpixelDetails = tsr_subpixel_payload(RawHistory3 >> SubpixelShift) & SUB_PIXEL_BIT_MASK;
 
 		Debug.r = UncompressParallaxFactor(SubpixelDetails);
-		Debug.b = float(SubpixelId);
+		//Debug.b = float(SubpixelId);
 	}
 
-	#if DEBUG_OUTPUT && 0
+	#if DEBUG_OUTPUT && 1
 	if (all(DispatchThreadId < HistoryInfo_ViewportMax * SUB_PIXEL_GRID_SIZE))
 	{
 		DebugOutput[DispatchThreadId] = Debug;
diff --git a/Engine/Shaders/Private/TemporalSuperResolution/TSRDecimateHistory.usf b/Engine/Shaders/Private/TemporalSuperResolution/TSRDecimateHistory.usf
--- a/Engine/Shaders/Private/TemporalSuperResolution/TSRDecimateHistory.usf
+++ b/Engine/Shaders/Private/TemporalSuperResolution/TSRDecimateHistory.usf
@@ -15,6 +15,8 @@
 
 #define CONFIG_MANUAL_LDS_SPILL 1
 
+#define CONFIG_UPDATE_SUBPIXEL_DETAIL (DIM_HIGH_FREQUENCY_ONLY)
+
 #define CONFIG_VELOCITY_SAMPLES 9
 #if CONFIG_VELOCITY_SAMPLES != 1
 	#define CONFIG_VELOCITY_SAMPLE_PAIRS ((CONFIG_VELOCITY_SAMPLES - 1) / 2)
@@ -40,7 +42,7 @@ Texture2D<float2> DilatedVelocityTexture;
 Texture2D<float> ClosestDepthTexture;
 Texture2D<uint> PrevUseCountTexture;
 Texture2D<uint> PrevClosestDepthTexture;
-Texture2D<uint> ParallaxFactorTexture;
+Texture2D<tsr_subpixel_payload> ParallaxFactorTexture;
 
 Texture2D<tsr_halfC> PrevHistory_LowFrequency;
 Texture2D<tsr_halfC> PrevHistory_HighFrequency;
@@ -57,6 +59,7 @@ RWTexture2D<tsr_halfC> PredictionSceneColorOutput;
 RWTexture2D<float2> HoleFilledVelocityOutput;
 RWTexture2D<tsr_half2> HoleFilledVelocityMaskOutput;
 RWTexture2D<tsr_half> ParallaxRejectionMaskOutput;
+RWTexture2D<tsr_subpixel_payload> HistorySubpixelDetailsOutput;
 
 
 //------------------------------------------------------- LDS
@@ -133,6 +136,9 @@ void MainCS(
 	// Sample current frame informations.
 	float2 EncodedScreenVelocity;
 	float DeviceZ;
+	#if CONFIG_UPDATE_SUBPIXEL_DETAIL
+		tsr_subpixel_payload ParallaxFactorBits;
+	#endif
 	tsr_half ParallaxFactor;
 	tsr_half RejectHighFrequency = tsr_half(0.0);
 	ISOLATE
@@ -148,8 +154,13 @@ void MainCS(
 		// Issue texture fetechs
 		ISOLATE
 		{
+			tsr_subpixel_payload RawParallaxFactor = ParallaxFactorTexture[InputPixelPos];
+
 			DeviceZ = ClosestDepthTexture[InputPixelPos];
-			ParallaxFactor = UncompressParallaxFactor((ParallaxFactorTexture[InputPixelPos] >> 4) & 0xF);
+			ParallaxFactor = UncompressParallaxFactor(RawParallaxFactor & tsr_subpixel_payload(SUB_PIXEL_PARALLAX_FACTOR_BIT_MASK));
+			#if CONFIG_UPDATE_SUBPIXEL_DETAIL
+				ParallaxFactorBits = RawParallaxFactor & tsr_subpixel_payload(SUB_PIXEL_PARALLAX_FACTOR_BIT_MASK);
+			#endif
 
 			#if CONFIG_VELOCITY_SAMPLES == 1
 				EncodedScreenVelocity = DilatedVelocityTexture[InputPixelPos];
@@ -329,11 +340,10 @@ void MainCS(
 		tsr_ushort2 InputPixelPos = (
 			Map8x8Tile2x2Lane(GroupThreadIndex) +
 			(tsr_ushort2(InputInfo_ViewportMin) + tsr_ushort2(GroupId) * tsr_ushort2(TILE_SIZE, TILE_SIZE)));
-
+			
+		float2 ScreenPos = ApplyScreenTransform(float2(InputPixelPos), InputPixelPosToScreenPos);
+		float2 PrevScreenPos = ScreenPos - DecodeVelocityFromTexture(float4(EncodedScreenVelocity, 0.0, 0.0)).xy;
 		{
-			float2 ScreenPos = ApplyScreenTransform(float2(InputPixelPos), InputPixelPosToScreenPos);
-
-			float2 PrevScreenPos = ScreenPos - DecodeVelocityFromTexture(float4(EncodedScreenVelocity, 0.0, 0.0)).xy;
 			float2 PrevInputBufferUV = (InputInfo_ScreenPosToViewportScale * PrevScreenPos + InputInfo_ScreenPosToViewportBias) * InputInfo_ExtentInverse;
 		
 			FBilinearSampleInfos BilinearInter = GetBilinearSampleLevelInfos(PrevInputBufferUV, InputInfo_Extent, InputInfo_ExtentInverse);
@@ -425,8 +435,14 @@ void MainCS(
 				DecimatedColorSamples[0] = PrevHistory_LowFrequency.SampleLevel(GlobalBilinearClampedSampler, KernelHistoryBufferUV, 0);
 			}
 			#endif
-				
-			SubpixelNeigborhood = GatherPrevSubpixelNeighborhood(PrevHistory_SubpixelDetails, KernelHistoryBufferUV);
+		
+			#if CONFIG_UPDATE_SUBPIXEL_DETAIL
+				float2 SubpixelBufferUV = ApplyScreenTransform(PrevScreenPos, ScreenPosToPrevSubpixelDetails);
+			#else
+				float2 SubpixelBufferUV = KernelHistoryBufferUV;
+			#endif
+
+			SubpixelNeigborhood = GatherPrevSubpixelNeighborhood(PrevHistory_SubpixelDetails, SubpixelBufferUV);
 		}
 	}
 	
@@ -458,6 +474,8 @@ void MainCS(
 
 	// Process the history samples
 	tsr_half ParallaxRejectionMask = 0.0;
+	
+	tsr_subpixel_details FinalSubpixelDetails;
 
 	ISOLATE
 	{
@@ -479,14 +497,19 @@ void MainCS(
 	
 		// Compute HistoryMaxRejectionFactor
 		tsr_half HistoryMaxRejectionFactor;
+		tsr_subpixel_details PrevSubpixelDetails;
 		{
 			uint EncodedHistoryMaxRejectionFactor = 0;
+			
+			PrevSubpixelDetails = 0;
 
 			UNROLL_N(SUB_PIXEL_COUNT)
 			for (uint SubpixelId = 0; SubpixelId < SUB_PIXEL_COUNT; SubpixelId++)
 			{
 				tsr_subpixel_payload SubpixelPayload = GetSubpixelPayload(SubpixelNeigborhood, SubpixelId);
 				tsr_subpixel_payload EncodedRejectionFactor = (SubpixelPayload >> SUB_PIXEL_PARALLAX_FACTOR_BIT_OFFSET) & SUB_PIXEL_PARALLAX_FACTOR_BIT_MASK;
+				
+				PrevSubpixelDetails |= SubpixelPayload << tsr_subpixel_payload(SUB_PIXEL_BIT_COUNT * SubpixelId);
 
 				EncodedHistoryMaxRejectionFactor = max(EncodedHistoryMaxRejectionFactor, EncodedRejectionFactor);
 			}
@@ -552,29 +575,7 @@ void MainCS(
 			tsr_half HistorySubdetailRejectionFactor = max(HistoryMaxRejectionFactor / ParallaxFactor, tsr_half(1.0));
 
 			float4 ThisClip = float4(ScreenPos, DeviceZ, 1);
-			float4 PrevClip = mul(ThisClip, View.ClipToPrevClip);
-			float2 PrevScreen = PrevClip.xy / PrevClip.w;
-			
-			float4 RotationalPrevClip = mul(ThisClip, RotationalClipToPrevClip);
-			float2 RotationalPrevScreen = RotationalPrevClip.xy / RotationalPrevClip.w;
-			
-			tsr_half PixelParallaxVelocity = tsr_half(0.5) * length(tsr_half2(RotationalPrevScreen - PrevScreen) * tsr_half2(InputInfo_ViewportSize));
-			tsr_half HistorySubdetailAdditionalParralax = (HistorySubdetailRejectionFactor - tsr_half(1.0)) * PixelParallaxVelocity;
-
-			// PixelParallaxVelocity can be +inf and generate NaN if HistorySubdetailRejectionFactor = 1.0;
-			HistorySubdetailAdditionalParralax = -min(-HistorySubdetailAdditionalParralax, tsr_half(0.0));
-
-			tsr_half HistoryParallaxRejectionMask = saturate(tsr_half(1.0) - HistorySubdetailAdditionalParralax);
-
-			ParallaxRejectionMask = min(ParallaxRejectionMask, HistoryParallaxRejectionMask);
-		}
-
-		// Reject subpixel details that might be in history but not input just based on their parrallax
-		{
-			tsr_half HistorySubdetailRejectionFactor = max(HistoryMaxRejectionFactor / ParallaxFactor, tsr_half(1.0));
-
-			float4 ThisClip = float4(ScreenPos, DeviceZ, 1);
-			float4 PrevClip = mul(ThisClip, View.ClipToPrevClip);
+			float4 PrevClip = mul(ThisClip, View.ClipToPrevClip); 
 			float2 PrevScreen = PrevClip.xy / PrevClip.w;
 			
 			float4 RotationalPrevClip = mul(ThisClip, RotationalClipToPrevClip);
@@ -619,6 +620,37 @@ void MainCS(
 				//HoleFilledVelocityMask = tsr_half(HoleFillingOutOfRange) * tsr_half(saturate(MaxHollFillPixelVelocity - VelocityDelta));
 			}
 		}
+		
+		// Update subpixel detail
+		#if CONFIG_UPDATE_SUBPIXEL_DETAIL
+		{
+			FinalSubpixelDetails = PrevSubpixelDetails;
+
+			{
+				const tsr_subpixel_details ParallaxFactorMask = tsr_subpixel_details(
+					(SUB_PIXEL_PARALLAX_FACTOR_BIT_MASK << (SUB_PIXEL_PARALLAX_FACTOR_BIT_OFFSET + 0 * SUB_PIXEL_BIT_COUNT)) | 
+					(SUB_PIXEL_PARALLAX_FACTOR_BIT_MASK << (SUB_PIXEL_PARALLAX_FACTOR_BIT_OFFSET + 1 * SUB_PIXEL_BIT_COUNT)) | 
+					(SUB_PIXEL_PARALLAX_FACTOR_BIT_MASK << (SUB_PIXEL_PARALLAX_FACTOR_BIT_OFFSET + 2 * SUB_PIXEL_BIT_COUNT)) | 
+					(SUB_PIXEL_PARALLAX_FACTOR_BIT_MASK << (SUB_PIXEL_PARALLAX_FACTOR_BIT_OFFSET + 3 * SUB_PIXEL_BIT_COUNT)) | 
+					0x0);
+			
+				// Reset parallax factor
+				FinalSubpixelDetails = select(bIsOffscreenOrParallaxRejected, FinalSubpixelDetails & ~ParallaxFactorMask, FinalSubpixelDetails);
+			}
+		
+			// TODO(TSR): Move SubpixelId CPU
+			bool2 bKOSign = InputJitter < 0.0;
+			uint SubpixelId = dot(uint2(bKOSign), uint2(1, SUB_PIXEL_GRID_SIZE));
+
+			tsr_subpixel_payload SubpixelShift = tsr_subpixel_payload(SubpixelId) * tsr_subpixel_payload(SUB_PIXEL_BIT_COUNT);
+
+			tsr_subpixel_payload SubpixelPayload = (ParallaxFactorBits << tsr_subpixel_payload(SUB_PIXEL_PARALLAX_FACTOR_BIT_OFFSET));
+
+			tsr_subpixel_details UpdatedFinalSubpixelDetails = (FinalSubpixelDetails & (~(tsr_subpixel_payload(SUB_PIXEL_BIT_MASK) << SubpixelShift))) | (SubpixelPayload << SubpixelShift);
+
+			FinalSubpixelDetails = UpdatedFinalSubpixelDetails;
+		}
+		#endif // CONFIG_UPDATE_SUBPIXEL_DETAIL
 	}
 	
 	// Output half res.
@@ -699,6 +731,9 @@ void MainCS(
 		HoleFilledVelocityOutput[OutputPixelPos] = EncodedScreenVelocity;
 		HoleFilledVelocityMaskOutput[InputPixelPos] = tsr_half2(0.0, RejectHighFrequency);
 		ParallaxRejectionMaskOutput[OutputPixelPos] = ParallaxRejectionMask;
+		#if CONFIG_UPDATE_SUBPIXEL_DETAIL
+			HistorySubpixelDetailsOutput[OutputPixelPos] = FinalSubpixelDetails;
+		#endif
 
 		#if DEBUG_OUTPUT
 		{
diff --git a/Engine/Shaders/Private/TemporalSuperResolution/TSRUpdateHistory.usf b/Engine/Shaders/Private/TemporalSuperResolution/TSRUpdateHistory.usf
--- a/Engine/Shaders/Private/TemporalSuperResolution/TSRUpdateHistory.usf
+++ b/Engine/Shaders/Private/TemporalSuperResolution/TSRUpdateHistory.usf
@@ -41,11 +41,13 @@
 	#define CONFIG_HIGH_FREQUENCY_ONLY 1
 	#define CONFIG_HIGH_FREQUENCY_YCOCG 0
 	#define CONFIG_OUTPUT_HALFRES_HISTORY 1
+	#define CONFIG_UPDATE_SUBPIXEL_DETAIL 0
 
 #else
 	#define CONFIG_HIGH_FREQUENCY_ONLY 0
 	#define CONFIG_HIGH_FREQUENCY_YCOCG 1
 	#define CONFIG_OUTPUT_HALFRES_HISTORY 0
+	#define CONFIG_UPDATE_SUBPIXEL_DETAIL 1
 
 #endif
 
@@ -158,8 +160,9 @@ groupshared tsr_half4 SharedArray9[GROUP_SHARED_SIZE];
 groupshared tsr_half4 SharedArray10[GROUP_SHARED_SIZE];
 groupshared tsr_half4 SharedArray11[GROUP_SHARED_SIZE];
 
-groupshared tsr_subpixel_payload2 SharedParallaxFactorBits[GROUP_SHARED_SIZE];
-
+#if CONFIG_UPDATE_SUBPIXEL_DETAIL
+	groupshared tsr_subpixel_payload2 SharedParallaxFactorBits[GROUP_SHARED_SIZE];
+#endif
 
 //------------------------------------------------------- FUNCTIONS
 
@@ -290,7 +293,9 @@ void MainCS(
 	tsr_half2 OutputPixelVelocity = tsr_half(0.0).xx;
 	tsr_half2 NoiseFiltering = tsr_half(0.0).xx;
 	tsr_half2x2 HoleFilledVelocityMask = dpv_interleave_mono_registers(tsr_half(1.0).xx);
-	tsr_subpixel_payload2 ParallaxFactorBits = tsr_subpixel_payload(0x0).xx;
+	#if CONFIG_UPDATE_SUBPIXEL_DETAIL
+		tsr_subpixel_payload2 ParallaxFactorBits = tsr_subpixel_payload(0x0).xx;
+	#endif
 	bool2 bIsResponsiveAAPixel = false;
 	bool2 bIsOffScreen = false;
 	bool2 bIsDisoccluded = false;
@@ -441,7 +446,9 @@ void MainCS(
 				/* out */ bIsDisoccluded);
 
 			// Process input texel, forcing to pack register.
-			ParallaxFactorBits    = dpv_force_interleave_registers_array(RawParallaxFactorBits) & tsr_subpixel_payload(SUB_PIXEL_PARALLAX_FACTOR_BIT_MASK);
+			#if CONFIG_UPDATE_SUBPIXEL_DETAIL
+				ParallaxFactorBits    = dpv_force_interleave_registers_array(RawParallaxFactorBits) & tsr_subpixel_payload(SUB_PIXEL_PARALLAX_FACTOR_BIT_MASK);
+			#endif
 			LowFrequencyRejection = dpv_force_interleave_registers_array(RawLowFrequencyRejection);
 			SceneStencilRef       = dpv_force_interleave_registers_array(RawSceneStencilRef);
 			
@@ -520,8 +527,10 @@ void MainCS(
 	tsr_half2 ClosestInputLuma4;
 	
 	
-	bool2 bUpdateSubpixelDetails;
-	bool2x2 bKOSign;
+	#if CONFIG_UPDATE_SUBPIXEL_DETAIL
+		bool2 bUpdateSubpixelDetails;
+		bool2x2 bKOSign;
+	#endif
 
 	PLATFORM_SPECIFIC_ISOLATE
 	{
@@ -529,6 +538,7 @@ void MainCS(
 		tsr_half2x2 dKO = tsr_half2x2(PPCo - PPCk);
 
 		// Compresses the sign of dKO into SGPR to keep VGPR pressure low.
+		#if CONFIG_UPDATE_SUBPIXEL_DETAIL
 		{
 			float SizeOfOutputPixelInInputPixelWidth = 0.5 * (InputInfo_ViewportSize.x * HistoryInfo_ViewportSizeInverse.x);
 
@@ -538,6 +548,7 @@ void MainCS(
 			bKOSign[0] = dKO[0] < 0.0;
 			bKOSign[1] = dKO[1] < 0.0;
 		}
+		#endif
 		
 		tsr_halfC RawInputColorArray[CONFIG_SAMPLES][DPV_PIXEL_PER_LANE];
 		tsr_half4 RawTranslucencyColorArray[CONFIG_SAMPLES][DPV_PIXEL_PER_LANE];
@@ -679,7 +690,9 @@ void MainCS(
 				SharedArray11[LocalGroupThreadIndex] = tsr_half4(TranslucencyMaxColor[2], TranslucencyMaxColor[3]);
 			#endif
 
-			SharedParallaxFactorBits[LocalGroupThreadIndex] = ParallaxFactorBits;
+			#if CONFIG_UPDATE_SUBPIXEL_DETAIL
+				SharedParallaxFactorBits[LocalGroupThreadIndex] = ParallaxFactorBits;
+			#endif
 		#else
 			SharedArray0[LocalGroupThreadIndex] = tsr_half4(dpv_lo(FilteredInputColor), dpv_lo(LowFrequencyRejection));
 			SharedArray1[LocalGroupThreadIndex] = tsr_half4(dpv_lo(InputMinColor), dpv_lo(InputPixelAlignement));
@@ -705,7 +718,9 @@ void MainCS(
 
 	tsr_half4x2 PrevHistoryTranslucency;
 	
-	tsr_subpixel_details2 PrevSubpixelDetails;
+	#if CONFIG_UPDATE_SUBPIXEL_DETAIL
+		tsr_subpixel_details2 PrevSubpixelDetails;
+	#endif
 
 	PLATFORM_SPECIFIC_ISOLATE
 	{
@@ -819,6 +834,7 @@ void MainCS(
 				AccumulateTranslucency  = dpv_scale(AccumulateTranslucency, FinalMultiplier);
 			}
 
+			#if CONFIG_UPDATE_SUBPIXEL_DETAIL
 			FSubpixelNeighborhood SubpixelNeighborhood0 = GatherPrevSubpixelNeighborhood(PrevHistory_SubpixelDetails, dpv_lo(PrevHistoryBufferUV));
 			FSubpixelNeighborhood SubpixelNeighborhood1 = GatherPrevSubpixelNeighborhood(PrevHistory_SubpixelDetails, dpv_hi(PrevHistoryBufferUV));
 			{
@@ -832,6 +848,7 @@ void MainCS(
 					PrevSubpixelDetails[1] |= SubpixelPayload1 << tsr_subpixel_payload(SUB_PIXEL_BIT_COUNT * SubpixelId);
 				}
 			}
+			#endif // CONFIG_UPDATE_SUBPIXEL_DETAIL
 			
 			#if CONFIG_SCENE_COLOR_APLHA
 				AccumulateLowFrequency  = -dpv_min(-AccumulateLowFrequency, tsr_half(0.0).xxxx);
@@ -907,7 +924,9 @@ void MainCS(
 				tsr_half4 RawLDS11 = SharedArray11[LocalGroupThreadIndex];
 			#endif
 
-			ParallaxFactorBits = SharedParallaxFactorBits[LocalGroupThreadIndex];
+			#if CONFIG_UPDATE_SUBPIXEL_DETAIL
+				ParallaxFactorBits = SharedParallaxFactorBits[LocalGroupThreadIndex];
+			#endif
 
 			FilteredInputColor = tsr_half3x2(RawLDS0.rg, RawLDS0.ba, RawLDS1.rg);
 			InputMinColor = tsr_half3x2(RawLDS2.rg, RawLDS2.ba, RawLDS3.rg);
@@ -1415,6 +1434,7 @@ void MainCS(
 	#endif
 
 	// Update subpixel detail
+	#if CONFIG_UPDATE_SUBPIXEL_DETAIL
 	tsr_subpixel_details2 FinalSubpixelDetails;
 	{
 		FinalSubpixelDetails = PrevSubpixelDetails;
@@ -1452,6 +1472,7 @@ void MainCS(
 			FinalSubpixelDetails.y = UpdatedFinalSubpixelDetails.y;
 		}
 	}
+	#endif // CONFIG_UPDATE_SUBPIXEL_DETAIL
 
 	// Compute final output
 	tsr_halfCx2 FinalOutputColor;
@@ -1600,7 +1621,9 @@ void MainCS(
 					HistoryOutput_Translucency[dpv_lo(LocalHistoryPixelPos)] = dpv_lo(FinalTranslucencyColor).rgb;
 					HistoryOutput_TranslucencyAlpha[dpv_lo(LocalHistoryPixelPos)] = dpv_lo(FinalTranslucencyColor).a;
 				#endif
-				HistoryOutput_SubpixelDetails[dpv_lo(LocalHistoryPixelPos)] = dpv_lo(FinalSubpixelDetails);
+				#if CONFIG_UPDATE_SUBPIXEL_DETAIL
+					HistoryOutput_SubpixelDetails[dpv_lo(LocalHistoryPixelPos)] = dpv_lo(FinalSubpixelDetails);
+				#endif
 			}
 		
 			// Output final history hi pixel.
@@ -1618,7 +1641,9 @@ void MainCS(
 					HistoryOutput_Translucency[dpv_hi(LocalHistoryPixelPos)] = dpv_hi(FinalTranslucencyColor).rgb;
 					HistoryOutput_TranslucencyAlpha[dpv_hi(LocalHistoryPixelPos)] = dpv_hi(FinalTranslucencyColor).a;
 				#endif
-				HistoryOutput_SubpixelDetails[dpv_hi(LocalHistoryPixelPos)] = dpv_hi(FinalSubpixelDetails);
+				#if CONFIG_UPDATE_SUBPIXEL_DETAIL
+					HistoryOutput_SubpixelDetails[dpv_hi(LocalHistoryPixelPos)] = dpv_hi(FinalSubpixelDetails);
+				#endif
 			}
 			#endif // !CONFIG_ENABLE_DUAL_PIXEL_VECTORIZATION
 		}
diff --git a/Engine/Source/Runtime/Renderer/Private/PostProcess/TemporalSuperResolution.cpp b/Engine/Source/Runtime/Renderer/Private/PostProcess/TemporalSuperResolution.cpp
--- a/Engine/Source/Runtime/Renderer/Private/PostProcess/TemporalSuperResolution.cpp
+++ b/Engine/Source/Runtime/Renderer/Private/PostProcess/TemporalSuperResolution.cpp
@@ -135,6 +135,8 @@ END_SHADER_PARAMETER_STRUCT()
 BEGIN_SHADER_PARAMETER_STRUCT(FTSRPrevHistoryParameters, )
 	SHADER_PARAMETER_STRUCT(FScreenPassTextureViewportParameters, PrevHistoryInfo)
 	SHADER_PARAMETER(FScreenTransform, ScreenPosToPrevHistoryBufferUV)
+	SHADER_PARAMETER(FScreenTransform, ScreenPosToPrevSubpixelDetails)
+	SHADER_PARAMETER(FVector2f, PrevSubpixelDetailsExtent)
 	SHADER_PARAMETER(float, HistoryPreExposureCorrection)
 END_SHADER_PARAMETER_STRUCT()
 
@@ -262,6 +264,7 @@ class FTSRDecimateHistoryCS : public FTSRShader
 		SHADER_PARAMETER_RDG_TEXTURE_UAV(RWTexture2D, HoleFilledVelocityOutput)
 		SHADER_PARAMETER_RDG_TEXTURE_UAV(RWTexture2D, HoleFilledVelocityMaskOutput)
 		SHADER_PARAMETER_RDG_TEXTURE_UAV(RWTexture2D, ParallaxRejectionMaskOutput)
+		SHADER_PARAMETER_RDG_TEXTURE_UAV(RWTexture2D, HistorySubpixelDetailsOutput)
 
 		SHADER_PARAMETER_RDG_TEXTURE_UAV(RWTexture2D, DebugOutput)
 	END_SHADER_PARAMETER_STRUCT()
@@ -594,6 +597,9 @@ ITemporalUpscaler::FOutputs AddTemporalSuperResolutionPasses(
 	// Whether the UpdateHistory pass also output half res history to reduce memory bottleneck up the FTSRDecimateHistoryCS.
 	const bool bHalfResHistory = bHistoryHighFrequencyOnly;
 
+	// Whether the subpixel details are stored at input or output resolution
+	const bool bSubPixelDetailsAtOutputResolution = !bHistoryHighFrequencyOnly;
+
 	EPixelFormat ColorFormat = bSupportsAlpha ? PF_FloatRGBA : PF_FloatR11G11B10;
 
 	int32 RejectionAntiAliasingQuality = FMath::Clamp(CVarTSRRejectionAntiAliasingQuality.GetValueOnRenderThread(), 1, 2);
@@ -897,17 +903,24 @@ ITemporalUpscaler::FOutputs AddTemporalSuperResolutionPasses(
 	// Setup the shader parameters for previous frame history
 	FTSRPrevHistoryParameters PrevHistoryParameters;
 	{
-
 		// Setup prev history parameters.
 		FScreenPassTextureViewport PrevHistoryViewport(PrevHistory.HighFrequency->Desc.Extent, InputHistory.OutputViewportRect);
+		FScreenPassTextureViewport PrevSubpixelDetailsViewport(PrevHistory.SubpixelDetails->Desc.Extent, bSubPixelDetailsAtOutputResolution ? InputHistory.OutputViewportRect : InputHistory.InputViewportRect);
+
 		if (bCameraCut)
 		{
 			PrevHistoryViewport.Extent = FIntPoint(1, 1);
 			PrevHistoryViewport.Rect = FIntRect(FIntPoint(0, 0), FIntPoint(1, 1));
+			PrevSubpixelDetailsViewport.Extent = FIntPoint(1, 1);
+			PrevSubpixelDetailsViewport.Rect = FIntRect(FIntPoint(0, 0), FIntPoint(1, 1));
 		}
+
 		PrevHistoryParameters.PrevHistoryInfo = GetScreenPassTextureViewportParameters(PrevHistoryViewport);
 		PrevHistoryParameters.ScreenPosToPrevHistoryBufferUV = FScreenTransform::ChangeTextureBasisFromTo(
 			PrevHistoryViewport, FScreenTransform::ETextureBasis::ScreenPosition, FScreenTransform::ETextureBasis::TextureUV);
+		PrevHistoryParameters.ScreenPosToPrevSubpixelDetails = FScreenTransform::ChangeTextureBasisFromTo(
+			PrevSubpixelDetailsViewport, FScreenTransform::ETextureBasis::ScreenPosition, FScreenTransform::ETextureBasis::TextureUV);
+		PrevHistoryParameters.PrevSubpixelDetailsExtent = PrevSubpixelDetailsViewport.Extent;
 		PrevHistoryParameters.HistoryPreExposureCorrection = View.PreExposure / View.PrevViewInfo.SceneColorPreExposure;
 	}
 
@@ -943,8 +956,20 @@ ITemporalUpscaler::FOutputs AddTemporalSuperResolutionPasses(
 		Desc.Format = PF_R8;
 		History.TranslucencyAlpha = GraphBuilder.CreateTexture(Desc, TEXT("TSR.History.TranslucencyAlpha"));
 
-		Desc.Format = PF_R16_UINT;
-		History.SubpixelDetails = GraphBuilder.CreateTexture(Desc, TEXT("TSR.History.SubpixelInfo"));
+		if (bSubPixelDetailsAtOutputResolution)
+		{
+			Desc.Format = PF_R16_UINT;
+			History.SubpixelDetails = GraphBuilder.CreateTexture(Desc, TEXT("TSR.History.SubpixelInfo"));
+		}
+		else
+		{
+			FRDGTextureDesc SubpixelDetailsDesc = FRDGTextureDesc::Create2D(
+				InputExtent,
+				PF_R16_UINT,
+				FClearValueBinding::None,
+				TexCreate_ShaderResource | TexCreate_UAV);
+			History.SubpixelDetails = GraphBuilder.CreateTexture(SubpixelDetailsDesc, TEXT("TSR.History.SubpixelInfo"));
+		}
 	}
 
 	// Decimate input to flicker at same frequency as input.
@@ -1035,6 +1060,7 @@ ITemporalUpscaler::FOutputs AddTemporalSuperResolutionPasses(
 		PassParameters->HoleFilledVelocityOutput = GraphBuilder.CreateUAV(HoleFilledVelocityTexture);
 		PassParameters->HoleFilledVelocityMaskOutput = GraphBuilder.CreateUAV(HoleFilledVelocityMaskTexture);
 		PassParameters->ParallaxRejectionMaskOutput = GraphBuilder.CreateUAV(ParallaxRejectionMaskTexture);
+		PassParameters->HistorySubpixelDetailsOutput = GraphBuilder.CreateUAV(History.SubpixelDetails);
 		PassParameters->DebugOutput = CreateDebugUAV(LowFrequencyExtent, TEXT("Debug.TSR.DecimateHistory"));
 
 		FTSRDecimateHistoryCS::FPermutationDomain PermutationVector;
@@ -1501,6 +1527,7 @@ ITemporalUpscaler::FOutputs AddTemporalSuperResolutionPasses(
 	if (!View.bStatePrevViewInfoIsReadOnly)
 	{
 		FTSRHistory& OutputHistory = View.ViewState->PrevFrameViewInfo.TSRHistory;
+		OutputHistory.InputViewportRect = InputRect;
 		OutputHistory.OutputViewportRect = FIntRect(FIntPoint(0, 0), HistorySize);
 
 		// Extract filterable history
diff --git a/Engine/Source/Runtime/Renderer/Private/SceneRendering.h b/Engine/Source/Runtime/Renderer/Private/SceneRendering.h
--- a/Engine/Source/Runtime/Renderer/Private/SceneRendering.h
+++ b/Engine/Source/Runtime/Renderer/Private/SceneRendering.h
@@ -894,6 +894,7 @@ struct FTSRHistory
 	// Filterable half-res output resolution
 	TRefCountPtr<IPooledRenderTarget> HighFrequencyHalfRes;
 
+	FIntRect InputViewportRect;
 	FIntRect OutputViewportRect;
 
 
