Fix the overflow of lighting on mobile caused by half precision.

diff --git a/Engine/Shaders/Private/MobileShadingModels.ush b/Engine/Shaders/Private/MobileShadingModels.ush
--- a/Engine/Shaders/Private/MobileShadingModels.ush
+++ b/Engine/Shaders/Private/MobileShadingModels.ush
@@ -239,7 +239,8 @@ half Vis_SmithJointApprox_Mobile(half a, half NoV, half NoL)
 {
 	float Vis_SmithV = NoL * (NoV * (1 - a) + a);
 	float Vis_SmithL = NoV * (NoL * (1 - a) + a);
-	return 0.5f * rcp(max(Vis_SmithV + Vis_SmithL, 0.0001f));
+	// clamp the visibility term to 0~1
+	return saturate(0.5f * rcp(max(Vis_SmithV + Vis_SmithL, 0.0001f)));
 }
 
 // [Schlick 1994, "An Inexpensive BRDF Model for Physically-Based Rendering"]
diff --git a/Engine/Shaders/Private/SingleLayerWaterShading.ush b/Engine/Shaders/Private/SingleLayerWaterShading.ush
--- a/Engine/Shaders/Private/SingleLayerWaterShading.ush
+++ b/Engine/Shaders/Private/SingleLayerWaterShading.ush
@@ -164,8 +164,10 @@ WaterVolumeLightingOutput EvaluateWaterVolumeLighting(
 
 		const float3 OpticalDepth = ExtinctionCoeff * BehindWaterDeltaDepth;
 		float3 Transmittance = exp(-OpticalDepth);
-		float3 ScatteredLuminance = ScatteringCoeff * (AmbScattLuminance + SunScattLuminance * DirectionalLightShadow);
-		ScatteredLuminance = (ScatteredLuminance - ScatteredLuminance * Transmittance) / ExtinctionCoeffSafe;
+		float3 IncomingLuminance = (AmbScattLuminance + SunScattLuminance * DirectionalLightShadow);
+		// The SafeScatteringAmount should be in 0~1, add saturate to fix an artifact caused by half precision on mobile.
+		float3 SafeScatteringAmount = saturate(ScatteringCoeff * (1.0f - Transmittance) / ExtinctionCoeffSafe);
+		float3 ScatteredLuminance = IncomingLuminance * SafeScatteringAmount;
 
 #else
 		// TODO Make the volumetric shadow part work again
