Strata - rough refraction - Added plot of gaussian lobe.

diff --git a/Engine/Shaders/Private/Strata/StrataRoughRefraction.usf b/Engine/Shaders/Private/Strata/StrataRoughRefraction.usf
--- a/Engine/Shaders/Private/Strata/StrataRoughRefraction.usf
+++ b/Engine/Shaders/Private/Strata/StrataRoughRefraction.usf
@@ -336,6 +336,22 @@ FStrataLobeStatistic StrataGetRefractedLobe2(FStrataLobeStatistic WiLobe, float3
 	return WoLobe;
 }
 
+// Gaussian function with a=1/(c*sqrt(2*PI)) so it represents the probability density function of a normally distributed random variable with expected value mu = b for some variance and integrale = 1;
+float GaussianFunction(float x, float Mean, float Variance)
+{
+	float StandardDeviation = sqrt(Variance);
+	float FactorA = 1 / (StandardDeviation * sqrt(2 * PI));
+	float Bell = exp(-((x - Mean) * (x - Mean)) / (2 * Variance));
+	return FactorA * Bell;
+}
+float GaussianFunction(float2 P, float2 Mean, float Variance)
+{
+	float StandardDeviation = sqrt(Variance);
+	float FactorA = 1 / (StandardDeviation * sqrt(2 * PI));
+	float2 Bell = exp(-((P.x - Mean.x) * (P.x - Mean.x)) / (2 * Variance) - ((P.y - Mean.y) * (P.y - Mean.y)) / (2 * Variance));
+	return FactorA * Bell;
+}
+
 void VisualizeRoughRefractionPS(
 	float4 SVPos : SV_POSITION,
 	out float4 OutColor : SV_Target0)
@@ -347,6 +363,9 @@ void VisualizeRoughRefractionPS(
 	float3 WorldDir = GetScreenWorldDir(SVPos);
 	float3 CamWorldPos = LWCToFloat(PrimaryView.WorldCameraOrigin);
 
+	FLobeStat LobeState = LobStatisticsBuffer[0];
+
+	const float CountScale = 0.01f;
 	
 	FShaderPrintContext Context = GetUIContext(all(uint2(PixelPos.xy) == uint2(10, 10)));
 	FUI UI = GetUI(Context);
@@ -382,17 +401,25 @@ void VisualizeRoughRefractionPS(
 			const int2 Coord = (UV * 64.0f + 0.5f);
 			if (all(Coord >= 0) && all(Coord < 64))
 			{
-				const float SampleCount = float(SampleCountTexture[Coord].x);
-				OutColor = float4(GetHSVDebugColor(SampleCount*0.01f), 0.0);
+				if (Coord.y < 32)
+				{
+					const float SampleCount = float(SampleCountTexture[Coord].x) * CountScale;
+					OutColor = float4(GetHSVDebugColor(SampleCount), 0.0);
+				}
+				else
+				{
+					float Weight = GaussianFunction((P.xy - RefractionP.xy) / VisProjectedRayRadius, LobeState.Mean.xy, LobeState.Variance.x);
+					OutColor = float4(GetHSVDebugColor(Weight), 0.0);
+				}
 			}
 		}
 	}
 
-#if 0
+#if 1
 	// Draw projected distribution as a curve
 	float3 PrevP = -1.0f;
 	float PrevValue = -1.0f;
-	for (float u = 0.0; u < 1.0; u += 0.01)
+	for (float u = 0.0; u < 1.0; u += 0.005f)
 	{
 		float3 P = lerp(SlabP0, SlabP1, u);
 		P.y = 0.5 * (SlabP0.y + SlabP1.y);
@@ -402,8 +429,8 @@ void VisualizeRoughRefractionPS(
 		const int2 Coord = (UV * 64.0f + 0.5f);
 		if (all(Coord >= 0) && all(Coord < 64))
 		{
-			const float SampleCount = float(SampleCountTexture[Coord].x);
-			const float Value = float(SampleCount) * 0.05;
+			const float SampleCount = float(SampleCountTexture[Coord].x) * CountScale;
+			const float Value = float(SampleCount);
 
 			if (PrevValue != -1.0f)
 			{
@@ -411,13 +438,36 @@ void VisualizeRoughRefractionPS(
 					PrevP +  float3(0, 0, PrevValue), 
 					P +  float3(0, 0, Value),
 					float4(1, 1, 0, 0.5));
-				//float4(GetHSVDebugColor(SampleCount * 0.01f), 0.5));
 			}
 
 			PrevP = P;
 			PrevValue = Value;
 		}
 	}
+
+	// Draw gaussian curve mapping to recovered variance
+	PrevP = -1.0f;
+	PrevValue = -1.0f;
+	for (float u = 0.0; u < 1.0; u += 0.005f)
+	{
+		float3 P = lerp(SlabP0, SlabP1, u);
+		P.y = 0.5 * (SlabP0.y + SlabP1.y);
+		P.z = SlabBottomZ;
+
+		float Weight = GaussianFunction((P.x - RefractionP.x) / VisProjectedRayRadius, LobeState.Mean.x, LobeState.Variance.x);
+		const float Value = Weight;
+
+		if (PrevValue != -1.0f)
+		{
+			AddLineWS(Context,
+				PrevP + float3(0, 0, PrevValue),
+				P + float3(0, 0, Value),
+				float4(0, 1, 0, 0.5));
+		}
+
+		PrevP = P;
+		PrevValue = Value;
+	}
 #endif
 
 #if 0
@@ -485,7 +535,6 @@ void VisualizeRoughRefractionPS(
 #endif
 
 	// Our mean / variance
-	FLobeStat LobeState = LobStatisticsBuffer[0];
 	DrawCircleXY(Context, float3(SlabCenter.xy + LobeState.Mean * VisProjectedRayRadius, SlabBottomZ + 0.05), LobeState.Variance * VisProjectedRayRadius, float4(1, 1, 0, 1));
 
 	Newline(Context);
