Fix hair strands shadow with RT shadow.

RT shadow are now done in translated world space, thus this CL remove the temp. hack for computing transversal in world space.

diff --git a/Engine/Shaders/Private/HairStrands/HairStrandsRaytracing.ush b/Engine/Shaders/Private/HairStrands/HairStrandsRaytracing.ush
--- a/Engine/Shaders/Private/HairStrands/HairStrandsRaytracing.ush
+++ b/Engine/Shaders/Private/HairStrands/HairStrandsRaytracing.ush
@@ -12,13 +12,11 @@
 #include "../HairStrands/HairStrandsVoxelPageTraversal.ush"
 
 
+// RayOrigin is in translated world space
 float InternalTraverseHair(uint2 PixelCoord, float3 RayOrigin, float3 RayDirection, float InMinT, float InOcclusionThreshold, float3 VoxelRandom, float RayRandom)
 {
 	float HitT = InMinT;
 
-	// Note: VirtualVoxel.TranslatedWorldOffset should be equivalent to LWCHackToFloat(PrimaryView.WorldCameraOrigin), but insures the value is consistent
-	RayOrigin = RayOrigin + VirtualVoxel.TranslatedWorldOffset; // HAIR_LWC
-
 	// Depth bias
 	// Origin is shifted voxels away towards the light + a constant bias of the size of the voxel
 	const float3 NormalizedDepthBias = (RayRandom * RayDirection + VoxelRandom);
