Fix missing RT scale on hair strands geometry.

diff --git a/Engine/Shaders/Private/HairStrands/HairStrandsVertexFactory.ush b/Engine/Shaders/Private/HairStrands/HairStrandsVertexFactory.ush
--- a/Engine/Shaders/Private/HairStrands/HairStrandsVertexFactory.ush
+++ b/Engine/Shaders/Private/HairStrands/HairStrandsVertexFactory.ush
@@ -401,17 +401,18 @@ RAY_TRACING_ENTRY_INTERSECTION(MaterialIS)
 		VertexIndex1 = HairStrandsVF.CulledVertexIdsBuffer[VertexIndex1];
 	}
 
+	const float EffectiveRadius = HairStrandsVF.Radius * HairStrandsVF.RaytracingRadiusScale;
 	const float3 PositionOffset = HairStrandsVF.PositionOffsetBuffer[0].xyz;
 	FHairControlPoint P0 = UnpackHairControlPoint(
 		HairStrandsVF.PositionBuffer[VertexIndex0],
 		PositionOffset,
-		HairStrandsVF.Radius * ScaleR0,
+		EffectiveRadius * ScaleR0,
 		HairStrandsVF.RootScale,
 		HairStrandsVF.TipScale);
 	FHairControlPoint P1 = UnpackHairControlPoint(
 		HairStrandsVF.PositionBuffer[VertexIndex1],
 		PositionOffset,
-		HairStrandsVF.Radius * ScaleR1,
+		EffectiveRadius * ScaleR1,
 		HairStrandsVF.RootScale,
 		HairStrandsVF.TipScale);
 
@@ -481,19 +482,20 @@ FMaterialPixelParameters GetMaterialPixelParameters(float3 RayOrigin, float3 Ray
 		VertexIndex1 = HairStrandsVF.CulledVertexIdsBuffer[VertexIndex1];
 	}
 
+	const float EffectiveRadius = HairStrandsVF.Radius * HairStrandsVF.RaytracingRadiusScale;
 	float UCoord = HitAttributes.Barycentrics.x;
 	float VCoord = HitAttributes.Barycentrics.y;
 	const float3 PositionOffset = HairStrandsVF.PositionOffsetBuffer[0].xyz;
 	FHairControlPoint P0 = UnpackHairControlPoint(
 		HairStrandsVF.PositionBuffer[VertexIndex0],
 		PositionOffset,
-		HairStrandsVF.Radius,
+		EffectiveRadius,
 		HairStrandsVF.RootScale,
 		HairStrandsVF.TipScale);
 	FHairControlPoint P1 = UnpackHairControlPoint(
 		HairStrandsVF.PositionBuffer[VertexIndex1],
 		PositionOffset,
-		HairStrandsVF.Radius,
+		EffectiveRadius,
 		HairStrandsVF.RootScale,
 		HairStrandsVF.TipScale);
 
@@ -539,13 +541,13 @@ FMaterialPixelParameters GetMaterialPixelParameters(float3 RayOrigin, float3 Ray
 		FHairControlPoint P0 = UnpackHairControlPoint(
 			HairStrandsVF.PositionBuffer[VertexIndex0],
 			PositionOffset,
-			HairStrandsVF.Radius,
+			EffectiveRadius,
 			HairStrandsVF.RootScale,
 			HairStrandsVF.TipScale);
 		FHairControlPoint P1 = UnpackHairControlPoint(
 			HairStrandsVF.PositionBuffer[VertexIndex1],
 			PositionOffset,
-			HairStrandsVF.Radius,
+			EffectiveRadius,
 			HairStrandsVF.RootScale,
 			HairStrandsVF.TipScale);
 		float4 A = float4(P0.Position, P0.WorldRadius);
