Caches TSR's FilterFrequency pass' input into LDS

diff --git a/Engine/Shaders/Private/TemporalSuperResolution/TSRFilterFrequencies.usf b/Engine/Shaders/Private/TemporalSuperResolution/TSRFilterFrequencies.usf
--- a/Engine/Shaders/Private/TemporalSuperResolution/TSRFilterFrequencies.usf
+++ b/Engine/Shaders/Private/TemporalSuperResolution/TSRFilterFrequencies.usf
@@ -13,11 +13,26 @@
 
 //------------------------------------------------------- CONFIG
 
+#if CONFIG_SCENE_COLOR_APLHA
+	#define CONFIG_CACHE_LDS 0
+#else
+	#define CONFIG_CACHE_LDS 1
+#endif
+
+#if CONFIG_CACHE_LDS
+	#define SUBTILE_SIZE 8
+	#define SUBTILE_COUNT 2
+	#define SUBTILE_PIXEL_COUNT (SUBTILE_SIZE * SUBTILE_SIZE)
+
+	#define LDS_PADDING 4
+#endif
+
 #define TILE_SIZE 16
 
 #define CONFIG_ACCUMULATE_KARIS_WEIGHTING 1
 
 #define CONFIG_KERNEL (KERNEL_3X3)
+#define CONFIG_KERNEL_SAMPLE_COUNT 9
 
 // Use a spread factor the size of the clamping box in the history comparison to avoid one flickering input pixel
 // to flicker the entire comparison neighborhood
@@ -41,6 +56,17 @@ RWTexture2D<tsr_halfC> FilteredPredictionSceneColorOutput;
 RWTexture2D<tsr_half>  InputSceneColorLdrLumaOutput;
 
 
+//------------------------------------------------------- LDS
+
+#if CONFIG_CACHE_LDS
+	#define LDS_TILE_SIZE (TILE_SIZE + 2 * LDS_PADDING)
+	#define LDS_SIZE (LDS_TILE_SIZE * LDS_TILE_SIZE)
+
+	groupshared tsr_half4 SharedArray0[LDS_SIZE];
+	groupshared tsr_half4 SharedArray1[LDS_SIZE];
+#endif
+
+
 //------------------------------------------------------- FUNCTIONS
 
 tsr_half ComputeColorWeight(tsr_halfC Color)
@@ -59,6 +85,29 @@ tsr_half ParallaxRejectionMaskToHistoryValidity(tsr_half ParallaxRejectionMask)
 	return (ParallaxRejectionMask < tsr_half(PARALLAX_REJECTION_MASK_THRESHOLD)) ? tsr_half(0.0) : tsr_half(1.0);
 }
 
+#if CONFIG_CACHE_LDS
+void ProcessAndStoreToLDS(tsr_short2 SubTileFetchCoordOffset, tsr_short2 SubTilePixelOffset, tsr_halfC Color, tsr_halfC History, tsr_half ParallaxRejectionMask)
+{
+	#if CONFIG_SCENE_COLOR_OVERFLOW
+		#if CONFIG_SCENE_COLOR_APLHA
+			Color = min(Color, LargestSceneColorRGBA);
+		#else
+			Color = min(Color, LargestSceneColorRGB);
+		#endif
+	#endif
+		
+	tsr_half HistorySampleValidity = ParallaxRejectionMaskToHistoryValidity(ParallaxRejectionMask);
+	tsr_half InputWeight = HistorySampleValidity * ComputeColorWeight(Color);
+	tsr_half HistoryWeight = HistorySampleValidity * ComputeColorWeight(History);
+
+	tsr_short Store = dot(SubTileFetchCoordOffset, tsr_short2(1, LDS_TILE_SIZE)) + dot(SubTilePixelOffset, tsr_short2(1, LDS_TILE_SIZE));
+
+	SharedArray0[Store] = tsr_half4(Color, InputWeight);
+	SharedArray1[Store] = tsr_half4(History, HistoryWeight);
+} 
+#endif // CONFIG_CACHE_LDS
+
+
 //------------------------------------------------------- ENTRY POINT
 
 [numthreads(TILE_SIZE, TILE_SIZE, 1)]
@@ -70,138 +119,178 @@ void MainCS(
 
 	float4 Debug = 0.0;
 	
-	tsr_short2 DispatchThreadId = (
-		tsr_short2(GroupId) * tsr_short2(TILE_SIZE, TILE_SIZE) +
-		Map16x16Tile2x2Lane(GroupThreadIndex));
+	#if CONFIG_CACHE_LDS
+		tsr_short2 SubTilePixelOffset = Map8x8Tile2x2Lane(GroupThreadIndex % SUBTILE_PIXEL_COUNT);
+		tsr_short2 SubTileOffset = tsr_short2(tsr_ushort2(tsr_ushort(GroupThreadIndex) / tsr_ushort(SUBTILE_PIXEL_COUNT), tsr_ushort(GroupThreadIndex) / tsr_ushort(SUBTILE_PIXEL_COUNT * SUBTILE_COUNT)) % tsr_ushort(SUBTILE_COUNT));
+
+		tsr_short2 TilePixelCoord = tsr_short2(InputPixelPosMin) + tsr_short2(GroupId) * tsr_short(TILE_SIZE);
+		tsr_short2 TilePixelOffset = SubTilePixelOffset + tsr_short(SUBTILE_SIZE) * SubTileOffset;
+	
+		// Load LDS
+		{
+			tsr_short2 FetchOrigin = TilePixelCoord - tsr_short(LDS_PADDING);
+
+			tsr_short2 SubTileFetchCoordOffset0 = tsr_short(SUBTILE_SIZE) * SubTileOffset;
+			tsr_short2 SubTileFetchCoordOffset1 = tsr_short(SUBTILE_SIZE) * select(SubTileOffset.y == 0, tsr_short2(2, SubTileOffset.x), tsr_short2(SubTileOffset.x, 2));
+			tsr_short2 SubTileFetchCoordOffset2 = tsr_short(SUBTILE_SIZE) * tsr_short(2);
 		
-	tsr_short2 InputPixelPos = tsr_short2(InputPixelPosMin) + DispatchThreadId;
+			tsr_short2 ClampedFetchCoord0 = ClampPixelOffset((FetchOrigin + SubTileFetchCoordOffset0) + SubTilePixelOffset, InputInfo_ViewportMin, InputInfo_ViewportMax);
+			tsr_short2 ClampedFetchCoord1 = ClampPixelOffset((FetchOrigin + SubTileFetchCoordOffset1) + SubTilePixelOffset, InputInfo_ViewportMin, InputInfo_ViewportMax);
+			tsr_short2 ClampedFetchCoord2 = ClampPixelOffset((FetchOrigin + SubTileFetchCoordOffset2) + SubTilePixelOffset, InputInfo_ViewportMin, InputInfo_ViewportMax);
 
-	// Fetch center of the kernel
-	tsr_halfC CenterInput;
-	tsr_halfC FilteredInput;
-	tsr_half FilteredInputWeight;
-	tsr_halfC FilteredHistory;
-	tsr_half FilteredHistoryWeight;
+			tsr_halfC Color0 = InputTexture[ClampedFetchCoord0];
+			tsr_halfC Color1 = InputTexture[ClampedFetchCoord1];
+			tsr_halfC Color2 = InputTexture[ClampedFetchCoord2];
 
-	{
-		tsr_half SampleSpatialWeight = 1.0;
+			tsr_halfC History0 = PredictionSceneColorTexture[ClampedFetchCoord0];
+			tsr_halfC History1 = PredictionSceneColorTexture[ClampedFetchCoord1];
+			tsr_halfC History2 = PredictionSceneColorTexture[ClampedFetchCoord2];
 
-		tsr_halfC InputSample = InputTexture[InputPixelPos];
-		tsr_halfC HistorySample = PredictionSceneColorTexture[InputPixelPos];
-		tsr_half HistorySampleValidity = ParallaxRejectionMaskToHistoryValidity(
-			ParallaxRejectionMaskTexture[InputPixelPos]);
+			tsr_half ParallaxRejectionMask0 = ParallaxRejectionMaskTexture[ClampedFetchCoord0];
+			tsr_half ParallaxRejectionMask1 = ParallaxRejectionMaskTexture[ClampedFetchCoord1];
+			tsr_half ParallaxRejectionMask2 = ParallaxRejectionMaskTexture[ClampedFetchCoord2];
 		
-		tsr_half InputWeight = ComputeColorWeight(InputSample);
-		tsr_half HistoryWeight = HistorySampleValidity * ComputeColorWeight(HistorySample);
+			ProcessAndStoreToLDS(SubTileFetchCoordOffset0, SubTilePixelOffset, Color0, History0, ParallaxRejectionMask0);
+			ProcessAndStoreToLDS(SubTileFetchCoordOffset1, SubTilePixelOffset, Color1, History1, ParallaxRejectionMask1);
+			ProcessAndStoreToLDS(SubTileFetchCoordOffset2, SubTilePixelOffset, Color2, History2, ParallaxRejectionMask2);
+		}
 
-		CenterInput = InputSample;
+		GroupMemoryBarrierWithGroupSync();
+		
+		tsr_short2 InputPixelPos = TilePixelCoord + TilePixelOffset;
+		
+		{
+			tsr_short LoadPos = dot(TilePixelOffset + tsr_short(LDS_PADDING), tsr_short2(1, LDS_TILE_SIZE));
 
-		FilteredInput       = (SampleSpatialWeight * InputWeight) * InputSample;
-		FilteredInputWeight = (SampleSpatialWeight * InputWeight);
+			Debug.rgb = SharedArray0[LoadPos].rgb;
+		}
+
+	#else // !CONFIG_CACHE_LDS
+		tsr_short2 DispatchThreadId = (
+			tsr_short2(GroupId) * tsr_short2(TILE_SIZE, TILE_SIZE) +
+			Map16x16Tile2x2Lane(GroupThreadIndex));
 		
-		FilteredHistory       = (SampleSpatialWeight * HistoryWeight) * HistorySample;
-		FilteredHistoryWeight = (SampleSpatialWeight * HistoryWeight);
-	}
-	
-	uint StartSampleId;
-	uint SampleCount;
-	uint SampleTrackId;
-	
-	#if CONFIG_KERNEL == KERNEL_1X1
-	{
-		StartSampleId = 0;
-		SampleCount = 0;
-	}
-	#elif CONFIG_KERNEL == KERNEL_3X3 || CONFIG_KERNEL == KERNEL_3X3_PLUS
-	{
-		StartSampleId = 0;
-		SampleCount = 9;
-	}
-	#else
-		#error Unknown kernel
-	#endif
+		tsr_short2 InputPixelPos = tsr_short2(InputPixelPosMin) + DispatchThreadId;
 
-	#if CONFIG_USE_DYNAMIC_SPREAD_FACTOR
-		tsr_short2 DynamicSpreadFactor = (
-			tsr_short(CONFIG_KERNEL_SPREAD_FACTOR) +
-			tsr_short2((tsr_ushort2(InputPixelPos) % tsr_ushort(2)) ^ (tsr_ushort2(InputPixelPos.yx) % tsr_ushort(2))));
-	#else
-		tsr_short2 DynamicSpreadFactor = tsr_short(CONFIG_KERNEL_SPREAD_FACTOR);
 	#endif
 
-	// Fetch neighborhood of the kernel
-	UNROLL
-	for (uint SampleId = StartSampleId; SampleId < SampleCount; SampleId++)
+
+	// Process the samples
+	tsr_halfC CenterInput;
+	tsr_halfC FilteredInput;
+	tsr_halfC FilteredHistory;
 	{
-		tsr_short2 Offset;
-		tsr_half SampleSpatialWeight;
+		FilteredInput = tsr_half(0.0);
+		FilteredHistory = tsr_half(0.0);
+
+		tsr_half FilteredInputWeight = tsr_half(0.0);
+		tsr_half FilteredHistoryWeight = tsr_half(0.0);
 		
-		#if CONFIG_KERNEL == KERNEL_1X1
-		{
-			Offset = 0;
-			SampleSpatialWeight = tsr_half(1.0);
-			continue;
-		}
-		#elif CONFIG_KERNEL == KERNEL_3X3 || CONFIG_KERNEL == KERNEL_3X3_PLUS
+		#if CONFIG_USE_DYNAMIC_SPREAD_FACTOR
+			tsr_short2 DynamicSpreadFactor = (
+				tsr_short(CONFIG_KERNEL_SPREAD_FACTOR) +
+				tsr_short2((tsr_ushort2(InputPixelPos) % tsr_ushort(2)) ^ (tsr_ushort2(InputPixelPos.yx) % tsr_ushort(2))));
+		#else
+			tsr_short2 DynamicSpreadFactor = tsr_short(CONFIG_KERNEL_SPREAD_FACTOR);
+		#endif
+
+		UNROLL_N(CONFIG_KERNEL_SAMPLE_COUNT)
+		for (uint SampleId = 0; SampleId < CONFIG_KERNEL_SAMPLE_COUNT; SampleId++)
 		{
-			if (SampleId == 4)
+			tsr_short2 Offset;
+			tsr_half SampleSpatialWeight;
+		
+			#if CONFIG_KERNEL == KERNEL_1X1
 			{
-				continue;
+				Offset = 0;
+				SampleSpatialWeight = tsr_half(1.0);
 			}
-			Offset = tsr_short2(kOffsets3x3[SampleId]);
+			#elif CONFIG_KERNEL == KERNEL_3X3 || CONFIG_KERNEL == KERNEL_3X3_PLUS
+			{
+				Offset = tsr_short2(kOffsets3x3[SampleId]);
 			
-			#if CONFIG_KERNEL == KERNEL_3X3_PLUS
-				if (abs(Offset.x) + abs(Offset.y) > 1)
-				{
-					continue;
-				}
-			#endif
+				#if CONFIG_KERNEL == KERNEL_3X3_PLUS
+					if (abs(Offset.x) + abs(Offset.y) > 1)
+					{
+						continue;
+					}
+				#endif
 			
-			#if 0
-				SampleSpatialWeight = tsr_half(1.0);
-			#elif 1
-				SampleSpatialWeight = tsr_half(rcp(abs(Offset.x) + 1.0) * rcp(abs(Offset.y) + 1.0));
+				#if 0
+					SampleSpatialWeight = tsr_half(1.0);
+				#elif 1
+					SampleSpatialWeight = tsr_half(rcp(abs(Offset.x) + 1.0) * rcp(abs(Offset.y) + 1.0));
+				#else
+					SampleSpatialWeight = tsr_half(rcp(abs(Offset.x) + abs(Offset.y)));
+				#endif
+			}
 			#else
-				SampleSpatialWeight = tsr_half(rcp(abs(Offset.x) + abs(Offset.y)));
+				#error Unknown kernel
 			#endif
-		}
-		#else
-			#error Unknown kernel
-		#endif
 		
-		tsr_short2 SampleInputPixelPos = InputPixelPos + ClampPixelOffset(
-			InputPixelPos,
-			DynamicSpreadFactor * Offset, Offset,
-			InputPixelPosMin, InputPixelPosMax);
-
-		tsr_halfC InputSample = InputTexture[SampleInputPixelPos];
-		tsr_halfC HistorySample = PredictionSceneColorTexture[SampleInputPixelPos];
-		tsr_half HistorySampleValidity = ParallaxRejectionMaskToHistoryValidity(
-			ParallaxRejectionMaskTexture[SampleInputPixelPos]);
-
-		#if CONFIG_SCENE_COLOR_OVERFLOW
-			#if CONFIG_SCENE_COLOR_APLHA
-				InputSample = min(InputSample, LargestSceneColorRGBA);
-			#else
-				InputSample = min(InputSample, LargestSceneColorRGB);
-			#endif
-		#endif
+			tsr_halfC InputSample;
+			tsr_halfC HistorySample;
+
+			tsr_half InputWeight;
+			tsr_half HistoryWeight;
+			#if CONFIG_CACHE_LDS
+			{
+				tsr_short LoadPos = (
+					dot(TilePixelOffset, tsr_short2(1, LDS_TILE_SIZE)) +
+					dot(tsr_short(LDS_PADDING).xx, tsr_short2(1, LDS_TILE_SIZE)) +
+					dot(DynamicSpreadFactor * Offset, tsr_short2(1, LDS_TILE_SIZE)));
 
-		tsr_half InputWeight = HistorySampleValidity * ComputeColorWeight(InputSample);
-		tsr_half HistoryWeight = HistorySampleValidity * ComputeColorWeight(HistorySample);
+				InputSample = SharedArray0[LoadPos].rgb;
+				HistorySample = SharedArray1[LoadPos].rgb;
 
-		FilteredInput       += (SampleSpatialWeight * InputWeight) * InputSample;
-		FilteredInputWeight += (SampleSpatialWeight * InputWeight);
+				InputWeight = SharedArray0[LoadPos].a;
+				HistoryWeight = SharedArray1[LoadPos].a;
+			}
+			#else // !CONFIG_CACHE_LDS
+			{
+				tsr_short2 SampleInputPixelPos = InputPixelPos + ClampPixelOffset(
+					InputPixelPos,
+					DynamicSpreadFactor * Offset, Offset,
+					InputPixelPosMin, InputPixelPosMax);
+
+				InputSample = InputTexture[SampleInputPixelPos];
+				HistorySample = PredictionSceneColorTexture[SampleInputPixelPos];
+
+				tsr_half HistorySampleValidity = ParallaxRejectionMaskToHistoryValidity(
+					ParallaxRejectionMaskTexture[SampleInputPixelPos]);
+
+				#if CONFIG_SCENE_COLOR_OVERFLOW
+					#if CONFIG_SCENE_COLOR_APLHA
+						InputSample = min(InputSample, LargestSceneColorRGBA);
+					#else
+						InputSample = min(InputSample, LargestSceneColorRGB);
+					#endif
+				#endif
+
+				InputWeight = HistorySampleValidity * ComputeColorWeight(InputSample);
+				HistoryWeight = HistorySampleValidity * ComputeColorWeight(HistorySample);
+			}
+			#endif // !CONFIG_CACHE_LDS
+
+			FilteredInput       += (SampleSpatialWeight * InputWeight) * InputSample;
+			FilteredInputWeight += (SampleSpatialWeight * InputWeight);
 		
-		FilteredHistory       += (SampleSpatialWeight * HistoryWeight) * HistorySample;
-		FilteredHistoryWeight += (SampleSpatialWeight * HistoryWeight);
-	}
+			FilteredHistory       += (SampleSpatialWeight * HistoryWeight) * HistorySample;
+			FilteredHistoryWeight += (SampleSpatialWeight * HistoryWeight);
 
-	{
-		FilteredInput *= SafeRcp(FilteredInputWeight);
-		FilteredHistory *= SafeRcp(FilteredHistoryWeight);
+			if (all(Offset == tsr_short(0)))
+			{
+				CenterInput = InputSample;
+			}
+
+		}
+		
+		{
+			FilteredInput *= SafeRcp(FilteredInputWeight);
+			FilteredHistory *= SafeRcp(FilteredHistoryWeight);
+		}
 	}
-	
+
 	// Compute the LDR luma used for spatial anti-aliasing's computation of the TexelOffset
 	#if DIM_OUTPUT_ANTI_ALIASING_LUMA
 	tsr_half CenterLumaLDR;
@@ -222,7 +311,7 @@ void MainCS(
 	#endif // DIM_OUTPUT_ANTI_ALIASING_LUMA
 	
 	tsr_short2 OutputPixelPos;
-	#if 1
+	#if !CONFIG_CACHE_LDS
 		ISOLATE
 		{
 			OutputPixelPos = (
