Fix for shader compile issue in hair strands vertex factory

diff --git a/Engine/Shaders/Private/HairStrands/HairStrandsVertexFactory.ush b/Engine/Shaders/Private/HairStrands/HairStrandsVertexFactory.ush
--- a/Engine/Shaders/Private/HairStrands/HairStrandsVertexFactory.ush
+++ b/Engine/Shaders/Private/HairStrands/HairStrandsVertexFactory.ush
@@ -538,6 +538,7 @@ FMaterialPixelParameters GetMaterialPixelParameters(float3 RayOrigin, float3 Ray
 	{
 		// find central point of the hair strand
 		const float3 PositionOffset = HairStrandsVF.PositionOffsetBuffer[0].xyz;
+		const float EffectiveRadius = HairStrandsVF.Radius * HairStrandsVF.RaytracingRadiusScale;
 		FHairControlPoint P0 = UnpackHairControlPoint(
 			HairStrandsVF.PositionBuffer[VertexIndex0],
 			PositionOffset,
