Strata - update rough refraction with a simple, but fitted on measured data, model taking into account depth.

diff --git a/Engine/Shaders/Private/Strata/Strata.ush b/Engine/Shaders/Private/Strata/Strata.ush
--- a/Engine/Shaders/Private/Strata/Strata.ush
+++ b/Engine/Shaders/Private/Strata/Strata.ush
@@ -585,6 +585,32 @@ FStrataOpaqueRoughRefractionData StrataUnpackOpaqueRoughRefractionData(float3 In
 	return Out;
 }
 
+
+// Returns the variance of rays projected on the XY plane after scattering at the interface of a slab of matter,
+// having the given Roughness and a normal along Z; after traveling along -Z from air to water (eta = 1.333).
+// That variance has been generated from unit vectors, so it is a measure of the spread of light on the projected XY plane.
+// Using Thales, one can scale that spread to any distance. The Gaussian kernel can be evaluated for a depth d using a Variance = GetVarianceFromRoughness(Roughness) * d.
+// 
+//			| V = -Z
+//			|
+//			|
+//			|
+//			|
+//			v
+//	------------------- ^	Surfcace of roughness R
+//		   /|\			|
+//		  / | \			| Depth D
+//		 /  |  \		|
+//		|<----->|		v
+//			Variance = GetVarianceFromRoughness(R) * D;
+float GetVarianceFromRoughness(float Roughness)
+{
+	float Roughness2 = Roughness * Roughness;
+	float Roughness3 = Roughness2 * Roughness;
+	float Roughness4 = Roughness2 * Roughness2;
+	return saturate(0.0000483367f - 0.000707943f * Roughness - 0.00840911f * Roughness2 + 0.095437f * Roughness3 - 0.0478129f * Roughness4);
+}
+
 ///////////////////////////////////////////////////////////////////////////////
 // Sub-surface 
 
@@ -2690,7 +2716,7 @@ void PackStrataOut(
 			const FStrataOperator RootOperator = StrataTree.Operators[Strata.OperatorIndex];
 			OpaqueRoughRefractionData.Coverage = RootOperator.Coverage;
 			// STRATA_TODO Fake variance place holder for now
-			OpaqueRoughRefractionData.VarianceCm = RootOperator.OpaqueRoughRefractTopRoughness * RootOperator.OpaqueRoughRefractThicknessCm;
+			OpaqueRoughRefractionData.VarianceCm = GetVarianceFromRoughness(RootOperator.OpaqueRoughRefractTopRoughness) * RootOperator.OpaqueRoughRefractThicknessCm;
 		}
 		
 		int SimpleMaterial_BSDFIndex = 0;
diff --git a/Engine/Shaders/Private/Strata/StrataRoughRefraction.usf b/Engine/Shaders/Private/Strata/StrataRoughRefraction.usf
--- a/Engine/Shaders/Private/Strata/StrataRoughRefraction.usf
+++ b/Engine/Shaders/Private/Strata/StrataRoughRefraction.usf
@@ -6,12 +6,31 @@
 
 
 
+// Gaussian function with a=1/(c*sqrt(2*PI)) so it represents the probability density function of a normally distributed random variable with expected value mu = b for some variance and integrale = 1;
+float GaussianFunction(float x, float Mean, float Variance)
+{
+	float StandardDeviation = sqrt(Variance);
+	float FactorA = 1 / (StandardDeviation * sqrt(2 * PI));
+	float Bell = exp(-abs((x - Mean) * (x - Mean)) / (2 * Variance));
+	return FactorA * Bell;
+}
+float GaussianFunction(float2 P, float2 Mean, float Variance)
+{
+	float StandardDeviation = sqrt(Variance);
+	float FactorA = 1 / (StandardDeviation * sqrt(2 * PI));
+	float2 Bell = exp(-((P.x - Mean.x) * (P.x - Mean.x)) / (2 * Variance) - ((P.y - Mean.y) * (P.y - Mean.y)) / (2 * Variance));
+	return FactorA * Bell;
+}
+
+
+
 #if OPAQUE_ROUGH_REFRACTION_PS
 
 #include "/Engine/Private/Strata/Strata.ush"
 
 Texture2D<float4> SeparatedOpaqueRoughRefractionSceneColor;
 float2 BlurDirection;
+float BlurScale;
 
 void OpaqueRoughRefractionPS(
 	float4 SVPos : SV_POSITION,
@@ -20,26 +39,49 @@ void OpaqueRoughRefractionPS(
 	float2 PixelPos = SVPos.xy;
 	float2 UvBuffer = PixelPos * View.BufferSizeAndInvSize.zw;
 
+	// First read unblurred layers at the bottom of the top layer
+	OutColor = SeparatedOpaqueRoughRefractionSceneColor[PixelPos];
+
 #if PERMUTATION_ENABLE_BLUR && STRATA_OPAQUE_ROUGH_REFRACTION_ENABLED
 	FStrataOpaqueRoughRefractionData ORRData = StrataUnpackOpaqueRoughRefractionData(Strata.OpaqueRoughRefractionTexture[PixelPos]);
 
-	if (ORRData.OpaqueRoughRefractionEnabled)
+	const float DeviceZ = LookupDeviceZ(UvBuffer);
+	const float SceneDepthCm = ConvertFromDeviceZ(DeviceZ);
+
+	const float2 PixelRadiusCm2 = SceneDepthCm * GetTanHalfFieldOfView() * View.ViewSizeAndInvSize.zw;
+	const float PixelRadiusCm = min(PixelRadiusCm2.x, PixelRadiusCm2.y);
+
+	const float PerceptualBlurAcceptanceFactor = 1.5f;
+	const bool OpaqueRoughRefractionEnabled = ORRData.OpaqueRoughRefractionEnabled && (PixelRadiusCm < PerceptualBlurAcceptanceFactor * ORRData.VarianceCm);
+
+	// We are now going to blur the bottom layer luminance according to the top layer variance resulting from the roughness.
+	// For this we have: measure GGX variance for a perpendicaular view, eta=1.33 and varying roughness. At runtim that variance is approximated by a Gaussian.
+
+	if (OpaqueRoughRefractionEnabled)
 	{
-		// BAD placeholder blur
-		float4 SourceData = 0.0f;
+		// Proceed blurring the bottom layers if needed according to the blur footprint.
+		float3 SourceData = 0.0f;
 		float SampleCount = 0.0f;
-		for (float u = -4.5f; u <= 4.5f; u += 1.0f)
+
+		const float VariancePixelSpace = ORRData.VarianceCm / PixelRadiusCm;
+		const float StandardDeviationPixelSpace = sqrt(VariancePixelSpace);
+
+		// A area of radius StandardDeviation*3 around the mean of a guassian normal distribution (=mu) encompass 99% of the energy.
+		const float RadiusAroundCenterMu = 3.0f;
+		for (float u = -RadiusAroundCenterMu; u <= RadiusAroundCenterMu; u += 0.5f) // steps of 0.5 means 13 taps
 		{
-			SourceData += SeparatedOpaqueRoughRefractionSceneColor[PixelPos + BlurDirection * u * ORRData.VarianceCm * 10000.0];
-			SampleCount++;
+			const float OffsetPixelSpace = u * StandardDeviationPixelSpace;
+			const float Weight = GaussianFunction(OffsetPixelSpace, 0.0f, VariancePixelSpace);
+			SourceData += SeparatedOpaqueRoughRefractionSceneColor[PixelPos + BlurDirection * OffsetPixelSpace].rgb * Weight;
+			SampleCount += Weight;
 		}
-		OutColor = SourceData / SampleCount;
+		const float3 BlurredSceneColor = SampleCount > 0.0 ? SourceData / SampleCount : 0.0f;
+
+		// Combine blurred layers below with unblurred according to coverage
+		OutColor = float4(lerp(OutColor.rgb, BlurredSceneColor, ORRData.Coverage), OutColor.a);
+		//OutColor *= float4(0, 1, 0, 1);
 	}
-	else
 #endif
-	{
-		OutColor = SeparatedOpaqueRoughRefractionSceneColor[PixelPos];
-	}
 
 }
 
@@ -377,22 +419,6 @@ FStrataLobeStatistic StrataGetRefractedLobe2(FStrataLobeStatistic WiLobe, float3
 	return WoLobe;
 }
 
-// Gaussian function with a=1/(c*sqrt(2*PI)) so it represents the probability density function of a normally distributed random variable with expected value mu = b for some variance and integrale = 1;
-float GaussianFunction(float x, float Mean, float Variance)
-{
-	float StandardDeviation = sqrt(Variance);
-	float FactorA = 1 / (StandardDeviation * sqrt(2 * PI));
-	float Bell = exp(-((x - Mean) * (x - Mean)) / (2 * Variance));
-	return FactorA * Bell;
-}
-float GaussianFunction(float2 P, float2 Mean, float Variance)
-{
-	float StandardDeviation = sqrt(Variance);
-	float FactorA = 1 / (StandardDeviation * sqrt(2 * PI));
-	float2 Bell = exp(-((P.x - Mean.x) * (P.x - Mean.x)) / (2 * Variance) - ((P.y - Mean.y) * (P.y - Mean.y)) / (2 * Variance));
-	return FactorA * Bell;
-}
-
 void VisualizeRoughRefractionPS(
 	float4 SVPos : SV_POSITION,
 	out float4 OutColor : SV_Target0)
diff --git a/Engine/Source/Runtime/Renderer/Private/Strata/StrataRoughRefraction.cpp b/Engine/Source/Runtime/Renderer/Private/Strata/StrataRoughRefraction.cpp
--- a/Engine/Source/Runtime/Renderer/Private/Strata/StrataRoughRefraction.cpp
+++ b/Engine/Source/Runtime/Renderer/Private/Strata/StrataRoughRefraction.cpp
@@ -19,6 +19,12 @@ static TAutoConsoleVariable<int32> CVarStrataOpaqueMaterialRoughRefraction(
 	TEXT("Enable Strata opaque material rough refractions effect from top layers over layers below."),
 	ECVF_ReadOnly | ECVF_RenderThreadSafe);
 
+static TAutoConsoleVariable<float> CVarStrataOpaqueMaterialRoughRefractionBlurScale(
+	TEXT("r.Strata.OpaqueMaterialRoughRefraction.BlurScale"),
+	1,
+	TEXT("Scale opaque rough refraction strengh for debug purposes."),
+	ECVF_ReadOnly | ECVF_RenderThreadSafe);
+
 
 namespace Strata
 {
@@ -38,10 +44,12 @@ class FOpaqueRoughRefractionPS : public FGlobalShader
 
 	BEGIN_SHADER_PARAMETER_STRUCT(FParameters, )
 		SHADER_PARAMETER_STRUCT_REF(FViewUniformShaderParameters, ViewUniformBuffer)
+		SHADER_PARAMETER_STRUCT_INCLUDE(FSceneTextureShaderParameters, SceneTextures)
 		SHADER_PARAMETER_RDG_UNIFORM_BUFFER(FStrataGlobalUniformParameters, Strata)
 		SHADER_PARAMETER_STRUCT_INCLUDE(Strata::FStrataTilePassVS::FParameters, StrataTile)
 		SHADER_PARAMETER_RDG_TEXTURE(Texture2D, SeparatedOpaqueRoughRefractionSceneColor)
 		SHADER_PARAMETER(FVector2f, BlurDirection)
+		SHADER_PARAMETER(float, BlurScale)
 		RENDER_TARGET_BINDING_SLOTS()
 	END_SHADER_PARAMETER_STRUCT();
 
@@ -97,6 +105,7 @@ void AddStrataOpaqueRoughRefractionPasses(
 		FSceneTextureParameters SceneTextureParameters = GetSceneTextureParameters(GraphBuilder);
 		FOpaqueRoughRefractionPS::FParameters* PassParameters = GraphBuilder.AllocParameters<FOpaqueRoughRefractionPS::FParameters>();
 		PassParameters->ViewUniformBuffer = View.ViewUniformBuffer;
+		PassParameters->SceneTextures = GetSceneTextureShaderParameters(SceneTextures.UniformBuffer);
 		PassParameters->Strata = Strata::BindStrataGlobalUniformParameters(View.StrataSceneData);
 		PassParameters->SeparatedOpaqueRoughRefractionSceneColor = SeparatedOpaqueRoughRefractionSceneColor;
 		PassParameters->RenderTargets[0] = FRenderTargetBinding(TempTexture, LoadAction);
@@ -152,6 +161,7 @@ void AddStrataOpaqueRoughRefractionPasses(
 		FSceneTextureParameters SceneTextureParameters = GetSceneTextureParameters(GraphBuilder);
 		FOpaqueRoughRefractionPS::FParameters* PassParameters = GraphBuilder.AllocParameters<FOpaqueRoughRefractionPS::FParameters>();
 		PassParameters->ViewUniformBuffer = View.ViewUniformBuffer;
+		PassParameters->SceneTextures = GetSceneTextureShaderParameters(SceneTextures.UniformBuffer);
 		PassParameters->Strata = Strata::BindStrataGlobalUniformParameters(View.StrataSceneData);
 		PassParameters->SeparatedOpaqueRoughRefractionSceneColor = TempTexture;
 		PassParameters->RenderTargets[0] = FRenderTargetBinding(SceneColorTexture, ERenderTargetLoadAction::ELoad);
@@ -170,7 +180,7 @@ void AddStrataOpaqueRoughRefractionPasses(
 		PassParameters->StrataTile = Strata::SetTileParameters(GraphBuilder, View, StrataTileType, PrimitiveType);
 
 		GraphBuilder.AddPass(
-			RDG_EVENT_NAME("OpaqueRoughRefraction - Pass2)"),
+			RDG_EVENT_NAME("OpaqueRoughRefraction - Pass1)"),
 			PassParameters,
 			ERDGPassFlags::Raster,
 			[&View, TileVertexShader, PixelShader, PassParameters, StrataTileType, PrimitiveType](FRHICommandList& RHICmdList)
@@ -208,6 +218,7 @@ void AddStrataOpaqueRoughRefractionPasses(
 		FSceneTextureParameters SceneTextureParameters = GetSceneTextureParameters(GraphBuilder);
 		FOpaqueRoughRefractionPS::FParameters* PassParameters = GraphBuilder.AllocParameters<FOpaqueRoughRefractionPS::FParameters>();
 		PassParameters->ViewUniformBuffer = View.ViewUniformBuffer;
+		PassParameters->SceneTextures = GetSceneTextureShaderParameters(SceneTextures.UniformBuffer);
 		PassParameters->Strata = Strata::BindStrataGlobalUniformParameters(View.StrataSceneData);
 		PassParameters->SeparatedOpaqueRoughRefractionSceneColor = SeparatedOpaqueRoughRefractionSceneColor;
 		PassParameters->RenderTargets[0] = FRenderTargetBinding(SceneColorTexture, ERenderTargetLoadAction::ELoad);
