Strata - opaque rough refraction
- Pluged in blurscale
- Now sqmple scene color using bilinear filteirng
- Added noise to blur position for TAA to smooth things out

diff --git a/Engine/Shaders/Private/Strata/StrataRoughRefraction.usf b/Engine/Shaders/Private/Strata/StrataRoughRefraction.usf
--- a/Engine/Shaders/Private/Strata/StrataRoughRefraction.usf
+++ b/Engine/Shaders/Private/Strata/StrataRoughRefraction.usf
@@ -44,6 +44,7 @@ void OpaqueRoughRefractionPS(
 
 #if PERMUTATION_ENABLE_BLUR && STRATA_OPAQUE_ROUGH_REFRACTION_ENABLED
 	FStrataOpaqueRoughRefractionData ORRData = StrataUnpackOpaqueRoughRefractionData(Strata.OpaqueRoughRefractionTexture[PixelPos]);
+	ORRData.VarianceCm *= BlurScale;
 
 	const float DeviceZ = LookupDeviceZ(UvBuffer);
 	const float SceneDepthCm = ConvertFromDeviceZ(DeviceZ);
@@ -51,7 +52,7 @@ void OpaqueRoughRefractionPS(
 	const float2 PixelRadiusCm2 = SceneDepthCm * GetTanHalfFieldOfView() * View.ViewSizeAndInvSize.zw;
 	const float PixelRadiusCm = min(PixelRadiusCm2.x, PixelRadiusCm2.y);
 
-	const float PerceptualBlurAcceptanceFactor = 1.5f;
+	const float PerceptualBlurAcceptanceFactor = 2.5f;
 	const bool OpaqueRoughRefractionEnabled = ORRData.OpaqueRoughRefractionEnabled && (PixelRadiusCm < PerceptualBlurAcceptanceFactor * ORRData.VarianceCm);
 
 	// We are now going to blur the bottom layer luminance according to the top layer variance resulting from the roughness.
@@ -60,25 +61,37 @@ void OpaqueRoughRefractionPS(
 	if (OpaqueRoughRefractionEnabled)
 	{
 		// Proceed blurring the bottom layers if needed according to the blur footprint.
-		float3 SourceData = 0.0f;
-		float SampleCount = 0.0f;
+		float3 WeightedLuminance = 0.0f;
+		float AccumulatedWeights = 0.0f;
 
 		const float VariancePixelSpace = ORRData.VarianceCm / PixelRadiusCm;
 		const float StandardDeviationPixelSpace = sqrt(VariancePixelSpace);
 
 		// A area of radius StandardDeviation*3 around the mean of a guassian normal distribution (=mu) encompass 99% of the energy.
 		const float RadiusAroundCenterMu = 3.0f;
-		for (float u = -RadiusAroundCenterMu; u <= RadiusAroundCenterMu; u += 0.5f) // steps of 0.5 means 13 taps
+		const float BlurStep = 0.5f;
+
+		// Randomly offset the blur center position around according to step size to leverage TAA and smooth the blur.
+		const float BlurRandomness = InterleavedGradientNoise(PixelPos, float(View.StateFrameIndexMod8)) - 0.5; // In [-0.5, 0.5]
+		const float RandAngle = BlurRandomness * 2.0f * PI;
+		float CosA, SinA;
+		sincos(RandAngle, SinA, CosA);
+		const float2 BlurOrigin = PixelPos + BlurStep * StandardDeviationPixelSpace * float2(CosA, SinA);
+
+		for (float u = -RadiusAroundCenterMu; u <= RadiusAroundCenterMu; u += BlurStep) // steps of 0.5 means 13 taps
 		{
 			const float OffsetPixelSpace = u * StandardDeviationPixelSpace;
 			const float Weight = GaussianFunction(OffsetPixelSpace, 0.0f, VariancePixelSpace);
-			SourceData += SeparatedOpaqueRoughRefractionSceneColor[PixelPos + BlurDirection * OffsetPixelSpace].rgb * Weight;
-			SampleCount += Weight;
+			const float2 UV = (BlurOrigin + BlurDirection * OffsetPixelSpace) * View.BufferSizeAndInvSize.zw;
+			WeightedLuminance += SeparatedOpaqueRoughRefractionSceneColor.SampleLevel(View.SharedBilinearClampedSampler, UV, 0).rgb * Weight;
+			AccumulatedWeights += Weight;
 		}
-		const float3 BlurredSceneColor = SampleCount > 0.0 ? SourceData / SampleCount : 0.0f;
+
+		const float bIsBlurValid = AccumulatedWeights > 0.0;
+		const float3 BlurredSceneColor = bIsBlurValid ? WeightedLuminance / AccumulatedWeights : 0.0f;
 
 		// Combine blurred layers below with unblurred according to coverage
-		OutColor = float4(lerp(OutColor.rgb, BlurredSceneColor, ORRData.Coverage), OutColor.a);
+		OutColor = float4(lerp(OutColor.rgb, BlurredSceneColor, bIsBlurValid ? ORRData.Coverage : 0.0), OutColor.a);
 		//OutColor *= float4(0, 1, 0, 1);
 	}
 #endif
diff --git a/Engine/Source/Runtime/Renderer/Private/Strata/StrataRoughRefraction.cpp b/Engine/Source/Runtime/Renderer/Private/Strata/StrataRoughRefraction.cpp
--- a/Engine/Source/Runtime/Renderer/Private/Strata/StrataRoughRefraction.cpp
+++ b/Engine/Source/Runtime/Renderer/Private/Strata/StrataRoughRefraction.cpp
@@ -21,9 +21,9 @@ static TAutoConsoleVariable<int32> CVarStrataOpaqueMaterialRoughRefraction(
 
 static TAutoConsoleVariable<float> CVarStrataOpaqueMaterialRoughRefractionBlurScale(
 	TEXT("r.Strata.OpaqueMaterialRoughRefraction.BlurScale"),
-	1,
+	1.0f,
 	TEXT("Scale opaque rough refraction strengh for debug purposes."),
-	ECVF_ReadOnly | ECVF_RenderThreadSafe);
+	ECVF_RenderThreadSafe);
 
 
 namespace Strata
@@ -110,6 +110,7 @@ void AddStrataOpaqueRoughRefractionPasses(
 		PassParameters->SeparatedOpaqueRoughRefractionSceneColor = SeparatedOpaqueRoughRefractionSceneColor;
 		PassParameters->RenderTargets[0] = FRenderTargetBinding(TempTexture, LoadAction);
 		PassParameters->BlurDirection = FVector2f(1.0f, 0.0f);
+		PassParameters->BlurScale = FMath::Max(0.0f, CVarStrataOpaqueMaterialRoughRefractionBlurScale.GetValueOnAnyThread());
 
 		FOpaqueRoughRefractionPS::FPermutationDomain PermutationVector;
 		PermutationVector.Set<FOpaqueRoughRefractionPS::FEnableBlur>(true);
@@ -166,6 +167,7 @@ void AddStrataOpaqueRoughRefractionPasses(
 		PassParameters->SeparatedOpaqueRoughRefractionSceneColor = TempTexture;
 		PassParameters->RenderTargets[0] = FRenderTargetBinding(SceneColorTexture, ERenderTargetLoadAction::ELoad);
 		PassParameters->BlurDirection = FVector2f(0.0f, 1.0f);
+		PassParameters->BlurScale = FMath::Max(0.0f, CVarStrataOpaqueMaterialRoughRefractionBlurScale.GetValueOnAnyThread());
 
 		FOpaqueRoughRefractionPS::FPermutationDomain PermutationVector;
 		PermutationVector.Set<FOpaqueRoughRefractionPS::FEnableBlur>(true);
@@ -223,6 +225,7 @@ void AddStrataOpaqueRoughRefractionPasses(
 		PassParameters->SeparatedOpaqueRoughRefractionSceneColor = SeparatedOpaqueRoughRefractionSceneColor;
 		PassParameters->RenderTargets[0] = FRenderTargetBinding(SceneColorTexture, ERenderTargetLoadAction::ELoad);
 		PassParameters->BlurDirection = FVector2f(0.0f, 0.0f);
+		PassParameters->BlurScale = FMath::Max(0.0f, CVarStrataOpaqueMaterialRoughRefractionBlurScale.GetValueOnAnyThread());
 
 		FOpaqueRoughRefractionPS::FPermutationDomain PermutationVector;
 		PermutationVector.Set<FOpaqueRoughRefractionPS::FEnableBlur>(false);
