Reenable OIT parameters binding.

diff --git a/Engine/Shaders/Private/BasePassPixelShader.usf b/Engine/Shaders/Private/BasePassPixelShader.usf
--- a/Engine/Shaders/Private/BasePassPixelShader.usf
+++ b/Engine/Shaders/Private/BasePassPixelShader.usf
@@ -135,7 +135,7 @@ float NormalCurvatureToRoughness(float3 WorldNormal)
 #define POST_PROCESS_SUBSURFACE ((MATERIAL_SHADINGMODEL_SUBSURFACE_PROFILE || MATERIAL_SHADINGMODEL_EYE) && USES_GBUFFER)
 #define MATERIAL_WORKS_WITH_DUAL_SOURCE_COLOR_BLENDING (MATERIAL_SHADINGMODEL_THIN_TRANSLUCENT || STRATA_BLENDING_TRANSLUCENT_COLOREDTRANSMITTANCE )
 
-#define OIT_ENABLED (PROJECT_OIT && (MATERIALBLENDING_TRANSLUCENT || MATERIALBLENDING_ADDITIVE || MATERIAL_WORKS_WITH_DUAL_SOURCE_COLOR_BLENDING))
+#define OIT_ENABLED (PROJECT_OIT && PLATFORM_SUPPORTS_ROV && (MATERIALBLENDING_TRANSLUCENT || MATERIALBLENDING_ADDITIVE || MATERIAL_WORKS_WITH_DUAL_SOURCE_COLOR_BLENDING))
 
 #if OIT_ENABLED
 #define OIT_IS_BASEPASS 1
diff --git a/Engine/Shaders/Private/OITCommon.ush b/Engine/Shaders/Private/OITCommon.ush
--- a/Engine/Shaders/Private/OITCommon.ush
+++ b/Engine/Shaders/Private/OITCommon.ush
@@ -230,6 +230,8 @@ uint2 GetStoringCoord(uint2 InBaseCoord, uint InSampleIt)
 // Multi-Layer Alpha Blending (MLAB) - Sample are inserted in a sorted manner and overflow sample are merged with alpha blending
 void AddOITSample_MLAB(uint2 InPixelCoord, float3 InColor, float3 InTrans, float InDepth)
 {
+	RASTER_ORDERED_VIEW_LOCK();
+
 	// ROV read to for pixel synchronization
 	uint SampleCount = min(TranslucentBasePass.OIT.OutOITSampleCount[InPixelCoord], TranslucentBasePass.OIT.MaxSamplePerPixel);
 
@@ -293,6 +295,8 @@ void AddOITSample_MLAB(uint2 InPixelCoord, float3 InColor, float3 InTrans, float
 	// Write to ROV to unlock pixel
 	SampleCount = min(SampleCount + 1, TranslucentBasePass.OIT.MaxSamplePerPixel);
 	TranslucentBasePass.OIT.OutOITSampleCount[InPixelCoord] = SampleCount;
+
+	RASTER_ORDERED_VIEW_UNLOCK();
 }
 
 // Add a translucent sample to the list of existing translucent samples
diff --git a/Engine/Source/Runtime/Renderer/Private/BasePassRendering.h b/Engine/Source/Runtime/Renderer/Private/BasePassRendering.h
--- a/Engine/Source/Runtime/Renderer/Private/BasePassRendering.h
+++ b/Engine/Source/Runtime/Renderer/Private/BasePassRendering.h
@@ -100,7 +100,7 @@ BEGIN_GLOBAL_SHADER_PARAMETER_STRUCT(FTranslucentBasePassUniformParameters,)
 	SHADER_PARAMETER_STRUCT(FSceneTextureUniformParameters, SceneTextures)
 	SHADER_PARAMETER_STRUCT(FStrataForwardPassUniformParameters, Strata)
 	SHADER_PARAMETER_STRUCT(FLightCloudTransmittanceParameters, ForwardDirLightCloudShadow)
-	//SHADER_PARAMETER_STRUCT(FOITBasePassUniformParameters, OIT)
+	SHADER_PARAMETER_STRUCT(FOITBasePassUniformParameters, OIT)
 	// Material SSR
 	SHADER_PARAMETER(FVector4f, HZBUvFactorAndInvFactor)
 	SHADER_PARAMETER(FVector4f, PrevScreenPositionScaleBias)
diff --git a/Engine/Source/Runtime/Renderer/Private/OIT/OIT.cpp b/Engine/Source/Runtime/Renderer/Private/OIT/OIT.cpp
--- a/Engine/Source/Runtime/Renderer/Private/OIT/OIT.cpp
+++ b/Engine/Source/Runtime/Renderer/Private/OIT/OIT.cpp
@@ -89,7 +89,7 @@ static bool IsOITSortedTrianglesSupported(EShaderPlatform InShaderPlatform)
 
 static bool IsOITSortedPixelsSupported(EShaderPlatform InShaderPlatform)
 {
-	return RHISupportsComputeShaders(InShaderPlatform) && FDataDrivenShaderPlatformInfo::GetSupportsOIT(InShaderPlatform);
+	return RHISupportsComputeShaders(InShaderPlatform) && FDataDrivenShaderPlatformInfo::GetSupportsROV(InShaderPlatform) && FDataDrivenShaderPlatformInfo::GetSupportsOIT(InShaderPlatform);
 }
 
 ///////////////////////////////////////////////////////////////////////////////////////////////////
diff --git a/Engine/Source/Runtime/Renderer/Private/TranslucentRendering.cpp b/Engine/Source/Runtime/Renderer/Private/TranslucentRendering.cpp
--- a/Engine/Source/Runtime/Renderer/Private/TranslucentRendering.cpp
+++ b/Engine/Source/Runtime/Renderer/Private/TranslucentRendering.cpp
@@ -1073,7 +1073,7 @@ TRDGUniformBufferRef<FTranslucentBasePassUniformParameters> CreateTranslucentBas
 	BasePassParameters.PreIntegratedGFTexture = GSystemTextures.PreintegratedGF->GetRenderTargetItem().ShaderResourceTexture;
 	BasePassParameters.PreIntegratedGFSampler = TStaticSamplerState<SF_Bilinear, AM_Clamp, AM_Clamp, AM_Clamp>::GetRHI();
 
-	//OIT::SetOITParameters(GraphBuilder, View, BasePassParameters.OIT, OITData);
+	OIT::SetOITParameters(GraphBuilder, View, BasePassParameters.OIT, OITData);
 
 	return GraphBuilder.CreateUniformBuffer(&BasePassParameters);
 }
