Strata - Opaque rough refraction - fixed coverage that was always 1.

diff --git a/Engine/Shaders/Private/Strata/Strata.ush b/Engine/Shaders/Private/Strata/Strata.ush
--- a/Engine/Shaders/Private/Strata/Strata.ush
+++ b/Engine/Shaders/Private/Strata/Strata.ush
@@ -774,6 +774,7 @@ struct FStrataOperator
 	FStrataLobeStatistic BottomUpRefractionLobeStat;	// The lobe used for the scene rough refraction from back to front, e.g. used for opaque rough refraction of incoming light within sub layers
 
 	// Opaque rough refraction
+	float				 OpaqueRoughRefractCoverage;	// The coverage of the top layer to account for in order to fade out the influence on the blur.
 	float				 OpaqueRoughRefractThicknessCm;	// The thickness of stack material from the top layer to the bottom layer.
 	float				 OpaqueRoughRefractTopRoughness;// The roughness of the top layer. This is enough because we currently assume all the bottom layers have a constant IOR, so eta will after that always be 1 and no refraction will occur after.
 
@@ -2714,7 +2715,7 @@ void PackStrataOut(
 		{
 			uint RootOperatorIndex = Strata.OperatorIndex;
 			const FStrataOperator RootOperator = StrataTree.Operators[Strata.OperatorIndex];
-			OpaqueRoughRefractionData.Coverage = RootOperator.Coverage;
+			OpaqueRoughRefractionData.Coverage = RootOperator.OpaqueRoughRefractCoverage;
 			// STRATA_TODO Fake variance place holder for now
 			OpaqueRoughRefractionData.VarianceCm = GetVarianceFromRoughness(RootOperator.OpaqueRoughRefractTopRoughness) * RootOperator.OpaqueRoughRefractThicknessCm;
 		}
diff --git a/Engine/Shaders/Private/Strata/StrataTree.ush b/Engine/Shaders/Private/Strata/StrataTree.ush
--- a/Engine/Shaders/Private/Strata/StrataTree.ush
+++ b/Engine/Shaders/Private/Strata/StrataTree.ush
@@ -79,8 +79,9 @@ void UpdateSingleBSDFOperatorCoverageTransmittance(
 
 	// Extract the true thickness of the BSDF before it is cast into the normalised Slab thickness we used.
 	// for opaque rough refraction, we are only interested in thickness of non bottom layers.
-	Op.OpaqueRoughRefractThicknessCm  = CurrentBSDF.bIsBottom ? 0.0f : CurrentBSDF.ThicknessCm;
-	Op.OpaqueRoughRefractTopRoughness = CurrentBSDF.bIsTop    ? Roughness : 0.0f;
+	Op.OpaqueRoughRefractThicknessCm  = CurrentBSDF.bIsBottom ? 0.0f					: CurrentBSDF.ThicknessCm;
+	Op.OpaqueRoughRefractCoverage	  = CurrentBSDF.bIsTop    ? CurrentBSDF.Coverage	: 0.0f;
+	Op.OpaqueRoughRefractTopRoughness = CurrentBSDF.bIsTop    ? Roughness				: 0.0f;
 
 	PostProcessBSDFBeforeLighting(CurrentBSDF);
 
@@ -167,7 +168,8 @@ void UpdateSingleOperatorCoverageTransmittance(inout FStrataTree StrataTree, int
 		Op.BottomUpRefractionLobeStat = StrataWeightLobe(OpA.BottomUpRefractionLobeStat, Weight);
 
 		// Coverage does not affect rough refraction data
-		Op.OpaqueRoughRefractThicknessCm = OpA.OpaqueRoughRefractThicknessCm;
+		Op.OpaqueRoughRefractThicknessCm  = OpA.OpaqueRoughRefractThicknessCm;
+		Op.OpaqueRoughRefractCoverage     = OpA.OpaqueRoughRefractCoverage * Weight;
 		Op.OpaqueRoughRefractTopRoughness = OpA.OpaqueRoughRefractTopRoughness;
 		break;
 	}
@@ -222,8 +224,9 @@ void UpdateSingleOperatorCoverageTransmittance(inout FStrataTree StrataTree, int
 			Op.BottomUpRefractionLobeStat.E *= saturate(OpA.Transmittance + (1.0f - OpA.Coverage)); // Combine with top layer medium and specular transmittance
 		}
 
-		// Vertical layering accumulates both operators thickness and only keep the top layer roughness, see OpaqueRoughRefractTopRoughness declaration for the details.
+		// Vertical layering accumulates both operators thickness and only keep the top layer roughness and coverage, see OpaqueRoughRefractTopRoughness declaration for the details.
 		Op.OpaqueRoughRefractThicknessCm  = OpA.OpaqueRoughRefractThicknessCm + OpB.OpaqueRoughRefractThicknessCm;
+		Op.OpaqueRoughRefractCoverage     = OpA.OpaqueRoughRefractCoverage;
 		Op.OpaqueRoughRefractTopRoughness = OpA.OpaqueRoughRefractTopRoughness;
 
 		Op.VerticalTop_TopDownRefractionLobeStat = OpA.TopDownRefractionLobeStat;
@@ -246,7 +249,8 @@ void UpdateSingleOperatorCoverageTransmittance(inout FStrataTree StrataTree, int
 		Op.BottomUpRefractionLobeStat = StrataHorizontalMixLobes(OpA.BottomUpRefractionLobeStat, OpB.BottomUpRefractionLobeStat, BMix);
 
 		// Horizontal mixes both operators thickness
-		Op.OpaqueRoughRefractThicknessCm  = lerp(OpA.OpaqueRoughRefractThicknessCm, OpB.OpaqueRoughRefractThicknessCm, BMix);
+		Op.OpaqueRoughRefractThicknessCm  = lerp(OpA.OpaqueRoughRefractThicknessCm,  OpB.OpaqueRoughRefractThicknessCm,  BMix);
+		Op.OpaqueRoughRefractCoverage     = lerp(OpA.OpaqueRoughRefractCoverage,     OpB.OpaqueRoughRefractCoverage,     BMix);
 		Op.OpaqueRoughRefractTopRoughness = lerp(OpA.OpaqueRoughRefractTopRoughness, OpB.OpaqueRoughRefractTopRoughness, BMix);
 		break;
 	}
@@ -265,7 +269,8 @@ void UpdateSingleOperatorCoverageTransmittance(inout FStrataTree StrataTree, int
 
 		Op.BottomUpRefractionLobeStat = StrataHorizontalMixLobes(OpA.BottomUpRefractionLobeStat, OpB.BottomUpRefractionLobeStat, 0.5f);
 
-		Op.OpaqueRoughRefractThicknessCm  = lerp(OpA.OpaqueRoughRefractThicknessCm, OpB.OpaqueRoughRefractThicknessCm, 0.5f);
+		Op.OpaqueRoughRefractThicknessCm  = lerp(OpA.OpaqueRoughRefractThicknessCm,  OpB.OpaqueRoughRefractThicknessCm,  0.5f);
+		Op.OpaqueRoughRefractCoverage     = lerp(OpA.OpaqueRoughRefractCoverage,     OpB.OpaqueRoughRefractCoverage,     0.5f);
 		Op.OpaqueRoughRefractTopRoughness = lerp(OpA.OpaqueRoughRefractTopRoughness, OpB.OpaqueRoughRefractTopRoughness, 0.5f);
 		break;
 	}
@@ -279,6 +284,7 @@ void UpdateSingleOperatorCoverageTransmittance(inout FStrataTree StrataTree, int
 		Op.TopDownRefractionWorldNormal = OpA.TopDownRefractionWorldNormal;
 		Op.BottomUpRefractionLobeStat = OpA.BottomUpRefractionLobeStat;
 		Op.OpaqueRoughRefractThicknessCm = OpA.OpaqueRoughRefractThicknessCm;
+		Op.OpaqueRoughRefractCoverage = OpA.OpaqueRoughRefractCoverage;
 		Op.OpaqueRoughRefractTopRoughness = OpA.OpaqueRoughRefractTopRoughness;
 		break;
 	}
