Fix shader compilation issue with TSR when optimization are disabled

diff --git a/Engine/Shaders/Private/TemporalSuperResolution/TSRFilterFrequencies.usf b/Engine/Shaders/Private/TemporalSuperResolution/TSRFilterFrequencies.usf
--- a/Engine/Shaders/Private/TemporalSuperResolution/TSRFilterFrequencies.usf
+++ b/Engine/Shaders/Private/TemporalSuperResolution/TSRFilterFrequencies.usf
@@ -180,6 +180,7 @@ void MainCS(
 	tsr_halfC FilteredInput;
 	tsr_halfC FilteredHistory;
 	{
+		CenterInput = tsr_half(0.0);
 		FilteredInput = tsr_half(0.0);
 		FilteredHistory = tsr_half(0.0);
 
