Fix for compilation issue on a console.

diff --git a/Engine/Shaders/Private/Strata/StrataRoughRefraction.usf b/Engine/Shaders/Private/Strata/StrataRoughRefraction.usf
--- a/Engine/Shaders/Private/Strata/StrataRoughRefraction.usf
+++ b/Engine/Shaders/Private/Strata/StrataRoughRefraction.usf
@@ -40,7 +40,7 @@ void OpaqueRoughRefractionPS(
 	float2 UvBuffer = PixelPos * View.BufferSizeAndInvSize.zw;
 
 	// First read unblurred layers at the bottom of the top layer
-	OutColor = SeparatedOpaqueRoughRefractionSceneColor[PixelPos];
+	OutColor = SeparatedOpaqueRoughRefractionSceneColor[uint2(PixelPos)];
 
 #if PERMUTATION_ENABLE_BLUR && STRATA_OPAQUE_ROUGH_REFRACTION_ENABLED
 	FStrataOpaqueRoughRefractionData ORRData = StrataUnpackOpaqueRoughRefractionData(Strata.OpaqueRoughRefractionTexture[PixelPos]);
