Strata - CLoth BSDF update to Charlie using current controls.
Cloth energy conservation table has also been updated.

diff --git a/Engine/Shaders/Private/BRDF.ush b/Engine/Shaders/Private/BRDF.ush
--- a/Engine/Shaders/Private/BRDF.ush
+++ b/Engine/Shaders/Private/BRDF.ush
@@ -640,3 +640,38 @@ float Vis_Cloth( float NoV, float NoL )
 	return rcp( 4 * ( NoL + NoV - NoL * NoV ) );
 }
 
+float D_Charlie(float Roughness, float NoH)
+{
+	float InvR = 1 / Roughness;
+	float Cos2H = NoH * NoH;
+	float Sin2H = 1 - Cos2H;
+	return (2 + InvR) * pow(Sin2H, InvR * 0.5) / (2 * PI);
+}
+
+// [Estevez and Kulla 2017, "Production Friendly Microfacet Sheen BRDF"]
+float Vis_Charlie_L(float x, float r)
+{
+	r = saturate(r);
+	r = 1.0 - (1. - r) * (1. - r);
+
+	float a = lerp(25.3245 , 21.5473 , r);
+	float b = lerp( 3.32435,  3.82987, r);
+	float c = lerp( 0.16801,  0.19823, r);
+	float d = lerp(-1.27393, -1.97760, r);
+	float e = lerp(-4.85967, -4.32054, r);
+
+	return a * rcp( (1 + b * pow(x, c)) + d * x + e);
+}
+float Vis_Charlie(float Roughness, float NoV, float NoL)
+{
+	float VisV = NoV < 0.5 ? exp(Vis_Charlie_L(NoV, Roughness)) : exp(2 * Vis_Charlie_L(0.5, Roughness) - Vis_Charlie_L(1 - NoV, Roughness));
+	float VisL = NoL < 0.5 ? exp(Vis_Charlie_L(NoL, Roughness)) : exp(2 * Vis_Charlie_L(0.5, Roughness) - Vis_Charlie_L(1 - NoL, Roughness));
+
+	return rcp(((1 + VisV + VisL) * (4 * NoV * NoL)));
+}
+
+float Vis_Ashikhmin(float NoV, float NoL)
+{
+	return rcp(4 * (NoL + NoV - NoL * NoV));
+}
+
diff --git a/Engine/Shaders/Private/ShadingEnergyConservationTable.usf b/Engine/Shaders/Private/ShadingEnergyConservationTable.usf
--- a/Engine/Shaders/Private/ShadingEnergyConservationTable.usf
+++ b/Engine/Shaders/Private/ShadingEnergyConservationTable.usf
@@ -91,8 +91,14 @@ void BuildEnergyTableCS(uint3 DispatchThreadId : SV_DispatchThreadID)
 		const float NoL = L.z;
 		const float3 H = normalize(V + L);
 		const float NoH = H.z;
+#if STRATA_ENABLED
+		const float DCloth = D_Charlie(Roughness, NoH);
+		const float VCloth = Vis_Ashikhmin(NoV, NoL);
+#else
 		const float DCloth = D_InvGGX(Alpha.x * Alpha.y, NoH);
-		const float ClothWeight = NoL > 0 && NoV > 0 ? (PI * DCloth * Vis_Cloth(NoV, NoL)) : 0.0;
+		const float VCloth = Vis_Cloth(NoV, NoL);
+#endif
+		const float ClothWeight = NoL > 0 && NoV > 0 ? (PI * DCloth * VCloth) : 0.0;
 
 		Weight = float2(ClothWeight, ClothWeight * Pow5(1 - dot(V, H)));
 #elif BUILD_ENERGY_TABLE == 3
diff --git a/Engine/Shaders/Private/Strata/StrataEvaluation.ush b/Engine/Shaders/Private/Strata/StrataEvaluation.ush
--- a/Engine/Shaders/Private/Strata/StrataEvaluation.ush
+++ b/Engine/Shaders/Private/Strata/StrataEvaluation.ush
@@ -476,35 +476,49 @@ FStrataEvaluateResult StrataEvaluateBSDFCommon(FStrataBSDFContext BSDFContext, F
 			BRANCH
 			if (BSDF_GETHASFUZZ(BSDFContext.BSDF))
 			{
-				const float Fuzz = SLAB_FUZZ_AMOUNT(BSDFContext.BSDF);
+				const float FuzzAmount = SLAB_FUZZ_AMOUNT(BSDFContext.BSDF);
 				const float3 FuzzF0  = SLAB_FUZZ_COLOR(BSDFContext.BSDF);
-				
-				FBxDFEnergyTerms EnergyTerms = ComputeClothEnergyTerms(SafeRoughness, BSDFContext.Context.NoV, F0, F90);
 
-				// Generalized inverse microfacet specular
-				float ClothD   = D_InvGGX(Alpha2, BSDFContext.Context.NoH);
-				float ClothVis = Vis_Cloth(BSDFContext.Context.NoV, AreaLightContext.NoL);
+				const float ClothF90 = 1.0f;
+				FBxDFEnergyTerms EnergyTerms = ComputeClothEnergyTerms(SafeRoughness, BSDFContext.Context.NoV, FuzzF0, ClothF90);
+
+				// The specular and diffuse components below the fuzz are only attenuated linearly according to the amount of fuzz.
+				const float3 Cloth_DirectionalAlbedo_SpecularTransmission = lerp(1.0, ComputeEnergyPreservation(EnergyTerms), FuzzAmount);
+
+				// Charlie cloth modle for D and Ashikhmin for visibility (more efficient than Charlie).
+				float ClothD = D_Charlie(SafeRoughness, BSDFContext.Context.NoH);
+				float ClothVis = Vis_Ashikhmin(BSDFContext.Context.NoV, AreaLightContext.NoL);
 				float3 ClothF  = F_Schlick(FuzzF0, BSDFContext.Context.VoH);
 				float3 ClothSpecularPathValue = ClothD * ClothVis * ClothF * ComputeEnergyConservation(EnergyTerms);
 
 				// Area light are not supported by the cloth BRDF
-				float3 ClothIntegratedSpecularValue = ShadowTerms.SurfaceShadow * AreaLightContext.NoL * AreaLightContext.AreaLight.FalloffColor * AreaLightContext.Falloff * ClothSpecularPathValue;
+				float3 ClothIntegratedSpecularValue = ShadowTerms.SurfaceShadow * AreaLightContext.NoL * AreaLightContext.AreaLight.FalloffColor * AreaLightContext.Falloff * ClothSpecularPathValue * FuzzAmount;
 				
-				Sample.SpecularPathValue		= lerp(Sample.SpecularPathValue,		ClothSpecularPathValue,			Fuzz);
-				Sample.IntegratedSpecularValue	= lerp(Sample.IntegratedSpecularValue,	ClothIntegratedSpecularValue,	Fuzz);
+				// Apply specular transmittance to diffuse and specular lob from slab medium and interface sitting below the layer of fuzz.
+				Sample.DiffusePathValue			*= Cloth_DirectionalAlbedo_SpecularTransmission;
+				Sample.IntegratedDiffuseValue	*= Cloth_DirectionalAlbedo_SpecularTransmission;
+				Sample.SpecularPathValue		*= Cloth_DirectionalAlbedo_SpecularTransmission;
+				Sample.IntegratedSpecularValue	*= Cloth_DirectionalAlbedo_SpecularTransmission;
+				
+				Sample.SpecularPathValue		+= ClothSpecularPathValue;
+				Sample.IntegratedSpecularValue  += ClothIntegratedSpecularValue;
 
 				// This is not good. We should really have a separate PDF for cloth it self associated with a probability.
 				// We should also output probability to sample cloth, diffuse, secular, etc.
 				// We will revisit when the path tracer is getting up with Strata.
-				Sample.SpecularPDF		 = lerp(Sample.SpecularPDF,			BSDFContext.SatNoV / PI,	Fuzz); // Per "Production Friendly Microfacet Sheen BRDF", hemispherical sampling give good result as the roughness is usually high.
-				
-				// Fuzz does not affect throughput.
+				//Sample.SpecularPDF = lerp(Sample.SpecularPDF, BSDFContext.SatNoV / PI, FuzzAmount); // Per "Production Friendly Microfacet Sheen BRDF", hemispherical sampling give good result as the roughness is usually high.
+				// STRATA_TODO we should have a separated cloth PDF and Weight
 
 				if (bHaziness)
 				{
-					Sample.SpecularHazePathValue= lerp(Sample.SpecularHazePathValue,	ClothSpecularPathValue,		Fuzz);
-					Sample.SpecularHazePDF		= lerp(Sample.SpecularHazePDF,			BSDFContext.SatNoV / PI,	Fuzz);
+					Sample.SpecularHazePathValue*= Cloth_DirectionalAlbedo_SpecularTransmission;
+					Sample.SpecularHazePathValue+= ClothSpecularPathValue;
+
+					// STRATA_TODO we should have a separated cloth PDF and Weight
+					//Sample.SpecularHazePDF		= lerp(Sample.SpecularHazePDF,			BSDFContext.SatNoV / PI,	Fuzz);;
 				}
+
+				Sample.Throughput *= Cloth_DirectionalAlbedo_SpecularTransmission;
 				break;
 			}
 			#endif
@@ -1038,15 +1052,24 @@ FStrataEnvLightResult StrataEvaluateForEnvLight(FStrataBSDFContext BSDFContext,
 			BRANCH
 			if (BSDF_GETHASFUZZ(BSDFContext.BSDF))
 			{
-				const float Fuzz = SLAB_FUZZ_AMOUNT(BSDFContext.BSDF);
-				const float3 FuzzColor  = SLAB_FUZZ_COLOR(BSDFContext.BSDF);
-
-				// This should be a lerp, but SkyLightDiffuse from ReflectionEnvironmentPixelShader.usf add the contribution...
-				// STRATA_TODO fix the Fuzz model
-				//StrataEnvLightResult.DiffuseWeight = lerp(StrataEnvLightResult.DiffuseWeight, FuzzColor, Fuzz);
-				StrataEnvLightResult.DiffuseWeight += FuzzColor * Fuzz;
-				StrataEnvLightResult.DiffuseColor = saturate(StrataEnvLightResult.DiffuseColor + FuzzColor * Fuzz);
-				// Fuzz does not affect throughput.
+				const float FuzzAmount = SLAB_FUZZ_AMOUNT(BSDFContext.BSDF);
+				const float3 FuzzF0 = SLAB_FUZZ_COLOR(BSDFContext.BSDF);
+
+				const float ClothF90 = 1.0f;
+				FBxDFEnergyTerms ClothEnergyTerms = ComputeClothEnergyTerms(SafeRoughness, BSDFContext.Context.NoV, FuzzF0, ClothF90);
+
+				// The specular and diffuse components below the fuzz are only attenuated linearly according to the amount of fuzz.
+				const float3 Cloth_DirectionalAlbedo_SpecularTransmission = lerp(1.0, ComputeEnergyPreservation(ClothEnergyTerms), FuzzAmount);
+
+				StrataEnvLightResult.DiffuseColor		*= Cloth_DirectionalAlbedo_SpecularTransmission;
+				StrataEnvLightResult.DiffuseWeight		*= Cloth_DirectionalAlbedo_SpecularTransmission;
+				StrataEnvLightResult.SpecularWeight		*= Cloth_DirectionalAlbedo_SpecularTransmission;
+				StrataEnvLightResult.SpecularHazeWeight *= Cloth_DirectionalAlbedo_SpecularTransmission;
+				
+				// We currently send the cloth environment lighting to be taken into account by the diffuse term.
+				// STARTA_TODO: true a mix of diffuse and specular contribution according to roughness.
+				StrataEnvLightResult.DiffuseColor		+= ClothEnergyTerms.E * FuzzAmount;
+				StrataEnvLightResult.DiffuseWeight		+= ClothEnergyTerms.E * FuzzAmount;
 				break;
 			}
 			#endif
