Fixed issues with opaque rough refraction not matching reference (cannot scale linearly varaicne because it is a non linear value, it must be on standard deviation). Now we are in a better position.
Variance is now computed from ray intersection with the plane at normalised depth. Not from projected of unit ray vector onto the plane (perpendicularly).

diff --git a/Engine/Shaders/Private/Strata/Strata.ush b/Engine/Shaders/Private/Strata/Strata.ush
--- a/Engine/Shaders/Private/Strata/Strata.ush
+++ b/Engine/Shaders/Private/Strata/Strata.ush
@@ -586,9 +586,9 @@ FStrataOpaqueRoughRefractionData StrataUnpackOpaqueRoughRefractionData(float3 In
 }
 
 
-// Returns the variance of rays projected on the XY plane after scattering at the interface of a slab of matter,
+// Returns the variance of refracted rays intersection with on the XY plane after scattering at the interface of a slab of matter,
 // having the given Roughness and a normal along Z; after traveling along -Z from air to water (eta = 1.333).
-// That variance has been generated from unit vectors, so it is a measure of the spread of light on the projected XY plane.
+// That variance has been generated from unit vectors, so it is a measure of the spread of light on the XY plane after a depth of 1.
 // Using Thales, one can scale that spread to any distance. The Gaussian kernel can be evaluated for a depth d using a Variance = GetVarianceFromRoughness(Roughness) * d.
 // 
 //			| V = -Z
@@ -608,7 +608,7 @@ float GetVarianceFromRoughness(float Roughness)
 	float Roughness2 = Roughness * Roughness;
 	float Roughness3 = Roughness2 * Roughness;
 	float Roughness4 = Roughness2 * Roughness2;
-	return saturate(0.0000483367f - 0.000707943f * Roughness - 0.00840911f * Roughness2 + 0.095437f * Roughness3 - 0.0478129f * Roughness4);
+	return saturate(0.0000226575 - 0.000332564 * Roughness - 0.0168059 * Roughness2 + 0.115569 * Roughness3 - 0.0526336 * Roughness4);
 }
 
 ///////////////////////////////////////////////////////////////////////////////
@@ -2716,8 +2716,10 @@ void PackStrataOut(
 			uint RootOperatorIndex = Strata.OperatorIndex;
 			const FStrataOperator RootOperator = StrataTree.Operators[Strata.OperatorIndex];
 			OpaqueRoughRefractionData.Coverage = RootOperator.OpaqueRoughRefractCoverage;
-			// STRATA_TODO Fake variance place holder for now
-			OpaqueRoughRefractionData.VarianceCm = GetVarianceFromRoughness(RootOperator.OpaqueRoughRefractTopRoughness) * RootOperator.OpaqueRoughRefractThicknessCm;
+
+			const float StandardDeviationCm = sqrt(GetVarianceFromRoughness(RootOperator.OpaqueRoughRefractTopRoughness));
+			const float StandardDeviationCmForThickness = StandardDeviationCm * RootOperator.OpaqueRoughRefractThicknessCm;
+			OpaqueRoughRefractionData.VarianceCm = StandardDeviationCmForThickness * StandardDeviationCmForThickness;
 		}
 		
 		int SimpleMaterial_BSDFIndex = 0;
diff --git a/Engine/Shaders/Private/Strata/StrataRoughRefraction.usf b/Engine/Shaders/Private/Strata/StrataRoughRefraction.usf
--- a/Engine/Shaders/Private/Strata/StrataRoughRefraction.usf
+++ b/Engine/Shaders/Private/Strata/StrataRoughRefraction.usf
@@ -52,8 +52,12 @@ void OpaqueRoughRefractionPS(
 	const float2 PixelRadiusCm2 = SceneDepthCm * GetTanHalfFieldOfView() * View.ViewSizeAndInvSize.zw;
 	const float PixelRadiusCm = min(PixelRadiusCm2.x, PixelRadiusCm2.y);
 
+	const float StandardDeviationCm = sqrt(ORRData.VarianceCm);
+	const float StandardDeviationPixelSpace = StandardDeviationCm / PixelRadiusCm;	// Be careful to do that conversion on linear standard deviation, not variance
+	const float VariancePixelSpace = StandardDeviationPixelSpace * StandardDeviationPixelSpace;
+
 	const float PerceptualBlurAcceptanceFactor = 2.5f;
-	const bool OpaqueRoughRefractionEnabled = ORRData.OpaqueRoughRefractionEnabled && (PixelRadiusCm < PerceptualBlurAcceptanceFactor * ORRData.VarianceCm);
+	const bool OpaqueRoughRefractionEnabled = ORRData.OpaqueRoughRefractionEnabled;// && (PixelRadiusCm < PerceptualBlurAcceptanceFactor* ORRData.VarianceCm);
 
 	// We are now going to blur the bottom layer luminance according to the top layer variance resulting from the roughness.
 	// For this we have: measure GGX variance for a perpendicaular view, eta=1.33 and varying roughness. At runtim that variance is approximated by a Gaussian.
@@ -64,9 +68,6 @@ void OpaqueRoughRefractionPS(
 		float3 WeightedLuminance = 0.0f;
 		float AccumulatedWeights = 0.0f;
 
-		const float VariancePixelSpace = ORRData.VarianceCm / PixelRadiusCm;
-		const float StandardDeviationPixelSpace = sqrt(VariancePixelSpace);
-
 		// A area of radius StandardDeviation*3 around the mean of a guassian normal distribution (=mu) encompass 99% of the energy.
 		const float RadiusAroundCenterMu = 3.0f;
 		const float BlurStep = 0.5f;
@@ -321,10 +322,11 @@ void EvaluateRoughRefractionLobeCS(uint3 ThreadId : SV_DispatchThreadID)
 					float2 SampleUV = L.xy * 0.5 + 0.5;
 					SampleCountTextureUAV[uint2(SampleUV * 64.0f)] += 1;
 
-					LobStatUpdate(LobeState, L.xy);
+					const float3 RayOnPlane = L.xyz / abs(L.z);
+					LobStatUpdate(LobeState, RayOnPlane.xy);
 
 					Count++;
-					Mean += L.xy;
+					Mean += RayOnPlane.xy;
 				}
 			}
 		}
@@ -369,6 +371,7 @@ uint TraceSqrtSampleCount;
 
 RWBuffer<float2> RoughRefracDataUAV;
 
+// Shader used to export variance based on roughness for curve fitting
 [numthreads(THREADGROUP_SIZE, 1, 1)]
 void RoughRefracDataCS(uint3 ThreadId : SV_DispatchThreadID)
 {
@@ -388,7 +391,9 @@ void RoughRefracDataCS(uint3 ThreadId : SV_DispatchThreadID)
 			float Value = 0;
 			if (GetGGXDirection(RandomUV, Roughness, V, L, Value, Eta12) && Value > 0.0f)
 			{
-				LobStatUpdate(LobeState, L.xy);
+				// Project onto the plane at distance 1 to consider the blur resulting at a distance of 1.
+				const float3 RayOnPlane = L.xyz / abs(L.z);
+				LobStatUpdate(LobeState, RayOnPlane.xy);
 			}
 		}
 	}
@@ -579,6 +584,7 @@ void VisualizeRoughRefractionPS(
 		AddLineWS(Context, SlabOrigin + SlabEdgeY * Offset + SlabEdgeZ, SlabOrigin + SlabEdgeY * Offset + SlabEdgeX + SlabEdgeZ, float4(0.2, 0.2, 0.2, 1.0));
 	}
 
+	// Draw the rays
 	for (float u = 0.5f; u < TraceDomainSize; u++)
 	{
 		for (float v = 0.5f; v < TraceDomainSize; v++)
@@ -588,12 +594,14 @@ void VisualizeRoughRefractionPS(
 			float Value = 0;
 			if (GetGGXDirection(RandomUV, UI.Roughness, UI.V, L, Value, UI.Eta12))
 			{
-				if(Value > 0.0f) AddLineWS(Context, RefractionP, RefractionP + L * VisProjectedRayRadius, float4(0, Value, 0, 0.3));
-				//if(Value > 0.0f) AddLineWS(Context, RefractionP + L * VisProjectedRayRadius, float3((RefractionP + L * VisProjectedRayRadius).xy, SlabCenter.z - 0.5* UI.SlabThickness), float4(0, 1, 0, 0.3));
+				const float3 RayOnPlane = L.xyz / abs(L.z);
+				if(Value > 0.0f) AddLineWS(Context, RefractionP, RefractionP + RayOnPlane * VisProjectedRayRadius, float4(0, Value, 0, 0.3));
+				//if(Value > 0.0f) AddLineWS(Context, RefractionP + RayOnPlane * VisProjectedRayRadius, float3((RefractionP + L * VisProjectedRayRadius).xy, SlabCenter.z - 0.5* UI.SlabThickness), float4(0, 1, 0, 0.3));
 			}
 			else
 			{
-				if (Value > 0.0f) AddLineWS(Context, RefractionP, RefractionP + L * VisProjectedRayRadius, float4(Value, 0, 0, 0.3));
+				const float3 RayOnPlane = L.xyz / abs(L.z);
+				if (Value > 0.0f) AddLineWS(Context, RefractionP, RefractionP + RayOnPlane * VisProjectedRayRadius, float4(Value, 0, 0, 0.3));
 			}
 		}
 	}
