Strata - fix for compilation issue with DXC.

diff --git a/Engine/Shaders/Private/Strata/StrataMaterialClassification.usf b/Engine/Shaders/Private/Strata/StrataMaterialClassification.usf
--- a/Engine/Shaders/Private/Strata/StrataMaterialClassification.usf
+++ b/Engine/Shaders/Private/Strata/StrataMaterialClassification.usf
@@ -315,13 +315,14 @@ void BSDFTileMainCS(uint2 GroupThreadId : SV_GroupThreadID, uint2 GroupId : SV_G
 	{
 		FStrataAddressing StrataAddressing = GetStrataPixelDataByteOffset(PixelCoord, uint2(View.BufferSizeAndInvSize.xy), MaxBytesPerPixel);
 		FStrataPixelHeader StrataPixelHeader = UnpackStrataHeaderIn(MaterialTextureArray, StrataAddressing, TopLayerTexture);
-		BSDFCount = StrataPixelHeader.BSDFCount;
+		BSDFCount = min(StrataPixelHeader.BSDFCount, STRATA_TILE_MAX_BSDF_COUNT);
 
 		if (BSDFCount > 0)
 		{
 			FStrataBSDFOffset Offsets = (FStrataBSDFOffset)0;
 			Offsets.BSDFCount = BSDFCount;
 
+			UNROLL_N(STRATA_TILE_MAX_BSDF_COUNT) 
 			for (uint BSDFIndex = 0; BSDFIndex < BSDFCount; ++BSDFIndex)
 			{
 				Offsets.BSDFOffsets[BSDFIndex] = StrataAddressing.CurrentIndex;
diff --git a/Engine/Shaders/Private/Strata/StrataTile.ush b/Engine/Shaders/Private/Strata/StrataTile.ush
--- a/Engine/Shaders/Private/Strata/StrataTile.ush
+++ b/Engine/Shaders/Private/Strata/StrataTile.ush
@@ -73,10 +73,13 @@ FStrataBSDFTile UnpackBSDFTile(uint In)
 ////////////////////////////////////////////////////////////////////////////////////////////////////////////
 // Strata BSDF offset
 
+// STRATA_TODO MErge with incoming CL and simply make STRATA_MAX_BSDF_COUNT only 8 or 12 and remove STRATA_TILE_MAX_BSDF_COUNT
+#define STRATA_TILE_MAX_BSDF_COUNT 8
+
 // Store the material slice index offset for each BSDF, for pixels containing several BSDFs
 struct FStrataBSDFOffset
 {
-	uint BSDFOffsets[STRATA_MAX_BSDF_COUNT];
+	uint BSDFOffsets[STRATA_TILE_MAX_BSDF_COUNT];
 	uint BSDFCount;
 };
 
@@ -84,14 +87,14 @@ uint PackBSDFOffset(FStrataBSDFOffset In)
 {
 	// Packed as relative cumulative offsets, stored into 32 bits
 	// | 4 bits | 2 bits | 2 bits | 2 bits | ... | 2 bits |
-	//   BSDF 0   BSDF 1   BSDF 2   BSDF 3  ...   BSDF 15 (STRATA_MAX_BSDF_COUNT)
+	//   BSDF 0   BSDF 1   BSDF 2   BSDF 3  ...   BSDF 15 (STRATA_TILE_MAX_BSDF_COUNT)
 
 	uint Out = 0;
-	const uint BSDFCount = min(In.BSDFCount, STRATA_MAX_BSDF_COUNT);
+	const uint BSDFCount = min(In.BSDFCount, STRATA_TILE_MAX_BSDF_COUNT);
 	if (BSDFCount > 0)
 	{
 		Out = In.BSDFOffsets[0] & 0xF;
-		UNROLL
+		UNROLL_N(STRATA_TILE_MAX_BSDF_COUNT)
 		for (uint i = 1; i < BSDFCount; i++)
 		{
 			uint Curr = (In.BSDFOffsets[i  ] & 0x3);
@@ -107,11 +110,11 @@ FStrataBSDFOffset UnpackBSDFOffset(uint In, uint InBSDFCount)
 	FStrataBSDFOffset Out = (FStrataBSDFOffset)0;
 	if (InBSDFCount > 0)
 	{
-		Out.BSDFCount = min(InBSDFCount, STRATA_MAX_BSDF_COUNT);
+		Out.BSDFCount = min(InBSDFCount, STRATA_TILE_MAX_BSDF_COUNT);
 		Out.BSDFOffsets[0] = (In & 0xF);
 
 		FStrataBSDFOffset Out = (FStrataBSDFOffset)0;
-		UNROLL
+		UNROLL_N(STRATA_TILE_MAX_BSDF_COUNT)
 		for (uint i = 1; i < Out.BSDFCount; i++)
 		{
 			Out.BSDFOffsets[i] = Out.BSDFOffsets[i-1] + ((In>>(2u*i)) & 0x3);
