Fixed nan introduced by fog when it is used on materials used by DrawCanvas

diff --git a/Engine/Shaders/Private/HeightFogCommon.ush b/Engine/Shaders/Private/HeightFogCommon.ush
--- a/Engine/Shaders/Private/HeightFogCommon.ush
+++ b/Engine/Shaders/Private/HeightFogCommon.ush
@@ -88,7 +88,7 @@ half4 GetExponentialHeightFog(float3 WorldPositionRelativeToCamera, float Exclud
 	float3 CameraToReceiver = WorldPositionRelativeToCamera;
 	CameraToReceiver.z += LWCHackToFloat(PrimaryView.WorldCameraOrigin).z - WorldObserverOrigin.z; // Compensate this vector for clamping the observer height 
 	float CameraToReceiverLengthSqr = dot(CameraToReceiver, CameraToReceiver);
-	float CameraToReceiverLengthInv = rsqrt(CameraToReceiverLengthSqr);
+	float CameraToReceiverLengthInv = rsqrt(max(CameraToReceiverLengthSqr, 0.00000001f));
 	float CameraToReceiverLength = CameraToReceiverLengthSqr * CameraToReceiverLengthInv;
 	half3 CameraToReceiverNormalized = CameraToReceiver * CameraToReceiverLengthInv;
 
