Support fallback to coarser pages in directional light SMRT
- Check the page table at the ray start as a guess at what resolution is available; iterate in that clipmap
- Support fallback while iterating, albeit at a performance cost

For current mapping strategies, similar performance and comparable quality to the "trilinear-like" solution today
- Resolution changes can now appear more directly at page boundaries rather than explicitly across clipmap boundaries, so depending on the given situation you may see "blocky" versions of the range bands rather than the range bands themselves. Not really a better or worse artifact, but worth noting.

NOTE: Fallbacks to very coarse data should technically also involve significant optimal biases, but we do not want to compute an optimal bias per tap. For now the texel dither optimal bias is sufficient based on the clipmap we are iterating in. We could put a minimum clamp on this but there's a direct tradeoff with losing potentially important fine-scale shadows.

diff --git a/Engine/Shaders/Private/VirtualShadowMaps/VirtualShadowMapProjectionCommon.ush b/Engine/Shaders/Private/VirtualShadowMaps/VirtualShadowMapProjectionCommon.ush
--- a/Engine/Shaders/Private/VirtualShadowMaps/VirtualShadowMapProjectionCommon.ush
+++ b/Engine/Shaders/Private/VirtualShadowMaps/VirtualShadowMapProjectionCommon.ush
@@ -17,17 +17,16 @@ VirtualShadowMapProjectionCommon.ush:
 
 #define VIRTUAL_SHADOW_MAP_MIN_OCCLUDER_DISTANCE 1e-6f
 
-float CalcClipmapLevelFloat(FVirtualShadowMapProjectionShaderData BaseProjectionData, FLWCVector3 WorldPosition)
+float CalcAbsoluteClipmapLevel(FVirtualShadowMapProjectionShaderData BaseProjectionData, FLWCVector3 WorldPosition)
 {
 	float DistanceToClipmapOrigin = length(LWCToFloat(LWCSubtract(WorldPosition, BaseProjectionData.ClipmapWorldOrigin)));
-	float Log2Distance = log2(DistanceToClipmapOrigin);
-	float AbsoluteLevel = Log2Distance + BaseProjectionData.ClipmapResolutionLodBias;
-	return AbsoluteLevel;
+	return log2(DistanceToClipmapOrigin);
 }
 
 int CalcClipmapLevel(FVirtualShadowMapProjectionShaderData BaseProjectionData, FLWCVector3 WorldPosition)
 {
-	return int(floor(CalcClipmapLevelFloat(BaseProjectionData, WorldPosition)));
+	float BiasedLevel = CalcAbsoluteClipmapLevel(BaseProjectionData, WorldPosition) + BaseProjectionData.ClipmapResolutionLodBias;
+	return int(floor(BiasedLevel));
 }
 
 struct FVirtualShadowMapSample
@@ -149,13 +148,11 @@ FVirtualShadowMapSample SampleVirtualShadowMapClipmap(int VirtualShadowMapId, fl
 
 #define DEBUG_VSM_CLIPMAP_LEVEL_SEARCH 0
 #if !DEBUG_VSM_CLIPMAP_LEVEL_SEARCH
-	FShadowPageTranslationResult Page = ShadowVirtualToPhysicalUV(VirtualShadowMapId, ShadowMapUV);	
-
 	uint2 BasePage = uint2(ShadowMapUV * VSM_LEVEL0_DIM_PAGES_XY);
 	FShadowPhysicalPage PhysicalPageEntry = ShadowGetPhysicalPage(CalcPageOffset(VirtualShadowMapId, 0, BasePage));
 	if (PhysicalPageEntry.bAnyLODValid)
 	{	
-		uint ClipmapLevelOffset = Page.LODOffset;
+		uint ClipmapLevelOffset = PhysicalPageEntry.LODOffset;
 		int ClipmapLevelId = VirtualShadowMapId + ClipmapLevelOffset;
 
 		Result.VirtualTexelAddressFloat = ShadowMapUV * float(CalcLevelDimsTexels(0));
@@ -172,7 +169,7 @@ FVirtualShadowMapSample SampleVirtualShadowMapClipmap(int VirtualShadowMapId, fl
 			uint2 VirtualTexelAddressMin = vPage * VSM_PAGE_SIZE;
 			uint2 VirtualTexelAddressMax = VirtualTexelAddressMin + (VSM_PAGE_SIZE - 1);
 
-			FVirtualShadowMapClipmapRelativeTransform Transform = CalcClipmapRelativeTransform(VirtualShadowMapId, Page.LODOffset);
+			FVirtualShadowMapClipmapRelativeTransform Transform = CalcClipmapRelativeTransform(VirtualShadowMapId, ClipmapLevelOffset);
 			float2 ClipmapUV = ShadowMapUV * Transform.Scale + Transform.Bias.xy;
 			DepthLevelScale = Transform.Scale;
 			DepthLevelBias = Transform.Bias.z;
diff --git a/Engine/Shaders/Private/VirtualShadowMaps/VirtualShadowMapProjectionDirectional.ush b/Engine/Shaders/Private/VirtualShadowMaps/VirtualShadowMapProjectionDirectional.ush
--- a/Engine/Shaders/Private/VirtualShadowMaps/VirtualShadowMapProjectionDirectional.ush
+++ b/Engine/Shaders/Private/VirtualShadowMaps/VirtualShadowMapProjectionDirectional.ush
@@ -12,19 +12,7 @@ VirtualShadowMapProjectionDirectional.ush:
 #include "VirtualShadowMapPageAccessCommon.ush"
 #include "VirtualShadowMapProjectionCommon.ush"
 #include "VirtualShadowMapSMRTCommon.ush"
-
-
-struct FSMRTClipmapRayState
-{
-	int ClipmapId;			// Pre-offset virtual shadow map ID of specific clipmap level
-	int ClipmapIdAdjacent;	// Pre-offset virtual shadow map ID of specific clipmap level
-	float3 RayStartShadowTexelsZ;	
-	float3 RayStepShadowTexelsZ;
-	FVirtualShadowMapClipmapRelativeTransform BaseToAdjacentTexels;
-	uint2 PhysicalTexelAddress;
-	uint2 VirtualTexelAddress;
-};
-
+#include "../Visualization.ush"
 
 float ComputeOptimalSlopeBiasDirectional(
 	FVirtualShadowMapProjectionShaderData ProjectionData,
@@ -38,7 +26,7 @@ float ComputeOptimalSlopeBiasDirectional(
 	// Clamp to avoid excessive degenerate slope biases causing flickering lit pixels
 	float2 Clamp = 0.1f;
 	DepthSlopeUV = bClamp ? clamp(DepthSlopeUV, -Clamp, Clamp) : DepthSlopeUV;
-		
+
 	float2 OffsetUV = TexelOffset / float(CalcLevelDimsTexels(0));
 	float OptimalSlopeBias = 2.0f * max(0.0f, dot(DepthSlopeUV, OffsetUV));
 
@@ -46,7 +34,17 @@ float ComputeOptimalSlopeBiasDirectional(
 }
 
 
-FSMRTClipmapRayState SMRTClipmapRayInitialize(
+struct FSMRTClipmapRayState
+{
+	int ClipmapId;			// Pre-offset virtual shadow map ID of specific clipmap level
+	float3 RayStartUVZ;	
+	float3 RayStepUVZ;
+	uint2 PhysicalTexelAddress;
+	uint2 VirtualTexelAddress;
+	int SampledClipmapId;
+};
+
+FSMRTClipmapRayState SMRTClipmapWithFallbackRayInitialize(
 	int VirtualShadowMapId,
 	FLWCVector3 RayOriginWorld,
 	float3 RayDir,
@@ -56,31 +54,46 @@ FSMRTClipmapRayState SMRTClipmapRayInitialize(
 	float2 TexelDitherOffset)
 {
 	const FVirtualShadowMapProjectionShaderData BaseProjectionData = GetVirtualShadowMapProjectionData(VirtualShadowMapId);
+	
+	// This is the (unbiased) most detailed clipmap level that we are guaranteed to have a page table entry for
+	float ClipmapLevelFloat = CalcAbsoluteClipmapLevel(BaseProjectionData, RayOriginWorld);
+	int ClipmapIndex = max(0, int(floor(ClipmapLevelFloat)) - BaseProjectionData.ClipmapLevel);
 
-	// TODO: Refactor parts of this that only depend on RayOriginWorld outside the ray loop
-	// The page at RayOriginWorld has a high chance of being mapped since that is close or equal to to the receiver point
-	// Thus we start our clipmap sampling around that page even as we walk through the depth space
-	const float ClipmapLevelFloat = CalcClipmapLevelFloat(BaseProjectionData, RayOriginWorld);
-	const int ClipmapIndex = max(0, int(floor(ClipmapLevelFloat)) - BaseProjectionData.ClipmapLevel);
-	const float ClipmapLevelFrac = frac(ClipmapLevelFloat);
-	int ClipmapIndexAdjacent = ClipmapIndex + (ClipmapLevelFrac > 0.5f ? 1 : -1);
-	ClipmapIndexAdjacent = clamp(ClipmapIndexAdjacent, 0, BaseProjectionData.ClipmapLevelCount - 1);
+	// Look up the page table at the ray origin as a proxy for the available resolution
+	// We will iterate in the space of this clipmap level with the expectation that it is a
+	// good guess at nearby mapped pages.
+	{
+		// Clipmap origin should be near shaded samples, so shadow translated world should be regular range
+		int ClipmapId = VirtualShadowMapId + ClipmapIndex;
+		const FVirtualShadowMapProjectionShaderData ProjectionData = GetVirtualShadowMapProjectionData(ClipmapId);
+		float3 RayOriginShadowTranslatedWorld = LWCToFloat(LWCAdd(RayOriginWorld, ProjectionData.PreViewTranslation));
+		float3 RayStart = RayOriginShadowTranslatedWorld + RayDir * RayStartOffset;
+		float3 RayStartUVZ = mul(float4(RayStart, 1.0f), ProjectionData.TranslatedWorldToShadowUVMatrix).xyz;
+		
+		// NOTE: We don't need the actual sample here, we're just looking for the best mapped page at the ray origin
+		// Trusting a lot in DCE here... may be better to refactor at some point
+		FVirtualShadowMapSample ShadowSample = SampleVirtualShadowMapClipmap(ClipmapId, RayStartUVZ.xy);
+		if (ShadowSample.bValid && ShadowSample.VirtualShadowMapId > ClipmapId)
+		{
+			// NOTE: We could bias this (as long as we clamp it to 0 still) to allow sampling higher resolution
+			// data in the case where we just happened to hit a pixel with lower resolution, but note that performance
+			// suffers when we have to use fallback levels in the SMRT loop, so this is fine for now.
+			// We may have to revisit this if we start getting finer grained with caching and mapped page resolutions,
+			// but likely if this matters here we would already be seeing visible resolution seams between pages regardless.
+			ClipmapIndex += (ShadowSample.VirtualShadowMapId - ClipmapId);
+		}
+	}
 
 	const FVirtualShadowMapProjectionShaderData ProjectionData = GetVirtualShadowMapProjectionData(VirtualShadowMapId + ClipmapIndex);
 
 	// Clipmap origin should be near shaded samples, so shadow translated world should be regular range
 	float3 RayOriginShadowTranslatedWorld = LWCToFloat(LWCAdd(RayOriginWorld, ProjectionData.PreViewTranslation));
-
 	float3 RayStart = RayOriginShadowTranslatedWorld + RayDir * RayStartOffset;
 	float3 RayVector = RayDir * RayLength;
 
-	const float3 RayStartUV = mul(float4(RayStart, 1.0f), ProjectionData.TranslatedWorldToShadowUVMatrix).xyz;
-	const float3 RayStepUV  = mul(float4(RayVector, 0.0f), ProjectionData.TranslatedWorldToShadowUVMatrix).xyz;
-
-	const float LevelDims = float(CalcLevelDimsTexels(0));
-	float3 RayStartShadowTexelsZ = float3(LevelDims, LevelDims, 1.0f) * RayStartUV.xyz;
-	float3 RayStepShadowTexelsZ  = float3(LevelDims, LevelDims, 1.0f) * RayStepUV.xyz;
-
+	float3 RayStartUVZ = mul(float4(RayStart, 1.0f), ProjectionData.TranslatedWorldToShadowUVMatrix).xyz;
+	float3 RayStepUVZ  = mul(float4(RayVector, 0.0f), ProjectionData.TranslatedWorldToShadowUVMatrix).xyz;
+	
 	// Texel dither to hide aliasing
 	// Note that this is directly scaled in texel space of the clipmap level, so it can create visual
 	// discontinuities at clipmap boundaries. We could scale this by distance instead of texel size but
@@ -92,22 +105,15 @@ FSMRTClipmapRayState SMRTClipmapRayInitialize(
 		float OptimalBias = ComputeOptimalSlopeBiasDirectional(ProjectionData, TexelDitherOffset, WorldNormal);
 		// Subtract off any portion of the bias that was already covered by the ray start offset (usually contact shadows)
 		OptimalBias = max(0.0f, OptimalBias - abs(RayStartOffset * ProjectionData.ShadowViewToClipMatrix._33));
-		
-		RayStartShadowTexelsZ.xy += TexelDitherOffset;
-		RayStartShadowTexelsZ.z += OptimalBias;
+
+		RayStartUVZ.xy += TexelDitherOffset / float(CalcLevelDimsTexels(0));
+		RayStartUVZ.z += OptimalBias;
 	}
-		
-	// We are working in texels rather than UVs, so need to scale bias appropriately
-	FVirtualShadowMapClipmapRelativeTransform BaseToAdjacentTexels =
-		CalcClipmapRelativeTransform(VirtualShadowMapId + ClipmapIndex, ClipmapIndexAdjacent - ClipmapIndex);
-	BaseToAdjacentTexels.Bias.xy *= LevelDims;
 
 	FSMRTClipmapRayState Result;
 	Result.ClipmapId = VirtualShadowMapId + ClipmapIndex;
-	Result.ClipmapIdAdjacent = VirtualShadowMapId + ClipmapIndexAdjacent;
-	Result.RayStartShadowTexelsZ = RayStartShadowTexelsZ;
-	Result.RayStepShadowTexelsZ = RayStepShadowTexelsZ;
-	Result.BaseToAdjacentTexels = BaseToAdjacentTexels;
+	Result.RayStartUVZ = RayStartUVZ;
+	Result.RayStepUVZ = RayStepUVZ;
 	Result.VirtualTexelAddress = uint2(0xFFFFFFFF, 0xFFFFFFFF);
 	Result.PhysicalTexelAddress = uint2(0xFFFFFFFF, 0xFFFFFFFF);
 	return Result;
@@ -118,36 +124,20 @@ FSMRTSample SMRTFindSample(inout FSMRTClipmapRayState RayState, float SampleTime
 	FSMRTSample Sample = (FSMRTSample)0;
 	Sample.bValid = false;
 	
-	const float3 SampleTexelsZ = RayState.RayStartShadowTexelsZ + RayState.RayStepShadowTexelsZ * SampleTime;
-
-	uint2 PhysicalTexelAddress = uint2(0, 0);
-	uint2 VirtualTexelAddress = uint2(SampleTexelsZ.xy);
-	bool bValid = VirtualToPhysicalTexel(RayState.ClipmapId, 0, VirtualTexelAddress, PhysicalTexelAddress);
-
-	float2 ToBaseDepthScaleBias = float2(1.0f, 0.0f);
-	if (!bValid)
-	{
-		// Replace with adjacent level sample
-		float2 AdjacentTexels = SampleTexelsZ.xy * RayState.BaseToAdjacentTexels.Scale + RayState.BaseToAdjacentTexels.Bias.xy;
-		VirtualTexelAddress = uint2(AdjacentTexels.xy);
-		bValid = VirtualToPhysicalTexel(RayState.ClipmapIdAdjacent, 0, VirtualTexelAddress, PhysicalTexelAddress);
-
-		// These are to bring the ray back into the base depth range
-		ToBaseDepthScaleBias.x = 1.0f / RayState.BaseToAdjacentTexels.Scale;
-		ToBaseDepthScaleBias.y = - RayState.BaseToAdjacentTexels.Bias.z * ToBaseDepthScaleBias.x;
-	}
-
-	if (bValid)
+	const float3 SampleUVZ = RayState.RayStartUVZ + RayState.RayStepUVZ * SampleTime;
+	FVirtualShadowMapSample ShadowSample = SampleVirtualShadowMapClipmap(RayState.ClipmapId, SampleUVZ.xy);
+	
+	if (ShadowSample.bValid)
 	{
-		float SampleDepth = ToBaseDepthScaleBias.x * SampleVirtualShadowMapPhysicalDepth(PhysicalTexelAddress) + ToBaseDepthScaleBias.y;
 		Sample.bValid = true;
-		Sample.SampleDepth = SampleDepth;
-		Sample.ReferenceDepth = SampleTexelsZ.z;
-		Sample.CompareToleranceScale = RayState.RayStepShadowTexelsZ.z;
+		Sample.SampleDepth = ShadowSample.Depth;
+		Sample.ReferenceDepth = SampleUVZ.z;
+		Sample.CompareToleranceScale = RayState.RayStepUVZ.z;
 
 		// Debug
-		RayState.VirtualTexelAddress = VirtualTexelAddress;
-		RayState.PhysicalTexelAddress = PhysicalTexelAddress;
+		RayState.VirtualTexelAddress = ShadowSample.VirtualTexelAddress;
+		RayState.PhysicalTexelAddress = ShadowSample.PhysicalTexelAddress;
+		RayState.SampledClipmapId = ShadowSample.VirtualShadowMapId;
 	}
 
 	return Sample;
@@ -158,8 +148,8 @@ float SMRTComputeOccluderDistance(FSMRTClipmapRayState RayState, float SampleDep
 	float4x4 ViewToClip = GetVirtualShadowMapProjectionData(RayState.ClipmapId).ShadowViewToClipMatrix;
 
 	// No perspective projection, so simple difference gets us the distance
-	float OccluderViewZ = (SampleDepth                      - ViewToClip._43) / ViewToClip._33;
-	float ReceiverViewZ = (RayState.RayStartShadowTexelsZ.z - ViewToClip._43) / ViewToClip._33;
+	float OccluderViewZ = (SampleDepth            - ViewToClip._43) / ViewToClip._33;
+	float ReceiverViewZ = (RayState.RayStartUVZ.z - ViewToClip._43) / ViewToClip._33;
 	float Result = ReceiverViewZ - OccluderViewZ;
 	return max(VIRTUAL_SHADOW_MAP_MIN_OCCLUDER_DISTANCE, Result);
 }
@@ -169,6 +159,7 @@ float SMRTComputeOccluderDistance(FSMRTClipmapRayState RayState, float SampleDep
 #include "VirtualShadowMapSMRTTemplate.ush"
 #undef SMRT_TEMPLATE_RAY_STRUCT
 
+
 float3 GetRandomDirectionalLightRayDir(FLightShaderParameters Light, float2 E)
 {
 	float3 RayDir = Light.Direction;
@@ -239,15 +230,16 @@ FVirtualShadowMapSampleResult TraceDirectional(
 
 			float2 TexelDitherOffset = (RandSample.zw - 0.5f) * TexelDitherScale;
 
-			FSMRTClipmapRayState RayState = SMRTClipmapRayInitialize(VirtualShadowMapId, WorldPosition, RayDir, RayLength, RayStartOffset, WorldNormal, TexelDitherOffset);
+			FSMRTClipmapRayState RayState = SMRTClipmapWithFallbackRayInitialize(VirtualShadowMapId, WorldPosition, RayDir, RayLength, RayStartOffset, WorldNormal, TexelDitherOffset);
 			FSMRTResult SMRTResult = SMRTRayCast(RayState, SamplesPerRay, StepOffset);
 
-			RayMissCount += SMRTResult.bValidHit ? 0 : 1;
-
 			// Debug output (DCE'd if not used)
-			Result.ClipmapIndexOrMipLevel = RayState.ClipmapId - VirtualShadowMapId;
+			Result.ClipmapIndexOrMipLevel = RayState.SampledClipmapId - VirtualShadowMapId;
 			Result.VirtualTexelAddress = RayState.VirtualTexelAddress;
 			Result.PhysicalTexelAddress =  RayState.PhysicalTexelAddress;
+			//Result.GeneralDebug = GreenToRedTurbo((RayState.SampledClipmapId - RayState.ClipmapId) / 2.0f);
+			
+			RayMissCount += SMRTResult.bValidHit ? 0 : 1;
 
 	#if SMRT_ADAPTIVE_RAY_COUNT && COMPILER_SUPPORTS_WAVE_VOTE
 			if( i == 0 )
