Added option to disable usage of a baked landscape mesh on mobile platforms (r.Mobile.LandscapeMesh)

diff --git a/Engine/Plugins/Experimental/GPULightmass/Source/GPULightmass/Private/Scene/Landscape.h b/Engine/Plugins/Experimental/GPULightmass/Source/GPULightmass/Private/Scene/Landscape.h
--- a/Engine/Plugins/Experimental/GPULightmass/Source/GPULightmass/Private/Scene/Landscape.h
+++ b/Engine/Plugins/Experimental/GPULightmass/Source/GPULightmass/Private/Scene/Landscape.h
@@ -32,6 +32,7 @@ public:
 		FVector4f    HeightmapScaleBias;
 		FVector4f    WeightmapScaleBias;
 		float       WeightmapSubsectionOffset;
+		bool		bUseMobileLandscapeMesh;
 	};
 
 	struct FLandscapeSectionRayTracingState
diff --git a/Engine/Plugins/Experimental/GPULightmass/Source/GPULightmass/Private/Scene/Scene.cpp b/Engine/Plugins/Experimental/GPULightmass/Source/GPULightmass/Private/Scene/Scene.cpp
--- a/Engine/Plugins/Experimental/GPULightmass/Source/GPULightmass/Private/Scene/Scene.cpp
+++ b/Engine/Plugins/Experimental/GPULightmass/Source/GPULightmass/Private/Scene/Scene.cpp
@@ -1241,9 +1241,10 @@ void FScene::AddGeometryInstanceFromComponent(ULandscapeComponent* InComponent)
 	InstanceRenderState.LocalBounds = InComponent->CalcBounds(FTransform::Identity);
 	InstanceRenderState.bCastShadow = InComponent->CastShadow && InComponent->bCastStaticShadow;
 
+	const bool bUseMobileLandscapeMesh = UseMobileLandscapeMesh(GShaderPlatformForFeatureLevel[FeatureLevel]);
 	const int8 SubsectionSizeLog2 = FMath::CeilLogTwo(InComponent->SubsectionSizeQuads + 1);
 	InstanceRenderState.SharedBuffersKey = (SubsectionSizeLog2 & 0xf) | ((InComponent->NumSubsections & 0xf) << 4) |
-		(FeatureLevel <= ERHIFeatureLevel::ES3_1 ? 0 : (1 << 30)) | (InComponent->XYOffsetmapTexture == nullptr ? 0 : 1 << 31);
+		(bUseMobileLandscapeMesh ? 0 : (1 << 30)) | (InComponent->XYOffsetmapTexture == nullptr ? 0 : 1 << 31);
 	InstanceRenderState.SharedBuffersKey |= 1 << 29; // Use this bit to indicate it is GPULightmass specific buffer (which only has FixedGridVertexFactory created)
 
 	TArray<UMaterialInterface*> AvailableMaterials;
@@ -1279,6 +1280,7 @@ void FScene::AddGeometryInstanceFromComponent(ULandscapeComponent* InComponent)
 	Initializer.HeightmapScaleBias         = (FVector4f)InComponent->HeightmapScaleBias;
 	Initializer.WeightmapScaleBias         = (FVector4f)InComponent->WeightmapScaleBias;
 	Initializer.WeightmapSubsectionOffset  = InComponent->WeightmapSubsectionOffset;
+	Initializer.bUseMobileLandscapeMesh	   = bUseMobileLandscapeMesh;
 
 	TArray<int32> RelevantPointLightsToAddOnRenderThread = AddAllPossiblyRelevantLightsToGeometry(LightScene.PointLights, Instance);
 	TArray<int32> RelevantSpotLightsToAddOnRenderThread = AddAllPossiblyRelevantLightsToGeometry(LightScene.SpotLights, Instance);
@@ -1301,7 +1303,7 @@ void FScene::AddGeometryInstanceFromComponent(ULandscapeComponent* InComponent)
 		{
 			InstanceRenderState.SharedBuffers = new FLandscapeSharedBuffers(
 				InstanceRenderState.SharedBuffersKey, Initializer.SubsectionSizeQuads, Initializer.NumSubsections,
-				LocalFeatureLevel);
+				LocalFeatureLevel, Initializer.bUseMobileLandscapeMesh);
 
 			FLandscapeComponentSceneProxy::SharedBuffersMap.Add(InstanceRenderState.SharedBuffersKey, InstanceRenderState.SharedBuffers);
 
diff --git a/Engine/Shaders/Private/LandscapeVertexFactory.ush b/Engine/Shaders/Private/LandscapeVertexFactory.ush
--- a/Engine/Shaders/Private/LandscapeVertexFactory.ush
+++ b/Engine/Shaders/Private/LandscapeVertexFactory.ush
@@ -9,6 +9,11 @@
 #define VERTEX_FACTORY_MODIFIES_INTERPOLATION 1
 #define VERTEX_FACTORY_HEIGHTFIELD 1
 
+#ifndef MOBILE_LANDSCAPE_MESH
+#define MOBILE_LANDSCAPE_MESH 0
+#endif
+
+
 // Set by FLightMapTexturePolicy
 #include "/Engine/Generated/UniformBuffers/PrecomputedLightingBuffer.ush"
 
@@ -16,7 +21,7 @@
 	#define XYOFFSET_SCALE (1.0f/256.f)
 #endif
 
-#if (ES3_1_PROFILE)
+#if MOBILE_LANDSCAPE_MESH
 /* Offset for UV localization for large UV values. */
 float2 TexCoordOffset;
 #endif
@@ -65,7 +70,7 @@ float4 GetLodValues(uint ComponentIndex)
 #if FIXED_GRID 
 	// Fixed grid has fixed Lod
 	return LandscapeFixedGrid.LodValues;
-#elif FEATURE_LEVEL < FEATURE_LEVEL_SM4
+#elif MOBILE_LANDSCAPE_MESH
 	// On mobile we always use the Lod 0 mesh data
 	return float4(0, 0, LandscapeParameters.SubsectionSizeVertsLayerUVPan.xy);
 #else
@@ -82,7 +87,7 @@ float4 GetLodValues(uint ComponentIndex)
  */
 float2 GetLodBias(uint ComponentIndex)
 {
-#if FIXED_GRID || FEATURE_LEVEL < FEATURE_LEVEL_SM4
+#if FIXED_GRID || MOBILE_LANDSCAPE_MESH
 	// Fixed grid and mobile don't use LodBias
 	return 0;
 #else
@@ -159,15 +164,14 @@ float CalcLOD(uint ComponentIndex, float2 xyLocalToSubsection, float2 Subsection
 
 struct FVertexFactoryInput
 {
-#if FEATURE_LEVEL >= FEATURE_LEVEL_SM4
-	float4 Position: ATTRIBUTE0;
-
+#if MOBILE_LANDSCAPE_MESH
+	float4 PackedPosition : ATTRIBUTE0;
+	float4 LODHeights0 : ATTRIBUTE1;
+	float4 LODHeights1 : ATTRIBUTE2;
+#else
+	float4 Position : ATTRIBUTE0;
 	// Dynamic instancing related attributes with InstanceIdOffset : ATTRIBUTE1
 	VF_GPUSCENE_DECLARE_INPUT_BLOCK(1)
-#else
-	float4 PackedPosition: ATTRIBUTE0; 
-	float4 LODHeights0 : ATTRIBUTE1;  
-	float4 LODHeights1 : ATTRIBUTE2;
 #endif
 };
 
@@ -245,7 +249,7 @@ uint GetNumRayTracingDynamicMeshVerticesIndirect()
 struct FVertexFactoryInterpolantsVSToPS
 {
 	float2	LayerTexCoord			: TEXCOORD0; // xy == texcoord
-#if (ES3_1_PROFILE)
+#if MOBILE_LANDSCAPE_MESH
 	float2	WeightMapTexCoord		: TEXCOORD1;
 	float4	TransformedTexCoord0	: TEXCOORD2;
 	float4  TransformedTexCoord1	: TEXCOORD3;
@@ -254,7 +258,7 @@ struct FVertexFactoryInterpolantsVSToPS
 	float4	TransformedTexCoords	: TEXCOORD2;
 #endif
 #if NEEDS_LIGHTMAP_COORDINATE
-#if (ES3_1_PROFILE)
+#if FEATURE_LEVEL == FEATURE_LEVEL_ES3_1
 	float2	LightMapCoordinate[2]	: TEXCOORD4;
 	float2	ShadowMapCoordinate		: TEXCOORD6;
 #else
@@ -275,7 +279,7 @@ struct FLandscapeTexCoords
 {
 	float2	LayerTexCoord : TEXCOORD0; // xy == texcoord
 	float2	WeightMapTexCoord;
-#if !(ES3_1_PROFILE)
+#if !MOBILE_LANDSCAPE_MESH
 	float2	HeightMapTexCoord;
 #endif
 #if NEEDS_LIGHTMAP_COORDINATE
@@ -300,7 +304,7 @@ struct FVertexFactoryRayTracingInterpolants
 #if NEEDS_VERTEX_FACTORY_INTERPOLATION
 float2 VertexFactoryGetRayTracingTextureCoordinate( FVertexFactoryRayTracingInterpolants Interpolants )
 {
-#if (ES3_1_PROFILE)
+#if MOBILE_LANDSCAPE_MESH
 	return Interpolants.InterpolantsVSToPS.WeightMapTexCoord.xy;
 #else
 	return Interpolants.InterpolantsVSToPS.WeightHeightMapTexCoord.zw;
@@ -370,7 +374,7 @@ FLandscapeTexCoords GetLandscapeTexCoords(FVertexFactoryInput Input, FVertexFact
 
 	Result.WeightMapTexCoord = Intermediates.LocalPosition.xy * LandscapeParameters.WeightmapUVScaleBias.xy + LandscapeParameters.WeightmapUVScaleBias.zw + Intermediates.InputPosition.zw * LandscapeParameters.SubsectionOffsetParams.zz;
 
-#if !(ES3_1_PROFILE)
+#if !MOBILE_LANDSCAPE_MESH
 	Result.HeightMapTexCoord = Intermediates.LocalPosition.xy * LandscapeParameters.HeightmapUVScaleBias.xy + LandscapeParameters.HeightmapUVScaleBias.zw + 0.5*LandscapeParameters.HeightmapUVScaleBias.xy + Intermediates.InputPosition.zw * LandscapeParameters.SubsectionOffsetParams.xy;
 #endif
 
@@ -396,10 +400,10 @@ float3x3 VertexFactoryGetPerPixelTangentBasis(FVertexFactoryInterpolantsVSToPS I
 {
 	float3x3 Result;
 #if PIXELSHADER || RAYHITGROUPSHADER || COMPUTESHADER
-#if FEATURE_LEVEL >= FEATURE_LEVEL_SM4
-	float4 SampleValue = Texture2DSample(LandscapeParameters.NormalmapTexture, LandscapeParameters.NormalmapTextureSampler, Interpolants.WeightHeightMapTexCoord.zw);
-#else
+#if MOBILE_LANDSCAPE_MESH
 	float4 SampleValue = Texture2DSample(LandscapeParameters.NormalmapTexture, LandscapeParameters.NormalmapTextureSampler, Interpolants.WeightMapTexCoord);
+#else
+	float4 SampleValue = Texture2DSample(LandscapeParameters.NormalmapTexture, LandscapeParameters.NormalmapTextureSampler, Interpolants.WeightHeightMapTexCoord.zw);
 #endif
 #if WEBGL
 	float2 SampleNormal = float2(SampleValue.b, SampleValue.a) * float2(2.0,2.0) - float2(1.0,1.0);
@@ -420,7 +424,7 @@ FMaterialPixelParameters GetMaterialPixelParameters(FVertexFactoryInterpolantsVS
 	FMaterialPixelParameters Result = MakeInitializedMaterialPixelParameters();
 
 #if NEEDS_LIGHTMAP_COORDINATE
-	#if (ES3_1_PROFILE)
+	#if FEATURE_LEVEL == FEATURE_LEVEL_ES3_1
 		// Not supported in pixel shader
 		float2 LightmapUVs = float2(0, 0);
 	#else
@@ -433,19 +437,19 @@ FMaterialPixelParameters GetMaterialPixelParameters(FVertexFactoryInterpolantsVS
 #if NUM_MATERIAL_TEXCOORDS     // XY layer
 	Result.TexCoords[0] = Interpolants.LayerTexCoord.xy;
 #if NUM_MATERIAL_TEXCOORDS > 1 // VS calcualted TexCoord 1, default is XZ layer
-	#if (ES3_1_PROFILE)
+	#if MOBILE_LANDSCAPE_MESH
 		Result.TexCoords[1] = Interpolants.TransformedTexCoord0.xy;
 	#else
 		Result.TexCoords[1] = Interpolants.TransformedTexCoords.xy;
 	#endif
 #if NUM_MATERIAL_TEXCOORDS > 2 // VS calcualted TexCoord 2, default is YZ layer
-	#if (ES3_1_PROFILE)
+	#if MOBILE_LANDSCAPE_MESH
 		Result.TexCoords[2] = Interpolants.TransformedTexCoord1.xy;
 	#else
 		Result.TexCoords[2] = Interpolants.TransformedTexCoords.zw;
 	#endif
 #if NUM_MATERIAL_TEXCOORDS > 3 // Weightmap
-	#if (ES3_1_PROFILE)
+	#if MOBILE_LANDSCAPE_MESH
 		Result.TexCoords[3] = Interpolants.WeightMapTexCoord;
 	#else
 		Result.TexCoords[3] = Interpolants.WeightHeightMapTexCoord.xy;
@@ -457,7 +461,7 @@ FMaterialPixelParameters GetMaterialPixelParameters(FVertexFactoryInterpolantsVS
 		Result.TexCoords[4] = float2(0,0);
 	#endif
 #if NUM_MATERIAL_TEXCOORDS > 5 // Height map
-	#if (ES3_1_PROFILE)
+	#if MOBILE_LANDSCAPE_MESH
 		Result.TexCoords[5] = float2(0,0); // Height map UV not supported
 	#else
 		Result.TexCoords[5] = Interpolants.WeightHeightMapTexCoord.zw;
@@ -486,7 +490,7 @@ FMaterialPixelParameters GetMaterialPixelParameters(FVertexFactoryInterpolantsVS
 	Result.LightmapUVs = LightmapUVs;
 #endif
 
-#if (ES3_1_PROFILE) && PIXELSHADER
+#if MOBILE_LANDSCAPE_MESH && PIXELSHADER
 	Result.LayerWeights.xy = Texture2DSample(LandscapeParameters.NormalmapTexture, LandscapeParameters.NormalmapTextureSampler, Interpolants.WeightMapTexCoord).rg;
 	Result.LayerWeights.z = saturate(LandscapeParameters.BlendableLayerMask.z - dot(Result.LayerWeights.xy, LandscapeParameters.BlendableLayerMask.xy));
 	Result.LayerWeights.w = 0; // Total weight need to be 1
@@ -519,7 +523,7 @@ FMaterialVertexParameters GetMaterialVertexParameters(FVertexFactoryInput Input,
 
 #if NUM_MATERIAL_TEXCOORDS_VERTEX     // XY layer
 	Result.TexCoords[0] = LandscapeTexCoords.LayerTexCoord.xy;
-	#if (ES3_1_PROFILE)
+	#if MOBILE_LANDSCAPE_MESH
 		Result.TexCoordOffset = TexCoordOffset;
 	#endif
 #if NUM_MATERIAL_TEXCOORDS_VERTEX > 1 // XZ layer
@@ -535,7 +539,7 @@ FMaterialVertexParameters GetMaterialVertexParameters(FVertexFactoryInput Input,
 		Result.TexCoords[4] = float2(0,0);
 	#endif
 #if NUM_MATERIAL_TEXCOORDS_VERTEX > 5 // Height map
-	#if (ES3_1_PROFILE)
+	#if MOBILE_LANDSCAPE_MESH
 		Result.TexCoords[5] = float2(0,0); // Height map UV not supported
 	#else
 		Result.TexCoords[5] = LandscapeTexCoords.HeightMapTexCoord;
@@ -560,7 +564,7 @@ FMaterialVertexParameters GetMaterialVertexParameters(FVertexFactoryInput Input,
 #if NEEDS_LIGHTMAP_COORDINATE
 void GetLightMapCoordinates(FVertexFactoryInterpolantsVSToPS Interpolants, out float2 LightmapUV0, out float2 LightmapUV1, out uint LightmapDataIndex)
 {
-#if (ES3_1_PROFILE)
+#if FEATURE_LEVEL == FEATURE_LEVEL_ES3_1
 	LightmapUV0 = Interpolants.LightMapCoordinate[0].xy;
 	LightmapUV1 = Interpolants.LightMapCoordinate[1].xy;
 #else
@@ -577,7 +581,7 @@ void GetLightMapCoordinates(FVertexFactoryInterpolantsVSToPS Interpolants, out f
 
 void GetShadowMapCoordinate(FVertexFactoryInterpolantsVSToPS Interpolants, out float2 ShadowMapCoordinate, out uint LightmapDataIndex)
 {
-#if (ES3_1_PROFILE)
+#if FEATURE_LEVEL == FEATURE_LEVEL_ES3_1
 	ShadowMapCoordinate = Interpolants.ShadowMapCoordinate.xy;
 #else
 	ShadowMapCoordinate = Interpolants.LightMapCoordinate.zw;
@@ -600,9 +604,7 @@ FVertexFactoryIntermediates GetVertexFactoryIntermediates(FVertexFactoryInput In
 	Intermediates.LodValues = GetLodValues(Intermediates.ComponentIndex);
 	Intermediates.LodBias = GetLodBias(Intermediates.ComponentIndex);
 
-#if FEATURE_LEVEL >= FEATURE_LEVEL_SM4
-	Intermediates.InputPosition = Input.Position;
-#else
+#if MOBILE_LANDSCAPE_MESH
 	Intermediates.InputPosition.xy = Input.PackedPosition.xy * 255.f;
 
 	uint PackedExtraData = (uint)(Input.PackedPosition.z * 255.f);
@@ -611,6 +613,8 @@ FVertexFactoryIntermediates GetVertexFactoryIntermediates(FVertexFactoryInput In
 	Intermediates.InputPosition.zw = float2(SubX, SubY);
 	float HoleMaxLod = (float)(PackedExtraData >> 5);
 	float HoleVertexLod = (float)((PackedExtraData >> 2) & 7);
+#else
+	Intermediates.InputPosition = Input.Position;
 #endif
 
 	float2 xy = Intermediates.InputPosition.xy * Intermediates.LodValues.w;
@@ -620,11 +624,7 @@ FVertexFactoryIntermediates GetVertexFactoryIntermediates(FVertexFactoryInput In
 	float MorphAlpha = LODCalculated - LodValue;
 
 	// InputPositionLODAdjusted : Position for actual LOD in base LOD units
-#if FEATURE_LEVEL >= FEATURE_LEVEL_SM4
-	float2 ActualLODCoordsInt = floor(Intermediates.InputPosition.xy * pow(2, -(LodValue - Intermediates.LodValues.x)));
-	float InvLODScaleFactor = pow(2, -LodValue);
-#else
-
+#if MOBILE_LANDSCAPE_MESH
 	// Clamp lod values around holes until we reach HoleMaxLod. Kill morphing when clamped.
 	bool bMorph = LodValue < HoleVertexLod || HoleMaxLod <= LodValue;
 	MorphAlpha = bMorph ? MorphAlpha : 0;
@@ -650,6 +650,9 @@ FVertexFactoryIntermediates GetVertexFactoryIntermediates(FVertexFactoryInput In
 	float InvLODScaleFactor = 1.f / (float)(1 << (uint)LodValue);
 	// LodValues.x is always 0 on mobile.
 	float2 ActualLODCoordsInt = floor(Intermediates.InputPosition.xy * InvLODScaleFactor);
+#else
+	float2 ActualLODCoordsInt = floor(Intermediates.InputPosition.xy * pow(2, -(LodValue - Intermediates.LodValues.x)));
+	float InvLODScaleFactor = pow(2, -LodValue);
 #endif
 
 	// Base to Actual LOD, Base to Next LOD
@@ -661,7 +664,7 @@ FVertexFactoryIntermediates GetVertexFactoryIntermediates(FVertexFactoryInput In
 	float2 InputPositionNextLOD = NextLODCoordsInt / CoordTranslate.y;
 
 	// Get the height and normal XY for current and next LOD out of the textures
-#if FEATURE_LEVEL >= FEATURE_LEVEL_SM4
+#if !MOBILE_LANDSCAPE_MESH
 	float2 SampleCoords = InputPositionLODAdjusted * LandscapeParameters.HeightmapUVScaleBias.xy + LandscapeParameters.HeightmapUVScaleBias.zw + 0.5*LandscapeParameters.HeightmapUVScaleBias.xy + Intermediates.InputPosition.zw * LandscapeParameters.SubsectionOffsetParams.xy;
 	float4 SampleValue = Texture2DSampleLevel(LandscapeParameters.HeightmapTexture, LandscapeParameters.HeightmapTextureSampler, SampleCoords, LodValue-Intermediates.LodBias.x);
 	float Height = DecodePackedHeight(SampleValue.xy);
@@ -686,13 +689,13 @@ FVertexFactoryIntermediates GetVertexFactoryIntermediates(FVertexFactoryInput In
 
 	Intermediates.LocalPosition = lerp( float3(InputPositionLODAdjusted, Height), float3(InputPositionNextLOD, HeightNextLOD), MorphAlpha );
 
-#if FEATURE_LEVEL >= FEATURE_LEVEL_SM4
+#if MOBILE_LANDSCAPE_MESH
+	Intermediates.WorldNormal = float3(0.0, 0.0, 1.0);
+#else
 	float2 Normal = float2(SampleValue.b, SampleValue.a);
 	float2 NormalNextLOD = float2(SampleValueNextLOD.b, SampleValueNextLOD.a);
 	float2 InterpNormal = lerp( Normal, NormalNextLOD, MorphAlpha ) * float2(2.0,2.0) - float2(1.0,1.0);
 	Intermediates.WorldNormal = float3( InterpNormal, sqrt(max(1.0-dot(InterpNormal,InterpNormal),0.0)) );
-#else
-	Intermediates.WorldNormal = float3( 0.0, 0.0, 1.0 );
 #endif
 	
 	return Intermediates;
@@ -729,7 +732,7 @@ FVertexFactoryInterpolantsVSToPS VertexFactoryGetInterpolantsVSToPS(FVertexFacto
 	Interpolants = (FVertexFactoryInterpolantsVSToPS)0;
 	FLandscapeTexCoords LandscapeTexCoords = GetLandscapeTexCoords(Input, Intermediates);
 
-#if (ES3_1_PROFILE)
+#if MOBILE_LANDSCAPE_MESH
 	Interpolants.LayerTexCoord	= LandscapeTexCoords.LayerTexCoord;
 	Interpolants.WeightMapTexCoord	= LandscapeTexCoords.WeightMapTexCoord; 
 	#if NUM_MATERIAL_TEXCOORDS
@@ -764,10 +767,10 @@ FVertexFactoryInterpolantsVSToPS VertexFactoryGetInterpolantsVSToPS(FVertexFacto
 		#endif // 1
 
 	#endif // NUM_MATERIAL_TEXCOORDS
-#endif // (ES3_1_PROFILE)
+#endif // MOBILE_LANDSCAPE_MESH
 
 #if NEEDS_LIGHTMAP_COORDINATE
-	#if (ES3_1_PROFILE)
+	#if FEATURE_LEVEL == FEATURE_LEVEL_ES3_1
 		Interpolants.LightMapCoordinate[0].xy = LandscapeTexCoords.LightMapCoordinate * PrecomputedLightingBuffer.LightMapCoordinateScaleBias.xy + PrecomputedLightingBuffer.LightMapCoordinateScaleBias.zw;
 		Interpolants.LightMapCoordinate[0].y *= 0.5;
 		Interpolants.LightMapCoordinate[1].xy = Interpolants.LightMapCoordinate[0].xy;
@@ -823,12 +826,12 @@ FVertexFactoryRayTracingInterpolants VertexFactoryInterpolate(FVertexFactoryRayT
 {
 	FVertexFactoryRayTracingInterpolants O;
 	
-#if NEEDS_LIGHTMAP_COORDINATE && !(ES3_1_PROFILE)
+#if NEEDS_LIGHTMAP_COORDINATE && FEATURE_LEVEL > FEATURE_LEVEL_ES3_1
 	INTERPOLATE_MEMBER(InterpolantsVSToPS.LightMapCoordinate);
 #endif
 
 	INTERPOLATE_MEMBER(InterpolantsVSToPS.LayerTexCoord);
-#if (ES3_1_PROFILE)
+#if MOBILE_LANDSCAPE_MESH
 	INTERPOLATE_MEMBER(InterpolantsVSToPS.WeightMapTexCoord);
 	INTERPOLATE_MEMBER(InterpolantsVSToPS.TransformedTexCoord0);
 	INTERPOLATE_MEMBER(InterpolantsVSToPS.TransformedTexCoord1);
diff --git a/Engine/Shaders/Private/MaterialTemplate.ush b/Engine/Shaders/Private/MaterialTemplate.ush
--- a/Engine/Shaders/Private/MaterialTemplate.ush
+++ b/Engine/Shaders/Private/MaterialTemplate.ush
@@ -417,7 +417,7 @@ struct FMaterialPixelParameters
 	/** Per-particle properties. Only valid for particle vertex factories. */
 	FMaterialParticleParameters Particle;
 
-#if ES3_1_PROFILE
+#if FEATURE_LEVEL == FEATURE_LEVEL_ES3_1
 	float4 LayerWeights;
 #endif
 
@@ -426,7 +426,7 @@ struct FMaterialPixelParameters
 	FTexCoordScalesParams TexCoordScalesParams;
 #endif
 
-#if POST_PROCESS_MATERIAL && (FEATURE_LEVEL <= FEATURE_LEVEL_ES3_1)
+#if POST_PROCESS_MATERIAL && FEATURE_LEVEL == FEATURE_LEVEL_ES3_1
 	/** Used in mobile custom pp material to preserve original SceneColor Alpha */
 	half BackupSceneColorAlpha;
 #endif
@@ -526,7 +526,7 @@ struct FMaterialVertexParameters
 	half4 VertexColor;
 #if NUM_MATERIAL_TEXCOORDS_VERTEX
 	float2 TexCoords[NUM_MATERIAL_TEXCOORDS_VERTEX];
-	#if ES3_1_PROFILE
+	#if FEATURE_LEVEL == FEATURE_LEVEL_ES3_1
 	float2 TexCoordOffset; // Offset for UV localization for large UV values
 	#endif
 #endif
diff --git a/Engine/Source/Developer/Android/AndroidTargetPlatform/Private/AndroidTargetPlatform.cpp b/Engine/Source/Developer/Android/AndroidTargetPlatform/Private/AndroidTargetPlatform.cpp
--- a/Engine/Source/Developer/Android/AndroidTargetPlatform/Private/AndroidTargetPlatform.cpp
+++ b/Engine/Source/Developer/Android/AndroidTargetPlatform/Private/AndroidTargetPlatform.cpp
@@ -323,6 +323,19 @@ bool FAndroidTargetPlatform::SupportsLandscapeMeshLODStreaming() const
 	return bStreamLandscapeMeshLODs;
 }
 
+bool FAndroidTargetPlatform::UseMobileLandscapeMesh() const
+{
+	// By default mobile uses landscape mesh
+	static bool bUseMobileLandscapeMesh = true;
+	static bool bInitialized = false;
+	if (!bInitialized)
+	{
+		GetConfigSystem()->GetBool(TEXT("/Script/Engine.RendererSettings"), TEXT("r.Mobile.LandscapeMesh"), bUseMobileLandscapeMesh, GEngineIni);
+		bInitialized = true;
+	}
+	return bUseMobileLandscapeMesh;
+}
+
 /* ITargetPlatform overrides
  *****************************************************************************/
 
@@ -419,6 +432,9 @@ bool FAndroidTargetPlatform::SupportsFeature( ETargetPlatformFeatures Feature )
 
 		case ETargetPlatformFeatures::DistanceFieldAO:
 			return UsesDistanceFields();
+
+		case ETargetPlatformFeatures::MobileLandscapeMesh:
+			return SupportsFeature(ETargetPlatformFeatures::MobileRendering) && UseMobileLandscapeMesh();
 			
 		default:
 			break;
diff --git a/Engine/Source/Developer/Android/AndroidTargetPlatform/Private/AndroidTargetPlatform.h b/Engine/Source/Developer/Android/AndroidTargetPlatform/Private/AndroidTargetPlatform.h
--- a/Engine/Source/Developer/Android/AndroidTargetPlatform/Private/AndroidTargetPlatform.h
+++ b/Engine/Source/Developer/Android/AndroidTargetPlatform/Private/AndroidTargetPlatform.h
@@ -185,6 +185,7 @@ protected:
 	bool SupportsVulkan() const;
 	bool SupportsLandscapeMeshLODStreaming() const;
 	bool SupportsVulkanSM5() const;
+	bool UseMobileLandscapeMesh() const;
 
 protected:
 
diff --git a/Engine/Source/Developer/IOS/IOSTargetPlatform/Private/IOSTargetPlatform.cpp b/Engine/Source/Developer/IOS/IOSTargetPlatform/Private/IOSTargetPlatform.cpp
--- a/Engine/Source/Developer/IOS/IOSTargetPlatform/Private/IOSTargetPlatform.cpp
+++ b/Engine/Source/Developer/IOS/IOSTargetPlatform/Private/IOSTargetPlatform.cpp
@@ -464,6 +464,20 @@ bool FIOSTargetPlatform::CanSupportRemoteShaderCompile() const
 	return false; // !bRemoteCompilingEnabled;
 }
 
+bool FIOSTargetPlatform::UseMobileLandscapeMesh() const
+{
+	// By default mobile uses landscape mesh
+	static bool bUseMobileLandscapeMesh = true;
+	static bool bInitialized = false;
+
+	if (!bInitialized)
+	{
+		GetConfigSystem()->GetBool(TEXT("/Script/Engine.RendererSettings"), TEXT("r.Mobile.LandscapeMesh"), bUseMobileLandscapeMesh, GEngineIni);
+		bInitialized = true;
+	}
+	return bUseMobileLandscapeMesh;
+}
+
 bool FIOSTargetPlatform::SupportsFeature( ETargetPlatformFeatures Feature ) const
 {
 	switch (Feature)
@@ -489,6 +503,9 @@ bool FIOSTargetPlatform::SupportsFeature( ETargetPlatformFeatures Feature ) cons
 		case ETargetPlatformFeatures::DistanceFieldAO:
 			return UsesDistanceFields();
 
+		case ETargetPlatformFeatures::MobileLandscapeMesh:
+			return SupportsFeature(ETargetPlatformFeatures::MobileRendering) && UseMobileLandscapeMesh();
+
 		default:
 			break;
 	}
diff --git a/Engine/Source/Developer/IOS/IOSTargetPlatform/Private/IOSTargetPlatform.h b/Engine/Source/Developer/IOS/IOSTargetPlatform/Private/IOSTargetPlatform.h
--- a/Engine/Source/Developer/IOS/IOSTargetPlatform/Private/IOSTargetPlatform.h
+++ b/Engine/Source/Developer/IOS/IOSTargetPlatform/Private/IOSTargetPlatform.h
@@ -114,6 +114,8 @@ public:
 	{
 		return bDistanceField;
 	}
+
+	bool UseMobileLandscapeMesh() const;
 private:
 
 	// Handles received pong messages from the LauncherDaemon.
diff --git a/Engine/Source/Developer/TargetPlatform/Public/Interfaces/ITargetPlatform.h b/Engine/Source/Developer/TargetPlatform/Public/Interfaces/ITargetPlatform.h
--- a/Engine/Source/Developer/TargetPlatform/Public/Interfaces/ITargetPlatform.h
+++ b/Engine/Source/Developer/TargetPlatform/Public/Interfaces/ITargetPlatform.h
@@ -110,6 +110,9 @@ enum class ETargetPlatformFeatures
 
 	/** The platform supports communication (reading and writing data) between a target a connected PC. */
 	DirectDataExchange,
+
+	/** Whether platform uses a baked landscape mesh */
+	MobileLandscapeMesh,
 };
 
 enum class EPlatformAuthentication
diff --git a/Engine/Source/Developer/Windows/WindowsTargetPlatform/Private/GenericWindowsTargetPlatform.h b/Engine/Source/Developer/Windows/WindowsTargetPlatform/Private/GenericWindowsTargetPlatform.h
--- a/Engine/Source/Developer/Windows/WindowsTargetPlatform/Private/GenericWindowsTargetPlatform.h
+++ b/Engine/Source/Developer/Windows/WindowsTargetPlatform/Private/GenericWindowsTargetPlatform.h
@@ -207,6 +207,19 @@ public:
 			return bCachedSupportsMobileRendering;
 		}
 
+		if (Feature == ETargetPlatformFeatures::MobileLandscapeMesh)
+		{
+			// By default mobile uses landscape mesh
+			static bool bUseMobileLandscapeMesh = true;
+			static bool bInitialized = false;
+			if (!bInitialized)
+			{
+				this->GetConfigSystem()->GetBool(TEXT("/Script/Engine.RendererSettings"), TEXT("r.Mobile.LandscapeMesh"), bUseMobileLandscapeMesh, GEngineIni);
+				bInitialized = true;
+			}
+			return bUseMobileLandscapeMesh && SupportsFeature(ETargetPlatformFeatures::MobileRendering);
+		}
+		
 		return TSuper::SupportsFeature(Feature);
 	}
 
diff --git a/Engine/Source/Runtime/Engine/Private/Materials/HLSLMaterialTranslator.cpp b/Engine/Source/Runtime/Engine/Private/Materials/HLSLMaterialTranslator.cpp
--- a/Engine/Source/Runtime/Engine/Private/Materials/HLSLMaterialTranslator.cpp
+++ b/Engine/Source/Runtime/Engine/Private/Materials/HLSLMaterialTranslator.cpp
@@ -7103,7 +7103,8 @@ bool FHLSLMaterialTranslator::GetStaticBoolValue(int32 BoolIndex, bool& bSucceed
 
 int32 FHLSLMaterialTranslator::StaticTerrainLayerWeight(FName LayerName,int32 Default)
 {
-	if (GetFeatureLevel() <= ERHIFeatureLevel::ES3_1 && ShaderFrequency != SF_Pixel)
+	const bool bUseMobileLandscapeMesh = UseMobileLandscapeMesh(Platform);
+	if (bUseMobileLandscapeMesh && ShaderFrequency != SF_Pixel)
 	{
 		return Errorf(TEXT("Landscape layer weights are only available in the pixel shader."));
 	}
@@ -7143,7 +7144,7 @@ int32 FHLSLMaterialTranslator::StaticTerrainLayerWeight(FName LayerName,int32 De
 	else
 	{			
 		int32 WeightmapCode;
-		if (GetFeatureLevel() <= ERHIFeatureLevel::ES3_1 && NumActiveTerrainLayerWeightParameters <= 3 && bAtLeastOneWeightBasedBlend)
+		if (bUseMobileLandscapeMesh && NumActiveTerrainLayerWeightParameters <= 3 && bAtLeastOneWeightBasedBlend)
 		{
 			// Mobile can pack 3 layers into the normal map texture B and A channels, implying the 3rd using weight based blending
 			// Layer texture is sampled into Parameters.LayerWeights in LandscapeVertexFactory.ush
@@ -11192,7 +11193,7 @@ int32 FHLSLMaterialTranslator::SpeedTree(int32 GeometryArg, int32 WindArg, int32
 	*/
 int32 FHLSLMaterialTranslator::TextureCoordinateOffset()
 {
-	if (FeatureLevel < ERHIFeatureLevel::SM5 && ShaderFrequency == SF_Vertex)
+	if (UseMobileLandscapeMesh(Platform) && ShaderFrequency == SF_Vertex)
 	{
 		return AddInlinedCodeChunkZeroDeriv(MCT_Float2, TEXT("Parameters.TexCoordOffset"));
 	}
diff --git a/Engine/Source/Runtime/Landscape/Private/Landscape.cpp b/Engine/Source/Runtime/Landscape/Private/Landscape.cpp
--- a/Engine/Source/Runtime/Landscape/Private/Landscape.cpp
+++ b/Engine/Source/Runtime/Landscape/Private/Landscape.cpp
@@ -131,6 +131,11 @@ namespace LandscapeCookStats
 // Set this to 0 to disable landscape cooking and thus disable it on device.
 #define ENABLE_LANDSCAPE_COOKING 1
 
+static bool UseMobileLandscapeMesh(const ITargetPlatform* TargetPlatform)
+{
+	return TargetPlatform->SupportsFeature(ETargetPlatformFeatures::MobileLandscapeMesh);
+}
+
 #define LOCTEXT_NAMESPACE "Landscape"
 
 static void PrintNumLandscapeShadows()
@@ -252,7 +257,7 @@ void ULandscapeComponent::BeginCacheForCookedPlatformData(const ITargetPlatform*
 {
 	Super::BeginCacheForCookedPlatformData(TargetPlatform);
 
-	if (TargetPlatform->SupportsFeature(ETargetPlatformFeatures::MobileRendering) && !HasAnyFlags(RF_ClassDefaultObject))
+	if (UseMobileLandscapeMesh(TargetPlatform) && !HasAnyFlags(RF_ClassDefaultObject))
 	{
 		CheckGenerateLandscapePlatformData(true, TargetPlatform);
 	}
@@ -350,7 +355,7 @@ void ULandscapeComponent::Serialize(FArchive& Ar)
 	Ar.UsingCustomVersion(FEditorObjectVersion::GUID);
 
 #if WITH_EDITOR
-	if (Ar.IsCooking() && !HasAnyFlags(RF_ClassDefaultObject) && Ar.CookingTarget()->SupportsFeature(ETargetPlatformFeatures::MobileRendering))
+	if (Ar.IsCooking() && !HasAnyFlags(RF_ClassDefaultObject) && UseMobileLandscapeMesh(Ar.CookingTarget()))
 	{
 		// for -oldcook:
 		// the old cooker calls BeginCacheForCookedPlatformData after the package export set is tagged, so the mobile material doesn't get saved, so we have to do CheckGenerateLandscapePlatformData in serialize
@@ -421,30 +426,33 @@ void ULandscapeComponent::Serialize(FArchive& Ar)
 			Ar.Serialize(Object, sizeof(UObject*));
 		}
 	}
-	else if (Ar.IsCooking() && !HasAnyFlags(RF_ClassDefaultObject) && !Ar.CookingTarget()->SupportsFeature(ETargetPlatformFeatures::DeferredRendering))
+	else if (Ar.IsCooking() && !HasAnyFlags(RF_ClassDefaultObject))
 	{
-		// These properties are only used for SM4+ so we back them up and clear them before serializing them.
-		UTexture2D* BackupHeightmapTexture = nullptr;
-		UTexture2D* BackupXYOffsetmapTexture = nullptr;
-		TArray<UMaterialInstanceConstant*> BackupMaterialInstances;
-		TArray<UTexture2D*> BackupWeightmapTextures;
+		const bool bUseMobileLandscapeMesh = UseMobileLandscapeMesh(Ar.CookingTarget());
+		
+		if (bUseMobileLandscapeMesh && !Ar.CookingTarget()->SupportsFeature(ETargetPlatformFeatures::DeferredRendering))
+		{
+			// These are used for SM5 rendering or if MobileLandscapeMesh is disabled 
+			UTexture2D* BackupHeightmapTexture = nullptr;
+			UTexture2D* BackupXYOffsetmapTexture = nullptr;
+			TArray<UMaterialInstanceConstant*> BackupMaterialInstances;
+			TArray<UTexture2D*> BackupWeightmapTextures;
 
-		Exchange(HeightmapTexture, BackupHeightmapTexture);
-		Exchange(BackupXYOffsetmapTexture, XYOffsetmapTexture);
-		Exchange(BackupMaterialInstances, MaterialInstances);
-		Exchange(BackupWeightmapTextures, WeightmapTextures);
+			Exchange(HeightmapTexture, BackupHeightmapTexture);
+			Exchange(BackupXYOffsetmapTexture, XYOffsetmapTexture);
+			Exchange(BackupMaterialInstances, MaterialInstances);
+			Exchange(BackupWeightmapTextures, WeightmapTextures);
 
-		Super::Serialize(Ar);
+			Super::Serialize(Ar);
 
-		Exchange(HeightmapTexture, BackupHeightmapTexture);
-		Exchange(BackupXYOffsetmapTexture, XYOffsetmapTexture);
-		Exchange(BackupMaterialInstances, MaterialInstances);
-		Exchange(BackupWeightmapTextures, WeightmapTextures);
-	}
-	else
-		if (Ar.IsCooking() && !HasAnyFlags(RF_ClassDefaultObject) && !Ar.CookingTarget()->SupportsFeature(ETargetPlatformFeatures::MobileRendering))
+			Exchange(HeightmapTexture, BackupHeightmapTexture);
+			Exchange(BackupXYOffsetmapTexture, XYOffsetmapTexture);
+			Exchange(BackupMaterialInstances, MaterialInstances);
+			Exchange(BackupWeightmapTextures, WeightmapTextures);
+		}
+		else if (!bUseMobileLandscapeMesh)
 		{
-			// These properties are only used for mobile so we back them up and clear them before serializing them.
+			// These properties are only when MobileLandscapeMesh is enabled so we back them up and clear them before serializing them.
 			TArray<UMaterialInterface*> BackupMobileMaterialInterfaces;
 			TArray<UTexture2D*> BackupMobileWeightmapTextures;
 
@@ -457,10 +465,16 @@ void ULandscapeComponent::Serialize(FArchive& Ar)
 			Exchange(MobileWeightmapTextures, BackupMobileWeightmapTextures);
 		}
 		else
-#endif
 		{
+			// Serialize both mobile landscape mesh and heightmap properties
 			Super::Serialize(Ar);
 		}
+	}
+	else
+#endif
+	{
+		Super::Serialize(Ar);
+	}
 
 	if (Ar.IsLoading() && Ar.CustomVer(FRenderingObjectVersion::GUID) < FRenderingObjectVersion::MapBuildDataSeparatePackage)
 	{
@@ -540,7 +554,7 @@ void ULandscapeComponent::Serialize(FArchive& Ar)
 #if ENABLE_LANDSCAPE_COOKING
 	if (bCooked)
 	{
-		bool bCookedMobileData = Ar.IsCooking() && Ar.CookingTarget()->SupportsFeature(ETargetPlatformFeatures::MobileRendering);
+		bool bCookedMobileData = Ar.IsCooking() && UseMobileLandscapeMesh(Ar.CookingTarget());
 		Ar << bCookedMobileData;
 
 		// Saving for cooking path
@@ -981,7 +995,7 @@ void ULandscapeComponent::PostLoad()
 		}
 		else
 		{
-			if (GMaxRHIFeatureLevel <= ERHIFeatureLevel::ES3_1)
+			if (UseMobileLandscapeMesh(GMaxRHIShaderPlatform))
 			{
 				UE_LOG(LogLandscape, Error, TEXT("Landscape component (%d, %d) Does not have a valid mobile combination material. To correct this issue, open the map in the editor and resave the map."), SectionBaseX, SectionBaseY);
 			}
@@ -992,9 +1006,11 @@ void ULandscapeComponent::PostLoad()
 	if (!HasAnyFlags(RF_ClassDefaultObject))
 	{
 		UWorld* World = GetWorld();
-
+		ERHIFeatureLevel::Type FeatureLevel = ((GEngine->GetDefaultWorldFeatureLevel() == ERHIFeatureLevel::ES3_1) || (World && (World->FeatureLevel <= ERHIFeatureLevel::ES3_1))) 
+			? ERHIFeatureLevel::ES3_1 : GMaxRHIFeatureLevel;
+				
 		// If we're loading on a platform that doesn't require cooked data, but defaults to a mobile feature level, generate or preload data from the DDC
-		if (!FPlatformProperties::RequiresCookedData() && ((GEngine->GetDefaultWorldFeatureLevel() <= ERHIFeatureLevel::ES3_1) || (World && (World->FeatureLevel <= ERHIFeatureLevel::ES3_1))))
+		if (!FPlatformProperties::RequiresCookedData() && UseMobileLandscapeMesh(GShaderPlatformForFeatureLevel[FeatureLevel]))
 		{
 			CheckGenerateLandscapePlatformData(false, nullptr);
 		}
@@ -1426,7 +1442,7 @@ FPrimitiveSceneProxy* ULandscapeComponent::CreateSceneProxy()
 
 	const auto FeatureLevel = GetWorld()->FeatureLevel;
 	FPrimitiveSceneProxy* Proxy = nullptr;
-	if (FeatureLevel >= ERHIFeatureLevel::SM5)
+	if (FeatureLevel >= ERHIFeatureLevel::SM5 || !UseMobileLandscapeMesh(GShaderPlatformForFeatureLevel[FeatureLevel]))
 	{
 		Proxy = new FLandscapeComponentSceneProxy(this);
 	}
@@ -2182,7 +2198,7 @@ void ALandscapeProxy::OnFeatureLevelChanged(ERHIFeatureLevel::Type NewFeatureLev
 
 	UpdateAllComponentMaterialInstances();
 
-	if (NewFeatureLevel <= ERHIFeatureLevel::ES3_1)
+	if (UseMobileLandscapeMesh(GShaderPlatformForFeatureLevel[NewFeatureLevel]))
 	{
 		for (ULandscapeComponent* Component : LandscapeComponents)
 		{
diff --git a/Engine/Source/Runtime/Landscape/Private/LandscapeEdit.cpp b/Engine/Source/Runtime/Landscape/Private/LandscapeEdit.cpp
--- a/Engine/Source/Runtime/Landscape/Private/LandscapeEdit.cpp
+++ b/Engine/Source/Runtime/Landscape/Private/LandscapeEdit.cpp
@@ -647,7 +647,7 @@ void ULandscapeComponent::PreFeatureLevelChange(ERHIFeatureLevel::Type PendingFe
 {
 	Super::PreFeatureLevelChange(PendingFeatureLevel);
 
-	if (PendingFeatureLevel <= ERHIFeatureLevel::ES3_1)
+	if (UseMobileLandscapeMesh(GShaderPlatformForFeatureLevel[PendingFeatureLevel]))
 	{
 		// See if we need to cook platform data for mobile preview in editor
 		CheckGenerateLandscapePlatformData(false, nullptr);
@@ -5153,7 +5153,7 @@ void ALandscapeStreamingProxy::PostEditChangeProperty(FPropertyChangedEvent& Pro
 
 			UWorld* World = GetWorld();
 
-			if (World != nullptr && World->FeatureLevel <= ERHIFeatureLevel::ES3_1)
+			if (World != nullptr && UseMobileLandscapeMesh(GShaderPlatformForFeatureLevel[World->FeatureLevel]))
 			{
 				for (ULandscapeComponent * Component : LandscapeComponents)
 				{
@@ -5557,7 +5557,7 @@ void ALandscape::PostEditChangeProperty(FPropertyChangedEvent& PropertyChangedEv
 
 				UWorld* World = GetWorld();
 
-				if (World != nullptr && World->FeatureLevel <= ERHIFeatureLevel::ES3_1)
+				if (World != nullptr && UseMobileLandscapeMesh(GShaderPlatformForFeatureLevel[World->FeatureLevel]))
 				{
 					for (ULandscapeComponent * Component : LandscapeComponents)
 					{
@@ -5701,7 +5701,7 @@ void ULandscapeComponent::PostEditChangeProperty(FPropertyChangedEvent& Property
 
 			UWorld* World = GetWorld();
 
-			if (World != nullptr && World->FeatureLevel <= ERHIFeatureLevel::ES3_1)
+			if (World != nullptr && UseMobileLandscapeMesh(GShaderPlatformForFeatureLevel[World->FeatureLevel]))
 			{
 				CheckGenerateLandscapePlatformData(false, nullptr);
 			}
diff --git a/Engine/Source/Runtime/Landscape/Private/LandscapeRender.cpp b/Engine/Source/Runtime/Landscape/Private/LandscapeRender.cpp
--- a/Engine/Source/Runtime/Landscape/Private/LandscapeRender.cpp
+++ b/Engine/Source/Runtime/Landscape/Private/LandscapeRender.cpp
@@ -1115,7 +1115,13 @@ FLandscapeComponentSceneProxy::FLandscapeComponentSceneProxy(ULandscapeComponent
 	EnableGPUSceneSupportFlags();
 
 	const ERHIFeatureLevel::Type FeatureLevel = GetScene().GetFeatureLevel();
-	if (FeatureLevel >= ERHIFeatureLevel::SM5)
+	const bool bMobileLandscapeMesh = UseMobileLandscapeMesh(GetScene().GetShaderPlatform());
+
+	if (bMobileLandscapeMesh)
+	{
+		AvailableMaterials.Append(InComponent->MobileMaterialInterfaces);
+	}
+	else
 	{
 		if (InComponent->GetLandscapeProxy()->bUseDynamicMaterialInstance)
 		{
@@ -1126,17 +1132,14 @@ FLandscapeComponentSceneProxy::FLandscapeComponentSceneProxy(ULandscapeComponent
 			AvailableMaterials.Append(InComponent->MaterialInstances);
 		}
 	}
-	else
-	{
-		AvailableMaterials.Append(InComponent->MobileMaterialInterfaces);
-	}
+
 
 	LODIndexToMaterialIndex = InComponent->LODIndexToMaterialIndex;
 	check(LODIndexToMaterialIndex.Num() == MaxLOD + 1);
 
 	SetLevelColor(FLinearColor(1.f, 1.f, 1.f));
-
-	if (FeatureLevel <= ERHIFeatureLevel::ES3_1)
+			
+	if (bMobileLandscapeMesh)
 	{
 		HeightmapTexture = nullptr;
 		HeightmapSubsectionOffsetU = 0;
@@ -1291,12 +1294,12 @@ FLandscapeComponentSceneProxy::FLandscapeComponentSceneProxy(ULandscapeComponent
 		}
 	}
 #endif
-
+	
 	const int8 SubsectionSizeLog2 = FMath::CeilLogTwo(InComponent->SubsectionSizeQuads + 1);
 	SharedBuffersKey = (SubsectionSizeLog2 & 0xf) | ((NumSubsections & 0xf) << 4) |
-		(FeatureLevel <= ERHIFeatureLevel::ES3_1 ? 0 : 1 << 30) | (XYOffsetmapTexture == nullptr ? 0 : 1 << 31);
+		(bMobileLandscapeMesh ? 0 : 1 << 30) | (XYOffsetmapTexture == nullptr ? 0 : 1 << 31);
 
-	bSupportsHeightfieldRepresentation = FeatureLevel > ERHIFeatureLevel::ES3_1;
+	bSupportsHeightfieldRepresentation = !bMobileLandscapeMesh;
 	bSupportsMeshCardRepresentation = true;
 
 #if WITH_EDITOR
@@ -1345,7 +1348,7 @@ void FLandscapeComponentSceneProxy::CreateRenderThreadResources()
 	{
 		SharedBuffers = new FLandscapeSharedBuffers(
 			SharedBuffersKey, SubsectionSizeQuads, NumSubsections,
-			FeatureLevel);
+			FeatureLevel, false);
 
 		FLandscapeComponentSceneProxy::SharedBuffersMap.Add(SharedBuffersKey, SharedBuffers);
 
@@ -2651,14 +2654,11 @@ void FLandscapeVertexBuffer::InitRHI()
 //
 
 template <typename INDEX_TYPE>
-void FLandscapeSharedBuffers::CreateIndexBuffers(ERHIFeatureLevel::Type InFeatureLevel)
+void FLandscapeSharedBuffers::CreateIndexBuffers()
 {
-	if (InFeatureLevel <= ERHIFeatureLevel::ES3_1)
+	if (bUseMobileLandscapeMesh && !bVertexScoresComputed)
 	{
-		if (!bVertexScoresComputed)
-		{
-			bVertexScoresComputed = ComputeVertexScores();
-		}
+		bVertexScoresComputed = ComputeVertexScores();
 	}
 
 	TArray<INDEX_TYPE> VertexToIndexMap;
@@ -2683,7 +2683,7 @@ void FLandscapeSharedBuffers::CreateIndexBuffers(ERHIFeatureLevel::Type InFeatur
 		MaxIndexFull = 0;
 		MinIndexFull = MAX_int32;
 
-		if (InFeatureLevel <= ERHIFeatureLevel::ES3_1)
+		if (bUseMobileLandscapeMesh)
 		{
 			// mobile version shares vertices across LODs to save memory
 			float MipRatio = (float)SubsectionSizeQuads / (float)LodSubsectionSizeQuads; // Morph current MIP to base MIP
@@ -2893,7 +2893,7 @@ void FLandscapeSharedBuffers::CreateGrassIndexBuffer()
 }
 #endif
 
-FLandscapeSharedBuffers::FLandscapeSharedBuffers(const int32 InSharedBuffersKey, const int32 InSubsectionSizeQuads, const int32 InNumSubsections, const ERHIFeatureLevel::Type InFeatureLevel)
+FLandscapeSharedBuffers::FLandscapeSharedBuffers(const int32 InSharedBuffersKey, const int32 InSubsectionSizeQuads, const int32 InNumSubsections, const ERHIFeatureLevel::Type InFeatureLevel, bool bInUseMobileLandscapeMesh)
 	: SharedBuffersKey(InSharedBuffersKey)
 	, NumIndexBuffers(FMath::CeilLogTwo(InSubsectionSizeQuads + 1))
 	, SubsectionSizeVerts(InSubsectionSizeQuads + 1)
@@ -2902,14 +2902,16 @@ FLandscapeSharedBuffers::FLandscapeSharedBuffers(const int32 InSharedBuffersKey,
 	, FixedGridVertexFactory(nullptr)
 	, VertexBuffer(nullptr)
 	, bUse32BitIndices(false)
+
+	, bUseMobileLandscapeMesh(bInUseMobileLandscapeMesh)
 #if WITH_EDITOR
 	, GrassIndexBuffer(nullptr)
 #endif
 {
 	NumVertices = FMath::Square(SubsectionSizeVerts) * FMath::Square(NumSubsections);
-	if (InFeatureLevel > ERHIFeatureLevel::ES3_1)
+	if (!bUseMobileLandscapeMesh)
 	{
-		// Vertex Buffer cannot be shared
+		// Vertex Buffer cannot be shared 
 		VertexBuffer = new FLandscapeVertexBuffer(InFeatureLevel, NumVertices, SubsectionSizeVerts, NumSubsections);
 	}
 	IndexBuffers = new FIndexBuffer * [NumIndexBuffers];
@@ -2927,7 +2929,7 @@ FLandscapeSharedBuffers::FLandscapeSharedBuffers(const int32 InSharedBuffersKey,
 	if (NumVertices > 65535)
 	{
 		bUse32BitIndices = true;
-		CreateIndexBuffers<uint32>(InFeatureLevel);
+		CreateIndexBuffers<uint32>();
 #if WITH_EDITOR
 		if (InFeatureLevel > ERHIFeatureLevel::ES3_1)
 		{
@@ -2937,7 +2939,7 @@ FLandscapeSharedBuffers::FLandscapeSharedBuffers(const int32 InSharedBuffersKey,
 	}
 	else
 	{
-		CreateIndexBuffers<uint16>(InFeatureLevel);
+		CreateIndexBuffers<uint16>();
 #if WITH_EDITOR
 		if (InFeatureLevel > ERHIFeatureLevel::ES3_1)
 		{
@@ -3122,8 +3124,8 @@ bool FLandscapeVertexFactory::ShouldCompilePermutation(const FVertexFactoryShade
 {
 	// only compile landscape materials for landscape vertex factory
 	// The special engine materials must be compiled for the landscape vertex factory because they are used with it for wireframe, etc.
-	return IsFeatureLevelSupported(Parameters.Platform, ERHIFeatureLevel::SM5) &&
-		(Parameters.MaterialParameters.bIsUsedWithLandscape || Parameters.MaterialParameters.bIsSpecialEngineMaterial);
+	const bool  bSupportsHeightfieldRepresentation = IsFeatureLevelSupported(Parameters.Platform, ERHIFeatureLevel::SM5) || !UseMobileLandscapeMesh(Parameters.Platform);
+	return bSupportsHeightfieldRepresentation && (Parameters.MaterialParameters.bIsUsedWithLandscape || Parameters.MaterialParameters.bIsSpecialEngineMaterial);
 }
 
 void FLandscapeVertexFactory::ModifyCompilationEnvironment(const FVertexFactoryShaderPermutationParameters& Parameters, FShaderCompilerEnvironment& OutEnvironment)
@@ -3718,11 +3720,12 @@ void ULandscapeComponent::GetStreamingRenderAssetInfo(FStreamingTextureLevelCont
 	}
 
 	ERHIFeatureLevel::Type FeatureLevel = LevelContext.GetFeatureLevel();
-	int32 MaterialInstanceCount = FeatureLevel >= ERHIFeatureLevel::SM5 ? GetMaterialInstanceCount() : MobileMaterialInterfaces.Num();
+	const bool bUseMobileLandscapeMesh = UseMobileLandscapeMesh(GShaderPlatformForFeatureLevel[FeatureLevel]);
+	int32 MaterialInstanceCount = bUseMobileLandscapeMesh ? MobileMaterialInterfaces.Num() : GetMaterialInstanceCount();
 
 	for (int32 MaterialIndex = 0; MaterialIndex < MaterialInstanceCount; ++MaterialIndex)
 	{
-		const UMaterialInterface* MaterialInterface = FeatureLevel >= ERHIFeatureLevel::SM5 ? GetMaterialInstance(MaterialIndex) : ToRawPtr(MobileMaterialInterfaces[MaterialIndex]);
+		const UMaterialInterface* MaterialInterface = bUseMobileLandscapeMesh ? ToRawPtr(MobileMaterialInterfaces[MaterialIndex]) : GetMaterialInstance(MaterialIndex);
 
 		// Normal usage...
 		// Enumerate the textures used by the material.
diff --git a/Engine/Source/Runtime/Landscape/Private/LandscapeRenderMobile.cpp b/Engine/Source/Runtime/Landscape/Private/LandscapeRenderMobile.cpp
--- a/Engine/Source/Runtime/Landscape/Private/LandscapeRenderMobile.cpp
+++ b/Engine/Source/Runtime/Landscape/Private/LandscapeRenderMobile.cpp
@@ -20,10 +20,15 @@ static TAutoConsoleVariable<int32> CVarMobileLandscapeHoleMesh(
 	TEXT("Set to 0 to skip loading of landscape hole meshes on mobile."),
 	ECVF_Default);
 
+static TAutoConsoleVariable<int32> CVarUseMobileLandscapeMesh(
+	TEXT("r.Mobile.LandscapeMesh"),
+	1,
+	TEXT("Whether to use baked landscape mesh for a mobile"),
+	ECVF_ReadOnly);
+
 bool FLandscapeVertexFactoryMobile::ShouldCompilePermutation(const FVertexFactoryShaderPermutationParameters& Parameters)
 {
-	auto FeatureLevel = GetMaxSupportedFeatureLevel(Parameters.Platform);
-	return (FeatureLevel == ERHIFeatureLevel::ES3_1) &&
+	return UseMobileLandscapeMesh(Parameters.Platform) &&
 		(Parameters.MaterialParameters.bIsUsedWithLandscape || Parameters.MaterialParameters.bIsSpecialEngineMaterial);
 }
 
@@ -197,7 +202,7 @@ void FLandscapeFixedGridVertexFactoryMobile::ModifyCompilationEnvironment(const
 	
 bool FLandscapeFixedGridVertexFactoryMobile::ShouldCompilePermutation(const FVertexFactoryShaderPermutationParameters& Parameters)
 {
-	return GetMaxSupportedFeatureLevel(Parameters.Platform) == ERHIFeatureLevel::ES3_1 &&
+	return UseMobileLandscapeMesh(Parameters.Platform) &&
 		(Parameters.MaterialParameters.bIsUsedWithLandscape || Parameters.MaterialParameters.bIsSpecialEngineMaterial);
 }
 
@@ -381,7 +386,7 @@ void FLandscapeComponentSceneProxyMobile::CreateRenderThreadResources()
 	{
 		SharedBuffers = new FLandscapeSharedBuffers(
 			SharedBuffersKey, SubsectionSizeQuads, NumSubsections,
-			GetScene().GetFeatureLevel());
+			FeatureLevel, true); 
 
 		FLandscapeComponentSceneProxy::SharedBuffersMap.Add(SharedBuffersKey, SharedBuffers);
 	}
diff --git a/Engine/Source/Runtime/Landscape/Public/LandscapeRender.h b/Engine/Source/Runtime/Landscape/Public/LandscapeRender.h
--- a/Engine/Source/Runtime/Landscape/Public/LandscapeRender.h
+++ b/Engine/Source/Runtime/Landscape/Public/LandscapeRender.h
@@ -321,6 +321,7 @@ public:
 	FIndexBuffer** IndexBuffers;
 	FLandscapeIndexRanges* IndexRanges;
 	bool bUse32BitIndices;
+	bool bUseMobileLandscapeMesh;
 #if WITH_EDITOR
 	FIndexBuffer* GrassIndexBuffer;
 	TArray<int32, TInlineAllocator<8>> GrassIndexMipOffsets;
@@ -330,10 +331,10 @@ public:
 	TArray<FIndexBuffer*> ZeroOffsetIndexBuffers;
 #endif
 
-	FLandscapeSharedBuffers(int32 SharedBuffersKey, int32 SubsectionSizeQuads, int32 NumSubsections, ERHIFeatureLevel::Type FeatureLevel);
+	FLandscapeSharedBuffers(int32 SharedBuffersKey, int32 SubsectionSizeQuads, int32 NumSubsections, ERHIFeatureLevel::Type FeatureLevel, bool bUseMobileLandscapeMesh);
 
 	template <typename INDEX_TYPE>
-	void CreateIndexBuffers(ERHIFeatureLevel::Type InFeatureLevel);
+	void CreateIndexBuffers();
 	
 #if WITH_EDITOR
 	template <typename INDEX_TYPE>
diff --git a/Engine/Source/Runtime/Landscape/Public/LandscapeRenderMobile.h b/Engine/Source/Runtime/Landscape/Public/LandscapeRenderMobile.h
--- a/Engine/Source/Runtime/Landscape/Public/LandscapeRenderMobile.h
+++ b/Engine/Source/Runtime/Landscape/Public/LandscapeRenderMobile.h
@@ -50,6 +50,7 @@ public:
 	{
 		FVertexFactory::ModifyCompilationEnvironment(Parameters, OutEnvironment);
 		OutEnvironment.SetDefine(TEXT("NUM_VF_PACKED_INTERPOLANTS"), TEXT("1"));
+		OutEnvironment.SetDefine(TEXT("MOBILE_LANDSCAPE_MESH"), TEXT("1"));
 	}
 
 	// FRenderResource interface.
diff --git a/Engine/Source/Runtime/RenderCore/Private/RenderUtils.cpp b/Engine/Source/Runtime/RenderCore/Private/RenderUtils.cpp
--- a/Engine/Source/Runtime/RenderCore/Private/RenderUtils.cpp
+++ b/Engine/Source/Runtime/RenderCore/Private/RenderUtils.cpp
@@ -1681,6 +1681,16 @@ RENDERCORE_API bool UseVirtualTextureLightmap(const FStaticFeatureLevel InFeatur
 	return bUseVirtualTextureLightmap;
 }
 
+RENDERCORE_API bool UseMobileLandscapeMesh(EShaderPlatform ShaderPlatform)
+{
+	if (IsMobilePlatform(ShaderPlatform))
+	{
+		static FShaderPlatformCachedIniValue<bool> PerPlatformCVar(TEXT("r.Mobile.LandscapeMesh"));
+		return (PerPlatformCVar.Get(ShaderPlatform) != 0);
+	}
+	return false;
+}
+
 RENDERCORE_API bool ExcludeNonPipelinedShaderTypes(EShaderPlatform ShaderPlatform)
 {
 	if (RHISupportsShaderPipelines(ShaderPlatform))
diff --git a/Engine/Source/Runtime/RenderCore/Public/RenderUtils.h b/Engine/Source/Runtime/RenderCore/Public/RenderUtils.h
--- a/Engine/Source/Runtime/RenderCore/Public/RenderUtils.h
+++ b/Engine/Source/Runtime/RenderCore/Public/RenderUtils.h
@@ -767,6 +767,11 @@ inline bool UseNonNaniteVirtualShadowMaps(EShaderPlatform ShaderPlatform, const
 */
 RENDERCORE_API bool UseVirtualTextureLightmap(const FStaticFeatureLevel InFeatureLevel, const class ITargetPlatform* TargetPlatform = nullptr);
 
+/**
+*	Checks if platform uses a baked landscape mesh (mobile)
+*/
+RENDERCORE_API bool UseMobileLandscapeMesh(EShaderPlatform ShaderPlatform);
+
 /**
  *  Checks if the non-pipeline shaders will not be compild and ones from FShaderPipeline used instead.
  */
