Uses the first HLSL template

diff --git a/Engine/Shaders/Private/TemporalSuperResolution/TSRCommon.ush b/Engine/Shaders/Private/TemporalSuperResolution/TSRCommon.ush
--- a/Engine/Shaders/Private/TemporalSuperResolution/TSRCommon.ush
+++ b/Engine/Shaders/Private/TemporalSuperResolution/TSRCommon.ush
@@ -863,8 +863,8 @@ void IsOffScreenOrDisoccluded(uint bCameraCut, float2x2 ScreenPos, tsr_half2 Par
 	bool2 bIsOutOfBounds = max(abs(ScreenPos[0]), abs(ScreenPos[1])) >= 1.0;
 	bool2 bIsParallaxRejected = ParallaxRejectionMask < tsr_half(PARALLAX_REJECTION_MASK_THRESHOLD);
 
-	bIsOffScreen = bIsCameraCut || bIsOutOfBounds;
-	bIsDisoccluded = !bIsOffScreen && bIsParallaxRejected;
+	bIsOffScreen = or(bIsCameraCut, bIsOutOfBounds);
+	bIsDisoccluded = and(!bIsOffScreen, bIsParallaxRejected);
 }
 
 tsr_half MeasureBackbufferLDRQuantizationError()
diff --git a/Engine/Shaders/Private/TemporalSuperResolution/TSRUpdateHistory.usf b/Engine/Shaders/Private/TemporalSuperResolution/TSRUpdateHistory.usf
--- a/Engine/Shaders/Private/TemporalSuperResolution/TSRUpdateHistory.usf
+++ b/Engine/Shaders/Private/TemporalSuperResolution/TSRUpdateHistory.usf
@@ -166,6 +166,12 @@ groupshared tsr_half4 SharedArray11[GROUP_SHARED_SIZE];
 
 //------------------------------------------------------- FUNCTIONS
 
+template<typename T>
+void CorrectExposure(inout T X)
+{
+	X *= tsr_half(HistoryPreExposureCorrection);
+}
+
 CALL_SITE_DEBUGLOC 
 tsr_ushort2x2 Map8x8Tile2x2LaneDPV(uint GroupThreadIndex)
 {
@@ -465,7 +471,7 @@ void MainCS(
 		#if CONFIG_REJECTION_ANTI_ALIASING
 		{
 			// Detect if HistoryBufferUV would be outside of the viewport.
-			tsr_half2 TotalRejection = select(bIsOffScreen || bIsDisoccluded, tsr_half(0.0).xx, saturate(LowFrequencyRejection * tsr_half(4.0)));
+			tsr_half2 TotalRejection = select(or(bIsOffScreen, bIsDisoccluded), tsr_half(0.0).xx, saturate(LowFrequencyRejection * tsr_half(4.0)));
 
 			tsr_ushort2 EncodedInputTexelOffset = dpv_force_interleave_registers_array(RawEncodedInputTexelOffset);
 			
@@ -562,7 +568,7 @@ void MainCS(
 			// Soften the spatial kernel when there is high noise in the input to make parrallax disocclusion of detailed geometry less distracting.
 			tsr_half2 OffScreenInputToHistoryFactor = tsr_half(1.0) - tsr_half(0.5) * NoiseFiltering;
 
-			tsr_half2 KernelInputToHistoryFactor = select(bIsOffScreen || bIsDisoccluded, OffScreenInputToHistoryFactor, tsr_half(InputToHistoryFactor));
+			tsr_half2 KernelInputToHistoryFactor = select(or(bIsOffScreen, bIsDisoccluded), OffScreenInputToHistoryFactor, tsr_half(InputToHistoryFactor));
 
 			UNROLL_N(CONFIG_SAMPLES)
 			for (uint SampleId = 0; SampleId < CONFIG_SAMPLES; SampleId++)
@@ -882,19 +888,19 @@ void MainCS(
 
 		// Correct history
 		{
-			PrevHistoryMommentMin[0] *= tsr_half(HistoryPreExposureCorrection);
-			PrevHistoryMommentMin[1] *= tsr_half(HistoryPreExposureCorrection);
-			PrevHistoryMommentMin[2] *= tsr_half(HistoryPreExposureCorrection);
-			PrevHistoryMommentMax[0] *= tsr_half(HistoryPreExposureCorrection);
-			PrevHistoryMommentMax[1] *= tsr_half(HistoryPreExposureCorrection);
-			PrevHistoryMommentMax[2] *= tsr_half(HistoryPreExposureCorrection);
-			PrevHighFrequencyColor[0] *= tsr_half(HistoryPreExposureCorrection);
-			PrevHighFrequencyColor[1] *= tsr_half(HistoryPreExposureCorrection);
-			PrevHighFrequencyColor[2] *= tsr_half(HistoryPreExposureCorrection);
+			CorrectExposure(PrevHistoryMommentMin[0]);
+			CorrectExposure(PrevHistoryMommentMin[1]);
+			CorrectExposure(PrevHistoryMommentMin[2]);
+			CorrectExposure(PrevHistoryMommentMax[0]);
+			CorrectExposure(PrevHistoryMommentMax[1]);
+			CorrectExposure(PrevHistoryMommentMax[2]);
+			CorrectExposure(PrevHighFrequencyColor[0]);
+			CorrectExposure(PrevHighFrequencyColor[1]);
+			CorrectExposure(PrevHighFrequencyColor[2]);
 			#if !CONFIG_HIGH_FREQUENCY_ONLY
-				PrevLowFrequencyColor[0] *= tsr_half(HistoryPreExposureCorrection);
-				PrevLowFrequencyColor[1] *= tsr_half(HistoryPreExposureCorrection);
-				PrevLowFrequencyColor[2] *= tsr_half(HistoryPreExposureCorrection);
+				CorrectExposure(PrevLowFrequencyColor[0]);
+				CorrectExposure(PrevLowFrequencyColor[1]);
+				CorrectExposure(PrevLowFrequencyColor[2]);
 			#endif
 
 			// Fixes negative validity that can happen due to negative lobes of Catrom.
@@ -1033,7 +1039,7 @@ void MainCS(
 			tsr_halfCx2 ClampedPrevHighFrequencyColor = clamp(PrevHighFrequencyColor, dpv_scale(InputMinColor, PrevHistoryWeight), dpv_scale(InputMaxColor, PrevHistoryWeight));
 		#endif
 		
-		tsr_half2 TotalRejection = select(bIsOffScreen || bIsDisoccluded, tsr_half(0.0).xx, saturate(LowFrequencyRejection * tsr_half(4.0)));
+		tsr_half2 TotalRejection = select(or(bIsOffScreen, bIsDisoccluded), tsr_half(0.0).xx, saturate(LowFrequencyRejection * tsr_half(4.0)));
 
 		tsr_half2 BlendClamp = TotalRejection * HoleFilledVelocityMask[1] * HighFrequencyRejection;
 
@@ -1130,7 +1136,7 @@ void MainCS(
 		//	PrevHistoryRejectionWeight = 1.0;
 		//#endif
 		
-		PrevHistoryRejectionWeight = select(bIsOffScreen || bIsDisoccluded, tsr_half(0.0), PrevHistoryRejectionWeight);
+		PrevHistoryRejectionWeight = select(or(bIsOffScreen, bIsDisoccluded), tsr_half(0.0), PrevHistoryRejectionWeight);
 
 		tsr_half2 DesiredCurrentContribution = max(Histeresis * InputPixelAlignement, tsr_half(0.0));
 
@@ -1316,7 +1322,7 @@ void MainCS(
 
 			// Completely discard the history
 			{
-				tsr_half2 TotalRejection = select(bIsOffScreen || bIsDisoccluded, tsr_half(0.0), saturate(LowFrequencyRejection * tsr_half(4.0)));
+				tsr_half2 TotalRejection = select(or(bIsOffScreen, bIsDisoccluded), tsr_half(0.0), saturate(LowFrequencyRejection * tsr_half(4.0)));
 
 				tsr_half2 TotalRejectionBlending = max(BlendFinal, HoleFilledVelocityMask[0]);
 
@@ -1449,7 +1455,7 @@ void MainCS(
 				0x0);
 			
 			// Reset parallax factor
-			FinalSubpixelDetails = select(bIsOffScreen || bIsDisoccluded, FinalSubpixelDetails & ~ParallaxFactorMask, FinalSubpixelDetails);
+			FinalSubpixelDetails = select(or(bIsOffScreen, bIsDisoccluded), FinalSubpixelDetails & ~ParallaxFactorMask, FinalSubpixelDetails);
 		}
 		
 		uint2 SubpixelId = dpv_dot(uint2x2(bKOSign), uint2(1, SUB_PIXEL_GRID_SIZE));
diff --git a/Engine/Source/Runtime/Renderer/Private/PostProcess/TemporalSuperResolution.cpp b/Engine/Source/Runtime/Renderer/Private/PostProcess/TemporalSuperResolution.cpp
--- a/Engine/Source/Runtime/Renderer/Private/PostProcess/TemporalSuperResolution.cpp
+++ b/Engine/Source/Runtime/Renderer/Private/PostProcess/TemporalSuperResolution.cpp
@@ -178,6 +178,7 @@ public:
 	{
 		OutEnvironment.CompilerFlags.Add(CFLAG_AllowRealTypes);
 		OutEnvironment.CompilerFlags.Add(CFLAG_WarningsAsErrors);
+		OutEnvironment.CompilerFlags.Add(CFLAG_HLSL2021);
 	}
 }; // class FTemporalSuperResolutionShader
 
