Stata - update visualize to show complexity

diff --git a/Engine/Shaders/Private/Strata/StrataVisualize.usf b/Engine/Shaders/Private/Strata/StrataVisualize.usf
--- a/Engine/Shaders/Private/Strata/StrataVisualize.usf
+++ b/Engine/Shaders/Private/Strata/StrataVisualize.usf
@@ -241,6 +241,9 @@ void VisualizeMaterialPS(
 		FStrataPixelHeader Header = UnpackStrataHeaderIn(Strata.MaterialTextureArray, StrataAddressing, Strata.TopLayerTexture);
 		const uint PixelStrataDataByteOffsetPostHeader = StrataAddressing.ReadBytes;
 
+		const bool bIsSimpleMaterial = IsSimpleMaterial(Header) || Header.BSDFCount == 0;
+		const bool bIsSingleMaterial = !IsSimpleMaterial(Header) && IsSingleMaterial(Header);
+
 		float3 BSDFFontColor = float3(0.25, 1.0, 0.25);
 		float3 BSDFFontColorGrey = float3(0.25, 0.25, 0.25);
 		float3 INFOFontColor = float3(1.0, 0.5, 0.25);
@@ -250,6 +253,9 @@ void VisualizeMaterialPS(
 		uint2 TopLeft = TopLeftInit;
 
 #define PRINTPARAMCHAR(X) PrintCharacter(PixelPos, OutColor.xyz, PARAMFontColor, TopLeft, X);
+#define PRINTPARAMCHARRed(X) PrintCharacter(PixelPos, OutColor.xyz, float3(1.0, 0.7, 0.7), TopLeft, X);
+#define PRINTPARAMCHARYel(X) PrintCharacter(PixelPos, OutColor.xyz, float3(1.0, 1.0, 0.7), TopLeft, X);
+#define PRINTPARAMCHARGre(X) PrintCharacter(PixelPos, OutColor.xyz, float3(0.7, 1.0, 0.7), TopLeft, X);
 #define PRINTINFOCHAR(X) PrintCharacter(PixelPos, OutColor.xyz, INFOFontColor, TopLeft, X);
 #define PRINTBSDFCHAR(X) PrintCharacter(PixelPos, OutColor.xyz, BSDFFontColor, TopLeft, X);
 #define PRINTBSDFCHARGREY(X) PrintCharacter(PixelPos, OutColor.xyz, BSDFFontColorGrey, TopLeft, X);
@@ -287,9 +293,6 @@ void VisualizeMaterialPS(
 #define HORIZONTALTAB  TopLeft.x += 8;
 
 #if PERMUTATION_BSDF_PASS==0
-		PRINTINFOCHAR(_T_); PRINTINFOCHAR(_O_); PRINTINFOCHAR(_P_);
-		NEXTLINE;
-		HORIZONTALSEPARATION;
 		PRINTBSDFCOUNT();
 		NEXTLINE;
 		PRINTMATERIALAO();
@@ -304,6 +307,20 @@ void VisualizeMaterialPS(
 		NEXTLINE;
 		PRINTMATERIALINDSHAD();
 		NEXTLINE;
+		PRINTBSDFCHAR(_T_); PRINTBSDFCHAR(_Y_); PRINTBSDFCHAR(_P_); PRINTBSDFCHAR(_E_) PRINTBSDFCHAR(_SPC_)
+		if (bIsSimpleMaterial)
+		{
+			PRINTPARAMCHARGre(_S_); PRINTPARAMCHARGre(_I_); PRINTPARAMCHARGre(_M_); PRINTPARAMCHARGre(_P_); PRINTPARAMCHARGre(_L_); PRINTPARAMCHARGre(_E_);
+		}
+		else if(bIsSingleMaterial)
+		{
+			PRINTPARAMCHARYel(_S_); PRINTPARAMCHARYel(_I_); PRINTPARAMCHARYel(_N_); PRINTPARAMCHARYel(_G_); PRINTPARAMCHARYel(_L_); PRINTPARAMCHARYel(_E_);
+		}
+		else
+		{
+			PRINTPARAMCHARRed(_C_); PRINTPARAMCHARRed(_O_); PRINTPARAMCHARRed(_M_); PRINTPARAMCHARRed(_P_); PRINTPARAMCHARRed(_L_); PRINTPARAMCHARRed(_E_); PRINTPARAMCHARRed(_X_);
+		}
+		NEXTLINE;
 
 		PRINTCLASSIFROUGH();
 		NEXTLINE;
@@ -334,9 +351,11 @@ void VisualizeMaterialPS(
 		}
 
 		HORIZONTALSEPARATION;
-#else
+		PRINTINFOCHAR(_T_); PRINTINFOCHAR(_O_); PRINTINFOCHAR(_P_);
 		NEXTLINE;
 		HORIZONTALSEPARATION;
+#else
+		NEXTLINE;
 		NEXTLINE;
 		NEXTLINE;
 		NEXTLINE;
@@ -367,6 +386,8 @@ void VisualizeMaterialPS(
 		}
 
 		HORIZONTALSEPARATION;
+		NEXTLINE;
+		HORIZONTALSEPARATION;
 #endif
 
 		LOOP
