Fixes short-circuit in TSR when alpha channel is enabled

diff --git a/Engine/Shaders/Private/TemporalSuperResolution/TSRUpdateHistory.usf b/Engine/Shaders/Private/TemporalSuperResolution/TSRUpdateHistory.usf
--- a/Engine/Shaders/Private/TemporalSuperResolution/TSRUpdateHistory.usf
+++ b/Engine/Shaders/Private/TemporalSuperResolution/TSRUpdateHistory.usf
@@ -1541,7 +1541,7 @@ void MainCS(
 		// (0.995 chosen to accomodate handling of 254/255)
 		#if CONFIG_SCENE_COLOR_APLHA
 		{
-			FinalOutputColor[3] = FinalOutputColor[3] > 0.995 ? tsr_half(1.0).xx : FinalOutputColor[3];
+			FinalOutputColor[3] = select(FinalOutputColor[3] > 0.995, tsr_half(1.0), FinalOutputColor[3]);
 		}
 		#endif
 	}
