Fixed strata compilation issue with FXC.

diff --git a/Engine/Shaders/Private/Strata/StrataVisualize.usf b/Engine/Shaders/Private/Strata/StrataVisualize.usf
--- a/Engine/Shaders/Private/Strata/StrataVisualize.usf
+++ b/Engine/Shaders/Private/Strata/StrataVisualize.usf
@@ -390,7 +390,7 @@ void VisualizeMaterialPS(
 		HORIZONTALSEPARATION;
 #endif
 
-		LOOP
+		UNROLL_N(STRATA_MAX_BSDF_COUNT) // FXC requires that we specify an upper bound here. Not DXC.
 		for (uint BSDFIndex = 0; BSDFIndex < Header.BSDFCount; ++BSDFIndex)
 		{
 			const FStrataBSDF BSDF = UnpackStrataBSDFIn(Strata.MaterialTextureArray, StrataAddressing, Header);
