Change Strata BSDF offset to fit up to 10 slabs data.

Also add function to seek BSDF data, and helper to compute offset address.

diff --git a/Engine/Shaders/Private/Strata/StrataMaterialClassification.usf b/Engine/Shaders/Private/Strata/StrataMaterialClassification.usf
--- a/Engine/Shaders/Private/Strata/StrataMaterialClassification.usf
+++ b/Engine/Shaders/Private/Strata/StrataMaterialClassification.usf
@@ -315,14 +315,14 @@ void BSDFTileMainCS(uint2 GroupThreadId : SV_GroupThreadID, uint2 GroupId : SV_G
 	{
 		FStrataAddressing StrataAddressing = GetStrataPixelDataByteOffset(PixelCoord, uint2(View.BufferSizeAndInvSize.xy), MaxBytesPerPixel);
 		FStrataPixelHeader StrataPixelHeader = UnpackStrataHeaderIn(MaterialTextureArray, StrataAddressing, TopLayerTexture);
-		BSDFCount = min(StrataPixelHeader.BSDFCount, STRATA_TILE_MAX_BSDF_COUNT);
+		BSDFCount = min(StrataPixelHeader.BSDFCount, STRATA_MAX_BSDF_COUNT_FOR_BDSFOFFSET);
 
 		if (BSDFCount > 0)
 		{
 			FStrataBSDFOffset Offsets = (FStrataBSDFOffset)0;
 			Offsets.BSDFCount = BSDFCount;
 
-			UNROLL_N(STRATA_TILE_MAX_BSDF_COUNT) 
+			UNROLL_N(STRATA_MAX_BSDF_COUNT_FOR_BDSFOFFSET)
 			for (uint BSDFIndex = 0; BSDFIndex < BSDFCount; ++BSDFIndex)
 			{
 				Offsets.BSDFOffsets[BSDFIndex] = StrataAddressing.CurrentIndex;
diff --git a/Engine/Shaders/Private/Strata/StrataTile.ush b/Engine/Shaders/Private/Strata/StrataTile.ush
--- a/Engine/Shaders/Private/Strata/StrataTile.ush
+++ b/Engine/Shaders/Private/Strata/StrataTile.ush
@@ -73,33 +73,33 @@ FStrataBSDFTile UnpackBSDFTile(uint In)
 ////////////////////////////////////////////////////////////////////////////////////////////////////////////
 // Strata BSDF offset
 
-// STRATA_TODO MErge with incoming CL and simply make STRATA_MAX_BSDF_COUNT only 8 or 12 and remove STRATA_TILE_MAX_BSDF_COUNT
-#define STRATA_TILE_MAX_BSDF_COUNT 8
+// BSDF offsets are packed into 32bits to max. BSDF is < STRATA_MAX_BSDF_COUNT, which could leads to some issue.
+#define STRATA_MAX_BSDF_COUNT_FOR_BDSFOFFSET 10u
 
 // Store the material slice index offset for each BSDF, for pixels containing several BSDFs
 struct FStrataBSDFOffset
 {
-	uint BSDFOffsets[STRATA_TILE_MAX_BSDF_COUNT];
+	uint BSDFOffsets[STRATA_MAX_BSDF_COUNT_FOR_BDSFOFFSET];
 	uint BSDFCount;
 };
 
 uint PackBSDFOffset(FStrataBSDFOffset In)
 {
 	// Packed as relative cumulative offsets, stored into 32 bits
-	// | 4 bits | 2 bits | 2 bits | 2 bits | ... | 2 bits |
-	//   BSDF 0   BSDF 1   BSDF 2   BSDF 3  ...   BSDF 15 (STRATA_TILE_MAX_BSDF_COUNT)
+	// | 3 bits | 3 bits | 3 bits | 3 bits | ... | 3 bits |
+	//   BSDF 0   BSDF 1   BSDF 2   BSDF 3  ...   BSDF 9
 
 	uint Out = 0;
-	const uint BSDFCount = min(In.BSDFCount, STRATA_TILE_MAX_BSDF_COUNT);
+	const uint BSDFCount = min(In.BSDFCount, STRATA_MAX_BSDF_COUNT_FOR_BDSFOFFSET);
 	if (BSDFCount > 0)
 	{
-		Out = In.BSDFOffsets[0] & 0xF;
-		UNROLL_N(STRATA_TILE_MAX_BSDF_COUNT)
+		Out = In.BSDFOffsets[0] & 0x7;
+		UNROLL_N(STRATA_MAX_BSDF_COUNT_FOR_BDSFOFFSET)
 		for (uint i = 1; i < BSDFCount; i++)
 		{
-			uint Curr = (In.BSDFOffsets[i  ] & 0x3);
-			uint Prev = (In.BSDFOffsets[i-1] & 0x3);
-			Out = Out | (((Curr - Prev) & 0x3) << (2u*i));
+			const uint Curr = In.BSDFOffsets[i];
+			const uint Prev = In.BSDFOffsets[i-1];
+			Out = Out | (((Curr - Prev) & 0x7) << (3u*i));
 		}
 	}
 	return Out;
@@ -110,19 +110,40 @@ FStrataBSDFOffset UnpackBSDFOffset(uint In, uint InBSDFCount)
 	FStrataBSDFOffset Out = (FStrataBSDFOffset)0;
 	if (InBSDFCount > 0)
 	{
-		Out.BSDFCount = min(InBSDFCount, STRATA_TILE_MAX_BSDF_COUNT);
-		Out.BSDFOffsets[0] = (In & 0xF);
+		Out.BSDFCount = min(InBSDFCount, STRATA_MAX_BSDF_COUNT_FOR_BDSFOFFSET);
+		Out.BSDFOffsets[0] = (In & 0x7);
 
 		FStrataBSDFOffset Out = (FStrataBSDFOffset)0;
-		UNROLL_N(STRATA_TILE_MAX_BSDF_COUNT)
+		UNROLL_N(STRATA_MAX_BSDF_COUNT_FOR_BDSFOFFSET)
 		for (uint i = 1; i < Out.BSDFCount; i++)
 		{
-			Out.BSDFOffsets[i] = Out.BSDFOffsets[i-1] + ((In>>(2u*i)) & 0x3);
+			Out.BSDFOffsets[i] = Out.BSDFOffsets[i-1] + ((In>>(3u*i)) & 0x7);
 		}
 	}
 	return Out;
 }
 
+uint UnpackBSDFOffsetAtIndex(uint In, uint InIndex, uint InBSDFCount)
+{
+	uint Out = 0;
+	if (InBSDFCount > 0)
+	{
+		const uint BSDFCount = min(InBSDFCount, STRATA_MAX_BSDF_COUNT_FOR_BDSFOFFSET);
+		UNROLL_N(STRATA_MAX_BSDF_COUNT_FOR_BDSFOFFSET)
+		for (uint i = 0; i <= InIndex; i++)
+		{
+			Out += ((In >> (3u * i)) & 0x7);
+		}
+	}
+	return Out;
+}
+
+void StrataSeekBSDF(inout FStrataAddressing Addressing, uint AddressOffset)
+{
+	Addressing.CurrentIndex = AddressOffset;
+	Addressing.ReadBytes = Addressing.CurrentIndex * 4;
+}
+
 void StrataSeekBSDF(inout FStrataAddressing Addressing, FStrataBSDFOffset Offsets, uint BSDFIndex)
 {
 	if (BSDFIndex < Offsets.BSDFCount)
