Strata - fixed 8bit dithering value in case it is used later.

diff --git a/Engine/Shaders/Private/Strata/Strata.ush b/Engine/Shaders/Private/Strata/Strata.ush
--- a/Engine/Shaders/Private/Strata/Strata.ush
+++ b/Engine/Shaders/Private/Strata/Strata.ush
@@ -1902,7 +1902,7 @@ void StrataDitherBSDF(inout FStrataBSDF BSDF, in float Dither)
 	{
 		// We execute a dither that can go in positive/negative direction to avoid surfaces to looks overall brighter.
 		const float CenteredDither = Dither - 0.5f;
-		const float Dither8bits = CenteredDither * 1.0f / 128.0f;
+		const float Dither8bits = CenteredDither * 1.0f / 255.0f;
 		const float Dither7bits = CenteredDither * 1.0f / 127.0f;
 		const float Dither6bits = CenteredDither * 1.0f / 63.0f;
 		const float3 Dither776 = float3(Dither7bits, Dither7bits, Dither6bits);
