Fixes regression in TSR on the edge of the viewport

diff --git a/Engine/Shaders/Private/TemporalSuperResolution/TSRFilterFrequencies.usf b/Engine/Shaders/Private/TemporalSuperResolution/TSRFilterFrequencies.usf
--- a/Engine/Shaders/Private/TemporalSuperResolution/TSRFilterFrequencies.usf
+++ b/Engine/Shaders/Private/TemporalSuperResolution/TSRFilterFrequencies.usf
@@ -134,9 +134,9 @@ void MainCS(
 			tsr_short2 SubTileFetchCoordOffset1 = tsr_short(SUBTILE_SIZE) * select(SubTileOffset.y == 0, tsr_short2(2, SubTileOffset.x), tsr_short2(SubTileOffset.x, 2));
 			tsr_short2 SubTileFetchCoordOffset2 = tsr_short(SUBTILE_SIZE) * tsr_short(2);
 		
-			tsr_short2 ClampedFetchCoord0 = ClampPixelOffset((FetchOrigin + SubTileFetchCoordOffset0) + SubTilePixelOffset, InputInfo_ViewportMin, InputInfo_ViewportMax);
-			tsr_short2 ClampedFetchCoord1 = ClampPixelOffset((FetchOrigin + SubTileFetchCoordOffset1) + SubTilePixelOffset, InputInfo_ViewportMin, InputInfo_ViewportMax);
-			tsr_short2 ClampedFetchCoord2 = ClampPixelOffset((FetchOrigin + SubTileFetchCoordOffset2) + SubTilePixelOffset, InputInfo_ViewportMin, InputInfo_ViewportMax);
+			tsr_short2 ClampedFetchCoord0 = ClampPixelOffset((FetchOrigin + SubTileFetchCoordOffset0) + SubTilePixelOffset, InputPixelPosMin, InputPixelPosMax);
+			tsr_short2 ClampedFetchCoord1 = ClampPixelOffset((FetchOrigin + SubTileFetchCoordOffset1) + SubTilePixelOffset, InputPixelPosMin, InputPixelPosMax);
+			tsr_short2 ClampedFetchCoord2 = ClampPixelOffset((FetchOrigin + SubTileFetchCoordOffset2) + SubTilePixelOffset, InputPixelPosMin, InputPixelPosMax);
 
 			tsr_halfC Color0 = InputTexture[ClampedFetchCoord0];
 			tsr_halfC Color1 = InputTexture[ClampedFetchCoord1];
