Groundwork for in-engine GPU-based DXR validation layer

This adds r.D3D12.RayTracing.GPUValidation that performs basic ray tracing geometry validation on the GPU.
Currently only the basic inactive triangle rules are checked. Additional checks will be built on this foundation.

diff --git a/Engine/Shaders/Private/RayTracing/RayTracingHitGroupCommon.ush b/Engine/Shaders/Private/RayTracing/RayTracingHitGroupCommon.ush
--- a/Engine/Shaders/Private/RayTracing/RayTracingHitGroupCommon.ush
+++ b/Engine/Shaders/Private/RayTracing/RayTracingHitGroupCommon.ush
@@ -44,22 +44,23 @@ uint3 LoadIndices32Bit(ByteAddressBuffer SourceBuffer, uint OffsetInBytes)
 	return SourceBuffer.Load3(OffsetInBytes);
 }
 
-float3 LoadVertexPositionFloat3(ByteAddressBuffer SourceBuffer, uint Index, uint StrideInBytes)
+float3 LoadVertexPositionFloat3(ByteAddressBuffer SourceBuffer, uint BaseOffsetInBytes, uint Index, uint StrideInBytes)
 {
-	uint OffsetInBytes = Index * StrideInBytes;
+	uint OffsetInBytes = BaseOffsetInBytes + Index * StrideInBytes;
 	return asfloat(SourceBuffer.Load3(OffsetInBytes));
 }
 
 // Explicitly declare "system" resources that will be available to all hit group shaders.
 
-FTriangleBaseAttributes LoadTriangleBaseAttributes(uint PrimitiveId)
+FTriangleBaseAttributes LoadTriangleBaseAttributes(
+	ByteAddressBuffer IndexBuffer, uint IndexBufferOffsetInBytes, uint IndexBufferStride, 
+	ByteAddressBuffer VertexBuffer, uint VertexBufferOffsetInBytes, uint VertexBufferStride, 
+	uint PrimitiveId)
 {
 	FTriangleBaseAttributes Result = (FTriangleBaseAttributes)0;
 
 	// Fetch vertex indices and positions, then compute local space normal and transform it to world space
 
-	const uint IndexBufferOffsetInBytes = HitGroupSystemRootConstants.IndexBufferOffsetInBytes;
-	const uint IndexBufferStride = HitGroupSystemRootConstants.GetIndexStride();
 	const uint BaseIndex = PrimitiveId * 3;
 
 	if (IndexBufferStride == 0)
@@ -69,24 +70,31 @@ FTriangleBaseAttributes LoadTriangleBaseAttributes(uint PrimitiveId)
 	}
 	else if (IndexBufferStride == 2)
 	{
-		Result.Indices = LoadIndices16Bit(HitGroupSystemIndexBuffer, IndexBufferOffsetInBytes + BaseIndex * IndexBufferStride);
+		Result.Indices = LoadIndices16Bit(IndexBuffer, IndexBufferOffsetInBytes + BaseIndex * IndexBufferStride);
 	}
 	else
 	{
-		Result.Indices = LoadIndices32Bit(HitGroupSystemIndexBuffer, IndexBufferOffsetInBytes + BaseIndex * IndexBufferStride);
+		Result.Indices = LoadIndices32Bit(IndexBuffer, IndexBufferOffsetInBytes + BaseIndex * IndexBufferStride);
 	}
 
 	// Fetch vertex positions (in local space)
-	const uint VertexStride = HitGroupSystemRootConstants.GetVertexStride();
-
 	// #dxr_todo: UE-72160 handle various vertex formats, for now only supporting float3
-	Result.LocalPositions[0] = LoadVertexPositionFloat3(HitGroupSystemVertexBuffer, Result.Indices[0], VertexStride);
-	Result.LocalPositions[1] = LoadVertexPositionFloat3(HitGroupSystemVertexBuffer, Result.Indices[1], VertexStride);
-	Result.LocalPositions[2] = LoadVertexPositionFloat3(HitGroupSystemVertexBuffer, Result.Indices[2], VertexStride);
+	Result.LocalPositions[0] = LoadVertexPositionFloat3(VertexBuffer, Result.Indices[0], VertexBufferOffsetInBytes, VertexBufferStride);
+	Result.LocalPositions[1] = LoadVertexPositionFloat3(VertexBuffer, Result.Indices[1], VertexBufferOffsetInBytes, VertexBufferStride);
+	Result.LocalPositions[2] = LoadVertexPositionFloat3(VertexBuffer, Result.Indices[2], VertexBufferOffsetInBytes, VertexBufferStride);
 
 	return Result;
 }
 
+FTriangleBaseAttributes LoadTriangleBaseAttributes(uint PrimitiveId)
+{
+	uint IndexBufferOffsetInBytes = HitGroupSystemRootConstants.IndexBufferOffsetInBytes;
+	uint IndexBufferStride = HitGroupSystemRootConstants.GetIndexStride();
+	uint VertexStride = HitGroupSystemRootConstants.GetVertexStride();
+	uint VertexBufferOffsetInBytes = 0; // Base offset is already applied when the buffer is bound to the hit group
+	return LoadTriangleBaseAttributes(HitGroupSystemIndexBuffer, IndexBufferOffsetInBytes, IndexBufferStride, HitGroupSystemVertexBuffer, VertexBufferOffsetInBytes, VertexStride, PrimitiveId);
+}
+
 uint GetInstanceUserData()
 {
 	return InstanceID();
diff --git a/Engine/Shaders/Private/RayTracing/RayTracingValidation.usf b/Engine/Shaders/Private/RayTracing/RayTracingValidation.usf
new file mode 100644
--- /dev/null
+++ b/Engine/Shaders/Private/RayTracing/RayTracingValidation.usf
@@ -0,0 +1,90 @@
+// Copyright Epic Games, Inc. All Rights Reserved.
+
+#include "../Common.ush"
+
+#ifndef NUM_THREADS_X
+#define NUM_THREADS_X 1024
+#endif
+
+#if PLATFORM_SUPPORTS_DIAGNOSTIC_BUFFER	
+
+#include "/Engine/Private/RayTracing/RayTracingHitGroupCommon.ush"
+
+ByteAddressBuffer IndexBuffer;
+ByteAddressBuffer VertexBuffer;
+uint VertexBufferStride;
+uint VertexBufferOffsetInBytes;
+uint IndexBufferOffsetInBytes;
+uint IndexBufferStride;
+uint NumPrimitives;
+
+bool IsNaN(float F)
+{
+	uint U = asuint(F);
+	uint ExponentMask = 0x7f800000;
+	uint MantissaMask = 0x007fffff;
+	return (U & ExponentMask) == ExponentMask 
+		&& (U & MantissaMask) != 0;
+}
+
+// Returns true if all vertices are finite or if vertex X components are NaN. DXR inactive primitive rule applies:
+// Triangles are considered "inactive" (but legal input to acceleration structure build) if the x component of each vertex is NaN.
+void CheckTriangle(FTriangleBaseAttributes Tri, uint PrimitiveIndex)
+{
+	// All X-s are NaNs: triangle is legal and inactive
+	if (IsNaN(Tri.LocalPositions[0].x) && IsNaN(Tri.LocalPositions[1].x) && IsNaN(Tri.LocalPositions[2].x))
+	{
+		return;
+	}
+
+	// As of 2022-03-28, some UE systems mark triangles as culled by only setting the provoking vertex to NaN
+	// In practice, the drivers handle this, so we allow it for now.
+	// Technically, this is a DXR spec violation.
+	if (IsNaN(Tri.LocalPositions[0].x))
+	{
+		return;
+	}
+	
+	// Active triangles must have non-NaN positions
+	for (uint i=0; i<3; ++i)
+	{
+		if (IsNaN(Tri.LocalPositions[i].x) || IsNaN(Tri.LocalPositions[i].y) || IsNaN(Tri.LocalPositions[i].z))
+		{
+			uint AssertID = 0xF69EAF06; // Unique assert ID (TODO: add support for full human-readable assert strings)
+			UEReportAssertWithPayload(AssertID,
+			uint4(PrimitiveIndex, 
+				asuint(Tri.LocalPositions[i].x), 
+				asuint(Tri.LocalPositions[i].y), 
+				asuint(Tri.LocalPositions[i].z)));
+			return;
+		}
+	}
+}
+
+[numthreads(NUM_THREADS_X, 1, 1)]
+void RayTracingValidateGeometryBuildParamsCS(uint3 DispatchThreadId : SV_DispatchThreadID)
+{
+	uint PrimitiveIndex = DispatchThreadId.x;
+
+	if (PrimitiveIndex >= NumPrimitives)
+	{
+		return;
+	}
+
+	FTriangleBaseAttributes Tri = LoadTriangleBaseAttributes(
+		IndexBuffer, IndexBufferOffsetInBytes, IndexBufferStride,
+		VertexBuffer, VertexBufferOffsetInBytes, VertexBufferStride,
+		PrimitiveIndex);
+	
+	CheckTriangle(Tri, PrimitiveIndex);
+}
+
+#else // PLATFORM_SUPPORTS_DIAGNOSTIC_BUFFER
+
+[numthreads(NUM_THREADS_X, 1, 1)]
+void RayTracingValidateGeometryBuildParamsCS(uint3 DispatchThreadId : SV_DispatchThreadID)
+{
+}
+
+#endif // PLATFORM_SUPPORTS_DIAGNOSTIC_BUFFER
+
diff --git a/Engine/Source/Runtime/D3D12RHI/Private/D3D12RayTracing.cpp b/Engine/Source/Runtime/D3D12RHI/Private/D3D12RayTracing.cpp
--- a/Engine/Source/Runtime/D3D12RHI/Private/D3D12RayTracing.cpp
+++ b/Engine/Source/Runtime/D3D12RHI/Private/D3D12RayTracing.cpp
@@ -9,6 +9,7 @@
 #include "Containers/DynamicRHIResourceArray.h"
 #include "Experimental/Containers/SherwoodHashTable.h"
 #include "BuiltInRayTracingShaders.h"
+#include "RayTracingValidationShaders.h"
 #include "Hash/CityHash.h"
 #include "HAL/CriticalSection.h"
 #include "HAL/IConsoleManager.h"
@@ -87,6 +88,13 @@ static FAutoConsoleVariableRef CVarD3D12RayTracingDeduplicateSamplers(
 	TEXT("Use an exhaustive search to deduplicate sampler descriptors when generating shader binding tables. Reduces sampler heap usage at the cost of some CPU time. (default = 1)")
 );
 
+static int32 GD3D12RayTracingGPUValidation = 0;
+static FAutoConsoleVariableRef CVarD3D12RayTracingGPUValidation(
+	TEXT("r.D3D12.RayTracing.GPUValidation"),
+	GD3D12RayTracingGPUValidation,
+	TEXT("Whether to perform validation of ray tracing geometry and other structures on the GPU. Requires Shader Model 6. (default = 0)")
+);
+
 // Ray tracing stat counters
 
 DECLARE_STATS_GROUP(TEXT("D3D12RHI: Ray Tracing"), STATGROUP_D3D12RayTracing, STATCAT_Advanced);
@@ -4409,6 +4417,16 @@ void FD3D12CommandContext::RHIBuildAccelerationStructures(const TArrayView<const
 		}		
 	}
 
+	if (GD3D12RayTracingGPUValidation)
+	{
+		TRHICommandList_RecursiveHazardous<FD3D12CommandContext> RHICmdList(this, GetGPUMask());
+
+		for (const FRayTracingGeometryBuildParams& P : Params)
+		{
+			FRayTracingValidateGeometryBuildParamsCS::Dispatch(RHICmdList, P);
+		}
+	}
+
 	TArrayView<D3D12_BUILD_RAYTRACING_ACCELERATION_STRUCTURE_DESC> LocalBuildDescs = MakeArrayView(BuildDescs.GetData(), BuildDescs.Num());
 	BuildAccelerationStructuresInternal(LocalBuildDescs);
 
diff --git a/Engine/Source/Runtime/RenderCore/Private/RayTracingValidationShaders.cpp b/Engine/Source/Runtime/RenderCore/Private/RayTracingValidationShaders.cpp
new file mode 100644
--- /dev/null
+++ b/Engine/Source/Runtime/RenderCore/Private/RayTracingValidationShaders.cpp
@@ -0,0 +1,76 @@
+// Copyright Epic Games, Inc. All Rights Reserved.
+
+#include "RayTracingValidationShaders.h"
+#include "ShaderParameterUtils.h"
+#include "PipelineStateCache.h"
+
+#if RHI_RAYTRACING
+
+IMPLEMENT_GLOBAL_SHADER(FRayTracingValidateGeometryBuildParamsCS, "/Engine/Private/RayTracing/RayTracingValidation.usf", "RayTracingValidateGeometryBuildParamsCS", SF_Compute);
+
+void FRayTracingValidateGeometryBuildParamsCS::Dispatch(FRHICommandList& RHICmdList, const FRayTracingGeometryBuildParams& Params)
+{
+	const FRayTracingGeometryInitializer& Initializer = Params.Geometry->GetInitializer();
+
+	TShaderMapRef<FRayTracingValidateGeometryBuildParamsCS> ComputeShader(GetGlobalShaderMap(GMaxRHIFeatureLevel));
+	FRHIComputeShader* ShaderRHI = ComputeShader.GetComputeShader();
+	SetComputePipelineState(RHICmdList, ShaderRHI);
+
+	// TODO: handle non-indexed geometry
+	if (Initializer.IndexBuffer == nullptr)
+	{
+		return;
+	}
+
+	TWideStringBuilder<256> EventName;
+	EventName.Append(TEXT("RTValidation"));
+	if (!Initializer.DebugName.IsNone())
+	{
+		FString DebugNameString = Initializer.DebugName.ToString();
+		EventName.Append(TEXT(" - "));
+		EventName.Append(*DebugNameString);
+	}
+	
+	RHICmdList.PushEvent(EventName.GetData(), FColor::Black);
+
+	const uint32 IndexStride = Initializer.IndexBuffer->GetStride();
+
+	const FRawBufferShaderResourceViewInitializer IBViewInitializer(Initializer.IndexBuffer);
+	FShaderResourceViewRHIRef IndexBufferSRV = RHICreateShaderResourceView(IBViewInitializer);
+
+	for (const FRayTracingGeometrySegment& Segment : Initializer.Segments)
+	{
+		if (Segment.VertexBufferElementType != VET_Float3)
+		{
+			// Only Float3 vertex positions are currently supported
+			continue;
+		}
+
+		const uint32 IndexBufferOffsetInBytes = Segment.FirstPrimitive * IndexStride * 3;
+
+		SetShaderValue(RHICmdList, ShaderRHI, ComputeShader->VertexBufferStrideParam, Segment.VertexBufferStride);
+		SetShaderValue(RHICmdList, ShaderRHI, ComputeShader->VertexBufferOffsetInBytesParam, Segment.VertexBufferOffset);
+		SetShaderValue(RHICmdList, ShaderRHI, ComputeShader->IndexBufferOffsetInBytesParam, IndexBufferOffsetInBytes);
+		SetShaderValue(RHICmdList, ShaderRHI, ComputeShader->IndexBufferStrideParam, IndexStride);
+		SetShaderValue(RHICmdList, ShaderRHI, ComputeShader->NumPrimitivesParam, Segment.NumPrimitives);
+
+		const FRawBufferShaderResourceViewInitializer VBViewInitializer(Segment.VertexBuffer);
+		FShaderResourceViewRHIRef VertexBufferSRV = RHICreateShaderResourceView(VBViewInitializer);
+
+		SetSRVParameter(RHICmdList, ShaderRHI, ComputeShader->VertexBufferParam, VertexBufferSRV);
+		SetSRVParameter(RHICmdList, ShaderRHI, ComputeShader->IndexBufferParam, IndexBufferSRV);
+
+		// TODO: handle arbitrary large meshes that may overrun the 1D dispatch limit
+		const uint32 MaxDispatchDimension = 65536; // D3D11_CS_DISPATCH_MAX_THREAD_GROUPS_PER_DIMENSION
+		const uint32 NumGroupsX = FMath::Min((Segment.NumPrimitives + NumThreadsX - 1) / NumThreadsX, MaxDispatchDimension);
+
+		RHICmdList.DispatchComputeShader(NumGroupsX, 1, 1);
+	}
+
+	SetSRVParameter(RHICmdList, ShaderRHI, ComputeShader->VertexBufferParam, nullptr);
+	SetSRVParameter(RHICmdList, ShaderRHI, ComputeShader->IndexBufferParam, nullptr);
+
+	RHICmdList.PopEvent();
+}
+
+#endif // RHI_RAYTRACING
diff --git a/Engine/Source/Runtime/RenderCore/Public/RayTracingValidationShaders.h b/Engine/Source/Runtime/RenderCore/Public/RayTracingValidationShaders.h
new file mode 100644
--- /dev/null
+++ b/Engine/Source/Runtime/RenderCore/Public/RayTracingValidationShaders.h
@@ -0,0 +1,55 @@
+// Copyright Epic Games, Inc. All Rights Reserved.
+
+#pragma once
+
+#include "BuiltInRayTracingShaders.h"
+#include "ShaderCompilerCore.h"
+
+#if RHI_RAYTRACING
+
+class RENDERCORE_API FRayTracingValidateGeometryBuildParamsCS : public FBuiltInRayTracingShader
+{
+	DECLARE_GLOBAL_SHADER(FRayTracingValidateGeometryBuildParamsCS);
+
+public:
+	FRayTracingValidateGeometryBuildParamsCS() = default;
+	FRayTracingValidateGeometryBuildParamsCS(const ShaderMetaType::CompiledShaderInitializerType& Initializer)
+		: FBuiltInRayTracingShader(Initializer)
+	{
+		IndexBufferParam.Bind(Initializer.ParameterMap, TEXT("IndexBuffer"), SPF_Optional);
+		VertexBufferParam.Bind(Initializer.ParameterMap, TEXT("VertexBuffer"), SPF_Optional);
+		VertexBufferStrideParam.Bind(Initializer.ParameterMap, TEXT("VertexBufferStride"), SPF_Optional);
+		VertexBufferOffsetInBytesParam.Bind(Initializer.ParameterMap, TEXT("VertexBufferOffsetInBytes"), SPF_Optional);
+		IndexBufferOffsetInBytesParam.Bind(Initializer.ParameterMap, TEXT("IndexBufferOffsetInBytes"), SPF_Optional);
+		IndexBufferStrideParam.Bind(Initializer.ParameterMap, TEXT("IndexBufferStride"), SPF_Optional);
+		NumPrimitivesParam.Bind(Initializer.ParameterMap, TEXT("NumPrimitives"), SPF_Optional);
+	}
+
+	static bool ShouldCompilePermutation(const FGlobalShaderPermutationParameters& Parameters)
+	{
+		return FBuiltInRayTracingShader::ShouldCompilePermutation(Parameters);
+	}
+
+	// Large thread group to handle large meshes with a single 1D dispatch
+	static const uint32 NumThreadsX = 1024;
+
+	static inline void ModifyCompilationEnvironment(const FGlobalShaderPermutationParameters& Parameters, FShaderCompilerEnvironment& OutEnvironment)
+	{
+		OutEnvironment.CompilerFlags.Add(CFLAG_ForceDXC);
+		OutEnvironment.SetDefine(TEXT("NUM_THREADS_X"), NumThreadsX);
+		FBuiltInRayTracingShader::ModifyCompilationEnvironment(Parameters, OutEnvironment);
+	}
+
+	static void Dispatch(FRHICommandList& RHICmdList, const FRayTracingGeometryBuildParams& Params);
+
+	LAYOUT_FIELD(FShaderResourceParameter, IndexBufferParam);
+	LAYOUT_FIELD(FShaderResourceParameter, VertexBufferParam);
+	LAYOUT_FIELD(FShaderParameter, VertexBufferStrideParam);
+	LAYOUT_FIELD(FShaderParameter, VertexBufferOffsetInBytesParam);
+	LAYOUT_FIELD(FShaderParameter, IndexBufferOffsetInBytesParam);
+	LAYOUT_FIELD(FShaderParameter, IndexBufferStrideParam);
+	LAYOUT_FIELD(FShaderParameter, NumPrimitivesParam);
+
+};
+
+#endif // RHI_RAYTRACING
