Fix LoadVertexPositionFloat3() parameter order. Fixes ray tracing triangle normal reconstruction.

diff --git a/Engine/Shaders/Private/RayTracing/RayTracingHitGroupCommon.ush b/Engine/Shaders/Private/RayTracing/RayTracingHitGroupCommon.ush
--- a/Engine/Shaders/Private/RayTracing/RayTracingHitGroupCommon.ush
+++ b/Engine/Shaders/Private/RayTracing/RayTracingHitGroupCommon.ush
@@ -79,9 +79,9 @@ FTriangleBaseAttributes LoadTriangleBaseAttributes(
 
 	// Fetch vertex positions (in local space)
 	// #dxr_todo: UE-72160 handle various vertex formats, for now only supporting float3
-	Result.LocalPositions[0] = LoadVertexPositionFloat3(VertexBuffer, Result.Indices[0], VertexBufferOffsetInBytes, VertexBufferStride);
-	Result.LocalPositions[1] = LoadVertexPositionFloat3(VertexBuffer, Result.Indices[1], VertexBufferOffsetInBytes, VertexBufferStride);
-	Result.LocalPositions[2] = LoadVertexPositionFloat3(VertexBuffer, Result.Indices[2], VertexBufferOffsetInBytes, VertexBufferStride);
+	Result.LocalPositions[0] = LoadVertexPositionFloat3(VertexBuffer, VertexBufferOffsetInBytes, Result.Indices[0], VertexBufferStride);
+	Result.LocalPositions[1] = LoadVertexPositionFloat3(VertexBuffer, VertexBufferOffsetInBytes, Result.Indices[1], VertexBufferStride);
+	Result.LocalPositions[2] = LoadVertexPositionFloat3(VertexBuffer, VertexBufferOffsetInBytes, Result.Indices[2], VertexBufferStride);
 
 	return Result;
 }
