Add validation for ray tracing geometry index buffer values

- Check that index values are within the D3D12_RAYTRACING_GEOMETRY_TRIANGLES_DESC::VertexCount range.

diff --git a/Engine/Shaders/Private/RayTracing/RayTracingValidation.usf b/Engine/Shaders/Private/RayTracing/RayTracingValidation.usf
--- a/Engine/Shaders/Private/RayTracing/RayTracingValidation.usf
+++ b/Engine/Shaders/Private/RayTracing/RayTracingValidation.usf
@@ -17,6 +17,7 @@ uint VertexBufferOffsetInBytes;
 uint IndexBufferOffsetInBytes;
 uint IndexBufferStride;
 uint NumPrimitives;
+uint MaxVertices;
 
 bool IsNaN(float F)
 {
@@ -76,6 +77,20 @@ void RayTracingValidateGeometryBuildParamsCS(uint3 DispatchThreadId : SV_Dispatc
 		VertexBuffer, VertexBufferOffsetInBytes, VertexBufferStride,
 		PrimitiveIndex);
 	
+	{
+		// Check that index buffer value is within the expected range, based on VertexCount field of D3D12_RAYTRACING_GEOMETRY_TRIANGLES_DESC.
+		// DXR spec for VertexCount: Number of vertices (positions) in VertexBuffer. If an index buffer is present, this must be at least the maximum index value in the index buffer + 1.
+		for (uint i = 0; i < 3; ++i)
+		{
+			uint Index = Tri.Indices[i];
+			if (Index >= MaxVertices)
+			{
+				uint AssertID = 0x5ECFB346; // Unique assert ID (TODO: add support for full human-readable assert strings)
+				UEReportAssertWithPayload(AssertID, uint4(PrimitiveIndex, Index, MaxVertices, 0));
+			}
+		}
+	}
+
 	CheckTriangle(Tri, PrimitiveIndex);
 }
 
diff --git a/Engine/Source/Runtime/RenderCore/Private/RayTracingValidationShaders.cpp b/Engine/Source/Runtime/RenderCore/Private/RayTracingValidationShaders.cpp
--- a/Engine/Source/Runtime/RenderCore/Private/RayTracingValidationShaders.cpp
+++ b/Engine/Source/Runtime/RenderCore/Private/RayTracingValidationShaders.cpp
@@ -53,6 +53,7 @@ void FRayTracingValidateGeometryBuildParamsCS::Dispatch(FRHICommandList& RHICmdL
 		SetShaderValue(RHICmdList, ShaderRHI, ComputeShader->IndexBufferOffsetInBytesParam, IndexBufferOffsetInBytes);
 		SetShaderValue(RHICmdList, ShaderRHI, ComputeShader->IndexBufferStrideParam, IndexStride);
 		SetShaderValue(RHICmdList, ShaderRHI, ComputeShader->NumPrimitivesParam, Segment.NumPrimitives);
+		SetShaderValue(RHICmdList, ShaderRHI, ComputeShader->MaxVerticesParam, Segment.MaxVertices);
 
 		const FRawBufferShaderResourceViewInitializer VBViewInitializer(Segment.VertexBuffer);
 		FShaderResourceViewRHIRef VertexBufferSRV = RHICreateShaderResourceView(VBViewInitializer);
diff --git a/Engine/Source/Runtime/RenderCore/Public/RayTracingValidationShaders.h b/Engine/Source/Runtime/RenderCore/Public/RayTracingValidationShaders.h
--- a/Engine/Source/Runtime/RenderCore/Public/RayTracingValidationShaders.h
+++ b/Engine/Source/Runtime/RenderCore/Public/RayTracingValidationShaders.h
@@ -23,6 +23,7 @@ public:
 		IndexBufferOffsetInBytesParam.Bind(Initializer.ParameterMap, TEXT("IndexBufferOffsetInBytes"), SPF_Optional);
 		IndexBufferStrideParam.Bind(Initializer.ParameterMap, TEXT("IndexBufferStride"), SPF_Optional);
 		NumPrimitivesParam.Bind(Initializer.ParameterMap, TEXT("NumPrimitives"), SPF_Optional);
+		MaxVerticesParam.Bind(Initializer.ParameterMap, TEXT("MaxVertices"), SPF_Optional);
 	}
 
 	static bool ShouldCompilePermutation(const FGlobalShaderPermutationParameters& Parameters)
@@ -49,6 +50,7 @@ public:
 	LAYOUT_FIELD(FShaderParameter, IndexBufferOffsetInBytesParam);
 	LAYOUT_FIELD(FShaderParameter, IndexBufferStrideParam);
 	LAYOUT_FIELD(FShaderParameter, NumPrimitivesParam);
+	LAYOUT_FIELD(FShaderParameter, MaxVerticesParam);
 
 };
 
