Removed useless ifdef

diff --git a/Engine/Shaders/Private/LandscapeVertexFactory.ush b/Engine/Shaders/Private/LandscapeVertexFactory.ush
--- a/Engine/Shaders/Private/LandscapeVertexFactory.ush
+++ b/Engine/Shaders/Private/LandscapeVertexFactory.ush
@@ -398,11 +398,7 @@ float3x3 VertexFactoryGetPerPixelTangentBasis(FVertexFactoryInterpolantsVSToPS I
 #else
 	float4 SampleValue = Texture2DSample(LandscapeParameters.NormalmapTexture, LandscapeParameters.NormalmapTextureSampler, Interpolants.WeightHeightMapTexCoord.zw);
 #endif
-#if WEBGL
 	float2 SampleNormal = float2(SampleValue.b, SampleValue.a) * float2(2.0,2.0) - float2(1.0,1.0);
-#else
-	float2 SampleNormal = float2(SampleValue.b, SampleValue.a) * float2(2.0,2.0) - float2(1.0,1.0);
-#endif
 	float3 WorldNormal = float3( SampleNormal, sqrt(max(1.0-dot(SampleNormal,SampleNormal),0.0)) );
 
 	Result = CalcTangentBasisFromWorldNormal(WorldNormal);
