Fix typo & missing spec occlusion on second specular lobe (haziness) with Lumen.

diff --git a/Engine/Shaders/Private/DiffuseIndirectComposite.usf b/Engine/Shaders/Private/DiffuseIndirectComposite.usf
--- a/Engine/Shaders/Private/DiffuseIndirectComposite.usf
+++ b/Engine/Shaders/Private/DiffuseIndirectComposite.usf
@@ -230,7 +230,7 @@ FDirectLighting StrataLumenLighting(
 		#if STRATA_FASTPATH==0
 		if (BSDF_GETHASHAZINESS(BSDF))
 		{
-			Out.Specular = BSDFThroughput * InSpecularLighting * StrataEnvLight.SpecularHazeWeight;
+			Out.Specular += BSDFThroughput * InSpecularLighting * StrataEnvLight.SpecularHazeWeight * Occlusion.SpecularOcclusion;
 		}
 		#endif 
 	}
