Removed last remnant of the now-defunct #if WEBGL

diff --git a/Engine/Shaders/Private/MobileBasePassPixelShader.usf b/Engine/Shaders/Private/MobileBasePassPixelShader.usf
--- a/Engine/Shaders/Private/MobileBasePassPixelShader.usf
+++ b/Engine/Shaders/Private/MobileBasePassPixelShader.usf
@@ -378,12 +378,7 @@ half3 GetImageBasedReflectionLighting(FMaterialPixelParameters MaterialParameter
 	GatherSpecularIBL(MaterialParameters, MobileReflectionCapture.Texture, MobileReflectionCapture.TextureSampler, Roughness, IndirectIrradiance, MobileReflectionCapture.Params.z, ImageBasedReflections);
 #endif
 
-#if WEBGL
-	// need a rgb swizzle instead of the existing rgba swizzle, we should add it if another use case comes up. 
-	return ImageBasedReflections.bgr;
-#else
 	return ImageBasedReflections.rgb;
-#endif
 }
 
 half3 FrameBufferBlendOp(half4 Source)
