Landscape: Pull PackHeight and UnpackHeight into a common shader utility file.

diff --git a/Engine/Shaders/Private/Landscape/LandscapeCommon.ush b/Engine/Shaders/Private/Landscape/LandscapeCommon.ush
new file mode 100644
--- /dev/null
+++ b/Engine/Shaders/Private/Landscape/LandscapeCommon.ush
@@ -0,0 +1,16 @@
+// Copyright Epic Games, Inc. All Rights Reserved.
+
+#pragma once
+
+// Height is packed from a [0, 2^16-1] value to two bytes, each of which we need to provide as a 0-1 float for writing 
+float2 PackHeight(float InHeight)
+{
+	uint Height = (uint) clamp(InHeight, 0.0f, 65535.0f);
+	return float2((float) (Height >> 8) / 255.0, (float) (Height & 0xFF) / 255.0);
+}
+
+float UnpackHeight(float2 InPackedHeight)
+{
+	// Convert each of the two bytes back to [0,255], then assemble
+	return float(((int) round(InPackedHeight.r * 255.0) << 8) | (int) round(InPackedHeight.g * 255.0));
+}
\ No newline at end of file
diff --git a/Engine/Shaders/Private/Landscape/LandscapeLayersHeightmapsPS.usf b/Engine/Shaders/Private/Landscape/LandscapeLayersHeightmapsPS.usf
--- a/Engine/Shaders/Private/Landscape/LandscapeLayersHeightmapsPS.usf
+++ b/Engine/Shaders/Private/Landscape/LandscapeLayersHeightmapsPS.usf
@@ -1,6 +1,7 @@
 // Copyright Epic Games, Inc. All Rights Reserved.
 
 #include "../Common.ush"
+#include "LandscapeCommon.ush"
 
 // ----------------------------------------------------------------------------------
 
@@ -76,18 +77,6 @@ Texture2D<float4> InSourceHeightmap;
 // ----------------------------------------------------------------------------------
 // Util functions : 
 
-float2 PackHeight(float InExtractedHeight)
-{
-	InExtractedHeight = clamp(InExtractedHeight, 0.0f, 65535.0f);
-	int iHeight = (int) InExtractedHeight;
-	return float2((float) ((iHeight - (iHeight & 255)) >> 8) / 255.0, (float) (iHeight & 255) / 255.0);
-}
-
-float UnpackHeight(float2 InPackedHeight)
-{
-	return float(((int) round(InPackedHeight.r * 255.0) << 8) | (int) round(InPackedHeight.g * 255.0));
-}
-
 float ExtractAlpha(float2 InPackedAlpha)
 {
 	int AlphaWithFlags = ((int) round(InPackedAlpha.r * 255.0f) << 8) | (int) round(InPackedAlpha.g * 255.0f);
diff --git a/Engine/Shaders/Private/LandscapeLayersPS.usf b/Engine/Shaders/Private/LandscapeLayersPS.usf
--- a/Engine/Shaders/Private/LandscapeLayersPS.usf
+++ b/Engine/Shaders/Private/LandscapeLayersPS.usf
@@ -1,6 +1,7 @@
 // Copyright Epic Games, Inc. All Rights Reserved.
 
 #include "Common.ush"
+#include "Landscape/LandscapeCommon.ush"
 
 Texture2D ReadTexture1;
 Texture2D ReadTexture2;							// This one is optional, it's valid when OutputConfig.z == 1
@@ -14,14 +15,9 @@ float CurrentMipComponentVertexCount;
 float2 CurrentMipTextureSize;
 float2 ParentMipTextureSize;
 
-float ExtractHeight(float2 InPackedHeight)
-{
-	return float(((int)round(InPackedHeight.r * 255.0) << 8) | (int)round(InPackedHeight.g * 255.0));
-}
-
 float ExtractAndScaleHeight(float2 InPackedHeight)
 {
-	float Height = ExtractHeight(InPackedHeight) - 32768.0;
+	float Height = UnpackHeight(InPackedHeight) - 32768.0;
 	return Height * LandscapeGridScale.z * TERRAIN_ZSCALE;
 }
 
@@ -32,18 +28,11 @@ float ExtractAlpha(float2 InPackedAlpha)
 	return float(Alpha) / 65532.0; // 0xFFFC in float
 }
 
-float2 PackHeight(float InExtractedHeight)
-{
-	InExtractedHeight = clamp(InExtractedHeight, 0.0f, 65535.0f);
-	int iHeight = (int)InExtractedHeight;
-	return float2((float)((iHeight - (iHeight & 255)) >> 8) / 255.0, (float)(iHeight & 255) / 255.0);
-}
-
 void PSHeightmapMain(float2 InTextureCoordinates : TEXCOORD0, out float4 OutColor : SV_Target0)
 {
 	float LayerVisibility = LayerInfo.y;
 	float4 SourceColor = ReadTexture1.Sample(ReadTexture1Sampler, InTextureCoordinates);
-	float Height = ExtractHeight(SourceColor.rg);
+	float Height = UnpackHeight(SourceColor.rg);
 	
 	// Perform calculation 0 based
 	Height -= 32768.0f;
@@ -59,7 +48,7 @@ void PSHeightmapMain(float2 InTextureCoordinates : TEXCOORD0, out float4 OutColo
 	if (OutputConfig.z == 1.0)
 	{
 		float2 BaseColor = ReadTexture2.Sample(ReadTexture2Sampler, InTextureCoordinates).rg;
-		float BaseHeight = ExtractHeight(BaseColor.rg);
+		float BaseHeight = UnpackHeight(BaseColor.rg);
 		BaseHeight -= 32768.0;
 
 		float LayerBlendMode = LayerInfo.z;
@@ -219,8 +208,8 @@ void PSHeightmapMainMips(float2 InTextureCoordinates : TEXCOORD0, out float4 Out
 	{
 		float4 SourceSample = ReadTexture1.SampleLevel(ReadTexture1Sampler, SourceUV, 0);
 		float4 DownSample = ReadTexture1.SampleLevel(ReadTexture1Sampler, SourceUV + float2(0.0, ParentMipTexelSize.y), 0);
-		float SourceHeight = ExtractHeight(SourceSample.rg);
-		float DownHeight = ExtractHeight(DownSample.rg);
+		float SourceHeight = UnpackHeight(SourceSample.rg);
+		float DownHeight = UnpackHeight(DownSample.rg);
 
 		if (IsMinBorderTexelY) // on top border
 		{
@@ -243,8 +232,8 @@ void PSHeightmapMainMips(float2 InTextureCoordinates : TEXCOORD0, out float4 Out
 	{
 		float4 RightSample = ReadTexture1.SampleLevel(ReadTexture1Sampler, SourceUV + float2(ParentMipTexelSize.x, 0.0), 0);
 		float4 DownRightSample = ReadTexture1.SampleLevel(ReadTexture1Sampler, SourceUV + ParentMipTexelSize, 0);
-		float RightHeight = ExtractHeight(RightSample.rg);
-		float DownRightHeight = ExtractHeight(DownRightSample.rg);
+		float RightHeight = UnpackHeight(RightSample.rg);
+		float DownRightHeight = UnpackHeight(DownRightSample.rg);
 
 		if (IsMinBorderTexelY) // on top border
 		{
@@ -270,10 +259,10 @@ void PSHeightmapMainMips(float2 InTextureCoordinates : TEXCOORD0, out float4 Out
 		float4 RightSample = ReadTexture1.SampleLevel(ReadTexture1Sampler, SourceUV + float2(ParentMipTexelSize.x, 0.0), 0);
 		float4 DownRightSample = ReadTexture1.SampleLevel(ReadTexture1Sampler, SourceUV + ParentMipTexelSize, 0);
 
-		float SourceHeight = ExtractHeight(SourceSample.rg);
-		float DownHeight = ExtractHeight(DownSample.rg);
-		float RightHeight = ExtractHeight(RightSample.rg);
-		float DownRightHeight = ExtractHeight(DownRightSample.rg);
+		float SourceHeight = UnpackHeight(SourceSample.rg);
+		float DownHeight = UnpackHeight(DownSample.rg);
+		float RightHeight = UnpackHeight(RightSample.rg);
+		float DownRightHeight = UnpackHeight(DownRightSample.rg);
 
 		if (IsMinBorderTexelY) // on top border
 		{
