Strata SLW now works correctly with separated main light luminance for distance field shadow.
SLW was also broken since the strata tree refactor.

diff --git a/Engine/Shaders/Private/BasePassPixelShader.usf b/Engine/Shaders/Private/BasePassPixelShader.usf
--- a/Engine/Shaders/Private/BasePassPixelShader.usf
+++ b/Engine/Shaders/Private/BasePassPixelShader.usf
@@ -1131,8 +1131,8 @@ void FPixelShaderInOut_MainPS(
 
 	// Override GBuffer data with Strata SLW water BSDF to run forward shadfing code.
 	// STRATA_TODO: run the shading through a strata path? (by adding a special BSDF?)
-	SanitizeStrataSingleLayerWater(PixelMaterialInputs.FrontMaterial.InlinedBSDF);
-	FStrataBSDF SLWBSDF = PixelMaterialInputs.FrontMaterial.InlinedBSDF;
+	SanitizeStrataSingleLayerWater(StrataPixelHeader.StrataTree.BSDFs[0]);
+	FStrataBSDF SLWBSDF = StrataPixelHeader.StrataTree.BSDFs[0];
 
 	BaseColor = SLW_BASECOLOR(SLWBSDF);
 	Metallic = SLW_METALLIC(SLWBSDF);
@@ -1967,14 +1967,28 @@ void FPixelShaderInOut_MainPS(
 				Out.MRT[GBUFFER_HAS_VELOCITY ? 6 : 5] = 0;
 			}
 
-	#if SINGLE_LAYER_WATER_SEPARATED_MAIN_LIGHT
+		#if SINGLE_LAYER_WATER_SEPARATED_MAIN_LIGHT && STRATA_ENABLED
+			if (bSeparateWaterMainDirLightLuminance)
+			{
+				// STRATA_TODO: we should have a define set to the correct output index.
+				#if PIXELSHADEROUTPUT_MRT2
+					Out.MRT[2] = float4(SeparatedWaterMainDirLightLuminance, 0.0f);
+				#elif PIXELSHADEROUTPUT_MRT1
+					Out.MRT[1] = float4(SeparatedWaterMainDirLightLuminance, 0.0f);
+				#else
+					Out.MRT[0] = float4(SeparatedWaterMainDirLightLuminance, 0.0f);
+				#endif
+			}
+		#endif
+
+		#if SINGLE_LAYER_WATER_SEPARATED_MAIN_LIGHT && !STRATA_ENABLED
 			if (GBuffer.ShadingModelID == SHADINGMODELID_SINGLELAYERWATER)
 			{
 				// In deferred, we always output the directional light in a separated buffer.
 				// This is used to apply distance field shadows or light function to the main directional light.
 				Out.MRT[GBUFFER_HAS_VELOCITY ? 7 : 6] = float4(SeparatedWaterMainDirLightLuminance * View.PreExposure, 1.0f);
 			}
-	#endif
+		#endif
 		}
 		#else
 		{
diff --git a/Engine/Shaders/Private/SingleLayerWaterShading.ush b/Engine/Shaders/Private/SingleLayerWaterShading.ush
--- a/Engine/Shaders/Private/SingleLayerWaterShading.ush
+++ b/Engine/Shaders/Private/SingleLayerWaterShading.ush
@@ -105,7 +105,8 @@ WaterVolumeLightingOutput EvaluateWaterVolumeLighting(
 		const float BehindWaterDeltaDepth = CameraIsUnderWater ? WaterDepth : max(0.0f, SceneDepth - WaterDepth);
 
 #if STRATA_INLINE_SINGLELAYERWATER
-		FStrataBSDF SLWBSDF = PixelMaterialInputs.FrontMaterial.InlinedBSDF;
+		SanitizeStrataSingleLayerWater(MaterialParameters.StrataTree.BSDFs[0]);
+		FStrataBSDF SLWBSDF = MaterialParameters.StrataTree.BSDFs[0];
 		const float3 Albedo = SLW_WATERALBEDO(SLWBSDF);
 		const float3 ExtinctionCoeff = SLW_WATEREXTINCTION(SLWBSDF);
 		const float3 ScatteringCoeff = Albedo * ExtinctionCoeff;
diff --git a/Engine/Shaders/Private/Strata/StrataLegacyConversion.ush b/Engine/Shaders/Private/Strata/StrataLegacyConversion.ush
--- a/Engine/Shaders/Private/Strata/StrataLegacyConversion.ush
+++ b/Engine/Shaders/Private/Strata/StrataLegacyConversion.ush
@@ -135,7 +135,7 @@ FStrataData StrataConvertLegacyMaterialDynamic(
 		const float3 WaterExtinction = WaterScatteringCoefficients + WaterAbsorptionCoefficients;
 		const float3 WaterAlbedo = WaterScatteringCoefficients / WaterExtinction;
 
-		return GetStrataSingleLayerWaterBSDF(
+		FStrataData Top_SLW = GetStrataSingleLayerWaterBSDF(
 			BaseColor, 				// BaseColor
 			Metallic, 				// Metallic 
 			Specular, 				// Specular
@@ -148,7 +148,14 @@ FStrataData StrataConvertLegacyMaterialDynamic(
 			ColorScaleBehindWater,	// ColorScaleBehindWater
 			SharedLocalBasisIndex);
 
-		// No need to take into account opacity weighting for translucent material, as the coverage is assumed to be 1
+		FStrataLegacyParameters Bot = InitStrataLegacyParameters(PixelFootprint, SharedLocalBasisIndex, 1.0f);
+
+		FStrataData Bot_Slab = CreateLegacySlab(Bot, SharedLocalBasisTypes, StrataTree,	/*OperatorIndex*/1, /*BSDFIndex*/1);
+		FStrataData WTop_Slab = StrataWeight(Top_SLW , 1.0, StrataTree,					/*OperatorIndex*/2, /*MaxDistanceFromLeaves*/1);
+		FStrataData WBot_Slab = StrataWeight(Bot_Slab, 0.0, StrataTree,					/*OperatorIndex*/3, /*MaxDistanceFromLeaves*/1);
+		FStrataData Out = StrataVerticalLayering(WTop_Slab, WBot_Slab, StrataTree,		/*OperatorIndex*/4, /*MaxDistanceFromLeaves*/2);
+		// Out maps to LegacyUpdateBSDFsOperators
+		return Out;
 	}
 	#endif
 	#if MATERIAL_SHADINGMODEL_HAIR
@@ -522,6 +529,7 @@ FStrataData StrataConvertLegacyMaterialDynamic(
 	Out.InlinedBSDF = ParamBlendVertical.InlinedBSDF;
 #endif
 
+	// Out maps to LegacyUpdateBSDFsOperators
 	return Out;
 }
 
@@ -600,20 +608,32 @@ FStrataData StrataConvertLegacyMaterialStatic(
 		const float3 WaterExtinction = WaterScatteringCoefficients + WaterAbsorptionCoefficients;
 		const float3 WaterAlbedo = WaterScatteringCoefficients / WaterExtinction;
 
-		Out = GetStrataSingleLayerWaterBSDF(
-			BaseColor, 				// BaseColor
-			Metallic, 				// Metallic 
-			Specular, 				// Specular
-			Roughness, 				// Roughness 
-			Emissive, 				// Emissive 
-			Opacity,				// TopMaterialOpacity
-			WaterAlbedo, 			// WaterAlbedo
-			WaterExtinction, 		// WaterExtinction
-			WaterPhaseG, 			// WaterPhaseG
-			ColorScaleBehindWater,	// ColorScaleBehindWater
-			SharedLocalBasisIndex);
-
-		// No need to take into account opacity weighting for translucent material, as the coverage is assumed to be 1
+		FStrataData Top_SLW = PromoteParameterBlendedBSDFToOperator(
+			GetStrataSingleLayerWaterBSDF(
+				BaseColor, 				// BaseColor
+				Metallic, 				// Metallic 
+				Specular, 				// Specular
+				Roughness, 				// Roughness 
+				Emissive, 				// Emissive 
+				Opacity,				// TopMaterialOpacity
+				WaterAlbedo, 			// WaterAlbedo
+				WaterExtinction, 		// WaterExtinction
+				WaterPhaseG, 			// WaterPhaseG
+				ColorScaleBehindWater,	// ColorScaleBehindWater
+				SharedLocalBasisIndex),
+			StrataTree,
+			0		/*OperatorIndex*/,
+			0		/*BSDFIndex*/,
+			true	/*In.bIsBottom*/,
+			true	/*In.bIsTop*/);
+
+		FStrataLegacyParameters Bot = InitStrataLegacyParameters(PixelFootprint, SharedLocalBasisIndex, 1.0f);
+		FStrataData Bot_Slab = CreateLegacySlab(Bot, SharedLocalBasisTypes, StrataTree,	/*OperatorIndex*/1, /*BSDFIndex*/1);
+
+		FStrataData WTop_Slab = StrataWeight(Top_SLW, 1.0, StrataTree,					/*OperatorIndex*/2, /*MaxDistanceFromLeaves*/1);
+		FStrataData WBot_Slab = StrataWeight(Bot_Slab, 0.0, StrataTree,					/*OperatorIndex*/3, /*MaxDistanceFromLeaves*/1);
+		Out = StrataVerticalLayering(WTop_Slab, WBot_Slab, StrataTree,					/*OperatorIndex*/4, /*MaxDistanceFromLeaves*/2);
+		// Out maps to LegacyUpdateBSDFsOperators
 	}
 	else
 	{
diff --git a/Engine/Source/Runtime/Renderer/Private/SingleLayerWaterRendering.cpp b/Engine/Source/Runtime/Renderer/Private/SingleLayerWaterRendering.cpp
--- a/Engine/Source/Runtime/Renderer/Private/SingleLayerWaterRendering.cpp
+++ b/Engine/Source/Runtime/Renderer/Private/SingleLayerWaterRendering.cpp
@@ -809,6 +809,7 @@ void FDeferredShadingSceneRenderer::RenderSingleLayerWaterInner(
 		const bool bNeverClear = true;
 		BasePassTextures[BasePassTextureCount++] = FTextureRenderTargetBinding(SceneWithoutWaterTextures.SeparatedMainDirLightTexture, bNeverClear);
 	}
+	Strata::AppendStrataMRTs(*this, BasePassTextureCount, BasePassTextures);
 	TArrayView<FTextureRenderTargetBinding> BasePassTexturesView = MakeArrayView(BasePassTextures.GetData(), BasePassTextureCount);
 
 	FRDGTextureRef WhiteForwardScreenSpaceShadowMask = SystemTextures.White;
