Fix for debugprobe not compiling on some platforms.

diff --git a/Engine/Shaders/Private/DebugProbes.usf b/Engine/Shaders/Private/DebugProbes.usf
--- a/Engine/Shaders/Private/DebugProbes.usf
+++ b/Engine/Shaders/Private/DebugProbes.usf
@@ -162,7 +162,7 @@ void MainPS(
 		tMax.z *= tDelta.z;
 
 		LOOP
-		while (all(dot(StartP - P, StartP - P) <= 5.1f * 5.1f))	// stop after a distance of 5 voxels is reached 
+		while (dot(StartP - P, StartP - P) <= 5.1f * 5.1f)	// stop after a distance of 5 voxels is reached 
 		{
 			const float3 SphereCenter = BoxMin + (floor(P) + 0.5) * VoxelSizeCm;
 
