Strata - fixes for SLW with Strata

diff --git a/Engine/Shaders/Private/Strata/Strata.ush b/Engine/Shaders/Private/Strata/Strata.ush
--- a/Engine/Shaders/Private/Strata/Strata.ush
+++ b/Engine/Shaders/Private/Strata/Strata.ush
@@ -2767,6 +2767,7 @@ void PackStrataOut(
 					case STRATA_BSDF_TYPE_HAIR:
 					{
 						bIsSimpleMaterial = false;
+						bIsSingleMaterial = false;
 
 						TopLayerTotalWeight += BSDF.TopLayerDataWeight;
 						TopLayerData.Roughness += BSDF.TopLayerDataWeight * FullyRough;
@@ -2777,6 +2778,7 @@ void PackStrataOut(
 					case STRATA_BSDF_TYPE_SINGLELAYERWATER:
 					{
 						bIsSimpleMaterial = false;
+						bIsSingleMaterial = false;
 						EmissiveColor += BSDF_GETEMISSIVE(BSDF) * BSDF.LuminanceWeight;
 
 						TopLayerTotalWeight += BSDF.TopLayerDataWeight;
diff --git a/Engine/Shaders/Private/Strata/StrataVisualize.usf b/Engine/Shaders/Private/Strata/StrataVisualize.usf
--- a/Engine/Shaders/Private/Strata/StrataVisualize.usf
+++ b/Engine/Shaders/Private/Strata/StrataVisualize.usf
@@ -390,6 +390,8 @@ void VisualizeMaterialPS(
 		HORIZONTALSEPARATION;
 #endif
 
+
+		Header.BSDFCount = min(Header.BSDFCount, STRATA_MAX_BSDF_COUNT);
 		UNROLL_N(STRATA_MAX_BSDF_COUNT) // FXC requires that we specify an upper bound here. Not DXC.
 		for (uint BSDFIndex = 0; BSDFIndex < Header.BSDFCount; ++BSDFIndex)
 		{
@@ -572,16 +574,15 @@ void VisualizeMaterialPS(
 			case STRATA_BSDF_TYPE_SINGLELAYERWATER:
 #if PERMUTATION_BSDF_PASS==2
 				PRINTBSDFCHAR(_W_); PRINTBSDFCHAR(_A_); PRINTBSDFCHAR(_T_); PRINTBSDFCHAR(_E_); PRINTBSDFCHAR(_R_);
-				PRINTBSDFCHAR(_SPC_); PRINTBSDFCHAR(_SPC_); PRINTBSDFCHAR(_SPC_); PRINTBSDFCHAR(_SPC_); PRINTBSDFCHAR(_SPC_);
-				PRINTBSDFCHAR(_T_); PRINTBSDFCHAR(_O_); PRINTBSDFCHAR(_D_); PRINTBSDFCHAR(_O_); PRINTBSDFCHAR(_SPC_);
+				NEXTLINE;
 
 				HORIZONTALTAB; PRINTPARAMCHAR(_A_); PRINTPARAMCHAR(_L_); PRINTPARAMCHAR(_B_); PRINTPARAMCHAR(_E_); PRINTPARAMCHAR(_D_); PRINTPARAMCHAR(_O_);
 				HORIZONTALTAB; PRINTFLOAT3(SLW_BASECOLOR(BSDF));
 				NEXTLINE;
-				HORIZONTALTAB; PRINTPARAMCHAR(_M_); PRINTPARAMCHAR(_E_); PRINTPARAMCHAR(_T_); PRINTPARAMCHAR(_A_); PRINTPARAMCHAR(_L_); PRINTPARAMCHAR(_O_);
+				HORIZONTALTAB; PRINTPARAMCHAR(_M_); PRINTPARAMCHAR(_E_); PRINTPARAMCHAR(_T_); PRINTPARAMCHAR(_A_); PRINTPARAMCHAR(_L_); PRINTPARAMCHAR(_SPC_);
 				HORIZONTALTAB; PRINTFLOAT1(SLW_METALLIC(BSDF));
 				NEXTLINE;
-				HORIZONTALTAB; PRINTPARAMCHAR(_S_); PRINTPARAMCHAR(_P_); PRINTPARAMCHAR(_E_); PRINTPARAMCHAR(_C_); PRINTPARAMCHAR(_U_); PRINTPARAMCHAR(_SPC_);
+				HORIZONTALTAB; PRINTPARAMCHAR(_S_); PRINTPARAMCHAR(_P_); PRINTPARAMCHAR(_E_); PRINTPARAMCHAR(_C_); PRINTPARAMCHAR(_SPC_); PRINTPARAMCHAR(_SPC_);
 				HORIZONTALTAB; PRINTFLOAT1(SLW_SPECULAR(BSDF));
 				NEXTLINE;
 				HORIZONTALTAB; PRINTPARAMCHAR(_R_); PRINTPARAMCHAR(_O_); PRINTPARAMCHAR(_U_); PRINTPARAMCHAR(_G_); PRINTPARAMCHAR(_H_); PRINTPARAMCHAR(_SPC_);
@@ -596,8 +597,14 @@ void VisualizeMaterialPS(
 				NEXTLINE;
 				NEXTLINE;
 				NEXTLINE;
+				NEXTLINE;
 #endif
 				break;
+
+			default:
+				PRINTBSDFCHAR(_E_); PRINTBSDFCHAR(_R_); PRINTBSDFCHAR(_R_); PRINTBSDFCHAR(_O_); PRINTBSDFCHAR(_R_);
+				NEXTLINE;
+				break;
 			}
 
 			HORIZONTALSEPARATION;
