Stata - fixe for wrong specular/f0 usage with SLW.

diff --git a/Engine/Shaders/Private/SingleLayerWaterComposite.usf b/Engine/Shaders/Private/SingleLayerWaterComposite.usf
--- a/Engine/Shaders/Private/SingleLayerWaterComposite.usf
+++ b/Engine/Shaders/Private/SingleLayerWaterComposite.usf
@@ -119,7 +119,7 @@ SingleLayerWaterCompositeOutput SingleLayerWaterComposite(float2 BufferUV, float
 			const float3 BSDFCoverage = 1.0f;	// Since the SLW BSDF is always isolated and composited onto the strata material buffer, this is assumed to be 1 here.
 
 			// The specular over the water in completely in the hand of the user. We do not fade out metalness for instance.
-			float3 EnvBrdfValue = BSDFCoverage * EnvBRDF(SLW_SPECULAR(SLWaterBSDF), SafeRoughness, saturate(NoV));
+			float3 EnvBrdfValue = BSDFCoverage * EnvBRDF(F0, SafeRoughness, saturate(NoV));
 
 #else // STRATA_ENABLED
 
