Strata visualize now adapts to dxc (loop) and fxc (only 2 bsdf becauseit cannot handle loop and unroll is sooooo sloooooow).

diff --git a/Engine/Shaders/Private/Strata/StrataVisualize.usf b/Engine/Shaders/Private/Strata/StrataVisualize.usf
--- a/Engine/Shaders/Private/Strata/StrataVisualize.usf
+++ b/Engine/Shaders/Private/Strata/StrataVisualize.usf
@@ -11,6 +11,8 @@
 #include "/Engine/Private/MiniFontCommon.ush"
 #include "/Engine/Private/ShaderPrint.ush"
 
+// This is for the fxc compiler to compile faster.
+#define STRATA_VIZ_MAX_BSDF 2
 
 void DrawReferentialTWS(float3 P /*In Translated World Space*/, float3 X, float3 Y, float3 N, float3 InColor)
 {
@@ -391,9 +393,14 @@ void VisualizeMaterialPS(
 #endif
 
 
-		Header.BSDFCount = min(Header.BSDFCount, STRATA_MAX_BSDF_COUNT);
-		UNROLL_N(STRATA_MAX_BSDF_COUNT) // FXC requires that we specify an upper bound here. Not DXC.
+#if COMPILER_DXC == 1
+		LOOP
 		for (uint BSDFIndex = 0; BSDFIndex < Header.BSDFCount; ++BSDFIndex)
+#else
+		Header.BSDFCount = min(Header.BSDFCount, STRATA_VIZ_MAX_BSDF);
+		UNROLL_N(STRATA_VIZ_MAX_BSDF) // FXC requires that we specify an upper bound here. Not DXC.
+		for (uint BSDFIndex = 0; BSDFIndex < Header.BSDFCount; ++BSDFIndex)
+#endif
 		{
 			const FStrataBSDF BSDF = UnpackStrataBSDFIn(Strata.MaterialTextureArray, StrataAddressing, Header);
 
