Improve the specular matching between raster and path tracer for subsurface scattering.

1. Use the decoding used by raster for the path tracer.
2. Use direct encoding from material roughness to lobe roughness instead of going through the average, which leads to underestimation (2%) and clamping (20%) of the value space.

Also clean up the code a little.

diff --git a/Engine/Shaders/Private/MobileShadingModels.ush b/Engine/Shaders/Private/MobileShadingModels.ush
--- a/Engine/Shaders/Private/MobileShadingModels.ush
+++ b/Engine/Shaders/Private/MobileShadingModels.ush
@@ -129,11 +129,18 @@ half CalculateCurvature(half3 WorldNormal, float3 WorldPosition)
 	return CurvatureApprox;
 }
 
-void GetProfileDualSpecular(uint SubsurfaceProfileInt, out half2 AverageToRoughness, out half LobeMix)
+void GetProfileDualSpecular(uint SubsurfaceProfileInt, half Roughness, half Opacity, out half2 LobeRoughness, out half LobeMix)
 {
 	half4 Data = GetSubsurfaceProfileTexture(SSSS_DUAL_SPECULAR_OFFSET, SubsurfaceProfileInt);
-	AverageToRoughness = Data.xy * SSSS_MAX_DUAL_SPECULAR_ROUGHNESS;
-	LobeMix = Data.z;
+	half2 MaterialRoughnessToLobeRoughness = Data.xy * SSSS_MAX_DUAL_SPECULAR_ROUGHNESS;
+	LobeMix = Data.z; 
+
+	half RoughnessLerpFactor = saturate(Opacity * 10.0f);
+	MaterialRoughnessToLobeRoughness.x = lerp(1.0f, MaterialRoughnessToLobeRoughness.x, RoughnessLerpFactor);
+	MaterialRoughnessToLobeRoughness.y = lerp(1.0f, MaterialRoughnessToLobeRoughness.y, RoughnessLerpFactor);
+
+	LobeRoughness.x = max(saturate(Roughness * MaterialRoughnessToLobeRoughness.x), 0.02f);
+	LobeRoughness.y = saturate(Roughness * MaterialRoughnessToLobeRoughness.y);
 }
 
 half4 GetSSProfilePreIntegratedValue(uint SubsurfaceProfileInt, half NoL, half Curvature)
@@ -204,17 +211,12 @@ void InitShadingModelContext(inout FMobileShadingModelContext ShadingModelContex
 				
 		ShadingModelContext.Curvature = GBuffer.CustomData.g;
 
-		half2 AverageToRoughness;
-		GetProfileDualSpecular(ShadingModelContext.SubsurfaceProfileInt, AverageToRoughness, ShadingModelContext.LobeMix);
+		half2 LobeRoughness;
+		GetProfileDualSpecular(ShadingModelContext.SubsurfaceProfileInt, GBuffer.Roughness, 
+			ShadingModelContext.Opacity, LobeRoughness, ShadingModelContext.LobeMix);
 
-		half AverageRoughness = GBuffer.Roughness;
-		ShadingModelContext.Lobe0Roughness = max(saturate(AverageRoughness * AverageToRoughness.x), 0.02f);
-		ShadingModelContext.Lobe1Roughness = saturate(AverageRoughness * AverageToRoughness.y);
-
-		// Smoothly lerp to default single GGX lobe as Opacity approaches 0
-		half RoughnessLerpFactor = saturate(ShadingModelContext.Opacity * 10.0f);
-		ShadingModelContext.Lobe0Roughness = lerp(1.0f, ShadingModelContext.Lobe0Roughness, RoughnessLerpFactor);
-		ShadingModelContext.Lobe1Roughness = lerp(1.0f, ShadingModelContext.Lobe1Roughness, RoughnessLerpFactor);
+		ShadingModelContext.Lobe0Roughness = LobeRoughness.x;
+		ShadingModelContext.Lobe1Roughness = LobeRoughness.y;
 	}
 	else if (GBuffer.ShadingModelID == SHADINGMODELID_HAIR)
 	{
@@ -413,7 +415,9 @@ FMobileDirectLighting MobileSubsurfaceProfileBxDF(FMobileShadingModelContext Sha
 
 	Lighting.Diffuse = BurleyDiffuse.rgb * ShadingModelContext.DiffuseColor;
 
-	Lighting.Specular = Shadow * NoL * MobileDualSpecularGGX(GBuffer.Roughness, ShadingModelContext.Lobe0Roughness, ShadingModelContext.Lobe1Roughness, ShadingModelContext.LobeMix, ShadingModelContext.EnvBrdf, ShadingModelContext.SpecularColor, ShadingModelContext.NoV, NoH, NoL, VoH);
+	half AverageRoughness = lerp(ShadingModelContext.Lobe0Roughness, ShadingModelContext.Lobe1Roughness, ShadingModelContext.LobeMix);
+
+	Lighting.Specular = Shadow * NoL * MobileDualSpecularGGX(AverageRoughness, ShadingModelContext.Lobe0Roughness, ShadingModelContext.Lobe1Roughness, ShadingModelContext.LobeMix, ShadingModelContext.EnvBrdf, ShadingModelContext.SpecularColor, ShadingModelContext.NoV, NoH, NoL, VoH);
 
 	return Lighting;
 }
diff --git a/Engine/Shaders/Private/PathTracing/Material/PathTracingMaterialSampling.ush b/Engine/Shaders/Private/PathTracing/Material/PathTracingMaterialSampling.ush
--- a/Engine/Shaders/Private/PathTracing/Material/PathTracingMaterialSampling.ush
+++ b/Engine/Shaders/Private/PathTracing/Material/PathTracingMaterialSampling.ush
@@ -423,18 +423,11 @@ FPathTracingPayload UnpackPathTracingPayloadMaterial(FPackedPathTracingPayload P
 				// TODO: non-burley profile allows a blend of diffuse and SSS which should be accounted for here
 			}
 			// Decode dual roughness info
+			float Lobe0Roughness = 0;
+			float Lobe1Roughness = 0;
+			float LobeMix = 0;
 
-			// See logic in ShadingModels.ush
-			float4 Data = View.SSProfilesTexture.Load(int3(SSSS_DUAL_SPECULAR_OFFSET, ProfileId, 0));
-			float AverageToRoughness0 = Data.x * SSSS_MAX_DUAL_SPECULAR_ROUGHNESS;
-			float AverageToRoughness1 = Data.y * SSSS_MAX_DUAL_SPECULAR_ROUGHNESS;
-			float LobeMix = Data.z;
-			float AverageRoughness = HitPayload.Roughness;
-			float Lobe0Roughness = max(saturate(AverageRoughness * AverageToRoughness0), 0.02f);
-			float Lobe1Roughness = saturate(AverageRoughness * AverageToRoughness1);
-
-			Lobe0Roughness = lerp(1.0f, Lobe0Roughness, saturate(HitPayload.Opacity * 10.0f));
-			Lobe1Roughness = lerp(1.0f, Lobe1Roughness, saturate(HitPayload.Opacity * 10.0f));
+			GetSubsurfaceProfileDualSpecular(ProfileId, HitPayload.Roughness, HitPayload.Opacity, Lobe0Roughness, Lobe1Roughness, LobeMix);
 
 			HitPayload.SetDualRoughnessSpecular(float3(Lobe0Roughness, Lobe1Roughness, LobeMix));
 			
diff --git a/Engine/Shaders/Private/ShadingModels.ush b/Engine/Shaders/Private/ShadingModels.ush
--- a/Engine/Shaders/Private/ShadingModels.ush
+++ b/Engine/Shaders/Private/ShadingModels.ush
@@ -562,16 +562,14 @@ float ApproximateHG(float cosJ, float g)
 	return (ISO_PHASE_FUNC_Normalized * (1.0f - g2) / max( 1e-5, gcos2));
 }
 
-void GetProfileDualSpecular(FGBufferData GBuffer, inout float AverageToRoughness0, inout float AverageToRoughness1, inout float LobeMix)
+void GetProfileDualSpecular(uint SubsurfaceProfileInt, float Roughness, float Opacity, out float LobeRoughness0, out float LobeRoughness1, out float LobeMix)
 {
 #if !FORWARD_SHADING
-	float DummyOpacity = 0;
-	float DummyMaterialRoughnessToAverage = 0;
-	GetSubsurfaceProfileDualSpecular(ExtractSubsurfaceProfileInt(GBuffer), DummyOpacity, AverageToRoughness0, AverageToRoughness1, LobeMix, DummyMaterialRoughnessToAverage);
+	GetSubsurfaceProfileDualSpecular(SubsurfaceProfileInt, Roughness, Opacity, LobeRoughness0, LobeRoughness1, LobeMix);
 #else
 	// Disable dual lobe, as subsurface profile doesn't work with forward.
-	AverageToRoughness0 = 1.f;
-	AverageToRoughness1 = 1.f;
+	Lobe0Roughness = Roughness;
+	Lobe1Roughness = Roughness;
 	LobeMix = 0.f;
 #endif
 }
@@ -583,17 +581,16 @@ FDirectLighting SubsurfaceProfileBxDF( FGBufferData GBuffer, half3 N, half3 V, h
 	SphereMaxNoH( Context, AreaLight.SphereSinAlpha, true );
 	Context.NoV = saturate( abs( Context.NoV ) + 1e-5 );
 
-	float AverageToRoughness0;
-	float AverageToRoughness1;
-	float LobeMix;
-	GetProfileDualSpecular(GBuffer, AverageToRoughness0, AverageToRoughness1, LobeMix);
+	uint SubsurfaceProfileId = ExtractSubsurfaceProfileInt(GBuffer);
+	float Opacity = GBuffer.CustomData.a;
+	float Roughness = GBuffer.Roughness;
 
-	float AverageRoughness = GBuffer.Roughness;
-	float Lobe0Roughness = max(saturate(AverageRoughness * AverageToRoughness0), 0.02f);
-	float Lobe1Roughness = saturate(AverageRoughness * AverageToRoughness1);
+	float Lobe0Roughness = 0;
+	float Lobe1Roughness = 0;
+	float LobeMix = 0;
 
-	const float Opacity = GBuffer.CustomData.a;
-	GetSubsurfaceProfileRoughnesses(Lobe0Roughness, Lobe1Roughness, Opacity);
+	GetProfileDualSpecular(SubsurfaceProfileId, Roughness, Opacity, Lobe0Roughness, Lobe1Roughness, LobeMix);
+	float AverageRoughness = lerp(Lobe0Roughness, Lobe1Roughness, LobeMix);
 
 	// Take the average roughness instead of compute a separate energy term for each roughness
 	const FBxDFEnergyTerms EnergyTerms = ComputeGGXSpecEnergyTerms(AverageRoughness, Context.NoV, GBuffer.SpecularColor);
diff --git a/Engine/Shaders/Private/ShadingModelsMaterial.ush b/Engine/Shaders/Private/ShadingModelsMaterial.ush
--- a/Engine/Shaders/Private/ShadingModelsMaterial.ush
+++ b/Engine/Shaders/Private/ShadingModelsMaterial.ush
@@ -62,17 +62,6 @@ void SetGBufferForShadingModel(
 		{
 			GBuffer.CustomData.rgb = EncodeSubsurfaceProfile(SubsurfaceProfile);
 			GBuffer.CustomData.a = Opacity;
-
-			// Average roughness for dual specular.
-			uint SubsurfaceProfileUInt = uint(SubsurfaceProfile * 255.0f + 0.5f);
-
-			float AverageToRoughness0 = 0;
-			float AverageToRoughness1 = 0;
-			float LobeMix = 0;
-			float MaterialRoughnessToAverage = 0;
-			GetSubsurfaceProfileDualSpecular(SubsurfaceProfileUInt, Opacity, AverageToRoughness0, AverageToRoughness1, LobeMix, MaterialRoughnessToAverage);
-
-			GBuffer.Roughness = saturate(GBuffer.Roughness * MaterialRoughnessToAverage);
 		}
 	#if SUBSURFACE_PROFILE_OPACITY_THRESHOLD
 		else
diff --git a/Engine/Shaders/Private/Strata/Strata.ush b/Engine/Shaders/Private/Strata/Strata.ush
--- a/Engine/Shaders/Private/Strata/Strata.ush
+++ b/Engine/Shaders/Private/Strata/Strata.ush
@@ -877,17 +877,11 @@ FStrataData GetStrataSlabBSDF(
 		// Average roughness for dual specular.
 		const uint SubsurfaceProfileUInt = StrataSubsurfaceProfileIdTo8bits(SSSProfileID);
 
-		float AverageToRoughness0 = 0;
-		float AverageToRoughness1 = 0;
+		float LobeRoughness0 = 0;
+		float LobeRoughness1 = 0;
 		float LobeMix = 0;
-		float MaterialRoughnessToAverage = 0;
-		GetSubsurfaceProfileDualSpecular(SubsurfaceProfileUInt, SSSMFPScale, AverageToRoughness0, AverageToRoughness1, LobeMix, MaterialRoughnessToAverage);
 
-		float AverageRoughness = saturate(Roughness * MaterialRoughnessToAverage); // This is the roughness originallys stored into the gbuffer
-		float Lobe0Roughness   = saturate(AverageRoughness * AverageToRoughness0);
-		float Lobe1Roughness   = saturate(AverageRoughness * AverageToRoughness1);
-
-		GetSubsurfaceProfileRoughnesses(Lobe0Roughness, Lobe1Roughness, SSSMFPScale);
+		GetSubsurfaceProfileDualSpecular(SubsurfaceProfileUInt, Roughness, SSSMFPScale, LobeRoughness0, LobeRoughness1, LobeMix);
 
 		// Convertion from roughness0/1 to roughness+haziness
 		// * This encoding is based on how roughness0 + haziness is translated to roughness1, see ComputeHazyLobeRoughness
diff --git a/Engine/Shaders/Private/SubsurfaceProfileCommon.ush b/Engine/Shaders/Private/SubsurfaceProfileCommon.ush
--- a/Engine/Shaders/Private/SubsurfaceProfileCommon.ush
+++ b/Engine/Shaders/Private/SubsurfaceProfileCommon.ush
@@ -76,24 +76,19 @@ half4 GetSubsurfaceProfileTexture(Texture2D InSSProfilesTexture, SamplerState In
 #endif 
 }
 
-void GetSubsurfaceProfileDualSpecular(uint SubsurfaceProfileInt, float Opacity, inout float AverageToRoughness0, inout float AverageToRoughness1, inout float LobeMix, inout float MaterialRoughnessToAverage)
+void GetSubsurfaceProfileDualSpecular(uint SubsurfaceProfileInt, float Roughness, float Opacity, out float LobeRoughness0, out float LobeRoughness1, out float LobeMix)
 {
 	const half4 Data = GetSubsurfaceProfileTexture(SSSS_DUAL_SPECULAR_OFFSET, SubsurfaceProfileInt);
-	AverageToRoughness0 = Data.x * SSSS_MAX_DUAL_SPECULAR_ROUGHNESS;
-	AverageToRoughness1 = Data.y * SSSS_MAX_DUAL_SPECULAR_ROUGHNESS;
-	LobeMix = Data.z;
 
 	// Smooth blend out dual specular when opacity is low, we have the extra SSSS_OPACITY_THRESHOLD_EPS so that we fade out by the time we
 	// get to 0.01, as opposed to 0.0.
-	MaterialRoughnessToAverage = lerp(1.0f, Data.w * SSSS_MAX_DUAL_SPECULAR_ROUGHNESS, saturate((Opacity - SSSS_OPACITY_THRESHOLD_EPS) * 10.0f));
-}
+	float MaterialRoughnessToLobeRoughness0 = lerp(1.0f, Data.x * SSSS_MAX_DUAL_SPECULAR_ROUGHNESS, saturate((Opacity - SSSS_OPACITY_THRESHOLD_EPS) * 10.0f));
+	float MaterialRoughnessToLobeRoughness1 = lerp(1.0f, Data.y * SSSS_MAX_DUAL_SPECULAR_ROUGHNESS, saturate((Opacity - SSSS_OPACITY_THRESHOLD_EPS) * 10.0f));
+	LobeMix = Data.z;
 
-void GetSubsurfaceProfileRoughnesses(inout float Lobe0Roughness, inout float Lobe1Roughness, float Opacity)
-{
-	// Smoothly lerp to default single GGX lobe as Opacity approaches 0, before reverting to SHADINGMODELID_DEFAULT_LIT.
-	// See SUBSURFACE_PROFILE_OPACITY_THRESHOLD in ShadingModelsMaterial.ush.
-	Lobe0Roughness = lerp(1.0f, Lobe0Roughness, saturate(Opacity * 10.0f));
-	Lobe1Roughness = lerp(1.0f, Lobe1Roughness, saturate(Opacity * 10.0f));
+	// Avoid specular explosion, and approximate `area lights`.
+	LobeRoughness0 = max(saturate(Roughness * MaterialRoughnessToLobeRoughness0), 0.02f);
+	LobeRoughness1 = saturate(Roughness * MaterialRoughnessToLobeRoughness1);
 }
 
 // Surface albedo and mean free path length
diff --git a/Engine/Source/Runtime/Engine/Private/Rendering/SubsurfaceProfile.cpp b/Engine/Source/Runtime/Engine/Private/Rendering/SubsurfaceProfile.cpp
--- a/Engine/Source/Runtime/Engine/Private/Rendering/SubsurfaceProfile.cpp
+++ b/Engine/Source/Runtime/Engine/Private/Rendering/SubsurfaceProfile.cpp
@@ -485,11 +485,9 @@ void FSubsurfaceProfileTexture::CreateTexture(FRHICommandListImmediate& RHICmdLi
 		TextureRow[SSSS_BOUNDARY_COLOR_BLEED_OFFSET].A = Data.bEnableBurley ? SSS_TYPE_BURLEY : SSS_TYPE_SSSS;
 
 		float MaterialRoughnessToAverage = Data.Roughness0 * (1.0f - Data.LobeMix) + Data.Roughness1 * Data.LobeMix;
-		float AverageToRoughness0 = Data.Roughness0 / MaterialRoughnessToAverage;
-		float AverageToRoughness1 = Data.Roughness1 / MaterialRoughnessToAverage;
 
-		TextureRow[SSSS_DUAL_SPECULAR_OFFSET].R = FMath::Clamp(AverageToRoughness0 / SSSS_MAX_DUAL_SPECULAR_ROUGHNESS, 0.0f, 1.0f);
-		TextureRow[SSSS_DUAL_SPECULAR_OFFSET].G = FMath::Clamp(AverageToRoughness1 / SSSS_MAX_DUAL_SPECULAR_ROUGHNESS, 0.0f, 1.0f);
+		TextureRow[SSSS_DUAL_SPECULAR_OFFSET].R = FMath::Clamp(Data.Roughness0 / SSSS_MAX_DUAL_SPECULAR_ROUGHNESS, 0.0f, 1.0f);
+		TextureRow[SSSS_DUAL_SPECULAR_OFFSET].G = FMath::Clamp(Data.Roughness1 / SSSS_MAX_DUAL_SPECULAR_ROUGHNESS, 0.0f, 1.0f);
 		TextureRow[SSSS_DUAL_SPECULAR_OFFSET].B = Data.LobeMix;
 		TextureRow[SSSS_DUAL_SPECULAR_OFFSET].A = FMath::Clamp(MaterialRoughnessToAverage / SSSS_MAX_DUAL_SPECULAR_ROUGHNESS, 0.0f, 1.0f);
 
