Strata  - fix for SLW refraction to match legacy (do not disrupt transform matrix)

diff --git a/Engine/Shaders/Private/BasePassPixelShader.usf b/Engine/Shaders/Private/BasePassPixelShader.usf
--- a/Engine/Shaders/Private/BasePassPixelShader.usf
+++ b/Engine/Shaders/Private/BasePassPixelShader.usf
@@ -1140,7 +1140,6 @@ void FPixelShaderInOut_MainPS(
 	Roughness = SLW_ROUGHNESS(SLWBSDF);
 	Opacity = SLW_TOPMATERIALOPACITY(SLWBSDF);
 	MaterialParameters.WorldNormal = StrataPixelHeader.SharedLocalBases.Normals[BSDF_GETSHAREDLOCALBASISID(SLWBSDF)];
-	MaterialParameters.TangentToWorld[2] = MaterialParameters.WorldNormal;
 	Anisotropy = 0.0f;
 	ShadingModel = MATERIAL_SHADINGMODEL_SINGLELAYERWATER;
 #endif
