Strata - last but not least difference fix: forgotten preexposition of separated main dir light which is the last difference. Now it all matches.

diff --git a/Engine/Shaders/Private/BasePassPixelShader.usf b/Engine/Shaders/Private/BasePassPixelShader.usf
--- a/Engine/Shaders/Private/BasePassPixelShader.usf
+++ b/Engine/Shaders/Private/BasePassPixelShader.usf
@@ -1971,11 +1971,11 @@ void FPixelShaderInOut_MainPS(
 			{
 				// STRATA_TODO: we should have a define set to the correct output index.
 				#if PIXELSHADEROUTPUT_MRT2
-					Out.MRT[2] = float4(SeparatedWaterMainDirLightLuminance, 0.0f);
+					Out.MRT[2] = float4(SeparatedWaterMainDirLightLuminance * View.PreExposure, 0.0f);
 				#elif PIXELSHADEROUTPUT_MRT1
-					Out.MRT[1] = float4(SeparatedWaterMainDirLightLuminance, 0.0f);
+					Out.MRT[1] = float4(SeparatedWaterMainDirLightLuminance * View.PreExposure, 0.0f);
 				#else
-					Out.MRT[0] = float4(SeparatedWaterMainDirLightLuminance, 0.0f);
+					Out.MRT[0] = float4(SeparatedWaterMainDirLightLuminance * View.PreExposure, 0.0f);
 				#endif
 			}
 		#endif
