Fix for strata compilation with new dual lob code.

diff --git a/Engine/Shaders/Private/Strata/Strata.ush b/Engine/Shaders/Private/Strata/Strata.ush
--- a/Engine/Shaders/Private/Strata/Strata.ush
+++ b/Engine/Shaders/Private/Strata/Strata.ush
@@ -886,8 +886,8 @@ FStrataData GetStrataSlabBSDF(
 		// Convertion from roughness0/1 to roughness+haziness
 		// * This encoding is based on how roughness0 + haziness is translated to roughness1, see ComputeHazyLobeRoughness
 		// * The mix factor is hardoded to 0.5 i.e., 0.5 * Fs(roughness0) + 0.5 * Fs(roughness1)
-		Haziness = (Lobe1Roughness - Lobe0Roughness) / (1.f - min(Lobe0Roughness, 0.99f));
-		Roughness = Lobe0Roughness;
+		Haziness = (LobeRoughness1 - LobeRoughness0) / (1.f - min(LobeRoughness0, 0.99f));
+		Roughness = LobeRoughness0;
 
 		SSSMFP = GetSubsurfaceProfileMFPInCm(SubsurfaceProfileUInt).xyz * SSSMFPScale;
 	}
