Fix hair adaptive voxel size outputting invalid value when the camera was inside the volume.

diff --git a/Engine/Shaders/Private/HairStrands/HairStrandsAABB.usf b/Engine/Shaders/Private/HairStrands/HairStrandsAABB.usf
--- a/Engine/Shaders/Private/HairStrands/HairStrandsAABB.usf
+++ b/Engine/Shaders/Private/HairStrands/HairStrandsAABB.usf
@@ -55,7 +55,8 @@ void Main(uint2 DispatchThreadId : SV_DispatchThreadID)
 		const FHairAABB ViewBound = Transform(Bound, View.TranslatedWorldToCameraView);
 
 		// Compute voxel size to have 1px = 1 voxels
-		const float PixelWorldSize = PixelSizeAtDepth1 * ViewBound.Min.z;
+		const float DefaultNearClip = 0.001f;
+		const float PixelWorldSize = PixelSizeAtDepth1 * max(ViewBound.Min.z, DefaultNearClip);
 		float VoxelWorldSize = PixelWorldSize * NumPixelPerVoxel;
 		WriteHairAABB(MacroGroupId, Bound, OutMacroGroupAABBBuffer);
 
