Strata - fixed normals on debug probe.

diff --git a/Engine/Shaders/Private/DebugProbes.usf b/Engine/Shaders/Private/DebugProbes.usf
--- a/Engine/Shaders/Private/DebugProbes.usf
+++ b/Engine/Shaders/Private/DebugProbes.usf
@@ -306,6 +306,7 @@ void MainPS(
 
 		FStrataData StrataData = PromoteParameterBlendedBSDFToOperator(StrataSlabData, StrataPixelHeader.StrataTree, /*int OperatorIndex*/0, /*int BSDFIndex*/0, /*int bIsBottom*/1, /*int bIsTop*/1);
 		StrataPixelHeader.StrataTree.BSDFs[0].LuminanceWeight = 1.0f;
+		StrataPixelHeader.StrataTree.BSDFs[0].TopLayerDataWeight = 1.0f;
 
 		FStrataSubsurfaceData SSSData = (FStrataSubsurfaceData)0;
 		FStrataTopLayerData TopLayerData = (FStrataTopLayerData)0;
