Fixed rough refractions to not leak accorss depth differences and outside viewport colors.

diff --git a/Engine/Shaders/Private/Strata/StrataRoughRefraction.usf b/Engine/Shaders/Private/Strata/StrataRoughRefraction.usf
--- a/Engine/Shaders/Private/Strata/StrataRoughRefraction.usf
+++ b/Engine/Shaders/Private/Strata/StrataRoughRefraction.usf
@@ -79,12 +79,37 @@ void OpaqueRoughRefractionPS(
 		sincos(RandAngle, SinA, CosA);
 		const float2 BlurOrigin = PixelPos + BlurStep * StandardDeviationPixelSpace * float2(CosA, SinA);
 
+		const float4 UvViewRectMinMax = float4(float2(View.ViewRectMinAndSize.xy) + 0.5f, float2(View.ViewRectMinAndSize.xy + View.BufferSizeAndInvSize.xy) - 0.5f) * View.BufferSizeAndInvSize.zwzw;
+
 		for (float u = -RadiusAroundCenterMu; u <= RadiusAroundCenterMu; u += BlurStep) // steps of 0.5 means 13 taps
 		{
 			const float OffsetPixelSpace = u * StandardDeviationPixelSpace;
-			const float Weight = GaussianFunction(OffsetPixelSpace, 0.0f, VariancePixelSpace);
-			const float2 UV = (BlurOrigin + BlurDirection * OffsetPixelSpace) * View.BufferSizeAndInvSize.zw;
-			WeightedLuminance += SeparatedOpaqueRoughRefractionSceneColor.SampleLevel(View.SharedBilinearClampedSampler, UV, 0).rgb * Weight;
+			const float2 SampleUV = (BlurOrigin + BlurDirection * OffsetPixelSpace) * View.BufferSizeAndInvSize.zw; 
+			const uint2 SamplePixelPos = uint2(SampleUV * View.BufferSizeAndInvSize.xy);
+
+			float Weight = GaussianFunction(OffsetPixelSpace, 0.0f, VariancePixelSpace);
+
+			// Do not accept rough refraction if the SampleUV is out of the view rect
+			const bool bSampleIsWithinViewRect = all(SampleUV >= UvViewRectMinMax.xy) && all(SampleUV <= UvViewRectMinMax.zw);
+			Weight *= bSampleIsWithinViewRect ? 1.0f : 0.0f;
+
+			// Do not accept rough refraction if the sampled material does not do rough refraction. Solved well by depth difference below so disabled.
+			//FStrataOpaqueRoughRefractionData ORRSampleData = StrataUnpackOpaqueRoughRefractionData(Strata.OpaqueRoughRefractionTexture[SamplePixelPos]);
+			//Weight *= ORRSampleData.OpaqueRoughRefractionEnabled ? 1.0f : 0.0f;
+
+			// Do not accept rough refraction if the depth difference is too high. Solves a lot of problems such as: 
+			//  - materials/sky being further away leaking,
+			//  - two materials hacing rough refraction touching on screen with a too large depth difference.
+			// Using the sample UV with bilinear interpolation for such an exclusion mechanism is not perfect but does work well in most cases.
+			// The true solution would be to only take the average of the 4 neightboor samples according to bilinear weights modulated by the exclusion filter itself.
+			const float DepthDifferenceThresholdCm = 5.0f;
+			const float SampleSceneDepthCm = ConvertFromDeviceZ(LookupDeviceZ(SampleUV));
+			Weight *= abs(SampleSceneDepthCm - SceneDepthCm) < DepthDifferenceThresholdCm ? 1.0f : 0.0f;
+
+			// Do not accepth rough refraction if the normal difference is too high: 
+			// => this could prove complex to make work in many cases with high frequency normalmapped surfaces. A geometry normal would be better to use.
+
+			WeightedLuminance += SeparatedOpaqueRoughRefractionSceneColor.SampleLevel(View.SharedBilinearClampedSampler, SampleUV, 0).rgb * Weight;
 			AccumulatedWeights += Weight;
 		}
 
