Add more Font/Debug color.

diff --git a/Engine/Shaders/Private/ShaderPrint.ush b/Engine/Shaders/Private/ShaderPrint.ush
--- a/Engine/Shaders/Private/ShaderPrint.ush
+++ b/Engine/Shaders/Private/ShaderPrint.ush
@@ -178,6 +178,12 @@ FFontColor GetDefaultFontColor()							{ FFontColor Out; Out.Color = float3(1,1,
 #define FontTurquoise	InitFontColor(26.f / 255.f, 188.f / 255.f, 156.f / 255.f)
 #define FontSilver		InitFontColor(189.f / 255.f, 195.f / 255.f, 199.f / 255.f)
 #define FontEmerald		InitFontColor(46.f / 255.f, 204.f / 255.f, 113.f / 255.f)
+#define FontLightRed	InitFontColor(0.75f, 0.50f, 0.50f)
+#define FontLightGreen	InitFontColor(0.50f, 0.75f, 0.50f)
+#define FontLightBlue	InitFontColor(0.50f, 0.50f, 0.75f)
+#define FontDarkRed		InitFontColor(0.50f, 0.25f, 0.25f)
+#define FontDarkGreen	InitFontColor(0.25f, 0.50f, 0.25f)
+#define FontDarkBlue	InitFontColor(0.25f, 0.25f, 0.50f)
 
 ///////////////////////////////////////////////////////////////////////////////////////////////////
 // Primitive colors
@@ -195,8 +201,12 @@ FFontColor GetDefaultFontColor()							{ FFontColor Out; Out.Color = float3(1,1,
 #define ColorTurquoise	float4( 26.f / 255.f, 188.f / 255.f, 156.f / 255.f, 1)
 #define ColorSilver		float4(189.f / 255.f, 195.f / 255.f, 199.f / 255.f, 1)
 #define ColorEmerald	float4( 46.f / 255.f, 204.f / 255.f, 113.f / 255.f, 1)
-#define ColorLightGreen	float4(0.5f, 0.75f, 0.5f, 1.0f)
-#define ColorDarkGreen	float4(0.25f,0.50f, 0.25f, 1.0f)
+#define ColorLightRed	float4(0.75f, 0.50f, 0.50f, 1.0f)
+#define ColorLightGreen	float4(0.50f, 0.75f, 0.50f, 1.0f)
+#define ColorLightBlue	float4(0.50f, 0.50f, 0.75f, 1.0f)
+#define ColorDarkRed	float4(0.50f, 0.25f, 0.25f, 1.0f)
+#define ColorDarkGreen	float4(0.25f, 0.50f, 0.25f, 1.0f)
+#define ColorDarkBlue	float4(0.25f, 0.25f, 0.50f, 1.0f)
 
 ///////////////////////////////////////////////////////////////////////////////////////////////////
 // Util pack/unpac functions
