Change strata visualization to use ShaderPrint.

Add r.Strata.DebugMode CVar to toggle view mode. 1: default material vis. 2: BSDF count.

diff --git a/Engine/Shaders/Private/Strata/StrataVisualize.usf b/Engine/Shaders/Private/Strata/StrataVisualize.usf
--- a/Engine/Shaders/Private/Strata/StrataVisualize.usf
+++ b/Engine/Shaders/Private/Strata/StrataVisualize.usf
@@ -7,9 +7,10 @@
 #define STRATA_SSS_MATERIAL_OVERRIDE 0
 #define STRATA_SSS_TRANSMISSION 0  // Disable for now, as the SSS profile texture need to be bound
 #include "/Engine/Private/Strata/StrataEvaluation.ush"
+#include "/Engine/Private/Strata/StrataTile.ush"
 
-#include "/Engine/Private/MiniFontCommon.ush"
 #include "/Engine/Private/ShaderPrint.ush"
+#include "../ColorMap.ush"
 
 // This is for the fxc compiler to compile faster.
 #define STRATA_VIZ_MAX_BSDF 2
@@ -116,537 +117,367 @@ void DrawPixelFootprint(float3 P, float3 dPdx, float3 dPdy, uint2 PixelCoord)
 	}
 }
 
-void VisualizeMaterialPS(
-	float4 SVPos : SV_POSITION,
-	out float4 OutColor : SV_Target0)
-{
-	const uint2 PixelPos = uint2(SVPos.xy);
-	const uint2 MaterialPos = View.CursorPosition;
-	float2 BufferUV = SvPositionToBufferUV(float4(MaterialPos, SVPos.zw));
-	const int  InspectSquare = 10;
+///////////////////////////////////////////////////////////////////////////////////////////////////
+// Material Print
 
-	const float2 ScreenPosition = SvPositionToScreenPosition(SVPos).xy;
-	const float2 ScreenMaterialPosition = SvPositionToScreenPosition(float4(MaterialPos, 0.5f, 1)).xy;
-	float DeviceZ = SampleDeviceZFromSceneTextures(BufferUV);
-	float SceneDepth = ConvertFromDeviceZ(DeviceZ);
-	const float3 WorldPosition = mul(float4(ScreenMaterialPosition * SceneDepth, SceneDepth, 1), PrimaryView.ScreenToTranslatedWorld).xyz;
-	const float3 V = -normalize(WorldPosition - PrimaryView.TranslatedWorldCameraOrigin);
+#if SHADER_MATERIALPRINT
 
-	FStrataAddressing StrataAddressing = GetStrataPixelDataByteOffset(MaterialPos, uint2(View.BufferSizeAndInvSize.xy), Strata.MaxBytesPerPixel);
-	FStrataPixelHeader Header = UnpackStrataHeaderIn(Strata.MaterialTextureArray, StrataAddressing, Strata.TopLayerTexture);
-	bool bStrataMaterial = Header.BSDFCount > 0;
+#ifndef MULTIPASS_ENABLE 
+#define MULTIPASS_ENABLE 0
+#endif
 
-#if 0
-	OutColor = 0;
-	if (Header.BSD)
+uint BSDFIndex;
+RWBuffer<uint> RWPositionOffsetBuffer;
+
+// View mode:
+// 1: Plot a single pixel
+// 2: Color code the number of BSDFs (green -> red)
+// 3: Add layer visualization
+void PrintBool(inout FShaderPrintContext Ctx, bool In, FFontColor InColor)
+{
+	if (In)
 	{
-		OutColor.w = 0.05f;
-		if (Header.BSDFCount == 1)
-		{
-			OutColor = float4(0, 0, 1, 0);
-		}
-		else if (Header.BSDFCount == 2)
-		{
-			OutColor = float4(0, 1, 0, 0);
-		}
-		else if (Header.BSDFCount == 3)
-		{
-			OutColor = float4(1, 1, 0, 0);
-		}
-		else if (Header.BSDFCount == 4)
-		{
-			OutColor = float4(1, 0, 0, 0);
-		}
-		else if (Header.BSDFCount > 4)
-		{
-			OutColor = float4(1, 1, 1, 0);
-		}
+		Print(Ctx, TEXT("Yes  "), FontGreen);
 	}
-	return;
-#endif
-
-#if 0	// STRATA_TODO generalise this visualisation to more properties
+	else
 	{
-		FStrataAddressing StrataAddressing = GetStrataPixelDataByteOffset(PixelPos, uint2(View.BufferSizeAndInvSize.xy), Strata.MaxBytesPerPixel);
-		const uint PixelStrataDataByteOffsetStart = StrataAddressing.ReadBytes;
+		Print(Ctx, TEXT("No   "), FontRed);
+	}
+}
 
-		FStrataPixelHeader Header2 = UnpackStrataHeaderIn(Strata.MaterialTextureArray, StrataAddressing, Strata.TopLayerTexture);
-		const uint PixelStrataDataByteOffsetPostHeader = StrataAddressing.ReadBytes;
+void PrintPixelType(inout FShaderPrintContext Ctx, bool bIsSimpleMaterial, bool bIsSingleMaterial, FFontColor InColor)
+{
+		 if (bIsSimpleMaterial)	Print(Ctx, TEXT("Simple  "), InColor);
+	else if (bIsSingleMaterial)	Print(Ctx, TEXT("Single  "), InColor);
+	else						Print(Ctx, TEXT("Complex "), InColor);
+}
 
-		if (HasSubsurface(Header2) || HasSubsurfaceProfile(Header2))
-		{
-			OutColor = float4(HasSubsurfaceProfile(Header2) ? 1 : 0, HasSubsurface(Header2) ? 1 : 0, 0, 0.5);
-			return;
-		}
-	}
-#endif
+void PrintBSDF(inout FShaderPrintContext Context, inout FStrataAddressing StrataAddressing, FStrataPixelHeader Header, float3 WorldPosition, float3 V)
+{
+	const FFontColor FontBSDF         = FontSilver;
+	const FFontColor FontBSDFType     = FontEmerald;
+	const FFontColor FontBSDFStateName= FontWhite;
+	const FFontColor FontBSDFPropName = FontOrange;
+	const FFontColor FontBSDFPropValu = FontWhite;
+
+	const FStrataBSDF BSDF = UnpackStrataBSDFIn(Strata.MaterialTextureArray, StrataAddressing, Header);
 
-	// Print pixel footprint & SSS footprint
-	const float3 P = WorldPosition;
-	const float3 dPdx = ddx(P);
-	const float3 dPdy = ddy(P);
-	if (bStrataMaterial && all(PixelPos == MaterialPos))
+	// Draw Referential
 	{
-		DrawPixelFootprint(P, dPdx, dPdy, PixelPos);
+		const float3 DummyL = float3(0, 0, 1);
+		FStrataBSDFContext BSDFContext = StrataCreateBSDFContext(Header, BSDF, StrataAddressing, V, DummyL);
+		DrawReferentialTWS(WorldPosition, BSDFContext.X, BSDFContext.Y, BSDFContext.N, float3(1, 1, 0));
 	}
 
-	OutColor = float4(0.0f, 0.0f, 0.0f, 1.0f);
-
-	// Plot BSDFs
-	bool bPlotValid = false;
-	if (bStrataMaterial)
+	switch (BSDF_GETTYPE(BSDF))
 	{
-		FStrataAddressing StrataAddressing = GetStrataPixelDataByteOffset(PixelPos, uint2(View.BufferSizeAndInvSize.xy), Strata.MaxBytesPerPixel);
-
-		FStrataPixelHeader Header = UnpackStrataHeaderIn(Strata.MaterialTextureArray, StrataAddressing, Strata.TopLayerTexture);
-
-		for (uint BSDFIndex = 0; BSDFIndex < Header.BSDFCount; ++BSDFIndex)
+		case STRATA_BSDF_TYPE_SLAB:
 		{
-			const FStrataBSDF BSDF = UnpackStrataBSDFIn(Strata.MaterialTextureArray, StrataAddressing, Header);
-
-			// Referential
-			if (all(PixelPos == MaterialPos))
+			Print(Context, TEXT("Slab "), FontBSDFType);
+			Newline(Context);
+			Print(Context, TEXT(" NormalID Aniso TopLayer Scattering SingleScatt. IsThin Weight"), FontBSDFStateName);
+			Newline(Context);
+
+			Print(Context, TEXT(" "), FontSilver);
+			Print(Context, BSDF_GETSHAREDLOCALBASISID(BSDF), FontBSDF, 2, 1);	Print(Context, TEXT("       "), FontSilver);
+			PrintBool(Context, BSDF_GETHASANISOTROPY(BSDF), FontBSDF);			Print(Context, TEXT(" "), FontSilver);
+			PrintBool(Context, BSDF_GETISTOPLAYER(BSDF), FontBSDF);				Print(Context, TEXT("    "), FontSilver);
+			PrintBool(Context, BSDF_GETHASSSS(BSDF), FontBSDF);					Print(Context, TEXT("      "), FontSilver);
+			PrintBool(Context, BSDF_GETISSIMPLEVOLUME(BSDF), FontBSDF);			Print(Context, TEXT("        "), FontSilver);
+			PrintBool(Context, BSDF_GETISTHIN(BSDF), FontBSDF);					Print(Context, TEXT("  "), FontSilver);
+
+			if (BSDF_GETHASGREYWEIGHT(BSDF))
 			{
-				FStrataBSDFContext BSDFContext = StrataCreateBSDFContext(Header, BSDF, StrataAddressing, V, float3(0,0,1));
-				DrawReferentialTWS(WorldPosition, BSDFContext.X, BSDFContext.Y, BSDFContext.N, float3(1,1,0));
+				Print(Context, BSDF.LuminanceWeight.x, FontSilver);
 			}
-
-			// Plot
+			else
 			{
-				const uint2 Resolution = 180;
-				const uint OffsetX = 50 + BSDFIndex * 200;
-				const uint OffsetY = 50;
-				OutColor.xyz += DrawPlot(V, Header, BSDF, StrataAddressing, PixelPos, uint2(OffsetX, OffsetY), Resolution);
+				Print(Context, BSDF.LuminanceWeight.r, FontLightRed);
+				Print(Context, BSDF.LuminanceWeight.g, FontLightGreen);
+				Print(Context, BSDF.LuminanceWeight.b, FontLightBlue);
 			}
-		}
+			Newline(Context);
 
-		// Background
-		if (dot(OutColor.xyz, float3(1, 1, 1)) > 0)
-		{
-			OutColor.w = 0;
-			bPlotValid = true;
-		}
-	}
+			Print(Context, TEXT(" Diffuse     "), FontBSDFPropName);	Print(Context, SLAB_DIFFUSEALBEDO(BSDF), FontBSDFPropValu);			Newline(Context);
 
-	// Print BSDFs attributes
-	if (bStrataMaterial)
-	{
-		uint2 RectMin = uint2(MaterialPos + InspectSquare);
-		uint2 RectMax = RectMin;
+			const bool bGreyF0 = SLAB_F0(BSDF).x == SLAB_F0(BSDF).y && SLAB_F0(BSDF).x == SLAB_F0(BSDF).z;
+			if (bGreyF0)
+			{
+			Print(Context, TEXT(" F0          "), FontBSDFPropName);	Print(Context, SLAB_F0(BSDF).x, FontBSDFPropValu);					Newline(Context);
+			}
+			else
+			{
+			Print(Context, TEXT(" F0          "), FontBSDFPropName);	Print(Context, SLAB_F0(BSDF), FontBSDFPropValu);					Newline(Context);
+			}
 
-		const FStrataTopLayerData TopLayerData = StrataUnpackTopLayerData(Strata.TopLayerTexture.Load(uint3(MaterialPos, 0)));
-		const FStrataSubsurfaceData SSSData = StrataUnpackSSSData(Strata.SSSTexture.Load(uint3(MaterialPos, 0)));
 
-		FStrataAddressing StrataAddressing = GetStrataPixelDataByteOffset(MaterialPos, uint2(View.BufferSizeAndInvSize.xy), Strata.MaxBytesPerPixel);
-		const uint PixelStrataDataByteOffsetStart = StrataAddressing.ReadBytes;
+			if (BSDF_GETHASF90(BSDF))
+			{
+			const bool bGreyF90 = SLAB_F90(BSDF).x == SLAB_F90(BSDF).y && SLAB_F90(BSDF).x == SLAB_F90(BSDF).z;
+			if (bGreyF90)
+			{
+			Print(Context, TEXT(" F90         "), FontBSDFPropName);	Print(Context, SLAB_F90(BSDF).x, FontBSDFPropValu);					Newline(Context);
+			}
+			else
+			{
+			Print(Context, TEXT(" F90         "), FontBSDFPropName);	Print(Context, SLAB_F90(BSDF), FontBSDFPropValu);					Newline(Context);
+			}
+			}
 
-		FStrataPixelHeader Header = UnpackStrataHeaderIn(Strata.MaterialTextureArray, StrataAddressing, Strata.TopLayerTexture);
-		const uint PixelStrataDataByteOffsetPostHeader = StrataAddressing.ReadBytes;
+			Print(Context, TEXT(" Roughness   "), FontBSDFPropName);	Print(Context, SLAB_ROUGHNESS(BSDF), FontBSDFPropValu);				Newline(Context);
 
-		const bool bIsSimpleMaterial = IsSimpleMaterial(Header) || Header.BSDFCount == 0;
-		const bool bIsSingleMaterial = !IsSimpleMaterial(Header) && IsSingleMaterial(Header);
+			if (BSDF_GETHASANISOTROPY(BSDF))
+			{
+			Print(Context, TEXT(" Anisotropy  "), FontBSDFPropName);	Print(Context, SLAB_ANISOTROPY(BSDF), FontBSDFPropValu);			Newline(Context);
+			}
 
-		float3 BSDFFontColor = float3(0.25, 1.0, 0.25);
-		float3 BSDFFontColorGrey = float3(0.25, 0.25, 0.25);
-		float3 INFOFontColor = float3(1.0, 0.5, 0.25);
-		float3 PARAMFontColor = float3(0.8, 0.8, 0.8);
-		float3 FontColor = float3(0.25, 1.0, 1.0);
-		uint2 TopLeftInit = MaterialPos + InspectSquare + uint2(3, 3);
-		uint2 TopLeft = TopLeftInit;
-
-#define PRINTPARAMCHAR(X) PrintCharacter(PixelPos, OutColor.xyz, PARAMFontColor, TopLeft, X);
-#define PRINTPARAMCHARRed(X) PrintCharacter(PixelPos, OutColor.xyz, float3(1.0, 0.7, 0.7), TopLeft, X);
-#define PRINTPARAMCHARYel(X) PrintCharacter(PixelPos, OutColor.xyz, float3(1.0, 1.0, 0.7), TopLeft, X);
-#define PRINTPARAMCHARGre(X) PrintCharacter(PixelPos, OutColor.xyz, float3(0.7, 1.0, 0.7), TopLeft, X);
-#define PRINTINFOCHAR(X) PrintCharacter(PixelPos, OutColor.xyz, INFOFontColor, TopLeft, X);
-#define PRINTBSDFCHAR(X) PrintCharacter(PixelPos, OutColor.xyz, BSDFFontColor, TopLeft, X);
-#define PRINTBSDFCHARGREY(X) PrintCharacter(PixelPos, OutColor.xyz, BSDFFontColorGrey, TopLeft, X);
-
-#define PRINTBSDFWEIGHT()	PRINTBSDFCHAR(_SPC_); PRINTBSDFCHAR(_MINUS_); PRINTBSDFCHAR(_W_); PRINTBSDFCHAR(_E_); PRINTBSDFCHAR(_I_); PRINTBSDFCHAR(_G_); PRINTBSDFCHAR(_H_);PRINTBSDFCHAR(_T_);  PRINTBSDFCHAR(_SPC_); PrintFloat(PixelPos, OutColor.xyz, FontColor, TopLeft, BSDF.LuminanceWeight.x); TopLeft.x+=8*7; PRINTBSDFCHAR(_SPC_); PrintFloat(PixelPos, OutColor.xyz, FontColor, TopLeft, BSDF.LuminanceWeight.y); TopLeft.x+=8*7; PRINTBSDFCHAR(_SPC_); PrintFloat(PixelPos, OutColor.xyz, FontColor, TopLeft, BSDF.LuminanceWeight.z); TopLeft.x+=8*7; 
-#define PRINTBSDFGREYWEIGHT()PRINTBSDFCHAR(_SPC_);PRINTBSDFCHARGREY(_MINUS_); PRINTBSDFCHARGREY(_W_); PRINTBSDFCHARGREY(_E_); PRINTBSDFCHARGREY(_I_); PRINTBSDFCHARGREY(_G_); PRINTBSDFCHARGREY(_H_);PRINTBSDFCHARGREY(_T_);  PRINTBSDFCHAR(_SPC_); PrintFloat(PixelPos, OutColor.xyz, FontColor, TopLeft, Unpack10F(BSDF_GETWEIGHT10F(BSDF))); TopLeft.x+=8*23; 
-#define PRINTBSDFSHAREDLOCALBASISID()	PRINTBSDFCHAR(_SPC_); PRINTBSDFCHAR(_MINUS_); PRINTBSDFCHAR(_N_); PRINTBSDFCHAR(_R_); PRINTBSDFCHAR(_M_); PRINTBSDFCHAR(_I_); PRINTBSDFCHAR(_D_); PRINTBSDFCHAR(_SPC_); PrintFloatNoFraction(PixelPos, OutColor.xyz, FontColor, TopLeft, float(BSDF_GETSHAREDLOCALBASISID(BSDF)), 1); TopLeft.x+=8*1;
-#define PRINTBSDFANISO()	PRINTBSDFCHAR(_SPC_); PRINTBSDFCHAR(_MINUS_); PRINTBSDFCHAR(_A_); PRINTBSDFCHAR(_N_); PRINTBSDFCHAR(_I_); PRINTBSDFCHAR(_S_); PRINTBSDFCHAR(_O_); PRINTBSDFCHAR(_SPC_); PrintFloatNoFraction(PixelPos, OutColor.xyz, FontColor, TopLeft, float(BSDF_GETHASANISOTROPY(BSDF)), 1); TopLeft.x+=8*1;
-#define PRINTBSDFTOPLAY()	PRINTBSDFCHAR(_SPC_); PRINTBSDFCHAR(_MINUS_); PRINTBSDFCHAR(_T_); PRINTBSDFCHAR(_O_); PRINTBSDFCHAR(_P_); PRINTBSDFCHAR(_L_); PRINTBSDFCHAR(_A_); PRINTBSDFCHAR(_Y_); PRINTBSDFCHAR(_SPC_); PrintFloatNoFraction(PixelPos, OutColor.xyz, FontColor, TopLeft, float(BSDF_GETISTOPLAYER(BSDF)), 1); TopLeft.x+=8*1;
-#define PRINTBSDFHASSCAT()	PRINTBSDFCHAR(_SPC_); PRINTBSDFCHAR(_MINUS_); PRINTBSDFCHAR(_H_); PRINTBSDFCHAR(_A_); PRINTBSDFCHAR(_S_); PRINTBSDFCHAR(_S_); PRINTBSDFCHAR(_C_); PRINTBSDFCHAR(_A_); PRINTBSDFCHAR(_T_); PRINTBSDFCHAR(_SPC_); PrintFloatNoFraction(PixelPos, OutColor.xyz, FontColor, TopLeft, float(BSDF_GETHASSSS(BSDF)), 1); TopLeft.x+=8*1;
-#define PRINTBSDFSIMPSCAT()	PRINTBSDFCHAR(_SPC_); PRINTBSDFCHAR(_MINUS_); PRINTBSDFCHAR(_S_); PRINTBSDFCHAR(_I_); PRINTBSDFCHAR(_M_); PRINTBSDFCHAR(_P_); PRINTBSDFCHAR(_S_); PRINTBSDFCHAR(_C_); PRINTBSDFCHAR(_A_); PRINTBSDFCHAR(_T_); PRINTBSDFCHAR(_SPC_); PrintFloatNoFraction(PixelPos, OutColor.xyz, FontColor, TopLeft, float(BSDF_GETISSIMPLEVOLUME(BSDF)), 1); TopLeft.x+=8*1;
-#define PRINTBSDFISTHIN()	PRINTBSDFCHAR(_SPC_); PRINTBSDFCHAR(_MINUS_); PRINTBSDFCHAR(_I_); PRINTBSDFCHAR(_S_); PRINTBSDFCHAR(_T_); PRINTBSDFCHAR(_H_); PRINTBSDFCHAR(_I_); PRINTBSDFCHAR(_N_); PRINTBSDFCHAR(_SPC_); PrintFloatNoFraction(PixelPos, OutColor.xyz, FontColor, TopLeft, float(BSDF_GETISTHIN(BSDF)), 1); TopLeft.x+=8*1;
-#define PRINTFLOAT3(X) PrintFloat(PixelPos, OutColor.xyz, PARAMFontColor, TopLeft, X.x); TopLeft.x+=8*7; PRINTBSDFCHAR(_SPC_); PrintFloat(PixelPos, OutColor.xyz, PARAMFontColor, TopLeft, X.y); TopLeft.x+=8*7; PRINTBSDFCHAR(_SPC_); PrintFloat(PixelPos, OutColor.xyz, PARAMFontColor, TopLeft, X.z);
-#define PRINTFLOAT1(Value) PrintFloat(PixelPos, OutColor.xyz, PARAMFontColor, TopLeft, Value); TopLeft.x+=8*7;
-#define PRINTFLOATNOFRAC(Value, DigitCount) PrintFloatNoFraction(PixelPos, OutColor.xyz, PARAMFontColor, TopLeft, Value, DigitCount); TopLeft.x+=8*7;
-#define SIGNOFFLOAT1(Value) if(Value < 0.0f) { PrintCharacter(PixelPos, OutColor.xyz, PARAMFontColor, TopLeft, _MINUS_); } /*else { PrintCharacter(PixelPos, OutColor.xyz, PARAMFontColor, TopLeft, _SPC_);}*/
-
-#define PRINTBSDFCOUNT()	PRINTBSDFCHAR(min(_0_+Header.BSDFCount,_F_)); PRINTBSDFCHAR(_SPC_); PRINTBSDFCHAR(_B_); PRINTBSDFCHAR(_S_); PRINTBSDFCHAR(_D_); PRINTBSDFCHAR(_F_);
-#define PRINTMATERIALAO()	PRINTBSDFCHAR(_A_); PRINTBSDFCHAR(_O_); PRINTBSDFCHAR(_SPC_); PRINTFLOAT1(StrataGetAO(Header));
-#define PRINTMATERIALSKIPPRECSHADOWMASK()	PRINTBSDFCHAR(_H_); PRINTBSDFCHAR(_A_); PRINTBSDFCHAR(_S_); PRINTBSDFCHAR(_SPC_); PRINTBSDFCHAR(_P_); PRINTBSDFCHAR(_R_); PRINTBSDFCHAR(_E_); PRINTBSDFCHAR(_S_); PRINTBSDFCHAR(_H_); PRINTBSDFCHAR(_A_); PRINTBSDFCHAR(_D_); PRINTBSDFCHAR(_SPC_); PRINTBSDFCHAR(_0_+(HasPrecShadowMask(Header)?1:0));
-#define PRINTMATERIALZEROPRECSHADOWMASK()	PRINTBSDFCHAR(_Z_); PRINTBSDFCHAR(_E_); PRINTBSDFCHAR(_R_); PRINTBSDFCHAR(_O_); PRINTBSDFCHAR(_P_); PRINTBSDFCHAR(_R_); PRINTBSDFCHAR(_E_); PRINTBSDFCHAR(_S_); PRINTBSDFCHAR(_H_); PRINTBSDFCHAR(_A_); PRINTBSDFCHAR(_D_); PRINTBSDFCHAR(_SPC_); PRINTBSDFCHAR(_0_+(HasZeroPrecShadowMask(Header)?1:0));
-#define PRINTMATERIALSLW()	PRINTBSDFCHAR(_S_); PRINTBSDFCHAR(_L_); PRINTBSDFCHAR(_W_); PRINTBSDFCHAR(_SPC_); PRINTBSDFCHAR(_0_+(IsSingleLayerWater(Header)?1:0));
-#define PRINTMATERIALCONTSHAD()	PRINTBSDFCHAR(_C_); PRINTBSDFCHAR(_O_); PRINTBSDFCHAR(_N_);  PRINTBSDFCHAR(_T_); PRINTBSDFCHAR(_SPC_);PRINTBSDFCHAR(_S_); PRINTBSDFCHAR(_H_); PRINTBSDFCHAR(_A_);  PRINTBSDFCHAR(_D_); PRINTBSDFCHAR(_SPC_); PRINTBSDFCHAR(_0_+(DoesCastContactShadow(Header)?1:0));
-#define PRINTMATERIALINDSHAD()	PRINTBSDFCHAR(_I_); PRINTBSDFCHAR(_N_); PRINTBSDFCHAR(_D_); PRINTBSDFCHAR(_SPC_); PRINTBSDFCHAR(_S_); PRINTBSDFCHAR(_H_); PRINTBSDFCHAR(_A_);  PRINTBSDFCHAR(_D_); PRINTBSDFCHAR(_SPC_); PRINTBSDFCHAR(_0_+(HasDynamicIndirectShadowCasterRepresentation(Header)?1:0));
-
-#define PRINTCLASSIFROUGH()		PRINTBSDFCHAR(_C_); PRINTBSDFCHAR(_L_); PRINTBSDFCHAR(_A_); PRINTBSDFCHAR(_S_); PRINTBSDFCHAR(_DOT_); PRINTBSDFCHAR(_R_) PRINTBSDFCHAR(_O_) PRINTBSDFCHAR(_U_) PRINTBSDFCHAR(_G_) PRINTBSDFCHAR(_H_) PRINTBSDFCHAR(_SPC_);	PRINTBSDFCHAR(_SPC_); PRINTFLOAT1(TopLayerData.Roughness);
-#define PRINTCLASSIFBASCOL()	PRINTBSDFCHAR(_C_); PRINTBSDFCHAR(_L_); PRINTBSDFCHAR(_A_); PRINTBSDFCHAR(_S_); PRINTBSDFCHAR(_DOT_); PRINTBSDFCHAR(_B_) PRINTBSDFCHAR(_A_) PRINTBSDFCHAR(_S_) PRINTBSDFCHAR(_C_) PRINTBSDFCHAR(_O_) PRINTBSDFCHAR(_L_);	PRINTBSDFCHAR(_SPC_); PRINTFLOAT3(SSSData.BaseColor);
-#define PRINTCLASSIFISPROF()	PRINTBSDFCHAR(_C_); PRINTBSDFCHAR(_L_); PRINTBSDFCHAR(_A_); PRINTBSDFCHAR(_S_); PRINTBSDFCHAR(_DOT_); PRINTBSDFCHAR(_I_) PRINTBSDFCHAR(_S_) PRINTBSDFCHAR(_P_) PRINTBSDFCHAR(_R_) PRINTBSDFCHAR(_O_) PRINTBSDFCHAR(_F_);	PRINTBSDFCHAR(_SPC_); PRINTBSDFCHAR(SSSData.bIsProfile?_1_:_0_);
-#define PRINTCLASSIFPROID()		PRINTBSDFCHAR(_C_); PRINTBSDFCHAR(_L_); PRINTBSDFCHAR(_A_); PRINTBSDFCHAR(_S_); PRINTBSDFCHAR(_DOT_); PRINTBSDFCHAR(_P_) PRINTBSDFCHAR(_R_) PRINTBSDFCHAR(_O_) PRINTBSDFCHAR(_I_) PRINTBSDFCHAR(_D_) PRINTBSDFCHAR(_SPC_);	PRINTBSDFCHAR(_SPC_); PRINTFLOATNOFRAC(SSSData.ProfileId, 3);
-#define PRINTCLASSIFPROSCL()	PRINTBSDFCHAR(_C_); PRINTBSDFCHAR(_L_); PRINTBSDFCHAR(_A_); PRINTBSDFCHAR(_S_); PRINTBSDFCHAR(_DOT_); PRINTBSDFCHAR(_P_) PRINTBSDFCHAR(_R_) PRINTBSDFCHAR(_O_) PRINTBSDFCHAR(_S_) PRINTBSDFCHAR(_C_) PRINTBSDFCHAR(_L_);	PRINTBSDFCHAR(_SPC_); PRINTFLOAT1(SSSData.ProfileRadiusScale);
-#define PRINTCLASSIFSSMFP()	PRINTBSDFCHAR(_C_); PRINTBSDFCHAR(_L_); PRINTBSDFCHAR(_A_); PRINTBSDFCHAR(_S_); PRINTBSDFCHAR(_DOT_); PRINTBSDFCHAR(_S_) PRINTBSDFCHAR(_S_) PRINTBSDFCHAR(_M_) PRINTBSDFCHAR(_F_) PRINTBSDFCHAR(_P_);	PRINTBSDFCHAR(_SPC_); PRINTFLOAT3(SSSData.MFP);
-
-#define NEXTLINE  TopLeft.y += 8; RectMax.x = max(RectMax.x, TopLeft.x); RectMax.y = max(RectMax.y, TopLeft.y); TopLeft.x = TopLeftInit.x; 
-#define HORIZONTALSEPARATION  TopLeft.y += 4;
-#define HORIZONTALTAB  TopLeft.x += 8;
-
-#if PERMUTATION_BSDF_PASS==0
-		PRINTBSDFCOUNT();
-		NEXTLINE;
-		PRINTMATERIALAO();
-		NEXTLINE;
-		PRINTMATERIALSLW();
-		NEXTLINE;
-		PRINTMATERIALSKIPPRECSHADOWMASK();
-		NEXTLINE;
-		PRINTMATERIALZEROPRECSHADOWMASK();
-		NEXTLINE;
-		PRINTMATERIALCONTSHAD();
-		NEXTLINE;
-		PRINTMATERIALINDSHAD();
-		NEXTLINE;
-		PRINTBSDFCHAR(_T_); PRINTBSDFCHAR(_Y_); PRINTBSDFCHAR(_P_); PRINTBSDFCHAR(_E_) PRINTBSDFCHAR(_SPC_)
-		if (bIsSimpleMaterial)
-		{
-			PRINTPARAMCHARGre(_S_); PRINTPARAMCHARGre(_I_); PRINTPARAMCHARGre(_M_); PRINTPARAMCHARGre(_P_); PRINTPARAMCHARGre(_L_); PRINTPARAMCHARGre(_E_);
-		}
-		else if(bIsSingleMaterial)
-		{
-			PRINTPARAMCHARYel(_S_); PRINTPARAMCHARYel(_I_); PRINTPARAMCHARYel(_N_); PRINTPARAMCHARYel(_G_); PRINTPARAMCHARYel(_L_); PRINTPARAMCHARYel(_E_);
-		}
-		else
-		{
-			PRINTPARAMCHARRed(_C_); PRINTPARAMCHARRed(_O_); PRINTPARAMCHARRed(_M_); PRINTPARAMCHARRed(_P_); PRINTPARAMCHARRed(_L_); PRINTPARAMCHARRed(_E_); PRINTPARAMCHARRed(_X_);
-		}
-		NEXTLINE;
+			if (BSDF_GETHASHAZINESS(BSDF))
+			{
+			Print(Context, TEXT(" Haziness    "), FontBSDFPropName);	Print(Context, SLAB_HAZINESS(BSDF), FontBSDFPropValu);				Newline(Context);
+			}
+			
+			if (BSDF_GETHASSSS(BSDF) || BSDF_GETISSIMPLEVOLUME(BSDF) || BSDF_GETISTHIN(BSDF))
+			{
 
-		PRINTCLASSIFROUGH();
-		NEXTLINE;
-		if (HasSubsurface(Header))
-		{
-			PRINTCLASSIFBASCOL();
-			NEXTLINE;
-			PRINTCLASSIFISPROF();
-			NEXTLINE;
-			if (SSSData.bIsProfile)
+			if (BSDF_GETHASSSSPROFILE(BSDF) && BSDF_GETISSIMPLEVOLUME(BSDF) == 0)
 			{
-				PRINTCLASSIFPROID();
-				NEXTLINE;
-				PRINTCLASSIFPROSCL();
-				NEXTLINE;
+			Print(Context, TEXT(" SSS ID.     "), FontBSDFPropName);	Print(Context, SLAB_SSSPROFILEID(BSDF), FontBSDFPropValu);			Newline(Context);
+			Print(Context, TEXT(" SSS Radius  "), FontBSDFPropName);	Print(Context, SLAB_SSSPROFILERADIUSSCALE(BSDF), FontBSDFPropValu);	Newline(Context);
 			}
 			else
 			{
-				PRINTCLASSIFSSMFP();
-				NEXTLINE;
-				NEXTLINE;
+			Print(Context, TEXT(" SSS MFP     "), FontBSDFPropName);	Print(Context, SLAB_SSSMFP(BSDF), FontBSDFPropValu);				Newline(Context);
+			}
+
+			} // SSS
+
+			if (BSDF_GETHASFUZZ(BSDF))
+			{
+			Print(Context, TEXT(" Fuzz Amount "), FontBSDFPropName);	Print(Context, SLAB_FUZZ_AMOUNT(BSDF), FontBSDFPropValu);			Newline(Context);
+			Print(Context, TEXT(" Fuzz Color  "), FontBSDFPropName);	Print(Context, SLAB_FUZZ_COLOR(BSDF), FontBSDFPropValu);			Newline(Context);
 			}
 		}
-		else
-		{
-			PRINTBSDFCHAR(_N_); PRINTBSDFCHAR(_O_); PRINTBSDFCHAR(_S_); PRINTBSDFCHAR(_S_); PRINTBSDFCHAR(_S_);
-			NEXTLINE;
-		}
+		break;
 
-		HORIZONTALSEPARATION;
-		PRINTINFOCHAR(_T_); PRINTINFOCHAR(_O_); PRINTINFOCHAR(_P_);
-		NEXTLINE;
-		HORIZONTALSEPARATION;
-#else
-		NEXTLINE;
-		NEXTLINE;
-		NEXTLINE;
-		NEXTLINE;
-		NEXTLINE;
-		NEXTLINE;
-		NEXTLINE;
-		NEXTLINE;
-
-		NEXTLINE;
-		if (HasSubsurface(Header))
+		case STRATA_BSDF_TYPE_HAIR:
 		{
-			NEXTLINE;
-			NEXTLINE;
-			if (SSSData.bIsProfile)
+			Print(Context, TEXT("Hair"), FontBSDFType);
+			Newline(Context);
+			Print(Context, TEXT(" NormalID Aniso IsTopLayer LuminanceWeight"), FontBSDF);
+			Newline(Context);
+
+			Print(Context, BSDF_GETSHAREDLOCALBASISID(BSDF), FontBSDF);
+			PrintBool(Context, BSDF_GETHASANISOTROPY(BSDF), FontBSDF);
+			PrintBool(Context, BSDF_GETISTOPLAYER(BSDF), FontBSDF);
+			if (BSDF_GETHASGREYWEIGHT(BSDF))
 			{
-				NEXTLINE;
-				NEXTLINE;
+				Print(Context, BSDF.LuminanceWeight.x, FontSilver);
 			}
 			else
 			{
-				NEXTLINE;
-				NEXTLINE;
+				Print(Context, BSDF.LuminanceWeight.r, FontLightRed);
+				Print(Context, BSDF.LuminanceWeight.g, FontLightGreen);
+				Print(Context, BSDF.LuminanceWeight.b, FontLightBlue);
 			}
+			Newline(Context);
+
+			Print(Context, TEXT(" BaseColor         "), FontBSDFPropName);	Print(Context, HAIR_BASECOLOR(BSDF), FontBSDFPropValu);				Newline(Context);
+			Print(Context, TEXT(" Specular          "), FontBSDFPropName);	Print(Context, HAIR_SPECULAR(BSDF), FontBSDFPropValu);				Newline(Context);
+			Print(Context, TEXT(" Roughness         "), FontBSDFPropName);	Print(Context, HAIR_ROUGHNESS(BSDF), FontBSDFPropValu);				Newline(Context);
+			Print(Context, TEXT(" Scatter           "), FontBSDFPropName);	Print(Context, HAIR_SCATTER(BSDF), FontBSDFPropValu);				Newline(Context);
+			Print(Context, TEXT(" Backlit           "), FontBSDFPropName);	Print(Context, HAIR_BACKLIT(BSDF), FontBSDFPropValu);				Newline(Context);
+			Print(Context, TEXT(" Has Transmittance "), FontBSDFPropName);	Print(Context, HAIR_COMPLEXTRANSMITTANCE(BSDF), FontBSDFPropValu);	Newline(Context);
 		}
-		else
+		break;
+
+		case STRATA_BSDF_TYPE_SINGLELAYERWATER:
 		{
-			NEXTLINE;
+			Print(Context, TEXT("Water"), FontBSDFType);
+
+			Newline(Context);
+			Print(Context, TEXT(" BaseColor            "), FontBSDFPropName);	Print(Context, SLW_BASECOLOR(BSDF), FontBSDFPropValu);			Newline(Context);
+			Print(Context, TEXT(" Metallic             "), FontBSDFPropName);	Print(Context, SLW_METALLIC(BSDF), FontBSDFPropValu);			Newline(Context);
+			Print(Context, TEXT(" Specular             "), FontBSDFPropName);	Print(Context, SLW_SPECULAR(BSDF), FontBSDFPropValu);			Newline(Context);
+			Print(Context, TEXT(" Roughness            "), FontBSDFPropName);	Print(Context, SLW_ROUGHNESS(BSDF), FontBSDFPropValu);			Newline(Context);
+			Print(Context, TEXT(" Top Material Opacity "), FontBSDFPropName);	Print(Context, SLW_TOPMATERIALOPACITY(BSDF), FontBSDFPropValu);	Newline(Context);
 		}
+		break;
 
-		HORIZONTALSEPARATION;
-		NEXTLINE;
-		HORIZONTALSEPARATION;
-#endif
+	default:
+		{
+			Print(Context, TEXT("Error"), FontRed);
+		}
+		break;
+	}
+
+	Newline(Context);
+}
+
+void PrintPixel(uint2 InCoord, float3 InWorldPosition, float3 V)
+{
+	FShaderPrintContext Context;
+	#if MULTIPASS_ENABLE
+	if (BSDFIndex > 0)
+	{
+		float2 PrintOffset;
+		PrintOffset.x = asfloat(RWPositionOffsetBuffer[0]);
+		PrintOffset.y = asfloat(RWPositionOffsetBuffer[1]);
+		Context = InitShaderPrintContext(true, PrintOffset);
+	}
+	else
+	#endif
+	{
+		Context = InitShaderPrintContext(true, uint2(50, 50));
+	}
+
+	FStrataAddressing StrataAddressing = GetStrataPixelDataByteOffset(InCoord, uint2(View.BufferSizeAndInvSize.xy), Strata.MaxBytesPerPixel);
+	const uint FootPrint_Start = StrataAddressing.ReadBytes;
+
+	FStrataPixelHeader Header = UnpackStrataHeaderIn(Strata.MaterialTextureArray, StrataAddressing, Strata.TopLayerTexture);
+	const uint FootPrint_PostHeader = StrataAddressing.ReadBytes;
+
+	const FFontColor FontBSDFCount = FontEmerald;
+	const FFontColor FontHeaderStateName = FontWhite;
+	const FFontColor FontHeaderPropName = FontWhite;
+	const FFontColor FontHeaderPropValue = FontSilver;
 
+	// Header
+	if (BSDFIndex < Header.BSDFCount && BSDFIndex == 0)
+	{
+		const FStrataTopLayerData TopLayerData = StrataUnpackTopLayerData(Strata.TopLayerTexture.Load(uint3(InCoord, 0)));
+		const FStrataSubsurfaceData SSSData = StrataUnpackSSSData(Strata.SSSTexture.Load(uint3(InCoord, 0)));
 
-#if COMPILER_DXC == 1
+		const bool bStrataMaterial = Header.BSDFCount > 0;
+		const bool bIsSimpleMaterial = IsSimpleMaterial(Header) || Header.BSDFCount == 0;
+		const bool bIsSingleMaterial = !IsSimpleMaterial(Header) && IsSingleMaterial(Header);
+
+		Print(Context, Header.BSDFCount, FontBSDFCount, 2, 0);
+		Print(Context, TEXT("BSDF"), FontBSDFCount);
+		Print(Context, TEXT(" - "), FontWhite);
+		PrintPixelType(Context, bIsSimpleMaterial, bIsSingleMaterial, FontLightRed);
+		Newline(Context);
+
+		Print(Context, TEXT(" AO    TopRoughness PreShadow ZeroShadow ContacShadow Ind.Occluder IsWater HasSSS"), FontHeaderStateName);
+		Newline(Context);
+
+		Print(Context, TEXT(" "));
+		Print(Context, StrataGetAO(Header), FontHeaderPropValue, 5, 3);									Print(Context, TEXT(" "));
+		Print(Context, TopLayerData.Roughness, FontHeaderPropValue, 5, 3);								Print(Context, TEXT("        "));
+		PrintBool(Context, HasPrecShadowMask(Header), FontHeaderPropValue);								Print(Context, TEXT("     "));
+		PrintBool(Context, HasZeroPrecShadowMask(Header), FontHeaderPropValue);							Print(Context, TEXT("      "));
+		PrintBool(Context, DoesCastContactShadow(Header), FontHeaderPropValue);							Print(Context, TEXT("        "));
+		PrintBool(Context, HasDynamicIndirectShadowCasterRepresentation(Header), FontHeaderPropValue);	Print(Context, TEXT("        "));
+		PrintBool(Context, IsSingleLayerWater(Header), FontHeaderPropValue);							Print(Context, TEXT("   "));
+		PrintBool(Context, HasSubsurface(Header), FontHeaderPropValue);
+		Newline(Context);
+		Newline(Context);
+	}
+
+	// BSDFs
+#if MULTIPASS_ENABLE
+	if (BSDFIndex < Header.BSDFCount)
+	{
+		if (BSDFIndex > 0)
+		{
+			StrataSeekBSDF(StrataAddressing, UnpackBSDFOffset(Strata.BSDFOffsetTexture[InCoord], Header.BSDFCount), BSDFIndex);
+		}
+		PrintBSDF(Context, StrataAddressing, Header, InWorldPosition, V);
+
+		RWPositionOffsetBuffer[0] = asuint(Context.Pos.x);
+		RWPositionOffsetBuffer[1] = asuint(Context.Pos.y);
+	}
+#else // MULTIPASS_ENABLE
+	if (BSDFIndex == 0 && Header.BSDFCount > 0)
+	{
+	#if 1
+		#if COMPILER_DXC == 1
 		LOOP
 		for (uint BSDFIndex = 0; BSDFIndex < Header.BSDFCount; ++BSDFIndex)
-#else
+		#else
 		Header.BSDFCount = min(Header.BSDFCount, STRATA_VIZ_MAX_BSDF);
 		UNROLL_N(STRATA_VIZ_MAX_BSDF) // FXC requires that we specify an upper bound here. Not DXC.
-		for (uint BSDFIndex = 0; BSDFIndex < Header.BSDFCount; ++BSDFIndex)
-#endif
+		for (uint LocalBSDFIndex = 0; LocalBSDFIndex < Header.BSDFCount; ++LocalBSDFIndex)
+		#endif
+	#endif
 		{
-			const FStrataBSDF BSDF = UnpackStrataBSDFIn(Strata.MaterialTextureArray, StrataAddressing, Header);
-
-			TopLeft.x = TopLeftInit.x;
+			PrintBSDF(Context, StrataAddressing, Header, InWorldPosition, V);
+		}
+	}
+#endif // MULTIPASS_ENABLE
 
-			switch (BSDF_GETTYPE(BSDF))
-			{
-			case STRATA_BSDF_TYPE_SLAB:
-#if PERMUTATION_BSDF_PASS==0
-				PRINTBSDFCHAR(_S_); PRINTBSDFCHAR(_L_); PRINTBSDFCHAR(_A_); PRINTBSDFCHAR(_B_); PRINTBSDFCHAR(_SPC_);
-				PRINTBSDFCHAR(_SPC_); PRINTBSDFCHAR(_SPC_); PRINTBSDFCHAR(_SPC_); PRINTBSDFCHAR(_SPC_); PRINTBSDFCHAR(_SPC_);
-				if (BSDF_GETHASGREYWEIGHT(BSDF))
-				{
-					PRINTBSDFGREYWEIGHT();
-				}
-				else
-				{
-					PRINTBSDFWEIGHT();
-				}
-				PRINTBSDFSHAREDLOCALBASISID();
-				PRINTBSDFANISO();
-				PRINTBSDFTOPLAY();
-				PRINTBSDFHASSCAT();
-				PRINTBSDFSIMPSCAT();
-				PRINTBSDFISTHIN();
-				NEXTLINE;
-
-				HORIZONTALTAB; PRINTPARAMCHAR(_D_); PRINTPARAMCHAR(_I_); PRINTPARAMCHAR(_F_); PRINTPARAMCHAR(_F_); PRINTPARAMCHAR(_U_); PRINTPARAMCHAR(_S_); 
-				HORIZONTALTAB; PRINTFLOAT3(SLAB_DIFFUSEALBEDO(BSDF));
-				NEXTLINE;
-
-				HORIZONTALTAB; PRINTPARAMCHAR(_F_); PRINTPARAMCHAR(_0_); PRINTPARAMCHAR(_SPC_); PRINTPARAMCHAR(_SPC_); PRINTPARAMCHAR(_SPC_); PRINTPARAMCHAR(_SPC_);
-				HORIZONTALTAB; PRINTFLOAT3(SLAB_F0(BSDF));
-				NEXTLINE;
-
-				if (BSDF_GETHASF90(BSDF))
-				{
-					HORIZONTALTAB; PRINTPARAMCHAR(_F_); PRINTPARAMCHAR(_9_); PRINTPARAMCHAR(_0_); PRINTPARAMCHAR(_SPC_); PRINTPARAMCHAR(_SPC_); PRINTPARAMCHAR(_SPC_);
-					HORIZONTALTAB; PRINTFLOAT3(SLAB_F90(BSDF));
-					NEXTLINE;
-				}
-
-				HORIZONTALTAB; PRINTPARAMCHAR(_R_); PRINTPARAMCHAR(_O_); PRINTPARAMCHAR(_U_); PRINTPARAMCHAR(_G_); PRINTPARAMCHAR(_H_); PRINTPARAMCHAR(_SPC_);
-				HORIZONTALTAB; PRINTFLOAT1(SLAB_ROUGHNESS(BSDF));
-				NEXTLINE;
-
-				if (BSDF_GETHASANISOTROPY(BSDF))
-				{				
-					HORIZONTALTAB; PRINTPARAMCHAR(_A_); PRINTPARAMCHAR(_N_); PRINTPARAMCHAR(_I_); PRINTPARAMCHAR(_S_); PRINTPARAMCHAR(_O_); PRINTPARAMCHAR(_SPC_);
-					HORIZONTALTAB; PRINTFLOAT1(SLAB_ANISOTROPY(BSDF));
-					NEXTLINE;
-				}
-
-				if (BSDF_GETHASHAZINESS(BSDF))
-				{
-					HORIZONTALTAB; PRINTPARAMCHAR(_H_); PRINTPARAMCHAR(_A_); PRINTPARAMCHAR(_Z_); PRINTPARAMCHAR(_I_); PRINTPARAMCHAR(_N_); PRINTPARAMCHAR(_SPC_);
-					HORIZONTALTAB; PRINTFLOAT1(SLAB_HAZINESS(BSDF));
-					NEXTLINE;
-				}
-
-				if (BSDF_GETHASSSS(BSDF) || BSDF_GETISSIMPLEVOLUME(BSDF) || BSDF_GETISTHIN(BSDF))
-				{
-					if (BSDF_GETHASSSSPROFILE(BSDF) && BSDF_GETISSIMPLEVOLUME(BSDF)==0)
-					{
-						HORIZONTALTAB; PRINTPARAMCHAR(_S_); PRINTPARAMCHAR(_S_); PRINTPARAMCHAR(_S_); PRINTPARAMCHAR(_SPC_); PRINTPARAMCHAR(_I_); PRINTPARAMCHAR(_D_); PRINTPARAMCHAR(_SPC_);
-						HORIZONTALTAB; PRINTFLOATNOFRAC(floor(SLAB_SSSPROFILEID(BSDF)*255.0f+0.5f),3);
-						NEXTLINE;
-						HORIZONTALTAB; PRINTPARAMCHAR(_S_); PRINTPARAMCHAR(_S_); PRINTPARAMCHAR(_S_); PRINTPARAMCHAR(_S_); PRINTPARAMCHAR(_C_); PRINTPARAMCHAR(_L_); PRINTPARAMCHAR(_SPC_);
-						HORIZONTALTAB; PRINTFLOAT1(SLAB_SSSPROFILERADIUSSCALE(BSDF));
-						NEXTLINE;
-					}
-					else
-					{
-						HORIZONTALTAB; PRINTPARAMCHAR(_S_); PRINTPARAMCHAR(_S_); PRINTPARAMCHAR(_S_); PRINTPARAMCHAR(_M_); PRINTPARAMCHAR(_F_); PRINTPARAMCHAR(_P_);
-						HORIZONTALTAB; PRINTFLOAT3(SLAB_SSSMFP(BSDF));
-						NEXTLINE;
-					}
-				}
-
-				if (BSDF_GETHASFUZZ(BSDF))
-				{
-					HORIZONTALTAB; PRINTPARAMCHAR(_F_); PRINTPARAMCHAR(_U_); PRINTPARAMCHAR(_Z_); PRINTPARAMCHAR(_Z_); PRINTPARAMCHAR(_SPC_); PRINTPARAMCHAR(_SPC_); PRINTPARAMCHAR(_SPC_);
-					HORIZONTALTAB; PRINTFLOAT1(SLAB_FUZZ_AMOUNT(BSDF));
-					NEXTLINE;
-					HORIZONTALTAB; PRINTPARAMCHAR(_F_); PRINTPARAMCHAR(_U_); PRINTPARAMCHAR(_Z_); PRINTPARAMCHAR(_C_); PRINTPARAMCHAR(_O_); PRINTPARAMCHAR(_L_); PRINTPARAMCHAR(_SPC_);
-					HORIZONTALTAB; PRINTFLOAT3(SLAB_FUZZ_COLOR(BSDF));
-					NEXTLINE;
-				}
-#else
-				NEXTLINE;
-				NEXTLINE;
-				NEXTLINE;
-				if (BSDF_GETHASF90(BSDF))
-				{
-					NEXTLINE;
-				}
-				if (BSDF_GETHASANISOTROPY(BSDF))
-				{
-					NEXTLINE;
-					NEXTLINE;
-				}
-				else
-				{
-					NEXTLINE;
-				}
-
-				if (BSDF_GETHASHAZINESS(BSDF))
-				{
-					NEXTLINE;
-				}
-
-				if (BSDF_GETHASSSS(BSDF) || BSDF_GETISSIMPLEVOLUME(BSDF) || BSDF_GETISTHIN(BSDF))
-				{
-					if (BSDF_GETHASSSSPROFILE(BSDF) && BSDF_GETISSIMPLEVOLUME(BSDF) == 0)
-					{
-						NEXTLINE;
-						NEXTLINE;
-					}
-					else
-					{
-						NEXTLINE;
-					}
-				}
-
-				if (BSDF_GETHASFUZZ(BSDF))
-				{
-					NEXTLINE;
-					NEXTLINE;
-				}
-#endif
-				break;
-
-			case STRATA_BSDF_TYPE_HAIR:
-#if PERMUTATION_BSDF_PASS==1
-				PRINTBSDFCHAR(_H_); PRINTBSDFCHAR(_A_); PRINTBSDFCHAR(_I_); PRINTBSDFCHAR(_R_); PRINTBSDFCHAR(_SPC_);
-				PRINTBSDFCHAR(_SPC_); PRINTBSDFCHAR(_SPC_); PRINTBSDFCHAR(_SPC_); PRINTBSDFCHAR(_SPC_); PRINTBSDFCHAR(_SPC_);
-				if (BSDF_GETHASGREYWEIGHT(BSDF))
-				{
-					PRINTBSDFGREYWEIGHT();
-				}
-				else
-				{
-					PRINTBSDFWEIGHT();
-				}
-				PRINTBSDFSHAREDLOCALBASISID();
-				PRINTBSDFANISO();
-				PRINTBSDFTOPLAY();
-				NEXTLINE;
-
-				HORIZONTALTAB; PRINTPARAMCHAR(_A_); PRINTPARAMCHAR(_L_); PRINTPARAMCHAR(_B_); PRINTPARAMCHAR(_E_); PRINTPARAMCHAR(_D_); PRINTPARAMCHAR(_O_);
-				HORIZONTALTAB; PRINTFLOAT3(HAIR_BASECOLOR(BSDF));
-				NEXTLINE;
-				HORIZONTALTAB; PRINTPARAMCHAR(_R_); PRINTPARAMCHAR(_O_); PRINTPARAMCHAR(_U_); PRINTPARAMCHAR(_G_); PRINTPARAMCHAR(_H_); PRINTPARAMCHAR(_SPC_);
-				HORIZONTALTAB; PRINTFLOAT1(HAIR_ROUGHNESS(BSDF));
-				NEXTLINE;
-				HORIZONTALTAB; PRINTPARAMCHAR(_S_); PRINTPARAMCHAR(_C_); PRINTPARAMCHAR(_A_); PRINTPARAMCHAR(_T_); PRINTPARAMCHAR(_T_); PRINTPARAMCHAR(_SPC_);
-				HORIZONTALTAB; PRINTFLOAT1(HAIR_SCATTER(BSDF));
-				NEXTLINE;
-				HORIZONTALTAB; PRINTPARAMCHAR(_S_); PRINTPARAMCHAR(_P_); PRINTPARAMCHAR(_E_); PRINTPARAMCHAR(_C_); PRINTPARAMCHAR(_U_); PRINTPARAMCHAR(_SPC_);
-				HORIZONTALTAB; PRINTFLOAT1(HAIR_SPECULAR(BSDF));
-				NEXTLINE;
-				HORIZONTALTAB; PRINTPARAMCHAR(_B_); PRINTPARAMCHAR(_A_); PRINTPARAMCHAR(_C_); PRINTPARAMCHAR(_K_); PRINTPARAMCHAR(_L_); PRINTPARAMCHAR(_SPC_);
-				HORIZONTALTAB; PRINTFLOAT1(HAIR_BACKLIT(BSDF));
-				NEXTLINE;
-				HORIZONTALTAB; PRINTPARAMCHAR(_T_); PRINTPARAMCHAR(_R_); PRINTPARAMCHAR(_A_); PRINTPARAMCHAR(_N_); PRINTPARAMCHAR(_S_); PRINTPARAMCHAR(_SPC_);
-				HORIZONTALTAB; PRINTFLOAT1(HAIR_COMPLEXTRANSMITTANCE(BSDF));
-				NEXTLINE;
-#else
-				NEXTLINE;
-				NEXTLINE;
-				NEXTLINE;
-				NEXTLINE;
-				NEXTLINE;
-				NEXTLINE;
-#endif
-				break;
-
-			case STRATA_BSDF_TYPE_SINGLELAYERWATER:
-#if PERMUTATION_BSDF_PASS==2
-				PRINTBSDFCHAR(_W_); PRINTBSDFCHAR(_A_); PRINTBSDFCHAR(_T_); PRINTBSDFCHAR(_E_); PRINTBSDFCHAR(_R_);
-				NEXTLINE;
-
-				HORIZONTALTAB; PRINTPARAMCHAR(_A_); PRINTPARAMCHAR(_L_); PRINTPARAMCHAR(_B_); PRINTPARAMCHAR(_E_); PRINTPARAMCHAR(_D_); PRINTPARAMCHAR(_O_);
-				HORIZONTALTAB; PRINTFLOAT3(SLW_BASECOLOR(BSDF));
-				NEXTLINE;
-				HORIZONTALTAB; PRINTPARAMCHAR(_M_); PRINTPARAMCHAR(_E_); PRINTPARAMCHAR(_T_); PRINTPARAMCHAR(_A_); PRINTPARAMCHAR(_L_); PRINTPARAMCHAR(_SPC_);
-				HORIZONTALTAB; PRINTFLOAT1(SLW_METALLIC(BSDF));
-				NEXTLINE;
-				HORIZONTALTAB; PRINTPARAMCHAR(_S_); PRINTPARAMCHAR(_P_); PRINTPARAMCHAR(_E_); PRINTPARAMCHAR(_C_); PRINTPARAMCHAR(_SPC_); PRINTPARAMCHAR(_SPC_);
-				HORIZONTALTAB; PRINTFLOAT1(SLW_SPECULAR(BSDF));
-				NEXTLINE;
-				HORIZONTALTAB; PRINTPARAMCHAR(_R_); PRINTPARAMCHAR(_O_); PRINTPARAMCHAR(_U_); PRINTPARAMCHAR(_G_); PRINTPARAMCHAR(_H_); PRINTPARAMCHAR(_SPC_);
-				HORIZONTALTAB; PRINTFLOAT1(SLW_ROUGHNESS(BSDF));
-				NEXTLINE;
-				HORIZONTALTAB; PRINTPARAMCHAR(_O_); PRINTPARAMCHAR(_P_); PRINTPARAMCHAR(_A_); PRINTPARAMCHAR(_C_); PRINTPARAMCHAR(_I_); PRINTPARAMCHAR(_T_);
-				HORIZONTALTAB; PRINTFLOAT1(SLW_TOPMATERIALOPACITY(BSDF));
-				NEXTLINE;
-#else
-				NEXTLINE;
-				NEXTLINE;
-				NEXTLINE;
-				NEXTLINE;
-				NEXTLINE;
-				NEXTLINE;
+	// Memory footprint
+#if MULTIPASS_ENABLE
+	if (BSDFIndex+1 == Header.BSDFCount)
 #endif
-				break;
+	{
+		const uint FootPrint_PostBSDFs = StrataAddressing.ReadBytes;
 
-			default:
-				PRINTBSDFCHAR(_E_); PRINTBSDFCHAR(_R_); PRINTBSDFCHAR(_R_); PRINTBSDFCHAR(_O_); PRINTBSDFCHAR(_R_);
-				NEXTLINE;
-				break;
-			}
+		const uint HeaderSize = FootPrint_PostHeader - FootPrint_Start;
+		const uint BSDFsSize  = FootPrint_PostBSDFs - FootPrint_PostHeader;
+		const uint TotalSize  = FootPrint_PostBSDFs - FootPrint_Start;
 
-			HORIZONTALSEPARATION;
-		}
+		Print(Context, TEXT("Memory"), FontEmerald); Newline(Context);
+		Print(Context, TEXT(" Header "), FontSilver); Print(Context, HeaderSize, FontSilver, 3, 0); Print(Context, TEXT("bytes"), FontSilver); Newline(Context);
+		Print(Context, TEXT(" BSDFs  "), FontSilver); Print(Context, BSDFsSize,  FontSilver, 3, 0); Print(Context, TEXT("bytes"), FontSilver); Newline(Context);
+		Print(Context, TEXT(" Total  "), FontOrange); Print(Context, TotalSize,  FontOrange, 3, 0); Print(Context, TEXT("bytes"), FontOrange); Newline(Context);
+	}
+}
+
+[numthreads(1, 1, 1)]
+void MaterialPrintInfoCS(uint3 DispatchThreadId : SV_DispatchThreadID)
+{
+	const uint2 PixelPos = View.CursorPosition;
+	const float2 BufferUV = SvPositionToBufferUV(float4(PixelPos + 0.5f, 0, 0));
+	const float2 ScreenPosition = SvPositionToScreenPosition(float4(PixelPos, 0.5f, 1)).xy;
+	const float DeviceZ = SampleDeviceZFromSceneTextures(BufferUV);
+	const float SceneDepth = ConvertFromDeviceZ(DeviceZ);
+	const float3 WorldPosition = mul(float4(ScreenPosition * SceneDepth, SceneDepth, 1), PrimaryView.ScreenToTranslatedWorld).xyz;
+	const float3 V = -normalize(WorldPosition - PrimaryView.TranslatedWorldCameraOrigin);
+
+	PrintPixel(PixelPos, WorldPosition, V);
+}
+#endif // SHADER_MATERIALPRINT
+
+///////////////////////////////////////////////////////////////////////////////////////////////////
+// Material Visualize
 
-		const uint PixelStrataDataByteOffsetPostBSDF = StrataAddressing.ReadBytes;
+#if SHADER_MATERIALVISUALIZE
 
-#if PERMUTATION_BSDF_PASS==0
-		PRINTINFOCHAR(_B_); PRINTINFOCHAR(_O_); PRINTINFOCHAR(_T_); PRINTINFOCHAR(_T_); PRINTINFOCHAR(_O_); PRINTINFOCHAR(_M_);
-		NEXTLINE;
-		HORIZONTALSEPARATION;
+uint ViewMode;
 
+void VisualizeMaterialPS(
+	float4 SVPos : SV_POSITION,
+	out float4 OutColor : SV_Target0)
+{
+	const uint2 PixelPos = uint2(SVPos.xy);
+	float2 BufferUV = SvPositionToBufferUV(float4(PixelPos, SVPos.zw));
 
-		PRINTBSDFCHAR(_H_); PRINTBSDFCHAR(_E_); PRINTBSDFCHAR(_A_); PRINTBSDFCHAR(_D_); PRINTBSDFCHAR(_E_); PRINTBSDFCHAR(_R_);
-		PRINTBSDFCHAR(_B_); PRINTBSDFCHAR(_Y_); PRINTBSDFCHAR(_T_); PRINTBSDFCHAR(_E_); PRINTBSDFCHAR(_S_); PRINTBSDFCHAR(_SPC_);
-		PRINTFLOATNOFRAC(float(PixelStrataDataByteOffsetPostHeader - PixelStrataDataByteOffsetStart), 3);
-		NEXTLINE;
+	const float2 ScreenPosition = SvPositionToScreenPosition(SVPos).xy;
+	const float2 ScreenMaterialPosition = SvPositionToScreenPosition(float4(PixelPos, 0.5f, 1)).xy;
+	float DeviceZ = SampleDeviceZFromSceneTextures(BufferUV);
+	float SceneDepth = ConvertFromDeviceZ(DeviceZ);
+	const float3 WorldPosition = mul(float4(ScreenMaterialPosition * SceneDepth, SceneDepth, 1), PrimaryView.ScreenToTranslatedWorld).xyz;
+	const float3 V = -normalize(WorldPosition - PrimaryView.TranslatedWorldCameraOrigin);
 
-		PRINTBSDFCHAR(_B_); PRINTBSDFCHAR(_S_); PRINTBSDFCHAR(_D_); PRINTBSDFCHAR(_F_); PRINTBSDFCHAR(_S_); PRINTBSDFCHAR(_SPC_);
-		PRINTBSDFCHAR(_B_); PRINTBSDFCHAR(_Y_); PRINTBSDFCHAR(_T_); PRINTBSDFCHAR(_E_); PRINTBSDFCHAR(_S_); PRINTBSDFCHAR(_SPC_);
-		PRINTFLOATNOFRAC(float(PixelStrataDataByteOffsetPostBSDF - PixelStrataDataByteOffsetPostHeader), 3);
-		NEXTLINE;
+	FStrataAddressing StrataAddressing = GetStrataPixelDataByteOffset(PixelPos, uint2(View.BufferSizeAndInvSize.xy), Strata.MaxBytesPerPixel);
+	FStrataPixelHeader Header = UnpackStrataHeaderIn(Strata.MaterialTextureArray, StrataAddressing, Strata.TopLayerTexture);
+	bool bStrataMaterial = Header.BSDFCount > 0;
 
-		PRINTBSDFCHAR(_T_); PRINTBSDFCHAR(_O_); PRINTBSDFCHAR(_T_); PRINTBSDFCHAR(_A_); PRINTBSDFCHAR(_L_); PRINTBSDFCHAR(_SPC_);
-		PRINTBSDFCHAR(_B_); PRINTBSDFCHAR(_Y_); PRINTBSDFCHAR(_T_); PRINTBSDFCHAR(_E_); PRINTBSDFCHAR(_S_); PRINTBSDFCHAR(_SPC_);
-		PRINTFLOATNOFRAC(float(PixelStrataDataByteOffsetPostBSDF - PixelStrataDataByteOffsetStart), 3);
-		NEXTLINE;
+	OutColor = 0;
 
-		// Background
-		if (all(PixelPos > RectMin) && all(PixelPos < RectMax) && !bPlotValid)
+	// BSDF count
+	if (ViewMode == 2)
+	{
+		OutColor.w = 0.1f;
+		if (Header.BSDFCount > 0)
 		{
-			OutColor.w = 0.0f;				// Set background visiblity to 0 in the bottom right part of the screen to better reads the debug values
+			const float BSDFCost = saturate(float(Header.BSDFCount) / 4.f);
+			OutColor.xyz = ColorMapTurbo(BSDFCost);
 		}
-#endif
+	}
+	else
+	{
+		OutColor = float4(0.0f, 0.0f, 0.0f, 1.0f);
 	}
 }
-
-
+#endif // SHADER_MATERIALVISUALIZE
\ No newline at end of file
diff --git a/Engine/Source/Runtime/Renderer/Private/Strata/Strata.cpp b/Engine/Source/Runtime/Renderer/Private/Strata/Strata.cpp
--- a/Engine/Source/Runtime/Renderer/Private/Strata/Strata.cpp
+++ b/Engine/Source/Runtime/Renderer/Private/Strata/Strata.cpp
@@ -66,6 +66,12 @@ static TAutoConsoleVariable<int32> CVarStrataTileOverflow(
 	TEXT("Scale the number of Strata tile for overflowing tiles containing multi-BSDFs pixels. (0: 0%, 1: 100%. Default 1.0f)."),
 	ECVF_RenderThreadSafe);
 
+static TAutoConsoleVariable<int32> CVarStrataDebugMode(
+	TEXT("r.Strata.DebugMode"),
+	1,
+	TEXT("Strata debug view mode."),
+	ECVF_RenderThreadSafe);
+
 IMPLEMENT_GLOBAL_SHADER_PARAMETER_STRUCT(FStrataGlobalUniformParameters, "Strata");
 
 void FStrataViewData::Reset()
@@ -413,19 +419,63 @@ TRDGUniformBufferRef<FStrataGlobalUniformParameters> BindStrataGlobalUniformPara
 // Debug
 
 #define VISUALIZE_MATERIAL_PASS_COUNT 3
+#define MULTIPASS_ENABLE 0
+
+class FMaterialPrintInfoCS : public FGlobalShader
+{
+	DECLARE_GLOBAL_SHADER(FMaterialPrintInfoCS);
+	SHADER_USE_PARAMETER_STRUCT(FMaterialPrintInfoCS, FGlobalShader);
+
+	using FPermutationDomain = TShaderPermutationDomain<>;
+
+	BEGIN_SHADER_PARAMETER_STRUCT(FParameters, )
+		SHADER_PARAMETER(uint32, BSDFIndex)
+		SHADER_PARAMETER_STRUCT_REF(FViewUniformShaderParameters, ViewUniformBuffer)
+		SHADER_PARAMETER_RDG_UNIFORM_BUFFER(FStrataGlobalUniformParameters, Strata)
+		SHADER_PARAMETER_STRUCT_INCLUDE(FSceneTextureParameters, SceneTextures)
+		SHADER_PARAMETER_STRUCT_INCLUDE(ShaderPrint::FShaderParameters, ShaderPrintParameters)
+		SHADER_PARAMETER_RDG_BUFFER_UAV(RWBuffer<uint>, RWPositionOffsetBuffer)
+	END_SHADER_PARAMETER_STRUCT()
+
+	static FPermutationDomain RemapPermutation(FPermutationDomain PermutationVector)
+	{
+		return PermutationVector;
+	}
+
+	static bool CanRunStrataVizualizeMaterial(EShaderPlatform Platform)
+	{
+		// On some consoles, this ALU heavy shader (and with optimisation disables for the sake of low compilation time) would spill registers. So only keep it for the editor.
+		return IsPCPlatform(Platform);
+	}
+
+	static bool ShouldCompilePermutation(const FGlobalShaderPermutationParameters& Parameters)
+	{
+		return GetMaxSupportedFeatureLevel(Parameters.Platform) >= ERHIFeatureLevel::SM5 && Strata::IsStrataEnabled() && CanRunStrataVizualizeMaterial(Parameters.Platform);
+	}
+
+	static void ModifyCompilationEnvironment(const FGlobalShaderPermutationParameters& Parameters, FShaderCompilerEnvironment& OutEnvironment)
+	{
+		FGlobalShader::ModifyCompilationEnvironment(Parameters, OutEnvironment);
+
+		// Stay debug and skip optimizations to reduce compilation time on this long shader.
+		OutEnvironment.CompilerFlags.Add(CFLAG_Debug);
+		OutEnvironment.SetDefine(TEXT("SHADER_MATERIALPRINT"), 1);
+		OutEnvironment.SetDefine(TEXT("MULTIPASS_ENABLE"), MULTIPASS_ENABLE);
+	}
+};
+IMPLEMENT_GLOBAL_SHADER(FMaterialPrintInfoCS, "/Engine/Private/Strata/StrataVisualize.usf", "MaterialPrintInfoCS", SF_Compute);
 
 class FVisualizeMaterialPS : public FGlobalShader
 {
 	DECLARE_GLOBAL_SHADER(FVisualizeMaterialPS);
 	SHADER_USE_PARAMETER_STRUCT(FVisualizeMaterialPS, FGlobalShader);
 
-	class FBSDFPass : SHADER_PERMUTATION_INT("PERMUTATION_BSDF_PASS", VISUALIZE_MATERIAL_PASS_COUNT);
-	using FPermutationDomain = TShaderPermutationDomain<FBSDFPass>;
+	using FPermutationDomain = TShaderPermutationDomain<>;
 
 	BEGIN_SHADER_PARAMETER_STRUCT(FParameters, )
+		SHADER_PARAMETER(uint32, ViewMode)
 		SHADER_PARAMETER_STRUCT_REF(FViewUniformShaderParameters, ViewUniformBuffer)
 		SHADER_PARAMETER_RDG_UNIFORM_BUFFER(FStrataGlobalUniformParameters, Strata)
-		SHADER_PARAMETER_TEXTURE(Texture2D, MiniFontTexture)
 		SHADER_PARAMETER_STRUCT_INCLUDE(FSceneTextureParameters, SceneTextures)
 		SHADER_PARAMETER_STRUCT_INCLUDE(ShaderPrint::FShaderParameters, ShaderPrintParameters)
 		RENDER_TARGET_BINDING_SLOTS()
@@ -444,8 +494,7 @@ class FVisualizeMaterialPS : public FGlobalShader
 
 	static bool ShouldCompilePermutation(const FGlobalShaderPermutationParameters& Parameters)
 	{
-		return GetMaxSupportedFeatureLevel(Parameters.Platform) >= ERHIFeatureLevel::SM5 && Strata::IsStrataEnabled()
-			&& CanRunStrataVizualizeMaterial(Parameters.Platform);
+		return GetMaxSupportedFeatureLevel(Parameters.Platform) >= ERHIFeatureLevel::SM5 && Strata::IsStrataEnabled() && CanRunStrataVizualizeMaterial(Parameters.Platform);
 	}
 
 	static void ModifyCompilationEnvironment(const FGlobalShaderPermutationParameters& Parameters, FShaderCompilerEnvironment& OutEnvironment)
@@ -454,6 +503,7 @@ class FVisualizeMaterialPS : public FGlobalShader
 
 		// Stay debug and skip optimizations to reduce compilation time on this long shader.
 		OutEnvironment.CompilerFlags.Add(CFLAG_Debug);
+		OutEnvironment.SetDefine(TEXT("SHADER_MATERIALVISUALIZE"), 1);
 	}
 };
 IMPLEMENT_GLOBAL_SHADER(FVisualizeMaterialPS, "/Engine/Private/Strata/StrataVisualize.usf", "VisualizeMaterialPS", SF_Pixel);
@@ -465,24 +515,53 @@ static void AddVisualizeMaterialPasses(FRDGBuilder& GraphBuilder, const FViewInf
 	{
 		if (!ShaderPrint::IsEnabled(View)) { ShaderPrint::SetEnabled(true); }
 		ShaderPrint::RequestSpaceForLines(64);
+		ShaderPrint::RequestSpaceForCharacters(1024);
 
 		FSceneTextureParameters SceneTextureParameters = GetSceneTextureParameters(GraphBuilder);
-		FVisualizeMaterialPS::FParameters* PassParameters = GraphBuilder.AllocParameters<FVisualizeMaterialPS::FParameters>();
-		PassParameters->ViewUniformBuffer = View.ViewUniformBuffer;
-		PassParameters->Strata = Strata::BindStrataGlobalUniformParameters(View);
-		PassParameters->MiniFontTexture = GetMiniFontTexture();
-		PassParameters->SceneTextures = GetSceneTextureParameters(GraphBuilder);
-		PassParameters->RenderTargets[0] = FRenderTargetBinding(SceneColorTexture, ERenderTargetLoadAction::ELoad);
-		ShaderPrint::SetParameters(GraphBuilder, View.ShaderPrintData, PassParameters->ShaderPrintParameters);
-
-		for (uint32 j = 0; j < VISUALIZE_MATERIAL_PASS_COUNT; ++j)
+
+		// Print Material info
+		{
+			uint32 BSDFIndex = 0;
+			FRDGBufferUAVRef PrintOffsetBufferUAV = nullptr;
+
+			#if MULTIPASS_ENABLE
+			FRDGBufferRef PrintOffsetBuffer = GraphBuilder.CreateBuffer(FRDGBufferDesc::CreateBufferDesc(4, 2), TEXT("Strata.DebugPrintPositionOffset"));
+			PrintOffsetBufferUAV = GraphBuilder.CreateUAV(PrintOffsetBuffer, PF_R32_UINT);
+			AddClearUAVPass(GraphBuilder, PrintOffsetBufferUAV, 50u);
+			const uint32 MaxBSDFCount = 8;
+
+			for (BSDFIndex; BSDFIndex < MaxBSDFCount; ++BSDFIndex)
+			#endif
+			{
+				FMaterialPrintInfoCS::FParameters* PassParameters = GraphBuilder.AllocParameters<FMaterialPrintInfoCS::FParameters>();
+				PassParameters->BSDFIndex = BSDFIndex;
+				PassParameters->RWPositionOffsetBuffer = PrintOffsetBufferUAV;
+				PassParameters->ViewUniformBuffer = View.ViewUniformBuffer;
+				PassParameters->Strata = Strata::BindStrataGlobalUniformParameters(View);
+				PassParameters->SceneTextures = GetSceneTextureParameters(GraphBuilder);
+				ShaderPrint::SetParameters(GraphBuilder, View.ShaderPrintData, PassParameters->ShaderPrintParameters);
+
+				TShaderMapRef<FMaterialPrintInfoCS> ComputeShader(View.ShaderMap);
+				FComputeShaderUtils::AddPass(GraphBuilder, RDG_EVENT_NAME("Strata::VisualizeMaterial(Print)"), ComputeShader, PassParameters, FIntVector(1, 1, 1));
+			}
+		}
+
+		// Draw material debug
+		const uint32 ViewMode = FMath::Max(0, CVarStrataDebugMode.GetValueOnRenderThread());
+		if (ViewMode > 1)
 		{
+			FVisualizeMaterialPS::FParameters* PassParameters = GraphBuilder.AllocParameters<FVisualizeMaterialPS::FParameters>();
+			PassParameters->ViewUniformBuffer = View.ViewUniformBuffer;
+			PassParameters->ViewMode = ViewMode;
+			PassParameters->Strata = Strata::BindStrataGlobalUniformParameters(View);
+			PassParameters->SceneTextures = GetSceneTextureParameters(GraphBuilder);
+			PassParameters->RenderTargets[0] = FRenderTargetBinding(SceneColorTexture, ERenderTargetLoadAction::ELoad);
+			ShaderPrint::SetParameters(GraphBuilder, View.ShaderPrintData, PassParameters->ShaderPrintParameters);
+
 			FVisualizeMaterialPS::FPermutationDomain PermutationVector;
-			PermutationVector.Set<typename FVisualizeMaterialPS::FBSDFPass>(j);
 			TShaderMapRef<FVisualizeMaterialPS> PixelShader(View.ShaderMap, PermutationVector);
 
-			FPixelShaderUtils::AddFullscreenPass<FVisualizeMaterialPS>(GraphBuilder, View.ShaderMap, RDG_EVENT_NAME("Strata::VisualizeMaterial"),
-				PixelShader, PassParameters, View.ViewRect, PreMultipliedColorTransmittanceBlend);
+			FPixelShaderUtils::AddFullscreenPass<FVisualizeMaterialPS>(GraphBuilder, View.ShaderMap, RDG_EVENT_NAME("Strata::VisualizeMaterial(Draw)"), PixelShader, PassParameters, View.ViewRect, PreMultipliedColorTransmittanceBlend);
 		}
 	}
 }
