Fix for shader compile error in forward shading

#rnx

diff --git a/Engine/Shaders/Private/ShadingModels.ush b/Engine/Shaders/Private/ShadingModels.ush
--- a/Engine/Shaders/Private/ShadingModels.ush
+++ b/Engine/Shaders/Private/ShadingModels.ush
@@ -568,8 +568,8 @@ void GetProfileDualSpecular(uint SubsurfaceProfileInt, float Roughness, float Op
 	GetSubsurfaceProfileDualSpecular(SubsurfaceProfileInt, Roughness, Opacity, LobeRoughness0, LobeRoughness1, LobeMix);
 #else
 	// Disable dual lobe, as subsurface profile doesn't work with forward.
-	Lobe0Roughness = Roughness;
-	Lobe1Roughness = Roughness;
+	LobeRoughness0 = Roughness;
+	LobeRoughness1 = Roughness;
 	LobeMix = 0.f;
 #endif
 }
