Fixed ray tracing translucency was looking overbright when the tonemapper was disabled.

#udn

diff --git a/Engine/Shaders/Private/RayTracing/RayTracingPrimaryRays.usf b/Engine/Shaders/Private/RayTracing/RayTracingPrimaryRays.usf
--- a/Engine/Shaders/Private/RayTracing/RayTracingPrimaryRays.usf
+++ b/Engine/Shaders/Private/RayTracing/RayTracingPrimaryRays.usf
@@ -338,10 +338,7 @@ RAY_TRACING_ENTRY_RAYGEN(RayTracingPrimaryRaysRGS)
 		FinalAlpha = SceneColorTexture.SampleLevel(GlobalPointClampedSampler, UV, 0).w;
 	}
 
-	if(ShouldUsePreExposure) 
-	{
-		PathRadiance.rgb *= View.PreExposure;
-	}
+	PathRadiance.rgb *= View.PreExposure;
 
 	PathRadiance = ClampToHalfFloatRange(PathRadiance);
 	ColorOutput[DispatchThreadId] = float4(PathRadiance, FinalAlpha);
