Remove duplication of clip logic in the velocity pass.

diff --git a/Engine/Shaders/Private/MaterialTemplate.ush b/Engine/Shaders/Private/MaterialTemplate.ush
--- a/Engine/Shaders/Private/MaterialTemplate.ush
+++ b/Engine/Shaders/Private/MaterialTemplate.ush
@@ -3081,20 +3081,6 @@ void GetMaterialClippingShadowDepth(FMaterialPixelParameters Parameters, FPixelM
 	#endif
 }
 
-// The material blending mode alone is not enough information. We also need to know if the material is ThinTranslucent because in this case we should never clip.
-// Indeed, ThinTranslucent surfaces always have a coverage of 1 and opacity represent the coverage of the opaque lit material sitting on top of the translucent one.
-void GetMaterialClippingVelocity(FMaterialPixelParameters Parameters, FPixelMaterialInputs PixelMaterialInputs, bool bIsThinTranslucent)
-{
-	ClipLODTransition(Parameters);
-	#if MATERIALBLENDING_MASKED && MATERIAL_DITHER_OPACITY_MASK
-		CLIP(GetMaterialMaskInputRaw(PixelMaterialInputs) - 1.0f / 255.0f);
-	#elif MATERIALBLENDING_MASKED
-		CLIP(GetMaterialMask(PixelMaterialInputs)); 
-	#elif MATERIALBLENDING_TRANSLUCENT || MATERIALBLENDING_ADDITIVE || MATERIALBLENDING_MODULATE
-		CLIP(bIsThinTranslucent ? 1.0f : GetMaterialOpacity(PixelMaterialInputs) - 1.0 / 255.0 - GetMaterialOpacityMaskClipValue());
-	#endif
-}
-
 #if MATERIAL_DITHER_OPACITY_MASK
 float DitheredOpacityMaskToOpacity(float Mask)
 {
@@ -3140,6 +3126,18 @@ void GetMaterialCoverageAndClipping(FMaterialPixelParameters Parameters, FPixelM
 #endif
 }
 
+// The material blending mode alone is not enough information. We also need to know if the material is ThinTranslucent because in this case we should never clip.
+// Indeed, ThinTranslucent surfaces always have a coverage of 1 and opacity represent the coverage of the opaque lit material sitting on top of the translucent one.
+void GetMaterialClippingVelocity(FMaterialPixelParameters Parameters, FPixelMaterialInputs PixelMaterialInputs, bool bIsThinTranslucent)
+{
+#if MATERIALBLENDING_TRANSLUCENT || MATERIALBLENDING_ADDITIVE || MATERIALBLENDING_MODULATE
+	ClipLODTransition(Parameters);
+	CLIP(bIsThinTranslucent ? 1.0f : GetMaterialOpacity(PixelMaterialInputs) - 1.0 / 255.0 - GetMaterialOpacityMaskClipValue());
+#else
+	GetMaterialCoverageAndClipping(Parameters, PixelMaterialInputs);
+#endif
+}
+
 #define MATERIALBLENDING_MASKED_USING_COVERAGE (FORWARD_SHADING && MATERIALBLENDING_MASKED && SUPPORTS_PIXEL_COVERAGE)
 #if MATERIALBLENDING_MASKED_USING_COVERAGE
 
diff --git a/Engine/Shaders/Private/VelocityShader.usf b/Engine/Shaders/Private/VelocityShader.usf
--- a/Engine/Shaders/Private/VelocityShader.usf
+++ b/Engine/Shaders/Private/VelocityShader.usf
@@ -167,7 +167,6 @@ void MainPixelShader(
 	FMaterialPixelParameters MaterialParameters = GetMaterialPixelParameters(Input.FactoryInterpolants, SvPosition);
 	FPixelMaterialInputs PixelMaterialInputs;
 	CalcMaterialParameters(MaterialParameters, PixelMaterialInputs, SvPosition, bIsFrontFace);
-	GetMaterialClippingVelocity(MaterialParameters, PixelMaterialInputs, VELOCITY_THIN_TRANSLUCENT_MODE > 0);
 
 #if OUTPUT_PIXEL_DEPTH_OFFSET
 	ApplyPixelDepthOffsetToMaterialParameters(MaterialParameters, PixelMaterialInputs, OutDepth);
@@ -176,7 +175,7 @@ void MainPixelShader(
 #if MATERIALBLENDING_MASKED_USING_COVERAGE
 	OutCoverage = DiscardMaterialWithPixelCoverage(MaterialParameters, PixelMaterialInputs);
 #else
-	GetMaterialCoverageAndClipping(MaterialParameters, PixelMaterialInputs);
+	GetMaterialClippingVelocity(MaterialParameters, PixelMaterialInputs, VELOCITY_THIN_TRANSLUCENT_MODE > 0);
 #endif
 
 	// 3d velocity, includes camera an object motion
