Rename Strata light transmittance from 'LightWeight' into WeightL.

diff --git a/Engine/Shaders/Private/Strata/Strata.ush b/Engine/Shaders/Private/Strata/Strata.ush
--- a/Engine/Shaders/Private/Strata/Strata.ush
+++ b/Engine/Shaders/Private/Strata/Strata.ush
@@ -174,16 +174,16 @@ struct FStrataBSDF
 //            BASIS(0)      32 |  Normal: Oct 16/16 | Tangent: 11/11/10
 //            BASIS(1)      32 |  Normal: Oct 16/16 | Tangent: 11/11/10
 //                             |
-//            STATE(0)      16 |  BSDF Type | Normal Id | Aniso | TopLayer | SSS | SSSProfile | GreyWeigh | Haze | F90 | Simple vol. | MFP Plugged | Fuzz | Thin
-//                             |  3           2           1       1          1     1            1           1      1     1             1             1      1 
+//            STATE(0)      16 |  BSDF Type | Normal Id | Aniso | TopLayer | SSS | SSSProfile | GreyWeighV | Haze | F90 | Simple vol. | MFP Plugged | Fuzz | Thin
+//                             |  3           2           1       1          1     1            1            1      1     1             1             1      1 
 //            BSDF(0)       ?? |  Adaptive
-//            STATE(1)      16 |  BSDF Type | Normal Id | Aniso | TopLayer | SSS | SSSProfile | GreyWeigh | Haze | F90 | Simple vol. | MFP Plugged | Fuzz | Thin
-//                             |  3           2           1       1          1     1            1           1      1     1             1             1      1 
+//            STATE(1)      16 |  BSDF Type | Normal Id | Aniso | TopLayer | SSS | SSSProfile | GreyWeighV | Haze | F90 | Simple vol. | MFP Plugged | Fuzz | Thin
+//                             |  3           2           1       1          1     1            1            1      1     1             1             1      1 
 //            BSDF(1)       ?? |  Adaptive
 //                             |
 // * LAYOUT 1 -------------------------------------------------------------------------------------------------------------------------------------------------------------------------
-//  (Single)  HEADER        30 |  State | AO | BSDF Types | Aniso | TopLayer | SSS | SSSProfile | GreyWeigh | Haze | F90 | Simple vol. | MFP Plugged | Fuzz | Thin
-//                             |  8       8    4            1       1          1     1            1           1      1     1             1             1      1            
+//  (Single)  HEADER        30 |  State | AO | BSDF Types | Aniso | TopLayer | SSS | SSSProfile | GreyWeighV | Haze | F90 | Simple vol. | MFP Plugged | Fuzz | Thin
+//                             |  8       8    4            1       1          1     1            1            1      1     1             1             1      1            
 //            BASIS(0)       0 | (Use top layer normal)
 //            BSDF(0)       ?? |  Adaptive
 // * LAYOUT 0 -------------------------------------------------------------------------------------------------------------------------------------------------------------------------          
@@ -280,7 +280,7 @@ struct FStrataBSDF
 #define STATE_BIT_COUNT_TOPLAYER				1
 #define STATE_BIT_COUNT_HASSSS					1
 #define STATE_BIT_COUNT_HASSSSPROFILE			1
-#define STATE_BIT_COUNT_GREYWEIGHT				1
+#define STATE_BIT_COUNT_GREYWEIGHT_V			1
 #define STATE_BIT_COUNT_HAZINESS				1
 #define STATE_BIT_COUNT_F90						1
 #define STATE_BIT_COUNT_SIMPLEVOLUME			1
@@ -289,7 +289,7 @@ struct FStrataBSDF
 #define STATE_BIT_COUNT_ISTHIN					1
 #define STATE_BIT_COUNT_BSDF_TYPE				3
 #define STATE_BIT_COUNT_SHAREDLOCALBASESID		2
-#define STATE_BIT_COUNT_HASLIGHTWEIGHT			1
+#define STATE_BIT_COUNT_HASWEIGHT_L				1
 #define STATE_BIT_COUNT___UNUSED___				5
 #define STATE_BIT_COUNT_WEIGHT10F				10
 
@@ -298,8 +298,8 @@ struct FStrataBSDF
 #define STATE_BIT_OFFSET_TOPLAYER				(STATE_BIT_OFFSET_ANISOTROPY		+ STATE_BIT_COUNT_ANISOTROPY)
 #define STATE_BIT_OFFSET_HASSSS					(STATE_BIT_OFFSET_TOPLAYER			+ STATE_BIT_COUNT_TOPLAYER)
 #define STATE_BIT_OFFSET_HASSSSPROFILE			(STATE_BIT_OFFSET_HASSSS			+ STATE_BIT_COUNT_HASSSS)
-#define STATE_BIT_OFFSET_GREYWEIGHT				(STATE_BIT_OFFSET_HASSSSPROFILE		+ STATE_BIT_COUNT_HASSSSPROFILE)
-#define STATE_BIT_OFFSET_HAZINESS				(STATE_BIT_OFFSET_GREYWEIGHT		+ STATE_BIT_COUNT_GREYWEIGHT)
+#define STATE_BIT_OFFSET_GREYWEIGHT_V			(STATE_BIT_OFFSET_HASSSSPROFILE		+ STATE_BIT_COUNT_HASSSSPROFILE)
+#define STATE_BIT_OFFSET_HAZINESS				(STATE_BIT_OFFSET_GREYWEIGHT_V		+ STATE_BIT_COUNT_GREYWEIGHT_V)
 #define STATE_BIT_OFFSET_F90					(STATE_BIT_OFFSET_HAZINESS			+ STATE_BIT_COUNT_HAZINESS)
 #define STATE_BIT_OFFSET_SIMPLEVOLUME			(STATE_BIT_OFFSET_F90				+ STATE_BIT_COUNT_F90)
 #define STATE_BIT_OFFSET_MFPPLUGGED				(STATE_BIT_OFFSET_SIMPLEVOLUME		+ STATE_BIT_COUNT_SIMPLEVOLUME)
@@ -307,7 +307,7 @@ struct FStrataBSDF
 #define STATE_BIT_OFFSET_ISTHIN					(STATE_BIT_OFFSET_HASFUZZ			+ STATE_BIT_COUNT_HASFUZZ)
 #define STATE_BIT_OFFSET_BSDF_TYPE				(STATE_BIT_OFFSET_ISTHIN			+ STATE_BIT_COUNT_ISTHIN)
 #define STATE_BIT_OFFSET_SHAREDLOCALBASESID		(STATE_BIT_OFFSET_BSDF_TYPE			+ STATE_BIT_COUNT_BSDF_TYPE)
-#define STATE_BIT_OFFSET_HASLIGHTWEIGHT			(STATE_BIT_OFFSET_SHAREDLOCALBASESID+ STATE_BIT_COUNT_SHAREDLOCALBASESID)
+#define STATE_BIT_OFFSET_HASWEIGHT_L			(STATE_BIT_OFFSET_SHAREDLOCALBASESID+ STATE_BIT_COUNT_SHAREDLOCALBASESID)
 #define STATE_BIT_OFFSET_WEIGHT10F				22
 
 // Layout 0 (Fast)
@@ -352,8 +352,8 @@ struct FStrataBSDF
 #define BSDF_SETHASSSSPROFILE(X, Profi)	WRITE_BITS(X.State, STATE_BIT_COUNT_HASSSSPROFILE, STATE_BIT_OFFSET_HASSSSPROFILE, Profi)
 
 // Indicates if the BSDF weight is grey scale so that it can be stored in a single float (e.g. top layer BSDFs)
-#define BSDF_GETHASGREYWEIGHT(X)		READ_BITS(X.State,  STATE_BIT_COUNT_GREYWEIGHT, STATE_BIT_OFFSET_GREYWEIGHT)
-#define BSDF_SETHASGREYWEIGHT(X, Grey)	WRITE_BITS(X.State, STATE_BIT_COUNT_GREYWEIGHT, STATE_BIT_OFFSET_GREYWEIGHT, Grey)
+#define BSDF_GETHASGREYWEIGHT_V(X)		READ_BITS(X.State,  STATE_BIT_COUNT_GREYWEIGHT_V, STATE_BIT_OFFSET_GREYWEIGHT_V)
+#define BSDF_SETHASGREYWEIGHT_V(X, Grey)WRITE_BITS(X.State, STATE_BIT_COUNT_GREYWEIGHT_V, STATE_BIT_OFFSET_GREYWEIGHT_V, Grey)
 
 // Indicates if the BSDF has haziness data
 #define BSDF_GETHASHAZINESS(X)			READ_BITS(X.State,  STATE_BIT_COUNT_HAZINESS, STATE_BIT_OFFSET_HAZINESS)
@@ -384,8 +384,8 @@ struct FStrataBSDF
 #define BSDF_SETISTHIN(X, Value)		WRITE_BITS(X.State, STATE_BIT_COUNT_ISTHIN, STATE_BIT_OFFSET_ISTHIN, Value)
 
 // Indicates if the BSDF has light weight data
-#define BSDF_GETHASLIGHTWEIGHT(X)		READ_BITS(X.State,  STATE_BIT_COUNT_HASLIGHTWEIGHT, STATE_BIT_OFFSET_HASLIGHTWEIGHT)
-#define BSDF_SETHASLIGHTWEIGHT(X, Value) WRITE_BITS(X.State, STATE_BIT_COUNT_HASLIGHTWEIGHT, STATE_BIT_OFFSET_HASLIGHTWEIGHT, Value)
+#define BSDF_GETHASWEIGHT_L(X)			READ_BITS(X.State,  STATE_BIT_COUNT_HASWEIGHT_L, STATE_BIT_OFFSET_HASWEIGHT_L)
+#define BSDF_SETHASWEIGHT_L(X, Value)	WRITE_BITS(X.State, STATE_BIT_COUNT_HASWEIGHT_L, STATE_BIT_OFFSET_HASWEIGHT_L, Value)
 
 #if STRATA_INLINE_SHADING
 #define BSDF_SETEMISSIVE(X, V)			X.Emissive = V
@@ -1535,7 +1535,7 @@ FStrataData StrataHorizontalMixingParameterBlending(FStrataData Background, FStr
 	BSDF_SETHASHAZINESS(ResultBSDF,		BSDF_GETHASHAZINESS(ForegroundBSDF)				| BSDF_GETHASHAZINESS(BackgroundBSDF));
 	BSDF_SETHASF90(ResultBSDF,			BSDF_GETHASF90(ForegroundBSDF)					| BSDF_GETHASF90(BackgroundBSDF));
 	BSDF_SETHASFUZZ(ResultBSDF,			BSDF_GETHASFUZZ(ForegroundBSDF)					| BSDF_GETHASFUZZ(BackgroundBSDF));
-	BSDF_SETHASLIGHTWEIGHT(ResultBSDF,	BSDF_GETHASLIGHTWEIGHT(ForegroundBSDF)			| BSDF_GETHASLIGHTWEIGHT(BackgroundBSDF));
+	BSDF_SETHASWEIGHT_L(ResultBSDF,		BSDF_GETHASWEIGHT_L(ForegroundBSDF)				| BSDF_GETHASWEIGHT_L(BackgroundBSDF));
 
 	BSDF_SETEMISSIVE(ResultBSDF,		lerp(BSDF_GETEMISSIVE(BackgroundBSDF),			BSDF_GETEMISSIVE(ForegroundBSDF),			ForegroundOtherMixFactor));
 	SLAB_DIFFUSEALBEDO(ResultBSDF) =	lerp(SLAB_DIFFUSEALBEDO(BackgroundBSDF),		SLAB_DIFFUSEALBEDO(ForegroundBSDF),			ForegroundOtherMixFactor);
@@ -2952,11 +2952,11 @@ void PackStrataOut(
 				BRANCH
 				if (bIsVisible)
 				{
-					const uint GreyScaleThroughput = StrataHasGreyScaleWeight(BSDF.LuminanceWeightV) ? 1 : 0;
-					BSDF_SETHASGREYWEIGHT(BSDF, GreyScaleThroughput);
+					const uint GreyScaleThroughputV = StrataHasGreyScaleWeight(BSDF.LuminanceWeightV) ? 1 : 0;
+					BSDF_SETHASGREYWEIGHT_V(BSDF, GreyScaleThroughputV);
 
 					const bool bHasLightTransmittance = any(BSDF.LuminanceWeightL < 1.0f);
-					BSDF_SETHASLIGHTWEIGHT(BSDF, bHasLightTransmittance);
+					BSDF_SETHASWEIGHT_L(BSDF, bHasLightTransmittance);
 
 					// A. Layout0 - Simple encoding (aka. fast-path): store header & data
 					if (bHasFastEncoding && BSDFIdx == SimpleMaterial_BSDFIndex)
@@ -3008,7 +3008,7 @@ void PackStrataOut(
 						}
 					}
 					// C. Layout2 - Weight for Regular/Complex path
-					else if (GreyScaleThroughput > 0)
+					else if (GreyScaleThroughputV > 0)
 					{
 						BSDF_SETWEIGHT10F(BSDF, Pack10F(BSDF.LuminanceWeightV.x));
 						STRATA_STORE_UINT1(BSDF.State);
@@ -3081,7 +3081,7 @@ void PackStrataOut(
 							STRATA_STORE_UINT1(PackColorLinearToGamma2AlphaLinear(float4(SLAB_FUZZ_COLOR(BSDF), SLAB_FUZZ_AMOUNT(BSDF))));
 						}
 
-						if (BSDF_GETHASLIGHTWEIGHT(BSDF))
+						if (BSDF_GETHASWEIGHT_L(BSDF))
 						{
 							STRATA_STORE_UINT1(PackR11G11B10F(BSDF.LuminanceWeightL));
 						}
@@ -3250,7 +3250,7 @@ FStrataBSDF UnpackFastPathStrataBSDFIn(FStrataMaterialContainer StrataBuffer, in
 	// Force a BSDF state to a single slab with isotropic specular, no subsurface scattering or any other options.
 	// This is to help the compiler understand disabled path.
 	BSDF_SETTYPE(OutBSDF, STRATA_BSDF_TYPE_SLAB);
-	BSDF_SETHASGREYWEIGHT(OutBSDF, 1);
+	BSDF_SETHASGREYWEIGHT_V(OutBSDF, 1);
 	BSDF_SETHASF90(OutBSDF, 0);
 	BSDF_SETHASANISOTROPY(OutBSDF, 0);
 	BSDF_SETISTOPLAYER(OutBSDF, 1);
@@ -3259,6 +3259,7 @@ FStrataBSDF UnpackFastPathStrataBSDFIn(FStrataMaterialContainer StrataBuffer, in
 	BSDF_SETHASHAZINESS(OutBSDF, 0);
 	BSDF_SETISSIMPLEVOLUME(OutBSDF, 0);
 	BSDF_SETHASFUZZ(OutBSDF, 0);
+	BSDF_SETHASWEIGHT_L(OutBSDF, 0);
 
 	return OutBSDF;
 }
@@ -3297,7 +3298,7 @@ FStrataBSDF UnpackStrataBSDFIn(FStrataMaterialContainer StrataBuffer, inout FStr
 	else
 	{
 		OutBSDF.State = StrataLoadUint1(StrataBuffer, StrataAddressing);
-		if (BSDF_GETHASGREYWEIGHT(OutBSDF))
+		if (BSDF_GETHASGREYWEIGHT_V(OutBSDF))
 		{
 			OutBSDF.LuminanceWeightV = Unpack10F(BSDF_GETWEIGHT10F(OutBSDF));
 		}
@@ -3374,7 +3375,7 @@ FStrataBSDF UnpackStrataBSDFIn(FStrataMaterialContainer StrataBuffer, inout FStr
 			StrataLoad_ColorGamma2ToLinearAlphaLinear(StrataBuffer, StrataAddressing, SLAB_FUZZ_COLOR(OutBSDF), SLAB_FUZZ_AMOUNT(OutBSDF));
 		}
 
-		if (BSDF_GETHASLIGHTWEIGHT(OutBSDF))
+		if (BSDF_GETHASWEIGHT_L(OutBSDF))
 		{
 			StrataLoad_R11G11B10F(StrataBuffer, StrataAddressing, OutBSDF.LuminanceWeightL);
 		}
diff --git a/Engine/Shaders/Private/Strata/StrataVisualize.usf b/Engine/Shaders/Private/Strata/StrataVisualize.usf
--- a/Engine/Shaders/Private/Strata/StrataVisualize.usf
+++ b/Engine/Shaders/Private/Strata/StrataVisualize.usf
@@ -129,7 +129,9 @@ void PrintBSDF(inout FShaderPrintContext Context, inout FStrataAddressing Strata
 	{
 		case STRATA_BSDF_TYPE_SLAB:
 		{
-			const bool bGreyWeightV = BSDF_GETHASGREYWEIGHT(BSDF);
+			const bool bWeightL = BSDF_GETHASWEIGHT_L(BSDF);
+			const bool bGreyWeightV = BSDF_GETHASGREYWEIGHT_V(BSDF);
+			const bool bGreyWeightL = BSDF.LuminanceWeightL.x == BSDF.LuminanceWeightL.y && BSDF.LuminanceWeightL.x == BSDF.LuminanceWeightL.z;
 
 			Print(Context, TEXT("Slab "), FontBSDFType);
 			Newline(Context);
@@ -138,7 +140,10 @@ void PrintBSDF(inout FShaderPrintContext Context, inout FStrataAddressing Strata
 			{
 				Print(Context, TEXT("           "), FontBSDFStateName);
 			}
-			Print(Context, TEXT("WeightL"), FontBSDFStateName);
+			if (bWeightL)
+			{
+				Print(Context, TEXT("WeightL"), FontBSDFStateName);
+			}
 			Newline(Context);
 
 			Print(Context, TEXT(" "), FontSilver);
@@ -150,23 +155,24 @@ void PrintBSDF(inout FShaderPrintContext Context, inout FStrataAddressing Strata
 			PrintBool(Context, BSDF_GETISTHIN(BSDF), FontBSDF);					Print(Context, TEXT("  "), FontSilver);
 
 			// View Weight
-			if (BSDF_GETHASGREYWEIGHT(BSDF))
-			{
-				Print(Context, BSDF.LuminanceWeightV.x, FontSilver, 6, 4);
-			}
-			else
 			{
-				Print(Context, BSDF.LuminanceWeightV.r, FontLightRed,   6, 4);
-				Print(Context, BSDF.LuminanceWeightV.g, FontLightGreen, 6, 4);
-				Print(Context, BSDF.LuminanceWeightV.b, FontLightBlue,  6, 4);
+				if (bGreyWeightV)
+				{
+					Print(Context, BSDF.LuminanceWeightV.x, FontSilver, 6, 4);
+				}
+				else
+				{
+					Print(Context, BSDF.LuminanceWeightV.r, FontLightRed,   6, 4);
+					Print(Context, BSDF.LuminanceWeightV.g, FontLightGreen, 6, 4);
+					Print(Context, BSDF.LuminanceWeightV.b, FontLightBlue,  6, 4);
+				}
 			}
 
 			// Lighy weight
-			if (any(BSDF.LuminanceWeightL < 1.f))
+			if (bWeightL)
 			{
 				Print(Context, TEXT(" "), FontSilver);
-				const bool bGreyLuminanceL = BSDF.LuminanceWeightL.x == BSDF.LuminanceWeightL.y && BSDF.LuminanceWeightL.x == BSDF.LuminanceWeightL.z;
-				if (bGreyLuminanceL)
+				if (bGreyWeightL)
 				{
 					Print(Context, BSDF.LuminanceWeightL.x, FontSilver);
 				}
@@ -250,7 +256,7 @@ void PrintBSDF(inout FShaderPrintContext Context, inout FStrataAddressing Strata
 			Print(Context, BSDF_GETSHAREDLOCALBASISID(BSDF), FontBSDF);
 			PrintBool(Context, BSDF_GETHASANISOTROPY(BSDF), FontBSDF);
 			PrintBool(Context, BSDF_GETISTOPLAYER(BSDF), FontBSDF);
-			if (BSDF_GETHASGREYWEIGHT(BSDF))
+			if (BSDF_GETHASGREYWEIGHT_V(BSDF))
 			{
 				Print(Context, BSDF.LuminanceWeightV.x, FontSilver);
 			}
