Strata - fixed crash with r.Strata.Debug.AdvancedVisualizationShaders=1

diff --git a/Engine/Shaders/Private/BasePassPixelShader.usf b/Engine/Shaders/Private/BasePassPixelShader.usf
--- a/Engine/Shaders/Private/BasePassPixelShader.usf
+++ b/Engine/Shaders/Private/BasePassPixelShader.usf
@@ -1732,9 +1732,9 @@ void FPixelShaderInOut_MainPS(
 		FStrataOpaqueRoughRefractionData OpaqueRoughRefractionData = (FStrataOpaqueRoughRefractionData)0;
 		FStrataIntegrationSettings Settings = InitStrataIntegrationSettings(false /*bForceFullyRough*/, StrataStruct.bRoughDiffuse, StrataStruct.PeelLayersAboveDepth);
 
-#if STRATA_ADVANCED_DEBUG_ENABLED
-		Settings.SliceStoringDebugStrataTree = all(View.CursorPosition == PixelPos) ? Strata.SliceStoringDebugStrataTree : Settings.SliceStoringDebugStrataTree;
-#endif
+	#if STRATA_ADVANCED_DEBUG_ENABLED
+		Settings.SliceStoringDebugStrataTree = all(View.CursorPosition == PixelPos) ? StrataStruct.SliceStoringDebugStrataTree : Settings.SliceStoringDebugStrataTree;
+	#endif
 
 		// Generate the strata material data to write out
 		FStrataAddressing StrataAddressing = GetStrataPixelDataByteOffset(PixelPos, uint2(ResolvedView.BufferSizeAndInvSize.xy), StrataStruct.MaxBytesPerPixel);
diff --git a/Engine/Shaders/Private/Strata/Strata.ush b/Engine/Shaders/Private/Strata/Strata.ush
--- a/Engine/Shaders/Private/Strata/Strata.ush
+++ b/Engine/Shaders/Private/Strata/Strata.ush
@@ -814,7 +814,7 @@ FStrataTree GetInitialisedStrataTree()
 }
 
 void StrataPackOutStrataTree(
-	inout RWTexture2DArray<uint> ExtraMaterialDataUAV,
+	RWTexture2DArray<uint> ExtraMaterialDataUAV,
 	in FStrataTree StrataTree,
 	in int SliceStoringDebugStrataTree)
 {
@@ -837,7 +837,7 @@ void StrataPackOutStrataTree(
 		WRITE_UINT(StrataTree.BSDFs[i].bIsBottom);
 		WRITE_UINT(StrataTree.BSDFs[i].bIsTop);
 
-#if STRATA_INLINE_SHADING
+	#if STRATA_INLINE_SHADING
 		WRITE_FLOAT(StrataTree.BSDFs[i].Coverage);
 
 		WRITE_FLOAT3(StrataTree.BSDFs[i].Emissive);
@@ -845,7 +845,7 @@ void StrataPackOutStrataTree(
 		WRITE_FLOAT(StrataTree.BSDFs[i].TmpMFP);
 
 		WRITE_FLOAT(StrataTree.BSDFs[i].TopLayerDataWeight);
-#endif
+	#endif
 
 		// Skipping specific single layer water data
 
