Fix Lumen lighting fetch onto non-initialized textures when Strata is enabled with material having haziness or SSS profile.

diff --git a/Engine/Shaders/Private/DiffuseIndirectComposite.usf b/Engine/Shaders/Private/DiffuseIndirectComposite.usf
--- a/Engine/Shaders/Private/DiffuseIndirectComposite.usf
+++ b/Engine/Shaders/Private/DiffuseIndirectComposite.usf
@@ -230,7 +230,9 @@ FDirectLighting StrataLumenLighting(
 		#if STRATA_FASTPATH==0
 		if (BSDF_GETHASHAZINESS(BSDF))
 		{
-			Out.Specular += BSDFThroughput * InSpecularLighting * StrataEnvLight.SpecularHazeWeight * Occlusion.SpecularOcclusion;
+			const float HazeRoughness = StrataComputeHazeRoughness(Roughness);
+			const float3 InHazeGlossyLighting = MixSpecularAndRoughReflections(HazeRoughness, InSpecularLighting, InRoughSpecularIndirectLighting);
+			Out.Specular += BSDFThroughput * InHazeGlossyLighting * StrataEnvLight.SpecularHazeWeight * Occlusion.SpecularOcclusion;
 		}
 		#endif 
 	}
