Fix shader compilation with fxc when strata is enabled.

diff --git a/Engine/Shaders/Private/Strata/Strata.ush b/Engine/Shaders/Private/Strata/Strata.ush
--- a/Engine/Shaders/Private/Strata/Strata.ush
+++ b/Engine/Shaders/Private/Strata/Strata.ush
@@ -827,7 +827,12 @@ void StrataPackOutStrataTree(
 	WRITE_UINT(StrataTree.BSDFCount);
 	WRITE_UINT(StrataTree.OperatorCount);
 
+#if COMPILER_DXC == 1
 	LOOP
+#else
+	StrataTree.BSDFCount = min(StrataTree.BSDFCount, STRATA_MAX_BSDF_COUNT);
+	STRATA_UNROLL_N(STRATA_MAX_BSDF_COUNT)
+#endif
 	for(int i=0; i< StrataTree.BSDFCount; ++i)
 	{
 		WRITE_UINT(StrataTree.BSDFs[i].State);
@@ -877,7 +882,12 @@ void StrataUnpackInStrataTree(
 	READ_UINT(StrataTree.BSDFCount);
 	READ_UINT(StrataTree.OperatorCount);
 
+#if COMPILER_DXC == 1
 	LOOP
+#else
+	StrataTree.BSDFCount = min(StrataTree.BSDFCount, STRATA_MAX_BSDF_COUNT);
+	STRATA_UNROLL_N(STRATA_MAX_BSDF_COUNT)
+#endif
 	for (int i = 0; i < StrataTree.BSDFCount; ++i)
 	{
 		READ_UINT(StrataTree.BSDFs[i].State);
