Clustered local lights and reflections are supported on mobile forward and deferred.

Local light shadows are supported on mobile forward and deferred, need to enable full depth prepass to generate the ShadowMaskTexture.

Remove the 4 maximum local lights limitation on mobile forward.

Remove the "High Quality Reflection" on mobile forward.

Disable simple lights on mobile forward.

diff --git a/Engine/Plugins/Runtime/Oculus/OculusVR/Source/OculusEditor/Private/OculusToolWidget.cpp b/Engine/Plugins/Runtime/Oculus/OculusVR/Source/OculusEditor/Private/OculusToolWidget.cpp
--- a/Engine/Plugins/Runtime/Oculus/OculusVR/Source/OculusEditor/Private/OculusToolWidget.cpp
+++ b/Engine/Plugins/Runtime/Oculus/OculusVR/Source/OculusEditor/Private/OculusToolWidget.cpp
@@ -285,7 +285,6 @@ void SOculusToolWidget::RebuildLayout()
 	AddSimpleSetting(box, SimpleSettings.Find(FName("MobileShaderStaticAndCSMShadowReceivers")));
 	AddSimpleSetting(box, SimpleSettings.Find(FName("MobileShaderAllowDistanceFieldShadows")));
 	AddSimpleSetting(box, SimpleSettings.Find(FName("MobileShaderAllowMovableDirectionalLights")));
-	AddSimpleSetting(box, SimpleSettings.Find(FName("MobileNumDynamicPointLights")));
 	AddSimpleSetting(box, SimpleSettings.Find(FName("MobileMovableSpotlights")));
 
 	box = NewCategory(scroller, FText::GetEmpty());
@@ -546,30 +545,6 @@ void SOculusToolWidget::Construct(const FArguments& InArgs)
 		&SOculusToolWidget::MobileShaderAllowMovableDirectionalLightsDisable }
 	);
 
-	SimpleSettings.Add(FName("MobileNumDynamicPointLights"), {
-		FName("MobileNumDynamicPointLights"),
-		LOCTEXT("MobileNumDynamicPointLightsDescription", "Your project does not contain any movable lights. Max Movable Spotlights / Point Lights can be set to 0 to reduce shader permutations."),
-		&SOculusToolWidget::MobileNumDynamicPointLightsVisibility,
-		TArray<SimpleSettingAction>(),
-		(int)SupportFlags::SupportMobile
-		});
-	SimpleSettings.Find(FName("MobileNumDynamicPointLights"))->actions.Add(
-		{ LOCTEXT("MobileNumDynamicPointLightsButton", "Set Max Movable Spotlights / Point Lights to 0"),
-		&SOculusToolWidget::MobileNumDynamicPointLightsFix }
-	);
-
-	SimpleSettings.Add(FName("MobileMovableSpotlights"), {
-		FName("MobileMovableSpotlights"),
-		LOCTEXT("MobileMovableSpotlightsDescription", "Your project does not contain any movable lights. Support Movable Spotlights can be disabled to reduce shader permutations."),
-		&SOculusToolWidget::MobileMovableSpotlightsDisableVisibility,
-		TArray<SimpleSettingAction>(),
-		(int)SupportFlags::SupportMobile
-		});
-	SimpleSettings.Find(FName("MobileMovableSpotlights"))->actions.Add(
-		{ LOCTEXT("MobileMovableSpotlightsButton", "Disable Support Movable Spotlights"),
-		&SOculusToolWidget::MobileMovableSpotlightsDisable }
-	);
-
 	auto scroller = SNew(SScrollBox);
 	ScrollingContainer = scroller;
 	RebuildLayout();
@@ -948,64 +923,6 @@ EVisibility SOculusToolWidget::MobileShaderAllowMovableDirectionalLightsVisibili
 	return EVisibility::Visible;
 }
 
-FReply SOculusToolWidget::MobileNumDynamicPointLightsFix(bool text)
-{
-	URendererSettings* Settings = GetMutableDefault<URendererSettings>();
-	Settings->MobileNumDynamicPointLights = 0;
-	Settings->UpdateSinglePropertyInConfigFile(Settings->GetClass()->FindPropertyByName(GET_MEMBER_NAME_CHECKED(URendererSettings, MobileNumDynamicPointLights)), Settings->GetDefaultConfigFilename());
-	SuggestRestart();
-	return FReply::Handled();
-}
-
-EVisibility SOculusToolWidget::MobileNumDynamicPointLightsVisibility(FName tag) const
-{
-	URendererSettings* Settings = GetMutableDefault<URendererSettings>();
-	if (Settings->MobileNumDynamicPointLights == 0)
-	{
-		return EVisibility::Collapsed;
-	}
-
-	for (const auto& kvp : DynamicLights)
-	{
-		AActor* owner = kvp.Value->GetOwner();
-		if (owner != NULL && owner->IsRootComponentMovable())
-		{
-			return EVisibility::Collapsed;
-		}
-	}
-
-	return EVisibility::Visible;
-}
-
-FReply SOculusToolWidget::MobileMovableSpotlightsDisable(bool text)
-{
-	URendererSettings* Settings = GetMutableDefault<URendererSettings>();
-	Settings->bMobileAllowMovableSpotlights = false;
-	Settings->UpdateSinglePropertyInConfigFile(Settings->GetClass()->FindPropertyByName(GET_MEMBER_NAME_CHECKED(URendererSettings, bMobileAllowMovableSpotlights)), Settings->GetDefaultConfigFilename());
-	SuggestRestart();
-	return FReply::Handled();
-}
-
-EVisibility SOculusToolWidget::MobileMovableSpotlightsDisableVisibility(FName tag) const
-{
-	URendererSettings* Settings = GetMutableDefault<URendererSettings>();
-	if (!Settings->bMobileAllowMovableSpotlights)
-	{
-		return EVisibility::Collapsed;
-	}
-
-	for (const auto& kvp : DynamicLights)
-	{
-		AActor* owner = kvp.Value->GetOwner();
-		if (owner != NULL && owner->IsRootComponentMovable())
-		{
-			return EVisibility::Collapsed;
-		}
-	}
-
-	return EVisibility::Visible;
-}
-
 void SOculusToolWidget::OnShowButtonChanged(ECheckBoxState NewState)
 {
 	GConfig->SetBool(TEXT("/Script/OculusEditor.OculusEditorSettings"), TEXT("bAddMenuOption"), NewState == ECheckBoxState::Checked ? true : false, FString::Printf(TEXT("%sDefaultEditor.ini"), *FPaths::SourceConfigDir()));
diff --git a/Engine/Plugins/Runtime/Oculus/OculusVR/Source/OculusEditor/Private/OculusToolWidget.h b/Engine/Plugins/Runtime/Oculus/OculusVR/Source/OculusEditor/Private/OculusToolWidget.h
--- a/Engine/Plugins/Runtime/Oculus/OculusVR/Source/OculusEditor/Private/OculusToolWidget.h
+++ b/Engine/Plugins/Runtime/Oculus/OculusVR/Source/OculusEditor/Private/OculusToolWidget.h
@@ -126,12 +126,6 @@ protected:
 	FReply MobileShaderAllowMovableDirectionalLightsDisable(bool text);
 	EVisibility MobileShaderAllowMovableDirectionalLightsVisibility(FName tag) const;
 
-	FReply MobileNumDynamicPointLightsFix(bool text);
-	EVisibility MobileNumDynamicPointLightsVisibility(FName tag) const;
-
-	FReply MobileMovableSpotlightsDisable(bool text);
-	EVisibility MobileMovableSpotlightsDisableVisibility(FName tag) const;
-
 	FReply SelectLight(FString lightName);
 	FReply IgnoreLight(FString lightName);
 	
diff --git a/Engine/Shaders/Private/MobileBasePassPixelShader.usf b/Engine/Shaders/Private/MobileBasePassPixelShader.usf
--- a/Engine/Shaders/Private/MobileBasePassPixelShader.usf
+++ b/Engine/Shaders/Private/MobileBasePassPixelShader.usf
@@ -10,6 +10,9 @@
 #define MobileSceneTextures MobileBasePass.SceneTextures
 #define EyeAdaptationStruct MobileBasePass
 
+#define ForwardLightData MobileBasePass.Forward
+#define ReflectionStruct MobileBasePass.ReflectionsParameters
+
 #ifndef MOBILE_QL_FORCE_DISABLE_PREINTEGRATEDGF
 #define MOBILE_QL_FORCE_DISABLE_PREINTEGRATEDGF 0
 #endif
@@ -23,18 +26,12 @@
 #endif
 
 #define APPLY_AO 0
-#define APPLY_DISTANCE_FIELD 0
 
 #if (MATERIALBLENDING_MASKED || MATERIALBLENDING_SOLID) 
 	#if ENABLE_AMBIENT_OCCLUSION && !MATERIAL_SHADINGMODEL_UNLIT
 		#undef APPLY_AO
 		#define APPLY_AO 1
 	#endif
-
-	#if ENABLE_DISTANCE_FIELD
-		#undef APPLY_DISTANCE_FIELD
-		#define APPLY_DISTANCE_FIELD 1
-	#endif
 #endif
 
 #if APPLY_AO
@@ -47,30 +44,8 @@
 	#define SIMPLE_SINGLE_LAYER_WATER 1
 #endif
 
-#if MAX_DYNAMIC_POINT_LIGHTS > 0
-#if VARIABLE_NUM_DYNAMIC_POINT_LIGHTS
-int NumDynamicPointLights;
-#else
-#define NumDynamicPointLights NUM_DYNAMIC_POINT_LIGHTS
-#endif
-
-#define DynamicSpotLightShadowTexture		MobileDirectionalLight.DirectionalLightShadowTexture
-#define DynamicSpotLightShadowSampler		MobileDirectionalLight.DirectionalLightShadowSampler
-#define DynamicSpotLightShadowBufferSize	MobileDirectionalLight.DirectionalLightShadowSize
-
-#define LightShadowType_PointLight	1
-#define LightShadowType_SpotLight	2
-#define LightShadowType_Shadow		4
-#define ValidLightType (LightShadowType_PointLight | LightShadowType_SpotLight)
-
-#endif
-
-#ifndef PROJECT_MOBILE_ENABLE_MOVABLE_SPOTLIGHTS
-#define PROJECT_MOBILE_ENABLE_MOVABLE_SPOTLIGHTS 0
-#endif
-
-#ifndef PROJECT_MOBILE_ENABLE_MOVABLE_SPOTLIGHTS_SHADOW
-#define PROJECT_MOBILE_ENABLE_MOVABLE_SPOTLIGHTS_SHADOW 0
+#ifndef PROJECT_MOBILE_ENABLE_MOVABLE_SPOTLIGHT_SHADOWS
+#define PROJECT_MOBILE_ENABLE_MOVABLE_SPOTLIGHT_SHADOWS 0
 #endif
 
 #ifndef MOBILE_QL_FORCE_FULLY_ROUGH
@@ -79,18 +54,14 @@ int NumDynamicPointLights;
 #ifndef MOBILE_QL_FORCE_NONMETAL
 #define MOBILE_QL_FORCE_NONMETAL 0
 #endif
-#ifndef MOBILE_QL_FORCE_LQ_REFLECTIONS
-#define MOBILE_QL_FORCE_LQ_REFLECTIONS 0
-#endif
 #ifndef MOBILE_QL_DISABLE_MATERIAL_NORMAL
 #define MOBILE_QL_DISABLE_MATERIAL_NORMAL 0
 #endif
 
 #define FULLY_ROUGH (MATERIAL_FULLY_ROUGH || MOBILE_QL_FORCE_FULLY_ROUGH)
 #define NONMETAL (MATERIAL_NONMETAL || MOBILE_QL_FORCE_NONMETAL)
-#define HQ_REFLECTIONS (MATERIAL_HQ_FORWARD_REFLECTION_CAPTURES && !MOBILE_QL_FORCE_LQ_REFLECTIONS)
 #define FORCE_VERTEX_NORMAL (MOBILE_QL_DISABLE_MATERIAL_NORMAL)
-#define SUPPORT_SPOTLIGHTS_SHADOW (MATERIALBLENDING_SOLID || MATERIALBLENDING_MASKED) && PROJECT_MOBILE_ENABLE_MOVABLE_SPOTLIGHTS_SHADOW
+#define SUPPORT_SPOTLIGHTS_SHADOW (MATERIALBLENDING_SOLID || MATERIALBLENDING_MASKED) && PROJECT_MOBILE_ENABLE_MOVABLE_SPOTLIGHT_SHADOWS
 #define DEFERRED_SHADING_PATH (MOBILE_DEFERRED_SHADING && ((MATERIALBLENDING_SOLID || MATERIALBLENDING_MASKED) && !MATERIAL_SHADINGMODEL_SINGLELAYERWATER))
 
 // SceneDepthAux always enabled for IOS, on other platforms only enabled when in forward and MobileHDR=true
@@ -103,284 +74,18 @@ int NumDynamicPointLights;
 #endif
 #define TRANSLUCENCY_NON_DIRECTIONAL (MATERIALBLENDING_ANY_TRANSLUCENT && (TRANSLUCENCY_LIGHTING_VOLUMETRIC_NONDIRECTIONAL || TRANSLUCENCY_LIGHTING_VOLUMETRIC_PERVERTEX_NONDIRECTIONAL) && !MATERIAL_SHADINGMODEL_SINGLELAYERWATER)
 
-#define ALLOW_CUBE_REFLECTIONS HQ_REFLECTIONS
-
 #include "SHCommon.ush"
 #include "/Engine/Generated/Material.ush"
 #include "MobileBasePassCommon.ush"
 #include "/Engine/Generated/VertexFactory.ush"
-#include "ReflectionEnvironmentShared.ush"
 #include "LightmapCommon.ush"  
-#include "MobileShadingModels.ush"
-#include "DynamicLightingCommon.ush"
-#include "PlanarReflectionShared.ush"
+#include "MobileLightingCommon.ush"
 #include "ShadingModelsMaterial.ush"
 
 #if MATERIAL_SHADINGMODEL_SINGLELAYERWATER
 	#include "SingleLayerWaterShading.ush"
 #endif
 
-/** Prenormalized capture of the scene that's closest to the object being rendered. */
-#if HQ_REFLECTIONS
-	#define MAX_HQ_REFLECTIONS 3
-	TextureCube ReflectionCubemap0;
-	SamplerState ReflectionCubemapSampler0;
-	TextureCube ReflectionCubemap1;
-	SamplerState ReflectionCubemapSampler1;
-	TextureCube ReflectionCubemap2;
-	SamplerState ReflectionCubemapSampler2;
-	// x,y,z - inverted average brightness for 0, 1, 2; w - sky cube texture max mips.
-	float4 ReflectionAverageBrigtness;
-	float4 ReflectanceMaxValueRGBM;
-	float4 ReflectionPositionsAndRadii[MAX_HQ_REFLECTIONS];
-	float4 ReflectionTilePositions[MAX_HQ_REFLECTIONS];
-		#if ALLOW_CUBE_REFLECTIONS
-		float4x4 CaptureBoxTransformArray[MAX_HQ_REFLECTIONS];
-		float4 CaptureBoxScalesArray[MAX_HQ_REFLECTIONS];
-		#endif
-#endif
-
-#ifndef MOBILE_USE_CSM_BRANCH
-#define MOBILE_USE_CSM_BRANCH 0
-#endif
-
-#if MOBILE_USE_CSM_BRANCH
-	uint UseCSM;
-#endif
-
-bool IsCSMEnabled()
-{
-#if MOBILE_USE_CSM_BRANCH
-	return UseCSM != 0u;
-#else
-	return true;
-#endif
-}
-
-// Half precision version for mobile's high quality reflection
-half MobileComputeMixingWeight(half IndirectIrradiance, half AverageBrightness, half Roughness)
-{
-	// Mirror surfaces should have no mixing, so they match reflections from other sources (SSR, planar reflections)
-	half MixingAlpha = smoothstep(0, 1, saturate(Roughness * View.ReflectionEnvironmentRoughnessMixingScaleBiasAndLargestWeight.x + View.ReflectionEnvironmentRoughnessMixingScaleBiasAndLargestWeight.y));
-
-	// We have high frequency directional data but low frequency spatial data in the envmap.
-	// We have high frequency spatial data but low frequency directional data in the lightmap.
-	// So, we combine the two for the best of both. This is done by removing the low spatial frequencies from the envmap and replacing them with the lightmap data.
-	// This is only done with luma so as to not get odd color shifting.
-	half MixingWeight = IndirectIrradiance / max(AverageBrightness, .0001f);
-
-	MixingWeight = min(MixingWeight, View.ReflectionEnvironmentRoughnessMixingScaleBiasAndLargestWeight.z);
-
-	return lerp(1.0f, MixingWeight, MixingAlpha);
-}
-
-half3 GetLookupVectorForBoxCaptureMobile(half3 ReflectionVector, half3 WorldPosition, half4 BoxCapturePositionAndRadius, half4x4 RelativeWorldToBox, half4 BoxScales, half3 LocalCaptureOffset, out half DistanceAlpha)
-{
-	// Transform the ray into the local space of the box, where it is an AABB with mins at -1 and maxs at 1
-	half3 LocalRayStart = mul(half4(WorldPosition - BoxCapturePositionAndRadius.xyz, 1), RelativeWorldToBox).xyz;
-	half3 LocalRayDirection = mul(half4(ReflectionVector, 0), RelativeWorldToBox).xyz;
-
-	half3 InvRayDir = rcp(LocalRayDirection);
-
-	//find the ray intersection with each of the 3 planes defined by the minimum extrema.
-	half3 FirstPlaneIntersections = -InvRayDir - LocalRayStart * InvRayDir;
-	//find the ray intersection with each of the 3 planes defined by the maximum extrema.
-	half3 SecondPlaneIntersections = InvRayDir - LocalRayStart * InvRayDir;
-	//get the furthest of these intersections along the ray
-	half3 FurthestPlaneIntersections = max(FirstPlaneIntersections, SecondPlaneIntersections);
-
-	//clamp the intersections to be between RayOrigin and RayEnd on the ray
-	half Intersection = min(FurthestPlaneIntersections.x, min(FurthestPlaneIntersections.y, FurthestPlaneIntersections.z));
-
-	// Compute the reprojected vector
-	half3 IntersectPosition = WorldPosition + Intersection * ReflectionVector;
-	half3 ProjectedCaptureVector = IntersectPosition - (BoxCapturePositionAndRadius.xyz + LocalCaptureOffset);
-
-	// Compute the distance from the receiving pixel to the box for masking
-	// Apply local to world scale to take scale into account without transforming back to world space
-	// Shrink the box by the transition distance (BoxScales.w) so that the fade happens inside the box influence area	
-	half BoxDistance = ComputeDistanceFromBoxToPoint(-(BoxScales.xyz - .5f * BoxScales.w), BoxScales.xyz - .5f * BoxScales.w, LocalRayStart * BoxScales.xyz);
-
-	// Setup a fade based on receiver distance to the box, hides the box influence shape
-	DistanceAlpha = 1.0 - smoothstep(0, .7f * BoxScales.w, BoxDistance);
-
-	return ProjectedCaptureVector;
-}
-
-half3 GetLookupVectorForSphereCaptureMobile(half3 ReflectionVector, half3 WorldPosition, half4 SphereCapturePositionAndRadius, half NormalizedDistanceToCapture, half3 LocalCaptureOffset, inout half DistanceAlpha)
-{
-	half3 ProjectedCaptureVector = ReflectionVector;
-	half ProjectionSphereRadius = SphereCapturePositionAndRadius.w;
-	half SphereRadiusSquared = ProjectionSphereRadius * ProjectionSphereRadius;
-
-	half3 LocalPosition = WorldPosition - SphereCapturePositionAndRadius.xyz;
-	half LocalPositionSqr = dot(LocalPosition, LocalPosition);
-
-	// Find the intersection between the ray along the reflection vector and the capture's sphere
-	half3 QuadraticCoef;
-	QuadraticCoef.x = 1;
-	QuadraticCoef.y = dot(ReflectionVector, LocalPosition);
-	QuadraticCoef.z = LocalPositionSqr - SphereRadiusSquared;
-
-	half Determinant = QuadraticCoef.y * QuadraticCoef.y - QuadraticCoef.z;
-
-	// Only continue if the ray intersects the sphere
-	FLATTEN
-		if (Determinant >= 0)
-		{
-			half FarIntersection = sqrt(Determinant) - QuadraticCoef.y;
-
-			half3 LocalIntersectionPosition = LocalPosition + FarIntersection * ReflectionVector;
-			ProjectedCaptureVector = LocalIntersectionPosition - LocalCaptureOffset;
-			// Note: some compilers don't handle smoothstep min > max (this was 1, .6)
-			//DistanceAlpha = 1.0 - smoothstep(.6, 1, NormalizedDistanceToCapture);
-
-			half x = saturate(2.5 * NormalizedDistanceToCapture - 1.5);
-			DistanceAlpha = 1 - x * x*(3 - 2 * x);
-		}
-	return ProjectedCaptureVector;
-}
-
-half3 GetMobileSkyLightReflection(half3 ReflectionVector, half Roughness, half CubemapMaxMip)
-{
-	half AbsoluteSpecularMip = ComputeReflectionCaptureMipFromRoughness(Roughness, CubemapMaxMip);
-	half3 Reflection = MobileReflectionCapture.Texture.SampleLevel(MobileReflectionCapture.TextureSampler, ReflectionVector, AbsoluteSpecularMip).rgb;
-	return Reflection * ResolvedView.SkyLightColor.rgb;
-}
-
-void GatherSpecularIBL(FMaterialPixelParameters MaterialParameters
-	, TextureCube ReflectionCube
-	, SamplerState ReflectionSampler
-	, half Roughness
-	, half IndirectIrradiance
-	, half MaxValue
-	, inout half4 ImageBasedReflections
-#if HQ_REFLECTIONS
-	, int ReflectionIndex
-	, inout half3 ExtraIndirectSpecular
-	, inout half2 CompositedAverageBrightness
-#endif
-)
-{
-	half3 ProjectedCaptureVector = MaterialParameters.ReflectionVector;
-
-#if HQ_REFLECTIONS
-
-	float CaptureRadius = ReflectionPositionsAndRadii[ReflectionIndex].w;
-	if (CaptureRadius == 0.0f)
-	{
-		return;
-	}
-	
-	half DistanceAlpha = 0.0;
-
-	FLWCVector3 CaptureWorldPosition = MakeLWCVector3(ReflectionTilePositions[ReflectionIndex].xyz, ReflectionPositionsAndRadii[ReflectionIndex].xyz);
-	float3 CaptureTranslatedWorldPosition = LWCToFloat(LWCAdd(CaptureWorldPosition, ResolvedView.PreViewTranslation));
-
-	half3 CaptureVector = MaterialParameters.WorldPosition_CamRelative - CaptureTranslatedWorldPosition;
-	half CaptureVectorLength = sqrt(dot(CaptureVector, CaptureVector));
-	half NormalizedDistanceToCapture = saturate(CaptureVectorLength / CaptureRadius);
-	ProjectedCaptureVector = GetOffSpecularPeakReflectionDir(MaterialParameters.WorldNormal, MaterialParameters.ReflectionVector, Roughness);
-	
-	if (CaptureVectorLength < CaptureRadius || CaptureRadius <= 0)
-	{
-#if ALLOW_CUBE_REFLECTIONS
-		if (CaptureBoxScalesArray[ReflectionIndex].w > 0)
-		{
-			ProjectedCaptureVector = GetLookupVectorForBoxCaptureMobile(MaterialParameters.ReflectionVector, MaterialParameters.WorldPosition_CamRelative, float4(CaptureTranslatedWorldPosition, CaptureRadius), CaptureBoxTransformArray[ReflectionIndex], CaptureBoxScalesArray[ReflectionIndex], half3(0.0f, 0.0f, 0.0f), DistanceAlpha);
-		}
-		else
-#endif
-		{
-			ProjectedCaptureVector = GetLookupVectorForSphereCaptureMobile(MaterialParameters.ReflectionVector, MaterialParameters.WorldPosition_CamRelative, float4(CaptureTranslatedWorldPosition, CaptureRadius), NormalizedDistanceToCapture, half3(0.0f, 0.0f, 0.0f), DistanceAlpha);
-		}
-	}
-	half UsingSkyReflection = CaptureRadius <= 0.0f;
-
-#else
-	half UsingSkyReflection = MobileReflectionCapture.Params.y > 0.0f;
-#endif
-
-	// Fetch from cubemap and convert to linear HDR
-	half3 SpecularIBL = 0.0f;
-
-
-	if (UsingSkyReflection)
-	{
-		// Apply sky colour if the reflection map is the sky.
-#if HQ_REFLECTIONS
-		half AbsoluteSpecularMip = ComputeReflectionCaptureMipFromRoughness(Roughness, ReflectionAverageBrigtness.w);
-		half4 SpecularIBLSample = ReflectionCube.SampleLevel(ReflectionSampler, ProjectedCaptureVector, AbsoluteSpecularMip);
-		ExtraIndirectSpecular = SpecularIBLSample.rgb * ResolvedView.SkyLightColor.rgb;
-#else
-		half AbsoluteSpecularMip = ComputeReflectionCaptureMipFromRoughness(Roughness, MobileReflectionCapture.Params.y);
-		half4 SpecularIBLSample = ReflectionCube.SampleLevel(ReflectionSampler, ProjectedCaptureVector, AbsoluteSpecularMip);
-		ImageBasedReflections.rgb = SpecularIBLSample.rgb * ResolvedView.SkyLightColor.rgb * ImageBasedReflections.a;
-#endif
-	}
-	else
-	{
-		half AbsoluteSpecularMip = ComputeReflectionCaptureMipFromRoughness(Roughness, ResolvedView.ReflectionCubemapMaxMip);
-		half4 SpecularIBLSample = ReflectionCube.SampleLevel(ReflectionSampler, ProjectedCaptureVector, AbsoluteSpecularMip);
-
-		SpecularIBL = RGBMDecode(SpecularIBLSample, MaxValue);
-		SpecularIBL = SpecularIBL * SpecularIBL;
-#if HQ_REFLECTIONS
-		SpecularIBL *= DistanceAlpha;
-
-		// Under operator (back to front)
-		ImageBasedReflections.rgb += SpecularIBL * ImageBasedReflections.a;
-		ImageBasedReflections.a *= 1 - DistanceAlpha;
-
-		half AverageBrightness = ReflectionAverageBrigtness[ReflectionIndex];
-		CompositedAverageBrightness.x += AverageBrightness * DistanceAlpha * CompositedAverageBrightness.y;
-		CompositedAverageBrightness.y *= 1 - DistanceAlpha;
-#else
-		half ReflectionCaptureBrightness = MobileReflectionCapture.Params.w;
-		ImageBasedReflections.rgb = SpecularIBL * ImageBasedReflections.a * ReflectionCaptureBrightness;
-#if ALLOW_STATIC_LIGHTING
-		//To keep ImageBasedReflectionLighting conherence with PC, use ComputeMixingWeight instead of InvReflectionAverageBrightness to calulate the IBL constribution
-		ImageBasedReflections.rgb *= MobileComputeMixingWeight(IndirectIrradiance, MobileReflectionCapture.Params.x, Roughness);
-#endif
-		ImageBasedReflections.rgb *= View.IndirectLightingColorScale;
-#endif
-	}
-}
-
-#if HQ_REFLECTIONS
-void BlendReflectionCaptures(FMaterialPixelParameters MaterialParameters, half Roughness, half IndirectIrradiance, inout half4 ImageBasedReflections)
-{
-	half2 CompositedAverageBrightness = half2(0.0f, 1.0f);
-	half3 ExtraIndirectSpecular = 0.0f;
-
-	GatherSpecularIBL(MaterialParameters, ReflectionCubemap0, ReflectionCubemapSampler0, Roughness, IndirectIrradiance, ReflectanceMaxValueRGBM.x, ImageBasedReflections, 0, ExtraIndirectSpecular, CompositedAverageBrightness);
-	GatherSpecularIBL(MaterialParameters, ReflectionCubemap1, ReflectionCubemapSampler1, Roughness, IndirectIrradiance, ReflectanceMaxValueRGBM.y, ImageBasedReflections, 1, ExtraIndirectSpecular, CompositedAverageBrightness);
-	GatherSpecularIBL(MaterialParameters, ReflectionCubemap2, ReflectionCubemapSampler2, Roughness, IndirectIrradiance, ReflectanceMaxValueRGBM.z, ImageBasedReflections, 2, ExtraIndirectSpecular, CompositedAverageBrightness);
-
-	ImageBasedReflections.rgb *= View.IndirectLightingColorScale;
-	CompositedAverageBrightness.x *= Luminance(View.IndirectLightingColorScale);
-
-#if ALLOW_STATIC_LIGHTING
-	ImageBasedReflections.rgb *= ComputeMixingWeight(IndirectIrradiance, CompositedAverageBrightness.x, Roughness);
-#endif
-
-	ImageBasedReflections.rgb += ImageBasedReflections.a * ExtraIndirectSpecular;
-}
-#endif // HQ_REFLECTIONS
-
-half3 GetImageBasedReflectionLighting(FMaterialPixelParameters MaterialParameters, half Roughness, half IndirectIrradiance, half CompositeAlpha)
-{
-	half4 ImageBasedReflections = half4(0, 0, 0, CompositeAlpha);
-
-#if HQ_REFLECTIONS
-	BlendReflectionCaptures(MaterialParameters, Roughness, IndirectIrradiance, ImageBasedReflections);
-#else
-	GatherSpecularIBL(MaterialParameters, MobileReflectionCapture.Texture, MobileReflectionCapture.TextureSampler, Roughness, IndirectIrradiance, MobileReflectionCapture.Params.z, ImageBasedReflections);
-#endif
-
-	return ImageBasedReflections.rgb;
-}
-
 half3 FrameBufferBlendOp(half4 Source)
 {
 	half4 Dest = half4 (0,0,0,0);
@@ -411,109 +116,6 @@ void ApplyPixelDepthOffsetForMobileBasePass(inout FMaterialPixelParameters Mater
     float PixelDepthOffset = ApplyPixelDepthOffsetToMaterialParameters(MaterialParameters, PixelMaterialInputs, OutDepth);
 }
 
-half ApplyPrecomputedShadowMask(half ShadowMap, half Shadow)
-{
-#if ALLOW_STATIC_LIGHTING
-	return min(ShadowMap, Shadow);
-#else
-	return ShadowMap;
-#endif
-}
-
-#if MAX_DYNAMIC_POINT_LIGHTS > 0
-void AccumulateLightingOfDynamicPointLight(
-FMaterialPixelParameters MaterialParameters, 
-FMobileShadingModelContext ShadingModelContext,
-FGBufferData GBuffer,
-float4 LightPositionAndInvRadius,
-float4 LightTilePosition,
-float4 LightColorAndFalloffExponent, 
-float4 SpotLightDirectionAndSpecularScale, 
-float4 SpotLightAnglesAndSoftTransitionScaleAndLightShadowType, 
-#if SUPPORT_SPOTLIGHTS_SHADOW
-FPCFSamplerSettings Settings,
-float4 SpotLightShadowSharpenAndShadowFadeFraction,
-float4 SpotLightShadowmapMinMax,
-float4x4 SpotLightShadowWorldToShadowMatrix,
-#endif
-inout FMobileLightAccumulator MobileLightAccumulator)
-{
-	uint LightShadowType = SpotLightAnglesAndSoftTransitionScaleAndLightShadowType.w;
-
-	float FadedShadow = 1.0f;
-	
-	FLWCVector3 LightWorldPosition = MakeLWCVector3(LightTilePosition.xyz, LightPositionAndInvRadius.xyz);
-	FLWCVector3 LightTranslatedWorldPosition = LWCAdd(LightWorldPosition, ResolvedView.PreViewTranslation);
-	float3 LocalPosition = MaterialParameters.WorldPosition_CamRelative - LWCToFloat(LightTranslatedWorldPosition);
-	float3 ToLight = -LocalPosition;
-	float DistanceSqr = dot(ToLight, ToLight);
-	half3 L = ToLight * rsqrt(DistanceSqr);
-	float3 PointH = normalize(float3(MaterialParameters.CameraVector + L));
-
-	half PointNoL = max(0, dot(MaterialParameters.WorldNormal, L));
-	float PointNoH = max(0, dot(MaterialParameters.WorldNormal, PointH));
-	half PointVoH = max(0, dot(MaterialParameters.CameraVector, PointH));
-
-#if SUPPORT_SPOTLIGHTS_SHADOW
-		if ((LightShadowType & LightShadowType_Shadow) == LightShadowType_Shadow)
-		{
-
-			float4 HomogeneousShadowPosition = mul(float4(LocalPosition, 1), SpotLightShadowWorldToShadowMatrix);
-			float2 ShadowUVs = HomogeneousShadowPosition.xy / HomogeneousShadowPosition.w;
-			if (all(ShadowUVs >= SpotLightShadowmapMinMax.xy && ShadowUVs <= SpotLightShadowmapMinMax.zw))
-			{
-				// Clamp pixel depth in light space for shadowing opaque, because areas of the shadow depth buffer that weren't rendered to will have been cleared to 1
-				// We want to force the shadow comparison to result in 'unshadowed' in that case, regardless of whether the pixel being shaded is in front or behind that plane
-				// Invert ShadowZ as the shadow space has been changed (but not yet the filtering code)
-				float ShadowZ = 1.0f - HomogeneousShadowPosition.z;
-				float LightSpacePixelDepthForOpaque = min(ShadowZ, 0.99999f);
-				Settings.SceneDepth = LightSpacePixelDepthForOpaque;
-				Settings.TransitionScale = SpotLightAnglesAndSoftTransitionScaleAndLightShadowType.z * lerp(SpotLightShadowSharpenAndShadowFadeFraction.z, 1.0, PointNoL);
-
-				half Shadow = MobileShadowPCF(ShadowUVs, Settings);
-
-				Shadow = saturate((Shadow - 0.5) * SpotLightShadowSharpenAndShadowFadeFraction.x + 0.5);
-
-				FadedShadow = lerp(1.0f, Square(Shadow), SpotLightShadowSharpenAndShadowFadeFraction.y);
-			}
-		}
-#endif
-
-	if ((LightShadowType & ValidLightType) != 0)
-	{
-		#if TRANSLUCENCY_NON_DIRECTIONAL
-			PointNoL = 1.f;
-		#endif
-
-		float Attenuation;
-
-		if (LightColorAndFalloffExponent.w == 0)
-		{
-			// Sphere falloff (technically just 1/d2 but this avoids inf)
-			Attenuation = 1 / (DistanceSqr + 1);
-
-			float LightRadiusMask = Square(saturate(1 - Square(DistanceSqr * (LightPositionAndInvRadius.w * LightPositionAndInvRadius.w))));
-			Attenuation *= LightRadiusMask;
-		}
-		else
-		{
-			Attenuation = RadialAttenuation(ToLight * LightPositionAndInvRadius.w, LightColorAndFalloffExponent.w);
-		}
-
-#if PROJECT_MOBILE_ENABLE_MOVABLE_SPOTLIGHTS
-		if ((LightShadowType & LightShadowType_SpotLight) == LightShadowType_SpotLight)
-		{
-			Attenuation *= SpotAttenuation(L, -SpotLightDirectionAndSpecularScale.xyz, SpotLightAnglesAndSoftTransitionScaleAndLightShadowType.xy);
-		}
-#endif
-
-		FMobileDirectLighting Lighting = MobileIntegrateBxDF(ShadingModelContext, GBuffer, PointNoL, PointNoH, PointVoH, MaterialParameters.CameraVector, L, FadedShadow);
-
-		MobileLightAccumulator_Add(MobileLightAccumulator, Lighting.Diffuse, Lighting.Specular * SpotLightDirectionAndSpecularScale.w, (Attenuation)*LightColorAndFalloffExponent.rgb * (1.0 / PI));
-	}
-}
-#endif
-
 // Force early depth_stencil for non-masked material that use VT feedback
 #if (NUM_VIRTUALTEXTURE_SAMPLES || LIGHTMAP_VT_ENABLED) && !(MATERIALBLENDING_MASKED || USE_DITHERED_LOD_TRANSITION || OUTPUT_PIXEL_DEPTH_OFFSET)
 	#define PIXELSHADER_EARLYDEPTHSTENCIL EARLYDEPTHSTENCIL	
@@ -798,196 +400,44 @@ void Main(
 	EncodeGBuffer(GBuffer, OutGBufferA, OutGBufferB, OutGBufferC, OutGBufferD, OutGBufferE, OutGBufferF, OutGBufferVelocity);
 #else
 
-#if ENABLE_SKY_LIGHT
-	half3 SkyDiffuseLighting = GetSkySHDiffuseSimple(MaterialParameters.WorldNormal);
-	half3 DiffuseLookup = SkyDiffuseLighting * ResolvedView.SkyLightColor.rgb;
-	IndirectIrradiance += Luminance(DiffuseLookup);
-#endif
-
-	half3 L = MobileDirectionalLight.DirectionalLightDirectionAndShadowTransition.xyz;
-	half NoL = max(0, dot(MaterialParameters.WorldNormal, L));
-	float3 H = normalize(float3(MaterialParameters.CameraVector + L));
-	half NoH = max(0, dot(MaterialParameters.WorldNormal, H));
-	half VoH = max(0, dot(MaterialParameters.CameraVector, H));
-
-	float  ShadowPositionZ = 0;
-#if DIRECTIONAL_LIGHT_CSM && !MATERIAL_SHADINGMODEL_SINGLELAYERWATER && !APPLY_DISTANCE_FIELD
-	// Cascaded Shadow Map
-	if (IsCSMEnabled())
-	{
-		half ShadowMap = MobileDirectionalLightCSM(MaterialParameters.ScreenPosition.xy, MaterialParameters.ScreenPosition.w, NoL, ShadowPositionZ);
-		Shadow = ApplyPrecomputedShadowMask(ShadowMap, Shadow);
-	}
-#endif /* DIRECTIONAL_LIGHT_CSM */
-
-#if APPLY_DISTANCE_FIELD
-	if (ShadowPositionZ == 0)
-	{
-		half ShadowMap = Texture2DSample(MobileBasePass.ScreenSpaceShadowMaskTexture, MobileBasePass.ScreenSpaceShadowMaskSampler, SvPositionToBufferUV(SvPosition)).x;
-		ShadowMap = DecodeLightAttenuation(ShadowMap);
-		Shadow = ApplyPrecomputedShadowMask(ShadowMap, Shadow);
-	}
-#endif
-
-#if TRANSLUCENCY_NON_DIRECTIONAL
-	NoL = 1.f;
-	ShadingModelContext.SpecularColor = 0.f;
-#endif
+	half4 DynamicShadowFactors = 1.0f;
+	half NoL = 0.0f;
 
-	// Direct Lighting (Directional light) + IBL
-	FMobileDirectLighting Lighting = MobileIntegrateBxDF(ShadingModelContext, GBuffer, NoL, NoH, VoH, MaterialParameters.CameraVector, L, Shadow);
-	// MobileDirectionalLight.DirectionalLightDistanceFadeMADAndSpecularScale.z saves SpecularScale for direction light.
-	MobileLightAccumulator_Add(MobileLightAccumulator, Lighting.Diffuse, Lighting.Specular * MobileDirectionalLight.DirectionalLightDistanceFadeMADAndSpecularScale.z, MobileDirectionalLight.DirectionalLightColor.rgb);
-
-#if	!(MATERIAL_SINGLE_SHADINGMODEL && MATERIAL_SHADINGMODEL_HAIR)
-	// Environment map has been prenormalized, scale by lightmap luminance
-	half3 SpecularIBL = GetImageBasedReflectionLighting(MaterialParameters, GBuffer.Roughness, IndirectIrradiance, 1.0f);
-	half3 SpecularIBLLighting = (half3)0.0f;
-	#if MATERIAL_PLANAR_FORWARD_REFLECTIONS
-		BRANCH
-		if (abs(dot(PlanarReflectionStruct.ReflectionPlane.xyz, 1)) > .0001f)
-		{
-			half4 PlanarReflection = GetPlanarReflection(MaterialParameters.WorldPosition_CamRelative, MaterialParameters.WorldNormal, GBuffer.Roughness);
-			// Planar reflections win over reflection environment
-			SpecularIBL = lerp(SpecularIBL, PlanarReflection.rgb, PlanarReflection.a);
-		}
-	#endif
-	#if MATERIAL_SHADINGMODEL_CLEAR_COAT
-		half F = GetEnvBRDF(0.04, ShadingModelContext.ClearCoatRoughness, ShadingModelContext.NoV).x;
-		F *= ShadingModelContext.ClearCoat;
-		half LayerAttenuation = (1 - F);
-
-		// Fc * Vis
-		#if FORWARDSHADING_USE_HQ_ENV_BRDF
-			half2 AB = PreIntegratedGF.SampleLevel(PreIntegratedGFSampler, float2(ShadingModelContext.NoV, GBuffer.Roughness), 0).rg;
-		#else
-			half2 AB = EnvBRDFApproxLazarov(GBuffer.Roughness, ShadingModelContext.NoV);
-		#endif
-		SpecularIBLLighting += SpecularIBL * LayerAttenuation * (ShadingModelContext.SpecularColor * AB.x + AB.y * saturate(50 * ShadingModelContext.SpecularColor.g) * (1 - ShadingModelContext.ClearCoat));
+	// Directional light
+	AccumulateDirectionalLighting(GBuffer, ShadingModelContext, MaterialParameters.CameraVector, MaterialParameters.ScreenPosition, SvPosition, Shadow, DynamicShadowFactors, NoL, MobileLightAccumulator);
 
-		SpecularIBL = GetImageBasedReflectionLighting(MaterialParameters, ShadingModelContext.ClearCoatRoughness, IndirectIrradiance, F);
-		SpecularIBLLighting += SpecularIBL;
+	// Skylight
+	bool bApplySkyLighting = !(MATERIAL_TWOSIDED && LQ_TEXTURE_LIGHTMAP) || NoL == 0;
+	AccumulateSkyLighting(GBuffer, ShadingModelContext, ResolvedView.SkyLightColor.rgb, bApplySkyLighting, IndirectIrradiance, MobileLightAccumulator);
 
-	#else
-		SpecularIBLLighting += SpecularIBL * ShadingModelContext.EnvBrdf;
-	#endif
+	float2 LocalPosition = SvPosition.xy - ResolvedView.ViewRectMin.xy;
+	uint GridIndex = ComputeLightGridCellIndex(uint2(LocalPosition.x, LocalPosition.y), SvPosition.w);
 
-		MobileLightAccumulator_Add(MobileLightAccumulator, 0.0f, SpecularIBLLighting, 1.0f);
+	// Reflection IBL
+#if !(MATERIAL_SINGLE_SHADINGMODEL && MATERIAL_SHADINGMODEL_HAIR)
+	AccumulateReflection(GBuffer
+		, ShadingModelContext
+		, MaterialParameters.WorldPosition_CamRelative
+		, MaterialParameters.ReflectionVector
+		, IndirectIrradiance
+		, GridIndex
+		, MobileLightAccumulator);
 #endif
 
 	// Local lights
-#if MAX_DYNAMIC_POINT_LIGHTS > 0 && !MATERIAL_SHADINGMODEL_SINGLELAYERWATER
-
-		if(NumDynamicPointLights > 0)
-		{
-
-			#if SUPPORT_SPOTLIGHTS_SHADOW
-				FPCFSamplerSettings Settings;
-				Settings.ShadowDepthTexture = DynamicSpotLightShadowTexture;
-				Settings.ShadowDepthTextureSampler = DynamicSpotLightShadowSampler;
-				Settings.ShadowBufferSize = DynamicSpotLightShadowBufferSize;
-				Settings.bSubsurface = false;
-				Settings.bTreatMaxDepthUnshadowed = false;
-				Settings.DensityMulConstant = 0;
-				Settings.ProjectionDepthBiasParameters = 0;
-			#endif
-
-			AccumulateLightingOfDynamicPointLight(MaterialParameters, 
-												ShadingModelContext,
-												GBuffer,
-												MobileMovablePointLight0.LightPositionAndInvRadius,
-												MobileMovablePointLight0.LightTilePosition,
-												MobileMovablePointLight0.LightColorAndFalloffExponent, 
-												MobileMovablePointLight0.SpotLightDirectionAndSpecularScale, 
-												MobileMovablePointLight0.SpotLightAnglesAndSoftTransitionScaleAndLightShadowType, 
-												#if SUPPORT_SPOTLIGHTS_SHADOW
-												Settings,
-												MobileMovablePointLight0.SpotLightShadowSharpenAndShadowFadeFraction,
-												MobileMovablePointLight0.SpotLightShadowmapMinMax,
-												MobileMovablePointLight0.SpotLightShadowWorldToShadowMatrix,
-												#endif
-												MobileLightAccumulator);
-
-			if (MAX_DYNAMIC_POINT_LIGHTS > 1 && NumDynamicPointLights > 1)
-			{
-				AccumulateLightingOfDynamicPointLight(MaterialParameters, 
-														ShadingModelContext,
-														GBuffer,
-													MobileMovablePointLight1.LightPositionAndInvRadius, 
-													MobileMovablePointLight1.LightTilePosition,
-														MobileMovablePointLight1.LightColorAndFalloffExponent, 
-														MobileMovablePointLight1.SpotLightDirectionAndSpecularScale, 
-														MobileMovablePointLight1.SpotLightAnglesAndSoftTransitionScaleAndLightShadowType, 
-														#if SUPPORT_SPOTLIGHTS_SHADOW
-														Settings,
-														MobileMovablePointLight1.SpotLightShadowSharpenAndShadowFadeFraction,
-														MobileMovablePointLight1.SpotLightShadowmapMinMax,
-														MobileMovablePointLight1.SpotLightShadowWorldToShadowMatrix,
-														#endif
-														MobileLightAccumulator);
-
-				if (MAX_DYNAMIC_POINT_LIGHTS > 2 && NumDynamicPointLights > 2)
-				{
-					AccumulateLightingOfDynamicPointLight(MaterialParameters, 
-															ShadingModelContext,
-															GBuffer,
-															MobileMovablePointLight2.LightPositionAndInvRadius, 
-														MobileMovablePointLight2.LightTilePosition,
-															MobileMovablePointLight2.LightColorAndFalloffExponent, 
-															MobileMovablePointLight2.SpotLightDirectionAndSpecularScale, 
-															MobileMovablePointLight2.SpotLightAnglesAndSoftTransitionScaleAndLightShadowType, 
-															#if SUPPORT_SPOTLIGHTS_SHADOW
-															Settings,
-															MobileMovablePointLight2.SpotLightShadowSharpenAndShadowFadeFraction,
-															MobileMovablePointLight2.SpotLightShadowmapMinMax,
-															MobileMovablePointLight2.SpotLightShadowWorldToShadowMatrix,
-															#endif
-															MobileLightAccumulator);
-
-					if (MAX_DYNAMIC_POINT_LIGHTS > 3 && NumDynamicPointLights > 3)
-					{
-						AccumulateLightingOfDynamicPointLight(MaterialParameters, 
-																ShadingModelContext,
-																GBuffer,
-																MobileMovablePointLight3.LightPositionAndInvRadius, 
-															MobileMovablePointLight3.LightTilePosition,
-																MobileMovablePointLight3.LightColorAndFalloffExponent, 
-																MobileMovablePointLight3.SpotLightDirectionAndSpecularScale, 
-																MobileMovablePointLight3.SpotLightAnglesAndSoftTransitionScaleAndLightShadowType, 
-																#if SUPPORT_SPOTLIGHTS_SHADOW
-																Settings,
-																MobileMovablePointLight3.SpotLightShadowSharpenAndShadowFadeFraction,
-																MobileMovablePointLight3.SpotLightShadowmapMinMax,
-																MobileMovablePointLight3.SpotLightShadowWorldToShadowMatrix,
-																#endif
-																MobileLightAccumulator);
-					}
-				}
-			}
-		}
-
-#endif
-
-	// Skylight
-#if ENABLE_SKY_LIGHT
-	//@mw todo
-	// TODO: Also need to do specular.
-	#if MATERIAL_TWOSIDED && LQ_TEXTURE_LIGHTMAP
-	if (NoL == 0)
+#if  !MATERIAL_SHADINGMODEL_SINGLELAYERWATER && ENABLE_CLUSTERED_LIGHTS
 	{
-	#endif
+		const uint EyeIndex = 0;
+		const FCulledLightsGridData CulledLightGridData = GetCulledLightsGrid(GridIndex, EyeIndex);
+		half4 LocalLightDynamicShadowFactors = 1.0f;
 
-	#if MATERIAL_SHADINGMODEL_SINGLELAYERWATER
-		ShadingModelContext.WaterDiffuseIndirectLuminance += SkyDiffuseLighting;
-	#endif
-		half3 DiffuseColorForSky = ShadingModelContext.DiffuseColor;
-	#if MATERIAL_SHADINGMODEL_SUBSURFACE || MATERIAL_SHADINGMODEL_PREINTEGRATED_SKIN
-		DiffuseColorForSky += ShadingModelContext.SubsurfaceColor;
-	#endif
-		MobileLightAccumulator_Add(MobileLightAccumulator, SkyDiffuseLighting, 0.0f, half3(ResolvedView.SkyLightColor.rgb) * DiffuseColorForSky * MaterialAO);
-	#if MATERIAL_TWOSIDED && LQ_TEXTURE_LIGHTMAP
+		#if USE_SHADOWMASKTEXTURE
+			LocalLightDynamicShadowFactors = DynamicShadowFactors;
+		#endif
+		
+		uint LightingChannelMask = GetPrimitive_LightingChannelMask(MaterialParameters.PrimitiveId);
+		AccumulateLightGridLocalLighting(CulledLightGridData, ShadingModelContext, GBuffer, MaterialParameters.WorldPosition_CamRelative, MaterialParameters.CameraVector, EyeIndex, 0, LocalLightDynamicShadowFactors, LightingChannelMask, MobileLightAccumulator);
 	}
-	#endif
 #endif
 
 #if MATERIAL_SHADINGMODEL_SINGLELAYERWATER
diff --git a/Engine/Shaders/Private/MobileDeferredShading.usf b/Engine/Shaders/Private/MobileDeferredShading.usf
--- a/Engine/Shaders/Private/MobileDeferredShading.usf
+++ b/Engine/Shaders/Private/MobileDeferredShading.usf
@@ -18,21 +18,12 @@
 #define PreIntegratedGFSampler	ReflectionStruct.PreIntegratedGFSampler
 #endif
 
-#include "DynamicLightingCommon.ush"
-#include "ReflectionEnvironmentShared.ush"
-#include "LightGridCommon.ush"
-#include "MobileShadingModels.ush"
+#include "MobileLightingCommon.ush"
 #include "IESLightProfilesCommon.ush"
 #if MATERIAL_SHADER
 #include "LightFunctionCommon.ush"
 #endif
-#include "PlanarReflectionShared.ush"
 
-#define REFLECTION_COMPOSITE_USE_BLENDED_REFLECTION_CAPTURES 1
-#define REFLECTION_COMPOSITE_SUPPORT_SKYLIGHT_BLEND 1
-#define REFLECTION_COMPOSITE_HAS_SPHERE_CAPTURES 1
-#define REFLECTION_COMPOSITE_HAS_BOX_CAPTURES 1
-#include "ReflectionEnvironmentComposite.ush"
 #if METAL_PROFILE
 #include "/Engine/Public/Platform/Metal/MetalSubpassSupport.ush"
 #elif VULKAN_PROFILE
@@ -45,20 +36,6 @@
 #define IS_SPOT_LIGHT 0 
 #endif 
 
-struct FMobileLightData
-{
-	float3 TranslatedWorldPosition;
-	float  InvRadius;
-	float3 Color;
-	float  FalloffExponent;
-	float3 Direction;
-	float2 SpotAngles;
-	float SourceRadius;
-	float SpecularScale;
-	bool bInverseSquared;
-	bool bSpotLight;
-};
-
 void FetchGBuffer(in float2 UV, out float4 GBufferA, out float4 GBufferB, out float4 GBufferC, out float4 GBufferD, out float SceneDepth)
 {
 #if VULKAN_PROFILE
@@ -146,166 +123,26 @@ half ComputeLightFunctionMultiplier(float3 TranslatedWorldPosition)
 #endif
 }
 
-void GetDirectLighting(
-	FMobileLightData LightData, 
-	FMobileShadingModelContext ShadingModelContext, 
-	FGBufferData GBuffer, 
-	float3 TranslatedWorldPosition, 
-	half3 CameraVector,
-	inout FMobileLightAccumulator MobileLightAccumulator)
-{
-	float3 ToLight = LightData.TranslatedWorldPosition - TranslatedWorldPosition;
-	float DistanceSqr = dot(ToLight, ToLight);
-	float3 L = ToLight * rsqrt(DistanceSqr);
-		
-	float Attenuation = 0.0;
-	if (LightData.bInverseSquared)
-	{
-		// Sphere falloff (technically just 1/d2 but this avoids inf)
-		Attenuation = 1.0f / (DistanceSqr + 1.0f);
-		Attenuation *= Square(saturate(1 - Square(DistanceSqr * Square(LightData.InvRadius))));
-	}
-	else
-	{
-		Attenuation = RadialAttenuation(ToLight * LightData.InvRadius, LightData.FalloffExponent);
-	}
-
-	if (LightData.bSpotLight)
-	{
-		Attenuation *= SpotAttenuation(L, -LightData.Direction, LightData.SpotAngles);
-	}
-		
-	if (Attenuation > 0.0)
-	{
-		half3 H = normalize(CameraVector + L);
-		half NoL = max(0.0, dot(GBuffer.WorldNormal, L));
-		half NoH = max(0.0, dot(GBuffer.WorldNormal, H));
-		half VoH = max(0.0, dot(CameraVector, H));
-		FMobileDirectLighting Lighting = MobileIntegrateBxDF(ShadingModelContext, GBuffer, NoL, NoH, VoH, CameraVector, L);
-		MobileLightAccumulator_Add(MobileLightAccumulator, Lighting.Diffuse, Lighting.Specular * LightData.SpecularScale, LightData.Color * (1.0 / PI) * Attenuation);
-	}
-}
-
-/**
- * Adds local lighting using the light grid, does not apply directional lights, as they are done elsewhere.
- * Does not support dynamic shadows, as these require the per-light shadow mask.
- */
-void GetLightGridLocalLighting(
-	const FCulledLightsGridData InLightGridData, 
-	FMobileShadingModelContext ShadingModelContext, 
-	FGBufferData GBuffer, 
-	float3 TranslatedWorldPosition, 
-	half3 CameraVector, 
-	uint EyeIndex, 
-	uint FirstNonSimpleLightIndex,
-	inout FMobileLightAccumulator MobileLightAccumulator)
-{
-	// Limit max to ForwardLightData.NumLocalLights.
-	// This prevents GPU hangs when the PS tries to read from uninitialized NumCulledLightsGrid buffer
-	const uint NumLocalLights = min(InLightGridData.NumLocalLights, GetNumLocalLights(EyeIndex));
-	LOOP
-	for (uint LocalLightListIndex = FirstNonSimpleLightIndex; LocalLightListIndex < NumLocalLights; LocalLightListIndex++)
-	{
-		const FLocalLightData LocalLight = GetLocalLightData(InLightGridData.DataStartIndex + LocalLightListIndex, EyeIndex);
-
-		// The lights are sorted such that all that support clustered deferred are at the beginning, there might be others
-		// (e.g., lights with dynamic shadows) so we break out when the condition fails.
-		if (!LocalLight.bClusteredDeferredSupported)
-		{
-			break;
-		}
-		
-		FMobileLightData LightData = (FMobileLightData)0;
-		LightData.TranslatedWorldPosition = LocalLight.LightPositionAndInvRadius.xyz;
-		LightData.InvRadius = LocalLight.LightPositionAndInvRadius.w;
-		// extra-early out since we know light grid is sloppy and all lights in list are radial (have a range)
-		// appears useless
-		float invLightRadiusSq = LightData.InvRadius*LightData.InvRadius;
-		float3 LightToCamera = LightData.TranslatedWorldPosition - TranslatedWorldPosition;
-		if (dot(LightToCamera, LightToCamera) * invLightRadiusSq > 1.0f)
-		{
-			continue;
-		}
-
-		LightData.Color = LocalLight.LightColorAndFalloffExponent.xyz;
-		LightData.FalloffExponent = LocalLight.LightColorAndFalloffExponent.w;
-		LightData.Direction = LocalLight.LightDirectionAndShadowMask.xyz;
-		LightData.SpotAngles = LocalLight.SpotAnglesAndSourceRadiusPacked.xy;
-		LightData.bInverseSquared = LightData.FalloffExponent == 0;
-		LightData.SpecularScale = 1.0;
-		uint LightTypeAndPackedShadowMapChannelMask = asuint(LocalLight.LightDirectionAndShadowMask.w);
-		// bits [17:16] really
-		uint LightType = LightTypeAndPackedShadowMapChannelMask >> 16;
-		// TODO: not sure if this check is even needed, standard deferred always seems to set it to true?
-		LightData.bSpotLight = (LightType == LIGHT_TYPE_SPOT);
-		
-		GetDirectLighting(LightData, ShadingModelContext, GBuffer, TranslatedWorldPosition, CameraVector, MobileLightAccumulator);
-	}
-}
-
-half3 ReflectionEnvironmentSkyLighting(
+void ReflectionEnvironmentSkyLighting(
 	FMobileShadingModelContext ShadingModelContext,
 	FGBufferData GBuffer,
 	float3 TranslatedWorldPosition,
 	half3 ReflectionVector,
-	uint GridIndex)
+	uint GridIndex,
+	inout FMobileLightAccumulator MobileLightAccumulator)
 {
-	half3 Result = 0;
-
-#if ENABLE_CLUSTERED_REFLECTION
-	uint NumCulledEntryIndex = (ForwardLightData.NumGridCells + GridIndex) * NUM_CULLED_LIGHTS_GRID_STRIDE;
-	uint NumLocalReflectionCaptures = min(ForwardLightData.NumCulledLightsGrid[NumCulledEntryIndex + 0], ForwardLightData.NumReflectionCaptures);
-	uint DataStartIndex = ForwardLightData.NumCulledLightsGrid[NumCulledEntryIndex + 1];
-
-	half3 SpecularIBL = CompositeReflectionCapturesAndSkylightTWS(
-		1.0f,
-		TranslatedWorldPosition,
-		ReflectionVector,//RayDirection,
-		GBuffer.Roughness,
-		GBuffer.IndirectIrradiance,
-		1.0f,
-		0.0f,
-		NumLocalReflectionCaptures,
-		DataStartIndex,
-		0,
-		true);
-
-	BRANCH
-	if (abs(dot(PlanarReflectionStruct.ReflectionPlane.xyz, 1)) > .0001f)
-	{
-		half4 PlanarReflection = GetPlanarReflection(TranslatedWorldPosition, GBuffer.WorldNormal, GBuffer.Roughness);
-		// Planar reflections win over reflection environment
-		SpecularIBL = lerp(SpecularIBL, PlanarReflection.rgb, PlanarReflection.a);
-	}
-
-	Result += SpecularIBL * ShadingModelContext.EnvBrdf;
-#endif
-
-#if ENABLE_SKY_LIGHT
-	#if !ENABLE_CLUSTERED_REFLECTION
-		half SkyAverageBrightness = 1.0f;
-		half3 SpecularIBL = GetSkyLightReflection(ReflectionVector, GBuffer.Roughness, SkyAverageBrightness);
-	#if ALLOW_STATIC_LIGHTING
-		SpecularIBL *= ComputeMixingWeight(GBuffer.IndirectIrradiance, SkyAverageBrightness, GBuffer.Roughness);
-	#endif	
-		Result += SpecularIBL * ShadingModelContext.EnvBrdf;
-	#endif
-
-	half3 SkyDiffuseLighting = GetSkySHDiffuseSimple(GBuffer.WorldNormal);
-	half3 DiffuseColorForSky = ShadingModelContext.DiffuseColor;
-	
-	#if MOBILE_SHADINGMODEL_SUPPORT
-	if (GBuffer.ShadingModelID == SHADINGMODELID_SUBSURFACE || GBuffer.ShadingModelID == SHADINGMODELID_PREINTEGRATED_SKIN)
-	{
-		half3 SubsurfaceColor = ExtractSubsurfaceColor(GBuffer);
-		// Add subsurface energy to diffuse
-		DiffuseColorForSky += SubsurfaceColor;
-	}
-	#endif
-	Result += SkyDiffuseLighting * half3(View.SkyLightColor.rgb) * DiffuseColorForSky * GBuffer.GBufferAO;
-#endif
-		
-	return Result;
+	half IndirectIrradiance = GBuffer.IndirectIrradiance;
+	// Skylight
+	AccumulateSkyLighting(GBuffer, ShadingModelContext, View.SkyLightColor.rgb, true, IndirectIrradiance, MobileLightAccumulator);
+
+	// IBL
+	AccumulateReflection(GBuffer
+		, ShadingModelContext
+		, TranslatedWorldPosition
+		, ReflectionVector
+		, IndirectIrradiance
+		, GridIndex
+		, MobileLightAccumulator);
 }
   
 void MobileDirectionalLightPS(
@@ -328,32 +165,14 @@ void MobileDirectionalLightPS(
 	float3 TranslatedWorldPosition = mul(float4(ScreenPos * GBuffer.Depth, GBuffer.Depth, 1), PrimaryView.ScreenToTranslatedWorld).xyz;
 
 	half3 CameraVector = normalize(PrimaryView.TranslatedWorldCameraOrigin - TranslatedWorldPosition);
-	half NoV = max(0, dot(GBuffer.WorldNormal, CameraVector));
+	half NoV = dot(GBuffer.WorldNormal, CameraVector);
 	half3 ReflectionVector = GBuffer.WorldNormal * (NoV * 2.0) - CameraVector;
-	
-	// Directional light
-	half3 L = MobileDirectionalLight.DirectionalLightDirectionAndShadowTransition.xyz;
-	half NoL = max(0, dot(GBuffer.WorldNormal, L));
-	half3 H = normalize(CameraVector + L);
-	half NoH = max(0, dot(GBuffer.WorldNormal, H));
-	half VoH = max(0, dot(CameraVector, H));
+	NoV = max(0, NoV);
 
 	FMobileLightAccumulator MobileLightAccumulator = (FMobileLightAccumulator)0;
 	// Check movable light param to determine if we should be using precomputed shadows
 	half Shadow = LightFunctionParameters2.z > 0.0f ? 1.0f : GBuffer.PrecomputedShadowFactors.r;
-	
-#if APPLY_CSM && !ENABLE_DISTANCE_FIELD
-	// TODO: we don't need to fully compute ScreenPosition here
-	float  ShadowPositionZ = 0;
 	float4 ScreenPosition = SvPositionToScreenPosition(float4(SvPosition.xyz, GBuffer.Depth));
-	float ShadowMap = MobileDirectionalLightCSM(ScreenPosition.xy, GBuffer.Depth, NoL, ShadowPositionZ);
-	Shadow = min(ShadowMap, Shadow);
-#endif
-
-#if ENABLE_DISTANCE_FIELD
-	Shadow = Texture2DSample(MobileBasePass.ScreenSpaceShadowMaskTexture, MobileBasePass.ScreenSpaceShadowMaskSampler, SvPositionToBufferUV(SvPosition)).x;
-	Shadow = DecodeLightAttenuation(Shadow);
-#endif
 
 	FMobileShadingModelContext ShadingModelContext = (FMobileShadingModelContext)0;
 	InitShadingModelContext(ShadingModelContext, GBuffer, CameraVector);
@@ -365,17 +184,20 @@ void MobileDirectionalLightPS(
 	{
 		const uint EyeIndex = 0;
 		const FCulledLightsGridData CulledLightGridData = GetCulledLightsGrid(GridIndex, EyeIndex);
-		GetLightGridLocalLighting(CulledLightGridData, ShadingModelContext, GBuffer, TranslatedWorldPosition, CameraVector, EyeIndex, 0, MobileLightAccumulator);
+		half4 LocalLightDynamicShadowFactors = 1.0f;
+		AccumulateLightGridLocalLighting(CulledLightGridData, ShadingModelContext, GBuffer, TranslatedWorldPosition, CameraVector, EyeIndex, 0, LocalLightDynamicShadowFactors, MobileLightAccumulator);
 	}
 #endif
 
-	FMobileDirectLighting Lighting = MobileIntegrateBxDF(ShadingModelContext, GBuffer, NoL, NoH, VoH, CameraVector, L, Shadow);
-	MobileLightAccumulator_Add(MobileLightAccumulator, Lighting.Diffuse, Lighting.Specular * MobileDirectionalLight.DirectionalLightDistanceFadeMADAndSpecularScale.z, MobileDirectionalLight.DirectionalLightColor.rgb);
+	// Directional light
+	half4 DynamicShadowFactors = 1.0f;
+	half NoL = 0.0f;
+	AccumulateDirectionalLighting(GBuffer, ShadingModelContext, CameraVector, ScreenPosition, SvPosition, Shadow, DynamicShadowFactors, NoL, MobileLightAccumulator);
+
 
 #if ENABLE_CLUSTERED_REFLECTION || ENABLE_SKY_LIGHT
 	// If we have a single directional light, apply relfection and sky contrubution here
-	half3 ReflectionAndSky = ReflectionEnvironmentSkyLighting(ShadingModelContext, GBuffer, TranslatedWorldPosition, ReflectionVector, GridIndex);
-	MobileLightAccumulator_Add(MobileLightAccumulator, ReflectionAndSky, 0.0, 1.0);
+	ReflectionEnvironmentSkyLighting(ShadingModelContext, GBuffer, TranslatedWorldPosition, ReflectionVector, GridIndex, MobileLightAccumulator);
 #endif
 
 	half3 Color = MobileLightAccumulator_GetResult(MobileLightAccumulator);
@@ -391,6 +213,15 @@ void MobileDirectionalLightPS(
 #endif
 }
 
+#if SUPPORT_SPOTLIGHTS_SHADOW
+float4			SpotLightShadowSharpenAndFadeFractionAndReceiverDepthBiasAndSoftTransitionScale;
+float4			SpotLightShadowmapMinMax;
+float4x4		SpotLightShadowWorldToShadowMatrix;
+Texture2D		LocalLightShadowTexture;
+SamplerState	LocalLightShadowSampler;
+float4			LocalLightShadowBufferSize;
+#endif
+
 void MobileRadialLightPS(
 	float4 InScreenPosition : TEXCOORD0,
 	float4 SVPos			: SV_POSITION,
@@ -404,6 +235,8 @@ void MobileRadialLightPS(
 #endif
 )
 {
+	ResolvedView = ResolveView();
+
 	float2 ScreenUV = InScreenPosition.xy / InScreenPosition.w * View.ScreenPositionScaleBias.xy + View.ScreenPositionScaleBias.wz;
 	FGBufferData GBuffer = FetchAndDecodeGBuffer(ScreenUV);
 	// With a perspective projection, the clip space position is NDC * Clip.w
@@ -412,25 +245,50 @@ void MobileRadialLightPS(
 	float3 TranslatedWorldPosition = mul(float4(ClipPosition, GBuffer.Depth, 1), PrimaryView.ScreenToTranslatedWorld).xyz;
 	half3 CameraVector = normalize(PrimaryView.TranslatedWorldCameraOrigin - TranslatedWorldPosition);
 	half NoV = max(0, dot(GBuffer.WorldNormal, CameraVector));
-		
-	FMobileLightData LightData = (FMobileLightData)0;
-	{
-		LightData.TranslatedWorldPosition = DeferredLightUniforms.TranslatedWorldPosition;
-		LightData.InvRadius = DeferredLightUniforms.InvRadius;
-		LightData.Color = DeferredLightUniforms.Color;
-		LightData.FalloffExponent = DeferredLightUniforms.FalloffExponent;
-		LightData.Direction = DeferredLightUniforms.Direction;
-		LightData.SpotAngles = DeferredLightUniforms.SpotAngles;
-		LightData.SpecularScale = 1.0;
-		LightData.bInverseSquared = DeferredLightUniforms.FalloffExponent == 0;
-		LightData.bSpotLight = IS_SPOT_LIGHT; 
-	}
 
 	FMobileShadingModelContext ShadingModelContext = (FMobileShadingModelContext)0;
 	InitShadingModelContext(ShadingModelContext, GBuffer, CameraVector);
 			
 	FMobileLightAccumulator MobileLightAccumulator = (FMobileLightAccumulator)0;
-	GetDirectLighting(LightData, ShadingModelContext, GBuffer, TranslatedWorldPosition, CameraVector, MobileLightAccumulator);
+
+	FMobileShadowInfo ShadowInfo;
+#if SUPPORT_SPOTLIGHTS_SHADOW
+	FPCFSamplerSettings Settings;
+	Settings.ShadowDepthTexture = LocalLightShadowTexture;
+	Settings.ShadowDepthTextureSampler = LocalLightShadowSampler;
+	Settings.ShadowBufferSize = LocalLightShadowBufferSize;
+	Settings.bSubsurface = false;
+	Settings.bTreatMaxDepthUnshadowed = false;
+	Settings.DensityMulConstant = 0;
+	Settings.ProjectionDepthBiasParameters = 0;
+	
+	ShadowInfo.Settings = Settings;
+	ShadowInfo.SpotLightShadowSharpen = SpotLightShadowSharpenAndFadeFractionAndReceiverDepthBiasAndSoftTransitionScale.x;
+	ShadowInfo.SpotLightFadeFraction = SpotLightShadowSharpenAndFadeFractionAndReceiverDepthBiasAndSoftTransitionScale.y;
+	ShadowInfo.SpotLightReceiverDepthBias = SpotLightShadowSharpenAndFadeFractionAndReceiverDepthBiasAndSoftTransitionScale.z;
+	ShadowInfo.SpotLightSoftTransitionScale = SpotLightShadowSharpenAndFadeFractionAndReceiverDepthBiasAndSoftTransitionScale.w;
+	ShadowInfo.SpotLightShadowmapMinMax = SpotLightShadowmapMinMax;
+	ShadowInfo.SpotLightShadowWorldToShadowMatrix = SpotLightShadowWorldToShadowMatrix;
+#endif
+
+	AccumulateLocalLighting(GBuffer,
+		ShadingModelContext,
+		TranslatedWorldPosition,
+		CameraVector,
+		DeferredLightUniforms.TranslatedWorldPosition,
+		DeferredLightUniforms.InvRadius,
+		DeferredLightUniforms.Color,
+		DeferredLightUniforms.FalloffExponent,
+		DeferredLightUniforms.Direction,
+		1.0f,
+		DeferredLightUniforms.SpotAngles,
+		DeferredLightUniforms.FalloffExponent == 0,
+		IS_SPOT_LIGHT,
+		DeferredLightUniforms.ShadowedBits > 0,
+		ShadowInfo,
+		MobileLightAccumulator
+		);
+
 	half3 Color = MobileLightAccumulator_GetResult(MobileLightAccumulator);
 	
 	half LightAttenuation = ComputeLightProfileMultiplier(TranslatedWorldPosition, DeferredLightUniforms.TranslatedWorldPosition, -DeferredLightUniforms.Direction, DeferredLightUniforms.Tangent);
@@ -473,7 +331,10 @@ void MobileReflectionEnvironmentSkyLightingPS(
 
 	float2 LocalPosition = SvPosition.xy - View.ViewRectMin.xy;
 	uint GridIndex = ComputeLightGridCellIndex(uint2(LocalPosition.x, LocalPosition.y), GBuffer.Depth);
-	half3 ReflectionAndSky = ReflectionEnvironmentSkyLighting(ShadingModelContext, GBuffer, TranslatedWorldPosition, ReflectionVector, GridIndex);
+	
+	FMobileLightAccumulator MobileLightAccumulator = (FMobileLightAccumulator)0;
+	ReflectionEnvironmentSkyLighting(ShadingModelContext, GBuffer, TranslatedWorldPosition, ReflectionVector, GridIndex, MobileLightAccumulator);
+	half3 ReflectionAndSky = MobileLightAccumulator_GetResult(MobileLightAccumulator);
 
 	ReflectionAndSky *= View.PreExposure;
 
diff --git a/Engine/Shaders/Private/MobileLightingCommon.ush b/Engine/Shaders/Private/MobileLightingCommon.ush
new file mode 100644
--- /dev/null
+++ b/Engine/Shaders/Private/MobileLightingCommon.ush
@@ -0,0 +1,612 @@
+// Copyright Epic Games, Inc. All Rights Reserved.
+
+#pragma once
+
+#include "MobileShadingModels.ush"
+#include "ReflectionEnvironmentShared.ush"
+#include "PlanarReflectionShared.ush"
+#include "DynamicLightingCommon.ush"
+#include "LightGridCommon.ush"
+
+#ifndef TRANSLUCENCY_NON_DIRECTIONAL
+#define  TRANSLUCENCY_NON_DIRECTIONAL 0
+#endif // !TRANSLUCENCY_NON_DIRECTIONAL
+
+#ifndef MOBILE_USE_CSM_BRANCH
+#define MOBILE_USE_CSM_BRANCH 0
+#endif
+
+#ifndef LQ_TEXTURE_LIGHTMAP
+#define  LQ_TEXTURE_LIGHTMAP 0
+#endif
+
+#define REFLECTION_COMPOSITE_USE_BLENDED_REFLECTION_CAPTURES 1
+#define REFLECTION_COMPOSITE_SUPPORT_SKYLIGHT_BLEND 1
+#define REFLECTION_COMPOSITE_HAS_SPHERE_CAPTURES 1
+#define REFLECTION_COMPOSITE_HAS_BOX_CAPTURES 1
+#include "ReflectionEnvironmentComposite.ush"
+
+/*------------------------------------------------------------------------------
+	Mobile Light Accumulation
+------------------------------------------------------------------------------*/
+struct FMobileLightAccumulator
+{
+	float3 TotalLight;
+};
+
+// accumulate light, can be called multiple times
+void MobileLightAccumulator_Add(inout FMobileLightAccumulator In, half3 DiffuseTotalLight, half3 SpecularTotalLight, half3 CommonMultiplier)
+{
+#if FULLY_ROUGH
+	In.TotalLight += DiffuseTotalLight * CommonMultiplier;
+#else
+	In.TotalLight += (DiffuseTotalLight + SpecularTotalLight) * CommonMultiplier;
+#endif
+}
+
+half3 MobileLightAccumulator_GetResult(FMobileLightAccumulator In)
+{
+	return min(65000.0, In.TotalLight);
+}
+
+/*------------------------------------------------------------------------------
+	Mobile Shadow.
+------------------------------------------------------------------------------*/
+
+half ApplyPrecomputedShadowMask(half ShadowMap, half Shadow)
+{
+#if ALLOW_STATIC_LIGHTING
+	return min(ShadowMap, Shadow);
+#else
+	return ShadowMap;
+#endif
+}
+
+#if MOBILE_USE_CSM_BRANCH
+	uint UseCSM;
+#endif
+
+bool IsCSMEnabled()
+{
+#if MOBILE_USE_CSM_BRANCH
+	return UseCSM != 0u;
+#else
+	return true;
+#endif
+}
+
+half MobileShadowPCF(float2 ShadowUVs, FPCFSamplerSettings Settings)
+{
+#if MOBILE_SHADOW_QUALITY == 0
+	half ShadowMap = ManualNoFiltering(ShadowUVs, Settings);
+#elif MOBILE_SHADOW_QUALITY == 1
+	half ShadowMap = Manual1x1PCF(ShadowUVs, Settings);
+#elif MOBILE_SHADOW_QUALITY == 2
+	half ShadowMap = Manual3x3PCF(ShadowUVs, Settings);
+#elif MOBILE_SHADOW_QUALITY == 3
+	half ShadowMap = Manual5x5PCF(ShadowUVs, Settings);
+#else
+	#error Unsupported MOBILE_SHADOW_QUALITY value.
+#endif
+
+return ShadowMap;
+}
+
+// Add fading CSM plane:
+#define FADE_CSM 1
+
+half MobileDirectionalLightCSM(float2 ScreenPosition, float SceneDepth, half NoL, inout float ShadowPositionZ)
+{
+	half ShadowMap = 1;
+	ShadowPositionZ = 0;
+	FPCFSamplerSettings Settings;
+	Settings.ShadowDepthTexture = MobileDirectionalLight.DirectionalLightShadowTexture;
+	Settings.ShadowDepthTextureSampler = MobileDirectionalLight.DirectionalLightShadowSampler;
+	Settings.TransitionScale = MobileDirectionalLight.DirectionalLightDirectionAndShadowTransition.w * lerp(MobileDirectionalLight.DirectionalLightDistanceFadeMADAndSpecularScale.w, 1.0, NoL);
+	Settings.ShadowBufferSize = MobileDirectionalLight.DirectionalLightShadowSize;
+	Settings.bSubsurface = false;
+	Settings.bTreatMaxDepthUnshadowed = false;
+	Settings.DensityMulConstant = 0;
+	Settings.ProjectionDepthBiasParameters = 0;
+
+	float4 ShadowPosition = float4(0, 0, 0, 0);
+	for (int i = 0; i < MAX_MOBILE_SHADOWCASCADES; i++)
+	{
+		if (SceneDepth < MobileDirectionalLight.DirectionalLightShadowDistances[i])
+		{
+		#if MOBILE_MULTI_VIEW
+			ShadowPosition = mul(float4(ScreenPosition.x, ScreenPosition.y, SceneDepth, 1), ResolvedView.MobileMultiviewShadowTransform);
+			ShadowPosition = mul(ShadowPosition, MobileDirectionalLight.DirectionalLightScreenToShadow[i]);
+		#else
+			ShadowPosition = mul(float4(ScreenPosition.x, ScreenPosition.y, SceneDepth, 1), MobileDirectionalLight.DirectionalLightScreenToShadow[i]);
+		#endif
+			ShadowPositionZ = ShadowPosition.z;
+			break; // position found.
+		}
+	}
+
+	// Process CSM only when ShadowPosition is valid.
+	if (ShadowPosition.z > 0)
+	{
+		// Clamp pixel depth in light space for shadowing opaque, because areas of the shadow depth buffer that weren't rendered to will have been cleared to 1
+		// We want to force the shadow comparison to result in 'unshadowed' in that case, regardless of whether the pixel being shaded is in front or behind that plane
+
+		// Invert ShadowZ as the shadow space has been changed (but not yet the filtering code)
+		float ShadowZ = 1.0f - ShadowPosition.z;
+		float LightSpacePixelDepthForOpaque = min(ShadowZ, 0.99999f);
+		Settings.SceneDepth = LightSpacePixelDepthForOpaque;
+
+		ShadowMap = MobileShadowPCF(ShadowPosition.xy, Settings);
+
+		#if FADE_CSM
+			float Fade = saturate(SceneDepth * MobileDirectionalLight.DirectionalLightDistanceFadeMADAndSpecularScale.x + MobileDirectionalLight.DirectionalLightDistanceFadeMADAndSpecularScale.y);
+			// lerp out shadow based on fade params.
+			ShadowMap = lerp(ShadowMap, 1.0, Fade * Fade);
+		#endif
+	}
+	return ShadowMap;
+}
+
+/*------------------------------------------------------------------------------
+	Mobile Directional Light.
+------------------------------------------------------------------------------*/
+
+void AccumulateDirectionalLighting(FGBufferData GBuffer, FMobileShadingModelContext ShadingModelContext, half3 CameraVector, float4 ScreenPosition, float4 SvPosition, half Shadow, inout half4 DynamicShadowFactors, out half NoL, inout FMobileLightAccumulator MobileLightAccumulator)
+{
+	half3 L = MobileDirectionalLight.DirectionalLightDirectionAndShadowTransition.xyz;
+	NoL = max(0, dot(GBuffer.WorldNormal, L));
+	float3 H = normalize(float3(CameraVector + L));
+	half NoH = max(0, dot(GBuffer.WorldNormal, H));
+	half VoH = max(0, dot(CameraVector, H));
+
+#if USE_SHADOWMASKTEXTURE
+	half4 PreviewShadowMapChannelMask = half4(1.0f, 0.0f, 0.0f, 0.0f);
+	DynamicShadowFactors = Texture2DSample(MobileBasePass.ScreenSpaceShadowMaskTexture, MobileBasePass.ScreenSpaceShadowMaskSampler, SvPositionToBufferUV(SvPosition));
+	DynamicShadowFactors = DecodeLightAttenuation(DynamicShadowFactors);
+		
+	#if ENABLE_CLUSTERED_LIGHTS
+		const FDirectionalLightData DirectionalLightData = GetDirectionalLightData(0);
+		PreviewShadowMapChannelMask = UnpackShadowMapChannelMask(DirectionalLightData.DirectionalLightShadowMapChannelMask >> 4);
+	#endif
+
+	half ShadowMap = dot(PreviewShadowMapChannelMask, DynamicShadowFactors);
+	Shadow = ApplyPrecomputedShadowMask(ShadowMap, Shadow);
+#else
+
+	float  ShadowPositionZ = 0;
+	half ShadowMap = MobileDirectionalLightCSM(ScreenPosition.xy, ScreenPosition.w, NoL, ShadowPositionZ);
+
+	#if !MOBILE_DEFERRED_SHADING && DIRECTIONAL_LIGHT_CSM && !MATERIAL_SHADINGMODEL_SINGLELAYERWATER
+		// Cascaded Shadow Map
+		if (IsCSMEnabled())
+		{		
+			Shadow = ApplyPrecomputedShadowMask(ShadowMap, Shadow);
+		}
+	#elif MOBILE_DEFERRED_SHADING && APPLY_CSM
+		Shadow = min(ShadowMap, Shadow);
+	#endif
+#endif
+
+#if TRANSLUCENCY_NON_DIRECTIONAL
+	NoL = 1.f;
+	ShadingModelContext.SpecularColor = 0.f;
+#endif
+
+	// Direct Lighting (Directional light) + IBL
+	FMobileDirectLighting Lighting = MobileIntegrateBxDF(ShadingModelContext, GBuffer, NoL, NoH, VoH, CameraVector, L, Shadow);
+	// MobileDirectionalLight.DirectionalLightDistanceFadeMADAndSpecularScale.z saves SpecularScale for direction light.
+	MobileLightAccumulator_Add(MobileLightAccumulator, Lighting.Diffuse, Lighting.Specular * MobileDirectionalLight.DirectionalLightDistanceFadeMADAndSpecularScale.z, MobileDirectionalLight.DirectionalLightColor.rgb);
+}
+
+/*------------------------------------------------------------------------------
+	Mobile SkyLight.
+------------------------------------------------------------------------------*/
+void AccumulateSkyLighting(FGBufferData GBuffer, FMobileShadingModelContext ShadingModelContext, half3 SkyLightColor, bool bApplySkyLighting, inout half IndirectIrradiance, inout FMobileLightAccumulator MobileLightAccumulator)
+{
+#if ENABLE_SKY_LIGHT
+	//@mw todo
+	// TODO: Also need to do specular.
+	if (bApplySkyLighting)
+	{
+		half3 SkyDiffuseLighting = GetSkySHDiffuseSimple(GBuffer.WorldNormal);
+		half3 DiffuseLookup = SkyDiffuseLighting * SkyLightColor;
+		IndirectIrradiance += Luminance(DiffuseLookup);
+
+	#if MATERIAL_SHADINGMODEL_SINGLELAYERWATER
+		ShadingModelContext.WaterDiffuseIndirectLuminance += SkyDiffuseLighting;
+	#endif
+
+		half3 DiffuseColorForSky = ShadingModelContext.DiffuseColor;
+
+	#if MOBILE_SHADINGMODEL_SUPPORT
+		if (GBuffer.ShadingModelID == SHADINGMODELID_SUBSURFACE || GBuffer.ShadingModelID == SHADINGMODELID_PREINTEGRATED_SKIN)
+		{
+			// Add subsurface energy to diffuse
+			DiffuseColorForSky += ShadingModelContext.SubsurfaceColor;
+		}
+	#endif
+
+		MobileLightAccumulator_Add(MobileLightAccumulator, SkyDiffuseLighting, 0.0f, SkyLightColor * DiffuseColorForSky * GBuffer.GBufferAO);
+	}
+#endif
+}
+
+/*------------------------------------------------------------------------------
+	Mobile Reflection.
+------------------------------------------------------------------------------*/
+
+#if IS_MOBILE_BASE_PASS
+/** Prenormalized capture of the scene that's closest to the object being rendered. */
+half3 GetMobileSkyLightReflection(half3 ReflectionVector, half Roughness, half CubemapMaxMip)
+{
+	half AbsoluteSpecularMip = ComputeReflectionCaptureMipFromRoughness(Roughness, CubemapMaxMip);
+	half4 Reflection = MobileReflectionCapture.Texture.SampleLevel(MobileReflectionCapture.TextureSampler, ReflectionVector, AbsoluteSpecularMip);
+	return Reflection.rgb * ResolvedView.SkyLightColor.rgb;
+}
+#endif
+
+half3 GetImageBasedReflectionLighting(half3 ReflectionVector
+	, float3 TranslatedWorldPosition
+	, half Roughness
+	, half IndirectIrradiance
+	, half CompositeAlpha
+	, uint GridIndex
+)
+{
+	half3 SpecularIBL = (half3)0.0f;
+
+#if ENABLE_CLUSTERED_REFLECTION
+	uint NumCulledEntryIndex = (ForwardLightData.NumGridCells + GridIndex) * NUM_CULLED_LIGHTS_GRID_STRIDE;
+	uint NumLocalReflectionCaptures = min(ForwardLightData.NumCulledLightsGrid[NumCulledEntryIndex + 0], ForwardLightData.NumReflectionCaptures);
+	uint DataStartIndex = ForwardLightData.NumCulledLightsGrid[NumCulledEntryIndex + 1];
+
+	SpecularIBL = CompositeReflectionCapturesAndSkylightTWS(
+		1.0f,
+		TranslatedWorldPosition,
+		ReflectionVector,//RayDirection,
+		Roughness,
+		IndirectIrradiance,
+		1.0f,
+		0.0f,
+		NumLocalReflectionCaptures,
+		DataStartIndex,
+		0,
+		true);
+
+#elif MOBILE_DEFERRED_LIGHTING
+
+	float SkyAverageBrightness = 1.0f;
+	SpecularIBL = GetSkyLightReflection(ReflectionVector, Roughness, SkyAverageBrightness);
+	
+	#if ALLOW_STATIC_LIGHTING
+	SpecularIBL *= ComputeMixingWeight(IndirectIrradiance, SkyAverageBrightness, Roughness);
+	#endif
+
+#else // !(ENABLE_CLUSTERED_REFLECTION || MOBILE_DEFERRED_LIGHTING)
+
+	half UsingSkyReflection = MobileReflectionCapture.Params.y > 0.0f;
+	if (UsingSkyReflection)
+	{
+		// Apply sky colour if the reflection map is the sky.
+		SpecularIBL = GetMobileSkyLightReflection(ReflectionVector, Roughness, MobileReflectionCapture.Params.y);
+	}
+	else
+	{
+		half AbsoluteSpecularMip = ComputeReflectionCaptureMipFromRoughness(Roughness, ResolvedView.ReflectionCubemapMaxMip);
+		half4 SpecularIBLSample = MobileReflectionCapture.Texture.SampleLevel(MobileReflectionCapture.TextureSampler, ReflectionVector, AbsoluteSpecularMip);
+		
+		half MaxValue = MobileReflectionCapture.Params.z;
+		SpecularIBL = RGBMDecode(SpecularIBLSample, MaxValue);
+		SpecularIBL = SpecularIBL * SpecularIBL;
+
+		half ReflectionCaptureBrightness = MobileReflectionCapture.Params.w;
+		SpecularIBL = SpecularIBL * ReflectionCaptureBrightness;
+	}
+
+	#if ALLOW_STATIC_LIGHTING
+	SpecularIBL *= ComputeMixingWeight(IndirectIrradiance, MobileReflectionCapture.Params.x, Roughness);
+	#endif
+
+#endif
+
+	SpecularIBL *= CompositeAlpha;
+
+	return SpecularIBL;
+}
+
+void AccumulateReflection(FGBufferData GBuffer
+	, FMobileShadingModelContext ShadingModelContext
+	, float3 TranslatedWorldPosition
+	, half3 ReflectionVector
+	, half IndirectIrradiance
+	, uint GridIndex
+	, inout FMobileLightAccumulator MobileLightAccumulator)
+{
+	half3 SpecularIBLLighting = (half3)0.0f;
+
+	half3 SpecularIBL = GetImageBasedReflectionLighting(ReflectionVector
+		, TranslatedWorldPosition
+		, GBuffer.Roughness
+		, IndirectIrradiance
+		, 1.0f
+		, GridIndex
+	);
+
+#if MOBILE_DEFERRED_LIGHTING || MATERIAL_PLANAR_FORWARD_REFLECTIONS
+	BRANCH
+	if (abs(dot(PlanarReflectionStruct.ReflectionPlane.xyz, 1)) > .0001f)
+	{
+		half4 PlanarReflection = GetPlanarReflection(TranslatedWorldPosition, GBuffer.WorldNormal, GBuffer.Roughness);
+		// Planar reflections win over reflection environment
+		SpecularIBL = lerp(SpecularIBL, PlanarReflection.rgb, PlanarReflection.a);
+	}
+#endif
+
+#if MOBILE_SHADINGMODEL_SUPPORT
+	if (GBuffer.ShadingModelID == SHADINGMODELID_CLEAR_COAT)
+	{
+		half F = GetEnvBRDF(0.04, ShadingModelContext.ClearCoatRoughness, ShadingModelContext.NoV).x;
+		F *= ShadingModelContext.ClearCoat;
+		half LayerAttenuation = (1 - F);
+
+		// Fc * Vis
+	#if FORWARDSHADING_USE_HQ_ENV_BRDF
+		half2 AB = PreIntegratedGF.SampleLevel(PreIntegratedGFSampler, float2(ShadingModelContext.NoV, GBuffer.Roughness), 0).rg;
+	#else
+		half2 AB = EnvBRDFApproxLazarov(GBuffer.Roughness, ShadingModelContext.NoV);
+	#endif
+		SpecularIBLLighting += SpecularIBL * LayerAttenuation * (ShadingModelContext.SpecularColor * AB.x + AB.y * saturate(50 * ShadingModelContext.SpecularColor.g) * (1 - ShadingModelContext.ClearCoat));
+
+		SpecularIBL = GetImageBasedReflectionLighting(ReflectionVector
+			, TranslatedWorldPosition
+			,ShadingModelContext.ClearCoatRoughness
+			, IndirectIrradiance
+			, F
+			, GridIndex
+		);
+		SpecularIBLLighting += SpecularIBL;
+	}
+	else
+#endif
+	{
+		SpecularIBLLighting += SpecularIBL * ShadingModelContext.EnvBrdf;
+	}
+
+	MobileLightAccumulator_Add(MobileLightAccumulator, 0.0f, SpecularIBLLighting, 1.0f);
+}
+
+/*------------------------------------------------------------------------------
+	Mobile Local lights.
+------------------------------------------------------------------------------*/
+
+struct FMobileShadowInfo
+{
+#if USE_SHADOWMASKTEXTURE
+	half Shadow;
+#elif SUPPORT_SPOTLIGHTS_SHADOW
+	FPCFSamplerSettings Settings;
+	float SpotLightShadowSharpen;
+	float SpotLightFadeFraction;
+	float SpotLightReceiverDepthBias;
+	float SpotLightSoftTransitionScale;
+	float4 SpotLightShadowmapMinMax;
+	float4x4 SpotLightShadowWorldToShadowMatrix;
+#endif
+};
+
+struct FMobileLightData
+{
+	float3 TranslatedWorldPosition;
+	float  InvRadius;
+	float3 Color;
+	float  FalloffExponent;
+	float3 Direction;
+	float2 SpotAngles;
+	float SourceRadius;
+	float SpecularScale;
+	bool bInverseSquared;
+	bool bSpotLight;
+	bool bCastSpotLightShadow;
+
+// shadow info
+	FMobileShadowInfo ShadowInfo;
+};
+
+void AccumulateLocalLightingInner(FMobileLightData LightData
+	, FGBufferData GBuffer
+	, FMobileShadingModelContext ShadingModelContext
+	, float3 TranslatedWorldPosition
+	, half3 CameraVector
+	, inout FMobileLightAccumulator MobileLightAccumulator)
+{
+	float3 ToLight = LightData.TranslatedWorldPosition - TranslatedWorldPosition;
+	float3 LocalPosition = -ToLight;
+	float DistanceSqr = dot(ToLight, ToLight);
+	half3 L = ToLight * rsqrt(DistanceSqr);
+	float3 H = normalize(float3(CameraVector + L));
+
+	half NoL = max(0, dot(GBuffer.WorldNormal, L));
+	float NoH = max(0, dot(GBuffer.WorldNormal, H));
+	half VoH = max(0, dot(CameraVector, H));
+
+	float Attenuation = 0.0;
+	if (LightData.bInverseSquared)
+	{
+		// Sphere falloff (technically just 1/d2 but this avoids inf)
+		Attenuation = 1.0f / (DistanceSqr + 1.0f);
+		Attenuation *= Square(saturate(1 - Square(DistanceSqr * Square(LightData.InvRadius))));
+	}
+	else
+	{
+		Attenuation = RadialAttenuation(ToLight * LightData.InvRadius, LightData.FalloffExponent);
+	}
+
+	if (LightData.bSpotLight)
+	{
+		Attenuation *= SpotAttenuation(L, -LightData.Direction, LightData.SpotAngles);
+	}
+
+	if (Attenuation > 0.0)
+	{
+		half3 H = normalize(CameraVector + L);
+		half NoL = max(0.0, dot(GBuffer.WorldNormal, L));
+		half NoH = max(0.0, dot(GBuffer.WorldNormal, H));
+		half VoH = max(0.0, dot(CameraVector, H));
+
+		half Shadow = 1.0f;
+
+		if (LightData.bCastSpotLightShadow)
+		{
+#if USE_SHADOWMASKTEXTURE
+			Shadow = LightData.ShadowInfo.Shadow;
+#elif SUPPORT_SPOTLIGHTS_SHADOW
+			float4 HomogeneousShadowPosition = mul(float4(LocalPosition, 1), LightData.ShadowInfo.SpotLightShadowWorldToShadowMatrix);
+			float2 ShadowUVs = HomogeneousShadowPosition.xy / HomogeneousShadowPosition.w;
+			if (all(ShadowUVs >= LightData.ShadowInfo.SpotLightShadowmapMinMax.xy && ShadowUVs <= LightData.ShadowInfo.SpotLightShadowmapMinMax.zw))
+			{
+				// Clamp pixel depth in light space for shadowing opaque, because areas of the shadow depth buffer that weren't rendered to will have been cleared to 1
+				// We want to force the shadow comparison to result in 'unshadowed' in that case, regardless of whether the pixel being shaded is in front or behind that plane
+				// Invert ShadowZ as the shadow space has been changed (but not yet the filtering code)
+				float ShadowZ = 1.0f - HomogeneousShadowPosition.z;
+				float LightSpacePixelDepthForOpaque = min(ShadowZ, 0.99999f);
+				LightData.ShadowInfo.Settings.SceneDepth = LightSpacePixelDepthForOpaque;
+				LightData.ShadowInfo.Settings.TransitionScale = LightData.ShadowInfo.SpotLightSoftTransitionScale * lerp(LightData.ShadowInfo.SpotLightReceiverDepthBias, 1.0, NoL);
+
+				Shadow = MobileShadowPCF(ShadowUVs, LightData.ShadowInfo.Settings);
+
+				Shadow = saturate((Shadow - 0.5) * LightData.ShadowInfo.SpotLightShadowSharpen + 0.5f);
+
+				Shadow = lerp(1.0f, Square(Shadow), LightData.ShadowInfo.SpotLightFadeFraction);
+			}
+#endif
+		}
+
+		FMobileDirectLighting Lighting = MobileIntegrateBxDF(ShadingModelContext, GBuffer, NoL, NoH, VoH, CameraVector, L, Shadow);
+		MobileLightAccumulator_Add(MobileLightAccumulator, Lighting.Diffuse, Lighting.Specular * LightData.SpecularScale, LightData.Color * (1.0 / PI) * Attenuation);
+	}
+}
+
+void AccumulateLocalLighting(FGBufferData GBuffer
+	, FMobileShadingModelContext ShadingModelContext
+	, float3 TranslatedWorldPosition
+	, half3 CameraVector
+	, float3 LightTranslatedWorldPosition
+	, float InvRadius
+	, float3 Color
+	, float FalloffExponent
+	, float3 Direction
+	, float SpecularScale
+	, float2 SpotLightAngles
+	, bool bInverseSquared
+	, bool bSpotLight
+	, bool bCastSpotLightShadow
+	, FMobileShadowInfo ShadowInfo
+	, inout FMobileLightAccumulator MobileLightAccumulator)
+{
+	FMobileLightData LightData;
+	{
+		LightData.TranslatedWorldPosition = LightTranslatedWorldPosition;
+		LightData.InvRadius = InvRadius;
+		LightData.Color = Color;
+		LightData.FalloffExponent = FalloffExponent;
+		LightData.Direction = Direction;
+		LightData.SpotAngles = SpotLightAngles;
+		LightData.SpecularScale = SpecularScale;
+		LightData.bInverseSquared = bInverseSquared;
+		LightData.bSpotLight = bSpotLight;
+		LightData.bCastSpotLightShadow = bCastSpotLightShadow;
+
+		LightData.ShadowInfo = ShadowInfo;
+	}
+
+	AccumulateLocalLightingInner(LightData
+		, GBuffer
+		, ShadingModelContext
+		, TranslatedWorldPosition
+		, CameraVector
+		, MobileLightAccumulator);
+}
+
+/**
+ * Adds local lighting using the light grid, does not apply directional lights, as they are done elsewhere.
+ */
+#if ENABLE_CLUSTERED_LIGHTS
+void AccumulateLightGridLocalLighting(const FCulledLightsGridData InLightGridData
+	, FMobileShadingModelContext ShadingModelContext
+	, FGBufferData GBuffer
+	, float3 TranslatedWorldPosition
+	, half3 CameraVector
+	, uint EyeIndex
+	, uint FirstNonSimpleLightIndex
+	, half4 DynamicShadowFactors
+#if !MOBILE_DEFERRED_LIGHTING
+	, uint LightingChannelMask
+#endif
+	, inout FMobileLightAccumulator MobileLightAccumulator)
+{
+	// Limit max to ForwardLightData.NumLocalLights.
+	// This prevents GPU hangs when the PS tries to read from uninitialized NumCulledLightsGrid buffer
+	const uint NumLocalLights = min(InLightGridData.NumLocalLights, GetNumLocalLights(EyeIndex));
+	LOOP
+	for (uint LocalLightListIndex = FirstNonSimpleLightIndex; LocalLightListIndex < NumLocalLights; LocalLightListIndex++)
+	{
+		const FLocalLightData LocalLight = GetLocalLightData(InLightGridData.DataStartIndex + LocalLightListIndex, EyeIndex);
+
+#if MOBILE_DEFERRED_LIGHTING
+		// The lights are sorted such that all that support clustered deferred are at the beginning, there might be others
+		// (e.g., lights with dynamic shadows) so we break out when the condition fails.
+		if (!LocalLight.bClusteredDeferredSupported)
+		{
+			break;
+		}
+#endif
+
+		// extra-early out since we know light grid is sloppy and all lights in list are radial (have a range)
+		// appears useless
+		float invLightRadiusSq = LocalLight.LightPositionAndInvRadius.w * LocalLight.LightPositionAndInvRadius.w;
+		if (length2(TranslatedWorldPosition - LocalLight.LightPositionAndInvRadius.xyz) * invLightRadiusSq > 1.0f)
+		{
+			continue;
+		}
+
+		uint LightTypeAndPackedShadowMapChannelMask = asuint(LocalLight.LightDirectionAndShadowMask.w);
+
+#if !MOBILE_DEFERRED_LIGHTING
+		// Check lighting channels for mobile forward
+		if (((LightTypeAndPackedShadowMapChannelMask >> 8) & LightingChannelMask) == 0)
+		{
+			continue;
+		}
+#endif
+
+		// bits [17:16] really
+		uint LightType = LightTypeAndPackedShadowMapChannelMask >> 16;
+
+		half4 PreviewShadowMapChannelMask = UnpackShadowMapChannelMask(LightTypeAndPackedShadowMapChannelMask >> 4);
+		half Shadow = dot(PreviewShadowMapChannelMask, DynamicShadowFactors);
+
+		FMobileShadowInfo ShadowInfo = (FMobileShadowInfo)0;
+#if USE_SHADOWMASKTEXTURE
+		ShadowInfo.Shadow = Shadow;
+#endif
+
+		AccumulateLocalLighting(GBuffer,
+			ShadingModelContext,
+			TranslatedWorldPosition,
+			CameraVector,
+			LocalLight.LightPositionAndInvRadius.xyz,
+			LocalLight.LightPositionAndInvRadius.w,
+			LocalLight.LightColorAndFalloffExponent.xyz,
+			LocalLight.LightColorAndFalloffExponent.w,
+			LocalLight.LightDirectionAndShadowMask.xyz,
+			1.0f,
+			LocalLight.SpotAnglesAndSourceRadiusPacked.xy,
+			LocalLight.LightColorAndFalloffExponent.w == 0,
+			(LightType == LIGHT_TYPE_SPOT),
+			(LightTypeAndPackedShadowMapChannelMask & 0xFF) != 0,
+			ShadowInfo,
+			MobileLightAccumulator
+			);
+	}
+}
+#endif
\ No newline at end of file
diff --git a/Engine/Shaders/Private/MobileShadingModels.ush b/Engine/Shaders/Private/MobileShadingModels.ush
--- a/Engine/Shaders/Private/MobileShadingModels.ush
+++ b/Engine/Shaders/Private/MobileShadingModels.ush
@@ -449,103 +449,4 @@ FMobileDirectLighting MobileIntegrateBxDF(FMobileShadingModelContext ShadingMode
 #endif
 
 	return MobileDefaultLitBxDF(ShadingModelContext, GBuffer, NoL, NoH, VoH, V, L, Shadow);
-}
-
-/*------------------------------------------------------------------------------
-	Mobile Light Accumulation
-------------------------------------------------------------------------------*/
-struct FMobileLightAccumulator
-{
-	float3 TotalLight;
-};
-
-// accumulate light, can be called multiple times
-void MobileLightAccumulator_Add(inout FMobileLightAccumulator In, half3 DiffuseTotalLight, half3 SpecularTotalLight, half3 CommonMultiplier)
-{
-#if FULLY_ROUGH
-	In.TotalLight += DiffuseTotalLight * CommonMultiplier;
-#else
-	In.TotalLight += (DiffuseTotalLight + SpecularTotalLight) * CommonMultiplier;
-#endif
-}
-
-half3 MobileLightAccumulator_GetResult(FMobileLightAccumulator In)
-{
-	return min(65000.0, In.TotalLight);
-}
-
-/*------------------------------------------------------------------------------
-	Mobile Shadow.
-------------------------------------------------------------------------------*/
-
-half MobileShadowPCF(float2 ShadowUVs, FPCFSamplerSettings Settings)
-{
-#if MOBILE_SHADOW_QUALITY == 0
-	half ShadowMap = ManualNoFiltering(ShadowUVs, Settings);
-#elif MOBILE_SHADOW_QUALITY == 1
-	half ShadowMap = Manual1x1PCF(ShadowUVs, Settings);
-#elif MOBILE_SHADOW_QUALITY == 2
-	half ShadowMap = Manual3x3PCF(ShadowUVs, Settings);
-#elif MOBILE_SHADOW_QUALITY == 3
-	half ShadowMap = Manual5x5PCF(ShadowUVs, Settings);
-#else
-	#error Unsupported MOBILE_SHADOW_QUALITY value.
-#endif
-
-return ShadowMap;
-}
-
-// Add fading CSM plane:
-#define FADE_CSM 1
-
-half MobileDirectionalLightCSM(float2 ScreenPosition, float SceneDepth, half NoL, inout float ShadowPositionZ)
-{
-	half ShadowMap = 1;
-	ShadowPositionZ = 0;
-	FPCFSamplerSettings Settings;
-	Settings.ShadowDepthTexture = MobileDirectionalLight.DirectionalLightShadowTexture;
-	Settings.ShadowDepthTextureSampler = MobileDirectionalLight.DirectionalLightShadowSampler;
-	Settings.TransitionScale = MobileDirectionalLight.DirectionalLightDirectionAndShadowTransition.w * lerp(MobileDirectionalLight.DirectionalLightDistanceFadeMADAndSpecularScale.w, 1.0, NoL);
-	Settings.ShadowBufferSize = MobileDirectionalLight.DirectionalLightShadowSize;
-	Settings.bSubsurface = false;
-	Settings.bTreatMaxDepthUnshadowed = false;
-	Settings.DensityMulConstant = 0;
-	Settings.ProjectionDepthBiasParameters = 0;
-
-	float4 ShadowPosition = float4(0, 0, 0, 0);
-	for (int i = 0; i < MAX_MOBILE_SHADOWCASCADES; i++)
-	{
-		if (SceneDepth < MobileDirectionalLight.DirectionalLightShadowDistances[i])
-		{
-		#if MOBILE_MULTI_VIEW
-			ShadowPosition = mul(float4(ScreenPosition.x, ScreenPosition.y, SceneDepth, 1), ResolvedView.MobileMultiviewShadowTransform);
-			ShadowPosition = mul(ShadowPosition, MobileDirectionalLight.DirectionalLightScreenToShadow[i]);
-		#else
-			ShadowPosition = mul(float4(ScreenPosition.x, ScreenPosition.y, SceneDepth, 1), MobileDirectionalLight.DirectionalLightScreenToShadow[i]);
-		#endif
-			ShadowPositionZ = ShadowPosition.z;
-			break; // position found.
-		}
-	}
-
-	// Process CSM only when ShadowPosition is valid.
-	if (ShadowPosition.z > 0)
-	{
-		// Clamp pixel depth in light space for shadowing opaque, because areas of the shadow depth buffer that weren't rendered to will have been cleared to 1
-		// We want to force the shadow comparison to result in 'unshadowed' in that case, regardless of whether the pixel being shaded is in front or behind that plane
-
-		// Invert ShadowZ as the shadow space has been changed (but not yet the filtering code)
-		float ShadowZ = 1.0f - ShadowPosition.z;
-		float LightSpacePixelDepthForOpaque = min(ShadowZ, 0.99999f);
-		Settings.SceneDepth = LightSpacePixelDepthForOpaque;
-
-		ShadowMap = MobileShadowPCF(ShadowPosition.xy, Settings);
-
-		#if FADE_CSM
-			float Fade = saturate(SceneDepth * MobileDirectionalLight.DirectionalLightDistanceFadeMADAndSpecularScale.x + MobileDirectionalLight.DirectionalLightDistanceFadeMADAndSpecularScale.y);
-			// lerp out shadow based on fade params.
-			ShadowMap = lerp(ShadowMap, 1.0, Fade * Fade);
-		#endif
-	}
-	return ShadowMap;
-}
+}
\ No newline at end of file
diff --git a/Engine/Shaders/Private/ReflectionEnvironmentShared.ush b/Engine/Shaders/Private/ReflectionEnvironmentShared.ush
--- a/Engine/Shaders/Private/ReflectionEnvironmentShared.ush
+++ b/Engine/Shaders/Private/ReflectionEnvironmentShared.ush
@@ -200,16 +200,16 @@ float3 GetLookupVectorForSphereCapture(float3 ReflectionVector, float3 WorldPosi
 	return ProjectedCaptureVector;
 }
 
-float ComputeMixingWeight(float IndirectIrradiance, float AverageBrightness, float Roughness)
+half ComputeMixingWeight(half IndirectIrradiance, half AverageBrightness, half Roughness)
 {
 	// Mirror surfaces should have no mixing, so they match reflections from other sources (SSR, planar reflections)
-	float MixingAlpha = smoothstep(0, 1, saturate(Roughness * View.ReflectionEnvironmentRoughnessMixingScaleBiasAndLargestWeight.x + View.ReflectionEnvironmentRoughnessMixingScaleBiasAndLargestWeight.y));
+	half MixingAlpha = smoothstep(0, 1, saturate(Roughness * View.ReflectionEnvironmentRoughnessMixingScaleBiasAndLargestWeight.x + View.ReflectionEnvironmentRoughnessMixingScaleBiasAndLargestWeight.y));
 
 	// We have high frequency directional data but low frequency spatial data in the envmap.
 	// We have high frequency spatial data but low frequency directional data in the lightmap.
 	// So, we combine the two for the best of both. This is done by removing the low spatial frequencies from the envmap and replacing them with the lightmap data.
 	// This is only done with luma so as to not get odd color shifting.
-	float MixingWeight = IndirectIrradiance / max(AverageBrightness, .0001f);
+	half MixingWeight = IndirectIrradiance / max(AverageBrightness, .0001f);
 
 	MixingWeight = min(MixingWeight, View.ReflectionEnvironmentRoughnessMixingScaleBiasAndLargestWeight.z);
 
diff --git a/Engine/Shaders/Private/ShadowProjectionPixelShader.usf b/Engine/Shaders/Private/ShadowProjectionPixelShader.usf
--- a/Engine/Shaders/Private/ShadowProjectionPixelShader.usf
+++ b/Engine/Shaders/Private/ShadowProjectionPixelShader.usf
@@ -337,7 +337,7 @@ void Main(
 	// RETURN_COLOR not needed unless writing to SceneColor;
 	float FadedShadow = lerp(1.0f, Square(Shadow), ShadowFadeFraction * PerObjectDistanceFadeFraction);
 
-#if FORWARD_SHADING
+#if FORWARD_SHADING || SHADING_PATH_MOBILE 
 	float LightInfluenceMask = GetLightInfluenceMask(TranslateWorldPosition);
 	// Constrain shadowing from this light to pixels inside the light's influence, since other non-overlapping lights are packed into the same channel
 	FadedShadow = lerp(1, FadedShadow, LightInfluenceMask);
diff --git a/Engine/Source/Editor/MaterialEditor/Private/MaterialStatsCommon.cpp b/Engine/Source/Editor/MaterialEditor/Private/MaterialStatsCommon.cpp
--- a/Engine/Source/Editor/MaterialEditor/Private/MaterialStatsCommon.cpp
+++ b/Engine/Source/Editor/MaterialEditor/Private/MaterialStatsCommon.cpp
@@ -384,11 +384,6 @@ void FMaterialStatsUtils::GetRepresentativeShaderTypesAndDescriptions(TMap<FName
 				static auto* CVarAllowDistanceFieldShadows = IConsoleManager::Get().FindTConsoleVariableDataInt(TEXT("r.Mobile.AllowDistanceFieldShadows"));
 				const bool bAllowDistanceFieldShadows = CVarAllowDistanceFieldShadows->GetValueOnAnyThread() != 0;
 
-				static auto* CVarPointLights = IConsoleManager::Get().FindTConsoleVariableDataInt(TEXT("r.MobileNumDynamicPointLights"));
-				const bool bPointLights = CVarPointLights->GetValueOnAnyThread() > 0;
-				
-				const int32 NumPointLights = CVarPointLights->GetValueOnAnyThread();
-				
 				if (bAllowDistanceFieldShadows)// distance field shadows
 				{
 					// distance field shadows only shaders
@@ -411,10 +406,9 @@ void FMaterialStatsUtils::GetRepresentativeShaderTypesAndDescriptions(TMap<FName
 								.Add(FRepresentativeShaderInfo(ERepresentativeShader::StationarySurfaceCSM, FName(ShaderNameStr), Description));
 						}
 
-						if (bPointLights) // add point lights shaders + distance field shadows
 						{
 							MobileBasePassShaderName(false, TEXT("FMobileDistanceFieldShadowsAndLQLightMapPolicy"), INT32_MAX, bMobileHDR, bOnlySkyPermutation, ShaderNameStr);
-							const FString Description = FString::Printf(TEXT("Mobile base pass shader with distance field shadows, CSM and %d point light(s) %s"), NumPointLights, DescSuffix);
+							const FString Description = FString::Printf(TEXT("Mobile base pass shader with distance field shadows, CSM and local light(s) %s"), DescSuffix);
 
 							FRepresentativeShaderInfo ShaderInfo = FRepresentativeShaderInfo(ERepresentativeShader::StationarySurfaceNPointLights, FName(ShaderNameStr), Description);
 
@@ -430,10 +424,9 @@ void FMaterialStatsUtils::GetRepresentativeShaderTypesAndDescriptions(TMap<FName
 						.Add(FRepresentativeShaderInfo(ERepresentativeShader::StationarySurface, FName(ShaderNameStr),
 							FString::Printf(TEXT("Mobile base pass shader with static lighting%s"), DescSuffix)));
 
-					if (bPointLights) // add point lights + lightmap
 					{
 						MobileBasePassShaderName(false, TEXT("TLightMapPolicyLQ"), INT32_MAX, bMobileHDR, bOnlySkyPermutation, ShaderNameStr);
-						const FString Description = FString::Printf(TEXT("Mobile base pass shader with static lighting and %d point light(s) %s"), NumPointLights, DescSuffix);
+						const FString Description = FString::Printf(TEXT("Mobile base pass shader with static lighting and local light(s) %s"), DescSuffix);
 
 						FRepresentativeShaderInfo ShaderInfo = FRepresentativeShaderInfo(ERepresentativeShader::StationarySurfaceNPointLights, FName(ShaderNameStr), Description);
 
diff --git a/Engine/Source/Runtime/Core/Private/HAL/ConsoleManager.cpp b/Engine/Source/Runtime/Core/Private/HAL/ConsoleManager.cpp
--- a/Engine/Source/Runtime/Core/Private/HAL/ConsoleManager.cpp
+++ b/Engine/Source/Runtime/Core/Private/HAL/ConsoleManager.cpp
@@ -2529,12 +2529,6 @@ static TAutoConsoleVariable<int32> CVarMobileShadingPath(
 		 "1: Deferred shading"),
 	ECVF_RenderThreadSafe | ECVF_ReadOnly);
 
-static TAutoConsoleVariable<int32> CVarMobileNumDynamicPointLights(
-	TEXT("r.MobileNumDynamicPointLights"),
-	4,
-	TEXT("The number of dynamic point lights to support on mobile devices. Setting this to 0 for games which do not require dynamic point lights will reduce the number of shaders generated."), 
-	ECVF_RenderThreadSafe | ECVF_ReadOnly);
-
 static TAutoConsoleVariable<int32> CVarMobileEnableStaticAndCSMShadowReceivers(
 	TEXT("r.Mobile.EnableStaticAndCSMShadowReceivers"),
 	1,
@@ -2620,6 +2614,17 @@ static TAutoConsoleVariable<int32> CVarMobileAllowPerPixelShadingModels(
 	ECVF_ReadOnly | ECVF_RenderThreadSafe
 );
 
+static TAutoConsoleVariable<int32> CVarMobileForwardSupportLocalLights(
+	TEXT("r.Mobile.Forward.EnableLocalLights"),
+	1,
+	TEXT("Enable local lights support on mobile forward. 0 is disabled, 1 is enabled (default)"),
+	ECVF_ReadOnly | ECVF_RenderThreadSafe);
+
+static TAutoConsoleVariable<int32> CVarMobileForwardEnableClusteredReflections(
+	TEXT("r.Mobile.Forward.EnableClusteredReflections"),
+	0,
+	TEXT("Whether to enable clustered reflections on mobile forward, it's always supported on mobile deferred."),
+	ECVF_ReadOnly | ECVF_RenderThreadSafe);
 
 static TAutoConsoleVariable<int32> CVarMobileSupportGPUScene(
 	TEXT("r.Mobile.SupportGPUScene"),
diff --git a/Engine/Source/Runtime/Engine/Classes/Engine/RendererSettings.h b/Engine/Source/Runtime/Engine/Classes/Engine/RendererSettings.h
--- a/Engine/Source/Runtime/Engine/Classes/Engine/RendererSettings.h
+++ b/Engine/Source/Runtime/Engine/Classes/Engine/RendererSettings.h
@@ -1007,6 +1007,20 @@ class ENGINE_API URendererSettings : public UDeveloperSettings
 		ConfigRestartRequired = true))
 		uint32 bMobileEnableMovableLightCSMShaderCulling : 1;
 
+	UPROPERTY(config, EditAnywhere, Category = Mobile, meta = (
+		ConsoleVariable = "r.Mobile.Forward.EnableLocalLights",
+		DisplayName = "Enable local lights support on mobile forward",
+		ToolTip = "Enable local lights support with clustered lighting on mobile forward. 0 is disabled, 1 is enabled (default). Changing this setting requires restarting the editor.",
+		ConfigRestartRequired = true))
+		uint32 bMobileForwardEnableLocalLights : 1;
+
+	UPROPERTY(config, EditAnywhere, Category = Mobile, meta = (
+		ConsoleVariable = "r.Mobile.Forward.EnableClusteredReflections",
+		DisplayName = "Enable clustered reflections on mobile forward",
+		ToolTip = "Whether to enable clustered reflections on mobile forward, it's always supported on mobile deferred. Changing this setting requires restarting the editor.",
+		ConfigRestartRequired = true))
+		uint32 bMobileForwardEnableClusteredReflections : 1;
+
 	UPROPERTY(config, EditAnywhere, Category = MobileShaderPermutationReduction, meta = (
 		ConsoleVariable = "r.Mobile.EnableNoPrecomputedLightingCSMShader",
 		DisplayName = "Support CSM on levels with Force No Precomputed Lighting enabled",
@@ -1030,20 +1044,6 @@ class ENGINE_API URendererSettings : public UDeveloperSettings
 		uint32 bMobileAllowMovableDirectionalLights : 1;
 
 	UPROPERTY(config, EditAnywhere, Category = MobileShaderPermutationReduction, meta = (
-		ConsoleVariable = "r.MobileNumDynamicPointLights", DisplayName = "Max Movable Spotlights / Point Lights", ClampMax = 4,
-		ToolTip = "The number of dynamic spotlights or point lights to support on mobile devices. Setting this to 0 for games which do not require dynamic spotlights or point lights will reduce the number of shaders generated. Changing this setting requires restarting the editor.",
-		ConfigRestartRequired = true))
-		uint32 MobileNumDynamicPointLights;
-
-	UPROPERTY(config, EditAnywhere, Category = MobileShaderPermutationReduction, meta = (
-		ConsoleVariable = "r.Mobile.EnableMovableSpotlights",
-		DisplayName = "Support Movable Spotlights",
-		ToolTip = "Generate shaders for primitives to receive lighting from movable spotlights. This incurs an additional cost when processing movable lights. Changing this setting requires restarting the editor.",
-		ConfigRestartRequired = true))
-		uint32 bMobileAllowMovableSpotlights : 1;
-
-	UPROPERTY(config, EditAnywhere, Category = MobileShaderPermutationReduction, meta = (
-		EditCondition = "bMobileAllowMovableSpotlights",
 		ConsoleVariable = "r.Mobile.EnableMovableSpotlightsShadow",
 		DisplayName = "Support Movable SpotlightShadows",
 		ToolTip = "Generate shaders for primitives to receive shadow from movable spotlights. Changing this setting requires restarting the editor.",
diff --git a/Engine/Source/Runtime/Engine/Classes/Materials/Material.h b/Engine/Source/Runtime/Engine/Classes/Materials/Material.h
--- a/Engine/Source/Runtime/Engine/Classes/Materials/Material.h
+++ b/Engine/Source/Runtime/Engine/Classes/Materials/Material.h
@@ -814,7 +814,6 @@ public:
 
 	/* 
 	 * Forward renderer: enables multiple parallax-corrected reflection captures that blend together.
-	 * Mobile renderer: blend between nearest 3 reflection captures, but reduces the number of samplers available to the material as two more samplers will be used for reflection cubemaps.
 	 */
 	UPROPERTY(EditAnywhere, BlueprintReadOnly, Category = ForwardShading, meta = (DisplayName = "High Quality Reflections"))
 	uint8 bUseHQForwardReflections : 1;
diff --git a/Engine/Source/Runtime/Engine/Private/Components/DirectionalLightComponent.cpp b/Engine/Source/Runtime/Engine/Private/Components/DirectionalLightComponent.cpp
--- a/Engine/Source/Runtime/Engine/Private/Components/DirectionalLightComponent.cpp
+++ b/Engine/Source/Runtime/Engine/Private/Components/DirectionalLightComponent.cpp
@@ -353,8 +353,10 @@ public:
 				const bool bUsingDeferredRenderer = Scene == nullptr ? true : Scene->GetShadingPath() == EShadingPath::Deferred;
 				bUseWholeSceneCSMForMovableObjects = Component->Mobility == EComponentMobility::Stationary && !Component->bUseInsetShadowsForMovableObjects && !bUsingDeferredRenderer;
 			}
+			// Modulated shadow is only supported on mobile forward
+			bCastModulatedShadows = Component->bCastModulatedShadows && (Scene == nullptr || !IsMobileDeferredShadingEnabled(Scene->GetShaderPlatform()));
 		}
-		bCastModulatedShadows = Component->bCastModulatedShadows;
+		
 		ModulatedShadowColor = FLinearColor(Component->ModulatedShadowColor);
 		ShadowAmount = Component->ShadowAmount;
 	}
diff --git a/Engine/Source/Runtime/Engine/Private/Components/LightComponent.cpp b/Engine/Source/Runtime/Engine/Private/Components/LightComponent.cpp
--- a/Engine/Source/Runtime/Engine/Private/Components/LightComponent.cpp
+++ b/Engine/Source/Runtime/Engine/Private/Components/LightComponent.cpp
@@ -377,10 +377,6 @@ FLightSceneProxy::FLightSceneProxy(const ULightComponent* InLightComponent)
 	, ShadowAmount(1.0f)
 	, SamplesPerPixel(1)
 	, DeepShadowLayerDistribution(InLightComponent->DeepShadowLayerDistribution)
-	, bMobileMovablePointLightUniformBufferNeedsUpdate(false)
-	, bMobileMovablePointLightShouldBeRender(false)
-	, bMobileMovablePointLightShouldCastShadow(false)
-	, MobileMovablePointLightShadowmapMinMax(0.0f)
 #if ACTOR_HAS_LABELS
 	, OwnerNameOrLabel(InLightComponent->GetOwner() ? InLightComponent->GetOwner()->GetActorNameOrLabel() : InLightComponent->GetName())
 #endif
diff --git a/Engine/Source/Runtime/Engine/Private/Components/PointLightComponent.cpp b/Engine/Source/Runtime/Engine/Private/Components/PointLightComponent.cpp
--- a/Engine/Source/Runtime/Engine/Private/Components/PointLightComponent.cpp
+++ b/Engine/Source/Runtime/Engine/Private/Components/PointLightComponent.cpp
@@ -108,9 +108,7 @@ static bool IsPointLightSupported(const UPointLightComponent* InLight)
 		if (!IsMobileDeferredShadingEnabled(GMaxRHIShaderPlatform))
 		{
 			// if project does not support dynamic point lights on mobile do not add them to the renderer 
-			static auto* CVarPointLights = IConsoleManager::Get().FindTConsoleVariableDataInt(TEXT("r.MobileNumDynamicPointLights"));
-			const bool bPointLights = CVarPointLights->GetValueOnAnyThread() > 0;
-			return bPointLights;
+			return MobileForwardEnableLocalLights(GMaxRHIShaderPlatform);
 		}
 	}
 	return true;
diff --git a/Engine/Source/Runtime/Engine/Private/Components/SpotLightComponent.cpp b/Engine/Source/Runtime/Engine/Private/Components/SpotLightComponent.cpp
--- a/Engine/Source/Runtime/Engine/Private/Components/SpotLightComponent.cpp
+++ b/Engine/Source/Runtime/Engine/Private/Components/SpotLightComponent.cpp
@@ -260,9 +260,7 @@ static bool IsSpotLightSupported(const USpotLightComponent* InLight)
 		if (!IsMobileDeferredShadingEnabled(GMaxRHIShaderPlatform))
 		{
 			// if project does not support dynamic point/spot lights on mobile do not add them to the renderer 
-			static auto* CVarPointLights = IConsoleManager::Get().FindTConsoleVariableDataInt(TEXT("r.MobileNumDynamicPointLights"));
-			const bool bPointLights = CVarPointLights->GetValueOnAnyThread() > 0;
-			return bPointLights;
+			return MobileForwardEnableLocalLights(GMaxRHIShaderPlatform);
 		}
 	}
 	return true;
diff --git a/Engine/Source/Runtime/Engine/Private/LocalLightSceneProxy.cpp b/Engine/Source/Runtime/Engine/Private/LocalLightSceneProxy.cpp
--- a/Engine/Source/Runtime/Engine/Private/LocalLightSceneProxy.cpp
+++ b/Engine/Source/Runtime/Engine/Private/LocalLightSceneProxy.cpp
@@ -141,6 +141,9 @@ void FLocalLightSceneProxy::UpdateRadius(float ComponentRadius)
 
 	// Min to avoid div by 0 (NaN in InvRadius)
 	InvRadius = 1.0f / FMath::Max(0.00001f, ComponentRadius);
+}
 
-	bMobileMovablePointLightUniformBufferNeedsUpdate = true;
+bool FLocalLightSceneProxy::IsLocalLight() const
+{
+	return true;
 }
\ No newline at end of file
diff --git a/Engine/Source/Runtime/Engine/Private/MapBuildData.cpp b/Engine/Source/Runtime/Engine/Private/MapBuildData.cpp
--- a/Engine/Source/Runtime/Engine/Private/MapBuildData.cpp
+++ b/Engine/Source/Runtime/Engine/Private/MapBuildData.cpp
@@ -497,9 +497,9 @@ void UMapBuildDataRegistry::Serialize(FArchive& Ar)
 void UMapBuildDataRegistry::PostLoad()
 {
 	Super::PostLoad();
-	bool bUsesMobileDeferredShading = IsMobileDeferredShadingEnabled(GMaxRHIShaderPlatform);
-	bool bFullDataRequired = GMaxRHIFeatureLevel >= ERHIFeatureLevel::SM5 || bUsesMobileDeferredShading;
-	bool bEncodedDataRequired = (GIsEditor || (GMaxRHIFeatureLevel == ERHIFeatureLevel::ES3_1 && !bUsesMobileDeferredShading));
+	bool bMobileEnableClusteredReflections = MobileEnableClusteredReflections(GMaxRHIShaderPlatform);
+	bool bFullDataRequired = GMaxRHIFeatureLevel >= ERHIFeatureLevel::SM5 || bMobileEnableClusteredReflections;
+	bool bEncodedDataRequired = (GIsEditor || (GMaxRHIFeatureLevel == ERHIFeatureLevel::ES3_1 && !bMobileEnableClusteredReflections));
 
 	HandleLegacyEncodedCubemapData();
 
diff --git a/Engine/Source/Runtime/Engine/Private/SceneManagement.cpp b/Engine/Source/Runtime/Engine/Private/SceneManagement.cpp
--- a/Engine/Source/Runtime/Engine/Private/SceneManagement.cpp
+++ b/Engine/Source/Runtime/Engine/Private/SceneManagement.cpp
@@ -1167,10 +1167,7 @@ void FReadOnlyCVARCache::Init()
 	static const auto CVarMobileAllowMovableDirectionalLights = IConsoleManager::Get().FindTConsoleVariableDataInt(TEXT("r.Mobile.AllowMovableDirectionalLights"));
 	static const auto CVarMobileEnableStaticAndCSMShadowReceivers = IConsoleManager::Get().FindTConsoleVariableDataInt(TEXT("r.Mobile.EnableStaticAndCSMShadowReceivers"));
 	static const auto CVarMobileAllowDistanceFieldShadows = IConsoleManager::Get().FindTConsoleVariableDataInt(TEXT("r.Mobile.AllowDistanceFieldShadows"));
-	static const auto CVarMobileNumDynamicPointLights = IConsoleManager::Get().FindTConsoleVariableDataInt(TEXT("r.MobileNumDynamicPointLights"));
 	static const auto CVarMobileSkyLightPermutation = IConsoleManager::Get().FindTConsoleVariableDataInt(TEXT("r.Mobile.SkyLightPermutation"));
-	static const auto CVarMobileEnableMovableSpotLights = IConsoleManager::Get().FindTConsoleVariableDataInt(TEXT("r.Mobile.EnableMovableSpotlights"));
-	static const auto CVarMobileEnableMovableSpotLightsShadow = IConsoleManager::Get().FindTConsoleVariableDataInt(TEXT("r.Mobile.EnableMovableSpotlightsShadow"));
 	static const auto CVarMobileEnableNoPrecomputedLightingCSMShader = IConsoleManager::Get().FindTConsoleVariableDataInt(TEXT("r.Mobile.EnableNoPrecomputedLightingCSMShader"));
 
 	const bool bForceAllPermutations = CVarSupportAllShaderPermutations && CVarSupportAllShaderPermutations->GetValueOnAnyThread() != 0;
@@ -1185,10 +1182,7 @@ void FReadOnlyCVARCache::Init()
 	bMobileAllowMovableDirectionalLights = CVarMobileAllowMovableDirectionalLights->GetValueOnAnyThread() != 0;
 	bMobileAllowDistanceFieldShadows = CVarMobileAllowDistanceFieldShadows->GetValueOnAnyThread() != 0;
 	bMobileEnableStaticAndCSMShadowReceivers = CVarMobileEnableStaticAndCSMShadowReceivers->GetValueOnAnyThread() != 0;
-	NumMobileMovablePointLights = CVarMobileNumDynamicPointLights->GetValueOnAnyThread();
 	MobileSkyLightPermutation = CVarMobileSkyLightPermutation->GetValueOnAnyThread();
-	bMobileEnableMovableSpotlights = CVarMobileEnableMovableSpotLights->GetValueOnAnyThread() != 0;
-	bMobileEnableMovableSpotlightsShadow = bMobileEnableMovableSpotlights && CVarMobileEnableMovableSpotLightsShadow->GetValueOnAnyThread() != 0;
 	bMobileEnableNoPrecomputedLightingCSMShader = CVarMobileEnableNoPrecomputedLightingCSMShader->GetValueOnAnyThread() != 0;
 
 	const bool bShowMissmatchedLowQualityLightmapsWarning = (!bEnableLowQualityLightmaps) && (GEngine->bShouldGenerateLowQualityLightmaps_DEPRECATED);
diff --git a/Engine/Source/Runtime/Engine/Private/SceneUtils.cpp b/Engine/Source/Runtime/Engine/Private/SceneUtils.cpp
--- a/Engine/Source/Runtime/Engine/Private/SceneUtils.cpp
+++ b/Engine/Source/Runtime/Engine/Private/SceneUtils.cpp
@@ -150,15 +150,15 @@ ENGINE_API uint32 GetDefaultMSAACount(const FStaticFeatureLevel InFeatureLevel,
 				bool bMobilePixelProjectedReflection = IsUsingMobilePixelProjectedReflection(ShaderPlatform);
 				
 				static const auto MobileEarlyZPassCVar = IConsoleManager::Get().FindTConsoleVariableDataInt(TEXT("r.Mobile.EarlyZPass"));
-				bool bMobileAmbientOcclusion = UseMobileAmbientOcclusion(ShaderPlatform) && IsMobileHDR();
-				bool bMobileDistanceFieldEnabled = IsMobileDistanceFieldEnabled(ShaderPlatform);
-				bool bIsFullDepthPrepassEnabled = bMobileDistanceFieldEnabled || bMobileAmbientOcclusion || MobileEarlyZPassCVar->GetValueOnAnyThread() > 0;
+				const bool bMobileAmbientOcclusion = UseMobileAmbientOcclusion(ShaderPlatform) && IsMobileHDR();
+				const bool bMobileUsesShadowMaskTexture = MobileUsesShadowMaskTexture(ShaderPlatform);
+				bool bIsFullDepthPrepassEnabled = bMobileAmbientOcclusion || bMobileUsesShadowMaskTexture || MobileEarlyZPassCVar->GetValueOnAnyThread() > 0;
 
 				bRendererSupportMSAA = bRHISupportsMSAA && !bMobilePixelProjectedReflection && !bIsFullDepthPrepassEnabled;
 
 				if (!bRendererSupportMSAA)
 				{
-					FailedReason = FString::Printf(TEXT("RHISupportsMSAA %d, MobilePixelProjectedReflection %d, MobileAmbientOcclusion %d, MobileDistanceFieldEnabled %d, MobileEarlyZPass %d"), bRHISupportsMSAA ? 1 : 0, bMobilePixelProjectedReflection ? 1 : 0, bMobileAmbientOcclusion ? 1 : 0, bMobileDistanceFieldEnabled ? 1 : 0, MobileEarlyZPassCVar->GetValueOnAnyThread());
+					FailedReason = FString::Printf(TEXT("RHISupportsMSAA %d, MobilePixelProjectedReflection %d, MobileAmbientOcclusion %d, MobileUsesShadowMaskTexture %d, MobileEarlyZPass %d"), bRHISupportsMSAA ? 1 : 0, bMobilePixelProjectedReflection ? 1 : 0, bMobileAmbientOcclusion ? 1 : 0, bMobileUsesShadowMaskTexture? 1 : 0, MobileEarlyZPassCVar->GetValueOnAnyThread());
 				}
 			}
 			else
diff --git a/Engine/Source/Runtime/Engine/Private/SceneView.cpp b/Engine/Source/Runtime/Engine/Private/SceneView.cpp
--- a/Engine/Source/Runtime/Engine/Private/SceneView.cpp
+++ b/Engine/Source/Runtime/Engine/Private/SceneView.cpp
@@ -47,11 +47,6 @@ IMPLEMENT_GLOBAL_SHADER_PARAMETER_STRUCT(FPrimitiveUniformShaderParameters, "Pri
 IMPLEMENT_STATIC_AND_SHADER_UNIFORM_BUFFER_STRUCT(FViewUniformShaderParameters, "View", View);
 IMPLEMENT_STATIC_AND_SHADER_UNIFORM_BUFFER_STRUCT(FInstancedViewUniformShaderParameters, "InstancedView", InstancedView);
 IMPLEMENT_GLOBAL_SHADER_PARAMETER_STRUCT(FMobileDirectionalLightShaderParameters, "MobileDirectionalLight");
-IMPLEMENT_GLOBAL_SHADER_PARAMETER_STRUCT(FMobileMovablePointLightUniformShaderParameters, "MobileMovablePointLight0");
-IMPLEMENT_GLOBAL_SHADER_PARAMETER_ALIAS_STRUCT(FMobileMovablePointLightUniformShaderParameters, MobileMovablePointLight1);
-IMPLEMENT_GLOBAL_SHADER_PARAMETER_ALIAS_STRUCT(FMobileMovablePointLightUniformShaderParameters, MobileMovablePointLight2);
-IMPLEMENT_GLOBAL_SHADER_PARAMETER_ALIAS_STRUCT(FMobileMovablePointLightUniformShaderParameters, MobileMovablePointLight3);
-
 
 static TAutoConsoleVariable<float> CVarSSRMaxRoughness(
 	TEXT("r.SSR.MaxRoughness"),
@@ -346,9 +341,6 @@ EVertexColorViewMode::Type GVertexColorViewMode = EVertexColorViewMode::Color;
 /** Global primitive uniform buffer resource containing identity transformations. */
 ENGINE_API TGlobalResource<FIdentityPrimitiveUniformBuffer> GIdentityPrimitiveUniformBuffer;
 
-/** Global primitive uniform buffer resource containing identity transformations. */
-ENGINE_API TGlobalResource<FDummyMovablePointLightUniformBuffer> GDummyMovablePointLightUniformBuffer;
-
 FSceneViewStateReference::~FSceneViewStateReference()
 {
 	Destroy();
diff --git a/Engine/Source/Runtime/Engine/Private/ShaderCompiler/ShaderCompiler.cpp b/Engine/Source/Runtime/Engine/Private/ShaderCompiler/ShaderCompiler.cpp
--- a/Engine/Source/Runtime/Engine/Private/ShaderCompiler/ShaderCompiler.cpp
+++ b/Engine/Source/Runtime/Engine/Private/ShaderCompiler/ShaderCompiler.cpp
@@ -5958,14 +5958,8 @@ void GlobalBeginCompileShader(
 
 	if (IsMobilePlatform((EShaderPlatform)Target.Platform))
 	{
-		static FShaderPlatformCachedIniValue<bool> MobileEnableMovableSpotlightsIniValue(TEXT("r.Mobile.EnableMovableSpotlights"));
-		static FShaderPlatformCachedIniValue<bool> MobileEnableMovableSpotlightsShadowIniValue(TEXT("r.Mobile.EnableMovableSpotlightsShadow"));
-
-		bool bMobileEnableMovableSpotlights = (MobileEnableMovableSpotlightsIniValue.Get((EShaderPlatform)Target.Platform) != 0);
-		Input.Environment.SetDefine(TEXT("PROJECT_MOBILE_ENABLE_MOVABLE_SPOTLIGHTS"), bMobileEnableMovableSpotlights ? 1 : 0);
-
-		bool bMobileEnableMovableSpotlightsShadow = (MobileEnableMovableSpotlightsShadowIniValue.Get((EShaderPlatform)Target.Platform) != 0);
-		Input.Environment.SetDefine(TEXT("PROJECT_MOBILE_ENABLE_MOVABLE_SPOTLIGHTS_SHADOW"), bMobileEnableMovableSpotlights && bMobileEnableMovableSpotlightsShadow ? 1 : 0);
+		const bool bMobileMovableSpotlightShadowsEnabled = IsMobileMovableSpotlightShadowsEnabled(Target.GetPlatform());
+		Input.Environment.SetDefine(TEXT("PROJECT_MOBILE_ENABLE_MOVABLE_SPOTLIGHT_SHADOWS"), bMobileMovableSpotlightShadowsEnabled ? 1 : 0);
 	}
 
 	const double TileSize = FLargeWorldRenderScalar::GetTileSize();
diff --git a/Engine/Source/Runtime/Engine/Private/ShaderCompiler/ShaderGenerationUtil.cpp b/Engine/Source/Runtime/Engine/Private/ShaderCompiler/ShaderGenerationUtil.cpp
--- a/Engine/Source/Runtime/Engine/Private/ShaderCompiler/ShaderGenerationUtil.cpp
+++ b/Engine/Source/Runtime/Engine/Private/ShaderCompiler/ShaderGenerationUtil.cpp
@@ -96,7 +96,6 @@ void FShaderCompileUtilities::ApplyFetchEnvironment(FShaderGlobalDefines& SrcDef
 	FETCH_COMPILE_BOOL(PLATFORM_SUPPORTS_RENDERTARGET_WRITE_MASK);
 	FETCH_COMPILE_BOOL(PLATFORM_SUPPORTS_PER_PIXEL_DBUFFER_MASK);
 	FETCH_COMPILE_BOOL(PLATFORM_SUPPORTS_DISTANCE_FIELDS);
-	FETCH_COMPILE_BOOL(PROJECT_MOBILE_ENABLE_MOVABLE_SPOTLIGHTS);
 	FETCH_COMPILE_BOOL(COMPILE_SHADERS_FOR_DEVELOPMENT_ALLOWED);
 
 	// note that we are doing an if so that if we call ApplyFetchEnvironment() twice, we get the logical OR of bSupportsDualBlending support
@@ -385,12 +384,6 @@ static FShaderGlobalDefines FetchShaderGlobalDefines(EShaderPlatform TargetPlatf
 	Ret.PLATFORM_SUPPORTS_RENDERTARGET_WRITE_MASK = RHISupportsRenderTargetWriteMask(EShaderPlatform(TargetPlatform)) ? 1 : 0;
 	Ret.PLATFORM_SUPPORTS_PER_PIXEL_DBUFFER_MASK = FDataDrivenShaderPlatformInfo::GetSupportsPerPixelDBufferMask(EShaderPlatform(TargetPlatform)) ? 1 : 0;
 	Ret.PLATFORM_SUPPORTS_DISTANCE_FIELDS = DoesPlatformSupportDistanceFields(EShaderPlatform(TargetPlatform)) ? 1 : 0;
-	
-	if (bIsMobilePlatform)
-	{
-		static IConsoleVariable* CVar = IConsoleManager::Get().FindConsoleVariable(TEXT("r.Mobile.EnableMovableSpotlights"));
-		Ret.PROJECT_MOBILE_ENABLE_MOVABLE_SPOTLIGHTS = CVar ? (CVar->GetInt() != 0) : 0;
-	}
 
 	Ret.bSupportsDualBlending = RHISupportsDualSourceBlending(TargetPlatform);
 
diff --git a/Engine/Source/Runtime/Engine/Public/LocalLightSceneProxy.h b/Engine/Source/Runtime/Engine/Public/LocalLightSceneProxy.h
--- a/Engine/Source/Runtime/Engine/Public/LocalLightSceneProxy.h
+++ b/Engine/Source/Runtime/Engine/Public/LocalLightSceneProxy.h
@@ -53,6 +53,8 @@ public:
 
 	virtual bool GetPerObjectProjectedShadowInitializer(const FBoxSphereBounds& SubjectBounds, class FPerObjectProjectedShadowInitializer& OutInitializer) const override;
 
+	virtual bool IsLocalLight() const override;
+
 protected:
 
 	/** Updates the light scene info's radius from the component. */
diff --git a/Engine/Source/Runtime/Engine/Public/SceneManagement.h b/Engine/Source/Runtime/Engine/Public/SceneManagement.h
--- a/Engine/Source/Runtime/Engine/Public/SceneManagement.h
+++ b/Engine/Source/Runtime/Engine/Public/SceneManagement.h
@@ -1277,74 +1277,16 @@ BEGIN_SHADER_PARAMETER_STRUCT(FLightShaderParameters, ENGINE_API)
 END_SHADER_PARAMETER_STRUCT()
 
 
-// Movable point light uniform buffer for mobile
-BEGIN_GLOBAL_SHADER_PARAMETER_STRUCT(FMobileMovablePointLightUniformShaderParameters,ENGINE_API)
-	SHADER_PARAMETER(FVector4f, LightPositionAndInvRadius)
-	SHADER_PARAMETER(FVector4f, LightTilePosition) // w unused
-	SHADER_PARAMETER(FVector4f, LightColorAndFalloffExponent)
-	SHADER_PARAMETER(FVector4f, SpotLightDirectionAndSpecularScale)
-	SHADER_PARAMETER(FVector4f, SpotLightAnglesAndSoftTransitionScaleAndLightShadowType) //xy SpotAngles, z SoftTransitionScale, w LightShadowType if (w&1 == 1) is pointlight, (w&2 == 2) is spotlight, (w&4 == 4) is with shadow
-	SHADER_PARAMETER(FVector4f, SpotLightShadowSharpenAndShadowFadeFraction) // x ShadowSharpen, y ShadowFadFraction
+// Movable local light shadow parameters for mobile deferred
+BEGIN_SHADER_PARAMETER_STRUCT(FMobileMovableLocalLightShadowParameters,ENGINE_API)
+	SHADER_PARAMETER(FVector4f, SpotLightShadowSharpenAndFadeFractionAndReceiverDepthBiasAndSoftTransitionScale) // x ShadowSharpen, y ShadowFadFraction, z ReceiverDepthBias, w SoftTransitionScale
 	SHADER_PARAMETER(FVector4f, SpotLightShadowmapMinMax)
 	SHADER_PARAMETER(FMatrix44f, SpotLightShadowWorldToShadowMatrix)
+	SHADER_PARAMETER(FVector4f, LocalLightShadowBufferSize)
+	SHADER_PARAMETER_TEXTURE(Texture2D, LocalLightShadowTexture)
+	SHADER_PARAMETER_SAMPLER(SamplerState, LocalLightShadowSampler)
 END_GLOBAL_SHADER_PARAMETER_STRUCT()
 
-/** Initializes the movable point light uniform shader parameters. */
-FORCEINLINE FMobileMovablePointLightUniformShaderParameters GetMovablePointLightUniformShaderParameters(
-	const FVector4f& LightPositionAndInvRadius,
-	const FVector4f& LightTilePosition,
-	const FVector4f& LightColorAndFalloffExponent,
-	const FVector4f& SpotLightDirectionAndSpecularScale,
-	const FVector4f& SpotLightAnglesAndSoftTransitionScaleAndLightShadowType,
-	const FVector4f& SpotLightShadowSharpenAndShadowFadeFraction,
-	const FVector4f& SpotLightShadowmapMinMax,
-	const FMatrix44f& SpotLightShadowWorldToShadowMatrix
-)
-{
-	FMobileMovablePointLightUniformShaderParameters Result;
-	Result.LightPositionAndInvRadius = LightPositionAndInvRadius;
-	Result.LightTilePosition = LightTilePosition;
-	Result.LightColorAndFalloffExponent = LightColorAndFalloffExponent;
-	Result.SpotLightDirectionAndSpecularScale = SpotLightDirectionAndSpecularScale;
-	Result.SpotLightAnglesAndSoftTransitionScaleAndLightShadowType = SpotLightAnglesAndSoftTransitionScaleAndLightShadowType;
-	Result.SpotLightShadowSharpenAndShadowFadeFraction = SpotLightShadowSharpenAndShadowFadeFraction;
-	Result.SpotLightShadowmapMinMax = SpotLightShadowmapMinMax;
-	Result.SpotLightShadowWorldToShadowMatrix = SpotLightShadowWorldToShadowMatrix;
-
-	return Result;
-}
-
-FORCEINLINE FMobileMovablePointLightUniformShaderParameters GetDummyMovablePointLightUniformShaderParameters()
-{
-	return GetMovablePointLightUniformShaderParameters(
-		FVector4f(ForceInitToZero),
-		FVector4f(ForceInitToZero),
-		FVector4f(ForceInitToZero),
-		FVector4f(ForceInitToZero),
-		FVector4f(ForceInitToZero),
-		FVector4f(ForceInitToZero),
-		FVector4f(ForceInitToZero),
-		FMatrix44f(ForceInitToZero)
-	);
-}
-
-/**
- * Dummy mobile movable point light uniform buffer.
- */
-class FDummyMovablePointLightUniformBuffer : public TUniformBuffer<FMobileMovablePointLightUniformShaderParameters>
-{
-public:
-
-	/** Default constructor. */
-	FDummyMovablePointLightUniformBuffer()
-	{
-		SetContents(GetDummyMovablePointLightUniformShaderParameters());
-	}
-};
-
-/** Global primitive uniform buffer resource containing identity transformations. */
-extern ENGINE_API TGlobalResource<FDummyMovablePointLightUniformBuffer> GDummyMovablePointLightUniformBuffer;
-
 /**
  * Generic parameters used to render a light
  * Has a 1:1 mapping with FLightShaderParameters, but can also be used in other contexts
@@ -1441,6 +1383,7 @@ public:
 	virtual float GetSourceRadius() const { return 0.0f; }
 	virtual bool IsInverseSquared() const { return true; }
 	virtual bool IsRectLight() const { return false; }
+	virtual bool  IsLocalLight() const { return false; }
 	virtual bool HasSourceTexture() const { return false; }
 	virtual float GetLightSourceAngle() const { return 0.0f; }
 	virtual float GetShadowSourceAngleFactor() const { return 1.0f; }
@@ -1662,16 +1605,6 @@ public:
 	virtual FLinearColor GetCloudScatteredLuminanceScale() const { return FLinearColor::White; }
 	virtual bool GetUsePerPixelAtmosphereTransmittance() const { return false; }
 
-	FORCEINLINE void SetMobileMovablePointLightUniformBufferNeedsUpdate(bool bInMobileMovablePointLightUniformBufferNeedsUpdate)
-	{
-		bMobileMovablePointLightUniformBufferNeedsUpdate = bInMobileMovablePointLightUniformBufferNeedsUpdate;
-	}
-
-	FORCEINLINE FRHIUniformBuffer* GetMobileMovablePointLightUniformBufferRHI() const
-	{
-		return MobileMovablePointLightUniformBuffer.GetReference();
-	}
-
 protected:
 
 	friend class FScene;
@@ -1853,21 +1786,6 @@ protected:
 	/** Deep shadow layer distribution. */
 	float DeepShadowLayerDistribution;
 
-	/** If this is TRUE, the light's mobile movable point light uniform buffer needs to be updated before it can be used for mobile base pass rendering. */
-	bool bMobileMovablePointLightUniformBufferNeedsUpdate;
-
-	/** Cached ShouldBeRender for mobile, since if the ShouldBeRender is changed we have to update the movable point lights uniform buffer. */
-	bool bMobileMovablePointLightShouldBeRender;
-
-	/** Cached DynamicShadows show flag for mobile, since if the show flag is changed we have to update the movable point lights uniform buffer. */
-	bool bMobileMovablePointLightShouldCastShadow;
-
-	/** Cached the spotlight shadow map min and max value for mobile, since if the value is changed we have to update the movable point lights uniform buffer. */
-	FVector4f MobileMovablePointLightShadowmapMinMax;
-
-	/** The movable point light's uniform buffer for mobile. */
-	TUniformBufferRef<FMobileMovablePointLightUniformShaderParameters> MobileMovablePointLightUniformBuffer;
-
 	/**
 	 * Updates the light proxy's cached transforms.
 	 * @param InLightToWorld - The new light-to-world transform.
@@ -3356,10 +3274,7 @@ struct FReadOnlyCVARCache
 	bool bMobileAllowMovableDirectionalLights;
 	bool bMobileAllowDistanceFieldShadows;
 	bool bMobileEnableStaticAndCSMShadowReceivers;
-	int32 NumMobileMovablePointLights;
 	int32 MobileSkyLightPermutation;
-	bool bMobileEnableMovableSpotlights;
-	bool bMobileEnableMovableSpotlightsShadow;
 	bool bMobileEnableNoPrecomputedLightingCSMShader;
 	
 	bool bInitialized;
diff --git a/Engine/Source/Runtime/MaterialShaderQualitySettings/Classes/ShaderPlatformQualitySettings.h b/Engine/Source/Runtime/MaterialShaderQualitySettings/Classes/ShaderPlatformQualitySettings.h
--- a/Engine/Source/Runtime/MaterialShaderQualitySettings/Classes/ShaderPlatformQualitySettings.h
+++ b/Engine/Source/Runtime/MaterialShaderQualitySettings/Classes/ShaderPlatformQualitySettings.h
@@ -38,7 +38,6 @@ public:
 		, bForceFullyRough(false)
 		, bForceNonMetal(false)
 		, bForceDisableLMDirectionality(false)
-		, bForceLQReflections(false)
 		, bForceDisablePreintegratedGF(false)
 		, bDisableMaterialNormalCalculation(false)
 		, MobileShadowQuality(EMobileShadowQuality::PCF_3x3)
@@ -60,9 +59,6 @@ public:
 	UPROPERTY(EditAnywhere, Config, Meta = (DisplayName = "Disable Lightmap directionality"), Category = "Quality")
 	bool bForceDisableLMDirectionality;
 
-	UPROPERTY(EditAnywhere, Config, Meta = (DisplayName = "Force low quality reflections"), Category = "Quality")
-	bool bForceLQReflections;
-
 	UPROPERTY(EditAnywhere, Config, Meta = (DisplayName = "Force not use preintegrated GF for simple IBL"), Category = "Quality")
 	bool bForceDisablePreintegratedGF;
 
diff --git a/Engine/Source/Runtime/MaterialShaderQualitySettings/Private/MaterialShaderQualitySettings.cpp b/Engine/Source/Runtime/MaterialShaderQualitySettings/Private/MaterialShaderQualitySettings.cpp
--- a/Engine/Source/Runtime/MaterialShaderQualitySettings/Private/MaterialShaderQualitySettings.cpp
+++ b/Engine/Source/Runtime/MaterialShaderQualitySettings/Private/MaterialShaderQualitySettings.cpp
@@ -160,7 +160,6 @@ bool FMaterialQualityOverrides::HasAnyOverridesSet() const
 		|| bForceDisableLMDirectionality != DefaultOverrides.bForceDisableLMDirectionality
 		|| bForceFullyRough != DefaultOverrides.bForceFullyRough
 		|| bForceNonMetal != DefaultOverrides.bForceNonMetal
-		|| bForceLQReflections != DefaultOverrides.bForceLQReflections
 		|| bForceDisablePreintegratedGF != DefaultOverrides.bForceDisablePreintegratedGF
 		|| bDisableMaterialNormalCalculation != DefaultOverrides.bDisableMaterialNormalCalculation
 		|| bDiscardQualityDuringCook != DefaultOverrides.bDiscardQualityDuringCook;
diff --git a/Engine/Source/Runtime/RenderCore/Private/RenderUtils.cpp b/Engine/Source/Runtime/RenderCore/Private/RenderUtils.cpp
--- a/Engine/Source/Runtime/RenderCore/Private/RenderUtils.cpp
+++ b/Engine/Source/Runtime/RenderCore/Private/RenderUtils.cpp
@@ -1152,8 +1152,8 @@ RENDERCORE_API bool MobileRequiresSceneDepthAux(const FStaticShaderPlatform Plat
 RENDERCORE_API bool SupportsTextureCubeArray(ERHIFeatureLevel::Type FeatureLevel)
 {
 	return FeatureLevel >= ERHIFeatureLevel::SM5 
-		// mobile deferred requries ES3.2 feature set
-		|| IsMobileDeferredShadingEnabled(GMaxRHIShaderPlatform);
+		// requries ES3.2 feature set
+		|| MobileEnableClusteredReflections(GetFeatureLevelShaderPlatform(FeatureLevel));
 }
 
 RENDERCORE_API bool MaskedInEarlyPass(const FStaticShaderPlatform Platform)
@@ -1203,6 +1203,31 @@ RENDERCORE_API bool IsMobileDistanceFieldEnabled(const FStaticShaderPlatform Pla
 	return IsMobilePlatform(Platform) && (FDataDrivenShaderPlatformInfo::GetSupportsMobileDistanceField(Platform)/* || IsD3DPlatform(Platform)*/) && IsUsingDistanceFields(Platform);
 }
 
+RENDERCORE_API bool IsMobileMovableSpotlightShadowsEnabled(const FStaticShaderPlatform Platform)
+{
+	static FShaderPlatformCachedIniValue<bool> MobileMovableSpotlightShadowsEnabledIniValue(TEXT("r.Mobile.EnableMovableSpotlightsShadow"));
+	return MobileMovableSpotlightShadowsEnabledIniValue.Get(Platform);
+}
+
+RENDERCORE_API bool MobileForwardEnableLocalLights(const FStaticShaderPlatform Platform)
+{
+	static FShaderPlatformCachedIniValue<bool> MobileForwardEnableLocalLightsIniValue(TEXT("r.Mobile.Forward.EnableLocalLights"));
+	return MobileForwardEnableLocalLightsIniValue.Get(Platform);
+}
+
+RENDERCORE_API bool MobileEnableClusteredReflections(const FStaticShaderPlatform Platform)
+{
+	static FShaderPlatformCachedIniValue<bool> MobileForwardEnableClusteredReflectionsIniValue(TEXT("r.Mobile.Forward.EnableClusteredReflections"));
+	return IsMobileDeferredShadingEnabled(Platform) || MobileForwardEnableClusteredReflectionsIniValue.Get(Platform);
+}
+
+RENDERCORE_API bool MobileUsesShadowMaskTexture(const FStaticShaderPlatform Platform)
+{
+	// Only distance field shadow needs to render shadow mask texture on mobile deferred, normal shadows need to be rendered separately because of handling lighting channels.
+	// Besides distance field shadow, with clustered lighting and shadow of local light enabled, shadows will render to shadow mask texture on mobile forward, lighting channels are handled in base pass shader.
+	return IsMobileDistanceFieldEnabled(Platform) || (!IsMobileDeferredShadingEnabled(Platform) && IsMobileMovableSpotlightShadowsEnabled(Platform) && MobileForwardEnableLocalLights(Platform));
+}
+
 RENDERCORE_API bool MobileBasePassAlwaysUsesCSM(const FStaticShaderPlatform Platform)
 {
 	static TConsoleVariableData<int32>* CVar = IConsoleManager::Get().FindTConsoleVariableDataInt(TEXT("r.Mobile.Shadow.CSMShaderCullingMethod"));
diff --git a/Engine/Source/Runtime/RenderCore/Private/Shader.cpp b/Engine/Source/Runtime/RenderCore/Private/Shader.cpp
--- a/Engine/Source/Runtime/RenderCore/Private/Shader.cpp
+++ b/Engine/Source/Runtime/RenderCore/Private/Shader.cpp
@@ -1540,18 +1540,8 @@ void ShaderMapAppendKeyString(EShaderPlatform Platform, FString& KeyString)
 		KeyString += IsUsingEmulatedUniformBuffers(Platform) ? TEXT("_NoUB") : TEXT("");
 
 		{
-			static FShaderPlatformCachedIniValue<bool> MobileEnableMovableSpotlightsIniValue(TEXT("r.Mobile.EnableMovableSpotlights"));
-			static FShaderPlatformCachedIniValue<bool> MobileEnableMovableSpotlightsShadowIniValue(TEXT("r.Mobile.EnableMovableSpotlightsShadow"));
-			static FShaderPlatformCachedIniValue<int32> MobileNumDynamicPointLightsIniValue(TEXT("r.MobileNumDynamicPointLights"));
-
-			bool bMobileEnableMovableSpotlights = (MobileEnableMovableSpotlightsIniValue.Get(Platform) != 0);
-			KeyString += (bMobileEnableMovableSpotlights) ? TEXT("_MSPTL") : TEXT("");
-
-			bool bMobileEnableMovableSpotlightsShadow = (MobileEnableMovableSpotlightsShadowIniValue.Get(Platform) != 0);
-			KeyString += (bMobileEnableMovableSpotlights && bMobileEnableMovableSpotlightsShadow) ? TEXT("S") : TEXT("");
-
-			const int32 MobileNumDynamicPointLights = MobileNumDynamicPointLightsIniValue.Get(Platform);
-			KeyString += (MobileNumDynamicPointLights != 0)? FString::Printf(TEXT("_MobDynPL_%d"), MobileNumDynamicPointLights): TEXT("");
+			const bool bMobileMovableSpotlightShadowsEnabled = IsMobileMovableSpotlightShadowsEnabled(Platform);
+			KeyString += bMobileMovableSpotlightShadowsEnabled ? TEXT("S") : TEXT("");
 		}
 		
 		{
@@ -1572,8 +1562,22 @@ void ShaderMapAppendKeyString(EShaderPlatform Platform, FString& KeyString)
 		}
 
 		{
-			static IConsoleVariable* CVar = IConsoleManager::Get().FindConsoleVariable(TEXT("r.Mobile.ShadingPath"));
-			KeyString += (CVar && CVar->GetInt() != 0) ? TEXT("_MobDSh") : TEXT("");
+			static IConsoleVariable* MobileShadingPathCVar = IConsoleManager::Get().FindConsoleVariable(TEXT("r.Mobile.ShadingPath"));			
+			if (MobileShadingPathCVar)
+			{
+				if (MobileShadingPathCVar->GetInt() != 0)
+				{
+					KeyString += TEXT("_MobDSh");
+				}
+				else
+				{
+					static IConsoleVariable* MobileForwardEnableClusteredReflectionsCVAR = IConsoleManager::Get().FindConsoleVariable(TEXT("r.Mobile.Forward.EnableClusteredReflections"));
+					if (MobileForwardEnableClusteredReflectionsCVAR && MobileForwardEnableClusteredReflectionsCVAR->GetInt() != 0)
+					{
+						KeyString += TEXT("_MobFCR");
+					}
+				}
+			}
 		}
 
 		{
diff --git a/Engine/Source/Runtime/RenderCore/Public/RenderUtils.h b/Engine/Source/Runtime/RenderCore/Public/RenderUtils.h
--- a/Engine/Source/Runtime/RenderCore/Public/RenderUtils.h
+++ b/Engine/Source/Runtime/RenderCore/Public/RenderUtils.h
@@ -503,6 +503,14 @@ RENDERCORE_API bool UseMobileAmbientOcclusion(const FStaticShaderPlatform Platfo
 
 RENDERCORE_API bool IsMobileDistanceFieldEnabled(const FStaticShaderPlatform Platform);
 
+RENDERCORE_API bool IsMobileMovableSpotlightShadowsEnabled(const FStaticShaderPlatform Platform);
+
+RENDERCORE_API bool MobileForwardEnableLocalLights(const FStaticShaderPlatform Platform);
+
+RENDERCORE_API bool MobileEnableClusteredReflections(const FStaticShaderPlatform Platform);
+
+RENDERCORE_API bool MobileUsesShadowMaskTexture(const FStaticShaderPlatform Platform);
+
 RENDERCORE_API bool MobileBasePassAlwaysUsesCSM(const FStaticShaderPlatform Platform);
 
 RENDERCORE_API bool SupportsGen4TAA(const FStaticShaderPlatform Platform);
diff --git a/Engine/Source/Runtime/RenderCore/Public/ShaderMaterial.h b/Engine/Source/Runtime/RenderCore/Public/ShaderMaterial.h
--- a/Engine/Source/Runtime/RenderCore/Public/ShaderMaterial.h
+++ b/Engine/Source/Runtime/RenderCore/Public/ShaderMaterial.h
@@ -37,7 +37,6 @@ struct FShaderGlobalDefines
 	bool PLATFORM_SUPPORTS_RENDERTARGET_WRITE_MASK;
 	bool PLATFORM_SUPPORTS_PER_PIXEL_DBUFFER_MASK;
 	bool PLATFORM_SUPPORTS_DISTANCE_FIELDS;
-	bool PROJECT_MOBILE_ENABLE_MOVABLE_SPOTLIGHTS;
 
 	// This is a sepcial one. The flag COMPILE_SHADERS_FOR_DEVELOPMENT is set by the cvar r.CompileShadersForDevelopment, but only
 	// if bAllowDevelopmentShaderCompile is true in the call to GlobalBeginCompileShader(). So the flag COMPILE_SHADERS_FOR_DEVELOPMENT_ALLOWED
diff --git a/Engine/Source/Runtime/Renderer/Private/EditorPrimitivesRendering.cpp b/Engine/Source/Runtime/Renderer/Private/EditorPrimitivesRendering.cpp
--- a/Engine/Source/Runtime/Renderer/Private/EditorPrimitivesRendering.cpp
+++ b/Engine/Source/Runtime/Renderer/Private/EditorPrimitivesRendering.cpp
@@ -125,7 +125,7 @@ bool FEditorPrimitivesBasePassMeshProcessor::ProcessMobileShadingPath(const FMes
 	typedef FUniformLightMapPolicy LightMapPolicyType;
 
 	const FVertexFactory* VertexFactory = MeshBatch.VertexFactory;
-	const int32 NumMovablePointLights = 0;
+	const bool bEnableLocalLights = false;
 	const bool bEnableSkyLight = false;
 
 	TMeshProcessorShaders<
@@ -133,7 +133,7 @@ bool FEditorPrimitivesBasePassMeshProcessor::ProcessMobileShadingPath(const FMes
 		TMobileBasePassPSPolicyParamType<FUniformLightMapPolicy>> BasePassShaders;
 	if (!MobileBasePass::GetShaders(
 		NoLightmapPolicy.GetIndirectPolicy(),
-		NumMovablePointLights,
+		bEnableLocalLights,
 		Material,
 		VertexFactory->GetType(),
 		bEnableSkyLight,
diff --git a/Engine/Source/Runtime/Renderer/Private/LightGridInjection.cpp b/Engine/Source/Runtime/Renderer/Private/LightGridInjection.cpp
--- a/Engine/Source/Runtime/Renderer/Private/LightGridInjection.cpp
+++ b/Engine/Source/Runtime/Renderer/Private/LightGridInjection.cpp
@@ -29,9 +29,6 @@ LightGridInjection.cpp
 #include "Components/LightComponent.h"
 #include "Engine/MapBuildDataRegistry.h"
 
-// Workaround for platforms that don't support implicit conversion from 16bit integers on the CPU to uint32 in the shader
-#define	CHANGE_LIGHTINDEXTYPE_SIZE	(PLATFORM_MAC || PLATFORM_IOS) 
-
 int32 GLightGridPixelSize = 64;
 FAutoConsoleVariableRef CVarLightGridPixelSize(
 	TEXT("r.Forward.LightGridPixelSize"),
@@ -172,7 +169,7 @@ public:
 
 	static bool ShouldCompilePermutation(const FGlobalShaderPermutationParameters& Parameters)
 	{
-		return IsFeatureLevelSupported(Parameters.Platform, ERHIFeatureLevel::SM5) || IsMobileDeferredShadingEnabled(Parameters.Platform);
+		return true;
 	}
 
 	static void ModifyCompilationEnvironment(const FGlobalShaderPermutationParameters& Parameters, FShaderCompilerEnvironment& OutEnvironment)
@@ -209,7 +206,7 @@ public:
 
 	static bool ShouldCompilePermutation(const FGlobalShaderPermutationParameters& Parameters)
 	{
-		return IsFeatureLevelSupported(Parameters.Platform, ERHIFeatureLevel::SM5) || IsMobileDeferredShadingEnabled(Parameters.Platform);
+		return true;
 	}
 
 	static void ModifyCompilationEnvironment(const FGlobalShaderPermutationParameters& Parameters, FShaderCompilerEnvironment& OutEnvironment)
@@ -616,7 +613,7 @@ void FSceneRenderer::ComputeLightGrid(FRDGBuilder& GraphBuilder, bool bCullLight
 		FVector ZParams = GetLightGridZParams(View.NearClippingDistance, FarPlane + 10.f);
 		ForwardLightData->LightGridZParams = (FVector3f)ZParams;
 
-		const uint64 NumIndexableLights = CHANGE_LIGHTINDEXTYPE_SIZE && !bAllowFormatConversion ? (1llu << (sizeof(FLightIndexType32) * 8llu)) : (1llu << (sizeof(FLightIndexType) * 8llu));
+		const uint64 NumIndexableLights = !bAllowFormatConversion ? (1llu << (sizeof(FLightIndexType32) * 8llu)) : (1llu << (sizeof(FLightIndexType) * 8llu));
 
 		if ((uint64)ForwardLocalLightData.Num() > NumIndexableLights)
 		{
@@ -630,7 +627,7 @@ void FSceneRenderer::ComputeLightGrid(FRDGBuilder& GraphBuilder, bool bCullLight
 		}
 
 #if ENABLE_LIGHT_CULLING_VIEW_SPACE_BUILD_DATA
-		const SIZE_T LightIndexTypeSize = CHANGE_LIGHTINDEXTYPE_SIZE && !bAllowFormatConversion ? sizeof(FLightIndexType32) : sizeof(FLightIndexType);
+		const SIZE_T LightIndexTypeSize = !bAllowFormatConversion ? sizeof(FLightIndexType32) : sizeof(FLightIndexType);
 		// Fuse these loops as I see no reason why not and we build some temporary data that is needed in the build pass and is 
 		// not needed to be stored permanently.
 #else // !ENABLE_LIGHT_CULLING_VIEW_SPACE_BUILD_DATA
@@ -638,7 +635,7 @@ void FSceneRenderer::ComputeLightGrid(FRDGBuilder& GraphBuilder, bool bCullLight
 		ForwardLightDataPerView.Emplace(ForwardLightData);
 	}
 
-	const SIZE_T LightIndexTypeSize = CHANGE_LIGHTINDEXTYPE_SIZE && !bAllowFormatConversion ? sizeof(FLightIndexType32) : sizeof(FLightIndexType);
+	const SIZE_T LightIndexTypeSize = !bAllowFormatConversion ? sizeof(FLightIndexType32) : sizeof(FLightIndexType);
 
 	for (int32 ViewIndex = 0; ViewIndex < Views.Num(); ViewIndex++)
 	{
@@ -718,6 +715,7 @@ void FSceneRenderer::ComputeLightGrid(FRDGBuilder& GraphBuilder, bool bCullLight
 				AddClearUAVPass(GraphBuilder, PassParameters->RWStartOffsetGrid, 0xFFFFFFFF);
 				AddClearUAVPass(GraphBuilder, PassParameters->RWNextCulledLightLink, 0);
 				AddClearUAVPass(GraphBuilder, NextCulledLightDataUAV, 0);
+				AddClearUAVPass(GraphBuilder, NumCulledLightsGridUAV, 0);
 				FComputeShaderUtils::AddPass(GraphBuilder, RDG_EVENT_NAME("LightGridInject:LinkedList"), ComputeShader, PassParameters, NumGroups);
 
 				{
@@ -900,5 +898,3 @@ void FDeferredShadingSceneRenderer::RenderForwardShadowProjections(
 		}
 	}
 }
-
-#undef CHANGE_LIGHTINDEXTYPE_SIZE
diff --git a/Engine/Source/Runtime/Renderer/Private/LightRendering.cpp b/Engine/Source/Runtime/Renderer/Private/LightRendering.cpp
--- a/Engine/Source/Runtime/Renderer/Private/LightRendering.cpp
+++ b/Engine/Source/Runtime/Renderer/Private/LightRendering.cpp
@@ -773,11 +773,14 @@ static bool LightRequiresDenosier(const FLightSceneInfo& LightSceneInfo)
 	return false;
 }
 
-
+bool FSceneRenderer::AllowSimpleLights() const
+{
+	return bAllowSimpleLights == 1;
+}
 
 void FSceneRenderer::GatherAndSortLights(FSortedLightSetSceneInfo& OutSortedLights, bool bShadowedLightsInClustered)
 {
-	if (bAllowSimpleLights)
+	if (AllowSimpleLights())
 	{
 		GatherSimpleLights(ViewFamily, Views, OutSortedLights.SimpleLights);
 	}
diff --git a/Engine/Source/Runtime/Renderer/Private/LightSceneInfo.cpp b/Engine/Source/Runtime/Renderer/Private/LightSceneInfo.cpp
--- a/Engine/Source/Runtime/Renderer/Private/LightSceneInfo.cpp
+++ b/Engine/Source/Runtime/Renderer/Private/LightSceneInfo.cpp
@@ -10,6 +10,7 @@
 #include "ScenePrivate.h"
 #include "DistanceFieldLightingShared.h"
 #include "Misc/LargeWorldRenderPosition.h"
+#include "LocalLightSceneProxy.h"
 
 int32 GWholeSceneShadowUnbuiltInteractionThreshold = 500;
 static FAutoConsoleVariableRef CVarWholeSceneShadowUnbuiltInteractionThreshold(
@@ -79,10 +80,8 @@ void FLightSceneInfo::AddToScene()
 	bool bIsValidLightTypeMobile = false;
 	if (Scene->GetShadingPath() == EShadingPath::Mobile && Proxy->IsMovable())
 	{
-		static const auto MobileEnableMovableSpotLightsVar = IConsoleManager::Get().FindTConsoleVariableDataInt(TEXT("r.Mobile.EnableMovableSpotLights"));
 		const uint8 LightType = Proxy->GetLightType();
-		bIsValidLightTypeMobile = LightType == LightType_Rect || LightType == LightType_Point
-			|| (LightType == LightType_Spot && MobileEnableMovableSpotLightsVar->GetValueOnRenderThread());
+		bIsValidLightTypeMobile = LightType == LightType_Rect || LightType == LightType_Point || LightType == LightType_Spot;
 	}
 
 	// Only need to create light interactions for lights that can cast a shadow,
@@ -118,12 +117,6 @@ void FLightSceneInfo::AddToScene()
 			{
 				CreateLightPrimitiveInteraction(LightSceneInfoCompact, PrimitiveSceneInfoCompact);
 			});
-
-			if (bIsValidLightTypeMobile)
-			{
-				Proxy->MobileMovablePointLightUniformBuffer = TUniformBufferRef<FMobileMovablePointLightUniformShaderParameters>::CreateUniformBufferImmediate(GetDummyMovablePointLightUniformShaderParameters(), UniformBuffer_MultiFrame);
-				Proxy->bMobileMovablePointLightUniformBufferNeedsUpdate = true;
-			}
 		}
 	}
 }
@@ -244,147 +237,46 @@ FLightPrimitiveInteraction* FLightSceneInfo::GetDynamicInteractionStaticPrimitiv
 	return DynamicInteractionStaticPrimitiveList;
 }
 
-void FLightSceneInfo::ConditionalUpdateMobileMovablePointLightUniformBuffer(const FSceneRenderer* SceneRenderer)
+bool FLightSceneInfo::SetupMobileMovableLocalLightShadowParameters(const FViewInfo& View, const TArray<FVisibleLightInfo, SceneRenderingAllocator>& VisibleLightInfos, FMobileMovableLocalLightShadowParameters& MobileMovableLocalLightShadowParameters) const
 {
-	QUICK_SCOPE_CYCLE_COUNTER(STAT_FLightSceneProxy_UpdateMobileMovablePointLightUniformBuffer);
-
-	static auto* MobileNumDynamicPointLightsCVar = IConsoleManager::Get().FindTConsoleVariableDataInt(TEXT("r.MobileNumDynamicPointLights"));
-	const int32 MobileNumDynamicPointLights = MobileNumDynamicPointLightsCVar->GetValueOnRenderThread();
-
-	static auto* MobileEnableMovableSpotlightsCVar = IConsoleManager::Get().FindTConsoleVariableDataInt(TEXT("r.Mobile.EnableMovableSpotlights"));
-	const int32 MobileEnableMovableSpotlights = MobileEnableMovableSpotlightsCVar->GetValueOnRenderThread();
-
-	static auto* EnableMovableSpotlightsShadowCVar = IConsoleManager::Get().FindTConsoleVariableDataInt(TEXT("r.Mobile.EnableMovableSpotlightsShadow"));
-	const int32 EnableMovableSpotlightsShadow = EnableMovableSpotlightsShadowCVar->GetValueOnRenderThread();
+	QUICK_SCOPE_CYCLE_COUNTER(STAT_FLightSceneProxy_SetupMobileMovableLocalLightShadowParameters);
 
-	checkSlow(MobileNumDynamicPointLights > 0 && SceneRenderer);
+	check(Proxy->IsLocalLight());
+	const uint8 LightType = Proxy->GetLightType();
 
-	FVector4f LightPositionAndInvRadius;
-	FVector4f LightColorAndFalloffExponent;
-	FVector4f SpotLightDirectionAndSpecularScale;
-	FVector4f SpotLightAnglesAndSoftTransitionScaleAndLightShadowType;
-	FVector4f SpotLightShadowSharpenAndShadowFadeFraction;
-	FVector4f SpotLightShadowmapMinMax;
-	FVector4f LightTilePosition;
-	FMatrix44f SpotLightWorldToShadowMatrix;
+	FLightRenderParameters LightParameters;
+	Proxy->GetLightShaderParameters(LightParameters);
 
-	bool bShouldBeRender = false;
-	bool bShouldCastShadow = false;
-
-	for (int32 ViewIndex = 0; ViewIndex < SceneRenderer->Views.Num(); ++ViewIndex)
-	{
-		bShouldBeRender = ShouldRenderLight(SceneRenderer->Views[ViewIndex]);
-
-		if (bShouldBeRender)
-		{
-			break;
-		}
-	}
+	bool bShouldCastShadow = View.Family->EngineShowFlags.DynamicShadows
+		&& !IsSimpleForwardShadingEnabled(View.GetShaderPlatform())
+		&& GetShadowQuality() > 0
+		&& LightType == LightType_Spot
+		&& VisibleLightInfos[Id].AllProjectedShadows.Num() > 0
+		&& VisibleLightInfos[Id].AllProjectedShadows.Last()->bAllocated;
 
-	if (bShouldBeRender)
+	if (bShouldCastShadow)
 	{
-		const uint8 LightType = Proxy->GetLightType();
-
-		const bool bIsValidLightType =
-			LightType == LightType_Point
-			|| LightType == LightType_Rect
-			|| (LightType == LightType_Spot && MobileEnableMovableSpotlights);
-
-		checkSlow(bIsValidLightType && Proxy->IsMovable());
-
-		FLightRenderParameters LightParameters;
-		Proxy->GetLightShaderParameters(LightParameters);
-
-		float LightFadeFactor = 0.0f;
-
-		for (int32 ViewIndex = 0; ViewIndex < SceneRenderer->Views.Num(); ++ViewIndex)
-		{
-			LightFadeFactor = FMath::Max(LightFadeFactor, GetLightFadeFactor(SceneRenderer->Views[ViewIndex], Proxy));
-		}
-
-		LightParameters.Color *= LightFadeFactor;
+		FProjectedShadowInfo* ProjectedShadowInfo = VisibleLightInfos[Id].AllProjectedShadows.Last();
+		checkSlow(ProjectedShadowInfo && ProjectedShadowInfo->CacheMode != SDCM_StaticPrimitivesOnly);
 
-		if (Proxy->IsInverseSquared())
-		{
-			LightParameters.FalloffExponent = 0;
-		}
-
-		uint32 LightShadowType = LightType == LightType_Spot ? 2 : 1;
-
-		bShouldCastShadow = SceneRenderer->ViewFamily.EngineShowFlags.DynamicShadows
-			&& !IsSimpleForwardShadingEnabled(SceneRenderer->ShaderPlatform)
-			&& GetShadowQuality() > 0
-			&& EnableMovableSpotlightsShadow != 0
-			&& LightType == LightType_Spot
-			&& SceneRenderer->VisibleLightInfos[Id].AllProjectedShadows.Num() > 0
-			&& SceneRenderer->VisibleLightInfos[Id].AllProjectedShadows.Last()->bAllocated;
-
-		LightShadowType |= bShouldCastShadow ? 4 : 0;
-
-		float SoftTransitionScale = 0.0f;
-		float ShadowFadeFraction = 0.0f;
-
-		if (bShouldCastShadow)
-		{
-			FProjectedShadowInfo* ProjectedShadowInfo = SceneRenderer->VisibleLightInfos[Id].AllProjectedShadows.Last();
-			checkSlow(ProjectedShadowInfo && ProjectedShadowInfo->CacheMode != SDCM_StaticPrimitivesOnly);
-
-			const float TransitionSize = ProjectedShadowInfo->ComputeTransitionSize();
-			checkSlow(TransitionSize > 0.0f);
-
-			SoftTransitionScale = 1.0f / TransitionSize;
-
-			for (int32 ViewIndex = 0; ViewIndex < SceneRenderer->Views.Num(); ++ViewIndex)
-			{
-				ShadowFadeFraction = FMath::Max(ShadowFadeFraction, ProjectedShadowInfo->FadeAlphas[ViewIndex]);
-			}
+		const float TransitionSize = ProjectedShadowInfo->ComputeTransitionSize();
+		checkSlow(TransitionSize > 0.0f);
 
-			SpotLightShadowSharpenAndShadowFadeFraction = FVector4f(Proxy->GetShadowSharpen() * 7.0f + 1.0f, ShadowFadeFraction, ProjectedShadowInfo->GetShaderReceiverDepthBias(), 0.0f);
+		float SoftTransitionScale = 1.0f / TransitionSize;
+		float ShadowFadeFraction = ProjectedShadowInfo->FadeAlphas[0];
 
-			const FMatrix WorldToShadowMatrix = ProjectedShadowInfo->GetWorldToShadowMatrix(SpotLightShadowmapMinMax);
-			SpotLightWorldToShadowMatrix = FMatrix44f(FTranslationMatrix(LightParameters.WorldPosition) * WorldToShadowMatrix);
-		}
+		MobileMovableLocalLightShadowParameters.SpotLightShadowSharpenAndFadeFractionAndReceiverDepthBiasAndSoftTransitionScale = FVector4f(Proxy->GetShadowSharpen() * 7.0f + 1.0f, ShadowFadeFraction, ProjectedShadowInfo->GetShaderReceiverDepthBias(), SoftTransitionScale);
 
-		const FLargeWorldRenderPosition AbsoluteWorldPosition(LightParameters.WorldPosition);
+		const FMatrix WorldToShadowMatrix = ProjectedShadowInfo->GetWorldToShadowMatrix(MobileMovableLocalLightShadowParameters.SpotLightShadowmapMinMax);
+		MobileMovableLocalLightShadowParameters.SpotLightShadowWorldToShadowMatrix = FMatrix44f(FTranslationMatrix(LightParameters.WorldPosition) * WorldToShadowMatrix);
 
-		LightPositionAndInvRadius = FVector4f(AbsoluteWorldPosition.GetOffset(), LightParameters.InvRadius);
-		LightTilePosition = FVector4f(AbsoluteWorldPosition.GetTile(), 0);
-		LightColorAndFalloffExponent = FVector4f(LightParameters.Color, LightParameters.FalloffExponent);
-		SpotLightDirectionAndSpecularScale = FVector4f(LightParameters.Direction.X, LightParameters.Direction.Y, LightParameters.Direction.Z, Proxy->GetSpecularScale());
-		SpotLightAnglesAndSoftTransitionScaleAndLightShadowType = FVector4f(LightParameters.SpotAngles.X, LightParameters.SpotAngles.Y, SoftTransitionScale, LightShadowType);
+		const FIntPoint ShadowBufferResolution = ProjectedShadowInfo->GetShadowBufferResolution();
+		MobileMovableLocalLightShadowParameters.LocalLightShadowBufferSize = FVector4f(ShadowBufferResolution.X, ShadowBufferResolution.Y, 1.0f / ShadowBufferResolution.X, 1.0f / ShadowBufferResolution.Y);
+		MobileMovableLocalLightShadowParameters.LocalLightShadowTexture = ProjectedShadowInfo->RenderTargets.DepthTarget->GetRHI();
+		MobileMovableLocalLightShadowParameters.LocalLightShadowSampler = TStaticSamplerState<SF_Point, AM_Clamp, AM_Clamp, AM_Clamp>::GetRHI();
 	}
 
-	if (bShouldBeRender != Proxy->bMobileMovablePointLightShouldBeRender ||
-		bShouldCastShadow != Proxy->bMobileMovablePointLightShouldCastShadow ||
-		SpotLightShadowmapMinMax != Proxy->MobileMovablePointLightShadowmapMinMax)
-	{
-		Proxy->bMobileMovablePointLightUniformBufferNeedsUpdate = true;
-
-		Proxy->bMobileMovablePointLightShouldBeRender = bShouldBeRender;
-
-		Proxy->bMobileMovablePointLightShouldCastShadow = bShouldCastShadow;
-
-		Proxy->MobileMovablePointLightShadowmapMinMax = SpotLightShadowmapMinMax;
-	}
-
-	if (Proxy->bMobileMovablePointLightUniformBufferNeedsUpdate)
-	{
-		const FMobileMovablePointLightUniformShaderParameters MobileMovablePointLightUniformShaderParameters =
-			GetMovablePointLightUniformShaderParameters(
-				LightPositionAndInvRadius,
-				LightTilePosition,
-				LightColorAndFalloffExponent,
-				SpotLightDirectionAndSpecularScale,
-				SpotLightAnglesAndSoftTransitionScaleAndLightShadowType,
-				SpotLightShadowSharpenAndShadowFadeFraction,
-				SpotLightShadowmapMinMax,
-				SpotLightWorldToShadowMatrix
-			);
-
-		Proxy->MobileMovablePointLightUniformBuffer.UpdateUniformBufferImmediate(MobileMovablePointLightUniformShaderParameters);
-
-		Proxy->bMobileMovablePointLightUniformBufferNeedsUpdate = false;
-	}
+	return bShouldCastShadow;
 }
 
 bool FLightSceneInfo::ShouldRecordShadowSubjectsForMobile() const
diff --git a/Engine/Source/Runtime/Renderer/Private/LightSceneInfo.h b/Engine/Source/Runtime/Renderer/Private/LightSceneInfo.h
--- a/Engine/Source/Runtime/Renderer/Private/LightSceneInfo.h
+++ b/Engine/Source/Runtime/Renderer/Private/LightSceneInfo.h
@@ -13,6 +13,7 @@
 #include "Math/GenericOctree.h"
 #include "PrimitiveSceneProxy.h"
 #include "Templates/UniquePtr.h"
+#include "SceneRendering.h"
 
 class FLightPrimitiveInteraction;
 class FLightSceneInfo;
@@ -340,8 +341,7 @@ public:
 		return (uint32)LightSceneInfo->Id;
 	}
 
-	// Update the mobile movable point light uniform buffer before it is used for mobile base pass rendering.
-	void ConditionalUpdateMobileMovablePointLightUniformBuffer(const class FSceneRenderer* SceneRenderer);
+	bool SetupMobileMovableLocalLightShadowParameters(const FViewInfo& View, const TArray<FVisibleLightInfo, SceneRenderingAllocator>& VisibleLightInfos, FMobileMovableLocalLightShadowParameters& MobileMovableLocalLightShadowParameters) const;
 
 	bool ShouldRecordShadowSubjectsForMobile() const;
 };
diff --git a/Engine/Source/Runtime/Renderer/Private/MobileBasePass.cpp b/Engine/Source/Runtime/Renderer/Private/MobileBasePass.cpp
--- a/Engine/Source/Runtime/Renderer/Private/MobileBasePass.cpp
+++ b/Engine/Source/Runtime/Renderer/Private/MobileBasePass.cpp
@@ -30,7 +30,7 @@ bool MobileUsesNoLightMapPermutation(const FMeshMaterialShaderPermutationParamet
 	return true;
 }
 
-template <ELightMapPolicyType Policy, int32 NumMovablePointLights>
+template <ELightMapPolicyType Policy, bool bEnableLocalLights>
 bool GetUniformMobileBasePassShaders(
 	const FMaterial& Material, 
 	FVertexFactoryType* VertexFactoryType, 
@@ -49,11 +49,11 @@ bool GetUniformMobileBasePassShaders(
 
 		if (bEnableSkyLight)
 		{
-			ShaderTypes.AddShaderType<TMobileBasePassPS<TUniformLightMapPolicy<Policy>, HDR_LINEAR_64, true, NumMovablePointLights>>();
+			ShaderTypes.AddShaderType<TMobileBasePassPS<TUniformLightMapPolicy<Policy>, HDR_LINEAR_64, true, bEnableLocalLights>>();
 		}
 		else
 		{
-			ShaderTypes.AddShaderType<TMobileBasePassPS<TUniformLightMapPolicy<Policy>, HDR_LINEAR_64, false, NumMovablePointLights>>();
+			ShaderTypes.AddShaderType<TMobileBasePassPS<TUniformLightMapPolicy<Policy>, HDR_LINEAR_64, false, bEnableLocalLights>>();
 		}	
 	}
 	else
@@ -62,11 +62,11 @@ bool GetUniformMobileBasePassShaders(
 
 		if (bEnableSkyLight)
 		{
-			ShaderTypes.AddShaderType<TMobileBasePassPS<TUniformLightMapPolicy<Policy>, LDR_GAMMA_32, true, NumMovablePointLights>>();
+			ShaderTypes.AddShaderType<TMobileBasePassPS<TUniformLightMapPolicy<Policy>, LDR_GAMMA_32, true, bEnableLocalLights>>();
 		}
 		else
 		{
-			ShaderTypes.AddShaderType<TMobileBasePassPS<TUniformLightMapPolicy<Policy>, LDR_GAMMA_32, false, NumMovablePointLights>>();
+			ShaderTypes.AddShaderType<TMobileBasePassPS<TUniformLightMapPolicy<Policy>, LDR_GAMMA_32, false, bEnableLocalLights>>();
 		}			
 	}
 
@@ -81,7 +81,7 @@ bool GetUniformMobileBasePassShaders(
 	return true;
 }
 
-template <int32 NumMovablePointLights>
+template <bool bEnableLocalLights>
 bool GetMobileBasePassShaders(
 	ELightMapPolicyType LightMapPolicyType, 
 	const FMaterial& Material, 
@@ -94,25 +94,25 @@ bool GetMobileBasePassShaders(
 	switch (LightMapPolicyType)
 	{
 	case LMP_NO_LIGHTMAP:
-		return GetUniformMobileBasePassShaders<LMP_NO_LIGHTMAP, NumMovablePointLights>(Material, VertexFactoryType, bEnableSkyLight, VertexShader, PixelShader);
+		return GetUniformMobileBasePassShaders<LMP_NO_LIGHTMAP, bEnableLocalLights>(Material, VertexFactoryType, bEnableSkyLight, VertexShader, PixelShader);
 	case LMP_LQ_LIGHTMAP:
-		return GetUniformMobileBasePassShaders<LMP_LQ_LIGHTMAP, NumMovablePointLights>(Material, VertexFactoryType, bEnableSkyLight, VertexShader, PixelShader);
+		return GetUniformMobileBasePassShaders<LMP_LQ_LIGHTMAP, bEnableLocalLights>(Material, VertexFactoryType, bEnableSkyLight, VertexShader, PixelShader);
 	case LMP_MOBILE_DISTANCE_FIELD_SHADOWS_AND_LQ_LIGHTMAP:
-		return GetUniformMobileBasePassShaders<LMP_MOBILE_DISTANCE_FIELD_SHADOWS_AND_LQ_LIGHTMAP, NumMovablePointLights>(Material, VertexFactoryType, bEnableSkyLight, VertexShader, PixelShader);
+		return GetUniformMobileBasePassShaders<LMP_MOBILE_DISTANCE_FIELD_SHADOWS_AND_LQ_LIGHTMAP, bEnableLocalLights>(Material, VertexFactoryType, bEnableSkyLight, VertexShader, PixelShader);
 	case LMP_MOBILE_DISTANCE_FIELD_SHADOWS_LIGHTMAP_AND_CSM:
-		return GetUniformMobileBasePassShaders<LMP_MOBILE_DISTANCE_FIELD_SHADOWS_LIGHTMAP_AND_CSM, NumMovablePointLights>(Material, VertexFactoryType, bEnableSkyLight, VertexShader, PixelShader);
+		return GetUniformMobileBasePassShaders<LMP_MOBILE_DISTANCE_FIELD_SHADOWS_LIGHTMAP_AND_CSM, bEnableLocalLights>(Material, VertexFactoryType, bEnableSkyLight, VertexShader, PixelShader);
 	case LMP_MOBILE_DIRECTIONAL_LIGHT_CSM_AND_LIGHTMAP:
-		return GetUniformMobileBasePassShaders<LMP_MOBILE_DIRECTIONAL_LIGHT_CSM_AND_LIGHTMAP, NumMovablePointLights>(Material, VertexFactoryType, bEnableSkyLight, VertexShader, PixelShader);
+		return GetUniformMobileBasePassShaders<LMP_MOBILE_DIRECTIONAL_LIGHT_CSM_AND_LIGHTMAP, bEnableLocalLights>(Material, VertexFactoryType, bEnableSkyLight, VertexShader, PixelShader);
 	case LMP_MOBILE_DIRECTIONAL_LIGHT_AND_SH_INDIRECT:
-		return GetUniformMobileBasePassShaders<LMP_MOBILE_DIRECTIONAL_LIGHT_AND_SH_INDIRECT, NumMovablePointLights>(Material, VertexFactoryType, bEnableSkyLight, VertexShader, PixelShader);
+		return GetUniformMobileBasePassShaders<LMP_MOBILE_DIRECTIONAL_LIGHT_AND_SH_INDIRECT, bEnableLocalLights>(Material, VertexFactoryType, bEnableSkyLight, VertexShader, PixelShader);
 	case LMP_MOBILE_DIRECTIONAL_LIGHT_CSM_AND_SH_INDIRECT:
-		return GetUniformMobileBasePassShaders<LMP_MOBILE_DIRECTIONAL_LIGHT_CSM_AND_SH_INDIRECT, NumMovablePointLights>(Material, VertexFactoryType, bEnableSkyLight, VertexShader, PixelShader);
+		return GetUniformMobileBasePassShaders<LMP_MOBILE_DIRECTIONAL_LIGHT_CSM_AND_SH_INDIRECT, bEnableLocalLights>(Material, VertexFactoryType, bEnableSkyLight, VertexShader, PixelShader);
 	case LMP_MOBILE_MOVABLE_DIRECTIONAL_LIGHT_WITH_LIGHTMAP:
-		return GetUniformMobileBasePassShaders<LMP_MOBILE_MOVABLE_DIRECTIONAL_LIGHT_WITH_LIGHTMAP, NumMovablePointLights>(Material, VertexFactoryType, bEnableSkyLight, VertexShader, PixelShader);
+		return GetUniformMobileBasePassShaders<LMP_MOBILE_MOVABLE_DIRECTIONAL_LIGHT_WITH_LIGHTMAP, bEnableLocalLights>(Material, VertexFactoryType, bEnableSkyLight, VertexShader, PixelShader);
 	case LMP_MOBILE_MOVABLE_DIRECTIONAL_LIGHT_CSM_WITH_LIGHTMAP:
-		return GetUniformMobileBasePassShaders<LMP_MOBILE_MOVABLE_DIRECTIONAL_LIGHT_CSM_WITH_LIGHTMAP, NumMovablePointLights>(Material, VertexFactoryType, bEnableSkyLight, VertexShader, PixelShader);
+		return GetUniformMobileBasePassShaders<LMP_MOBILE_MOVABLE_DIRECTIONAL_LIGHT_CSM_WITH_LIGHTMAP, bEnableLocalLights>(Material, VertexFactoryType, bEnableSkyLight, VertexShader, PixelShader);
 	case LMP_MOBILE_DIRECTIONAL_LIGHT_CSM:
-		return GetUniformMobileBasePassShaders<LMP_MOBILE_DIRECTIONAL_LIGHT_CSM, NumMovablePointLights>(Material, VertexFactoryType, bEnableSkyLight, VertexShader, PixelShader);
+		return GetUniformMobileBasePassShaders<LMP_MOBILE_DIRECTIONAL_LIGHT_CSM, bEnableLocalLights>(Material, VertexFactoryType, bEnableSkyLight, VertexShader, PixelShader);
 	default:										
 		check(false);
 		return true;
@@ -121,7 +121,7 @@ bool GetMobileBasePassShaders(
 
 bool MobileBasePass::GetShaders(
 	ELightMapPolicyType LightMapPolicyType,
-	int32 NumMovablePointLights, 
+	bool bEnableLocalLights,
 	const FMaterial& MaterialResource,
 	FVertexFactoryType* VertexFactoryType,
 	bool bEnableSkyLight, 
@@ -134,10 +134,9 @@ bool MobileBasePass::GetShaders(
 		bEnableSkyLight = !bEnableSkyLight;
 	}
 
-	switch (NumMovablePointLights)
+	if (bEnableLocalLights)
 	{
-	case INT32_MAX:
-		return GetMobileBasePassShaders<INT32_MAX>(
+		return GetMobileBasePassShaders<true>(
 			LightMapPolicyType,
 			MaterialResource,
 			VertexFactoryType,
@@ -145,9 +144,10 @@ bool MobileBasePass::GetShaders(
 			VertexShader,
 			PixelShader
 			);
-	case 0:
-	default:
-		return GetMobileBasePassShaders<0>(
+	}
+	else
+	{
+		return GetMobileBasePassShaders<false>(
 			LightMapPolicyType,
 			MaterialResource,
 			VertexFactoryType,
@@ -177,18 +177,6 @@ const FLightSceneInfo* MobileBasePass::GetDirectionalLightInfo(const FScene* Sce
 	return MobileDirectionalLight;
 }
 
-int32 MobileBasePass::CalcNumMovablePointLights(const FMaterial& InMaterial, const FPrimitiveSceneProxy* InPrimitiveSceneProxy)
-{
-	const FReadOnlyCVARCache& ReadOnlyCVARCache = FReadOnlyCVARCache::Get();
-	const bool bIsUnlit = InMaterial.GetShadingModels().IsUnlit();
-	int32 OutNumMovablePointLights = (InPrimitiveSceneProxy && !bIsUnlit) ? FMath::Min<int32>(InPrimitiveSceneProxy->GetPrimitiveSceneInfo()->NumMobileMovablePointLights, ReadOnlyCVARCache.NumMobileMovablePointLights) : 0;
-	if (OutNumMovablePointLights > 0)
-	{
-		OutNumMovablePointLights = INT32_MAX;
-	}
-	return OutNumMovablePointLights;
-}
-
 bool MobileBasePass::StaticCanReceiveCSM(const FLightSceneInfo* LightSceneInfo, const FPrimitiveSceneProxy* PrimitiveSceneProxy)
 {
 	// For movable directional lights, when CSM culling is disabled the default behavior is to receive CSM.
@@ -556,26 +544,6 @@ void TMobileBasePassPSPolicyParamType<FUniformLightMapPolicy>::GetShaderBindings
 			}
 			ShaderBindings.Add(ReflectionParameter, ReflectionUB);
 		}
-		
-		if (NumDynamicPointLightsParameter.IsBound())
-		{
-			static FHashedName MobileMovablePointLightHashedName[MAX_BASEPASS_DYNAMIC_POINT_LIGHTS] = { FHashedName(TEXT("MobileMovablePointLight0")), FHashedName(TEXT("MobileMovablePointLight1")), FHashedName(TEXT("MobileMovablePointLight2")), FHashedName(TEXT("MobileMovablePointLight3")) };
-
-			// Set dynamic point lights
-			FMobileBasePassMovableLightInfo LightInfo(PrimitiveSceneProxy);
-			ShaderBindings.Add(NumDynamicPointLightsParameter, LightInfo.NumMovablePointLights);
-			for (int32 i = 0; i < MAX_BASEPASS_DYNAMIC_POINT_LIGHTS; ++i)
-			{
-				if (i < LightInfo.NumMovablePointLights && LightInfo.MovablePointLightUniformBuffer[i])
-				{
-					ShaderBindings.Add(GetUniformBufferParameter(MobileMovablePointLightHashedName[i]), LightInfo.MovablePointLightUniformBuffer[i]);
-				}
-				else
-				{
-					ShaderBindings.Add(GetUniformBufferParameter(MobileMovablePointLightHashedName[i]), GDummyMovablePointLightUniformBuffer.GetUniformBufferRHI());
-				}
-			}
-		}
 	}
 	else
 	{
@@ -703,8 +671,11 @@ bool FMobileBasePassMeshProcessor::Process(
 	
 	if (Scene && Scene->SkyLight)
 	{
+		// Clustered reflection is always enabled on mobile deferred translucent material
+		const bool bEnableClusteredReflections = MobileEnableClusteredReflections(Scene->GetShaderPlatform());
+
 		//The stationary skylight contribution has been added both to the LowQuality Lightmap and ILC on mobile, so we should skip the sky light spherical harmonic contribution for it.
-		bool bStationarySkyLightHasBeenApplied = MobileBasePass::StationarySkyLightHasBeenApplied(Scene, LightMapPolicyType);
+		bool bStationarySkyLightHasBeenApplied = MobileBasePass::StationarySkyLightHasBeenApplied(Scene, LightMapPolicyType) && !bEnableClusteredReflections;
 
 		//Two side material should enable sky light for the back face since only the front face has light map and it will be corrected in base pass shader.
 		bool bSkipStationarySkyLight = bStationarySkyLightHasBeenApplied && !MaterialResource.IsTwoSided();
@@ -712,15 +683,20 @@ bool FMobileBasePassMeshProcessor::Process(
 		bEnableSkyLight = ShadingModels.IsLit() && Scene->ShouldRenderSkylightInBasePass(BlendMode) && (!bSkipStationarySkyLight);
 	}
 
-	int32 NumMovablePointLights = 0;
-	if (!bUsesDeferredShading)
+	bool bEnableLocalLights = false;
+	if (Scene && PrimitiveSceneProxy && ShadingModels.IsLit())
 	{
-		NumMovablePointLights = MobileBasePass::CalcNumMovablePointLights(MaterialResource, PrimitiveSceneProxy);
+		const bool bIsTranslucentMaterialOnDeferred = Scene && bTranslucentBasePass && IsMobileDeferredShadingEnabled(Scene->GetShaderPlatform());
+		// Clustered lighting could only be enabled on mobile forward and always enabled on mobile deferred translucent material
+		if (!bUsesDeferredShading && (MobileForwardEnableLocalLights(Scene->GetShaderPlatform()) || bIsTranslucentMaterialOnDeferred))
+		{
+			bEnableLocalLights = PrimitiveSceneProxy->GetPrimitiveSceneInfo()->NumMobileMovableLocalLights > 0;
+		}
 	}
 
 	if (!MobileBasePass::GetShaders(
 		LightMapPolicyType,
-		NumMovablePointLights,
+		bEnableLocalLights,
 		MaterialResource,
 		MeshBatch.VertexFactory->GetType(),
 		bEnableSkyLight,
diff --git a/Engine/Source/Runtime/Renderer/Private/MobileBasePassRendering.cpp b/Engine/Source/Runtime/Renderer/Private/MobileBasePassRendering.cpp
--- a/Engine/Source/Runtime/Renderer/Private/MobileBasePassRendering.cpp
+++ b/Engine/Source/Runtime/Renderer/Private/MobileBasePassRendering.cpp
@@ -18,6 +18,8 @@
 #include "FramePro/FrameProProfiler.h"
 #include "PostProcess/PostProcessPixelProjectedReflectionMobile.h"
 #include "Engine/SubsurfaceProfile.h"
+#include "LocalLightSceneProxy.h"
+#include "ReflectionEnvironment.h"
 
 // Changing this causes a full shader recompile
 static TAutoConsoleVariable<int32> CVarMobileDisableVertexFog(
@@ -26,20 +28,14 @@ static TAutoConsoleVariable<int32> CVarMobileDisableVertexFog(
 	TEXT("Set to 1 to disable vertex fogging in all mobile shaders."),
 	ECVF_ReadOnly | ECVF_RenderThreadSafe);
 
-static TAutoConsoleVariable<int32> CVarMobileEnableMovableSpotLights(
-	TEXT("r.Mobile.EnableMovableSpotlights"),
-	0,
-	TEXT("If 1 then enable movable spotlight support"),
-	ECVF_ReadOnly | ECVF_RenderThreadSafe);
-
 static TAutoConsoleVariable<int32> CVarMobileEnableMovableSpotLightShadows(
 	TEXT("r.Mobile.EnableMovableSpotlightsShadow"),
 	0,
 	TEXT("If 1 then enable movable spotlight shadow support"),
 	ECVF_ReadOnly | ECVF_RenderThreadSafe);
 
-static TAutoConsoleVariable<int32> CVarMobileMaxVisibleMovableSpotLightsShadow(
-	TEXT("r.Mobile.MaxVisibleMovableSpotLightsShadow"),
+static TAutoConsoleVariable<int32> CVarMobileMaxVisibleMovableSpotLightShadows(
+	TEXT("r.Mobile.MaxVisibleMovableSpotLightShadows"),
 	8,
 	TEXT("The max number of visible spotlighs can cast shadow sorted by screen size, should be as less as possible for performance reason"),
 	ECVF_RenderThreadSafe);
@@ -68,11 +64,6 @@ static TAutoConsoleVariable<int32> CVarMobileUseHWsRGBEncoding(
 	TEXT("1: Use GPU HW to convert linear to sRGB automatically (device must support sRGB write control)\n"),
 	ECVF_RenderThreadSafe);
 
-namespace MobileBasePass
-{
-	FShaderPlatformCachedIniValue<int32> MobileNumDynamicPointLightsIniValue(TEXT("r.MobileNumDynamicPointLights"));
-};
-
 bool ShouldCacheShaderByPlatformAndOutputFormat(EShaderPlatform Platform, EOutputFormat OutputFormat)
 {
 	bool bSupportsMobileHDR = IsMobileHDR();
@@ -89,23 +80,27 @@ bool ShouldCacheShaderByPlatformAndOutputFormat(EShaderPlatform Platform, EOutpu
 	IMPLEMENT_MATERIAL_SHADER_TYPE(template<>, TMobileBasePassVS##LightMapPolicyName##LDRGamma32, TEXT("/Engine/Private/MobileBasePassVertexShader.usf"), TEXT("Main"), SF_Vertex); \
 	IMPLEMENT_MATERIAL_SHADER_TYPE(template<>, TMobileBasePassVS##LightMapPolicyName##HDRLinear64, TEXT("/Engine/Private/MobileBasePassVertexShader.usf"), TEXT("Main"), SF_Vertex);
 
-#define IMPLEMENT_MOBILE_SHADING_BASEPASS_LIGHTMAPPED_PIXEL_SHADER_TYPE(LightMapPolicyType,LightMapPolicyName,NumMovablePointLights) \
-	typedef TMobileBasePassPS< LightMapPolicyType, LDR_GAMMA_32, false, NumMovablePointLights > TMobileBasePassPS##LightMapPolicyName##NumMovablePointLights##LDRGamma32; \
-	typedef TMobileBasePassPS< LightMapPolicyType, HDR_LINEAR_64, false, NumMovablePointLights > TMobileBasePassPS##LightMapPolicyName##NumMovablePointLights##HDRLinear64; \
-	typedef TMobileBasePassPS< LightMapPolicyType, LDR_GAMMA_32, true, NumMovablePointLights > TMobileBasePassPS##LightMapPolicyName##NumMovablePointLights##LDRGamma32##Skylight; \
-	typedef TMobileBasePassPS< LightMapPolicyType, HDR_LINEAR_64, true, NumMovablePointLights > TMobileBasePassPS##LightMapPolicyName##NumMovablePointLights##HDRLinear64##Skylight; \
-	IMPLEMENT_MATERIAL_SHADER_TYPE(template<>, TMobileBasePassPS##LightMapPolicyName##NumMovablePointLights##LDRGamma32, TEXT("/Engine/Private/MobileBasePassPixelShader.usf"), TEXT("Main"), SF_Pixel); \
-	IMPLEMENT_MATERIAL_SHADER_TYPE(template<>, TMobileBasePassPS##LightMapPolicyName##NumMovablePointLights##HDRLinear64, TEXT("/Engine/Private/MobileBasePassPixelShader.usf"), TEXT("Main"), SF_Pixel); \
-	IMPLEMENT_MATERIAL_SHADER_TYPE(template<>, TMobileBasePassPS##LightMapPolicyName##NumMovablePointLights##LDRGamma32##Skylight, TEXT("/Engine/Private/MobileBasePassPixelShader.usf"), TEXT("Main"), SF_Pixel); \
-	IMPLEMENT_MATERIAL_SHADER_TYPE(template<>, TMobileBasePassPS##LightMapPolicyName##NumMovablePointLights##HDRLinear64##Skylight, TEXT("/Engine/Private/MobileBasePassPixelShader.usf"), TEXT("Main"), SF_Pixel);
-
-static_assert(MAX_BASEPASS_DYNAMIC_POINT_LIGHTS == 4, "If you change MAX_BASEPASS_DYNAMIC_POINT_LIGHTS, you need to add shader types below");
+#define IMPLEMENT_MOBILE_SHADING_BASEPASS_LIGHTMAPPED_PIXEL_SHADER_TYPE(LightMapPolicyType,LightMapPolicyName) \
+	typedef TMobileBasePassPS< LightMapPolicyType, LDR_GAMMA_32, false, false > TMobileBasePassPS##LightMapPolicyName##LDRGamma32; \
+	typedef TMobileBasePassPS< LightMapPolicyType, HDR_LINEAR_64, false, false > TMobileBasePassPS##LightMapPolicyName##HDRLinear64; \
+	typedef TMobileBasePassPS< LightMapPolicyType, LDR_GAMMA_32, true, false > TMobileBasePassPS##LightMapPolicyName##LDRGamma32##Skylight; \
+	typedef TMobileBasePassPS< LightMapPolicyType, HDR_LINEAR_64, true, false > TMobileBasePassPS##LightMapPolicyName##HDRLinear64##Skylight; \
+	typedef TMobileBasePassPS< LightMapPolicyType, LDR_GAMMA_32, false, true > TMobileBasePassPS##LightMapPolicyName##LDRGamma32##EnableLocalLights; \
+	typedef TMobileBasePassPS< LightMapPolicyType, HDR_LINEAR_64, false, true > TMobileBasePassPS##LightMapPolicyName##HDRLinear64##EnableLocalLights; \
+	typedef TMobileBasePassPS< LightMapPolicyType, LDR_GAMMA_32, true, true > TMobileBasePassPS##LightMapPolicyName##LDRGamma32##Skylight##EnableLocalLights; \
+	typedef TMobileBasePassPS< LightMapPolicyType, HDR_LINEAR_64, true, true > TMobileBasePassPS##LightMapPolicyName##HDRLinear64##Skylight##EnableLocalLights; \
+	IMPLEMENT_MATERIAL_SHADER_TYPE(template<>, TMobileBasePassPS##LightMapPolicyName##LDRGamma32, TEXT("/Engine/Private/MobileBasePassPixelShader.usf"), TEXT("Main"), SF_Pixel); \
+	IMPLEMENT_MATERIAL_SHADER_TYPE(template<>, TMobileBasePassPS##LightMapPolicyName##HDRLinear64, TEXT("/Engine/Private/MobileBasePassPixelShader.usf"), TEXT("Main"), SF_Pixel); \
+	IMPLEMENT_MATERIAL_SHADER_TYPE(template<>, TMobileBasePassPS##LightMapPolicyName##LDRGamma32##Skylight, TEXT("/Engine/Private/MobileBasePassPixelShader.usf"), TEXT("Main"), SF_Pixel); \
+	IMPLEMENT_MATERIAL_SHADER_TYPE(template<>, TMobileBasePassPS##LightMapPolicyName##HDRLinear64##Skylight, TEXT("/Engine/Private/MobileBasePassPixelShader.usf"), TEXT("Main"), SF_Pixel) \
+	IMPLEMENT_MATERIAL_SHADER_TYPE(template<>, TMobileBasePassPS##LightMapPolicyName##LDRGamma32##EnableLocalLights, TEXT("/Engine/Private/MobileBasePassPixelShader.usf"), TEXT("Main"), SF_Pixel); \
+	IMPLEMENT_MATERIAL_SHADER_TYPE(template<>, TMobileBasePassPS##LightMapPolicyName##HDRLinear64##EnableLocalLights, TEXT("/Engine/Private/MobileBasePassPixelShader.usf"), TEXT("Main"), SF_Pixel); \
+	IMPLEMENT_MATERIAL_SHADER_TYPE(template<>, TMobileBasePassPS##LightMapPolicyName##LDRGamma32##Skylight##EnableLocalLights, TEXT("/Engine/Private/MobileBasePassPixelShader.usf"), TEXT("Main"), SF_Pixel); \
+	IMPLEMENT_MATERIAL_SHADER_TYPE(template<>, TMobileBasePassPS##LightMapPolicyName##HDRLinear64##Skylight##EnableLocalLights, TEXT("/Engine/Private/MobileBasePassPixelShader.usf"), TEXT("Main"), SF_Pixel);
 
-// Permutations for the number of point lights to support. INT32_MAX indicates the shader should use branching to support a variable number of point lights.
 #define IMPLEMENT_MOBILE_SHADING_BASEPASS_LIGHTMAPPED_SHADER_TYPE(LightMapPolicyType,LightMapPolicyName) \
 	IMPLEMENT_MOBILE_SHADING_BASEPASS_LIGHTMAPPED_VERTEX_SHADER_TYPE(LightMapPolicyType, LightMapPolicyName) \
-	IMPLEMENT_MOBILE_SHADING_BASEPASS_LIGHTMAPPED_PIXEL_SHADER_TYPE(LightMapPolicyType, LightMapPolicyName, 0) \
-	IMPLEMENT_MOBILE_SHADING_BASEPASS_LIGHTMAPPED_PIXEL_SHADER_TYPE(LightMapPolicyType, LightMapPolicyName, INT32_MAX)
+	IMPLEMENT_MOBILE_SHADING_BASEPASS_LIGHTMAPPED_PIXEL_SHADER_TYPE(LightMapPolicyType, LightMapPolicyName)
 
 // Implement shader types per lightmap policy 
 IMPLEMENT_MOBILE_SHADING_BASEPASS_LIGHTMAPPED_SHADER_TYPE(TUniformLightMapPolicy<LMP_NO_LIGHTMAP>, FNoLightMapPolicy);
@@ -131,41 +126,12 @@ bool TMobileBasePassPSPolicyParamType<LightMapPolicyType>::ModifyCompilationEnvi
 	OutEnvironment.SetDefine(TEXT("MOBILE_QL_FORCE_FULLY_ROUGH"), QualityOverrides.bEnableOverride && QualityOverrides.bForceFullyRough != 0 ? 1u : 0u);
 	OutEnvironment.SetDefine(TEXT("MOBILE_QL_FORCE_NONMETAL"), QualityOverrides.bEnableOverride && QualityOverrides.bForceNonMetal != 0 ? 1u : 0u);
 	OutEnvironment.SetDefine(TEXT("QL_FORCEDISABLE_LM_DIRECTIONALITY"), QualityOverrides.bEnableOverride && QualityOverrides.bForceDisableLMDirectionality != 0 ? 1u : 0u);
-	OutEnvironment.SetDefine(TEXT("MOBILE_QL_FORCE_LQ_REFLECTIONS"), QualityOverrides.bEnableOverride && QualityOverrides.bForceLQReflections != 0 ? 1u : 0u);
 	OutEnvironment.SetDefine(TEXT("MOBILE_QL_FORCE_DISABLE_PREINTEGRATEDGF"), QualityOverrides.bEnableOverride && QualityOverrides.bForceDisablePreintegratedGF != 0 ? 1u : 0u);
 	OutEnvironment.SetDefine(TEXT("MOBILE_SHADOW_QUALITY"), (uint32)QualityOverrides.MobileShadowQuality);
 	OutEnvironment.SetDefine(TEXT("MOBILE_QL_DISABLE_MATERIAL_NORMAL"), QualityOverrides.bEnableOverride && QualityOverrides.bDisableMaterialNormalCalculation);
 	return true;
 }
 
-FMobileBasePassMovableLightInfo::FMobileBasePassMovableLightInfo(const FPrimitiveSceneProxy* InSceneProxy)
-: NumMovablePointLights(0)
-{
-	static auto* MobileNumDynamicPointLightsCVar = IConsoleManager::Get().FindTConsoleVariableDataInt(TEXT("r.MobileNumDynamicPointLights"));
-	const int32 MobileNumDynamicPointLights = MobileNumDynamicPointLightsCVar->GetValueOnRenderThread();
-
-	if (InSceneProxy != nullptr)
-	{
-		for (FLightPrimitiveInteraction* LPI = InSceneProxy->GetPrimitiveSceneInfo()->LightList; LPI && NumMovablePointLights < MobileNumDynamicPointLights; LPI = LPI->GetNextLight())
-		{
-			FLightSceneInfo* LightSceneInfo = LPI->GetLight();
-			FLightSceneProxy* LightProxy = LightSceneInfo->Proxy;
-			const uint8 LightType = LightProxy->GetLightType(); 
-			const bool bIsValidLightType =
-				  LightType == LightType_Point
-				|| LightType == LightType_Rect
-				|| (LightType == LightType_Spot && CVarMobileEnableMovableSpotLights.GetValueOnRenderThread());
-				
-			if (bIsValidLightType && LightProxy->IsMovable() && (LightProxy->GetLightingChannelMask() & InSceneProxy->GetLightingChannelMask()) != 0)
-			{
-				MovablePointLightUniformBuffer[NumMovablePointLights] = LightProxy->GetMobileMovablePointLightUniformBufferRHI();
-
-				NumMovablePointLights++;
-			}
-		}
-	}
-}
-
 extern void SetupDummyForwardLightUniformParameters(FRDGBuilder& GraphBuilder, FForwardLightData& ForwardLightData);
 
 void SetupMobileBasePassUniformParameters(
@@ -218,8 +184,10 @@ void SetupMobileBasePassUniformParameters(
 	BasePassParameters.AmbientOcclusionTexture = AmbientOcclusionTexture;
 	BasePassParameters.AmbientOcclusionSampler = TStaticSamplerState<SF_Bilinear, AM_Clamp, AM_Clamp, AM_Clamp>::GetRHI();
 	BasePassParameters.AmbientOcclusionStaticFraction = FMath::Clamp(View.FinalPostProcessSettings.AmbientOcclusionStaticFraction, 0.0f, 1.0f);
+
+	const bool bMobileUsesShadowMaskTexture = MobileUsesShadowMaskTexture(View.GetShaderPlatform());
 	
-	if (IsMobileDistanceFieldEnabled(View.GetShaderPlatform()) && GScreenSpaceShadowMaskTextureMobileOutputs.ScreenSpaceShadowMaskTextureMobile.IsValid())
+	if (bMobileUsesShadowMaskTexture && GScreenSpaceShadowMaskTextureMobileOutputs.ScreenSpaceShadowMaskTextureMobile.IsValid())
 	{
 		FRDGTextureRef ScreenShadowMaskTexture = GraphBuilder.RegisterExternalTexture(GScreenSpaceShadowMaskTextureMobileOutputs.ScreenSpaceShadowMaskTextureMobile, TEXT("ScreenSpaceShadowMaskTextureMobile"));
 		BasePassParameters.ScreenSpaceShadowMaskTexture = ScreenShadowMaskTexture;
@@ -230,6 +198,8 @@ void SetupMobileBasePassUniformParameters(
 		BasePassParameters.ScreenSpaceShadowMaskTexture = GSystemTextures.GetWhiteDummy(GraphBuilder);
 		BasePassParameters.ScreenSpaceShadowMaskSampler = TStaticSamplerState<SF_Point, AM_Clamp, AM_Clamp, AM_Clamp>::GetRHI();
 	}
+
+	SetupReflectionUniformParameters(View, BasePassParameters.ReflectionsParameters);
 }
 
 TRDGUniformBufferRef<FMobileBasePassUniformParameters> CreateMobileBasePassUniformBuffer(
@@ -301,14 +271,6 @@ void SetupMobileDirectionalLightUniformParameters(
 			}
 		}
 	}
-	
-	if (SceneView.MobileMovableSpotLightsShadowInfo.ShadowDepthTexture != nullptr)
-	{
-		checkSlow(Params.DirectionalLightShadowTexture == SceneView.MobileMovableSpotLightsShadowInfo.ShadowDepthTexture || Params.DirectionalLightShadowTexture == GWhiteTexture->TextureRHI);
-
-		Params.DirectionalLightShadowSize = SceneView.MobileMovableSpotLightsShadowInfo.ShadowBufferSize;
-		Params.DirectionalLightShadowTexture = SceneView.MobileMovableSpotLightsShadowInfo.ShadowDepthTexture;
-	}
 }
 
 void SetupMobileSkyReflectionUniformParameters(FSkyLightSceneProxy* SkyLight, FMobileReflectionCaptureShaderParameters& Parameters)
diff --git a/Engine/Source/Runtime/Renderer/Private/MobileBasePassRendering.h b/Engine/Source/Runtime/Renderer/Private/MobileBasePassRendering.h
--- a/Engine/Source/Runtime/Renderer/Private/MobileBasePassRendering.h
+++ b/Engine/Source/Runtime/Renderer/Private/MobileBasePassRendering.h
@@ -38,6 +38,7 @@ BEGIN_GLOBAL_SHADER_PARAMETER_STRUCT(FMobileBasePassUniformParameters, )
 	SHADER_PARAMETER_STRUCT(FPlanarReflectionUniformParameters, PlanarReflection) // Single global planar reflection for the forward pass.
 	SHADER_PARAMETER_STRUCT(FMobileSceneTextureUniformParameters, SceneTextures)
 	SHADER_PARAMETER_STRUCT(FDebugViewModeUniformParameters, DebugViewMode)
+	SHADER_PARAMETER_STRUCT(FReflectionUniformParameters, ReflectionsParameters)
 	SHADER_PARAMETER_TEXTURE(Texture2D, PreIntegratedGFTexture)
 	SHADER_PARAMETER_SAMPLER(SamplerState, PreIntegratedGFSampler)
 	SHADER_PARAMETER_RDG_BUFFER_SRV(Buffer<float4>, EyeAdaptationBuffer)
@@ -91,18 +92,6 @@ enum EOutputFormat
 	HDR_LINEAR_64,
 };
 
-#define MAX_BASEPASS_DYNAMIC_POINT_LIGHTS 4
-
-/* Info for dynamic point or spot lights rendered in base pass */
-class FMobileBasePassMovableLightInfo
-{
-public:
-	FMobileBasePassMovableLightInfo(const FPrimitiveSceneProxy* InSceneProxy);
-
-	int32 NumMovablePointLights;
-	FRHIUniformBuffer* MovablePointLightUniformBuffer[MAX_BASEPASS_DYNAMIC_POINT_LIGHTS];
-};
-
 bool ShouldCacheShaderByPlatformAndOutputFormat(EShaderPlatform Platform, EOutputFormat OutputFormat);
 
 template<typename LightMapPolicyType>
@@ -272,8 +261,6 @@ public:
 		HQReflectionTilePositions.Bind(Initializer.ParameterMap, TEXT("ReflectionTilePositions"));
 		HQReflectionCaptureBoxTransformArray.Bind(Initializer.ParameterMap, TEXT("CaptureBoxTransformArray"));
 		HQReflectionCaptureBoxScalesArray.Bind(Initializer.ParameterMap, TEXT("CaptureBoxScalesArray"));
-
-		NumDynamicPointLightsParameter.Bind(Initializer.ParameterMap, TEXT("NumDynamicPointLights"));
 						
 		CSMDebugHintParams.Bind(Initializer.ParameterMap, TEXT("CSMDebugHint"));
 		UseCSMParameter.Bind(Initializer.ParameterMap, TEXT("UseCSM"));
@@ -295,8 +282,6 @@ private:
 	LAYOUT_FIELD(FShaderParameter, HQReflectionCaptureBoxTransformArray);
 	LAYOUT_FIELD(FShaderParameter, HQReflectionCaptureBoxScalesArray);
 
-	LAYOUT_FIELD(FShaderParameter, NumDynamicPointLightsParameter);
-
 	LAYOUT_FIELD(FShaderParameter, CSMDebugHintParams);
 	LAYOUT_FIELD(FShaderParameter, UseCSMParameter);
 	
@@ -354,7 +339,7 @@ namespace MobileBasePass
 
 	bool GetShaders(
 		ELightMapPolicyType LightMapPolicyType,
-		int32 NumMovablePointLights, 
+		bool bEnableLocalLights,
 		const FMaterial& MaterialResource,
 		FVertexFactoryType* VertexFactoryType,
 		bool bEnableSkyLight, 
@@ -362,7 +347,6 @@ namespace MobileBasePass
 		TShaderRef<TMobileBasePassPSPolicyParamType<FUniformLightMapPolicy>>& PixelShader);
 
 	const FLightSceneInfo* GetDirectionalLightInfo(const FScene* Scene, const FPrimitiveSceneProxy* PrimitiveSceneProxy);
-	int32 CalcNumMovablePointLights(const FMaterial& InMaterial, const FPrimitiveSceneProxy* InPrimitiveSceneProxy);
 
 	bool StaticCanReceiveCSM(const FLightSceneInfo* LightSceneInfo, const FPrimitiveSceneProxy* PrimitiveSceneProxy);
 
@@ -370,8 +354,6 @@ namespace MobileBasePass
 	void SetTranslucentRenderState(FMeshPassProcessorRenderState& DrawRenderState, const FMaterial& Material);
 
 	bool StationarySkyLightHasBeenApplied(const FScene* Scene, ELightMapPolicyType LightMapPolicyType);
-
-	extern FShaderPlatformCachedIniValue<int32> MobileNumDynamicPointLightsIniValue;
 };
 
 
@@ -387,7 +369,7 @@ inline bool UseSkylightPermutation(bool bEnableSkyLight, int32 MobileSkyLightPer
 	}
 }
 
-template< typename LightMapPolicyType, EOutputFormat OutputFormat, bool bEnableSkyLight, int32 NumMovablePointLights>
+template< typename LightMapPolicyType, EOutputFormat OutputFormat, bool bEnableSkyLight, bool bEnableLocalLights>
 class TMobileBasePassPS : public TMobileBasePassPSBaseType<LightMapPolicyType>
 {
 	DECLARE_SHADER_TYPE(TMobileBasePassPS,MeshMaterial);
@@ -398,7 +380,6 @@ public:
 		// We compile the point light shader combinations based on the project settings
 		static auto* MobileSkyLightPermutationCVar = IConsoleManager::Get().FindTConsoleVariableDataInt(TEXT("r.Mobile.SkyLightPermutation"));
 
-		const int32 MobileNumDynamicPointLights = MobileBasePass::MobileNumDynamicPointLightsIniValue.Get(Parameters.Platform);
 		const int32 MobileSkyLightPermutationOptions = MobileSkyLightPermutationCVar->GetValueOnAnyThread();
 		const bool bDeferredShading = IsMobileDeferredShadingEnabled(Parameters.Platform);
 		
@@ -406,6 +387,8 @@ public:
 		const bool bMaterialUsesForwardShading = bIsLit && 
 			(IsTranslucentBlendMode(Parameters.MaterialParameters.BlendMode) || Parameters.MaterialParameters.ShadingModels.HasShadingModel(MSM_SingleLayerWater));
 
+		// Translucent materials always support clustered shading on mobile deferred
+		const bool bSupportsLocalLights = (!bDeferredShading && MobileForwardEnableLocalLights(Parameters.Platform)) || (bDeferredShading && bMaterialUsesForwardShading);
 		// Only compile skylight version for lit materials
 		const bool bShouldCacheBySkylight = !bEnableSkyLight || bIsLit;
 
@@ -415,18 +398,16 @@ public:
 			return false;
 		}
 
-		// Deferred shading does not need SkyLight and PointLight permutations
+		// Deferred shading does not need SkyLight and LocalLight permutations
 		// TODO: skip skylight permutations for deferred
-		const bool bShouldCacheByShading = (!bDeferredShading || bMaterialUsesForwardShading) || (NumMovablePointLights == 0);
-
-		const bool bShouldCacheByNumDynamicPointLights =
-			(NumMovablePointLights == 0 ||
-			(bIsLit && NumMovablePointLights == INT32_MAX && MobileNumDynamicPointLights > 0));		// single shader for variable number of point lights
+		const bool bForwardShading = !bDeferredShading || bMaterialUsesForwardShading;
+		const bool bShouldCacheByShading = (bForwardShading || !bEnableLocalLights);
+		const bool bShouldCacheByLocalLights = !bEnableLocalLights || (bIsLit && bEnableLocalLights == bSupportsLocalLights);
 
 		return TMobileBasePassPSBaseType<LightMapPolicyType>::ShouldCompilePermutation(Parameters) && 
 				ShouldCacheShaderByPlatformAndOutputFormat(Parameters.Platform, OutputFormat) && 
 				bShouldCacheBySkylight && 
-				bShouldCacheByNumDynamicPointLights &&
+				bShouldCacheByLocalLights &&
 				bShouldCacheByShading;
 	}
 	
@@ -435,27 +416,21 @@ public:
 		static auto* MobileUseHWsRGBEncodingCVAR = IConsoleManager::Get().FindTConsoleVariableDataInt(TEXT("r.Mobile.UseHWsRGBEncoding"));
 		const bool bMobileUseHWsRGBEncoding = (MobileUseHWsRGBEncodingCVAR && MobileUseHWsRGBEncodingCVAR->GetValueOnAnyThread() == 1);
 
+		const bool bMobileUsesShadowMaskTexture = MobileUsesShadowMaskTexture(Parameters.Platform);
+		const bool bEnableClusteredReflections = MobileEnableClusteredReflections(Parameters.Platform);
+		const bool bTranslucentMaterial = IsTranslucentBlendMode(Parameters.MaterialParameters.BlendMode) || Parameters.MaterialParameters.ShadingModels.HasShadingModel(MSM_SingleLayerWater);
+
 		TMobileBasePassPSBaseType<LightMapPolicyType>::ModifyCompilationEnvironment(Parameters, OutEnvironment);
 		OutEnvironment.SetDefine(TEXT("ENABLE_SKY_LIGHT"), bEnableSkyLight);
 		OutEnvironment.SetDefine(TEXT("OUTPUT_GAMMA_SPACE"), OutputFormat == LDR_GAMMA_32 && !bMobileUseHWsRGBEncoding);
 		OutEnvironment.SetDefine(TEXT("OUTPUT_MOBILE_HDR"), OutputFormat == HDR_LINEAR_64 ? 1u : 0u);
-		if (NumMovablePointLights == INT32_MAX)
-		{
-			const int32 MaxDynamicPointLights = FMath::Clamp(MobileBasePass::MobileNumDynamicPointLightsIniValue.Get(Parameters.Platform), 0, MAX_BASEPASS_DYNAMIC_POINT_LIGHTS);
-			
-			OutEnvironment.SetDefine(TEXT("MAX_DYNAMIC_POINT_LIGHTS"), (uint32)MaxDynamicPointLights);
-			OutEnvironment.SetDefine(TEXT("VARIABLE_NUM_DYNAMIC_POINT_LIGHTS"), (uint32)1);
-		}
-		else
-		{
-			OutEnvironment.SetDefine(TEXT("MAX_DYNAMIC_POINT_LIGHTS"), (uint32)NumMovablePointLights);
-			OutEnvironment.SetDefine(TEXT("VARIABLE_NUM_DYNAMIC_POINT_LIGHTS"), (uint32)0);
-			OutEnvironment.SetDefine(TEXT("NUM_DYNAMIC_POINT_LIGHTS"), (uint32)NumMovablePointLights);
-		}
 
 		OutEnvironment.SetDefine(TEXT("ENABLE_AMBIENT_OCCLUSION"), IsMobileAmbientOcclusionEnabled(Parameters.Platform) ? 1u : 0u);
 
-		OutEnvironment.SetDefine(TEXT("ENABLE_DISTANCE_FIELD"), IsMobileDistanceFieldEnabled(Parameters.Platform));
+		FForwardLightingParameters::ModifyCompilationEnvironment(Parameters.Platform, OutEnvironment);
+		OutEnvironment.SetDefine(TEXT("ENABLE_CLUSTERED_LIGHTS"), bEnableLocalLights ? 1u : 0u);
+		OutEnvironment.SetDefine(TEXT("ENABLE_CLUSTERED_REFLECTION"), bEnableClusteredReflections ? 1u : 0u);
+		OutEnvironment.SetDefine(TEXT("USE_SHADOWMASKTEXTURE"), bMobileUsesShadowMaskTexture && !bTranslucentMaterial ? 1u : 0u);
 	}
 	
 	/** Initialization constructor. */
diff --git a/Engine/Source/Runtime/Renderer/Private/MobileDeferredShadingPass.cpp b/Engine/Source/Runtime/Renderer/Private/MobileDeferredShadingPass.cpp
--- a/Engine/Source/Runtime/Renderer/Private/MobileDeferredShadingPass.cpp
+++ b/Engine/Source/Runtime/Renderer/Private/MobileDeferredShadingPass.cpp
@@ -7,6 +7,7 @@
 #include "PostProcess/SceneFilterRendering.h"
 #include "PipelineStateCache.h"
 #include "PlanarReflectionRendering.h"
+#include "LocalLightSceneProxy.h"
 
 int32 GMobileUseClusteredDeferredShading = 0;
 static FAutoConsoleVariableRef CVarMobileUseClusteredDeferredShading(
@@ -65,12 +66,17 @@ class FMobileDirectionalLightFunctionPS : public FMaterialShader
 		OutEnvironment.SetDefine(TEXT(PREPROCESSOR_TO_STRING(MAX_MOBILE_SHADOWCASCADES)), GetMobileMaxShadowCascades());
 		OutEnvironment.SetDefine(TEXT("SUPPORTS_TEXTURECUBE_ARRAY"), 1);
 		OutEnvironment.SetDefine(TEXT("USE_LIGHT_FUNCTION"), Parameters.MaterialParameters.bIsDefaultMaterial ? 0 : 1);
-		OutEnvironment.SetDefine(TEXT("ENABLE_DISTANCE_FIELD"), IsMobileDistanceFieldEnabled(Parameters.Platform));
+		OutEnvironment.SetDefine(TEXT("USE_SHADOWMASKTEXTURE"), MobileUsesShadowMaskTexture(Parameters.Platform) ? 1u : 0u);
 		OutEnvironment.SetDefine(TEXT("MATERIAL_SHADER"), 1);
 	}
 
-	static FPermutationDomain RemapPermutationVector(FPermutationDomain PermutationVector)
+	static FPermutationDomain RemapPermutationVector(FPermutationDomain PermutationVector, EShaderPlatform Platform)
 	{
+		if (MobileUsesShadowMaskTexture(Platform))
+		{
+			PermutationVector.Set<FApplyCSM>(false);
+		}
+
 		if (PermutationVector.Get<FApplyCSM>() == false)
 		{
 			PermutationVector.Set<FShadowQuality>(0);
@@ -89,7 +95,7 @@ class FMobileDirectionalLightFunctionPS : public FMaterialShader
 		
 		FPermutationDomain PermutationVector(Parameters.PermutationId);
 		// Compile out the shader if this permutation gets remapped.
-		if (RemapPermutationVector(PermutationVector) != PermutationVector)
+		if (RemapPermutationVector(PermutationVector, Parameters.Platform) != PermutationVector)
 		{
 			return false;
 		}
@@ -102,7 +108,8 @@ class FMobileDirectionalLightFunctionPS : public FMaterialShader
 		bool bUseClusteredLights = GMobileUseClusteredDeferredShading != 0;
 		bool bClustredReflection = bInlineReflectionAndSky && (View.NumBoxReflectionCaptures + View.NumSphereReflectionCaptures) > 0;
 		bool bEnableSkyLight = bInlineReflectionAndSky && bSkyLight;
-		int32 ShadowQuality = bDynamicShadows ? (int32)GetShadowQuality() : 0;
+		const bool bMobileUsesShadowMaskTexture = MobileUsesShadowMaskTexture(View.GetShaderPlatform());
+		int32 ShadowQuality = bDynamicShadows && !bMobileUsesShadowMaskTexture ? (int32)GetShadowQuality() : 0;
 				
 		FPermutationDomain PermutationVector;
 		PermutationVector.Set<FMobileDirectionalLightFunctionPS::FEnableClustredLights>(bUseClusteredLights);
@@ -135,13 +142,15 @@ public:
 
 	class FSpotLightDim			: SHADER_PERMUTATION_BOOL("IS_SPOT_LIGHT");
 	class FIESProfileDim		: SHADER_PERMUTATION_BOOL("USE_IES_PROFILE");
-	using FPermutationDomain = TShaderPermutationDomain<FSpotLightDim, FIESProfileDim>;
+	class FSpotLightShadowDim	: SHADER_PERMUTATION_BOOL("SUPPORT_SPOTLIGHTS_SHADOW");
+	using FPermutationDomain = TShaderPermutationDomain<FSpotLightDim, FIESProfileDim, FSpotLightShadowDim>;
 
 	BEGIN_SHADER_PARAMETER_STRUCT(FParameters, )
 		SHADER_PARAMETER(FMatrix44f, TranslatedWorldToLight)
 		SHADER_PARAMETER(FVector4f, LightFunctionParameters)
 		SHADER_PARAMETER(FVector3f, LightFunctionParameters2)
 		SHADER_PARAMETER_STRUCT_REF(FDeferredLightUniformStruct, DeferredLightUniforms)
+		SHADER_PARAMETER_STRUCT_INCLUDE(FMobileMovableLocalLightShadowParameters, MobileMovableLocalLightShadow)
 		SHADER_PARAMETER_TEXTURE(Texture2D, IESTexture)
 		SHADER_PARAMETER_SAMPLER(SamplerState, IESTextureSampler)
 		SHADER_PARAMETER_TEXTURE(Texture2D, PreIntegratedGF)
@@ -156,15 +165,35 @@ public:
 		{
 			return false;
 		}
+
+		FPermutationDomain PermutationVector(Parameters.PermutationId);
+		// Compile out the shader if this permutation gets remapped.
+		if (RemapPermutationVector(PermutationVector, Parameters.Platform) != PermutationVector)
+		{
+			return false;
+		}
+
 		return true;
 	}
 
+	static FPermutationDomain RemapPermutationVector(FPermutationDomain PermutationVector, EShaderPlatform Platform)
+	{
+		if (!IsMobileMovableSpotlightShadowsEnabled(Platform))
+		{
+			PermutationVector.Set<FSpotLightShadowDim>(false);
+		}
+
+		return PermutationVector;
+	}
+
 	static void ModifyCompilationEnvironment(const FMaterialShaderPermutationParameters& Parameters, FShaderCompilerEnvironment& OutEnvironment)
 	{
 		FMaterialShader::ModifyCompilationEnvironment(Parameters, OutEnvironment);
 		OutEnvironment.SetDefine(TEXT("USE_LIGHT_FUNCTION"), Parameters.MaterialParameters.bIsDefaultMaterial ? 0 : 1);
 		OutEnvironment.SetDefine(TEXT("MATERIAL_SHADER"), 1);
 		OutEnvironment.SetDefine(TEXT("GLOBAL_PREINTEGRATEDGF"), 1);
+		OutEnvironment.SetDefine(TEXT("USE_SHADOWMASKTEXTURE"), 0);
+		OutEnvironment.SetDefine(TEXT("ENABLE_CLUSTERED_LIGHTS"), 0);
 	}
 
 	static void SetParameters(FRHICommandList& RHICmdList, const TShaderRef<FMobileRadialLightFunctionPS>& Shader, const FViewInfo& View, const FMaterialRenderProxy* Proxy, const FMaterial& Material, const FParameters& Parameters)
@@ -386,7 +415,9 @@ static void RenderDirectionalLight(FRHICommandListImmediate& RHICmdList, const F
 	PassParameters.ReflectionsParameters = CreateUniformBufferImmediate(ReflectionUniformParameters, UniformBuffer_SingleDraw);
 	PassParameters.LightFunctionParameters = FVector4f(1.0f, 1.0f, 0.0f, 0.0f);
 
-	if (IsMobileDistanceFieldEnabled(View.GetShaderPlatform()) && GScreenSpaceShadowMaskTextureMobileOutputs.ScreenSpaceShadowMaskTextureMobile.IsValid())
+	const bool bMobileUsesShadowMaskTexture = MobileUsesShadowMaskTexture(View.GetShaderPlatform());
+
+	if (bMobileUsesShadowMaskTexture && GScreenSpaceShadowMaskTextureMobileOutputs.ScreenSpaceShadowMaskTextureMobile.IsValid())
 	{
 		PassParameters.ScreenSpaceShadowMaskTexture = GScreenSpaceShadowMaskTextureMobileOutputs.ScreenSpaceShadowMaskTextureMobile->GetRHI();
 		PassParameters.ScreenSpaceShadowMaskSampler = TStaticSamplerState<SF_Point, AM_Clamp, AM_Clamp, AM_Clamp>::GetRHI();
@@ -549,7 +580,8 @@ static void RenderLocalLight(
 	const FScene& Scene, 
 	const FViewInfo& View, 
 	const FLightSceneInfo& LightSceneInfo, 
-	const FCachedLightMaterial& DefaultLightMaterial)
+	const FCachedLightMaterial& DefaultLightMaterial,
+	const TArray<FVisibleLightInfo, SceneRenderingAllocator>& VisibleLightInfos)
 {
 	uint8 LightingChannelMask = LightSceneInfo.Proxy->GetLightingChannelMask();
 	if (!LightSceneInfo.ShouldRenderLight(View) || LightingChannelMask == 0)
@@ -568,7 +600,8 @@ static void RenderLocalLight(
 	FString LightNameWithLevel;
 	FSceneRenderer::GetLightNameForDrawEvent(LightSceneInfo.Proxy, LightNameWithLevel);
 	SCOPED_DRAW_EVENTF(RHICmdList, LocalLight, TEXT("%s"), *LightNameWithLevel);
-
+	check(LightSceneInfo.Proxy->IsLocalLight());
+	
 	if (GMobileUseLightStencilCulling != 0)
 	{
 		RenderLocalLight_StencilMask(RHICmdList, Scene, View, LightSceneInfo);
@@ -614,9 +647,12 @@ static void RenderLocalLight(
 	{
 		LightFunctionMaterialProxy = LightSceneInfo.Proxy->GetLightFunctionMaterial();
 	}
+	FMobileRadialLightFunctionPS::FParameters PassParameters;
+	const bool bShouldCastShadow = LightSceneInfo.SetupMobileMovableLocalLightShadowParameters(View, VisibleLightInfos, PassParameters.MobileMovableLocalLightShadow);
 	FMobileRadialLightFunctionPS::FPermutationDomain PermutationVector;
 	PermutationVector.Set<FMobileRadialLightFunctionPS::FSpotLightDim>(bIsSpotLight);
 	PermutationVector.Set<FMobileRadialLightFunctionPS::FIESProfileDim>(bUseIESTexture);
+	PermutationVector.Set<FMobileRadialLightFunctionPS::FSpotLightShadowDim>(bShouldCastShadow);
 	FCachedLightMaterial LightMaterial;
 	TShaderRef<FMobileRadialLightFunctionPS> PixelShader;
 	GetLightMaterial(DefaultLightMaterial, LightFunctionMaterialProxy, PermutationVector.ToDimensionValueId(), LightMaterial, PixelShader);
@@ -632,13 +668,12 @@ static void RenderLocalLight(
 	FDeferredLightVS::FParameters ParametersVS = FDeferredLightVS::GetParameters(View, &LightSceneInfo);
 	SetShaderParameters(RHICmdList, VertexShader, VertexShader.GetVertexShader(), ParametersVS);
 
-	FMobileRadialLightFunctionPS::FParameters PassParameters;
 	PassParameters.DeferredLightUniforms = TUniformBufferRef<FDeferredLightUniformStruct>::CreateUniformBufferImmediate(GetDeferredLightParameters(View, LightSceneInfo), EUniformBufferUsage::UniformBuffer_SingleFrame);
 	PassParameters.IESTexture = IESTextureResource->TextureRHI;
 	PassParameters.IESTextureSampler = IESTextureResource->SamplerStateRHI;
 	const float TanOuterAngle = bIsSpotLight ? FMath::Tan(LightSceneInfo.Proxy->GetOuterConeAngle()) : 1.0f;
 	PassParameters.LightFunctionParameters = FVector4f(TanOuterAngle, 1.0f /*ShadowFadeFraction*/, bIsSpotLight ? 1.0f : 0.0f, bIsPointLight ? 1.0f : 0.0f);
-	PassParameters.LightFunctionParameters2 = FVector3f(LightSceneInfo.Proxy->GetLightFunctionFadeDistance(), LightSceneInfo.Proxy->GetLightFunctionDisabledBrightness(),	0.0f);
+	PassParameters.LightFunctionParameters2 = FVector3f(LightSceneInfo.Proxy->GetLightFunctionFadeDistance(), LightSceneInfo.Proxy->GetLightFunctionDisabledBrightness(), IsMobileMovableSpotlightShadowsEnabled(Scene.GetShaderPlatform()) ? 1.0f : 0.0f);
 	const FVector Scale = LightSceneInfo.Proxy->GetLightFunctionScale();
 	// Switch x and z so that z of the user specified scale affects the distance along the light direction
 	const FVector InverseScale = FVector(1.f / Scale.Z, 1.f / Scale.Y, 1.f / Scale.X);
@@ -765,7 +800,8 @@ void MobileDeferredShadingPass(
 	int32 NumViews,
 	const FViewInfo& View,
 	const FScene& Scene, 
-	const FSortedLightSetSceneInfo &SortedLightSet)
+	const FSortedLightSetSceneInfo& SortedLightSet,
+	const TArray<FVisibleLightInfo, SceneRenderingAllocator>& VisibleLightInfos)
 {
 	SCOPED_DRAW_EVENT(RHICmdList, DeferredShading);
 	
@@ -787,16 +823,26 @@ void MobileDeferredShadingPass(
 
 	// Render non-clustered local lights
 	int32 NumLights = SortedLightSet.SortedLights.Num();
+	const int32 UnbatchedLightStart = SortedLightSet.UnbatchedLightStart;
 	int32 StandardDeferredStart = SortedLightSet.SimpleLightsEnd;
 	if (GMobileUseClusteredDeferredShading != 0)
 	{
 		StandardDeferredStart = SortedLightSet.ClusteredSupportedEnd;
 	}
 
-	for (int32 LightIdx = StandardDeferredStart; LightIdx < NumLights; ++LightIdx)
+	// Draw non-shadowed non-light function lights
+	for (int32 LightIdx = StandardDeferredStart; LightIdx < UnbatchedLightStart; ++LightIdx)
+	{
+		const FSortedLightSceneInfo& SortedLight = SortedLightSet.SortedLights[LightIdx];
+		const FLightSceneInfo& LightSceneInfo = *SortedLight.LightSceneInfo;
+		RenderLocalLight(RHICmdList, Scene, View, LightSceneInfo, DefaultMaterial, VisibleLightInfos);
+	}
+
+	// Draw shadowed and light function lights
+	for (int32 LightIdx = UnbatchedLightStart; LightIdx < NumLights; ++LightIdx)
 	{
 		const FSortedLightSceneInfo& SortedLight = SortedLightSet.SortedLights[LightIdx];
 		const FLightSceneInfo& LightSceneInfo = *SortedLight.LightSceneInfo;
-		RenderLocalLight(RHICmdList, Scene, View, LightSceneInfo, DefaultMaterial);
+		RenderLocalLight(RHICmdList, Scene, View, LightSceneInfo, DefaultMaterial, VisibleLightInfos);
 	}
 }
\ No newline at end of file
diff --git a/Engine/Source/Runtime/Renderer/Private/MobileDeferredShadingPass.h b/Engine/Source/Runtime/Renderer/Private/MobileDeferredShadingPass.h
--- a/Engine/Source/Runtime/Renderer/Private/MobileDeferredShadingPass.h
+++ b/Engine/Source/Runtime/Renderer/Private/MobileDeferredShadingPass.h
@@ -4,6 +4,7 @@
 
 #include "CoreMinimal.h"
 #include "RenderGraphDefinitions.h"
+#include "SceneRendering.h"
 
 class FMobileSceneTextureUniformParameters;
 class FRDGBuilder;
@@ -21,4 +22,5 @@ void MobileDeferredShadingPass(
 	int32 NumViews,
 	const FViewInfo& View,
 	const FScene& Scene, 
-	const FSortedLightSetSceneInfo &SortedLightSet);
+	const FSortedLightSetSceneInfo& SortedLightSet,
+	const TArray<FVisibleLightInfo, SceneRenderingAllocator>& VisibleLightInfos);
diff --git a/Engine/Source/Runtime/Renderer/Private/MobileShadingRenderer.cpp b/Engine/Source/Runtime/Renderer/Private/MobileShadingRenderer.cpp
--- a/Engine/Source/Runtime/Renderer/Private/MobileShadingRenderer.cpp
+++ b/Engine/Source/Runtime/Renderer/Private/MobileShadingRenderer.cpp
@@ -172,6 +172,7 @@ static void RenderOpaqueFX(
 BEGIN_SHADER_PARAMETER_STRUCT(FMobileRenderPassParameters, RENDERER_API)
 	SHADER_PARAMETER_STRUCT_INCLUDE(FViewShaderParameters, View)
 	SHADER_PARAMETER_RDG_UNIFORM_BUFFER(FMobileBasePassUniformParameters, MobileBasePass)
+	SHADER_PARAMETER_STRUCT_REF(FReflectionCaptureShaderData, ReflectionCapture)
 	RDG_BUFFER_ACCESS_ARRAY(DrawIndirectArgsBuffers)
 	RDG_BUFFER_ACCESS_ARRAY(InstanceIdOffsetBuffers)
 	RENDER_TARGET_BINDING_SLOTS()
@@ -230,9 +231,12 @@ FMobileSceneRenderer::FMobileSceneRenderer(const FSceneViewFamily* InViewFamily,
 	bShouldRenderCustomDepth = false;
 	bRequiresPixelProjectedPlanarRelfectionPass = false;
 	bRequiresAmbientOcclusionPass = false;
+	bRequiresShadowProjections = false;
 	bIsFullDepthPrepassEnabled = Scene->EarlyZPassMode == DDM_AllOpaque;
 	bIsMaskedOnlyDepthPrepassEnabled = Scene->EarlyZPassMode == DDM_MaskedOnly;
 	bRequiresSceneDepthAux = MobileRequiresSceneDepthAux(ShaderPlatform);
+	bEnableClusteredLocalLights = bDeferredShading || MobileForwardEnableLocalLights(ShaderPlatform);
+	bEnableClusteredReflections = MobileEnableClusteredReflections(ShaderPlatform);
 	
 	StandardTranslucencyPass = ViewFamily.AllowTranslucencyAfterDOF() ? ETranslucencyPass::TPT_StandardTranslucency : ETranslucencyPass::TPT_AllTranslucency;
 	StandardTranslucencyMeshPass = TranslucencyPassToMeshPass(StandardTranslucencyPass);
@@ -440,8 +444,9 @@ void FMobileSceneRenderer::InitViews(FRDGBuilder& GraphBuilder, FSceneTexturesCo
 		&& !ViewFamily.EngineShowFlags.VisualizeLightCulling
 		&& !ViewFamily.UseDebugViewPS();
 
+	bRequiresShadowProjections = MobileUsesShadowMaskTexture(ShaderPlatform);
+
 	bShouldRenderHZB = ShouldRenderHZB();
-		
 
 	// Whether we need to store depth for post-processing
 	// On PowerVR we see flickering of shadows and depths not updating correctly if targets are discarded.
@@ -453,8 +458,6 @@ void FMobileSceneRenderer::InitViews(FRDGBuilder& GraphBuilder, FSceneTexturesCo
 	bKeepDepthContent =
 		bRequiresMultiPass ||
 		bForceDepthResolve ||
-		bRequiresAmbientOcclusionPass ||
-		bRequiresDistanceField ||
 		bRequiresPixelProjectedPlanarRelfectionPass ||
 		bSeparateTranslucencyActive ||
 		Views[0].bIsReflectionCapture ||
@@ -513,13 +516,13 @@ void FMobileSceneRenderer::InitViews(FRDGBuilder& GraphBuilder, FSceneTexturesCo
 		InitSkyAtmosphereForViews(RHICmdList);
 	}
 
-	if (bRequiresDistanceField)
+	if (bRequiresShadowProjections)
 	{
-		InitMobileSDFShadowingOutputs(RHICmdList, SceneTexturesConfig.Extent);
+		InitMobileShadowProjectionOutputs(RHICmdList, SceneTexturesConfig.Extent);
 	}
 	else
 	{
-		ReleaseMobileSDFShadowingOutputs();
+		ReleaseMobileShadowProjectionOutputs();
 	}
 
 	// Find out whether custom depth pass should be rendered.
@@ -545,9 +548,6 @@ void FMobileSceneRenderer::InitViews(FRDGBuilder& GraphBuilder, FSceneTexturesCo
 		PrepareViewVisibilityLists();
 	}
 
-	/** Before SetupMobileBasePassAfterShadowInit, we need to update the uniform buffer and shadow info for all movable point lights.*/
-	UpdateMovablePointLightUniformBufferAndShadowInfo();
-
 	SetupMobileBasePassAfterShadowInit(BasePassDepthStencilAccess, ViewCommandsPerView, InstanceCullingManager);
 
 	// if we kicked off ILC update via task, wait and finalize.
@@ -556,6 +556,7 @@ void FMobileSceneRenderer::InitViews(FRDGBuilder& GraphBuilder, FSceneTexturesCo
 		Scene->IndirectLightingCache.FinalizeCacheUpdates(Scene, *this, ILCTaskData);
 	}
 
+	const bool bEnableClusteredShading = bEnableClusteredLocalLights || bEnableClusteredReflections;
 	// initialize per-view uniform buffer.  Pass in shadow info as necessary.
 	for (int32 ViewIndex = Views.Num() - 1; ViewIndex >= 0; --ViewIndex)
 	{
@@ -596,7 +597,7 @@ void FMobileSceneRenderer::InitViews(FRDGBuilder& GraphBuilder, FSceneTexturesCo
 		}
 	}
 
-	if (bDeferredShading)
+	if (bEnableClusteredShading)
 	{
 		SetupSceneReflectionCaptureBuffer(RHICmdList);
 	}
@@ -779,12 +780,19 @@ void FMobileSceneRenderer::Render(FRDGBuilder& GraphBuilder)
 	}
 
 	FSortedLightSetSceneInfo& SortedLightSet = *GraphBuilder.AllocObject<FSortedLightSetSceneInfo>();
-	if (bDeferredShading)
+
+	const bool bEnableClusteredShading = bEnableClusteredLocalLights || bEnableClusteredReflections;
+	if (bEnableClusteredShading)
 	{
-		GatherAndSortLights(SortedLightSet);
+		// Shadows are applied in clustered shading on mobile forward and separately on mobile deferred.
+		bool bShadowedLightsInClustered = bRequiresShadowProjections && !bDeferredShading;
+		GatherAndSortLights(SortedLightSet, bShadowedLightsInClustered);
 		int32 NumReflectionCaptures = Views[0].NumBoxReflectionCaptures + Views[0].NumSphereReflectionCaptures;
-		bool bCullLightsToGrid = (NumReflectionCaptures > 0 || GMobileUseClusteredDeferredShading != 0);
-		ComputeLightGrid(GraphBuilder, bCullLightsToGrid, SortedLightSet);
+		bool bCullLightsToGrid = ((bEnableClusteredReflections && NumReflectionCaptures > 0) || (bEnableClusteredLocalLights && SortedLightSet.ClusteredSupportedEnd > SortedLightSet.SimpleLightsEnd));
+		if (bCullLightsToGrid)
+		{
+			ComputeLightGrid(GraphBuilder, true, SortedLightSet);
+		}
 	}
 
 	// Generate the Sky/Atmosphere look up tables
@@ -846,7 +854,7 @@ void FMobileSceneRenderer::Render(FRDGBuilder& GraphBuilder)
 		SceneTextures.MobileSetupMode = EMobileSceneTextureSetupMode::SceneDepth;
 		SceneTextures.MobileUniformBuffer = CreateMobileSceneTextureUniformBuffer(GraphBuilder, SceneTextures.MobileSetupMode);
 
-		if (bRequiresDistanceField)
+		if (bRequiresShadowProjections)
 		{
 			CSV_SCOPED_TIMING_STAT_EXCLUSIVE(RenderMobileShadowProjections);
 			RDG_GPU_STAT_SCOPE(GraphBuilder, ShadowProjection);
@@ -1063,6 +1071,7 @@ void FMobileSceneRenderer::RenderForward(FRDGBuilder& GraphBuilder, FRDGTextureR
 		FMobileRenderPassParameters* PassParameters = GraphBuilder.AllocParameters<FMobileRenderPassParameters>();
 		PassParameters->View = View.GetShaderParameters();
 		PassParameters->MobileBasePass = CreateMobileBasePassUniformBuffer(GraphBuilder, View, EMobileBasePass::Opaque, SetupMode, MobileBasePassTextures);
+		PassParameters->ReflectionCapture = View.ReflectionCaptureUniformBuffer;
 		PassParameters->RenderTargets = BasePassRenderTargets;
 	
 		BuildInstanceCullingDrawParams(GraphBuilder, View, PassParameters);
@@ -1181,6 +1190,7 @@ void FMobileSceneRenderer::RenderForwardMultiPass(FRDGBuilder& GraphBuilder, FMo
 	FMobileRenderPassParameters* SecondPassParameters = GraphBuilder.AllocParameters<FMobileRenderPassParameters>();
 	*SecondPassParameters = *PassParameters;
 	SecondPassParameters->MobileBasePass = CreateMobileBasePassUniformBuffer(GraphBuilder, View, EMobileBasePass::Translucent, SetupMode);
+	SecondPassParameters->ReflectionCapture = View.ReflectionCaptureUniformBuffer;
 	SecondPassParameters->RenderTargets = BasePassRenderTargets;
 
 	GraphBuilder.AddPass(
@@ -1310,7 +1320,7 @@ void FMobileSceneRenderer::RenderDeferred(FRDGBuilder& GraphBuilder, const FSort
 
 	TArrayView<FRDGTextureRef> BasePassTexturesView = MakeArrayView(ColorTargets);
 
-	FRenderTargetBindingSlots BasePassRenderTargets = GetRenderTargetBindings(ERenderTargetLoadAction::ENoAction, BasePassTexturesView);
+	FRenderTargetBindingSlots BasePassRenderTargets = GetRenderTargetBindings(ERenderTargetLoadAction::EClear, BasePassTexturesView);
 	BasePassRenderTargets.DepthStencil = bIsFullDepthPrepassEnabled ? 
 		FDepthStencilBinding(SceneTextures.Depth.Target, ERenderTargetLoadAction::ELoad, ERenderTargetLoadAction::ELoad, FExclusiveDepthStencil::DepthRead_StencilWrite) : 
 		FDepthStencilBinding(SceneTextures.Depth.Target, ERenderTargetLoadAction::EClear, ERenderTargetLoadAction::EClear, FExclusiveDepthStencil::DepthWrite_StencilWrite);
@@ -1349,6 +1359,7 @@ void FMobileSceneRenderer::RenderDeferred(FRDGBuilder& GraphBuilder, const FSort
 		auto* PassParameters = GraphBuilder.AllocParameters<FMobileRenderPassParameters>();
 		PassParameters->View = View.GetShaderParameters();
 		PassParameters->MobileBasePass = CreateMobileBasePassUniformBuffer(GraphBuilder, View, EMobileBasePass::Opaque, SetupMode, MobileBasePassTextures);
+		PassParameters->ReflectionCapture = View.ReflectionCaptureUniformBuffer;
 		PassParameters->RenderTargets = BasePassRenderTargets;
 		
 		BuildInstanceCullingDrawParams(GraphBuilder, View, PassParameters);
@@ -1372,10 +1383,6 @@ void FMobileSceneRenderer::RenderDeferredSinglePass(FRDGBuilder& GraphBuilder, c
 		ERDGPassFlags::Raster,
 		[this, PassParameters, ViewIndex, NumViews, &View, &SceneTextures, &SortedLightSet, bUsingPixelLocalStorage](FRHICommandListImmediate& RHICmdList)
 	{
-		if (GIsEditor && !View.bIsSceneCapture)
-		{
-			DrawClearQuad(RHICmdList, View.BackgroundColor);
-		}
 		// Depth pre-pass
 		RHICmdList.SetCurrentStat(GET_STATID(STAT_CLM_MobilePrePass));
 		RenderMaskedPrePass(RHICmdList, View);
@@ -1397,7 +1404,7 @@ void FMobileSceneRenderer::RenderDeferredSinglePass(FRDGBuilder& GraphBuilder, c
 		RenderDecals(RHICmdList, View);
 		// SceneColor write, SceneDepth is read only
 		RHICmdList.NextSubpass();
-		MobileDeferredShadingPass(RHICmdList, ViewIndex, NumViews, View, *Scene, SortedLightSet);
+		MobileDeferredShadingPass(RHICmdList, ViewIndex, NumViews, View, *Scene, SortedLightSet, VisibleLightInfos);
 
 		if (bUsingPixelLocalStorage)
 		{
@@ -1422,11 +1429,6 @@ void FMobileSceneRenderer::RenderDeferredMultiPass(FRDGBuilder& GraphBuilder, cl
 		ERDGPassFlags::Raster,
 		[this, PassParameters, &View, &SceneTextures](FRHICommandListImmediate& RHICmdList)
 	{
-		if (GIsEditor && !View.bIsSceneCapture)
-		{
-			DrawClearQuad(RHICmdList, View.BackgroundColor);
-		}
-
 		// Depth pre-pass
 		RHICmdList.SetCurrentStat(GET_STATID(STAT_CLM_MobilePrePass));
 		RenderMaskedPrePass(RHICmdList, View);
@@ -1455,6 +1457,7 @@ void FMobileSceneRenderer::RenderDeferredMultiPass(FRDGBuilder& GraphBuilder, cl
 	auto* SecondPassParameters = GraphBuilder.AllocParameters<FMobileRenderPassParameters>();
 	*SecondPassParameters = *PassParameters;
 	SecondPassParameters->MobileBasePass = CreateMobileBasePassUniformBuffer(GraphBuilder, View, EMobileBasePass::Translucent, SetupMode);
+	SecondPassParameters->ReflectionCapture = View.ReflectionCaptureUniformBuffer;
 	SecondPassParameters->RenderTargets = BasePassRenderTargets;
 
 	GraphBuilder.AddPass(
@@ -1478,6 +1481,7 @@ void FMobileSceneRenderer::RenderDeferredMultiPass(FRDGBuilder& GraphBuilder, cl
 	auto* ThirdPassParameters = GraphBuilder.AllocParameters<FMobileRenderPassParameters>();
 	*ThirdPassParameters = *PassParameters;
 	ThirdPassParameters->MobileBasePass = CreateMobileBasePassUniformBuffer(GraphBuilder, View, EMobileBasePass::Translucent, SetupMode);
+	ThirdPassParameters->ReflectionCapture = View.ReflectionCaptureUniformBuffer;
 	ThirdPassParameters->RenderTargets = BasePassRenderTargets;
 
 	GraphBuilder.AddPass(
@@ -1487,7 +1491,7 @@ void FMobileSceneRenderer::RenderDeferredMultiPass(FRDGBuilder& GraphBuilder, cl
 		[this, ThirdPassParameters, ViewIndex, NumViews, &View, &SceneTextures, &SortedLightSet](FRHICommandListImmediate& RHICmdList)
 	{
 		RHICmdList.NextSubpass();
-		MobileDeferredShadingPass(RHICmdList, ViewIndex, NumViews, View, *Scene, SortedLightSet);
+		MobileDeferredShadingPass(RHICmdList, ViewIndex, NumViews, View, *Scene, SortedLightSet, VisibleLightInfos);
 		RenderFog(RHICmdList, View);
 		// Draw translucency.
 		RenderTranslucency(RHICmdList, View);
@@ -1688,75 +1692,6 @@ void FMobileSceneRenderer::PreTonemapMSAA(FRHICommandListImmediate& RHICmdList,
 		EDRF_UseTriangleOptimization);
 }
 
-/** Before SetupMobileBasePassAfterShadowInit, we need to update the uniform buffer and shadow info for all movable point lights.*/
-void FMobileSceneRenderer::UpdateMovablePointLightUniformBufferAndShadowInfo()
-{
-	static auto* MobileNumDynamicPointLightsCVar = IConsoleManager::Get().FindTConsoleVariableDataInt(TEXT("r.MobileNumDynamicPointLights"));
-	const int32 MobileNumDynamicPointLights = MobileNumDynamicPointLightsCVar->GetValueOnRenderThread();
-
-	static auto* MobileEnableMovableSpotlightsCVar = IConsoleManager::Get().FindTConsoleVariableDataInt(TEXT("r.Mobile.EnableMovableSpotlights"));
-	const int32 MobileEnableMovableSpotlights = MobileEnableMovableSpotlightsCVar->GetValueOnRenderThread();
-
-	static auto* EnableMovableSpotlightShadowsCVar = IConsoleManager::Get().FindTConsoleVariableDataInt(TEXT("r.Mobile.EnableMovableSpotlightsShadow"));
-	const int32 EnableMovableSpotlightShadows = EnableMovableSpotlightShadowsCVar->GetValueOnRenderThread();
-
-	if (MobileNumDynamicPointLights > 0)
-	{
-		bool bShouldDynamicShadows = ViewFamily.EngineShowFlags.DynamicShadows
-			&& !IsSimpleForwardShadingEnabled(ShaderPlatform)
-			&& GetShadowQuality() > 0
-			&& EnableMovableSpotlightShadows != 0;
-
-		for (auto LightIt = Scene->Lights.CreateConstIterator(); LightIt; ++LightIt)
-		{
-			const FLightSceneInfoCompact& LightSceneInfoCompact = *LightIt;
-			FLightSceneInfo* LightSceneInfo = LightSceneInfoCompact.LightSceneInfo;
-
-			FLightSceneProxy* LightProxy = LightSceneInfo->Proxy;
-			const uint8 LightType = LightProxy->GetLightType();
-
-			const bool bIsValidLightType =
-				LightType == LightType_Point
-				|| LightType == LightType_Rect
-				|| (LightType == LightType_Spot && MobileEnableMovableSpotlights);
-
-			if (bIsValidLightType && LightProxy->IsMovable())
-			{
-				LightSceneInfo->ConditionalUpdateMobileMovablePointLightUniformBuffer(this);
-
-				bool bDynamicShadows = bShouldDynamicShadows
-					&& LightType == LightType_Spot
-					&& VisibleLightInfos[LightSceneInfo->Id].AllProjectedShadows.Num() > 0
-					&& VisibleLightInfos[LightSceneInfo->Id].AllProjectedShadows.Last()->bAllocated;
-
-				if (bDynamicShadows)
-				{
-					FProjectedShadowInfo *ProjectedShadowInfo = VisibleLightInfos[LightSceneInfo->Id].AllProjectedShadows.Last();
-					checkSlow(ProjectedShadowInfo && ProjectedShadowInfo->CacheMode != SDCM_StaticPrimitivesOnly);
-
-					const FIntPoint ShadowBufferResolution = ProjectedShadowInfo->GetShadowBufferResolution();
-					
-					for (int32 ViewIndex = 0; ViewIndex < Views.Num(); ++ViewIndex)
-					{
-						FViewInfo& View = Views[ViewIndex];
-
-						FMobileMovableSpotLightsShadowInfo& MobileMovableSpotLightsShadowInfo = View.MobileMovableSpotLightsShadowInfo;
-
-						checkSlow(MobileMovableSpotLightsShadowInfo.ShadowDepthTexture == nullptr
-							|| MobileMovableSpotLightsShadowInfo.ShadowDepthTexture == ProjectedShadowInfo->RenderTargets.DepthTarget->GetRHI());
-
-						if (MobileMovableSpotLightsShadowInfo.ShadowDepthTexture == nullptr)
-						{
-							MobileMovableSpotLightsShadowInfo.ShadowDepthTexture = ProjectedShadowInfo->RenderTargets.DepthTarget->GetRHI();
-							MobileMovableSpotLightsShadowInfo.ShadowBufferSize = FVector4f(ShadowBufferResolution.X, ShadowBufferResolution.Y, 1.0f / ShadowBufferResolution.X, 1.0f / ShadowBufferResolution.Y);
-						}
-					}
-				}
-			}
-		}
-	}
-}
-
 bool FMobileSceneRenderer::ShouldRenderHZB()
 {
 	static const auto MobileAmbientOcclusionTechniqueCVar = IConsoleManager::Get().FindTConsoleVariableDataInt(TEXT("r.Mobile.AmbientOcclusionTechnique"));
@@ -1811,3 +1746,8 @@ void FMobileSceneRenderer::RenderHZB(FRDGBuilder& GraphBuilder, FRDGTextureRef S
 		}
 	}
 }
+
+bool FMobileSceneRenderer::AllowSimpleLights() const
+{
+	return FSceneRenderer::AllowSimpleLights() && bDeferredShading;
+}
diff --git a/Engine/Source/Runtime/Renderer/Private/PrimitiveSceneInfo.cpp b/Engine/Source/Runtime/Renderer/Private/PrimitiveSceneInfo.cpp
--- a/Engine/Source/Runtime/Renderer/Private/PrimitiveSceneInfo.cpp
+++ b/Engine/Source/Runtime/Renderer/Private/PrimitiveSceneInfo.cpp
@@ -168,7 +168,7 @@ FPrimitiveSceneInfo::FPrimitiveSceneInfo(UPrimitiveComponent* InComponent,FScene
 	LightList(NULL),
 	LastRenderTime(-FLT_MAX),
 	Scene(InScene),
-	NumMobileMovablePointLights(0),
+	NumMobileMovableLocalLights(0),
 	bShouldRenderInMainPass(InComponent->SceneProxy->ShouldRenderInMainPass()),
 	bVisibleInRealTimeSkyCapture(InComponent->SceneProxy->IsVisibleInRealTimeSkyCaptures()),
 #if RHI_RAYTRACING
diff --git a/Engine/Source/Runtime/Renderer/Private/RendererScene.cpp b/Engine/Source/Runtime/Renderer/Private/RendererScene.cpp
--- a/Engine/Source/Runtime/Renderer/Private/RendererScene.cpp
+++ b/Engine/Source/Runtime/Renderer/Private/RendererScene.cpp
@@ -1185,7 +1185,6 @@ FScene::FScene(UWorld* InWorld, bool bInRequiresHitProxies, bool bInIsEditorScen
 ,	NumUnbuiltReflectionCaptures(0)
 ,	NumMobileStaticAndCSMLights_RenderThread(0)
 ,	NumMobileMovableDirectionalLights_RenderThread(0)
-,	MobileWholeSceneShadowAtlasSize(0, 0)
 ,	GPUSkinCache(nullptr)
 ,	SceneLODHierarchy(this)
 ,	RuntimeVirtualTexturePrimitiveHideMaskEditor(0)
@@ -1863,7 +1862,7 @@ void FScene::AssignAvailableShadowMapChannelForLight(FLightSceneInfo* LightScene
 			return;
 		}
 	}
-	else if (LightSceneInfo->Proxy->GetLightType() == LightType_Directional)
+	else if (LightSceneInfo->Proxy->GetLightType() == LightType_Directional && !IsMobilePlatform(GetShaderPlatform()))
 	{
 		// The implementation of forward lighting in ShadowProjectionPixelShader.usf does not support binding the directional light to channel 3.
 		// This is related to the USE_FADE_PLANE feature that encodes the CSM blend factor the alpha channel.
@@ -1965,8 +1964,9 @@ void FScene::AddLightSceneInfo_RenderThread(FLightSceneInfo* LightSceneInfo)
 		RectProxy->AtlasSlotIndex = RectLightAtlas::AddRectLightTexture(RectProxy->SourceTexture);
 	}
 
-	const bool bForwardShading = IsForwardShadingEnabled(GetShaderPlatform());
-	if (bForwardShading && (LightSceneInfo->Proxy->CastsDynamicShadow() || LightSceneInfo->Proxy->GetLightFunctionMaterial()))
+	const EShaderPlatform ShaderPlatform = GetShaderPlatform();
+	const bool bAssignShadowMapChannel = IsForwardShadingEnabled(ShaderPlatform) ||  (IsMobilePlatform(ShaderPlatform) && MobileUsesShadowMaskTexture(ShaderPlatform));
+	if (bAssignShadowMapChannel && (LightSceneInfo->Proxy->CastsDynamicShadow() || LightSceneInfo->Proxy->GetLightFunctionMaterial()))
 	{
 		AssignAvailableShadowMapChannelForLight(LightSceneInfo);
 	}
@@ -3122,7 +3122,6 @@ void FScene::UpdateLightColorAndBrightness(ULightComponent* Light)
 					{
 						Scene->Lights[ LightSceneInfo->Id ].Color = NewParameters.NewColor;
 					}
-					LightSceneInfo->Proxy->SetMobileMovablePointLightUniformBufferNeedsUpdate(true);
 				}
 			});
 	}
@@ -3842,12 +3841,16 @@ void FScene::UpdateEarlyZPassMode()
 	{
 		EarlyZPassMode = DDM_None;
 
-		if (MaskedInEarlyPass(GetFeatureLevelShaderPlatform(FeatureLevel)))
+		const EShaderPlatform ShaderPlatform = GetShaderPlatform();
+		static const bool bMaskedInEarlyPass = MaskedInEarlyPass(ShaderPlatform);
+		if (bMaskedInEarlyPass)
 		{
 			EarlyZPassMode = DDM_MaskedOnly;
 		}
 
-		if (IsMobileDistanceFieldEnabled(GetShaderPlatform()) || IsMobileAmbientOcclusionEnabled(GetShaderPlatform()))
+		const bool bIsMobileAmbientOcclusionEnabled = IsMobileAmbientOcclusionEnabled(ShaderPlatform);
+		const bool bMobileUsesShadowMaskTexture = MobileUsesShadowMaskTexture(ShaderPlatform);
+		if (bIsMobileAmbientOcclusionEnabled || bMobileUsesShadowMaskTexture)
 		{
 			EarlyZPassMode = DDM_AllOpaque;
 		}
diff --git a/Engine/Source/Runtime/Renderer/Private/SceneCore.cpp b/Engine/Source/Runtime/Renderer/Private/SceneCore.cpp
--- a/Engine/Source/Runtime/Renderer/Private/SceneCore.cpp
+++ b/Engine/Source/Runtime/Renderer/Private/SceneCore.cpp
@@ -172,7 +172,7 @@ FLightPrimitiveInteraction::FLightPrimitiveInteraction(
 	bHasTranslucentObjectShadow(bInHasTranslucentObjectShadow),
 	bHasInsetObjectShadow(bInHasInsetObjectShadow),
 	bSelfShadowOnly(false),
-	bMobileDynamicPointLight(false)
+	bMobileDynamicLocalLight(false)
 {
 	// Determine whether this light-primitive interaction produces a shadow.
 	if(PrimitiveSceneInfo->Proxy->HasStaticLighting())
@@ -237,19 +237,21 @@ FLightPrimitiveInteraction::FLightPrimitiveInteraction(
 		if (PrimitiveSceneInfo->Scene->GetShadingPath() == EShadingPath::Mobile && LightSceneInfo->Proxy->IsMovable())
 		{
 			const uint8 LightType = LightSceneInfo->Proxy->GetLightType();
-			static const auto MobileEnableMovableSpotLightsVar = IConsoleManager::Get().FindTConsoleVariableDataInt(TEXT("r.Mobile.EnableMovableSpotLights"));
 
 			const bool bIsValidLightType = 
 				   LightType == LightType_Rect
 				|| LightType == LightType_Point
-				|| (LightType == LightType_Spot && MobileEnableMovableSpotLightsVar->GetValueOnRenderThread());
+				|| LightType == LightType_Spot;
 
 			if( bIsValidLightType )
 			{
-				bMobileDynamicPointLight = true;
-				PrimitiveSceneInfo->NumMobileMovablePointLights++;
-				// The mobile renderer needs to update the shader bindings of movable point lights uniform buffer, so we have to update any static meshes in drawlists
-				PrimitiveSceneInfo->BeginDeferredUpdateStaticMeshes();
+				bMobileDynamicLocalLight = true;
+				PrimitiveSceneInfo->NumMobileMovableLocalLights++;
+				if (PrimitiveSceneInfo->NumMobileMovableLocalLights == 1)
+				{
+					// Update static meshes to choose the shader permutation with local lights.
+					PrimitiveSceneInfo->BeginDeferredUpdateStaticMeshes();
+				}
 			} 
 		}
 	}
@@ -306,12 +308,15 @@ FLightPrimitiveInteraction::~FLightPrimitiveInteraction()
 
 	FlushCachedShadowMapData();
 
-	// Track mobile movable point light count
-	if (bMobileDynamicPointLight)
+	// Track mobile movable local light count
+	if (bMobileDynamicLocalLight)
 	{
-		PrimitiveSceneInfo->NumMobileMovablePointLights--;
-		// The mobile renderer needs to use a different shader for movable point lights, so we have to update any static meshes in drawlists
-		PrimitiveSceneInfo->BeginDeferredUpdateStaticMeshes();
+		PrimitiveSceneInfo->NumMobileMovableLocalLights--;
+		if (PrimitiveSceneInfo->NumMobileMovableLocalLights == 0)
+		{
+			// Update static meshes to choose the shader permutation without local lights.
+			PrimitiveSceneInfo->BeginDeferredUpdateStaticMeshes();
+		}
 	}
 
 	// Remove the interaction from the light's interaction list.
diff --git a/Engine/Source/Runtime/Renderer/Private/SceneCore.h b/Engine/Source/Runtime/Renderer/Private/SceneCore.h
--- a/Engine/Source/Runtime/Renderer/Private/SceneCore.h
+++ b/Engine/Source/Runtime/Renderer/Private/SceneCore.h
@@ -41,7 +41,7 @@ public:
 	bool HasTranslucentObjectShadow() const { return bHasTranslucentObjectShadow; }
 	bool HasInsetObjectShadow() const { return bHasInsetObjectShadow; }
 	bool CastsSelfShadowOnly() const { return bSelfShadowOnly; }
-	bool IsMobileDynamicPointLight() const { return bMobileDynamicPointLight; }
+	bool IsMobileDynamicLocalLight() const { return bMobileDynamicLocalLight; }
 
 	FORCEINLINE bool IsNaniteMeshProxy() const { return bNaniteMeshProxy; }
 	FORCEINLINE bool ProxySupportsGPUScene() const { return bProxySupportsGPUScene; }
@@ -112,8 +112,8 @@ private:
 	/** True if the primitive only shadows itself. */
 	uint32 bSelfShadowOnly : 1;
 
-	/** True this is a mobile dynamic point light interaction. */
-	uint32 bMobileDynamicPointLight : 1;
+	/** True this is a mobile dynamic local light interaction. */
+	uint32 bMobileDynamicLocalLight : 1;
 
 	/** If true then all meshes drawn by the primitive scene proxy are Nanite meshes. Caches the result of FPrimitiveSceneProxy::IsNaniteMesh() */
 	uint32 bNaniteMeshProxy : 1;
diff --git a/Engine/Source/Runtime/Renderer/Private/ScenePrivate.h b/Engine/Source/Runtime/Renderer/Private/ScenePrivate.h
--- a/Engine/Source/Runtime/Renderer/Private/ScenePrivate.h
+++ b/Engine/Source/Runtime/Renderer/Private/ScenePrivate.h
@@ -3038,9 +3038,6 @@ public:
 	int32 NumMobileStaticAndCSMLights_RenderThread;
 	int32 NumMobileMovableDirectionalLights_RenderThread;
 
-	/** Cached shadow atlas size for mobile, since if the resolution is changed we have to update the movable point lights uniform buffer. */
-	FIntPoint MobileWholeSceneShadowAtlasSize;
-
 	FSceneVelocityData VelocityData;
 
 	/** GPU Skinning cache, if enabled */
diff --git a/Engine/Source/Runtime/Renderer/Private/SceneRendering.h b/Engine/Source/Runtime/Renderer/Private/SceneRendering.h
--- a/Engine/Source/Runtime/Renderer/Private/SceneRendering.h
+++ b/Engine/Source/Runtime/Renderer/Private/SceneRendering.h
@@ -132,13 +132,6 @@ protected:
 	TArray<const FPrimitiveSceneInfo*, SceneRenderingAllocator> ShadowSubjectPrimitives;
 };
 
-class FMobileMovableSpotLightsShadowInfo
-{
-public:
-	FVector4f ShadowBufferSize = FVector4f(0.0f);
-	FRHITexture* ShadowDepthTexture = nullptr;
-};
-
 /** Information about a visible light which is specific to the view it's visible in. */
 class FVisibleLightViewInfo
 {
@@ -1299,9 +1292,6 @@ public:
 
 	FHairStrandsViewData HairStrandsViewData;
 
-	//Spotlight shadow info for mobile.
-	FMobileMovableSpotLightsShadowInfo MobileMovableSpotLightsShadowInfo;
-
 	/** Parameters for exponential height fog. */
 	FVector4f ExponentialFogParameters;
 	FVector4f ExponentialFogParameters2;
@@ -1908,6 +1898,7 @@ public:
 	virtual void RenderHitProxies(FRDGBuilder& GraphBuilder) {}
 	virtual bool ShouldRenderVelocities() const { return false; }
 	virtual bool ShouldRenderPrePass() const { return false; }
+	virtual bool AllowSimpleLights() const;
 
 	/** Creates a scene renderer based on the current feature level. */
 	RENDERER_API static FSceneRenderer* CreateSceneRenderer(const FSceneViewFamily* InViewFamily, FHitProxyConsumer* HitProxyConsumer);
@@ -2318,6 +2309,7 @@ public:
 
 	virtual bool ShouldRenderPrePass() const override;
 
+	virtual bool AllowSimpleLights() const override;
 
 protected:
 	/** Finds the visible dynamic shadows for each view. */
@@ -2395,9 +2387,6 @@ protected:
 
 	void RenderPixelProjectedReflection(FRDGBuilder& GraphBuilder, FRDGTextureRef SceneColorTexture, FRDGTextureRef SceneDepthTexture, FRDGTextureRef PixelProjectedReflectionTexture, const FPlanarReflectionSceneProxy* PlanarReflectionSceneProxy);
 
-	/** Before SetupMobileBasePassAfterShadowInit, we need to update the uniform buffer and shadow info for all movable point lights.*/
-	void UpdateMovablePointLightUniformBufferAndShadowInfo();
-
 	void RenderMobileShadowProjections(FRDGBuilder& GraphBuilder, FRDGTextureRef SceneDepthTexture);
 private:
 	const bool bGammaSpace;
@@ -2417,6 +2406,9 @@ private:
 	bool bIsFullDepthPrepassEnabled;
 	bool bIsMaskedOnlyDepthPrepassEnabled;
 	bool bRequiresSceneDepthAux;
+	bool bEnableClusteredLocalLights;
+	bool bEnableClusteredReflections;
+	bool bRequiresShadowProjections;
 
 	ETranslucencyPass::Type StandardTranslucencyPass;
 	EMeshPass::Type StandardTranslucencyMeshPass;
diff --git a/Engine/Source/Runtime/Renderer/Private/ShadowDepthRendering.cpp b/Engine/Source/Runtime/Renderer/Private/ShadowDepthRendering.cpp
--- a/Engine/Source/Runtime/Renderer/Private/ShadowDepthRendering.cpp
+++ b/Engine/Source/Runtime/Renderer/Private/ShadowDepthRendering.cpp
@@ -280,9 +280,9 @@ static TAutoConsoleVariable<int32> CVarSupportPointLightWholeSceneShadows(
 static bool MobileUsesPerspectiveCorrectShadowPermutation(EShaderPlatform ShaderPlatform)
 {
 	// Required only for spotlight shadows on mobile
-	static FShaderPlatformCachedIniValue<bool> MobileEnableMovableSpotlightsShadowIniValue(TEXT("r.Mobile.EnableMovableSpotlightsShadow"));
-	const bool bMobileEnableMovableSpotlightsShadow = (MobileEnableMovableSpotlightsShadowIniValue.Get(ShaderPlatform) != 0);
-	return bMobileEnableMovableSpotlightsShadow;
+	static FShaderPlatformCachedIniValue<bool> MobileEnableMovableSpotlightShadowsIniValue(TEXT("r.Mobile.EnableMovableSpotlightsShadow"));
+	const bool bMobileEnableMovableSpotlightShadows = (MobileEnableMovableSpotlightShadowsIniValue.Get(ShaderPlatform) != 0);
+	return bMobileEnableMovableSpotlightShadows;
 }
 
 
diff --git a/Engine/Source/Runtime/Renderer/Private/ShadowRendering.cpp b/Engine/Source/Runtime/Renderer/Private/ShadowRendering.cpp
--- a/Engine/Source/Runtime/Renderer/Private/ShadowRendering.cpp
+++ b/Engine/Source/Runtime/Renderer/Private/ShadowRendering.cpp
@@ -2371,7 +2371,7 @@ void FDeferredShadingSceneRenderer::RenderDeferredShadowProjections(
 
 void FMobileSceneRenderer::RenderModulatedShadowProjections(FRHICommandListImmediate& RHICmdList, int32 ViewIndex, const FViewInfo& View)
 {
-	if (IsSimpleForwardShadingEnabled(ShaderPlatform) || !ViewFamily.EngineShowFlags.DynamicShadows || View.bIsPlanarReflection)
+	if (IsSimpleForwardShadingEnabled(ShaderPlatform) || !ViewFamily.EngineShowFlags.DynamicShadows || View.bIsPlanarReflection || bRequiresShadowProjections)
 	{
 		return;
 	}
@@ -2420,7 +2420,7 @@ void FMobileSceneRenderer::RenderModulatedShadowProjections(FRHICommandListImmed
 	}
 }
 
-void InitMobileSDFShadowingOutputs(FRHICommandListImmediate& RHICmdList, const FIntPoint& Extent)
+void InitMobileShadowProjectionOutputs(FRHICommandListImmediate& RHICmdList, const FIntPoint& Extent)
 {
 	const FIntPoint& BufferSize = Extent;
 
@@ -2435,6 +2435,8 @@ void FMobileSceneRenderer::RenderMobileShadowProjections(
 	FRDGBuilder& GraphBuilder, 
 	FRDGTextureRef SceneDepthTexture)
 {
+	RDG_RHI_EVENT_SCOPE(GraphBuilder, RenderMobileShadowProjections);
+
 	FRDGTextureRef ScreenShadowMaskTexture = GraphBuilder.RegisterExternalTexture(GScreenSpaceShadowMaskTextureMobileOutputs.ScreenSpaceShadowMaskTextureMobile, TEXT("ScreenSpaceShadowMaskTextureMobile"));
 	AddClearRenderTargetPass(GraphBuilder, ScreenShadowMaskTexture);
 
@@ -2447,15 +2449,19 @@ void FMobileSceneRenderer::RenderMobileShadowProjections(
 		const FVisibleLightInfo& VisibleLightInfo = VisibleLightInfos[LightSceneInfo->Id];
 		const FLightSceneProxy* LightSceneProxy = LightSceneInfo->Proxy;
 
-		RenderShadowProjections(GraphBuilder, SceneTextures,
-			ScreenShadowMaskTexture,
-			nullptr,
-			LightSceneInfo,
-			false);
+		// Local light shadows don't render to shadow mask texture on mobile deferred
+		if (LightSceneProxy->GetLightType() == LightType_Directional || !IsMobileDeferredShadingEnabled(ShaderPlatform))
+		{
+			RenderShadowProjections(GraphBuilder, SceneTextures,
+				ScreenShadowMaskTexture,
+				nullptr,
+				LightSceneInfo,
+				true);
+		}
 	}
 }
 
-void ReleaseMobileSDFShadowingOutputs()
+void ReleaseMobileShadowProjectionOutputs()
 {
 	GScreenSpaceShadowMaskTextureMobileOutputs.Release();
 }
diff --git a/Engine/Source/Runtime/Renderer/Private/ShadowRendering.h b/Engine/Source/Runtime/Renderer/Private/ShadowRendering.h
--- a/Engine/Source/Runtime/Renderer/Private/ShadowRendering.h
+++ b/Engine/Source/Runtime/Renderer/Private/ShadowRendering.h
@@ -59,8 +59,8 @@ void OverrideWithDefaultMaterialForShadowDepth(
 	ERHIFeatureLevel::Type InFeatureLevel
 	);
 
-void InitMobileSDFShadowingOutputs(FRHICommandListImmediate& RHICmdList, const FIntPoint& Extent);
-void ReleaseMobileSDFShadowingOutputs();
+void InitMobileShadowProjectionOutputs(FRHICommandListImmediate& RHICmdList, const FIntPoint& Extent);
+void ReleaseMobileShadowProjectionOutputs();
 
 enum EShadowDepthRenderMode
 {
@@ -1176,7 +1176,7 @@ public:
 	static bool ShouldCompilePermutation(const FGlobalShaderPermutationParameters& Parameters)
 	{
 		return IsFeatureLevelSupported(Parameters.Platform, ERHIFeatureLevel::SM5)
-			|| (bUseTransmission == 0 && SubPixelShadow == 0 && IsMobileDistanceFieldEnabled(Parameters.Platform));
+			|| (bUseTransmission == 0 && SubPixelShadow == 0 && MobileUsesShadowMaskTexture(Parameters.Platform));
 	}
 
 	/**
@@ -1190,6 +1190,9 @@ public:
 		OutEnvironment.SetDefine(TEXT("SUBPIXEL_SHADOW"), (uint32)(SubPixelShadow ? 1 : 0));
 		OutEnvironment.SetDefine(TEXT("USE_FADE_PLANE"), (uint32)(bUseFadePlane ? 1 : 0));
 		OutEnvironment.SetDefine(TEXT("USE_TRANSMISSION"), (uint32)(bUseTransmission ? 1 : 0));
+
+		const bool bMobileVulkanForceDepthRead = IsVulkanMobilePlatform(Parameters.Platform) && MobileUsesShadowMaskTexture(Parameters.Platform);
+		OutEnvironment.SetDefine(TEXT("FORCE_DEPTH_TEXTURE_READS"), (uint32)(bMobileVulkanForceDepthRead ? 1 : 0));
 	}
 
 	/**
diff --git a/Engine/Source/Runtime/Renderer/Private/ShadowSetup.cpp b/Engine/Source/Runtime/Renderer/Private/ShadowSetup.cpp
--- a/Engine/Source/Runtime/Renderer/Private/ShadowSetup.cpp
+++ b/Engine/Source/Runtime/Renderer/Private/ShadowSetup.cpp
@@ -5173,11 +5173,11 @@ void FSceneRenderer::AllocateShadowDepthTargets(FRHICommandListImmediate& RHICmd
 		MobileDynamicSpotlightShadows.Sort(FCompareFProjectedShadowInfoByResolution());
 
 		//Limit the number of spotlights shadow for performance reason
-		static const auto MobileMaxVisibleMovableSpotLightsShadowCVar = IConsoleManager::Get().FindTConsoleVariableDataInt(TEXT("r.Mobile.MaxVisibleMovableSpotLightsShadow"));
-		if (MobileMaxVisibleMovableSpotLightsShadowCVar)
+		static const auto MobileMaxVisibleMovableSpotLightShadowsCVar = IConsoleManager::Get().FindTConsoleVariableDataInt(TEXT("r.Mobile.MaxVisibleMovableSpotLightShadows"));
+		if (MobileMaxVisibleMovableSpotLightShadowsCVar)
 		{
-			int32 MobileMaxVisibleMovableSpotLightsShadow = MobileMaxVisibleMovableSpotLightsShadowCVar->GetValueOnRenderThread();
-			MobileDynamicSpotlightShadows.RemoveAt(MobileMaxVisibleMovableSpotLightsShadow, FMath::Max(MobileDynamicSpotlightShadows.Num() - MobileMaxVisibleMovableSpotLightsShadow, 0), false);
+			int32 MobileMaxVisibleMovableSpotLightShadows = MobileMaxVisibleMovableSpotLightShadowsCVar->GetValueOnRenderThread();
+			MobileDynamicSpotlightShadows.RemoveAt(MobileMaxVisibleMovableSpotLightShadows, FMath::Max(MobileDynamicSpotlightShadows.Num() - MobileMaxVisibleMovableSpotLightShadows, 0), false);
 		}
 
 		MobileWholeSceneDirectionalShadows.Append(MobileDynamicSpotlightShadows);
@@ -5629,7 +5629,7 @@ FDynamicShadowsTaskData* FSceneRenderer::BeginInitDynamicShadows(bool bRunningEa
 
 	const bool bProjectEnablePointLightShadows = Scene->ReadOnlyCVARCache.bEnablePointLightShadows && !bMobile; // Point light shadow is unsupported on mobile for now.
 	const bool bProjectEnableMovableDirectionLightShadows = !bMobile || Scene->ReadOnlyCVARCache.bMobileAllowMovableDirectionalLights;
-	const bool bProjectEnableMovableSpotLightShadows = !bMobile || Scene->ReadOnlyCVARCache.bMobileEnableMovableSpotlightsShadow;
+	const bool bProjectEnableMovableSpotLightShadows = !bMobile || IsMobileMovableSpotlightShadowsEnabled(ShaderPlatform);
 
 	uint32 NumPointShadowCachesUpdatedThisFrame = 0;
 	uint32 NumSpotShadowCachesUpdatedThisFrame = 0;
@@ -5871,10 +5871,6 @@ void FSceneRenderer::AllocateMobileCSMAndSpotLightShadowDepthTargets(FRHICommand
 				WholeSceneAtlasSize.Y = 1 << FMath::CeilLogTwo(WholeSceneAtlasSize.Y);
 			}
 
-			bool bResolutionChanged = Scene->MobileWholeSceneShadowAtlasSize != WholeSceneAtlasSize;
-
-			Scene->MobileWholeSceneShadowAtlasSize = WholeSceneAtlasSize;
-
 			FPooledRenderTargetDesc WholeSceneShadowMapDesc2D(FPooledRenderTargetDesc::Create2DDesc(WholeSceneAtlasSize, PF_ShadowDepth, FClearValueBinding::DepthOne, TexCreate_None, TexCreate_DepthStencilTargetable | TexCreate_ShaderResource, false));
 			WholeSceneShadowMapDesc2D.Flags |= GFastVRamConfig.ShadowCSM;
 			GRenderTargetPool.FindFreeElement(RHICmdList, WholeSceneShadowMapDesc2D, ShadowMapAtlas.RenderTargets.DepthTarget, TEXT("MobileCSMAndSpotLightShadowmap"));
@@ -5888,12 +5884,6 @@ void FSceneRenderer::AllocateMobileCSMAndSpotLightShadowDepthTargets(FRHICommand
 					ProjectedShadowInfo->RenderTargets.CopyReferencesFromRenderTargets(ShadowMapAtlas.RenderTargets);
 					ProjectedShadowInfo->SetupShadowDepthView(this);
 					ShadowMapAtlas.Shadows.Add(ProjectedShadowInfo);
-
-					if (bResolutionChanged)
-					{
-						FLightSceneInfo* LightSceneInfo = (FLightSceneInfo*)(&ProjectedShadowInfo->GetLightSceneInfo());
-						LightSceneInfo->Proxy->SetMobileMovablePointLightUniformBufferNeedsUpdate(true);
-					}
 				}
 			}
 		}
diff --git a/Engine/Source/Runtime/Renderer/Public/PrimitiveSceneInfo.h b/Engine/Source/Runtime/Renderer/Public/PrimitiveSceneInfo.h
--- a/Engine/Source/Runtime/Renderer/Public/PrimitiveSceneInfo.h
+++ b/Engine/Source/Runtime/Renderer/Public/PrimitiveSceneInfo.h
@@ -378,8 +378,8 @@ public:
 	/** The scene the primitive is in. */
 	FScene* Scene;
 
-	/** The number of movable point lights for mobile */
-	int32 NumMobileMovablePointLights;
+	/** The number of movable local lights for mobile */
+	int32 NumMobileMovableLocalLights;
 
 	/** Set to true for the primitive to be rendered in the main pass to be visible in a view. */
 	bool bShouldRenderInMainPass : 1;
