VulkanRHI: Nanite's VisBuffer wasn't being included in BuildHzb for the R64 format (used only by Vulkan).  This should help Nanite performance in Vulkan.

diff --git a/Engine/Shaders/Private/HZB.usf b/Engine/Shaders/Private/HZB.usf
--- a/Engine/Shaders/Private/HZB.usf
+++ b/Engine/Shaders/Private/HZB.usf
@@ -181,7 +181,7 @@ void HZBBuildCS(
 	float4 DeviceZ = Gather4(ParentTextureMip, ParentTextureMipSampler, BufferUV);
 
 	uint2 SrcPixelCoord = (DispatchThreadId << 1) + PixelViewPortMinMax.xy;
-#if VIS_BUFFER_FORMAT == 1
+#if (VIS_BUFFER_FORMAT == 1) || (VIS_BUFFER_FORMAT == 3)
 	DeviceZ = max(DeviceZ, Gather4VisZ(BufferUV, SrcPixelCoord));
 #elif VIS_BUFFER_FORMAT == 2
 	DeviceZ = Gather4VisZ(BufferUV, SrcPixelCoord);
