Strata - fix for warning due to type.

diff --git a/Engine/Shaders/Private/Strata/Strata.ush b/Engine/Shaders/Private/Strata/Strata.ush
--- a/Engine/Shaders/Private/Strata/Strata.ush
+++ b/Engine/Shaders/Private/Strata/Strata.ush
@@ -847,7 +847,7 @@ void StrataPackOutStrataTree(
 
 		WRITE_FLOAT3(StrataTree.BSDFs[i].Emissive);
 		WRITE_FLOAT(StrataTree.BSDFs[i].ThicknessCm);
-		WRITE_FLOAT(StrataTree.BSDFs[i].TmpMFP);
+		WRITE_FLOAT3(StrataTree.BSDFs[i].TmpMFP);
 
 		WRITE_FLOAT(StrataTree.BSDFs[i].TopLayerDataWeight);
 	#endif
@@ -902,7 +902,7 @@ void StrataUnpackInStrataTree(
 
 		READ_FLOAT3(StrataTree.BSDFs[i].Emissive);
 		READ_FLOAT(StrataTree.BSDFs[i].ThicknessCm);
-		READ_FLOAT(StrataTree.BSDFs[i].TmpMFP);
+		READ_FLOAT3(StrataTree.BSDFs[i].TmpMFP);
 
 		READ_FLOAT(StrataTree.BSDFs[i].TopLayerDataWeight);
 	#endif
