Fix for broken fog visible line due to bad LWC math. This is a local (correct) fix and will be brought up as a problem.
A follow up after 19765907.

diff --git a/Engine/Shaders/Private/MobileFog.usf b/Engine/Shaders/Private/MobileFog.usf
--- a/Engine/Shaders/Private/MobileFog.usf
+++ b/Engine/Shaders/Private/MobileFog.usf
@@ -49,8 +49,10 @@ void MobileFogPS(
 	if (ResolvedView.SkyAtmosphereApplyCameraAerialPerspectiveVolume > 0.0f)
 	{
 		float3 WorldPosition = WorldPositionRelativeToCamera + ResolvedView.TranslatedWorldCameraOrigin - LWCHackToFloat(ResolvedView.PreViewTranslation);
-		float4 NDCPosition = LWCMultiply(MakeLWCVector4(WorldPosition, 1.0f), ResolvedView.WorldToClip);
-				
+
+	//	float4 NDCPosition = LWCMultiply(MakeLWCVector4(WorldPosition, 1.0f), PrimaryView.WorldToClip);		// LWC_TODO: MakeLWCVector4 seems broken
+		float4 NDCPosition = LWCMultiply(LWCPromote(float4(WorldPosition, 1.0f)), PrimaryView.WorldToClip);	// Using LWCPromote make the math work as expected.
+
 		// Sample the aerial perspective (AP). It is also blended under the VertexFog parameter.
 		FinalFog = GetAerialPerspectiveLuminanceTransmittanceWithFogOver(
 			ResolvedView.RealTimeReflectionCapture, 
