Fix unlit view mode with Strata.

diff --git a/Engine/Shaders/Private/BasePassPixelShader.usf b/Engine/Shaders/Private/BasePassPixelShader.usf
--- a/Engine/Shaders/Private/BasePassPixelShader.usf
+++ b/Engine/Shaders/Private/BasePassPixelShader.usf
@@ -1775,6 +1775,11 @@ void FPixelShaderInOut_MainPS(
 			StrataStruct.SSSTextureUAV[PixelPos] = StrataPackSSSData(SSSData);
 		}
 
+		// Unlit view mode
+	#if USE_DEVELOPMENT_SHADERS
+		Color = lerp(Color, TopLayerData.BaseColor, View.UnlitViewmodeMask);
+	#endif
+
 		Color += EmissiveLuminance;
 	}
 
diff --git a/Engine/Shaders/Private/Strata/Strata.ush b/Engine/Shaders/Private/Strata/Strata.ush
--- a/Engine/Shaders/Private/Strata/Strata.ush
+++ b/Engine/Shaders/Private/Strata/Strata.ush
@@ -535,6 +535,7 @@ FStrataIrradianceAndOcclusion StrataUnpackIrradianceAndOcclusion(uint In /*8bits
 // Top layer
 struct FStrataTopLayerData
 {
+	float3 BaseColor; // This data is only available during the base pass. It is not packed/stored otherwise.
 	float3 WorldNormal;
 	float  Roughness;
 	bool   bIsValid;
@@ -2928,6 +2929,8 @@ void PackStrataOut(
 					}
 					}
 
+					TopLayerData.BaseColor += BSDF.TopLayerDataWeight * StrataGetBSDFBaseColor(BSDF);
+
 					float3x3 TangentBasis = StrataGetBSDFSharedBasis_InlineShading(StrataPixelHeader, BSDF, StrataAddressing);
 					float3 N = TangentBasis[2];
 					TopLayerData.WorldNormal += N * BSDF.TopLayerDataWeight;
@@ -2950,6 +2953,7 @@ void PackStrataOut(
 		// Finalize top layer data
 		TopLayerData.WorldNormal= TopLayerTotalWeight > 0.0f	? normalize(TopLayerData.WorldNormal)			: 0.0f;
 		TopLayerData.Roughness	= TopLayerTotalWeight > 0.0f	? TopLayerData.Roughness / TopLayerTotalWeight	: 0.0f;
+		TopLayerData.BaseColor	= TopLayerTotalWeight > 0.0f	? TopLayerData.BaseColor / TopLayerTotalWeight	: 0.0f;
 
 		// Set storage layout as either: fast(0), single(1), or complex(2)
 		AndIsSimpleMaterial(StrataPixelHeader, BSDFVisibleCount == 1);
