Prevent clip(x) from being used in Nanite compute shaders because CS doesn't support clip

diff --git a/Engine/Shaders/Private/MaterialTemplate.ush b/Engine/Shaders/Private/MaterialTemplate.ush
--- a/Engine/Shaders/Private/MaterialTemplate.ush
+++ b/Engine/Shaders/Private/MaterialTemplate.ush
@@ -137,6 +137,9 @@
 #if IS_NANITE_PASS
 	#define ALLOW_CLIP 0
 	#define CLIP(x)
+	#if COMPUTESHADER
+		#define clip(x)
+	#endif
 #else
 	#define ALLOW_CLIP 1
 	#define CLIP(x) clip(x)
