Define ddx, ddy, and fwidth to 0 for Nanite computer shaders

diff --git a/Engine/Shaders/Private/MaterialTemplate.ush b/Engine/Shaders/Private/MaterialTemplate.ush
--- a/Engine/Shaders/Private/MaterialTemplate.ush
+++ b/Engine/Shaders/Private/MaterialTemplate.ush
@@ -139,6 +139,9 @@
 	#define CLIP(x)
 	#if COMPUTESHADER
 		#define clip(x)
+		#define ddx(x) 0
+		#define ddy(x) 0
+		#define fwidth(x) 0
 	#endif
 #else
 	#define ALLOW_CLIP 1
