Fix Lumen composition when Strata rough refraction is enabled.

diff --git a/Engine/Shaders/Private/DiffuseIndirectComposite.usf b/Engine/Shaders/Private/DiffuseIndirectComposite.usf
--- a/Engine/Shaders/Private/DiffuseIndirectComposite.usf
+++ b/Engine/Shaders/Private/DiffuseIndirectComposite.usf
@@ -62,6 +62,11 @@
   #endif
 #endif
 
+#if STRATA_ENABLED && STRATA_OPAQUE_ROUGH_REFRACTION_ENABLED
+RWTexture2D<float3> OutOpaqueRoughRefractionSceneColor;
+RWTexture2D<float3> OutSubSurfaceSceneColor;
+#endif
+
 //------------------------------------------------------- INCLUDES
 
 #include "Common.ush"
@@ -186,13 +191,27 @@ float3 CombineRoughSpecular(FGBufferData GBuffer, float NoV, float3 RayTracedRef
 	return Lighting;
 }
 
+float GetSSSCheckerboadSpecularScale(uint2 PixelPos, bool bNeedsSeparateLightAccumulation)
+{
+	float SpecularScale = 1.0f;
+	if (bNeedsSeparateLightAccumulation &&
+		View.bSubsurfacePostprocessEnabled > 0 && View.bCheckerboardSubsurfaceProfileRendering > 0)
+	{
+		bool bChecker = CheckerFromPixelPos(PixelPos);
+		// Adjust for checkerboard. only apply non-diffuse lighting (including emissive) 
+		// to the specular component, otherwise lighting is applied twice
+		SpecularScale = !bChecker;
+	}
+	return SpecularScale;
+}
+
 ///////////////////////////////////////////////////////////////////////////////////////////////////////////////////////////////////////
 // Lumen diffuse evaluation for strata material.
 
 // STRATA_TODO: add tile permutation for reducing VGPR pressures
 
 #if STRATA_ENABLED
-FDirectLighting StrataLumenLighting(
+FStrataDeferredLighting StrataLumenLighting(
 	uint2 PixelPos,
 	FStrataMaterialContainer MaterialBuffer,
 	FStrataAddressing StrataAddressing,
@@ -203,7 +222,7 @@ FDirectLighting StrataLumenLighting(
 	in float3 InSpecularLighting,
 	in float3 InRoughSpecularIndirectLighting)
 {
-	FDirectLighting Out = (FDirectLighting)0;
+	FStrataDeferredLighting Out = GetInitialisedStrataDeferredLighting();
 	const FStrataIntegrationSettings Settings = InitStrataIntegrationSettings(false /*bForceFullyRough*/, Strata.bRoughDiffuse, Strata.PeelLayersAboveDepth);
 
 	Strata_for (uint BSDFIndex = 0; BSDFIndex < StrataPixelHeader.BSDFCount; ++BSDFIndex)
@@ -224,17 +243,25 @@ FDirectLighting StrataLumenLighting(
 
 		const float3 InGlossyLighting = MixSpecularAndRoughReflections(Roughness, InSpecularLighting, InRoughSpecularIndirectLighting);
 
-		Out.Diffuse  += BSDFThroughput * InDiffuseLighting * StrataEnvLight.DiffuseWeight  * Occlusion.DiffuseOcclusion;
-		Out.Specular += BSDFThroughput * InGlossyLighting  * StrataEnvLight.SpecularWeight * Occlusion.SpecularOcclusion;
+		FDirectLighting OutBSDF = (FDirectLighting)0;
+		OutBSDF.Diffuse  += BSDFThroughput * InDiffuseLighting * StrataEnvLight.DiffuseWeight * Occlusion.DiffuseOcclusion;
+		OutBSDF.Specular += BSDFThroughput * InGlossyLighting  * StrataEnvLight.SpecularWeight * Occlusion.SpecularOcclusion;
 
 		#if STRATA_FASTPATH==0
 		if (BSDF_GETHASHAZINESS(BSDF))
 		{
 			const float HazeRoughness = StrataComputeHazeRoughness(Roughness);
 			const float3 InHazeGlossyLighting = MixSpecularAndRoughReflections(HazeRoughness, InSpecularLighting, InRoughSpecularIndirectLighting);
-			Out.Specular += BSDFThroughput * InHazeGlossyLighting * StrataEnvLight.SpecularHazeWeight * Occlusion.SpecularOcclusion;
+			OutBSDF.Specular += BSDFThroughput * InHazeGlossyLighting * StrataEnvLight.SpecularHazeWeight * Occlusion.SpecularOcclusion;
 		}
 		#endif 
+
+		// SSS Checkerboard
+		OutBSDF.Specular *= GetSSSCheckerboadSpecularScale(PixelPos, StrataEnvLight.bSubsurface);
+
+		FLightAccumulator Accumulator = (FLightAccumulator)0;
+		LightAccumulator_AddSplit(Accumulator, OutBSDF.Diffuse, OutBSDF.Specular, OutBSDF.Diffuse, 1.f /*CommonMultiplier*/, StrataEnvLight.bSubsurface);
+		AccumulateStrataDeferredLighting(Out, Accumulator, StrataEnvLight.bSubsurface, BSDF_GETISTOPLAYER(BSDF));
 	}
 
 	return Out;
@@ -275,6 +302,7 @@ void MainPS(
 
 	// Apply diffuse indirect.
 	#if DIM_APPLY_DIFFUSE_INDIRECT
+	OutAddColor = 0;
 	{
 		FDirectLighting IndirectLighting = (FDirectLighting)0;
 
@@ -382,7 +410,7 @@ void MainPS(
 				const FStrataTopLayerData TopLayerData = StrataUnpackTopLayerData(Strata.TopLayerTexture.Load(uint3(PixelPos, 0)));
 				const FStrataSubsurfaceData SSSData = StrataUnpackSSSData(Strata.SSSTexture.Load(uint3(PixelPos, 0)));
 
-				IndirectLighting = StrataLumenLighting(
+				const FStrataDeferredLighting IndirectLighting_Strata = StrataLumenLighting(
 					PixelPos,
 					Strata.MaterialTextureArray,
 					StrataAddressing,
@@ -392,6 +420,13 @@ void MainPS(
 					DiffuseIndirectLighting,
 					SpecularIndirectLighting,
 					RoughSpecularIndirectLighting);
+
+				OutAddColor = IndirectLighting_Strata.SceneColor;
+			  #if STRATA_OPAQUE_ROUGH_REFRACTION_ENABLED
+				const uint2 OutCoord = PixelPos;
+				OutOpaqueRoughRefractionSceneColor[OutCoord] = IndirectLighting_Strata.OpaqueRoughRefractionSceneColor;
+				OutSubSurfaceSceneColor[OutCoord] = IndirectLighting_Strata.SubSurfaceSceneColor;
+			  #endif
 			}
 			#else // STRATA_ENABLED
 			{
@@ -449,16 +484,9 @@ void MainPS(
 			#endif // STRATA_ENABLED
 		}
 
-		if (Material.bNeedsSeparateLightAccumulation &&
-			View.bSubsurfacePostprocessEnabled > 0 && View.bCheckerboardSubsurfaceProfileRendering > 0)
-		{
-			bool bChecker = CheckerFromSceneColorUV(SceneBufferUV);
-
-			// Adjust for checkerboard. only apply non-diffuse lighting (including emissive) 
-			// to the specular component, otherwise lighting is applied twice
-			IndirectLighting.Specular *= !bChecker;
-		}
-
+		// Accumulate lighting into the final buffers
+	#if !STRATA_ENABLED
+		IndirectLighting.Specular *= GetSSSCheckerboadSpecularScale(PixelPos, Material.bNeedsSeparateLightAccumulation);
 		FLightAccumulator LightAccumulator = (FLightAccumulator)0;
 		LightAccumulator_Add(
 			LightAccumulator,
@@ -467,6 +495,8 @@ void MainPS(
 			1.0f,
 			Material.bNeedsSeparateLightAccumulation);
 		OutAddColor = LightAccumulator_GetResult(LightAccumulator);
+	#endif // !STRATA_ENABLED
+
 	}
 	#endif
 
diff --git a/Engine/Source/Runtime/Renderer/Private/IndirectLightRendering.cpp b/Engine/Source/Runtime/Renderer/Private/IndirectLightRendering.cpp
--- a/Engine/Source/Runtime/Renderer/Private/IndirectLightRendering.cpp
+++ b/Engine/Source/Runtime/Renderer/Private/IndirectLightRendering.cpp
@@ -162,6 +162,9 @@ class FDiffuseIndirectCompositePS : public FGlobalShader
 
 		SHADER_PARAMETER_RDG_TEXTURE_UAV(RWTexture2D<float4>, PassDebugOutput)
 
+		SHADER_PARAMETER_RDG_TEXTURE_UAV(RWTexture2D<float3>, OutOpaqueRoughRefractionSceneColor)
+		SHADER_PARAMETER_RDG_TEXTURE_UAV(RWTexture2D<float3>, OutSubSurfaceSceneColor)
+
 		SHADER_PARAMETER(FVector2f, BufferUVToOutputPixelPosition)
 		SHADER_PARAMETER_RDG_TEXTURE(Texture2D, EyeAdaptation)
 		SHADER_PARAMETER_RDG_TEXTURE_ARRAY(Texture2D<uint>, CompressedMetadata, [2])
@@ -1162,6 +1165,13 @@ void FDeferredShadingSceneRenderer::RenderDiffuseIndirectAndAmbientOcclusion(
 			{
 				check(Strata::IsStrataEnabled());
 
+				// Rough refraction targets
+				if (Strata::IsStrataOpaqueMaterialRoughRefractionEnabled())
+				{
+					PassParameters->OutOpaqueRoughRefractionSceneColor = GraphBuilder.CreateUAV(Scene->StrataSceneData.SeparatedOpaqueRoughRefractionSceneColor);
+					PassParameters->OutSubSurfaceSceneColor = GraphBuilder.CreateUAV(Scene->StrataSceneData.SeparatedSubSurfaceSceneColor);
+				}
+
 				Strata::FStrataTilePassVS::FPermutationDomain VSPermutationVector;
 				VSPermutationVector.Set< Strata::FStrataTilePassVS::FEnableDebug >(false);
 				VSPermutationVector.Set< Strata::FStrataTilePassVS::FEnableTexCoordScreenVector >(false);
diff --git a/Engine/Source/Runtime/Renderer/Private/Strata/StrataRoughRefraction.cpp b/Engine/Source/Runtime/Renderer/Private/Strata/StrataRoughRefraction.cpp
--- a/Engine/Source/Runtime/Renderer/Private/Strata/StrataRoughRefraction.cpp
+++ b/Engine/Source/Runtime/Renderer/Private/Strata/StrataRoughRefraction.cpp
@@ -125,7 +125,7 @@ void AddStrataOpaqueRoughRefractionPasses(
 		PassParameters->StrataTile = Strata::SetTileParameters(GraphBuilder, View, StrataTileType, PrimitiveType);
 
 		GraphBuilder.AddPass(
-			RDG_EVENT_NAME("OpaqueRoughRefraction - Pass0"),
+			RDG_EVENT_NAME("Strata::OpaqueRoughRefraction(Blur,Horizontal)"),
 			PassParameters,
 			ERDGPassFlags::Raster,
 			[&View, TileVertexShader, PixelShader, PassParameters, StrataTileType, PrimitiveType](FRHICommandList& RHICmdList)
@@ -182,7 +182,7 @@ void AddStrataOpaqueRoughRefractionPasses(
 		PassParameters->StrataTile = Strata::SetTileParameters(GraphBuilder, View, StrataTileType, PrimitiveType);
 
 		GraphBuilder.AddPass(
-			RDG_EVENT_NAME("OpaqueRoughRefraction - Pass1)"),
+			RDG_EVENT_NAME("Strata::OpaqueRoughRefraction(Blur,Vertical)"),
 			PassParameters,
 			ERDGPassFlags::Raster,
 			[&View, TileVertexShader, PixelShader, PassParameters, StrataTileType, PrimitiveType](FRHICommandList& RHICmdList)
@@ -240,7 +240,7 @@ void AddStrataOpaqueRoughRefractionPasses(
 		PassParameters->StrataTile = Strata::SetTileParameters(GraphBuilder, View, StrataTileType, PrimitiveType);
 
 		GraphBuilder.AddPass(
-			RDG_EVENT_NAME("OpaqueRoughRefraction - Add SSSOnly tiles"),
+			RDG_EVENT_NAME("Strata::OpaqueRoughRefraction(SSS)"),
 			PassParameters,
 			ERDGPassFlags::Raster,
 			[&View, TileVertexShader, PixelShader, PassParameters, StrataTileType, PrimitiveType](FRHICommandList& RHICmdList)
