Strata tree vis now show each slab independentely.
Common functions have been made commin in StrataLightingCommon.ush.

diff --git a/Engine/Shaders/Private/Strata/Strata.ush b/Engine/Shaders/Private/Strata/Strata.ush
--- a/Engine/Shaders/Private/Strata/Strata.ush
+++ b/Engine/Shaders/Private/Strata/Strata.ush
@@ -2704,6 +2704,8 @@ struct FStrataIntegrationSettings
 {
 	bool bForceFullyRough;
 	bool bRoughDiffuseEnabled;
+
+	// Embbeded debug settings
 	int  SliceStoringDebugStrataTree;
 	int  PeelLayersAboveDepth;
 };
@@ -2713,11 +2715,17 @@ FStrataIntegrationSettings InitStrataIntegrationSettings(bool bForceFullyRough,
 	FStrataIntegrationSettings Out;
 	Out.bForceFullyRough = bForceFullyRough;
 	Out.bRoughDiffuseEnabled = bRoughDiffuseEnabled;
+
 	Out.PeelLayersAboveDepth = PeelLayersAboveDepth;
 	Out.SliceStoringDebugStrataTree = -1;
 	return Out;
 }
 
+FStrataIntegrationSettings InitStrataIntegrationSettings()
+{
+	return InitStrataIntegrationSettings(false, true, -1);
+}
+
 ///////////////////////////////////////////////////////////////////////////////
 // Strata deferred base-pass output
 
diff --git a/Engine/Shaders/Private/Strata/StrataEnvironmentLighting.ush b/Engine/Shaders/Private/Strata/StrataEnvironmentLighting.ush
--- a/Engine/Shaders/Private/Strata/StrataEnvironmentLighting.ush
+++ b/Engine/Shaders/Private/Strata/StrataEnvironmentLighting.ush
@@ -10,6 +10,9 @@
 
 #if STRATA_ENABLED
 
+#define USE_STRATA_ENV_LIGHTING_COMMON 1
+#include "StrataLightingCommon.ush"
+
 // Env. lighting evaluation for strata material.
 // Unpack BSDF on-the-fly
 FStrataDeferredLighting StrataEnvironmentLighting(
@@ -20,7 +23,7 @@ FStrataDeferredLighting StrataEnvironmentLighting(
 	float3 TranslatedWorldPosition,
 	float3 CameraToPixel,
 	float3 V,
-	float AmbientOcclusion,
+	float ScreenAmbientOcclusion,
 	float TopLayerSpecularContributionFactor,
 	FStrataIntegrationSettings Settings,
 	FStrataMaterialContainer MaterialBuffer,
@@ -39,7 +42,7 @@ FStrataDeferredLighting StrataEnvironmentLighting(
 	FUpsampleDFAOOutput UpsampleDFAOOutput = UpsampleDFAO(BufferUV, SceneDepth);
 #endif
 
-	float CombinedScreenAndMaterialAO = StrataGetAO(StrataPixelHeader) * AmbientOcclusion;
+	float CombinedScreenAndMaterialAO = StrataGetAO(StrataPixelHeader) * ScreenAmbientOcclusion;
 
 	
 #if FEATURE_LEVEL >= FEATURE_LEVEL_SM5
@@ -73,107 +76,27 @@ FStrataDeferredLighting StrataEnvironmentLighting(
 		BentNormal = ApplyDFAO(UpsampleDFAOOutput, StrataEnvLight.DiffuseNormal);
 	#endif
 
-		// Diffuse component
 		float3 DiffuseLighting = 0;
-	#if ENABLE_DYNAMIC_SKY_LIGHT
-		const bool bProcessFrontFaceDiffuse = any(StrataEnvLight.DiffuseWeight > 0.0f);
-		const bool bProcessBackFaceDiffuse = any(StrataEnvLight.DiffuseBackFaceWeight > 0.0f);
-		if (bProcessFrontFaceDiffuse || bProcessBackFaceDiffuse)
-		{
-			// Compute the common sky visibility factors
-			FSkyLightVisibilityData SkyVisData = GetSkyLightVisibilityData(StrataBSDFContext.N, StrataBSDFContext.N, StrataGetAO(StrataPixelHeader), AmbientOcclusion, BentNormal);
-
-			if (bProcessFrontFaceDiffuse)
-			{
-				// Finally sample the sky diffuse contribution (spherical harmonic, Lambert BRDF)
-				float3 DiffuseLookup = GetSkySHDiffuse(BentNormal) * View.SkyLightColor.rgb;
-				// And accumulate
-				// Note: Use diffuse directional albedo (i.e., DiffuseWeight) as first order approximation for env. integration (STRATA_TODO instead compute SH coefficients for Chan)
-				DiffuseLighting = CloudVolumetricAOShadow * BSDFThroughput * (SkyVisData.SkyDiffuseLookUpMul * DiffuseLookup + SkyVisData.SkyDiffuseLookUpAdd) * StrataEnvLight.DiffuseWeight;
-			}
-			if (bProcessBackFaceDiffuse)
-			{
-				// We do not evaluate back face sky light visibility data because all the data we have is for the front face only.  This could be evaluated at some cost.
-				// However, we do apply SkyVisData.SkyDiffuseLookUpMul for scaling consistency.
-
-				// Finally sample the sky diffuse contribution (spherical harmonic, Lambert BRDF) along the opposite normal direction
-				float3 DiffuseLookup = GetSkySHDiffuse(-BentNormal) * View.SkyLightColor.rgb;
-				// And accumulate
-				// Note: Use diffuse directional albedo (i.e., DiffuseWeight) as first order approximation for env. integration (STRATA_TODO instead compute SH coefficients for Chan)
-				DiffuseLighting += CloudVolumetricAOShadow * BSDFThroughput * (SkyVisData.SkyDiffuseLookUpMul * DiffuseLookup) * StrataEnvLight.DiffuseBackFaceWeight;
-			}
-		}
-	#endif // ENABLE_DYNAMIC_SKY_LIGHT
-
-		// Specular component
-		const bool bIsTopLayer = BSDF_GETISTOPLAYER(BSDF);
 		float3 SpecularLighting = 0;
-		if (any(StrataEnvLight.SpecularWeight > 0.0f))
-		{
-			float IndirectIrradiance = 0.0;						// STRATA_TODO GBuffer.IndirectIrradiance, that will likely have to take 8 mnore bits on the header
-		#if ENABLE_SKY_LIGHT && ALLOW_STATIC_LIGHTING
-			BRANCH
-			// Add in diffuse contribution from dynamic skylights so reflection captures will have something to mix with
-			if (ReflectionStruct.SkyLightParameters.y > 0 && ReflectionStruct.SkyLightParameters.z > 0)
-			{
-				IndirectIrradiance += GetDynamicSkyIndirectIrradiance(BentNormal, StrataBSDFContext.N);
-			}
-		#endif
-
-			// Compute some extra occlusion information from DFAO and sky light data
-			float IndirectSpecularOcclusion = 1.0f;
-			float3 ExtraIndirectSpecular = 0.0f;
-
-		#if SUPPORT_DFAO_INDIRECT_OCCLUSION
-			float IndirectDiffuseOcclusion;
-			const bool bTwoSideFoliage = false;
-			GetDistanceFieldAOSpecularOcclusion(BentNormal, StrataEnvLight.SpecularDirection, StrataEnvLight.SpecularSafeRoughness, bTwoSideFoliage, IndirectSpecularOcclusion, IndirectDiffuseOcclusion, ExtraIndirectSpecular);
-			// Apply DFAO to IndirectIrradiance before mixing with indirect specular
-			IndirectIrradiance *= IndirectDiffuseOcclusion;
-		#endif
-
-			float RoughnessSquared = StrataEnvLight.SpecularSafeRoughness * StrataEnvLight.SpecularSafeRoughness;
-			float SpecularOcclusion = IndirectSpecularOcclusion * GetSpecularOcclusion(StrataBSDFContext.SatNoV, RoughnessSquared, CombinedScreenAndMaterialAO);
-
-			const bool bCompositeSkylight = true;
-			SpecularLighting += BSDFThroughput * StrataEnvLight.SpecularWeight *
-				CompositeReflectionCapturesAndSkylightTWS(
-					(bIsTopLayer ? TopLayerSpecularContributionFactor : 1.0f) * SpecularOcclusion,
-					TranslatedWorldPosition,
-					StrataEnvLight.SpecularDirection,
-					StrataEnvLight.SpecularSafeRoughness,
-					IndirectIrradiance,
-					IndirectSpecularOcclusion,
-					ExtraIndirectSpecular,
-					NumCulledReflectionCaptures,
-					CaptureDataStartIndex,
-					0,
-					bCompositeSkylight);
-
-		#if STRATA_FASTPATH==0
-			if (BSDF_GETHASHAZINESS(BSDF))
-			{
-				SpecularLighting += BSDFThroughput * StrataEnvLight.SpecularHazeWeight *
-					CompositeReflectionCapturesAndSkylightTWS(
-						(bIsTopLayer ? TopLayerSpecularContributionFactor : 1.0f) * SpecularOcclusion,
-						TranslatedWorldPosition,
-						StrataEnvLight.SpecularDirection,
-						StrataEnvLight.SpecularHazeSafeRoughness,
-						IndirectIrradiance,
-						IndirectSpecularOcclusion,
-						ExtraIndirectSpecular,
-						NumCulledReflectionCaptures,
-						CaptureDataStartIndex,
-						0,
-						bCompositeSkylight);
-
-				// SSR is traced for the sharpest lob. The smoothest one does not rely on SSR so we need to lower energy coming from SSR according to the lobe blend weight.
-				// And we also try to make the transition smooth using Haziness
-				SSRReductionFactor -= bIsTopLayer ? dot(BSDFThroughput, float3(1.0f / 3.0f, 1.0f / 3.0f, 1.0f / 3.0f)) * StrataEnvLight.SSRReduction : 0.0f;
-			}
-		#endif 
-		}
+		StrataEnvLightingCommon(
+			StrataEnvLight,
+			StrataPixelHeader,
+			StrataBSDFContext,
+			BSDF,
+			BentNormal,
+			BSDFThroughput,
+			CaptureDataStartIndex,
+			NumCulledReflectionCaptures,
+			ScreenAmbientOcclusion,
+			CloudVolumetricAOShadow,
+			TopLayerSpecularContributionFactor,
+			TranslatedWorldPosition,
+			CombinedScreenAndMaterialAO,
+			SSRReductionFactor,
+			DiffuseLighting,
+			SpecularLighting);
 
+		const bool bIsTopLayer = BSDF_GETISTOPLAYER(BSDF);
 		if (bIsTopLayer)
 		{
 			// The specular path weight applied on SSR. It must account for throughput even for top surface because it also contains coverage.
diff --git a/Engine/Shaders/Private/Strata/StrataForwardLighting.ush b/Engine/Shaders/Private/Strata/StrataForwardLighting.ush
--- a/Engine/Shaders/Private/Strata/StrataForwardLighting.ush
+++ b/Engine/Shaders/Private/Strata/StrataForwardLighting.ush
@@ -2,6 +2,8 @@
 
 #pragma once
 
+#define USE_STRATA_FORWARD_LIGHTING_COMMON 1
+#include "StrataLightingCommon.ush"
 
 #define FORWARD_PER_PIXEL_SHADING (FORWARD_SHADING || TRANSLUCENCY_LIGHTING_SURFACE_FORWARDSHADING || MATERIAL_SHADINGMODEL_SINGLELAYERWATER)
 
@@ -30,126 +32,6 @@ void GetPrecomputedIndirectLightingAndSkyLight(
 	out float3 OutSubsurfaceLighting,
 	out float OutIndirectIrradiance);
 
-// Common function for forward lighting per pixel using light data
-float3 StrataForwardLightingCommon(
-	float Dither,
-	FStrataIntegrationSettings Settings,
-	FDeferredLightData LightData,
-	float3 ToLight,
-	float LightMask,
-	float4 LightAttenuation,
-	FRectTexture RectTexture,
-	uint LightChannelMask,
-	uint PrimitiveLightingChannelMask,
-	float3 AbsoluteWorldPosition,
-	float3 TranslatedWorldPosition,
-	float SceneDepth,
-	float3 BSDFColoredVisibility,
-	FStrataPixelHeader StrataPixelHeader,
-	FStrataBSDFContext StrataBSDFContext,
-	inout bool bEvaluateHasBeenExecuted,
-	inout FStrataEvaluateResult BSDFEvaluate)
-{
-	bEvaluateHasBeenExecuted = false;
-	float3 Color = 0.0;
-
-	if (LightMask > 0.0)
-	{
-		// Evaluate the ShadowTerm that can then be used when integrating the lighting
-		FShadowTerms ShadowTerms = { StrataGetAO(StrataPixelHeader), 1.0, 1.0, InitHairTransmittanceData() };
-
-		float4 PrecomputedShadowFactors = HasZeroPrecShadowMask(StrataPixelHeader) ? 0.0f : 1.0f;
-	#if TRANSLUCENCY_LIGHTING_SURFACE_FORWARDSHADING
-		PrecomputedShadowFactors.x = ComputeDirectionalLightStaticShadowing(TranslatedWorldPosition).x;
-	#else
-		PrecomputedShadowFactors.x = 1;
-	#endif
-
-		const uint FakeShadingModelID = 0;
-		const float FakeContactShadowOpacity = 1.0f;
-		GetShadowTerms(SceneDepth, PrecomputedShadowFactors, FakeShadingModelID, FakeContactShadowOpacity,
-			LightData, TranslatedWorldPosition, ToLight, LightAttenuation, Dither, ShadowTerms);
-
-		float Roughness = StrataGetBSDFRoughness(StrataBSDFContext.BSDF);
-		FAreaLightIntegrateContext AreaLightContext = InitAreaLightIntegrateContext();
-		BRANCH
-		if (ShadowTerms.SurfaceShadow + ShadowTerms.TransmissionShadow > 0)
-		{
-			BSDFEvaluate = (FStrataEvaluateResult)0;
-
-		#if NON_DIRECTIONAL_DIRECT_LIGHTING
-			float Lighting;
-			if (LightData.bRectLight)
-			{
-				FRect Rect = GetRect(ToLight, LightData);
-				if (!IsRectVisible(Rect))
-				{
-					LightMask = 0.0f; // Rect light can be non visible due to barn door occlusion
-				}
-				AreaLightContext = CreateRectIntegrateContext(Roughness, StrataBSDFContext.N, StrataBSDFContext.V, Rect, RectTexture);
-				Lighting = IntegrateLight(Rect, RectTexture);
-
-				// We must have the evaluate inside the if due to the rectlight texture: it must be now be ambiguous which texture is going ot be used.
-				// After te compilation, a local resource must map to a unique global resource (the default or the actual rect light texture).
-				BSDFEvaluate = StrataIntegrateBSDF(StrataBSDFContext, ShadowTerms, AreaLightContext, Settings);
-				bEvaluateHasBeenExecuted = true;
-			}
-			else
-			{
-				FCapsuleLight Capsule = GetCapsule(ToLight, LightData);
-				AreaLightContext = CreateCapsuleIntegrateContext(Roughness, StrataBSDFContext.N, StrataBSDFContext.V, Capsule, LightData.bInverseSquared);
-				Lighting = IntegrateLight(Capsule, LightData.bInverseSquared);
-
-				BSDFEvaluate = StrataIntegrateBSDF(StrataBSDFContext, ShadowTerms, AreaLightContext, Settings);
-				bEvaluateHasBeenExecuted = true;
-			}
-
-			FLATTEN
-			if (LightChannelMask & PrimitiveLightingChannelMask)
-			{
-				float3 DiffuseLuminance = Diffuse_Lambert(BSDFEvaluate.DiffuseColor) * Lighting;
-				const float3 LightCommonMultiplier = LightData.Color * LightMask;
-				Color += (DiffuseLuminance * LightCommonMultiplier + BSDFEvaluate.EmissivePathValue) * BSDFColoredVisibility;
-			}
-		#else
-			if (LightData.bRectLight)
-			{
-				FRect Rect = GetRect(ToLight, LightData);
-				if (!IsRectVisible(Rect))
-				{
-					LightMask = 0.0f; // Rect light can be non visible due to barn door occlusion
-				}
-				AreaLightContext = CreateRectIntegrateContext(Roughness, StrataBSDFContext.N, StrataBSDFContext.V, Rect, RectTexture);
-
-				// We must have the evaluate inside the if due to the rectlight texture: it must be now be ambiguous which texture is going ot be used.
-				// After te compilation, a local resource must map to a unique global resource (the default or the actual rect light texture).
-				BSDFEvaluate = StrataIntegrateBSDF(StrataBSDFContext, ShadowTerms, AreaLightContext, Settings);
-				bEvaluateHasBeenExecuted = true;
-			}
-			else
-			{
-				FCapsuleLight Capsule = GetCapsule(ToLight, LightData);
-				AreaLightContext = CreateCapsuleIntegrateContext(Roughness, StrataBSDFContext.N, StrataBSDFContext.V, Capsule, LightData.bInverseSquared);
-
-				BSDFEvaluate = StrataIntegrateBSDF(StrataBSDFContext, ShadowTerms, AreaLightContext, Settings);
-				bEvaluateHasBeenExecuted = true;
-			}
-
-			FLATTEN
-			if (LightChannelMask & PrimitiveLightingChannelMask)
-			{
-				float3 DiffuseLuminance = BSDFEvaluate.IntegratedDiffuseValue;
-				float3 SpecularLuminance = BSDFEvaluate.IntegratedSpecularValue * LightData.SpecularScale;
-				const float3 LightCommonMultiplier = LightData.Color * LightMask;
-				Color += ((DiffuseLuminance + SpecularLuminance) * LightCommonMultiplier + BSDFEvaluate.EmissivePathValue) * BSDFColoredVisibility;
-			}
-		#endif
-		}
-	}
-
-	return Color;
-}
-
 // This code should map to GetForwardDirectLightingSplit
 float3 StrataForwardLighting(
 	uint EyeIndex,
@@ -257,7 +139,7 @@ float3 StrataForwardLighting(
 	if (StrataTree.BSDFCount > 0)
 	{
 		// Update tree (coverage/transmittance/luminace weights)
-		StrataUpdateTree(NullStrataAddressing, StrataPixelHeader, StrataTree, StrataData, V, Settings, OutCoverage, OutTransmittancePreCoverage); ////////////////////////////////////////////////////////////////////
+		StrataUpdateTree(NullStrataAddressing, StrataPixelHeader, StrataTree, StrataData, V, Settings, OutCoverage, OutTransmittancePreCoverage);
 
 		//
 		// Forward lighting
@@ -487,7 +369,6 @@ float3 StrataForwardLighting(
 							RectTexture,
 							LightChannelMask,
 							PrimitiveLightingChannelMask,
-							AbsoluteWorldPosition,
 							TranslatedWorldPosition,
 							SceneDepth,
 							LuminanceWeight(StrataBSDFContext, CurrentBSDF),
@@ -571,7 +452,6 @@ float3 StrataForwardLighting(
 						RectTexture,
 						LightChannelMask,
 						PrimitiveLightingChannelMask,
-						AbsoluteWorldPosition,
 						TranslatedWorldPosition,
 						SceneDepth,
 						LuminanceWeight(StrataBSDFContext, CurrentBSDF),
diff --git a/Engine/Shaders/Private/Strata/StrataLightingCommon.ush b/Engine/Shaders/Private/Strata/StrataLightingCommon.ush
new file mode 100644
--- /dev/null
+++ b/Engine/Shaders/Private/Strata/StrataLightingCommon.ush
@@ -0,0 +1,254 @@
+// Copyright Epic Games, Inc. All Rights Reserved.
+
+#pragma once
+
+#ifdef USE_STRATA_FORWARD_LIGHTING_COMMON
+
+// Common function for forward lighting per pixel using light data
+float3 StrataForwardLightingCommon(
+	float Dither,
+	FStrataIntegrationSettings Settings,
+	FDeferredLightData LightData,
+	float3 ToLight,
+	float LightMask,
+	float4 LightAttenuation,
+	FRectTexture RectTexture,
+	uint LightChannelMask,
+	uint PrimitiveLightingChannelMask,
+	float3 TranslatedWorldPosition,
+	float SceneDepth,
+	float3 BSDFColoredVisibility,
+	FStrataPixelHeader StrataPixelHeader,
+	FStrataBSDFContext StrataBSDFContext,
+	inout bool bEvaluateHasBeenExecuted,
+	inout FStrataEvaluateResult BSDFEvaluate)
+{
+	bEvaluateHasBeenExecuted = false;
+	float3 Color = 0.0;
+
+	if (LightMask > 0.0)
+	{
+		// Evaluate the ShadowTerm that can then be used when integrating the lighting
+		FShadowTerms ShadowTerms = { StrataGetAO(StrataPixelHeader), 1.0, 1.0, InitHairTransmittanceData() };
+
+		float4 PrecomputedShadowFactors = HasZeroPrecShadowMask(StrataPixelHeader) ? 0.0f : 1.0f;
+	#if TRANSLUCENCY_LIGHTING_SURFACE_FORWARDSHADING
+		PrecomputedShadowFactors.x = ComputeDirectionalLightStaticShadowing(TranslatedWorldPosition).x;
+	#else
+		PrecomputedShadowFactors.x = 1;
+	#endif
+
+		const uint FakeShadingModelID = 0;
+		const float FakeContactShadowOpacity = 1.0f;
+		GetShadowTerms(SceneDepth, PrecomputedShadowFactors, FakeShadingModelID, FakeContactShadowOpacity,
+			LightData, TranslatedWorldPosition, ToLight, LightAttenuation, Dither, ShadowTerms);
+
+		float Roughness = StrataGetBSDFRoughness(StrataBSDFContext.BSDF);
+		FAreaLightIntegrateContext AreaLightContext = InitAreaLightIntegrateContext();
+		BRANCH
+		if (ShadowTerms.SurfaceShadow + ShadowTerms.TransmissionShadow > 0)
+		{
+			BSDFEvaluate = (FStrataEvaluateResult)0;
+
+		#if NON_DIRECTIONAL_DIRECT_LIGHTING
+			float Lighting;
+			if (LightData.bRectLight)
+			{
+				FRect Rect = GetRect(ToLight, LightData);
+				if (!IsRectVisible(Rect))
+				{
+					LightMask = 0.0f; // Rect light can be non visible due to barn door occlusion
+				}
+				AreaLightContext = CreateRectIntegrateContext(Roughness, StrataBSDFContext.N, StrataBSDFContext.V, Rect, RectTexture);
+				Lighting = IntegrateLight(Rect, RectTexture);
+
+				// We must have the evaluate inside the if due to the rectlight texture: it must be now be ambiguous which texture is going ot be used.
+				// After te compilation, a local resource must map to a unique global resource (the default or the actual rect light texture).
+				BSDFEvaluate = StrataIntegrateBSDF(StrataBSDFContext, ShadowTerms, AreaLightContext, Settings);
+				bEvaluateHasBeenExecuted = true;
+			}
+			else
+			{
+				FCapsuleLight Capsule = GetCapsule(ToLight, LightData);
+				AreaLightContext = CreateCapsuleIntegrateContext(Roughness, StrataBSDFContext.N, StrataBSDFContext.V, Capsule, LightData.bInverseSquared);
+				Lighting = IntegrateLight(Capsule, LightData.bInverseSquared);
+
+				BSDFEvaluate = StrataIntegrateBSDF(StrataBSDFContext, ShadowTerms, AreaLightContext, Settings);
+				bEvaluateHasBeenExecuted = true;
+			}
+
+			FLATTEN
+			if (LightChannelMask & PrimitiveLightingChannelMask)
+			{
+				float3 DiffuseLuminance = Diffuse_Lambert(BSDFEvaluate.DiffuseColor) * Lighting;
+				const float3 LightCommonMultiplier = LightData.Color * LightMask;
+				Color += (DiffuseLuminance * LightCommonMultiplier + BSDFEvaluate.EmissivePathValue) * BSDFColoredVisibility;
+			}
+		#else
+			if (LightData.bRectLight)
+			{
+				FRect Rect = GetRect(ToLight, LightData);
+				if (!IsRectVisible(Rect))
+				{
+					LightMask = 0.0f; // Rect light can be non visible due to barn door occlusion
+				}
+				AreaLightContext = CreateRectIntegrateContext(Roughness, StrataBSDFContext.N, StrataBSDFContext.V, Rect, RectTexture);
+
+				// We must have the evaluate inside the if due to the rectlight texture: it must be now be ambiguous which texture is going ot be used.
+				// After te compilation, a local resource must map to a unique global resource (the default or the actual rect light texture).
+				BSDFEvaluate = StrataIntegrateBSDF(StrataBSDFContext, ShadowTerms, AreaLightContext, Settings);
+				bEvaluateHasBeenExecuted = true;
+			}
+			else
+			{
+				FCapsuleLight Capsule = GetCapsule(ToLight, LightData);
+				AreaLightContext = CreateCapsuleIntegrateContext(Roughness, StrataBSDFContext.N, StrataBSDFContext.V, Capsule, LightData.bInverseSquared);
+
+				BSDFEvaluate = StrataIntegrateBSDF(StrataBSDFContext, ShadowTerms, AreaLightContext, Settings);
+				bEvaluateHasBeenExecuted = true;
+			}
+
+			FLATTEN
+			if (LightChannelMask & PrimitiveLightingChannelMask)
+			{
+				float3 DiffuseLuminance = BSDFEvaluate.IntegratedDiffuseValue;
+				float3 SpecularLuminance = BSDFEvaluate.IntegratedSpecularValue * LightData.SpecularScale;
+				const float3 LightCommonMultiplier = LightData.Color * LightMask;
+				Color += ((DiffuseLuminance + SpecularLuminance) * LightCommonMultiplier + BSDFEvaluate.EmissivePathValue) * BSDFColoredVisibility;
+			}
+		#endif
+		}
+	}
+
+	return Color;
+}
+
+#endif // USE_STRATA_FORWARD_LIGHTING_COMMON
+
+
+
+#ifdef USE_STRATA_ENV_LIGHTING_COMMON
+
+void StrataEnvLightingCommon(
+	in FStrataEnvLightResult StrataEnvLight,
+	in FStrataPixelHeader StrataPixelHeader,
+	in FStrataBSDFContext StrataBSDFContext,
+	in FStrataBSDF BSDF,
+	in float3 BentNormal,
+	in float3 BSDFThroughput,
+	in uint CaptureDataStartIndex,
+	in uint NumCulledReflectionCaptures,
+	in float ScreenAmbientOcclusion,
+	in float CloudVolumetricAOShadow,
+	in float TopLayerSpecularContributionFactor,
+	in float3 TranslatedWorldPosition,
+	in float CombinedScreenAndMaterialAO,
+	inout float SSRReductionFactor,
+	inout float3 DiffuseLighting,
+	inout float3 SpecularLighting)
+{
+	// Diffuse component
+	DiffuseLighting = 0;
+#if ENABLE_DYNAMIC_SKY_LIGHT
+	const bool bProcessFrontFaceDiffuse = any(StrataEnvLight.DiffuseWeight > 0.0f);
+	const bool bProcessBackFaceDiffuse = any(StrataEnvLight.DiffuseBackFaceWeight > 0.0f);
+	if (bProcessFrontFaceDiffuse || bProcessBackFaceDiffuse)
+	{
+		// Compute the common sky visibility factors
+		FSkyLightVisibilityData SkyVisData = GetSkyLightVisibilityData(StrataBSDFContext.N, StrataBSDFContext.N, StrataGetAO(StrataPixelHeader), ScreenAmbientOcclusion, BentNormal);
+
+		if (bProcessFrontFaceDiffuse)
+		{
+			// Finally sample the sky diffuse contribution (spherical harmonic, Lambert BRDF)
+			float3 DiffuseLookup = GetSkySHDiffuse(BentNormal) * View.SkyLightColor.rgb;
+			// And accumulate
+			// Note: Use diffuse directional albedo (i.e., DiffuseWeight) as first order approximation for env. integration (STRATA_TODO instead compute SH coefficients for Chan)
+			DiffuseLighting = CloudVolumetricAOShadow * BSDFThroughput * (SkyVisData.SkyDiffuseLookUpMul * DiffuseLookup + SkyVisData.SkyDiffuseLookUpAdd) * StrataEnvLight.DiffuseWeight;
+		}
+		if (bProcessBackFaceDiffuse)
+		{
+			// We do not evaluate back face sky light visibility data because all the data we have is for the front face only.  This could be evaluated at some cost.
+			// However, we do apply SkyVisData.SkyDiffuseLookUpMul for scaling consistency.
+
+			// Finally sample the sky diffuse contribution (spherical harmonic, Lambert BRDF) along the opposite normal direction
+			float3 DiffuseLookup = GetSkySHDiffuse(-BentNormal) * View.SkyLightColor.rgb;
+			// And accumulate
+			// Note: Use diffuse directional albedo (i.e., DiffuseWeight) as first order approximation for env. integration (STRATA_TODO instead compute SH coefficients for Chan)
+			DiffuseLighting += CloudVolumetricAOShadow * BSDFThroughput * (SkyVisData.SkyDiffuseLookUpMul * DiffuseLookup) * StrataEnvLight.DiffuseBackFaceWeight;
+		}
+	}
+#endif // ENABLE_DYNAMIC_SKY_LIGHT
+
+	// Specular component
+	const bool bIsTopLayer = BSDF_GETISTOPLAYER(BSDF);
+	SpecularLighting = 0;
+	if (any(StrataEnvLight.SpecularWeight > 0.0f))
+	{
+		float IndirectIrradiance = 0.0;						// STRATA_TODO GBuffer.IndirectIrradiance, that will likely have to take 8 mnore bits on the header
+	#if ENABLE_SKY_LIGHT && ALLOW_STATIC_LIGHTING
+		BRANCH
+		// Add in diffuse contribution from dynamic skylights so reflection captures will have something to mix with
+		if (ReflectionStruct.SkyLightParameters.y > 0 && ReflectionStruct.SkyLightParameters.z > 0)
+		{
+			IndirectIrradiance += GetDynamicSkyIndirectIrradiance(BentNormal, StrataBSDFContext.N);
+		}
+	#endif
+
+		// Compute some extra occlusion information from DFAO and sky light data
+		float IndirectSpecularOcclusion = 1.0f;
+		float3 ExtraIndirectSpecular = 0.0f;
+
+	#if SUPPORT_DFAO_INDIRECT_OCCLUSION
+		float IndirectDiffuseOcclusion;
+		const bool bTwoSideFoliage = false;
+		GetDistanceFieldAOSpecularOcclusion(BentNormal, StrataEnvLight.SpecularDirection, StrataEnvLight.SpecularSafeRoughness, bTwoSideFoliage, IndirectSpecularOcclusion, IndirectDiffuseOcclusion, ExtraIndirectSpecular);
+		// Apply DFAO to IndirectIrradiance before mixing with indirect specular
+		IndirectIrradiance *= IndirectDiffuseOcclusion;
+	#endif
+
+		float RoughnessSquared = StrataEnvLight.SpecularSafeRoughness * StrataEnvLight.SpecularSafeRoughness;
+		float SpecularOcclusion = IndirectSpecularOcclusion * GetSpecularOcclusion(StrataBSDFContext.SatNoV, RoughnessSquared, CombinedScreenAndMaterialAO);
+
+		const bool bCompositeSkylight = true;
+		SpecularLighting += BSDFThroughput * StrataEnvLight.SpecularWeight *
+			CompositeReflectionCapturesAndSkylightTWS(
+				(bIsTopLayer ? TopLayerSpecularContributionFactor : 1.0f) * SpecularOcclusion,
+				TranslatedWorldPosition,
+				StrataEnvLight.SpecularDirection,
+				StrataEnvLight.SpecularSafeRoughness,
+				IndirectIrradiance,
+				IndirectSpecularOcclusion,
+				ExtraIndirectSpecular,
+				NumCulledReflectionCaptures,
+				CaptureDataStartIndex,
+				0,
+				bCompositeSkylight);
+
+	#if STRATA_FASTPATH==0
+		if (BSDF_GETHASHAZINESS(BSDF))
+		{
+			SpecularLighting += BSDFThroughput * StrataEnvLight.SpecularHazeWeight *
+				CompositeReflectionCapturesAndSkylightTWS(
+					(bIsTopLayer ? TopLayerSpecularContributionFactor : 1.0f) * SpecularOcclusion,
+					TranslatedWorldPosition,
+					StrataEnvLight.SpecularDirection,
+					StrataEnvLight.SpecularHazeSafeRoughness,
+					IndirectIrradiance,
+					IndirectSpecularOcclusion,
+					ExtraIndirectSpecular,
+					NumCulledReflectionCaptures,
+					CaptureDataStartIndex,
+					0,
+					bCompositeSkylight);
+
+			// SSR is traced for the sharpest lob. The smoothest one does not rely on SSR so we need to lower energy coming from SSR according to the lobe blend weight.
+			// And we also try to make the transition smooth using Haziness
+			SSRReductionFactor -= bIsTopLayer ? dot(BSDFThroughput, float3(1.0f / 3.0f, 1.0f / 3.0f, 1.0f / 3.0f)) * StrataEnvLight.SSRReduction : 0.0f;
+		}
+	#endif 
+	}
+}
+
+#endif // STRATA_ENV_LIGHTING_COMMON
+
+
diff --git a/Engine/Shaders/Private/Strata/StrataVisualize.usf b/Engine/Shaders/Private/Strata/StrataVisualize.usf
--- a/Engine/Shaders/Private/Strata/StrataVisualize.usf
+++ b/Engine/Shaders/Private/Strata/StrataVisualize.usf
@@ -3,7 +3,7 @@
 #include "/Engine/Private/Common.ush"
 #include "/Engine/Private/DeferredShadingCommon.ush"
 
-#ifdef SHADER_DEBUGSTRATATREE
+#ifdef SHADER_DEBUGSTRATATREE_PS
 #define STRATA_INLINE_SHADING 1
 #else
 #define STRATA_INLINE_SHADING 0
@@ -455,7 +455,7 @@ void VisualizeMaterialPS(
 ///////////////////////////////////////////////////////////////////////////////////////////////////
 // Material Visualize
 
-#if SHADER_DEBUGSTRATATREE
+#if SHADER_DEBUGSTRATATREE_TEST_CS
 
 [numthreads(1, 1, 1)]
 void MaterialDebugStrataTreeCS(uint3 DispatchThreadId : SV_DispatchThreadID)
@@ -486,18 +486,187 @@ void MaterialDebugStrataTreeCS(uint3 DispatchThreadId : SV_DispatchThreadID)
 	}
 }
 
-float4 GetMaterialOnSphere(
-	in float2 PixelPos,
+#endif // SHADER_DEBUGSTRATATREE_TEST_CS
+
+#if SHADER_DEBUGSTRATATREE_PS
+
+#include "../DeferredShadingCommon.ush"
+#define SUPPORT_CONTACT_SHADOWS 0
+#include "../DeferredLightingCommon.ush"
+
+#include "StrataEvaluation.ush"
+
+// These defines are needed for environment lighting in StrataLightingCommon
+#define ENABLE_DYNAMIC_SKY_LIGHT 1
+#define ENABLE_SKY_LIGHT 1
+#define SUPPORT_DFAO_INDIRECT_OCCLUSION 0
+#define STRATA_FASTPATH 0
+#define REFLECTION_COMPOSITE_USE_BLENDED_REFLECTION_CAPTURES 0
+#include "../ReflectionEnvironmentShared.ush"
+#include "../SkyLightingDiffuseShared.ush"
+#include "../ReflectionEnvironmentComposite.ush"
+
+#define USE_STRATA_ENV_LIGHTING_COMMON 1
+#define USE_STRATA_FORWARD_LIGHTING_COMMON 1
+#include "StrataLightingCommon.ush"
+
+float4 GetBackgroundCheckboardColor(float2 UV)
+{
+	const float Count = 16.0f;
+	int2 UVi = int2(UV * Count);
+	bool bIsBlack = (UVi.x % 2) > 0 ? true : false;
+	bIsBlack = (UVi.y % 2) > 0 ? !bIsBlack : bIsBlack;
+	const float3 Color = (bIsBlack ? 0.04 : 0.18);
+	return float4(Color, 0.0);
+}
+
+void GetMaterialOnSphere(
+	in float4 SVPos,
 	in float2 Pos,
 	in float  Size,
 	in FStrataBSDF BSDF,
 	inout float4 OutColor)
 {
-	if (all(PixelPos > Pos) && all(PixelPos < (Pos + Size)))
+	const float2 PixelPos = SVPos.xy;
+
+	if (BSDF_GETTYPE(BSDF) == STRATA_BSDF_TYPE_SLAB && all(PixelPos > Pos) && all(PixelPos < (Pos + Size)))
 	{
-		OutColor = float4(StrataGetBSDFDiffuseColor(BSDF), 0);
+		const float2 UV = (PixelPos - Pos) / Size;
+		OutColor = GetBackgroundCheckboardColor(UV);
+
+		float3 AbsoluteWorldPosition = LWCHackToFloat(PrimaryView.WorldCameraOrigin);
+		float3 TranslatedWorldPosition = PrimaryView.TranslatedWorldCameraOrigin;
+		float SceneDepth = 0.0;
+
+		float3 V = -View.ViewForward;// float3(0, 0, 1);
+		const float3 RayO = V + (UV.x-0.5) * View.ViewRight + -(UV.y-0.5) * View.ViewUp;
+		const float3 RayD = -V;
+		const float4 Sphere = float4(0.0f, 0.0f, 0.f, 0.48);
+		const float2 Hit = RayIntersectSphere(RayO, RayD, Sphere);
+		if (Hit.x >= 0)
+		{
+			float3 P = RayO + RayD * Hit.x;
+			float3 N = normalize(P - Sphere.xyz);
+			const float ThicknessRatio = (Hit.y - Hit.x) / Sphere.w;
+
+			if (BSDF_GETISSIMPLEVOLUME(BSDF))
+			{
+				// Apply transmittance on the background
+
+				float3 DiffuseColor = SLAB_DIFFUSEALBEDO(BSDF);
+				FParticipatingMedia PM = StrataSlabCreateParticipatingMedia(DiffuseColor, SLAB_SSSMFP(BSDF));
+				const float3 Transmittance = IsotropicMediumSlabTransmittance(PM, ThicknessRatio * STRATA_SIMPLEVOLUME_THICKNESS_M, 1.f);
+
+				OutColor.rgb *= Transmittance * BSDF.Coverage + (1.0f - BSDF.Coverage);
+			}
+			else
+			{
+				OutColor.rgb *= 0.0f;
+			}
+
+			// Create a head that will map to a single BSDF
+			FStrataPixelHeader StrataPixelHeader = InitialiseStrataPixelHeader();
+			StrataPixelHeader.BSDFCount = 1;
+
+			FStrataIntegrationSettings Settings = InitStrataIntegrationSettings();
+			FStrataAddressing NullStrataAddressing = (FStrataAddressing)0;
+
+			float3x3 TangentBasis;
+			TangentBasis[2] = N;
+			TangentBasis[0] = normalize(N*float3(1, 0, 0));
+			TangentBasis[1] = cross(N, TangentBasis[0]);
+			TangentBasis[0] =-cross(N, TangentBasis[1]);
+
+			// Apply a white directional light
+			if(ForwardLightData.HasDirectionalLight)
+			{
+				FDeferredLightData DirLightData = (FDeferredLightData)0;
+				DirLightData.Color = ForwardLightData.DirectionalLightColor;
+				DirLightData.FalloffExponent = 0;
+				DirLightData.Direction = ForwardLightData.DirectionalLightDirection;
+				DirLightData.DistanceFadeMAD = ForwardLightData.DirectionalLightDistanceFadeMAD;
+				DirLightData.bRadialLight = false;
+				DirLightData.SpecularScale = 1;
+				float4 DirLightAttenuation = float4(1, 1, 1, 1);
+				DirLightData.ShadowedBits = 1;
+				DirLightData.ShadowMapChannelMask.x = 1;
+				DirLightData.HairTransmittance = InitHairTransmittanceData();
+
+				float Dither = 0.0f;
+				float3 ToLight = DirLightData.Direction;
+				float LightMask = 1.0f;
+				FRectTexture RectTexture = InitRectTexture();
+				uint LightChannelMask = 0xFFFFFFFF;
+				uint PrimitiveLightingChannelMask = LightChannelMask;
+
+				float3 BSDFColoredVisibility = 1.0f;
+
+				// Patch the BSDF to appear as a top layer
+				BSDF.LuminanceWeightV = 1.0f;
+				BSDF.LuminanceWeightL = 1.0f;
+				BSDF.bIsBottom = true;
+				BSDF.bIsTop = true;
+				BSDF.Coverage = 1.0f;
+				BSDF.Emissive = 0.0f;
+
+				FStrataBSDFContext StrataBSDFContext = StrataCreateBSDFContext(TangentBasis, BSDF, V, ToLight);
+
+				bool bEvaluateHasBeenExecuted = false;
+				FStrataEvaluateResult BSDFEvaluate = (FStrataEvaluateResult)0;
+				float3 DirLightLuminance = StrataForwardLightingCommon(
+					Dither,
+					Settings,
+					DirLightData,
+					ToLight,
+					LightMask,
+					DirLightAttenuation,
+					RectTexture,
+					LightChannelMask,
+					PrimitiveLightingChannelMask,
+					TranslatedWorldPosition,
+					SceneDepth,
+					BSDFColoredVisibility,
+					StrataPixelHeader,
+					StrataBSDFContext,
+					bEvaluateHasBeenExecuted,
+					BSDFEvaluate);
+
+				OutColor.rgb += DirLightLuminance * View.PreExposure;
+			}
+
+			// Apply the sky box
+			{
+				float3 ToLight = V;
+				FStrataBSDFContext StrataBSDFContext = StrataCreateBSDFContext(TangentBasis, BSDF, V, ToLight);
+
+				const bool bEnableSpecular = ReflectionStruct.SkyLightParameters.y > 0.0f;
+				FStrataEnvLightResult StrataEnvLight = StrataEvaluateForEnvLight(StrataBSDFContext, bEnableSpecular, Settings);
+
+				float SSRReductionFactor = 0.0;
+				float3 DiffuseLighting = 0.0;
+				float3 SpecularLighting = 0.0;
+				StrataEnvLightingCommon(
+					StrataEnvLight,
+					StrataPixelHeader,
+					StrataBSDFContext,
+					BSDF,
+					N,	//in float3 BentNormal,
+					1.0,//in float3 BSDFThroughput,
+					0,	//in uint CaptureDataStartIndex,
+					0,	//in uint NumCulledReflectionCaptures,
+					1.0,//in float ScreenAmbientOcclusion,
+					1.0,//in float CloudVolumetricAOShadow,
+					1.0,//in float TopLayerSpecularContributionFactor,
+					TranslatedWorldPosition,
+					1.0,//in float CombinedScreenAndMaterialAO,
+					SSRReductionFactor,
+					DiffuseLighting,
+					SpecularLighting);
+
+				OutColor.rgb += (DiffuseLighting + SpecularLighting) * View.PreExposure;;
+			}
+		}
 	}
-	return OutColor;
 }
 
 void MaterialDebugStrataTreePS(
@@ -510,7 +679,7 @@ void MaterialDebugStrataTreePS(
 	FStrataPixelHeader Header = UnpackStrataHeaderIn(Strata.MaterialTextureArray, StrataAddressing, Strata.TopLayerTexture);
 	if (Header.BSDFCount > 0)
 	{
-		OutColor = float4(0.0f, 0.0f, 0.0f, 0.5f);
+		OutColor = float4(0.0f, 0.0f, 0.0f, 0.9f);
 
 		FStrataTree StrataTree = GetInitialisedStrataTree();
 		StrataUnpackInStrataTree(
@@ -521,21 +690,21 @@ void MaterialDebugStrataTreePS(
 		const float2 PixelPos = SVPos.xy;
 
 		float2 Pos = float2(100.0f, 100.0f);
-		const float SphereSize = 100.0f;
+		const float SphereSize = 200.0f;
 		for (int i = 0; i < StrataTree.BSDFCount; i++)
 		{
-			OutColor = GetMaterialOnSphere(
-				PixelPos,
+			GetMaterialOnSphere(
+				SVPos,
 				Pos,
 				SphereSize,
 				StrataTree.BSDFs[i],
 				OutColor);
-		
+
 			Pos.y += SphereSize + 10.0f;
 		}
 	}
 
-	OutColor = float4(pow(OutColor.rgb,1.0/2.2), OutColor.a);
+	OutColor = float4(pow(OutColor.rgb, 1.0 / 2.2), OutColor.a);
 }
 
-#endif // SHADER_DEBUGSTRATATREE
+#endif // SHADER_DEBUGSTRATATREE_PS
diff --git a/Engine/Source/Runtime/Renderer/Private/Strata/StrataVisualize.cpp b/Engine/Source/Runtime/Renderer/Private/Strata/StrataVisualize.cpp
--- a/Engine/Source/Runtime/Renderer/Private/Strata/StrataVisualize.cpp
+++ b/Engine/Source/Runtime/Renderer/Private/Strata/StrataVisualize.cpp
@@ -6,6 +6,7 @@
 #include "ScreenPass.h"
 #include "ShaderCompiler.h"
 #include "PixelShaderUtils.h"
+#include "BasePassRendering.h"
 
 
 static TAutoConsoleVariable<int32> CVarStrataClassificationDebug(
@@ -151,7 +152,8 @@ class FMaterialDebugStrataTreeCS : public FGlobalShader
 
 		// Stay debug and skip optimizations to reduce compilation time on this long shader.
 		OutEnvironment.CompilerFlags.Add(CFLAG_Debug);
-		OutEnvironment.SetDefine(TEXT("SHADER_DEBUGSTRATATREE"), 1);
+		OutEnvironment.CompilerFlags.Add(CFLAG_ForceDXC);
+		OutEnvironment.SetDefine(TEXT("SHADER_DEBUGSTRATATREE_TEST_CS"), 1);
 	}
 };
 IMPLEMENT_GLOBAL_SHADER(FMaterialDebugStrataTreeCS, "/Engine/Private/Strata/StrataVisualize.usf", "MaterialDebugStrataTreeCS", SF_Compute);
@@ -166,6 +168,9 @@ class FMaterialDebugStrataTreePS : public FGlobalShader
 	BEGIN_SHADER_PARAMETER_STRUCT(FParameters, )
 		SHADER_PARAMETER_STRUCT_REF(FViewUniformShaderParameters, ViewUniformBuffer)
 		SHADER_PARAMETER_RDG_UNIFORM_BUFFER(FStrataGlobalUniformParameters, Strata)
+		SHADER_PARAMETER_STRUCT_REF(FReflectionUniformParameters, ReflectionStruct)
+		SHADER_PARAMETER_STRUCT_REF(FReflectionCaptureShaderData, ReflectionCapture)
+		SHADER_PARAMETER_RDG_UNIFORM_BUFFER(FForwardLightData, ForwardLightData)
 		RENDER_TARGET_BINDING_SLOTS()
 	END_SHADER_PARAMETER_STRUCT()
 
@@ -185,7 +190,7 @@ class FMaterialDebugStrataTreePS : public FGlobalShader
 
 		// Stay debug and skip optimizations to reduce compilation time on this long shader.
 		OutEnvironment.CompilerFlags.Add(CFLAG_Debug);
-		OutEnvironment.SetDefine(TEXT("SHADER_DEBUGSTRATATREE"), 1);
+		OutEnvironment.SetDefine(TEXT("SHADER_DEBUGSTRATATREE_PS"), 1);
 	}
 };
 IMPLEMENT_GLOBAL_SHADER(FMaterialDebugStrataTreePS, "/Engine/Private/Strata/StrataVisualize.usf", "MaterialDebugStrataTreePS", SF_Pixel);
@@ -264,6 +269,9 @@ static void AddVisualizeMaterialPasses(FRDGBuilder& GraphBuilder, const FViewInf
 					FMaterialDebugStrataTreePS::FParameters* PassParameters = GraphBuilder.AllocParameters<FMaterialDebugStrataTreePS::FParameters>();
 					PassParameters->ViewUniformBuffer = View.ViewUniformBuffer;
 					PassParameters->Strata = Strata::BindStrataGlobalUniformParameters(View);
+					PassParameters->ReflectionStruct = CreateReflectionUniformBuffer(View, UniformBuffer_SingleFrame);
+					PassParameters->ReflectionCapture = View.ReflectionCaptureUniformBuffer;
+					PassParameters->ForwardLightData = View.ForwardLightingResources.ForwardLightUniformBuffer;
 					PassParameters->RenderTargets[0] = FRenderTargetBinding(SceneColorTexture, ERenderTargetLoadAction::ELoad);
 
 					FMaterialDebugStrataTreePS::FPermutationDomain PermutationVector;
