Use MakePrecise() on WPO
This was happening for FXC compiler but not DXC compiler.
It was causing broken velocity vectors, and broken editor outlines on SM6.

diff --git a/Engine/Shaders/Private/MaterialTemplate.ush b/Engine/Shaders/Private/MaterialTemplate.ush
--- a/Engine/Shaders/Private/MaterialTemplate.ush
+++ b/Engine/Shaders/Private/MaterialTemplate.ush
@@ -2736,51 +2736,63 @@ half GetMaterialMaskedOpacity(FPixelMaterialInputs PixelMaterialInputs)
 }
 #endif
 
-// UE-102725: Workaround for precision issue that only occurs with world position offsets in DXBC shaders.
-#if IR_LANGUAGE_DXBC
-#define WPO_PRECISE(T) precise T
+// UE-102725: Workaround for precision issue that only occurs with world position offsets in D3D shaders.
+#if USE_WPO_PRECISE
+#define WPO_PRECISE(T) MakePrecise(T)
 #else
 #define WPO_PRECISE(T) T
 #endif
 
-WPO_PRECISE(float3) GetMaterialWorldPositionOffset(FMaterialVertexParameters Parameters)
+float3 GetMaterialWorldPositionOffsetInternal(FMaterialVertexParameters Parameters)
 {
-	#if IS_NANITE_PASS
-		BRANCH
+%s;
+}
+
+float3 GetMaterialWorldPositionOffset(FMaterialVertexParameters Parameters)
+{
+#if IS_NANITE_PASS
+	BRANCH
 		if (!Parameters.bEvaluateWorldPositionOffset)
 		{
-			return float3(0,0,0);
+			return float3(0, 0, 0);
 		}
-	#endif
+#endif
 
-	#if USE_INSTANCING || USE_INSTANCE_CULLING
-		// skip if this instance is hidden
-		if (Parameters.PerInstanceParams.z < 1.f)
-		{
-			return float3(0,0,0);
-		}
-	#endif
+#if USE_INSTANCING || USE_INSTANCE_CULLING
+	// skip if this instance is hidden
+	if (Parameters.PerInstanceParams.z < 1.f)
+	{
+		return float3(0, 0, 0);
+	}
+#endif
+
+	return WPO_PRECISE(GetMaterialWorldPositionOffsetInternal(Parameters));
+}
+
+float3 GetMaterialPreviousWorldPositionOffsetInternal(FMaterialVertexParameters Parameters)
+{
 %s;
 }
 
-WPO_PRECISE(float3) GetMaterialPreviousWorldPositionOffset(FMaterialVertexParameters Parameters)
+float3 GetMaterialPreviousWorldPositionOffset(FMaterialVertexParameters Parameters)
 {
-	#if IS_NANITE_PASS
-		BRANCH
+#if IS_NANITE_PASS
+	BRANCH
 		if (!Parameters.bEvaluateWorldPositionOffset)
 		{
-			return float3(0,0,0);
+			return float3(0, 0, 0);
 		}
-	#endif
+#endif
 
-	#if USE_INSTANCING || USE_INSTANCE_CULLING
-		// skip if this instance is hidden
-		if (Parameters.PerInstanceParams.z < 1.f)
-		{
-			return float3(0,0,0);
-		}
-	#endif
-%s;
+#if USE_INSTANCING || USE_INSTANCE_CULLING
+	// skip if this instance is hidden
+	if (Parameters.PerInstanceParams.z < 1.f)
+	{
+		return float3(0, 0, 0);
+	}
+#endif
+
+	return WPO_PRECISE(GetMaterialPreviousWorldPositionOffsetInternal(Parameters));
 }
 
 
diff --git a/Engine/Source/Developer/Windows/ShaderFormatD3D/Private/D3DShaderCompiler.cpp b/Engine/Source/Developer/Windows/ShaderFormatD3D/Private/D3DShaderCompiler.cpp
--- a/Engine/Source/Developer/Windows/ShaderFormatD3D/Private/D3DShaderCompiler.cpp
+++ b/Engine/Source/Developer/Windows/ShaderFormatD3D/Private/D3DShaderCompiler.cpp
@@ -1172,6 +1172,8 @@ void CompileD3DShader(const FShaderCompilerInput& Input, FShaderCompilerOutput&
 	// Set additional defines.
 	AdditionalDefines.SetDefine(TEXT("COMPILER_HLSL"), 1);
 	AdditionalDefines.SetDefine(TEXT("PLATFORM_SUPPORTS_ROV"), 1); // Assume min. spec HW supports with DX12/SM5
+
+	AdditionalDefines.SetDefine(TEXT("USE_WPO_PRECISE"), 1);
 	
 	if (bUseDXC)
 	{
@@ -1190,11 +1192,7 @@ void CompileD3DShader(const FShaderCompilerInput& Input, FShaderCompilerOutput&
 			AdditionalDefines.SetDefine(TEXT("PLATFORM_SUPPORTS_REAL_TYPES"), 1);
 		}
 	}
-	else
-	{
-		AdditionalDefines.SetDefine(TEXT("IR_LANGUAGE_DXBC"), 1);
-	}
-
+	
 	if (Input.bSkipPreprocessedCache)
 	{
 		if (!FFileHelper::LoadFileToString(PreprocessedShaderSource, *Input.VirtualSourceFilePath))
