Add retro-reflective vector computation to Strata context.

These are not currently used.

diff --git a/Engine/Shaders/Private/Strata/StrataEvaluation.ush b/Engine/Shaders/Private/Strata/StrataEvaluation.ush
--- a/Engine/Shaders/Private/Strata/StrataEvaluation.ush
+++ b/Engine/Shaders/Private/Strata/StrataEvaluation.ush
@@ -77,10 +77,6 @@ float3 StrataGetOffSpecularPeakReflectionDir(float3 Normal, float3 ReflectionVec
 	return lerp(Normal, ReflectionVector, (1 - a) * (sqrt(1 - a) + a));
 }
 
-// When haziness is used, both lobes take equal part in the blending. Only roughness changes.
-#define SPECULAR_HAZINESS_WEIGHT 0.5f
-
-
 ///////////////////////////////////////////////////////////////////////////////
 // BSDF evaluate and sampling
 
@@ -94,6 +90,7 @@ struct FStrataBSDFContext
 	float3 V;
 	float3 R;
 	float3 H;
+	float3 B;
 
 	float3 L;	// There to initialise the BxDFContext. Only used by StrataEvaluateBSDF, not by StrataImportanceSampleBSDF or StrataEvaluateForEnvLight
 
@@ -104,6 +101,7 @@ struct FStrataBSDFContext
 	float3x3 TangentBasis;
 	float3 TangentV;
 	float3 TangentH;
+	float3 TangentB;
 };
 
 FStrataBSDFContext StrataCreateBSDFContext(float3x3 TangentBasis, FStrataBSDF BSDF, float3 V, float3 L, bool bHasValidL=true)
@@ -119,6 +117,7 @@ FStrataBSDFContext StrataCreateBSDFContext(float3x3 TangentBasis, FStrataBSDF BS
 	BSDFContext.R = 2 * dot(BSDFContext.V, BSDFContext.N) * BSDFContext.N - BSDFContext.V;
 	BSDFContext.L = bHasValidL ? L : BSDFContext.R;
 	BSDFContext.H = normalize(BSDFContext.V + BSDFContext.L);
+	BSDFContext.B = normalize(BSDFContext.R + BSDFContext.L);
 
 	BSDFContext.Context = (BxDFContext)0;
 	if (BSDF_GETHASANISOTROPY(BSDF) != 0)
@@ -133,6 +132,7 @@ FStrataBSDFContext StrataCreateBSDFContext(float3x3 TangentBasis, FStrataBSDF BS
 
 	BSDFContext.TangentV = mul(BSDFContext.TangentBasis, BSDFContext.V);
 	BSDFContext.TangentH = mul(BSDFContext.TangentBasis, BSDFContext.H);
+	BSDFContext.TangentB = mul(BSDFContext.TangentBasis, BSDFContext.B);
 
 	BSDFContext.SatNoL = saturate(BSDFContext.Context.NoL);
 	BSDFContext.SatNoV = saturate(BSDFContext.Context.NoV);
@@ -158,11 +158,15 @@ void StrataUpdateBSDFContext(inout FStrataBSDFContext BSDFContext, float3 NewL)
 	// Update all the data related to L
 	BSDFContext.L = NewL;
 	Init(BSDFContext.Context, BSDFContext.N, BSDFContext.V, BSDFContext.L);
-	BSDFContext.H = normalize(BSDFContext.V + BSDFContext.L);
-	BSDFContext.R = 2 * dot(BSDFContext.V, BSDFContext.N) * BSDFContext.N - BSDFContext.V;
+	BSDFContext.R		 = 2 * dot(BSDFContext.V, BSDFContext.N) * BSDFContext.N - BSDFContext.V;
+	BSDFContext.H		 = normalize(BSDFContext.V + BSDFContext.L);
+	BSDFContext.B		 = normalize(BSDFContext.R + BSDFContext.L);
+
 	BSDFContext.TangentH = mul(BSDFContext.TangentBasis, BSDFContext.H);
-	BSDFContext.SatNoL = saturate(BSDFContext.Context.NoL);
-	BSDFContext.SatNoV = saturate(BSDFContext.Context.NoV);
+	BSDFContext.TangentB = mul(BSDFContext.TangentBasis, BSDFContext.B);
+
+	BSDFContext.SatNoL	 = saturate(BSDFContext.Context.NoL);
+	BSDFContext.SatNoV	 = saturate(BSDFContext.Context.NoV);
 }
 
 ///////////////////////////////////////////////////////////////////////////////
