Strata - visualisation - do not show BSDF with coverage of 0 and do not present BSDF as operator to the user.

diff --git a/Engine/Shaders/Private/Strata/StrataVisualize.usf b/Engine/Shaders/Private/Strata/StrataVisualize.usf
--- a/Engine/Shaders/Private/Strata/StrataVisualize.usf
+++ b/Engine/Shaders/Private/Strata/StrataVisualize.usf
@@ -476,11 +476,14 @@ void MaterialDebugStrataTreeCS(uint3 DispatchThreadId : SV_DispatchThreadID)
 		FShaderPrintContext Context;
 		Context = InitShaderPrintContext(true, uint2(50, 50));
 
+		const int BSDFCount = StrataTree.BSDFCount;
+		// We remove BSDF from operators to no confuse the user. Behind the scene BSDF are operators sitting the leaves of the tree.
+		const int OperatorCount = StrataTree.OperatorCount - StrataTree.BSDFCount;
 
-		Print(Context, StrataTree.BSDFCount, FontWhite, 2, 0);
+		Print(Context, BSDFCount, FontWhite, 2, 0);
 		Print(Context, TEXT("BSDFs"), FontWhite);
 		Newline(Context);
-		Print(Context, StrataTree.OperatorCount, FontWhite, 2, 0);
+		Print(Context, OperatorCount, FontWhite, 2, 0);
 		Print(Context, TEXT("Operators"), FontWhite);
 		Newline(Context);
 	}
@@ -701,14 +704,17 @@ void MaterialDebugStrataTreePS(
 #endif
 		for (int i = 0; i < StrataTree.BSDFCount; i++)
 		{
-			GetMaterialOnSphere(
-				SVPos,
-				Pos,
-				SphereSize,
-				StrataTree.BSDFs[i],
-				OutColor);
-
-			Pos.y += SphereSize + 10.0f;
+			if (StrataTree.BSDFs[i].Coverage > 0.0f)
+			{
+				GetMaterialOnSphere(
+					SVPos,
+					Pos,
+					SphereSize,
+					StrataTree.BSDFs[i],
+					OutColor);
+
+				Pos.y += SphereSize + 10.0f;
+			}
 		}
 	}
 
