Strata - debug starta tree visualisation load BSDFs on the fly to reduce compiler issue and increase occupancy.

diff --git a/Engine/Shaders/Private/Strata/Strata.ush b/Engine/Shaders/Private/Strata/Strata.ush
--- a/Engine/Shaders/Private/Strata/Strata.ush
+++ b/Engine/Shaders/Private/Strata/Strata.ush
@@ -825,6 +825,7 @@ void StrataPackOutStrataTree(
 #define WRITE_FLOAT3(Data) WRITE_FLOAT(Data.x);WRITE_FLOAT(Data.y);WRITE_FLOAT(Data.z);
 #define WRITE_FLOAT4(Data) WRITE_FLOAT(Data.x);WRITE_FLOAT(Data.y);WRITE_FLOAT(Data.z);;WRITE_FLOAT(Data.w);
 
+	WRITE_UINT(STRATA_USES_CONVERSION_FROM_LEGACY ? 1 : 0);
 	WRITE_UINT(StrataTree.BSDFCount);
 	WRITE_UINT(StrataTree.OperatorCount);
 
@@ -868,60 +869,75 @@ void StrataPackOutStrataTree(
 #undef WRITE_FLOAT4
 }
 
-void StrataUnpackInStrataTree(
-	Texture2DArray<uint> StrataMaterialData,
-	inout FStrataTree StrataTree,
-	in int SliceStoringDebugStrataTree)
+struct FStrataTreeHeader
 {
-	StrataTree = GetInitialisedStrataTree();
-	int OutCoord = 0;
+	int ConvertedFromLegacy;
+	int BSDFCount;
+	int OperatorCount;
+};
+
 #define READ_UINT(Data) Data = StrataMaterialData[uint3(OutCoord++, 0, SliceStoringDebugStrataTree)];
 #define READ_FLOAT(Data) Data = asfloat(StrataMaterialData[uint3(OutCoord++, 0, SliceStoringDebugStrataTree)]);
 #define READ_FLOAT3(Data) READ_FLOAT(Data.x);READ_FLOAT(Data.y);READ_FLOAT(Data.z);
 #define READ_FLOAT4(Data) READ_FLOAT(Data.x);READ_FLOAT(Data.y);READ_FLOAT(Data.z);READ_FLOAT(Data.w);
 
-	READ_UINT(StrataTree.BSDFCount);
-	READ_UINT(StrataTree.OperatorCount);
+void StrataUnpackInStrataTreeHeader(
+	Texture2DArray<uint> StrataMaterialData,
+	inout FStrataTreeHeader StrataTreeHeader,
+	in int SliceStoringDebugStrataTree)
+{
+	StrataTreeHeader = (FStrataTreeHeader)0;
 
-#if COMPILER_DXC == 1
-	LOOP
-#else
-	StrataTree.BSDFCount = min(StrataTree.BSDFCount, STRATA_MAX_BSDF_COUNT);
-	STRATA_UNROLL_N(STRATA_MAX_BSDF_COUNT)
-#endif
-	for (int i = 0; i < StrataTree.BSDFCount; ++i)
-	{
-		READ_UINT(StrataTree.BSDFs[i].State);
-		READ_UINT(StrataTree.BSDFs[i].OperatorIndex);
-		READ_FLOAT3(StrataTree.BSDFs[i].LuminanceWeightV);
-		READ_FLOAT3(StrataTree.BSDFs[i].LuminanceWeightL);
-		READ_UINT(StrataTree.BSDFs[i].bIsBottom);
-		READ_UINT(StrataTree.BSDFs[i].bIsTop);
+	int OutCoord = 0;
 
-	#if STRATA_INLINE_SHADING
-		READ_FLOAT(StrataTree.BSDFs[i].Coverage);
+	READ_UINT(StrataTreeHeader.ConvertedFromLegacy);
+	READ_UINT(StrataTreeHeader.BSDFCount);
+	READ_UINT(StrataTreeHeader.OperatorCount);
+}
 
-		READ_FLOAT3(StrataTree.BSDFs[i].Emissive);
-		READ_FLOAT(StrataTree.BSDFs[i].ThicknessCm);
-		READ_FLOAT3(StrataTree.BSDFs[i].TmpMFP);
+FStrataBSDF StrataUnpackInStrataTreeBSDF(
+	uint BSDFIndex,
+	Texture2DArray<uint> StrataMaterialData,
+	inout FStrataTreeHeader StrataTreeHeader,
+	in int SliceStoringDebugStrataTree)
+{
+	int OutCoord = 3;												// Header
+	OutCoord += BSDFIndex * (STRATA_INLINE_SHADING ? 39 : 30);		// Skip BSDFs
 
-		READ_FLOAT(StrataTree.BSDFs[i].TopLayerDataWeight);
-	#endif
+	FStrataBSDF BSDF = (FStrataBSDF)0;
 
-		// Skipping specific single layer water data
+	READ_UINT(BSDF.State);
+	READ_UINT(BSDF.OperatorIndex);
+	READ_FLOAT3(BSDF.LuminanceWeightV);
+	READ_FLOAT3(BSDF.LuminanceWeightL);
+	READ_UINT(BSDF.bIsBottom);
+	READ_UINT(BSDF.bIsTop);
 
-		READ_FLOAT4(StrataTree.BSDFs[i].VGPRs[0]);
-		READ_FLOAT4(StrataTree.BSDFs[i].VGPRs[1]);
-		READ_FLOAT4(StrataTree.BSDFs[i].VGPRs[2]);
-		READ_FLOAT4(StrataTree.BSDFs[i].VGPRs[3]);
-		READ_FLOAT4(StrataTree.BSDFs[i].VGPRs[4]);
-	}
+#if STRATA_INLINE_SHADING
+	READ_FLOAT(BSDF.Coverage);
+
+	READ_FLOAT3(BSDF.Emissive);
+	READ_FLOAT(BSDF.ThicknessCm);
+	READ_FLOAT3(BSDF.TmpMFP);
+
+	READ_FLOAT(BSDF.TopLayerDataWeight);
+#endif
+
+	// Skipping specific single layer water data
+
+	READ_FLOAT4(BSDF.VGPRs[0]);
+	READ_FLOAT4(BSDF.VGPRs[1]);
+	READ_FLOAT4(BSDF.VGPRs[2]);
+	READ_FLOAT4(BSDF.VGPRs[3]);
+	READ_FLOAT4(BSDF.VGPRs[4]);
+
+	return BSDF;
+}
 
 #undef READ_UINT
 #undef READ_FLOAT
 #undef READ_FLOAT3
 #undef READ_FLOAT4
-}
 
 
 struct FStrataData
diff --git a/Engine/Shaders/Private/Strata/StrataVisualize.usf b/Engine/Shaders/Private/Strata/StrataVisualize.usf
--- a/Engine/Shaders/Private/Strata/StrataVisualize.usf
+++ b/Engine/Shaders/Private/Strata/StrataVisualize.usf
@@ -466,22 +466,27 @@ void MaterialDebugStrataTreeCS(uint3 DispatchThreadId : SV_DispatchThreadID)
 	FStrataPixelHeader Header = UnpackStrataHeaderIn(Strata.MaterialTextureArray, StrataAddressing, Strata.TopLayerTexture);
 	if (Header.BSDFCount > 0)
 	{
-		FStrataTree StrataTree = GetInitialisedStrataTree();
-		StrataUnpackInStrataTree(
-			Strata.MaterialTextureArray,
-			StrataTree,
-			Strata.SliceStoringDebugStrataTree + STRATA_BASE_PASS_MRT_OUTPUT_COUNT);
+		const int SliceStoringDebugStrataTree = Strata.SliceStoringDebugStrataTree + STRATA_BASE_PASS_MRT_OUTPUT_COUNT;
 
+		FStrataTreeHeader StrataTreeHeader = (FStrataTreeHeader)0;
+		StrataUnpackInStrataTreeHeader(
+			Strata.MaterialTextureArray,
+			StrataTreeHeader,
+			SliceStoringDebugStrataTree);
 
 		FShaderPrintContext Context;
 		Context = InitShaderPrintContext(true, uint2(50, 50));
 
-		const int BSDFCount = StrataTree.BSDFCount;
+		const int BSDFCount = StrataTreeHeader.BSDFCount;
 		// We remove BSDF from operators to no confuse the user. Behind the scene BSDF are operators sitting the leaves of the tree.
-		const int OperatorCount = StrataTree.OperatorCount - StrataTree.BSDFCount;
+		const int OperatorCount = StrataTreeHeader.OperatorCount - StrataTreeHeader.BSDFCount;
 
 		Print(Context, BSDFCount, FontWhite, 2, 0);
 		Print(Context, TEXT("BSDFs"), FontWhite);
+		if (StrataTreeHeader.ConvertedFromLegacy)
+		{
+			Print(Context, TEXT(" (Converted from legacy material)"), FontWhite);
+		}
 		Newline(Context);
 		Print(Context, OperatorCount, FontWhite, 2, 0);
 		Print(Context, TEXT("Operators"), FontWhite);
@@ -684,11 +689,13 @@ void MaterialDebugStrataTreePS(
 	{
 		OutColor = float4(0.0f, 0.0f, 0.0f, 0.9f);
 
-		FStrataTree StrataTree = GetInitialisedStrataTree();
-		StrataUnpackInStrataTree(
+		const int SliceStoringDebugStrataTree = Strata.SliceStoringDebugStrataTree + STRATA_BASE_PASS_MRT_OUTPUT_COUNT;
+
+		FStrataTreeHeader StrataTreeHeader = (FStrataTreeHeader)0;
+		StrataUnpackInStrataTreeHeader(
 			Strata.MaterialTextureArray,
-			StrataTree,
-			Strata.SliceStoringDebugStrataTree + STRATA_BASE_PASS_MRT_OUTPUT_COUNT);
+			StrataTreeHeader,
+			SliceStoringDebugStrataTree);
 
 		const float2 PixelPos = SVPos.xy;
 
@@ -699,18 +706,24 @@ void MaterialDebugStrataTreePS(
 #if COMPILER_DXC == 1
 		LOOP
 #else
-		StrataTree.BSDFCount = min(StrataTree.BSDFCount, 2); 
+		StrataTree.BSDFCount = min(StrataTreeHeader.BSDFCount, 2);
 		STRATA_UNROLL_N(2)
 #endif
-		for (int i = 0; i < StrataTree.BSDFCount; i++)
+		for (int i = 0; i < StrataTreeHeader.BSDFCount; i++)
 		{
-			if (StrataTree.BSDFs[i].Coverage > 0.0f)
+			FStrataBSDF BSDF = StrataUnpackInStrataTreeBSDF(
+				i,
+				Strata.MaterialTextureArray,
+				StrataTreeHeader,
+				SliceStoringDebugStrataTree);
+
+			if (BSDF.Coverage > 0.0f)
 			{
 				GetMaterialOnSphere(
 					SVPos,
 					Pos,
 					SphereSize,
-					StrataTree.BSDFs[i],
+					BSDF,
 					OutColor);
 
 				Pos.y += SphereSize + 10.0f;
