Strata - fix for emissive to only be affected by LuminanceWeightV (translucent non forward shaded path)

diff --git a/Engine/Shaders/Private/Strata/StrataForwardLighting.ush b/Engine/Shaders/Private/Strata/StrataForwardLighting.ush
--- a/Engine/Shaders/Private/Strata/StrataForwardLighting.ush
+++ b/Engine/Shaders/Private/Strata/StrataForwardLighting.ush
@@ -329,7 +329,7 @@ float3 StrataForwardLighting(
 			#endif
 
 				const float LightVolumeNoL = 1.0f;
-				Color += (SurfaceLighting + BSDFEvaluate.EmissivePathValue) * LuminanceWeight(LightVolumeNoL, CurrentBSDF);
+				Color += SurfaceLighting * LuminanceWeight(LightVolumeNoL, CurrentBSDF) + BSDFEvaluate.EmissivePathValue * CurrentBSDF.LuminanceWeightV;
 
 			#elif FORWARD_PER_PIXEL_SHADING
 
