Strata - fixed bad emissive usage in forward shaded translucent.

diff --git a/Engine/Shaders/Private/Strata/StrataForwardLighting.ush b/Engine/Shaders/Private/Strata/StrataForwardLighting.ush
--- a/Engine/Shaders/Private/Strata/StrataForwardLighting.ush
+++ b/Engine/Shaders/Private/Strata/StrataForwardLighting.ush
@@ -333,8 +333,14 @@ float3 StrataForwardLighting(
 
 			#elif FORWARD_PER_PIXEL_SHADING
 
+
+				////
+				//// Evaluate emissive
+				////
+				Color += BSDF_GETEMISSIVE(CurrentBSDF) * CurrentBSDF.LuminanceWeightV;
+
 				////
-				//// Evaluate Directional light only
+				//// Evaluate the single directional light selected for forward lighting.
 				////
 				{
 					float3 DirLightL = DirLightData.Direction;	// Already normalized
diff --git a/Engine/Shaders/Private/Strata/StrataLightingCommon.ush b/Engine/Shaders/Private/Strata/StrataLightingCommon.ush
--- a/Engine/Shaders/Private/Strata/StrataLightingCommon.ush
+++ b/Engine/Shaders/Private/Strata/StrataLightingCommon.ush
@@ -82,7 +82,7 @@ float3 StrataForwardLightingCommon(
 			{
 				float3 DiffuseLuminance = Diffuse_Lambert(BSDFEvaluate.DiffuseColor) * Lighting;
 				const float3 LightCommonMultiplier = LightData.Color * LightMask;
-				Color += (DiffuseLuminance * LightCommonMultiplier + BSDFEvaluate.EmissivePathValue) * BSDFColoredVisibility;
+				Color += DiffuseLuminance * LightCommonMultiplier * BSDFColoredVisibility;
 			}
 		#else
 			if (LightData.bRectLight)
@@ -114,7 +114,7 @@ float3 StrataForwardLightingCommon(
 				float3 DiffuseLuminance = BSDFEvaluate.IntegratedDiffuseValue;
 				float3 SpecularLuminance = BSDFEvaluate.IntegratedSpecularValue * LightData.SpecularScale;
 				const float3 LightCommonMultiplier = LightData.Color * LightMask;
-				Color += ((DiffuseLuminance + SpecularLuminance) * LightCommonMultiplier + BSDFEvaluate.EmissivePathValue) * BSDFColoredVisibility;
+				Color += (DiffuseLuminance + SpecularLuminance) * LightCommonMultiplier * BSDFColoredVisibility;
 			}
 		#endif
 		}
diff --git a/Engine/Shaders/Private/Strata/StrataVisualize.usf b/Engine/Shaders/Private/Strata/StrataVisualize.usf
--- a/Engine/Shaders/Private/Strata/StrataVisualize.usf
+++ b/Engine/Shaders/Private/Strata/StrataVisualize.usf
@@ -572,6 +572,13 @@ void GetMaterialOnSphere(
 				OutColor.rgb *= 0.0f;
 			}
 
+			// Patch the BSDF to appear as a top layer
+			BSDF.LuminanceWeightV = 1.0f;
+			BSDF.LuminanceWeightL = 1.0f;
+			BSDF.bIsBottom = true;
+			BSDF.bIsTop = true;
+			BSDF.Coverage = 1.0f;
+
 			// Create a head that will map to a single BSDF
 			FStrataPixelHeader StrataPixelHeader = InitialiseStrataPixelHeader();
 			StrataPixelHeader.BSDFCount = 1;
@@ -585,6 +592,9 @@ void GetMaterialOnSphere(
 			TangentBasis[1] = cross(N, TangentBasis[0]);
 			TangentBasis[0] =-cross(N, TangentBasis[1]);
 
+			// Apply emissive
+			OutColor.rgb += BSDF_GETEMISSIVE(BSDF) * View.PreExposure;
+
 			// Apply a white directional light
 			if(ForwardLightData.HasDirectionalLight)
 			{
@@ -609,14 +619,6 @@ void GetMaterialOnSphere(
 
 				float3 BSDFColoredVisibility = 1.0f;
 
-				// Patch the BSDF to appear as a top layer
-				BSDF.LuminanceWeightV = 1.0f;
-				BSDF.LuminanceWeightL = 1.0f;
-				BSDF.bIsBottom = true;
-				BSDF.bIsTop = true;
-				BSDF.Coverage = 1.0f;
-				BSDF.Emissive = 0.0f;
-
 				FStrataBSDFContext StrataBSDFContext = StrataCreateBSDFContext(TangentBasis, BSDF, V, ToLight);
 
 				bool bEvaluateHasBeenExecuted = false;
