Fix for broken starta material classification when using waveops.

diff --git a/Engine/Shaders/Private/Strata/StrataMaterialClassification.usf b/Engine/Shaders/Private/Strata/StrataMaterialClassification.usf
--- a/Engine/Shaders/Private/Strata/StrataMaterialClassification.usf
+++ b/Engine/Shaders/Private/Strata/StrataMaterialClassification.usf
@@ -98,7 +98,7 @@ void TileMainCS(uint2 DispatchThreadId : SV_DispatchThreadID, uint LinearIndex :
 	const bool bTileContainsSingle							= WaveActiveAnyTrue(bContainsSingleMaterial);
 	const bool bTileContainsComplex							= WaveActiveAnyTrue(bContainsComplexMaterial);
 	const bool bTileContainsOpaqueRoughRefraction			= WaveActiveAnyTrue(bContainsOpaqueRoughRefraction);
-	const bool bTileContainsScreenSpaceSubsurfaceScattering	= WaveActiveAnyTrue(bHasScreenSpaceSubsurfaceScattering);
+	const bool bTileContainsScreenSpaceSubsurfaceScattering	= WaveActiveAnyTrue(bContainsScreenSpaceSubsurfaceScattering);
 #else // PERMUTATION_WAVE_OPS
 
 	s_TileFlags[LinearIndex] = 
