Fix shader build error on PCD3D_ES31.
#rnx

diff --git a/Engine/Shaders/Private/MaterialTemplate.ush b/Engine/Shaders/Private/MaterialTemplate.ush
--- a/Engine/Shaders/Private/MaterialTemplate.ush
+++ b/Engine/Shaders/Private/MaterialTemplate.ush
@@ -2738,7 +2738,7 @@ half GetMaterialMaskedOpacity(FPixelMaterialInputs PixelMaterialInputs)
 
 // UE-102725: Workaround for precision issue that only occurs with world position offsets in D3D shaders.
 #if USE_WPO_PRECISE
-#define WPO_PRECISE(T) MakePrecise(T)
+#define WPO_PRECISE(T) INVARIANT(T)
 #else
 #define WPO_PRECISE(T) T
 #endif
