Strata - fixed errors with transmittanceV and transmittanceL (swapped values and coverage that hsould not be applied on transmittanceL)

diff --git a/Engine/Shaders/Private/Strata/StrataTree.ush b/Engine/Shaders/Private/Strata/StrataTree.ush
--- a/Engine/Shaders/Private/Strata/StrataTree.ush
+++ b/Engine/Shaders/Private/Strata/StrataTree.ush
@@ -216,8 +216,8 @@ void UpdateSingleOperatorCoverageTransmittance(inout FStrataTree StrataTree, int
 #endif
 
 		Op.VerticalTop_Coverage		  = OpA.Coverage;
-		Op.VerticalTop_TransmittanceL = OpA.TransmittanceV;
-		Op.VerticalTop_TransmittanceV = OpA.TransmittanceL;
+		Op.VerticalTop_TransmittanceV = OpA.TransmittanceV;
+		Op.VerticalTop_TransmittanceL = OpA.TransmittanceL;
 
 		// Project top lob through bottom layer assuming medium have same IOR.
 		{
@@ -332,7 +332,7 @@ void UpdateAllBSDFWithBottomUpOperatorVisit_Weight(
 
 	const float Weight = saturate(Op.Weight);
 	CurrentBSDF.LuminanceWeightV *= Weight;
-	CurrentBSDF.LuminanceWeightL *= Weight;
+	//CurrentBSDF.LuminanceWeightL *= Weight; A weight affecting the coverage of a BSDF won't affect the weight to light. Only the LuminanceWeightV to simulate a change of coverage.
 	CurrentBSDF.Coverage *= Weight;
 
 	CurrentBSDF.TopLayerDataWeight *= Weight;
@@ -353,7 +353,7 @@ void UpdateAllBSDFWithBottomUpOperatorVisit_Horizontal(
 	const float BMix = Mix;
 	const float Weight = PreviousIsInputA > 0 ? AMix : BMix;
 	CurrentBSDF.LuminanceWeightV *= Weight;
-	CurrentBSDF.LuminanceWeightL *= Weight;
+	//CurrentBSDF.LuminanceWeightL *= Weight; A weight affecting the coverage of a BSDF won't affect the weight to light. Only the LuminanceWeightV to simulate a change of coverage.
 	CurrentBSDF.Coverage *= Weight;
 
 	CurrentBSDF.TopLayerDataWeight *= Weight;
