Strata - clarify the fact that SSS is not translucent "see-through" as of today.

diff --git a/Engine/Shaders/Private/Strata/StrataEvaluation.ush b/Engine/Shaders/Private/Strata/StrataEvaluation.ush
--- a/Engine/Shaders/Private/Strata/StrataEvaluation.ush
+++ b/Engine/Shaders/Private/Strata/StrataEvaluation.ush
@@ -586,8 +586,8 @@ FStrataEvaluateResult StrataEvaluateBSDFCommon(FStrataBSDFContext BSDFContext, F
 				const float3 RefracV = refract(BSDFContext.V, -BSDFContext.N, OneOverIOR);
 				const float PhaseFunction = ApproximateHG(dot(-BSDFContext.L, RefracV), PhaseFunctionAnisotropy);
 
-				Sample.ThroughputV				= TransmissionThroughput * PhaseFunction;
-				Sample.ThroughputL				= TransmissionThroughput * PhaseFunction;
+				Sample.ThroughputV				= OpaqueBSDFThroughput; // SSS is not translucent as of today
+				Sample.ThroughputL				= OpaqueBSDFThroughput; // idem
 				Sample.TransmissionPathValue	= TransmissionThroughput * PhaseFunction;
 				Sample.TransmissionPDF			= 1.0f / (4.0f * PI);		// STRATA_TODO this currently match the uniform sphere sampling from StrataImportanceSampleBSDF
 
