Strata - added cube for material vismode 3

diff --git a/Engine/Shaders/Private/Strata/StrataVisualize.usf b/Engine/Shaders/Private/Strata/StrataVisualize.usf
--- a/Engine/Shaders/Private/Strata/StrataVisualize.usf
+++ b/Engine/Shaders/Private/Strata/StrataVisualize.usf
@@ -528,7 +528,147 @@ float4 GetBackgroundCheckboardColor(float2 UV)
 	return float4(Color, 0.0);
 }
 
-void GetMaterialOnSphere(
+// Updated from http://jcgt.org/published/0007/03/04/
+bool slabs(float3 p0, float3 p1, float3 rayOrigin, float3 invRaydir, out float outTMin, out float outTMax)
+{
+	float3 t0 = (p0 - rayOrigin) * invRaydir;
+	float3 t1 = (p1 - rayOrigin) * invRaydir;
+	float3 tmin = min(t0, t1), tmax = max(t0, t1);
+	float maxtmin = max(max(tmin.x, tmin.y), tmin.z);
+	float mintmax = min(min(tmax.x, tmax.y), tmax.z);
+	outTMin = maxtmin;
+	outTMax = mintmax;
+	return maxtmin <= mintmax;
+}
+
+void EvaluateLighting(
+	in FStrataBSDF BSDF,
+	in float3 V,
+	in float3x3 TangentBasis,
+	in float ThicknessRatio,
+	inout float4 OutColor)
+{
+	float3 AbsoluteWorldPosition = LWCHackToFloat(PrimaryView.WorldCameraOrigin);
+	float3 TranslatedWorldPosition = PrimaryView.TranslatedWorldCameraOrigin;
+	float SceneDepth = 0.0;
+
+	const float3 N = TangentBasis[2];
+
+	if (BSDF_GETISSIMPLEVOLUME(BSDF))
+	{
+		// Apply transmittance on the background
+
+		float3 DiffuseColor = SLAB_DIFFUSEALBEDO(BSDF);
+		FParticipatingMedia PM = StrataSlabCreateParticipatingMedia(DiffuseColor, SLAB_SSSMFP(BSDF));
+		const float3 Transmittance = IsotropicMediumSlabTransmittance(PM, ThicknessRatio * STRATA_SIMPLEVOLUME_THICKNESS_M, 1.f);
+
+		OutColor.rgb *= Transmittance * BSDF.Coverage + (1.0f - BSDF.Coverage);
+	}
+	else
+	{
+		OutColor.rgb *= 0.0f;
+	}
+
+	// Patch the BSDF to appear as a top layer
+	BSDF.LuminanceWeightV = 1.0f;
+	BSDF.LuminanceWeightL = 1.0f;
+	BSDF.bIsBottom = true;
+	BSDF.bIsTop = true;
+	BSDF.Coverage = 1.0f;
+
+	// Create a head that will map to a single BSDF
+	FStrataPixelHeader StrataPixelHeader = InitialiseStrataPixelHeader();
+	StrataPixelHeader.BSDFCount = 1;
+
+	FStrataIntegrationSettings Settings = InitStrataIntegrationSettings();
+	FStrataAddressing NullStrataAddressing = (FStrataAddressing)0;
+
+	// Apply emissive
+	OutColor.rgb += BSDF_GETEMISSIVE(BSDF) * View.PreExposure;
+
+	// Apply a white directional light
+	if (ForwardLightData.HasDirectionalLight)
+	{
+		FDeferredLightData DirLightData = (FDeferredLightData)0;
+		DirLightData.Color = ForwardLightData.DirectionalLightColor;
+		DirLightData.FalloffExponent = 0;
+		DirLightData.Direction = ForwardLightData.DirectionalLightDirection;
+		DirLightData.DistanceFadeMAD = ForwardLightData.DirectionalLightDistanceFadeMAD;
+		DirLightData.bRadialLight = false;
+		DirLightData.SpecularScale = 1;
+		float4 DirLightAttenuation = float4(1, 1, 1, 1);
+		DirLightData.ShadowedBits = 1;
+		DirLightData.ShadowMapChannelMask.x = 1;
+		DirLightData.HairTransmittance = InitHairTransmittanceData();
+
+		float Dither = 0.0f;
+		float3 ToLight = DirLightData.Direction;
+		float LightMask = 1.0f;
+		FRectTexture RectTexture = InitRectTexture();
+		uint LightChannelMask = 0xFFFFFFFF;
+		uint PrimitiveLightingChannelMask = LightChannelMask;
+
+		float3 BSDFColoredVisibility = 1.0f;
+
+		FStrataBSDFContext StrataBSDFContext = StrataCreateBSDFContext(TangentBasis, BSDF, V, ToLight);
+
+		bool bEvaluateHasBeenExecuted = false;
+		FStrataEvaluateResult BSDFEvaluate = (FStrataEvaluateResult)0;
+		float3 DirLightLuminance = StrataForwardLightingCommon(
+			Dither,
+			Settings,
+			DirLightData,
+			ToLight,
+			LightMask,
+			DirLightAttenuation,
+			RectTexture,
+			LightChannelMask,
+			PrimitiveLightingChannelMask,
+			TranslatedWorldPosition,
+			SceneDepth,
+			BSDFColoredVisibility,
+			StrataPixelHeader,
+			StrataBSDFContext,
+			bEvaluateHasBeenExecuted,
+			BSDFEvaluate);
+
+		OutColor.rgb += DirLightLuminance * View.PreExposure;
+	}
+
+	// Apply the sky box
+	{
+		float3 ToLight = V;
+		FStrataBSDFContext StrataBSDFContext = StrataCreateBSDFContext(TangentBasis, BSDF, V, ToLight);
+
+		const bool bEnableSpecular = ReflectionStruct.SkyLightParameters.y > 0.0f;
+		FStrataEnvLightResult StrataEnvLight = StrataEvaluateForEnvLight(StrataBSDFContext, bEnableSpecular, Settings);
+
+		float SSRReductionFactor = 0.0;
+		float3 DiffuseLighting = 0.0;
+		float3 SpecularLighting = 0.0;
+		StrataEnvLightingCommon(
+			StrataEnvLight,
+			StrataPixelHeader,
+			StrataBSDFContext,
+			BSDF,
+			N,	//in float3 BentNormal,
+			1.0,//in float3 BSDFThroughput,
+			0,	//in uint CaptureDataStartIndex,
+			0,	//in uint NumCulledReflectionCaptures,
+			1.0,//in float ScreenAmbientOcclusion,
+			1.0,//in float CloudVolumetricAOShadow,
+			1.0,//in float TopLayerSpecularContributionFactor,
+			TranslatedWorldPosition,
+			1.0,//in float CombinedScreenAndMaterialAO,
+			SSRReductionFactor,
+			DiffuseLighting,
+			SpecularLighting);
+
+		OutColor.rgb += (DiffuseLighting + SpecularLighting) * View.PreExposure;;
+	}
+}
+
+void GetMaterialOnCube(
 	in float4 SVPos,
 	in float2 Pos,
 	in float  Size,
@@ -542,9 +682,51 @@ void GetMaterialOnSphere(
 		const float2 UV = (PixelPos - Pos) / Size;
 		OutColor = GetBackgroundCheckboardColor(UV);
 
-		float3 AbsoluteWorldPosition = LWCHackToFloat(PrimaryView.WorldCameraOrigin);
-		float3 TranslatedWorldPosition = PrimaryView.TranslatedWorldCameraOrigin;
-		float SceneDepth = 0.0;
+		const float3 RayO = -View.ViewForward;
+		const float3 RayD = View.ViewForward + (UV.x - 0.5) * View.ViewRight -(UV.y - 0.5) * View.ViewUp;
+		float3 V = -RayD;
+
+		float2 Hit;
+		const float CubeSize = 0.25;
+		bool bIntersection = slabs(-CubeSize, CubeSize, RayO, 1/RayD, Hit.x, Hit.y);
+		if (bIntersection)
+		{
+			float3 P = RayO + RayD * Hit.x;
+			float3 N = normalize(P);
+			float3 AbsN = abs(N);
+			N = AbsN.x > AbsN.y && AbsN.x > AbsN.z ? float3(1 * sign(N.x), 0, 0) : N;
+			N = AbsN.y > AbsN.z && AbsN.y > AbsN.x ? float3(0, 1 * sign(N.y), 0) : N;
+			N = AbsN.z > AbsN.x && AbsN.z > AbsN.y ? float3(0, 0, 1 * sign(N.z)) : N;
+			const float ThicknessRatio = abs(Hit.y - Hit.x) / CubeSize;
+
+			float3x3 TangentBasis;
+			TangentBasis[2] = N;
+			TangentBasis[0] = AbsN.x > AbsN.y && AbsN.x > AbsN.z ? float3(0, 1, 0) : (AbsN.y > AbsN.z && AbsN.y > AbsN.x ? float3(0, 0, 1) : float3(1, 0, 0));
+			TangentBasis[1] = AbsN.x > AbsN.y && AbsN.x > AbsN.z ? float3(0, 0, 1) : (AbsN.y > AbsN.z && AbsN.y > AbsN.x ? float3(1, 0, 0) : float3(0, 1, 0));
+
+			EvaluateLighting(
+				BSDF,
+				V,
+				TangentBasis,
+				ThicknessRatio,
+				OutColor);
+		}
+	}
+}
+
+void GetMaterialOnSphere(
+	in float4 SVPos,
+	in float2 Pos,
+	in float  Size,
+	in FStrataBSDF BSDF,
+	inout float4 OutColor)
+{
+	const float2 PixelPos = SVPos.xy;
+
+	if (BSDF_GETTYPE(BSDF) == STRATA_BSDF_TYPE_SLAB && all(PixelPos > Pos) && all(PixelPos < (Pos + Size)))
+	{
+		const float2 UV = (PixelPos - Pos) / Size;
+		OutColor = GetBackgroundCheckboardColor(UV);
 
 		float3 V = -View.ViewForward;// float3(0, 0, 1);
 		const float3 RayO = V + (UV.x-0.5) * View.ViewRight + -(UV.y-0.5) * View.ViewUp;
@@ -557,124 +739,18 @@ void GetMaterialOnSphere(
 			float3 N = normalize(P - Sphere.xyz);
 			const float ThicknessRatio = (Hit.y - Hit.x) / Sphere.w;
 
-			if (BSDF_GETISSIMPLEVOLUME(BSDF))
-			{
-				// Apply transmittance on the background
-
-				float3 DiffuseColor = SLAB_DIFFUSEALBEDO(BSDF);
-				FParticipatingMedia PM = StrataSlabCreateParticipatingMedia(DiffuseColor, SLAB_SSSMFP(BSDF));
-				const float3 Transmittance = IsotropicMediumSlabTransmittance(PM, ThicknessRatio * STRATA_SIMPLEVOLUME_THICKNESS_M, 1.f);
-
-				OutColor.rgb *= Transmittance * BSDF.Coverage + (1.0f - BSDF.Coverage);
-			}
-			else
-			{
-				OutColor.rgb *= 0.0f;
-			}
-
-			// Patch the BSDF to appear as a top layer
-			BSDF.LuminanceWeightV = 1.0f;
-			BSDF.LuminanceWeightL = 1.0f;
-			BSDF.bIsBottom = true;
-			BSDF.bIsTop = true;
-			BSDF.Coverage = 1.0f;
-
-			// Create a head that will map to a single BSDF
-			FStrataPixelHeader StrataPixelHeader = InitialiseStrataPixelHeader();
-			StrataPixelHeader.BSDFCount = 1;
-
-			FStrataIntegrationSettings Settings = InitStrataIntegrationSettings();
-			FStrataAddressing NullStrataAddressing = (FStrataAddressing)0;
-
 			float3x3 TangentBasis;
 			TangentBasis[2] = N;
-			TangentBasis[0] = normalize(N*float3(1, 0, 0));
+			TangentBasis[0] = normalize(N * float3(1, 0, 0));
 			TangentBasis[1] = cross(N, TangentBasis[0]);
-			TangentBasis[0] =-cross(N, TangentBasis[1]);
-
-			// Apply emissive
-			OutColor.rgb += BSDF_GETEMISSIVE(BSDF) * View.PreExposure;
-
-			// Apply a white directional light
-			if(ForwardLightData.HasDirectionalLight)
-			{
-				FDeferredLightData DirLightData = (FDeferredLightData)0;
-				DirLightData.Color = ForwardLightData.DirectionalLightColor;
-				DirLightData.FalloffExponent = 0;
-				DirLightData.Direction = ForwardLightData.DirectionalLightDirection;
-				DirLightData.DistanceFadeMAD = ForwardLightData.DirectionalLightDistanceFadeMAD;
-				DirLightData.bRadialLight = false;
-				DirLightData.SpecularScale = 1;
-				float4 DirLightAttenuation = float4(1, 1, 1, 1);
-				DirLightData.ShadowedBits = 1;
-				DirLightData.ShadowMapChannelMask.x = 1;
-				DirLightData.HairTransmittance = InitHairTransmittanceData();
-
-				float Dither = 0.0f;
-				float3 ToLight = DirLightData.Direction;
-				float LightMask = 1.0f;
-				FRectTexture RectTexture = InitRectTexture();
-				uint LightChannelMask = 0xFFFFFFFF;
-				uint PrimitiveLightingChannelMask = LightChannelMask;
-
-				float3 BSDFColoredVisibility = 1.0f;
-
-				FStrataBSDFContext StrataBSDFContext = StrataCreateBSDFContext(TangentBasis, BSDF, V, ToLight);
-
-				bool bEvaluateHasBeenExecuted = false;
-				FStrataEvaluateResult BSDFEvaluate = (FStrataEvaluateResult)0;
-				float3 DirLightLuminance = StrataForwardLightingCommon(
-					Dither,
-					Settings,
-					DirLightData,
-					ToLight,
-					LightMask,
-					DirLightAttenuation,
-					RectTexture,
-					LightChannelMask,
-					PrimitiveLightingChannelMask,
-					TranslatedWorldPosition,
-					SceneDepth,
-					BSDFColoredVisibility,
-					StrataPixelHeader,
-					StrataBSDFContext,
-					bEvaluateHasBeenExecuted,
-					BSDFEvaluate);
-
-				OutColor.rgb += DirLightLuminance * View.PreExposure;
-			}
-
-			// Apply the sky box
-			{
-				float3 ToLight = V;
-				FStrataBSDFContext StrataBSDFContext = StrataCreateBSDFContext(TangentBasis, BSDF, V, ToLight);
-
-				const bool bEnableSpecular = ReflectionStruct.SkyLightParameters.y > 0.0f;
-				FStrataEnvLightResult StrataEnvLight = StrataEvaluateForEnvLight(StrataBSDFContext, bEnableSpecular, Settings);
-
-				float SSRReductionFactor = 0.0;
-				float3 DiffuseLighting = 0.0;
-				float3 SpecularLighting = 0.0;
-				StrataEnvLightingCommon(
-					StrataEnvLight,
-					StrataPixelHeader,
-					StrataBSDFContext,
-					BSDF,
-					N,	//in float3 BentNormal,
-					1.0,//in float3 BSDFThroughput,
-					0,	//in uint CaptureDataStartIndex,
-					0,	//in uint NumCulledReflectionCaptures,
-					1.0,//in float ScreenAmbientOcclusion,
-					1.0,//in float CloudVolumetricAOShadow,
-					1.0,//in float TopLayerSpecularContributionFactor,
-					TranslatedWorldPosition,
-					1.0,//in float CombinedScreenAndMaterialAO,
-					SSRReductionFactor,
-					DiffuseLighting,
-					SpecularLighting);
-
-				OutColor.rgb += (DiffuseLighting + SpecularLighting) * View.PreExposure;;
-			}
+			TangentBasis[0] = -cross(N, TangentBasis[1]);
+
+			EvaluateLighting(
+				BSDF,
+				V,
+				TangentBasis,
+				ThicknessRatio,
+				OutColor);
 		}
 	}
 }
@@ -702,7 +778,7 @@ void MaterialDebugStrataTreePS(
 		const float2 PixelPos = SVPos.xy;
 
 		float2 Pos = float2(100.0f, 100.0f);
-		const float SphereSize = 200.0f;
+		const float MaterialPanelSize = 200.0f;
 
 		// STRATA_TODO: read BSDF and Operators "on the fly".
 #if COMPILER_DXC == 1
@@ -724,11 +800,18 @@ void MaterialDebugStrataTreePS(
 				GetMaterialOnSphere(
 					SVPos,
 					Pos,
-					SphereSize,
+					MaterialPanelSize,
+					BSDF,
+					OutColor);
+
+				GetMaterialOnCube(
+					SVPos,
+					Pos + float2(MaterialPanelSize + 10.0f, 0.0f),
+					MaterialPanelSize,
 					BSDF,
 					OutColor);
 
-				Pos.y += SphereSize + 10.0f;
+				Pos.y += MaterialPanelSize + 10.0f;
 			}
 		}
 	}
