Reorganize RT shader code.

This changelist does not change code, but only move code around to ease coming RT support for Strata.

diff --git a/Engine/Shaders/Private/RayTracing/RayTracingCommon.ush b/Engine/Shaders/Private/RayTracing/RayTracingCommon.ush
--- a/Engine/Shaders/Private/RayTracing/RayTracingCommon.ush
+++ b/Engine/Shaders/Private/RayTracing/RayTracingCommon.ush
@@ -102,6 +102,44 @@
 #define PLATFORM_RAYTRACING_INSTANCE_FLAG_FORCE_OPAQUE						0x4
 #define PLATFORM_RAYTRACING_INSTANCE_FLAG_FORCE_NON_OPAQUE					0x8
 
+// must match ERayTracingInstanceFlags
+#define RAYTRACING_INSTANCE_FLAG_NONE					0
+#define RAYTRACING_INSTANCE_FLAG_TRIANGLE_CULL_DISABLE	(1 << 1)
+#define RAYTRACING_INSTANCE_FLAG_TRIANGLE_CULL_REVERSE	(1 << 2)
+#define RAYTRACING_INSTANCE_FLAG_FORCE_OPAQUE			(1 << 3)
+#define RAYTRACING_INSTANCE_FLAG_FORCE_NON_OPAQUE		(1 << 4)
+
+#define RAY_TRACING_BLEND_MODE_OPAQUE				0
+#define RAY_TRACING_BLEND_MODE_ALPHA_COMPOSITE		1
+#define RAY_TRACING_BLEND_MODE_TRANSLUCENT			2
+#define RAY_TRACING_BLEND_MODE_ADDITIVE				3
+#define RAY_TRACING_BLEND_MODE_MODULATE				4
+#define RAY_TRACING_BLEND_MODE_ALPHA_HOLDOUT		5
+
+// Check how Flags are packed in the FPackedMaterialClosestHitPayload before adding to ensure bits are free
+// There are two sets of flags:
+//   - "inputs" are set before calling TraceRay
+//   - "outputs" are set during execution of the CHS and are valid after TraceRay returns
+// These two sets of masks can overlap because they never co-exist
+#define RAY_TRACING_PAYLOAD_OUTPUT_FLAG_NONE                               0
+#define RAY_TRACING_PAYLOAD_OUTPUT_FLAG_FRONT_FACE                   (1 << 0) // Indicates that ray has hit front face of a primitive. Set by closest hit shader.
+#define RAY_TRACING_PAYLOAD_OUTPUT_FLAG_TWO_SIDED                    (1 << 1) // Ray hit a two-sided material
+#define RAY_TRACING_PAYLOAD_OUTPUT_FLAG_UNUSED2                      (1 << 2) // Free for future use
+#define RAY_TRACING_PAYLOAD_OUTPUT_FLAG_UNUSED3                      (1 << 3) // Free for future use
+#define RAY_TRACING_PAYLOAD_OUTPUT_FLAG_UNUSED4                      (1 << 4) // Free for future use
+#define RAY_TRACING_PAYLOAD_OUTPUT_FLAG_UNUSED5                      (1 << 5) // Free for future use
+
+#define RAY_TRACING_PAYLOAD_INPUT_FLAG_NONE								   0	
+#define RAY_TRACING_PAYLOAD_INPUT_FLAG_MINIMAL_PAYLOAD               (1 << 0) // Indicates that closest hit shader should only fill FMinimalPayload (HitT), skipping full material evaluation. Set by RayGen shader, before TraceRay().
+#define RAY_TRACING_PAYLOAD_INPUT_FLAG_ENABLE_SKY_LIGHT_CONTRIBUTION (1 << 1) // Indicates that hit shaders should add in the Sky Light contribution in their output. Set by RayGen shader, before TraceRay().
+#define RAY_TRACING_PAYLOAD_INPUT_FLAG_LUMEN_PAYLOAD		         (1 << 2) // Indicates that closest hit shader should evaluate geometry normal and instanceid.
+#define RAY_TRACING_PAYLOAD_INPUT_FLAG_IGNORE_TRANSLUCENT            (1 << 3) // Indicates that translucent materials should be ignored by IgnoreHit() in anyhit shaders.
+#define RAY_TRACING_PAYLOAD_INPUT_FLAG_UNUSED4                       (1 << 4) // Free for future use
+#define RAY_TRACING_PAYLOAD_INPUT_FLAG_UNUSED5                       (1 << 5) // Free for future use
+
+///////////////////////////////////////////////////////////////////////////////////////////////////
+// Instance descriptor
+
 // Must match D3D12_RAYTRACING_INSTANCE_DESC / VkAccelerationStructureInstanceKHR
 struct FRayTracingInstanceDescriptor
 {
@@ -129,6 +167,9 @@
 
 #endif // OVERRIDE_PLATFORM_RAYTRACING_INSTANCE_DESCRITOR
 
+///////////////////////////////////////////////////////////////////////////////////////////////////
+// Attributes and small payload
+
 // Indices and vertex positions of the ray tracing geometry triangle.
 // Used for loading and interpolation of per-vertex attributes and computing per-triangle normals.
 struct FTriangleBaseAttributes
@@ -217,6 +258,9 @@
 
 };
 
+///////////////////////////////////////////////////////////////////////////////////////////////////
+// Ray cone
+
 struct FRayCone
 {
 	float Width;
@@ -247,6 +291,9 @@
 // Include utility fonction for world position reconstruction 
 #include "../PositionReconstructionCommon.ush" 
 
+///////////////////////////////////////////////////////////////////////////////////////////////////
+// Helpers
+
 uint CalcLinearIndex(uint2 PixelCoord)
 {
 	return PixelCoord.y * uint(View.BufferSizeAndInvSize.x) + PixelCoord.x;
@@ -262,12 +309,69 @@
 	return DispatchThreadId * UpscaleFactor + uint2(SubPixelId & (UpscaleFactor - 1), SubPixelId / UpscaleFactor);
 }
 
-// must match ERayTracingInstanceFlags
-#define RAYTRACING_INSTANCE_FLAG_NONE					0
-#define RAYTRACING_INSTANCE_FLAG_TRIANGLE_CULL_DISABLE	(1 << 1)
-#define RAYTRACING_INSTANCE_FLAG_TRIANGLE_CULL_REVERSE	(1 << 2)
-#define RAYTRACING_INSTANCE_FLAG_FORCE_OPAQUE			(1 << 3)
-#define RAYTRACING_INSTANCE_FLAG_FORCE_NON_OPAQUE		(1 << 4)
+// WorldNormal is the vector towards which the ray position will be offseted.
+void ApplyPositionBias(inout float3 RayOrigin, float3 RayDirection, const float3 WorldNormal, const float MaxNormalBias)
+{
+	// Apply normal perturbation when defining ray to:
+	// * avoid self intersection with current underlying triangle
+	// * hide mismatch between shading surface & geometric surface
+	//
+	// While using shading normal is not correct (we should use the 
+	// geometry normal, but it is not available atm/too costly to compute), 
+	// it is good enough for a cheap solution
+	const float MinBias = 0.01f;
+	const float MaxBias = max(MinBias, MaxNormalBias);
+	const float NormalBias = lerp(MaxBias, MinBias, saturate(dot(WorldNormal, RayDirection)));
 
+	RayOrigin += WorldNormal * NormalBias;
+}
+
+void ApplyCameraRelativeDepthBias(inout float3 RayOrigin, float3 RayDirection, uint2 PixelCoord, float DeviceZ, const float3 WorldNormal, const float AbsoluteNormalBias)
+{
+	float3 TranslatedWorldPosition = ReconstructTranslatedWorldPositionFromDeviceZ(PixelCoord, DeviceZ);
+	float3 CameraDirection = TranslatedWorldPosition - View.TranslatedWorldCameraOrigin;
+	float DistanceToCamera = length(CameraDirection);
+	CameraDirection = normalize(CameraDirection);
+	float Epsilon = 1.0e-4;
+	float RelativeBias = DistanceToCamera * Epsilon;
+	//float ProjectedBias = RelativeBias / dot(RayDirection, WorldNormal);
+
+	float RayBias = max(RelativeBias, AbsoluteNormalBias);
+	RayOrigin -= CameraDirection * RayBias;
+	ApplyPositionBias(RayOrigin, RayDirection, WorldNormal, RayBias);
+}
+
+void ApplyPositionBias(inout RayDesc Ray, const float3 WorldNormal, const float MaxNormalBias)
+{
+	ApplyPositionBias(Ray.Origin, Ray.Direction, WorldNormal, MaxNormalBias);
+}
+
+void ApplyCameraRelativeDepthBias(inout RayDesc Ray, uint2 PixelCoord, float DeviceZ, const float3 WorldNormal, const float AbsoluteNormalBias)
+{
+	ApplyCameraRelativeDepthBias(Ray.Origin, Ray.Direction, PixelCoord, DeviceZ, WorldNormal, AbsoluteNormalBias);
+}
+
+#if !COMPUTESHADER || PLATFORM_SUPPORTS_INLINE_RAY_TRACING
+
+RayDesc CreatePrimaryRay(float2 UV)
+{
+	float2 ScreenPosition = (UV - View.ScreenPositionScaleBias.wz) / View.ScreenPositionScaleBias.xy;
+
+	RayDesc Ray;
+	float4 NearPoint = mul(float4(ScreenPosition, 1, 1), View.ClipToTranslatedWorld);
+	float4 FarPoint = mul(float4(ScreenPosition, 0, 1), View.ClipToTranslatedWorld);
+	Ray.Origin = NearPoint.xyz * rcp(NearPoint.w);
+	Ray.Direction = normalize(NearPoint.w * FarPoint.xyz - FarPoint.w * NearPoint.xyz);
+
+	// extract far plane (ortho only, so ignore distinction between T and Z)
+	float4 Back = mul(float4(0, 0, 0, 1), View.ClipToView);
+	Ray.TMin = 0;
+	Ray.TMax = Back.w > 0 ? (Back.z / Back.w - View.NearPlane) : POSITIVE_INFINITY;
+
+	return Ray;
+}
+
+#endif // !COMPUTESHADER || PLATFORM_SUPPORTS_INLINE_RAY_TRACING
+
 uint TranslateRayTracingInstanceFlags(uint InFlags)
 {
@@ -297,33 +401,10 @@
 	return Result;
 }
 
-#define RAY_TRACING_BLEND_MODE_OPAQUE				0
-#define RAY_TRACING_BLEND_MODE_ALPHA_COMPOSITE		1
-#define RAY_TRACING_BLEND_MODE_TRANSLUCENT			2
-#define RAY_TRACING_BLEND_MODE_ADDITIVE				3
-#define RAY_TRACING_BLEND_MODE_MODULATE				4
-#define RAY_TRACING_BLEND_MODE_ALPHA_HOLDOUT		5
+///////////////////////////////////////////////////////////////////////////////////////////////////
+// Material payload
 
-// Check how Flags are packed in the FPackedMaterialClosestHitPayload before adding to ensure bits are free
-// There are two sets of flags:
-//   - "inputs" are set before calling TraceRay
-//   - "outputs" are set during execution of the CHS and are valid after TraceRay returns
-// These two sets of masks can overlap because they never co-exist
-#define RAY_TRACING_PAYLOAD_OUTPUT_FLAG_NONE                               0
-#define RAY_TRACING_PAYLOAD_OUTPUT_FLAG_FRONT_FACE                   (1 << 0) // Indicates that ray has hit front face of a primitive. Set by closest hit shader.
-#define RAY_TRACING_PAYLOAD_OUTPUT_FLAG_TWO_SIDED                    (1 << 1) // Ray hit a two-sided material
-#define RAY_TRACING_PAYLOAD_OUTPUT_FLAG_UNUSED2                      (1 << 2) // Free for future use
-#define RAY_TRACING_PAYLOAD_OUTPUT_FLAG_UNUSED3                      (1 << 3) // Free for future use
-#define RAY_TRACING_PAYLOAD_OUTPUT_FLAG_UNUSED4                      (1 << 4) // Free for future use
-#define RAY_TRACING_PAYLOAD_OUTPUT_FLAG_UNUSED5                      (1 << 5) // Free for future use
+#if !COMPUTESHADER
 
-#define RAY_TRACING_PAYLOAD_INPUT_FLAG_NONE								   0	
-#define RAY_TRACING_PAYLOAD_INPUT_FLAG_MINIMAL_PAYLOAD               (1 << 0) // Indicates that closest hit shader should only fill FMinimalPayload (HitT), skipping full material evaluation. Set by RayGen shader, before TraceRay().
-#define RAY_TRACING_PAYLOAD_INPUT_FLAG_ENABLE_SKY_LIGHT_CONTRIBUTION (1 << 1) // Indicates that hit shaders should add in the Sky Light contribution in their output. Set by RayGen shader, before TraceRay().
-#define RAY_TRACING_PAYLOAD_INPUT_FLAG_LUMEN_PAYLOAD		         (1 << 2) // Indicates that closest hit shader should evaluate geometry normal and instanceid.
-#define RAY_TRACING_PAYLOAD_INPUT_FLAG_IGNORE_TRANSLUCENT            (1 << 3) // Indicates that translucent materials should be ignored by IgnoreHit() in anyhit shaders.
-#define RAY_TRACING_PAYLOAD_INPUT_FLAG_UNUSED4                       (1 << 4) // Free for future use
-#define RAY_TRACING_PAYLOAD_INPUT_FLAG_UNUSED5                       (1 << 5) // Free for future use
-
 struct FMaterialClosestHitPayload : FMinimalPayload
 {
@@ -369,62 +450,6 @@
 	bool HasAnisotropy() { return abs(Anisotropy) >= 0.001f; }
 };
 
-// WorldNormal is the vector towards which the ray position will be offseted.
-void ApplyPositionBias(inout float3 RayOrigin, float3 RayDirection, const float3 WorldNormal, const float MaxNormalBias)
-{
-	// Apply normal perturbation when defining ray to:
-	// * avoid self intersection with current underlying triangle
-	// * hide mismatch between shading surface & geometric surface
-	//
-	// While using shading normal is not correct (we should use the 
-	// geometry normal, but it is not available atm/too costly to compute), 
-	// it is good enough for a cheap solution
-	const float MinBias = 0.01f;
-	const float MaxBias = max(MinBias, MaxNormalBias);
-	const float NormalBias = lerp(MaxBias, MinBias, saturate(dot(WorldNormal, RayDirection)));
-
-	RayOrigin += WorldNormal * NormalBias;
-}
-
-void ApplyCameraRelativeDepthBias(inout float3 RayOrigin, float3 RayDirection, uint2 PixelCoord, float DeviceZ, const float3 WorldNormal, const float AbsoluteNormalBias)
-{
-	float3 TranslatedWorldPosition = ReconstructTranslatedWorldPositionFromDeviceZ(PixelCoord, DeviceZ);
-	float3 CameraDirection = TranslatedWorldPosition - View.TranslatedWorldCameraOrigin;
-	float DistanceToCamera = length(CameraDirection);
-	CameraDirection = normalize(CameraDirection);
-	float Epsilon = 1.0e-4;
-	float RelativeBias = DistanceToCamera * Epsilon;
-	//float ProjectedBias = RelativeBias / dot(RayDirection, WorldNormal);
-
-	float RayBias = max(RelativeBias, AbsoluteNormalBias);
-	RayOrigin -= CameraDirection * RayBias;
-	ApplyPositionBias(RayOrigin, RayDirection, WorldNormal, RayBias);
-}
-
-#if !COMPUTESHADER || PLATFORM_SUPPORTS_INLINE_RAY_TRACING
-
-RayDesc CreatePrimaryRay(float2 UV)
-{
-	float2 ScreenPosition = (UV - View.ScreenPositionScaleBias.wz) / View.ScreenPositionScaleBias.xy;
-
-	RayDesc Ray;
-	float4 NearPoint = mul(float4(ScreenPosition, 1, 1), View.ClipToTranslatedWorld);
-	float4 FarPoint  = mul(float4(ScreenPosition, 0, 1), View.ClipToTranslatedWorld);
-	Ray.Origin = NearPoint.xyz * rcp(NearPoint.w);
-	Ray.Direction = normalize(NearPoint.w * FarPoint.xyz - FarPoint.w * NearPoint.xyz);
-
-	// extract far plane (ortho only, so ignore distinction between T and Z)
-	float4 Back = mul(float4(0, 0, 0, 1), View.ClipToView);
-	Ray.TMin = 0;
-	Ray.TMax = Back.w > 0 ? (Back.z / Back.w - View.NearPlane) : POSITIVE_INFINITY;
-
-	return Ray;
-}
-
-#endif // !COMPUTESHADER || PLATFORM_SUPPORTS_INLINE_RAY_TRACING
-
-#if !COMPUTESHADER
-
 struct FPackedMaterialClosestHitPayload : FMinimalPayload
 {
 	// float FMinimalPayload::HitT                       // 4  bytes
@@ -630,15 +655,12 @@
 	return Output;
 }
 
-void ApplyPositionBias(inout RayDesc Ray, const float3 WorldNormal, const float MaxNormalBias)
-{
-	ApplyPositionBias(Ray.Origin, Ray.Direction, WorldNormal, MaxNormalBias);
-}
+#endif // !COMPUTESHADER
 
-void ApplyCameraRelativeDepthBias(inout RayDesc Ray, uint2 PixelCoord, float DeviceZ, const float3 WorldNormal, const float AbsoluteNormalBias)
-{
-	ApplyCameraRelativeDepthBias(Ray.Origin, Ray.Direction, PixelCoord, DeviceZ, WorldNormal, AbsoluteNormalBias);
-}
+///////////////////////////////////////////////////////////////////////////////////////////////////
+// Tracing util functions
 
+#if !COMPUTESHADER
+
 void TraceVisibilityRayPacked(
 	inout FPackedMaterialClosestHitPayload PackedPayload,
diff --git a/Engine/Shaders/Private/RayTracing/RayTracingMaterialHitShaders.usf b/Engine/Shaders/Private/RayTracing/RayTracingMaterialHitShaders.usf
--- a/Engine/Shaders/Private/RayTracing/RayTracingMaterialHitShaders.usf
+++ b/Engine/Shaders/Private/RayTracing/RayTracingMaterialHitShaders.usf
@@ -231,103 +231,13 @@
 	return Normal;
 }
 
-RAY_TRACING_ENTRY_CLOSEST_HIT(MaterialCHS,
-	FPackedMaterialClosestHitPayload, PackedPayload,
-	FRayTracingIntersectionAttributes, Attributes)
+void GetMaterialPayload(
+	FPixelMaterialInputs PixelMaterialInputs,
+	FMaterialPixelParameters MaterialParameters,
+	FVertexFactoryInterpolantsVSToPS Interpolants,
+	bool bIsEnableSkyLightContribution,
+	inout FMaterialClosestHitPayload Payload)
 {
-	PackedPayload.HitT = RayTCurrent();
-
-#if USE_MATERIAL_CLOSEST_HIT_SHADER
-
-	bool bIsLumenPayload = PackedPayload.IsLumenPayload();
-	if (bIsLumenPayload)
-	{
-		float3 SurfaceNormal = GetLumenGeometryNormal();
-		const int PrimitiveIndex = GetInstanceUserData();
-		const int PrimitiveInstanceIndex = InstanceIndex() - GetBaseInstanceIndex();
-
-		PackedPayload.SetGeometryNormal(SurfaceNormal);
-		PackedPayload.SetSceneInstanceIndex(
-			GetGPUSceneInstanceIndex(PrimitiveIndex, PrimitiveInstanceIndex)
-		);
-	}
-
-	if (PackedPayload.IsMinimalPayloadMode())
-	{
-		// Minimal payload mode only fills FMinimalPayload::HitT, skipping actual material evaluation.
-		// This mode is used when tracing shadow rays against masked geometry.
-		// Dynamic branch is used to avoid compiling an extra shader permutation.
-		return;
-	}
-
-	GlobalTextureMipBias = PackedPayload.GetMipBias();
-
-	const uint2 PixelCoord = PackedPayload.GetPixelCoord();
-
-	ResolvedView = ResolveView();
-
-	FVertexFactoryInterpolantsVSToPS Interpolants;
-	float4 SvPosition = float4(PixelCoord.xy, 0.0, 1.0);
-	
-	FRayCone PropagatedCone = PropagateRayCone(PackedPayload.GetRayCone(), 0 /* surface curvature */, RayTCurrent());
-	CalcInterpolants(PropagatedCone, Attributes, Interpolants);
-
-#if VF_SUPPORTS_RAYTRACING_PREPARE_MATERIAL_PIXEL_PARAMETERS
-	float3 GeoNormal = 0;
-	FMaterialPixelParameters MaterialParameters = GetMaterialPixelParameters(WorldRayOrigin(), WorldRayDirection(), RayTCurrent(), PrimitiveIndex(), Attributes, HitKind(), GeoNormal);
-#else
-	FMaterialPixelParameters MaterialParameters = GetMaterialPixelParameters(Interpolants, SvPosition);
-#endif
-	FPixelMaterialInputs PixelMaterialInputs;
-
-	{
-		float4 ScreenPosition = SvPositionToResolvedScreenPosition(SvPosition);
-		float3 TranslatedWorldPosition = WorldRayOrigin() + WorldRayDirection() * RayTCurrent();
-		
-		bool bIsFrontFace = HitKind() == HIT_KIND_TRIANGLE_FRONT_FACE;
-
-		// #dxr_todo: UE-72130 support world position offset 
-		// #if USE_WORLD_POSITION_EXCLUDING_SHADER_OFFSETS
-		// CalcMaterialParametersEx(MaterialParameters, PixelMaterialInputs, SvPosition, ScreenPosition, bIsFrontFace, TranslatedWorldPosition, BasePassInterpolants.PixelPositionExcludingWPO);
-		// #else
-		CalcMaterialParametersEx(MaterialParameters, PixelMaterialInputs, SvPosition, ScreenPosition, bIsFrontFace, TranslatedWorldPosition, TranslatedWorldPosition);
-		// #endif
-	}
-
-	FMaterialClosestHitPayload Payload = (FMaterialClosestHitPayload)0;
-
-	/**
-	 * Set common material attributes for both full and simplified materials
-	 **/
-	Payload.ShadingModelID = GetMaterialShadingModel(PixelMaterialInputs);
-
-	#if MATERIALBLENDING_ALPHACOMPOSITE
-		uint MaterialBlendingMode = RAY_TRACING_BLEND_MODE_ALPHA_COMPOSITE;
-	#elif MATERIALBLENDING_ALPHAHOLDOUT
-		uint MaterialBlendingMode = RAY_TRACING_BLEND_MODE_ALPHA_HOLDOUT;
-	#elif MATERIALBLENDING_TRANSLUCENT
-		uint MaterialBlendingMode = RAY_TRACING_BLEND_MODE_TRANSLUCENT;
-	#elif MATERIALBLENDING_ADDITIVE
-		uint MaterialBlendingMode = RAY_TRACING_BLEND_MODE_ADDITIVE;
-	#elif MATERIALBLENDING_MODULATE
-		uint MaterialBlendingMode = RAY_TRACING_BLEND_MODE_MODULATE;
-	#else
-		uint MaterialBlendingMode = RAY_TRACING_BLEND_MODE_OPAQUE;
-	#endif
-
-	Payload.BlendingMode = MaterialBlendingMode;
-	Payload.PrimitiveLightingChannelMask = GetPrimitive_LightingChannelMask(MaterialParameters.PrimitiveId);
-
-	Payload.HitT = RayTCurrent();
-	if (HitKind() == HIT_KIND_TRIANGLE_FRONT_FACE)
-	{
-		Payload.SetFrontFace();
-	}
-
-#if TWO_SIDED_MATERIAL
-	Payload.SetTwoSided();
-#endif
-
 	// Store the results in local variables and reuse instead of calling the functions multiple times.
 	half3 BaseColor = GetMaterialBaseColor(PixelMaterialInputs);
 #if MATERIALBLENDING_TRANSLUCENT && REFRACTION_USE_INDEX_OF_REFRACTION
@@ -447,7 +357,7 @@
 		Payload.WorldNormal,
 		VolumetricLightmapBrickTextureUVs,
 		Payload.ShadingModelID,
-		PackedPayload.IsEnableSkyLightContribution(),
+		bIsEnableSkyLightContribution,
 		DiffuseIndirectLighting,
 		SubsurfaceIndirectLighting,
 		IndirectIrradiance);
@@ -460,7 +370,7 @@
 	
 	Payload.IndirectIrradiance += DiffuseIndirectLighting; // Contains both lightmap & skylight
 
-#else
+#else // SIMPLIFIED_MATERIAL_SHADER
 
 	/**
 	 * Set material attributes for simplified materials
@@ -499,21 +409,119 @@
 	Payload.Opacity = GetMaterialOpacity(PixelMaterialInputs);
 	Payload.Ior = 1;
 
-#endif
+#endif // SIMPLIFIED_MATERIAL_SHADER
+}
 
-	PackedPayload = PackRayTracingPayload(Payload, PropagatedCone);
+RAY_TRACING_ENTRY_CLOSEST_HIT(MaterialCHS,
+	FPackedMaterialClosestHitPayload, PackedPayload,
+	FRayTracingIntersectionAttributes, Attributes)
+{
+	PackedPayload.HitT = RayTCurrent();
 
-	// Override packed payload with custom Lumen data
+#if USE_MATERIAL_CLOSEST_HIT_SHADER
+
+	bool bIsLumenPayload = PackedPayload.IsLumenPayload();
 	if (bIsLumenPayload)
 	{
+		float3 SurfaceNormal = GetLumenGeometryNormal();
 		const int PrimitiveIndex = GetInstanceUserData();
 		const int PrimitiveInstanceIndex = InstanceIndex() - GetBaseInstanceIndex();
 
+		PackedPayload.SetGeometryNormal(SurfaceNormal);
 		PackedPayload.SetSceneInstanceIndex(
 			GetGPUSceneInstanceIndex(PrimitiveIndex, PrimitiveInstanceIndex)
 		);
 	}
 
+	if (PackedPayload.IsMinimalPayloadMode())
+	{
+		// Minimal payload mode only fills FMinimalPayload::HitT, skipping actual material evaluation.
+		// This mode is used when tracing shadow rays against masked geometry.
+		// Dynamic branch is used to avoid compiling an extra shader permutation.
+		return;
+	}
+
+	GlobalTextureMipBias = PackedPayload.GetMipBias();
+
+	const uint2 PixelCoord = PackedPayload.GetPixelCoord();
+
+	ResolvedView = ResolveView();
+
+	FVertexFactoryInterpolantsVSToPS Interpolants;
+	float4 SvPosition = float4(PixelCoord.xy, 0.0, 1.0);
+	
+	FRayCone PropagatedCone = PropagateRayCone(PackedPayload.GetRayCone(), 0 /* surface curvature */, RayTCurrent());
+	CalcInterpolants(PropagatedCone, Attributes, Interpolants);
+
+#if VF_SUPPORTS_RAYTRACING_PREPARE_MATERIAL_PIXEL_PARAMETERS
+	float3 GeoNormal = 0;
+	FMaterialPixelParameters MaterialParameters = GetMaterialPixelParameters(WorldRayOrigin(), WorldRayDirection(), RayTCurrent(), PrimitiveIndex(), Attributes, HitKind(), GeoNormal);
+#else
+	FMaterialPixelParameters MaterialParameters = GetMaterialPixelParameters(Interpolants, SvPosition);
+#endif
+	FPixelMaterialInputs PixelMaterialInputs;
+
+	{
+		float4 ScreenPosition = SvPositionToResolvedScreenPosition(SvPosition);
+		float3 TranslatedWorldPosition = WorldRayOrigin() + WorldRayDirection() * RayTCurrent();
+		
+		bool bIsFrontFace = HitKind() == HIT_KIND_TRIANGLE_FRONT_FACE;
+
+		// #dxr_todo: UE-72130 support world position offset 
+		// #if USE_WORLD_POSITION_EXCLUDING_SHADER_OFFSETS
+		// CalcMaterialParametersEx(MaterialParameters, PixelMaterialInputs, SvPosition, ScreenPosition, bIsFrontFace, TranslatedWorldPosition, BasePassInterpolants.PixelPositionExcludingWPO);
+		// #else
+		CalcMaterialParametersEx(MaterialParameters, PixelMaterialInputs, SvPosition, ScreenPosition, bIsFrontFace, TranslatedWorldPosition, TranslatedWorldPosition);
+		// #endif
+	}
+
+	FMaterialClosestHitPayload Payload = (FMaterialClosestHitPayload)0;
+
+	/**
+	 * Set common material attributes for both full and simplified materials
+	 **/
+	Payload.ShadingModelID = GetMaterialShadingModel(PixelMaterialInputs);
+
+	#if MATERIALBLENDING_ALPHACOMPOSITE
+		uint MaterialBlendingMode = RAY_TRACING_BLEND_MODE_ALPHA_COMPOSITE;
+	#elif MATERIALBLENDING_ALPHAHOLDOUT
+		uint MaterialBlendingMode = RAY_TRACING_BLEND_MODE_ALPHA_HOLDOUT;
+	#elif MATERIALBLENDING_TRANSLUCENT
+		uint MaterialBlendingMode = RAY_TRACING_BLEND_MODE_TRANSLUCENT;
+	#elif MATERIALBLENDING_ADDITIVE
+		uint MaterialBlendingMode = RAY_TRACING_BLEND_MODE_ADDITIVE;
+	#elif MATERIALBLENDING_MODULATE
+		uint MaterialBlendingMode = RAY_TRACING_BLEND_MODE_MODULATE;
+	#else
+		uint MaterialBlendingMode = RAY_TRACING_BLEND_MODE_OPAQUE;
+	#endif
+
+	Payload.BlendingMode = MaterialBlendingMode;
+	Payload.PrimitiveLightingChannelMask = GetPrimitive_LightingChannelMask(MaterialParameters.PrimitiveId);
+
+	Payload.HitT = RayTCurrent();
+	if (HitKind() == HIT_KIND_TRIANGLE_FRONT_FACE)
+	{
+		Payload.SetFrontFace();
+	}
+
+#if TWO_SIDED_MATERIAL
+	Payload.SetTwoSided();
+#endif
+
+	// Fill payload with material data
+	GetMaterialPayload(PixelMaterialInputs, MaterialParameters, Interpolants, PackedPayload.IsEnableSkyLightContribution(), Payload);
+
+	PackedPayload = PackRayTracingPayload(Payload, PropagatedCone);
+
+	// Override packed payload with custom Lumen data
+	if (bIsLumenPayload)
+	{
+		const int PrimitiveIndex = GetInstanceUserData();
+		const int PrimitiveInstanceIndex = InstanceIndex() - GetBaseInstanceIndex();
+		PackedPayload.SetSceneInstanceIndex(GetGPUSceneInstanceIndex(PrimitiveIndex, PrimitiveInstanceIndex));
+	}
+
 #endif // USE_MATERIAL_CLOSEST_HIT_SHADER
 }
 
