Fix Strata shader compilation warning.

diff --git a/Engine/Shaders/Private/Strata/Strata.ush b/Engine/Shaders/Private/Strata/Strata.ush
--- a/Engine/Shaders/Private/Strata/Strata.ush
+++ b/Engine/Shaders/Private/Strata/Strata.ush
@@ -391,7 +391,7 @@ struct FStrataBSDF
 #define BSDF_SETEMISSIVE(X, V)			X.Emissive = V
 #define BSDF_GETEMISSIVE(X)				X.Emissive
 #else
-#define BSDF_SETEMISSIVE(X, V)			X
+#define BSDF_SETEMISSIVE(X, V)			
 #define BSDF_GETEMISSIVE(X)				0.0f
 #endif
 
@@ -399,7 +399,7 @@ struct FStrataBSDF
 #define BSDF_SETTHICKNESSCM(X, V)		X.ThicknessCm = V
 #define BSDF_GETTHICKNESSCM(X)			X.ThicknessCm
 #else
-#define BSDF_SETTHICKNESSCM(X, V)		X
+#define BSDF_SETTHICKNESSCM(X, V)		
 #define BSDF_GETTHICKNESSCM(X)			STRATA_SIMPLEVOLUME_THICKNESS_CM
 #endif
 
