Fix another shader typo.

diff --git a/Engine/Shaders/Private/RayTracing/RayTracingCommon.ush b/Engine/Shaders/Private/RayTracing/RayTracingCommon.ush
--- a/Engine/Shaders/Private/RayTracing/RayTracingCommon.ush
+++ b/Engine/Shaders/Private/RayTracing/RayTracingCommon.ush
@@ -342,6 +342,8 @@ void ApplyCameraRelativeDepthBias(inout float3 RayOrigin, float3 RayDirection, u
 	ApplyPositionBias(RayOrigin, RayDirection, WorldNormal, RayBias);
 }
 
+#if !COMPUTESHADER
+
 void ApplyPositionBias(inout RayDesc Ray, const float3 WorldNormal, const float MaxNormalBias)
 {
 	ApplyPositionBias(Ray.Origin, Ray.Direction, WorldNormal, MaxNormalBias);
@@ -352,6 +354,8 @@ void ApplyCameraRelativeDepthBias(inout RayDesc Ray, uint2 PixelCoord, float Dev
 	ApplyCameraRelativeDepthBias(Ray.Origin, Ray.Direction, PixelCoord, DeviceZ, WorldNormal, AbsoluteNormalBias);
 }
 
+#endif // !COMPUTESHADER
+
 #if !COMPUTESHADER || PLATFORM_SUPPORTS_INLINE_RAY_TRACING
 
 RayDesc CreatePrimaryRay(float2 UV)
