Abstract strata top layer reading code.

diff --git a/Engine/Shaders/Private/Strata/Strata.ush b/Engine/Shaders/Private/Strata/Strata.ush
--- a/Engine/Shaders/Private/Strata/Strata.ush
+++ b/Engine/Shaders/Private/Strata/Strata.ush
@@ -36,36 +36,6 @@
 #define STRATA_UNROLL_N(X) UNROLL_N(X)
 #endif
 
-// Strata Rendertarget / UAV setup
-// RT0: Color
-// [insert here potential velocity rendertarget and prec shadow]
-// RT2: Strata uint
-// RT3: Strata uint
-// RT4: Top normal
-//
-// UAV SSS data written out if needed
-// UAV if strata data overrun the MRT output size
-
-struct FRWStrataMaterialContainerStruct
-{
-	uint MaterialRenderTargets[STRATA_BASE_PASS_MRT_OUTPUT_COUNT];
-//	RWTexture2DArray<uint> ExtraMaterialDataUAV; // If this is used, we cannot run shaders on debug...
-};
-#define FRWStrataMaterialContainer	FRWStrataMaterialContainerStruct
-#define FStrataMaterialContainer	Texture2DArray<uint>
-
-FRWStrataMaterialContainer InitialiseRWStrataMaterialContainer(RWTexture2DArray<uint> ExtraMaterialDataUAV)
-{
-	FRWStrataMaterialContainer RWStrataMaterialContainer;
-	UNROLL
-	for (int i = 0; i < STRATA_BASE_PASS_MRT_OUTPUT_COUNT; ++i)
-	{
-		RWStrataMaterialContainer.MaterialRenderTargets[i] = 0;
-	}
-//	RWStrataMaterialContainer.ExtraMaterialDataUAV = ExtraMaterialDataUAV;
-	return RWStrataMaterialContainer;
-}
-
 // During the base pass or forward rendering, shared local bases are simply available from registers.
 // If a BSDF is loaded for processing from the strata buffer, normals are unpacked on demand to reduce VGPR pressure and increase occupancy.
 #ifndef STRATA_INLINE_SHADING
@@ -1210,6 +1180,9 @@ FStrataData GetStrataSingleLayerWaterBSDF(
 	return StrataData;
 }
 
+///////////////////////////////////////////////////////////////////////////////
+// Material buffer Addressing
+
 // Optimised for power of two because it relies on a division done using bit shift
 uint DivideAndRoundUp(uint Dividend, uint Divisor, uint DivisorAsBitShift)
 {
@@ -1219,9 +1192,7 @@ uint DivideAndRoundUp(uint Dividend, uint Divisor, uint DivisorAsBitShift)
 struct FStrataAddressing
 {
 	uint  CurrentIndex;
-
 	uint2 PixelCoords;
-
 	uint  ReadBytes;
 };
 
@@ -1234,8 +1205,42 @@ FStrataAddressing GetStrataPixelDataByteOffset(uint2 PixelPos, uint2 ViewBufferS
 	return StrataAddressing;
 }
 
+///////////////////////////////////////////////////////////////////////////////
+// Material buffer Read/Write
+
+// Strata Rendertarget / UAV setup
+// RT0: Color
+// [insert here potential velocity rendertarget and prec shadow]
+// RT2: Strata uint
+// RT3: Strata uint
+// RT4: Top normal
+//
+// UAV SSS data written out if needed
+// UAV if strata data overrun the MRT output size
+
+struct FRWStrataMaterialContainerStruct
+{
+	uint MaterialRenderTargets[STRATA_BASE_PASS_MRT_OUTPUT_COUNT];
+//	RWTexture2DArray<uint> ExtraMaterialDataUAV; // If this is used, we cannot run shaders on debug...
+};
+
+#define FRWStrataMaterialContainer	FRWStrataMaterialContainerStruct
+FRWStrataMaterialContainer InitialiseRWStrataMaterialContainer(RWTexture2DArray<uint> ExtraMaterialDataUAV)
+{
+	FRWStrataMaterialContainer RWStrataMaterialContainer;
+	UNROLL
+	for (int i = 0; i < STRATA_BASE_PASS_MRT_OUTPUT_COUNT; ++i)
+	{
+		RWStrataMaterialContainer.MaterialRenderTargets[i] = 0;
+	}
+//	RWStrataMaterialContainer.ExtraMaterialDataUAV = ExtraMaterialDataUAV;
+	return RWStrataMaterialContainer;
+}
+
 #if COMPILER_SWITCH // These platforms are complaining so simply disabling Strata there for now.
 
+#define FStrataMaterialContainer Texture2DArray<uint>
+
 void StrataStoreUint1(inout FRWStrataMaterialContainer StrataBuffer, RWTexture2DArray<uint> ExtraMaterialDataUAV, inout FStrataAddressing StrataAddressing, uint Data) {}
 #define STRATA_STORE_UINT1(x) 
 
@@ -1243,7 +1248,9 @@ void FinalizeWrites(inout FRWStrataMaterialContainer StrataBuffer, inout FStrata
 
 uint StrataLoadUint1(FStrataMaterialContainer StrataBuffer, inout FStrataAddressing StrataAddressing) { return 0; }
 
-#else
+#else // COMPILER_SWITCH
+
+#define FStrataMaterialContainer Texture2DArray<uint>
 
 void StrataStoreUint1(inout FRWStrataMaterialContainer StrataBuffer, RWTexture2DArray<uint> ExtraMaterialDataUAV, inout FStrataAddressing StrataAddressing, uint Data)
 {
@@ -1252,9 +1259,9 @@ void StrataStoreUint1(inout FRWStrataMaterialContainer StrataBuffer, RWTexture2D
 	{
 		StrataBuffer.MaterialRenderTargets[0] = StrataAddressing.CurrentIndex == 0 ? Data : StrataBuffer.MaterialRenderTargets[0];
 		StrataBuffer.MaterialRenderTargets[1] = StrataAddressing.CurrentIndex == 1 ? Data : StrataBuffer.MaterialRenderTargets[1];
-#if STRATA_BASE_PASS_MRT_OUTPUT_COUNT != 2
-#error Strata STRATA_BASE_PASS_MRT_OUTPUT_COUNT has been update but not StrataStore function
-#endif
+		#if STRATA_BASE_PASS_MRT_OUTPUT_COUNT != 2
+		#error Strata STRATA_BASE_PASS_MRT_OUTPUT_COUNT has been update but not StrataStore function
+		#endif
 	}
 	else
 	{
@@ -1262,6 +1269,7 @@ void StrataStoreUint1(inout FRWStrataMaterialContainer StrataBuffer, RWTexture2D
 	}
 	StrataAddressing.CurrentIndex++;
 }
+
 #define STRATA_STORE_UINT1(x)	StrataStoreUint1(StrataBuffer, ExtraMaterialDataUAV, StrataAddressing, x)
 
 void FinalizeWrites(inout FRWStrataMaterialContainer StrataBuffer, inout FStrataAddressing StrataAddressing)
@@ -1278,8 +1286,20 @@ uint StrataLoadUint1(FStrataMaterialContainer StrataBuffer, inout FStrataAddress
 	return Data;
 }
 
-#endif
+uint StrataLoadUint1AtIndexOffset(in FStrataMaterialContainer StrataBuffer, in FStrataAddressing StrataAddressing, uint IndexOffset)
+{
+	const uint MipLevel = 0;
+	return StrataBuffer.mips[MipLevel][uint3(StrataAddressing.PixelCoords, IndexOffset)];
+}
 
+#endif // COMPILER_SWITCH
+
+#define FStrataTopLayerDataContainer Texture2D<uint>
+
+uint StrataLoadTopLayerData(FStrataTopLayerDataContainer InTopLayerContainer, inout FStrataAddressing StrataAddressing)
+{
+	return InTopLayerContainer.Load(uint3(StrataAddressing.PixelCoords, 0));
+}
 
 ///////////////////////////////////////////////////////////////////////////////
 // Shared local bases
@@ -1399,7 +1419,7 @@ float3x3 StrataGetBSDFSharedBasis_DeferredShading(in FStrataPixelHeader StrataPi
 		return GetTangentBasis(StrataUnpackTopLayerData(StrataPixelHeader.PackedTopLayerData).WorldNormal);
 	}
 
-	uint PackedBasis = StrataPixelHeader.StrataBuffer[uint3(StrataAddressing.PixelCoords, StrataPixelHeader.SharedLocalBasesIndexOffset + BSDF_GETSHAREDLOCALBASISID(BSDF))];
+	uint PackedBasis = StrataLoadUint1AtIndexOffset(StrataPixelHeader.StrataBuffer, StrataAddressing, StrataPixelHeader.SharedLocalBasesIndexOffset + BSDF_GETSHAREDLOCALBASISID(BSDF));
 
 	uint BasisType   = StrataGetSharedLocalBasisType(StrataPixelHeader.SharedLocalBasesTypes_PackedHeader, BSDF_GETSHAREDLOCALBASISID(BSDF));
 	float3x3 OutTangentBasis;
@@ -3219,7 +3239,7 @@ uint PackStrataHeader(uint InBSDFCount, FStrataPixelHeader InHeader)
 
 // Unpack only strata header
 // If this is changed, please update the compiler side material size evaluation in StrataMaterial.cpp
-FStrataPixelHeader UnpackStrataHeaderIn(FStrataMaterialContainer StrataBuffer, inout FStrataAddressing StrataAddressing, Texture2D<uint> InStrataTopLayerTexture)
+FStrataPixelHeader UnpackStrataHeaderIn(FStrataMaterialContainer StrataBuffer, inout FStrataAddressing StrataAddressing, FStrataTopLayerDataContainer InStrataTopLayerTexture)
 {
 	FStrataPixelHeader Out = InitialiseStrataPixelHeader();
 
@@ -3253,7 +3273,7 @@ FStrataPixelHeader UnpackStrataHeaderIn(FStrataMaterialContainer StrataBuffer, i
 		Out.SharedLocalBasesTypes_PackedHeader = PackedHeader;
 		Out.SharedLocalBasesIndexOffset = 0; // Unused
 
-		Out.PackedTopLayerData = InStrataTopLayerTexture.Load(uint3(StrataAddressing.PixelCoords, 0));
+		Out.PackedTopLayerData = StrataLoadTopLayerData(InStrataTopLayerTexture, StrataAddressing);
 		StrataAddressing.ReadBytes += 4;
 	  #endif
 	}
