Reorganize read/write strata data functions

diff --git a/Engine/Shaders/Private/Strata/Strata.ush b/Engine/Shaders/Private/Strata/Strata.ush
--- a/Engine/Shaders/Private/Strata/Strata.ush
+++ b/Engine/Shaders/Private/Strata/Strata.ush
@@ -1206,7 +1206,7 @@ FStrataAddressing GetStrataPixelDataByteOffset(uint2 PixelPos, uint2 ViewBufferS
 }
 
 ///////////////////////////////////////////////////////////////////////////////
-// Material buffer Read/Write
+// Material buffer Write
 
 // Strata Rendertarget / UAV setup
 // RT0: Color
@@ -1237,23 +1237,9 @@ FRWStrataMaterialContainer InitialiseRWStrataMaterialContainer(RWTexture2DArray<
 	return RWStrataMaterialContainer;
 }
 
-#if COMPILER_SWITCH // These platforms are complaining so simply disabling Strata there for now.
-
-#define FStrataMaterialContainer Texture2DArray<uint>
-
-void StrataStoreUint1(inout FRWStrataMaterialContainer StrataBuffer, RWTexture2DArray<uint> ExtraMaterialDataUAV, inout FStrataAddressing StrataAddressing, uint Data) {}
-#define STRATA_STORE_UINT1(x) 
-
-void FinalizeWrites(inout FRWStrataMaterialContainer StrataBuffer, inout FStrataAddressing StrataAddressing) {}
-
-uint StrataLoadUint1(FStrataMaterialContainer StrataBuffer, inout FStrataAddressing StrataAddressing) { return 0; }
-
-#else // COMPILER_SWITCH
-
-#define FStrataMaterialContainer Texture2DArray<uint>
-
 void StrataStoreUint1(inout FRWStrataMaterialContainer StrataBuffer, RWTexture2DArray<uint> ExtraMaterialDataUAV, inout FStrataAddressing StrataAddressing, uint Data)
 {
+  #if !COMPILER_SWITCH
 	// We heavily rely on the compiler to optimize out the if code
 	if (StrataAddressing.CurrentIndex < STRATA_BASE_PASS_MRT_OUTPUT_COUNT)
 	{
@@ -1268,6 +1254,7 @@ void StrataStoreUint1(inout FRWStrataMaterialContainer StrataBuffer, RWTexture2D
 		ExtraMaterialDataUAV[uint3(StrataAddressing.PixelCoords, StrataAddressing.CurrentIndex - STRATA_BASE_PASS_MRT_OUTPUT_COUNT)] = Data;
 	}
 	StrataAddressing.CurrentIndex++;
+  #endif // !COMPILER_SWITCH
 }
 
 #define STRATA_STORE_UINT1(x)	StrataStoreUint1(StrataBuffer, ExtraMaterialDataUAV, StrataAddressing, x)
@@ -1277,23 +1264,33 @@ void FinalizeWrites(inout FRWStrataMaterialContainer StrataBuffer, inout FStrata
 	// NOP
 }
 
+///////////////////////////////////////////////////////////////////////////////
+// Material buffer read
+
+#define FStrataMaterialContainer Texture2DArray<uint>
+
 uint StrataLoadUint1(FStrataMaterialContainer StrataBuffer, inout FStrataAddressing StrataAddressing)
 {
+	uint Data = 0;
+  #if !COMPILER_SWITCH
 	const uint MipLevel = 0;
-	uint Data = StrataBuffer.mips[MipLevel][uint3(StrataAddressing.PixelCoords, StrataAddressing.CurrentIndex)];
+	Data = StrataBuffer.mips[MipLevel][uint3(StrataAddressing.PixelCoords, StrataAddressing.CurrentIndex)];
 	StrataAddressing.CurrentIndex++;
 	StrataAddressing.ReadBytes += 4;
+  #endif
 	return Data;
 }
 
 uint StrataLoadUint1AtIndexOffset(in FStrataMaterialContainer StrataBuffer, in FStrataAddressing StrataAddressing, uint IndexOffset)
 {
+	uint Data = 0;
+  #if !COMPILER_SWITCH
 	const uint MipLevel = 0;
-	return StrataBuffer.mips[MipLevel][uint3(StrataAddressing.PixelCoords, IndexOffset)];
+	Data = StrataBuffer.mips[MipLevel][uint3(StrataAddressing.PixelCoords, IndexOffset)];
+  #endif
+	return Data;
 }
 
-#endif // COMPILER_SWITCH
-
 #define FStrataTopLayerDataContainer Texture2D<uint>
 
 uint StrataLoadTopLayerData(FStrataTopLayerDataContainer InTopLayerContainer, inout FStrataAddressing StrataAddressing)
