Fix missing OVERRIDE_PLATFORM_RAYTRACING_INSTANCE_DESCRITOR guard.

diff --git a/Engine/Shaders/Private/RayTracing/RayTracingCommon.ush b/Engine/Shaders/Private/RayTracing/RayTracingCommon.ush
--- a/Engine/Shaders/Private/RayTracing/RayTracingCommon.ush
+++ b/Engine/Shaders/Private/RayTracing/RayTracingCommon.ush
@@ -92,6 +92,7 @@
 [shader("miss")] void name(inout payload_type payload_name)
 #endif //RAY_TRACING_ENTRY_MISS
 
+#ifndef OVERRIDE_PLATFORM_RAYTRACING_INSTANCE_DESCRITOR
 
 // Must match D3D12_RAYTRACING_INSTANCE_FLAGS / VkGeometryInstanceFlagBitsKHR
 #define PLATFORM_RAYTRACING_INSTANCE_FLAG_NONE								0
@@ -100,6 +101,8 @@
 #define PLATFORM_RAYTRACING_INSTANCE_FLAG_FORCE_OPAQUE						0x4
 #define PLATFORM_RAYTRACING_INSTANCE_FLAG_FORCE_NON_OPAQUE					0x8
 
+#endif // OVERRIDE_PLATFORM_RAYTRACING_INSTANCE_DESCRITOR
+
 // must match ERayTracingInstanceFlags
 #define RAYTRACING_INSTANCE_FLAG_NONE					0
 #define RAYTRACING_INSTANCE_FLAG_TRIANGLE_CULL_DISABLE	(1 << 1)
