Strata - small udpate to debug data written.

diff --git a/Engine/Shaders/Private/Strata/Strata.ush b/Engine/Shaders/Private/Strata/Strata.ush
--- a/Engine/Shaders/Private/Strata/Strata.ush
+++ b/Engine/Shaders/Private/Strata/Strata.ush
@@ -722,7 +722,7 @@ float StrataComputeHazeWeight(float Roughness0, float Haziness)    { float Rough
 struct FStrataOperator
 {
 	// STRATA_TODO: pack in a single uint
-	uint ParentIndex;
+	int ParentIndex;
 	uint LayerDepth;
 	uint MaxDistanceFromLeaves;							// The largest depth distance from any tree leaf representing a BSDF
 	uint Type;											// The type of operator
@@ -786,18 +786,28 @@ FStrataTree GetInitialisedStrataTree()
 
 void StrataPackOutStrataTree(
 	RWTexture2DArray<uint> ExtraMaterialDataUAV,
+	in uint RootOperatorIndex,
 	in FStrataTree StrataTree,
 	in int SliceStoringDebugStrataTree)
 {
-	int OutCoord = 0;
-#define WRITE_UINT(Data) ExtraMaterialDataUAV[uint3(OutCoord++, 0, SliceStoringDebugStrataTree)] = Data;
-#define WRITE_FLOAT(Data) ExtraMaterialDataUAV[uint3(OutCoord++, 0, SliceStoringDebugStrataTree)] = asuint(Data);
+#define WRITE_INT(Data) ExtraMaterialDataUAV[uint3(OutCoord++, OutCoordY, SliceStoringDebugStrataTree)] = asuint(Data);
+#define WRITE_UINT(Data) ExtraMaterialDataUAV[uint3(OutCoord++, OutCoordY, SliceStoringDebugStrataTree)] = Data;
+#define WRITE_FLOAT(Data) ExtraMaterialDataUAV[uint3(OutCoord++, OutCoordY, SliceStoringDebugStrataTree)] = asuint(Data);
 #define WRITE_FLOAT3(Data) WRITE_FLOAT(Data.x);WRITE_FLOAT(Data.y);WRITE_FLOAT(Data.z);
 #define WRITE_FLOAT4(Data) WRITE_FLOAT(Data.x);WRITE_FLOAT(Data.y);WRITE_FLOAT(Data.z);;WRITE_FLOAT(Data.w);
 
+	// Write out header
+	int OutCoord = 0;
+	int OutCoordY = 0;
+
 	WRITE_UINT(STRATA_USES_CONVERSION_FROM_LEGACY ? 1 : 0);
 	WRITE_UINT(StrataTree.BSDFCount);
 	WRITE_UINT(StrataTree.OperatorCount);
+	WRITE_UINT(RootOperatorIndex);
+
+	// Write out BSDFs
+	OutCoord = 0;
+	OutCoordY = 1;
 
 #if COMPILER_DXC == 1
 	LOOP
@@ -808,7 +818,7 @@ void StrataPackOutStrataTree(
 	for(int i=0; i< StrataTree.BSDFCount; ++i)
 	{
 		WRITE_UINT(StrataTree.BSDFs[i].State);
-		WRITE_UINT(StrataTree.BSDFs[i].OperatorIndex);
+		WRITE_INT(StrataTree.BSDFs[i].OperatorIndex);
 		WRITE_FLOAT3(StrataTree.BSDFs[i].LuminanceWeightV);
 		WRITE_FLOAT3(StrataTree.BSDFs[i].LuminanceWeightL);
 		WRITE_UINT(StrataTree.BSDFs[i].bIsBottom);
@@ -833,6 +843,29 @@ void StrataPackOutStrataTree(
 		WRITE_FLOAT4(StrataTree.BSDFs[i].VGPRs[4]);
 	}
 
+	// Write out operators
+	OutCoord = 0;
+	OutCoordY = 2;
+
+#if COMPILER_DXC == 1
+	LOOP
+#else
+	StrataTree.OperatorCount = min(StrataTree.OperatorCount, STRATA_MAX_OPERATOR_COUNT);
+	STRATA_UNROLL_N(STRATA_MAX_OPERATOR_COUNT)
+#endif
+	for (i = 0; i < StrataTree.OperatorCount; ++i)
+	{
+		WRITE_INT(StrataTree.Operators[i].ParentIndex);
+		WRITE_UINT(StrataTree.Operators[i].LayerDepth);
+		WRITE_UINT(StrataTree.Operators[i].MaxDistanceFromLeaves);
+		WRITE_UINT(StrataTree.Operators[i].Type);
+		WRITE_INT(StrataTree.Operators[i].LeftIndex);
+		WRITE_INT(StrataTree.Operators[i].RightIndex);
+		WRITE_FLOAT(StrataTree.Operators[i].Weight);
+		WRITE_FLOAT(StrataTree.Operators[i].Coverage);
+		// Skip the rest for now
+	}
+
 #undef WRITE_UINT
 #undef WRITE_FLOAT
 #undef WRITE_FLOAT3
@@ -844,10 +877,12 @@ struct FStrataTreeHeader
 	int ConvertedFromLegacy;
 	int BSDFCount;
 	int OperatorCount;
+	uint RootOperatorIndex;
 };
 
-#define READ_UINT(Data) Data = StrataMaterialData[uint3(OutCoord++, 0, SliceStoringDebugStrataTree)];
-#define READ_FLOAT(Data) Data = asfloat(StrataMaterialData[uint3(OutCoord++, 0, SliceStoringDebugStrataTree)]);
+#define READ_INT(Data) Data = asint(StrataMaterialData[uint3(OutCoord++, OutCoordY, SliceStoringDebugStrataTree)]);
+#define READ_UINT(Data) Data = StrataMaterialData[uint3(OutCoord++, OutCoordY, SliceStoringDebugStrataTree)];
+#define READ_FLOAT(Data) Data = asfloat(StrataMaterialData[uint3(OutCoord++, OutCoordY, SliceStoringDebugStrataTree)]);
 #define READ_FLOAT3(Data) READ_FLOAT(Data.x);READ_FLOAT(Data.y);READ_FLOAT(Data.z);
 #define READ_FLOAT4(Data) READ_FLOAT(Data.x);READ_FLOAT(Data.y);READ_FLOAT(Data.z);READ_FLOAT(Data.w);
 
@@ -859,10 +894,12 @@ void StrataUnpackInStrataTreeHeader(
 	StrataTreeHeader = (FStrataTreeHeader)0;
 
 	int OutCoord = 0;
+	int OutCoordY= 0;
 
 	READ_UINT(StrataTreeHeader.ConvertedFromLegacy);
 	READ_UINT(StrataTreeHeader.BSDFCount);
 	READ_UINT(StrataTreeHeader.OperatorCount);
+	READ_UINT(StrataTreeHeader.RootOperatorIndex);
 }
 
 FStrataBSDF StrataUnpackInStrataTreeBSDF(
@@ -871,13 +908,13 @@ FStrataBSDF StrataUnpackInStrataTreeBSDF(
 	inout FStrataTreeHeader StrataTreeHeader,
 	in int SliceStoringDebugStrataTree)
 {
-	int OutCoord = 3;												// Header
-	OutCoord += BSDFIndex * (STRATA_INLINE_SHADING ? 39 : 30);		// Skip BSDFs
+	int OutCoord = BSDFIndex * (STRATA_INLINE_SHADING ? 39 : 30);		// Skip BSDFs
+	int OutCoordY= 1;
 
 	FStrataBSDF BSDF = (FStrataBSDF)0;
 
 	READ_UINT(BSDF.State);
-	READ_UINT(BSDF.OperatorIndex);
+	READ_INT(BSDF.OperatorIndex);
 	READ_FLOAT3(BSDF.LuminanceWeightV);
 	READ_FLOAT3(BSDF.LuminanceWeightL);
 	READ_UINT(BSDF.bIsBottom);
@@ -904,6 +941,30 @@ FStrataBSDF StrataUnpackInStrataTreeBSDF(
 	return BSDF;
 }
 
+FStrataOperator StrataUnpackInStrataTreeOperator(
+	uint OperatorIndex,
+	Texture2DArray<uint> StrataMaterialData,
+	inout FStrataTreeHeader StrataTreeHeader,
+	in int SliceStoringDebugStrataTree)
+{
+	int OutCoord = OperatorIndex * 8;	// Skip Operators
+	int OutCoordY = 2;
+
+	FStrataOperator Op = (FStrataOperator)0;
+
+	READ_INT(Op.ParentIndex);
+	READ_UINT(Op.LayerDepth);
+	READ_UINT(Op.MaxDistanceFromLeaves);
+	READ_UINT(Op.Type);
+	READ_INT(Op.LeftIndex);
+	READ_INT(Op.RightIndex);
+	READ_FLOAT(Op.Weight);
+	READ_FLOAT(Op.Coverage);
+
+	return Op;
+}
+
+#undef READ_INT
 #undef READ_UINT
 #undef READ_FLOAT
 #undef READ_FLOAT3
@@ -2874,16 +2935,17 @@ void PackStrataOut(
 		FStrataAddressing NullStrataAddressing = (FStrataAddressing)0;	// Fake unused in StrataCreateBSDFContext when using Forward inline shading
 		StrataUpdateTree(NullStrataAddressing, StrataPixelHeader, StrataTree, V, Settings);
 
+		const uint RootOperatorIndex = Strata.OperatorIndex;
+
 	#if STRATA_ADVANCED_DEBUG_ENABLED
 		if (Settings.SliceStoringDebugStrataTree > -1)
 		{
-			StrataPackOutStrataTree(ExtraMaterialDataUAV, StrataTree, Settings.SliceStoringDebugStrataTree);
+			StrataPackOutStrataTree(ExtraMaterialDataUAV, RootOperatorIndex, StrataTree, Settings.SliceStoringDebugStrataTree);
 		}
 	#endif
 
 		// Process the rough refraction data (this code is only used, and not culled out, if effectively used in base pass for instance)
 		{
-			uint RootOperatorIndex = Strata.OperatorIndex;
 			const FStrataOperator RootOperator = StrataTree.Operators[Strata.OperatorIndex];
 			OpaqueRoughRefractionData.Coverage = RootOperator.OpaqueRoughRefractCoverage;
 
