Starta - more material topology debug

diff --git a/Engine/Shaders/Private/Strata/StrataVisualize.usf b/Engine/Shaders/Private/Strata/StrataVisualize.usf
--- a/Engine/Shaders/Private/Strata/StrataVisualize.usf
+++ b/Engine/Shaders/Private/Strata/StrataVisualize.usf
@@ -3,7 +3,7 @@
 #include "/Engine/Private/Common.ush"
 #include "/Engine/Private/DeferredShadingCommon.ush"
 
-#ifdef SHADER_DEBUGSTRATATREE_PS
+#if defined(SHADER_DEBUGSTRATATREE_PS) || defined(SHADER_DEBUGSTRATATREE_CS)
 #define STRATA_INLINE_SHADING 1
 #else
 #define STRATA_INLINE_SHADING 0
@@ -455,7 +455,11 @@ void VisualizeMaterialPS(
 ///////////////////////////////////////////////////////////////////////////////////////////////////
 // Material Visualize
 
-#if SHADER_DEBUGSTRATATREE_TEST_CS
+#define BSDF_ELEMENT_SIZE 210
+#define BSDF_ELEMENT_STARTX 100.0f + 450.0f
+#define BSDF_ELEMENT_STARTY 100
+
+#if SHADER_DEBUGSTRATATREE_CS
 
 [numthreads(1, 1, 1)]
 void MaterialDebugStrataTreeCS(uint3 DispatchThreadId : SV_DispatchThreadID)
@@ -478,23 +482,162 @@ void MaterialDebugStrataTreeCS(uint3 DispatchThreadId : SV_DispatchThreadID)
 		Context = InitShaderPrintContext(true, uint2(50, 50));
 
 		const int BSDFCount = StrataTreeHeader.BSDFCount;
-		// We remove BSDF from operators to no confuse the user. Behind the scene BSDF are operators sitting the leaves of the tree.
-		const int OperatorCount = StrataTreeHeader.OperatorCount - StrataTreeHeader.BSDFCount;
 
-		Print(Context, BSDFCount, FontWhite, 2, 0);
-		Print(Context, TEXT("BSDFs"), FontWhite);
+		Print(Context, BSDFCount, FontYellow, 2, 0);
+		Print(Context, TEXT("BSDFs"), FontYellow);
 		if (StrataTreeHeader.ConvertedFromLegacy)
 		{
-			Print(Context, TEXT(" (Converted from legacy material)"), FontWhite);
+			Print(Context, TEXT(" (Converted from legacy material)"), FontOrange);
 		}
 		Newline(Context);
-		Print(Context, OperatorCount, FontWhite, 2, 0);
+		Print(Context, StrataTreeHeader.OperatorCount, FontWhite, 2, 0);
 		Print(Context, TEXT("Operators"), FontWhite);
-		Newline(Context);
+
+		#define MAX_LAYER_COUNT 4
+		float LayerAccCoverage[MAX_LAYER_COUNT];
+		for (int i = 0; i < MAX_LAYER_COUNT; ++i)
+		{
+			LayerAccCoverage[i] = 0.0f;
+		}
+
+		Context = InitShaderPrintContext(true, uint2(50, BSDF_ELEMENT_STARTY));
+	#if COMPILER_DXC == 1
+		LOOP
+	#else
+		StrataTreeHeader.OperatorCount = min(StrataTreeHeader.OperatorCount, 2);
+		STRATA_UNROLL_N(2)
+	#endif
+		for (i = 0; i < StrataTreeHeader.OperatorCount; i++)
+		{
+			FStrataOperator Op = StrataUnpackInStrataTreeOperator(
+				i,
+				Strata.MaterialTextureArray,
+				StrataTreeHeader,
+				SliceStoringDebugStrataTree);
+
+			Print(Context, TEXT("OperatorIndex = "), FontYellow);
+			Print(Context, i, FontYellow, 2, 0);
+			Newline(Context);
+			Print(Context, TEXT("ParentIndex = "), FontWhite);
+			Print(Context, Op.ParentIndex, FontWhite, 2, 0);
+			Newline(Context);
+			Print(Context, TEXT("LayerDepth = "), FontWhite);
+			Print(Context, Op.LayerDepth, FontWhite, 2, 0);
+			Newline(Context);
+			Print(Context, TEXT("MaxDistanceFromLeaves = "), FontWhite);
+			Print(Context, Op.MaxDistanceFromLeaves, FontWhite, 2, 0);
+			Newline(Context);
+			Print(Context, TEXT("Type = "), FontWhite);
+			if (Op.Type == STRATA_OPERATOR_WEIGHT)
+			{
+				Print(Context, TEXT("WEIGHT"), FontWhite);
+			}
+			else if (Op.Type == STRATA_OPERATOR_VERTICAL)
+			{
+				Print(Context, TEXT("VERTICAL"), FontWhite);
+			}
+			else if (Op.Type == STRATA_OPERATOR_HORIZONTAL)
+			{
+				Print(Context, TEXT("HORIZONTAL"), FontWhite);
+			}
+			else if (Op.Type == STRATA_OPERATOR_ADD)
+			{
+				Print(Context, TEXT("ADD"), FontWhite);
+			}
+			else if (Op.Type == STRATA_OPERATOR_BSDF)
+			{
+				Print(Context, TEXT("BSDF"), FontWhite);
+			}
+			else if (Op.Type == STRATA_OPERATOR_BSDF_LEGACY)
+			{
+				Print(Context, TEXT("BSDF_LEGACY"), FontWhite);
+			}
+			else if (Op.Type == STRATA_OPERATOR_THINFILM)
+			{
+				Print(Context, TEXT("THINFILM"), FontWhite);
+			}
+			else
+			{
+				Print(Context, TEXT("UNKOWN"), FontWhite);
+			}
+			Newline(Context);
+			Print(Context, TEXT("LeftIndex = "), FontWhite);
+			Print(Context, Op.LeftIndex, FontWhite, 2, 0);
+			Newline(Context);
+			Print(Context, TEXT("RightIndex = "), FontWhite);
+			Print(Context, Op.RightIndex, FontWhite, 2, 0);
+			Newline(Context);
+			Print(Context, TEXT("Weight = "), FontWhite);
+			Print(Context, Op.Weight, FontWhite, 5, 2);
+			Newline(Context);
+			Print(Context, TEXT("Coverage = "), FontWhite);
+			Print(Context, Op.Coverage, FontWhite, 5, 2);
+			Newline(Context);
+			Newline(Context);
+		}
+
+	#if COMPILER_DXC == 1
+		LOOP
+	#else
+		StrataTreeHeader.BSDFCount = min(StrataTreeHeader.BSDFCount, 2);
+		STRATA_UNROLL_N(2)
+	#endif
+		for (i = 0; i < StrataTreeHeader.BSDFCount; i++)
+		{
+			FStrataBSDF BSDF = StrataUnpackInStrataTreeBSDF(
+				i,
+				Strata.MaterialTextureArray,
+				StrataTreeHeader,
+				SliceStoringDebugStrataTree);
+
+			const bool bActive = any(BSDF.LuminanceWeightV > 0.0);
+			FFontColor BSDFFont = FontDarkRed;
+
+			Context = InitShaderPrintContext(true, uint2(300, BSDF_ELEMENT_STARTY + BSDF_ELEMENT_SIZE * i));
+			Print(Context, TEXT("BSDFIndex = "), FontYellow);
+			Print(Context, i, FontYellow, 2, 0);
+			if (!bActive)
+			{
+				Print(Context, TEXT("DISABLED"), FontDarkRed);
+			}
+			else
+			{
+				BSDFFont = FontWhite;
+			}
+			Newline(Context);
+			Print(Context, TEXT("OperatorIndex = "), BSDFFont);
+			Print(Context, BSDF.OperatorIndex, BSDFFont, 2, 0);
+			Newline(Context);
+			Print(Context, TEXT("Coverage = "), BSDFFont);
+			Print(Context, BSDF.Coverage, BSDFFont, 5, 2);
+			Newline(Context);
+			Print(Context, TEXT("LumWghtV = "), BSDFFont);
+			Print(Context, BSDF.LuminanceWeightV.r, FontLightRed, 5, 2);
+			Print(Context, BSDF.LuminanceWeightV.g, FontLightGreen, 5, 2);
+			Print(Context, BSDF.LuminanceWeightV.b, FontLightBlue, 5, 2);
+			Newline(Context);
+			Print(Context, TEXT("LumWghtL = "), BSDFFont);
+			Print(Context, BSDF.LuminanceWeightL.r, FontLightRed, 5, 2);
+			Print(Context, BSDF.LuminanceWeightL.g, FontLightGreen, 5, 2);
+			Print(Context, BSDF.LuminanceWeightL.b, FontLightBlue, 5, 2);
+			Newline(Context);
+			Print(Context, TEXT("Diffuse  = "), BSDFFont);
+			const float3 BSDFDiffuse = StrataGetBSDFDiffuseColor(BSDF);
+			Print(Context, BSDFDiffuse.r, FontLightRed, 5, 2);
+			Print(Context, BSDFDiffuse.g, FontLightGreen, 5, 2);
+			Print(Context, BSDFDiffuse.b, FontLightBlue, 5, 2);
+			Newline(Context);
+		}
+
+		Context = InitShaderPrintContext(true, uint2(675, BSDF_ELEMENT_STARTY - 50));
+		Print(Context, TEXT("BSDF visualization"), FontWhite);
+
+		Context = InitShaderPrintContext(true, uint2(1165, BSDF_ELEMENT_STARTY - 50));
+		Print(Context, TEXT("Topology"), FontWhite);
 	}
 }
 
-#endif // SHADER_DEBUGSTRATATREE_TEST_CS
+#endif // SHADER_DEBUGSTRATATREE_CS
 
 #if SHADER_DEBUGSTRATATREE_PS
 
@@ -682,8 +825,8 @@ void GetMaterialOnCube(
 		const float2 UV = (PixelPos - Pos) / Size;
 		OutColor = GetBackgroundCheckboardColor(UV);
 
-		const float3 RayO = -View.ViewForward;
-		const float3 RayD = View.ViewForward + (UV.x - 0.5) * View.ViewRight -(UV.y - 0.5) * View.ViewUp;
+		const float3 RayO = -View.ViewForward*1.2;
+		const float3 RayD = View.ViewForward + 1.0 * (UV.x - 0.5) * View.ViewRight - 1.0 * (UV.y - 0.5) * View.ViewUp;
 		float3 V = -RayD;
 
 		float2 Hit;
@@ -728,7 +871,7 @@ void GetMaterialOnSphere(
 		const float2 UV = (PixelPos - Pos) / Size;
 		OutColor = GetBackgroundCheckboardColor(UV);
 
-		float3 V = -View.ViewForward;// float3(0, 0, 1);
+		float3 V = -View.ViewForward;
 		const float3 RayO = V + (UV.x-0.5) * View.ViewRight + -(UV.y-0.5) * View.ViewUp;
 		const float3 RayD = -V;
 		const float4 Sphere = float4(0.0f, 0.0f, 0.f, 0.48);
@@ -755,6 +898,43 @@ void GetMaterialOnSphere(
 	}
 }
 
+void GetMaterialSlice(
+	in float4 SVPos,
+	in float2 Pos,
+	in float2  Size,
+	in FStrataBSDF BSDF,
+	inout float4 OutColor)
+{
+	const float2 PixelPos = SVPos.xy;
+
+	if (BSDF_GETTYPE(BSDF) == STRATA_BSDF_TYPE_SLAB && all(int2(PixelPos) > int2(Pos)) && all(int2(PixelPos) < int2(Pos + Size)))
+	{
+		const float2 UV = (PixelPos - Pos) / Size;
+		OutColor = GetBackgroundCheckboardColor((PixelPos - Pos) / Size.x);
+
+		// View and normal along up.
+		float3 V = float3(0, 0, 1);
+		float3 N = float3(0, 0, 1);
+		const float ThicknessRatio = 1.0;
+
+		float3x3 TangentBasis;
+		TangentBasis[2] = N;
+		TangentBasis[0] = float3(1, 0, 0);
+		TangentBasis[1] = float3(0, 1, 0);
+
+		EvaluateLighting(
+			BSDF,
+			V,
+			TangentBasis,
+			ThicknessRatio,
+			OutColor);
+	}
+	else
+	{
+		OutColor = 0;
+	}
+}
+
 void MaterialDebugStrataTreePS(
 	float4 SVPos : SV_POSITION,
 	out float4 OutColor : SV_Target0)
@@ -765,7 +945,12 @@ void MaterialDebugStrataTreePS(
 	FStrataPixelHeader Header = UnpackStrataHeaderIn(Strata.MaterialTextureArray, StrataAddressing, Strata.TopLayerTexture);
 	if (Header.BSDFCount > 0)
 	{
-		OutColor = float4(0.0f, 0.0f, 0.0f, 0.9f);
+		const float2 PixelPos = SVPos.xy;
+		OutColor = float4(0.0f, 0.0f, 0.0f, PixelPos.x < 540.0f ? 0.5f : 1.0f); // Background
+		if (int(PixelPos.x) == 280 && int(PixelPos.y) > 90)
+		{
+			OutColor = float4(1.0f, 1.0f, 1.0f, 0.5f);							// Operator/BSDF separator
+		}
 
 		const int SliceStoringDebugStrataTree = Strata.SliceStoringDebugStrataTree + STRATA_BASE_PASS_MRT_OUTPUT_COUNT;
 
@@ -775,16 +960,15 @@ void MaterialDebugStrataTreePS(
 			StrataTreeHeader,
 			SliceStoringDebugStrataTree);
 
-		const float2 PixelPos = SVPos.xy;
 
-		float2 Pos = float2(100.0f, 100.0f);
+		float2 Pos = float2(BSDF_ELEMENT_STARTX, BSDF_ELEMENT_STARTY);
 		const float MaterialPanelSize = 200.0f;
 
 		// STRATA_TODO: read BSDF and Operators "on the fly".
 #if COMPILER_DXC == 1
 		LOOP
 #else
-		StrataTree.BSDFCount = min(StrataTreeHeader.BSDFCount, 2);
+		StrataTreeHeader.BSDFCount = min(StrataTreeHeader.BSDFCount, 2);
 		STRATA_UNROLL_N(2)
 #endif
 		for (int i = 0; i < StrataTreeHeader.BSDFCount; i++)
@@ -795,7 +979,7 @@ void MaterialDebugStrataTreePS(
 				StrataTreeHeader,
 				SliceStoringDebugStrataTree);
 
-			if (BSDF.Coverage > 0.0f)
+			//if (BSDF.Coverage > 0.0f)
 			{
 				GetMaterialOnSphere(
 					SVPos,
@@ -814,6 +998,59 @@ void MaterialDebugStrataTreePS(
 				Pos.y += MaterialPanelSize + 10.0f;
 			}
 		}
+
+		#define MAX_LAYER_COUNT 4
+		float LayerAccCoverage[MAX_LAYER_COUNT];
+		for (i = 0; i < MAX_LAYER_COUNT; ++i)
+		{
+			LayerAccCoverage[i] = 0.0f;
+		}
+
+		float2 MatTopOrigin = float2(BSDF_ELEMENT_STARTX + 450, BSDF_ELEMENT_STARTY);// float2(BSDF_ELEMENT_STARTX, 550);
+		float2 MatTopLayerSize = float2(410, 100);
+		
+		// This is a really basic material topology debug. STRATA_TODO would be to draw the material graph itself with anottation
+#if COMPILER_DXC == 1
+		LOOP
+#else
+		StrataTreeHeader.BSDFCount = min(StrataTreeHeader.BSDFCount, 2);
+		STRATA_UNROLL_N(2)
+#endif
+		for (i = 0; i < StrataTreeHeader.BSDFCount; i++)
+		{
+			FStrataBSDF BSDF = StrataUnpackInStrataTreeBSDF(
+				i,
+				Strata.MaterialTextureArray,
+				StrataTreeHeader,
+				SliceStoringDebugStrataTree);
+
+			FStrataOperator Op = StrataUnpackInStrataTreeOperator(
+				BSDF.OperatorIndex,
+				Strata.MaterialTextureArray,
+				StrataTreeHeader,
+				SliceStoringDebugStrataTree);
+
+			const uint LayerDepth = Op.LayerDepth;
+
+			if (LayerDepth < MAX_LAYER_COUNT)
+			{
+				// Draw the slab information
+				float2 LayerOrigin = MatTopOrigin + MatTopLayerSize.xy * float2(LayerAccCoverage[LayerDepth], LayerDepth);
+				float2 SlabSize = MatTopLayerSize.xy * float2(BSDF.Coverage, 1.0);
+
+				if (BSDF_GETTYPE(BSDF) == STRATA_BSDF_TYPE_SLAB && all(int2(PixelPos) >= int2(LayerOrigin)) && all(int2(PixelPos) <= int2(LayerOrigin + SlabSize)))
+				{
+					GetMaterialSlice(
+						SVPos,
+						LayerOrigin,
+						SlabSize,
+						BSDF,
+						OutColor);
+				}
+
+				LayerAccCoverage[LayerDepth] += BSDF.Coverage;
+			}
+		}
 	}
 
 	OutColor = float4(pow(OutColor.rgb, 1.0 / 2.2), OutColor.a);
diff --git a/Engine/Source/Runtime/Renderer/Private/Strata/StrataVisualize.cpp b/Engine/Source/Runtime/Renderer/Private/Strata/StrataVisualize.cpp
--- a/Engine/Source/Runtime/Renderer/Private/Strata/StrataVisualize.cpp
+++ b/Engine/Source/Runtime/Renderer/Private/Strata/StrataVisualize.cpp
@@ -153,7 +153,7 @@ class FMaterialDebugStrataTreeCS : public FGlobalShader
 
 		// Stay debug and skip optimizations to reduce compilation time on this long shader.
 		OutEnvironment.CompilerFlags.Add(CFLAG_Debug);
-		OutEnvironment.SetDefine(TEXT("SHADER_DEBUGSTRATATREE_TEST_CS"), 1);
+		OutEnvironment.SetDefine(TEXT("SHADER_DEBUGSTRATATREE_CS"), 1);
 	}
 };
 IMPLEMENT_GLOBAL_SHADER(FMaterialDebugStrataTreeCS, "/Engine/Private/Strata/StrataVisualize.usf", "MaterialDebugStrataTreeCS", SF_Compute);
