Strata - update to link BSDF indices and Operator indices via colors

diff --git a/Engine/Shaders/Private/Strata/StrataVisualize.usf b/Engine/Shaders/Private/Strata/StrataVisualize.usf
--- a/Engine/Shaders/Private/Strata/StrataVisualize.usf
+++ b/Engine/Shaders/Private/Strata/StrataVisualize.usf
@@ -515,8 +515,8 @@ void MaterialDebugStrataTreeCS(uint3 DispatchThreadId : SV_DispatchThreadID)
 				StrataTreeHeader,
 				SliceStoringDebugStrataTree);
 
-			Print(Context, TEXT("OperatorIndex = "), FontYellow);
-			Print(Context, i, FontYellow, 2, 0);
+			Print(Context, TEXT("OperatorIndex = "), FontCyan);
+			Print(Context, i, FontCyan, 2, 0);
 			Newline(Context);
 			Print(Context, TEXT("ParentIndex = "), FontWhite);
 			Print(Context, Op.ParentIndex, FontWhite, 2, 0);
@@ -561,8 +561,16 @@ void MaterialDebugStrataTreeCS(uint3 DispatchThreadId : SV_DispatchThreadID)
 				Print(Context, TEXT("UNKOWN"), FontWhite);
 			}
 			Newline(Context);
-			Print(Context, TEXT("LeftIndex = "), FontWhite);
-			Print(Context, Op.LeftIndex, FontWhite, 2, 0);
+			if (Op.Type == STRATA_OPERATOR_BSDF)
+			{
+				Print(Context, TEXT("BSDFIndex = "), FontYellow);
+				Print(Context, Op.LeftIndex, FontYellow, 2, 0);
+			}
+			else
+			{
+				Print(Context, TEXT("LeftIndex = "), FontWhite);
+				Print(Context, Op.LeftIndex, FontWhite, 2, 0);
+			}
 			Newline(Context);
 			Print(Context, TEXT("RightIndex = "), FontWhite);
 			Print(Context, Op.RightIndex, FontWhite, 2, 0);
@@ -605,8 +613,8 @@ void MaterialDebugStrataTreeCS(uint3 DispatchThreadId : SV_DispatchThreadID)
 				BSDFFont = FontWhite;
 			}
 			Newline(Context);
-			Print(Context, TEXT("OperatorIndex = "), BSDFFont);
-			Print(Context, BSDF.OperatorIndex, BSDFFont, 2, 0);
+			Print(Context, TEXT("OperatorIndex = "), FontCyan);
+			Print(Context, BSDF.OperatorIndex, FontCyan, 2, 0);
 			Newline(Context);
 			Print(Context, TEXT("Coverage = "), BSDFFont);
 			Print(Context, BSDF.Coverage, BSDFFont, 5, 2);
