Mobile GL Deferred ReflectionCapture UBO fix (passed too many elements and hit the GL limit)

diff --git a/Engine/Shaders/Private/LightGridInjection.usf b/Engine/Shaders/Private/LightGridInjection.usf
--- a/Engine/Shaders/Private/LightGridInjection.usf
+++ b/Engine/Shaders/Private/LightGridInjection.usf
@@ -5,6 +5,7 @@
 #include "Common.ush"
 #include "Definitions.usf"
 #include "LargeWorldCoordinates.ush"
+#include "ReflectionEnvironmentShared.ush"
 
 RWBuffer<uint> RWNumCulledLightsGrid;
 RWBuffer<uint> RWCulledLightDataGrid;
@@ -250,10 +251,10 @@ void LightGridInjectionCS(
 		LOOP
 		for (uint ReflectionCaptureIndex = 0; ReflectionCaptureIndex < NumReflectionCaptures; ReflectionCaptureIndex++)
 		{
-			FLWCVector3 CaptureWorldPosition = MakeLWCVector3(ReflectionCapture.TilePosition[ReflectionCaptureIndex].xyz, ReflectionCapture.PositionAndRadius[ReflectionCaptureIndex].xyz);
+			FLWCVector3 CaptureWorldPosition = MakeLWCVector3(GetReflectionTilePosition(ReflectionCaptureIndex).xyz, GetReflectionPositionAndRadius(ReflectionCaptureIndex).xyz);
 			float3 CaptureTranslatedWorldPosition = LWCToFloat(LWCAdd(CaptureWorldPosition, PrimaryView.PreViewTranslation));
 			float3 ViewSpaceCapturePosition = mul(float4(CaptureTranslatedWorldPosition, 1), View.TranslatedWorldToView).xyz;
-			float CaptureRadius = ReflectionCapture.PositionAndRadius[ReflectionCaptureIndex].w;
+			float CaptureRadius = GetReflectionPositionAndRadius(ReflectionCaptureIndex).w;
 
 			float BoxDistanceSq = ComputeSquaredDistanceFromBoxToPoint(ViewTileCenter, ViewTileExtent, ViewSpaceCapturePosition);
 
diff --git a/Engine/Shaders/Private/ReflectionEnvironmentComposite.ush b/Engine/Shaders/Private/ReflectionEnvironmentComposite.ush
--- a/Engine/Shaders/Private/ReflectionEnvironmentComposite.ush
+++ b/Engine/Shaders/Private/ReflectionEnvironmentComposite.ush
@@ -56,11 +56,11 @@ float3 CompositeReflectionCapturesAndSkylightTWS(
 	#endif
 #endif
 
-		FLWCVector3 CaptureWorldPosition = MakeLWCVector3(ReflectionCapture.TilePosition[CaptureIndex].xyz, ReflectionCapture.PositionAndRadius[CaptureIndex].xyz);
+		FLWCVector3 CaptureWorldPosition = MakeLWCVector3(GetReflectionTilePosition(CaptureIndex).xyz, GetReflectionPositionAndRadius(CaptureIndex).xyz);
 		float3 CaptureTranslatedWorldPosition = LWCToFloat(LWCAdd(CaptureWorldPosition, PrimaryView.PreViewTranslation));
-		float CaptureRadius = ReflectionCapture.PositionAndRadius[CaptureIndex].w;
-			
-		float4 CaptureProperties = ReflectionCapture.CaptureProperties[CaptureIndex];
+		float CaptureRadius = GetReflectionPositionAndRadius(CaptureIndex).w;
+
+		float4 CaptureProperties = GetReflectionCaptureProperties(CaptureIndex);
 
 		float3 CaptureVector = TranslatedWorldPosition - CaptureTranslatedWorldPosition;
 		float CaptureVectorLength = sqrt(dot(CaptureVector, CaptureVector));		
@@ -70,7 +70,7 @@ float3 CompositeReflectionCapturesAndSkylightTWS(
 		if (CaptureVectorLength < CaptureRadius)
 		{
 			float3 ProjectedCaptureVector = RayDirection;
-			float4 CaptureOffsetAndAverageBrightness = ReflectionCapture.CaptureOffsetAndAverageBrightness[CaptureIndex];
+			float4 CaptureOffsetAndAverageBrightness = GetReflectionCaptureOffsetAndAverageBrightness(CaptureIndex);
 
 			// Fade out based on distance to capture
 			float DistanceAlpha = 0;
@@ -83,7 +83,8 @@ float3 CompositeReflectionCapturesAndSkylightTWS(
 					BRANCH if (CaptureProperties.b > 0)
 					#endif
 					{
-						ProjectedCaptureVector = GetLookupVectorForBoxCapture(RayDirection, TranslatedWorldPosition, float4(CaptureTranslatedWorldPosition, CaptureRadius), ReflectionCapture.BoxTransform[CaptureIndex], ReflectionCapture.BoxScales[CaptureIndex], CaptureOffsetAndAverageBrightness.xyz, DistanceAlpha);
+						ProjectedCaptureVector = GetLookupVectorForBoxCapture(RayDirection, TranslatedWorldPosition, float4(CaptureTranslatedWorldPosition, CaptureRadius),
+																			 GetReflectionBoxTransform(CaptureIndex), GetReflectionBoxScales(CaptureIndex), CaptureOffsetAndAverageBrightness.xyz, DistanceAlpha);
 					}
 				#endif
 
@@ -123,13 +124,13 @@ float3 CompositeReflectionCapturesAndSkylightTWS(
 
 	float3 ProjectedCaptureVector = RayDirection;
 
-	FLWCVector3 SingleCaptureWorldPosition = MakeLWCVector3(ReflectionCapture.TilePosition[SingleCaptureIndex].xyz, ReflectionCapture.PositionAndRadius[SingleCaptureIndex].xyz);
+	FLWCVector3 SingleCaptureWorldPosition = MakeLWCVector3(GetReflectionTilePosition(SingleCaptureIndex).xyz, GetReflectionPositionAndRadius(SingleCaptureIndex).xyz);
 	float3 SingleCaptureTranslatedWorldPosition = LWCToFloat(LWCAdd(SingleCaptureWorldPosition, PrimaryView.PreViewTranslation));
-	float SingleCaptureRadius = ReflectionCapture.PositionAndRadius[SingleCaptureIndex].w;
+	float SingleCaptureRadius = GetReflectionPositionAndRadius(SingleCaptureIndex).w;
 
-	float4 SingleCaptureOffsetAndAverageBrightness = ReflectionCapture.CaptureOffsetAndAverageBrightness[SingleCaptureIndex];
-	float SingleCaptureBrightness = ReflectionCapture.CaptureProperties[SingleCaptureIndex].x;
-	float SingleCaptureArrayIndex = ReflectionCapture.CaptureProperties[SingleCaptureIndex].y;
+	float4 SingleCaptureOffsetAndAverageBrightness = GetReflectionCaptureOffsetAndAverageBrightness(SingleCaptureIndex);
+	float SingleCaptureBrightness = GetReflectionCaptureProperties(SingleCaptureIndex).x;
+	float SingleCaptureArrayIndex = GetReflectionCaptureProperties(SingleCaptureIndex).y;
 
 	#define APPROXIMATE_CONTINUOUS_SINGLE_CAPTURE_PARALLAX 0
 	#if APPROXIMATE_CONTINUOUS_SINGLE_CAPTURE_PARALLAX
@@ -281,4 +282,4 @@ float3 CompositeReflectionCapturesAndSkylight(
 		SingleCaptureIndex,
 		bCompositeSkylight,
 		EyeIndex);
-}
\ No newline at end of file
+}
diff --git a/Engine/Shaders/Private/ReflectionEnvironmentShared.ush b/Engine/Shaders/Private/ReflectionEnvironmentShared.ush
--- a/Engine/Shaders/Private/ReflectionEnvironmentShared.ush
+++ b/Engine/Shaders/Private/ReflectionEnvironmentShared.ush
@@ -215,3 +215,19 @@ half ComputeMixingWeight(half IndirectIrradiance, half AverageBrightness, half R
 
 	return lerp(1.0f, MixingWeight, MixingAlpha);
 }
+
+#if ES3_1_PROFILE
+#define GetReflectionPositionAndRadius(CaptureIndex) ReflectionCaptureES31.PositionAndRadius[CaptureIndex]
+#define GetReflectionCaptureProperties(CaptureIndex) ReflectionCaptureES31.CaptureProperties[CaptureIndex]
+#define GetReflectionCaptureOffsetAndAverageBrightness(CaptureIndex) ReflectionCaptureES31.CaptureOffsetAndAverageBrightness[CaptureIndex]
+#define GetReflectionBoxTransform(CaptureIndex) ReflectionCaptureES31.BoxTransform[CaptureIndex]
+#define GetReflectionBoxScales(CaptureIndex) ReflectionCaptureES31.BoxScales[CaptureIndex]
+#define GetReflectionTilePosition(CaptureIndex) ReflectionCaptureES31.TilePosition[CaptureIndex]
+#else
+#define GetReflectionPositionAndRadius(CaptureIndex) ReflectionCaptureSM5.PositionAndRadius[CaptureIndex]
+#define GetReflectionCaptureProperties(CaptureIndex) ReflectionCaptureSM5.CaptureProperties[CaptureIndex]
+#define GetReflectionCaptureOffsetAndAverageBrightness(CaptureIndex) ReflectionCaptureSM5.CaptureOffsetAndAverageBrightness[CaptureIndex]
+#define GetReflectionBoxTransform(CaptureIndex) ReflectionCaptureSM5.BoxTransform[CaptureIndex]
+#define GetReflectionBoxScales(CaptureIndex) ReflectionCaptureSM5.BoxScales[CaptureIndex]
+#define GetReflectionTilePosition(CaptureIndex) ReflectionCaptureSM5.TilePosition[CaptureIndex]
+#endif
\ No newline at end of file
diff --git a/Engine/Source/Runtime/Renderer/Private/LightGridInjection.cpp b/Engine/Source/Runtime/Renderer/Private/LightGridInjection.cpp
--- a/Engine/Source/Runtime/Renderer/Private/LightGridInjection.cpp
+++ b/Engine/Source/Runtime/Renderer/Private/LightGridInjection.cpp
@@ -148,6 +148,7 @@ public:
 
 	BEGIN_SHADER_PARAMETER_STRUCT(FParameters, )
 		SHADER_PARAMETER_STRUCT_REF(FReflectionCaptureShaderData, ReflectionCapture)
+		SHADER_PARAMETER_STRUCT_REF(FMobileReflectionCaptureShaderData, MobileReflectionCaptureData)
 		SHADER_PARAMETER_STRUCT_REF(FViewUniformShaderParameters, View)
 		SHADER_PARAMETER_RDG_BUFFER_UAV(RWBuffer<uint>, RWNumCulledLightsGrid)
 		SHADER_PARAMETER_RDG_BUFFER_UAV(RWBuffer<uint>, RWCulledLightDataGrid)
@@ -215,7 +216,7 @@ public:
 		OutEnvironment.SetDefine(TEXT("THREADGROUP_SIZE"), LightGridInjectionGroupSize);
 		FForwardLightingParameters::ModifyCompilationEnvironment(Parameters.Platform, OutEnvironment);
 		OutEnvironment.SetDefine(TEXT("LIGHT_LINK_STRIDE"), LightLinkStride);
-		OutEnvironment.SetDefine(TEXT("MAX_CAPTURES"), GMaxNumReflectionCaptures);
+		OutEnvironment.SetDefine(TEXT("MAX_CAPTURES"), GetMaxNumReflectionCaptures(Parameters.Platform));
 		OutEnvironment.SetDefine(TEXT("ENABLE_LIGHT_CULLING_VIEW_SPACE_BUILD_DATA"), ENABLE_LIGHT_CULLING_VIEW_SPACE_BUILD_DATA);
 	}
 };
@@ -680,7 +681,16 @@ void FSceneRenderer::ComputeLightGrid(FRDGBuilder& GraphBuilder, bool bCullLight
 			FLightGridInjectionCS::FParameters *PassParameters = GraphBuilder.AllocParameters<FLightGridInjectionCS::FParameters>();
 
 			PassParameters->View                    = View.ViewUniformBuffer;
-			PassParameters->ReflectionCapture       = View.ReflectionCaptureUniformBuffer;
+			
+			if (IsMobilePlatform(View.GetShaderPlatform()))
+			{
+				PassParameters->MobileReflectionCaptureData = View.MobileReflectionCaptureUniformBuffer;
+			}
+			else
+			{
+				PassParameters->ReflectionCapture = View.ReflectionCaptureUniformBuffer;
+			}
+
 			PassParameters->RWNumCulledLightsGrid   = NumCulledLightsGridUAV;
 			PassParameters->RWCulledLightDataGrid   = CulledLightDataGridUAV;
 			PassParameters->RWNextCulledLightLink   = GraphBuilder.CreateUAV(NextCulledLightLinkBuffer, PF_R32_UINT);
diff --git a/Engine/Source/Runtime/Renderer/Private/MobileDeferredShadingPass.cpp b/Engine/Source/Runtime/Renderer/Private/MobileDeferredShadingPass.cpp
--- a/Engine/Source/Runtime/Renderer/Private/MobileDeferredShadingPass.cpp
+++ b/Engine/Source/Runtime/Renderer/Private/MobileDeferredShadingPass.cpp
@@ -52,7 +52,7 @@ class FMobileDirectionalLightFunctionPS : public FMaterialShader
 
 	BEGIN_SHADER_PARAMETER_STRUCT(FParameters, )
 		SHADER_PARAMETER_STRUCT_REF(FMobileDirectionalLightShaderParameters, MobileDirectionalLight)
-		SHADER_PARAMETER_STRUCT_REF(FReflectionCaptureShaderData, ReflectionCaptureData)
+		SHADER_PARAMETER_STRUCT_REF(FMobileReflectionCaptureShaderData, MobileReflectionCaptureData)
 		SHADER_PARAMETER(FMatrix44f, TranslatedWorldToLight)
 		SHADER_PARAMETER(FVector4f, LightFunctionParameters)
 		SHADER_PARAMETER(FVector3f, LightFunctionParameters2)
@@ -236,7 +236,7 @@ public:
 
 	BEGIN_SHADER_PARAMETER_STRUCT(FParameters, )
 		SHADER_PARAMETER_STRUCT_REF(FViewUniformShaderParameters, View)
-		SHADER_PARAMETER_STRUCT_REF(FReflectionCaptureShaderData, ReflectionCaptureData)
+		SHADER_PARAMETER_STRUCT_REF(FMobileReflectionCaptureShaderData, MobileReflectionCaptureData)
 	END_SHADER_PARAMETER_STRUCT()
 
 	static bool ShouldCompilePermutation(const FGlobalShaderPermutationParameters& Parameters)
@@ -337,7 +337,7 @@ void RenderReflectionEnvironmentSkyLighting(FRHICommandListImmediate& RHICmdList
 
 	FMobileReflectionEnvironmentSkyLightingPS::FParameters PassParameters;
 	PassParameters.View = GetShaderBinding(View.ViewUniformBuffer);
-	PassParameters.ReflectionCaptureData = GetShaderBinding(View.ReflectionCaptureUniformBuffer);
+	PassParameters.MobileReflectionCaptureData = GetShaderBinding(View.MobileReflectionCaptureUniformBuffer);
 
 	TShaderMapRef<FPostProcessVS> VertexShader(View.ShaderMap);
 
@@ -428,7 +428,7 @@ static void RenderDirectionalLight(FRHICommandListImmediate& RHICmdList, const F
 
 	FMobileDirectionalLightFunctionPS::FParameters PassParameters;
 	PassParameters.MobileDirectionalLight = Scene.UniformBuffers.MobileDirectionalLightUniformBuffers[LightingChannel + 1];
-	PassParameters.ReflectionCaptureData = GetShaderBinding(View.ReflectionCaptureUniformBuffer);
+	PassParameters.MobileReflectionCaptureData = GetShaderBinding(View.MobileReflectionCaptureUniformBuffer);
 	PassParameters.LightFunctionParameters = FVector4f(1.0f, 1.0f, 0.0f, 0.0f);
 
 	const bool bMobileUsesShadowMaskTexture = MobileUsesShadowMaskTexture(View.GetShaderPlatform());
diff --git a/Engine/Source/Runtime/Renderer/Private/ReflectionEnvironment.cpp b/Engine/Source/Runtime/Renderer/Private/ReflectionEnvironment.cpp
--- a/Engine/Source/Runtime/Renderer/Private/ReflectionEnvironment.cpp
+++ b/Engine/Source/Runtime/Renderer/Private/ReflectionEnvironment.cpp
@@ -457,4 +457,5 @@ void FReflectionEnvironmentCubemapArray::UpdateMaxCubemaps(uint32 InMaxCubemaps,
 }
 
 IMPLEMENT_STATIC_UNIFORM_BUFFER_SLOT(ReflectionCapture);
-IMPLEMENT_STATIC_AND_SHADER_UNIFORM_BUFFER_STRUCT_EX(FReflectionCaptureShaderData, "ReflectionCapture", ReflectionCapture, FShaderParametersMetadata::EUsageFlags::NoEmulatedUniformBuffer);
+IMPLEMENT_STATIC_AND_SHADER_UNIFORM_BUFFER_STRUCT(FReflectionCaptureShaderData, "ReflectionCaptureSM5", ReflectionCapture);
+IMPLEMENT_STATIC_AND_SHADER_UNIFORM_BUFFER_STRUCT_EX(FMobileReflectionCaptureShaderData, "ReflectionCaptureES31", ReflectionCapture, FShaderParametersMetadata::EUsageFlags::NoEmulatedUniformBuffer);
diff --git a/Engine/Source/Runtime/Renderer/Private/ReflectionEnvironmentCapture.cpp b/Engine/Source/Runtime/Renderer/Private/ReflectionEnvironmentCapture.cpp
--- a/Engine/Source/Runtime/Renderer/Private/ReflectionEnvironmentCapture.cpp
+++ b/Engine/Source/Runtime/Renderer/Private/ReflectionEnvironmentCapture.cpp
@@ -459,7 +459,7 @@ int32 FindOrAllocateCubemapIndex(FScene* Scene, const UReflectionCaptureComponen
 		Scene->ReflectionSceneData.AllocatedReflectionCaptureState.Add(Component, FCaptureComponentSceneState(CubemapIndex));
 		Scene->ReflectionSceneData.AllocatedReflectionCaptureStateHasChanged = true;
 
-		check(CubemapIndex < GMaxNumReflectionCaptures);
+		check(CubemapIndex < GetMaxNumReflectionCaptures(Scene->GetShaderPlatform()));
 	}
 
 	check(CubemapIndex >= 0);
@@ -749,7 +749,7 @@ void FScene::AllocateReflectionCaptures(const TArray<UReflectionCaptureComponent
 	{
 		if (SupportsTextureCubeArray(GetFeatureLevel()))
 		{
-			int32_t PlatformMaxNumReflectionCaptures = FMath::Min(FMath::FloorToInt(GMaxTextureArrayLayers / 6.0f), GMaxNumReflectionCaptures);
+			int32_t PlatformMaxNumReflectionCaptures = FMath::Min(FMath::FloorToInt(GMaxTextureArrayLayers / 6.0f), GetMaxNumReflectionCaptures(GetShaderPlatform()));
 
 			for (int32 CaptureIndex = 0; CaptureIndex < NewCaptures.Num(); CaptureIndex++)
 			{
diff --git a/Engine/Source/Runtime/Renderer/Private/SceneRendering.cpp b/Engine/Source/Runtime/Renderer/Private/SceneRendering.cpp
--- a/Engine/Source/Runtime/Renderer/Private/SceneRendering.cpp
+++ b/Engine/Source/Runtime/Renderer/Private/SceneRendering.cpp
@@ -112,6 +112,11 @@ bool IsDynamicInstancingEnabled(ERHIFeatureLevel::Type FeatureLevel)
 		&& UseGPUScene(GMaxRHIShaderPlatform, FeatureLevel);
 }
 
+int32 GetMaxNumReflectionCaptures(EShaderPlatform ShaderPlatform)
+{
+	return IsMobilePlatform(ShaderPlatform) ? GMobileMaxNumReflectionCaptures : GMaxNumReflectionCaptures;
+}
+
 int32 GDumpInstancingStats = 0;
 FAutoConsoleVariableRef CVarDumpInstancingStats(
 	TEXT("r.MeshDrawCommands.LogDynamicInstancingStats"),
diff --git a/Engine/Source/Runtime/Renderer/Private/SceneRendering.h b/Engine/Source/Runtime/Renderer/Private/SceneRendering.h
--- a/Engine/Source/Runtime/Renderer/Private/SceneRendering.h
+++ b/Engine/Source/Runtime/Renderer/Private/SceneRendering.h
@@ -845,6 +845,21 @@ BEGIN_GLOBAL_SHADER_PARAMETER_STRUCT(FReflectionCaptureShaderData,)
 	SHADER_PARAMETER_ARRAY(FVector4f,BoxScales,[GMaxNumReflectionCaptures])
 END_GLOBAL_SHADER_PARAMETER_STRUCT()
 
+static const int32 GMobileMaxNumReflectionCaptures = 100;
+BEGIN_GLOBAL_SHADER_PARAMETER_STRUCT(FMobileReflectionCaptureShaderData, )
+	SHADER_PARAMETER_ARRAY(FVector4f,PositionAndRadius,[GMobileMaxNumReflectionCaptures])
+	// W is unused
+	SHADER_PARAMETER_ARRAY(FVector4f,TilePosition,[GMobileMaxNumReflectionCaptures])
+	// R is brightness, G is array index, B is shape
+	SHADER_PARAMETER_ARRAY(FVector4f,CaptureProperties,[GMobileMaxNumReflectionCaptures])
+	SHADER_PARAMETER_ARRAY(FVector4f,CaptureOffsetAndAverageBrightness,[GMobileMaxNumReflectionCaptures])
+	// Stores the box transform for a box shape, other data is packed for other shapes
+	SHADER_PARAMETER_ARRAY(FMatrix44f,BoxTransform,[GMobileMaxNumReflectionCaptures])
+	SHADER_PARAMETER_ARRAY(FVector4f,BoxScales,[GMobileMaxNumReflectionCaptures])
+END_GLOBAL_SHADER_PARAMETER_STRUCT()
+
+extern int32 GetMaxNumReflectionCaptures(EShaderPlatform ShaderPlatform);
+
 // Structure in charge of storing all information about TAA's history.
 struct FTemporalAAHistory
 {
@@ -1404,6 +1419,7 @@ public:
 	int32 NumSphereReflectionCaptures;
 	float FurthestReflectionCaptureDistance;
 	TUniformBufferRef<FReflectionCaptureShaderData> ReflectionCaptureUniformBuffer;
+	TUniformBufferRef<FMobileReflectionCaptureShaderData> MobileReflectionCaptureUniformBuffer;
 
 	// Sky / Atmosphere textures (transient owned by this view info) and pointer to constants owned by SkyAtmosphere proxy.
 	TRefCountPtr<IPooledRenderTarget> SkyAtmosphereCameraAerialPerspectiveVolume;
diff --git a/Engine/Source/Runtime/Renderer/Private/SceneVisibility.cpp b/Engine/Source/Runtime/Renderer/Private/SceneVisibility.cpp
--- a/Engine/Source/Runtime/Renderer/Private/SceneVisibility.cpp
+++ b/Engine/Source/Runtime/Renderer/Private/SceneVisibility.cpp
@@ -4984,12 +4984,10 @@ void FDeferredShadingSceneRenderer::InitViews(FRDGBuilder& GraphBuilder, const F
 	}
 }
 
-void FSceneRenderer::SetupSceneReflectionCaptureBuffer(FRHICommandListImmediate& RHICmdList)
+template<class T>
+void CreateReflectionCaptureUniformBuffer(const TArray<FReflectionCaptureSortData>& SortedCaptures, TUniformBufferRef<T>& OutReflectionCaptureUniformBuffer)
 {
-	FReflectionCaptureShaderData SamplePositionsBuffer;
-
-	const TArray<FReflectionCaptureSortData>& SortedCaptures = Scene->ReflectionSceneData.SortedCaptures;
-
+	T SamplePositionsBuffer;
 	for (int32 CaptureIndex = 0; CaptureIndex < SortedCaptures.Num(); CaptureIndex++)
 	{
 		SamplePositionsBuffer.PositionAndRadius[CaptureIndex] = FVector4f(SortedCaptures[CaptureIndex].RelativePosition, SortedCaptures[CaptureIndex].Radius);
@@ -5001,13 +4999,39 @@ void FSceneRenderer::SetupSceneReflectionCaptureBuffer(FRHICommandListImmediate&
 		SamplePositionsBuffer.BoxScales[CaptureIndex] = SortedCaptures[CaptureIndex].BoxScales;
 	}
 
-	auto ReflectionCaptureUniformBuffer = TUniformBufferRef<FReflectionCaptureShaderData>::CreateUniformBufferImmediate(SamplePositionsBuffer, UniformBuffer_SingleFrame);
+	OutReflectionCaptureUniformBuffer = TUniformBufferRef<T>::CreateUniformBufferImmediate(SamplePositionsBuffer, UniformBuffer_SingleFrame);
+}
+
+void FSceneRenderer::SetupSceneReflectionCaptureBuffer(FRHICommandListImmediate& RHICmdList)
+{
+	const TArray<FReflectionCaptureSortData>& SortedCaptures = Scene->ReflectionSceneData.SortedCaptures;
+
+	TUniformBufferRef<FMobileReflectionCaptureShaderData> MobileReflectionCaptureUniformBuffer;
+	TUniformBufferRef<FReflectionCaptureShaderData> ReflectionCaptureUniformBuffer;
+	bool bIsMobileClusteredReflections = IsMobilePlatform(ShaderPlatform) && MobileEnableClusteredReflections(ShaderPlatform);
+
+	if (bIsMobileClusteredReflections)
+	{
+		CreateReflectionCaptureUniformBuffer(SortedCaptures, MobileReflectionCaptureUniformBuffer);
+	}
+	else
+	{
+		CreateReflectionCaptureUniformBuffer(SortedCaptures, ReflectionCaptureUniformBuffer);
+	}
 
 	for (int32 ViewIndex = 0; ViewIndex < Views.Num(); ++ViewIndex)
 	{
 		FViewInfo& View = Views[ViewIndex];
 
-		View.ReflectionCaptureUniformBuffer = ReflectionCaptureUniformBuffer;
+		if (bIsMobileClusteredReflections)
+		{
+			View.MobileReflectionCaptureUniformBuffer = MobileReflectionCaptureUniformBuffer;
+		}
+		else
+		{
+			View.ReflectionCaptureUniformBuffer = ReflectionCaptureUniformBuffer;
+		}
+		
 		View.NumBoxReflectionCaptures = 0;
 		View.NumSphereReflectionCaptures = 0;
 		View.FurthestReflectionCaptureDistance = 0.0f;
