2m45s to 43s when booting the editor with cold cache due to fxc being fxc

#rnx

diff --git a/Engine/Shaders/Private/ClusteredDeferredShadingPixelShader.usf b/Engine/Shaders/Private/ClusteredDeferredShadingPixelShader.usf
--- a/Engine/Shaders/Private/ClusteredDeferredShadingPixelShader.usf
+++ b/Engine/Shaders/Private/ClusteredDeferredShadingPixelShader.usf
@@ -52,7 +52,8 @@
 #if STRATA_ENABLED
 	#define USE_PASS_PER_SHADING_MODEL 0
 #else
-	#define USE_PASS_PER_SHADING_MODEL 1
+	// Only enable on non-FXC compilers since FXC takes ages to compile all the shader model passes.
+	#define USE_PASS_PER_SHADING_MODEL (COMPILER_DXC == 1 || COMPILER_PSSL == 1 || COMPILER_METAL == 1)
 #endif
 
 // Extracts out the simple lights into their own loop that is executed first, independently of shading model.
