Fixed appearance to be better preserved when using verticl blending of top layer being translucent.

diff --git a/Engine/Shaders/Private/Strata/Strata.ush b/Engine/Shaders/Private/Strata/Strata.ush
--- a/Engine/Shaders/Private/Strata/Strata.ush
+++ b/Engine/Shaders/Private/Strata/Strata.ush
@@ -1879,13 +1879,14 @@ FStrataBSDF StrataVerticalLayeringParameterBlending(in FStrataBSDF TopBSDF, in F
 	FParticipatingMedia PM				= StrataSlabCreateParticipatingMedia(UnusedDiffuseColor, SLAB_SSSMFP(TopBSDF));
 	const float3 SlabTransmittanceV		= IsotropicMediumSlabTransmittance(PM, STRATA_SIMPLEVOLUME_THICKNESS_M, NoV);
 	const float3 SlabTransmittanceL		= IsotropicMediumSlabTransmittance(PM, STRATA_SIMPLEVOLUME_THICKNESS_M, 1.0f /*Normal incidence i.e., NoL == 1.f*/);
+	const float3 SlabTransmittance		= SlabTransmittanceV * SlabTransmittanceL;
 
 	// STRATA_TODO account for energy conservation here.
 
 	FVerticalLayeringInfo Info			= GetVerticalLayeringInfo(SafeTopCoverage, SafeBaseCoverage);
 	// Compute some ways to intepoldate and accumulate parameters, but always normalise because coverage is applied later in UpdateSingleBSDFOperatorCoverageTransmittance.
-	const float  TopToBottomLerp		= saturate((Info.TransmittanceTopAndBottom * dot(SlabTransmittanceV,0.33.xxx) + Info.SurfaceBottom) / max(STRATA_EPSILON, Info.Coverage));
-	const float3 BottomColorFactor		= (Info.TransmittanceTopAndBottom * SlabTransmittanceV + Info.SurfaceBottom) / max(STRATA_EPSILON, Info.Coverage);
+	const float  TopToBottomLerp		= saturate((Info.TransmittanceTopAndBottom * dot(SlabTransmittance,0.33.xxx) + Info.SurfaceBottom) / max(STRATA_EPSILON, Info.Coverage));
+	const float3 BottomColorFactor		= (Info.TransmittanceTopAndBottom * SlabTransmittance + Info.SurfaceBottom) / max(STRATA_EPSILON, Info.Coverage);
 	const float  TopFactor				= Info.SurfaceTop / max(STRATA_EPSILON, Info.Coverage);
 	
 	// Add both layers thickness
