Change FurnaceTest to loop over sample set.

This reduces shader compilation from 10s -> 3s

diff --git a/Engine/Shaders/Private/ShadingFurnaceTest.usf b/Engine/Shaders/Private/ShadingFurnaceTest.usf
--- a/Engine/Shaders/Private/ShadingFurnaceTest.usf
+++ b/Engine/Shaders/Private/ShadingFurnaceTest.usf
@@ -114,7 +114,7 @@ void MainPS(
 		float3 FinalRadiance = 0.f;
 
 		// Set loop (Cosine/GGX/... distribuciton)
-		UNROLL
+		LOOP
 		for (uint Set = 0; Set < NumSets; Set++)
 		{
 			// Samples loop
