Path Tracer: Remove RayCone from path tracer payload as it was not actually used

If we want to re-introduce this notion, we should be able to handle this directly in the material given the world position and camera matrix.

This frees up 4 bytes (in preparation for Strata support) and also reduces the path state structure size by 4 bytes.

I've measured around 3% speedup from this removal.

diff --git a/Engine/Plugins/Experimental/GPULightmass/Shaders/Private/LightmapPathTracing.usf b/Engine/Plugins/Experimental/GPULightmass/Shaders/Private/LightmapPathTracing.usf
--- a/Engine/Plugins/Experimental/GPULightmass/Shaders/Private/LightmapPathTracing.usf
+++ b/Engine/Plugins/Experimental/GPULightmass/Shaders/Private/LightmapPathTracing.usf
@@ -154,7 +154,7 @@ void AccumulateRadiance(inout float3 TotalRadiance, float3 PathRadiance)
 	}
 }
 
-FPathTracingPayload TraceTransparentRay(RayDesc Ray, FRayCone RayCone, bool IsCameraRay, bool LastBounce, bool IncludeEmission, uint2 LaunchIndex, uint NumLights, inout RandomSequence RandSequence, inout float3 PathThroughput, inout float3 Radiance)
+FPathTracingPayload TraceTransparentRay(RayDesc Ray, bool IsCameraRay, bool LastBounce, bool IncludeEmission, uint2 LaunchIndex, uint NumLights, inout RandomSequence RandSequence, inout float3 PathThroughput, inout float3 Radiance)
 {
 	const uint RayFlags = IsCameraRay && EnableCameraBackfaceCulling ? RAY_FLAG_CULL_BACK_FACING_TRIANGLES : 0;
 	const uint MissShaderIndex = 0;
@@ -173,7 +173,7 @@ FPathTracingPayload TraceTransparentRay(RayDesc Ray, FRayCone RayCone, bool IsCa
 		// also avoid noise from hot spots (they can be properly
 		// importance sampled if a capturing skylight is added)
 		const bool bIgnoreSky = !IsCameraRay;
-		FPackedPathTracingPayload PackedPayload = InitPathTracingPayload(RayCone, LaunchIndex, bIgnoreSky);
+		FPackedPathTracingPayload PackedPayload = InitPathTracingPayload(LaunchIndex, bIgnoreSky);
 		// Trace the ray
 		TraceRay(
 			TLAS,
@@ -346,9 +346,6 @@ void PathTracingKernel(
 	float LightPickingCdf[RAY_TRACING_LIGHT_COUNT_MAXIMUM];
 	
 	Ray.Direction = normalize(Ray.Direction);
-
-	FRayCone RayCone = (FRayCone)0;
-	RayCone.SpreadAngle = View.EyeToPixelSpreadAngle;
 	
 	// path state variables (these cary information between bounces)
 	float3 PathThroughput = 1.0;
@@ -368,7 +365,6 @@ void PathTracingKernel(
 		if (bIsCameraRay)
 		{
 			// GPULightmass fakes a 'camera' ray that hits the lightmap texel perpendicularly
-			Payload.RayCone = (FRayCone)0;
 			Payload.TranslatedWorldPos = TranslatedWorldPosition;
 			Payload.WorldNormal = ShadingNormal;
 			Payload.Radiance = float3(0, 0, 0);
@@ -389,7 +385,7 @@ void PathTracingKernel(
 		}
 		else
 		{
-			Payload = TraceTransparentRay(Ray, RayCone, false , bIsLastBounce, bIncludeEmissive, LaunchIndex, NumVisibleLights, RandSequence, PathThroughput, Radiance);
+			Payload = TraceTransparentRay(Ray, false , bIsLastBounce, bIncludeEmissive, LaunchIndex, NumVisibleLights, RandSequence, PathThroughput, Radiance);
 		}
 		
 		if (Payload.IsMiss())
@@ -740,9 +736,6 @@ void PathTracingKernel(
 		// from this point on, we don't need to include lights in the trace call
 		// because NEE handled it for us
 		NumVisibleLights = 0;
-		
-		// update ray cone for the next bounce
-		RayCone = Payload.RayCone;
 	}
 	
 #if USE_IRRADIANCE_CACHING
diff --git a/Engine/Shaders/Private/PathTracing/PathTracing.usf b/Engine/Shaders/Private/PathTracing/PathTracing.usf
--- a/Engine/Shaders/Private/PathTracing/PathTracing.usf
+++ b/Engine/Shaders/Private/PathTracing/PathTracing.usf
@@ -59,8 +59,6 @@ FPathState LoadPathStateData(uint Index)
 	Output.Ray.Direction = Data.RayDirection;
 	Output.Ray.TMin = 0;
 	Output.Ray.TMax = POSITIVE_INFINITY;
-	Output.RayCone.Width = f16tof32(Data.PackedRayCone);
-	Output.RayCone.SpreadAngle = f16tof32(Data.PackedRayCone >> 16);
 	Output.PathThroughput = Data.PathThroughput;
 	Output.PathRoughness = f16tof32(Data.PackedRoughnessSigma.x);
 	Output.SigmaT.x = f16tof32(Data.PackedRoughnessSigma.x >> 16);
@@ -83,7 +81,6 @@ void StorePathStateData(FPathState PathState, uint Index)
 	Data.PackedAlbedoNormal.z = f32tof16(PathState.Normal.y) | (f32tof16(PathState.Normal.z) << 16);
 	Data.RayOrigin = PathState.Ray.Origin;
 	Data.RayDirection = PathState.Ray.Direction;
-	Data.PackedRayCone = f32tof16(PathState.RayCone.Width) | (f32tof16(PathState.RayCone.SpreadAngle) << 16);
 	Data.PathThroughput = PathState.PathThroughput;
 	Data.PackedRoughnessSigma.x = f32tof16(PathState.PathRoughness) | (f32tof16(PathState.SigmaT.x) << 16);
 	Data.PackedRoughnessSigma.y = f32tof16(PathState.SigmaT.y) | (f32tof16(PathState.SigmaT.z) << 16);
diff --git a/Engine/Shaders/Private/PathTracing/PathTracingCommon.ush b/Engine/Shaders/Private/PathTracing/PathTracingCommon.ush
--- a/Engine/Shaders/Private/PathTracing/PathTracingCommon.ush
+++ b/Engine/Shaders/Private/PathTracing/PathTracingCommon.ush
@@ -46,17 +46,13 @@ float3 PayloadDecodeUnitVector(uint Encoded)
 struct FPackedPathTracingPayload : FMinimalPayload
 {
 	// float FMinimalPayload::HitT                       // 4  bytes
-	uint PackedRayCone;                                  // 4  bytes (fp16 width, fp16 spread)
 	uint PackedPixelCoord;                               // 4  bytes
-	uint PackedData[13];								 // 52 bytes (encoded data, depends on shading model and ray type)
+	uint PackedData[14];								 // 56 bytes (encoded data, depends on shading model and ray type)
 														 // 64 bytes total
 
 	void SetPixelCoord(uint2 PixelCoord) { PackedPixelCoord = (PixelCoord.x & 0xFFFF) | (PixelCoord.y << 16); }
 	uint2 GetPixelCoord() { return uint2(PackedPixelCoord & 0xFFFF, PackedPixelCoord >> 16); }
 
-	FRayCone GetRayCone() { return UnpackRayCone(PackedRayCone); }
-	void SetRayCone(FRayCone RayCone) { PackedRayCone = PackRayCone(RayCone); }
-
 	uint GetFlags() { return (PackedData[7] >> 23) & 0x3F; }
 	void SetFlag(uint Flag) { PackedData[7] |= Flag << 23; }
 	uint GetPrimitiveLightingChannelMask() { return (PackedData[7] >> 29) & 0x7; }
@@ -119,7 +115,6 @@ struct FPackedPathTracingPayload : FMinimalPayload
 // This payload structure is the expanded version of the above which is more convenient to work with
 struct FPathTracingPayload : FMinimalPayload
 {
-	FRayCone RayCone;
 	float3 Radiance;
 	float3 WorldGeoNormal;           // normal of the actual triangle (faceted)
 	float3 WorldNormal;              // normal output of the material (includes normal/bump)
@@ -190,10 +185,9 @@ struct FPathTracingPayload : FMinimalPayload
 	}
 };
 
-FPackedPathTracingPayload InitPathTracingPayload(FRayCone RayCone, uint2 PixelCoord, bool bIgnoreSkyMaterials)
+FPackedPathTracingPayload InitPathTracingPayload(uint2 PixelCoord, bool bIgnoreSkyMaterials)
 {
 	FPackedPathTracingPayload Output = (FPackedPathTracingPayload)0;
-	Output.PackedRayCone = PackRayCone(RayCone);
 	Output.SetPixelCoord(PixelCoord);
 	if (bIgnoreSkyMaterials)
 	{
@@ -220,7 +214,6 @@ FPackedPathTracingPayload PackPathTracingPayload(FPathTracingPayload Input)
 {
 	FPackedPathTracingPayload Output = (FPackedPathTracingPayload)0;
 	Output.HitT = Input.HitT;
-	Output.PackedRayCone = PackRayCone(Input.RayCone);
 
 	Output.PackedData[0] = f32tof16(Input.Radiance.x);
 	Output.PackedData[0] |= f32tof16(Input.Radiance.y) << 16;
@@ -273,7 +266,6 @@ FPathTracingPayload UnpackPathTracingPayload(FPackedPathTracingPayload Input, Ra
 	FPathTracingPayload Output = (FPathTracingPayload)0;
 
 	Output.HitT = Input.HitT;
-	Output.RayCone = Input.GetRayCone();
 	Output.Radiance.x = f16tof32(Input.PackedData[0]);
 	Output.Radiance.y = f16tof32(Input.PackedData[0] >> 16);
 	Output.Radiance.z = f16tof32(Input.PackedData[1]);
diff --git a/Engine/Shaders/Private/PathTracing/PathTracingCore.ush b/Engine/Shaders/Private/PathTracing/PathTracingCore.ush
--- a/Engine/Shaders/Private/PathTracing/PathTracingCore.ush
+++ b/Engine/Shaders/Private/PathTracing/PathTracingCore.ush
@@ -68,11 +68,10 @@ struct FPathState
 	float3 Albedo;                  //  6 bytes // half precision
 	float3 Normal;                  //  6 bytes // half precision
 	RayDesc Ray;                    // 12+12 = 24 bytes (TMin/TMax do not need to be stored)
-	FRayCone RayCone;               //  4 bytes // half precision
 	float3 PathThroughput;          // 12 bytes
 	float PathRoughness;            //  2 bytes // half precision
 	float3 SigmaT; // Extinction    //  6 bytes // half precision 
-									// 88 bytes total (see FPathTracingPackedPathState)
+									// 84 bytes total (see FPathTracingPackedPathState)
 
 	void WritePixel(uint2 TextureIndex)
 	{
@@ -132,7 +131,7 @@ FPathTracingPayload TraceTransparentRay(inout FPathState PathState, bool bIsCame
 		// also avoid noise from hot spots (they can be properly
 		// importance sampled if a capturing skylight is added)
 		const bool bIgnoreSky = !bIsCameraRay;
-		FPackedPathTracingPayload PackedPayload = InitPathTracingPayload(PathState.RayCone, PathState.PixelIndex, bIgnoreSky);
+		FPackedPathTracingPayload PackedPayload = InitPathTracingPayload(PathState.PixelIndex, bIgnoreSky);
 		// Trace the ray
 		TraceRay(
 			TLAS,
@@ -149,7 +148,6 @@ FPathTracingPayload TraceTransparentRay(inout FPathState PathState, bool bIsCame
 			// we didn't hit any real geometry, but we did hit the volume's blocker geometry
 			// create a virtual hit with it here
 			PackedPayload = VolumeGetBlockerHit(VolumeIsectList.BlockerID, PathState.Ray.Origin, PathState.Ray.Direction, VolumeIsectList.BlockerHitT);
-			PackedPayload.SetRayCone(PropagateRayCone(PathState.RayCone, 0.0, VolumeIsectList.BlockerHitT));
 		}
 
 		// Loop over lights to capture their contribution
@@ -167,7 +165,6 @@ FPathTracingPayload TraceTransparentRay(inout FPathState PathState, bool bIsCame
 				//       but is not consistent with how reflections behave and the sorting order will be arbitrary
 				FPathTracingPayload LightHitPayload = (FPathTracingPayload)0; // clear all fields
 				LightHitPayload.HitT = LightHit.HitT;
-				LightHitPayload.RayCone = PropagateRayCone(PathState.RayCone, 0.0, LightHit.HitT);
 				LightHitPayload.Radiance = LightHit.Radiance;
 				LightHitPayload.ShadingModelID = SHADINGMODELID_UNLIT;
 				// alpha composite assumes premult, which allows emissive hits that are fully transparent
@@ -273,7 +270,6 @@ FPathTracingPayload TraceTransparentRay(inout FPathState PathState, bool bIsCame
 							{
 								// stash this hit for next time
 								Payload = CreateMediumHitPayload(Interval.VolumeTMin, TranslatedWorldPos, Result);
-								Payload.RayCone = PropagateRayCone(PathState.RayCone, 0.0, Interval.VolumeTMin);
 								PayloadThroughput = Contrib / SelectionWeight;
 							}
 						}
@@ -447,7 +443,7 @@ struct FProbeResult
 FProbeResult TraceProbeRay(RayDesc Ray, uint2 PixelCoord)
 {
 	// Trace a short ray to see if we escaped the surface
-	FPackedPathTracingPayload PackedPayload = InitPathTracingPayload((FRayCone)0, PixelCoord, true);
+	FPackedPathTracingPayload PackedPayload = InitPathTracingPayload(PixelCoord, true);
 	// TODO: use smaller payload + flag to skip shading work? (only want smooth+geo normals)
 	//       however supporting blockers would require some shading ...
 	const uint RayFlags = 0;
@@ -729,8 +725,6 @@ FPathState CreatePathState(int2 PixelIndex)
 		Output.Ray.Origin = Origin + Output.Ray.Direction * (View.NearPlane * rcp(dot(ViewZ, Output.Ray.Direction)));
 	}
 
-	Output.RayCone.SpreadAngle = View.EyeToPixelSpreadAngle;
-
 	// path state variables (these cary information between bounces)
 	Output.PathThroughput = 1.0;
 	Output.PathRoughness = 0;
@@ -1183,9 +1177,6 @@ bool PathTracingKernel(inout FPathState PathState, int Bounce)
 
 	}
 
-	// update ray cone for the next bounce
-	PathState.RayCone = Payload.RayCone;
-
 	// this path should continue bouncing
 	return true;
 }
diff --git a/Engine/Shaders/Private/PathTracing/PathTracingMaterialHitShader.usf b/Engine/Shaders/Private/PathTracing/PathTracingMaterialHitShader.usf
--- a/Engine/Shaders/Private/PathTracing/PathTracingMaterialHitShader.usf
+++ b/Engine/Shaders/Private/PathTracing/PathTracingMaterialHitShader.usf
@@ -44,8 +44,6 @@ RAY_TRACING_ENTRY_CLOSEST_HIT(PathTracingMaterialCHS,
 
 	float4 SvPosition = float4(PixelCoord.xy, 0.0, 1.0);
 
-	FRayCone PropagatedCone = PropagateRayCone(PackedPayload.GetRayCone(), 0 /* surface curvature */, RayTCurrent());
-
 #if VF_SUPPORTS_RAYTRACING_PREPARE_MATERIAL_PIXEL_PARAMETERS
 	// this is a newer codepath that is both more flexible and allows for more direct calculation compared to the other codepath
 	// TODO: implement such a method for all vertex factories
@@ -54,7 +52,7 @@ RAY_TRACING_ENTRY_CLOSEST_HIT(PathTracingMaterialCHS,
 #else
 	FVertexFactoryInterpolantsVSToPS Interpolants;
 	float3 GeoNormal = 0;
-	CalcInterpolants(PropagatedCone, Attributes, Interpolants, GeoNormal);
+	CalcInterpolants((FRayCone)0, Attributes, Interpolants, GeoNormal);
 
 	FMaterialPixelParameters MaterialParameters = GetMaterialPixelParameters(Interpolants, SvPosition);
 #endif
@@ -75,7 +73,6 @@ RAY_TRACING_ENTRY_CLOSEST_HIT(PathTracingMaterialCHS,
 	}
 
 	FPathTracingPayload Payload = (FPathTracingPayload)0;
-	Payload.RayCone = PropagatedCone;
 
 	/**
 	 * Set common material attributes for both full and simplified materials
@@ -381,7 +378,7 @@ RAY_TRACING_ENTRY_ANY_HIT(PathTracingMaterialAHS,
 	FVertexFactoryInterpolantsVSToPS Interpolants;
 	float4 SvPosition = float4(PixelCoord.xy, 0.0, 1.0);
 
-	CalcInterpolants(PackedPayload.GetRayCone(), Attributes, Interpolants);
+	CalcInterpolants((FRayCone)0, Attributes, Interpolants);
 
 	FMaterialPixelParameters MaterialParameters = GetMaterialPixelParameters(Interpolants, SvPosition);
 	FPixelMaterialInputs PixelMaterialInputs;
diff --git a/Engine/Shaders/Private/RayTracing/RayTracingCalcInterpolants.ush b/Engine/Shaders/Private/RayTracing/RayTracingCalcInterpolants.ush
--- a/Engine/Shaders/Private/RayTracing/RayTracingCalcInterpolants.ush
+++ b/Engine/Shaders/Private/RayTracing/RayTracingCalcInterpolants.ush
@@ -40,7 +40,7 @@ void CalcInterpolants(in FRayCone RayCone, in FRayTracingIntersectionAttributes
 	float InvWorldArea = rsqrt(dot(Unnormalized, Unnormalized));
 	FaceNormal = Unnormalized * InvWorldArea;
 
-#if (NUM_TEX_COORD_INTERPOLATORS || USE_PARTICLE_SUBUVS) && !VERTEX_FACTORY_MODIFIES_INTERPOLATION
+#if USE_RAYTRACED_TEXTURE_RAYCONE_LOD && (NUM_TEX_COORD_INTERPOLATORS || USE_PARTICLE_SUBUVS) && !VERTEX_FACTORY_MODIFIES_INTERPOLATION
 	float2 TA = TexCoords[1] - TexCoords[0];
 	float2 TB = TexCoords[2] - TexCoords[0];
 	float UvArea = abs((TA.x * TB.y) - (TA.y * TB.x));
diff --git a/Engine/Shaders/Shared/RayTracingTypes.h b/Engine/Shaders/Shared/RayTracingTypes.h
--- a/Engine/Shaders/Shared/RayTracingTypes.h
+++ b/Engine/Shaders/Shared/RayTracingTypes.h
@@ -45,13 +45,12 @@ struct FPathTracingPackedPathState {
 	uint16    Normal[3];
 	FVector3f RayOrigin;
 	FVector3f RayDirection;
-	uint32    RayCone;
 	FVector3f PathThroughput;
 	uint16    PathRoughness;
 	uint16    SigmaT[3];
 };
 
-static_assert(sizeof(FPathTracingPackedPathState) == 88, "Packed Path State size should be minimized");
+static_assert(sizeof(FPathTracingPackedPathState) == 84, "Packed Path State size should be minimized");
 
 #else
 
@@ -87,7 +86,6 @@ struct FPathTracingPackedPathState {
 	uint3     PackedAlbedoNormal;
 	float3    RayOrigin;
 	float3    RayDirection;
-	uint      PackedRayCone;
 	float3    PathThroughput;
 	uint2     PackedRoughnessSigma;
 };
