Strata - clustered shading now works with opaque rough refraction.

diff --git a/Engine/Shaders/Private/ClusteredDeferredShadingPixelShader.usf b/Engine/Shaders/Private/ClusteredDeferredShadingPixelShader.usf
--- a/Engine/Shaders/Private/ClusteredDeferredShadingPixelShader.usf
+++ b/Engine/Shaders/Private/ClusteredDeferredShadingPixelShader.usf
@@ -168,6 +168,8 @@ float4 GetSimpleLightLighting(
 #if STRATA_ENABLED && USE_HAIR_LIGHTING == 0
 	FStrataAddressing StrataAddressing,
 	FStrataPixelHeader StrataPixelHeader,
+	inout float3 OutOpaqueRoughRefractionSceneColor,
+	inout float3 OutSubSurfaceSceneColor,
 #else
 	FScreenSpaceData ScreenSpaceData,
 #endif 
@@ -210,7 +212,7 @@ float4 GetSimpleLightLighting(
 		const float3 V = -CameraVector;
 		const float LightMask = GetLocalLightAttenuation(TranslatedWorldPosition, LightData, ToLight, L);
 
-		CompositedLighting += StrataDeferredLighting(
+		FStrataDeferredLighting StrataLighting = StrataDeferredLighting(
 			LightData,
 			V,
 			L,
@@ -220,14 +222,29 @@ float4 GetSimpleLightLighting(
 			Strata.MaterialTextureArray,
 			StrataAddressing,
 			StrataPixelHeader).SceneColor;
+
+		CompositedLighting += StrataLighting.SceneColor;
+	#if STRATA_OPAQUE_ROUGH_REFRACTION_ENABLED
+		OutOpaqueRoughRefractionSceneColor += StrataLighting.OpaqueRoughRefractionSceneColor;
+		OutSubSurfaceSceneColor += StrataLighting.SubSurfaceSceneColor;
+	#endif
+
 	#else
+
 		FSimpleDeferredLightData SimpleLightData = (FSimpleDeferredLightData)0;
 		SimpleLightData.TranslatedWorldPosition = LocalLight.LightPositionAndInvRadius.xyz;
 		SimpleLightData.InvRadius = LocalLight.LightPositionAndInvRadius.w;
 		SimpleLightData.Color = LocalLight.LightColorAndFalloffExponent.xyz;
 		SimpleLightData.FalloffExponent = LocalLight.LightColorAndFalloffExponent.w;
 		SimpleLightData.bInverseSquared = SimpleLightData.FalloffExponent == 0;
-		CompositedLighting += GetSimpleDynamicLighting(TranslatedWorldPosition, CameraVector, ScreenSpaceData.GBuffer.WorldNormal, ScreenSpaceData.AmbientOcclusion, ScreenSpaceData.GBuffer.DiffuseColor, ScreenSpaceData.GBuffer.SpecularColor, ScreenSpaceData.GBuffer.Roughness, SimpleLightData);
+		FStrataDeferredLighting StrataLighting = GetSimpleDynamicLighting(TranslatedWorldPosition, CameraVector, ScreenSpaceData.GBuffer.WorldNormal, ScreenSpaceData.AmbientOcclusion, ScreenSpaceData.GBuffer.DiffuseColor, ScreenSpaceData.GBuffer.SpecularColor, ScreenSpaceData.GBuffer.Roughness, SimpleLightData);
+
+		CompositedLighting += StrataLighting.SceneColor;
+	#if STRATA_OPAQUE_ROUGH_REFRACTION_ENABLED
+		OutOpaqueRoughRefractionSceneColor += StrataLighting.OpaqueRoughRefractionSceneColor;
+		OutSubSurfaceSceneColor += StrataLighting.SubSurfaceSceneColor;
+	#endif
+
 	#endif
 	}
 	FirstNonSimpleLightIndex = LocalLightListIndex;
@@ -244,6 +261,8 @@ float4 GetLightGridLocalLighting(
 #if STRATA_ENABLED && USE_HAIR_LIGHTING == 0
 	FStrataAddressing StrataAddressing,
 	FStrataPixelHeader StrataPixelHeader,
+	inout float3 OutOpaqueRoughRefractionSceneColor,
+	inout float3 OutSubSurfaceSceneColor,
 #else 
 	FScreenSpaceData ScreenSpaceData, 
 #endif
@@ -360,7 +379,7 @@ float4 GetLightGridLocalLighting(
 		GetShadowTerms(SceneDepth, PrecomputedShadowFactors, FakeShadingModelID, FakeContactShadowOpacity, LightData, TranslatedWorldPosition, L, LightAttenuation, Dither, ShadowTerms);
 
 		const float LightMask = 1;
-		DirectLighting += StrataDeferredLighting(
+		FStrataDeferredLighting StrataLighting = StrataDeferredLighting(
 			LightData,
 			V,
 			L,
@@ -369,7 +388,14 @@ float4 GetLightGridLocalLighting(
 			ShadowTerms,
 			Strata.MaterialTextureArray,
 			StrataAddressing,
-			StrataPixelHeader).SceneColor;
+			StrataPixelHeader);
+
+		DirectLighting += StrataLighting.SceneColor;
+	#if STRATA_OPAQUE_ROUGH_REFRACTION_ENABLED
+		OutOpaqueRoughRefractionSceneColor += StrataLighting.OpaqueRoughRefractionSceneColor;
+		OutSubSurfaceSceneColor += StrataLighting.SubSurfaceSceneColor;
+	#endif
+
 	#elif 1
 		// NOTE: uses AO data like tiled deferred (as opposed to forward path)
 		// TODO: Passing in 0,0 as SVPos, what is meaning of this? Only used if 'REFERENCE_QUALITY' is on.
@@ -418,7 +444,12 @@ void ClusteredShadingPixelShader(
 	float2 ScreenUV : TEXCOORD0,
 	float3 ScreenVector : TEXCOORD1,
 	in float4 SvPosition : SV_Position,
-	out float4 OutColor : SV_Target0)
+	out float4 OutColor : SV_Target0
+#if STRATA_OPAQUE_ROUGH_REFRACTION_ENABLED
+	, out float3 OutOpaqueRoughRefractionSceneColor : SV_Target1
+	, out float3 OutSubSurfaceSceneColor : SV_Target2
+#endif
+	)
 {
 	uint2 PixelPos = SvPosition.xy;
 
@@ -439,6 +470,8 @@ void ClusteredShadingPixelShader(
 	FStrataPixelHeader StrataPixelHeader = UnpackStrataHeaderIn(Strata.MaterialTextureArray, StrataAddressing, Strata.TopLayerTexture);
 	
 	float4 CompositedLighting = 0;
+	float3 OpaqueRoughRefractionSceneColor = 0;
+	float3 SubSurfaceSceneColor = 0;
 
 	BRANCH
 	if (StrataPixelHeader.BSDFCount > 0 || CulledLightGridData.NumLocalLights == 0)
@@ -449,14 +482,20 @@ void ClusteredShadingPixelShader(
 		uint FirstNonSimpleLightIndex = 0;
 
 		// Simple lights
-		CompositedLighting += GetSimpleLightLighting(StrataAddressing, StrataPixelHeader, CulledLightGridData, TranslatedWorldPosition, CameraVector, EyeIndex, FirstNonSimpleLightIndex);
+		CompositedLighting += GetSimpleLightLighting(StrataAddressing, StrataPixelHeader, OpaqueRoughRefractionSceneColor, SubSurfaceSceneColor,
+			CulledLightGridData, TranslatedWorldPosition, CameraVector, EyeIndex, FirstNonSimpleLightIndex);
 
 		// Regular lights
-		CompositedLighting += GetLightGridLocalLighting(StrataAddressing, StrataPixelHeader, CulledLightGridData, TranslatedWorldPosition, CameraVector, ScreenUV, EyeIndex, Dither, FirstNonSimpleLightIndex);
+		CompositedLighting += GetLightGridLocalLighting(StrataAddressing, StrataPixelHeader, OpaqueRoughRefractionSceneColor, SubSurfaceSceneColor,
+			CulledLightGridData, TranslatedWorldPosition, CameraVector, ScreenUV, EyeIndex, Dither, FirstNonSimpleLightIndex);
 	}
 	
 #if !VISUALIZE_LIGHT_CULLING
 	CompositedLighting *= View.PreExposure;
+#if STRATA_OPAQUE_ROUGH_REFRACTION_ENABLED
+	OutOpaqueRoughRefractionSceneColor = OpaqueRoughRefractionSceneColor * View.PreExposure;
+	OutSubSurfaceSceneColor = SubSurfaceSceneColor * View.PreExposure;
+#endif
 #endif
 	OutColor = CompositedLighting;
 }
diff --git a/Engine/Source/Runtime/Renderer/Private/ClusteredDeferredShadingPass.cpp b/Engine/Source/Runtime/Renderer/Private/ClusteredDeferredShadingPass.cpp
--- a/Engine/Source/Runtime/Renderer/Private/ClusteredDeferredShadingPass.cpp
+++ b/Engine/Source/Runtime/Renderer/Private/ClusteredDeferredShadingPass.cpp
@@ -24,9 +24,9 @@
 #include "ShaderPrintParameters.h"
 #include "ShaderPrint.h"
 #include "VirtualShadowMaps/VirtualShadowMapArray.h"
-
 #include "SceneFilterRendering.h"
 #include "PostProcessing.h"
+#include "Strata/Strata.h"
 
 // This is used to switch on and off the clustered deferred shading implementation, that uses the light grid to perform shading.
 int32 GUseClusteredDeferredShading = 0;
@@ -148,7 +148,8 @@ static void InternalAddClusteredDeferredShadingPass(
 	EStrataTileType TileType,
 	FRDGTextureRef ShadowMaskBits,
 	FVirtualShadowMapArray& VirtualShadowMapArray,
-	FRDGBufferSRVRef HairTransmittanceBuffer)
+	FRDGBufferSRVRef HairTransmittanceBuffer,
+	FStrataSceneData* StrataSceneData)
 {
 	check(SortedLightsSet.ClusteredSupportedEnd > 0);
 	const FIntPoint SceneTextureExtent = SceneTextures.Config.Extent;
@@ -171,6 +172,12 @@ static void InternalAddClusteredDeferredShadingPass(
 	{
 		PassParameters->RenderTargets[0] = FRenderTargetBinding(View.HairStrandsViewData.VisibilityData.SampleLightingTexture, ERenderTargetLoadAction::ELoad);
 	}
+	if (Strata::IsStrataOpaqueMaterialRoughRefractionEnabled())
+	{
+		check(StrataSceneData);
+		PassParameters->RenderTargets[1] = FRenderTargetBinding(StrataSceneData->SeparatedOpaqueRoughRefractionSceneColor, ERenderTargetLoadAction::ELoad);
+		PassParameters->RenderTargets[2] = FRenderTargetBinding(StrataSceneData->SeparatedSubSurfaceSceneColor, ERenderTargetLoadAction::ELoad);
+	}
 
 	// VS - Strata tile parameters
 	EPrimitiveType PrimitiveType = PT_TriangleList;
@@ -201,7 +208,17 @@ static void InternalAddClusteredDeferredShadingPass(
 			InRHICmdList.ApplyCachedRenderTargets(GraphicsPSOInit);
 
 			// Additive blend to accumulate lighting contributions.
-			GraphicsPSOInit.BlendState = TStaticBlendState<CW_RGBA, BO_Add, BF_One, BF_One, BO_Add, BF_One, BF_One>::GetRHI();
+			if (Strata::IsStrataOpaqueMaterialRoughRefractionEnabled())
+			{
+				GraphicsPSOInit.BlendState = TStaticBlendState<
+					CW_RGBA, BO_Add, BF_One, BF_One, BO_Add, BF_One, BF_One,
+					CW_RGBA, BO_Add, BF_One, BF_One, BO_Add, BF_One, BF_One,
+					CW_RGBA, BO_Add, BF_One, BF_One, BO_Add, BF_One, BF_One>::GetRHI();
+			}
+			else
+			{
+				GraphicsPSOInit.BlendState = TStaticBlendState<CW_RGBA, BO_Add, BF_One, BF_One, BO_Add, BF_One, BF_One>::GetRHI();
+			}
 
 			GraphicsPSOInit.RasterizerState = TStaticRasterizerState<FM_Solid, CM_None>::GetRHI();
 			GraphicsPSOInit.DepthStencilState = TStaticDepthStencilState<false, CF_Always>::GetRHI();
@@ -265,8 +282,12 @@ void FDeferredShadingSceneRenderer::AddClusteredDeferredShadingPass(
 		{
 			FViewInfo& View = Views[ViewIndex];
 
+			FStrataSceneData* StrataSceneData = nullptr;
+
 			if (Strata::IsStrataEnabled())
 			{
+				StrataSceneData = &Scene->StrataSceneData;
+
 				InternalAddClusteredDeferredShadingPass(
 					GraphBuilder,
 					View,
@@ -276,7 +297,8 @@ void FDeferredShadingSceneRenderer::AddClusteredDeferredShadingPass(
 					EStrataTileType::ESimple,
 					ShadowMaskBits,
 					VirtualShadowMapArray,
-					nullptr);
+					nullptr, 
+					StrataSceneData);
 
 				InternalAddClusteredDeferredShadingPass(
 					GraphBuilder,
@@ -287,7 +309,8 @@ void FDeferredShadingSceneRenderer::AddClusteredDeferredShadingPass(
 					EStrataTileType::ESingle,
 					ShadowMaskBits,
 					VirtualShadowMapArray,
-					nullptr);
+					nullptr,
+					StrataSceneData);
 
 				InternalAddClusteredDeferredShadingPass(
 					GraphBuilder,
@@ -298,7 +321,8 @@ void FDeferredShadingSceneRenderer::AddClusteredDeferredShadingPass(
 					EStrataTileType::EComplex,
 					ShadowMaskBits,
 					VirtualShadowMapArray,
-					nullptr);
+					nullptr,
+					StrataSceneData);
 			}
 			else
 			{
@@ -311,7 +335,8 @@ void FDeferredShadingSceneRenderer::AddClusteredDeferredShadingPass(
 					EStrataTileType::ECount,
 					ShadowMaskBits,
 					VirtualShadowMapArray,
-					nullptr);
+					nullptr,
+					StrataSceneData);
 			}
 
 			if (HairStrands::HasViewHairStrandsData(View))
@@ -326,7 +351,8 @@ void FDeferredShadingSceneRenderer::AddClusteredDeferredShadingPass(
 					EStrataTileType::ECount,
 					HairStrandsShadowMaskBits,
 					VirtualShadowMapArray,
-					GraphBuilder.CreateSRV(TransmittanceMask.TransmittanceMask, FHairStrandsTransmittanceMaskData::Format));
+					GraphBuilder.CreateSRV(TransmittanceMask.TransmittanceMask, FHairStrandsTransmittanceMaskData::Format),
+					StrataSceneData);
 			}
 		}
 	}
