Fix potential nan in fog calculations.
Seen when doing offscreen tile mesh rendering with SM6.

diff --git a/Engine/Shaders/Private/HeightFogCommon.ush b/Engine/Shaders/Private/HeightFogCommon.ush
--- a/Engine/Shaders/Private/HeightFogCommon.ush
+++ b/Engine/Shaders/Private/HeightFogCommon.ush
@@ -211,9 +211,10 @@ half4 CalculateHeightFog(float3 WorldPositionRelativeToCamera)
 	float ExcludeDistance = 0;
 
 #if PERMUTATION_SUPPORT_VOLUMETRIC_FOG
-	float InvCosAngle = length(WorldPositionRelativeToCamera) / dot(WorldPositionRelativeToCamera, View.ViewForward);
 	// Volumetric fog covers up to MaxDistance along ViewZ, exclude analytical fog from this range
-	ExcludeDistance = max(View.VolumetricFogMaxDistance * InvCosAngle, 0);
+	float CosAngle = dot(normalize(WorldPositionRelativeToCamera), View.ViewForward);
+	float InvCosAngle = (CosAngle > FLT_EPSILON) ? rcp(CosAngle) : 0;
+	ExcludeDistance = View.VolumetricFogMaxDistance * InvCosAngle;
 #endif
 
 	float4 FogInscatteringAndOpacity = GetExponentialHeightFog(WorldPositionRelativeToCamera, ExcludeDistance);
