Shader compilation fix

- fix for shader compilation on an NDA'd platform introduced in CL 20008404

#rnx

diff --git a/Engine/Shaders/Private/PathTracing/Material/PathTracingRadianceProbe.ush b/Engine/Shaders/Private/PathTracing/Material/PathTracingRadianceProbe.ush
--- a/Engine/Shaders/Private/PathTracing/Material/PathTracingRadianceProbe.ush
+++ b/Engine/Shaders/Private/PathTracing/Material/PathTracingRadianceProbe.ush
@@ -55,7 +55,7 @@ FMaterialSample RadianceProbe_SampleMaterial(
 {
 	float PdfAdjust = 1.0f; 
 
-	float2 Point = RandSample.yx;
+	float2 SamplePoint = RandSample.yx;
 
 #ifdef LIGHTMAP_PATH_TRACING_MAIN_RG
 #if USE_FIRST_BOUNCE_RAY_GUIDING
@@ -88,9 +88,9 @@ FMaterialSample RadianceProbe_SampleMaterial(
 			LastRowPrefixSum = RowPrefixSum;
 			RowPrefixSum = RayGuidingCDFY[BatchedTiles[TileIndex].WorkingSetPosition / GPreviewLightmapPhysicalTileSize * CDF_TILE_SIZE / 4 + FinalPosition];
 			
-			if (RowPrefixSum > Point.y)
+			if (RowPrefixSum > SamplePoint.y)
 			{
-				Point.y = (y + (Point.y - LastRowPrefixSum) / (RowPrefixSum - LastRowPrefixSum)) / DIRECTIONAL_BINS_ONE_DIM;
+				SamplePoint.y = (y + (SamplePoint.y - LastRowPrefixSum) / (RowPrefixSum - LastRowPrefixSum)) / DIRECTIONAL_BINS_ONE_DIM;
 				
 				PdfAdjust *= (RowPrefixSum - LastRowPrefixSum) * DIRECTIONAL_BINS_ONE_DIM;
 				
@@ -110,9 +110,9 @@ FMaterialSample RadianceProbe_SampleMaterial(
 			LastPrefixSum = PrefixSum;
 			PrefixSum = RayGuidingCDFX[BatchedTiles[TileIndex].WorkingSetPosition / GPreviewLightmapPhysicalTileSize * CDF_TILE_SIZE + FinalPosition];
 			
-			if (PrefixSum > Point.x)
+			if (PrefixSum > SamplePoint.x)
 			{
-				Point.x = (x + (Point.x - LastPrefixSum) / (PrefixSum - LastPrefixSum)) / DIRECTIONAL_BINS_ONE_DIM;									
+				SamplePoint.x = (x + (SamplePoint.x - LastPrefixSum) / (PrefixSum - LastPrefixSum)) / DIRECTIONAL_BINS_ONE_DIM;									
 				PdfAdjust *= (PrefixSum - LastPrefixSum) * DIRECTIONAL_BINS_ONE_DIM;
 				break;
 			}
@@ -125,10 +125,10 @@ FMaterialSample RadianceProbe_SampleMaterial(
 
 #ifdef LIGHTMAP_PATH_TRACING_MAIN_RG
 	// Use cosine hemisphere sampling for surface lightmaps
-	float4 SampledValue = CosineSampleHemisphereConcentric(Point);
+	float4 SampledValue = CosineSampleHemisphereConcentric(SamplePoint);
 #else
 	// Use uniform hemisphere sampling for volumetric lightmaps (the two hemispheres are meant to be merged)
-	float4 SampledValue = UniformSampleHemisphere(Point);
+	float4 SampledValue = UniformSampleHemisphere(SamplePoint);
 #endif
 	float3 OutDirection = TangentToWorld(SampledValue.xyz, N_World);
 	
