fix for HDR in editor (scRGB): we were applying AP1->sRGB color space conversion twice in case of scRGB: once during the LUT computation, and once during the tonemapping through ST2084ToScRGB. Since scRGB might contain negative values and I didn't want to change the format of the LUT, it's simpler to store AP1 encoded values in the LUT and keep applying ST2084ToScRGB during the tonemapping: the RT format will then be 16F, allowing proper negative values.

Thanks to [REDACTED_HANDLE] for spotting this in the first place !

diff --git a/Engine/Shaders/Private/PostProcessCombineLUTs.usf b/Engine/Shaders/Private/PostProcessCombineLUTs.usf
--- a/Engine/Shaders/Private/PostProcessCombineLUTs.usf
+++ b/Engine/Shaders/Private/PostProcessCombineLUTs.usf
@@ -372,8 +372,12 @@ float4 CombineLUTsCommon(float2 InUV, uint InLayerIndex)
 		// 1000 nit ODT
 		float3 ODTColor = ACESOutputTransforms1000( GradedColor );
 
-		// Convert from AP1 to specified output gamut
-		ODTColor = mul( AP1_2_Output, ODTColor );
+		// Convert from AP1 to specified output gamut except for ScRGB as it might contain negative values. 
+		// In this case, the AP1->output conversion happens through ST2084ToScRGB in the tonemap function
+		if( GetOutputDevice() != TONEMAPPER_OUTPUT_ACES1000nitScRGB )
+		{
+			ODTColor = mul(AP1_2_Output, ODTColor);
+		}
 
 		// Apply conversion to ST-2084 (Dolby PQ)
 		OutDeviceColor = LinearToST2084( ODTColor );
@@ -385,8 +389,12 @@ float4 CombineLUTsCommon(float2 InUV, uint InLayerIndex)
 		// 2000 nit ODT
 		float3 ODTColor = ACESOutputTransforms2000( GradedColor );
 
-		// Convert from AP1 to specified output gamut
-		ODTColor = mul( AP1_2_Output, ODTColor );
+		// Convert from AP1 to specified output gamut except for ScRGB as it might contain negative values. 
+		// In this case, the AP1->output conversion happens through ST2084ToScRGB in the tonemap function
+		if ( GetOutputDevice() != TONEMAPPER_OUTPUT_ACES2000nitScRGB)
+		{
+			ODTColor = mul(AP1_2_Output, ODTColor);
+		}
 
 		// Apply conversion to ST-2084 (Dolby PQ)
 		OutDeviceColor = LinearToST2084( ODTColor );
