Fix crash when enabling Strata & RT Skylighting at the same time.

Sky lighting is not supported with Strata

diff --git a/Engine/Shaders/Private/RayTracing/RayTracingSkyLightRGS.usf b/Engine/Shaders/Private/RayTracing/RayTracingSkyLightRGS.usf
--- a/Engine/Shaders/Private/RayTracing/RayTracingSkyLightRGS.usf
+++ b/Engine/Shaders/Private/RayTracing/RayTracingSkyLightRGS.usf
@@ -23,6 +23,10 @@ RAY_TRACING_ENTRY_RAYGEN(SkyLightRGS)
 	uint2 DispatchThreadId = DispatchRaysIndex().xy + View.ViewRectMin.xy;
 	uint2 PixelCoord = GetPixelCoord(DispatchThreadId, UpscaleFactor);
 
+#if STRATA_ENABLED
+	RWSkyOcclusionMaskUAV[DispatchThreadId] = float4(0,0,0,1);
+	RWSkyOcclusionRayDistanceUAV[DispatchThreadId] = float2(200000, 0);
+#else // STRATA_ENABLED
 	// Get G-Buffer surface data
 	float2 InvBufferSize = View.BufferSizeAndInvSize.zw;
 	float2 UV = (float2(PixelCoord) + 0.5) * InvBufferSize;
@@ -94,4 +98,5 @@ RAY_TRACING_ENTRY_RAYGEN(SkyLightRGS)
 
 	RWSkyOcclusionMaskUAV[DispatchThreadId] = float4(ClampToHalfFloatRange(DiffuseExitantRadiance.rgb), AmbientOcclusion);
 	RWSkyOcclusionRayDistanceUAV[DispatchThreadId] = float2(HitDistance, SamplesPerPixel);
+#endif // STRATA_ENABLED
 }
\ No newline at end of file
