Use half precision for world vertex normal interpolators on mobile

diff --git a/Engine/Shaders/Private/Common.ush b/Engine/Shaders/Private/Common.ush
--- a/Engine/Shaders/Private/Common.ush
+++ b/Engine/Shaders/Private/Common.ush
@@ -571,17 +571,9 @@ MaterialFloat4 TextureCubeArraySampleLevel(TextureCubeArray Tex, SamplerState Sa
 #endif
 
 #else
-
 	#define TANGENTTOWORLD0					TEXCOORD10
 	#define TANGENTTOWORLD2					TEXCOORD11
-
-	// TangentToWorld0 is float4 only to match D3D11
-#if METAL_PROFILE || COMPILER_GLSL_ES3_1 || VULKAN_PROFILE
-		//@todo-rco: FIXME!
-		#define TANGENTTOWORLD_INTERPOLATOR_BLOCK	float4 TangentToWorld0 : TANGENTTOWORLD0; float4	TangentToWorld2	: TANGENTTOWORLD2;
-#else
-	#define TANGENTTOWORLD_INTERPOLATOR_BLOCK	MaterialFloat4 TangentToWorld0 : TANGENTTOWORLD0; MaterialFloat4	TangentToWorld2	: TANGENTTOWORLD2;
-#endif
+	#define TANGENTTOWORLD_INTERPOLATOR_BLOCK	half4 TangentToWorld0 : TANGENTTOWORLD0; half4 TangentToWorld2 : TANGENTTOWORLD2;
 #endif
 
 MaterialFloat Luminance( MaterialFloat3 LinearColor )
diff --git a/Engine/Shaders/Private/MobileBasePassCommon.ush b/Engine/Shaders/Private/MobileBasePassCommon.ush
--- a/Engine/Shaders/Private/MobileBasePassCommon.ush
+++ b/Engine/Shaders/Private/MobileBasePassCommon.ush
@@ -23,17 +23,10 @@
 #define PACK_INTERPOLANTS (USE_VERTEX_FOG && NUM_VF_PACKED_INTERPOLANTS > 0 && (ES3_1_PROFILE))
 #define LANDSCAPE_BUG_WORKAROUND (IOS && IS_MOBILE_BASEPASS_VERTEX_SHADER && PACK_INTERPOLANTS)
 
-#if COMPILER_METAL
-// @todo-rco: @todo metal: FIXME!
-#define INTERP_TYPE float4
-#else
-#define INTERP_TYPE half4
-#endif
-
 struct FSharedMobileBasePassInterpolants
 {
 #if USE_VERTEX_FOG && !PACK_INTERPOLANTS
-	INTERP_TYPE VertexFog			: TEXCOORD7;
+	half4 VertexFog			: TEXCOORD7;
 #endif
 
 	float4 PixelPosition	: TEXCOORD8; // xyz = world position, w = clip z
@@ -54,7 +47,7 @@ struct FSharedMobileBasePassInterpolants
 // as it currently uses indexing for interpolator naming.  Otherwise
 // the VS to PS attribute mapping won't match
 #if LANDSCAPE_BUG_WORKAROUND
-	INTERP_TYPE DummyInterp		: DUMMY_INTERP;
+	half4 DummyInterp		: DUMMY_INTERP;
 #endif
 };
 
