Fix for AtmosphericLightVectorto work in vertex shaders.

diff --git a/Engine/Shaders/Private/MaterialTemplate.ush b/Engine/Shaders/Private/MaterialTemplate.ush
--- a/Engine/Shaders/Private/MaterialTemplate.ush
+++ b/Engine/Shaders/Private/MaterialTemplate.ush
@@ -1534,10 +1534,9 @@ float3 MaterialExpressionGetHairColorFromMelanin(float Melanin, float Redness, f
 	return GetHairColorFromMelanin(Melanin, Redness, DyeColor);
 }
 
-float3 MaterialExpressionAtmosphericLightVector(FMaterialPixelParameters Parameters)
-{
-	return ResolvedView.AtmosphereLightDirection[0].xyz;
-}
+#define DEFINE_ATMOSPHERELIGHTVECTOR(MaterialParamType) float3 MaterialExpressionAtmosphericLightVector(MaterialParamType Parameters) {return ResolvedView.AtmosphereLightDirection[0].xyz;}
+DEFINE_ATMOSPHERELIGHTVECTOR(FMaterialPixelParameters)
+DEFINE_ATMOSPHERELIGHTVECTOR(FMaterialVertexParameters)
 
 float3 MaterialExpressionAtmosphericLightColor(FMaterialPixelParameters Parameters)
 {
