ShaderPrint
- Added AddCircleSS
- Clean up AddLine in Favor of AddLineSS

diff --git a/Engine/Shaders/Private/ShaderPrint.ush b/Engine/Shaders/Private/ShaderPrint.ush
--- a/Engine/Shaders/Private/ShaderPrint.ush
+++ b/Engine/Shaders/Private/ShaderPrint.ush
@@ -843,13 +843,9 @@ void AddLineSS(float2 InPosInPixel0, float2 InPosInPixel1, float4 Color0, float4
 	}
 }
 void AddLineSS(float2 Pos0, float2 Pos1, float4 Color)											{ AddLineSS(Pos0, Pos1, Color, Color); }
-void AddLine(float2 Pos0, float2 Pos1, float4 Color)											{ AddLineSS(Pos0, Pos1, Color); }
-void AddLine(float2 Pos0, float2 Pos1, float4 Color0, float4 Color1)							{ AddLineSS(Pos0, Pos1, Color0, Color1); }
 
 void AddLineSS(FShaderPrintContext Ctx, float2 Pos0, float2 Pos1, float4 Color0, float4 Color1) { if (Ctx.bIsActive) { AddLineSS(Pos0, Pos1, Color0, Color1); } }
 void AddLineSS(FShaderPrintContext Ctx, float2 Pos0, float2 Pos1, float4 Color)					{ if (Ctx.bIsActive) { AddLineSS(Pos0, Pos1, Color, Color); } }
-void AddLine(FShaderPrintContext Ctx, float2 Pos0, float2 Pos1, float4 Color)					{ if (Ctx.bIsActive) { AddLineSS(Pos0, Pos1, Color); } }
-void AddLine(FShaderPrintContext Ctx, float2 Pos0, float2 Pos1, float4 Color0, float4 Color1)	{ if (Ctx.bIsActive) { AddLineSS(Pos0, Pos1, Color0, Color1); } }
 
 ///////////////////////////////////////////////////////////////////////////////////////////////////
 // Quad
@@ -938,6 +934,30 @@ void AddCrossTWS(FShaderPrintContext Ctx, float3 Pos, float Size, float4 Color)
 void AddCrossWS (FShaderPrintContext Ctx, float3 Pos, float Size, float4 Color)	{ if (Ctx.bIsActive) { AddCrossTWS(Pos + ShaderPrintData.TranslatedWorldOffset, Size, Color); } }
 void AddCross   (FShaderPrintContext Ctx, float3 Pos, float Size, float4 Color)	{ if (Ctx.bIsActive) { AddCrossTWS(Pos + ShaderPrintData.TranslatedWorldOffset, Size, Color); } }
 
+///////////////////////////////////////////////////////////////////////////////////////////////////
+// Circles
+
+void AddCircleSS(float2 Center, float Radius, float4 Color)
+{
+	const float TStep = 0.05f;
+
+	float S, C;
+	sincos(0.0, S, C);
+	float2 PrevP = Center + float2(C, S) * Radius;
+
+	for (float t = TStep; t <= 1.0; t += TStep)
+	{
+		sincos(t * 2 * PI, S, C);
+		float2 P = Center + float2(C, S) * Radius;
+
+		AddLineSS(PrevP, P, Color);
+		PrevP = P;
+	}
+}
+
+void AddCircleSS(FShaderPrintContext Ctx, float2 Center, float Radius, float4 Color) { if (Ctx.bIsActive) { AddCircleSS(Center, Radius, Color); } }
+
+
 ///////////////////////////////////////////////////////////////////////////////////////////////////
 // Referential
 
