Strata - added strata tree view as a graph. Only for DXC since FXC would drown.

diff --git a/Engine/Shaders/Private/Strata/StrataVisualize.usf b/Engine/Shaders/Private/Strata/StrataVisualize.usf
--- a/Engine/Shaders/Private/Strata/StrataVisualize.usf
+++ b/Engine/Shaders/Private/Strata/StrataVisualize.usf
@@ -459,6 +459,25 @@ void VisualizeMaterialPS(
 #define BSDF_ELEMENT_STARTX 100.0f + 450.0f
 #define BSDF_ELEMENT_STARTY 100
 
+#if SHADER_DEBUGSTRATATREE_PS || SHADER_DEBUGSTRATATREE_CS
+
+#define TREE_MAX_DEPTH			4
+#define TREE_START_POS			float2(1300, 350)
+#define TREE_BRANCH				float2(200, 75)
+#define TREE_COLOR				float4(0.8, 0.8, 0.8, 1.0)
+
+float2 GetBranchOffset(FStrataOperator Op, bool bIsLeft)
+{
+	if (Op.Type == STRATA_OPERATOR_WEIGHT || Op.Type == STRATA_OPERATOR_THINFILM)
+	{
+		return float2(0, TREE_BRANCH.y);
+	}
+	return TREE_BRANCH * (bIsLeft ? float2(-1, 1) : float2(1, 1));
+}
+
+#endif // SHADER_DEBUGSTRATATREE_PS || SHADER_DEBUGSTRATATREE_CS
+
+
 #if SHADER_DEBUGSTRATATREE_CS
 
 [numthreads(1, 1, 1)]
@@ -642,6 +661,128 @@ void MaterialDebugStrataTreeCS(uint3 DispatchThreadId : SV_DispatchThreadID)
 
 		Context = InitShaderPrintContext(true, uint2(1165, BSDF_ELEMENT_STARTY - 50));
 		Print(Context, TEXT("Topology"), FontWhite);
+
+		// Debug print the strata tree only for DXC for the sake of compilation performance.
+	#if COMPILER_DXC
+		uint NodeIndexStack[TREE_MAX_DEPTH];
+		uint ChildrenVisitedStack[TREE_MAX_DEPTH]; // 0:none, 1:left, 2:left&right
+		float2 PrintPosStack[TREE_MAX_DEPTH];
+
+		int StackPtr = 0;
+		NodeIndexStack[StackPtr] = StrataTreeHeader.RootOperatorIndex;
+		ChildrenVisitedStack[StackPtr] = 0;
+		PrintPosStack[StackPtr] = TREE_START_POS;
+
+		while (StackPtr >= 0)
+		{
+			uint OperatorIndex = NodeIndexStack[StackPtr];
+			FStrataOperator Op = StrataUnpackInStrataTreeOperator(
+				OperatorIndex,
+				Strata.MaterialTextureArray,
+				StrataTreeHeader,
+				SliceStoringDebugStrataTree);
+
+			Context = InitShaderPrintContext(true, uint2(PrintPosStack[StackPtr]));
+			if (Op.Type == STRATA_OPERATOR_WEIGHT)
+			{
+				Print(Context, TEXT("W"), FontWhite);
+			}
+			else if (Op.Type == STRATA_OPERATOR_VERTICAL)
+			{
+				Print(Context, TEXT("V"), FontWhite);
+			}
+			else if (Op.Type == STRATA_OPERATOR_HORIZONTAL)
+			{
+				Print(Context, TEXT("H"), FontWhite);
+			}
+			else if (Op.Type == STRATA_OPERATOR_ADD)
+			{
+				Print(Context, TEXT("A"), FontWhite);
+			}
+			else if (Op.Type == STRATA_OPERATOR_BSDF)
+			{
+				Context = InitShaderPrintContext(true, uint2(PrintPosStack[StackPtr]) - uint2(3,0));
+				Print(Context, TEXT("B"), FontYellow);
+				Print(Context, Op.LeftIndex, FontYellow);
+			}
+			else if (Op.Type == STRATA_OPERATOR_BSDF_LEGACY)
+			{
+				Print(Context, TEXT("BL"), FontWhite);
+			}
+			else if (Op.Type == STRATA_OPERATOR_THINFILM)
+			{
+				Print(Context, TEXT("TF"), FontWhite);
+			}
+			else
+			{
+				Print(Context, TEXT("UNKOWN"), FontWhite);
+			}
+
+			const float2 ParentPrintPos = PrintPosStack[StackPtr];
+
+			AddCircleSS(Context, ParentPrintPos, 20, TREE_COLOR);
+
+			if (Op.LayerDepth == TREE_MAX_DEPTH - 1 || Op.Type == STRATA_OPERATOR_BSDF || Op.Type == STRATA_OPERATOR_BSDF_LEGACY)
+			{
+				StackPtr--;
+				continue;
+			}
+
+			if (ChildrenVisitedStack[StackPtr] == 0)
+			{
+				if (Op.LeftIndex >= 0)
+				{
+					ChildrenVisitedStack[StackPtr] = 1;
+
+					StackPtr++;
+					NodeIndexStack[StackPtr] = Op.LeftIndex;
+					ChildrenVisitedStack[StackPtr] = 0;
+
+					PrintPosStack[StackPtr] = ParentPrintPos + GetBranchOffset(Op, true);
+					AddLineSS(Context, ParentPrintPos, PrintPosStack[StackPtr], TREE_COLOR);
+				}
+				else if (Op.RightIndex >= 0)
+				{
+					ChildrenVisitedStack[StackPtr] = 2;
+
+					StackPtr++;
+					NodeIndexStack[StackPtr] = Op.RightIndex;
+					ChildrenVisitedStack[StackPtr] = 0;
+
+					PrintPosStack[StackPtr] = ParentPrintPos + GetBranchOffset(Op, false);
+					AddLineSS(Context, ParentPrintPos, PrintPosStack[StackPtr], TREE_COLOR);
+				}
+				else
+				{
+					StackPtr--;
+				}
+			}
+			else if (ChildrenVisitedStack[StackPtr] == 1)
+			{
+				if (Op.RightIndex >= 0)
+				{
+					ChildrenVisitedStack[StackPtr] = 2;
+
+					StackPtr++;
+					NodeIndexStack[StackPtr] = Op.RightIndex;
+					ChildrenVisitedStack[StackPtr] = 0;
+
+					PrintPosStack[StackPtr] = ParentPrintPos + GetBranchOffset(Op, false);
+					AddLineSS(Context, ParentPrintPos, PrintPosStack[StackPtr], TREE_COLOR);
+				}
+				else
+				{
+					StackPtr--;
+				}
+			}
+			else
+			{
+				StackPtr--;
+			}
+		}
+
+	#endif // COMPILER_DXC
+
 	}
 }
 
diff --git a/Engine/Source/Runtime/Renderer/Private/Strata/StrataVisualize.cpp b/Engine/Source/Runtime/Renderer/Private/Strata/StrataVisualize.cpp
--- a/Engine/Source/Runtime/Renderer/Private/Strata/StrataVisualize.cpp
+++ b/Engine/Source/Runtime/Renderer/Private/Strata/StrataVisualize.cpp
@@ -203,7 +203,7 @@ static void AddVisualizeMaterialPasses(FRDGBuilder& GraphBuilder, const FViewInf
 	if (View.Family->EngineShowFlags.VisualizeStrataMaterial)
 	{
 		if (!ShaderPrint::IsEnabled(View)) { ShaderPrint::SetEnabled(true); }
-		ShaderPrint::RequestSpaceForLines(64);
+		ShaderPrint::RequestSpaceForLines(1024);
 		ShaderPrint::RequestSpaceForCharacters(1024);
 
 		const uint32 ViewMode = FMath::Max(0, CVarStrataDebugVisualizeMode.GetValueOnRenderThread());
