Fix Strata shader compilation warning.

diff --git a/Engine/Shaders/Private/Strata/Strata.ush b/Engine/Shaders/Private/Strata/Strata.ush
--- a/Engine/Shaders/Private/Strata/Strata.ush
+++ b/Engine/Shaders/Private/Strata/Strata.ush
@@ -857,16 +857,16 @@ void StrataPackOutStrataTree(
 	StrataTree.OperatorCount = min(StrataTree.OperatorCount, STRATA_MAX_OPERATOR_COUNT);
 	STRATA_UNROLL_N(STRATA_MAX_OPERATOR_COUNT)
 #endif
-	for (i = 0; i < StrataTree.OperatorCount; ++i)
+	for (uint j = 0; j < StrataTree.OperatorCount; ++j)
 	{
-		WRITE_INT(StrataTree.Operators[i].ParentIndex);
-		WRITE_UINT(StrataTree.Operators[i].LayerDepth);
-		WRITE_UINT(StrataTree.Operators[i].MaxDistanceFromLeaves);
-		WRITE_UINT(StrataTree.Operators[i].Type);
-		WRITE_INT(StrataTree.Operators[i].LeftIndex);
-		WRITE_INT(StrataTree.Operators[i].RightIndex);
-		WRITE_FLOAT(StrataTree.Operators[i].Weight);
-		WRITE_FLOAT(StrataTree.Operators[i].Coverage);
+		WRITE_INT(StrataTree.Operators[j].ParentIndex);
+		WRITE_UINT(StrataTree.Operators[j].LayerDepth);
+		WRITE_UINT(StrataTree.Operators[j].MaxDistanceFromLeaves);
+		WRITE_UINT(StrataTree.Operators[j].Type);
+		WRITE_INT(StrataTree.Operators[j].LeftIndex);
+		WRITE_INT(StrataTree.Operators[j].RightIndex);
+		WRITE_FLOAT(StrataTree.Operators[j].Weight);
+		WRITE_FLOAT(StrataTree.Operators[j].Coverage);
 		// Skip the rest for now
 	}
 
