Fixes shader compilation failure in TSR when alpha channel through post processing is enabled

diff --git a/Engine/Shaders/Private/TemporalSuperResolution/TSRUpdateHistory.usf b/Engine/Shaders/Private/TemporalSuperResolution/TSRUpdateHistory.usf
--- a/Engine/Shaders/Private/TemporalSuperResolution/TSRUpdateHistory.usf
+++ b/Engine/Shaders/Private/TemporalSuperResolution/TSRUpdateHistory.usf
@@ -1662,7 +1662,11 @@ void MainCS(
 					HistoryOutput_Metadata[dpv_lo(LocalHistoryPixelPos)]      = tsr_half2(dpv_lo(FinalLowFrequencyWeight), dpv_lo(FinalHistoryValidity));
 				#endif
 				#if CONFIG_ACCUMULATE_TRANSLUCENCY_SEPARATLY
-					HistoryOutput_ColorArray[tsr_short3(dpv_lo(LocalHistoryPixelPos), HistoryOutput_ArrayIndices_Translucency)] = dpv_lo(FinalTranslucencyColor).rgb;
+					#if CONFIG_SCENE_COLOR_APLHA
+						HistoryOutput_ColorArray[tsr_short3(dpv_lo(LocalHistoryPixelPos), HistoryOutput_ArrayIndices_Translucency)] = dpv_lo(FinalTranslucencyColor);
+					#else
+						HistoryOutput_ColorArray[tsr_short3(dpv_lo(LocalHistoryPixelPos), HistoryOutput_ArrayIndices_Translucency)] = dpv_lo(FinalTranslucencyColor).rgb;
+					#endif 
 					HistoryOutput_TranslucencyAlpha[dpv_lo(LocalHistoryPixelPos)] = dpv_lo(FinalTranslucencyColor).a;
 				#endif
 				#if CONFIG_UPDATE_SUBPIXEL_DETAIL
@@ -1682,7 +1686,11 @@ void MainCS(
 					HistoryOutput_Metadata[dpv_hi(LocalHistoryPixelPos)]      = tsr_half2(dpv_hi(FinalLowFrequencyWeight), dpv_hi(FinalHistoryValidity));
 				#endif
 				#if CONFIG_ACCUMULATE_TRANSLUCENCY_SEPARATLY
-					HistoryOutput_ColorArray[tsr_short3(dpv_hi(LocalHistoryPixelPos), HistoryOutput_ArrayIndices_Translucency)] = dpv_hi(FinalTranslucencyColor).rgb;
+					#if CONFIG_SCENE_COLOR_APLHA
+						HistoryOutput_ColorArray[tsr_short3(dpv_hi(LocalHistoryPixelPos), HistoryOutput_ArrayIndices_Translucency)] = dpv_hi(FinalTranslucencyColor);
+					#else
+						HistoryOutput_ColorArray[tsr_short3(dpv_hi(LocalHistoryPixelPos), HistoryOutput_ArrayIndices_Translucency)] = dpv_hi(FinalTranslucencyColor).rgb;
+					#endif
 					HistoryOutput_TranslucencyAlpha[dpv_hi(LocalHistoryPixelPos)] = dpv_hi(FinalTranslucencyColor).a;
 				#endif
 				#if CONFIG_UPDATE_SUBPIXEL_DETAIL
