Fix missing specular lighting into RT reflections.

diff --git a/Engine/Shaders/Private/RayTracing/RayTracingLightingMS.usf b/Engine/Shaders/Private/RayTracing/RayTracingLightingMS.usf
--- a/Engine/Shaders/Private/RayTracing/RayTracingLightingMS.usf
+++ b/Engine/Shaders/Private/RayTracing/RayTracingLightingMS.usf
@@ -101,7 +101,7 @@ RAY_TRACING_ENTRY_MISS(RayTracingLightingMS, FPackedMaterialClosestHitPayload, P
 	// LWC_TODO - lighting wants translated world space
 	float SurfaceShadow = 1.0f;
 	FGBufferData GBufferData = GetGBufferDataFromPayload(Payload);
-	float3 LightContribution = GetDynamicLighting(WorldPosition, V, GBufferData, 1.0f, GBufferData.ShadingModelID, LightData, LightAttenuation, Dither, SVPos, SurfaceShadow).xyz;
+	float3 LightContribution = GetDynamicLighting(WorldPosition, CameraVector, GBufferData, 1.0f, GBufferData.ShadingModelID, LightData, LightAttenuation, Dither, SVPos, SurfaceShadow).xyz;
 #endif // STRATA_ENABLED
 
 	float3 AccumulatedRadiance = Payload.Radiance + LightContribution * LightProfileMultiplier;
