Fix RT shader compilation error when strata was enabled, but the shader wasn't a material (e.g., material function).

diff --git a/Engine/Shaders/Private/RayTracing/RayTracingMaterialHitShaders.usf b/Engine/Shaders/Private/RayTracing/RayTracingMaterialHitShaders.usf
--- a/Engine/Shaders/Private/RayTracing/RayTracingMaterialHitShaders.usf
+++ b/Engine/Shaders/Private/RayTracing/RayTracingMaterialHitShaders.usf
@@ -243,6 +243,7 @@ void GetMaterialPayload(
 	inout FMaterialClosestHitPayload Payload)
 #if STRATA_ENABLED
 {
+  #if MATERIAL_IS_STRATA
 	const float Dither = 0.5f;
 	float3 EmissiveLuminance = 0.0f;
 	bool bStrataSubsurfaceEnable = false;
@@ -265,10 +266,6 @@ void GetMaterialPayload(
 
 	// Generate the strata material data to write out
 	FStrataAddressing StrataAddressing = GetStrataPixelDataByteOffset(0,0,0);
-	
-#if STRATA_USES_RESOURCE_MATERIALCONTAINER != 0
-#error This is not correct
-#endif
 
 	const float3 CamVector = MaterialParameters.CameraVector;	
 	Payload.StrataData = InitialiseRWStrataMaterialContainer();
@@ -297,6 +294,7 @@ void GetMaterialPayload(
 		Payload.Opacity = 0;
 	}
 	#endif // MATERIALBLENDING_MASKED
+  #endif // MATERIAL_IS_STRATA
 }
 #else // STRATA_ENABLED
 {
