Strata - renaming to clarify what is throughput, transmittance, along V or N, and what is transmittance for matter layered above.

diff --git a/Engine/Shaders/Private/DebugProbes.usf b/Engine/Shaders/Private/DebugProbes.usf
--- a/Engine/Shaders/Private/DebugProbes.usf
+++ b/Engine/Shaders/Private/DebugProbes.usf
@@ -307,7 +307,7 @@ void MainPS(
 
 		FStrataData StrataData = PromoteParameterBlendedBSDFToOperator(StrataSlabData, StrataPixelHeader.StrataTree, /*int OperatorIndex*/0, /*int BSDFIndex*/0, /*int LayerDepth*/0, /*int bIsBottom*/1);
 		StrataPixelHeader.StrataTree.BSDFs[0].LuminanceWeightV = 1.0f;
-		StrataPixelHeader.StrataTree.BSDFs[0].LuminanceWeightL = 1.0f;
+		StrataPixelHeader.StrataTree.BSDFs[0].TransmittanceAboveAlongN = 1.0f;
 		StrataPixelHeader.StrataTree.BSDFs[0].TopLayerDataWeight = 1.0f;
 
 		FStrataSubsurfaceData SSSData = (FStrataSubsurfaceData)0;
diff --git a/Engine/Shaders/Private/Strata/Strata.ush b/Engine/Shaders/Private/Strata/Strata.ush
--- a/Engine/Shaders/Private/Strata/Strata.ush
+++ b/Engine/Shaders/Private/Strata/Strata.ush
@@ -99,7 +99,10 @@ struct FStrataBSDF
 	// STRATA_TODO pack the following
 	int					OperatorIndex;
 	float3				LuminanceWeightV;
-	float3				LuminanceWeightL;
+
+	float3				TransmittanceAboveAlongN;	// The transmittance towards the top of the material along the normal ignoring this BSDF, e.g. only accumulating matter above it.
+//	float3				TransmittanceBelow;			// TODO
+
 	int					bIsBottom;
 	int					bIsTop;
 
@@ -263,7 +266,7 @@ struct FStrataBSDF
 #define STATE_BIT_COUNT_ISTHIN					1
 #define STATE_BIT_COUNT_BSDF_TYPE				3
 #define STATE_BIT_COUNT_SHAREDLOCALBASESID		2
-#define STATE_BIT_COUNT_HASWEIGHT_L				1
+#define STATE_BIT_COUNT_HASTRANSABOVE			1
 #define STATE_BIT_COUNT___UNUSED___				5
 #define STATE_BIT_COUNT_WEIGHT10F				10
 
@@ -281,7 +284,7 @@ struct FStrataBSDF
 #define STATE_BIT_OFFSET_ISTHIN					(STATE_BIT_OFFSET_HASFUZZ			+ STATE_BIT_COUNT_HASFUZZ)
 #define STATE_BIT_OFFSET_BSDF_TYPE				(STATE_BIT_OFFSET_ISTHIN			+ STATE_BIT_COUNT_ISTHIN)
 #define STATE_BIT_OFFSET_SHAREDLOCALBASESID		(STATE_BIT_OFFSET_BSDF_TYPE			+ STATE_BIT_COUNT_BSDF_TYPE)
-#define STATE_BIT_OFFSET_HASWEIGHT_L			(STATE_BIT_OFFSET_SHAREDLOCALBASESID+ STATE_BIT_COUNT_SHAREDLOCALBASESID)
+#define STATE_BIT_OFFSET_HASTRANSABOVE			(STATE_BIT_OFFSET_SHAREDLOCALBASESID+ STATE_BIT_COUNT_SHAREDLOCALBASESID)
 #define STATE_BIT_OFFSET_WEIGHT10F				22
 
 // Layout 0 (Fast)
@@ -358,8 +361,8 @@ struct FStrataBSDF
 #define BSDF_SETISTHIN(X, Value)		WRITE_BITS(X.State, STATE_BIT_COUNT_ISTHIN, STATE_BIT_OFFSET_ISTHIN, Value)
 
 // Indicates if the BSDF has light weight data
-#define BSDF_GETHASWEIGHT_L(X)			READ_BITS(X.State,  STATE_BIT_COUNT_HASWEIGHT_L, STATE_BIT_OFFSET_HASWEIGHT_L)
-#define BSDF_SETHASWEIGHT_L(X, Value)	WRITE_BITS(X.State, STATE_BIT_COUNT_HASWEIGHT_L, STATE_BIT_OFFSET_HASWEIGHT_L, Value)
+#define BSDF_GETHASTRANSABOVE(X)		READ_BITS(X.State,  STATE_BIT_COUNT_HASTRANSABOVE, STATE_BIT_OFFSET_HASTRANSABOVE)
+#define BSDF_SETHASTRANSABOVE(X, Value)	WRITE_BITS(X.State, STATE_BIT_COUNT_HASTRANSABOVE, STATE_BIT_OFFSET_HASTRANSABOVE, Value)
 
 #if STRATA_INLINE_SHADING
 #define BSDF_SETEMISSIVE(X, V)			X.Emissive = V
@@ -737,12 +740,12 @@ struct FStrataOperator
 	// Coverage / Transmittance
 
 	float  Coverage;									// The summarised coverage of all the matter (from the sub tree) represented by that operator
-	float3 TransmittanceL;								// This is the transmittance of the matter itself. If coverage is 0.25, then 25% of the light will be affected by the transmittance and 75% of light won't be.
-	float3 TransmittanceV;								// This is the transmittance of the matter itself. If coverage is 0.25, then 25% of the light will be affected by the transmittance and 75% of light won't be.
+	float3 ThroughputAlongV;							// This is the throughput of the matter represented by this operator (matter and fuzz transmittance). If coverage is 0.25, then 25% of the light will be affected by the transmittance and 75% of light won't be.
+	float3 TransmittanceAlongN;							// This is the transmittance of the matter represented by this operator. If coverage is 0.25, then 25% of the light will be affected by the transmittance and 75% of light won't be.
 
 	float  VerticalTop_Coverage;						// Only used by vertical layering operators. It represent the coverage of all the matter of the top layer (from the left sub tree).
-	float3 VerticalTop_TransmittanceV;					// Only used by vertical layering operators. It represent the transmittance of matter of the top layer (from the left sub tree, and towards the view).
-	float3 VerticalTop_TransmittanceL;					// Only used by vertical layering operators. It represent the transmittance of matter of the top layer (from the left sub tree, and towards the light).
+	float3 VerticalTop_ThroughputAlongV;				// Only used by vertical layering operators. It represent the transmittance of matter of the top layer (from the left sub tree, and towards the view).
+	float3 VerticalTop_TransmittanceAlongN;				// Only used by vertical layering operators. It represent the transmittance of matter of the top layer (from the left sub tree, and along the normal of the surface).
 
 	// Rough refraction
 	float3				 TopDownRefractionWorldNormal;	// The normal used for the single refracted direction
@@ -824,7 +827,7 @@ void StrataPackOutStrataTree(
 		WRITE_UINT(StrataTree.BSDFs[i].State);
 		WRITE_INT(StrataTree.BSDFs[i].OperatorIndex);
 		WRITE_FLOAT3(StrataTree.BSDFs[i].LuminanceWeightV);
-		WRITE_FLOAT3(StrataTree.BSDFs[i].LuminanceWeightL);
+		WRITE_FLOAT3(StrataTree.BSDFs[i].TransmittanceAboveAlongN);
 		WRITE_UINT(StrataTree.BSDFs[i].bIsBottom);
 		WRITE_UINT(StrataTree.BSDFs[i].bIsTop);
 
@@ -920,7 +923,7 @@ FStrataBSDF StrataUnpackInStrataTreeBSDF(
 	READ_UINT(BSDF.State);
 	READ_INT(BSDF.OperatorIndex);
 	READ_FLOAT3(BSDF.LuminanceWeightV);
-	READ_FLOAT3(BSDF.LuminanceWeightL);
+	READ_FLOAT3(BSDF.TransmittanceAboveAlongN);
 	READ_UINT(BSDF.bIsBottom);
 	READ_UINT(BSDF.bIsTop);
 
@@ -1758,7 +1761,7 @@ FStrataData StrataHorizontalMixingParameterBlending(FStrataData Background, FStr
 	BSDF_SETHASHAZINESS(ResultBSDF,		BSDF_GETHASHAZINESS(ForegroundBSDF)				| BSDF_GETHASHAZINESS(BackgroundBSDF));
 	BSDF_SETHASF90(ResultBSDF,			BSDF_GETHASF90(ForegroundBSDF)					| BSDF_GETHASF90(BackgroundBSDF));
 	BSDF_SETHASFUZZ(ResultBSDF,			BSDF_GETHASFUZZ(ForegroundBSDF)					| BSDF_GETHASFUZZ(BackgroundBSDF));
-	BSDF_SETHASWEIGHT_L(ResultBSDF,		BSDF_GETHASWEIGHT_L(ForegroundBSDF)				| BSDF_GETHASWEIGHT_L(BackgroundBSDF));
+	BSDF_SETHASTRANSABOVE(ResultBSDF,	BSDF_GETHASTRANSABOVE(ForegroundBSDF)			| BSDF_GETHASTRANSABOVE(BackgroundBSDF));
 
 	BSDF_SETEMISSIVE(ResultBSDF,		lerp(BSDF_GETEMISSIVE(BackgroundBSDF),			BSDF_GETEMISSIVE(ForegroundBSDF),			ForegroundOtherMixFactor));
 	SLAB_DIFFUSEALBEDO(ResultBSDF) =	lerp(SLAB_DIFFUSEALBEDO(BackgroundBSDF),		SLAB_DIFFUSEALBEDO(ForegroundBSDF),			ForegroundOtherMixFactor);
@@ -1950,8 +1953,8 @@ FStrataBSDF StrataVerticalLayeringParameterBlending(in FStrataBSDF TopBSDF, in F
 	const float3 UnusedDiffuseColor		= 0;	// Unused when only transmittance is needed
 	FParticipatingMedia PM				= StrataSlabCreateParticipatingMedia(UnusedDiffuseColor, SLAB_SSSMFP(TopBSDF));
 	const float3 SlabTransmittanceV		= IsotropicMediumSlabTransmittance(PM, STRATA_SIMPLEVOLUME_THICKNESS_M, NoV);
-	const float3 SlabTransmittanceL		= IsotropicMediumSlabTransmittance(PM, STRATA_SIMPLEVOLUME_THICKNESS_M, 1.0f /*Normal incidence i.e., NoL == 1.f*/);
-	const float3 SlabTransmittance		= SlabTransmittanceV * SlabTransmittanceL;
+	const float3 SlabTransmittanceN		= IsotropicMediumSlabTransmittance(PM, STRATA_SIMPLEVOLUME_THICKNESS_M, 1.0f /*Normal incidence i.e., NoL == 1.f*/);
+	const float3 SlabTransmittance		= SlabTransmittanceV * SlabTransmittanceN;
 
 	// STRATA_TODO account for energy conservation here.
 
@@ -2959,7 +2962,7 @@ void StrataUpdateTree(
 		Settings);
 
 	OutCoverage = saturate(StrataTree.Operators[StrataData.OperatorIndex].Coverage);
-	OutTransmittancePreCoverage = saturate(StrataTree.Operators[StrataData.OperatorIndex].TransmittanceV);
+	OutTransmittancePreCoverage = saturate(StrataTree.Operators[StrataData.OperatorIndex].ThroughputAlongV);
 }
 
 #endif // STRATA_INLINE_SHADING && STRATA_CLAMPED_BSDF_COUNT > 0
@@ -3189,8 +3192,8 @@ void PackStrataOut(
 					const uint GreyScaleThroughputV = StrataHasGreyScaleWeight(BSDF.LuminanceWeightV) ? 1 : 0;
 					BSDF_SETHASGREYWEIGHT_V(BSDF, GreyScaleThroughputV);
 
-					const bool bHasLightTransmittance = any(BSDF.LuminanceWeightL < 1.0f);
-					BSDF_SETHASWEIGHT_L(BSDF, bHasLightTransmittance);
+					const bool bTransmittanceAboveAlongNRequired = any(BSDF.TransmittanceAboveAlongN < 1.0f);
+					BSDF_SETHASTRANSABOVE(BSDF, bTransmittanceAboveAlongNRequired);
 
 					// A. Layout0 - Simple encoding (aka. fast-path): store header & data
 					if (bHasFastEncoding && BSDFIdx == SimpleMaterial_BSDFIndex)
@@ -3315,9 +3318,9 @@ void PackStrataOut(
 							STRATA_STORE_UINT1(PackColorLinearToGamma2AlphaLinear(float4(SLAB_FUZZ_COLOR(BSDF), SLAB_FUZZ_AMOUNT(BSDF))));
 						}
 
-						if (BSDF_GETHASWEIGHT_L(BSDF))
+						if (BSDF_GETHASTRANSABOVE(BSDF))
 						{
-							STRATA_STORE_UINT1(PackR11G11B10F(BSDF.LuminanceWeightL));
+							STRATA_STORE_UINT1(PackR11G11B10F(BSDF.TransmittanceAboveAlongN));
 						}
 
 						// 8-24 bytes
@@ -3464,7 +3467,7 @@ FStrataBSDF UnpackFastPathStrataBSDFIn(FStrataMaterialContainer StrataBuffer, in
 	uint Data1 = StrataLoadUint1(StrataBuffer, StrataAddressing);
 	OutBSDF.State = 0;
 	OutBSDF.LuminanceWeightV = 1.0f;
-	OutBSDF.LuminanceWeightL = 1.0f;
+	OutBSDF.TransmittanceAboveAlongN = 1.0f;
 
 	const uint PackedRoughness8bits = 0xFF    & (Data0 >>  HEADER_FASTENCODING_BIT_COUNT);
 	const uint PackedData8Bits		= 0xFF    & (Data0 >> (HEADER_FASTENCODING_BIT_COUNT + 8));
@@ -3493,7 +3496,7 @@ FStrataBSDF UnpackFastPathStrataBSDFIn(FStrataMaterialContainer StrataBuffer, in
 	BSDF_SETHASHAZINESS(OutBSDF, 0);
 	BSDF_SETISSIMPLEVOLUME(OutBSDF, 0);
 	BSDF_SETHASFUZZ(OutBSDF, 0);
-	BSDF_SETHASWEIGHT_L(OutBSDF, 0);
+	BSDF_SETHASTRANSABOVE(OutBSDF, 0);
 
 	return OutBSDF;
 }
@@ -3510,7 +3513,7 @@ FStrataBSDF UnpackStrataBSDFIn(FStrataMaterialContainer StrataBuffer, inout FStr
 	}
 
 	FStrataBSDF OutBSDF = (FStrataBSDF)0;
-	OutBSDF.LuminanceWeightL = 1.0f;
+	OutBSDF.TransmittanceAboveAlongN = 1.0f;
 
 	float DummyFloat = 0.0f;
 
@@ -3609,9 +3612,9 @@ FStrataBSDF UnpackStrataBSDFIn(FStrataMaterialContainer StrataBuffer, inout FStr
 			StrataLoad_ColorGamma2ToLinearAlphaLinear(StrataBuffer, StrataAddressing, SLAB_FUZZ_COLOR(OutBSDF), SLAB_FUZZ_AMOUNT(OutBSDF));
 		}
 
-		if (BSDF_GETHASWEIGHT_L(OutBSDF))
+		if (BSDF_GETHASTRANSABOVE(OutBSDF))
 		{
-			StrataLoad_R11G11B10F(StrataBuffer, StrataAddressing, OutBSDF.LuminanceWeightL);
+			StrataLoad_R11G11B10F(StrataBuffer, StrataAddressing, OutBSDF.TransmittanceAboveAlongN);
 		}
 	}
 	break;
diff --git a/Engine/Shaders/Private/Strata/StrataEvaluation.ush b/Engine/Shaders/Private/Strata/StrataEvaluation.ush
--- a/Engine/Shaders/Private/Strata/StrataEvaluation.ush
+++ b/Engine/Shaders/Private/Strata/StrataEvaluation.ush
@@ -177,14 +177,14 @@ float3 LuminanceWeight(in float SatNoL, in FStrataBSDF InBSDF)
 	// LuminanceWeight is the absorption from the shading point towards the surface 
 	// in the normal direction. To compute the absorption towards the light, we 
 	// reweight it with new distance L'
-	// LuminanceWeightL  = exp(sigma * L)  with L=1 for normal incidence. 
-	// LuminanceWeightL' = exp(sigma * L') with L'= L . 1/NoL
-	//                   = exp(sigma * L . 1/NoL)
-	//                   = pow(exp(sigma * L), 1/NoL)
-	//                   = pow(LuminanceWeightL, 1/NoL)
+	// TransmittanceAboveAlongN = exp(sigma * T)  at normal incidence and for a thickness T. 
+	// TransmittanceAboveAlongL	= exp(sigma * L)  with L= T . 1/NoL
+	//							= exp(sigma * T . 1/NoL)
+	//							= pow(exp(sigma * T), 1/NoL)
+	// TransmittanceAboveAlongL	= pow(TransmittanceAboveAlongN, 1/NoL)
 	const float DistanceL = rcp(max(SatNoL, 0.001f));
-	const float3 EffectiveLuminanceWeightL = any(InBSDF.LuminanceWeightL < 1.f) ? pow(max(InBSDF.LuminanceWeightL, 0.0001f), DistanceL) : InBSDF.LuminanceWeightL;
-	return InBSDF.LuminanceWeightV * EffectiveLuminanceWeightL;
+	const float3 TransmittanceAboveAlongL = any(InBSDF.TransmittanceAboveAlongN < 1.f) ? pow(max(InBSDF.TransmittanceAboveAlongN, 0.0001f), DistanceL) : InBSDF.TransmittanceAboveAlongN;
+	return InBSDF.LuminanceWeightV * TransmittanceAboveAlongL;
 }
 
 float3 LuminanceWeight(in FStrataBSDFContext InContext, in FStrataBSDF InBSDF)
@@ -222,8 +222,8 @@ struct FStrataEvaluateResult
 	float  TransmissionPDF;
 	float  ThinTransmissionPathPDF;
 
-	float3 ThroughputV;					// Throughput to the next layer (from the view to the shading point)
-	float3 ThroughputL;					// Throughput to the next layer (from the light to the shading point)
+	float3 ThroughputV;					// Throughput to the next layer (from the view to the shading point). Contains many things such as transmittance or cloth energy conservation/preservation factors.
+	float3 TransmittanceAlongN;			// Transmittance to the next layer along the surface normal. This is used later to compute the transmittance towards the light.
 	bool   bSubsurface;					// True if we need to separate the subsurface light contribution for the screen space diffusion process.
 	bool   bApplyProjectedSolidAngle;	// True if the Saturate(NoL) factor should be applied or not over the value of the BSDF
 };
@@ -494,8 +494,8 @@ FStrataEvaluateResult StrataEvaluateBSDFCommon(FStrataBSDFContext BSDFContext, F
 			// we do not need to add emissive for the BRDF TotalSpec or TotalDiff values as this is handled separately
 			Sample.EmissivePathValue = BSDF_GETEMISSIVE(BSDFContext.BSDF);
 
-			Sample.ThroughputV = OpaqueBSDFThroughput;
-			Sample.ThroughputL = OpaqueBSDFThroughput;
+			Sample.ThroughputV			= OpaqueBSDFThroughput;
+			Sample.TransmittanceAlongN	= OpaqueBSDFThroughput;
 
 			////
 			//// Evaluate cloth fuzz layered on top of the slab.
@@ -548,6 +548,7 @@ FStrataEvaluateResult StrataEvaluateBSDFCommon(FStrataBSDFContext BSDFContext, F
 				}
 
 				Sample.ThroughputV *= Cloth_DirectionalAlbedo_SpecularTransmission;
+				// The specular transmission is ignored towards the direction of the light. (TransmittanceAlongN can only store transmittance)
 				break;
 			}
 			#endif
@@ -587,7 +588,7 @@ FStrataEvaluateResult StrataEvaluateBSDFCommon(FStrataBSDFContext BSDFContext, F
 				const float PhaseFunction = ApproximateHG(dot(-BSDFContext.L, RefracV), PhaseFunctionAnisotropy);
 
 				Sample.ThroughputV				= OpaqueBSDFThroughput; // SSS is not translucent as of today
-				Sample.ThroughputL				= OpaqueBSDFThroughput; // idem
+				Sample.TransmittanceAlongN		= OpaqueBSDFThroughput; // idem
 				Sample.TransmissionPathValue	= TransmissionThroughput * PhaseFunction;
 				Sample.TransmissionPDF			= 1.0f / (4.0f * PI);		// STRATA_TODO this currently match the uniform sphere sampling from StrataImportanceSampleBSDF
 
@@ -647,17 +648,17 @@ FStrataEvaluateResult StrataEvaluateBSDFCommon(FStrataBSDFContext BSDFContext, F
 
 				// Shading point <-> Light throughput
 				// Compute the transmittance at normal incidence (instead of BSDFContext.Context.NoL), and will compute the final value during evaluation 
-				const float3 SlabTransmittanceL = IsotropicMediumSlabTransmittance(PM, STRATA_SIMPLEVOLUME_THICKNESS_M, 1.f /*NoL with L==N*/);
-				const float3 SpecularTransmissionL = 1.0f;
+				const float3 SlabTransmittanceN = IsotropicMediumSlabTransmittance(PM, STRATA_SIMPLEVOLUME_THICKNESS_M, 1.f /*NoL with L==N*/);
+				// The specular transmission is ignored towards the direction of the light. (TransmittanceAlongN can only store transmittance)
 
 				const float Phase = IsotropicPhase();
 
 				// Now lerp between the optically thin and optically thin medium models.
 				// The diffuse and throughput account for the GGX interface SpecularTransmission.
-				Sample.DiffusePathValue = lerp(Sample.DiffusePathValue,		SlabDirectionalAlbedo * SpecularTransmissionV,		DiffuseToVolumeBlend);
-				Sample.DiffusePDF		= lerp(Sample.DiffusePDF,			Phase,												DiffuseToVolumeBlend);
-				Sample.ThroughputV		= lerp(Sample.ThroughputV,			SlabTransmittanceV    * SpecularTransmissionV,		DiffuseToVolumeBlend);
-				Sample.ThroughputL		= lerp(Sample.ThroughputL,			SlabTransmittanceL    * SpecularTransmissionL,		DiffuseToVolumeBlend);
+				Sample.DiffusePathValue		= lerp(Sample.DiffusePathValue,		SlabDirectionalAlbedo * SpecularTransmissionV,		DiffuseToVolumeBlend);
+				Sample.DiffusePDF			= lerp(Sample.DiffusePDF,			Phase,												DiffuseToVolumeBlend);
+				Sample.ThroughputV			= lerp(Sample.ThroughputV,			SlabTransmittanceV    * SpecularTransmissionV,		DiffuseToVolumeBlend);
+				Sample.TransmittanceAlongN	= lerp(Sample.TransmittanceAlongN,	SlabTransmittanceN,									DiffuseToVolumeBlend);
 
 				Sample.bSubsurface = false;	// It should already be the case because Enforce in this case because BSDF_GETHASSSS should be false when BSDF_GETISSIMPLEVOLUME is true
 
@@ -697,7 +698,7 @@ FStrataEvaluateResult StrataEvaluateBSDFCommon(FStrataBSDFContext BSDFContext, F
 			Sample.SpecularPathValue	= HairShading(GBuffer, BSDFContext.L, BSDFContext.V, BSDFContext.N, ShadowTerms.TransmissionShadow, HairTransmittance, BacklitEnabled, Area, Random);
 			Sample.SpecularPDF			= 1.0f / (4.0f * PI);		// STRATA_TODO this currently match the uniform sphere sampling from StrataImportanceSampleBSDF
 			Sample.ThroughputV			= OpaqueBSDFThroughput; 
-			Sample.ThroughputL			= OpaqueBSDFThroughput; 
+			Sample.TransmittanceAlongN	= OpaqueBSDFThroughput;
 			Sample.IntegratedSpecularValue = ShadowTerms.TransmissionShadow * AreaLightContext.AreaLight.FalloffColor * AreaLightContext.Falloff * Sample.SpecularPathValue;
 		}
 		break;
diff --git a/Engine/Shaders/Private/Strata/StrataTree.ush b/Engine/Shaders/Private/Strata/StrataTree.ush
--- a/Engine/Shaders/Private/Strata/StrataTree.ush
+++ b/Engine/Shaders/Private/Strata/StrataTree.ush
@@ -97,22 +97,21 @@ void UpdateSingleBSDFOperatorCoverageTransmittance(
 	FStrataEvaluateResult BSDFEvaluate = StrataEvaluateBSDF(StrataBSDFContext, Settings);
 
 	float CurrentBSDFCoverage		= CurrentBSDF.Coverage;
-	float3 CurrentBSDFTransmittanceV = BSDFEvaluate.ThroughputV;
-	float3 CurrentBSDFTransmittanceL = BSDFEvaluate.ThroughputL;
 
 	// We multiply this weight with coverage:
 	//  - Because when coming from parameter blending, we need to scale the output down according to coverage resulting from parameter blending.
 	//  - That will be applied on emissive and reflected light, so we do not need to apply that anywhere else.
 	//  - When coming from non parameter blending, this is equal identity 1 and the operator parsing will continue to correctly weight the BSDF according to the tree.
 	CurrentBSDF.LuminanceWeightV = float3(1.0f, 1.0f, 1.0f) * CurrentBSDFCoverage;
-	CurrentBSDF.LuminanceWeightL = float3(1.0f, 1.0f, 1.0f); // Don't weight by CurrentBSDFCoverage, otherwise coverage would be applied twice
+
+	CurrentBSDF.TransmittanceAboveAlongN = float3(1.0f, 1.0f, 1.0f); // Initialise to "no matter above"
 
 	// All BSDF are considered for contribution as the top BSDF can have their coverage dip to 0 and then another BSDF need to be considered for the top layer representation.
 	CurrentBSDF.TopLayerDataWeight = 1.0f;
 
-	Op.Coverage = CurrentBSDFCoverage;
-	Op.TransmittanceV = CurrentBSDFTransmittanceV;
-	Op.TransmittanceL = CurrentBSDFTransmittanceL;
+	Op.Coverage				= CurrentBSDFCoverage;
+	Op.ThroughputAlongV		= BSDFEvaluate.ThroughputV;
+	Op.TransmittanceAlongN	= BSDFEvaluate.TransmittanceAlongN;
 
 	// Evaluate refraction lobes.
 	const FBxDFEnergyTerms EnergyTerms = ComputeGGXSpecEnergyTerms(Roughness, StrataBSDFContext.SatNoV, StrataGetBSDFSpecularColor(CurrentBSDF), StrataGetBSDFSpecularF90(CurrentBSDF));
@@ -120,7 +119,7 @@ void UpdateSingleBSDFOperatorCoverageTransmittance(
 		const float TopDownInterfaceEta = CurrentBSDF.bIsTop ? AIR_IOR / WATER_IOR : 1.0f;
 		FStrataLobeStatistic RefractedLobe = StrataGetRefractedLobe(StrataGetDiracLobe(V), EnergyTerms.E, Roughness, TopDownInterfaceEta);
 		// Account for transmittance of the medium (but ignoring scattering for now)
-		RefractedLobe.E *= CurrentBSDFTransmittanceV;
+		RefractedLobe.E *= BSDFEvaluate.ThroughputV;
 		// Reduce according to coverage. It is a approximation to handle rough refraction reduction according to coverage when relying on a single lob.
 		RefractedLobe = StrataWeightLobe(RefractedLobe, CurrentBSDFCoverage);
 		Op.TopDownRefractionLobeStat = RefractedLobe;
@@ -132,7 +131,7 @@ void UpdateSingleBSDFOperatorCoverageTransmittance(
 		const float3 NormalDown = float3(0, 0, -1);
 		FStrataLobeStatistic RefractedLobe = StrataGetRefractedLobe(StrataGetDiracLobe(NormalDown), EnergyTerms.E, Roughness, BottomUpInterfaceEta);
 		// Account for transmittance of the medium (but ignoring scattering for now)
-		RefractedLobe.E *= CurrentBSDFTransmittanceV;
+		RefractedLobe.E *= BSDFEvaluate.ThroughputV;
 		// Reduce according to coverage. It is a approximation to handle rough refraction reduction according to coverage when relying on a single lob.
 		RefractedLobe = StrataWeightLobe(RefractedLobe, CurrentBSDFCoverage);
 		Op.BottomUpRefractionLobeStat = RefractedLobe;
@@ -166,9 +165,9 @@ void UpdateSingleOperatorCoverageTransmittance(inout FStrataTree StrataTree, int
 	case STRATA_OPERATOR_WEIGHT:
 	{
 		const float Weight = saturate(Op.Weight);
-		Op.Coverage = Weight * OpA.Coverage;
-		Op.TransmittanceV = OpA.TransmittanceV;
-		Op.TransmittanceL = OpA.TransmittanceL;
+		Op.Coverage				= Weight * OpA.Coverage;
+		Op.ThroughputAlongV		= OpA.ThroughputAlongV;
+		Op.TransmittanceAlongN	= OpA.TransmittanceAlongN;
 
 		Op.TopDownRefractionLobeStat = StrataWeightLobe(OpA.TopDownRefractionLobeStat, Weight);
 		Op.TopDownRefractionWorldNormal = OpA.TopDownRefractionWorldNormal * Weight;
@@ -184,10 +183,10 @@ void UpdateSingleOperatorCoverageTransmittance(inout FStrataTree StrataTree, int
 
 	case STRATA_OPERATOR_VERTICAL:
 	{
-		const float3 TopTransmittanceV = OpA.TransmittanceV;
-		const float3 TopTransmittanceL = OpA.TransmittanceL;
-		const float3 BotTransmittanceV = OpB.TransmittanceV;
-		const float3 BotTransmittanceL = OpB.TransmittanceL;
+		const float3 TopThroughputAlongV	= OpA.ThroughputAlongV;
+		const float3 TopTransmittanceAlongN = OpA.TransmittanceAlongN;
+		const float3 BotThroughputAlongV	= OpB.ThroughputAlongV;
+		const float3 BotTransmittanceAlongN = OpB.TransmittanceAlongN;
 		const float TopCoverage = OpA.Coverage;
 		const float BotCoverage = OpB.Coverage;
 
@@ -198,36 +197,36 @@ void UpdateSingleOperatorCoverageTransmittance(inout FStrataTree StrataTree, int
 
 		// If a layer has a lower coverage than the other, we account for that when compuing the Transmittance.
 		const float TransmittanceOne = 1.0f;
-		const float3 TopTransmittanceAdjustedV = lerp(TransmittanceOne, TopTransmittanceV, TopCoverage * LayerMaxCoverageInv);
-		const float3 TopTransmittanceAdjustedL = lerp(TransmittanceOne, TopTransmittanceL, TopCoverage * LayerMaxCoverageInv);
-		const float3 BotTransmittanceAdjustedV = lerp(TransmittanceOne, BotTransmittanceV, BotCoverage * LayerMaxCoverageInv);
-		const float3 BotTransmittanceAdjustedL = lerp(TransmittanceOne, BotTransmittanceL, BotCoverage * LayerMaxCoverageInv);
+		const float3 TopTransmittanceAdjustedV = lerp(TransmittanceOne, TopThroughputAlongV,	TopCoverage * LayerMaxCoverageInv);
+		const float3 TopTransmittanceAdjustedL = lerp(TransmittanceOne, TopTransmittanceAlongN, TopCoverage * LayerMaxCoverageInv);
+		const float3 BotTransmittanceAdjustedV = lerp(TransmittanceOne, BotThroughputAlongV,	BotCoverage * LayerMaxCoverageInv);
+		const float3 BotTransmittanceAdjustedL = lerp(TransmittanceOne, BotTransmittanceAlongN, BotCoverage * LayerMaxCoverageInv);
 
 		// Now we update the material transmittance and coverage
 		Op.Coverage = LayerMaxCoverage;
-		Op.TransmittanceV = TopTransmittanceAdjustedV * BotTransmittanceAdjustedV;
-		Op.TransmittanceL = TopTransmittanceAdjustedL * BotTransmittanceAdjustedL;
+		Op.ThroughputAlongV		= TopTransmittanceAdjustedV * BotTransmittanceAdjustedV;
+		Op.TransmittanceAlongN	= TopTransmittanceAdjustedL * BotTransmittanceAdjustedL;
 #else
 		FVerticalLayeringInfo Info = GetVerticalLayeringInfo(TopCoverage, BotCoverage);
 
-		Op.Coverage					= Info.Coverage;
-		Op.TransmittanceV			= Info.TransmittanceOnlyTop * TopTransmittanceV + Info.TransmittanceOnlyBottom * BotTransmittanceV + Info.TransmittanceTopAndBottom * TopTransmittanceV * BotTransmittanceV;
-		Op.TransmittanceL			= Info.TransmittanceOnlyTop * TopTransmittanceL + Info.TransmittanceOnlyBottom * BotTransmittanceL + Info.TransmittanceTopAndBottom * TopTransmittanceL * BotTransmittanceL;
+		Op.Coverage				= Info.Coverage;
+		Op.ThroughputAlongV		= Info.TransmittanceOnlyTop * TopThroughputAlongV +		Info.TransmittanceOnlyBottom * BotThroughputAlongV +	Info.TransmittanceTopAndBottom * TopThroughputAlongV * BotThroughputAlongV;
+		Op.TransmittanceAlongN	= Info.TransmittanceOnlyTop * TopTransmittanceAlongN +	Info.TransmittanceOnlyBottom * BotTransmittanceAlongN + Info.TransmittanceTopAndBottom * TopTransmittanceAlongN * BotTransmittanceAlongN;
 #endif
 
-		Op.VerticalTop_Coverage		  = OpA.Coverage;
-		Op.VerticalTop_TransmittanceV = OpA.TransmittanceV;
-		Op.VerticalTop_TransmittanceL = OpA.TransmittanceL;
+		Op.VerticalTop_Coverage				= OpA.Coverage;
+		Op.VerticalTop_ThroughputAlongV		= OpA.ThroughputAlongV;
+		Op.VerticalTop_TransmittanceAlongN	= OpA.TransmittanceAlongN;
 
 		// Project top lob through bottom layer assuming medium have same IOR.
 		{
 			const float TopDownInterfaceEta = 1.0f;
 			Op.TopDownRefractionLobeStat = StrataGetRefractedLobe(StrataOppositeLobe(OpA.TopDownRefractionLobeStat), 
 				/*InterfaceFDG*/0.0f, StrataLobeVarianceToRoughness(OpB.TopDownRefractionLobeStat.Sigma), TopDownInterfaceEta);
-			Op.TopDownRefractionLobeStat.E *= saturate(OpB.TransmittanceV + (1.0f - OpB.Coverage)); // Combine with botton layer medium and specular transmittance
+			Op.TopDownRefractionLobeStat.E *= saturate(OpB.ThroughputAlongV + (1.0f - OpB.Coverage)); // Combine with botton layer medium and specular transmittance
 
 			// The ratio of BottomSurfaceOnlyVisibile and visible through top layer (with a scale of 0.5 because at maximum both surfaces will be visible) as a fonction of the total coverage.
-			const float RefractionNormalMix = saturate((Info.SurfaceBottom + 0.5f * Info.TransmittanceTopAndBottom * dot(TopTransmittanceV, 0.33.xxx)) / max(STRATA_EPSILON, Info.Coverage));
+			const float RefractionNormalMix = saturate((Info.SurfaceBottom + 0.5f * Info.TransmittanceTopAndBottom * dot(TopThroughputAlongV, 0.33.xxx)) / max(STRATA_EPSILON, Info.Coverage));
 			Op.TopDownRefractionWorldNormal = lerp(OpA.TopDownRefractionWorldNormal, OpB.TopDownRefractionWorldNormal, RefractionNormalMix);
 		}
 
@@ -237,7 +236,7 @@ void UpdateSingleOperatorCoverageTransmittance(inout FStrataTree StrataTree, int
 			const float BottomUpInterfaceEta = bIsTop ? WATER_IOR / AIR_IOR : 1.0f;
 			Op.BottomUpRefractionLobeStat = StrataGetRefractedLobe(StrataOppositeLobe(OpB.BottomUpRefractionLobeStat), 
 				/*InterfaceFDG*/0.0f, StrataLobeVarianceToRoughness(OpA.BottomUpRefractionLobeStat.Sigma), BottomUpInterfaceEta);
-			Op.BottomUpRefractionLobeStat.E *= saturate(OpA.TransmittanceV + (1.0f - OpA.Coverage)); // Combine with top layer medium and specular transmittance
+			Op.BottomUpRefractionLobeStat.E *= saturate(OpA.ThroughputAlongV + (1.0f - OpA.Coverage)); // Combine with top layer medium and specular transmittance
 		}
 
 		// Vertical layering accumulates both operators thickness and only keep the top layer roughness and coverage, see OpaqueRoughRefractTopRoughness declaration for the details.
@@ -257,8 +256,8 @@ void UpdateSingleOperatorCoverageTransmittance(inout FStrataTree StrataTree, int
 		const float BMix = Mix;
 
 		Op.Coverage = AMix * OpA.Coverage + BMix * OpB.Coverage;
-		Op.TransmittanceV = (AMix * OpA.Coverage * OpA.TransmittanceV + BMix * OpB.Coverage * OpB.TransmittanceV) / max(STRATA_EPSILON, AMix * OpA.Coverage + BMix * OpB.Coverage);
-		Op.TransmittanceL = (AMix * OpA.Coverage * OpA.TransmittanceL + BMix * OpB.Coverage * OpB.TransmittanceL) / max(STRATA_EPSILON, AMix * OpA.Coverage + BMix * OpB.Coverage);
+		Op.ThroughputAlongV    = (AMix * OpA.Coverage * OpA.ThroughputAlongV    + BMix * OpB.Coverage * OpB.ThroughputAlongV)    / max(STRATA_EPSILON, AMix * OpA.Coverage + BMix * OpB.Coverage);
+		Op.TransmittanceAlongN = (AMix * OpA.Coverage * OpA.TransmittanceAlongN + BMix * OpB.Coverage * OpB.TransmittanceAlongN) / max(STRATA_EPSILON, AMix * OpA.Coverage + BMix * OpB.Coverage);
 
 		Op.TopDownRefractionLobeStat = StrataHorizontalMixLobes(OpA.TopDownRefractionLobeStat, OpB.TopDownRefractionLobeStat, BMix);
 		Op.TopDownRefractionWorldNormal = lerp(OpA.TopDownRefractionWorldNormal, OpB.TopDownRefractionWorldNormal, BMix);
@@ -279,8 +278,8 @@ void UpdateSingleOperatorCoverageTransmittance(inout FStrataTree StrataTree, int
 		const float AMixFactor = SafeABSDFCoverage / max(STRATA_EPSILON, SafeABSDFCoverage + SafeBBSDFCoverage);
 
 		Op.Coverage = saturate(OpA.Coverage + OpB.Coverage);
-		Op.TransmittanceV = lerp(OpB.TransmittanceV, OpA.TransmittanceV, AMixFactor);
-		Op.TransmittanceL = lerp(OpB.TransmittanceL, OpA.TransmittanceL, AMixFactor);
+		Op.ThroughputAlongV    = lerp(OpB.ThroughputAlongV,    OpA.ThroughputAlongV,    AMixFactor);
+		Op.TransmittanceAlongN = lerp(OpB.TransmittanceAlongN, OpA.TransmittanceAlongN, AMixFactor);
 
 		Op.TopDownRefractionLobeStat = StrataHorizontalMixLobes(OpA.TopDownRefractionLobeStat, OpB.TopDownRefractionLobeStat, 0.5f);
 		Op.TopDownRefractionWorldNormal = lerp(OpA.TopDownRefractionWorldNormal, OpB.TopDownRefractionWorldNormal, 0.5f);
@@ -297,8 +296,8 @@ void UpdateSingleOperatorCoverageTransmittance(inout FStrataTree StrataTree, int
 	{
 		// Thin-film does not affect these property
 		Op.Coverage = OpA.Coverage;
-		Op.TransmittanceV = OpA.TransmittanceV;
-		Op.TransmittanceL = OpA.TransmittanceL;
+		Op.ThroughputAlongV = OpA.ThroughputAlongV;
+		Op.TransmittanceAlongN = OpA.TransmittanceAlongN;
 		Op.TopDownRefractionLobeStat = OpA.TopDownRefractionLobeStat;
 		Op.TopDownRefractionWorldNormal = OpA.TopDownRefractionWorldNormal;
 		Op.BottomUpRefractionLobeStat = OpA.BottomUpRefractionLobeStat;
@@ -332,7 +331,6 @@ void UpdateAllBSDFWithBottomUpOperatorVisit_Weight(
 
 	const float Weight = saturate(Op.Weight);
 	CurrentBSDF.LuminanceWeightV *= Weight;
-	//CurrentBSDF.LuminanceWeightL *= Weight; A weight affecting the coverage of a BSDF won't affect the weight to light. Only the LuminanceWeightV to simulate a change of coverage.
 	CurrentBSDF.Coverage *= Weight;
 
 	CurrentBSDF.TopLayerDataWeight *= Weight;
@@ -353,7 +351,6 @@ void UpdateAllBSDFWithBottomUpOperatorVisit_Horizontal(
 	const float BMix = Mix;
 	const float Weight = PreviousIsInputA > 0 ? AMix : BMix;
 	CurrentBSDF.LuminanceWeightV *= Weight;
-	//CurrentBSDF.LuminanceWeightL *= Weight; A weight affecting the coverage of a BSDF won't affect the weight to light. Only the LuminanceWeightV to simulate a change of coverage.
 	CurrentBSDF.Coverage *= Weight;
 
 	CurrentBSDF.TopLayerDataWeight *= Weight;
@@ -371,8 +368,8 @@ void UpdateAllBSDFWithBottomUpOperatorVisit_Vertical(
 
 	// PreviousIsInputA > 0 means it is a BSDF affecting the top layer so we do not affect it by anything for this vertical layering.
 	// Otherise, the BSDF comes from the bottom part so we affect it weights from the gathered top layer coverage/transmittance
-	CurrentBSDF.LuminanceWeightV *= PreviousIsInputA > 0 ? 1.0f : saturate(1.0f - Op.VerticalTop_Coverage) + saturate(Op.VerticalTop_Coverage) * Op.VerticalTop_TransmittanceV;
-	CurrentBSDF.LuminanceWeightL *= PreviousIsInputA > 0 ? 1.0f : saturate(1.0f - Op.VerticalTop_Coverage) + saturate(Op.VerticalTop_Coverage) * Op.VerticalTop_TransmittanceL;
+	CurrentBSDF.LuminanceWeightV		*= PreviousIsInputA > 0 ? 1.0f : saturate(1.0f - Op.VerticalTop_Coverage) + saturate(Op.VerticalTop_Coverage) * Op.VerticalTop_ThroughputAlongV;
+	CurrentBSDF.TransmittanceAboveAlongN*= PreviousIsInputA > 0 ? 1.0f : saturate(1.0f - Op.VerticalTop_Coverage) + saturate(Op.VerticalTop_Coverage) * Op.VerticalTop_TransmittanceAlongN;
 
 	// If the BSDF is from the bottom, we reduce its contribution according to the top layer coverage only.
 	CurrentBSDF.TopLayerDataWeight *= PreviousIsInputA ? 1.0 : saturate(1.0f - Op.VerticalTop_Coverage);
diff --git a/Engine/Shaders/Private/Strata/StrataVisualize.usf b/Engine/Shaders/Private/Strata/StrataVisualize.usf
--- a/Engine/Shaders/Private/Strata/StrataVisualize.usf
+++ b/Engine/Shaders/Private/Strata/StrataVisualize.usf
@@ -134,9 +134,9 @@ void PrintBSDF(inout FShaderPrintContext Context, inout FStrataAddressing Strata
 	{
 		case STRATA_BSDF_TYPE_SLAB:
 		{
-			const bool bWeightL = BSDF_GETHASWEIGHT_L(BSDF);
+			const bool bWeightL = BSDF_GETHASTRANSABOVE(BSDF);
 			const bool bGreyWeightV = BSDF_GETHASGREYWEIGHT_V(BSDF);
-			const bool bGreyWeightL = BSDF.LuminanceWeightL.x == BSDF.LuminanceWeightL.y && BSDF.LuminanceWeightL.x == BSDF.LuminanceWeightL.z;
+			const bool bGreyTopTrans = BSDF.TransmittanceAboveAlongN.x == BSDF.TransmittanceAboveAlongN.y && BSDF.TransmittanceAboveAlongN.x == BSDF.TransmittanceAboveAlongN.z;
 
 			Print(Context, TEXT("Slab "), FontBSDFType);
 			Newline(Context);
@@ -174,18 +174,18 @@ void PrintBSDF(inout FShaderPrintContext Context, inout FStrataAddressing Strata
 			}
 
 			// Lighy weight
-			if (bWeightL)
+			if (bGreyTopTrans)
 			{
 				Print(Context, TEXT(" "), FontSilver);
-				if (bGreyWeightL)
+				if (bGreyTopTrans)
 				{
-					Print(Context, BSDF.LuminanceWeightL.x, FontSilver);
+					Print(Context, BSDF.TransmittanceAboveAlongN.x, FontSilver);
 				}
 				else	
 				{
-					Print(Context, BSDF.LuminanceWeightL.r, FontLightRed,   6, 4);
-					Print(Context, BSDF.LuminanceWeightL.g, FontLightGreen, 6, 4);
-					Print(Context, BSDF.LuminanceWeightL.b, FontLightBlue,  6, 4);
+					Print(Context, BSDF.TransmittanceAboveAlongN.r, FontLightRed,   6, 4);
+					Print(Context, BSDF.TransmittanceAboveAlongN.g, FontLightGreen, 6, 4);
+					Print(Context, BSDF.TransmittanceAboveAlongN.b, FontLightBlue,  6, 4);
 				}
 			}
 			Newline(Context);
@@ -643,10 +643,10 @@ void MaterialDebugStrataTreeCS(uint3 DispatchThreadId : SV_DispatchThreadID)
 			Print(Context, BSDF.LuminanceWeightV.g, FontLightGreen, 5, 2);
 			Print(Context, BSDF.LuminanceWeightV.b, FontLightBlue, 5, 2);
 			Newline(Context);
-			Print(Context, TEXT("LumWghtL = "), BSDFFont);
-			Print(Context, BSDF.LuminanceWeightL.r, FontLightRed, 5, 2);
-			Print(Context, BSDF.LuminanceWeightL.g, FontLightGreen, 5, 2);
-			Print(Context, BSDF.LuminanceWeightL.b, FontLightBlue, 5, 2);
+			Print(Context, TEXT("TopTrans = "), BSDFFont);
+			Print(Context, BSDF.TransmittanceAboveAlongN.r, FontLightRed, 5, 2);
+			Print(Context, BSDF.TransmittanceAboveAlongN.g, FontLightGreen, 5, 2);
+			Print(Context, BSDF.TransmittanceAboveAlongN.b, FontLightBlue, 5, 2);
 			Newline(Context);
 			Print(Context, TEXT("Diffuse  = "), BSDFFont);
 			const float3 BSDFDiffuse = StrataGetBSDFDiffuseColor(BSDF);
@@ -863,7 +863,7 @@ void EvaluateLighting(
 
 	// Patch the BSDF to appear as a top layer
 	BSDF.LuminanceWeightV = 1.0f;
-	BSDF.LuminanceWeightL = 1.0f;
+	BSDF.TransmittanceAboveAlongN = 1.0f;
 	BSDF.bIsBottom = true;
 	BSDF.bIsTop = true;
 	BSDF.Coverage = 1.0f;
