Strata - Added TransmittanceBelowAlongN which will be useful later with the path tracer when evaluating a single BSDF and choosing a refractions path.

diff --git a/Engine/Shaders/Private/Strata/Strata.ush b/Engine/Shaders/Private/Strata/Strata.ush
--- a/Engine/Shaders/Private/Strata/Strata.ush
+++ b/Engine/Shaders/Private/Strata/Strata.ush
@@ -100,8 +100,8 @@ struct FStrataBSDF
 	int					OperatorIndex;
 	float3				LuminanceWeightV;
 
-	float3				TransmittanceAboveAlongN;	// The transmittance towards the top of the material along the normal ignoring this BSDF, e.g. only accumulating matter above it.
-//	float3				TransmittanceBelow;			// TODO
+	float3				TransmittanceAboveAlongN;	// The transmittance towards the top of the material along the normal and ignoring this BSDF, e.g. only accumulating matter above it.
+	float3				TransmittanceBelowAlongN;	// The transmittance towards the bottom of the material along the normal and ignoring this BSDF, e.g. only accumulating matter below it.
 
 	int					bIsBottom;
 	int					bIsTop;
@@ -747,6 +747,10 @@ struct FStrataOperator
 	float3 VerticalTop_ThroughputAlongV;				// Only used by vertical layering operators. It represent the transmittance of matter of the top layer (from the left sub tree, and towards the view).
 	float3 VerticalTop_TransmittanceAlongN;				// Only used by vertical layering operators. It represent the transmittance of matter of the top layer (from the left sub tree, and along the normal of the surface).
 
+	float  VerticalBot_Coverage;						// Analoguous to VerticalTop_ data but for the bottom layer.
+	float3 VerticalBot_ThroughputAlongV;				// Idem.
+	float3 VerticalBot_TransmittanceAlongN;				// Idem.
+
 	// Rough refraction
 	float3				 TopDownRefractionWorldNormal;	// The normal used for the single refracted direction
 	FStrataLobeStatistic TopDownRefractionLobeStat;		// The lobe used for the scene rough refraction from front to back, e.g. used for translucent rough refraciton
@@ -828,6 +832,7 @@ void StrataPackOutStrataTree(
 		WRITE_INT(StrataTree.BSDFs[i].OperatorIndex);
 		WRITE_FLOAT3(StrataTree.BSDFs[i].LuminanceWeightV);
 		WRITE_FLOAT3(StrataTree.BSDFs[i].TransmittanceAboveAlongN);
+		WRITE_FLOAT3(StrataTree.BSDFs[i].TransmittanceBelowAlongN);
 		WRITE_UINT(StrataTree.BSDFs[i].bIsBottom);
 		WRITE_UINT(StrataTree.BSDFs[i].bIsTop);
 
@@ -915,7 +920,7 @@ FStrataBSDF StrataUnpackInStrataTreeBSDF(
 	inout FStrataTreeHeader StrataTreeHeader,
 	in int SliceStoringDebugStrataTree)
 {
-	int OutCoord = BSDFIndex * (STRATA_INLINE_SHADING ? 39 : 30);		// Skip BSDFs
+	int OutCoord = BSDFIndex * (STRATA_INLINE_SHADING ? 42 : 33);		// Skip BSDFs
 	int OutCoordY= 1;
 
 	FStrataBSDF BSDF = (FStrataBSDF)0;
@@ -924,6 +929,7 @@ FStrataBSDF StrataUnpackInStrataTreeBSDF(
 	READ_INT(BSDF.OperatorIndex);
 	READ_FLOAT3(BSDF.LuminanceWeightV);
 	READ_FLOAT3(BSDF.TransmittanceAboveAlongN);
+	READ_FLOAT3(BSDF.TransmittanceBelowAlongN);
 	READ_UINT(BSDF.bIsBottom);
 	READ_UINT(BSDF.bIsTop);
 
diff --git a/Engine/Shaders/Private/Strata/StrataTree.ush b/Engine/Shaders/Private/Strata/StrataTree.ush
--- a/Engine/Shaders/Private/Strata/StrataTree.ush
+++ b/Engine/Shaders/Private/Strata/StrataTree.ush
@@ -218,6 +218,10 @@ void UpdateSingleOperatorCoverageTransmittance(inout FStrataTree StrataTree, int
 		Op.VerticalTop_ThroughputAlongV		= OpA.ThroughputAlongV;
 		Op.VerticalTop_TransmittanceAlongN	= OpA.TransmittanceAlongN;
 
+		Op.VerticalBot_Coverage				= OpB.Coverage;
+		Op.VerticalBot_ThroughputAlongV		= OpB.ThroughputAlongV;
+		Op.VerticalBot_TransmittanceAlongN	= OpB.TransmittanceAlongN;
+
 		// Project top lob through bottom layer assuming medium have same IOR.
 		{
 			const float TopDownInterfaceEta = 1.0f;
@@ -370,6 +374,7 @@ void UpdateAllBSDFWithBottomUpOperatorVisit_Vertical(
 	// Otherise, the BSDF comes from the bottom part so we affect it weights from the gathered top layer coverage/transmittance
 	CurrentBSDF.LuminanceWeightV		*= PreviousIsInputA > 0 ? 1.0f : saturate(1.0f - Op.VerticalTop_Coverage) + saturate(Op.VerticalTop_Coverage) * Op.VerticalTop_ThroughputAlongV;
 	CurrentBSDF.TransmittanceAboveAlongN*= PreviousIsInputA > 0 ? 1.0f : saturate(1.0f - Op.VerticalTop_Coverage) + saturate(Op.VerticalTop_Coverage) * Op.VerticalTop_TransmittanceAlongN;
+	CurrentBSDF.TransmittanceBelowAlongN*= PreviousIsInputA ==0 ? 1.0f : saturate(1.0f - Op.VerticalBot_Coverage) + saturate(Op.VerticalBot_Coverage) * Op.VerticalBot_TransmittanceAlongN;
 
 	// If the BSDF is from the bottom, we reduce its contribution according to the top layer coverage only.
 	CurrentBSDF.TopLayerDataWeight *= PreviousIsInputA ? 1.0 : saturate(1.0f - Op.VerticalTop_Coverage);
diff --git a/Engine/Shaders/Private/Strata/StrataVisualize.usf b/Engine/Shaders/Private/Strata/StrataVisualize.usf
--- a/Engine/Shaders/Private/Strata/StrataVisualize.usf
+++ b/Engine/Shaders/Private/Strata/StrataVisualize.usf
@@ -648,6 +648,11 @@ void MaterialDebugStrataTreeCS(uint3 DispatchThreadId : SV_DispatchThreadID)
 			Print(Context, BSDF.TransmittanceAboveAlongN.g, FontLightGreen, 5, 2);
 			Print(Context, BSDF.TransmittanceAboveAlongN.b, FontLightBlue, 5, 2);
 			Newline(Context);
+			Print(Context, TEXT("BotTrans = "), BSDFFont);
+			Print(Context, BSDF.TransmittanceBelowAlongN.r, FontLightRed, 5, 2);
+			Print(Context, BSDF.TransmittanceBelowAlongN.g, FontLightGreen, 5, 2);
+			Print(Context, BSDF.TransmittanceBelowAlongN.b, FontLightBlue, 5, 2);
+			Newline(Context);
 			Print(Context, TEXT("Diffuse  = "), BSDFFont);
 			const float3 BSDFDiffuse = StrataGetBSDFDiffuseColor(BSDF);
 			Print(Context, BSDFDiffuse.r, FontLightRed, 5, 2);
@@ -864,6 +869,7 @@ void EvaluateLighting(
 	// Patch the BSDF to appear as a top layer
 	BSDF.LuminanceWeightV = 1.0f;
 	BSDF.TransmittanceAboveAlongN = 1.0f;
+	BSDF.TransmittanceBelowAlongN = 1.0f;
 	BSDF.bIsBottom = true;
 	BSDF.bIsTop = true;
 	BSDF.Coverage = 1.0f;
