Fix shader compilation error.

diff --git a/Engine/Shaders/Private/Strata/StrataMaterialClassification.usf b/Engine/Shaders/Private/Strata/StrataMaterialClassification.usf
--- a/Engine/Shaders/Private/Strata/StrataMaterialClassification.usf
+++ b/Engine/Shaders/Private/Strata/StrataMaterialClassification.usf
@@ -244,7 +244,7 @@ void ArgsMainCS(uint2 DispatchThreadId : SV_DispatchThreadID)
 	if (all(DispatchThreadId == 0))
 	{
 		const uint TileCount = TileDrawIndirectDataBuffer[0].x;
-		const uint DispatchX = min(TileCount, OverflowTileCount.x);
+		const uint DispatchX = min(TileCount, uint(OverflowTileCount.x));
 		const uint DispatchY = DivideAndRoundUp(TileCount, OverflowTileCount.x);
 
 		TileDispatchIndirectDataBuffer[0] = DispatchX;
diff --git a/Engine/Shaders/Private/Strata/StrataTile.ush b/Engine/Shaders/Private/Strata/StrataTile.ush
--- a/Engine/Shaders/Private/Strata/StrataTile.ush
+++ b/Engine/Shaders/Private/Strata/StrataTile.ush
@@ -113,7 +113,6 @@ FStrataBSDFOffset UnpackBSDFOffset(uint In, uint InBSDFCount)
 		Out.BSDFCount = min(InBSDFCount, STRATA_MAX_BSDF_COUNT_FOR_BDSFOFFSET);
 		Out.BSDFOffsets[0] = (In & 0x7);
 
-		FStrataBSDFOffset Out = (FStrataBSDFOffset)0;
 		UNROLL_N(STRATA_MAX_BSDF_COUNT_FOR_BDSFOFFSET)
 		for (uint i = 1; i < Out.BSDFCount; i++)
 		{
