Strata - improved "strata tree as a graph" look.

diff --git a/Engine/Shaders/Private/Strata/StrataVisualize.usf b/Engine/Shaders/Private/Strata/StrataVisualize.usf
--- a/Engine/Shaders/Private/Strata/StrataVisualize.usf
+++ b/Engine/Shaders/Private/Strata/StrataVisualize.usf
@@ -465,6 +465,7 @@ void VisualizeMaterialPS(
 #define TREE_START_POS			float2(1300, 350)
 #define TREE_BRANCH				float2(200, 75)
 #define TREE_COLOR				float4(0.8, 0.8, 0.8, 1.0)
+#define TREE_NODE_RADIUS		20.0
 
 float2 GetBranchOffset(FStrataOperator Op, bool bIsLeft)
 {
@@ -725,7 +726,7 @@ void MaterialDebugStrataTreeCS(uint3 DispatchThreadId : SV_DispatchThreadID)
 
 			const float2 ParentPrintPos = PrintPosStack[StackPtr];
 
-			AddCircleSS(Context, ParentPrintPos, 20, TREE_COLOR);
+			AddCircleSS(Context, ParentPrintPos, TREE_NODE_RADIUS, TREE_COLOR);
 
 			if (Op.LayerDepth == TREE_MAX_DEPTH - 1 || Op.Type == STRATA_OPERATOR_BSDF || Op.Type == STRATA_OPERATOR_BSDF_LEGACY)
 			{
@@ -743,8 +744,10 @@ void MaterialDebugStrataTreeCS(uint3 DispatchThreadId : SV_DispatchThreadID)
 					NodeIndexStack[StackPtr] = Op.LeftIndex;
 					ChildrenVisitedStack[StackPtr] = 0;
 
-					PrintPosStack[StackPtr] = ParentPrintPos + GetBranchOffset(Op, true);
-					AddLineSS(Context, ParentPrintPos, PrintPosStack[StackPtr], TREE_COLOR);
+					const float2 BranchOffset = GetBranchOffset(Op, true);
+					const float2 BranchOffsetNorm = normalize(BranchOffset);
+					PrintPosStack[StackPtr] = ParentPrintPos + BranchOffset;
+					AddLineSS(Context, ParentPrintPos + TREE_NODE_RADIUS * BranchOffsetNorm, PrintPosStack[StackPtr] - TREE_NODE_RADIUS * BranchOffsetNorm, TREE_COLOR);
 				}
 				else if (Op.RightIndex >= 0)
 				{
@@ -754,8 +757,10 @@ void MaterialDebugStrataTreeCS(uint3 DispatchThreadId : SV_DispatchThreadID)
 					NodeIndexStack[StackPtr] = Op.RightIndex;
 					ChildrenVisitedStack[StackPtr] = 0;
 
-					PrintPosStack[StackPtr] = ParentPrintPos + GetBranchOffset(Op, false);
-					AddLineSS(Context, ParentPrintPos, PrintPosStack[StackPtr], TREE_COLOR);
+					const float2 BranchOffset = GetBranchOffset(Op, false);
+					const float2 BranchOffsetNorm = normalize(BranchOffset);
+					PrintPosStack[StackPtr] = ParentPrintPos + BranchOffset;
+					AddLineSS(Context, ParentPrintPos + TREE_NODE_RADIUS * BranchOffsetNorm, PrintPosStack[StackPtr] - TREE_NODE_RADIUS * BranchOffsetNorm, TREE_COLOR);
 				}
 				else
 				{
@@ -772,8 +777,10 @@ void MaterialDebugStrataTreeCS(uint3 DispatchThreadId : SV_DispatchThreadID)
 					NodeIndexStack[StackPtr] = Op.RightIndex;
 					ChildrenVisitedStack[StackPtr] = 0;
 
-					PrintPosStack[StackPtr] = ParentPrintPos + GetBranchOffset(Op, false);
-					AddLineSS(Context, ParentPrintPos, PrintPosStack[StackPtr], TREE_COLOR);
+					const float2 BranchOffset = GetBranchOffset(Op, false);
+					const float2 BranchOffsetNorm = normalize(BranchOffset);
+					PrintPosStack[StackPtr] = ParentPrintPos + BranchOffset;
+					AddLineSS(Context, ParentPrintPos + TREE_NODE_RADIUS * BranchOffsetNorm, PrintPosStack[StackPtr] - TREE_NODE_RADIUS * BranchOffsetNorm, TREE_COLOR);
 				}
 				else
 				{
