Strata - fix for potential out of bound memory access.

diff --git a/Engine/Shaders/Private/Strata/StrataTile.ush b/Engine/Shaders/Private/Strata/StrataTile.ush
--- a/Engine/Shaders/Private/Strata/StrataTile.ush
+++ b/Engine/Shaders/Private/Strata/StrataTile.ush
@@ -94,7 +94,7 @@ uint PackBSDFOffset(FStrataBSDFOffset In)
 	if (BSDFCount > 0)
 	{
 		Out = In.BSDFOffsets[0] & 0x7;
-		UNROLL_N(STRATA_MAX_BSDF_COUNT_FOR_BDSFOFFSET)
+		UNROLL_N(STRATA_MAX_BSDF_COUNT_FOR_BDSFOFFSET-1)
 		for (uint i = 1; i < BSDFCount; i++)
 		{
 			const uint Curr = In.BSDFOffsets[i];
@@ -113,7 +113,7 @@ FStrataBSDFOffset UnpackBSDFOffset(uint In, uint InBSDFCount)
 		Out.BSDFCount = min(InBSDFCount, STRATA_MAX_BSDF_COUNT_FOR_BDSFOFFSET);
 		Out.BSDFOffsets[0] = (In & 0x7);
 
-		UNROLL_N(STRATA_MAX_BSDF_COUNT_FOR_BDSFOFFSET)
+		UNROLL_N(STRATA_MAX_BSDF_COUNT_FOR_BDSFOFFSET-1)
 		for (uint i = 1; i < Out.BSDFCount; i++)
 		{
 			Out.BSDFOffsets[i] = Out.BSDFOffsets[i-1] + ((In>>(3u*i)) & 0x7);
