Extract mobile GBuffer encoding into a separate function, before adding mobile specific GBuffer layout
Store IndirectIrradiance in GbufferA.z so it has 10 bits of precision

diff --git a/Engine/Shaders/Private/DeferredShadingCommon.ush b/Engine/Shaders/Private/DeferredShadingCommon.ush
--- a/Engine/Shaders/Private/DeferredShadingCommon.ush
+++ b/Engine/Shaders/Private/DeferredShadingCommon.ush
@@ -459,6 +459,84 @@ float4 ComputeIndirectLightingSampleE(uint2 TracingPixelCoord, uint TracingRayIn
 		Hammersley16(TracingRayIndex, TracingRayCount, Seed1));
 }
 
+#if SHADING_PATH_MOBILE
+/** Mobile specific encoding of GBuffer data */
+void MobileEncodeGBuffer(
+	FGBufferData GBuffer,
+	out half4 OutGBufferA,
+	out half4 OutGBufferB,
+	out half4 OutGBufferC,
+	out half4 OutGBufferD
+)
+{
+	if (GBuffer.ShadingModelID == SHADINGMODELID_UNLIT)
+	{
+		OutGBufferA = 0;
+		OutGBufferB = 0;
+		OutGBufferC = 0;
+		OutGBufferD = 0;
+	}
+	else
+	{
+		OutGBufferA.rg = UnitVectorToOctahedron( normalize(GBuffer.WorldNormal) ) * 0.5f + 0.5f;
+#if ALLOW_STATIC_LIGHTING
+		OutGBufferA.b = EncodeIndirectIrradiance(GBuffer.IndirectIrradiance * GBuffer.GBufferAO);
+#else
+		OutGBufferA.b = 1;
+#endif
+		OutGBufferA.a = GBuffer.PerObjectGBufferData;		
+
+		OutGBufferB.r = GBuffer.Metallic;
+		OutGBufferB.g = GBuffer.Specular;
+		OutGBufferB.b = GBuffer.Roughness;
+		OutGBufferB.a = EncodeShadingModelIdAndSelectiveOutputMask(GBuffer.ShadingModelID, 0);
+
+		OutGBufferC.rgb = EncodeBaseColor( GBuffer.BaseColor );
+
+#if ALLOW_STATIC_LIGHTING
+		OutGBufferC.a = GBuffer.PrecomputedShadowFactors.x;
+#else
+		OutGBufferC.a = GBuffer.GBufferAO;
+#endif
+		OutGBufferD = GBuffer.CustomData;
+	}
+}
+
+FGBufferData MobileDecodeGBuffer(half4 InGBufferA, half4 InGBufferB, half4 InGBufferC, half4 InGBufferD)
+{
+	FGBufferData GBuffer = (FGBufferData)0;
+	GBuffer.WorldNormal = OctahedronToUnitVector(InGBufferA.xy * 2.0f - 1.0f);
+#if ALLOW_STATIC_LIGHTING
+	GBuffer.IndirectIrradiance = DecodeIndirectIrradiance(InGBufferA.z);
+#else
+	GBuffer.IndirectIrradiance = 1;
+#endif
+	GBuffer.PerObjectGBufferData = InGBufferA.a;
+
+	GBuffer.Metallic	= InGBufferB.r;
+	GBuffer.Specular	= InGBufferB.g;
+	GBuffer.Roughness	= max(0.015625, InGBufferB.b);
+	// Note: must match GetShadingModelId standalone function logic
+	// Also Note: SimpleElementPixelShader directly sets SV_Target2 ( GBufferB ) to indicate unlit.
+	// An update there will be required if this layout changes.
+	GBuffer.ShadingModelID = DecodeShadingModelId(InGBufferB.a);
+	GBuffer.SelectiveOutputMask = 0;
+	GBuffer.BaseColor = DecodeBaseColor(InGBufferC.rgb);
+#if ALLOW_STATIC_LIGHTING
+	GBuffer.GBufferAO = 1;
+	GBuffer.PrecomputedShadowFactors = InGBufferC.a;
+#else
+	GBuffer.GBufferAO = InGBufferC.a;
+	GBuffer.PrecomputedShadowFactors = 1.0;
+#endif
+
+#if	MOBILE_EXTENDED_GBUFFER
+	GBuffer.CustomData = HasCustomGBufferData(GBuffer.ShadingModelID) ? InGBufferD : 0;
+#endif
+	return GBuffer;
+}
+
+#endif //SHADING_PATH_MOBILE
 
 /** Populates OutGBufferA, B and C */
 void EncodeGBuffer(
@@ -472,6 +550,11 @@ void EncodeGBuffer(
 	float QuantizationBias = 0		// -0.5 to 0.5 random float. Used to bias quantization.
 	)
 {
+#if SHADING_PATH_MOBILE
+	MobileEncodeGBuffer(GBuffer, OutGBufferA, OutGBufferB, OutGBufferC, OutGBufferD);
+	OutGBufferE = 0;
+	OutGBufferVelocity = 0;
+#else
 	if (GBuffer.ShadingModelID == SHADINGMODELID_UNLIT)
 	{
 		OutGBufferA = 0;
@@ -482,11 +565,7 @@ void EncodeGBuffer(
 	}
 	else
 	{
-#if MOBILE_DEFERRED_SHADING
-		OutGBufferA.rg = UnitVectorToOctahedron( normalize(GBuffer.WorldNormal) ) * 0.5f + 0.5f;
-		OutGBufferA.b = GBuffer.PrecomputedShadowFactors.x;
-		OutGBufferA.a = GBuffer.PerObjectGBufferData;		
-#elif 1
+#if 1
 		OutGBufferA.rgb = EncodeNormal( GBuffer.WorldNormal );
 		OutGBufferA.a = GBuffer.PerObjectGBufferData;
 #else
@@ -524,6 +603,7 @@ void EncodeGBuffer(
 #else
 	OutGBufferVelocity = 0;
 #endif
+#endif// SHADING_PATH_MOBILE
 }
 
 // SubsurfaceProfile does deferred lighting with a checker board pixel pattern
diff --git a/Engine/Shaders/Private/MobileBasePassPixelShader.usf b/Engine/Shaders/Private/MobileBasePassPixelShader.usf
--- a/Engine/Shaders/Private/MobileBasePassPixelShader.usf
+++ b/Engine/Shaders/Private/MobileBasePassPixelShader.usf
@@ -400,14 +400,9 @@ void Main(
 	half Shadow = GetPrimaryPrecomputedShadowMask(LightmapVTPageTableResult, Interpolants, MaterialParameters).r;
 
 #if DEFERRED_SHADING_PATH
-	float4 OutGBufferE;
-	float4 OutGBufferF;
-	float4 OutGBufferVelocity = 0;
-
 	GBuffer.IndirectIrradiance = IndirectIrradiance;
 	GBuffer.PrecomputedShadowFactors.r = Shadow;
-
-	EncodeGBuffer(GBuffer, OutGBufferA, OutGBufferB, OutGBufferC, OutGBufferD, OutGBufferE, OutGBufferF, OutGBufferVelocity);
+	MobileEncodeGBuffer(GBuffer, OutGBufferA, OutGBufferB, OutGBufferC, OutGBufferD);
 #else
 
 	half4 DynamicShadowFactors = 1.0f;
diff --git a/Engine/Shaders/Private/MobileDeferredShading.usf b/Engine/Shaders/Private/MobileDeferredShading.usf
--- a/Engine/Shaders/Private/MobileDeferredShading.usf
+++ b/Engine/Shaders/Private/MobileDeferredShading.usf
@@ -35,7 +35,7 @@
 #define IS_SPOT_LIGHT 0 
 #endif 
 
-void FetchGBuffer(in float2 UV, out float4 GBufferA, out float4 GBufferB, out float4 GBufferC, out float4 GBufferD, out float SceneDepth)
+void FetchGBuffer(in float2 UV, out half4 GBufferA, out half4 GBufferB, out half4 GBufferC, out half4 GBufferD, out float SceneDepth)
 {
 	GBufferD = 0;
 
@@ -75,33 +75,12 @@ void FetchGBuffer(in float2 UV, out float4 GBufferA, out float4 GBufferB, out fl
 FGBufferData FetchAndDecodeGBuffer(in float2 UV)
 {
 	float SceneDepth = 0;
-	float4 GBufferA = 0;
-	float4 GBufferB = 0;
-	float4 GBufferC = 0;
-	float4 GBufferD = 0;
+	half4 GBufferA = 0;
+	half4 GBufferB = 0;
+	half4 GBufferC = 0;
+	half4 GBufferD = 0;
 	FetchGBuffer(UV, GBufferA, GBufferB, GBufferC, GBufferD, SceneDepth);
-	
-	FGBufferData GBuffer = (FGBufferData)0;
-	GBuffer.WorldNormal = OctahedronToUnitVector( GBufferA.xy * 2.0f - 1.0f );
-	GBuffer.PrecomputedShadowFactors = GBufferA.z;
-	GBuffer.PerObjectGBufferData = GBufferA.a;  
-	GBuffer.Metallic	= GBufferB.r;
-	GBuffer.Specular	= GBufferB.g;
-	GBuffer.Roughness	= max(0.015625, GBufferB.b);
-	// Note: must match GetShadingModelId standalone function logic
-	// Also Note: SimpleElementPixelShader directly sets SV_Target2 ( GBufferB ) to indicate unlit.
-	// An update there will be required if this layout changes.
-	GBuffer.ShadingModelID = DecodeShadingModelId(GBufferB.a);
-	GBuffer.SelectiveOutputMask = DecodeSelectiveOutputMask(GBufferB.a);
-	GBuffer.BaseColor = DecodeBaseColor(GBufferC.rgb);
-#if ALLOW_STATIC_LIGHTING
-	GBuffer.GBufferAO = 1;
-	GBuffer.IndirectIrradiance = DecodeIndirectIrradiance(GBufferC.a);
-#else
-	GBuffer.GBufferAO = GBufferC.a;
-	GBuffer.IndirectIrradiance = 1;
-#endif
-	GBuffer.CustomData = HasCustomGBufferData(GBuffer.ShadingModelID) ? GBufferD : 0;
+	FGBufferData GBuffer = MobileDecodeGBuffer(GBufferA, GBufferB, GBufferC, GBufferD);
 	GBuffer.Depth = SceneDepth;
 	return GBuffer;
 }
