Add STRATA_COMPILER.

diff --git a/Engine/Shaders/Private/Strata/Strata.ush b/Engine/Shaders/Private/Strata/Strata.ush
--- a/Engine/Shaders/Private/Strata/Strata.ush
+++ b/Engine/Shaders/Private/Strata/Strata.ush
@@ -28,7 +28,8 @@
 // DXC has issue with unrolling certain part of the code. To address temporarly the warning-as-error issues, 
 // we redefine the unroll function and fall back to loop on DXC
 // To remove this, one must test PC + DX12 + RayTracing and make sure raytracing shaders compile.
-#if COMPILER_DXC == 1
+#define STRATA_COMPILER (COMPILER_DXC)
+#if STRATA_COMPILER == 1
 #define STRATA_UNROLL LOOP
 #define STRATA_UNROLL_N(X) LOOP
 #else
@@ -820,7 +821,7 @@ void StrataPackOutStrataTree(
 	OutCoord = 0;
 	OutCoordY = 1;
 
-#if COMPILER_DXC == 1
+#if STRATA_COMPILER == 1
 	LOOP
 #else
 	StrataTree.BSDFCount = min(StrataTree.BSDFCount, STRATA_MAX_BSDF_COUNT);
@@ -859,7 +860,7 @@ void StrataPackOutStrataTree(
 	OutCoord = 0;
 	OutCoordY = 2;
 
-#if COMPILER_DXC == 1
+#if STRATA_COMPILER == 1
 	LOOP
 #else
 	StrataTree.OperatorCount = min(StrataTree.OperatorCount, STRATA_MAX_OPERATOR_COUNT);
diff --git a/Engine/Shaders/Private/Strata/StrataVisualize.usf b/Engine/Shaders/Private/Strata/StrataVisualize.usf
--- a/Engine/Shaders/Private/Strata/StrataVisualize.usf
+++ b/Engine/Shaders/Private/Strata/StrataVisualize.usf
@@ -521,7 +521,7 @@ void MaterialDebugStrataTreeCS(uint3 DispatchThreadId : SV_DispatchThreadID)
 		}
 
 		Context = InitShaderPrintContext(true, uint2(50, BSDF_ELEMENT_STARTY));
-	#if COMPILER_DXC == 1
+	#if STRATA_COMPILER == 1
 		LOOP
 	#else
 		StrataTreeHeader.OperatorCount = min(StrataTreeHeader.OperatorCount, 2);
@@ -604,7 +604,7 @@ void MaterialDebugStrataTreeCS(uint3 DispatchThreadId : SV_DispatchThreadID)
 			Newline(Context);
 		}
 
-	#if COMPILER_DXC == 1
+	#if STRATA_COMPILER == 1
 		LOOP
 	#else
 		StrataTreeHeader.BSDFCount = min(StrataTreeHeader.BSDFCount, 2);
@@ -669,7 +669,7 @@ void MaterialDebugStrataTreeCS(uint3 DispatchThreadId : SV_DispatchThreadID)
 		Print(Context, TEXT("Topology"), FontWhite);
 
 		// Debug print the strata tree only for DXC for the sake of compilation performance.
-	#if COMPILER_DXC
+	#if STRATA_COMPILER
 		uint NodeIndexStack[TREE_MAX_DEPTH];
 		uint ChildrenVisitedStack[TREE_MAX_DEPTH]; // 0:none, 1:left, 2:left&right
 		float2 PrintPosStack[TREE_MAX_DEPTH];
@@ -793,7 +793,7 @@ void MaterialDebugStrataTreeCS(uint3 DispatchThreadId : SV_DispatchThreadID)
 			}
 		}
 
-	#endif // COMPILER_DXC
+	#endif // STRATA_COMPILER
 
 	}
 }
@@ -1127,7 +1127,7 @@ void MaterialDebugStrataTreePS(
 		const float MaterialPanelSize = 200.0f;
 
 		// STRATA_TODO: read BSDF and Operators "on the fly".
-#if COMPILER_DXC == 1
+#if STRATA_COMPILER == 1
 		LOOP
 #else
 		StrataTreeHeader.BSDFCount = min(StrataTreeHeader.BSDFCount, 2);
@@ -1172,7 +1172,7 @@ void MaterialDebugStrataTreePS(
 		float2 MatTopLayerSize = float2(410, 100);
 		
 		// This is a really basic material topology debug. STRATA_TODO would be to draw the material graph itself with anottation
-#if COMPILER_DXC == 1
+#if STRATA_COMPILER == 1
 		LOOP
 #else
 		StrataTreeHeader.BSDFCount = min(StrataTreeHeader.BSDFCount, 2);
