Remove Strata shader compilation warning.

diff --git a/Engine/Shaders/Private/Strata/StrataEvaluation.ush b/Engine/Shaders/Private/Strata/StrataEvaluation.ush
--- a/Engine/Shaders/Private/Strata/StrataEvaluation.ush
+++ b/Engine/Shaders/Private/Strata/StrataEvaluation.ush
@@ -773,15 +773,15 @@ void StrataAddLobeWithMIS(inout float3 Weight, inout float Pdf, float3 LobeWeigh
 		LobeWeight *= 1 / LobeProb;
 
 		// See discussion in MISWeightRobust for why this is better than LobePdf / (Pdf + LobePdf)
-		float MISWeight;
+		float LocalMISWeight;
 		if (Pdf < LobePdf)
-			MISWeight = 1 / (1 + Pdf / LobePdf);
+			LocalMISWeight = 1 / (1 + Pdf / LobePdf);
 		else if (LobePdf < Pdf)
-			MISWeight = 1 - 1 / (1 + LobePdf / Pdf);
+			LocalMISWeight = 1 - 1 / (1 + LobePdf / Pdf);
 		else
-			MISWeight = 0.5f; // avoid (rare) inf/inf
+			LocalMISWeight = 0.5f; // avoid (rare) inf/inf
 
-		Weight = lerp(Weight, LobeWeight, MISWeight);
+		Weight = lerp(Weight, LobeWeight, LocalMISWeight);
 		Pdf += LobePdf;
 	}
 }
