Fixes a source of aliasing in TSR under movement

diff --git a/Engine/Shaders/Private/TemporalSuperResolution/TSRCommon.ush b/Engine/Shaders/Private/TemporalSuperResolution/TSRCommon.ush
--- a/Engine/Shaders/Private/TemporalSuperResolution/TSRCommon.ush
+++ b/Engine/Shaders/Private/TemporalSuperResolution/TSRCommon.ush
@@ -80,6 +80,13 @@
 #define SUB_PIXEL_PARALLAX_FACTOR_BIT_OFFSET 0
 #define SUB_PIXEL_PARALLAX_FACTOR_BIT_MASK ((1 << SUB_PIXEL_PARALLAX_FACTOR_BIT_COUNT) - 1)
 
+// Minimum number of sample of each texel when reprojecting the high frequency history.
+// Value too low causes area of low sample count to have pixel dominating too much in reprojection causing input res aliasing on texture details
+#define MIN_REPROJECTION_SAMPLE_COUNT 2.0
+
+// Minimum weight of a pixel in history.
+#define MIN_IMPLICIT_HISTORY_WEIGHT ceil(tsr_half(255.0) * rcp(tsr_half(MAX_SAMPLE_COUNT)) * tsr_half(MIN_REPROJECTION_SAMPLE_COUNT)) * rcp(tsr_half(255.0))
+
 
 // K = Center of the nearest input pixel.
 // O = Center of the output pixel.
@@ -621,8 +628,6 @@ half2 SafeRcp(half2 x)
 
 #endif
 
-#define MIN_IMPLICIT_HISTORY_WEIGHT ceil(tsr_half(255.0) * rcp(tsr_half(MAX_SAMPLE_COUNT)) * tsr_half(0.25)) * rcp(tsr_half(255.0))
-
 CALL_SITE_DEBUGLOC
 tsr_half2 ComputeLowFrequencyContribution(tsr_half2 HistoryValidity)
 {
