Apply INVARIANT to debug mode vertex shader so that it matches depth pass.

diff --git a/Engine/Shaders/Private/DebugViewModeVertexShader.usf b/Engine/Shaders/Private/DebugViewModeVertexShader.usf
--- a/Engine/Shaders/Private/DebugViewModeVertexShader.usf
+++ b/Engine/Shaders/Private/DebugViewModeVertexShader.usf
@@ -84,7 +84,7 @@ void Main(
 
 	{
 		float4 RasterizedWorldPosition = VertexFactoryGetRasterizedWorldPosition(Input, VFIntermediates, WorldPos);
-		Output.Position = mul(RasterizedWorldPosition, ResolvedView.TranslatedWorldToClip);
+		Output.Position = INVARIANT(mul(RasterizedWorldPosition, ResolvedView.TranslatedWorldToClip));
 	}
 
 	#if INSTANCED_STEREO && !MULTI_VIEW
