Remove USE_WPO_PRECISE
Instead make the output VelocityPrevScreenPosition interpolant INVARIANT which is a better match for the other approaches to INVARIANT in the shader code.
Initially the USE_WPO_PRECISE code was added for DXBC only, to prevent mismatch between depth and base pass. This type of error was later more widely fixed by 19392831.
This fix led to another need for USE_WPO_PRECISE. That came from the output position now being (correctly) INVARIANT but the output VelocityPrevScreenPosition didn't have INVARIANT. This could lead to very different calculations between current and previous position and therefore incorrect velocity vectors.

diff --git a/Engine/Shaders/Private/BasePassVertexShader.usf b/Engine/Shaders/Private/BasePassVertexShader.usf
--- a/Engine/Shaders/Private/BasePassVertexShader.usf
+++ b/Engine/Shaders/Private/BasePassVertexShader.usf
@@ -72,7 +72,7 @@ void Main(
 		float3 ODS = OffsetODS(RasterizedWorldPosition.xyz, ResolvedView.TranslatedWorldCameraOrigin.xyz, ResolvedView.StereoIPD);
 		ClipSpacePosition = mul(float4(RasterizedWorldPosition.xyz + ODS, 1.0), ResolvedView.TranslatedWorldToClip);
 	#else
-		ClipSpacePosition = INVARIANT(mul(RasterizedWorldPosition, ResolvedView.TranslatedWorldToClip));
+		ClipSpacePosition = mul(RasterizedWorldPosition, ResolvedView.TranslatedWorldToClip);
 	#endif
 		Output.Position = INVARIANT(ClipSpacePosition);
 	}
@@ -244,7 +244,7 @@ void Main(
 		}
 
 		// compute the old screen pos with the old world position and the old camera matrix
-		Output.BasePassInterpolants.VelocityPrevScreenPosition = PrevTranslatedWorldPosition; 
+		Output.BasePassInterpolants.VelocityPrevScreenPosition = INVARIANT(PrevTranslatedWorldPosition); 
 	}
 #endif	// WRITES_VELOCITY_TO_GBUFFER
 }
diff --git a/Engine/Shaders/Private/MaterialTemplate.ush b/Engine/Shaders/Private/MaterialTemplate.ush
--- a/Engine/Shaders/Private/MaterialTemplate.ush
+++ b/Engine/Shaders/Private/MaterialTemplate.ush
@@ -2735,18 +2735,6 @@ half GetMaterialMaskedOpacity(FPixelMaterialInputs PixelMaterialInputs)
 }
 #endif
 
-// UE-102725: Workaround for precision issue that only occurs with world position offsets in D3D shaders.
-#if USE_WPO_PRECISE
-#define WPO_PRECISE(T) INVARIANT(T)
-#else
-#define WPO_PRECISE(T) T
-#endif
-
-float3 GetMaterialWorldPositionOffsetInternal(FMaterialVertexParameters Parameters)
-{
-%s;
-}
-
 float3 GetMaterialWorldPositionOffset(FMaterialVertexParameters Parameters)
 {
 #if IS_NANITE_PASS
@@ -2765,11 +2753,6 @@ float3 GetMaterialWorldPositionOffset(FMaterialVertexParameters Parameters)
 	}
 #endif
 
-	return WPO_PRECISE(GetMaterialWorldPositionOffsetInternal(Parameters));
-}
-
-float3 GetMaterialPreviousWorldPositionOffsetInternal(FMaterialVertexParameters Parameters)
-{
 %s;
 }
 
@@ -2791,10 +2774,9 @@ float3 GetMaterialPreviousWorldPositionOffset(FMaterialVertexParameters Paramete
 	}
 #endif
 
-	return WPO_PRECISE(GetMaterialPreviousWorldPositionOffsetInternal(Parameters));
+%s;
 }
 
-
 // .rgb:SubsurfaceColor, .a:SSProfileId in 0..1 range
 half4 GetMaterialSubsurfaceDataRaw(FPixelMaterialInputs PixelMaterialInputs)
 {
diff --git a/Engine/Source/Developer/Windows/ShaderFormatD3D/Private/D3DShaderCompiler.cpp b/Engine/Source/Developer/Windows/ShaderFormatD3D/Private/D3DShaderCompiler.cpp
--- a/Engine/Source/Developer/Windows/ShaderFormatD3D/Private/D3DShaderCompiler.cpp
+++ b/Engine/Source/Developer/Windows/ShaderFormatD3D/Private/D3DShaderCompiler.cpp
@@ -1179,8 +1179,6 @@ void CompileD3DShader(const FShaderCompilerInput& Input, FShaderCompilerOutput&
 	AdditionalDefines.SetDefine(TEXT("COMPILER_HLSL"), 1);
 	AdditionalDefines.SetDefine(TEXT("PLATFORM_SUPPORTS_ROV"), 1); // Assume min. spec HW supports with DX12/SM5
 
-	AdditionalDefines.SetDefine(TEXT("USE_WPO_PRECISE"), 1);
-	
 	if (bUseDXC)
 	{
 		AdditionalDefines.SetDefine(TEXT("PLATFORM_SUPPORTS_SM6_0_WAVE_OPERATIONS"), 1);
