Strata - fix for very broken AO/irradiance in many cases.

diff --git a/Engine/Shaders/Private/BasePassPixelShader.usf b/Engine/Shaders/Private/BasePassPixelShader.usf
--- a/Engine/Shaders/Private/BasePassPixelShader.usf
+++ b/Engine/Shaders/Private/BasePassPixelShader.usf
@@ -1714,7 +1714,7 @@ void FPixelShaderInOut_MainPS(
 	// const FShadingOcclusion ShadingOcclusion = ApplyBentNormal(MaterialParameters.CameraVector, MaterialParameters.WorldNormal, MaterialParameters.TangentToWorld, GetBentNormalZero(MaterialParameters), GBuffer.Roughness, MaterialAO);
 	{
 		StrataPixelHeader.IrradianceAO									= InitIrradianceAndOcclusion();
-		StrataPixelHeader.IrradianceAO.MaterialAO						= GBuffer.GBufferAO;	// STRATA_TODO: AOMultiBounce(Luminance(GBuffer.SpecularColor), ShadingOcclusion.SpecOcclusion).g; // FIXME: ALLOW_STATIC_LIGHTING == 0 expects this to be AO
+		StrataPixelHeader.IrradianceAO.MaterialAO						= MaterialAO;			// STRATA_TODO: AOMultiBounce(Luminance(GBuffer.SpecularColor), ShadingOcclusion.SpecOcclusion).g; // FIXME: ALLOW_STATIC_LIGHTING == 0 expects this to be AO
 		StrataPixelHeader.IrradianceAO.IndirectIrradiance				= IndirectIrradiance;	// STRATA_TODO
 		StrataPixelHeader.IrradianceAO.DiffuseIndirectSampleOcclusion	= GetDiffuseIndirectSampleOcclusion(StrataPixelHeader.SharedLocalBases, MaterialParameters.CameraVector, In.SvPosition.xy, StrataPixelHeader.IrradianceAO.MaterialAO);
 	}
diff --git a/Engine/Shaders/Private/Strata/Strata.ush b/Engine/Shaders/Private/Strata/Strata.ush
--- a/Engine/Shaders/Private/Strata/Strata.ush
+++ b/Engine/Shaders/Private/Strata/Strata.ush
@@ -232,6 +232,19 @@ struct FStrataBSDF
 #error Strata Header is larger than 32 bits
 #endif
 
+// A header always has a common state represented by stat bits (low 8 bits) and Irradiance/AO bits (higher 8 bits)
+#define HEADER_GETCOMMONSTATES(X)					READ_BITS( X, (HEADER_BIT_COUNT_STATE + HEADER_BIT_COUNT_IRRADIANCE_AO), HEADER_BIT_OFFSET_STATE)
+#define HEADER_SETCOMMONSTATES(X, NewValue)			WRITE_BITS(X, (HEADER_BIT_COUNT_STATE + HEADER_BIT_COUNT_IRRADIANCE_AO), HEADER_BIT_OFFSET_STATE, NewValue)
+
+#define HEADER_GETSTATE(X)							READ_BITS( X, HEADER_BIT_COUNT_STATE, HEADER_BIT_OFFSET_STATE)
+#define HEADER_SETSTATE(X, NewValue)				WRITE_BITS(X, HEADER_BIT_COUNT_STATE, HEADER_BIT_OFFSET_STATE, NewValue)
+
+#define HEADER_GETIRRADIANCE_AO(X)					READ_BITS( X, HEADER_BIT_COUNT_IRRADIANCE_AO, HEADER_BIT_OFFSET_IRRADIANCE_AO)
+#define HEADER_SETIRRADIANCE_AO(X, NewValue)		WRITE_BITS(X, HEADER_BIT_COUNT_IRRADIANCE_AO, HEADER_BIT_OFFSET_IRRADIANCE_AO, NewValue)
+
+#define HEADER_GETSHADINGMODELS(X)					READ_BITS( X, HEADER_BIT_COUNT_SHADINGMODELS, HEADER_BIT_OFFSET_SHADINGMODELS)
+#define HEADER_SETSHADINGMODELS(X, NewValue)		WRITE_BITS(X, HEADER_BIT_COUNT_SHADINGMODELS, HEADER_BIT_OFFSET_SHADINGMODELS, NewValue)
+
 #define HEADER_GETBSDFCOUNT(X)						READ_BITS( X, HEADER_BIT_COUNT_BSDF_COUNT, HEADER_BIT_OFFSET_BSDF_COUNT)
 #define HEADER_SETBSDFCOUNT(X, NewValue)			WRITE_BITS(X, HEADER_BIT_COUNT_BSDF_COUNT, HEADER_BIT_OFFSET_BSDF_COUNT, NewValue)
 
@@ -241,15 +254,6 @@ struct FStrataBSDF
 #define HEADER_GETSHAREDLOCALBASESTYPE(X)			READ_BITS( X, HEADER_BIT_COUNT_SHAREDLOCALBASES_TYPE, HEADER_BIT_OFFSET_SHAREDLOCALBASES_TYPE)
 #define HEADER_SETSHAREDLOCALBASESTYPE(X, NewValue)	WRITE_BITS(X, HEADER_BIT_COUNT_SHAREDLOCALBASES_TYPE, HEADER_BIT_OFFSET_SHAREDLOCALBASES_TYPE, NewValue)
 
-#define HEADER_GETIRRADIANCE_AO(X)					READ_BITS( X, HEADER_BIT_COUNT_IRRADIANCE_AO, HEADER_BIT_OFFSET_IRRADIANCE_AO)
-#define HEADER_SETIRRADIANCE_AO(X, NewValue)		WRITE_BITS(X, HEADER_BIT_COUNT_IRRADIANCE_AO, HEADER_BIT_OFFSET_IRRADIANCE_AO, NewValue)
-
-#define HEADER_GETSTATE(X)							READ_BITS( X, HEADER_BIT_COUNT_STATE, HEADER_BIT_OFFSET_STATE)
-#define HEADER_SETSTATE(X, NewValue)				WRITE_BITS(X, HEADER_BIT_COUNT_STATE, HEADER_BIT_OFFSET_STATE, NewValue)
-
-#define HEADER_GETSHADINGMODELS(X)					READ_BITS( X, HEADER_BIT_COUNT_SHADINGMODELS, HEADER_BIT_OFFSET_SHADINGMODELS)
-#define HEADER_SETSHADINGMODELS(X, NewValue)		WRITE_BITS(X, HEADER_BIT_COUNT_SHADINGMODELS, HEADER_BIT_OFFSET_SHADINGMODELS, NewValue)
-
 ///////////////////////////////////////////////////////////////////////////////
 // BSDF State
  
@@ -487,7 +491,7 @@ uint StrataPackIrradianceAndOcclusion(FStrataIrradianceAndOcclusion In, float Qu
 	// No space for AO. Multiply IndirectIrradiance by AO instead of storing.
 	return min(0xFFu, (EncodeIndirectIrradiance(In.IndirectIrradiance * In.MaterialAO) + QuantizationBias * (1.0 / 255.0)) * 0xFFu);
 #else
-	return saturate(In.MaterialAO) * 0xFFu;
+	return uint(saturate(In.MaterialAO) * 255.0f);
 #endif
 }
 
@@ -503,7 +507,7 @@ FStrataIrradianceAndOcclusion StrataUnpackIrradianceAndOcclusion(uint In /*8bits
 	Out.DiffuseIndirectSampleOcclusion = 0x0u;
 	Out.IndirectIrradiance = DecodeIndirectIrradiance(In * rcp(255));
 #else
-	Out.MaterialAO = In * rcp(255);
+	Out.MaterialAO = float(In)* rcp(255.0f);
 	Out.DiffuseIndirectSampleOcclusion = 0x0u;
 	Out.IndirectIrradiance = 1.f;
 #endif
@@ -3145,7 +3149,7 @@ void PackStrataOut(
 		AndIsSimpleMaterial(StrataPixelHeader, BSDFVisibleCount == 1);
 		AndIsSingleMaterial(StrataPixelHeader, BSDFVisibleCount == 1);
 
-		// Store irradiance and AO into the state bits of the header
+		// Store irradiance and AO into the state bits of the header so that they can be written later as part of the common header state bits.
 	#if STRATA_INLINE_SHADING
 		HEADER_SETIRRADIANCE_AO(StrataPixelHeader.State, StrataPackIrradianceAndOcclusion(StrataPixelHeader.IrradianceAO, 0));
 	#endif
@@ -3218,7 +3222,7 @@ void PackStrataOut(
 						// Data0 (Header_State|Header_AO|Roughness|Diffuse8bits)
 						{
 							uint Out = 0;
-							HEADER_SETSTATE(Out, StrataPixelHeader.State);
+							HEADER_SETCOMMONSTATES(Out, StrataPixelHeader.State);
 
 							uint Packed16 = (PackedDiffuse8Bits << 8) | PackedRoughness8bits;
 							Out = (Out & HEADER_FASTENCODING_MASK) | (Packed16 << HEADER_FASTENCODING_BIT_COUNT);
@@ -3242,7 +3246,8 @@ void PackStrataOut(
 						// Data0 (Header_State|Header_AO|Header_BSDFTypes|BSDF_State)
 						{
 							uint Out = 0;
-							HEADER_SETSTATE(Out, StrataPixelHeader.State);
+							HEADER_SETCOMMONSTATES(Out, StrataPixelHeader.State);
+
 							HEADER_SETSHADINGMODELS(Out, GetShadingModels(StrataPixelHeader));
 							Out = Out & HEADER_SINGLEENCODING_MASK;
 							Out = Out | ((BSDF.State & STATE_SINGLEENCODING_MASK) << HEADER_SINGLEENCODING_BIT_COUNT);
@@ -3379,7 +3384,7 @@ uint PackStrataHeader(uint InBSDFCount, FStrataPixelHeader InHeader)
 	HEADER_SETSHAREDLOCALBASESTYPE(Out, InHeader.SharedLocalBases.Types);
 #endif
 
-	HEADER_SETSTATE(Out, InHeader.State);
+	HEADER_SETCOMMONSTATES(Out, InHeader.State);
 
 	return Out;
 }
@@ -3391,7 +3396,7 @@ FStrataPixelHeader UnpackStrataHeaderIn(FStrataMaterialContainer StrataBuffer, i
 	FStrataPixelHeader Out = InitialiseStrataPixelHeader();
 
 	uint PackedHeader = StrataLoadUint1(StrataBuffer, StrataAddressing);
-	Out.State		= HEADER_GETSTATE(PackedHeader);
+	Out.State		= HEADER_GETCOMMONSTATES(PackedHeader);	// Unpack state and AO data at the same time 
 	Out.BSDFCount   = HEADER_GETBSDFCOUNT(PackedHeader);
 
 	const bool bIsSimpleMaterial = (PackedHeader & HEADER_MASK_ISSIMPLEMATERIAL) > 0;
