Path Tracer: Take the normal map into account at the exit point of the random walk for better detail preservation in SSS

diff --git a/Engine/Shaders/Private/PathTracing/PathTracingCore.ush b/Engine/Shaders/Private/PathTracing/PathTracingCore.ush
--- a/Engine/Shaders/Private/PathTracing/PathTracingCore.ush
+++ b/Engine/Shaders/Private/PathTracing/PathTracingCore.ush
@@ -439,6 +439,7 @@ struct FProbeResult
 {
 	float HitT;
 	float3 WorldNormal;
+	float3 WorldSmoothNormal;
 	float3 WorldGeoNormal;
 	int FrontFace;
 
@@ -470,7 +471,8 @@ FProbeResult TraceProbeRay(RayDesc Ray, uint2 PixelCoord)
 	FPathTracingPayload Payload = UnpackPathTracingPayload(PackedPayload, Ray);
 	FProbeResult Result;
 	Result.HitT = PackedPayload.HitT;
-	Result.WorldNormal = Payload.WorldSmoothNormal;
+	Result.WorldNormal = Payload.WorldNormal;
+	Result.WorldSmoothNormal = Payload.WorldSmoothNormal;
 	Result.WorldGeoNormal = Payload.WorldGeoNormal;
 	Result.FrontFace = Payload.IsFrontFace();
 	return Result;
@@ -654,7 +656,7 @@ bool ProcessSubsurfaceRandomWalk(inout FPathTracingPayload Payload, inout float3
 				// we hit the boundary! overwrite the current payload and exit the walk with a diffuse scattering event
 				Payload.TranslatedWorldPos = Ray.Origin + ProbeResult.HitT * Ray.Direction;
 				Payload.WorldNormal = ProbeResult.WorldNormal;
-				Payload.WorldSmoothNormal = ProbeResult.WorldNormal;
+				Payload.WorldSmoothNormal = ProbeResult.WorldSmoothNormal;
 				Payload.WorldGeoNormal = ProbeResult.WorldGeoNormal;
 				// make sure normal is pointed outward from the object so we capture illumination from the exterior
 				if (ProbeResult.FrontFace != Payload.IsFrontFace())
