Fix material dual source blending on macOS and enable the feature in DataDrivenPlatformInfo.ini.   Add dual source blending slot index to pixel shader MRT defination when MATERIAL_WORKS_WITH_DUAL_SOURCE_COLOR_BLENDING is defined and dual source blending is enabled.

diff --git a/Engine/Config/Mac/DataDrivenPlatformInfo.ini b/Engine/Config/Mac/DataDrivenPlatformInfo.ini
--- a/Engine/Config/Mac/DataDrivenPlatformInfo.ini
+++ b/Engine/Config/Mac/DataDrivenPlatformInfo.ini
@@ -47,6 +47,7 @@ bSupportsPercentageCloserShadows=true
 bSupportsCapsuleShadows=true
 bSupports4ComponentUAVReadWrite=true
 bSupportsVolumeTextureAtomics=false
+bSupportsDualSourceBlending=true
 
 [ShaderPlatform METAL_MRT_MAC]
 Language=Metal
diff --git a/Engine/Shaders/Private/PixelShaderOutputCommon.ush b/Engine/Shaders/Private/PixelShaderOutputCommon.ush
--- a/Engine/Shaders/Private/PixelShaderOutputCommon.ush
+++ b/Engine/Shaders/Private/PixelShaderOutputCommon.ush
@@ -70,11 +70,19 @@ void MainPS
 		OPTIONAL_IsFrontFace
 
 #if PIXELSHADEROUTPUT_MRT0
+	#if DUAL_SOURCE_COLOR_BLENDING_ENABLED && MATERIAL_WORKS_WITH_DUAL_SOURCE_COLOR_BLENDING
+		, out float4 OutTarget0 DUAL_SOURCE_BLENDING_SLOT(0) : SV_Target0
+	#else
 		, out float4 OutTarget0 : SV_Target0
+	#endif
 #endif
 
 #if PIXELSHADEROUTPUT_MRT1
+	#if DUAL_SOURCE_COLOR_BLENDING_ENABLED && MATERIAL_WORKS_WITH_DUAL_SOURCE_COLOR_BLENDING
+		, out float4 OutTarget1 DUAL_SOURCE_BLENDING_SLOT(1) : SV_Target1
+	#else
 		, out float4 OutTarget1 : SV_Target1
+	#endif
 #endif
 
 #if PIXELSHADEROUTPUT_MRT2
@@ -229,4 +237,4 @@ void MainPS
 #if OUTPUT_PIXEL_DEPTH_OFFSET
 	OutDepth = PixelShaderOut.Depth;
 #endif 
-}
\ No newline at end of file
+}
