Fixed bad translucent shadow on translucent surfaces.
UDN https://epicgames.lightning.force.com/lightning/r/Case/5004z00001dnA8jAAE/view

diff --git a/Engine/Shaders/Private/BasePassPixelShader.usf b/Engine/Shaders/Private/BasePassPixelShader.usf
--- a/Engine/Shaders/Private/BasePassPixelShader.usf
+++ b/Engine/Shaders/Private/BasePassPixelShader.usf
@@ -273,8 +273,9 @@ float3 GetTranslucencyVolumeLighting(
 			// Apply a stable offset to the world position used for shadowing, which blurs out high frequency details in the shadowmap with many layers
 			float4 HomogeneousShadowPosition = mul(float4(LWCHackToFloat(MaterialParameters.AbsoluteWorldPosition) + MaterialParameters.LightingPositionOffset, 1), TranslucentSelfShadow.WorldToShadowMatrix);
 			float2 ShadowUVs = HomogeneousShadowPosition.xy / HomogeneousShadowPosition.w;
+			float ShadowZ = 1 - HomogeneousShadowPosition.z;
 			// Lookup the shadow density at the point being shaded
-			float3 ShadowDensity = CalculateTranslucencyShadowingDensity(ShadowUVs, HomogeneousShadowPosition.z) / GetMaterialTranslucentMultipleScatteringExtinction();
+			float3 ShadowDensity = CalculateTranslucencyShadowingDensity(ShadowUVs, ShadowZ) / GetMaterialTranslucentMultipleScatteringExtinction();
 			// Compute colored transmission based on the density that the light ray passed through
 			float3 SelfShadowing = saturate(exp(-ShadowDensity * GetMaterialTranslucentSelfShadowDensityScale()));
 			// Compute a second shadow gradient to add interesting information in the shadowed area of the first
