Strata - Added support for translucent shadow.

diff --git a/Engine/Shaders/Private/Strata/StrataForwardLighting.ush b/Engine/Shaders/Private/Strata/StrataForwardLighting.ush
--- a/Engine/Shaders/Private/Strata/StrataForwardLighting.ush
+++ b/Engine/Shaders/Private/Strata/StrataForwardLighting.ush
@@ -301,8 +301,9 @@ float3 StrataForwardLighting(
 					// Apply a stable offset to the world position used for shadowing, which blurs out high frequency details in the shadowmap with many layers
 					float4 HomogeneousShadowPosition = mul(float4(LWCHackToFloat(AbsoluteWorldPosition) + LightingPositionOffset, 1), TranslucentSelfShadow.WorldToShadowMatrix);
 					float2 ShadowUVs = HomogeneousShadowPosition.xy / HomogeneousShadowPosition.w;
+					float ShadowZ = 1 - HomogeneousShadowPosition.z;
 					// Lookup the shadow density at the point being shaded
-					float3 ShadowDensity = CalculateTranslucencyShadowingDensity(ShadowUVs, HomogeneousShadowPosition.z) / GetMaterialTranslucentMultipleScatteringExtinction();
+					float3 ShadowDensity = CalculateTranslucencyShadowingDensity(ShadowUVs, ShadowZ) / GetMaterialTranslucentMultipleScatteringExtinction();
 					// Compute colored transmission based on the density that the light ray passed through
 					float3 SelfShadowing = saturate(exp(-ShadowDensity * GetMaterialTranslucentSelfShadowDensityScale()));
 					// Compute a second shadow gradient to add interesting information in the shadowed area of the first
diff --git a/Engine/Shaders/Private/TranslucentShadowDepthShaders.usf b/Engine/Shaders/Private/TranslucentShadowDepthShaders.usf
--- a/Engine/Shaders/Private/TranslucentShadowDepthShaders.usf
+++ b/Engine/Shaders/Private/TranslucentShadowDepthShaders.usf
@@ -85,7 +85,40 @@ void MainOpacityPS(
 	// Evaluate the mask for translucent materials
 	GetMaterialClippingShadowDepth(MaterialParameters, PixelMaterialInputs);
 
-	float Density = GetMaterialOpacity(PixelMaterialInputs) * GetMaterialTranslucentShadowDensityScale();
+#if STRATA_ENABLED
+	FStrataData StrataData = PixelMaterialInputs.FrontMaterial;
+
+	FStrataAddressing NullStrataAddressing = (FStrataAddressing)0;	// Fake unused in StrataCreateBSDFContext when using Forward inline shading
+
+	FStrataPixelHeader StrataPixelHeader = InitialiseStrataPixelHeader();
+	StrataPixelHeader.StrataTree = MaterialParameters.StrataTree;
+	StrataPixelHeader.BSDFCount = MaterialParameters.StrataTree.BSDFCount;
+	StrataPixelHeader.SharedLocalBases = MaterialParameters.SharedLocalBases;
+
+	// Retrieve density from coverage and colored transmittance.
+	float Density = 0.0f;
+
+	float Sigma = 0.0f;
+	float3 RefractionWorldNormal = MaterialParameters.CameraVector;
+	FStrataTree StrataTree = StrataPixelHeader.StrataTree;
+	if (StrataTree.BSDFCount > 0)
+	{
+		const float3 V = MaterialParameters.CameraVector;
+
+		// Update tree (coverage/transmittance/luminace weights)
+		const FStrataIntegrationSettings Settings = InitStrataIntegrationSettings(false /*bForceFullyRough*/, Strata.bRoughDiffuse /*bRoughDiffuseEnabled*/, Strata.PeelLayersAboveDepth);
+		StrataUpdateTree(NullStrataAddressing, StrataPixelHeader, StrataTree, V, Settings);
+
+		const float StrataTreeRootOperatorIndex = StrataData.OperatorIndex;
+		const FStrataOperator RootOperator = StrataTree.Operators[StrataTreeRootOperatorIndex];
+
+		Density = (1.0 - RootOperator.Coverage) + RootOperator.Coverage * dot(RootOperator.ThroughputAlongV, (1.0f / 3.0f).xxx);
+	}
+#else
+	float Density = GetMaterialOpacity(PixelMaterialInputs);
+#endif
+
+	Density *= GetMaterialTranslucentShadowDensityScale();
 
 #if PERSPECTIVE_CORRECT_DEPTH
 	/** Used to normalize the outputted depth */
