Path Tracer: Fix for potential NaNs in hair shading model

diff --git a/Engine/Shaders/Private/PathTracing/Material/PathTracingHair.ush b/Engine/Shaders/Private/PathTracing/Material/PathTracingHair.ush
--- a/Engine/Shaders/Private/PathTracing/Material/PathTracingHair.ush
+++ b/Engine/Shaders/Private/PathTracing/Material/PathTracingHair.ush
@@ -40,7 +40,7 @@ float ZinkeHairDiffuse(float3 L)
 	// optimized fit to avoid acos
 	// error is around ~0.001 
 	const float CosPhiCosTheta = L.y;
-	return 0.0786133 * CosTheta + CosPhiCosTheta * (0.125 + 0.0463867 * CosPhiCosTheta * rcp(CosTheta));
+	return 0.0786133 * CosTheta + CosPhiCosTheta * (0.125 + 0.0463867 * CosPhiCosTheta * rcp(CosTheta + 1e-12));
 #endif
 }
 
@@ -128,8 +128,8 @@ float FastASin(float x) {
 float Mp(float SinThetaI, float SinThetaO, float v)
 {
 	float rv = rcp(v);
-	float CosThetaO2 = 1 - SinThetaO * SinThetaO;
-	float CosThetaI2 = 1 - SinThetaI * SinThetaI;
+	float CosThetaO2 = saturate(1 - SinThetaO * SinThetaO);
+	float CosThetaI2 = saturate(1 - SinThetaI * SinThetaI);
 	float CosProduct = sqrt(CosThetaO2 * CosThetaI2);
 	// "Path tracing in Production" - Siggraph 2018 Course
 	// "From Bricks to Bunnies - Adapting a Raytracer to New Requirements" - Luke Emrose
