Changes TSR's UpdateHistory pass to sample the history before the input color to control the upscaling kernel based on history's state

diff --git a/Engine/Shaders/Private/TemporalSuperResolution/TSRUpdateHistory.usf b/Engine/Shaders/Private/TemporalSuperResolution/TSRUpdateHistory.usf
--- a/Engine/Shaders/Private/TemporalSuperResolution/TSRUpdateHistory.usf
+++ b/Engine/Shaders/Private/TemporalSuperResolution/TSRUpdateHistory.usf
@@ -72,7 +72,7 @@
 #define CONFIG_LOAD_TRANSLUCENCY_REJECTION 0
 
 #if CONFIG_COMPILE_FP16 && !CONFIG_SCENE_COLOR_APLHA
-	#define CONFIG_MANUAL_LDS_SPILL 1
+	#define CONFIG_MANUAL_LDS_SPILL 0 // TODO
 #else
 	#define CONFIG_MANUAL_LDS_SPILL 0
 #endif
@@ -565,193 +565,6 @@
 		LowFrequencyRejection = min(LowFrequencyRejection, TranslucencyRejection);
 	#endif
 	
-	// Filter input scene color at predictor frequency.
-	tsr_halfCx2 FilteredInputColor;
-	tsr_halfCx2 InputMinColor;
-	tsr_halfCx2 InputMaxColor;
-	
-	tsr_half4x2 FilteredTranslucencyColor;
-	tsr_half4x2 TranslucencyMinColor;
-	tsr_half4x2 TranslucencyMaxColor;
-
-	tsr_half2 InputPixelAlignement;
-	tsr_half2 ClosestInputLuma4;
-	
-	
-	#if CONFIG_UPDATE_SUBPIXEL_DETAIL
-		bool2 bUpdateSubpixelDetails;
-		bool2x2 bKOSign;
-	#endif
-
-	PLATFORM_SPECIFIC_ISOLATE
-	{
-		// Vector in pixel between pixel K -> O.
-		tsr_half2x2 dKO = tsr_half2x2(PPCo - PPCk);
-
-		// Compresses the sign of dKO into SGPR to keep VGPR pressure low.
-		#if CONFIG_UPDATE_SUBPIXEL_DETAIL
-		{
-			float SizeOfOutputPixelInInputPixelWidth = 0.5 * (InputInfo_ViewportSize.x * HistoryInfo_ViewportSizeInverse.x);
-
-			bUpdateSubpixelDetails[0] = all(abs(dpv_lo(dKO)) < SizeOfOutputPixelInInputPixelWidth);
-			bUpdateSubpixelDetails[1] = all(abs(dpv_hi(dKO)) < SizeOfOutputPixelInInputPixelWidth);
-		
-			bKOSign[0] = dKO[0] < 0.0;
-			bKOSign[1] = dKO[1] < 0.0;
-		}
-		#endif
-		
-		tsr_halfC RawInputColorArray[CONFIG_SAMPLES][DPV_PIXEL_PER_LANE];
-		tsr_half4 RawTranslucencyColorArray[CONFIG_SAMPLES][DPV_PIXEL_PER_LANE];
-
-		tsr_half2x2 RawdPPArray[CONFIG_SAMPLES];
-		tsr_half2   RawSampleSpatialWeightArray[CONFIG_SAMPLES];
-		
-		// Issues overlapped texture fetches
-		PLATFORM_SPECIFIC_ISOLATE
-		{
-			// Soften the spatial kernel when there is high noise in the input to make parrallax disocclusion of detailed geometry less distracting.
-			tsr_half2 OffScreenInputToHistoryFactor = tsr_half(1.0) - tsr_half(0.5) * NoiseFiltering;
-
-			tsr_half2 KernelInputToHistoryFactor = select(or(bIsOffScreen, bIsDisoccluded), OffScreenInputToHistoryFactor, tsr_half(InputToHistoryFactor));
-
-			UNROLL_N(CONFIG_SAMPLES)
-			for (uint SampleId = 0; SampleId < CONFIG_SAMPLES; SampleId++)
-			{
-				tsr_short2x2 SampleInputPixelPos;
-				tsr_half2x2 PixelOffset;
-				ComputeInputKernelSamplePosition(
-					InputPixelPos, dKO, SampleId,
-					/* out */ SampleInputPixelPos,
-					/* out */ PixelOffset);
-				
-				tsr_half2x2 dPP = PixelOffset - dKO;
-				tsr_half2 SampleSpatialWeight = ComputeSampleWeigth(KernelInputToHistoryFactor, dPP, /* MinimalContribution = */ float(0.005));
-				
-				RawdPPArray[SampleId] = dPP; 
-				RawSampleSpatialWeightArray[SampleId] = SampleSpatialWeight;
-
-				UNROLL_N(DPV_PIXEL_PER_LANE)
-				for (uint PixelId = 0; PixelId < DPV_PIXEL_PER_LANE; PixelId++)
-				{
-					RawInputColorArray[SampleId][PixelId] = InputSceneColorTexture[dpv_access_pixel(SampleInputPixelPos, PixelId)];
-					#if CONFIG_ACCUMULATE_TRANSLUCENCY_SEPARATLY
-						RawTranslucencyColorArray[SampleId][PixelId] = InputSceneTranslucencyTexture[dpv_access_pixel(SampleInputPixelPos, PixelId)];
-					#endif
-				}
-			} // for (uint SampleId = 0; SampleId < CONFIG_SAMPLES; SampleId++)
-		}
-			
-		tsr_half2 FilteredInputColorWeight = tsr_half(0.0);
-		
-		FilteredInputColor = tsr_half(0.0);
-		#if CONFIG_SCENE_COLOR_APLHA
-			InputMinColor = dpv_interleave_mono_registers(+LargestSceneColorYCoCgA);
-			InputMaxColor = dpv_interleave_mono_registers(-LargestSceneColorYCoCgA);
-		#else
-			InputMinColor = dpv_interleave_mono_registers(+LargestSceneColorYCoCg);
-			InputMaxColor = dpv_interleave_mono_registers(-LargestSceneColorYCoCg);
-		#endif
-		TranslucencyMinColor = dpv_interleave_mono_registers(+LargestSceneColorRGBA);
-		TranslucencyMaxColor = dpv_interleave_mono_registers(-LargestSceneColorRGBA);
-
-		UNROLL_N(CONFIG_SAMPLES)
-		for (uint SampleId = 0; SampleId < CONFIG_SAMPLES; SampleId++)
-		{
-			tsr_half2 SampleSpatialWeight = RawSampleSpatialWeightArray[SampleId];
-			tsr_halfCx2 InputColor = dpv_force_interleave_registers_array(RawInputColorArray[SampleId]);
-			tsr_half4x2 TranslucencyColor = dpv_force_interleave_registers_array(RawTranslucencyColorArray[SampleId]);
-
-			#if CONFIG_SCENE_COLOR_OVERFLOW
-				#if CONFIG_SCENE_COLOR_APLHA
-					InputColor = min(InputColor, dpv_interleave_mono_registers(LargestSceneColorRGBA));
-				#else
-					InputColor = min(InputColor, dpv_interleave_mono_registers(LargestSceneColorRGB));
-				#endif
-			#endif
-
-			tsr_half2 ToneWeight = HdrWeight4(InputColor);
-
-			FilteredInputColor       += dpv_scale((SampleSpatialWeight * ToneWeight), InputColor);
-			FilteredInputColorWeight += (SampleSpatialWeight * ToneWeight);
-
-			if (SampleId == 0)
-			{
-				FilteredTranslucencyColor = TranslucencyColor;
-			}
-
-			#if CONFIG_HIGH_FREQUENCY_YCOCG
-				InputMinColor = min(InputMinColor, TransformColorForClampingBox(InputColor));
-				InputMaxColor = max(InputMaxColor, TransformColorForClampingBox(InputColor));
-			#else
-				InputMinColor = min(InputMinColor, InputColor);
-				InputMaxColor = max(InputMaxColor, InputColor);
-			#endif
-
-			#if CONFIG_ACCUMULATE_TRANSLUCENCY_SEPARATLY
-			{
-				TranslucencyMinColor = min(TranslucencyMinColor, TranslucencyColor);
-				TranslucencyMaxColor = max(TranslucencyMaxColor, TranslucencyColor);
-				
-				if (SampleId == 0)
-				{
-					FilteredTranslucencyColor = TranslucencyColor;
-				}
-			}
-			#endif
-		}
-		
-		FilteredInputColor = dpv_scale(FilteredInputColor, rcp(FilteredInputColorWeight));
-
-		InputPixelAlignement = ComputeSampleWeigth(tsr_half(InputToHistoryFactor), dKO, /* MinimalContribution = */ float(0.0));
-		
-		#if CONFIG_SCENE_COLOR_OVERFLOW
-		{
-			#if CONFIG_SCENE_COLOR_APLHA
-				FilteredInputColor = min(FilteredInputColor, dpv_interleave_mono_registers(LargestSceneColorRGBA));
-			#else
-				FilteredInputColor = min(FilteredInputColor, dpv_interleave_mono_registers(LargestSceneColorRGB));
-			#endif
-			ClosestInputLuma4 = min(ClosestInputLuma4, LargestSceneColorYCoCg.x);
-			FilteredTranslucencyColor = min(FilteredTranslucencyColor, dpv_interleave_mono_registers(LargestSceneColorRGBA));
-		}
-		#endif
-	}
-		
-	// Spills to LDS to make room in VGPR for history sampling.
-	#if CONFIG_MANUAL_LDS_SPILL
-	PLATFORM_SPECIFIC_ISOLATE
-	{
-		uint LocalGroupThreadIndex = GetGroupThreadIndex(GroupThreadIndex, GroupWaveIndex);
-
-		#if CONFIG_ENABLE_DUAL_PIXEL_VECTORIZATION
-			SharedArray0[LocalGroupThreadIndex] = tsr_half4(FilteredInputColor[0], FilteredInputColor[1]);
-			SharedArray1[LocalGroupThreadIndex] = tsr_half4(FilteredInputColor[2], LowFrequencyRejection);
-			SharedArray2[LocalGroupThreadIndex] = tsr_half4(InputMinColor[0], InputMinColor[1]);
-			SharedArray3[LocalGroupThreadIndex] = tsr_half4(InputMinColor[2], InputPixelAlignement);
-			SharedArray4[LocalGroupThreadIndex] = tsr_half4(InputMaxColor[0], InputMaxColor[1]);
-			SharedArray5[LocalGroupThreadIndex] = tsr_half4(InputMaxColor[2], OutputPixelVelocity);
-
-			#if CONFIG_ACCUMULATE_TRANSLUCENCY_SEPARATLY
-				SharedArray6[LocalGroupThreadIndex] = tsr_half4(FilteredTranslucencyColor[0], FilteredTranslucencyColor[1]);
-				SharedArray7[LocalGroupThreadIndex] = tsr_half4(FilteredTranslucencyColor[2], FilteredTranslucencyColor[3]);
-				SharedArray8[LocalGroupThreadIndex] = tsr_half4(TranslucencyMinColor[0], TranslucencyMinColor[1]);
-				SharedArray9[LocalGroupThreadIndex] = tsr_half4(TranslucencyMinColor[2], TranslucencyMinColor[3]);
-				SharedArray10[LocalGroupThreadIndex] = tsr_half4(TranslucencyMaxColor[0], TranslucencyMaxColor[1]);
-				SharedArray11[LocalGroupThreadIndex] = tsr_half4(TranslucencyMaxColor[2], TranslucencyMaxColor[3]);
-			#endif
-
-			#if CONFIG_UPDATE_SUBPIXEL_DETAIL
-				SharedParallaxFactorBits[LocalGroupThreadIndex] = ParallaxFactorBits;
-			#endif
-		#else
-			SharedArray0[LocalGroupThreadIndex] = tsr_half4(dpv_lo(FilteredInputColor), dpv_lo(LowFrequencyRejection));
-			SharedArray1[LocalGroupThreadIndex] = tsr_half4(dpv_lo(InputMinColor), dpv_lo(InputPixelAlignement));
-			SharedArray2[LocalGroupThreadIndex] = tsr_half4(dpv_lo(InputMaxColor), dpv_lo(OutputPixelVelocity));
-		#endif
-	}
-	#endif
-	
 	// Reproject history
 	#if !CONFIG_HIGH_FREQUENCY_ONLY
 		tsr_halfCx2 PrevLowFrequencyColor;
@@ -953,6 +766,205 @@
 			PrevHistoryValidity = max(PrevHistoryValidity, tsr_half(0.0).xx);
 		}
 	}
+	
+	// TODO: CONFIG_MANUAL_LDS_SPILL
+
+	// Filter input scene color at predictor frequency.
+	tsr_halfCx2 FilteredInputColor;
+	tsr_halfCx2 InputMinColor;
+	tsr_halfCx2 InputMaxColor;
+	
+	tsr_half4x2 FilteredTranslucencyColor;
+	tsr_half4x2 TranslucencyMinColor;
+	tsr_half4x2 TranslucencyMaxColor;
+
+	tsr_half2 InputPixelAlignement;
+	tsr_half2 ClosestInputLuma4;
+	
+	
+	#if CONFIG_UPDATE_SUBPIXEL_DETAIL
+		bool2 bUpdateSubpixelDetails;
+		bool2x2 bKOSign;
+	#endif
+
+	PLATFORM_SPECIFIC_ISOLATE
+	{
+		// Vector in pixel between pixel K -> O.
+		tsr_half2x2 dKO = tsr_half2x2(PPCo - PPCk);
+
+		// Compresses the sign of dKO into SGPR to keep VGPR pressure low.
+		#if CONFIG_UPDATE_SUBPIXEL_DETAIL
+		{
+			float SizeOfOutputPixelInInputPixelWidth = 0.5 * (InputInfo_ViewportSize.x * HistoryInfo_ViewportSizeInverse.x);
+
+			bUpdateSubpixelDetails[0] = all(abs(dpv_lo(dKO)) < SizeOfOutputPixelInInputPixelWidth);
+			bUpdateSubpixelDetails[1] = all(abs(dpv_hi(dKO)) < SizeOfOutputPixelInInputPixelWidth);
+		
+			bKOSign[0] = dKO[0] < 0.0;
+			bKOSign[1] = dKO[1] < 0.0;
+		}
+		#endif
+		
+		tsr_halfC RawInputColorArray[CONFIG_SAMPLES][DPV_PIXEL_PER_LANE];
+		tsr_half4 RawTranslucencyColorArray[CONFIG_SAMPLES][DPV_PIXEL_PER_LANE];
+
+		tsr_half2x2 RawdPPArray[CONFIG_SAMPLES];
+		tsr_half2   RawSampleSpatialWeightArray[CONFIG_SAMPLES];
+		
+		// Issues overlapped texture fetches
+		PLATFORM_SPECIFIC_ISOLATE
+		{
+			// Soften the spatial kernel when there is high noise in the input to make parrallax disocclusion of detailed geometry less distracting.
+			tsr_half2 OffScreenInputToHistoryFactor = tsr_half(1.0) - tsr_half(0.5) * NoiseFiltering;
+
+			bool2 bIsRefining = PrevHistoryValidity == 0.0;
+			//bIsRefining = true;
+
+			tsr_half2 TotalRejection = select(or(or(bIsOffScreen, bIsDisoccluded), !bIsRefining), tsr_half(0.0).xx, saturate(LowFrequencyRejection * tsr_half(4.0)));
+
+			tsr_half2 KernelInputToHistoryFactor = lerp(OffScreenInputToHistoryFactor, tsr_half(InputToHistoryFactor), TotalRejection);
+			tsr_half2 KernelInputToHistoryAlignmentFactor = lerp(tsr_half(1.0).xx, tsr_half(InputToHistoryFactor), TotalRejection);
+			
+			//InputPixelAlignement = ComputeSampleWeigth(tsr_half(InputToHistoryFactor), dKO, /* MinimalContribution = */ float(0.0));
+			InputPixelAlignement = ComputeSampleWeigth(KernelInputToHistoryAlignmentFactor, dKO, /* MinimalContribution = */ float(0.0));
+		
+			UNROLL_N(CONFIG_SAMPLES)
+			for (uint SampleId = 0; SampleId < CONFIG_SAMPLES; SampleId++)
+			{
+				tsr_short2x2 SampleInputPixelPos;
+				tsr_half2x2 PixelOffset;
+				ComputeInputKernelSamplePosition(
+					InputPixelPos, dKO, SampleId,
+					/* out */ SampleInputPixelPos,
+					/* out */ PixelOffset);
+				
+				tsr_half2x2 dPP = PixelOffset - dKO;
+				tsr_half2 SampleSpatialWeight = ComputeSampleWeigth(KernelInputToHistoryFactor, dPP, /* MinimalContribution = */ float(0.005));
+				
+				RawdPPArray[SampleId] = dPP; 
+				RawSampleSpatialWeightArray[SampleId] = SampleSpatialWeight;
+
+				UNROLL_N(DPV_PIXEL_PER_LANE)
+				for (uint PixelId = 0; PixelId < DPV_PIXEL_PER_LANE; PixelId++)
+				{
+					RawInputColorArray[SampleId][PixelId] = InputSceneColorTexture[dpv_access_pixel(SampleInputPixelPos, PixelId)];
+					#if CONFIG_ACCUMULATE_TRANSLUCENCY_SEPARATLY
+						RawTranslucencyColorArray[SampleId][PixelId] = InputSceneTranslucencyTexture[dpv_access_pixel(SampleInputPixelPos, PixelId)];
+					#endif
+				}
+			} // for (uint SampleId = 0; SampleId < CONFIG_SAMPLES; SampleId++)
+		}
+			
+		tsr_half2 FilteredInputColorWeight = tsr_half(0.0);
+		
+		FilteredInputColor = tsr_half(0.0);
+		#if CONFIG_SCENE_COLOR_APLHA
+			InputMinColor = dpv_interleave_mono_registers(+LargestSceneColorYCoCgA);
+			InputMaxColor = dpv_interleave_mono_registers(-LargestSceneColorYCoCgA);
+		#else
+			InputMinColor = dpv_interleave_mono_registers(+LargestSceneColorYCoCg);
+			InputMaxColor = dpv_interleave_mono_registers(-LargestSceneColorYCoCg);
+		#endif
+		TranslucencyMinColor = dpv_interleave_mono_registers(+LargestSceneColorRGBA);
+		TranslucencyMaxColor = dpv_interleave_mono_registers(-LargestSceneColorRGBA);
+
+		UNROLL_N(CONFIG_SAMPLES)
+		for (uint SampleId = 0; SampleId < CONFIG_SAMPLES; SampleId++)
+		{
+			tsr_half2 SampleSpatialWeight = RawSampleSpatialWeightArray[SampleId];
+			tsr_halfCx2 InputColor = dpv_force_interleave_registers_array(RawInputColorArray[SampleId]);
+			tsr_half4x2 TranslucencyColor = dpv_force_interleave_registers_array(RawTranslucencyColorArray[SampleId]);
+
+			#if CONFIG_SCENE_COLOR_OVERFLOW
+				#if CONFIG_SCENE_COLOR_APLHA
+					InputColor = min(InputColor, dpv_interleave_mono_registers(LargestSceneColorRGBA));
+				#else
+					InputColor = min(InputColor, dpv_interleave_mono_registers(LargestSceneColorRGB));
+				#endif
+			#endif
+
+			tsr_half2 ToneWeight = HdrWeight4(InputColor);
+
+			FilteredInputColor       += dpv_scale((SampleSpatialWeight * ToneWeight), InputColor);
+			FilteredInputColorWeight += (SampleSpatialWeight * ToneWeight);
+
+			if (SampleId == 0)
+			{
+				FilteredTranslucencyColor = TranslucencyColor;
+			}
+
+			#if CONFIG_HIGH_FREQUENCY_YCOCG
+				InputMinColor = min(InputMinColor, TransformColorForClampingBox(InputColor));
+				InputMaxColor = max(InputMaxColor, TransformColorForClampingBox(InputColor));
+			#else
+				InputMinColor = min(InputMinColor, InputColor);
+				InputMaxColor = max(InputMaxColor, InputColor);
+			#endif
+
+			#if CONFIG_ACCUMULATE_TRANSLUCENCY_SEPARATLY
+			{
+				TranslucencyMinColor = min(TranslucencyMinColor, TranslucencyColor);
+				TranslucencyMaxColor = max(TranslucencyMaxColor, TranslucencyColor);
+				
+				if (SampleId == 0)
+				{
+					FilteredTranslucencyColor = TranslucencyColor;
+				}
+			}
+			#endif
+		}
+		
+		FilteredInputColor = dpv_scale(FilteredInputColor, rcp(FilteredInputColorWeight));
+
+		#if CONFIG_SCENE_COLOR_OVERFLOW
+		{
+			#if CONFIG_SCENE_COLOR_APLHA
+				FilteredInputColor = min(FilteredInputColor, dpv_interleave_mono_registers(LargestSceneColorRGBA));
+			#else
+				FilteredInputColor = min(FilteredInputColor, dpv_interleave_mono_registers(LargestSceneColorRGB));
+			#endif
+			ClosestInputLuma4 = min(ClosestInputLuma4, LargestSceneColorYCoCg.x);
+			FilteredTranslucencyColor = min(FilteredTranslucencyColor, dpv_interleave_mono_registers(LargestSceneColorRGBA));
+		}
+		#endif
+	}
+		
+	// TODO: CONFIG_MANUAL_LDS_SPILL
+
+
+	// Spills to LDS to make room in VGPR for history sampling.
+	#if CONFIG_MANUAL_LDS_SPILL
+	PLATFORM_SPECIFIC_ISOLATE
+	{
+		uint LocalGroupThreadIndex = GetGroupThreadIndex(GroupThreadIndex, GroupWaveIndex);
+	
+		#if CONFIG_ENABLE_DUAL_PIXEL_VECTORIZATION
+			SharedArray0[LocalGroupThreadIndex] = tsr_half4(FilteredInputColor[0], FilteredInputColor[1]);
+			SharedArray1[LocalGroupThreadIndex] = tsr_half4(FilteredInputColor[2], LowFrequencyRejection);
+			SharedArray2[LocalGroupThreadIndex] = tsr_half4(InputMinColor[0], InputMinColor[1]);
+			SharedArray3[LocalGroupThreadIndex] = tsr_half4(InputMinColor[2], InputPixelAlignement);
+			SharedArray4[LocalGroupThreadIndex] = tsr_half4(InputMaxColor[0], InputMaxColor[1]);
+			SharedArray5[LocalGroupThreadIndex] = tsr_half4(InputMaxColor[2], OutputPixelVelocity);
+	
+			#if CONFIG_ACCUMULATE_TRANSLUCENCY_SEPARATLY
+				SharedArray6[LocalGroupThreadIndex] = tsr_half4(FilteredTranslucencyColor[0], FilteredTranslucencyColor[1]);
+				SharedArray7[LocalGroupThreadIndex] = tsr_half4(FilteredTranslucencyColor[2], FilteredTranslucencyColor[3]);
+				SharedArray8[LocalGroupThreadIndex] = tsr_half4(TranslucencyMinColor[0], TranslucencyMinColor[1]);
+				SharedArray9[LocalGroupThreadIndex] = tsr_half4(TranslucencyMinColor[2], TranslucencyMinColor[3]);
+				SharedArray10[LocalGroupThreadIndex] = tsr_half4(TranslucencyMaxColor[0], TranslucencyMaxColor[1]);
+				SharedArray11[LocalGroupThreadIndex] = tsr_half4(TranslucencyMaxColor[2], TranslucencyMaxColor[3]);
+			#endif
+	
+			#if CONFIG_UPDATE_SUBPIXEL_DETAIL
+				SharedParallaxFactorBits[LocalGroupThreadIndex] = ParallaxFactorBits;
+			#endif
+		#else
+			SharedArray0[LocalGroupThreadIndex] = tsr_half4(dpv_lo(FilteredInputColor), dpv_lo(LowFrequencyRejection));
+			SharedArray1[LocalGroupThreadIndex] = tsr_half4(dpv_lo(InputMinColor), dpv_lo(InputPixelAlignement));
+			SharedArray2[LocalGroupThreadIndex] = tsr_half4(dpv_lo(InputMaxColor), dpv_lo(OutputPixelVelocity));
+		#endif
+	}
+	#endif
 		
 	#if CONFIG_MANUAL_LDS_SPILL
 	PLATFORM_SPECIFIC_ISOLATE
@@ -1131,11 +1143,7 @@
 		#endif
 
 		// Compute the newly rejected validity
-		tsr_half2 RejectedValidity = (
-			min(ClampedPrevHistoryValidity, PredictionOnlyValidity - DesiredCurrentContribution) +
-			max(ClampedPrevHistoryValidity - PredictionOnlyValidity + DesiredCurrentContribution, tsr_half(0.0)) * PrevHistoryRejectionWeight);
+		tsr_half2 RejectedValidity = ClampedPrevHistoryValidity * PrevHistoryRejectionWeight;
 
-		RejectedValidity = ClampedPrevHistoryValidity * PrevHistoryRejectionWeight;
-
 		// Compute the maximum output validity.
 		tsr_half2 OutputValidity = (
