Fixing translucent object compositing in ray-traced reflections.

diff --git a/Engine/Shaders/Private/RayTracing/RayTracingReflections.usf b/Engine/Shaders/Private/RayTracing/RayTracingReflections.usf
--- a/Engine/Shaders/Private/RayTracing/RayTracingReflections.usf
+++ b/Engine/Shaders/Private/RayTracing/RayTracingReflections.usf
@@ -310,7 +310,6 @@ RAY_TRACING_ENTRY_RAYGEN(RayTracingReflectionsRGS)
 	if (bIsValidPixel)
 	{
 		{
-			bool bIsValidSample = false;
 			bool bNeedsCapture = true;
 
 			float3 PathRadiance = 0;
@@ -664,7 +663,7 @@ RAY_TRACING_ENTRY_RAYGEN(RayTracingReflectionsRGS)
 					{
 						RayDesc TransmissionRay = TopLayerRay;
 						TransmissionRay.TMin = 0.01;
-						TransmissionRay.TMax = PackedPayload.HitT;
+						TransmissionRay.TMax = PackedPayload.IsHit() ? PackedPayload.HitT : 100000.0f;
 						Transmission = Transmit(TransmissionRay, PixelCoord, TopLayerRadiance);
 					}
 
@@ -713,7 +712,6 @@ RAY_TRACING_ENTRY_RAYGEN(RayTracingReflectionsRGS)
 						ImaginaryHitDistance += PackedPayload.HitT;
 					}
 
-					bIsValidSample = true;
 					if (BounceIndex == 0 && (ClosestHitDistance == DENOISER_INVALID_HIT_DISTANCE || ClosestHitDistance == DENOISER_MISS_HIT_DISTANCE || ClosestHitDistance > PackedPayload.HitT))
 					{
 						ClosestHitDistance = PackedPayload.HitT;
@@ -728,7 +726,6 @@ RAY_TRACING_ENTRY_RAYGEN(RayTracingReflectionsRGS)
 					float SkyAverageBrightness = 1.0f;
 					float3 SkyLighting = GetSkyLightReflection(TopLayerRay.Direction, TopLayerRoughness, SkyAverageBrightness);
 					TopLayerRadiance = SkyLighting;
-					bIsValidSample = true;
 					bSkyWasHit = true;	// To stop the main path recusion early
 
 					bApplyHeightFog = HeightFog > 0;
@@ -760,7 +757,7 @@ RAY_TRACING_ENTRY_RAYGEN(RayTracingReflectionsRGS)
 						ImaginaryReflectedSurfaceHit = false;
 					}
 
-					TopLayerRadiance = 0;
+					//TopLayerRadiance = 0;
 				}
 
 				// Apply height fog on the sample. (for clear coat: we apply the fog the same way on each layer reflection)
@@ -830,7 +827,7 @@ RAY_TRACING_ENTRY_RAYGEN(RayTracingReflectionsRGS)
 				}
 			}
 			
-			if (bIsValidSample)
+			if (any(PathRadiance) > 0.0)
 			{
 				bNeedsCapture |= BounceIndex == LocalMaxBounces && !(bSkyWasHit || ClosestHitDistance < 0.0f);
 				if (UseReflectionCaptures && bNeedsCapture)
@@ -1010,22 +1007,6 @@ RAY_TRACING_ENTRY_RAYGEN(RayTracingReflectionsRGS)
 			ReflectedColor *= DistanceFade * max(RoughnessFade, ClearCoatRoughnessFade) * rcp(LocalSamplesPerPixel);
 		}
 
-		// Handle per-sample translucency
-		if (EnableTranslucency && (SampleMode == SAMPLE_PHASE_MONOLITHIC || SampleMode == SAMPLE_PHASE_RESOLVE))
-		{
-			RayDesc TransmissionRay;
-			TransmissionRay.Origin = View.TranslatedWorldCameraOrigin;
-			TransmissionRay.Direction = CameraDirection;
-			TransmissionRay.TMin = 0.0;
-			TransmissionRay.TMax = dot(OriginalTranslatedWorldPosition - View.TranslatedWorldCameraOrigin, CameraDirection) - 0.1;
-			float3 Inscattering = 0.0;
-			float Transmission = Transmit(TransmissionRay, PixelCoord, Inscattering);
-			//#dxr_todo: add Inscattering support (this is emission only so far)
-			//#dxr_todo: depending on if there is a translucency pass the component might already be attenuated. Check for double contribution
-			ReflectedColor.rgb *= Transmission;
-			ReflectedColor.rgb += Inscattering;
-		}
-
 		// Transform the sample information into the expected output format
 		ReflectedColor.xyz *= View.PreExposure;
 
