Support shading models for mobile deferred path when static lighting is disabled

diff --git a/Engine/Shaders/Private/DeferredShadingCommon.ush b/Engine/Shaders/Private/DeferredShadingCommon.ush
--- a/Engine/Shaders/Private/DeferredShadingCommon.ush
+++ b/Engine/Shaders/Private/DeferredShadingCommon.ush
@@ -460,6 +460,31 @@ float4 ComputeIndirectLightingSampleE(uint2 TracingPixelCoord, uint TracingRayIn
 }
 
 #if SHADING_PATH_MOBILE
+#ifndef ENABLE_SHADINGMODEL_SUPPORT_MOBILE_DEFERRED
+#define ENABLE_SHADINGMODEL_SUPPORT_MOBILE_DEFERRED 0
+#endif
+#define MOBILE_SHADINGMODEL_SUPPORT (!MOBILE_DEFERRED_SHADING || ENABLE_SHADINGMODEL_SUPPORT_MOBILE_DEFERRED)
+
+float MobileEncodeShadingModelIdAndColorChannel(uint ShadingModelId, float Color)
+{
+	uint ColorValue = (uint)round(clamp(Color, 0.0, 1.0) * 31.0) << 3; // 5bit color
+	return (float)(ColorValue | ShadingModelId) / 255.0;
+}
+
+uint MobileDecodeShadingModelId(float InPackedChannel)
+{
+	uint ShadingModelId = (uint)round(InPackedChannel * 255.0);
+#if !MOBILE_EXTENDED_GBUFFER
+	ShadingModelId = ShadingModelId & 0x7;
+#endif
+	return ShadingModelId;
+}
+
+float MobileDecodeColorChannel(float InPackedChannel)
+{
+	return (((uint)(InPackedChannel * 255.0) & ~0x7) >> 3) / 31.0;
+}
+
 /** Mobile specific encoding of GBuffer data */
 void MobileEncodeGBuffer(
 	FGBufferData GBuffer,
@@ -478,6 +503,10 @@ void MobileEncodeGBuffer(
 	}
 	else
 	{
+#if !MOBILE_EXTENDED_GBUFFER
+		// We allocate only 3bits for shading model ID
+		GBuffer.ShadingModelID = (GBuffer.ShadingModelID > 7 ? SHADINGMODELID_DEFAULT_LIT : GBuffer.ShadingModelID);
+#endif
 		OutGBufferA.rg = UnitVectorToOctahedron( normalize(GBuffer.WorldNormal) ) * 0.5f + 0.5f;
 #if ALLOW_STATIC_LIGHTING
 		OutGBufferA.b = EncodeIndirectIrradiance(GBuffer.IndirectIrradiance * GBuffer.GBufferAO);
@@ -489,16 +518,44 @@ void MobileEncodeGBuffer(
 		OutGBufferB.r = GBuffer.Metallic;
 		OutGBufferB.g = GBuffer.Specular;
 		OutGBufferB.b = GBuffer.Roughness;
-		OutGBufferB.a = EncodeShadingModelIdAndSelectiveOutputMask(GBuffer.ShadingModelID, 0);
+		OutGBufferB.a = GBuffer.ShadingModelID / 255.0;
 
 		OutGBufferC.rgb = EncodeBaseColor( GBuffer.BaseColor );
-
 #if ALLOW_STATIC_LIGHTING
 		OutGBufferC.a = GBuffer.PrecomputedShadowFactors.x;
 #else
 		OutGBufferC.a = GBuffer.GBufferAO;
 #endif
 		OutGBufferD = GBuffer.CustomData;
+
+#if MOBILE_SHADINGMODEL_SUPPORT && !MOBILE_EXTENDED_GBUFFER
+		if (GBuffer.ShadingModelID == SHADINGMODELID_SUBSURFACE || 
+			GBuffer.ShadingModelID == SHADINGMODELID_PREINTEGRATED_SKIN)
+		{
+			OutGBufferA.b = GBuffer.CustomData.a; // Opacity
+			OutGBufferB.r = GBuffer.CustomData.r; // SubsurfaceColor.R
+			OutGBufferC.a = GBuffer.CustomData.g; // SubsurfaceColor.G
+			// encode SubsurfaceColor.B into 5 bits and 3 bits for shading model 
+			OutGBufferB.a = MobileEncodeShadingModelIdAndColorChannel(GBuffer.ShadingModelID, GBuffer.CustomData.b);
+		}
+		else if (GBuffer.ShadingModelID == SHADINGMODELID_CLEAR_COAT)
+		{
+			OutGBufferA.b = GBuffer.CustomData.x; // ClearCoat
+			OutGBufferC.a = GBuffer.CustomData.y; // ClearCoatRoughness
+		}
+		else if (GBuffer.ShadingModelID == SHADINGMODELID_SUBSURFACE_PROFILE)
+		{
+			OutGBufferA.b = GBuffer.CustomData.a; // Opacity
+			OutGBufferB.r = GBuffer.CustomData.r; // SubsurfaceProfile
+			OutGBufferC.a = GBuffer.CustomData.g; // Curvature
+		}
+		else if (GBuffer.ShadingModelID == SHADINGMODELID_HAIR)
+		{
+			OutGBufferA.b = GBuffer.CustomData.z; // Backlit
+			OutGBufferB.r = GBuffer.CustomData.x; // 
+			OutGBufferC.a = GBuffer.CustomData.y; // 
+		}
+#endif
 	}
 }
 
@@ -519,7 +576,7 @@ FGBufferData MobileDecodeGBuffer(half4 InGBufferA, half4 InGBufferB, half4 InGBu
 	// Note: must match GetShadingModelId standalone function logic
 	// Also Note: SimpleElementPixelShader directly sets SV_Target2 ( GBufferB ) to indicate unlit.
 	// An update there will be required if this layout changes.
-	GBuffer.ShadingModelID = DecodeShadingModelId(InGBufferB.a);
+	GBuffer.ShadingModelID = MobileDecodeShadingModelId(InGBufferB.a);
 	GBuffer.SelectiveOutputMask = 0;
 	GBuffer.BaseColor = DecodeBaseColor(InGBufferC.rgb);
 #if ALLOW_STATIC_LIGHTING
@@ -530,8 +587,43 @@ FGBufferData MobileDecodeGBuffer(half4 InGBufferA, half4 InGBufferB, half4 InGBu
 	GBuffer.PrecomputedShadowFactors = 1.0;
 #endif
 
-#if	MOBILE_EXTENDED_GBUFFER
-	GBuffer.CustomData = HasCustomGBufferData(GBuffer.ShadingModelID) ? InGBufferD : 0;
+#if MOBILE_SHADINGMODEL_SUPPORT
+	#if	MOBILE_EXTENDED_GBUFFER
+		GBuffer.CustomData = HasCustomGBufferData(GBuffer.ShadingModelID) ? InGBufferD : 0;
+	#else
+		if (GBuffer.ShadingModelID == SHADINGMODELID_SUBSURFACE ||
+			GBuffer.ShadingModelID == SHADINGMODELID_PREINTEGRATED_SKIN)
+		{
+			GBuffer.CustomData.a = InGBufferA.b; // Opacity
+			GBuffer.CustomData.r = InGBufferB.r; // SubsurfaceColor.R
+			GBuffer.CustomData.g = InGBufferC.a; // SubsurfaceColor.G
+			GBuffer.CustomData.b = MobileDecodeColorChannel(InGBufferB.a); // SubsurfaceColor.B
+			GBuffer.GBufferAO = 1;	// no space for AO
+			GBuffer.Metallic = 0;	// no space for Metallic
+		}
+		else if (GBuffer.ShadingModelID == SHADINGMODELID_CLEAR_COAT)
+		{
+			GBuffer.CustomData.x = InGBufferA.b; // ClearCoat
+			GBuffer.CustomData.y = InGBufferC.a; // ClearCoatRoughness
+			GBuffer.GBufferAO = 1;	// no space for AO
+		}
+		else if (GBuffer.ShadingModelID == SHADINGMODELID_SUBSURFACE_PROFILE)
+		{
+			GBuffer.CustomData.a = InGBufferA.b; // Opacity
+			GBuffer.CustomData.r = InGBufferB.r; // SubsurfaceProfile
+			GBuffer.CustomData.g = InGBufferC.a; // Curvature
+			GBuffer.GBufferAO = 1;	// no space for AO
+			GBuffer.Metallic = 0;	// no space for Metallic
+		}
+		else if (GBuffer.ShadingModelID == SHADINGMODELID_HAIR)
+		{
+			GBuffer.CustomData.z = InGBufferA.b; // Backlit
+			GBuffer.CustomData.x = InGBufferB.r; // 
+			GBuffer.CustomData.y = InGBufferC.a; //
+			GBuffer.GBufferAO = 1;	// no space for AO
+			GBuffer.Metallic = 0;	// no space for Metallic
+		}
+	#endif
 #endif
 	return GBuffer;
 }
diff --git a/Engine/Shaders/Private/MobileShadingModels.ush b/Engine/Shaders/Private/MobileShadingModels.ush
--- a/Engine/Shaders/Private/MobileShadingModels.ush
+++ b/Engine/Shaders/Private/MobileShadingModels.ush
@@ -51,11 +51,6 @@
 #define DEFERRED_SHADING_PATH 0
 #endif
 
-#ifndef ENABLE_SHADINGMODEL_SUPPORT_MOBILE_DEFERRED
-#define ENABLE_SHADINGMODEL_SUPPORT_MOBILE_DEFERRED 0
-#endif
-
-#define MOBILE_SHADINGMODEL_SUPPORT (!MOBILE_DEFERRED_SHADING || ENABLE_SHADINGMODEL_SUPPORT_MOBILE_DEFERRED)
 
 /*------------------------------------------------------------------------------
 	Mobile Shading Models
@@ -205,12 +200,14 @@ void InitShadingModelContext(inout FMobileShadingModelContext ShadingModelContex
 	else if (GBuffer.ShadingModelID == SHADINGMODELID_SUBSURFACE || GBuffer.ShadingModelID == SHADINGMODELID_PREINTEGRATED_SKIN)
 	{
 		ShadingModelContext.SubsurfaceColor = ExtractSubsurfaceColor(GBuffer);
+		ShadingModelContext.Opacity = GBuffer.CustomData.a;
 	}
 	else if (GBuffer.ShadingModelID == SHADINGMODELID_SUBSURFACE_PROFILE)
 	{
 		ShadingModelContext.SubsurfaceProfileInt = ExtractSubsurfaceProfileInt(GBuffer);
 				
 		ShadingModelContext.Curvature = GBuffer.CustomData.g;
+		ShadingModelContext.Opacity = GBuffer.CustomData.a;
 
 		half2 LobeRoughness;
 		GetProfileDualSpecular(ShadingModelContext.SubsurfaceProfileInt, GBuffer.Roughness, 
diff --git a/Engine/Source/Runtime/Engine/Private/ShaderCompiler/ShaderCompiler.cpp b/Engine/Source/Runtime/Engine/Private/ShaderCompiler/ShaderCompiler.cpp
--- a/Engine/Source/Runtime/Engine/Private/ShaderCompiler/ShaderCompiler.cpp
+++ b/Engine/Source/Runtime/Engine/Private/ShaderCompiler/ShaderCompiler.cpp
@@ -5480,7 +5480,7 @@ void GlobalBeginCompileShader(
 
 			bool bGLESDeferredShading = Target.Platform == SP_OPENGL_ES3_1_ANDROID;
 			Input.Environment.SetDefine(TEXT("USE_GLES_FBF_DEFERRED"), bGLESDeferredShading ? 1 : 0);
-			Input.Environment.SetDefine(TEXT("MOBILE_EXTENDED_GBUFFER"), MobileEnableExtenedGBuffer((EShaderPlatform)Target.Platform) ? 1 : 0);
+			Input.Environment.SetDefine(TEXT("MOBILE_EXTENDED_GBUFFER"), MobileUsesExtenedGBuffer((EShaderPlatform)Target.Platform) ? 1 : 0);
 		}
 	}
 
diff --git a/Engine/Source/Runtime/RenderCore/Private/RenderUtils.cpp b/Engine/Source/Runtime/RenderCore/Private/RenderUtils.cpp
--- a/Engine/Source/Runtime/RenderCore/Private/RenderUtils.cpp
+++ b/Engine/Source/Runtime/RenderCore/Private/RenderUtils.cpp
@@ -1233,8 +1233,8 @@ RENDERCORE_API bool MobileUsesShadowMaskTexture(const FStaticShaderPlatform Plat
 	return IsMobileDistanceFieldEnabled(Platform) || (!IsMobileDeferredShadingEnabled(Platform) && IsMobileMovableSpotlightShadowsEnabled(Platform) && MobileForwardEnableLocalLights(Platform));
 }
 
-// Whether to support more than 4 color attachments for GBuffer. Required for shading models with a custom data
-RENDERCORE_API bool MobileEnableExtenedGBuffer(FStaticShaderPlatform ShaderPlatform)
+// Whether to support more than 4 color attachments for GBuffer 
+RENDERCORE_API bool MobileUsesExtenedGBuffer(FStaticShaderPlatform ShaderPlatform)
 {
 	// Android GLES: uses PLS for deferred shading and limited to 128 bits
 	// Vulkan requires:
@@ -1244,6 +1244,15 @@ RENDERCORE_API bool MobileEnableExtenedGBuffer(FStaticShaderPlatform ShaderPlatf
 	return (ShaderPlatform != SP_OPENGL_ES3_1_ANDROID) && false;
 }
 
+// Required for shading models with a custom data
+RENDERCORE_API bool MobileUsesGBufferCustomData(FStaticShaderPlatform ShaderPlatform)
+{
+	static const auto CVarAllowStaticLighting = IConsoleManager::Get().FindTConsoleVariableDataInt(TEXT("r.AllowStaticLighting"));
+	static const bool bAllowStaticLighting = CVarAllowStaticLighting->GetValueOnAnyThread() != 0;
+	// we can pack CustomData into static lighting related space
+	return MobileUsesExtenedGBuffer(ShaderPlatform) || !bAllowStaticLighting;
+}
+
 RENDERCORE_API bool MobileBasePassAlwaysUsesCSM(const FStaticShaderPlatform Platform)
 {
 	static TConsoleVariableData<int32>* CVar = IConsoleManager::Get().FindTConsoleVariableDataInt(TEXT("r.Mobile.Shadow.CSMShaderCullingMethod"));
diff --git a/Engine/Source/Runtime/RenderCore/Private/Shader.cpp b/Engine/Source/Runtime/RenderCore/Private/Shader.cpp
--- a/Engine/Source/Runtime/RenderCore/Private/Shader.cpp
+++ b/Engine/Source/Runtime/RenderCore/Private/Shader.cpp
@@ -1567,7 +1567,7 @@ void ShaderMapAppendKeyString(EShaderPlatform Platform, FString& KeyString)
 			{
 				if (MobileShadingPathCVar->GetInt() != 0)
 				{
-					KeyString += (MobileEnableExtenedGBuffer(Platform) ? TEXT("_MobDShEx") : TEXT("_MobDSh"));
+					KeyString += (MobileUsesExtenedGBuffer(Platform) ? TEXT("_MobDShEx") : TEXT("_MobDSh"));
 				}
 				else
 				{
diff --git a/Engine/Source/Runtime/RenderCore/Public/RenderUtils.h b/Engine/Source/Runtime/RenderCore/Public/RenderUtils.h
--- a/Engine/Source/Runtime/RenderCore/Public/RenderUtils.h
+++ b/Engine/Source/Runtime/RenderCore/Public/RenderUtils.h
@@ -511,7 +511,9 @@ RENDERCORE_API bool MobileEnableClusteredReflections(const FStaticShaderPlatform
 
 RENDERCORE_API bool MobileUsesShadowMaskTexture(const FStaticShaderPlatform Platform);
 
-RENDERCORE_API bool MobileEnableExtenedGBuffer(const FStaticShaderPlatform Platform);
+RENDERCORE_API bool MobileUsesExtenedGBuffer(const FStaticShaderPlatform Platform);
+
+RENDERCORE_API bool MobileUsesGBufferCustomData(const FStaticShaderPlatform Platform);
 
 RENDERCORE_API bool MobileBasePassAlwaysUsesCSM(const FStaticShaderPlatform Platform);
 
diff --git a/Engine/Source/Runtime/Renderer/Private/MobileBasePass.cpp b/Engine/Source/Runtime/Renderer/Private/MobileBasePass.cpp
--- a/Engine/Source/Runtime/Renderer/Private/MobileBasePass.cpp
+++ b/Engine/Source/Runtime/Renderer/Private/MobileBasePass.cpp
@@ -25,8 +25,8 @@ uint8 GetMobileShadingModelStencilValue(FMaterialShadingModelField ShadingModel)
 		return 0u;
 	}
 	
-	// mark everyhing as MSM_DefaultLit if extended GBuffer is not supported
-	return MobileEnableExtenedGBuffer(GMaxRHIShaderPlatform) ? 2u : 1u;
+	// mark everyhing as MSM_DefaultLit if GBuffer CustomData is not supported
+	return MobileUsesGBufferCustomData(GMaxRHIShaderPlatform) ? 2u : 1u;
 }
 
 bool MobileUsesNoLightMapPermutation(const FMeshMaterialShaderPermutationParameters& Parameters)
diff --git a/Engine/Source/Runtime/Renderer/Private/MobileBasePassRendering.h b/Engine/Source/Runtime/Renderer/Private/MobileBasePassRendering.h
--- a/Engine/Source/Runtime/Renderer/Private/MobileBasePassRendering.h
+++ b/Engine/Source/Runtime/Renderer/Private/MobileBasePassRendering.h
@@ -133,6 +133,11 @@ public:
 	static void ModifyCompilationEnvironment(const FMaterialShaderPermutationParameters& Parameters, FShaderCompilerEnvironment& OutEnvironment)
 	{
 		FMeshMaterialShader::ModifyCompilationEnvironment(Parameters, OutEnvironment);
+		
+		if (IsMobileDeferredShadingEnabled(Parameters.Platform))
+		{
+			OutEnvironment.SetDefine(TEXT("ENABLE_SHADINGMODEL_SUPPORT_MOBILE_DEFERRED"), MobileUsesGBufferCustomData(Parameters.Platform));
+		}
 	}
 
 	void GetShaderBindings(
@@ -234,6 +239,11 @@ public:
 		// This define simply lets the compilation environment know that we are using a Base Pass PixelShader.
 		OutEnvironment.SetDefine(TEXT("IS_BASE_PASS"), 1);
 		OutEnvironment.SetDefine(TEXT("IS_MOBILE_BASE_PASS"), 1);
+		
+		if (IsMobileDeferredShadingEnabled(Parameters.Platform))
+		{
+			OutEnvironment.SetDefine(TEXT("ENABLE_SHADINGMODEL_SUPPORT_MOBILE_DEFERRED"), MobileUsesGBufferCustomData(Parameters.Platform));
+		}
 
 		// Modify compilation environment depending upon material shader quality level settings.
 		ModifyCompilationEnvironmentForQualityLevel(Parameters.Platform, Parameters.MaterialParameters.QualityLevel, OutEnvironment);
diff --git a/Engine/Source/Runtime/Renderer/Private/MobileDeferredShadingPass.cpp b/Engine/Source/Runtime/Renderer/Private/MobileDeferredShadingPass.cpp
--- a/Engine/Source/Runtime/Renderer/Private/MobileDeferredShadingPass.cpp
+++ b/Engine/Source/Runtime/Renderer/Private/MobileDeferredShadingPass.cpp
@@ -82,7 +82,7 @@ class FMobileDirectionalLightFunctionPS : public FMaterialShader
 			PermutationVector.Set<FShadowQuality>(0);
 		}
 
-		if (!MobileEnableExtenedGBuffer(Platform))
+		if (!MobileUsesGBufferCustomData(Platform))
 		{
 			PermutationVector.Set<FEnableShadingModelSupport>(false);
 		}
@@ -189,7 +189,7 @@ public:
 			PermutationVector.Set<FSpotLightShadowDim>(false);
 		}
 
-		if (!MobileEnableExtenedGBuffer(Platform))
+		if (!MobileUsesGBufferCustomData(Platform))
 		{
 			PermutationVector.Set<FEnableShadingModelSupport>(false);
 		}
@@ -248,7 +248,7 @@ public:
 		}
 
 		FPermutationDomain PermutationVector(Parameters.PermutationId);
-		if (!MobileEnableExtenedGBuffer(Parameters.Platform) && PermutationVector.Get<FEnableShadingModelSupport>())
+		if (!MobileUsesGBufferCustomData(Parameters.Platform) && PermutationVector.Get<FEnableShadingModelSupport>())
 		{
 			return false;
 		}
@@ -342,7 +342,7 @@ void RenderReflectionEnvironmentSkyLighting(FRHICommandListImmediate& RHICmdList
 	TShaderMapRef<FPostProcessVS> VertexShader(View.ShaderMap);
 
 	// Do two passes, first masking DefautLit, second masking all other shading models
-	int32 NumPasses = MobileEnableExtenedGBuffer(Scene.GetShaderPlatform()) ? 2 : 1;
+	int32 NumPasses = MobileUsesGBufferCustomData(Scene.GetShaderPlatform()) ? 2 : 1;
 	uint8 PassShadingModelStencilValue[2] =
 	{
 		GetMobileShadingModelStencilValue(MSM_DefaultLit),
@@ -459,7 +459,7 @@ static void RenderDirectionalLight(FRHICommandListImmediate& RHICmdList, const F
 	const bool bDynamicShadows = DirectionalLight.Proxy->CastsDynamicShadow() && (LightingChannel == 0u) && View.Family->EngineShowFlags.DynamicShadows;
 
 	// Do two passes, first masking DefautLit, second masking all other shading models
-	int32 NumPasses = MobileEnableExtenedGBuffer(Scene.GetShaderPlatform()) ? 2 : 1;
+	int32 NumPasses = MobileUsesGBufferCustomData(Scene.GetShaderPlatform()) ? 2 : 1;
 	uint8 PassShadingModelStencilValue[2] =
 	{
 		GetMobileShadingModelStencilValue(MSM_DefaultLit),
@@ -720,7 +720,7 @@ static void RenderLocalLight(
 	PassParameters.TranslatedWorldToLight = FMatrix44f(FTranslationMatrix(-View.ViewMatrices.GetPreViewTranslation()) * WorldToLight);
 
 	// Do two passes, first masking DefautLit, second masking all other shading models
-	int32 NumPasses = MobileEnableExtenedGBuffer(Scene.GetShaderPlatform()) ? 2 : 1;
+	int32 NumPasses = MobileUsesGBufferCustomData(Scene.GetShaderPlatform()) ? 2 : 1;
 	uint8 PassShadingModelStencilValue[2] =
 	{
 		GetMobileShadingModelStencilValue(MSM_DefaultLit),
@@ -807,7 +807,7 @@ static void RenderSimpleLights(
 
 	// Setup PSOs we going to use for light rendering 
 	// Do two passes, first masking DefautLit, second masking all other shading models
-	int32 NumPasses = MobileEnableExtenedGBuffer(Scene.GetShaderPlatform()) ? 2 : 1;
+	int32 NumPasses = MobileUsesGBufferCustomData(Scene.GetShaderPlatform()) ? 2 : 1;
 	uint8 PassShadingModelStencilValue[2] =
 	{
 		GetMobileShadingModelStencilValue(MSM_DefaultLit),
diff --git a/Engine/Source/Runtime/Renderer/Private/MobileShadingRenderer.cpp b/Engine/Source/Runtime/Renderer/Private/MobileShadingRenderer.cpp
--- a/Engine/Source/Runtime/Renderer/Private/MobileShadingRenderer.cpp
+++ b/Engine/Source/Runtime/Renderer/Private/MobileShadingRenderer.cpp
@@ -1329,7 +1329,7 @@ void FMobileSceneRenderer::RenderDeferred(FRDGBuilder& GraphBuilder, const FSort
 		ColorTargets.Add(SceneTextures.GBufferB);
 		ColorTargets.Add(SceneTextures.GBufferC);
 		
-		if (MobileEnableExtenedGBuffer(ShaderPlatform))
+		if (MobileUsesExtenedGBuffer(ShaderPlatform))
 		{
 			ColorTargets.Add(SceneTextures.GBufferD);
 		}
