Fix a cook failure on android.

diff --git a/Engine/Shaders/Private/PostProcessMaterialShaders.usf b/Engine/Shaders/Private/PostProcessMaterialShaders.usf
--- a/Engine/Shaders/Private/PostProcessMaterialShaders.usf
+++ b/Engine/Shaders/Private/PostProcessMaterialShaders.usf
@@ -11,6 +11,10 @@
 #define POST_PROCESS_MATERIAL 0
 #endif
 
+#ifndef POST_PROCESS_MATERIAL_MOBILE
+#define POST_PROCESS_MATERIAL_MOBILE 0
+#endif
+
 #if (POST_PROCESS_MATERIAL == 0)
 #error POST_PROCESS_MATERIAL must be defined to non-zero in the shader compilation environment.
 #endif
diff --git a/Engine/Shaders/Private/SceneCapturePixelShader.usf b/Engine/Shaders/Private/SceneCapturePixelShader.usf
--- a/Engine/Shaders/Private/SceneCapturePixelShader.usf
+++ b/Engine/Shaders/Private/SceneCapturePixelShader.usf
@@ -5,7 +5,7 @@
 =============================================================================*/
 
 // Sample SceneDepthAux texture on mobile
-#define POST_PROCESS_MATERIAL 1
+#define POST_PROCESS_MATERIAL_MOBILE 1
 
 #include "Common.ush"
 #include "DeferredShadingCommon.ush"
diff --git a/Engine/Shaders/Private/SceneTexturesCommon.ush b/Engine/Shaders/Private/SceneTexturesCommon.ush
--- a/Engine/Shaders/Private/SceneTexturesCommon.ush
+++ b/Engine/Shaders/Private/SceneTexturesCommon.ush
@@ -184,7 +184,7 @@ float LookupDeviceZ( float2 ScreenUV )
 #elif FORCE_DEPTH_TEXTURE_READS || COMPILER_SWITCH
 	// native Depth buffer lookup
 	return Texture2DSampleLevel(MobileSceneTextures.SceneDepthTexture, MobileSceneTextures.SceneDepthTextureSampler, ScreenUV, 0).r;
-#elif (POST_PROCESS_MATERIAL || COMPUTESHADER) && !POST_PROCESS_AR_PASSTHROUGH
+#elif (POST_PROCESS_MATERIAL || POST_PROCESS_MATERIAL_MOBILE || COMPUTESHADER) && !POST_PROCESS_AR_PASSTHROUGH
 	#if MOBILE_DEFERRED_SHADING
 		return Texture2DSampleLevel(MobileSceneTextures.SceneDepthTexture, MobileSceneTextures.SceneDepthTextureSampler, ScreenUV, 0).r;
 	#else
diff --git a/Engine/Source/Runtime/Renderer/Private/DistortionRendering.cpp b/Engine/Source/Runtime/Renderer/Private/DistortionRendering.cpp
--- a/Engine/Source/Runtime/Renderer/Private/DistortionRendering.cpp
+++ b/Engine/Source/Runtime/Renderer/Private/DistortionRendering.cpp
@@ -207,7 +207,7 @@ public:
 		if (IsMobilePlatform(Parameters.Platform))
 		{
 			// use same path for scene textures as post-process material
-			OutEnvironment.SetDefine(TEXT("POST_PROCESS_MATERIAL"), 1);
+			OutEnvironment.SetDefine(TEXT("POST_PROCESS_MATERIAL_MOBILE"), 1);
 		}
 
 		// Skip the material clip if depth test should not be done
