Strata - added some checkes for when changing number of render targets.

diff --git a/Engine/Shaders/Private/BasePassPixelShader.usf b/Engine/Shaders/Private/BasePassPixelShader.usf
--- a/Engine/Shaders/Private/BasePassPixelShader.usf
+++ b/Engine/Shaders/Private/BasePassPixelShader.usf
@@ -1753,6 +1753,9 @@ void FPixelShaderInOut_MainPS(
 			SSSData, TopLayerData, OpaqueRoughRefractionData);
 
 		// Write out MRT data
+	#if STRATA_BASE_PASS_MRT_OUTPUT_COUNT != 2
+	#error Strata STRATA_BASE_PASS_MRT_OUTPUT_COUNT has been update but not StrataOutput
+	#endif
 		Out.StrataOutput[0] = RWStrataMaterialContainer.MaterialRenderTargets[0];
 		Out.StrataOutput[1] = RWStrataMaterialContainer.MaterialRenderTargets[1];
 		Out.StrataOutput[2] = StrataPackTopLayerData(TopLayerData);
diff --git a/Engine/Shaders/Private/PixelShaderOutputCommon.ush b/Engine/Shaders/Private/PixelShaderOutputCommon.ush
--- a/Engine/Shaders/Private/PixelShaderOutputCommon.ush
+++ b/Engine/Shaders/Private/PixelShaderOutputCommon.ush
@@ -90,6 +90,9 @@ void MainPS
 #endif
 
 #if STRATA_OPAQUE_DEFERRED
+#if STRATA_BASE_PASS_MRT_OUTPUT_COUNT != 2
+#error Strata STRATA_BASE_PASS_MRT_OUTPUT_COUNT has been update but not the uint MRTs
+#endif
 #if PIXELSHADEROUTPUT_MRT3
 #error Strata cannot map to such a case
 #endif
