Change legacy SubSurface shading model conversion to be plain diffuse rather than proper sub-surface (i.e. with screen-space diffusion).

diff --git a/Engine/Shaders/Private/Strata/StrataLegacyConversion.ush b/Engine/Shaders/Private/Strata/StrataLegacyConversion.ush
--- a/Engine/Shaders/Private/Strata/StrataLegacyConversion.ush
+++ b/Engine/Shaders/Private/Strata/StrataLegacyConversion.ush
@@ -227,13 +227,17 @@ FStrataData StrataConvertLegacyMaterialDynamic(
 		// and a constant scaling for matching roughly the behavior
 		SSSMFP = SSSMFP * 0.5f;
 
+		// Convert legacy sub-surface scattering model to proper sub-surface material (using screen-space diffusion in raster mode)
+		// STATA_TODO: add something similar to legacy 'wrap' lighting behavior
+		const bool bConvertToSSS = false;
+
 		Top.BaseColor = ModifiedBaseColor;
 		Top.Specular = Specular;
 		Top.Metallic = Metallic;
 		Top.Roughness = Roughness;
 		Top.Anisotropy = Anisotropy;
-		Top.SSSMFP = SSSMFP;
-		Top.SSSMFPScale = Opacity;
+		Top.SSSMFP = bConvertToSSS ? SSSMFP : 0;
+		Top.SSSMFPScale = bConvertToSSS ? Opacity : 1.f;
 		Top.Emissive = Emissive;
 
 		// No need to take into account opacity weighting for translucent material, as the coverage is assumed to be 1
