Strata - fix for strata simple to not write more uint than necessary.

diff --git a/Engine/Shaders/Private/Strata/Strata.ush b/Engine/Shaders/Private/Strata/Strata.ush
--- a/Engine/Shaders/Private/Strata/Strata.ush
+++ b/Engine/Shaders/Private/Strata/Strata.ush
@@ -3238,6 +3238,9 @@ void PackStrataOut(
 						}
 						// Ensure the rest of the BSDF is not stored with the regular path
 						BSDF.State = 0;
+						// And break now to not write more data
+						BSDFCount = 1;
+						break;
 					}
 					// B. Layout1 - Simple encoding: single BSDF, whose header & BSDF state are merged
 					else if (bHasSingleEncoding && BSDFIdx == SimpleMaterial_BSDFIndex)
