Strata - More detailed material bytes information.

diff --git a/Engine/Shaders/Private/Strata/StrataVisualize.usf b/Engine/Shaders/Private/Strata/StrataVisualize.usf
--- a/Engine/Shaders/Private/Strata/StrataVisualize.usf
+++ b/Engine/Shaders/Private/Strata/StrataVisualize.usf
@@ -380,16 +380,35 @@ void PrintPixel(uint2 InCoord, float3 InWorldPosition, float3 V)
 	// Memory footprint
 	if (BSDFIndex+1 == Header.BSDFCount)
 	{
+		const bool bStrataMaterial = Header.BSDFCount > 0;
+		const bool bIsSimpleMaterial = IsSimpleMaterial(Header) || Header.BSDFCount == 0;
+		const bool bIsSingleMaterial = !IsSimpleMaterial(Header) && IsSingleMaterial(Header);
+
 		const uint FootPrint_PostBSDFs = StrataAddressing.ReadBytes;
 
 		const uint HeaderSize = FootPrint_PostHeader - FootPrint_Start;
 		const uint BSDFsSize  = FootPrint_PostBSDFs - FootPrint_PostHeader;
 		const uint TotalSize  = FootPrint_PostBSDFs - FootPrint_Start;
 
+		const uint TopLayerDataBytes = 4;	// bytes 
 		Print(Context, TEXT("Memory"), FontEmerald); Newline(Context);
-		Print(Context, TEXT(" Header "), FontSilver); Print(Context, HeaderSize, FontSilver, 3, 0); Print(Context, TEXT("bytes"), FontSilver); Newline(Context);
-		Print(Context, TEXT(" BSDFs  "), FontSilver); Print(Context, BSDFsSize,  FontSilver, 3, 0); Print(Context, TEXT("bytes"), FontSilver); Newline(Context);
-		Print(Context, TEXT(" Total  "), FontOrange); Print(Context, TotalSize,  FontOrange, 3, 0); Print(Context, TEXT("bytes"), FontOrange); Newline(Context);
+		if (bIsSingleMaterial)
+		{
+			Print(Context, TEXT(" Header       "), FontSilver); Print(Context, HeaderSize - TopLayerDataBytes,				FontSilver, 3, 0); Print(Context, TEXT("bytes"), FontSilver); Newline(Context);
+			Print(Context, TEXT(" TopNormalTex "), FontSilver); Print(Context, TopLayerDataBytes,							FontSilver, 3, 0); Print(Context, TEXT("bytes"), FontSilver); Newline(Context);
+			Print(Context, TEXT(" BSDF         "), FontSilver); Print(Context, BSDFsSize,									FontSilver, 3, 0); Print(Context, TEXT("bytes"), FontSilver); Newline(Context);
+		}
+		else if (bIsSimpleMaterial)
+		{
+			Print(Context, TEXT(" Header+BSDF  "), FontSilver); Print(Context, HeaderSize + BSDFsSize - TopLayerDataBytes,	FontSilver, 3, 0); Print(Context, TEXT("bytes"), FontSilver); Newline(Context);
+			Print(Context, TEXT(" TopNormalTex "), FontSilver); Print(Context, TopLayerDataBytes,							FontSilver, 3, 0); Print(Context, TEXT("bytes"), FontSilver); Newline(Context);
+		}
+		else
+		{
+			Print(Context, TEXT(" Header+Norm  "), FontSilver); Print(Context, HeaderSize,									FontSilver, 3, 0); Print(Context, TEXT("bytes"), FontSilver); Newline(Context);
+			Print(Context, TEXT(" BSDFs        "), FontSilver); Print(Context, BSDFsSize,									FontSilver, 3, 0); Print(Context, TEXT("bytes"), FontSilver); Newline(Context);
+		}
+		Print(Context, TEXT(" Total        "), FontOrange); Print(Context, TotalSize,  FontOrange, 3, 0); Print(Context, TEXT("bytes"), FontOrange); Newline(Context);
 	}
 }
 
