Strata - Fixed shader compilation on console.

diff --git a/Engine/Shaders/Private/Strata/Strata.ush b/Engine/Shaders/Private/Strata/Strata.ush
--- a/Engine/Shaders/Private/Strata/Strata.ush
+++ b/Engine/Shaders/Private/Strata/Strata.ush
@@ -3208,7 +3208,8 @@ void PackStrataOut(
 					BSDF_SETHASTRANSABOVE(BSDF, bTransmittanceAboveAlongNRequired);
 
 					// A. Layout0 - Simple encoding (aka. fast-path): store header & data
-					if (bHasFastEncoding && BSDFIdx == SimpleMaterial_BSDFIndex)
+					const bool bFastEncodedBSDF = bHasFastEncoding && BSDFIdx == SimpleMaterial_BSDFIndex;
+					if (bFastEncodedBSDF)
 					{
 						// Now we pack diffuse and F0 in a special way: both encoded as R7G7B6 with gamma 2.0. The low bit cound will be hidden by dithering and TAA.
 						// We pack 32 bits in the first uint and the remaining bits are pack is the lowest significant bit of the second uint.
@@ -3238,9 +3239,6 @@ void PackStrataOut(
 						}
 						// Ensure the rest of the BSDF is not stored with the regular path
 						BSDF.State = 0;
-						// And break now to not write more data
-						BSDFCount = 1;
-						break;
 					}
 					// B. Layout1 - Simple encoding: single BSDF, whose header & BSDF state are merged
 					else if (bHasSingleEncoding && BSDFIdx == SimpleMaterial_BSDFIndex)
@@ -3273,89 +3271,92 @@ void PackStrataOut(
 					}
 
 					// Layout1 & Layout2
-					const uint BSDFType = BSDF_GETTYPE(BSDF);
-					switch (BSDF_GETTYPE(BSDF))
+					if (!bFastEncodedBSDF)
 					{
-					case STRATA_BSDF_TYPE_SLAB:
-					{
-						// Now we pack diffuse and F0 in a special way: both encoded as R7G7B6 with gamma 2.0. The low bit cound will be hidden by dithering and TAA.
-						// We pack 32 bits in the first uint and the remaining bits are pack is the lowest significant bit of the second uint.
-						const uint PackedDiffuse20Bits = PackR7G7B6Gamma2(SLAB_DIFFUSEALBEDO(BSDF));
-						const uint PackedF020Bits = PackR7G7B6Gamma2(SLAB_F0(BSDF));
-						const uint PackedData32Bits = ((PackedDiffuse20Bits << 12) & 0xFFFFF000) | (PackedF020Bits & 0xFFF);
-						const uint PackedData8Bits = (PackedF020Bits >> 12) & 0xFF;
-
-						STRATA_STORE_UINT1(PackedData32Bits);
-						STRATA_STORE_UINT1(PackedData8Bits | PackRGBA8(float4(0.0f, SLAB_ROUGHNESS(BSDF), (SLAB_ANISOTROPY(BSDF) + 1.f) * 0.5f, 0.f)));
-
-						if (BSDF_GETHASF90(BSDF) || BSDF_GETHASHAZINESS(BSDF))
+						const uint BSDFType = BSDF_GETTYPE(BSDF);
+						switch (BSDF_GETTYPE(BSDF))
+						{
+						case STRATA_BSDF_TYPE_SLAB:
 						{
-							// What is important is to maintain the hue and saturation, so we scale the color by the maximum of its components
-							float3 F90 = saturate(SLAB_F90(BSDF));
-							const float Divisor = max(F90.r, max(F90.g, F90.b));
-							F90 = Divisor > 0.0f ? F90 / Divisor : 1.0f;
-							float3 F90YCoCg = LinearRGB_2_NormalisedYCoCg(F90);
+							// Now we pack diffuse and F0 in a special way: both encoded as R7G7B6 with gamma 2.0. The low bit cound will be hidden by dithering and TAA.
+							// We pack 32 bits in the first uint and the remaining bits are pack is the lowest significant bit of the second uint.
+							const uint PackedDiffuse20Bits = PackR7G7B6Gamma2(SLAB_DIFFUSEALBEDO(BSDF));
+							const uint PackedF020Bits = PackR7G7B6Gamma2(SLAB_F0(BSDF));
+							const uint PackedData32Bits = ((PackedDiffuse20Bits << 12) & 0xFFFFF000) | (PackedF020Bits & 0xFFF);
+							const uint PackedData8Bits = (PackedF020Bits >> 12) & 0xFF;
 
-							STRATA_STORE_UINT1(PackRGBA8(float4(F90YCoCg.y, F90YCoCg.z, 0, SLAB_HAZINESS(BSDF))));
-						}
+							STRATA_STORE_UINT1(PackedData32Bits);
+							STRATA_STORE_UINT1(PackedData8Bits | PackRGBA8(float4(0.0f, SLAB_ROUGHNESS(BSDF), (SLAB_ANISOTROPY(BSDF) + 1.f) * 0.5f, 0.f)));
 
-						const bool bIsSimpleVolume = BSDF_GETISSIMPLEVOLUME(BSDF);
-						const bool bIsThin = BSDF_GETISTHIN(BSDF);
-						if (BSDF_GETHASSSS(BSDF) || bIsSimpleVolume || bIsThin)
-						{
-							if (!bIsThin && BSDF_GETHASSSS(BSDF))
+							if (BSDF_GETHASF90(BSDF) || BSDF_GETHASHAZINESS(BSDF))
 							{
-								SSSData.BaseColor = SLAB_DIFFUSEALBEDO(BSDF);
-								SSSData.Specular = F0ToDielectricSpecular(SLAB_F0(BSDF).x);
+								// What is important is to maintain the hue and saturation, so we scale the color by the maximum of its components
+								float3 F90 = saturate(SLAB_F90(BSDF));
+								const float Divisor = max(F90.r, max(F90.g, F90.b));
+								F90 = Divisor > 0.0f ? F90 / Divisor : 1.0f;
+								float3 F90YCoCg = LinearRGB_2_NormalisedYCoCg(F90);
 
-								// Profile ID: we only keep the latest ID, and profile will overtake any non-profile based SSS
-								SSSData.bIsValid = true;
-								SSSData.bIsProfile = SSSData.bIsProfile || BSDF_GETHASSSSPROFILE(BSDF);
-								SSSData.ProfileId = StrataSubsurfaceProfileIdTo8bits(SLAB_SSSPROFILEID(BSDF));
-								SSSData.ProfileRadiusScale = SLAB_SSSPROFILERADIUSSCALE(BSDF);
+								STRATA_STORE_UINT1(PackRGBA8(float4(F90YCoCg.y, F90YCoCg.z, 0, SLAB_HAZINESS(BSDF))));
+							}
 
-								SSSData.MFP = SLAB_SSSMFP(BSDF);
+							const bool bIsSimpleVolume = BSDF_GETISSIMPLEVOLUME(BSDF);
+							const bool bIsThin = BSDF_GETISTHIN(BSDF);
+							if (BSDF_GETHASSSS(BSDF) || bIsSimpleVolume || bIsThin)
+							{
+								if (!bIsThin && BSDF_GETHASSSS(BSDF))
+								{
+									SSSData.BaseColor = SLAB_DIFFUSEALBEDO(BSDF);
+									SSSData.Specular = F0ToDielectricSpecular(SLAB_F0(BSDF).x);
+
+									// Profile ID: we only keep the latest ID, and profile will overtake any non-profile based SSS
+									SSSData.bIsValid = true;
+									SSSData.bIsProfile = SSSData.bIsProfile || BSDF_GETHASSSSPROFILE(BSDF);
+									SSSData.ProfileId = StrataSubsurfaceProfileIdTo8bits(SLAB_SSSPROFILEID(BSDF));
+									SSSData.ProfileRadiusScale = SLAB_SSSPROFILERADIUSSCALE(BSDF);
+
+									SSSData.MFP = SLAB_SSSMFP(BSDF);
+								}
+
+								if (BSDF_GETHASSSSPROFILE(BSDF) && !bIsSimpleVolume)
+								{
+									// Simple volume overrides SSS profile because it takes over for when not at the bottom of the BSDF layer, or during forward rendering.
+									STRATA_STORE_UINT1(PackRGBA8(float4(SLAB_SSSPROFILEID(BSDF), SLAB_SSSPROFILERADIUSSCALE(BSDF), 0, 0)));
+								}
+								else
+								{
+									// Path used for bottom most layer SSS, simple volume and two sided lighting.
+									STRATA_STORE_UINT1(PackR11G11B10F(SLAB_SSSMFP(BSDF)));
+								}
 							}
 
-							if (BSDF_GETHASSSSPROFILE(BSDF) && !bIsSimpleVolume)
+							if (BSDF_GETHASFUZZ(BSDF))
 							{
-								// Simple volume overrides SSS profile because it takes over for when not at the bottom of the BSDF layer, or during forward rendering.
-								STRATA_STORE_UINT1(PackRGBA8(float4(SLAB_SSSPROFILEID(BSDF), SLAB_SSSPROFILERADIUSSCALE(BSDF), 0, 0)));
+								STRATA_STORE_UINT1(PackColorLinearToGamma2AlphaLinear(float4(SLAB_FUZZ_COLOR(BSDF), SLAB_FUZZ_AMOUNT(BSDF))));
 							}
-							else
+
+							if (BSDF_GETHASTRANSABOVE(BSDF))
 							{
-								// Path used for bottom most layer SSS, simple volume and two sided lighting.
-								STRATA_STORE_UINT1(PackR11G11B10F(SLAB_SSSMFP(BSDF)));
+								STRATA_STORE_UINT1(PackR11G11B10F(BSDF.TransmittanceAboveAlongN));
 							}
-						}
 
-						if (BSDF_GETHASFUZZ(BSDF))
+							// 8-24 bytes
+						}
+						break;
+						case STRATA_BSDF_TYPE_HAIR:
 						{
-							STRATA_STORE_UINT1(PackColorLinearToGamma2AlphaLinear(float4(SLAB_FUZZ_COLOR(BSDF), SLAB_FUZZ_AMOUNT(BSDF))));
+							STRATA_STORE_UINT1(PackColorLinearToGamma2AlphaLinear(float4(HAIR_BASECOLOR(BSDF), HAIR_ROUGHNESS(BSDF))));
+							STRATA_STORE_UINT1(PackRGBA8(float4(HAIR_SCATTER(BSDF), HAIR_SPECULAR(BSDF), HAIR_BACKLIT(BSDF), HAIR_COMPLEXTRANSMITTANCE(BSDF))));
+							// 8 bytes
 						}
-
-						if (BSDF_GETHASTRANSABOVE(BSDF))
+						break;
+						case STRATA_BSDF_TYPE_SINGLELAYERWATER:
 						{
-							STRATA_STORE_UINT1(PackR11G11B10F(BSDF.TransmittanceAboveAlongN));
+							STRATA_STORE_UINT1(PackColorLinearToGamma2AlphaLinear(float4(SLW_BASECOLOR(BSDF), 0.0f)));
+							STRATA_STORE_UINT1(PackRGBA8(float4(SLW_METALLIC(BSDF), SLW_SPECULAR(BSDF), SLW_ROUGHNESS(BSDF), SLW_TOPMATERIALOPACITY(BSDF))));
+							STRATA_STORE_UINT1(PackR11G11B10F(SLW_SEPARATEDMAINDIRLIGHT(BSDF)));	// STRATA_TODO make this an optional load/store base on if DistanceFieldShadow are enabled
+						}
+						break;
 						}
-
-						// 8-24 bytes
-					}
-					break;
-					case STRATA_BSDF_TYPE_HAIR:
-					{
-						STRATA_STORE_UINT1(PackColorLinearToGamma2AlphaLinear(float4(HAIR_BASECOLOR(BSDF), HAIR_ROUGHNESS(BSDF))));
-						STRATA_STORE_UINT1(PackRGBA8(float4(HAIR_SCATTER(BSDF), HAIR_SPECULAR(BSDF), HAIR_BACKLIT(BSDF), HAIR_COMPLEXTRANSMITTANCE(BSDF))));
-						// 8 bytes
-					}
-					break;
-					case STRATA_BSDF_TYPE_SINGLELAYERWATER:
-					{
-						STRATA_STORE_UINT1(PackColorLinearToGamma2AlphaLinear(float4(SLW_BASECOLOR(BSDF), 0.0f)));
-						STRATA_STORE_UINT1(PackRGBA8(float4(SLW_METALLIC(BSDF), SLW_SPECULAR(BSDF), SLW_ROUGHNESS(BSDF), SLW_TOPMATERIALOPACITY(BSDF))));
-						STRATA_STORE_UINT1(PackR11G11B10F(SLW_SEPARATEDMAINDIRLIGHT(BSDF)));	// STRATA_TODO make this an optional load/store base on if DistanceFieldShadow are enabled
-					}
-					break;
 					}
 
 					// Stop writing if above budget.
