Fix incorrect sharpen when using local exposure
- Need to apply local exposure to neighbours.

diff --git a/Engine/Shaders/Private/PostProcessTonemap.usf b/Engine/Shaders/Private/PostProcessTonemap.usf
--- a/Engine/Shaders/Private/PostProcessTonemap.usf
+++ b/Engine/Shaders/Private/PostProcessTonemap.usf
@@ -442,16 +442,21 @@ float4 TonemapCommonPS(
 
 	float ExposureScale = ExposureScaleVignette.x;
 
+	float LocalExposureOneOverPreExposure = View.OneOverPreExposure;
+
 #if USE_LOCAL_EXPOSURE
-	float LuminanceVal = CalculateEyeAdaptationLuminance(SceneColor.rgb);
-	float LogLuminance = log2(LuminanceVal);
-	float MiddleGreyLumValue = ExposureScaleVignette.y;
+	{
+		float LuminanceVal = CalculateEyeAdaptationLuminance(SceneColor.rgb);
+		float LogLuminance = log2(LuminanceVal);
+		float MiddleGreyLumValue = ExposureScaleVignette.y;
 
-	float BaseLogLum = CalculateBaseLogLuminance(LogLuminance, EyeAdaptation_LocalExposureBlurredLuminanceBlend, ExposureScale, FullViewUV, LumBilateralGrid, BlurredLogLum, LumBilateralGridSampler, BlurredLogLumSampler);
-	float LocalExposure = CalculateLocalExposure(LogLuminance + log2(ExposureScale), BaseLogLum, MiddleGreyLumValue, EyeAdaptation_LocalExposureContrastScale, EyeAdaptation_LocalExposureDetailStrength);
+		float BaseLogLum = CalculateBaseLogLuminance(LogLuminance, EyeAdaptation_LocalExposureBlurredLuminanceBlend, ExposureScale, FullViewUV, LumBilateralGrid, BlurredLogLum, LumBilateralGridSampler, BlurredLogLumSampler);
+		float LocalExposure = CalculateLocalExposure(LogLuminance + log2(ExposureScale), BaseLogLum, MiddleGreyLumValue, EyeAdaptation_LocalExposureContrastScale, EyeAdaptation_LocalExposureDetailStrength);
 
-	// Apply before bloom
-	SceneColor.rgb *= LocalExposure;
+		// Apply before bloom
+		SceneColor.rgb *= LocalExposure;
+		LocalExposureOneOverPreExposure *= LocalExposure;
+	}
 #endif
 
 #if DEBUG_TONEMAPPER
@@ -487,8 +492,8 @@ float4 TonemapCommonPS(
 				
 		half4 LuminanceNeighbors;
 
-		half3 C1 = LookupSceneColor(SceneUV, float2(PQMP.QuadVector.x, 0)).rgb * View.OneOverPreExposure;
-		half3 C3 = LookupSceneColor(SceneUV, float2(0, PQMP.QuadVector.y)).rgb * View.OneOverPreExposure;
+		half3 C1 = LookupSceneColor(SceneUV, float2(PQMP.QuadVector.x, 0)).rgb * LocalExposureOneOverPreExposure;
+		half3 C3 = LookupSceneColor(SceneUV, float2(0, PQMP.QuadVector.y)).rgb * LocalExposureOneOverPreExposure;
 		LuminanceNeighbors.x = Luminance(C1);
 		LuminanceNeighbors.y = Luminance(C3);
 		
@@ -497,10 +502,10 @@ float4 TonemapCommonPS(
 		LuminanceNeighbors.z =  PQMPReadX(PQMP,A0);
 		LuminanceNeighbors.w =  PQMPReadY(PQMP,A0);
 	#else
-		half3 C1 = LookupSceneColor(SceneUV, float2(-1,  0)) * View.OneOverPreExposure;
-		half3 C2 = LookupSceneColor(SceneUV, float2( 1,  0)) * View.OneOverPreExposure;
-		half3 C3 = LookupSceneColor(SceneUV, float2( 0, -1)) * View.OneOverPreExposure;
-		half3 C4 = LookupSceneColor(SceneUV, float2( 0,  1)) * View.OneOverPreExposure;
+		half3 C1 = LookupSceneColor(SceneUV, float2(-1,  0)) * LocalExposureOneOverPreExposure;
+		half3 C2 = LookupSceneColor(SceneUV, float2( 1,  0)) * LocalExposureOneOverPreExposure;
+		half3 C3 = LookupSceneColor(SceneUV, float2( 0, -1)) * LocalExposureOneOverPreExposure;
+		half3 C4 = LookupSceneColor(SceneUV, float2( 0,  1)) * LocalExposureOneOverPreExposure;
 		half4 LuminanceNeighbors = half4(Luminance(C1), Luminance(C2), Luminance(C3), Luminance(C4));
 	#endif
 
