Treat max depth as unshadowed when calculating translucent lighting volume.

diff --git a/Engine/Shaders/Private/TranslucentLightInjectionShaders.usf b/Engine/Shaders/Private/TranslucentLightInjectionShaders.usf
--- a/Engine/Shaders/Private/TranslucentLightInjectionShaders.usf
+++ b/Engine/Shaders/Private/TranslucentLightInjectionShaders.usf
@@ -7,6 +7,8 @@
 #include "Common.ush"
 #include "SHCommon.ush"
 
+#define TREAT_MAXDEPTH_UNSHADOWED 1
+
 #if INJECTION_PIXEL_SHADER
 	#include "/Engine/Generated/Material.ush"
 	#include "VolumeLightingCommon.ush"
