Add ShaderPrint triangle support.

This allows to create triangle-based primitive from shaders. Lines and triangles are stored into the same buffer. ShaderSymbol will be merged into this same buffer later on.
This CL also fix an out of bound write issue when allocating lines.

diff --git a/Engine/Shaders/Private/ShaderPrint.ush b/Engine/Shaders/Private/ShaderPrint.ush
--- a/Engine/Shaders/Private/ShaderPrint.ush
+++ b/Engine/Shaders/Private/ShaderPrint.ush
@@ -715,77 +715,287 @@ float AddSlider(inout FShaderPrintContext InContext, FShaderPrintText InTextEntr
 }
 
 ///////////////////////////////////////////////////////////////////////////////////////////////////
-// Primitives (Line, AABB, OBB, Referentials, ...)
+// Primitives types
+
+// Type of primitives
+#define SHADER_PRINT_TYPE_LINE 0
+#define SHADER_PRINT_TYPE_TRIANGLE 1
+
+// Size of packed data, in uint
+#define SHADER_PRINT_UINT_COUNT_LINE 8
+#define SHADER_PRINT_UINT_COUNT_TRIANGLE 12
+
+// Counter offset
+#define SHADER_PRINT_COUNTER_OFFSET_LINE 0
+#define SHADER_PRINT_COUNTER_OFFSET_TRIANGLE 1
+#define SHADER_PRINT_COUNTER_OFFSET_FREE0 2
+#define SHADER_PRINT_COUNTER_OFFSET_FREE1 3
+#define SHADER_PRINT_COUNTER_COUNT 4
+
+// Primitive buffer layout
+//  ------------------------------------------------------------------------------------------------------------------------------------------------
+// ||  Line    | Triangle |  Empty  |  Empty  ||  Line0  |  Line1  |  Line2  |  Line3  |  ...  | Triangle0 | Triangle1 | Triangle2 | Triangle3 | ...
+// || Counter  | Counter  | Counter | Counter ||         |         |         |         |  ...  |           |           |           |           | ...
+// || (1uint)  | (1uint)  | (1uint) | (1uint) || (8uint) | (8uint) | (8uint) | (8uint) |  ...  | (12uint)  | (12uint)  | (12uint)  | (12uint)  | ...
+//  ------------------------------------------------------------------------------------------------------------------------------------------------
+RWStructuredBuffer<uint> ShaderPrint_RWPrimitivesBuffer;
+
+uint GetPrimitiveLineOffset    (uint InIndex) { return SHADER_PRINT_COUNTER_COUNT + SHADER_PRINT_UINT_COUNT_LINE * InIndex; }
+uint GetPrimitiveTriangleOffset(uint InIndex) { return SHADER_PRINT_COUNTER_COUNT + SHADER_PRINT_UINT_COUNT_LINE * ShaderPrintData.MaxLineCount + SHADER_PRINT_UINT_COUNT_TRIANGLE * InIndex; }
 
-struct FShaderDrawElement 
+///////////////////////////////////////////////////////////////////////////////////////////////////
+// Triangle-based Primitives (Triangle, Point, ...)
+
+struct FTriangleElement 
 {
 	float3 Pos0;
 	float3 Pos1;
-	float4 Color0;
-	float4 Color1;
+	float3 Pos2;
+	float4 Color;
 	bool bIsScreenSpace;
 };
 
-FShaderDrawElement UnpackShaderElement(StructuredBuffer<uint> InPrimitiveBuffer, uint InIndex)
+struct FPackedTriangleElement
 {
-	const uint Index8 = (1 + InIndex) * 8;
+	uint4 Packed0;
+	uint4 Packed1;
+	uint4 Packed2;
+};
 
-	FShaderDrawElement Out = (FShaderDrawElement)0;
+FTriangleElement UnpackTriangleElement(FPackedTriangleElement In)
+{
+	FTriangleElement Out = (FTriangleElement)0;
 	{
-		Out.Pos0.x = asfloat(InPrimitiveBuffer[Index8 + 0]);
-		Out.Pos0.y = asfloat(InPrimitiveBuffer[Index8 + 1]);
-		Out.Pos0.z = asfloat(InPrimitiveBuffer[Index8 + 2]);
-
-		const uint PackedColor0 = InPrimitiveBuffer[Index8 + 3];
-		Out.Color0 = float4((PackedColor0 >> 24) & 0xFF, (PackedColor0 >> 16) & 0xFF, (PackedColor0 >> 8) & 0xFF, (PackedColor0) & 0xFF) / 255.0f;
+		Out.Pos0.x = asfloat(In.Packed0.x);
+		Out.Pos0.y = asfloat(In.Packed0.y);
+		Out.Pos0.z = asfloat(In.Packed0.z);
+		Out.Color  = float4((In.Packed0.w >> 24) & 0xFF, (In.Packed0.w >> 16) & 0xFF, (In.Packed0.w >> 8) & 0xFF, (In.Packed0.w) & 0xFF) / 255.0f;
 	}
 	{
-		Out.Pos1.x = asfloat(InPrimitiveBuffer[Index8 + 4]);
-		Out.Pos1.y = asfloat(InPrimitiveBuffer[Index8 + 5]);
-		Out.Pos1.z = asfloat(InPrimitiveBuffer[Index8 + 6]);
-
-		const uint PackedColor1 = InPrimitiveBuffer[Index8 + 7];
-		Out.Color1 = float4((PackedColor1 >> 24) & 0xFF, (PackedColor1 >> 16) & 0xFF, (PackedColor1 >> 8) & 0xFF, (PackedColor1) & 0xFF) / 255.0f;
+		Out.Pos1.x = asfloat(In.Packed1.x);
+		Out.Pos1.y = asfloat(In.Packed1.y);
+		Out.Pos1.z = asfloat(In.Packed1.z);
 
 		// The 'space' info is stored into the alpha's LSB
-		Out.bIsScreenSpace = (PackedColor1 & 0x1) > 0u;
+		Out.bIsScreenSpace = (In.Packed1.w & 0x1) > 0u;
+	}
+	{
+		Out.Pos2.x = asfloat(In.Packed2.x);
+		Out.Pos2.y = asfloat(In.Packed2.y);
+		Out.Pos2.z = asfloat(In.Packed2.z);
 	}
 	return Out;
 }
 
-RWStructuredBuffer<uint> ShaderPrint_RWLinesBuffer;
+FPackedTriangleElement PackTriangleElement(FTriangleElement In)
+{
+	uint4 PackedC = uint4(255.0f * saturate(In.Color));
+
+	FPackedTriangleElement Out = (FPackedTriangleElement)0;
+	Out.Packed0.x = asuint(In.Pos0.x);
+	Out.Packed0.y = asuint(In.Pos0.y);
+	Out.Packed0.z = asuint(In.Pos0.z);
+	Out.Packed0.w = (PackedC.x << 24) | (PackedC.y << 16) | (PackedC.z << 8) | (PackedC.w);
+
+	Out.Packed1.x = asuint(In.Pos1.x);
+	Out.Packed1.y = asuint(In.Pos1.y);
+	Out.Packed1.z = asuint(In.Pos1.z);
+	Out.Packed1.w = In.bIsScreenSpace ? 0x1 : 0x0;
+
+	Out.Packed2.x = asuint(In.Pos2.x);
+	Out.Packed2.y = asuint(In.Pos2.y);
+	Out.Packed2.z = asuint(In.Pos2.z);
+	Out.Packed2.w = 0;
+	return Out;
+}
 
-bool AllocateShaderDrawElement(uint Count, inout uint OutIndex)
+FTriangleElement UnpackTriangleElement(StructuredBuffer<uint> InPrimitiveBuffer, uint InIndex)
+{
+	const uint Index12 = GetPrimitiveTriangleOffset(InIndex);
+
+	FPackedTriangleElement Out = (FPackedTriangleElement)0;
+	Out.Packed0.x = InPrimitiveBuffer[Index12 + 0];
+	Out.Packed0.y = InPrimitiveBuffer[Index12 + 1];
+	Out.Packed0.z = InPrimitiveBuffer[Index12 + 2];
+	Out.Packed0.w = InPrimitiveBuffer[Index12 + 3];
+
+	Out.Packed1.x = InPrimitiveBuffer[Index12 + 4];
+	Out.Packed1.y = InPrimitiveBuffer[Index12 + 5];
+	Out.Packed1.z = InPrimitiveBuffer[Index12 + 6];
+	Out.Packed1.w = InPrimitiveBuffer[Index12 + 7];
+
+	Out.Packed2.x = InPrimitiveBuffer[Index12 + 8];
+	Out.Packed2.y = InPrimitiveBuffer[Index12 + 9];
+	Out.Packed2.z = InPrimitiveBuffer[Index12 + 10];
+	Out.Packed2.w = InPrimitiveBuffer[Index12 + 11];
+	return UnpackTriangleElement(Out);
+}
+
+bool AllocateTriangleElement(uint Count, inout uint OutIndex)
 {
 	OutIndex = 0;
-	if (ShaderPrintData.MaxLineCount == 0)
+	if (ShaderPrintData.MaxTriangleCount == 0)
 	{
 		return false;
 	}
 
-	InterlockedAdd(ShaderPrint_RWLinesBuffer[0], Count, OutIndex);
-	return (1+OutIndex) < ShaderPrintData.MaxLineCount;
+	InterlockedAdd(ShaderPrint_RWPrimitivesBuffer[SHADER_PRINT_COUNTER_OFFSET_TRIANGLE], Count, OutIndex);
+	return (OutIndex + Count) < ShaderPrintData.MaxTriangleCount;
+}
+
+void AddTriangleElement(FTriangleElement In, uint Index)
+{
+	const uint Index12 = GetPrimitiveTriangleOffset(Index);
+
+	FPackedTriangleElement Out = PackTriangleElement(In);
+	ShaderPrint_RWPrimitivesBuffer[Index12 + 0] = Out.Packed0.x;
+	ShaderPrint_RWPrimitivesBuffer[Index12 + 1] = Out.Packed0.y;
+	ShaderPrint_RWPrimitivesBuffer[Index12 + 2] = Out.Packed0.z;
+	ShaderPrint_RWPrimitivesBuffer[Index12 + 3] = Out.Packed0.w;
+
+	ShaderPrint_RWPrimitivesBuffer[Index12 + 4] = Out.Packed1.x;
+	ShaderPrint_RWPrimitivesBuffer[Index12 + 5] = Out.Packed1.y;
+	ShaderPrint_RWPrimitivesBuffer[Index12 + 6] = Out.Packed1.z;
+	ShaderPrint_RWPrimitivesBuffer[Index12 + 7] = Out.Packed1.w;
+
+	ShaderPrint_RWPrimitivesBuffer[Index12 + 8] = Out.Packed2.x;
+	ShaderPrint_RWPrimitivesBuffer[Index12 + 9] = Out.Packed2.y;
+	ShaderPrint_RWPrimitivesBuffer[Index12 + 10]= Out.Packed2.z;
+	ShaderPrint_RWPrimitivesBuffer[Index12 + 11]= Out.Packed2.w;
+}
+
+///////////////////////////////////////////////////////////////////////////////////////////////////
+// Filled Quad
+
+void AddFilledQuadTWS(float3 Pos0, float3 Pos1, float3 Pos2, float3 Pos3, float4 Color)
+{
+	uint Offset = 0;
+	if (AllocateTriangleElement(2, Offset))
+	{
+		FTriangleElement Element;
+		Element.bIsScreenSpace = false;
+		Element.Color = Color;
+
+		Element.Pos0 = Pos0;
+		Element.Pos1 = Pos1;
+		Element.Pos2 = Pos2;
+		AddTriangleElement(Element, Offset + 0);
+
+		Element.Pos0 = Pos0;
+		Element.Pos1 = Pos2;
+		Element.Pos2 = Pos3;
+		AddTriangleElement(Element, Offset + 1);
+	}
+}
+
+void AddFilledQuadWS (float3 Pos0, float3 Pos1, float3 Pos2, float3 Pos3, float4 Color)							 { AddFilledQuadTWS(Pos0 + ShaderPrintData.TranslatedWorldOffset, Pos1 + ShaderPrintData.TranslatedWorldOffset, Pos2 + ShaderPrintData.TranslatedWorldOffset, Pos3 + ShaderPrintData.TranslatedWorldOffset, Color); }
+void AddFilledQuadTWS(FShaderPrintContext Ctx, float3 Pos0, float3 Pos1, float3 Pos2, float3 Pos3, float4 Color) { if (Ctx.bIsActive) { AddFilledQuadTWS(Pos0, Pos1, Pos2, Pos3, Color); } }
+void AddFilledQuadWS (FShaderPrintContext Ctx, float3 Pos0, float3 Pos1, float3 Pos2, float3 Pos3, float4 Color) { if (Ctx.bIsActive) { AddFilledQuadTWS(Pos0 + ShaderPrintData.TranslatedWorldOffset, Pos1 + ShaderPrintData.TranslatedWorldOffset, Pos2 + ShaderPrintData.TranslatedWorldOffset, Pos3 + ShaderPrintData.TranslatedWorldOffset, Color); } }
+
+///////////////////////////////////////////////////////////////////////////////////////////////////
+// Line-based Primitives (Line, AABB, OBB, Referentials, ...)
+
+struct FLineElement 
+{
+	float3 Pos0;
+	float3 Pos1;
+	float4 Color0;
+	float4 Color1;
+	bool bIsScreenSpace;
+};
+
+struct FPackedLineElement
+{
+	uint4 Packed0;
+	uint4 Packed1;
+};
+
+FLineElement UnpackLineElement(FPackedLineElement In)
+{
+	FLineElement Out = (FLineElement)0;
+	{
+		Out.Pos0.x = asfloat(In.Packed0.x);
+		Out.Pos0.y = asfloat(In.Packed0.y);
+		Out.Pos0.z = asfloat(In.Packed0.z);
+		Out.Color0 = float4((In.Packed0.w >> 24) & 0xFF, (In.Packed0.w >> 16) & 0xFF, (In.Packed0.w >> 8) & 0xFF, (In.Packed0.w) & 0xFF) / 255.0f;
+	}
+	{
+		Out.Pos1.x = asfloat(In.Packed1.x);
+		Out.Pos1.y = asfloat(In.Packed1.y);
+		Out.Pos1.z = asfloat(In.Packed1.z);
+		Out.Color1 = float4((In.Packed1.w >> 24) & 0xFF, (In.Packed1.w >> 16) & 0xFF, (In.Packed1.w >> 8) & 0xFF, (In.Packed1.w) & 0xFF) / 255.0f;
+
+		// The 'space' info is stored into the alpha's LSB
+		Out.bIsScreenSpace = (In.Packed1.w & 0x1) > 0u;
+	}
+	return Out;
 }
 
-void AddShaderElement(FShaderDrawElement E, uint Index)
+FPackedLineElement PackLineElement(FLineElement In)
 {
-	uint4 PackedC0 = uint4(255.0f * saturate(E.Color0));
-	uint4 PackedC1 = uint4(255.0f * saturate(E.Color1));
+	uint4 PackedC0 = uint4(255.0f * saturate(In.Color0));
+	uint4 PackedC1 = uint4(255.0f * saturate(In.Color1));
 
 	// Store 'space' info into the alpha's LSB
 	PackedC1 = PackedC1 & 0xFE;
-	PackedC1 = PackedC1 | (E.bIsScreenSpace ? 0x1 : 0x0);
-
-	const uint Index8 = (1 + Index) * 8;
-	ShaderPrint_RWLinesBuffer[Index8 + 0] = asuint(E.Pos0.x);
-	ShaderPrint_RWLinesBuffer[Index8 + 1] = asuint(E.Pos0.y);
-	ShaderPrint_RWLinesBuffer[Index8 + 2] = asuint(E.Pos0.z);
-	ShaderPrint_RWLinesBuffer[Index8 + 3] = (PackedC0.x << 24) | (PackedC0.y << 16) | (PackedC0.z << 8) | (PackedC0.w);
-
-	ShaderPrint_RWLinesBuffer[Index8 + 4] = asuint(E.Pos1.x);
-	ShaderPrint_RWLinesBuffer[Index8 + 5] = asuint(E.Pos1.y);
-	ShaderPrint_RWLinesBuffer[Index8 + 6] = asuint(E.Pos1.z);
-	ShaderPrint_RWLinesBuffer[Index8 + 7] = (PackedC1.x << 24) | (PackedC1.y << 16) | (PackedC1.z << 8) | (PackedC1.w);
+	PackedC1 = PackedC1 | (In.bIsScreenSpace ? 0x1 : 0x0);
+
+	FPackedLineElement Out = (FPackedLineElement)0;
+	Out.Packed0.x = asuint(In.Pos0.x);
+	Out.Packed0.y = asuint(In.Pos0.y);
+	Out.Packed0.z = asuint(In.Pos0.z);
+	Out.Packed0.w = (PackedC0.x << 24) | (PackedC0.y << 16) | (PackedC0.z << 8) | (PackedC0.w);
+
+	Out.Packed1.x = asuint(In.Pos1.x);
+	Out.Packed1.y = asuint(In.Pos1.y);
+	Out.Packed1.z = asuint(In.Pos1.z);
+	Out.Packed1.w = (PackedC1.x << 24) | (PackedC1.y << 16) | (PackedC1.z << 8) | (PackedC1.w);
+	return Out;
+}
+
+FLineElement UnpackLineElement(StructuredBuffer<uint> InPrimitiveBuffer, uint InIndex)
+{
+	const uint Index8 = GetPrimitiveLineOffset(InIndex);
+
+	FPackedLineElement Out = (FPackedLineElement)0;
+	Out.Packed0.x = InPrimitiveBuffer[Index8 + 0];
+	Out.Packed0.y = InPrimitiveBuffer[Index8 + 1];
+	Out.Packed0.z = InPrimitiveBuffer[Index8 + 2];
+	Out.Packed0.w = InPrimitiveBuffer[Index8 + 3];
+
+	Out.Packed1.x = InPrimitiveBuffer[Index8 + 4];
+	Out.Packed1.y = InPrimitiveBuffer[Index8 + 5];
+	Out.Packed1.z = InPrimitiveBuffer[Index8 + 6];
+	Out.Packed1.w = InPrimitiveBuffer[Index8 + 7];
+	return UnpackLineElement(Out);
+}
+
+bool AllocateLineElement(uint Count, inout uint OutIndex)
+{
+	OutIndex = 0;
+	if (ShaderPrintData.MaxLineCount == 0)
+	{
+		return false;
+	}
+
+	InterlockedAdd(ShaderPrint_RWPrimitivesBuffer[SHADER_PRINT_COUNTER_OFFSET_LINE], Count, OutIndex);
+	return (OutIndex + Count) < ShaderPrintData.MaxLineCount;
+}
+
+void AddLineElement(FLineElement In, uint Index)
+{
+	const uint Index8 = GetPrimitiveLineOffset(Index);
+
+	FPackedLineElement Out = PackLineElement(In);
+	ShaderPrint_RWPrimitivesBuffer[Index8 + 0] = Out.Packed0.x;
+	ShaderPrint_RWPrimitivesBuffer[Index8 + 1] = Out.Packed0.y;
+	ShaderPrint_RWPrimitivesBuffer[Index8 + 2] = Out.Packed0.z;
+	ShaderPrint_RWPrimitivesBuffer[Index8 + 3] = Out.Packed0.w;
+
+	ShaderPrint_RWPrimitivesBuffer[Index8 + 4] = Out.Packed1.x;
+	ShaderPrint_RWPrimitivesBuffer[Index8 + 5] = Out.Packed1.y;
+	ShaderPrint_RWPrimitivesBuffer[Index8 + 6] = Out.Packed1.z;
+	ShaderPrint_RWPrimitivesBuffer[Index8 + 7] = Out.Packed1.w;
 }
 
 int2 GetCursorPos()
@@ -799,15 +1009,15 @@ int2 GetCursorPos()
 void AddLineTWS(float3 Pos0, float3 Pos1, float4 Color0, float4 Color1)
 {
 	uint Offset = 0;
-	if (AllocateShaderDrawElement(1, Offset))
+	if (AllocateLineElement(1, Offset))
 	{
-		FShaderDrawElement Element;
+		FLineElement Element;
 		Element.Pos0 = Pos0;
 		Element.Pos1 = Pos1;
 		Element.Color0 = Color0;
 		Element.Color1 = Color1;
 		Element.bIsScreenSpace = false;
-		AddShaderElement(Element, Offset);
+		AddLineElement(Element, Offset);
 	}
 }
 
@@ -831,15 +1041,15 @@ void AddLineSS(float2 InPosInPixel0, float2 InPosInPixel1, float4 Color0, float4
 	const float2 Pos1 = float2(InPosInPixel1) / float2(ShaderPrintData.Resolution);
 
 	uint Offset = 0;
-	if (AllocateShaderDrawElement(1, Offset))
+	if (AllocateLineElement(1, Offset))
 	{
-		FShaderDrawElement Element;
+		FLineElement Element;
 		Element.Pos0 = float3(Pos0, 0);
 		Element.Pos1 = float3(Pos1, 0);
 		Element.Color0 = Color0;
 		Element.Color1 = Color1;
 		Element.bIsScreenSpace = true;
-		AddShaderElement(Element, Offset);
+		AddLineElement(Element, Offset);
 	}
 }
 void AddLineSS(float2 Pos0, float2 Pos1, float4 Color)											{ AddLineSS(Pos0, Pos1, Color, Color); }
@@ -853,15 +1063,15 @@ void AddLineSS(FShaderPrintContext Ctx, float2 Pos0, float2 Pos1, float4 Color)
 void AddQuadTWS(float3 Pos0, float3 Pos1, float3 Pos2, float3 Pos3, float4 Color)
 {
 	uint Offset = 0;
-	if (AllocateShaderDrawElement(4, Offset))
+	if (AllocateLineElement(4, Offset))
 	{
-		FShaderDrawElement Element;
+		FLineElement Element;
 		Element.bIsScreenSpace = false;
 		Element.Color0 = Element.Color1 = Color;
-		Element.Pos0 = Pos0; Element.Pos1 = Pos1; AddShaderElement(Element, Offset + 0);
-		Element.Pos0 = Pos1; Element.Pos1 = Pos2; AddShaderElement(Element, Offset + 1);
-		Element.Pos0 = Pos2; Element.Pos1 = Pos3; AddShaderElement(Element, Offset + 2);
-		Element.Pos0 = Pos3; Element.Pos1 = Pos0; AddShaderElement(Element, Offset + 3);
+		Element.Pos0 = Pos0; Element.Pos1 = Pos1; AddLineElement(Element, Offset + 0);
+		Element.Pos0 = Pos1; Element.Pos1 = Pos2; AddLineElement(Element, Offset + 1);
+		Element.Pos0 = Pos2; Element.Pos1 = Pos3; AddLineElement(Element, Offset + 2);
+		Element.Pos0 = Pos3; Element.Pos1 = Pos0; AddLineElement(Element, Offset + 3);
 	}
 }
 
@@ -878,7 +1088,7 @@ void AddQuad   (FShaderPrintContext Ctx, float3 Pos0, float3 Pos1, float3 Pos2,
 void AddAABBTWS(float3 Min, float3 Max, float4 Color)
 {
 	uint Offset = 0;
-	if (AllocateShaderDrawElement(12, Offset))
+	if (AllocateLineElement(12, Offset))
 	{
 		float3 P0 = float3(Min.x, Min.y, Min.z);
 		float3 P1 = float3(Max.x, Min.y, Min.z);
@@ -889,24 +1099,24 @@ void AddAABBTWS(float3 Min, float3 Max, float4 Color)
 		float3 P6 = float3(Max.x, Max.y, Max.z);
 		float3 P7 = float3(Min.x, Max.y, Max.z);
 
-		FShaderDrawElement Element;
+		FLineElement Element;
 		Element.bIsScreenSpace = false;
 		Element.Color0 = Element.Color1 = Color;
 
-		Element.Pos0 = P0; Element.Pos1 = P1; AddShaderElement(Element, Offset + 0);
-		Element.Pos0 = P1; Element.Pos1 = P2; AddShaderElement(Element, Offset + 1);
-		Element.Pos0 = P2; Element.Pos1 = P3; AddShaderElement(Element, Offset + 2);
-		Element.Pos0 = P3; Element.Pos1 = P0; AddShaderElement(Element, Offset + 3);
+		Element.Pos0 = P0; Element.Pos1 = P1; AddLineElement(Element, Offset + 0);
+		Element.Pos0 = P1; Element.Pos1 = P2; AddLineElement(Element, Offset + 1);
+		Element.Pos0 = P2; Element.Pos1 = P3; AddLineElement(Element, Offset + 2);
+		Element.Pos0 = P3; Element.Pos1 = P0; AddLineElement(Element, Offset + 3);
 
-		Element.Pos0 = P4; Element.Pos1 = P5; AddShaderElement(Element, Offset + 4);
-		Element.Pos0 = P5; Element.Pos1 = P6; AddShaderElement(Element, Offset + 5);
-		Element.Pos0 = P6; Element.Pos1 = P7; AddShaderElement(Element, Offset + 6);
-		Element.Pos0 = P7; Element.Pos1 = P4; AddShaderElement(Element, Offset + 7);
+		Element.Pos0 = P4; Element.Pos1 = P5; AddLineElement(Element, Offset + 4);
+		Element.Pos0 = P5; Element.Pos1 = P6; AddLineElement(Element, Offset + 5);
+		Element.Pos0 = P6; Element.Pos1 = P7; AddLineElement(Element, Offset + 6);
+		Element.Pos0 = P7; Element.Pos1 = P4; AddLineElement(Element, Offset + 7);
 
-		Element.Pos0 = P0; Element.Pos1 = P4; AddShaderElement(Element, Offset + 8);
-		Element.Pos0 = P1; Element.Pos1 = P5; AddShaderElement(Element, Offset + 9);
-		Element.Pos0 = P2; Element.Pos1 = P6; AddShaderElement(Element, Offset +10);
-		Element.Pos0 = P3; Element.Pos1 = P7; AddShaderElement(Element, Offset +11);
+		Element.Pos0 = P0; Element.Pos1 = P4; AddLineElement(Element, Offset + 8);
+		Element.Pos0 = P1; Element.Pos1 = P5; AddLineElement(Element, Offset + 9);
+		Element.Pos0 = P2; Element.Pos1 = P6; AddLineElement(Element, Offset +10);
+		Element.Pos0 = P3; Element.Pos1 = P7; AddLineElement(Element, Offset +11);
 	}
 }
 
@@ -1029,7 +1239,7 @@ void AddLineTriangle   (FShaderPrintContext Ctx, float3 P0, float3 P1, float3 P2
 void AddOBBTWS(float3 Min, float3 Max, float4 Color, float4x4 Transform)
 {
 	uint Offset = 0;
-	if (AllocateShaderDrawElement(12, Offset))
+	if (AllocateLineElement(12, Offset))
 	{
 		float3 P0 = mul(float4(Min.x, Min.y, Min.z, 1.0f), Transform).xyz;
 		float3 P1 = mul(float4(Max.x, Min.y, Min.z, 1.0f), Transform).xyz;
@@ -1040,24 +1250,24 @@ void AddOBBTWS(float3 Min, float3 Max, float4 Color, float4x4 Transform)
 		float3 P6 = mul(float4(Max.x, Max.y, Max.z, 1.0f), Transform).xyz;
 		float3 P7 = mul(float4(Min.x, Max.y, Max.z, 1.0f), Transform).xyz;
 
-		FShaderDrawElement Element;
+		FLineElement Element;
 		Element.bIsScreenSpace = false;
 		Element.Color0 = Element.Color1 = Color;
 
-		Element.Pos0 = P0; Element.Pos1 = P1; AddShaderElement(Element, Offset + 0);
-		Element.Pos0 = P1; Element.Pos1 = P2; AddShaderElement(Element, Offset + 1);
-		Element.Pos0 = P2; Element.Pos1 = P3; AddShaderElement(Element, Offset + 2);
-		Element.Pos0 = P3; Element.Pos1 = P0; AddShaderElement(Element, Offset + 3);
+		Element.Pos0 = P0; Element.Pos1 = P1; AddLineElement(Element, Offset + 0);
+		Element.Pos0 = P1; Element.Pos1 = P2; AddLineElement(Element, Offset + 1);
+		Element.Pos0 = P2; Element.Pos1 = P3; AddLineElement(Element, Offset + 2);
+		Element.Pos0 = P3; Element.Pos1 = P0; AddLineElement(Element, Offset + 3);
 
-		Element.Pos0 = P4; Element.Pos1 = P5; AddShaderElement(Element, Offset + 4);
-		Element.Pos0 = P5; Element.Pos1 = P6; AddShaderElement(Element, Offset + 5);
-		Element.Pos0 = P6; Element.Pos1 = P7; AddShaderElement(Element, Offset + 6);
-		Element.Pos0 = P7; Element.Pos1 = P4; AddShaderElement(Element, Offset + 7);
+		Element.Pos0 = P4; Element.Pos1 = P5; AddLineElement(Element, Offset + 4);
+		Element.Pos0 = P5; Element.Pos1 = P6; AddLineElement(Element, Offset + 5);
+		Element.Pos0 = P6; Element.Pos1 = P7; AddLineElement(Element, Offset + 6);
+		Element.Pos0 = P7; Element.Pos1 = P4; AddLineElement(Element, Offset + 7);
 
-		Element.Pos0 = P0; Element.Pos1 = P4; AddShaderElement(Element, Offset + 8);
-		Element.Pos0 = P1; Element.Pos1 = P5; AddShaderElement(Element, Offset + 9);
-		Element.Pos0 = P2; Element.Pos1 = P6; AddShaderElement(Element, Offset + 10);
-		Element.Pos0 = P3; Element.Pos1 = P7; AddShaderElement(Element, Offset + 11);
+		Element.Pos0 = P0; Element.Pos1 = P4; AddLineElement(Element, Offset + 8);
+		Element.Pos0 = P1; Element.Pos1 = P5; AddLineElement(Element, Offset + 9);
+		Element.Pos0 = P2; Element.Pos1 = P6; AddLineElement(Element, Offset + 10);
+		Element.Pos0 = P3; Element.Pos1 = P7; AddLineElement(Element, Offset + 11);
 	}
 }
 
diff --git a/Engine/Shaders/Private/ShaderPrintDrawPrimitive.usf b/Engine/Shaders/Private/ShaderPrintDrawPrimitive.usf
--- a/Engine/Shaders/Private/ShaderPrintDrawPrimitive.usf
+++ b/Engine/Shaders/Private/ShaderPrintDrawPrimitive.usf
@@ -19,14 +19,14 @@ float4x4 TranslatedWorldToClip;
 
 SamplerState DepthSampler;
 
-#if GPU_DEBUG_RENDERING_VS 
+#if GPU_DEBUG_RENDERING_LINE_VS 
 void ShaderDrawDebugVS(
 	in uint InstanceId : SV_InstanceID,
 	in uint VertexId : SV_VertexID,
 	out float4 Position : SV_POSITION,
 	out float4 ColorAlpha : TEXCOORD0)
 {
-	const FShaderDrawElement Item = UnpackShaderElement(ShaderDrawDebugPrimitive, InstanceId);
+	const FLineElement Item = UnpackLineElement(ShaderDrawDebugPrimitive, InstanceId);
 
 	float4 VertexPositionTWS = float4(VertexId == 0 ? Item.Pos0 : Item.Pos1, 1.0f);
 	float4 VertexColor = VertexId == 0 ? Item.Color0 : Item.Color1;
@@ -53,6 +53,40 @@ void ShaderDrawDebugVS(
 }
 #endif
 
+#if GPU_DEBUG_RENDERING_TRIANGLE_VS 
+void ShaderDrawDebugVS(
+	in uint InstanceId : SV_InstanceID,
+	in uint VertexId : SV_VertexID,
+	out float4 Position : SV_POSITION,
+	out float4 ColorAlpha : TEXCOORD0)
+{
+	const FTriangleElement Item = UnpackTriangleElement(ShaderDrawDebugPrimitive, InstanceId);
+
+	float4 VertexPositionTWS = float4(VertexId == 0 ? Item.Pos0 : (VertexId == 1 ? Item.Pos1 : Item.Pos2), 1.0f);
+	float4 VertexColor = Item.Color;
+
+	if (Item.bIsScreenSpace)
+	{
+		Position = VertexPositionTWS;
+		Position.xy = Position.xy * 2 - 1;
+		Position.y = -Position.y;
+
+		Position.z = 1;
+		Position.w = 1;
+	}
+	else
+	{
+		// Optional world offset. For 'locked' debug drawing, this offset translates the world 
+		// position recorded for a given translated world offset, into the current translated 
+		// world offset
+		VertexPositionTWS.xyz += TranslatedWorldOffsetConversion;
+
+		Position = mul(VertexPositionTWS, TranslatedWorldToClip);
+	}
+	ColorAlpha = VertexColor;
+}
+#endif
+
 #if GPU_DEBUG_RENDERING_PS
 void ShaderDrawDebugPS(
 	in float4 Position : SV_POSITION,
@@ -84,10 +118,6 @@ void ShaderDrawDebugClearCS(uint3 DispatchThreadId : SV_DispatchThreadID)
 	RWElementBuffer[1] = 0;
 	RWElementBuffer[2] = 0;
 	RWElementBuffer[3] = 0;
-	RWElementBuffer[4] = 0;
-	RWElementBuffer[5] = 0;
-	RWElementBuffer[6] = 0;
-	RWElementBuffer[7] = 0;
 }
 #endif
 
@@ -95,14 +125,18 @@ void ShaderDrawDebugClearCS(uint3 DispatchThreadId : SV_DispatchThreadID)
 
 StructuredBuffer<uint> ElementBuffer;
 RWBuffer<uint> RWIndirectArgs;
-uint MaxInstanceCount;
+uint PrimitiveType;
 
 [numthreads(1, 1, 1)]
 void ShaderDrawDebugCopyCS(uint3 DispatchThreadId : SV_DispatchThreadID)
 {
-	uint IndirectArgInstanceCount = min(MaxInstanceCount, ElementBuffer[0]);
+	const bool bLine    = PrimitiveType == SHADER_PRINT_TYPE_LINE;
+	const uint Index    = bLine ? SHADER_PRINT_COUNTER_OFFSET_LINE : SHADER_PRINT_COUNTER_OFFSET_TRIANGLE;
+	const uint MaxCount = bLine ? ShaderPrintData.MaxLineCount : ShaderPrintData.MaxTriangleCount;
+
+	const uint IndirectArgInstanceCount = min(MaxCount, ElementBuffer[Index]);
 
-	RWIndirectArgs[0] = 2; // Line primitive
+	RWIndirectArgs[0] = bLine ? 2 : 3; // Line or Triangle primitive
 	RWIndirectArgs[1] = IndirectArgInstanceCount;
 	RWIndirectArgs[2] = 0;
 	RWIndirectArgs[3] = 0;
diff --git a/Engine/Source/Runtime/Renderer/Private/ScenePrivate.h b/Engine/Source/Runtime/Renderer/Private/ScenePrivate.h
--- a/Engine/Source/Runtime/Renderer/Private/ScenePrivate.h
+++ b/Engine/Source/Runtime/Renderer/Private/ScenePrivate.h
@@ -623,7 +623,7 @@ struct FHLODSceneNodeVisibilityState
 struct FShaderPrintStateData
 {
 	TRefCountPtr<FRDGPooledBuffer> StateBuffer;
-	TRefCountPtr<FRDGPooledBuffer> LineBuffer;
+	TRefCountPtr<FRDGPooledBuffer> PrimitiveBuffer;
 	FVector PreViewTranslation = FVector::ZeroVector;
 	bool bIsLocked = false;
 
@@ -632,7 +632,7 @@ struct FShaderPrintStateData
 		bIsLocked = false;
 		PreViewTranslation = FVector::ZeroVector;
 		StateBuffer = nullptr;
-		LineBuffer = nullptr;
+		PrimitiveBuffer = nullptr;
 	}
 };
 
diff --git a/Engine/Source/Runtime/Renderer/Private/ShaderPrint.cpp b/Engine/Source/Runtime/Renderer/Private/ShaderPrint.cpp
--- a/Engine/Source/Runtime/Renderer/Private/ShaderPrint.cpp
+++ b/Engine/Source/Runtime/Renderer/Private/ShaderPrint.cpp
@@ -62,6 +62,12 @@ namespace ShaderPrint
 		TEXT("ShaderPrint max line count.\n"),
 		ECVF_Cheat | ECVF_RenderThreadSafe);
 
+	static TAutoConsoleVariable<int32> CVarMaxTriangleCount(
+		TEXT("r.ShaderPrint.MaxTriangle"),
+		32,
+		TEXT("ShaderPrint max triangle count.\n"),
+		ECVF_Cheat | ECVF_RenderThreadSafe);
+
 	static TAutoConsoleVariable<int32> CVarDrawLock(
 		TEXT("r.ShaderPrint.Lock"),
 		0,
@@ -74,6 +80,7 @@ namespace ShaderPrint
 	static uint32 GWidgetRequestCount = 0;
 	static uint32 GCharacterRequestCount = 0;
 	static uint32 GLineRequestCount = 0;
+	static uint32 GTriangleRequestCount = 0;
 	static bool GShaderPrintEnableOverride = false;
 	static FViewInfo* GDefaultView = nullptr;
 
@@ -121,13 +128,17 @@ namespace ShaderPrint
 		return IsEnabled() ? MaxValueCount : 0;
 	}
 
-	// Returns the element count for a whole line
-	static uint32 GetLineElementCount()
+	static uint32 GetMaxTriangleCount()
 	{
-		// x,y,z (3) and a packed color (1) for both the start and end points (2x)
-		return 8;
+		uint32 MaxValueCount = FMath::Max(CVarMaxTriangleCount.GetValueOnRenderThread() + int32(GTriangleRequestCount), 0);
+		return IsEnabled() ? MaxValueCount : 0;
 	}
 
+	// Returns the number of uints used for counters, a line element, and a triangle elements
+	static uint32 GetCountersUintSize()       { return 4;  }
+	static uint32 GetPackedLineUintSize()     { return 8; }
+	static uint32 GetPackedTriangleUintSize() { return 12; }
+
 	// Get symbol buffer size
 	// This is some multiple of the value buffer size to allow for maximum value->symbol expansion
 	static uint32 GetMaxSymbolCountFromValueCount(uint32 MaxValueCount)
@@ -171,6 +182,7 @@ namespace ShaderPrint
 		Out.MaxSymbolCount = GetMaxSymbolCountFromValueCount(Setup.MaxValueCount);
 		Out.MaxStateCount = Setup.MaxStateCount;
 		Out.MaxLineCount = Setup.MaxLineCount;
+		Out.MaxTriangleCount = Setup.MaxTriangleCount;
 		Out.TranslatedWorldOffset = FVector3f(Setup.PreViewTranslation);
 
 		return TUniformBufferRef<ShaderPrint::FShaderPrintCommonParameters>::CreateUniformBufferImmediate(Out, UniformBuffer_SingleFrame);
@@ -185,7 +197,7 @@ namespace ShaderPrint
 		OutParameters.Common = Data.UniformBuffer;
 		OutParameters.ShaderPrint_StateBuffer = GraphBuilder.CreateSRV(Data.ShaderPrintStateBuffer);
 		OutParameters.ShaderPrint_RWValuesBuffer = GraphBuilder.CreateUAV(Data.ShaderPrintValueBuffer);
-		OutParameters.ShaderPrint_RWLinesBuffer = GraphBuilder.CreateUAV(Data.ShaderPrintLineBuffer);
+		OutParameters.ShaderPrint_RWPrimitivesBuffer = GraphBuilder.CreateUAV(Data.ShaderPrintPrimitiveBuffer);
 	}
 
 	void SetParameters(FRDGBuilder& GraphBuilder, FShaderParameters& OutParameters)
@@ -235,6 +247,11 @@ namespace ShaderPrint
 		GLineRequestCount += InCount;
 	}
 
+	void RequestSpaceForTriangles(uint32 InCount)
+	{
+		GTriangleRequestCount += InCount;
+	}
+
 	bool IsEnabled()
 	{
 		return CVarEnable.GetValueOnAnyThread() != 0 || GShaderPrintEnableOverride;
@@ -438,7 +455,8 @@ namespace ShaderPrint
 		BEGIN_SHADER_PARAMETER_STRUCT(FParameters, )
 			SHADER_PARAMETER_RDG_BUFFER_SRV(StructuredBuffer, ElementBuffer)
 			SHADER_PARAMETER_RDG_BUFFER_UAV(RWBuffer, RWIndirectArgs)
-			SHADER_PARAMETER(uint32, MaxInstanceCount)
+			SHADER_PARAMETER(uint32, PrimitiveType)
+			SHADER_PARAMETER_STRUCT_REF(FShaderPrintCommonParameters, ShaderPrintData)
 		END_SHADER_PARAMETER_STRUCT()
 
 		static bool ShouldCompilePermutation(const FGlobalShaderPermutationParameters& Parameters)
@@ -462,6 +480,7 @@ namespace ShaderPrint
 		SHADER_USE_PARAMETER_STRUCT(FShaderDrawDebugVS, FGlobalShader);
 
 		BEGIN_SHADER_PARAMETER_STRUCT(FParameters, )
+			SHADER_PARAMETER_STRUCT_REF(FShaderPrintCommonParameters, Common)
 			SHADER_PARAMETER(FVector3f, TranslatedWorldOffsetConversion)
 			SHADER_PARAMETER(FMatrix44f, TranslatedWorldToClip)
 			SHADER_PARAMETER_STRUCT_REF(FViewUniformShaderParameters, View)
@@ -470,6 +489,9 @@ namespace ShaderPrint
 			RDG_BUFFER_ACCESS(IndirectBuffer, ERHIAccess::IndirectArgs)
 		END_SHADER_PARAMETER_STRUCT()
 
+		class FPrimitiveType : SHADER_PERMUTATION_INT("PERMUTATION_PRIMITIVE_TYPE", 2);
+		using FPermutationDomain = TShaderPermutationDomain<FPrimitiveType>;
+
 		static bool ShouldCompilePermutation(const FGlobalShaderPermutationParameters& Parameters)
 		{
 			return IsSupported(Parameters.Platform);
@@ -479,14 +501,20 @@ namespace ShaderPrint
 		{
 			FGlobalShader::ModifyCompilationEnvironment(Parameters, OutEnvironment);
 			OutEnvironment.SetDefine(TEXT("GPU_DEBUG_RENDERING"), 1);
-			OutEnvironment.SetDefine(TEXT("GPU_DEBUG_RENDERING_VS"), 1);
-			OutEnvironment.SetDefine(TEXT("GPU_DEBUG_RENDERING_PS"), 0);
+			FPermutationDomain PermutationVector(Parameters.PermutationId);
+			if (PermutationVector.Get<FPrimitiveType>() == 0)
+			{
+				OutEnvironment.SetDefine(TEXT("GPU_DEBUG_RENDERING_LINE_VS"), 1);
+			}
+			else
+			{
+				OutEnvironment.SetDefine(TEXT("GPU_DEBUG_RENDERING_TRIANGLE_VS"), 1);
+			}
 		}
 	};
 
 	IMPLEMENT_GLOBAL_SHADER(FShaderDrawDebugVS, "/Engine/Private/ShaderPrintDrawPrimitive.usf", "ShaderDrawDebugVS", SF_Vertex);
 
-
 	class FShaderDrawDebugPS : public FGlobalShader
 	{
 		DECLARE_GLOBAL_SHADER(FShaderDrawDebugPS);
@@ -513,7 +541,6 @@ namespace ShaderPrint
 		{
 			FGlobalShader::ModifyCompilationEnvironment(Parameters, OutEnvironment);
 			OutEnvironment.SetDefine(TEXT("GPU_DEBUG_RENDERING"), 1);
-			OutEnvironment.SetDefine(TEXT("GPU_DEBUG_RENDERING_VS"), 0);
 			OutEnvironment.SetDefine(TEXT("GPU_DEBUG_RENDERING_PS"), 1);
 		}
 	};
@@ -540,6 +567,7 @@ namespace ShaderPrint
 		MaxValueCount = GetMaxValueCount();
 		MaxStateCount = GetMaxWidgetCount();
 		MaxLineCount = GetMaxLineCount();
+		MaxTriangleCount = GetMaxTriangleCount();
 	}
 
 	FShaderPrintSetup::FShaderPrintSetup(FSceneView const& View)
@@ -554,6 +582,7 @@ namespace ShaderPrint
 		MaxValueCount = GetMaxValueCount();
 		MaxStateCount = GetMaxWidgetCount();
 		MaxLineCount = GetMaxLineCount();
+		MaxTriangleCount = GetMaxTriangleCount();
 	}
 
 	FShaderPrintData CreateShaderPrintData(FRDGBuilder& GraphBuilder, FShaderPrintSetup const& InSetup, FSceneViewState* InViewState)
@@ -570,7 +599,7 @@ namespace ShaderPrint
 		{
 			ShaderPrintData.ShaderPrintValueBuffer = GraphBuilder.RegisterExternalBuffer(GEmptyBuffer->Buffer);
 			ShaderPrintData.ShaderPrintStateBuffer = GraphBuilder.RegisterExternalBuffer(GEmptyBuffer->Buffer);
-			ShaderPrintData.ShaderPrintLineBuffer = GraphBuilder.RegisterExternalBuffer(GEmptyBuffer->Buffer);
+			ShaderPrintData.ShaderPrintPrimitiveBuffer = GraphBuilder.RegisterExternalBuffer(GEmptyBuffer->Buffer);
 			return ShaderPrintData;
 		}
 
@@ -620,18 +649,19 @@ namespace ShaderPrint
 			const bool bLockBufferThisFrame = IsDrawLocked() && InViewState != nullptr && !InViewState->ShaderPrintStateData.bIsLocked;
 			ERDGBufferFlags Flags = bLockBufferThisFrame ? ERDGBufferFlags::MultiFrame : ERDGBufferFlags::None;
 
-			ShaderPrintData.ShaderPrintLineBuffer = GraphBuilder.CreateBuffer(FRDGBufferDesc::CreateStructuredDesc(4, GetLineElementCount() * InSetup.MaxLineCount), TEXT("ShaderPrint.LineBuffer"), Flags);
+			const uint32 UintElementCount = GetCountersUintSize() + GetPackedLineUintSize() * InSetup.MaxLineCount + GetPackedTriangleUintSize() * InSetup.MaxTriangleCount;
+			ShaderPrintData.ShaderPrintPrimitiveBuffer = GraphBuilder.CreateBuffer(FRDGBufferDesc::CreateStructuredDesc(4, UintElementCount), TEXT("ShaderPrint.PrimitiveBuffer"), Flags);
 
 			// Clear buffer counter
 			{
 				TShaderMapRef<FShaderDrawDebugClearCS> ComputeShader(GlobalShaderMap);
 
 				FShaderDrawDebugClearCS::FParameters* PassParameters = GraphBuilder.AllocParameters<FShaderDrawDebugClearCS::FParameters>();
-				PassParameters->RWElementBuffer = GraphBuilder.CreateUAV(ShaderPrintData.ShaderPrintLineBuffer);
+				PassParameters->RWElementBuffer = GraphBuilder.CreateUAV(ShaderPrintData.ShaderPrintPrimitiveBuffer);
 				ClearUnusedGraphResources(ComputeShader, PassParameters);
 
 				GraphBuilder.AddPass(
-					RDG_EVENT_NAME("ShaderPrint::CreateShaderPrintData (Clear lines)"),
+					RDG_EVENT_NAME("ShaderPrint::CreateShaderPrintData (Clear primitives)"),
 					PassParameters,
 					ERDGPassFlags::Compute,
 					[PassParameters, ComputeShader](FRHICommandList& RHICmdList)
@@ -644,14 +674,14 @@ namespace ShaderPrint
 			{
 				if (IsDrawLocked() && !InViewState->ShaderPrintStateData.bIsLocked)
 				{
-					InViewState->ShaderPrintStateData.LineBuffer = GraphBuilder.ConvertToExternalBuffer(ShaderPrintData.ShaderPrintLineBuffer);
+					InViewState->ShaderPrintStateData.PrimitiveBuffer = GraphBuilder.ConvertToExternalBuffer(ShaderPrintData.ShaderPrintPrimitiveBuffer);
 					InViewState->ShaderPrintStateData.PreViewTranslation = InSetup.PreViewTranslation;
 					InViewState->ShaderPrintStateData.bIsLocked = true;
 				}
 
 				if (!IsDrawLocked() && InViewState->ShaderPrintStateData.bIsLocked)
 				{
-					InViewState->ShaderPrintStateData.LineBuffer = nullptr;
+					InViewState->ShaderPrintStateData.PrimitiveBuffer = nullptr;
 					InViewState->ShaderPrintStateData.PreViewTranslation = FVector::ZeroVector;
 					InViewState->ShaderPrintStateData.bIsLocked = false;
 				}
@@ -660,7 +690,7 @@ namespace ShaderPrint
 
 		return ShaderPrintData;
 	}
-
+	\
 	FShaderPrintData CreateShaderPrintData(FRDGBuilder& GraphBuilder, FShaderPrintSetup const& InSetup)
 	{
 		return CreateShaderPrintData(GraphBuilder, InSetup, nullptr);
@@ -694,6 +724,7 @@ namespace ShaderPrint
 		GCharacterRequestCount = 0;
 		GWidgetRequestCount = 0;
 		GLineRequestCount = 0;
+		GTriangleRequestCount = 0;
 	}
 
 	static void InternalDrawView_Characters(
@@ -830,13 +861,16 @@ namespace ShaderPrint
 		}
 	}
 
-	static void InternalDrawView_Lines(
+	static void InternalDrawView_Primitives(
 		FRDGBuilder& GraphBuilder,
-		FRDGBufferRef ShaderPrintLineBuffer,
-		FIntRect ViewRect,
-		FIntRect UnscaledViewRect,
-		FMatrix const& TranslatedWorldToClip,
-		FVector TranslatedWorldOffsetConversion,
+		const FShaderPrintData& ShaderPrintData,
+		FRDGBufferRef ShaderPrintPrimitiveBuffer,
+		const FIntRect& ViewRect,
+		const FIntRect& UnscaledViewRect,
+		const FMatrix & TranslatedWorldToClip,
+		const FVector& TranslatedWorldOffsetConversion,
+		const bool bLines,
+		const bool bLocked,
 		FRDGTextureRef OutputTexture,
 		FRDGTextureRef DepthTexture)
 	{
@@ -844,17 +878,16 @@ namespace ShaderPrint
 
 		FRDGBufferRef IndirectBuffer = GraphBuilder.CreateBuffer(FRDGBufferDesc::CreateIndirectDesc<FRHIDrawIndirectParameters>(1), TEXT("ShaderDraw.IndirectBuffer"), ERDGBufferFlags::None);
 		{
-			const uint32 MaxLineCount = ShaderPrintLineBuffer->Desc.NumElements / GetLineElementCount();
-
 			FShaderDrawDebugCopyCS::FParameters* Parameters = GraphBuilder.AllocParameters<FShaderDrawDebugCopyCS::FParameters>();
-			Parameters->ElementBuffer = GraphBuilder.CreateSRV(ShaderPrintLineBuffer);
+			Parameters->ElementBuffer = GraphBuilder.CreateSRV(ShaderPrintPrimitiveBuffer);
 			Parameters->RWIndirectArgs = GraphBuilder.CreateUAV(IndirectBuffer, PF_R32_UINT);
-			Parameters->MaxInstanceCount = MaxLineCount;
+			Parameters->ShaderPrintData = ShaderPrintData.UniformBuffer;
+			Parameters->PrimitiveType = bLines ? 0u : 1u;
 
 			TShaderMapRef<FShaderDrawDebugCopyCS> ComputeShader(GlobalShaderMap);
 			ClearUnusedGraphResources(ComputeShader, Parameters);
 			GraphBuilder.AddPass(
-				RDG_EVENT_NAME("ShaderPrint::CopyLineArgs"),
+				RDG_EVENT_NAME("ShaderPrint::CopyLineArgs(%s%s)", bLines ? TEXT("Lines") : TEXT("Triangles"), bLocked ? TEXT(",Locked") : TEXT("")),
 				Parameters,
 				ERDGPassFlags::Compute,
 				[Parameters, ComputeShader](FRHICommandList& RHICmdList)
@@ -863,7 +896,9 @@ namespace ShaderPrint
 				});
 		}
 
-		TShaderMapRef<FShaderDrawDebugVS> VertexShader(GlobalShaderMap);
+		FShaderDrawDebugVS::FPermutationDomain PermutationVector;
+		PermutationVector.Set<FShaderDrawDebugVS::FPrimitiveType>(bLines ? 0u : 1u);
+		TShaderMapRef<FShaderDrawDebugVS> VertexShader(GlobalShaderMap, PermutationVector);
 		TShaderMapRef<FShaderDrawDebugPS> PixelShader(GlobalShaderMap);
 
 		FShaderDrawVSPSParameters* PassParameters = GraphBuilder.AllocParameters<FShaderDrawVSPSParameters >();
@@ -876,8 +911,9 @@ namespace ShaderPrint
 		PassParameters->PS.OriginalBufferInvSize = FVector2f(1.f / DepthTexture->Desc.Extent.X, 1.f / DepthTexture->Desc.Extent.Y);
 		PassParameters->PS.DepthTexture = DepthTexture;
 		PassParameters->PS.DepthSampler = TStaticSamplerState<SF_Point, AM_Clamp, AM_Clamp, AM_Clamp>::GetRHI();
-		PassParameters->VS.ShaderDrawDebugPrimitive = GraphBuilder.CreateSRV(ShaderPrintLineBuffer);
+		PassParameters->VS.ShaderDrawDebugPrimitive = GraphBuilder.CreateSRV(ShaderPrintPrimitiveBuffer);
 		PassParameters->VS.IndirectBuffer = IndirectBuffer;
+		PassParameters->VS.Common = ShaderPrintData.UniformBuffer;
 
 		ValidateShaderParameters(PixelShader, PassParameters->PS);
 		ClearUnusedGraphResources(PixelShader, &PassParameters->PS, { IndirectBuffer });
@@ -886,10 +922,10 @@ namespace ShaderPrint
 
 		const FIntRect Viewport = UnscaledViewRect;
 		GraphBuilder.AddPass(
-			RDG_EVENT_NAME("ShaderPrint::DrawLines"),
+			RDG_EVENT_NAME("ShaderPrint::Draw(%s%s)", bLines ? TEXT("Lines") : TEXT("Triangles"), bLocked ? TEXT(",Locked") : TEXT("")),
 			PassParameters,
 			ERDGPassFlags::Raster,
-			[VertexShader, PixelShader, PassParameters, IndirectBuffer, Viewport](FRHICommandList& RHICmdList)
+			[VertexShader, PixelShader, PassParameters, IndirectBuffer, Viewport, bLines](FRHICommandList& RHICmdList)
 			{
 				// Marks the indirect draw parameter as used by the pass, given it's not used directly by any of the shaders.
 				PassParameters->VS.IndirectBuffer->MarkResourceAsUsed();
@@ -899,7 +935,7 @@ namespace ShaderPrint
 				GraphicsPSOInit.DepthStencilState = TStaticDepthStencilState<false, CF_Always>::GetRHI();
 				GraphicsPSOInit.BlendState = TStaticBlendState<CW_RGBA, BO_Add, BF_One, BF_InverseSourceAlpha, BO_Add, BF_Zero, BF_One>::GetRHI(); // Premultiplied-alpha composition
 				GraphicsPSOInit.RasterizerState = TStaticRasterizerState<FM_Solid, CM_None, true>::GetRHI();
-				GraphicsPSOInit.PrimitiveType = PT_LineList;
+				GraphicsPSOInit.PrimitiveType = bLines ? PT_LineList : PT_TriangleList;
 				GraphicsPSOInit.BoundShaderState.VertexDeclarationRHI = GEmptyVertexDeclaration.VertexDeclarationRHI;
 				GraphicsPSOInit.BoundShaderState.VertexShaderRHI = VertexShader.GetVertexShader();
 				GraphicsPSOInit.BoundShaderState.PixelShaderRHI = PixelShader.GetPixelShader();
@@ -924,16 +960,23 @@ namespace ShaderPrint
 
 		// Lines
 		{
-			FRDGBufferRef DataBuffer = View.ShaderPrintData.ShaderPrintLineBuffer;
-			InternalDrawView_Lines(GraphBuilder, DataBuffer, View.ViewRect, View.UnscaledViewRect, View.ViewMatrices.GetTranslatedViewProjectionMatrix(), FVector::ZeroVector, OutputTexture.Texture, DepthTexture.Texture);
+			FRDGBufferRef DataBuffer = View.ShaderPrintData.ShaderPrintPrimitiveBuffer;
+			InternalDrawView_Primitives(GraphBuilder, View.ShaderPrintData, DataBuffer, View.ViewRect, View.UnscaledViewRect, View.ViewMatrices.GetTranslatedViewProjectionMatrix(), FVector::ZeroVector, true /*bLines*/, false /*bLocked*/, OutputTexture.Texture, DepthTexture.Texture);
+		}
+
+		// Triangles
+		{
+			FRDGBufferRef DataBuffer = View.ShaderPrintData.ShaderPrintPrimitiveBuffer;
+			InternalDrawView_Primitives(GraphBuilder, View.ShaderPrintData, DataBuffer, View.ViewRect, View.UnscaledViewRect, View.ViewMatrices.GetTranslatedViewProjectionMatrix(), FVector::ZeroVector, false /*bLines*/, false /*bLocked*/, OutputTexture.Texture, DepthTexture.Texture);
 		}
 
-		// Locked lines
+		// Locked Lines/Triangles
 		if (View.ViewState && View.ViewState->ShaderPrintStateData.bIsLocked)
 		{
 			const FVector LockedToCurrentTranslatedOffset = View.ViewMatrices.GetPreViewTranslation() - View.ViewState->ShaderPrintStateData.PreViewTranslation;
-			FRDGBufferRef DataBuffer = GraphBuilder.RegisterExternalBuffer(View.ViewState->ShaderPrintStateData.LineBuffer);
-			InternalDrawView_Lines(GraphBuilder, DataBuffer, View.ViewRect, View.UnscaledViewRect, View.ViewMatrices.GetTranslatedViewProjectionMatrix(), LockedToCurrentTranslatedOffset, OutputTexture.Texture, DepthTexture.Texture);
+			FRDGBufferRef DataBuffer = GraphBuilder.RegisterExternalBuffer(View.ViewState->ShaderPrintStateData.PrimitiveBuffer);
+			InternalDrawView_Primitives(GraphBuilder, View.ShaderPrintData, DataBuffer, View.ViewRect, View.UnscaledViewRect, View.ViewMatrices.GetTranslatedViewProjectionMatrix(), LockedToCurrentTranslatedOffset, true  /*bLines*/, true/*bLocked*/, OutputTexture.Texture, DepthTexture.Texture);
+			InternalDrawView_Primitives(GraphBuilder, View.ShaderPrintData, DataBuffer, View.ViewRect, View.UnscaledViewRect, View.ViewMatrices.GetTranslatedViewProjectionMatrix(), LockedToCurrentTranslatedOffset, false /*bLines*/, true/*bLocked*/, OutputTexture.Texture, DepthTexture.Texture);
 		}
 
 		// Characters
diff --git a/Engine/Source/Runtime/Renderer/Public/ShaderPrintParameters.h b/Engine/Source/Runtime/Renderer/Public/ShaderPrintParameters.h
--- a/Engine/Source/Runtime/Renderer/Public/ShaderPrintParameters.h
+++ b/Engine/Source/Runtime/Renderer/Public/ShaderPrintParameters.h
@@ -32,6 +32,7 @@ namespace ShaderPrint
 		SHADER_PARAMETER(uint32, MaxSymbolCount)
 		SHADER_PARAMETER(uint32, MaxStateCount)
 		SHADER_PARAMETER(uint32, MaxLineCount)
+		SHADER_PARAMETER(uint32, MaxTriangleCount)
 	END_GLOBAL_SHADER_PARAMETER_STRUCT()
 
 	// ShaderPrint parameter struct declaration
@@ -39,7 +40,7 @@ namespace ShaderPrint
 		SHADER_PARAMETER_STRUCT_REF(FShaderPrintCommonParameters, Common)
 		SHADER_PARAMETER_RDG_BUFFER_SRV(StructuredBuffer<uint2>, ShaderPrint_StateBuffer)
 		SHADER_PARAMETER_RDG_BUFFER_UAV(RWStructuredBuffer<ShaderPrintItem>, ShaderPrint_RWValuesBuffer)
-		SHADER_PARAMETER_RDG_BUFFER_UAV(RWStructuredBuffer<uint>, ShaderPrint_RWLinesBuffer)
+		SHADER_PARAMETER_RDG_BUFFER_UAV(RWStructuredBuffer<uint>, ShaderPrint_RWPrimitivesBuffer)
 	END_SHADER_PARAMETER_STRUCT()
 
 	// Does the platform support the ShaderPrint system?
@@ -97,6 +98,8 @@ namespace ShaderPrint
 		uint32 MaxStateCount;
 		/** Initial size of line buffer. Will also be increased by RequestSpaceForLines(). */
 		uint32 MaxLineCount;
+		/** Initial size of triangle buffer. Will also be increased by RequestSpaceForLines(). */
+		uint32 MaxTriangleCount;
 	};
 
 	/** Create the shader print render data. This allocates and clears the render buffers. */
@@ -123,5 +126,5 @@ struct RENDERER_API FShaderPrintData
 
 	FRDGBufferRef ShaderPrintValueBuffer = nullptr;
 	FRDGBufferRef ShaderPrintStateBuffer = nullptr;
-	FRDGBufferRef ShaderPrintLineBuffer = nullptr;
+	FRDGBufferRef ShaderPrintPrimitiveBuffer = nullptr;
 };
