Fix Strata data loading for simple/single material, and avoid loading & running extra logic.

diff --git a/Engine/Shaders/Private/ClusteredDeferredShadingPixelShader.usf b/Engine/Shaders/Private/ClusteredDeferredShadingPixelShader.usf
--- a/Engine/Shaders/Private/ClusteredDeferredShadingPixelShader.usf
+++ b/Engine/Shaders/Private/ClusteredDeferredShadingPixelShader.usf
@@ -15,16 +15,6 @@
 // If the textures are atlased or something like that in the future it should be ok to turn on.
 #define USE_SOURCE_TEXTURE 0
 
-#include "HairStrands/HairStrandsVisibilityCommonStruct.ush" 
-#include "Common.ush"
-#include "Strata/Strata.ush"
-#include "DeferredShadingCommon.ush"  
-#include "BRDF.ush"
-#include "ShadingModels.ush"
-#include "LightGridCommon.ush"
-#include "DeferredLightingCommon.ush"
-#include "VirtualShadowMaps/VirtualShadowMapTransmissionCommon.ush"
-
 #if STRATA_ENABLED
   #if STRATA_TILETYPE == 0
 	#define STRATA_FASTPATH 1
@@ -35,9 +25,19 @@
   #else
 	#error Strata tile type non-implemented
   #endif
+#endif
+
+#include "HairStrands/HairStrandsVisibilityCommonStruct.ush" 
+#include "Common.ush"
+#include "Strata/Strata.ush"
+#include "DeferredShadingCommon.ush"  
+#include "BRDF.ush"
+#include "ShadingModels.ush"
+#include "LightGridCommon.ush"
+#include "DeferredLightingCommon.ush"
+#include "VirtualShadowMaps/VirtualShadowMapTransmissionCommon.ush"
 #include "Strata/StrataEvaluation.ush"
 #include "Strata/StrataDeferredLighting.ush"
-#endif
 
 #if USE_HAIR_LIGHTING == 1
 	#include "HairStrands/HairStrandsCommon.ush"
diff --git a/Engine/Shaders/Private/ReflectionEnvironmentPixelShader.usf b/Engine/Shaders/Private/ReflectionEnvironmentPixelShader.usf
--- a/Engine/Shaders/Private/ReflectionEnvironmentPixelShader.usf
+++ b/Engine/Shaders/Private/ReflectionEnvironmentPixelShader.usf
@@ -4,6 +4,19 @@
 	ReflectionEnvironmentComputeShaders - functionality to apply local cubemaps.
 =============================================================================*/
   
+#if STRATA_ENABLED
+#define STRATA_INLINE_SHADING 0
+#if STRATA_TILETYPE == 0
+	#define STRATA_FASTPATH 1
+#elif STRATA_TILETYPE == 1
+	#define STRATA_SINGLEPATH 1
+#elif STRATA_TILETYPE == 2
+	// COMPLEX PATH
+#else
+	#error Strata tile type non-implemented
+#endif
+#endif
+
 #include "Common.ush"
 #include "DeferredShadingCommon.ush"  
 #include "BRDF.ush"
@@ -18,8 +31,8 @@
 #define REFLECTION_COMPOSITE_USE_BLENDED_REFLECTION_CAPTURES 1
 #define REFLECTION_COMPOSITE_SUPPORT_SKYLIGHT_BLEND 1
 #include "ReflectionEnvironmentComposite.ush"
-
-
+#include "/Engine/Private/Strata/StrataEvaluation.ush"
+#include "/Engine/Private/Strata/StrataEnvironmentLighting.ush"
 
 Texture2D ScreenSpaceReflectionsTexture;
 SamplerState ScreenSpaceReflectionsSampler;
@@ -37,22 +50,6 @@ Texture2D AmbientOcclusionTexture;
 SamplerState AmbientOcclusionSampler;
 
 
-#if STRATA_ENABLED
-#define STRATA_INLINE_SHADING 0
-#if STRATA_TILETYPE == 0
-	#define STRATA_FASTPATH 1
-#elif STRATA_TILETYPE == 1
-	#define STRATA_SINGLEPATH 1
-#elif STRATA_TILETYPE == 2
-	// COMPLEX PATH
-#else
-	#error Strata tile type non-implemented
-#endif
-#include "/Engine/Private/Strata/StrataEvaluation.ush"
-#include "/Engine/Private/Strata/StrataEnvironmentLighting.ush"
-#endif
-
-
 float3 GatherRadiance(float CompositeAlpha, float3 TranslatedWorldPosition, float3 RayDirection, float Roughness, float3 BentNormal, float IndirectIrradiance, uint ShadingModelID, uint NumCulledReflectionCaptures, uint CaptureDataStartIndex)
 {
 	// Indirect occlusion from DFAO, which should be applied to reflection captures and skylight specular, but not SSR
diff --git a/Engine/Shaders/Private/Strata/Strata.ush b/Engine/Shaders/Private/Strata/Strata.ush
--- a/Engine/Shaders/Private/Strata/Strata.ush
+++ b/Engine/Shaders/Private/Strata/Strata.ush
@@ -24,7 +24,6 @@
 #include "/Engine/Shared/StrataDefinitions.h"
 #include "/Engine/Private/Strata/StrataStatisticalOperators.ush"
 
-
 // DXC has issue with unrolling certain part of the code. To address temporarly the warning-as-error issues, 
 // we redefine the unroll function and fall back to loop on DXC
 // To remove this, one must test PC + DX12 + RayTracing and make sure raytracing shaders compile.
@@ -75,6 +74,14 @@
 #define STRATA_USES_RESOURCE_MATERIALCONTAINER (!RAYCALLABLESHADER && !RAYHITGROUPSHADER && !RAYGENSHADER && !RAYMISSSHADER)
 #endif
 
+#ifndef STRATA_FASTPATH 
+#define STRATA_FASTPATH 0
+#endif
+
+#ifndef STRATA_SINGLEPATH 
+#define STRATA_SINGLEPATH 0
+#endif
+
 // When converting legacy materials to Strata, we need to override transmittance/coverage to be independent from the material and directly linked to opacity
 #define STRATA_LEGACY_PREMULT_ALPHA_OVERRIDE (STRATA_USES_CONVERSION_FROM_LEGACY && MATERIALBLENDING_ALPHACOMPOSITE)
 
@@ -1500,8 +1507,8 @@ struct FStrataPixelHeader
 	uint State;
 };
 
-bool IsSimpleMaterial(const FStrataPixelHeader In)											{ return (In.State & HEADER_MASK_ISSIMPLEMATERIAL) != 0; }
-bool IsSingleMaterial(const FStrataPixelHeader In)											{ return (In.State & HEADER_MASK_ISSINGLEMATERIAL) != 0; }
+bool IsSimpleMaterial(const FStrataPixelHeader In)											{ return (In.State & HEADER_MASK_ISSIMPLEMATERIAL) != 0 || STRATA_FASTPATH; }
+bool IsSingleMaterial(const FStrataPixelHeader In)											{ return (In.State & HEADER_MASK_ISSINGLEMATERIAL) != 0 || STRATA_SINGLEPATH; }
 bool HasPrecShadowMask(const FStrataPixelHeader In)											{ return (In.State & HEADER_MASK_HASPRECSHADOWMASK) != 0; }
 bool HasZeroPrecShadowMask(const FStrataPixelHeader In)										{ return (In.State & HEADER_MASK_ZEROPRECSHADOWMASK) != 0; }
 bool IsSingleLayerWater(const FStrataPixelHeader In)										{ return (In.State & HEADER_MASK_ISSINGLELAYERWATER) != 0; }
@@ -3402,8 +3409,8 @@ FStrataPixelHeader UnpackStrataHeaderIn(FStrataMaterialContainer StrataBuffer, i
 	Out.State		= HEADER_GETCOMMONSTATES(PackedHeader);	// Unpack state and AO data at the same time 
 	Out.BSDFCount   = HEADER_GETBSDFCOUNT(PackedHeader);
 
-	const bool bIsSimpleMaterial = (PackedHeader & HEADER_MASK_ISSIMPLEMATERIAL) > 0;
-	const bool bIsSingleMaterial = (PackedHeader & HEADER_MASK_ISSINGLEMATERIAL) > 0;
+	const bool bIsSimpleMaterial = (PackedHeader & HEADER_MASK_ISSIMPLEMATERIAL) > 0 || STRATA_FASTPATH;
+	const bool bIsSingleMaterial = (PackedHeader & HEADER_MASK_ISSINGLEMATERIAL) > 0 || STRATA_SINGLEPATH;
 
 	// Use simple or single layout for decoding data
 	BRANCH
diff --git a/Engine/Shaders/Private/Strata/StrataEvaluation.ush b/Engine/Shaders/Private/Strata/StrataEvaluation.ush
--- a/Engine/Shaders/Private/Strata/StrataEvaluation.ush
+++ b/Engine/Shaders/Private/Strata/StrataEvaluation.ush
@@ -23,14 +23,6 @@
 #include "../HairStrands/HairStrandsDeepTransmittanceCommon.ush"
 #include "../HairStrands/HairStrandsDeepTransmittanceDualScattering.ush"
 
-#ifndef STRATA_FASTPATH 
-#define STRATA_FASTPATH 0
-#endif
-
-#ifndef STRATA_SINGLE 
-#define STRATA_SINGLE 0
-#endif
-
 #ifndef STRATA_TRANSLUCENT_ENABLED
 #define STRATA_TRANSLUCENT_ENABLED 0
 #endif
