Remove unnecessary/unused use of INVARIANT in shader code.
INVARIANT is now only used on specific vertex shader outputs.

diff --git a/Engine/Shaders/Private/Common.ush b/Engine/Shaders/Private/Common.ush
--- a/Engine/Shaders/Private/Common.ush
+++ b/Engine/Shaders/Private/Common.ush
@@ -621,23 +621,6 @@ uint3 Mod(uint3 a, uint3 b)
 #endif
 }
 
-MaterialFloat UnClampedPow(MaterialFloat X, MaterialFloat Y)
-{
-	return pow(X, INVARIANT(Y));
-}
-MaterialFloat2 UnClampedPow(MaterialFloat2 X, MaterialFloat2 Y)
-{
-	return pow(X, INVARIANT(Y));
-}
-MaterialFloat3 UnClampedPow(MaterialFloat3 X, MaterialFloat3 Y)
-{
-	return pow(X, INVARIANT(Y));
-}
-MaterialFloat4 UnClampedPow(MaterialFloat4 X, MaterialFloat4 Y)
-{
-	return pow(X, INVARIANT(Y));
-}
-
 #define POW_CLAMP 0.000001f
 
 // Clamp the base, so it's never <= 0.0f (INF/NaN).
diff --git a/Engine/Shaders/Private/LandscapeVertexFactory.ush b/Engine/Shaders/Private/LandscapeVertexFactory.ush
--- a/Engine/Shaders/Private/LandscapeVertexFactory.ush
+++ b/Engine/Shaders/Private/LandscapeVertexFactory.ush
@@ -333,13 +333,13 @@ FInstanceSceneData GetInstanceData(FVertexFactoryIntermediates Intermediates)
 
 float3 GetLocalPosition(FVertexFactoryIntermediates Intermediates)
 {
-	return INVARIANT(Intermediates.LocalPosition+float3(Intermediates.InputPosition.zw * LandscapeParameters.SubsectionOffsetParams.ww,0));
+	return Intermediates.LocalPosition+float3(Intermediates.InputPosition.zw * LandscapeParameters.SubsectionOffsetParams.ww,0);
 }
 
 float4 VertexFactoryGetWorldPosition(FVertexFactoryInput Input, FVertexFactoryIntermediates Intermediates)
 {
 	FLWCMatrix LocalToWorld = GetInstanceData(Intermediates).LocalToWorld;
-	return INVARIANT(TransformLocalToTranslatedWorld(GetLocalPosition(Intermediates), LocalToWorld));
+	return TransformLocalToTranslatedWorld(GetLocalPosition(Intermediates), LocalToWorld);
 }
 
 float3 VertexFactoryGetWorldNormal(FVertexFactoryInput Input, FVertexFactoryIntermediates Intermediates)
@@ -353,7 +353,7 @@ float4 VertexFactoryGetPreviousWorldPosition(FVertexFactoryInput Input, FVertexF
 	float3 LocalPosition = GetLocalPosition(Intermediates);
 	FLWCMatrix LocalToWorld = GetInstanceData(Intermediates).LocalToWorld;
 
-	FLWCVector3 WorldPosition = LWCInvariantMultiply(LocalPosition, LocalToWorld);
+	FLWCVector3 WorldPosition = TransformLocalToWorld(LocalPosition, LocalToWorld);
 	return float4(LWCToFloat(LWCAdd(WorldPosition, ResolvedView.PrevPreViewTranslation)), 1.0f);
 }
 
diff --git a/Engine/Shaders/Private/LargeWorldCoordinates.ush b/Engine/Shaders/Private/LargeWorldCoordinates.ush
--- a/Engine/Shaders/Private/LargeWorldCoordinates.ush
+++ b/Engine/Shaders/Private/LargeWorldCoordinates.ush
@@ -253,13 +253,6 @@ FLWCVector3 LWCMultiply(float3 Position, FLWCMatrix InMatrix)
 	return MakeLWCVector3(LWCGetTile(InMatrix), Offset);
 }
 
-FLWCVector3 LWCInvariantMultiply(float3 Position, FLWCMatrix InMatrix)
-{
-	// Explicit operations rather than mul() avoids z-fighting between depth/base passes
-	float3 Offset = INVARIANT((Position.xxx * InMatrix.M[0].xyz + Position.yyy * InMatrix.M[1].xyz + Position.zzz * InMatrix.M[2].xyz) + InMatrix.M[3].xyz);
-	return MakeLWCVector3(LWCGetTile(InMatrix), Offset);
-}
-
 FLWCVector4 LWCMultiply(float4 Position, FLWCMatrix InMatrix)
 {
 	float4 Offset = mul(Position, InMatrix.M);
diff --git a/Engine/Shaders/Private/LocalVertexFactory.ush b/Engine/Shaders/Private/LocalVertexFactory.ush
--- a/Engine/Shaders/Private/LocalVertexFactory.ush
+++ b/Engine/Shaders/Private/LocalVertexFactory.ush
@@ -776,7 +776,7 @@ float4 CalcWorldPosition(float4 Position, FLWCMatrix LocalToWorld)
 	// Transform from mesh to world space
 	return TransformLocalToTranslatedWorld(LocalPos.xyz);
 	*/
-	return INVARIANT(TransformLocalToTranslatedWorld(float3(mul(Position, CalcSliceTransform(dot(Position.xyz, SplineMeshDir))).xyz), LocalToWorld));
+	return TransformLocalToTranslatedWorld(float3(mul(Position, CalcSliceTransform(dot(Position.xyz, SplineMeshDir))).xyz), LocalToWorld);
 #else
 	return TransformLocalToTranslatedWorld(Position.xyz, LocalToWorld);
 #endif
diff --git a/Engine/Shaders/Private/VertexFactoryCommon.ush b/Engine/Shaders/Private/VertexFactoryCommon.ush
--- a/Engine/Shaders/Private/VertexFactoryCommon.ush
+++ b/Engine/Shaders/Private/VertexFactoryCommon.ush
@@ -65,8 +65,7 @@
 
 FLWCVector3 TransformLocalToWorld(float3 LocalPosition, FLWCMatrix LocalToWorld)
 {
-	//invariant code!
-	return LWCInvariantMultiply(LocalPosition, LocalToWorld);
+	return LWCMultiply(LocalPosition, LocalToWorld);
 }
 
 FLWCVector3 TransformLocalToWorld(float3 LocalPosition)
