Use 0.18 (MidGray) as a black background replacement for a mobile eye adaptation

diff --git a/Engine/Shaders/Private/PostProcessMobile.usf b/Engine/Shaders/Private/PostProcessMobile.usf
--- a/Engine/Shaders/Private/PostProcessMobile.usf
+++ b/Engine/Shaders/Private/PostProcessMobile.usf
@@ -359,8 +359,8 @@ void BloomPS_Mobile(
 	half Intensity = dot(AverageColor, half3(1.0f, 1.0f, 1.0f) / 3.0f);
 	if (Intensity == 0.0)
 	{
-		// Override intensity for black pixels to 0.1 This behaviour is controlled by the cvar "r.EyeAdaptation.BlackHistogramBucketInfluence"
-		Intensity = (0.1 - EyeAdaptation_BlackHistogramBucketInfluence);
+		// Override intensity for black pixels to MidGray. This behaviour is controlled by the cvar "r.EyeAdaptation.BlackHistogramBucketInfluence"
+		Intensity = (0.18 - EyeAdaptation_BlackHistogramBucketInfluence);
 	}
 	Intensity = max(Intensity, EyeAdaptation_LuminanceMin);
 	const float LogIntensity = clamp(log2(Intensity), -10.0f, 20.0f);
@@ -1367,7 +1367,7 @@ void Histogram_MainCS(
 				// case is so the black background pixels in the editor have no effect. But if we have cases where
 				// pixel values can actually be black, we want to set EyeAdaptation_LastHistogramBucketInfluence=1.0.
 				// This value is controlled by the cvar "r.EyeAdaptation.BlackHistogramBucketInfluence"
-				// THIS IS DONE in a Bloom setup, black pixels will set intencity to 0.1
+				// THIS IS DONE in a Bloom setup, black pixels will set intencity to 0.18
 				//if (Bucket0 == 0)
 				//{
 				//	Weight0 *= EyeAdaptation_BlackHistogramBucketInfluence;
