Fixes for Metal automation tests not running to completion and QA-Materials build lighting crashing.

- ShadingFurnaceTest MainPS() OutColor was returning float3,  need to return float4 when RGBA texture is bound on Metal otherwise we get validation errors.
- DistanceFieldShadowing; Fetch VS not PS.
- MetalRenderTargetTexture Cube resolve; cube copy to array seems to work as intended,  remove asserts to allow this area to be more fully tested.

#rnx

diff --git a/Engine/Shaders/Private/ShadingFurnaceTest.usf b/Engine/Shaders/Private/ShadingFurnaceTest.usf
--- a/Engine/Shaders/Private/ShadingFurnaceTest.usf
+++ b/Engine/Shaders/Private/ShadingFurnaceTest.usf
@@ -29,7 +29,7 @@ uint NumSamplesPerSet;
 
 void MainPS(
 	float4 SVPos : SV_POSITION,
-	out float3 OutColor : SV_Target0)
+	out float4 OutColor : SV_Target0)
 {
 	// Non strata path (todo: add strata path)
 	const float2 ScreenUV = SvPositionToBufferUV(SVPos);
@@ -245,16 +245,16 @@ void MainPS(
 
 				FinalRadiance += Weight * SampleRadiance;
 			}
-			OutColor = FinalRadiance;
+			OutColor = float4(FinalRadiance, 1.f);
 		}
 	}
 	else
 	{
-		OutColor = InRadiance;
+		OutColor = float4(InRadiance, 1.f);
 	}
 
 }
 
 #endif // SHADER_FURNACE_ANALYTIC
 
-////////////////////////////////////////////////////////////////////////////////////////////////////////////
\ No newline at end of file
+////////////////////////////////////////////////////////////////////////////////////////////////////////////
diff --git a/Engine/Source/Runtime/Apple/MetalRHI/Private/MetalRenderTarget.cpp b/Engine/Source/Runtime/Apple/MetalRHI/Private/MetalRenderTarget.cpp
--- a/Engine/Source/Runtime/Apple/MetalRHI/Private/MetalRenderTarget.cpp
+++ b/Engine/Source/Runtime/Apple/MetalRHI/Private/MetalRenderTarget.cpp
@@ -68,12 +68,6 @@ void FMetalRHICommandContext::RHICopyToResolveTarget(FRHITexture* SourceTextureR
 			}
 		}
 
-		checkf( SourceDesc.IsTexture2D()   || SourceDesc.IsTextureCube(), TEXT("Only Tex2D & Cube are tested to work so far!"));
-		checkf(!SourceDesc.IsTextureCube() || SourceDesc.ArraySize == 1,  TEXT("Cube arrays might not work yet."));
-
-		checkf( DestinationDesc.IsTexture2D()   || DestinationDesc.IsTextureCube(), TEXT("Only Tex2D & Cube are tested to work so far!"));
-		checkf(!DestinationDesc.IsTextureCube() || DestinationDesc.ArraySize == 1,  TEXT("Cube arrays might not work yet."));
-
 		mtlpp::Origin Origin(0, 0, 0);
 		mtlpp::Size Size(0, 0, 1);
 		if (ResolveParams.Rect.IsValid())
diff --git a/Engine/Source/Runtime/Renderer/Private/DistanceFieldShadowing.cpp b/Engine/Source/Runtime/Renderer/Private/DistanceFieldShadowing.cpp
--- a/Engine/Source/Runtime/Renderer/Private/DistanceFieldShadowing.cpp
+++ b/Engine/Source/Runtime/Renderer/Private/DistanceFieldShadowing.cpp
@@ -1115,7 +1115,7 @@ void FProjectedShadowInfo::RenderRayTracedDistanceFieldProjection(
 				SetGraphicsPipelineState(RHICmdList, GraphicsPSOInit, 0);
 
 				SetShaderParameters(RHICmdList, PixelShader, PixelShader.GetPixelShader(), PassParameters->PS);
-				SetShaderParameters(RHICmdList, VertexShader, VertexShader.GetPixelShader(), PassParameters->VS);
+				SetShaderParameters(RHICmdList, VertexShader, VertexShader.GetVertexShader(), PassParameters->VS);
 
 				//@todo - depth bounds test for local lights
 				if (bDirectionalLight)
