Enable Strata decal.

diff --git a/Engine/Shaders/Private/DBufferDecalShared.ush b/Engine/Shaders/Private/DBufferDecalShared.ush
--- a/Engine/Shaders/Private/DBufferDecalShared.ush
+++ b/Engine/Shaders/Private/DBufferDecalShared.ush
@@ -98,9 +98,6 @@ FStrataDBuffer GetDBufferData(float2 BufferUV, uint RTMaskBit)
 
 void ApplyDBufferData(in FStrataDBuffer StrataDBufferData, inout FStrataPixelHeader StrataPixelHeader, inout FStrataData StrataData)
 {
-  // Strata - Disabling temporary DBuffer prior to optimization
-  #if 0
-
 	// Early out if DBuffer has no coverage
 	if (StrataDBufferData.Coverage <= 0)
 		return;
@@ -245,7 +242,6 @@ void ApplyDBufferData(in FStrataDBuffer StrataDBufferData, inout FStrataPixelHea
 
 	#undef BSDF
 	}
-  #endif // Strata - Disabling temporary DBuffer prior to optimization
 }
 #else // STRATA_ENABLED
 struct FDBufferData
