Add INVARIANT to hit proxy shader to deal with cases where WPO is giving mismatch with actual base pass results.
#rnx

diff --git a/Engine/Shaders/Private/HitProxyVertexShader.usf b/Engine/Shaders/Private/HitProxyVertexShader.usf
--- a/Engine/Shaders/Private/HitProxyVertexShader.usf
+++ b/Engine/Shaders/Private/HitProxyVertexShader.usf
@@ -65,7 +65,7 @@ void Main(
 
 	{
 		float4 RasterizedWorldPosition = VertexFactoryGetRasterizedWorldPosition(Input, VFIntermediates, WorldPosition);
-		Output.Position = mul(RasterizedWorldPosition, View.TranslatedWorldToClip);
+		Output.Position = INVARIANT(mul(RasterizedWorldPosition, View.TranslatedWorldToClip));
 	}
 
 #if USE_GLOBAL_CLIP_PLANE
