Fix GPU crash in HistogramEyeAdaptation on mobile. LogLuminance can be above 1. Make sure Bucket0 value never gets above HistogramSize. Not doing so can end up reading out of boundary in groupshared SharedHistogram array

#rnx

diff --git a/Engine/Shaders/Private/PostProcessMobile.usf b/Engine/Shaders/Private/PostProcessMobile.usf
--- a/Engine/Shaders/Private/PostProcessMobile.usf
+++ b/Engine/Shaders/Private/PostProcessMobile.usf
@@ -1356,7 +1356,7 @@ void Histogram_MainCS(
 				Bucket0 = (uint)(fBucket);
 				Bucket1 = Bucket0 + 1;
 
-//				Bucket0 = min(Bucket0, HistogramSizeMinusOne);
+				Bucket0 = min(Bucket0, HistogramSizeMinusOne);
 				Bucket1 = min(Bucket1, HistogramSizeMinusOne);
 
 				// Weighted blend between the two buckets.
