Fixed ClearCoat in Lumen Reflections when using Hit Lighting

diff --git a/Engine/Shaders/Private/RayTracing/RayTracingCommon.ush b/Engine/Shaders/Private/RayTracing/RayTracingCommon.ush
--- a/Engine/Shaders/Private/RayTracing/RayTracingCommon.ush
+++ b/Engine/Shaders/Private/RayTracing/RayTracingCommon.ush
@@ -766,7 +766,8 @@ struct FPackedMaterialClosestHitPayload : FMinimalPayload
 	// Member function for (IsLumenPayload() == 1)
 	void SetSceneInstanceIndex(uint GPUSceneInstanceIndex)
 	{
-		PackedCustomData = GPUSceneInstanceIndex;
+		// Note: overwriting WorldTangentAndAnisotropy[1] so GetWorldTangent() won't work in Lumen
+		WorldTangentAndAnisotropy[1] = GPUSceneInstanceIndex;
 	}
 
 	void SetGeometryNormal(float3 GeometryNormal)
@@ -776,7 +777,7 @@ struct FPackedMaterialClosestHitPayload : FMinimalPayload
 		RadianceAndNormal[2] |= f32tof16(GeometryNormal.z) << 16;
 	}
 
-	int GetSceneInstanceIndex() { return PackedCustomData; }
+	int GetSceneInstanceIndex() { return WorldTangentAndAnisotropy[1]; }
 	float3 GetGeometryNormal() { return GetWorldNormal(); }
 };
 
