Pass SV_ViewID to VIEW_ID as uint, no more casting to float

The uint->float cast was a hack to support old devices with bad drivers. Now it's interfering with mobile preview in DX12, so time to get rid of these hacks.

diff --git a/Engine/Shaders/Private/MobileBasePassCommon.ush b/Engine/Shaders/Private/MobileBasePassCommon.ush
--- a/Engine/Shaders/Private/MobileBasePassCommon.ush
+++ b/Engine/Shaders/Private/MobileBasePassCommon.ush
@@ -40,7 +40,7 @@ struct FSharedMobileBasePassInterpolants
 #endif
 
 #if MOBILE_MULTI_VIEW
-	float MultiViewId : VIEW_ID;
+	nointerpolation uint MultiViewId : VIEW_ID;
 #endif
 
 // Keep this dummy at the bottom of the structure for the DXC compiler
diff --git a/Engine/Shaders/Private/MobileBasePassPixelShader.usf b/Engine/Shaders/Private/MobileBasePassPixelShader.usf
--- a/Engine/Shaders/Private/MobileBasePassPixelShader.usf
+++ b/Engine/Shaders/Private/MobileBasePassPixelShader.usf
@@ -219,7 +219,7 @@ void Main(
 #endif
 
 #if MOBILE_MULTI_VIEW
-	ResolvedView = ResolveView(uint(BasePassInterpolants.MultiViewId));
+	ResolvedView = ResolveView(BasePassInterpolants.MultiViewId);
 #else
 	ResolvedView = ResolveView();
 #endif
diff --git a/Engine/Shaders/Private/MobileBasePassVertexShader.usf b/Engine/Shaders/Private/MobileBasePassVertexShader.usf
--- a/Engine/Shaders/Private/MobileBasePassVertexShader.usf
+++ b/Engine/Shaders/Private/MobileBasePassVertexShader.usf
@@ -46,16 +46,10 @@ void Main(
 	const uint EyeIndex = GetEyeIndex(InstanceId);
 	ResolvedView = ResolveView(EyeIndex);
 	LayerIndex = EyeIndex;
-	Output.BasePassInterpolants.MultiViewId = float(EyeIndex);
+	Output.BasePassInterpolants.MultiViewId = EyeIndex;
 #elif MOBILE_MULTI_VIEW
-#if COMPILER_GLSL_ES3_1
-	const int MultiViewId = int(ViewId);
-	ResolvedView = ResolveView(uint(MultiViewId));
-	Output.BasePassInterpolants.MultiViewId = float(MultiViewId);
-#else
 	ResolvedView = ResolveView(ViewId);
-	Output.BasePassInterpolants.MultiViewId = float(ViewId);
-#endif
+	Output.BasePassInterpolants.MultiViewId = ViewId;
 #else
 	ResolvedView = ResolveView();
 #endif
diff --git a/Engine/Shaders/Private/PostProcessAmbientOcclusionMobile.usf b/Engine/Shaders/Private/PostProcessAmbientOcclusionMobile.usf
--- a/Engine/Shaders/Private/PostProcessAmbientOcclusionMobile.usf
+++ b/Engine/Shaders/Private/PostProcessAmbientOcclusionMobile.usf
@@ -784,7 +784,7 @@ void GTAOHorizonSearchIntegralSpatialFilterCS(
 	, out uint LayerIndex : SV_RenderTargetArrayIndex
 #elif MOBILE_MULTI_VIEW
 	, in uint ViewId : SV_ViewID
-	, out float MultiViewIndex : VIEW_ID
+	, out nointerpolation uint MultiViewIndex : VIEW_ID
 #endif
 )
 {
@@ -793,14 +793,8 @@ void GTAOHorizonSearchIntegralSpatialFilterCS(
 	ResolvedView = ResolveView(EyeIndex);
 	LayerIndex = EyeIndex;
 #elif MOBILE_MULTI_VIEW
-#if COMPILER_GLSL_ES3_1
-	const int MultiViewId = int(ViewId);
-	ResolvedView = ResolveView(uint(MultiViewId));
-	MultiViewIndex = float(MultiViewId);
-#else
 	ResolvedView = ResolveView(ViewId);
-	MultiViewIndex = float(ViewId);
-#endif
+	MultiViewIndex = ViewId;
 #else
 	ResolvedView = ResolveView();
 #endif
@@ -855,7 +849,7 @@ void GTAOHorizonSearchIntegralCS(
 	, out uint LayerIndex : SV_RenderTargetArrayIndex
 #elif MOBILE_MULTI_VIEW
 	, in uint ViewId : SV_ViewID
-	, out float MultiViewIndex : VIEW_ID
+	, out nointerpolation uint MultiViewIndex : VIEW_ID
 #endif
 )
 {
@@ -864,14 +858,8 @@ void GTAOHorizonSearchIntegralCS(
 	ResolvedView = ResolveView(EyeIndex);
 	LayerIndex = EyeIndex;
 #elif MOBILE_MULTI_VIEW
-#if COMPILER_GLSL_ES3_1
-	const int MultiViewId = int(ViewId);
-	ResolvedView = ResolveView(uint(MultiViewId));
-	MultiViewIndex = float(MultiViewId);
-#else
 	ResolvedView = ResolveView(ViewId);
-	MultiViewIndex = float(ViewId);
-#endif
+	MultiViewIndex = ViewId;
 #else
 	ResolvedView = ResolveView();
 #endif
diff --git a/Engine/Shaders/Private/PostProcessPixelProjectedReflectionMobile.usf b/Engine/Shaders/Private/PostProcessPixelProjectedReflectionMobile.usf
--- a/Engine/Shaders/Private/PostProcessPixelProjectedReflectionMobile.usf
+++ b/Engine/Shaders/Private/PostProcessPixelProjectedReflectionMobile.usf
@@ -131,7 +131,7 @@ void ProjectionPassCS(
 	, out uint LayerIndex : SV_RenderTargetArrayIndex
 #elif MOBILE_MULTI_VIEW
 	, in uint ViewId : SV_ViewID
-	, out float MultiViewIndex : VIEW_ID
+	, out nointerpolation uint MultiViewIndex : VIEW_ID
 #endif
 )
 {
@@ -140,14 +140,8 @@ void ProjectionPassCS(
 	ResolvedView = ResolveView(EyeIndex);
 	LayerIndex = EyeIndex;
 #elif MOBILE_MULTI_VIEW
-#if COMPILER_GLSL_ES3_1
-	const int MultiViewId = int(ViewId);
-	ResolvedView = ResolveView(uint(MultiViewId));
-	MultiViewIndex = float(MultiViewId);
-#else
 	ResolvedView = ResolveView(ViewId);
-	MultiViewIndex = float(ViewId);
-#endif
+	MultiViewIndex = ViewId;
 #else
 	ResolvedView = ResolveView();
 #endif
@@ -198,7 +192,7 @@ void ReflectionPassVS(
 	, out uint LayerIndex : SV_RenderTargetArrayIndex
 #elif MOBILE_MULTI_VIEW
 	, in uint ViewId : SV_ViewID
-	, out float MultiViewIndex : VIEW_ID
+	, out nointerpolation uint MultiViewIndex : VIEW_ID
 #endif
 )
 {
@@ -207,14 +201,8 @@ void ReflectionPassVS(
 	ResolvedView = ResolveView(EyeIndex);
 	LayerIndex = EyeIndex;
 #elif MOBILE_MULTI_VIEW
-#if COMPILER_GLSL_ES3_1
-	const int MultiViewId = int(ViewId);
-	ResolvedView = ResolveView(uint(MultiViewId));
-	MultiViewIndex = float(MultiViewId);
-#else
 	ResolvedView = ResolveView(ViewId);
-	MultiViewIndex = float(ViewId);
-#endif
+	MultiViewIndex = ViewId;
 #else
 	ResolvedView = ResolveView();
 #endif
@@ -298,12 +286,12 @@ void ReflectionPassPS(
 	in float4 SvPosition : SV_Position,
 	in float4 PixelPosition : TEXCOORD0,
 #if MOBILE_MULTI_VIEW
-	in float MultiViewId : VIEW_ID,
+	in nointerpolation uint MultiViewId : VIEW_ID,
 #endif
 	out HALF4_TYPE OutColor : SV_Target0)
 {
 #if MOBILE_MULTI_VIEW
-	ResolvedView = ResolveView(uint(MultiViewId));
+	ResolvedView = ResolveView(MultiViewId);
 #else
 	ResolvedView = ResolveView();
 #endif
