Fix shader warning when Strata is enabled.

diff --git a/Engine/Shaders/Private/Strata/StrataVisualize.usf b/Engine/Shaders/Private/Strata/StrataVisualize.usf
--- a/Engine/Shaders/Private/Strata/StrataVisualize.usf
+++ b/Engine/Shaders/Private/Strata/StrataVisualize.usf
@@ -538,151 +538,157 @@ void MaterialDebugStrataTreeCS(uint3 DispatchThreadId : SV_DispatchThreadID)
 
 		#define MAX_LAYER_COUNT 4
 		float LayerAccCoverage[MAX_LAYER_COUNT];
-		for (int i = 0; i < MAX_LAYER_COUNT; ++i)
 		{
-			LayerAccCoverage[i] = 0.0f;
+			for (int i = 0; i < MAX_LAYER_COUNT; ++i)
+			{
+				LayerAccCoverage[i] = 0.0f;
+			}
 		}
 
 		Context = InitShaderPrintContext(true, uint2(50, BSDF_ELEMENT_STARTY));
-	#if STRATA_COMPILER == 1
-		LOOP
-	#else
-		StrataTreeHeader.OperatorCount = min(StrataTreeHeader.OperatorCount, 2);
-		STRATA_UNROLL_N(2)
-	#endif
-		for (i = 0; i < StrataTreeHeader.OperatorCount; i++)
 		{
-			FStrataOperator Op = StrataUnpackInStrataTreeOperator(
-				i,
-				Strata.MaterialTextureArray,
-				StrataTreeHeader,
-				SliceStoringDebugStrataTree);
-
-			Print(Context, TEXT("OperatorIndex = "), FontCyan);
-			Print(Context, i, FontCyan, 2, 0);
-			Newline(Context);
-			Print(Context, TEXT("ParentIndex = "), FontWhite);
-			Print(Context, Op.ParentIndex, FontWhite, 2, 0);
-			Newline(Context);
-			Print(Context, TEXT("LayerDepth = "), FontWhite);
-			Print(Context, Op.LayerDepth, FontWhite, 2, 0);
-			Newline(Context);
-			Print(Context, TEXT("MaxDistanceFromLeaves = "), FontWhite);
-			Print(Context, Op.MaxDistanceFromLeaves, FontWhite, 2, 0);
-			Newline(Context);
-			Print(Context, TEXT("Type = "), FontWhite);
-			if (Op.Type == STRATA_OPERATOR_WEIGHT)
-			{
-				Print(Context, TEXT("WEIGHT"), FontWhite);
-			}
-			else if (Op.Type == STRATA_OPERATOR_VERTICAL)
-			{
-				Print(Context, TEXT("VERTICAL"), FontWhite);
-			}
-			else if (Op.Type == STRATA_OPERATOR_HORIZONTAL)
-			{
-				Print(Context, TEXT("HORIZONTAL"), FontWhite);
-			}
-			else if (Op.Type == STRATA_OPERATOR_ADD)
-			{
-				Print(Context, TEXT("ADD"), FontWhite);
-			}
-			else if (Op.Type == STRATA_OPERATOR_BSDF)
-			{
-				Print(Context, TEXT("BSDF"), FontWhite);
-			}
-			else if (Op.Type == STRATA_OPERATOR_BSDF_LEGACY)
-			{
-				Print(Context, TEXT("BSDF_LEGACY"), FontWhite);
-			}
-			else if (Op.Type == STRATA_OPERATOR_THINFILM)
-			{
-				Print(Context, TEXT("THINFILM"), FontWhite);
-			}
-			else
-			{
-				Print(Context, TEXT("UNKOWN"), FontWhite);
-			}
-			Newline(Context);
-			if (Op.Type == STRATA_OPERATOR_BSDF)
-			{
-				Print(Context, TEXT("BSDFIndex = "), FontYellow);
-				Print(Context, Op.LeftIndex, FontYellow, 2, 0);
-			}
-			else
+		#if STRATA_COMPILER == 1
+			LOOP
+		#else
+			StrataTreeHeader.OperatorCount = min(StrataTreeHeader.OperatorCount, 2);
+			STRATA_UNROLL_N(2)
+		#endif
+			for (int i = 0; i < StrataTreeHeader.OperatorCount; i++)
 			{
-				Print(Context, TEXT("LeftIndex = "), FontWhite);
-				Print(Context, Op.LeftIndex, FontWhite, 2, 0);
+				FStrataOperator Op = StrataUnpackInStrataTreeOperator(
+					i,
+					Strata.MaterialTextureArray,
+					StrataTreeHeader,
+					SliceStoringDebugStrataTree);
+
+				Print(Context, TEXT("OperatorIndex = "), FontCyan);
+				Print(Context, i, FontCyan, 2, 0);
+				Newline(Context);
+				Print(Context, TEXT("ParentIndex = "), FontWhite);
+				Print(Context, Op.ParentIndex, FontWhite, 2, 0);
+				Newline(Context);
+				Print(Context, TEXT("LayerDepth = "), FontWhite);
+				Print(Context, Op.LayerDepth, FontWhite, 2, 0);
+				Newline(Context);
+				Print(Context, TEXT("MaxDistanceFromLeaves = "), FontWhite);
+				Print(Context, Op.MaxDistanceFromLeaves, FontWhite, 2, 0);
+				Newline(Context);
+				Print(Context, TEXT("Type = "), FontWhite);
+				if (Op.Type == STRATA_OPERATOR_WEIGHT)
+				{
+					Print(Context, TEXT("WEIGHT"), FontWhite);
+				}
+				else if (Op.Type == STRATA_OPERATOR_VERTICAL)
+				{
+					Print(Context, TEXT("VERTICAL"), FontWhite);
+				}
+				else if (Op.Type == STRATA_OPERATOR_HORIZONTAL)
+				{
+					Print(Context, TEXT("HORIZONTAL"), FontWhite);
+				}
+				else if (Op.Type == STRATA_OPERATOR_ADD)
+				{
+					Print(Context, TEXT("ADD"), FontWhite);
+				}
+				else if (Op.Type == STRATA_OPERATOR_BSDF)
+				{
+					Print(Context, TEXT("BSDF"), FontWhite);
+				}
+				else if (Op.Type == STRATA_OPERATOR_BSDF_LEGACY)
+				{
+					Print(Context, TEXT("BSDF_LEGACY"), FontWhite);
+				}
+				else if (Op.Type == STRATA_OPERATOR_THINFILM)
+				{
+					Print(Context, TEXT("THINFILM"), FontWhite);
+				}
+				else
+				{
+					Print(Context, TEXT("UNKOWN"), FontWhite);
+				}
+				Newline(Context);
+				if (Op.Type == STRATA_OPERATOR_BSDF)
+				{
+					Print(Context, TEXT("BSDFIndex = "), FontYellow);
+					Print(Context, Op.LeftIndex, FontYellow, 2, 0);
+				}
+				else
+				{
+					Print(Context, TEXT("LeftIndex = "), FontWhite);
+					Print(Context, Op.LeftIndex, FontWhite, 2, 0);
+				}
+				Newline(Context);
+				Print(Context, TEXT("RightIndex = "), FontWhite);
+				Print(Context, Op.RightIndex, FontWhite, 2, 0);
+				Newline(Context);
+				Print(Context, TEXT("Weight = "), FontWhite);
+				Print(Context, Op.Weight, FontWhite, 5, 2);
+				Newline(Context);
+				Print(Context, TEXT("Coverage = "), FontWhite);
+				Print(Context, Op.Coverage, FontWhite, 5, 2);
+				Newline(Context);
+				Newline(Context);
 			}
-			Newline(Context);
-			Print(Context, TEXT("RightIndex = "), FontWhite);
-			Print(Context, Op.RightIndex, FontWhite, 2, 0);
-			Newline(Context);
-			Print(Context, TEXT("Weight = "), FontWhite);
-			Print(Context, Op.Weight, FontWhite, 5, 2);
-			Newline(Context);
-			Print(Context, TEXT("Coverage = "), FontWhite);
-			Print(Context, Op.Coverage, FontWhite, 5, 2);
-			Newline(Context);
-			Newline(Context);
 		}
 
-	#if STRATA_COMPILER == 1
-		LOOP
-	#else
-		StrataTreeHeader.BSDFCount = min(StrataTreeHeader.BSDFCount, 2);
-		STRATA_UNROLL_N(2)
-	#endif
-		for (i = 0; i < StrataTreeHeader.BSDFCount; i++)
 		{
-			FStrataBSDF BSDF = StrataUnpackInStrataTreeBSDF(
-				i,
-				Strata.MaterialTextureArray,
-				StrataTreeHeader,
-				SliceStoringDebugStrataTree);
+		#if STRATA_COMPILER == 1
+			LOOP
+		#else
+			StrataTreeHeader.BSDFCount = min(StrataTreeHeader.BSDFCount, 2);
+			STRATA_UNROLL_N(2)
+		#endif
+			for (int i = 0; i < StrataTreeHeader.BSDFCount; i++)
+			{
+				FStrataBSDF BSDF = StrataUnpackInStrataTreeBSDF(
+					i,
+					Strata.MaterialTextureArray,
+					StrataTreeHeader,
+					SliceStoringDebugStrataTree);
 
-			const bool bActive = any(BSDF.LuminanceWeightV > 0.0);
-			FFontColor BSDFFont = FontDarkRed;
+				const bool bActive = any(BSDF.LuminanceWeightV > 0.0);
+				FFontColor BSDFFont = FontDarkRed;
 
-			Context = InitShaderPrintContext(true, uint2(300, BSDF_ELEMENT_STARTY + BSDF_ELEMENT_SIZE * i));
-			Print(Context, TEXT("BSDFIndex = "), FontYellow);
-			Print(Context, i, FontYellow, 2, 0);
-			if (!bActive)
-			{
-				Print(Context, TEXT("DISABLED"), FontDarkRed);
-			}
-			else
-			{
-				BSDFFont = FontWhite;
+				Context = InitShaderPrintContext(true, uint2(300, BSDF_ELEMENT_STARTY + BSDF_ELEMENT_SIZE * i));
+				Print(Context, TEXT("BSDFIndex = "), FontYellow);
+				Print(Context, i, FontYellow, 2, 0);
+				if (!bActive)
+				{
+					Print(Context, TEXT("DISABLED"), FontDarkRed);
+				}
+				else
+				{
+					BSDFFont = FontWhite;
+				}
+				Newline(Context);
+				Print(Context, TEXT("OperatorIndex = "), FontCyan);
+				Print(Context, BSDF.OperatorIndex, FontCyan, 2, 0);
+				Newline(Context);
+				Print(Context, TEXT("Coverage = "), BSDFFont);
+				Print(Context, BSDF.Coverage, BSDFFont, 5, 2);
+				Newline(Context);
+				Print(Context, TEXT("LumWghtV = "), BSDFFont);
+				Print(Context, BSDF.LuminanceWeightV.r, FontLightRed, 5, 2);
+				Print(Context, BSDF.LuminanceWeightV.g, FontLightGreen, 5, 2);
+				Print(Context, BSDF.LuminanceWeightV.b, FontLightBlue, 5, 2);
+				Newline(Context);
+				Print(Context, TEXT("TopTrans = "), BSDFFont);
+				Print(Context, BSDF.TransmittanceAboveAlongN.r, FontLightRed, 5, 2);
+				Print(Context, BSDF.TransmittanceAboveAlongN.g, FontLightGreen, 5, 2);
+				Print(Context, BSDF.TransmittanceAboveAlongN.b, FontLightBlue, 5, 2);
+				Newline(Context);
+				Print(Context, TEXT("BotTrans = "), BSDFFont);
+				Print(Context, BSDF.TransmittanceBelowAlongN.r, FontLightRed, 5, 2);
+				Print(Context, BSDF.TransmittanceBelowAlongN.g, FontLightGreen, 5, 2);
+				Print(Context, BSDF.TransmittanceBelowAlongN.b, FontLightBlue, 5, 2);
+				Newline(Context);
+				Print(Context, TEXT("Diffuse  = "), BSDFFont);
+				const float3 BSDFDiffuse = StrataGetBSDFDiffuseColor(BSDF);
+				Print(Context, BSDFDiffuse.r, FontLightRed, 5, 2);
+				Print(Context, BSDFDiffuse.g, FontLightGreen, 5, 2);
+				Print(Context, BSDFDiffuse.b, FontLightBlue, 5, 2);
+				Newline(Context);
 			}
-			Newline(Context);
-			Print(Context, TEXT("OperatorIndex = "), FontCyan);
-			Print(Context, BSDF.OperatorIndex, FontCyan, 2, 0);
-			Newline(Context);
-			Print(Context, TEXT("Coverage = "), BSDFFont);
-			Print(Context, BSDF.Coverage, BSDFFont, 5, 2);
-			Newline(Context);
-			Print(Context, TEXT("LumWghtV = "), BSDFFont);
-			Print(Context, BSDF.LuminanceWeightV.r, FontLightRed, 5, 2);
-			Print(Context, BSDF.LuminanceWeightV.g, FontLightGreen, 5, 2);
-			Print(Context, BSDF.LuminanceWeightV.b, FontLightBlue, 5, 2);
-			Newline(Context);
-			Print(Context, TEXT("TopTrans = "), BSDFFont);
-			Print(Context, BSDF.TransmittanceAboveAlongN.r, FontLightRed, 5, 2);
-			Print(Context, BSDF.TransmittanceAboveAlongN.g, FontLightGreen, 5, 2);
-			Print(Context, BSDF.TransmittanceAboveAlongN.b, FontLightBlue, 5, 2);
-			Newline(Context);
-			Print(Context, TEXT("BotTrans = "), BSDFFont);
-			Print(Context, BSDF.TransmittanceBelowAlongN.r, FontLightRed, 5, 2);
-			Print(Context, BSDF.TransmittanceBelowAlongN.g, FontLightGreen, 5, 2);
-			Print(Context, BSDF.TransmittanceBelowAlongN.b, FontLightBlue, 5, 2);
-			Newline(Context);
-			Print(Context, TEXT("Diffuse  = "), BSDFFont);
-			const float3 BSDFDiffuse = StrataGetBSDFDiffuseColor(BSDF);
-			Print(Context, BSDFDiffuse.r, FontLightRed, 5, 2);
-			Print(Context, BSDFDiffuse.g, FontLightGreen, 5, 2);
-			Print(Context, BSDFDiffuse.b, FontLightBlue, 5, 2);
-			Newline(Context);
 		}
 
 		Context = InitShaderPrintContext(true, uint2(675, BSDF_ELEMENT_STARTY - 50));
@@ -1150,90 +1156,96 @@ void MaterialDebugStrataTreePS(
 		const float MaterialPanelSize = 200.0f;
 
 		// STRATA_TODO: read BSDF and Operators "on the fly".
-#if STRATA_COMPILER == 1
-		LOOP
-#else
-		StrataTreeHeader.BSDFCount = min(StrataTreeHeader.BSDFCount, 2);
-		STRATA_UNROLL_N(2)
-#endif
-		for (int i = 0; i < StrataTreeHeader.BSDFCount; i++)
 		{
-			FStrataBSDF BSDF = StrataUnpackInStrataTreeBSDF(
-				i,
-				Strata.MaterialTextureArray,
-				StrataTreeHeader,
-				SliceStoringDebugStrataTree);
-
-			//if (BSDF.Coverage > 0.0f)
+		#if STRATA_COMPILER == 1
+			LOOP
+		#else
+			StrataTreeHeader.BSDFCount = min(StrataTreeHeader.BSDFCount, 2);
+			STRATA_UNROLL_N(2)
+		#endif
+			for (int i = 0; i < StrataTreeHeader.BSDFCount; i++)
 			{
-				GetMaterialOnSphere(
-					SVPos,
-					Pos,
-					MaterialPanelSize,
-					BSDF,
-					OutColor);
-
-				GetMaterialOnCube(
-					SVPos,
-					Pos + float2(MaterialPanelSize + 10.0f, 0.0f),
-					MaterialPanelSize,
-					BSDF,
-					OutColor);
-
-				Pos.y += MaterialPanelSize + 10.0f;
+				FStrataBSDF BSDF = StrataUnpackInStrataTreeBSDF(
+					i,
+					Strata.MaterialTextureArray,
+					StrataTreeHeader,
+					SliceStoringDebugStrataTree);
+
+				//if (BSDF.Coverage > 0.0f)
+				{
+					GetMaterialOnSphere(
+						SVPos,
+						Pos,
+						MaterialPanelSize,
+						BSDF,
+						OutColor);
+
+					GetMaterialOnCube(
+						SVPos,
+						Pos + float2(MaterialPanelSize + 10.0f, 0.0f),
+						MaterialPanelSize,
+						BSDF,
+						OutColor);
+
+					Pos.y += MaterialPanelSize + 10.0f;
+				}
 			}
 		}
 
 		#define MAX_LAYER_COUNT 4
 		float LayerAccCoverage[MAX_LAYER_COUNT];
-		for (i = 0; i < MAX_LAYER_COUNT; ++i)
 		{
-			LayerAccCoverage[i] = 0.0f;
+			for (int i = 0; i < MAX_LAYER_COUNT; ++i)
+			{
+				LayerAccCoverage[i] = 0.0f;
+			}
 		}
 
 		float2 MatTopOrigin = float2(BSDF_ELEMENT_STARTX + 450, BSDF_ELEMENT_STARTY);// float2(BSDF_ELEMENT_STARTX, 550);
 		float2 MatTopLayerSize = float2(410, 100);
 		
 		// This is a really basic material topology debug. STRATA_TODO would be to draw the material graph itself with anottation
-#if STRATA_COMPILER == 1
-		LOOP
-#else
-		StrataTreeHeader.BSDFCount = min(StrataTreeHeader.BSDFCount, 2);
-		STRATA_UNROLL_N(2)
-#endif
-		for (i = 0; i < StrataTreeHeader.BSDFCount; i++)
 		{
-			FStrataBSDF BSDF = StrataUnpackInStrataTreeBSDF(
-				i,
-				Strata.MaterialTextureArray,
-				StrataTreeHeader,
-				SliceStoringDebugStrataTree);
-
-			FStrataOperator Op = StrataUnpackInStrataTreeOperator(
-				BSDF.OperatorIndex,
-				Strata.MaterialTextureArray,
-				StrataTreeHeader,
-				SliceStoringDebugStrataTree);
+		#if STRATA_COMPILER == 1
+			LOOP
+		#else
+			StrataTreeHeader.BSDFCount = min(StrataTreeHeader.BSDFCount, 2);
+			STRATA_UNROLL_N(2)
+		#endif
+			for (int i = 0; i < StrataTreeHeader.BSDFCount; i++)
+			{
+				FStrataBSDF BSDF = StrataUnpackInStrataTreeBSDF(
+					i,
+					Strata.MaterialTextureArray,
+					StrataTreeHeader,
+					SliceStoringDebugStrataTree);
 
-			const uint LayerDepth = Op.LayerDepth;
+				FStrataOperator Op = StrataUnpackInStrataTreeOperator(
+					BSDF.OperatorIndex,
+					Strata.MaterialTextureArray,
+					StrataTreeHeader,
+					SliceStoringDebugStrataTree);
 
-			if (LayerDepth < MAX_LAYER_COUNT)
-			{
-				// Draw the slab information
-				float2 LayerOrigin = MatTopOrigin + MatTopLayerSize.xy * float2(LayerAccCoverage[LayerDepth], LayerDepth);
-				float2 SlabSize = MatTopLayerSize.xy * float2(BSDF.Coverage, 1.0);
+				const uint LayerDepth = Op.LayerDepth;
 
-				if (BSDF_GETTYPE(BSDF) == STRATA_BSDF_TYPE_SLAB && all(int2(PixelPos) >= int2(LayerOrigin)) && all(int2(PixelPos) <= int2(LayerOrigin + SlabSize)))
+				if (LayerDepth < MAX_LAYER_COUNT)
 				{
-					GetMaterialSlice(
-						SVPos,
-						LayerOrigin,
-						SlabSize,
-						BSDF,
-						OutColor);
+					// Draw the slab information
+					float2 LayerOrigin = MatTopOrigin + MatTopLayerSize.xy * float2(LayerAccCoverage[LayerDepth], LayerDepth);
+					float2 SlabSize = MatTopLayerSize.xy * float2(BSDF.Coverage, 1.0);
+
+					if (BSDF_GETTYPE(BSDF) == STRATA_BSDF_TYPE_SLAB && all(int2(PixelPos) >= int2(LayerOrigin)) && all(int2(PixelPos) <= int2(LayerOrigin + SlabSize)))
+					{
+						GetMaterialSlice(
+							SVPos,
+							LayerOrigin,
+							SlabSize,
+							BSDF,
+							OutColor);
+					}
+
+					LayerAccCoverage[LayerDepth] += BSDF.Coverage;
 				}
-
-				LayerAccCoverage[LayerDepth] += BSDF.Coverage;
 			}
 		}
 	}
