Fix shader compilation with Strata.

diff --git a/Engine/Shaders/Private/VirtualShadowMaps/VirtualShadowMapProjection.usf b/Engine/Shaders/Private/VirtualShadowMaps/VirtualShadowMapProjection.usf
--- a/Engine/Shaders/Private/VirtualShadowMaps/VirtualShadowMapProjection.usf
+++ b/Engine/Shaders/Private/VirtualShadowMaps/VirtualShadowMapProjection.usf
@@ -51,7 +51,7 @@ FProjectionShadingInfo GetProjectionShadingInfo(uint2 PixelPos)
 	const FStrataPixelHeader StrataPixelHeader = UnpackStrataHeaderIn(Strata.MaterialTextureArray, StrataAddressing, Strata.TopLayerTexture);
 	const FStrataTopLayerData TopLayerData = StrataUnpackTopLayerData(Strata.TopLayerTexture.Load(uint3(PixelPos, 0)));
 
-	const FStrataSubsurfaceHeader SSSHeader = StrataLoadSubsurfaceHeader(Strata.SSSTexture, PixelPos));
+	const FStrataSubsurfaceHeader SSSHeader = StrataLoadSubsurfaceHeader(Strata.SSSTexture, PixelPos);
 	const bool bIsValid = StrataSubSurfaceHeaderGetIsValid(SSSHeader);
 	const bool bIsProfile = StrataSubSurfaceHeaderGetIsProfile(SSSHeader);
 	Out.bIsSubsurface = bIsValid && !bIsProfile;
