Replaced linear sampler for SceneDepthWithoutWater with point sampler in VolumetricRenderTarget.usf

#rnx

diff --git a/Engine/Shaders/Private/VolumetricRenderTarget.usf b/Engine/Shaders/Private/VolumetricRenderTarget.usf
--- a/Engine/Shaders/Private/VolumetricRenderTarget.usf
+++ b/Engine/Shaders/Private/VolumetricRenderTarget.usf
@@ -316,6 +316,7 @@ float4 VolumetricTextureValidUvRect;
 
 #if PERMUTATION_RENDER_UNDERWATER_BUFFER || PERMUTATION_RENDER_CAMERA_UNDERWATER
 Texture2D WaterLinearDepthTexture;
+SamplerState WaterLinearDepthSampler;
 float4 SceneWithoutSingleLayerWaterViewRect;
 float2 FullResolutionToWaterBufferScale;
 #endif
@@ -410,7 +411,7 @@ void ComposeVolumetricRTOverScenePS(
 	// Offset the uv to the view buffer region and take into account dynamic resolution scaling.
 	float2 WaterDepthScreenUV = SceneWithoutSingleLayerWaterViewRect.xy + WaterVolumeUV * (View.ViewSizeAndInvSize.xy * View.BufferSizeAndInvSize.zw);
 
-	float PixelLinearDepth = WaterLinearDepthTexture.SampleLevel(LinearTextureSampler, WaterDepthScreenUV, 0).r * UseWaterLinearDepthTexture;
+	float PixelLinearDepth = WaterLinearDepthTexture.SampleLevel(WaterLinearDepthSampler, WaterDepthScreenUV, 0).r * UseWaterLinearDepthTexture;
 
 	float3 WorldPosition = LWCHackToFloat(SvPositionToWorld(float4(SVPos.xy, 0.5, 1.0)));
 	WorldPosition = normalize(WorldPosition - LWCHackToFloat(PrimaryView.WorldCameraOrigin)) * PixelLinearDepth + LWCHackToFloat(PrimaryView.WorldCameraOrigin);
@@ -439,7 +440,7 @@ void ComposeVolumetricRTOverScenePS(
 #if PERMUTATION_RENDER_CAMERA_UNDERWATER
 	// Now check that we are compositing a pixel that is not "water" to avoid applying clouds twice (they are already composited with the behind water layer scene).
 	// We also lack depth information behind the water surface now so the composition would be wrong anyway.
-	float WaterTestPixelLinearDepth = WaterLinearDepthTexture.SampleLevel(LinearTextureSampler, ScreenUV, 0).r * UseWaterLinearDepthTexture;
+	float WaterTestPixelLinearDepth = WaterLinearDepthTexture.SampleLevel(WaterLinearDepthSampler, ScreenUV, 0).r * UseWaterLinearDepthTexture;
 	if (WaterTestPixelLinearDepth > PixelDistanceFromView)
 	{
 		// This pixel contains water, so skip it because clouds have already been composited in the "behind water scene color".
diff --git a/Engine/Source/Runtime/Renderer/Private/VolumetricRenderTarget.cpp b/Engine/Source/Runtime/Renderer/Private/VolumetricRenderTarget.cpp
--- a/Engine/Source/Runtime/Renderer/Private/VolumetricRenderTarget.cpp
+++ b/Engine/Source/Runtime/Renderer/Private/VolumetricRenderTarget.cpp
@@ -572,6 +572,7 @@ class FComposeVolumetricRTOverScenePS : public FGlobalShader
 		SHADER_PARAMETER_RDG_TEXTURE(Texture2D, WaterLinearDepthTexture)
 		SHADER_PARAMETER_RDG_TEXTURE(Texture2DMS<float>, MSAADepthTexture)
 		SHADER_PARAMETER_SAMPLER(SamplerState, LinearTextureSampler)
+		SHADER_PARAMETER_SAMPLER(SamplerState, WaterLinearDepthSampler)
 		RENDER_TARGET_BINDING_SLOTS()
 		SHADER_PARAMETER(float, UvOffsetScale)
 		SHADER_PARAMETER(float, UvOffsetSampleAcceptanceWeight)
@@ -672,6 +673,7 @@ void ComposeVolumetricRenderTargetOverScene(
 		GetTextureSafeUvCoordBound(PassParameters->VolumetricTexture, PassParameters->VolumetricTextureValidCoordRect, PassParameters->VolumetricTextureValidUvRect);
 
 		PassParameters->WaterLinearDepthTexture = WaterPassData.DepthTexture;
+		PassParameters->WaterLinearDepthSampler = TStaticSamplerState<SF_Point>::GetRHI();
 		if (bShouldRenderSingleLayerWater)
 		{
 			const FSceneWithoutWaterTextures::FView& WaterPassViewData = WaterPassData.Views[ViewIndex];
@@ -736,6 +738,7 @@ void ComposeVolumetricRenderTargetOverSceneUnderWater(
 		PassParameters->VolumetricDepthTexture = VolumetricDepthTexture;
 		PassParameters->WaterLinearDepthTexture = WaterPassData.DepthTexture;
 		PassParameters->LinearTextureSampler = TStaticSamplerState<SF_Bilinear>::GetRHI();
+		PassParameters->WaterLinearDepthSampler = TStaticSamplerState<SF_Point>::GetRHI();
 		PassParameters->UvOffsetScale = VolumetricCloudRT.GetUvNoiseScale();
 		PassParameters->UvOffsetSampleAcceptanceWeight = GetUvNoiseSampleAcceptanceWeight();
 		PassParameters->FullResolutionToVolumetricBufferResolutionScale = FVector2f(1.0f / float(GetMainDownsampleFactor(VRTMode)), float(GetMainDownsampleFactor(VRTMode)));
