Fix duplicated instance id semantic (UE-150289).

- Duplicating semantics in vertex shader signature is no longer allowed (by DXC). This change moves such parameters from VS signature into the vertex factory input structures, reusing existing parameters where applicable.

diff --git a/Engine/Plugins/Experimental/Water/Shaders/Private/WaterMeshVertexFactory.ush b/Engine/Plugins/Experimental/Water/Shaders/Private/WaterMeshVertexFactory.ush
--- a/Engine/Plugins/Experimental/Water/Shaders/Private/WaterMeshVertexFactory.ush
+++ b/Engine/Plugins/Experimental/Water/Shaders/Private/WaterMeshVertexFactory.ush
@@ -33,6 +33,7 @@ struct FVertexFactoryInput
 #endif
 
 	VF_GPUSCENE_DECLARE_INPUT_BLOCK(13)
+	VF_INSTANCED_STEREO_DECLARE_INPUT_BLOCK()
 };
 
 /** 
@@ -46,6 +47,7 @@ struct FPositionOnlyVertexFactoryInput
 	float4 InstanceData1 : ATTRIBUTE9; 
 
 	VF_GPUSCENE_DECLARE_INPUT_BLOCK(1)
+	VF_INSTANCED_STEREO_DECLARE_INPUT_BLOCK()
 };
 
 /** 
@@ -60,6 +62,7 @@ struct FPositionAndNormalOnlyVertexFactoryInput
 	float4 InstanceData1 : ATTRIBUTE9; 
 
 	VF_GPUSCENE_DECLARE_INPUT_BLOCK(1)
+	VF_INSTANCED_STEREO_DECLARE_INPUT_BLOCK()
 };
 
 /** 
diff --git a/Engine/Plugins/FX/Niagara/Shaders/Private/NiagaraMeshVertexFactory.ush b/Engine/Plugins/FX/Niagara/Shaders/Private/NiagaraMeshVertexFactory.ush
--- a/Engine/Plugins/FX/Niagara/Shaders/Private/NiagaraMeshVertexFactory.ush
+++ b/Engine/Plugins/FX/Niagara/Shaders/Private/NiagaraMeshVertexFactory.ush
@@ -76,6 +76,7 @@ struct FVertexFactoryInput
 
 #if (VF_USE_PRIMITIVE_SCENE_DATA == 1)
 	VF_GPUSCENE_DECLARE_INPUT_BLOCK(13)
+	VF_INSTANCED_STEREO_DECLARE_INPUT_BLOCK()
 #else
 	uint InstanceId	: SV_InstanceID;
 #endif
diff --git a/Engine/Plugins/FX/Niagara/Shaders/Private/NiagaraRibbonVertexFactory.ush b/Engine/Plugins/FX/Niagara/Shaders/Private/NiagaraRibbonVertexFactory.ush
--- a/Engine/Plugins/FX/Niagara/Shaders/Private/NiagaraRibbonVertexFactory.ush
+++ b/Engine/Plugins/FX/Niagara/Shaders/Private/NiagaraRibbonVertexFactory.ush
@@ -115,7 +115,10 @@ struct FVertexFactoryInput
 	// Optional instance ID for vertex layered rendering
 #if FEATURE_LEVEL >= FEATURE_LEVEL_SM4 && ONEPASS_POINTLIGHT_SHADOW && USING_VERTEX_SHADER_LAYER
 	uint InstanceId					: SV_InstanceID;
+#else
+	VF_INSTANCED_STEREO_DECLARE_INPUT_BLOCK()
 #endif
+
 };
 
 // The interpolation data required for an interp vertex.
diff --git a/Engine/Shaders/Private/AnisotropyPassShader.usf b/Engine/Shaders/Private/AnisotropyPassShader.usf
--- a/Engine/Shaders/Private/AnisotropyPassShader.usf
+++ b/Engine/Shaders/Private/AnisotropyPassShader.usf
@@ -33,7 +33,6 @@ void MainVertexShader(
 	, out float OutGlobalClipPlaneDistance : SV_ClipDistance
 #endif
 #if INSTANCED_STEREO
-	, uint InstanceId : SV_InstanceID
 	#if !MULTI_VIEW
 		, out float OutClipDistance : SV_ClipDistance1
 	#else
@@ -43,17 +42,14 @@ void MainVertexShader(
 	)
 {
 #if INSTANCED_STEREO
-	const uint EyeIndex = GetEyeIndex(InstanceId);
-	ResolvedView = ResolveView(EyeIndex);
+	const uint EyeIndex = GetEyeIndexFromVF(Input);
 	#if !MULTI_VIEW
 		OutClipDistance = 0.0;
 	#else
 		ViewportIndex = EyeIndex;
 	#endif
-#else
-	uint EyeIndex = 0;
-	ResolvedView = ResolveView();
 #endif
+	ResolvedView = ResolveViewFromVF(Input);
 	
 	FVertexFactoryIntermediates VFIntermediates = GetVertexFactoryIntermediates(Input);
 	float4 WorldPos = VertexFactoryGetWorldPosition(Input, VFIntermediates);
diff --git a/Engine/Shaders/Private/BasePassVertexShader.usf b/Engine/Shaders/Private/BasePassVertexShader.usf
--- a/Engine/Shaders/Private/BasePassVertexShader.usf
+++ b/Engine/Shaders/Private/BasePassVertexShader.usf
@@ -29,7 +29,6 @@ void Main(
 	, out float OutGlobalClipPlaneDistance : SV_ClipDistance
 #endif
 #if INSTANCED_STEREO
-	, uint InstanceId : SV_InstanceID
 	#if !MULTI_VIEW
 		, out float OutClipDistance : SV_ClipDistance1
 	#else
@@ -37,19 +36,16 @@ void Main(
 	#endif
 #endif
 	)
-{	
+{
 #if INSTANCED_STEREO
-	const uint EyeIndex = GetEyeIndex(InstanceId);
-	ResolvedView = ResolveView(EyeIndex);
+	const uint EyeIndex = GetEyeIndexFromVF(Input);
 	#if !MULTI_VIEW
 		OutClipDistance = 0.0;
 	#else
 		ViewportIndex = EyeIndex;
 	#endif
-#else
-	uint EyeIndex = 0;
-	ResolvedView = ResolveView();
 #endif
+	ResolvedView = ResolveViewFromVF(Input);
 
 	FVertexFactoryIntermediates VFIntermediates = GetVertexFactoryIntermediates(Input);
 	float4 WorldPositionExcludingWPO = VertexFactoryGetWorldPosition(Input, VFIntermediates);
diff --git a/Engine/Shaders/Private/DebugViewModeVertexShader.usf b/Engine/Shaders/Private/DebugViewModeVertexShader.usf
--- a/Engine/Shaders/Private/DebugViewModeVertexShader.usf
+++ b/Engine/Shaders/Private/DebugViewModeVertexShader.usf
@@ -48,7 +48,6 @@ void Main(
 	, out float OutGlobalClipPlaneDistance : SV_ClipDistance
 #endif
 #if INSTANCED_STEREO
-	, uint InstanceId : SV_InstanceID
 	#if !MULTI_VIEW
 		, out float OutClipDistance : SV_ClipDistance1
 	#else
@@ -58,16 +57,14 @@ void Main(
 	)
 {
 #if INSTANCED_STEREO
-	const uint EyeIndex = GetEyeIndex(InstanceId);
-	ResolvedView = ResolveView(EyeIndex);
+	const uint EyeIndex = GetEyeIndexFromVF(Input);
 	#if !MULTI_VIEW
 		OutClipDistance = 0.0;
 	#else
 		ViewportIndex = EyeIndex;
 	#endif
-#else
-	ResolvedView = ResolveView();
 #endif
+	ResolvedView = ResolveViewFromVF(Input);
 
 	FVertexFactoryIntermediates VFIntermediates = GetVertexFactoryIntermediates(Input);
 	float4 WorldPos = VertexFactoryGetWorldPosition(Input, VFIntermediates);
diff --git a/Engine/Shaders/Private/DepthOnlyVertexShader.usf b/Engine/Shaders/Private/DepthOnlyVertexShader.usf
--- a/Engine/Shaders/Private/DepthOnlyVertexShader.usf
+++ b/Engine/Shaders/Private/DepthOnlyVertexShader.usf
@@ -36,7 +36,6 @@ void Main(
 	, out float OutGlobalClipPlaneDistance : SV_ClipDistance
 #endif
 #if INSTANCED_STEREO
-	, uint InstanceId : SV_InstanceID
 	#if !MULTI_VIEW
 		, out float OutClipDistance : SV_ClipDistance1
 	#else
@@ -46,16 +45,14 @@ void Main(
 	)
 {
 #if INSTANCED_STEREO
-	uint EyeIndex = GetEyeIndex(InstanceId);
-	ResolvedView = ResolveView(EyeIndex);
+	uint EyeIndex = GetEyeIndexFromVF(Input);
 	#if !MULTI_VIEW
 		OutClipDistance = 0.0;
 	#else
 		ViewportIndex = EyeIndex;
 	#endif
-#else
-	ResolvedView = ResolveView();
 #endif
+	ResolvedView = ResolveViewFromVF(Input);
 
 	FVertexFactoryIntermediates VFIntermediates = GetVertexFactoryIntermediates(Input);
 	float4 WorldPos = VertexFactoryGetWorldPosition(Input, VFIntermediates);
diff --git a/Engine/Shaders/Private/DistortAccumulateVS.usf b/Engine/Shaders/Private/DistortAccumulateVS.usf
--- a/Engine/Shaders/Private/DistortAccumulateVS.usf
+++ b/Engine/Shaders/Private/DistortAccumulateVS.usf
@@ -37,7 +37,6 @@ void Main(
 	, out float OutGlobalClipPlaneDistance : SV_ClipDistance
 #endif
 #if INSTANCED_STEREO
-	, uint InstanceId : SV_InstanceID
 	#if !MULTI_VIEW
 		, out float OutClipDistance : SV_ClipDistance1
 	#else
@@ -47,17 +46,14 @@ void Main(
 	)
 {
 #if INSTANCED_STEREO
-	const uint EyeIndex = GetEyeIndex(InstanceId);
-	ResolvedView = ResolveView(EyeIndex);
+	const uint EyeIndex = GetEyeIndexFromVF(Input);
 	#if !MULTI_VIEW
 		OutClipDistance = 0.0;
 	#else
 		ViewportIndex = EyeIndex;
 	#endif
-#else
-	uint EyeIndex = 0;
-	ResolvedView = ResolveView();
 #endif
+	ResolvedView = ResolveViewFromVF(Input);
 
 	FVertexFactoryIntermediates VFIntermediates = GetVertexFactoryIntermediates(Input);
 	float4	WorldPosition = VertexFactoryGetWorldPosition(Input, VFIntermediates);
diff --git a/Engine/Shaders/Private/GeometryCacheVertexFactory.ush b/Engine/Shaders/Private/GeometryCacheVertexFactory.ush
--- a/Engine/Shaders/Private/GeometryCacheVertexFactory.ush
+++ b/Engine/Shaders/Private/GeometryCacheVertexFactory.ush
@@ -71,6 +71,7 @@ struct FVertexFactoryInput
 #endif
 
 	VF_GPUSCENE_DECLARE_INPUT_BLOCK(13)
+	VF_INSTANCED_STEREO_DECLARE_INPUT_BLOCK()
 };
 
 // RHI_RAYTRACING
@@ -117,6 +118,7 @@ struct FPositionOnlyVertexFactoryInput
 {
 	float4	Position	: ATTRIBUTE0;
 	VF_GPUSCENE_DECLARE_INPUT_BLOCK(1)
+	VF_INSTANCED_STEREO_DECLARE_INPUT_BLOCK()
 };
 
 struct FPositionAndNormalOnlyVertexFactoryInput
@@ -124,6 +126,7 @@ struct FPositionAndNormalOnlyVertexFactoryInput
 	float4	Position	: ATTRIBUTE0;
 	float4	Normal		: ATTRIBUTE1;
 	VF_GPUSCENE_DECLARE_INPUT_BLOCK(2)
+	VF_INSTANCED_STEREO_DECLARE_INPUT_BLOCK()
 };
 
 /** for depth-only pass */
diff --git a/Engine/Shaders/Private/GeometryCollectionVertexFactory.ush b/Engine/Shaders/Private/GeometryCollectionVertexFactory.ush
--- a/Engine/Shaders/Private/GeometryCollectionVertexFactory.ush
+++ b/Engine/Shaders/Private/GeometryCollectionVertexFactory.ush
@@ -88,6 +88,7 @@ struct FVertexFactoryInput
 #endif
 	// Dynamic instancing related attributes with InstanceIdOffset : ATTRIBUTE13
 	VF_GPUSCENE_DECLARE_INPUT_BLOCK(13)
+	VF_INSTANCED_STEREO_DECLARE_INPUT_BLOCK()
 
 #if NEEDS_LIGHTMAP_COORDINATE && !MANUAL_VERTEX_FETCH
 	float2	LightMapCoordinate : ATTRIBUTE15;
@@ -110,6 +111,7 @@ struct FPositionOnlyVertexFactoryInput
 
 	// Dynamic instancing related attributes with InstanceIdOffset : ATTRIBUTE1
 	VF_GPUSCENE_DECLARE_INPUT_BLOCK(1)
+	VF_INSTANCED_STEREO_DECLARE_INPUT_BLOCK()
 
 #if MANUAL_VERTEX_FETCH
 	uint VertexId : SV_VertexID;
@@ -126,6 +128,7 @@ struct FPositionAndNormalOnlyVertexFactoryInput
 
 	// Dynamic instancing related attributes with InstanceIdOffset : ATTRIBUTE1
 	VF_GPUSCENE_DECLARE_INPUT_BLOCK(1)
+	VF_INSTANCED_STEREO_DECLARE_INPUT_BLOCK()
 
 #if MANUAL_VERTEX_FETCH
 	uint VertexId : SV_VertexID;
diff --git a/Engine/Shaders/Private/GpuSkinVertexFactory.ush b/Engine/Shaders/Private/GpuSkinVertexFactory.ush
--- a/Engine/Shaders/Private/GpuSkinVertexFactory.ush
+++ b/Engine/Shaders/Private/GpuSkinVertexFactory.ush
@@ -118,6 +118,7 @@ struct FVertexFactoryInput
 
 	// Dynamic instancing related attributes with InstanceIdOffset : ATTRIBUTE16
 	VF_GPUSCENE_DECLARE_INPUT_BLOCK(16)
+	VF_INSTANCED_STEREO_DECLARE_INPUT_BLOCK()
 };
 
 #if !GPUSKIN_MORPH_BLEND && GPUSKIN_APEX_CLOTH
diff --git a/Engine/Shaders/Private/HairStrands/HairCardsVertexFactory.ush b/Engine/Shaders/Private/HairStrands/HairCardsVertexFactory.ush
--- a/Engine/Shaders/Private/HairStrands/HairCardsVertexFactory.ush
+++ b/Engine/Shaders/Private/HairStrands/HairCardsVertexFactory.ush
@@ -52,6 +52,7 @@ struct FVertexFactoryInput
 #endif
 
 	VF_GPUSCENE_DECLARE_INPUT_BLOCK(13)
+	VF_INSTANCED_STEREO_DECLARE_INPUT_BLOCK()
 };
 
 #if RAYHITGROUPSHADER
diff --git a/Engine/Shaders/Private/HairStrands/HairStrandsVertexFactory.ush b/Engine/Shaders/Private/HairStrands/HairStrandsVertexFactory.ush
--- a/Engine/Shaders/Private/HairStrands/HairStrandsVertexFactory.ush
+++ b/Engine/Shaders/Private/HairStrands/HairStrandsVertexFactory.ush
@@ -93,6 +93,9 @@ struct FVertexFactoryInput
 	// Dynamic instancing related attributes with InstanceIdOffset : ATTRIBUTE13
 	VF_GPUSCENE_DECLARE_INPUT_BLOCK(13)
 
+	// Instanced stereo related params
+	VF_INSTANCED_STEREO_DECLARE_INPUT_BLOCK()
+
 	uint VertexId : SV_VertexID;
 };
 
diff --git a/Engine/Shaders/Private/LandscapeVertexFactory.ush b/Engine/Shaders/Private/LandscapeVertexFactory.ush
--- a/Engine/Shaders/Private/LandscapeVertexFactory.ush
+++ b/Engine/Shaders/Private/LandscapeVertexFactory.ush
@@ -172,6 +172,7 @@ struct FVertexFactoryInput
 	// Dynamic instancing related attributes with InstanceIdOffset : ATTRIBUTE1
 	VF_GPUSCENE_DECLARE_INPUT_BLOCK(1)
 #endif
+	VF_INSTANCED_STEREO_DECLARE_INPUT_BLOCK()
 };
 
 // RHI_RAYTRACING
diff --git a/Engine/Shaders/Private/LocalVertexFactory.ush b/Engine/Shaders/Private/LocalVertexFactory.ush
--- a/Engine/Shaders/Private/LocalVertexFactory.ush
+++ b/Engine/Shaders/Private/LocalVertexFactory.ush
@@ -161,6 +161,8 @@ struct FVertexFactoryInput
 
 #if USE_INSTANCING
 	uint InstanceId	: SV_InstanceID;
+#else
+	VF_INSTANCED_STEREO_DECLARE_INPUT_BLOCK()
 #endif
 
 #if GPUSKIN_PASS_THROUGH && !MANUAL_VERTEX_FETCH
@@ -259,6 +261,7 @@ struct FPositionOnlyVertexFactoryInput
 
 	// Dynamic instancing related attributes with InstanceIdOffset : ATTRIBUTE1
 	VF_GPUSCENE_DECLARE_INPUT_BLOCK(1)
+	VF_INSTANCED_STEREO_DECLARE_INPUT_BLOCK()
 
 #if (USE_INSTANCING || USE_INSTANCE_CULLING) && !MANUAL_VERTEX_FETCH && (VF_USE_PRIMITIVE_SCENE_DATA != 2)
 	float4 InstanceOrigin : ATTRIBUTE8;  // per-instance random in w 
@@ -286,6 +289,7 @@ struct FPositionAndNormalOnlyVertexFactoryInput
 
 	// Dynamic instancing related attributes with InstanceIdOffset : ATTRIBUTE1
 	VF_GPUSCENE_DECLARE_INPUT_BLOCK(1)
+	VF_INSTANCED_STEREO_DECLARE_INPUT_BLOCK()
 
 #if (USE_INSTANCING || USE_INSTANCE_CULLING) && !MANUAL_VERTEX_FETCH && (VF_USE_PRIMITIVE_SCENE_DATA != 2)
 	float4 InstanceOrigin : ATTRIBUTE8;  // per-instance random in w 
diff --git a/Engine/Shaders/Private/MeshParticleVertexFactory.ush b/Engine/Shaders/Private/MeshParticleVertexFactory.ush
--- a/Engine/Shaders/Private/MeshParticleVertexFactory.ush
+++ b/Engine/Shaders/Private/MeshParticleVertexFactory.ush
@@ -84,6 +84,8 @@ struct FVertexFactoryInput
 /** Optional instance ID for vertex layered rendering */
 #if (PARTICLE_MESH_INSTANCED)
 	uint InstanceId	: SV_InstanceID;
+#else
+	VF_INSTANCED_STEREO_DECLARE_INPUT_BLOCK()
 #endif
 
 #if PASS_NEEDS_VERTEX_ID
diff --git a/Engine/Shaders/Private/MobileBasePassVertexShader.usf b/Engine/Shaders/Private/MobileBasePassVertexShader.usf
--- a/Engine/Shaders/Private/MobileBasePassVertexShader.usf
+++ b/Engine/Shaders/Private/MobileBasePassVertexShader.usf
@@ -35,7 +35,6 @@ void Main(
 	FVertexFactoryInput Input
 	, out FMobileShadingBasePassVSOutput Output
 #if INSTANCED_STEREO
-	, uint InstanceId : SV_InstanceID
 	, out uint LayerIndex : SV_RenderTargetArrayIndex
 #elif MOBILE_MULTI_VIEW
 	, in uint ViewId : SV_ViewID
@@ -43,7 +42,7 @@ void Main(
 	)
 {
 #if INSTANCED_STEREO
-	const uint EyeIndex = GetEyeIndex(InstanceId);
+	const uint EyeIndex = GetEyeIndexFromVF(Input);
 	ResolvedView = ResolveView(EyeIndex);
 	LayerIndex = EyeIndex;
 	Output.BasePassInterpolants.MultiViewId = EyeIndex;
diff --git a/Engine/Shaders/Private/ParticleBeamTrailVertexFactory.ush b/Engine/Shaders/Private/ParticleBeamTrailVertexFactory.ush
--- a/Engine/Shaders/Private/ParticleBeamTrailVertexFactory.ush
+++ b/Engine/Shaders/Private/ParticleBeamTrailVertexFactory.ush
@@ -16,6 +16,8 @@ struct FVertexFactoryInput
 #if USE_DYNAMIC_PARAMETERS
 	float4 DynamicParameter0		: ATTRIBUTE5;
 #endif	//USE_DYNAMIC_PARAMETERS
+
+	VF_INSTANCED_STEREO_DECLARE_INPUT_BLOCK()
 };
 
 struct FVertexFactoryInterpolantsVSToPS
diff --git a/Engine/Shaders/Private/ParticleSpriteVertexFactory.ush b/Engine/Shaders/Private/ParticleSpriteVertexFactory.ush
--- a/Engine/Shaders/Private/ParticleSpriteVertexFactory.ush
+++ b/Engine/Shaders/Private/ParticleSpriteVertexFactory.ush
@@ -28,6 +28,8 @@ struct FVertexFactoryInput
 #endif	//USE_DYNAMIC_PARAMETERS
 	float2	TexCoord				: ATTRIBUTE4;
 	uint VertexId					: SV_VertexID;
+
+	VF_INSTANCED_STEREO_DECLARE_INPUT_BLOCK()
 };
 
 struct FVertexFactoryInterpolantsVSToPS
diff --git a/Engine/Shaders/Private/PositionOnlyDepthVertexShader.usf b/Engine/Shaders/Private/PositionOnlyDepthVertexShader.usf
--- a/Engine/Shaders/Private/PositionOnlyDepthVertexShader.usf
+++ b/Engine/Shaders/Private/PositionOnlyDepthVertexShader.usf
@@ -15,7 +15,6 @@ void Main(
 	, out float OutGlobalClipPlaneDistance : SV_ClipDistance
 #endif
 #if INSTANCED_STEREO
-	, uint InstanceId : SV_InstanceID
 	#if !MULTI_VIEW
 		, out float OutClipDistance : SV_ClipDistance1
 	#else
@@ -25,17 +24,15 @@ void Main(
 	)
 {
 #if INSTANCED_STEREO
-	uint EyeIndex = GetEyeIndex(InstanceId);
-	ResolvedView = ResolveView(EyeIndex);
+	const uint EyeIndex = GetEyeIndexFromVF(Input);
 
 	#if !MULTI_VIEW
 		OutClipDistance = 0.0;
 	#else
 		ViewportIndex = EyeIndex;
 	#endif
-#else
-	ResolvedView = ResolveView();
 #endif
+	ResolvedView = ResolveViewFromVF(Input);
 
 	float4 WorldPos = VertexFactoryGetWorldPosition(Input);
 
diff --git a/Engine/Shaders/Private/VelocityShader.usf b/Engine/Shaders/Private/VelocityShader.usf
--- a/Engine/Shaders/Private/VelocityShader.usf
+++ b/Engine/Shaders/Private/VelocityShader.usf
@@ -45,7 +45,6 @@ void MainVertexShader(
 	, out float OutGlobalClipPlaneDistance : SV_ClipDistance
 #endif
 #if INSTANCED_STEREO
-	, uint InstanceId : SV_InstanceID
 	#if !MULTI_VIEW
 		, out float OutClipDistance : SV_ClipDistance1
 	#else
@@ -55,16 +54,14 @@ void MainVertexShader(
 	)
 {
 #if INSTANCED_STEREO
-	const uint EyeIndex = GetEyeIndex(InstanceId);
-	ResolvedView = ResolveView(EyeIndex);
+	const uint EyeIndex = GetEyeIndexFromVF(Input);
 	#if !MULTI_VIEW
 		OutClipDistance = 0.0;
 	#else
 		ViewportIndex = EyeIndex;
 	#endif
-#else
-	ResolvedView = ResolveView();
 #endif
+	ResolvedView = ResolveViewFromVF(Input);
 	
 	FVertexFactoryIntermediates VFIntermediates = GetVertexFactoryIntermediates(Input);
 	float4 TranslatedWorldPosition = VertexFactoryGetWorldPosition( Input, VFIntermediates );
diff --git a/Engine/Shaders/Private/VertexFactoryCommon.ush b/Engine/Shaders/Private/VertexFactoryCommon.ush
--- a/Engine/Shaders/Private/VertexFactoryCommon.ush
+++ b/Engine/Shaders/Private/VertexFactoryCommon.ush
@@ -3,6 +3,21 @@
 #include "Common.ush"
 #include "SceneData.ush"
 
+//
+// Functions to get instanced stereo properties from a VF. They have to be macros as hlslcc does not support member functions (as of May 2022 it is still used for a couple platforms),
+// whereas defining them as overloaded functions (taking VF as parameter) doesn't work as some VFs can be identical for typeless HLSL overloading rules.
+//
+#if INSTANCED_STEREO
+	#define GetInstanceIdFromVF(VFInput)				(VFInput.InstanceId)
+	#define GetEyeIndexFromVF(VFInput)					(IsInstancedStereo() ? (GetInstanceIdFromVF(VFInput) & 1) : 0)
+	#define ResolveViewFromVF(VFInput)					(ResolveView(GetEyeIndexFromVF(VFInput)))
+#else
+	#define GetInstanceIdFromVF(VFInput)				(0)
+	#define GetEyeIndexFromVF(VFInput)					(0)
+	#define ResolveViewFromVF(VFInput)					(ResolveView())
+#endif // INSTANCED_STEREO
+
+
 //  
 // Required VF inputs to support dynamic instancing 
 //
@@ -27,6 +42,10 @@
 		_VFInputOut_.InstanceIdOffset = (_PrimitiveId_) | VF_TREAT_INSTANCE_ID_OFFSET_AS_PRIMITIVE_ID_FLAG;	\
 		_VFInputOut_.DrawInstanceId = _DrawInstanceId_;
 
+	#define VF_INSTANCED_STEREO_DECLARE_INPUT_BLOCK()
+
+	#undef GetInstanceIdFromVF
+	#define GetInstanceIdFromVF(VFInput)				(VFInput.DrawInstanceId)
 
 #elif (VF_USE_PRIMITIVE_SCENE_DATA == 2)	
 	#define VF_GPUSCENE_DECLARE_INPUT_BLOCK(Attr0) \
@@ -54,12 +73,40 @@
 
 	#define VF_GPUSCENE_SET_INPUT_FOR_RT(_VFInputOut_, _PrimitiveId_, _DrawInstanceId_)
 
+	#undef GetInstanceIdFromVF
+	#define GetInstanceIdFromVF(VFInput)				(VFInput.DrawInstanceId)
+
 #else
 	#define VF_GPUSCENE_DECLARE_INPUT_BLOCK(Attr0) 
 	#define VF_GPUSCENE_GET_INTERMEDIATES(VFInput) GetSceneDataIntermediates()
 	#define VF_GPUSCENE_GET_LIGHTMAP_UV_SCALE_BIAS(VFInput, LightmapDataIndex) PrecomputedLightingBuffer.LightMapCoordinateScaleBias
 	#define VF_GPUSCENE_GET_SHADOWMAP_UV_SCALE_BIAS(VFInput, LightmapDataIndex) PrecomputedLightingBuffer.ShadowMapCoordinateScaleBias
 	#define VF_GPUSCENE_SET_INPUT_FOR_RT(_VFInputOut_, _PrimitiveId_, _DrawInstanceId_)
+
+	#if INSTANCED_STEREO
+		#define VF_INSTANCED_STEREO_DECLARE_INPUT_BLOCK() \
+			uint InstanceId : SV_InstanceID;
+	#else
+		#define VF_INSTANCED_STEREO_DECLARE_INPUT_BLOCK()
+	#endif 
+
+#endif
+
+// In case other paths in the above #if miss the declaration, catch this early
+#if !defined(VF_INSTANCED_STEREO_DECLARE_INPUT_BLOCK)
+	#error VF_INSTANCED_STEREO_DECLARE_INPUT_BLOCK() should be defined for a vertex factory
+#endif
+
+#if !defined(GetInstanceIdFromVF)
+	#error GetInstanceIdFromVF() should be defined.
+#endif
+
+#if !defined(GetEyeIndexFromVF)
+	#error GetEyeIndexFromVF() should be defined.
+#endif
+
+#if !defined(ResolveViewFromVF)
+	#error ResolveViewFromVF() should be defined.
 #endif
 
 
@@ -183,4 +230,3 @@ float3 OffsetODS(float3 WorldPosition, float3 WorldCameraPosition, float IPD)
 	return -A * D + B * T;
 }
 #endif
-
