Fix GPU crash when ray tracing debug view is enabled:
* Crash was caused by shader compiler generating invalid code for TextureVirtualSample on some platforms.
* Slightly modified the code as a workaround until the compiler issue is resolved.

diff --git a/Engine/Shaders/Private/VirtualTextureCommon.ush b/Engine/Shaders/Private/VirtualTextureCommon.ush
--- a/Engine/Shaders/Private/VirtualTextureCommon.ush
+++ b/Engine/Shaders/Private/VirtualTextureCommon.ush
@@ -701,18 +701,15 @@ MaterialFloat4 TextureVirtualSample(
 	VTPageTableResult PageTableResult, uint LayerIndex,
 	VTUniform Uniform)
 {
-	if (IsValid(PageTableResult, LayerIndex))
-	{
-		const float2 pUV = VTComputePhysicalUVs(PageTableResult, LayerIndex, Uniform);
+	const float2 pUV = VTComputePhysicalUVs(PageTableResult, LayerIndex, Uniform);
+	const bool bValid = IsValid(PageTableResult, LayerIndex);
+
 #if VIRTUAL_TEXTURE_ANISOTROPIC_FILTERING
-		return Physical.SampleGrad(PhysicalSampler, pUV, PageTableResult.dUVdx, PageTableResult.dUVdy);
+	return bValid ? Physical.SampleGrad(PhysicalSampler, pUV, PageTableResult.dUVdx, PageTableResult.dUVdy) : Uniform.FallbackValue;
 #else
-		// no need for dUVdx/dUVdy unless we have anistropic filtering enabled
-		return Physical.SampleLevel(PhysicalSampler, pUV, 0.0f);
+	// no need for dUVdx/dUVdy unless we have anistropic filtering enabled
+	return bValid ? Physical.SampleLevel(PhysicalSampler, pUV, 0.0f) : Uniform.FallbackValue;
 #endif
-	}
-
-	return Uniform.FallbackValue;
 }
 
 // SamplePhysicalTexture: Aniso Off
@@ -721,14 +718,11 @@ MaterialFloat4 TextureVirtualSampleLevel(
 	VTPageTableResult PageTableResult, uint LayerIndex,
 	VTUniform Uniform)
 {
-	if (IsValid(PageTableResult, LayerIndex))
-	{
-		const float2 pUV = VTComputePhysicalUVs(PageTableResult, LayerIndex, Uniform);
-		// No need to apply dUVdx/dUVdy, don't support anisotropic when sampling a specific level
-		return Physical.SampleLevel(PhysicalSampler, pUV, 0.0f);
-	}
+	const float2 pUV = VTComputePhysicalUVs(PageTableResult, LayerIndex, Uniform);
+	const bool bValid = IsValid(PageTableResult, LayerIndex);
+	// No need to apply dUVdx/dUVdy, don't support anisotropic when sampling a specific level
 
-	return Uniform.FallbackValue;
+	return bValid ? Physical.SampleLevel(PhysicalSampler, pUV, 0.0f) : Uniform.FallbackValue;
 }
 
 
