Rename BSDF bit for naming consistency.

diff --git a/Engine/Shaders/Private/Strata/Strata.ush b/Engine/Shaders/Private/Strata/Strata.ush
--- a/Engine/Shaders/Private/Strata/Strata.ush
+++ b/Engine/Shaders/Private/Strata/Strata.ush
@@ -265,116 +265,116 @@ struct FStrataBSDF
 // BSDF State
  
 // Bit count for each component stored in a Strata BSDF state
-#define STATE_BIT_COUNT_ANISOTROPY				1
-#define STATE_BIT_COUNT_TOPLAYER				1
-#define STATE_BIT_COUNT_HASSSS					1
-#define STATE_BIT_COUNT_HASSSSPROFILE			1
-#define STATE_BIT_COUNT_GREYWEIGHT_V			1
-#define STATE_BIT_COUNT_HAZINESS				1
-#define STATE_BIT_COUNT_F90						1
-#define STATE_BIT_COUNT_SIMPLEVOLUME			1
-#define STATE_BIT_COUNT_MFPPLUGGED				1
-#define STATE_BIT_COUNT_HASFUZZ					1
-#define STATE_BIT_COUNT_ISTHIN					1
-#define STATE_BIT_COUNT_BSDF_TYPE				3
-#define STATE_BIT_COUNT_SHAREDLOCALBASESID		2
-#define STATE_BIT_COUNT_HASTRANSABOVE			1
-#define STATE_BIT_COUNT___UNUSED___				5
-#define STATE_BIT_COUNT_WEIGHT10F				10
+#define BSDF_BIT_COUNT_ANISOTROPY				1
+#define BSDF_BIT_COUNT_TOPLAYER					1
+#define BSDF_BIT_COUNT_HASSSS					1
+#define BSDF_BIT_COUNT_HASSSSPROFILE			1
+#define BSDF_BIT_COUNT_GREYWEIGHT_V				1
+#define BSDF_BIT_COUNT_HAZINESS					1
+#define BSDF_BIT_COUNT_F90						1
+#define BSDF_BIT_COUNT_SIMPLEVOLUME				1
+#define BSDF_BIT_COUNT_MFPPLUGGED				1
+#define BSDF_BIT_COUNT_HASFUZZ					1
+#define BSDF_BIT_COUNT_ISTHIN					1
+#define BSDF_BIT_COUNT_BSDF_TYPE				3
+#define BSDF_BIT_COUNT_SHAREDLOCALBASESID		2
+#define BSDF_BIT_COUNT_HASTRANSABOVE			1
+#define BSDF_BIT_COUNT___UNUSED___				5
+#define BSDF_BIT_COUNT_WEIGHT10F				10
 
 // Bit offset of each component store in BSDF state 
-#define STATE_BIT_OFFSET_ANISOTROPY				0
-#define STATE_BIT_OFFSET_TOPLAYER				(STATE_BIT_OFFSET_ANISOTROPY		+ STATE_BIT_COUNT_ANISOTROPY)
-#define STATE_BIT_OFFSET_HASSSS					(STATE_BIT_OFFSET_TOPLAYER			+ STATE_BIT_COUNT_TOPLAYER)
-#define STATE_BIT_OFFSET_HASSSSPROFILE			(STATE_BIT_OFFSET_HASSSS			+ STATE_BIT_COUNT_HASSSS)
-#define STATE_BIT_OFFSET_GREYWEIGHT_V			(STATE_BIT_OFFSET_HASSSSPROFILE		+ STATE_BIT_COUNT_HASSSSPROFILE)
-#define STATE_BIT_OFFSET_HAZINESS				(STATE_BIT_OFFSET_GREYWEIGHT_V		+ STATE_BIT_COUNT_GREYWEIGHT_V)
-#define STATE_BIT_OFFSET_F90					(STATE_BIT_OFFSET_HAZINESS			+ STATE_BIT_COUNT_HAZINESS)
-#define STATE_BIT_OFFSET_SIMPLEVOLUME			(STATE_BIT_OFFSET_F90				+ STATE_BIT_COUNT_F90)
-#define STATE_BIT_OFFSET_MFPPLUGGED				(STATE_BIT_OFFSET_SIMPLEVOLUME		+ STATE_BIT_COUNT_SIMPLEVOLUME)
-#define STATE_BIT_OFFSET_HASFUZZ				(STATE_BIT_OFFSET_MFPPLUGGED		+ STATE_BIT_COUNT_MFPPLUGGED)
-#define STATE_BIT_OFFSET_ISTHIN					(STATE_BIT_OFFSET_HASFUZZ			+ STATE_BIT_COUNT_HASFUZZ)
-#define STATE_BIT_OFFSET_BSDF_TYPE				(STATE_BIT_OFFSET_ISTHIN			+ STATE_BIT_COUNT_ISTHIN)
-#define STATE_BIT_OFFSET_SHAREDLOCALBASESID		(STATE_BIT_OFFSET_BSDF_TYPE			+ STATE_BIT_COUNT_BSDF_TYPE)
-#define STATE_BIT_OFFSET_HASTRANSABOVE			(STATE_BIT_OFFSET_SHAREDLOCALBASESID+ STATE_BIT_COUNT_SHAREDLOCALBASESID)
-#define STATE_BIT_OFFSET_WEIGHT10F				22
+#define BSDF_BIT_OFFSET_ANISOTROPY				0
+#define BSDF_BIT_OFFSET_TOPLAYER				(BSDF_BIT_OFFSET_ANISOTROPY			+ BSDF_BIT_COUNT_ANISOTROPY)
+#define BSDF_BIT_OFFSET_HASSSS					(BSDF_BIT_OFFSET_TOPLAYER			+ BSDF_BIT_COUNT_TOPLAYER)
+#define BSDF_BIT_OFFSET_HASSSSPROFILE			(BSDF_BIT_OFFSET_HASSSS				+ BSDF_BIT_COUNT_HASSSS)
+#define BSDF_BIT_OFFSET_GREYWEIGHT_V			(BSDF_BIT_OFFSET_HASSSSPROFILE		+ BSDF_BIT_COUNT_HASSSSPROFILE)
+#define BSDF_BIT_OFFSET_HAZINESS				(BSDF_BIT_OFFSET_GREYWEIGHT_V		+ BSDF_BIT_COUNT_GREYWEIGHT_V)
+#define BSDF_BIT_OFFSET_F90						(BSDF_BIT_OFFSET_HAZINESS			+ BSDF_BIT_COUNT_HAZINESS)
+#define BSDF_BIT_OFFSET_SIMPLEVOLUME			(BSDF_BIT_OFFSET_F90				+ BSDF_BIT_COUNT_F90)
+#define BSDF_BIT_OFFSET_MFPPLUGGED				(BSDF_BIT_OFFSET_SIMPLEVOLUME		+ BSDF_BIT_COUNT_SIMPLEVOLUME)
+#define BSDF_BIT_OFFSET_HASFUZZ					(BSDF_BIT_OFFSET_MFPPLUGGED			+ BSDF_BIT_COUNT_MFPPLUGGED)
+#define BSDF_BIT_OFFSET_ISTHIN					(BSDF_BIT_OFFSET_HASFUZZ			+ BSDF_BIT_COUNT_HASFUZZ)
+#define BSDF_BIT_OFFSET_BSDF_TYPE				(BSDF_BIT_OFFSET_ISTHIN				+ BSDF_BIT_COUNT_ISTHIN)
+#define BSDF_BIT_OFFSET_SHAREDLOCALBASESID		(BSDF_BIT_OFFSET_BSDF_TYPE			+ BSDF_BIT_COUNT_BSDF_TYPE)
+#define BSDF_BIT_OFFSET_HASTRANSABOVE			(BSDF_BIT_OFFSET_SHAREDLOCALBASESID	+ BSDF_BIT_COUNT_SHAREDLOCALBASESID)
+#define BSDF_BIT_OFFSET_WEIGHT10F				22
 
 // Layout 0 (Fast)
-#define STATE_FASTENCODING_BIT_COUNT			(0u)
-#define STATE_FASTENCODING_MASK					(0u)
+#define BSDF_FASTENCODING_BIT_COUNT				(0u)
+#define BSDF_FASTENCODING_MASK					(0u)
 // Layout 1 (Single BSDF)
-#define STATE_SINGLEENCODING_BIT_COUNT			(STATE_BIT_OFFSET_ISTHIN + STATE_BIT_COUNT_ISTHIN)
-#define STATE_SINGLEENCODING_MASK				((1u<<STATE_SINGLEENCODING_BIT_COUNT)-1u)
+#define BSDF_SINGLEENCODING_BIT_COUNT			(BSDF_BIT_OFFSET_ISTHIN + BSDF_BIT_COUNT_ISTHIN)
+#define BSDF_SINGLEENCODING_MASK				((1u<<BSDF_SINGLEENCODING_BIT_COUNT)-1u)
 // Layout 2 (Complex)
-#define STATE_COMPLEXCODING_BIT_COUNT			(32u)
-#define STATE_COMPLEXFASTENCODING_MASK			(0xFFFFFFFFu)
+#define BSDF_COMPLEXCODING_BIT_COUNT			(32u)
+#define BSDF_COMPLEXFASTENCODING_MASK			(0xFFFFFFFFu)
 
 // Sanity check
-#if (STATE_BIT_OFFSET_WEIGHT10F + STATE_BIT_COUNT_WEIGHT10F) > 32
+#if (BSDF_BIT_OFFSET_WEIGHT10F + BSDF_BIT_COUNT_WEIGHT10F) > 32
 #error Strata BSDF state is larger than 32 bits
 #endif
 
 // Accessors for the BSDF type
-#define BSDF_GETTYPE(X)					READ_BITS(X.State,  STATE_BIT_COUNT_BSDF_TYPE, STATE_BIT_OFFSET_BSDF_TYPE)
-#define BSDF_SETTYPE(X, NewType)		WRITE_BITS(X.State, STATE_BIT_COUNT_BSDF_TYPE, STATE_BIT_OFFSET_BSDF_TYPE, NewType)
+#define BSDF_GETTYPE(X)					READ_BITS(X.State,  BSDF_BIT_COUNT_BSDF_TYPE, BSDF_BIT_OFFSET_BSDF_TYPE)
+#define BSDF_SETTYPE(X, NewType)		WRITE_BITS(X.State, BSDF_BIT_COUNT_BSDF_TYPE, BSDF_BIT_OFFSET_BSDF_TYPE, NewType)
 
 // Accessors for shared local bases
 // This assumes Normals[STRATA_MAX_SHAREDLOCALBASES_REGISTERS]
-#define BSDF_GETSHAREDLOCALBASISID(X)	READ_BITS(X.State,  STATE_BIT_COUNT_SHAREDLOCALBASESID, STATE_BIT_OFFSET_SHAREDLOCALBASESID)
-#define BSDF_SETSHAREDLOCALBASISID(X, NrmlIdx)	WRITE_BITS(X.State, STATE_BIT_COUNT_SHAREDLOCALBASESID, STATE_BIT_OFFSET_SHAREDLOCALBASESID, NrmlIdx)
+#define BSDF_GETSHAREDLOCALBASISID(X)	READ_BITS(X.State,  BSDF_BIT_COUNT_SHAREDLOCALBASESID, BSDF_BIT_OFFSET_SHAREDLOCALBASESID)
+#define BSDF_SETSHAREDLOCALBASISID(X, NrmlIdx)	WRITE_BITS(X.State, BSDF_BIT_COUNT_SHAREDLOCALBASESID, BSDF_BIT_OFFSET_SHAREDLOCALBASESID, NrmlIdx)
 
 // Accessors for HasAnisotropy
 // Indicates if the BSDF has anisotropic feature
-#define BSDF_GETHASANISOTROPY(X)		READ_BITS(X.State,  STATE_BIT_COUNT_ANISOTROPY, STATE_BIT_OFFSET_ANISOTROPY)
-#define BSDF_SETHASANISOTROPY(X, Aniso)	WRITE_BITS(X.State, STATE_BIT_COUNT_ANISOTROPY, STATE_BIT_OFFSET_ANISOTROPY, Aniso)
+#define BSDF_GETHASANISOTROPY(X)		READ_BITS(X.State,  BSDF_BIT_COUNT_ANISOTROPY, BSDF_BIT_OFFSET_ANISOTROPY)
+#define BSDF_SETHASANISOTROPY(X, Aniso)	WRITE_BITS(X.State, BSDF_BIT_COUNT_ANISOTROPY, BSDF_BIT_OFFSET_ANISOTROPY, Aniso)
 
 // Indicates if the BSDF is part of the top layer
-#define BSDF_GETISTOPLAYER(X)			READ_BITS(X.State,  STATE_BIT_COUNT_TOPLAYER, STATE_BIT_OFFSET_TOPLAYER)
-#define BSDF_SETISTOPLAYER(X, IsTop)	WRITE_BITS(X.State, STATE_BIT_COUNT_TOPLAYER, STATE_BIT_OFFSET_TOPLAYER, IsTop)
+#define BSDF_GETISTOPLAYER(X)			READ_BITS(X.State,  BSDF_BIT_COUNT_TOPLAYER, BSDF_BIT_OFFSET_TOPLAYER)
+#define BSDF_SETISTOPLAYER(X, IsTop)	WRITE_BITS(X.State, BSDF_BIT_COUNT_TOPLAYER, BSDF_BIT_OFFSET_TOPLAYER, IsTop)
 
 // Indicates if the BSDF has scattering component
-#define BSDF_GETHASSSS(X)				READ_BITS(X.State,  STATE_BIT_COUNT_HASSSS, STATE_BIT_OFFSET_HASSSS)
-#define BSDF_SETHASSSS(X, Scatt)		WRITE_BITS(X.State, STATE_BIT_COUNT_HASSSS, STATE_BIT_OFFSET_HASSSS, Scatt)
+#define BSDF_GETHASSSS(X)				READ_BITS(X.State,  BSDF_BIT_COUNT_HASSSS, BSDF_BIT_OFFSET_HASSSS)
+#define BSDF_SETHASSSS(X, Scatt)		WRITE_BITS(X.State, BSDF_BIT_COUNT_HASSSS, BSDF_BIT_OFFSET_HASSSS, Scatt)
 
 // Indicates if the BSDF has a SSS profile
-#define BSDF_GETHASSSSPROFILE(X)		READ_BITS(X.State,  STATE_BIT_COUNT_HASSSSPROFILE, STATE_BIT_OFFSET_HASSSSPROFILE)
-#define BSDF_SETHASSSSPROFILE(X, Profi)	WRITE_BITS(X.State, STATE_BIT_COUNT_HASSSSPROFILE, STATE_BIT_OFFSET_HASSSSPROFILE, Profi)
+#define BSDF_GETHASSSSPROFILE(X)		READ_BITS(X.State,  BSDF_BIT_COUNT_HASSSSPROFILE, BSDF_BIT_OFFSET_HASSSSPROFILE)
+#define BSDF_SETHASSSSPROFILE(X, Profi)	WRITE_BITS(X.State, BSDF_BIT_COUNT_HASSSSPROFILE, BSDF_BIT_OFFSET_HASSSSPROFILE, Profi)
 
 // Indicates if the BSDF weight is grey scale so that it can be stored in a single float (e.g. top layer BSDFs)
-#define BSDF_GETHASGREYWEIGHT_V(X)		READ_BITS(X.State,  STATE_BIT_COUNT_GREYWEIGHT_V, STATE_BIT_OFFSET_GREYWEIGHT_V)
-#define BSDF_SETHASGREYWEIGHT_V(X, Grey)WRITE_BITS(X.State, STATE_BIT_COUNT_GREYWEIGHT_V, STATE_BIT_OFFSET_GREYWEIGHT_V, Grey)
+#define BSDF_GETHASGREYWEIGHT_V(X)		READ_BITS(X.State,  BSDF_BIT_COUNT_GREYWEIGHT_V, BSDF_BIT_OFFSET_GREYWEIGHT_V)
+#define BSDF_SETHASGREYWEIGHT_V(X, Grey)WRITE_BITS(X.State, BSDF_BIT_COUNT_GREYWEIGHT_V, BSDF_BIT_OFFSET_GREYWEIGHT_V, Grey)
 
 // Indicates if the BSDF has haziness data
-#define BSDF_GETHASHAZINESS(X)			READ_BITS(X.State,  STATE_BIT_COUNT_HAZINESS, STATE_BIT_OFFSET_HAZINESS)
-#define BSDF_SETHASHAZINESS(X, Haze)	WRITE_BITS(X.State, STATE_BIT_COUNT_HAZINESS, STATE_BIT_OFFSET_HAZINESS, Haze)
+#define BSDF_GETHASHAZINESS(X)			READ_BITS(X.State,  BSDF_BIT_COUNT_HAZINESS, BSDF_BIT_OFFSET_HAZINESS)
+#define BSDF_SETHASHAZINESS(X, Haze)	WRITE_BITS(X.State, BSDF_BIT_COUNT_HAZINESS, BSDF_BIT_OFFSET_HAZINESS, Haze)
 
 // Indicates if the BSDF has edge color data
-#define BSDF_GETHASF90(X)				READ_BITS(X.State,  STATE_BIT_COUNT_F90, STATE_BIT_OFFSET_F90)
-#define BSDF_SETHASF90(X, Col)			WRITE_BITS(X.State, STATE_BIT_COUNT_F90, STATE_BIT_OFFSET_F90, Col)
+#define BSDF_GETHASF90(X)				READ_BITS(X.State,  BSDF_BIT_COUNT_F90, BSDF_BIT_OFFSET_F90)
+#define BSDF_SETHASF90(X, Col)			WRITE_BITS(X.State, BSDF_BIT_COUNT_F90, BSDF_BIT_OFFSET_F90, Col)
 
 // Accessors for the 10bits float weight used if is it detected being a grey scale
-#define BSDF_GETWEIGHT10F(X)			READ_BITS(X.State,  STATE_BIT_COUNT_WEIGHT10F, STATE_BIT_OFFSET_WEIGHT10F)
-#define BSDF_SETWEIGHT10F(X, Weight)	WRITE_BITS(X.State, STATE_BIT_COUNT_WEIGHT10F, STATE_BIT_OFFSET_WEIGHT10F, Weight)
+#define BSDF_GETWEIGHT10F(X)			READ_BITS(X.State,  BSDF_BIT_COUNT_WEIGHT10F, BSDF_BIT_OFFSET_WEIGHT10F)
+#define BSDF_SETWEIGHT10F(X, Weight)	WRITE_BITS(X.State, BSDF_BIT_COUNT_WEIGHT10F, BSDF_BIT_OFFSET_WEIGHT10F, Weight)
 
 // Indicates if the BSDF is part of the top layer
-#define BSDF_GETISSIMPLEVOLUME(X)		READ_BITS(X.State,  STATE_BIT_COUNT_SIMPLEVOLUME, STATE_BIT_OFFSET_SIMPLEVOLUME)
-#define BSDF_SETISSIMPLEVOLUME(X, Value) WRITE_BITS(X.State, STATE_BIT_COUNT_SIMPLEVOLUME, STATE_BIT_OFFSET_SIMPLEVOLUME, Value)
+#define BSDF_GETISSIMPLEVOLUME(X)		READ_BITS(X.State,  BSDF_BIT_COUNT_SIMPLEVOLUME, BSDF_BIT_OFFSET_SIMPLEVOLUME)
+#define BSDF_SETISSIMPLEVOLUME(X, Value) WRITE_BITS(X.State, BSDF_BIT_COUNT_SIMPLEVOLUME, BSDF_BIT_OFFSET_SIMPLEVOLUME, Value)
 
 // Indicates if the BSDF is part of the top layer
-#define BSDF_GETHASMFP(X)				READ_BITS(X.State,  STATE_BIT_COUNT_MFPPLUGGED, STATE_BIT_OFFSET_MFPPLUGGED)
-#define BSDF_SETHASMFP(X, Value)		WRITE_BITS(X.State, STATE_BIT_COUNT_MFPPLUGGED, STATE_BIT_OFFSET_MFPPLUGGED, Value)
+#define BSDF_GETHASMFP(X)				READ_BITS(X.State,  BSDF_BIT_COUNT_MFPPLUGGED, BSDF_BIT_OFFSET_MFPPLUGGED)
+#define BSDF_SETHASMFP(X, Value)		WRITE_BITS(X.State, BSDF_BIT_COUNT_MFPPLUGGED, BSDF_BIT_OFFSET_MFPPLUGGED, Value)
 
 // Indicates if the BSDF is part of the top layer
-#define BSDF_GETHASFUZZ(X)				READ_BITS(X.State,  STATE_BIT_COUNT_HASFUZZ, STATE_BIT_OFFSET_HASFUZZ)
-#define BSDF_SETHASFUZZ(X, Value)		WRITE_BITS(X.State, STATE_BIT_COUNT_HASFUZZ, STATE_BIT_OFFSET_HASFUZZ, Value)
+#define BSDF_GETHASFUZZ(X)				READ_BITS(X.State,  BSDF_BIT_COUNT_HASFUZZ, BSDF_BIT_OFFSET_HASFUZZ)
+#define BSDF_SETHASFUZZ(X, Value)		WRITE_BITS(X.State, BSDF_BIT_COUNT_HASFUZZ, BSDF_BIT_OFFSET_HASFUZZ, Value)
 
 // Indicates if the BSDF is part of the top layer
-#define BSDF_GETISTHIN(X)				READ_BITS(X.State,  STATE_BIT_COUNT_ISTHIN, STATE_BIT_OFFSET_ISTHIN)
-#define BSDF_SETISTHIN(X, Value)		WRITE_BITS(X.State, STATE_BIT_COUNT_ISTHIN, STATE_BIT_OFFSET_ISTHIN, Value)
+#define BSDF_GETISTHIN(X)				READ_BITS(X.State,  BSDF_BIT_COUNT_ISTHIN, BSDF_BIT_OFFSET_ISTHIN)
+#define BSDF_SETISTHIN(X, Value)		WRITE_BITS(X.State, BSDF_BIT_COUNT_ISTHIN, BSDF_BIT_OFFSET_ISTHIN, Value)
 
 // Indicates if the BSDF has light weight data
-#define BSDF_GETHASTRANSABOVE(X)		READ_BITS(X.State,  STATE_BIT_COUNT_HASTRANSABOVE, STATE_BIT_OFFSET_HASTRANSABOVE)
-#define BSDF_SETHASTRANSABOVE(X, Value)	WRITE_BITS(X.State, STATE_BIT_COUNT_HASTRANSABOVE, STATE_BIT_OFFSET_HASTRANSABOVE, Value)
+#define BSDF_GETHASTRANSABOVE(X)		READ_BITS(X.State,  BSDF_BIT_COUNT_HASTRANSABOVE, BSDF_BIT_OFFSET_HASTRANSABOVE)
+#define BSDF_SETHASTRANSABOVE(X, Value)	WRITE_BITS(X.State, BSDF_BIT_COUNT_HASTRANSABOVE, BSDF_BIT_OFFSET_HASTRANSABOVE, Value)
 
 #if STRATA_INLINE_SHADING
 #define BSDF_SETEMISSIVE(X, V)			X.Emissive = V
@@ -394,10 +394,10 @@ struct FStrataBSDF
 
 ///////////////////////////////////////////////////////////////////////////////
 // Slab BSDF members
-#define SLAB_COMPLEX_PATH_MASK (TO_OFFSET_BITMASK(STATE_BIT_COUNT_ANISOTROPY, STATE_BIT_OFFSET_ANISOTROPY)		| TO_OFFSET_BITMASK(STATE_BIT_COUNT_HASSSS, STATE_BIT_OFFSET_HASSSS) | \
-								TO_OFFSET_BITMASK(STATE_BIT_COUNT_HASSSSPROFILE, STATE_BIT_OFFSET_HASSSSPROFILE)| TO_OFFSET_BITMASK(STATE_BIT_COUNT_HAZINESS, STATE_BIT_OFFSET_HAZINESS) | \
-								TO_OFFSET_BITMASK(STATE_BIT_COUNT_F90, STATE_BIT_OFFSET_F90)                    | TO_OFFSET_BITMASK(STATE_BIT_COUNT_SIMPLEVOLUME, STATE_BIT_OFFSET_SIMPLEVOLUME) | \
-								TO_OFFSET_BITMASK(STATE_BIT_COUNT_HASFUZZ, STATE_BIT_OFFSET_HASFUZZ)            | TO_OFFSET_BITMASK(STATE_BIT_COUNT_ISTHIN, STATE_BIT_OFFSET_ISTHIN) )
+#define SLAB_COMPLEX_PATH_MASK (TO_OFFSET_BITMASK(BSDF_BIT_COUNT_ANISOTROPY, BSDF_BIT_OFFSET_ANISOTROPY)		| TO_OFFSET_BITMASK(BSDF_BIT_COUNT_HASSSS, BSDF_BIT_OFFSET_HASSSS) | \
+								TO_OFFSET_BITMASK(BSDF_BIT_COUNT_HASSSSPROFILE, BSDF_BIT_OFFSET_HASSSSPROFILE)	| TO_OFFSET_BITMASK(BSDF_BIT_COUNT_HAZINESS, BSDF_BIT_OFFSET_HAZINESS) | \
+								TO_OFFSET_BITMASK(BSDF_BIT_COUNT_F90, BSDF_BIT_OFFSET_F90)						| TO_OFFSET_BITMASK(BSDF_BIT_COUNT_SIMPLEVOLUME, BSDF_BIT_OFFSET_SIMPLEVOLUME) | \
+								TO_OFFSET_BITMASK(BSDF_BIT_COUNT_HASFUZZ, BSDF_BIT_OFFSET_HASFUZZ)				| TO_OFFSET_BITMASK(BSDF_BIT_COUNT_ISTHIN, BSDF_BIT_OFFSET_ISTHIN) )
 
 #define SLAB_DIFFUSEALBEDO(X)			X.VGPRs[0].xyz
 #define SLAB_ROUGHNESS(X)				X.VGPRs[0].w
@@ -3328,7 +3328,7 @@ void PackStrataOut(
 
 							HEADER_SETSHADINGMODELS(Out, GetShadingModels(StrataPixelHeader));
 							Out = Out & HEADER_SINGLEENCODING_MASK;
-							Out = Out | ((BSDF.State & STATE_SINGLEENCODING_MASK) << HEADER_SINGLEENCODING_BIT_COUNT);
+							Out = Out | ((BSDF.State & BSDF_SINGLEENCODING_MASK) << HEADER_SINGLEENCODING_BIT_COUNT);
 							STRATA_STORE_UINT1(Out);
 
 						#if (HEADER_SIMPLEENCODING_BIT_COUNT) > 32
@@ -3623,7 +3623,7 @@ FStrataBSDF UnpackStrataBSDFIn(FStrataMaterialContainer StrataBuffer, inout FStr
 	#if STRATA_DEFERRED_SHADING
 		Data = StrataHeader.SharedLocalBasesTypes_PackedHeader;
 	#endif
-		OutBSDF.State = (Data>>HEADER_SINGLEENCODING_BIT_COUNT) & STATE_SINGLEENCODING_MASK;
+		OutBSDF.State = (Data>>HEADER_SINGLEENCODING_BIT_COUNT) & BSDF_SINGLEENCODING_MASK;
 		OutBSDF.LuminanceWeightV = 1.0f;
 		const uint BSDFType = firstbitlow(GetShadingModels(StrataHeader));
 		BSDF_SETTYPE(OutBSDF, BSDFType);
