Add sanity check for strata single layout header.

diff --git a/Engine/Shaders/Private/Strata/Strata.ush b/Engine/Shaders/Private/Strata/Strata.ush
--- a/Engine/Shaders/Private/Strata/Strata.ush
+++ b/Engine/Shaders/Private/Strata/Strata.ush
@@ -313,6 +313,9 @@ struct FStrataBSDF
 #if (BSDF_BIT_OFFSET_WEIGHT10F + BSDF_BIT_COUNT_WEIGHT10F) > 32
 #error Strata BSDF state is larger than 32 bits
 #endif
+#if (HEADER_SINGLEENCODING_BIT_COUNT + BSDF_SINGLEENCODING_BIT_COUNT) > 32u
+#error Strata Single layout header is larger than 32 bits
+#endif
 
 // Accessors for the BSDF type
 #define BSDF_GETTYPE(X)					READ_BITS(X.State,  BSDF_BIT_COUNT_BSDF_TYPE, BSDF_BIT_OFFSET_BSDF_TYPE)
