Fix shader compile error when packaging Hololens (UE-152869).

diff --git a/Engine/Shaders/Private/VertexFactoryCommon.ush b/Engine/Shaders/Private/VertexFactoryCommon.ush
--- a/Engine/Shaders/Private/VertexFactoryCommon.ush
+++ b/Engine/Shaders/Private/VertexFactoryCommon.ush
@@ -73,6 +73,8 @@
 
 	#define VF_GPUSCENE_SET_INPUT_FOR_RT(_VFInputOut_, _PrimitiveId_, _DrawInstanceId_)
 
+	#define VF_INSTANCED_STEREO_DECLARE_INPUT_BLOCK()
+
 	#undef GetInstanceIdFromVF
 	#define GetInstanceIdFromVF(VFInput)				(VFInput.DrawInstanceId)
 
