* Add SSS phase function support to Strata (only used for 'single scattering as transmission' on SSS material).
* Change SSS type to allows more enums. This will be used by SSS wrapping mode later on.
* Add bool for enabling/disable light diffusion on Strata material (on by default).
* Fix shader compilation issue on preview material due to decal code

diff --git a/Engine/Shaders/Private/BasePassPixelShader.usf b/Engine/Shaders/Private/BasePassPixelShader.usf
--- a/Engine/Shaders/Private/BasePassPixelShader.usf
+++ b/Engine/Shaders/Private/BasePassPixelShader.usf
@@ -1762,7 +1762,7 @@ void FPixelShaderInOut_MainPS(
 			// We only enable opaque rough refraction if the material is coated or has SSS.
 			// BSDFs and Operators indices must match StrataLegacyConversion.ush.
 			OpaqueRoughRefractionData.OpaqueRoughRefractionEnabled = (StrataPixelHeader.StrataTree.Operators[2].Weight > 0.0f && StrataPixelHeader.StrataTree.Operators[3].Weight > 0.0f) 
-				|| BSDF_GETHASSSS(StrataPixelHeader.StrataTree.BSDFs[0]) || BSDF_GETHASSSS(StrataPixelHeader.StrataTree.BSDFs[1]);
+				|| BSDF_GETSSSTYPE(StrataPixelHeader.StrataTree.BSDFs[0]) != SSS_TYPE_INVALID || BSDF_GETSSSTYPE(StrataPixelHeader.StrataTree.BSDFs[1]) != SSS_TYPE_INVALID;
 		}
 	#endif
 		StrataStruct.OpaqueRoughRefractionTextureUAV[PixelPos] = StrataPackOpaqueRoughRefractionData(OpaqueRoughRefractionData);
diff --git a/Engine/Shaders/Private/DBufferDecalShared.ush b/Engine/Shaders/Private/DBufferDecalShared.ush
--- a/Engine/Shaders/Private/DBufferDecalShared.ush
+++ b/Engine/Shaders/Private/DBufferDecalShared.ush
@@ -96,6 +96,7 @@ FStrataDBuffer GetDBufferData(float2 BufferUV, uint RTMaskBit)
 #define STRATA_DBUFFER_RESPONSE_NORMAL    (MATERIALDECALRESPONSEMASK & 0x2)
 #define STRATA_DBUFFER_RESPONSE_ROUGHNESS (MATERIALDECALRESPONSEMASK & 0x4)
 
+#if MATERIAL_IS_STRATA
 void ApplyDBufferData(in FStrataDBuffer StrataDBufferData, inout FStrataPixelHeader StrataPixelHeader, inout FStrataData StrataData)
 {
 	// Early out if DBuffer has no coverage
@@ -161,8 +162,7 @@ void ApplyDBufferData(in FStrataDBuffer StrataDBufferData, inout FStrataPixelHea
 					SLAB_SSSMFP(BSDF) = lerp(0.f, SLAB_SSSMFP(BSDF), StrataDBufferData.OneMinusCoverage_BaseColor);
 					if (all(SLAB_SSSMFP(BSDF) == 0.0f))
 					{
-						BSDF_SETHASMFP(BSDF, 0);
-						BSDF_SETHASSSS(BSDF, 0);
+						BSDF_SETSSSTYPE(BSDF, SSS_TYPE_INVALID);
 					}
 				#endif
 
@@ -243,6 +243,7 @@ void ApplyDBufferData(in FStrataDBuffer StrataDBufferData, inout FStrataPixelHea
 	#undef BSDF
 	}
 }
+#endif // MATERIAL_IS_STRATA
 #else // STRATA_ENABLED
 struct FDBufferData
 {
diff --git a/Engine/Shaders/Private/DebugProbes.usf b/Engine/Shaders/Private/DebugProbes.usf
--- a/Engine/Shaders/Private/DebugProbes.usf
+++ b/Engine/Shaders/Private/DebugProbes.usf
@@ -297,6 +297,8 @@ void MainPS(
 			/*float SSSProfileID					*/0.0f,
 			/*float3 SSSMFP							*/0.0f,
 			/*float SSSMFPScale						*/0.0f,
+			/*float SSSPhaseAniso					*/0.0f,
+			/*float UseSSSDiffuse					*/0.0f,
 			/*float3 Emissive						*/0.0f,
 			/*float Haziness						*/0.0f,
 			/*float FuzzAmount						*/0.0f,
diff --git a/Engine/Shaders/Private/Strata/Strata.ush b/Engine/Shaders/Private/Strata/Strata.ush
--- a/Engine/Shaders/Private/Strata/Strata.ush
+++ b/Engine/Shaders/Private/Strata/Strata.ush
@@ -267,8 +267,7 @@ struct FStrataBSDF
 // Bit count for each component stored in a Strata BSDF state
 #define BSDF_BIT_COUNT_ANISOTROPY				1
 #define BSDF_BIT_COUNT_TOPLAYER					1
-#define BSDF_BIT_COUNT_HASSSS					1
-#define BSDF_BIT_COUNT_HASSSSPROFILE			1
+#define BSDF_BIT_COUNT_SSSTYPE					2
 #define BSDF_BIT_COUNT_GREYWEIGHT_V				1
 #define BSDF_BIT_COUNT_HAZINESS					1
 #define BSDF_BIT_COUNT_F90						1
@@ -285,9 +284,8 @@ struct FStrataBSDF
 // Bit offset of each component store in BSDF state 
 #define BSDF_BIT_OFFSET_ANISOTROPY				0
 #define BSDF_BIT_OFFSET_TOPLAYER				(BSDF_BIT_OFFSET_ANISOTROPY			+ BSDF_BIT_COUNT_ANISOTROPY)
-#define BSDF_BIT_OFFSET_HASSSS					(BSDF_BIT_OFFSET_TOPLAYER			+ BSDF_BIT_COUNT_TOPLAYER)
-#define BSDF_BIT_OFFSET_HASSSSPROFILE			(BSDF_BIT_OFFSET_HASSSS				+ BSDF_BIT_COUNT_HASSSS)
-#define BSDF_BIT_OFFSET_GREYWEIGHT_V			(BSDF_BIT_OFFSET_HASSSSPROFILE		+ BSDF_BIT_COUNT_HASSSSPROFILE)
+#define BSDF_BIT_OFFSET_SSSTYPE					(BSDF_BIT_OFFSET_TOPLAYER			+ BSDF_BIT_COUNT_TOPLAYER)
+#define BSDF_BIT_OFFSET_GREYWEIGHT_V			(BSDF_BIT_OFFSET_SSSTYPE			+ BSDF_BIT_COUNT_SSSTYPE)
 #define BSDF_BIT_OFFSET_HAZINESS				(BSDF_BIT_OFFSET_GREYWEIGHT_V		+ BSDF_BIT_COUNT_GREYWEIGHT_V)
 #define BSDF_BIT_OFFSET_F90						(BSDF_BIT_OFFSET_HAZINESS			+ BSDF_BIT_COUNT_HAZINESS)
 #define BSDF_BIT_OFFSET_SIMPLEVOLUME			(BSDF_BIT_OFFSET_F90				+ BSDF_BIT_COUNT_F90)
@@ -309,6 +307,12 @@ struct FStrataBSDF
 #define BSDF_COMPLEXCODING_BIT_COUNT			(32u)
 #define BSDF_COMPLEXFASTENCODING_MASK			(0xFFFFFFFFu)
 
+// Sub-surface scattering types
+#define SSS_TYPE_INVALID						(0u)
+#define SSS_TYPE_WRAP							(1u)
+#define SSS_TYPE_DIFFUSION						(2u)
+#define SSS_TYPE_DIFFUSION_PROFILE				(3u)
+
 // Sanity check
 #if (BSDF_BIT_OFFSET_WEIGHT10F + BSDF_BIT_COUNT_WEIGHT10F) > 32
 #error Strata BSDF state is larger than 32 bits
@@ -336,12 +340,8 @@ struct FStrataBSDF
 #define BSDF_SETISTOPLAYER(X, IsTop)	WRITE_BITS(X.State, BSDF_BIT_COUNT_TOPLAYER, BSDF_BIT_OFFSET_TOPLAYER, IsTop)
 
 // Indicates if the BSDF has scattering component
-#define BSDF_GETHASSSS(X)				READ_BITS(X.State,  BSDF_BIT_COUNT_HASSSS, BSDF_BIT_OFFSET_HASSSS)
-#define BSDF_SETHASSSS(X, Scatt)		WRITE_BITS(X.State, BSDF_BIT_COUNT_HASSSS, BSDF_BIT_OFFSET_HASSSS, Scatt)
-
-// Indicates if the BSDF has a SSS profile
-#define BSDF_GETHASSSSPROFILE(X)		READ_BITS(X.State,  BSDF_BIT_COUNT_HASSSSPROFILE, BSDF_BIT_OFFSET_HASSSSPROFILE)
-#define BSDF_SETHASSSSPROFILE(X, Profi)	WRITE_BITS(X.State, BSDF_BIT_COUNT_HASSSSPROFILE, BSDF_BIT_OFFSET_HASSSSPROFILE, Profi)
+#define BSDF_GETSSSTYPE(X)				READ_BITS(X.State,  BSDF_BIT_COUNT_SSSTYPE, BSDF_BIT_OFFSET_SSSTYPE)
+#define BSDF_SETSSSTYPE(X, Scatt)		WRITE_BITS(X.State, BSDF_BIT_COUNT_SSSTYPE, BSDF_BIT_OFFSET_SSSTYPE, Scatt)
 
 // Indicates if the BSDF weight is grey scale so that it can be stored in a single float (e.g. top layer BSDFs)
 #define BSDF_GETHASGREYWEIGHT_V(X)		READ_BITS(X.State,  BSDF_BIT_COUNT_GREYWEIGHT_V, BSDF_BIT_OFFSET_GREYWEIGHT_V)
@@ -397,10 +397,10 @@ struct FStrataBSDF
 
 ///////////////////////////////////////////////////////////////////////////////
 // Slab BSDF members
-#define SLAB_COMPLEX_PATH_MASK (TO_OFFSET_BITMASK(BSDF_BIT_COUNT_ANISOTROPY, BSDF_BIT_OFFSET_ANISOTROPY)		| TO_OFFSET_BITMASK(BSDF_BIT_COUNT_HASSSS, BSDF_BIT_OFFSET_HASSSS) | \
-								TO_OFFSET_BITMASK(BSDF_BIT_COUNT_HASSSSPROFILE, BSDF_BIT_OFFSET_HASSSSPROFILE)	| TO_OFFSET_BITMASK(BSDF_BIT_COUNT_HAZINESS, BSDF_BIT_OFFSET_HAZINESS) | \
-								TO_OFFSET_BITMASK(BSDF_BIT_COUNT_F90, BSDF_BIT_OFFSET_F90)						| TO_OFFSET_BITMASK(BSDF_BIT_COUNT_SIMPLEVOLUME, BSDF_BIT_OFFSET_SIMPLEVOLUME) | \
-								TO_OFFSET_BITMASK(BSDF_BIT_COUNT_HASFUZZ, BSDF_BIT_OFFSET_HASFUZZ)				| TO_OFFSET_BITMASK(BSDF_BIT_COUNT_ISTHIN, BSDF_BIT_OFFSET_ISTHIN) )
+#define SLAB_COMPLEX_PATH_MASK (TO_OFFSET_BITMASK(BSDF_BIT_COUNT_ANISOTROPY, BSDF_BIT_OFFSET_ANISOTROPY)		| \
+								TO_OFFSET_BITMASK(BSDF_BIT_COUNT_SSSTYPE,    BSDF_BIT_OFFSET_SSSTYPE)			| TO_OFFSET_BITMASK(BSDF_BIT_COUNT_HAZINESS,     BSDF_BIT_OFFSET_HAZINESS) | \
+								TO_OFFSET_BITMASK(BSDF_BIT_COUNT_F90,        BSDF_BIT_OFFSET_F90)				| TO_OFFSET_BITMASK(BSDF_BIT_COUNT_SIMPLEVOLUME, BSDF_BIT_OFFSET_SIMPLEVOLUME) | \
+								TO_OFFSET_BITMASK(BSDF_BIT_COUNT_HASFUZZ,    BSDF_BIT_OFFSET_HASFUZZ)			| TO_OFFSET_BITMASK(BSDF_BIT_COUNT_ISTHIN,       BSDF_BIT_OFFSET_ISTHIN) )
 
 #define SLAB_DIFFUSEALBEDO(X)			X.VGPRs[0].xyz
 #define SLAB_ROUGHNESS(X)				X.VGPRs[0].w
@@ -414,7 +414,7 @@ struct FStrataBSDF
 #define SLAB_SSSPROFILERADIUSSCALE(X)	X.VGPRs[3].y
 // Or explicit MFP in centimeter
 #define SLAB_SSSMFP(X)					X.VGPRs[3].xyz
-// Free slot:                           X.VGPRs[3].w
+#define SLAB_SSSPHASEANISOTROPY(X)		X.VGPRs[3].w
 // Fuzz to simulate cloth
 #define SLAB_FUZZ_COLOR(X)				X.VGPRs[4].xyz
 #define SLAB_FUZZ_AMOUNT(X)				X.VGPRs[4].w
@@ -1108,7 +1108,7 @@ FStrataData GetStrataSlabBSDF(
 	float3 BaseColor, float3 EdgeColor, float Specular, float Metallic,	//  Metalness workflow
 	float3 DiffuseAlbedo, float3 F0, float3 F90,						// !Metalness workflow
 	float Roughness, float Anisotropy,
-	float SSSProfileID, float3 SSSMFP, float SSSMFPScale,
+	float SSSProfileID, float3 SSSMFP, float SSSMFPScale, float SSSPhaseAnisotropy, float UseSSSDiffusion,
 	float3 Emissive, 
 	float Haziness, 
 	float FuzzAmount, float3 FuzzColor,
@@ -1186,8 +1186,7 @@ FStrataData GetStrataSlabBSDF(
 	BSDF_SETEMISSIVE		(StrataData.InlinedBSDF,		Emissive);
 	BSDF_SETHASANISOTROPY	(StrataData.InlinedBSDF,		bHasAnisotropy);
 	BSDF_SETISTOPLAYER		(StrataData.InlinedBSDF,		0);
-	BSDF_SETHASSSS			(StrataData.InlinedBSDF,		bHasSSS ? 1 : 0);
-	BSDF_SETHASSSSPROFILE	(StrataData.InlinedBSDF,		bHasSSSProfile ? 1 : 0);
+	BSDF_SETSSSTYPE			(StrataData.InlinedBSDF,		bHasSSS ? (bHasSSSProfile ? SSS_TYPE_DIFFUSION_PROFILE : SSS_TYPE_DIFFUSION) : SSS_TYPE_INVALID);
 	BSDF_SETISTHIN			(StrataData.InlinedBSDF,		bIsThin ? 1 : 0);
 	BSDF_SETHASMFP			(StrataData.InlinedBSDF,		bMFPPlugged ? 1 : 0);
 	BSDF_SETHASHAZINESS		(StrataData.InlinedBSDF,		bHasHaziness);
@@ -1224,8 +1223,9 @@ FStrataData GetStrataSlabBSDF(
 	}
 	else
 	{
-		SLAB_SSSMFP	(StrataData.InlinedBSDF)			= SSSMFP * SSSMFPScale;
+		SLAB_SSSMFP	(StrataData.InlinedBSDF)				= SSSMFP * SSSMFPScale;
 	}
+	SLAB_SSSPHASEANISOTROPY(StrataData.InlinedBSDF)			= SSSPhaseAnisotropy;
 
 	SLAB_FUZZ_AMOUNT		(StrataData.InlinedBSDF)		= FuzzAmount;
 	SLAB_FUZZ_COLOR			(StrataData.InlinedBSDF)		= FuzzColor;
@@ -1819,6 +1819,11 @@ float HorizontalMixingParameterBlendingBSDFCoverageToNormalMix(FStrataData Backg
 #endif
 }
 
+// Merge two type of SSS type, by using the most complex behavior, i.e., No < Wrap < Diffusion < Diffusion Profile
+uint StrataMergeSSSType(uint SSSType0, uint SSSType1)
+{
+	return max(SSSType0, SSSType1);
+}
 
 // Note about parameter blending
 //  - We can only parameter blend Slab nodes. This error is handled in the compiler.
@@ -1843,8 +1848,7 @@ FStrataData StrataHorizontalMixingParameterBlending(FStrataData Background, FStr
 
 	BSDF_SETISTOPLAYER(ResultBSDF,		0);
 	BSDF_SETHASANISOTROPY(ResultBSDF,	BSDF_GETHASANISOTROPY(ForegroundBSDF)			| BSDF_GETHASANISOTROPY(BackgroundBSDF));
-	BSDF_SETHASSSS(ResultBSDF,			BSDF_GETHASSSS(ForegroundBSDF)					| BSDF_GETHASSSS(BackgroundBSDF));
-	BSDF_SETHASSSSPROFILE(ResultBSDF,	BSDF_GETHASSSSPROFILE(ForegroundBSDF)			| BSDF_GETHASSSSPROFILE(BackgroundBSDF));	// Problem: cannot blend SSS profiles...
+	BSDF_SETSSSTYPE(ResultBSDF,			StrataMergeSSSType(BSDF_GETSSSTYPE(ForegroundBSDF), BSDF_GETSSSTYPE(BackgroundBSDF)));
 	BSDF_SETISTHIN(ResultBSDF,			BSDF_GETISTHIN(ForegroundBSDF)					| BSDF_GETISTHIN(BackgroundBSDF));
 	BSDF_SETHASMFP(ResultBSDF,			BSDF_GETHASMFP(ForegroundBSDF)					| BSDF_GETHASMFP(BackgroundBSDF));
 	BSDF_SETHASHAZINESS(ResultBSDF,		BSDF_GETHASHAZINESS(ForegroundBSDF)				| BSDF_GETHASHAZINESS(BackgroundBSDF));
@@ -1864,10 +1868,10 @@ FStrataData StrataHorizontalMixingParameterBlending(FStrataData Background, FStr
 	SLAB_HAZINESS(ResultBSDF) =			lerp(SLAB_HAZINESS(BackgroundBSDF),				SLAB_HAZINESS(ForegroundBSDF),				ForegroundOtherMixFactor);
 	SLAB_FUZZ_AMOUNT(ResultBSDF) =		lerp(SLAB_FUZZ_AMOUNT(BackgroundBSDF),			SLAB_FUZZ_AMOUNT(ForegroundBSDF),			ForegroundOtherMixFactor);
 
-	if (BSDF_GETHASSSSPROFILE(ResultBSDF))
+	if (BSDF_GETSSSTYPE(ResultBSDF) == SSS_TYPE_DIFFUSION_PROFILE)
 	{
-		const bool bForegroundHasSSSProfile = BSDF_GETHASSSSPROFILE(ForegroundBSDF);
-		const bool bBackgroundHasSSSProfile = BSDF_GETHASSSSPROFILE(BackgroundBSDF);
+		const bool bForegroundHasSSSProfile = BSDF_GETSSSTYPE(ForegroundBSDF) == SSS_TYPE_DIFFUSION_PROFILE;
+		const bool bBackgroundHasSSSProfile = BSDF_GETSSSTYPE(BackgroundBSDF) == SSS_TYPE_DIFFUSION_PROFILE;
 
 		// As mentioned above, we cannot blend SSSProfiles. So we select the one that triggered the profile, foreground having priority.
 		SLAB_SSSPROFILEID(ResultBSDF) = bForegroundHasSSSProfile ? SLAB_SSSPROFILEID(ForegroundBSDF) : SLAB_SSSPROFILEID(BackgroundBSDF);
@@ -1881,8 +1885,7 @@ FStrataData StrataHorizontalMixingParameterBlending(FStrataData Background, FStr
 		if (SLAB_SSSPROFILERADIUSSCALE(ResultBSDF) == 0.0f)
 		{
 			SLAB_SSSPROFILEID(ResultBSDF) = SSS_PROFILE_ID_INVALID;
-			BSDF_SETHASSSS(ResultBSDF, 0);
-			BSDF_SETHASSSSPROFILE(ResultBSDF, 0);
+			BSDF_SETSSSTYPE(ResultBSDF, SSS_TYPE_INVALID);
 		}
 	}
 	else
@@ -1892,14 +1895,14 @@ FStrataData StrataHorizontalMixingParameterBlending(FStrataData Background, FStr
 		// We still try to disable SSS if possible: for performance and also to preserve the surface color (especially specular reflection color).
 		if (all(SLAB_SSSMFP(ResultBSDF) == 0.0f))
 		{
-			BSDF_SETHASSSS(ResultBSDF, 0);
+			BSDF_SETSSSTYPE(ResultBSDF, SSS_TYPE_INVALID);
 		}
 	}
 
 
 #if STRATA_INLINE_SHADING
 	// At this stage we do not know if this strata BSDF is going to be using simple volume lighting so always store the MFP on the side to not override SSSPROFILE data.
-	ResultBSDF.TmpMFP =				lerp(BackgroundBSDF.TmpMFP,					ForegroundBSDF.TmpMFP,						ForegroundOtherMixFactor);
+	ResultBSDF.TmpMFP =					lerp(BackgroundBSDF.TmpMFP,						ForegroundBSDF.TmpMFP,						ForegroundOtherMixFactor);
 
 	// We blend the parameter above, and we also need to lerp the Coverage to match the total energy.
 	ResultBSDF.Coverage =				lerp(BackgroundBSDF.Coverage,					ForegroundBSDF.Coverage,					HorizontalMixValue);
@@ -1938,8 +1941,7 @@ FStrataData StrataAddParameterBlending(FStrataData A, FStrataData B, float AMixF
 
 	BSDF_SETISTOPLAYER(ResultBSDF,		0);
 	BSDF_SETHASANISOTROPY(ResultBSDF,	BSDF_GETHASANISOTROPY(BBSDF)					| BSDF_GETHASANISOTROPY(ABSDF));
-	BSDF_SETHASSSS(ResultBSDF,			BSDF_GETHASSSS(BBSDF)							| BSDF_GETHASSSS(ABSDF));
-	BSDF_SETHASSSSPROFILE(ResultBSDF,	BSDF_GETHASSSSPROFILE(BBSDF)					| BSDF_GETHASSSSPROFILE(ABSDF));	// Problem: cannot blend SSS profiles...
+	BSDF_SETSSSTYPE(ResultBSDF,			StrataMergeSSSType(BSDF_GETSSSTYPE(BBSDF)		, BSDF_GETSSSTYPE(ABSDF)));
 	BSDF_SETISTHIN(ResultBSDF,			BSDF_GETISTHIN(BBSDF)							| BSDF_GETISTHIN(ABSDF));
 	BSDF_SETHASMFP(ResultBSDF,			BSDF_GETHASMFP(BBSDF)							| BSDF_GETHASMFP(ABSDF));
 	BSDF_SETHASHAZINESS(ResultBSDF,		BSDF_GETHASHAZINESS(BBSDF)						| BSDF_GETHASHAZINESS(ABSDF));
@@ -1961,10 +1963,10 @@ FStrataData StrataAddParameterBlending(FStrataData A, FStrataData B, float AMixF
 	SLAB_HAZINESS(ResultBSDF) =			lerp(SLAB_HAZINESS(BBSDF),						SLAB_HAZINESS(ABSDF),						AMixFactor);
 	SLAB_FUZZ_AMOUNT(ResultBSDF) =		lerp(SLAB_FUZZ_AMOUNT(BBSDF),					SLAB_FUZZ_AMOUNT(ABSDF),					AMixFactor);
 
-	if (BSDF_GETHASSSSPROFILE(ResultBSDF))
+	if (BSDF_GETSSSTYPE(ResultBSDF) == SSS_TYPE_DIFFUSION_PROFILE)
 	{
 		// As mentioned above, we cannot blend SSSProfiles. We do select the appropriate profile, ABSDF being selected if both have SSSProfile.
-		if (BSDF_GETHASSSSPROFILE(ABSDF))
+		if (BSDF_GETSSSTYPE(ABSDF) == SSS_TYPE_DIFFUSION_PROFILE)
 		{
 			SLAB_SSSPROFILEID(ResultBSDF) = SLAB_SSSPROFILEID(ABSDF);
 			SLAB_SSSPROFILERADIUSSCALE(ResultBSDF) = SLAB_SSSPROFILERADIUSSCALE(ABSDF);
@@ -2027,8 +2029,7 @@ FStrataBSDF StrataVerticalLayeringParameterBlending(in FStrataBSDF TopBSDF, in F
 
 	BSDF_SETISTOPLAYER(ResultBSDF,		0);
 	BSDF_SETHASANISOTROPY(ResultBSDF,	0);																					// Anisotropy is disabled because it changes a lot of things
-	BSDF_SETHASSSS(ResultBSDF,			BSDF_GETHASSSS(BaseBSDF));															// We only keep SSS if the bottom layer has it
-	BSDF_SETHASSSSPROFILE(ResultBSDF,	0);																					// Problem: cannot blend SSS profiles so we simply disable SSSProfile
+	BSDF_SETSSSTYPE(ResultBSDF,			min(BSDF_GETSSSTYPE(BaseBSDF), SSS_TYPE_DIFFUSION));								// We only keep SSS if the bottom layer has it. Problem: cannot blend SSS profiles so we simply disable SSSProfile in this case, i.e., min(..., SSS_TYPE_DIFFUSION)
 	BSDF_SETISTHIN(ResultBSDF,			BSDF_GETISTHIN(BaseBSDF));
 	BSDF_SETHASMFP(ResultBSDF,			BSDF_GETHASMFP(BaseBSDF)						| BSDF_GETHASMFP(TopBSDF));			// We keep MFP of both to make sure bottom layer on translucent can also be effective
 	BSDF_SETHASHAZINESS(ResultBSDF,		BSDF_GETHASHAZINESS(TopBSDF));														// We only keep haziness if the top layer has it
@@ -2195,16 +2196,13 @@ void SanitizeStrataSlab(inout FStrataBSDF BSDF)
 #endif
 
 	// BSDF feature tests must be executed for profile path to not interact with MFP, especially for forward rendered simple volumetrics.
-	if (BSDF_GETHASSSS(BSDF))
+	if (BSDF_GETSSSTYPE(BSDF) == SSS_TYPE_DIFFUSION_PROFILE)
 	{
-		if (BSDF_GETHASSSSPROFILE(BSDF))
-		{
-			SLAB_SSSPROFILERADIUSSCALE(BSDF) = saturate(SLAB_SSSPROFILERADIUSSCALE(BSDF));
-		}
-		else
-		{
-			SLAB_SSSMFP(BSDF) = clamp(SLAB_SSSMFP(BSDF), 0.0f, Max111110BitsFloat3);
-		}
+		SLAB_SSSPROFILERADIUSSCALE(BSDF) = saturate(SLAB_SSSPROFILERADIUSSCALE(BSDF));
+	}
+	else if (BSDF_GETSSSTYPE(BSDF) == SSS_TYPE_DIFFUSION)
+	{
+		SLAB_SSSMFP(BSDF) = clamp(SLAB_SSSMFP(BSDF), 0.0f, Max111110BitsFloat3);
 	}
 
 	SLAB_FUZZ_AMOUNT(BSDF) = saturate(SLAB_FUZZ_AMOUNT(BSDF));
@@ -2330,8 +2328,7 @@ void EnableSlabBSDFSimpleVolumetric(inout FStrataBSDF BSDF)
 	EnableSlabMFPComputationSpace(BSDF);
 
 	// Mark the BSDF has not having any SSS through post process.
-	BSDF_SETHASSSS(BSDF, 0);
-	BSDF_SETHASSSSPROFILE(BSDF, 0);
+	BSDF_SETSSSTYPE(BSDF, SSS_TYPE_INVALID);
 
 	// Notify that the Slab can use SSS post-process only if at the bottom (layered considered optically thick, e.g. cannot see-through)
 	// When forward rendering is used or if a layer is not a bottom layer, SSS post cannot be used and simple scattering is used.
@@ -2422,7 +2419,7 @@ float3 StrataGetBSDFSubSurfaceColor(in FStrataBSDF BSDF)
 {
 	switch (BSDF_GETTYPE(BSDF))
 	{
-	case STRATA_BSDF_TYPE_SLAB: 			return BSDF_GETHASSSS(BSDF) ? SLAB_DIFFUSEALBEDO(BSDF) : 0; // A more accurate version would be IsotropicMediumSlabPunctualDirectionalAlbedo(StrataSlabCreateParticipatingMedia(SLAB_DIFFUSEALBEDO(BSDF), SLAB_SSSMFP(BSDFContext.BSDF)));
+	case STRATA_BSDF_TYPE_SLAB: 			return BSDF_GETSSSTYPE(BSDF) != SSS_TYPE_INVALID ? SLAB_DIFFUSEALBEDO(BSDF) : 0; // A more accurate version would be IsotropicMediumSlabPunctualDirectionalAlbedo(StrataSlabCreateParticipatingMedia(SLAB_DIFFUSEALBEDO(BSDF), SLAB_SSSMFP(BSDFContext.BSDF)));
 	case STRATA_BSDF_TYPE_HAIR: 			return 0;
 	case STRATA_BSDF_TYPE_SINGLELAYERWATER: return 0;
 	}
@@ -2515,14 +2512,14 @@ float StrataGetLegacyShadingModels(in FStrataBSDF BSDF)
 	case STRATA_BSDF_TYPE_SLAB:
 	{
 		// Base on features, we can only partially recovers the original shading models. Shading models like foliage, eye, clear coat are impossible/harder to infer.
-		if (BSDF_GETHASSSSPROFILE(BSDF))	{ return SHADINGMODELID_SUBSURFACE_PROFILE; }
-		else if (BSDF_GETHASSSS(BSDF))		{ return SHADINGMODELID_SUBSURFACE; }
-		else if (BSDF_GETHASFUZZ(BSDF))		{ return SHADINGMODELID_CLOTH; }
-		else if (BSDF_GETHASFUZZ(BSDF))		{ return SHADINGMODELID_TWOSIDED_FOLIAGE; }
-		else								{ return SHADINGMODELID_DEFAULT_LIT; }
+		if (BSDF_GETSSSTYPE(BSDF) == SSS_TYPE_DIFFUSION_PROFILE){ return SHADINGMODELID_SUBSURFACE_PROFILE; }
+		else if (BSDF_GETSSSTYPE(BSDF) != SSS_TYPE_INVALID)		{ return SHADINGMODELID_SUBSURFACE; }
+		else if (BSDF_GETHASFUZZ(BSDF))							{ return SHADINGMODELID_CLOTH; }
+		else if (BSDF_GETHASFUZZ(BSDF))							{ return SHADINGMODELID_TWOSIDED_FOLIAGE; }
+		else													{ return SHADINGMODELID_DEFAULT_LIT; }
 	}
-	case STRATA_BSDF_TYPE_HAIR: 			{ return SHADINGMODELID_HAIR; }
-	case STRATA_BSDF_TYPE_SINGLELAYERWATER: { return SHADINGMODELID_SINGLELAYERWATER; }
+	case STRATA_BSDF_TYPE_HAIR: 								{ return SHADINGMODELID_HAIR; }
+	case STRATA_BSDF_TYPE_SINGLELAYERWATER:						{ return SHADINGMODELID_SINGLELAYERWATER; }
 	}
 	return 0.f;
 }
@@ -2947,6 +2944,8 @@ void StrataConvertFromDBuffer(
 		0.f,				// SSSProfileID
 		0.f,				// SSSMFP
 		0.f,				// SSSMFPScale
+		0.f,				// SSSPhaseAniso
+		0.f,				// UseSSSDiffusion
 		0.f,				// Emissive
 		0.f,				// Haziness
 		0.f,				// FuzzAmount
@@ -3165,12 +3164,13 @@ void PackStrataOut(
 						TopLayerTotalWeight += BSDF.TopLayerDataWeight;
 						TopLayerData.Roughness += BSDF.TopLayerDataWeight * SLAB_ROUGHNESS(BSDF);
 
+						const bool bHasSSS = BSDF_GETSSSTYPE(BSDF) != SSS_TYPE_INVALID;
 						const bool bIsThin = BSDF_GETISTHIN(BSDF);
 						const bool bIsSimpleVolume = BSDF_GETISSIMPLEVOLUME(BSDF);
-						bStrataSubsurfaceEnable = bStrataSubsurfaceEnable || (BSDF_GETHASSSS(BSDF) && BSDF.Coverage > 0.0f && !bIsSimpleVolume && !bIsThin); // Should it be coverage>0 or any(LuminanceWeight>0) ?
+						bStrataSubsurfaceEnable = bStrataSubsurfaceEnable || (bHasSSS && BSDF.Coverage > 0.0f && !bIsSimpleVolume && !bIsThin); // Should it be coverage>0 or any(LuminanceWeight>0) ?
 
 						const bool bSSSMask = !bIsThin && BSDF.bIsBottom;	// SSS and Thin can only be on the bottom layer, so ignoring bIsSimpleVolume here.
-						SetHasSubsurface(StrataPixelHeader, (bSSSMask && BSDF_GETHASSSS(BSDF)));
+						SetHasSubsurface(StrataPixelHeader, (bSSSMask && bHasSSS));
 
 						break;
 					}
@@ -3367,7 +3367,7 @@ void PackStrataOut(
 							const uint PackedData8Bits = (PackedF020Bits >> 12) & 0xFF;
 
 							STRATA_STORE_UINT1(PackedData32Bits);
-							STRATA_STORE_UINT1(PackedData8Bits | PackRGBA8(float4(0.0f, SLAB_ROUGHNESS(BSDF), (SLAB_ANISOTROPY(BSDF) + 1.f) * 0.5f, 0.f)));
+							STRATA_STORE_UINT1(PackedData8Bits | PackRGBA8(float4(0.0f, SLAB_ROUGHNESS(BSDF), (SLAB_ANISOTROPY(BSDF) + 1.f) * 0.5f, (SLAB_SSSPHASEANISOTROPY(BSDF) + 1.f) * 0.5f)));
 
 							if (BSDF_GETHASF90(BSDF) || BSDF_GETHASHAZINESS(BSDF))
 							{
@@ -3382,11 +3382,12 @@ void PackStrataOut(
 
 							const bool bIsSimpleVolume = BSDF_GETISSIMPLEVOLUME(BSDF);
 							const bool bIsThin = BSDF_GETISTHIN(BSDF);
-							if (BSDF_GETHASSSS(BSDF) || bIsSimpleVolume || bIsThin)
+							const bool bHasSSS = BSDF_GETSSSTYPE(BSDF) != SSS_TYPE_INVALID;
+							if (bHasSSS || bIsSimpleVolume || bIsThin)
 							{
-								if (!bIsThin && BSDF_GETHASSSS(BSDF))
+								if (!bIsThin && bHasSSS)
 								{
-									const bool bIsProfile = StrataSubSurfaceHeaderGetIsProfile(SSSData.Header) || BSDF_GETHASSSSPROFILE(BSDF);
+									const bool bIsProfile = StrataSubSurfaceHeaderGetIsProfile(SSSData.Header) || BSDF_GETSSSTYPE(BSDF) == SSS_TYPE_DIFFUSION_PROFILE;
 									StrataSubSurfaceHeaderSetIsValid(SSSData.Header, true);
 									StrataSubSurfaceHeaderSetIsProfile(SSSData.Header, bIsProfile);
 									if (bIsProfile)
@@ -3401,7 +3402,7 @@ void PackStrataOut(
 									StrataSubsurfaceExtrasSetBaseColor(SSSData.Extras, SLAB_DIFFUSEALBEDO(BSDF));
 								}
 
-								if (BSDF_GETHASSSSPROFILE(BSDF) && !bIsSimpleVolume)
+								if (BSDF_GETSSSTYPE(BSDF) == SSS_TYPE_DIFFUSION_PROFILE && !bIsSimpleVolume)
 								{
 									// Simple volume overrides SSS profile because it takes over for when not at the bottom of the BSDF layer, or during forward rendering.
 									STRATA_STORE_UINT1(PackRGBA8(float4(SLAB_SSSPROFILEID(BSDF), SLAB_SSSPROFILERADIUSSCALE(BSDF), 0, 0)));
@@ -3556,8 +3557,9 @@ void UnpackFastPathSlabBSDF(inout FStrataBSDF OutBSDF, uint2 PackedData01)
 	SLAB_F0(OutBSDF)			= UnpackR7G7B6Gamma2(PackedF020Bits);
 
 	float4 Data1 = UnpackRGBA8(PackedData01.y);
-	SLAB_ROUGHNESS(OutBSDF)		= Data1.y;
-	SLAB_ANISOTROPY(OutBSDF)	= Data1.z * 2.f - 1.f;
+	SLAB_ROUGHNESS(OutBSDF)			= Data1.y;
+	SLAB_ANISOTROPY(OutBSDF)		= Data1.z * 2.f - 1.f;
+	SLAB_SSSPHASEANISOTROPY(OutBSDF)= Data1.2 * 2.f - 1.f;
 }
 
 FStrataBSDF UnpackFastPathStrataBSDFIn(FStrataMaterialContainer StrataBuffer, inout FStrataAddressing StrataAddressing, FStrataPixelHeader StrataPixelHeader)
@@ -3593,8 +3595,7 @@ FStrataBSDF UnpackFastPathStrataBSDFIn(FStrataMaterialContainer StrataBuffer, in
 	BSDF_SETHASF90(OutBSDF, 0);
 	BSDF_SETHASANISOTROPY(OutBSDF, 0);
 	BSDF_SETISTOPLAYER(OutBSDF, 1);
-	BSDF_SETHASSSS(OutBSDF, 0);
-	BSDF_SETHASSSSPROFILE(OutBSDF, 0);
+	BSDF_SETSSSTYPE(OutBSDF, SSS_TYPE_INVALID);
 	BSDF_SETHASHAZINESS(OutBSDF, 0);
 	BSDF_SETISSIMPLEVOLUME(OutBSDF, 0);
 	BSDF_SETHASFUZZ(OutBSDF, 0);
@@ -3680,9 +3681,10 @@ FStrataBSDF UnpackStrataBSDFIn(FStrataMaterialContainer StrataBuffer, inout FStr
 
 		const bool bIsSimpleVolume = BSDF_GETISSIMPLEVOLUME(OutBSDF);
 		const bool bIsThin = BSDF_GETISTHIN(OutBSDF);
-		if (BSDF_GETHASSSS(OutBSDF) || bIsSimpleVolume || bIsThin)
+		const bool bHasSSS = BSDF_GETSSSTYPE(OutBSDF) != SSS_TYPE_INVALID;
+		if (bHasSSS || bIsSimpleVolume || bIsThin)
 		{
-			if (BSDF_GETHASSSSPROFILE(OutBSDF) && !bIsSimpleVolume && !bIsThin)
+			if (BSDF_GETSSSTYPE(OutBSDF) == SSS_TYPE_DIFFUSION_PROFILE && !bIsSimpleVolume && !bIsThin)
 			{
 				float4 Data = UnpackRGBA8(StrataLoadUint1(StrataBuffer, StrataAddressing));
 				SLAB_SSSPROFILEID(OutBSDF) = Data.r;
diff --git a/Engine/Shaders/Private/Strata/StrataEvaluation.ush b/Engine/Shaders/Private/Strata/StrataEvaluation.ush
--- a/Engine/Shaders/Private/Strata/StrataEvaluation.ush
+++ b/Engine/Shaders/Private/Strata/StrataEvaluation.ush
@@ -244,7 +244,7 @@ FStrataEvaluateResult StrataEvaluateBSDFCommon(FStrataBSDFContext BSDFContext, F
 		#if STRATA_FASTPATH==0
 			float HazeSafeRoughness		= SafeRoughness;
 		#endif
-			const bool bHasSubsurface	= BSDF_GETHASSSS(BSDFContext.BSDF) > 0;
+			const bool bHasSubsurface	= BSDF_GETSSSTYPE(BSDFContext.BSDF) != SSS_TYPE_INVALID;
 			const bool bHasAnisotropy	= BSDF_GETHASANISOTROPY(BSDFContext.BSDF);
 			const bool bHaziness		= BSDF_GETHASHAZINESS(BSDFContext.BSDF);
 
@@ -550,8 +550,8 @@ FStrataEvaluateResult StrataEvaluateBSDFCommon(FStrataBSDFContext BSDFContext, F
 			////
 
 			#if STRATA_FASTPATH==0 && STRATA_SSS_TRANSMISSION
-			const bool bHasProfile = BSDF_GETHASSSSPROFILE(BSDFContext.BSDF);
-			if (BSDF_GETHASSSS(BSDFContext.BSDF))
+			const bool bHasProfile = BSDF_GETSSSTYPE(BSDFContext.BSDF) == SSS_TYPE_DIFFUSION_PROFILE;
+			if (BSDF_GETSSSTYPE(BSDFContext.BSDF) != SSS_TYPE_INVALID)
 			{
 				const float ThicknessInCm = DecodeThickness(ShadowTerms.TransmissionThickness) * SSSS_MAX_TRANSMISSION_PROFILE_DISTANCE;
 
@@ -572,7 +572,7 @@ FStrataEvaluateResult StrataEvaluateBSDFCommon(FStrataBSDFContext BSDFContext, F
 					const float3 SubsurfaceAlebdo = SLAB_DIFFUSEALBEDO(BSDFContext.BSDF);
 					const float3 MeanFreePathInCm = SLAB_SSSMFP(BSDFContext.BSDF);
 					TransmissionThroughput	= GetBurleyTransmissionProfile(SubsurfaceAlebdo, MeanFreePathInCm, ThicknessInCm);
-					PhaseFunctionAnisotropy = 0.93f; // STRATA_TODO: expose phase function eccentricity. For now, forward scattering matching default SSS profile value
+					PhaseFunctionAnisotropy = SLAB_SSSPHASEANISOTROPY(BSDFContext.BSDF);
 					OneOverIOR				= 1.f / DielectricF0ToIor(F0.y);
 				}
 
@@ -969,7 +969,7 @@ FStrataEnvLightResult StrataEvaluateForEnvLight(FStrataBSDFContext BSDFContext,
 			}
 
 			StrataEnvLightResult.DiffuseNormal = BSDFContext.N;
-			StrataEnvLightResult.bSubsurface = BSDF_GETHASSSS(BSDFContext.BSDF) > 0;
+			StrataEnvLightResult.bSubsurface = BSDF_GETSSSTYPE(BSDFContext.BSDF) != SSS_TYPE_INVALID;
 
 			float SafeRoughness = MakeRoughnessSafe(SLAB_ROUGHNESS(BSDFContext.BSDF));
 
diff --git a/Engine/Shaders/Private/Strata/StrataLegacyConversion.ush b/Engine/Shaders/Private/Strata/StrataLegacyConversion.ush
--- a/Engine/Shaders/Private/Strata/StrataLegacyConversion.ush
+++ b/Engine/Shaders/Private/Strata/StrataLegacyConversion.ush
@@ -20,6 +20,8 @@ struct FStrataLegacyParameters
 	float  SSSProfileID;
 	float3 SSSMFP;
 	float  SSSMFPScale;
+	float  SSSPhaseAniso;
+	bool   UseSSSDiffusion;
 	float3 Emissive;
 	float  FuzzAmount;
 	float3 FuzzColor;
@@ -51,6 +53,8 @@ FStrataLegacyParameters InitStrataLegacyParameters(FStrataPixelFootprint InPixel
 	Out.SSSProfileID = 0.f;
 	Out.SSSMFP = Zeros;
 	Out.SSSMFPScale = 1.f;
+	Out.SSSPhaseAniso = 0.f;
+	Out.UseSSSDiffusion = false;
 	Out.Emissive = Zeros;
 	Out.FuzzAmount = 0.f;
 	Out.FuzzColor = Zeros;
@@ -87,7 +91,9 @@ FStrataData CreateLegacySlab(
 			In.Anisotropy,		// Anisotropy
 			In.SSSProfileID,	// SSSProfileID
 			In.SSSMFP,			// SSSMFP
-			In.SSSMFPScale,	// SSSMFPScale
+			In.SSSMFPScale,		// SSSMFPScale
+			In.SSSPhaseAniso,	// SSSPhaseAniso
+			In.UseSSSDiffusion, // UseSSSDiffusion
 			In.Emissive,		// Emissive
 			0.f,				// Haziness
 			In.FuzzAmount,		// FuzzAmount
@@ -268,6 +274,8 @@ FStrataData StrataConvertLegacyMaterialDynamic(
 		Top.Anisotropy = Anisotropy;
 		Top.SSSProfileID = SubSurfaceProfileId;
 		Top.SSSMFPScale = Opacity;
+		Top.SSSPhaseAniso = 0.93f;
+		Top.UseSSSDiffusion = true;
 		Top.Emissive = Emissive;
 
 		// No need to take into account opacity weighting for translucent material, as the coverage is assumed to be 1
diff --git a/Engine/Shaders/Private/Strata/StrataVisualize.usf b/Engine/Shaders/Private/Strata/StrataVisualize.usf
--- a/Engine/Shaders/Private/Strata/StrataVisualize.usf
+++ b/Engine/Shaders/Private/Strata/StrataVisualize.usf
@@ -157,12 +157,12 @@ void PrintBSDF(inout FShaderPrintContext Context, inout FStrataAddressing Strata
 			Newline(Context);
 
 			Print(Context, TEXT(" "), FontSilver);
-			Print(Context, BSDF_GETSHAREDLOCALBASISID(BSDF), FontBSDF, 2, 1);	Print(Context, TEXT("       "), FontSilver);
-			PrintBool(Context, BSDF_GETHASANISOTROPY(BSDF), FontBSDF);			Print(Context, TEXT(" "), FontSilver);
-			PrintBool(Context, BSDF_GETISTOPLAYER(BSDF), FontBSDF);				Print(Context, TEXT("    "), FontSilver);
-			PrintBool(Context, BSDF_GETHASSSS(BSDF), FontBSDF);					Print(Context, TEXT("      "), FontSilver);
-			PrintBool(Context, BSDF_GETISSIMPLEVOLUME(BSDF), FontBSDF);			Print(Context, TEXT("        "), FontSilver);
-			PrintBool(Context, BSDF_GETISTHIN(BSDF), FontBSDF);					Print(Context, TEXT("  "), FontSilver);
+			Print(Context, BSDF_GETSHAREDLOCALBASISID(BSDF), FontBSDF, 2, 1);			Print(Context, TEXT("       "), FontSilver);
+			PrintBool(Context, BSDF_GETHASANISOTROPY(BSDF), FontBSDF);					Print(Context, TEXT(" "), FontSilver);
+			PrintBool(Context, BSDF_GETISTOPLAYER(BSDF), FontBSDF);						Print(Context, TEXT("    "), FontSilver);
+			PrintBool(Context, BSDF_GETSSSTYPE(BSDF) != SSS_TYPE_INVALID, FontBSDF);	Print(Context, TEXT("      "), FontSilver);
+			PrintBool(Context, BSDF_GETISSIMPLEVOLUME(BSDF), FontBSDF);					Print(Context, TEXT("        "), FontSilver);
+			PrintBool(Context, BSDF_GETISTHIN(BSDF), FontBSDF);							Print(Context, TEXT("  "), FontSilver);
 
 			// View Weight
 			{
@@ -233,10 +233,10 @@ void PrintBSDF(inout FShaderPrintContext Context, inout FStrataAddressing Strata
 			Print(Context, TEXT(" Haziness    "), FontBSDFPropName);	Print(Context, SLAB_HAZINESS(BSDF), FontBSDFPropValu);				Newline(Context);
 			}
 			
-			if (BSDF_GETHASSSS(BSDF) || BSDF_GETISSIMPLEVOLUME(BSDF) || BSDF_GETISTHIN(BSDF))
+			if (BSDF_GETSSSTYPE(BSDF) != SSS_TYPE_INVALID || BSDF_GETISSIMPLEVOLUME(BSDF) || BSDF_GETISTHIN(BSDF))
 			{
 
-			if (BSDF_GETHASSSSPROFILE(BSDF) && BSDF_GETISSIMPLEVOLUME(BSDF) == 0)
+			if (BSDF_GETSSSTYPE(BSDF) == SSS_TYPE_DIFFUSION_PROFILE && BSDF_GETISSIMPLEVOLUME(BSDF) == 0)
 			{
 			Print(Context, TEXT(" SSS ID.     "), FontBSDFPropName);	Print(Context, SLAB_SSSPROFILEID(BSDF), FontBSDFPropValu);			Newline(Context);
 			Print(Context, TEXT(" SSS Radius  "), FontBSDFPropName);	Print(Context, SLAB_SSSPROFILERADIUSSCALE(BSDF), FontBSDFPropValu);	Newline(Context);
diff --git a/Engine/Source/Runtime/Engine/Classes/Materials/MaterialExpressionStrata.h b/Engine/Source/Runtime/Engine/Classes/Materials/MaterialExpressionStrata.h
--- a/Engine/Source/Runtime/Engine/Classes/Materials/MaterialExpressionStrata.h
+++ b/Engine/Source/Runtime/Engine/Classes/Materials/MaterialExpressionStrata.h
@@ -266,6 +266,12 @@ class UMaterialExpressionStrataSlabBSDF : public UMaterialExpressionStrataBSDF
 	UPROPERTY()
 	FExpressionInput SSSMFPScale;
 
+	/**
+	 * Phase function anisotropy. Positive value elongates the phase function along the lignt direction, causing forward scattering. Negative value elongates the phase function backward to the light direction, causing backward scattering.  (type = float, unitless, defaults to 1, valid value -1..1)
+	 */
+	UPROPERTY()
+	FExpressionInput SSSPhaseAnisotropy;
+
 	/**
 	 * Emissive color on top of the surface (type = float3, unit = luminance, default = 0)
 	 */
@@ -304,6 +310,10 @@ class UMaterialExpressionStrataSlabBSDF : public UMaterialExpressionStrataBSDF
 	UPROPERTY(EditAnywhere, Category = Mode)
 	uint32 bUseMetalness : 1;
 
+	/** Whether to use light diffusion (i.e., SSS diffusion) or wrap-approximation for material with scattering behavior. */
+	UPROPERTY(EditAnywhere, Category = Mode, meta = (DisplayName = "Use Subsurface Diffusion"))
+	uint32 bUseSSSDiffusion : 1;
+
 	//~ Begin UMaterialExpression Interface
 #if WITH_EDITOR
 	virtual int32 Compile(class FMaterialCompiler* Compiler, int32 OutputIndex) override;
diff --git a/Engine/Source/Runtime/Engine/Private/Materials/HLSLMaterialTranslator.cpp b/Engine/Source/Runtime/Engine/Private/Materials/HLSLMaterialTranslator.cpp
--- a/Engine/Source/Runtime/Engine/Private/Materials/HLSLMaterialTranslator.cpp
+++ b/Engine/Source/Runtime/Engine/Private/Materials/HLSLMaterialTranslator.cpp
@@ -10272,8 +10272,8 @@ int32 FHLSLMaterialTranslator::StrataSlabBSDF(
 	int32 BaseColor, int32 EdgeColor, int32 Specular, int32 Metallic,
 	int32 DiffuseAlbedo, int32 F0, int32 F90,
 	int32 Roughness, int32 Anisotropy,
-	int32 SSSProfileId, int32 SSSMFP, int32 SSSMFPScale, 
-	int32 EmissiveColor,	
+	int32 SSSProfileId, int32 SSSMFP, int32 SSSMFPScale, int32 SSSPhaseAniso, int32 UseSSSDiffusion,
+	int32 EmissiveColor,
 	int32 Haziness, 
 	int32 FuzzAmount, int32 FuzzColor,
 	int32 Thickness, 
@@ -10291,7 +10291,7 @@ int32 FHLSLMaterialTranslator::StrataSlabBSDF(
 			return INDEX_NONE;
 		}
 		return AddCodeChunk(
-			MCT_Strata, TEXT("PromoteParameterBlendedBSDFToOperator(GetStrataSlabBSDF(Parameters.StrataPixelFootprint, %s, %s, %s, %s, %s, %s, %s, %s, %s, %s, %s, %s, %s, %s, %s, %s, %s, %s, %s, Parameters.SharedLocalBases.Types) /* Normal = %s ; Tangent = %s */, Parameters.StrataTree, %u, %u, %u, %u)"),
+			MCT_Strata, TEXT("PromoteParameterBlendedBSDFToOperator(GetStrataSlabBSDF(Parameters.StrataPixelFootprint, %s, %s, %s, %s, %s, %s, %s, %s, %s, %s, %s, %s, %s, %s, %s, %s, %s, %s, %s, %s, %s, Parameters.SharedLocalBases.Types) /* Normal = %s ; Tangent = %s */, Parameters.StrataTree, %u, %u, %u, %u)"),
 			*StrataGetCastParameterCode(UseMetalness,	MCT_Float),
 			*StrataGetCastParameterCode(BaseColor,		MCT_Float3),
 			*StrataGetCastParameterCode(EdgeColor,		MCT_Float3),
@@ -10305,6 +10305,8 @@ int32 FHLSLMaterialTranslator::StrataSlabBSDF(
 			*StrataGetCastParameterCode(SSSProfileId,	MCT_Float),
 			*StrataGetCastParameterCode(SSSMFP,			MCT_Float3),
 			*StrataGetCastParameterCode(SSSMFPScale,	MCT_Float),
+			*StrataGetCastParameterCode(SSSPhaseAniso,  MCT_Float),
+			*StrataGetCastParameterCode(UseSSSDiffusion,MCT_Float),
 			*StrataGetCastParameterCode(EmissiveColor,	MCT_Float3),
 			*StrataGetCastParameterCode(Haziness,		MCT_Float),
 			*StrataGetCastParameterCode(FuzzAmount,		MCT_Float),
@@ -10320,7 +10322,7 @@ int32 FHLSLMaterialTranslator::StrataSlabBSDF(
 	}
 	
 	return AddCodeChunk(
-		MCT_Strata, TEXT("GetStrataSlabBSDF(Parameters.StrataPixelFootprint, %s, %s, %s, %s, %s, %s, %s, %s, %s, %s, %s, %s, %s, %s, %s, %s, %s, %s, %s, Parameters.SharedLocalBases.Types) /* Normal = %s ; Tangent = %s */"),
+		MCT_Strata, TEXT("GetStrataSlabBSDF(Parameters.StrataPixelFootprint, %s, %s, %s, %s, %s, %s, %s, %s, %s, %s, %s, %s, %s, %s, %s, %s, %s, %s, %s, %s, %s, Parameters.SharedLocalBases.Types) /* Normal = %s ; Tangent = %s */"),
 		*StrataGetCastParameterCode(UseMetalness,	MCT_Float),
 		*StrataGetCastParameterCode(BaseColor,		MCT_Float3),
 		*StrataGetCastParameterCode(EdgeColor,		MCT_Float3),
@@ -10334,6 +10336,8 @@ int32 FHLSLMaterialTranslator::StrataSlabBSDF(
 		*StrataGetCastParameterCode(SSSProfileId,	MCT_Float),
 		*StrataGetCastParameterCode(SSSMFP,			MCT_Float3),
 		*StrataGetCastParameterCode(SSSMFPScale,	MCT_Float),
+		*StrataGetCastParameterCode(SSSPhaseAniso,	MCT_Float),
+		*StrataGetCastParameterCode(UseSSSDiffusion,MCT_Float),
 		*StrataGetCastParameterCode(EmissiveColor,	MCT_Float3),
 		*StrataGetCastParameterCode(Haziness,		MCT_Float),
 		*StrataGetCastParameterCode(FuzzAmount,		MCT_Float),
diff --git a/Engine/Source/Runtime/Engine/Private/Materials/HLSLMaterialTranslator.h b/Engine/Source/Runtime/Engine/Private/Materials/HLSLMaterialTranslator.h
--- a/Engine/Source/Runtime/Engine/Private/Materials/HLSLMaterialTranslator.h
+++ b/Engine/Source/Runtime/Engine/Private/Materials/HLSLMaterialTranslator.h
@@ -1022,7 +1022,7 @@ protected:
 		int32 BaseColor, int32 EdgeColor, int32 Specular, int32 Metallic,
 		int32 DiffuseAlbedo, int32 F0, int32 F90,
 		int32 Roughness, int32 Anisotropy,
-		int32 SSSProfileId, int32 SSSMFP, int32 SSSMFPScale,
+		int32 SSSProfileId, int32 SSSMFP, int32 SSSMFPScale, int32 SSSPhaseAniso, int32 bUseSSSDiffusion,
 		int32 EmissiveColor, 
 		int32 Haziness, 
 		int32 FuzzAmount, int32 FuzzColor, 
diff --git a/Engine/Source/Runtime/Engine/Private/Materials/MaterialExpressions.cpp b/Engine/Source/Runtime/Engine/Private/Materials/MaterialExpressions.cpp
--- a/Engine/Source/Runtime/Engine/Private/Materials/MaterialExpressions.cpp
+++ b/Engine/Source/Runtime/Engine/Private/Materials/MaterialExpressions.cpp
@@ -21916,6 +21916,7 @@ int32 UMaterialExpressionStrataBSDF::CompilePreview(class FMaterialCompiler* Com
 UMaterialExpressionStrataSlabBSDF::UMaterialExpressionStrataSlabBSDF(const FObjectInitializer& ObjectInitializer)
 	: Super(ObjectInitializer)
 	, bUseMetalness(true)
+	, bUseSSSDiffusion(true)
 {
 	struct FConstructorStatics
 	{
@@ -21983,6 +21984,8 @@ int32 UMaterialExpressionStrataSlabBSDF::Compile(class FMaterialCompiler* Compil
 		SSSProfileCodeChunk != INDEX_NONE ? SSSProfileCodeChunk : Compiler->Constant(0.0f),	
 		CompileWithDefaultFloat3(Compiler, SSSMFP, 0.0f, 0.0f, 0.0f),
 		CompileWithDefaultFloat1(Compiler, SSSMFPScale, 1.0f),
+		CompileWithDefaultFloat1(Compiler, SSSPhaseAnisotropy, 0.0f),
+		bUseSSSDiffusion ? Compiler->Constant(1.0f) : Compiler->Constant(0.0f),
 		CompileWithDefaultFloat3(Compiler, EmissiveColor, 0.0f, 0.0f, 0.0f),
 		CompileWithDefaultFloat1(Compiler, Haziness, 0.0f),
 		CompileWithDefaultFloat1(Compiler, FuzzAmount, 0.0f),
@@ -22019,6 +22022,7 @@ const TArray<FExpressionInput*> UMaterialExpressionStrataSlabBSDF::GetInputs()
 	Result.Add(&Tangent);
 	Result.Add(&SSSMFP);
 	Result.Add(&SSSMFPScale);
+	Result.Add(&SSSPhaseAnisotropy);
 	Result.Add(&EmissiveColor);
 	Result.Add(&Haziness);
 	Result.Add(&Thickness);
@@ -22151,21 +22155,25 @@ uint32 UMaterialExpressionStrataSlabBSDF::GetInputType(int32 InputIndex)
 	}
 	else if (InputIndex == (10 + SkipDisabledInputOffset))
 	{
-		return MCT_Float3; // Emissive Color
+		return MCT_Float1; // SSSPhaseAniso
 	}
 	else if (InputIndex == (11 + SkipDisabledInputOffset))
 	{
-		return MCT_Float1; // Haziness
+		return MCT_Float3; // Emissive Color
 	}
 	else if (InputIndex == (12 + SkipDisabledInputOffset))
 	{
-		return MCT_Float1; // Thickness
+		return MCT_Float1; // Haziness
 	}
 	else if (InputIndex == (13 + SkipDisabledInputOffset))
 	{
-		return MCT_Float1; // FuzzAmount
+		return MCT_Float1; // Thickness
 	}
 	else if (InputIndex == (14 + SkipDisabledInputOffset))
+	{
+		return MCT_Float1; // FuzzAmount
+	}
+	else if (InputIndex == (15 + SkipDisabledInputOffset))
 	{
 		return MCT_Float3; // FuzzColor
 	}
@@ -22239,21 +22247,25 @@ FName UMaterialExpressionStrataSlabBSDF::GetInputName(int32 InputIndex) const
 	}
 	else if (InputIndex == (10 + SkipDisabledInputOffset))
 	{
-		return TEXT("Emissive Color");
+		return TEXT("SSS Phase Anisotropy");
 	}
 	else if (InputIndex == (11 + SkipDisabledInputOffset))
 	{
-		return TEXT("Haziness");
+		return TEXT("Emissive Color");
 	}
 	else if (InputIndex == (12 + SkipDisabledInputOffset))
 	{
-		return TEXT("Thickness");
+		return TEXT("Haziness");
 	}
 	else if (InputIndex == (13 + SkipDisabledInputOffset))
 	{
-		return TEXT("FuzzAmount");
+		return TEXT("Thickness");
 	}
 	else if (InputIndex == (14 + SkipDisabledInputOffset))
+	{
+		return TEXT("FuzzAmount");
+	}
+	else if (InputIndex == (15 + SkipDisabledInputOffset))
 	{
 		return TEXT("FuzzColor");
 	}
@@ -22353,6 +22365,10 @@ void UMaterialExpressionStrataSlabBSDF::GetConnectorToolTip(int32 InputIndex, in
 	{
 		Super::GetConnectorToolTip(18, INDEX_NONE, OutToolTip);
 	}
+	else if (InputIndex == (16 + SkipDisabledInputOffset))
+	{
+		Super::GetConnectorToolTip(19, INDEX_NONE, OutToolTip);
+	}
 }
 
 bool UMaterialExpressionStrataSlabBSDF::IsResultStrataMaterial(int32 OutputIndex)
diff --git a/Engine/Source/Runtime/Engine/Public/MaterialCompiler.h b/Engine/Source/Runtime/Engine/Public/MaterialCompiler.h
--- a/Engine/Source/Runtime/Engine/Public/MaterialCompiler.h
+++ b/Engine/Source/Runtime/Engine/Public/MaterialCompiler.h
@@ -539,7 +539,7 @@ public:
 		int32 BaseColor, int32 EdgeColor, int32 Specular, int32 Metallic,
 		int32 DiffuseAlbedo, int32 F0, int32 F90,
 		int32 Roughness, int32 Anisotropy,
-		int32 SSSProfileId, int32 SSSMFP, int32 SSSMFPScale,
+		int32 SSSProfileId, int32 SSSMFP, int32 SSSMFPScale, int32 SSSPhaseAniso, int32 UseSSSDiffusion,
 		int32 EmissiveColor, 
 		int32 Haziness, 
 		int32 FuzzAmount, int32 FuzzColor,
@@ -1082,7 +1082,7 @@ public:
 		int32 BaseColor, int32 EdgeColor, int32 Specular, int32 Metallic,
 		int32 DiffuseAlbedo, int32 F0, int32 F90,
 		int32 Roughness, int32 Anisotropy,
-		int32 SSSProfileId, int32 SSSMFP, int32 SSSMFPScale,
+		int32 SSSProfileId, int32 SSSMFP, int32 SSSMFPScale, int32 SSSPhaseAniso, int32 UseSSSDiffusion,
 		int32 EmissiveColor, 
 		int32 Haziness, 
 		int32 FuzzAmount, int32 FuzzColor,
@@ -1095,7 +1095,7 @@ public:
 			BaseColor, EdgeColor, Specular, Metallic,
 			DiffuseAlbedo, F0, F90,
 			Roughness, Anisotropy,
-			SSSProfileId, SSSMFP, SSSMFPScale,
+			SSSProfileId, SSSMFP, SSSMFPScale, SSSPhaseAniso, UseSSSDiffusion,
 			EmissiveColor, 
 			Haziness, 
 			FuzzAmount, FuzzColor,
