Fix shader typo.

diff --git a/Engine/Shaders/Private/Strata/Strata.ush b/Engine/Shaders/Private/Strata/Strata.ush
--- a/Engine/Shaders/Private/Strata/Strata.ush
+++ b/Engine/Shaders/Private/Strata/Strata.ush
@@ -3559,7 +3559,7 @@ void UnpackFastPathSlabBSDF(inout FStrataBSDF OutBSDF, uint2 PackedData01)
 	float4 Data1 = UnpackRGBA8(PackedData01.y);
 	SLAB_ROUGHNESS(OutBSDF)			= Data1.y;
 	SLAB_ANISOTROPY(OutBSDF)		= Data1.z * 2.f - 1.f;
-	SLAB_SSSPHASEANISOTROPY(OutBSDF)= Data1.2 * 2.f - 1.f;
+	SLAB_SSSPHASEANISOTROPY(OutBSDF)= Data1.w * 2.f - 1.f;
 }
 
 FStrataBSDF UnpackFastPathStrataBSDFIn(FStrataMaterialContainer StrataBuffer, inout FStrataAddressing StrataAddressing, FStrataPixelHeader StrataPixelHeader)
