Add SSS phase anisotropy debug display.

diff --git a/Engine/Shaders/Private/Strata/StrataVisualize.usf b/Engine/Shaders/Private/Strata/StrataVisualize.usf
--- a/Engine/Shaders/Private/Strata/StrataVisualize.usf
+++ b/Engine/Shaders/Private/Strata/StrataVisualize.usf
@@ -246,6 +246,7 @@ void PrintBSDF(inout FShaderPrintContext Context, inout FStrataAddressing Strata
 			Print(Context, TEXT(" SSS MFP     "), FontBSDFPropName);	Print(Context, SLAB_SSSMFP(BSDF), FontBSDFPropValu);				Newline(Context);
 			}
 
+			Print(Context, TEXT(" SSS Phase   "), FontBSDFPropName);	Print(Context, SLAB_SSSPHASEANISOTROPY(BSDF), FontBSDFPropValu);	Newline(Context);
 			} // SSS
 
 			if (BSDF_GETHASFUZZ(BSDF))
