Path Tracer: expose azimuthal roughness on hair via the scatter (metallic) parameter. The default value of 0 corresponds to 1.0 azimuthal roughness (current behavior) while increasing the scatter amount makes the hair more and more forward scattering, leading to a softer appearance. The lowest possible roughness is limited to 0.1 to limit noise.

diff --git a/Engine/Shaders/Private/PathTracing/Material/PathTracingHair.ush b/Engine/Shaders/Private/PathTracing/Material/PathTracingHair.ush
--- a/Engine/Shaders/Private/PathTracing/Material/PathTracingHair.ush
+++ b/Engine/Shaders/Private/PathTracing/Material/PathTracingHair.ush
@@ -119,6 +119,7 @@ float LogI0(float x)
 float FastASin(float x) {
 	// Maximum error: 10^-3.35
 	// TODO: consider replacing acosFast/asinFast which have approximation error 10^-2.0 (with one less fma, but also with a slightly sub-optimal constant)
+	x = clamp(x, -1.0, 1.0);
 	float a = abs(x);
 	float r = (0.5 * PI) + x * (-0.207034 + 0.0531013 * x);
 	r = (0.5 * PI) - r * sqrt(1 - a);
@@ -239,7 +240,7 @@ FHairData PrepareHairData(float3 V, float3 Color, float LongitudinalRoughness, f
 	float3 SinScale; // sin(ScaleAngle * 2^k) for k=0,1,2
 	float3 CosScale; // cos(ScaleAngle * 2^k) for k=0,1,2
 	SinScale.x = sin(ScaleAngle);
-	CosScale.x = TrigInverse(SinScale.x);
+	CosScale.x = cos(ScaleAngle);
 	SinScale.y = 2 * CosScale.x * SinScale.x;
 	CosScale.y = Pow2(CosScale.x) - Pow2(SinScale.x);
 	SinScale.z = 2 * CosScale.y * SinScale.y;
@@ -305,7 +306,8 @@ FMaterialSample Hair_SampleMaterial(
 
 	float3 V = mul(HairBasis, V_World); // TODO: simplify this since HairBasis is already sort of aligned with V ?
 	float h = 2 * Payload.CustomData.y - 1; // remap back to [-1,1]
-	FHairData HairData = PrepareHairData(V, Payload.BaseColor, Payload.Roughness, 1.0, h, Payload.Specular);
+	float AzimuthalRoughness = lerp(1.0, 0.1, Payload.Metallic);
+	FHairData HairData = PrepareHairData(V, Payload.BaseColor, Payload.Roughness, AzimuthalRoughness, h, Payload.Specular);
 
 	// 1) Decide which lobe to sample
 	float u = RandSample.z * HairData.LobeCdf.w;
@@ -388,7 +390,8 @@ FMaterialEval Hair_EvalMaterial(
 	float3 L = mul(HairBasis, L_World);
 
 	float h = 2 * Payload.CustomData.y - 1; // remap back to [-1,1]
-	FHairData HairData = PrepareHairData(V, Payload.BaseColor, Payload.Roughness, 1.0, h, Payload.Specular);
+	float AzimuthalRoughness = lerp(1.0, 0.1, Payload.Metallic);
+	FHairData HairData = PrepareHairData(V, Payload.BaseColor, Payload.Roughness, AzimuthalRoughness, h, Payload.Specular);
 
 	// precompute data based of lighting direction
 	float SinThetaI = L.z;
