Hair Strands: Remove degenerate segments when preparing bounding boxes for hair ray tracing. Also measure the surface area of the split boxes to ensure it is worthwhile to perform the split. This reduces unnecessary splitting, improving performance in some cases.

This also appears to improve the stability of the hair bvh build step.

diff --git a/Engine/Shaders/Private/HairStrands/HairStrandsRaytracingGeometry.usf b/Engine/Shaders/Private/HairStrands/HairStrandsRaytracingGeometry.usf
--- a/Engine/Shaders/Private/HairStrands/HairStrandsRaytracingGeometry.usf
+++ b/Engine/Shaders/Private/HairStrands/HairStrandsRaytracingGeometry.usf
@@ -149,6 +149,15 @@ void SetBodySegment(uint InVertexIndex)
 #endif
 
 
+float GetAABBHalfAreaForSegment(float4 A, float4 B)
+{
+	float3 BoxLo = min(A.xyz - A.w, B.xyz - B.w);
+	float3 BoxHi = max(A.xyz + A.w, B.xyz + B.w);
+	float3 BoxDiag = BoxHi - BoxLo;
+	return dot(BoxDiag.xyz, BoxDiag.yzx);
+}
+
+
 [numthreads(GROUP_SIZE, 1, 1)]
 void MainCS(uint2 DispatchThreadId : SV_DispatchThreadID)
 {
@@ -196,14 +205,41 @@ void MainCS(uint2 DispatchThreadId : SV_DispatchThreadID)
 		FHairControlPoint P1 = GetControlPoint(SourceIndex1);
 		float4 A = float4(P0.Position, P0.WorldRadius * LodRadiusScale0);
 		float4 B = float4(P1.Position, P1.WorldRadius * LodRadiusScale1);
-		float4 SplitA = A;
-		for (int i = 0; i < RaytracingProceduralSplits; i++)
+		if ((A.w == 0 && B.w == 0) || length2(B.xyz - A.xyz) == 0 || any(!IsFinite(A)) || any(!IsFinite(B)))
 		{
-			float4 SplitB = lerp(A, B, float(i + 1) / float(RaytracingProceduralSplits));
-			float3 Lo = min(SplitA.xyz - SplitA.w, SplitB.xyz - SplitB.w);
-			float3 Hi = max(SplitA.xyz + SplitA.w, SplitB.xyz + SplitB.w);
-			SetAABB(RaytracingProceduralSplits * VertexIndex0 + i, Lo, Hi);
-			SplitA = SplitB;
+			// this segment is degenerate, don't bother adding it
+			SetInvalidSegment(VertexIndex0);
+		}
+		else
+		{
+			// measure aabb surface area prior to splitting
+			float FullArea = GetAABBHalfAreaForSegment(A, B);
+
+			float4 SplitA = A;
+			float SumArea = 0.0;
+			for (int i = 0; i < RaytracingProceduralSplits; i++)
+			{
+				float4 SplitB = lerp(A, B, float(i + 1) / float(RaytracingProceduralSplits));
+				float3 Lo = min(SplitA.xyz - SplitA.w, SplitB.xyz - SplitB.w);
+				float3 Hi = max(SplitA.xyz + SplitA.w, SplitB.xyz + SplitB.w);
+				SumArea += GetAABBHalfAreaForSegment(SplitA, SplitB);
+				SetAABB(RaytracingProceduralSplits * VertexIndex0 + i, Lo, Hi);
+				SplitA = SplitB;
+			}
+#if 1
+			if (FullArea <= 2 * SumArea)
+			{
+				// not worth splitting? use full bbox and make other splits degenerate
+				float3 Lo = min(A.xyz - A.w, B.xyz - B.w);
+				float3 Hi = max(A.xyz + A.w, B.xyz + B.w);
+				SetAABB(RaytracingProceduralSplits * VertexIndex0 + 0, Lo, Hi);
+				for (int i = 1; i < RaytracingProceduralSplits; i++)
+				{
+					const float NaN = asfloat(0x7FC00000);
+					SetAABB(RaytracingProceduralSplits * VertexIndex0 + i, NaN, NaN);
+				}
+			}
+#endif
 		}
 	}
 #else
