Strata - PDO now relies on in shader depth test to avoid UAV write or material data when the depth buffer does not match.
(we cannot rely on LateZ to cull UAV writes)

diff --git a/Engine/Shaders/Private/BasePassPixelShader.usf b/Engine/Shaders/Private/BasePassPixelShader.usf
--- a/Engine/Shaders/Private/BasePassPixelShader.usf
+++ b/Engine/Shaders/Private/BasePassPixelShader.usf
@@ -1720,6 +1720,20 @@ void FPixelShaderInOut_MainPS(
 	GBuffer.PrecomputedShadowFactors = PrecomputedShadowFactors;
 	GBuffer.Velocity = EncodedVelocity;
 
+#if OUTPUT_PIXEL_DEPTH_OFFSET
+	// When in deferred, opaque materials with pixel depth offset must execute a custom depth test in order to avoid conflicting UAV writes.
+	const float PSDeviceZWithOffset  = Out.Depth;
+	const float TexDeviceZWithOffset = OpaqueBasePass.ResolvedSceneDepthTexture.Load(int3(In.SvPosition.xy, 0)).r;
+
+	uint PSDeviceZWithOffsetUINT  = uint(PSDeviceZWithOffset * 16777215.0f + 0.5f);	// 16777215 = 2^24 - 1
+	uint TexDeviceZWithOffsetUINT = uint(TexDeviceZWithOffset * 16777215.0f + 0.5f);
+	
+	const bool DepthTest_Equal = OpaqueBasePass.Is24BitUnormDepthStencil ?
+		(PSDeviceZWithOffsetUINT == TexDeviceZWithOffsetUINT || PSDeviceZWithOffsetUINT == (TexDeviceZWithOffsetUINT - 1)) :	// 24 bit unorm, need to test a bit more, likely due to 24unorm conversion to float when loading form the texture.
+		(PSDeviceZWithOffset == TexDeviceZWithOffset); 																			// 32 bit float, 1 to 1 mapping.
+
+	if (DepthTest_Equal)
+#endif
 	{
 		float3 EmissiveLuminance = 0.0f;
 		uint2 PixelPos = uint2(In.SvPosition.xy);
@@ -2129,6 +2143,19 @@ void FPixelShaderInOut_MainPS(
 #endif
 }
 
+#if STRATA_OPAQUE_DEFERRED
+
+// Strata opaque deferred materials require early depth stencil test to avoid writing through UAV from the pixel shader when the depth surface is not from the current pixel shader material
+// Notes: 
+//		- Pixel depth offset prevents the usage of early depth test because the depth is generated in shader. So we can only relay on LateZ in this case.
+//		- DEPTHSTENCIL_EARLYTEST_LATEWRITE does not help in this case due to UAV writes we need to skip.
+//		- In order to avoid UAV writes, we will run in shader manual depth tests.
+#if !OUTPUT_PIXEL_DEPTH_OFFSET
+#define PIXELSHADER_EARLYDEPTHSTENCIL EARLYDEPTHSTENCIL
+#endif
+
+#else // STRATA_OPAQUE_DEFERRED
+
 // If virtual texture is enabled then use early depth test so that UAV feedback buffer writes respect the depth test
 #if NUM_VIRTUALTEXTURE_SAMPLES || LIGHTMAP_VT_ENABLED || STRATA_OPAQUE_DEFERRED || OIT_ENABLED
 	#if COMPILER_SUPPORTS_DEPTHSTENCIL_EARLYTEST_LATEWRITE
@@ -2141,6 +2168,8 @@ void FPixelShaderInOut_MainPS(
 	#endif
 #endif
 
+#endif // STRATA_OPAQUE_DEFERRED
+
 
 // the following needs to match to the code in FSceneTextures::GetGBufferRenderTargets()
 #define PIXELSHADEROUTPUT_BASEPASS 1
diff --git a/Engine/Shaders/Private/Strata/Strata.ush b/Engine/Shaders/Private/Strata/Strata.ush
--- a/Engine/Shaders/Private/Strata/Strata.ush
+++ b/Engine/Shaders/Private/Strata/Strata.ush
@@ -3107,7 +3107,7 @@ void PackStrataOut(
 
 		const uint RootOperatorIndex = Strata.OperatorIndex;
 
-	#if STRATA_ADVANCED_DEBUG_ENABLED
+	#if STRATA_ADVANCED_DEBUG_ENABLED && STRATA_USES_RESOURCE_MATERIALCONTAINER
 		if (Settings.SliceStoringDebugStrataTree > -1)
 		{
 			StrataPackOutStrataTree(ExtraMaterialDataUAV, RootOperatorIndex, StrataTree, Settings.SliceStoringDebugStrataTree);
diff --git a/Engine/Source/Runtime/Apple/AGXRHI/Private/AGXRHI.cpp b/Engine/Source/Runtime/Apple/AGXRHI/Private/AGXRHI.cpp
--- a/Engine/Source/Runtime/Apple/AGXRHI/Private/AGXRHI.cpp
+++ b/Engine/Source/Runtime/Apple/AGXRHI/Private/AGXRHI.cpp
@@ -804,10 +804,10 @@ FAGXDynamicRHI::FAGXDynamicRHI(ERHIFeatureLevel::Type RequestedFeatureLevel)
 		GPixelFormats[PF_FloatR11G11B10		].Supported			= true;
 	}
 	
-		GPixelFormats[PF_DepthStencil		].PlatformFormat	= (uint32)MTLPixelFormatDepth32Float_Stencil8;
-		GPixelFormats[PF_DepthStencil		].BlockBytes		= 4;
-
+	GPixelFormats[PF_DepthStencil		].PlatformFormat	= (uint32)MTLPixelFormatDepth32Float_Stencil8;
+	GPixelFormats[PF_DepthStencil		].BlockBytes		= 4;
 	GPixelFormats[PF_DepthStencil		].Supported			= true;
+
 	GPixelFormats[PF_ShadowDepth		].PlatformFormat	= (uint32)MTLPixelFormatDepth32Float;
 	GPixelFormats[PF_ShadowDepth		].BlockBytes		= 4;
 	GPixelFormats[PF_ShadowDepth		].Supported			= true;
@@ -842,6 +842,7 @@ FAGXDynamicRHI::FAGXDynamicRHI(ERHIFeatureLevel::Type RequestedFeatureLevel)
 	if(bSupportsD24S8)
 	{
 		GPixelFormats[PF_DepthStencil	].PlatformFormat	= (uint32)MTLPixelFormatDepth24Unorm_Stencil8;
+		GPixelFormats[PF_DepthStencil].bIs24BitUnormDepthStencil = true;
 	}
 	else
 	{
diff --git a/Engine/Source/Runtime/Apple/MetalRHI/Private/MetalRHI.cpp b/Engine/Source/Runtime/Apple/MetalRHI/Private/MetalRHI.cpp
--- a/Engine/Source/Runtime/Apple/MetalRHI/Private/MetalRHI.cpp
+++ b/Engine/Source/Runtime/Apple/MetalRHI/Private/MetalRHI.cpp
@@ -853,6 +853,7 @@ FMetalDynamicRHI::FMetalDynamicRHI(ERHIFeatureLevel::Type RequestedFeatureLevel)
 	if(bSupportsD24S8)
 	{
 		GPixelFormats[PF_DepthStencil	].PlatformFormat	= (uint32)mtlpp::PixelFormat::Depth24Unorm_Stencil8;
+		GPixelFormats[PF_DepthStencil	].bIs24BitUnormDepthStencil = true;
 	}
 	else
 	{
diff --git a/Engine/Source/Runtime/D3D12RHI/Private/D3D12RHI.cpp b/Engine/Source/Runtime/D3D12RHI/Private/D3D12RHI.cpp
--- a/Engine/Source/Runtime/D3D12RHI/Private/D3D12RHI.cpp
+++ b/Engine/Source/Runtime/D3D12RHI/Private/D3D12RHI.cpp
@@ -108,6 +108,7 @@ FD3D12DynamicRHI::FD3D12DynamicRHI(const TArray<TSharedPtr<FD3D12Adapter>>& Chos
 		GPixelFormats[PF_DepthStencil].PlatformFormat = DXGI_FORMAT_R24G8_TYPELESS;
 		GPixelFormats[PF_DepthStencil].BlockBytes = 4;
 		GPixelFormats[PF_DepthStencil].Supported = true;
+		GPixelFormats[PF_DepthStencil].bIs24BitUnormDepthStencil = true;
 		GPixelFormats[PF_X24_G8].PlatformFormat = DXGI_FORMAT_X24_TYPELESS_G8_UINT;
 		GPixelFormats[PF_X24_G8].BlockBytes = 4;
 	}
diff --git a/Engine/Source/Runtime/RHI/Private/RHI.cpp b/Engine/Source/Runtime/RHI/Private/RHI.cpp
--- a/Engine/Source/Runtime/RHI/Private/RHI.cpp
+++ b/Engine/Source/Runtime/RHI/Private/RHI.cpp
@@ -2631,6 +2631,7 @@ FPixelFormatInfo::FPixelFormatInfo(
 	, BlockBytes(InBlockBytes)
 	, NumComponents(InNumComponents)
 	, Supported(InSupported)
+	, bIs24BitUnormDepthStencil(true)
 {
 }
 
diff --git a/Engine/Source/Runtime/RHI/Public/RHI.h b/Engine/Source/Runtime/RHI/Public/RHI.h
--- a/Engine/Source/Runtime/RHI/Public/RHI.h
+++ b/Engine/Source/Runtime/RHI/Public/RHI.h
@@ -860,7 +860,8 @@ struct FPixelFormatInfo
 	uint32						PlatformFormat{ 0 };
 
 	/** Whether the texture format is supported on the current platform/ rendering combination	*/
-	bool						Supported;
+	uint8						Supported : 1;
+	uint8						bIs24BitUnormDepthStencil : 1;	// If false, 32 bit float is assumed
 };
 
 extern RHI_API FPixelFormatInfo GPixelFormats[PF_MAX];		// Maps members of EPixelFormat to a FPixelFormatInfo describing the format.
diff --git a/Engine/Source/Runtime/Renderer/Private/BasePassRendering.cpp b/Engine/Source/Runtime/Renderer/Private/BasePassRendering.cpp
--- a/Engine/Source/Runtime/Renderer/Private/BasePassRendering.cpp
+++ b/Engine/Source/Runtime/Renderer/Private/BasePassRendering.cpp
@@ -619,6 +619,7 @@ TRDGUniformBufferRef<FOpaqueBasePassUniformParameters> CreateOpaqueBasePassUnifo
 		{
 			BasePassParameters.ResolvedSceneDepthTexture = ForwardBasePassTextures.SceneDepthIfResolved;
 		}
+		BasePassParameters.Is24BitUnormDepthStencil = ForwardBasePassTextures.bIs24BitUnormDepthStencil ? 1 : 0;
 	}
 
 	// DBuffer Decals
@@ -907,6 +908,12 @@ void FDeferredShadingSceneRenderer::RenderBasePass(
 		ForwardBasePassTextures.ScreenSpaceAO = SceneTextures.ScreenSpaceAO;
 		ForwardBasePassTextures.ScreenSpaceShadowMask = ForwardShadowMaskTexture;
 	}
+	else if (!ExclusiveDepthStencil.IsDepthWrite())
+	{
+		// If depth write is not enabled, we can bound the depth texture as read only
+		ForwardBasePassTextures.SceneDepthIfResolved = SceneTextures.Depth.Resolve;
+	}
+	ForwardBasePassTextures.bIs24BitUnormDepthStencil = ForwardBasePassTextures.SceneDepthIfResolved ? GPixelFormats[ForwardBasePassTextures.SceneDepthIfResolved->Desc.Format].bIs24BitUnormDepthStencil : 1;
 
 	GraphBuilder.SetCommandListStat(GET_STATID(STAT_CLM_BasePass));
 	RenderBasePassInternal(GraphBuilder, SceneTextures, BasePassRenderTargets, BasePassDepthStencilAccess, ForwardBasePassTextures, DBufferTextures, bDoParallelBasePass, bRenderLightmapDensity, InstanceCullingManager, bNaniteEnabled, NaniteRasterResults);
diff --git a/Engine/Source/Runtime/Renderer/Private/BasePassRendering.h b/Engine/Source/Runtime/Renderer/Private/BasePassRendering.h
--- a/Engine/Source/Runtime/Renderer/Private/BasePassRendering.h
+++ b/Engine/Source/Runtime/Renderer/Private/BasePassRendering.h
@@ -60,6 +60,7 @@ struct FForwardBasePassTextures
 	FRDGTextureRef ScreenSpaceAO = nullptr;
 	FRDGTextureRef ScreenSpaceShadowMask = nullptr;
 	FRDGTextureRef SceneDepthIfResolved = nullptr;
+	bool bIs24BitUnormDepthStencil = false;
 };
 
 BEGIN_GLOBAL_SHADER_PARAMETER_STRUCT(FSharedBasePassUniformParameters,)
@@ -91,6 +92,7 @@ BEGIN_GLOBAL_SHADER_PARAMETER_STRUCT(FOpaqueBasePassUniformParameters,)
 	SHADER_PARAMETER_SAMPLER(SamplerState, PreIntegratedGFSampler)
 	SHADER_PARAMETER(FVector4f, SceneWithoutSingleLayerWaterMinMaxUV)
 	SHADER_PARAMETER(FVector4f, DistortionParams)
+	SHADER_PARAMETER(int32, Is24BitUnormDepthStencil)
 	// Misc
 	SHADER_PARAMETER_RDG_TEXTURE(Texture2D, EyeAdaptationTexture)
 END_GLOBAL_SHADER_PARAMETER_STRUCT()
diff --git a/Engine/Source/Runtime/VulkanRHI/Private/VulkanDevice.cpp b/Engine/Source/Runtime/VulkanRHI/Private/VulkanDevice.cpp
--- a/Engine/Source/Runtime/VulkanRHI/Private/VulkanDevice.cpp
+++ b/Engine/Source/Runtime/VulkanRHI/Private/VulkanDevice.cpp
@@ -615,6 +615,7 @@ void FVulkanDevice::SetupFormats()
 		// prefer VK_FORMAT_D24_UNORM_S8_UINT
 		MapFormatSupport(PF_DepthStencil, { VK_FORMAT_D24_UNORM_S8_UINT, VK_FORMAT_D32_SFLOAT_S8_UINT, VK_FORMAT_D16_UNORM_S8_UINT }, ComponentMappingRIII);
 		MapFormatSupport(PF_X24_G8, { VK_FORMAT_D24_UNORM_S8_UINT, VK_FORMAT_D32_SFLOAT_S8_UINT, VK_FORMAT_D16_UNORM_S8_UINT }, ComponentMappingRR01);
+		GPixelFormats[PF_DepthStencil].bIs24BitUnormDepthStencil = true;
 	}
 	else
 	{
diff --git a/Engine/Source/Runtime/Windows/D3D11RHI/Private/D3D11Device.cpp b/Engine/Source/Runtime/Windows/D3D11RHI/Private/D3D11Device.cpp
--- a/Engine/Source/Runtime/Windows/D3D11RHI/Private/D3D11Device.cpp
+++ b/Engine/Source/Runtime/Windows/D3D11RHI/Private/D3D11Device.cpp
@@ -121,6 +121,7 @@ FD3D11DynamicRHI::FD3D11DynamicRHI(IDXGIFactory1* InDXGIFactory1, D3D_FEATURE_LE
 	{
 		GPixelFormats[PF_DepthStencil].PlatformFormat = DXGI_FORMAT_R24G8_TYPELESS;
 		GPixelFormats[PF_DepthStencil].BlockBytes = 4;
+		GPixelFormats[PF_DepthStencil].bIs24BitUnormDepthStencil = true;
 		GPixelFormats[PF_X24_G8].PlatformFormat = DXGI_FORMAT_X24_TYPELESS_G8_UINT;
 		GPixelFormats[PF_X24_G8].BlockBytes = 4;
 	}
