Fix HitProxyVertexShader using the wrong offset parameter in to VertexFetch_HitProxyIdBuffer

diff --git a/Engine/Shaders/Private/GeometryCollectionVertexFactory.ush b/Engine/Shaders/Private/GeometryCollectionVertexFactory.ush
--- a/Engine/Shaders/Private/GeometryCollectionVertexFactory.ush
+++ b/Engine/Shaders/Private/GeometryCollectionVertexFactory.ush
@@ -23,6 +23,12 @@
 
 	Buffer<float4> VertexFetch_BoneOriginBuffer;
 	Buffer<float4> VertexFetch_BoneLightmapBuffer;
+
+	#define VF_REQUIRES_HITPROXY_INDIRECTION 1
+	int VertexFactoryGetVertexFetchParameter(int ParameterIndex)
+	{
+		return GeometryCollectionVF.VertexFetch_Parameters[ParameterIndex];
+	}
 #endif
 
 // Since we don't have template functions it will have to be this way.
diff --git a/Engine/Shaders/Private/HitProxyVertexShader.usf b/Engine/Shaders/Private/HitProxyVertexShader.usf
--- a/Engine/Shaders/Private/HitProxyVertexShader.usf
+++ b/Engine/Shaders/Private/HitProxyVertexShader.usf
@@ -18,6 +18,17 @@
 Buffer<float4> VertexFetch_HitProxyIdBuffer;
 #endif
 
+#ifndef VF_REQUIRES_HITPROXY_INDIRECTION 
+#define VF_REQUIRES_HITPROXY_INDIRECTION 0
+#endif
+
+#if VF_REQUIRES_HITPROXY_INDIRECTION == 0
+	int VertexFactoryGetVertexFetchParameter(int ParameterIndex)
+	{
+		return 0;
+	}
+#endif
+
 struct FHitProxyVSToPS
 {
 	FVertexFactoryInterpolantsVSToPS FactoryInterpolants;
@@ -55,7 +66,7 @@ void Main(
 
 #if USE_PER_VERTEX_HITPROXY_ID
 	#if MANUAL_VERTEX_FETCH
-		Output.InstanceHitProxyId = VertexFetch_HitProxyIdBuffer[(VF_VertexOffset + Input.VertexId)] FMANUALFETCH_COLOR_COMPONENT_SWIZZLE;
+		Output.InstanceHitProxyId = VertexFetch_HitProxyIdBuffer[(VertexFactoryGetVertexFetchParameter(VF_VertexOffset) + Input.VertexId)] FMANUALFETCH_COLOR_COMPONENT_SWIZZLE;
 	#else
 		Output.InstanceHitProxyId = float4(0,0,0,0);
 	#endif
diff --git a/Engine/Shaders/Private/LocalVertexFactory.ush b/Engine/Shaders/Private/LocalVertexFactory.ush
--- a/Engine/Shaders/Private/LocalVertexFactory.ush
+++ b/Engine/Shaders/Private/LocalVertexFactory.ush
@@ -76,8 +76,16 @@ float4 SplineParams[10];
 	Buffer<float4> VertexFetch_InstanceOriginBuffer;
 	Buffer<float4> VertexFetch_InstanceTransformBuffer;
 	Buffer<float4> VertexFetch_InstanceLightmapBuffer;
+
+	#define VF_REQUIRES_HITPROXY_INDIRECTION 1
+	int VertexFactoryGetVertexFetchParameter(int ParameterIndex)
+	{
+		return LocalVF.VertexFetch_Parameters[ParameterIndex];
+	}
 #endif //! MANUAL_VERTEX_FETCH
 
+#define VF_REQUIRES_HITPROXY_INDIRECTION 1
+
 /**
  * Per-vertex inputs from bound vertex buffers
  */
