Added to mini font PrintSmallUint

diff --git a/Engine/Shaders/Private/MiniFontCommon.ush b/Engine/Shaders/Private/MiniFontCommon.ush
--- a/Engine/Shaders/Private/MiniFontCommon.ush
+++ b/Engine/Shaders/Private/MiniFontCommon.ush
@@ -167,3 +167,17 @@ void PrintSmallFloat(int2 PixelPos, inout float3 OutColor, float3 FontColor, int
 	// after period
 	PrintCharacter(PixelPos, OutColor, FontColor, Cursor, ExtractDigitFromFloat(Number, 0.1));
 }
+
+void PrintSmallUint(int2 PixelPos, inout float3 OutColor, float3 FontColor, int2 LeftTop, uint Number)
+{
+	int2 Cursor = LeftTop;
+	float FloatNumber = float(Number) + 0.05;
+
+	// before period
+	FLATTEN
+	if (Number >= 10)
+	{
+		PrintCharacter(PixelPos, OutColor, FontColor, Cursor, ExtractDigitFromFloat(FloatNumber, 10));
+	}
+	PrintCharacter(PixelPos, OutColor, FontColor, Cursor, ExtractDigitFromFloat(FloatNumber, 1));
+}
