Fixed height fog looking badin deferred when viewed over large distance.
EngineTest with height fog ran succesfully.

diff --git a/Engine/Shaders/Private/HeightFogPixelShader.usf b/Engine/Shaders/Private/HeightFogPixelShader.usf
--- a/Engine/Shaders/Private/HeightFogPixelShader.usf
+++ b/Engine/Shaders/Private/HeightFogPixelShader.usf
@@ -52,7 +52,13 @@ void ExponentialPixelMain(
 
 	bool bIsRendered = (DeviceZ != 0.0);
 
+#if 1
+	// Using SvPositionToTranslatedWorld is more stable and accurate as compared to the ScreenVector over large distance (SkyAtmosphere like).
+	// However, we do not to clamp DeviceZ to a reasinable FarZ in order to avoid Inf values.
+	float3 WorldPositionRelativeToCamera = SvPositionToTranslatedWorld(float4(SVPos.xy, max(1.5e-10, DeviceZ), 1.0));
+#else
 	float3 WorldPositionRelativeToCamera = ScreenVector.xyz * SceneDepth;
+#endif
 
 	float ZSlice = log2(SceneDepth * View.VolumetricFogGridZParams.x + View.VolumetricFogGridZParams.y) * View.VolumetricFogGridZParams.z * View.VolumetricFogInvGridSize.z;
 	
