Temporary shader compilation fix.

diff --git a/Engine/Shaders/Private/LightGridInjection.usf b/Engine/Shaders/Private/LightGridInjection.usf
--- a/Engine/Shaders/Private/LightGridInjection.usf
+++ b/Engine/Shaders/Private/LightGridInjection.usf
@@ -434,6 +434,10 @@ void DebugLightGridPS(
 	in float4 SVPos : SV_POSITION,
 	out float4 OutLuminanceTransmittance : SV_Target0)
 {
+	OutLuminanceTransmittance = float4(0, 0, 0, 1);
+	
+	// Temp fix
+#if !COMPILER_SWITCH
 	ResolvedView = ResolveView();
 
 	const uint EyeIndex = 0;
@@ -445,8 +449,6 @@ void DebugLightGridPS(
 	const uint2 LocalPixelPosDynRes = float2(LocalPixelPos) * View.ViewResolutionFraction;
 	const uint2 LocalPixelPosDynRes2 = float2(LocalPixelPos+1) * View.ViewResolutionFraction;
 
-	OutLuminanceTransmittance = float4(0, 0, 0, 1);
-
 	if (any(LocalPixelPosDynRes >= uint2(ResolvedView.ViewSizeAndInvSize.xy)))
 	{
 		return;
@@ -505,6 +507,7 @@ void DebugLightGridPS(
 	{
 		OutLuminanceTransmittance.rgba = float4(0.25, 0.25, 0.25, 0.0);
 	}
+#endif
 }
 
 #endif // SHADER_DEBUG_LIGHT_GRID_PS
