Fix shader compilation error.

diff --git a/Engine/Shaders/Private/LightGridInjection.usf b/Engine/Shaders/Private/LightGridInjection.usf
--- a/Engine/Shaders/Private/LightGridInjection.usf
+++ b/Engine/Shaders/Private/LightGridInjection.usf
@@ -421,9 +421,7 @@ void LightGridCompactCS(
 #ifdef SHADER_DEBUG_LIGHT_GRID_PS
 
 #define SUPPORT_CONTACT_SHADOWS 0
-#include "DeferredLightingCommon.ush"
 #include "LightGridCommon.ush"
-#include "/Engine/Private/MiniFontCommon.ush"
 #include "/Engine/Private/ShaderPrint.ush"
 #include "ColorMap.ush"
 
@@ -435,9 +433,7 @@ void DebugLightGridPS(
 	out float4 OutLuminanceTransmittance : SV_Target0)
 {
 	OutLuminanceTransmittance = float4(0, 0, 0, 1);
-	
-	// Temp fix
-#if !COMPILER_SWITCH
+
 	ResolvedView = ResolveView();
 
 	const uint EyeIndex = 0;
@@ -507,7 +503,6 @@ void DebugLightGridPS(
 	{
 		OutLuminanceTransmittance.rgba = float4(0.25, 0.25, 0.25, 0.0);
 	}
-#endif
 }
 
 #endif // SHADER_DEBUG_LIGHT_GRID_PS
