Fix shader compile error
#rnx

diff --git a/Engine/Shaders/Private/HairStrands/HairStrandsVelocity.usf b/Engine/Shaders/Private/HairStrands/HairStrandsVelocity.usf
--- a/Engine/Shaders/Private/HairStrands/HairStrandsVelocity.usf
+++ b/Engine/Shaders/Private/HairStrands/HairStrandsVelocity.usf
@@ -111,7 +111,7 @@ void MainCS(uint3 DispatchThreadId : SV_DispatchThreadID, uint2 InGroupId : SV_G
 				#endif
 
 				#if PERMUTATION_VELOCITY == VELOCITY_TYPE_MAX
-				const float4 CurrentVelocity = DecodeVelocityFromTexture(EncodedVelocity);
+				const float3 CurrentVelocity = DecodeVelocityFromTexture(EncodedVelocity);
 				const float CurrentVelocityMagnitude2 = dot(CurrentVelocity, CurrentVelocity);
 				if (CurrentVelocityMagnitude2 > MaxVelocityMagnitude2)
 				{
