* Add screen aligned quad to shader print.
* Expose triangle count request

diff --git a/Engine/Shaders/Private/ShaderPrint.ush b/Engine/Shaders/Private/ShaderPrint.ush
--- a/Engine/Shaders/Private/ShaderPrint.ush
+++ b/Engine/Shaders/Private/ShaderPrint.ush
@@ -895,6 +895,41 @@
 void AddFilledQuadWS (FShaderPrintContext Ctx, float3 Pos0, float3 Pos1, float3 Pos2, float3 Pos3, float4 Color) { if (Ctx.bIsActive) { AddFilledQuadTWS(Pos0 + ShaderPrintData.TranslatedWorldOffset, Pos1 + ShaderPrintData.TranslatedWorldOffset, Pos2 + ShaderPrintData.TranslatedWorldOffset, Pos3 + ShaderPrintData.TranslatedWorldOffset, Color); } }
 
 ///////////////////////////////////////////////////////////////////////////////////////////////////
+// Filled Quad (Screen-Space)
+
+void AddFilledQuadSS(float2 InPosInPixel0, float2 InPosInPixel2, float4 Color)
+{
+	const float2 InPosInPixel1 = float2(InPosInPixel2.x, InPosInPixel0.y);
+	const float2 InPosInPixel3 = float2(InPosInPixel0.x, InPosInPixel2.y);
+
+	const float Depth = 0.5f;
+	const float3 Pos0 = float3((InPosInPixel0) / float2(ShaderPrintData.Resolution), Depth);
+	const float3 Pos1 = float3((InPosInPixel1) / float2(ShaderPrintData.Resolution), Depth);
+	const float3 Pos2 = float3((InPosInPixel2) / float2(ShaderPrintData.Resolution), Depth);
+	const float3 Pos3 = float3((InPosInPixel3) / float2(ShaderPrintData.Resolution), Depth);
+
+	uint Offset = 0;
+	if (AllocateTriangleElement(2, Offset))
+	{
+		FTriangleElement Element;
+		Element.bIsScreenSpace = true;
+		Element.Color = Color;
+
+		Element.Pos0 = Pos0;
+		Element.Pos1 = Pos1;
+		Element.Pos2 = Pos2;
+		AddTriangleElement(Element, Offset + 0);
+
+		Element.Pos0 = Pos0;
+		Element.Pos1 = Pos2;
+		Element.Pos2 = Pos3;
+		AddTriangleElement(Element, Offset + 1);
+	}
+}
+
+void AddFilledQuadSS(FShaderPrintContext Ctx, float2 Pos0, float2 Pos1, float4 Color) { if (Ctx.bIsActive) { AddFilledQuadSS(Pos0, Pos1, Color); } }
+
+///////////////////////////////////////////////////////////////////////////////////////////////////
 // Line-based Primitives (Line, AABB, OBB, Referentials, ...)
 
 struct FLineElement 
diff --git a/Engine/Source/Runtime/Renderer/Public/ShaderPrintParameters.h b/Engine/Source/Runtime/Renderer/Public/ShaderPrintParameters.h
--- a/Engine/Source/Runtime/Renderer/Public/ShaderPrintParameters.h
+++ b/Engine/Source/Runtime/Renderer/Public/ShaderPrintParameters.h
@@ -71,6 +71,7 @@
 	 */
 	RENDERER_API void RequestSpaceForCharacters(uint32 MaxElementCount);
 	RENDERER_API void RequestSpaceForLines(uint32 MaxElementCount);
+	RENDERER_API void RequestSpaceForTriangles(uint32 MaxElementCount);
 
 	/** Structure containing setup for shader print capturing. */
 	struct RENDERER_API FShaderPrintSetup
