Improve texel dither to be more continuous and avoid artifacts at the edges of pages (which are now dynamically picked)
Add a world dither. Not necessarily useful but can be somewhat analygous to PCF radius.
Add permutation to SMRT samples = 0 (and possibly more in the future).

diff --git a/Engine/Shaders/Private/VirtualShadowMaps/VirtualShadowMapPageManagement.usf b/Engine/Shaders/Private/VirtualShadowMaps/VirtualShadowMapPageManagement.usf
--- a/Engine/Shaders/Private/VirtualShadowMaps/VirtualShadowMapPageManagement.usf
+++ b/Engine/Shaders/Private/VirtualShadowMaps/VirtualShadowMapPageManagement.usf
@@ -428,7 +428,7 @@
 				PrevPageAddress += OffsetScale * (PrevProjection.ClipmapCornerOffset - Projection.ClipmapCornerOffset);
 			}
 
-			if (all(PrevPageAddress >= 0) && all(PrevPageAddress < (VSM_LEVEL0_DIM_PAGES_XY >> MipLevel)))
+			if (all(PrevPageAddress >= 0) && all(PrevPageAddress < (uint(VSM_LEVEL0_DIM_PAGES_XY) >> MipLevel)))
 			{
 				uint PrevGlobalOffset = CalcPageOffset(PrevVirtualShadowMapId, MipLevel, uint2(PrevPageAddress));
 
@@ -1126,7 +1126,8 @@
 [numthreads(TILE_THREAD_GROUP_SIZE_XY, TILE_THREAD_GROUP_SIZE_XY, 1)]
 void BuildHZBPerPageCS(uint GroupThreadIndex : SV_GroupIndex, uint GroupIndex : SV_GroupID)
 {
-	uint2 SrcTileOffset = GetTileOffset(GroupIndex, PhysicalPagesForHzb);
+	// We don't care about the array index here (not doing a static/dynamic split - HZB only for static)
+	uint2 SrcTileOffset = GetTileOffset(GroupIndex, PhysicalPagesForHzb).xy;
 
 	uint2 RemappedGroupThreadIndex = InitialTilePixelPositionForReduction2x2(LOG2_TILE_SIZE_XY - 1U, GroupThreadIndex);
 
diff --git a/Engine/Shaders/Private/VirtualShadowMaps/VirtualShadowMapProjection.usf b/Engine/Shaders/Private/VirtualShadowMaps/VirtualShadowMapProjection.usf
--- a/Engine/Shaders/Private/VirtualShadowMaps/VirtualShadowMapProjection.usf
+++ b/Engine/Shaders/Private/VirtualShadowMaps/VirtualShadowMapProjection.usf
@@ -151,11 +151,12 @@
 	return RayLength;
 }
 
-float3 GetEstimatedGeoWorldNormal(float3 TranslatedWorldPosition, FProjectionShadingInfo ShadingInfo)
+// NOTE: Not currently used by occasionally useful for testing
+float3 GetEstimatedGeoWorldNormal(float3 TranslatedWorldPosition, float3 ShadingNormal)
 {
 	// Figure out slope, we do world space since that is the space where we might override using the shading normal...
-	float3 TranslatedWorldPositionDDX = DDX(TranslatedWorldPosition);
-	float3 TranslatedWorldPositionDDY = DDY(TranslatedWorldPosition);
+	float3 TranslatedWorldPositionDDX = ddx_fine(TranslatedWorldPosition);
+	float3 TranslatedWorldPositionDDY = ddy_fine(TranslatedWorldPosition);
 	float3 EstimatedGeoWorldNormal = cross(TranslatedWorldPositionDDX, TranslatedWorldPositionDDY);
 
 	// Handle NaNs; will cause it to revert to shading normal below
@@ -166,9 +167,9 @@
 	// NOTE: Gbuffer normal is not the surface normal for hair; hair lighting takes a different path but possibly
 	// necessary to do some sort of special casing here (disable normal offset and biasing entirely?).
 	// If the estimated geo normal is too far out we assume it's broken (derivative includes other surfaces or background) and fall back to the shading normal
-	if (dot(ShadingInfo.WorldNormal, EstimatedGeoWorldNormal) < 0.25f)
+	if (dot(ShadingNormal, EstimatedGeoWorldNormal) < 0.25f)
 	{
-		EstimatedGeoWorldNormal = ShadingInfo.WorldNormal;
+		EstimatedGeoWorldNormal = ShadingNormal;
 	}
 #endif
 
@@ -210,6 +211,8 @@
 float SMRTRayLengthScale;			// Directional lights
 float SMRTCotMaxRayAngleFromLight;	// Spot/point lights
 float SMRTTexelDitherScale;			// Currently only directional lights
+float SMRTWorldDitherScale;			// Currently only directional lights
+uint bSMRTUseAdaptiveRayCount;
 
 // One pass projection
 RWTexture2D< uint4 > OutShadowMaskBits;
@@ -268,9 +271,9 @@
 	BRANCH
 	if (bInLightRegion && bValidPixel)
 	{
-		const bool bBackfaceCull = (bCullBackfacingPixels != 0) && !bIsHairInput && !ShadingInfo.bIsSubsurface;
-
-		const float3 WorldNormal = (bIsHairInput || ShadingInfo.bIsHair) ? L : ShadingInfo.WorldNormal;
+		const bool bBackfaceCull = (bCullBackfacingPixels > 0) && !bIsHairInput && !ShadingInfo.bIsSubsurface;
+		
+		float3 WorldNormal = (bIsHairInput || ShadingInfo.bIsHair) ? L : ShadingInfo.WorldNormal;
 		TranslatedWorldPosition += WorldNormal * NormalBiasLength;
 
 		// Do not run contact shadow when computing the hair shadow mask (i.e. shadow mask applied on hair, has the scene 
@@ -292,6 +295,8 @@
 
 		if (SMRTRayCount > 0)
 		{
+			const bool bUseAdaptiveRayCount = bSMRTUseAdaptiveRayCount > 0;
+
 		#if DIRECTIONAL_LIGHT
 			Result = TraceDirectional(
 				VirtualShadowMapId,
@@ -306,7 +311,9 @@
 				Noise,
 				bBackfaceCull,
 				WorldNormal,
-				SMRTTexelDitherScale);
+				SMRTTexelDitherScale,
+				SMRTWorldDitherScale,
+				bUseAdaptiveRayCount);
 		#else
 			Result = TraceLocalLight(
 				VirtualShadowMapId,
@@ -320,7 +327,8 @@
 				SMRTCotMaxRayAngleFromLight,
 				Noise,
 				bBackfaceCull,
-				WorldNormal);
+				WorldNormal,
+				bUseAdaptiveRayCount);
 		#endif // DIRECTIONAL_LIGHT
 		}
 		else
diff --git a/Engine/Shaders/Private/VirtualShadowMaps/VirtualShadowMapProjectionDirectional.ush b/Engine/Shaders/Private/VirtualShadowMaps/VirtualShadowMapProjectionDirectional.ush
--- a/Engine/Shaders/Private/VirtualShadowMaps/VirtualShadowMapProjectionDirectional.ush
+++ b/Engine/Shaders/Private/VirtualShadowMaps/VirtualShadowMapProjectionDirectional.ush
@@ -14,9 +14,8 @@
 #include "VirtualShadowMapSMRTCommon.ush"
 #include "../Visualization.ush"
 
-float ComputeOptimalSlopeBiasDirectional(
+float2 ComputeDepthSlopeDirectionalUV(
 	FVirtualShadowMapProjectionShaderData ProjectionData,
-	float2 TexelOffset,
 	float3 EstimatedGeoWorldNormal,
 	bool bClamp = true)
 {
@@ -24,16 +23,48 @@
 	float2 DepthSlopeUV = -NormalPlaneUV.xy / NormalPlaneUV.z;
 
 	// Clamp to avoid excessive degenerate slope biases causing flickering lit pixels
-	float2 Clamp = 0.1f;
+	float2 Clamp = 0.02f;
 	DepthSlopeUV = bClamp ? clamp(DepthSlopeUV, -Clamp, Clamp) : DepthSlopeUV;
 
-	float2 OffsetUV = TexelOffset / float(CalcLevelDimsTexels(0));
-	float OptimalSlopeBias = 2.0f * max(0.0f, dot(DepthSlopeUV, OffsetUV));
-
-	return OptimalSlopeBias;
+	return DepthSlopeUV;
 }
 
+float ComputeOptimalSlopeBiasDirectional(
+	float2 DepthSlopeUV,
+	float2 OffsetUV)
+{
+	return 2.0f * max(0.0f, dot(DepthSlopeUV, OffsetUV));
+}
 
+int GetMappedClipmapId(int VirtualShadowMapId, FLWCVector3 RayOriginWorld)
+{
+	// This is the (unbiased) most detailed clipmap level that we are guaranteed to have a page table entry for
+	const FVirtualShadowMapProjectionShaderData BaseProjectionData = GetVirtualShadowMapProjectionData(VirtualShadowMapId);		
+	float ClipmapLevelFloat = CalcAbsoluteClipmapLevel(BaseProjectionData, RayOriginWorld);
+	int ClipmapIndex = max(0, int(floor(ClipmapLevelFloat)) - BaseProjectionData.ClipmapLevel);
+
+	// Clipmap origin should be near shaded samples, so shadow translated world should be regular range
+	int ClipmapId = VirtualShadowMapId + ClipmapIndex;
+	const FVirtualShadowMapProjectionShaderData ProjectionData = GetVirtualShadowMapProjectionData(ClipmapId);
+	float3 RayOriginShadowTranslatedWorld = LWCToFloat(LWCAdd(RayOriginWorld, ProjectionData.PreViewTranslation));
+	float3 RayStartUVZ = mul(float4(RayOriginShadowTranslatedWorld, 1.0f), ProjectionData.TranslatedWorldToShadowUVMatrix).xyz;
+
+	// NOTE: We don't need the actual sample here, we're just looking for the best mapped page at the ray origin
+	// Trusting a lot in DCE here... may be better to refactor at some point
+	FVirtualShadowMapSample ShadowSample = SampleVirtualShadowMapClipmap(ClipmapId, RayStartUVZ.xy);
+	if (ShadowSample.bValid && ShadowSample.VirtualShadowMapId > ClipmapId)
+	{
+		// NOTE: We could bias this (as long as we clamp it to 0 still) to allow sampling higher resolution
+		// data in the case where we just happened to hit a pixel with lower resolution, but note that performance
+		// suffers when we have to use fallback levels in the SMRT loop, so this is fine for now.
+		// We may have to revisit this if we start getting finer grained with caching and mapped page resolutions,
+		// but likely if this matters here we would already be seeing visible resolution seams between pages regardless.
+		ClipmapIndex += (ShadowSample.VirtualShadowMapId - ClipmapId);
+	}
+
+	return VirtualShadowMapId + ClipmapIndex;
+}
+
 struct FSMRTClipmapRayState
 {
 	int ClipmapId;			// Pre-offset virtual shadow map ID of specific clipmap level
@@ -44,50 +75,15 @@
 	int SampledClipmapId;
 };
 
-FSMRTClipmapRayState SMRTClipmapWithFallbackRayInitialize(
-	int VirtualShadowMapId,
-	FLWCVector3 RayOriginWorld,
+FSMRTClipmapRayState SMRTClipmapRayInitialize(
+	const FVirtualShadowMapProjectionShaderData ProjectionData,
+	float3 RayOriginShadowTranslatedWorld,
 	float3 RayDir,
 	float RayLength,
 	float RayStartOffset,
-	half3 WorldNormal,
-	float2 TexelDitherOffset)
+	float2 DepthSlopeUV,
+	float2 TexelOffset)
 {
-	const FVirtualShadowMapProjectionShaderData BaseProjectionData = GetVirtualShadowMapProjectionData(VirtualShadowMapId);
-	
-	// This is the (unbiased) most detailed clipmap level that we are guaranteed to have a page table entry for
-	float ClipmapLevelFloat = CalcAbsoluteClipmapLevel(BaseProjectionData, RayOriginWorld);
-	int ClipmapIndex = max(0, int(floor(ClipmapLevelFloat)) - BaseProjectionData.ClipmapLevel);
-
-	// Look up the page table at the ray origin as a proxy for the available resolution
-	// We will iterate in the space of this clipmap level with the expectation that it is a
-	// good guess at nearby mapped pages.
-	{
-		// Clipmap origin should be near shaded samples, so shadow translated world should be regular range
-		int ClipmapId = VirtualShadowMapId + ClipmapIndex;
-		const FVirtualShadowMapProjectionShaderData ProjectionData = GetVirtualShadowMapProjectionData(ClipmapId);
-		float3 RayOriginShadowTranslatedWorld = LWCToFloat(LWCAdd(RayOriginWorld, ProjectionData.PreViewTranslation));
-		float3 RayStart = RayOriginShadowTranslatedWorld + RayDir * RayStartOffset;
-		float3 RayStartUVZ = mul(float4(RayStart, 1.0f), ProjectionData.TranslatedWorldToShadowUVMatrix).xyz;
-		
-		// NOTE: We don't need the actual sample here, we're just looking for the best mapped page at the ray origin
-		// Trusting a lot in DCE here... may be better to refactor at some point
-		FVirtualShadowMapSample ShadowSample = SampleVirtualShadowMapClipmap(ClipmapId, RayStartUVZ.xy);
-		if (ShadowSample.bValid && ShadowSample.VirtualShadowMapId > ClipmapId)
-		{
-			// NOTE: We could bias this (as long as we clamp it to 0 still) to allow sampling higher resolution
-			// data in the case where we just happened to hit a pixel with lower resolution, but note that performance
-			// suffers when we have to use fallback levels in the SMRT loop, so this is fine for now.
-			// We may have to revisit this if we start getting finer grained with caching and mapped page resolutions,
-			// but likely if this matters here we would already be seeing visible resolution seams between pages regardless.
-			ClipmapIndex += (ShadowSample.VirtualShadowMapId - ClipmapId);
-		}
-	}
-
-	const FVirtualShadowMapProjectionShaderData ProjectionData = GetVirtualShadowMapProjectionData(VirtualShadowMapId + ClipmapIndex);
-
-	// Clipmap origin should be near shaded samples, so shadow translated world should be regular range
-	float3 RayOriginShadowTranslatedWorld = LWCToFloat(LWCAdd(RayOriginWorld, ProjectionData.PreViewTranslation));
 	float3 RayStart = RayOriginShadowTranslatedWorld + RayDir * RayStartOffset;
 	float3 RayVector = RayDir * RayLength;
 
@@ -98,20 +94,17 @@
 	// Note that this is directly scaled in texel space of the clipmap level, so it can create visual
 	// discontinuities at clipmap boundaries. We could scale this by distance instead of texel size but
 	// in the case where texels are large enough to be obvious there are already visual discontinuities.
-	if (any(TexelDitherOffset != 0))
 	{
-		// When using a texel bias we have to compute a plane-based offset; with any significant magnification
-		// moving the sampling location like this can produce self-shadowing artifacts.
-		float OptimalBias = ComputeOptimalSlopeBiasDirectional(ProjectionData, TexelDitherOffset, WorldNormal);
-		// Subtract off any portion of the bias that was already covered by the ray start offset (usually contact shadows)
+		float OptimalBias = ComputeOptimalSlopeBiasDirectional(DepthSlopeUV, TexelOffset);
+		// Subtract off any portion of the bias that was already covered by the ray start offset (usually screen ray)
 		OptimalBias = max(0.0f, OptimalBias - abs(RayStartOffset * ProjectionData.ShadowViewToClipMatrix._33));
 
-		RayStartUVZ.xy += TexelDitherOffset / float(CalcLevelDimsTexels(0));
+		RayStartUVZ.xy += TexelOffset;
 		RayStartUVZ.z += OptimalBias;
 	}
 
 	FSMRTClipmapRayState Result;
-	Result.ClipmapId = VirtualShadowMapId + ClipmapIndex;
+	Result.ClipmapId = ProjectionData.VirtualShadowMapId;
 	Result.RayStartUVZ = RayStartUVZ;
 	Result.RayStepUVZ = RayStepUVZ;
 	Result.VirtualTexelAddress = uint2(0xFFFFFFFF, 0xFFFFFFFF);
@@ -192,10 +185,12 @@
 	uint MaxRayCount,
 	uint SamplesPerRay,
 	float RayLengthScale,
-	const half Noise,
+	const float Noise,
 	bool bCullBackface,
 	half3 WorldNormal,
-	half TexelDitherScale)
+	float TexelDitherScale,
+	float WorldDitherScale,
+	bool bUseAdaptiveRayCount)
 {
 	FLWCVector3 WorldPosition = LWCSubtract(LWCPromote(TranslatedWorldPosition), PrimaryView.PreViewTranslation);
 	float3 ViewPosition = mul(float4(TranslatedWorldPosition, 1.0f), View.TranslatedWorldToView).xyz;
@@ -207,6 +202,33 @@
 
 	if (!bCullBackface || !IsBackfaceToDirectionalLight(WorldNormal, Light.Direction, Light.SourceRadius))
 	{
+		// Find the best resolution mapped clipmap at the ray origin location
+		const int ClipmapId = GetMappedClipmapId(VirtualShadowMapId, WorldPosition);
+		const FVirtualShadowMapProjectionShaderData ProjectionData = GetVirtualShadowMapProjectionData(ClipmapId);
+
+		// This function is designed to vary in depth in roughly the same way as shadow texel density does,
+		// thus the dependence on ClipmapResolutionLodBias. The goal here is to define a texel dither that
+		// is "smooth" and doesn't have obvious artifacts at page edges of different clipmap levels.
+		// 
+		// This could become a problem later if we break too far from that global resolution bias. That said,
+		// these cases would have unavoidable discontinuities at page edges, so that would have to be handled
+		// in some different way.
+		//
+		// 0.5 is arbitrary but makes it consistent with the previous dither scale cvar
+		float TexelDither = 0.0f;
+		if (TexelDitherScale > 0.0f)
+		{
+			TexelDither =
+				((0.5f / float(CalcLevelDimsTexels(0))) * TexelDitherScale * DistanceFromViewOrigin) /
+				(exp2(ProjectionData.ClipmapLevel - ProjectionData.ClipmapResolutionLodBias));
+		}
+
+		// ~Constant world-space dither/blur
+		float WorldDither = WorldDitherScale * ProjectionData.ShadowViewToClipMatrix._11;
+		float DitherScale = max(TexelDither, WorldDither);
+
+		float2 DepthSlopeUV = ComputeDepthSlopeDirectionalUV(ProjectionData, WorldNormal);
+
 		// Compute max ray length based on view depth
 		// This affects both how far we are willing to trace across the shadow map (for clipmaps this is related to view depth)
 		// and the maximum size a penumbra can be.
@@ -214,7 +236,8 @@
 		// Too low values will greatly restrict how large penumbras can be in screen space.
 		float RayLength = RayLengthScale * DistanceFromViewOrigin;
 
-		float StepOffset = Noise;
+		// Clipmap origin should be near shaded samples, so shadow translated world should be regular range
+		float3 RayOriginShadowTranslatedWorld = LWCToFloat(LWCAdd(WorldPosition, ProjectionData.PreViewTranslation));
 
 		uint RayMissCount = 0;
 		uint i = 0;
@@ -228,10 +251,17 @@
 
 			float3 RayDir = GetRandomDirectionalLightRayDir(Light, E);
 
-			float2 TexelDitherOffset = (RandSample.zw - 0.5f) * TexelDitherScale;
+			float2 TexelOffset = (RandSample.zw - 0.5f) * DitherScale;
 
-			FSMRTClipmapRayState RayState = SMRTClipmapWithFallbackRayInitialize(VirtualShadowMapId, WorldPosition, RayDir, RayLength, RayStartOffset, WorldNormal, TexelDitherOffset);
-			FSMRTResult SMRTResult = SMRTRayCast(RayState, SamplesPerRay, StepOffset);
+			FSMRTClipmapRayState RayState = SMRTClipmapRayInitialize(
+				ProjectionData,
+				RayOriginShadowTranslatedWorld,
+				RayDir,
+				RayLength,
+				RayStartOffset,
+				DepthSlopeUV,
+				TexelOffset);
+			FSMRTResult SMRTResult = SMRTRayCast(RayState, SamplesPerRay, Noise);
 
 			// Debug output (DCE'd if not used)
 			Result.ClipmapIndexOrMipLevel = RayState.SampledClipmapId - VirtualShadowMapId;
@@ -241,36 +271,33 @@
 			
 			RayMissCount += SMRTResult.bValidHit ? 0 : 1;
 
-	#if SMRT_ADAPTIVE_RAY_COUNT && COMPILER_SUPPORTS_WAVE_VOTE
-			if( i == 0 )
+			if (bUseAdaptiveRayCount)
 			{
-				bool bHit = SMRTResult.bValidHit;
-				if( bHit )
+				if( i == 0 )
 				{
-					// TODO: Some simple heuristic for directional lights
+					bool bHit = SMRTResult.bValidHit;
+
+					// All lanes missed
+					bool bAllLanesMiss = WaveActiveAllTrue( !bHit );
+					if( bAllLanesMiss )
+					{
+						break;
+					}
 				}
-
-				// All lanes miss or hit something close enough to not be noisy.
-				bool bAllLanesMiss = WaveActiveAllTrue( !bHit );
-				if( bAllLanesMiss )
+				else if( i == 1 )
 				{
-					break;
+					// After 2 iterations and all have hit, assume umbra
+					bool bAllLanesHit = WaveActiveAllTrue( RayMissCount == 0 );
+					if( bAllLanesHit )
+					{
+						break;
+					}
 				}
 			}
-			else if( i == 1 )
-			{
-				// After 2 iterations and all have hit, must be in umbra.
-				bool bAllLanesHit = WaveActiveAllTrue( RayMissCount == 0 );
-				if( bAllLanesHit )
-				{
-					break;
-				}
-			}
-	#endif
 		}
 		uint RayCount = min(i + 1, MaxRayCount);		// break vs regular for loop exit
 
-		Result.ShadowFactor = SMRTComputeRayMissFactor(RayMissCount, RayCount, MaxRayCount, Noise);;
+		Result.ShadowFactor = SMRTComputeRayMissFactor(RayMissCount, RayCount, MaxRayCount, Noise);
 		Result.OccluderDistance = -1.0f;	// TODO?
 		Result.RayCount = RayCount;
 	}
diff --git a/Engine/Shaders/Private/VirtualShadowMaps/VirtualShadowMapProjectionSpot.ush b/Engine/Shaders/Private/VirtualShadowMaps/VirtualShadowMapProjectionSpot.ush
--- a/Engine/Shaders/Private/VirtualShadowMaps/VirtualShadowMapProjectionSpot.ush
+++ b/Engine/Shaders/Private/VirtualShadowMaps/VirtualShadowMapProjectionSpot.ush
@@ -13,11 +13,6 @@
 #include "VirtualShadowMapProjectionCommon.ush"
 #include "VirtualShadowMapSMRTCommon.ush"
 
-#ifndef SMRT_ADAPTIVE_RAY_COUNT
-#define SMRT_ADAPTIVE_RAY_COUNT 0
-#endif
-
-
 struct FSMRTSingleRayState
 {
 	int VirtualShadowMapId;
@@ -275,7 +270,8 @@
 	float CotMaxRayAngleFromLight,
 	const float Noise,
 	const bool bCullBackface,
-	float3 WorldNormal)
+	float3 WorldNormal,
+	bool bUseAdaptiveRayCount)
 {
 	const uint LightType = Light.SpotAngles.x > -2.0f ? LIGHT_TYPE_SPOT : LIGHT_TYPE_POINT;
 	const float3 ToLight = Light.TranslatedWorldPosition - TranslatedWorldPosition;
@@ -369,43 +365,45 @@
 			Result.PhysicalTexelAddress = RayState.PhysicalTexelAddress;
 			Result.OccluderDistance = SMRTResult.OccluderDistance;
 
-#if SMRT_ADAPTIVE_RAY_COUNT && COMPILER_SUPPORTS_WAVE_VOTE
-			if( i == 0 )
+			if (bUseAdaptiveRayCount)
 			{
-				bool bHit = SMRTResult.bValidHit;
-				#if 0			// Disabled heuristic penumbra size path for now as it creates some speckle artifacts near the edges of sharp shadows
-					if( bHit )
-					{
-					#if 0
-						//float WorldBluringRadius = ComputeLightSampleWorldBluringRadius( TranslatedWorldPosition, LIGHT_TYPE_SPOT, Light, SMRTResult.OccluderDistance );
-						//float PixelBluringRadius = ProjectWorldDistanceToPixelDistance( SceneDepth, WorldBluringRadius );
-						float WorldBluringRadius = Light.SourceRadius * SMRTResult.OccluderDistance / max( DistToLight - SMRTResult.OccluderDistance, 0.001 );
-						float PixelBluringRadius = WorldBluringRadius * View.ViewSizeAndInvSize.x * 0.5 / ( SceneDepth * View.ClipToView[0][0] );
-						bHit = PixelBluringRadius > 2;
-					#else
-						bHit = SMRTResult.OccluderDistance * Light.SourceRadius * ( 0.25 * View.ViewSizeAndInvSize.y / ( SceneDepth * View.ClipToView[1][1] ) ) > max( DistToLight - Result.OccluderDistance, 0.001 );
+				if( i == 0 )
+				{
+					bool bHit = SMRTResult.bValidHit;
+					#if 0			// Disabled heuristic penumbra size path for now as it creates some speckle artifacts near the edges of sharp shadows
+						if( bHit )
+						{
+						#if 0
+							//float WorldBluringRadius = ComputeLightSampleWorldBluringRadius( TranslatedWorldPosition, LIGHT_TYPE_SPOT, Light, SMRTResult.OccluderDistance );
+							//float PixelBluringRadius = ProjectWorldDistanceToPixelDistance( SceneDepth, WorldBluringRadius );
+							float WorldBluringRadius = Light.SourceRadius * SMRTResult.OccluderDistance / max( DistToLight - SMRTResult.OccluderDistance, 0.001 );
+							float PixelBluringRadius = WorldBluringRadius * View.ViewSizeAndInvSize.x * 0.5 / ( SceneDepth * View.ClipToView[0][0] );
+							bHit = PixelBluringRadius > 2;
+						#else
+							bHit = SMRTResult.OccluderDistance * Light.SourceRadius * ( 0.25 * View.ViewSizeAndInvSize.y / ( SceneDepth * View.ClipToView[1][1] ) ) > max( DistToLight - Result.OccluderDistance, 0.001 );
+						#endif
+						}
 					#endif
+
+					// All lanes miss or hit something close enough to not be noisy.
+					bool bAllLanesMiss = WaveActiveAllTrue( !bHit );
+					if( bAllLanesMiss )
+					{
+						break;
 					}
-				#endif
-
-				// All lanes miss or hit something close enough to not be noisy.
-				bool bAllLanesMiss = WaveActiveAllTrue( !bHit );
-				if( bAllLanesMiss )
+				}
+				else if( i == 1 )
 				{
-					break;
+					// After 2 iterations and all have hit, must be in umbra.
+					// TODO: Modify this heuristic based on ray counts to avoid blocky artifacts in near-umbra
+					bool bAllLanesHit = WaveActiveAllTrue( RayMissCount == 0 );
+					if( bAllLanesHit )
+					{
+						break;
+					}
 				}
 			}
-			else if( i == 1 )
-			{
-				// After 2 iterations and all have hit, must be in umbra.
-				// TODO: Modify this heuristic based on ray counts to avoid blocky artifacts in near-umbra
-				bool bAllLanesHit = WaveActiveAllTrue( RayMissCount == 0 );
-				if( bAllLanesHit )
-				{
-					break;
-				}
-			}
-#endif
+
 			uint Seed = asuint( StepOffset );
 			StepOffset = ( EvolveSobolSeed( Seed ) >> 8 ) * 5.96046447754e-08; // * 2^-24;
 		}
diff --git a/Engine/Shaders/Private/VirtualShadowMaps/VirtualShadowMapSMRTTemplate.ush b/Engine/Shaders/Private/VirtualShadowMaps/VirtualShadowMapSMRTTemplate.ush
--- a/Engine/Shaders/Private/VirtualShadowMaps/VirtualShadowMapSMRTTemplate.ush
+++ b/Engine/Shaders/Private/VirtualShadowMaps/VirtualShadowMapSMRTTemplate.ush
@@ -18,6 +18,11 @@
 // included multiple times as our hacky way to generate template-like instantiations
 // before HLSL supports templates natively.
 
+// If defined, unrolls the loop and ignores the NumSteps parameter in favor of the define
+#ifndef SMRT_TEMPLATE_STATIC_SAMPLES_PER_RAY
+#define SMRT_TEMPLATE_STATIC_SAMPLES_PER_RAY -1
+#endif
+
 FSMRTResult SMRTRayCast(
 	inout SMRT_TEMPLATE_RAY_STRUCT RayState,
 	int NumSteps,
@@ -38,6 +43,10 @@
 	const float DeltaTimeBias = 2.0 * TimeScale * ( TimeBias - 0.5 * TimeScale );
 
 	bool bValidHit =  false;
+#if SMRT_TEMPLATE_STATIC_SAMPLES_PER_RAY >= 0
+	NumSteps = SMRT_TEMPLATE_STATIC_SAMPLES_PER_RAY;
+	[unroll]
+#endif
 	for (int i = 0; i <= NumSteps; i++)
 	{
 		const float SampleTime	= i == NumSteps ? 0 : Pow2( TimeScale * i + TimeBias );
@@ -64,7 +73,7 @@
 			{
 				// Add a small relative error to the comparison to avoid missing surfaces due to numeric precision issues
 				// Without this there are occasionally flickering fireflies in fully shadowed regions with SMRT
-				const float EpsScale = 1.02f;
+				const float EpsScale = 1.5f;
 				const float CompareTolerance = abs(Sample.CompareToleranceScale * DeltaTime) * EpsScale;
 				const bool bBehind = (SampleDepth - ReferenceDepth) > CompareTolerance;
 
diff --git a/Engine/Source/Runtime/Renderer/Private/VirtualShadowMaps/VirtualShadowMapProjection.cpp b/Engine/Source/Runtime/Renderer/Private/VirtualShadowMaps/VirtualShadowMapProjection.cpp
--- a/Engine/Source/Runtime/Renderer/Private/VirtualShadowMaps/VirtualShadowMapProjection.cpp
+++ b/Engine/Source/Runtime/Renderer/Private/VirtualShadowMaps/VirtualShadowMapProjection.cpp
@@ -101,14 +101,21 @@
 	ECVF_Scalability | ECVF_RenderThreadSafe
 );
 
-static TAutoConsoleVariable<float> CVarSMRTTexelDitherScale(
-	TEXT( "r.Shadow.Virtual.SMRT.TexelDitherScale" ),
+static TAutoConsoleVariable<float> CVarSMRTTexelDitherScaleDirectional(
+	TEXT( "r.Shadow.Virtual.SMRT.TexelDitherScaleDirectional" ),
 	2.0f,
-	TEXT( "Applies a dither to the shadow map ray casts to help hide aliasing due to insufficient shadow resolution.\n" )
-	TEXT( "This is usually desirable, but it can occasionally cause shadows from thin geometry to separate from their casters at shallow light angles." ),
+	TEXT( "Applies a dither to the shadow map ray casts for directional lights to help hide aliasing due to insufficient shadow resolution.\n" )
+	TEXT( "Setting this too high can cause shadows light leaks near occluders." ),
 	ECVF_Scalability | ECVF_RenderThreadSafe
 );
 
+static TAutoConsoleVariable<float> CVarSMRTWorldDitherScaleDirectional(
+	TEXT( "r.Shadow.Virtual.SMRT.WorldDitherScaleDirectional" ),
+	0.0f,
+	TEXT( "Applies a world." ),
+	ECVF_Scalability | ECVF_RenderThreadSafe
+);
+
 TAutoConsoleVariable<int32> CVarSMRTExtrapolateWithSlope(
 	TEXT("r.Shadow.Virtual.SMRT.ExtrapolateWithSlope"),
 	1,
@@ -199,19 +206,20 @@
 	SHADER_USE_PARAMETER_STRUCT(FVirtualShadowMapProjectionCS, FGlobalShader)
 	
 	class FDirectionalLightDim		: SHADER_PERMUTATION_BOOL("DIRECTIONAL_LIGHT");
-	class FSMRTAdaptiveRayCountDim	: SHADER_PERMUTATION_BOOL("SMRT_ADAPTIVE_RAY_COUNT");
 	class FSMRTExtrapolateSlopeDim	: SHADER_PERMUTATION_BOOL("SMRT_EXTRAPOLATE_WITH_SLOPE");
 	class FOnePassProjectionDim		: SHADER_PERMUTATION_BOOL("ONE_PASS_PROJECTION");
 	class FHairStrandsDim			: SHADER_PERMUTATION_BOOL("HAS_HAIR_STRANDS");
 	class FVisualizeOutputDim		: SHADER_PERMUTATION_BOOL("VISUALIZE_OUTPUT");
+	// -1 means dynamic count
+	class FSMRTStaticSampleCount	: SHADER_PERMUTATION_RANGE_INT("SMRT_TEMPLATE_STATIC_SAMPLES_PER_RAY", -1, 2);
 
 	using FPermutationDomain = TShaderPermutationDomain<
 		FDirectionalLightDim,
 		FOnePassProjectionDim,
-		FSMRTAdaptiveRayCountDim,
 		FHairStrandsDim,
 		FVisualizeOutputDim,
-		FSMRTExtrapolateSlopeDim>;
+		FSMRTExtrapolateSlopeDim,
+		FSMRTStaticSampleCount>;
 
 	BEGIN_SHADER_PARAMETER_STRUCT(FParameters, )
 		SHADER_PARAMETER_STRUCT_INCLUDE(FVirtualShadowMapSamplingParameters, SamplingParameters)
@@ -228,6 +236,8 @@
 		SHADER_PARAMETER(float, SMRTRayLengthScale)
 		SHADER_PARAMETER(float, SMRTCotMaxRayAngleFromLight)
 		SHADER_PARAMETER(float, SMRTTexelDitherScale)
+		SHADER_PARAMETER(float, SMRTWorldDitherScale)
+		SHADER_PARAMETER(uint32, bSMRTUseAdaptiveRayCount)
 		SHADER_PARAMETER(uint32, InputType)
 		SHADER_PARAMETER(uint32, bCullBackfacingPixels)
 		// One pass projection parameters
@@ -252,10 +262,8 @@
 		FForwardLightingParameters::ModifyCompilationEnvironment(Parameters.Platform, OutEnvironment);
 
 		FPermutationDomain PermutationVector(Parameters.PermutationId);
-		if (PermutationVector.Get<FSMRTAdaptiveRayCountDim>())
-		{
-			OutEnvironment.CompilerFlags.Add(CFLAG_WaveOperations);
-		}
+		// TODO: We may no longer need this with SM6 requirement, but shouldn't hurt
+		OutEnvironment.CompilerFlags.Add(CFLAG_WaveOperations);
 
 		OutEnvironment.CompilerFlags.Add(CFLAG_Wave32);
 		if (FDataDrivenShaderPlatformInfo::GetSupportsRealTypes(Parameters.Platform) == ERHIFeatureSupport::RuntimeGuaranteed)
@@ -296,9 +304,10 @@
 	const FLightSceneProxy* LightProxy = nullptr,
 	int32 VirtualShadowMapId = INDEX_NONE)
 {
+	check(GRHISupportsWaveOperations);
+
 	// Use hair strands data (i.e., hair voxel tracing) only for Gbuffer input for casting hair shadow onto opaque geometry.
 	const bool bHasHairStrandsData = HairStrands::HasViewHairStrandsData(View);
-	const bool bAdaptiveRayCount = GRHISupportsWaveOperations && CVarSMRTAdaptiveRayCount.GetValueOnRenderThread() != 0;
 
 	FVirtualShadowMapProjectionCS::FParameters* PassParameters = GraphBuilder.AllocParameters< FVirtualShadowMapProjectionCS::FParameters >();
 	PassParameters->SamplingParameters = VirtualShadowMapArray.GetSamplingParameters(GraphBuilder);
@@ -309,7 +318,7 @@
 	PassParameters->NormalBias = GetNormalBiasForShader();
 	PassParameters->InputType = uint32(InputType);
 	PassParameters->bCullBackfacingPixels = VirtualShadowMapArray.ShouldCullBackfacingPixels() ? 1 : 0;
-	PassParameters->SMRTTexelDitherScale = CVarSMRTTexelDitherScale.GetValueOnRenderThread();
+	PassParameters->bSMRTUseAdaptiveRayCount = CVarSMRTAdaptiveRayCount.GetValueOnRenderThread() != 0 ? 1 : 0;
 	PassParameters->Strata = Strata::BindStrataGlobalUniformParameters(View);
 	if (bHasHairStrandsData)
 	{
@@ -342,6 +351,8 @@
 		PassParameters->SMRTSamplesPerRay = CVarSMRTSamplesPerRayDirectional.GetValueOnRenderThread();
 		PassParameters->SMRTRayLengthScale = CVarSMRTRayLengthScaleDirectional.GetValueOnRenderThread();
 		PassParameters->SMRTCotMaxRayAngleFromLight = 0.0f;	// unused in this path
+		PassParameters->SMRTTexelDitherScale = CVarSMRTTexelDitherScaleDirectional.GetValueOnRenderThread();
+		PassParameters->SMRTWorldDitherScale = CVarSMRTWorldDitherScaleDirectional.GetValueOnRenderThread();
 	}
 	else
 	{
@@ -349,6 +360,8 @@
 		PassParameters->SMRTSamplesPerRay = CVarSMRTSamplesPerRayLocal.GetValueOnRenderThread();
 		PassParameters->SMRTRayLengthScale = 0.0f;		// unused in this path
 		PassParameters->SMRTCotMaxRayAngleFromLight = 1.0f / FMath::Tan(CVarSMRTMaxRayAngleFromLight.GetValueOnRenderThread());
+		PassParameters->SMRTTexelDitherScale = 0.0f;	// not yet implemented
+		PassParameters->SMRTWorldDitherScale = 0.0f;	// not yet implemented
 	}
 	
 	bool bDebugOutput = false;
@@ -365,13 +378,17 @@
 	}
 #endif
 
+	// If the requested samples per ray matches one of our static permutations, pick that one
+	// Otherwise use the dynamic samples per ray permutation (-1).
+	int StaticSamplesPerRay = PassParameters->SMRTSamplesPerRay == 0 ? PassParameters->SMRTSamplesPerRay : -1;
+
 	FVirtualShadowMapProjectionCS::FPermutationDomain PermutationVector;
 	PermutationVector.Set< FVirtualShadowMapProjectionCS::FDirectionalLightDim >( bDirectionalLight );
 	PermutationVector.Set< FVirtualShadowMapProjectionCS::FOnePassProjectionDim >( bOnePassProjection );
-	PermutationVector.Set< FVirtualShadowMapProjectionCS::FSMRTAdaptiveRayCountDim >( bAdaptiveRayCount );
 	PermutationVector.Set< FVirtualShadowMapProjectionCS::FSMRTExtrapolateSlopeDim >( CVarSMRTExtrapolateWithSlope.GetValueOnRenderThread() != 0 );
 	PermutationVector.Set< FVirtualShadowMapProjectionCS::FHairStrandsDim >( bHasHairStrandsData );
 	PermutationVector.Set< FVirtualShadowMapProjectionCS::FVisualizeOutputDim >( bDebugOutput );
+	PermutationVector.Set< FVirtualShadowMapProjectionCS::FSMRTStaticSampleCount >( StaticSamplesPerRay );
 
 	auto ComputeShader = View.ShaderMap->GetShader< FVirtualShadowMapProjectionCS >( PermutationVector );
 	ClearUnusedGraphResources( ComputeShader, PassParameters );
@@ -380,8 +397,10 @@
 	const FIntPoint GroupCount = FIntPoint::DivideAndRoundUp( ProjectionRect.Size(), 8 );
 	FComputeShaderUtils::AddPass(
 		GraphBuilder,
-		RDG_EVENT_NAME("VirtualShadowMapProjection(RayCount:%s,Input:%s%s)", 
-			bAdaptiveRayCount ? TEXT("Adaptive") : TEXT("Static"), 
+		RDG_EVENT_NAME("VirtualShadowMapProjection(RayCount:%u(%s),SamplesPerRay:%u,Input:%s%s)",
+			PassParameters->SMRTRayCount,
+			PassParameters->bSMRTUseAdaptiveRayCount ? TEXT("Adaptive") : TEXT("Static"), 
+			PassParameters->SMRTSamplesPerRay,
 			ToString(InputType),
 			bDebugOutput ? TEXT(",Debug") : TEXT("")),
 		ComputeShader,
