Add support for callbacks to simulate any-hit and intersection shaders in inline ray tracing:
- Using templates to support passing custom structures in a new function TraceRayInlineWithCallback().
- The callback structure should have OnAnyHit() and OnProceduralIntersection() function with special signatures.
- Added default callback structure (FDefaultTraceRayInlineCallback) with previous behavior (only opaque triangles intersections).
- Rework nanite procedural intersection to use callbacks instead of being hardcoded in TraceRayInline.

diff --git a/Engine/Shaders/Private/RayTracing/RayTracingBarycentrics.usf b/Engine/Shaders/Private/RayTracing/RayTracingBarycentrics.usf
--- a/Engine/Shaders/Private/RayTracing/RayTracingBarycentrics.usf
+++ b/Engine/Shaders/Private/RayTracing/RayTracingBarycentrics.usf
@@ -3,6 +3,10 @@
 #include "../Common.ush"
 #include "RayTracingCommon.ush"
 
+#if ENABLE_TRACE_RAY_INLINE_PROCEDURAL_PRIMITIVE
+#include "../Nanite/NaniteRayTrace.ush"
+#endif
+
 #if COMPUTESHADER
 #include "RayTracingDeferredMaterials.ush"
 #include "TraceRayInline.ush"
@@ -12,6 +16,41 @@
 RaytracingAccelerationStructure TLAS;
 float RTDebugVisualizationNaniteCutError;
 
+#if COMPUTESHADER
+struct FNaniteTraceRayInlineCallback : FDefaultTraceRayInlineCallback
+{
+	float2 Barycentrics;
+	float NaniteCutError;
+
+	bool OnProceduralPrimitive(float3 ObjectRayOrigin, float3 ObjectRayDirection, float RayTMin, inout float RayTCurrent, FTraceRayInlineProceduralIntersectionParameters ProceduralIntersectionParameters)
+	{
+#if ENABLE_TRACE_RAY_INLINE_PROCEDURAL_PRIMITIVE
+		uint PrimitiveID = ProceduralIntersectionParameters.InstanceID;
+		FPrimitiveSceneData PrimitiveData = GetPrimitiveData(PrimitiveID);
+		uint PrimitiveInstanceIndex = 0; // TODO
+		uint InstanceId = PrimitiveData.InstanceSceneDataOffset + PrimitiveInstanceIndex;
+		FInstanceSceneData InstanceData = GetInstanceSceneData(InstanceId, View.InstanceSceneDataSOAStride);
+
+		if (InstanceData.ValidInstance && InstanceData.NaniteRuntimeResourceID != 0xFFFFFFFFu)
+		{
+			const float3 RayOrg = ObjectRayOrigin;
+			const float3 RayDir = ObjectRayDirection;
+			const float TMin = RayTMin;
+
+			FNaniteIntersectionResult NaniteIntersectionResult;
+			if (NaniteIntersection(RayOrg, RayDir, InstanceData, PrimitiveData, NaniteCutError, TMin, RayTCurrent, NaniteIntersectionResult))
+			{
+				Barycentrics = NaniteIntersectionResult.Barycentrics;
+				return true;
+			}
+		}
+#endif // ENABLE_TRACE_RAY_INLINE_PROCEDURAL_PRIMITIVE
+		return false;
+	}
+};
+#endif // COMPUTESHADER
+
+
 void RayTracingBarycentricsCommon(uint2 DispatchThreadIndex)
 {
 	uint2 PixelCoord = DispatchThreadIndex.xy + View.ViewRectMin;
@@ -27,17 +66,21 @@
 #if COMPUTESHADER
 	FTraceRayInlineContext TraceRayInlineContext = CreateTraceRayInlineContext();
 	TraceRayInlineContext.bProcedural = ENABLE_TRACE_RAY_INLINE_PROCEDURAL_PRIMITIVE;
-	TraceRayInlineContext.NaniteCutError = RTDebugVisualizationNaniteCutError;
-	FTraceRayInlineResult TraceRayResult = TraceRayInline(
+
+	FNaniteTraceRayInlineCallback NaniteCallback = (FNaniteTraceRayInlineCallback)0;
+	NaniteCallback.NaniteCutError = RTDebugVisualizationNaniteCutError;
+
+	FTraceRayInlineResult TraceRayResult = TraceRayInlineWithCallback(
 		TLAS,   // AccelerationStructure
 		RayFlags,
 		RAY_TRACING_MASK_ALL,             // InstanceInclusionMask
 		Ray,    // RayDesc
-		TraceRayInlineContext
+		TraceRayInlineContext,
+		NaniteCallback
 	);
 
 	Payload.HitT = TraceRayResult.HitT;
-	Payload.Barycentrics = TraceRayResult.Barycentrics;
+	Payload.Barycentrics = TraceRayResult.bIsProcedural ? NaniteCallback.Barycentrics : TraceRayResult.Barycentrics;
 #else
 	TraceRay(
 		TLAS,   // AccelerationStructure
diff --git a/Engine/Shaders/Private/RayTracing/RayTracingDebugTraversal.usf b/Engine/Shaders/Private/RayTracing/RayTracingDebugTraversal.usf
--- a/Engine/Shaders/Private/RayTracing/RayTracingDebugTraversal.usf
+++ b/Engine/Shaders/Private/RayTracing/RayTracingDebugTraversal.usf
@@ -4,6 +4,10 @@
 #include "../ShadingCommon.ush"
 #include "../ColorMap.ush"
 
+#if ENABLE_TRACE_RAY_INLINE_PROCEDURAL_PRIMITIVE
+#include "../Nanite/NaniteRayTrace.ush"
+#endif
+
 #include "RayTracingCommon.ush"
 #include "RayTracingDeferredMaterials.ush"
 
@@ -19,6 +23,39 @@
 float TraversalClusterScale;
 float TraversalTriangleScale;
 
+struct FNaniteTraceRayInlineCallback : FDefaultTraceRayInlineCallback
+{
+	uint ClusterIntersectionCount;
+
+	bool OnProceduralPrimitive(float3 ObjectRayOrigin, float3 ObjectRayDirection, float RayTMin, inout float RayTCurrent, FTraceRayInlineProceduralIntersectionParameters ProceduralIntersectionParameters)
+	{
+#if ENABLE_TRACE_RAY_INLINE_PROCEDURAL_PRIMITIVE
+		uint PrimitiveID = ProceduralIntersectionParameters.InstanceID;
+		FPrimitiveSceneData PrimitiveData = GetPrimitiveData(PrimitiveID);
+		uint PrimitiveInstanceIndex = 0; // TODO
+		uint InstanceId = PrimitiveData.InstanceSceneDataOffset + PrimitiveInstanceIndex;
+		FInstanceSceneData InstanceData = GetInstanceSceneData(InstanceId, View.InstanceSceneDataSOAStride);
+
+		if (InstanceData.ValidInstance && InstanceData.NaniteRuntimeResourceID != 0xFFFFFFFFu)
+		{
+			const float3 RayOrg = ObjectRayOrigin;
+			const float3 RayDir = ObjectRayDirection;
+			const float TMin = RayTMin;
+
+			const float NaniteCutError = 0.0f;
+
+			FNaniteIntersectionResult NaniteIntersectionResult;
+			if (NaniteIntersection(RayOrg, RayDir, InstanceData, PrimitiveData, NaniteCutError, TMin, RayTCurrent, NaniteIntersectionResult))
+			{
+				ClusterIntersectionCount = NaniteIntersectionResult.ClusterIntersectionCount;
+				return true;
+			}
+		}
+#endif
+		return false;
+	}
+};
+
 [numthreads(INLINE_RAY_TRACING_THREAD_GROUP_SIZE_X, INLINE_RAY_TRACING_THREAD_GROUP_SIZE_Y, 1)]
 void RayTracingDebugTraversalCS(uint3 DispatchThread : SV_DispatchThreadID)
 {
@@ -30,14 +67,16 @@
 
 	uint RayFlags = RAY_FLAG_CULL_BACK_FACING_TRIANGLES | RAY_FLAG_FORCE_OPAQUE;
 
+	FNaniteTraceRayInlineCallback NaniteCallback = (FNaniteTraceRayInlineCallback)0;
+
 	FTraceRayInlineContext TraceRayInlineContext = CreateTraceRayInlineContext();
 	TraceRayInlineContext.bProcedural = ENABLE_TRACE_RAY_INLINE_PROCEDURAL_PRIMITIVE;
-	TraceRayInlineContext.NaniteCutError = 0.0f;
-	FTraceRayInlineResult TraceResult = TraceRayInline(TLAS, RayFlags, RAY_TRACING_MASK_ALL, Ray, TraceRayInlineContext);
 
+	FTraceRayInlineResult TraceResult = TraceRayInlineWithCallback(TLAS, RayFlags, RAY_TRACING_MASK_ALL, Ray, TraceRayInlineContext, NaniteCallback);
+
 	float4 Result = 0;
 
-	const uint Intersections = TraceResult.TraversalStatistics.NodeIntersectionCount + TraceResult.TraversalStatistics.ClusterIntersectionCount + TraceResult.TraversalStatistics.TriangleIntersectionCount;
+	const uint Intersections = TraceResult.TraversalStatistics.NodeIntersectionCount + NaniteCallback.ClusterIntersectionCount + TraceResult.TraversalStatistics.TriangleIntersectionCount;
 	switch (VisualizationMode)
 	{
 	default:
@@ -45,7 +84,7 @@
 		Result.rgb = ColorMapTurbo(Intersections / (TraversalBoxScale + TraversalTriangleScale));
 		break;
 	case RAY_TRACING_DEBUG_VIZ_TRAVERSAL_CLUSTER:
-		Result.rgb = ColorMapTurbo(TraceResult.TraversalStatistics.ClusterIntersectionCount / TraversalClusterScale);
+		Result.rgb = ColorMapTurbo(NaniteCallback.ClusterIntersectionCount / TraversalClusterScale);
 		break;
 	case RAY_TRACING_DEBUG_VIZ_TRAVERSAL_NODE:
 		Result.rgb = ColorMapTurbo(TraceResult.TraversalStatistics.NodeIntersectionCount / TraversalBoxScale);
diff --git a/Engine/Shaders/Private/RayTracing/TraceRayInline.ush b/Engine/Shaders/Private/RayTracing/TraceRayInline.ush
--- a/Engine/Shaders/Private/RayTracing/TraceRayInline.ush
+++ b/Engine/Shaders/Private/RayTracing/TraceRayInline.ush
@@ -20,10 +20,6 @@
 #	define ENABLE_TRACE_RAY_INLINE_TRAVERSAL_STATISTICS 0
 #endif // ENABLE_TRACE_RAY_INLINE_TRAVERSAL_STATISTICS
 
-#if ENABLE_TRACE_RAY_INLINE_PROCEDURAL_PRIMITIVE
-#include "../Nanite/NaniteRayTrace.ush"
-#endif
-
 #if OVERRIDE_TRACERAYINLINE_USH
 #include "/Platform/Private/TraceRayInline.ush"
 
@@ -39,66 +35,86 @@
 #define PLATFORM_SUPPORTS_INLINE_RAY_TRACING_MAX_ITERATIONS 0
 #define PLATFORM_SUPPORTS_INLINE_RAY_TRACING_INSTANCE_MATRIX 1
 
-FTraceRayInlineResult TraceRayInline(RaytracingAccelerationStructure AccelerationStructure,
+template<typename TTraceRayInlineCallback>
+FTraceRayInlineResult TraceRayInlineWithCallback(RaytracingAccelerationStructure AccelerationStructure,
 	uint RayFlags,
 	uint InstanceInclusionMask,
 	RayDesc Ray,
-	FTraceRayInlineContext TraceRayInlineContext)
+	FTraceRayInlineContext TraceRayInlineContext,
+	inout TTraceRayInlineCallback Callback)
 {
 	FTraceRayInlineResult TraceRayResult = (FTraceRayInlineResult)0;
-	TraceRayResult.SetMiss();
+	// Compiler has issues with the line below
+	//TraceRayResult.SetMiss();
+	TraceRayResult.HitT = -1;
 	TraceRayResult.bIsCompleted = true;
 
 #if ENABLE_TRACE_RAY_INLINE_PROCEDURAL_PRIMITIVE
-	float2 ProcBarycentrics;
-	FTraceRayInlineTraversalStatistics ProcTraversalStatistics = (FTraceRayInlineTraversalStatistics)0;
-
-	RayQuery<RAY_FLAG_FORCE_OPAQUE> RayQ;
-#else 
-	RayQuery<RAY_FLAG_SKIP_PROCEDURAL_PRIMITIVES | RAY_FLAG_FORCE_OPAQUE> RayQ;
+	RayQuery<RAY_FLAG_NONE> RayQ;
+#else
+	RayQuery<RAY_FLAG_SKIP_PROCEDURAL_PRIMITIVES> RayQ;
 #endif // ENABLE_TRACE_RAY_INLINE_PROCEDURAL_PRIMITIVE
 
 	RayQ.TraceRayInline(AccelerationStructure, RayFlags, InstanceInclusionMask, Ray);
 
 	while (RayQ.Proceed())
 	{
-#if ENABLE_TRACE_RAY_INLINE_PROCEDURAL_PRIMITIVE
 		switch (RayQ.CandidateType())
 		{
+#if ENABLE_TRACE_RAY_INLINE_PROCEDURAL_PRIMITIVE
 		case CANDIDATE_PROCEDURAL_PRIMITIVE:
 		{
 			if (TraceRayInlineContext.bProcedural)
 			{
-				uint PrimitiveID = RayQ.CandidateInstanceID();
-				FPrimitiveSceneData PrimitiveData = GetPrimitiveData(PrimitiveID);
-				uint PrimitiveInstanceIndex = 0; // TODO
-				uint InstanceId = PrimitiveData.InstanceSceneDataOffset + PrimitiveInstanceIndex;
-				FInstanceSceneData InstanceData = GetInstanceSceneData(InstanceId, View.InstanceSceneDataSOAStride);
+				FTraceRayInlineProceduralIntersectionParameters CandidateProceduralIntersection;
+				CandidateProceduralIntersection.InstanceID = RayQ.CandidateInstanceID();
+				CandidateProceduralIntersection.InstanceIndex = RayQ.CandidateInstanceIndex();
+				CandidateProceduralIntersection.PrimitiveIndex = RayQ.CandidatePrimitiveIndex();
+				CandidateProceduralIntersection.GeometryIndex = RayQ.CandidateGeometryIndex();
+				CandidateProceduralIntersection.InstanceContributionToHitGroupIndex = RayQ.CandidateInstanceContributionToHitGroupIndex();
 
-				if (InstanceData.ValidInstance && InstanceData.NaniteRuntimeResourceID != 0xFFFFFFFFu)
+				const float3 RayOrg = RayQ.CandidateObjectRayOrigin();
+				const float3 RayDir = RayQ.CandidateObjectRayDirection();
+				const float TMin = RayQ.RayTMin();
+				const float TMax = RayQ.CommittedRayT();
+
+				float IntersectionT = TMax;
+
+				if (Callback.OnProceduralPrimitive(RayOrg, RayDir, TMin, IntersectionT, CandidateProceduralIntersection))
 				{
-					const float3 RayOrg = RayQ.CandidateObjectRayOrigin();
-					const float3 RayDir = RayQ.CandidateObjectRayDirection();
-					const float TMin = RayQ.RayTMin();
-					const float TMax = RayQ.CommittedRayT();
-
-					float IntersectionT = TMax;
-
-					FNaniteIntersectionResult NaniteIntersectionResult;
-					if (NaniteIntersection(RayOrg, RayDir, InstanceData, PrimitiveData, TraceRayInlineContext.NaniteCutError, TMin, IntersectionT, NaniteIntersectionResult))
-					{
-						RayQ.CommitProceduralPrimitiveHit(IntersectionT);
-						ProcBarycentrics = NaniteIntersectionResult.Barycentrics;
-						ProcTraversalStatistics.NodeIntersectionCount = NaniteIntersectionResult.NodeIntersectionCount;
-						ProcTraversalStatistics.ClusterIntersectionCount = NaniteIntersectionResult.ClusterIntersectionCount;
-						ProcTraversalStatistics.TriangleIntersectionCount = NaniteIntersectionResult.TriangleIntersectionCount;
-					}
+					RayQ.CommitProceduralPrimitiveHit(IntersectionT);
 				}
 			}					
 			break;
 		}
+#endif // ENABLE_TRACE_RAY_INLINE_PROCEDURAL_PRIMITIVE
+		case CANDIDATE_NON_OPAQUE_TRIANGLE:
+		{
+			FTraceRayInlineResult CandidateTraceRayResult;
+			CandidateTraceRayResult.HitT = RayQ.CandidateTriangleRayT();
+			CandidateTraceRayResult.InstanceID = RayQ.CandidateInstanceID();
+			CandidateTraceRayResult.InstanceIndex = RayQ.CandidateInstanceIndex();
+			CandidateTraceRayResult.PrimitiveIndex = RayQ.CandidatePrimitiveIndex();
+			CandidateTraceRayResult.GeometryIndex = RayQ.CandidateGeometryIndex();
+			CandidateTraceRayResult.Barycentrics = RayQ.CandidateTriangleBarycentrics();
+			CandidateTraceRayResult.bIsFrontFace = RayQ.CandidateTriangleFrontFace();
+			CandidateTraceRayResult.InstanceContributionToHitGroupIndex = RayQ.CandidateInstanceContributionToHitGroupIndex();
+
+			float3x4 ObjectToWorld3x4 = RayQ.CandidateObjectToWorld3x4();
+			CandidateTraceRayResult.ObjectToWorld3x4 = ObjectToWorld3x4;
+
+			// placeholder
+			CandidateTraceRayResult.WorldGeometryNormal = RayQ.WorldRayDirection();
+
+			CandidateTraceRayResult.Bookmark[0] = asuint(TraceRayResult.HitT);
+			CandidateTraceRayResult.Bookmark[1] = asuint(TraceRayResult.HitT);
+			if (Callback.OnAnyHit(RayQ.CandidateObjectRayOrigin(), RayQ.CandidateObjectRayDirection(), CandidateTraceRayResult))
+			{
+				RayQ.CommitNonOpaqueTriangleHit();
+			}
+			break;
 		}
-#endif // ENABLE_TRACE_RAY_INLINE_PROCEDURAL_PRIMITIVE
+		}
 	}
 
 	switch (RayQ.CommittedStatus())
@@ -132,13 +148,14 @@
 #if ENABLE_TRACE_RAY_INLINE_PROCEDURAL_PRIMITIVE
 	case COMMITTED_PROCEDURAL_PRIMITIVE_HIT:
 	{
+		TraceRayResult.bIsProcedural = true;
 		TraceRayResult.HitT = RayQ.CommittedRayT();
 		TraceRayResult.InstanceID = RayQ.CommittedInstanceID();
 		TraceRayResult.InstanceIndex = RayQ.CommittedInstanceIndex();
 		TraceRayResult.PrimitiveIndex = RayQ.CommittedPrimitiveIndex();
 		TraceRayResult.GeometryIndex = RayQ.CommittedGeometryIndex();
-		TraceRayResult.Barycentrics = ProcBarycentrics;
-		TraceRayResult.bIsFrontFace = true; // TODO
+		TraceRayResult.Barycentrics = float2(0.0f, 0.0f); // ignored for procedural
+		TraceRayResult.bIsFrontFace = true; // ignored for procedural
 		TraceRayResult.InstanceContributionToHitGroupIndex = RayQ.CommittedInstanceContributionToHitGroupIndex();
 
 		// placeholder
@@ -150,10 +167,6 @@
 #if PLATFORM_SUPPORTS_INLINE_RAY_TRACING_INSTANCE_MATRIX
 		TraceRayResult.ObjectToWorld3x4 = RayQ.CommittedObjectToWorld3x4();
 #endif // PLATFORM_SUPPORTS_INLINE_RAY_TRACING_INSTANCE_MATRIX
-
-#if ENABLE_TRACE_RAY_INLINE_TRAVERSAL_STATISTICS
-		TraceRayResult.TraversalStatistics = ProcTraversalStatistics;
-#endif // ENABLE_TRACE_RAY_INLINE_TRAVERSAL_STATISTICS
 		break;
 	}
 #endif // ENABLE_TRACE_RAY_INLINE_PROCEDURAL_PRIMITIVE
@@ -162,6 +175,17 @@
 	return TraceRayResult;
 }
 
+FTraceRayInlineResult TraceRayInline(RaytracingAccelerationStructure AccelerationStructure,
+	uint RayFlags,
+	uint InstanceInclusionMask,
+	RayDesc InRay,
+	FTraceRayInlineContext TraceRayInlineContext)
+{
+	FDefaultTraceRayInlineCallback DefaultCallback;
+	RayFlags |= RAY_FLAG_FORCE_OPAQUE;
+	return TraceRayInlineWithCallback(AccelerationStructure, RayFlags, InstanceInclusionMask, InRay, TraceRayInlineContext, DefaultCallback);
+}
+
 #endif // OVERRIDE_TRACERAYINLINE_USH
 
 #ifndef PLATFORM_SUPPORTS_INLINE_RAY_TRACING_TRIANGLE_NORMALS
diff --git a/Engine/Shaders/Private/RayTracing/TraceRayInlineCommon.ush b/Engine/Shaders/Private/RayTracing/TraceRayInlineCommon.ush
--- a/Engine/Shaders/Private/RayTracing/TraceRayInlineCommon.ush
+++ b/Engine/Shaders/Private/RayTracing/TraceRayInlineCommon.ush
@@ -44,13 +44,13 @@
 struct FTraceRayInlineResult
 {
 	float HitT;
-	float2 Barycentrics;
+	float2 Barycentrics; // Only valid if bIsProcedural is false
 	uint InstanceIndex;
 	uint InstanceID;
 	uint InstanceContributionToHitGroupIndex;
 	uint GeometryIndex;
 	uint PrimitiveIndex;
-	bool bIsFrontFace;
+	bool bIsFrontFace; // Only valid if bIsProcedural is false
 
 	uint Bookmark[2];
 
@@ -63,6 +63,9 @@
 	// Can be false only if PLATFORM_SUPPORTS_INLINE_RAY_TRACING_MAX_ITERATIONS is 1
 	bool bIsCompleted;
 
+	// Can only be true if ENABLE_TRACE_RAY_INLINE_PROCEDURAL_PRIMITIVE is 1
+	bool bIsProcedural;
+
 	// Only valid if shader was compiled with ENABLE_TRACE_RAY_INLINE_TRAVERSAL_STATISTICS
 	FTraceRayInlineTraversalStatistics TraversalStatistics;
 
@@ -79,14 +82,38 @@
 
 	// Only used if shader was compiled with ENABLE_TRACE_RAY_INLINE_PROCEDURAL_PRIMITIVE
 	bool bProcedural;
-	float NaniteCutError;
 };
 
+struct FTraceRayInlineProceduralIntersectionParameters
+{
+	uint InstanceIndex;
+	uint InstanceID;
+	uint InstanceContributionToHitGroupIndex;
+	uint GeometryIndex;
+	uint PrimitiveIndex;
+};
+
+// Default callback structure for use with TraceRayInline.
+// If custom functions are needed the new structure should derive from this structure.
+struct FDefaultTraceRayInlineCallback
+{	
+	// Returning true means to accept the hit and false means to continue search (default accept)
+	bool OnAnyHit(float3 ObjectRayOrigin, float3 ObjectRayDirection, FTraceRayInlineResult Hit)
+	{ 
+		return true; 
+	}
+
+	// Returning true means accept the procedural hit and false means to skip it (default skip)
+	bool OnProceduralPrimitive(float3 ObjectRayOrigin, float3 ObjectRayDirection, float RayTMin, inout float RayTCurrent, FTraceRayInlineProceduralIntersectionParameters ProceduralIntersectionParameters)
+	{
+		return false;
+	}
+};
+
 FTraceRayInlineContext CreateTraceRayInlineContext()
 {
 	FTraceRayInlineContext Context;
 	Context.bProcedural = false;
 	Context.MaxIterations = 0xFFFFFFFF;
-	Context.NaniteCutError = 0.0f;
 	return Context;
 }
