Fix a possible duplication of InstanceId (UE-153227).

diff --git a/Engine/Shaders/Private/LocalVertexFactory.ush b/Engine/Shaders/Private/LocalVertexFactory.ush
--- a/Engine/Shaders/Private/LocalVertexFactory.ush
+++ b/Engine/Shaders/Private/LocalVertexFactory.ush
@@ -269,7 +269,6 @@ struct FPositionOnlyVertexFactoryInput
 
 	// Dynamic instancing related attributes with InstanceIdOffset : ATTRIBUTE1
 	VF_GPUSCENE_DECLARE_INPUT_BLOCK(1)
-	VF_INSTANCED_STEREO_DECLARE_INPUT_BLOCK()
 
 #if (USE_INSTANCING || USE_INSTANCE_CULLING) && !MANUAL_VERTEX_FETCH && (VF_USE_PRIMITIVE_SCENE_DATA != 2)
 	float4 InstanceOrigin : ATTRIBUTE8;  // per-instance random in w 
@@ -280,6 +279,8 @@ struct FPositionOnlyVertexFactoryInput
 
 #if USE_INSTANCING
 	uint InstanceId : SV_InstanceID;
+#else
+	VF_INSTANCED_STEREO_DECLARE_INPUT_BLOCK()
 #endif
 
 #if MANUAL_VERTEX_FETCH
@@ -297,7 +298,6 @@ struct FPositionAndNormalOnlyVertexFactoryInput
 
 	// Dynamic instancing related attributes with InstanceIdOffset : ATTRIBUTE1
 	VF_GPUSCENE_DECLARE_INPUT_BLOCK(1)
-	VF_INSTANCED_STEREO_DECLARE_INPUT_BLOCK()
 
 #if (USE_INSTANCING || USE_INSTANCE_CULLING) && !MANUAL_VERTEX_FETCH && (VF_USE_PRIMITIVE_SCENE_DATA != 2)
 	float4 InstanceOrigin : ATTRIBUTE8;  // per-instance random in w 
@@ -308,6 +308,8 @@ struct FPositionAndNormalOnlyVertexFactoryInput
 
 #if USE_INSTANCING
 	uint InstanceId : SV_InstanceID;
+#else
+	VF_INSTANCED_STEREO_DECLARE_INPUT_BLOCK()
 #endif
 
 #if MANUAL_VERTEX_FETCH
