Add multiple reflection capture support for single layer water.

diff --git a/Engine/Shaders/Private/SingleLayerWaterComposite.usf b/Engine/Shaders/Private/SingleLayerWaterComposite.usf
--- a/Engine/Shaders/Private/SingleLayerWaterComposite.usf
+++ b/Engine/Shaders/Private/SingleLayerWaterComposite.usf
@@ -96,8 +96,20 @@ SingleLayerWaterCompositeOutput SingleLayerWaterComposite(float2 BufferUV, float
 			float IndirectIrradiance = 0.0;
 			float IndirectSpecularOcclusion = 1.0f;
 			float3 ExtraIndirectSpecular = 0;
-			uint CaptureDataStartIndex = 0;
-			uint NumCulledReflectionCaptures = 0;
+			
+			//support reflection captures
+			#if FEATURE_LEVEL >= FEATURE_LEVEL_SM5
+				float2 LocalPosition = SvPosition.xy - View.ViewRectMin.xy;
+
+				uint GridIndex = ComputeLightGridCellIndex(uint2(LocalPosition.x, LocalPosition.y), GBuffer.Depth);
+				uint NumCulledEntryIndex = (ForwardLightData.NumGridCells + GridIndex) * NUM_CULLED_LIGHTS_GRID_STRIDE;
+				uint NumCulledReflectionCaptures = min(ForwardLightData.NumCulledLightsGrid[NumCulledEntryIndex + 0], ForwardLightData.NumReflectionCaptures);
+				uint CaptureDataStartIndex = ForwardLightData.NumCulledLightsGrid[NumCulledEntryIndex + 1];
+			#else
+				uint NumCulledReflectionCaptures = 0;
+				uint CaptureDataStartIndex = 0;
+			#endif
+
 			//
 			float3 N = WorldNormal;
 			float3 V = -CameraToPixel;
@@ -154,8 +166,20 @@ SingleLayerWaterCompositeOutput SingleLayerWaterComposite(float2 BufferUV, float
 			float IndirectIrradiance = 0.0;// GBuffer.IndirectIrradiance;
 			float IndirectSpecularOcclusion = 1.0f;
 			float3 ExtraIndirectSpecular = 0;
-			uint CaptureDataStartIndex = 0;
-			uint NumCulledReflectionCaptures = ForwardLightData.NumReflectionCaptures;
+
+			//support reflection captures
+			#if FEATURE_LEVEL >= FEATURE_LEVEL_SM5
+				float2 LocalPosition = SvPosition.xy - View.ViewRectMin.xy;
+
+				uint GridIndex = ComputeLightGridCellIndex(uint2(LocalPosition.x, LocalPosition.y), GBuffer.Depth);
+				uint NumCulledEntryIndex = (ForwardLightData.NumGridCells + GridIndex) * NUM_CULLED_LIGHTS_GRID_STRIDE;
+				uint NumCulledReflectionCaptures = min(ForwardLightData.NumCulledLightsGrid[NumCulledEntryIndex + 0], ForwardLightData.NumReflectionCaptures);
+				uint CaptureDataStartIndex = ForwardLightData.NumCulledLightsGrid[NumCulledEntryIndex + 1];
+			#else
+				uint NumCulledReflectionCaptures = 0;
+				uint CaptureDataStartIndex = 0;
+			#endif
+
 			//
 			float3 N = GBuffer.WorldNormal;
 			float3 V = -CameraToPixel;
diff --git a/Engine/Source/Runtime/Renderer/Private/SingleLayerWaterRendering.cpp b/Engine/Source/Runtime/Renderer/Private/SingleLayerWaterRendering.cpp
--- a/Engine/Source/Runtime/Renderer/Private/SingleLayerWaterRendering.cpp
+++ b/Engine/Source/Runtime/Renderer/Private/SingleLayerWaterRendering.cpp
@@ -178,6 +178,12 @@ class FSingleLayerWaterCompositePS : public FGlobalShader
 	{
 		return IsFeatureLevelSupported(Parameters.Platform, ERHIFeatureLevel::SM5);
 	}
+
+    static void ModifyCompilationEnvironment(const FGlobalShaderPermutationParameters& Parameters, FShaderCompilerEnvironment& OutEnvironment)
+	{
+		FGlobalShader::ModifyCompilationEnvironment(Parameters, OutEnvironment);
+		FForwardLightingParameters::ModifyCompilationEnvironment(Parameters.Platform, OutEnvironment);//Support reflection captures
+	}
 };
 
 IMPLEMENT_GLOBAL_SHADER(FSingleLayerWaterCompositePS, "/Engine/Private/SingleLayerWaterComposite.usf", "SingleLayerWaterCompositePS", SF_Pixel);
@@ -1028,4 +1034,4 @@ FMeshPassProcessor* CreateSingleLayerWaterPassProcessor(const FScene* Scene, con
 	return new(FMemStack::Get()) FSingleLayerWaterPassMeshProcessor(Scene, InViewIfDynamicMeshCommand, DrawRenderState, InDrawListContext);
 }
 
-FRegisterPassProcessorCreateFunction RegisterSingleLayerWaterPass(&CreateSingleLayerWaterPassProcessor, EShadingPath::Deferred, EMeshPass::SingleLayerWaterPass, EMeshPassFlags::MainView);
\ No newline at end of file
+FRegisterPassProcessorCreateFunction RegisterSingleLayerWaterPass(&CreateSingleLayerWaterPassProcessor, EShadingPath::Deferred, EMeshPass::SingleLayerWaterPass, EMeshPassFlags::MainView);
