Change ShaderPrint symbol buffers layout.

It will ease merging all shader print RW buffer into a single one.

diff --git a/Engine/Shaders/Private/ShaderPrint.ush b/Engine/Shaders/Private/ShaderPrint.ush
--- a/Engine/Shaders/Private/ShaderPrint.ush
+++ b/Engine/Shaders/Private/ShaderPrint.ush
@@ -100,44 +100,46 @@
 //
 
 ///////////////////////////////////////////////////////////////////////////////////////////////////
-// Uniforms and structures
-
-struct FShaderPrintItem
-{
-	float2 ScreenPos; // Position in normalized coordinates
-	int Value;        // Cast to value or symbol
-	int Type;         // SHADER_PRINT_TYPE_* defines how to read Value
-	float3 Color;	  // Color 
-	uint Metadata;	  // Metadata (used for widget to store extra state data
-};
-
-// Needs to match C++ code in ShaderPrint.cpp
-struct FPackedShaderPrintItem
-{
-	uint ScreenPos16bits; // Position in normalized coordinates
-	int Value;            // Cast to value or symbol
-	uint TypeAndColor;    //
-	uint Metadata;	      //
-};
-
-StructuredBuffer<uint> ShaderPrint_StateBuffer;
-RWStructuredBuffer<FPackedShaderPrintItem> ShaderPrint_RWValuesBuffer;
-
-///////////////////////////////////////////////////////////////////////////////////////////////////
 // Input types
 
 // Content of FShaderPrintItem.Type defines what to cast FShaderPrintItem.Value as
-#define SHADER_PRINT_TYPE_SYMBOL 0
-#define SHADER_PRINT_TYPE_FLOAT  1
-#define SHADER_PRINT_TYPE_INT    2
-#define SHADER_PRINT_TYPE_UINT   3
-#define SHADER_PRINT_TYPE_HEX    4
 
-#define SHADER_PRINT_TYPE_SLIDER 5
-#define SHADER_PRINT_TYPE_CHECK  6
+// Type of primitives
+#define SHADER_PRINT_TYPE_SYMBOL   0
+#define SHADER_PRINT_TYPE_FLOAT    1
+#define SHADER_PRINT_TYPE_INT      2
+#define SHADER_PRINT_TYPE_UINT     3
+#define SHADER_PRINT_TYPE_HEX      4
+#define SHADER_PRINT_TYPE_SLIDER   5
+#define SHADER_PRINT_TYPE_CHECK    6
+// Type of primitives (to be merged with the list above once all buffer are merged)
+#define SHADER_PRINT_TYPE_LINE     0
+#define SHADER_PRINT_TYPE_TRIANGLE 1
 
-#define SHADER_PRINT_COUNT_OFFSET 0
-#define SHADER_PRINT_VALUE_OFFSET 1
+// Size of packed data, in uint
+#define SHADER_PRINT_UINT_COUNT_SYMBOL   4
+#define SHADER_PRINT_UINT_COUNT_LINE     8
+#define SHADER_PRINT_UINT_COUNT_TRIANGLE 12
+
+// Counter offset
+#define SHADER_PRINT_COUNTER_OFFSET_LINE     0
+#define SHADER_PRINT_COUNTER_OFFSET_TRIANGLE 1
+#define SHADER_PRINT_COUNTER_OFFSET_SYMBOL   2
+#define SHADER_PRINT_COUNTER_OFFSET_FREE     3
+#define SHADER_PRINT_COUNTER_COUNT           4
+
+uint GetSymbolOffset(uint InIndex)				{ return SHADER_PRINT_COUNTER_COUNT + SHADER_PRINT_UINT_COUNT_SYMBOL * InIndex; }
+uint GetPrimitiveLineOffset(uint InIndex)		{ return SHADER_PRINT_COUNTER_COUNT + SHADER_PRINT_UINT_COUNT_LINE * InIndex; }
+uint GetPrimitiveTriangleOffset(uint InIndex)	{ return SHADER_PRINT_COUNTER_COUNT + SHADER_PRINT_UINT_COUNT_LINE * ShaderPrintData.MaxLineCount + SHADER_PRINT_UINT_COUNT_TRIANGLE * InIndex; }
+
+void ClearCounters(RWStructuredBuffer<uint> InRWBuffer)
+{
+	// Symbol/print Counter
+	InRWBuffer[SHADER_PRINT_COUNTER_OFFSET_LINE] = 0;
+	InRWBuffer[SHADER_PRINT_COUNTER_OFFSET_TRIANGLE] = 0;
+	InRWBuffer[SHADER_PRINT_COUNTER_OFFSET_SYMBOL] = 0;
+	InRWBuffer[SHADER_PRINT_COUNTER_OFFSET_FREE] = 0;
+}
 
 ///////////////////////////////////////////////////////////////////////////////////////////////////
 // State 
@@ -153,6 +155,27 @@
 #define SHADER_PRINT_STATE_HASH_MASK 0xFFFFFF
 
 ///////////////////////////////////////////////////////////////////////////////////////////////////
+// Output structures
+
+StructuredBuffer<uint> ShaderPrint_StateBuffer;
+
+// Symbols
+//  ----------------------------------------------------------------------------------------------
+// ||  Line    | Triangle |  Symbol |  Empty  ||  Symb0  |  Symb1  |  Symb2  |  Symb3  |  ...  ...
+// || Counter  | Counter  | Counter | Counter ||         |         |         |         |  ...  ...
+// || (1uint)  | (1uint)  | (1uint) | (1uint) || (4uint) | (4uint) | (4uint) | (4uint) |  ...  ...
+//  ----------------------------------------------------------------------------------------------
+RWStructuredBuffer<uint> ShaderPrint_RWValuesBuffer;
+
+// Primitive buffer layout
+//  ------------------------------------------------------------------------------------------------------------------------------------------------
+// ||  Line    | Triangle |  Empty  |  Empty  ||  Line0  |  Line1  |  Line2  |  Line3  |  ...  | Triangle0 | Triangle1 | Triangle2 | Triangle3 | ...
+// || Counter  | Counter  | Counter | Counter ||         |         |         |         |  ...  |           |           |           |           | ...
+// || (1uint)  | (1uint)  | (1uint) | (1uint) || (8uint) | (8uint) | (8uint) | (8uint) |  ...  | (12uint)  | (12uint)  | (12uint)  | (12uint)  | ...
+//  ------------------------------------------------------------------------------------------------------------------------------------------------
+RWStructuredBuffer<uint> ShaderPrint_RWPrimitivesBuffer;
+
+///////////////////////////////////////////////////////////////////////////////////////////////////
 // Font Color
 
 struct FFontColor
@@ -213,6 +236,24 @@
 ///////////////////////////////////////////////////////////////////////////////////////////////////
 // Util pack/unpac functions
 
+struct FShaderPrintItem
+{
+	float2 ScreenPos; // Position in normalized coordinates
+	int Value;        // Cast to value or symbol
+	int Type;         // SHADER_PRINT_TYPE_* defines how to read Value
+	float3 Color;	  // Color 
+	uint Metadata;	  // Metadata (used for widget to store extra state data
+};
+
+// Needs to match C++ code in ShaderPrint.cpp
+struct FPackedShaderPrintItem
+{
+	uint ScreenPos16bits; // Position in normalized coordinates
+	uint Value;           // Cast to value or symbol
+	uint TypeAndColor;    //
+	uint Metadata;	      //
+};
+
 FPackedShaderPrintItem PackShaderPrintItem(FShaderPrintItem In)
 {
 	const uint3 Color8bits = saturate(In.Color) * 0xFF;
@@ -220,7 +261,7 @@
 
 	FPackedShaderPrintItem Out;
 	Out.ScreenPos16bits = ScreenPos16bit.x | (ScreenPos16bit.y<<16);
-	Out.Value = In.Value;
+	Out.Value = asuint(In.Value);
 	Out.TypeAndColor = (Color8bits.z << 24) | (Color8bits.y << 16) | (Color8bits.x << 8) | (In.Type & 0xFF);
 	Out.Metadata = In.Metadata;
 	return Out;
@@ -232,7 +273,7 @@
 
 	FShaderPrintItem Out;
 	Out.ScreenPos = f16tof32(ScreenPos16bits);
-	Out.Value	  = In.Value;
+	Out.Value	  = asint(In.Value);
 	Out.Type	  = (In.TypeAndColor) & 0xFF;
 	Out.Color.x   = float((In.TypeAndColor >>  8) & 0xFF) / 255.f;
 	Out.Color.y   = float((In.TypeAndColor >> 16) & 0xFF) / 255.f;
@@ -241,6 +282,30 @@
 	return Out;
 }
 
+void WriteSymbol(uint Offset, FShaderPrintItem In, RWStructuredBuffer<uint> InRWBuffer)
+{
+	const uint Offset4 = GetSymbolOffset(Offset);
+
+	const FPackedShaderPrintItem Packed = PackShaderPrintItem(In);
+	InRWBuffer[Offset4 + 0] = Packed.ScreenPos16bits;
+	InRWBuffer[Offset4 + 1] = Packed.Value;
+	InRWBuffer[Offset4 + 2] = Packed.TypeAndColor;
+	InRWBuffer[Offset4 + 3] = Packed.Metadata;
+}
+
+FShaderPrintItem ReadSymbol(uint Offset, StructuredBuffer<uint> InBuffer)
+{
+	const uint Offset4 = GetSymbolOffset(Offset);
+
+	FPackedShaderPrintItem Packed = (FPackedShaderPrintItem)0;
+	Packed.ScreenPos16bits	= InBuffer[Offset4 + 0];
+	Packed.Value			= InBuffer[Offset4 + 1];
+	Packed.TypeAndColor		= InBuffer[Offset4 + 2];
+	Packed.Metadata			= InBuffer[Offset4 + 3];
+
+	return UnpackShaderPrintItem(Packed);
+}
+
 ///////////////////////////////////////////////////////////////////////////////////////////////////
 // Global variables
 
@@ -291,7 +356,7 @@
 
 	// Buffer counter is stored in first element .Value
 	int IndexToStore = 0;
-	InterlockedAdd(ShaderPrint_RWValuesBuffer[SHADER_PRINT_COUNT_OFFSET].Value, 1, IndexToStore);
+	InterlockedAdd(ShaderPrint_RWValuesBuffer[SHADER_PRINT_COUNTER_OFFSET_SYMBOL], 1, IndexToStore);
 
 	// Prevent writing off the buffer
 	// Note that counter still increases so need clamp when reading it in later passes
@@ -300,8 +365,7 @@
 		return;
 	}
 
-	// Because counter is in first element, read/writes need to be offset by 1
-	ShaderPrint_RWValuesBuffer[IndexToStore + SHADER_PRINT_VALUE_OFFSET] = PackShaderPrintItem(Item);
+	WriteSymbol(IndexToStore, Item, ShaderPrint_RWValuesBuffer);
 }
 
 void ShaderPrint_Internal(in float2 ScreenPos, in int Value, in FFontColor FontColor, in uint Metadata, in int Type)
@@ -369,7 +433,7 @@
 		// Prevent writing off the buffer
 		// Note that counter still increases so need clamp when reading it in later passes
 		int IndexToStore = 0;
-		InterlockedAdd(ShaderPrint_RWValuesBuffer[SHADER_PRINT_COUNT_OFFSET].Value, Count, IndexToStore);
+		InterlockedAdd(ShaderPrint_RWValuesBuffer[SHADER_PRINT_COUNTER_OFFSET_SYMBOL], Count, IndexToStore);
 		if (uint(IndexToStore + Count) < ShaderPrintData.MaxValueCount)
 		{
 			for (uint i = Begin; i < End; ++i)
@@ -381,7 +445,7 @@
 				Item.Color = InColor.Color;
 				Item.Metadata = 0u;
 
-				ShaderPrint_RWValuesBuffer[IndexToStore + SHADER_PRINT_VALUE_OFFSET] = PackShaderPrintItem(Item);
+				WriteSymbol(IndexToStore, Item, ShaderPrint_RWValuesBuffer);
 
 				++IndexToStore;
 				InPos.x += ShaderPrintData.FontSpacing.x;
@@ -554,10 +618,7 @@
 
 void ShaderPrintFilterOneThread()
 {
-	// Atomic flag is stored in first element .Type Value
-	uint PrevValue;
-	InterlockedCompareExchange(ShaderPrint_RWValuesBuffer[SHADER_PRINT_COUNT_OFFSET].TypeAndColor, 0, 1, PrevValue);
-	ShaderPrintFilter(PrevValue == 0);
+	// Deprecated
 }
 
 ///////////////////////////////////////////////////////////////////////////////////////////////////
@@ -719,32 +780,6 @@
 ///////////////////////////////////////////////////////////////////////////////////////////////////
 // Primitives types
 
-// Type of primitives
-#define SHADER_PRINT_TYPE_LINE 0
-#define SHADER_PRINT_TYPE_TRIANGLE 1
-
-// Size of packed data, in uint
-#define SHADER_PRINT_UINT_COUNT_LINE 8
-#define SHADER_PRINT_UINT_COUNT_TRIANGLE 12
-
-// Counter offset
-#define SHADER_PRINT_COUNTER_OFFSET_LINE 0
-#define SHADER_PRINT_COUNTER_OFFSET_TRIANGLE 1
-#define SHADER_PRINT_COUNTER_OFFSET_FREE0 2
-#define SHADER_PRINT_COUNTER_OFFSET_FREE1 3
-#define SHADER_PRINT_COUNTER_COUNT 4
-
-// Primitive buffer layout
-//  ------------------------------------------------------------------------------------------------------------------------------------------------
-// ||  Line    | Triangle |  Empty  |  Empty  ||  Line0  |  Line1  |  Line2  |  Line3  |  ...  | Triangle0 | Triangle1 | Triangle2 | Triangle3 | ...
-// || Counter  | Counter  | Counter | Counter ||         |         |         |         |  ...  |           |           |           |           | ...
-// || (1uint)  | (1uint)  | (1uint) | (1uint) || (8uint) | (8uint) | (8uint) | (8uint) |  ...  | (12uint)  | (12uint)  | (12uint)  | (12uint)  | ...
-//  ------------------------------------------------------------------------------------------------------------------------------------------------
-RWStructuredBuffer<uint> ShaderPrint_RWPrimitivesBuffer;
-
-uint GetPrimitiveLineOffset    (uint InIndex) { return SHADER_PRINT_COUNTER_COUNT + SHADER_PRINT_UINT_COUNT_LINE * InIndex; }
-uint GetPrimitiveTriangleOffset(uint InIndex) { return SHADER_PRINT_COUNTER_COUNT + SHADER_PRINT_UINT_COUNT_LINE * ShaderPrintData.MaxLineCount + SHADER_PRINT_UINT_COUNT_TRIANGLE * InIndex; }
-
 ///////////////////////////////////////////////////////////////////////////////////////////////////
 // Triangle-based Primitives (Triangle, Point, ...)
 
diff --git a/Engine/Shaders/Private/ShaderPrintDraw.usf b/Engine/Shaders/Private/ShaderPrintDraw.usf
--- a/Engine/Shaders/Private/ShaderPrintDraw.usf
+++ b/Engine/Shaders/Private/ShaderPrintDraw.usf
@@ -11,11 +11,11 @@ uint2 ViewportResolution;
 uint FrameIndex;
 uint FrameThreshold;
 
-StructuredBuffer<FPackedShaderPrintItem> ValuesBuffer;
-StructuredBuffer<FPackedShaderPrintItem> SymbolsBuffer;
+StructuredBuffer<uint> ValuesBuffer;
+StructuredBuffer<uint> SymbolsBuffer;
 
-RWStructuredBuffer<FPackedShaderPrintItem> RWValuesBuffer;
-RWStructuredBuffer<FPackedShaderPrintItem> RWSymbolsBuffer;
+RWStructuredBuffer<uint> RWValuesBuffer;
+RWStructuredBuffer<uint> RWSymbolsBuffer;
 RWStructuredBuffer<uint> RWStateBuffer;
 
 RWBuffer<uint> RWIndirectDispatchArgsBuffer;
@@ -30,17 +30,10 @@ RWBuffer<uint> RWIndirectDrawArgsBuffer;
 [numthreads(1, 1, 1)]
 void InitValueBufferCS()
 {
-	// Clear .Value in the first buffer element which is (ab)used as the buffer counter
-	// Clear .Type in the first buffer element which is (ab)used as the ShaderPrintFilterOneThread() atomic value
-
 	// Symbol/print Counter
-	RWValuesBuffer[SHADER_PRINT_COUNT_OFFSET].ScreenPos16bits = 0;
-	RWValuesBuffer[SHADER_PRINT_COUNT_OFFSET].Value = 0;
-	RWValuesBuffer[SHADER_PRINT_COUNT_OFFSET].TypeAndColor = 0;
-	RWValuesBuffer[SHADER_PRINT_COUNT_OFFSET].Metadata = 0;
+	ClearCounters(RWValuesBuffer);
 }
 
-
 // --------------------------------------------------------------------------
 // Fill the indirect params buffer ready for the value->symbol compute pass
 // Also clear the symbol buffer ready to be filled
@@ -48,17 +41,15 @@ void InitValueBufferCS()
 [numthreads(1, 1, 1)]
 void BuildIndirectDispatchArgsCS()
 {
-	uint ValueCount = min(ShaderPrintData.MaxValueCount, uint(ValuesBuffer[SHADER_PRINT_COUNT_OFFSET].Value));
+	uint ValueCount = min(ShaderPrintData.MaxValueCount, ValuesBuffer[SHADER_PRINT_COUNTER_OFFSET_SYMBOL]);
 
 	RWIndirectDispatchArgsBuffer[0] = (ValueCount + 63)/64;
 	RWIndirectDispatchArgsBuffer[1] = 1;
 	RWIndirectDispatchArgsBuffer[2] = 1;
 
-	// Clear .Value in the first buffer element which is (ab)used as the buffer counter
-	RWSymbolsBuffer[SHADER_PRINT_COUNT_OFFSET].Value = 0;
+	ClearCounters(RWSymbolsBuffer);
 }
 
-
 // --------------------------------------------------------------------------
 // Read the values buffer and convert to the symbols buffer
 // This is required for converting a single number to a full set of symbols
@@ -73,7 +64,7 @@ void AddSymbol(inout float2 ScreenPos, in int Symbol, in float3 Color, uint Type
 	{
 		// Buffer counter is stored in first element .Value
 		int IndexToStore = 0;
-		InterlockedAdd(RWSymbolsBuffer[SHADER_PRINT_COUNT_OFFSET].Value, 1, IndexToStore);
+		InterlockedAdd(RWSymbolsBuffer[SHADER_PRINT_COUNTER_OFFSET_SYMBOL], 1, IndexToStore);
 	
 		// Prevent writing off the buffer
 		// Note that counter still increases so need to clamp when reading it in later passes
@@ -82,8 +73,7 @@ void AddSymbol(inout float2 ScreenPos, in int Symbol, in float3 Color, uint Type
 			return;
 		}
 
-		// Because counter is in first element, read/writes need to be offset by 1
-		RWSymbolsBuffer[IndexToStore + SHADER_PRINT_VALUE_OFFSET] = PackShaderPrintItem(Item);
+		WriteSymbol(IndexToStore, Item, RWSymbolsBuffer);
 	}
 
 	ScreenPos.x += ShaderPrintData.FontSpacing.x;
@@ -213,13 +203,13 @@ void AddHexSymbols(inout float2 ScreenPos, uint Value, in float3 Color)
 [numthreads(64, 1, 1)]
 void BuildSymbolBufferCS(uint3 DispatchThreadId : SV_DispatchThreadID)
 {
-	uint ValueCount = min(ShaderPrintData.MaxValueCount, uint(ValuesBuffer[SHADER_PRINT_COUNT_OFFSET].Value));
+	uint ValueCount = min(ShaderPrintData.MaxValueCount, ValuesBuffer[SHADER_PRINT_COUNTER_OFFSET_SYMBOL]);
 	if (DispatchThreadId.x >= ValueCount)
 	{
 		return;
 	}
-
-	FShaderPrintItem Value = UnpackShaderPrintItem(ValuesBuffer[DispatchThreadId.x + SHADER_PRINT_VALUE_OFFSET]);
+	
+	FShaderPrintItem Value = ReadSymbol(DispatchThreadId.x, ValuesBuffer);
 	
 	if (Value.Type == SHADER_PRINT_TYPE_SYMBOL)
 		AddSymbol(Value.ScreenPos, Value.Value, Value.Color);
@@ -281,7 +271,7 @@ void CompactStateBufferCS(uint3 DispatchThreadId : SV_DispatchThreadID)
 [numthreads(1, 1, 1)]
 void BuildIndirectDrawArgsCS()
 {
-	uint SymbolCount = min(ShaderPrintData.MaxSymbolCount, uint(SymbolsBuffer[SHADER_PRINT_COUNT_OFFSET].Value));
+	uint SymbolCount = min(ShaderPrintData.MaxSymbolCount, SymbolsBuffer[SHADER_PRINT_COUNTER_OFFSET_SYMBOL]);
 
 	RWIndirectDrawArgsBuffer[0] = 6;
 	RWIndirectDrawArgsBuffer[1] = SymbolCount * 2; // Main glyph + drop shadow
@@ -308,7 +298,7 @@ void DrawSymbolsVS(
 	float2 DropShadowOffset = bDropShadow ? ShaderPrintData.FontSize.xy * float2(0.2f, -0.2f) : 0;
 
 	InstanceId = InstanceId >> 1;
-	FShaderPrintItem Symbol = UnpackShaderPrintItem(SymbolsBuffer[InstanceId + SHADER_PRINT_VALUE_OFFSET]);
+	FShaderPrintItem Symbol = ReadSymbol(InstanceId, SymbolsBuffer);
 	OutType = Symbol.Type;
 	OutValue = Symbol.Value;
 	OutColor = Symbol.Color;
diff --git a/Engine/Shaders/Private/ShaderPrintDrawPrimitive.usf b/Engine/Shaders/Private/ShaderPrintDrawPrimitive.usf
--- a/Engine/Shaders/Private/ShaderPrintDrawPrimitive.usf
+++ b/Engine/Shaders/Private/ShaderPrintDrawPrimitive.usf
@@ -114,10 +114,7 @@ RWStructuredBuffer<uint> RWElementBuffer;
 [numthreads(1, 1, 1)]
 void ShaderDrawDebugClearCS(uint3 DispatchThreadId : SV_DispatchThreadID)
 {
-	RWElementBuffer[0] = 0; // Only element which really need to be cleared
-	RWElementBuffer[1] = 0;
-	RWElementBuffer[2] = 0;
-	RWElementBuffer[3] = 0;
+	ClearCounters(RWElementBuffer);
 }
 #endif
 
diff --git a/Engine/Source/Runtime/Renderer/Private/ShaderPrint.cpp b/Engine/Source/Runtime/Renderer/Private/ShaderPrint.cpp
--- a/Engine/Source/Runtime/Renderer/Private/ShaderPrint.cpp
+++ b/Engine/Source/Runtime/Renderer/Private/ShaderPrint.cpp
@@ -86,30 +86,7 @@ namespace ShaderPrint
 
 	//////////////////////////////////////////////////////////////////////////////////////////////////
 	// Struct & Functions
-	
-	// Structure used by shader buffers to store values and symbols
-	struct FPackedShaderPrintItem
-	{
-		uint32 ScreenPos16bits;
-		int32  Value;
-		uint32 TypeAndColor;
-		uint32 Metadata;
-	};
 
-	// Empty buffer for binding when ShaderPrint is disabled
-	class FEmptyBuffer : public FBufferWithRDG
-	{
-	public:
-		void InitRHI() override
-		{
-			Buffer = AllocatePooledBuffer(FRDGBufferDesc::CreateStructuredDesc(sizeof(FPackedShaderPrintItem), 1), TEXT("ShaderPrint.EmptyValueBuffer"));
-		}
-	};
-
-	FBufferWithRDG* GEmptyBuffer = new TGlobalResource<FEmptyBuffer>();
-
-	// Get value buffer size
-	// Note that if the ShaderPrint system is disabled we still want to bind a minimal buffer
 	static uint32 GetMaxValueCount()
 	{
 		uint32 MaxValueCount = FMath::Max(CVarMaxCharacterCount.GetValueOnRenderThread() + int32(GCharacterRequestCount), 0);
@@ -135,9 +112,10 @@ namespace ShaderPrint
 	}
 
 	// Returns the number of uints used for counters, a line element, and a triangle elements
-	static uint32 GetCountersUintSize()       { return 4;  }
+	static uint32 GetCountersUintSize()       { return 4; }
 	static uint32 GetPackedLineUintSize()     { return 8; }
 	static uint32 GetPackedTriangleUintSize() { return 12; }
+	static uint32 GetPackedSymbolUintSize()   { return 4; }
 
 	// Get symbol buffer size
 	// This is some multiple of the value buffer size to allow for maximum value->symbol expansion
@@ -156,6 +134,18 @@ namespace ShaderPrint
 		return CVarDrawLock.GetValueOnRenderThread() > 0;
 	}
 
+	// Empty buffer for binding when ShaderPrint is disabled
+	class FEmptyBuffer : public FBufferWithRDG
+	{
+	public:
+		void InitRHI() override
+		{
+			Buffer = AllocatePooledBuffer(FRDGBufferDesc::CreateStructuredDesc(4, GetCountersUintSize()), TEXT("ShaderPrint.EmptyValueBuffer"));
+		}
+	};
+
+	FBufferWithRDG* GEmptyBuffer = new TGlobalResource<FEmptyBuffer>();
+
 	//////////////////////////////////////////////////////////////////////////////////////////////////
 	// Uniform buffer
 	
@@ -279,7 +269,7 @@ namespace ShaderPrint
 		SHADER_USE_PARAMETER_STRUCT(FShaderInitValueBufferCS, FGlobalShader);
 
 		BEGIN_SHADER_PARAMETER_STRUCT(FParameters, )
-			SHADER_PARAMETER_RDG_BUFFER_UAV(RWStructuredBuffer<FPackedShaderPrintItem>, RWValuesBuffer)
+			SHADER_PARAMETER_RDG_BUFFER_UAV(RWStructuredBuffer<uint>, RWValuesBuffer)
 		END_SHADER_PARAMETER_STRUCT()
 
 		static bool ShouldCompilePermutation(FGlobalShaderPermutationParameters const& Parameters)
@@ -299,8 +289,8 @@ namespace ShaderPrint
 
 		BEGIN_SHADER_PARAMETER_STRUCT(FParameters, )
 			SHADER_PARAMETER_STRUCT_REF(FShaderPrintCommonParameters, Common)
-			SHADER_PARAMETER_RDG_BUFFER_SRV(StructuredBuffer<FPackedShaderPrintItem>, ValuesBuffer)
-			SHADER_PARAMETER_RDG_BUFFER_UAV(RWStructuredBuffer<FPackedShaderPrintItem>, RWSymbolsBuffer)
+			SHADER_PARAMETER_RDG_BUFFER_SRV(StructuredBuffer<uint>, ValuesBuffer)
+			SHADER_PARAMETER_RDG_BUFFER_UAV(RWStructuredBuffer<uint>, RWSymbolsBuffer)
 			SHADER_PARAMETER_RDG_BUFFER_UAV(RWBuffer<uint>, RWIndirectDispatchArgsBuffer)
 		END_SHADER_PARAMETER_STRUCT()
 
@@ -344,8 +334,8 @@ namespace ShaderPrint
 		BEGIN_SHADER_PARAMETER_STRUCT(FParameters, )
 			SHADER_PARAMETER(uint32, FrameIndex)
 			SHADER_PARAMETER_STRUCT_REF(FShaderPrintCommonParameters, Common)
-			SHADER_PARAMETER_RDG_BUFFER_SRV(StructuredBuffer<FPackedShaderPrintItem>, ValuesBuffer)
-			SHADER_PARAMETER_RDG_BUFFER_UAV(RWStructuredBuffer<FPackedShaderPrintItem>, RWSymbolsBuffer)
+			SHADER_PARAMETER_RDG_BUFFER_SRV(StructuredBuffer<uint>, ValuesBuffer)
+			SHADER_PARAMETER_RDG_BUFFER_UAV(RWStructuredBuffer<uint>, RWSymbolsBuffer)
 			SHADER_PARAMETER_RDG_BUFFER_UAV(RWStructuredBuffer<uint>, RWStateBuffer)
 			RDG_BUFFER_ACCESS(IndirectDispatchArgsBuffer, ERHIAccess::IndirectArgs)
 		END_SHADER_PARAMETER_STRUCT()
@@ -367,7 +357,7 @@ namespace ShaderPrint
 
 		BEGIN_SHADER_PARAMETER_STRUCT(FParameters, )
 			SHADER_PARAMETER_STRUCT_REF(FShaderPrintCommonParameters, Common)
-			SHADER_PARAMETER_RDG_BUFFER_SRV(StructuredBuffer<FPackedShaderPrintItem>, SymbolsBuffer)
+			SHADER_PARAMETER_RDG_BUFFER_SRV(StructuredBuffer<uint>, SymbolsBuffer)
 			SHADER_PARAMETER_RDG_BUFFER_UAV(RWBuffer<uint>, RWIndirectDrawArgsBuffer)
 		END_SHADER_PARAMETER_STRUCT()
 
@@ -394,7 +384,7 @@ namespace ShaderPrint
 			RENDER_TARGET_BINDING_SLOTS()
 			SHADER_PARAMETER_STRUCT_REF(FShaderPrintCommonParameters, Common)
 			SHADER_PARAMETER_TEXTURE(Texture2D, MiniFontTexture)
-			SHADER_PARAMETER_RDG_BUFFER_SRV(StructuredBuffer<FPackedShaderPrintItem>, SymbolsBuffer)
+			SHADER_PARAMETER_RDG_BUFFER_SRV(StructuredBuffer<uint>, SymbolsBuffer)
 			RDG_BUFFER_ACCESS(IndirectDrawArgsBuffer, ERHIAccess::IndirectArgs)
 		END_SHADER_PARAMETER_STRUCT()
 
@@ -608,8 +598,8 @@ namespace ShaderPrint
 		// Characters/Widgets
 		{
 			// Initialize output buffer and store in the view info
-			// Values buffer contains Count + 1 elements. The first element is only used as a counter.
-			ShaderPrintData.ShaderPrintValueBuffer = GraphBuilder.CreateBuffer(FRDGBufferDesc::CreateStructuredDesc(sizeof(FPackedShaderPrintItem), InSetup.MaxValueCount + 1), TEXT("ShaderPrint.ValueBuffer"));
+			const uint32 UintElementCount = GetCountersUintSize() + GetPackedSymbolUintSize() * InSetup.MaxValueCount;
+			ShaderPrintData.ShaderPrintValueBuffer = GraphBuilder.CreateBuffer(FRDGBufferDesc::CreateStructuredDesc(4, UintElementCount), TEXT("ShaderPrint.ValueBuffer"));
 
 			// State buffer is retrieved from the view state, or created if it does not exist
 			if (InViewState != nullptr)
@@ -620,7 +610,7 @@ namespace ShaderPrint
 				}
 				else
 				{
-					ShaderPrintData.ShaderPrintStateBuffer = GraphBuilder.CreateBuffer(FRDGBufferDesc::CreateStructuredDesc(sizeof(uint32), (3 * InSetup.MaxStateCount) + 1), TEXT("ShaderPrint.StateBuffer"));
+					ShaderPrintData.ShaderPrintStateBuffer = GraphBuilder.CreateBuffer(FRDGBufferDesc::CreateStructuredDesc(4, (3 * InSetup.MaxStateCount) + 1), TEXT("ShaderPrint.StateBuffer"));
 					AddClearUAVPass(GraphBuilder, GraphBuilder.CreateUAV(ShaderPrintData.ShaderPrintStateBuffer, PF_R32_UINT), 0u);
 					InViewState->ShaderPrintStateData.StateBuffer = GraphBuilder.ConvertToExternalBuffer(ShaderPrintData.ShaderPrintStateBuffer);
 				}
@@ -735,8 +725,8 @@ namespace ShaderPrint
 		FScreenPassTexture OutputTexture)
 	{
 		// Initialize graph managed resources
-		// Symbols buffer contains Count + 1 elements. The first element is only used as a counter.
-		FRDGBufferRef SymbolBuffer = GraphBuilder.CreateBuffer(FRDGBufferDesc::CreateStructuredDesc(sizeof(FPackedShaderPrintItem), GetMaxSymbolCount() + 1), TEXT("ShaderPrint.SymbolBuffer"));
+		const uint32 UintElementCount = GetCountersUintSize() + GetPackedSymbolUintSize() * GetMaxSymbolCount();
+		FRDGBufferRef SymbolBuffer = GraphBuilder.CreateBuffer(FRDGBufferDesc::CreateStructuredDesc(4, UintElementCount), TEXT("ShaderPrint.SymbolBuffer"));
 		FRDGBufferRef IndirectDispatchArgsBuffer = GraphBuilder.CreateBuffer(FRDGBufferDesc::CreateIndirectDesc(4), TEXT("ShaderPrint.IndirectDispatchArgs"));
 		FRDGBufferRef IndirectDrawArgsBuffer = GraphBuilder.CreateBuffer(FRDGBufferDesc::CreateIndirectDesc(5), TEXT("ShaderPrint.IndirectDrawArgs"));
 
