Fix issues when compiling callable shaders using ddx/ddy.

diff --git a/Engine/Shaders/Private/Common.ush b/Engine/Shaders/Private/Common.ush
--- a/Engine/Shaders/Private/Common.ush
+++ b/Engine/Shaders/Private/Common.ush
@@ -222,7 +222,7 @@ float VectorSum(float4 V) { return V.x + V.y + V.z + V.w; }
 // Using SV_ClipDistance has overhead (15% slower base pass in triangle bound test scene on PS4) so projects have to opt-in
 #define USE_GLOBAL_CLIP_PLANE (PROJECT_ALLOW_GLOBAL_CLIP_PLANE && !MATERIAL_DOMAIN_POSTPROCESS && !MATERIAL_DOMAIN_UI)
 
-#if RAYHITGROUPSHADER || RAYMISSHADER
+#if RAYHITGROUPSHADER || RAYMISSHADER || RAYCALLABLESHADER
 
 // These built-ins are not available in ray tracing
 // Define dummy versions so that ray-tracing materials will at least compile
diff --git a/Engine/Shaders/Private/LWCOperations.ush b/Engine/Shaders/Private/LWCOperations.ush
--- a/Engine/Shaders/Private/LWCOperations.ush
+++ b/Engine/Shaders/Private/LWCOperations.ush
@@ -41,7 +41,7 @@ LWC_COMPARE_OP(LWCLessEqual, <=, FLWCType, FFloatType, FBoolType)
 #undef LWC_COMPARE_OP
 
 // No ddxy inside ray tracing shaders
-#if RAYHITGROUPSHADER || RAYMISSHADER || USE_FORCE_TEXTURE_MIP
+#if RAYHITGROUPSHADER || RAYMISSHADER || RAYCALLABLESHADER || USE_FORCE_TEXTURE_MIP
 FFloatType LWCDdx(FLWCType V) { return (FFloatType)0; }
 FFloatType LWCDdy(FLWCType V) { return (FFloatType)0; }
 #else
