Allows TSR to upscale separate translucency independently from input resolution

This also implements r.Translucency.ScreenPercentage.Basis to experiment controling the translucency resolution independently from the primary screen percentage

diff --git a/Engine/Shaders/Private/TemporalSuperResolution/TSRUpdateHistory.usf b/Engine/Shaders/Private/TemporalSuperResolution/TSRUpdateHistory.usf
--- a/Engine/Shaders/Private/TemporalSuperResolution/TSRUpdateHistory.usf
+++ b/Engine/Shaders/Private/TemporalSuperResolution/TSRUpdateHistory.usf
@@ -110,8 +110,26 @@ Texture2D<tsr_ushort> AntiAliasingTexture;
 Texture2D<tsr_half> NoiseFilteringTexture;
 Texture2D<tsr_half2> HoleFilledVelocityMaskTexture;
 
-float4 HistoryPixelPosToScreenPos;
-float4 HistoryPixelPosToPPCo;
+float2 TranslucencyInfo_Extent;
+float2 TranslucencyInfo_ExtentInverse;
+float2 TranslucencyInfo_ScreenPosToViewportScale;
+float2 TranslucencyInfo_ScreenPosToViewportBias;
+uint2  TranslucencyInfo_ViewportMin;
+uint2  TranslucencyInfo_ViewportMax;
+float2 TranslucencyInfo_ViewportSize;
+float2 TranslucencyInfo_ViewportSizeInverse;
+float2 TranslucencyInfo_UVViewportMin;
+float2 TranslucencyInfo_UVViewportMax;
+float2 TranslucencyInfo_UVViewportSize;
+float2 TranslucencyInfo_UVViewportSizeInverse;
+float2 TranslucencyInfo_UVViewportBilinearMin;
+float2 TranslucencyInfo_UVViewportBilinearMax;
+int2   TranslucencyPixelPosMin;
+int2   TranslucencyPixelPosMax;
+
+FScreenTransform HistoryPixelPosToScreenPos;
+FScreenTransform HistoryPixelPosToInputPPCo;
+FScreenTransform HistoryPixelPosToTranslucencyPPCo;
 float3 HistoryQuantizationError;
 float MinTranslucencyRejection;
 float InvWeightClampingPixelSpeed;
@@ -202,10 +220,12 @@ tsr_half3x2 GetColorChannel(tsr_half4x2 Color)
 }
 
 void ComputeInputKernelSamplePosition(
-	tsr_short2x2 InputPixelPos,
+	tsr_short2x2 PixelPos,
 	tsr_half2x2 dKO,
 	uint SampleId,
-	out tsr_short2x2 SampleInputPixelPos,
+	int2 PixelPosMin,
+	int2 PixelPosMax,
+	out tsr_short2x2 SamplePixelPos,
 	out tsr_half2x2 PixelOffset)
 {
 	
@@ -214,10 +234,10 @@ void ComputeInputKernelSamplePosition(
 		tsr_short2 iPixelOffset = tsr_short2(kOffsets3x3[kSquareIndexes3x3[SampleId]]);
 		PixelOffset = dpv_interleave_registers(tsr_half2(iPixelOffset), tsr_half2(iPixelOffset));
 				
-		SampleInputPixelPos = InputPixelPos + ClampPixelOffset(
-			InputPixelPos,
+		SamplePixelPos = PixelPos + ClampPixelOffset(
+			PixelPos,
 			dpv_interleave_registers(iPixelOffset, iPixelOffset), iPixelOffset,
-			InputPixelPosMin, InputPixelPosMax);
+			PixelPosMin, PixelPosMax);
 	}
 	#elif CONFIG_SAMPLES == 5 || CONFIG_SAMPLES == 6
 	{
@@ -237,7 +257,7 @@ void ComputeInputKernelSamplePosition(
 				PixelOffset[1] = asfloat(asuint(1.0).xx | (asuint(dKO[1]) & uint(0x80000000)));
 			#endif
 						
-			SampleInputPixelPos = ClampPixelOffset(InputPixelPos + iPixelOffset, InputPixelPosMin, InputPixelPosMax);
+			SamplePixelPos = ClampPixelOffset(PixelPos + iPixelOffset, PixelPosMin, PixelPosMax);
 		}
 		else
 		#endif
@@ -245,15 +265,15 @@ void ComputeInputKernelSamplePosition(
 			tsr_short2 iPixelOffset = tsr_short2(kOffsets3x3[kPlusIndexes3x3[SampleId]]);
 			PixelOffset = dpv_interleave_registers(tsr_half2(iPixelOffset), tsr_half2(iPixelOffset));
 					
-			SampleInputPixelPos = InputPixelPos + ClampPixelOffset(
-				InputPixelPos,
+			SamplePixelPos = PixelPos + ClampPixelOffset(
+				PixelPos,
 				dpv_interleave_registers(iPixelOffset, iPixelOffset), iPixelOffset,
-				InputPixelPosMin, InputPixelPosMax);
+				PixelPosMin, PixelPosMax);
 		}
 	}
 	#elif CONFIG_SAMPLES == 2
 	{
-		SampleInputPixelPos = InputPixelPos;
+		SamplePixelPos = PixelPos;
 		PixelOffset = 0;
 	}
 	#else
@@ -328,13 +348,13 @@ void MainCS(
 	float2x2 ScreenPos = ApplyScreenTransform(float2x2(HistoryPixelPos), HistoryPixelPosToScreenPos);
 	
 	// Pixel coordinate of the center of output pixel O in the input viewport.
-	float2x2 PPCo = ApplyScreenTransform(float2x2(HistoryPixelPos), HistoryPixelPosToPPCo);
+	float2x2 InputPPCo = ApplyScreenTransform(float2x2(HistoryPixelPos), HistoryPixelPosToInputPPCo);
 
-	// Pixel coordinate of the center of the nearest input pixel K.
-	float2x2 PPCk = floor(PPCo) + 0.5;
+	// Pixel coordinate of the center of the nearest input pixel K in the input viewport.
+	float2x2 InputPPCk = floor(InputPPCo) + 0.5;
 	
 	tsr_short2x2 InputPixelPos = ClampPixelOffset(
-		tsr_short2x2(PPCo),
+		tsr_short2x2(InputPPCo),
 		InputPixelPosMin, InputPixelPosMax);
 
 	// Fetch reprojections related informations.
@@ -401,7 +421,7 @@ void MainCS(
 			}
 
 			// Vector in pixel between pixel K -> O.
-			tsr_half2x2 dKO = tsr_half2x2(PPCo - PPCk);
+			tsr_half2x2 dInputKO = tsr_half2x2(InputPPCo - InputPPCk);
 		
 			UNROLL_N(CONFIG_SAMPLES - 1)
 			for (uint NeighborId = 0; NeighborId < (CONFIG_SAMPLES - 1); NeighborId++)
@@ -409,7 +429,8 @@ void MainCS(
 				tsr_short2x2 SampleInputPixelPos;
 				tsr_half2x2 PixelOffset;
 				ComputeInputKernelSamplePosition(
-					InputPixelPos, dKO, /* SampleId = */ NeighborId + 1,
+					InputPixelPos, dInputKO, /* SampleId = */ NeighborId + 1,
+					InputPixelPosMin, InputPixelPosMax,
 					/* out */ SampleInputPixelPos,
 					/* out */ PixelOffset);
 			
@@ -426,7 +447,7 @@ void MainCS(
 			for (uint PixelId2 = 0; PixelId2 < DPV_PIXEL_PER_LANE; PixelId2++)
 			{
 				tsr_short2 LocalInputPixelPos0 = dpv_access_pixel(InputPixelPos, PixelId2);
-				tsr_short2 LocalInputPixelPos1 = LocalInputPixelPos0 + tsr_short2(sign(dpv_access_pixel(dKO, PixelId2)));
+				tsr_short2 LocalInputPixelPos1 = LocalInputPixelPos0 + tsr_short2(sign(dpv_access_pixel(dInputKO, PixelId2)));
 				LocalInputPixelPos1 = ClampPixelOffset(LocalInputPixelPos1, InputPixelPosMin, InputPixelPosMax);
 
 				RawEncodedVelocityNeighborhood[0][PixelId2] = DilatedVelocityTexture[tsr_short2(LocalInputPixelPos0.x, LocalInputPixelPos0.y)];
@@ -443,7 +464,7 @@ void MainCS(
 		PLATFORM_SPECIFIC_ISOLATE
 		{
 			// Vector in pixel between pixel K -> O.
-			tsr_half2x2 dKO = tsr_half2x2(PPCo - PPCk);
+			tsr_half2x2 dInputKO = tsr_half2x2(InputPPCo - InputPPCk);
 		
 			// Process velocity
 			{
@@ -457,7 +478,7 @@ void MainCS(
 					float2 V01 = DecodeVelocityFromTexture(float4(RawEncodedVelocityNeighborhood[2][PixelId], 0.0, 0.0)).xy;
 					float2 V11 = DecodeVelocityFromTexture(float4(RawEncodedVelocityNeighborhood[3][PixelId], 0.0, 0.0)).xy;
 
-					tsr_half2 BilinearInterp = abs(dpv_access_pixel(dKO, PixelId));
+					tsr_half2 BilinearInterp = abs(dpv_access_pixel(dInputKO, PixelId));
 					
 					float Bilinear00 = saturate(1.0 - BilinearInterp.x) * saturate(1.0 - BilinearInterp.y);
 					float Bilinear10 = saturate(0.0 + BilinearInterp.x) * saturate(1.0 - BilinearInterp.y);
@@ -526,10 +547,10 @@ void MainCS(
 			tsr_half2x2 TexelOffset = DecodeSpatialAntiAliasingOffset(EncodedInputTexelOffset);
 			
 			{
-				PPCo += dpv_scale(TexelOffset, tsr_half(1.0) - TotalRejection);
-				PPCk = floor(PPCo) + 0.5;
+				InputPPCo += dpv_scale(TexelOffset, tsr_half(1.0) - TotalRejection);
+				InputPPCk = floor(InputPPCo) + 0.5;
 				InputPixelPos = ClampPixelOffset(
-					tsr_short2x2(PPCo),
+					tsr_short2x2(InputPPCo),
 					InputPixelPosMin, InputPixelPosMax);
 			}
 		}
@@ -784,24 +805,24 @@ void MainCS(
 	
 	#if CONFIG_UPDATE_SUBPIXEL_DETAIL
 		bool2 bUpdateSubpixelDetails;
-		bool2x2 bKOSign;
+		bool2x2 bInputKOSign;
 	#endif
 
 	PLATFORM_SPECIFIC_ISOLATE
 	{
 		// Vector in pixel between pixel K -> O.
-		tsr_half2x2 dKO = tsr_half2x2(PPCo - PPCk);
+		tsr_half2x2 dInputKO = tsr_half2x2(InputPPCo - InputPPCk);
 
-		// Compresses the sign of dKO into SGPR to keep VGPR pressure low.
+		// Compresses the sign of dInputKO into SGPR to keep VGPR pressure low.
 		#if CONFIG_UPDATE_SUBPIXEL_DETAIL
 		{
 			float SizeOfOutputPixelInInputPixelWidth = 0.5 * (InputInfo_ViewportSize.x * HistoryInfo_ViewportSizeInverse.x);
 
-			bUpdateSubpixelDetails[0] = all(abs(dpv_lo(dKO)) < SizeOfOutputPixelInInputPixelWidth);
-			bUpdateSubpixelDetails[1] = all(abs(dpv_hi(dKO)) < SizeOfOutputPixelInInputPixelWidth);
+			bUpdateSubpixelDetails[0] = all(abs(dpv_lo(dInputKO)) < SizeOfOutputPixelInInputPixelWidth);
+			bUpdateSubpixelDetails[1] = all(abs(dpv_hi(dInputKO)) < SizeOfOutputPixelInInputPixelWidth);
 		
-			bKOSign[0] = dKO[0] < 0.0;
-			bKOSign[1] = dKO[1] < 0.0;
+			bInputKOSign[0] = dInputKO[0] < 0.0;
+			bInputKOSign[1] = dInputKO[1] < 0.0;
 		}
 		#endif
 		
@@ -814,6 +835,20 @@ void MainCS(
 		// Issues overlapped texture fetches
 		PLATFORM_SPECIFIC_ISOLATE
 		{
+			// Pixel coordinate of the center of output pixel O in the input viewport.
+			float2x2 TranslucencyPPCo = ApplyScreenTransform(float2x2(HistoryPixelPos), HistoryPixelPosToTranslucencyPPCo);
+
+			// Pixel coordinate of the center of the nearest input pixel K in the input viewport.
+			float2x2 TranslucencyPPCk = floor(TranslucencyPPCo) + 0.5;
+	
+			tsr_short2x2 TranslucencyPixelPos = ClampPixelOffset(
+				tsr_short2x2(TranslucencyPPCo),
+				TranslucencyPixelPosMin, TranslucencyPixelPosMax);
+				
+			// Vector in pixel between pixel K -> O.
+			tsr_half2x2 dTranslucencyKO = tsr_half2x2(TranslucencyPPCo - TranslucencyPPCk);
+
+
 			// Soften the spatial kernel when there is high noise in the input to make parrallax disocclusion of detailed geometry less distracting.
 			tsr_half2 OffScreenInputToHistoryFactor = tsr_half(1.0) - tsr_half(0.5) * NoiseFiltering;
 
@@ -825,31 +860,43 @@ void MainCS(
 			tsr_half2 KernelInputToHistoryFactor = lerp(OffScreenInputToHistoryFactor, tsr_half(InputToHistoryFactor), TotalRejection);
 			tsr_half2 KernelInputToHistoryAlignmentFactor = lerp(tsr_half(1.0).xx, tsr_half(InputToHistoryFactor), TotalRejection);
 			
-			//InputPixelAlignement = ComputeSampleWeigth(tsr_half(InputToHistoryFactor), dKO, /* MinimalContribution = */ float(0.0));
-			InputPixelAlignement = ComputeSampleWeigth(KernelInputToHistoryAlignmentFactor, dKO, /* MinimalContribution = */ float(0.0));
+			//InputPixelAlignement = ComputeSampleWeigth(tsr_half(InputToHistoryFactor), dInputKO, /* MinimalContribution = */ float(0.0));
+			InputPixelAlignement = ComputeSampleWeigth(KernelInputToHistoryAlignmentFactor, dInputKO, /* MinimalContribution = */ float(0.0));
 		
 			UNROLL_N(CONFIG_SAMPLES)
 			for (uint SampleId = 0; SampleId < CONFIG_SAMPLES; SampleId++)
 			{
-				tsr_short2x2 SampleInputPixelPos;
-				tsr_half2x2 PixelOffset;
+				tsr_short2x2 InputSamplePixelPos;
+				tsr_half2x2 InputPixelOffset;
 				ComputeInputKernelSamplePosition(
-					InputPixelPos, dKO, SampleId,
-					/* out */ SampleInputPixelPos,
-					/* out */ PixelOffset);
+					InputPixelPos, dInputKO, SampleId,
+					InputPixelPosMin, InputPixelPosMax,
+					/* out */ InputSamplePixelPos,
+					/* out */ InputPixelOffset);
 				
-				tsr_half2x2 dPP = PixelOffset - dKO;
+				tsr_short2x2 TranslucencySamplePixelPos;
+				tsr_half2x2 TranslucencyPixelOffset;
+				ComputeInputKernelSamplePosition(
+					TranslucencyPixelPos, dTranslucencyKO, SampleId,
+					TranslucencyPixelPosMin, TranslucencyPixelPosMax,
+					/* out */ TranslucencySamplePixelPos,
+					/* out */ TranslucencyPixelOffset);
+				
+				tsr_half2x2 dPP = InputPixelOffset - dInputKO;
 				tsr_half2 SampleSpatialWeight = ComputeSampleWeigth(KernelInputToHistoryFactor, dPP, /* MinimalContribution = */ float(0.005));
 				
+				tsr_half2x2 dTranslucencyPP = TranslucencyPixelOffset - dTranslucencyKO;
+				tsr_half2 TranslucencySampleSpatialWeight = ComputeSampleWeigth(KernelInputToHistoryFactor, dTranslucencyPP, /* MinimalContribution = */ float(0.005));
+				
 				RawdPPArray[SampleId] = dPP; 
 				RawSampleSpatialWeightArray[SampleId] = SampleSpatialWeight;
 
 				UNROLL_N(DPV_PIXEL_PER_LANE)
 				for (uint PixelId = 0; PixelId < DPV_PIXEL_PER_LANE; PixelId++)
 				{
-					RawInputColorArray[SampleId][PixelId] = InputSceneColorTexture[dpv_access_pixel(SampleInputPixelPos, PixelId)];
+					RawInputColorArray[SampleId][PixelId] = InputSceneColorTexture[dpv_access_pixel(InputSamplePixelPos, PixelId)];
 					#if CONFIG_ACCUMULATE_TRANSLUCENCY_SEPARATLY
-						RawTranslucencyColorArray[SampleId][PixelId] = InputSceneTranslucencyTexture[dpv_access_pixel(SampleInputPixelPos, PixelId)];
+						RawTranslucencyColorArray[SampleId][PixelId] = InputSceneTranslucencyTexture[dpv_access_pixel(TranslucencySamplePixelPos, PixelId)];
 					#endif
 				}
 			} // for (uint SampleId = 0; SampleId < CONFIG_SAMPLES; SampleId++)
@@ -1504,7 +1551,7 @@ void MainCS(
 			FinalSubpixelDetails = select(or(bIsOffScreen, bIsDisoccluded), FinalSubpixelDetails & ~ParallaxFactorMask, FinalSubpixelDetails);
 		}
 		
-		uint2 SubpixelId = dpv_dot(uint2x2(bKOSign), uint2(1, SUB_PIXEL_GRID_SIZE));
+		uint2 SubpixelId = dpv_dot(uint2x2(bInputKOSign), uint2(1, SUB_PIXEL_GRID_SIZE));
 
 		tsr_subpixel_payload2 SubpixelShift = tsr_subpixel_payload2(SubpixelId) * tsr_subpixel_payload(SUB_PIXEL_BIT_COUNT);
 
diff --git a/Engine/Source/Runtime/Renderer/Private/PostProcess/TemporalSuperResolution.cpp b/Engine/Source/Runtime/Renderer/Private/PostProcess/TemporalSuperResolution.cpp
--- a/Engine/Source/Runtime/Renderer/Private/PostProcess/TemporalSuperResolution.cpp
+++ b/Engine/Source/Runtime/Renderer/Private/PostProcess/TemporalSuperResolution.cpp
@@ -515,9 +515,15 @@ class FTSRUpdateHistoryCS : public FTSRShader
 		SHADER_PARAMETER_RDG_TEXTURE(Texture2D, AntiAliasingTexture)
 		SHADER_PARAMETER_RDG_TEXTURE(Texture2D, NoiseFilteringTexture)
 		SHADER_PARAMETER_RDG_TEXTURE(Texture2D, HoleFilledVelocityMaskTexture)
-		
+
+		SHADER_PARAMETER_STRUCT(FScreenPassTextureViewportParameters, TranslucencyInfo)
+		SHADER_PARAMETER(FIntPoint, TranslucencyPixelPosMin)
+		SHADER_PARAMETER(FIntPoint, TranslucencyPixelPosMax)
+
 		SHADER_PARAMETER(FScreenTransform, HistoryPixelPosToScreenPos)
-		SHADER_PARAMETER(FScreenTransform, HistoryPixelPosToPPCo)
+		SHADER_PARAMETER(FScreenTransform, HistoryPixelPosToInputPPCo)
+		SHADER_PARAMETER(FScreenTransform, HistoryPixelPosToTranslucencyPPCo)
+
 		SHADER_PARAMETER(FVector3f, HistoryQuantizationError)
 		SHADER_PARAMETER(float, MinTranslucencyRejection)
 		SHADER_PARAMETER(float, InvWeightClampingPixelSpeed)
@@ -1475,8 +1481,6 @@ ITemporalUpscaler::FOutputs AddTemporalSuperResolutionPasses(
 
 	// Update temporal history.
 	{
-		ensure(SeparateTranslucencyRect.Size() == InputRect.Size() || !bHasSeparateTranslucency);
-
 		static const TCHAR* const kUpdateQualityNames[] = {
 			TEXT("Low"),
 			TEXT("Medium"),
@@ -1501,9 +1505,15 @@ ITemporalUpscaler::FOutputs AddTemporalSuperResolutionPasses(
 		PassParameters->NoiseFilteringTexture = NoiseFilteringTexture;
 		PassParameters->HoleFilledVelocityMaskTexture = HoleFilledVelocityMaskTexture;
 
+		PassParameters->TranslucencyInfo = GetScreenPassTextureViewportParameters(FScreenPassTextureViewport(
+			PassParameters->TranslucencyRejectionTexture->Desc.Extent, SeparateTranslucencyRect));
+		PassParameters->TranslucencyPixelPosMin = PassParameters->TranslucencyInfo.ViewportMin;
+		PassParameters->TranslucencyPixelPosMax = PassParameters->TranslucencyInfo.ViewportMax - 1;
+
 		FScreenTransform HistoryPixelPosToViewportUV = (FScreenTransform::Identity + 0.5f) * CommonParameters.HistoryInfo.ViewportSizeInverse;
 		PassParameters->HistoryPixelPosToScreenPos = HistoryPixelPosToViewportUV * FScreenTransform::ViewportUVToScreenPos;
-		PassParameters->HistoryPixelPosToPPCo = HistoryPixelPosToViewportUV * CommonParameters.InputInfo.ViewportSize + CommonParameters.InputJitter + CommonParameters.InputPixelPosMin;
+		PassParameters->HistoryPixelPosToInputPPCo = HistoryPixelPosToViewportUV * CommonParameters.InputInfo.ViewportSize + CommonParameters.InputJitter + CommonParameters.InputPixelPosMin;
+		PassParameters->HistoryPixelPosToTranslucencyPPCo = HistoryPixelPosToViewportUV * PassParameters->TranslucencyInfo.ViewportSize + CommonParameters.InputJitter * PassParameters->TranslucencyInfo.ViewportSize / CommonParameters.InputInfo.ViewportSize + SeparateTranslucencyRect.Min;
 		PassParameters->HistoryQuantizationError = (FVector3f)ComputePixelFormatQuantizationError(History.ColorArray->Desc.Format);
 		PassParameters->MinTranslucencyRejection = TranslucencyRejectionTexture == nullptr ? 1.0 : 0.0;
 		PassParameters->InvWeightClampingPixelSpeed = 1.0f / CVarTSRWeightClampingPixelSpeed.GetValueOnRenderThread();
diff --git a/Engine/Source/Runtime/Renderer/Private/TranslucentRendering.cpp b/Engine/Source/Runtime/Renderer/Private/TranslucentRendering.cpp
--- a/Engine/Source/Runtime/Renderer/Private/TranslucentRendering.cpp
+++ b/Engine/Source/Runtime/Renderer/Private/TranslucentRendering.cpp
@@ -19,6 +19,8 @@
 #include "PixelShaderUtils.h"
 #include "OIT/OIT.h"
 #include "OIT/OITParameters.h"
+#include "DynamicResolutionState.h"
+#include "TemporalAA.h"
 
 DECLARE_CYCLE_STAT(TEXT("TranslucencyTimestampQueryFence Wait"), STAT_TranslucencyTimestampQueryFence_Wait, STATGROUP_SceneRendering);
 DECLARE_CYCLE_STAT(TEXT("TranslucencyTimestampQuery Wait"), STAT_TranslucencyTimestampQuery_Wait, STATGROUP_SceneRendering);
@@ -34,6 +36,13 @@
 	TEXT("<0 is treated like 100."),
 	ECVF_Scalability | ECVF_Default);
 
+static TAutoConsoleVariable<int32> CVarTranslucencyScreenPercentageBasis(
+	TEXT("r.Translucency.ScreenPercentage.Basis"), 0,
+	TEXT("Basis of the translucency's screen percentage (Experimental).\n")
+	TEXT(" 0: Uses the primary view's resolution (notably scaling with r.ScreenPercentage and r.DynamicRes.*)\n")
+	TEXT(" 1: Uses the secondary view's resolution (temporal upscale's output resolution)"),
+	ECVF_Scalability | ECVF_Default);
+
 static TAutoConsoleVariable<float> CVarTranslucencyMinScreenPercentage(
 	TEXT("r.Translucency.DynamicRes.MinScreenPercentage"),
 	DynamicRenderScaling::FractionToPercentage(DynamicRenderScaling::FHeuristicSettings::kDefaultMinResolutionFraction),
@@ -82,7 +91,7 @@
 {
 	DynamicRenderScaling::FHeuristicSettings BucketSetting;
 	BucketSetting.Model = DynamicRenderScaling::EHeuristicModel::Quadratic;
-	BucketSetting.bModelScalesWithPrimaryScreenPercentage = true;
+	BucketSetting.bModelScalesWithPrimaryScreenPercentage = CVarTranslucencyScreenPercentageBasis.GetValueOnAnyThread() != 1;
 	BucketSetting.MinResolutionFraction = DynamicRenderScaling::GetPercentageCVarToFraction(CVarTranslucencyMinScreenPercentage);
 	BucketSetting.BudgetMs              = CVarTranslucencyTimeBudget.GetValueOnAnyThread();
 	BucketSetting.ChangeThreshold       = DynamicRenderScaling::GetPercentageCVarToFraction(CVarTranslucencyChangeThreshold);
@@ -229,6 +238,12 @@
 		MaxTranslucencyResolutionFraction = ActiveViewFamily.DynamicResolutionUpperBounds[GDynamicTranslucencyResolution];
 	}
 
+	if (CVarTranslucencyScreenPercentageBasis.GetValueOnRenderThread() == 1)
+	{
+		TranslucencyResolutionFraction /= ActiveViewFamily.DynamicResolutionFractions[GDynamicPrimaryResolutionFraction];
+		MaxTranslucencyResolutionFraction /= ActiveViewFamily.DynamicResolutionUpperBounds[GDynamicPrimaryResolutionFraction];
+	}
+
 	FSeparateTranslucencyDimensions Dimensions;
 	// TODO: this should be MaxTranslucencyResolutionFraction instead of TranslucencyResolutionFraction to keep the size of render target stable, but the SvPositionToBuffer() is broken in material.
 	Dimensions.Extent = GetScaledExtent(ActiveViewFamily.SceneTexturesConfig.Extent, TranslucencyResolutionFraction);
@@ -1207,7 +1222,6 @@
 		return;
 	}
 
-	RDG_EVENT_SCOPE(GraphBuilder, "%s", TranslucencyPassToString(TranslucencyPass));
 	RDG_WAIT_FOR_TASKS_CONDITIONAL(GraphBuilder, IsTranslucencyWaitForTasksEnabled());
 
 	const bool bIsModulate = TranslucencyPass == ETranslucencyPass::TPT_TranslucencyAfterDOFModulate;
@@ -1413,7 +1427,7 @@
 
 	// Create a shared depth texture at the correct resolution.
 	FRDGTextureMSAA SharedDepthTexture;
-	const bool bIsScalingTranslucency = SeparateTranslucencyDimensions.Scale < 1.0f;
+	const bool bIsScalingTranslucency = SeparateTranslucencyDimensions.Scale != 1.0f;
 	if (bIsScalingTranslucency)
 	{
 		const FRDGTextureDesc Desc = FRDGTextureDesc::Create2D(
@@ -1510,7 +1524,7 @@
 		}
 
 		FTranslucencyPassResources& TranslucencyPassResources = OutTranslucencyResourceMap->Get(ViewIndex, ETranslucencyPass::TPT_TranslucencyAfterDOF);
-		if (SharedUpscaledPostDOFTranslucencyColor && TranslucencyPassResources.IsValid() && TranslucencyPassResources.ViewRect.Size() != View.ViewRect.Size())
+		if (SharedUpscaledPostDOFTranslucencyColor && TranslucencyPassResources.IsValid() && TranslucencyPassResources.ViewRect.Size() != View.ViewRect.Size() && ITemporalUpscaler::GetMainTAAPassConfig(View) != EMainTAAPassConfig::TSR)
 		{
 			FTranslucencyComposition TranslucencyComposition;
 			TranslucencyComposition.Operation = FTranslucencyComposition::EOperation::UpscaleOnly;
