Refactored bent normal to be simpler for future Strata changes

diff --git a/Engine/Shaders/Private/BasePassPixelShader.usf b/Engine/Shaders/Private/BasePassPixelShader.usf
--- a/Engine/Shaders/Private/BasePassPixelShader.usf
+++ b/Engine/Shaders/Private/BasePassPixelShader.usf
@@ -797,10 +797,14 @@ float DotSpecularSG( float Roughness, float3 N, float3 V, FSphericalGaussian Lig
 #endif
 }
 
-float3 GetBentNormalZero(in FMaterialPixelParameters MaterialParameters)
+float3 GetWorldBentNormalZero(in FMaterialPixelParameters MaterialParameters)
 {
 #if NUM_MATERIAL_OUTPUTS_GETBENTNORMAL > 0
-	return GetBentNormal0(MaterialParameters);
+	#if MATERIAL_TANGENTSPACENORMAL
+		return normalize(TransformTangentVectorToWorld(MaterialParameters.TangentToWorld, GetBentNormal0(MaterialParameters)));
+	#else
+		return GetBentNormal0(MaterialParameters);
+	#endif
 #else
 	return MaterialParameters.WorldNormal;
 #endif
@@ -809,8 +813,7 @@ float3 GetBentNormalZero(in FMaterialPixelParameters MaterialParameters)
 FShadingOcclusion ApplyBentNormal(
 	in float3 CameraVector, 
 	in float3 WorldNormal, 
-	in float3x3 TangentToWorld, 
-	in float3 BentNormal0, 
+	in float3 WorldBentNormal0, 
 	in float Roughness,
 	in float MaterialAO)
 {
@@ -820,11 +823,7 @@ FShadingOcclusion ApplyBentNormal(
 	Out.BentNormal    = WorldNormal;
 
 #if NUM_MATERIAL_OUTPUTS_GETBENTNORMAL > 0
-	#if MATERIAL_TANGENTSPACENORMAL
-	Out.BentNormal = normalize( TransformTangentVectorToWorld(TangentToWorld, BentNormal0) );
-	#else
-	Out.BentNormal = BentNormal0;
-	#endif
+	Out.BentNormal = WorldBentNormal0;
 
 	FSphericalGaussian HemisphereSG = Hemisphere_ToSphericalGaussian(WorldNormal);
 	FSphericalGaussian NormalSG = ClampedCosine_ToSphericalGaussian(WorldNormal);
@@ -850,19 +849,6 @@ FShadingOcclusion ApplyBentNormal(
 	return Out;
 }
 
-float3 GetBentNormal(in FMaterialPixelParameters MaterialParameters)
-{
-#if NUM_MATERIAL_OUTPUTS_GETBENTNORMAL > 0
-	#if MATERIAL_TANGENTSPACENORMAL
-	return normalize(TransformTangentVectorToWorld(MaterialParameters.TangentToWorld, GetBentNormalZero(MaterialParameters)));
-	#else
-	return GetBentNormalZero(MaterialParameters);
-	#endif
-#else
-	return MaterialParameters.WorldNormal;
-#endif
-}
-
 #if STRATA_ENABLED
 uint GetDiffuseIndirectSampleOcclusion(FSharedLocalBases SharedLocalBases, float3 V, float2 SvPosition, float MaterialAO)
 {
@@ -901,14 +887,14 @@ uint GetDiffuseIndirectSampleOcclusion(FSharedLocalBases SharedLocalBases, float
 	return DiffuseIndirectSampleOcclusion;
 }
 #else // STRATA_ENABLED
-uint GetDiffuseIndirectSampleOcclusion(FGBufferData GBuffer, float3 V, float3 WorldNormal, float3 BentNormal, float2 SvPosition, float MaterialAO)
+uint GetDiffuseIndirectSampleOcclusion(FGBufferData GBuffer, float3 V, float3 WorldNormal, float3 WorldBentNormal, float2 SvPosition, float MaterialAO)
 {
 	uint DiffuseIndirectSampleOcclusion = 0;
 #if GBUFFER_HAS_DIFFUSE_SAMPLE_OCCLUSION && !MATERIAL_SHADINGMODEL_UNLIT
 	if (MaterialAO < 1.0f)
 	{
 		FSphericalGaussian HemisphereSG = Hemisphere_ToSphericalGaussian(WorldNormal);
-		FSphericalGaussian VisibleSG = BentNormalAO_ToSphericalGaussian(BentNormal, MaterialAO);
+		FSphericalGaussian VisibleSG = BentNormalAO_ToSphericalGaussian(WorldBentNormal, MaterialAO);
 		float VisibilityThreshold = InterleavedGradientNoise(SvPosition, View.StateFrameIndexMod8);
 
 		uint OcclusionMask = 0x0;
@@ -1328,12 +1314,12 @@ void FPixelShaderInOut_MainPS(
 		InputBentNormal = OctahedronToUnitVector(oct1);
 	}
 	
-	const FShadingOcclusion ShadingOcclusion = ApplyBentNormal(MaterialParameters.CameraVector, InputBentNormal, MaterialParameters.TangentToWorld, GetBentNormalZero(MaterialParameters), GBuffer.Roughness, MaterialAO);
+	const FShadingOcclusion ShadingOcclusion = ApplyBentNormal(MaterialParameters.CameraVector, InputBentNormal, GetWorldBentNormalZero(MaterialParameters), GBuffer.Roughness, MaterialAO);
 
 	// FIXME: ALLOW_STATIC_LIGHTING == 0 expects this to be AO
 	GBuffer.GBufferAO = AOMultiBounce( Luminance( GBuffer.SpecularColor ), ShadingOcclusion.SpecOcclusion ).g;
 	#if !STRATA_INLINE_SINGLELAYERWATER
-	GBuffer.DiffuseIndirectSampleOcclusion = GetDiffuseIndirectSampleOcclusion(GBuffer, MaterialParameters.CameraVector, MaterialParameters.WorldNormal, GetBentNormal(MaterialParameters), In.SvPosition.xy, MaterialAO);
+	GBuffer.DiffuseIndirectSampleOcclusion = GetDiffuseIndirectSampleOcclusion(GBuffer, MaterialParameters.CameraVector, MaterialParameters.WorldNormal, GetWorldBentNormalZero(MaterialParameters), In.SvPosition.xy, MaterialAO);
 	#endif
 #endif // !STRATA_ENABLED
 
