Strata - added support for precomputed lighting to the deferred base opaque pass shaders.

diff --git a/Engine/Shaders/Private/BasePassPixelShader.usf b/Engine/Shaders/Private/BasePassPixelShader.usf
--- a/Engine/Shaders/Private/BasePassPixelShader.usf
+++ b/Engine/Shaders/Private/BasePassPixelShader.usf
@@ -44,8 +44,37 @@
 #define STRATA_ENABLED 0
 #endif
 
+struct FShadingOcclusion
+{
+	float  DiffOcclusion;
+	float  SpecOcclusion;
+	float3 BentNormal;
+};
+
+#define MATERIAL_STRATA_PRECOMPUTED_LIGHTING (MATERIAL_IS_STRATA && STRATA_ENABLED)
+
+#if MATERIAL_STRATA_PRECOMPUTED_LIGHTING
+#if STRATA_CLAMPED_BSDF_COUNT <= 0
+#define STRATA_BSDF_COUNT_FOR_PRECOMPUTED 1
+#else
+#define STRATA_BSDF_COUNT_FOR_PRECOMPUTED STRATA_CLAMPED_BSDF_COUNT
+#endif
+struct FStrataBSDFForPrecomputedLighting
+{
+	float3 DiffuseNormal;// STRATA_TODO test packed on a uint4 as R11G11B10
+	float3 DiffuseColor;
+	float3 DiffuseColorForIndirect;
+	float3 SubsurfaceColorForIndirect;
+	FShadingOcclusion ShadingOcclusion;
+};
+struct FStrataPrecomputedLightingData
+{
+	uint BSDFCount;
+	FStrataBSDFForPrecomputedLighting BSDFs[STRATA_BSDF_COUNT_FOR_PRECOMPUTED];
+};
+#endif
+
 // For the base pass STRATA_INLINE_SHADING is defined in BasePassRendering.h
-#include "Strata/Strata.ush"
 #include "SHCommon.ush"
 #include "/Engine/Generated/Material.ush"
 #include "BasePassCommon.ush"
@@ -171,15 +200,10 @@ float NormalCurvatureToRoughness(float3 WorldNormal)
 #include "ForwardLightingCommon.ush"
 #endif
 
-struct FShadingOcclusion
-{
-	float  DiffOcclusion;
-	float  SpecOcclusion;
-	float3 BentNormal;
-};
-
-#if STRATA_TRANSLUCENT_FORWARD || STRATA_FORWARD_SHADING
+#if STRATA_TRANSLUCENT_FORWARD || STRATA_FORWARD_SHADING || (MATERIAL_STRATA_PRECOMPUTED_LIGHTING && STRATA_OPAQUE_DEFERRED)
 #include "/Engine/Private/Strata/StrataEvaluation.ush"
+#endif
+#if STRATA_TRANSLUCENT_FORWARD || STRATA_FORWARD_SHADING
 #include "/Engine/Private/Strata/StrataForwardLighting.ush"
 #endif
 
@@ -850,6 +874,70 @@ FShadingOcclusion ApplyBentNormal(
 }
 
 #if STRATA_ENABLED
+
+#if MATERIAL_STRATA_PRECOMPUTED_LIGHTING
+FStrataBSDFForPrecomputedLighting StrataPrepareBSDFForPrecomputedLighting(
+	in FStrataIntegrationSettings Settings,
+	in FStrataPixelHeader StrataPixelHeader,
+	in FStrataBSDF CurrentBSDF,
+	in float3 V,
+	in float3 WorldBentNormal0)
+{
+	FStrataBSDFForPrecomputedLighting BSDFForPrecomputedLighting = (FStrataBSDFForPrecomputedLighting)0;
+
+	FStrataAddressing NullStrataAddressing = (FStrataAddressing)0;
+	FStrataBSDFContext StrataBSDFContext = StrataCreateBSDFContext(StrataPixelHeader, CurrentBSDF, NullStrataAddressing, V);
+
+	// Update a context specific for environment lighting
+	FStrataBSDFContext EnvBSDFContext = StrataBSDFContext;
+	EnvBSDFContext.N = BSDFForPrecomputedLighting.ShadingOcclusion.BentNormal;
+	StrataUpdateBSDFContext(EnvBSDFContext, EnvBSDFContext.L);
+
+	// Evaluate the bent normal
+	BSDFForPrecomputedLighting.ShadingOcclusion = ApplyBentNormal(
+		V,
+		EnvBSDFContext.N,
+		WorldBentNormal0,
+		StrataGetBSDFRoughness(CurrentBSDF),
+		StrataPixelHeader.IrradianceAO.MaterialAO);
+
+	// For diffuse, we specify a perpendicular to the surface light direction for the transmittance to light to not be view dependent.
+	float DiffuseEnvLightingNoL = 1.0f;
+	float3 LuminanceWeightFinal = LuminanceWeight(DiffuseEnvLightingNoL, CurrentBSDF);
+
+	// And evaluate diffuse parameters
+	const bool bEnableSpecular = false;
+	FStrataEnvLightResult StrataEnvLight = StrataEvaluateForEnvLight(EnvBSDFContext, bEnableSpecular, Settings);
+	BSDFForPrecomputedLighting.DiffuseNormal = StrataEnvLight.DiffuseNormal;
+	BSDFForPrecomputedLighting.DiffuseColor = StrataEnvLight.DiffuseColor * LuminanceWeightFinal;
+	BSDFForPrecomputedLighting.DiffuseColorForIndirect = StrataEnvLight.DiffuseWeight * LuminanceWeightFinal;
+	BSDFForPrecomputedLighting.SubsurfaceColorForIndirect = StrataEnvLight.DiffuseBackFaceWeight * LuminanceWeightFinal;
+
+	return BSDFForPrecomputedLighting;
+}
+
+float3 StrataGetBSDFPrecomputedLighting(
+	in FStrataBSDFForPrecomputedLighting BSDFForPrecomputedLighting,
+	in FMaterialPixelParameters MaterialParameters,
+	in FVertexFactoryInterpolantsVSToPS Interpolants,
+	in FBasePassInterpolantsVSToPS BasePassInterpolants,
+	in VTPageTableResult LightmapVTPageTableResult,
+	in float3 VolumetricLightmapBrickTextureUVs)
+{
+	float3 DiffuseLighting = 0.0f;
+	float3 SubsurfaceLighting = 0.0f;
+	float  IndirectIrrad1iance = 0.0f;
+	const bool bEvaluateBackface = any(BSDFForPrecomputedLighting.SubsurfaceColorForIndirect > 0.0);	// With Strata we have no other way to know before if backface lighting will be needed
+	GetPrecomputedIndirectLightingAndSkyLight(MaterialParameters, Interpolants, BasePassInterpolants, LightmapVTPageTableResult, bEvaluateBackface, BSDFForPrecomputedLighting.DiffuseNormal, VolumetricLightmapBrickTextureUVs,
+		DiffuseLighting, SubsurfaceLighting, IndirectIrrad1iance);
+
+	float3 OutLuminance = (DiffuseLighting * BSDFForPrecomputedLighting.DiffuseColorForIndirect + SubsurfaceLighting * BSDFForPrecomputedLighting.SubsurfaceColorForIndirect)
+		* AOMultiBounce(BSDFForPrecomputedLighting.DiffuseColor, BSDFForPrecomputedLighting.ShadingOcclusion.DiffOcclusion); // STRATA_TODO This should only be applied on the aggregated diffuse colors, not each contribution independently. But we also need to account for DiffOcclusion which looks normal dependent
+
+	return OutLuminance;
+}
+#endif
+
 uint GetDiffuseIndirectSampleOcclusion(FSharedLocalBases SharedLocalBases, float3 V, float2 SvPosition, float MaterialAO)
 {
 	uint DiffuseIndirectSampleOcclusion = 0;
@@ -1723,7 +1811,11 @@ void FPixelShaderInOut_MainPS(
 	{
 		float3 EmissiveLuminance = 0.0f;
 		uint2 PixelPos = uint2(In.SvPosition.xy);
+		const float3 WorldBentNormal0 = GetWorldBentNormalZero(MaterialParameters);
 
+	#if MATERIAL_STRATA_PRECOMPUTED_LIGHTING
+		FStrataPrecomputedLightingData PrecomputedLightingData = (FStrataPrecomputedLightingData)0;
+	#endif
 		FStrataSubsurfaceData SSSData = (FStrataSubsurfaceData)0;
 		FStrataTopLayerData TopLayerData = (FStrataTopLayerData)0;
 		FStrataOpaqueRoughRefractionData OpaqueRoughRefractionData = (FStrataOpaqueRoughRefractionData)0;
@@ -1742,8 +1834,26 @@ void FPixelShaderInOut_MainPS(
 			Dither,
 			Settings,
 			StrataAddressing,
-			StrataPixelHeader, StrataData, MaterialParameters.CameraVector, bStrataSubsurfaceEnable, EmissiveLuminance,
-			SSSData, TopLayerData, OpaqueRoughRefractionData);
+			StrataPixelHeader, StrataData, MaterialParameters.CameraVector, WorldBentNormal0, bStrataSubsurfaceEnable, EmissiveLuminance,
+			SSSData, TopLayerData, OpaqueRoughRefractionData
+	#if MATERIAL_STRATA_PRECOMPUTED_LIGHTING
+			, PrecomputedLightingData
+	#endif
+		);
+
+		float3 PreComputedLightingLuminance = 0.0f;
+	#if MATERIAL_STRATA_PRECOMPUTED_LIGHTING
+		for (uint i = 0; i < PrecomputedLightingData.BSDFCount; ++i)
+		{
+			PreComputedLightingLuminance += StrataGetBSDFPrecomputedLighting(
+				PrecomputedLightingData.BSDFs[i],
+				MaterialParameters,
+				Interpolants,
+				BasePassInterpolants,
+				LightmapVTPageTableResult,
+				VolumetricLightmapBrickTextureUVs);
+		}
+	#endif
 
 		// Write out MRT data
 	#if STRATA_BASE_PASS_MRT_OUTPUT_COUNT != 2
@@ -1781,7 +1891,7 @@ void FPixelShaderInOut_MainPS(
 		Color = lerp(Color, TopLayerData.BaseColor, View.UnlitViewmodeMask);
 	#endif
 
-		Color += EmissiveLuminance;
+		Color += EmissiveLuminance + PreComputedLightingLuminance;
 	}
 
 #endif // STRATA_OPAQUE_DEFERRED
diff --git a/Engine/Shaders/Private/DebugProbes.usf b/Engine/Shaders/Private/DebugProbes.usf
--- a/Engine/Shaders/Private/DebugProbes.usf
+++ b/Engine/Shaders/Private/DebugProbes.usf
@@ -278,7 +278,8 @@ void MainPS(
 		SetCastContactShadow(StrataPixelHeader, 0);
 
 		uint SharedLocalBasisIndex = 0;
-		StrataPixelHeader.SharedLocalBases.Normals[SharedLocalBasisIndex] = normalize(Normal);
+		const float3 NormalizedNormal = normalize(Normal);
+		StrataPixelHeader.SharedLocalBases.Normals[SharedLocalBasisIndex] = NormalizedNormal;
 		StrataPixelHeader.SharedLocalBases.Count = 1;
 
 		FStrataPixelFootprint PixelFootprint = (FStrataPixelFootprint)0;
@@ -329,7 +330,7 @@ void MainPS(
 			Dither,
 			Settings,
 			StrataAddressing,
-			StrataPixelHeader, StrataData, -RayDir, bStrataSubsurfaceEnable, EmissiveLuminance,
+			StrataPixelHeader, StrataData, -RayDir, NormalizedNormal, bStrataSubsurfaceEnable, EmissiveLuminance,
 			SSSData, TopLayerData, OpaqueRoughRefractionData);
 
 		// Write out MRTs. Note: this only works for deferred with basepassvelocity off
diff --git a/Engine/Shaders/Private/Strata/Strata.ush b/Engine/Shaders/Private/Strata/Strata.ush
--- a/Engine/Shaders/Private/Strata/Strata.ush
+++ b/Engine/Shaders/Private/Strata/Strata.ush
@@ -42,6 +42,12 @@
 #define STRATA_INLINE_SHADING 0
 #endif
 
+// This is used to evaluate some precomputed lighting data while we pack out Strata material.
+// It is also used to remove reference to unwanted resourced when not needed.
+#ifndef MATERIAL_STRATA_PRECOMPUTED_LIGHTING
+#define MATERIAL_STRATA_PRECOMPUTED_LIGHTING 0
+#endif
+
 #ifndef STRATA_DEFERRED_SHADING
 #define STRATA_DEFERRED_SHADING (STRATA_INLINE_SHADING == 0)
 #endif 
@@ -3130,6 +3136,15 @@ void StrataUpdateTree(
 
 uint PackStrataHeader(uint BSDFCount, FStrataPixelHeader In);
 
+#if MATERIAL_STRATA_PRECOMPUTED_LIGHTING
+FStrataBSDFForPrecomputedLighting StrataPrepareBSDFForPrecomputedLighting(
+	in FStrataIntegrationSettings Settings,
+	in FStrataPixelHeader StrataPixelHeader,
+	in FStrataBSDF CurrentBSDF,
+	in float3 V,
+	in float3 WorldBentNormal0);
+#endif
+
 // If this is changed, please update the compiler side material size evaluation in StrataMaterial.cpp
 void PackStrataOut(
 	inout FRWStrataMaterialContainer StrataBuffer,
@@ -3142,15 +3157,25 @@ void PackStrataOut(
 	FStrataPixelHeader StrataPixelHeader,
 	FStrataData Strata,
 	float3 V,
+	float3 WorldBentNormal0,
 	inout bool bStrataSubsurfaceEnable,
 	inout float3 EmissiveColor,
 	inout FStrataSubsurfaceData SSSData,
 	inout FStrataTopLayerData TopLayerData,
-	inout FStrataOpaqueRoughRefractionData OpaqueRoughRefractionData)
+	inout FStrataOpaqueRoughRefractionData OpaqueRoughRefractionData
+#if MATERIAL_STRATA_PRECOMPUTED_LIGHTING
+	,
+	inout FStrataPrecomputedLightingData IndirectLightingData
+#endif
+)
 {
 	// This only exists with inline shading and when we are going to write out BSDFs (UpdateAllXXX functions needs to be defined)
 #if STRATA_INLINE_SHADING && STRATA_CLAMPED_BSDF_COUNT > 0
 
+#if MATERIAL_STRATA_PRECOMPUTED_LIGHTING
+	IndirectLightingData = (FStrataPrecomputedLightingData)0;
+#endif 
+
 	bStrataSubsurfaceEnable = false;
 	EmissiveColor = 0.0f;
 
@@ -3515,6 +3540,17 @@ void PackStrataOut(
 						}
 					}
 
+					// Convert BSDF to data ready to evaluate precomputed lighting
+					#if MATERIAL_STRATA_PRECOMPUTED_LIGHTING
+					IndirectLightingData.BSDFs[IndirectLightingData.BSDFCount++] = StrataPrepareBSDFForPrecomputedLighting(
+						Settings,
+						StrataPixelHeader,
+						BSDF,
+						V,
+						WorldBentNormal0);
+					#endif
+
+
 					// Stop writing if above budget.
 					BSDFCount++;
 				}
diff --git a/Engine/Shaders/Private/Strata/StrataForwardLighting.ush b/Engine/Shaders/Private/Strata/StrataForwardLighting.ush
--- a/Engine/Shaders/Private/Strata/StrataForwardLighting.ush
+++ b/Engine/Shaders/Private/Strata/StrataForwardLighting.ush
@@ -10,13 +10,13 @@
 // Forward declarations
 void GetVolumeLightingNonDirectional(float4 AmbientLightingVector, float3 DiffuseColor, inout float3 InterpolatedLighting, out float4 VolumeLighting);
 void GetVolumeLightingDirectional(float4 AmbientLightingVector, float3 DirectionalLightingVector, float3 WorldNormal, float3 DiffuseColor, inout float3 InterpolatedLighting, out float4 VolumeLighting);
-float3 GetBentNormalZero(in FMaterialPixelParameters MaterialParameters);
+
+float3 GetWorldBentNormalZero(in FMaterialPixelParameters MaterialParameters);
 
 FShadingOcclusion ApplyBentNormal(
 	in float3 CameraVector,
 	in float3 WorldNormal,
-	in float3x3 TangentToWorld,
-	in float3 BentNormal0,
+	in float3 WorldBentNormal0,
 	in float Roughness,
 	in float MaterialAO);
 
@@ -160,7 +160,7 @@ float3 StrataForwardLighting(
 
 				// Compute the bent normal from the BSDF normal
 				const float Roughness = StrataGetBSDFRoughness(CurrentBSDF);
-				const FShadingOcclusion ShadingOcclusion = ApplyBentNormal(MaterialParameters.CameraVector, StrataBSDFContext.N, MaterialParameters.TangentToWorld, GetBentNormalZero(MaterialParameters), Roughness,  StrataGetAO(StrataPixelHeader));
+				const FShadingOcclusion ShadingOcclusion = ApplyBentNormal(MaterialParameters.CameraVector, StrataBSDFContext.N, GetWorldBentNormalZero(MaterialParameters), Roughness,  StrataGetAO(StrataPixelHeader));
 
 				float IndirectOcclusion = 1.0f;
 				float IndirectIrradiance = 0.0f; 
