Strata - split out export function from strata.ush. (PackStrataOut only for now).

diff --git a/Engine/Shaders/Private/BasePassPixelShader.usf b/Engine/Shaders/Private/BasePassPixelShader.usf
--- a/Engine/Shaders/Private/BasePassPixelShader.usf
+++ b/Engine/Shaders/Private/BasePassPixelShader.usf
@@ -200,6 +200,9 @@ float NormalCurvatureToRoughness(float3 WorldNormal)
 #include "ForwardLightingCommon.ush"
 #endif
 
+#if STRATA_OPAQUE_DEFERRED
+#include "/Engine/Private/Strata/StrataExport.ush"
+#endif
 #if STRATA_TRANSLUCENT_FORWARD || STRATA_FORWARD_SHADING || (MATERIAL_STRATA_PRECOMPUTED_LIGHTING && STRATA_OPAQUE_DEFERRED)
 #include "/Engine/Private/Strata/StrataEvaluation.ush"
 #endif
@@ -821,19 +824,6 @@ float DotSpecularSG( float Roughness, float3 N, float3 V, FSphericalGaussian Lig
 #endif
 }
 
-float3 GetWorldBentNormalZero(in FMaterialPixelParameters MaterialParameters)
-{
-#if NUM_MATERIAL_OUTPUTS_GETBENTNORMAL > 0
-	#if MATERIAL_TANGENTSPACENORMAL
-		return normalize(TransformTangentVectorToWorld(MaterialParameters.TangentToWorld, GetBentNormal0(MaterialParameters)));
-	#else
-		return GetBentNormal0(MaterialParameters);
-	#endif
-#else
-	return MaterialParameters.WorldNormal;
-#endif
-}
-
 FShadingOcclusion ApplyBentNormal(
 	in float3 CameraVector, 
 	in float3 WorldNormal, 
diff --git a/Engine/Shaders/Private/DebugProbes.usf b/Engine/Shaders/Private/DebugProbes.usf
--- a/Engine/Shaders/Private/DebugProbes.usf
+++ b/Engine/Shaders/Private/DebugProbes.usf
@@ -7,6 +7,9 @@
 #define STRATA_INLINE_SHADING 1
 #include "Strata/Strata.ush"
 
+#define MATERIAL_STRATA_PRECOMPUTED_LIGHTING 0
+#include "Strata/StrataExport.ush"
+
 #include "DeferredShadingCommon.ush"
 #include "VelocityCommon.ush"
 
diff --git a/Engine/Shaders/Private/RayTracing/RayTracingMaterialHitShaders.usf b/Engine/Shaders/Private/RayTracing/RayTracingMaterialHitShaders.usf
--- a/Engine/Shaders/Private/RayTracing/RayTracingMaterialHitShaders.usf
+++ b/Engine/Shaders/Private/RayTracing/RayTracingMaterialHitShaders.usf
@@ -34,6 +34,8 @@
 #include "/Engine/Private/LightmapCommon.ush"
 #include "/Engine/Private/Lumen/LumenHardwareRayTracingPayloadCommon.ush"
 
+#define MATERIAL_STRATA_PRECOMPUTED_LIGHTING 0
+#include "/Engine/Private/Strata/StrataExport.ush"
 
 /** Computes sky diffuse lighting, including precomputed shadowing. */
 void GetSkyLighting(
@@ -247,6 +249,7 @@ void GetMaterialPayload(
 	const float Dither = 0.5f;
 	float3 EmissiveLuminance = 0.0f;
 	bool bStrataSubsurfaceEnable = false;
+	const float3 WorldBentNormal0 = GetWorldBentNormalZero(MaterialParameters);
 	FStrataData StrataData = PixelMaterialInputs.FrontMaterial;
 
 	// Initialise a Strata header with normal in registers
@@ -276,6 +279,7 @@ void GetMaterialPayload(
 		StrataAddressing,
 		StrataPixelHeader, StrataData, 
 		CamVector, 
+		WorldBentNormal0,
 		bStrataSubsurfaceEnable, 
 		EmissiveLuminance,
 		SSSData, TopLayerData, OpaqueRoughRefractionData
diff --git a/Engine/Shaders/Private/Strata/Strata.ush b/Engine/Shaders/Private/Strata/Strata.ush
--- a/Engine/Shaders/Private/Strata/Strata.ush
+++ b/Engine/Shaders/Private/Strata/Strata.ush
@@ -3134,439 +3134,8 @@ void StrataUpdateTree(
 
 #endif // STRATA_INLINE_SHADING && STRATA_CLAMPED_BSDF_COUNT > 0
 
-uint PackStrataHeader(uint BSDFCount, FStrataPixelHeader In);
-
-#if MATERIAL_STRATA_PRECOMPUTED_LIGHTING
-FStrataBSDFForPrecomputedLighting StrataPrepareBSDFForPrecomputedLighting(
-	in FStrataIntegrationSettings Settings,
-	in FStrataPixelHeader StrataPixelHeader,
-	in FStrataBSDF CurrentBSDF,
-	in float3 V,
-	in float3 WorldBentNormal0);
-#endif
-
-// If this is changed, please update the compiler side material size evaluation in StrataMaterial.cpp
-void PackStrataOut(
-	inout FRWStrataMaterialContainer StrataBuffer,
-#if STRATA_USES_RESOURCE_MATERIALCONTAINER
-	RWTexture2DArray<uint> ExtraMaterialDataUAV,
-#endif
-	float Dither,
-	FStrataIntegrationSettings Settings,
-	inout FStrataAddressing StrataAddressing,
-	FStrataPixelHeader StrataPixelHeader,
-	FStrataData Strata,
-	float3 V,
-	float3 WorldBentNormal0,
-	inout bool bStrataSubsurfaceEnable,
-	inout float3 EmissiveColor,
-	inout FStrataSubsurfaceData SSSData,
-	inout FStrataTopLayerData TopLayerData,
-	inout FStrataOpaqueRoughRefractionData OpaqueRoughRefractionData
-#if MATERIAL_STRATA_PRECOMPUTED_LIGHTING
-	,
-	inout FStrataPrecomputedLightingData IndirectLightingData
-#endif
-)
-{
-	// This only exists with inline shading and when we are going to write out BSDFs (UpdateAllXXX functions needs to be defined)
-#if STRATA_INLINE_SHADING && STRATA_CLAMPED_BSDF_COUNT > 0
-
-#if MATERIAL_STRATA_PRECOMPUTED_LIGHTING
-	IndirectLightingData = (FStrataPrecomputedLightingData)0;
-#endif 
-
-	bStrataSubsurfaceEnable = false;
-	EmissiveColor = 0.0f;
-
-	const float FullyRough = 1.0f;
-
-	// While packing Strata layer data, Classification/SSS/TopLayer data are extracted & stored for dedicated passes
-	// This avoid to run a post-basepass which would re-read all the material data
-	SSSData = (FStrataSubsurfaceData)0;
-	TopLayerData = (FStrataTopLayerData)0;
-	OpaqueRoughRefractionData = (FStrataOpaqueRoughRefractionData)0;
-
-	SetIsSimpleMaterial(StrataPixelHeader, true); // non strata material or sky pixels are considered as simple (so that (strata simple pixels) + (non-strata pixels) can build a simple tile).
-	SetIsSingleMaterial(StrataPixelHeader, true); // non strata material or sky pixels are considered as simple (so that (strata simple pixels) + (non-strata pixels) can build a simple tile).
-
-	uint BSDFVisibleCount = 0;
-	
-	FStrataTree StrataTree = StrataPixelHeader.StrataTree;
-	if (StrataTree.BSDFCount > 0)
-	{
-		// Update tree (coverage/transmittance/luminace weights)
-		FStrataAddressing NullStrataAddressing = (FStrataAddressing)0;	// Fake unused in StrataCreateBSDFContext when using Forward inline shading
-		StrataUpdateTree(NullStrataAddressing, StrataPixelHeader, StrataTree, V, Settings);
-
-		const uint RootOperatorIndex = Strata.OperatorIndex;
-
-	#if STRATA_ADVANCED_DEBUG_ENABLED && STRATA_USES_RESOURCE_MATERIALCONTAINER
-		if (Settings.SliceStoringDebugStrataTree > -1)
-		{
-			StrataPackOutStrataTree(ExtraMaterialDataUAV, RootOperatorIndex, StrataTree, Settings.SliceStoringDebugStrataTree);
-		}
-	#endif
-
-		// Process the rough refraction data (this code is only used, and not culled out, if effectively used in base pass for instance)
-		{
-			const FStrataOperator RootOperator = StrataTree.Operators[Strata.OperatorIndex];
-			OpaqueRoughRefractionData.Coverage = RootOperator.OpaqueRoughRefractCoverage;
-
-			const float StandardDeviationCm = sqrt(GetVarianceFromRoughness(RootOperator.OpaqueRoughRefractTopRoughness));
-			const float StandardDeviationCmForThickness = StandardDeviationCm * RootOperator.OpaqueRoughRefractThicknessCm;
-			OpaqueRoughRefractionData.VarianceCm = StandardDeviationCmForThickness * StandardDeviationCmForThickness;
-		}
-		
-		int SimpleMaterial_BSDFIndex = 0;
-		float TopLayerTotalWeight = 0.0f;
-		
-		{
-			STRATA_UNROLL_N(STRATA_CLAMPED_BSDF_COUNT)
-			for (int BSDFIdx = 0; BSDFIdx < StrataTree.BSDFCount; ++BSDFIdx)
-			{
-			#define BSDF StrataTree.BSDFs[BSDFIdx]
-
-				const bool bIsVisible = StrataIsBSDFVisible(BSDF);
-
-				if (bIsVisible)
-				{
-					BSDFVisibleCount++;
-
-					// Dither BSDF data to hide more aggressive bit reduction.
-					// Sanitize happening after will clamp values that could have gone above and below valid range.
-					StrataDitherBSDF(BSDF, Dither);
-
-					bool bIsSimpleMaterial = true;
-					bool bIsSingleMaterial = true;
-					switch (BSDF_GETTYPE(BSDF))
-					{
-					case STRATA_BSDF_TYPE_SLAB:
-					{
-						EmissiveColor += BSDF_GETEMISSIVE(BSDF) * BSDF.LuminanceWeightV;
-
-						// Simple and single bsdf do not allow weights other than 1
-						const bool LuminanceWeightEqualOne = all(BSDF.LuminanceWeightV == 1.0f);
-
-						// Update simple material compatibility
-						bIsSimpleMaterial = IsStrataSlabFastPathCompatible(BSDF) && LuminanceWeightEqualOne;
-
-						// Update single material compatibility. For now, single materials don't support anisotropy, 
-						// they use the complex pass, as the toplayer data does not contain the full frame basis, only the top normal.
-						bIsSingleMaterial = BSDF_GETHASANISOTROPY(BSDF) == 0 && LuminanceWeightEqualOne;
-
-						TopLayerTotalWeight += BSDF.TopLayerDataWeight;
-						TopLayerData.Roughness += BSDF.TopLayerDataWeight * SLAB_ROUGHNESS(BSDF);
-
-						const bool bHasSSS = BSDF_GETSSSTYPE(BSDF) != SSS_TYPE_INVALID;
-						const bool bIsThin = BSDF_GETISTHIN(BSDF);
-						const bool bIsSimpleVolume = BSDF_GETISSIMPLEVOLUME(BSDF);
-						bStrataSubsurfaceEnable = bStrataSubsurfaceEnable || (bHasSSS && BSDF.Coverage > 0.0f && !bIsSimpleVolume && !bIsThin); // Should it be coverage>0 or any(LuminanceWeight>0) ?
-
-						const bool bSSSMask = !bIsThin && BSDF.bIsBottom;	// SSS and Thin can only be on the bottom layer, so ignoring bIsSimpleVolume here.
-						SetHasSubsurface(StrataPixelHeader, (bSSSMask && bHasSSS));
-
-						break;
-					}
-
-					case STRATA_BSDF_TYPE_HAIR:
-					{
-						bIsSimpleMaterial = false;
-						bIsSingleMaterial = false;
-
-						TopLayerTotalWeight += BSDF.TopLayerDataWeight;
-						TopLayerData.Roughness += BSDF.TopLayerDataWeight * FullyRough;
-
-						break;
-					}
-
-					case STRATA_BSDF_TYPE_SINGLELAYERWATER:
-					{
-						bIsSimpleMaterial = false;
-						bIsSingleMaterial = false;
-						EmissiveColor += BSDF_GETEMISSIVE(BSDF) * BSDF.LuminanceWeightV;
-
-						TopLayerTotalWeight += BSDF.TopLayerDataWeight;
-						TopLayerData.Roughness += BSDF.TopLayerDataWeight * SLW_ROUGHNESS(BSDF);
-
-						break;
-					}
-					}
-
-					TopLayerData.BaseColor += BSDF.TopLayerDataWeight * StrataGetBSDFBaseColor(BSDF);
-
-					float3x3 TangentBasis = StrataGetBSDFSharedBasis_InlineShading(StrataPixelHeader, BSDF, StrataAddressing);
-					float3 N = TangentBasis[2];
-					TopLayerData.WorldNormal += N * BSDF.TopLayerDataWeight;
-
-					AddShadingModels(StrataPixelHeader, StrataShadingModelBit(BSDF_GETTYPE(BSDF)));
-					AndIsSimpleMaterial(StrataPixelHeader, bIsSimpleMaterial);
-					AndIsSingleMaterial(StrataPixelHeader, bIsSingleMaterial);
-
-					SimpleMaterial_BSDFIndex = BSDFIdx;
-
-					// Notify that the BSDF is at the top for SSR to only affect reflection there and not on the lower layers
-					BSDF_SETISTOPLAYER(BSDF, BSDF.bIsTop ? 1 : 0);
-				}
-		
-		
-			#undef BSDF
-			}
-		}
-
-		// Finalize top layer data
-		TopLayerData.WorldNormal= TopLayerTotalWeight > 0.0f	? normalize(TopLayerData.WorldNormal)			: 0.0f;
-		TopLayerData.Roughness	= TopLayerTotalWeight > 0.0f	? TopLayerData.Roughness / TopLayerTotalWeight	: 0.0f;
-		TopLayerData.BaseColor	= TopLayerTotalWeight > 0.0f	? TopLayerData.BaseColor / TopLayerTotalWeight	: 0.0f;
-
-		// Set storage layout as either: fast(0), single(1), or complex(2)
-		AndIsSimpleMaterial(StrataPixelHeader, BSDFVisibleCount == 1);
-		AndIsSingleMaterial(StrataPixelHeader, BSDFVisibleCount == 1);
-
-		// Store irradiance and AO into the state bits of the header so that they can be written later as part of the common header state bits.
-	#if STRATA_INLINE_SHADING
-		HEADER_SETIRRADIANCE_AO(StrataPixelHeader.State, StrataPackIrradianceAndOcclusion(StrataPixelHeader.IrradianceAO, 0));
-	#endif
-
-		// Now write out Strata data
-
-		///////////////////////////////////////////////////////////////////////////
-		// 3 types of encodings
-		// * A. Layout0: Simple encoding	(use top layer normal)
-		// * B. Layout1: Single encoding	(use top layer normal)
-		// * C. Layout2: Complex encoding	(use basis)
-		const bool bHasFastEncoding = IsSimpleMaterial(StrataPixelHeader);
-		const bool bHasSingleEncoding = IsSingleMaterial(StrataPixelHeader);
-
-		// (Layout2)
-		if (!bHasFastEncoding && !bHasSingleEncoding)
-		{
-			// 1. the header (Regular/Complex encoding)
-			const uint PackedHeader = PackStrataHeader(BSDFVisibleCount, StrataPixelHeader);
-			STRATA_STORE_UINT1(PackedHeader);
-
-			// 1.1 tangent basis (Regular/Complex encoding)
-#if STRATA_INLINE_SHADING
-			UNROLL
-			for (uint i = 0; i < StrataPixelHeader.SharedLocalBases.Count; ++i)
-			{
-				const uint BasisType = StrataGetSharedLocalBasisType(StrataPixelHeader.SharedLocalBases.Types, i);
-				if (BasisType == STRATA_BASIS_TYPE_NORMAL)
-				{
-					STRATA_STORE_UINT1(StrataPackNormal(StrataPixelHeader.SharedLocalBases.Normals[i]));
-				}
-				else // if (BasisType == STRATA_BASIS_TYPE_TANGENT)
-				{
-					STRATA_STORE_UINT1(StrataPackNormalAndTangent(StrataPixelHeader.SharedLocalBases.Normals[i], StrataPixelHeader.SharedLocalBases.Tangents[i]));
-				}
-			}
-#endif
-		}
-
-		{
-			int BSDFCount = 0;
-			STRATA_UNROLL_N(STRATA_CLAMPED_BSDF_COUNT)
-			for (int BSDFIdx = 0; BSDFIdx < StrataTree.BSDFCount; ++BSDFIdx)
-			{
-			#define BSDF StrataTree.BSDFs[BSDFIdx]
-
-				const bool bIsVisible = StrataIsBSDFVisible(BSDF);
-
-				// Only write visible BSDF
-				BRANCH
-				if (bIsVisible)
-				{
-					const uint GreyScaleThroughputV = StrataHasGreyScaleWeight(BSDF.LuminanceWeightV) ? 1 : 0;
-					BSDF_SETHASGREYWEIGHT_V(BSDF, GreyScaleThroughputV);
-
-					const bool bTransmittanceAboveAlongNRequired = any(BSDF.TransmittanceAboveAlongN < 1.0f);
-					BSDF_SETHASTRANSABOVE(BSDF, bTransmittanceAboveAlongNRequired);
-
-					// A. Layout0 - Simple encoding (aka. fast-path): store header & data
-					const bool bFastEncodedBSDF = bHasFastEncoding && BSDFIdx == SimpleMaterial_BSDFIndex;
-					if (bFastEncodedBSDF)
-					{
-						// Now we pack diffuse and F0 in a special way: both encoded as R7G7B6 with gamma 2.0. The low bit cound will be hidden by dithering and TAA.
-						// We pack 32 bits in the first uint and the remaining bits are pack is the lowest significant bit of the second uint.
-						const uint PackedDiffuse20Bits = PackR7G7B6Gamma2(SLAB_DIFFUSEALBEDO(BSDF));
-						const uint PackedDiffuse12Bits = PackedDiffuse20Bits & 0xFFF;
-						const uint PackedDiffuse8Bits = (PackedDiffuse20Bits >> 12) & 0xFF;
-						const uint PackedF020Bits = PackR7G7B6Gamma2(SLAB_F0(BSDF));
-						const uint PackedRoughness8bits = PackR8(SLAB_ROUGHNESS(BSDF));
-
-						// Data0 (Header_State|Header_AO|Roughness|Diffuse8bits)
-						{
-							uint Out = 0;
-							HEADER_SETCOMMONSTATES(Out, StrataPixelHeader.State);
-
-							uint Packed16 = (PackedDiffuse8Bits << 8) | PackedRoughness8bits;
-							Out = (Out & HEADER_FASTENCODING_MASK) | (Packed16 << HEADER_FASTENCODING_BIT_COUNT);
-							STRATA_STORE_UINT1(Out);
-
-						#if (HEADER_FASTENCODING_BIT_COUNT + 16) > 32
-						#error Strata fast path header is > 32bits
-						#endif
-						}
-						// Data1 (F0|Diffuse12bits)
-						{
-							uint Out = PackedF020Bits | (PackedDiffuse12Bits << 20);
-							STRATA_STORE_UINT1(Out);
-						}
-						// Ensure the rest of the BSDF is not stored with the regular path
-						BSDF.State = 0;
-					}
-					// B. Layout1 - Simple encoding: single BSDF, whose header & BSDF state are merged
-					else if (bHasSingleEncoding && BSDFIdx == SimpleMaterial_BSDFIndex)
-					{
-						// Data0 (Header_State|Header_AO|Header_BSDFTypes|BSDF_State)
-						{
-							uint Out = 0;
-							HEADER_SETCOMMONSTATES(Out, StrataPixelHeader.State);
-
-							HEADER_SETSHADINGMODELS(Out, GetShadingModels(StrataPixelHeader));
-							Out = Out & HEADER_SINGLEENCODING_MASK;
-							Out = Out | ((BSDF.State & BSDF_SINGLEENCODING_MASK) << HEADER_SINGLEENCODING_BIT_COUNT);
-							STRATA_STORE_UINT1(Out);
-
-						#if (HEADER_SIMPLEENCODING_BIT_COUNT) > 32
-						#error Strata fast path header is > 32bits
-						#endif
-						}
-					}
-					// C. Layout2 - Weight for Regular/Complex path
-					else if (GreyScaleThroughputV > 0)
-					{
-						BSDF_SETWEIGHT10F(BSDF, Pack10F(BSDF.LuminanceWeightV.x));
-						STRATA_STORE_UINT1(BSDF.State);
-					}
-					else
-					{
-						STRATA_STORE_UINT1(BSDF.State);
-						STRATA_STORE_UINT1(PackR11G11B10F(BSDF.LuminanceWeightV));
-					}
-
-					// Layout1 & Layout2
-					if (!bFastEncodedBSDF)
-					{
-						const uint BSDFType = BSDF_GETTYPE(BSDF);
-						switch (BSDF_GETTYPE(BSDF))
-						{
-						case STRATA_BSDF_TYPE_SLAB:
-						{
-							// Now we pack diffuse and F0 in a special way: both encoded as R7G7B6 with gamma 2.0. The low bit cound will be hidden by dithering and TAA.
-							// We pack 32 bits in the first uint and the remaining bits are pack is the lowest significant bit of the second uint.
-							const uint PackedDiffuse20Bits = PackR7G7B6Gamma2(SLAB_DIFFUSEALBEDO(BSDF));
-							const uint PackedF020Bits = PackR7G7B6Gamma2(SLAB_F0(BSDF));
-							const uint PackedData32Bits = ((PackedDiffuse20Bits << 12) & 0xFFFFF000) | (PackedF020Bits & 0xFFF);
-							const uint PackedData8Bits = (PackedF020Bits >> 12) & 0xFF;
-
-							STRATA_STORE_UINT1(PackedData32Bits);
-							STRATA_STORE_UINT1(PackedData8Bits | PackRGBA8(float4(0.0f, SLAB_ROUGHNESS(BSDF), (SLAB_ANISOTROPY(BSDF) + 1.f) * 0.5f, (SLAB_SSSPHASEANISOTROPY(BSDF) + 1.f) * 0.5f)));
-
-							if (BSDF_GETHASF90(BSDF) || BSDF_GETHASHAZINESS(BSDF))
-							{
-								// What is important is to maintain the hue and saturation, so we scale the color by the maximum of its components
-								float3 F90 = saturate(SLAB_F90(BSDF));
-								const float Divisor = max(F90.r, max(F90.g, F90.b));
-								F90 = Divisor > 0.0f ? F90 / Divisor : 1.0f;
-								float3 F90YCoCg = LinearRGB_2_NormalisedYCoCg(F90);
-
-								STRATA_STORE_UINT1(PackRGBA8(float4(F90YCoCg.y, F90YCoCg.z, 0, SLAB_HAZINESS(BSDF))));
-							}
-
-							const bool bIsSimpleVolume = BSDF_GETISSIMPLEVOLUME(BSDF);
-							const bool bIsThin = BSDF_GETISTHIN(BSDF);
-							const bool bHasSSS = BSDF_GETSSSTYPE(BSDF) != SSS_TYPE_INVALID;
-							if (bHasSSS || bIsSimpleVolume || bIsThin)
-							{
-								if (!bIsThin && bHasSSS)
-								{
-									StrataSubSurfaceHeaderSetSSSType(SSSData.Header, BSDF_GETSSSTYPE(BSDF));
-									if (BSDF_GETSSSTYPE(BSDF) == SSS_TYPE_WRAP)
-									{
-										StrataSubSurfaceHeaderSetWrap(SSSData.Header, SLAB_SSSMFP(BSDF), SLAB_SSSPHASEANISOTROPY(BSDF));
-									}
-									else if (BSDF_GETSSSTYPE(BSDF) == SSS_TYPE_DIFFUSION_PROFILE)
-									{
-										StrataSubSurfaceHeaderSetProfile(SSSData.Header, SLAB_SSSPROFILERADIUSSCALE(BSDF), StrataSubsurfaceProfileIdTo8bits(SLAB_SSSPROFILEID(BSDF)));
-									}
-									else
-									{
-										StrataSubSurfaceHeaderSetNonProfile(SSSData.Header, SLAB_SSSMFP(BSDF));
-									}
-
-									StrataSubsurfaceExtrasSetBaseColor(SSSData.Extras, SLAB_DIFFUSEALBEDO(BSDF));
-								}
-
-								if (BSDF_GETSSSTYPE(BSDF) == SSS_TYPE_DIFFUSION_PROFILE && !bIsSimpleVolume)
-								{
-									// Simple volume overrides SSS profile because it takes over for when not at the bottom of the BSDF layer, or during forward rendering.
-									STRATA_STORE_UINT1(PackRGBA8(float4(SLAB_SSSPROFILEID(BSDF), SLAB_SSSPROFILERADIUSSCALE(BSDF), 0, 0)));
-								}
-								else
-								{
-									// Path used for bottom most layer SSS, simple volume and two sided lighting.
-									STRATA_STORE_UINT1(PackR11G11B10F(SLAB_SSSMFP(BSDF)));
-								}
-							}
-
-							if (BSDF_GETHASFUZZ(BSDF))
-							{
-								STRATA_STORE_UINT1(PackColorLinearToGamma2AlphaLinear(float4(SLAB_FUZZ_COLOR(BSDF), SLAB_FUZZ_AMOUNT(BSDF))));
-							}
-
-							if (BSDF_GETHASTRANSABOVE(BSDF))
-							{
-								STRATA_STORE_UINT1(PackR11G11B10F(BSDF.TransmittanceAboveAlongN));
-							}
-
-							// 8-24 bytes
-						}
-						break;
-						case STRATA_BSDF_TYPE_HAIR:
-						{
-							STRATA_STORE_UINT1(PackColorLinearToGamma2AlphaLinear(float4(HAIR_BASECOLOR(BSDF), HAIR_ROUGHNESS(BSDF))));
-							STRATA_STORE_UINT1(PackRGBA8(float4(HAIR_SCATTER(BSDF), HAIR_SPECULAR(BSDF), HAIR_BACKLIT(BSDF), HAIR_COMPLEXTRANSMITTANCE(BSDF))));
-							// 8 bytes
-						}
-						break;
-						case STRATA_BSDF_TYPE_SINGLELAYERWATER:
-						{
-							STRATA_STORE_UINT1(PackColorLinearToGamma2AlphaLinear(float4(SLW_BASECOLOR(BSDF), 0.0f)));
-							STRATA_STORE_UINT1(PackRGBA8(float4(SLW_METALLIC(BSDF), SLW_SPECULAR(BSDF), SLW_ROUGHNESS(BSDF), SLW_TOPMATERIALOPACITY(BSDF))));
-							STRATA_STORE_UINT1(PackR11G11B10F(SLW_SEPARATEDMAINDIRLIGHT(BSDF)));	// STRATA_TODO make this an optional load/store base on if DistanceFieldShadow are enabled
-						}
-						break;
-						}
-					}
-
-					// Convert BSDF to data ready to evaluate precomputed lighting
-					#if MATERIAL_STRATA_PRECOMPUTED_LIGHTING
-					IndirectLightingData.BSDFs[IndirectLightingData.BSDFCount++] = StrataPrepareBSDFForPrecomputedLighting(
-						Settings,
-						StrataPixelHeader,
-						BSDF,
-						V,
-						WorldBentNormal0);
-					#endif
-
-
-					// Stop writing if above budget.
-					BSDFCount++;
-				}
-
-			#undef BSDF
-			}
-
-			checkSlow(BSDFCount == BSDFVisibleCount);
-		}
-
-		FinalizeWrites(StrataBuffer, StrataAddressing);
-	}
-	
-#endif // STRATA_INLINE_SHADING && STRATA_CLAMPED_BSDF_COUNT > 0
-
-}
+///////////////////////////////////////////////////////////////////////////////
+// Header and BSDF packing
 
 // Pack only strata header
 uint PackStrataHeader(uint InBSDFCount, FStrataPixelHeader InHeader)
diff --git a/Engine/Shaders/Private/Strata/StrataExport.ush b/Engine/Shaders/Private/Strata/StrataExport.ush
new file mode 100644
--- /dev/null
+++ b/Engine/Shaders/Private/Strata/StrataExport.ush
@@ -0,0 +1,448 @@
+// Copyright Epic Games, Inc. All Rights Reserved.
+
+#pragma once
+
+#if STRATA_ENABLED
+
+// This is used to evaluate some precomputed lighting data while we pack out Strata material.
+// It is also used to remove reference to unwanted resourced when not needed.
+#ifndef MATERIAL_STRATA_PRECOMPUTED_LIGHTING
+#error MATERIAL_STRATA_PRECOMPUTED_LIGHTING must be defined from Strata.ush before including StrataExport.ush
+#endif
+
+#ifndef STRATA_USES_RESOURCE_MATERIALCONTAINER
+#define STRATA_USES_RESOURCE_MATERIALCONTAINER must be defined from Strata.ush before including StrataExport.ush
+#endif
+
+#if MATERIAL_STRATA_PRECOMPUTED_LIGHTING
+FStrataBSDFForPrecomputedLighting StrataPrepareBSDFForPrecomputedLighting(
+	in FStrataIntegrationSettings Settings,
+	in FStrataPixelHeader StrataPixelHeader,
+	in FStrataBSDF CurrentBSDF,
+	in float3 V,
+	in float3 WorldBentNormal0);
+#endif
+
+// If this is changed, please update the compiler side material size evaluation in StrataMaterial.cpp
+void PackStrataOut(
+	inout FRWStrataMaterialContainer StrataBuffer,
+#if STRATA_USES_RESOURCE_MATERIALCONTAINER
+	RWTexture2DArray<uint> ExtraMaterialDataUAV,
+#endif
+	float Dither,
+	FStrataIntegrationSettings Settings,
+	inout FStrataAddressing StrataAddressing,
+	FStrataPixelHeader StrataPixelHeader,
+	FStrataData Strata,
+	float3 V,
+	float3 WorldBentNormal0,
+	inout bool bStrataSubsurfaceEnable,
+	inout float3 EmissiveColor,
+	inout FStrataSubsurfaceData SSSData,
+	inout FStrataTopLayerData TopLayerData,
+	inout FStrataOpaqueRoughRefractionData OpaqueRoughRefractionData
+#if MATERIAL_STRATA_PRECOMPUTED_LIGHTING
+	,
+	inout FStrataPrecomputedLightingData IndirectLightingData
+#endif
+)
+{
+	// This only exists with inline shading and when we are going to write out BSDFs (UpdateAllXXX functions needs to be defined)
+#if STRATA_INLINE_SHADING && STRATA_CLAMPED_BSDF_COUNT > 0
+
+#if MATERIAL_STRATA_PRECOMPUTED_LIGHTING
+	IndirectLightingData = (FStrataPrecomputedLightingData)0;
+#endif 
+
+	bStrataSubsurfaceEnable = false;
+	EmissiveColor = 0.0f;
+
+	const float FullyRough = 1.0f;
+
+	// While packing Strata layer data, Classification/SSS/TopLayer data are extracted & stored for dedicated passes
+	// This avoid to run a post-basepass which would re-read all the material data
+	SSSData = (FStrataSubsurfaceData)0;
+	TopLayerData = (FStrataTopLayerData)0;
+	OpaqueRoughRefractionData = (FStrataOpaqueRoughRefractionData)0;
+
+	SetIsSimpleMaterial(StrataPixelHeader, true); // non strata material or sky pixels are considered as simple (so that (strata simple pixels) + (non-strata pixels) can build a simple tile).
+	SetIsSingleMaterial(StrataPixelHeader, true); // non strata material or sky pixels are considered as simple (so that (strata simple pixels) + (non-strata pixels) can build a simple tile).
+
+	uint BSDFVisibleCount = 0;
+	
+	FStrataTree StrataTree = StrataPixelHeader.StrataTree;
+	if (StrataTree.BSDFCount > 0)
+	{
+		// Update tree (coverage/transmittance/luminace weights)
+		FStrataAddressing NullStrataAddressing = (FStrataAddressing)0;	// Fake unused in StrataCreateBSDFContext when using Forward inline shading
+		StrataUpdateTree(NullStrataAddressing, StrataPixelHeader, StrataTree, V, Settings);
+
+		const uint RootOperatorIndex = Strata.OperatorIndex;
+
+	#if STRATA_ADVANCED_DEBUG_ENABLED && STRATA_USES_RESOURCE_MATERIALCONTAINER
+		if (Settings.SliceStoringDebugStrataTree > -1)
+		{
+			StrataPackOutStrataTree(ExtraMaterialDataUAV, RootOperatorIndex, StrataTree, Settings.SliceStoringDebugStrataTree);
+		}
+	#endif
+
+		// Process the rough refraction data (this code is only used, and not culled out, if effectively used in base pass for instance)
+		{
+			const FStrataOperator RootOperator = StrataTree.Operators[Strata.OperatorIndex];
+			OpaqueRoughRefractionData.Coverage = RootOperator.OpaqueRoughRefractCoverage;
+
+			const float StandardDeviationCm = sqrt(GetVarianceFromRoughness(RootOperator.OpaqueRoughRefractTopRoughness));
+			const float StandardDeviationCmForThickness = StandardDeviationCm * RootOperator.OpaqueRoughRefractThicknessCm;
+			OpaqueRoughRefractionData.VarianceCm = StandardDeviationCmForThickness * StandardDeviationCmForThickness;
+		}
+		
+		int SimpleMaterial_BSDFIndex = 0;
+		float TopLayerTotalWeight = 0.0f;
+		
+		{
+			STRATA_UNROLL_N(STRATA_CLAMPED_BSDF_COUNT)
+			for (int BSDFIdx = 0; BSDFIdx < StrataTree.BSDFCount; ++BSDFIdx)
+			{
+			#define BSDF StrataTree.BSDFs[BSDFIdx]
+
+				const bool bIsVisible = StrataIsBSDFVisible(BSDF);
+
+				if (bIsVisible)
+				{
+					BSDFVisibleCount++;
+
+					// Dither BSDF data to hide more aggressive bit reduction.
+					// Sanitize happening after will clamp values that could have gone above and below valid range.
+					StrataDitherBSDF(BSDF, Dither);
+
+					bool bIsSimpleMaterial = true;
+					bool bIsSingleMaterial = true;
+					switch (BSDF_GETTYPE(BSDF))
+					{
+					case STRATA_BSDF_TYPE_SLAB:
+					{
+						EmissiveColor += BSDF_GETEMISSIVE(BSDF) * BSDF.LuminanceWeightV;
+
+						// Simple and single bsdf do not allow weights other than 1
+						const bool LuminanceWeightEqualOne = all(BSDF.LuminanceWeightV == 1.0f);
+
+						// Update simple material compatibility
+						bIsSimpleMaterial = IsStrataSlabFastPathCompatible(BSDF) && LuminanceWeightEqualOne;
+
+						// Update single material compatibility. For now, single materials don't support anisotropy, 
+						// they use the complex pass, as the toplayer data does not contain the full frame basis, only the top normal.
+						bIsSingleMaterial = BSDF_GETHASANISOTROPY(BSDF) == 0 && LuminanceWeightEqualOne;
+
+						TopLayerTotalWeight += BSDF.TopLayerDataWeight;
+						TopLayerData.Roughness += BSDF.TopLayerDataWeight * SLAB_ROUGHNESS(BSDF);
+
+						const bool bHasSSS = BSDF_GETSSSTYPE(BSDF) != SSS_TYPE_INVALID;
+						const bool bIsThin = BSDF_GETISTHIN(BSDF);
+						const bool bIsSimpleVolume = BSDF_GETISSIMPLEVOLUME(BSDF);
+						bStrataSubsurfaceEnable = bStrataSubsurfaceEnable || (bHasSSS && BSDF.Coverage > 0.0f && !bIsSimpleVolume && !bIsThin); // Should it be coverage>0 or any(LuminanceWeight>0) ?
+
+						const bool bSSSMask = !bIsThin && BSDF.bIsBottom;	// SSS and Thin can only be on the bottom layer, so ignoring bIsSimpleVolume here.
+						SetHasSubsurface(StrataPixelHeader, (bSSSMask && bHasSSS));
+
+						break;
+					}
+
+					case STRATA_BSDF_TYPE_HAIR:
+					{
+						bIsSimpleMaterial = false;
+						bIsSingleMaterial = false;
+
+						TopLayerTotalWeight += BSDF.TopLayerDataWeight;
+						TopLayerData.Roughness += BSDF.TopLayerDataWeight * FullyRough;
+
+						break;
+					}
+
+					case STRATA_BSDF_TYPE_SINGLELAYERWATER:
+					{
+						bIsSimpleMaterial = false;
+						bIsSingleMaterial = false;
+						EmissiveColor += BSDF_GETEMISSIVE(BSDF) * BSDF.LuminanceWeightV;
+
+						TopLayerTotalWeight += BSDF.TopLayerDataWeight;
+						TopLayerData.Roughness += BSDF.TopLayerDataWeight * SLW_ROUGHNESS(BSDF);
+
+						break;
+					}
+					}
+
+					TopLayerData.BaseColor += BSDF.TopLayerDataWeight * StrataGetBSDFBaseColor(BSDF);
+
+					float3x3 TangentBasis = StrataGetBSDFSharedBasis_InlineShading(StrataPixelHeader, BSDF, StrataAddressing);
+					float3 N = TangentBasis[2];
+					TopLayerData.WorldNormal += N * BSDF.TopLayerDataWeight;
+
+					AddShadingModels(StrataPixelHeader, StrataShadingModelBit(BSDF_GETTYPE(BSDF)));
+					AndIsSimpleMaterial(StrataPixelHeader, bIsSimpleMaterial);
+					AndIsSingleMaterial(StrataPixelHeader, bIsSingleMaterial);
+
+					SimpleMaterial_BSDFIndex = BSDFIdx;
+
+					// Notify that the BSDF is at the top for SSR to only affect reflection there and not on the lower layers
+					BSDF_SETISTOPLAYER(BSDF, BSDF.bIsTop ? 1 : 0);
+				}
+		
+		
+			#undef BSDF
+			}
+		}
+
+		// Finalize top layer data
+		TopLayerData.WorldNormal= TopLayerTotalWeight > 0.0f	? normalize(TopLayerData.WorldNormal)			: 0.0f;
+		TopLayerData.Roughness	= TopLayerTotalWeight > 0.0f	? TopLayerData.Roughness / TopLayerTotalWeight	: 0.0f;
+		TopLayerData.BaseColor	= TopLayerTotalWeight > 0.0f	? TopLayerData.BaseColor / TopLayerTotalWeight	: 0.0f;
+
+		// Set storage layout as either: fast(0), single(1), or complex(2)
+		AndIsSimpleMaterial(StrataPixelHeader, BSDFVisibleCount == 1);
+		AndIsSingleMaterial(StrataPixelHeader, BSDFVisibleCount == 1);
+
+		// Store irradiance and AO into the state bits of the header so that they can be written later as part of the common header state bits.
+	#if STRATA_INLINE_SHADING
+		HEADER_SETIRRADIANCE_AO(StrataPixelHeader.State, StrataPackIrradianceAndOcclusion(StrataPixelHeader.IrradianceAO, 0));
+	#endif
+
+		// Now write out Strata data
+
+		///////////////////////////////////////////////////////////////////////////
+		// 3 types of encodings
+		// * A. Layout0: Simple encoding	(use top layer normal)
+		// * B. Layout1: Single encoding	(use top layer normal)
+		// * C. Layout2: Complex encoding	(use basis)
+		const bool bHasFastEncoding = IsSimpleMaterial(StrataPixelHeader);
+		const bool bHasSingleEncoding = IsSingleMaterial(StrataPixelHeader);
+
+		// (Layout2)
+		if (!bHasFastEncoding && !bHasSingleEncoding)
+		{
+			// 1. the header (Regular/Complex encoding)
+			const uint PackedHeader = PackStrataHeader(BSDFVisibleCount, StrataPixelHeader);
+			STRATA_STORE_UINT1(PackedHeader);
+
+			// 1.1 tangent basis (Regular/Complex encoding)
+		#if STRATA_INLINE_SHADING
+			UNROLL
+			for (uint i = 0; i < StrataPixelHeader.SharedLocalBases.Count; ++i)
+			{
+				const uint BasisType = StrataGetSharedLocalBasisType(StrataPixelHeader.SharedLocalBases.Types, i);
+				if (BasisType == STRATA_BASIS_TYPE_NORMAL)
+				{
+					STRATA_STORE_UINT1(StrataPackNormal(StrataPixelHeader.SharedLocalBases.Normals[i]));
+				}
+				else // if (BasisType == STRATA_BASIS_TYPE_TANGENT)
+				{
+					STRATA_STORE_UINT1(StrataPackNormalAndTangent(StrataPixelHeader.SharedLocalBases.Normals[i], StrataPixelHeader.SharedLocalBases.Tangents[i]));
+				}
+			}
+		#endif
+		}
+
+		{
+			int BSDFCount = 0;
+			STRATA_UNROLL_N(STRATA_CLAMPED_BSDF_COUNT)
+			for (int BSDFIdx = 0; BSDFIdx < StrataTree.BSDFCount; ++BSDFIdx)
+			{
+			#define BSDF StrataTree.BSDFs[BSDFIdx]
+
+				const bool bIsVisible = StrataIsBSDFVisible(BSDF);
+
+				// Only write visible BSDF
+				BRANCH
+				if (bIsVisible)
+				{
+					const uint GreyScaleThroughputV = StrataHasGreyScaleWeight(BSDF.LuminanceWeightV) ? 1 : 0;
+					BSDF_SETHASGREYWEIGHT_V(BSDF, GreyScaleThroughputV);
+
+					const bool bTransmittanceAboveAlongNRequired = any(BSDF.TransmittanceAboveAlongN < 1.0f);
+					BSDF_SETHASTRANSABOVE(BSDF, bTransmittanceAboveAlongNRequired);
+
+					// A. Layout0 - Simple encoding (aka. fast-path): store header & data
+					const bool bFastEncodedBSDF = bHasFastEncoding && BSDFIdx == SimpleMaterial_BSDFIndex;
+					if (bFastEncodedBSDF)
+					{
+						// Now we pack diffuse and F0 in a special way: both encoded as R7G7B6 with gamma 2.0. The low bit cound will be hidden by dithering and TAA.
+						// We pack 32 bits in the first uint and the remaining bits are pack is the lowest significant bit of the second uint.
+						const uint PackedDiffuse20Bits = PackR7G7B6Gamma2(SLAB_DIFFUSEALBEDO(BSDF));
+						const uint PackedDiffuse12Bits = PackedDiffuse20Bits & 0xFFF;
+						const uint PackedDiffuse8Bits = (PackedDiffuse20Bits >> 12) & 0xFF;
+						const uint PackedF020Bits = PackR7G7B6Gamma2(SLAB_F0(BSDF));
+						const uint PackedRoughness8bits = PackR8(SLAB_ROUGHNESS(BSDF));
+
+						// Data0 (Header_State|Header_AO|Roughness|Diffuse8bits)
+						{
+							uint Out = 0;
+							HEADER_SETCOMMONSTATES(Out, StrataPixelHeader.State);
+
+							uint Packed16 = (PackedDiffuse8Bits << 8) | PackedRoughness8bits;
+							Out = (Out & HEADER_FASTENCODING_MASK) | (Packed16 << HEADER_FASTENCODING_BIT_COUNT);
+							STRATA_STORE_UINT1(Out);
+
+						#if (HEADER_FASTENCODING_BIT_COUNT + 16) > 32
+						#error Strata fast path header is > 32bits
+						#endif
+						}
+						// Data1 (F0|Diffuse12bits)
+						{
+							uint Out = PackedF020Bits | (PackedDiffuse12Bits << 20);
+							STRATA_STORE_UINT1(Out);
+						}
+						// Ensure the rest of the BSDF is not stored with the regular path
+						BSDF.State = 0;
+					}
+					// B. Layout1 - Simple encoding: single BSDF, whose header & BSDF state are merged
+					else if (bHasSingleEncoding && BSDFIdx == SimpleMaterial_BSDFIndex)
+					{
+						// Data0 (Header_State|Header_AO|Header_BSDFTypes|BSDF_State)
+						{
+							uint Out = 0;
+							HEADER_SETCOMMONSTATES(Out, StrataPixelHeader.State);
+
+							HEADER_SETSHADINGMODELS(Out, GetShadingModels(StrataPixelHeader));
+							Out = Out & HEADER_SINGLEENCODING_MASK;
+							Out = Out | ((BSDF.State & BSDF_SINGLEENCODING_MASK) << HEADER_SINGLEENCODING_BIT_COUNT);
+							STRATA_STORE_UINT1(Out);
+
+						#if (HEADER_SIMPLEENCODING_BIT_COUNT) > 32
+						#error Strata fast path header is > 32bits
+						#endif
+						}
+					}
+					// C. Layout2 - Weight for Regular/Complex path
+					else if (GreyScaleThroughputV > 0)
+					{
+						BSDF_SETWEIGHT10F(BSDF, Pack10F(BSDF.LuminanceWeightV.x));
+						STRATA_STORE_UINT1(BSDF.State);
+					}
+					else
+					{
+						STRATA_STORE_UINT1(BSDF.State);
+						STRATA_STORE_UINT1(PackR11G11B10F(BSDF.LuminanceWeightV));
+					}
+
+					// Layout1 & Layout2
+					if (!bFastEncodedBSDF)
+					{
+						const uint BSDFType = BSDF_GETTYPE(BSDF);
+						switch (BSDF_GETTYPE(BSDF))
+						{
+						case STRATA_BSDF_TYPE_SLAB:
+						{
+							// Now we pack diffuse and F0 in a special way: both encoded as R7G7B6 with gamma 2.0. The low bit cound will be hidden by dithering and TAA.
+							// We pack 32 bits in the first uint and the remaining bits are pack is the lowest significant bit of the second uint.
+							const uint PackedDiffuse20Bits = PackR7G7B6Gamma2(SLAB_DIFFUSEALBEDO(BSDF));
+							const uint PackedF020Bits = PackR7G7B6Gamma2(SLAB_F0(BSDF));
+							const uint PackedData32Bits = ((PackedDiffuse20Bits << 12) & 0xFFFFF000) | (PackedF020Bits & 0xFFF);
+							const uint PackedData8Bits = (PackedF020Bits >> 12) & 0xFF;
+
+							STRATA_STORE_UINT1(PackedData32Bits);
+							STRATA_STORE_UINT1(PackedData8Bits | PackRGBA8(float4(0.0f, SLAB_ROUGHNESS(BSDF), (SLAB_ANISOTROPY(BSDF) + 1.f) * 0.5f, (SLAB_SSSPHASEANISOTROPY(BSDF) + 1.f) * 0.5f)));
+
+							if (BSDF_GETHASF90(BSDF) || BSDF_GETHASHAZINESS(BSDF))
+							{
+								// What is important is to maintain the hue and saturation, so we scale the color by the maximum of its components
+								float3 F90 = saturate(SLAB_F90(BSDF));
+								const float Divisor = max(F90.r, max(F90.g, F90.b));
+								F90 = Divisor > 0.0f ? F90 / Divisor : 1.0f;
+								float3 F90YCoCg = LinearRGB_2_NormalisedYCoCg(F90);
+
+								STRATA_STORE_UINT1(PackRGBA8(float4(F90YCoCg.y, F90YCoCg.z, 0, SLAB_HAZINESS(BSDF))));
+							}
+
+							const bool bIsSimpleVolume = BSDF_GETISSIMPLEVOLUME(BSDF);
+							const bool bIsThin = BSDF_GETISTHIN(BSDF);
+							const bool bHasSSS = BSDF_GETSSSTYPE(BSDF) != SSS_TYPE_INVALID;
+							if (bHasSSS || bIsSimpleVolume || bIsThin)
+							{
+								if (!bIsThin && bHasSSS)
+								{
+									StrataSubSurfaceHeaderSetSSSType(SSSData.Header, BSDF_GETSSSTYPE(BSDF));
+									if (BSDF_GETSSSTYPE(BSDF) == SSS_TYPE_WRAP)
+									{
+										StrataSubSurfaceHeaderSetWrap(SSSData.Header, SLAB_SSSMFP(BSDF), SLAB_SSSPHASEANISOTROPY(BSDF));
+									}
+									else if (BSDF_GETSSSTYPE(BSDF) == SSS_TYPE_DIFFUSION_PROFILE)
+									{
+										StrataSubSurfaceHeaderSetProfile(SSSData.Header, SLAB_SSSPROFILERADIUSSCALE(BSDF), StrataSubsurfaceProfileIdTo8bits(SLAB_SSSPROFILEID(BSDF)));
+									}
+									else
+									{
+										StrataSubSurfaceHeaderSetNonProfile(SSSData.Header, SLAB_SSSMFP(BSDF));
+									}
+
+									StrataSubsurfaceExtrasSetBaseColor(SSSData.Extras, SLAB_DIFFUSEALBEDO(BSDF));
+								}
+
+								if (BSDF_GETSSSTYPE(BSDF) == SSS_TYPE_DIFFUSION_PROFILE && !bIsSimpleVolume)
+								{
+									// Simple volume overrides SSS profile because it takes over for when not at the bottom of the BSDF layer, or during forward rendering.
+									STRATA_STORE_UINT1(PackRGBA8(float4(SLAB_SSSPROFILEID(BSDF), SLAB_SSSPROFILERADIUSSCALE(BSDF), 0, 0)));
+								}
+								else
+								{
+									// Path used for bottom most layer SSS, simple volume and two sided lighting.
+									STRATA_STORE_UINT1(PackR11G11B10F(SLAB_SSSMFP(BSDF)));
+								}
+							}
+
+							if (BSDF_GETHASFUZZ(BSDF))
+							{
+								STRATA_STORE_UINT1(PackColorLinearToGamma2AlphaLinear(float4(SLAB_FUZZ_COLOR(BSDF), SLAB_FUZZ_AMOUNT(BSDF))));
+							}
+
+							if (BSDF_GETHASTRANSABOVE(BSDF))
+							{
+								STRATA_STORE_UINT1(PackR11G11B10F(BSDF.TransmittanceAboveAlongN));
+							}
+
+							// 8-24 bytes
+						}
+						break;
+						case STRATA_BSDF_TYPE_HAIR:
+						{
+							STRATA_STORE_UINT1(PackColorLinearToGamma2AlphaLinear(float4(HAIR_BASECOLOR(BSDF), HAIR_ROUGHNESS(BSDF))));
+							STRATA_STORE_UINT1(PackRGBA8(float4(HAIR_SCATTER(BSDF), HAIR_SPECULAR(BSDF), HAIR_BACKLIT(BSDF), HAIR_COMPLEXTRANSMITTANCE(BSDF))));
+							// 8 bytes
+						}
+						break;
+						case STRATA_BSDF_TYPE_SINGLELAYERWATER:
+						{
+							STRATA_STORE_UINT1(PackColorLinearToGamma2AlphaLinear(float4(SLW_BASECOLOR(BSDF), 0.0f)));
+							STRATA_STORE_UINT1(PackRGBA8(float4(SLW_METALLIC(BSDF), SLW_SPECULAR(BSDF), SLW_ROUGHNESS(BSDF), SLW_TOPMATERIALOPACITY(BSDF))));
+							STRATA_STORE_UINT1(PackR11G11B10F(SLW_SEPARATEDMAINDIRLIGHT(BSDF)));	// STRATA_TODO make this an optional load/store base on if DistanceFieldShadow are enabled
+						}
+						break;
+						}
+					}
+
+					// Convert BSDF to data ready to evaluate precomputed lighting
+					#if MATERIAL_STRATA_PRECOMPUTED_LIGHTING
+					IndirectLightingData.BSDFs[IndirectLightingData.BSDFCount++] = StrataPrepareBSDFForPrecomputedLighting(
+						Settings,
+						StrataPixelHeader,
+						BSDF,
+						V,
+						WorldBentNormal0);
+					#endif
+
+
+					// Stop writing if above budget.
+					BSDFCount++;
+				}
+
+			#undef BSDF
+			}
+
+			checkSlow(BSDFCount == BSDFVisibleCount);
+		}
+
+		FinalizeWrites(StrataBuffer, StrataAddressing);
+	}
+#endif // STRATA_INLINE_SHADING && STRATA_CLAMPED_BSDF_COUNT > 0
+
+}
+
+#endif // STRATA_ENABLED
