Strata - forgotten file to go with 20436182.

diff --git a/Engine/Shaders/Private/MaterialTemplate.ush b/Engine/Shaders/Private/MaterialTemplate.ush
--- a/Engine/Shaders/Private/MaterialTemplate.ush
+++ b/Engine/Shaders/Private/MaterialTemplate.ush
@@ -3392,3 +3392,16 @@ float ApplyPixelDepthOffsetToMaterialParameters(inout FMaterialPixelParameters M
 
 	return PixelDepthOffset;
 }
+
+float3 GetWorldBentNormalZero(in FMaterialPixelParameters MaterialParameters)
+{
+#if NUM_MATERIAL_OUTPUTS_GETBENTNORMAL > 0
+	#if MATERIAL_TANGENTSPACENORMAL
+		return normalize(TransformTangentVectorToWorld(MaterialParameters.TangentToWorld, GetBentNormal0(MaterialParameters)));
+	#else
+		return GetBentNormal0(MaterialParameters);
+	#endif
+#else
+	return MaterialParameters.WorldNormal;
+#endif
+}
