Strata - moved more code to StrataExport. Also clarified that the extra struct  is for opaque only.

diff --git a/Engine/Shaders/Private/BasePassPixelShader.usf b/Engine/Shaders/Private/BasePassPixelShader.usf
--- a/Engine/Shaders/Private/BasePassPixelShader.usf
+++ b/Engine/Shaders/Private/BasePassPixelShader.usf
@@ -44,36 +44,6 @@
 #define STRATA_ENABLED 0
 #endif
 
-struct FShadingOcclusion
-{
-	float  DiffOcclusion;
-	float  SpecOcclusion;
-	float3 BentNormal;
-};
-
-#define MATERIAL_STRATA_PRECOMPUTED_LIGHTING (MATERIAL_IS_STRATA && STRATA_ENABLED)
-
-#if MATERIAL_STRATA_PRECOMPUTED_LIGHTING
-#if STRATA_CLAMPED_BSDF_COUNT <= 0
-#define STRATA_BSDF_COUNT_FOR_PRECOMPUTED 1
-#else
-#define STRATA_BSDF_COUNT_FOR_PRECOMPUTED STRATA_CLAMPED_BSDF_COUNT
-#endif
-struct FStrataBSDFForPrecomputedLighting
-{
-	float3 DiffuseNormal;// STRATA_TODO test packed on a uint4 as R11G11B10
-	float3 DiffuseColor;
-	float3 DiffuseColorForIndirect;
-	float3 SubsurfaceColorForIndirect;
-	FShadingOcclusion ShadingOcclusion;
-};
-struct FStrataPrecomputedLightingData
-{
-	uint BSDFCount;
-	FStrataBSDFForPrecomputedLighting BSDFs[STRATA_BSDF_COUNT_FOR_PRECOMPUTED];
-};
-#endif
-
 // For the base pass STRATA_INLINE_SHADING is defined in BasePassRendering.h
 #include "SHCommon.ush"
 #include "/Engine/Generated/Material.ush"
@@ -140,6 +110,13 @@ float NormalCurvatureToRoughness(float3 WorldNormal)
 	return saturate(CurvatureApprox * View.NormalCurvatureToRoughnessScaleBias.x + View.NormalCurvatureToRoughnessScaleBias.y);
 }
 
+struct FShadingOcclusion
+{
+	float  DiffOcclusion;
+	float  SpecOcclusion;
+	float3 BentNormal;
+};
+
 #if TRANSLUCENT_SELF_SHADOWING
 	#include "ShadowProjectionCommon.ush"	 
 #endif
@@ -200,15 +177,17 @@ float NormalCurvatureToRoughness(float3 WorldNormal)
 #include "ForwardLightingCommon.ush"
 #endif
 
-#if STRATA_OPAQUE_DEFERRED
-#include "/Engine/Private/Strata/StrataExport.ush"
-#endif
-#if STRATA_TRANSLUCENT_FORWARD || STRATA_FORWARD_SHADING || (MATERIAL_STRATA_PRECOMPUTED_LIGHTING && STRATA_OPAQUE_DEFERRED)
+#define MATERIAL_STRATA_OPAQUE_PRECOMPUTED_LIGHTING (MATERIAL_IS_STRATA && STRATA_ENABLED && STRATA_OPAQUE_DEFERRED)
+
+#if STRATA_TRANSLUCENT_FORWARD || STRATA_FORWARD_SHADING || MATERIAL_STRATA_OPAQUE_PRECOMPUTED_LIGHTING
 #include "/Engine/Private/Strata/StrataEvaluation.ush"
 #endif
 #if STRATA_TRANSLUCENT_FORWARD || STRATA_FORWARD_SHADING
 #include "/Engine/Private/Strata/StrataForwardLighting.ush"
 #endif
+#if MATERIAL_STRATA_OPAQUE_PRECOMPUTED_LIGHTING
+#include "/Engine/Private/Strata/StrataExport.ush"
+#endif
 
 // Separate main directional light from scene color buffer if water needs to receive distance field shadows.
 // Sperated dir light should match with FShaderCompileUtilities::FetchGBufferParamsRuntime.
@@ -865,47 +844,7 @@ FShadingOcclusion ApplyBentNormal(
 
 #if STRATA_ENABLED
 
-#if MATERIAL_STRATA_PRECOMPUTED_LIGHTING
-FStrataBSDFForPrecomputedLighting StrataPrepareBSDFForPrecomputedLighting(
-	in FStrataIntegrationSettings Settings,
-	in FStrataPixelHeader StrataPixelHeader,
-	in FStrataBSDF CurrentBSDF,
-	in float3 V,
-	in float3 WorldBentNormal0)
-{
-	FStrataBSDFForPrecomputedLighting BSDFForPrecomputedLighting = (FStrataBSDFForPrecomputedLighting)0;
-
-	FStrataAddressing NullStrataAddressing = (FStrataAddressing)0;
-	FStrataBSDFContext StrataBSDFContext = StrataCreateBSDFContext(StrataPixelHeader, CurrentBSDF, NullStrataAddressing, V);
-
-	// Update a context specific for environment lighting
-	FStrataBSDFContext EnvBSDFContext = StrataBSDFContext;
-	EnvBSDFContext.N = BSDFForPrecomputedLighting.ShadingOcclusion.BentNormal;
-	StrataUpdateBSDFContext(EnvBSDFContext, EnvBSDFContext.L);
-
-	// Evaluate the bent normal
-	BSDFForPrecomputedLighting.ShadingOcclusion = ApplyBentNormal(
-		V,
-		EnvBSDFContext.N,
-		WorldBentNormal0,
-		StrataGetBSDFRoughness(CurrentBSDF),
-		StrataPixelHeader.IrradianceAO.MaterialAO);
-
-	// For diffuse, we specify a perpendicular to the surface light direction for the transmittance to light to not be view dependent.
-	float DiffuseEnvLightingNoL = 1.0f;
-	float3 LuminanceWeightFinal = LuminanceWeight(DiffuseEnvLightingNoL, CurrentBSDF);
-
-	// And evaluate diffuse parameters
-	const bool bEnableSpecular = false;
-	FStrataEnvLightResult StrataEnvLight = StrataEvaluateForEnvLight(EnvBSDFContext, bEnableSpecular, Settings);
-	BSDFForPrecomputedLighting.DiffuseNormal = StrataEnvLight.DiffuseNormal;
-	BSDFForPrecomputedLighting.DiffuseColor = StrataEnvLight.DiffuseColor * LuminanceWeightFinal;
-	BSDFForPrecomputedLighting.DiffuseColorForIndirect = StrataEnvLight.DiffuseWeight * LuminanceWeightFinal;
-	BSDFForPrecomputedLighting.SubsurfaceColorForIndirect = StrataEnvLight.DiffuseBackFaceWeight * LuminanceWeightFinal;
-
-	return BSDFForPrecomputedLighting;
-}
-
+#if MATERIAL_STRATA_OPAQUE_PRECOMPUTED_LIGHTING
 float3 StrataGetBSDFPrecomputedLighting(
 	in FStrataBSDFForPrecomputedLighting BSDFForPrecomputedLighting,
 	in FMaterialPixelParameters MaterialParameters,
@@ -1815,7 +1754,7 @@ void FPixelShaderInOut_MainPS(
 		uint2 PixelPos = uint2(In.SvPosition.xy);
 		const float3 WorldBentNormal0 = GetWorldBentNormalZero(MaterialParameters);
 
-	#if MATERIAL_STRATA_PRECOMPUTED_LIGHTING
+	#if MATERIAL_STRATA_OPAQUE_PRECOMPUTED_LIGHTING
 		FStrataPrecomputedLightingData PrecomputedLightingData = (FStrataPrecomputedLightingData)0;
 	#endif
 		FStrataSubsurfaceData SSSData = (FStrataSubsurfaceData)0;
@@ -1838,13 +1777,13 @@ void FPixelShaderInOut_MainPS(
 			StrataAddressing,
 			StrataPixelHeader, StrataData, MaterialParameters.CameraVector, WorldBentNormal0, bStrataSubsurfaceEnable, EmissiveLuminance,
 			SSSData, TopLayerData, OpaqueRoughRefractionData
-	#if MATERIAL_STRATA_PRECOMPUTED_LIGHTING
+	#if MATERIAL_STRATA_OPAQUE_PRECOMPUTED_LIGHTING
 			, PrecomputedLightingData
 	#endif
 		);
 
 		float3 PreComputedLightingLuminance = 0.0f;
-	#if MATERIAL_STRATA_PRECOMPUTED_LIGHTING
+	#if MATERIAL_STRATA_OPAQUE_PRECOMPUTED_LIGHTING
 		for (uint i = 0; i < PrecomputedLightingData.BSDFCount; ++i)
 		{
 			PreComputedLightingLuminance += StrataGetBSDFPrecomputedLighting(
diff --git a/Engine/Shaders/Private/DebugProbes.usf b/Engine/Shaders/Private/DebugProbes.usf
--- a/Engine/Shaders/Private/DebugProbes.usf
+++ b/Engine/Shaders/Private/DebugProbes.usf
@@ -7,7 +7,7 @@
 #define STRATA_INLINE_SHADING 1
 #include "Strata/Strata.ush"
 
-#define MATERIAL_STRATA_PRECOMPUTED_LIGHTING 0
+#define MATERIAL_STRATA_OPAQUE_PRECOMPUTED_LIGHTING 0
 #include "Strata/StrataExport.ush"
 
 #include "DeferredShadingCommon.ush"
diff --git a/Engine/Shaders/Private/RayTracing/RayTracingMaterialHitShaders.usf b/Engine/Shaders/Private/RayTracing/RayTracingMaterialHitShaders.usf
--- a/Engine/Shaders/Private/RayTracing/RayTracingMaterialHitShaders.usf
+++ b/Engine/Shaders/Private/RayTracing/RayTracingMaterialHitShaders.usf
@@ -34,7 +34,7 @@
 #include "/Engine/Private/LightmapCommon.ush"
 #include "/Engine/Private/Lumen/LumenHardwareRayTracingPayloadCommon.ush"
 
-#define MATERIAL_STRATA_PRECOMPUTED_LIGHTING 0
+#define MATERIAL_STRATA_OPAQUE_PRECOMPUTED_LIGHTING 0
 #include "/Engine/Private/Strata/StrataExport.ush"
 
 /** Computes sky diffuse lighting, including precomputed shadowing. */
diff --git a/Engine/Shaders/Private/Strata/Strata.ush b/Engine/Shaders/Private/Strata/Strata.ush
--- a/Engine/Shaders/Private/Strata/Strata.ush
+++ b/Engine/Shaders/Private/Strata/Strata.ush
@@ -42,12 +42,6 @@
 #define STRATA_INLINE_SHADING 0
 #endif
 
-// This is used to evaluate some precomputed lighting data while we pack out Strata material.
-// It is also used to remove reference to unwanted resourced when not needed.
-#ifndef MATERIAL_STRATA_PRECOMPUTED_LIGHTING
-#define MATERIAL_STRATA_PRECOMPUTED_LIGHTING 0
-#endif
-
 #ifndef STRATA_DEFERRED_SHADING
 #define STRATA_DEFERRED_SHADING (STRATA_INLINE_SHADING == 0)
 #endif 
diff --git a/Engine/Shaders/Private/Strata/StrataExport.ush b/Engine/Shaders/Private/Strata/StrataExport.ush
--- a/Engine/Shaders/Private/Strata/StrataExport.ush
+++ b/Engine/Shaders/Private/Strata/StrataExport.ush
@@ -6,22 +6,87 @@
 
 // This is used to evaluate some precomputed lighting data while we pack out Strata material.
 // It is also used to remove reference to unwanted resourced when not needed.
-#ifndef MATERIAL_STRATA_PRECOMPUTED_LIGHTING
-#error MATERIAL_STRATA_PRECOMPUTED_LIGHTING must be defined from Strata.ush before including StrataExport.ush
+#ifndef MATERIAL_STRATA_OPAQUE_PRECOMPUTED_LIGHTING
+#error MATERIAL_STRATA_OPAQUE_PRECOMPUTED_LIGHTING must be defined before including StrataExport.ush
 #endif
 
 #ifndef STRATA_USES_RESOURCE_MATERIALCONTAINER
 #define STRATA_USES_RESOURCE_MATERIALCONTAINER must be defined from Strata.ush before including StrataExport.ush
 #endif
 
-#if MATERIAL_STRATA_PRECOMPUTED_LIGHTING
+#if MATERIAL_STRATA_OPAQUE_PRECOMPUTED_LIGHTING
+
+// Forward declaration
+FShadingOcclusion ApplyBentNormal(
+	in float3 CameraVector,
+	in float3 WorldNormal,
+	in float3 WorldBentNormal0,
+	in float Roughness,
+	in float MaterialAO);
+
+
+#if STRATA_CLAMPED_BSDF_COUNT <= 0
+	#define STRATA_BSDF_COUNT_FOR_PRECOMPUTED 1
+#else
+	#define STRATA_BSDF_COUNT_FOR_PRECOMPUTED STRATA_CLAMPED_BSDF_COUNT
+#endif
+
+struct FStrataBSDFForPrecomputedLighting
+{
+	float3 DiffuseNormal;// STRATA_TODO test packed on a uint4 as R11G11B10
+	float3 DiffuseColor;
+	float3 DiffuseColorForIndirect;
+	float3 SubsurfaceColorForIndirect;
+	FShadingOcclusion ShadingOcclusion;
+};
+
+struct FStrataPrecomputedLightingData
+{
+	uint BSDFCount;
+	FStrataBSDFForPrecomputedLighting BSDFs[STRATA_BSDF_COUNT_FOR_PRECOMPUTED];
+};
+
 FStrataBSDFForPrecomputedLighting StrataPrepareBSDFForPrecomputedLighting(
 	in FStrataIntegrationSettings Settings,
 	in FStrataPixelHeader StrataPixelHeader,
 	in FStrataBSDF CurrentBSDF,
 	in float3 V,
-	in float3 WorldBentNormal0);
-#endif
+	in float3 WorldBentNormal0)
+{
+	FStrataBSDFForPrecomputedLighting BSDFForPrecomputedLighting = (FStrataBSDFForPrecomputedLighting)0;
+
+	FStrataAddressing NullStrataAddressing = (FStrataAddressing)0;
+	FStrataBSDFContext StrataBSDFContext = StrataCreateBSDFContext(StrataPixelHeader, CurrentBSDF, NullStrataAddressing, V);
+
+	// Update a context specific for environment lighting
+	FStrataBSDFContext EnvBSDFContext = StrataBSDFContext;
+	EnvBSDFContext.N = BSDFForPrecomputedLighting.ShadingOcclusion.BentNormal;
+	StrataUpdateBSDFContext(EnvBSDFContext, EnvBSDFContext.L);
+
+	// Evaluate the bent normal
+	BSDFForPrecomputedLighting.ShadingOcclusion = ApplyBentNormal(
+		V,
+		EnvBSDFContext.N,
+		WorldBentNormal0,
+		StrataGetBSDFRoughness(CurrentBSDF),
+		StrataPixelHeader.IrradianceAO.MaterialAO);
+
+	// For diffuse, we specify a perpendicular to the surface light direction for the transmittance to light to not be view dependent.
+	float DiffuseEnvLightingNoL = 1.0f;
+	float3 LuminanceWeightFinal = LuminanceWeight(DiffuseEnvLightingNoL, CurrentBSDF);
+
+	// And evaluate diffuse parameters
+	const bool bEnableSpecular = false;
+	FStrataEnvLightResult StrataEnvLight = StrataEvaluateForEnvLight(EnvBSDFContext, bEnableSpecular, Settings);
+	BSDFForPrecomputedLighting.DiffuseNormal = StrataEnvLight.DiffuseNormal;
+	BSDFForPrecomputedLighting.DiffuseColor = StrataEnvLight.DiffuseColor * LuminanceWeightFinal;
+	BSDFForPrecomputedLighting.DiffuseColorForIndirect = StrataEnvLight.DiffuseWeight * LuminanceWeightFinal;
+	BSDFForPrecomputedLighting.SubsurfaceColorForIndirect = StrataEnvLight.DiffuseBackFaceWeight * LuminanceWeightFinal;
+
+	return BSDFForPrecomputedLighting;
+}
+
+#endif // MATERIAL_STRATA_OPAQUE_PRECOMPUTED_LIGHTING
 
 // If this is changed, please update the compiler side material size evaluation in StrataMaterial.cpp
 void PackStrataOut(
@@ -41,7 +106,7 @@ void PackStrataOut(
 	inout FStrataSubsurfaceData SSSData,
 	inout FStrataTopLayerData TopLayerData,
 	inout FStrataOpaqueRoughRefractionData OpaqueRoughRefractionData
-#if MATERIAL_STRATA_PRECOMPUTED_LIGHTING
+#if MATERIAL_STRATA_OPAQUE_PRECOMPUTED_LIGHTING
 	,
 	inout FStrataPrecomputedLightingData IndirectLightingData
 #endif
@@ -50,7 +115,7 @@ void PackStrataOut(
 	// This only exists with inline shading and when we are going to write out BSDFs (UpdateAllXXX functions needs to be defined)
 #if STRATA_INLINE_SHADING && STRATA_CLAMPED_BSDF_COUNT > 0
 
-#if MATERIAL_STRATA_PRECOMPUTED_LIGHTING
+#if MATERIAL_STRATA_OPAQUE_PRECOMPUTED_LIGHTING
 	IndirectLightingData = (FStrataPrecomputedLightingData)0;
 #endif 
 
@@ -419,7 +484,7 @@ void PackStrataOut(
 					}
 
 					// Convert BSDF to data ready to evaluate precomputed lighting
-					#if MATERIAL_STRATA_PRECOMPUTED_LIGHTING
+					#if MATERIAL_STRATA_OPAQUE_PRECOMPUTED_LIGHTING
 					IndirectLightingData.BSDFs[IndirectLightingData.BSDFCount++] = StrataPrepareBSDFForPrecomputedLighting(
 						Settings,
 						StrataPixelHeader,
