Fix shader compile error when  NEEDS_BASEPASS_PIXEL_FOGGING and  PROJECT_SUPPORT_SKY_ATMOSPHERE are off

diff --git a/Engine/Shaders/Private/BasePassPixelShader.usf b/Engine/Shaders/Private/BasePassPixelShader.usf
--- a/Engine/Shaders/Private/BasePassPixelShader.usf
+++ b/Engine/Shaders/Private/BasePassPixelShader.usf
@@ -1438,6 +1438,8 @@ void FPixelShaderInOut_MainPS(
 				#if FORWARD_SHADING || TRANSLUCENCY_LIGHTING_SURFACE_FORWARDSHADING || MATERIAL_SHADINGMODEL_SINGLELAYERWATER
 
 					float3 DirectionalLightAtmosphereTransmittance = 1.0f;
+
+					#if NEEDS_BASEPASS_PIXEL_FOGGING && PROJECT_SUPPORT_SKY_ATMOSPHERE
 					const uint LightIndex = 0;
 					if (ResolvedView.AtmosphereLightDiscCosHalfApexAngle_PPTrans[LightIndex].y > 0.0f)
 					{
@@ -1447,6 +1449,7 @@ void FPixelShaderInOut_MainPS(
 							PlanetCenterToTranslatedWorldPos, ResolvedView.AtmosphereLightDirection[LightIndex].xyz, ResolvedView.SkyAtmosphereBottomRadiusKm, ResolvedView.SkyAtmosphereTopRadiusKm,
 							View.TransmittanceLutTexture, View.TransmittanceLutTextureSampler);
 					}
+					#endif // NEEDS_BASEPASS_PIXEL_FOGGING && PROJECT_SUPPORT_SKY_ATMOSPHERE
 
 					FDeferredLightingSplit ForwardDirectLighting = GetForwardDirectLightingSplit(
 						GridIndex, MaterialParameters.WorldPosition_CamRelative, MaterialParameters.CameraVector, GBuffer, NearestResolvedDepthScreenUV, MaterialParameters.PrimitiveId, EyeIndex, Dither, 
