Fix Strata visualization name typo

diff --git a/Engine/Shaders/Private/Strata/StrataVisualize.usf b/Engine/Shaders/Private/Strata/StrataVisualize.usf
--- a/Engine/Shaders/Private/Strata/StrataVisualize.usf
+++ b/Engine/Shaders/Private/Strata/StrataVisualize.usf
@@ -121,7 +121,7 @@ void PrintPixelType(inout FShaderPrintContext Ctx, bool bIsSimpleMaterial, bool
 void PrintSSSType(inout FShaderPrintContext Ctx, uint SSSType, FFontColor InColor)
 {
 		 if (SSSType == SSS_TYPE_WRAP)				Print(Ctx, TEXT("Wrap"), InColor);
-	else if (SSSType == SSS_TYPE_DIFFUSION)			Print(Ctx, TEXT("Diffuse"), InColor);
+	else if (SSSType == SSS_TYPE_DIFFUSION)			Print(Ctx, TEXT("Diffusion"), InColor);
 	else if (SSSType == SSS_TYPE_DIFFUSION_PROFILE)	Print(Ctx, TEXT("Diffusion Profile"), InColor);
 	else /*	(SSSType == SSS_TYPE_INVALID) */		Print(Ctx, TEXT("Invalid"), InColor);
 }
