Add pre-integrated SSS material conversion.

diff --git a/Engine/Shaders/Private/Strata/StrataLegacyConversion.ush b/Engine/Shaders/Private/Strata/StrataLegacyConversion.ush
--- a/Engine/Shaders/Private/Strata/StrataLegacyConversion.ush
+++ b/Engine/Shaders/Private/Strata/StrataLegacyConversion.ush
@@ -219,7 +219,7 @@ FStrataData StrataConvertLegacyMaterialDynamic(
 		FinalWeight = Opacity;
 #endif
 	}
-	else if (ShadingModel == SHADINGMODELID_SUBSURFACE || ShadingModel == SHADINGMODELID_PREINTEGRATED_SKIN)
+	else if (ShadingModel == SHADINGMODELID_SUBSURFACE)
 	{
 		const float ThicknessInCm = STRATA_SIMPLEVOLUME_THICKNESS_CM;
 		const float3 MFPInCm = TransmittanceToMeanFreePath(SubSurfaceColor, ThicknessInCm * CENTIMETER_TO_METER) * METER_TO_CENTIMETER;
@@ -253,6 +253,34 @@ FStrataData StrataConvertLegacyMaterialDynamic(
 
 		FinalWeight = Opacity;
 #endif
+	}
+	else if (ShadingModel == SHADINGMODELID_PREINTEGRATED_SKIN)
+	{
+		// Use default profile MFP (e.g., human skin) as MFP value for converting hardcoded pre-integrated SSS texture for deferred material.
+		const float3 MFPInCm = float3(1.0f, 0.088964f, 0.072095f) * 2.6748f * 0.1f;
+
+		// Legacy material uses Subsurface color as transmission 'tinting', but we can represent that with a single layer. So instead we take 
+		// the max color of SubSurfaceColro & BaseColor
+		Top.BaseColor = max(SubSurfaceColor, BaseColor); 
+		Top.Specular = Specular;
+		Top.Metallic = Metallic;
+		Top.Roughness = Roughness;
+		Top.Anisotropy = Anisotropy;
+		Top.SSSMFP = max(1e-05f, MFPInCm); // Ensure the MFP is not null to force the material as have SSS
+		Top.SSSMFPScale = 1.0f;
+		Top.SSSProfileID = 0.f;
+		Top.SSSMFPScale = 1.f - Opacity;
+		Top.SSSPhaseAniso = 0.93f;
+		Top.UseSSSDiffusion = true;
+		Top.Emissive = Emissive;
+
+		// No need to take into account opacity weighting for translucent material, as the coverage is assumed to be 1
+		Bot.Weight = 0.f;
+
+		FinalWeight = 1.0f;
+#if STRATA_TRANSLUCENT_MATERIAL
+		FinalWeight = Opacity;
+#endif
 	}
 	else if (ShadingModel == SHADINGMODELID_SUBSURFACE_PROFILE)
 	{
