Add initial support for rendering water in path tracer:
* Add missing water Material ID cases for material hit shaders
* Make water translucent in the initial version
* Disable underwater postprocess material when running PT

diff --git a/Engine/Plugins/Experimental/Water/Source/Runtime/Private/WaterSubsystem.cpp b/Engine/Plugins/Experimental/Water/Source/Runtime/Private/WaterSubsystem.cpp
--- a/Engine/Plugins/Experimental/Water/Source/Runtime/Private/WaterSubsystem.cpp
+++ b/Engine/Plugins/Experimental/Water/Source/Runtime/Private/WaterSubsystem.cpp
@@ -666,7 +666,7 @@ void UWaterSubsystem::ComputeUnderwaterPostProcess(FVector ViewLocation, FSceneV
 
 	SceneView->UnderwaterDepth = CachedDepthUnderwater;
 
-	if (!bUnderwaterForPostProcess || !IsUnderwaterPostProcessEnabled())
+	if (!bUnderwaterForPostProcess || !IsUnderwaterPostProcessEnabled() || SceneView->Family->EngineShowFlags.PathTracing)
 	{
 		UnderwaterPostProcessVolume.PostProcessProperties.bIsEnabled = false;
 		UnderwaterPostProcessVolume.PostProcessProperties.Settings = nullptr;
diff --git a/Engine/Shaders/Private/PathTracing/Material/PathTracingMaterialSampling.ush b/Engine/Shaders/Private/PathTracing/Material/PathTracingMaterialSampling.ush
--- a/Engine/Shaders/Private/PathTracing/Material/PathTracingMaterialSampling.ush
+++ b/Engine/Shaders/Private/PathTracing/Material/PathTracingMaterialSampling.ush
@@ -146,6 +146,7 @@ float3 EstimateMaterialAlbedo(FPathTracingPayload Payload)
 		case SHADINGMODELID_SUBSURFACE_PROFILE:
 			Albedo = BrdfWeight * (Payload.DiffuseColor + Payload.SpecularColor + Payload.GetSubsurfaceColor());
 			break;
+		case SHADINGMODELID_SINGLELAYERWATER:
 		case SHADINGMODELID_DEFAULT_LIT:
 		case SHADINGMODELID_CLEAR_COAT:
 		case SHADINGMODELID_EYE:
@@ -232,6 +233,7 @@ void AccumulateAlbedoNormal(inout FPathTracingPayload Payload, float3 PathThroug
 		Albedo += BrdfWeight * PathThroughput * (Payload.DiffuseColor + RoughnessWeight * Payload.SpecularColor + Payload.GetSubsurfaceColor());
 		Normal += Luminance(BrdfWeight * PathThroughput) * Payload.WorldNormal;
 		break;
+	case SHADINGMODELID_SINGLELAYERWATER:
 	case SHADINGMODELID_DEFAULT_LIT:
 	case SHADINGMODELID_CLEAR_COAT:
 	case SHADINGMODELID_EYE:
@@ -291,6 +293,7 @@ void AdjustMaterialRoughness(inout FPathTracingPayload Payload, float PathRoughn
 			// clamp coat roughness
 			Payload.SetClearCoatRoughness(max(Payload.GetClearCoatRoughness(), PathRoughness));
 			// fallthrough
+		case SHADINGMODELID_SINGLELAYERWATER:
 		case SHADINGMODELID_TWOSIDED_FOLIAGE:
 		case SHADINGMODELID_DEFAULT_LIT:
 		case SHADINGMODELID_THIN_TRANSLUCENT:
@@ -532,6 +535,7 @@ FMaterialSample SampleMaterial(
 		case SHADINGMODELID_DEFAULT_LIT:
 		case SHADINGMODELID_SUBSURFACE:
 		case SHADINGMODELID_PREINTEGRATED_SKIN:
+		case SHADINGMODELID_SINGLELAYERWATER:
 			Result = DefaultLit_SampleMaterial(RayDirection, Payload, RandSample);
 			break;
 		case SHADINGMODELID_SUBSURFACE_PROFILE:
@@ -634,6 +638,7 @@ FMaterialEval EvalMaterial(
 		case SHADINGMODELID_DEFAULT_LIT:
 		case SHADINGMODELID_SUBSURFACE:
 		case SHADINGMODELID_PREINTEGRATED_SKIN:
+		case SHADINGMODELID_SINGLELAYERWATER:
 			Result = DefaultLit_EvalMaterial(IncomingDirection, OutgoingDirection, Payload);
 			break;
 		case SHADINGMODELID_SUBSURFACE_PROFILE:
diff --git a/Engine/Shaders/Private/PathTracing/PathTracingCore.ush b/Engine/Shaders/Private/PathTracing/PathTracingCore.ush
--- a/Engine/Shaders/Private/PathTracing/PathTracingCore.ush
+++ b/Engine/Shaders/Private/PathTracing/PathTracingCore.ush
@@ -1198,6 +1198,22 @@
 		}
 	}
 
+	if (Payload.ShadingModelID == SHADINGMODELID_SINGLELAYERWATER)
+	{
+		float3 LocalSigmaT = max(Payload.CustomData.xyz, 1e-8f);
+		if (Payload.IsFrontFace())
+		{
+			// entering
+			PathState.SigmaT += LocalSigmaT;
+		}
+		else
+		{
+			// exiting
+			PathState.SigmaT -= LocalSigmaT;
+			PathState.SigmaT = max(PathState.SigmaT, 0);
+		}
+	}
+
 	// If we are using Material sampling for lights
 	if (MISMode != 1 && bDoLightLoop)
 	{
diff --git a/Engine/Shaders/Private/PathTracing/PathTracingMaterialHitShader.usf b/Engine/Shaders/Private/PathTracing/PathTracingMaterialHitShader.usf
--- a/Engine/Shaders/Private/PathTracing/PathTracingMaterialHitShader.usf
+++ b/Engine/Shaders/Private/PathTracing/PathTracingMaterialHitShader.usf
@@ -204,14 +204,34 @@ RAY_TRACING_ENTRY_CLOSEST_HIT(PathTracingMaterialCHS,
 	else
 #endif
 
-#if MATERIAL_SHADINGMODEL_THIN_TRANSLUCENT
-	if (Payload.ShadingModelID == SHADINGMODELID_THIN_TRANSLUCENT)
+#if MATERIAL_SHADINGMODEL_SINGLELAYERWATER
+	if (Payload.ShadingModelID == SHADINGMODELID_SINGLELAYERWATER)
 	{
-		Payload.CustomData.xyz = GetThinTranslucentMaterialOutput0(MaterialParameters);
+		const float3 ScatteringCoeff = max(0.0f, LWCToFloat(GetSingleLayerWaterMaterialOutput0(MaterialParameters)));
+		const float3 AbsorptionCoeff = max(0.0f, LWCToFloat(GetSingleLayerWaterMaterialOutput1(MaterialParameters)));
+		const float PhaseG = clamp(LWCToFloat(GetSingleLayerWaterMaterialOutput2(MaterialParameters)), -1.0f, 1.0f);
+
+		const float3 WaterExtinction = ScatteringCoeff + AbsorptionCoeff;
+		const float3 WaterAlbedo = ScatteringCoeff / WaterExtinction;
+
+		Payload.BaseColor.xyz = WaterAlbedo;
+		Payload.CustomData.xyz = WaterExtinction;
+		Payload.CustomData.w = PhaseG;
+
+		Payload.Ior = DielectricF0ToIor(DielectricSpecularToF0(Specular));
+		Payload.BlendingMode = RAY_TRACING_BLEND_MODE_TRANSLUCENT;
 	}
 	else
 #endif
 
+#if MATERIAL_SHADINGMODEL_THIN_TRANSLUCENT
+		if (Payload.ShadingModelID == SHADINGMODELID_THIN_TRANSLUCENT)
+		{
+			Payload.CustomData.xyz = GetThinTranslucentMaterialOutput0(MaterialParameters);
+		}
+		else
+#endif
+
 #if MATERIAL_SHADINGMODEL_HAIR
 	if (Payload.ShadingModelID == SHADINGMODELID_HAIR)
 	{
diff --git a/Engine/Shaders/Private/RayTracing/RayTracingMaterialHitShaders.usf b/Engine/Shaders/Private/RayTracing/RayTracingMaterialHitShaders.usf
--- a/Engine/Shaders/Private/RayTracing/RayTracingMaterialHitShaders.usf
+++ b/Engine/Shaders/Private/RayTracing/RayTracingMaterialHitShaders.usf
@@ -382,6 +382,25 @@ void GetMaterialPayload(
 	}
 #endif
 
+#if MATERIAL_SHADINGMODEL_SINGLELAYERWATER
+	if (Payload.ShadingModelID == SHADINGMODELID_SINGLELAYERWATER)
+	{
+		const float3 ScatteringCoeff = max(0.0f, LWCToFloat(GetSingleLayerWaterMaterialOutput0(MaterialParameters)));
+		const float3 AbsorptionCoeff = max(0.0f, LWCToFloat(GetSingleLayerWaterMaterialOutput1(MaterialParameters)));
+		const float PhaseG = clamp(LWCToFloat(GetSingleLayerWaterMaterialOutput2(MaterialParameters)), -1.0f, 1.0f);
+
+		const float3 WaterExtinction = ScatteringCoeff + AbsorptionCoeff;
+		const float3 WaterAlbedo = ScatteringCoeff / WaterExtinction;
+
+		Payload.BaseColor.xyz = WaterAlbedo;
+		Payload.CustomData.xyz = WaterExtinction;
+		Payload.CustomData.w = PhaseG;
+
+		Payload.Ior = DielectricF0ToIor(DielectricSpecularToF0(Specular));
+		Payload.BlendingMode = RAY_TRACING_BLEND_MODE_TRANSLUCENT;
+	}
+#endif
+
 #if MATERIAL_SHADINGMODEL_THIN_TRANSLUCENT
 	if (Payload.ShadingModelID == SHADINGMODELID_THIN_TRANSLUCENT)
 	{
