Merge Strata  thin surface & wrap evaluation code.

diff --git a/Engine/Shaders/Private/Strata/StrataEvaluation.ush b/Engine/Shaders/Private/Strata/StrataEvaluation.ush
--- a/Engine/Shaders/Private/Strata/StrataEvaluation.ush
+++ b/Engine/Shaders/Private/Strata/StrataEvaluation.ush
@@ -201,7 +201,6 @@ struct FStrataEvaluateResult
 	float3 SpecularPathValue;
 	float3 SpecularHazePathValue;
 	float3 TransmissionPathValue;
-	float3 ThinTransmissionPathValue;
 
 	// The following probabilities are part of the pdf to work correctly with a monter carlo based integrator and as expected are not applied on the xxxPathValue.
 	// So PathProbability should be applied on PathValue in the rasteriser to recover the correct balance.
@@ -212,7 +211,6 @@ struct FStrataEvaluateResult
 	float  SpecularPDF;
 	float  SpecularHazePDF;
 	float  TransmissionPDF;
-	float  ThinTransmissionPathPDF;
 
 	float3 ThroughputV;					// Throughput to the next layer (from the view to the shading point). Contains many things such as transmittance or cloth energy conservation/preservation factors.
 	float3 TransmittanceAlongN;			// Transmittance to the next layer along the surface normal. This is used later to compute the transmittance towards the light.
@@ -550,11 +548,12 @@ FStrataEvaluateResult StrataEvaluateBSDFCommon(FStrataBSDFContext BSDFContext, F
 
 			// Simple port of the legacy code to ensure we can reproduce the look first
 			#if STRATA_FASTPATH==0
-			if (BSDF_GETSSSTYPE(BSDFContext.BSDF) == SSS_TYPE_WRAP)
+			const bool bIsThin = BSDF_GETISTHIN(BSDFContext.BSDF);
+			if (BSDF_GETSSSTYPE(BSDFContext.BSDF) == SSS_TYPE_WRAP || bIsThin)
 			{
-				const float3 DiffuseColor = float3(1, 1, 1);
+				const float3 SlabDiffuseColor = bIsThin ? DiffuseColor : float3(1, 1, 1);
 				const float TransmittanceNoL = 1.0f;
-				const FParticipatingMedia PM = StrataSlabCreateParticipatingMedia(DiffuseColor, SLAB_SSSMFP(BSDFContext.BSDF));
+				const FParticipatingMedia PM = StrataSlabCreateParticipatingMedia(SlabDiffuseColor, SLAB_SSSMFP(BSDFContext.BSDF));
 				const float3 TransmissionThroughput = IsotropicMediumSlabTransmittance(PM, STRATA_SIMPLEVOLUME_THICKNESS_M, TransmittanceNoL);
 
 				const float VoL = dot(BSDFContext.L, -BSDFContext.V);
@@ -568,9 +567,10 @@ FStrataEvaluateResult StrataEvaluateBSDFCommon(FStrataBSDFContext BSDFContext, F
 				// 
 				// Note: The function ApproximateHG is not correctly normalized, and we should be using HenyeyGreensteinPhase instead. However the look is quite 
 				//       different and we need to handle legacy conversion if we need to change this. 
+				const float HemisphericPhaseFunction = saturate(-BSDFContext.Context.NoL);
 				const float IsotropicPhaseFunction = lerp(saturate(NoH) / (2.0f * PI), 1.f / (4.f * PI), abs(PhaseAnisotropy));
 				const float AnisotropicPhaseFunction = ApproximateHG(VoL, PhaseAnisotropy);
-				const float EffectivePhaseFunction = lerp(AnisotropicPhaseFunction, IsotropicPhaseFunction, lerp(1.f, 0.5f, abs(PhaseAnisotropy)));
+				const float EffectivePhaseFunction = bIsThin ? HemisphericPhaseFunction : lerp(AnisotropicPhaseFunction, IsotropicPhaseFunction, lerp(1.f, 0.5f, abs(PhaseAnisotropy)));
 
 				Sample.TransmissionPathValue = TransmissionThroughput * EffectivePhaseFunction * DirectionalAlbedo_SpecularTransmission;
 				Sample.TransmissionPDF = 1.0f / (4.0f * PI); // STRATA_TODO this currently match the uniform sphere sampling from StrataImportanceSampleBSDF
@@ -620,39 +620,6 @@ FStrataEvaluateResult StrataEvaluateBSDFCommon(FStrataBSDFContext BSDFContext, F
 			}
 			#endif
 
-			////
-			//// Evaluate Thin light transmittance BSDF, transmitting light on the other side of thin surfaces.
-			////
-
-			#if STRATA_FASTPATH==0
-			BRANCH
-			if (BSDF_GETISTHIN(BSDFContext.BSDF))
-			{
-				FParticipatingMedia PM = StrataSlabCreateParticipatingMedia(DiffuseColor, SLAB_SSSMFP(BSDFContext.BSDF));
-				const float BackFaceNoL = -BSDFContext.Context.NoL;
-				const float BackFaceNoLSat = saturate(BackFaceNoL);
-
-				// We consider single scattering and perpendicular to the surface light transmittance.
-				// Using BackFaceNoL instead of 1 for TransmittanceNoL would be more correct but it gives some high frequency details looking weird without multiple scattering.
-				const float TransmittanceNoL = 1.0f;	
-				const float3 SlabTransmittance = IsotropicMediumSlabTransmittance(PM, STRATA_SIMPLEVOLUME_THICKNESS_M, TransmittanceNoL);
-
-				// We recover the transmitted luminance from as the back face Lambert affected by the transmittance.
-				// The MFP is specified on the slab node and then transmittance recoved from it.
-				// We assume the back face diffuse albedo is white=1 and no diffusion happens at this stage.
-				const float3 BackFaceWhiteDiffuseColor = float3(1.0f, 1.0f, 1.0f);
-
-				// The diffuse is attenuated by the slab transmittance and the GGX interface SpecularTransmission.
-				Sample.ThinTransmissionPathValue = Diffuse_Lambert(BackFaceWhiteDiffuseColor) * SlabTransmittance * DirectionalAlbedo_SpecularTransmission;
-				Sample.ThinTransmissionPathPDF = BackFaceNoLSat / PI;
-
-				// We do not want thin (a.k.a. two sided lighting) contribution to go the SSS post process path.
-				Sample.bSubsurface = false;
-
-				Sample.IntegratedDiffuseValue += ShadowTerms.TransmissionShadow * Sample.ThinTransmissionPathValue * BackFaceNoLSat * AreaLightContext.Falloff * AreaLightContext.AreaLight.FalloffColor;
-			}
-			#endif
-
 			////
 			//// Evaluate a layer of participating media: scattering and transmittance.
 			//// This is used for optically thin translucent objects, or non-bottom layer of a material.
