Strata - fixed compilation issue without dxc. Also a nice clean up.

diff --git a/Engine/Shaders/Private/BasePassPixelShader.usf b/Engine/Shaders/Private/BasePassPixelShader.usf
--- a/Engine/Shaders/Private/BasePassPixelShader.usf
+++ b/Engine/Shaders/Private/BasePassPixelShader.usf
@@ -844,29 +844,6 @@ FShadingOcclusion ApplyBentNormal(
 
 #if STRATA_ENABLED
 
-#if MATERIAL_STRATA_OPAQUE_PRECOMPUTED_LIGHTING
-float3 StrataGetBSDFPrecomputedLighting(
-	in FStrataBSDFForPrecomputedLighting BSDFForPrecomputedLighting,
-	in FMaterialPixelParameters MaterialParameters,
-	in FVertexFactoryInterpolantsVSToPS Interpolants,
-	in FBasePassInterpolantsVSToPS BasePassInterpolants,
-	in VTPageTableResult LightmapVTPageTableResult,
-	in float3 VolumetricLightmapBrickTextureUVs)
-{
-	float3 DiffuseLighting = 0.0f;
-	float3 SubsurfaceLighting = 0.0f;
-	float  IndirectIrrad1iance = 0.0f;
-	const bool bEvaluateBackface = any(BSDFForPrecomputedLighting.SubsurfaceColorForIndirect > 0.0);	// With Strata we have no other way to know before if backface lighting will be needed
-	GetPrecomputedIndirectLightingAndSkyLight(MaterialParameters, Interpolants, BasePassInterpolants, LightmapVTPageTableResult, bEvaluateBackface, BSDFForPrecomputedLighting.DiffuseNormal, VolumetricLightmapBrickTextureUVs,
-		DiffuseLighting, SubsurfaceLighting, IndirectIrrad1iance);
-
-	float3 OutLuminance = (DiffuseLighting * BSDFForPrecomputedLighting.DiffuseColorForIndirect + SubsurfaceLighting * BSDFForPrecomputedLighting.SubsurfaceColorForIndirect)
-		* AOMultiBounce(BSDFForPrecomputedLighting.DiffuseColor, BSDFForPrecomputedLighting.ShadingOcclusion.DiffOcclusion); // STRATA_TODO This should only be applied on the aggregated diffuse colors, not each contribution independently. But we also need to account for DiffOcclusion which looks normal dependent
-
-	return OutLuminance;
-}
-#endif
-
 uint GetDiffuseIndirectSampleOcclusion(FSharedLocalBases SharedLocalBases, float3 V, float2 SvPosition, float MaterialAO)
 {
 	uint DiffuseIndirectSampleOcclusion = 0;
@@ -1757,9 +1734,6 @@ void FPixelShaderInOut_MainPS(
 		uint2 PixelPos = uint2(In.SvPosition.xy);
 		const float3 WorldBentNormal0 = GetWorldBentNormalZero(MaterialParameters);
 
-	#if MATERIAL_STRATA_OPAQUE_PRECOMPUTED_LIGHTING
-		FStrataPrecomputedLightingData PrecomputedLightingData = (FStrataPrecomputedLightingData)0;
-	#endif
 		FStrataSubsurfaceData SSSData = (FStrataSubsurfaceData)0;
 		FStrataTopLayerData TopLayerData = (FStrataTopLayerData)0;
 		FStrataOpaqueRoughRefractionData OpaqueRoughRefractionData = (FStrataOpaqueRoughRefractionData)0;
@@ -1781,24 +1755,14 @@ void FPixelShaderInOut_MainPS(
 			StrataPixelHeader, StrataData, MaterialParameters.CameraVector, WorldBentNormal0, bStrataSubsurfaceEnable, EmissiveLuminance,
 			SSSData, TopLayerData, OpaqueRoughRefractionData
 	#if MATERIAL_STRATA_OPAQUE_PRECOMPUTED_LIGHTING
-			, PrecomputedLightingData
+			,MaterialParameters
+			,Interpolants
+			,BasePassInterpolants
+			,LightmapVTPageTableResult
+			,VolumetricLightmapBrickTextureUVs
 	#endif
 		);
 
-		float3 PreComputedLightingLuminance = 0.0f;
-	#if MATERIAL_STRATA_OPAQUE_PRECOMPUTED_LIGHTING
-		for (uint i = 0; i < PrecomputedLightingData.BSDFCount; ++i)
-		{
-			PreComputedLightingLuminance += StrataGetBSDFPrecomputedLighting(
-				PrecomputedLightingData.BSDFs[i],
-				MaterialParameters,
-				Interpolants,
-				BasePassInterpolants,
-				LightmapVTPageTableResult,
-				VolumetricLightmapBrickTextureUVs);
-		}
-	#endif
-
 		// Write out MRT data
 	#if STRATA_BASE_PASS_MRT_OUTPUT_COUNT != 2
 	#error Strata STRATA_BASE_PASS_MRT_OUTPUT_COUNT has been update but not StrataOutput
@@ -1835,7 +1799,7 @@ void FPixelShaderInOut_MainPS(
 		Color = lerp(Color, TopLayerData.BaseColor, View.UnlitViewmodeMask);
 	#endif
 
-		Color += EmissiveLuminance + PreComputedLightingLuminance;
+		Color += EmissiveLuminance;
 	}
 
 #endif // STRATA_OPAQUE_DEFERRED
diff --git a/Engine/Shaders/Private/Strata/StrataExport.ush b/Engine/Shaders/Private/Strata/StrataExport.ush
--- a/Engine/Shaders/Private/Strata/StrataExport.ush
+++ b/Engine/Shaders/Private/Strata/StrataExport.ush
@@ -24,66 +24,67 @@ FShadingOcclusion ApplyBentNormal(
 	in float Roughness,
 	in float MaterialAO);
 
-
-#if STRATA_CLAMPED_BSDF_COUNT <= 0
-	#define STRATA_BSDF_COUNT_FOR_PRECOMPUTED 1
-#else
-	#define STRATA_BSDF_COUNT_FOR_PRECOMPUTED STRATA_CLAMPED_BSDF_COUNT
-#endif
-
-struct FStrataBSDFForPrecomputedLighting
-{
-	float3 DiffuseNormal;// STRATA_TODO test packed on a uint4 as R11G11B10
-	float3 DiffuseColor;
-	float3 DiffuseColorForIndirect;
-	float3 SubsurfaceColorForIndirect;
-	FShadingOcclusion ShadingOcclusion;
-};
-
-struct FStrataPrecomputedLightingData
-{
-	uint BSDFCount;
-	FStrataBSDFForPrecomputedLighting BSDFs[STRATA_BSDF_COUNT_FOR_PRECOMPUTED];
-};
-
-FStrataBSDFForPrecomputedLighting StrataPrepareBSDFForPrecomputedLighting(
+void GetPrecomputedIndirectLightingAndSkyLight(
+	FMaterialPixelParameters MaterialParameters,
+	FVertexFactoryInterpolantsVSToPS Interpolants,
+	FBasePassInterpolantsVSToPS BasePassInterpolants,
+	VTPageTableResult LightmapVTPageTableResult,
+	bool bEvaluateBackface,
+	float3 DiffuseDir,
+	float3 VolumetricLightmapBrickTextureUVs,
+	out float3 OutDiffuseLighting,
+	out float3 OutSubsurfaceLighting,
+	out float OutIndirectIrradiance);
+
+float3 StrataGetBSDFPrecomputedLighting(
 	in FStrataIntegrationSettings Settings,
 	in FStrataPixelHeader StrataPixelHeader,
 	in FStrataBSDF CurrentBSDF,
 	in float3 V,
-	in float3 WorldBentNormal0)
+	in float3 WorldBentNormal0,
+	in FMaterialPixelParameters MaterialParameters,
+	in FVertexFactoryInterpolantsVSToPS Interpolants,
+	in FBasePassInterpolantsVSToPS BasePassInterpolants,
+	in VTPageTableResult LightmapVTPageTableResult,
+	in float3 VolumetricLightmapBrickTextureUVs)
 {
-	FStrataBSDFForPrecomputedLighting BSDFForPrecomputedLighting = (FStrataBSDFForPrecomputedLighting)0;
-
 	FStrataAddressing NullStrataAddressing = (FStrataAddressing)0;
 	FStrataBSDFContext StrataBSDFContext = StrataCreateBSDFContext(StrataPixelHeader, CurrentBSDF, NullStrataAddressing, V);
 
-	// Update a context specific for environment lighting
-	FStrataBSDFContext EnvBSDFContext = StrataBSDFContext;
-	EnvBSDFContext.N = BSDFForPrecomputedLighting.ShadingOcclusion.BentNormal;
-	StrataUpdateBSDFContext(EnvBSDFContext, EnvBSDFContext.L);
-
-	// Evaluate the bent normal
-	BSDFForPrecomputedLighting.ShadingOcclusion = ApplyBentNormal(
+	// Evaluate shading occlusion data and bent normal
+	FShadingOcclusion ShadingOcclusion = ApplyBentNormal(
 		V,
-		EnvBSDFContext.N,
+		StrataBSDFContext.N,
 		WorldBentNormal0,
 		StrataGetBSDFRoughness(CurrentBSDF),
 		StrataPixelHeader.IrradianceAO.MaterialAO);
 
+	// Update a context specific for environment lighting
+	FStrataBSDFContext EnvBSDFContext = StrataBSDFContext;
+	EnvBSDFContext.N = ShadingOcclusion.BentNormal;
+	StrataUpdateBSDFContext(EnvBSDFContext, EnvBSDFContext.L);
+
 	// For diffuse, we specify a perpendicular to the surface light direction for the transmittance to light to not be view dependent.
 	float DiffuseEnvLightingNoL = 1.0f;
 	float3 LuminanceWeightFinal = LuminanceWeight(DiffuseEnvLightingNoL, CurrentBSDF);
 
-	// And evaluate diffuse parameters
+	// Evaluate material environment ligthing parameters
 	const bool bEnableSpecular = false;
 	FStrataEnvLightResult StrataEnvLight = StrataEvaluateForEnvLight(EnvBSDFContext, bEnableSpecular, Settings);
-	BSDFForPrecomputedLighting.DiffuseNormal = StrataEnvLight.DiffuseNormal;
-	BSDFForPrecomputedLighting.DiffuseColor = StrataEnvLight.DiffuseColor * LuminanceWeightFinal;
-	BSDFForPrecomputedLighting.DiffuseColorForIndirect = StrataEnvLight.DiffuseWeight * LuminanceWeightFinal;
-	BSDFForPrecomputedLighting.SubsurfaceColorForIndirect = StrataEnvLight.DiffuseBackFaceWeight * LuminanceWeightFinal;
 
-	return BSDFForPrecomputedLighting;
+	// Fetch precomputed lighting
+	float3 DiffuseLighting = 0.0f;
+	float3 SubsurfaceLighting = 0.0f;
+	float  IndirectIrrad1iance = 0.0f;
+	const bool bEvaluateBackface = any(StrataEnvLight.DiffuseBackFaceWeight > 0.0);	// With Strata we have no other way to know if backface lighting will be needed
+	GetPrecomputedIndirectLightingAndSkyLight(MaterialParameters, Interpolants, BasePassInterpolants, LightmapVTPageTableResult, bEvaluateBackface, StrataEnvLight.DiffuseNormal, VolumetricLightmapBrickTextureUVs,
+		DiffuseLighting, SubsurfaceLighting, IndirectIrrad1iance);
+
+	// Evaluate precomputed lighting according to material parameters
+	float3 OutLuminance = (DiffuseLighting * StrataEnvLight.DiffuseWeight + SubsurfaceLighting * StrataEnvLight.DiffuseBackFaceWeight) * LuminanceWeightFinal
+		* AOMultiBounce(StrataEnvLight.DiffuseColor, ShadingOcclusion.DiffOcclusion); // STRATA_TODO This should only be applied on the aggregated diffuse colors, not each contribution independently. But we also need to account for DiffOcclusion which looks normal dependent
+
+	return OutLuminance;
 }
 
 #endif // MATERIAL_STRATA_OPAQUE_PRECOMPUTED_LIGHTING
@@ -108,17 +109,17 @@ void PackStrataOut(
 	inout FStrataOpaqueRoughRefractionData OpaqueRoughRefractionData
 #if MATERIAL_STRATA_OPAQUE_PRECOMPUTED_LIGHTING
 	,
-	inout FStrataPrecomputedLightingData IndirectLightingData
+	in FMaterialPixelParameters MaterialParameters,
+	in FVertexFactoryInterpolantsVSToPS Interpolants,
+	in FBasePassInterpolantsVSToPS BasePassInterpolants,
+	in VTPageTableResult LightmapVTPageTableResult,
+	in float3 VolumetricLightmapBrickTextureUVs
 #endif
 )
 {
 	// This only exists with inline shading and when we are going to write out BSDFs (UpdateAllXXX functions needs to be defined)
 #if STRATA_INLINE_SHADING && STRATA_CLAMPED_BSDF_COUNT > 0
 
-#if MATERIAL_STRATA_OPAQUE_PRECOMPUTED_LIGHTING
-	IndirectLightingData = (FStrataPrecomputedLightingData)0;
-#endif 
-
 	bStrataSubsurfaceEnable = false;
 	EmissiveColor = 0.0f;
 
@@ -485,12 +486,17 @@ void PackStrataOut(
 
 					// Convert BSDF to data ready to evaluate precomputed lighting
 					#if MATERIAL_STRATA_OPAQUE_PRECOMPUTED_LIGHTING
-					IndirectLightingData.BSDFs[IndirectLightingData.BSDFCount++] = StrataPrepareBSDFForPrecomputedLighting(
+					EmissiveColor += StrataGetBSDFPrecomputedLighting(
 						Settings,
 						StrataPixelHeader,
 						BSDF,
 						V,
-						WorldBentNormal0);
+						WorldBentNormal0,
+						MaterialParameters,
+						Interpolants,
+						BasePassInterpolants,
+						LightmapVTPageTableResult,
+						VolumetricLightmapBrickTextureUVs);
 					#endif
 
 
