Change thin / two-sided material to go through the SSS wrap lighting path which allows to properly compute shadow transmission, similarly to legacy material.

diff --git a/Engine/Shaders/Private/Strata/Strata.ush b/Engine/Shaders/Private/Strata/Strata.ush
--- a/Engine/Shaders/Private/Strata/Strata.ush
+++ b/Engine/Shaders/Private/Strata/Strata.ush
@@ -1189,7 +1189,7 @@ FStrataData GetStrataSlabBSDF(
 	const bool bSSSProfilePlugged = SSSProfileID > 0.0f;
 	const bool bMFPPlugged = any(SSSMFP > 0);
 	bool bIsThin = bMFPPlugged && Thickness < STRATA_LAYER_ISTHIN_THICKNESS_THRESHOLD_CM;
-	bool bHasSSSWrap = !bIsThin && bMFPPlugged && !UseSSSDiffusion;
+	bool bHasSSSWrap = bMFPPlugged && !UseSSSDiffusion;
 	bool bHasSSSProfile = !bIsThin && bSSSProfilePlugged && SSSMFPScale > 0.f;	// When no profile is provide, we use Burley. STRATA_TODO we should clarify as burley is also an option of the profile.
 	bool bHasSSS = bMFPPlugged || bHasSSSProfile || bHasSSSWrap;
 
@@ -3354,11 +3354,10 @@ FStrataBSDF UnpackStrataBSDFIn(FStrataMaterialContainer StrataBuffer, inout FStr
 		}
 
 		const bool bIsSimpleVolume = BSDF_GETISSIMPLEVOLUME(OutBSDF);
-		const bool bIsThin = BSDF_GETISTHIN(OutBSDF);
 		const bool bHasSSS = BSDF_GETSSSTYPE(OutBSDF) != SSS_TYPE_INVALID;
-		if (bHasSSS || bIsSimpleVolume || bIsThin)
+		if (bHasSSS || bIsSimpleVolume)
 		{
-			if (BSDF_GETSSSTYPE(OutBSDF) == SSS_TYPE_DIFFUSION_PROFILE && !bIsSimpleVolume && !bIsThin)
+			if (BSDF_GETSSSTYPE(OutBSDF) == SSS_TYPE_DIFFUSION_PROFILE && !bIsSimpleVolume)
 			{
 				float4 Data = UnpackRGBA8(StrataLoadUint1(StrataBuffer, StrataAddressing));
 				SLAB_SSSPROFILEID(OutBSDF) = Data.r;
diff --git a/Engine/Shaders/Private/Strata/StrataEvaluation.ush b/Engine/Shaders/Private/Strata/StrataEvaluation.ush
--- a/Engine/Shaders/Private/Strata/StrataEvaluation.ush
+++ b/Engine/Shaders/Private/Strata/StrataEvaluation.ush
@@ -548,9 +548,10 @@ FStrataEvaluateResult StrataEvaluateBSDFCommon(FStrataBSDFContext BSDFContext, F
 
 			// Simple port of the legacy code to ensure we can reproduce the look first
 			#if STRATA_FASTPATH==0
-			const bool bIsThin = BSDF_GETISTHIN(BSDFContext.BSDF);
-			if (BSDF_GETSSSTYPE(BSDFContext.BSDF) == SSS_TYPE_WRAP || bIsThin)
+			if (BSDF_GETSSSTYPE(BSDFContext.BSDF) == SSS_TYPE_WRAP)
 			{
+				const bool bIsThin = BSDF_GETISTHIN(BSDFContext.BSDF);
+
 				const float3 SlabDiffuseColor = bIsThin ? DiffuseColor : float3(1, 1, 1);
 				const float TransmittanceNoL = 1.0f;
 				const FParticipatingMedia PM = StrataSlabCreateParticipatingMedia(SlabDiffuseColor, SLAB_SSSMFP(BSDFContext.BSDF));
@@ -1072,7 +1073,7 @@ FStrataEnvLightResult StrataEvaluateForEnvLight(FStrataBSDFContext BSDFContext,
 
 			#if STRATA_FASTPATH==0
 			BRANCH
-			if (BSDF_GETISTHIN(BSDFContext.BSDF) || BSDF_GETSSSTYPE(BSDFContext.BSDF) == SSS_TYPE_WRAP)
+			if (BSDF_GETSSSTYPE(BSDFContext.BSDF) == SSS_TYPE_WRAP)
 			{
 				FParticipatingMedia PM = StrataSlabCreateParticipatingMedia(DiffuseColor, SLAB_SSSMFP(BSDFContext.BSDF));
 
diff --git a/Engine/Shaders/Private/Strata/StrataExport.ush b/Engine/Shaders/Private/Strata/StrataExport.ush
--- a/Engine/Shaders/Private/Strata/StrataExport.ush
+++ b/Engine/Shaders/Private/Strata/StrataExport.ush
@@ -205,9 +205,9 @@ void PackStrataOut(
 						const bool bHasSSS = BSDF_GETSSSTYPE(BSDF) != SSS_TYPE_INVALID;
 						const bool bIsThin = BSDF_GETISTHIN(BSDF);
 						const bool bIsSimpleVolume = BSDF_GETISSIMPLEVOLUME(BSDF);
-						bStrataSubsurfaceEnable = bStrataSubsurfaceEnable || (bHasSSS && BSDF.Coverage > 0.0f && !bIsSimpleVolume && !bIsThin); // Should it be coverage>0 or any(LuminanceWeight>0) ?
+						bStrataSubsurfaceEnable = bStrataSubsurfaceEnable || (bHasSSS && BSDF.Coverage > 0.0f && !bIsSimpleVolume); // Should it be coverage>0 or any(LuminanceWeight>0) ?
 
-						const bool bSSSMask = !bIsThin && BSDF.bIsBottom;	// SSS and Thin can only be on the bottom layer, so ignoring bIsSimpleVolume here.
+						const bool bSSSMask = BSDF.bIsBottom;	// SSS and Thin can only be on the bottom layer, so ignoring bIsSimpleVolume here.
 						SetHasSubsurface(StrataPixelHeader, (bSSSMask && bHasSSS));
 
 						break;
@@ -419,11 +419,10 @@ void PackStrataOut(
 							}
 
 							const bool bIsSimpleVolume = BSDF_GETISSIMPLEVOLUME(BSDF);
-							const bool bIsThin = BSDF_GETISTHIN(BSDF);
 							const bool bHasSSS = BSDF_GETSSSTYPE(BSDF) != SSS_TYPE_INVALID;
-							if (bHasSSS || bIsSimpleVolume || bIsThin)
+							if (bHasSSS || bIsSimpleVolume)
 							{
-								if (!bIsThin && bHasSSS)
+								if (bHasSSS)
 								{
 									StrataSubSurfaceHeaderSetSSSType(SSSData.Header, BSDF_GETSSSTYPE(BSDF));
 									if (BSDF_GETSSSTYPE(BSDF) == SSS_TYPE_WRAP)
diff --git a/Engine/Shaders/Private/Strata/StrataLegacyConversion.ush b/Engine/Shaders/Private/Strata/StrataLegacyConversion.ush
--- a/Engine/Shaders/Private/Strata/StrataLegacyConversion.ush
+++ b/Engine/Shaders/Private/Strata/StrataLegacyConversion.ush
@@ -356,7 +356,7 @@ FStrataData StrataConvertLegacyMaterialDynamic(
 	{
 		// Set a thickness that will enabled the thin lighting model (corresponding to the legacy two-sided lighting model)
 		const float Thickness = STRATA_LAYER_ISTHIN_THICKNESS_THRESHOLD_CM - 1e-5f;
-		const float3 MFP = TransmittanceToMeanFreePath(SubSurfaceColor /*TransmittanceColor*/, Thickness * CENTIMETER_TO_METER) * METER_TO_CENTIMETER;
+		const float3 MFPInCm = TransmittanceToMeanFreePath(SubSurfaceColor /*TransmittanceColor*/, Thickness * CENTIMETER_TO_METER) * METER_TO_CENTIMETER;
 
 		Top.BaseColor = BaseColor;
 		Top.Specular = Specular;
@@ -364,9 +364,11 @@ FStrataData StrataConvertLegacyMaterialDynamic(
 		Top.Roughness = Roughness;
 		Top.Anisotropy = Anisotropy;
 		Top.SSSProfileID = 0.f;
-		Top.SSSMFP = MFP;
-		Top.SSSMFPScale = 1.f;
+		Top.SSSMFP = max(1e-05f, MFPInCm); // Ensure the MFP is not null to force the material as have SSS
+		Top.SSSMFPScale = 1.0f;
+		Top.SSSPhaseAniso = 1.0f - Opacity; // Opaque-Thick: isotropic phase function, Thin: forward phase scattering function
 		Top.Emissive = Emissive;
+		Top.UseSSSDiffusion = false;
 		Top.Thickness = Thickness;
 
 		// Take in account the legacy opacity weighting for translucent material
