Strata - Added debug for indirect irradiance to viz.

diff --git a/Engine/Shaders/Private/Strata/StrataVisualize.usf b/Engine/Shaders/Private/Strata/StrataVisualize.usf
--- a/Engine/Shaders/Private/Strata/StrataVisualize.usf
+++ b/Engine/Shaders/Private/Strata/StrataVisualize.usf
@@ -360,17 +360,20 @@ void PrintPixel(uint2 InCoord, float3 InWorldPosition, float3 V)
 		const bool bIsSimpleMaterial = IsSimpleMaterial(Header) || Header.BSDFCount == 0;
 		const bool bIsSingleMaterial = !IsSimpleMaterial(Header) && IsSingleMaterial(Header);
 
+		FStrataIrradianceAndOcclusion IrradianceAndOcclusion = StrataGetIrradianceAndAO(Header);
+
 		Print(Context, Header.BSDFCount, FontBSDFCount, 2, 0);
 		Print(Context, TEXT("BSDF"), FontBSDFCount);
 		Print(Context, TEXT(" - "), FontWhite);
 		PrintPixelType(Context, bIsSimpleMaterial, bIsSingleMaterial, FontLightRed);
 		Newline(Context);
 
-		Print(Context, TEXT(" AO    TopRoughness PreShadow ZeroShadow ContacShadow Ind.Occluder IsWater HasSSS"), FontHeaderStateName);
+		Print(Context, TEXT(" AO    IndIrr TopRoughness PreShadow ZeroShadow ContacShadow Ind.Occluder IsWater HasSSS"), FontHeaderStateName);
 		Newline(Context);
 
 		Print(Context, TEXT(" "));
-		Print(Context, StrataGetAO(Header), FontHeaderPropValue, 5, 3);									Print(Context, TEXT(" "));
+		Print(Context, IrradianceAndOcclusion.MaterialAO, FontHeaderPropValue, 5, 3);					Print(Context, TEXT(" "));
+		Print(Context, IrradianceAndOcclusion.IndirectIrradiance, FontHeaderPropValue, 6, 3);			Print(Context, TEXT(" "));
 		Print(Context, TopLayerData.Roughness, FontHeaderPropValue, 5, 3);								Print(Context, TEXT("        "));
 		PrintBool(Context, HasPrecShadowMask(Header), FontHeaderPropValue);								Print(Context, TEXT("     "));
 		PrintBool(Context, HasZeroPrecShadowMask(Header), FontHeaderPropValue);							Print(Context, TEXT("      "));
