Fix shader compile error.

diff --git a/Engine/Shaders/Private/DistanceFieldShadowingShared.ush b/Engine/Shaders/Private/DistanceFieldShadowingShared.ush
--- a/Engine/Shaders/Private/DistanceFieldShadowingShared.ush
+++ b/Engine/Shaders/Private/DistanceFieldShadowingShared.ush
@@ -137,9 +137,9 @@ float ShadowRayTraceThroughCulledObjects(
 				
 				float MaxEncodedDistance = DFAssetData.DistanceFieldToVolumeScaleBias.x + DFAssetData.DistanceFieldToVolumeScaleBias.y;
 
-				if (DFObjectData.bMostlyTwoSided)
+				if (bExpandSurface || DFObjectData.bMostlyTwoSided)
 				{
-					MaxEncodedDistance -= VolumeTwoSidedMeshDistanceBias;
+					MaxEncodedDistance -= ExpandSurfaceDistance;
 				}
 
 				// Prevent incorrect shadowing when sampling invalid bricks by limiting VolumeMaxSphereRadius to MaxEncodedDistance
@@ -155,7 +155,6 @@ float ShadowRayTraceThroughCulledObjects(
 #endif
 				float MinStepSize = 1.0f / (4 * MaxSteps);
 
-				float MinDistance = 1000000;
 				uint StepIndex = 0;
 
 				LOOP
@@ -174,8 +173,6 @@ float ShadowRayTraceThroughCulledObjects(
 						DistanceField -= ExpandSurfaceAmount;
 					}
 
-					MinDistance = min(MinDistance, DistanceField);
-
 					// Don't allow occlusion within an object's self shadow distance
 					float SelfShadowVisibility = 1 - saturate(SampleRayTime * SelfShadowScale);
 
