Fix for subsurface color affecting lighting even when subsurface scattering feature is disabled.

diff --git a/Engine/Shaders/Private/BasePassPixelShader.usf b/Engine/Shaders/Private/BasePassPixelShader.usf
--- a/Engine/Shaders/Private/BasePassPixelShader.usf
+++ b/Engine/Shaders/Private/BasePassPixelShader.usf
@@ -1071,6 +1071,10 @@ void FPixelShaderInOut_MainPS(
 	if (ShadingModel == SHADINGMODELID_SUBSURFACE || ShadingModel == SHADINGMODELID_PREINTEGRATED_SKIN || ShadingModel == SHADINGMODELID_SUBSURFACE_PROFILE || ShadingModel == SHADINGMODELID_TWOSIDED_FOLIAGE || ShadingModel == SHADINGMODELID_CLOTH || ShadingModel == SHADINGMODELID_EYE)
 	{
 		float4 SubsurfaceData = GetMaterialSubsurfaceData(PixelMaterialInputs);
+		
+		// If the subsurface feature is disabled, the subsurface color should be black
+		// This will also ensure that the subsurface profile does not contribute when subsurface scattering is disabled
+		SubsurfaceData *= View.bSubsurfacePostprocessEnabled;
 
 		if (false) // Dummy if to make the ifdef logic play nicely
 		{
diff --git a/Engine/Shaders/Private/MobileBasePassPixelShader.usf b/Engine/Shaders/Private/MobileBasePassPixelShader.usf
--- a/Engine/Shaders/Private/MobileBasePassPixelShader.usf
+++ b/Engine/Shaders/Private/MobileBasePassPixelShader.usf
@@ -289,6 +289,10 @@ void Main(
 
 #if (MATERIAL_SHADINGMODEL_SUBSURFACE || MATERIAL_SHADINGMODEL_PREINTEGRATED_SKIN || MATERIAL_SHADINGMODEL_SUBSURFACE_PROFILE) && MOBILE_SHADINGMODEL_SUPPORT
 	half4 SubsurfaceData = GetMaterialSubsurfaceData(PixelMaterialInputs);
+	
+	// If the subsurface feature is disabled, the subsurface color should be black
+	// This will also ensure that the subsurface profile does not contribute when subsurface scattering is disabled
+	SubsurfaceData *= View.bSubsurfacePostprocessEnabled;
 
 	SubsurfaceColor = SubsurfaceData.rgb * ResolvedView.DiffuseOverrideParameter.w + ResolvedView.DiffuseOverrideParameter.xyz;
 
