Strata - working stationary sky light, bent normal, specular occlusion and material AO input.

diff --git a/Engine/Shaders/Private/BasePassPixelShader.usf b/Engine/Shaders/Private/BasePassPixelShader.usf
--- a/Engine/Shaders/Private/BasePassPixelShader.usf
+++ b/Engine/Shaders/Private/BasePassPixelShader.usf
@@ -1711,8 +1711,8 @@ void FPixelShaderInOut_MainPS(
 	// const FShadingOcclusion ShadingOcclusion = ApplyBentNormal(MaterialParameters.CameraVector, MaterialParameters.WorldNormal, MaterialParameters.TangentToWorld, GetBentNormalZero(MaterialParameters), GBuffer.Roughness, MaterialAO);
 	{
 		StrataPixelHeader.IrradianceAO									= InitIrradianceAndOcclusion();
-		StrataPixelHeader.IrradianceAO.MaterialAO						= MaterialAO;			// STRATA_TODO: AOMultiBounce(Luminance(GBuffer.SpecularColor), ShadingOcclusion.SpecOcclusion).g; // FIXME: ALLOW_STATIC_LIGHTING == 0 expects this to be AO
-		StrataPixelHeader.IrradianceAO.IndirectIrradiance				= IndirectIrradiance;	// STRATA_TODO
+		StrataPixelHeader.IrradianceAO.MaterialAO						= MaterialAO;			// This is the material AO, it will be converted to multibounce AO using ShadingOcclusion.SpecOcclusion in the PackStrataOut function.
+		StrataPixelHeader.IrradianceAO.IndirectIrradiance				= IndirectIrradiance;
 		StrataPixelHeader.IrradianceAO.DiffuseIndirectSampleOcclusion	= GetDiffuseIndirectSampleOcclusion(StrataPixelHeader.SharedLocalBases, MaterialParameters.CameraVector, In.SvPosition.xy, StrataPixelHeader.IrradianceAO.MaterialAO);
 	}
 
@@ -1995,9 +1995,9 @@ void FPixelShaderInOut_MainPS(
 				GBuffer.GenericAO = float(GBuffer.DiffuseIndirectSampleOcclusion) * rcp(255) + (0.5 / 255.0);
 			#elif ALLOW_STATIC_LIGHTING
 				// No space for AO. Multiply IndirectIrradiance by AO instead of storing.
-				GBuffer.GenericAO = EncodeIndirectIrradiance(GBuffer.IndirectIrradiance * GBuffer.GBufferAO) + QuantizationBias * (1.0 / 255.0);
+				GBuffer.GenericAO = EncodeIndirectIrradiance(GBuffer.IndirectIrradiance * GBuffer.GBufferAO) + QuantizationBias * (1.0 / 255.0);	// Stationary sky light path
 			#else
-				GBuffer.GenericAO = GBuffer.GBufferAO;
+				GBuffer.GenericAO = GBuffer.GBufferAO;	// Movable sky light path
 			#endif
 
 			EncodeGBufferToMRT(Out, GBuffer, QuantizationBias);
diff --git a/Engine/Shaders/Private/ReflectionEnvironmentComposite.ush b/Engine/Shaders/Private/ReflectionEnvironmentComposite.ush
--- a/Engine/Shaders/Private/ReflectionEnvironmentComposite.ush
+++ b/Engine/Shaders/Private/ReflectionEnvironmentComposite.ush
@@ -174,7 +174,7 @@ float3 CompositeReflectionCapturesAndSkylightTWS(
 			float3 SkyLighting = GetSkyLightReflection(RayDirection, Roughness, SkyAverageBrightness);
 		#endif
 
-		// Normalize for static skylight types which mix with lightmaps
+		// Normalize for static skylight types which mix with lightmaps, material ambient occlusion as well as diffuse/specular occlusion.
 		bool bNormalize = ReflectionStruct.SkyLightParameters.z < 1 && ALLOW_STATIC_LIGHTING;
 
 		FLATTEN
diff --git a/Engine/Shaders/Private/Strata/StrataExport.ush b/Engine/Shaders/Private/Strata/StrataExport.ush
--- a/Engine/Shaders/Private/Strata/StrataExport.ush
+++ b/Engine/Shaders/Private/Strata/StrataExport.ush
@@ -46,7 +46,8 @@ float3 StrataGetBSDFPrecomputedLighting(
 	in FVertexFactoryInterpolantsVSToPS Interpolants,
 	in FBasePassInterpolantsVSToPS BasePassInterpolants,
 	in VTPageTableResult LightmapVTPageTableResult,
-	in float3 VolumetricLightmapBrickTextureUVs)
+	in float3 VolumetricLightmapBrickTextureUVs,
+	inout float2 SpecMultiBounceAO_IndirectIrradiance)
 {
 	FStrataAddressing NullStrataAddressing = (FStrataAddressing)0;
 	FStrataBSDFContext StrataBSDFContext = StrataCreateBSDFContext(StrataPixelHeader, CurrentBSDF, NullStrataAddressing, V);
@@ -75,14 +76,17 @@ float3 StrataGetBSDFPrecomputedLighting(
 	// Fetch precomputed lighting
 	float3 DiffuseLighting = 0.0f;
 	float3 SubsurfaceLighting = 0.0f;
-	float  IndirectIrrad1iance = 0.0f;
+	float  IndirectIrradiance = 0.0f;
 	const bool bEvaluateBackface = any(StrataEnvLight.DiffuseBackFaceWeight > 0.0);	// With Strata we have no other way to know if backface lighting will be needed
 	GetPrecomputedIndirectLightingAndSkyLight(MaterialParameters, Interpolants, BasePassInterpolants, LightmapVTPageTableResult, bEvaluateBackface, StrataEnvLight.DiffuseNormal, VolumetricLightmapBrickTextureUVs,
-		DiffuseLighting, SubsurfaceLighting, IndirectIrrad1iance);
+		DiffuseLighting, SubsurfaceLighting, IndirectIrradiance);
 
-	// Evaluate precomputed lighting according to material parameters
-	float3 OutLuminance = (DiffuseLighting * StrataEnvLight.DiffuseWeight + SubsurfaceLighting * StrataEnvLight.DiffuseBackFaceWeight) * LuminanceWeightFinal
-		* AOMultiBounce(StrataEnvLight.DiffuseColor, ShadingOcclusion.DiffOcclusion); // STRATA_TODO This should only be applied on the aggregated diffuse colors, not each contribution independently. But we also need to account for DiffOcclusion which looks normal dependent
+	const float3 DiffMultiBounceAO = AOMultiBounce(StrataEnvLight.DiffuseColor, ShadingOcclusion.DiffOcclusion);
+	SpecMultiBounceAO_IndirectIrradiance.x = AOMultiBounce(Luminance(StrataEnvLight.SpecularColor), ShadingOcclusion.SpecOcclusion).g;
+	SpecMultiBounceAO_IndirectIrradiance.y = IndirectIrradiance;
+
+	// Evaluate precomputed lighting according to material parameters matching environment lighting.
+	float3 OutLuminance = (DiffuseLighting * StrataEnvLight.DiffuseWeight + SubsurfaceLighting * StrataEnvLight.DiffuseBackFaceWeight) * LuminanceWeightFinal * DiffMultiBounceAO;
 
 	return OutLuminance;
 }
@@ -267,11 +271,6 @@ void PackStrataOut(
 		AndIsSimpleMaterial(StrataPixelHeader, BSDFVisibleCount == 1);
 		AndIsSingleMaterial(StrataPixelHeader, BSDFVisibleCount == 1);
 
-		// Store irradiance and AO into the state bits of the header so that they can be written later as part of the common header state bits.
-	#if STRATA_INLINE_SHADING
-		HEADER_SETIRRADIANCE_AO(StrataPixelHeader.State, StrataPackIrradianceAndOcclusion(StrataPixelHeader.IrradianceAO, 0));
-	#endif
-
 		// Now write out Strata data
 
 		///////////////////////////////////////////////////////////////////////////
@@ -307,6 +306,57 @@ void PackStrataOut(
 		#endif
 		}
 
+		// Now evaluate Specular Occlusion and Indirect Irradiance.
+		{
+			float BSDFCount = 0;
+			float2 SpecMultiBounceAO_IndirectIrradiance_Acc = 0.0f;
+			STRATA_UNROLL_N(STRATA_CLAMPED_BSDF_COUNT)
+			for (int BSDFIdx = 0; BSDFIdx < StrataTree.BSDFCount; ++BSDFIdx)
+			{
+			#define BSDF StrataTree.BSDFs[BSDFIdx]
+
+				const bool bIsVisible = StrataIsBSDFVisible(BSDF);
+
+				// Only write visible BSDF
+				BRANCH
+				if (bIsVisible)
+				{
+					// Convert BSDF to data ready to evaluate precomputed lighting
+					#if MATERIAL_STRATA_OPAQUE_PRECOMPUTED_LIGHTING
+					float2 SpecMultiBounceAO_IndirectIrradiance = 0.0f;
+					EmissiveColor += StrataGetBSDFPrecomputedLighting(
+						Settings,
+						StrataPixelHeader,
+						BSDF,
+						V,
+						WorldBentNormal0,
+						MaterialParameters,
+						Interpolants,
+						BasePassInterpolants,
+						LightmapVTPageTableResult,
+						VolumetricLightmapBrickTextureUVs,
+						SpecMultiBounceAO_IndirectIrradiance);
+					SpecMultiBounceAO_IndirectIrradiance_Acc += SpecMultiBounceAO_IndirectIrradiance;
+					#endif
+
+					// Stop writing if above budget.
+					BSDFCount++;
+				}
+
+			#undef BSDF
+			}
+
+			// Write the combined SpecMultiBounceAO for all BSDFs as a single AO value. This is a trade off to have a mix of all the interactions, and it is correct for a single BSDF.
+			float2 SpecMultiBounceAO_IndirectIrradiance = SpecMultiBounceAO_IndirectIrradiance_Acc * rcp(max(1.0f, BSDFCount));
+			StrataPixelHeader.IrradianceAO.MaterialAO			= SpecMultiBounceAO_IndirectIrradiance.x;
+			StrataPixelHeader.IrradianceAO.IndirectIrradiance	= SpecMultiBounceAO_IndirectIrradiance.y;
+		}
+		
+		// Store irradiance and AO into the state bits of the header so that they can be written later as part of the common header state bits.
+	#if STRATA_INLINE_SHADING
+		HEADER_SETIRRADIANCE_AO(StrataPixelHeader.State, StrataPackIrradianceAndOcclusion(StrataPixelHeader.IrradianceAO, 0));
+	#endif
+
 		{
 			int BSDFCount = 0;
 			STRATA_UNROLL_N(STRATA_CLAMPED_BSDF_COUNT)
@@ -483,22 +533,6 @@ void PackStrataOut(
 						}
 					}
 
-					// Convert BSDF to data ready to evaluate precomputed lighting
-					#if MATERIAL_STRATA_OPAQUE_PRECOMPUTED_LIGHTING
-					EmissiveColor += StrataGetBSDFPrecomputedLighting(
-						Settings,
-						StrataPixelHeader,
-						BSDF,
-						V,
-						WorldBentNormal0,
-						MaterialParameters,
-						Interpolants,
-						BasePassInterpolants,
-						LightmapVTPageTableResult,
-						VolumetricLightmapBrickTextureUVs);
-					#endif
-
-
 					// Stop writing if above budget.
 					BSDFCount++;
 				}
diff --git a/Engine/Shaders/Private/Strata/StrataLightingCommon.ush b/Engine/Shaders/Private/Strata/StrataLightingCommon.ush
--- a/Engine/Shaders/Private/Strata/StrataLightingCommon.ush
+++ b/Engine/Shaders/Private/Strata/StrataLightingCommon.ush
@@ -147,6 +147,8 @@
 	inout float3 DiffuseLighting,
 	inout float3 SpecularLighting)
 {
+	FStrataIrradianceAndOcclusion StrataIrradianceAndOcclusion = StrataGetIrradianceAndAO(StrataPixelHeader);
+
 	// Diffuse component
 	DiffuseLighting = 0;
 #if ENABLE_DYNAMIC_SKY_LIGHT
@@ -155,7 +157,7 @@
 	if (bProcessFrontFaceDiffuse || bProcessBackFaceDiffuse)
 	{
 		// Compute the common sky visibility factors
-		FSkyLightVisibilityData SkyVisData = GetSkyLightVisibilityData(StrataBSDFContext.N, StrataBSDFContext.N, StrataGetAO(StrataPixelHeader), ScreenAmbientOcclusion, BentNormal);
+		FSkyLightVisibilityData SkyVisData = GetSkyLightVisibilityData(StrataBSDFContext.N, StrataBSDFContext.N, StrataIrradianceAndOcclusion.MaterialAO, ScreenAmbientOcclusion, BentNormal);
 
 		if (bProcessFrontFaceDiffuse)
 		{
@@ -184,7 +186,7 @@
 	SpecularLighting = 0;
 	if (any(StrataEnvLight.SpecularWeight > 0.0f))
 	{
-		float IndirectIrradiance = 0.0;						// STRATA_TODO GBuffer.IndirectIrradiance, that will likely have to take 8 mnore bits on the header
+		float IndirectIrradiance = StrataIrradianceAndOcclusion.IndirectIrradiance;
 	#if ENABLE_SKY_LIGHT && ALLOW_STATIC_LIGHTING
 		BRANCH
 		// Add in diffuse contribution from dynamic skylights so reflection captures will have something to mix with
