Respect PRIMITIVE_SCENE_DATA_FLAG_DECAL_RECEIVER in path tracer.

diff --git a/Engine/Shaders/Private/PathTracing/PathTracingCommon.ush b/Engine/Shaders/Private/PathTracing/PathTracingCommon.ush
--- a/Engine/Shaders/Private/PathTracing/PathTracingCommon.ush
+++ b/Engine/Shaders/Private/PathTracing/PathTracingCommon.ush
@@ -11,6 +11,7 @@
 
 // These flags are analogous to RAY_TRACING_PAYLOAD_* flags. There are currently 6 bits available
 #define PATH_TRACING_PAYLOAD_OUTPUT_FLAG_FRONT_FACE                   (1 << 0) // Indicates that ray has hit the front face of a primitive. Set by closest hit shader.
+#define PATH_TRACING_PAYLOAD_OUTPUT_FLAG_DECAL_RECEIVER               (1 << 1) // Indicates that ray has hit a primitive that receives decals. Set by closest hit shader.
 
 #define PATH_TRACING_PAYLOAD_INPUT_FLAG_IGNORE_SKY_MATERIALS          (1 << 0) // Indicates that the AHS/CHS should skip processing if the material is tagged as sky
 #define PATH_TRACING_PAYLOAD_INPUT_FLAG_PATHTRACER_VISIBILITY         (1 << 1) // Indicates that this is a path tracer visibility ray (which supports material evaluation for transparent shadows and fake caustics)
@@ -140,6 +141,9 @@ struct FPathTracingPayload : FMinimalPayload
 	void SetFrontFace() { Flags |= PATH_TRACING_PAYLOAD_OUTPUT_FLAG_FRONT_FACE; }
 	bool IsFrontFace() { return (Flags & PATH_TRACING_PAYLOAD_OUTPUT_FLAG_FRONT_FACE) != 0; }
 
+	void SetDecalReceiver() { Flags |= PATH_TRACING_PAYLOAD_OUTPUT_FLAG_DECAL_RECEIVER; }
+	bool IsDecalReceiver() { return (Flags & PATH_TRACING_PAYLOAD_OUTPUT_FLAG_DECAL_RECEIVER) != 0; }
+
 
 	// Various ways to interpret CustomData (depending on ShadingModelID)
 	float GetClearCoat() { return CustomData.x; }
diff --git a/Engine/Shaders/Private/PathTracing/PathTracingCore.ush b/Engine/Shaders/Private/PathTracing/PathTracingCore.ush
--- a/Engine/Shaders/Private/PathTracing/PathTracingCore.ush
+++ b/Engine/Shaders/Private/PathTracing/PathTracingCore.ush
@@ -321,7 +321,7 @@ FPathTracingPayload TraceTransparentRay(inout FPathState PathState, bool bIsCame
 		FPathTracingPayload HitPayload = UnpackPathTracingPayloadMaterial(PackedPayload, PathState.Ray);
 
 #if PLATFORM_SUPPORTS_CALLABLE_SHADERS
-		if (EnableDecals)
+		if (EnableDecals && HitPayload.IsDecalReceiver())
 		{
 			float3 TranslatedWorldPos = HitPayload.TranslatedWorldPos;
 			float3 WorldNormal = HitPayload.WorldNormal;
diff --git a/Engine/Shaders/Private/PathTracing/PathTracingMaterialHitShader.usf b/Engine/Shaders/Private/PathTracing/PathTracingMaterialHitShader.usf
--- a/Engine/Shaders/Private/PathTracing/PathTracingMaterialHitShader.usf
+++ b/Engine/Shaders/Private/PathTracing/PathTracingMaterialHitShader.usf
@@ -96,6 +96,11 @@ RAY_TRACING_ENTRY_CLOSEST_HIT(PathTracingMaterialCHS,
 	Payload.BlendingMode = MaterialBlendingMode;
 	Payload.PrimitiveLightingChannelMask = GetPrimitive_LightingChannelMask(MaterialParameters.PrimitiveId);
 
+	if ((GetPrimitiveData(MaterialParameters.PrimitiveId).Flags & PRIMITIVE_SCENE_DATA_FLAG_DECAL_RECEIVER) != 0)
+	{
+		Payload.SetDecalReceiver();
+	}
+
 	Payload.HitT = RayTCurrent();
 	if (HitKind() == HIT_KIND_TRIANGLE_FRONT_FACE)
 	{
