Apply GetRayTracingQualitySwitch/GetPathTracingQualitySwitch/GetLightmassReplaceState() to raytracing callable and miss shaders.

diff --git a/Engine/Shaders/Private/Common.ush b/Engine/Shaders/Private/Common.ush
--- a/Engine/Shaders/Private/Common.ush
+++ b/Engine/Shaders/Private/Common.ush
@@ -1756,7 +1756,7 @@ bool GetNaniteReplaceState()
 
 bool GetRayTracingQualitySwitch()
 {
-#if (RAYHITGROUPSHADER || LUMEN_CARD_CAPTURE) && !PATH_TRACING
+#if (RAYHITGROUPSHADER || RAYMISSSHADER || RAYCALLABLESHADER || LUMEN_CARD_CAPTURE) && !PATH_TRACING
 	return true;
 #else
 	return false;
@@ -1765,7 +1765,7 @@ bool GetRayTracingQualitySwitch()
 
 bool GetPathTracingQualitySwitch()
 {
-#if RAYHITGROUPSHADER && PATH_TRACING && !SIMPLIFIED_MATERIAL_SHADER
+#if (RAYHITGROUPSHADER || RAYMISSSHADER || RAYCALLABLESHADER) && PATH_TRACING && !SIMPLIFIED_MATERIAL_SHADER
 	return true;
 #else
 	return false;
@@ -1774,7 +1774,7 @@ bool GetPathTracingQualitySwitch()
 
 bool GetLightmassReplaceState()
 {
-#if RAYHITGROUPSHADER && PATH_TRACING && SIMPLIFIED_MATERIAL_SHADER
+#if (RAYHITGROUPSHADER || RAYMISSSHADER || RAYCALLABLESHADER) && PATH_TRACING && SIMPLIFIED_MATERIAL_SHADER
 	return true;
 #else
 	return false;
