Add a missing InstanceLocalBoundsCenter and InstanceLocalBoundsExtent in a non-GPUScene path

diff --git a/Engine/Shaders/Private/SceneData.ush b/Engine/Shaders/Private/SceneData.ush
--- a/Engine/Shaders/Private/SceneData.ush
+++ b/Engine/Shaders/Private/SceneData.ush
@@ -185,6 +185,8 @@ FPrimitiveSceneData GetPrimitiveDataFromUniformBuffer()
 	PrimitiveData.ObjectBoundsZ						= Primitive.ObjectBoundsZ;
 	PrimitiveData.InstancePayloadDataOffset			= Primitive.InstancePayloadDataOffset;
 	PrimitiveData.InstancePayloadDataStride			= Primitive.InstancePayloadDataStride;
+	PrimitiveData.InstanceLocalBoundsCenter			= Primitive.InstanceLocalBoundsCenter;
+	PrimitiveData.InstanceLocalBoundsExtent			= Primitive.InstanceLocalBoundsExtent;
 	PrimitiveData.WireframeColor					= Primitive.WireframeColor;
 	PrimitiveData.LevelColor						= Primitive.LevelColor;
 	PrimitiveData.NaniteImposterIndex				= Primitive.NaniteImposterIndexAndFilterFlags & NANITE_IMPOSTER_INDEX_MASK;
@@ -302,6 +304,7 @@ FPrimitiveSceneData GetPrimitiveData(uint PrimitiveId)
 	PrimitiveData.LevelColor					= LoadPrimitivePrimitiveSceneDataElement(PrimitiveIndex, 30).xyz;
 	PrimitiveData.PersistentPrimitiveIndex		= asint(LoadPrimitivePrimitiveSceneDataElement(PrimitiveIndex, 30).w);
 	PrimitiveData.CameraDistanceCullMinMaxSquared = LoadPrimitivePrimitiveSceneDataElement(PrimitiveIndex, 31).xy;
+	// NOTE: Please make sure GetPrimitiveDataFromUniformBuffer() gets updated as well when adding new members here
 
 	uint NaniteImposterIndexAndFilterFlags		= asuint(LoadPrimitivePrimitiveSceneDataElement(PrimitiveIndex, 29).w);
 	PrimitiveData.NaniteFilterFlags				= NaniteImposterIndexAndFilterFlags >> NANITE_IMPOSTER_INDEX_NUM_BITS;
