Add small Diffuse Chan optimization (removed quater rate operation) + remove energy correction as the visual delta is very small and probably not worth at the moment.

Overall gain during lighting pass for simple tile is ~10% on PS4.

diff --git a/Engine/Shaders/Private/BRDF.ush b/Engine/Shaders/Private/BRDF.ush
--- a/Engine/Shaders/Private/BRDF.ush
+++ b/Engine/Shaders/Private/BRDF.ush
@@ -321,7 +321,7 @@ float3 Diffuse_Chan( float3 DiffuseColor, float a2, float NoV, float NoL, float
 	NoH = saturate(NoH);
 
 	// a2 = 2 / ( 1 + exp2( 18 * g )
-	float g = saturate( (1.0 / 18.0) * log2( 2 / a2 - 1 ) );
+	float g = saturate( (1.0 / 18.0) * log2( 2 * rcpFast(a2) - 1 ) );
 
 	float F0 = VoH + Pow5( 1 - VoH );
 	float FdV = 1 - 0.75 * Pow5( 1 - NoV );
@@ -331,7 +331,7 @@ float3 Diffuse_Chan( float3 DiffuseColor, float a2, float NoV, float NoL, float
 	float Fd = lerp( F0, FdV * FdL, saturate( 2.2 * g - 0.5 ) );
 
 	// Retro reflectivity contribution.
-	float Fb = ( (34.5 * g - 59 ) * g + 24.5 ) * VoH * exp2( -max( 73.2 * g - 21.2, 8.9 ) * sqrt( NoH ) );
+	float Fb = ( (34.5 * g - 59 ) * g + 24.5 ) * VoH * exp2( -max( 73.2 * g - 21.2, 8.9 ) * sqrtFast( NoH ) );
 	// It fades out when lights become area lights in order to avoid visual artefacts.
 	Fb *= RetroReflectivityWeight;
 	
diff --git a/Engine/Shaders/Private/Strata/StrataEvaluation.ush b/Engine/Shaders/Private/Strata/StrataEvaluation.ush
--- a/Engine/Shaders/Private/Strata/StrataEvaluation.ush
+++ b/Engine/Shaders/Private/Strata/StrataEvaluation.ush
@@ -281,7 +281,10 @@ FStrataEvaluateResult StrataEvaluateBSDFCommon(FStrataBSDFContext BSDFContext, F
 
 			if (STRATA_DIFFUSE_CHAN && Settings.bRoughDiffuseEnabled)
 			{
-				// If the specular layer is anisotropic, the diffuse takes the 'main' roughness value rather than the tangent/bitangent value
+				// * If the specular layer is anisotropic, the diffuse takes the 'main' roughness value rather than the tangent/bitangent value
+				// * The Chan model bakes transmittance specular directional albedo assuming F=0.04. In previous code we reapplied 
+				//   this transmittance, as the energy preservation code we remove it as well. However the visual effect was small 
+				//   and the cost was rather large (~8%). This is why we removed it in recent iteration
 				Sample.DiffusePathValue = Diffuse_Chan(DiffuseColor, Alpha2, NoV, AreaLightContext.NoL, VoH, NoH, GetAreaLightDiffuseMicroReflWeight(AreaLightContext.AreaLight));
 			}
 			else
@@ -319,25 +322,9 @@ FStrataEvaluateResult StrataEvaluateBSDFCommon(FStrataBSDFContext BSDFContext, F
 					HazeMSScale = ComputeEnergyConservation(EnergyTerms);
 				}
 			#endif
-
-				// We remove the Chan model baked transmittance specular directional albedo assuming F=0.04.
-				// Then we apply the transmittance according to the current surface setup.
-				// To make this a safe operation:
-				//	- We clamp SpecularTChan to a threshold to avoid div by 0 when the material is a metal. 
-				//	- When the material tends to be metal, DirectionalAlbedo_SpecularTransmission will naturally tend to be 0
-				if (STRATA_DIFFUSE_CHAN > 0 && Settings.bRoughDiffuseEnabled)
-				{
-					const float ChanF0 = 0.04;
-					FBxDFEnergyTerms EnergyTerms = ComputeGGXSpecEnergyTerms(SafeRoughness, BSDFContext.Context.NoV, ChanF0, 1.0f);
-					const float3 SpecularTransmitionChanF0ToMaterialF0 = DirectionalAlbedo_SpecularTransmission / max(STRATA_EPSILON, ComputeEnergyPreservation(EnergyTerms));
-					Sample.DiffusePathValue *= SpecularTransmitionChanF0ToMaterialF0;
-					Sample.IntegratedDiffuseValue *= SpecularTransmitionChanF0ToMaterialF0;
-				}
-				else
-				{
-					Sample.DiffusePathValue *= DirectionalAlbedo_SpecularTransmission;
-					Sample.IntegratedDiffuseValue *= DirectionalAlbedo_SpecularTransmission;
-				}
+			
+				Sample.DiffusePathValue *= DirectionalAlbedo_SpecularTransmission;
+				Sample.IntegratedDiffuseValue *= DirectionalAlbedo_SpecularTransmission;
 			}
 
 
