Fix for potential NaNs in cloud depth texture with half float.

diff --git a/Engine/Shaders/Private/VolumetricCloud.usf b/Engine/Shaders/Private/VolumetricCloud.usf
--- a/Engine/Shaders/Private/VolumetricCloud.usf
+++ b/Engine/Shaders/Private/VolumetricCloud.usf
@@ -474,7 +474,7 @@
 		DeviceZ = max(0.000000000001, DeviceZ);
 #endif
 		float3 DepthBufferTranslatedWorldPos = SvPositionToTranslatedWorld(float4(SvPosition.xy, DeviceZ, 1.0));
-		const float TDepthBufferKm = length(DepthBufferTranslatedWorldPos * CENTIMETER_TO_KILOMETER - RayTranslatedWorldOriginKm);
+		const float TDepthBufferKm = min(length(DepthBufferTranslatedWorldPos * CENTIMETER_TO_KILOMETER - RayTranslatedWorldOriginKm), MaxHalfFloat);
 		OutColor1.y = TDepthBufferKm;
 
 		//
