Strata - fixed AO/bent normal to also work for complex material.

diff --git a/Engine/Shaders/Private/Strata/StrataExport.ush b/Engine/Shaders/Private/Strata/StrataExport.ush
--- a/Engine/Shaders/Private/Strata/StrataExport.ush
+++ b/Engine/Shaders/Private/Strata/StrataExport.ush
@@ -271,41 +271,6 @@
 		AndIsSimpleMaterial(StrataPixelHeader, BSDFVisibleCount == 1);
 		AndIsSingleMaterial(StrataPixelHeader, BSDFVisibleCount == 1);
 
-		// Now write out Strata data
-
-		///////////////////////////////////////////////////////////////////////////
-		// 3 types of encodings
-		// * A. Layout0: Simple encoding	(use top layer normal)
-		// * B. Layout1: Single encoding	(use top layer normal)
-		// * C. Layout2: Complex encoding	(use basis)
-		const bool bHasFastEncoding = IsSimpleMaterial(StrataPixelHeader);
-		const bool bHasSingleEncoding = IsSingleMaterial(StrataPixelHeader);
-
-		// (Layout2)
-		if (!bHasFastEncoding && !bHasSingleEncoding)
-		{
-			// 1. the header (Regular/Complex encoding)
-			const uint PackedHeader = PackStrataHeader(BSDFVisibleCount, StrataPixelHeader);
-			STRATA_STORE_UINT1(PackedHeader);
-
-			// 1.1 tangent basis (Regular/Complex encoding)
-		#if STRATA_INLINE_SHADING
-			UNROLL
-			for (uint i = 0; i < StrataPixelHeader.SharedLocalBases.Count; ++i)
-			{
-				const uint BasisType = StrataGetSharedLocalBasisType(StrataPixelHeader.SharedLocalBases.Types, i);
-				if (BasisType == STRATA_BASIS_TYPE_NORMAL)
-				{
-					STRATA_STORE_UINT1(StrataPackNormal(StrataPixelHeader.SharedLocalBases.Normals[i]));
-				}
-				else // if (BasisType == STRATA_BASIS_TYPE_TANGENT)
-				{
-					STRATA_STORE_UINT1(StrataPackNormalAndTangent(StrataPixelHeader.SharedLocalBases.Normals[i], StrataPixelHeader.SharedLocalBases.Tangents[i]));
-				}
-			}
-		#endif
-		}
-
 		// Now evaluate Specular Occlusion and Indirect Irradiance.
 		{
 			float BSDFCount = 0;
@@ -357,6 +322,41 @@
 		HEADER_SETIRRADIANCE_AO(StrataPixelHeader.State, StrataPackIrradianceAndOcclusion(StrataPixelHeader.IrradianceAO, 0));
 	#endif
 
+		// Now write out Strata data
+
+		///////////////////////////////////////////////////////////////////////////
+		// 3 types of encodings
+		// * A. Layout0: Simple encoding	(use top layer normal)
+		// * B. Layout1: Single encoding	(use top layer normal)
+		// * C. Layout2: Complex encoding	(use basis)
+		const bool bHasFastEncoding = IsSimpleMaterial(StrataPixelHeader);
+		const bool bHasSingleEncoding = IsSingleMaterial(StrataPixelHeader);
+
+		// (Layout2)
+		if (!bHasFastEncoding && !bHasSingleEncoding)
 		{
+			// 1. the header (Regular/Complex encoding)
+			const uint PackedHeader = PackStrataHeader(BSDFVisibleCount, StrataPixelHeader);
+			STRATA_STORE_UINT1(PackedHeader);
+
+			// 1.1 tangent basis (Regular/Complex encoding)
+		#if STRATA_INLINE_SHADING
+			UNROLL
+			for (uint i = 0; i < StrataPixelHeader.SharedLocalBases.Count; ++i)
+			{
+				const uint BasisType = StrataGetSharedLocalBasisType(StrataPixelHeader.SharedLocalBases.Types, i);
+				if (BasisType == STRATA_BASIS_TYPE_NORMAL)
+				{
+					STRATA_STORE_UINT1(StrataPackNormal(StrataPixelHeader.SharedLocalBases.Normals[i]));
+				}
+				else // if (BasisType == STRATA_BASIS_TYPE_TANGENT)
+				{
+					STRATA_STORE_UINT1(StrataPackNormalAndTangent(StrataPixelHeader.SharedLocalBases.Normals[i], StrataPixelHeader.SharedLocalBases.Tangents[i]));
+				}
+			}
+		#endif
+		}
+
+		{
 			int BSDFCount = 0;
 			STRATA_UNROLL_N(STRATA_CLAMPED_BSDF_COUNT)
