Fixed shadowed local lights looking dull at high exposure inside the volumetric fog. This was due to LocalShadowedLightScattering not being pre exposed. It now is preexposed.
We could pre expose more texture like that in the volumetric fog system later.

diff --git a/Engine/Shaders/Private/VolumetricFog.usf b/Engine/Shaders/Private/VolumetricFog.usf
--- a/Engine/Shaders/Private/VolumetricFog.usf
+++ b/Engine/Shaders/Private/VolumetricFog.usf
@@ -386,7 +386,8 @@
 			//OutScattering.rgb += DeferredLightUniforms.Color * .0000001f;
 		}
 
-		OutScattering.rgb /= (float)NumSuperSamples;
+		// We pre-expose the buffer containing shadowed local lights luminance contribution. This helps maintaining details and color at high exposure for very bright scenes.
+		OutScattering.rgb *= View.PreExposure / (float)NumSuperSamples;
 	}
 }
 
@@ -759,8 +760,9 @@
 
 	LightScattering /= (float)NumSuperSamples;
 
-	// Shadowed point and spot lights were computed earlier
-	LightScattering += LocalShadowedLightScattering[GridCoordinate].xyz;
+	// Shadowed point and spot lights were computed earlier. 
+	// Note: this texture was pre-exposed from InjectShadowedLocalLightPS, se we revert pre-exposure.
+	LightScattering += View.OneOverPreExposure * LocalShadowedLightScattering[GridCoordinate].xyz;
 
 	float4 MaterialScatteringAndAbsorption = VBufferA[GridCoordinate];
 	float Extinction = MaterialScatteringAndAbsorption.w + Luminance(MaterialScatteringAndAbsorption.xyz);
