Path Tracer: Add new raymask for hair primitives and skip hair tracing when the ray origin is large

A significant performance drop has been observed when raytracing hair strands from a large distance. Since hair strands are generally more expensive to ray trace than ordinary triangles, this can lead to TDRs

While viewing hair from a large distance is not common, this can occur in the path tracer when enabling atmospheric scattering as rays bouncing in the upper atmosphere can bounce back towards the hair.

Probe rays (as used for SSS now always skip hair, which should bring a performance improvement to SSS rendering with hair.

At the same time, fix the ray masking setup logic to allow hair strands to not cast shadows and do not trace against far field geometry.

diff --git a/Engine/Plugins/Runtime/HairStrands/Source/HairStrandsCore/Private/GroomComponent.cpp b/Engine/Plugins/Runtime/HairStrands/Source/HairStrandsCore/Private/GroomComponent.cpp
--- a/Engine/Plugins/Runtime/HairStrands/Source/HairStrandsCore/Private/GroomComponent.cpp
+++ b/Engine/Plugins/Runtime/HairStrands/Source/HairStrandsCore/Private/GroomComponent.cpp
@@ -581,9 +581,24 @@ public:
 					RayTracingInstance.BuildInstanceMaskAndFlags(GetScene().GetFeatureLevel(), ERayTracingInstanceLayer::NearField, RayTracingMask);
 
 					// If thin shadow is requested, we ensure that regular shadow mask is not added.
+					// TODO: can this logic live in BuildInstanceMaskAndFlags instead?
 					if (RayTracingMask == RAY_TRACING_MASK_THIN_SHADOW)
 					{
-						RayTracingInstance.Mask &= ~RAY_TRACING_MASK_SHADOW;
+						if (RayTracingInstance.Mask & RAY_TRACING_MASK_SHADOW)
+						{
+							// geometry casts shadows, make sure it is in only one of the shadow casting groups, so it can be treated seperately if desired
+							RayTracingInstance.Mask &= ~RAY_TRACING_MASK_SHADOW;
+							RayTracingInstance.Mask |= RAY_TRACING_MASK_THIN_SHADOW;
+						}
+						else
+						{
+							// if the geometry does not cast shadows, remove this flag
+							RayTracingInstance.Mask &= ~RAY_TRACING_MASK_THIN_SHADOW;
+						}
+						// make sure geometry is only in the hair group
+						RayTracingInstance.Mask &= ~RAY_TRACING_MASK_OPAQUE;
+						RayTracingInstance.Mask &= ~RAY_TRACING_MASK_TRANSLUCENT;
+						RayTracingInstance.Mask |= RAY_TRACING_MASK_HAIR_STRANDS;
 					}
 
 					OutRayTracingInstances.Add(RayTracingInstance);
diff --git a/Engine/Shaders/Private/PathTracing/PathTracingCore.ush b/Engine/Shaders/Private/PathTracing/PathTracingCore.ush
--- a/Engine/Shaders/Private/PathTracing/PathTracingCore.ush
+++ b/Engine/Shaders/Private/PathTracing/PathTracingCore.ush
@@ -8,6 +8,10 @@
 #define USE_RAY_TRACING_DECAL_GRID	1
 #define USE_ENERGY_CONSERVATION 1
 
+// Ignore hair strands if the ray origin is beyond this radius. This is meant to workaround a performance cliff that can occur on some hardware.
+// In ordinary scenes, the most likely occurance of this problem is when combining hair and volumetric atmosphere as rays from several Km can travel back towards the hair and cause long delays or even TDRs
+#define PATH_TRACING_SKIP_HAIR_DISTANCE      1e6
+
 #include "../Common.ush"
 #include "../PostProcessCommon.ush"
 #include "../RectLight.ush"
@@ -109,6 +113,7 @@ void AccumulateRadiance(inout float3 TotalRadiance, float3 PathRadiance, bool bI
 FPathTracingPayload TraceTransparentRay(inout FPathState PathState, bool bIsCameraRay, bool bLastBounce, bool bIncludeEmission, uint NumVisibleLights, inout FVolumeSegment VolumeSegment)
 {
 	const uint RayFlags = bIsCameraRay && EnableCameraBackfaceCulling ? RAY_FLAG_CULL_BACK_FACING_TRIANGLES : 0;
+	const uint InstanceInclusionMask = RAY_TRACING_MASK_OPAQUE + RAY_TRACING_MASK_TRANSLUCENT + (length2(PathState.Ray.Origin) < Pow2(PATH_TRACING_SKIP_HAIR_DISTANCE) ? RAY_TRACING_MASK_HAIR_STRANDS : 0);
 	const uint MissShaderIndex = 0;
 	float2 RISRandSample = RandomSequence_GenerateSample2D(PathState.RandSequence);
 	FRISContext HitSample = InitRISContext(RISRandSample.x);
@@ -140,7 +145,7 @@ FPathTracingPayload TraceTransparentRay(inout FPathState PathState, bool bIsCame
 		TraceRay(
 			TLAS,
 			RayFlags,
-			RAY_TRACING_MASK_ALL,
+			InstanceInclusionMask,
 			RAY_TRACING_SHADER_SLOT_MATERIAL,
 			RAY_TRACING_NUM_SHADER_SLOTS,
 			MissShaderIndex,
@@ -455,7 +460,7 @@ FPathTracingPayload TraceTransparentRay(inout FPathState PathState, bool bIsCame
 float3 TraceTransparentVisibilityRay(RayDesc Ray, uint2 PixelCoord, float PathRoughness, inout RandomSequence RandSequence)
 {
 	const uint RayFlags = RAY_FLAG_ACCEPT_FIRST_HIT_AND_END_SEARCH | RAY_FLAG_SKIP_CLOSEST_HIT_SHADER;
-	const uint InstanceInclusionMask = RAY_TRACING_MASK_SHADOW | RAY_TRACING_MASK_THIN_SHADOW;
+	const uint InstanceInclusionMask = RAY_TRACING_MASK_SHADOW | (length2(Ray.Origin) < Pow2(PATH_TRACING_SKIP_HAIR_DISTANCE) ? RAY_TRACING_MASK_THIN_SHADOW : 0);
 	const uint RayContributionToHitGroupIndex = RAY_TRACING_SHADER_SLOT_SHADOW;
 	const uint MultiplierForGeometryContributionToShaderIndex = RAY_TRACING_NUM_SHADER_SLOTS;
 	const uint MissShaderIndex = 0;
@@ -514,11 +519,12 @@ FProbeResult TraceProbeRay(RayDesc Ray, uint2 PixelCoord)
 	// TODO: use smaller payload + flag to skip shading work? (only want smooth+geo normals)
 	//       however supporting blockers would require some shading ...
 	const uint RayFlags = 0;
+	const uint InstanceInclusionMask = RAY_TRACING_MASK_ALL - RAY_TRACING_MASK_THIN_SHADOW - RAY_TRACING_MASK_HAIR_STRANDS; // Ignore hair strands for SSS rays
 	const uint MissShaderIndex = 0;
 	TraceRay(
 		TLAS,
 		RayFlags,
-		RAY_TRACING_MASK_ALL,
+		InstanceInclusionMask,
 		RAY_TRACING_SHADER_SLOT_MATERIAL,
 		RAY_TRACING_NUM_SHADER_SLOTS,
 		MissShaderIndex,
diff --git a/Engine/Shaders/Shared/RayTracingDefinitions.h b/Engine/Shaders/Shared/RayTracingDefinitions.h
--- a/Engine/Shaders/Shared/RayTracingDefinitions.h
+++ b/Engine/Shaders/Shared/RayTracingDefinitions.h
@@ -21,6 +21,7 @@
 #define RAY_TRACING_MASK_SHADOW						0x08    // Whether the geometry is visible for shadow rays
 #define RAY_TRACING_MASK_SCENE_CAPTURE				0x10    // Whether the geometry is visible when rendering scene capture views
 #define RAY_TRACING_MASK_FAR_FIELD					0x20
+#define RAY_TRACING_MASK_HAIR_STRANDS               0x40    // For primary ray tracing against hair
 #define RAY_TRACING_MASK_ALL						0xFF
 
 #define RAY_TRACING_SHADER_SLOT_MATERIAL	0
