Disable diffuse lighting evaluation when diffuse albedo is black.

Diffuse albedo is not disabled when running with simple volume it relies on a lerp between diffuse and volumetric.

There was no measureable gain with this change on conductor surface.

diff --git a/Engine/Shaders/Private/Strata/StrataEvaluation.ush b/Engine/Shaders/Private/Strata/StrataEvaluation.ush
--- a/Engine/Shaders/Private/Strata/StrataEvaluation.ush
+++ b/Engine/Shaders/Private/Strata/StrataEvaluation.ush
@@ -279,7 +279,7 @@ FStrataEvaluateResult StrataEvaluateBSDFCommon(FStrataBSDFContext BSDFContext, F
 			//// Evaluate the diffuse component.
 			////
 
-			if (STRATA_DIFFUSE_CHAN && Settings.bRoughDiffuseEnabled)
+			if (STRATA_DIFFUSE_CHAN && Settings.bRoughDiffuseEnabled && any(DiffuseColor > 0))
 			{
 				// * If the specular layer is anisotropic, the diffuse takes the 'main' roughness value rather than the tangent/bitangent value
 				// * The Chan model bakes transmittance specular directional albedo assuming F=0.04. In previous code we reapplied 
