Fix Strata emissive surface with Lumen reflection.

The unlit BSDF wasn't properly handle when evaluating raytracing material and the opacity value was incorrectly packed.

diff --git a/Engine/Shaders/Private/RayTracing/RayTracingCommon.ush b/Engine/Shaders/Private/RayTracing/RayTracingCommon.ush
--- a/Engine/Shaders/Private/RayTracing/RayTracingCommon.ush
+++ b/Engine/Shaders/Private/RayTracing/RayTracingCommon.ush
@@ -561,36 +561,38 @@
 #endif
 };
 
-FPackedMaterialClosestHitPayload PackRayTracingPayload(FMaterialClosestHitPayload Input, in FRayCone RayCone)
+FPackedMaterialClosestHitPayload PackRayTracingPayload(FMaterialClosestHitPayload In, in FRayCone RayCone)
 {
-	FPackedMaterialClosestHitPayload Output = (FPackedMaterialClosestHitPayload)0;
-	Output.HitT = Input.HitT;
-	Output.SetRayCone(RayCone);
-	Output.SetFlags(Input.Flags);
-	Output.SetBlendingMode(Input.BlendingMode);
-	Output.SetPrimitiveLightingChannelMask(Input.PrimitiveLightingChannelMask);	
-	Output.SetRadiance(Input.Radiance);
-	Output.StrataData = Input.StrataData;
+	FPackedMaterialClosestHitPayload Out = (FPackedMaterialClosestHitPayload)0;
+	Out.HitT = In.HitT;
+	Out.SetRayCone(RayCone);
+	Out.SetFlags(In.Flags);
+	Out.SetBlendingMode(In.BlendingMode);
+	Out.SetPrimitiveLightingChannelMask(In.PrimitiveLightingChannelMask);
+	Out.SetRadiance(In.Radiance);
+	Out.StrataData = In.StrataData;
+	Out.SetOpacity(In.Opacity);
 
-	return Output;
+	return Out;
 }
 
-FMaterialClosestHitPayload UnpackRayTracingPayload(FPackedMaterialClosestHitPayload Input, RayDesc Ray)
+FMaterialClosestHitPayload UnpackRayTracingPayload(FPackedMaterialClosestHitPayload In, RayDesc Ray)
 {
-	FMaterialClosestHitPayload Output = (FMaterialClosestHitPayload)0;
-	Output.HitT = Input.HitT;
-	Output.RayCone = Input.GetRayCone();
-	Output.BlendingMode = Input.GetBlendingMode();
-	Output.PrimitiveLightingChannelMask = Input.GetPrimitiveLightingChannelMask();
-	Output.Flags = Input.GetFlags();
-	Output.Radiance = Input.GetRadiance();
-	Output.WorldNormal = Input.GetWorldNormal();
-	Output.ShadingModelID = Input.GetShadingModelID();
-	Output.IndirectIrradiance = Input.GetIndirectIrradiance();
-	Output.StrataData = Input.StrataData;
-	Output.TranslatedWorldPos = Ray.Origin + Output.HitT * Ray.Direction;
-	Output.RayDirection = Input.GetRayDirection();
-	return Output;
+	FMaterialClosestHitPayload Out = (FMaterialClosestHitPayload)0;
+	Out.HitT = In.HitT;
+	Out.RayCone = In.GetRayCone();
+	Out.BlendingMode = In.GetBlendingMode();
+	Out.PrimitiveLightingChannelMask = In.GetPrimitiveLightingChannelMask();
+	Out.Flags = In.GetFlags();
+	Out.Radiance = In.GetRadiance();
+	Out.WorldNormal = In.GetWorldNormal();
+	Out.ShadingModelID = In.GetShadingModelID();
+	Out.IndirectIrradiance = In.GetIndirectIrradiance();
+	Out.StrataData = In.StrataData;
+	Out.TranslatedWorldPos = Ray.Origin + Out.HitT * Ray.Direction;
+	Out.RayDirection = In.GetRayDirection();
+	Out.Opacity = In.GetOpacity();
+	return Out;
 }
 
 #else // STRATA_ENABLED
diff --git a/Engine/Shaders/Private/RayTracing/RayTracingMaterialHitShaders.usf b/Engine/Shaders/Private/RayTracing/RayTracingMaterialHitShaders.usf
--- a/Engine/Shaders/Private/RayTracing/RayTracingMaterialHitShaders.usf
+++ b/Engine/Shaders/Private/RayTracing/RayTracingMaterialHitShaders.usf
@@ -248,10 +248,23 @@ void GetMaterialPayload(
   #if MATERIAL_IS_STRATA
 	const float Dither = 0.5f;
 	float3 EmissiveLuminance = 0.0f;
+	float TotalCoverage = 1.f;
 	bool bStrataSubsurfaceEnable = false;
 	const float3 WorldBentNormal0 = GetWorldBentNormalZero(MaterialParameters);
 	FStrataData StrataData = PixelMaterialInputs.FrontMaterial;
 
+	// We detect the isolated unlit BSDF node case to make sure lots of code is optimised out, when used, there can only be one.
+	// When unlit is selected as part of dynamic shading model, the SLAB model is used instead.
+	// This allows us to avoid having the unlit node case inside the general BSDF loop and 0 BSDF to process in such pixels.
+	#define STRATA_OPTIMIZED_UNLIT MATERIAL_SHADINGMODEL_UNLIT
+
+  #if STRATA_OPTIMIZED_UNLIT
+	FStrataBSDF UnlitBSDF = PixelMaterialInputs.FrontMaterial.InlinedBSDF;
+	EmissiveLuminance = BSDF_GETEMISSIVE(UnlitBSDF);
+	#if !STRATA_OPAQUE_MATERIAL
+	TotalCoverage = UnlitBSDF.Coverage;
+	#endif
+  #else // STRATA_OPTIMIZED_UNLIT
 	// Initialise a Strata header with normal in registers
 	FStrataPixelHeader StrataPixelHeader = InitialiseStrataPixelHeader();
 	StrataPixelHeader.StrataTree = MaterialParameters.StrataTree;
@@ -286,11 +299,13 @@ void GetMaterialPayload(
 	);
 
 	Payload.StrataData.PackedTopLayerData = StrataPackTopLayerData(TopLayerData);
+  #endif // STRATA_OPTIMIZED_UNLIT
+
 	Payload.Radiance += EmissiveLuminance;
 	Payload.TranslatedWorldPos = LWCToFloat(LWCAdd(MaterialParameters.AbsoluteWorldPosition, ResolvedView.PreViewTranslation));
 
 	// Opacity
-	Payload.Opacity = 1.0f;
+	Payload.Opacity = TotalCoverage;
 	#if MATERIALBLENDING_MASKED
 	// regular masked mode - binary decision
 	if (GetMaterialMask(PixelMaterialInputs) < 0)
