Update VulkanRT HLSL code to use templated vk::RawBufferLoad

diff --git a/Engine/Shaders/Private/RayTracing/TraceRayInlineVulkan.ush b/Engine/Shaders/Private/RayTracing/TraceRayInlineVulkan.ush
--- a/Engine/Shaders/Private/RayTracing/TraceRayInlineVulkan.ush
+++ b/Engine/Shaders/Private/RayTracing/TraceRayInlineVulkan.ush
@@ -14,33 +14,6 @@
 
 #define PLATFORM_SUPPORTS_INLINE_RAY_TRACING_TRIANGLE_ATTRIBUTES 1
 
-uint2 VulkanRawBufferLoad2(uint64_t BaseAddress, uint OffsetInBytes)
-{
-	uint2 Result;
-	Result.x = vk::RawBufferLoad(BaseAddress + OffsetInBytes);
-	Result.y = vk::RawBufferLoad(BaseAddress + OffsetInBytes + 4);
-	return Result;
-}
-
-uint3 VulkanRawBufferLoad3(uint64_t BaseAddress, uint OffsetInBytes)
-{
-	uint3 Result;
-	Result.x = vk::RawBufferLoad(BaseAddress + OffsetInBytes);
-	Result.y = vk::RawBufferLoad(BaseAddress + OffsetInBytes + 4);
-	Result.z = vk::RawBufferLoad(BaseAddress + OffsetInBytes + 8);
-	return Result;
-}
-
-uint4 VulkanRawBufferLoad4(uint64_t BaseAddress, uint OffsetInBytes)
-{
-	uint4 Result;
-	Result.x = vk::RawBufferLoad(BaseAddress + OffsetInBytes);
-	Result.y = vk::RawBufferLoad(BaseAddress + OffsetInBytes + 4);
-	Result.z = vk::RawBufferLoad(BaseAddress + OffsetInBytes + 8);
-	Result.w = vk::RawBufferLoad(BaseAddress + OffsetInBytes + 12);
-	return Result;
-}
-
 uint3 VulkanLoadIndices16Bit(uint64_t SourceBufferBaseAddress, uint OffsetInBytes)
 {
 	uint3 Result;
@@ -48,7 +21,7 @@ uint3 VulkanLoadIndices16Bit(uint64_t SourceBufferBaseAddress, uint OffsetInByte
 	// ByteAddressBuffer loads must be aligned to DWORD boundary.
 	// We can load 2 DWORDs (4 SHORTs) at a time and extract 3 SHORTs (first 3 or second 3).
 	uint AlignedOffsetInBytes = OffsetInBytes & (~3);
-	const uint2 PackedIndices = VulkanRawBufferLoad2(SourceBufferBaseAddress, AlignedOffsetInBytes);
+	const uint2 PackedIndices = vk::RawBufferLoad<uint2>(SourceBufferBaseAddress + AlignedOffsetInBytes);
 
 	if (AlignedOffsetInBytes == OffsetInBytes)
 	{
@@ -70,13 +43,13 @@ uint3 VulkanLoadIndices16Bit(uint64_t SourceBufferBaseAddress, uint OffsetInByte
 
 uint3 VulkanLoadIndices32Bit(uint64_t SourceBufferBaseAddress, uint OffsetInBytes)
 {
-	return VulkanRawBufferLoad3(SourceBufferBaseAddress, OffsetInBytes);
+	return vk::RawBufferLoad<uint3>(SourceBufferBaseAddress + OffsetInBytes);
 }
 
 float3 VulkanLoadVertexPositionFloat3(uint64_t SourceBufferBaseAddress, uint Index, uint StrideInBytes)
 {
 	uint OffsetInBytes = Index * StrideInBytes;
-	return asfloat(VulkanRawBufferLoad3(SourceBufferBaseAddress, OffsetInBytes));
+	return asfloat(vk::RawBufferLoad<uint3>(SourceBufferBaseAddress + OffsetInBytes));
 }
 
 // Counterpart of FVulkanRayTracingGeometryParameters in VulkanRayTracing.cpp
