Change Strata wrap lighting to match legacy behavior precisely.

This allows to preserve the intended look for now. This will be revisited later on once the final Strata scattering evaluation will be put in place.

diff --git a/Engine/Shaders/Private/Strata/StrataEvaluation.ush b/Engine/Shaders/Private/Strata/StrataEvaluation.ush
--- a/Engine/Shaders/Private/Strata/StrataEvaluation.ush
+++ b/Engine/Shaders/Private/Strata/StrataEvaluation.ush
@@ -533,36 +533,48 @@ FStrataEvaluateResult StrataEvaluateBSDFCommon(FStrataBSDFContext BSDFContext, F
 			//// Evaluate approximated SSS (using wrap lighting). Temp code. STRATA_TODO: unify evaluation
 			////
 
-			// Simple port of the legacy code to ensure we can reproduce the look first
 			#if STRATA_FASTPATH==0
 			if (BSDF_GETSSSTYPE(BSDFContext.BSDF) == SSS_TYPE_WRAP)
 			{
+				// The wrap lighting is a non-physically based shading which intends to match legacy behavior.
+				// For doing so the look MFP & phase function anisotropy are respectively reinterpreted as 'SubSurfaceColor' and 'Opacity', per legacy shading model
 				const bool bIsThin = BSDF_GETISTHIN(BSDFContext.BSDF);
-
-				const float3 SlabDiffuseColor = bIsThin ? DiffuseColor : float3(1, 1, 1);
 				const float TransmittanceNoL = 1.0f;
+				const float3 SlabDiffuseColor = bIsThin ? DiffuseColor : float3(1, 1, 1);
 				const FParticipatingMedia PM = StrataSlabCreateParticipatingMedia(SlabDiffuseColor, SLAB_SSSMFP(BSDFContext.BSDF));
-				const float3 TransmissionThroughput = IsotropicMediumSlabTransmittance(PM, STRATA_SIMPLEVOLUME_THICKNESS_M, TransmittanceNoL);
-
-				const float VoL = dot(BSDFContext.L, -BSDFContext.V);
-				const float NoH = dot(BSDFContext.N, BSDFContext.H);
-				const float PhaseAnisotropy = clamp(SLAB_SSSPHASEANISOTROPY(BSDFContext.BSDF), -0.98f, 0.98f);
-
-				// Derive scattering distribution
-				// The 'effective' phase function is composed of:
-				// * An anisotropic part: a Henyey-Greenstein function
-				// * An   isotropic part: a lerp between an isotropic phase function and a pseudo-lambertian term to provide surface 'normal' dependence/effect
-				// 
-				// Note: The function ApproximateHG is not correctly normalized, and we should be using HenyeyGreensteinPhase instead. However the look is quite 
-				//       different and we need to handle legacy conversion if we need to change this. 
-				const float HemisphericPhaseFunction = saturate(-BSDFContext.Context.NoL);
-				const float IsotropicPhaseFunction = lerp(saturate(NoH) / (2.0f * PI), 1.f / (4.f * PI), abs(PhaseAnisotropy));
-				const float AnisotropicPhaseFunction = ApproximateHG(VoL, PhaseAnisotropy);
-				const float EffectivePhaseFunction = bIsThin ? HemisphericPhaseFunction : lerp(AnisotropicPhaseFunction, IsotropicPhaseFunction, lerp(1.f, 0.5f, abs(PhaseAnisotropy)));
-
-				Sample.TransmissionPathValue = TransmissionThroughput * EffectivePhaseFunction * DirectionalAlbedo_SpecularTransmission;
-				Sample.TransmissionPDF = 1.0f / (4.0f * PI); // STRATA_TODO this currently match the uniform sphere sampling from StrataImportanceSampleBSDF
+				const float3 SubSurfaceColor = IsotropicMediumSlabTransmittance(PM, STRATA_SIMPLEVOLUME_THICKNESS_M, TransmittanceNoL);
+				const float Opacity = max(0.f, 1.f - SLAB_SSSPHASEANISOTROPY(BSDFContext.BSDF));
+
+				float3 TransmissionThroughput = 0;
+				if (bIsThin)
+				{
+					// Legacy Foliage
+					
+					// http://blog.stevemcauley.com/2011/12/03/energy-conserving-wrapped-diffuse/
+					const float Wrap = 0.5;
+					const float WrapNoL = saturate((-dot(BSDFContext.N, BSDFContext.L) + Wrap) / Square(1 + Wrap));
+					// Scatter distribution
+					const float VoL = dot(BSDFContext.V, BSDFContext.L);
+					const float Scatter = D_GGX(0.6 * 0.6, saturate(-VoL));
+
+					TransmissionThroughput = (WrapNoL * Scatter) * SubSurfaceColor;
+				}
+				else
+				{
+					// Legacy Subsurface
+					 
+					// To get an effect when you see through the material (hard coded pow constant)
+					const float InScatter = pow(saturate(dot(BSDFContext.L, -BSDFContext.V)), 12) * lerp(3, .1f, Opacity);
+					// Wrap around lighting, /(PI*2) to be energy consistent (hack do get some view dependnt and light dependent effect)
+					// Opacity of 0 gives no normal dependent lighting, Opacity of 1 gives strong normal contribution
+					const float NormalContribution = saturate(dot(BSDFContext.N, BSDFContext.H) * Opacity + 1 - Opacity);
+					const float BackScatter = /* GBuffer.GBufferAO */ NormalContribution / (PI * 2);
+					// Lerp to never exceed 1 (energy conserving)
+					TransmissionThroughput = lerp(BackScatter, 1, InScatter) * SubSurfaceColor;
+				}
 
+				Sample.TransmissionPDF = 1.0f / (4.0f * PI); // STRATA_TODO this currently match the uniform sphere sampling from StrataImportanceSampleBSDF
+				Sample.TransmissionPathValue = TransmissionThroughput * DirectionalAlbedo_SpecularTransmission;
 				Sample.IntegratedDiffuseValue += ShadowTerms.TransmissionShadow * Sample.TransmissionPathValue * AreaLightContext.AreaLight.FalloffColor * AreaLightContext.Falloff;
 			}
 			#endif
diff --git a/Engine/Shaders/Private/Strata/StrataLegacyConversion.ush b/Engine/Shaders/Private/Strata/StrataLegacyConversion.ush
--- a/Engine/Shaders/Private/Strata/StrataLegacyConversion.ush
+++ b/Engine/Shaders/Private/Strata/StrataLegacyConversion.ush
@@ -231,7 +231,7 @@ FStrataData StrataConvertLegacyMaterialDynamic(
 		Top.Anisotropy = Anisotropy;
 		Top.SSSMFP = max(1e-05f, MFPInCm); // Ensure the MFP is not null to force the material as have SSS
 		Top.SSSMFPScale = 1.0f;
-		Top.SSSPhaseAniso = 1.0f - Opacity; // Opaque-Thick: isotropic phase function, Thin: forward phase scattering function
+		Top.SSSPhaseAniso = saturate(1.f - Opacity); // Opaque-Thick: isotropic phase function, Thin: forward phase scattering function
 		Top.Emissive = Emissive;
 		Top.UseSSSDiffusion = false;
 		Top.Thickness = ThicknessInCm;
@@ -260,7 +260,7 @@ FStrataData StrataConvertLegacyMaterialDynamic(
 		const float3 MFPInCm = float3(1.0f, 0.088964f, 0.072095f) * 2.6748f * 0.1f;
 
 		// Legacy material uses Subsurface color as transmission 'tinting', but we can represent that with a single layer. So instead we take 
-		// the max color of SubSurfaceColro & BaseColor
+		// the max color of SubSurfaceColor & BaseColor
 		Top.BaseColor = max(SubSurfaceColor, BaseColor); 
 		Top.Specular = Specular;
 		Top.Metallic = Metallic;
@@ -366,7 +366,7 @@ FStrataData StrataConvertLegacyMaterialDynamic(
 		Top.SSSProfileID = 0.f;
 		Top.SSSMFP = max(1e-05f, MFPInCm); // Ensure the MFP is not null to force the material as have SSS
 		Top.SSSMFPScale = 1.0f;
-		Top.SSSPhaseAniso = 1.0f - Opacity; // Opaque-Thick: isotropic phase function, Thin: forward phase scattering function
+		Top.SSSPhaseAniso = 1.0f - saturate(Opacity); // Opaque-Thick: isotropic phase function, Thin: forward phase scattering function
 		Top.Emissive = Emissive;
 		Top.UseSSSDiffusion = false;
 		Top.Thickness = Thickness;
