Fix for PostProcess Materials set to "OutputAlpha=False" outputting fully transparent images instead of opaque.

#rnx

diff --git a/Engine/Shaders/Private/PostProcessMaterialShaders.usf b/Engine/Shaders/Private/PostProcessMaterialShaders.usf
--- a/Engine/Shaders/Private/PostProcessMaterialShaders.usf
+++ b/Engine/Shaders/Private/PostProcessMaterialShaders.usf
@@ -198,7 +198,11 @@ void MainPS(
 	// fill out other related material parameters
 	CalcMaterialParametersPost(Parameters, PixelMaterialInputs, SvPosition, true);
 
-	float Alpha = 1.0f;
+	// In Unreal, an alpha channel value of 0.0f means that the pixel is opaque and a value of 1.0f means that there is translucency.
+	// This is the opposite of how the alpha channel is interpreted in formats like PNG. For this reason, rendered images that are exported
+	// from the engine will have their alpha channel inverted to match the PNG convention. However, for consistency inside the rendering chain,
+	// we need to stick to 0.0f as the default value.
+	float Alpha = 0.0f;
 #if STRATA_ENABLED
 
 	// Unlit forces a single BSDF
@@ -259,7 +263,11 @@ void MainPS_VideoOverlay(
 #if MATERIAL_OUTPUT_OPACITY_AS_ALPHA
 	const float Alpha = GetMaterialOpacity(PixelMaterialInputs);
 #else
-	const float Alpha = 1.0f;
+	// In Unreal, an alpha channel value of 0.0f means that the pixel is opaque and a value of 1.0f means that there is translucency.
+	// This is the opposite of how the alpha channel is interpreted in formats like PNG. For this reason, rendered images that are exported
+	// from the engine will have their alpha channel inverted to match the PNG convention. However, for consistency inside the rendering chain,
+	// we need to stick to 0.0f as the default value.
+	const float Alpha = 0.0f;
 #endif
 	OutColor = float4(GetMaterialEmissive(PixelMaterialInputs), Alpha );
 }
