Remove unecessary STRATA_OPTIMIZED_UNLIT since now it is defined globally.

diff --git a/Engine/Shaders/Private/RayTracing/RayTracingMaterialHitShaders.usf b/Engine/Shaders/Private/RayTracing/RayTracingMaterialHitShaders.usf
--- a/Engine/Shaders/Private/RayTracing/RayTracingMaterialHitShaders.usf
+++ b/Engine/Shaders/Private/RayTracing/RayTracingMaterialHitShaders.usf
@@ -253,11 +253,6 @@ void GetMaterialPayload(
 	const float3 WorldBentNormal0 = GetWorldBentNormalZero(MaterialParameters);
 	FStrataData StrataData = PixelMaterialInputs.FrontMaterial;
 
-	// We detect the isolated unlit BSDF node case to make sure lots of code is optimised out, when used, there can only be one.
-	// When unlit is selected as part of dynamic shading model, the SLAB model is used instead.
-	// This allows us to avoid having the unlit node case inside the general BSDF loop and 0 BSDF to process in such pixels.
-	#define STRATA_OPTIMIZED_UNLIT MATERIAL_SHADINGMODEL_UNLIT
-
   #if STRATA_OPTIMIZED_UNLIT
 	FStrataBSDF UnlitBSDF = PixelMaterialInputs.FrontMaterial.InlinedBSDF;
 	EmissiveLuminance = BSDF_GETEMISSIVE(UnlitBSDF);
