Path Tracer: Implement SSS support for the eye shading model

Since the eye shading model already uses the payload's custom vector for storing the caustic normal, I had to add a new field to explicitly store the subsurface radius. This only applies to the unpacked payload however.

diff --git a/Engine/Shaders/Private/PathTracing/Material/PathTracingMaterialSampling.ush b/Engine/Shaders/Private/PathTracing/Material/PathTracingMaterialSampling.ush
--- a/Engine/Shaders/Private/PathTracing/Material/PathTracingMaterialSampling.ush
+++ b/Engine/Shaders/Private/PathTracing/Material/PathTracingMaterialSampling.ush
@@ -144,12 +144,12 @@ float3 EstimateMaterialAlbedo(FPathTracingPayload Payload)
 		case SHADINGMODELID_SUBSURFACE:
 		case SHADINGMODELID_PREINTEGRATED_SKIN:
 		case SHADINGMODELID_SUBSURFACE_PROFILE:
+		case SHADINGMODELID_EYE:
 			Albedo = BrdfWeight * (Payload.DiffuseColor + Payload.SpecularColor + Payload.GetSubsurfaceColor());
 			break;
 		case SHADINGMODELID_SINGLELAYERWATER:
 		case SHADINGMODELID_DEFAULT_LIT:
 		case SHADINGMODELID_CLEAR_COAT:
-		case SHADINGMODELID_EYE:
 			// TODO: include clearcoat albedo as well
 			Albedo = BrdfWeight * (Payload.DiffuseColor + Payload.SpecularColor);
 			break;
@@ -230,13 +230,13 @@ void AccumulateAlbedoNormal(inout FPathTracingPayload Payload, float3 PathThroug
 	case SHADINGMODELID_SUBSURFACE:
 	case SHADINGMODELID_PREINTEGRATED_SKIN:
 	case SHADINGMODELID_SUBSURFACE_PROFILE:
+	case SHADINGMODELID_EYE:
 		Albedo += BrdfWeight * PathThroughput * (Payload.DiffuseColor + RoughnessWeight * Payload.SpecularColor + Payload.GetSubsurfaceColor());
 		Normal += Luminance(BrdfWeight * PathThroughput) * Payload.WorldNormal;
 		break;
 	case SHADINGMODELID_SINGLELAYERWATER:
 	case SHADINGMODELID_DEFAULT_LIT:
 	case SHADINGMODELID_CLEAR_COAT:
-	case SHADINGMODELID_EYE:
 		Albedo += BrdfWeight * PathThroughput * (Payload.DiffuseColor + RoughnessWeight * Payload.SpecularColor);
 		Normal += Luminance(BrdfWeight * PathThroughput) * Payload.WorldNormal;
 		break;
@@ -346,6 +346,24 @@ float GetAverageRoughness(FPathTracingPayload Payload)
 #endif
 }
 
+float3 DecodeSSSProfileRadius(uint ProfileId, float3 DiffuseColor, float Opacity)
+{
+	// NOTE: would be better to just precompute all this on the CPU and just have a single value to lookup
+
+	// Burley parameterization
+	float3 SurfaceAlbedo = View.SSProfilesTexture.Load(int3(BSSS_SURFACEALBEDO_OFFSET, ProfileId, 0)).xyz;
+	float3 DiffuseMeanFreePath = DecodeDiffuseMeanFreePath(View.SSProfilesTexture.Load(int3(BSSS_DMFP_OFFSET, ProfileId, 0))).xyz;
+	float WorldUnitScale = DecodeWorldUnitScale(View.SSProfilesTexture.Load(int3(SSSS_TINT_SCALE_OFFSET, ProfileId, 0)).a) * BURLEY_CM_2_MM;
+
+	// Opacity acts as a per-pixel radius multiplier
+	// NOTE: this seems backwards? Opacity=0 acts like default-lit while Opacity=1 acts like SSS?
+	// NOTE2: Confirm if this interpretation of opacity is correct ...
+	float3 SSSRadius = GetMFPFromDMFPApprox(SurfaceAlbedo, DiffuseColor, Opacity * WorldUnitScale * DiffuseMeanFreePath);
+
+	return SSSRadius * BURLEY_MM_2_CM;
+
+}
+
 
 FPathTracingPayload UnpackPathTracingPayloadMaterial(FPackedPathTracingPayload PackedInput, RayDesc Ray)
 {
@@ -364,7 +382,7 @@ FPathTracingPayload UnpackPathTracingPayloadMaterial(FPackedPathTracingPayload P
 		{
 			// TODO: is this accurate enough? hard to measure exactly since the raster path uses a very different approach
 			float SSSRadius = (1 - HitPayload.Opacity) * 10.0; // simple formula, up to 10cm radius
-			HitPayload.SetSubsurfaceRadius(SSSRadius);
+			HitPayload.SubsurfaceRadius = SSSRadius;
 
 			// Reset opacity to avoid special cases elsewhere
 			HitPayload.Opacity = 1;
@@ -396,21 +414,10 @@ FPathTracingPayload UnpackPathTracingPayloadMaterial(FPackedPathTracingPayload P
 #else
 		case SHADINGMODELID_SUBSURFACE_PROFILE:
 		{
-			uint ProfileId = uint(HitPayload.CustomData.a * 255.0 + 0.5);
+			uint ProfileId = HitPayload.GetSubsurfaceProfileId();
 			if (GetSubsurfaceProfileUseBurley(ProfileId))
 			{
-				// NOTE: would be better to just precompute all this on the CPU and just have a single value to lookup
-
-				// Burley parameterization
-				float3 SurfaceAlbedo = View.SSProfilesTexture.Load(int3(BSSS_SURFACEALBEDO_OFFSET, ProfileId, 0)).xyz;
-				float3 DiffuseMeanFreePath = DecodeDiffuseMeanFreePath(View.SSProfilesTexture.Load(int3(BSSS_DMFP_OFFSET, ProfileId, 0))).xyz;
-				float WorldUnitScale = DecodeWorldUnitScale(View.SSProfilesTexture.Load(int3(SSSS_TINT_SCALE_OFFSET, ProfileId, 0)).a) * BURLEY_CM_2_MM;
-
-				// Opacity acts as a per-pixel radius multiplier
-				// NOTE: this seems backwards? Opacity=0 acts like default-lit while Opacity=1 acts like SSS?
-				float3 SSSRadius =GetMFPFromDMFPApprox(SurfaceAlbedo, HitPayload.DiffuseColor, HitPayload.Opacity * WorldUnitScale * DiffuseMeanFreePath);
-
-				HitPayload.SetSubsurfaceRadius(SSSRadius * BURLEY_MM_2_CM);
+				HitPayload.SubsurfaceRadius = DecodeSSSProfileRadius(ProfileId, HitPayload.DiffuseColor, HitPayload.Opacity);
 
 				// burley mode doesn't have any diffuse component
 				HitPayload.SetSubsurfaceColor(HitPayload.DiffuseColor);
@@ -418,11 +425,9 @@ FPathTracingPayload UnpackPathTracingPayloadMaterial(FPackedPathTracingPayload P
 			}
 			else
 			{
-				// How to handle this case? a diffusion kernel is tabulated into the texture but I don't see a way to extract the original
-				// radius/color controls
-				HitPayload.SetSubsurfaceRadius(0);
+				// NOTE: Non-burley scattering is deprecated going forward, so not worth supporting this mode
+				HitPayload.SubsurfaceRadius = 0.0;
 				HitPayload.SetSubsurfaceColor(0);
-				// TODO: non-burley profile allows a blend of diffuse and SSS which should be accounted for here
 			}
 			// Decode dual roughness info
 			float Lobe0Roughness = 0;
@@ -432,11 +437,30 @@ FPathTracingPayload UnpackPathTracingPayloadMaterial(FPackedPathTracingPayload P
 			GetSubsurfaceProfileDualSpecular(ProfileId, HitPayload.Roughness, HitPayload.Opacity, Lobe0Roughness, Lobe1Roughness, LobeMix);
 
 			HitPayload.SetDualRoughnessSpecular(float3(Lobe0Roughness, Lobe1Roughness, LobeMix));
-			
 			// Reset opacity to avoid special cases elsewhere
 			HitPayload.Opacity = 1;
 			break;
 		}
+		case SHADINGMODELID_EYE:
+		{
+			uint ProfileId = HitPayload.GetSubsurfaceProfileId();
+			if (GetSubsurfaceProfileUseBurley(ProfileId))
+			{
+				HitPayload.SubsurfaceRadius = DecodeSSSProfileRadius(ProfileId, HitPayload.DiffuseColor, 1.0);
+
+				// eye still needs a diffuse component for the iris center
+				float IrisMask = HitPayload.Anisotropy;
+				HitPayload.SetSubsurfaceColor((1 - IrisMask) * HitPayload.DiffuseColor);
+				HitPayload.DiffuseColor *= IrisMask;
+			}
+			else
+			{
+				// NOTE: Non-burley scattering is deprecated going forward, so not worth supporting this mode
+				HitPayload.SubsurfaceRadius = 0.0;
+				HitPayload.SetSubsurfaceColor(0);
+			}
+			break;
+		}
 #endif
 	}
 
diff --git a/Engine/Shaders/Private/PathTracing/PathTracingCommon.ush b/Engine/Shaders/Private/PathTracing/PathTracingCommon.ush
--- a/Engine/Shaders/Private/PathTracing/PathTracingCommon.ush
+++ b/Engine/Shaders/Private/PathTracing/PathTracingCommon.ush
@@ -135,6 +135,7 @@ struct FPathTracingPayload : FMinimalPayload
 	float4 CustomData;
 	float3 CustomVector;
 	float3 WorldTangent;
+	float3 SubsurfaceRadius;
 	float  Anisotropy;
 	float3 TranslatedWorldPos;
 
@@ -151,12 +152,11 @@ struct FPathTracingPayload : FMinimalPayload
 	void SetClearCoatRoughness(float ClearCoatRoughness) { CustomData.y = ClearCoatRoughness; }
 	float3 GetClearCoatBottomNormal() { return CustomVector; }
 
+	uint GetSubsurfaceProfileId() { return uint(CustomData.a * 255.0 + 0.5); }
+
 	float3 GetSubsurfaceColor() { return CustomData.xyz; }
 	void SetSubsurfaceColor(float3 SubsurfaceColor) { CustomData.xyz = SubsurfaceColor; }
 	
-	float3 GetSubsurfaceRadius() { return CustomVector.xyz; }
-	void SetSubsurfaceRadius(float3 SubsurfaceRadius) { CustomVector.xyz = SubsurfaceRadius; }
-
 	// Steal an unused float3 (the only material with dual specular does not support anisotropy)
 	// Data is encoded as (x: roughness 0, y: roughness 1, z: roughness mix)
 	float3 GetDualRoughnessSpecular() { return WorldTangent; }
@@ -185,7 +185,8 @@ struct FPathTracingPayload : FMinimalPayload
 	{
 		return ShadingModelID == SHADINGMODELID_SUBSURFACE ||
 			   ShadingModelID == SHADINGMODELID_PREINTEGRATED_SKIN ||
-			   ShadingModelID == SHADINGMODELID_SUBSURFACE_PROFILE;
+			   ShadingModelID == SHADINGMODELID_SUBSURFACE_PROFILE ||
+			   ShadingModelID == SHADINGMODELID_EYE;
 	}
 };
 
diff --git a/Engine/Shaders/Private/PathTracing/PathTracingCore.ush b/Engine/Shaders/Private/PathTracing/PathTracingCore.ush
--- a/Engine/Shaders/Private/PathTracing/PathTracingCore.ush
+++ b/Engine/Shaders/Private/PathTracing/PathTracingCore.ush
@@ -572,7 +572,7 @@ bool ProcessSubsurfaceRandomWalk(inout FPathTracingPayload Payload, inout float3
 	}
 
 	float3 SSSColor = Payload.GetSubsurfaceColor();
-	float3 SSSRadius = Payload.GetSubsurfaceRadius();
+	float3 SSSRadius = Payload.SubsurfaceRadius;
 
 	if (SimplifySSS || all(SSSRadius == 0) || all(SSSColor == 0) || MaxSSSBounces == 0)
 	{
