Fix NaN value with Strata simple volumetric caused be low albedo value & dither.

This ensures that the albedo & F0 are always in 0..1

diff --git a/Engine/Shaders/Private/Strata/Strata.ush b/Engine/Shaders/Private/Strata/Strata.ush
--- a/Engine/Shaders/Private/Strata/Strata.ush
+++ b/Engine/Shaders/Private/Strata/Strata.ush
@@ -2337,11 +2337,11 @@ void StrataSanitizeBSDF(inout FStrataBSDF BSDF)
 // This is used to respect artistic desire, e.g. shutdown diffuse or specular.
 float3 DitherIfNonBlack(float3 Input, float3 Dither)
 {
+	float3 NewInput = Input + Dither;
 	return float3(
-		Input.x == 0.0 ? Input.x : Input.x + Dither.x,
-		Input.y == 0.0 ? Input.y : Input.y + Dither.y,
-		Input.z == 0.0 ? Input.z : Input.z + Dither.z
-	);
+		NewInput.x >= 0.0 && NewInput.x <= 1.0 ? NewInput.x : Input.x,
+		NewInput.y >= 0.0 && NewInput.y <= 1.0 ? NewInput.y : Input.y,
+		NewInput.z >= 0.0 && NewInput.z <= 1.0 ? NewInput.z : Input.z);
 }
 
 void StrataDitherBSDF(inout FStrataBSDF BSDF, in float Dither)
