Fix the conversion between LinearRGB and Lab space.

diff --git a/Engine/Shaders/Private/ColorSpace.ush b/Engine/Shaders/Private/ColorSpace.ush
--- a/Engine/Shaders/Private/ColorSpace.ush
+++ b/Engine/Shaders/Private/ColorSpace.ush
@@ -65,9 +65,15 @@ float3 XYZ_2_LinearRGB( float3 XYZ )
 // CIE LAB
 // https://en.wikipedia.org/wiki/Lab_color_space
 //------------------------------------------------------------------
+// Assume XYZ is in nominal range to reduce the operation on a constant coefficient 100.
+// E.g., whilte RGB (1, 1, 1) maps to XYZ (0.95047, 1.00, 1.08883)
+// for a reference white of D65. 
+// Use XYZ_WHITE_REF_DXX_NOMINAL instead for LinearRGB_2_LAB() and LAB_2_LinearRGB()
 
 static const float3 XYZ_WHITE_REF_D65 = float3(95.047, 100.0, 108.883);
 static const float3 XYZ_WHITE_REF_D50 = float3(96.6797, 100.0, 82.5188);
+static const float3 XYZ_WHITE_REF_D65_NOMINAL = float3(0.95047, 1.00, 1.08883);
+static const float3 XYZ_WHITE_REF_D50_NOMINAL = float3(0.966797, 1.00, 0.825188);
 static const float XYZ_2_LAB_DELTA_SQUARED = 0.04280618311; // (6/29)^3
 static const float XYZ_2_LAB_DELTA_CUBED = 0.00885645167; // (6/29)^3
 
diff --git a/Engine/Shaders/Private/Strata/StrataLegacyConversion.ush b/Engine/Shaders/Private/Strata/StrataLegacyConversion.ush
--- a/Engine/Shaders/Private/Strata/StrataLegacyConversion.ush
+++ b/Engine/Shaders/Private/Strata/StrataLegacyConversion.ush
@@ -452,10 +452,10 @@ FStrataData StrataConvertLegacyMaterialDynamic(
 		// Because we are using a single slab for simulating a 'coverage mix' between two BSDFs (a opaque metallic/conductor and a non-scattering translucent medium), 
 		// we need to do some approximation. In order to reduce avoid hue shift due to difference between 'solid' base-color/diffuse-albedo and 'transparent' transmission 
 		// color, we lerp BaseColor with transmittance color in LAB space so that the interpolated color minized its 'distance' 
-		const float3 BaseColor_LAB			= LinearRGB_2_LAB(BaseColor, XYZ_WHITE_REF_D65);
-		const float3 TransmittanceColor_LAB = LinearRGB_2_LAB(TransmittanceColor, XYZ_WHITE_REF_D65);
+		const float3 BaseColor_LAB			= LinearRGB_2_LAB(BaseColor, XYZ_WHITE_REF_D65_NOMINAL);
+		const float3 TransmittanceColor_LAB = LinearRGB_2_LAB(TransmittanceColor, XYZ_WHITE_REF_D65_NOMINAL);
 		const float3 Color_LAB			    = lerp(TransmittanceColor_LAB, BaseColor_LAB, Opacity);
-		const float3 Color					= LAB_2_LinearRGB(Color_LAB, XYZ_WHITE_REF_D65);
+		const float3 Color					= LAB_2_LinearRGB(Color_LAB, XYZ_WHITE_REF_D65_NOMINAL);
 
 		// Lerp the transmittance color rather than MFP as it is more perceptually linear
 		const float3 MFP = TransmittanceToMeanFreePath(lerp(Color, Zeros, Opacity), Thickness * CENTIMETER_TO_METER) * METER_TO_CENTIMETER;
