Strata - update material mode to be simple, single, complex, SLWater, to be able to share bit with other incoming modes.

diff --git a/Engine/Shaders/Private/BasePassPixelShader.usf b/Engine/Shaders/Private/BasePassPixelShader.usf
--- a/Engine/Shaders/Private/BasePassPixelShader.usf
+++ b/Engine/Shaders/Private/BasePassPixelShader.usf
@@ -1106,7 +1106,7 @@ void FPixelShaderInOut_MainPS(
 #endif
 
 #if STRATA_ENABLED && STRATA_INLINE_SINGLELAYERWATER
-	SetIsSingleLayerWater(StrataPixelHeader, true);
+	SetMaterialMode(StrataPixelHeader, HEADER_MATERIALMODE_SLWATER);
 
 	// Override GBuffer data with Strata SLW water BSDF to run forward shadfing code.
 	// STRATA_TODO: run the shading through a strata path? (by adding a special BSDF?)
diff --git a/Engine/Shaders/Private/Strata/Strata.ush b/Engine/Shaders/Private/Strata/Strata.ush
--- a/Engine/Shaders/Private/Strata/Strata.ush
+++ b/Engine/Shaders/Private/Strata/Strata.ush
@@ -196,13 +196,11 @@ struct FStrataBSDF
 // Bit count for each component stored in a Strata header
 // 
 // State bit masks
-#define HEADER_MASK_ISSIMPLEMATERIAL			(1u << 0)
-#define HEADER_MASK_ISSINGLEMATERIAL			(1u << 1)
-#define HEADER_MASK_HASPRECSHADOWMASK			(1u << 2)
-#define HEADER_MASK_ZEROPRECSHADOWMASK			(1u << 3)
-#define HEADER_MASK_CASTCONTACTSHADOW			(1u << 4)
-#define HEADER_MASK_HASDYNINDIRECTSHADOWCASTER	(1u << 5)
-#define HEADER_MASK_ISSINGLELAYERWATER			(1u << 6)	// Used to avoid loading BSDF or Categorisation data when generating tiles containing SLW material. Avoid running another categorisation.
+#define HEADER_MASK_MATERIALMODE				(0x7)
+#define HEADER_MASK_HASPRECSHADOWMASK			(1u << 3)
+#define HEADER_MASK_ZEROPRECSHADOWMASK			(1u << 4)
+#define HEADER_MASK_CASTCONTACTSHADOW			(1u << 5)
+#define HEADER_MASK_HASDYNINDIRECTSHADOWCASTER	(1u << 6)
 #define HEADER_MASK_HASSUBSURFACE				(1u << 7)
 #define HEADER_MASK_BIT_COUNT					8
 
@@ -263,6 +261,13 @@ struct FStrataBSDF
 #define HEADER_GETSHAREDLOCALBASESTYPE(X)			READ_BITS( X, HEADER_BIT_COUNT_SHAREDLOCALBASES_TYPE, HEADER_BIT_OFFSET_SHAREDLOCALBASES_TYPE)
 #define HEADER_SETSHAREDLOCALBASESTYPE(X, NewValue)	WRITE_BITS(X, HEADER_BIT_COUNT_SHAREDLOCALBASES_TYPE, HEADER_BIT_OFFSET_SHAREDLOCALBASES_TYPE, NewValue)
 
+#define HEADER_MATERIALMODE_SIMPLE				0
+#define HEADER_MATERIALMODE_SINGLE				1
+#define HEADER_MATERIALMODE_COMPLEX				2
+#define HEADER_MATERIALMODE_SLWATER				3
+// Up to 8 material modes: 4 remains.
+// STRATA_TODO: I believe we could merge ShadingModels & MaterialModes together with slab simple, slab single, slab complex, water, hair, etc.
+
 ///////////////////////////////////////////////////////////////////////////////
 // BSDF State
  
@@ -1646,11 +1651,14 @@ struct FStrataPixelHeader
 	uint State;
 };
 
-bool IsSimpleMaterial(const FStrataPixelHeader In)											{ return (In.State & HEADER_MASK_ISSIMPLEMATERIAL) != 0 || STRATA_FASTPATH; }
-bool IsSingleMaterial(const FStrataPixelHeader In)											{ return (In.State & HEADER_MASK_ISSINGLEMATERIAL) != 0 || STRATA_SINGLEPATH; }
+void SetMaterialMode(inout FStrataPixelHeader Out, int MateralMode)							{ Out.State = (Out.State & ~HEADER_MASK_MATERIALMODE) | (MateralMode & HEADER_MASK_MATERIALMODE); }
+bool IsSimpleMaterial(const FStrataPixelHeader In)											{ return (In.State & HEADER_MASK_MATERIALMODE) == HEADER_MATERIALMODE_SIMPLE || STRATA_FASTPATH; }
+bool IsSingleMaterial(const FStrataPixelHeader In)											{ return (In.State & HEADER_MASK_MATERIALMODE) == HEADER_MATERIALMODE_SINGLE || STRATA_SINGLEPATH; }
+bool IsComplexMaterial(const FStrataPixelHeader In)											{ return (In.State & HEADER_MASK_MATERIALMODE) == HEADER_MATERIALMODE_COMPLEX && !STRATA_FASTPATH && !STRATA_SINGLEPATH; }
+bool IsSingleLayerWater(const FStrataPixelHeader In)										{ return (In.State & HEADER_MASK_MATERIALMODE) == HEADER_MATERIALMODE_SLWATER; }
+
 bool HasPrecShadowMask(const FStrataPixelHeader In)											{ return (In.State & HEADER_MASK_HASPRECSHADOWMASK) != 0; }
 bool HasZeroPrecShadowMask(const FStrataPixelHeader In)										{ return (In.State & HEADER_MASK_ZEROPRECSHADOWMASK) != 0; }
-bool IsSingleLayerWater(const FStrataPixelHeader In)										{ return (In.State & HEADER_MASK_ISSINGLELAYERWATER) != 0; }
 bool DoesCastContactShadow(const FStrataPixelHeader In)										{ return (In.State & HEADER_MASK_CASTCONTACTSHADOW) != 0; }
 uint GetShadingModels(const FStrataPixelHeader In)											{ return (In.State & HEADER_MASK_SHADINGMODELS)>>HEADER_BIT_OFFSET_SHADINGMODELS; }
 bool IsStrataMaterial(const FStrataPixelHeader In)											{ return (In.State & HEADER_MASK_SHADINGMODELS) > 0 || IsSimpleMaterial(In); }
@@ -1659,18 +1667,12 @@ bool HasDynamicIndirectShadowCasterRepresentation(const FStrataPixelHeader In)
 
 bool StrataHasShadingModel(const FStrataPixelHeader In, uint ShadingModel)					{ return StrataHasShadingModel(GetShadingModels(In), ShadingModel); }
 
-void SetIsSimpleMaterial(inout FStrataPixelHeader Out, bool bIn)							{ Out.State |= (bIn ? HEADER_MASK_ISSIMPLEMATERIAL : 0u); }
-void SetIsSingleMaterial(inout FStrataPixelHeader Out, bool bIn)							{ Out.State |= (bIn ? HEADER_MASK_ISSINGLEMATERIAL : 0u); }
 void SetHasPrecShadowMask(inout FStrataPixelHeader Out, bool bIn)							{ Out.State |= (bIn ? HEADER_MASK_HASPRECSHADOWMASK : 0u); }
 void SetZeroPrecShadowMask(inout FStrataPixelHeader Out, bool bIn)							{ Out.State |= (bIn ? HEADER_MASK_ZEROPRECSHADOWMASK : 0u); }
-void SetIsSingleLayerWater(inout FStrataPixelHeader Out, bool bIn)							{ Out.State |= (bIn ? HEADER_MASK_ISSINGLELAYERWATER : 0u); }
 void SetCastContactShadow (inout FStrataPixelHeader Out, bool bIn)							{ Out.State |= (bIn ? HEADER_MASK_CASTCONTACTSHADOW : 0u); }
 void SetDynamicIndirectShadowCasterRepresentation(inout FStrataPixelHeader Out, bool bIn)	{ Out.State |= (bIn ? HEADER_MASK_HASDYNINDIRECTSHADOWCASTER : 0u); }
 void SetHasSubsurface(inout FStrataPixelHeader Out, bool bIn)								{ Out.State |= (bIn ? HEADER_MASK_HASSUBSURFACE : 0u); }
-
 void AddShadingModels(inout FStrataPixelHeader Out, uint bIn)								{ Out.State |= ((bIn<<HEADER_BIT_OFFSET_SHADINGMODELS) & HEADER_MASK_SHADINGMODELS); }
-void AndIsSimpleMaterial(inout FStrataPixelHeader Out, bool bIn)							{ Out.State = (Out.State & ~HEADER_MASK_ISSIMPLEMATERIAL) | ((Out.State & HEADER_MASK_ISSIMPLEMATERIAL) & (bIn ? HEADER_MASK_ISSIMPLEMATERIAL : 0u)); }
-void AndIsSingleMaterial(inout FStrataPixelHeader Out, bool bIn)							{ Out.State = (Out.State & ~HEADER_MASK_ISSINGLEMATERIAL) | ((Out.State & HEADER_MASK_ISSINGLEMATERIAL) & (bIn ? HEADER_MASK_ISSINGLEMATERIAL : 0u)); }
 
 FStrataPixelHeader InitialiseStrataPixelHeader()
 {
@@ -1687,6 +1689,7 @@ FStrataPixelHeader InitialiseStrataPixelHeader()
 	Out.PackedTopLayerData = 0;
 #endif
 	Out.State = 0;
+	SetMaterialMode(Out, HEADER_MATERIALMODE_SIMPLE); // Should be a NOP since SIMPLE=0. This is for non strata material or sky pixels to be considered as simple (so that (strata simple pixels) + (non-strata pixels) can build a simple tile).
 	return Out;
 }
 
@@ -3159,8 +3162,8 @@ FStrataPixelHeader UnpackStrataHeaderIn(FStrataMaterialContainer StrataBuffer, i
 	Out.State		= HEADER_GETCOMMONSTATES(PackedHeader);	// Unpack state and AO data at the same time 
 	Out.BSDFCount   = HEADER_GETBSDFCOUNT(PackedHeader);
 
-	const bool bIsSimpleMaterial = (PackedHeader & HEADER_MASK_ISSIMPLEMATERIAL) > 0;
-	const bool bIsSingleMaterial = (PackedHeader & HEADER_MASK_ISSINGLEMATERIAL) > 0;
+	const bool bIsSimpleMaterial = IsSimpleMaterial(Out);
+	const bool bIsSingleMaterial = IsSingleMaterial(Out);
 
 	// Use simple or single layout for decoding data
 	BRANCH
diff --git a/Engine/Shaders/Private/Strata/StrataExport.ush b/Engine/Shaders/Private/Strata/StrataExport.ush
--- a/Engine/Shaders/Private/Strata/StrataExport.ush
+++ b/Engine/Shaders/Private/Strata/StrataExport.ush
@@ -135,9 +135,6 @@ void PackStrataOut(
 	TopLayerData = (FStrataTopLayerData)0;
 	OpaqueRoughRefractionData = (FStrataOpaqueRoughRefractionData)0;
 
-	SetIsSimpleMaterial(StrataPixelHeader, true); // non strata material or sky pixels are considered as simple (so that (strata simple pixels) + (non-strata pixels) can build a simple tile).
-	SetIsSingleMaterial(StrataPixelHeader, true); // non strata material or sky pixels are considered as simple (so that (strata simple pixels) + (non-strata pixels) can build a simple tile).
-
 	uint BSDFVisibleCount = 0;
 	
 	FStrataTree StrataTree = StrataPixelHeader.StrataTree;
@@ -165,8 +162,11 @@ void PackStrataOut(
 			const float StandardDeviationCmForThickness = StandardDeviationCm * RootOperator.OpaqueRoughRefractThicknessCm;
 			OpaqueRoughRefractionData.VarianceCm = StandardDeviationCmForThickness * StandardDeviationCmForThickness;
 		}
-		
-		int SimpleMaterial_BSDFIndex = 0;
+
+		bool bIsSimpleMaterial = true;
+		bool bIsSingleMaterial = true;
+		bool bIsOnlySlab = true;
+		int OneBSDFMaterial_Index = 0; // Use for simple and single material modes. Represent the last visible BSDF index.
 		float TopLayerTotalWeight = 0.0f;
 		
 		{
@@ -185,8 +185,6 @@ void PackStrataOut(
 					// Sanitize happening after will clamp values that could have gone above and below valid range.
 					StrataDitherBSDF(BSDF, Dither);
 
-					bool bIsSimpleMaterial = true;
-					bool bIsSingleMaterial = true;
 					switch (BSDF_GETTYPE(BSDF))
 					{
 					case STRATA_BSDF_TYPE_SLAB:
@@ -197,11 +195,11 @@ void PackStrataOut(
 						const bool LuminanceWeightEqualOne = all(BSDF.LuminanceWeightV == 1.0f);
 
 						// Update simple material compatibility
-						bIsSimpleMaterial = IsStrataSlabFastPathCompatible(BSDF) && LuminanceWeightEqualOne;
+						bIsSimpleMaterial &= IsStrataSlabFastPathCompatible(BSDF) && LuminanceWeightEqualOne;
 
 						// Update single material compatibility. For now, single materials don't support anisotropy, 
 						// they use the complex pass, as the toplayer data does not contain the full frame basis, only the top normal.
-						bIsSingleMaterial = BSDF_GETHASANISOTROPY(BSDF) == 0 && LuminanceWeightEqualOne;
+						bIsSingleMaterial &= BSDF_GETHASANISOTROPY(BSDF) == 0 && LuminanceWeightEqualOne;
 
 						TopLayerTotalWeight += BSDF.TopLayerDataWeight;
 						TopLayerData.Roughness += BSDF.TopLayerDataWeight * SLAB_ROUGHNESS(BSDF);
@@ -219,6 +217,7 @@ void PackStrataOut(
 
 					case STRATA_BSDF_TYPE_HAIR:
 					{
+						bIsOnlySlab = false;
 						bIsSimpleMaterial = false;
 						bIsSingleMaterial = false;
 
@@ -230,6 +229,7 @@ void PackStrataOut(
 
 					case STRATA_BSDF_TYPE_SINGLELAYERWATER:
 					{
+						bIsOnlySlab = false;
 						bIsSimpleMaterial = false;
 						bIsSingleMaterial = false;
 						EmissiveColor += BSDF_GETEMISSIVE(BSDF) * BSDF.LuminanceWeightV;
@@ -248,10 +248,8 @@ void PackStrataOut(
 					TopLayerData.WorldNormal += N * BSDF.TopLayerDataWeight;
 
 					AddShadingModels(StrataPixelHeader, StrataShadingModelBit(BSDF_GETTYPE(BSDF)));
-					AndIsSimpleMaterial(StrataPixelHeader, bIsSimpleMaterial);
-					AndIsSingleMaterial(StrataPixelHeader, bIsSingleMaterial);
 
-					SimpleMaterial_BSDFIndex = BSDFIdx;
+					OneBSDFMaterial_Index = BSDFIdx;
 
 					// Notify that the BSDF is at the top for SSR to only affect reflection there and not on the lower layers
 					BSDF_SETISTOPLAYER(BSDF, BSDF.bIsTop ? 1 : 0);
@@ -268,8 +266,12 @@ void PackStrataOut(
 		TopLayerData.BaseColor	= TopLayerTotalWeight > 0.0f	? TopLayerData.BaseColor / TopLayerTotalWeight	: 0.0f;
 
 		// Set storage layout as either: fast(0), single(1), or complex(2)
-		AndIsSimpleMaterial(StrataPixelHeader, BSDFVisibleCount == 1);
-		AndIsSingleMaterial(StrataPixelHeader, BSDFVisibleCount == 1);
+		if (bIsOnlySlab)
+		{
+			bIsSimpleMaterial &= BSDFVisibleCount == 1;
+			bIsSingleMaterial &= BSDFVisibleCount == 1;
+			SetMaterialMode(StrataPixelHeader, bIsSimpleMaterial ? HEADER_MATERIALMODE_SIMPLE : bIsSingleMaterial ? HEADER_MATERIALMODE_SINGLE : HEADER_MATERIALMODE_COMPLEX);
+		}
 
 		// Now evaluate Specular Occlusion and Indirect Irradiance.
 		{
@@ -377,7 +379,7 @@ void PackStrataOut(
 					BSDF_SETHASTRANSABOVE(BSDF, bTransmittanceAboveAlongNRequired);
 
 					// A. Layout0 - Simple encoding (aka. fast-path): store header & data
-					const bool bFastEncodedBSDF = bHasFastEncoding && BSDFIdx == SimpleMaterial_BSDFIndex;
+					const bool bFastEncodedBSDF = bHasFastEncoding && BSDFIdx == OneBSDFMaterial_Index;
 					if (bFastEncodedBSDF)
 					{
 						// Now we pack diffuse and F0 in a special way: both encoded as R7G7B6 with gamma 2.0. The low bit cound will be hidden by dithering and TAA.
@@ -410,7 +412,7 @@ void PackStrataOut(
 						BSDF.State = 0;
 					}
 					// B. Layout1 - Simple encoding: single BSDF, whose header & BSDF state are merged
-					else if (bHasSingleEncoding && BSDFIdx == SimpleMaterial_BSDFIndex)
+					else if (bHasSingleEncoding && BSDFIdx == OneBSDFMaterial_Index)
 					{
 						// Data0 (Header_State|Header_AO|Header_BSDFTypes|BSDF_State)
 						{
diff --git a/Engine/Shaders/Private/Strata/StrataVisualize.usf b/Engine/Shaders/Private/Strata/StrataVisualize.usf
--- a/Engine/Shaders/Private/Strata/StrataVisualize.usf
+++ b/Engine/Shaders/Private/Strata/StrataVisualize.usf
@@ -111,11 +111,13 @@ void PrintBool(inout FShaderPrintContext Ctx, bool In, FFontColor InColor)
 	}
 }
 
-void PrintPixelType(inout FShaderPrintContext Ctx, bool bIsSimpleMaterial, bool bIsSingleMaterial, FFontColor InColor)
+void PrintPixelType(inout FShaderPrintContext Ctx, in FStrataPixelHeader Header, FFontColor InColor)
 {
-		 if (bIsSimpleMaterial)	Print(Ctx, TEXT("Simple  "), InColor);
-	else if (bIsSingleMaterial)	Print(Ctx, TEXT("Single  "), InColor);
-	else						Print(Ctx, TEXT("Complex "), InColor);
+	if (IsSimpleMaterial(Header))			Print(Ctx, TEXT("Simple  "), InColor);
+	else if (IsSingleMaterial(Header))		Print(Ctx, TEXT("Single  "), InColor);
+	else if (IsComplexMaterial(Header))		Print(Ctx, TEXT("Complex "), InColor);
+	else if (IsSingleLayerWater(Header))	Print(Ctx, TEXT("SLWater "), InColor);
+	else									Print(Ctx, TEXT("!ERROR! "), FontRed);
 }
 
 void PrintSSSType(inout FShaderPrintContext Ctx, uint SSSType, FFontColor InColor)
@@ -365,10 +367,10 @@ void PrintPixel(uint2 InCoord, float3 InWorldPosition, float3 V)
 		Print(Context, Header.BSDFCount, FontBSDFCount, 2, 0);
 		Print(Context, TEXT("BSDF"), FontBSDFCount);
 		Print(Context, TEXT(" - "), FontWhite);
-		PrintPixelType(Context, bIsSimpleMaterial, bIsSingleMaterial, FontLightRed);
+		PrintPixelType(Context, Header, FontLightRed);
 		Newline(Context);
 
-		Print(Context, TEXT(" AO    IndIrr TopRoughness PreShadow ZeroShadow ContacShadow Ind.Occluder IsWater HasSSS"), FontHeaderStateName);
+		Print(Context, TEXT(" AO    IndIrr TopRoughness PreShadow ZeroShadow ContacShadow Ind.Occluder HasSSS"), FontHeaderStateName);
 		Newline(Context);
 
 		Print(Context, TEXT(" "));
@@ -379,7 +381,6 @@ void PrintPixel(uint2 InCoord, float3 InWorldPosition, float3 V)
 		PrintBool(Context, HasZeroPrecShadowMask(Header), FontHeaderPropValue);							Print(Context, TEXT("      "));
 		PrintBool(Context, DoesCastContactShadow(Header), FontHeaderPropValue);							Print(Context, TEXT("        "));
 		PrintBool(Context, HasDynamicIndirectShadowCasterRepresentation(Header), FontHeaderPropValue);	Print(Context, TEXT("        "));
-		PrintBool(Context, IsSingleLayerWater(Header), FontHeaderPropValue);							Print(Context, TEXT("   "));
 		PrintBool(Context, HasSubsurface(Header), FontHeaderPropValue);
 		Newline(Context);
 		Newline(Context);
