Fix Strata shader compilation with fxc

diff --git a/Engine/Shaders/Private/Strata/Strata.ush b/Engine/Shaders/Private/Strata/Strata.ush
--- a/Engine/Shaders/Private/Strata/Strata.ush
+++ b/Engine/Shaders/Private/Strata/Strata.ush
@@ -1651,7 +1651,10 @@ struct FStrataPixelHeader
 	uint State;
 };
 
-void SetMaterialMode(inout FStrataPixelHeader Out, int MateralMode)							{ Out.State = (Out.State & ~HEADER_MASK_MATERIALMODE) | (MateralMode & HEADER_MASK_MATERIALMODE); }
+// Workaround for legacy shader compiler, claiming FStrataPixelHeader has uninitialized state when passed as a 'inout' parameter
+void InternalSetMaterialMode(inout uint OutState, uint MaterialMode)						{ OutState = (OutState & ~HEADER_MASK_MATERIALMODE) | (MaterialMode & HEADER_MASK_MATERIALMODE); }
+#define SetMaterialMode(InHeader, InMode)													InternalSetMaterialMode(InHeader.State, InMode)
+
 bool IsSimpleMaterial(const FStrataPixelHeader In)											{ return (In.State & HEADER_MASK_MATERIALMODE) == HEADER_MATERIALMODE_SIMPLE || STRATA_FASTPATH; }
 bool IsSingleMaterial(const FStrataPixelHeader In)											{ return (In.State & HEADER_MASK_MATERIALMODE) == HEADER_MATERIALMODE_SINGLE || STRATA_SINGLEPATH; }
 bool IsComplexMaterial(const FStrataPixelHeader In)											{ return (In.State & HEADER_MASK_MATERIALMODE) == HEADER_MATERIALMODE_COMPLEX && !STRATA_FASTPATH && !STRATA_SINGLEPATH; }
diff --git a/Engine/Shaders/Private/Strata/StrataExport.ush b/Engine/Shaders/Private/Strata/StrataExport.ush
--- a/Engine/Shaders/Private/Strata/StrataExport.ush
+++ b/Engine/Shaders/Private/Strata/StrataExport.ush
@@ -195,11 +195,11 @@ void PackStrataOut(
 						const bool LuminanceWeightEqualOne = all(BSDF.LuminanceWeightV == 1.0f);
 
 						// Update simple material compatibility
-						bIsSimpleMaterial &= IsStrataSlabFastPathCompatible(BSDF) && LuminanceWeightEqualOne;
+						bIsSimpleMaterial = bIsSimpleMaterial && IsStrataSlabFastPathCompatible(BSDF) && LuminanceWeightEqualOne;
 
 						// Update single material compatibility. For now, single materials don't support anisotropy, 
 						// they use the complex pass, as the toplayer data does not contain the full frame basis, only the top normal.
-						bIsSingleMaterial &= BSDF_GETHASANISOTROPY(BSDF) == 0 && LuminanceWeightEqualOne;
+						bIsSingleMaterial = bIsSingleMaterial && BSDF_GETHASANISOTROPY(BSDF) == 0 && LuminanceWeightEqualOne;
 
 						TopLayerTotalWeight += BSDF.TopLayerDataWeight;
 						TopLayerData.Roughness += BSDF.TopLayerDataWeight * SLAB_ROUGHNESS(BSDF);
@@ -268,8 +268,8 @@ void PackStrataOut(
 		// Set storage layout as either: fast(0), single(1), or complex(2)
 		if (bIsOnlySlab)
 		{
-			bIsSimpleMaterial &= BSDFVisibleCount == 1;
-			bIsSingleMaterial &= BSDFVisibleCount == 1;
+			bIsSimpleMaterial = bIsSimpleMaterial && BSDFVisibleCount == 1;
+			bIsSingleMaterial = bIsSingleMaterial && BSDFVisibleCount == 1;
 			SetMaterialMode(StrataPixelHeader, bIsSimpleMaterial ? HEADER_MATERIALMODE_SIMPLE : bIsSingleMaterial ? HEADER_MATERIALMODE_SINGLE : HEADER_MATERIALMODE_COMPLEX);
 		}
 
