Fix dithering issue on F0 value.

Change dithering to be done on 'gamma compressed' data rather than linear data as dithering on low value cause high relative error. For instance a dither bias on 6bits encoding cause a +/- 1.5% on F0 value, causing visual artifact on character.

With this change bias is reduced for lower value, but increase on higher value. From tests, it does not seems very noticealble.Maybe more test will be needed.

diff --git a/Engine/Shaders/Private/Strata/Strata.ush b/Engine/Shaders/Private/Strata/Strata.ush
--- a/Engine/Shaders/Private/Strata/Strata.ush
+++ b/Engine/Shaders/Private/Strata/Strata.ush
@@ -2724,11 +2724,21 @@ float4 UnpackRGBA8(uint rgba)
 	return Out;
 }
 
-uint PackR7G7B6Gamma2(float3 rgb)
+uint PackR7G7B6Gamma2(float3 rgb, float Dither)
 {
-	uint r = (uint(sqrt(saturate(rgb.r)) * 127.0f) << 0);
-	uint g = (uint(sqrt(saturate(rgb.g)) * 127.0f) << 7);
-	uint b = (uint(sqrt(saturate(rgb.b)) *  63.0f) << 14);
+	// We execute a dither that can go in positive/negative direction to avoid surfaces to looks overall brighter.
+	const float CenteredDither = Dither - 0.5f;
+	const float Dither8bits = CenteredDither * 1.0f / 255.0f;
+	const float Dither7bits = CenteredDither * 1.0f / 127.0f;
+	const float Dither6bits = CenteredDither * 1.0f / 63.0f;
+	const float3 Dither776 = float3(Dither7bits, Dither7bits, Dither6bits);
+
+	// Dither value to hide more aggressive bit reduction.
+	const float3 dithered_rgb = DitherIfNonBlack(sqrt(saturate(rgb.r)), Dither776);
+
+	uint r = (uint(dithered_rgb.r * 127.0f) << 0);
+	uint g = (uint(dithered_rgb.g * 127.0f) << 7);
+	uint b = (uint(dithered_rgb.b *  63.0f) << 14);
 	return r | g | b;
 }
 
diff --git a/Engine/Shaders/Private/Strata/StrataExport.ush b/Engine/Shaders/Private/Strata/StrataExport.ush
--- a/Engine/Shaders/Private/Strata/StrataExport.ush
+++ b/Engine/Shaders/Private/Strata/StrataExport.ush
@@ -181,10 +181,6 @@ void PackStrataOut(
 				{
 					BSDFVisibleCount++;
 
-					// Dither BSDF data to hide more aggressive bit reduction.
-					// Sanitize happening after will clamp values that could have gone above and below valid range.
-					StrataDitherBSDF(BSDF, Dither);
-
 					switch (BSDF_GETTYPE(BSDF))
 					{
 					case STRATA_BSDF_TYPE_SLAB:
@@ -384,10 +380,10 @@ void PackStrataOut(
 					{
 						// Now we pack diffuse and F0 in a special way: both encoded as R7G7B6 with gamma 2.0. The low bit cound will be hidden by dithering and TAA.
 						// We pack 32 bits in the first uint and the remaining bits are pack is the lowest significant bit of the second uint.
-						const uint PackedDiffuse20Bits = PackR7G7B6Gamma2(SLAB_DIFFUSEALBEDO(BSDF));
+						const uint PackedDiffuse20Bits = PackR7G7B6Gamma2(SLAB_DIFFUSEALBEDO(BSDF), Dither);
 						const uint PackedDiffuse12Bits = PackedDiffuse20Bits & 0xFFF;
 						const uint PackedDiffuse8Bits = (PackedDiffuse20Bits >> 12) & 0xFF;
-						const uint PackedF020Bits = PackR7G7B6Gamma2(SLAB_F0(BSDF));
+						const uint PackedF020Bits = PackR7G7B6Gamma2(SLAB_F0(BSDF), Dither);
 						const uint PackedRoughness8bits = PackR8(SLAB_ROUGHNESS(BSDF));
 
 						// Data0 (Header_State|Header_AO|Roughness|Diffuse8bits)
@@ -451,8 +447,8 @@ void PackStrataOut(
 						{
 							// Now we pack diffuse and F0 in a special way: both encoded as R7G7B6 with gamma 2.0. The low bit cound will be hidden by dithering and TAA.
 							// We pack 32 bits in the first uint and the remaining bits are pack is the lowest significant bit of the second uint.
-							const uint PackedDiffuse20Bits = PackR7G7B6Gamma2(SLAB_DIFFUSEALBEDO(BSDF));
-							const uint PackedF020Bits = PackR7G7B6Gamma2(SLAB_F0(BSDF));
+							const uint PackedDiffuse20Bits = PackR7G7B6Gamma2(SLAB_DIFFUSEALBEDO(BSDF), Dither);
+							const uint PackedF020Bits = PackR7G7B6Gamma2(SLAB_F0(BSDF), Dither);
 							const uint PackedData32Bits = ((PackedDiffuse20Bits << 12) & 0xFFFFF000) | (PackedF020Bits & 0xFFF);
 							const uint PackedData8Bits = (PackedF020Bits >> 12) & 0xFF;
 
