Add back pixel footprint visualization.

diff --git a/Engine/Shaders/Private/Strata/Strata.ush b/Engine/Shaders/Private/Strata/Strata.ush
--- a/Engine/Shaders/Private/Strata/Strata.ush
+++ b/Engine/Shaders/Private/Strata/Strata.ush
@@ -1157,18 +1157,26 @@ struct FStrataPixelFootprint
 	float PixelRadiusInWorldSpace; // In cm
 };
 
+FStrataPixelFootprint StrataGetPixelFootprint(float3 dPdx, float3 dPdy)
+{
+	const float dX = lengthFast(dPdx);
+	const float dY = lengthFast(dPdy);
+
+	FStrataPixelFootprint Out;
+	Out.PixelRadiusInWorldSpace = min(dX, dY) * 0.5f;
+	return Out;
+}
+
 FStrataPixelFootprint StrataGetPixelFootprint(float3 InWorldPosition)
 {
 #if CLOUD_LAYER_PIXEL_SHADER==1
 	// Clouds materials feature dynamic loops and thus cannot support derivatives.
-	FStrataPixelFootprint Out = (FStrataPixelFootprint)0;
+	return (FStrataPixelFootprint)0;
 #else
-	const float dX = lengthFast(ddx(InWorldPosition));
-	const float dY = lengthFast(ddy(InWorldPosition));
-	FStrataPixelFootprint Out;
-	Out.PixelRadiusInWorldSpace = min(dX, dY) * 0.5f;
+	const float3 dPdx = ddx(InWorldPosition);
+	const float3 dPdy = ddy(InWorldPosition);
+	return StrataGetPixelFootprint(dPdx, dPdy);
 #endif
-	return Out;
 }
 
 ///////////////////////////////////////////////////////////////////////////////
diff --git a/Engine/Shaders/Private/Strata/StrataVisualize.usf b/Engine/Shaders/Private/Strata/StrataVisualize.usf
--- a/Engine/Shaders/Private/Strata/StrataVisualize.usf
+++ b/Engine/Shaders/Private/Strata/StrataVisualize.usf
@@ -76,14 +76,10 @@ void DrawPixelFootprint(float3 P, float3 dPdx, float3 dPdy, uint2 PixelCoord)
 			MFP = StrataSubSurfaceHeaderGetMFP(SSSHeader);
 		}
 
-		FStrataPixelFootprint Footprint = StrataGetPixelFootprint(P);
+		FStrataPixelFootprint Footprint = StrataGetPixelFootprint(dPdx, dPdy);
 		AddDrawPixelFootprint(P, dPdx, dPdy, 0.5f, false, ColorRed);
 		AddDrawPixelFootprint(P, dPdx, dPdy, Footprint.PixelRadiusInWorldSpace, true, ColorOrange);
 		AddDrawPixelFootprint(P, dPdx, dPdy, max3(MFP.x, MFP.y, MFP.z), true, ColorCyan);
-
-		// Normal
-		const float3 N = normalize(cross(normalize(dPdx), normalize(dPdy)));
-		AddLineTWS(P, P + N * 10.f, ColorBlue);
 	}
 }
 
@@ -435,18 +431,32 @@ void PrintPixel(uint2 InCoord, float3 InWorldPosition, float3 V)
 	}
 }
 
-[numthreads(1, 1, 1)]
-void MaterialPrintInfoCS(uint3 DispatchThreadId : SV_DispatchThreadID)
+float3 GetWorldPositionFromPixelPos(uint2 PixelPos)
 {
-	const uint2 PixelPos = (float2(View.CursorPosition) * View.ViewResolutionFraction);
 	const float2 BufferUV = SvPositionToBufferUV(float4(PixelPos + 0.5f, 0, 0));
 	const float2 ScreenPosition = SvPositionToScreenPosition(float4(PixelPos, 0.5f, 1)).xy;
 	const float DeviceZ = SampleDeviceZFromSceneTextures(BufferUV);
 	const float SceneDepth = ConvertFromDeviceZ(DeviceZ);
 	const float3 WorldPosition = mul(float4(ScreenPosition * SceneDepth, SceneDepth, 1), PrimaryView.ScreenToTranslatedWorld).xyz;
+	return WorldPosition;
+}
+
+[numthreads(1, 1, 1)]
+void MaterialPrintInfoCS(uint3 DispatchThreadId : SV_DispatchThreadID)
+{
+	const uint2 PixelPos = (float2(View.CursorPosition) * View.ViewResolutionFraction);
+	const float3 WorldPosition = GetWorldPositionFromPixelPos(PixelPos);
 	const float3 V = -normalize(WorldPosition - PrimaryView.TranslatedWorldCameraOrigin);
 
+	// Print pixel info
 	PrintPixel(PixelPos, WorldPosition, V);
+
+	// Print pixel footprint
+	const float3 WorldPositionX = GetWorldPositionFromPixelPos(PixelPos + uint2(1, 0));
+	const float3 WorldPositionY = GetWorldPositionFromPixelPos(PixelPos + uint2(0, 1));
+	const float3 dPdx = WorldPositionX - WorldPosition;
+	const float3 dPdy = WorldPositionY - WorldPosition;
+	DrawPixelFootprint(WorldPosition, dPdx, dPdy, PixelPos);
 }
 #endif // SHADER_MATERIALPRINT
 
