Add Strata visualization background for better readability.

diff --git a/Engine/Shaders/Private/Strata/StrataVisualize.usf b/Engine/Shaders/Private/Strata/StrataVisualize.usf
--- a/Engine/Shaders/Private/Strata/StrataVisualize.usf
+++ b/Engine/Shaders/Private/Strata/StrataVisualize.usf
@@ -124,7 +124,23 @@ void PrintSSSType(inout FShaderPrintContext Ctx, uint SSSType, FFontColor InColo
 	else /*	(SSSType == SSS_TYPE_INVALID) */		Print(Ctx, TEXT("Invalid"), InColor);
 }
 
-void PrintBSDF(inout FShaderPrintContext Context, inout FStrataAddressing StrataAddressing, FStrataPixelHeader Header, float3 WorldPosition, float3 V)
+void PrintBackgroundRect(inout FShaderPrintContext Ctx, float2 StartRect, float2 EndRect)
+{
+	const float4 RectBackgroundColor = float4(1, 1, 1, 0.25f);
+	AddFilledQuadSS(
+		(StartRect - ShaderPrintData.FontSize) * ShaderPrintData.Resolution,
+		(EndRect + ShaderPrintData.FontSize) * ShaderPrintData.Resolution,
+		RectBackgroundColor);
+}
+
+void Newline(inout FShaderPrintContext Context, inout float2 MaxRect)
+{
+	MaxRect.x = max(MaxRect.x, Context.Pos.x);
+	MaxRect.y = Context.Pos.y;
+	Newline(Context);
+}
+
+void PrintBSDF(inout FShaderPrintContext Context, inout FStrataAddressing StrataAddressing, FStrataPixelHeader Header, float3 WorldPosition, float3 V, inout float2 RectMax)
 {
 	const FFontColor FontBSDF         = FontSilver;
 	const FFontColor FontBSDFType     = FontEmerald;
@@ -150,7 +166,7 @@ void PrintBSDF(inout FShaderPrintContext Context, inout FStrataAddressing Strata
 			const bool bGreyTopTrans = BSDF.TransmittanceAboveAlongN.x == BSDF.TransmittanceAboveAlongN.y && BSDF.TransmittanceAboveAlongN.x == BSDF.TransmittanceAboveAlongN.z;
 
 			Print(Context, TEXT("Slab "), FontBSDFType);
-			Newline(Context);
+			Newline(Context, RectMax);
 			Print(Context, TEXT(" NormalID Aniso TopLayer Scattering SingleScatt. IsThin WeightV "), FontBSDFStateName);
 			if (!bGreyWeightV)
 			{
@@ -160,7 +176,7 @@ void PrintBSDF(inout FShaderPrintContext Context, inout FStrataAddressing Strata
 			{
 				Print(Context, TEXT("WeightL"), FontBSDFStateName);
 			}
-			Newline(Context);
+			Newline(Context, RectMax);
 
 			Print(Context, TEXT(" "), FontSilver);
 			Print(Context, BSDF_GETSHAREDLOCALBASISID(BSDF), FontBSDF, 2, 1);			Print(Context, TEXT("       "), FontSilver);
@@ -199,18 +215,18 @@ void PrintBSDF(inout FShaderPrintContext Context, inout FStrataAddressing Strata
 					Print(Context, BSDF.TransmittanceAboveAlongN.b, FontLightBlue,  6, 4);
 				}
 			}
-			Newline(Context);
+			Newline(Context, RectMax);
 
-			Print(Context, TEXT(" Diffuse     "), FontBSDFPropName);	Print(Context, SLAB_DIFFUSEALBEDO(BSDF), FontBSDFPropValu);			Newline(Context);
+			Print(Context, TEXT(" Diffuse     "), FontBSDFPropName);	Print(Context, SLAB_DIFFUSEALBEDO(BSDF), FontBSDFPropValu);			Newline(Context, RectMax);
 
 			const bool bGreyF0 = SLAB_F0(BSDF).x == SLAB_F0(BSDF).y && SLAB_F0(BSDF).x == SLAB_F0(BSDF).z;
 			if (bGreyF0)
 			{
-			Print(Context, TEXT(" F0          "), FontBSDFPropName);	Print(Context, SLAB_F0(BSDF).x, FontBSDFPropValu);					Newline(Context);
+			Print(Context, TEXT(" F0          "), FontBSDFPropName);	Print(Context, SLAB_F0(BSDF).x, FontBSDFPropValu);					Newline(Context, RectMax);
 			}
 			else
 			{
-			Print(Context, TEXT(" F0          "), FontBSDFPropName);	Print(Context, SLAB_F0(BSDF), FontBSDFPropValu);					Newline(Context);
+			Print(Context, TEXT(" F0          "), FontBSDFPropName);	Print(Context, SLAB_F0(BSDF), FontBSDFPropValu);					Newline(Context, RectMax);
 			}
 
 
@@ -219,24 +235,24 @@ void PrintBSDF(inout FShaderPrintContext Context, inout FStrataAddressing Strata
 			const bool bGreyF90 = SLAB_F90(BSDF).x == SLAB_F90(BSDF).y && SLAB_F90(BSDF).x == SLAB_F90(BSDF).z;
 			if (bGreyF90)
 			{
-			Print(Context, TEXT(" F90         "), FontBSDFPropName);	Print(Context, SLAB_F90(BSDF).x, FontBSDFPropValu);					Newline(Context);
+			Print(Context, TEXT(" F90         "), FontBSDFPropName);	Print(Context, SLAB_F90(BSDF).x, FontBSDFPropValu);					Newline(Context, RectMax);
 			}
 			else
 			{
-			Print(Context, TEXT(" F90         "), FontBSDFPropName);	Print(Context, SLAB_F90(BSDF), FontBSDFPropValu);					Newline(Context);
+			Print(Context, TEXT(" F90         "), FontBSDFPropName);	Print(Context, SLAB_F90(BSDF), FontBSDFPropValu);					Newline(Context, RectMax);
 			}
 			}
 
-			Print(Context, TEXT(" Roughness   "), FontBSDFPropName);	Print(Context, SLAB_ROUGHNESS(BSDF), FontBSDFPropValu);				Newline(Context);
+			Print(Context, TEXT(" Roughness   "), FontBSDFPropName);	Print(Context, SLAB_ROUGHNESS(BSDF), FontBSDFPropValu);				Newline(Context, RectMax);
 
 			if (BSDF_GETHASANISOTROPY(BSDF))
 			{
-			Print(Context, TEXT(" Anisotropy  "), FontBSDFPropName);	Print(Context, SLAB_ANISOTROPY(BSDF), FontBSDFPropValu);			Newline(Context);
+			Print(Context, TEXT(" Anisotropy  "), FontBSDFPropName);	Print(Context, SLAB_ANISOTROPY(BSDF), FontBSDFPropValu);			Newline(Context, RectMax);
 			}
 
 			if (BSDF_GETHASHAZINESS(BSDF))
 			{
-			Print(Context, TEXT(" Haziness    "), FontBSDFPropName);	Print(Context, SLAB_HAZINESS(BSDF), FontBSDFPropValu);				Newline(Context);
+			Print(Context, TEXT(" Haziness    "), FontBSDFPropName);	Print(Context, SLAB_HAZINESS(BSDF), FontBSDFPropValu);				Newline(Context, RectMax);
 			}
 			
 			if (BSDF_GETSSSTYPE(BSDF) != SSS_TYPE_INVALID || BSDF_GETISSIMPLEVOLUME(BSDF) || BSDF_GETISTHIN(BSDF))
@@ -245,26 +261,26 @@ void PrintBSDF(inout FShaderPrintContext Context, inout FStrataAddressing Strata
 			const FStrataSubsurfaceHeader SSSHeader = StrataLoadSubsurfaceHeader(Strata.SSSTexture, StrataAddressing.PixelCoords);
 			const uint SSSType = StrataSubSurfaceHeaderGetSSSType(SSSHeader);
 			const float Opacity = StrataSubSurfaceHeaderGetWrapOpacity(SSSHeader);
-			Print(Context, TEXT(" SSS Type    "), FontBSDFPropName);	PrintSSSType(Context, SSSType, FontYellow);	Newline(Context);
-			Print(Context, TEXT(" SSS Opacity "), FontBSDFPropName);	Print(Context, Opacity, FontBSDFPropValu);	Newline(Context);
+			Print(Context, TEXT(" SSS Type    "), FontBSDFPropName);	PrintSSSType(Context, SSSType, FontYellow);	Newline(Context, RectMax);
+			Print(Context, TEXT(" SSS Opacity "), FontBSDFPropName);	Print(Context, Opacity, FontBSDFPropValu);	Newline(Context, RectMax);
 
 			if (BSDF_GETSSSTYPE(BSDF) == SSS_TYPE_DIFFUSION_PROFILE && BSDF_GETISSIMPLEVOLUME(BSDF) == 0)
 			{
-			Print(Context, TEXT(" SSS ID.     "), FontBSDFPropName);	Print(Context, SLAB_SSSPROFILEID(BSDF), FontBSDFPropValu);			Newline(Context);
-			Print(Context, TEXT(" SSS Radius  "), FontBSDFPropName);	Print(Context, SLAB_SSSPROFILERADIUSSCALE(BSDF), FontBSDFPropValu);	Newline(Context);
+			Print(Context, TEXT(" SSS ID.     "), FontBSDFPropName);	Print(Context, SLAB_SSSPROFILEID(BSDF), FontBSDFPropValu);			Newline(Context, RectMax);
+			Print(Context, TEXT(" SSS Radius  "), FontBSDFPropName);	Print(Context, SLAB_SSSPROFILERADIUSSCALE(BSDF), FontBSDFPropValu);	Newline(Context, RectMax);
 			}
 			else
 			{
-			Print(Context, TEXT(" SSS MFP     "), FontBSDFPropName);	Print(Context, SLAB_SSSMFP(BSDF), FontBSDFPropValu);				Newline(Context);
+			Print(Context, TEXT(" SSS MFP     "), FontBSDFPropName);	Print(Context, SLAB_SSSMFP(BSDF), FontBSDFPropValu);				Newline(Context, RectMax);
 			}
 
-			Print(Context, TEXT(" SSS Phase   "), FontBSDFPropName);	Print(Context, SLAB_SSSPHASEANISOTROPY(BSDF), FontBSDFPropValu);	Newline(Context);
+			Print(Context, TEXT(" SSS Phase   "), FontBSDFPropName);	Print(Context, SLAB_SSSPHASEANISOTROPY(BSDF), FontBSDFPropValu);	Newline(Context, RectMax);
 			} // SSS
 
 			if (BSDF_GETHASFUZZ(BSDF))
 			{
-			Print(Context, TEXT(" Fuzz Amount "), FontBSDFPropName);	Print(Context, SLAB_FUZZ_AMOUNT(BSDF), FontBSDFPropValu);			Newline(Context);
-			Print(Context, TEXT(" Fuzz Color  "), FontBSDFPropName);	Print(Context, SLAB_FUZZ_COLOR(BSDF), FontBSDFPropValu);			Newline(Context);
+			Print(Context, TEXT(" Fuzz Amount "), FontBSDFPropName);	Print(Context, SLAB_FUZZ_AMOUNT(BSDF), FontBSDFPropValu);			Newline(Context, RectMax);
+			Print(Context, TEXT(" Fuzz Color  "), FontBSDFPropName);	Print(Context, SLAB_FUZZ_COLOR(BSDF), FontBSDFPropValu);			Newline(Context, RectMax);
 			}
 		}
 		break;
@@ -272,9 +288,9 @@ void PrintBSDF(inout FShaderPrintContext Context, inout FStrataAddressing Strata
 		case STRATA_BSDF_TYPE_HAIR:
 		{
 			Print(Context, TEXT("Hair"), FontBSDFType);
-			Newline(Context);
+			Newline(Context, RectMax);
 			Print(Context, TEXT(" NormalID Aniso IsTopLayer LuminanceWeight"), FontBSDF);
-			Newline(Context);
+			Newline(Context, RectMax);
 
 			Print(Context, BSDF_GETSHAREDLOCALBASISID(BSDF), FontBSDF);
 			PrintBool(Context, BSDF_GETHASANISOTROPY(BSDF), FontBSDF);
@@ -289,14 +305,14 @@ void PrintBSDF(inout FShaderPrintContext Context, inout FStrataAddressing Strata
 				Print(Context, BSDF.LuminanceWeightV.g, FontLightGreen);
 				Print(Context, BSDF.LuminanceWeightV.b, FontLightBlue);
 			}
-			Newline(Context);
-
-			Print(Context, TEXT(" BaseColor         "), FontBSDFPropName);	Print(Context, HAIR_BASECOLOR(BSDF), FontBSDFPropValu);				Newline(Context);
-			Print(Context, TEXT(" Specular          "), FontBSDFPropName);	Print(Context, HAIR_SPECULAR(BSDF), FontBSDFPropValu);				Newline(Context);
-			Print(Context, TEXT(" Roughness         "), FontBSDFPropName);	Print(Context, HAIR_ROUGHNESS(BSDF), FontBSDFPropValu);				Newline(Context);
-			Print(Context, TEXT(" Scatter           "), FontBSDFPropName);	Print(Context, HAIR_SCATTER(BSDF), FontBSDFPropValu);				Newline(Context);
-			Print(Context, TEXT(" Backlit           "), FontBSDFPropName);	Print(Context, HAIR_BACKLIT(BSDF), FontBSDFPropValu);				Newline(Context);
-			Print(Context, TEXT(" Has Transmittance "), FontBSDFPropName);	Print(Context, HAIR_COMPLEXTRANSMITTANCE(BSDF), FontBSDFPropValu);	Newline(Context);
+			Newline(Context, RectMax);
+
+			Print(Context, TEXT(" BaseColor         "), FontBSDFPropName);	Print(Context, HAIR_BASECOLOR(BSDF), FontBSDFPropValu);				Newline(Context, RectMax);
+			Print(Context, TEXT(" Specular          "), FontBSDFPropName);	Print(Context, HAIR_SPECULAR(BSDF), FontBSDFPropValu);				Newline(Context, RectMax);
+			Print(Context, TEXT(" Roughness         "), FontBSDFPropName);	Print(Context, HAIR_ROUGHNESS(BSDF), FontBSDFPropValu);				Newline(Context, RectMax);
+			Print(Context, TEXT(" Scatter           "), FontBSDFPropName);	Print(Context, HAIR_SCATTER(BSDF), FontBSDFPropValu);				Newline(Context, RectMax);
+			Print(Context, TEXT(" Backlit           "), FontBSDFPropName);	Print(Context, HAIR_BACKLIT(BSDF), FontBSDFPropValu);				Newline(Context, RectMax);
+			Print(Context, TEXT(" Has Transmittance "), FontBSDFPropName);	Print(Context, HAIR_COMPLEXTRANSMITTANCE(BSDF), FontBSDFPropValu);	Newline(Context, RectMax);
 		}
 		break;
 
@@ -305,11 +321,11 @@ void PrintBSDF(inout FShaderPrintContext Context, inout FStrataAddressing Strata
 			Print(Context, TEXT("Water"), FontBSDFType);
 
 			Newline(Context);
-			Print(Context, TEXT(" BaseColor            "), FontBSDFPropName);	Print(Context, SLW_BASECOLOR(BSDF), FontBSDFPropValu);			Newline(Context);
-			Print(Context, TEXT(" Metallic             "), FontBSDFPropName);	Print(Context, SLW_METALLIC(BSDF), FontBSDFPropValu);			Newline(Context);
-			Print(Context, TEXT(" Specular             "), FontBSDFPropName);	Print(Context, SLW_SPECULAR(BSDF), FontBSDFPropValu);			Newline(Context);
-			Print(Context, TEXT(" Roughness            "), FontBSDFPropName);	Print(Context, SLW_ROUGHNESS(BSDF), FontBSDFPropValu);			Newline(Context);
-			Print(Context, TEXT(" Top Material Opacity "), FontBSDFPropName);	Print(Context, SLW_TOPMATERIALOPACITY(BSDF), FontBSDFPropValu);	Newline(Context);
+			Print(Context, TEXT(" BaseColor            "), FontBSDFPropName);	Print(Context, SLW_BASECOLOR(BSDF), FontBSDFPropValu);			Newline(Context, RectMax);
+			Print(Context, TEXT(" Metallic             "), FontBSDFPropName);	Print(Context, SLW_METALLIC(BSDF), FontBSDFPropValu);			Newline(Context, RectMax);
+			Print(Context, TEXT(" Specular             "), FontBSDFPropName);	Print(Context, SLW_SPECULAR(BSDF), FontBSDFPropValu);			Newline(Context, RectMax);
+			Print(Context, TEXT(" Roughness            "), FontBSDFPropName);	Print(Context, SLW_ROUGHNESS(BSDF), FontBSDFPropValu);			Newline(Context, RectMax);
+			Print(Context, TEXT(" Top Material Opacity "), FontBSDFPropName);	Print(Context, SLW_TOPMATERIALOPACITY(BSDF), FontBSDFPropValu);	Newline(Context, RectMax);
 		}
 		break;
 
@@ -349,6 +365,10 @@ void PrintPixel(uint2 InCoord, float3 InWorldPosition, float3 V)
 	const FFontColor FontHeaderPropName = FontWhite;
 	const FFontColor FontHeaderPropValue = FontSilver;
 
+	// Backgroun rect. for better print visibility. Track min/max point
+	float2 RectMin = Context.Pos;
+	float2 RectMax = Context.StartPos;
+
 	// Header
 	if (BSDFIndex < Header.BSDFCount && BSDFIndex == 0)
 	{
@@ -360,14 +380,15 @@ void PrintPixel(uint2 InCoord, float3 InWorldPosition, float3 V)
 
 		FStrataIrradianceAndOcclusion IrradianceAndOcclusion = StrataGetIrradianceAndAO(Header);
 
+		RectMin = Context.Pos;
 		Print(Context, Header.BSDFCount, FontBSDFCount, 2, 0);
 		Print(Context, TEXT("BSDF"), FontBSDFCount);
 		Print(Context, TEXT(" - "), FontWhite);
 		PrintPixelType(Context, Header, FontLightRed);
-		Newline(Context);
+		Newline(Context, RectMax);
 
 		Print(Context, TEXT(" AO    IndIrr TopRoughness PreShadow ZeroShadow ContacShadow Ind.Occluder HasSSS"), FontHeaderStateName);
-		Newline(Context);
+		Newline(Context, RectMax);
 
 		Print(Context, TEXT(" "));
 		Print(Context, IrradianceAndOcclusion.MaterialAO, FontHeaderPropValue, 5, 3);					Print(Context, TEXT(" "));
@@ -378,8 +399,10 @@ void PrintPixel(uint2 InCoord, float3 InWorldPosition, float3 V)
 		PrintBool(Context, DoesCastContactShadow(Header), FontHeaderPropValue);							Print(Context, TEXT("        "));
 		PrintBool(Context, HasDynamicIndirectShadowCasterRepresentation(Header), FontHeaderPropValue);	Print(Context, TEXT("        "));
 		PrintBool(Context, HasSubsurface(Header), FontHeaderPropValue);
+		Newline(Context, RectMax);
 		Newline(Context);
-		Newline(Context);
+		
+		PrintBackgroundRect(Context, RectMin, RectMax);
 	}
 
 	// BSDFs
@@ -390,7 +413,11 @@ void PrintPixel(uint2 InCoord, float3 InWorldPosition, float3 V)
 			const uint AddressOffset = UnpackBSDFOffsetAtIndex(Strata.BSDFOffsetTexture[InCoord], BSDFIndex, Header.BSDFCount);
 			StrataSeekBSDF(StrataAddressing, AddressOffset);
 		}
-		PrintBSDF(Context, StrataAddressing, Header, InWorldPosition, V);
+
+		RectMin = Context.Pos;
+		RectMax = Context.StartPos;
+		PrintBSDF(Context, StrataAddressing, Header, InWorldPosition, V, RectMax);
+		PrintBackgroundRect(Context, RectMin, RectMax);
 
 		RWPositionOffsetBuffer[0] = asuint(Context.Pos.x);
 		RWPositionOffsetBuffer[1] = asuint(Context.Pos.y);
@@ -410,24 +437,27 @@ void PrintPixel(uint2 InCoord, float3 InWorldPosition, float3 V)
 		const uint TotalSize  = FootPrint_PostBSDFs - FootPrint_Start;
 
 		const uint TopLayerDataBytes = 4;	// bytes 
-		Print(Context, TEXT("Memory"), FontEmerald); Newline(Context);
+		RectMin = Context.Pos;
+		RectMax = Context.StartPos;
+		Print(Context, TEXT("Memory"), FontEmerald); Newline(Context, RectMax);
 		if (bIsSingleMaterial)
 		{
-			Print(Context, TEXT(" Header       "), FontSilver); Print(Context, HeaderSize - TopLayerDataBytes,				FontSilver, 3, 0); Print(Context, TEXT("bytes"), FontSilver); Newline(Context);
-			Print(Context, TEXT(" TopNormalTex "), FontSilver); Print(Context, TopLayerDataBytes,							FontSilver, 3, 0); Print(Context, TEXT("bytes"), FontSilver); Newline(Context);
-			Print(Context, TEXT(" BSDF         "), FontSilver); Print(Context, BSDFsSize,									FontSilver, 3, 0); Print(Context, TEXT("bytes"), FontSilver); Newline(Context);
+			Print(Context, TEXT(" Header       "), FontSilver); Print(Context, HeaderSize - TopLayerDataBytes,				FontSilver, 3, 0); Print(Context, TEXT("bytes"), FontSilver); Newline(Context, RectMax);
+			Print(Context, TEXT(" TopNormalTex "), FontSilver); Print(Context, TopLayerDataBytes,							FontSilver, 3, 0); Print(Context, TEXT("bytes"), FontSilver); Newline(Context, RectMax);
+			Print(Context, TEXT(" BSDF         "), FontSilver); Print(Context, BSDFsSize,									FontSilver, 3, 0); Print(Context, TEXT("bytes"), FontSilver); Newline(Context, RectMax);
 		}
 		else if (bIsSimpleMaterial)
 		{
-			Print(Context, TEXT(" Header+BSDF  "), FontSilver); Print(Context, HeaderSize + BSDFsSize - TopLayerDataBytes,	FontSilver, 3, 0); Print(Context, TEXT("bytes"), FontSilver); Newline(Context);
-			Print(Context, TEXT(" TopNormalTex "), FontSilver); Print(Context, TopLayerDataBytes,							FontSilver, 3, 0); Print(Context, TEXT("bytes"), FontSilver); Newline(Context);
+			Print(Context, TEXT(" Header+BSDF  "), FontSilver); Print(Context, HeaderSize + BSDFsSize - TopLayerDataBytes,	FontSilver, 3, 0); Print(Context, TEXT("bytes"), FontSilver); Newline(Context, RectMax);
+			Print(Context, TEXT(" TopNormalTex "), FontSilver); Print(Context, TopLayerDataBytes,							FontSilver, 3, 0); Print(Context, TEXT("bytes"), FontSilver); Newline(Context, RectMax);
 		}
 		else
 		{
-			Print(Context, TEXT(" Header+Norm  "), FontSilver); Print(Context, HeaderSize,									FontSilver, 3, 0); Print(Context, TEXT("bytes"), FontSilver); Newline(Context);
-			Print(Context, TEXT(" BSDFs        "), FontSilver); Print(Context, BSDFsSize,									FontSilver, 3, 0); Print(Context, TEXT("bytes"), FontSilver); Newline(Context);
+			Print(Context, TEXT(" Header+Norm  "), FontSilver); Print(Context, HeaderSize,									FontSilver, 3, 0); Print(Context, TEXT("bytes"), FontSilver); Newline(Context, RectMax);
+			Print(Context, TEXT(" BSDFs        "), FontSilver); Print(Context, BSDFsSize,									FontSilver, 3, 0); Print(Context, TEXT("bytes"), FontSilver); Newline(Context, RectMax);
 		}
-		Print(Context, TEXT(" Total        "), FontOrange); Print(Context, TotalSize,  FontOrange, 3, 0); Print(Context, TEXT("bytes"), FontOrange); Newline(Context);
+		Print(Context, TEXT(" Total        "), FontOrange); Print(Context, TotalSize,  FontOrange, 3, 0); Print(Context, TEXT("bytes"), FontOrange); Newline(Context, RectMax);
+		PrintBackgroundRect(Context, RectMin, RectMax);
 	}
 }
 
