Strata - Fixed spec/diffuse error from a previous improvement.

diff --git a/Engine/Shaders/Private/Strata/Strata.ush b/Engine/Shaders/Private/Strata/Strata.ush
--- a/Engine/Shaders/Private/Strata/Strata.ush
+++ b/Engine/Shaders/Private/Strata/Strata.ush
@@ -2742,7 +2742,7 @@ uint PackR7G7B6Gamma2(float3 rgb, float Dither)
 	const float3 Dither776 = float3(Dither7bits, Dither7bits, Dither6bits);
 
 	// Dither value to hide more aggressive bit reduction.
-	const float3 dithered_rgb = DitherIfNonBlack(sqrt(saturate(rgb.r)), Dither776);
+	const float3 dithered_rgb = DitherIfNonBlack(sqrt(saturate(rgb)), Dither776);
 
 	uint r = (uint(dithered_rgb.r * 127.0f) << 0);
 	uint g = (uint(dithered_rgb.g * 127.0f) << 7);
