Strata - Made second roughness more controlable with roughness + weight.
Renamed second roughness from haziness (more will be done later).
This prepares the upcoming CL that will add a fast clear coat path (legacy).

diff --git a/Engine/Shaders/Private/DebugProbes.usf b/Engine/Shaders/Private/DebugProbes.usf
--- a/Engine/Shaders/Private/DebugProbes.usf
+++ b/Engine/Shaders/Private/DebugProbes.usf
@@ -304,7 +304,8 @@ void MainPS(
 			/*float SSSPhaseAniso					*/0.0f,
 			/*float UseSSSDiffuse					*/0.0f,
 			/*float3 Emissive						*/0.0f,
-			/*float Haziness						*/0.0f,
+			/*float SecondRoughness					*/0.0f,
+			/*float SecondRoughnessWeight			*/0.0f,
 			/*float FuzzAmount						*/0.0f,
 			/*float3 FuzzColor						*/0.0f,
 			/*float Thickness						*/0.0f,
diff --git a/Engine/Shaders/Private/DiffuseIndirectComposite.usf b/Engine/Shaders/Private/DiffuseIndirectComposite.usf
--- a/Engine/Shaders/Private/DiffuseIndirectComposite.usf
+++ b/Engine/Shaders/Private/DiffuseIndirectComposite.usf
@@ -256,9 +256,9 @@ FStrataDeferredLighting StrataLumenLighting(
 		#if STRATA_FASTPATH==0
 		if (BSDF_GETHASHAZINESS(BSDF))
 		{
-			const float HazeRoughness = StrataComputeHazeRoughness(Roughness);
-			const float3 InHazeGlossyLighting = MixSpecularAndRoughReflections(HazeRoughness, InSpecularLighting, InRoughSpecularLighting);
-			OutBSDF.Specular += BSDFThroughput * InHazeGlossyLighting * StrataEnvLight.SpecularHazeWeight * Occlusion.SpecularOcclusion;
+			FHaziness Haziness = UnpackHaziness(SLAB_HAZINESS(BSDF));
+			const float3 InHazeGlossyLighting = MixSpecularAndRoughReflections(Haziness.Roughness, InSpecularLighting, InRoughSpecularLighting);
+			OutBSDF.Specular += BSDFThroughput * InHazeGlossyLighting * Haziness.Weight * Occlusion.SpecularOcclusion;
 		}
 		#endif 
 
diff --git a/Engine/Shaders/Private/Strata/Strata.ush b/Engine/Shaders/Private/Strata/Strata.ush
--- a/Engine/Shaders/Private/Strata/Strata.ush
+++ b/Engine/Shaders/Private/Strata/Strata.ush
@@ -414,7 +414,6 @@ struct FStrataBSDF
 #define SLAB_F0(X)						X.VGPRs[1].xyz
 #define SLAB_ANISOTROPY(X)				X.VGPRs[1].w
 #define SLAB_F90(X)						X.VGPRs[2].xyz
-// Haziness
 #define SLAB_HAZINESS(X)				X.VGPRs[2].w
 // Either with SSS profile
 #define SLAB_SSSPROFILEID(X)			X.VGPRs[3].x
@@ -485,6 +484,58 @@ uint   StrataPackNormal24(in float3 Normal);
 float3 StrataUnpackNormal24(uint PackedNormal);
 bool   IsStrataSlabFastPathCompatible(in FStrataBSDF BSDF);
 
+///////////////////////////////////////////////////////////////////////////////
+// Haziness
+
+// bFastClearCoat==false means that Weight is the blend in amount of the secondary specular lob over the first specular lob.
+// bFastClearCoat==true  means fast clear coat is used, Weight is coverage of a top layer specular. 
+struct FHaziness
+{
+	bool  bFastClearCoat;
+	float Roughness;
+	float Weight;
+};
+
+FHaziness InitialiseHaziness()
+{
+	FHaziness Haziness = (FHaziness)0;
+	return Haziness;
+}
+
+// returns 16 bits: 1 bit bFastClearCoat, 7 bits Weight, 8 bits roughness
+float PackHaziness(FHaziness Haziness)
+{
+	uint High8Bits = (Haziness.bFastClearCoat ? 0x80 : 0x00) | uint(saturate(Haziness.Weight) * 127);
+	uint Data = (High8Bits << 8) | uint(saturate(Haziness.Roughness) * 255);
+	return asfloat(Data);
+}
+
+FHaziness UnpackHaziness(float FloatData)
+{
+	uint Data = asuint(FloatData);
+
+	FHaziness Haziness;
+	Haziness.bFastClearCoat = Data & 0x8000;
+	Haziness.Weight = float(uint((Data>>8) & 0x7F)) * (1.0f / 127.0f);
+	Haziness.Roughness = float(uint(Data & 0xFF)) * (1.0f / 255.0f);
+	return Haziness;
+}
+
+float LerpHazinessParameterBlending(FStrataBSDF A, FStrataBSDF B, float MixValue)
+{
+	FHaziness HazinessA = UnpackHaziness(SLAB_HAZINESS(A));
+	FHaziness HazinessB = UnpackHaziness(SLAB_HAZINESS(B));
+
+	FHaziness OutputHaziness;
+	OutputHaziness.Weight			= lerp(HazinessA.Weight,	HazinessB.Weight,		MixValue);
+	OutputHaziness.Roughness		= lerp(HazinessA.Roughness, HazinessB.Roughness,	MixValue);
+	OutputHaziness.bFastClearCoat	= HazinessA.bFastClearCoat || HazinessB.bFastClearCoat;		// Maintain the fast legacy path first
+
+	return PackHaziness(OutputHaziness);
+}
+
+//BSDF_SETHASHAZINESS(ResultBSDF, BSDF_GETHASHAZINESS(ForegroundBSDF) | BSDF_GETHASHAZINESS(BackgroundBSDF));	// TODO
+
 ///////////////////////////////////////////////////////////////////////////////
 // Irradiance and occlusion
 struct FStrataIrradianceAndOcclusion
@@ -821,55 +872,56 @@ bool StrataHasGreyScaleWeight(in float3 Weight)
 	return Weight.x == Weight.y && Weight.y == Weight.z;
 }
 
-void StrataComputeHazeRoughnessAndWeight(
-	float Roughness0,
-	float Haziness,
-	inout float Roughness1,
-	inout float MixLobe)
-{	
-	// Fade out Haziness as the roughness reaches 0.5f, as otherwise it create
-	//       27 |\
-	// Extent   | \
-	//          |  \
-	//        3 |    ----
-	//          0  0.25  1
-	//           Rougness
-
-	// Lerp factor for adapting haziness parameters
-	const float S2 = saturate(2 * Roughness0);
-	const float S4 = saturate(4 * Roughness0);
-
-	// For low roughness, the haze extent is larger, and the core highlight is lower. 
-	// This allows to emphase the haze effect at low roughness, and fade out haze as roughness increase.
-	const float HazeExtent = lerp(27.0f, 3.0f, S4);
-	const float HazeCore   = lerp(0.01f, 0.5f, S2);
-
-	// Reference "A Composite BRDF Model for Hazy Gloss"
-	 
-	// Compute wide roughness
-	float alpha_n = Square(Roughness0); 
-	float lambda_h = HazeExtent;
-	const float alpha_w = alpha_n * (1 + lambda_h); // Eq.8
-	Roughness1 = sqrtFast(alpha_w);
-
-	// Compute mix weights
-	const float beta_h = Haziness;
-	const float r_c = HazeCore;
-	const float w = alpha_w;
-
-	const float p = 1.0 / Square(1 + lambda_h); // Table 2, GGX
-	const float b = 2 * (r_c * (1 - w) + w * p);
-	const float D = Square(b) - 4 * (b - 1) * r_c;
-	const float u = (b - sqrtFast(D)) / (2 * b - 2); // Eq. 12, 13, 14
-	const float u2 = u * u;
-	const float w2 = 2 * w;
-	const float k_h = ((u - u2) * w2 * beta_h) / (Square(1 - u) + (u - u2) * w2 + u2); // Eq. 15
-	const float r = r_c + (1 - p) * k_h; // Eq. 5
-	MixLobe = k_h / r; // Eq. 4
-}
-
-float StrataComputeHazeRoughness(float Roughness0) { float Roughness1=0; float MixLobe=0; StrataComputeHazeRoughnessAndWeight(Roughness0, 1.0f, Roughness1, MixLobe); return Roughness1; }
-float StrataComputeHazeWeight(float Roughness0, float Haziness)    { float Roughness1=0; float MixLobe=0; StrataComputeHazeRoughnessAndWeight(Roughness0, Haziness, Roughness1, MixLobe); return MixLobe; }
+// STRATA_TODO: move this code to be a Haziness node. So keeping this code here for now.
+////////void StrataComputeHazeRoughnessAndWeight(
+////////	float Roughness0,
+////////	float Haziness,
+////////	inout float Roughness1,
+////////	inout float MixLobe)
+////////{	
+////////	// Fade out Haziness as the roughness reaches 0.5f, as otherwise it create
+////////	//       27 |\
+////////	// Extent   | \
+////////	//          |  \
+////////	//        3 |    ----
+////////	//          0  0.25  1
+////////	//           Rougness
+////////
+////////	// Lerp factor for adapting haziness parameters
+////////	const float S2 = saturate(2 * Roughness0);
+////////	const float S4 = saturate(4 * Roughness0);
+////////
+////////	// For low roughness, the haze extent is larger, and the core highlight is lower. 
+////////	// This allows to emphase the haze effect at low roughness, and fade out haze as roughness increase.
+////////	const float HazeExtent = lerp(27.0f, 3.0f, S4);
+////////	const float HazeCore   = lerp(0.01f, 0.5f, S2);
+////////
+////////	// Reference "A Composite BRDF Model for Hazy Gloss"
+////////	 
+////////	// Compute wide roughness
+////////	float alpha_n = Square(Roughness0); 
+////////	float lambda_h = HazeExtent;
+////////	const float alpha_w = alpha_n * (1 + lambda_h); // Eq.8
+////////	Roughness1 = sqrtFast(alpha_w);
+////////
+////////	// Compute mix weights
+////////	const float beta_h = Haziness;
+////////	const float r_c = HazeCore;
+////////	const float w = alpha_w;
+////////
+////////	const float p = 1.0 / Square(1 + lambda_h); // Table 2, GGX
+////////	const float b = 2 * (r_c * (1 - w) + w * p);
+////////	const float D = Square(b) - 4 * (b - 1) * r_c;
+////////	const float u = (b - sqrtFast(D)) / (2 * b - 2); // Eq. 12, 13, 14
+////////	const float u2 = u * u;
+////////	const float w2 = 2 * w;
+////////	const float k_h = ((u - u2) * w2 * beta_h) / (Square(1 - u) + (u - u2) * w2 + u2); // Eq. 15
+////////	const float r = r_c + (1 - p) * k_h; // Eq. 5
+////////	MixLobe = k_h / r; // Eq. 4
+////////}
+////////
+////////float StrataComputeHazeRoughness(float Roughness0) { float Roughness1=0; float MixLobe=0; StrataComputeHazeRoughnessAndWeight(Roughness0, 1.0f, Roughness1, MixLobe); return Roughness1; }
+////////float StrataComputeHazeWeight(float Roughness0, float Haziness)    { float Roughness1=0; float MixLobe=0; StrataComputeHazeRoughnessAndWeight(Roughness0, Haziness, Roughness1, MixLobe); return MixLobe; }
 
 ///////////////////////////////////////////////////////////////////////////////
 // Material & BSDF description
@@ -1190,7 +1242,7 @@ FStrataData GetStrataSlabBSDF(
 	float Roughness, float Anisotropy,
 	float SSSProfileID, float3 SSSMFP, float SSSMFPScale, float SSSPhaseAnisotropy, float UseSSSDiffusion,
 	float3 Emissive, 
-	float Haziness, 
+	float SecondRoughness, float SecondRoughnessWeight,
 	float FuzzAmount, float3 FuzzColor,
 	float Thickness,
 	uint SharedLocalBasisIndex, inout uint SharedLocalBasisTypes)
@@ -1206,6 +1258,8 @@ FStrataData GetStrataSlabBSDF(
 	bool bHasSSSProfile = !bIsThin && bSSSProfilePlugged && SSSMFPScale > 0.f;	// When no profile is provide, we use Burley. STRATA_TODO we should clarify as burley is also an option of the profile.
 	bool bHasSSS = bMFPPlugged || bHasSSSProfile || bHasSSSWrap;
 
+	FHaziness Haziness = InitialiseHaziness();
+
 	// Derive roughness and haziness from profile data
 	if (bSSSProfilePlugged)
 	{
@@ -1214,24 +1268,22 @@ FStrataData GetStrataSlabBSDF(
 
 		float LobeRoughness0 = 0;
 		float LobeRoughness1 = 0;
-		float LobeMix = 0;
 
-		GetSubsurfaceProfileDualSpecular(SubsurfaceProfileUInt, Roughness, SSSMFPScale, LobeRoughness0, LobeRoughness1, LobeMix);
+		GetSubsurfaceProfileDualSpecular(SubsurfaceProfileUInt, Roughness, SSSMFPScale, LobeRoughness0, LobeRoughness1, Haziness.Weight);
 
 		// Convertion from roughness0/1 to roughness+haziness
 		// * This encoding is based on how roughness0 + haziness is translated to roughness1, see ComputeHazyLobeRoughness
 		// * The mix factor is hardoded to 0.5 i.e., 0.5 * Fs(roughness0) + 0.5 * Fs(roughness1)
-		Haziness = (LobeRoughness1 - LobeRoughness0) / (1.f - min(LobeRoughness0, 0.99f));
+		Haziness.Roughness = (LobeRoughness1 - LobeRoughness0) / (1.f - min(LobeRoughness0, 0.99f));
 		Roughness = LobeRoughness0;
 
 		SSSMFP = GetSubsurfaceProfileMFPInCm(SubsurfaceProfileUInt).xyz * SSSMFPScale;
 	}
 	else
 	{
-		// Convert Haziness into a mix weight
-		// Haze roughness is inexpensive to compute, unlike mix factor. This is why we comput its value here, and 
-		// store it in place of the original haziness value. This saves computation time at runtime.
-		Haziness = Haziness > 0 ? StrataComputeHazeWeight(Roughness, Haziness) : 0.f;
+		Haziness.Roughness = SecondRoughness;
+		Haziness.Weight = SecondRoughnessWeight;
+		Haziness.bFastClearCoat = false;
 	}
 
 	// If MFP is smaller than the pixel's footprint, remove SSS attributes
@@ -1253,7 +1305,7 @@ FStrataData GetStrataSlabBSDF(
 		Thickness = STRATA_SIMPLEVOLUME_THICKNESS_CM;
 	}
 
-	const uint bHasHaziness = Haziness > 0.0f ? 1 : 0;
+	const uint bHasHaziness = ((Haziness.Roughness != Roughness  && Haziness.Weight > 0.0f) || Haziness.bFastClearCoat) ? 1 : 0;
 	const uint bHasAnisotropy = Anisotropy != 0 ? 1 : 0;
 	const uint bHasFuzz = FuzzAmount > 0.0f ? 1 : 0;
 	if (bHasAnisotropy)
@@ -1302,7 +1354,8 @@ FStrataData GetStrataSlabBSDF(
 
 	SLAB_ROUGHNESS			(StrataData.InlinedBSDF)		= Roughness;
 	SLAB_ANISOTROPY			(StrataData.InlinedBSDF)		= Anisotropy;
-	SLAB_HAZINESS			(StrataData.InlinedBSDF)		= Haziness;
+
+	SLAB_HAZINESS(StrataData.InlinedBSDF)					= PackHaziness(Haziness);
 
 	if (bHasSSSProfile)
 	{
@@ -1955,8 +2008,9 @@ FStrataData StrataHorizontalMixingParameterBlending(FStrataData Background, FStr
 	BSDF_SETTHICKNESSCM(ResultBSDF,		lerp(BSDF_GETTHICKNESSCM(BackgroundBSDF),		BSDF_GETTHICKNESSCM(ForegroundBSDF),		ForegroundOtherMixFactor));
 	SLAB_ROUGHNESS(ResultBSDF) =		lerp(SLAB_ROUGHNESS(BackgroundBSDF),			SLAB_ROUGHNESS(ForegroundBSDF),				ForegroundOtherMixFactor);
 	SLAB_ANISOTROPY(ResultBSDF) =		lerp(SLAB_ANISOTROPY(BackgroundBSDF),			SLAB_ANISOTROPY(ForegroundBSDF),			ForegroundOtherMixFactor);
-	SLAB_HAZINESS(ResultBSDF) =			lerp(SLAB_HAZINESS(BackgroundBSDF),				SLAB_HAZINESS(ForegroundBSDF),				ForegroundOtherMixFactor);
 	SLAB_FUZZ_AMOUNT(ResultBSDF) =		lerp(SLAB_FUZZ_AMOUNT(BackgroundBSDF),			SLAB_FUZZ_AMOUNT(ForegroundBSDF),			ForegroundOtherMixFactor);
+	SLAB_HAZINESS(ResultBSDF) =			LerpHazinessParameterBlending(BackgroundBSDF,	ForegroundBSDF,								ForegroundOtherMixFactor);
+	
 
 	if (BSDF_GETSSSTYPE(ResultBSDF) == SSS_TYPE_DIFFUSION_PROFILE)
 	{
@@ -2050,8 +2104,8 @@ FStrataData StrataAddParameterBlending(FStrataData A, FStrataData B, float AMixF
 	BSDF_SETTHICKNESSCM(ResultBSDF,		lerp(BSDF_GETTHICKNESSCM(BBSDF),				BSDF_GETTHICKNESSCM(ABSDF),					AMixFactor));
 	SLAB_ROUGHNESS(ResultBSDF) =		lerp(SLAB_ROUGHNESS(BBSDF),						SLAB_ROUGHNESS(ABSDF),						AMixFactor);
 	SLAB_ANISOTROPY(ResultBSDF) =		lerp(SLAB_ANISOTROPY(BBSDF),					SLAB_ANISOTROPY(ABSDF),						AMixFactor);
-	SLAB_HAZINESS(ResultBSDF) =			lerp(SLAB_HAZINESS(BBSDF),						SLAB_HAZINESS(ABSDF),						AMixFactor);
 	SLAB_FUZZ_AMOUNT(ResultBSDF) =		lerp(SLAB_FUZZ_AMOUNT(BBSDF),					SLAB_FUZZ_AMOUNT(ABSDF),					AMixFactor);
+	SLAB_HAZINESS(ResultBSDF) =			LerpHazinessParameterBlending(BBSDF,			ABSDF,										AMixFactor);
 
 	if (BSDF_GETSSSTYPE(ResultBSDF) == SSS_TYPE_DIFFUSION_PROFILE)
 	{
@@ -2158,8 +2212,8 @@ FStrataBSDF StrataVerticalLayeringParameterBlending(in FStrataBSDF TopBSDF, in F
 	SLAB_FUZZ_COLOR(ResultBSDF) =		saturate(BottomColorFactor * SLAB_FUZZ_COLOR(BaseBSDF)						+ TopFactor * SLAB_FUZZ_COLOR(TopBSDF));
 		
 	SLAB_ROUGHNESS(ResultBSDF) =		lerp(SLAB_ROUGHNESS(TopBSDF),					SLAB_ROUGHNESS(BaseBSDF),	TopToBottomLerp);
-	SLAB_HAZINESS(ResultBSDF) =			lerp(SLAB_HAZINESS(TopBSDF),					SLAB_HAZINESS(BaseBSDF),	TopToBottomLerp);
 	SLAB_FUZZ_AMOUNT(ResultBSDF) =		lerp(SLAB_FUZZ_AMOUNT(TopBSDF),					SLAB_FUZZ_AMOUNT(BaseBSDF), TopToBottomLerp);
+	SLAB_HAZINESS(ResultBSDF) =			LerpHazinessParameterBlending(TopBSDF,			BaseBSDF,					TopToBottomLerp);
 
 	// Anisotropy is disabled with vertical blending. 
 	SLAB_ANISOTROPY(ResultBSDF) =		0.f;
@@ -2280,7 +2334,7 @@ void SanitizeStrataSlab(inout FStrataBSDF BSDF)
 	SLAB_F90(BSDF) = saturate(SLAB_F90(BSDF));
 	SLAB_ROUGHNESS(BSDF) = SanitizeRoughness(SLAB_ROUGHNESS(BSDF));
 	SLAB_ANISOTROPY(BSDF) = clamp(SLAB_ANISOTROPY(BSDF), -1.f, 1.f);
-	SLAB_HAZINESS(BSDF) = saturate(SLAB_HAZINESS(BSDF));
+	// SLAB_HAZINESS is sanitize when packed, see PackHaziness.
 #if STRATA_INLINE_SHADING
 	BSDF.TmpMFP = clamp(BSDF.TmpMFP, 0.0f, Max111110BitsFloat3);
 #endif
@@ -3042,7 +3096,8 @@ void StrataConvertFromDBuffer(
 		0.f,				// SSSPhaseAniso
 		0.f,				// UseSSSDiffusion
 		0.f,				// Emissive
-		0.f,				// Haziness
+		0.f,				// SecondRoughness
+		0.f,				// SecondRoughnessWeight
 		0.f,				// FuzzAmount
 		0.f,				// FuzzColor
 		DefaultThickness,
@@ -3360,7 +3415,8 @@ FStrataBSDF UnpackStrataBSDFIn(FStrataMaterialContainer StrataBuffer, inout FStr
 		const bool bHasF90 = BSDF_GETHASF90(OutBSDF);
 		if (bHasF90 || BSDF_GETHASHAZINESS(OutBSDF))
 		{
-			float4 Data = UnpackRGBA8(StrataLoadUint1(StrataBuffer, StrataAddressing));
+			uint RawData = StrataLoadUint1(StrataBuffer, StrataAddressing);
+			float2 Data = UnpackRGBA8(RawData).xy;
 
 			const float F90Co = Data.x;
 			const float F90Cg = Data.y;
@@ -3374,7 +3430,7 @@ FStrataBSDF UnpackStrataBSDFIn(FStrataMaterialContainer StrataBuffer, inout FStr
 			F90 = Divisor > 0.0f ? F90 / Divisor : F90;
 			SLAB_F90(OutBSDF) = F90;
 
-			SLAB_HAZINESS(OutBSDF) = Data.w;
+			SLAB_HAZINESS(OutBSDF) = asfloat(RawData >> 16); // Keeping the top 16 bits reinterpreted as float.
 		}
 
 		const bool bIsSimpleVolume = BSDF_GETISSIMPLEVOLUME(OutBSDF);
diff --git a/Engine/Shaders/Private/Strata/StrataEvaluation.ush b/Engine/Shaders/Private/Strata/StrataEvaluation.ush
--- a/Engine/Shaders/Private/Strata/StrataEvaluation.ush
+++ b/Engine/Shaders/Private/Strata/StrataEvaluation.ush
@@ -315,8 +315,9 @@ FStrataEvaluateResult StrataEvaluateBSDFCommon(FStrataBSDFContext BSDFContext, F
 			#if STRATA_FASTPATH==0
 				if (bHaziness)
 				{
-					HazeWeight = SLAB_HAZINESS(BSDFContext.BSDF);
-					HazeSafeRoughness = MakeRoughnessSafe(StrataComputeHazeRoughness(SafeRoughness));
+					FHaziness Haziness = UnpackHaziness(SLAB_HAZINESS(BSDFContext.BSDF));
+					HazeWeight = Haziness.Weight;
+					HazeSafeRoughness = MakeRoughnessSafe(Haziness.Roughness);
 
 					FBxDFEnergyTerms EnergyTerms = ComputeGGXSpecEnergyTerms(HazeSafeRoughness, BSDFContext.Context.NoV, F0, F90);
 					HazeMSScale = ComputeEnergyConservation(EnergyTerms);
@@ -828,11 +829,12 @@ FStrataImportanceSampleResult StrataImportanceSampleBSDF(FStrataBSDFContext BSDF
 			float SpecularRoughness = SafeRoughness;
 			if (Out.Term == SHADING_TERM_SPECULAR && bHaziness)
 			{
-				const float HazeWeight = SLAB_HAZINESS(BSDFContext.BSDF);
-				const float HazeRoughness = MakeRoughnessSafe(StrataComputeHazeRoughness(SafeRoughness));
+				FHaziness Haziness = UnpackHaziness(SLAB_HAZINESS(BSDFContext.BSDF));
+				const float HazeWeight = Haziness.Weight;
+				const float HazeSafeRoughness = MakeRoughnessSafe(Haziness.Roughness);
 
 				SpecularLobePDF   = 1.f - HazeWeight;
-				SpecularRoughness = StrataSelectLobe(E.x, SpecularLobePDF, SafeRoughness, HazeRoughness);
+				SpecularRoughness = StrataSelectLobe(E.x, SpecularLobePDF, SafeRoughness, HazeSafeRoughness);
 			}
 
 			// Sample Diffuse lobe
@@ -1012,8 +1014,9 @@ FStrataEnvLightResult StrataEvaluateForEnvLight(FStrataBSDFContext BSDFContext,
 				const bool bHasHaziness = BSDF_GETHASHAZINESS(BSDFContext.BSDF);
 				if (bHasHaziness)
 				{
-					const float SafeRoughnessHaze = MakeRoughnessSafe(StrataComputeHazeRoughness(SafeRoughness));
-					const float HazeWeight = SLAB_HAZINESS(BSDFContext.BSDF);
+					FHaziness Haziness = UnpackHaziness(SLAB_HAZINESS(BSDFContext.BSDF));
+					const float HazeWeight = Haziness.Weight;
+					const float HazeSafeRoughness = MakeRoughnessSafe(Haziness.Roughness);
 
 					// Smoothly fade in haziness while fading out SSR to make sure we do not add energy and avoid popping.
 					const float HazinessFadeIn = saturate(HazeWeight / 0.1f);
@@ -1023,10 +1026,10 @@ FStrataEnvLightResult StrataEvaluateForEnvLight(FStrataBSDFContext BSDFContext,
 					StrataEnvLightResult.SpecularWeight *= lerp(1.0f, (1.0f - HazeWeight), HazinessFadeIn);
 
 					// Compute the directional albedo for hazy specular
-					FBxDFEnergyTerms EnergyTerms = ComputeGGXSpecEnergyTerms(SafeRoughnessHaze, BSDFContext.Context.NoV, F0, F90);
+					FBxDFEnergyTerms EnergyTerms = ComputeGGXSpecEnergyTerms(HazeSafeRoughness, BSDFContext.Context.NoV, F0, F90);
 					const float3 EvalEnvBRDFHaze = EnergyTerms.E;
 					StrataEnvLightResult.SpecularHazeWeight = lerp(0.0f, HazeWeight, HazinessFadeIn) * EvalEnvBRDFHaze * SpecularMicroOcclusion;
-					StrataEnvLightResult.SpecularHazeSafeRoughness = SafeRoughnessHaze;
+					StrataEnvLightResult.SpecularHazeSafeRoughness = HazeSafeRoughness;
 				}
 				#endif
 			}
diff --git a/Engine/Shaders/Private/Strata/StrataExport.ush b/Engine/Shaders/Private/Strata/StrataExport.ush
--- a/Engine/Shaders/Private/Strata/StrataExport.ush
+++ b/Engine/Shaders/Private/Strata/StrataExport.ush
@@ -463,7 +463,9 @@ void PackStrataOut(
 								F90 = Divisor > 0.0f ? F90 / Divisor : 1.0f;
 								float3 F90YCoCg = LinearRGB_2_NormalisedYCoCg(F90);
 
-								STRATA_STORE_UINT1(PackRGBA8(float4(F90YCoCg.y, F90YCoCg.z, 0, SLAB_HAZINESS(BSDF))));
+								uint F90Data = PackRGBA8(float4(F90YCoCg.y, F90YCoCg.z, 0, 0));
+								uint HazinessData = asuint(SLAB_HAZINESS(BSDF));
+								STRATA_STORE_UINT1(HazinessData << 16 | F90Data);
 							}
 
 							const bool bIsSimpleVolume = BSDF_GETISSIMPLEVOLUME(BSDF);
diff --git a/Engine/Shaders/Private/Strata/StrataLegacyConversion.ush b/Engine/Shaders/Private/Strata/StrataLegacyConversion.ush
--- a/Engine/Shaders/Private/Strata/StrataLegacyConversion.ush
+++ b/Engine/Shaders/Private/Strata/StrataLegacyConversion.ush
@@ -95,7 +95,8 @@ FStrataData CreateLegacySlab(
 			In.SSSPhaseAniso,	// SSSPhaseAniso
 			In.UseSSSDiffusion, // UseSSSDiffusion
 			In.Emissive,		// Emissive
-			0.f,				// Haziness
+			0.f,				// SecondRoughness
+			0.f,				// SecondRoughnessWeight
 			In.FuzzAmount,		// FuzzAmount
 			In.FuzzColor,		// FuzzColor
 			In.Thickness,
diff --git a/Engine/Shaders/Private/Strata/StrataLightingCommon.ush b/Engine/Shaders/Private/Strata/StrataLightingCommon.ush
--- a/Engine/Shaders/Private/Strata/StrataLightingCommon.ush
+++ b/Engine/Shaders/Private/Strata/StrataLightingCommon.ush
@@ -184,7 +184,11 @@ void StrataEnvLightingCommon(
 	// Specular component
 	const bool bIsTopLayer = BSDF_GETISTOPLAYER(BSDF);
 	SpecularLighting = 0;
+#if STRATA_FASTPATH==0
+	if (any((StrataEnvLight.SpecularWeight + StrataEnvLight.SpecularHazeWeight) > 0.0f))
+#else
 	if (any(StrataEnvLight.SpecularWeight > 0.0f))
+#endif
 	{
 		float IndirectIrradiance = StrataIrradianceAndOcclusion.IndirectIrradiance;
 	#if ENABLE_SKY_LIGHT && ALLOW_STATIC_LIGHTING
diff --git a/Engine/Shaders/Private/Strata/StrataVisualize.usf b/Engine/Shaders/Private/Strata/StrataVisualize.usf
--- a/Engine/Shaders/Private/Strata/StrataVisualize.usf
+++ b/Engine/Shaders/Private/Strata/StrataVisualize.usf
@@ -222,65 +222,70 @@ void PrintBSDF(inout FShaderPrintContext Context, inout FStrataAddressing Strata
 			const bool bGreyF0 = SLAB_F0(BSDF).x == SLAB_F0(BSDF).y && SLAB_F0(BSDF).x == SLAB_F0(BSDF).z;
 			if (bGreyF0)
 			{
-			Print(Context, TEXT(" F0          "), FontBSDFPropName);	Print(Context, SLAB_F0(BSDF).x, FontBSDFPropValu);					Newline(Context, RectMax);
+				Print(Context, TEXT(" F0          "), FontBSDFPropName);	Print(Context, SLAB_F0(BSDF).x, FontBSDFPropValu);				Newline(Context, RectMax);
 			}
 			else
 			{
-			Print(Context, TEXT(" F0          "), FontBSDFPropName);	Print(Context, SLAB_F0(BSDF), FontBSDFPropValu);					Newline(Context, RectMax);
+				Print(Context, TEXT(" F0          "), FontBSDFPropName);	Print(Context, SLAB_F0(BSDF), FontBSDFPropValu);				Newline(Context, RectMax);
 			}
 
 
 			if (BSDF_GETHASF90(BSDF))
 			{
-			const bool bGreyF90 = SLAB_F90(BSDF).x == SLAB_F90(BSDF).y && SLAB_F90(BSDF).x == SLAB_F90(BSDF).z;
-			if (bGreyF90)
-			{
-			Print(Context, TEXT(" F90         "), FontBSDFPropName);	Print(Context, SLAB_F90(BSDF).x, FontBSDFPropValu);					Newline(Context, RectMax);
-			}
-			else
-			{
-			Print(Context, TEXT(" F90         "), FontBSDFPropName);	Print(Context, SLAB_F90(BSDF), FontBSDFPropValu);					Newline(Context, RectMax);
-			}
+				const bool bGreyF90 = SLAB_F90(BSDF).x == SLAB_F90(BSDF).y && SLAB_F90(BSDF).x == SLAB_F90(BSDF).z;
+				if (bGreyF90)
+				{
+					Print(Context, TEXT(" F90         "), FontBSDFPropName);	Print(Context, SLAB_F90(BSDF).x, FontBSDFPropValu);			Newline(Context, RectMax);
+				}
+				else
+				{
+					Print(Context, TEXT(" F90         "), FontBSDFPropName);	Print(Context, SLAB_F90(BSDF), FontBSDFPropValu);			Newline(Context, RectMax);
+				}
 			}
 
 			Print(Context, TEXT(" Roughness   "), FontBSDFPropName);	Print(Context, SLAB_ROUGHNESS(BSDF), FontBSDFPropValu);				Newline(Context, RectMax);
 
 			if (BSDF_GETHASANISOTROPY(BSDF))
 			{
-			Print(Context, TEXT(" Anisotropy  "), FontBSDFPropName);	Print(Context, SLAB_ANISOTROPY(BSDF), FontBSDFPropValu);			Newline(Context, RectMax);
+				Print(Context, TEXT(" Anisotropy  "), FontBSDFPropName);	Print(Context, SLAB_ANISOTROPY(BSDF), FontBSDFPropValu);		Newline(Context, RectMax);
 			}
 
 			if (BSDF_GETHASHAZINESS(BSDF))
 			{
-			Print(Context, TEXT(" Haziness    "), FontBSDFPropName);	Print(Context, SLAB_HAZINESS(BSDF), FontBSDFPropValu);				Newline(Context, RectMax);
+				FHaziness Haziness = UnpackHaziness(SLAB_HAZINESS(BSDF));
+
+				Print(Context, TEXT(" Haziness    "),		FontBSDFPropName);																		Newline(Context, RectMax);
+				Print(Context, TEXT("     Roughness    "),	FontBSDFPropName);	Print(Context, Haziness.Roughness,				FontBSDFPropValu);	Newline(Context, RectMax);
+				Print(Context, TEXT("     Weight       "),	FontBSDFPropName);	Print(Context, Haziness.Weight,					FontBSDFPropValu);	Newline(Context, RectMax);
+				Print(Context, TEXT("     FastClearCoat"),	FontBSDFPropName);	Print(Context, Haziness.bFastClearCoat ?1u:0u,	FontBSDFPropValu);	Newline(Context, RectMax);
 			}
-			
+
 			if (BSDF_GETSSSTYPE(BSDF) != SSS_TYPE_INVALID || BSDF_GETISSIMPLEVOLUME(BSDF) || BSDF_GETISTHIN(BSDF))
 			{
 
-			const FStrataSubsurfaceHeader SSSHeader = StrataLoadSubsurfaceHeader(Strata.SSSTexture, StrataAddressing.PixelCoords);
-			const uint SSSType = StrataSubSurfaceHeaderGetSSSType(SSSHeader);
-			const float Opacity = StrataSubSurfaceHeaderGetWrapOpacity(SSSHeader);
-			Print(Context, TEXT(" SSS Type    "), FontBSDFPropName);	PrintSSSType(Context, SSSType, FontYellow);	Newline(Context, RectMax);
-			Print(Context, TEXT(" SSS Opacity "), FontBSDFPropName);	Print(Context, Opacity, FontBSDFPropValu);	Newline(Context, RectMax);
+				const FStrataSubsurfaceHeader SSSHeader = StrataLoadSubsurfaceHeader(Strata.SSSTexture, StrataAddressing.PixelCoords);
+				const uint SSSType = StrataSubSurfaceHeaderGetSSSType(SSSHeader);
+				const float Opacity = StrataSubSurfaceHeaderGetWrapOpacity(SSSHeader);
+				Print(Context, TEXT(" SSS Type    "), FontBSDFPropName);	PrintSSSType(Context, SSSType, FontYellow);	Newline(Context, RectMax);
+				Print(Context, TEXT(" SSS Opacity "), FontBSDFPropName);	Print(Context, Opacity, FontBSDFPropValu);	Newline(Context, RectMax);
 
-			if (BSDF_GETSSSTYPE(BSDF) == SSS_TYPE_DIFFUSION_PROFILE && BSDF_GETISSIMPLEVOLUME(BSDF) == 0)
-			{
-			Print(Context, TEXT(" SSS ID.     "), FontBSDFPropName);	Print(Context, SLAB_SSSPROFILEID(BSDF), FontBSDFPropValu);			Newline(Context, RectMax);
-			Print(Context, TEXT(" SSS Radius  "), FontBSDFPropName);	Print(Context, SLAB_SSSPROFILERADIUSSCALE(BSDF), FontBSDFPropValu);	Newline(Context, RectMax);
-			}
-			else
-			{
-			Print(Context, TEXT(" SSS MFP     "), FontBSDFPropName);	Print(Context, SLAB_SSSMFP(BSDF), FontBSDFPropValu);				Newline(Context, RectMax);
-			}
+				if (BSDF_GETSSSTYPE(BSDF) == SSS_TYPE_DIFFUSION_PROFILE && BSDF_GETISSIMPLEVOLUME(BSDF) == 0)
+				{
+					Print(Context, TEXT(" SSS ID.     "), FontBSDFPropName);	Print(Context, SLAB_SSSPROFILEID(BSDF), FontBSDFPropValu);			Newline(Context, RectMax);
+					Print(Context, TEXT(" SSS Radius  "), FontBSDFPropName);	Print(Context, SLAB_SSSPROFILERADIUSSCALE(BSDF), FontBSDFPropValu);	Newline(Context, RectMax);
+				}
+				else
+				{
+					Print(Context, TEXT(" SSS MFP     "), FontBSDFPropName);	Print(Context, SLAB_SSSMFP(BSDF), FontBSDFPropValu);				Newline(Context, RectMax);
+				}
 
-			Print(Context, TEXT(" SSS Phase   "), FontBSDFPropName);	Print(Context, SLAB_SSSPHASEANISOTROPY(BSDF), FontBSDFPropValu);	Newline(Context, RectMax);
+				Print(Context, TEXT(" SSS Phase   "), FontBSDFPropName);	Print(Context, SLAB_SSSPHASEANISOTROPY(BSDF), FontBSDFPropValu);		Newline(Context, RectMax);
 			} // SSS
 
 			if (BSDF_GETHASFUZZ(BSDF))
 			{
-			Print(Context, TEXT(" Fuzz Amount "), FontBSDFPropName);	Print(Context, SLAB_FUZZ_AMOUNT(BSDF), FontBSDFPropValu);			Newline(Context, RectMax);
-			Print(Context, TEXT(" Fuzz Color  "), FontBSDFPropName);	Print(Context, SLAB_FUZZ_COLOR(BSDF), FontBSDFPropValu);			Newline(Context, RectMax);
+				Print(Context, TEXT(" Fuzz Amount "), FontBSDFPropName);	Print(Context, SLAB_FUZZ_AMOUNT(BSDF), FontBSDFPropValu);				Newline(Context, RectMax);
+				Print(Context, TEXT(" Fuzz Color  "), FontBSDFPropName);	Print(Context, SLAB_FUZZ_COLOR(BSDF), FontBSDFPropValu);				Newline(Context, RectMax);
 			}
 		}
 		break;
diff --git a/Engine/Source/Runtime/Engine/Classes/Materials/MaterialExpressionStrata.h b/Engine/Source/Runtime/Engine/Classes/Materials/MaterialExpressionStrata.h
--- a/Engine/Source/Runtime/Engine/Classes/Materials/MaterialExpressionStrata.h
+++ b/Engine/Source/Runtime/Engine/Classes/Materials/MaterialExpressionStrata.h
@@ -279,10 +279,16 @@ class UMaterialExpressionStrataSlabBSDF : public UMaterialExpressionStrataBSDF
 	FExpressionInput EmissiveColor;
 
 	/**
-	 * Haziness controls the relative roughness of a second specular lobe. 0 means disabled and 1 means the second lobe specular lobe will lerp the current roughness to fully rough. (type = float, unitless, default = 0)
+	 * Controls the roughness of a secondary specular lobe. Roughness of 0 (smooth) is a mirror reflection and 1 (rough) is completely matte or diffuse. Does not influence diffuse roughness. (type = float, unit = unitless, defaults to 0.5)
 	 */
 	UPROPERTY()
-	FExpressionInput Haziness;
+	FExpressionInput SecondRoughness;
+
+	/**
+	 * The weight of the second specular lobe using SecondRoughness. The first specular using Roughness will have a weight of (1 - SecondRoughnessWeight). (type = float, unitless, default = 0)
+	 */
+	UPROPERTY()
+	FExpressionInput SecondRoughnessWeight;
 
 	/**
 	 * The slab thickness. (type = float, centimeters, default = 0.01 centimeter = 0.1 millimeter)
@@ -330,7 +336,7 @@ class UMaterialExpressionStrataSlabBSDF : public UMaterialExpressionStrataBSDF
 
 	bool HasEdgeColor() const;
 	bool HasFuzz() const;
-	bool HasHaziness() const;
+	bool HasSecondRoughness() const;
 	bool HasSSS() const;
 	bool HasSSSProfile() const;
 	bool HasMFPPluggedIn() const;
diff --git a/Engine/Source/Runtime/Engine/Private/Materials/HLSLMaterialTranslator.cpp b/Engine/Source/Runtime/Engine/Private/Materials/HLSLMaterialTranslator.cpp
--- a/Engine/Source/Runtime/Engine/Private/Materials/HLSLMaterialTranslator.cpp
+++ b/Engine/Source/Runtime/Engine/Private/Materials/HLSLMaterialTranslator.cpp
@@ -1957,9 +1957,9 @@ void FHLSLMaterialTranslator::GetMaterialEnvironment(EShaderPlatform InPlatform,
 					{
 						if (!It.IsDiscarded() && It.MaxDistanceFromLeaves == DistanceToLeaves)
 						{
-							StrataMaterialDescription += FString::Printf(TEXT("\tIdx=%d Op=%s ParentIdx=%d LeftIndex=%d RightIndex=%d BSDFIdx=%d LayerDepth=%d IsTop=%d IsBot=%d BSDFType=%s SSS=%d MFP=%d F90=%d Hazy=%d Fuzz=%d  Aniso=%d \r\n"),
+							StrataMaterialDescription += FString::Printf(TEXT("\tIdx=%d Op=%s ParentIdx=%d LeftIndex=%d RightIndex=%d BSDFIdx=%d LayerDepth=%d IsTop=%d IsBot=%d BSDFType=%s SSS=%d MFP=%d F90=%d Rough2=%d Fuzz=%d  Aniso=%d \r\n"),
 								It.Index, GetStrataOperatorStr(It.OperatorType), It.ParentIndex, It.LeftIndex, It.RightIndex, It.BSDFIndex, It.LayerDepth, It.bIsTop, It.bIsBottom, 
-								*GetStrataBSDFName(It.BSDFType), It.bBSDFHasSSS, It.bBSDFHasMFPPluggedIn, It.bBSDFHasEdgeColor, It.bBSDFHasHaziness, It.bBSDFHasFuzz, It.bBSDFHasAnisotropy);
+								*GetStrataBSDFName(It.BSDFType), It.bBSDFHasSSS, It.bBSDFHasMFPPluggedIn, It.bBSDFHasEdgeColor, It.bBSDFHasSecondRoughness, It.bBSDFHasFuzz, It.bBSDFHasAnisotropy);
 						}
 					}
 				}
@@ -10007,7 +10007,7 @@ bool FHLSLMaterialTranslator::StrataGenerateDerivedMaterialOperatorData()
 				{
 				case STRATA_BSDF_TYPE_SLAB:
 				{
-					bStrataMaterialIsSimple = bStrataMaterialIsSimple && !bMayHaveColoredWeight && !It.bBSDFHasAnisotropy && !It.bBSDFHasEdgeColor && !It.bBSDFHasFuzz && !It.bBSDFHasHaziness && !It.bBSDFHasMFPPluggedIn && !It.bBSDFHasSSS;
+					bStrataMaterialIsSimple = bStrataMaterialIsSimple && !bMayHaveColoredWeight && !It.bBSDFHasAnisotropy && !It.bBSDFHasEdgeColor && !It.bBSDFHasFuzz && !It.bBSDFHasSecondRoughness && !It.bBSDFHasMFPPluggedIn && !It.bBSDFHasSSS;
 					bStrataMaterialIsSingle = bStrataMaterialIsSingle && !bMayHaveColoredWeight && !It.bBSDFHasAnisotropy;
 					break;
 				}
@@ -10101,7 +10101,7 @@ bool FHLSLMaterialTranslator::StrataGenerateDerivedMaterialOperatorData()
 					StrataMaterialRequestedSizeByte += UintByteSize;
 					StrataMaterialRequestedSizeByte += UintByteSize;
 
-					if (It.bBSDFHasEdgeColor || It.bBSDFHasHaziness)
+					if (It.bBSDFHasEdgeColor || It.bBSDFHasSecondRoughness)
 					{
 						StrataMaterialRequestedSizeByte += UintByteSize;
 					}
@@ -10282,7 +10282,7 @@ int32 FHLSLMaterialTranslator::StrataSlabBSDF(
 	int32 Roughness, int32 Anisotropy,
 	int32 SSSProfileId, int32 SSSMFP, int32 SSSMFPScale, int32 SSSPhaseAniso, int32 UseSSSDiffusion,
 	int32 EmissiveColor,
-	int32 Haziness, 
+	int32 SecondRoughness, int32 SecondRoughnessWeight,
 	int32 FuzzAmount, int32 FuzzColor,
 	int32 Thickness, 
 	int32 Normal, int32 Tangent, const FString& SharedLocalBasisIndexMacro,
@@ -10299,27 +10299,28 @@ int32 FHLSLMaterialTranslator::StrataSlabBSDF(
 			return INDEX_NONE;
 		}
 		return AddCodeChunk(
-			MCT_Strata, TEXT("PromoteParameterBlendedBSDFToOperator(GetStrataSlabBSDF(Parameters.StrataPixelFootprint, %s, %s, %s, %s, %s, %s, %s, %s, %s, %s, %s, %s, %s, %s, %s, %s, %s, %s, %s, %s, %s, Parameters.SharedLocalBases.Types) /* Normal = %s ; Tangent = %s */, Parameters.StrataTree, %u, %u, %u, %u)"),
-			*StrataGetCastParameterCode(UseMetalness,	MCT_Float),
-			*StrataGetCastParameterCode(BaseColor,		MCT_Float3),
-			*StrataGetCastParameterCode(EdgeColor,		MCT_Float3),
-			*StrataGetCastParameterCode(Specular,		MCT_Float),
-			*StrataGetCastParameterCode(Metallic,		MCT_Float),
-			*StrataGetCastParameterCode(DiffuseAlbedo,	MCT_Float3),
-			*StrataGetCastParameterCode(F0,				MCT_Float3),
-			*StrataGetCastParameterCode(F90,			MCT_Float3),
-			*StrataGetCastParameterCode(Roughness,		MCT_Float),
-			*StrataGetCastParameterCode(Anisotropy,		MCT_Float),
-			*StrataGetCastParameterCode(SSSProfileId,	MCT_Float),
-			*StrataGetCastParameterCode(SSSMFP,			MCT_Float3),
-			*StrataGetCastParameterCode(SSSMFPScale,	MCT_Float),
-			*StrataGetCastParameterCode(SSSPhaseAniso,  MCT_Float),
-			*StrataGetCastParameterCode(UseSSSDiffusion,MCT_Float),
-			*StrataGetCastParameterCode(EmissiveColor,	MCT_Float3),
-			*StrataGetCastParameterCode(Haziness,		MCT_Float),
-			*StrataGetCastParameterCode(FuzzAmount,		MCT_Float),
-			*StrataGetCastParameterCode(FuzzColor,		MCT_Float3),
-			*StrataGetCastParameterCode(Thickness,		MCT_Float),
+			MCT_Strata, TEXT("PromoteParameterBlendedBSDFToOperator(GetStrataSlabBSDF(Parameters.StrataPixelFootprint, %s, %s, %s, %s, %s, %s, %s, %s, %s, %s, %s, %s, %s, %s, %s, %s, %s, %s, %s, %s, %s, %s, Parameters.SharedLocalBases.Types) /* Normal = %s ; Tangent = %s */, Parameters.StrataTree, %u, %u, %u, %u)"),
+			*StrataGetCastParameterCode(UseMetalness,			MCT_Float),
+			*StrataGetCastParameterCode(BaseColor,				MCT_Float3),
+			*StrataGetCastParameterCode(EdgeColor,				MCT_Float3),
+			*StrataGetCastParameterCode(Specular,				MCT_Float),
+			*StrataGetCastParameterCode(Metallic,				MCT_Float),
+			*StrataGetCastParameterCode(DiffuseAlbedo,			MCT_Float3),
+			*StrataGetCastParameterCode(F0,						MCT_Float3),
+			*StrataGetCastParameterCode(F90,					MCT_Float3),
+			*StrataGetCastParameterCode(Roughness,				MCT_Float),
+			*StrataGetCastParameterCode(Anisotropy,				MCT_Float),
+			*StrataGetCastParameterCode(SSSProfileId,			MCT_Float),
+			*StrataGetCastParameterCode(SSSMFP,					MCT_Float3),
+			*StrataGetCastParameterCode(SSSMFPScale,			MCT_Float),
+			*StrataGetCastParameterCode(SSSPhaseAniso,			MCT_Float),
+			*StrataGetCastParameterCode(UseSSSDiffusion,		MCT_Float),
+			*StrataGetCastParameterCode(EmissiveColor,			MCT_Float3),
+			*StrataGetCastParameterCode(SecondRoughness,		MCT_Float),
+			*StrataGetCastParameterCode(SecondRoughnessWeight,	MCT_Float),
+			*StrataGetCastParameterCode(FuzzAmount,				MCT_Float),
+			*StrataGetCastParameterCode(FuzzColor,				MCT_Float3),
+			*StrataGetCastParameterCode(Thickness,				MCT_Float),
 			*SharedLocalBasisIndexMacro,
 			*NormalCode,
 			*TangentCode,
@@ -10330,27 +10331,28 @@ int32 FHLSLMaterialTranslator::StrataSlabBSDF(
 	}
 	
 	return AddCodeChunk(
-		MCT_Strata, TEXT("GetStrataSlabBSDF(Parameters.StrataPixelFootprint, %s, %s, %s, %s, %s, %s, %s, %s, %s, %s, %s, %s, %s, %s, %s, %s, %s, %s, %s, %s, %s, Parameters.SharedLocalBases.Types) /* Normal = %s ; Tangent = %s */"),
-		*StrataGetCastParameterCode(UseMetalness,	MCT_Float),
-		*StrataGetCastParameterCode(BaseColor,		MCT_Float3),
-		*StrataGetCastParameterCode(EdgeColor,		MCT_Float3),
-		*StrataGetCastParameterCode(Specular,		MCT_Float),
-		*StrataGetCastParameterCode(Metallic,		MCT_Float),
-		*StrataGetCastParameterCode(DiffuseAlbedo,	MCT_Float3),
-		*StrataGetCastParameterCode(F0,				MCT_Float3),
-		*StrataGetCastParameterCode(F90,			MCT_Float3),
-		*StrataGetCastParameterCode(Roughness,		MCT_Float),
-		*StrataGetCastParameterCode(Anisotropy,		MCT_Float),
-		*StrataGetCastParameterCode(SSSProfileId,	MCT_Float),
-		*StrataGetCastParameterCode(SSSMFP,			MCT_Float3),
-		*StrataGetCastParameterCode(SSSMFPScale,	MCT_Float),
-		*StrataGetCastParameterCode(SSSPhaseAniso,	MCT_Float),
-		*StrataGetCastParameterCode(UseSSSDiffusion,MCT_Float),
-		*StrataGetCastParameterCode(EmissiveColor,	MCT_Float3),
-		*StrataGetCastParameterCode(Haziness,		MCT_Float),
-		*StrataGetCastParameterCode(FuzzAmount,		MCT_Float),
-		*StrataGetCastParameterCode(FuzzColor,		MCT_Float3),
-		*StrataGetCastParameterCode(Thickness,		MCT_Float),
+		MCT_Strata, TEXT("GetStrataSlabBSDF(Parameters.StrataPixelFootprint, %s, %s, %s, %s, %s, %s, %s, %s, %s, %s, %s, %s, %s, %s, %s, %s, %s, %s, %s, %s, %s, %s, Parameters.SharedLocalBases.Types) /* Normal = %s ; Tangent = %s */"),
+		*StrataGetCastParameterCode(UseMetalness,			MCT_Float),
+		*StrataGetCastParameterCode(BaseColor,				MCT_Float3),
+		*StrataGetCastParameterCode(EdgeColor,				MCT_Float3),
+		*StrataGetCastParameterCode(Specular,				MCT_Float),
+		*StrataGetCastParameterCode(Metallic,				MCT_Float),
+		*StrataGetCastParameterCode(DiffuseAlbedo,			MCT_Float3),
+		*StrataGetCastParameterCode(F0,						MCT_Float3),
+		*StrataGetCastParameterCode(F90,					MCT_Float3),
+		*StrataGetCastParameterCode(Roughness,				MCT_Float),
+		*StrataGetCastParameterCode(Anisotropy,				MCT_Float),
+		*StrataGetCastParameterCode(SSSProfileId,			MCT_Float),
+		*StrataGetCastParameterCode(SSSMFP,					MCT_Float3),
+		*StrataGetCastParameterCode(SSSMFPScale,			MCT_Float),
+		*StrataGetCastParameterCode(SSSPhaseAniso,			MCT_Float),
+		*StrataGetCastParameterCode(UseSSSDiffusion,		MCT_Float),
+		*StrataGetCastParameterCode(EmissiveColor,			MCT_Float3),
+		*StrataGetCastParameterCode(SecondRoughness,		MCT_Float),
+		*StrataGetCastParameterCode(SecondRoughnessWeight,	MCT_Float),
+		*StrataGetCastParameterCode(FuzzAmount,				MCT_Float),
+		*StrataGetCastParameterCode(FuzzColor,				MCT_Float3),
+		*StrataGetCastParameterCode(Thickness,				MCT_Float),
 		*SharedLocalBasisIndexMacro,
 		*NormalCode,
 		*TangentCode
diff --git a/Engine/Source/Runtime/Engine/Private/Materials/HLSLMaterialTranslator.h b/Engine/Source/Runtime/Engine/Private/Materials/HLSLMaterialTranslator.h
--- a/Engine/Source/Runtime/Engine/Private/Materials/HLSLMaterialTranslator.h
+++ b/Engine/Source/Runtime/Engine/Private/Materials/HLSLMaterialTranslator.h
@@ -1025,7 +1025,7 @@ protected:
 		int32 Roughness, int32 Anisotropy,
 		int32 SSSProfileId, int32 SSSMFP, int32 SSSMFPScale, int32 SSSPhaseAniso, int32 bUseSSSDiffusion,
 		int32 EmissiveColor, 
-		int32 Haziness, 
+		int32 SecondRoughness, int32 SecondRoughnessWeight,
 		int32 FuzzAmount, int32 FuzzColor, 
 		int32 Thickness,
 		int32 Normal, int32 Tangent, const FString& SharedLocalBasisIndexMacro,
diff --git a/Engine/Source/Runtime/Engine/Private/Materials/MaterialExpressions.cpp b/Engine/Source/Runtime/Engine/Private/Materials/MaterialExpressions.cpp
--- a/Engine/Source/Runtime/Engine/Private/Materials/MaterialExpressions.cpp
+++ b/Engine/Source/Runtime/Engine/Private/Materials/MaterialExpressions.cpp
@@ -22000,7 +22000,7 @@ int32 UMaterialExpressionStrataSlabBSDF::Compile(class FMaterialCompiler* Compil
 
 	const bool bHasEdgeColor = HasEdgeColor(); // This accounts for EdgeColor and also F90 when the non metalness workflow is selected.
 	const bool bHasFuzz = HasFuzz();
-	const bool bHasHaziness = HasHaziness();
+	const bool bHasSecondRoughness = HasSecondRoughness();
 	const bool bHasMFPPluggedIn = HasMFPPluggedIn();
 	int32 SSSProfileCodeChunk = INDEX_NONE;
 	const bool bHasSSS = HasSSS();
@@ -22045,7 +22045,8 @@ int32 UMaterialExpressionStrataSlabBSDF::Compile(class FMaterialCompiler* Compil
 		CompileWithDefaultFloat1(Compiler, SSSPhaseAnisotropy, 0.0f),
 		bUseSSSDiffusion ? Compiler->Constant(1.0f) : Compiler->Constant(0.0f),
 		CompileWithDefaultFloat3(Compiler, EmissiveColor, 0.0f, 0.0f, 0.0f),
-		CompileWithDefaultFloat1(Compiler, Haziness, 0.0f),
+		CompileWithDefaultFloat1(Compiler, SecondRoughness, 0.0f),
+		CompileWithDefaultFloat1(Compiler, SecondRoughnessWeight, 0.0f),
 		CompileWithDefaultFloat1(Compiler, FuzzAmount, 0.0f),
 		CompileWithDefaultFloat3(Compiler, FuzzColor, 0.0f, 0.0f, 0.0f),
 		CompileWithDefaultFloat1(Compiler, Thickness, STRATA_LAYER_DEFAULT_THICKNESS_CM),
@@ -22082,7 +22083,8 @@ const TArray<FExpressionInput*> UMaterialExpressionStrataSlabBSDF::GetInputs()
 	Result.Add(&SSSMFPScale);
 	Result.Add(&SSSPhaseAnisotropy);
 	Result.Add(&EmissiveColor);
-	Result.Add(&Haziness);
+	Result.Add(&SecondRoughness);
+	Result.Add(&SecondRoughnessWeight);
 	Result.Add(&Thickness);
 	Result.Add(&FuzzAmount);
 	Result.Add(&FuzzColor);
@@ -22221,17 +22223,21 @@ uint32 UMaterialExpressionStrataSlabBSDF::GetInputType(int32 InputIndex)
 	}
 	else if (InputIndex == (12 + SkipDisabledInputOffset))
 	{
-		return MCT_Float1; // Haziness
+		return MCT_Float1; // SecondRoughness
 	}
 	else if (InputIndex == (13 + SkipDisabledInputOffset))
 	{
-		return MCT_Float1; // Thickness
+		return MCT_Float1; // SecondRoughnessWeight
 	}
 	else if (InputIndex == (14 + SkipDisabledInputOffset))
 	{
-		return MCT_Float1; // FuzzAmount
+		return MCT_Float1; // Thickness
 	}
 	else if (InputIndex == (15 + SkipDisabledInputOffset))
+	{
+		return MCT_Float1; // FuzzAmount
+	}
+	else if (InputIndex == (16 + SkipDisabledInputOffset))
 	{
 		return MCT_Float3; // FuzzColor
 	}
@@ -22313,17 +22319,21 @@ FName UMaterialExpressionStrataSlabBSDF::GetInputName(int32 InputIndex) const
 	}
 	else if (InputIndex == (12 + SkipDisabledInputOffset))
 	{
-		return TEXT("Haziness");
+		return TEXT("Second Roughness");
 	}
 	else if (InputIndex == (13 + SkipDisabledInputOffset))
 	{
-		return TEXT("Thickness");
+		return TEXT("Second Roughness Weight");
 	}
 	else if (InputIndex == (14 + SkipDisabledInputOffset))
 	{
-		return TEXT("FuzzAmount");
+		return TEXT("Thickness");
 	}
 	else if (InputIndex == (15 + SkipDisabledInputOffset))
+	{
+		return TEXT("FuzzAmount");
+	}
+	else if (InputIndex == (16 + SkipDisabledInputOffset))
 	{
 		return TEXT("FuzzColor");
 	}
@@ -22427,6 +22437,10 @@ void UMaterialExpressionStrataSlabBSDF::GetConnectorToolTip(int32 InputIndex, in
 	{
 		Super::GetConnectorToolTip(19, INDEX_NONE, OutToolTip);
 	}
+	else if (InputIndex == (17 + SkipDisabledInputOffset))
+	{
+		Super::GetConnectorToolTip(20, INDEX_NONE, OutToolTip);
+	}
 }
 
 bool UMaterialExpressionStrataSlabBSDF::IsResultStrataMaterial(int32 OutputIndex)
@@ -22476,7 +22490,7 @@ FStrataOperator* UMaterialExpressionStrataSlabBSDF::StrataGenerateMaterialTopolo
 	StrataOperator.BSDFType = STRATA_BSDF_TYPE_SLAB;
 	StrataOperator.bBSDFHasEdgeColor = HasEdgeColor();
 	StrataOperator.bBSDFHasFuzz = HasFuzz();
-	StrataOperator.bBSDFHasHaziness = HasHaziness();
+	StrataOperator.bBSDFHasSecondRoughness = HasSecondRoughness();
 	StrataOperator.bBSDFHasSSS = HasSSS();
 	StrataOperator.bBSDFHasMFPPluggedIn = HasMFPPluggedIn();
 	StrataOperator.bBSDFHasAnisotropy = HasAnisotropy();
@@ -22508,9 +22522,9 @@ bool UMaterialExpressionStrataSlabBSDF::HasFuzz() const
 	return FuzzAmount.IsConnected();
 }
 
-bool UMaterialExpressionStrataSlabBSDF::HasHaziness() const
+bool UMaterialExpressionStrataSlabBSDF::HasSecondRoughness() const
 {
-	return Haziness.IsConnected();
+	return SecondRoughnessWeight.IsConnected();
 }
 
 bool UMaterialExpressionStrataSlabBSDF::HasAnisotropy() const
diff --git a/Engine/Source/Runtime/Engine/Public/MaterialCompiler.h b/Engine/Source/Runtime/Engine/Public/MaterialCompiler.h
--- a/Engine/Source/Runtime/Engine/Public/MaterialCompiler.h
+++ b/Engine/Source/Runtime/Engine/Public/MaterialCompiler.h
@@ -111,7 +111,7 @@ struct FStrataOperator
 	bool  bBSDFHasMFPPluggedIn;
 	bool  bBSDFHasEdgeColor;
 	bool  bBSDFHasFuzz;
-	bool  bBSDFHasHaziness;
+	bool  bBSDFHasSecondRoughness;
 	bool  bBSDFHasAnisotropy;
 
 	// Data derived after the tree has been built.
@@ -137,7 +137,7 @@ struct FStrataOperator
 		bBSDFHasMFPPluggedIn = false;
 		bBSDFHasEdgeColor = false;
 		bBSDFHasFuzz = false;
-		bBSDFHasHaziness = false;
+		bBSDFHasSecondRoughness = false;
 		bBSDFHasAnisotropy = false;
 
 		MaxDistanceFromLeaves = 0;
@@ -154,7 +154,7 @@ struct FStrataOperator
 		bBSDFHasMFPPluggedIn = A.bBSDFHasMFPPluggedIn || B.bBSDFHasMFPPluggedIn;
 		bBSDFHasEdgeColor = A.bBSDFHasEdgeColor || B.bBSDFHasEdgeColor;
 		bBSDFHasFuzz = A.bBSDFHasFuzz || B.bBSDFHasFuzz;
-		bBSDFHasHaziness = A.bBSDFHasHaziness || B.bBSDFHasHaziness;
+		bBSDFHasSecondRoughness = A.bBSDFHasSecondRoughness || B.bBSDFHasSecondRoughness;
 		bBSDFHasAnisotropy = A.bBSDFHasAnisotropy || B.bBSDFHasAnisotropy;
 	}
 
@@ -164,7 +164,7 @@ struct FStrataOperator
 		bBSDFHasMFPPluggedIn = A.bBSDFHasMFPPluggedIn;
 		bBSDFHasEdgeColor = A.bBSDFHasEdgeColor;
 		bBSDFHasFuzz = A.bBSDFHasFuzz;
-		bBSDFHasHaziness = A.bBSDFHasHaziness;
+		bBSDFHasSecondRoughness = A.bBSDFHasSecondRoughness;
 		bBSDFHasAnisotropy = A.bBSDFHasAnisotropy;
 	}
 
@@ -568,7 +568,7 @@ public:
 		int32 Roughness, int32 Anisotropy,
 		int32 SSSProfileId, int32 SSSMFP, int32 SSSMFPScale, int32 SSSPhaseAniso, int32 UseSSSDiffusion,
 		int32 EmissiveColor, 
-		int32 Haziness, 
+		int32 SecondRoughness, int32 SecondRoughnessWeight,
 		int32 FuzzAmount, int32 FuzzColor,
 		int32 Thickness,
 		int32 Normal, int32 Tangent, const FString& SharedLocalBasisIndexMacro,
@@ -1116,7 +1116,7 @@ public:
 		int32 Roughness, int32 Anisotropy,
 		int32 SSSProfileId, int32 SSSMFP, int32 SSSMFPScale, int32 SSSPhaseAniso, int32 UseSSSDiffusion,
 		int32 EmissiveColor, 
-		int32 Haziness, 
+		int32 SecondRoughness, int32 SecondRoughnessWeight,
 		int32 FuzzAmount, int32 FuzzColor,
 		int32 Thickness,
 		int32 Normal, int32 Tangent, const FString& SharedLocalBasisIndexMacro,
@@ -1129,7 +1129,7 @@ public:
 			Roughness, Anisotropy,
 			SSSProfileId, SSSMFP, SSSMFPScale, SSSPhaseAniso, UseSSSDiffusion,
 			EmissiveColor, 
-			Haziness, 
+			SecondRoughness, SecondRoughnessWeight,
 			FuzzAmount, FuzzColor,
 			Thickness,
 			Normal, Tangent, SharedLocalBasisIndexMacro,
