Add SSGI support for Strata.

diff --git a/Engine/Shaders/Private/DiffuseIndirectComposite.usf b/Engine/Shaders/Private/DiffuseIndirectComposite.usf
--- a/Engine/Shaders/Private/DiffuseIndirectComposite.usf
+++ b/Engine/Shaders/Private/DiffuseIndirectComposite.usf
@@ -7,7 +7,7 @@
 
 //------------------------------------------------------- CONFIGS
 
-#if DIM_APPLY_DIFFUSE_INDIRECT == 1 // Denoised lighting buffer
+#if DIM_APPLY_DIFFUSE_INDIRECT == 1 // Denoised lighting buffer (SSGI)
 	#define CONFIG_SIGNAL_PROCESSING SIGNAL_PROCESSING_SSGI
 	#define COMPILE_SIGNAL_COLOR 1
 	#define MAX_SIGNAL_BATCH_SIZE 1
@@ -50,7 +50,7 @@
 #define COMPILE_MOMENT1_ACCUMULATOR 1
 #define COMPILE_BOX_KERNEL 1
 
-#if STRATA_ENABLED && DIM_APPLY_DIFFUSE_INDIRECT == 4  
+#if STRATA_ENABLED && (DIM_APPLY_DIFFUSE_INDIRECT == 1 || DIM_APPLY_DIFFUSE_INDIRECT == 4)
   #if STRATA_TILETYPE == 0
 	#define STRATA_FASTPATH 1
   #elif STRATA_TILETYPE == 1
@@ -209,13 +209,15 @@ float GetSSSCheckerboadSpecularScale(uint2 PixelPos, bool bNeedsSeparateLightAcc
 // Lumen diffuse evaluation for strata material.
 
 #if STRATA_ENABLED
-FStrataDeferredLighting StrataLumenLighting(
+FStrataDeferredLighting StrataIndirectLighting(
 	uint2 PixelPos,
 	FStrataMaterialContainer MaterialBuffer,
 	inout FStrataAddressing StrataAddressing,
 	FStrataPixelHeader StrataPixelHeader,
 	in float3 V,
-	in float  InAO)
+	in float  InAO,
+	in float3 InputDiffuseLighting,
+	in float3 InputSpecularLighting)
 {
 	FStrataDeferredLighting Out = GetInitialisedStrataDeferredLighting();
 	const FStrataIntegrationSettings Settings = InitStrataIntegrationSettings(false /*bForceFullyRough*/, Strata.bRoughDiffuse, Strata.PeelLayersAboveDepth);
@@ -226,12 +228,20 @@ FStrataDeferredLighting StrataLumenLighting(
 
 	Strata_for(uint BSDFIndex = 0, BSDFIndex < StrataPixelHeader.BSDFCount, ++BSDFIndex)
 	{
+		// * Lumen integration output indirect lighting per BSDF
+		// * SSGI  integration output a single indirect lighting value for all BSDF
+	#if DIM_APPLY_DIFFUSE_INDIRECT == 4
 		const uint2 TileCoord = GetTilePositionFlatten(Tile0, Tile0Coord, BSDFIndex);
 		const uint2 InCoord   = TileCoord * STRATA_TILE_SIZE + LocalTileCoord;
 
 		const float3 InDiffuseLighting			= DiffuseIndirect_Textures_0.Load(uint3(InCoord, 0)).rgb;
 		const float3 InRoughSpecularLighting	= DiffuseIndirect_Textures_1.Load(uint3(InCoord, 0)).rgb;
 		const float3 InSpecularLighting			= DiffuseIndirect_Textures_2.Load(uint3(InCoord, 0)).rgb;
+	#else
+		const float3 InDiffuseLighting			= InputDiffuseLighting;
+		const float3 InRoughSpecularLighting	= InputSpecularLighting;
+		const float3 InSpecularLighting			= InputSpecularLighting;
+	#endif
 
 		// STRATA_TODO: Override the BSDF F0/DiffuseColor when bVisualizeDiffuseIndirect is set
 		FStrataBSDF BSDF = UnpackStrataBSDF(MaterialBuffer, StrataAddressing, StrataPixelHeader);
@@ -414,13 +424,15 @@ void MainPS(
 				FStrataAddressing StrataAddressing = GetStrataPixelDataByteOffset(PixelPos, uint2(View.BufferSizeAndInvSize.xy), Strata.MaxBytesPerPixel);
 				const FStrataPixelHeader StrataPixelHeader = UnpackStrataHeaderIn(Strata.MaterialTextureArray, StrataAddressing, Strata.TopLayerTexture);
 
-				const FStrataDeferredLighting IndirectLighting_Strata = StrataLumenLighting(
+				const FStrataDeferredLighting IndirectLighting_Strata = StrataIndirectLighting(
 					PixelPos,
 					Strata.MaterialTextureArray,
 					StrataAddressing,
 					StrataPixelHeader,
 					V,
-					FinalAmbientOcclusion);
+					FinalAmbientOcclusion,
+					DiffuseIndirectLighting,
+					SpecularIndirectLighting);
 
 				OutAddColor = IndirectLighting_Strata.SceneColor;
 			  #if STRATA_OPAQUE_ROUGH_REFRACTION_ENABLED
diff --git a/Engine/Shaders/Private/SSRT/SSRTDiffuseIndirect.usf b/Engine/Shaders/Private/SSRT/SSRTDiffuseIndirect.usf
--- a/Engine/Shaders/Private/SSRT/SSRTDiffuseIndirect.usf
+++ b/Engine/Shaders/Private/SSRT/SSRTDiffuseIndirect.usf
@@ -106,7 +106,7 @@
 #include "../BRDF.ush"
 #include "../MonteCarlo.ush"
 #include "../SceneTextureParameters.ush"
-
+#include "../Strata/Strata.ush"
 
 //------------------------------------------------------- PARAMETERS
 
@@ -189,10 +189,11 @@ uint2 ComputePixelPosition(uint2 GroupId, uint2 GroupPixelOffset)
 
 void UpdateLane2DCoordinateInformations(
 	uint2 PixelPosition,
-	out float2 BufferUV,
-	out float2 ScreenPos)
+	inout float2 BufferUV,
+	inout float2 ScreenPos,
+	inout float2 FullResPixelPosition)
 {
-	float2 FullResPixelPosition = PixelPosition * PixelPositionToFullResPixel + FullResPixelOffset;
+	FullResPixelPosition = PixelPosition * PixelPositionToFullResPixel + FullResPixelOffset;
 
 	// TODO: split screen
 	BufferUV = FullResPixelPosition * View.BufferSizeAndInvSize.zw;
@@ -267,21 +268,33 @@ void MainCS(
 			uint2 GroupPixelOffset = DecodeGroupPixelOffset(GroupPixelId);
             uint2 PixelPosition = ComputePixelPosition(GroupId, GroupPixelOffset);
 				
-			float2 BufferUV;
-			float2 ScreenPos;
-			UpdateLane2DCoordinateInformations(PixelPosition, /* out */ BufferUV, /* out */ ScreenPos);
+			float2 BufferUV = 0;
+			float2 ScreenPos = 0;
+			float2 FullResPixelPosition = 0;
+			UpdateLane2DCoordinateInformations(PixelPosition, /* out */ BufferUV, /* out */ ScreenPos, /* out */ FullResPixelPosition);
 	
+		#if STRATA_ENABLED
+			const FStrataTopLayerData TopLayerData = StrataUnpackTopLayerData(Strata.TopLayerTexture.Load(uint3(FullResPixelPosition, 0)));
+			const float Roughness = TopLayerData.Roughness;
+			const float3 WorldNormal = TopLayerData.WorldNormal;
+			const bool bIsValid = IsStrataMaterial(TopLayerData);
+		#else // STRATA_ENABLED
 			FGBufferData GBuffer = GetGBufferDataFromSceneTextures(BufferUV);
+			const float Roughness = GBuffer.Roughness;
+			const float3 WorldNormal = GBuffer.WorldNormal;
+			const bool bIsValid = GBuffer.ShadingModelID != SHADINGMODELID_UNLIT;
+		#endif // STRATA_ENABLED
+
 			float DeviceZ = SceneDepthTexture.SampleLevel(SceneDepthTextureSampler, BufferUV, 0).r;
 
 			float ProbeOcclusionDistance = ProbeOcclusionDistanceTexture.Load(int3(PixelPosition, 0)).x;
 
-			bool bTraceRay = GBuffer.ShadingModelID != SHADINGMODELID_UNLIT;
+			bool bTraceRay = bIsValid;
 				
 			#if CONFIG_LDS_STORE_VIEW_N
-				SharedMemory0[TILE_PIXEL_COUNT * 0 | GroupPixelId] = CompressN(mul(float4(GBuffer.WorldNormal, 0), View.TranslatedWorldToView).xyz);
+				SharedMemory0[TILE_PIXEL_COUNT * 0 | GroupPixelId] = CompressN(mul(float4(WorldNormal, 0), View.TranslatedWorldToView).xyz);
 			#else
-				SharedMemory0[TILE_PIXEL_COUNT * 0 | GroupPixelId] = CompressN(GBuffer.WorldNormal);
+				SharedMemory0[TILE_PIXEL_COUNT * 0 | GroupPixelId] = CompressN(WorldNormal);
 			#endif
 			SharedMemory0[TILE_PIXEL_COUNT * 1 | GroupPixelId] = asuint(bTraceRay ? DeviceZ : -1.0);
 
@@ -291,7 +304,7 @@ void MainCS(
 			}
 			#elif DIM_LIGHTING_TERM == SPECULAR_TERM
 			{
-				SharedMemory1[GroupPixelId] = asuint(GBuffer.Roughness);
+				SharedMemory1[GroupPixelId] = asuint(Roughness);
 			}
 			#else
 				#error Unnimplemented
@@ -383,9 +396,10 @@ void MainCS(
 			uint2 GroupPixelOffset = DecodeGroupPixelOffset(GroupPixelId);
 			uint2 PixelPosition = ComputePixelPosition(GroupId, GroupPixelOffset);
 	
-			float2 BufferUV;
-			float2 ScreenPos;
-			UpdateLane2DCoordinateInformations(PixelPosition, /* out */ BufferUV, /* out */ ScreenPos);
+			float2 BufferUV = 0;
+			float2 ScreenPos = 0;
+			float2 FullResPixelPosition = 0;
+			UpdateLane2DCoordinateInformations(PixelPosition, /* out */ BufferUV, /* out */ ScreenPos, /* out */ FullResPixelPosition); 
 	
 			float3 PositionTranslatedWorld = mul(float4(ScreenPos * SceneDepth, SceneDepth, 1), View.ScreenToTranslatedWorld).xyz;
 			float3 V = normalize(View.TranslatedWorldCameraOrigin - PositionTranslatedWorld);
