Strata - added haziness cointorl node to feed slab second roughness and weight.

diff --git a/Engine/Shaders/Private/Strata/Strata.ush b/Engine/Shaders/Private/Strata/Strata.ush
--- a/Engine/Shaders/Private/Strata/Strata.ush
+++ b/Engine/Shaders/Private/Strata/Strata.ush
@@ -872,56 +872,55 @@ bool StrataHasGreyScaleWeight(in float3 Weight)
 	return Weight.x == Weight.y && Weight.y == Weight.z;
 }
 
-// STRATA_TODO: move this code to be a Haziness node. So keeping this code here for now.
-////////void StrataComputeHazeRoughnessAndWeight(
-////////	float Roughness0,
-////////	float Haziness,
-////////	inout float Roughness1,
-////////	inout float MixLobe)
-////////{	
-////////	// Fade out Haziness as the roughness reaches 0.5f, as otherwise it create
-////////	//       27 |\
-////////	// Extent   | \
-////////	//          |  \
-////////	//        3 |    ----
-////////	//          0  0.25  1
-////////	//           Rougness
-////////
-////////	// Lerp factor for adapting haziness parameters
-////////	const float S2 = saturate(2 * Roughness0);
-////////	const float S4 = saturate(4 * Roughness0);
-////////
-////////	// For low roughness, the haze extent is larger, and the core highlight is lower. 
-////////	// This allows to emphase the haze effect at low roughness, and fade out haze as roughness increase.
-////////	const float HazeExtent = lerp(27.0f, 3.0f, S4);
-////////	const float HazeCore   = lerp(0.01f, 0.5f, S2);
-////////
-////////	// Reference "A Composite BRDF Model for Hazy Gloss"
-////////	 
-////////	// Compute wide roughness
-////////	float alpha_n = Square(Roughness0); 
-////////	float lambda_h = HazeExtent;
-////////	const float alpha_w = alpha_n * (1 + lambda_h); // Eq.8
-////////	Roughness1 = sqrtFast(alpha_w);
-////////
-////////	// Compute mix weights
-////////	const float beta_h = Haziness;
-////////	const float r_c = HazeCore;
-////////	const float w = alpha_w;
-////////
-////////	const float p = 1.0 / Square(1 + lambda_h); // Table 2, GGX
-////////	const float b = 2 * (r_c * (1 - w) + w * p);
-////////	const float D = Square(b) - 4 * (b - 1) * r_c;
-////////	const float u = (b - sqrtFast(D)) / (2 * b - 2); // Eq. 12, 13, 14
-////////	const float u2 = u * u;
-////////	const float w2 = 2 * w;
-////////	const float k_h = ((u - u2) * w2 * beta_h) / (Square(1 - u) + (u - u2) * w2 + u2); // Eq. 15
-////////	const float r = r_c + (1 - p) * k_h; // Eq. 5
-////////	MixLobe = k_h / r; // Eq. 4
-////////}
-////////
-////////float StrataComputeHazeRoughness(float Roughness0) { float Roughness1=0; float MixLobe=0; StrataComputeHazeRoughnessAndWeight(Roughness0, 1.0f, Roughness1, MixLobe); return Roughness1; }
-////////float StrataComputeHazeWeight(float Roughness0, float Haziness)    { float Roughness1=0; float MixLobe=0; StrataComputeHazeRoughnessAndWeight(Roughness0, Haziness, Roughness1, MixLobe); return MixLobe; }
+void StrataComputeHazeRoughnessAndWeight(
+	float Roughness0,
+	float Haziness,
+	inout float Roughness1,
+	inout float MixLobe)
+{	
+	// Fade out Haziness as the roughness reaches 0.5f, as otherwise it create
+	//       27 |\
+	// Extent   | \
+	//          |  \
+	//        3 |    ----
+	//          0  0.25  1
+	//           Rougness
+
+	// Lerp factor for adapting haziness parameters
+	const float S2 = saturate(2 * Roughness0);
+	const float S4 = saturate(4 * Roughness0);
+
+	// For low roughness, the haze extent is larger, and the core highlight is lower. 
+	// This allows to emphase the haze effect at low roughness, and fade out haze as roughness increase.
+	const float HazeExtent = lerp(27.0f, 3.0f, S4);
+	const float HazeCore   = lerp(0.01f, 0.5f, S2);
+
+	// Reference "A Composite BRDF Model for Hazy Gloss"
+	 
+	// Compute wide roughness
+	float alpha_n = Square(Roughness0); 
+	float lambda_h = HazeExtent;
+	const float alpha_w = alpha_n * (1 + lambda_h); // Eq.8
+	Roughness1 = sqrtFast(alpha_w);
+
+	// Compute mix weights
+	const float beta_h = Haziness;
+	const float r_c = HazeCore;
+	const float w = alpha_w;
+
+	const float p = 1.0 / Square(1 + lambda_h); // Table 2, GGX
+	const float b = 2 * (r_c * (1 - w) + w * p);
+	const float D = Square(b) - 4 * (b - 1) * r_c;
+	const float u = (b - sqrtFast(D)) / (2 * b - 2); // Eq. 12, 13, 14
+	const float u2 = u * u;
+	const float w2 = 2 * w;
+	const float k_h = ((u - u2) * w2 * beta_h) / (Square(1 - u) + (u - u2) * w2 + u2); // Eq. 15
+	const float r = r_c + (1 - p) * k_h; // Eq. 5
+	MixLobe = k_h / r; // Eq. 4
+}
+
+float StrataComputeHazeRoughness(float Roughness0)				{ float Roughness1=0; float MixLobe=0; StrataComputeHazeRoughnessAndWeight(Roughness0, 1.0f,     Roughness1, MixLobe);	return Roughness1; }
+float StrataComputeHazeWeight(float Roughness0, float Haziness)	{ float Roughness1=0; float MixLobe=0; StrataComputeHazeRoughnessAndWeight(Roughness0, Haziness, Roughness1, MixLobe);	return MixLobe; }
 
 ///////////////////////////////////////////////////////////////////////////////
 // Material & BSDF description
diff --git a/Engine/Source/Runtime/Engine/Classes/Materials/MaterialExpressionStrata.h b/Engine/Source/Runtime/Engine/Classes/Materials/MaterialExpressionStrata.h
--- a/Engine/Source/Runtime/Engine/Classes/Materials/MaterialExpressionStrata.h
+++ b/Engine/Source/Runtime/Engine/Classes/Materials/MaterialExpressionStrata.h
@@ -894,3 +894,32 @@ class UMaterialExpressionStrataMetalnessToDiffuseAlbedoF0 : public UMaterialExpr
 	//~ End UMaterialExpression Interface
 };
 
+UCLASS(MinimalAPI, collapsecategories, hidecategories = Object, DisplayName = "Strata Haziness-To-Secondary-Roughness")
+class UMaterialExpressionStrataHazinessToSecondaryRoughness : public UMaterialExpressionStrataBSDF
+{
+	GENERATED_UCLASS_BODY()
+
+	/**
+	* The base roughness of the surface. It represented the smoothest part of the reflection.
+	*/
+	UPROPERTY()
+	FExpressionInput BaseRoughness;
+
+	/**
+	* Haziness represent the amount of irregularity of the surface. A high value will lead to a second rough specular lobe causing the surface too look `milky`.
+	*/
+	UPROPERTY()
+	FExpressionInput Haziness;
+
+	//~ Begin UMaterialExpression Interface
+#if WITH_EDITOR
+	virtual int32 Compile(class FMaterialCompiler* Compiler, int32 OutputIndex) override;
+	virtual void GetCaption(TArray<FString>& OutCaptions) const override;
+	virtual uint32 GetOutputType(int32 OutputIndex) override;
+	virtual uint32 GetInputType(int32 InputIndex) override;
+	virtual void GetConnectorToolTip(int32 InputIndex, int32 OutputIndex, TArray<FString>& OutToolTip) override;
+	virtual void GetExpressionToolTip(TArray<FString>& OutToolTip) override;
+#endif
+	//~ End UMaterialExpression Interface
+};
+
diff --git a/Engine/Source/Runtime/Engine/Private/Materials/HLSLMaterialTranslator.cpp b/Engine/Source/Runtime/Engine/Private/Materials/HLSLMaterialTranslator.cpp
--- a/Engine/Source/Runtime/Engine/Private/Materials/HLSLMaterialTranslator.cpp
+++ b/Engine/Source/Runtime/Engine/Private/Materials/HLSLMaterialTranslator.cpp
@@ -10769,6 +10769,30 @@ int32 FHLSLMaterialTranslator::StrataMetalnessToDiffuseAlbedoF0(int32 BaseColor,
 	return INDEX_NONE;
 }
 
+int32 FHLSLMaterialTranslator::StrataHazinessToSecondaryRoughness(int32 BaseRoughness, int32 Haziness, int32 OutputIndex)
+{
+	if (OutputIndex == INDEX_NONE)
+	{
+		return INDEX_NONE;
+	}
+
+	switch (OutputIndex)
+	{
+	case 0:
+		return AddCodeChunk(MCT_Float1,
+			TEXT("StrataComputeHazeRoughness(saturate(%s))"),
+			*GetParameterCode(BaseRoughness));
+		break;
+	case 1:
+		return AddCodeChunk(MCT_Float1,
+			TEXT("StrataComputeHazeWeight(saturate(%s), saturate(%s))"),
+			*GetParameterCode(BaseRoughness),
+			*GetParameterCode(Haziness));
+		break;
+	}
+	return INDEX_NONE;
+}
+
 int32 FHLSLMaterialTranslator::MapARPassthroughCameraUV(int32 UV)
 {
 	if (UV == INDEX_NONE)
diff --git a/Engine/Source/Runtime/Engine/Private/Materials/HLSLMaterialTranslator.h b/Engine/Source/Runtime/Engine/Private/Materials/HLSLMaterialTranslator.h
--- a/Engine/Source/Runtime/Engine/Private/Materials/HLSLMaterialTranslator.h
+++ b/Engine/Source/Runtime/Engine/Private/Materials/HLSLMaterialTranslator.h
@@ -1062,6 +1062,7 @@ protected:
 	virtual int32 StrataThinFilmParameterBlending(int32 A, int32 Thickness, int32 IOR, int32 BSDFNormalCodeChunk, FStrataOperator* PromoteToOperator) override;
 	virtual int32 StrataTransmittanceToMFP(int32 TransmittanceColor, int32 DesiredThickness, int32 OutputIndex) override;
 	virtual int32 StrataMetalnessToDiffuseAlbedoF0(int32 BaseColor, int32 Specular, int32 Metallic, int32 OutputIndex) override;
+	virtual int32 StrataHazinessToSecondaryRoughness(int32 BaseRoughness, int32 Haziness, int32 OutputIndex) override;
 	
 	virtual FStrataOperator& StrataCompilationRegisterOperator(int32 OperatorType, UMaterialExpression* Expression, UMaterialExpression* Parent, bool bUseParameterBlending = false) override;
 	virtual FStrataOperator& StrataCompilationGetOperator(UMaterialExpression* Expression) override;
diff --git a/Engine/Source/Runtime/Engine/Private/Materials/MaterialExpressions.cpp b/Engine/Source/Runtime/Engine/Private/Materials/MaterialExpressions.cpp
--- a/Engine/Source/Runtime/Engine/Private/Materials/MaterialExpressions.cpp
+++ b/Engine/Source/Runtime/Engine/Private/Materials/MaterialExpressions.cpp
@@ -23998,6 +23998,108 @@ void UMaterialExpressionStrataMetalnessToDiffuseAlbedoF0::GetExpressionToolTip(T
 #endif // WITH_EDITOR
 
 
+
+UMaterialExpressionStrataHazinessToSecondaryRoughness::UMaterialExpressionStrataHazinessToSecondaryRoughness(const FObjectInitializer& ObjectInitializer)
+	: Super(ObjectInitializer)
+{
+	struct FConstructorStatics
+	{
+		FText NAME_Strata;
+		FConstructorStatics() : NAME_Strata(LOCTEXT("Strata Helpers", "Strata Helpers")) { }
+	};
+	static FConstructorStatics ConstructorStatics;
+#if WITH_EDITORONLY_DATA
+	MenuCategories.Add(ConstructorStatics.NAME_Strata);
+
+	bShowOutputNameOnPin = true;
+
+	Outputs.Reset();
+	Outputs.Add(FExpressionOutput(TEXT("Second Roughness")));
+	Outputs.Add(FExpressionOutput(TEXT("Second Roughness Weight")));
+#endif
+}
+
+#if WITH_EDITOR
+int32 UMaterialExpressionStrataHazinessToSecondaryRoughness::Compile(class FMaterialCompiler* Compiler, int32 OutputIndex)
+{
+	return Compiler->StrataHazinessToSecondaryRoughness(
+		BaseRoughness.GetTracedInput().Expression ? BaseRoughness.Compile(Compiler) : Compiler->Constant(0.1f),
+		Haziness.GetTracedInput().Expression ? Haziness.Compile(Compiler) : Compiler->Constant(0.5f),
+		OutputIndex);
+}
+
+void UMaterialExpressionStrataHazinessToSecondaryRoughness::GetCaption(TArray<FString>& OutCaptions) const
+{
+	OutCaptions.Add(TEXT("Strata Haziness-To-Secondary-Roughness"));
+}
+
+uint32 UMaterialExpressionStrataHazinessToSecondaryRoughness::GetOutputType(int32 OutputIndex)
+{
+	switch (OutputIndex)
+	{
+	case 0:
+		return MCT_Float1; // Second Roughness
+		break;
+	case 1:
+		return MCT_Float1; // Second Roughness Weight
+		break;
+	}
+
+	check(false);
+	return MCT_Float1;
+}
+
+uint32 UMaterialExpressionStrataHazinessToSecondaryRoughness::GetInputType(int32 InputIndex)
+{
+	switch (InputIndex)
+	{
+	case 0:
+		return MCT_Float1; // BaseRoughness
+		break;
+	case 1:
+		return MCT_Float1; // Haziness
+		break;
+	}
+
+	check(false);
+	return MCT_Float1;
+}
+void UMaterialExpressionStrataHazinessToSecondaryRoughness::GetConnectorToolTip(int32 InputIndex, int32 OutputIndex, TArray<FString>& OutToolTip)
+{
+	if (InputIndex != INDEX_NONE)
+	{
+		switch (InputIndex)
+		{
+		case 0:
+			ConvertToMultilineToolTip(TEXT("The base roughness of the surface. It represented the smoothest part of the reflection."), 80, OutToolTip);
+			break;
+		case 1:
+			ConvertToMultilineToolTip(TEXT("Haziness represent the amount of irregularity of the surface. A high value will lead to a second rough specular lobe causing the surface too look `milky`."), 80, OutToolTip);
+			break;
+		}
+	}
+	else if (OutputIndex != INDEX_NONE)
+	{
+		switch (OutputIndex)
+		{
+		case 0:
+			ConvertToMultilineToolTip(TEXT("The roughness of the second lobe."), 80, OutToolTip);
+			break;
+		case 1:
+			ConvertToMultilineToolTip(TEXT("The weight of the secondary specular lobe, while the primary specular lobe will have a weight of (1 - SecondRoughnessWeight)."), 80, OutToolTip);
+			break;
+		}
+	}
+}
+
+void UMaterialExpressionStrataHazinessToSecondaryRoughness::GetExpressionToolTip(TArray<FString>& OutToolTip)
+{
+	ConvertToMultilineToolTip(TEXT("Compute a second specular lobe roughness from a base surface roughness and haziness. This parameterisation ensure that the haziness makes physically and is perceptually easy to author."), 80, OutToolTip);
+
+}
+#endif // WITH_EDITOR
+
+
 UMaterialExpressionExecBegin::UMaterialExpressionExecBegin(const FObjectInitializer& ObjectInitializer)
 	: Super(ObjectInitializer)
 {
diff --git a/Engine/Source/Runtime/Engine/Public/MaterialCompiler.h b/Engine/Source/Runtime/Engine/Public/MaterialCompiler.h
--- a/Engine/Source/Runtime/Engine/Public/MaterialCompiler.h
+++ b/Engine/Source/Runtime/Engine/Public/MaterialCompiler.h
@@ -604,6 +604,7 @@ public:
 	virtual int32 StrataThinFilmParameterBlending(int32 A, int32 Thickness, int32 IOR, int32 BSDFNormalCodeChunk, FStrataOperator* PromoteToOperator) = 0;
 	virtual int32 StrataTransmittanceToMFP(int32 TransmittanceColor, int32 DesiredThickness, int32 OutputIndex) = 0;
 	virtual int32 StrataMetalnessToDiffuseAlbedoF0(int32 BaseColor, int32 Specular, int32 Metallic, int32 OutputIndex) = 0;
+	virtual int32 StrataHazinessToSecondaryRoughness(int32 BaseRoughness, int32 Haziness, int32 OutputIndex) = 0;
 	
 	/**
 	 * Register an operator of the tree representation the Strata material and its topology.
@@ -1249,6 +1250,11 @@ public:
 		return Compiler->StrataMetalnessToDiffuseAlbedoF0(BaseColor, Specular, Metallic, OutputIndex);
 	}
 
+	virtual int32 StrataHazinessToSecondaryRoughness(int32 BaseRoughness, int32 Haziness, int32 OutputIndex) override
+	{
+		return Compiler->StrataHazinessToSecondaryRoughness(BaseRoughness, Haziness, OutputIndex);
+	}
+
 	virtual FStrataOperator& StrataCompilationRegisterOperator(int32 OperatorType, UMaterialExpression* Expression, UMaterialExpression* Parent, bool bUseParameterBlending = false) override
 	{
 		return Compiler->StrataCompilationRegisterOperator(OperatorType, Expression, Parent, bUseParameterBlending);
