Strata - SSS profile dual specular is now mapping correctly from legacy to strata.

diff --git a/Engine/Shaders/Private/Strata/Strata.ush b/Engine/Shaders/Private/Strata/Strata.ush
--- a/Engine/Shaders/Private/Strata/Strata.ush
+++ b/Engine/Shaders/Private/Strata/Strata.ush
@@ -1265,16 +1265,7 @@ FStrataData GetStrataSlabBSDF(
 		// Average roughness for dual specular.
 		const uint SubsurfaceProfileUInt = StrataSubsurfaceProfileIdTo8bits(SSSProfileID);
 
-		float LobeRoughness0 = 0;
-		float LobeRoughness1 = 0;
-
-		GetSubsurfaceProfileDualSpecular(SubsurfaceProfileUInt, Roughness, SSSMFPScale, LobeRoughness0, LobeRoughness1, Haziness.Weight);
-
-		// Convertion from roughness0/1 to roughness+haziness
-		// * This encoding is based on how roughness0 + haziness is translated to roughness1, see ComputeHazyLobeRoughness
-		// * The mix factor is hardoded to 0.5 i.e., 0.5 * Fs(roughness0) + 0.5 * Fs(roughness1)
-		Haziness.Roughness = (LobeRoughness1 - LobeRoughness0) / (1.f - min(LobeRoughness0, 0.99f));
-		Roughness = LobeRoughness0;
+		GetSubsurfaceProfileDualSpecular(SubsurfaceProfileUInt, Roughness, SSSMFPScale, Roughness, Haziness.Roughness, Haziness.Weight);
 
 		SSSMFP = GetSubsurfaceProfileMFPInCm(SubsurfaceProfileUInt).xyz * SSSMFPScale;
 	}
