* Fix Strata hair BSDF missing tree update.
* Change material header BSDF to handle a single BSDF type rather than a combination of BSDF type (which isn't supported at the moment.

diff --git a/Engine/Shaders/Private/Strata/Strata.ush b/Engine/Shaders/Private/Strata/Strata.ush
--- a/Engine/Shaders/Private/Strata/Strata.ush
+++ b/Engine/Shaders/Private/Strata/Strata.ush
@@ -844,16 +844,7 @@ uint StrataSubsurfaceProfileIdTo8bits(float In)
 }
 
 ///////////////////////////////////////////////////////////////////////////////
-// Shading models
-uint StrataShadingModelBit(uint ShadingModel)
-{
-	return 1u << ShadingModel;
-}
-
-bool StrataHasShadingModel(uint ShadingModels, uint ShadingModel)
-{
-	return (ShadingModels & (1u << ShadingModel)) > 0;
-}
+// Misc
 
 // Return true if a BSDF has a custom area light integrator
 bool StrataHasAreaLightIntegrator(in FStrataBSDF BSDF)
@@ -1719,7 +1710,7 @@ bool IsStrataMaterial(const FStrataPixelHeader In)											{ return (In.State
 bool HasSubsurface(const FStrataPixelHeader In)												{ return (In.State & HEADER_MASK_HASSUBSURFACE) != 0; }
 bool HasDynamicIndirectShadowCasterRepresentation(const FStrataPixelHeader In)				{ return (In.State & HEADER_MASK_HASDYNINDIRECTSHADOWCASTER) != 0; }
 
-bool StrataHasShadingModel(const FStrataPixelHeader In, uint ShadingModel)					{ return StrataHasShadingModel(GetShadingModels(In), ShadingModel); }
+bool StrataHasShadingModel(const FStrataPixelHeader In, uint ShadingModel)					{ return GetShadingModels(In) == ShadingModel; }
 
 void SetHasPrecShadowMask(inout FStrataPixelHeader Out, bool bIn)							{ Out.State |= (bIn ? HEADER_MASK_HASPRECSHADOWMASK : 0u); }
 void SetZeroPrecShadowMask(inout FStrataPixelHeader Out, bool bIn)							{ Out.State |= (bIn ? HEADER_MASK_ZEROPRECSHADOWMASK : 0u); }
@@ -3063,7 +3054,7 @@ void StrataConvertFromDBuffer(
 	OutHeader.IrradianceAO = InitIrradianceAndOcclusion();
 #endif
 	OutHeader.State = 0;
-	AddShadingModels(OutHeader, StrataShadingModelBit(STRATA_BSDF_TYPE_SLAB));
+	HEADER_SETSHADINGMODELS(OutHeader.State, STRATA_BSDF_TYPE_SLAB);
 
 	// Decal only support a single Slab BSDF with no-anisotropy
 	const FStrataPixelFootprint PixelFootprint = (FStrataPixelFootprint)0;
@@ -3241,7 +3232,7 @@ FStrataPixelHeader UnpackStrataHeaderIn(FStrataMaterialContainer StrataBuffer, i
 		{
 			// Clear shading model, and ensurewe have only STRATA_BSDF_TYPE_SLAB
 			Out.State = Out.State & (~HEADER_MASK_SHADINGMODELS);
-			HEADER_SETSHADINGMODELS(Out.State, StrataShadingModelBit(STRATA_BSDF_TYPE_SLAB));
+			HEADER_SETSHADINGMODELS(Out.State, STRATA_BSDF_TYPE_SLAB);
 		}
 		else // if (bIsSingleMaterial)
 		{
diff --git a/Engine/Shaders/Private/Strata/StrataExport.ush b/Engine/Shaders/Private/Strata/StrataExport.ush
--- a/Engine/Shaders/Private/Strata/StrataExport.ush
+++ b/Engine/Shaders/Private/Strata/StrataExport.ush
@@ -243,7 +243,8 @@ void PackStrataOut(
 					float3 N = TangentBasis[2];
 					TopLayerData.WorldNormal += N * BSDF.TopLayerDataWeight;
 
-					AddShadingModels(StrataPixelHeader, StrataShadingModelBit(BSDF_GETTYPE(BSDF)));
+					// Set shading mode to the last visible BSDF. Normally all BSDF should have the same type
+					HEADER_SETSHADINGMODELS(StrataPixelHeader.State, BSDF_GETTYPE(BSDF));
 
 					OneBSDFMaterial_Index = BSDFIdx;
 
@@ -267,6 +268,10 @@ void PackStrataOut(
 			bIsSimpleMaterial = bIsSimpleMaterial && BSDFVisibleCount == 1;
 			bIsSingleMaterial = bIsSingleMaterial && BSDFVisibleCount == 1;
 			SetMaterialMode(StrataPixelHeader, bIsSimpleMaterial ? HEADER_MATERIALMODE_SIMPLE : bIsSingleMaterial ? HEADER_MATERIALMODE_SINGLE : HEADER_MATERIALMODE_COMPLEX);
+		}		
+		else if (!IsSingleLayerWater(StrataPixelHeader))
+		{
+			SetMaterialMode(StrataPixelHeader, HEADER_MATERIALMODE_COMPLEX);
 		}
 
 		// Now evaluate Specular Occlusion and Indirect Irradiance.
diff --git a/Engine/Shaders/Private/Strata/StrataLegacyConversion.ush b/Engine/Shaders/Private/Strata/StrataLegacyConversion.ush
--- a/Engine/Shaders/Private/Strata/StrataLegacyConversion.ush
+++ b/Engine/Shaders/Private/Strata/StrataLegacyConversion.ush
@@ -624,16 +624,28 @@ FStrataData StrataConvertLegacyMaterialStatic(
 	}
 	else if (ShadingModel == SHADINGMODELID_HAIR)
 	{
-		Out = GetStrataHairBSDF(
-			BaseColor,	   // BaseColor
-			Metallic,	   // Scatter
-			Specular,	   // Specular
-			Roughness,	   // Roughness
-			ClearCoat,	   // Backlit
-			Emissive,	   // EmissiveColor
-			SharedLocalBasisIndex);
+		FStrataData Top_Hair = PromoteParameterBlendedBSDFToOperator(
+			GetStrataHairBSDF(
+				BaseColor,	   // BaseColor
+				Metallic,	   // Scatter
+				Specular,	   // Specular
+				Roughness,	   // Roughness
+				ClearCoat,	   // Backlit
+				Emissive,	   // EmissiveColor
+				SharedLocalBasisIndex),
+				StrataTree,
+				0		/*OperatorIndex*/,
+				0		/*BSDFIndex*/,
+				0		/*int LayerDepth*/,
+				1		/*int bIsBottom*/);
 
-		// No need to take into account opacity weighting for translucent material, as the coverage is assumed to be 1
+		FStrataLegacyParameters Bot = InitStrataLegacyParameters(PixelFootprint, SharedLocalBasisIndex, 1.0f);
+		FStrataData Bot_Slab = CreateLegacySlab(Bot, SharedLocalBasisTypes, StrataTree,	/*OperatorIndex*/1, /*BSDFIndex*/1);
+
+		FStrataData WTop_Slab = StrataWeight(Top_Hair, 1.0, StrataTree,					/*OperatorIndex*/2, /*MaxDistanceFromLeaves*/1);
+		FStrataData WBot_Slab = StrataWeight(Bot_Slab, 0.0, StrataTree,					/*OperatorIndex*/3, /*MaxDistanceFromLeaves*/1);
+		Out = StrataVerticalLayering(WTop_Slab, WBot_Slab, StrataTree,					/*OperatorIndex*/4, /*MaxDistanceFromLeaves*/2);
+		// Out maps to LegacyUpdateBSDFsOperators
 	}
 	else if (ShadingModel == SHADINGMODELID_SINGLELAYERWATER)
 	{
@@ -653,11 +665,11 @@ FStrataData StrataConvertLegacyMaterialStatic(
 				WaterPhaseG, 			// WaterPhaseG
 				ColorScaleBehindWater,	// ColorScaleBehindWater
 				SharedLocalBasisIndex),
-			StrataTree,
-			0		/*OperatorIndex*/,
-			0		/*BSDFIndex*/,
-			0		/*int LayerDepth*/,
-			1		/*int bIsBottom*/);
+				StrataTree,
+				0		/*OperatorIndex*/,
+				0		/*BSDFIndex*/,
+				0		/*int LayerDepth*/,
+				1		/*int bIsBottom*/);
 
 		FStrataLegacyParameters Bot = InitStrataLegacyParameters(PixelFootprint, SharedLocalBasisIndex, 1.0f);
 		FStrataData Bot_Slab = CreateLegacySlab(Bot, SharedLocalBasisTypes, StrataTree,	/*OperatorIndex*/1, /*BSDFIndex*/1);
