Add SimpleClearCoatTransmittance shared between path tracer, rateriser and soon Strata.

diff --git a/Engine/Shaders/Private/BRDF.ush b/Engine/Shaders/Private/BRDF.ush
--- a/Engine/Shaders/Private/BRDF.ush
+++ b/Engine/Shaders/Private/BRDF.ush
@@ -675,3 +675,40 @@ float Vis_Ashikhmin(float NoV, float NoL)
 	return rcp(4 * (NoL + NoV - NoL * NoV));
 }
 
+// Special medium transmittance evaluation for the clear coat BSDF.
+float3 SimpleClearCoatTransmittance(float NoL, float NoV, float Metallic, float3 BaseColor)
+{
+	float3 Transmittance = 1.0;
+
+	float ClearCoatCoverage = Metallic;
+	if (ClearCoatCoverage > 0.0)
+	{
+		float LayerThickness = 1.0;	// Assume a normalized thickness
+
+		// The path length inside the medium is from light to bottom surface, then back from bottom surface to the viewer.
+		// This assuming an infinite and flat medium slab.
+		float ThinDistance = LayerThickness * (rcp(NoV) + rcp(NoL));
+
+		// BaseColor represents reflected color viewed at 0 incidence angle for the view and light direction... 
+		float3 TransmittanceColor = Diffuse_Lambert(BaseColor);
+		// ... and after being affected by the substrate transmittance when traveling through it, back and forth.
+		// Because of this, extinction is normalized by traveling through layer thickness twice
+		float3 ExtinctionCoefficient = -log(TransmittanceColor) / (2.0 * LayerThickness);
+
+		// Optical depth is evaluated for ThinDistance computed above. But we also subtract 2.0 * LayerThickness
+		// in order to respect the fact that BaseColor represents reflected color viewed at 0 incidence angle. => In this case, OpticalDepth will be 0 and thus Transmittance=1.
+		float3 OpticalDepth = ExtinctionCoefficient * max(ThinDistance - 2.0 * LayerThickness, 0.0);
+
+		// Compute transmittance as a fonction of OpticalDepth
+		Transmittance = exp(-OpticalDepth);
+		// And as a fucntion of the ClearCoatCoverage.
+		Transmittance = lerp(1.0, Transmittance, ClearCoatCoverage);
+	}
+
+	// Note: this is a really special way to handle clear coat transmittance.
+	//  - The basecolor of the bottom layer is used as the transmittance color.
+	//	- So basically, a colored metal as bottom layer will look even more saturated at grazing angle.
+	//	- Transmittance is only computed for metals, for some reason plastic do not receive transmittance.
+
+	return Transmittance;
+}
diff --git a/Engine/Shaders/Private/PathTracing/Material/PathTracingClearCoat.ush b/Engine/Shaders/Private/PathTracing/Material/PathTracingClearCoat.ush
--- a/Engine/Shaders/Private/PathTracing/Material/PathTracingClearCoat.ush
+++ b/Engine/Shaders/Private/PathTracing/Material/PathTracingClearCoat.ush
@@ -17,27 +17,6 @@
 #include "PathTracingFresnel.ush"
 #include "PathTracingEnergyConservation.ush"
 
-float3 CalcThinTransmission(float NoL, float NoV, FPathTracingPayload Payload)
-{
-	float3 Transmission = 1.0;
-
-	float AbsorptionMix = Payload.Metallic;
-	if (AbsorptionMix > 0.0)
-	{
-		float LayerThickness = 1.0;	// Assume normalized thickness
-		float ThinDistance = LayerThickness * (rcp(NoV) + rcp(NoL));
-
-		// Base color represents reflected color viewed at 0 incidence angle, after being absorbed through the substrate.
-		// Because of this, extinction is normalized by traveling through layer thickness twice
-		float3 TransmissionColor = Diffuse_Lambert(Payload.BaseColor);
-		float3 ExtinctionCoefficient = -log(TransmissionColor) / (2.0 * LayerThickness);
-		float3 OpticalDepth = ExtinctionCoefficient * max(ThinDistance - 2.0 * LayerThickness, 0.0);
-		Transmission = exp(-OpticalDepth);
-		Transmission = lerp(1.0, Transmission, AbsorptionMix);
-	}
-	return Transmission;
-}
-
 float CalcClearCoatSelectionPdf(float CoatEnergy)
 {
 #if EVAL_TOP_LAYER == 0
@@ -112,7 +91,7 @@ FMaterialEval ClearCoat_EvalMaterial(
 	Payload.WorldNormal = N_World;
 
 	// Calculate transmission through substrate
-	float3 Transmission = CalcThinTransmission(NoL, NoV, Payload);
+	float3 Transmission = SimpleClearCoatTransmittance(NoL, NoV, Payload.Metallic, Payload.BaseColor);
 
 	BottomEval.Weight *= lerp(1.0, Transmission, Payload.GetClearCoat());
 	BottomEval.Weight *= 1.0 - Coat.E; // remove energy reflected by clearcoat to maintain energy preservation
diff --git a/Engine/Shaders/Private/ShadingModels.ush b/Engine/Shaders/Private/ShadingModels.ush
--- a/Engine/Shaders/Private/ShadingModels.ush
+++ b/Engine/Shaders/Private/ShadingModels.ush
@@ -265,28 +265,6 @@ float3 SimpleShading( float3 DiffuseColor, float3 SpecularColor, float Roughness
 		(D * Vis) * F * ComputeEnergyConservation(EnergyTerms);
 }
 
-float3 CalcThinTransmission(float NoL, float NoV, FGBufferData GBuffer)
-{
-	float3 Transmission = 1.0;
-	float AbsorptionMix = GBuffer.Metallic;
-	if (AbsorptionMix > 0.0)
-	{
-		// Normalized layer thickness documented for clarity
-		float NormalizedLayerThickness = 1.0;
-		float ThinDistance = NormalizedLayerThickness * (rcp(NoV) + rcp(NoL));
-
-		// Base color represents reflected color viewed at 0 incidence angle, after being absorbed through the substrate.
-		// Because of this, extinction is normalized by traveling through layer thickness twice
-		float3 TransmissionColor = Diffuse_Lambert(GBuffer.BaseColor);
-		float3 ExtinctionCoefficient = -log(TransmissionColor) / (2.0 * NormalizedLayerThickness);
-		float3 OpticalDepth = ExtinctionCoefficient * max(ThinDistance - 2.0 * NormalizedLayerThickness, 0.0);
-		Transmission = saturate(exp(-OpticalDepth));
-		Transmission = lerp(1.0, Transmission, AbsorptionMix);
-	}
-
-	return Transmission;
-}
-
 float RefractBlend(float VoH, float Eta)
 {
 	// Refraction blend factor for normal component of VoH
@@ -482,7 +460,7 @@ FDirectLighting ClearCoatBxDF( FGBufferData GBuffer, half3 N, half3 V, half3 L,
 	FBxDFEnergyTerms EnergyTermsBottom = ComputeGGXSpecEnergyTerms(GBuffer.Roughness, BottomContext.NoV, GBuffer.SpecularColor);
 
 	// Absorption
-	float3 Transmission = CalcThinTransmission(BottomContext.NoL, BottomContext.NoV, GBuffer);
+	float3 Transmission = SimpleClearCoatTransmittance(BottomContext.NoL, BottomContext.NoV, GBuffer.Metallic, GBuffer.BaseColor);
 
 	// Default Lit
 	float3 DefaultDiffuse = (Falloff * NoL) * AreaLight.FalloffColor * Diffuse_Lambert(GBuffer.DiffuseColor) * ComputeEnergyPreservation(EnergyTermsBottom);
