Fix ShadingModels unpacking due to previous change.

diff --git a/Engine/Shaders/Private/Strata/Strata.ush b/Engine/Shaders/Private/Strata/Strata.ush
--- a/Engine/Shaders/Private/Strata/Strata.ush
+++ b/Engine/Shaders/Private/Strata/Strata.ush
@@ -3363,7 +3363,7 @@ FStrataBSDF UnpackStrataBSDFIn(FStrataMaterialContainer StrataBuffer, inout FStr
 	#endif
 		OutBSDF.State = (Data>>HEADER_SINGLEENCODING_BIT_COUNT) & BSDF_SINGLEENCODING_MASK;
 		OutBSDF.LuminanceWeightV = 1.0f;
-		const uint BSDFType = firstbitlow(GetShadingModels(StrataHeader));
+		const uint BSDFType = GetShadingModels(StrataHeader);
 		BSDF_SETTYPE(OutBSDF, BSDFType);
 
 		// For now, single encode doesn't support anisotropy, as the top layer data only contains the normal, not the full frame basis
