Strata - visualise errors are more explicit

diff --git a/Engine/Shaders/Private/Strata/StrataVisualize.usf b/Engine/Shaders/Private/Strata/StrataVisualize.usf
--- a/Engine/Shaders/Private/Strata/StrataVisualize.usf
+++ b/Engine/Shaders/Private/Strata/StrataVisualize.usf
@@ -113,7 +113,7 @@ void PrintPixelType(inout FShaderPrintContext Ctx, in FStrataPixelHeader Header,
 	else if (IsSingleMaterial(Header))		Print(Ctx, TEXT("Single  "), InColor);
 	else if (IsComplexMaterial(Header))		Print(Ctx, TEXT("Complex "), InColor);
 	else if (IsSingleLayerWater(Header))	Print(Ctx, TEXT("SLWater "), InColor);
-	else									Print(Ctx, TEXT("!ERROR! "), FontRed);
+	else									Print(Ctx, TEXT("Error - Unkown pixel type"), FontRed);
 }
 
 void PrintSSSType(inout FShaderPrintContext Ctx, uint SSSType, FFontColor InColor)
@@ -336,7 +336,7 @@ void PrintBSDF(inout FShaderPrintContext Context, inout FStrataAddressing Strata
 
 	default:
 		{
-			Print(Context, TEXT("Error"), FontRed);
+			Print(Context, TEXT("Error - Uknown BSDF type"), FontRed);
 		}
 		break;
 	}
