Add Strata custom Eye BSDF.

This allows to have 'fast'/single layer path eye which support dual normal & caustic computation. This first version does not have SSS profile support and is missing the alternative iris normal decoding. This will be added in subsequent CLs.

diff --git a/Engine/Shaders/Private/DBufferDecalShared.ush b/Engine/Shaders/Private/DBufferDecalShared.ush
--- a/Engine/Shaders/Private/DBufferDecalShared.ush
+++ b/Engine/Shaders/Private/DBufferDecalShared.ush
@@ -225,6 +225,18 @@ void ApplyDBufferData(in FStrataDBuffer StrataDBufferData, inout FStrataPixelHea
 				break;
 			}
 
+			case STRATA_BSDF_TYPE_EYE:
+			{
+			#if STRATA_DBUFFER_RESPONSE_BASECOLOR
+				EYE_DIFFUSEALBEDO(BSDF) = EYE_DIFFUSEALBEDO(BSDF) * StrataDBufferData.OneMinusCoverage_BaseColor + StrataDBufferData.BaseColor;
+			#endif
+
+			#if STRATA_DBUFFER_RESPONSE_ROUGHNESS
+				EYE_ROUGHNESS(BSDF) = EYE_ROUGHNESS(BSDF) * StrataDBufferData.OneMinusCoverage_Roughness + StrataDBufferData.Roughness;
+			#endif
+				break;
+			}
+
 			case STRATA_BSDF_TYPE_SINGLELAYERWATER:
 			{
 			#if STRATA_DBUFFER_RESPONSE_BASECOLOR
diff --git a/Engine/Shaders/Private/ShadowProjectionPixelShader.usf b/Engine/Shaders/Private/ShadowProjectionPixelShader.usf
--- a/Engine/Shaders/Private/ShadowProjectionPixelShader.usf
+++ b/Engine/Shaders/Private/ShadowProjectionPixelShader.usf
@@ -252,7 +252,7 @@ void Main(
 		FStrataPixelHeader StrataHeader = UnpackStrataHeaderIn(Strata.MaterialTextureArray, StrataAddressing, Strata.TopLayerTexture);
 		if (IsStrataMaterial(StrataHeader))
 		{
-			bIsEyeOrHair				= StrataGetBSDFType(StrataHeader) == STRATA_BSDF_TYPE_HAIR;
+			bIsEyeOrHair				= StrataGetBSDFType(StrataHeader) == STRATA_BSDF_TYPE_HAIR || StrataGetBSDFType(StrataHeader) == STRATA_BSDF_TYPE_EYE;
 			bHasSubsurface				= HasSubsurface(StrataHeader) || bIsEyeOrHair;
 			bHasSubsurfaceTransmission	= HasSubsurface(StrataHeader);
 		}
diff --git a/Engine/Shaders/Private/Strata/Strata.ush b/Engine/Shaders/Private/Strata/Strata.ush
--- a/Engine/Shaders/Private/Strata/Strata.ush
+++ b/Engine/Shaders/Private/Strata/Strata.ush
@@ -145,7 +145,7 @@ struct FStrataBSDF
 // * Simple : store a single Slab BSDF (with no extra feature)
 // 
 // Header State (8bits):
-// * MATERIALMODE
+// * MATERIALMODE (None/SlabSimple/SlabSingle/SlabComplex/SLW/Hair/Eye)
 // * HASPRECSHADOWMASK			
 // * ZEROPRECSHADOWMASK			
 // * CASTCONTACTSHADOW			
@@ -259,7 +259,8 @@ struct FStrataBSDF
 #define HEADER_MATERIALMODE_SLAB_COMPLEX		3
 #define HEADER_MATERIALMODE_SLWATER				4
 #define HEADER_MATERIALMODE_HAIR				5
-// Up to 8 material modes: 2 remains.
+#define HEADER_MATERIALMODE_EYE					6
+// Up to 8 material modes: 1 remains.
 
 ///////////////////////////////////////////////////////////////////////////////
 // BSDF State
@@ -437,6 +438,15 @@ struct FStrataBSDF
 #define HAIR_BACKLIT(X)					X.VGPRs[1].z
 #define HAIR_COMPLEXTRANSMITTANCE(X)	X.VGPRs[1].w
 
+///////////////////////////////////////////////////////////////////////////////
+// Eye BSDF members
+#define EYE_DIFFUSEALBEDO(X)			X.VGPRs[0].xyz
+#define EYE_F0(X)						0.04f
+#define EYE_ROUGHNESS(X)				X.VGPRs[0].w
+#define EYE_IRISDISTANCE(X)				X.VGPRs[1].x
+#define EYE_IRISMASK(X)					X.VGPRs[1].y
+#define EYE_IRISNORMAL(X)				X.VGPRs[2].xyz
+
 ///////////////////////////////////////////////////////////////////////////////
 // Single Layer Water BSDF members
 #define SLW_BASECOLOR(X)				X.VGPRs[0].xyz
@@ -1418,6 +1428,25 @@ FStrataData GetStrataHairBSDF(float3 BaseColor, float Scatter, float Specular, f
 	return StrataData;
 }
 
+FStrataData GetStrataEyeBSDF(float3 DiffuseAlbedo, float Roughness, float IrisMask, float IrisDistance, float3 IrisNormal, float3 Emissive, uint SharedLocalBasisIndex)
+{
+	FStrataData StrataData = GetInitialisedStrataData();
+	BSDF_SETTYPE					(StrataData.InlinedBSDF, STRATA_BSDF_TYPE_EYE);
+	BSDF_SETSHAREDLOCALBASISID		(StrataData.InlinedBSDF, SharedLocalBasisIndex);
+	BSDF_SETEMISSIVE				(StrataData.InlinedBSDF, Emissive);
+	EYE_DIFFUSEALBEDO				(StrataData.InlinedBSDF) = DiffuseAlbedo;
+	EYE_ROUGHNESS					(StrataData.InlinedBSDF) = Roughness;
+	EYE_IRISDISTANCE				(StrataData.InlinedBSDF) = IrisDistance;
+	EYE_IRISMASK					(StrataData.InlinedBSDF) = IrisMask;
+	EYE_IRISNORMAL					(StrataData.InlinedBSDF) = IrisNormal;
+
+#if STRATA_INLINE_SHADING
+	StrataData.InlinedBSDF.Coverage = 1.0f;
+#endif
+
+	return StrataData;
+}
+
 FStrataData GetStrataSingleLayerWaterBSDF(
 	float3 BaseColor, float Metallic, float Specular, float Roughness, float3 Emissive, float TopMaterialOpacity,
 	float3 WaterAlbedo, float3 WaterExtinction, float WaterPhaseG, float3 ColorScaleBehindWater, uint SharedLocalBasisIndex)
@@ -1697,6 +1726,7 @@ bool IsSingleMaterial(const FStrataPixelHeader In)											{ return (In.State
 bool IsComplexMaterial(const FStrataPixelHeader In)											{ return (In.State & HEADER_MASK_MATERIALMODE) == HEADER_MATERIALMODE_SLAB_COMPLEX && !STRATA_FASTPATH && !STRATA_SINGLEPATH; }
 bool IsSingleLayerWater(const FStrataPixelHeader In)										{ return (In.State & HEADER_MASK_MATERIALMODE) == HEADER_MATERIALMODE_SLWATER; }
 bool IsHair(const FStrataPixelHeader In)													{ return (In.State & HEADER_MASK_MATERIALMODE) == HEADER_MATERIALMODE_HAIR; }
+bool IsEye(const FStrataPixelHeader In)														{ return (In.State & HEADER_MASK_MATERIALMODE) == HEADER_MATERIALMODE_EYE; }
 uint StrataGetBSDFType(const FStrataPixelHeader In)												
 { 
 	switch (GetMaterialMode(In))
@@ -1706,6 +1736,7 @@ uint StrataGetBSDFType(const FStrataPixelHeader In)
 	case HEADER_MATERIALMODE_SLAB_COMPLEX:	return STRATA_BSDF_TYPE_SLAB;
 	case HEADER_MATERIALMODE_SLWATER:		return STRATA_BSDF_TYPE_SINGLELAYERWATER;
 	case HEADER_MATERIALMODE_HAIR:			return STRATA_BSDF_TYPE_HAIR;
+	case HEADER_MATERIALMODE_EYE:			return STRATA_BSDF_TYPE_EYE;
 	default:								return STRATA_BSDF_TYPE_SLAB;
 	}
 }
@@ -2348,6 +2379,16 @@ void SanitizeStrataHair(inout FStrataBSDF BSDF)
 	HAIR_BACKLIT(BSDF) = saturate(HAIR_BACKLIT(BSDF));
 }
 
+void SanitizeStrataEye(inout FStrataBSDF BSDF)
+{
+	BSDF_SETEMISSIVE(BSDF, clamp(BSDF_GETEMISSIVE(BSDF), 0.0, Max10BitsFloat));
+	EYE_DIFFUSEALBEDO(BSDF) = saturate(EYE_DIFFUSEALBEDO(BSDF));
+	EYE_ROUGHNESS(BSDF) = SanitizeRoughness(EYE_ROUGHNESS(BSDF));
+	EYE_IRISDISTANCE(BSDF) = saturate(EYE_IRISDISTANCE(BSDF));
+	EYE_IRISMASK(BSDF) = saturate(EYE_IRISMASK(BSDF));
+	EYE_IRISNORMAL(BSDF) = normalize(EYE_IRISNORMAL(BSDF));
+}
+
 void SanitizeStrataSingleLayerWater(inout FStrataBSDF BSDF)
 {
 	BSDF_SETEMISSIVE(BSDF, clamp(BSDF_GETEMISSIVE(BSDF), 0.0, Max10BitsFloat));
@@ -2378,6 +2419,11 @@ void StrataSanitizeBSDF(inout FStrataBSDF BSDF)
 			SanitizeStrataHair(BSDF);
 			break;
 		}
+		case STRATA_BSDF_TYPE_EYE:
+		{
+			SanitizeStrataEye(BSDF);
+			break;
+		}
 		case STRATA_BSDF_TYPE_SINGLELAYERWATER:
 		{
 			SanitizeStrataSingleLayerWater(BSDF);
@@ -2492,6 +2538,7 @@ float3 StrataGetBSDFBaseColor(in FStrataBSDF BSDF)
 	{
 	case STRATA_BSDF_TYPE_SLAB: 			return lerp(SLAB_DIFFUSEALBEDO(BSDF), SLAB_F0(BSDF), F0RGBToMetallic(SLAB_F0(BSDF)));  // This BaseColor recovery based on Metallic recovers is not accurate, and only used for legacy conversion purpose (e.g., PostProcessMaterial texture access)
 	case STRATA_BSDF_TYPE_HAIR: 			return HAIR_BASECOLOR(BSDF);
+	case STRATA_BSDF_TYPE_EYE: 				return EYE_DIFFUSEALBEDO(BSDF);
 	case STRATA_BSDF_TYPE_SINGLELAYERWATER: return SLW_BASECOLOR(BSDF);
 	}
 	return 0.f;
@@ -2503,6 +2550,7 @@ float3 StrataGetBSDFDiffuseColor(in FStrataBSDF BSDF)
 	{
 	case STRATA_BSDF_TYPE_SLAB: 			return SLAB_DIFFUSEALBEDO(BSDF);
 	case STRATA_BSDF_TYPE_HAIR: 			return HAIR_BASECOLOR(BSDF);
+	case STRATA_BSDF_TYPE_EYE: 				return EYE_DIFFUSEALBEDO(BSDF);
 	case STRATA_BSDF_TYPE_SINGLELAYERWATER: return SLW_BASECOLOR(BSDF);
 	}
 	return 0.f;
@@ -2514,6 +2562,7 @@ float3 StrataGetBSDFSpecularF0(in FStrataBSDF BSDF)
 	{
 	case STRATA_BSDF_TYPE_SLAB: 			return SLAB_F0(BSDF);
 	case STRATA_BSDF_TYPE_HAIR: 			return DielectricSpecularToF0(HAIR_SPECULAR(BSDF));
+	case STRATA_BSDF_TYPE_EYE: 				return EYE_F0(BSDF);
 	case STRATA_BSDF_TYPE_SINGLELAYERWATER: return DielectricSpecularToF0(SLW_SPECULAR(BSDF));
 	}
 	return 0.f;
@@ -2526,6 +2575,7 @@ float3 StrataGetBSDFSpecularF90(in FStrataBSDF BSDF)
 	{
 	case STRATA_BSDF_TYPE_SLAB: 			return SLAB_F90(BSDF);
 	case STRATA_BSDF_TYPE_HAIR: 			return 1.0f;
+	case STRATA_BSDF_TYPE_EYE: 				return 1.0f;
 	case STRATA_BSDF_TYPE_SINGLELAYERWATER: return 1.0f;
 	}
 	return 0.f;
@@ -2535,8 +2585,9 @@ void StrataSetBSDFSpecularF0(in FStrataBSDF BSDF, float3 F0)
 {
 	switch (BSDF_GETTYPE(BSDF))
 	{
-	case STRATA_BSDF_TYPE_SLAB: 			SLAB_F0(BSDF) = F0;
+	case STRATA_BSDF_TYPE_SLAB: 			SLAB_F0(BSDF) = F0; break;
 	case STRATA_BSDF_TYPE_HAIR: 			HAIR_SPECULAR(BSDF) = F0ToDielectricSpecular(max3(F0.x, F0.y, F0.z)); break;
+	case STRATA_BSDF_TYPE_EYE: 				/* EYE_F0 us hardcoded */ break;
 	case STRATA_BSDF_TYPE_SINGLELAYERWATER: SLW_SPECULAR(BSDF) = F0ToDielectricSpecular(max3(F0.x, F0.y, F0.z)); break;
 	}
 }
@@ -2545,8 +2596,9 @@ void StrataSetBSDFSpecularF90(in FStrataBSDF BSDF, float3 F90)
 {
 	switch (BSDF_GETTYPE(BSDF))
 	{
-	case STRATA_BSDF_TYPE_SLAB: 			SLAB_F90(BSDF) = F90;
+	case STRATA_BSDF_TYPE_SLAB: 			SLAB_F90(BSDF) = F90; break;
 	case STRATA_BSDF_TYPE_HAIR: 			break;
+	case STRATA_BSDF_TYPE_EYE: 				break;
 	case STRATA_BSDF_TYPE_SINGLELAYERWATER: break;
 	}
 }
@@ -2557,6 +2609,7 @@ float3 StrataGetBSDFSubSurfaceColor(in FStrataBSDF BSDF)
 	{
 	case STRATA_BSDF_TYPE_SLAB: 			return BSDF_GETSSSTYPE(BSDF) != SSS_TYPE_INVALID ? SLAB_DIFFUSEALBEDO(BSDF) : 0; // A more accurate version would be IsotropicMediumSlabPunctualDirectionalAlbedo(StrataSlabCreateParticipatingMedia(SLAB_DIFFUSEALBEDO(BSDF), SLAB_SSSMFP(BSDFContext.BSDF)));
 	case STRATA_BSDF_TYPE_HAIR: 			return 0;
+	case STRATA_BSDF_TYPE_EYE: 				return EYE_DIFFUSEALBEDO(BSDF);
 	case STRATA_BSDF_TYPE_SINGLELAYERWATER: return 0;
 	}
 	return 0.f;
@@ -2568,6 +2621,7 @@ float StrataGetBSDFSpecular(in FStrataBSDF BSDF)
 	{
 	case STRATA_BSDF_TYPE_SLAB: 			return F0ToDielectricSpecular(SLAB_F0(BSDF).x);
 	case STRATA_BSDF_TYPE_HAIR: 			return HAIR_SPECULAR(BSDF);
+	case STRATA_BSDF_TYPE_EYE: 				return F0ToDielectricSpecular(EYE_F0(BSDF));
 	case STRATA_BSDF_TYPE_SINGLELAYERWATER: return SLW_SPECULAR(BSDF);
 	}
 	return 0.f;
@@ -2579,6 +2633,7 @@ float StrataGetBSDFMetallic(in FStrataBSDF BSDF)
 	{
 	case STRATA_BSDF_TYPE_SLAB: 			return F0RGBToMetallic(SLAB_F0(BSDF)); // Metallic recovers is not accurate, and only used for legacy conversion purpose (e.g., PostProcessMaterial texture access)
 	case STRATA_BSDF_TYPE_HAIR: 			return 0;
+	case STRATA_BSDF_TYPE_EYE: 				return 0;
 	case STRATA_BSDF_TYPE_SINGLELAYERWATER: return 0;
 	}
 	return 0.f;
@@ -2590,6 +2645,7 @@ float StrataGetBSDFRoughness(in FStrataBSDF BSDF)
 	{
 	case STRATA_BSDF_TYPE_SLAB: 			return SLAB_ROUGHNESS(BSDF);
 	case STRATA_BSDF_TYPE_HAIR: 			return HAIR_ROUGHNESS(BSDF);
+	case STRATA_BSDF_TYPE_EYE: 				return EYE_ROUGHNESS(BSDF);
 	case STRATA_BSDF_TYPE_SINGLELAYERWATER: return SLW_ROUGHNESS(BSDF);
 	}
 	return 0.f;
@@ -2601,6 +2657,7 @@ void StrataSetBSDFRoughness(inout FStrataBSDF BSDF, in float Roughness)
 	{
 	case STRATA_BSDF_TYPE_SLAB: 			SLAB_ROUGHNESS(BSDF) = Roughness; break;
 	case STRATA_BSDF_TYPE_HAIR: 			HAIR_ROUGHNESS(BSDF) = Roughness; break;
+	case STRATA_BSDF_TYPE_EYE: 				EYE_ROUGHNESS(BSDF)  = Roughness; break;
 	case STRATA_BSDF_TYPE_SINGLELAYERWATER: SLW_ROUGHNESS(BSDF)  = Roughness; break;
 	}
 }
@@ -2611,6 +2668,7 @@ float StrataGetBSDFAnisotropy(in FStrataBSDF BSDF)
 	{
 	case STRATA_BSDF_TYPE_SLAB: 			return SLAB_ANISOTROPY(BSDF);
 	case STRATA_BSDF_TYPE_HAIR: 			return 0;
+	case STRATA_BSDF_TYPE_EYE: 				return 0;
 	case STRATA_BSDF_TYPE_SINGLELAYERWATER: return 0;
 	}
 	return 0.f;
@@ -2655,6 +2713,7 @@ float StrataGetLegacyShadingModels(in FStrataBSDF BSDF)
 		else													{ return SHADINGMODELID_DEFAULT_LIT; }
 	}
 	case STRATA_BSDF_TYPE_HAIR: 								{ return SHADINGMODELID_HAIR; }
+	case STRATA_BSDF_TYPE_EYE: 									{ return SHADINGMODELID_EYE; }
 	case STRATA_BSDF_TYPE_SINGLELAYERWATER:						{ return SHADINGMODELID_SINGLELAYERWATER; }
 	}
 	return 0.f;
@@ -3459,6 +3518,17 @@ FStrataBSDF UnpackStrataBSDFIn(FStrataMaterialContainer StrataBuffer, inout FStr
 		// 8 bytes
 	}
 	break;
+	case STRATA_BSDF_TYPE_EYE:
+	{
+		StrataLoad_ColorGamma2ToLinearAlphaLinear(StrataBuffer, StrataAddressing, EYE_DIFFUSEALBEDO(OutBSDF), EYE_ROUGHNESS(OutBSDF));		
+		float4 Data = UnpackRGBA8(StrataLoadUint1(StrataBuffer, StrataAddressing));
+		EYE_IRISMASK(OutBSDF) = Data.x;
+		EYE_IRISDISTANCE(OutBSDF) = Data.y;
+		EYE_IRISNORMAL(OutBSDF) = OctahedronToUnitVector(Data.zw * 2.0f - 1.f);
+		// STRATA_TODO - Add SSS profile support
+		// 8 bytes
+	}
+	break;
 	case STRATA_BSDF_TYPE_SINGLELAYERWATER:
 	{
 		StrataLoad_ColorGamma2ToLinearAlphaLinear(StrataBuffer, StrataAddressing, SLW_BASECOLOR(OutBSDF), DummyFloat);
diff --git a/Engine/Shaders/Private/Strata/StrataEvaluation.ush b/Engine/Shaders/Private/Strata/StrataEvaluation.ush
--- a/Engine/Shaders/Private/Strata/StrataEvaluation.ush
+++ b/Engine/Shaders/Private/Strata/StrataEvaluation.ush
@@ -694,6 +694,140 @@ FStrataEvaluateResult StrataEvaluateBSDFCommon(FStrataBSDFContext BSDFContext, F
 			Sample.IntegratedSpecularValue = ShadowTerms.TransmissionShadow * AreaLightContext.AreaLight.FalloffColor * AreaLightContext.Falloff * Sample.SpecularPathValue;
 		}
 		break;
+		
+		case STRATA_BSDF_TYPE_EYE:
+		{
+			float3 DiffuseColor	= EYE_DIFFUSEALBEDO(BSDFContext.BSDF);
+			float3 F0			= EYE_F0(BSDFContext.BSDF);
+			float3 F90			= 1.0f;
+			const float SafeRoughness	= MakeRoughnessSafe(EYE_ROUGHNESS(BSDFContext.BSDF));
+
+			if (Settings.bForceFullyRough)
+			{
+				// When rendering reflection captures, the BSDF roughness has already been forced to 1 using View.RoughnessOverrideParameter (see StrataSanitizeBSDF).
+				EnvBRDFApproxFullyRough(DiffuseColor, F0, F90);
+			}
+
+			float Alpha2 = Pow4(SafeRoughness);
+
+			Init(BSDFContext.Context, BSDFContext.N, BSDFContext.V, AreaLightContext.L);
+			SphereMaxNoH(BSDFContext.Context, AreaLightContext.AreaLight.SphereSinAlpha, true);
+
+			BSDFContext.Context.NoV = saturate(max(abs(BSDFContext.Context.NoV), STRATA_EPSILON));
+
+			////
+			//// Evaluate the diffuse component.
+			////
+
+			const float IrisNoL = saturate(dot(EYE_IRISNORMAL(BSDFContext.BSDF), BSDFContext.L));
+
+			// Blend in the negative intersection normal to create some concavity
+			// Not great as it ties the concavity to the convexity of the cornea surface
+			// No good justification for that. On the other hand, if we're just looking to
+			// introduce some concavity, this does the job.
+			const float3 CausticNormal = normalize(lerp(EYE_IRISNORMAL(BSDFContext.BSDF), -BSDFContext.N, EYE_IRISMASK(BSDFContext.BSDF) * EYE_IRISDISTANCE(BSDFContext.BSDF)));
+
+			// Add-hoc legacy shading mode for eye. 
+			const float Power = lerp(12, 1, IrisNoL);
+			const float Caustic = 0.8 + 0.2 * (Power + 1) * pow(saturate(dot(CausticNormal, BSDFContext.L)), Power);
+
+			const float Iris   = IrisNoL * Caustic;
+			const float Sclera = BSDFContext.Context.NoL;
+
+			Sample.DiffusePathValue = Diffuse_Lambert(DiffuseColor) * lerp(Sclera, Iris, EYE_IRISMASK(BSDFContext.BSDF));
+			Sample.IntegratedDiffuseValue += ShadowTerms.SurfaceShadow * AreaLightContext.NoL * Sample.DiffusePathValue * AreaLightContext.Falloff * AreaLightContext.AreaLight.FalloffColor;
+
+			Sample.DiffuseColor		= DiffuseColor;
+			Sample.DiffusePDF		= BSDFContext.SatNoL / PI;
+			Sample.bSubsurface		= BSDF_GETSSSTYPE(BSDFContext.BSDF) == SSS_TYPE_DIFFUSION || BSDF_GETSSSTYPE(BSDFContext.BSDF) == SSS_TYPE_DIFFUSION_PROFILE;
+
+			//
+			// Apply energy conservation on the diffuse component
+			// If the specular layer is anisotropic, the energy term is computed onto the 'main' roughness [Kulla 2019]
+			//
+			float3 MSScale = 1;
+			float3 DirectionalAlbedo_SpecularTransmission = 1.0f;
+			{
+				FBxDFEnergyTerms EnergyTerms = ComputeGGXSpecEnergyTerms(SafeRoughness, BSDFContext.Context.NoV, F0, F90);
+				DirectionalAlbedo_SpecularTransmission = ComputeEnergyPreservation(EnergyTerms);
+				MSScale = ComputeEnergyConservation(EnergyTerms);
+			}
+			
+			Sample.DiffusePathValue *= DirectionalAlbedo_SpecularTransmission;
+			Sample.IntegratedDiffuseValue *= DirectionalAlbedo_SpecularTransmission;
+
+			////
+			//// Evaluate the specular component.
+			//// This takes into account multiple scattering, micro occlusion.
+			//// Note: anisotropy completely disables area integrations. Lights fall back to punctual.
+			////
+
+			float3 RectLightSpec = 0;
+			float3 RectLightSpecHaze = 0;
+			float D = 0;
+			float Vis = 0;
+			float PDF = 0;
+			const bool bIsRectLight = IntegrationType == INTEGRATION_AREA_LIGHT_RECT;
+
+			if (bIsRectLight)
+			{
+				// In this case, we set D to 1 and Vis will contain the area light / GGX integration
+				RectLightSpec = RectGGXApproxLTC(SafeRoughness, F0, BSDFContext.N, BSDFContext.V, AreaLightContext.AreaLight.Rect, AreaLightContext.AreaLight.Texture);
+
+				const float H_PDF = VisibleGGXPDF(BSDFContext.TangentV, BSDFContext.TangentH, Alpha2);
+				PDF = RayPDFToReflectionRayPDF(BSDFContext.Context.VoH, H_PDF);
+			}
+			else
+			{
+				// Special override for roughness for supporting area light integrator with Sphere/Tube/Disk light, which modifies/increase roughness.
+
+				// Generalized microfacet specular
+				float Alpha2Spec = Alpha2;
+				if(IntegrationType == INTEGRATION_PUNCTUAL_LIGHT)
+				{
+					D = D_GGX(Alpha2Spec, BSDFContext.Context.NoH);
+				}
+				else
+				{
+					const float Energy = EnergyNormalization(Alpha2Spec, BSDFContext.Context.VoH, AreaLightContext.AreaLight);
+					D = D_GGX(Alpha2Spec, BSDFContext.Context.NoH) * Energy;
+				}
+				// STRATA_TODO Should we use Vis_SmithJointApprox here?
+				Vis = Vis_SmithJoint(Alpha2Spec, BSDFContext.Context.NoV, AreaLightContext.NoL);
+				const float H_PDF = VisibleGGXPDF(BSDFContext.TangentV, BSDFContext.TangentH, Alpha2Spec);
+				PDF = RayPDFToReflectionRayPDF(BSDFContext.Context.VoH, H_PDF);
+			}
+
+			const float SpecularMicroOcclusion = F0RGBToMicroOcclusion(F0);
+			const float3 FresnelTerm = SpecularMicroOcclusion * F_Schlick(F0, F90, BSDFContext.Context.VoH);
+
+			Sample.SpecularPathProbability	= 1.0f;
+			Sample.SpecularPDF				= PDF * Sample.SpecularPathProbability;
+
+			if (bIsRectLight)
+			{
+				Sample.SpecularPathValue       = RectLightSpec * MSScale * SpecularMicroOcclusion;
+				Sample.IntegratedSpecularValue = ShadowTerms.SurfaceShadow * /* AreaLightContext.NoL and falloff is part of the LTC integration already */ Sample.SpecularPathValue * Sample.SpecularPathProbability;
+			}
+			else
+			{
+				Sample.SpecularPathValue        = D * Vis * FresnelTerm * MSScale;
+				Sample.IntegratedSpecularValue += ShadowTerms.SurfaceShadow * AreaLightContext.NoL * AreaLightContext.AreaLight.FalloffColor * AreaLightContext.Falloff * Sample.SpecularPathValue * Sample.SpecularPathProbability;
+			}
+
+			////
+			//// Evaluate emissive and set the sample throughput (transmittance to next layer) corresponding to an opaque slab of matter.
+			////
+			
+			// we do not need to add emissive for the BRDF TotalSpec or TotalDiff values as this is handled separately
+			Sample.EmissivePathValue = BSDF_GETEMISSIVE(BSDFContext.BSDF);
+
+			Sample.ThroughputV			= OpaqueBSDFThroughput;
+			Sample.TransmittanceAlongN	= OpaqueBSDFThroughput;
+
+			break;
+		}
+		break;
 
 		//case STRATA_BSDF_TYPE_VOLUMETRICFOGCLOUD:
 		//case STRATA_BSDF_TYPE_UNLIT:
@@ -921,6 +1055,67 @@ FStrataImportanceSampleResult StrataImportanceSampleBSDF(FStrataBSDFContext BSDF
 			break;
 		}
 
+		case STRATA_BSDF_TYPE_EYE:
+		{
+			// TODO manage iris normal for the diffuse term
+			const float3 DiffuseColor	= EYE_DIFFUSEALBEDO(BSDFContext.BSDF);
+			const float3 F0 			= EYE_F0(BSDFContext.BSDF);
+			const float3 F90			= 1.0f;
+			const float SafeRoughness   = MakeRoughnessSafe(EYE_ROUGHNESS(BSDFContext.BSDF));
+
+			// Pick diffuse/specular lobe
+			float TermPDF = Out.Term == SHADING_TERM_DIFFUSE ? 1 : 0;
+			Out.Term = InTermMask;
+			if (Out.Term == (SHADING_TERM_DIFFUSE | SHADING_TERM_SPECULAR))
+			{
+				// Use (main) specular lobe as a heuristic to pick select between specular/diffuse lobe
+				const FBxDFEnergyTerms SpecEnergyTerms = ComputeGGXSpecEnergyTerms(SafeRoughness, BSDFContext.Context.NoV, F0, F90);
+
+				// Probability of picking diffuse lobe vs. specular lobe
+				TermPDF  = StrataLobeSelectionProb(DiffuseColor * (1 - SpecEnergyTerms.E), SpecEnergyTerms.E);
+				Out.Term = StrataSelectLobe(E.x, TermPDF, SHADING_TERM_DIFFUSE, SHADING_TERM_SPECULAR);
+			}
+
+			// Sample Diffuse lobe
+			if (Out.Term == SHADING_TERM_DIFFUSE)
+			{
+				float4 ImportanceSample = CosineSampleHemisphere(E);
+				Out.L   = mul(ImportanceSample.xyz, BSDFContext.TangentBasis);
+				Out.PDF = ImportanceSample.w;
+			}
+
+			// Sample Specular lobe
+			if (Out.Term == SHADING_TERM_SPECULAR)
+			{
+				float4 TangentH = ImportanceSampleVisibleGGX(UniformSampleDisk(E), Pow4(SafeRoughness), BSDFContext.TangentV);
+
+				const float HPDF = TangentH.w;
+				const float3 H = mul(TangentH.xyz, BSDFContext.TangentBasis);
+				const float VoH = saturate(dot(BSDFContext.V, H));
+
+				Out.L   = 2 * dot(BSDFContext.V, H) * H - BSDFContext.V;
+				Out.PDF = RayPDFToReflectionRayPDF(VoH, HPDF);
+			}
+
+			// Evaluate lobes
+			StrataUpdateBSDFContext(BSDFContext, Out.L);
+			const FStrataEvaluateResult Evaluate = StrataEvaluateBSDF(BSDFContext, Settings);
+
+			// Add lobes mixtures
+			Out.PDF = 0;
+			Out.Weight = 0;
+			if (InTermMask & SHADING_TERM_DIFFUSE)
+			{
+				StrataAddLobeWithMIS(Out.Weight, Out.PDF, Evaluate.DiffusePathValue / Evaluate.DiffusePDF, Evaluate.DiffusePDF, TermPDF);
+			}
+
+			if (InTermMask & SHADING_TERM_SPECULAR)
+			{
+				StrataAddLobeWithMIS(Out.Weight, Out.PDF, Evaluate.SpecularPathValue / Evaluate.SpecularPDF, Evaluate.SpecularPDF, (1-TermPDF));
+			}
+			break;
+		}
+
 		//case STRATA_BSDF_TYPE_VOLUMETRICFOGCLOUD:
 		//case STRATA_BSDF_TYPE_UNLIT:
 		//Nothing to do in this case because these BSDF are evaluated in other specialised passes.
@@ -1156,6 +1351,47 @@ FStrataEnvLightResult StrataEvaluateForEnvLight(FStrataBSDFContext BSDFContext,
 			break;
 		}
 
+		case STRATA_BSDF_TYPE_EYE:	
+		{
+			float3 DiffuseColor		= EYE_DIFFUSEALBEDO(BSDFContext.BSDF);
+			float3 F0 				= EYE_F0(BSDFContext.BSDF);
+			float3 F90 				= 1.0f;
+
+			if (Settings.bForceFullyRough)
+			{
+				// When rendering reflection captures, the BSDF roughness has already been forced to 1 using View.RoughnessOverrideParameter (see StrataSanitizeBSDF).
+				EnvBRDFApproxFullyRough(DiffuseColor, F0, F90);
+			}
+
+			StrataEnvLightResult.DiffuseNormal = BSDFContext.N;
+			StrataEnvLightResult.bSubsurface = BSDF_GETSSSTYPE(BSDFContext.BSDF) != SSS_TYPE_INVALID; // TODO profile
+
+			float SafeRoughness = MakeRoughnessSafe(SLAB_ROUGHNESS(BSDFContext.BSDF));
+
+			float3 EvalEnvBRDF = 0;
+			float3 DirectionalAlbedo_SpecularTransmission = 0;
+			{
+				FBxDFEnergyTerms SpecularEnergyTerms = ComputeGGXSpecEnergyTerms(SafeRoughness, BSDFContext.Context.NoV, F0, F90);
+				const float DiffuseDirectionalAlbedo = 1.0f;
+				EvalEnvBRDF									= SpecularEnergyTerms.E;
+				DirectionalAlbedo_SpecularTransmission		= ComputeEnergyPreservation(SpecularEnergyTerms);
+				StrataEnvLightResult.DiffuseWeight			= DiffuseColor * DiffuseDirectionalAlbedo * DirectionalAlbedo_SpecularTransmission;
+				StrataEnvLightResult.DiffuseColor			= DiffuseColor;
+			}
+
+			BRANCH
+			if (bEnableSpecular)
+			{
+				const float SpecularMicroOcclusion = F0RGBToMicroOcclusion(F0);
+				StrataEnvLightResult.SpecularDirection = StrataGetOffSpecularPeakReflectionDir(BSDFContext.N, BSDFContext.R, SafeRoughness);
+				StrataEnvLightResult.SpecularWeight = EvalEnvBRDF * SpecularMicroOcclusion;
+				StrataEnvLightResult.SpecularSafeRoughness = SafeRoughness;
+				StrataEnvLightResult.SpecularColor = F0;
+			}
+
+			break;
+		}
+
 		//case STRATA_BSDF_TYPE_VOLUMETRICFOGCLOUD:
 		//case STRATA_BSDF_TYPE_UNLIT:
 		//case STRATA_BSDF_TYPE_SINGLELAYERWATER:
diff --git a/Engine/Shaders/Private/Strata/StrataExport.ush b/Engine/Shaders/Private/Strata/StrataExport.ush
--- a/Engine/Shaders/Private/Strata/StrataExport.ush
+++ b/Engine/Shaders/Private/Strata/StrataExport.ush
@@ -224,6 +224,19 @@ void PackStrataOut(
 						break;
 					}
 
+					case STRATA_BSDF_TYPE_EYE:
+					{
+						bIsOnlySlab = false;
+						bIsSimpleMaterial = false;
+						bIsSingleMaterial = false;
+
+						TopLayerTotalWeight += BSDF.TopLayerDataWeight;
+						TopLayerData.Roughness += BSDF.TopLayerDataWeight * FullyRough;
+
+						SetMaterialMode(StrataPixelHeader, HEADER_MATERIALMODE_EYE);
+						break;
+					}
+
 					case STRATA_BSDF_TYPE_SINGLELAYERWATER:
 					{
 						bIsOnlySlab = false;
@@ -521,6 +534,14 @@ void PackStrataOut(
 							// 8 bytes
 						}
 						break;
+						case STRATA_BSDF_TYPE_EYE:
+						{
+							const float2 EncodedIrisNormal = UnitVectorToOctahedron(EYE_IRISNORMAL(BSDF));
+							STRATA_STORE_UINT1(PackColorLinearToGamma2AlphaLinear(float4(EYE_DIFFUSEALBEDO(BSDF), EYE_ROUGHNESS(BSDF))));
+							STRATA_STORE_UINT1(PackRGBA8(float4(EYE_IRISMASK(BSDF), EYE_IRISDISTANCE(BSDF), EncodedIrisNormal * 0.5f + 0.5f)));
+							// 8 bytes
+						}
+						break;
 						case STRATA_BSDF_TYPE_SINGLELAYERWATER:
 						{
 							STRATA_STORE_UINT1(PackColorLinearToGamma2AlphaLinear(float4(SLW_BASECOLOR(BSDF), 0.0f)));
diff --git a/Engine/Shaders/Private/Strata/StrataLegacyConversion.ush b/Engine/Shaders/Private/Strata/StrataLegacyConversion.ush
--- a/Engine/Shaders/Private/Strata/StrataLegacyConversion.ush
+++ b/Engine/Shaders/Private/Strata/StrataLegacyConversion.ush
@@ -124,6 +124,9 @@ FStrataData StrataConvertLegacyMaterialDynamic(
 	float3 TransmittanceColor,
 	float3 WaterScatteringCoefficients, float3 WaterAbsorptionCoefficients, float WaterPhaseG, float3 ColorScaleBehindWater,
 	uint ShadingModel,
+	float3 RawNormal,
+	float3 RawTangent,
+	float3 RawClearCoatNormal,
 	uint SharedLocalBasisIndex,
 	uint ClearCoatBottomNormal_SharedLocalBasisIndex,
 	inout uint SharedLocalBasisTypes,
@@ -180,7 +183,25 @@ FStrataData StrataConvertLegacyMaterialDynamic(
 		// No need to take into account opacity weighting for translucent material, as the coverage is assumed to be 1
 	}
 	#endif
+	#if MATERIAL_SHADINGMODEL_EYE
+	if (ShadingModel == SHADINGMODELID_EYE)
+	{
+		const float IrisMask = ClearCoat;
+		const float IrisDistance = ClearCoatRoughness;
+		const float3 IrisNormal = RawClearCoatNormal;
+
+		return GetStrataEyeBSDF(
+			BaseColor,	        // DiffuseAlbedo
+			Roughness,	        // Roughness
+			ClearCoat,	        // IrisMask
+			ClearCoatRoughness, // IrisDistance
+			IrisNormal,         // IrisNormal TODO
+			Emissive,	        // EmissiveColor
+			SharedLocalBasisIndex);
 
+		// No need to take into account opacity weighting for translucent material, as the coverage is assumed to be 1
+	}
+	#endif
 	// Top/Bottom
 	FStrataLegacyParameters Top = InitStrataLegacyParameters(PixelFootprint, SharedLocalBasisIndex, 1.0f);
 	FStrataLegacyParameters Bot = InitStrataLegacyParameters(PixelFootprint, SharedLocalBasisIndex, 1.0f);
@@ -401,48 +422,6 @@ FStrataData StrataConvertLegacyMaterialDynamic(
 		FinalWeight = Opacity;
 #endif
 	}
-	else if (ShadingModel == SHADINGMODELID_EYE)
-	{
-		const float IrisMask = ClearCoat;
-		const float IrisDistance = ClearCoatRoughness;
-
-		// Lerp transmittance rather MFP, as transmittance is perceptually linear unlike MFP
-		const float Thickness = DefaultThickness;
-		const float3 MFP = TransmittanceToMeanFreePath(IrisMask.xxx, Thickness * CENTIMETER_TO_METER) * METER_TO_CENTIMETER;
-
-		// Cornea/Sclera
-		Top.BaseColor = BaseColor;
-		Top.Specular = Specular;
-		Top.Metallic = 0.f;
-		Top.Roughness = Roughness;
-		Top.Anisotropy = 0.f;
-		Top.SSSProfileID = SubSurfaceProfileId;
-		Top.SSSMFP = MFP;
-		Top.SSSMFPScale = Opacity;
-		Top.Emissive = Emissive;
-		Top.SharedLocalBasisIndex = SharedLocalBasisIndex;
-		Top.LayerDepth = 0;
-		Top.bIsBottom = 0;
-		Top.Weight = 1.f;
-
-		// Iris
-		Bot.BaseColor = BaseColor;
-		Bot.Specular = Specular;
-		Bot.Metallic = .0f;
-		Bot.Roughness = Roughness;
-		Bot.Anisotropy = 0.f;
-		Bot.SSSProfileID = SubSurfaceProfileId;
-		Bot.SSSMFP = Zeros;
-		Bot.SSSMFPScale = Opacity;
-		Bot.Emissive = Emissive;
-		Bot.SharedLocalBasisIndex = ClearCoatBottomNormal_SharedLocalBasisIndex;
-		Bot.LayerDepth = 1;
-		Bot.bIsBottom = 1;
-		Bot.Weight = 1.f;
-
-		// No need to take into account opacity weighting for translucent material, as the coverage is assumed to be 1
-		FinalWeight = 1.0f;
-	}
 	else if (ShadingModel == SHADINGMODELID_THIN_TRANSLUCENT)	
 	{
 		const float Thickness = DefaultThickness;
@@ -580,6 +559,9 @@ FStrataData StrataConvertLegacyMaterialStatic(
 	float3 TransmittanceColor,
 	float3 WaterScatteringCoefficients, float3 WaterAbsorptionCoefficients, float WaterPhaseG, float3 ColorScaleBehindWater,
 	uint ShadingModel,
+	float3 RawNormal,
+	float3 RawTangent,
+	float3 RawClearCoatNormal,
 	uint SharedLocalBasisIndex,
 	uint ClearCoatBottomNormal_SharedLocalBasisIndex,
 	inout uint SharedLocalBasisTypes,
@@ -633,11 +615,11 @@ FStrataData StrataConvertLegacyMaterialStatic(
 				ClearCoat,	   // Backlit
 				Emissive,	   // EmissiveColor
 				SharedLocalBasisIndex),
-				StrataTree,
-				0		/*OperatorIndex*/,
-				0		/*BSDFIndex*/,
-				0		/*int LayerDepth*/,
-				1		/*int bIsBottom*/);
+			StrataTree,
+			0		/*OperatorIndex*/,
+			0		/*BSDFIndex*/,
+			0		/*int LayerDepth*/,
+			1		/*int bIsBottom*/);
 
 		FStrataLegacyParameters Bot = InitStrataLegacyParameters(PixelFootprint, SharedLocalBasisIndex, 1.0f);
 		FStrataData Bot_Slab = CreateLegacySlab(Bot, SharedLocalBasisTypes, StrataTree,	/*OperatorIndex*/1, /*BSDFIndex*/1);
@@ -647,6 +629,35 @@ FStrataData StrataConvertLegacyMaterialStatic(
 		Out = StrataVerticalLayering(WTop_Slab, WBot_Slab, StrataTree,					/*OperatorIndex*/4, /*MaxDistanceFromLeaves*/2);
 		// Out maps to LegacyUpdateBSDFsOperators
 	}
+	else if (ShadingModel == SHADINGMODELID_EYE)
+	{
+		const float IrisMask = ClearCoat;
+		const float IrisDistance = ClearCoatRoughness;
+		const float3 IrisNormal = RawClearCoatNormal;
+
+		FStrataData Top_Eye = PromoteParameterBlendedBSDFToOperator(
+			GetStrataEyeBSDF(
+				BaseColor,	   // DiffuseAlbedo
+				Roughness,	   // Roughness
+				IrisMask,	   // IrisMask
+				IrisDistance,  // IrisDistance
+				IrisNormal,	   // IrisNormal
+				Emissive,	   // EmissiveColor
+				SharedLocalBasisIndex),
+			StrataTree,
+			0		/*OperatorIndex*/,
+			0		/*BSDFIndex*/,
+			0		/*int LayerDepth*/,
+			1		/*int bIsBottom*/);
+
+		FStrataLegacyParameters Bot = InitStrataLegacyParameters(PixelFootprint, SharedLocalBasisIndex, 1.0f);
+		FStrataData Bot_Slab = CreateLegacySlab(Bot, SharedLocalBasisTypes, StrataTree,	/*OperatorIndex*/1, /*BSDFIndex*/1);
+
+		FStrataData WTop_Slab = StrataWeight(Top_Eye, 1.0, StrataTree,					/*OperatorIndex*/2, /*MaxDistanceFromLeaves*/1);
+		FStrataData WBot_Slab = StrataWeight(Bot_Slab, 0.0, StrataTree,					/*OperatorIndex*/3, /*MaxDistanceFromLeaves*/1);
+		Out = StrataVerticalLayering(WTop_Slab, WBot_Slab, StrataTree,					/*OperatorIndex*/4, /*MaxDistanceFromLeaves*/2);
+		// Out maps to LegacyUpdateBSDFsOperators
+	}
 	else if (ShadingModel == SHADINGMODELID_SINGLELAYERWATER)
 	{
 		const float3 WaterExtinction = WaterScatteringCoefficients + WaterAbsorptionCoefficients;
@@ -665,11 +676,11 @@ FStrataData StrataConvertLegacyMaterialStatic(
 				WaterPhaseG, 			// WaterPhaseG
 				ColorScaleBehindWater,	// ColorScaleBehindWater
 				SharedLocalBasisIndex),
-				StrataTree,
-				0		/*OperatorIndex*/,
-				0		/*BSDFIndex*/,
-				0		/*int LayerDepth*/,
-				1		/*int bIsBottom*/);
+			StrataTree,
+			0		/*OperatorIndex*/,
+			0		/*BSDFIndex*/,
+			0		/*int LayerDepth*/,
+			1		/*int bIsBottom*/);
 
 		FStrataLegacyParameters Bot = InitStrataLegacyParameters(PixelFootprint, SharedLocalBasisIndex, 1.0f);
 		FStrataData Bot_Slab = CreateLegacySlab(Bot, SharedLocalBasisTypes, StrataTree,	/*OperatorIndex*/1, /*BSDFIndex*/1);
@@ -692,6 +703,9 @@ FStrataData StrataConvertLegacyMaterialStatic(
 			TransmittanceColor,
 			WaterScatteringCoefficients, WaterAbsorptionCoefficients, WaterPhaseG, ColorScaleBehindWater,
 			ShadingModel,
+			RawNormal,
+			RawTangent,
+			RawClearCoatNormal,
 			SharedLocalBasisIndex,
 			ClearCoatBottomNormal_SharedLocalBasisIndex,
 			SharedLocalBasisTypes,
diff --git a/Engine/Shaders/Private/Strata/StrataVisualize.usf b/Engine/Shaders/Private/Strata/StrataVisualize.usf
--- a/Engine/Shaders/Private/Strata/StrataVisualize.usf
+++ b/Engine/Shaders/Private/Strata/StrataVisualize.usf
@@ -113,6 +113,8 @@ void PrintPixelType(inout FShaderPrintContext Ctx, in FStrataPixelHeader Header,
 	else if (IsSingleMaterial(Header))		Print(Ctx, TEXT("Single  "), InColor);
 	else if (IsComplexMaterial(Header))		Print(Ctx, TEXT("Complex "), InColor);
 	else if (IsSingleLayerWater(Header))	Print(Ctx, TEXT("SLWater "), InColor);
+	else if (IsHair(Header))				Print(Ctx, TEXT("Hair    "), InColor);
+	else if (IsEye(Header))					Print(Ctx, TEXT("Eye     "), InColor);
 	else									Print(Ctx, TEXT("Error - Unkown pixel type"), FontRed);
 }
 
@@ -321,6 +323,35 @@ void PrintBSDF(inout FShaderPrintContext Context, inout FStrataAddressing Strata
 		}
 		break;
 
+		case STRATA_BSDF_TYPE_EYE:
+		{
+			Print(Context, TEXT("Eye"), FontBSDFType);
+			Newline(Context, RectMax);
+			Print(Context, TEXT(" NormalID IsTopLayer LuminanceWeight"), FontBSDF);
+			Newline(Context, RectMax);
+
+			Print(Context, BSDF_GETSHAREDLOCALBASISID(BSDF), FontBSDF);
+			PrintBool(Context, BSDF_GETISTOPLAYER(BSDF), FontBSDF);
+			if (BSDF_GETHASGREYWEIGHT_V(BSDF))
+			{
+				Print(Context, BSDF.LuminanceWeightV.x, FontSilver);
+			}
+			else
+			{
+				Print(Context, BSDF.LuminanceWeightV.r, FontLightRed);
+				Print(Context, BSDF.LuminanceWeightV.g, FontLightGreen);
+				Print(Context, BSDF.LuminanceWeightV.b, FontLightBlue);
+			}
+			Newline(Context, RectMax);
+
+			Print(Context, TEXT(" Diffuse Albedo    "), FontBSDFPropName);	Print(Context, EYE_DIFFUSEALBEDO(BSDF), FontBSDFPropValu);			Newline(Context, RectMax);
+			Print(Context, TEXT(" F0                "), FontBSDFPropName);	Print(Context, EYE_F0(BSDF), FontBSDFPropValu);						Newline(Context, RectMax);
+			Print(Context, TEXT(" Roughness         "), FontBSDFPropName);	Print(Context, EYE_ROUGHNESS(BSDF), FontBSDFPropValu);				Newline(Context, RectMax);
+			Print(Context, TEXT(" Iris Mask         "), FontBSDFPropName);	Print(Context, EYE_IRISDISTANCE(BSDF), FontBSDFPropValu);			Newline(Context, RectMax);
+			Print(Context, TEXT(" Iris Distance     "), FontBSDFPropName);	Print(Context, EYE_IRISMASK(BSDF), FontBSDFPropValu);				Newline(Context, RectMax);
+		}
+		break;
+
 		case STRATA_BSDF_TYPE_SINGLELAYERWATER:
 		{
 			Print(Context, TEXT("Water"), FontBSDFType);
diff --git a/Engine/Shaders/Shared/StrataDefinitions.h b/Engine/Shaders/Shared/StrataDefinitions.h
--- a/Engine/Shaders/Shared/StrataDefinitions.h
+++ b/Engine/Shaders/Shared/StrataDefinitions.h
@@ -26,7 +26,8 @@
 #define STRATA_BSDF_TYPE_UNLIT				2
 #define STRATA_BSDF_TYPE_HAIR				3
 #define STRATA_BSDF_TYPE_SINGLELAYERWATER	4
-// When more than 5 BSDF must exists, please update STATE_BIT_COUNT_BSDF and FStrataClassification.ShadingModels packing in Strata.ush
+#define STRATA_BSDF_TYPE_EYE				5
+// When more than 7 BSDF must exists, please update STATE_BIT_COUNT_BSDF and FStrataClassification.ShadingModels packing in Strata.ush
 
 // The size of strata material classification tiles on screen
 #define STRATA_TILE_SIZE					8
diff --git a/Engine/Source/Runtime/Engine/Classes/Engine/EngineTypes.h b/Engine/Source/Runtime/Engine/Classes/Engine/EngineTypes.h
--- a/Engine/Source/Runtime/Engine/Classes/Engine/EngineTypes.h
+++ b/Engine/Source/Runtime/Engine/Classes/Engine/EngineTypes.h
@@ -673,6 +673,7 @@ enum EStrataShadingModel
 	SSM_SubsurfaceLit			UMETA(DisplayName = "SubsurfaceLit"),
 	SSM_VolumetricFogCloud		UMETA(DisplayName = "VolumetricFogCloud"),
 	SSM_Hair					UMETA(DisplayName = "Hair"),
+	SSM_Eye						UMETA(DisplayName = "Eye"),
 	SSM_SingleLayerWater		UMETA(DisplayName = "SingleLayerWater"),
 	SSM_LightFunction			UMETA(DisplayName = "LightFunction"),
 	SSM_PostProcess				UMETA(DisplayName = "PostProcess"),
diff --git a/Engine/Source/Runtime/Engine/Classes/Materials/MaterialExpressionStrata.h b/Engine/Source/Runtime/Engine/Classes/Materials/MaterialExpressionStrata.h
--- a/Engine/Source/Runtime/Engine/Classes/Materials/MaterialExpressionStrata.h
+++ b/Engine/Source/Runtime/Engine/Classes/Materials/MaterialExpressionStrata.h
@@ -477,6 +477,66 @@ class UMaterialExpressionStrataHairBSDF : public UMaterialExpressionStrataBSDF
 	//~ End UMaterialExpression Interface
 };
 
+UCLASS(MinimalAPI, collapsecategories, hidecategories = Object, DisplayName = "Strata Eye BSDF")
+class UMaterialExpressionStrataEyeBSDF : public UMaterialExpressionStrataBSDF
+{
+	GENERATED_UCLASS_BODY()
+		
+	/**
+	 * Hair fiber base color resulting from single and multiple scattering combined. (type = float3, unit = unitless, defaults to black)
+	 */
+	UPROPERTY()
+	FExpressionInput DiffuseColor;
+		
+	/**
+	 * Controls how rough the Material is. Roughness of 0 (smooth) is a mirror reflection and 1 (rough) is completely matte or diffuse (type = float, unit = unitless, defaults to 0.5)
+	 */
+	UPROPERTY()
+	FExpressionInput Roughness;
+
+	/**
+	 * Normal of the sclera and cornea (type = float3, unit = unitless, defaults to +X vector)
+	 */
+	UPROPERTY()
+	FExpressionInput CorneaNormal;
+
+	/**
+	 * Normal of the iris (type = float3, unit = unitless, defaults to +X vector)
+	 */
+	UPROPERTY()
+	FExpressionInput IrisNormal;
+
+	/**
+	 * Mask defining the iris surface (type = float, unit = unitless, defaults to 0.0)
+	 */
+	UPROPERTY()
+	FExpressionInput IrisMask;
+
+	/**
+	 * Distance from the center of the iris (type = float, unit = unitless, defaults to 0.0)
+	 */
+	UPROPERTY()
+	FExpressionInput IrisDistance;
+
+	/**
+	 * Emissive color on top of the surface (type = float3, unit = luminance, defaults to 0.0)
+	 */
+	UPROPERTY()
+	FExpressionInput EmissiveColor;
+
+	//~ Begin UMaterialExpression Interface
+#if WITH_EDITOR
+	virtual int32 Compile(class FMaterialCompiler* Compiler, int32 OutputIndex) override;
+	virtual void GetCaption(TArray<FString>& OutCaptions) const override;
+	virtual uint32 GetOutputType(int32 OutputIndex) override;
+	virtual uint32 GetInputType(int32 InputIndex) override;
+	virtual bool IsResultStrataMaterial(int32 OutputIndex) override;
+	virtual void GatherStrataMaterialInfo(FStrataMaterialInfo& StrataMaterialInfo, int32 OutputIndex) override;
+	virtual FStrataOperator* StrataGenerateMaterialTopologyTree(class FMaterialCompiler* Compiler, class UMaterialExpression* Parent, int32 OutputIndex) override;
+#endif
+	//~ End UMaterialExpression Interface
+};
+
 UCLASS(MinimalAPI, collapsecategories, hidecategories = Object, DisplayName = "Strata Single Layer Water BSDF")
 class UMaterialExpressionStrataSingleLayerWaterBSDF : public UMaterialExpressionStrataBSDF
 {
diff --git a/Engine/Source/Runtime/Engine/Private/Materials/HLSLMaterialTranslator.cpp b/Engine/Source/Runtime/Engine/Private/Materials/HLSLMaterialTranslator.cpp
--- a/Engine/Source/Runtime/Engine/Private/Materials/HLSLMaterialTranslator.cpp
+++ b/Engine/Source/Runtime/Engine/Private/Materials/HLSLMaterialTranslator.cpp
@@ -9736,6 +9736,7 @@ bool FHLSLMaterialTranslator::StrataGenerateDerivedMaterialOperatorData()
 		bool bHasUnlit = false;
 		bool bHasVFogCloud = false;
 		bool bHasHair = false;
+		bool bHasEye = false;
 		bool bHasSLW = false;
 		bool bHasSlab = false;
 
@@ -9787,6 +9788,7 @@ bool FHLSLMaterialTranslator::StrataGenerateDerivedMaterialOperatorData()
 				bHasUnlit |= CurrentOperator.BSDFType == STRATA_BSDF_TYPE_UNLIT;
 				bHasVFogCloud |= CurrentOperator.BSDFType == STRATA_BSDF_TYPE_VOLUMETRICFOGCLOUD;
 				bHasHair |= CurrentOperator.BSDFType == STRATA_BSDF_TYPE_HAIR;
+				bHasEye |= CurrentOperator.BSDFType == STRATA_BSDF_TYPE_EYE;
 				bHasSLW |= CurrentOperator.BSDFType == STRATA_BSDF_TYPE_SINGLELAYERWATER;
 				bHasSlab |= CurrentOperator.BSDFType == STRATA_BSDF_TYPE_SLAB;
 				break;
@@ -9813,13 +9815,13 @@ bool FHLSLMaterialTranslator::StrataGenerateDerivedMaterialOperatorData()
 		bStrataOutputsOpaqueRoughRefractions = bStrataUsesVerticalLayering && bIsOpaqueOrMasked;
 		bStrataMaterialIsUnlitNode = bHasUnlit;
 
-		if ((bHasUnlit || bHasVFogCloud || bHasHair || bHasSLW) && StrataMaterialBSDFCount > 1)
+		if ((bHasUnlit || bHasVFogCloud || bHasHair || bHasEye || bHasSLW) && StrataMaterialBSDFCount > 1)
 		{
 			Errorf(TEXT("Unlit, Fog/Cloud, Hair or SingleLayerWater must be used in isolation. See %s (asset: %s).\r\n"), *Material->GetDebugName(), *Material->GetAssetPath().ToString());
 			// Even though we could support Unlit with slab.
 		}
 
-		if ((bHasUnlit || bHasVFogCloud || bHasHair || bHasSLW) && bOperatorEncountered)
+		if ((bHasUnlit || bHasVFogCloud || bHasHair || bHasEye || bHasSLW) && bOperatorEncountered)
 		{
 			Errorf(TEXT("Unlit, Fog/Cloud, Hair or SingleLayerWater cannot be used with operators. See %s (asset: %s).\r\n"), *Material->GetDebugName(), *Material->GetAssetPath().ToString());
 			// This is because it will results in simpler lighting loops focusin on slab.
@@ -10024,6 +10026,12 @@ bool FHLSLMaterialTranslator::StrataGenerateDerivedMaterialOperatorData()
 					bStrataMaterialIsSingle = false;
 					break;
 				}
+				case STRATA_BSDF_TYPE_EYE:
+				{
+					bStrataMaterialIsSimple = false;
+					bStrataMaterialIsSingle = false;
+					break;
+				}
 				case STRATA_BSDF_TYPE_SINGLELAYERWATER:
 				{
 					bStrataMaterialIsSimple = false;
@@ -10128,6 +10136,12 @@ bool FHLSLMaterialTranslator::StrataGenerateDerivedMaterialOperatorData()
 					StrataMaterialRequestedSizeByte += UintByteSize;
 					break;
 				}
+				case STRATA_BSDF_TYPE_EYE:
+				{
+					StrataMaterialRequestedSizeByte += UintByteSize;
+					StrataMaterialRequestedSizeByte += UintByteSize;
+					break;
+				}
 				case STRATA_BSDF_TYPE_SINGLELAYERWATER:
 				{
 					StrataMaterialRequestedSizeByte += UintByteSize;
@@ -10382,12 +10396,13 @@ int32 FHLSLMaterialTranslator::StrataConversionFromLegacy(
 {
 	const FString NormalCode = GetParameterCode(Normal);
 	const FString TangentCode = Tangent != INDEX_NONE ? *GetParameterCode(Tangent) : TEXT("NONE");
+	const int32 RawTangent = Tangent != INDEX_NONE ? Tangent : Normal;
 
 	const FString ClearCoat_NormalCode = GetParameterCode(ClearCoat_Normal);
 	const FString ClearCoat_TangentCode = Tangent != INDEX_NONE ? *GetParameterCode(ClearCoat_Tangent) : TEXT("NONE");
 
 	return AddCodeChunk(
-		MCT_Strata, TEXT("StrataConvertLegacyMaterial%s(Parameters.StrataPixelFootprint, %s, %s, %s, %s, %s, %s, %s, %s, %s, %s, %s, %s, %s, %s, %s, %s, %s, %s, %s, Parameters.SharedLocalBases.Types, Parameters.StrataTree) /* Normal = %s ; Tangent = %s ; ClearCoat_Normal = %s ; ClearCoat_Tangent = %s */"),
+		MCT_Strata, TEXT("StrataConvertLegacyMaterial%s(Parameters.StrataPixelFootprint, %s, %s, %s, %s, %s, %s, %s, %s, %s, %s, %s, %s, %s, %s, %s, %s, %s, %s, %s, %s, %s, %s, Parameters.SharedLocalBases.Types, Parameters.StrataTree) /* Normal = %s ; Tangent = %s ; ClearCoat_Normal = %s ; ClearCoat_Tangent = %s */"),
 		bHasDynamicShadingModels ? TEXT("Dynamic") : TEXT("Static"),
 		*StrataGetCastParameterCode(BaseColor,						MCT_Float3),
 		*StrataGetCastParameterCode(Specular,						MCT_Float),
@@ -10406,6 +10421,10 @@ int32 FHLSLMaterialTranslator::StrataConversionFromLegacy(
 		*StrataGetCastParameterCode(WaterPhaseG,					MCT_Float),
 		*StrataGetCastParameterCode(ColorScaleBehindWater,			MCT_Float3),
 		*GetParameterCode(ShadingModel),
+		// Raw access to Normal/Tanget/ClearCoatNormal for conversion purpose
+		*StrataGetCastParameterCode(Normal,							MCT_Float3),
+		*StrataGetCastParameterCode(RawTangent,						MCT_Float3),
+		*StrataGetCastParameterCode(ClearCoat_Normal,				MCT_Float3),
 		*SharedLocalBasisIndexMacro,
 		*ClearCoat_SharedLocalBasisIndexMacro,
 		// Regular normal basis
@@ -10452,6 +10471,22 @@ int32 FHLSLMaterialTranslator::StrataHairBSDF(int32 BaseColor, int32 Scatter, in
 	);
 }
 
+int32 FHLSLMaterialTranslator::StrataEyeBSDF(int32 DiffuseAlbedo, int32 Roughness, int32 IrisMask, int32 IrisDistance, int32 EmissiveColor, int32 CorneaNormal, int32 IrisNormal, const FString& SharedLocalBasisIndexMacro)
+{
+	return AddCodeChunk(
+		MCT_Strata, TEXT("GetStrataEyeBSDF(%s, %s, %s, %s, %s, %s, %s) /* Cornea:%s Iris:%s */"),
+		*StrataGetCastParameterCode(DiffuseAlbedo, MCT_Float3),
+		*StrataGetCastParameterCode(Roughness, MCT_Float),
+		*StrataGetCastParameterCode(IrisMask, MCT_Float),
+		*StrataGetCastParameterCode(IrisDistance, MCT_Float),
+		*StrataGetCastParameterCode(IrisNormal, MCT_Float),
+		*StrataGetCastParameterCode(EmissiveColor, MCT_Float3),
+		*SharedLocalBasisIndexMacro,
+		*GetParameterCode(CorneaNormal),
+		*GetParameterCode(IrisNormal)
+	);
+}
+
 int32 FHLSLMaterialTranslator::StrataSingleLayerWaterBSDF(
 	int32 BaseColor, int32 Metallic, int32 Specular, int32 Roughness, int32 EmissiveColor, int32 TopMaterialOpacity,
 	int32 WaterAlbedo, int32 WaterExtinction, int32 WaterPhaseG, int32 ColorScaleBehindWater, int32 Normal, const FString& SharedLocalBasisIndexMacro,
diff --git a/Engine/Source/Runtime/Engine/Private/Materials/HLSLMaterialTranslator.h b/Engine/Source/Runtime/Engine/Private/Materials/HLSLMaterialTranslator.h
--- a/Engine/Source/Runtime/Engine/Private/Materials/HLSLMaterialTranslator.h
+++ b/Engine/Source/Runtime/Engine/Private/Materials/HLSLMaterialTranslator.h
@@ -1046,6 +1046,7 @@ protected:
 	virtual int32 StrataVolumetricFogCloudBSDF(int32 Albedo, int32 Extinction, int32 EmissiveColor, int32 AmbientOcclusion) override;
 	virtual int32 StrataUnlitBSDF(int32 EmissiveColor, int32 TransmittanceColor) override;
 	virtual int32 StrataHairBSDF(int32 BaseColor, int32 Scatter, int32 Specular, int32 Roughness, int32 Backlit, int32 EmissiveColor, int32 Tangent, const FString& SharedLocalBasisIndexMacro) override;
+	virtual int32 StrataEyeBSDF(int32 DiffuseColor, int32 Roughness, int32 IrisMask, int32 IrisDistance, int32 EmissiveColor, int32 CorneaNormal, int32 IrisNormal, const FString& SharedLocalBasisIndexMacro) override;
 	virtual int32 StrataSingleLayerWaterBSDF(
 		int32 BaseColor, int32 Metallic, int32 Specular, int32 Roughness, 
 		int32 EmissiveColor, int32 TopMaterialOpacity, int32 WaterAlbedo, int32 WaterExtinction, int32 WaterPhaseG, 
diff --git a/Engine/Source/Runtime/Engine/Private/Materials/Material.cpp b/Engine/Source/Runtime/Engine/Private/Materials/Material.cpp
--- a/Engine/Source/Runtime/Engine/Private/Materials/Material.cpp
+++ b/Engine/Source/Runtime/Engine/Private/Materials/Material.cpp
@@ -2926,7 +2926,7 @@ static void AddStrataShadingModelFromMaterialShadingModel(FStrataMaterialInfo& O
 	if (InShadingModels.HasShadingModel(MSM_TwoSidedFoliage))	{ OutInfo.AddShadingModel(EStrataShadingModel::SSM_SubsurfaceLit); }
 	if (InShadingModels.HasShadingModel(MSM_Hair))				{ OutInfo.AddShadingModel(EStrataShadingModel::SSM_Hair); }
 	if (InShadingModels.HasShadingModel(MSM_Cloth))				{ OutInfo.AddShadingModel(EStrataShadingModel::SSM_DefaultLit); }
-	if (InShadingModels.HasShadingModel(MSM_Eye))				{ OutInfo.AddShadingModel(EStrataShadingModel::SSM_SubsurfaceLit); }
+	if (InShadingModels.HasShadingModel(MSM_Eye))				{ OutInfo.AddShadingModel(EStrataShadingModel::SSM_Eye); }
 	if (InShadingModels.HasShadingModel(MSM_SingleLayerWater))	{ OutInfo.AddShadingModel(EStrataShadingModel::SSM_SingleLayerWater); }
 	if (InShadingModels.HasShadingModel(MSM_ThinTranslucent))	{ OutInfo.AddShadingModel(EStrataShadingModel::SSM_DefaultLit); }
 }
@@ -3295,7 +3295,7 @@ void UMaterial::ConvertMaterialToStrataMaterial()
 		}
 		if (ClearCoatOutput)
 		{
-			MoveConnectionTo(*ClearCoatOutput->GetInput(0), ConvertNode, 17);		 // ClearCoatNormal
+			CopyConnectionTo(*ClearCoatOutput->GetInput(0), ConvertNode, 17);		 // ClearCoatNormal
 		}
 	}
 
@@ -4515,6 +4515,11 @@ void UMaterial::RebuildShadingModelField()
 				ShadingModel = MSM_Hair;
 				BlendMode = EBlendMode::BLEND_Opaque;
 			}
+			else if (StrataMaterialInfo.HasOnlyShadingModel(SSM_Eye))
+			{
+				MaterialDomain = EMaterialDomain::MD_Surface;
+				ShadingModel = MSM_Eye;
+			}
 			else if (StrataMaterialInfo.HasOnlyShadingModel(SSM_SingleLayerWater))
 			{
 				MaterialDomain = EMaterialDomain::MD_Surface;
diff --git a/Engine/Source/Runtime/Engine/Private/Materials/MaterialExpressions.cpp b/Engine/Source/Runtime/Engine/Private/Materials/MaterialExpressions.cpp
--- a/Engine/Source/Runtime/Engine/Private/Materials/MaterialExpressions.cpp
+++ b/Engine/Source/Runtime/Engine/Private/Materials/MaterialExpressions.cpp
@@ -21729,6 +21729,10 @@ int32 UMaterialExpressionStrataLegacyConversion::Compile(class FMaterialCompiler
 		{
 			Compiler->StrataCompilationGetOperator(this).BSDFRegisteredSharedLocalBasis = NewRegisteredSharedLocalBasis;
 		}
+		else if (ConvertedStrataMaterialInfo.HasShadingModel(SSM_Eye))
+		{
+			Compiler->StrataCompilationGetOperator(this).BSDFRegisteredSharedLocalBasis = NewRegisteredSharedLocalBasis;
+		}
 		else if (ConvertedStrataMaterialInfo.HasShadingModel(SSM_SingleLayerWater))
 		{
 			Compiler->StrataCompilationGetOperator(this).BSDFRegisteredSharedLocalBasis = NewRegisteredSharedLocalBasis;
@@ -21869,6 +21873,7 @@ void UMaterialExpressionStrataLegacyConversion::GatherStrataMaterialInfo(FStrata
 	if (ConvertedStrataMaterialInfo.HasShadingModel(SSM_SubsurfaceLit))			{ StrataMaterialInfo.AddShadingModel(SSM_SubsurfaceLit); }
 	if (ConvertedStrataMaterialInfo.HasShadingModel(SSM_VolumetricFogCloud))	{ StrataMaterialInfo.AddShadingModel(SSM_VolumetricFogCloud); }
 	if (ConvertedStrataMaterialInfo.HasShadingModel(SSM_Hair))					{ StrataMaterialInfo.AddShadingModel(SSM_Hair); }
+	if (ConvertedStrataMaterialInfo.HasShadingModel(SSM_Eye))					{ StrataMaterialInfo.AddShadingModel(SSM_Eye); }
 	if (ConvertedStrataMaterialInfo.HasShadingModel(SSM_SingleLayerWater))		{ StrataMaterialInfo.AddShadingModel(SSM_SingleLayerWater); }
 	if (SubsurfaceProfile)														{ StrataMaterialInfo.AddSubsurfaceProfile(SubsurfaceProfile); }
 	if (ConvertedStrataMaterialInfo.HasShadingModelFromExpression())			{ StrataMaterialInfo.SetShadingModelFromExpression(true); }
@@ -21931,6 +21936,12 @@ FStrataOperator* UMaterialExpressionStrataLegacyConversion::StrataGenerateMateri
 			Operator.BSDFType = STRATA_BSDF_TYPE_HAIR;
 			return &Operator;
 		}
+		else if (ConvertedStrataMaterialInfo.HasShadingModel(SSM_Eye))
+		{
+			FStrataOperator& Operator = Compiler->StrataCompilationRegisterOperator(STRATA_OPERATOR_BSDF_LEGACY, this, Parent);
+			Operator.BSDFType = STRATA_BSDF_TYPE_EYE;
+			return &Operator;
+		}
 		else if (ConvertedStrataMaterialInfo.HasShadingModel(SSM_SingleLayerWater))
 		{
 			FStrataOperator& Operator = Compiler->StrataCompilationRegisterOperator(STRATA_OPERATOR_BSDF_LEGACY, this, Parent);
@@ -23041,6 +23052,92 @@ FStrataOperator* UMaterialExpressionStrataHairBSDF::StrataGenerateMaterialTopolo
 }
 #endif // WITH_EDITOR
 
+UMaterialExpressionStrataEyeBSDF::UMaterialExpressionStrataEyeBSDF(const FObjectInitializer& ObjectInitializer)
+	: Super(ObjectInitializer)
+{
+	struct FConstructorStatics
+	{
+		FText NAME_Strata;
+		FConstructorStatics() : NAME_Strata(LOCTEXT("Strata BSDFs", "Strata BSDFs")) { }
+	};
+	static FConstructorStatics ConstructorStatics;
+#if WITH_EDITORONLY_DATA
+	MenuCategories.Add(ConstructorStatics.NAME_Strata);
+#endif
+}
+
+#if WITH_EDITOR
+int32 UMaterialExpressionStrataEyeBSDF::Compile(class FMaterialCompiler* Compiler, int32 OutputIndex)
+{
+	int32 CorneaNormalCodeChunk = CompileWithDefaultTangentWS(Compiler, CorneaNormal);
+	const FStrataRegisteredSharedLocalBasis NewRegisteredSharedLocalBasis = StrataCompilationInfoCreateSharedLocalBasis(Compiler, CorneaNormalCodeChunk);
+
+	int32 OutputCodeChunk = Compiler->StrataEyeBSDF(
+		CompileWithDefaultFloat3(Compiler, DiffuseColor, 0.0f, 0.0f, 0.0f),
+		CompileWithDefaultFloat1(Compiler, Roughness,	 0.5f),
+		CompileWithDefaultFloat1(Compiler, IrisMask,	 0.0f),
+		CompileWithDefaultFloat1(Compiler, IrisDistance, 0.0f),
+		CompileWithDefaultFloat3(Compiler, IrisNormal,   0.0f, 0.0f, 1.0f),
+		CompileWithDefaultFloat3(Compiler, EmissiveColor,0.0f, 0.0f, 0.0f),
+		CorneaNormalCodeChunk,
+		Compiler->GetStrataSharedLocalBasisIndexMacro(NewRegisteredSharedLocalBasis));
+
+	return OutputCodeChunk;
+}
+
+void UMaterialExpressionStrataEyeBSDF::GetCaption(TArray<FString>& OutCaptions) const
+{
+	OutCaptions.Add(TEXT("Strata Eye BSDF"));
+}
+
+uint32 UMaterialExpressionStrataEyeBSDF::GetOutputType(int32 OutputIndex)
+{
+	return MCT_Strata;
+}
+
+uint32 UMaterialExpressionStrataEyeBSDF::GetInputType(int32 InputIndex)
+{
+	switch (InputIndex)
+	{
+	case 0: return MCT_Float3; // DiffuseColor
+	case 1: return MCT_Float1; // Roughness
+	case 2: return MCT_Float3; // Cornea normal
+	case 3: return MCT_Float3; // IrisNormal
+	case 4: return MCT_Float1; // IrisMask
+	case 5: return MCT_Float1; // IrisDistance
+	case 6: return MCT_Float3; // EmissiveColor
+	}
+
+	check(false);
+	return MCT_Float1;
+}
+
+bool UMaterialExpressionStrataEyeBSDF::IsResultStrataMaterial(int32 OutputIndex)
+{
+	return true;
+}
+
+void UMaterialExpressionStrataEyeBSDF::GatherStrataMaterialInfo(FStrataMaterialInfo& StrataMaterialInfo, int32 OutputIndex)
+{
+	// Track connected inputs
+	if (DiffuseColor.IsConnected())		{ StrataMaterialInfo.AddPropertyConnected(MP_BaseColor); }
+	if (Roughness.IsConnected())		{ StrataMaterialInfo.AddPropertyConnected(MP_Roughness); }
+	if (CorneaNormal.IsConnected())		{ StrataMaterialInfo.AddPropertyConnected(MP_Normal); }
+	if (IrisNormal.IsConnected())		{ StrataMaterialInfo.AddPropertyConnected(MP_Tangent); }
+	if (IrisMask.IsConnected())			{ StrataMaterialInfo.AddPropertyConnected(MP_CustomData0); }
+	if (IrisDistance.IsConnected())		{ StrataMaterialInfo.AddPropertyConnected(MP_CustomData1); }
+	if (EmissiveColor.IsConnected())	{ StrataMaterialInfo.AddPropertyConnected(MP_EmissiveColor); }
+
+	StrataMaterialInfo.AddShadingModel(SSM_Eye);
+}
+
+FStrataOperator* UMaterialExpressionStrataEyeBSDF::StrataGenerateMaterialTopologyTree(class FMaterialCompiler* Compiler, class UMaterialExpression* Parent, int32 OutputIndex)
+{
+	FStrataOperator& StrataOperator = Compiler->StrataCompilationRegisterOperator(STRATA_OPERATOR_BSDF, this, Parent);
+	StrataOperator.BSDFType = STRATA_BSDF_TYPE_EYE;
+	return &StrataOperator;
+}
+#endif // WITH_EDITOR
 
 
 UMaterialExpressionStrataSingleLayerWaterBSDF::UMaterialExpressionStrataSingleLayerWaterBSDF(const FObjectInitializer& ObjectInitializer)
diff --git a/Engine/Source/Runtime/Engine/Private/Materials/StrataMaterial.cpp b/Engine/Source/Runtime/Engine/Private/Materials/StrataMaterial.cpp
--- a/Engine/Source/Runtime/Engine/Private/Materials/StrataMaterial.cpp
+++ b/Engine/Source/Runtime/Engine/Private/Materials/StrataMaterial.cpp
@@ -21,6 +21,9 @@ FString GetStrataBSDFName(uint8 BSDFType)
 	case STRATA_BSDF_TYPE_HAIR:
 		return TEXT("HAIR");
 		break;
+	case STRATA_BSDF_TYPE_EYE:
+		return TEXT("EYE");
+		break;
 	case STRATA_BSDF_TYPE_SINGLELAYERWATER:
 		return TEXT("SINGLELAYERWATER");
 		break;
diff --git a/Engine/Source/Runtime/Engine/Public/MaterialCompiler.h b/Engine/Source/Runtime/Engine/Public/MaterialCompiler.h
--- a/Engine/Source/Runtime/Engine/Public/MaterialCompiler.h
+++ b/Engine/Source/Runtime/Engine/Public/MaterialCompiler.h
@@ -589,6 +589,7 @@ public:
 	virtual int32 StrataVolumetricFogCloudBSDF(int32 Albedo, int32 Extinction, int32 EmissiveColor, int32 AmbientOcclusion) = 0;
 	virtual int32 StrataUnlitBSDF(int32 EmissiveColor, int32 TransmittanceColor) = 0;
 	virtual int32 StrataHairBSDF(int32 BaseColor, int32 Scatter, int32 Specular, int32 Roughness, int32 Backlit, int32 EmissiveColor, int32 Tangent, const FString& SharedLocalBasisIndexMacro) = 0;
+	virtual int32 StrataEyeBSDF(int32 DiffuseAlbedo, int32 Roughness, int32 IrisMask, int32 IrisDistance, int32 EmissiveColor, int32 CorneaNormal, int32 IrisNormal, const FString& SharedLocalBasisIndexMacro) = 0;
 	virtual int32 StrataSingleLayerWaterBSDF(
 		int32 BaseColor, int32 Metallic, int32 Specular, int32 Roughness, int32 EmissiveColor, int32 TopMaterialOpacity, 
 		int32 WaterAlbedo, int32 WaterExtinction, int32 WaterPhaseG, int32 ColorScaleBehindWater, int32 Normal, const FString& SharedLocalBasisIndexMacro, FStrataOperator* PromoteToOperator) = 0;
@@ -1181,6 +1182,11 @@ public:
 		return Compiler->StrataHairBSDF(BaseColor, Scatter, Specular, Roughness, Backlit, EmissiveColor, Tangent, SharedLocalBasisIndexMacro);
 	}
 
+	virtual int32 StrataEyeBSDF(int32 DiffuseAlbedo, int32 Roughness, int32 IrisMask, int32 IrisDistance, int32 EmissiveColor, int32 CorneaNormal, int32 IrisNormal, const FString& SharedLocalBasisIndexMacro) override
+	{
+		return Compiler->StrataEyeBSDF(DiffuseAlbedo, Roughness, IrisMask, IrisDistance, EmissiveColor, CorneaNormal, IrisNormal, SharedLocalBasisIndexMacro);
+	}
+
 	virtual int32 StrataSingleLayerWaterBSDF(
 		int32 BaseColor, int32 Metallic, int32 Specular, int32 Roughness, int32 EmissiveColor, int32 TopMaterialOpacity, 
 		int32 WaterAlbedo, int32 WaterExtinction, int32 WaterPhaseG, int32 ColorScaleBehindWater, int32 Normal, const FString& SharedLocalBasisIndexMacro, FStrataOperator* PromoteToOperator) override
