Strata
- added optimised clear coat path (hard coded absorption, top F0=0.04 & F90=1, shared normal only).
- This is needed for simple clear coat to perform well and better repsect previous gbuffe rmemory usage.
- SSR handle that change of roughness  and F0 on top.
- there will be follow ups: Lumen to be fixed, a few visual mismatch to fix, etc.

diff --git a/Engine/Shaders/Private/DebugProbes.usf b/Engine/Shaders/Private/DebugProbes.usf
--- a/Engine/Shaders/Private/DebugProbes.usf
+++ b/Engine/Shaders/Private/DebugProbes.usf
@@ -306,6 +306,7 @@ void MainPS(
 			/*float3 Emissive						*/0.0f,
 			/*float SecondRoughness					*/0.0f,
 			/*float SecondRoughnessWeight			*/0.0f,
+			/*flaot SecondRoughnessAsSimpleClearCoat*/0.0f,
 			/*float FuzzAmount						*/0.0f,
 			/*float3 FuzzColor						*/0.0f,
 			/*float Thickness						*/0.0f,
diff --git a/Engine/Shaders/Private/Strata/Strata.ush b/Engine/Shaders/Private/Strata/Strata.ush
--- a/Engine/Shaders/Private/Strata/Strata.ush
+++ b/Engine/Shaders/Private/Strata/Strata.ush
@@ -141,7 +141,7 @@ struct FStrataBSDF
 // ---------------
 // 3 types of storages layout:
 // * Complex: store many (mixed) BSDFs
-// * Single : store a single BSDF (with all feature apart from anisotropy)
+// * Single : store a single BSDF (with all feature apart from anisotropy, because we sample from the toplayer data texture which stores normal but not the tangent)
 // * Simple : store a single Slab BSDF (with no extra feature)
 // 
 // Header State (8bits):
@@ -490,11 +490,11 @@ bool   IsStrataSlabFastPathCompatible(in FStrataBSDF BSDF);
 ///////////////////////////////////////////////////////////////////////////////
 // Haziness
 
-// bFastClearCoat==false means that Weight is the blend in amount of the secondary specular lob over the first specular lob.
-// bFastClearCoat==true  means fast clear coat is used, Weight is coverage of a top layer specular. 
+// bSimpleClearCoat==false means that Weight is the blend in amount of the secondary specular lob over the first specular lob.
+// bSimpleClearCoat==true  means fast clear coat is used, Weight is coverage of a top layer specular. 
 struct FHaziness
 {
-	bool  bFastClearCoat;
+	bool  bSimpleClearCoat;
 	float Roughness;
 	float Weight;
 };
@@ -508,7 +508,7 @@ FHaziness InitialiseHaziness()
 // returns 16 bits: 1 bit bFastClearCoat, 7 bits Weight, 8 bits roughness
 float PackHaziness(FHaziness Haziness)
 {
-	uint High8Bits = (Haziness.bFastClearCoat ? 0x80 : 0x00) | uint(saturate(Haziness.Weight) * 127);
+	uint High8Bits = (Haziness.bSimpleClearCoat ? 0x80 : 0x00) | uint(saturate(Haziness.Weight) * 127);
 	uint Data = (High8Bits << 8) | uint(saturate(Haziness.Roughness) * 255);
 	return asfloat(Data);
 }
@@ -518,9 +518,9 @@ FHaziness UnpackHaziness(float FloatData)
 	uint Data = asuint(FloatData);
 
 	FHaziness Haziness;
-	Haziness.bFastClearCoat = Data & 0x8000;
-	Haziness.Weight = float(uint((Data>>8) & 0x7F)) * (1.0f / 127.0f);
-	Haziness.Roughness = float(uint(Data & 0xFF)) * (1.0f / 255.0f);
+	Haziness.bSimpleClearCoat	= Data & 0x8000;
+	Haziness.Weight				= float(uint((Data>>8) & 0x7F)) * (1.0f / 127.0f);
+	Haziness.Roughness			= float(uint(Data & 0xFF)) * (1.0f / 255.0f);
 	return Haziness;
 }
 
@@ -532,7 +532,7 @@ float LerpHazinessParameterBlending(FStrataBSDF A, FStrataBSDF B, float MixValue
 	FHaziness OutputHaziness;
 	OutputHaziness.Weight			= lerp(HazinessA.Weight,	HazinessB.Weight,		MixValue);
 	OutputHaziness.Roughness		= lerp(HazinessA.Roughness, HazinessB.Roughness,	MixValue);
-	OutputHaziness.bFastClearCoat	= HazinessA.bFastClearCoat || HazinessB.bFastClearCoat;		// Maintain the fast legacy path first
+	OutputHaziness.bSimpleClearCoat = HazinessA.bSimpleClearCoat || HazinessB.bSimpleClearCoat;		// Maintain the fast legacy path first
 
 	return PackHaziness(OutputHaziness);
 }
@@ -1235,7 +1235,7 @@ FStrataData GetStrataSlabBSDF(
 	float Roughness, float Anisotropy,
 	float SSSProfileID, float3 SSSMFP, float SSSMFPScale, float SSSPhaseAnisotropy, float UseSSSDiffusion,
 	float3 Emissive, 
-	float SecondRoughness, float SecondRoughnessWeight,
+	float SecondRoughness, float SecondRoughnessWeight, float SecondRoughnessAsSimpleClearCoat,
 	float FuzzAmount, float3 FuzzColor,
 	float Thickness,
 	uint SharedLocalBasisIndex, inout uint SharedLocalBasisTypes)
@@ -1267,7 +1267,7 @@ FStrataData GetStrataSlabBSDF(
 	{
 		Haziness.Roughness = SecondRoughness;
 		Haziness.Weight = SecondRoughnessWeight;
-		Haziness.bFastClearCoat = false;
+		Haziness.bSimpleClearCoat = SecondRoughnessAsSimpleClearCoat > 0.0f;
 	}
 
 	// If MFP is smaller than the pixel's footprint, remove SSS attributes
@@ -1289,7 +1289,7 @@ FStrataData GetStrataSlabBSDF(
 		Thickness = STRATA_SIMPLEVOLUME_THICKNESS_CM;
 	}
 
-	const uint bHasHaziness = ((Haziness.Roughness != Roughness  && Haziness.Weight > 0.0f) || Haziness.bFastClearCoat) ? 1 : 0;
+	const uint bHasHaziness = ((Haziness.Roughness != Roughness  && Haziness.Weight > 0.0f) || Haziness.bSimpleClearCoat) ? 1 : 0;
 	const uint bHasAnisotropy = Anisotropy != 0 ? 1 : 0;
 	const uint bHasFuzz = FuzzAmount > 0.0f ? 1 : 0;
 	if (bHasAnisotropy)
@@ -3142,6 +3142,7 @@ void StrataConvertFromDBuffer(
 		0.f,				// Emissive
 		0.f,				// SecondRoughness
 		0.f,				// SecondRoughnessWeight
+		0.f,				// SecondRoughnessAsSimpleClearCoat
 		0.f,				// FuzzAmount
 		0.f,				// FuzzColor
 		DefaultThickness,
diff --git a/Engine/Shaders/Private/Strata/StrataEnvironmentLighting.ush b/Engine/Shaders/Private/Strata/StrataEnvironmentLighting.ush
--- a/Engine/Shaders/Private/Strata/StrataEnvironmentLighting.ush
+++ b/Engine/Shaders/Private/Strata/StrataEnvironmentLighting.ush
@@ -100,7 +100,7 @@ FStrataDeferredLighting StrataEnvironmentLighting(
 		if (bIsTopLayer)
 		{
 			// The specular path weight applied on SSR. It must account for throughput even for top surface because it also contains coverage.
-			TopLayerEnvBRDF += BSDFThroughput * StrataEnvLight.SpecularWeight;
+			TopLayerEnvBRDF += BSDFThroughput * StrataEnvLight.SSRSpecularWeight;
 		}
 
 		// STRATA_TODO ReflectionEnvironment(...) for static reflection volumes.
diff --git a/Engine/Shaders/Private/Strata/StrataEvaluation.ush b/Engine/Shaders/Private/Strata/StrataEvaluation.ush
--- a/Engine/Shaders/Private/Strata/StrataEvaluation.ush
+++ b/Engine/Shaders/Private/Strata/StrataEvaluation.ush
@@ -305,6 +305,10 @@
 			float3 HazeMSScale = 1;
 			float3 DirectionalAlbedo_SpecularTransmission = 1.0f;
 			float HazeWeight = 0.f;
+			bool  bHazeAsSimpleClearCoat = false;
+			float3 HazeClearCoatTransmittance = 1.0f;
+			float3 HazeF0 = F0;
+			float3 HazeF90= F90;
 			{
 				{
 					FBxDFEnergyTerms EnergyTerms = ComputeGGXSpecEnergyTerms(SafeRoughness, BSDFContext.Context.NoV, F0, F90);
@@ -318,8 +322,17 @@
 					FHaziness Haziness = UnpackHaziness(SLAB_HAZINESS(BSDFContext.BSDF));
 					HazeWeight = Haziness.Weight;
 					HazeSafeRoughness = MakeRoughnessSafe(Haziness.Roughness);
+					bHazeAsSimpleClearCoat = Haziness.bSimpleClearCoat;
 
-					FBxDFEnergyTerms EnergyTerms = ComputeGGXSpecEnergyTerms(HazeSafeRoughness, BSDFContext.Context.NoV, F0, F90);
+					BRANCH
+					if (bHazeAsSimpleClearCoat)
+					{
+						HazeF0 = 0.04f;
+						HazeF90 = 1.0f;
+						HazeClearCoatTransmittance = SimpleClearCoatTransmittance(max(0.0001, AreaLightContext.NoL), max(0.0001, NoV), StrataGetBSDFMetallic(BSDFContext.BSDF), StrataGetBSDFBaseColor(BSDFContext.BSDF));
+					}
+
+					FBxDFEnergyTerms EnergyTerms = ComputeGGXSpecEnergyTerms(HazeSafeRoughness, BSDFContext.Context.NoV, HazeF0, HazeF90);
 					HazeMSScale = ComputeEnergyConservation(EnergyTerms);
 				}
 			#endif
@@ -328,7 +341,6 @@
 				Sample.IntegratedDiffuseValue *= DirectionalAlbedo_SpecularTransmission;
 			}
 
-
 			////
 			//// Evaluate the specular component.
 			//// This takes into account multiple scattering, micro occlusion.
@@ -387,7 +399,7 @@
 					#if STRATA_FASTPATH==0
 					if (bHaziness)
 					{
-						RectLightSpecHaze = RectGGXApproxLTC(HazeSafeRoughness, F0, BSDFContext.N, BSDFContext.V, AreaLightContext.AreaLight.Rect, AreaLightContext.AreaLight.Texture);
+						RectLightSpecHaze = RectGGXApproxLTC(HazeSafeRoughness, HazeF0, BSDFContext.N, BSDFContext.V, AreaLightContext.AreaLight.Rect, AreaLightContext.AreaLight.Texture);
 					
 						const float H_PDF = VisibleGGXPDF(BSDFContext.TangentV, BSDFContext.TangentH, Pow4(HazeSafeRoughness));
 						HazePDF = RayPDFToReflectionRayPDF(BSDFContext.Context.VoH, H_PDF);
@@ -441,29 +453,57 @@
 			const float SpecularMicroOcclusion = F0RGBToMicroOcclusion(F0);
 			const float3 FresnelTerm = SpecularMicroOcclusion * F_Schlick(F0, F90, BSDFContext.Context.VoH);
 
-			Sample.SpecularPathProbability		= bHaziness	? (1.0f - HazeWeight) : 1.0f;
-			Sample.SpecularHazePathProbability	= bHaziness	? HazeWeight : 0.0f;
+			const float HazeSpecularMicroOcclusion = F0RGBToMicroOcclusion(HazeF0);
+			const float3 HazeFresnelTerm = SpecularMicroOcclusion * F_Schlick(HazeF0, HazeF90, BSDFContext.Context.VoH);
 
-			Sample.SpecularPDF				= PDF * Sample.SpecularPathProbability;
-			Sample.SpecularHazePDF			= HazePDF * Sample.SpecularHazePathProbability;
+			const float3 TopLayerSpecularTransmittionApprox = saturate(1.0f - HazeFresnelTerm); // Simple and do not use the LUT for now
 
+			float3 CommonTerm = 0.0f;
 			if (bIsRectLight)
 			{
 				Sample.SpecularPathValue     = RectLightSpec     * MSScale     * SpecularMicroOcclusion;
-				Sample.SpecularHazePathValue = RectLightSpecHaze * HazeMSScale * SpecularMicroOcclusion;
+				Sample.SpecularHazePathValue = RectLightSpecHaze * HazeMSScale * HazeSpecularMicroOcclusion;
 
-				Sample.IntegratedSpecularValue = ShadowTerms.SurfaceShadow * /* AreaLightContext.NoL and falloff is part of the LTC integration already */
-					(Sample.SpecularPathValue * Sample.SpecularPathProbability + Sample.SpecularHazePathValue * Sample.SpecularHazePathProbability);
+				CommonTerm = ShadowTerms.SurfaceShadow; /* AreaLightContext.NoL and falloff is part of the LTC integration already */
 			}
 			else
 			{
-				Sample.SpecularPathValue     = D     * Vis     * FresnelTerm * MSScale;
-				Sample.SpecularHazePathValue = HazeD * HazeVis * FresnelTerm * HazeMSScale;
+				Sample.SpecularPathValue     = D     * Vis     * FresnelTerm     * MSScale     * HazeClearCoatTransmittance;
+				Sample.SpecularHazePathValue = HazeD * HazeVis * HazeFresnelTerm * HazeMSScale;
 
-				Sample.IntegratedSpecularValue += ShadowTerms.SurfaceShadow * AreaLightContext.NoL * AreaLightContext.AreaLight.FalloffColor * AreaLightContext.Falloff *
-					(Sample.SpecularPathValue * Sample.SpecularPathProbability + Sample.SpecularHazePathValue * Sample.SpecularHazePathProbability);
+				CommonTerm = ShadowTerms.SurfaceShadow * AreaLightContext.NoL * AreaLightContext.AreaLight.FalloffColor * AreaLightContext.Falloff;
 			}
 
+			// Reminder: 
+			//	- bHazeAsSimpleClearCoat == false means we have two specular lobes that are lerped.
+			//	- bHazeAsSimpleClearCoat == true  means we have a bottom specular lobe (roughness and F0 user specified) and a top specular lob (hard coded absortion, F0=0.04 and user specified roughness.
+			if (bHazeAsSimpleClearCoat)
+			{
+				const float TopLayerCoverage = HazeWeight;
+				const float3 BottomLayerFactors = lerp(1.0, HazeClearCoatTransmittance * TopLayerSpecularTransmittionApprox, TopLayerCoverage);
+
+				const float BottomLayerFactorsGrey = dot(BottomLayerFactors, (1.0 / 3.0).xxx);
+				Sample.SpecularPathProbability     = BottomLayerFactorsGrey / (TopLayerCoverage + BottomLayerFactorsGrey);
+				Sample.SpecularHazePathProbability = TopLayerCoverage       / (TopLayerCoverage + BottomLayerFactorsGrey);
+
+				// Clear coat applied over bottom layer diffuse component
+				Sample.DiffusePathValue			*= BottomLayerFactors;
+				Sample.IntegratedDiffuseValue	*= BottomLayerFactors;
+				// And to the bottom layer specular on top of whiche we add the top layer specular
+				Sample.IntegratedSpecularValue	 = CommonTerm * (Sample.SpecularPathValue * BottomLayerFactors + Sample.SpecularHazePathValue * TopLayerCoverage);
+			}
+			else
+			{
+				Sample.SpecularPathProbability     = bHaziness ? (1.0f - HazeWeight) : 1.0f;
+				Sample.SpecularHazePathProbability = bHaziness ?         HazeWeight  : 0.0f;
+
+				Sample.IntegratedSpecularValue = CommonTerm * (Sample.SpecularPathValue * Sample.SpecularPathProbability + Sample.SpecularHazePathValue * Sample.SpecularHazePathProbability);
+			}
+
+			Sample.SpecularPDF = PDF * Sample.SpecularPathProbability;
+			Sample.SpecularHazePDF = HazePDF * Sample.SpecularHazePathProbability;
+
+
 			////
 			//// Evaluate emissive and set the sample throughput (transmittance to next layer) corresponding to an opaque slab of matter.
 			////
@@ -1142,6 +1182,9 @@
 	float  SpecularHazeSafeRoughness;
 	float  SSRReduction;
 
+	// This can vary depending if the simple clear coat is used.
+	float3 SSRSpecularWeight;
+
 	bool   bSubsurface; // True if we need to separate the subsurface light contribution for the screen space diffusion process.
 };
 
@@ -1180,6 +1223,45 @@
 				StrataEnvLightResult.DiffuseColor			= DiffuseColor;
 			}
 
+			#if STRATA_FASTPATH==0
 			BRANCH
+			if (BSDF_GETISSIMPLEVOLUME(BSDFContext.BSDF))
+			{
+				FParticipatingMedia PM = StrataSlabCreateParticipatingMedia(DiffuseColor, SLAB_SSSMFP(BSDFContext.BSDF));
+				const float DiffuseToVolumeBlend = StrataSlabDiffuseToVolumeBlend(PM);
+
+				// The environment response of has been measured for an hemisphere having a uniform radiance of 1. 
+				// This DirectionalAlbedo also contains the integral over the hemisphere according to the phase function.
+				// The spherical harmonic used for environment lighting contains the diffuse integral over the hemisphere as DiffLambert = int_{\omega} {1 * NoL d\Omega} = PI.
+				// We do not want to be affected by the NoL term here since The DirectionalAlbedo also contains the integral over the hemisphere according to the phase function.
+				// Integral of a value=1 over the Hemisphere is int_{\omega} {1  d\Omega} = 2PI.
+				// So, as an approximation, we recover the uniform hemisphere luminance as UniformEnvL = DiffLambert * (1 / PI) * (2 * PI).
+				// We consider that the diffuse integration of the environment SH over the hemisphere multiplied with the light response will be a good approximation similar to the "split sum integral".
+				float3 SlabDirectionalAlbedo = IsotropicMediumSlabEnvDirectionalAlbedo(PM) * ((1 / PI)* (2 * PI));
+
+			#if 1
+				// Because the default EnvBRDF does not account for multiple scattering, we rescale the participating media response according to EnvBRDF.
+				// This is to ensure a monotone light response behavior when transitioning from the Diffuse to the Volumetric model (without that, multiple scattering can result in higher luminance while trnasitionning from diffuse to volume light model)
+				// STRATA_TODO: We should NOT have to do that. This should be investigated futher...
+				const float MaxDiffuseWeight = max3(StrataEnvLightResult.DiffuseWeight.x, StrataEnvLightResult.DiffuseWeight.y, StrataEnvLightResult.DiffuseWeight.z);
+				const float MaxSlabDirectionalAlbedo = max3(SlabDirectionalAlbedo.x, SlabDirectionalAlbedo.y, SlabDirectionalAlbedo.z);
+				if (MaxDiffuseWeight < MaxSlabDirectionalAlbedo)
+				{
+					// We only want to do that when the slab light response is higher than the diffse response, so that when the mfp is large (and scattering low), the light response still work correctly according to thickness and MFP.
+					SlabDirectionalAlbedo *= MaxDiffuseWeight / MaxSlabDirectionalAlbedo;
+				}
+			#endif
+
+				// The diffuse is attenuated by the GGX interface SpecularTransmission.
+				SlabDirectionalAlbedo *= DirectionalAlbedo_SpecularTransmission;
+
+				StrataEnvLightResult.DiffuseWeight = lerp(StrataEnvLightResult.DiffuseWeight, SlabDirectionalAlbedo, DiffuseToVolumeBlend);
+				// We keep the same normal
+				StrataEnvLightResult.bSubsurface = false;
+			}
+			#endif
+			
+
+			BRANCH
 			if (bEnableSpecular)
 			{
@@ -1205,6 +1287,8 @@
 				StrataEnvLightResult.SpecularSafeRoughness = SafeRoughness;
 				StrataEnvLightResult.SpecularColor = F0;
 
+				StrataEnvLightResult.SSRSpecularWeight = StrataEnvLightResult.SpecularWeight;
+
 				#if STRATA_FASTPATH==0
 				const bool bHasHaziness = BSDF_GETHASHAZINESS(BSDFContext.BSDF);
 				if (bHasHaziness)
@@ -1212,61 +1296,55 @@
 					FHaziness Haziness = UnpackHaziness(SLAB_HAZINESS(BSDFContext.BSDF));
 					const float HazeWeight = Haziness.Weight;
 					const float HazeSafeRoughness = MakeRoughnessSafe(Haziness.Roughness);
+					const bool  bHazeAsSimpleClearCoat = Haziness.bSimpleClearCoat;
+					
+					BRANCH
+					if (bHazeAsSimpleClearCoat)
+					{
+						const float3 HazeClearCoatTransmittance = SimpleClearCoatTransmittance(max(0.0001, BSDFContext.Context.NoL), max(0.0001, BSDFContext.Context.NoV), StrataGetBSDFMetallic(BSDFContext.BSDF), StrataGetBSDFBaseColor(BSDFContext.BSDF));
 
-					// Smoothly fade in haziness while fading out SSR to make sure we do not add energy and avoid popping.
-					const float HazinessFadeIn = saturate(HazeWeight / 0.1f);
-					StrataEnvLightResult.SSRReduction = lerp(0.0f, HazeWeight, HazinessFadeIn);
+						const float HazeF0 = 0.04f;
+						const float HazeF90= 1.0f;
+						const float SpecularMicroOcclusion = F0RGBToMicroOcclusion(HazeF0);
+						FBxDFEnergyTerms EnergyTerms = ComputeGGXSpecEnergyTerms(HazeSafeRoughness, BSDFContext.Context.NoV, HazeF0, HazeF90);
+						const float3 EvalEnvBRDFHaze = EnergyTerms.E;
+						const float3 TopLayerSpecularTransmittionApprox = saturate(1.0f - EvalEnvBRDFHaze); // Simple and do not use extra LUT for now
 
-					// Apply blend factor on the sharp specular contribution
-					StrataEnvLightResult.SpecularWeight *= lerp(1.0f, (1.0f - HazeWeight), HazinessFadeIn);
+						const float TopLayerCoverage = HazeWeight;
+						const float3 BottomLayerFactors = lerp(1.0, HazeClearCoatTransmittance * TopLayerSpecularTransmittionApprox, TopLayerCoverage);
 
-					// Compute the directional albedo for hazy specular
-					FBxDFEnergyTerms EnergyTerms = ComputeGGXSpecEnergyTerms(HazeSafeRoughness, BSDFContext.Context.NoV, F0, F90);
-					const float3 EvalEnvBRDFHaze = EnergyTerms.E;
-					StrataEnvLightResult.SpecularHazeWeight = lerp(0.0f, HazeWeight, HazinessFadeIn) * EvalEnvBRDFHaze * SpecularMicroOcclusion;
-					StrataEnvLightResult.SpecularHazeSafeRoughness = HazeSafeRoughness;
+						StrataEnvLightResult.SpecularWeight *= BottomLayerFactors;
+						StrataEnvLightResult.DiffuseWeight  *= BottomLayerFactors;
+
+						StrataEnvLightResult.SpecularHazeWeight = TopLayerCoverage * EvalEnvBRDFHaze * SpecularMicroOcclusion;
+						StrataEnvLightResult.SpecularHazeSafeRoughness = HazeSafeRoughness;
+
+						StrataEnvLightResult.SSRSpecularWeight = lerp(StrataEnvLightResult.SpecularWeight, StrataEnvLightResult.SpecularHazeWeight, TopLayerCoverage);
+					}
+					else
+					{
+						// Smoothly fade in haziness while fading out SSR to make sure we do not add energy and avoid popping.
+						const float HazinessFadeIn = saturate(HazeWeight / 0.1f);
+						StrataEnvLightResult.SSRReduction = lerp(0.0f, HazeWeight, HazinessFadeIn);
+
+						// Apply blend factor on the sharp specular contribution
+						const float BlendFactor = lerp(1.0f, (1.0f - HazeWeight), HazinessFadeIn);
+						StrataEnvLightResult.SpecularWeight *= BlendFactor;
+						StrataEnvLightResult.SSRSpecularWeight *= BlendFactor;
+
+						// Compute the directional albedo for hazy specular
+						FBxDFEnergyTerms EnergyTerms = ComputeGGXSpecEnergyTerms(HazeSafeRoughness, BSDFContext.Context.NoV, F0, F90);
+						const float3 EvalEnvBRDFHaze = EnergyTerms.E;
+						StrataEnvLightResult.SpecularHazeWeight = lerp(0.0f, HazeWeight, HazinessFadeIn) * EvalEnvBRDFHaze * SpecularMicroOcclusion;
+						StrataEnvLightResult.SpecularHazeSafeRoughness = HazeSafeRoughness;
+					}
+
+
+
+
 				}
 				#endif
 			}
 
 			#if STRATA_FASTPATH==0
 			BRANCH
-			if (BSDF_GETISSIMPLEVOLUME(BSDFContext.BSDF))
-			{
-				FParticipatingMedia PM = StrataSlabCreateParticipatingMedia(DiffuseColor, SLAB_SSSMFP(BSDFContext.BSDF));
-				const float DiffuseToVolumeBlend = StrataSlabDiffuseToVolumeBlend(PM);
-
-				// The environment response of has been measured for an hemisphere having a uniform radiance of 1. 
-				// This DirectionalAlbedo also contains the integral over the hemisphere according to the phase function.
-				// The spherical harmonic used for environment lighting contains the diffuse integral over the hemisphere as DiffLambert = int_{\omega} {1 * NoL d\Omega} = PI.
-				// We do not want to be affected by the NoL term here since The DirectionalAlbedo also contains the integral over the hemisphere according to the phase function.
-				// Integral of a value=1 over the Hemisphere is int_{\omega} {1  d\Omega} = 2PI.
-				// So, as an approximation, we recover the uniform hemisphere luminance as UniformEnvL = DiffLambert * (1 / PI) * (2 * PI).
-				// We consider that the diffuse integration of the environment SH over the hemisphere multiplied with the light response will be a good approximation similar to the "split sum integral".
-				float3 SlabDirectionalAlbedo = IsotropicMediumSlabEnvDirectionalAlbedo(PM) * ((1 / PI)* (2 * PI));
-
-			#if 1
-				// Because the default EnvBRDF does not account for multiple scattering, we rescale the participating media response according to EnvBRDF.
-				// This is to ensure a monotone light response behavior when transitioning from the Diffuse to the Volumetric model (without that, multiple scattering can result in higher luminance while trnasitionning from diffuse to volume light model)
-				// STRATA_TODO: We should NOT have to do that. This should be investigated futher...
-				const float MaxDiffuseWeight = max3(StrataEnvLightResult.DiffuseWeight.x, StrataEnvLightResult.DiffuseWeight.y, StrataEnvLightResult.DiffuseWeight.z);
-				const float MaxSlabDirectionalAlbedo = max3(SlabDirectionalAlbedo.x, SlabDirectionalAlbedo.y, SlabDirectionalAlbedo.z);
-				if (MaxDiffuseWeight < MaxSlabDirectionalAlbedo)
-				{
-					// We only want to do that when the slab light response is higher than the diffse response, so that when the mfp is large (and scattering low), the light response still work correctly according to thickness and MFP.
-					SlabDirectionalAlbedo *= MaxDiffuseWeight / MaxSlabDirectionalAlbedo;
-				}
-			#endif
-
-				// The diffuse is attenuated by the GGX interface SpecularTransmission.
-				SlabDirectionalAlbedo *= DirectionalAlbedo_SpecularTransmission;
-
-				StrataEnvLightResult.DiffuseWeight = lerp(StrataEnvLightResult.DiffuseWeight, SlabDirectionalAlbedo, DiffuseToVolumeBlend);
-				// We keep the same normal
-				StrataEnvLightResult.bSubsurface = false;
-			}
-			#endif
-			
-
-			#if STRATA_FASTPATH==0
-			BRANCH
@@ -1309,6 +1387,7 @@
 				StrataEnvLightResult.DiffuseWeight		*= Cloth_DirectionalAlbedo_SpecularTransmission;
 				StrataEnvLightResult.SpecularWeight		*= Cloth_DirectionalAlbedo_SpecularTransmission;
 				StrataEnvLightResult.SpecularHazeWeight *= Cloth_DirectionalAlbedo_SpecularTransmission;
+				StrataEnvLightResult.SSRSpecularWeight	*= Cloth_DirectionalAlbedo_SpecularTransmission;
 				
 				// We currently send the cloth environment lighting to be taken into account by the diffuse term.
 				// STARTA_TODO: true a mix of diffuse and specular contribution according to roughness.
diff --git a/Engine/Shaders/Private/Strata/StrataExport.ush b/Engine/Shaders/Private/Strata/StrataExport.ush
--- a/Engine/Shaders/Private/Strata/StrataExport.ush
+++ b/Engine/Shaders/Private/Strata/StrataExport.ush
@@ -198,7 +198,20 @@ void PackStrataOut(
 						bIsSingleMaterial = bIsSingleMaterial && BSDF_GETHASANISOTROPY(BSDF) == 0 && LuminanceWeightEqualOne;
 
 						TopLayerTotalWeight += BSDF.TopLayerDataWeight;
-						TopLayerData.Roughness += BSDF.TopLayerDataWeight * SLAB_ROUGHNESS(BSDF);
+						float TopLayerRoughnessContribution = BSDF.TopLayerDataWeight * SLAB_ROUGHNESS(BSDF);
+						float TopLayerBaseColorContribution = BSDF.TopLayerDataWeight * StrataGetBSDFBaseColor(BSDF);
+						if (BSDF_GETHASHAZINESS(BSDF))
+						{
+							FHaziness Haziness = UnpackHaziness(SLAB_HAZINESS(BSDF));
+							if (Haziness.bSimpleClearCoat)
+							{
+								TopLayerRoughnessContribution = BSDF.TopLayerDataWeight * lerp(SLAB_ROUGHNESS(BSDF), Haziness.Roughness, Haziness.Weight);
+								TopLayerBaseColorContribution = BSDF.TopLayerDataWeight * lerp(0.0, 0.04f, F0RGBToMetallic(0.04f));
+							}
+						}
+						TopLayerData.Roughness += TopLayerRoughnessContribution;
+						TopLayerData.BaseColor += TopLayerBaseColorContribution;
+
 
 						const bool bHasSSS = BSDF_GETSSSTYPE(BSDF) != SSS_TYPE_INVALID;
 						const bool bIsThin = BSDF_GETISTHIN(BSDF);
@@ -219,6 +232,7 @@ void PackStrataOut(
 
 						TopLayerTotalWeight += BSDF.TopLayerDataWeight;
 						TopLayerData.Roughness += BSDF.TopLayerDataWeight * FullyRough;
+						TopLayerData.BaseColor += BSDF.TopLayerDataWeight * StrataGetBSDFBaseColor(BSDF);
 
 						SetMaterialMode(StrataPixelHeader, HEADER_MATERIALMODE_HAIR);
 						break;
@@ -246,13 +260,13 @@ void PackStrataOut(
 
 						TopLayerTotalWeight += BSDF.TopLayerDataWeight;
 						TopLayerData.Roughness += BSDF.TopLayerDataWeight * SLW_ROUGHNESS(BSDF);
+						TopLayerData.BaseColor += BSDF.TopLayerDataWeight * StrataGetBSDFBaseColor(BSDF);
 
 						SetMaterialMode(StrataPixelHeader, HEADER_MATERIALMODE_SLWATER);
 						break;
 					}
 					}
 
-					TopLayerData.BaseColor += BSDF.TopLayerDataWeight * StrataGetBSDFBaseColor(BSDF);
 
 					float3x3 TangentBasis = StrataGetBSDFSharedBasis_InlineShading(StrataPixelHeader, BSDF, StrataAddressing);
 					float3 N = TangentBasis[2];
diff --git a/Engine/Shaders/Private/Strata/StrataLegacyConversion.ush b/Engine/Shaders/Private/Strata/StrataLegacyConversion.ush
--- a/Engine/Shaders/Private/Strata/StrataLegacyConversion.ush
+++ b/Engine/Shaders/Private/Strata/StrataLegacyConversion.ush
@@ -97,6 +97,7 @@ FStrataData CreateLegacySlab(
 			In.Emissive,		// Emissive
 			0.f,				// SecondRoughness
 			0.f,				// SecondRoughnessWeight
+			0.f,				// SecondRoughnessAsSimpleClearCoat
 			In.FuzzAmount,		// FuzzAmount
 			In.FuzzColor,		// FuzzColor
 			In.Thickness,
diff --git a/Engine/Shaders/Private/Strata/StrataVisualize.usf b/Engine/Shaders/Private/Strata/StrataVisualize.usf
--- a/Engine/Shaders/Private/Strata/StrataVisualize.usf
+++ b/Engine/Shaders/Private/Strata/StrataVisualize.usf
@@ -256,10 +256,10 @@ void PrintBSDF(inout FShaderPrintContext Context, inout FStrataAddressing Strata
 			{
 				FHaziness Haziness = UnpackHaziness(SLAB_HAZINESS(BSDF));
 
-				Print(Context, TEXT(" Haziness    "),		FontBSDFPropName);																		Newline(Context, RectMax);
-				Print(Context, TEXT("     Roughness    "),	FontBSDFPropName);	Print(Context, Haziness.Roughness,				FontBSDFPropValu);	Newline(Context, RectMax);
-				Print(Context, TEXT("     Weight       "),	FontBSDFPropName);	Print(Context, Haziness.Weight,					FontBSDFPropValu);	Newline(Context, RectMax);
-				Print(Context, TEXT("     FastClearCoat"),	FontBSDFPropName);	Print(Context, Haziness.bFastClearCoat ?1u:0u,	FontBSDFPropValu);	Newline(Context, RectMax);
+				Print(Context, TEXT(" Haziness    "),			FontBSDFPropName);																		Newline(Context, RectMax);
+				Print(Context, TEXT("     Roughness    "),		FontBSDFPropName);	Print(Context, Haziness.Roughness,				FontBSDFPropValu);	Newline(Context, RectMax);
+				Print(Context, TEXT("     Weight       "),		FontBSDFPropName);	Print(Context, Haziness.Weight,					FontBSDFPropValu);	Newline(Context, RectMax);
+				Print(Context, TEXT("     SimpleClearCoat"),	FontBSDFPropName);	Print(Context, Haziness.bSimpleClearCoat?1u:0u,	FontBSDFPropValu);	Newline(Context, RectMax);
 			}
 
 			if (BSDF_GETSSSTYPE(BSDF) != SSS_TYPE_INVALID || BSDF_GETISSIMPLEVOLUME(BSDF) || BSDF_GETISTHIN(BSDF))
diff --git a/Engine/Source/Runtime/Engine/Classes/Materials/MaterialExpressionStrata.h b/Engine/Source/Runtime/Engine/Classes/Materials/MaterialExpressionStrata.h
--- a/Engine/Source/Runtime/Engine/Classes/Materials/MaterialExpressionStrata.h
+++ b/Engine/Source/Runtime/Engine/Classes/Materials/MaterialExpressionStrata.h
@@ -29,8 +29,13 @@ class UMaterialExpressionStrataBSDF : public UMaterialExpression
 	GENERATED_UCLASS_BODY()
 
 #if WITH_EDITOR
-		virtual int32 CompilePreview(class FMaterialCompiler* Compiler, int32 OutputIndex) override;
+	virtual int32 CompilePreview(class FMaterialCompiler* Compiler, int32 OutputIndex) override;
 #endif
+
+	float DielectricSpecularToF0(float SpecularIn)
+	{
+		return 0.08f * SpecularIn;
+	};
 };
 
 
@@ -345,6 +350,75 @@ class UMaterialExpressionStrataSlabBSDF : public UMaterialExpressionStrataBSDF
 	//~ End UMaterialExpression Interface
 };
 
+UCLASS(MinimalAPI, collapsecategories, hidecategories = Object, DisplayName = "Strata Simple Clear Coat")
+class UMaterialExpressionStrataSimpleClearCoatBSDF : public UMaterialExpressionStrataBSDF
+{
+	GENERATED_UCLASS_BODY()
+
+	/**
+	 * Defines the overall color of the Material. (type = float3, unit = unitless, defaults to 0.18)
+	 */
+	UPROPERTY()
+	FExpressionInput BaseColor;
+
+	/**
+	 * Controls how \"metal-like\" your surface looks like. 0 means dielectric, 1 means conductor (type = float, unit = unitless, defaults to 0)
+	 */
+	UPROPERTY()
+	FExpressionInput Metallic;
+	
+	/**
+	 * Used to scale the current amount of specularity on non-metallic surfaces and is a value between 0 and 1 (type = float, unit = unitless, defaults to plastic 0.5)
+	 */
+	UPROPERTY()
+	FExpressionInput Specular;
+
+	/**
+	 * Controls how rough the bottom layer of the material is. Roughness of 0 (smooth) is a mirror reflection and 1 (rough) is completely matte or diffuse. (type = float, unit = unitless, defaults to 0.5)
+	 */
+	UPROPERTY()
+	FExpressionInput Roughness;
+
+	/**
+	 * Controls the coverage of the clear coat layer: 0 means no clear coat, 1 means coat is fully visible. (type = float, unit = unitless, defaults to 0.5)
+	 */
+	UPROPERTY()
+	FExpressionInput ClearCoatCoverage;
+
+	/**
+	 * Controls how rough the top layer of the material is. Roughness of 0 (smooth) is a mirror reflection and 1 (rough) is completely matte or diffuse. (type = float, unit = unitless, defaults to 0.5)
+	 */
+	UPROPERTY()
+	FExpressionInput ClearCoatRoughness;
+
+	/**
+	 * Take the surface normal as input. The normal is considered tangent or world space according to the space properties on the main material node. (type = float3, unit = unitless, defaults to vertex normal)
+	 */
+	UPROPERTY()
+	FExpressionInput Normal;
+
+	/**
+	 * Emissive color of the medium (type = float3, unit = luminance, default = 0)
+	 */
+	UPROPERTY()
+	FExpressionInput EmissiveColor;
+
+	//~ Begin UMaterialExpression Interface
+#if WITH_EDITOR
+	virtual int32 Compile(class FMaterialCompiler* Compiler, int32 OutputIndex) override;
+	virtual void GetCaption(TArray<FString>& OutCaptions) const override;
+	virtual uint32 GetOutputType(int32 OutputIndex) override;
+	virtual uint32 GetInputType(int32 InputIndex) override;
+	virtual bool IsResultStrataMaterial(int32 OutputIndex) override;
+	virtual void GatherStrataMaterialInfo(FStrataMaterialInfo& StrataMaterialInfo, int32 OutputIndex) override;
+	virtual FStrataOperator* StrataGenerateMaterialTopologyTree(class FMaterialCompiler* Compiler, class UMaterialExpression* Parent, int32 OutputIndex) override;
+	virtual FName GetInputName(int32 InputIndex) const override;
+	virtual void GetConnectorToolTip(int32 InputIndex, int32 OutputIndex, TArray<FString>& OutToolTip) override;
+	virtual const TArray<FExpressionInput*> GetInputs() override;
+#endif
+	//~ End UMaterialExpression Interface
+};
+
 UCLASS(MinimalAPI, collapsecategories, hidecategories = Object, DisplayName = "Strata Volumetric-Fog-Cloud BSDF")
 class UMaterialExpressionStrataVolumetricFogCloudBSDF : public UMaterialExpressionStrataBSDF
 {
diff --git a/Engine/Source/Runtime/Engine/Private/Materials/HLSLMaterialTranslator.cpp b/Engine/Source/Runtime/Engine/Private/Materials/HLSLMaterialTranslator.cpp
--- a/Engine/Source/Runtime/Engine/Private/Materials/HLSLMaterialTranslator.cpp
+++ b/Engine/Source/Runtime/Engine/Private/Materials/HLSLMaterialTranslator.cpp
@@ -1959,7 +1959,7 @@ void FHLSLMaterialTranslator::GetMaterialEnvironment(EShaderPlatform InPlatform,
 						{
 							StrataMaterialDescription += FString::Printf(TEXT("\tIdx=%d Op=%s ParentIdx=%d LeftIndex=%d RightIndex=%d BSDFIdx=%d LayerDepth=%d IsTop=%d IsBot=%d BSDFType=%s SSS=%d MFP=%d F90=%d Rough2=%d Fuzz=%d  Aniso=%d \r\n"),
 								It.Index, GetStrataOperatorStr(It.OperatorType), It.ParentIndex, It.LeftIndex, It.RightIndex, It.BSDFIndex, It.LayerDepth, It.bIsTop, It.bIsBottom, 
-								*GetStrataBSDFName(It.BSDFType), It.bBSDFHasSSS, It.bBSDFHasMFPPluggedIn, It.bBSDFHasEdgeColor, It.bBSDFHasSecondRoughness, It.bBSDFHasFuzz, It.bBSDFHasAnisotropy);
+								*GetStrataBSDFName(It.BSDFType), It.bBSDFHasSSS, It.bBSDFHasMFPPluggedIn, It.bBSDFHasEdgeColor, It.bBSDFHasSecondRoughnessOrSimpleClearCoat, It.bBSDFHasFuzz, It.bBSDFHasAnisotropy);
 						}
 					}
 				}
@@ -10016,7 +10016,7 @@ bool FHLSLMaterialTranslator::StrataGenerateDerivedMaterialOperatorData()
 				{
 				case STRATA_BSDF_TYPE_SLAB:
 				{
-					bStrataMaterialIsSimple = bStrataMaterialIsSimple && !bMayHaveColoredWeight && !It.bBSDFHasAnisotropy && !It.bBSDFHasEdgeColor && !It.bBSDFHasFuzz && !It.bBSDFHasSecondRoughness && !It.bBSDFHasMFPPluggedIn && !It.bBSDFHasSSS;
+					bStrataMaterialIsSimple = bStrataMaterialIsSimple && !bMayHaveColoredWeight && !It.bBSDFHasAnisotropy && !It.bBSDFHasEdgeColor && !It.bBSDFHasFuzz && !It.bBSDFHasSecondRoughnessOrSimpleClearCoat && !It.bBSDFHasMFPPluggedIn && !It.bBSDFHasSSS;
 					bStrataMaterialIsSingle = bStrataMaterialIsSingle && !bMayHaveColoredWeight && !It.bBSDFHasAnisotropy;
 					break;
 				}
@@ -10116,7 +10116,7 @@ bool FHLSLMaterialTranslator::StrataGenerateDerivedMaterialOperatorData()
 					StrataMaterialRequestedSizeByte += UintByteSize;
 					StrataMaterialRequestedSizeByte += UintByteSize;
 
-					if (It.bBSDFHasEdgeColor || It.bBSDFHasSecondRoughness)
+					if (It.bBSDFHasEdgeColor || It.bBSDFHasSecondRoughnessOrSimpleClearCoat)
 					{
 						StrataMaterialRequestedSizeByte += UintByteSize;
 					}
@@ -10303,7 +10303,7 @@ int32 FHLSLMaterialTranslator::StrataSlabBSDF(
 	int32 Roughness, int32 Anisotropy,
 	int32 SSSProfileId, int32 SSSMFP, int32 SSSMFPScale, int32 SSSPhaseAniso, int32 UseSSSDiffusion,
 	int32 EmissiveColor,
-	int32 SecondRoughness, int32 SecondRoughnessWeight,
+	int32 SecondRoughness, int32 SecondRoughnessWeight, int32 SecondRoughnessAsSimpleClearCoat,
 	int32 FuzzAmount, int32 FuzzColor,
 	int32 Thickness, 
 	int32 Normal, int32 Tangent, const FString& SharedLocalBasisIndexMacro,
@@ -10320,7 +10320,7 @@ int32 FHLSLMaterialTranslator::StrataSlabBSDF(
 			return INDEX_NONE;
 		}
 		return AddCodeChunk(
-			MCT_Strata, TEXT("PromoteParameterBlendedBSDFToOperator(GetStrataSlabBSDF(Parameters.StrataPixelFootprint, %s, %s, %s, %s, %s, %s, %s, %s, %s, %s, %s, %s, %s, %s, %s, %s, %s, %s, %s, %s, %s, %s, Parameters.SharedLocalBases.Types) /* Normal = %s ; Tangent = %s */, Parameters.StrataTree, %u, %u, %u, %u)"),
+			MCT_Strata, TEXT("PromoteParameterBlendedBSDFToOperator(GetStrataSlabBSDF(Parameters.StrataPixelFootprint, %s, %s, %s, %s, %s, %s, %s, %s, %s, %s, %s, %s, %s, %s, %s, %s, %s, %s, %s, %s, %s, %s, %s, Parameters.SharedLocalBases.Types) /* Normal = %s ; Tangent = %s */, Parameters.StrataTree, %u, %u, %u, %u)"),
 			*StrataGetCastParameterCode(UseMetalness,			MCT_Float),
 			*StrataGetCastParameterCode(BaseColor,				MCT_Float3),
 			*StrataGetCastParameterCode(EdgeColor,				MCT_Float3),
@@ -10339,6 +10339,7 @@ int32 FHLSLMaterialTranslator::StrataSlabBSDF(
 			*StrataGetCastParameterCode(EmissiveColor,			MCT_Float3),
 			*StrataGetCastParameterCode(SecondRoughness,		MCT_Float),
 			*StrataGetCastParameterCode(SecondRoughnessWeight,	MCT_Float),
+			*StrataGetCastParameterCode(SecondRoughnessAsSimpleClearCoat, MCT_Float),
 			*StrataGetCastParameterCode(FuzzAmount,				MCT_Float),
 			*StrataGetCastParameterCode(FuzzColor,				MCT_Float3),
 			*StrataGetCastParameterCode(Thickness,				MCT_Float),
@@ -10352,7 +10353,7 @@ int32 FHLSLMaterialTranslator::StrataSlabBSDF(
 	}
 	
 	return AddCodeChunk(
-		MCT_Strata, TEXT("GetStrataSlabBSDF(Parameters.StrataPixelFootprint, %s, %s, %s, %s, %s, %s, %s, %s, %s, %s, %s, %s, %s, %s, %s, %s, %s, %s, %s, %s, %s, %s, Parameters.SharedLocalBases.Types) /* Normal = %s ; Tangent = %s */"),
+		MCT_Strata, TEXT("GetStrataSlabBSDF(Parameters.StrataPixelFootprint, %s, %s, %s, %s, %s, %s, %s, %s, %s, %s, %s, %s, %s, %s, %s, %s, %s, %s, %s, %s, %s, %s, %s, Parameters.SharedLocalBases.Types) /* Normal = %s ; Tangent = %s */"),
 		*StrataGetCastParameterCode(UseMetalness,			MCT_Float),
 		*StrataGetCastParameterCode(BaseColor,				MCT_Float3),
 		*StrataGetCastParameterCode(EdgeColor,				MCT_Float3),
@@ -10371,6 +10372,7 @@ int32 FHLSLMaterialTranslator::StrataSlabBSDF(
 		*StrataGetCastParameterCode(EmissiveColor,			MCT_Float3),
 		*StrataGetCastParameterCode(SecondRoughness,		MCT_Float),
 		*StrataGetCastParameterCode(SecondRoughnessWeight,	MCT_Float),
+		*StrataGetCastParameterCode(SecondRoughnessAsSimpleClearCoat, MCT_Float),
 		*StrataGetCastParameterCode(FuzzAmount,				MCT_Float),
 		*StrataGetCastParameterCode(FuzzColor,				MCT_Float3),
 		*StrataGetCastParameterCode(Thickness,				MCT_Float),
diff --git a/Engine/Source/Runtime/Engine/Private/Materials/HLSLMaterialTranslator.h b/Engine/Source/Runtime/Engine/Private/Materials/HLSLMaterialTranslator.h
--- a/Engine/Source/Runtime/Engine/Private/Materials/HLSLMaterialTranslator.h
+++ b/Engine/Source/Runtime/Engine/Private/Materials/HLSLMaterialTranslator.h
@@ -1025,7 +1025,7 @@ protected:
 		int32 Roughness, int32 Anisotropy,
 		int32 SSSProfileId, int32 SSSMFP, int32 SSSMFPScale, int32 SSSPhaseAniso, int32 bUseSSSDiffusion,
 		int32 EmissiveColor, 
-		int32 SecondRoughness, int32 SecondRoughnessWeight,
+		int32 SecondRoughness, int32 SecondRoughnessWeight, int32 SecondRoughnessAsSimpleClearCoat,
 		int32 FuzzAmount, int32 FuzzColor, 
 		int32 Thickness,
 		int32 Normal, int32 Tangent, const FString& SharedLocalBasisIndexMacro,
diff --git a/Engine/Source/Runtime/Engine/Private/Materials/MaterialExpressions.cpp b/Engine/Source/Runtime/Engine/Private/Materials/MaterialExpressions.cpp
--- a/Engine/Source/Runtime/Engine/Private/Materials/MaterialExpressions.cpp
+++ b/Engine/Source/Runtime/Engine/Private/Materials/MaterialExpressions.cpp
@@ -22058,6 +22058,7 @@ int32 UMaterialExpressionStrataSlabBSDF::Compile(class FMaterialCompiler* Compil
 		CompileWithDefaultFloat3(Compiler, EmissiveColor, 0.0f, 0.0f, 0.0f),
 		CompileWithDefaultFloat1(Compiler, SecondRoughness, 0.0f),
 		CompileWithDefaultFloat1(Compiler, SecondRoughnessWeight, 0.0f),
+		Compiler->Constant(0.0f),										// SecondRoughnessAsSimpleClearCoat
 		CompileWithDefaultFloat1(Compiler, FuzzAmount, 0.0f),
 		CompileWithDefaultFloat3(Compiler, FuzzColor, 0.0f, 0.0f, 0.0f),
 		CompileWithDefaultFloat1(Compiler, Thickness, STRATA_LAYER_DEFAULT_THICKNESS_CM),
@@ -22501,7 +22502,7 @@ FStrataOperator* UMaterialExpressionStrataSlabBSDF::StrataGenerateMaterialTopolo
 	StrataOperator.BSDFType = STRATA_BSDF_TYPE_SLAB;
 	StrataOperator.bBSDFHasEdgeColor = HasEdgeColor();
 	StrataOperator.bBSDFHasFuzz = HasFuzz();
-	StrataOperator.bBSDFHasSecondRoughness = HasSecondRoughness();
+	StrataOperator.bBSDFHasSecondRoughnessOrSimpleClearCoat = HasSecondRoughness();
 	StrataOperator.bBSDFHasSSS = HasSSS();
 	StrataOperator.bBSDFHasMFPPluggedIn = HasMFPPluggedIn();
 	StrataOperator.bBSDFHasAnisotropy = HasAnisotropy();
@@ -22546,6 +22547,248 @@ bool UMaterialExpressionStrataSlabBSDF::HasAnisotropy() const
 
 
 
+UMaterialExpressionStrataSimpleClearCoatBSDF::UMaterialExpressionStrataSimpleClearCoatBSDF(const FObjectInitializer& ObjectInitializer)
+	: Super(ObjectInitializer)
+{
+	struct FConstructorStatics
+	{
+		FText NAME_Strata;
+		FConstructorStatics() : NAME_Strata(LOCTEXT("Strata BSDFs", "Strata BSDFs")) { }
+	};
+	static FConstructorStatics ConstructorStatics;
+#if WITH_EDITORONLY_DATA
+	MenuCategories.Add(ConstructorStatics.NAME_Strata);
+#endif
+}
+
+#if WITH_EDITOR
+int32 UMaterialExpressionStrataSimpleClearCoatBSDF::Compile(class FMaterialCompiler* Compiler, int32 OutputIndex)
+{
+	const bool  bUseMetalness = true;
+	const float DefaultSpecular = 0.5f;
+	const float DefaultF0 = DielectricSpecularToF0(DefaultSpecular);
+
+	int32 NormalCodeChunk = CompileWithDefaultNormalWS(Compiler, Normal);
+	const int32 NullTangentCodeChunk = INDEX_NONE;
+	const FStrataRegisteredSharedLocalBasis NewRegisteredSharedLocalBasis = StrataCompilationInfoCreateSharedLocalBasis(Compiler, NormalCodeChunk, NullTangentCodeChunk);
+
+	FStrataOperator& StrataOperator = Compiler->StrataCompilationGetOperator(this);
+	StrataOperator.BSDFRegisteredSharedLocalBasis = NewRegisteredSharedLocalBasis;
+
+	int32 OutputCodeChunk = Compiler->StrataSlabBSDF(
+		bUseMetalness ? Compiler->Constant(1.0f) : Compiler->Constant(0.0f),
+
+		// Metalness workflow
+		CompileWithDefaultFloat3(Compiler, BaseColor, 0.18f, 0.18f, 0.18f),
+		Compiler->Constant3(1.0f, 1.0f, 1.0f),					// EdgeColor
+		CompileWithDefaultFloat1(Compiler, Specular, DefaultSpecular),
+		CompileWithDefaultFloat1(Compiler, Metallic, 0.0f),
+
+		// !Metalness workflow unused for this node
+		Compiler->Constant3(0.18f, 0.18f, 0.18f),				// DiffuseAlbedo
+		Compiler->Constant3(DefaultF0, DefaultF0, DefaultF0),	// F0
+		Compiler->Constant3(1.0f, 1.0f, 1.0f),					// F90
+		
+		CompileWithDefaultFloat1(Compiler, Roughness, 0.5f),
+		Compiler->Constant(0.0f),								// Anisotropy
+		Compiler->Constant(0.0f),								// SSSProfile
+		Compiler->Constant3(0.0f, 0.0f, 0.0f),					// SSSMFP
+		Compiler->Constant(0.0f),								// SSSMFPScale
+		Compiler->Constant(0.0f),								// SSSPhaseAnisotropy
+		Compiler->Constant(0.0f),								// bUseSSSDiffusion
+		CompileWithDefaultFloat3(Compiler, EmissiveColor, 0.0f, 0.0f, 0.0f),
+		CompileWithDefaultFloat1(Compiler, ClearCoatRoughness, 0.1f),
+		CompileWithDefaultFloat1(Compiler, ClearCoatCoverage, 1.0f),
+		Compiler->Constant(1.0f),								// SecondRoughnessAsSimpleClearCoat == true for UMaterialExpressionStrataSimpleClearCoatBSDF
+		Compiler->Constant(0.0f),								// FuzzAmount
+		Compiler->Constant3(0.0f, 0.0f, 0.0f),					// FuzzColor
+		Compiler->Constant(STRATA_LAYER_DEFAULT_THICKNESS_CM),	// Thickness
+		NormalCodeChunk,
+		NullTangentCodeChunk,
+		Compiler->GetStrataSharedLocalBasisIndexMacro(NewRegisteredSharedLocalBasis),
+		!StrataOperator.bUseParameterBlending || (StrataOperator.bUseParameterBlending && StrataOperator.bRootOfParameterBlendingSubTree) ? &StrataOperator : nullptr);
+
+	return OutputCodeChunk;
+}
+
+
+const TArray<FExpressionInput*> UMaterialExpressionStrataSimpleClearCoatBSDF::GetInputs()
+{
+	TArray<FExpressionInput*> Result;
+	Result.Add(&BaseColor);
+	Result.Add(&Metallic);
+	Result.Add(&Specular);
+	Result.Add(&Roughness);
+	Result.Add(&ClearCoatRoughness);
+	Result.Add(&ClearCoatCoverage);
+	Result.Add(&Normal);
+	Result.Add(&EmissiveColor);
+	return Result;
+}
+
+void UMaterialExpressionStrataSimpleClearCoatBSDF::GetCaption(TArray<FString>& OutCaptions) const
+{
+	OutCaptions.Add(TEXT("Strata Simple Clear Coat"));
+}
+
+uint32 UMaterialExpressionStrataSimpleClearCoatBSDF::GetOutputType(int32 OutputIndex)
+{
+	return MCT_Strata;
+}
+
+uint32 UMaterialExpressionStrataSimpleClearCoatBSDF::GetInputType(int32 InputIndex)
+{
+	if (InputIndex == 0)
+	{
+		return MCT_Float3; // BaseColor
+	}
+	else if (InputIndex == 1)
+	{
+		return MCT_Float1; // Metallic
+	}
+	else if (InputIndex == 2)
+	{
+		return MCT_Float1; // Specular
+	}
+	else if (InputIndex == 3)
+	{
+		return MCT_Float1; // Roughness
+	}
+	else if (InputIndex == 4)
+	{
+		return MCT_Float1; // ClearCoatRoughness
+	}
+	else if (InputIndex == 5)
+	{
+		return MCT_Float1; // ClearCoatCoverage
+	}
+	else if (InputIndex == 6)
+	{
+		return MCT_Float3; // Normal
+	}
+	else if (InputIndex == 7)
+	{
+		return MCT_Float3; // Emissive Color
+	}
+
+	check(false);
+	return MCT_Float1;
+}
+
+FName UMaterialExpressionStrataSimpleClearCoatBSDF::GetInputName(int32 InputIndex) const
+{
+	if (InputIndex == 0)
+	{
+		return TEXT("BaseColor");
+	}
+	else if (InputIndex == 1)
+	{
+		return TEXT("Metallic");
+	}
+	else if (InputIndex == 2)
+	{
+		return TEXT("Specular");
+	}
+	else if (InputIndex == 3)
+	{
+		return TEXT("Roughness");
+	}
+	else if (InputIndex == 4)
+	{
+		return TEXT("Clear Coat Roughness");
+	}
+	else if (InputIndex == 5)
+	{
+		return TEXT("Clear Coat Coverage");
+	}
+	else if (InputIndex == 6)
+	{
+		return TEXT("Normal");
+	}
+	else if (InputIndex == 7)
+	{
+		return TEXT("Emissive Color");
+	}
+
+	return TEXT("Unknown");
+}
+
+void UMaterialExpressionStrataSimpleClearCoatBSDF::GetConnectorToolTip(int32 InputIndex, int32 OutputIndex, TArray<FString>& OutToolTip)
+{
+	if (OutputIndex == 0)
+	{
+		OutToolTip.Add(TEXT("TT Ouput"));
+		return;
+	}
+
+	if (InputIndex == 0)
+	{
+		Super::GetConnectorToolTip(0, INDEX_NONE, OutToolTip);
+	}
+	else if (InputIndex == 1)
+	{
+		Super::GetConnectorToolTip(1, INDEX_NONE, OutToolTip);
+	}
+	else if (InputIndex == 2)
+	{
+		Super::GetConnectorToolTip(2, INDEX_NONE, OutToolTip);
+	}
+	else if (InputIndex == 3)
+	{
+		Super::GetConnectorToolTip(3, INDEX_NONE, OutToolTip);
+	}
+	else if (InputIndex == 4)
+	{
+		Super::GetConnectorToolTip(4, INDEX_NONE, OutToolTip);
+	}
+	else if (InputIndex == 5)
+	{
+		Super::GetConnectorToolTip(5, INDEX_NONE, OutToolTip);
+	}
+	else if (InputIndex == 6)
+	{
+		Super::GetConnectorToolTip(6, INDEX_NONE, OutToolTip);
+	}
+	else if (InputIndex == 7)
+	{
+		Super::GetConnectorToolTip(7, INDEX_NONE, OutToolTip);
+	}
+}
+
+bool UMaterialExpressionStrataSimpleClearCoatBSDF::IsResultStrataMaterial(int32 OutputIndex)
+{
+	return true;
+}
+
+void UMaterialExpressionStrataSimpleClearCoatBSDF::GatherStrataMaterialInfo(FStrataMaterialInfo& StrataMaterialInfo, int32 OutputIndex)
+{
+	// Track connected inputs
+	if (BaseColor.IsConnected()) { StrataMaterialInfo.AddPropertyConnected(MP_BaseColor); }
+	if (Metallic.IsConnected()) { StrataMaterialInfo.AddPropertyConnected(MP_Metallic); }
+	if (Specular.IsConnected()) { StrataMaterialInfo.AddPropertyConnected(MP_Specular); }
+	if (Roughness.IsConnected()) { StrataMaterialInfo.AddPropertyConnected(MP_Roughness); }
+	if (Normal.IsConnected()) { StrataMaterialInfo.AddPropertyConnected(MP_Normal); }
+	if (EmissiveColor.IsConnected()) { StrataMaterialInfo.AddPropertyConnected(MP_EmissiveColor); }
+
+	StrataMaterialInfo.AddShadingModel(SSM_DefaultLit);
+}
+
+FStrataOperator* UMaterialExpressionStrataSimpleClearCoatBSDF::StrataGenerateMaterialTopologyTree(class FMaterialCompiler* Compiler, class UMaterialExpression* Parent, int32 OutputIndex)
+{
+	FStrataOperator& StrataOperator = Compiler->StrataCompilationRegisterOperator(STRATA_OPERATOR_BSDF, this, Parent);
+	StrataOperator.BSDFType = STRATA_BSDF_TYPE_SLAB;
+	StrataOperator.bBSDFHasEdgeColor = false;
+	StrataOperator.bBSDFHasFuzz = false;
+	StrataOperator.bBSDFHasSecondRoughnessOrSimpleClearCoat = true;	// This node explicitely requires simple clear coat
+	StrataOperator.bBSDFHasSSS = false;
+	StrataOperator.bBSDFHasMFPPluggedIn = false;
+	StrataOperator.bBSDFHasAnisotropy = false;
+	return &StrataOperator;
+}
+#endif // WITH_EDITOR
+
+
+
 UMaterialExpressionStrataVolumetricFogCloudBSDF::UMaterialExpressionStrataVolumetricFogCloudBSDF(const FObjectInitializer& ObjectInitializer)
 	: Super(ObjectInitializer)
 {
diff --git a/Engine/Source/Runtime/Engine/Public/MaterialCompiler.h b/Engine/Source/Runtime/Engine/Public/MaterialCompiler.h
--- a/Engine/Source/Runtime/Engine/Public/MaterialCompiler.h
+++ b/Engine/Source/Runtime/Engine/Public/MaterialCompiler.h
@@ -111,7 +111,7 @@ struct FStrataOperator
 	bool  bBSDFHasMFPPluggedIn;
 	bool  bBSDFHasEdgeColor;
 	bool  bBSDFHasFuzz;
-	bool  bBSDFHasSecondRoughness;
+	bool  bBSDFHasSecondRoughnessOrSimpleClearCoat;
 	bool  bBSDFHasAnisotropy;
 
 	// Data derived after the tree has been built.
@@ -137,7 +137,7 @@ struct FStrataOperator
 		bBSDFHasMFPPluggedIn = false;
 		bBSDFHasEdgeColor = false;
 		bBSDFHasFuzz = false;
-		bBSDFHasSecondRoughness = false;
+		bBSDFHasSecondRoughnessOrSimpleClearCoat = false;
 		bBSDFHasAnisotropy = false;
 
 		MaxDistanceFromLeaves = 0;
@@ -154,7 +154,7 @@ struct FStrataOperator
 		bBSDFHasMFPPluggedIn = A.bBSDFHasMFPPluggedIn || B.bBSDFHasMFPPluggedIn;
 		bBSDFHasEdgeColor = A.bBSDFHasEdgeColor || B.bBSDFHasEdgeColor;
 		bBSDFHasFuzz = A.bBSDFHasFuzz || B.bBSDFHasFuzz;
-		bBSDFHasSecondRoughness = A.bBSDFHasSecondRoughness || B.bBSDFHasSecondRoughness;
+		bBSDFHasSecondRoughnessOrSimpleClearCoat = A.bBSDFHasSecondRoughnessOrSimpleClearCoat || B.bBSDFHasSecondRoughnessOrSimpleClearCoat;
 		bBSDFHasAnisotropy = A.bBSDFHasAnisotropy || B.bBSDFHasAnisotropy;
 	}
 
@@ -164,7 +164,7 @@ struct FStrataOperator
 		bBSDFHasMFPPluggedIn = A.bBSDFHasMFPPluggedIn;
 		bBSDFHasEdgeColor = A.bBSDFHasEdgeColor;
 		bBSDFHasFuzz = A.bBSDFHasFuzz;
-		bBSDFHasSecondRoughness = A.bBSDFHasSecondRoughness;
+		bBSDFHasSecondRoughnessOrSimpleClearCoat = A.bBSDFHasSecondRoughnessOrSimpleClearCoat;
 		bBSDFHasAnisotropy = A.bBSDFHasAnisotropy;
 	}
 
@@ -568,7 +568,7 @@ public:
 		int32 Roughness, int32 Anisotropy,
 		int32 SSSProfileId, int32 SSSMFP, int32 SSSMFPScale, int32 SSSPhaseAniso, int32 UseSSSDiffusion,
 		int32 EmissiveColor, 
-		int32 SecondRoughness, int32 SecondRoughnessWeight,
+		int32 SecondRoughness, int32 SecondRoughnessWeight, int32 SecondRoughnessAsSimpleClearCoat,
 		int32 FuzzAmount, int32 FuzzColor,
 		int32 Thickness,
 		int32 Normal, int32 Tangent, const FString& SharedLocalBasisIndexMacro,
@@ -1118,7 +1118,7 @@ public:
 		int32 Roughness, int32 Anisotropy,
 		int32 SSSProfileId, int32 SSSMFP, int32 SSSMFPScale, int32 SSSPhaseAniso, int32 UseSSSDiffusion,
 		int32 EmissiveColor, 
-		int32 SecondRoughness, int32 SecondRoughnessWeight,
+		int32 SecondRoughness, int32 SecondRoughnessWeight, int32 SecondRoughnessAsSimpleClearCoat,
 		int32 FuzzAmount, int32 FuzzColor,
 		int32 Thickness,
 		int32 Normal, int32 Tangent, const FString& SharedLocalBasisIndexMacro,
@@ -1131,7 +1131,7 @@ public:
 			Roughness, Anisotropy,
 			SSSProfileId, SSSMFP, SSSMFPScale, SSSPhaseAniso, UseSSSDiffusion,
 			EmissiveColor, 
-			SecondRoughness, SecondRoughnessWeight,
+			SecondRoughness, SecondRoughnessWeight, SecondRoughnessAsSimpleClearCoat,
 			FuzzAmount, FuzzColor,
 			Thickness,
 			Normal, Tangent, SharedLocalBasisIndexMacro,
