Strata - simple clear coat fixes for wrong math with coverage and HazeClearCoatTransmittance

diff --git a/Engine/Shaders/Private/Strata/StrataEvaluation.ush b/Engine/Shaders/Private/Strata/StrataEvaluation.ush
--- a/Engine/Shaders/Private/Strata/StrataEvaluation.ush
+++ b/Engine/Shaders/Private/Strata/StrataEvaluation.ush
@@ -301,10 +301,10 @@ FStrataEvaluateResult StrataEvaluateBSDFCommon(FStrataBSDFContext BSDFContext, F
 			// Apply energy conservation on the diffuse component
 			// If the specular layer is anisotropic, the energy term is computed onto the 'main' roughness [Kulla 2019]
 			//
-			float3 MSScale = 1;
-			float3 HazeMSScale = 1;
+			float3 MSScale = 1.0f;
+			float3 HazeMSScale = 1.0f;
 			float3 DirectionalAlbedo_SpecularTransmission = 1.0f;
-			float HazeWeight = 0.f;
+			float HazeWeight = 0.0f;
 			bool  bHazeAsSimpleClearCoat = false;
 			float3 HazeClearCoatTransmittance = 1.0f;
 			float3 HazeF0 = F0;
@@ -451,12 +451,9 @@ FStrataEvaluateResult StrataEvaluateBSDFCommon(FStrataBSDFContext BSDFContext, F
 			}
 
 			const float SpecularMicroOcclusion = F0RGBToMicroOcclusion(F0);
-			const float3 FresnelTerm = SpecularMicroOcclusion * F_Schlick(F0, F90, BSDFContext.Context.VoH);
 
 			const float HazeSpecularMicroOcclusion = F0RGBToMicroOcclusion(HazeF0);
-			const float3 HazeFresnelTerm = SpecularMicroOcclusion * F_Schlick(HazeF0, HazeF90, BSDFContext.Context.VoH);
-
-			const float3 TopLayerSpecularTransmittionApprox = saturate(1.0f - HazeFresnelTerm); // Simple and do not use the LUT for now
+			const float3 HazeFresnelTerm = HazeSpecularMicroOcclusion * F_Schlick(HazeF0, HazeF90, BSDFContext.Context.VoH);
 
 			float3 CommonTerm = 0.0f;
 			if (bIsRectLight)
@@ -468,7 +465,9 @@ FStrataEvaluateResult StrataEvaluateBSDFCommon(FStrataBSDFContext BSDFContext, F
 			}
 			else
 			{
-				Sample.SpecularPathValue     = D     * Vis     * FresnelTerm     * MSScale     * HazeClearCoatTransmittance;
+				const float3 FresnelTerm = SpecularMicroOcclusion * F_Schlick(F0, F90, BSDFContext.Context.VoH);
+
+				Sample.SpecularPathValue     = D     * Vis     * FresnelTerm     * MSScale;
 				Sample.SpecularHazePathValue = HazeD * HazeVis * HazeFresnelTerm * HazeMSScale;
 
 				CommonTerm = ShadowTerms.SurfaceShadow * AreaLightContext.NoL * AreaLightContext.AreaLight.FalloffColor * AreaLightContext.Falloff;
@@ -480,7 +479,11 @@ FStrataEvaluateResult StrataEvaluateBSDFCommon(FStrataBSDFContext BSDFContext, F
 			if (bHazeAsSimpleClearCoat)
 			{
 				const float TopLayerCoverage = HazeWeight;
-				const float3 BottomLayerFactors = lerp(1.0, HazeClearCoatTransmittance * TopLayerSpecularTransmittionApprox, TopLayerCoverage);
+				const float3 TopLayerSpecularTransmittionApprox = saturate(1.0f - HazeFresnelTerm); // Simple and do not use the LUT for now
+				const float3 BottomLayerFactors = lerp(1.0f, HazeClearCoatTransmittance * TopLayerSpecularTransmittionApprox, TopLayerCoverage);
+
+				// Apply the top clear coat throughput onto the bottom layer
+				Sample.SpecularPathValue *= BottomLayerFactors;
 
 				const float BottomLayerFactorsGrey = dot(BottomLayerFactors, (1.0 / 3.0).xxx);
 				Sample.SpecularPathProbability     = BottomLayerFactorsGrey / (TopLayerCoverage + BottomLayerFactorsGrey);
