Strata - simple clear coat is a better match now.
Culled out some code from the fast path.

diff --git a/Engine/Shaders/Private/Strata/StrataEvaluation.ush b/Engine/Shaders/Private/Strata/StrataEvaluation.ush
--- a/Engine/Shaders/Private/Strata/StrataEvaluation.ush
+++ b/Engine/Shaders/Private/Strata/StrataEvaluation.ush
@@ -329,7 +329,9 @@ FStrataEvaluateResult StrataEvaluateBSDFCommon(FStrataBSDFContext BSDFContext, F
 					{
 						HazeF0 = 0.04f;
 						HazeF90 = 1.0f;
-						HazeClearCoatTransmittance = SimpleClearCoatTransmittance(max(0.0001, AreaLightContext.NoL), max(0.0001, NoV), StrataGetBSDFMetallic(BSDFContext.BSDF), StrataGetBSDFBaseColor(BSDFContext.BSDF));
+
+						BxDFContext ClearCoatContext = RefractClearCoatContext(BSDFContext.Context);
+						HazeClearCoatTransmittance = SimpleClearCoatTransmittance(ClearCoatContext.NoL, ClearCoatContext.NoV, StrataGetBSDFMetallic(BSDFContext.BSDF), StrataGetBSDFBaseColor(BSDFContext.BSDF));
 					}
 
 					FBxDFEnergyTerms EnergyTerms = ComputeGGXSpecEnergyTerms(HazeSafeRoughness, BSDFContext.Context.NoV, HazeF0, HazeF90);
@@ -451,9 +453,15 @@ FStrataEvaluateResult StrataEvaluateBSDFCommon(FStrataBSDFContext BSDFContext, F
 			}
 
 			const float SpecularMicroOcclusion = F0RGBToMicroOcclusion(F0);
-
+			const float3 FresnelTerm = SpecularMicroOcclusion * F_Schlick(F0, F90, BSDFContext.Context.VoH);
+			
+		#if STRATA_FASTPATH==0
 			const float HazeSpecularMicroOcclusion = F0RGBToMicroOcclusion(HazeF0);
 			const float3 HazeFresnelTerm = HazeSpecularMicroOcclusion * F_Schlick(HazeF0, HazeF90, BSDFContext.Context.VoH);
+		#else
+			const float HazeSpecularMicroOcclusion = SpecularMicroOcclusion;
+			const float3 HazeFresnelTerm = FresnelTerm;
+		#endif
 
 			float3 CommonTerm = 0.0f;
 			if (bIsRectLight)
@@ -465,14 +473,13 @@ FStrataEvaluateResult StrataEvaluateBSDFCommon(FStrataBSDFContext BSDFContext, F
 			}
 			else
 			{
-				const float3 FresnelTerm = SpecularMicroOcclusion * F_Schlick(F0, F90, BSDFContext.Context.VoH);
-
 				Sample.SpecularPathValue     = D     * Vis     * FresnelTerm     * MSScale;
 				Sample.SpecularHazePathValue = HazeD * HazeVis * HazeFresnelTerm * HazeMSScale;
 
 				CommonTerm = ShadowTerms.SurfaceShadow * AreaLightContext.NoL * AreaLightContext.AreaLight.FalloffColor * AreaLightContext.Falloff;
 			}
 
+		#if STRATA_FASTPATH==0
 			// Reminder: 
 			//	- bHazeAsSimpleClearCoat == false means we have two specular lobes that are lerped.
 			//	- bHazeAsSimpleClearCoat == true  means we have a bottom specular lobe (roughness and F0 user specified) and a top specular lob (hard coded absortion, F0=0.04 and user specified roughness.
@@ -496,6 +503,7 @@ FStrataEvaluateResult StrataEvaluateBSDFCommon(FStrataBSDFContext BSDFContext, F
 				Sample.IntegratedSpecularValue	 = CommonTerm * (Sample.SpecularPathValue * BottomLayerFactors + Sample.SpecularHazePathValue * TopLayerCoverage);
 			}
 			else
+		#endif
 			{
 				Sample.SpecularPathProbability     = bHaziness ? (1.0f - HazeWeight) : 1.0f;
 				Sample.SpecularHazePathProbability = bHaziness ?         HazeWeight  : 0.0f;
@@ -1304,7 +1312,8 @@ FStrataEnvLightResult StrataEvaluateForEnvLight(FStrataBSDFContext BSDFContext,
 					BRANCH
 					if (bHazeAsSimpleClearCoat)
 					{
-						const float3 HazeClearCoatTransmittance = SimpleClearCoatTransmittance(max(0.0001, BSDFContext.Context.NoL), max(0.0001, BSDFContext.Context.NoV), StrataGetBSDFMetallic(BSDFContext.BSDF), StrataGetBSDFBaseColor(BSDFContext.BSDF));
+						BxDFContext ClearCoatContext = RefractClearCoatContext(BSDFContext.Context);
+						const float3 HazeClearCoatTransmittance = SimpleClearCoatTransmittance(ClearCoatContext.NoL, ClearCoatContext.NoV, StrataGetBSDFMetallic(BSDFContext.BSDF), StrataGetBSDFBaseColor(BSDFContext.BSDF));
 
 						const float HazeF0 = 0.04f;
 						const float HazeF90= 1.0f;
@@ -1341,10 +1350,6 @@ FStrataEnvLightResult StrataEvaluateForEnvLight(FStrataBSDFContext BSDFContext,
 						StrataEnvLightResult.SpecularHazeWeight = lerp(0.0f, HazeWeight, HazinessFadeIn) * EvalEnvBRDFHaze * SpecularMicroOcclusion;
 						StrataEnvLightResult.SpecularHazeSafeRoughness = HazeSafeRoughness;
 					}
-
-
-
-
 				}
 				#endif
 			}
