Strata - legacy conversion of clear coat now uses relies on SecondRoughnessAsSimpleClearCoat.

diff --git a/Engine/Shaders/Private/Strata/StrataLegacyConversion.ush b/Engine/Shaders/Private/Strata/StrataLegacyConversion.ush
--- a/Engine/Shaders/Private/Strata/StrataLegacyConversion.ush
+++ b/Engine/Shaders/Private/Strata/StrataLegacyConversion.ush
@@ -23,6 +23,9 @@ struct FStrataLegacyParameters
 	float  SSSPhaseAniso;
 	bool   UseSSSDiffusion;
 	float3 Emissive;
+	float  SecondRoughness;
+	float  SecondRoughnessWeight;
+	bool   SecondRoughnessAsSimpleClearCoat;
 	float  FuzzAmount;
 	float3 FuzzColor;
 	float  Thickness;
@@ -56,6 +59,9 @@ FStrataLegacyParameters InitStrataLegacyParameters(FStrataPixelFootprint InPixel
 	Out.SSSPhaseAniso = 0.f;
 	Out.UseSSSDiffusion = false;
 	Out.Emissive = Zeros;
+	Out.SecondRoughness = 0.0f;
+	Out.SecondRoughnessWeight = 0.0f;
+	Out.SecondRoughnessAsSimpleClearCoat = false;
 	Out.FuzzAmount = 0.f;
 	Out.FuzzColor = Zeros;
 	Out.Thickness = STRATA_LAYER_DEFAULT_THICKNESS_CM;
@@ -79,27 +85,27 @@ FStrataData CreateLegacySlab(
 	return PromoteParameterBlendedBSDFToOperator(
 		GetStrataSlabBSDF(
 			In.PixelFootprint,
-			In.UseMetalness,	// UseMetalness
-			In.BaseColor,		// BaseColor		-  Metalness workflow
-			Ones,				// EdgeColor		-  Metalness workflow
-			In.Specular,		// Specular			-  Metalness workflow
-			In.Metallic,		// Metallic			-  Metalness workflow
-			In.DiffuseAlbedo,	// DiffuseAlbedo	- !Metalness workflow
-			In.F0,				// F0				- !Metalness workflow
-			In.F90,				// F90				- !Metalness workflow
-			In.Roughness,		// Roughness
-			In.Anisotropy,		// Anisotropy
-			In.SSSProfileID,	// SSSProfileID
-			In.SSSMFP,			// SSSMFP
-			In.SSSMFPScale,		// SSSMFPScale
-			In.SSSPhaseAniso,	// SSSPhaseAniso
-			In.UseSSSDiffusion, // UseSSSDiffusion
-			In.Emissive,		// Emissive
-			0.f,				// SecondRoughness
-			0.f,				// SecondRoughnessWeight
-			0.f,				// SecondRoughnessAsSimpleClearCoat
-			In.FuzzAmount,		// FuzzAmount
-			In.FuzzColor,		// FuzzColor
+			In.UseMetalness,									// UseMetalness
+			In.BaseColor,										// BaseColor		-  Metalness workflow
+			Ones,												// EdgeColor		-  Metalness workflow
+			In.Specular,										// Specular			-  Metalness workflow
+			In.Metallic,										// Metallic			-  Metalness workflow
+			In.DiffuseAlbedo,									// DiffuseAlbedo	- !Metalness workflow
+			In.F0,												// F0				- !Metalness workflow
+			In.F90,												// F90				- !Metalness workflow
+			In.Roughness,										// Roughness
+			In.Anisotropy,										// Anisotropy
+			In.SSSProfileID,									// SSSProfileID
+			In.SSSMFP,											// SSSMFP
+			In.SSSMFPScale,										// SSSMFPScale
+			In.SSSPhaseAniso,									// SSSPhaseAniso
+			In.UseSSSDiffusion,									// UseSSSDiffusion
+			In.Emissive,										// Emissive
+			In.SecondRoughness,									// SecondRoughness
+			In.SecondRoughnessWeight,							// SecondRoughnessWeight
+			In.SecondRoughnessAsSimpleClearCoat ? 1.0f: 0.0f,	// SecondRoughnessAsSimpleClearCoat
+			In.FuzzAmount,										// FuzzAmount
+			In.FuzzColor,										// FuzzColor
 			In.Thickness,
 			In.SharedLocalBasisIndex, SharedLocalBasisTypes),
 		StrataTree,
@@ -137,7 +143,7 @@ FStrataData StrataConvertLegacyMaterialDynamic(
 	const float3 Zeros = float3(0, 0, 0);
 	const float3 Ones  = float3(1, 1, 1);
 
-	// Can only mix Unlit / DefaultLit / Sub-Surface / Preintegrated-Skin / Subsurface-Profile / ClearCoat / Foliage / Cloth / Eye/ Thin
+	// Can only mix Unlit / DefaultLit / Sub-Surface / Preintegrated-Skin / Subsurface-Profile / ClearCoat / Foliage / Cloth / Eye / Thin
 
 	// Note: If hair or single layer water are enabled with other shading model, the code generation won't be ideal, as the compiler will struggle with unrolling the BSDFs
 	#if MATERIAL_SHADINGMODEL_SINGLELAYERWATER
@@ -328,47 +334,24 @@ FStrataData StrataConvertLegacyMaterialDynamic(
 	}
 	else if (ShadingModel == SHADINGMODELID_CLEAR_COAT)
 	{
-		// === BOTTOM layer === 
-		Bot.BaseColor = BaseColor;
-		Bot.Specular = Specular;
-		Bot.Metallic = Metallic;
-		Bot.Roughness = Roughness;
-		Bot.Anisotropy = Anisotropy;
-		Bot.SSSProfileID = 0.f;
-		Bot.SSSMFP = Zeros;
-		Bot.SSSMFPScale = 0.f;
-		Bot.Emissive = Zeros;
-		Bot.SharedLocalBasisIndex = ClearCoatBottomNormal_SharedLocalBasisIndex;
-		Bot.LayerDepth = 1;
-		Bot.bIsBottom = 1;
-		Bot.Weight = 1.f;
-		// === TOP layer === 
-		 																								   
-		// The top layer has a hard coded specular value of 0.5 (F0 = 0.04)
-		// BaseColor = 0 to only feature absorption, no scattering
-
-		// The original clear coat is a complex assemblage of arbitrary functions that do not always make sense. 
-		// To simplify things, we set the top slab BSDF as having a constant Grey scale transmittance. 
-		// As for the original, this is achieved with coverage so both transmittance and specular contribution vanishes
-
-		// Now setup the mean free path with a hard coded transmittance of 0.75 when viewing the surface perpendicularly
-		const float Thickness = DefaultThickness;
-		const float3 MFP = TransmittanceToMeanFreePath(0.75f, Thickness * CENTIMETER_TO_METER) * METER_TO_CENTIMETER;
-
-		Top.BaseColor = Zeros;
-		Top.Specular = 0.5f;
-		Top.Metallic = 0.f;
-		Top.Roughness = ClearCoatRoughness;
+		Top.BaseColor = BaseColor;
+		Top.Specular = Specular;
+		Top.Metallic = Metallic;
+		Top.Roughness = Roughness;
 		Top.Anisotropy = Anisotropy;
 		Top.SSSProfileID = 0.f;
-		Top.SSSMFP = MFP;
-		Top.SSSMFPScale = 1.f;
-		Top.Emissive = Emissive;
-		Top.Thickness = Thickness;
-		Top.SharedLocalBasisIndex = SharedLocalBasisIndex;
-		Top.LayerDepth = 0;
-		Top.bIsBottom = 0;
-		Top.Weight = ClearCoat;
+		Top.SSSMFP = Zeros;
+		Top.SSSMFPScale = 0.f;
+		Top.SecondRoughnessWeight = ClearCoat;
+		Top.SecondRoughness = ClearCoatRoughness;
+		Top.SecondRoughnessAsSimpleClearCoat = true; // Haziness as simple clear coat layer.
+		Top.Emissive = Zeros;
+		Top.SharedLocalBasisIndex = ClearCoatBottomNormal_SharedLocalBasisIndex;
+		Top.LayerDepth = 1;
+		Top.bIsBottom = 1;
+		Top.Weight = 1.f;
+
+		Bot.Weight = 0.f;
 
 		FinalWeight = 1.0f;
 #if STRATA_TRANSLUCENT_MATERIAL
