Strata - legacy convertion node now only use a single slab.
Legacy material are now 20bytes.

diff --git a/Engine/Shaders/Private/Strata/StrataLegacyConversion.ush b/Engine/Shaders/Private/Strata/StrataLegacyConversion.ush
--- a/Engine/Shaders/Private/Strata/StrataLegacyConversion.ush
+++ b/Engine/Shaders/Private/Strata/StrataLegacyConversion.ush
@@ -31,9 +31,6 @@ struct FStrataLegacyParameters
 	float  Thickness;
 	uint   SharedLocalBasisIndex;
 
-	int LayerDepth;
-	int bIsBottom;
-
 	float Weight;
 };
 
@@ -66,8 +63,6 @@ FStrataLegacyParameters InitStrataLegacyParameters(FStrataPixelFootprint InPixel
 	Out.FuzzColor = Zeros;
 	Out.Thickness = STRATA_LAYER_DEFAULT_THICKNESS_CM;
 	Out.SharedLocalBasisIndex = InSharedLocalBasisIndex;
-	Out.LayerDepth = 0;
-	Out.bIsBottom = 1;
 	Out.Weight = InWeight;
 	Out.PixelFootprint = InPixelFootprint;
 	return Out;
@@ -111,8 +106,9 @@ FStrataData CreateLegacySlab(
 		StrataTree,
 		OperatorIndex,
 		BSDFIndex,
-		In.LayerDepth,
-		In.bIsBottom);
+		0,														//LayerDepth,
+		1														//bIsBottom
+		);
 }
 
 // Convert legacy shading models - Dynamic
@@ -167,10 +163,7 @@ FStrataData StrataConvertLegacyMaterialDynamic(
 
 		FStrataLegacyParameters Bot = InitStrataLegacyParameters(PixelFootprint, SharedLocalBasisIndex, 1.0f);
 
-		FStrataData Bot_Slab = CreateLegacySlab(Bot, SharedLocalBasisTypes, StrataTree,	/*OperatorIndex*/1, /*BSDFIndex*/1);
-		FStrataData WTop_Slab = StrataWeight(Top_SLW , 1.0, StrataTree,					/*OperatorIndex*/2, /*MaxDistanceFromLeaves*/1);
-		FStrataData WBot_Slab = StrataWeight(Bot_Slab, 0.0, StrataTree,					/*OperatorIndex*/3, /*MaxDistanceFromLeaves*/1);
-		FStrataData Out = StrataVerticalLayering(WTop_Slab, WBot_Slab, StrataTree,		/*OperatorIndex*/4, /*MaxDistanceFromLeaves*/2);
+		FStrataData WTop_Slab = StrataWeight(Top_SLW , 1.0, StrataTree,					/*OperatorIndex*/1, /*MaxDistanceFromLeaves*/1);
 		// Out maps to LegacyUpdateBSDFsOperators
 		return Out;
 	}
@@ -209,23 +202,21 @@ FStrataData StrataConvertLegacyMaterialDynamic(
 		// No need to take into account opacity weighting for translucent material, as the coverage is assumed to be 1
 	}
 	#endif
-	// Top/Bottom
-	FStrataLegacyParameters Top = InitStrataLegacyParameters(PixelFootprint, SharedLocalBasisIndex, 1.0f);
-	FStrataLegacyParameters Bot = InitStrataLegacyParameters(PixelFootprint, SharedLocalBasisIndex, 1.0f);
+
+
+	FStrataLegacyParameters LegacySlab = InitStrataLegacyParameters(PixelFootprint, SharedLocalBasisIndex, 1.0f);
 
 	float FinalWeight = 1.0f;
 
 	if (ShadingModel == SHADINGMODELID_UNLIT)
 	{
 		// Unlit is handle with a emissive slab
-		Top.BaseColor = Zeros;
-		Top.Specular = 0.f;
-		Top.Metallic = 0.f;
-		Top.Roughness = 0.f;
-		Top.Anisotropy = Anisotropy;
-		Top.Emissive = Emissive;
-
-		Bot.Weight = 0.f;
+		LegacySlab.BaseColor = Zeros;
+		LegacySlab.Specular = 0.f;
+		LegacySlab.Metallic = 0.f;
+		LegacySlab.Roughness = 0.f;
+		LegacySlab.Anisotropy = Anisotropy;
+		LegacySlab.Emissive = Emissive;
 
 		FinalWeight = 1.0f;
 #if STRATA_TRANSLUCENT_MATERIAL
@@ -234,14 +225,12 @@ FStrataData StrataConvertLegacyMaterialDynamic(
 	}
 	else if (ShadingModel == SHADINGMODELID_DEFAULT_LIT)
 	{
-		Top.BaseColor = BaseColor;
-		Top.Specular = Specular;
-		Top.Metallic = Metallic;
-		Top.Roughness = Roughness;
-		Top.Anisotropy = Anisotropy;
-		Top.Emissive = Emissive;
-
-		Bot.Weight = 0.f;
+		LegacySlab.BaseColor = BaseColor;
+		LegacySlab.Specular = Specular;
+		LegacySlab.Metallic = Metallic;
+		LegacySlab.Roughness = Roughness;
+		LegacySlab.Anisotropy = Anisotropy;
+		LegacySlab.Emissive = Emissive;
 
 		FinalWeight = 1.0f;
 #if STRATA_TRANSLUCENT_MATERIAL
@@ -253,20 +242,17 @@ FStrataData StrataConvertLegacyMaterialDynamic(
 		const float ThicknessInCm = STRATA_SIMPLEVOLUME_THICKNESS_CM;
 		const float3 MFPInCm = TransmittanceToMeanFreePath(SubSurfaceColor, ThicknessInCm * CENTIMETER_TO_METER) * METER_TO_CENTIMETER;
 
-		Top.BaseColor = BaseColor;
-		Top.Specular = Specular;
-		Top.Metallic = Metallic;
-		Top.Roughness = Roughness;
-		Top.Anisotropy = Anisotropy;
-		Top.SSSMFP = max(1e-05f, MFPInCm); // Ensure the MFP is not null to force the material as have SSS
-		Top.SSSMFPScale = 1.0f;
-		Top.SSSPhaseAniso = saturate(1.f - Opacity); // Opaque-Thick: isotropic phase function, Thin: forward phase scattering function
-		Top.Emissive = Emissive;
-		Top.UseSSSDiffusion = false;
-		Top.Thickness = ThicknessInCm;
-
-		// No need to take into account opacity weighting for translucent material, as the coverage is assumed to be 1
-		Bot.Weight = 0.f;
+		LegacySlab.BaseColor = BaseColor;
+		LegacySlab.Specular = Specular;
+		LegacySlab.Metallic = Metallic;
+		LegacySlab.Roughness = Roughness;
+		LegacySlab.Anisotropy = Anisotropy;
+		LegacySlab.SSSMFP = max(1e-05f, MFPInCm); // Ensure the MFP is not null to force the material as have SSS
+		LegacySlab.SSSMFPScale = 1.0f;
+		LegacySlab.SSSPhaseAniso = saturate(1.f - Opacity); // Opaque-Thick: isotropic phase function, Thin: forward phase scattering function
+		LegacySlab.Emissive = Emissive;
+		LegacySlab.UseSSSDiffusion = false;
+		LegacySlab.Thickness = ThicknessInCm;
 
 		FinalWeight = 1.0f;
 
@@ -276,9 +262,9 @@ FStrataData StrataConvertLegacyMaterialDynamic(
 		// Instead for now we use a mfp of 0 and make the material vanish according to opacity.
 		// All this will be fixed when this material will be converted to LUT.
 		const float Transmittance = saturate(1.f - Opacity);
-		Top.SSSMFP = 0;// MFPInCm;
-		Top.SSSMFPScale = 1.0f;
-		Top.Thickness = ThicknessInCm;
+		LegacySlab.SSSMFP = 0;// MFPInCm;
+		LegacySlab.SSSMFPScale = 1.0f;
+		LegacySlab.Thickness = ThicknessInCm;
 
 		FinalWeight = Opacity;
 #endif
@@ -290,21 +276,18 @@ FStrataData StrataConvertLegacyMaterialDynamic(
 
 		// Legacy material uses Subsurface color as transmission 'tinting', but we can represent that with a single layer. So instead we take 
 		// the max color of SubSurfaceColor & BaseColor
-		Top.BaseColor = max(SubSurfaceColor, BaseColor); 
-		Top.Specular = Specular;
-		Top.Metallic = Metallic;
-		Top.Roughness = Roughness;
-		Top.Anisotropy = Anisotropy;
-		Top.SSSMFP = max(1e-05f, MFPInCm); // Ensure the MFP is not null to force the material as have SSS
-		Top.SSSMFPScale = 1.0f;
-		Top.SSSProfileID = 0.f;
-		Top.SSSMFPScale = 1.f - Opacity;
-		Top.SSSPhaseAniso = 0.93f;
-		Top.UseSSSDiffusion = true;
-		Top.Emissive = Emissive;
-
-		// No need to take into account opacity weighting for translucent material, as the coverage is assumed to be 1
-		Bot.Weight = 0.f;
+		LegacySlab.BaseColor = max(SubSurfaceColor, BaseColor); 
+		LegacySlab.Specular = Specular;
+		LegacySlab.Metallic = Metallic;
+		LegacySlab.Roughness = Roughness;
+		LegacySlab.Anisotropy = Anisotropy;
+		LegacySlab.SSSMFP = max(1e-05f, MFPInCm); // Ensure the MFP is not null to force the material as have SSS
+		LegacySlab.SSSMFPScale = 1.0f;
+		LegacySlab.SSSProfileID = 0.f;
+		LegacySlab.SSSMFPScale = 1.f - Opacity;
+		LegacySlab.SSSPhaseAniso = 0.93f;
+		LegacySlab.UseSSSDiffusion = true;
+		LegacySlab.Emissive = Emissive;
 
 		FinalWeight = 1.0f;
 #if STRATA_TRANSLUCENT_MATERIAL
@@ -313,19 +296,16 @@ FStrataData StrataConvertLegacyMaterialDynamic(
 	}
 	else if (ShadingModel == SHADINGMODELID_SUBSURFACE_PROFILE)
 	{
-		Top.BaseColor = BaseColor;
-		Top.Specular = Specular;
-		Top.Metallic = Metallic;
-		Top.Roughness = Roughness;
-		Top.Anisotropy = Anisotropy;
-		Top.SSSProfileID = SubSurfaceProfileId;
-		Top.SSSMFPScale = Opacity;
-		Top.SSSPhaseAniso = 0.93f;
-		Top.UseSSSDiffusion = true;
-		Top.Emissive = Emissive;
-
-		// No need to take into account opacity weighting for translucent material, as the coverage is assumed to be 1
-		Bot.Weight = 0.f;
+		LegacySlab.BaseColor = BaseColor;
+		LegacySlab.Specular = Specular;
+		LegacySlab.Metallic = Metallic;
+		LegacySlab.Roughness = Roughness;
+		LegacySlab.Anisotropy = Anisotropy;
+		LegacySlab.SSSProfileID = SubSurfaceProfileId;
+		LegacySlab.SSSMFPScale = Opacity;
+		LegacySlab.SSSPhaseAniso = 0.93f;
+		LegacySlab.UseSSSDiffusion = true;
+		LegacySlab.Emissive = Emissive;
 
 		FinalWeight = 1.0f;
 #if STRATA_TRANSLUCENT_MATERIAL
@@ -334,24 +314,20 @@ FStrataData StrataConvertLegacyMaterialDynamic(
 	}
 	else if (ShadingModel == SHADINGMODELID_CLEAR_COAT)
 	{
-		Top.BaseColor = BaseColor;
-		Top.Specular = Specular;
-		Top.Metallic = Metallic;
-		Top.Roughness = Roughness;
-		Top.Anisotropy = Anisotropy;
-		Top.SSSProfileID = 0.f;
-		Top.SSSMFP = Zeros;
-		Top.SSSMFPScale = 0.f;
-		Top.SecondRoughnessWeight = ClearCoat;
-		Top.SecondRoughness = ClearCoatRoughness;
-		Top.SecondRoughnessAsSimpleClearCoat = true; // Haziness as simple clear coat layer.
-		Top.Emissive = Zeros;
-		Top.SharedLocalBasisIndex = ClearCoatBottomNormal_SharedLocalBasisIndex;
-		Top.LayerDepth = 1;
-		Top.bIsBottom = 1;
-		Top.Weight = 1.f;
-
-		Bot.Weight = 0.f;
+		LegacySlab.BaseColor = BaseColor;
+		LegacySlab.Specular = Specular;
+		LegacySlab.Metallic = Metallic;
+		LegacySlab.Roughness = Roughness;
+		LegacySlab.Anisotropy = Anisotropy;
+		LegacySlab.SSSProfileID = 0.f;
+		LegacySlab.SSSMFP = Zeros;
+		LegacySlab.SSSMFPScale = 0.f;
+		LegacySlab.SecondRoughnessWeight = ClearCoat;
+		LegacySlab.SecondRoughness = ClearCoatRoughness;
+		LegacySlab.SecondRoughnessAsSimpleClearCoat = true; // Haziness as simple clear coat layer.
+		LegacySlab.Emissive = Zeros;
+		LegacySlab.SharedLocalBasisIndex = ClearCoatBottomNormal_SharedLocalBasisIndex;
+		LegacySlab.Weight = 1.f;
 
 		FinalWeight = 1.0f;
 #if STRATA_TRANSLUCENT_MATERIAL
@@ -364,21 +340,18 @@ FStrataData StrataConvertLegacyMaterialDynamic(
 		const float Thickness = STRATA_LAYER_ISTHIN_THICKNESS_THRESHOLD_CM - 1e-5f;
 		const float3 MFPInCm = TransmittanceToMeanFreePath(SubSurfaceColor /*TransmittanceColor*/, Thickness * CENTIMETER_TO_METER) * METER_TO_CENTIMETER;
 
-		Top.BaseColor = BaseColor;
-		Top.Specular = Specular;
-		Top.Metallic = Metallic;
-		Top.Roughness = Roughness;
-		Top.Anisotropy = Anisotropy;
-		Top.SSSProfileID = 0.f;
-		Top.SSSMFP = max(1e-05f, MFPInCm); // Ensure the MFP is not null to force the material as have SSS
-		Top.SSSMFPScale = 1.0f;
-		Top.SSSPhaseAniso = 1.0f - saturate(Opacity); // Opaque-Thick: isotropic phase function, Thin: forward phase scattering function
-		Top.Emissive = Emissive;
-		Top.UseSSSDiffusion = false;
-		Top.Thickness = Thickness;
-
-		// Take in account the legacy opacity weighting for translucent material
-		Bot.Weight = 0.f;
+		LegacySlab.BaseColor = BaseColor;
+		LegacySlab.Specular = Specular;
+		LegacySlab.Metallic = Metallic;
+		LegacySlab.Roughness = Roughness;
+		LegacySlab.Anisotropy = Anisotropy;
+		LegacySlab.SSSProfileID = 0.f;
+		LegacySlab.SSSMFP = max(1e-05f, MFPInCm); // Ensure the MFP is not null to force the material as have SSS
+		LegacySlab.SSSMFPScale = 1.0f;
+		LegacySlab.SSSPhaseAniso = 1.0f - saturate(Opacity); // Opaque-Thick: isotropic phase function, Thin: forward phase scattering function
+		LegacySlab.Emissive = Emissive;
+		LegacySlab.UseSSSDiffusion = false;
+		LegacySlab.Thickness = Thickness;
 
 		FinalWeight = 1.0f;
 #if STRATA_TRANSLUCENT_MATERIAL
@@ -387,19 +360,17 @@ FStrataData StrataConvertLegacyMaterialDynamic(
 	}
 	else if (ShadingModel == SHADINGMODELID_CLOTH)
 	{
-		Top.BaseColor = BaseColor;
-		Top.Specular = Specular;
-		Top.Metallic = Metallic;
-		Top.Roughness = Roughness;
-		Top.Anisotropy = Anisotropy;
-		Top.SSSProfileID = 0.f;
-		Top.SSSMFP = 0.f;
-		Top.SSSMFPScale = 0.f;
-		Top.Emissive = Emissive;
-		Top.FuzzAmount = ClearCoat;
-		Top.FuzzColor = SubSurfaceColor;
-
-		Bot.Weight = 0.f;
+		LegacySlab.BaseColor = BaseColor;
+		LegacySlab.Specular = Specular;
+		LegacySlab.Metallic = Metallic;
+		LegacySlab.Roughness = Roughness;
+		LegacySlab.Anisotropy = Anisotropy;
+		LegacySlab.SSSProfileID = 0.f;
+		LegacySlab.SSSMFP = 0.f;
+		LegacySlab.SSSMFPScale = 0.f;
+		LegacySlab.Emissive = Emissive;
+		LegacySlab.FuzzAmount = ClearCoat;
+		LegacySlab.FuzzColor = SubSurfaceColor;
 
 		FinalWeight = 1.0f;
 #if STRATA_TRANSLUCENT_MATERIAL
@@ -434,57 +405,18 @@ FStrataData StrataConvertLegacyMaterialDynamic(
 		const float3 BotAlbedo = Zeros;
 		const float3 BotMFP    = MFP;
 
-		Top.DiffuseAlbedo	= lerp(BotAlbedo, TopAlbedo, Opacity);
-		Top.SSSMFP			= lerp(BotMFP,    TopMFP,    Opacity);
-		Top.F0				= lerp(BotF0,     TopF0,     Opacity);
-		Top.F90				= Ones;
-		Top.UseMetalness	= false;
-		Top.Roughness 		= Roughness;
-		Top.Anisotropy 		= 0.f;
-		Top.SSSProfileID 	= 0.f;
-		Top.SSSMFPScale 	= 1.f;
-		Top.Emissive 		= Emissive;
-		Top.Thickness 		= Thickness;
-		Top.SharedLocalBasisIndex = SharedLocalBasisIndex;
-
-		Top.Weight  = 1.0f;
-		Bot.Weight  = 0.0f;
-
-		// Reference two layers approach
-	#if 0
-		// Top slab BSDF as a simple Disney material
-		Top.BaseColor = BaseColor;
-		Top.Specular = Specular;
-		Top.Metallic = Metallic;
-		Top.Roughness = Roughness;
-		Top.Anisotropy = Anisotropy;
-		Top.SSSProfileID = 0.f;
-		Top.SSSMFP = Zeros;
-		Top.SSSMFPScale = 0.f;
-		Top.Emissive = Emissive;
-		Top.SharedLocalBasisIndex = SharedLocalBasisIndex;
-
-		// Bottom slab BSDF will be a simple absorption only layer
-		const float Thickness = DefaultThickness;
-		const float3 MFP = TransmittanceToMeanFreePath(TransmittanceColor, Thickness * CENTIMETER_TO_METER) * METER_TO_CENTIMETER;
-
-		// Assign specular properties shared with the top layer.
-		Bot.BaseColor = Zeros;
-		Bot.Specular = Specular;
-		Bot.Metallic = .0f;
-		Bot.Roughness = Roughness;
-		Bot.Anisotropy = 0.f;
-		Bot.SSSProfileID = .0f;
-		Bot.SSSMFP = MFP;
-		Bot.SSSMFPScale = 1.f;
-		Bot.Emissive = Zeros;
-		Bot.Thickness = Thickness;
-		Bot.SharedLocalBasisIndex = ClearCoatBottomNormal_SharedLocalBasisIndex;
-
-		// No need to take into account opacity weighting for translucent material, as the coverage is assumed to be 1
-		// Now weight the top base material by opacity
-		Top.Weight = Opacity;
-	#endif
+		LegacySlab.DiffuseAlbedo	= lerp(BotAlbedo, TopAlbedo, Opacity);
+		LegacySlab.SSSMFP			= lerp(BotMFP,    TopMFP,    Opacity);
+		LegacySlab.F0				= lerp(BotF0,     TopF0,     Opacity);
+		LegacySlab.F90				= Ones;
+		LegacySlab.UseMetalness	= false;
+		LegacySlab.Roughness 		= Roughness;
+		LegacySlab.Anisotropy 		= 0.f;
+		LegacySlab.SSSProfileID 	= 0.f;
+		LegacySlab.SSSMFPScale 	= 1.f;
+		LegacySlab.Emissive 		= Emissive;
+		LegacySlab.Thickness 		= Thickness;
+		LegacySlab.SharedLocalBasisIndex = SharedLocalBasisIndex;
 
 		// Thin translucency model always have a coverage of 1. It also canont be used in a multiple shading model context.
 		FinalWeight = 1.0f;
@@ -492,16 +424,13 @@ FStrataData StrataConvertLegacyMaterialDynamic(
 
 	// Fixed layer structure for helping compiler to unroll and optimize shader
 	// This must match what is done in PackStrataOut
-	FStrataData Top_Slab = CreateLegacySlab(Top, SharedLocalBasisTypes, StrataTree,	/*OperatorIndex*/0, /*BSDFIndex*/0);
-	FStrataData Bot_Slab = CreateLegacySlab(Bot, SharedLocalBasisTypes, StrataTree,	/*OperatorIndex*/1, /*BSDFIndex*/1);
-	FStrataData WTop_Slab= StrataWeight(Top_Slab, Top.Weight, StrataTree,			/*OperatorIndex*/2, /*MaxDistanceFromLeaves*/1);
-	FStrataData WBot_Slab= StrataWeight(Bot_Slab, Bot.Weight, StrataTree,			/*OperatorIndex*/3, /*MaxDistanceFromLeaves*/1);
-	FStrataData Out		 = StrataVerticalLayering(WTop_Slab, WBot_Slab, StrataTree,	/*OperatorIndex*/4, /*MaxDistanceFromLeaves*/2);
+	FStrataData Slab = CreateLegacySlab(LegacySlab, SharedLocalBasisTypes, StrataTree,	/*OperatorIndex*/0, /*BSDFIndex*/0);
+	FStrataData Out  = StrataWeight(Slab, LegacySlab.Weight, StrataTree,				/*OperatorIndex*/1, /*MaxDistanceFromLeaves*/1);
 
 	// Take in account the legacy opacity weighting for translucent material
 	#if STRATA_LEGACY_MATERIAL_APPLIES_FINAL_WEIGHT
 	{
-		Out = StrataWeight(Out, FinalWeight, StrataTree,	/*OperatorIndex*/5, /*MaxDistanceFromLeaves*/3);
+		Out = StrataWeight(Out, FinalWeight, StrataTree,								/*OperatorIndex*/2, /*MaxDistanceFromLeaves*/2);
 	}
 	#endif
 
@@ -511,13 +440,7 @@ FStrataData StrataConvertLegacyMaterialDynamic(
 
 	// We also evaluate the inlined BSDF resulting from parameter blending in case it is used later.
 	// For instance, this is used for deferred decals. We only evaluate the top material as this is enough for decals.
-	FStrataData ParamBlendVertical = StrataWeightParameterBlending(Top_Slab, Top.Weight);
-//	if(Bot.Weight > 0.0)
-//	{
-//		const float DummyNoV = 1.0f;
-//		uint NewNormalIndex = BSDF_GETSHAREDLOCALBASISID(Top_Slab.InlinedBSDF);
-//		ParamBlendVertical = StrataVerticalLayeringParameterBlending(ParamBlendVertical, StrataWeightParameterBlending(Bot_Slab, Bot.Weight), NewNormalIndex, DummyNoV);
-//	}
+	FStrataData ParamBlendVertical = StrataWeightParameterBlending(Slab, LegacySlab.Weight);
 #if STRATA_LEGACY_MATERIAL_APPLIES_FINAL_WEIGHT
 	Out.InlinedBSDF = StrataWeightParameterBlending(ParamBlendVertical, FinalWeight).InlinedBSDF;
 #else
@@ -590,7 +513,7 @@ FStrataData StrataConvertLegacyMaterialStatic(
 	}
 	else if (ShadingModel == SHADINGMODELID_HAIR)
 	{
-		FStrataData Top_Hair = PromoteParameterBlendedBSDFToOperator(
+		FStrataData HairMat = PromoteParameterBlendedBSDFToOperator(
 			GetStrataHairBSDF(
 				BaseColor,	   // BaseColor
 				Metallic,	   // Scatter
@@ -605,12 +528,7 @@ FStrataData StrataConvertLegacyMaterialStatic(
 			0		/*int LayerDepth*/,
 			1		/*int bIsBottom*/);
 
-		FStrataLegacyParameters Bot = InitStrataLegacyParameters(PixelFootprint, SharedLocalBasisIndex, 1.0f);
-		FStrataData Bot_Slab = CreateLegacySlab(Bot, SharedLocalBasisTypes, StrataTree,	/*OperatorIndex*/1, /*BSDFIndex*/1);
-
-		FStrataData WTop_Slab = StrataWeight(Top_Hair, 1.0, StrataTree,					/*OperatorIndex*/2, /*MaxDistanceFromLeaves*/1);
-		FStrataData WBot_Slab = StrataWeight(Bot_Slab, 0.0, StrataTree,					/*OperatorIndex*/3, /*MaxDistanceFromLeaves*/1);
-		Out = StrataVerticalLayering(WTop_Slab, WBot_Slab, StrataTree,					/*OperatorIndex*/4, /*MaxDistanceFromLeaves*/2);
+		Out = StrataWeight(HairMat, 1.0, StrataTree, /*OperatorIndex*/1, /*MaxDistanceFromLeaves*/1);
 		// Out maps to LegacyUpdateBSDFsOperators
 	}
 	else if (ShadingModel == SHADINGMODELID_EYE)
@@ -619,7 +537,7 @@ FStrataData StrataConvertLegacyMaterialStatic(
 		const float IrisDistance = ClearCoatRoughness;
 		const float3 IrisNormal = RawClearCoatNormal;
 
-		FStrataData Top_Eye = PromoteParameterBlendedBSDFToOperator(
+		FStrataData EyeMat = PromoteParameterBlendedBSDFToOperator(
 			GetStrataEyeBSDF(
 				BaseColor,	   // DiffuseAlbedo
 				Roughness,	   // Roughness
@@ -634,12 +552,7 @@ FStrataData StrataConvertLegacyMaterialStatic(
 			0		/*int LayerDepth*/,
 			1		/*int bIsBottom*/);
 
-		FStrataLegacyParameters Bot = InitStrataLegacyParameters(PixelFootprint, SharedLocalBasisIndex, 1.0f);
-		FStrataData Bot_Slab = CreateLegacySlab(Bot, SharedLocalBasisTypes, StrataTree,	/*OperatorIndex*/1, /*BSDFIndex*/1);
-
-		FStrataData WTop_Slab = StrataWeight(Top_Eye, 1.0, StrataTree,					/*OperatorIndex*/2, /*MaxDistanceFromLeaves*/1);
-		FStrataData WBot_Slab = StrataWeight(Bot_Slab, 0.0, StrataTree,					/*OperatorIndex*/3, /*MaxDistanceFromLeaves*/1);
-		Out = StrataVerticalLayering(WTop_Slab, WBot_Slab, StrataTree,					/*OperatorIndex*/4, /*MaxDistanceFromLeaves*/2);
+		Out = StrataWeight(EyeMat, 1.0, StrataTree, /*OperatorIndex*/1, /*MaxDistanceFromLeaves*/1);
 		// Out maps to LegacyUpdateBSDFsOperators
 	}
 	else if (ShadingModel == SHADINGMODELID_SINGLELAYERWATER)
@@ -647,7 +560,7 @@ FStrataData StrataConvertLegacyMaterialStatic(
 		const float3 WaterExtinction = WaterScatteringCoefficients + WaterAbsorptionCoefficients;
 		const float3 WaterAlbedo = WaterScatteringCoefficients / WaterExtinction;
 
-		FStrataData Top_SLW = PromoteParameterBlendedBSDFToOperator(
+		FStrataData SLW = PromoteParameterBlendedBSDFToOperator(
 			GetStrataSingleLayerWaterBSDF(
 				BaseColor, 				// BaseColor
 				Metallic, 				// Metallic 
@@ -666,12 +579,7 @@ FStrataData StrataConvertLegacyMaterialStatic(
 			0		/*int LayerDepth*/,
 			1		/*int bIsBottom*/);
 
-		FStrataLegacyParameters Bot = InitStrataLegacyParameters(PixelFootprint, SharedLocalBasisIndex, 1.0f);
-		FStrataData Bot_Slab = CreateLegacySlab(Bot, SharedLocalBasisTypes, StrataTree,	/*OperatorIndex*/1, /*BSDFIndex*/1);
-
-		FStrataData WTop_Slab = StrataWeight(Top_SLW, 1.0, StrataTree,					/*OperatorIndex*/2, /*MaxDistanceFromLeaves*/1);
-		FStrataData WBot_Slab = StrataWeight(Bot_Slab, 0.0, StrataTree,					/*OperatorIndex*/3, /*MaxDistanceFromLeaves*/1);
-		Out = StrataVerticalLayering(WTop_Slab, WBot_Slab, StrataTree,					/*OperatorIndex*/4, /*MaxDistanceFromLeaves*/2);
+		Out = StrataWeight(SLW, 1.0, StrataTree, /*OperatorIndex*/1, /*MaxDistanceFromLeaves*/1);
 		// Out maps to LegacyUpdateBSDFsOperators
 	}
 	else
diff --git a/Engine/Shaders/Private/Strata/StrataTree.ush b/Engine/Shaders/Private/Strata/StrataTree.ush
--- a/Engine/Shaders/Private/Strata/StrataTree.ush
+++ b/Engine/Shaders/Private/Strata/StrataTree.ush
@@ -459,27 +459,17 @@ void LegacyUpdateBSDFsOperators(inout FStrataPixelHeader StrataPixelHeader, inou
 {
 	// This must match what is done in StrataConvertLegacyMaterialDynamic
 	UpdateSingleBSDFOperatorCoverageTransmittance(StrataPixelHeader, StrataTree, StrataTree.BSDFs[0], Settings, NullStrataAddressing, V);
-	UpdateSingleBSDFOperatorCoverageTransmittance(StrataPixelHeader, StrataTree, StrataTree.BSDFs[1], Settings, NullStrataAddressing, V);
 
 	// MaxDistanceFromLeaves = 1 
-	UpdateSingleOperatorCoverageTransmittance(StrataTree, /*OperatorIndex*/2);
-	UpdateSingleOperatorCoverageTransmittance(StrataTree, /*OperatorIndex*/3);
-	// MaxDistanceFromLeaves = 2 
-	UpdateSingleOperatorCoverageTransmittance(StrataTree, /*OperatorIndex*/4);
+	UpdateSingleOperatorCoverageTransmittance(StrataTree, /*OperatorIndex*/1);
 #if STRATA_LEGACY_MATERIAL_APPLIES_FINAL_WEIGHT
-	// MaxDistanceFromLeaves = 3 
-	UpdateSingleOperatorCoverageTransmittance(StrataTree, /*OperatorIndex*/5);
+	// MaxDistanceFromLeaves = 2 
+	UpdateSingleOperatorCoverageTransmittance(StrataTree, /*OperatorIndex*/2);
 #endif
 
-	UpdateAllBSDFWithBottomUpOperatorVisit_Weight	(StrataTree, StrataTree.BSDFs[0], /*OperatorIndex*/2, /*PreviousIsInputA*/1);
-	UpdateAllBSDFWithBottomUpOperatorVisit_Vertical	(StrataTree, StrataTree.BSDFs[0], /*OperatorIndex*/4, /*PreviousIsInputA*/1);
-#if STRATA_LEGACY_MATERIAL_APPLIES_FINAL_WEIGHT
-	UpdateAllBSDFWithBottomUpOperatorVisit_Weight	(StrataTree, StrataTree.BSDFs[0], /*OperatorIndex*/5, /*PreviousIsInputA*/1);
-#endif
-	UpdateAllBSDFWithBottomUpOperatorVisit_Weight	(StrataTree, StrataTree.BSDFs[1], /*OperatorIndex*/3, /*PreviousIsInputA*/1);
-	UpdateAllBSDFWithBottomUpOperatorVisit_Vertical	(StrataTree, StrataTree.BSDFs[1], /*OperatorIndex*/4, /*PreviousIsInputA*/0);
+	UpdateAllBSDFWithBottomUpOperatorVisit_Weight	(StrataTree, StrataTree.BSDFs[0], /*OperatorIndex*/1, /*PreviousIsInputA*/1);
 #if STRATA_LEGACY_MATERIAL_APPLIES_FINAL_WEIGHT
-	UpdateAllBSDFWithBottomUpOperatorVisit_Weight	(StrataTree, StrataTree.BSDFs[1], /*OperatorIndex*/5, /*PreviousIsInputA*/1);
+	UpdateAllBSDFWithBottomUpOperatorVisit_Weight	(StrataTree, StrataTree.BSDFs[0], /*OperatorIndex*/2, /*PreviousIsInputA*/1);
 #endif
 }
 
diff --git a/Engine/Source/Runtime/Engine/Private/Materials/MaterialExpressions.cpp b/Engine/Source/Runtime/Engine/Private/Materials/MaterialExpressions.cpp
--- a/Engine/Source/Runtime/Engine/Private/Materials/MaterialExpressions.cpp
+++ b/Engine/Source/Runtime/Engine/Private/Materials/MaterialExpressions.cpp
@@ -21881,28 +21881,20 @@ void UMaterialExpressionStrataLegacyConversion::GatherStrataMaterialInfo(FStrata
 
 FStrataOperator* UMaterialExpressionStrataLegacyConversion::StrataGenerateMaterialTopologyTree(class FMaterialCompiler* Compiler, class UMaterialExpression* Parent, int32 OutputIndex)
 {
-	// The legacy node can only be added once and we know the worst shape it can take so we hard code it here, 
-	// with fake addresses for BSDF defined in shader code.
-	UMaterialExpression* ChildTopExpression = reinterpret_cast<UMaterialExpression*>(1);
-	UMaterialExpression* ChildBottomExpression = reinterpret_cast<UMaterialExpression*>(2);
+	// Legacy conversion now requires a single slab
+	UMaterialExpression* ChildSlabExpression = reinterpret_cast<UMaterialExpression*>(1);
 
 	auto AddDefaultWorstCase = [&](bool bSSS)
 	{
-		FStrataOperator& StrataOperator = Compiler->StrataCompilationRegisterOperator(STRATA_OPERATOR_VERTICAL, this, Parent);
-
-		FStrataOperator& TopOperator = Compiler->StrataCompilationRegisterOperator(STRATA_OPERATOR_BSDF_LEGACY, ChildTopExpression, this);
-		TopOperator.BSDFType = STRATA_BSDF_TYPE_SLAB;
-		TopOperator.bBSDFHasSSS = false;
-		TopOperator.bBSDFHasMFPPluggedIn = true;
-		TopOperator.bBSDFHasFuzz = true;
-		FStrataOperator& BottomOperator = Compiler->StrataCompilationRegisterOperator(STRATA_OPERATOR_BSDF_LEGACY, ChildBottomExpression, this);
-		BottomOperator.BSDFType = STRATA_BSDF_TYPE_SLAB;
-		BottomOperator.bBSDFHasSSS = true;
-		BottomOperator.bBSDFHasMFPPluggedIn = true;
-		TopOperator.bBSDFHasFuzz = false;
-
-		AssignOperatorIndexIfNotNull(StrataOperator.LeftIndex, &TopOperator);
-		AssignOperatorIndexIfNotNull(StrataOperator.RightIndex, &BottomOperator);
+		FStrataOperator& StrataOperator = Compiler->StrataCompilationRegisterOperator(STRATA_OPERATOR_WEIGHT, this, Parent);
+
+		FStrataOperator& SlabOperator = Compiler->StrataCompilationRegisterOperator(STRATA_OPERATOR_BSDF_LEGACY, ChildSlabExpression, this);
+		SlabOperator.BSDFType = STRATA_BSDF_TYPE_SLAB;
+		SlabOperator.bBSDFHasSSS = true;
+		SlabOperator.bBSDFHasMFPPluggedIn = true;
+		SlabOperator.bBSDFHasFuzz = true;
+
+		AssignOperatorIndexIfNotNull(StrataOperator.LeftIndex, &SlabOperator);
 
 		return &StrataOperator;
 	};
