Add comments to avoid the ambiguous conversion between LinearRGB and Lab space in shader.

diff --git a/Engine/Shaders/Private/ColorSpace.ush b/Engine/Shaders/Private/ColorSpace.ush
--- a/Engine/Shaders/Private/ColorSpace.ush
+++ b/Engine/Shaders/Private/ColorSpace.ush
@@ -67,8 +67,10 @@ float3 XYZ_2_LinearRGB( float3 XYZ )
 //------------------------------------------------------------------
 // Assume XYZ is in nominal range to reduce the operation on a constant coefficient 100.
 // E.g., whilte RGB (1, 1, 1) maps to XYZ (0.95047, 1.00, 1.08883)
-// for a reference white of D65. 
-// Use XYZ_WHITE_REF_DXX_NOMINAL instead for LinearRGB_2_LAB() and LAB_2_LinearRGB()
+// for a reference white of D65. LAB is still in the original range.
+// Use XYZ_WHITE_REF_DXX_NOMINAL instead for LinearRGB_2_LAB() and LAB_2_LinearRGB().
+// E.g.,
+// LinearRGB_2_LAB((1.0f, 1.0f, 1.0f), XYZ_WHITE_REF_D65_NOMINAL) = (100.0f, 0.0f, 0.0f)
 
 static const float3 XYZ_WHITE_REF_D65 = float3(95.047, 100.0, 108.883);
 static const float3 XYZ_WHITE_REF_D50 = float3(96.6797, 100.0, 82.5188);
