Strata - translucent can correctly sample scene textures.

diff --git a/Engine/Shaders/Private/MaterialTemplate.ush b/Engine/Shaders/Private/MaterialTemplate.ush
--- a/Engine/Shaders/Private/MaterialTemplate.ush
+++ b/Engine/Shaders/Private/MaterialTemplate.ush
@@ -160,6 +160,11 @@
 #define STRATA_INLINE_SHADING 1
 #endif
 
+#if STRATA_TRANSLUCENT_FORWARD && !STRATA_MATERIAL_EXPORT_FROM_TRANSLUCENT && !defined(STRATA_DEFERRED_SHADING)
+// This is for translucent surface to be able to fetch opaque material data such as normal, albedo and more from SceneTextures (See SceneTextureLookup).
+#define STRATA_DEFERRED_SHADING 1
+#endif
+
 #if STRATA_ENABLED
 #include "/Engine/Private/Strata/Strata.ush"
 #else
@@ -2299,12 +2304,21 @@ float4 SceneTextureLookup(float2 UV, int SceneTextureIndex, bool bFiltered)
 
 	// 2. Post-process inputs which are dependent on material data
 	//    Note: For PostProcess material we always return the first layer dataCalcSceneCustomStencil
-	FStrataAddressing StrataAddressing   = GetStrataPixelDataByteOffset(PixelPos, uint2(View.BufferSizeAndInvSize.xy), Strata.MaxBytesPerPixel);
+#ifdef StrataStruct
+	FStrataAddressing StrataAddressing   = GetStrataPixelDataByteOffset(PixelPos, uint2(View.BufferSizeAndInvSize.xy), StrataStruct.MaxBytesPerPixel);
+	FStrataPixelHeader StrataPixelHeader = UnpackStrataHeaderIn(StrataStruct.MaterialTextureArray, StrataAddressing, StrataStruct.TopLayerTexture);
+	BRANCH
+	if (StrataPixelHeader.BSDFCount > 0)
+	{
+		const FStrataBSDF BSDF = UnpackStrataBSDFIn(StrataStruct.MaterialTextureArray, StrataAddressing, StrataPixelHeader);
+#else
+	FStrataAddressing StrataAddressing = GetStrataPixelDataByteOffset(PixelPos, uint2(View.BufferSizeAndInvSize.xy), Strata.MaxBytesPerPixel);
 	FStrataPixelHeader StrataPixelHeader = UnpackStrataHeaderIn(Strata.MaterialTextureArray, StrataAddressing, Strata.TopLayerTexture);
 	BRANCH
 	if (StrataPixelHeader.BSDFCount > 0)
 	{
 		const FStrataBSDF BSDF = UnpackStrataBSDFIn(Strata.MaterialTextureArray, StrataAddressing, StrataPixelHeader);
+#endif
 		switch (SceneTextureIndex)
 		{
 			// order needs to match to ESceneTextureId
diff --git a/Engine/Shaders/Private/Strata/Strata.ush b/Engine/Shaders/Private/Strata/Strata.ush
--- a/Engine/Shaders/Private/Strata/Strata.ush
+++ b/Engine/Shaders/Private/Strata/Strata.ush
@@ -1754,6 +1754,11 @@ void SetCastContactShadow (inout FStrataPixelHeader Out, bool bIn)							{ Out.S
 void SetDynamicIndirectShadowCasterRepresentation(inout FStrataPixelHeader Out, bool bIn)	{ Out.State |= (bIn ? HEADER_MASK_HASDYNINDIRECTSHADOWCASTER : 0u); }
 void SetHasSubsurface(inout FStrataPixelHeader Out, bool bIn)								{ Out.State |= (bIn ? HEADER_MASK_HASSUBSURFACE : 0u); }
 
+#if STRATA_DEFERRED_SHADING
+// Some compiler complains that FStrataPixelHeader is not completely initialised, Initialise with dummy resource seems to work when unused in code.
+FStrataMaterialContainer DummyStrataBuffer;
+#endif
+
 FStrataPixelHeader InitialiseStrataPixelHeader()
 {
 	FStrataPixelHeader Out;
@@ -1764,6 +1769,7 @@ FStrataPixelHeader InitialiseStrataPixelHeader()
 	Out.IrradianceAO = InitIrradianceAndOcclusion();
 #endif
 #if STRATA_DEFERRED_SHADING
+	Out.StrataBuffer = DummyStrataBuffer;
 	Out.SharedLocalBasesIndexOffset = 0;
 	Out.SharedLocalBasesTypes_PackedHeader = 0;
 	Out.PackedTopLayerData = 0;
diff --git a/Engine/Source/Runtime/Renderer/Private/Strata/Strata.cpp b/Engine/Source/Runtime/Renderer/Private/Strata/Strata.cpp
--- a/Engine/Source/Runtime/Renderer/Private/Strata/Strata.cpp
+++ b/Engine/Source/Runtime/Renderer/Private/Strata/Strata.cpp
@@ -474,13 +474,21 @@ void BindStrataForwardPasslUniformParameters(FRDGBuilder& GraphBuilder, const FV
 	FStrataSceneData* StrataSceneData = View.StrataViewData.SceneData;
 	if (IsStrataEnabled() && StrataSceneData)
 	{
+		OutStrataUniformParameters.MaxBytesPerPixel = StrataSceneData->MaxBytesPerPixel;
 		OutStrataUniformParameters.bRoughDiffuse = StrataSceneData->bRoughDiffuse ? 1u : 0u;
 		OutStrataUniformParameters.PeelLayersAboveDepth = StrataSceneData->PeelLayersAboveDepth;
+		OutStrataUniformParameters.MaterialTextureArray = StrataSceneData->MaterialTextureArray;
+		OutStrataUniformParameters.TopLayerTexture = StrataSceneData->TopLayerTexture;
 	}
 	else
 	{
+		const FRDGSystemTextures& SystemTextures = FRDGSystemTextures::Get(GraphBuilder);
+		FRDGTextureRef DefaultTextureArray = GSystemTextures.GetDefaultTexture(GraphBuilder, ETextureDimension::Texture2DArray, EPixelFormat::PF_R32_UINT, FClearValueBinding::Transparent);
+		OutStrataUniformParameters.MaxBytesPerPixel = 0;
 		OutStrataUniformParameters.bRoughDiffuse = 0;
 		OutStrataUniformParameters.PeelLayersAboveDepth = 0;
+		OutStrataUniformParameters.MaterialTextureArray = DefaultTextureArray;
+		OutStrataUniformParameters.TopLayerTexture = SystemTextures.DefaultNormal8Bit;
 	}
 }
 
diff --git a/Engine/Source/Runtime/Renderer/Private/Strata/Strata.h b/Engine/Source/Runtime/Renderer/Private/Strata/Strata.h
--- a/Engine/Source/Runtime/Renderer/Private/Strata/Strata.h
+++ b/Engine/Source/Runtime/Renderer/Private/Strata/Strata.h
@@ -28,8 +28,11 @@ BEGIN_SHADER_PARAMETER_STRUCT(FStrataBasePassUniformParameters, )
 END_SHADER_PARAMETER_STRUCT()
 
 BEGIN_SHADER_PARAMETER_STRUCT(FStrataForwardPassUniformParameters, )
+	SHADER_PARAMETER(uint32, MaxBytesPerPixel)
 	SHADER_PARAMETER(uint32, bRoughDiffuse)
 	SHADER_PARAMETER(uint32, PeelLayersAboveDepth)
+	SHADER_PARAMETER_RDG_TEXTURE(Texture2DArray<uint>, MaterialTextureArray)
+	SHADER_PARAMETER_RDG_TEXTURE(Texture2D<uint>, TopLayerTexture)
 END_SHADER_PARAMETER_STRUCT()
 
 BEGIN_SHADER_PARAMETER_STRUCT(FStrataTileParameter, )
