Strata - added haz / simple clear coat to forward translucent.

diff --git a/Engine/Shaders/Private/Strata/StrataForwardLighting.ush b/Engine/Shaders/Private/Strata/StrataForwardLighting.ush
--- a/Engine/Shaders/Private/Strata/StrataForwardLighting.ush
+++ b/Engine/Shaders/Private/Strata/StrataForwardLighting.ush
@@ -225,14 +225,20 @@ float3 StrataForwardLighting(
 					// Evaluate specular parameters without bent normal
 					FStrataEnvLightResult StrataEnvLight = StrataEvaluateForEnvLight(StrataBSDFContext, bEnableSpecular, Settings);
 
-					// Evaluate specular lighting
+					// Evaluate specular lighting for the main lobe
 					if (any(StrataEnvLight.SpecularWeight > 0.0f))
 					{
 						int SingleCaptureIndex = GetPrimitiveData(MaterialParameters).SingleCaptureIndex;
-						Color += GetImageBasedReflectionSpecular(MaterialParameters, StrataEnvLight.SpecularDirection, Roughness, IndirectIrradiance, GridIndex, SingleCaptureIndex, EyeIndex) *
+						Color += GetImageBasedReflectionSpecular(MaterialParameters, StrataEnvLight.SpecularDirection, StrataEnvLight.SpecularSafeRoughness, IndirectIrradiance, GridIndex, SingleCaptureIndex, EyeIndex) *
 								StrataEnvLight.SpecularWeight * IndirectOcclusion * LuminanceWeight(StrataBSDFContext, CurrentBSDF) * AOMultiBounce(StrataEnvLight.SpecularColor, ShadingOcclusion.SpecOcclusion);
 					}
-					// We skip hazy lob evaluation in forward.
+					// And for the second lobe if needed
+					if (any(StrataEnvLight.SpecularHazeWeight > 0.0f))
+					{
+						int SingleCaptureIndex = GetPrimitiveData(MaterialParameters).SingleCaptureIndex;
+						Color += GetImageBasedReflectionSpecular(MaterialParameters, StrataEnvLight.SpecularDirection, StrataEnvLight.SpecularHazeSafeRoughness, IndirectIrradiance, GridIndex, SingleCaptureIndex, EyeIndex) *
+							StrataEnvLight.SpecularHazeWeight * IndirectOcclusion * LuminanceWeight(StrataBSDFContext, CurrentBSDF) * AOMultiBounce(StrataEnvLight.SpecularColor, ShadingOcclusion.SpecOcclusion);
+					}
 				}
 				#endif
 
