Strata - second specular lobe has its own direction to not have a rough specular affect a sharp specular due to StrataGetOffSpecularPeakReflectionDir.

diff --git a/Engine/Shaders/Private/Strata/StrataEvaluation.ush b/Engine/Shaders/Private/Strata/StrataEvaluation.ush
--- a/Engine/Shaders/Private/Strata/StrataEvaluation.ush
+++ b/Engine/Shaders/Private/Strata/StrataEvaluation.ush
@@ -1188,7 +1188,7 @@ struct FStrataEnvLightResult
 	float3 SpecularWeight;
 	float  SpecularSafeRoughness;
 
-	// Specular haze maintains the same off peak specular direction as the non hazy specular lobe
+	float3 SpecularHazeDirection;	// We must use a different direction to not have a rough specular affect a sharp specular due to StrataGetOffSpecularPeakReflectionDir.
 	float3 SpecularHazeWeight;
 	float  SpecularHazeSafeRoughness;
 	float  SSRReduction;
@@ -1308,7 +1308,9 @@ FStrataEnvLightResult StrataEvaluateForEnvLight(FStrataBSDFContext BSDFContext,
 					const float HazeWeight = Haziness.Weight;
 					const float HazeSafeRoughness = MakeRoughnessSafe(Haziness.Roughness);
 					const bool  bHazeAsSimpleClearCoat = Haziness.bSimpleClearCoat;
-					
+
+					StrataEnvLightResult.SpecularHazeDirection = StrataGetOffSpecularPeakReflectionDir(BSDFContext.N, BSDFContext.R, HazeSafeRoughness);
+
 					BRANCH
 					if (bHazeAsSimpleClearCoat)
 					{
diff --git a/Engine/Shaders/Private/Strata/StrataForwardLighting.ush b/Engine/Shaders/Private/Strata/StrataForwardLighting.ush
--- a/Engine/Shaders/Private/Strata/StrataForwardLighting.ush
+++ b/Engine/Shaders/Private/Strata/StrataForwardLighting.ush
@@ -236,7 +236,7 @@ float3 StrataForwardLighting(
 					if (any(StrataEnvLight.SpecularHazeWeight > 0.0f))
 					{
 						int SingleCaptureIndex = GetPrimitiveData(MaterialParameters).SingleCaptureIndex;
-						Color += GetImageBasedReflectionSpecular(MaterialParameters, StrataEnvLight.SpecularDirection, StrataEnvLight.SpecularHazeSafeRoughness, IndirectIrradiance, GridIndex, SingleCaptureIndex, EyeIndex) *
+						Color += GetImageBasedReflectionSpecular(MaterialParameters, StrataEnvLight.SpecularHazeDirection, StrataEnvLight.SpecularHazeSafeRoughness, IndirectIrradiance, GridIndex, SingleCaptureIndex, EyeIndex) *
 							StrataEnvLight.SpecularHazeWeight * IndirectOcclusion * LuminanceWeight(StrataBSDFContext, CurrentBSDF) * AOMultiBounce(StrataEnvLight.SpecularColor, ShadingOcclusion.SpecOcclusion);
 					}
 				}
diff --git a/Engine/Shaders/Private/Strata/StrataLightingCommon.ush b/Engine/Shaders/Private/Strata/StrataLightingCommon.ush
--- a/Engine/Shaders/Private/Strata/StrataLightingCommon.ush
+++ b/Engine/Shaders/Private/Strata/StrataLightingCommon.ush
@@ -237,7 +237,7 @@ void StrataEnvLightingCommon(
 				CompositeReflectionCapturesAndSkylightTWS(
 					(bIsTopLayer ? TopLayerSpecularContributionFactor : 1.0f) * SpecularOcclusion,
 					TranslatedWorldPosition,
-					StrataEnvLight.SpecularDirection,
+					StrataEnvLight.SpecularHazeDirection,
 					StrataEnvLight.SpecularHazeSafeRoughness,
 					IndirectIrradiance,
 					IndirectSpecularOcclusion,
