Strata - fixed lumen compositing nans

diff --git a/Engine/Shaders/Private/DiffuseIndirectComposite.usf b/Engine/Shaders/Private/DiffuseIndirectComposite.usf
--- a/Engine/Shaders/Private/DiffuseIndirectComposite.usf
+++ b/Engine/Shaders/Private/DiffuseIndirectComposite.usf
@@ -423,23 +423,25 @@ void MainPS(
 			{			
 				FStrataAddressing StrataAddressing = GetStrataPixelDataByteOffset(PixelPos, uint2(View.BufferSizeAndInvSize.xy), Strata.MaxBytesPerPixel);
 				const FStrataPixelHeader StrataPixelHeader = UnpackStrataHeaderIn(Strata.MaterialTextureArray, StrataAddressing, Strata.TopLayerTexture);
-
-				const FStrataDeferredLighting IndirectLighting_Strata = StrataIndirectLighting(
-					PixelPos,
-					Strata.MaterialTextureArray,
-					StrataAddressing,
-					StrataPixelHeader,
-					V,
-					FinalAmbientOcclusion,
-					DiffuseIndirectLighting,
-					SpecularIndirectLighting);
-
-				OutAddColor = IndirectLighting_Strata.SceneColor;
-			  #if STRATA_OPAQUE_ROUGH_REFRACTION_ENABLED
-				const uint2 OutCoord = PixelPos;
-				OutOpaqueRoughRefractionSceneColor[OutCoord] = IndirectLighting_Strata.OpaqueRoughRefractionSceneColor;
-				OutSubSurfaceSceneColor[OutCoord] = IndirectLighting_Strata.SubSurfaceSceneColor;
-			  #endif
+				if (GetMaterialMode(StrataPixelHeader) > HEADER_MATERIALMODE_NONE)
+				{
+					const FStrataDeferredLighting IndirectLighting_Strata = StrataIndirectLighting(
+						PixelPos,
+						Strata.MaterialTextureArray,
+						StrataAddressing,
+						StrataPixelHeader,
+						V,
+						FinalAmbientOcclusion,
+						DiffuseIndirectLighting,
+						SpecularIndirectLighting);
+
+					OutAddColor = IndirectLighting_Strata.SceneColor;
+				#if STRATA_OPAQUE_ROUGH_REFRACTION_ENABLED
+					const uint2 OutCoord = PixelPos;
+					OutOpaqueRoughRefractionSceneColor[OutCoord] = IndirectLighting_Strata.OpaqueRoughRefractionSceneColor;
+					OutSubSurfaceSceneColor[OutCoord] = IndirectLighting_Strata.SubSurfaceSceneColor;
+				#endif
+				}
 			}
 			#else // STRATA_ENABLED
 			{
