Fix virtual textures not updating correctly in debug view modes.
Early depth stencil wasn't being set, meaning that the VT feedback buffer could be overwritted from occluded objects.

diff --git a/Engine/Shaders/Private/DebugViewModePixelShader.usf b/Engine/Shaders/Private/DebugViewModePixelShader.usf
--- a/Engine/Shaders/Private/DebugViewModePixelShader.usf
+++ b/Engine/Shaders/Private/DebugViewModePixelShader.usf
@@ -538,6 +538,11 @@ float4 VisualizeGPUSkinCache(in FDebugPSIn DebugInputs, in bool bIsFrontFace)
 	return ApplyDebugColor(DebugInputs, bIsFrontFace, SkinCacheDebugColor);
 }
 
+#if !OUTPUT_PIXEL_DEPTH_OFFSET
+#define PIXELSHADER_EARLYDEPTHSTENCIL EARLYDEPTHSTENCIL
+#endif
+
+PIXELSHADER_EARLYDEPTHSTENCIL
 void Main(
 	in FDebugPSIn DebugInputs
 	OPTIONAL_IsFrontFace,
@@ -590,9 +595,3 @@ void Main(
 
 	OutColor = RETURN_COLOR(VisualizeColor);	
 }
-
-
-
-
-
-
