Fix "Primary Rays" ray tracing debug view mode

* Rename ERayTracingPrimaryRaysFlag::ConsiderSurfaceScatter to PrimaryView and use it to initialize first ray instance inclusion mask in RayTracingPrimaryRays RGS

diff --git a/Engine/Shaders/Private/RayTracing/RayTracingPrimaryRays.usf b/Engine/Shaders/Private/RayTracing/RayTracingPrimaryRays.usf
--- a/Engine/Shaders/Private/RayTracing/RayTracingPrimaryRays.usf
+++ b/Engine/Shaders/Private/RayTracing/RayTracingPrimaryRays.usf
@@ -31,8 +31,8 @@
 
 #define	ERayTracingPrimaryRaysFlag_None                               0
 #define	ERayTracingPrimaryRaysFlag_UseGBufferForMaxDistance    (1u << 0)
-#define	ERayTracingPrimaryRaysFlag_ConsiderSurfaceScatter	   (1u << 1)
-#define	ERayTracingPrimaryRaysFlag_AllowSkipSkySample		   (1u << 2)
+#define	ERayTracingPrimaryRaysFlag_PrimaryView                 (1u << 1)
+#define	ERayTracingPrimaryRaysFlag_AllowSkipSkySample          (1u << 2)
 
 
 int SamplesPerPixel;
@@ -128,7 +128,8 @@ RAY_TRACING_ENTRY_RAYGEN(RayTracingPrimaryRaysRGS)
 	// Check if the Sky Light should affect reflection rays within translucency.
 	const bool bSkyLightAffectReflection = ShouldSkyLightAffectReflection();
 	
-	bool bHasScattered = (ERayTracingPrimaryRaysFlag_ConsiderSurfaceScatter & PrimaryRayFlags) != 0;
+	const bool bPrimaryView = (ERayTracingPrimaryRaysFlag_PrimaryView & PrimaryRayFlags) != 0;
+	bool bHasScattered = bPrimaryView;
 	float AccumulatedOpacity = 0.0;
 
 	// Integrated data by path tracing
@@ -140,7 +141,7 @@ RAY_TRACING_ENTRY_RAYGEN(RayTracingPrimaryRaysRGS)
 	for (uint RefractionRayIndex = 0; RefractionRayIndex < MaxRefractionRays; ++RefractionRayIndex)
 	{
 		const uint RefractionRayFlags = RefractionRayIndex == 0? RAY_FLAG_CULL_BACK_FACING_TRIANGLES : 0;
-		const uint RefractionInstanceInclusionMask = RefractionRayIndex == 0 || !TranslucencyRefraction ? RAY_TRACING_MASK_TRANSLUCENT : RAY_TRACING_MASK_ALL;
+		const uint RefractionInstanceInclusionMask = (!bPrimaryView && RefractionRayIndex == 0) || !TranslucencyRefraction ? RAY_TRACING_MASK_TRANSLUCENT : RAY_TRACING_MASK_ALL;
 		const bool bRefractionRayTraceSkyLightContribution = false;
 		const bool bRefractionDecoupleSampleGeneration = true;
 		const bool bRefractionEnableSkyLightContribution = true;
diff --git a/Engine/Source/Runtime/Renderer/Private/RayTracing/RayTracingDebug.cpp b/Engine/Source/Runtime/Renderer/Private/RayTracing/RayTracingDebug.cpp
--- a/Engine/Source/Runtime/Renderer/Private/RayTracing/RayTracingDebug.cpp
+++ b/Engine/Source/Runtime/Renderer/Private/RayTracing/RayTracingDebug.cpp
@@ -359,7 +359,7 @@ void FDeferredShadingSceneRenderer::RenderRayTracingDebug(FRDGBuilder& GraphBuil
 
 		RenderRayTracingPrimaryRaysView(
 				GraphBuilder, View, &OutputColor, &HitDistanceTexture, 1, 1, 1,
-			ERayTracingPrimaryRaysFlag::ConsiderSurfaceScatter);
+			ERayTracingPrimaryRaysFlag::PrimaryView);
 
 		AddDrawTexturePass(GraphBuilder, View, OutputColor, SceneColorTexture, View.ViewRect.Min, View.ViewRect.Min, View.ViewRect.Size());
 		return;
diff --git a/Engine/Source/Runtime/Renderer/Private/RayTracing/RaytracingOptions.h b/Engine/Source/Runtime/Renderer/Private/RayTracing/RaytracingOptions.h
--- a/Engine/Source/Runtime/Renderer/Private/RayTracing/RaytracingOptions.h
+++ b/Engine/Source/Runtime/Renderer/Private/RayTracing/RaytracingOptions.h
@@ -18,8 +18,8 @@ enum class ERayTracingPrimaryRaysFlag: uint32
 {
 	None                      =      0,
 	UseGBufferForMaxDistance  = 1 << 0,
-	ConsiderSurfaceScatter	  = 1 << 1,
-	AllowSkipSkySample		  = 1 << 2,
+	PrimaryView	              = 1 << 1,
+	AllowSkipSkySample        = 1 << 2,
 };
 
 ENUM_CLASS_FLAGS(ERayTracingPrimaryRaysFlag);
