Fix debug shader compile error for PDO materials introduced in 20702869
#rnx

diff --git a/Engine/Shaders/Private/DebugViewModePixelShader.usf b/Engine/Shaders/Private/DebugViewModePixelShader.usf
--- a/Engine/Shaders/Private/DebugViewModePixelShader.usf
+++ b/Engine/Shaders/Private/DebugViewModePixelShader.usf
@@ -540,6 +540,8 @@ float4 VisualizeGPUSkinCache(in FDebugPSIn DebugInputs, in bool bIsFrontFace)
 
 #if !OUTPUT_PIXEL_DEPTH_OFFSET
 #define PIXELSHADER_EARLYDEPTHSTENCIL EARLYDEPTHSTENCIL
+#else
+#define PIXELSHADER_EARLYDEPTHSTENCIL
 #endif
 
 PIXELSHADER_EARLYDEPTHSTENCIL
