Only use TextureCubeArraySampleLevel in Texture Editor when SM5 feature level is supported.
#rnx

diff --git a/Engine/Shaders/Private/SimpleElementPixelShader.usf b/Engine/Shaders/Private/SimpleElementPixelShader.usf
--- a/Engine/Shaders/Private/SimpleElementPixelShader.usf
+++ b/Engine/Shaders/Private/SimpleElementPixelShader.usf
@@ -388,7 +388,9 @@ void CubemapTextureProperties(
 		CurrentSliceIndex = floor(TextureCoordinate.y * NumSlices);
 	}
 
+#if FEATURE_LEVEL >= FEATURE_LEVEL_SM5
 	BaseColor = TextureCubeArraySampleLevel(CubeTexture, CubeTextureSampler, Direction.xzy, CurrentSliceIndex, MipLevel);
+#endif
 
 #else //TEXTURECUBE_ARRAY
 
