Fix shader compilation issue for PCD3D_ES31.
#rnx

diff --git a/Engine/Shaders/Private/SimpleElementPixelShader.usf b/Engine/Shaders/Private/SimpleElementPixelShader.usf
--- a/Engine/Shaders/Private/SimpleElementPixelShader.usf
+++ b/Engine/Shaders/Private/SimpleElementPixelShader.usf
@@ -390,6 +390,8 @@ void CubemapTextureProperties(
 
 #if FEATURE_LEVEL >= FEATURE_LEVEL_SM5
 	BaseColor = TextureCubeArraySampleLevel(CubeTexture, CubeTextureSampler, Direction.xzy, CurrentSliceIndex, MipLevel);
+#else
+	BaseColor = 0;
 #endif
 
 #else //TEXTURECUBE_ARRAY
