1.Supports more shading models on mobile, including twoside foliage, cloth and eye shading models.

2.Uses the same shading models and dynamic lightings code path on desktop.

3.Remove the division of PI on mobile directional light's color in C++ to line up with desktop in shaders.

diff --git a/Engine/Shaders/Private/BRDF.ush b/Engine/Shaders/Private/BRDF.ush
--- a/Engine/Shaders/Private/BRDF.ush
+++ b/Engine/Shaders/Private/BRDF.ush
@@ -8,17 +8,17 @@
 
 struct BxDFContext
 {
-	float NoV;
-	float NoL;
-	float VoL;
-	float NoH;
-	float VoH;
-	float XoV;
-	float XoL;
-	float XoH;
-	float YoV;
-	float YoL;
-	float YoH;
+	half NoV;
+	half NoL;
+	half VoL;
+	half NoH;
+	half VoH;
+	half XoV;
+	half XoL;
+	half XoH;
+	half YoV;
+	half YoL;
+	half YoH;
 #if 1	// Shit code until we get member functions
 };
 
@@ -60,6 +60,26 @@ void Init( inout BxDFContext Context, half3 N, half3 X, half3 Y, half3 V, half3
 	Context.YoH = (Context.YoL + Context.YoV) * InvLenH;
 }
 
+void InitMobile(inout BxDFContext Context, half3 N, half3 V, half3 L, half NoL)
+{
+	Context.NoL = NoL;
+	Context.NoV = dot(N, V);
+	Context.VoL = dot(V, L);
+	float3 H = normalize(float3(V + L));
+	Context.NoH = max(0, dot(N, H));
+	Context.VoH = max(0, dot(V, H));
+
+	//NoL = saturate( NoL );
+	//NoV = saturate( abs( NoV ) + 1e-5 );
+
+	Context.XoV = 0.0f;
+	Context.XoL = 0.0f;
+	Context.XoH = 0.0f;
+	Context.YoV = 0.0f;
+	Context.YoL = 0.0f;
+	Context.YoH = 0.0f;
+}
+
 // [ de Carpentier 2017, "Decima Engine: Advances in Lighting and AA" ]
 void SphereMaxNoH( inout BxDFContext Context, float SinAlpha, bool bNewtonIteration )
 {
@@ -255,7 +275,7 @@ void SphereMaxNoH( inout BxDFContext Context, float SinAlpha, bool bNewtonIterat
 // f = Microfacet specular BRDF = D*G*F / (4*NoL*NoV) = D*Vis*F
 
 
-float3 Diffuse_Lambert( float3 DiffuseColor )
+half3 Diffuse_Lambert( half3 DiffuseColor )
 {
 	return DiffuseColor * (1 / PI);
 }
diff --git a/Engine/Shaders/Private/BasePassCommon.ush b/Engine/Shaders/Private/BasePassCommon.ush
--- a/Engine/Shaders/Private/BasePassCommon.ush
+++ b/Engine/Shaders/Private/BasePassCommon.ush
@@ -4,6 +4,7 @@
 	BasePassCommon.ush: Base pass definitions used by both vertex and pixel shader
 =============================================================================*/
 
+#if SHADING_PATH_DEFERRED
 #if MATERIALBLENDING_ANY_TRANSLUCENT
 	#define ForwardLightData TranslucentBasePass.Shared.Forward
 	#define ReflectionStruct TranslucentBasePass.Shared.Reflection
@@ -139,3 +140,152 @@ float3 GetDirectionalLightingVectorFromTranslucentLightingVolume(float3 InnerVol
 	// Lerp between cascades
 	return lerp(OuterVector1, InnerVector1, FinalLerpFactor);
 }
+#endif
+
+#include "SphericalGaussian.ush"
+
+float NormalCurvatureToRoughness(float3 WorldNormal)
+{
+	float3 dNdx = ddx(WorldNormal);
+	float3 dNdy = ddy(WorldNormal);
+	float x = dot(dNdx, dNdx);
+	float y = dot(dNdy, dNdy);
+	float CurvatureApprox = pow(max(x, y), View.NormalCurvatureToRoughnessScaleBias.z);
+	return saturate(CurvatureApprox * View.NormalCurvatureToRoughnessScaleBias.x + View.NormalCurvatureToRoughnessScaleBias.y);
+}
+
+struct FShadingOcclusion
+{
+	half  DiffOcclusion;
+	half  SpecOcclusion;
+	half3 BentNormal;
+};
+
+float DotSpecularSG( float Roughness, float3 N, float3 V, FSphericalGaussian LightSG )
+{
+	float a = Pow2( max( 0.02, Roughness ) );
+	float a2 = a*a;
+	
+	float3 L = LightSG.Axis;
+	float3 H = normalize(V + L);
+
+	float NoV = saturate( abs( dot(N, V) ) + 1e-5 );
+
+	FSphericalGaussian NDF;
+	NDF.Axis = N;
+	NDF.Sharpness = 2 / a2;
+	NDF.Amplitude = rcp( PI * a2 );
+
+#if 0
+	{
+		// Reflect NDF
+		//float3 R = 2 * dot( V, N ) * N - V;
+		float3 R = 2 * NoV * N - V;
+
+		// Point lobe in off-specular peak direction
+		//R = lerp( N, R, (1 - a) * ( sqrt(1 - a) + a ) );
+		//R = normalize( R );
+
+	#if 0
+		// Warp
+		FSphericalGaussian SpecularSG;
+		SpecularSG.Axis = R;
+		SpecularSG.Sharpness = 0.5 / ( a2 * max( NoV, 0.1 ) );
+		SpecularSG.Amplitude = rcp( PI * a2 );
+	#else
+		FAnisoSphericalGaussian SpecularSG;
+		SpecularSG.AxisZ = R;
+		SpecularSG.AxisX = normalize( cross( N, SpecularSG.AxisZ ) );
+		SpecularSG.AxisY = normalize( cross( R, SpecularSG.AxisX ) );
+ 
+		// Second derivative of the sharpness with respect to how
+		// far we are from basis Axis direction
+		SpecularSG.SharpnessX = 0.25 / ( a2 * Pow2( max( NoV, 0.001 ) ) );
+		SpecularSG.SharpnessY = 0.25 / a2;
+		SpecularSG.Amplitude = rcp( PI * a2 );
+	#endif
+		return Dot( SpecularSG, LightSG );
+	}
+#elif 0
+	{
+		// Project LightSG into half vector space
+	#if 0
+		FSphericalGaussian WarpedLightSG;
+		WarpedLightSG.Axis		= H;
+		WarpedLightSG.Sharpness	= LightSG.Sharpness * 1.5 * NoV;
+		WarpedLightSG.Amplitude	= LightSG.Amplitude;
+	#else
+		FAnisoSphericalGaussian WarpedLightSG;
+		WarpedLightSG.AxisZ = H;
+		WarpedLightSG.AxisX = normalize( cross( N, WarpedLightSG.AxisZ ) );
+		WarpedLightSG.AxisY = normalize( cross( H, WarpedLightSG.AxisX ) );
+ 
+		// Second derivative of the sharpness with respect to how
+		// far we are from basis Axis direction
+		WarpedLightSG.SharpnessX = LightSG.Sharpness * 2 * Pow2( NoV );
+		WarpedLightSG.SharpnessY = LightSG.Sharpness * 2;
+		WarpedLightSG.Amplitude	= LightSG.Amplitude;
+	#endif
+
+		return Dot( WarpedLightSG, NDF );
+	}
+#else
+	{
+		// We can do the half space ASG method cheaper by assuming H is in the YZ plane.
+		float SharpnessX = LightSG.Sharpness * 2 * Pow2( NoV );
+		float SharpnessY = LightSG.Sharpness * 2;
+
+		float nu = NDF.Sharpness * 0.5;
+
+		FSphericalGaussian ConvolvedNDF;
+		ConvolvedNDF.Axis = NDF.Axis;
+		ConvolvedNDF.Sharpness = 2 * (nu * SharpnessY) / (nu + SharpnessY);
+		ConvolvedNDF.Amplitude = NDF.Amplitude * LightSG.Amplitude;
+		ConvolvedNDF.Amplitude *= PI * rsqrt( (nu + SharpnessX) * (nu + SharpnessY) );
+
+		//float3 AxisX = normalize( cross( N, V ) );
+		//ConvolvedNDF.Amplitude *= exp( -(nu * SharpnessX) / (nu + SharpnessX) * Pow2( dot( H, AxisX ) ) );
+
+		return Evaluate( ConvolvedNDF, H );
+	}
+#endif
+}
+
+FShadingOcclusion ApplyBentNormal(
+	in half3 CameraVector, 
+	in half3 WorldNormal, 
+	in float3 WorldBentNormal0, 
+	in half Roughness,
+	in half MaterialAO)
+{
+	FShadingOcclusion Out;
+	Out.DiffOcclusion = MaterialAO;
+	Out.SpecOcclusion = MaterialAO;
+	Out.BentNormal    = WorldNormal;
+
+#if NUM_MATERIAL_OUTPUTS_GETBENTNORMAL > 0
+	Out.BentNormal = WorldBentNormal0;
+
+	FSphericalGaussian HemisphereSG = Hemisphere_ToSphericalGaussian(WorldNormal);
+	FSphericalGaussian NormalSG = ClampedCosine_ToSphericalGaussian(WorldNormal);
+	FSphericalGaussian VisibleSG = BentNormalAO_ToSphericalGaussian(Out.BentNormal, Out.DiffOcclusion );
+	FSphericalGaussian DiffuseSG = Mul( NormalSG, VisibleSG );
+	
+	float VisibleCosAngle = sqrt( 1 - Out.DiffOcclusion );
+
+#if 1	// Mix full resolution normal with low res bent normal
+	Out.BentNormal = DiffuseSG.Axis;
+	//DiffOcclusion = saturate( Integral( DiffuseSG ) / Dot( NormalSG, HemisphereSG ) );
+	Out.DiffOcclusion = saturate( Integral( DiffuseSG ) * 0.42276995 );
+#endif
+
+	half3 N = WorldNormal;
+	half3 V = CameraVector;
+
+	Out.SpecOcclusion  = DotSpecularSG( Roughness, N, V, VisibleSG );
+	Out.SpecOcclusion /= DotSpecularSG( Roughness, N, V, HemisphereSG );
+
+	Out.SpecOcclusion = saturate(Out.SpecOcclusion );
+#endif
+	return Out;
+}
\ No newline at end of file
diff --git a/Engine/Shaders/Private/BasePassPixelShader.usf b/Engine/Shaders/Private/BasePassPixelShader.usf
--- a/Engine/Shaders/Private/BasePassPixelShader.usf
+++ b/Engine/Shaders/Private/BasePassPixelShader.usf
@@ -56,7 +56,6 @@
 #include "LightAccumulator.ush"
 #include "DeferredShadingCommon.ush"
 #include "VelocityCommon.ush"
-#include "SphericalGaussian.ush"
 #include "DBufferDecalShared.ush"
 #include "ShadingModelsSampling.ush"
 
@@ -94,23 +93,6 @@
 #include "Lumen/LumenTranslucencyVolumeShared.ush"
 #endif
 
-float NormalCurvatureToRoughness(float3 WorldNormal)
-{
-    float3 dNdx = ddx(WorldNormal);
-    float3 dNdy = ddy(WorldNormal);
-    float x = dot(dNdx, dNdx);
-    float y = dot(dNdy, dNdy);
-    float CurvatureApprox = pow(max(x, y), View.NormalCurvatureToRoughnessScaleBias.z);
-	return saturate(CurvatureApprox * View.NormalCurvatureToRoughnessScaleBias.x + View.NormalCurvatureToRoughnessScaleBias.y);
-}
-
-struct FShadingOcclusion
-{
-	float  DiffOcclusion;
-	float  SpecOcclusion;
-	float3 BentNormal;
-};
-
 #if TRANSLUCENT_SELF_SHADOWING
 	#include "ShadowProjectionCommon.ush"	 
 #endif
@@ -632,7 +614,7 @@ void GetPrecomputedIndirectLightingAndSkyLight(
 		LightmapUVType LightmapUV0, LightmapUV1;
 		uint LightmapDataIndex;
 		GetLightMapCoordinates(Interpolants, LightmapUV0, LightmapUV1, LightmapDataIndex);
-		OutDiffuseLighting = GetLightMapColorLQ(LightmapVTPageTableResult, LightmapUV0, LightmapUV1, LightmapDataIndex, DiffuseDir).rgb;
+		GetLightMapColorLQ(LightmapVTPageTableResult, LightmapUV0, LightmapUV1, LightmapDataIndex, DiffuseDir, bEvaluateBackface, OutDiffuseLighting, OutSubsurfaceLighting, OutIndirectIrradiance);
 
 	#endif
 
@@ -707,135 +689,6 @@ void ApplyPixelDepthOffsetForBasePass(inout FMaterialPixelParameters MaterialPar
 #endif
 }
 
-float DotSpecularSG( float Roughness, float3 N, float3 V, FSphericalGaussian LightSG )
-{
-	float a = Pow2( max( 0.02, Roughness ) );
-	float a2 = a*a;
-	
-	float3 L = LightSG.Axis;
-	float3 H = normalize(V + L);
-
-	float NoV = saturate( abs( dot(N, V) ) + 1e-5 );
-
-	FSphericalGaussian NDF;
-	NDF.Axis = N;
-	NDF.Sharpness = 2 / a2;
-	NDF.Amplitude = rcp( PI * a2 );
-
-#if 0
-	{
-		// Reflect NDF
-		//float3 R = 2 * dot( V, N ) * N - V;
-		float3 R = 2 * NoV * N - V;
-
-		// Point lobe in off-specular peak direction
-		//R = lerp( N, R, (1 - a) * ( sqrt(1 - a) + a ) );
-		//R = normalize( R );
-
-	#if 0
-		// Warp
-		FSphericalGaussian SpecularSG;
-		SpecularSG.Axis = R;
-		SpecularSG.Sharpness = 0.5 / ( a2 * max( NoV, 0.1 ) );
-		SpecularSG.Amplitude = rcp( PI * a2 );
-	#else
-		FAnisoSphericalGaussian SpecularSG;
-		SpecularSG.AxisZ = R;
-		SpecularSG.AxisX = normalize( cross( N, SpecularSG.AxisZ ) );
-		SpecularSG.AxisY = normalize( cross( R, SpecularSG.AxisX ) );
- 
-		// Second derivative of the sharpness with respect to how
-		// far we are from basis Axis direction
-		SpecularSG.SharpnessX = 0.25 / ( a2 * Pow2( max( NoV, 0.001 ) ) );
-		SpecularSG.SharpnessY = 0.25 / a2;
-		SpecularSG.Amplitude = rcp( PI * a2 );
-	#endif
-		return Dot( SpecularSG, LightSG );
-	}
-#elif 0
-	{
-		// Project LightSG into half vector space
-	#if 0
-		FSphericalGaussian WarpedLightSG;
-		WarpedLightSG.Axis		= H;
-		WarpedLightSG.Sharpness	= LightSG.Sharpness * 1.5 * NoV;
-		WarpedLightSG.Amplitude	= LightSG.Amplitude;
-	#else
-		FAnisoSphericalGaussian WarpedLightSG;
-		WarpedLightSG.AxisZ = H;
-		WarpedLightSG.AxisX = normalize( cross( N, WarpedLightSG.AxisZ ) );
-		WarpedLightSG.AxisY = normalize( cross( H, WarpedLightSG.AxisX ) );
- 
-		// Second derivative of the sharpness with respect to how
-		// far we are from basis Axis direction
-		WarpedLightSG.SharpnessX = LightSG.Sharpness * 2 * Pow2( NoV );
-		WarpedLightSG.SharpnessY = LightSG.Sharpness * 2;
-		WarpedLightSG.Amplitude	= LightSG.Amplitude;
-	#endif
-
-		return Dot( WarpedLightSG, NDF );
-	}
-#else
-	{
-		// We can do the half space ASG method cheaper by assuming H is in the YZ plane.
-		float SharpnessX = LightSG.Sharpness * 2 * Pow2( NoV );
-		float SharpnessY = LightSG.Sharpness * 2;
-
-		float nu = NDF.Sharpness * 0.5;
-
-		FSphericalGaussian ConvolvedNDF;
-		ConvolvedNDF.Axis = NDF.Axis;
-		ConvolvedNDF.Sharpness = 2 * (nu * SharpnessY) / (nu + SharpnessY);
-		ConvolvedNDF.Amplitude = NDF.Amplitude * LightSG.Amplitude;
-		ConvolvedNDF.Amplitude *= PI * rsqrt( (nu + SharpnessX) * (nu + SharpnessY) );
-
-		//float3 AxisX = normalize( cross( N, V ) );
-		//ConvolvedNDF.Amplitude *= exp( -(nu * SharpnessX) / (nu + SharpnessX) * Pow2( dot( H, AxisX ) ) );
-
-		return Evaluate( ConvolvedNDF, H );
-	}
-#endif
-}
-
-FShadingOcclusion ApplyBentNormal(
-	in float3 CameraVector, 
-	in float3 WorldNormal, 
-	in float3 WorldBentNormal0, 
-	in float Roughness,
-	in float MaterialAO)
-{
-	FShadingOcclusion Out;
-	Out.DiffOcclusion = MaterialAO;
-	Out.SpecOcclusion = MaterialAO;
-	Out.BentNormal    = WorldNormal;
-
-#if NUM_MATERIAL_OUTPUTS_GETBENTNORMAL > 0
-	Out.BentNormal = WorldBentNormal0;
-
-	FSphericalGaussian HemisphereSG = Hemisphere_ToSphericalGaussian(WorldNormal);
-	FSphericalGaussian NormalSG = ClampedCosine_ToSphericalGaussian(WorldNormal);
-	FSphericalGaussian VisibleSG = BentNormalAO_ToSphericalGaussian(Out.BentNormal, Out.DiffOcclusion );
-	FSphericalGaussian DiffuseSG = Mul( NormalSG, VisibleSG );
-	
-	float VisibleCosAngle = sqrt( 1 - Out.DiffOcclusion );
-
-#if 1	// Mix full resolution normal with low res bent normal
-	Out.BentNormal = DiffuseSG.Axis;
-	//DiffOcclusion = saturate( Integral( DiffuseSG ) / Dot( NormalSG, HemisphereSG ) );
-	Out.DiffOcclusion = saturate( Integral( DiffuseSG ) * 0.42276995 );
-#endif
-
-	float3 N = WorldNormal;
-	float3 V = CameraVector;
-
-	Out.SpecOcclusion  = DotSpecularSG( Roughness, N, V, VisibleSG );
-	Out.SpecOcclusion /= DotSpecularSG( Roughness, N, V, HemisphereSG );
-
-	Out.SpecOcclusion = saturate(Out.SpecOcclusion );
-#endif
-	return Out;
-}
-
 #if STRATA_ENABLED
 
 uint GetDiffuseIndirectSampleOcclusion(FSharedLocalBases SharedLocalBases, float3 V, float2 SvPosition, float MaterialAO)
diff --git a/Engine/Shaders/Private/ClearCoatCommon.ush b/Engine/Shaders/Private/ClearCoatCommon.ush
--- a/Engine/Shaders/Private/ClearCoatCommon.ush
+++ b/Engine/Shaders/Private/ClearCoatCommon.ush
@@ -18,28 +18,28 @@ float GetClearCoatRoughness(FGBufferData GBuffer)
 	return GBuffer.ShadingModelID == SHADINGMODELID_CLEAR_COAT ? GBuffer.CustomData.y : GBuffer.Roughness;
 }
 
-void RemapClearCoatDiffuseAndSpecularColor(FGBufferData GBuffer, float NoV, inout float3 DiffuseColor, inout float3 SpecularColor)
+void RemapClearCoatDiffuseAndSpecularColor(FGBufferData GBuffer, half NoV, inout half3 DiffuseColor, inout half3 SpecularColor)
 {
 	if (GBuffer.ShadingModelID == SHADINGMODELID_CLEAR_COAT)
 	{
 		// Attenuate base color and recompute diffuse color
-		float RefractionScale = ((NoV * 0.5 + 0.5) * NoV - 1) * saturate(1.25 - 1.25 * GBuffer.Roughness) + 1;
+		half RefractionScale = ((NoV * 0.5 + 0.5) * NoV - 1) * saturate(1.25 - 1.25 * GBuffer.Roughness) + 1;
 
-		float MetalSpec = 0.9;
-		float3 AbsorptionColor = GBuffer.BaseColor * (1 / MetalSpec);
-		float3 Absorption = AbsorptionColor * ((NoV - 1) * 0.85 * (1 - lerp(AbsorptionColor, Square(AbsorptionColor), -0.78)) + 1);
+		half MetalSpec = 0.9;
+		half3 AbsorptionColor = GBuffer.BaseColor * (1 / MetalSpec);
+		half3 Absorption = AbsorptionColor * ((NoV - 1) * 0.85 * (1 - lerp(AbsorptionColor, Square(AbsorptionColor), -0.78)) + 1);
 
-		float F0 = 0.04;
-		float Fc = Pow5(1 - NoV);
-		float F = Fc + (1 - Fc) * F0;
-		float ClearCoat = GBuffer.CustomData.x;
-		float LayerAttenuation = lerp(1, (1 - F), ClearCoat);
+		half F0 = 0.04;
+		half Fc = Pow5(1 - NoV);
+		half F = Fc + (1 - Fc) * F0;
+		half ClearCoat = GBuffer.CustomData.x;
+		half LayerAttenuation = lerp(1, (1 - F), ClearCoat);
 
-		float3 BaseColor = lerp(GBuffer.BaseColor * LayerAttenuation, MetalSpec * Absorption * RefractionScale, GBuffer.Metallic * ClearCoat);
+		half3 BaseColor = lerp(GBuffer.BaseColor * LayerAttenuation, MetalSpec * Absorption * RefractionScale, GBuffer.Metallic * ClearCoat);
 		//BaseColor += Dither / 255.f;
 		DiffuseColor = BaseColor - BaseColor * GBuffer.Metallic;
 
-		float Specular = lerp(GBuffer.Specular, RefractionScale, ClearCoat);
+		half Specular = lerp(GBuffer.Specular, RefractionScale, ClearCoat);
 		SpecularColor = ComputeF0(Specular, BaseColor, GBuffer.Metallic) * View.SpecularOverrideParameter.w + View.SpecularOverrideParameter.xyz;
 	}
 }
diff --git a/Engine/Shaders/Private/DeferredLightingCommon.ush b/Engine/Shaders/Private/DeferredLightingCommon.ush
--- a/Engine/Shaders/Private/DeferredLightingCommon.ush
+++ b/Engine/Shaders/Private/DeferredLightingCommon.ush
@@ -35,16 +35,16 @@
 struct FDeferredLightData
 {
 	float3 TranslatedWorldPosition;
-	float  InvRadius;
-	float3 Color;
-	float  FalloffExponent;
+	half  InvRadius;
+	half3 Color;
+	half  FalloffExponent;
 	float3 Direction;
 	float3 Tangent;
 	float SoftSourceRadius;
-	float2 SpotAngles;
+	half2 SpotAngles;
 	float SourceRadius;
 	float SourceLength;
-	float SpecularScale;
+	half SpecularScale;
 	float ContactShadowLength;
 	/** Intensity of non-shadow-casting contact shadows */
 	float ContactShadowNonShadowCastingIntensity;
@@ -100,6 +100,7 @@
 	return Fade * Fade;
 }
 
+#if SUPPORT_CONTACT_SHADOWS
 // Returns distance along ray that the first hit occurred, or negative on miss
 // Sets bOutHitCastDynamicShadow if the hit point is marked as a dynamic shadow caster
 float ShadowRayCast(
@@ -175,12 +176,13 @@
 	
 	return HitDistance;
 }
+#endif
 
 #ifndef SUPPORT_CONTACT_SHADOWS
 #error "Must set SUPPORT_CONTACT_SHADOWS"
 #endif
 
-void GetShadowTerms(float SceneDepth, float4 PrecomputedShadowFactors, uint ShadingModelID, float ContactShadowOpacity, FDeferredLightData LightData, float3 TranslatedWorldPosition, float3 L, float4 LightAttenuation, float Dither, inout FShadowTerms Shadow)
+void GetShadowTerms(float SceneDepth, float4 PrecomputedShadowFactors, uint ShadingModelID, float ContactShadowOpacity, FDeferredLightData LightData, float3 TranslatedWorldPosition, half3 L, half4 LightAttenuation, float Dither, inout FShadowTerms Shadow)
 {
 	float ContactShadowLength = 0.0f;
 	const float ContactShadowLengthScreenScale = View.ClipToView[1][1] * SceneDepth;
@@ -309,6 +311,9 @@
 	if (LightData.bInverseSquared)
 	{
 		LightMask = Square( saturate( 1 - Square( DistanceSqr * Square(LightData.InvRadius) ) ) );
+#if SHADING_PATH_MOBILE
+		LightMask *= 1.0f / (DistanceSqr + 1.0f);
+#endif
 	}
 	else
 	{
@@ -356,16 +361,15 @@
 	return Capsule;
 }
 
-/** Calculates lighting for a given position, normal, etc with a fully featured lighting model designed for quality. */
-FDeferredLightingSplit GetDynamicLightingSplit(
-	float3 TranslatedWorldPosition, float3 CameraVector, FGBufferData GBuffer, float AmbientOcclusion, uint ShadingModelID, 
-	FDeferredLightData LightData, float4 LightAttenuation, float Dither, uint2 SVPos, 
+FLightAccumulator AccumulateDynamicLighting(
+	float3 TranslatedWorldPosition, half3 CameraVector, FGBufferData GBuffer, half AmbientOcclusion, uint ShadingModelID,
+	FDeferredLightData LightData, half4 LightAttenuation, float Dither, uint2 SVPos, 
 	inout float SurfaceShadow)
 {
 	FLightAccumulator LightAccumulator = (FLightAccumulator)0;
 
-	float3 V = -CameraVector;
-	float3 N = GBuffer.WorldNormal;
+	half3 V = -CameraVector;
+	half3 N = GBuffer.WorldNormal;
 	BRANCH if( GBuffer.ShadingModelID == SHADINGMODELID_CLEAR_COAT && CLEAR_COAT_BOTTOM_NORMAL)
 	{
 		const float2 oct1 = ((float2(GBuffer.CustomData.a, GBuffer.CustomData.z) * 4) - (512.0/255.0)) + UnitVectorToOctahedron(GBuffer.WorldNormal);
@@ -374,11 +378,12 @@
 	
 	float3 L = LightData.Direction;	// Already normalized
 	float3 ToLight = L;
-	
+	half3 MaskedLightColor = LightData.Color;
 	float LightMask = 1;
 	if (LightData.bRadialLight)
 	{
 		LightMask = GetLocalLightAttenuation( TranslatedWorldPosition, LightData, ToLight, L );
+		MaskedLightColor *= LightMask;
 	}
 
 	LightAccumulator.EstimatedCost += 0.3f;		// running the PixelShader at all has a cost
@@ -398,11 +403,26 @@
 
 		LightAccumulator.EstimatedCost += 0.3f;		// add the cost of getting the shadow terms
 
+#if SHADING_PATH_MOBILE
+		const bool bNeedsSeparateSubsurfaceLightAccumulation = UseSubsurfaceProfile(GBuffer.ShadingModelID);
+		
+		FDirectLighting Lighting = (FDirectLighting)0;
+
+		half NoL = max(0, dot(GBuffer.WorldNormal, L));
+	#if TRANSLUCENCY_NON_DIRECTIONAL
+		NoL = 1.0f;
+	#endif
+		Lighting = EvaluateBxDF(GBuffer, N, V, L, NoL, Shadow);
+
+		Lighting.Specular *= LightData.SpecularScale;
+				
+		LightAccumulator_AddSplit( LightAccumulator, Lighting.Diffuse, Lighting.Specular, Lighting.Diffuse, MaskedLightColor * Shadow.SurfaceShadow, bNeedsSeparateSubsurfaceLightAccumulation );
+		LightAccumulator_AddSplit( LightAccumulator, Lighting.Transmission, 0.0f, Lighting.Transmission, MaskedLightColor * Shadow.TransmissionShadow, bNeedsSeparateSubsurfaceLightAccumulation );
+#else // SHADING_PATH_MOBILE
 		BRANCH
 		if( Shadow.SurfaceShadow + Shadow.TransmissionShadow > 0 )
 		{
 			const bool bNeedsSeparateSubsurfaceLightAccumulation = UseSubsurfaceProfile(GBuffer.ShadingModelID);
-			float3 LightColor = LightData.Color;
 
 		#if NON_DIRECTIONAL_DIRECT_LIGHTING
 			float Lighting;
@@ -421,7 +441,7 @@
 			}
 
 			float3 LightingDiffuse = Diffuse_Lambert( GBuffer.DiffuseColor ) * Lighting;
-			LightAccumulator_AddSplit(LightAccumulator, LightingDiffuse, 0.0f, 0, LightColor * LightMask * Shadow.SurfaceShadow, bNeedsSeparateSubsurfaceLightAccumulation);
+			LightAccumulator_AddSplit(LightAccumulator, LightingDiffuse, 0.0f, 0, MaskedLightColor * Shadow.SurfaceShadow, bNeedsSeparateSubsurfaceLightAccumulation);
 		#else
 			FDirectLighting Lighting;
 
@@ -449,14 +469,24 @@
 
 			Lighting.Specular *= LightData.SpecularScale;
 				
-			LightAccumulator_AddSplit( LightAccumulator, Lighting.Diffuse, Lighting.Specular, Lighting.Diffuse, LightColor * LightMask * Shadow.SurfaceShadow, bNeedsSeparateSubsurfaceLightAccumulation );
-			LightAccumulator_AddSplit( LightAccumulator, Lighting.Transmission, 0.0f, Lighting.Transmission, LightColor * LightMask * Shadow.TransmissionShadow, bNeedsSeparateSubsurfaceLightAccumulation );
+			LightAccumulator_AddSplit( LightAccumulator, Lighting.Diffuse, Lighting.Specular, Lighting.Diffuse, MaskedLightColor * Shadow.SurfaceShadow, bNeedsSeparateSubsurfaceLightAccumulation );
+			LightAccumulator_AddSplit( LightAccumulator, Lighting.Transmission, 0.0f, Lighting.Transmission, MaskedLightColor * Shadow.TransmissionShadow, bNeedsSeparateSubsurfaceLightAccumulation );
 
 			LightAccumulator.EstimatedCost += 0.4f;		// add the cost of the lighting computations (should sum up to 1 form one light)
 		#endif
 		}
+#endif // SHADING_PATH_MOBILE
 	}
+	return LightAccumulator;
+}
 
+/** Calculates lighting for a given position, normal, etc with a fully featured lighting model designed for quality. */
+FDeferredLightingSplit GetDynamicLightingSplit(
+	float3 TranslatedWorldPosition, float3 CameraVector, FGBufferData GBuffer, float AmbientOcclusion, uint ShadingModelID, 
+	FDeferredLightData LightData, float4 LightAttenuation, float Dither, uint2 SVPos, 
+	inout float SurfaceShadow)
+{
+	FLightAccumulator LightAccumulator = AccumulateDynamicLighting(TranslatedWorldPosition, CameraVector, GBuffer, AmbientOcclusion, ShadingModelID, LightData, LightAttenuation, Dither, SVPos, SurfaceShadow);
 	return LightAccumulator_GetResultSplit(LightAccumulator);
 }
 
diff --git a/Engine/Shaders/Private/DeferredShadingCommon.ush b/Engine/Shaders/Private/DeferredShadingCommon.ush
--- a/Engine/Shaders/Private/DeferredShadingCommon.ush
+++ b/Engine/Shaders/Private/DeferredShadingCommon.ush
@@ -329,34 +329,34 @@ bool HasAnisotropy(int SelectiveOutputMask)
 struct FGBufferData
 {
 	// normalized
-	float3 WorldNormal;
+	half3 WorldNormal;
 	// normalized, only valid if HAS_ANISOTROPY_MASK in SelectiveOutputMask
-	float3 WorldTangent;
+	half3 WorldTangent;
 	// 0..1 (derived from BaseColor, Metalness, Specular)
-	float3 DiffuseColor;
+	half3 DiffuseColor;
 	// 0..1 (derived from BaseColor, Metalness, Specular)
-	float3 SpecularColor;
+	half3 SpecularColor;
 	// 0..1, white for SHADINGMODELID_SUBSURFACE_PROFILE and SHADINGMODELID_EYE (apply BaseColor after scattering is more correct and less blurry)
-	float3 BaseColor;
+	half3 BaseColor;
 	// 0..1
-	float Metallic;
+	half Metallic;
 	// 0..1
-	float Specular;
+	half Specular;
 	// 0..1
-	float4 CustomData;
+	half4 CustomData;
 	// AO utility value
-	float GenericAO;
+	half GenericAO;
 	// Indirect irradiance luma
-	float IndirectIrradiance;
+	half IndirectIrradiance;
 	// Static shadow factors for channels assigned by Lightmass
 	// Lights using static shadowing will pick up the appropriate channel in their deferred pass
-	float4 PrecomputedShadowFactors;
+	half4 PrecomputedShadowFactors;
 	// 0..1
-	float Roughness;
+	half Roughness;
 	// -1..1, only valid if only valid if HAS_ANISOTROPY_MASK in SelectiveOutputMask
-	float Anisotropy;
+	half Anisotropy;
 	// 0..1 ambient occlusion  e.g.SSAO, wet surface mask, skylight mask, ...
-	float GBufferAO;
+	half GBufferAO;
 	// Bit mask for occlusion of the diffuse indirect samples
 	uint DiffuseIndirectSampleOcclusion;
 	// 0..255 
@@ -364,23 +364,26 @@ struct FGBufferData
 	// 0..255 
 	uint SelectiveOutputMask;
 	// 0..1, 2 bits, use CastContactShadow(GBuffer) or HasDynamicIndirectShadowCasterRepresentation(GBuffer) to extract
-	float PerObjectGBufferData;
+	half PerObjectGBufferData;
 	// in world units
-	float CustomDepth;
+	half CustomDepth;
 	// Custom depth stencil value
 	uint CustomStencil;
 	// in unreal units (linear), can be used to reconstruct world position,
 	// only valid when decoding the GBuffer as the value gets reconstructed from the Z buffer
-	float Depth;
+	half Depth;
 	// Velocity for motion blur (only used when WRITES_VELOCITY_TO_GBUFFER is enabled)
-	float4 Velocity;
+	half4 Velocity;
 
 	// 0..1, only needed by SHADINGMODELID_SUBSURFACE_PROFILE and SHADINGMODELID_EYE which apply BaseColor later
-	float3 StoredBaseColor;
+	half3 StoredBaseColor;
 	// 0..1, only needed by SHADINGMODELID_SUBSURFACE_PROFILE and SHADINGMODELID_EYE which apply Specular later
-	float StoredSpecular;
+	half StoredSpecular;
 	// 0..1, only needed by SHADINGMODELID_EYE which encodes Iris Distance inside Metallic
-	float StoredMetallic;
+	half StoredMetallic;
+
+	// Curvature for mobile subsurface profile
+	half Curvature;
 };
 
 bool CastContactShadow(FGBufferData GBufferData)
@@ -459,30 +462,57 @@ float4 ComputeIndirectLightingSampleE(uint2 TracingPixelCoord, uint TracingRayIn
 		Hammersley16(TracingRayIndex, TracingRayCount, Seed1));
 }
 
+#ifndef MOBILE_HIGH_QUALITY_BRDF
+#define MOBILE_HIGH_QUALITY_BRDF 0
+#endif
+
 #if SHADING_PATH_MOBILE
 #ifndef ENABLE_SHADINGMODEL_SUPPORT_MOBILE_DEFERRED
 #define ENABLE_SHADINGMODEL_SUPPORT_MOBILE_DEFERRED 0
 #endif
+
+#ifndef MOBILE_DEFERRED_SHADING
+#define MOBILE_DEFERRED_SHADING 0
+#endif
+
 #define MOBILE_SHADINGMODEL_SUPPORT (!MOBILE_DEFERRED_SHADING || ENABLE_SHADINGMODEL_SUPPORT_MOBILE_DEFERRED)
 
-float MobileEncodeShadingModelIdAndColorChannel(uint ShadingModelId, float Color)
+float MobileEncodeIdAndColorChannel(uint Id, float Color, bool b10Bits)
 {
-	uint ColorValue = (uint)round(clamp(Color, 0.0, 1.0) * 31.0) << 3; // 5bit color
-	return (float)(ColorValue | ShadingModelId) / 255.0;
+	if (b10Bits)
+	{
+		uint ColorValue = (uint)round(clamp(Color, 0.0, 1.0) * 63.0) << 4; // 6bit color
+		return (float)(ColorValue | (Id & 0x0F)) / 1023.0;
+	}
+	else // 8 bits
+	{
+		uint ColorValue = (uint)round(clamp(Color, 0.0, 1.0) * 31.0) << 3; // 5bit color
+		return (float)(ColorValue | (Id & 0x07)) / 255.0;
+	}
 }
 
-uint MobileDecodeShadingModelId(float InPackedChannel)
+uint MobileDecodeId(float InPackedChannel, bool b10Bits)
 {
-	uint ShadingModelId = (uint)round(InPackedChannel * 255.0);
-#if !MOBILE_EXTENDED_GBUFFER
-	ShadingModelId = ShadingModelId & 0x7;
-#endif
-	return ShadingModelId;
+	if (b10Bits)
+	{
+		return ((uint)round(InPackedChannel * 1023.0)) & 0x0F;
+	}
+	else // 8 bits
+	{
+		return ((uint)round(InPackedChannel * 255.0)) & 0x07;
+	}
 }
 
-float MobileDecodeColorChannel(float InPackedChannel)
+float MobileDecodeColorChannel(float InPackedChannel, bool b10Bits)
 {
-	return (((uint)(InPackedChannel * 255.0) & ~0x7) >> 3) / 31.0;
+	if (b10Bits)
+	{
+		return ((uint)(InPackedChannel * 1023.0) >> 4) / 63.0;
+	}
+	else // 8 bits
+	{
+		return ((uint)(InPackedChannel * 255.0) >> 3) / 31.0;
+	}
 }
 
 /** Mobile specific encoding of GBuffer data */
@@ -503,10 +533,6 @@ void MobileEncodeGBuffer(
 	}
 	else
 	{
-#if !MOBILE_EXTENDED_GBUFFER
-		// We allocate only 3bits for shading model ID
-		GBuffer.ShadingModelID = (GBuffer.ShadingModelID > 7 ? SHADINGMODELID_DEFAULT_LIT : GBuffer.ShadingModelID);
-#endif
 		OutGBufferA.rg = UnitVectorToOctahedron( normalize(GBuffer.WorldNormal) ) * 0.5f + 0.5f;
 #if ALLOW_STATIC_LIGHTING
 		OutGBufferA.b = EncodeIndirectIrradiance(GBuffer.IndirectIrradiance * GBuffer.GBufferAO);
@@ -528,33 +554,75 @@ void MobileEncodeGBuffer(
 #endif
 		OutGBufferD = GBuffer.CustomData;
 
-#if MOBILE_SHADINGMODEL_SUPPORT && !MOBILE_EXTENDED_GBUFFER
+#if MOBILE_SHADINGMODEL_SUPPORT
+	#if	MOBILE_EXTENDED_GBUFFER
+		if (GBuffer.ShadingModelID == SHADINGMODELID_EYE)
+		{
+			// encode Curvature into 5 bits and 3 bits for SubsurfaceProfile 
+			OutGBufferD.x = MobileEncodeIdAndColorChannel(GBuffer.CustomData.x, GBuffer.Curvature, false);
+		}
+	#else
 		if (GBuffer.ShadingModelID == SHADINGMODELID_SUBSURFACE || 
 			GBuffer.ShadingModelID == SHADINGMODELID_PREINTEGRATED_SKIN)
 		{
-			OutGBufferA.b = GBuffer.CustomData.a; // Opacity
-			OutGBufferB.r = GBuffer.CustomData.r; // SubsurfaceColor.R
+			// encode Opacity into 6 bits and 4 bits for shading model
+			OutGBufferA.b = MobileEncodeIdAndColorChannel(GBuffer.ShadingModelID, GBuffer.CustomData.a, true);
+			OutGBufferB.a = GBuffer.CustomData.r; // SubsurfaceColor.R
 			OutGBufferC.a = GBuffer.CustomData.g; // SubsurfaceColor.G
-			// encode SubsurfaceColor.B into 5 bits and 3 bits for shading model 
-			OutGBufferB.a = MobileEncodeShadingModelIdAndColorChannel(GBuffer.ShadingModelID, GBuffer.CustomData.b);
+			OutGBufferB.r = GBuffer.CustomData.b; // SubsurfaceColor.B
 		}
 		else if (GBuffer.ShadingModelID == SHADINGMODELID_CLEAR_COAT)
 		{
-			OutGBufferA.b = GBuffer.CustomData.x; // ClearCoat
+			// encode 4 bits for shading model
+			OutGBufferA.b = MobileEncodeIdAndColorChannel(GBuffer.ShadingModelID, 0, true);
+			OutGBufferB.a = GBuffer.CustomData.x; // ClearCoat
 			OutGBufferC.a = GBuffer.CustomData.y; // ClearCoatRoughness
 		}
 		else if (GBuffer.ShadingModelID == SHADINGMODELID_SUBSURFACE_PROFILE)
 		{
-			OutGBufferA.b = GBuffer.CustomData.a; // Opacity
-			OutGBufferB.r = GBuffer.CustomData.r; // SubsurfaceProfile
-			OutGBufferC.a = GBuffer.CustomData.g; // Curvature
+			// encode Opacity into 6 bits and 4 bits for shading model
+			OutGBufferA.b = MobileEncodeIdAndColorChannel(GBuffer.ShadingModelID, GBuffer.CustomData.a, true);
+			OutGBufferB.a = GBuffer.CustomData.r; // SubsurfaceProfile
+			OutGBufferC.a = GBuffer.Curvature; // Curvature
 		}
 		else if (GBuffer.ShadingModelID == SHADINGMODELID_HAIR)
 		{
-			OutGBufferA.b = GBuffer.CustomData.z; // Backlit
-			OutGBufferB.r = GBuffer.CustomData.x; // 
+			// encode Backlit into 6 bits and 4 bits for shading model
+			OutGBufferA.b = MobileEncodeIdAndColorChannel(GBuffer.ShadingModelID, GBuffer.CustomData.z, true);
+			OutGBufferB.a = GBuffer.CustomData.x; // 
 			OutGBufferC.a = GBuffer.CustomData.y; // 
 		}
+		else if (GBuffer.ShadingModelID == SHADINGMODELID_TWOSIDED_FOLIAGE)
+		{
+			// encode 4 bits for shading model
+			OutGBufferA.b = MobileEncodeIdAndColorChannel(GBuffer.ShadingModelID, 0, true);
+			OutGBufferB.a = GBuffer.CustomData.r; // SubsurfaceColor.R
+			OutGBufferC.a = GBuffer.CustomData.g; // SubsurfaceColor.G
+			OutGBufferB.r = GBuffer.CustomData.b; // SubsurfaceColor.B
+		}
+		else if (GBuffer.ShadingModelID == SHADINGMODELID_CLOTH)
+		{
+			// encode Cloth into 6 bits and 4 bits for shading model
+			OutGBufferA.b = MobileEncodeIdAndColorChannel(GBuffer.ShadingModelID, GBuffer.CustomData.a, true);
+			OutGBufferB.a = GBuffer.CustomData.r; // SubsurfaceColor.R
+			OutGBufferC.a = GBuffer.CustomData.g; // SubsurfaceColor.G
+			OutGBufferB.r = GBuffer.CustomData.b; // SubsurfaceColor.B
+		}
+		else if (GBuffer.ShadingModelID == SHADINGMODELID_EYE)
+		{
+			// encode IrisMask into 6 bits and 4 bits for shading model 
+			OutGBufferA.b = MobileEncodeIdAndColorChannel(GBuffer.ShadingModelID, GBuffer.CustomData.w, true);
+			// encode Curvature into 5 bits and 3 bits for SubsurfaceProfile 
+			OutGBufferB.a = MobileEncodeIdAndColorChannel(GBuffer.CustomData.x, GBuffer.Curvature, false);
+			OutGBufferC.a = GBuffer.CustomData.y; // IrisNormal_Oct.x
+			OutGBufferB.g = GBuffer.CustomData.z; // IrisNormal_Oct.y
+		}
+		else
+		{
+			// encode 4 bits for shading model
+			OutGBufferA.b = MobileEncodeIdAndColorChannel(GBuffer.ShadingModelID, 0, true);
+		}
+	#endif
 #endif
 	}
 }
@@ -572,11 +640,11 @@ FGBufferData MobileDecodeGBuffer(half4 InGBufferA, half4 InGBufferB, half4 InGBu
 
 	GBuffer.Metallic	= InGBufferB.r;
 	GBuffer.Specular	= InGBufferB.g;
-	GBuffer.Roughness	= max(0.015625, InGBufferB.b);
+	GBuffer.Roughness	= InGBufferB.b;
 	// Note: must match GetShadingModelId standalone function logic
 	// Also Note: SimpleElementPixelShader directly sets SV_Target2 ( GBufferB ) to indicate unlit.
 	// An update there will be required if this layout changes.
-	GBuffer.ShadingModelID = MobileDecodeShadingModelId(InGBufferB.a);
+	GBuffer.ShadingModelID = MOBILE_SHADINGMODEL_SUPPORT ? (uint)round(InGBufferB.a * 255.0f) : SHADINGMODELID_DEFAULT_LIT;
 	GBuffer.SelectiveOutputMask = 0;
 	GBuffer.BaseColor = DecodeBaseColor(InGBufferC.rgb);
 #if ALLOW_STATIC_LIGHTING
@@ -587,44 +655,94 @@ FGBufferData MobileDecodeGBuffer(half4 InGBufferA, half4 InGBufferB, half4 InGBu
 	GBuffer.PrecomputedShadowFactors = 1.0;
 #endif
 
+	GBuffer.StoredBaseColor = GBuffer.BaseColor;
+	GBuffer.StoredMetallic = GBuffer.Metallic;
+	GBuffer.StoredSpecular = GBuffer.Specular;
+
 #if MOBILE_SHADINGMODEL_SUPPORT
 	#if	MOBILE_EXTENDED_GBUFFER
 		GBuffer.CustomData = HasCustomGBufferData(GBuffer.ShadingModelID) ? InGBufferD : 0;
+		if (GBuffer.ShadingModelID == SHADINGMODELID_EYE)
+		{
+			GBuffer.Curvature = MobileDecodeColorChannel(InGBufferD.x, false); // Curvature
+		}
 	#else
+		GBuffer.ShadingModelID = MobileDecodeId(InGBufferA.b, true);
+
 		if (GBuffer.ShadingModelID == SHADINGMODELID_SUBSURFACE ||
 			GBuffer.ShadingModelID == SHADINGMODELID_PREINTEGRATED_SKIN)
 		{
-			GBuffer.CustomData.a = InGBufferA.b; // Opacity
-			GBuffer.CustomData.r = InGBufferB.r; // SubsurfaceColor.R
+			GBuffer.CustomData.a = MobileDecodeColorChannel(InGBufferA.b, true); // Opacity
+			GBuffer.CustomData.r = InGBufferB.a; // SubsurfaceColor.R
 			GBuffer.CustomData.g = InGBufferC.a; // SubsurfaceColor.G
-			GBuffer.CustomData.b = MobileDecodeColorChannel(InGBufferB.a); // SubsurfaceColor.B
+			GBuffer.CustomData.b = InGBufferB.r; // SubsurfaceColor.B
 			GBuffer.GBufferAO = 1;	// no space for AO
 			GBuffer.Metallic = 0;	// no space for Metallic
 		}
 		else if (GBuffer.ShadingModelID == SHADINGMODELID_CLEAR_COAT)
 		{
-			GBuffer.CustomData.x = InGBufferA.b; // ClearCoat
+			GBuffer.CustomData.x = InGBufferB.a; // ClearCoat
 			GBuffer.CustomData.y = InGBufferC.a; // ClearCoatRoughness
 			GBuffer.GBufferAO = 1;	// no space for AO
 		}
 		else if (GBuffer.ShadingModelID == SHADINGMODELID_SUBSURFACE_PROFILE)
 		{
-			GBuffer.CustomData.a = InGBufferA.b; // Opacity
-			GBuffer.CustomData.r = InGBufferB.r; // SubsurfaceProfile
-			GBuffer.CustomData.g = InGBufferC.a; // Curvature
+			GBuffer.CustomData.a = MobileDecodeColorChannel(InGBufferA.b, true); // Opacity
+			GBuffer.CustomData.r = InGBufferB.a; // SubsurfaceProfile
+			GBuffer.Curvature = InGBufferC.a; // Curvature
 			GBuffer.GBufferAO = 1;	// no space for AO
-			GBuffer.Metallic = 0;	// no space for Metallic
 		}
 		else if (GBuffer.ShadingModelID == SHADINGMODELID_HAIR)
 		{
-			GBuffer.CustomData.z = InGBufferA.b; // Backlit
-			GBuffer.CustomData.x = InGBufferB.r; // 
+			GBuffer.CustomData.z = MobileDecodeColorChannel(InGBufferA.b, true); // Backlit
+			GBuffer.CustomData.x = InGBufferB.a; // 
 			GBuffer.CustomData.y = InGBufferC.a; //
 			GBuffer.GBufferAO = 1;	// no space for AO
+		}
+		else if (GBuffer.ShadingModelID == SHADINGMODELID_TWOSIDED_FOLIAGE)
+		{
+			GBuffer.CustomData.r = InGBufferB.a; // SubsurfaceColor.R
+			GBuffer.CustomData.g = InGBufferC.a; // SubsurfaceColor.G
+			GBuffer.CustomData.b = InGBufferB.r; // SubsurfaceColor.B
+			GBuffer.GBufferAO = 1;	// no space for AO
 			GBuffer.Metallic = 0;	// no space for Metallic
 		}
+		else if (GBuffer.ShadingModelID == SHADINGMODELID_CLOTH)
+		{
+			GBuffer.CustomData.a = MobileDecodeColorChannel(InGBufferA.b, true); // Cloth
+			GBuffer.CustomData.r = InGBufferB.a; // SubsurfaceColor.R
+			GBuffer.CustomData.g = InGBufferC.a; // SubsurfaceColor.G
+			GBuffer.CustomData.b = InGBufferB.r; // SubsurfaceColor.B
+			GBuffer.GBufferAO = 1;	// no space for AO
+			GBuffer.Metallic = 0;	// no space for Metallic
+		}
+		else if (GBuffer.ShadingModelID == SHADINGMODELID_EYE)
+		{
+			GBuffer.CustomData.w = MobileDecodeColorChannel(InGBufferA.b, true); // IrisMask
+			GBuffer.CustomData.x = MobileDecodeId(InGBufferB.a, false); // SubsurfaceProfile
+			GBuffer.CustomData.y = InGBufferC.a; // IrisNormal_Oct.x
+			GBuffer.CustomData.z = InGBufferB.g; // IrisNormal_Oct.y
+			GBuffer.Curvature = MobileDecodeColorChannel(InGBufferB.a, false); // Curvature
+			
+			GBuffer.GBufferAO = 1;	// no space for AO
+			GBuffer.Specular = 0.5f;	// no space for Specular
+		}
 	#endif
+
+	if (GBuffer.ShadingModelID == SHADINGMODELID_EYE)
+	{
+		GBuffer.Metallic = 0.0;
+	#if IRIS_NORMAL
+		GBuffer.Specular = 0.25;
+	#endif
+	}
 #endif
+
+	// derived from BaseColor, Metalness, Specular
+	{
+		GBuffer.SpecularColor = ComputeF0(GBuffer.Specular, GBuffer.BaseColor, GBuffer.Metallic);
+		GBuffer.DiffuseColor = GBuffer.BaseColor - GBuffer.BaseColor * GBuffer.Metallic;
+	}
 	return GBuffer;
 }
 
@@ -997,12 +1115,19 @@ FScreenSpaceData GetScreenSpaceData(float2 UV, bool bGetNormalizedNormal = true)
 
 #endif
 // [ Jimenez et al. 2016, "Practical Realtime Strategies for Accurate Indirect Occlusion" ]
-float3 AOMultiBounce(float3 BaseColor, float AO)
+half3 AOMultiBounce(half3 BaseColor, half AO)
 {
-	float3 a = 2.0404 * BaseColor - 0.3324;
-	float3 b = -4.7951 * BaseColor + 0.6417;
-	float3 c = 2.7552 * BaseColor + 0.6903;
-	return max(AO, ((AO * a + b) * AO + c) * AO);
+	if (SHADING_PATH_MOBILE && !MOBILE_HIGH_QUALITY_BRDF)
+	{
+		return AO;
+	}
+	else
+	{
+		half3 a = 2.0404 * BaseColor - 0.3324;
+		half3 b = -4.7951 * BaseColor + 0.6417;
+		half3 c = 2.7552 * BaseColor + 0.6903;
+		return max(AO, ((AO * a + b) * AO + c) * AO);
+	}
 }
 
 #define LIGHTING_CHANNELS_TEXTURE_DISTANCE_FIELD_REPRESENTATION_BIT 3
\ No newline at end of file
diff --git a/Engine/Shaders/Private/LightmapCommon.ush b/Engine/Shaders/Private/LightmapCommon.ush
--- a/Engine/Shaders/Private/LightmapCommon.ush
+++ b/Engine/Shaders/Private/LightmapCommon.ush
@@ -75,8 +75,10 @@ half4 SampleLightmapVT(VTPageTableResult LightmapVTPageTableResult, uint LayerIn
 }
 #endif // LIGHTMAP_VT_ENABLED
 
-half4 GetLightMapColorLQ( VTPageTableResult LightmapVTPageTableResult, LightmapUVType LightmapUV0, LightmapUVType LightmapUV1, uint LightmapDataIndex, half3 WorldNormal )
+void GetLightMapColorLQ( VTPageTableResult LightmapVTPageTableResult, LightmapUVType LightmapUV0, LightmapUVType LightmapUV1, uint LightmapDataIndex, half3 WorldNormal, bool bEvaluateBackface, out half3 OutDiffuseLighting, out half3 OutSubsurfaceLighting, out half OutIndirectIrradiance)
 {
+	OutSubsurfaceLighting = 0;
+
 #if LIGHTMAP_VT_ENABLED
 	half4 Lightmap0 = SampleLightmapVT(LightmapVTPageTableResult, 0u, LightmapDataIndex, LightmapResourceCluster.VTLightMapTexture, LightmapResourceCluster.LightMapSampler);
 	half4 Lightmap1 = SampleLightmapVT(LightmapVTPageTableResult, 1u, LightmapDataIndex, LightmapResourceCluster.VTLightMapTexture_1, LightmapResourceCluster.LightMapSampler_1);
@@ -100,14 +102,32 @@ half4 GetLightMapColorLQ( VTPageTableResult LightmapVTPageTableResult, LightmapU
 
 	// Sample SH with normal
 	half Directionality = max( 0.0, dot( SH, float4(WorldNormal.yzx, 1) ) );	// 1 dot, 1 smax
+
+	#if SHADINGMODEL_REQUIRES_BACKFACE_LIGHTING
+	if (bEvaluateBackface)
+	{
+		half SubsurfaceDirectionality = max(0.0, dot(SH, float4(-WorldNormal.yzx, 1)));
+		half SubsurfaceLuma = L * SubsurfaceDirectionality;
+		OutSubsurfaceLighting = LogRGB * (SubsurfaceLuma / LogL);
+	}
+	#endif
 #else
 	half Directionality = 0.6;
+
+	#if SHADINGMODEL_REQUIRES_BACKFACE_LIGHTING
+	if (bEvaluateBackface)
+	{
+		half SubsurfaceLuma = L * Directionality;
+		OutSubsurfaceLighting = LogRGB * (SubsurfaceLuma / LogL);
+	}
+	#endif
 #endif
 		
 	half Luma = L * Directionality;
 	half3 Color = LogRGB * (Luma / LogL);				// 1 rcp, 1 smul, 1 vmul
 
-	return half4( Color, Luma );
+	OutDiffuseLighting = Color;
+	OutIndirectIrradiance = Luma;
 }
 
 void GetLightMapColorHQ( VTPageTableResult LightmapVTPageTableResult, LightmapUVType LightmapUV0, LightmapUVType LightmapUV1, uint LightmapDataIndex, half3 WorldNormal, float2 SvPositionXY, bool bEvaluateBackface, out half3 OutDiffuseLighting, out half3 OutSubsurfaceLighting )
diff --git a/Engine/Shaders/Private/MobileBasePassCommon.ush b/Engine/Shaders/Private/MobileBasePassCommon.ush
--- a/Engine/Shaders/Private/MobileBasePassCommon.ush
+++ b/Engine/Shaders/Private/MobileBasePassCommon.ush
@@ -4,6 +4,7 @@
 	MobileBasePassCommon.ush: Base pass definitions used by both vertex and pixel shader
 =============================================================================*/
 
+#include "BasePassCommon.ush"
 #define FogStruct MobileBasePass.Fog
 #define PlanarReflectionStruct MobileBasePass.PlanarReflection
 
diff --git a/Engine/Shaders/Private/MobileBasePassPixelShader.usf b/Engine/Shaders/Private/MobileBasePassPixelShader.usf
--- a/Engine/Shaders/Private/MobileBasePassPixelShader.usf
+++ b/Engine/Shaders/Private/MobileBasePassPixelShader.usf
@@ -126,11 +126,41 @@ void ApplyPixelDepthOffsetForMobileBasePass(inout FMaterialPixelParameters Mater
 	#define PIXELSHADER_EARLYDEPTHSTENCIL 	
 #endif
 
-void ComputeIndirect(VTPageTableResult LightmapVTPageTableResult, FVertexFactoryInterpolantsVSToPS Interpolants, float3 DiffuseDir, half3 DiffuseColorForIndirect, half MaterialAO, out half IndirectIrradiance, inout FMobileLightAccumulator MobileLightAccumulator)
+void GetSkyLighting(bool bApplySkyLighting, half3 SkyLightColor, half3 WorldNormal, bool bEvaluateBackface, inout half IndirectIrradiance, out half3 OutSkyDiffuseLighting, out half3 OutSkySubsurfaceLighting, out half3 OutWaterDiffuseIndirectLuminance)
 {
-	//To keep IndirectLightingCache conherence with PC, initialize the IndirectIrradiance to zero.
-	IndirectIrradiance = 0;
-	half3 IndirectColor = 0;
+	OutSkyDiffuseLighting = 0;
+	OutSkySubsurfaceLighting = 0;
+	OutWaterDiffuseIndirectLuminance = 0;
+
+#if ENABLE_SKY_LIGHT
+	//@mw todo
+	// TODO: Also need to do specular.
+	if (bApplySkyLighting)
+	{
+		half3 SkyDiffuse = GetSkySHDiffuseSimple(WorldNormal);
+		OutSkyDiffuseLighting = SkyDiffuse * SkyLightColor;
+		IndirectIrradiance += Luminance(OutSkyDiffuseLighting);
+
+	#if MATERIAL_SHADINGMODEL_SINGLELAYERWATER
+		OutWaterDiffuseIndirectLuminance += SkyDiffuse;
+	#endif
+
+	#if SHADINGMODEL_REQUIRES_BACKFACE_LIGHTING
+		if (bEvaluateBackface)
+		{
+			OutSkySubsurfaceLighting = GetSkySHDiffuseSimple(-WorldNormal) * SkyLightColor;
+		}
+	#endif
+	}
+#endif
+}
+
+void GetPrecomputedIndirectLightingAndSkyLight(VTPageTableResult LightmapVTPageTableResult, FVertexFactoryInterpolantsVSToPS Interpolants, bool bEvaluateBackface, half3 DiffuseDir, bool bApplySkyLighting, half3 SkyLightColor, out half3 OutDiffuseIndirectLighting, out half3 OutSubsurfaceIndirectLighting, out half OutIndirectIrradiance, out half3 OutWaterDiffuseIndirectLuminance)
+{
+	//To keep IndirectLightingCache conherence with PC, initialize the OutIndirectIrradiance to zero.
+	OutIndirectIrradiance = 0;
+	OutDiffuseIndirectLighting = 0;
+	OutSubsurfaceIndirectLighting = 0;
 
 	// Indirect Diffuse
 #if LQ_TEXTURE_LIGHTMAP
@@ -138,9 +168,7 @@ void ComputeIndirect(VTPageTableResult LightmapVTPageTableResult, FVertexFactory
 	uint LightmapDataIndex;
 	GetLightMapCoordinates(Interpolants, LightmapUV0, LightmapUV1, LightmapDataIndex);
 
-	half4 LightmapColor = GetLightMapColorLQ(LightmapVTPageTableResult, LightmapUV0, LightmapUV1, LightmapDataIndex, DiffuseDir);
-	IndirectColor += LightmapColor.rgb * DiffuseColorForIndirect * View.IndirectLightingColorScale;
-	IndirectIrradiance = LightmapColor.a;
+	GetLightMapColorLQ(LightmapVTPageTableResult, LightmapUV0, LightmapUV1, LightmapDataIndex, DiffuseDir, bEvaluateBackface, OutDiffuseIndirectLighting, OutSubsurfaceIndirectLighting, OutIndirectIrradiance);
 #elif CACHED_POINT_INDIRECT_LIGHTING
 	#if MATERIALBLENDING_MASKED || MATERIALBLENDING_SOLID
 		// Take the normal into account for opaque
@@ -160,24 +188,38 @@ void ComputeIndirect(VTPageTableResult LightmapVTPageTableResult, FVertexFactory
 		FThreeBandSHVector DiffuseTransferSH = CalcDiffuseTransferSH3(DiffuseDir, 1);
 
 		// Compute diffuse lighting which takes the normal into account
-		half3 DiffuseGI = max(half3(0, 0, 0), DotSH3(PointIndirectLighting, DiffuseTransferSH));
-
-		IndirectIrradiance = Luminance(DiffuseGI);
-		IndirectColor += DiffuseColorForIndirect * DiffuseGI * View.IndirectLightingColorScale;
+		OutDiffuseIndirectLighting += max(half3(0, 0, 0), DotSH3(PointIndirectLighting, DiffuseTransferSH));
+
+		#if SHADINGMODEL_REQUIRES_BACKFACE_LIGHTING
+		if (bEvaluateBackface)
+		{
+			FThreeBandSHVector SubsurfaceTransferSH = CalcDiffuseTransferSH3(-DiffuseDir, 1);
+			OutSubsurfaceIndirectLighting += max(half3(0, 0, 0), DotSH3(PointIndirectLighting, SubsurfaceTransferSH));
+		}
+		#endif
 	#else 
 		// Non-directional for translucency
 		// Ambient terms packed in xyz
 		// Already divided by PI and SH ambient on CPU
-		half3 PointIndirectLighting = IndirectLightingCache.IndirectLightingSHSingleCoefficient.rgb;
-		half3 DiffuseGI = PointIndirectLighting;
+		OutDiffuseIndirectLighting += IndirectLightingCache.IndirectLightingSHSingleCoefficient.rgb;
 
-		IndirectIrradiance = Luminance(DiffuseGI);
-		IndirectColor += DiffuseColorForIndirect * DiffuseGI * View.IndirectLightingColorScale;
+		#if SHADINGMODEL_REQUIRES_BACKFACE_LIGHTING
+		if (bEvaluateBackface)
+		{
+			OutSubsurfaceIndirectLighting += IndirectLightingCache.IndirectLightingSHSingleCoefficient.rgb;
+		}
+		#endif
 	#endif
+
+	OutIndirectIrradiance = Luminance(OutDiffuseIndirectLighting);
 #endif
-	IndirectIrradiance *= MaterialAO;
 
-	MobileLightAccumulator_Add(MobileLightAccumulator, IndirectColor, 0.0f, MaterialAO);
+	half3 SkyDiffuseLighting;
+	half3 SkySubsurfaceLighting;
+	GetSkyLighting(bApplySkyLighting, SkyLightColor, DiffuseDir, bEvaluateBackface, OutIndirectIrradiance, SkyDiffuseLighting, SkySubsurfaceLighting, OutWaterDiffuseIndirectLuminance);
+
+	OutDiffuseIndirectLighting += SkyDiffuseLighting;
+	OutSubsurfaceIndirectLighting += SkySubsurfaceLighting;
 }
 
 PIXELSHADER_EARLYDEPTHSTENCIL
@@ -287,24 +329,40 @@ void Main(
 	// 0..1, SubsurfaceProfileId = int(x * 255)
 	float SubsurfaceProfile = 0;
 
-#if (MATERIAL_SHADINGMODEL_SUBSURFACE || MATERIAL_SHADINGMODEL_PREINTEGRATED_SKIN || MATERIAL_SHADINGMODEL_SUBSURFACE_PROFILE) && MOBILE_SHADINGMODEL_SUPPORT
-	half4 SubsurfaceData = GetMaterialSubsurfaceData(PixelMaterialInputs);
-	
-	// If the subsurface feature is disabled, the subsurface color should be black
-	// This will also ensure that the subsurface profile does not contribute when subsurface scattering is disabled
-	SubsurfaceData *= View.bSubsurfacePostprocessEnabled;
-
-	SubsurfaceColor = SubsurfaceData.rgb * ResolvedView.DiffuseOverrideParameter.w + ResolvedView.DiffuseOverrideParameter.xyz;
-
-	SubsurfaceProfile = SubsurfaceData.a;
+#if (MATERIAL_SHADINGMODEL_SUBSURFACE || MATERIAL_SHADINGMODEL_PREINTEGRATED_SKIN || MATERIAL_SHADINGMODEL_SUBSURFACE_PROFILE || MATERIAL_SHADINGMODEL_TWOSIDED_FOLIAGE || MATERIAL_SHADINGMODEL_CLOTH || MATERIAL_SHADINGMODEL_EYE) 
+	if (ShadingModelID == SHADINGMODELID_SUBSURFACE || ShadingModelID == SHADINGMODELID_PREINTEGRATED_SKIN || ShadingModelID == SHADINGMODELID_SUBSURFACE_PROFILE || ShadingModelID == SHADINGMODELID_TWOSIDED_FOLIAGE || ShadingModelID == SHADINGMODELID_CLOTH || ShadingModelID == SHADINGMODELID_EYE)
+	{
+		half4 SubsurfaceData = GetMaterialSubsurfaceData(PixelMaterialInputs);
+
+		// If the subsurface feature is disabled, the subsurface color should be black
+		// This will also ensure that the subsurface profile does not contribute when subsurface scattering is disabled
+		SubsurfaceData *= View.bSubsurfacePostprocessEnabled;
+
+		if (ShadingModelID == SHADINGMODELID_SUBSURFACE || ShadingModelID == SHADINGMODELID_PREINTEGRATED_SKIN || ShadingModelID == SHADINGMODELID_TWOSIDED_FOLIAGE)
+		{
+			SubsurfaceColor = SubsurfaceData.rgb * ResolvedView.DiffuseOverrideParameter.w + ResolvedView.DiffuseOverrideParameter.xyz;
+		}
+		else if (ShadingModelID == SHADINGMODELID_CLOTH)
+		{
+			SubsurfaceColor = SubsurfaceData.rgb;
+		}
+
+		SubsurfaceProfile = SubsurfaceData.a;
+	}
 #endif
 
 #if APPLY_AO
 	half4 GatheredAmbientOcclusion = Texture2DSample(AmbientOcclusionTexture, AmbientOcclusionSampler, SvPositionToBufferUV(SvPosition));
 
-	MaterialAO *= GatheredAmbientOcclusion.r;
+	MaterialAO *= lerp(1.0, GatheredAmbientOcclusion.r, AmbientOcclusionStaticFraction);
 #endif
 
+	const half BaseMaterialCoverageOverWater = Opacity;
+	const half WaterVisibility = 1.0 - BaseMaterialCoverageOverWater;
+	// Fade out diffuse as this will be handled by the single scattering lighting. when over the water surface.
+	// We keep the SpecularColor for sun/water interactions
+	half3 WaterDiffuseIndirectLuminance = 0;
+
 	FGBufferData GBuffer = (FGBufferData)0;
 	GBuffer.GBufferAO = MaterialAO;
 	GBuffer.Depth = MaterialParameters.ScreenPosition.w;
@@ -324,13 +382,78 @@ void Main(
 		ShadingModelID
 	);
 
-	GBuffer.Roughness = max(0.015625f, GBuffer.Roughness);
-	
-	FMobileShadingModelContext ShadingModelContext = (FMobileShadingModelContext)0;
-	ShadingModelContext.Opacity = Opacity;
+	GBuffer.StoredBaseColor = GBuffer.BaseColor;
+	GBuffer.StoredMetallic = GBuffer.Metallic;
+	GBuffer.StoredSpecular = GBuffer.Specular;
+
+#if MATERIAL_SHADINGMODEL_EYE
+	if( GBuffer.ShadingModelID == SHADINGMODELID_EYE )
+	{
+		GBuffer.Metallic = 0.0;
+		#if IRIS_NORMAL
+		GBuffer.Specular = 0.25;
+		#endif
+	}
+#endif
+
+#if NONMETAL
+	GBuffer.DiffuseColor = GBuffer.BaseColor;
+	GBuffer.SpecularColor = 0.04;
+#else
+	GBuffer.SpecularColor = ComputeF0(GBuffer.Specular, GBuffer.BaseColor, GBuffer.Metallic);
+	GBuffer.DiffuseColor = GBuffer.BaseColor - GBuffer.BaseColor * GBuffer.Metallic;
+#endif
+
+#if FULLY_ROUGH
+	EnvBRDFApproxFullyRough(GBuffer.DiffuseColor, GBuffer.SpecularColor);
+#endif
 
+#if MOBILE_EMULATION && !MOBILE_DEFERRED_LIGHTING
+	{
+		// this feature is only needed for development/editor - we can compile it out for a shipping build (see r.CompileShadersForDevelopment cvar help)
+		GBuffer.DiffuseColor = GBuffer.DiffuseColor * ResolvedView.DiffuseOverrideParameter.w + ResolvedView.DiffuseOverrideParameter.xyz;
+		GBuffer.SpecularColor = GBuffer.SpecularColor * ResolvedView.SpecularOverrideParameter.w + ResolvedView.SpecularOverrideParameter.xyz;
+	}
+#endif
+
+#if MATERIAL_NORMAL_CURVATURE_TO_ROUGHNESS
+	// Curvature-to-roughness uses derivatives of the WorldVertexNormal, which is incompatible with centroid interpolation because
+	// the samples are not uniformly distributed. Therefore we use WorldVertexNormal_Center which is guaranteed to be center interpolated.
+	#if USE_WORLDVERTEXNORMAL_CENTER_INTERPOLATION
+	float GeometricAARoughness = NormalCurvatureToRoughness(MaterialParameters.WorldVertexNormal_Center);
+	#else
+	float GeometricAARoughness = NormalCurvatureToRoughness(MaterialParameters.TangentToWorld[2].xyz);
+	#endif
+	GBuffer.Roughness = max(GBuffer.Roughness, GeometricAARoughness);
+
+	#if MATERIAL_SHADINGMODEL_CLEAR_COAT
+	if (GBuffer.ShadingModelID == SHADINGMODELID_CLEAR_COAT)
+	{
+		GBuffer.CustomData.y = max(GBuffer.CustomData.y, GeometricAARoughness);
+	}
+	#endif
+#endif
+
+	half3 InputBentNormal = MaterialParameters.WorldNormal;
+
+	// Clear Coat Bottom Normal
+	#if MATERIAL_SHADINGMODEL_CLEAR_COAT
+	BRANCH if (GBuffer.ShadingModelID == SHADINGMODELID_CLEAR_COAT && CLEAR_COAT_BOTTOM_NORMAL)
+	{
+		const float2 oct1 = ((float2(GBuffer.CustomData.a, GBuffer.CustomData.z) * 4) - (512.0 / 255.0)) + UnitVectorToOctahedron(GBuffer.WorldNormal);
+		InputBentNormal = OctahedronToUnitVector(oct1);
+	}
+	#endif
+
+	const FShadingOcclusion ShadingOcclusion = ApplyBentNormal(MaterialParameters.CameraVector, InputBentNormal, GetWorldBentNormalZero(MaterialParameters), GBuffer.Roughness, GBuffer.GBufferAO);
+
+#if MOBILE_HIGH_QUALITY_BRDF
+	GBuffer.GBufferAO = AOMultiBounce(Luminance(GBuffer.SpecularColor), ShadingOcclusion.SpecOcclusion).g;
+#endif
+
+	
 #if MATERIAL_SHADINGMODEL_THIN_TRANSLUCENT
-	float TopMaterialCoverage = ShadingModelContext.Opacity;
+	float TopMaterialCoverage = Opacity;
 
 	// Adjust Opacity
 	{
@@ -349,26 +472,12 @@ void Main(
 		// The alpha of the material in translucent mode is one minus the amount that it is modulating the background by.
 		float MaterialOpacity = 1.0f - MaterialModulation;
 		
-		ShadingModelContext.Opacity = 1.0f - (1.0f - TopMaterialCoverage) * (1.0f - MaterialOpacity);
+		Opacity = 1.0f - (1.0f - TopMaterialCoverage) * (1.0f - MaterialOpacity);
 	}
 
 	GBuffer.BaseColor *= TopMaterialCoverage;
 #endif
 
-#if MATERIAL_SHADINGMODEL_SUBSURFACE_PROFILE && MOBILE_SHADINGMODEL_SUPPORT
-	if (GBuffer.ShadingModelID == SHADINGMODELID_SUBSURFACE_PROFILE)
-	{
-		#if MATERIAL_SUBSURFACE_PROFILE_USE_CURVATURE
-			GBuffer.CustomData.g = GetMaterialCustomData0(MaterialParameters);
-		#else
-			GBuffer.CustomData.g = CalculateCurvature(GBuffer.WorldNormal, MaterialParameters.WorldPosition_CamRelative);
-		#endif
-
-		GBuffer.CustomData.g = max(saturate(GBuffer.CustomData.g), 0.001f);
-	}
-#endif
-	InitShadingModelContext(ShadingModelContext, GBuffer, MaterialParameters.CameraVector);
-
 	VTPageTableResult LightmapVTPageTableResult = (VTPageTableResult)0.0f;
 #if LIGHTMAP_VT_ENABLED
 	{
@@ -385,21 +494,63 @@ void Main(
 	StoreVirtualTextureFeedback(MaterialParameters.VirtualTextureFeedback, 0, LightmapVTPageTableResult.PackedRequest);
 #endif
 
-	FMobileLightAccumulator MobileLightAccumulator = (FMobileLightAccumulator)0;
+	half IndirectIrradiance = 0;
+	half3 DiffuseIndirectLighting = 0;
+	half3 SubsurfaceIndirectLighting = 0;
 
-	half IndirectIrradiance;
-	half3 IndirectColor;
-	half3 DiffuseColorForIndirect = ShadingModelContext.DiffuseColor;
+	half3 DiffuseColorForIndirect = GBuffer.DiffuseColor;
+	half3 DiffuseDir = ShadingOcclusion.BentNormal;
 
-#if (MATERIAL_SHADINGMODEL_SUBSURFACE || MATERIAL_SHADINGMODEL_PREINTEGRATED_SKIN) && MOBILE_SHADINGMODEL_SUPPORT
+#if MATERIAL_SHADINGMODEL_SUBSURFACE || MATERIAL_SHADINGMODEL_PREINTEGRATED_SKIN
 	if (GBuffer.ShadingModelID == SHADINGMODELID_SUBSURFACE || GBuffer.ShadingModelID == SHADINGMODELID_PREINTEGRATED_SKIN)
 	{
-		DiffuseColorForIndirect += ShadingModelContext.SubsurfaceColor;
+		// Add subsurface energy to diffuse
+		//@todo - better subsurface handling for these shading models with skylight and precomputed GI
+		DiffuseColorForIndirect += SubsurfaceColor;
 	}
 #endif
 
+#if MATERIAL_SHADINGMODEL_CLOTH
+	if (GBuffer.ShadingModelID == SHADINGMODELID_CLOTH)
+	{
+		DiffuseColorForIndirect += SubsurfaceColor * saturate(GetMaterialCustomData0(MaterialParameters));
+	}
+#endif
+
+#if MATERIAL_SHADINGMODEL_HAIR
+	if (GBuffer.ShadingModelID == SHADINGMODELID_HAIR)
+	{
+		FHairTransmittanceData TransmittanceData = InitHairTransmittanceData(true);
+		half3 N = MaterialParameters.WorldNormal;
+		half3 V = MaterialParameters.CameraVector;
+		half3 L = normalize(V - N * dot(V, N));
+		DiffuseDir = L;
+		DiffuseColorForIndirect = 2 * PI * HairShading(GBuffer, L, V, N, 1, TransmittanceData, 0, 0.2, uint2(0, 0));
+	}
+#endif
+
+	half NoL = max(0, dot(GBuffer.WorldNormal, MobileDirectionalLight.DirectionalLightDirectionAndShadowTransition.xyz));
+	bool bApplySkyLighting = !(MATERIAL_TWOSIDED && LQ_TEXTURE_LIGHTMAP) || NoL == 0;
+
+	const bool bEvaluateBackface = GetShadingModelRequiresBackfaceLighting(GBuffer.ShadingModelID);
+
+	GetPrecomputedIndirectLightingAndSkyLight(LightmapVTPageTableResult, Interpolants, bEvaluateBackface, DiffuseDir, bApplySkyLighting, ResolvedView.SkyLightColor.rgb, DiffuseIndirectLighting, SubsurfaceIndirectLighting, IndirectIrradiance, WaterDiffuseIndirectLuminance);
+
 	// Apply MaterialAO since we don't have the DiffuseIndirectComposite pass on mobile deferred.
-	ComputeIndirect(LightmapVTPageTableResult, Interpolants, ShadingModelContext.DiffuseDir, DiffuseColorForIndirect, MaterialAO, IndirectIrradiance, MobileLightAccumulator);
+	IndirectIrradiance *= GBuffer.GBufferAO;
+
+	half3 DiffuseColor = (DiffuseIndirectLighting * DiffuseColorForIndirect + SubsurfaceIndirectLighting * SubsurfaceColor) * AOMultiBounce(GBuffer.BaseColor, ShadingOcclusion.DiffOcclusion);
+
+#if MATERIAL_SHADINGMODEL_SINGLELAYERWATER
+	// Fade out diffuse as this will be handled by the single scattering lighting in water material.
+	// We do this after the just above GetPrecomputedIndirectLightingAndSkyLight to keep ambiant lighting avialable.
+	// We also keep the SpecularColor for sun/water interactions.
+	GBuffer.DiffuseColor *= BaseMaterialCoverageOverWater;
+	DiffuseColor *= BaseMaterialCoverageOverWater;
+#endif
+
+	FLightAccumulator DirectLighting = (FLightAccumulator)0;
+	LightAccumulator_AddSplit(DirectLighting, DiffuseColor, 0.0f, DiffuseColor, 1.0f, false);
 
 	half Shadow = GetPrimaryPrecomputedShadowMask(LightmapVTPageTableResult, Interpolants, MaterialParameters).r;
 
@@ -409,29 +560,24 @@ void Main(
 	MobileEncodeGBuffer(GBuffer, OutGBufferA, OutGBufferB, OutGBufferC, OutGBufferD);
 #else
 
+	half3 CameraVector = -MaterialParameters.CameraVector;
+
 	half4 DynamicShadowFactors = 1.0f;
-	half NoL = 0.0f;
 
 	// Directional light
-	AccumulateDirectionalLighting(GBuffer, ShadingModelContext, MaterialParameters.CameraVector, MaterialParameters.ScreenPosition, SvPosition, Shadow, DynamicShadowFactors, NoL, MobileLightAccumulator);
-
-	// Skylight
-	bool bApplySkyLighting = !(MATERIAL_TWOSIDED && LQ_TEXTURE_LIGHTMAP) || NoL == 0;
-	AccumulateSkyLighting(GBuffer, ShadingModelContext, ResolvedView.SkyLightColor.rgb, bApplySkyLighting, IndirectIrradiance, MobileLightAccumulator);
+	AccumulateDirectionalLighting(GBuffer, MaterialParameters.WorldPosition_CamRelative, CameraVector, MaterialParameters.ScreenPosition, SvPosition, Shadow, DynamicShadowFactors, DirectLighting);
 
 	float2 LocalPosition = SvPosition.xy - ResolvedView.ViewRectMin.xy;
 	uint GridIndex = ComputeLightGridCellIndex(uint2(LocalPosition.x, LocalPosition.y), SvPosition.w);
 
 	// Reflection IBL
-#if !(MATERIAL_SINGLE_SHADINGMODEL && MATERIAL_SHADINGMODEL_HAIR)
 	AccumulateReflection(GBuffer
-		, ShadingModelContext
+		, CameraVector
 		, MaterialParameters.WorldPosition_CamRelative
 		, MaterialParameters.ReflectionVector
 		, IndirectIrradiance
 		, GridIndex
-		, MobileLightAccumulator);
-#endif
+		, DirectLighting);
 
 	// Local lights
 #if  !MATERIAL_SHADINGMODEL_SINGLELAYERWATER && ENABLE_CLUSTERED_LIGHTS
@@ -445,7 +591,7 @@ void Main(
 		#endif
 		
 		uint LightingChannelMask = GetPrimitive_LightingChannelMask(MaterialParameters.PrimitiveId);
-		AccumulateLightGridLocalLighting(CulledLightGridData, ShadingModelContext, GBuffer, MaterialParameters.WorldPosition_CamRelative, MaterialParameters.CameraVector, EyeIndex, 0, LocalLightDynamicShadowFactors, LightingChannelMask, MobileLightAccumulator);
+		AccumulateLightGridLocalLighting(CulledLightGridData, GBuffer, MaterialParameters.WorldPosition_CamRelative, CameraVector, EyeIndex, 0, LocalLightDynamicShadowFactors, LightingChannelMask, DirectLighting);
 	}
 #endif
 
@@ -456,8 +602,11 @@ void Main(
 
 		const bool CameraIsUnderWater = false;	// Fade out the material contribution over to water contribution according to material opacity.
 		const float3 SunIlluminance = ResolvedView.DirectionalLightColor.rgb * PI;			// times PI because it is divided by PI on CPU (=luminance) and we want illuminance here. 
-		const float3 WaterDiffuseIndirectIlluminance = ShadingModelContext.WaterDiffuseIndirectLuminance * PI;	// DiffuseIndirectLighting is luminance. So we need to multiply by PI to get illuminance.
-		const float3 EnvBrdf = ShadingModelContext.EnvBrdf;
+		const float3 WaterDiffuseIndirectIlluminance = WaterDiffuseIndirectLuminance * PI;	// DiffuseIndirectLighting is luminance. So we need to multiply by PI to get illuminance.
+		half3 N = GBuffer.WorldNormal;
+		half3 V = MaterialParameters.CameraVector;
+		half NoV = saturate(abs(dot(N, V)) + 1e-5);
+		const float3 EnvBrdf = GetEnvBRDF(GBuffer.SpecularColor, GBuffer.Roughness, NoV);
 		const uint EyeIndex = 0;
 
 		const float4 NullDistortionParams = 1.0f;
@@ -465,19 +614,19 @@ void Main(
 			MaterialParameters, PixelMaterialInputs, ResolvedView,
 			Shadow, GBuffer.Specular, NullDistortionParams,
 			SunIlluminance, WaterDiffuseIndirectIlluminance, EnvBrdf,
-			CameraIsUnderWater, ShadingModelContext.WaterVisibility, EyeIndex,
+			CameraIsUnderWater, WaterVisibility, EyeIndex,
 			bSeparateMainDirLightLuminance, SeparatedWaterMainDirLightLuminance);
 
 		// Accumulate luminance and occlude the background according to transmittance to view and mean transmittance to lights.
-		MobileLightAccumulator_Add(MobileLightAccumulator, WaterLighting.Luminance, 0.0f, 1.0f);
-		ShadingModelContext.Opacity = 1.0 - ((1.0 - ShadingModelContext.Opacity) * dot(WaterLighting.WaterToSceneToLightTransmittance, float3(1.0 / 3.0, 1.0 / 3.0, 1.0 / 3.0)));
+		LightAccumulator_AddSplit(DirectLighting, WaterLighting.Luminance, 0.0f, WaterLighting.Luminance, 1.0f, false);
+		Opacity = 1.0 - ((1.0 - Opacity) * dot(WaterLighting.WaterToSceneToLightTransmittance, float3(1.0 / 3.0, 1.0 / 3.0, 1.0 / 3.0)));
 	}
 #endif // MATERIAL_SHADINGMODEL_SINGLELAYERWATER
 
 #endif// DEFERRED_SHADING_PATH
 
 #if !MATERIAL_SHADINGMODEL_UNLIT
-	half3 Color = MobileLightAccumulator_GetResult(MobileLightAccumulator);
+	half3 Color = DirectLighting.TotalLight;
 #else
 	half3 Color = 0.0f;
 #endif
@@ -500,22 +649,22 @@ void Main(
 	Color += Emissive;
 
 #if !MATERIAL_SHADINGMODEL_UNLIT && MOBILE_EMULATION
-	Color = lerp(Color, ShadingModelContext.DiffuseColor, ResolvedView.UnlitViewmodeMask);
+	Color = lerp(Color, DiffuseColor, ResolvedView.UnlitViewmodeMask);
 #endif
 
 	// On mobile, water (an opaque material) is rendered as trnaslucent with forced premultiplied alpha blending (see MobileBasePass::SetTranslucentRenderState)
 	#if MATERIALBLENDING_ALPHACOMPOSITE || MATERIAL_SHADINGMODEL_SINGLELAYERWATER
-		OutColor = half4(Color * VertexFog.a + VertexFog.rgb * ShadingModelContext.Opacity, ShadingModelContext.Opacity);
+		OutColor = half4(Color * VertexFog.a + VertexFog.rgb * Opacity, Opacity);
 	#elif MATERIALBLENDING_ALPHAHOLDOUT
 		// not implemented for holdout
-		OutColor = half4(Color * VertexFog.a + VertexFog.rgb * ShadingModelContext.Opacity, ShadingModelContext.Opacity);
+		OutColor = half4(Color * VertexFog.a + VertexFog.rgb * Opacity, Opacity);
 	#elif MATERIALBLENDING_TRANSLUCENT
-		OutColor = half4(Color * VertexFog.a + VertexFog.rgb, ShadingModelContext.Opacity);
+		OutColor = half4(Color * VertexFog.a + VertexFog.rgb, Opacity);
 	#elif MATERIALBLENDING_ADDITIVE
-		OutColor = half4(Color * (VertexFog.a * ShadingModelContext.Opacity.x), 0.0f);
+		OutColor = half4(Color * (VertexFog.a * Opacity.x), 0.0f);
 	#elif MATERIALBLENDING_MODULATE
 		half3 FoggedColor = lerp(half3(1, 1, 1), Color, VertexFog.aaa * VertexFog.aaa);
-		OutColor = half4(FoggedColor, ShadingModelContext.Opacity);
+		OutColor = half4(FoggedColor, Opacity);
 	#else
 		OutColor.rgb = Color * VertexFog.a + VertexFog.rgb;
 
@@ -537,12 +686,7 @@ void Main(
 		OutColor.rgb *= ResolvedView.PreExposure;
 	#endif
 
-	#if MATERIAL_IS_SKY
-		// Sky materials can result in high luminance values, e.g. the sun disk. 
-		// This is so we make sure to at least stay within the boundaries of fp10 and not cause NaN on some platforms.
-		// We also half that range to also make sure we have room for other additive elements such as bloom, clouds or particle visual effects.
-		OutColor.rgb = min(OutColor.rgb, Max10BitsFloat.xxx * 0.5f);
-	#endif
+	OutColor.rgb = SafeGetOutColor(OutColor.rgb);
 
 #if USE_EDITOR_COMPOSITING && (MOBILE_EMULATION)
 	// Editor primitive depth testing
@@ -562,7 +706,7 @@ void Main(
 	FinalizeVirtualTextureFeedback(
 		MaterialParameters.VirtualTextureFeedback,
 		MaterialParameters.SvPosition,
-		ShadingModelContext.Opacity,
+		Opacity,
 		View.FrameNumber,
 		View.VTFeedbackBuffer
 	);
diff --git a/Engine/Shaders/Private/MobileDeferredShading.usf b/Engine/Shaders/Private/MobileDeferredShading.usf
--- a/Engine/Shaders/Private/MobileDeferredShading.usf
+++ b/Engine/Shaders/Private/MobileDeferredShading.usf
@@ -18,7 +18,6 @@
 #define PreIntegratedGFSampler	ReflectionStruct.PreIntegratedGFSampler
 
 #include "MobileLightingCommon.ush"
-#include "IESLightProfilesCommon.ush"
 #if MATERIAL_SHADER
 #include "LightFunctionCommon.ush"
 #endif
@@ -110,25 +109,21 @@ half ComputeLightFunctionMultiplier(float3 TranslatedWorldPosition)
 }
 
 void ReflectionEnvironmentSkyLighting(
-	FMobileShadingModelContext ShadingModelContext,
 	FGBufferData GBuffer,
+	half3 CameraVector,
 	float3 TranslatedWorldPosition,
 	half3 ReflectionVector,
 	uint GridIndex,
-	inout FMobileLightAccumulator MobileLightAccumulator)
+	inout FLightAccumulator DirectLighting)
 {
-	half IndirectIrradiance = GBuffer.IndirectIrradiance;
-	// Skylight
-	AccumulateSkyLighting(GBuffer, ShadingModelContext, View.SkyLightColor.rgb, true, IndirectIrradiance, MobileLightAccumulator);
-
 	// IBL
 	AccumulateReflection(GBuffer
-		, ShadingModelContext
+		, CameraVector
 		, TranslatedWorldPosition
 		, ReflectionVector
-		, IndirectIrradiance
+		, GBuffer.IndirectIrradiance
 		, GridIndex
-		, MobileLightAccumulator);
+		, DirectLighting);
 }
   
 void MobileDirectionalLightPS(
@@ -150,18 +145,16 @@ void MobileDirectionalLightPS(
 	float2 ScreenPos = UVAndScreenPos.zw;
 	float3 TranslatedWorldPosition = mul(float4(ScreenPos * GBuffer.Depth, GBuffer.Depth, 1), PrimaryView.ScreenToTranslatedWorld).xyz;
 
-	half3 CameraVector = normalize(PrimaryView.TranslatedWorldCameraOrigin - TranslatedWorldPosition);
-	half NoV = dot(GBuffer.WorldNormal, CameraVector);
-	half3 ReflectionVector = GBuffer.WorldNormal * (NoV * 2.0) - CameraVector;
-	NoV = max(0, NoV);
+	half3 CameraVector = normalize(TranslatedWorldPosition);
+	half3 V = -CameraVector;
+	half NoV = dot(GBuffer.WorldNormal, V);
+	half3 ReflectionVector = GBuffer.WorldNormal * (NoV * 2.0) - V;
+
+	FLightAccumulator DirectLighting = (FLightAccumulator)0;
 
-	FMobileLightAccumulator MobileLightAccumulator = (FMobileLightAccumulator)0;
 	// Check movable light param to determine if we should be using precomputed shadows
 	half Shadow = LightFunctionParameters2.z > 0.0f ? 1.0f : GBuffer.PrecomputedShadowFactors.r;
 	float4 ScreenPosition = SvPositionToScreenPosition(float4(SvPosition.xyz, GBuffer.Depth));
-
-	FMobileShadingModelContext ShadingModelContext = (FMobileShadingModelContext)0;
-	InitShadingModelContext(ShadingModelContext, GBuffer, CameraVector);
 	
 	float2 LocalPosition = SvPosition.xy - View.ViewRectMin.xy;
 	uint GridIndex = ComputeLightGridCellIndex(uint2(LocalPosition.x, LocalPosition.y), GBuffer.Depth);
@@ -171,29 +164,28 @@ void MobileDirectionalLightPS(
 		const uint EyeIndex = 0;
 		const FCulledLightsGridData CulledLightGridData = GetCulledLightsGrid(GridIndex, EyeIndex);
 		half4 LocalLightDynamicShadowFactors = 1.0f;
-		AccumulateLightGridLocalLighting(CulledLightGridData, ShadingModelContext, GBuffer, TranslatedWorldPosition, CameraVector, EyeIndex, 0, LocalLightDynamicShadowFactors, MobileLightAccumulator);
+		AccumulateLightGridLocalLighting(CulledLightGridData, GBuffer, TranslatedWorldPosition, CameraVector, EyeIndex, 0, LocalLightDynamicShadowFactors, DirectLighting);
 	}
 #endif
 
 	// Directional light
 	half4 DynamicShadowFactors = 1.0f;
-	half NoL = 0.0f;
-	AccumulateDirectionalLighting(GBuffer, ShadingModelContext, CameraVector, ScreenPosition, SvPosition, Shadow, DynamicShadowFactors, NoL, MobileLightAccumulator);
+	AccumulateDirectionalLighting(GBuffer, TranslatedWorldPosition, CameraVector, ScreenPosition, SvPosition, Shadow, DynamicShadowFactors, DirectLighting);
 
 #if ENABLE_CLUSTERED_REFLECTION || ENABLE_SKY_LIGHT || ENABLE_PLANAR_REFLECTION
 	// If we have a single directional light, apply relfection and sky contrubution here
-	ReflectionEnvironmentSkyLighting(ShadingModelContext, GBuffer, TranslatedWorldPosition, ReflectionVector, GridIndex, MobileLightAccumulator);
+	ReflectionEnvironmentSkyLighting(GBuffer, CameraVector, TranslatedWorldPosition, ReflectionVector, GridIndex, DirectLighting);
 #endif
 
-	half3 Color = MobileLightAccumulator_GetResult(MobileLightAccumulator);
-	half LightAttenuation = ComputeLightFunctionMultiplier(TranslatedWorldPosition);
+	half3 Color = DirectLighting.TotalLight * ComputeLightFunctionMultiplier(TranslatedWorldPosition);
 	// MobileHDR applies PreExposure in tonemapper
-	LightAttenuation *= View.PreExposure;
+	Color *= View.PreExposure;
 
+	Color = SafeGetOutColor(Color);
 #if defined(MOBILE_DEFERRED_SHADING) && USE_GLES_FBF_DEFERRED == 1
-	OutProxyAdditive.rgb = Color.rgb * LightAttenuation;
+	OutProxyAdditive.rgb = Color;
 #else
-	OutColor.rgb = Color.rgb * LightAttenuation;
+	OutColor.rgb = Color;
 	OutColor.a = 1;
 #endif
 }
@@ -207,6 +199,46 @@ SamplerState	LocalLightShadowSampler;
 float4			LocalLightShadowBufferSize;
 #endif
 
+FDeferredLightData SetupLightDataForStandardDeferred()
+{
+	// Build the light data struct using the DeferredLightUniforms and light defines
+	// We are heavily relying on the shader compiler to optimize out constant subexpressions in GetDynamicLighting()
+
+	FDeferredLightData LightData = (FDeferredLightData)0;
+	LightData.TranslatedWorldPosition = GetDeferredLightTranslatedWorldPosition();
+	LightData.InvRadius = DeferredLightUniforms.InvRadius;
+	LightData.Color = DeferredLightUniforms.Color;
+	LightData.FalloffExponent = DeferredLightUniforms.FalloffExponent;
+	LightData.Direction = DeferredLightUniforms.Direction;
+//	LightData.Tangent = DeferredLightUniforms.Tangent;
+	LightData.SpotAngles = DeferredLightUniforms.SpotAngles;
+//	LightData.SourceRadius = DeferredLightUniforms.SourceRadius;
+//	LightData.SourceLength = DeferredLightUniforms.SourceLength;
+//	LightData.SoftSourceRadius = DeferredLightUniforms.SoftSourceRadius;
+	LightData.SpecularScale = DeferredLightUniforms.SpecularScale;
+//	LightData.ContactShadowLength = abs(DeferredLightUniforms.ContactShadowLength);
+//	LightData.ContactShadowLengthInWS = DeferredLightUniforms.ContactShadowLength < 0.0f;
+//	LightData.ContactShadowNonShadowCastingIntensity = DeferredLightUniforms.ContactShadowNonShadowCastingIntensity;
+//	LightData.DistanceFadeMAD = DeferredLightUniforms.DistanceFadeMAD;
+//	LightData.ShadowMapChannelMask = DeferredLightUniforms.ShadowMapChannelMask;
+	LightData.ShadowedBits = DeferredLightUniforms.ShadowedBits;
+	LightData.bInverseSquared = DeferredLightUniforms.FalloffExponent == 0; // Directional lights don't use 'inverse squared attenuation'
+	LightData.bRadialLight = true;
+	//@todo - permutation opportunity
+	LightData.bSpotLight = IS_SPOT_LIGHT;
+	LightData.bRectLight = false;
+	
+//	LightData.RectLightBarnCosAngle = DeferredLightUniforms.RectLightBarnCosAngle;
+//	LightData.RectLightBarnLength = DeferredLightUniforms.RectLightBarnLength;
+//	LightData.RectLightAtlasMaxLevel = DeferredLightUniforms.RectLightAtlasMaxLevel;
+//	LightData.RectLightAtlasUVOffset = DeferredLightUniforms.RectLightAtlasUVOffset;
+//	LightData.RectLightAtlasUVScale  = DeferredLightUniforms.RectLightAtlasUVScale;
+
+	LightData.HairTransmittance = InitHairTransmittanceData();
+
+	return LightData;
+}
+
 void MobileRadialLightPS(
 	float4 InScreenPosition : TEXCOORD0,
 	float4 SVPos			: SV_POSITION,
@@ -228,16 +260,26 @@ void MobileRadialLightPS(
 	// With an orthographic projection, clip space is the same as NDC
 	float2 ClipPosition = InScreenPosition.xy / InScreenPosition.w * (View.ViewToClip[3][3] < 1.0f ? GBuffer.Depth : 1.0f);
 	float3 TranslatedWorldPosition = mul(float4(ClipPosition, GBuffer.Depth, 1), PrimaryView.ScreenToTranslatedWorld).xyz;
-	half3 CameraVector = normalize(PrimaryView.TranslatedWorldCameraOrigin - TranslatedWorldPosition);
-	half NoV = max(0, dot(GBuffer.WorldNormal, CameraVector));
-
-	FMobileShadingModelContext ShadingModelContext = (FMobileShadingModelContext)0;
-	InitShadingModelContext(ShadingModelContext, GBuffer, CameraVector);
+	half3 CameraVector = normalize(TranslatedWorldPosition);
 			
-	FMobileLightAccumulator MobileLightAccumulator = (FMobileLightAccumulator)0;
+	FLightAccumulator DirectLighting = (FLightAccumulator)0;
+	FDeferredLightData LightData = SetupLightDataForStandardDeferred();
+	
+	// Affect the light color by any other shadow/transmittance before the lighting is evaluated
+	{
+		half Attenuation = ComputeLightProfileMultiplier(TranslatedWorldPosition, DeferredLightUniforms.TranslatedWorldPosition, -DeferredLightUniforms.Direction, DeferredLightUniforms.Tangent);
+		LightData.Color *= Attenuation;
+	}
+
+	half Shadow = 1.0;
+#if SUPPORT_SPOTLIGHTS_SHADOW && IS_SPOT_LIGHT
+	float3 ToLight = LightData.TranslatedWorldPosition - TranslatedWorldPosition;
+	float3 LocalPosition = -ToLight;
+	float DistanceSqr = dot(ToLight, ToLight);
+	half3 L = ToLight * rsqrt(DistanceSqr);
+	half3 N = GBuffer.WorldNormal;
+	half NoL = saturate(dot(N, L));
 
-	FMobileShadowInfo ShadowInfo;
-#if SUPPORT_SPOTLIGHTS_SHADOW
 	FPCFSamplerSettings Settings;
 	Settings.ShadowDepthTexture = LocalLightShadowTexture;
 	Settings.ShadowDepthTextureSampler = LocalLightShadowSampler;
@@ -247,45 +289,46 @@ void MobileRadialLightPS(
 	Settings.DensityMulConstant = 0;
 	Settings.ProjectionDepthBiasParameters = 0;
 	
-	ShadowInfo.Settings = Settings;
-	ShadowInfo.SpotLightShadowSharpen = SpotLightShadowSharpenAndFadeFractionAndReceiverDepthBiasAndSoftTransitionScale.x;
-	ShadowInfo.SpotLightFadeFraction = SpotLightShadowSharpenAndFadeFractionAndReceiverDepthBiasAndSoftTransitionScale.y;
-	ShadowInfo.SpotLightReceiverDepthBias = SpotLightShadowSharpenAndFadeFractionAndReceiverDepthBiasAndSoftTransitionScale.z;
-	ShadowInfo.SpotLightSoftTransitionScale = SpotLightShadowSharpenAndFadeFractionAndReceiverDepthBiasAndSoftTransitionScale.w;
-	ShadowInfo.SpotLightShadowmapMinMax = SpotLightShadowmapMinMax;
-	ShadowInfo.SpotLightShadowWorldToShadowMatrix = SpotLightShadowWorldToShadowMatrix;
+	float SpotLightShadowSharpen = SpotLightShadowSharpenAndFadeFractionAndReceiverDepthBiasAndSoftTransitionScale.x;
+	float SpotLightFadeFraction = SpotLightShadowSharpenAndFadeFractionAndReceiverDepthBiasAndSoftTransitionScale.y;
+	float SpotLightReceiverDepthBias = SpotLightShadowSharpenAndFadeFractionAndReceiverDepthBiasAndSoftTransitionScale.z;
+	float SpotLightSoftTransitionScale = SpotLightShadowSharpenAndFadeFractionAndReceiverDepthBiasAndSoftTransitionScale.w;
+
+	float4 HomogeneousShadowPosition = mul(float4(LocalPosition, 1), SpotLightShadowWorldToShadowMatrix);
+	float2 ShadowUVs = HomogeneousShadowPosition.xy / HomogeneousShadowPosition.w;
+	if (all(ShadowUVs >= SpotLightShadowmapMinMax.xy && ShadowUVs <= SpotLightShadowmapMinMax.zw))
+	{
+		// Clamp pixel depth in light space for shadowing opaque, because areas of the shadow depth buffer that weren't rendered to will have been cleared to 1
+		// We want to force the shadow comparison to result in 'unshadowed' in that case, regardless of whether the pixel being shaded is in front or behind that plane
+		// Invert ShadowZ as the shadow space has been changed (but not yet the filtering code)
+		float ShadowZ = 1.0f - HomogeneousShadowPosition.z;
+		float LightSpacePixelDepthForOpaque = min(ShadowZ, 0.99999f);
+		Settings.SceneDepth = LightSpacePixelDepthForOpaque;
+		Settings.TransitionScale = SpotLightSoftTransitionScale * lerp(SpotLightReceiverDepthBias, 1.0, NoL);
+
+		Shadow = MobileShadowPCF(ShadowUVs, Settings);
+
+		Shadow = saturate((Shadow - 0.5) * SpotLightShadowSharpen + 0.5f);
+
+		Shadow = lerp(1.0f, Square(Shadow), SpotLightFadeFraction);
+	}
 #endif
 
-	AccumulateLocalLighting(GBuffer,
-		ShadingModelContext,
-		TranslatedWorldPosition,
-		CameraVector,
-		DeferredLightUniforms.TranslatedWorldPosition,
-		DeferredLightUniforms.InvRadius,
-		DeferredLightUniforms.Color,
-		DeferredLightUniforms.FalloffExponent,
-		DeferredLightUniforms.Direction,
-		1.0f,
-		DeferredLightUniforms.SpotAngles,
-		DeferredLightUniforms.FalloffExponent == 0,
-		IS_SPOT_LIGHT,
-		DeferredLightUniforms.ShadowedBits > 0,
-		ShadowInfo,
-		MobileLightAccumulator
-		);
-
-	half3 Color = MobileLightAccumulator_GetResult(MobileLightAccumulator);
-	
-	half LightAttenuation = ComputeLightProfileMultiplier(TranslatedWorldPosition, DeferredLightUniforms.TranslatedWorldPosition, -DeferredLightUniforms.Direction, DeferredLightUniforms.Tangent);
-	LightAttenuation*= ComputeLightFunctionMultiplier(TranslatedWorldPosition);
+	half4 LightAttenuation = half4(1, 1, Shadow, Shadow);
 
+	float SurfaceShadow = 0;
+	DirectLighting = GetDynamicLightingSplit(TranslatedWorldPosition, CameraVector, GBuffer, 1, GBuffer.ShadingModelID, LightData, LightAttenuation, 0, uint2(0, 0), SurfaceShadow);
+
+	half3 Color = DirectLighting.TotalLight * ComputeLightFunctionMultiplier(TranslatedWorldPosition);
+	
 	// MobileHDR applies PreExposure in tonemapper
-	LightAttenuation*= View.PreExposure;
+	Color *= View.PreExposure;
 
+	Color = SafeGetOutColor(Color);
 #if defined(MOBILE_DEFERRED_SHADING) && USE_GLES_FBF_DEFERRED == 1
-	OutProxyAdditive.rgb = Color.rgb * LightAttenuation;
+	OutProxyAdditive.rgb = Color;
 #else
-	OutColor.rgb = Color.rgb * LightAttenuation;
+	OutColor.rgb = Color;
 	OutColor.a = 1;
 #endif
 }
@@ -307,19 +350,17 @@ void MobileReflectionEnvironmentSkyLightingPS(
 	FGBufferData GBuffer = FetchAndDecodeGBuffer(UVAndScreenPos.xy);
 	float2 ScreenPos = UVAndScreenPos.zw;
 	float3 TranslatedWorldPosition = mul(float4(ScreenPos * GBuffer.Depth, GBuffer.Depth, 1), PrimaryView.ScreenToTranslatedWorld).xyz;
-	half3 CameraVector = normalize(PrimaryView.TranslatedWorldCameraOrigin - TranslatedWorldPosition);
-	half NoV = max(0, dot(GBuffer.WorldNormal, CameraVector));
-	half3 ReflectionVector = GBuffer.WorldNormal * (NoV * 2.0) - CameraVector;
-
-	FMobileShadingModelContext ShadingModelContext = (FMobileShadingModelContext)0;
-	InitShadingModelContext(ShadingModelContext, GBuffer, CameraVector);
+	half3 CameraVector = normalize(TranslatedWorldPosition);
+	half3 V = -CameraVector;
+	half NoV = max(0, dot(GBuffer.WorldNormal, V));
+	half3 ReflectionVector = GBuffer.WorldNormal * (NoV * 2.0) - V;
 
 	float2 LocalPosition = SvPosition.xy - View.ViewRectMin.xy;
 	uint GridIndex = ComputeLightGridCellIndex(uint2(LocalPosition.x, LocalPosition.y), GBuffer.Depth);
 	
-	FMobileLightAccumulator MobileLightAccumulator = (FMobileLightAccumulator)0;
-	ReflectionEnvironmentSkyLighting(ShadingModelContext, GBuffer, TranslatedWorldPosition, ReflectionVector, GridIndex, MobileLightAccumulator);
-	half3 ReflectionAndSky = MobileLightAccumulator_GetResult(MobileLightAccumulator);
+	FLightAccumulator DirectLighting = (FLightAccumulator)0;
+	ReflectionEnvironmentSkyLighting(GBuffer, CameraVector, TranslatedWorldPosition, ReflectionVector, GridIndex, DirectLighting);
+	half3 ReflectionAndSky = DirectLighting.TotalLight;
 
 	ReflectionAndSky *= View.PreExposure;
 
diff --git a/Engine/Shaders/Private/MobileLightingCommon.ush b/Engine/Shaders/Private/MobileLightingCommon.ush
--- a/Engine/Shaders/Private/MobileLightingCommon.ush
+++ b/Engine/Shaders/Private/MobileLightingCommon.ush
@@ -2,11 +2,15 @@
 
 #pragma once
 
-#include "MobileShadingModels.ush"
+#define SUPPORT_CONTACT_SHADOWS 0
+
+#include "ShadingModels.ush"
 #include "ReflectionEnvironmentShared.ush"
 #include "PlanarReflectionShared.ush"
-#include "DynamicLightingCommon.ush"
 #include "LightGridCommon.ush"
+#include "ShadowFilteringCommon.ush"
+#include "ClearCoatCommon.ush"
+#include "DeferredLightingCommon.ush"
 
 #ifndef TRANSLUCENCY_NON_DIRECTIONAL
 #define TRANSLUCENCY_NON_DIRECTIONAL 0
@@ -30,27 +34,44 @@
 #define REFLECTION_COMPOSITE_HAS_BOX_CAPTURES 1
 #include "ReflectionEnvironmentComposite.ush"
 
-/*------------------------------------------------------------------------------
-	Mobile Light Accumulation
-------------------------------------------------------------------------------*/
-struct FMobileLightAccumulator
+half GetSurfaceShadow(FGBufferData GBuffer, FShadowTerms ShadowTerms)
 {
-	float3 TotalLight;
-};
+	if (GBuffer.ShadingModelID == SHADINGMODELID_SUBSURFACE_PROFILE)
+	{
+		return 1.0f;
+	}
+	return ShadowTerms.SurfaceShadow;
+}
 
-// accumulate light, can be called multiple times
-void MobileLightAccumulator_Add(inout FMobileLightAccumulator In, half3 DiffuseTotalLight, half3 SpecularTotalLight, half3 CommonMultiplier)
+half GetTransmissionShadow(FGBufferData GBuffer, FShadowTerms ShadowTerms)
 {
-#if FULLY_ROUGH
-	In.TotalLight += DiffuseTotalLight * CommonMultiplier;
-#else
-	In.TotalLight += (DiffuseTotalLight + SpecularTotalLight) * CommonMultiplier;
-#endif
+	if (GBuffer.ShadingModelID == GBuffer.ShadingModelID == SHADINGMODELID_EYE)
+	{
+		return 1.0f;
+	}
+	return ShadowTerms.TransmissionShadow;
+}
+
+half3 SafeGetOutColor(half3 OutColor)
+{
+	// Sky materials can result in high luminance values, e.g. the sun disk. 
+	// The PreExposure could double the OutColor
+	// This is so we make sure to at least stay within the boundaries of fp10 and not cause NaN on some platforms.
+	// We also half that range to also make sure we have room for other additive elements such as bloom, clouds or particle visual effects.
+	OutColor = min(OutColor, Max111110BitsFloat3 * 0.5f);
+	return OutColor;
 }
 
-half3 MobileLightAccumulator_GetResult(FMobileLightAccumulator In)
+FLightAccumulator LightAccumulator_Add(FLightAccumulator A, FLightAccumulator B)
 {
-	return min(65000.0, In.TotalLight);
+	FLightAccumulator Sum = (FLightAccumulator)0;
+	Sum.TotalLight = A.TotalLight + B.TotalLight;
+	Sum.ScatterableLightLuma = A.ScatterableLightLuma + B.ScatterableLightLuma;
+	Sum.ScatterableLight = A.ScatterableLight + B.ScatterableLight;
+	Sum.EstimatedCost = A.EstimatedCost + B.EstimatedCost;
+	Sum.TotalLightDiffuse = A.TotalLightDiffuse + B.TotalLightDiffuse;
+	Sum.TotalLightSpecular = A.TotalLightSpecular + B.TotalLightSpecular;
+	return Sum;
 }
 
 /*------------------------------------------------------------------------------
@@ -103,7 +124,7 @@ return ShadowMap;
 #define ENABLE_MOBILE_CSM 1
 #endif
 
-half MobileDirectionalLightCSM(float2 ScreenPosition, float SceneDepth, half NoL, inout float ShadowPositionZ)
+half MobileDirectionalLightCSM(float2 ScreenPosition, float SceneDepth, inout float ShadowPositionZ)
 {
 	half ShadowMap = 1;
 #if ENABLE_MOBILE_CSM
@@ -111,7 +132,7 @@ half MobileDirectionalLightCSM(float2 ScreenPosition, float SceneDepth, half NoL
 	FPCFSamplerSettings Settings;
 	Settings.ShadowDepthTexture = MobileDirectionalLight.DirectionalLightShadowTexture;
 	Settings.ShadowDepthTextureSampler = MobileDirectionalLight.DirectionalLightShadowSampler;
-	Settings.TransitionScale = MobileDirectionalLight.DirectionalLightDirectionAndShadowTransition.w * lerp(MobileDirectionalLight.DirectionalLightDistanceFadeMADAndSpecularScale.w, 1.0, NoL);
+	Settings.TransitionScale = MobileDirectionalLight.DirectionalLightDirectionAndShadowTransition.w;
 	Settings.ShadowBufferSize = MobileDirectionalLight.DirectionalLightShadowSize;
 	Settings.bSubsurface = false;
 	Settings.bTreatMaxDepthUnshadowed = false;
@@ -161,14 +182,8 @@ half MobileDirectionalLightCSM(float2 ScreenPosition, float SceneDepth, half NoL
 	Mobile Directional Light.
 ------------------------------------------------------------------------------*/
 
-void AccumulateDirectionalLighting(FGBufferData GBuffer, FMobileShadingModelContext ShadingModelContext, half3 CameraVector, float4 ScreenPosition, float4 SvPosition, half Shadow, inout half4 DynamicShadowFactors, out half NoL, inout FMobileLightAccumulator MobileLightAccumulator)
+void AccumulateDirectionalLighting(FGBufferData GBuffer, float3 TranslatedWorldPosition, half3 CameraVector, float4 ScreenPosition, float4 SvPosition, half Shadow, inout half4 DynamicShadowFactors, inout FLightAccumulator DirectLighting)
 {
-	half3 L = MobileDirectionalLight.DirectionalLightDirectionAndShadowTransition.xyz;
-	NoL = max(0, dot(GBuffer.WorldNormal, L));
-	float3 H = normalize(float3(CameraVector + L));
-	half NoH = max(0, dot(GBuffer.WorldNormal, H));
-	half VoH = max(0, dot(CameraVector, H));
-
 #if USE_SHADOWMASKTEXTURE
 	half4 PreviewShadowMapChannelMask = half4(1.0f, 0.0f, 0.0f, 0.0f);
 	DynamicShadowFactors = Texture2DSample(MobileBasePass.ScreenSpaceShadowMaskTexture, MobileBasePass.ScreenSpaceShadowMaskSampler, SvPositionToBufferUV(SvPosition));
@@ -184,7 +199,7 @@ void AccumulateDirectionalLighting(FGBufferData GBuffer, FMobileShadingModelCont
 #else
 
 	float  ShadowPositionZ = 0;
-	half ShadowMap = MobileDirectionalLightCSM(ScreenPosition.xy, ScreenPosition.w, NoL, ShadowPositionZ);
+	half ShadowMap = MobileDirectionalLightCSM(ScreenPosition.xy, ScreenPosition.w, ShadowPositionZ);
 
 	#if !MOBILE_DEFERRED_SHADING && DIRECTIONAL_LIGHT_CSM && !MATERIAL_SHADINGMODEL_SINGLELAYERWATER
 		// Cascaded Shadow Map
@@ -197,48 +212,20 @@ void AccumulateDirectionalLighting(FGBufferData GBuffer, FMobileShadingModelCont
 	#endif
 #endif
 
-#if TRANSLUCENCY_NON_DIRECTIONAL
-	NoL = 1.f;
-	ShadingModelContext.SpecularColor = 0.f;
-#endif
-
-	// Direct Lighting (Directional light) + IBL
-	FMobileDirectLighting Lighting = MobileIntegrateBxDF(ShadingModelContext, GBuffer, NoL, NoH, VoH, CameraVector, L, Shadow);
-	// MobileDirectionalLight.DirectionalLightDistanceFadeMADAndSpecularScale.z saves SpecularScale for direction light.
-	MobileLightAccumulator_Add(MobileLightAccumulator, Lighting.Diffuse, Lighting.Specular * MobileDirectionalLight.DirectionalLightDistanceFadeMADAndSpecularScale.z, MobileDirectionalLight.DirectionalLightColor.rgb);
-}
-
-/*------------------------------------------------------------------------------
-	Mobile SkyLight.
-------------------------------------------------------------------------------*/
-void AccumulateSkyLighting(FGBufferData GBuffer, FMobileShadingModelContext ShadingModelContext, half3 SkyLightColor, bool bApplySkyLighting, inout half IndirectIrradiance, inout FMobileLightAccumulator MobileLightAccumulator)
-{
-#if ENABLE_SKY_LIGHT
-	//@mw todo
-	// TODO: Also need to do specular.
-	if (bApplySkyLighting)
-	{
-		half3 SkyDiffuseLighting = GetSkySHDiffuseSimple(GBuffer.WorldNormal);
-		half3 DiffuseLookup = SkyDiffuseLighting * SkyLightColor;
-		IndirectIrradiance += Luminance(DiffuseLookup);
-
-	#if MATERIAL_SHADINGMODEL_SINGLELAYERWATER
-		ShadingModelContext.WaterDiffuseIndirectLuminance += SkyDiffuseLighting;
-	#endif
+	FDeferredLightData LightData = (FDeferredLightData)0;
+	LightData.Color = MobileDirectionalLight.DirectionalLightColor.rgb;
+	LightData.FalloffExponent = 0;
+	LightData.Direction = MobileDirectionalLight.DirectionalLightDirectionAndShadowTransition.xyz;
+	LightData.bRadialLight = false;
+	LightData.SpecularScale = MobileDirectionalLight.DirectionalLightDistanceFadeMADAndSpecularScale.z;
+	LightData.ShadowedBits = 1;
+	LightData.HairTransmittance = InitHairTransmittanceData();
 
-		half3 DiffuseColorForSky = ShadingModelContext.DiffuseColor;
+	half4 LightAttenuation = half4(Shadow, Shadow, 1, 1);
 
-	#if MOBILE_SHADINGMODEL_SUPPORT
-		if (GBuffer.ShadingModelID == SHADINGMODELID_SUBSURFACE || GBuffer.ShadingModelID == SHADINGMODELID_PREINTEGRATED_SKIN)
-		{
-			// Add subsurface energy to diffuse
-			DiffuseColorForSky += ShadingModelContext.SubsurfaceColor;
-		}
-	#endif
-
-		MobileLightAccumulator_Add(MobileLightAccumulator, SkyDiffuseLighting, 0.0f, SkyLightColor * DiffuseColorForSky * GBuffer.GBufferAO);
-	}
-#endif
+	float SurfaceShadow = 0;
+	FLightAccumulator NewLighting = AccumulateDynamicLighting(TranslatedWorldPosition, CameraVector, GBuffer, 1, GBuffer.ShadingModelID, LightData, LightAttenuation, 0, uint2(0, 0), SurfaceShadow);
+	DirectLighting = LightAccumulator_Add(DirectLighting, NewLighting);
 }
 
 /*------------------------------------------------------------------------------
@@ -271,7 +258,7 @@ half3 GetImageBasedReflectionLighting(half3 ReflectionVector
 	uint DataStartIndex = ForwardLightData.NumCulledLightsGrid[NumCulledEntryIndex + 1];
 
 	SpecularIBL = CompositeReflectionCapturesAndSkylightTWS(
-		1.0f,
+		CompositeAlpha,
 		TranslatedWorldPosition,
 		ReflectionVector,//RayDirection,
 		Roughness,
@@ -292,6 +279,7 @@ half3 GetImageBasedReflectionLighting(half3 ReflectionVector
 	SpecularIBL *= ComputeMixingWeight(IndirectIrradiance, SkyAverageBrightness, Roughness);
 	#endif
 
+	SpecularIBL *= CompositeAlpha;
 #else // !(ENABLE_CLUSTERED_REFLECTION || MOBILE_DEFERRED_LIGHTING)
 
 	half UsingSkyReflection = MobileReflectionCapture.Params.y > 0.0f;
@@ -317,28 +305,42 @@ half3 GetImageBasedReflectionLighting(half3 ReflectionVector
 	SpecularIBL *= ComputeMixingWeight(IndirectIrradiance, MobileReflectionCapture.Params.x, Roughness);
 	#endif
 
-#endif
-
 	SpecularIBL *= CompositeAlpha;
+#endif
 
 	return SpecularIBL;
 }
 
 void AccumulateReflection(FGBufferData GBuffer
-	, FMobileShadingModelContext ShadingModelContext
+	, half3 CameraVector
 	, float3 TranslatedWorldPosition
 	, half3 ReflectionVector
 	, half IndirectIrradiance
 	, uint GridIndex
-	, inout FMobileLightAccumulator MobileLightAccumulator)
+	, inout FLightAccumulator DirectLighting)
 {
 	half3 SpecularIBLLighting = (half3)0.0f;
 
+	half3 N = GBuffer.WorldNormal;
+	half3 V = -CameraVector;
+	half NoV = saturate(abs(dot(N, V)) + 1e-5);
+
+	half3 TopLayerR = ReflectionVector;
+	half SpecularOcclusion = GBuffer.GBufferAO;
+	if (MOBILE_HIGH_QUALITY_BRDF)
+	{
+		// Point lobe in off-specular peak direction
+		ReflectionVector = GetOffSpecularPeakReflectionDir(N, ReflectionVector, GBuffer.Roughness);
+
+		half RoughnessSq = GBuffer.Roughness * GBuffer.Roughness;
+		SpecularOcclusion = GetSpecularOcclusion(NoV, RoughnessSq, SpecularOcclusion);
+	}
+
 	half3 SpecularIBL = GetImageBasedReflectionLighting(ReflectionVector
 		, TranslatedWorldPosition
 		, GBuffer.Roughness
 		, IndirectIrradiance
-		, 1.0f
+		, SpecularOcclusion
 		, GridIndex
 	);
 
@@ -352,197 +354,102 @@ void AccumulateReflection(FGBufferData GBuffer
 	}
 #endif
 
-#if MOBILE_SHADINGMODEL_SUPPORT
+	half3 DiffuseColor = GBuffer.DiffuseColor;
+	half3 SpecularColor = GBuffer.SpecularColor;
+
 	if (GBuffer.ShadingModelID == SHADINGMODELID_CLEAR_COAT)
 	{
-		half F = GetEnvBRDF(0.04, ShadingModelContext.ClearCoatRoughness, ShadingModelContext.NoV).x;
-		F *= ShadingModelContext.ClearCoat;
+		const half ClearCoat = GBuffer.CustomData.x;
+		const half ClearCoatRoughness = GBuffer.CustomData.y;
+
+		RemapClearCoatDiffuseAndSpecularColor(GBuffer, NoV, DiffuseColor, SpecularColor);
+
+		half F = GetEnvBRDF(0.04, ClearCoatRoughness, NoV).x;
+		F *= ClearCoat;
 		half LayerAttenuation = (1 - F);
 
 		// Fc * Vis
 	#if FORWARDSHADING_USE_HQ_ENV_BRDF
-		half2 AB = PreIntegratedGF.SampleLevel(PreIntegratedGFSampler, float2(ShadingModelContext.NoV, GBuffer.Roughness), 0).rg;
+		half2 AB = PreIntegratedGF.SampleLevel(PreIntegratedGFSampler, float2(NoV, GBuffer.Roughness), 0).rg;
 	#else
-		half2 AB = EnvBRDFApproxLazarov(GBuffer.Roughness, ShadingModelContext.NoV);
+		half2 AB = EnvBRDFApproxLazarov(GBuffer.Roughness, NoV);
 	#endif
-		SpecularIBLLighting += SpecularIBL * LayerAttenuation * (ShadingModelContext.SpecularColor * AB.x + AB.y * saturate(50 * ShadingModelContext.SpecularColor.g) * (1 - ShadingModelContext.ClearCoat));
+		SpecularIBLLighting += SpecularIBL * LayerAttenuation * (SpecularColor * AB.x + AB.y * saturate(50 * SpecularColor.g) * (1 - ClearCoat));
 
-		SpecularIBL = GetImageBasedReflectionLighting(ReflectionVector
+		SpecularIBL = GetImageBasedReflectionLighting(TopLayerR
 			, TranslatedWorldPosition
-			,ShadingModelContext.ClearCoatRoughness
+			, ClearCoatRoughness
 			, IndirectIrradiance
-			, F
+			, F * SpecularOcclusion
 			, GridIndex
 		);
 		SpecularIBLLighting += SpecularIBL;
 	}
+	else if (GBuffer.ShadingModelID == SHADINGMODELID_HAIR)
+	{
+		//Skip IBL for Hair
+	}
 	else
-#endif
 	{
-		SpecularIBLLighting += SpecularIBL * ShadingModelContext.EnvBrdf;
+		SpecularIBLLighting += SpecularIBL * GetEnvBRDF(SpecularColor, GBuffer.Roughness, NoV);
 	}
-
-	MobileLightAccumulator_Add(MobileLightAccumulator, 0.0f, SpecularIBLLighting, 1.0f);
+	LightAccumulator_AddSplit(DirectLighting, 0.0f, SpecularIBLLighting, 0.0f, 1.0f, false);
 }
 
 /*------------------------------------------------------------------------------
 	Mobile Local lights.
 ------------------------------------------------------------------------------*/
-
-struct FMobileShadowInfo
-{
-#if USE_SHADOWMASKTEXTURE
-	half Shadow;
-#elif SUPPORT_SPOTLIGHTS_SHADOW
-	FPCFSamplerSettings Settings;
-	float SpotLightShadowSharpen;
-	float SpotLightFadeFraction;
-	float SpotLightReceiverDepthBias;
-	float SpotLightSoftTransitionScale;
-	float4 SpotLightShadowmapMinMax;
-	float4x4 SpotLightShadowWorldToShadowMatrix;
-#endif
-};
-
-struct FMobileLightData
-{
-	float3 TranslatedWorldPosition;
-	float  InvRadius;
-	float3 Color;
-	float  FalloffExponent;
-	float3 Direction;
-	float2 SpotAngles;
-	float SourceRadius;
-	float SpecularScale;
-	bool bInverseSquared;
-	bool bSpotLight;
-	bool bCastSpotLightShadow;
-
-// shadow info
-	FMobileShadowInfo ShadowInfo;
-};
-
-void AccumulateLocalLightingInner(FMobileLightData LightData
-	, FGBufferData GBuffer
-	, FMobileShadingModelContext ShadingModelContext
-	, float3 TranslatedWorldPosition
-	, half3 CameraVector
-	, inout FMobileLightAccumulator MobileLightAccumulator)
-{
-	float3 ToLight = LightData.TranslatedWorldPosition - TranslatedWorldPosition;
-	float3 LocalPosition = -ToLight;
-	float DistanceSqr = dot(ToLight, ToLight);
-	half3 L = ToLight * rsqrt(DistanceSqr);
-	float3 H = normalize(float3(CameraVector + L));
-
-	half NoL = max(0, dot(GBuffer.WorldNormal, L));
-	float NoH = max(0, dot(GBuffer.WorldNormal, H));
-	half VoH = max(0, dot(CameraVector, H));
-
-	float Attenuation = 0.0;
-	if (LightData.bInverseSquared)
-	{
-		// Sphere falloff (technically just 1/d2 but this avoids inf)
-		Attenuation = 1.0f / (DistanceSqr + 1.0f);
-		Attenuation *= Square(saturate(1 - Square(DistanceSqr * Square(LightData.InvRadius))));
-	}
-	else
-	{
-		Attenuation = RadialAttenuation(ToLight * LightData.InvRadius, LightData.FalloffExponent);
-	}
-
-	if (LightData.bSpotLight)
-	{
-		Attenuation *= SpotAttenuation(L, -LightData.Direction, LightData.SpotAngles);
-	}
-
-	if (Attenuation > 0.0)
-	{
-		half3 H = normalize(CameraVector + L);
-		half NoL = max(0.0, dot(GBuffer.WorldNormal, L));
-		half NoH = max(0.0, dot(GBuffer.WorldNormal, H));
-		half VoH = max(0.0, dot(CameraVector, H));
-
-		half Shadow = 1.0f;
-
-		if (LightData.bCastSpotLightShadow)
-		{
-#if USE_SHADOWMASKTEXTURE
-			Shadow = LightData.ShadowInfo.Shadow;
-#elif SUPPORT_SPOTLIGHTS_SHADOW
-			float4 HomogeneousShadowPosition = mul(float4(LocalPosition, 1), LightData.ShadowInfo.SpotLightShadowWorldToShadowMatrix);
-			float2 ShadowUVs = HomogeneousShadowPosition.xy / HomogeneousShadowPosition.w;
-			if (all(ShadowUVs >= LightData.ShadowInfo.SpotLightShadowmapMinMax.xy && ShadowUVs <= LightData.ShadowInfo.SpotLightShadowmapMinMax.zw))
-			{
-				// Clamp pixel depth in light space for shadowing opaque, because areas of the shadow depth buffer that weren't rendered to will have been cleared to 1
-				// We want to force the shadow comparison to result in 'unshadowed' in that case, regardless of whether the pixel being shaded is in front or behind that plane
-				// Invert ShadowZ as the shadow space has been changed (but not yet the filtering code)
-				float ShadowZ = 1.0f - HomogeneousShadowPosition.z;
-				float LightSpacePixelDepthForOpaque = min(ShadowZ, 0.99999f);
-				LightData.ShadowInfo.Settings.SceneDepth = LightSpacePixelDepthForOpaque;
-				LightData.ShadowInfo.Settings.TransitionScale = LightData.ShadowInfo.SpotLightSoftTransitionScale * lerp(LightData.ShadowInfo.SpotLightReceiverDepthBias, 1.0, NoL);
-
-				Shadow = MobileShadowPCF(ShadowUVs, LightData.ShadowInfo.Settings);
-
-				Shadow = saturate((Shadow - 0.5) * LightData.ShadowInfo.SpotLightShadowSharpen + 0.5f);
-
-				Shadow = lerp(1.0f, Square(Shadow), LightData.ShadowInfo.SpotLightFadeFraction);
-			}
-#endif
-		}
-
-		FMobileDirectLighting Lighting = MobileIntegrateBxDF(ShadingModelContext, GBuffer, NoL, NoH, VoH, CameraVector, L, Shadow);
-		MobileLightAccumulator_Add(MobileLightAccumulator, Lighting.Diffuse, Lighting.Specular * LightData.SpecularScale, LightData.Color * (1.0 / PI) * Attenuation);
-	}
-}
-
-void AccumulateLocalLighting(FGBufferData GBuffer
-	, FMobileShadingModelContext ShadingModelContext
-	, float3 TranslatedWorldPosition
-	, half3 CameraVector
-	, float3 LightTranslatedWorldPosition
-	, float InvRadius
-	, float3 Color
-	, float FalloffExponent
-	, float3 Direction
-	, float SpecularScale
-	, float2 SpotLightAngles
-	, bool bInverseSquared
-	, bool bSpotLight
-	, bool bCastSpotLightShadow
-	, FMobileShadowInfo ShadowInfo
-	, inout FMobileLightAccumulator MobileLightAccumulator)
-{
-	FMobileLightData LightData;
-	{
-		LightData.TranslatedWorldPosition = LightTranslatedWorldPosition;
-		LightData.InvRadius = InvRadius;
-		LightData.Color = Color;
-		LightData.FalloffExponent = FalloffExponent;
-		LightData.Direction = Direction;
-		LightData.SpotAngles = SpotLightAngles;
-		LightData.SpecularScale = SpecularScale;
-		LightData.bInverseSquared = bInverseSquared;
-		LightData.bSpotLight = bSpotLight;
-		LightData.bCastSpotLightShadow = bCastSpotLightShadow;
-
-		LightData.ShadowInfo = ShadowInfo;
-	}
-
-	AccumulateLocalLightingInner(LightData
-		, GBuffer
-		, ShadingModelContext
-		, TranslatedWorldPosition
-		, CameraVector
-		, MobileLightAccumulator);
-}
-
 /**
  * Adds local lighting using the light grid, does not apply directional lights, as they are done elsewhere.
  */
 #if ENABLE_CLUSTERED_LIGHTS
+ /**
+  * Convert a FLocalLightData into a FDeferredLightData
+  */
+FDeferredLightData ConvertLocalToDeferredLight(const FLocalLightData In)
+{
+	// TODO: Not part of LightGrid ligth data, does it need to be added?
+	float SpecularScale = 1.0f;
+	const uint LightTypeAndPackedShadowMapChannelMask = asuint(In.LightDirectionAndShadowMask.w);
+
+	// bits [17:16] really
+	const uint LightType = LightTypeAndPackedShadowMapChannelMask >> 16;
+
+	FDeferredLightData Out = (FDeferredLightData)0;
+	Out.TranslatedWorldPosition = In.LightPositionAndInvRadius.xyz;
+	Out.InvRadius = In.LightPositionAndInvRadius.w;
+	Out.Color = In.LightColorAndFalloffExponent.xyz;
+	//Out.Color.xy = 0.0f;
+	Out.FalloffExponent = In.LightColorAndFalloffExponent.w;
+	Out.Direction = In.LightDirectionAndShadowMask.xyz;
+	Out.SpotAngles = In.SpotAnglesAndSourceRadiusPacked.xy;
+//	Out.SourceRadius = In.SpotAnglesAndSourceRadiusPacked.z;
+//	Out.SourceLength = f16tof32(asuint(In.SpotAnglesAndSourceRadiusPacked.w));
+//	Out.Tangent = In.LightTangentAndSoftSourceRadius.xyz;
+//	Out.SoftSourceRadius = In.LightTangentAndSoftSourceRadius.w;
+	Out.bInverseSquared = Out.FalloffExponent == 0;
+	// Only directional is not radial, and this loop is only non-directional lights (LIGHT_SOURCE_SHAPE > 0 in trad) 
+	Out.bRadialLight = true;
+	// Setting it to true follows how it is setup in standard deferred, where only the directional lights are not doing this (LIGHT_SOURCE_SHAPE > 0)
+	//Out.bSpotLight = true;
+	//Out.bSpotLight = Out.SpotAngles.x > -2.0f;
+	Out.SpecularScale = SpecularScale;
+	// TODO: not sure if this check is even needed, standard deferred always seems to set it to true?
+	Out.bSpotLight = (LightType == LIGHT_TYPE_SPOT);
+	// NOTE: supporting the rect-light code path adds ~30% overhead on NVIDIA (RTX 2060) even when not used! 
+	// Presumed register pressure increase again, not profiled in detail.
+	// Rect lights are now disabled by the sorting, so the shader should not handle them, for now.
+	Out.bRectLight = false;// (LightType == LIGHT_TYPE_RECT);
+	// bits [7:0]
+//	Out.ShadowedBits = (LightTypeAndPackedShadowMapChannelMask & 0xFF) != 0 ? 1 : 0;
+	// Static shadowing uses ShadowMapChannel, dynamic shadows are packed into light attenuation using PreviewShadowMapChannel
+//	Out.ShadowMapChannelMask = UnpackShadowMapChannelMask(LightTypeAndPackedShadowMapChannelMask);
+	Out.ShadowedBits = 1;
+	Out.HairTransmittance = InitHairTransmittanceData();
+	return Out;
+}
+
 void AccumulateLightGridLocalLighting(const FCulledLightsGridData InLightGridData
-	, FMobileShadingModelContext ShadingModelContext
 	, FGBufferData GBuffer
 	, float3 TranslatedWorldPosition
 	, half3 CameraVector
@@ -552,7 +459,7 @@ void AccumulateLightGridLocalLighting(const FCulledLightsGridData InLightGridDat
 #if !MOBILE_DEFERRED_LIGHTING
 	, uint LightingChannelMask
 #endif
-	, inout FMobileLightAccumulator MobileLightAccumulator)
+	, inout FLightAccumulator DirectLighting)
 {
 	// Limit max to ForwardLightData.NumLocalLights.
 	// This prevents GPU hangs when the PS tries to read from uninitialized NumCulledLightsGrid buffer
@@ -592,31 +499,22 @@ void AccumulateLightGridLocalLighting(const FCulledLightsGridData InLightGridDat
 		// bits [17:16] really
 		uint LightType = LightTypeAndPackedShadowMapChannelMask >> 16;
 
-		half4 PreviewShadowMapChannelMask = UnpackShadowMapChannelMask(LightTypeAndPackedShadowMapChannelMask >> 4);
-		half Shadow = dot(PreviewShadowMapChannelMask, DynamicShadowFactors);
-
-		FMobileShadowInfo ShadowInfo = (FMobileShadowInfo)0;
+		half Shadow = 1.0f;
 #if USE_SHADOWMASKTEXTURE
-		ShadowInfo.Shadow = Shadow;
+		if (LightType == LIGHT_TYPE_SPOT)
+		{
+			half4 PreviewShadowMapChannelMask = UnpackShadowMapChannelMask(LightTypeAndPackedShadowMapChannelMask >> 4);
+			Shadow = dot(PreviewShadowMapChannelMask, DynamicShadowFactors);
+		}
 #endif
 
-		AccumulateLocalLighting(GBuffer,
-			ShadingModelContext,
-			TranslatedWorldPosition,
-			CameraVector,
-			LocalLight.LightPositionAndInvRadius.xyz,
-			LocalLight.LightPositionAndInvRadius.w,
-			LocalLight.LightColorAndFalloffExponent.xyz,
-			LocalLight.LightColorAndFalloffExponent.w,
-			LocalLight.LightDirectionAndShadowMask.xyz,
-			1.0f,
-			LocalLight.SpotAnglesAndSourceRadiusPacked.xy,
-			LocalLight.LightColorAndFalloffExponent.w == 0,
-			(LightType == LIGHT_TYPE_SPOT),
-			(LightTypeAndPackedShadowMapChannelMask & 0xFF) != 0,
-			ShadowInfo,
-			MobileLightAccumulator
-			);
+		FDeferredLightData LightData = ConvertLocalToDeferredLight(LocalLight);
+
+		half4 LightAttenuation = half4(1, 1, Shadow, Shadow);
+
+		float SurfaceShadow = 0;
+		FLightAccumulator NewLighting = AccumulateDynamicLighting(TranslatedWorldPosition, CameraVector, GBuffer, 1, GBuffer.ShadingModelID, LightData, LightAttenuation, 0, uint2(0, 0), SurfaceShadow);
+		DirectLighting = LightAccumulator_Add(DirectLighting, NewLighting);
 	}
 }
 #endif
\ No newline at end of file
diff --git a/Engine/Shaders/Private/MobileShadingModels.ush b/Engine/Shaders/Private/MobileShadingModels.ush
deleted file mode 100644
diff --git a/Engine/Shaders/Private/RayTracing/RayTracingMaterialHitShaders.usf b/Engine/Shaders/Private/RayTracing/RayTracingMaterialHitShaders.usf
--- a/Engine/Shaders/Private/RayTracing/RayTracingMaterialHitShaders.usf
+++ b/Engine/Shaders/Private/RayTracing/RayTracingMaterialHitShaders.usf
@@ -148,6 +148,8 @@ void GetPrecomputedIndirectLightingAndSkyLight(
 	LightmapUVType SkyOcclusionUV = (LightmapUVType)0;
 	uint SkyOcclusionDataIndex = 0u;
 
+	const bool bEvaluateBackface = GetShadingModelRequiresBackfaceLighting(ShadingModel);
+
 #if HQ_TEXTURE_LIGHTMAP || LQ_TEXTURE_LIGHTMAP
 		LightmapUVType LightmapUV0, LightmapUV1;
 		uint LightmapDataIndex;
@@ -165,7 +167,7 @@ void GetPrecomputedIndirectLightingAndSkyLight(
 		OutDiffuseLighting = max(float3(0,0,0), DotSH3(IrradianceSH, DiffuseTransferSH)) / PI;
 
 		#if MATERIAL_SHADINGMODEL_TWOSIDED_FOLIAGE
-		if (ShadingModel == SHADINGMODELID_TWOSIDED_FOLIAGE)
+		if (bEvaluateBackface)
 		{
 			FThreeBandSHVector SubsurfaceTransferSH = CalcDiffuseTransferSH3(-DiffuseDir, 1);
 			OutSubsurfaceLighting += max(float3(0,0,0), DotSH3(IrradianceSH, SubsurfaceTransferSH)) / PI;
@@ -177,12 +179,12 @@ void GetPrecomputedIndirectLightingAndSkyLight(
 
 		SkyOcclusionUV = LightmapUV0;
 		SkyOcclusionDataIndex = LightmapDataIndex;
-		GetLightMapColorHQ(LightmapVTPageTableResult, LightmapUV0, LightmapUV1, LightmapDataIndex, DiffuseDir, MaterialParameters.SvPosition.xy, ShadingModel, OutDiffuseLighting, OutSubsurfaceLighting);
+		GetLightMapColorHQ(LightmapVTPageTableResult, LightmapUV0, LightmapUV1, LightmapDataIndex, DiffuseDir, MaterialParameters.SvPosition.xy, bEvaluateBackface, OutDiffuseLighting, OutSubsurfaceLighting);
 
 	// Low quality texture lightmaps
 	#elif LQ_TEXTURE_LIGHTMAP
 
-		OutDiffuseLighting = GetLightMapColorLQ(LightmapVTPageTableResult, LightmapUV0, LightmapUV1, LightmapDataIndex, DiffuseDir).rgb;
+		GetLightMapColorLQ(LightmapVTPageTableResult, LightmapUV0, LightmapUV1, LightmapDataIndex, DiffuseDir, bEvaluateBackface, OutDiffuseLighting, OutSubsurfaceLighting, OutIndirectIrradiance);
 
 	#endif
 
diff --git a/Engine/Shaders/Private/ReflectionEnvironmentShared.ush b/Engine/Shaders/Private/ReflectionEnvironmentShared.ush
--- a/Engine/Shaders/Private/ReflectionEnvironmentShared.ush
+++ b/Engine/Shaders/Private/ReflectionEnvironmentShared.ush
@@ -122,13 +122,13 @@ float3 GetSkySHDiffuseSimple(float3 Normal)
 
 
 // Point lobe in off-specular peak direction
-float3 GetOffSpecularPeakReflectionDir(float3 Normal, float3 ReflectionVector, float Roughness)
+half3 GetOffSpecularPeakReflectionDir(half3 Normal, half3 ReflectionVector, half Roughness)
 {
-	float a = Square(Roughness);
+	half a = Roughness * Roughness;
 	return lerp( Normal, ReflectionVector, (1 - a) * ( sqrt(1 - a) + a ) );	
 }
 
-float GetSpecularOcclusion(float NoV, float RoughnessSq, float AO)
+half GetSpecularOcclusion(half NoV, half RoughnessSq, half AO)
 {
 	return saturate( pow( NoV + AO, RoughnessSq ) - 1 + AO );
 }
diff --git a/Engine/Shaders/Private/ShadingCommon.ush b/Engine/Shaders/Private/ShadingCommon.ush
--- a/Engine/Shaders/Private/ShadingCommon.ush
+++ b/Engine/Shaders/Private/ShadingCommon.ush
@@ -4,6 +4,13 @@
 
 #include "Common.ush"
 
+#ifndef SHADING_PATH_MOBILE
+#define SHADING_PATH_MOBILE 0
+#endif
+
+#ifndef SHADING_PATH_DEFERRED
+#define SHADING_PATH_DEFERRED 0
+#endif
 
 ///////////////////////////////////////////////////////////////////////////////
 // Shading models
@@ -107,7 +114,7 @@ float F0RGBToDielectricSpecular(float3 F0)
 	return F0ToDielectricSpecular(max(F0.r, max(F0.g, F0.b)));
 }
 
-float DielectricSpecularToF0(float Specular)
+half DielectricSpecularToF0(half Specular)
 {
 	return 0.08f * Specular;
 }
@@ -135,7 +142,7 @@ float F0RGBToMicroOcclusion(float3 F0)
 	return F0ToMicroOcclusion(max(F0.r, max(F0.g, F0.b)));
 }
 
-float3 ComputeF0(float Specular, float3 BaseColor, float Metallic)
+half3 ComputeF0(half Specular, half3 BaseColor, half Metallic)
 {
 	return lerp(DielectricSpecularToF0(Specular).xxx, BaseColor, Metallic.xxx);
 }
diff --git a/Engine/Shaders/Private/ShadingModels.ush b/Engine/Shaders/Private/ShadingModels.ush
--- a/Engine/Shaders/Private/ShadingModels.ush
+++ b/Engine/Shaders/Private/ShadingModels.ush
@@ -12,6 +12,87 @@
 #include "HairBsdf.ush"
 #include "ShadingEnergyConservation.ush"
 
+#if SHADING_PATH_MOBILE
+#include "MobileGGX.ush"
+
+#ifndef MAX_MOBILE_SHADOWCASCADES
+#define MAX_MOBILE_SHADOWCASCADES 0
+#endif 
+
+#ifndef MOBILE_SHADOW_QUALITY
+#define MOBILE_SHADOW_QUALITY 2
+#endif
+
+//It's always 0 in mobile deferred lighting shaders
+#ifndef FULLY_ROUGH
+#define FULLY_ROUGH 0
+#endif
+
+//It's always 0 in mobile deferred lighting shaders
+#ifndef NONMETAL
+#define NONMETAL 0
+#endif
+
+//It's always 0 in mobile deferred lighting shaders
+#ifndef MATERIAL_SHADINGMODEL_SINGLELAYERWATER
+#define MATERIAL_SHADINGMODEL_SINGLELAYERWATER 0
+#endif
+
+//It's always 0 in mobile deferred lighting shaders
+#ifndef FORWARDSHADING_USE_HQ_ENV_BRDF
+#define  FORWARDSHADING_USE_HQ_ENV_BRDF 0
+#endif
+
+//It's always 1 in mobile deferred lighting shaders
+#ifndef MOBILE_DEFERRED_LIGHTING
+#define MOBILE_DEFERRED_LIGHTING 0
+#endif
+
+#ifndef DEFERRED_SHADING_PATH
+#define DEFERRED_SHADING_PATH 0
+#endif
+
+/*------------------------------------------------------------------------------
+	Mobile Shading Model Functions
+------------------------------------------------------------------------------*/
+half4 GetSSProfilePreIntegratedValue(uint SubsurfaceProfileInt, half NoL, half Curvature)
+{
+	float3 UV = float3((NoL * .5 + .5), Curvature, SubsurfaceProfileInt);
+
+	return Texture2DArraySampleLevel(View.SSProfilesPreIntegratedTexture, View.SSProfilesPreIntegratedSampler, UV, 0);
+}
+
+half3 GetEnvBRDF(half3 SpecularColor, half Roughness, half NoV)
+{
+#if FULLY_ROUGH
+	return 0.0f;
+#elif FORWARDSHADING_USE_HQ_ENV_BRDF || MOBILE_DEFERRED_LIGHTING
+	return EnvBRDF(SpecularColor, Roughness, NoV);
+#elif NONMETAL
+	// If nothing is hooked up to Metalic and Specular,
+	// then defaults are the same as a non-metal,
+	// so this define is safe.
+	return EnvBRDFApproxNonmetal(Roughness, NoV).xxx;
+#else
+	return EnvBRDFApprox(SpecularColor, Roughness, NoV);
+#endif
+}
+
+half3 MobileSpecularGGXInner(half D, half3 SpecularColor, half Roughness, half NoV, half NoL, half VoH, bool bHighQualityBRDF)
+{
+	half Vis = (Roughness * 0.25 + 0.25);
+	half3 F = GetEnvBRDF(SpecularColor, Roughness, NoV);
+
+	if (bHighQualityBRDF)
+	{
+		Vis = saturate(Vis_SmithJointApprox(Roughness * Roughness * Roughness * Roughness, NoV, NoL));
+		F = F_Schlick(SpecularColor, VoH);
+	}
+
+	return (D * Vis) * F;
+}
+#endif // SHADING_PATH_MOBILE
+
 #if 0
 void StandardShadingShared( float3 DiffuseColor, float3 SpecularColor, float Roughness, float3 V, half3 N )
 {
@@ -58,12 +139,12 @@ struct FDirectLighting
 
 struct FShadowTerms
 {
-	float	SurfaceShadow;
-	float	TransmissionShadow;
-	float	TransmissionThickness;
+	half	SurfaceShadow;
+	half	TransmissionShadow;
+	half	TransmissionThickness;
 	FHairTransmittanceData HairTransmittance;
 };
-FDirectLighting HairBxDF(FGBufferData GBuffer, half3 N, half3 V, half3 L, float Falloff, float NoL, FAreaLight AreaLight, FShadowTerms Shadow)
+FDirectLighting HairBxDF(FGBufferData GBuffer, half3 N, half3 V, half3 L, float Falloff, half NoL, FAreaLight AreaLight, FShadowTerms Shadow)
 {
 	const float3 BsdfValue = HairShading(GBuffer, L, V, N, Shadow.TransmissionShadow, Shadow.HairTransmittance, 1, 0, uint2(0, 0));
 
@@ -141,20 +222,25 @@ float3 SpecularGGX(float Roughness, float Anisotropy, float3 SpecularColor, BxDF
 	return (D * Vis) * F;
 }
 
-float3 SpecularGGX( float Roughness, float3 SpecularColor, BxDFContext Context, float NoL, FAreaLight AreaLight )
+float3 SpecularGGX( float Roughness, float3 SpecularColor, BxDFContext Context, half NoL, FAreaLight AreaLight )
 {
 	float a2 = Pow4( Roughness );
 	float Energy = EnergyNormalization( a2, Context.VoH, AreaLight );
 	
+#if SHADING_PATH_MOBILE
+	half D = GGX_Mobile(Roughness, Context.NoH) * Energy;
+	return MobileSpecularGGXInner(D, SpecularColor, Roughness, Context.NoV, NoL, Context.VoH, MOBILE_HIGH_QUALITY_BRDF);
+#else
 	// Generalized microfacet specular
 	float D = D_GGX( a2, Context.NoH ) * Energy;
 	float Vis = Vis_SmithJointApprox( a2, Context.NoV, NoL );
 	float3 F = F_Schlick( SpecularColor, Context.VoH );
 
 	return (D * Vis) * F;
+#endif
 }
 
-float3 DualSpecularGGX(float AverageRoughness, float Lobe0Roughness, float Lobe1Roughness, float LobeMix, float3 SpecularColor, BxDFContext Context, float NoL, FAreaLight AreaLight)
+half3 DualSpecularGGX(half AverageRoughness, half Lobe0Roughness, half Lobe1Roughness, half LobeMix, half3 SpecularColor, BxDFContext Context, half NoL, FAreaLight AreaLight)
 {
 	float AverageAlpha2 = Pow4(AverageRoughness);
 	float Lobe0Alpha2 = Pow4(Lobe0Roughness);
@@ -163,15 +249,21 @@ float3 DualSpecularGGX(float AverageRoughness, float Lobe0Roughness, float Lobe1
 	float Lobe0Energy = EnergyNormalization(Lobe0Alpha2, Context.VoH, AreaLight);
 	float Lobe1Energy = EnergyNormalization(Lobe1Alpha2, Context.VoH, AreaLight);
 
+#if SHADING_PATH_MOBILE
+	half D = lerp(GGX_Mobile(Lobe0Roughness, Context.NoH) * Lobe0Energy, GGX_Mobile(Lobe1Roughness, Context.NoH) * Lobe1Energy, LobeMix);
+
+	return MobileSpecularGGXInner(D, SpecularColor, AverageRoughness, Context.NoV, NoL, Context.VoH, MOBILE_HIGH_QUALITY_BRDF);
+#else
 	// Generalized microfacet specular
 	float D = lerp(D_GGX(Lobe0Alpha2, Context.NoH) * Lobe0Energy, D_GGX(Lobe1Alpha2, Context.NoH) * Lobe1Energy, LobeMix);
 	float Vis = Vis_SmithJointApprox(AverageAlpha2, Context.NoV, NoL); // Average visibility well approximates using two separate ones (one per lobe).
 	float3 F = F_Schlick(SpecularColor, Context.VoH);
 
 	return (D * Vis) * F;
+#endif
 }
 
-FDirectLighting DefaultLitBxDF( FGBufferData GBuffer, half3 N, half3 V, half3 L, float Falloff, float NoL, FAreaLight AreaLight, FShadowTerms Shadow )
+FDirectLighting DefaultLitBxDF( FGBufferData GBuffer, half3 N, half3 V, half3 L, float Falloff, half NoL, FAreaLight AreaLight, FShadowTerms Shadow )
 {
 	BxDFContext Context;
 	FDirectLighting Lighting;
@@ -196,7 +288,11 @@ FDirectLighting DefaultLitBxDF( FGBufferData GBuffer, half3 N, half3 V, half3 L,
 	}
 	else
 	{
+#if SHADING_PATH_MOBILE
+		InitMobile(Context, N, V, L, NoL);
+#else
 		Init(Context, N, V, L);
+#endif
 
 		NoV = Context.NoV;
 		VoH = Context.VoH;
@@ -272,7 +368,7 @@ float RefractBlend(float VoH, float Eta)
 	return Eta * VoH - sqrt(k);
 }
 
-float RefractBlendClearCoatApprox(float VoH)
+half RefractBlendClearCoatApprox(half VoH)
 {
 	// Polynomial approximation of refraction blend factor for normal component of VoH with fixed Eta (1/1.5):
 	return (0.63 - 0.22 * VoH) * VoH - 0.745;
@@ -296,9 +392,9 @@ BxDFContext RefractClearCoatContext(BxDFContext Context)
 	//  NoV2 = 1.0 / 1.5 * NoV - RefractBlendClearCoatApprox(VoH) * NoH
 
 	BxDFContext RefractedContext = Context;
-	float Eta = 1.0 / 1.5;
-	float RefractionBlendFactor = RefractBlendClearCoatApprox(Context.VoH);
-	float RefractionProjectionTerm = RefractionBlendFactor * Context.NoH;
+	half Eta = 1.0 / 1.5;
+	half RefractionBlendFactor = RefractBlendClearCoatApprox(Context.VoH);
+	half RefractionProjectionTerm = RefractionBlendFactor * Context.NoH;
 	RefractedContext.NoV = clamp(Eta * Context.NoV - RefractionProjectionTerm, 0.001, 1.0); // Due to CalcThinTransmission and Vis_SmithJointAniso, we need to make sure
 	RefractedContext.NoL = clamp(Eta * Context.NoL - RefractionProjectionTerm, 0.001, 1.0); // those values are not 0s to avoid NaNs.
 	RefractedContext.VoH = saturate(Eta * Context.VoH - RefractionBlendFactor);
@@ -307,12 +403,12 @@ BxDFContext RefractClearCoatContext(BxDFContext Context)
 	return RefractedContext;
 }
 
-FDirectLighting ClearCoatBxDF( FGBufferData GBuffer, half3 N, half3 V, half3 L, float Falloff, float NoL, FAreaLight AreaLight, FShadowTerms Shadow )
+FDirectLighting ClearCoatBxDF( FGBufferData GBuffer, half3 N, half3 V, half3 L, float Falloff, half NoL, FAreaLight AreaLight, FShadowTerms Shadow )
 {
-	const float ClearCoat			= GBuffer.CustomData.x;
-	const float ClearCoatRoughness	= max(GBuffer.CustomData.y, 0.02f);
-	const float Film = 1 * ClearCoat;
-	const float MetalSpec = 0.9;
+	const half ClearCoat			= GBuffer.CustomData.x;
+	const half ClearCoatRoughness	= max(GBuffer.CustomData.y, 0.02f);
+	const half Film = 1 * ClearCoat;
+	const half MetalSpec = 0.9;
 
 	FDirectLighting Lighting = {
 		float3(0.0, 0.0, 0.0),
@@ -348,7 +444,7 @@ FDirectLighting ClearCoatBxDF( FGBufferData GBuffer, half3 N, half3 V, half3 L,
 	// Note: the operation is not invertible for GBuffer.Roughness = 1.0, so roughness is clamped to 254.0/255.0
 	float SphereSinAlpha = AreaLight.SphereSinAlpha;
 	float RoughnessCompensation = 1 - Pow2(GBuffer.Roughness);
-	float Alpha = Pow2(ClearCoatRoughness);
+	half Alpha = ClearCoatRoughness * ClearCoatRoughness;
 	RoughnessCompensation = RoughnessCompensation > 0.0 ? (1 - Alpha) / RoughnessCompensation : 0.0;
 	AreaLight.SphereSinAlpha = saturate(AreaLight.SphereSinAlpha * RoughnessCompensation);
 
@@ -373,10 +469,14 @@ FDirectLighting ClearCoatBxDF( FGBufferData GBuffer, half3 N, half3 V, half3 L,
 		// Generalized microfacet specular
 		float a2 = Pow2(Alpha);
 		float ClearCoatEnergy = EnergyNormalization(a2, Context.VoH, AreaLight);
+		half Vis = Vis_SmithJointApprox(a2, Context.NoV, NoL);
+#if SHADING_PATH_MOBILE
+		half D = GGX_Mobile(ClearCoatRoughness, Context.NoH) * ClearCoatEnergy;
+		Vis = saturate(Vis);
+#else
 		float D = D_GGX(a2, Context.NoH) * ClearCoatEnergy;
-		float Vis = Vis_SmithJointApprox(a2, Context.NoV, NoL);
-
-		float Fr1 = D * Vis * F;
+#endif
+		half Fr1 = (D * Vis) * F;
 		Lighting.Specular = ClearCoat * AreaLight.FalloffColor * (Falloff * NoL * Fr1);
 	}
 	Lighting.Specular *= ComputeEnergyConservation(EnergyTermsCoat);
@@ -392,7 +492,7 @@ FDirectLighting ClearCoatBxDF( FGBufferData GBuffer, half3 N, half3 V, half3 L,
 	float3 FresnelCoeff = ComputeEnergyPreservation(EnergyTermsCoat); // 1.0 - F;
 	#else
 	// Preserve old behavior when energy conservation is disabled
-	float FresnelCoeff = 1.0 - F;
+	half FresnelCoeff = 1.0 - F;
 	#endif
 	FresnelCoeff *= FresnelCoeff;
 
@@ -463,8 +563,8 @@ FDirectLighting ClearCoatBxDF( FGBufferData GBuffer, half3 N, half3 V, half3 L,
 	float3 Transmission = SimpleClearCoatTransmittance(BottomContext.NoL, BottomContext.NoV, GBuffer.Metallic, GBuffer.BaseColor);
 
 	// Default Lit
-	float3 DefaultDiffuse = (Falloff * NoL) * AreaLight.FalloffColor * Diffuse_Lambert(GBuffer.DiffuseColor) * ComputeEnergyPreservation(EnergyTermsBottom);
-	float3 RefractedDiffuse = FresnelCoeff * Transmission * DefaultDiffuse;
+	half3 DefaultDiffuse = (Falloff * NoL) * AreaLight.FalloffColor * Diffuse_Lambert(GBuffer.DiffuseColor) * ComputeEnergyPreservation(EnergyTermsBottom);
+	half3 RefractedDiffuse = FresnelCoeff * Transmission * DefaultDiffuse;
 	Lighting.Diffuse = lerp(DefaultDiffuse, RefractedDiffuse, ClearCoat);
 
 	if (!bHasAnisotropy && bIsRect)
@@ -477,8 +577,8 @@ FDirectLighting ClearCoatBxDF( FGBufferData GBuffer, half3 N, half3 V, half3 L,
 	else
 	{
 		float a2 = Pow4(GBuffer.Roughness);
-		float D2 = 0;
-		float Vis2 = 0;
+		half D2 = 0;
+		half Vis2 = 0;
 
 		BRANCH
 		if (bHasAnisotropy)
@@ -492,12 +592,17 @@ FDirectLighting ClearCoatBxDF( FGBufferData GBuffer, half3 N, half3 V, half3 L,
 		}
 		else
 		{
-			D2 = D_GGX(a2, BottomContext.NoH);
 			// NoL is chosen to provide better parity with DefaultLit when ClearCoat=0
 			Vis2 = Vis_SmithJointApprox(a2, BottomContext.NoV, NoL);
+#if SHADING_PATH_MOBILE
+			D2 = GGX_Mobile(GBuffer.Roughness, BottomContext.NoH);
+			Vis2 = saturate(Vis2);
+#else
+			D2 = D_GGX(a2, BottomContext.NoH);
+#endif
 		}
-		float3 F = F_Schlick(GBuffer.SpecularColor, BottomContext.VoH);
-		float3 F_DefaultLit = F_Schlick(GBuffer.SpecularColor, Context.VoH);
+		half3 F = F_Schlick(GBuffer.SpecularColor, BottomContext.VoH);
+		half3 F_DefaultLit = F_Schlick(GBuffer.SpecularColor, Context.VoH);
 		
 		float Energy = 0;
 
@@ -540,7 +645,7 @@ float ApproximateHG(float cosJ, float g)
 	return (ISO_PHASE_FUNC_Normalized * (1.0f - g2) / max( 1e-5, gcos2));
 }
 
-void GetProfileDualSpecular(uint SubsurfaceProfileInt, float Roughness, float Opacity, out float LobeRoughness0, out float LobeRoughness1, out float LobeMix)
+void GetProfileDualSpecular(uint SubsurfaceProfileInt, half Roughness, half Opacity, out half LobeRoughness0, out half LobeRoughness1, out half LobeMix)
 {
 #if !FORWARD_SHADING
 	GetSubsurfaceProfileDualSpecular(SubsurfaceProfileInt, Roughness, Opacity, LobeRoughness0, LobeRoughness1, LobeMix);
@@ -552,7 +657,7 @@ void GetProfileDualSpecular(uint SubsurfaceProfileInt, float Roughness, float Op
 #endif
 }
 
-FDirectLighting SubsurfaceProfileBxDF( FGBufferData GBuffer, half3 N, half3 V, half3 L, float Falloff, float NoL, FAreaLight AreaLight, FShadowTerms Shadow )
+FDirectLighting SubsurfaceProfileBxDF( FGBufferData GBuffer, half3 N, half3 V, half3 L, float Falloff, half NoL, FAreaLight AreaLight, FShadowTerms Shadow )
 {
 	BxDFContext Context;
 	Init( Context, N, V, L );
@@ -560,21 +665,35 @@ FDirectLighting SubsurfaceProfileBxDF( FGBufferData GBuffer, half3 N, half3 V, h
 	Context.NoV = saturate( abs( Context.NoV ) + 1e-5 );
 
 	uint SubsurfaceProfileId = ExtractSubsurfaceProfileInt(GBuffer);
-	float Opacity = GBuffer.CustomData.a;
-	float Roughness = GBuffer.Roughness;
+	half Opacity = GBuffer.CustomData.a;
+	half Roughness = GBuffer.Roughness;
 
-	float Lobe0Roughness = 0;
-	float Lobe1Roughness = 0;
-	float LobeMix = 0;
+	half Lobe0Roughness = 0;
+	half Lobe1Roughness = 0;
+	half LobeMix = 0;
 
 	GetProfileDualSpecular(SubsurfaceProfileId, Roughness, Opacity, Lobe0Roughness, Lobe1Roughness, LobeMix);
-	float AverageRoughness = lerp(Lobe0Roughness, Lobe1Roughness, LobeMix);
+	half AverageRoughness = lerp(Lobe0Roughness, Lobe1Roughness, LobeMix);
 
 	// Take the average roughness instead of compute a separate energy term for each roughness
 	const FBxDFEnergyTerms EnergyTerms = ComputeGGXSpecEnergyTerms(AverageRoughness, Context.NoV, GBuffer.SpecularColor);
 
 	FDirectLighting Lighting;
+#if SHADING_PATH_MOBILE
+	half Curvature = GBuffer.Curvature;
+	half UnClampedNoL = dot(GBuffer.WorldNormal, L);
+	half ShadowFactor = 1.0f - sqrt(Shadow.SurfaceShadow);
+	// Rotate the world normal based on the shadow value, it's just a experimental value 
+	half UnClampedRotatedNoL = max(UnClampedNoL - max(2.0f * UnClampedNoL, 0.4f) * ShadowFactor, -1.0f);
+	half4 BurleyDiffuse = GetSSProfilePreIntegratedValue(SubsurfaceProfileId, UnClampedRotatedNoL, Curvature);
+	// asset specific color
+	half3 Tint = GetSubsurfaceProfileTexture(SSSS_TINT_SCALE_OFFSET, SubsurfaceProfileId).rgb;
+	
+	// Needs to apply shadow at here, since the preintegrated value has already counted it, and we need to skip applying shadow outside.
+	Lighting.Diffuse = lerp(AreaLight.FalloffColor * (Falloff * NoL) * Shadow.SurfaceShadow, BurleyDiffuse.rgb, Tint) * Diffuse_Lambert(GBuffer.DiffuseColor);
+#else
 	Lighting.Diffuse  = AreaLight.FalloffColor * (Falloff * NoL) * Diffuse_Burley( GBuffer.DiffuseColor, GBuffer.Roughness, Context.NoV, NoL, Context.VoH );
+#endif
 
 	if (IsRectLight(AreaLight))
 	{
@@ -612,7 +731,7 @@ FDirectLighting SubsurfaceProfileBxDF( FGBufferData GBuffer, half3 N, half3 V, h
 	return Lighting;
 }
 
-FDirectLighting ClothBxDF( FGBufferData GBuffer, half3 N, half3 V, half3 L, float Falloff, float NoL, FAreaLight AreaLight, FShadowTerms Shadow )
+FDirectLighting ClothBxDF( FGBufferData GBuffer, half3 N, half3 V, half3 L, float Falloff, half NoL, FAreaLight AreaLight, FShadowTerms Shadow )
 {
 	const float3 FuzzColor	= ExtractSubsurfaceColor(GBuffer);
 	const float  Cloth		= saturate(GBuffer.CustomData.a);
@@ -650,22 +769,22 @@ FDirectLighting ClothBxDF( FGBufferData GBuffer, half3 N, half3 V, half3 L, floa
 	return Lighting;
 }
 
-FDirectLighting SubsurfaceBxDF( FGBufferData GBuffer, half3 N, half3 V, half3 L, float Falloff, float NoL, FAreaLight AreaLight, FShadowTerms Shadow )
+FDirectLighting SubsurfaceBxDF( FGBufferData GBuffer, half3 N, half3 V, half3 L, float Falloff, half NoL, FAreaLight AreaLight, FShadowTerms Shadow )
 {
 	FDirectLighting Lighting = DefaultLitBxDF( GBuffer, N, V, L, Falloff, NoL, AreaLight, Shadow );
 	
-	float3 SubsurfaceColor = ExtractSubsurfaceColor(GBuffer);
-	float Opacity = GBuffer.CustomData.a;
+	half3 SubsurfaceColor = ExtractSubsurfaceColor(GBuffer);
+	half Opacity = GBuffer.CustomData.a;
 
-	float3 H = normalize(V + L);
+	half3 H = normalize(V + L);
 
 	// to get an effect when you see through the material
 	// hard coded pow constant
-	float InScatter = pow(saturate(dot(L, -V)), 12) * lerp(3, .1f, Opacity);
+	half InScatter = pow(saturate(dot(L, -V)), 12) * lerp(3, .1f, Opacity);
 	// wrap around lighting, /(PI*2) to be energy consistent (hack do get some view dependnt and light dependent effect)
 	// Opacity of 0 gives no normal dependent lighting, Opacity of 1 gives strong normal contribution
-	float NormalContribution = saturate(dot(N, H) * Opacity + 1 - Opacity);
-	float BackScatter = GBuffer.GBufferAO * NormalContribution / (PI * 2);
+	half NormalContribution = saturate(dot(N, H) * Opacity + 1 - Opacity);
+	half BackScatter = GBuffer.GBufferAO * NormalContribution / (PI * 2);
 	
 	// lerp to never exceed 1 (energy conserving)
 	Lighting.Transmission = AreaLight.FalloffColor * ( Falloff * lerp(BackScatter, 1, InScatter) ) * SubsurfaceColor;
@@ -673,26 +792,30 @@ FDirectLighting SubsurfaceBxDF( FGBufferData GBuffer, half3 N, half3 V, half3 L,
 	return Lighting;
 }
 
-FDirectLighting TwoSidedBxDF( FGBufferData GBuffer, half3 N, half3 V, half3 L, float Falloff, float NoL, FAreaLight AreaLight, FShadowTerms Shadow )
+FDirectLighting TwoSidedBxDF( FGBufferData GBuffer, half3 N, half3 V, half3 L, float Falloff, half NoL, FAreaLight AreaLight, FShadowTerms Shadow )
 {
 	FDirectLighting Lighting = DefaultLitBxDF( GBuffer, N, V, L, Falloff, NoL, AreaLight, Shadow );
 
-	float3 SubsurfaceColor = ExtractSubsurfaceColor(GBuffer);
+	half3 SubsurfaceColor = ExtractSubsurfaceColor(GBuffer);
 
 	// http://blog.stevemcauley.com/2011/12/03/energy-conserving-wrapped-diffuse/
-	float Wrap = 0.5;
-	float WrapNoL = saturate( ( -dot(N, L) + Wrap ) / Square( 1 + Wrap ) );
+	half Wrap = 0.5;
+	half WrapNoL = saturate( ( -dot(N, L) + Wrap ) / Square( 1 + Wrap ) );
 
 	// Scatter distribution
-	float VoL = dot(V, L);
+	half VoL = dot(V, L);
+#if SHADING_PATH_MOBILE
+	half Scatter = GGX_Mobile(0.6, saturate(-VoL));
+#else
 	float Scatter = D_GGX( 0.6*0.6, saturate( -VoL ) );
+#endif
 
 	Lighting.Transmission = AreaLight.FalloffColor * (Falloff * WrapNoL * Scatter) * SubsurfaceColor;
 
 	return Lighting;
 }
 
-FDirectLighting EyeBxDF( FGBufferData GBuffer, half3 N, half3 V, half3 L, float Falloff, float NoL, FAreaLight AreaLight, FShadowTerms Shadow )
+FDirectLighting EyeBxDF( FGBufferData GBuffer, half3 N, half3 V, half3 L, float Falloff, half NoL, FAreaLight AreaLight, FShadowTerms Shadow )
 {
 #if IRIS_NORMAL
 	const float2 CausticNormalDelta	= float2( GBuffer.StoredMetallic, GBuffer.StoredSpecular )	* 2 - (256.0/255.0);
@@ -738,11 +861,14 @@ FDirectLighting EyeBxDF( FGBufferData GBuffer, half3 N, half3 V, half3 L, float
 	{
 		float a2 = Pow4( GBuffer.Roughness );
 		float Energy = EnergyNormalization( a2, Context.VoH, AreaLight );
-
 		// Generalized microfacet specular
-		float D = D_GGX( a2, Context.NoH ) * Energy;
-		float Vis = Vis_SmithJointApprox( a2, Context.NoV, NoL );
-		
+		float Vis = Vis_SmithJointApprox(a2, Context.NoV, NoL);
+#if SHADING_PATH_MOBILE
+		half D = GGX_Mobile(GBuffer.Roughness, Context.NoH) * Energy;
+		Vis = saturate(Vis);
+#else
+		float D = D_GGX(a2, Context.NoH) * Energy;
+#endif
 		Lighting.Specular = AreaLight.FalloffColor * (Falloff * NoL) * D * Vis * F;
 	}
 
@@ -763,23 +889,43 @@ FDirectLighting EyeBxDF( FGBufferData GBuffer, half3 N, half3 V, half3 L, float
 
 	Lighting.Diffuse = 0;
 	Lighting.Transmission = AreaLight.FalloffColor * ( Falloff * lerp( Sclera, Iris, IrisMask ) * EnergyPreservation ) * Diffuse_Lambert( GBuffer.DiffuseColor );
+#if SHADING_PATH_MOBILE
+	uint SubsurfaceProfileId = ExtractSubsurfaceProfileInt(GBuffer);
+	half Curvature = GBuffer.Curvature;
+	half ShadowFactor = 1.0f - sqrt(Shadow.SurfaceShadow);
+	
+	// Rotate the world normal based on the shadow value, it's just a experimental value 
+	half UnClampedNoL = dot(GBuffer.WorldNormal, L);
+	half UnClampedRotatedNoL = max(UnClampedNoL - max(2.0f * UnClampedNoL, 0.4f) * ShadowFactor, -1.0f);
+	half4 BurleyDiffuse = GetSSProfilePreIntegratedValue(SubsurfaceProfileId, UnClampedRotatedNoL, Curvature);
+
+	// Rotate the world normal based on the shadow value, it's just a experimental value 
+	half UnClampedIrisNoL = dot(IrisNormal, L);
+	half UnClampedRotatedIrisNoL = max(UnClampedIrisNoL - max(2.0f * UnClampedIrisNoL, 0.4f) * ShadowFactor, -1.0f);
+	half4 IrisBurleyDiffuse = GetSSProfilePreIntegratedValue(SubsurfaceProfileId, UnClampedRotatedIrisNoL * Caustic, Curvature);
+	// asset specific color
+	half3 Tint = GetSubsurfaceProfileTexture(SSSS_TINT_SCALE_OFFSET, SubsurfaceProfileId).rgb;
+
+	// Needs to apply shadow at here, since the preintegrated value has already counted it, and we need to skip applying shadow outside.
+	Lighting.Transmission = lerp(Lighting.Transmission * Shadow.SurfaceShadow, Falloff * lerp(BurleyDiffuse.rgb, IrisBurleyDiffuse.rgb, IrisMask) * Diffuse_Lambert(GBuffer.DiffuseColor), Tint);
+#endif
 	return Lighting;
 }
 
-FDirectLighting PreintegratedSkinBxDF( FGBufferData GBuffer, half3 N, half3 V, half3 L, float Falloff, float NoL, FAreaLight AreaLight, FShadowTerms Shadow )
+FDirectLighting PreintegratedSkinBxDF( FGBufferData GBuffer, half3 N, half3 V, half3 L, float Falloff, half NoL, FAreaLight AreaLight, FShadowTerms Shadow )
 {
 	FDirectLighting Lighting = DefaultLitBxDF( GBuffer, N, V, L, Falloff, NoL, AreaLight, Shadow );
 	
-	float3 SubsurfaceColor = ExtractSubsurfaceColor(GBuffer);
-	float Opacity = GBuffer.CustomData.a;
+	half3 SubsurfaceColor = ExtractSubsurfaceColor(GBuffer);
+	half Opacity = GBuffer.CustomData.a;
 
-	float3 PreintegratedBRDF = Texture2DSampleLevel(View.PreIntegratedBRDF, View.PreIntegratedBRDFSampler, float2(saturate(dot(N, L) * .5 + .5), 1 - Opacity), 0).rgb;
+	half3 PreintegratedBRDF = Texture2DSampleLevel(View.PreIntegratedBRDF, View.PreIntegratedBRDFSampler, float2(saturate(dot(N, L) * .5 + .5), 1 - Opacity), 0).rgb;
 	Lighting.Transmission = AreaLight.FalloffColor * Falloff * PreintegratedBRDF * SubsurfaceColor;
 
 	return Lighting;
 }
 
-FDirectLighting IntegrateBxDF( FGBufferData GBuffer, half3 N, half3 V, half3 L, float Falloff, float NoL, FAreaLight AreaLight, FShadowTerms Shadow )
+FDirectLighting IntegrateBxDF( FGBufferData GBuffer, half3 N, half3 V, half3 L, float Falloff, half NoL, FAreaLight AreaLight, FShadowTerms Shadow )
 {
 	switch( GBuffer.ShadingModelID )
 	{
@@ -820,4 +966,4 @@ FDirectLighting EvaluateBxDF( FGBufferData GBuffer, half3 N, half3 V, half3 L, f
 	AreaLight.Texture = InitRectTexture();
 
 	return IntegrateBxDF( GBuffer, N, V, L, 1, NoL, AreaLight, Shadow );
-}
+}
\ No newline at end of file
diff --git a/Engine/Shaders/Private/ShadingModelsMaterial.ush b/Engine/Shaders/Private/ShadingModelsMaterial.ush
--- a/Engine/Shaders/Private/ShadingModelsMaterial.ush
+++ b/Engine/Shaders/Private/ShadingModelsMaterial.ush
@@ -62,6 +62,16 @@ void SetGBufferForShadingModel(
 		{
 			GBuffer.CustomData.rgb = EncodeSubsurfaceProfile(SubsurfaceProfile);
 			GBuffer.CustomData.a = Opacity;
+
+		#if SHADING_PATH_MOBILE
+			#if MATERIAL_SUBSURFACE_PROFILE_USE_CURVATURE
+				GBuffer.Curvature = GetMaterialCustomData0(MaterialParameters);
+			#else
+				GBuffer.Curvature = CalculateCurvature(GBuffer.WorldNormal, MaterialParameters.WorldPosition_CamRelative);
+			#endif
+
+			GBuffer.Curvature = clamp(GBuffer.Curvature, 0.001f, 1.0f);
+		#endif
 		}
 	#if SUBSURFACE_PROFILE_OPACITY_THRESHOLD
 		else
@@ -179,6 +189,16 @@ void SetGBufferForShadingModel(
 			GBuffer.CustomData.yz = UnitVectorToOctahedron( normalize(Tangent) ) * 0.5 + 0.5;
 		#endif
 	#endif
+
+	#if SHADING_PATH_MOBILE
+		#if MATERIAL_SHADINGMODEL_EYE_USE_CURVATURE
+			GBuffer.Curvature = Metallic;
+		#else
+			GBuffer.Curvature = CalculateCurvature(GBuffer.WorldNormal, MaterialParameters.WorldPosition_CamRelative);
+		#endif
+
+		GBuffer.Curvature = clamp(GBuffer.Curvature, 0.001f, 1.0f);
+	#endif
 	}
 #endif
 }
\ No newline at end of file
diff --git a/Engine/Shaders/Private/SubsurfaceProfileCommon.ush b/Engine/Shaders/Private/SubsurfaceProfileCommon.ush
--- a/Engine/Shaders/Private/SubsurfaceProfileCommon.ush
+++ b/Engine/Shaders/Private/SubsurfaceProfileCommon.ush
@@ -76,14 +76,14 @@ half4 GetSubsurfaceProfileTexture(Texture2D InSSProfilesTexture, SamplerState In
 #endif 
 }
 
-void GetSubsurfaceProfileDualSpecular(uint SubsurfaceProfileInt, float Roughness, float Opacity, out float LobeRoughness0, out float LobeRoughness1, out float LobeMix)
+void GetSubsurfaceProfileDualSpecular(uint SubsurfaceProfileInt, half Roughness, half Opacity, out half LobeRoughness0, out half LobeRoughness1, out half LobeMix)
 {
 	const half4 Data = GetSubsurfaceProfileTexture(SSSS_DUAL_SPECULAR_OFFSET, SubsurfaceProfileInt);
 
 	// Smooth blend out dual specular when opacity is low, we have the extra SSSS_OPACITY_THRESHOLD_EPS so that we fade out by the time we
 	// get to 0.01, as opposed to 0.0.
-	float MaterialRoughnessToLobeRoughness0 = lerp(1.0f, Data.x * SSSS_MAX_DUAL_SPECULAR_ROUGHNESS, saturate((Opacity - SSSS_OPACITY_THRESHOLD_EPS) * 10.0f));
-	float MaterialRoughnessToLobeRoughness1 = lerp(1.0f, Data.y * SSSS_MAX_DUAL_SPECULAR_ROUGHNESS, saturate((Opacity - SSSS_OPACITY_THRESHOLD_EPS) * 10.0f));
+	half MaterialRoughnessToLobeRoughness0 = lerp(1.0f, Data.x * SSSS_MAX_DUAL_SPECULAR_ROUGHNESS, saturate((Opacity - SSSS_OPACITY_THRESHOLD_EPS) * 10.0f));
+	half MaterialRoughnessToLobeRoughness1 = lerp(1.0f, Data.y * SSSS_MAX_DUAL_SPECULAR_ROUGHNESS, saturate((Opacity - SSSS_OPACITY_THRESHOLD_EPS) * 10.0f));
 	LobeMix = Data.z;
 
 	// Avoid specular explosion, and approximate `area lights`.
@@ -140,4 +140,22 @@ float4 GetSubsurfaceProfileDMFPInCm(int SubsurfaceProfileInt)
 	const float4 DiffuseMeanFreePath = DecodeDiffuseMeanFreePath(GetSubsurfaceProfileDiffuseMeanFreePath(SubsurfaceProfileInt));
 	const float WorldUnitScale = DecodeWorldUnitScale(GetSubsurfaceProfileWorldUnitScale(SubsurfaceProfileInt));
 	return DiffuseMeanFreePath * WorldUnitScale; // In cm
-}
\ No newline at end of file
+}
+
+#if SHADING_PATH_MOBILE
+half CalculateCurvature(half3 WorldNormal, float3 WorldPosition)
+{
+#if 0
+	half DeltaNormal = length(abs(DDX(WorldNormal)) + abs(DDY(WorldNormal)));
+	half DeltaPosition = length(abs(DDX(WorldPosition)) + abs(DDY(WorldPosition))) * BURLEY_CM_2_MM;
+	half CurvatureApprox = DeltaNormal / DeltaPosition;
+#else
+	half3 dNdx = ddx((HALF3_TYPE)WorldNormal);
+	half3 dNdy = ddy((HALF3_TYPE)WorldNormal);
+	half x = dot(dNdx, dNdx);
+	half y = dot(dNdy, dNdy);
+	half CurvatureApprox = pow(max(x, y), ResolvedView.NormalCurvatureToRoughnessScaleBias.z);
+#endif
+	return CurvatureApprox;
+}
+#endif
\ No newline at end of file
diff --git a/Engine/Source/Runtime/Core/Private/HAL/ConsoleManager.cpp b/Engine/Source/Runtime/Core/Private/HAL/ConsoleManager.cpp
--- a/Engine/Source/Runtime/Core/Private/HAL/ConsoleManager.cpp
+++ b/Engine/Source/Runtime/Core/Private/HAL/ConsoleManager.cpp
@@ -2643,6 +2643,13 @@ static TAutoConsoleVariable<int32> CVarMobileAllowPerPixelShadingModels(
 	ECVF_ReadOnly | ECVF_RenderThreadSafe
 );
 
+static TAutoConsoleVariable<int32> CVarMobileEnabledShadingModelsMask(
+	TEXT("r.Mobile.EnabledShadingModelsMask"),
+	0xFFFF,
+	TEXT("The mask of shader models that indicates which shading models could be enabled on mobile."),
+	ECVF_ReadOnly | ECVF_RenderThreadSafe
+);
+
 static TAutoConsoleVariable<int32> CVarMobileForwardSupportLocalLights(
 	TEXT("r.Mobile.Forward.EnableLocalLights"),
 	1,
diff --git a/Engine/Source/Runtime/Engine/Private/Materials/HLSLMaterialTranslator.cpp b/Engine/Source/Runtime/Engine/Private/Materials/HLSLMaterialTranslator.cpp
--- a/Engine/Source/Runtime/Engine/Private/Materials/HLSLMaterialTranslator.cpp
+++ b/Engine/Source/Runtime/Engine/Private/Materials/HLSLMaterialTranslator.cpp
@@ -1006,7 +1006,9 @@ bool FHLSLMaterialTranslator::Translate()
 			AddEstimatedTextureSample(1);
 		}
 
-		bUsesCurvature = FeatureLevel == ERHIFeatureLevel::ES3_1 && MaterialShadingModels.HasShadingModel(MSM_SubsurfaceProfile) && IsMaterialPropertyUsed(MP_CustomData0, Chunk[MP_CustomData0], FLinearColor(1, 0, 0, 0), 1);
+		bUsesCurvature = FeatureLevel == ERHIFeatureLevel::ES3_1 && 
+						((MaterialShadingModels.HasShadingModel(MSM_SubsurfaceProfile) && IsMaterialPropertyUsed(MP_CustomData0, Chunk[MP_CustomData0], FLinearColor(1, 0, 0, 0), 1))
+						|| (MaterialShadingModels.HasShadingModel(MSM_Eye) && IsMaterialPropertyUsed(MP_Opacity, Chunk[MP_Opacity], FLinearColor(1, 0, 0, 0), 1)));
 
 		if (BlendMode == BLEND_Modulate && MaterialShadingModels.IsLit() && !Material->IsDeferredDecal())
 		{
@@ -1562,7 +1564,7 @@ void FHLSLMaterialTranslator::ValidateShadingModelsForFeatureLevel(const FMateri
 {
 	if (FeatureLevel <= ERHIFeatureLevel::ES3_1)
 	{
-		const TArray<EMaterialShadingModel>& InvalidShadingModels = { MSM_Eye };
+		const TArray<EMaterialShadingModel>& InvalidShadingModels = {};
 		for (EMaterialShadingModel InvalidShadingModel : InvalidShadingModels)
 		{
 			if (ShadingModels.HasShadingModel(InvalidShadingModel))
@@ -1790,6 +1792,11 @@ void FHLSLMaterialTranslator::GetMaterialEnvironment(EShaderPlatform InPlatform,
 		{
 			OutEnvironment.SetDefine(TEXT("MATERIAL_SHADINGMODEL_EYE"), TEXT("1"));
 			NumSetMaterials++;
+
+			if (bUsesCurvature)
+			{
+				OutEnvironment.SetDefine(TEXT("MATERIAL_SHADINGMODEL_EYE_USE_CURVATURE"), TEXT("1"));
+			}
 		}
 		if (ShadingModels.HasShadingModel(MSM_SingleLayerWater))
 		{
@@ -9286,6 +9293,11 @@ int32 FHLSLMaterialTranslator::CustomPrimitiveData(int32 OutputIndex, EMaterialV
 
 int32 FHLSLMaterialTranslator::ShadingModel(EMaterialShadingModel InSelectedShadingModel)
 {
+	// If the shading model is masked out, fallback to default shading model
+	if ((EnabledShadingModelsMask(Platform) & (1 << (int32)InSelectedShadingModel)) == 0)
+	{
+		InSelectedShadingModel = MSM_DefaultLit;
+	}
 	ShadingModelsFromCompilation.AddShadingModel(InSelectedShadingModel);
 	return AddInlinedCodeChunk(MCT_ShadingModel, TEXT("%d"), InSelectedShadingModel);
 }
diff --git a/Engine/Source/Runtime/Engine/Private/Materials/MaterialShared.cpp b/Engine/Source/Runtime/Engine/Private/Materials/MaterialShared.cpp
--- a/Engine/Source/Runtime/Engine/Private/Materials/MaterialShared.cpp
+++ b/Engine/Source/Runtime/Engine/Private/Materials/MaterialShared.cpp
@@ -4983,6 +4983,7 @@ FText FMaterialAttributeDefinitionMap::GetAttributeOverrideForMaterial(const FGu
 		return Material->MaterialDomain == MD_Volume ? LOCTEXT("Albedo", "Albedo") : LOCTEXT("BaseColor", "Base Color");
 	case MP_Metallic:
 		CustomPinNames.Add({MSM_Hair, "Scatter"});
+		CustomPinNames.Add({ MSM_Eye, "Curvature" });
 		return FText::FromString(GetPinNameFromShadingModelField(Material->GetShadingModels(), CustomPinNames, "Metallic"));
 	case MP_Specular:
 		return LOCTEXT("Specular", "Specular");
diff --git a/Engine/Source/Runtime/RenderCore/Private/RenderUtils.cpp b/Engine/Source/Runtime/RenderCore/Private/RenderUtils.cpp
--- a/Engine/Source/Runtime/RenderCore/Private/RenderUtils.cpp
+++ b/Engine/Source/Runtime/RenderCore/Private/RenderUtils.cpp
@@ -1195,6 +1195,17 @@ RENDERCORE_API bool AllowPerPixelShadingModels(const FStaticShaderPlatform Platf
 	return true;
 }
 
+RENDERCORE_API int32 EnabledShadingModelsMask(const FStaticShaderPlatform Platform)
+{
+	static const auto CVar = IConsoleManager::Get().FindTConsoleVariableDataInt(TEXT("r.Mobile.EnabledShadingModelsMask"));
+
+	if (IsMobilePlatform(Platform) && CVar)
+	{
+		return CVar->GetValueOnAnyThread();
+	}
+	return 0xFFFF;
+}
+
 RENDERCORE_API uint64 GMobileAmbientOcclusionPlatformMask = 0;
 static_assert(SP_NumPlatforms <= sizeof(GMobileAmbientOcclusionPlatformMask) * 8, "GMobileAmbientOcclusionPlatformMask must be large enough to support all shader platforms");
 
diff --git a/Engine/Source/Runtime/RenderCore/Private/Shader.cpp b/Engine/Source/Runtime/RenderCore/Private/Shader.cpp
--- a/Engine/Source/Runtime/RenderCore/Private/Shader.cpp
+++ b/Engine/Source/Runtime/RenderCore/Private/Shader.cpp
@@ -1606,6 +1606,10 @@ void ShaderMapAppendKeyString(EShaderPlatform Platform, FString& KeyString)
 		{
 			KeyString += IsMobileDistanceFieldEnabled(Platform) ? TEXT("_MobSDF") : TEXT("");
 		}
+
+		{
+			KeyString +=  FString::Printf(TEXT("ESMM_%X"), EnabledShadingModelsMask(Platform));
+		}
 	}
 	else
 	{
diff --git a/Engine/Source/Runtime/RenderCore/Public/RenderUtils.h b/Engine/Source/Runtime/RenderCore/Public/RenderUtils.h
--- a/Engine/Source/Runtime/RenderCore/Public/RenderUtils.h
+++ b/Engine/Source/Runtime/RenderCore/Public/RenderUtils.h
@@ -499,6 +499,8 @@ RENDERCORE_API bool AllowPixelDepthOffset(const FStaticShaderPlatform Platform);
 
 RENDERCORE_API bool AllowPerPixelShadingModels(const FStaticShaderPlatform Platform);
 
+RENDERCORE_API int32 EnabledShadingModelsMask(const FStaticShaderPlatform Platform);
+
 RENDERCORE_API bool UseMobileAmbientOcclusion(const FStaticShaderPlatform Platform);
 
 RENDERCORE_API bool IsMobileDistanceFieldEnabled(const FStaticShaderPlatform Platform);
diff --git a/Engine/Source/Runtime/Renderer/Private/IndirectLightRendering.cpp b/Engine/Source/Runtime/Renderer/Private/IndirectLightRendering.cpp
--- a/Engine/Source/Runtime/Renderer/Private/IndirectLightRendering.cpp
+++ b/Engine/Source/Runtime/Renderer/Private/IndirectLightRendering.cpp
@@ -109,7 +109,7 @@ class FDiffuseIndirectCompositePS : public FGlobalShader
 
 	static bool ShouldCompilePermutation(const FGlobalShaderPermutationParameters& Parameters)
 	{
-		if (IsMetalPlatform(Parameters.Platform) && !IsMetalSM5Platform(Parameters.Platform))
+		if (IsMobilePlatform(Parameters.Platform))
 		{
 			return false;
 		}
diff --git a/Engine/Source/Runtime/Renderer/Private/MobileBasePassRendering.cpp b/Engine/Source/Runtime/Renderer/Private/MobileBasePassRendering.cpp
--- a/Engine/Source/Runtime/Renderer/Private/MobileBasePassRendering.cpp
+++ b/Engine/Source/Runtime/Renderer/Private/MobileBasePassRendering.cpp
@@ -233,7 +233,7 @@ void SetupMobileDirectionalLightUniformParameters(
 	FLightSceneInfo* Light = Scene.MobileDirectionalLights[ChannelIdx];
 	if (Light)
 	{
-		Params.DirectionalLightColor = Light->Proxy->GetSunIlluminanceAccountingForSkyAtmospherePerPixelTransmittance() / PI;
+		Params.DirectionalLightColor = Light->Proxy->GetSunIlluminanceAccountingForSkyAtmospherePerPixelTransmittance();
 		Params.DirectionalLightDirectionAndShadowTransition = FVector4f((FVector3f)-Light->Proxy->GetDirection(), 0.f);
 
 		const FVector2D FadeParams = Light->Proxy->GetDirectionalLightDistanceFadeParameters(FeatureLevel, Light->IsPrecomputedLightingValid(), SceneView.MaxShadowCascades);
