Fix a bug of mobile deferred lighting.

diff --git a/Engine/Shaders/Private/MobileDeferredShading.usf b/Engine/Shaders/Private/MobileDeferredShading.usf
--- a/Engine/Shaders/Private/MobileDeferredShading.usf
+++ b/Engine/Shaders/Private/MobileDeferredShading.usf
@@ -317,7 +317,7 @@ void MobileRadialLightPS(
 	half4 LightAttenuation = half4(1, 1, Shadow, Shadow);
 
 	float SurfaceShadow = 0;
-	DirectLighting = GetDynamicLightingSplit(TranslatedWorldPosition, CameraVector, GBuffer, 1, GBuffer.ShadingModelID, LightData, LightAttenuation, 0, uint2(0, 0), SurfaceShadow);
+	DirectLighting = AccumulateDynamicLighting(TranslatedWorldPosition, CameraVector, GBuffer, 1, GBuffer.ShadingModelID, LightData, LightAttenuation, 0, uint2(0, 0), SurfaceShadow);
 
 	half3 Color = DirectLighting.TotalLight * ComputeLightFunctionMultiplier(TranslatedWorldPosition);
 	
