Fix for shader compilation failures with forward shading enabled

- reported rendering issues with forward shading enable are due to defaulting materials
- EyeIndex is not used in the INSTANCED_STEREO case but it is still passed through some functions so we need to define it and default it to 0

diff --git a/Engine/Shaders/Private/BasePassVertexShader.usf b/Engine/Shaders/Private/BasePassVertexShader.usf
--- a/Engine/Shaders/Private/BasePassVertexShader.usf
+++ b/Engine/Shaders/Private/BasePassVertexShader.usf
@@ -44,6 +44,8 @@ void Main(
 	#else
 		ViewportIndex = EyeIndex;
 	#endif
+#else
+	const uint EyeIndex = 0;
 #endif
 	ResolvedView = ResolveViewFromVF(Input);
 
