Strata - fixed compilation issue due to recent SSS change.

diff --git a/Engine/Shaders/Private/PostProcessSubsurface.usf b/Engine/Shaders/Private/PostProcessSubsurface.usf
--- a/Engine/Shaders/Private/PostProcessSubsurface.usf
+++ b/Engine/Shaders/Private/PostProcessSubsurface.usf
@@ -457,8 +457,8 @@ bool IsCheckerBoardNeighborSubsurface(float2 SceneUV, int2 PixelOffset)
 	const float2 UV = SceneUV + PixelOffset * View.BufferSizeAndInvSize.zw;
 
 #if STRATA_ENABLED
-	const FStrataSubsurfaceData SSSData = LoadStataSSSData(UV);
-	const bool bHasSSSProfile = SSSData.bIsValid;
+	const FStrataSubsurfaceData SSSData = LoadStrataSSSData(UV);
+	const bool bHasSSSProfile = StrataSubSurfaceHeaderGetIsValid(SSSData.Header);
 #else
 	FScreenSpaceData ScreenSpaceData = GetScreenSpaceData(UV);
 	const bool bHasSSSProfile = UseSubsurfaceProfile(ScreenSpaceData.GBuffer.ShadingModelID);
