Strata - fix Strata simple claer coat not working with FXC due to asfloat and asuint.

diff --git a/Engine/Shaders/Private/Strata/Strata.ush b/Engine/Shaders/Private/Strata/Strata.ush
--- a/Engine/Shaders/Private/Strata/Strata.ush
+++ b/Engine/Shaders/Private/Strata/Strata.ush
@@ -134,6 +134,7 @@ struct FStrataBSDF
 #endif
 
 	float4 VGPRs[5];
+	uint   Haziness;					// Compilers such as FXC have troubles when using asfloat and asuint to propagate correct haziness bitfield around through registers. So we store that data as an actual uint.
 };
 
 ////////////////////////////////////////////////////////////////////////////////////////////////////////////////////////////////////////////////////////////////////////////////////////////////
@@ -408,7 +409,7 @@ struct FStrataBSDF
 #define SLAB_F0(X)						X.VGPRs[1].xyz
 #define SLAB_ANISOTROPY(X)				X.VGPRs[1].w
 #define SLAB_F90(X)						X.VGPRs[2].xyz
-#define SLAB_HAZINESS(X)				X.VGPRs[2].w
+#define SLAB_HAZINESS(X)				X.Haziness
 // Either with SSS profile
 #define SLAB_SSSPROFILEID(X)			X.VGPRs[3].x
 #define SLAB_SSSPROFILERADIUSSCALE(X)	X.VGPRs[3].y
@@ -506,17 +507,15 @@ FHaziness InitialiseHaziness()
 }
 
 // returns 16 bits: 1 bit bFastClearCoat, 7 bits Weight, 8 bits roughness
-float PackHaziness(FHaziness Haziness)
+uint PackHaziness(FHaziness Haziness)
 {
 	uint High8Bits = (Haziness.bSimpleClearCoat ? 0x80 : 0x00) | uint(saturate(Haziness.Weight) * 127);
 	uint Data = (High8Bits << 8) | uint(saturate(Haziness.Roughness) * 255);
-	return asfloat(Data);
+	return Data;
 }
 
-FHaziness UnpackHaziness(float FloatData)
+FHaziness UnpackHaziness(uint Data)
 {
-	uint Data = asuint(FloatData);
-
 	FHaziness Haziness;
 	Haziness.bSimpleClearCoat	= Data & 0x8000;
 	Haziness.Weight				= float(uint((Data>>8) & 0x7F)) * (1.0f / 127.0f);
@@ -524,7 +523,7 @@ FHaziness UnpackHaziness(float FloatData)
 	return Haziness;
 }
 
-float LerpHazinessParameterBlending(FStrataBSDF A, FStrataBSDF B, float MixValue)
+uint LerpHazinessParameterBlending(FStrataBSDF A, FStrataBSDF B, float MixValue)
 {
 	FHaziness HazinessA = UnpackHaziness(SLAB_HAZINESS(A));
 	FHaziness HazinessB = UnpackHaziness(SLAB_HAZINESS(B));
@@ -537,8 +536,6 @@ float LerpHazinessParameterBlending(FStrataBSDF A, FStrataBSDF B, float MixValue
 	return PackHaziness(OutputHaziness);
 }
 
-//BSDF_SETHASHAZINESS(ResultBSDF, BSDF_GETHASHAZINESS(ForegroundBSDF) | BSDF_GETHASHAZINESS(BackgroundBSDF));	// TODO
-
 ///////////////////////////////////////////////////////////////////////////////
 // Irradiance and occlusion
 struct FStrataIrradianceAndOcclusion
@@ -3469,7 +3466,7 @@ FStrataBSDF UnpackStrataBSDFIn(FStrataMaterialContainer StrataBuffer, inout FStr
 			F90 = Divisor > 0.0f ? F90 / Divisor : F90;
 			SLAB_F90(OutBSDF) = F90;
 
-			SLAB_HAZINESS(OutBSDF) = asfloat(RawData >> 16); // Keeping the top 16 bits reinterpreted as float.
+			SLAB_HAZINESS(OutBSDF) = RawData >> 16; // Keeping the top 16 bits
 		}
 
 		const bool bIsSimpleVolume = BSDF_GETISSIMPLEVOLUME(OutBSDF);
diff --git a/Engine/Shaders/Private/Strata/StrataExport.ush b/Engine/Shaders/Private/Strata/StrataExport.ush
--- a/Engine/Shaders/Private/Strata/StrataExport.ush
+++ b/Engine/Shaders/Private/Strata/StrataExport.ush
@@ -489,7 +489,7 @@ void PackStrataOut(
 								float3 F90YCoCg = LinearRGB_2_NormalisedYCoCg(F90);
 
 								uint F90Data = PackRGBA8(float4(F90YCoCg.y, F90YCoCg.z, 0, 0));
-								uint HazinessData = asuint(SLAB_HAZINESS(BSDF));
+								uint HazinessData = SLAB_HAZINESS(BSDF);
 								STRATA_STORE_UINT1(HazinessData << 16 | F90Data);
 							}
 
