Strata - minor renaming.

diff --git a/Engine/Shaders/Private/ReflectionEnvironmentPixelShader.usf b/Engine/Shaders/Private/ReflectionEnvironmentPixelShader.usf
--- a/Engine/Shaders/Private/ReflectionEnvironmentPixelShader.usf
+++ b/Engine/Shaders/Private/ReflectionEnvironmentPixelShader.usf
@@ -288,7 +288,7 @@ void ReflectionEnvironmentSkyLighting(
 
 		// A reduction of SSR contribution is needed when there is haziness because we only consider SSR for the sharpest lobe.
 		float SSRReductionFactor = 1.0f;
-		float3 TopLayerEnvBRDF = 0.0f;
+		float3 SSRTopLayerEnvBRDF = 0.0f;
 
 		const float CloudVolumetricAOShadow = GetCloudVolumetricAOShadow(TranslatedWorldPosition);
 		const FStrataIntegrationSettings Settings = InitStrataIntegrationSettings(false /*bForceFullyRough*/, Strata.bRoughDiffuse, Strata.PeelLayersAboveDepth);
@@ -309,11 +309,11 @@ void ReflectionEnvironmentSkyLighting(
 			StrataPixelHeader,
 			CloudVolumetricAOShadow,
 			SSRReductionFactor,
-			TopLayerEnvBRDF);
+			SSRTopLayerEnvBRDF);
 
 		// Screen space reflections added on top of environment lighting.
 		// Alpha is 0 to not affect the subsurface diffuse luma information.
-		const float4 SSRColor = float4(SSR.rgb * TopLayerEnvBRDF * SSRReductionFactor, 0.0);
+		const float4 SSRColor = float4(SSR.rgb * SSRTopLayerEnvBRDF * SSRReductionFactor, 0.0);
 
 		OutColor += StrataLighting.SceneColor + SSRColor;
 #if STRATA_OPAQUE_ROUGH_REFRACTION_ENABLED
diff --git a/Engine/Shaders/Private/Strata/StrataEnvironmentLighting.ush b/Engine/Shaders/Private/Strata/StrataEnvironmentLighting.ush
--- a/Engine/Shaders/Private/Strata/StrataEnvironmentLighting.ush
+++ b/Engine/Shaders/Private/Strata/StrataEnvironmentLighting.ush
@@ -31,7 +31,7 @@ FStrataDeferredLighting StrataEnvironmentLighting(
 	FStrataPixelHeader StrataPixelHeader,
 	float CloudVolumetricAOShadow,
 	inout float SSRReductionFactor,
-	inout float3 TopLayerEnvBRDF)
+	inout float3 SSRTopLayerEnvBRDF)
 {
 	FStrataDeferredLighting StrataLighting = GetInitialisedStrataDeferredLighting();
 
@@ -100,7 +100,7 @@ FStrataDeferredLighting StrataEnvironmentLighting(
 		if (bIsTopLayer)
 		{
 			// The specular path weight applied on SSR. It must account for throughput even for top surface because it also contains coverage.
-			TopLayerEnvBRDF += BSDFThroughput * StrataEnvLight.SSRSpecularWeight;
+			SSRTopLayerEnvBRDF += BSDFThroughput * StrataEnvLight.SSRSpecularWeight;
 		}
 
 		// STRATA_TODO ReflectionEnvironment(...) for static reflection volumes.
