Cleanup ray tracing decal shaders.
- Restructured decal shader parameters structure.
    - use member functions to handle input / output parameters.
    - store all data in array of floats to better match how we handle ray tracing payloads (useful to support mesh decal payload in followup CL).
    - this also further avoids situations that cause crash while creating RTPSO on some platforms (similar to CL 20646476).
- Moved logic that applies a decal to FPathTracingPayload to a separate function.
- Set SvPosition and ScreenPosition when evaluating decal material graph.

diff --git a/Engine/Shaders/Private/PathTracing/PathTracingCore.ush b/Engine/Shaders/Private/PathTracing/PathTracingCore.ush
--- a/Engine/Shaders/Private/PathTracing/PathTracingCore.ush
+++ b/Engine/Shaders/Private/PathTracing/PathTracingCore.ush
@@ -110,6 +110,44 @@ void AccumulateRadiance(inout float3 TotalRadiance, float3 PathRadiance, bool bI
 	TotalRadiance += PathRadiance;
 }
 
+void ApplyDecal(inout FPathTracingPayload HitPayload, FDecalShaderParams Params)
+{
+	const float DecalAlpha = Params.GetOpacity();
+	const uint DecalFlags = Params.GetFlags();
+
+	if (DecalFlags & DECAL_WRITE_BASE_COLOR_FLAG)
+	{
+		HitPayload.BaseColor = lerp(HitPayload.BaseColor, Params.GetBaseColor(), DecalAlpha);
+	}
+
+	if (DecalFlags & DECAL_WRITE_NORMAL_FLAG)
+	{
+		HitPayload.WorldNormal = lerp(HitPayload.WorldNormal, Params.GetWorldNormal(), DecalAlpha);
+	}
+
+	if (DecalFlags & DECAL_WRITE_ROUGHNESS_SPECULAR_METALLIC_FLAG)
+	{
+		HitPayload.Metallic = lerp(HitPayload.Metallic, Params.GetMetallic(), DecalAlpha);
+		HitPayload.Specular = lerp(HitPayload.Specular, Params.GetSpecular(), DecalAlpha);
+		HitPayload.Roughness = lerp(HitPayload.Roughness, Params.GetRoughness(), DecalAlpha);
+	}
+
+	if (DecalFlags & DECAL_WRITE_AMBIENT_OCCLUSION_FLAG)
+	{
+		// TODO: DecalAO?
+	}
+
+	if (DecalFlags & DECAL_WRITE_EMISSIVE_FLAG)
+	{
+		HitPayload.Radiance += Params.GetEmissive() * DecalAlpha;
+	}
+
+	HitPayload.DiffuseColor = HitPayload.BaseColor - HitPayload.BaseColor * HitPayload.Metallic;
+	HitPayload.SpecularColor = ComputeF0(HitPayload.Specular, HitPayload.BaseColor, HitPayload.Metallic);
+
+	// TODO: Payload.Opacity?
+}
+
 FPathTracingPayload TraceTransparentRay(inout FPathState PathState, bool bIsCameraRay, bool bLastBounce, bool bIncludeEmission, uint NumVisibleLights, inout FVolumeSegment VolumeSegment)
 {
 	const uint RayFlags = bIsCameraRay && EnableCameraBackfaceCulling ? RAY_FLAG_CULL_BACK_FACING_TRIANGLES : 0;
@@ -328,8 +366,7 @@ FPathTracingPayload TraceTransparentRay(inout FPathState PathState, bool bIsCame
 #if PLATFORM_SUPPORTS_CALLABLE_SHADERS
 		if (EnableDecals && HitPayload.IsDecalReceiver())
 		{
-			float3 TranslatedWorldPos = HitPayload.TranslatedWorldPos;
-			float3 WorldNormal = HitPayload.WorldNormal;
+			const float3 TranslatedWorldPos = HitPayload.TranslatedWorldPos;
 
 			FDecalLoopCount DecalLoopCount = DecalGridLookup(TranslatedWorldPos);
 
@@ -337,47 +374,13 @@ FPathTracingPayload TraceTransparentRay(inout FPathState PathState, bool bIsCame
 			{
 				uint DecalId = GetDecalId(Index, DecalLoopCount);
 
-				FDecalCallableShaderParams Params;
-				Params.TranslatedWorldPosition_MetallicSpecularRoughnessAO.rgb = TranslatedWorldPos;
-				Params.WorldNormal.rgb = HitPayload.WorldNormal;
-				Params.Color = float4(HitPayload.BaseColor, 0.0f);
-				Params.Emissive = float3(0.0f, 0.0f, 0.0f);
+				FDecalShaderParams Params = (FDecalShaderParams)0;
+				Params.SetInputTranslatedWorldPosition(TranslatedWorldPos);
+				Params.SetPixelCoord(PathState.PixelIndex);
 
 				CallShader(DecalId, Params);
 
-				const float DecalAlpha = Params.Color.a;
-
-				float DecalMetallic = Params.TranslatedWorldPosition_MetallicSpecularRoughnessAO.r;
-				if (DecalMetallic >= 0)
-				{
-					HitPayload.Metallic = lerp(HitPayload.Metallic, DecalMetallic, DecalAlpha);
-				}
-
-				float DecalSpecular = Params.TranslatedWorldPosition_MetallicSpecularRoughnessAO.g;
-				if (DecalSpecular >= 0)
-				{
-					HitPayload.Specular = lerp(HitPayload.Specular, DecalSpecular, DecalAlpha);
-				}
-
-				float3 DecalSpecularColor = ComputeF0(HitPayload.Specular, Params.Color.rgb, HitPayload.Metallic);
-				float3 DecalDiffuseColor = Params.Color.rgb - Params.Color.rgb * HitPayload.Metallic;
-
-				HitPayload.BaseColor = lerp(HitPayload.BaseColor, Params.Color.rgb, DecalAlpha);
-				HitPayload.DiffuseColor = lerp(HitPayload.DiffuseColor, DecalDiffuseColor, DecalAlpha);
-				HitPayload.SpecularColor = lerp(HitPayload.SpecularColor, DecalSpecularColor, DecalAlpha);
-
-				const float DecalRoughness = Params.TranslatedWorldPosition_MetallicSpecularRoughnessAO.b;
-				if (DecalRoughness >= 0.0f)
-				{
-					HitPayload.Roughness = lerp(HitPayload.Roughness, DecalRoughness, DecalAlpha);
-				}
-
-				HitPayload.WorldNormal = lerp(HitPayload.WorldNormal, Params.WorldNormal.rgb, DecalAlpha);
-
-				HitPayload.Radiance += Params.Emissive.rgb * DecalAlpha;
-
-				// TODO: DecalAO?
-				// TODO: Payload.Opacity?
+				ApplyDecal(HitPayload, Params);
 			}
 		}
 #endif
diff --git a/Engine/Shaders/Private/RayTracing/RayTracingCommon.ush b/Engine/Shaders/Private/RayTracing/RayTracingCommon.ush
--- a/Engine/Shaders/Private/RayTracing/RayTracingCommon.ush
+++ b/Engine/Shaders/Private/RayTracing/RayTracingCommon.ush
@@ -861,19 +861,70 @@ FMaterialClosestHitPayload UnpackRayTracingPayload(FPackedMaterialClosestHitPayl
 ///////////////////////////////////////////////////////////////////////////////////////////////////
 // Decal parameters
 
-#if !COMPUTESHADER
+#define DECAL_WRITE_BASE_COLOR_FLAG						(1 << 0)
+#define DECAL_WRITE_NORMAL_FLAG							(1 << 1)
+#define DECAL_WRITE_ROUGHNESS_SPECULAR_METALLIC_FLAG	(1 << 2)
+#define DECAL_WRITE_EMISSIVE_FLAG						(1 << 3)
+#define DECAL_WRITE_AMBIENT_OCCLUSION_FLAG				(1 << 4)
 
-struct FDecalCallableShaderParams
+struct FDecalShaderParams
 {
-	// In order to reduce params size reuse input params to store output
-	// Input: TranslatedWorldPosition / Output: MetallicSpecularRoughnessAO
-	float4 TranslatedWorldPosition_MetallicSpecularRoughnessAO;
-	float4 Color;
-	float3 WorldNormal;
-	float3 Emissive;
-};
+	float Data[15];
 
-#endif // !COMPUTESHADER
+	void SetInputTranslatedWorldPosition(float3 In)
+	{
+		Data[0] = In.x;
+		Data[1] = In.y;
+		Data[2] = In.z;
+	}
+	float3 GetInputTranslatedWorldPosition() { return float3(Data[0], Data[1], Data[2]); }
+
+	void SetPixelCoord(uint2 PixelCoord)
+	{
+		Data[3] = asfloat(PixelCoord.x);
+		Data[4] = asfloat(PixelCoord.y);
+	}
+	uint2 GetPixelCoord() { return asuint(float2(Data[3], Data[4])); }
+
+	void SetMetallic(float In) { Data[0] = In; }
+	float GetMetallic() { return Data[0]; }
+
+	void SetSpecular(float In) { Data[1] = In; }
+	float GetSpecular() { return Data[1]; }
+
+	void SetRoughness(float In) { Data[2] = In; }
+	float GetRoughness() { return Data[2]; }
+
+	void SetAO(float In) { Data[3] = In; }
+	float GetAO() { return Data[3]; }
+
+	void SetBaseColor(float3 In) {
+		Data[4] = In.r;
+		Data[5] = In.g;
+		Data[6] = In.b;
+	}
+	float3 GetBaseColor() { return float3(Data[4], Data[5], Data[6]); }
+
+	void SetOpacity(float In) { Data[7] = In; }
+	float GetOpacity() { return Data[7]; }
+
+	void SetWorldNormal(float3 In) {
+		Data[8] = In.r;
+		Data[9] = In.g;
+		Data[10] = In.b;
+	}
+	float3 GetWorldNormal() { return float3(Data[8], Data[9], Data[10]); }
+
+	void SetFlags(uint In) { Data[11] = asfloat(In); }
+	uint GetFlags() { return asuint(Data[11]); }
+
+	void SetEmissive(float3 In) {
+		Data[12] = In.r;
+		Data[13] = In.g;
+		Data[14] = In.b;
+	}
+	float3 GetEmissive() { return float3(Data[12], Data[13], Data[14]); }
+};
 
 ///////////////////////////////////////////////////////////////////////////////////////////////////
 // Tracing util functions
diff --git a/Engine/Shaders/Private/RayTracing/RayTracingDecalMaterialShader.usf b/Engine/Shaders/Private/RayTracing/RayTracingDecalMaterialShader.usf
--- a/Engine/Shaders/Private/RayTracing/RayTracingDecalMaterialShader.usf
+++ b/Engine/Shaders/Private/RayTracing/RayTracingDecalMaterialShader.usf
@@ -15,22 +15,16 @@
 #define DecalParams RayTracingDecalParameters.DecalParams
 #define DecalWriteFlags RayTracingDecalParameters.DecalWriteFlags
 
-#define DECAL_WRITE_BASE_COLOR_FLAG						(1 << 0)
-#define DECAL_WRITE_NORMAL_FLAG							(1 << 1)
-#define DECAL_WRITE_ROUGHNESS_SPECULAR_METALLIC_FLAG	(1 << 2)
-#define DECAL_WRITE_EMISSIVE_FLAG						(1 << 3)
-#define DECAL_WRITE_AMBIENT_OCCLUSION_FLAG				(1 << 4)
-
 // DECAL_PRIMITIVE informs material templates which functions to expose when rendering decals.
 #define DECAL_PRIMITIVE 1
-#include "/Engine/Generated/Material.ush"   
+#include "/Engine/Generated/Material.ush"
 
-RAY_TRACING_ENTRY_CALLABLE(RayTracingDecalMaterialShader,
-	FDecalCallableShaderParams, Params)
+RAY_TRACING_ENTRY_CALLABLE(RayTracingDecalMaterialShader, FDecalShaderParams, Params)
 {
 	ResolvedView = ResolveView();
 
-	const float3 TranslatedWorldPosition = Params.TranslatedWorldPosition_MetallicSpecularRoughnessAO.xyz;
+	const float3 TranslatedWorldPosition = Params.GetInputTranslatedWorldPosition();
+	const uint2 PixelCoord = Params.GetPixelCoord();
 
 	float4 DecalPos = mul(float4(TranslatedWorldPosition.xyz, 1), WorldToDecal);
 
@@ -61,8 +55,8 @@ RAY_TRACING_ENTRY_CALLABLE(RayTracingDecalMaterialShader,
 	MaterialParameters.TwoSidedSign = 1;
 	MaterialParameters.VertexColor = 1;
 	MaterialParameters.CameraVector = CameraVector;
-	MaterialParameters.SvPosition = 0.0f; // TODO
-	MaterialParameters.ScreenPosition = 0.0f; // TODO
+	MaterialParameters.SvPosition = float4(PixelCoord.xy, 0.0, 1.0);
+	MaterialParameters.ScreenPosition = SvPositionToResolvedScreenPosition(MaterialParameters.SvPosition);
 	MaterialParameters.LightVector = SwizzlePos;
 
 	MaterialParameters.AbsoluteWorldPosition = MaterialParameters.WorldPosition_NoOffsets = LWCSubtract(TranslatedWorldPosition.xyz, PrimaryView.PreViewTranslation);
@@ -72,42 +66,34 @@ RAY_TRACING_ENTRY_CALLABLE(RayTracingDecalMaterialShader,
 	CalcPixelMaterialInputs(MaterialParameters, PixelMaterialInputs);
 
 	const float DecalFading = saturate(4 - 4 * abs(SwizzlePos.z * 2 - 1)) * DecalParams.x;
-	Params.Color.a = GetMaterialOpacity(PixelMaterialInputs) * DecalFading;
+	Params.SetOpacity(GetMaterialOpacity(PixelMaterialInputs) * DecalFading);
 
 	if (DecalWriteFlags & DECAL_WRITE_BASE_COLOR_FLAG)
 	{
-		Params.Color.rgb = GetMaterialBaseColor(PixelMaterialInputs);
+		Params.SetBaseColor(GetMaterialBaseColor(PixelMaterialInputs));
 	}
 
 	if (DecalWriteFlags & DECAL_WRITE_NORMAL_FLAG)
 	{
-		Params.WorldNormal.rgb = MaterialParameters.WorldNormal;
+		Params.SetWorldNormal(MaterialParameters.WorldNormal);
 	}
 
 	if (DecalWriteFlags & DECAL_WRITE_ROUGHNESS_SPECULAR_METALLIC_FLAG)
 	{
-		Params.TranslatedWorldPosition_MetallicSpecularRoughnessAO.r = GetMaterialMetallic(PixelMaterialInputs);
-		Params.TranslatedWorldPosition_MetallicSpecularRoughnessAO.g = GetMaterialSpecular(PixelMaterialInputs);
-		Params.TranslatedWorldPosition_MetallicSpecularRoughnessAO.b = GetMaterialRoughness(PixelMaterialInputs);
-	}
-	else
-	{
-		Params.TranslatedWorldPosition_MetallicSpecularRoughnessAO.r = -1;
-		Params.TranslatedWorldPosition_MetallicSpecularRoughnessAO.g = -1;
-		Params.TranslatedWorldPosition_MetallicSpecularRoughnessAO.b = -1;
+		Params.SetMetallic(GetMaterialMetallic(PixelMaterialInputs));
+		Params.SetSpecular(GetMaterialSpecular(PixelMaterialInputs));
+		Params.SetRoughness(GetMaterialRoughness(PixelMaterialInputs));
 	}
 
 	if (DecalWriteFlags & DECAL_WRITE_AMBIENT_OCCLUSION_FLAG)
 	{
-		Params.TranslatedWorldPosition_MetallicSpecularRoughnessAO.a = GetMaterialAmbientOcclusion(PixelMaterialInputs);
-	}
-	else
-	{
-		Params.TranslatedWorldPosition_MetallicSpecularRoughnessAO.a = -1;
+		Params.SetAO(GetMaterialAmbientOcclusion(PixelMaterialInputs));
 	}
 
 	if (DecalWriteFlags & DECAL_WRITE_EMISSIVE_FLAG)
 	{
-		Params.Emissive = GetMaterialEmissive(PixelMaterialInputs);
+		Params.SetEmissive(GetMaterialEmissive(PixelMaterialInputs));
 	}
+
+	Params.SetFlags(DecalWriteFlags);
 }
diff --git a/Engine/Source/Runtime/D3D12RHI/Private/D3D12RayTracing.cpp b/Engine/Source/Runtime/D3D12RHI/Private/D3D12RayTracing.cpp
--- a/Engine/Source/Runtime/D3D12RHI/Private/D3D12RayTracing.cpp
+++ b/Engine/Source/Runtime/D3D12RHI/Private/D3D12RayTracing.cpp
@@ -5636,7 +5636,7 @@ void FD3D12CommandContext::RHISetRayTracingBindings(
 
 	auto BindingTask = [Bindings, Device, ShaderTable, Scene, Pipeline, &SharedCounter, IncrementCount, NumBindings, BindingType](uint32 WorkerIndex)
 	{
-		TRACE_CPUPROFILER_EVENT_SCOPE(SetRayTracingHitGroups_Task);
+		TRACE_CPUPROFILER_EVENT_SCOPE(SetRayTracingBindings_Task);
 	 
 		int32 BaseIndex = 0;
 		while (BaseIndex < (int32)NumBindings)
