Fix spelling: MAX_OCCUPENCY -> MAX_OCCUPANCY

diff --git a/Engine/Shaders/Private/VirtualShadowMaps/VirtualShadowMapProjection.usf b/Engine/Shaders/Private/VirtualShadowMaps/VirtualShadowMapProjection.usf
--- a/Engine/Shaders/Private/VirtualShadowMaps/VirtualShadowMapProjection.usf
+++ b/Engine/Shaders/Private/VirtualShadowMaps/VirtualShadowMapProjection.usf
@@ -96,7 +96,7 @@ void VirtualShadowMapCompositePS(
 
 
 SCHEDULER_MIN_PRESSURE
-MAX_OCCUPENCY
+MAX_OCCUPANCY
 
 // Screen space ray trace to attempt to skip over ambiguous regions near the receiver surface
 // Returns length at which to start the virtual shadow map ray; usually this is where the screen ray ended or went behind a surface
diff --git a/Engine/Shaders/Public/Platform.ush b/Engine/Shaders/Public/Platform.ush
--- a/Engine/Shaders/Public/Platform.ush
+++ b/Engine/Shaders/Public/Platform.ush
@@ -640,8 +640,8 @@ void ShaderYield()
 #define SCHEDULER_MIN_PRESSURE
 #endif
 
-#ifndef MAX_OCCUPENCY
-#define MAX_OCCUPENCY
+#ifndef MAX_OCCUPANCY
+#define MAX_OCCUPANCY
 #endif
 
 // ---------------------------------------------------- Interpolator attribute fallbacks
