Use proper sampling for cubemap preview when mip level is not specified.

diff --git a/Engine/Shaders/Private/SimpleElementPixelShader.usf b/Engine/Shaders/Private/SimpleElementPixelShader.usf
--- a/Engine/Shaders/Private/SimpleElementPixelShader.usf
+++ b/Engine/Shaders/Private/SimpleElementPixelShader.usf
@@ -384,30 +384,60 @@ void CubemapTextureProperties(
 
 	if (FaceIndex >= 0.0)
 	{
-		BaseColor = TextureCubeArraySampleLevel(CubeTexture, CubeTextureSampler, GetCubemapVector(TextureCoordinate * 2 - 1, FaceIndex), SliceIndex, MipLevel);
+		float4 CubeArrayCoordinates = float4(GetCubemapVector(TextureCoordinate * 2 - 1, FaceIndex), SliceIndex);
+		if (MipLevel >= 0.0)
+		{
+			BaseColor = TextureCubeArraySampleLevel(CubeTexture, CubeTextureSampler, CubeArrayCoordinates, MipLevel);
+		}
+		else
+		{
+			BaseColor = TextureCubeArraySample(CubeTexture, CubeTextureSampler, CubeArrayCoordinates);
+		}
 	}
 	else
 	{
 		// if the face index is not specified, display the cubemap as a longlat unwrap
 		float2 Angles = float2(2 * PI * (TextureCoordinate.x + 0.5f), PI * TextureCoordinate.y);
 		float s = sin(Angles.y);
-		float3 Direction = float3(s * sin(Angles.x), cos(Angles.y), -s * cos(Angles.x));
-		BaseColor = TextureCubeArraySampleLevel(CubeTexture, CubeTextureSampler, Direction.xzy, SliceIndex, MipLevel);
+		float4 CubeArrayCoordinates = float4(s * sin(Angles.x), -s * cos(Angles.x), cos(Angles.y), SliceIndex);
+		if (MipLevel >= 0.0)
+		{
+			BaseColor = TextureCubeArraySampleLevel(CubeTexture, CubeTextureSampler, CubeArrayCoordinates, MipLevel);
+		}
+		else
+		{
+			BaseColor = TextureCubeArraySample(CubeTexture, CubeTextureSampler, CubeArrayCoordinates);
+		}
 	}
 
 #else //TEXTURECUBE_ARRAY
 
 	if (FaceIndex >= 0.0)
 	{
-		BaseColor = TextureCubeSampleLevel(CubeTexture, CubeTextureSampler, GetCubemapVector(TextureCoordinate * 2 - 1, FaceIndex), MipLevel);
+		float3 CubeCoordinates = GetCubemapVector(TextureCoordinate * 2 - 1, FaceIndex);
+		if (MipLevel >= 0.0)
+		{
+			BaseColor = TextureCubeSampleLevel(CubeTexture, CubeTextureSampler, CubeCoordinates, MipLevel);
+		}
+		else
+		{
+			BaseColor = TextureCubeSample(CubeTexture, CubeTextureSampler, CubeCoordinates);
+		}
 	}
 	else
 	{
 		// if the face index is not specified, display the cubemap as a longlat unwrap
 		float2 Angles = float2(2 * PI * (TextureCoordinate.x + 0.5f), PI * TextureCoordinate.y);
 		float s = sin(Angles.y);
-		float3 Direction = float3(s * sin(Angles.x), cos(Angles.y), -s * cos(Angles.x));
-		BaseColor = TextureCubeSampleLevel(CubeTexture, CubeTextureSampler, Direction.xzy, MipLevel);
+		float3 CubeCoordinates = float3(s * sin(Angles.x), -s * cos(Angles.x), cos(Angles.y));
+		if (MipLevel >= 0.0)
+		{
+			BaseColor = TextureCubeSampleLevel(CubeTexture, CubeTextureSampler, CubeCoordinates, MipLevel);
+		}
+		else
+		{
+			BaseColor = TextureCubeSample(CubeTexture, CubeTextureSampler, CubeCoordinates);
+		}
 	}
 
 #endif //TEXTURECUBE_ARRAY
