Implemment support for Per Instance Custom Data for non-Nanite use cases (including ray tracing shaders)

This node is only supported on platforms that use GPU Scene. The use a VertexInterpolator node is no longer required for non-Nanite use-cases.

diff --git a/Engine/Shaders/Private/LocalVertexFactory.ush b/Engine/Shaders/Private/LocalVertexFactory.ush
--- a/Engine/Shaders/Private/LocalVertexFactory.ush
+++ b/Engine/Shaders/Private/LocalVertexFactory.ush
@@ -522,6 +522,11 @@ FMaterialPixelParameters GetMaterialPixelParameters(FVertexFactoryInterpolantsVS
 	Result.TwoSidedSign = 1;
 	Result.PrimitiveId = GetPrimitiveId(Interpolants);
 
+#if USES_PER_INSTANCE_CUSTOM_DATA && VF_USE_PRIMITIVE_SCENE_DATA
+	Result.CustomDataOffset = Interpolants.CustomDataOffset;
+	Result.CustomDataCount  = Interpolants.CustomDataCount;
+#endif
+
 #if HAS_INSTANCE_LOCAL_TO_WORLD_PS
 	Result.InstanceLocalToWorld = GetInstanceLocalToWorld(Interpolants);
 #endif
@@ -600,7 +605,7 @@ FMaterialVertexParameters GetMaterialVertexParameters(FVertexFactoryInput Input,
 		Result.InstanceWorldToLocal = GetInstanceData(Intermediates).WorldToLocal;
 		Result.PerInstanceParams = Intermediates.PerInstanceParams;
 	#endif
-	#if USES_PER_INSTANCE_CUSTOM_DATA
+	#if USES_PER_INSTANCE_CUSTOM_DATA && VF_USE_PRIMITIVE_SCENE_DATA
 		Result.CustomDataOffset = Intermediates.SceneData.InstanceData.CustomDataOffset;
 		Result.CustomDataCount = Intermediates.SceneData.InstanceData.CustomDataCount;
 	#endif
@@ -1074,6 +1079,11 @@ FVertexFactoryInterpolantsVSToPS VertexFactoryGetInterpolantsVSToPS(FVertexFacto
 	Interpolants.PerInstanceParams = Intermediates.PerInstanceParams;
 #endif
 
+#if USES_PER_INSTANCE_CUSTOM_DATA && VF_USE_PRIMITIVE_SCENE_DATA
+	Interpolants.CustomDataOffset = Intermediates.SceneData.InstanceData.CustomDataOffset;
+	Interpolants.CustomDataCount  = Intermediates.SceneData.InstanceData.CustomDataCount;
+#endif
+
 #if INSTANCED_STEREO
 	Interpolants.EyeIndex = 0;
 #endif
@@ -1211,6 +1221,12 @@ float4 VertexFactoryGetPreviousWorldPosition(FVertexFactoryInput Input, FVertexF
         O.InterpolantsVSToPS.LightmapDataIndex = a.InterpolantsVSToPS.LightmapDataIndex;
     #endif
 #endif	
+
+#if USES_PER_INSTANCE_CUSTOM_DATA && VF_USE_PRIMITIVE_SCENE_DATA
+		O.InterpolantsVSToPS.CustomDataOffset = a.InterpolantsVSToPS.CustomDataOffset;
+		O.InterpolantsVSToPS.CustomDataCount  = a.InterpolantsVSToPS.CustomDataCount;
+#endif
+
 		// Do we really need to interpolate TangentToWorld2 here? It should be replaced by the
 		// interpolated normal from 'whatever' interpolation scheme we're using
 	
diff --git a/Engine/Shaders/Private/LocalVertexFactoryCommon.ush b/Engine/Shaders/Private/LocalVertexFactoryCommon.ush
--- a/Engine/Shaders/Private/LocalVertexFactoryCommon.ush
+++ b/Engine/Shaders/Private/LocalVertexFactoryCommon.ush
@@ -35,7 +35,13 @@ struct FVertexFactoryInterpolantsVSToPS
 	#if NEEDS_LIGHTMAP_COORDINATE
 		nointerpolation uint LightmapDataIndex : LIGHTMAP_ID;
 	#endif
+	#if USES_PER_INSTANCE_CUSTOM_DATA
+		nointerpolation uint CustomDataOffset : CUSTOM_DATA_OFFSET;
+		nointerpolation uint CustomDataCount : CUSTOM_DATA_COUNT;
+	#endif
 #endif
+
+
 #if HAIR_STRAND_MESH_FACTORY
 		nointerpolation uint HairPrimitiveId		: HAIR_PRIMITIVE_ID;	// Control point ID
 		float2 HairPrimitiveUV						: HAIR_PRIMITIVE_UV;	// U: parameteric distance between the two surrounding control points. V: parametric distance along the width.
diff --git a/Engine/Shaders/Private/MaterialTemplate.ush b/Engine/Shaders/Private/MaterialTemplate.ush
--- a/Engine/Shaders/Private/MaterialTemplate.ush
+++ b/Engine/Shaders/Private/MaterialTemplate.ush
@@ -401,7 +401,7 @@ struct FMaterialPixelParameters
 	uint InstanceId;
 #endif
 
-#if IS_NANITE_PASS && USES_PER_INSTANCE_CUSTOM_DATA
+#if USES_PER_INSTANCE_CUSTOM_DATA && VF_USE_PRIMITIVE_SCENE_DATA
 	uint CustomDataOffset;
 	uint CustomDataCount;
 #endif
@@ -509,7 +509,7 @@ struct FMaterialVertexParameters
 	// TangentToWorld[2] is WorldVertexNormal
 	half3x3 TangentToWorld;
 
-#if USES_PER_INSTANCE_CUSTOM_DATA
+#if USES_PER_INSTANCE_CUSTOM_DATA && VF_USE_PRIMITIVE_SCENE_DATA
 	uint CustomDataOffset;
 	uint CustomDataCount;
 #endif
@@ -1073,7 +1073,7 @@ MaterialFloat2 GetDefaultSceneTextureUV(FMaterialPixelParameters Parameters, con
 /** Per Instance Custom Data Getter (Pixel Shader Variant - Visibility Buffer) */
 float GetPerInstanceCustomData(FMaterialPixelParameters Parameters, int Index, float DefaultValue)
 {
-#if IS_NANITE_PASS && USES_PER_INSTANCE_CUSTOM_DATA
+#if USES_PER_INSTANCE_CUSTOM_DATA && VF_USE_PRIMITIVE_SCENE_DATA
 	const uint FloatIndex = uint(Index);
 	BRANCH
 	if (FloatIndex < Parameters.CustomDataCount)
@@ -1118,7 +1118,7 @@ float GetPerInstanceCustomData(FMaterialVertexParameters Parameters, int Index,
 /** Get the per-instance custom data when instancing */
 MaterialFloat3 GetPerInstanceCustomData3Vector(FMaterialPixelParameters Parameters, int Index, MaterialFloat3 DefaultValue)
 {
-#if IS_NANITE_PASS && USES_PER_INSTANCE_CUSTOM_DATA
+#if USES_PER_INSTANCE_CUSTOM_DATA && VF_USE_PRIMITIVE_SCENE_DATA
 	return float3(GetPerInstanceCustomData(Parameters, Index + 0, DefaultValue.x),
 				  GetPerInstanceCustomData(Parameters, Index + 1, DefaultValue.y),
 				  GetPerInstanceCustomData(Parameters, Index + 2, DefaultValue.z));
