SMRT local lights updates
- Improvements to cube map filtering to eliminate various leaking artifacts at cube edges/corners
- Implement texel dither and optimal bias for local lights
- Remove some half stuff that was causing some sub-optimal compilation
- Fix VSM visualization alignment
- Add a general purpose permutation parameter that can be set for A/B testing projection performance that affects register allocation
- General cleanup

diff --git a/Engine/Shaders/Private/VirtualShadowMaps/VirtualShadowMapPageManagement.usf b/Engine/Shaders/Private/VirtualShadowMaps/VirtualShadowMapPageManagement.usf
--- a/Engine/Shaders/Private/VirtualShadowMaps/VirtualShadowMapPageManagement.usf
+++ b/Engine/Shaders/Private/VirtualShadowMaps/VirtualShadowMapPageManagement.usf
@@ -38,49 +38,27 @@
 
 StructuredBuffer<int> DirectionalLightIds;
 
-float ResolutionLodBiasLocal;
 float PageDilationBorderSizeDirectional;
 float PageDilationBorderSizeLocal;
 uint InputType;
 uint bCullBackfacingPixels;
 
-uint CalcMipLevelLocal(float Footprint)
+
+uint GetMipLevelLocal(int VirtualShadowMapId, float3 TranslatedWorldPosition, float SceneDepth)
 {
-	float MipLevelFloat = log2(Footprint) + ResolutionLodBiasLocal;
+	FVirtualShadowMapProjectionShaderData ProjectionData = GetVirtualShadowMapProjectionData(VirtualShadowMapId);
+
+	// If local lights are near the primary view the combined offset should be small
+	float3 ShadowTranslatedWorldPosition = TranslatedWorldPosition + 
+		LWCToFloat(LWCSubtract(ProjectionData.PreViewTranslation, PrimaryView.PreViewTranslation));
+
+	float Footprint = VirtualShadowMapCalcPixelFootprintLocal(ProjectionData, ShadowTranslatedWorldPosition, SceneDepth);
+
+	float MipLevelFloat = log2(Footprint) + ProjectionData.ResolutionLodBias;
 	uint MipLevel = uint(max(floor(MipLevelFloat), 0.0f));
 	return min(MipLevel, (VSM_MAX_MIP_LEVELS - 1U));
 }
 
-uint GetMipLevelLocal(int VirtualShadowMapId, float3 TranslatedWorldPosition, float SceneDepth)
-{
-	float Footprint = float(VSM_VIRTUAL_MAX_RESOLUTION_XY);
-	FVirtualShadowMapProjectionShaderData ProjectionData = GetVirtualShadowMapProjectionData(VirtualShadowMapId);
-
-	{
-		// Compute footprint by projecting the approximate size of a camera pixel at the given depth to shadow space
-		// NOTE: This doesn't take the screen XY position/FOV into account, which may or may not be desirable.
-		
-		// TODO: Roll into a uniform
-		float2 RadiusXY = 1.0f / (View.ViewSizeAndInvSize.xy * View.ViewToClip._m00_m11);
-		float RadiusScreen = min(RadiusXY.x, RadiusXY.y);
-		float DepthScale = SceneDepth * View.ViewToClip[2][3] + View.ViewToClip[3][3];
-
-		float RadiusWorld = DepthScale * RadiusScreen;
-
-		// If local lights are near the primary view the combined offset should be small
-		float3 ShadowTranslatedWorldPosition = TranslatedWorldPosition + 
-			LWCToFloat(LWCSubtract(ProjectionData.PreViewTranslation, PrimaryView.PreViewTranslation));
-		
-		float4 ShadowUVz = mul(float4(ShadowTranslatedWorldPosition, 1.0f), ProjectionData.TranslatedWorldToShadowUVMatrix);
-
-		float4 RadiusClipH = mul(float4(RadiusWorld, 0.0f, ShadowUVz.w, 1.0f), ProjectionData.ShadowViewToClipMatrix);
-		float RadiusClip = abs(RadiusClipH.x / RadiusClipH.w);
-		Footprint = RadiusClip * float(2 * VSM_VIRTUAL_MAX_RESOLUTION_XY);
-	}
-
-	return CalcMipLevelLocal(Footprint);
-}
-
 void MarkPage(uint VirtualShadowMapId, uint MipLevel, float3 TranslatedWorldPosition, bool bUsePageDilation, float2 PageDilationOffset)
 {
 	FVirtualShadowMapProjectionShaderData ProjectionData = GetVirtualShadowMapProjectionData(VirtualShadowMapId);
diff --git a/Engine/Shaders/Private/VirtualShadowMaps/VirtualShadowMapProjection.usf b/Engine/Shaders/Private/VirtualShadowMaps/VirtualShadowMapProjection.usf
--- a/Engine/Shaders/Private/VirtualShadowMaps/VirtualShadowMapProjection.usf
+++ b/Engine/Shaders/Private/VirtualShadowMaps/VirtualShadowMapProjection.usf
@@ -186,7 +186,7 @@ FLightShaderParameters ConvertFromLocal( const FLocalLightData LightData )
 	Light.Direction					= LightData.LightDirectionAndShadowMask.xyz;
 	Light.Tangent					= LightData.LightTangentAndSoftSourceRadius.xyz;
 	Light.SpotAngles				= LightData.SpotAnglesAndSourceRadiusPacked.xy;
-	Light.SpecularScale				= 1;
+	Light.SpecularScale				= LightData.RectBarnDoorAndVirtualShadowMapIdAndSpecularScale.w;
 	Light.SourceRadius				= LightData.SpotAnglesAndSourceRadiusPacked.z;
 	Light.SoftSourceRadius			= LightData.LightTangentAndSoftSourceRadius.w;
 	Light.SourceLength				= f16tof32(asuint(LightData.SpotAnglesAndSourceRadiusPacked.w));
@@ -211,7 +211,8 @@ int SMRTSamplesPerRay;
 float SMRTRayLengthScale;			// Directional lights
 float SMRTCotMaxRayAngleFromLight;	// Spot/point lights
 float SMRTTexelDitherScale;			// Currently only directional lights
-float SMRTWorldDitherScale;			// Currently only directional lights
+float SMRTExtrapolateSlope;
+float SMRTMaxSlopeBias;				// Currently only local lights
 uint bSMRTUseAdaptiveRayCount;
 
 // One pass projection
@@ -312,7 +313,7 @@ FVirtualShadowMapSampleResult ProjectLight(
 				bBackfaceCull,
 				WorldNormal,
 				SMRTTexelDitherScale,
-				SMRTWorldDitherScale,
+				SMRTExtrapolateSlope,
 				bUseAdaptiveRayCount);
 		#else
 			Result = TraceLocalLight(
@@ -328,6 +329,9 @@ FVirtualShadowMapSampleResult ProjectLight(
 				Noise,
 				bBackfaceCull,
 				WorldNormal,
+				SMRTTexelDitherScale,
+				SMRTExtrapolateSlope,
+				SMRTMaxSlopeBias,
 				bUseAdaptiveRayCount);
 		#endif // DIRECTIONAL_LIGHT
 		}
@@ -552,7 +556,6 @@ void VirtualShadowMapProjection(
 	const float Noise = InterleavedGradientNoise( SvPosition.xy, View.StateFrameIndexMod8 );
 
 	const FProjectionShadingInfo ShadingInfo = GetProjectionShadingInfo(PixelPos);
-	const float SubsurfaceOpacity = bIsHairInput ? 1.0f : GetSubsurfaceOpacityFromGbuffer( PixelPos );
 
 	FVirtualShadowMapSampleResult VisualizeResult = InitVirtualShadowMapSampleResult();
 
@@ -615,6 +618,7 @@ void VirtualShadowMapProjection(
 			Noise);
 		
 		float SSSTransmission = Result.ShadowFactor;
+		const float SubsurfaceOpacity = bIsHairInput ? 1.0f : GetSubsurfaceOpacityFromGbuffer( PixelPos );
 		if ( SubsurfaceOpacity < 1.0f )
 		{
 			SSSTransmission = ComputeSimpleSubsurfaceTransmissionFromVirtualShadowMap(
@@ -635,7 +639,7 @@ void VirtualShadowMapProjection(
 #if VISUALIZE_OUTPUT
 	if (InputType == INPUT_TYPE_GBUFFER)
 	{
-		OutputVisualize( PixelPos - uint2( ProjectionRect.xy ), VisualizeResult );
+		OutputVisualize( PixelPos , VisualizeResult );
 	}
 #endif // VISUALIZE_OUTPUT
 }
diff --git a/Engine/Shaders/Private/VirtualShadowMaps/VirtualShadowMapProjectionCommon.ush b/Engine/Shaders/Private/VirtualShadowMaps/VirtualShadowMapProjectionCommon.ush
--- a/Engine/Shaders/Private/VirtualShadowMaps/VirtualShadowMapProjectionCommon.ush
+++ b/Engine/Shaders/Private/VirtualShadowMaps/VirtualShadowMapProjectionCommon.ush
@@ -25,10 +25,36 @@ float CalcAbsoluteClipmapLevel(FVirtualShadowMapProjectionShaderData BaseProject
 
 int CalcClipmapLevel(FVirtualShadowMapProjectionShaderData BaseProjectionData, FLWCVector3 WorldPosition)
 {
-	float BiasedLevel = CalcAbsoluteClipmapLevel(BaseProjectionData, WorldPosition) + BaseProjectionData.ClipmapResolutionLodBias;
+	float BiasedLevel = CalcAbsoluteClipmapLevel(BaseProjectionData, WorldPosition) + BaseProjectionData.ResolutionLodBias;
 	return int(floor(BiasedLevel));
 }
 
+// Calculate the projected pixel footprint onto a local light, scaled in shadow map texels
+// ShadowTranslatedWorldPosition should already be translated by both the primary view and shadow ProjectionData.PreViewTranslation
+float VirtualShadowMapCalcPixelFootprintLocal(
+	FVirtualShadowMapProjectionShaderData ProjectionData,
+	float3 ShadowTranslatedWorldPosition,
+	float SceneDepth)
+{
+	// Compute footprint by projecting the approximate size of a camera pixel at the given depth to shadow space
+	// NOTE: This doesn't take the screen XY position/FOV into account, which may or may not be desirable.
+
+	// TODO: Roll into a uniform
+	float2 RadiusXY = 1.0f / (View.ViewSizeAndInvSize.xy * View.ViewToClip._m00_m11);
+	float RadiusScreen = min(RadiusXY.x, RadiusXY.y);
+	float DepthScale = SceneDepth * View.ViewToClip[2][3] + View.ViewToClip[3][3];
+
+	float RadiusWorld = DepthScale * RadiusScreen;
+
+	float4 ShadowUVz = mul(float4(ShadowTranslatedWorldPosition, 1.0f), ProjectionData.TranslatedWorldToShadowUVMatrix);
+
+	float4 RadiusClipH = mul(float4(RadiusWorld, 0.0f, ShadowUVz.w, 1.0f), ProjectionData.ShadowViewToClipMatrix);
+	float RadiusClip = abs(RadiusClipH.x / RadiusClipH.w);
+	float Footprint = RadiusClip * float(2 * VSM_VIRTUAL_MAX_RESOLUTION_XY);
+
+	return Footprint;
+}
+
 struct FVirtualShadowMapSample
 {
 	float Depth;
@@ -287,11 +313,11 @@ uint VirtualShadowMapGetCubeFace( float3 Dir )
 struct FVirtualShadowMapSampleResult
 {
 	bool bValid;
-	half ShadowFactor;	// 0 = fully occluded, 1 = no occlusion
+	float ShadowFactor;	// 0 = fully occluded, 1 = no occlusion
 	float OccluderDistance;
 
 	// Debug data; do not reference in non-debug permutations or performance may be affected!
-	uint ClipmapIndexOrMipLevel;	
+	uint ClipmapIndexOrMipLevel;
 	uint RayCount;
 	uint2 VirtualTexelAddress;
 	uint2 PhysicalTexelAddress;
diff --git a/Engine/Shaders/Private/VirtualShadowMaps/VirtualShadowMapProjectionDirectional.ush b/Engine/Shaders/Private/VirtualShadowMaps/VirtualShadowMapProjectionDirectional.ush
--- a/Engine/Shaders/Private/VirtualShadowMaps/VirtualShadowMapProjectionDirectional.ush
+++ b/Engine/Shaders/Private/VirtualShadowMaps/VirtualShadowMapProjectionDirectional.ush
@@ -23,7 +23,7 @@
 	float2 DepthSlopeUV = -NormalPlaneUV.xy / NormalPlaneUV.z;
 
 	// Clamp to avoid excessive degenerate slope biases causing flickering lit pixels
-	float2 Clamp = 0.02f;
+	float2 Clamp = 0.05f;
 	DepthSlopeUV = bClamp ? clamp(DepthSlopeUV, -Clamp, Clamp) : DepthSlopeUV;
 
 	return DepthSlopeUV;
@@ -70,6 +70,7 @@
 	int ClipmapId;			// Pre-offset virtual shadow map ID of specific clipmap level
 	float3 RayStartUVZ;	
 	float3 RayStepUVZ;
+	float ExtrapolateSlope;
 	uint2 PhysicalTexelAddress;
 	uint2 VirtualTexelAddress;
 	int SampledClipmapId;
@@ -82,7 +83,8 @@
 	float RayLength,
 	float RayStartOffset,
 	float2 DepthSlopeUV,
-	float2 TexelOffset)
+	float2 TexelOffset,
+	float ExtrapolateSlope)
 {
 	float3 RayStart = RayOriginShadowTranslatedWorld + RayDir * RayStartOffset;
 	float3 RayVector = RayDir * RayLength;
@@ -107,25 +109,26 @@
 	Result.ClipmapId = ProjectionData.VirtualShadowMapId;
 	Result.RayStartUVZ = RayStartUVZ;
 	Result.RayStepUVZ = RayStepUVZ;
+	// Max depth slope for slope-based extrapolation when using SMRT. Scale so that it doesn't change based on ZRangeScale.
+	Result.ExtrapolateSlope = abs(ExtrapolateSlope * ProjectionData.ShadowViewToClipMatrix._33);
 	Result.VirtualTexelAddress = uint2(0xFFFFFFFF, 0xFFFFFFFF);
 	Result.PhysicalTexelAddress = uint2(0xFFFFFFFF, 0xFFFFFFFF);
 	return Result;
 }
 
 FSMRTSample SMRTFindSample(inout FSMRTClipmapRayState RayState, float SampleTime)
-{
-	FSMRTSample Sample = (FSMRTSample)0;
-	Sample.bValid = false;
-	
+{	
 	const float3 SampleUVZ = RayState.RayStartUVZ + RayState.RayStepUVZ * SampleTime;
 	FVirtualShadowMapSample ShadowSample = SampleVirtualShadowMapClipmap(RayState.ClipmapId, SampleUVZ.xy);
-	
+
+	FSMRTSample Sample = InitSMRTSample();
+	Sample.bValid = ShadowSample.bValid;
+	Sample.ReferenceDepth = SampleUVZ.z;
+	Sample.ExtrapolateSlope = RayState.ExtrapolateSlope;
+
 	if (ShadowSample.bValid)
 	{
-		Sample.bValid = true;
 		Sample.SampleDepth = ShadowSample.Depth;
-		Sample.ReferenceDepth = SampleUVZ.z;
-		Sample.CompareToleranceScale = RayState.RayStepUVZ.z;
 
 		// Debug
 		RayState.VirtualTexelAddress = ShadowSample.VirtualTexelAddress;
@@ -136,17 +139,6 @@
 	return Sample;
 }
 
-float SMRTComputeOccluderDistance(FSMRTClipmapRayState RayState, float SampleDepth)
-{
-	float4x4 ViewToClip = GetVirtualShadowMapProjectionData(RayState.ClipmapId).ShadowViewToClipMatrix;
-
-	// No perspective projection, so simple difference gets us the distance
-	float OccluderViewZ = (SampleDepth            - ViewToClip._43) / ViewToClip._33;
-	float ReceiverViewZ = (RayState.RayStartUVZ.z - ViewToClip._43) / ViewToClip._33;
-	float Result = ReceiverViewZ - OccluderViewZ;
-	return max(VIRTUAL_SHADOW_MAP_MIN_OCCLUDER_DISTANCE, Result);
-}
-
 // Instantiate SMRTRayCast for FSMRTClipmapRayState
 #define SMRT_TEMPLATE_RAY_STRUCT FSMRTClipmapRayState
 #include "VirtualShadowMapSMRTTemplate.ush"
@@ -167,11 +159,11 @@
 }
 
 // Normal and LightDirection should be normalized
-bool IsBackfaceToDirectionalLight(half3 Normal, half3 LightDirection, half LightSourceRadius)
+bool IsBackfaceToDirectionalLight(float3 Normal, float3 LightDirection, float LightSourceRadius)
 {
 	// Allow a minimum of ~5 degrees of wiggle room to account for normal issues
-	half MinSinAlpha = 0.1;
-	half SinAlpha = max(abs(LightSourceRadius), MinSinAlpha);
+	float MinSinAlpha = 0.1;
+	float SinAlpha = max(abs(LightSourceRadius), MinSinAlpha);
 	return dot(Normal, LightDirection) < -SinAlpha;
 }
 
@@ -189,7 +181,7 @@
 	bool bCullBackface,
 	half3 WorldNormal,
 	float TexelDitherScale,
-	float WorldDitherScale,
+	float ExtrapolateSlope,
 	bool bUseAdaptiveRayCount)
 {
 	FLWCVector3 WorldPosition = LWCSubtract(LWCPromote(TranslatedWorldPosition), PrimaryView.PreViewTranslation);
@@ -207,7 +199,7 @@
 		const FVirtualShadowMapProjectionShaderData ProjectionData = GetVirtualShadowMapProjectionData(ClipmapId);
 
 		// This function is designed to vary in depth in roughly the same way as shadow texel density does,
-		// thus the dependence on ClipmapResolutionLodBias. The goal here is to define a texel dither that
+		// thus the dependence on Clipmap ResolutionLodBias. The goal here is to define a texel dither that
 		// is "smooth" and doesn't have obvious artifacts at page edges of different clipmap levels.
 		// 
 		// This could become a problem later if we break too far from that global resolution bias. That said,
@@ -215,18 +207,14 @@
 		// in some different way.
 		//
 		// 0.5 is arbitrary but makes it consistent with the previous dither scale cvar
-		float TexelDither = 0.0f;
+		float DitherScale = 0.0f;
 		if (TexelDitherScale > 0.0f)
 		{
-			TexelDither =
+			DitherScale =
 				((0.5f / float(CalcLevelDimsTexels(0))) * TexelDitherScale * DistanceFromViewOrigin) /
-				(exp2(ProjectionData.ClipmapLevel - ProjectionData.ClipmapResolutionLodBias));
+				(exp2(ProjectionData.ClipmapLevel - ProjectionData.ResolutionLodBias));
 		}
 
-		// ~Constant world-space dither/blur
-		float WorldDither = WorldDitherScale * ProjectionData.ShadowViewToClipMatrix._11;
-		float DitherScale = max(TexelDither, WorldDither);
-
 		float2 DepthSlopeUV = ComputeDepthSlopeDirectionalUV(ProjectionData, WorldNormal);
 
 		// Compute max ray length based on view depth
@@ -260,7 +248,8 @@
 				RayLength,
 				RayStartOffset,
 				DepthSlopeUV,
-				TexelOffset);
+				TexelOffset,
+				ExtrapolateSlope);
 			FSMRTResult SMRTResult = SMRTRayCast(RayState, SamplesPerRay, Noise);
 
 			// Debug output (DCE'd if not used)
diff --git a/Engine/Shaders/Private/VirtualShadowMaps/VirtualShadowMapProjectionSpot.ush b/Engine/Shaders/Private/VirtualShadowMaps/VirtualShadowMapProjectionSpot.ush
--- a/Engine/Shaders/Private/VirtualShadowMaps/VirtualShadowMapProjectionSpot.ush
+++ b/Engine/Shaders/Private/VirtualShadowMaps/VirtualShadowMapProjectionSpot.ush
@@ -13,11 +13,51 @@
 #include "VirtualShadowMapProjectionCommon.ush"
 #include "VirtualShadowMapSMRTCommon.ush"
 
+// Possibly parameterize this at some point, but tradeoffs are not always obvious
+#define VIRTUAL_SHADOW_MAP_MAX_SLOPE_BIAS 100.0f
+
+float2 ComputeDepthSlopeLocalUV(
+	int VirtualShadowMapId,
+	float3 ShadowTranslatedWorldPosition,
+	float3 WorldNormal)
+{
+	FVirtualShadowMapProjectionShaderData ProjectionData = GetVirtualShadowMapProjectionData(VirtualShadowMapId);
+
+	float4 NormalPlaneTranslatedWorld = float4(WorldNormal, -dot(WorldNormal, ShadowTranslatedWorldPosition));
+	float4 NormalPlaneUV = mul(NormalPlaneTranslatedWorld, ProjectionData.TranslatedWorldToShadowUVNormalMatrix);
+	float2 DepthSlopeUV = -NormalPlaneUV.xy / NormalPlaneUV.z;
+	return DepthSlopeUV;
+}
+
+float ComputeOptimalSlopeBiasLocal(
+	int VirtualShadowMapId,
+	float3 ShadowTranslatedWorldPosition,
+	float2 DepthSlopeUV,
+	bool bClamp = true)
+{
+	// NOTE: Better to reload the necessary data here based on VSM ID than keep it alive in registers!
+	FVirtualShadowMapProjectionShaderData ProjectionData = GetVirtualShadowMapProjectionData(VirtualShadowMapId);
+
+	float4 ShadowUVz = mul(float4(ShadowTranslatedWorldPosition, 1.0f), ProjectionData.TranslatedWorldToShadowUVMatrix);
+	ShadowUVz.xyz /= ShadowUVz.w;
+	FShadowPageTranslationResult Page = ShadowVirtualToPhysicalUV(ProjectionData.VirtualShadowMapId, ShadowUVz.xy);
+
+	float MipLevelDim = float(CalcLevelDimsTexels(Page.LODOffset));
+	float2 TexelCenter = float2(Page.VirtualTexelAddress) + 0.5f;
+	float2 TexelCenterOffset = TexelCenter - Page.VirtualTexelAddressFloat;
+	float2 TexelCenterOffsetUV = TexelCenterOffset / MipLevelDim;
+	float OptimalSlopeBias = 2.0f * max(0.0f, dot(DepthSlopeUV, TexelCenterOffsetUV));
+
+	return OptimalSlopeBias;
+}
+
+
 struct FSMRTSingleRayState
 {
 	int VirtualShadowMapId;
 	float3 RayStartUVz;	
 	float3 RayStepUVz;
+	float ExtrapolateSlope;
 	// Debug output
 	int MipLevel;
 	uint2 VirtualTexelAddress;
@@ -28,7 +68,10 @@
 	int VirtualShadowMapId,
 	// Shadow PreViewTranslation already added in
 	float3 RayStartShadowTranslatedWorld,
-	float3 RayEndShadowTranslatedWorld)
+	float3 RayEndShadowTranslatedWorld,
+	float ExtrapolateSlope,
+	float DepthBias,
+	bool bClampUVs = true)
 {
 	FVirtualShadowMapProjectionShaderData ProjectionData = GetVirtualShadowMapProjectionData(VirtualShadowMapId);
 
@@ -44,10 +87,19 @@
 	RayEndUVz.xyz = RayEndUVz.xyz / RayEndUVz.w;
 	float3 RayStepUVz = RayEndUVz.xyz - RayStartUVz.xyz;
 
+	// Offsets can move it slightly off the edge of spotlight projection. Clamp to edge.
+	// Note we do *not* want to do this if we are tracing dual face cubemap rays!
+	if (bClampUVs)
+	{
+		RayStartUVz.xy = saturate(RayStartUVz.xy);
+	}
+	RayStartUVz.z += DepthBias;
+
 	FSMRTSingleRayState Result;
 	Result.VirtualShadowMapId = VirtualShadowMapId;
 	Result.RayStartUVz = RayStartUVz.xyz;
 	Result.RayStepUVz = RayStepUVz.xyz;
+	Result.ExtrapolateSlope = ExtrapolateSlope * RayStepUVz.z;
 	Result.MipLevel = 0;
 	Result.VirtualTexelAddress = uint2(0xFFFFFFFF, 0xFFFFFFFF);
 	Result.PhysicalTexelAddress = uint2(0xFFFFFFFF, 0xFFFFFFFF);
@@ -56,10 +108,13 @@
 
 FSMRTSample SMRTFindSample(inout FSMRTSingleRayState RayState, float SampleTime)
 {
-	FSMRTSample Sample = (FSMRTSample)0;
-	Sample.bValid = false;
 	float3 UVz = RayState.RayStartUVz.xyz + RayState.RayStepUVz.xyz * SampleTime;
 
+	FSMRTSample Sample = InitSMRTSample();
+	Sample.bValid = false;
+	Sample.ReferenceDepth = UVz.z;
+	Sample.ExtrapolateSlope = RayState.ExtrapolateSlope;
+
 	if (all(UVz.xy == saturate(UVz.xy)))
 	{
 		FVirtualShadowMapSample SmSample = SampleVirtualShadowMap(RayState.VirtualShadowMapId, UVz.xy);
@@ -67,8 +122,7 @@
 		{
 			Sample.bValid = true;
 			Sample.SampleDepth = SmSample.Depth;
-			Sample.ReferenceDepth = UVz.z;
-			Sample.CompareToleranceScale = RayState.RayStepUVz.z;
+			
 			// Debug output
 			RayState.MipLevel = SmSample.MipLevel;
 			RayState.PhysicalTexelAddress = SmSample.PhysicalTexelAddress;
@@ -79,18 +133,6 @@
 	return Sample;
 }
 
-float SMRTComputeOccluderDistance(FSMRTSingleRayState RayState, float SampleDepth)
-{
-	// Avoid keeping around the registers
-	float4x4 ViewToClip = GetVirtualShadowMapProjectionData(RayState.VirtualShadowMapId).ShadowViewToClipMatrix;
-
-	// View space Z delta is a good enough stand-in for euclidean distance for now given our bent ray tests and so on
-	float OccluderViewZ = ViewToClip._43 / (SampleDepth            - ViewToClip._33);
-	float ReceiverViewZ = ViewToClip._43 / (RayState.RayStartUVz.z - ViewToClip._33);		
-	float Result = ReceiverViewZ - OccluderViewZ;	
-	return max(VIRTUAL_SHADOW_MAP_MIN_OCCLUDER_DISTANCE, Result);
-}
-
 // Instantiate SMRTRayCast for FSMRTSingleRayState
 #define SMRT_TEMPLATE_RAY_STRUCT FSMRTSingleRayState
 #include "VirtualShadowMapSMRTTemplate.ush"
@@ -101,45 +143,45 @@
 {
 	FSMRTSingleRayState Face0;
 	FSMRTSingleRayState Face1;
-	bool bSampleInFace0;
+	bool bSampleInFace1;
 };
 
 FSMRTTwoCubeFaceRayState SMRTTwoCubeFaceRayInitialize(
 	int VirtualShadowMapId0,
 	int VirtualShadowMapId1,
 	float3 RayStartShadowTranslatedWorld,
-	float3 RayEndShadowTranslatedWorld)
+	float3 RayEndShadowTranslatedWorld,
+	float ExtrapolateSlope,
+	float DepthBias)
 {
 	FSMRTTwoCubeFaceRayState Result;
-	Result.Face0 = SMRTSingleRayInitialize(VirtualShadowMapId0, RayStartShadowTranslatedWorld, RayEndShadowTranslatedWorld);
-	Result.Face1 = SMRTSingleRayInitialize(VirtualShadowMapId1, RayStartShadowTranslatedWorld, RayEndShadowTranslatedWorld);
-	Result.bSampleInFace0 = false;
+	Result.Face0 = SMRTSingleRayInitialize(VirtualShadowMapId0, RayStartShadowTranslatedWorld, RayEndShadowTranslatedWorld, ExtrapolateSlope, DepthBias, false);
+	Result.Face1 = SMRTSingleRayInitialize(VirtualShadowMapId1, RayStartShadowTranslatedWorld, RayEndShadowTranslatedWorld, ExtrapolateSlope, DepthBias, false);
+	Result.bSampleInFace1 = true;
 	return Result;
 }
 
 FSMRTSample SMRTFindSample(inout FSMRTTwoCubeFaceRayState RayState, float SampleTime)
 {
-	FSMRTSample Sample = SMRTFindSample(RayState.Face0, SampleTime);
-	RayState.bSampleInFace0 = Sample.bValid;
-	if (!RayState.bSampleInFace0)
+	// NOTE: Traces from END to START, so we start in face1
+	FSMRTSample Sample = InitSMRTSample();
+	if (RayState.bSampleInFace1)
 	{
 		Sample = SMRTFindSample(RayState.Face1, SampleTime);
+		if (!Sample.bValid)
+		{
+			Sample = SMRTFindSample(RayState.Face0, SampleTime);
+			Sample.bResetExtrapolation = true;
+			RayState.bSampleInFace1 = false;
+		}
+	}
+	else
+	{
+		Sample = SMRTFindSample(RayState.Face0, SampleTime);
 	}
 	return Sample;
 }
 
-float SMRTComputeOccluderDistance(FSMRTTwoCubeFaceRayState RayState, float SampleDepth)
-{
-	if (RayState.bSampleInFace0)
-	{
-		return SMRTComputeOccluderDistance(RayState.Face0, SampleDepth);
-	}
-	else
-	{
-		return SMRTComputeOccluderDistance(RayState.Face1, SampleDepth);
-	}
-}
-
 // Instantiate SMRTRayCast for FSMRTTwoCubeFaceRayState
 #define SMRT_TEMPLATE_RAY_STRUCT FSMRTTwoCubeFaceRayState
 #include "VirtualShadowMapSMRTTemplate.ush"
@@ -150,12 +192,14 @@
 	int VirtualShadowMapId,
 	float3 RayStartShadowTranslatedWorld,
 	float3 RayEndShadowTranslatedWorld,
+	float DepthBias,
 	int NumSteps,
 	float StepOffset,
+	float ExtrapolateSlope,
 	// To get debug data out
 	inout FSMRTSingleRayState OutRayState)
 {
-	OutRayState = SMRTSingleRayInitialize(VirtualShadowMapId, RayStartShadowTranslatedWorld, RayEndShadowTranslatedWorld);
+	OutRayState = SMRTSingleRayInitialize(VirtualShadowMapId, RayStartShadowTranslatedWorld, RayEndShadowTranslatedWorld, ExtrapolateSlope, DepthBias);
 	return SMRTRayCast(OutRayState, NumSteps, StepOffset);
 }
 
@@ -163,36 +207,38 @@
 	int VirtualShadowMapId,
 	float3 RayStartShadowTranslatedWorld,
 	float3 RayEndShadowTranslatedWorld,
+	float DepthBias,
 	int NumSteps,
 	float StepOffset,
+	float ExtrapolateSlope,
 	// To get debug data out
 	inout FSMRTSingleRayState OutRayState)
 {
 	uint RayStartFace = VirtualShadowMapGetCubeFace(RayStartShadowTranslatedWorld);
 	uint RayEndFace   = VirtualShadowMapGetCubeFace(RayEndShadowTranslatedWorld);
-		
+
 	FSMRTResult Result;
-	if (RayStartFace != RayEndFace)
+	if (WaveActiveAnyTrue(RayStartFace != RayEndFace))
 	{
 		FSMRTTwoCubeFaceRayState RayState = SMRTTwoCubeFaceRayInitialize(
 			VirtualShadowMapId + RayStartFace,
 			VirtualShadowMapId + RayEndFace,
-			RayStartShadowTranslatedWorld, RayEndShadowTranslatedWorld);
+			RayStartShadowTranslatedWorld, RayEndShadowTranslatedWorld, ExtrapolateSlope, DepthBias);
 		Result = SMRTRayCast(RayState, NumSteps, StepOffset);
 
-		if (RayState.bSampleInFace0)
+		if (RayState.bSampleInFace1)
 		{
-			OutRayState = RayState.Face0;
-		}
-		else
-		{
 			OutRayState = RayState.Face1;
 		}
+		else
+		{
+			OutRayState = RayState.Face0;
+		}
 	}
 	else
 	{
 		// Fast path: ray stays on a single cube map face (effectively a spot light)
-		OutRayState = SMRTSingleRayInitialize(VirtualShadowMapId + RayEndFace, RayStartShadowTranslatedWorld, RayEndShadowTranslatedWorld);
+		OutRayState = SMRTSingleRayInitialize(VirtualShadowMapId + RayEndFace, RayStartShadowTranslatedWorld, RayEndShadowTranslatedWorld, ExtrapolateSlope, DepthBias);
 		Result = SMRTRayCast(OutRayState, NumSteps, StepOffset);
 	}
 
@@ -258,6 +304,20 @@
 	return IsBackfaceToLocalLight(ToLight, Normal, LightSourceRadius, Dummy1, Dummy2, Dummy3);
 }
 
+// Shadow resolution scale that roughly matches the LOD function in VSM page marking
+float ComputeShadowResolutionScale(FVirtualShadowMapProjectionShaderData ProjectionData, float SceneDepth)
+{
+	const float2 RadiusXY = 1.0f / (View.ViewSizeAndInvSize.xy * View.ViewToClip._11_22);
+	const float RadiusScreen = min(RadiusXY.x, RadiusXY.y);
+	const float DepthScale = SceneDepth * View.ViewToClip._34 + View.ViewToClip._44;
+	const float ScreenPixelRadiusWorld = DepthScale * RadiusScreen;
+
+	// Clamp the minimum here since our shadows sadly don't actually have infinite resolution near the camera...
+	float ShadowResolutionScale = max(0.1f, ScreenPixelRadiusWorld * exp2(ProjectionData.ResolutionLodBias));
+
+	return ShadowResolutionScale;
+}
+
 FVirtualShadowMapSampleResult TraceLocalLight(
 	int VirtualShadowMapId,
 	FLightShaderParameters Light,
@@ -271,6 +331,9 @@
 	const float Noise,
 	const bool bCullBackface,
 	float3 WorldNormal,
+	float TexelDitherScale,
+	float ExtrapolateSlope,
+	float MaxSlopeBias,
 	bool bUseAdaptiveRayCount)
 {
 	const uint LightType = Light.SpotAngles.x > -2.0f ? LIGHT_TYPE_SPOT : LIGHT_TYPE_POINT;
@@ -291,12 +354,37 @@
 
 		uint TileID = MortonEncode( PixelPos );
 
+		int OffsetVirtualShadowMapId = VirtualShadowMapId;
+		if (LightType == LIGHT_TYPE_POINT)
+		{
+			OffsetVirtualShadowMapId += VirtualShadowMapGetCubeFace(-ToLight);
+		}
+
 		// Move from PrimaryView translated world to Shadow translated world
-		// NOTE: Cube map faces all share the same PreViewTranslation
-		FVirtualShadowMapProjectionShaderData BaseProjectionData = GetVirtualShadowMapProjectionData(VirtualShadowMapId);
-		float3 PrimaryToShadowTranslation = LWCToFloat(LWCSubtract(BaseProjectionData.PreViewTranslation, PrimaryView.PreViewTranslation));
-		float3 ShadowTranslatedWorld = TranslatedWorldPosition + PrimaryToShadowTranslation;
+		// NOTE: Cube map faces all share the same PreViewTranslation and projection depth scale
+		FVirtualShadowMapProjectionShaderData ProjectionData = GetVirtualShadowMapProjectionData(OffsetVirtualShadowMapId);
+		const float3 PrimaryToShadowTranslation = LWCToFloat(LWCSubtract(ProjectionData.PreViewTranslation, PrimaryView.PreViewTranslation));
+		const float3 ShadowTranslatedWorld = TranslatedWorldPosition + PrimaryToShadowTranslation;
 
+		const float2 DepthSlopeUV = ComputeDepthSlopeLocalUV(OffsetVirtualShadowMapId, ShadowTranslatedWorld, WorldNormal);
+
+		const float ShadowResolutionScale = ComputeShadowResolutionScale(ProjectionData, SceneDepth);
+		// Rough scaling from world units -> post projection depth.
+		// We use a radially symmetric mapping here instead of the true function to minimize issues crossing cubemap faces.
+		const float ShadowDepthFunctionScale = abs(ProjectionData.ShadowViewToClipMatrix._43 / length2(ToLight));
+
+		const float DitherScale = TexelDitherScale * ShadowResolutionScale;
+		const float MaxDitherSlope = MaxSlopeBias * DitherScale;
+		const float MaxDepthBias = MaxSlopeBias * ShadowResolutionScale * ShadowDepthFunctionScale;
+
+		// Moving the ray origin means we need to attempt to account for local receiver geometry to avoid
+		// incorrect self-shadowing. To this end we move the point on a plane aligned with the shading normal.
+		// It would be much better to use geometric (faceted) normals - as with optimal bias - but we do not
+		// have access to those. Shading normals that stray too far from the real geometry will cause artifacts.
+		const float3x3 ConeTangentToWorldMatrix = GetTangentBasis(ConeAxis);
+		const float3 NormalPlaneTangent = mul(ConeTangentToWorldMatrix, WorldNormal);
+		const float2 DepthSlopeTangent = -NormalPlaneTangent.xy / NormalPlaneTangent.z;
+
 		uint i = 0;
 		uint RayMissCount = 0;
 		for( ; i < MaxRayCount; i++ )
@@ -304,7 +392,6 @@
 			RandomSequence RandSequence;
 #if 0
 			RandomSequence_Initialize( RandSequence, PixelPos, i, View.StateFrameIndex, MaxRayCount );
-
 			float4 RandSample = RandomSequence_GenerateSample4D( RandSequence );
 #else
 			// Combine pixel-level and sample-level bits into the sample index (visible structure will be hidden by owen scrambling of the index)
@@ -320,20 +407,37 @@
 			float SinTheta2 = dot( LightUV, LightUV );
 			float SinTheta = sqrt( SinTheta2 );
 			float CosTheta = sqrt( 1 - SinTheta2 );
+			float3 Dir = mul(float3(LightUV, CosTheta), ConeTangentToWorldMatrix);
 
-			float3 Dir = TangentToWorld( float3( LightUV, CosTheta ), ConeAxis );
+			float3 RayStartWithDither = ShadowTranslatedWorld;
 
+			if (DitherScale > 0.0f)
+			{
+				// Offset following the plane of the receiver
+				float2 RandomOffset = (RandSample.zw - 0.5f) * DitherScale;
+				float OffsetZ = min(MaxDitherSlope, 2.0f * max(0.0f, dot(DepthSlopeTangent, RandomOffset)));
+				float3 Offset = mul(float3(RandomOffset, OffsetZ), ConeTangentToWorldMatrix);
+				RayStartWithDither += Offset;
+			}
+
 			// The further afield our rays go the poorer our approximation is as the "bend" due to our testing
 			// against a shadow map instead of along the ray increases. Thus we avoid going all the way to the light
 			// where the projection becomes extreme.
-			
+
 			// Clip ray to max angle to limit penumbra sizes
 			float ClampedEnd = DistToLight * 0.75 * saturate( 1.5 / ( CosTheta + CotMaxRayAngleFromLight * SinTheta ) );
 			float Start = min(RayStartOffset, ClampedEnd - 1e-6f);
 
-			float3 RayStart = ShadowTranslatedWorld + Dir * Start;
-			float3 RayEnd = ShadowTranslatedWorld + Dir * ClampedEnd;
+			float3 RayStart = RayStartWithDither + Dir * Start;
+			float3 RayEnd = RayStartWithDither + Dir * ClampedEnd;
 
+			float DepthBias = 0.0f;
+			if (MaxDepthBias > 0.0f)
+			{
+				// Clamp depth bias to avoid excessive degenerate slope biases causing flickering lit pixels
+				DepthBias = min(MaxDepthBias, ComputeOptimalSlopeBiasLocal(OffsetVirtualShadowMapId, RayStart, DepthSlopeUV));
+			}
+
 			FSMRTResult SMRTResult;
 			FSMRTSingleRayState RayState;	// For debug output
 			if ( LightType == LIGHT_TYPE_SPOT )
@@ -342,8 +446,10 @@
 					VirtualShadowMapId,
 					RayStart,
 					RayEnd,
+					DepthBias,
 					SamplesPerRay,
 					StepOffset,
+					ExtrapolateSlope,
 					RayState);
 			}
 			else
@@ -352,8 +458,10 @@
 					VirtualShadowMapId,
 					RayStart,
 					RayEnd,
+					DepthBias,
 					SamplesPerRay,
 					StepOffset,
+					ExtrapolateSlope,
 					RayState);
 			}
 
@@ -363,7 +471,6 @@
 			Result.ClipmapIndexOrMipLevel = RayState.MipLevel;
 			Result.VirtualTexelAddress = RayState.VirtualTexelAddress;
 			Result.PhysicalTexelAddress = RayState.PhysicalTexelAddress;
-			Result.OccluderDistance = SMRTResult.OccluderDistance;
 
 			if (bUseAdaptiveRayCount)
 			{
@@ -411,6 +518,7 @@
 
 		Result.ShadowFactor = SMRTComputeRayMissFactor(RayMissCount, RayCount, MaxRayCount, Noise);
 		Result.RayCount = RayCount;
+		Result.OccluderDistance = -1.0f;	// TODO?
 	}
 
 	return Result;
diff --git a/Engine/Shaders/Private/VirtualShadowMaps/VirtualShadowMapProjectionStructs.ush b/Engine/Shaders/Private/VirtualShadowMaps/VirtualShadowMapProjectionStructs.ush
--- a/Engine/Shaders/Private/VirtualShadowMaps/VirtualShadowMapProjectionStructs.ush
+++ b/Engine/Shaders/Private/VirtualShadowMaps/VirtualShadowMapProjectionStructs.ush
@@ -17,12 +17,13 @@ struct FVirtualShadowMapProjectionShaderData
 	float LightSourceRadius;
 
 	FLWCVector3 ClipmapWorldOrigin;
+	float ResolutionLodBias;
+
 	int2 ClipmapCornerOffset;
 	int ClipmapIndex;		// 0..ClipmapLevelCount-1
 	int ClipmapLevel;
 
-	int ClipmapLevelCount;	
-	float ClipmapResolutionLodBias;
+	int ClipmapLevelCount;
 
 	// Derived data for convenience when passing the structure around
 	int VirtualShadowMapId;
@@ -64,7 +65,7 @@ FVirtualShadowMapProjectionShaderData DecodeVirtualShadowMapProjectionData(ByteA
 	Result.LightSourceRadius				= asfloat(ProjectionData.Load (Base + (16*17 + 4*3)));
 
 	float3 NegativeClipmapWorldOriginOffset = asfloat(ProjectionData.Load3(Base + (16*18 + 4*0)));
-	Result.ClipmapResolutionLodBias			= asfloat(ProjectionData.Load (Base + (16*18 + 4*3)));
+	Result.ResolutionLodBias				= asfloat(ProjectionData.Load (Base + (16*18 + 4*3)));
 		
 	Result.ClipmapCornerOffset				= asint  (ProjectionData.Load2(Base + (16*19 + 4*0)));
 	Result.ClipmapIndex						= asint  (ProjectionData.Load (Base + (16*19 + 4*2)));
diff --git a/Engine/Shaders/Private/VirtualShadowMaps/VirtualShadowMapSMRTCommon.ush b/Engine/Shaders/Private/VirtualShadowMaps/VirtualShadowMapSMRTCommon.ush
--- a/Engine/Shaders/Private/VirtualShadowMaps/VirtualShadowMapSMRTCommon.ush
+++ b/Engine/Shaders/Private/VirtualShadowMaps/VirtualShadowMapSMRTCommon.ush
@@ -6,23 +6,29 @@ VirtualShadowMapSMRTCommon.ush:
 
 #pragma once
 
-#ifndef SMRT_EXTRAPOLATE_WITH_SLOPE
-#define SMRT_EXTRAPOLATE_WITH_SLOPE 1
-#endif
-
 struct FSMRTSample
 {
 	bool bValid;	
 	float SampleDepth;
 	float ReferenceDepth;
-	float CompareToleranceScale;
+	float ExtrapolateSlope;
+	bool bResetExtrapolation;
 };
 
+FSMRTSample InitSMRTSample()
+{
+	FSMRTSample Result;
+	Result.bValid = false;
+	Result.SampleDepth = 0;
+	Result.ReferenceDepth = 0;
+	Result.ExtrapolateSlope = 0;
+	Result.bResetExtrapolation = false;
+	return Result;
+}
+
 struct FSMRTResult
 {
 	bool bValidHit;
-	float OccluderDistance;
-	int SampleMissCount;
 };
 
 float SMRTComputeRayMissFactor(uint RayMissCount, uint RayCount, uint MaxRayCount, half Noise)
diff --git a/Engine/Shaders/Private/VirtualShadowMaps/VirtualShadowMapSMRTTemplate.ush b/Engine/Shaders/Private/VirtualShadowMaps/VirtualShadowMapSMRTTemplate.ush
--- a/Engine/Shaders/Private/VirtualShadowMaps/VirtualShadowMapSMRTTemplate.ush
+++ b/Engine/Shaders/Private/VirtualShadowMaps/VirtualShadowMapSMRTTemplate.ush
@@ -7,8 +7,8 @@
 1) Include "SMRTCommon.ush"
 2) Define a "ray state" structure (FRayState here, but must be unique)
 3) Implement FSMRTSample SMRTFindSample(inout FRayState RayState, float SampleTime)
-4) Implement SMRTComputeOccluderDistance(FRayState RayState, float SampleDepth)
-5) Include this file:
+  - NOTE: ReferenceDepth must always be set, even if the sample is invalid
+4) Include this file:
 #define SMRT_TEMPLATE_RAY_STRUCT FRayState
 #include "VirtualShadowMapSMRTTemplate.ush"
 #undef SMRT_TEMPLATE_RAY_STRUCT
@@ -32,15 +32,13 @@
 	float DepthHistoryTime = -1;
 	float DepthSlope = 0;
 
-	int SampleMissCount = 0;
-	float OccluderDistance = 100000;
-
 	const float TimeScale = -1.0f / NumSteps;
 	const float TimeBias = 1.0f + ( 1.0 - StepOffset ) * TimeScale;
 
-	// Backward delta. Time[i] - Time[i-1]
-	const float DeltaTimeScale = 2.0 * Pow2( TimeScale );
-	const float DeltaTimeBias = 2.0 * TimeScale * ( TimeBias - 0.5 * TimeScale );
+	// This doesn't get used in the first iteration of the loop and so will always
+	// be valid by the time that it does.
+	float PrevReferenceDepth = -1;
+	float PrevSampleTime = -1;
 
 	bool bValidHit =  false;
 #if SMRT_TEMPLATE_STATIC_SAMPLES_PER_RAY >= 0
@@ -49,14 +47,19 @@
 #endif
 	for (int i = 0; i <= NumSteps; i++)
 	{
-		const float SampleTime	= i == NumSteps ? 0 : Pow2( TimeScale * i + TimeBias );
-		const float DeltaTime	= i == NumSteps ? Pow2( TimeScale * StepOffset ) : DeltaTimeScale * i + DeltaTimeBias;
+		const float SampleTime = ( i == NumSteps ) ? 0 : Pow2( TimeScale * i + TimeBias );
 
 		FSMRTSample Sample = SMRTFindSample(RayState, SampleTime);
+		const float ReferenceDepth = Sample.ReferenceDepth;
+
+		if (Sample.bResetExtrapolation)
+		{
+			DepthSlope = Sample.ExtrapolateSlope;
+		}
+
 		if (Sample.bValid)
 		{
 			const float SampleDepth = Sample.SampleDepth;
-			const float ReferenceDepth = Sample.ReferenceDepth;
 			if (DepthHistory < 0)
 			{
 				// First valid sample we've seen. Do a regular depth compare.
@@ -64,25 +67,27 @@
 				DepthHistoryTime = SampleTime;
 				if (SampleDepth > ReferenceDepth)
 				{
-					OccluderDistance = SMRTComputeOccluderDistance(RayState, SampleDepth);
 					bValidHit = true;
 					break;
 				}
 			}
 			else
 			{
+				const float DeltaReferenceDepth = ReferenceDepth - PrevReferenceDepth;
+
 				// Add a small relative error to the comparison to avoid missing surfaces due to numeric precision issues
 				// Without this there are occasionally flickering fireflies in fully shadowed regions with SMRT
-				const float EpsScale = 1.5f;
-				const float CompareTolerance = abs(Sample.CompareToleranceScale * DeltaTime) * EpsScale;
+				const float EpsScale = 1.05f;
+				const float CompareTolerance = abs(DeltaReferenceDepth) * EpsScale;
+								
 				const bool bBehind = (SampleDepth - ReferenceDepth) > CompareTolerance;
-
-				float DeltaHistoryTime = SampleTime - DepthHistoryTime;
 				float DepthForComparison = SampleDepth;
 
+				float DeltaHistoryTime = SampleTime - DepthHistoryTime;
+
 				if (bBehind)
 				{
-					#if SMRT_EXTRAPOLATE_WITH_SLOPE
+					#if SMRT_EXTRAPOLATE_SLOPE
 						DepthForComparison = DepthSlope * DeltaHistoryTime + DepthHistory;
 					#else
 						DepthForComparison = DepthHistory;
@@ -90,18 +95,16 @@
 				}
 				else
 				{
-					// This must be the second valid sample we've seen, as the first will take the DepthHistory < 0 path
-					// Only update the slope if we sampled a new texel; otherwise it will just zero out the slope
-					// This can happen at high sample counts with rays that are nearly parallel to the light direction
-					if (SampleDepth != DepthHistory)	// Exact test is fine here; point sampled texel with be identical
+					if (SampleDepth != DepthHistory)
 					{
 						// NOTE: DCE will remove all this if SMRT_EXTRAPOLATE_WITH_SLOPE is false
-						const float SlopeClamp = 1e-5f;
+						const float SlopeClamp = Sample.ExtrapolateSlope;
 						DepthSlope = (SampleDepth - DepthHistory) / DeltaHistoryTime;
 						DepthSlope = clamp(DepthSlope, -SlopeClamp, SlopeClamp);
+
+						DepthHistory = SampleDepth;
+						DepthHistoryTime = SampleTime;
 					}
-					DepthHistory = SampleDepth;
-					DepthHistoryTime = SampleTime;
 				}
 
 				float DepthDiff = ReferenceDepth - DepthForComparison;
@@ -109,21 +112,17 @@
 				bool bHit = abs(DepthDiff + HalfCompareTolerance) < HalfCompareTolerance;
 				if (bHit)
 				{
-					OccluderDistance = SMRTComputeOccluderDistance(RayState, SampleDepth);
 					bValidHit = true;
 					break;
 				}
 			}
+
+			PrevSampleTime = SampleTime;
+			PrevReferenceDepth = ReferenceDepth;
 		}
-		else
-		{
-			++SampleMissCount;
-		}
 	}
 
 	FSMRTResult Result;
 	Result.bValidHit = bValidHit;
-	Result.OccluderDistance = OccluderDistance;
-	Result.SampleMissCount = SampleMissCount;
 	return Result;
 }
diff --git a/Engine/Shaders/Private/VirtualShadowMaps/VirtualShadowMapTransmissionCommon.ush b/Engine/Shaders/Private/VirtualShadowMaps/VirtualShadowMapTransmissionCommon.ush
--- a/Engine/Shaders/Private/VirtualShadowMaps/VirtualShadowMapTransmissionCommon.ush
+++ b/Engine/Shaders/Private/VirtualShadowMaps/VirtualShadowMapTransmissionCommon.ush
@@ -10,9 +10,9 @@ VirtualShadowMaps/VirtualShadowMapTransmissionCommon.ush:
 #include "../DeferredShadingCommon.ush"
 #include "VirtualShadowMapProjectionCommon.ush"
 
-half GetSubsurfaceOpacityFromGbuffer(uint2 PixelPos)
+float GetSubsurfaceOpacityFromGbuffer(uint2 PixelPos)
 {
-	half Opacity = 1.0;
+	float Opacity = 1.0;
 
 #if STRATA_ENABLED
 	Opacity = StrataGetOpacityFromSubSurface(StrataLoadSubsurfaceHeader(Strata.SSSTexture, PixelPos));
@@ -40,10 +40,10 @@ half GetSubsurfaceOpacityFromGbuffer(uint2 PixelPos)
 * Sample VSM to determine approximate thickness and compute simple subsurface transmission matching the
 * PCF path.
 */
-half ComputeSimpleSubsurfaceTransmissionFromVirtualShadowMap(
+float ComputeSimpleSubsurfaceTransmissionFromVirtualShadowMap(
 	int VirtualShadowMapId, 
 	float3 TranslatedWorldPosition,
-	half SubsurfaceOpacity)
+	float SubsurfaceOpacity)
 {
 	// If the receiver is a subsurface material, we do a very simple transmission approximation
 	// This matches the logic in the PCF kernel (ShadowFilteringCommon.ush) and the contact
diff --git a/Engine/Source/Runtime/Renderer/Private/VirtualShadowMaps/VirtualShadowMapArray.cpp b/Engine/Source/Runtime/Renderer/Private/VirtualShadowMaps/VirtualShadowMapArray.cpp
--- a/Engine/Source/Runtime/Renderer/Private/VirtualShadowMaps/VirtualShadowMapArray.cpp
+++ b/Engine/Source/Runtime/Renderer/Private/VirtualShadowMaps/VirtualShadowMapArray.cpp
@@ -1103,6 +1103,8 @@ void FVirtualShadowMapArray::BuildPageAllocations(
 			}
 		}
 
+		const float ResolutionLodBiasLocal = CVarResolutionLodBiasLocal.GetValueOnRenderThread();
+
 		for (FProjectedShadowInfo* ProjectedShadowInfo : VisibleLightInfo.AllProjectedShadows)
 		{
 			if (ProjectedShadowInfo->HasVirtualShadowMap())
@@ -1139,6 +1141,7 @@ void FVirtualShadowMapArray::BuildPageAllocations(
 					Data.PreViewTranslationLWCOffset			= PreViewTranslation.GetOffset();
 					Data.LightType								= ProjectedShadowInfo->GetLightSceneInfo().Proxy->GetLightType();
 					Data.LightSourceRadius						= ProjectedShadowInfo->GetLightSceneInfo().Proxy->GetSourceRadius();
+					Data.ResolutionLodBias						= ResolutionLodBiasLocal;
 				}
 
 				if (bDebugOutputEnabled)
@@ -1242,7 +1245,6 @@ void FVirtualShadowMapArray::BuildPageAllocations(
 					PassParameters->OutPageRequestFlags = PageRequestFlagsUAV;
 					PassParameters->ForwardLightData = View.ForwardLightingResources.ForwardLightUniformBuffer;
 					PassParameters->DirectionalLightIds = GraphBuilder.CreateSRV(DirectionalLightIdsRDG);
-					PassParameters->ResolutionLodBiasLocal = CVarResolutionLodBiasLocal.GetValueOnRenderThread();
 					PassParameters->PageDilationBorderSizeLocal = CVarPageDilationBorderSizeLocal.GetValueOnRenderThread();
 					PassParameters->PageDilationBorderSizeDirectional = CVarPageDilationBorderSizeDirectional.GetValueOnRenderThread();
 					PassParameters->bCullBackfacingPixels = ShouldCullBackfacingPixels() ? 1 : 0;
diff --git a/Engine/Source/Runtime/Renderer/Private/VirtualShadowMaps/VirtualShadowMapArray.h b/Engine/Source/Runtime/Renderer/Private/VirtualShadowMaps/VirtualShadowMapArray.h
--- a/Engine/Source/Runtime/Renderer/Private/VirtualShadowMaps/VirtualShadowMapArray.h
+++ b/Engine/Source/Runtime/Renderer/Private/VirtualShadowMaps/VirtualShadowMapArray.h
@@ -91,7 +91,8 @@ struct FVirtualShadowMapProjectionShaderData
 	// TODO: There are more local lights than directional
 	// We should move the directional-specific stuff out to its own structure.
 	FVector3f NegativeClipmapWorldOriginLWCOffset;	// Shares the LWCTile with PreViewTranslation
-	float ClipmapResolutionLodBias = 0.0f;
+	// Slightly different meaning for clipmaps (includes camera pixel size scaling stuff) and local lights (raw bias)
+	float ResolutionLodBias = 0.0f;
 
 	FIntPoint ClipmapCornerOffset = FIntPoint(0, 0);
 	int32 ClipmapIndex = 0;					// 0 .. ClipmapLevelCount-1
diff --git a/Engine/Source/Runtime/Renderer/Private/VirtualShadowMaps/VirtualShadowMapClipmap.cpp b/Engine/Source/Runtime/Renderer/Private/VirtualShadowMaps/VirtualShadowMapClipmap.cpp
--- a/Engine/Source/Runtime/Renderer/Private/VirtualShadowMaps/VirtualShadowMapClipmap.cpp
+++ b/Engine/Source/Runtime/Renderer/Private/VirtualShadowMaps/VirtualShadowMapClipmap.cpp
@@ -261,7 +261,7 @@ FVirtualShadowMapProjectionShaderData FVirtualShadowMapClipmap::GetProjectionSha
 	Data.ClipmapIndex = ClipmapIndex;
 	Data.ClipmapLevel = FirstLevel + ClipmapIndex;
 	Data.ClipmapLevelCount = LevelData.Num();
-	Data.ClipmapResolutionLodBias = ResolutionLodBias;
+	Data.ResolutionLodBias = ResolutionLodBias;
 	Data.ClipmapCornerOffset = Level.CornerOffset;
 	Data.LightSourceRadius = GetLightSceneInfo().Proxy->GetSourceRadius();
 
diff --git a/Engine/Source/Runtime/Renderer/Private/VirtualShadowMaps/VirtualShadowMapProjection.cpp b/Engine/Source/Runtime/Renderer/Private/VirtualShadowMaps/VirtualShadowMapProjection.cpp
--- a/Engine/Source/Runtime/Renderer/Private/VirtualShadowMaps/VirtualShadowMapProjection.cpp
+++ b/Engine/Source/Runtime/Renderer/Private/VirtualShadowMaps/VirtualShadowMapProjection.cpp
@@ -26,6 +26,8 @@
 #include "VirtualShadowMapClipmap.h"
 #include "HairStrands/HairStrandsData.h"
 
+#define MAX_TEST_PERMUTATION 0
+
 static TAutoConsoleVariable<float> CVarScreenRayLength(
 	TEXT( "r.Shadow.Virtual.ScreenRayLength" ),
 	0.015f,
@@ -101,6 +103,14 @@
 	ECVF_Scalability | ECVF_RenderThreadSafe
 );
 
+static TAutoConsoleVariable<float> CVarSMRTTexelDitherScaleLocal(
+	TEXT( "r.Shadow.Virtual.SMRT.TexelDitherScaleLocal" ),
+	2.0f,
+	TEXT( "Applies a dither to the shadow map ray casts for local lights to help hide aliasing due to insufficient shadow resolution.\n" )
+	TEXT( "Setting this too high can cause shadows light leaks near occluders." ),
+	ECVF_Scalability | ECVF_RenderThreadSafe
+);
+
 static TAutoConsoleVariable<float> CVarSMRTTexelDitherScaleDirectional(
 	TEXT( "r.Shadow.Virtual.SMRT.TexelDitherScaleDirectional" ),
 	2.0f,
@@ -109,20 +119,31 @@
 	ECVF_Scalability | ECVF_RenderThreadSafe
 );
 
-static TAutoConsoleVariable<float> CVarSMRTWorldDitherScaleDirectional(
-	TEXT( "r.Shadow.Virtual.SMRT.WorldDitherScaleDirectional" ),
-	0.0f,
-	TEXT( "Applies a world." ),
+TAutoConsoleVariable<float> CVarSMRTExtrapolateMaxSlopeLocal(
+	TEXT("r.Shadow.Virtual.SMRT.ExtrapolateMaxSlopeLocal"),
+	0.05f,
+	TEXT("Maximum depth slope when extrapolating behind occluders for local lights.\n")
+	TEXT("Higher values allow softer penumbra edges but can introduce light leaks behind second occluders.\n")
+	TEXT("Setting to 0 will disable slope extrapolation slightly improving projection performance, at the cost of reduced penumbra quality."),
 	ECVF_Scalability | ECVF_RenderThreadSafe
 );
 
-TAutoConsoleVariable<int32> CVarSMRTExtrapolateWithSlope(
-	TEXT("r.Shadow.Virtual.SMRT.ExtrapolateWithSlope"),
-	1,
-	TEXT("Use slope-based extrapolation behind occluders. This can increase the quality of shadow penumbra on surfaces aligned with the light direction."),
+TAutoConsoleVariable<float> CVarSMRTExtrapolateMaxSlopeDirectional(
+	TEXT("r.Shadow.Virtual.SMRT.ExtrapolateMaxSlopeDirectional"),
+	5.0f,
+	TEXT("Maximum depth slope when extrapolating behind occluders for directional lights.\n")
+	TEXT("Higher values allow softer penumbra edges but can introduce light leaks behind second occluders.\n")
+	TEXT("Setting to 0 will disable slope extrapolation slightly improving projection performance, at the cost of reduced penumbra quality."),
 	ECVF_Scalability | ECVF_RenderThreadSafe
 );
 
+TAutoConsoleVariable<float> CVarSMRTMaxSlopeBiasLocal(
+	TEXT("r.Shadow.Virtual.SMRT.MaxSlopeBiasLocal"),
+	50.0f,
+	TEXT("Maximum depth slope. Low values produce artifacts if shadow resolution is insufficient. High values can worsen light leaks near occluders and sparkly pixels in shadowed areas."),
+	ECVF_Scalability | ECVF_RenderThreadSafe
+);
+
 static TAutoConsoleVariable<int32> CVarForcePerLightShadowMaskClear(
 	TEXT( "r.Shadow.Virtual.ForcePerLightShadowMaskClear" ),
 	0,
@@ -130,6 +151,15 @@
 	ECVF_RenderThreadSafe
 );
 
+#if MAX_TEST_PERMUTATION > 0
+static TAutoConsoleVariable<int32> CVarTestPermutation(
+	TEXT( "r.Shadow.Virtual.ProjectionTestPermutation" ),
+	0,
+	TEXT( "Used for A/B testing projection shader changes. " ),
+	ECVF_RenderThreadSafe
+);
+#endif
+
 const TCHAR* ToString(EVirtualShadowMapProjectionInputType In)
 {
 	switch (In)
@@ -206,20 +236,26 @@
 	SHADER_USE_PARAMETER_STRUCT(FVirtualShadowMapProjectionCS, FGlobalShader)
 	
 	class FDirectionalLightDim		: SHADER_PERMUTATION_BOOL("DIRECTIONAL_LIGHT");
-	class FSMRTExtrapolateSlopeDim	: SHADER_PERMUTATION_BOOL("SMRT_EXTRAPOLATE_WITH_SLOPE");
 	class FOnePassProjectionDim		: SHADER_PERMUTATION_BOOL("ONE_PASS_PROJECTION");
 	class FHairStrandsDim			: SHADER_PERMUTATION_BOOL("HAS_HAIR_STRANDS");
 	class FVisualizeOutputDim		: SHADER_PERMUTATION_BOOL("VISUALIZE_OUTPUT");
+	class FExtrapolateSlopeDim		: SHADER_PERMUTATION_BOOL("SMRT_EXTRAPOLATE_SLOPE");
 	// -1 means dynamic count
 	class FSMRTStaticSampleCount	: SHADER_PERMUTATION_RANGE_INT("SMRT_TEMPLATE_STATIC_SAMPLES_PER_RAY", -1, 2);
+	// Used for A/B testing a change that affects reg allocation, etc.
+	class FTestDim					: SHADER_PERMUTATION_INT("TEST_PERMUTATION", MAX_TEST_PERMUTATION+1);
 
 	using FPermutationDomain = TShaderPermutationDomain<
 		FDirectionalLightDim,
 		FOnePassProjectionDim,
 		FHairStrandsDim,
 		FVisualizeOutputDim,
-		FSMRTExtrapolateSlopeDim,
-		FSMRTStaticSampleCount>;
+		FExtrapolateSlopeDim,
+		FSMRTStaticSampleCount
+#if MAX_TEST_PERMUTATION > 0
+		, FTestDim
+#endif
+	>;
 
 	BEGIN_SHADER_PARAMETER_STRUCT(FParameters, )
 		SHADER_PARAMETER_STRUCT_INCLUDE(FVirtualShadowMapSamplingParameters, SamplingParameters)
@@ -236,7 +272,8 @@
 		SHADER_PARAMETER(float, SMRTRayLengthScale)
 		SHADER_PARAMETER(float, SMRTCotMaxRayAngleFromLight)
 		SHADER_PARAMETER(float, SMRTTexelDitherScale)
-		SHADER_PARAMETER(float, SMRTWorldDitherScale)
+		SHADER_PARAMETER(float, SMRTExtrapolateSlope)
+		SHADER_PARAMETER(float, SMRTMaxSlopeBias)
 		SHADER_PARAMETER(uint32, bSMRTUseAdaptiveRayCount)
 		SHADER_PARAMETER(uint32, InputType)
 		SHADER_PARAMETER(uint32, bCullBackfacingPixels)
@@ -352,7 +389,8 @@
 		PassParameters->SMRTRayLengthScale = CVarSMRTRayLengthScaleDirectional.GetValueOnRenderThread();
 		PassParameters->SMRTCotMaxRayAngleFromLight = 0.0f;	// unused in this path
 		PassParameters->SMRTTexelDitherScale = CVarSMRTTexelDitherScaleDirectional.GetValueOnRenderThread();
-		PassParameters->SMRTWorldDitherScale = CVarSMRTWorldDitherScaleDirectional.GetValueOnRenderThread();
+		PassParameters->SMRTExtrapolateSlope = CVarSMRTExtrapolateMaxSlopeDirectional.GetValueOnRenderThread();
+		PassParameters->SMRTMaxSlopeBias = 0.0f; // unused in this path
 	}
 	else
 	{
@@ -360,8 +398,9 @@
 		PassParameters->SMRTSamplesPerRay = CVarSMRTSamplesPerRayLocal.GetValueOnRenderThread();
 		PassParameters->SMRTRayLengthScale = 0.0f;		// unused in this path
 		PassParameters->SMRTCotMaxRayAngleFromLight = 1.0f / FMath::Tan(CVarSMRTMaxRayAngleFromLight.GetValueOnRenderThread());
-		PassParameters->SMRTTexelDitherScale = 0.0f;	// not yet implemented
-		PassParameters->SMRTWorldDitherScale = 0.0f;	// not yet implemented
+		PassParameters->SMRTTexelDitherScale = CVarSMRTTexelDitherScaleLocal.GetValueOnRenderThread();
+		PassParameters->SMRTExtrapolateSlope = CVarSMRTExtrapolateMaxSlopeLocal.GetValueOnRenderThread();
+		PassParameters->SMRTMaxSlopeBias = CVarSMRTMaxSlopeBiasLocal.GetValueOnRenderThread();
 	}
 	
 	bool bDebugOutput = false;
@@ -385,10 +424,16 @@
 	FVirtualShadowMapProjectionCS::FPermutationDomain PermutationVector;
 	PermutationVector.Set< FVirtualShadowMapProjectionCS::FDirectionalLightDim >( bDirectionalLight );
 	PermutationVector.Set< FVirtualShadowMapProjectionCS::FOnePassProjectionDim >( bOnePassProjection );
-	PermutationVector.Set< FVirtualShadowMapProjectionCS::FSMRTExtrapolateSlopeDim >( CVarSMRTExtrapolateWithSlope.GetValueOnRenderThread() != 0 );
 	PermutationVector.Set< FVirtualShadowMapProjectionCS::FHairStrandsDim >( bHasHairStrandsData );
 	PermutationVector.Set< FVirtualShadowMapProjectionCS::FVisualizeOutputDim >( bDebugOutput );
+	PermutationVector.Set< FVirtualShadowMapProjectionCS::FExtrapolateSlopeDim >( PassParameters->SMRTExtrapolateSlope > 0.0f );
 	PermutationVector.Set< FVirtualShadowMapProjectionCS::FSMRTStaticSampleCount >( StaticSamplesPerRay );
+#if MAX_TEST_PERMUTATION > 0
+	{
+		int32 TestPermutation = FMath::Clamp(CVarTestPermutation.GetValueOnRenderThread(), 0, MAX_TEST_PERMUTATION);
+		PermutationVector.Set< FVirtualShadowMapProjectionCS::FTestDim >( TestPermutation );
+	}
+#endif
 
 	auto ComputeShader = View.ShaderMap->GetShader< FVirtualShadowMapProjectionCS >( PermutationVector );
 	ClearUnusedGraphResources( ComputeShader, PassParameters );
