Add SpecularScale to light grid data so it can be used in clustered shading and forward lighting passes.

diff --git a/Engine/Shaders/Private/ClusteredDeferredShadingPixelShader.usf b/Engine/Shaders/Private/ClusteredDeferredShadingPixelShader.usf
--- a/Engine/Shaders/Private/ClusteredDeferredShadingPixelShader.usf
+++ b/Engine/Shaders/Private/ClusteredDeferredShadingPixelShader.usf
@@ -123,8 +123,6 @@ bool IsLightVisible(const FLocalLightData In, float3 TranslatedWorldPosition)
  */
 FDeferredLightData ConvertLocalToDeferredLight(const FLocalLightData In)
 {
-	// TODO: Not part of LightGrid ligth data, does it need to be added?
-	float SpecularScale = 1.0f;
 	const uint LightTypeAndPackedShadowMapChannelMask = asuint(In.LightDirectionAndShadowMask.w);
 
 	// bits [17:16] really
@@ -148,7 +146,7 @@ FDeferredLightData ConvertLocalToDeferredLight(const FLocalLightData In)
 	// Setting it to true follows how it is setup in standard deferred, where only the directional lights are not doing this (LIGHT_SOURCE_SHAPE > 0)
 	//Out.bSpotLight = true;
 	//Out.bSpotLight = Out.SpotAngles.x > -2.0f;
-	Out.SpecularScale = SpecularScale;
+	Out.SpecularScale = In.RectBarnDoorAndVirtualShadowMapIdAndSpecularScale.w;
 	// TODO: not sure if this check is even needed, standard deferred always seems to set it to true?
 	Out.bSpotLight = (LightType == LIGHT_TYPE_SPOT);
 	// NOTE: supporting the rect-light code path adds ~30% overhead on NVIDIA (RTX 2060) even when not used! 
diff --git a/Engine/Shaders/Private/ForwardLightingCommon.ush b/Engine/Shaders/Private/ForwardLightingCommon.ush
--- a/Engine/Shaders/Private/ForwardLightingCommon.ush
+++ b/Engine/Shaders/Private/ForwardLightingCommon.ush
@@ -257,7 +257,7 @@ FDeferredLightingSplit GetForwardDirectLightingSplit(
 		LightData.bInverseSquared = LightData.FalloffExponent == 0;
 		LightData.bRadialLight = true;
 		LightData.bSpotLight = LightData.SpotAngles.x > -2.0f;
-		LightData.SpecularScale = SpecularScale;
+		LightData.SpecularScale = SpecularScale * LightData.RectBarnDoorAndVirtualShadowMapIdAndSpecularScale.w;
 		// Initialize rect light once parameters are supported onto FLocalLightData
 		// LightData.RectLightAtlasMaxLevel = LocalLight.RectLightAtlasMaxLevel;
 		// LightData.RectLightAtlasUVOffset = LocalLight.RectLightAtlasUVOffset;
diff --git a/Engine/Shaders/Private/LightGridCommon.ush b/Engine/Shaders/Private/LightGridCommon.ush
--- a/Engine/Shaders/Private/LightGridCommon.ush
+++ b/Engine/Shaders/Private/LightGridCommon.ush
@@ -172,7 +172,7 @@ struct FLocalLightData
 	float4 SpotAnglesAndSourceRadiusPacked;
     float4 LightDirectionAndShadowMask;
     float4 LightTangentAndSoftSourceRadius;
-	float4 RectBarnDoorAndVirtualShadowMapId;
+	float4 RectBarnDoorAndVirtualShadowMapIdAndSpecularScale;
 	/** Flag is true if clustered deferred is supported for this light. They are always first / together.*/
 	bool bClusteredDeferredSupported; 
 	/** Flag is true if it is a simple light. They are always first / together.*/
@@ -200,8 +200,8 @@ FLocalLightData GetLocalLightData(uint GridIndex, uint EyeIndex)
 		Result.LightDirectionAndShadowMask = ForwardLightData.ForwardLocalLightBuffer[LocalLightBaseIndex + 2];
 		Result.SpotAnglesAndSourceRadiusPacked = ForwardLightData.ForwardLocalLightBuffer[LocalLightBaseIndex + 3];
 		Result.LightTangentAndSoftSourceRadius = ForwardLightData.ForwardLocalLightBuffer[LocalLightBaseIndex + 4];
-		Result.RectBarnDoorAndVirtualShadowMapId = ForwardLightData.ForwardLocalLightBuffer[LocalLightBaseIndex + 5];
-		Result.VirtualShadowMapId = int(Result.RectBarnDoorAndVirtualShadowMapId.z);
+		Result.RectBarnDoorAndVirtualShadowMapIdAndSpecularScale = ForwardLightData.ForwardLocalLightBuffer[LocalLightBaseIndex + 5];
+		Result.VirtualShadowMapId = int(Result.RectBarnDoorAndVirtualShadowMapIdAndSpecularScale.z);
 #if INSTANCED_STEREO
 	}
 	else
@@ -215,8 +215,8 @@ FLocalLightData GetLocalLightData(uint GridIndex, uint EyeIndex)
 		Result.LightDirectionAndShadowMask = ForwardLightDataISR.ForwardLocalLightBuffer[LocalLightBaseIndex + 2];
 		Result.SpotAnglesAndSourceRadiusPacked = ForwardLightDataISR.ForwardLocalLightBuffer[LocalLightBaseIndex + 3];
 		Result.LightTangentAndSoftSourceRadius = ForwardLightDataISR.ForwardLocalLightBuffer[LocalLightBaseIndex + 4];
-		Result.RectBarnDoorAndVirtualShadowMapId = ForwardLightDataISR.ForwardLocalLightBuffer[LocalLightBaseIndex + 5];
-		Result.VirtualShadowMapId = int(Result.RectBarnDoorAndVirtualShadowMapId.z);
+		Result.RectBarnDoorAndVirtualShadowMapIdAndSpecularScale = ForwardLightDataISR.ForwardLocalLightBuffer[LocalLightBaseIndex + 5];
+		Result.VirtualShadowMapId = int(Result.RectBarnDoorAndVirtualShadowMapIdAndSpecularScale.z);
 	}
 #endif
 
diff --git a/Engine/Shaders/Private/VolumetricCloud.usf b/Engine/Shaders/Private/VolumetricCloud.usf
--- a/Engine/Shaders/Private/VolumetricCloud.usf
+++ b/Engine/Shaders/Private/VolumetricCloud.usf
@@ -1007,8 +1007,8 @@ void MainCommon(in FMaterialPixelParameters MaterialParameters, in float4 SvPosi
 						LightData.Tangent = LocalLight.LightTangentAndSoftSourceRadius.xyz;
 						LightData.SoftSourceRadius = LocalLight.LightTangentAndSoftSourceRadius.w;
 						LightData.bInverseSquared = LightData.FalloffExponent == 0;
-						LightData.RectLightBarnCosAngle = LocalLight.RectBarnDoorAndVirtualShadowMapId.x;
-						LightData.RectLightBarnLength = LocalLight.RectBarnDoorAndVirtualShadowMapId.y;
+						LightData.RectLightBarnCosAngle = LocalLight.RectBarnDoorAndVirtualShadowMapIdAndSpecularScale.x;
+						LightData.RectLightBarnLength = LocalLight.RectBarnDoorAndVirtualShadowMapIdAndSpecularScale.y;
 						LightData.bRadialLight = true;
 						LightData.bSpotLight = LightData.SpotAngles.x > -2.0f;
 						LightData.bRectLight = LightData.RectLightBarnLength > -2.0f;
diff --git a/Engine/Shaders/Private/VolumetricFog.usf b/Engine/Shaders/Private/VolumetricFog.usf
--- a/Engine/Shaders/Private/VolumetricFog.usf
+++ b/Engine/Shaders/Private/VolumetricFog.usf
@@ -717,8 +717,8 @@ void LightScatteringCS(
 				LightData.Tangent = LocalLight.LightTangentAndSoftSourceRadius.xyz;
 				LightData.SoftSourceRadius = LocalLight.LightTangentAndSoftSourceRadius.w;
 				LightData.bInverseSquared = LightData.FalloffExponent == 0;
-				LightData.RectLightBarnCosAngle = LocalLight.RectBarnDoorAndVirtualShadowMapId.x;
-				LightData.RectLightBarnLength = LocalLight.RectBarnDoorAndVirtualShadowMapId.y;
+				LightData.RectLightBarnCosAngle = LocalLight.RectBarnDoorAndVirtualShadowMapIdAndSpecularScale.x;
+				LightData.RectLightBarnLength = LocalLight.RectBarnDoorAndVirtualShadowMapIdAndSpecularScale.y;
 				LightData.RectLightAtlasUVOffset = 0;
 				LightData.RectLightAtlasUVScale  = 0;
 				LightData.RectLightAtlasMaxLevel = MAX_RECT_ATLAS_MIP;
diff --git a/Engine/Source/Runtime/Renderer/Private/BasePassRendering.h b/Engine/Source/Runtime/Renderer/Private/BasePassRendering.h
--- a/Engine/Source/Runtime/Renderer/Private/BasePassRendering.h
+++ b/Engine/Source/Runtime/Renderer/Private/BasePassRendering.h
@@ -52,7 +52,7 @@ public:
 	FVector4f LightDirectionAndShadowMapChannelMask;
 	FVector4f SpotAnglesAndSourceRadiusPacked;
 	FVector4f LightTangentAndSoftSourceRadius;
-	FVector4f RectBarnDoorAndVirtualShadowMapId;
+	FVector4f RectBarnDoorAndVirtualShadowMapIdAndSpecularScale;
 };
 
 struct FForwardBasePassTextures
diff --git a/Engine/Source/Runtime/Renderer/Private/LightGridInjection.cpp b/Engine/Source/Runtime/Renderer/Private/LightGridInjection.cpp
--- a/Engine/Source/Runtime/Renderer/Private/LightGridInjection.cpp
+++ b/Engine/Source/Runtime/Renderer/Private/LightGridInjection.cpp
@@ -377,7 +377,7 @@ void FSceneRenderer::ComputeLightGrid(FRDGBuilder& GraphBuilder, bool bCullLight
 
 					LightData.SpotAnglesAndSourceRadiusPacked = FVector4f(-2, 1, 0, *(float*)&PackedWInt);
 					LightData.LightTangentAndSoftSourceRadius = FVector4f(1.0f, 0.0f, 0.0f, 0.0f);
-					LightData.RectBarnDoorAndVirtualShadowMapId = FVector4f(0, -2, 0, 0);
+					LightData.RectBarnDoorAndVirtualShadowMapIdAndSpecularScale = FVector4f(0, -2, -1, 1);
 
 #if ENABLE_LIGHT_CULLING_VIEW_SPACE_BUILD_DATA
 					FVector4f ViewSpacePosAndRadius(FVector4f(View.ViewMatrices.GetViewMatrix().TransformPosition(SimpleLightPerViewData.Position)), SimpleLight.Radius);
@@ -467,8 +467,8 @@ void FSceneRenderer::ComputeLightGrid(FRDGBuilder& GraphBuilder, bool bCullLight
 
 						// NOTE: This cast of VirtualShadowMapId to float is not ideal, but bitcast has issues here with INDEX_NONE -> NaN
 						// and 32-bit floats have enough mantissa to cover all reasonable numbers here for now.
-						LightData.RectBarnDoorAndVirtualShadowMapId = FVector4f(LightParameters.RectLightBarnCosAngle, LightParameters.RectLightBarnLength, float(VirtualShadowMapId), 0);
-						checkSlow(int(LightData.RectBarnDoorAndVirtualShadowMapId.Z) == VirtualShadowMapId);
+						LightData.RectBarnDoorAndVirtualShadowMapIdAndSpecularScale = FVector4f(LightParameters.RectLightBarnCosAngle, LightParameters.RectLightBarnLength, float(VirtualShadowMapId), LightParameters.SpecularScale);
+						checkSlow(int(LightData.RectBarnDoorAndVirtualShadowMapIdAndSpecularScale.Z) == VirtualShadowMapId);
 
 						float VolumetricScatteringIntensity = LightProxy->GetVolumetricScatteringIntensity();
 
