Fix  /Engine/Private/ForwardLightingCommon.ush:260:55: error: no member named 'RectBarnDoorAndVirtualShadowMapIdAndSpecularScale' in 'FDeferredLightData'
#rnx

diff --git a/Engine/Shaders/Private/ForwardLightingCommon.ush b/Engine/Shaders/Private/ForwardLightingCommon.ush
--- a/Engine/Shaders/Private/ForwardLightingCommon.ush
+++ b/Engine/Shaders/Private/ForwardLightingCommon.ush
@@ -257,7 +257,7 @@ FDeferredLightingSplit GetForwardDirectLightingSplit(
 		LightData.bInverseSquared = LightData.FalloffExponent == 0;
 		LightData.bRadialLight = true;
 		LightData.bSpotLight = LightData.SpotAngles.x > -2.0f;
-		LightData.SpecularScale = SpecularScale * LightData.RectBarnDoorAndVirtualShadowMapIdAndSpecularScale.w;
+		LightData.SpecularScale = SpecularScale * LocalLight.RectBarnDoorAndVirtualShadowMapIdAndSpecularScale.w;
 		// Initialize rect light once parameters are supported onto FLocalLightData
 		// LightData.RectLightAtlasMaxLevel = LocalLight.RectLightAtlasMaxLevel;
 		// LightData.RectLightAtlasUVOffset = LocalLight.RectLightAtlasUVOffset;
