HDR: Device-encoding-only shader fix for ScRGB.

This correctly replicates the combine LUTs and tonemap shaders, where the final conversion from AP1 to sRGB is applied after the LUT.

diff --git a/Engine/Shaders/Private/PostProcessDeviceEncodingOnly.usf b/Engine/Shaders/Private/PostProcessDeviceEncodingOnly.usf
--- a/Engine/Shaders/Private/PostProcessDeviceEncodingOnly.usf
+++ b/Engine/Shaders/Private/PostProcessDeviceEncodingOnly.usf
@@ -100,7 +100,7 @@ float4 DeviceEncodingOnlyCommonPS(
 	}
 
 	// ACES 1000nit transform with PQ/2084 encoding, user specified gamut 
-	else if( GetOutputDevice() == TONEMAPPER_OUTPUT_ACES1000nitST2084 || GetOutputDevice() == TONEMAPPER_OUTPUT_ACES1000nitScRGB)
+	else if( GetOutputDevice() == TONEMAPPER_OUTPUT_ACES1000nitST2084 )
 	{		
 		// 1000 nit ODT
 		float3 ODTColor = ACESOutputTransforms1000( FilmColor );
@@ -113,7 +113,7 @@ float4 DeviceEncodingOnlyCommonPS(
 	}
 
 	// ACES 2000nit transform with PQ/2084 encoding, user specified gamut 
-	else if( GetOutputDevice() == TONEMAPPER_OUTPUT_ACES2000nitST2084 || GetOutputDevice() == TONEMAPPER_OUTPUT_ACES2000nitScRGB)
+	else if( GetOutputDevice() == TONEMAPPER_OUTPUT_ACES2000nitST2084 )
 	{		
 		// 2000 nit ODT
 		float3 ODTColor = ACESOutputTransforms2000( FilmColor );
@@ -123,7 +123,27 @@ float4 DeviceEncodingOnlyCommonPS(
 
 		// Apply conversion to ST-2084 (Dolby PQ)
 		OutDeviceColor = LinearToST2084( ODTColor );
-	}	
+	}
+
+	// ACES 1000nit transform to linear ScRGB
+	else if(GetOutputDevice() == TONEMAPPER_OUTPUT_ACES1000nitScRGB)
+	{
+		// 1000 nit ODT
+		float3 ODTColor = ACESOutputTransforms1000( FilmColor );
+
+		// Apply conversion to ScRGB
+		OutDeviceColor = ST2084ToScRGB( LinearToST2084( ODTColor ), TONEMAPPER_OUTPUT_ACES1000nitScRGB );
+	}
+
+	// ACES 2000nit transform to linear ScRGB
+	else if(GetOutputDevice() == TONEMAPPER_OUTPUT_ACES2000nitScRGB)
+	{
+		// 2000 nit ODT
+		float3 ODTColor = ACESOutputTransforms2000( FilmColor );
+
+		// Apply conversion to ScRGB
+		OutDeviceColor = ST2084ToScRGB( LinearToST2084( ODTColor ), TONEMAPPER_OUTPUT_ACES2000nitScRGB );
+	}
 	
 	else if( GetOutputDevice() == TONEMAPPER_OUTPUT_LinearEXR)
 	{
