Fixed shader compilation error in editor mobile preview

diff --git a/Engine/Shaders/Private/Common.ush b/Engine/Shaders/Private/Common.ush
--- a/Engine/Shaders/Private/Common.ush
+++ b/Engine/Shaders/Private/Common.ush
@@ -1792,13 +1792,11 @@ bool GetRuntimeVirtualTextureOutputSwitch()
 #endif
 }
 
-#if FEATURE_LEVEL >= FEATURE_LEVEL_SM4
 struct FWriteToSliceGeometryOutput
 {
 	FScreenVertexOutput Vertex;
 	uint LayerIndex : SV_RenderTargetArrayIndex;
 };
-#endif
 
 // Helper macro to globally ignore requests for non-offset world positions in materials when lower than shader model 4. We do this
 // because we are using an extra interpolator for this second world position, and in < ES31 there may not be enough
