Fix a single layer shader compilation error on mobile

diff --git a/Engine/Shaders/Private/MobileBasePassPixelShader.usf b/Engine/Shaders/Private/MobileBasePassPixelShader.usf
--- a/Engine/Shaders/Private/MobileBasePassPixelShader.usf
+++ b/Engine/Shaders/Private/MobileBasePassPixelShader.usf
@@ -562,9 +562,10 @@ void Main(
 	half3 CameraVector = -MaterialParameters.CameraVector;
 
 	half4 DynamicShadowFactors = 1.0f;
+	float DirectionalLightShadow = 1.0f;
 
 	// Directional light
-	AccumulateDirectionalLighting(GBuffer, MaterialParameters.WorldPosition_CamRelative, CameraVector, MaterialParameters.ScreenPosition, SvPosition, DynamicShadowFactors, DirectLighting);
+	AccumulateDirectionalLighting(GBuffer, MaterialParameters.WorldPosition_CamRelative, CameraVector, MaterialParameters.ScreenPosition, SvPosition, DynamicShadowFactors, DirectionalLightShadow, DirectLighting);
 
 	float2 LocalPosition = SvPosition.xy - ResolvedView.ViewRectMin.xy;
 	uint GridIndex = ComputeLightGridCellIndex(uint2(LocalPosition.x, LocalPosition.y), SvPosition.w);
@@ -611,7 +612,7 @@ void Main(
 		const float4 NullDistortionParams = 1.0f;
 		WaterVolumeLightingOutput WaterLighting = EvaluateWaterVolumeLighting(
 			MaterialParameters, PixelMaterialInputs, ResolvedView,
-			Shadow, GBuffer.Specular, NullDistortionParams,
+			DirectionalLightShadow, GBuffer.Specular, NullDistortionParams,
 			SunIlluminance, WaterDiffuseIndirectIlluminance, EnvBrdf,
 			CameraIsUnderWater, WaterVisibility, EyeIndex,
 			bSeparateMainDirLightLuminance, SeparatedWaterMainDirLightLuminance);
diff --git a/Engine/Shaders/Private/MobileDeferredShading.usf b/Engine/Shaders/Private/MobileDeferredShading.usf
--- a/Engine/Shaders/Private/MobileDeferredShading.usf
+++ b/Engine/Shaders/Private/MobileDeferredShading.usf
@@ -168,7 +168,8 @@ void MobileDirectionalLightPS(
 
 	// Directional light
 	half4 DynamicShadowFactors = 1.0f;
-	AccumulateDirectionalLighting(GBuffer, TranslatedWorldPosition, CameraVector, ScreenPosition, SvPosition, DynamicShadowFactors, DirectLighting);
+	float DirectionalLightShadow = 1.0f;
+	AccumulateDirectionalLighting(GBuffer, TranslatedWorldPosition, CameraVector, ScreenPosition, SvPosition, DynamicShadowFactors, DirectionalLightShadow, DirectLighting);
 
 #if ENABLE_CLUSTERED_REFLECTION || ENABLE_SKY_LIGHT || ENABLE_PLANAR_REFLECTION
 	// If we have a single directional light, apply relfection and sky contrubution here
diff --git a/Engine/Shaders/Private/MobileLightingCommon.ush b/Engine/Shaders/Private/MobileLightingCommon.ush
--- a/Engine/Shaders/Private/MobileLightingCommon.ush
+++ b/Engine/Shaders/Private/MobileLightingCommon.ush
@@ -182,7 +182,7 @@ half MobileDirectionalLightCSM(float2 ScreenPosition, float SceneDepth, inout fl
 	Mobile Directional Light.
 ------------------------------------------------------------------------------*/
 
-void AccumulateDirectionalLighting(FGBufferData GBuffer, float3 TranslatedWorldPosition, half3 CameraVector, float4 ScreenPosition, float4 SvPosition, inout half4 DynamicShadowFactors, inout FLightAccumulator DirectLighting)
+void AccumulateDirectionalLighting(FGBufferData GBuffer, float3 TranslatedWorldPosition, half3 CameraVector, float4 ScreenPosition, float4 SvPosition, inout half4 DynamicShadowFactors, inout float OutDirectionalLightShadow, inout FLightAccumulator DirectLighting)
 {
 #if USE_SHADOWMASKTEXTURE
 	half4 PreviewShadowMapChannelMask = half4(1.0f, 0.0f, 0.0f, 0.0f);
@@ -221,8 +221,7 @@ void AccumulateDirectionalLighting(FGBufferData GBuffer, float3 TranslatedWorldP
 
 	half4 LightAttenuation = half4(1, 1, DynamicShadowing, DynamicShadowing);
 
-	float SurfaceShadow = 0;
-	FLightAccumulator NewLighting = AccumulateDynamicLighting(TranslatedWorldPosition, CameraVector, GBuffer, 1, GBuffer.ShadingModelID, LightData, LightAttenuation, 0, uint2(0, 0), SurfaceShadow);
+	FLightAccumulator NewLighting = AccumulateDynamicLighting(TranslatedWorldPosition, CameraVector, GBuffer, 1, GBuffer.ShadingModelID, LightData, LightAttenuation, 0, uint2(0, 0), OutDirectionalLightShadow);
 	DirectLighting = LightAccumulator_Add(DirectLighting, NewLighting);
 }
 
