Fix shader compilation issue for PCD3D_SM6.
#rnx

diff --git a/Engine/Shaders/Private/SimpleElementPixelShader.usf b/Engine/Shaders/Private/SimpleElementPixelShader.usf
--- a/Engine/Shaders/Private/SimpleElementPixelShader.usf
+++ b/Engine/Shaders/Private/SimpleElementPixelShader.usf
@@ -357,17 +357,18 @@ void CubemapTextureProperties(
 
 #if TEXTURECUBE_ARRAY
 	// SliceIndex represents an index of a cubemap in the array
-	if (SliceIndex < 0.0)
+	float CurrentSliceIndex = SliceIndex;
+	if (CurrentSliceIndex < 0.0)
 	{
 		// if the slice index is not specified, display a portion of each cubemap from the TextureCube array
-		SliceIndex = floor(TextureCoordinate.y * NumSlices);
+		CurrentSliceIndex = floor(TextureCoordinate.y * NumSlices);
 	}
 #endif
 
 	if (MipLevel >= 0.0)
 	{
 #if TEXTURECUBE_ARRAY
-		BaseColor = TextureCubeArraySampleLevel(CubeTexture, CubeTextureSampler, float4(CubeCoordinates, SliceIndex), MipLevel);
+		BaseColor = TextureCubeArraySampleLevel(CubeTexture, CubeTextureSampler, float4(CubeCoordinates, CurrentSliceIndex), MipLevel);
 #else
 		BaseColor = TextureCubeSampleLevel(CubeTexture, CubeTextureSampler, CubeCoordinates, MipLevel);
 #endif
@@ -375,7 +376,7 @@ void CubemapTextureProperties(
 	else
 	{
 #if TEXTURECUBE_ARRAY
-		BaseColor = TextureCubeArraySample(CubeTexture, CubeTextureSampler, float4(CubeCoordinates, SliceIndex));
+		BaseColor = TextureCubeArraySample(CubeTexture, CubeTextureSampler, float4(CubeCoordinates, CurrentSliceIndex));
 #else
 		BaseColor = TextureCubeSample(CubeTexture, CubeTextureSampler, CubeCoordinates);
 #endif
