Fix shader compilation warning for CubemapTextureProperties.
#rnx

diff --git a/Engine/Shaders/Private/SimpleElementPixelShader.usf b/Engine/Shaders/Private/SimpleElementPixelShader.usf
--- a/Engine/Shaders/Private/SimpleElementPixelShader.usf
+++ b/Engine/Shaders/Private/SimpleElementPixelShader.usf
@@ -352,7 +352,7 @@ void CubemapTextureProperties(
 	{
 		float2 ScaledUVs = TextureCoordinate * 2 - 1;
 		// identity view matrix displays the face 0 properly oriented in space
-		CubeCoordinates = mul(float3(1, ScaledUVs.x, -ScaledUVs.y), ViewMatrix);
+		CubeCoordinates = mul(float3(1, ScaledUVs.x, -ScaledUVs.y), (float3x3)ViewMatrix);
 	}
 
 #if TEXTURECUBE_ARRAY
