Allow SceneDepthWithoutWater buffer to be read in Postprocess materials.

#rnx

diff --git a/Engine/Shaders/Private/MaterialTemplate.ush b/Engine/Shaders/Private/MaterialTemplate.ush
--- a/Engine/Shaders/Private/MaterialTemplate.ush
+++ b/Engine/Shaders/Private/MaterialTemplate.ush
@@ -91,7 +91,9 @@
 	#include "/Engine/Private/SkyAtmosphereCommon.ush"
 #endif
 
-#if MATERIAL_SHADINGMODEL_SINGLELAYERWATER
+#if MATERIAL_SHADINGMODEL_SINGLELAYERWATER || POST_PROCESS_MATERIAL // For WaterSampleSceneDepthWithoutWater() and SINGLE_LAYER_WATER_DEPTH_SCALE
+// This is currently included for post processing materials even if the SceneDepthWithoutWater texture is not sampled.
+// It might be worth it to add a switch to only include it if the texture is actually required by the material.
 	#include "/Engine/Private/SingleLayerWaterCommon.ush"
 #endif
 
@@ -1654,17 +1656,42 @@ float3 MaterialExpressionSkyAtmosphereDistantLightScatteredLuminance(FMaterialPi
 #endif
 }
 
-/** 
-	Get the scene depth from the scene below the single layer water surface. This is only valid in the single layer water rendering pass.
+#if POST_PROCESS_MATERIAL
+uint bSceneDepthWithoutWaterTextureAvailable;
+Texture2D SceneDepthWithoutSingleLayerWaterTexture;
+SamplerState SceneDepthWithoutSingleLayerWaterSampler;
+float4 SceneWithoutSingleLayerWaterMinMaxUV;
+float2 SceneWithoutSingleLayerWaterTextureSize;
+float2 SceneWithoutSingleLayerWaterInvTextureSize;
+#endif
+
+/**
+	Get the scene depth from the scene below the single layer water surface. This is only valid in the single layer water rendering pass or during post processing.
 	Returns the chosen fallback depth if the material doesn't support reading back the correct depth.
 */
 float MaterialExpressionSceneDepthWithoutWater(float2 ViewportUV, float FallbackDepth)
 {
 #if MATERIAL_SHADINGMODEL_SINGLELAYERWATER && !SCENE_TEXTURES_DISABLED && (!SIMPLE_SINGLE_LAYER_WATER || SINGLE_LAYER_WATER_SIMPLE_FORWARD)
-
 	const float2 ClampedViewportUV = clamp(ViewportUV, SingleLayerWater.SceneWithoutSingleLayerWaterMinMaxUV.xy, SingleLayerWater.SceneWithoutSingleLayerWaterMinMaxUV.zw);
-
-	return SingleLayerWater.SceneDepthWithoutSingleLayerWaterTexture.SampleLevel(SingleLayerWaterSceneDepthSampler, ClampedViewportUV, 0).x * SINGLE_LAYER_WATER_DEPTH_SCALE;
+	return WaterSampleSceneDepthWithoutWater(SingleLayerWater.SceneDepthWithoutSingleLayerWaterTexture, SingleLayerWaterSceneDepthSampler, ClampedViewportUV, 
+		SingleLayerWater.SceneWithoutSingleLayerWaterTextureSize, SingleLayerWater.SceneWithoutSingleLayerWaterInvTextureSize) * SINGLE_LAYER_WATER_DEPTH_SCALE;
+#elif POST_PROCESS_MATERIAL
+	if (bSceneDepthWithoutWaterTextureAvailable)
+	{
+		const float2 ClampedViewportUV = clamp(ViewportUV, SceneWithoutSingleLayerWaterMinMaxUV.xy, SceneWithoutSingleLayerWaterMinMaxUV.zw);
+		return WaterSampleSceneDepthWithoutWater(SceneDepthWithoutSingleLayerWaterTexture, SceneDepthWithoutSingleLayerWaterSampler, ClampedViewportUV, 
+			SceneWithoutSingleLayerWaterTextureSize, SceneWithoutSingleLayerWaterInvTextureSize) * SINGLE_LAYER_WATER_DEPTH_SCALE;
+	}
+	else
+	{
+#if SHADING_PATH_DEFERRED
+		return CalcSceneDepth(uint2(ViewportUV * View.BufferSizeAndInvSize.xy));
+#elif SHADING_PATH_MOBILE
+		return CalcSceneDepth(ViewportUV);
+#else
+		return FallbackDepth;
+#endif
+	}
 #else
 	return FallbackDepth;
 #endif
diff --git a/Engine/Shaders/Private/SingleLayerWaterCommon.ush b/Engine/Shaders/Private/SingleLayerWaterCommon.ush
--- a/Engine/Shaders/Private/SingleLayerWaterCommon.ush
+++ b/Engine/Shaders/Private/SingleLayerWaterCommon.ush
@@ -5,12 +5,15 @@
 
 // Scale for encoding depth into SceneDepthWithoutSingleLayerWaterTexture
 // Must match FDeferredShadingSceneRenderer::RenderUnderWaterFog
+// Originally used to work around precision issues with a float16 linear depth buffer.
+// TODO: Do we still need this?
 #define SINGLE_LAYER_WATER_DEPTH_SCALE 100.0f
 
 #ifndef SIMPLE_SINGLE_LAYER_WATER
 #define SIMPLE_SINGLE_LAYER_WATER 0
 #endif
 
+#if MATERIAL_SHADINGMODEL_SINGLELAYERWATER
 #if SUPPORTS_INDEPENDENT_SAMPLERS
 	#define SingleLayerWaterSceneColorSampler View.SharedBilinearClampedSampler
 #if METAL_PROFILE
@@ -21,7 +24,8 @@
 #else
 	#define SingleLayerWaterSceneColorSampler SingleLayerWater.SceneColorWithoutSingleLayerWaterSampler
 	#define SingleLayerWaterSceneDepthSampler SingleLayerWater.SceneDepthWithoutSingleLayerWaterSampler
-#endif
+#endif // SUPPORTS_INDEPENDENT_SAMPLERS
+#endif // MATERIAL_SHADINGMODEL_SINGLELAYERWATER
 
 // We use our own refract because the HLSL one was not working on all platforms (ES3.1 preview for instance)
 bool WaterRefract(float3 V, float3 N, float Eta, out float3 R)
diff --git a/Engine/Source/Runtime/Engine/Private/Materials/HLSLMaterialTranslator.cpp b/Engine/Source/Runtime/Engine/Private/Materials/HLSLMaterialTranslator.cpp
--- a/Engine/Source/Runtime/Engine/Private/Materials/HLSLMaterialTranslator.cpp
+++ b/Engine/Source/Runtime/Engine/Private/Materials/HLSLMaterialTranslator.cpp
@@ -9185,19 +9185,24 @@ int32 FHLSLMaterialTranslator::SceneDepthWithoutWater(int32 Offset, int32 Viewpo
 		return Errorf(TEXT("Cannot read scene depth without water from the vertex shader."));
 	}
 
-	if (!Material->GetShadingModels().HasShadingModel(MSM_SingleLayerWater))
-	{
-		return Errorf(TEXT("Can only read scene depth below water when material Shading Model is Single Layer Water."));
-	}
-	
-	if (Material->GetMaterialDomain() != MD_Surface)
-	{
-		return Errorf(TEXT("Can only read scene depth below water when material Domain is set to Surface."));
-	}
+	const EMaterialDomain MaterialDomain = Material->GetMaterialDomain();
 
-	if (IsTranslucentBlendMode(Material->GetBlendMode()))
+	if (MaterialDomain != MD_PostProcess)
 	{
-		return Errorf(TEXT("Can only read scene depth below water when material Blend Mode isn't translucent."));
+		if (!Material->GetShadingModels().HasShadingModel(MSM_SingleLayerWater))
+		{
+			return Errorf(TEXT("Can only read scene depth below water when material Shading Model is Single Layer Water or when material Domain is PostProcess."));
+		}
+
+		if (MaterialDomain != MD_Surface)
+		{
+			return Errorf(TEXT("Can only read scene depth below water when material Domain is set to Surface or PostProcess."));
+		}
+
+		if (IsTranslucentBlendMode(Material->GetBlendMode()))
+		{
+			return Errorf(TEXT("Can only read scene depth below water when material Blend Mode isn't translucent."));
+		}
 	}
 
 	if (Offset == INDEX_NONE && bUseOffset)
diff --git a/Engine/Source/Runtime/Renderer/Private/DeferredShadingRenderer.cpp b/Engine/Source/Runtime/Renderer/Private/DeferredShadingRenderer.cpp
--- a/Engine/Source/Runtime/Renderer/Private/DeferredShadingRenderer.cpp
+++ b/Engine/Source/Runtime/Renderer/Private/DeferredShadingRenderer.cpp
@@ -3428,7 +3428,7 @@ void FDeferredShadingSceneRenderer::Render(FRDGBuilder& GraphBuilder)
 					const FPerViewPipelineState& ViewPipelineState = GetViewPipelineState(View);
 					const bool bAnyLumenActive = ViewPipelineState.DiffuseIndirectMethod == EDiffuseIndirectMethod::Lumen || ViewPipelineState.ReflectionsMethod == EReflectionsMethod::Lumen;
 
-					AddPostProcessingPasses(GraphBuilder, View, ViewIndex, bAnyLumenActive, PostProcessingInputs, NaniteResults, InstanceCullingManager, &VirtualShadowMapArray, LumenFrameTemporaries);
+					AddPostProcessingPasses(GraphBuilder, View, ViewIndex, bAnyLumenActive, PostProcessingInputs, NaniteResults, InstanceCullingManager, &VirtualShadowMapArray, LumenFrameTemporaries, SceneWithoutWaterTextures);
 				}
 			}
 		}
diff --git a/Engine/Source/Runtime/Renderer/Private/PostProcess/PostProcessMaterial.cpp b/Engine/Source/Runtime/Renderer/Private/PostProcess/PostProcessMaterial.cpp
--- a/Engine/Source/Runtime/Renderer/Private/PostProcess/PostProcessMaterial.cpp
+++ b/Engine/Source/Runtime/Renderer/Private/PostProcess/PostProcessMaterial.cpp
@@ -22,6 +22,7 @@
 #include "SceneTextureParameters.h"
 #include "SystemTextures.h"
 #include "Strata/Strata.h"
+#include "SingleLayerWaterRendering.h"
 
 namespace
 {
@@ -546,6 +547,36 @@ FScreenPassTexture AddPostProcessMaterialPass(
 	PostProcessMaterialParameters->bFlipYAxis = Inputs.bFlipYAxis && !bForceIntermediateTarget;
 	PostProcessMaterialParameters->Strata = Strata::BindStrataGlobalUniformParameters(View);
 
+	// SceneDepthWithoutWater
+	const bool bHasValidSceneDepthWithoutWater = Inputs.SceneWithoutWaterTextures && Inputs.SceneWithoutWaterTextures->DepthTexture;
+	const bool bShouldUseBilinearSamplerForDepth = bHasValidSceneDepthWithoutWater && ShouldUseBilinearSamplerForDepthWithoutSingleLayerWater(Inputs.SceneWithoutWaterTextures->DepthTexture->Desc.Format);
+	PostProcessMaterialParameters->bSceneDepthWithoutWaterTextureAvailable = bHasValidSceneDepthWithoutWater;
+	PostProcessMaterialParameters->SceneDepthWithoutSingleLayerWaterSampler = bShouldUseBilinearSamplerForDepth ? TStaticSamplerState<SF_Bilinear>::GetRHI() : TStaticSamplerState<SF_Point>::GetRHI();
+	PostProcessMaterialParameters->SceneDepthWithoutSingleLayerWaterTexture = FRDGSystemTextures::Get(GraphBuilder).Black;
+	PostProcessMaterialParameters->SceneWithoutSingleLayerWaterMinMaxUV = FVector4f(0.0f, 0.0f, 1.0f, 1.0f);
+	PostProcessMaterialParameters->SceneWithoutSingleLayerWaterTextureSize = FVector2f(0.0f, 0.0f);
+	PostProcessMaterialParameters->SceneWithoutSingleLayerWaterInvTextureSize = FVector2f(0.0f, 0.0f);
+	if (bHasValidSceneDepthWithoutWater)
+	{
+		const bool bIsInstancedStereoSideBySide = View.bIsInstancedStereoEnabled && !View.bIsMultiViewEnabled && IStereoRendering::IsStereoEyeView(View);
+		int32 WaterViewIndex = INDEX_NONE;
+		if (bIsInstancedStereoSideBySide)
+		{
+			WaterViewIndex = View.PrimaryViewIndex; // The instanced view does not have MinMaxUV initialized, instead the primary view MinMaxUV covers both eyes
+		}
+		else
+		{
+			verify(View.Family->Views.Find(&View, WaterViewIndex));
+		}
+
+		PostProcessMaterialParameters->SceneDepthWithoutSingleLayerWaterTexture = Inputs.SceneWithoutWaterTextures->DepthTexture;
+		PostProcessMaterialParameters->SceneWithoutSingleLayerWaterMinMaxUV = Inputs.SceneWithoutWaterTextures->Views[WaterViewIndex].MinMaxUV;
+
+		const FIntVector DepthTextureSize = Inputs.SceneWithoutWaterTextures->DepthTexture->Desc.GetSize();
+		PostProcessMaterialParameters->SceneWithoutSingleLayerWaterTextureSize = FVector2f(DepthTextureSize.X, DepthTextureSize.Y);
+		PostProcessMaterialParameters->SceneWithoutSingleLayerWaterInvTextureSize = FVector2f(1.0f / DepthTextureSize.X, 1.0f / DepthTextureSize.Y);
+	}
+
 	ClearUnusedGraphResources(VertexShader, PixelShader, PostProcessMaterialParameters);
 
 	EScreenPassDrawFlags ScreenPassFlags = EScreenPassDrawFlags::AllowHMDHiddenAreaMask;
diff --git a/Engine/Source/Runtime/Renderer/Private/PostProcess/PostProcessMaterial.h b/Engine/Source/Runtime/Renderer/Private/PostProcess/PostProcessMaterial.h
--- a/Engine/Source/Runtime/Renderer/Private/PostProcess/PostProcessMaterial.h
+++ b/Engine/Source/Runtime/Renderer/Private/PostProcess/PostProcessMaterial.h
@@ -7,6 +7,7 @@
 #include "Strata/Strata.h"
 
 class UMaterialInterface;
+struct FSceneWithoutWaterTextures;
 
 const uint32 kPostProcessMaterialInputCountMax = 5;
 
@@ -47,6 +48,12 @@ BEGIN_SHADER_PARAMETER_STRUCT(FPostProcessMaterialParameters, )
 	SHADER_PARAMETER_RDG_BUFFER_SRV(Buffer<float4>, EyeAdaptationBuffer)
 	SHADER_PARAMETER(uint32, bFlipYAxis)
 	SHADER_PARAMETER(uint32, bMetalMSAAHDRDecode)
+	SHADER_PARAMETER(uint32, bSceneDepthWithoutWaterTextureAvailable)
+	SHADER_PARAMETER_RDG_TEXTURE(Texture2D, SceneDepthWithoutSingleLayerWaterTexture)
+	SHADER_PARAMETER_SAMPLER(SamplerState, SceneDepthWithoutSingleLayerWaterSampler)
+	SHADER_PARAMETER(FVector4f, SceneWithoutSingleLayerWaterMinMaxUV)
+	SHADER_PARAMETER(FVector2f, SceneWithoutSingleLayerWaterTextureSize)
+	SHADER_PARAMETER(FVector2f, SceneWithoutSingleLayerWaterInvTextureSize)
 	RENDER_TARGET_BINDING_SLOTS()
 END_SHADER_PARAMETER_STRUCT()
 
@@ -103,6 +110,9 @@ struct FPostProcessMaterialInputs
 	/** The uniform buffer containing all scene textures. */
 	FSceneTextureShaderParameters SceneTextures;
 
+	/** Depth and color textures of the scene without single layer water. May be nullptr if not available. */
+	const FSceneWithoutWaterTextures* SceneWithoutWaterTextures = nullptr;
+
 	/** Performs a vertical axis flip if the RHI allows it. */
 	bool bFlipYAxis = false;
 
diff --git a/Engine/Source/Runtime/Renderer/Private/PostProcess/PostProcessing.cpp b/Engine/Source/Runtime/Renderer/Private/PostProcess/PostProcessing.cpp
--- a/Engine/Source/Runtime/Renderer/Private/PostProcess/PostProcessing.cpp
+++ b/Engine/Source/Runtime/Renderer/Private/PostProcess/PostProcessing.cpp
@@ -195,13 +195,15 @@ bool ComposeSeparateTranslucencyInTSR(const FViewInfo& View);
 
 void AddPostProcessingPasses(
 	FRDGBuilder& GraphBuilder,
-	const FViewInfo& View, int32 ViewIndex,
+	const FViewInfo& View, 
+	int32 ViewIndex,
 	bool bAnyLumenActive,
 	const FPostProcessingInputs& Inputs,
 	const Nanite::FRasterResults* NaniteRasterResults,
 	FInstanceCullingManager& InstanceCullingManager,
 	FVirtualShadowMapArray* VirtualShadowMapArray, 
-	FLumenSceneFrameTemporaries& LumenFrameTemporaries)
+	FLumenSceneFrameTemporaries& LumenFrameTemporaries,
+	const FSceneWithoutWaterTextures& SceneWithoutWaterTextures)
 {
 	RDG_CSV_STAT_EXCLUSIVE_SCOPE(GraphBuilder, RenderPostProcessing);
 	QUICK_SCOPE_CYCLE_COUNTER(STAT_PostProcessing_Process);
@@ -401,6 +403,7 @@ void AddPostProcessingPasses(
 		PostProcessMaterialInputs.SetInput(EPostProcessMaterialInput::Velocity, Velocity);
 		PostProcessMaterialInputs.SceneTextures = GetSceneTextureShaderParameters(Inputs.SceneTextures);
 		PostProcessMaterialInputs.CustomDepthTexture = CustomDepth.Texture;
+		PostProcessMaterialInputs.SceneWithoutWaterTextures = &SceneWithoutWaterTextures;
 
 		return PostProcessMaterialInputs;
 	};
diff --git a/Engine/Source/Runtime/Renderer/Private/PostProcess/PostProcessing.h b/Engine/Source/Runtime/Renderer/Private/PostProcess/PostProcessing.h
--- a/Engine/Source/Runtime/Renderer/Private/PostProcess/PostProcessing.h
+++ b/Engine/Source/Runtime/Renderer/Private/PostProcess/PostProcessing.h
@@ -49,7 +49,8 @@ void AddPostProcessingPasses(
 	const Nanite::FRasterResults* NaniteRasterResults,
 	FInstanceCullingManager& InstanceCullingManager,
 	FVirtualShadowMapArray* VirtualShadowMapArray,
-	struct FLumenSceneFrameTemporaries& LumenFrameTemporaries);
+	struct FLumenSceneFrameTemporaries& LumenFrameTemporaries,
+	const FSceneWithoutWaterTextures& SceneWithoutWaterTextures);
 
 void AddDebugViewPostProcessingPasses(FRDGBuilder& GraphBuilder, const FViewInfo& View, const FPostProcessingInputs& Inputs, const Nanite::FRasterResults* NaniteRasterResults);
 
