Fixes TSR rejection on pre-exposure changes

diff --git a/Engine/Shaders/Private/TemporalSuperResolution/TSRDecimateHistory.usf b/Engine/Shaders/Private/TemporalSuperResolution/TSRDecimateHistory.usf
--- a/Engine/Shaders/Private/TemporalSuperResolution/TSRDecimateHistory.usf
+++ b/Engine/Shaders/Private/TemporalSuperResolution/TSRDecimateHistory.usf
@@ -537,6 +537,12 @@ void MainCS(
 		}
 	}
 
+	// Correct exposure of ReprojectedHistoryGuideColor
+	{
+		tsr_halfC ReprojectedHistoryGuideLinearColor = ReprojectedHistoryGuideColor * rcp(tsr_half(1.0) - ReprojectedHistoryGuideColor) * tsr_half(HistoryPreExposureCorrection);
+		ReprojectedHistoryGuideColor = ReprojectedHistoryGuideLinearColor * rcp(tsr_half(1.0) + ReprojectedHistoryGuideLinearColor);
+	}
+
 	// Output full res
 	ISOLATE
 	{
