Don't undo PreExposure in the output of ApplyLocalExposureCS.
- The output should use the same encoding as the input.
- Also set r.Bloom.ApplyLocalExposure to false by default for now.

diff --git a/Engine/Shaders/Private/PostProcessLocalExposure.usf b/Engine/Shaders/Private/PostProcessLocalExposure.usf
--- a/Engine/Shaders/Private/PostProcessLocalExposure.usf
+++ b/Engine/Shaders/Private/PostProcessLocalExposure.usf
@@ -41,9 +41,8 @@ void ApplyLocalExposureCS(uint2 DispatchThreadId : SV_DispatchThreadID)
 		float MiddleGreyLumValue = log2(0.18 * ExposureScaleMiddleGrey.y * EyeAdaptation_LocalExposureMiddleGreyExposureCompensation);
 
 		float4 SceneColor = InputTexture.Load(uint3(DispatchThreadId + Input_ViewportMin, 0));
-		SceneColor.rgb *= View.OneOverPreExposure;
 
-		float LuminanceVal = CalculateEyeAdaptationLuminance(SceneColor.rgb);
+		float LuminanceVal = CalculateEyeAdaptationLuminance(SceneColor.rgb * View.OneOverPreExposure);
 		float LogLuminance = log2(LuminanceVal);
 
 		float BaseLogLum = CalculateBaseLogLuminance(LogLuminance, EyeAdaptation_LocalExposureBlurredLuminanceBlend, ExposureScale, UV, LumBilateralGrid, BlurredLogLum, TextureSampler, TextureSampler);
diff --git a/Engine/Source/Runtime/Renderer/Private/PostProcess/PostProcessing.cpp b/Engine/Source/Runtime/Renderer/Private/PostProcess/PostProcessing.cpp
--- a/Engine/Source/Runtime/Renderer/Private/PostProcess/PostProcessing.cpp
+++ b/Engine/Source/Runtime/Renderer/Private/PostProcess/PostProcessing.cpp
@@ -83,8 +83,8 @@ TAutoConsoleVariable<float> CVarDepthOfFieldMaxSize(
 
 TAutoConsoleVariable<bool> CVarBloomApplyLocalExposure(
 	TEXT("r.Bloom.ApplyLocalExposure"),
-	true,
-	TEXT("Whether to apply local exposure when calculating bloom, default: true"),
+	false,
+	TEXT("Whether to apply local exposure when calculating bloom, default: false"),
 	ECVF_Scalability | ECVF_RenderThreadSafe);
 
 TAutoConsoleVariable<int32> CVarPostProcessingPropagateAlpha(
