Change cards guide to have a root/tip color to make debuggin easier.

diff --git a/Engine/Plugins/Runtime/HairStrands/Source/HairStrandsCore/Private/GroomDebug.cpp b/Engine/Plugins/Runtime/HairStrands/Source/HairStrandsCore/Private/GroomDebug.cpp
--- a/Engine/Plugins/Runtime/HairStrands/Source/HairStrandsCore/Private/GroomDebug.cpp
+++ b/Engine/Plugins/Runtime/HairStrands/Source/HairStrandsCore/Private/GroomDebug.cpp
@@ -770,9 +770,6 @@ static void AddDrawDebugCardsGuidesPass(
 	// Force ShaderPrint on.
 	ShaderPrint::SetEnabled(true);
 
-	const uint32 MaxCount = 128000;
-	ShaderPrint::RequestSpaceForLines(MaxCount);
-
 	if (ShaderPrintData == nullptr || !ShaderPrint::IsEnabled(*ShaderPrintData))
 	{
 		return;
@@ -790,11 +787,16 @@ static void AddDrawDebugCardsGuidesPass(
 	}
 
 	const FHairGroupInstance::FCards::FLOD& LOD = Instance->Cards.LODs[HairLODIndex];
-	
 	if (!LOD.Guides.Data)
 	{
 		return;
 	}
+	
+	const uint32 MaxCount = FMath::Max(
+		Instance->Guides.RestResource ? Instance->Guides.RestResource->GetVertexCount() * 2 : 0,
+		LOD.Guides.RestResource ? LOD.Guides.RestResource->GetVertexCount() * 2 : 0);
+	ShaderPrint::RequestSpaceForLines(MaxCount);
+
 	TShaderMapRef<FDrawDebugCardGuidesCS> ComputeShader(ShaderMap);
 	const bool bGuideValid			= Instance->Guides.RestResource != nullptr;
 	const bool bGuideDeformValid	= Instance->Guides.DeformedResource != nullptr;
diff --git a/Engine/Plugins/Runtime/HairStrands/Source/HairStrandsEditor/Private/Widgets/SGroomEditorViewport.cpp b/Engine/Plugins/Runtime/HairStrands/Source/HairStrandsEditor/Private/Widgets/SGroomEditorViewport.cpp
--- a/Engine/Plugins/Runtime/HairStrands/Source/HairStrandsEditor/Private/Widgets/SGroomEditorViewport.cpp
+++ b/Engine/Plugins/Runtime/HairStrands/Source/HairStrandsEditor/Private/Widgets/SGroomEditorViewport.cpp
@@ -69,6 +69,7 @@ FGroomEditorViewportClient::FGroomEditorViewportClient(FAdvancedPreviewScene& In
 	EngineShowFlags.DisableAdvancedFeatures();
 	EngineShowFlags.SetSnap(0);
 	EngineShowFlags.SetTemporalAA(true);
+	EngineShowFlags.SetShaderPrint(true);
 	
 	OverrideNearClipPlane(1.0f);	
 
@@ -89,8 +90,7 @@ void FGroomEditorViewportClient::Tick(float DeltaSeconds)
 
 void FGroomEditorViewportClient::Draw(FViewport* InViewport, FCanvas* Canvas)
 {
-	TSharedPtr<SGroomEditorViewport> GroomEditorViewport = GroomEditorViewportPtr.Pin();	
-	
+	TSharedPtr<SGroomEditorViewport> GroomEditorViewport = GroomEditorViewportPtr.Pin();
 	FEditorViewportClient::Draw(InViewport, Canvas);
 }
 
diff --git a/Engine/Shaders/Private/HairStrands/HairCardsDebug.usf b/Engine/Shaders/Private/HairStrands/HairCardsDebug.usf
--- a/Engine/Shaders/Private/HairStrands/HairCardsDebug.usf
+++ b/Engine/Shaders/Private/HairStrands/HairCardsDebug.usf
@@ -126,35 +126,40 @@ void MainCS(uint3 DispatchThreadId : SV_DispatchThreadID)
 	uint4 PackedP0 = 0;
 	uint4 PackedP1 = 0;
 	float3 Offset = 0;
-	float4 Color = float4(1, 1, 0, 1);
+	float4 ColorRoot = ColorWhite;
+	float4 ColorTip  = ColorWhite;
 
 	if (DebugMode == DEBUGMODE_REN_REST)
 	{
 		PackedP0 = RenRestPosition[VertexId];
 		PackedP1 = RenRestPosition[VertexId + 1];
 		Offset   = RenRestOffset;
-		Color	 = float4(1, 1, 0, 1);
+		ColorRoot= ColorYellow;
+		ColorTip = ColorPurple;
 	}
 	else if (DebugMode == DEBUGMODE_REN_DEFORMED)
 	{
 		PackedP0 = RenDeformedPosition[VertexId];
 		PackedP1 = RenDeformedPosition[VertexId + 1];
 		Offset   = RenDeformedOffset[0].xyz;
-		Color	 = float4(1, 1, 0, 1);
+		ColorRoot= ColorYellow;
+		ColorTip = ColorPurple;
 	}
 	else if (DebugMode == DEBUGMODE_SIM_REST)
 	{
 		PackedP0 = SimRestPosition[VertexId];
 		PackedP1 = SimRestPosition[VertexId + 1];
 		Offset   = SimRestOffset;
-		Color	 = float4(1, 0, 1, 1);
+		ColorRoot= ColorCyan;
+		ColorTip = ColorRed;
 	}
 	else if (DebugMode == DEBUGMODE_SIM_DEFORMED)
 	{
 		PackedP0 = SimDeformedPosition[VertexId];
 		PackedP1 = SimDeformedPosition[VertexId + 1];
 		Offset   = SimDeformedOffset[0].xyz;
-		Color	 = float4(1, 0, 1, 1);
+		ColorRoot= ColorCyan;
+		ColorTip = ColorRed;
 	}
 
 	const FHairControlPoint P0 = UnpackHairControlPoint(PackedP0, Offset, 1, 1, 1);
@@ -162,7 +167,9 @@ void MainCS(uint3 DispatchThreadId : SV_DispatchThreadID)
 
 	if (P1.Type != HAIR_CONTROLPOINT_START)
 	{
-		AddLineWS(mul(float4(P0.Position, 1), LocalToWorld), mul(float4(P1.Position,1), LocalToWorld), Color, Color);
+		const float4 Color0 = lerp(ColorRoot, ColorTip, P0.UCoord);
+		const float4 Color1 = lerp(ColorRoot, ColorTip, P1.UCoord);
+		AddLineWS(mul(float4(P0.Position, 1), LocalToWorld), mul(float4(P1.Position,1), LocalToWorld), Color0, Color1);
 	}
 }
 
