Make SSS anisotropy behavior consistent in Wrap mode between shadow and lighting code.

The aniso value is converted into an opacity with an abs() function.

diff --git a/Engine/Shaders/Private/Strata/StrataEvaluation.ush b/Engine/Shaders/Private/Strata/StrataEvaluation.ush
--- a/Engine/Shaders/Private/Strata/StrataEvaluation.ush
+++ b/Engine/Shaders/Private/Strata/StrataEvaluation.ush
@@ -595,7 +595,7 @@ FStrataEvaluateResult StrataEvaluateBSDFCommon(FStrataBSDFContext BSDFContext, F
 				const float3 SlabDiffuseColor = bIsThin ? DiffuseColor : float3(1, 1, 1);
 				const FParticipatingMedia PM = StrataSlabCreateParticipatingMedia(SlabDiffuseColor, SLAB_SSSMFP(BSDFContext.BSDF));
 				const float3 SubSurfaceColor = IsotropicMediumSlabTransmittance(PM, STRATA_SIMPLEVOLUME_THICKNESS_M, TransmittanceNoL);
-				const float Opacity = max(0.f, 1.f - SLAB_SSSPHASEANISOTROPY(BSDFContext.BSDF));
+				const float Opacity = 1.f - abs(SLAB_SSSPHASEANISOTROPY(BSDFContext.BSDF));
 
 				float3 TransmissionThroughput = 0;
 				if (bIsThin)
