Fixes unecessary LDS store in the TSR's RejectShading pass

diff --git a/Engine/Shaders/Private/TemporalSuperResolution/TSRRejectShading.usf b/Engine/Shaders/Private/TemporalSuperResolution/TSRRejectShading.usf
--- a/Engine/Shaders/Private/TemporalSuperResolution/TSRRejectShading.usf
+++ b/Engine/Shaders/Private/TemporalSuperResolution/TSRRejectShading.usf
@@ -140,6 +140,12 @@ struct TAccumulatorAvg
 	Type Value;
 	tsr_simd_half Weight;
 	
+	CALL_SITE_DEBUGLOC
+	static bool NeedWeightInAccumulate()
+	{
+		return true;
+	}
+
 	CALL_SITE_DEBUGLOC
 	void Start(Type CenterSample)
 	{
@@ -174,6 +180,12 @@ struct TAccumulatorMin
 {
 	Type Value;
 	
+	CALL_SITE_DEBUGLOC
+	static bool NeedWeightInAccumulate()
+	{
+		return false;
+	}
+
 	CALL_SITE_DEBUGLOC
 	void Start(Type CenterSample)
 	{
@@ -204,6 +216,12 @@ struct TAccumulatorMax
 {
 	Type Value;
 	
+	CALL_SITE_DEBUGLOC
+	static bool NeedWeightInAccumulate()
+	{
+		return false;
+	}
+
 	CALL_SITE_DEBUGLOC
 	void Start(Type CenterSample)
 	{
@@ -241,13 +259,21 @@ FAccumulator Convole3x1(FSampleType Color, tsr_simd_half Weight, const tsr_short
 	
 	FSampleType PreprocessedNeighborColor = FAccumulator::PreprocessNeighborSample(Color, Weight);
 
-	FSampleType Sample0 = AccessNeighborTexel(PreprocessedNeighborColor, +Offset, GroupThreadIndex);
-	tsr_simd_half Weight0 = AccessNeighborTexel(Weight, +Offset, GroupThreadIndex);
+	FSampleType Sample0 = AccessNeighborTexel(PreprocessedNeighborColor, Offset, GroupThreadIndex);
+	tsr_simd_half Weight0 = tsr_simd_half::Const(tsr_half(1.0));
+	if (FAccumulator::NeedWeightInAccumulate())
+	{
+		Weight0 = AccessNeighborTexel(Weight, Offset, GroupThreadIndex);
+	}
 	Accumulator.Accumulate(+Offset, Sample0, Weight0);
 	
 	const tsr_short2 OffsetN = tsr_short2(-int2(Offset)); // work arround a shader compiler bug
 	FSampleType Sample1 = AccessNeighborTexel(PreprocessedNeighborColor, OffsetN, GroupThreadIndex);
-	tsr_simd_half Weight1 = AccessNeighborTexel(Weight, OffsetN, GroupThreadIndex);
+	tsr_simd_half Weight1 = tsr_simd_half::Const(tsr_half(1.0));
+	if (FAccumulator::NeedWeightInAccumulate())
+	{
+		Weight1 = AccessNeighborTexel(Weight, OffsetN, GroupThreadIndex);
+	}
 	Accumulator.Accumulate(OffsetN, Sample1, Weight1);
 
 	Accumulator.Finish(3);
