VulkanRHI: Avoid "groupshared bool" on Vulkan, seems to have issues on NVIDIA.  Fixes crash in Vulkan engine tests (and playback of D3D12 captures in RDOC).  Fix reference implementation for future generations.

diff --git a/Engine/Shaders/Private/SingleLayerWaterComposite.usf b/Engine/Shaders/Private/SingleLayerWaterComposite.usf
--- a/Engine/Shaders/Private/SingleLayerWaterComposite.usf
+++ b/Engine/Shaders/Private/SingleLayerWaterComposite.usf
@@ -292,7 +292,7 @@ RWBuffer<uint> DrawIndirectDataUAV;
 RWBuffer<uint> DispatchIndirectDataUAV;
 RWBuffer<uint> WaterTileListDataUAV;
 
-groupshared bool ContainsWater[WORK_TILE_SIZE*WORK_TILE_SIZE];
+groupshared uint ContainsWater[WORK_TILE_SIZE*WORK_TILE_SIZE];
 
 [numthreads(WORK_TILE_SIZE, WORK_TILE_SIZE, 1)]
 void WaterTileCatergorisationCS(uint3 ThreadId : SV_DispatchThreadID, uint3 GroupId : SV_GroupID, uint3 GroupThreadId : SV_GroupThreadID)
@@ -335,7 +335,7 @@ void WaterTileCatergorisationCS(uint3 ThreadId : SV_DispatchThreadID, uint3 Grou
 		if (bContainsWater)
 		{
 			uint WriteToIndex;
-			InterlockedAdd(DrawIndirectDataUAV[0], 1, WriteToIndex);
+			InterlockedAdd(DrawIndirectDataUAV[1], 1, WriteToIndex);
 			InterlockedAdd(DispatchIndirectDataUAV[0], 1);
 			WaterTileListDataUAV[WriteToIndex] = GroupId.x | (GroupId.y << 16); // assumes 16bit is enought to represent a tiled resolution up to 65,535 :)
 		}
@@ -352,47 +352,47 @@ void WaterTileCatergorisationCS(uint3 ThreadId : SV_DispatchThreadID, uint3 Grou
 	#if STRATA_ENABLED
 		FStrataAddressing StrataAddressing = GetStrataPixelDataByteOffset(PixelPos, uint2(View.BufferSizeAndInvSize.xy), Strata.MaxBytesPerPixel);
 		FStrataPixelHeader StrataPixelHeader = UnpackStrataHeaderIn(Strata.MaterialTextureArray, StrataAddressing, Strata.TopLayerTexture);
-		bContainsWater = bContainsWater || (StrataPixelHeader.BSDFCount > 0 && IsSingleLayerWater(StrataPixelHeader));
+		bContainsWater = (StrataPixelHeader.BSDFCount > 0 && IsSingleLayerWater(StrataPixelHeader));
 	#else
 		FGBufferData GBuffer = GetGBufferDataFromSceneTextures(BufferUV);
-		bContainsWater = bContainsWater || GBuffer.ShadingModelID == SHADINGMODELID_SINGLELAYERWATER;
+		bContainsWater = (GBuffer.ShadingModelID == SHADINGMODELID_SINGLELAYERWATER);
 	#endif
 	}
 
 	const int LinearIndex = GroupThreadId.y * WORK_TILE_SIZE + GroupThreadId.x;
-	ContainsWater[LinearIndex] = bContainsWater;
+	ContainsWater[LinearIndex] = bContainsWater ? 1 : 0;
 	GroupMemoryBarrierWithGroupSync();
 
-	if (LinearIndex < 32) // 8*8 = 64 elements to merge
+	if ((LinearIndex < 32) && (ContainsWater[LinearIndex + 32] != 0)) // 8*8 = 64 elements to merge
 	{
-		ContainsWater[LinearIndex] = ContainsWater[LinearIndex] || ContainsWater[LinearIndex + 32];
+		ContainsWater[LinearIndex] = 1;
 	}
 	GroupMemoryBarrierWithGroupSync();
-	if (LinearIndex < 16)
+	if ((LinearIndex < 16) && (ContainsWater[LinearIndex + 16] != 0))
 	{
-		ContainsWater[LinearIndex] = ContainsWater[LinearIndex] || ContainsWater[LinearIndex + 16];
+		ContainsWater[LinearIndex] = 1;
 	}
 	GroupMemoryBarrierWithGroupSync();
 
 	// The smallest wave size is 16 on Intel hardware. So now we can do simple math operations without group sync.
 	// EDIT: for some reason group sync is needed until the end, otherwise some pixels are missing...
 
-	if (LinearIndex < 8)
+	if ((LinearIndex < 8) && (ContainsWater[LinearIndex + 8] != 0))
 	{
-		ContainsWater[LinearIndex] = ContainsWater[LinearIndex] || ContainsWater[LinearIndex + 8];
+		ContainsWater[LinearIndex] = 1;
 	}
 	GroupMemoryBarrierWithGroupSync();
-	if (LinearIndex < 4)
+	if ((LinearIndex < 4) && (ContainsWater[LinearIndex + 4] != 0))
 	{
-		ContainsWater[LinearIndex] = ContainsWater[LinearIndex] || ContainsWater[LinearIndex + 4];
+		ContainsWater[LinearIndex] = 1;
 	}
 	GroupMemoryBarrierWithGroupSync();
-	if (LinearIndex < 2)
+	if ((LinearIndex < 2) && (ContainsWater[LinearIndex + 2] != 0))
 	{
-		ContainsWater[LinearIndex] = ContainsWater[LinearIndex] || ContainsWater[LinearIndex + 2];
+		ContainsWater[LinearIndex] = 1;
 	}
 	GroupMemoryBarrierWithGroupSync();
-	if (LinearIndex < 1 && (ContainsWater[LinearIndex] || ContainsWater[LinearIndex + 1]))
+	if (LinearIndex < 1 && (bContainsWater || (ContainsWater[LinearIndex + 1] != 0)))
 	{
 		uint WriteToIndex;  
 		InterlockedAdd(DrawIndirectDataUAV[1], 1, WriteToIndex);
