Fix copyright notice in recently added file

diff --git a/Engine/Shaders/Private/RayTracing/RayTracingLightFunctionCommon.ush b/Engine/Shaders/Private/RayTracing/RayTracingLightFunctionCommon.ush
--- a/Engine/Shaders/Private/RayTracing/RayTracingLightFunctionCommon.ush
+++ b/Engine/Shaders/Private/RayTracing/RayTracingLightFunctionCommon.ush
@@ -1,4 +1,4 @@
-// Copyright 1998-2019 Epic Games, Inc. All Rights Reserved.
+// Copyright Epic Games, Inc. All Rights Reserved.
 
 
 // base light function material shader, so it can be folded where apprpriate
