Add support for the various Strata eye normal reconstruction.

TangentOutput0 output is not piped yet.

diff --git a/Engine/Shaders/Private/MaterialTemplate.ush b/Engine/Shaders/Private/MaterialTemplate.ush
--- a/Engine/Shaders/Private/MaterialTemplate.ush
+++ b/Engine/Shaders/Private/MaterialTemplate.ush
@@ -3474,3 +3474,38 @@ float3 GetWorldBentNormalZero(in FMaterialPixelParameters MaterialParameters)
 	return MaterialParameters.WorldNormal;
 #endif
 }
+
+// Return the Iris & Caustic normal which should be used for eye shading model
+void GetEyeNormals(
+	float IrisMask,
+	float IrisDistance,
+	in float3 InNormal,
+	in float3 InClearCoatNormal,
+	in float3 InCustomTangent,
+	inout float3 OutIrisNormal,
+	inout float3 OutCausticNormal)
+{
+#if IRIS_NORMAL
+  #if NUM_MATERIAL_OUTPUTS_GETTANGENTOUTPUT > 0
+	// Blend in the negative intersection normal to create some concavity
+	// Not great as it ties the concavity to the convexity of the cornea surface
+	// No good justification for that. On the other hand, if we're just looking to
+	// introduce some concavity, this does the job.
+	const float3 PlaneNormal = normalize(InCustomTangent);
+	OutCausticNormal = normalize(lerp(PlaneNormal, -InNormal, IrisMask * IrisDistance));
+  #else
+	const float3 PlaneNormal = InNormal;
+	OutCausticNormal = InNormal;
+  #endif
+
+  #if NUM_MATERIAL_OUTPUTS_CLEARCOATBOTTOMNORMAL > 0
+	OutIrisNormal = normalize(InClearCoatNormal);
+  #else
+	OutIrisNormal = PlaneNormal;
+  #endif
+#else
+  #if NUM_MATERIAL_OUTPUTS_GETTANGENTOUTPUT > 0
+	OutIrisNormal = normalize(InCustomTangent);
+  #endif
+#endif
+}
diff --git a/Engine/Shaders/Private/Strata/StrataLegacyConversion.ush b/Engine/Shaders/Private/Strata/StrataLegacyConversion.ush
--- a/Engine/Shaders/Private/Strata/StrataLegacyConversion.ush
+++ b/Engine/Shaders/Private/Strata/StrataLegacyConversion.ush
@@ -4,6 +4,20 @@
 
 #include "../Common.ush"
 
+///////////////////////////////////////////////////////////////////////////////////////////////////
+// Helpers
+
+void GetEyeNormals(
+	float IrisMask,
+	float IrisDistance,
+	in float3 InNormal,
+	in float3 InClearCoatNormal,
+	in float3 InCustomTangent,
+	inout float3 OutIrisNormal,
+	inout float3 OutCausticNormal);
+
+///////////////////////////////////////////////////////////////////////////////////////////////////
+
 struct FStrataLegacyParameters
 {
 	FStrataPixelFootprint PixelFootprint;
@@ -529,7 +543,10 @@ FStrataData StrataConvertLegacyMaterialStatic(
 	{
 		const float IrisMask = ClearCoat;
 		const float IrisDistance = ClearCoatRoughness;
-		const float3 IrisNormal = RawClearCoatNormal;
+
+		float3 IrisNormal = RawNormal;
+		float3 CausticNormal = RawNormal;
+		GetEyeNormals(IrisMask, IrisDistance, RawNormal, RawClearCoatNormal, RawClearCoatNormal /*RawCustomTangent*/, IrisNormal, CausticNormal);
 
 		FStrataData EyeMat = PromoteParameterBlendedBSDFToOperator(
 			GetStrataEyeBSDF(
diff --git a/Engine/Shaders/Private/Strata/StrataVisualize.usf b/Engine/Shaders/Private/Strata/StrataVisualize.usf
--- a/Engine/Shaders/Private/Strata/StrataVisualize.usf
+++ b/Engine/Shaders/Private/Strata/StrataVisualize.usf
@@ -349,6 +349,10 @@ void PrintBSDF(inout FShaderPrintContext Context, inout FStrataAddressing Strata
 			Print(Context, TEXT(" Roughness         "), FontBSDFPropName);	Print(Context, EYE_ROUGHNESS(BSDF), FontBSDFPropValu);				Newline(Context, RectMax);
 			Print(Context, TEXT(" Iris Mask         "), FontBSDFPropName);	Print(Context, EYE_IRISDISTANCE(BSDF), FontBSDFPropValu);			Newline(Context, RectMax);
 			Print(Context, TEXT(" Iris Distance     "), FontBSDFPropName);	Print(Context, EYE_IRISMASK(BSDF), FontBSDFPropValu);				Newline(Context, RectMax);
+
+			// Tangent basis for the iris
+			const float3x3 IrisTangentBasis = GetTangentBasis(EYE_IRISNORMAL(BSDF));
+			DrawReferentialTWS(WorldPosition, IrisTangentBasis[0], IrisTangentBasis[1], IrisTangentBasis[2], float3(0, 1, 1));
 		}
 		break;
 
diff --git a/Engine/Source/Runtime/Engine/Private/Materials/Material.cpp b/Engine/Source/Runtime/Engine/Private/Materials/Material.cpp
--- a/Engine/Source/Runtime/Engine/Private/Materials/Material.cpp
+++ b/Engine/Source/Runtime/Engine/Private/Materials/Material.cpp
@@ -3126,7 +3126,7 @@ void UMaterial::ConvertMaterialToStrataMaterial()
 					MoveConnectionTo(EditorOnly->Tangent, BottomSlabBSDF, 10);					// Tangent
 
 					check(ClearCoatBottomNormalOutput);
-					MoveConnectionTo(*ClearCoatBottomNormalOutput->GetInput(0), BottomSlabBSDF, 9);// ClearColorBottomNormal -> BottomSlabBSDF.Normal
+					CopyConnectionTo(*ClearCoatBottomNormalOutput->GetInput(0), BottomSlabBSDF, 9);// ClearColorBottomNormal -> BottomSlabBSDF.Normal
 
 					// Now weight the top base material by opacity.
 					UMaterialExpressionStrataSlabBSDF* TopSlabBSDF = NewObject<UMaterialExpressionStrataSlabBSDF>(this);
@@ -3353,6 +3353,10 @@ void UMaterial::ConvertMaterialToStrataMaterial()
 			MoveConnectionTo(*SingleLayerWaterOutput->GetInput(2), ConvertNode, 15); // WaterPhaseG
 			MoveConnectionTo(*SingleLayerWaterOutput->GetInput(3), ConvertNode, 16); // ColorScaleBehindWater
 		}
+		if (ClearCoatBottomNormalOutput)
+		{
+			CopyConnectionTo(*ClearCoatBottomNormalOutput->GetInput(0), ConvertNode, 17);		 // ClearCoatNormal
+		}
 	}
 
 	if (bInvalidateShader)
