Add iris plane support for Strata eye.

diff --git a/Engine/Shaders/Private/MaterialTemplate.ush b/Engine/Shaders/Private/MaterialTemplate.ush
--- a/Engine/Shaders/Private/MaterialTemplate.ush
+++ b/Engine/Shaders/Private/MaterialTemplate.ush
@@ -3475,7 +3475,7 @@ float3 GetWorldBentNormalZero(in FMaterialPixelParameters MaterialParameters)
 #endif
 }
 
-// Return the Iris & Caustic normal which should be used for eye shading model
+// Return the Iris & Iris plane normal, used for eye shading model
 void GetEyeNormals(
 	float IrisMask,
 	float IrisDistance,
@@ -3483,29 +3483,27 @@ void GetEyeNormals(
 	in float3 InClearCoatNormal,
 	in float3 InCustomTangent,
 	inout float3 OutIrisNormal,
-	inout float3 OutCausticNormal)
+	inout float3 OutIrisPlaneNormal)
 {
 #if IRIS_NORMAL
-  #if NUM_MATERIAL_OUTPUTS_GETTANGENTOUTPUT > 0
-	// Blend in the negative intersection normal to create some concavity
-	// Not great as it ties the concavity to the convexity of the cornea surface
-	// No good justification for that. On the other hand, if we're just looking to
-	// introduce some concavity, this does the job.
-	const float3 PlaneNormal = normalize(InCustomTangent);
-	OutCausticNormal = normalize(lerp(PlaneNormal, -InNormal, IrisMask * IrisDistance));
+  #if NUM_MATERIAL_OUTPUTS_CLEARCOATBOTTOMNORMAL > 0
+	OutIrisNormal = normalize(InClearCoatNormal);
   #else
-	const float3 PlaneNormal = InNormal;
-	OutCausticNormal = InNormal;
+	OutIrisNormal = InNormal;
   #endif
 
-  #if NUM_MATERIAL_OUTPUTS_CLEARCOATBOTTOMNORMAL > 0
-	OutIrisNormal = normalize(InClearCoatNormal);
+  #if NUM_MATERIAL_OUTPUTS_GETTANGENTOUTPUT > 0
+	OutIrisPlaneNormal = normalize(InCustomTangent);
   #else
-	OutIrisNormal = PlaneNormal;
+	OutIrisPlaneNormal = OutIrisNormal;
   #endif
 #else
   #if NUM_MATERIAL_OUTPUTS_GETTANGENTOUTPUT > 0
 	OutIrisNormal = normalize(InCustomTangent);
+	OutIrisPlaneNormal = OutIrisNormal;
+  #else
+	OutIrisNormal = InNormal;
+	OutIrisPlaneNormal = InNormal;
   #endif
 #endif
 }
diff --git a/Engine/Shaders/Private/Strata/Strata.ush b/Engine/Shaders/Private/Strata/Strata.ush
--- a/Engine/Shaders/Private/Strata/Strata.ush
+++ b/Engine/Shaders/Private/Strata/Strata.ush
@@ -448,6 +448,7 @@ struct FStrataBSDF
 #define EYE_IRISMASK(X)					X.VGPRs[1].y
 #define EYE_SSSPROFILEID(X)				X.VGPRs[1].z
 #define EYE_IRISNORMAL(X)				X.VGPRs[2].xyz
+#define EYE_IRISPLANENORMAL(X)			X.VGPRs[3].xyz
 
 ///////////////////////////////////////////////////////////////////////////////
 // Single Layer Water BSDF members
@@ -1426,7 +1427,7 @@ FStrataData GetStrataHairBSDF(float3 BaseColor, float Scatter, float Specular, f
 	return StrataData;
 }
 
-FStrataData GetStrataEyeBSDF(float3 DiffuseAlbedo, float Roughness, float IrisMask, float IrisDistance, float3 IrisNormal, float SSSProfileId, float3 Emissive, uint SharedLocalBasisIndex)
+FStrataData GetStrataEyeBSDF(float3 DiffuseAlbedo, float Roughness, float IrisMask, float IrisDistance, float3 IrisNormal, float3 IrisPlaneNormal, float SSSProfileId, float3 Emissive, uint SharedLocalBasisIndex)
 {
 	FStrataData StrataData = GetInitialisedStrataData();
 	BSDF_SETTYPE					(StrataData.InlinedBSDF, STRATA_BSDF_TYPE_EYE);
@@ -1437,6 +1438,7 @@ FStrataData GetStrataEyeBSDF(float3 DiffuseAlbedo, float Roughness, float IrisMa
 	EYE_IRISDISTANCE				(StrataData.InlinedBSDF) = IrisDistance;
 	EYE_IRISMASK					(StrataData.InlinedBSDF) = IrisMask;
 	EYE_IRISNORMAL					(StrataData.InlinedBSDF) = IrisNormal;
+	EYE_IRISPLANENORMAL				(StrataData.InlinedBSDF) = IrisPlaneNormal;
 	EYE_SSSPROFILEID				(StrataData.InlinedBSDF) = SSSProfileId;
 
 	const bool bHasSSS = SSSProfileId != 0;
@@ -2395,6 +2397,7 @@ void SanitizeStrataEye(inout FStrataBSDF BSDF)
 	EYE_IRISDISTANCE(BSDF) = saturate(EYE_IRISDISTANCE(BSDF));
 	EYE_IRISMASK(BSDF) = saturate(EYE_IRISMASK(BSDF));
 	EYE_IRISNORMAL(BSDF) = normalize(EYE_IRISNORMAL(BSDF));
+	EYE_IRISPLANENORMAL(BSDF) = normalize(EYE_IRISPLANENORMAL(BSDF));
 }
 
 void SanitizeStrataSingleLayerWater(inout FStrataBSDF BSDF)
@@ -3530,12 +3533,13 @@ FStrataBSDF UnpackStrataBSDFIn(FStrataMaterialContainer StrataBuffer, inout FStr
 	case STRATA_BSDF_TYPE_EYE:
 	{
 		StrataLoad_ColorGamma2ToLinearAlphaLinear(StrataBuffer, StrataAddressing, EYE_DIFFUSEALBEDO(OutBSDF), EYE_ROUGHNESS(OutBSDF));		
-		float4 Data = UnpackRGBA8(StrataLoadUint1(StrataBuffer, StrataAddressing));
-		EYE_IRISMASK(OutBSDF) = Data.x;
-		EYE_IRISDISTANCE(OutBSDF) = Data.y;
-		EYE_IRISNORMAL(OutBSDF) = OctahedronToUnitVector(Data.zw * 2.0f - 1.f);
-		// STRATA_TODO - Add SSS profile support
-		// 8 bytes
+		const float4 Data0 = UnpackRGBA8(StrataLoadUint1(StrataBuffer, StrataAddressing));
+		const float4 Data1 = UnpackRGBA8(StrataLoadUint1(StrataBuffer, StrataAddressing));
+		EYE_IRISMASK(OutBSDF) = Data0.x;
+		EYE_IRISDISTANCE(OutBSDF) = Data0.y;
+		EYE_IRISNORMAL(OutBSDF) = OctahedronToUnitVector(Data1.xy * 2.0f - 1.f);
+		EYE_IRISPLANENORMAL(OutBSDF) = OctahedronToUnitVector(Data1.zw * 2.0f - 1.f);
+		// 12 bytes
 	}
 	break;
 	case STRATA_BSDF_TYPE_SINGLELAYERWATER:
diff --git a/Engine/Shaders/Private/Strata/StrataEvaluation.ush b/Engine/Shaders/Private/Strata/StrataEvaluation.ush
--- a/Engine/Shaders/Private/Strata/StrataEvaluation.ush
+++ b/Engine/Shaders/Private/Strata/StrataEvaluation.ush
@@ -776,7 +776,7 @@ FStrataEvaluateResult StrataEvaluateBSDFCommon(FStrataBSDFContext BSDFContext, F
 			// Not great as it ties the concavity to the convexity of the cornea surface
 			// No good justification for that. On the other hand, if we're just looking to
 			// introduce some concavity, this does the job.
-			const float3 CausticNormal = normalize(lerp(EYE_IRISNORMAL(BSDFContext.BSDF), -BSDFContext.N, EYE_IRISMASK(BSDFContext.BSDF) * EYE_IRISDISTANCE(BSDFContext.BSDF)));
+			const float3 CausticNormal = normalize(lerp(EYE_IRISPLANENORMAL(BSDFContext.BSDF), -BSDFContext.N, EYE_IRISMASK(BSDFContext.BSDF) * EYE_IRISDISTANCE(BSDFContext.BSDF)));
 
 			// Add-hoc legacy shading mode for eye. 
 			const float Power = lerp(12, 1, IrisNoL);
diff --git a/Engine/Shaders/Private/Strata/StrataExport.ush b/Engine/Shaders/Private/Strata/StrataExport.ush
--- a/Engine/Shaders/Private/Strata/StrataExport.ush
+++ b/Engine/Shaders/Private/Strata/StrataExport.ush
@@ -553,9 +553,11 @@ void PackStrataOut(
 						case STRATA_BSDF_TYPE_EYE:
 						{
 							const float2 EncodedIrisNormal = UnitVectorToOctahedron(EYE_IRISNORMAL(BSDF));
+							const float2 EncodedIrisPlaneNormal = UnitVectorToOctahedron(EYE_IRISPLANENORMAL(BSDF));
 							STRATA_STORE_UINT1(PackColorLinearToGamma2AlphaLinear(float4(EYE_DIFFUSEALBEDO(BSDF), EYE_ROUGHNESS(BSDF))));
-							STRATA_STORE_UINT1(PackRGBA8(float4(EYE_IRISMASK(BSDF), EYE_IRISDISTANCE(BSDF), EncodedIrisNormal * 0.5f + 0.5f)));
-							// 8 bytes
+							STRATA_STORE_UINT1(PackRGBA8(float4(EYE_IRISMASK(BSDF), EYE_IRISDISTANCE(BSDF), 0, 0)));
+							STRATA_STORE_UINT1(PackRGBA8(float4(EncodedIrisNormal * 0.5f + 0.5f, EncodedIrisPlaneNormal * 0.5f + 0.5f)));
+							// 12 bytes
 
 							// Note: we don't store the SSS profile ID into the Eye data, as it is not needed there. The profile is only stored into the SSS data.
 							const bool bHasSSS = BSDF_GETSSSTYPE(BSDF) == SSS_TYPE_DIFFUSION_PROFILE;
diff --git a/Engine/Shaders/Private/Strata/StrataLegacyConversion.ush b/Engine/Shaders/Private/Strata/StrataLegacyConversion.ush
--- a/Engine/Shaders/Private/Strata/StrataLegacyConversion.ush
+++ b/Engine/Shaders/Private/Strata/StrataLegacyConversion.ush
@@ -14,7 +14,7 @@ void GetEyeNormals(
 	in float3 InClearCoatNormal,
 	in float3 InCustomTangent,
 	inout float3 OutIrisNormal,
-	inout float3 OutCausticNormal);
+	inout float3 OutIrisPlaneNormal);
 
 ///////////////////////////////////////////////////////////////////////////////////////////////////
 
@@ -547,8 +547,8 @@ FStrataData StrataConvertLegacyMaterialStatic(
 		const float IrisDistance = ClearCoatRoughness;
 
 		float3 IrisNormal = RawNormal;
-		float3 CausticNormal = RawNormal;
-		GetEyeNormals(IrisMask, IrisDistance, RawNormal, RawClearCoatNormal, RawCustomTangent, IrisNormal, CausticNormal);
+		float3 IrisPlaneNormal = RawNormal;
+		GetEyeNormals(IrisMask, IrisDistance, RawNormal, RawClearCoatNormal, RawCustomTangent, IrisNormal, IrisPlaneNormal);
 
 		FStrataData EyeMat = PromoteParameterBlendedBSDFToOperator(
 			GetStrataEyeBSDF(
@@ -557,6 +557,7 @@ FStrataData StrataConvertLegacyMaterialStatic(
 				IrisMask,	         // IrisMask
 				IrisDistance,        // IrisDistance
 				IrisNormal,	         // IrisNormal
+				IrisPlaneNormal,	 // IrisPlaneNormal
 				SubSurfaceProfileId, // SSS profile
 				Emissive,	         // EmissiveColor
 				SharedLocalBasisIndex),
diff --git a/Engine/Shaders/Private/Strata/StrataVisualize.usf b/Engine/Shaders/Private/Strata/StrataVisualize.usf
--- a/Engine/Shaders/Private/Strata/StrataVisualize.usf
+++ b/Engine/Shaders/Private/Strata/StrataVisualize.usf
@@ -356,6 +356,10 @@ void PrintBSDF(inout FShaderPrintContext Context, inout FStrataAddressing Strata
 			const float3x3 IrisTangentBasis = GetTangentBasis(EYE_IRISNORMAL(BSDF));
 			DrawReferentialTWS(WorldPosition, IrisTangentBasis[0], IrisTangentBasis[1], IrisTangentBasis[2], float3(0, 1, 1));
 
+			// Tangent basis for the iris plane
+			const float3x3 IrisPlaneTangentBasis = GetTangentBasis(EYE_IRISPLANENORMAL(BSDF));
+			DrawReferentialTWS(WorldPosition, IrisPlaneTangentBasis[0], IrisPlaneTangentBasis[1], IrisPlaneTangentBasis[2], float3(1, 0, 1));
+
 			const bool bHasSSS = HasSubsurface(Header);
 			const FStrataSubsurfaceHeader SSSHeader = StrataLoadSubsurfaceHeader(Strata.SSSTexture, StrataAddressing.PixelCoords);
 			const bool bIsValid = StrataSubSurfaceHeaderGetIsValid(SSSHeader);
diff --git a/Engine/Source/Runtime/Engine/Classes/Materials/MaterialExpressionStrata.h b/Engine/Source/Runtime/Engine/Classes/Materials/MaterialExpressionStrata.h
--- a/Engine/Source/Runtime/Engine/Classes/Materials/MaterialExpressionStrata.h
+++ b/Engine/Source/Runtime/Engine/Classes/Materials/MaterialExpressionStrata.h
@@ -586,6 +586,12 @@ class UMaterialExpressionStrataEyeBSDF : public UMaterialExpressionStrataBSDF
 	UPROPERTY()
 	FExpressionInput IrisNormal;
 
+	/**
+	 * Normal of the iris plane (type = float3, unit = unitless, defaults to +X vector)
+	 */
+	UPROPERTY()
+	FExpressionInput IrisPlaneNormal;
+
 	/**
 	 * Mask defining the iris surface (type = float, unit = unitless, defaults to 0.0)
 	 */
diff --git a/Engine/Source/Runtime/Engine/Private/Materials/HLSLMaterialTranslator.cpp b/Engine/Source/Runtime/Engine/Private/Materials/HLSLMaterialTranslator.cpp
--- a/Engine/Source/Runtime/Engine/Private/Materials/HLSLMaterialTranslator.cpp
+++ b/Engine/Source/Runtime/Engine/Private/Materials/HLSLMaterialTranslator.cpp
@@ -10492,15 +10492,16 @@ int32 FHLSLMaterialTranslator::StrataHairBSDF(int32 BaseColor, int32 Scatter, in
 	);
 }
 
-int32 FHLSLMaterialTranslator::StrataEyeBSDF(int32 DiffuseAlbedo, int32 Roughness, int32 IrisMask, int32 IrisDistance, int32 EmissiveColor, int32 CorneaNormal, int32 IrisNormal, int32 SSSProfileId, const FString& SharedLocalBasisIndexMacro)
+int32 FHLSLMaterialTranslator::StrataEyeBSDF(int32 DiffuseAlbedo, int32 Roughness, int32 IrisMask, int32 IrisDistance, int32 EmissiveColor, int32 CorneaNormal, int32 IrisNormal, int32 IrisPlaneNormal, int32 SSSProfileId, const FString& SharedLocalBasisIndexMacro)
 {
 	return AddCodeChunk(
-		MCT_Strata, TEXT("GetStrataEyeBSDF(%s, %s, %s, %s, %s, %s, %s) /* Cornea:%s Iris:%s */"),
+		MCT_Strata, TEXT("GetStrataEyeBSDF(%s, %s, %s, %s, %s, %s, %s, %s) /* Cornea:%s Iris:%s */"),
 		*StrataGetCastParameterCode(DiffuseAlbedo, MCT_Float3),
 		*StrataGetCastParameterCode(Roughness, MCT_Float),
 		*StrataGetCastParameterCode(IrisMask, MCT_Float),
 		*StrataGetCastParameterCode(IrisDistance, MCT_Float),
-		*StrataGetCastParameterCode(IrisNormal, MCT_Float),
+		*StrataGetCastParameterCode(IrisNormal, MCT_Float3),
+		*StrataGetCastParameterCode(IrisPlaneNormal, MCT_Float3),
 		*StrataGetCastParameterCode(SSSProfileId, MCT_Float),
 		*StrataGetCastParameterCode(EmissiveColor, MCT_Float3),
 		*SharedLocalBasisIndexMacro,
diff --git a/Engine/Source/Runtime/Engine/Private/Materials/HLSLMaterialTranslator.h b/Engine/Source/Runtime/Engine/Private/Materials/HLSLMaterialTranslator.h
--- a/Engine/Source/Runtime/Engine/Private/Materials/HLSLMaterialTranslator.h
+++ b/Engine/Source/Runtime/Engine/Private/Materials/HLSLMaterialTranslator.h
@@ -1047,7 +1047,7 @@ protected:
 	virtual int32 StrataVolumetricFogCloudBSDF(int32 Albedo, int32 Extinction, int32 EmissiveColor, int32 AmbientOcclusion) override;
 	virtual int32 StrataUnlitBSDF(int32 EmissiveColor, int32 TransmittanceColor) override;
 	virtual int32 StrataHairBSDF(int32 BaseColor, int32 Scatter, int32 Specular, int32 Roughness, int32 Backlit, int32 EmissiveColor, int32 Tangent, const FString& SharedLocalBasisIndexMacro) override;
-	virtual int32 StrataEyeBSDF(int32 DiffuseColor, int32 Roughness, int32 IrisMask, int32 IrisDistance, int32 EmissiveColor, int32 CorneaNormal, int32 IrisNormal, int32 SSSProfileId, const FString& SharedLocalBasisIndexMacro) override;
+	virtual int32 StrataEyeBSDF(int32 DiffuseColor, int32 Roughness, int32 IrisMask, int32 IrisDistance, int32 EmissiveColor, int32 CorneaNormal, int32 IrisNormal, int32 IrisPlaneNormal, int32 SSSProfileId, const FString& SharedLocalBasisIndexMacro) override;
 	virtual int32 StrataSingleLayerWaterBSDF(
 		int32 BaseColor, int32 Metallic, int32 Specular, int32 Roughness, 
 		int32 EmissiveColor, int32 TopMaterialOpacity, int32 WaterAlbedo, int32 WaterExtinction, int32 WaterPhaseG, 
diff --git a/Engine/Source/Runtime/Engine/Private/Materials/MaterialExpressions.cpp b/Engine/Source/Runtime/Engine/Private/Materials/MaterialExpressions.cpp
--- a/Engine/Source/Runtime/Engine/Private/Materials/MaterialExpressions.cpp
+++ b/Engine/Source/Runtime/Engine/Private/Materials/MaterialExpressions.cpp
@@ -23335,6 +23335,7 @@ int32 UMaterialExpressionStrataEyeBSDF::Compile(class FMaterialCompiler* Compile
 		CompileWithDefaultFloat1(Compiler, IrisMask,	 0.0f),
 		CompileWithDefaultFloat1(Compiler, IrisDistance, 0.0f),
 		CompileWithDefaultFloat3(Compiler, IrisNormal,   0.0f, 0.0f, 1.0f),
+		CompileWithDefaultFloat3(Compiler, IrisPlaneNormal, 0.0f, 0.0f, 1.0f),
 		SSSProfileCodeChunk != INDEX_NONE ? SSSProfileCodeChunk : Compiler->Constant(0.0f),
 		CompileWithDefaultFloat3(Compiler, EmissiveColor,0.0f, 0.0f, 0.0f),
 		CorneaNormalCodeChunk,
@@ -23361,9 +23362,10 @@ uint32 UMaterialExpressionStrataEyeBSDF::GetInputType(int32 InputIndex)
 	case 1: return MCT_Float1; // Roughness
 	case 2: return MCT_Float3; // Cornea normal
 	case 3: return MCT_Float3; // IrisNormal
-	case 4: return MCT_Float1; // IrisMask
-	case 5: return MCT_Float1; // IrisDistance
-	case 6: return MCT_Float3; // EmissiveColor
+	case 4: return MCT_Float3; // IrisPlaneNormal
+	case 5: return MCT_Float1; // IrisMask
+	case 6: return MCT_Float1; // IrisDistance
+	case 7: return MCT_Float3; // EmissiveColor
 	}
 
 	check(false);
@@ -23382,6 +23384,7 @@ void UMaterialExpressionStrataEyeBSDF::GatherStrataMaterialInfo(FStrataMaterialI
 	if (Roughness.IsConnected())		{ StrataMaterialInfo.AddPropertyConnected(MP_Roughness); }
 	if (CorneaNormal.IsConnected())		{ StrataMaterialInfo.AddPropertyConnected(MP_Normal); }
 	if (IrisNormal.IsConnected())		{ StrataMaterialInfo.AddPropertyConnected(MP_Tangent); }
+	if (IrisPlaneNormal.IsConnected())	{ StrataMaterialInfo.AddPropertyConnected(MP_Tangent); }
 	if (IrisMask.IsConnected())			{ StrataMaterialInfo.AddPropertyConnected(MP_CustomData0); }
 	if (IrisDistance.IsConnected())		{ StrataMaterialInfo.AddPropertyConnected(MP_CustomData1); }
 	if (EmissiveColor.IsConnected())	{ StrataMaterialInfo.AddPropertyConnected(MP_EmissiveColor); }
diff --git a/Engine/Source/Runtime/Engine/Public/MaterialCompiler.h b/Engine/Source/Runtime/Engine/Public/MaterialCompiler.h
--- a/Engine/Source/Runtime/Engine/Public/MaterialCompiler.h
+++ b/Engine/Source/Runtime/Engine/Public/MaterialCompiler.h
@@ -590,7 +590,7 @@ public:
 	virtual int32 StrataVolumetricFogCloudBSDF(int32 Albedo, int32 Extinction, int32 EmissiveColor, int32 AmbientOcclusion) = 0;
 	virtual int32 StrataUnlitBSDF(int32 EmissiveColor, int32 TransmittanceColor) = 0;
 	virtual int32 StrataHairBSDF(int32 BaseColor, int32 Scatter, int32 Specular, int32 Roughness, int32 Backlit, int32 EmissiveColor, int32 Tangent, const FString& SharedLocalBasisIndexMacro) = 0;
-	virtual int32 StrataEyeBSDF(int32 DiffuseAlbedo, int32 Roughness, int32 IrisMask, int32 IrisDistance, int32 EmissiveColor, int32 CorneaNormal, int32 IrisNormal, int32 SSSProfileId, const FString& SharedLocalBasisIndexMacro) = 0;
+	virtual int32 StrataEyeBSDF(int32 DiffuseAlbedo, int32 Roughness, int32 IrisMask, int32 IrisDistance, int32 EmissiveColor, int32 CorneaNormal, int32 IrisNormal, int32 IrisPlaneNormal, int32 SSSProfileId, const FString& SharedLocalBasisIndexMacro) = 0;
 	virtual int32 StrataSingleLayerWaterBSDF(
 		int32 BaseColor, int32 Metallic, int32 Specular, int32 Roughness, int32 EmissiveColor, int32 TopMaterialOpacity, 
 		int32 WaterAlbedo, int32 WaterExtinction, int32 WaterPhaseG, int32 ColorScaleBehindWater, int32 Normal, const FString& SharedLocalBasisIndexMacro, FStrataOperator* PromoteToOperator) = 0;
@@ -1185,9 +1185,9 @@ public:
 		return Compiler->StrataHairBSDF(BaseColor, Scatter, Specular, Roughness, Backlit, EmissiveColor, Tangent, SharedLocalBasisIndexMacro);
 	}
 
-	virtual int32 StrataEyeBSDF(int32 DiffuseAlbedo, int32 Roughness, int32 IrisMask, int32 IrisDistance, int32 EmissiveColor, int32 CorneaNormal, int32 IrisNormal, int32 SSSProfileId, const FString& SharedLocalBasisIndexMacro) override
+	virtual int32 StrataEyeBSDF(int32 DiffuseAlbedo, int32 Roughness, int32 IrisMask, int32 IrisDistance, int32 EmissiveColor, int32 CorneaNormal, int32 IrisNormal, int32 IrisPlaneNormal, int32 SSSProfileId, const FString& SharedLocalBasisIndexMacro) override
 	{
-		return Compiler->StrataEyeBSDF(DiffuseAlbedo, Roughness, IrisMask, IrisDistance, EmissiveColor, CorneaNormal, IrisNormal, SSSProfileId, SharedLocalBasisIndexMacro);
+		return Compiler->StrataEyeBSDF(DiffuseAlbedo, Roughness, IrisMask, IrisDistance, EmissiveColor, CorneaNormal, IrisNormal, IrisPlaneNormal, SSSProfileId, SharedLocalBasisIndexMacro);
 	}
 
 	virtual int32 StrataSingleLayerWaterBSDF(
