Fix SSS albedo on eye.

diff --git a/Engine/Shaders/Private/Strata/Strata.ush b/Engine/Shaders/Private/Strata/Strata.ush
--- a/Engine/Shaders/Private/Strata/Strata.ush
+++ b/Engine/Shaders/Private/Strata/Strata.ush
@@ -442,7 +442,8 @@ struct FStrataBSDF
 ///////////////////////////////////////////////////////////////////////////////
 // Eye BSDF members
 #define EYE_DIFFUSEALBEDO(X)			X.VGPRs[0].xyz
-#define EYE_F0(X)						0.04f
+#define EYE_F0(X)						X.VGPRs[2].w
+#define EYE_F90(X)						X.VGPRs[3].w
 #define EYE_ROUGHNESS(X)				X.VGPRs[0].w
 #define EYE_IRISDISTANCE(X)				X.VGPRs[1].x
 #define EYE_IRISMASK(X)					X.VGPRs[1].y
@@ -2586,7 +2587,7 @@ float3 StrataGetBSDFSpecularF90(in FStrataBSDF BSDF)
 	{
 	case STRATA_BSDF_TYPE_SLAB: 			return SLAB_F90(BSDF);
 	case STRATA_BSDF_TYPE_HAIR: 			return 1.0f;
-	case STRATA_BSDF_TYPE_EYE: 				return 1.0f;
+	case STRATA_BSDF_TYPE_EYE: 				return EYE_F90(BSDF);
 	case STRATA_BSDF_TYPE_SINGLELAYERWATER: return 1.0f;
 	}
 	return 0.f;
@@ -3539,6 +3540,19 @@ FStrataBSDF UnpackStrataBSDFIn(FStrataMaterialContainer StrataBuffer, inout FStr
 		EYE_IRISDISTANCE(OutBSDF) = Data0.y;
 		EYE_IRISNORMAL(OutBSDF) = OctahedronToUnitVector(Data1.xy * 2.0f - 1.f);
 		EYE_IRISPLANENORMAL(OutBSDF) = OctahedronToUnitVector(Data1.zw * 2.0f - 1.f);
+		EYE_F0(OutBSDF)  = 0.04f;
+		EYE_F90(OutBSDF) = 1.0f;
+
+		if (BSDF_GETSSSTYPE(OutBSDF) == SSS_TYPE_DIFFUSION_PROFILE)
+		{
+			const bool bChecker = CheckerFromPixelPos(StrataAddressing.PixelCoords);
+			float3 DiffuseAlbedo = EYE_DIFFUSEALBEDO(OutBSDF);
+			float SpecularFactor = 1.0f;
+			AdjustBaseColorAndSpecularColorForSubsurfaceProfileLighting(DiffuseAlbedo, SpecularFactor, bChecker);
+			EYE_DIFFUSEALBEDO(OutBSDF) = DiffuseAlbedo;
+			EYE_F0(OutBSDF)  *= SpecularFactor;
+			EYE_F90(OutBSDF) *= SpecularFactor;
+		}
 		// 12 bytes
 	}
 	break;
diff --git a/Engine/Shaders/Private/Strata/StrataEvaluation.ush b/Engine/Shaders/Private/Strata/StrataEvaluation.ush
--- a/Engine/Shaders/Private/Strata/StrataEvaluation.ush
+++ b/Engine/Shaders/Private/Strata/StrataEvaluation.ush
@@ -750,7 +750,7 @@ FStrataEvaluateResult StrataEvaluateBSDFCommon(FStrataBSDFContext BSDFContext, F
 		{
 			float3 DiffuseColor	= EYE_DIFFUSEALBEDO(BSDFContext.BSDF);
 			float3 F0			= EYE_F0(BSDFContext.BSDF);
-			float3 F90			= 1.0f;
+			float3 F90			= EYE_F90(BSDFContext.BSDF);
 			const float SafeRoughness	= MakeRoughnessSafe(EYE_ROUGHNESS(BSDFContext.BSDF));
 
 			if (Settings.bForceFullyRough)
@@ -1111,7 +1111,7 @@ FStrataImportanceSampleResult StrataImportanceSampleBSDF(FStrataBSDFContext BSDF
 			// TODO manage iris normal for the diffuse term
 			const float3 DiffuseColor	= EYE_DIFFUSEALBEDO(BSDFContext.BSDF);
 			const float3 F0 			= EYE_F0(BSDFContext.BSDF);
-			const float3 F90			= 1.0f;
+			const float3 F90			= EYE_F90(BSDFContext.BSDF);
 			const float SafeRoughness   = MakeRoughnessSafe(EYE_ROUGHNESS(BSDFContext.BSDF));
 
 			// Pick diffuse/specular lobe
@@ -1444,7 +1444,7 @@ FStrataEnvLightResult StrataEvaluateForEnvLight(FStrataBSDFContext BSDFContext,
 		{
 			float3 DiffuseColor		= EYE_DIFFUSEALBEDO(BSDFContext.BSDF);
 			float3 F0 				= EYE_F0(BSDFContext.BSDF);
-			float3 F90 				= 1.0f;
+			float3 F90 				= EYE_F90(BSDFContext.BSDF);
 
 			if (Settings.bForceFullyRough)
 			{
