Fixed depth to be invariant for lighting debug probes.

diff --git a/Engine/Shaders/Private/DebugProbes.usf b/Engine/Shaders/Private/DebugProbes.usf
--- a/Engine/Shaders/Private/DebugProbes.usf
+++ b/Engine/Shaders/Private/DebugProbes.usf
@@ -232,7 +232,7 @@ void MainPS(
 	float4 ClipPos = mul(float4(WorlPosition, 1.0f), LWCHackToFloat(ResolvedView.WorldToClip));
 	ClipPos /= ClipPos.w;
 	const float DeviceZ = ClipPos.z;
-	OutDepth = DeviceZ;
+	OutDepth = INVARIANT(DeviceZ);
 
 	if (RayT <= 0.0f)
 	{
