Fix a bug that the brightness is not consistent with MSAA enabled on IOS.

diff --git a/Engine/Shaders/Private/PostProcessTonemap.usf b/Engine/Shaders/Private/PostProcessTonemap.usf
--- a/Engine/Shaders/Private/PostProcessTonemap.usf
+++ b/Engine/Shaders/Private/PostProcessTonemap.usf
@@ -423,6 +423,13 @@ float4 TonemapCommonPS(
 		half4 SceneColor = SampleSceneColor(SceneUV);
 	#endif
 
+#if METAL_MSAA_HDR_DECODE && !USE_GAMMA_ONLY
+	// look for PreTonemapMSAA 
+	SceneColor.rgb *= rcp(SceneColor.r * (-0.299) + SceneColor.g * (-0.587) + SceneColor.b * (-0.114) + 1.0);
+	// Try to kill negatives and NaNs here
+	SceneColor.rgb = max(SceneColor.rgb, 0);
+#endif
+
 	SceneColor.rgb *= View.OneOverPreExposure;
 
 	float ExposureScale = ExposureScaleVignette.x;
@@ -528,13 +535,6 @@ float4 TonemapCommonPS(
 
 	#else 
 	{
-		#if METAL_MSAA_HDR_DECODE
-			// look for PreTonemapMSAA 
-			SceneColor.rgb *= rcp(SceneColor.r*(-0.299) + SceneColor.g*(-0.587) + SceneColor.b*(-0.114) + 1.0);
-			// Try to kill negatives and NaNs here
-			SceneColor.rgb = max(SceneColor.rgb, 0); 
-		#endif
-	
 		half3 LinearColor = SceneColor.rgb * ColorScale0.rgb;
 
 		#if SUPPORTS_SCENE_COLOR_APPLY_PARAMETERS
