Fix Strata eye shading when using dynamic shading model in legacy material.

diff --git a/Engine/Shaders/Private/Strata/StrataLegacyConversion.ush b/Engine/Shaders/Private/Strata/StrataLegacyConversion.ush
--- a/Engine/Shaders/Private/Strata/StrataLegacyConversion.ush
+++ b/Engine/Shaders/Private/Strata/StrataLegacyConversion.ush
@@ -197,6 +197,37 @@ FStrataData StrataConvertLegacyMaterialDynamic(
 		// No need to take into account opacity weighting for translucent material, as the coverage is assumed to be 1
 	}
 	#endif
+	#if MATERIAL_SHADINGMODEL_EYE
+	if (ShadingModel == SHADINGMODELID_EYE)
+	{
+		const float IrisMask = ClearCoat;
+		const float IrisDistance = ClearCoatRoughness;
+
+		float3 IrisNormal = RawNormal;
+		float3 IrisPlaneNormal = RawNormal;
+		GetEyeNormals(IrisMask, IrisDistance, RawNormal, RawClearCoatNormal, RawCustomTangent, IrisNormal, IrisPlaneNormal);
+
+		FStrataData EyeMat = PromoteParameterBlendedBSDFToOperator(
+			GetStrataEyeBSDF(
+				BaseColor,	         // DiffuseAlbedo
+				Roughness,	         // Roughness
+				IrisMask,	         // IrisMask
+				IrisDistance,        // IrisDistance
+				IrisNormal,	         // IrisNormal
+				IrisPlaneNormal,	 // IrisPlaneNormal
+				SubSurfaceProfileId, // SSS profile
+				Emissive,	         // EmissiveColor
+				SharedLocalBasisIndex),
+			StrataTree,
+			0		/*OperatorIndex*/,
+			0		/*BSDFIndex*/,
+			0		/*int LayerDepth*/,
+			1		/*int bIsBottom*/);
+
+		return StrataWeight(EyeMat, 1.0, StrataTree, /*OperatorIndex*/1, /*MaxDistanceFromLeaves*/1);
+		// Out maps to LegacyUpdateBSDFsOperators
+	}
+	#endif
 
 	FStrataLegacyParameters LegacySlab = InitStrataLegacyParameters(PixelFootprint, SharedLocalBasisIndex, 1.0f);
 
@@ -415,21 +446,6 @@ FStrataData StrataConvertLegacyMaterialDynamic(
 		// Thin translucency model always have a coverage of 1. It also canont be used in a multiple shading model context.
 		FinalWeight = 1.0f;
 	}
-	else if (ShadingModel == SHADINGMODELID_EYE) 
-	{
-		// Fallback to bright purple in case eye shading is selected when using the material attribute workflow.
-		LegacySlab.BaseColor = float3(0, 1, 1);
-		LegacySlab.Specular = 0;
-		LegacySlab.Metallic = 0;
-		LegacySlab.Roughness = 0;
-		LegacySlab.Anisotropy = 0;
-		LegacySlab.Emissive = 0;
-
-		FinalWeight = 1.0f;
-#if STRATA_TRANSLUCENT_MATERIAL
-		FinalWeight = 1.0f;
-#endif
-	}
 
 	// Fixed layer structure for helping compiler to unroll and optimize shader
 	// This must match what is done in PackStrataOut
