Visualize HDR now has an illuminance meter. It measures illuminance from behind the camera for now (might change later for surface normal)

diff --git a/Engine/Shaders/Private/DebugProbes.usf b/Engine/Shaders/Private/DebugProbes.usf
--- a/Engine/Shaders/Private/DebugProbes.usf
+++ b/Engine/Shaders/Private/DebugProbes.usf
@@ -65,6 +65,127 @@ void DebugSphereIntersection(
 	}
 }
 
+#if PERMUTATION_PASS == RENDER_BASEPASS
+void WriteOutMaterial(
+	in float PerObjectGBufferData,
+	in float SelectiveOutputMask,
+	in float3 BaseColor,
+	in float  Metallic,
+	in float  Specular,
+	in float  Roughness,
+	in float  IndirectIrradianceAndAO,
+	in float3 Normal
+#if STRATA_ENABLED
+	, inout float4 OutGBufferA
+	, inout uint OutGBufferB
+	, inout uint OutGBufferC
+	, inout uint OutGBufferD
+	, inout uint OutGBufferE
+#else
+	, inout float4 OutGBufferA
+	, inout float4 OutGBufferB
+	, inout float4 OutGBufferC
+	, inout float4 OutGBufferD
+	, inout float4 OutGBufferE
+#endif
+
+)
+{
+#if STRATA_ENABLED
+
+	FStrataPixelHeader StrataPixelHeader = InitialiseStrataPixelHeader();
+	StrataPixelHeader.StrataTree = GetInitialisedStrataTree();
+	StrataPixelHeader.IrradianceAO = InitIrradianceAndOcclusion();
+	StrataPixelHeader.IrradianceAO.MaterialAO = 1.0f - IndirectIrradianceAndAO;
+	SetCastContactShadow(StrataPixelHeader, 0);
+
+	uint SharedLocalBasisIndex = 0;
+	const float3 NormalizedNormal = normalize(Normal);
+	StrataPixelHeader.SharedLocalBases.Normals[SharedLocalBasisIndex] = NormalizedNormal;
+	StrataPixelHeader.SharedLocalBases.Count = 1;
+
+	FStrataPixelFootprint PixelFootprint = (FStrataPixelFootprint)0;
+	FStrataData StrataSlabData = GetStrataSlabBSDF(
+		/**FStrataPixelFootprint PixelFootprint */PixelFootprint,
+		/*float UseMetalness					*/true,
+		/*float3 BaseColor						*/BaseColor,
+		/*float3 EdgeColor						*/1.0f,
+		/*float Specular						*/Specular,
+		/*float Metallic						*/Metallic,
+		/*float3 DiffuseAlbedo					*/0.0f,
+		/*float3 F0								*/0.0f,
+		/*float3 F90							*/0.0f,
+		/*float Roughness						*/Roughness,
+		/*float Anisotropy						*/0.0f,
+		/*float SSSProfileID					*/0.0f,
+		/*float3 SSSMFP							*/0.0f,
+		/*float SSSMFPScale						*/0.0f,
+		/*float SSSPhaseAniso					*/0.0f,
+		/*float UseSSSDiffuse					*/0.0f,
+		/*float3 Emissive						*/0.0f,
+		/*float SecondRoughness					*/0.0f,
+		/*float SecondRoughnessWeight			*/0.0f,
+		/*flaot SecondRoughnessAsSimpleClearCoat*/0.0f,
+		/*float FuzzAmount						*/0.0f,
+		/*float3 FuzzColor						*/0.0f,
+		/*float Thickness						*/0.0f,
+		/*uint SharedLocalBasisIndex			*/SharedLocalBasisIndex,
+		/*inout uint SharedLocalBasisTypes		*/StrataPixelHeader.SharedLocalBases.Types);
+
+	FStrataData StrataData = PromoteParameterBlendedBSDFToOperator(StrataSlabData, StrataPixelHeader.StrataTree, /*int OperatorIndex*/0, /*int BSDFIndex*/0, /*int LayerDepth*/0, /*int bIsBottom*/1);
+	StrataPixelHeader.StrataTree.BSDFs[0].LuminanceWeightV = 1.0f;
+	StrataPixelHeader.StrataTree.BSDFs[0].TransmittanceAboveAlongN = 1.0f;
+	StrataPixelHeader.StrataTree.BSDFs[0].TopLayerDataWeight = 1.0f;
+
+	FStrataSubsurfaceData SSSData = (FStrataSubsurfaceData)0;
+	FStrataTopLayerData TopLayerData = (FStrataTopLayerData)0;
+	FStrataOpaqueRoughRefractionData OpaqueRoughRefractionData = (FStrataOpaqueRoughRefractionData)0;
+
+	float3 EmissiveLuminance = 0.0f;
+	bool bStrataSubsurfaceEnable = false;
+	const float Dither = 0.0f;
+	const FStrataIntegrationSettings Settings = InitStrataIntegrationSettings(false /*bForceFullyRough*/, bRoughDiffuse, 0/*PeelLayersAboveDepth*/);
+
+	FStrataAddressing StrataAddressing = GetStrataPixelDataByteOffset(PixelCoord, uint2(ResolvedView.BufferSizeAndInvSize.xy), MaxBytesPerPixel);
+	FRWStrataMaterialContainer RWStrataMaterialContainer = InitialiseRWStrataMaterialContainer(MaterialTextureArrayUAV);
+	PackStrataOut(
+		RWStrataMaterialContainer,
+		MaterialTextureArrayUAV,
+		Dither,
+		Settings,
+		StrataAddressing,
+		StrataPixelHeader, StrataData, -RayDir, NormalizedNormal, bStrataSubsurfaceEnable, EmissiveLuminance,
+		SSSData, TopLayerData, OpaqueRoughRefractionData);
+
+	// Write out MRTs. Note: this only works for deferred with basepassvelocity off
+	OutGBufferA = 0;
+	OutGBufferB = RWStrataMaterialContainer.MaterialRenderTargets[0];
+	OutGBufferC = RWStrataMaterialContainer.MaterialRenderTargets[1];
+	OutGBufferD = StrataPackTopLayerData(TopLayerData);
+	OutGBufferE = 0;
+
+#else // STRATA_ENABLED
+
+	// We write all outputs, the rasterizer will discard what is not bound. That is fine since we do not care that much about performance for this debug view.
+
+	OutGBufferA = float4(
+		EncodeNormal(Normal),
+		PerObjectGBufferData);
+
+	OutGBufferB = float4(
+		Metallic,
+		Specular,
+		Roughness,
+		EncodeShadingModelIdAndSelectiveOutputMask(SHADINGMODELID_DEFAULT_LIT, SelectiveOutputMask));
+
+	OutGBufferC = float4(BaseColor, IndirectIrradianceAndAO);
+	OutGBufferD = float4(0.0f, 0.0f, 0.0f, 0.0f);	// custom data
+	OutGBufferE = float4(1.0f, 1.0f, 1.0f, 1.0f);	// Pre computed shadow factor
+
+#endif // STRATA_ENABLED
+}
+#endif // PERMUTATION_PASS == RENDER_BASEPASS
+
 
 
 void MainPS(
@@ -107,11 +228,86 @@ void MainPS(
 	const float3 RayDir = normalize(TranslatedWorldPos.xyz - View.TranslatedWorldCameraOrigin);
 	const float3 RayOrig = LWCHackToFloat(ResolvedView.WorldCameraOrigin);
 
+#if PERMUTATION_VISUALIZEHDR /////////////////////////////////////////////////////////////////////////////////////////////////////////////////////////////////////
+
+	// In this case we want to show a small patch in front of the camera with pure white albedo.
+	// That is a way to measure the illuminance at the camera position, on the hemisphere pointing otwards the back of the camera.
+
+	const float DeviceZ = 0.999999f;
+
+	// IlluminanceMeter variables must match the one in PostProcessVisualizeHDR.usf
+	const float  IlluminanceMeterSize = 20;
+	const float  IlluminanceMeterHalfSize = IlluminanceMeterSize / 2;
+	const float2 IlluminanceMeterCenter = View.ViewSizeAndInvSize.xy * 0.5f - uint2(0, 100);
+
+	if(all(PixelCoord > (IlluminanceMeterCenter - IlluminanceMeterHalfSize)) && all(PixelCoord < (IlluminanceMeterCenter + IlluminanceMeterHalfSize)))
+	{
+	#if PERMUTATION_PASS == RENDER_VELOCITYPASS
+
+		float3 WorldPosition = RayOrig;
+		float3 Velocity = 0.0f;
+		float4 ScreenPos = mul(float4(WorldPosition.xyz + LWCHackToFloat(ResolvedView.PreViewTranslation), 1), ResolvedView.TranslatedWorldToClip);
+		float4 PrevScreenPos = mul(float4(WorldPosition.xyz + LWCHackToFloat(ResolvedView.PrevPreViewTranslation), 1), ResolvedView.PrevTranslatedWorldToClip);
+		Velocity = Calculate3DVelocity(ScreenPos, PrevScreenPos);
+		OutColor = EncodeVelocityToTexture(Velocity);
+
+	#elif PERMUTATION_PASS == RENDER_BASEPASS
+
+		const float PerObjectGBufferData = 0.0f;
+		const float SelectiveOutputMask = 0.0f;
+
+		const float3 WorldNormal = -View.ViewForward;
+		const float3 BaseColor = 1.0f;
+		const float  Metallic = 0.0f;
+		const float  Specular = 0.0f;
+		const float  Roughness = 1.0f;
+		const float  IndirectIrradianceAndAO = 0.0f;
+
+	#if STRATA_ENABLED
+		// EarlyDepthStencil does not work because we render a simple quad on screen with custom OutDepth.
+		// We cannot rely on late depth test because pixel shader would have already operated the write to strata buffer.
+		// So we must execute a manual depth test in shader (fine because basepass does not write depth)
+		float2 UvBuffer = SvPosition.xy * View.BufferSizeAndInvSize.zw;	// Uv for depth buffer read (size can be larger than viewport)
+		const float DepthBufferDeviceZ = LookupDeviceZ(UvBuffer);
+		if (DepthBufferDeviceZ < (DeviceZ + 0.000001))
+	#endif // STRATA_ENABLED
+		{
+			WriteOutMaterial(
+				PerObjectGBufferData,
+				SelectiveOutputMask,
+				BaseColor,
+				Metallic,
+				Specular,
+				Roughness,
+				IndirectIrradianceAndAO,
+				WorldNormal,
+				OutGBufferA,
+				OutGBufferB,
+				OutGBufferC,
+				OutGBufferD,
+				OutGBufferE);
+		}
+	#endif
+
+		OutColor = 0.0f;
+	}
+	else
+	{
+		discard;
+		return;
+	}
+
+
+
+#else // PERMUTATION_VISUALIZEHDR /////////////////////////////////////////////////////////////////////////////////////////////////////////////////////////////////////
+
+
+
 	int ProbeMode = false;
 	bool bProbeAttachedToCamera = false;
 	float RayT = -1.0f;
-	float3 WorlPosition = 0.0f;
-	float3 Normal = 0.0f;
+	float3 WorldPosition = 0.0f;
+	float3 WorldNormal = 0.0f;
 
 
 	// Render the camera probes (this could be done proceppdurally as flat sphere normal on screen rotated by inverse view matrix)
@@ -123,7 +319,7 @@ void MainPS(
 
 			DebugSphereIntersection(
 				RayOrig, RayDir, SphereCenter, SphereRadius, PROBE_MODE_SPECULAR, true,
-				ProbeMode, RayT, WorlPosition, Normal, bProbeAttachedToCamera);
+				ProbeMode, RayT, WorldPosition, WorldNormal, bProbeAttachedToCamera);
 		}
 		{
 			float3 SphereCenter = LWCHackToFloat(ResolvedView.WorldCameraOrigin) + (ResolvedView.ViewForward * 1.0f + ResolvedView.ViewRight * -0.5f + ResolvedView.ViewUp * -0.2f) * METER_TO_CENTIMETER;
@@ -131,7 +327,7 @@ void MainPS(
 
 			DebugSphereIntersection(
 				RayOrig, RayDir, SphereCenter, SphereRadius, PROBE_MODE_DIFFUSE, true,
-				ProbeMode, RayT, WorlPosition, Normal, bProbeAttachedToCamera);
+				ProbeMode, RayT, WorldPosition, WorldNormal, bProbeAttachedToCamera);
 		}
 	}
 
@@ -172,7 +368,7 @@ void MainPS(
 
 			DebugSphereIntersection(
 				RayOrig, RayDir, SphereCenter, SphereRadius, PROBE_MODE_BOTH, false,
-				ProbeMode, RayT, WorlPosition, Normal, bProbeAttachedToCamera);
+				ProbeMode, RayT, WorldPosition, WorldNormal, bProbeAttachedToCamera);
 
 			// Amanatides 3D-DDA 
 			if (tMax.x < tMax.y)
@@ -221,7 +417,7 @@ void MainPS(
 
 					DebugSphereIntersection(
 						RayOrig, RayDir, SphereCenter, SphereRadius, PROBE_MODE_BOTH, false,
-						ProbeMode, RayT, WorlPosition, Normal, bProbeAttachedToCamera);
+						ProbeMode, RayT, WorldPosition, WorldNormal, bProbeAttachedToCamera);
 				}
 			}
 		}
@@ -229,7 +425,7 @@ void MainPS(
 	}
 
 
-	float4 ClipPos = mul(float4(WorlPosition, 1.0f), LWCHackToFloat(ResolvedView.WorldToClip));
+	float4 ClipPos = mul(float4(WorldPosition, 1.0f), LWCHackToFloat(ResolvedView.WorldToClip));
 	ClipPos /= ClipPos.w;
 	const float DeviceZ = ClipPos.z;
 	OutDepth = INVARIANT(DeviceZ);
@@ -247,8 +443,8 @@ void MainPS(
 
 	if (!bProbeAttachedToCamera)
 	{
-		float4 ScreenPos	= mul(float4(WorlPosition.xyz + LWCHackToFloat(ResolvedView.PreViewTranslation), 1),	ResolvedView.TranslatedWorldToClip);
-		float4 PrevScreenPos= mul(float4(WorlPosition.xyz + LWCHackToFloat(ResolvedView.PrevPreViewTranslation), 1),ResolvedView.PrevTranslatedWorldToClip);
+		float4 ScreenPos	= mul(float4(WorldPosition.xyz + LWCHackToFloat(ResolvedView.PreViewTranslation), 1),	ResolvedView.TranslatedWorldToClip);
+		float4 PrevScreenPos= mul(float4(WorldPosition.xyz + LWCHackToFloat(ResolvedView.PrevPreViewTranslation), 1),ResolvedView.PrevTranslatedWorldToClip);
 		Velocity = Calculate3DVelocity(ScreenPos, PrevScreenPos);
 	}
 
@@ -266,117 +462,35 @@ void MainPS(
 	const float  IndirectIrradianceAndAO = 0.0f;
 
 #if STRATA_ENABLED
-
 	// EarlyDepthStencil does not work because we render a simple quad on screen with custom OutDepth.
 	// We cannot rely on late depth test because pixel shader would have already operated the write to strata buffer.
 	// So we must execute a manual depth test in shader (fine because basepass does not write depth)
 	float2 UvBuffer = SvPosition.xy * View.BufferSizeAndInvSize.zw;	// Uv for depth buffer read (size can be larger than viewport)
 	const float DepthBufferDeviceZ = LookupDeviceZ(UvBuffer);
-	if (DepthBufferDeviceZ < (DeviceZ+0.000001))
-	{
-		FStrataPixelHeader StrataPixelHeader = InitialiseStrataPixelHeader();
-		StrataPixelHeader.StrataTree = GetInitialisedStrataTree();
-		StrataPixelHeader.IrradianceAO = InitIrradianceAndOcclusion();
-		StrataPixelHeader.IrradianceAO.MaterialAO = 1.0f - IndirectIrradianceAndAO;
-		SetCastContactShadow(StrataPixelHeader, 0);
-
-		uint SharedLocalBasisIndex = 0;
-		const float3 NormalizedNormal = normalize(Normal);
-		StrataPixelHeader.SharedLocalBases.Normals[SharedLocalBasisIndex] = NormalizedNormal;
-		StrataPixelHeader.SharedLocalBases.Count = 1;
-
-		FStrataPixelFootprint PixelFootprint = (FStrataPixelFootprint)0;
-		FStrataData StrataSlabData = GetStrataSlabBSDF(
-			/**FStrataPixelFootprint PixelFootprint */PixelFootprint,
-			/*float UseMetalness					*/true,
-			/*float3 BaseColor						*/BaseColor,
-			/*float3 EdgeColor						*/1.0f,
-			/*float Specular						*/Specular,
-			/*float Metallic						*/Metallic,
-			/*float3 DiffuseAlbedo					*/0.0f,
-			/*float3 F0								*/0.0f,
-			/*float3 F90							*/0.0f,
-			/*float Roughness						*/Roughness,
-			/*float Anisotropy						*/0.0f,
-			/*float SSSProfileID					*/0.0f,
-			/*float3 SSSMFP							*/0.0f,
-			/*float SSSMFPScale						*/0.0f,
-			/*float SSSPhaseAniso					*/0.0f,
-			/*float UseSSSDiffuse					*/0.0f,
-			/*float3 Emissive						*/0.0f,
-			/*float SecondRoughness					*/0.0f,
-			/*float SecondRoughnessWeight			*/0.0f,
-			/*flaot SecondRoughnessAsSimpleClearCoat*/0.0f,
-			/*float FuzzAmount						*/0.0f,
-			/*float3 FuzzColor						*/0.0f,
-			/*float Thickness						*/0.0f,
-			/*uint SharedLocalBasisIndex			*/SharedLocalBasisIndex,
-			/*inout uint SharedLocalBasisTypes		*/StrataPixelHeader.SharedLocalBases.Types);
-
-		FStrataData StrataData = PromoteParameterBlendedBSDFToOperator(StrataSlabData, StrataPixelHeader.StrataTree, /*int OperatorIndex*/0, /*int BSDFIndex*/0, /*int LayerDepth*/0, /*int bIsBottom*/1);
-		StrataPixelHeader.StrataTree.BSDFs[0].LuminanceWeightV = 1.0f;
-		StrataPixelHeader.StrataTree.BSDFs[0].TransmittanceAboveAlongN = 1.0f;
-		StrataPixelHeader.StrataTree.BSDFs[0].TopLayerDataWeight = 1.0f;
-
-		FStrataSubsurfaceData SSSData = (FStrataSubsurfaceData)0;
-		FStrataTopLayerData TopLayerData = (FStrataTopLayerData)0;
-		FStrataOpaqueRoughRefractionData OpaqueRoughRefractionData = (FStrataOpaqueRoughRefractionData)0;
-
-		float3 EmissiveLuminance = 0.0f;
-		bool bStrataSubsurfaceEnable = false;
-		const float Dither = 0.0f;
-		const FStrataIntegrationSettings Settings = InitStrataIntegrationSettings(false /*bForceFullyRough*/, bRoughDiffuse, 0/*PeelLayersAboveDepth*/);
-
-		FStrataAddressing StrataAddressing = GetStrataPixelDataByteOffset(PixelCoord, uint2(ResolvedView.BufferSizeAndInvSize.xy), MaxBytesPerPixel);
-		FRWStrataMaterialContainer RWStrataMaterialContainer = InitialiseRWStrataMaterialContainer(MaterialTextureArrayUAV);
-		PackStrataOut(
-			RWStrataMaterialContainer, 
-			MaterialTextureArrayUAV,
-			Dither,
-			Settings,
-			StrataAddressing,
-			StrataPixelHeader, StrataData, -RayDir, NormalizedNormal, bStrataSubsurfaceEnable, EmissiveLuminance,
-			SSSData, TopLayerData, OpaqueRoughRefractionData);
-
-		// Write out MRTs. Note: this only works for deferred with basepassvelocity off
-		OutGBufferA = 0;
-		OutGBufferB = RWStrataMaterialContainer.MaterialRenderTargets[0];
-		OutGBufferC = RWStrataMaterialContainer.MaterialRenderTargets[1];
-		OutGBufferD = StrataPackTopLayerData(TopLayerData);
-		OutGBufferE = 0;
-	}
-	else
+	if (DepthBufferDeviceZ < (DeviceZ + 0.000001))
+#endif // STRATA_ENABLED
 	{
-		// If we are behind the front depth buffer, we must also discard to avoid writting to output MRTs.
-		discard;
-		return;
+		WriteOutMaterial(
+			PerObjectGBufferData,
+			SelectiveOutputMask,
+			BaseColor,
+			Metallic,
+			Specular,
+			Roughness,
+			IndirectIrradianceAndAO,
+			WorldNormal,
+			OutGBufferA,
+			OutGBufferB,
+			OutGBufferC,
+			OutGBufferD,
+			OutGBufferE);
 	}
-	
-#else // STRATA_ENABLED
-
-	// We write all outputs, the rasterizer will discard what is not bound. That is fine since we do not care that much about performance for this debug view.
-
-	OutGBufferA = float4(
-		EncodeNormal(Normal), 
-		PerObjectGBufferData);
-
-	OutGBufferB = float4(
-		Metallic, 
-		Specular, 
-		Roughness, 
-		EncodeShadingModelIdAndSelectiveOutputMask(SHADINGMODELID_DEFAULT_LIT, SelectiveOutputMask));
-
-	OutGBufferC = float4(BaseColor, IndirectIrradianceAndAO);
-	OutGBufferD = float4(0.0f, 0.0f, 0.0f, 0.0f);	// custom data
-	OutGBufferE = float4(1.0f, 1.0f, 1.0f, 1.0f);	// Pre computed shadow factor
-
-#endif // STRATA_ENABLED
 
 	OutColor = 0.0f;
 
 #endif // PERMUTATION_PASS
 
-
+#endif // PERMUTATION_VISUALIZEHDR /////////////////////////////////////////////////////////////////////////////////////////////////////////////////////////////////////
 
 #if PERMUTATION_PASS == RENDER_VELOCITYPASS || PERMUTATION_PASS == RENDER_DEPTHPREPASS || (PERMUTATION_PASS == RENDER_BASEPASS && STRATA_ENABLED == 0)
 	// Base pass with strata cannot write to depth because it is going to execute in shader depth testing
diff --git a/Engine/Shaders/Private/PostProcessVisualizeHDR.usf b/Engine/Shaders/Private/PostProcessVisualizeHDR.usf
--- a/Engine/Shaders/Private/PostProcessVisualizeHDR.usf
+++ b/Engine/Shaders/Private/PostProcessVisualizeHDR.usf
@@ -291,6 +291,68 @@ void MainPS(noperspective float4 UVAndScreenPos : TEXCOORD0, float4 SvPosition :
 		}
 	}
 
+	// Illuminance meter for the hemisphere in front of the camera
+	{
+		OutColor = float4(0.0f, 0.0f, 0.0f, 0.0f);
+
+		// IlluminanceMeter variables must match the one in DebugProbes.usf
+		const float  IlluminanceMeterSize = 20;
+		const float  IlluminanceMeterHalfSize = IlluminanceMeterSize / 2;
+		const float2 IlluminanceMeterCenter = View.ViewSizeAndInvSize.xy * 0.5f - uint2(0, 100);
+
+		const float2 IlluminanceMeterCenterUV = (0*Input_UVViewportMin + IlluminanceMeterCenter) * Input_ExtentInverse;
+
+		const float PixelDx = PixelPos.x - IlluminanceMeterCenter.x;
+		const float PixelDy = PixelPos.y - IlluminanceMeterCenter.y;
+		const float AbsPixelDx = abs(PixelDx);
+		const float AbsPixelDy = abs(PixelDy);
+		if (PixelDx > - 20 && PixelDx < 110 && PixelDy < 20 && PixelDy > -20)
+		{
+			float3 LuminanceAvg = 0.0f;
+			float SampleCount = 0.0f;
+			LOOP
+			for (float OffsetX = (-IlluminanceMeterHalfSize + 2); OffsetX < (IlluminanceMeterHalfSize-2); OffsetX++)
+			{
+				for (float OffsetY = (-IlluminanceMeterHalfSize + 2); OffsetY < (IlluminanceMeterHalfSize-2); OffsetY++)
+				{
+					float2 UV = IlluminanceMeterCenterUV + float2(OffsetX, OffsetY) * Input_ExtentInverse;
+					LuminanceAvg += Texture2DSampleLevel(HDRSceneColorTexture, HDRSceneColorSampler, UV, 0).rgb;
+					SampleCount++;
+				}
+			}
+			LuminanceAvg = LuminanceAvg / SampleCount;
+			LuminanceAvg *= View.OneOverPreExposure;
+
+			OutColor = Texture2DSample(SceneColorTexture, SceneColorSampler, UV); // Influenced area default scene color
+			if (AbsPixelDx < IlluminanceMeterHalfSize && AbsPixelDy < IlluminanceMeterHalfSize)
+			{
+				if (AbsPixelDx == IlluminanceMeterHalfSize - 1 || AbsPixelDy == IlluminanceMeterHalfSize - 1)
+				{
+					OutColor = float4(1.0, 1.0, 1.0, 1.0); // White border
+				}
+			}
+
+			float MaxLuminance = max(LuminanceAvg.r, max(LuminanceAvg.g, LuminanceAvg.b));
+			float MaxIlluminance = MaxLuminance * PI;
+
+			int2 TopLeft = IlluminanceMeterCenter + int2(11, 11);
+			// Illuminance
+			if (MaxIlluminance < 1000.0f)
+			{
+				PrintFloat(PixelPos, OutColor.xyz, float3(1, 1, 1), TopLeft, MaxIlluminance);
+			}
+			else
+			{
+				PrintFloatNoFractionLarge(PixelPos, OutColor.xyz, float3(1, 1, 1), TopLeft, MaxIlluminance);
+			}
+			TopLeft.x += 60;
+			PrintCharacter(PixelPos, OutColor.xyz, float3(1, 1, 1), TopLeft, _L_);
+			PrintCharacter(PixelPos, OutColor.xyz, float3(1, 1, 1), TopLeft, _U_);
+			PrintCharacter(PixelPos, OutColor.xyz, float3(1, 1, 1), TopLeft, _X_);
+			return;
+		}
+	}
+
 	if (EyeAdaptation_VisualizeDebugType == 0)
 	{
 		OutColor = Texture2DSample(SceneColorTexture, SceneColorSampler, UV);
diff --git a/Engine/Source/Runtime/Renderer/Private/DebugProbeRendering.cpp b/Engine/Source/Runtime/Renderer/Private/DebugProbeRendering.cpp
--- a/Engine/Source/Runtime/Renderer/Private/DebugProbeRendering.cpp
+++ b/Engine/Source/Runtime/Renderer/Private/DebugProbeRendering.cpp
@@ -8,6 +8,8 @@
 #include "SceneTextureParameters.h"
 #include "Strata/Strata.h"
 
+#define DEBUG_PROBE_ENABLED (!UE_BUILD_SHIPPING)
+
 //
 //	Deferred probes are only stamped in deferred mode.
 //
@@ -40,7 +42,8 @@ class FStampDeferredDebugProbePS : public FGlobalShader
 	SHADER_USE_PARAMETER_STRUCT(FStampDeferredDebugProbePS, FGlobalShader);
 
 	class FRenderPass : SHADER_PERMUTATION_RANGE_INT("PERMUTATION_PASS", 0, 3);
-	using FPermutationDomain = TShaderPermutationDomain<FRenderPass>;
+	class FVisualizeHDR : SHADER_PERMUTATION_BOOL("PERMUTATION_VISUALIZEHDR");
+	using FPermutationDomain = TShaderPermutationDomain<FRenderPass, FVisualizeHDR>;
 
 	BEGIN_SHADER_PARAMETER_STRUCT(FParameters, )
 		SHADER_PARAMETER_STRUCT_REF(FViewUniformShaderParameters, ViewUniformBuffer)
@@ -72,13 +75,27 @@ public:
 
 IMPLEMENT_GLOBAL_SHADER(FStampDeferredDebugProbePS, "/Engine/Private/DebugProbes.usf", "MainPS", SF_Pixel);
 
+#if DEBUG_PROBE_ENABLED
+
+static bool AnyViewRequireVisualizeHDR(TArrayView<const FViewInfo> Views)
+{
+	for (const FViewInfo& View : Views)
+	{
+		if (View.Family->EngineShowFlags.VisualizeHDR)
+		{
+			return true;
+		}
+	}
+	return false;
+}
  
 template<bool bEnableDepthWrite, ECompareFunction CompareFunction>
 static void CommonStampDeferredDebugProbeDrawCall(
 	FRDGBuilder& GraphBuilder,
 	const FViewInfo& View, 
 	FStampDeferredDebugProbePS::FParameters* PassParameters,
-	int32 RenderPass)
+	int32 RenderPass,
+	bool bVisualizeHDR)
 {
 	PassParameters->ViewUniformBuffer = View.ViewUniformBuffer;
 	PassParameters->MaterialTextureArrayUAV = View.StrataViewData.SceneData->MaterialTextureArrayUAVWithoutRTs;
@@ -88,6 +105,7 @@ static void CommonStampDeferredDebugProbeDrawCall(
 		
 	FStampDeferredDebugProbePS::FPermutationDomain PermutationVector;
 	PermutationVector.Set<FStampDeferredDebugProbePS::FRenderPass>(RenderPass);
+	PermutationVector.Set<FStampDeferredDebugProbePS::FVisualizeHDR>(bVisualizeHDR);
 	TShaderMapRef<FStampDeferredDebugProbePS> PixelShader(View.ShaderMap, PermutationVector);
 
 	FPixelShaderUtils::AddFullscreenPass<FStampDeferredDebugProbePS>(
@@ -98,18 +116,22 @@ static void CommonStampDeferredDebugProbeDrawCall(
 		TStaticDepthStencilState<bEnableDepthWrite, CompareFunction>::GetRHI());
 }
 
+#endif // DEBUG_PROBE_ENABLED
+
 void StampDeferredDebugProbeDepthPS(
 	FRDGBuilder& GraphBuilder,
 	TArrayView<const FViewInfo> Views,
 	const FRDGTextureRef SceneDepthTexture)
 {
+#if DEBUG_PROBE_ENABLED
 	RDG_EVENT_SCOPE(GraphBuilder, "StampDeferredDebugProbeDepth");
 	RDG_GPU_STAT_SCOPE(GraphBuilder, StampDeferredDebugProbe);
 
+	const bool bVisualizeHDR = AnyViewRequireVisualizeHDR(Views);
 	const bool bVisualizeLightingOnProbes = CVarVisualizeLightingOnProbes.GetValueOnRenderThread() > 0;
 	for (const FViewInfo& View : Views)
 	{
-		if (!(bVisualizeLightingOnProbes || View.Family->EngineShowFlags.VisualizeLightingOnProbes) || View.bIsReflectionCapture)
+		if (!(bVisualizeLightingOnProbes || View.Family->EngineShowFlags.VisualizeLightingOnProbes || bVisualizeHDR) || View.bIsReflectionCapture)
 		{
 			continue;
 		}
@@ -117,8 +139,9 @@ void StampDeferredDebugProbeDepthPS(
 		FStampDeferredDebugProbePS::FParameters* PassParameters = GraphBuilder.AllocParameters<FStampDeferredDebugProbePS::FParameters>();
 		PassParameters->RenderTargets.DepthStencil = FDepthStencilBinding(SceneDepthTexture, ERenderTargetLoadAction::ELoad, ERenderTargetLoadAction::ELoad, FExclusiveDepthStencil::DepthWrite_StencilWrite);
 
-		CommonStampDeferredDebugProbeDrawCall<true, CF_DepthNearOrEqual>(GraphBuilder, View, PassParameters, RENDER_DEPTHPREPASS);
+		CommonStampDeferredDebugProbeDrawCall<true, CF_DepthNearOrEqual>(GraphBuilder, View, PassParameters, RENDER_DEPTHPREPASS, bVisualizeHDR);
 	}
+#endif // DEBUG_PROBE_ENABLED
 }
 
 void StampDeferredDebugProbeMaterialPS(
@@ -127,13 +150,15 @@ void StampDeferredDebugProbeMaterialPS(
 	const FRenderTargetBindingSlots& BasePassRenderTargets,
 	const FMinimalSceneTextures& SceneTextures)
 {
+#if DEBUG_PROBE_ENABLED
 	RDG_EVENT_SCOPE(GraphBuilder, "StampDeferredDebugProbeMaterial");
 	RDG_GPU_STAT_SCOPE(GraphBuilder, StampDeferredDebugProbe);
 
+	const bool bVisualizeHDR = AnyViewRequireVisualizeHDR(Views);
 	const bool bVisualizeLightingOnProbes = CVarVisualizeLightingOnProbes.GetValueOnRenderThread() > 0;
 	for (const FViewInfo& View : Views)
 	{
-		if (!(bVisualizeLightingOnProbes || View.Family->EngineShowFlags.VisualizeLightingOnProbes) || View.bIsReflectionCapture)
+		if (!(bVisualizeLightingOnProbes || View.Family->EngineShowFlags.VisualizeLightingOnProbes || bVisualizeHDR) || View.bIsReflectionCapture)
 		{
 			continue;
 		}
@@ -146,14 +171,14 @@ void StampDeferredDebugProbeMaterialPS(
 			PassParameters->RenderTargets.DepthStencil = FDepthStencilBinding();
 			PassParameters->SceneTextures = CreateSceneTextureShaderParameters(GraphBuilder, View.GetSceneTexturesChecked(), View.GetFeatureLevel(), ESceneTextureSetupMode::SceneDepth);
 
-			CommonStampDeferredDebugProbeDrawCall<false, CF_Always>(GraphBuilder, View, PassParameters, RENDER_BASEPASS);
+			CommonStampDeferredDebugProbeDrawCall<false, CF_Always>(GraphBuilder, View, PassParameters, RENDER_BASEPASS, bVisualizeHDR);
 		}
 		else
 		{
-			CommonStampDeferredDebugProbeDrawCall<false, CF_DepthNearOrEqual>(GraphBuilder, View, PassParameters, RENDER_BASEPASS);
+			CommonStampDeferredDebugProbeDrawCall<false, CF_DepthNearOrEqual>(GraphBuilder, View, PassParameters, RENDER_BASEPASS, bVisualizeHDR);
 		}
-
 	}
+#endif // DEBUG_PROBE_ENABLED
 }
 
 void StampDeferredDebugProbeVelocityPS(
@@ -161,13 +186,15 @@ void StampDeferredDebugProbeVelocityPS(
 	TArrayView<const FViewInfo> Views,
 	const FRenderTargetBindingSlots& BasePassRenderTargets)
 {
+#if DEBUG_PROBE_ENABLED
 	RDG_EVENT_SCOPE(GraphBuilder, "StampDeferredDebugProbeVelocity");
 	RDG_GPU_STAT_SCOPE(GraphBuilder, StampDeferredDebugProbe);
 
+	const bool bVisualizeHDR = AnyViewRequireVisualizeHDR(Views);
 	const bool bVisualizeLightingOnProbes = CVarVisualizeLightingOnProbes.GetValueOnRenderThread() > 0;
 	for (const FViewInfo& View : Views)
 	{
-		if (!(bVisualizeLightingOnProbes || View.Family->EngineShowFlags.VisualizeLightingOnProbes) || View.bIsReflectionCapture)
+		if (!(bVisualizeLightingOnProbes || View.Family->EngineShowFlags.VisualizeLightingOnProbes || bVisualizeHDR) || View.bIsReflectionCapture)
 		{
 			continue;
 		}
@@ -176,7 +203,8 @@ void StampDeferredDebugProbeVelocityPS(
 		PassParameters->RenderTargets = BasePassRenderTargets;
 
 		const bool bRenderVelocity = true;
-		CommonStampDeferredDebugProbeDrawCall<false, CF_DepthNearOrEqual>(GraphBuilder, View, PassParameters, RENDER_VELOCITYPASS);
+		CommonStampDeferredDebugProbeDrawCall<false, CF_DepthNearOrEqual>(GraphBuilder, View, PassParameters, RENDER_VELOCITYPASS, bVisualizeHDR);
 	}
+#endif // DEBUG_PROBE_ENABLED
 }
 
diff --git a/Engine/Source/Runtime/Renderer/Private/PostProcess/PostProcessVisualizeHDR.cpp b/Engine/Source/Runtime/Renderer/Private/PostProcess/PostProcessVisualizeHDR.cpp
--- a/Engine/Source/Runtime/Renderer/Private/PostProcess/PostProcessVisualizeHDR.cpp
+++ b/Engine/Source/Runtime/Renderer/Private/PostProcess/PostProcessVisualizeHDR.cpp
@@ -205,6 +205,14 @@ FScreenPassTexture AddVisualizeHDRPass(FRDGBuilder& GraphBuilder, const FViewInf
 			Canvas.DrawShadowedString(X, Y += YStep, TEXT("Exposure Compensation (Curve):"), GetStatsFont(), FLinearColor(1, 1, 1));
 			Canvas.DrawShadowedString(X + ColumnWidth, Y, *Line, GetStatsFont(), FLinearColor(1, 1, 1));
 
+			const float IlluminanceMeterTextX = Output.ViewRect.Min.X + Output.ViewRect.Size().X * 0.5f - 180.0f;
+			const float IlluminanceMeterTextY = Output.ViewRect.Min.Y + Output.ViewRect.Size().Y * 0.5f - 130.0f;
+			Canvas.DrawShadowedString(IlluminanceMeterTextX, IlluminanceMeterTextY, TEXT("Illuminance meter measuring the back hemisphere of the camera"), GetStatsFont(), FLinearColor(1, 1, 1));
+
+			const float LuminanceMeterTextX = Output.ViewRect.Min.X + Output.ViewRect.Size().X * 0.5f - 50.0f;
+			const float LuminanceMeterTextY = Output.ViewRect.Min.Y + Output.ViewRect.Size().Y * 0.5f - 30.0f;
+			Canvas.DrawShadowedString(LuminanceMeterTextX, LuminanceMeterTextY, TEXT("Luminance meter"), GetStatsFont(), FLinearColor(1, 1, 1));
+
 			AutoExposureBias += CurveExposureBias;
 		}
 
