Change eye IOR of 1.38 (F0=0.25) to reflect correct cornea IOR.

diff --git a/Engine/Shaders/Private/Strata/Strata.ush b/Engine/Shaders/Private/Strata/Strata.ush
--- a/Engine/Shaders/Private/Strata/Strata.ush
+++ b/Engine/Shaders/Private/Strata/Strata.ush
@@ -98,6 +98,7 @@
 #define STRATA_SIMPLEVOLUME_THICKNESS_M		(STRATA_SIMPLEVOLUME_THICKNESS_CM * CENTIMETER_TO_METER)
 
 #define STRATA_EPSILON						1e-10f
+#define STRATA_EYE_DEFAULT_F0				0.028f
 
 ///////////////////////////////////////////////////////////////////////////////
 // BSDFs representation
@@ -1441,6 +1442,8 @@ FStrataData GetStrataEyeBSDF(float3 DiffuseAlbedo, float Roughness, float IrisMa
 	EYE_IRISNORMAL					(StrataData.InlinedBSDF) = IrisNormal;
 	EYE_IRISPLANENORMAL				(StrataData.InlinedBSDF) = IrisPlaneNormal;
 	EYE_SSSPROFILEID				(StrataData.InlinedBSDF) = SSSProfileId;
+	EYE_F0							(StrataData.InlinedBSDF) = STRATA_EYE_DEFAULT_F0;
+	EYE_F90							(StrataData.InlinedBSDF) = 1.0f;
 
 	const bool bHasSSS = SSSProfileId != 0;
 	BSDF_SETSSSTYPE(StrataData.InlinedBSDF, bHasSSS ? SSS_TYPE_DIFFUSION_PROFILE : SSS_TYPE_INVALID);
@@ -3540,7 +3543,7 @@ FStrataBSDF UnpackStrataBSDFIn(FStrataMaterialContainer StrataBuffer, inout FStr
 		EYE_IRISDISTANCE(OutBSDF) = Data0.y;
 		EYE_IRISNORMAL(OutBSDF) = OctahedronToUnitVector(Data1.xy * 2.0f - 1.f);
 		EYE_IRISPLANENORMAL(OutBSDF) = OctahedronToUnitVector(Data1.zw * 2.0f - 1.f);
-		EYE_F0(OutBSDF)  = 0.04f;
+		EYE_F0(OutBSDF)  = STRATA_EYE_DEFAULT_F0;
 		EYE_F90(OutBSDF) = 1.0f;
 
 		if (BSDF_GETSSSTYPE(OutBSDF) == SSS_TYPE_DIFFUSION_PROFILE)
