Fix Strata tile artifact due to invalid pixel being processed for single and simple tiles.

diff --git a/Engine/Shaders/Private/Strata/Strata.ush b/Engine/Shaders/Private/Strata/Strata.ush
--- a/Engine/Shaders/Private/Strata/Strata.ush
+++ b/Engine/Shaders/Private/Strata/Strata.ush
@@ -3305,7 +3305,9 @@ FStrataPixelHeader UnpackStrataHeaderIn(FStrataMaterialContainer StrataBuffer, i
 	BRANCH
 	if (bIsSimpleMaterial || bIsSingleMaterial)
 	{
-		Out.BSDFCount = 1;
+		// Within a tile classified as simple or single, some pixel might be invalid (e.g., sky pixels). 
+		// This ensures that invalid pixels are not processed, and keep these tiles as single/simple instead of complex
+		Out.BSDFCount = IsStrataMaterial(Out) ? 1 : 0; 
 
 	  #if STRATA_DEFERRED_SHADING
 		Out.StrataBuffer = StrataBuffer;
