Add eye compact encoding (store a partial header and then data).

diff --git a/Engine/Shaders/Private/Strata/Strata.ush b/Engine/Shaders/Private/Strata/Strata.ush
--- a/Engine/Shaders/Private/Strata/Strata.ush
+++ b/Engine/Shaders/Private/Strata/Strata.ush
@@ -181,6 +181,15 @@ struct FStrataBSDF
 //            STATE(0)       0 | 
 //            BSDF(0)       32 | F0_RGB | DiffuseAbedo(12/20)
 //                             | 20       12
+// * LAYOUT X -------------------------------------------------------------------------------------------------------------------------------------------------------------------------          
+//   (Eye)    HEADER        29 | State | A0 | IrisMask | IrisMask
+//                             | 8       8    8           8
+//            BASIS(0)       0 | (Use top layer normal)
+//            STATE(0)       0 | 
+//            BSDF(0)       32 | DiffuseAlbedo | Roughness
+//                             | 24              8
+//                          32 | IrisNormal    | IrisPlaneNormal
+//                             | 16              16
 //            -------------------------------------------------------------------------------------------------------------------------------------------------------------------------
 ////////////////////////////////////////////////////////////////////////////////////////////////////////////////////////////////////////////////////////////////////////////////////////////////
 
@@ -230,6 +239,9 @@ struct FStrataBSDF
 // Layout 2 (Complex)
 #define HEADER_COMPLEXENCODING_BIT_COUNT			(32u)
 #define HEADER_COMPLEXENCODING_MASK					(0xFFFFFFFFu)
+// Layout Eye
+#define HEADER_EYEENCODING_BIT_COUNT				(HEADER_BIT_COUNT_STATE + HEADER_BIT_COUNT_IRRADIANCE_AO)
+#define HEADER_EYEENCODING_MASK						((1u<<HEADER_EYEENCODING_BIT_COUNT)-1u)
 
 // Sanity check
 #if (HEADER_BIT_OFFSET_SHAREDLOCALBASES_TYPE + HEADER_BIT_COUNT_SHAREDLOCALBASES_TYPE) > 32
@@ -1447,7 +1459,6 @@ FStrataData GetStrataEyeBSDF(float3 DiffuseAlbedo, float Roughness, float IrisMa
 
 	const bool bHasSSS = SSSProfileId != 0;
 	BSDF_SETSSSTYPE(StrataData.InlinedBSDF, bHasSSS ? SSS_TYPE_DIFFUSION_PROFILE : SSS_TYPE_INVALID);
-
 #if STRATA_INLINE_SHADING
 	StrataData.InlinedBSDF.Coverage = 1.0f;
 #endif
@@ -1809,7 +1820,7 @@ float3x3 StrataGetBSDFSharedBasis_InlineShading(in FStrataPixelHeader StrataPixe
 float3x3 StrataGetBSDFSharedBasis_DeferredShading(in FStrataPixelHeader StrataPixelHeader, in FStrataBSDF BSDF, in FStrataAddressing StrataAddressing)
 {
 #if STRATA_DEFERRED_SHADING
-	if (IsSimpleMaterial(StrataPixelHeader) || IsSingleMaterial(StrataPixelHeader))
+	if (IsSimpleMaterial(StrataPixelHeader) || IsSingleMaterial(StrataPixelHeader) || IsEye(StrataPixelHeader))
 	{		
 		return GetTangentBasis(StrataUnpackTopLayerData(StrataPixelHeader.PackedTopLayerData).WorldNormal);
 	}
@@ -3300,10 +3311,11 @@ FStrataPixelHeader UnpackStrataHeaderIn(FStrataMaterialContainer StrataBuffer, i
 
 	const bool bIsSimpleMaterial = IsSimpleMaterial(Out);
 	const bool bIsSingleMaterial = IsSingleMaterial(Out);
+	const bool bIsEye            = IsEye(Out);
 
 	// Use simple or single layout for decoding data
 	BRANCH
-	if (bIsSimpleMaterial || bIsSingleMaterial)
+	if (bIsSimpleMaterial || bIsSingleMaterial || bIsEye)
 	{
 		// Within a tile classified as simple or single, some pixel might be invalid (e.g., sky pixels). 
 		// This ensures that invalid pixels are not processed, and keep these tiles as single/simple instead of complex
@@ -3438,6 +3450,20 @@ FStrataBSDF UnpackStrataBSDFIn(FStrataMaterialContainer StrataBuffer, inout FStr
 		// For now, single encode doesn't support anisotropy, as the top layer data only contains the normal, not the full frame basis
 		BSDF_SETHASANISOTROPY(OutBSDF, 0);
 	}
+	else if (IsEye(StrataHeader))
+	{
+		uint Data = 0;
+	#if STRATA_DEFERRED_SHADING
+		Data = StrataHeader.SharedLocalBasesTypes_PackedHeader;
+	#endif
+		OutBSDF.State = Data & HEADER_EYEENCODING_MASK;
+		OutBSDF.LuminanceWeightV = 1.0f;
+		BSDF_SETTYPE(OutBSDF, STRATA_BSDF_TYPE_EYE);
+		BSDF_SETSHAREDLOCALBASISID(OutBSDF, 0);
+		BSDF_SETHASANISOTROPY(OutBSDF, 0);
+		BSDF_SETISTOPLAYER(OutBSDF, 1);
+		BSDF_SETSSSTYPE(OutBSDF, HasSubsurface(StrataHeader) ? SSS_TYPE_DIFFUSION_PROFILE : SSS_TYPE_INVALID);
+	}
 	else
 	{
 		OutBSDF.State = StrataLoadUint1(StrataBuffer, StrataAddressing);
@@ -3538,16 +3564,23 @@ FStrataBSDF UnpackStrataBSDFIn(FStrataMaterialContainer StrataBuffer, inout FStr
 	break;
 	case STRATA_BSDF_TYPE_EYE:
 	{
-		StrataLoad_ColorGamma2ToLinearAlphaLinear(StrataBuffer, StrataAddressing, EYE_DIFFUSEALBEDO(OutBSDF), EYE_ROUGHNESS(OutBSDF));		
-		const float4 Data0 = UnpackRGBA8(StrataLoadUint1(StrataBuffer, StrataAddressing));
-		const float4 Data1 = UnpackRGBA8(StrataLoadUint1(StrataBuffer, StrataAddressing));
+		float2 Data0 = 0;
+	#if STRATA_DEFERRED_SHADING
+		Data0 = UnpackRGBA8(StrataHeader.SharedLocalBasesTypes_PackedHeader).zw;
+	#endif
 		EYE_IRISMASK(OutBSDF) = Data0.x;
 		EYE_IRISDISTANCE(OutBSDF) = Data0.y;
+
+		StrataLoad_ColorGamma2ToLinearAlphaLinear(StrataBuffer, StrataAddressing, EYE_DIFFUSEALBEDO(OutBSDF), EYE_ROUGHNESS(OutBSDF));
+		const float4 Data1 = UnpackRGBA8(StrataLoadUint1(StrataBuffer, StrataAddressing));
 		EYE_IRISNORMAL(OutBSDF) = OctahedronToUnitVector(Data1.xy * 2.0f - 1.f);
 		EYE_IRISPLANENORMAL(OutBSDF) = OctahedronToUnitVector(Data1.zw * 2.0f - 1.f);
 		EYE_F0(OutBSDF)  = STRATA_EYE_DEFAULT_F0;
 		EYE_F90(OutBSDF) = 1.0f;
 
+		// When loading Strata BSDF which have SSS enable, BaseColor and (optionally) Specular value can be overriden 
+		// based on the SSS method used (screen space SSS / checherboarded screen space SSS / ...)
+		#if STRATA_SSS_MATERIAL_OVERRIDE
 		if (BSDF_GETSSSTYPE(OutBSDF) == SSS_TYPE_DIFFUSION_PROFILE)
 		{
 			const bool bChecker = CheckerFromPixelPos(StrataAddressing.PixelCoords);
@@ -3558,7 +3591,8 @@ FStrataBSDF UnpackStrataBSDFIn(FStrataMaterialContainer StrataBuffer, inout FStr
 			EYE_F0(OutBSDF)  *= SpecularFactor;
 			EYE_F90(OutBSDF) *= SpecularFactor;
 		}
-		// 12 bytes
+		#endif
+		// 8 bytes
 	}
 	break;
 	case STRATA_BSDF_TYPE_SINGLELAYERWATER:
diff --git a/Engine/Shaders/Private/Strata/StrataExport.ush b/Engine/Shaders/Private/Strata/StrataExport.ush
--- a/Engine/Shaders/Private/Strata/StrataExport.ush
+++ b/Engine/Shaders/Private/Strata/StrataExport.ush
@@ -452,6 +452,22 @@ void PackStrataOut(
 						#endif
 						}
 					}
+					// C. Layout Eye
+					else if (BSDF_GETTYPE(BSDF) == STRATA_BSDF_TYPE_EYE)
+					{
+						// Data0 (Header_State|Header_AO|Data)
+						{
+							uint Out = 0;
+							HEADER_SETCOMMONSTATES(Out, StrataPixelHeader.State);
+							Out = Out & HEADER_EYEENCODING_MASK;
+							Out = Out | (PackRGBA8(float4(EYE_IRISMASK(BSDF), EYE_IRISDISTANCE(BSDF), 0, 0)) << 16u);
+							STRATA_STORE_UINT1(Out);
+
+						#if (HEADER_EYEENCODING_BIT_COUNT) > 32
+						#error Strata eye path header is > 32bits
+						#endif
+						}
+					}
 					// C. Layout2 - Weight for Regular/Complex path
 					else if (GreyScaleThroughputV > 0)
 					{
@@ -555,9 +571,8 @@ void PackStrataOut(
 							const float2 EncodedIrisNormal = UnitVectorToOctahedron(EYE_IRISNORMAL(BSDF));
 							const float2 EncodedIrisPlaneNormal = UnitVectorToOctahedron(EYE_IRISPLANENORMAL(BSDF));
 							STRATA_STORE_UINT1(PackColorLinearToGamma2AlphaLinear(float4(EYE_DIFFUSEALBEDO(BSDF), EYE_ROUGHNESS(BSDF))));
-							STRATA_STORE_UINT1(PackRGBA8(float4(EYE_IRISMASK(BSDF), EYE_IRISDISTANCE(BSDF), 0, 0)));
 							STRATA_STORE_UINT1(PackRGBA8(float4(EncodedIrisNormal * 0.5f + 0.5f, EncodedIrisPlaneNormal * 0.5f + 0.5f)));
-							// 12 bytes
+							// 8 bytes
 
 							// Note: we don't store the SSS profile ID into the Eye data, as it is not needed there. The profile is only stored into the SSS data.
 							const bool bHasSSS = BSDF_GETSSSTYPE(BSDF) == SSS_TYPE_DIFFUSION_PROFILE;
diff --git a/Engine/Shaders/Private/Strata/StrataVisualize.usf b/Engine/Shaders/Private/Strata/StrataVisualize.usf
--- a/Engine/Shaders/Private/Strata/StrataVisualize.usf
+++ b/Engine/Shaders/Private/Strata/StrataVisualize.usf
@@ -488,15 +488,20 @@ void PrintPixel(uint2 InCoord, float3 InWorldPosition, float3 V)
 
 		const uint FootPrint_PostBSDFs = StrataAddressing.ReadBytes;
 
-		const uint HeaderSize = FootPrint_PostHeader - FootPrint_Start;
-		const uint BSDFsSize  = FootPrint_PostBSDFs - FootPrint_PostHeader;
-		const uint TotalSize  = FootPrint_PostBSDFs - FootPrint_Start;
+		const FStrataSubsurfaceHeader SSSHeader = StrataLoadSubsurfaceHeader(Strata.SSSTexture, StrataAddressing.PixelCoords);
+		const bool bHasSSSData = StrataSubSurfaceHeaderGetIsValid(SSSHeader);
 
 		const uint TopLayerDataBytes = 4;	// bytes 
+		const uint SSSDataBytes = 8;	    // bytes 
+
+		const uint HeaderSize =  FootPrint_PostHeader - FootPrint_Start;
+		const uint BSDFsSize  =  FootPrint_PostBSDFs - FootPrint_PostHeader;
+		const uint TotalSize  = (FootPrint_PostBSDFs - FootPrint_Start) + (bHasSSSData ? SSSDataBytes : 0);
+
 		RectMin = Context.Pos;
 		RectMax = Context.StartPos;
 		Print(Context, TEXT("Memory"), FontEmerald); Newline(Context, RectMax);
-		if (bIsSingleMaterial)
+		if (bIsSingleMaterial || IsEye(Header))
 		{
 			Print(Context, TEXT(" Header       "), FontSilver); Print(Context, HeaderSize - TopLayerDataBytes,				FontSilver, 3, 0); Print(Context, TEXT("bytes"), FontSilver); Newline(Context, RectMax);
 			Print(Context, TEXT(" TopNormalTex "), FontSilver); Print(Context, TopLayerDataBytes,							FontSilver, 3, 0); Print(Context, TEXT("bytes"), FontSilver); Newline(Context, RectMax);
@@ -512,6 +517,12 @@ void PrintPixel(uint2 InCoord, float3 InWorldPosition, float3 V)
 			Print(Context, TEXT(" Header+Norm  "), FontSilver); Print(Context, HeaderSize,									FontSilver, 3, 0); Print(Context, TEXT("bytes"), FontSilver); Newline(Context, RectMax);
 			Print(Context, TEXT(" BSDFs        "), FontSilver); Print(Context, BSDFsSize,									FontSilver, 3, 0); Print(Context, TEXT("bytes"), FontSilver); Newline(Context, RectMax);
 		}
+
+		if (bHasSSSData)
+		{
+			Print(Context, TEXT(" SSS Data     "), FontSilver); Print(Context, SSSDataBytes, FontSilver, 3, 0); Print(Context, TEXT("bytes"), FontSilver); Newline(Context, RectMax);
+		}
+
 		Print(Context, TEXT(" Total        "), FontOrange); Print(Context, TotalSize,  FontOrange, 3, 0); Print(Context, TEXT("bytes"), FontOrange); Newline(Context, RectMax);
 		PrintBackgroundRect(Context, RectMin, RectMax);
 	}
