Add Strata compact hair storage.

Fix issue with eye storage introduced by previous CL.

diff --git a/Engine/Shaders/Private/Strata/Strata.ush b/Engine/Shaders/Private/Strata/Strata.ush
--- a/Engine/Shaders/Private/Strata/Strata.ush
+++ b/Engine/Shaders/Private/Strata/Strata.ush
@@ -190,6 +190,15 @@ struct FStrataBSDF
 //                             | 24              8
 //                          32 | IrisNormal    | IrisPlaneNormal
 //                             | 16              16
+// * LAYOUT X -------------------------------------------------------------------------------------------------------------------------------------------------------------------------          
+//   (Hair)   HEADER        29 | State | A0 | Free | Free
+//                             | 8       8    8           8
+//            BASIS(0)       0 | (Use top layer normal)
+//            STATE(0)       0 | 
+//            BSDF(0)       32 | BaseColor | Roughness
+//                             | 24              8
+//                          32 | Scatter | Specular | Backlit | ComplexTransmittance
+//                             | 8         8          8         8
 //            -------------------------------------------------------------------------------------------------------------------------------------------------------------------------
 ////////////////////////////////////////////////////////////////////////////////////////////////////////////////////////////////////////////////////////////////////////////////////////////////
 
@@ -242,6 +251,9 @@ struct FStrataBSDF
 // Layout Eye
 #define HEADER_EYEENCODING_BIT_COUNT				(HEADER_BIT_COUNT_STATE + HEADER_BIT_COUNT_IRRADIANCE_AO)
 #define HEADER_EYEENCODING_MASK						((1u<<HEADER_EYEENCODING_BIT_COUNT)-1u)
+// Layout Hair
+#define HEADER_HAIRENCODING_BIT_COUNT				(HEADER_BIT_COUNT_STATE + HEADER_BIT_COUNT_IRRADIANCE_AO)
+#define HEADER_HAIRENCODING_MASK					((1u<<HEADER_HAIRENCODING_BIT_COUNT)-1u)
 
 // Sanity check
 #if (HEADER_BIT_OFFSET_SHAREDLOCALBASES_TYPE + HEADER_BIT_COUNT_SHAREDLOCALBASES_TYPE) > 32
@@ -1820,7 +1832,7 @@ float3x3 StrataGetBSDFSharedBasis_InlineShading(in FStrataPixelHeader StrataPixe
 float3x3 StrataGetBSDFSharedBasis_DeferredShading(in FStrataPixelHeader StrataPixelHeader, in FStrataBSDF BSDF, in FStrataAddressing StrataAddressing)
 {
 #if STRATA_DEFERRED_SHADING
-	if (IsSimpleMaterial(StrataPixelHeader) || IsSingleMaterial(StrataPixelHeader) || IsEye(StrataPixelHeader))
+	if (IsSimpleMaterial(StrataPixelHeader) || IsSingleMaterial(StrataPixelHeader) || IsEye(StrataPixelHeader) || IsHair(StrataPixelHeader))
 	{		
 		return GetTangentBasis(StrataUnpackTopLayerData(StrataPixelHeader.PackedTopLayerData).WorldNormal);
 	}
@@ -3309,13 +3321,9 @@ FStrataPixelHeader UnpackStrataHeaderIn(FStrataMaterialContainer StrataBuffer, i
 	Out.State		= HEADER_GETCOMMONSTATES(PackedHeader);	// Unpack state and AO data at the same time 
 	Out.BSDFCount   = HEADER_GETBSDFCOUNT(PackedHeader);
 
-	const bool bIsSimpleMaterial = IsSimpleMaterial(Out);
-	const bool bIsSingleMaterial = IsSingleMaterial(Out);
-	const bool bIsEye            = IsEye(Out);
-
 	// Use simple or single layout for decoding data
 	BRANCH
-	if (bIsSimpleMaterial || bIsSingleMaterial || bIsEye)
+	if (IsSimpleMaterial(Out) || IsSingleMaterial(Out) || IsEye(Out) || IsHair(Out))
 	{
 		// Within a tile classified as simple or single, some pixel might be invalid (e.g., sky pixels). 
 		// This ensures that invalid pixels are not processed, and keep these tiles as single/simple instead of complex
@@ -3464,6 +3472,19 @@ FStrataBSDF UnpackStrataBSDFIn(FStrataMaterialContainer StrataBuffer, inout FStr
 		BSDF_SETISTOPLAYER(OutBSDF, 1);
 		BSDF_SETSSSTYPE(OutBSDF, HasSubsurface(StrataHeader) ? SSS_TYPE_DIFFUSION_PROFILE : SSS_TYPE_INVALID);
 	}
+	else if (IsHair(StrataHeader))
+	{
+		uint Data = 0;
+	#if STRATA_DEFERRED_SHADING
+		Data = StrataHeader.SharedLocalBasesTypes_PackedHeader;
+	#endif
+		OutBSDF.State = Data & HEADER_HAIRENCODING_MASK;
+		OutBSDF.LuminanceWeightV = 1.0f;
+		BSDF_SETTYPE(OutBSDF, STRATA_BSDF_TYPE_HAIR);
+		BSDF_SETSHAREDLOCALBASISID(OutBSDF, 0);
+		BSDF_SETHASANISOTROPY(OutBSDF, 0);
+		BSDF_SETISTOPLAYER(OutBSDF, 1);
+	}
 	else
 	{
 		OutBSDF.State = StrataLoadUint1(StrataBuffer, StrataAddressing);
@@ -3564,12 +3585,12 @@ FStrataBSDF UnpackStrataBSDFIn(FStrataMaterialContainer StrataBuffer, inout FStr
 	break;
 	case STRATA_BSDF_TYPE_EYE:
 	{
-		float2 Data0 = 0;
+		float4 Data0 = 0;
 	#if STRATA_DEFERRED_SHADING
-		Data0 = UnpackRGBA8(StrataHeader.SharedLocalBasesTypes_PackedHeader).zw;
+		Data0 = UnpackRGBA8(StrataHeader.SharedLocalBasesTypes_PackedHeader);
 	#endif
-		EYE_IRISMASK(OutBSDF) = Data0.x;
-		EYE_IRISDISTANCE(OutBSDF) = Data0.y;
+		EYE_IRISMASK(OutBSDF) = Data0.z;
+		EYE_IRISDISTANCE(OutBSDF) = Data0.w;
 
 		StrataLoad_ColorGamma2ToLinearAlphaLinear(StrataBuffer, StrataAddressing, EYE_DIFFUSEALBEDO(OutBSDF), EYE_ROUGHNESS(OutBSDF));
 		const float4 Data1 = UnpackRGBA8(StrataLoadUint1(StrataBuffer, StrataAddressing));
diff --git a/Engine/Shaders/Private/Strata/StrataExport.ush b/Engine/Shaders/Private/Strata/StrataExport.ush
--- a/Engine/Shaders/Private/Strata/StrataExport.ush
+++ b/Engine/Shaders/Private/Strata/StrataExport.ush
@@ -355,12 +355,15 @@ void PackStrataOut(
 		// 3 types of encodings
 		// * A. Layout0: Simple encoding	(use top layer normal)
 		// * B. Layout1: Single encoding	(use top layer normal)
-		// * C. Layout2: Complex encoding	(use basis)
+		// * C. Eye    : Custom encoding	(use top layer normal)
+		// * D. Hair   : Custom encoding	(use top layer normal)
+		// * E. Layout2: Complex encoding	(use basis)
 		const bool bHasFastEncoding = IsSimpleMaterial(StrataPixelHeader);
 		const bool bHasSingleEncoding = IsSingleMaterial(StrataPixelHeader);
+		const bool bCustomEncoding = IsHair(StrataPixelHeader) || IsEye(StrataPixelHeader);
 
 		// (Layout2)
-		if (!bHasFastEncoding && !bHasSingleEncoding)
+		if (!bHasFastEncoding && !bHasSingleEncoding && !bCustomEncoding)
 		{
 			// 1. the header (Regular/Complex encoding)
 			const uint PackedHeader = PackStrataHeader(BSDFVisibleCount, StrataPixelHeader);
@@ -440,35 +443,44 @@ void PackStrataOut(
 					else if (bHasSingleEncoding && BSDFIdx == OneBSDFMaterial_Index)
 					{
 						// Data0 (Header_State|Header_AO|Header_BSDFTypes|BSDF_State)
-						{
-							uint Out = 0;
-							HEADER_SETCOMMONSTATES(Out, StrataPixelHeader.State);
-							Out = Out & HEADER_SINGLEENCODING_MASK;
-							Out = Out | ((BSDF.State & BSDF_SINGLEENCODING_MASK) << HEADER_SINGLEENCODING_BIT_COUNT);
-							STRATA_STORE_UINT1(Out);
+						uint Out = 0;
+						HEADER_SETCOMMONSTATES(Out, StrataPixelHeader.State);
+						Out = Out & HEADER_SINGLEENCODING_MASK;
+						Out = Out | ((BSDF.State & BSDF_SINGLEENCODING_MASK) << HEADER_SINGLEENCODING_BIT_COUNT);
+						STRATA_STORE_UINT1(Out);
 
 						#if (HEADER_SIMPLEENCODING_BIT_COUNT) > 32
 						#error Strata fast path header is > 32bits
 						#endif
-						}
 					}
-					// C. Layout Eye
-					else if (BSDF_GETTYPE(BSDF) == STRATA_BSDF_TYPE_EYE)
+					// C. Layout X (Eye)
+					else if (IsEye(StrataPixelHeader))
 					{
 						// Data0 (Header_State|Header_AO|Data)
-						{
-							uint Out = 0;
-							HEADER_SETCOMMONSTATES(Out, StrataPixelHeader.State);
-							Out = Out & HEADER_EYEENCODING_MASK;
-							Out = Out | (PackRGBA8(float4(EYE_IRISMASK(BSDF), EYE_IRISDISTANCE(BSDF), 0, 0)) << 16u);
-							STRATA_STORE_UINT1(Out);
+						uint Out = 0;
+						HEADER_SETCOMMONSTATES(Out, StrataPixelHeader.State);
+						Out = Out & HEADER_EYEENCODING_MASK;
+						Out = Out | PackRGBA8(float4(0, 0, EYE_IRISMASK(BSDF), EYE_IRISDISTANCE(BSDF)));
+						STRATA_STORE_UINT1(Out);
 
 						#if (HEADER_EYEENCODING_BIT_COUNT) > 32
 						#error Strata eye path header is > 32bits
 						#endif
-						}
 					}
-					// C. Layout2 - Weight for Regular/Complex path
+					// D. Layout X (Hair)
+					else if (IsHair(StrataPixelHeader))
+					{
+						// Data0 (Header_State|Header_AO|Data)
+						uint Out = 0;
+						HEADER_SETCOMMONSTATES(Out, StrataPixelHeader.State);
+						Out = Out & HEADER_HAIRENCODING_MASK;
+						STRATA_STORE_UINT1(Out);
+
+						#if (HEADER_EYEENCODING_BIT_COUNT) > 32
+						#error Strata hair path header is > 32bits
+						#endif
+					}
+					// E. Layout2 - Weight for Regular/Complex path
 					else if (GreyScaleThroughputV > 0)
 					{
 						BSDF_SETWEIGHT10F(BSDF, Pack10F(BSDF.LuminanceWeightV.x));
