Volumetric cloud render target now uses distance conversion macros.

diff --git a/Engine/Shaders/Private/VolumetricRenderTarget.usf b/Engine/Shaders/Private/VolumetricRenderTarget.usf
--- a/Engine/Shaders/Private/VolumetricRenderTarget.usf
+++ b/Engine/Shaders/Private/VolumetricRenderTarget.usf
@@ -144,7 +144,7 @@ void ReconstructVolumetricRenderTargetPS(
 
 		// Sample participating media "front depth" for a better reprojection
 		float TracingVolumetricSampleDepthKm = SafeLoadTracingVolumetricDepthTexture(int2(SVPos.xy) / DownSampleFactor).x;
-		float TracingVolumetricSampleDepth = TracingVolumetricSampleDepthKm * 100000.0f; // Kilometers to centimeters
+		float TracingVolumetricSampleDepth = TracingVolumetricSampleDepthKm * KILOMETER_TO_CENTIMETER;
 		float DeviceZ = ConvertToDeviceZ(TracingVolumetricSampleDepth); // Approximation. Should try real DeviceZ
 
 		float4 CurrClip = float4(ScreenPosition, DeviceZ, 1); // Inverted Far Depth = 0.0f
@@ -209,7 +209,7 @@ void ReconstructVolumetricRenderTargetPS(
 
 				const float ReconstructDepthZ = HalfResDepthTexture.Load(int3(SVPos.xy + ViewViewRectMin, 0)).r;
 				const float3 WorldPosition = LWCHackToFloat(SvPositionToWorld(float4(CenterSample, ReconstructDepthZ, 1.0)));
-				const float PixelDistanceFromViewKm = length(WorldPosition - LWCHackToFloat(PrimaryView.WorldCameraOrigin)) * (1.0f / 100000.0f); // centimeter to kilometer
+				const float PixelDistanceFromViewKm = length(WorldPosition - LWCHackToFloat(PrimaryView.WorldCameraOrigin)) * CENTIMETER_TO_KILOMETER;
 
 				RGBA = HistoryRGBA;
 				Depths = HistoryDepths;
@@ -465,7 +465,7 @@ void ComposeVolumetricRTOverScenePS(
 #endif 
 
 	float PixelDistanceFromView   = length(PrimaryView.TranslatedWorldCameraOrigin - ScreenTranslatedWorldPosition);
-	float PixelDistanceFromViewKm = PixelDistanceFromView * (1.0f / 100000.0f); // centimeter to kilometer
+	float PixelDistanceFromViewKm = PixelDistanceFromView * CENTIMETER_TO_KILOMETER;
 
 
 
