Illuminance meter improvements (show mean and rdg, account for viewport min)

diff --git a/Engine/Shaders/Private/PostProcessVisualizeHDR.usf b/Engine/Shaders/Private/PostProcessVisualizeHDR.usf
--- a/Engine/Shaders/Private/PostProcessVisualizeHDR.usf
+++ b/Engine/Shaders/Private/PostProcessVisualizeHDR.usf
@@ -300,13 +300,14 @@ void MainPS(noperspective float4 UVAndScreenPos : TEXCOORD0, float4 SvPosition :
 		const float  IlluminanceMeterHalfSize = IlluminanceMeterSize / 2;
 		const float2 IlluminanceMeterCenter = View.ViewSizeAndInvSize.xy * 0.5f - uint2(0, 100);
 
-		const float2 IlluminanceMeterCenterUV = (0*Input_UVViewportMin + IlluminanceMeterCenter) * Input_ExtentInverse;
+		const float2 IlluminanceMeterCenterCoord = Input_UVViewportMin + IlluminanceMeterCenter;
+		const float2 IlluminanceMeterCenterUV = IlluminanceMeterCenterCoord * Input_ExtentInverse;
 
 		const float PixelDx = PixelPos.x - IlluminanceMeterCenter.x;
 		const float PixelDy = PixelPos.y - IlluminanceMeterCenter.y;
 		const float AbsPixelDx = abs(PixelDx);
 		const float AbsPixelDy = abs(PixelDy);
-		if (PixelDx > - 20 && PixelDx < 110 && PixelDy < 20 && PixelDy > -20)
+		if (PixelDx > - 20 && PixelDx < 200 && PixelDy < 40 && PixelDy > -20)
 		{
 			float3 LuminanceAvg = 0.0f;
 			float SampleCount = 0.0f;
@@ -332,23 +333,52 @@ void MainPS(noperspective float4 UVAndScreenPos : TEXCOORD0, float4 SvPosition :
 				}
 			}
 
-			float MaxLuminance = max(LuminanceAvg.r, max(LuminanceAvg.g, LuminanceAvg.b));
-			float MaxIlluminance = MaxLuminance * PI;
+			float3 Illuminance = LuminanceAvg * PI;
+			float MeanIlluminance = Luminance(Illuminance);
 
-			int2 TopLeft = IlluminanceMeterCenter + int2(11, 11);
-			// Illuminance
-			if (MaxIlluminance < 1000.0f)
+			int2 TopLeft = IlluminanceMeterCenterCoord + int2(11, 11);
+			// Mean illuminance
+			if (MeanIlluminance < 1000.0f)
 			{
-				PrintFloat(PixelPos, OutColor.xyz, float3(1, 1, 1), TopLeft, MaxIlluminance);
+				PrintFloat(PixelPos, OutColor.xyz, float3(1, 1, 1), TopLeft, MeanIlluminance);
 			}
 			else
 			{
-				PrintFloatNoFractionLarge(PixelPos, OutColor.xyz, float3(1, 1, 1), TopLeft, MaxIlluminance);
+				PrintFloatNoFractionLarge(PixelPos, OutColor.xyz, float3(1, 1, 1), TopLeft, MeanIlluminance);
 			}
 			TopLeft.x += 60;
 			PrintCharacter(PixelPos, OutColor.xyz, float3(1, 1, 1), TopLeft, _L_);
 			PrintCharacter(PixelPos, OutColor.xyz, float3(1, 1, 1), TopLeft, _U_);
 			PrintCharacter(PixelPos, OutColor.xyz, float3(1, 1, 1), TopLeft, _X_);
+
+			// Illuminance as RGB
+			TopLeft = IlluminanceMeterCenterCoord + int2(11, 22);
+			if (Illuminance.r < 1000.0f)
+			{
+				PrintFloat(PixelPos, OutColor.xyz, float3(1, 0.5, 0.5), TopLeft, Illuminance.r);
+			}
+			else
+			{
+				PrintFloatNoFractionLarge(PixelPos, OutColor.xyz, float3(1, 0.5, 0.5), TopLeft, Illuminance.r);
+			}
+			TopLeft.x += 60;
+			if (Illuminance.g < 1000.0f)
+			{
+				PrintFloat(PixelPos, OutColor.xyz, float3(0.5, 1, 0.5), TopLeft, Illuminance.g);
+			}
+			else
+			{
+				PrintFloatNoFractionLarge(PixelPos, OutColor.xyz, float3(0.5, 1, 0.5), TopLeft, Illuminance.g);
+			}
+			TopLeft.x += 60;
+			if (Illuminance.b < 1000.0f)
+			{
+				PrintFloat(PixelPos, OutColor.xyz, float3(0.5, 0.5, 1), TopLeft, Illuminance.b);
+			}
+			else
+			{
+				PrintFloatNoFractionLarge(PixelPos, OutColor.xyz, float3(0.5, 0.5, 1), TopLeft, Illuminance.b);
+			}
 			return;
 		}
 	}
