Fixed illuminance meter border to look good.

diff --git a/Engine/Shaders/Private/PostProcessVisualizeHDR.usf b/Engine/Shaders/Private/PostProcessVisualizeHDR.usf
--- a/Engine/Shaders/Private/PostProcessVisualizeHDR.usf
+++ b/Engine/Shaders/Private/PostProcessVisualizeHDR.usf
@@ -325,9 +325,12 @@ void MainPS(noperspective float4 UVAndScreenPos : TEXCOORD0, float4 SvPosition :
 			LuminanceAvg *= View.OneOverPreExposure;
 
 			OutColor = Texture2DSample(SceneColorTexture, SceneColorSampler, UV); // Influenced area default scene color
-			if (AbsPixelDx < IlluminanceMeterHalfSize && AbsPixelDy < IlluminanceMeterHalfSize)
+			const int IlluminanceMeterHalfSizeInt = int(IlluminanceMeterHalfSize);
+			const int AbsPixelDxInt = int(AbsPixelDx);
+			const int AbsPixelDyInt = int(AbsPixelDy);
+			if (AbsPixelDxInt < IlluminanceMeterHalfSizeInt && AbsPixelDyInt < IlluminanceMeterHalfSizeInt)
 			{
-				if (AbsPixelDx == IlluminanceMeterHalfSize - 1 || AbsPixelDy == IlluminanceMeterHalfSize - 1)
+				if (AbsPixelDxInt == (IlluminanceMeterHalfSizeInt - 1) || AbsPixelDyInt == (IlluminanceMeterHalfSizeInt-1))
 				{
 					OutColor = float4(1.0, 1.0, 1.0, 1.0); // White border
 				}
