Strata - Fixed export basecolor to actually be colored.

diff --git a/Engine/Shaders/Private/Strata/StrataExport.ush b/Engine/Shaders/Private/Strata/StrataExport.ush
--- a/Engine/Shaders/Private/Strata/StrataExport.ush
+++ b/Engine/Shaders/Private/Strata/StrataExport.ush
@@ -199,7 +199,7 @@ void PackStrataOut(
 
 						TopLayerTotalWeight += BSDF.TopLayerDataWeight;
 						float TopLayerRoughnessContribution = BSDF.TopLayerDataWeight * SLAB_ROUGHNESS(BSDF);
-						float TopLayerBaseColorContribution = BSDF.TopLayerDataWeight * StrataGetBSDFBaseColor(BSDF);
+						float3 TopLayerBaseColorContribution = BSDF.TopLayerDataWeight * StrataGetBSDFBaseColor(BSDF);
 						if (BSDF_GETHASHAZINESS(BSDF))
 						{
 							FHaziness Haziness = UnpackHaziness(SLAB_HAZINESS(BSDF));
