Strata - fixed debug probes compilation

diff --git a/Engine/Shaders/Private/DebugProbes.usf b/Engine/Shaders/Private/DebugProbes.usf
--- a/Engine/Shaders/Private/DebugProbes.usf
+++ b/Engine/Shaders/Private/DebugProbes.usf
@@ -38,6 +38,8 @@ int DebugProbesMode;
 #define RENDER_BASEPASS		1
 #define RENDER_VELOCITYPASS	2
 
+// Base pass with strata cannot write to depth because it is going to execute in shader depth testing
+#define DEPTH_OUTPUT_ENABLED (PERMUTATION_PASS == RENDER_VELOCITYPASS || PERMUTATION_PASS == RENDER_DEPTHPREPASS || (PERMUTATION_PASS == RENDER_BASEPASS && STRATA_ENABLED == 0))
 
 
 void DebugSphereIntersection(
@@ -67,6 +69,8 @@ void DebugSphereIntersection(
 
 #if PERMUTATION_PASS == RENDER_BASEPASS
 void WriteOutMaterial(
+	in uint2 PixelCoord,
+	in float3 RayDir,
 	in float PerObjectGBufferData,
 	in float SelectiveOutputMask,
 	in float3 BaseColor,
@@ -206,7 +210,9 @@ void MainPS(
 	, out float4 OutGBufferE	: SV_Target5
 	#endif
 #endif
+#if DEPTH_OUTPUT_ENABLED
 	, out float  OutDepth		: SV_DEPTH
+#endif
 )
 {
 	OutColor = 0.f;
@@ -273,6 +279,8 @@ void MainPS(
 	#endif // STRATA_ENABLED
 		{
 			WriteOutMaterial(
+				PixelCoord,
+				RayDir,
 				PerObjectGBufferData,
 				SelectiveOutputMask,
 				BaseColor,
@@ -428,7 +436,9 @@ void MainPS(
 	float4 ClipPos = mul(float4(WorldPosition, 1.0f), LWCHackToFloat(ResolvedView.WorldToClip));
 	ClipPos /= ClipPos.w;
 	const float DeviceZ = ClipPos.z;
+#if DEPTH_OUTPUT_ENABLED
 	OutDepth = INVARIANT(DeviceZ);
+#endif
 
 	if (RayT <= 0.0f)
 	{
@@ -471,6 +481,8 @@ void MainPS(
 #endif // STRATA_ENABLED
 	{
 		WriteOutMaterial(
+			PixelCoord,
+			RayDir,
 			PerObjectGBufferData,
 			SelectiveOutputMask,
 			BaseColor,
@@ -492,8 +504,7 @@ void MainPS(
 
 #endif // PERMUTATION_VISUALIZEHDR /////////////////////////////////////////////////////////////////////////////////////////////////////////////////////////////////////
 
-#if PERMUTATION_PASS == RENDER_VELOCITYPASS || PERMUTATION_PASS == RENDER_DEPTHPREPASS || (PERMUTATION_PASS == RENDER_BASEPASS && STRATA_ENABLED == 0)
-	// Base pass with strata cannot write to depth because it is going to execute in shader depth testing
+#if DEPTH_OUTPUT_ENABLED
 	OutDepth = DeviceZ;
 #endif
 
