Implemented GPU driven shadow atlas for hair Deep Opacity Maps. Merged FDeepShadowViewInfo and FDeepShadowTransform structs/buffers.

#rnx

diff --git a/Engine/Shaders/Private/HairStrands/HairStrandsDeepShadowAllocation.usf b/Engine/Shaders/Private/HairStrands/HairStrandsDeepShadowAllocation.usf
--- a/Engine/Shaders/Private/HairStrands/HairStrandsDeepShadowAllocation.usf
+++ b/Engine/Shaders/Private/HairStrands/HairStrandsDeepShadowAllocation.usf
@@ -47,10 +47,21 @@ int2				SlotResolution;
 uint				SlotIndexCount;
 uint				MacroGroupCount;
 
+uint2				AtlasResolution;
+float2				AtlasTexelSize;
+uint				MinAtlasTileResolution;
+uint				MinAtlasTileResolutionLog2;
+
 Buffer<int>								MacroGroupAABBBuffer;
 RWStructuredBuffer<FDeepShadowViewInfo>	OutShadowViewInfoBuffer;
-RWStructuredBuffer<FDeepShadowTransform>OutShadowTranslatedWorldToLightTransformBuffer;
 
+// Each bin corresponds to a power of two shadow atlas tile resolution. The resolution is calculated as pow(2, MinAtlasTileResolutionLog2 + BinIndex).
+// This gives us a maximum tile resolution of pow(2, MinAtlasTileResolutionLog2 + NUM_ATLAS_TILE_BINS - 1).
+#define NUM_ATLAS_TILE_BINS 6
+groupshared uint AtlasTileBinElementCounts[NUM_ATLAS_TILE_BINS];
+groupshared uint AtlasTileBins[NUM_ATLAS_TILE_BINS][MAX_SLOT_COUNT];
+groupshared uint4 AtlasTileAllocations[MAX_SLOT_COUNT];
+groupshared uint AtlasTileDesiredResolutions[MAX_SLOT_COUNT];
 
 float DegreesToRadians(float InDeg)
 {
@@ -110,7 +121,7 @@ void ComputeTranslatedWorldToLightClip(
 		MinZ = max(1.0f, MinZ);
 
 		const float SphereDistance			= length(TranslatedLightPosition - Center);
-		const float HalfFov					= min(MaxHafFovInRad, asin(Radius / SphereDistance));
+		const float HalfFov					= min(MaxHafFovInRad, atan(Radius / SphereDistance));
 
 		const float4x4 TranslatedWorldToLight	= LookAtMatrix(TranslatedLightPosition, Center, float3(0, 0, 1));
 		const float4x4 ProjMatrix				= ReversedZPerspectiveMatrix(HalfFov, 1, 1, MinZ, MaxZ);
@@ -120,13 +131,178 @@ void ComputeTranslatedWorldToLightClip(
 	}
 }
 
+float2 ComputeProjectedScreenSize(float4x4 TranslatedWorldToLightClipTransform)
+{
+	float4x4 LightClipToTranslatedWorld = Inverse(TranslatedWorldToLightClipTransform);
+	float4x4 LightClipToScreenClip = mul(LightClipToTranslatedWorld, View.TranslatedWorldToClip);
+	float2 ProjectedMinUV = 99999.0f;
+	float2 ProjectedMaxUV = -99999.0f;
+
+	UNROLL
+	for (int Y = 0; Y < 2; ++Y)
+	{
+		UNROLL
+		for (int X = 0; X < 2; ++X)
+		{
+			float4 LightClip = float4(X * 2.0f - 1.0f, Y * 2.0f - 1.0f, 0.0f, 1.0f);
+			float4 ScreenProjected = mul(LightClip, LightClipToScreenClip);
+			float2 ScreenProjectedUV = (ScreenProjected.xy / ScreenProjected.w) * float2(0.5f, -0.5f) + 0.5f;
+
+			ProjectedMinUV = min(ProjectedMinUV, ScreenProjectedUV);
+			ProjectedMaxUV = max(ProjectedMaxUV, ScreenProjectedUV);
+		}
+	}
+
+	float2 ProjectedSize = max(ProjectedMaxUV - ProjectedMinUV, 0.0f) * View.ViewSizeAndInvSize.xy;
+	return ProjectedSize;
+}
+
+void AddToAtlasTileBin(float2 ProjectedShadowSize, uint SlotIndex)
+{
+	float ProjectedShadowSizeMaxDim = ceil(max(ProjectedShadowSize.x, ProjectedShadowSize.y));
+	// Get the log2 of the next greater power of two.
+	uint RequestedShadowSizeLog2 = ProjectedShadowSizeMaxDim > 0.0f ? (uint)ceil(log2(ProjectedShadowSizeMaxDim)) : 0;
+	RequestedShadowSizeLog2 = clamp(RequestedShadowSizeLog2, MinAtlasTileResolutionLog2, (MinAtlasTileResolutionLog2 + NUM_ATLAS_TILE_BINS - 1));
+	uint BinIndex = RequestedShadowSizeLog2 - MinAtlasTileResolutionLog2;
+
+	// Append to bin
+	uint WriteIndex = 0;
+	InterlockedAdd(AtlasTileBinElementCounts[BinIndex], 1, WriteIndex);
+	// Note: Currently we have MAX_SLOT_COUNT == FHairStrandsDeepShadowData::MaxMacroGroupCount, 
+	// so by sizing the AtlasTileBins arrays to MAX_SLOT_COUNT we can't overflow. If we ever change that,
+	// we need to account for this here.
+	AtlasTileBins[BinIndex][WriteIndex] = SlotIndex;
+	
+	// Keep the actual desired resolution around
+	AtlasTileDesiredResolutions[SlotIndex] = max(MinAtlasTileResolution, (uint)ProjectedShadowSizeMaxDim);
+}
+
+void AllocateAtlasTiles()
+{
+	uint MaxAllowedTileResolution = 1u << (MinAtlasTileResolutionLog2 + NUM_ATLAS_TILE_BINS - 1);
+
+	// Allocating might fail when the requested tiles are too big, so we retry allocating with smaller tiles when that happens.
+	for (uint NumRetries = 0; NumRetries < NUM_ATLAS_TILE_BINS; ++NumRetries)
+	{
+		uint2 CurrentOffset = 0;
+		uint CurrentRowHeight = 0;
+		// Keep track of the actual maximum tile resolution requested in this allocation attempt
+		uint MaxActualTileResolution = 0;
+
+		// Loop over all bins starting with the one with the highest resolution
+		for (uint i = 0; i < NUM_ATLAS_TILE_BINS; ++i)
+		{
+			const uint BinIndex = NUM_ATLAS_TILE_BINS - 1 - i;
+			const uint NumElements = AtlasTileBinElementCounts[BinIndex];
+
+			// Skip empty bins
+			if (NumElements == 0)
+			{
+				continue;
+			}
+
+			uint BinResolution = 1u << (BinIndex + MinAtlasTileResolutionLog2);
+			BinResolution = min(BinResolution, MaxAllowedTileResolution); // Start clamping the resolution of larger tiles when previous allocations did not succeed
+			MaxActualTileResolution = max(MaxActualTileResolution, BinResolution);
+
+			// Allocate space for all elements in this bin
+			for (uint ElementIndex = 0; ElementIndex < NumElements; ++ElementIndex)
+			{
+				const uint SlotIndex = AtlasTileBins[BinIndex][ElementIndex];
+				const uint DesiredResolution = AtlasTileDesiredResolutions[SlotIndex];
+				const uint ActualResolution = min(DesiredResolution, BinResolution);
+
+				// Advance to the next row if we reached the right end of the atlas
+				if ((CurrentOffset.x + ActualResolution) > AtlasResolution.x)
+				{
+					CurrentOffset.x = 0;
+					CurrentOffset.y += CurrentRowHeight;
+					CurrentRowHeight = 0;
+				}
+
+				AtlasTileAllocations[SlotIndex] = uint4(CurrentOffset, ActualResolution.xx);
+				CurrentOffset.x += ActualResolution;
+				CurrentRowHeight = max(CurrentRowHeight, ActualResolution);
+			}
+		}
+
+		if ((CurrentOffset.y + CurrentRowHeight) <= AtlasResolution.y)
+		{
+			break;
+		}
+
+		// Half the maximum allowed tile resolution with each failed attempt.
+		// Note that we half the actual maximum requested resolution which allows us to skip certain resolutions where the corresponding bins are empty.
+		MaxAllowedTileResolution = MaxActualTileResolution / 2;
+	}
+}
+
+void ApplyAtlasTileScaleBias(float4x4 ShadowMatrix, uint SlotIndex, out float4x4 ModifiedTranslatedWorldToClip, out float4 AtlasScaleBias)
+{
+	uint4 TileOffsetSize = AtlasTileAllocations[SlotIndex];
+	float2 Scale = TileOffsetSize.zw * AtlasTexelSize;
+	float2 Bias = TileOffsetSize.xy * AtlasTexelSize;
+
+	AtlasScaleBias = float4(Scale, Bias);
+
+#if 1
+	Bias = Bias * float2(2.0f, -2.0f) - float2(1.0f, -1.0f) + float2(Scale.x, -Scale.y);
+
+	float4x4 ScaleBias = float4x4(
+		Scale.x, 0.0f, 0.0f, 0.0f,
+		0.0f, Scale.y, 0.0f, 0.0f,
+		0.0f, 0.0f, 1.0f, 0.0f,
+		Bias.x, Bias.y, 0.0f, 1.0f
+	);
+
+	ModifiedTranslatedWorldToClip = mul(ShadowMatrix, ScaleBias);
+#else
+	float4x4 ClipToUV = float4x4(
+		0.5f, 0.0f, 0.0f, 0.0f,
+		0.0f, -0.5f, 0.0f, 0.0f,
+		0.0f, 0.0f, 1.0f, 0.0f,
+		0.5f, 0.5f, 0.0f, 1.0f
+	);
+
+	float4x4 UVToClip = float4x4(
+		2.0f, 0.0f, 0.0f, 0.0f,
+		0.0f, -2.0f, 0.0f, 0.0f,
+		0.0f, 0.0f, 1.0f, 0.0f,
+		-1.0f, 1.0f, 0.0f, 1.0f
+	);
+	
+	float4x4 ScaleBias = float4x4(
+		Scale.x, 0.0f, 0.0f, 0.0f,
+		0.0f, Scale.y, 0.0f, 0.0f,
+		0.0f, 0.0f, 1.0f, 0.0f,
+		Bias.x, Bias.y, 0.0f, 1.0f
+	);
+
+	ModifiedTranslatedWorldToClip = mul(ShadowMatrix, mul(ClipToUV, mul(ScaleBias, UVToClip)));
+#endif
+}
+
 // This code assume we have less than 32 macro group (which fit into a single CU/SM)
 [numthreads(MAX_SLOT_COUNT, 1, 1)]
 void CreateViewInfo(uint2 DispatchThreadId : SV_DispatchThreadID)
 {
 	const uint SlotIndex = DispatchThreadId.x;
 
-	if (SlotIndex > SlotIndexCount)
+	// Clear LDS
+	{
+		UNROLL
+		for (uint WaveOffset = 0; WaveOffset < NUM_ATLAS_TILE_BINS; WaveOffset += MAX_SLOT_COUNT)
+		{
+			uint BinIndex = WaveOffset + SlotIndex;
+			if (BinIndex < NUM_ATLAS_TILE_BINS)
+			{
+				AtlasTileBinElementCounts[BinIndex] = 0;
+			}
+		}
+	}
+	GroupMemoryBarrierWithGroupSync();
+
+	if (SlotIndex >= SlotIndexCount)
 		return;
 
 	const uint MacroGroupId = MacroGroupIndices[SlotIndex].x;
@@ -140,10 +316,18 @@ void CreateViewInfo(uint2 DispatchThreadId : SV_DispatchThreadID)
 	FDeepShadowViewInfo ViewInfo;
 	ComputeTranslatedWorldToLightClip(ViewInfo.TranslatedWorldToClip, ViewInfo.MinRadiusAtDepth1, ViewInfo.ViewForward, TranslatedBound, SlotIndex);
 
-	FDeepShadowTransform Transform;
-	Transform.TranslatedWorldToClip = ViewInfo.TranslatedWorldToClip;
+	float2 ProjectedShadowSize = ComputeProjectedScreenSize(ViewInfo.TranslatedWorldToClip);
+	AddToAtlasTileBin(ProjectedShadowSize, SlotIndex);
+
+	GroupMemoryBarrierWithGroupSync();
+	if (SlotIndex == 0)
+	{
+		AllocateAtlasTiles();
+	}
+	GroupMemoryBarrierWithGroupSync();
+
+	ApplyAtlasTileScaleBias(ViewInfo.TranslatedWorldToClip, SlotIndex, ViewInfo.TranslatedWorldToClipScaledBiased, ViewInfo.AtlasScaleBias);
 
 	OutShadowViewInfoBuffer[SlotIndex] = ViewInfo;
-	OutShadowTranslatedWorldToLightTransformBuffer[SlotIndex] = Transform;
 }
 #endif
diff --git a/Engine/Shaders/Private/HairStrands/HairStrandsDeepShadowCommonStruct.ush b/Engine/Shaders/Private/HairStrands/HairStrandsDeepShadowCommonStruct.ush
--- a/Engine/Shaders/Private/HairStrands/HairStrandsDeepShadowCommonStruct.ush
+++ b/Engine/Shaders/Private/HairStrands/HairStrandsDeepShadowCommonStruct.ush
@@ -4,12 +4,9 @@
 
 struct FDeepShadowViewInfo
 {
+	float4x4 TranslatedWorldToClipScaledBiased; // Has the shadow atlas tile scale bias baked in
 	float4x4 TranslatedWorldToClip;
+	float4 AtlasScaleBias; // Scale/bias for transforming UV coordinates into the shadow atlas tile of this shadow view.
 	float3   ViewForward;
 	float    MinRadiusAtDepth1;
-};
-
-struct FDeepShadowTransform
-{
-	float4x4 TranslatedWorldToClip; // Store the transform into a struct to avoid the matrix to be loaded as transposed
 };
\ No newline at end of file
diff --git a/Engine/Shaders/Private/HairStrands/HairStrandsDeepShadowDebug.usf b/Engine/Shaders/Private/HairStrands/HairStrandsDeepShadowDebug.usf
--- a/Engine/Shaders/Private/HairStrands/HairStrandsDeepShadowDebug.usf
+++ b/Engine/Shaders/Private/HairStrands/HairStrandsDeepShadowDebug.usf
@@ -135,7 +135,7 @@ uint AllocatedSlotCount;
 uint MacroGroupCount;
 
 Buffer<int>	MacroGroupAABBBuffer;
-StructuredBuffer<FDeepShadowTransform> ShadowTranslatedWorldToLightTransformBuffer;
+StructuredBuffer<FDeepShadowViewInfo> ShadowViewInfoBuffer;
 
 [numthreads(1, 1, 1)]
 void MainCS(uint3 DispatchThreadId : SV_DispatchThreadID)
@@ -154,7 +154,7 @@ void MainCS(uint3 DispatchThreadId : SV_DispatchThreadID)
 	const float4 FrustumColor = float4(0.5f, 1, 0, 1);
 	for (uint SlotIt = 0; SlotIt < AllocatedSlotCount; ++SlotIt)
 	{
-		const float4x4 TranslatedWorldToClip = ShadowTranslatedWorldToLightTransformBuffer[SlotIt].TranslatedWorldToClip;
+		const float4x4 TranslatedWorldToClip = ShadowViewInfoBuffer[SlotIt].TranslatedWorldToClip;
 		const float4x4 ClipToTranslatedWorld = Inverse(TranslatedWorldToClip);
 
 		float4 P01 = float4(-1, -1, 0, 1);
diff --git a/Engine/Shaders/Private/HairStrands/HairStrandsDeepShadowMask.usf b/Engine/Shaders/Private/HairStrands/HairStrandsDeepShadowMask.usf
--- a/Engine/Shaders/Private/HairStrands/HairStrandsDeepShadowMask.usf
+++ b/Engine/Shaders/Private/HairStrands/HairStrandsDeepShadowMask.usf
@@ -35,7 +35,7 @@ float3 GetTranslatedWorldPosition(float2 UV, float SceneDepth)
 
 #if SHADER_SHADOWMASK_DEEPSHADOW
 
-StructuredBuffer<FDeepShadowTransform> DeepShadow_TranslatedWorldToLightTransformBuffer;
+StructuredBuffer<FDeepShadowViewInfo> DeepShadow_ViewInfoBuffer;
 uint				DeepShadow_bIsGPUDriven;
 
 float4x4			DeepShadow_CPUTranslatedWorldToLightTransform;
@@ -57,16 +57,19 @@ SamplerState 		ShadowSampler;
 // XY: UV of the shadow space
 //  Z: Normalized depth value in light clip space
 //  W: Positive if the position is valid, negative otherwise
-float4 ToLightPosition(float3 InTranslatedWorldPosition)
+float4 ToLightPosition(float3 InTranslatedWorldPosition, out float4 AtlasScaleBias)
 {
 	float4x4 TranslatedWorldToLightTransform = 0;
 	if (DeepShadow_bIsGPUDriven > 0)
 	{
-		TranslatedWorldToLightTransform = DeepShadow_TranslatedWorldToLightTransformBuffer[DeepShadow_SlotIndex].TranslatedWorldToClip;
+		FDeepShadowViewInfo ShadowViewInfo = DeepShadow_ViewInfoBuffer[DeepShadow_SlotIndex];
+		TranslatedWorldToLightTransform = ShadowViewInfo.TranslatedWorldToClip;
+		AtlasScaleBias = ShadowViewInfo.AtlasScaleBias;
 	}
 	else
 	{
 		TranslatedWorldToLightTransform = DeepShadow_CPUTranslatedWorldToLightTransform;
+		AtlasScaleBias = float4(1.0f, 1.0f, 0.0f, 0.0f);
 	}
 
 	float4 LightPos = mul(float4(InTranslatedWorldPosition,1), TranslatedWorldToLightTransform);
@@ -338,7 +341,8 @@ void MainPS(
 	const bool bIsHairPixel = HairStrands.HairCoverageTexture.Load(uint3(PixelCoord, 0)) >= 1;
 
 	const float3 TranslatedWorldPosition = GetTranslatedWorldPosition(UV, SceneDepth);
-	const float4 LightSpacePosition		 = ToLightPosition(TranslatedWorldPosition);
+	float4 AtlasScaleBias;
+	float4 LightSpacePosition		 = ToLightPosition(TranslatedWorldPosition, AtlasScaleBias);
 
 	// PCF or simple point sampler
 	float Visibility = 1;
@@ -348,6 +352,8 @@ void MainPS(
 		(LightSpacePosition.y >= 0 && LightSpacePosition.y <= 1);
 	if (!bIsHairPixel && bIsValid)
 	{
+		LightSpacePosition.xy = LightSpacePosition.xy * AtlasScaleBias.xy + AtlasScaleBias.zw;
+
 		uint2 ShadowAtlasResolution;
 		DeepShadow_FrontDepthTexture.GetDimensions(ShadowAtlasResolution.x, ShadowAtlasResolution.y);
 
diff --git a/Engine/Shaders/Private/HairStrands/HairStrandsDeepShadowVS.usf b/Engine/Shaders/Private/HairStrands/HairStrandsDeepShadowVS.usf
--- a/Engine/Shaders/Private/HairStrands/HairStrandsDeepShadowVS.usf
+++ b/Engine/Shaders/Private/HairStrands/HairStrandsDeepShadowVS.usf
@@ -53,7 +53,7 @@ void Main(
 	{
 		const uint AtlasSlotIndex = DeepRasterPass.AtlasSlotIndex;
 		FDeepShadowViewInfo DeepShadowViewInfo = DeepRasterPass.DeepShadowViewInfoBuffer[AtlasSlotIndex];
-		TranslatedWorldToClipMatrix = DeepShadowViewInfo.TranslatedWorldToClip;
+		TranslatedWorldToClipMatrix = DeepShadowViewInfo.TranslatedWorldToClipScaledBiased;
 		HairViewInfo.ViewForward	= DeepShadowViewInfo.ViewForward;
 		HairViewInfo.RadiusAtDepth1 = DeepShadowViewInfo.MinRadiusAtDepth1;
 	}
diff --git a/Engine/Shaders/Private/HairStrands/HairStrandsDeepTransmittanceMask.usf b/Engine/Shaders/Private/HairStrands/HairStrandsDeepTransmittanceMask.usf
--- a/Engine/Shaders/Private/HairStrands/HairStrandsDeepTransmittanceMask.usf
+++ b/Engine/Shaders/Private/HairStrands/HairStrandsDeepTransmittanceMask.usf
@@ -26,6 +26,7 @@ struct FTransmittanceSettings
 	int2	 DeepShadowResolution;
 	int2	 DeepShadowOffset;
 	float4x4 TranslatedWorldToLightTransform;
+	float4 AtlasScaleBias;
 	float3   LightDirection; // Direction from world position to the light
 	float3   TranslatedLightPosition;
 	bool	 bIsDirectional;
@@ -71,7 +72,7 @@ float3 GetTranslatedWorldPosition(float2 UV, float SceneDepth)
 
 #if SHADER_TRANSMITTANCE_DEEPSHADOW
 
-StructuredBuffer<FDeepShadowTransform> DeepShadow_TranslatedWorldToLightTransformBuffer;
+StructuredBuffer<FDeepShadowViewInfo> DeepShadow_ViewInfoBuffer;
 uint				DeepShadow_bIsGPUDriven;
 
 uint2				DeepShadow_Resolution;
@@ -104,7 +105,7 @@ FHairTransmittanceMask ComputeTransmittance(
 	FHairTransmittanceMask Out = InitHairTransmittanceMask();
 	{
 		// LightSpacePosition is 'TranslatedWorldPosition' in LightSpace (TexelCoord.XY / NDC Depth.Z)
-		const float3 LightSpacePosition = ToLightPosition(TranslatedWorldPosition, Settings.TranslatedWorldToLightTransform) * float3(Settings.DeepShadowResolution, 1) + float3(Settings.DeepShadowOffset, 0);
+		const float3 LightSpacePosition = ToLightPosition(TranslatedWorldPosition, Settings.TranslatedWorldToLightTransform) * float3(Settings.AtlasScaleBias.xy, 1) + float3(Settings.AtlasScaleBias.zw, 0);
 
 		float DepthBias = 0;
 		if (Settings.DeepShadowDepthBiasScale > 0)
@@ -185,13 +186,17 @@ void MainCS(uint2 DispatchThreadId : SV_DispatchThreadID)
 		const uint MacroGroupIndex = clamp(Sample.MacroGroupId, 0, MAX_HAIR_MACROGROUP_COUNT - 1);
 		const int3 SlotOffset_SlotIndex = DeepShadow_AtlasSlotOffsets_AtlasSlotIndex[MacroGroupIndex].xyz;
 		float4x4 TranslatedWorldToLightTransform = 0;
+		float4 AtlasScaleBias = float4(1.0f, 1.0f, 0.0f, 0.0f);
 		if (DeepShadow_bIsGPUDriven > 0)
 		{
-			TranslatedWorldToLightTransform = DeepShadow_TranslatedWorldToLightTransformBuffer[SlotOffset_SlotIndex.z].TranslatedWorldToClip;
+			FDeepShadowViewInfo ShadowViewInfo = DeepShadow_ViewInfoBuffer[SlotOffset_SlotIndex.z];
+			TranslatedWorldToLightTransform = ShadowViewInfo.TranslatedWorldToClip;
+			AtlasScaleBias = ShadowViewInfo.AtlasScaleBias * float4(DeepShadow_Resolution, DeepShadow_Resolution);
 		}	
 		else
 		{
 			TranslatedWorldToLightTransform = DeepShadow_CPUTranslatedWorldToLightTransforms[MacroGroupIndex];
+			AtlasScaleBias = float4(DeepShadow_Resolution, SlotOffset_SlotIndex.xy);
 		}
 		
 		const float SceneDepth = ConvertFromDeviceZ(Sample.Depth);
@@ -201,6 +206,7 @@ void MainCS(uint2 DispatchThreadId : SV_DispatchThreadID)
 		Settings.DeepShadowOffset = SlotOffset_SlotIndex.xy;
 		Settings.DeepShadowResolution = DeepShadow_Resolution;
 		Settings.TranslatedWorldToLightTransform = TranslatedWorldToLightTransform;
+		Settings.AtlasScaleBias = AtlasScaleBias;
 		Settings.bIsDirectional = IsLightDirectional(TranslatedLightPosition_LightDirection);
 		Settings.LightDirection = GetTranslatedLightDirection(TranslatedLightPosition_LightDirection, TranslatedWorldPosition);
 		Settings.TranslatedLightPosition = GetTranslatedLightPosition(TranslatedLightPosition_LightDirection, TranslatedWorldPosition);
diff --git a/Engine/Shaders/Private/HairStrands/HairStrandsVoxelDepthInjection.usf b/Engine/Shaders/Private/HairStrands/HairStrandsVoxelDepthInjection.usf
--- a/Engine/Shaders/Private/HairStrands/HairStrandsVoxelDepthInjection.usf
+++ b/Engine/Shaders/Private/HairStrands/HairStrandsVoxelDepthInjection.usf
@@ -98,7 +98,7 @@ void MainVS(
 	if (bIsGPUDriven)
 	{
 		FDeepShadowViewInfo DeepShadowViewInfo = DeepShadowViewInfoBuffer[AtlasSlotIndex];
-		TranslatedWorldToClipMatrix = DeepShadowViewInfo.TranslatedWorldToClip;
+		TranslatedWorldToClipMatrix = DeepShadowViewInfo.TranslatedWorldToClipScaledBiased;
 	}
 
 	OutTranslatedWorldPosition = TranslatedWorldPosition;
@@ -165,7 +165,7 @@ void MainPS(
 		if (bIsGPUDriven)
 		{
 			FDeepShadowViewInfo DeepShadowViewInfo = DeepShadowViewInfoBuffer[AtlasSlotIndex];
-			TranslatedWorldToClipMatrix = DeepShadowViewInfo.TranslatedWorldToClip;
+			TranslatedWorldToClipMatrix = DeepShadowViewInfo.TranslatedWorldToClipScaledBiased;
 		}
 		const float3 HitP = TranslatedWP0 + normalize(TranslatedWP1 - TranslatedWP0) * TraversalResult.HitT;
 		float4 ClipP = mul(float4(HitP, 1), TranslatedWorldToClipMatrix);
diff --git a/Engine/Source/Runtime/Renderer/Private/HairStrands/HairStrandsData.h b/Engine/Source/Runtime/Renderer/Private/HairStrands/HairStrandsData.h
--- a/Engine/Source/Runtime/Renderer/Private/HairStrands/HairStrandsData.h
+++ b/Engine/Source/Runtime/Renderer/Private/HairStrands/HairStrandsData.h
@@ -255,7 +255,7 @@ struct FHairStrandsDeepShadowResources
 
 	FRDGTextureRef DepthAtlasTexture = nullptr;
 	FRDGTextureRef LayersAtlasTexture = nullptr;
-	FRDGBufferRef DeepShadowTranslatedWorldToLightTransforms = nullptr;
+	FRDGBufferRef DeepShadowViewInfoBuffer = nullptr;
 };
 
 ////////////////////////////////////////////////////////////////////////////////////
diff --git a/Engine/Source/Runtime/Renderer/Private/HairStrands/HairStrandsDebug.cpp b/Engine/Source/Runtime/Renderer/Private/HairStrands/HairStrandsDebug.cpp
--- a/Engine/Source/Runtime/Renderer/Private/HairStrands/HairStrandsDebug.cpp
+++ b/Engine/Source/Runtime/Renderer/Private/HairStrands/HairStrandsDebug.cpp
@@ -662,7 +662,7 @@ class FDeepShadowInfoCS : public FGlobalShader
 		SHADER_PARAMETER(uint32, MacroGroupCount)
 		SHADER_PARAMETER_STRUCT_REF(FViewUniformShaderParameters, ViewUniformBuffer)
 		SHADER_PARAMETER_RDG_BUFFER_SRV(Buffer, MacroGroupAABBBuffer)
-		SHADER_PARAMETER_RDG_BUFFER_SRV(StructuredBuffer, ShadowTranslatedWorldToLightTransformBuffer)
+		SHADER_PARAMETER_RDG_BUFFER_SRV(StructuredBuffer, ShadowViewInfoBuffer)
 		SHADER_PARAMETER_RDG_TEXTURE_UAV(RWTexture2D, OutputTexture)
 		END_SHADER_PARAMETER_STRUCT()
 
@@ -704,7 +704,7 @@ static void AddDeepShadowInfoPass(
 	Parameters->MacroGroupCount = MacroGroupResources.MacroGroupCount;
 	Parameters->SceneTextures = SceneTextures;
 	Parameters->MacroGroupAABBBuffer = GraphBuilder.CreateSRV(MacroGroupResources.MacroGroupAABBsBuffer, PF_R32_SINT);
-	Parameters->ShadowTranslatedWorldToLightTransformBuffer = GraphBuilder.CreateSRV(DeepShadowResources.DeepShadowTranslatedWorldToLightTransforms);
+	Parameters->ShadowViewInfoBuffer = GraphBuilder.CreateSRV(DeepShadowResources.DeepShadowViewInfoBuffer);
 	ShaderPrint::SetParameters(GraphBuilder, View.ShaderPrintData, Parameters->ShaderPrintParameters);
 	Parameters->OutputTexture = GraphBuilder.CreateUAV(OutputTexture);
 
diff --git a/Engine/Source/Runtime/Renderer/Private/HairStrands/HairStrandsDeepShadow.cpp b/Engine/Source/Runtime/Renderer/Private/HairStrands/HairStrandsDeepShadow.cpp
--- a/Engine/Source/Runtime/Renderer/Private/HairStrands/HairStrandsDeepShadow.cpp
+++ b/Engine/Source/Runtime/Renderer/Private/HairStrands/HairStrandsDeepShadow.cpp
@@ -12,6 +12,9 @@
 static int32 GDeepShadowResolution = 2048;
 static FAutoConsoleVariableRef CVarDeepShadowResolution(TEXT("r.HairStrands.DeepShadow.Resolution"), GDeepShadowResolution, TEXT("Shadow resolution for Deep Opacity Map rendering. (default = 2048)"));
 
+static int32 GDeepShadowMinResolution = 64;
+static FAutoConsoleVariableRef CVarDeepShadowMinResolution(TEXT("r.HairStrands.DeepShadow.MinResolution"), GDeepShadowMinResolution, TEXT("Minimum shadow resolution for shadow atlas tiles for Deep Opacity Map rendering. (default = 64)"));
+
 static int32 GDeepShadowGPUDriven = 1;
 static FAutoConsoleVariableRef CVarDeepShadowGPUDriven(TEXT("r.HairStrands.DeepShadow.GPUDriven"), GDeepShadowGPUDriven, TEXT("Enable deep shadow to be driven by GPU bounding box, rather CPU ones. This allows more robust behavior"));
 
@@ -192,9 +195,13 @@ class FDeepShadowCreateViewInfoCS : public FGlobalShader
 		SHADER_PARAMETER(float, AABBScale)
 		SHADER_PARAMETER(float, MaxHafFovInRad)
 
+		SHADER_PARAMETER(FUintVector2, AtlasResolution)
+		SHADER_PARAMETER(FVector2f, AtlasTexelSize)
+		SHADER_PARAMETER(uint32, MinAtlasTileResolution)
+		SHADER_PARAMETER(uint32, MinAtlasTileResolutionLog2)
+
 		SHADER_PARAMETER_RDG_BUFFER_SRV(Buffer<int>, MacroGroupAABBBuffer)
 		SHADER_PARAMETER_RDG_BUFFER_UAV(RWStructuredBuffer<FDeepShadowViewInfo>, OutShadowViewInfoBuffer)
-		SHADER_PARAMETER_RDG_BUFFER_UAV(RWStructuredBuffer<float4x4>, OutShadowTranslatedWorldToLightTransformBuffer)
 		SHADER_PARAMETER_STRUCT_REF(FViewUniformShaderParameters, ViewUniformBuffer)
 
 	END_SHADER_PARAMETER_STRUCT()
@@ -276,7 +283,13 @@ void RenderHairStrandsDeepShadows(
 		const uint32 AtlasSlotX = FGenericPlatformMath::CeilToInt(FMath::Sqrt(static_cast<float>(DOMSlotCount)));
 		const FIntPoint AtlasSlotDimension(AtlasSlotX, AtlasSlotX == DOMSlotCount ? 1 : AtlasSlotX);
 		const FIntPoint AtlasSlotResolution(GDeepShadowResolution, GDeepShadowResolution);
-		const FIntPoint AtlasResolution(AtlasSlotResolution.X * AtlasSlotDimension.X, AtlasSlotResolution.Y * AtlasSlotDimension.Y);
+		FIntPoint AtlasResolution(AtlasSlotResolution.X * AtlasSlotDimension.X, AtlasSlotResolution.Y * AtlasSlotDimension.Y);
+
+		const bool bDeepShadowGPUDriven = GDeepShadowGPUDriven > 0;
+		if (bDeepShadowGPUDriven)
+		{
+			AtlasResolution = FIntPoint(GDeepShadowResolution, GDeepShadowResolution);
+		}
 
 		
 		DeepShadowResources.TotalAtlasSlotCount = 0;
@@ -325,9 +338,9 @@ void RenderHairStrandsDeepShadows(
 				DomData.LayerDistribution = LightProxy->GetDeepShadowLayerDistribution();
 				DomData.bIsLightDirectional = bIsDirectional;
 				DomData.LightId = LightInfo->Id;
-				DomData.ShadowResolution = AtlasSlotResolution;
+				DomData.ShadowResolution = bDeepShadowGPUDriven ? AtlasResolution.X : AtlasSlotResolution;
 				DomData.Bounds = MacroGroupBounds;
-				DomData.AtlasRect = FIntRect(AtlasRectOffset, AtlasRectOffset + AtlasSlotResolution);
+				DomData.AtlasRect = bDeepShadowGPUDriven ? FIntRect(0, 0, AtlasResolution.X, AtlasResolution.Y) : FIntRect(AtlasRectOffset, AtlasRectOffset + AtlasSlotResolution);
 				DomData.MacroGroupId = MacroGroup.MacroGroupId;
 				DomData.CPU_MinStrandRadiusAtDepth1 = MinStrandRadiusAtDepth1;
 				DomData.AtlasSlotIndex = TotalAtlasSlotIndex;
@@ -341,11 +354,10 @@ void RenderHairStrandsDeepShadows(
 		DeepShadowResources.TotalAtlasSlotCount = TotalAtlasSlotIndex;
 		DeepShadowResources.AtlasSlotResolution = AtlasSlotResolution;
 
-		FRDGBufferRef DeepShadowViewInfoBuffer = GraphBuilder.CreateBuffer(FRDGBufferDesc::CreateStructuredDesc(20 * sizeof(float), FMath::Max(1u, TotalAtlasSlotIndex)), TEXT("Hair.DeepShadowViewInfo"));
-		FRDGBufferRef DeepShadowTranslatedWorldToLightBuffer = GraphBuilder.CreateBuffer(FRDGBufferDesc::CreateStructuredDesc(16 * sizeof(float), FMath::Max(1u, TotalAtlasSlotIndex)), TEXT("Hair.DeepShadowTranslatedWorldToLightTransform"));
+		FRDGBufferRef DeepShadowViewInfoBuffer = GraphBuilder.CreateBuffer(FRDGBufferDesc::CreateStructuredDesc((16 + 16 + 4 + 3 + 1) * sizeof(float), FMath::Max(1u, TotalAtlasSlotIndex)), TEXT("Hair.DeepShadowViewInfo"));
 		FRDGBufferSRVRef DeepShadowViewInfoBufferSRV = GraphBuilder.CreateSRV(DeepShadowViewInfoBuffer);
 
-		DeepShadowResources.bIsGPUDriven = GDeepShadowGPUDriven > 0;
+		DeepShadowResources.bIsGPUDriven = bDeepShadowGPUDriven;
 		{
 			check(TotalAtlasSlotIndex < FHairStrandsDeepShadowResources::MaxAtlasSlotCount);
 
@@ -369,11 +381,14 @@ void RenderHairStrandsDeepShadows(
 			Parameters->MacroGroupCount = MacroGroupDatas.Num();
 			Parameters->MacroGroupAABBBuffer = GraphBuilder.CreateSRV(MacroGroupResources.MacroGroupAABBsBuffer, PF_R32_SINT);
 			Parameters->OutShadowViewInfoBuffer = GraphBuilder.CreateUAV(DeepShadowViewInfoBuffer);
-			Parameters->OutShadowTranslatedWorldToLightTransformBuffer = GraphBuilder.CreateUAV(DeepShadowTranslatedWorldToLightBuffer);
 
 			Parameters->MaxHafFovInRad = 0.5f * FMath::DegreesToRadians(GetDeepShadowMaxFovAngle());
 			Parameters->AABBScale = GetDeepShadowAABBScale();
 			Parameters->RasterizationScale = GetDeepShadowRasterizationScale();
+			Parameters->AtlasResolution = FUintVector2(AtlasResolution.X, AtlasResolution.Y);
+			Parameters->AtlasTexelSize = FVector2f(1.0f / AtlasResolution.X, 1.0f / AtlasResolution.Y);
+			Parameters->MinAtlasTileResolutionLog2 = FMath::FloorLog2(FMath::Clamp(GDeepShadowMinResolution, 16u, GDeepShadowResolution)); // A shadow map resolution of less than 16x16 is highly unlikely to ever be needed.
+			Parameters->MinAtlasTileResolution = 1u << Parameters->MinAtlasTileResolutionLog2;
 			Parameters->ViewUniformBuffer = View.ViewUniformBuffer;
 
 			// Currently support only 32 instance group at max
@@ -509,6 +524,6 @@ void RenderHairStrandsDeepShadows(
 		}
 		DeepShadowResources.DepthAtlasTexture = FrontDepthAtlasTexture;
 		DeepShadowResources.LayersAtlasTexture = DeepShadowLayersAtlasTexture;
-		DeepShadowResources.DeepShadowTranslatedWorldToLightTransforms = DeepShadowTranslatedWorldToLightBuffer;
+		DeepShadowResources.DeepShadowViewInfoBuffer = DeepShadowViewInfoBuffer;
 	}
 }
diff --git a/Engine/Source/Runtime/Renderer/Private/HairStrands/HairStrandsTransmittance.cpp b/Engine/Source/Runtime/Renderer/Private/HairStrands/HairStrandsTransmittance.cpp
--- a/Engine/Source/Runtime/Renderer/Private/HairStrands/HairStrandsTransmittance.cpp
+++ b/Engine/Source/Runtime/Renderer/Private/HairStrands/HairStrandsTransmittance.cpp
@@ -333,7 +333,7 @@ class FHairStrandsDeepShadowTransmittanceMaskCS : public FGlobalShader
 		SHADER_PARAMETER_RDG_TEXTURE(Texture2D, RayMarchMaskTexture)
 		SHADER_PARAMETER_RDG_TEXTURE(Texture2D, DeepShadow_FrontDepthTexture)
 		SHADER_PARAMETER_RDG_TEXTURE(Texture2D, DeepShadow_DomTexture)
-		SHADER_PARAMETER_RDG_BUFFER_SRV(StructuredBuffer, DeepShadow_TranslatedWorldToLightTransformBuffer)
+		SHADER_PARAMETER_RDG_BUFFER_SRV(StructuredBuffer, DeepShadow_ViewInfoBuffer)
 		SHADER_PARAMETER_RDG_BUFFER_UAV(RWBuffer, OutTransmittanceMask)
 
 		RDG_BUFFER_ACCESS(IndirectArgsBuffer, ERHIAccess::IndirectArgs)
@@ -359,7 +359,7 @@ struct FHairStrandsDeepShadowTransmittanceLightParams : FHairStrandsTransmittanc
 {
 	FIntVector4 DeepShadow_AtlasSlotOffsets_AtlasSlotIndex[FHairStrandsDeepShadowData::MaxMacroGroupCount];
 	FMatrix DeepShadow_CPUTranslatedWorldToLightTransforms[FHairStrandsDeepShadowData::MaxMacroGroupCount];
-	FRDGBufferSRVRef DeepShadow_TranslatedWorldToLightTransformBuffer = nullptr;
+	FRDGBufferSRVRef DeepShadow_ViewInfoBuffer = nullptr;
 	FIntPoint DeepShadow_Resolution = FIntPoint(0, 0);
 	bool DeepShadow_bIsGPUDriven = false;
 	FVector4f DeepShadow_LayerDepths = FVector4f(0, 0, 0, 0);
@@ -405,7 +405,7 @@ static FRDGBufferRef AddHairStrandsDeepShadowTransmittanceMaskPass(
 	Parameters->IndirectArgsBuffer = IndirectArgsBuffer;
 	Parameters->HairStrands = HairStrands::BindHairStrandsViewUniformParameters(View);
 	Parameters->DeepShadow_bIsGPUDriven = Params.DeepShadow_bIsGPUDriven ? 1 : 0;;
-	Parameters->DeepShadow_TranslatedWorldToLightTransformBuffer = Params.DeepShadow_TranslatedWorldToLightTransformBuffer;
+	Parameters->DeepShadow_ViewInfoBuffer = Params.DeepShadow_ViewInfoBuffer;
 	Parameters->RayMarchMaskTexture = ScreenShadowMaskSubPixelTexture ? ScreenShadowMaskSubPixelTexture : GraphBuilder.RegisterExternalTexture(GSystemTextures.WhiteDummy);
 	Parameters->ShadowChannelMask = FVector4f(0, 0, 0, 0);
 	Parameters->ShadowChannelMask[FMath::Clamp<uint32>(Params.ShadowChannelMask, 0, 3)] = 1.0f;
@@ -608,7 +608,7 @@ class FHairStrandsDeepShadowMaskPS : public FGlobalShader
 
 		SHADER_PARAMETER_RDG_TEXTURE(Texture2D, RayMarchMaskTexture)
 
-		SHADER_PARAMETER_RDG_BUFFER_SRV(StructuredBuffer, DeepShadow_TranslatedWorldToLightTransformBuffer)
+		SHADER_PARAMETER_RDG_BUFFER_SRV(StructuredBuffer, DeepShadow_ViewInfoBuffer)
 		SHADER_PARAMETER_RDG_TEXTURE(Texture2D, DeepShadow_FrontDepthTexture)
 		SHADER_PARAMETER_RDG_TEXTURE(Texture2D, DeepShadow_DomTexture)
 		SHADER_PARAMETER_RDG_TEXTURE(Texture2D, SceneDepthTexture)
@@ -637,7 +637,7 @@ struct FHairStrandsDeepShadowParams
 {
 	bool			bIsWholeSceneLight = false;
 	uint32			OutputChannel = ~0;
-	FRDGBufferSRVRef DeepShadow_TranslatedWorldToLightTransformBuffer = nullptr;
+	FRDGBufferSRVRef DeepShadow_ViewInfoBuffer = nullptr;
 	FMatrix			DeepShadow_CPUTranslatedWorldToLightTransform;
 	FIntRect		DeepShadow_AtlasRect;
 	FRDGTextureRef	DeepShadow_FrontDepthTexture = nullptr;
@@ -676,7 +676,7 @@ static void AddHairStrandsDeepShadowMaskPass(
 	Parameters->DeepShadow_SlotIndex = Params.DeepShadow_AtlasSlotIndex;
 	Parameters->DeepShadow_SlotOffset = FIntPoint(Params.DeepShadow_AtlasRect.Min.X, Params.DeepShadow_AtlasRect.Min.Y);
 	Parameters->DeepShadow_SlotResolution = FIntPoint(Params.DeepShadow_AtlasRect.Max.X - Params.DeepShadow_AtlasRect.Min.X, Params.DeepShadow_AtlasRect.Max.Y - Params.DeepShadow_AtlasRect.Min.Y);
-	Parameters->DeepShadow_TranslatedWorldToLightTransformBuffer = Params.DeepShadow_TranslatedWorldToLightTransformBuffer;
+	Parameters->DeepShadow_ViewInfoBuffer = Params.DeepShadow_ViewInfoBuffer;
 	Parameters->DeepShadow_bIsGPUDriven = Params.DeepShadow_bIsGPUDriven ? 1 : 0;;
 	Parameters->HairStrands = HairStrands::BindHairStrandsViewUniformParameters(View);
 	if (ShaderPrint::IsValid(View.ShaderPrintData))
@@ -785,16 +785,16 @@ static FHairStrandsTransmittanceMaskData InternalRenderHairStrandsTransmittanceM
 		memset(Params.DeepShadow_AtlasSlotOffsets_AtlasSlotIndex, 0, sizeof(Params.DeepShadow_AtlasSlotOffsets_AtlasSlotIndex));
 		memset(Params.DeepShadow_CPUTranslatedWorldToLightTransforms, 0, sizeof(Params.DeepShadow_CPUTranslatedWorldToLightTransforms));
 
-		FRDGBufferSRVRef DeepShadow_TranslatedWorldToLightTransformBufferSRV = nullptr;
+		FRDGBufferSRVRef DeepShadow_ViewInfoBufferSRV = nullptr;
 		for (const FHairStrandsMacroGroupData& MacroGroupData : MacroGroupDatas)
 		{
 			for (const FHairStrandsDeepShadowData& DeepShadowData : MacroGroupData.DeepShadowDatas)
 			{
 				if (DeepShadowData.LightId == LightSceneInfo->Id)
 				{
-					if (DeepShadow_TranslatedWorldToLightTransformBufferSRV == nullptr)
+					if (DeepShadow_ViewInfoBufferSRV == nullptr)
 					{
-						DeepShadow_TranslatedWorldToLightTransformBufferSRV = GraphBuilder.CreateSRV(DeepShadowResources.DeepShadowTranslatedWorldToLightTransforms);
+						DeepShadow_ViewInfoBufferSRV = GraphBuilder.CreateSRV(DeepShadowResources.DeepShadowViewInfoBuffer);
 					}
 
 					bHasFoundLight = true;
@@ -808,7 +808,7 @@ static FHairStrandsTransmittanceMaskData InternalRenderHairStrandsTransmittanceM
 					Params.DeepShadow_LayerDepths = ComputeDeepShadowLayerDepths(DeepShadowData.LayerDistribution);
 					Params.DeepShadow_AtlasSlotOffsets_AtlasSlotIndex[DeepShadowData.MacroGroupId] = FIntVector4(DeepShadowData.AtlasRect.Min.X, DeepShadowData.AtlasRect.Min.Y, DeepShadowData.AtlasSlotIndex, 0);
 					Params.DeepShadow_CPUTranslatedWorldToLightTransforms[DeepShadowData.MacroGroupId] = DeepShadowData.CPU_TranslatedWorldToLightTransform;
-					Params.DeepShadow_TranslatedWorldToLightTransformBuffer = DeepShadow_TranslatedWorldToLightTransformBufferSRV;
+					Params.DeepShadow_ViewInfoBuffer = DeepShadow_ViewInfoBufferSRV;
 					Params.DeepShadow_bIsGPUDriven = DeepShadowResources.bIsGPUDriven;
 				}
 			}
@@ -939,8 +939,8 @@ static void InternalRenderHairStrandsShadowMask(
 	bool bHasDeepShadow = false;
 	if (!IsHairStrandsForVoxelTransmittanceAndShadowEnable())
 	{
-		FRDGBufferRef DeepShadow_TranslatedWorldToLightTransformBuffer = nullptr;
-		FRDGBufferSRVRef DeepShadow_TranslatedWorldToLightTransformBufferSRV = nullptr;
+		FRDGBufferRef DeepShadow_ViewInfoBuffer = nullptr;
+		FRDGBufferSRVRef DeepShadow_ViewInfoBufferSRV = nullptr;
 
 		for (const FHairStrandsMacroGroupData& MacroGroupData : InMacroGroupDatas)
 		{
@@ -949,10 +949,10 @@ static void InternalRenderHairStrandsShadowMask(
 				if (DomData.LightId != LightSceneInfo->Id)
 					continue;
 
-				if (DeepShadow_TranslatedWorldToLightTransformBuffer == nullptr)
+				if (DeepShadow_ViewInfoBuffer == nullptr)
 				{
-					DeepShadow_TranslatedWorldToLightTransformBuffer = DeepShadowResources.DeepShadowTranslatedWorldToLightTransforms;
-					DeepShadow_TranslatedWorldToLightTransformBufferSRV = GraphBuilder.CreateSRV(DeepShadow_TranslatedWorldToLightTransformBuffer);
+					DeepShadow_ViewInfoBuffer = DeepShadowResources.DeepShadowViewInfoBuffer;
+					DeepShadow_ViewInfoBufferSRV = GraphBuilder.CreateSRV(DeepShadow_ViewInfoBuffer);
 				}
 
 				bHasDeepShadow = true;
@@ -960,7 +960,7 @@ static void InternalRenderHairStrandsShadowMask(
 
 				FHairStrandsDeepShadowParams Params;
 				Params.DeepShadow_AtlasSlotIndex = DomData.AtlasSlotIndex;
-				Params.DeepShadow_TranslatedWorldToLightTransformBuffer = DeepShadow_TranslatedWorldToLightTransformBufferSRV;
+				Params.DeepShadow_ViewInfoBuffer = DeepShadow_ViewInfoBufferSRV;
 				Params.DeepShadow_bIsGPUDriven = DeepShadowResources.bIsGPUDriven ? 1 : 0;
 				Params.DeepShadow_CPUTranslatedWorldToLightTransform = DomData.CPU_TranslatedWorldToLightTransform;
 				Params.DeepShadow_AtlasRect = DomData.AtlasRect;
