Fix for thread sync compiler error in HairStrandsDeepShadowAllocation.ush

#rnx

diff --git a/Engine/Shaders/Private/HairStrands/HairStrandsDeepShadowAllocation.usf b/Engine/Shaders/Private/HairStrands/HairStrandsDeepShadowAllocation.usf
--- a/Engine/Shaders/Private/HairStrands/HairStrandsDeepShadowAllocation.usf
+++ b/Engine/Shaders/Private/HairStrands/HairStrandsDeepShadowAllocation.usf
@@ -302,22 +302,22 @@ void CreateViewInfo(uint2 DispatchThreadId : SV_DispatchThreadID)
 	}
 	GroupMemoryBarrierWithGroupSync();
 
-	if (SlotIndex >= SlotIndexCount)
-		return;
-
-	const uint MacroGroupId = MacroGroupIndices[SlotIndex].x;
-
-	FHairAABB TranslatedBound = InitHairAABB();
-	if (MacroGroupId < MacroGroupCount)
+	FDeepShadowViewInfo ViewInfo;
+	if (SlotIndex < SlotIndexCount)
 	{
-		TranslatedBound = ReadHairAABB(MacroGroupId, MacroGroupAABBBuffer);
-	}
+		const uint MacroGroupId = MacroGroupIndices[SlotIndex].x;
 
-	FDeepShadowViewInfo ViewInfo;
-	ComputeTranslatedWorldToLightClip(ViewInfo.TranslatedWorldToClip, ViewInfo.MinRadiusAtDepth1, ViewInfo.ViewForward, TranslatedBound, SlotIndex);
+		FHairAABB TranslatedBound = InitHairAABB();
+		if (MacroGroupId < MacroGroupCount)
+		{
+			TranslatedBound = ReadHairAABB(MacroGroupId, MacroGroupAABBBuffer);
+		}
+
+		ComputeTranslatedWorldToLightClip(ViewInfo.TranslatedWorldToClip, ViewInfo.MinRadiusAtDepth1, ViewInfo.ViewForward, TranslatedBound, SlotIndex);
 
-	float2 ProjectedShadowSize = ComputeProjectedScreenSize(ViewInfo.TranslatedWorldToClip);
-	AddToAtlasTileBin(ProjectedShadowSize, SlotIndex);
+		float2 ProjectedShadowSize = ComputeProjectedScreenSize(ViewInfo.TranslatedWorldToClip);
+		AddToAtlasTileBin(ProjectedShadowSize, SlotIndex);
+	}
 
 	GroupMemoryBarrierWithGroupSync();
 	if (SlotIndex == 0)
@@ -326,8 +326,10 @@ void CreateViewInfo(uint2 DispatchThreadId : SV_DispatchThreadID)
 	}
 	GroupMemoryBarrierWithGroupSync();
 
-	ApplyAtlasTileScaleBias(ViewInfo.TranslatedWorldToClip, SlotIndex, ViewInfo.TranslatedWorldToClipScaledBiased, ViewInfo.AtlasScaleBias);
-
-	OutShadowViewInfoBuffer[SlotIndex] = ViewInfo;
+	if (SlotIndex < SlotIndexCount)
+	{
+		ApplyAtlasTileScaleBias(ViewInfo.TranslatedWorldToClip, SlotIndex, ViewInfo.TranslatedWorldToClipScaledBiased, ViewInfo.AtlasScaleBias);
+		OutShadowViewInfoBuffer[SlotIndex] = ViewInfo;
+	}
 }
 #endif
