Implemented support for cloud shadows on SingleLayerWater.

#rnx

diff --git a/Engine/Shaders/Private/BasePassCommon.ush b/Engine/Shaders/Private/BasePassCommon.ush
--- a/Engine/Shaders/Private/BasePassCommon.ush
+++ b/Engine/Shaders/Private/BasePassCommon.ush
@@ -56,8 +56,13 @@
 #define TRANSLUCENCY_PERVERTEX_LIGHTING_VOLUME (!FORWARD_SHADING && (TRANSLUCENCY_LIGHTING_VOLUMETRIC_PERVERTEX_DIRECTIONAL || TRANSLUCENCY_LIGHTING_VOLUMETRIC_PERVERTEX_NONDIRECTIONAL))
 #define TRANSLUCENCY_PERVERTEX_FORWARD_SHADING (FORWARD_SHADING && (TRANSLUCENCY_LIGHTING_VOLUMETRIC_PERVERTEX_DIRECTIONAL || TRANSLUCENCY_LIGHTING_VOLUMETRIC_PERVERTEX_NONDIRECTIONAL))
 
-// We consider that the cloud shadow is always enabled but the interpolator is only used when render lit translucent with forward shading. We only support per vertex cloud shadow as of today.
+// We consider that the cloud shadow is always enabled but the interpolator is only used when render water or lit translucent with forward shading. We only support per vertex cloud shadow as of today.
+#if MATERIAL_SHADINGMODEL_SINGLELAYERWATER
+// Checking for !SINGLE_LAYER_WATER_SIMPLE_FORWARD disables cloud shadows on low end platforms.
+#define NEEDS_BASEPASS_CLOUD_SHADOW_INTERPOLATOR (SUPPORT_CLOUD_SHADOW_ON_SINGLE_LAYER_WATER && !SINGLE_LAYER_WATER_SIMPLE_FORWARD)
+#else
 #define NEEDS_BASEPASS_CLOUD_SHADOW_INTERPOLATOR (SUPPORT_CLOUD_SHADOW_ON_FORWARD_LIT_TRANSLUCENT && IS_MATERIAL_TRANSLUCENT_AND_LIT && !TRANSLUCENCY_PERVERTEX_FORWARD_SHADING && !TRANSLUCENCY_PERVERTEX_LIGHTING_VOLUME)
+#endif
 
 struct FSharedBasePassInterpolants
 {
diff --git a/Engine/Shaders/Private/BasePassPixelShader.usf b/Engine/Shaders/Private/BasePassPixelShader.usf
--- a/Engine/Shaders/Private/BasePassPixelShader.usf
+++ b/Engine/Shaders/Private/BasePassPixelShader.usf
@@ -150,6 +150,8 @@
 #include "ThinTranslucentCommon.ush"
 
 #if TRANSLUCENCY_LIGHTING_SURFACE_LIGHTINGVOLUME || TRANSLUCENCY_LIGHTING_SURFACE_FORWARDSHADING || FORWARD_SHADING || MATERIAL_SHADINGMODEL_SINGLELAYERWATER || STRATA_TRANSLUCENT_FORWARD || STRATA_FORWARD_SHADING
+// Enables applying a cloud shadow factor to lighting calculations in ForwardLightingCommon.ush.
+#define ENABLE_FORWARD_CLOUD_SHADOW NEEDS_BASEPASS_CLOUD_SHADOW_INTERPOLATOR
 #include "ForwardLightingCommon.ush"
 #endif
 
@@ -1226,6 +1228,7 @@ void FPixelShaderInOut_MainPS(
 		float IndirectOcclusion = 1.0f;
 		float2 NearestResolvedDepthScreenUV = 0;
 		float DirectionalLightShadow = 1.0f;
+		float DirectionalLightCloudShadow = 1.0f;
 
 		#if FORWARD_SHADING && (MATERIALBLENDING_SOLID || MATERIALBLENDING_MASKED)
 			float2 NDC = MaterialParameters.ScreenPosition.xy / MaterialParameters.ScreenPosition.w;
@@ -1258,7 +1261,6 @@ void FPixelShaderInOut_MainPS(
 			#if FEATURE_LEVEL >= FEATURE_LEVEL_SM5
 				GridIndex = ComputeLightGridCellIndex((uint2)((MaterialParameters.SvPosition.xy - ResolvedView.ViewRectMin.xy) * View.LightProbeSizeRatioAndInvSizeRatio.zw), MaterialParameters.SvPosition.w, EyeIndex);
 
-				float DirectionalLightCloudShadow = 1.0f;
 				#if NEEDS_BASEPASS_CLOUD_SHADOW_INTERPOLATOR
 					DirectionalLightCloudShadow = BasePassInterpolants.VertexCloudShadow;
 				#endif
@@ -1338,6 +1340,7 @@ void FPixelShaderInOut_MainPS(
 #else // !STRATA_ENABLED
 
 	float DirectionalLightShadow = 1.0f; 
+	float DirectionalLightCloudShadow = 1.0f;
 	float IndirectOcclusion = 1.0f;
 					
 #endif // !STRATA_ENABLED
@@ -1476,7 +1479,7 @@ void FPixelShaderInOut_MainPS(
 		const float4 NullDistortionParams = 1.0f;
 		WaterVolumeLightingOutput WaterLighting = EvaluateWaterVolumeLighting(
 			MaterialParameters, PixelMaterialInputs, ResolvedView,
-			DirectionalLightShadow,
+			DirectionalLightShadow * DirectionalLightCloudShadow,
 			SingleLayerWater.SceneDepthWithoutSingleLayerWaterTexture, SingleLayerWaterSceneDepthSampler, // Scene depth texture
 			SingleLayerWater.SceneWithoutSingleLayerWaterTextureSize,
 			SingleLayerWater.SceneWithoutSingleLayerWaterInvTextureSize,
@@ -1492,7 +1495,7 @@ void FPixelShaderInOut_MainPS(
 #else
 		Color += EvaluateWaterVolumeLighting(
 			MaterialParameters, PixelMaterialInputs, ResolvedView,
-			DirectionalLightShadow,
+			DirectionalLightShadow * DirectionalLightCloudShadow,
 			SingleLayerWater.SceneColorWithoutSingleLayerWaterTexture, SingleLayerWaterSceneColorSampler,
 			SingleLayerWater.SceneDepthWithoutSingleLayerWaterTexture, SingleLayerWaterSceneDepthSampler,
 			SingleLayerWater.SceneWithoutSingleLayerWaterMinMaxUV.xy,
diff --git a/Engine/Shaders/Private/BasePassVertexShader.usf b/Engine/Shaders/Private/BasePassVertexShader.usf
--- a/Engine/Shaders/Private/BasePassVertexShader.usf
+++ b/Engine/Shaders/Private/BasePassVertexShader.usf
@@ -10,6 +10,9 @@
 
 #if TRANSLUCENCY_PERVERTEX_FORWARD_SHADING
 #include "ReflectionEnvironmentShared.ush"
+
+// Enables applying a cloud shadow factor to lighting calculations in ForwardLightingCommon.ush.
+#define ENABLE_FORWARD_CLOUD_SHADOW SUPPORT_CLOUD_SHADOW_ON_FORWARD_LIT_TRANSLUCENT
 #include "ForwardLightingCommon.ush"
 #endif
 
@@ -17,7 +20,7 @@
 #include "SkyAtmosphereCommon.ush"
 #endif
 
-#if (NEEDS_BASEPASS_VERTEX_FOGGING && MATERIAL_ENABLE_TRANSLUCENCY_CLOUD_FOGGING) || SUPPORT_CLOUD_SHADOW_ON_FORWARD_LIT_TRANSLUCENT
+#if (NEEDS_BASEPASS_VERTEX_FOGGING && MATERIAL_ENABLE_TRANSLUCENCY_CLOUD_FOGGING) || SUPPORT_CLOUD_SHADOW_ON_FORWARD_LIT_TRANSLUCENT || NEEDS_BASEPASS_CLOUD_SHADOW_INTERPOLATOR
 #include "VolumetricCloudCommon.ush"
 #endif
 
@@ -139,7 +142,7 @@ void Main(
 
 #endif // NEEDS_BASEPASS_VERTEX_FOGGING
 
-#if TRANSLUCENCY_ANY_PERVERTEX_LIGHTING || SUPPORT_CLOUD_SHADOW_ON_FORWARD_LIT_TRANSLUCENT
+#if TRANSLUCENCY_ANY_PERVERTEX_LIGHTING || SUPPORT_CLOUD_SHADOW_ON_FORWARD_LIT_TRANSLUCENT || NEEDS_BASEPASS_CLOUD_SHADOW_INTERPOLATOR
 	float3 WorldPositionForVertexLightingTranslated = VertexFactoryGetPositionForVertexLighting(Input, VFIntermediates, WorldPosition.xyz);
 	float3 WorldPositionForVertexLighting = WorldPositionForVertexLightingTranslated - LWCHackToFloat(ResolvedView.PreViewTranslation);
 #endif
@@ -194,22 +197,30 @@ void Main(
 
 #elif NEEDS_BASEPASS_CLOUD_SHADOW_INTERPOLATOR
 
+#if MATERIAL_SHADINGMODEL_SINGLELAYERWATER
+#define ForwardDirLightCloudShadowParams SingleLayerWater.ForwardDirLightCloudShadow
+#elif IS_MATERIAL_TRANSLUCENT_AND_LIT
+#define ForwardDirLightCloudShadowParams TranslucentBasePass.ForwardDirLightCloudShadow
+#endif
+
 	Output.BasePassInterpolants.VertexCloudShadow = 1.0f;
-	if (TranslucentBasePass.ForwardDirLightCloudShadow.CloudShadowmapStrength > 0.0f)
+	if (ForwardDirLightCloudShadowParams.CloudShadowmapStrength > 0.0f)
 	{
 		float OutOpticalDepth = 0.0f;
 		Output.BasePassInterpolants.VertexCloudShadow = lerp(
 			1.0f,
 			GetCloudVolumetricShadow(
 				WorldPositionForVertexLightingTranslated,
-				TranslucentBasePass.ForwardDirLightCloudShadow.CloudShadowmapTranslatedWorldToLightClipMatrix,
-				TranslucentBasePass.ForwardDirLightCloudShadow.CloudShadowmapFarDepthKm,
-				TranslucentBasePass.ForwardDirLightCloudShadow.CloudShadowmapTexture,
-				TranslucentBasePass.ForwardDirLightCloudShadow.CloudShadowmapSampler,
+				ForwardDirLightCloudShadowParams.CloudShadowmapTranslatedWorldToLightClipMatrix,
+				ForwardDirLightCloudShadowParams.CloudShadowmapFarDepthKm,
+				ForwardDirLightCloudShadowParams.CloudShadowmapTexture,
+				ForwardDirLightCloudShadowParams.CloudShadowmapSampler,
 				OutOpticalDepth),
-			TranslucentBasePass.ForwardDirLightCloudShadow.CloudShadowmapStrength);
+			ForwardDirLightCloudShadowParams.CloudShadowmapStrength);
 	}
 
+#undef ForwardDirLightCloudShadowParams
+
 #endif
 
 	#if PRECOMPUTED_IRRADIANCE_VOLUME_LIGHTING && TRANSLUCENCY_ANY_PERVERTEX_LIGHTING
diff --git a/Engine/Shaders/Private/ForwardLightingCommon.ush b/Engine/Shaders/Private/ForwardLightingCommon.ush
--- a/Engine/Shaders/Private/ForwardLightingCommon.ush
+++ b/Engine/Shaders/Private/ForwardLightingCommon.ush
@@ -31,8 +31,8 @@
 #include "VirtualShadowMaps/VirtualShadowMapProjectionCommon.ush"
 #endif
 
-#if SUPPORT_CLOUD_SHADOW_ON_FORWARD_LIT_TRANSLUCENT
-#include "VolumetricCloudCommon.ush" 
+#ifndef ENABLE_FORWARD_CLOUD_SHADOW
+#define ENABLE_FORWARD_CLOUD_SHADOW 0
 #endif
 
 float PrevSceneColorPreExposureInv;
@@ -208,7 +208,7 @@ FDeferredLightingSplit GetForwardDirectLightingSplit(
 		FLATTEN
 		if ((DirectionalLightData.DirectionalLightShadowMapChannelMask >> 8) & LightingChannelMask)
 		{
-		#if SUPPORT_CLOUD_SHADOW_ON_FORWARD_LIT_TRANSLUCENT
+		#if ENABLE_FORWARD_CLOUD_SHADOW
 			// This cannot go into LightAttenuation due to some distance based attenuation math being applied.
 			NewLighting.DiffuseLighting *= InDirectionalLightCloudShadow;
 			NewLighting.SpecularLighting *= InDirectionalLightCloudShadow;
@@ -333,7 +333,7 @@ float3 GetForwardDirectLightingForVertexLighting(uint GridIndex, float3 Translat
 		bool bUnused = false;
 		ShadowFactor *= ComputeDirectionalLightDynamicShadowing(TranslatedWorldPosition, SceneDepth, bUnused);
 
-	#if SUPPORT_CLOUD_SHADOW_ON_FORWARD_LIT_TRANSLUCENT
+	#if ENABLE_FORWARD_CLOUD_SHADOW
 		ShadowFactor *= InDirectionalLightCloudShadow;
 	#endif
 
diff --git a/Engine/Source/Runtime/Engine/Classes/Engine/RendererSettings.h b/Engine/Source/Runtime/Engine/Classes/Engine/RendererSettings.h
--- a/Engine/Source/Runtime/Engine/Classes/Engine/RendererSettings.h
+++ b/Engine/Source/Runtime/Engine/Classes/Engine/RendererSettings.h
@@ -904,6 +904,15 @@ class ENGINE_API URendererSettings : public UDeveloperSettings
 		ConfigRestartRequired = true))
 		uint32 bSupportCloudShadowOnForwardLitTranslucent : 1;
 
+	/**
+	"Enable cloud shadow on SingleLayerWater surface. This is evaluated per vertex to reduce GPU cost. The cloud system requires extra samplers/textures to be bound to vertex shaders."
+	*/
+	UPROPERTY(config, EditAnywhere, Category = ShaderPermutationReduction, meta = (
+		ConsoleVariable = "r.Water.SingleLayerWater.SupportCloudShadow", DisplayName = "Support Cloud Shadow On SingleLayerWater",
+		ToolTip = "Enable cloud shadow on SingleLayerWater. This is evaluated per vertex to reduce GPU cost and requires extra samplers/textures to be bound to vertex shaders. This is not implemented on mobile as VolumetricClouds are not available on these platforms.",
+		ConfigRestartRequired = true))
+		uint32 bSupportCloudShadowOnSingleLayerWater : 1;
+
 	/**
 	"Enable Strata materials (Beta)."
 	*/
diff --git a/Engine/Source/Runtime/Engine/Private/ShaderCompiler/ShaderCompiler.cpp b/Engine/Source/Runtime/Engine/Private/ShaderCompiler/ShaderCompiler.cpp
--- a/Engine/Source/Runtime/Engine/Private/ShaderCompiler/ShaderCompiler.cpp
+++ b/Engine/Source/Runtime/Engine/Private/ShaderCompiler/ShaderCompiler.cpp
@@ -5943,6 +5943,12 @@ void GlobalBeginCompileShader(
 		Input.Environment.SetDefine(TEXT("SUPPORT_CLOUD_SHADOW_ON_FORWARD_LIT_TRANSLUCENT"), bSupportCloudShadowOnForwardLitTranslucent ? 1 : 0);
 	}
 
+	{
+		static IConsoleVariable *CVar = IConsoleManager::Get().FindConsoleVariable(TEXT("r.Water.SingleLayerWater.SupportCloudShadow"));
+		const bool bSupportCloudShadowOnSingleLayerWater = CVar && CVar->GetInt() > 0;
+		Input.Environment.SetDefine(TEXT("SUPPORT_CLOUD_SHADOW_ON_SINGLE_LAYER_WATER"), bSupportCloudShadowOnSingleLayerWater ? 1 : 0);
+	}
+
 	{
 		static IConsoleVariable* CVar = IConsoleManager::Get().FindConsoleVariable(TEXT("r.Strata"));
 		const bool bStrata = CVar && CVar->GetInt() != 0;
diff --git a/Engine/Source/Runtime/Engine/Private/ShaderCompiler/ShaderGenerationUtil.cpp b/Engine/Source/Runtime/Engine/Private/ShaderCompiler/ShaderGenerationUtil.cpp
--- a/Engine/Source/Runtime/Engine/Private/ShaderCompiler/ShaderGenerationUtil.cpp
+++ b/Engine/Source/Runtime/Engine/Private/ShaderCompiler/ShaderGenerationUtil.cpp
@@ -91,6 +91,7 @@ void FShaderCompileUtilities::ApplyFetchEnvironment(FShaderGlobalDefines& SrcDef
 	FETCH_COMPILE_BOOL(PROJECT_SUPPORT_SKY_ATMOSPHERE);
 	FETCH_COMPILE_BOOL(PROJECT_SUPPORT_SKY_ATMOSPHERE_AFFECTS_HEIGHFOG);
 	FETCH_COMPILE_BOOL(SUPPORT_CLOUD_SHADOW_ON_FORWARD_LIT_TRANSLUCENT);
+	FETCH_COMPILE_BOOL(SUPPORT_CLOUD_SHADOW_ON_SINGLE_LAYER_WATER);
 	FETCH_COMPILE_BOOL(PROJECT_MOBILE_USE_LEGACY_SHADING);
 	FETCH_COMPILE_BOOL(POST_PROCESS_ALPHA);
 	FETCH_COMPILE_BOOL(PLATFORM_SUPPORTS_RENDERTARGET_WRITE_MASK);
@@ -366,6 +367,12 @@ static FShaderGlobalDefines FetchShaderGlobalDefines(EShaderPlatform TargetPlatf
 		Ret.SUPPORT_CLOUD_SHADOW_ON_FORWARD_LIT_TRANSLUCENT = bSupportCloudShadowOnForwardLitTranslucent ? 1 : 0;
 	}
 
+	{
+		static IConsoleVariable *CVar = IConsoleManager::Get().FindConsoleVariable(TEXT("r.Water.SingleLayerWater.SupportCloudShadow"));
+		const bool bSupportCloudShadowOnSingleLayerWater = CVar && CVar->GetInt() > 0;
+		Ret.SUPPORT_CLOUD_SHADOW_ON_SINGLE_LAYER_WATER = bSupportCloudShadowOnSingleLayerWater ? 1 : 0;
+	}
+
 	{
 		static IConsoleVariable* CVar = IConsoleManager::Get().FindConsoleVariable(TEXT("r.Mobile.UseLegacyShadingModel"));
 		Ret.PROJECT_MOBILE_USE_LEGACY_SHADING = CVar ? (CVar->GetInt() != 0) : 0;
diff --git a/Engine/Source/Runtime/RenderCore/Public/ShaderMaterial.h b/Engine/Source/Runtime/RenderCore/Public/ShaderMaterial.h
--- a/Engine/Source/Runtime/RenderCore/Public/ShaderMaterial.h
+++ b/Engine/Source/Runtime/RenderCore/Public/ShaderMaterial.h
@@ -32,6 +32,7 @@ struct FShaderGlobalDefines
 	bool PROJECT_SUPPORT_SKY_ATMOSPHERE;
 	bool PROJECT_SUPPORT_SKY_ATMOSPHERE_AFFECTS_HEIGHFOG;
 	bool SUPPORT_CLOUD_SHADOW_ON_FORWARD_LIT_TRANSLUCENT;
+	bool SUPPORT_CLOUD_SHADOW_ON_SINGLE_LAYER_WATER;
 	bool PROJECT_MOBILE_USE_LEGACY_SHADING;
 	bool POST_PROCESS_ALPHA;
 	bool PLATFORM_SUPPORTS_RENDERTARGET_WRITE_MASK;
diff --git a/Engine/Source/Runtime/Renderer/Private/SingleLayerWaterRendering.cpp b/Engine/Source/Runtime/Renderer/Private/SingleLayerWaterRendering.cpp
--- a/Engine/Source/Runtime/Renderer/Private/SingleLayerWaterRendering.cpp
+++ b/Engine/Source/Runtime/Renderer/Private/SingleLayerWaterRendering.cpp
@@ -70,6 +70,13 @@ static TAutoConsoleVariable<int32> CVarWaterSingleLayerShadersSupportDistanceFie
 	TEXT("Whether or not the single layer water material shaders are compiled with support for distance field shadow, i.e. output main directional light luminance in a separate render target. This is preconditioned on using deferred shading and having distance field support enabled in the project."),
 	ECVF_ReadOnly | ECVF_RenderThreadSafe);
 
+// The project setting for the cloud shadow to affect SingleLayerWater (enable/disable runtime and shader code).  This is not implemented on mobile as VolumetricClouds are not available on these platforms.
+static TAutoConsoleVariable<int32> CVarSupportCloudShadowOnSingleLayerWater(
+	TEXT("r.Water.SingleLayerWater.SupportCloudShadow"),
+	0,
+	TEXT("Enables cloud shadows on SingleLayerWater materials."),
+	ECVF_ReadOnly | ECVF_RenderThreadSafe);
+
 static TAutoConsoleVariable<int32> CVarWaterSingleLayerDistanceFieldShadow(
 	TEXT("r.Water.SingleLayer.DistanceFieldShadow"), 1,
 	TEXT("When using deferred, distance field shadow tracing is supported on single layer water. This cvar can be used to toggle it on/off at runtime."),
@@ -850,6 +857,7 @@ BEGIN_UNIFORM_BUFFER_STRUCT(FSingleLayerWaterPassUniformParameters,)
 	SHADER_PARAMETER(FVector4f, DistortionParams)
 	SHADER_PARAMETER(FVector2f, SceneWithoutSingleLayerWaterTextureSize)
 	SHADER_PARAMETER(FVector2f, SceneWithoutSingleLayerWaterInvTextureSize)
+	SHADER_PARAMETER_STRUCT(FLightCloudTransmittanceParameters, ForwardDirLightCloudShadow)
 END_UNIFORM_BUFFER_STRUCT()
 
 // At the moment we reuse the DeferredDecals static uniform buffer slot because it is currently unused in this pass.
@@ -923,6 +931,9 @@ void FDeferredShadingSceneRenderer::RenderSingleLayerWaterInner(
 			SetupDistortionParams(SLWUniformParameters.DistortionParams, View);
 			SLWUniformParameters.SceneWithoutSingleLayerWaterTextureSize = FVector2f(DepthTextureSize.X, DepthTextureSize.Y);
 			SLWUniformParameters.SceneWithoutSingleLayerWaterInvTextureSize = FVector2f(1.0f / DepthTextureSize.X, 1.0f / DepthTextureSize.Y);
+
+			const FLightSceneProxy *SelectedForwardDirectionalLightProxy = View.ForwardLightingResources.SelectedForwardDirectionalLightProxy;
+			SetupLightCloudTransmittanceParameters(GraphBuilder, Scene, View, SelectedForwardDirectionalLightProxy ? SelectedForwardDirectionalLightProxy->GetLightSceneInfo() : nullptr, SLWUniformParameters.ForwardDirLightCloudShadow);
 		}
 
 		FSingleLayerWaterPassParameters* PassParameters = GraphBuilder.AllocParameters<FSingleLayerWaterPassParameters>();
