ColorManagement: Initial compatibility for working color space in the output transform.

This focuses initally on integration in CombineLUTs and its device-encoding-only equivalent, but exposes a global color space conversion uniform buffer for later use within the renderer but also in plugins.

diff --git a/Engine/Shaders/Private/PostProcessCombineLUTs.usf b/Engine/Shaders/Private/PostProcessCombineLUTs.usf
--- a/Engine/Shaders/Private/PostProcessCombineLUTs.usf
+++ b/Engine/Shaders/Private/PostProcessCombineLUTs.usf
@@ -39,7 +39,7 @@ static const float LUTSize = 32;
 float WhiteTemp;
 float WhiteTint;
 
-float3 WhiteBalance( float3 LinearColor )
+float3 WhiteBalance( float3 LinearColor, const float3x3 WCS_2_XYZ, const float3x3 XYZ_2_WCS )
 {
 	float2 SrcWhiteDaylight = D_IlluminantChromaticity( WhiteTemp );
 	float2 SrcWhitePlankian = PlanckianLocusChromaticity( WhiteTemp );
@@ -61,7 +61,7 @@ float3 WhiteBalance( float3 LinearColor )
 	}
 
 	float3x3 WhiteBalanceMat = ChromaticAdaptation( SrcWhite, D65White );
-	WhiteBalanceMat = mul( XYZ_2_sRGB_MAT, mul( WhiteBalanceMat, sRGB_2_XYZ_MAT ) );
+	WhiteBalanceMat = mul( XYZ_2_WCS, mul( WhiteBalanceMat, WCS_2_XYZ ) );
 
 	return mul( WhiteBalanceMat, LinearColor );
 }
@@ -198,10 +198,7 @@ float4 CombineLUTsCommon(float2 InUV, uint InLayerIndex)
 
 	float4 OutColor = 0;
 	
-	// AP1 to Working space matrices
-	const float3x3 sRGB_2_AP1 = mul( XYZ_2_AP1_MAT, mul( D65_2_D60_CAT, sRGB_2_XYZ_MAT ) );
-	const float3x3 AP1_2_sRGB = mul( XYZ_2_sRGB_MAT, mul( D60_2_D65_CAT, AP1_2_XYZ_MAT ) );
-
+	const float3x3 AP1_2_sRGB = mul( XYZ_2_sRGB_MAT, mul( D60_2_D65_CAT, AP1_2_XYZ_MAT ) ); 
 	const float3x3 AP0_2_AP1 = mul( XYZ_2_AP1_MAT, AP0_2_XYZ_MAT );
 	const float3x3 AP1_2_AP0 = mul( XYZ_2_AP0_MAT, AP1_2_XYZ_MAT );
 
@@ -224,12 +221,12 @@ float4 CombineLUTsCommon(float2 InUV, uint InLayerIndex)
 	}
 
 #if !SKIP_TEMPERATURE
-	float3 BalancedColor = WhiteBalance( LinearColor );
+	float3 BalancedColor = WhiteBalance( LinearColor, (float3x3)WorkingColorSpace.ToXYZ, (float3x3)WorkingColorSpace.FromXYZ );
 #else
 	float3 BalancedColor = LinearColor;
 #endif
 
-	float3 ColorAP1 = mul( sRGB_2_AP1, BalancedColor );
+	float3 ColorAP1 = mul( (float3x3)WorkingColorSpace.ToAP1, BalancedColor );
 
 	// Expand bright saturated colors outside the sRGB gamut to fake wide gamut rendering.
  	float  LumaAP1 = dot( ColorAP1, AP1_RGB2Y );
@@ -260,7 +257,7 @@ float4 CombineLUTsCommon(float2 InUV, uint InLayerIndex)
 	ColorAP1 = ColorCorrectAll( ColorAP1 );
 
 	// Store for Linear HDR output without tone curve
-	float3 GradedColor = mul( AP1_2_sRGB, ColorAP1 );
+	float3 GradedColor = mul( (float3x3)WorkingColorSpace.FromAP1, ColorAP1 );
 
 	const float3x3 BlueCorrect =
 	{
@@ -287,8 +284,8 @@ float4 CombineLUTsCommon(float2 InUV, uint InLayerIndex)
 	// Uncorrect blue to maintain white point
 	ColorAP1 = lerp( ColorAP1, mul( BlueCorrectInvAP1, ColorAP1 ), BlueCorrection );
 
-	// Convert from AP1 to sRGB and clip out-of-gamut values
-	float3 FilmColor = max(0, mul( AP1_2_sRGB, ColorAP1 ));
+	// Convert from AP1 to the working color space and clip out-of-gamut values
+	float3 FilmColor = max(0, mul( (float3x3)WorkingColorSpace.FromAP1, ColorAP1 ));
 
 #if BLENDCOUNT > 1
 	{
@@ -350,8 +347,8 @@ float4 CombineLUTsCommon(float2 InUV, uint InLayerIndex)
 
 		// FIXME: Workaround for UE-29935, pushing all colors with a 0 component to black output
 		// Default parameters seem to cancel out (sRGB->XYZ->AP1->XYZ->sRGB), so should be okay for a temp fix
-		float3 OutputGamutColor = FilmColor;
-
+		float3 OutputGamutColor = WorkingColorSpace.bIsSRGB ? FilmColor : mul( AP1_2_Output, mul( (float3x3)WorkingColorSpace.ToAP1, FilmColor ) );
+		
 		// Apply conversion to sRGB (this must be an exact sRGB conversion else darks are bad).
 		OutDeviceColor = LinearToSrgb( OutputGamutColor );
 	}
@@ -359,8 +356,8 @@ float4 CombineLUTsCommon(float2 InUV, uint InLayerIndex)
 	// Rec 709, user specified gamut
 	else if( GetOutputDevice() == TONEMAPPER_OUTPUT_Rec709)
 	{
-		// Convert from sRGB to specified output gamut
-		float3 OutputGamutColor = mul( AP1_2_Output, mul( sRGB_2_AP1, FilmColor ) );
+		// Convert from the working color space to specified output gamut
+		float3 OutputGamutColor = mul( AP1_2_Output, mul( (float3x3)WorkingColorSpace.ToAP1, FilmColor ) );
 
 		// Didn't profile yet if the branching version would be faster (different linear segment).
 		OutDeviceColor = LinearTo709Branchless( OutputGamutColor );
@@ -370,7 +367,7 @@ float4 CombineLUTsCommon(float2 InUV, uint InLayerIndex)
 	else if( GetOutputDevice() == TONEMAPPER_OUTPUT_ACES1000nitST2084 || GetOutputDevice() == TONEMAPPER_OUTPUT_ACES1000nitScRGB)
 	{		
 		// 1000 nit ODT
-		float3 ODTColor = ACESOutputTransforms1000( GradedColor );
+		float3 ODTColor = ACESOutputTransforms1000( GradedColor, (float3x3)WorkingColorSpace.ToAP0 );
 
 		// Convert from AP1 to specified output gamut except for ScRGB as it might contain negative values. 
 		// In this case, the AP1->output conversion happens through ST2084ToScRGB in the tonemap function
@@ -387,7 +384,7 @@ float4 CombineLUTsCommon(float2 InUV, uint InLayerIndex)
 	else if( GetOutputDevice() == TONEMAPPER_OUTPUT_ACES2000nitST2084 || GetOutputDevice() == TONEMAPPER_OUTPUT_ACES2000nitScRGB)
 	{		
 		// 2000 nit ODT
-		float3 ODTColor = ACESOutputTransforms2000( GradedColor );
+		float3 ODTColor = ACESOutputTransforms2000( GradedColor, (float3x3)WorkingColorSpace.ToAP0 );
 
 		// Convert from AP1 to specified output gamut except for ScRGB as it might contain negative values. 
 		// In this case, the AP1->output conversion happens through ST2084ToScRGB in the tonemap function
@@ -402,7 +399,7 @@ float4 CombineLUTsCommon(float2 InUV, uint InLayerIndex)
 	
 	else if( GetOutputDevice() == TONEMAPPER_OUTPUT_LinearEXR)
 	{
-			float3 OutputGamutColor = mul( AP1_2_Output, mul( sRGB_2_AP1, GradedColor ) );
+			float3 OutputGamutColor = mul( AP1_2_Output, mul( (float3x3)WorkingColorSpace.ToAP1, GradedColor ) );
 			OutDeviceColor = LinearToST2084( OutputGamutColor );
 	}
 	// Linear HDR, including all color correction, but no tone curve
@@ -413,7 +410,7 @@ float4 CombineLUTsCommon(float2 InUV, uint InLayerIndex)
 	// "Linear" including all color correction and the tone curve, but no device gamma
 	else if (GetOutputDevice() == TONEMAPPER_OUTPUT_WithToneCurve)
 	{
-		float3 OutputGamutColor = mul(AP1_2_Output, mul(sRGB_2_AP1, FilmColorNoGamma));
+		float3 OutputGamutColor = mul( AP1_2_Output, mul( (float3x3)WorkingColorSpace.ToAP1, FilmColorNoGamma ) );
 
 		OutDeviceColor = OutputGamutColor;
 	}
@@ -422,8 +419,8 @@ float4 CombineLUTsCommon(float2 InUV, uint InLayerIndex)
 	// Gamma 2.2, user specified gamut
 	else
 	{
-		// Convert from sRGB to specified output gamut
-		float3 OutputGamutColor = mul( AP1_2_Output, mul( sRGB_2_AP1, FilmColor ) );
+		// Convert from the working color space to specified output gamut
+		float3 OutputGamutColor = mul( AP1_2_Output, mul( (float3x3)WorkingColorSpace.ToAP1, FilmColor ) );
 
 		// This is different than the prior "gamma" curve adjustment (but reusing the variable).
 		// For displays set to a gamma colorspace.
diff --git a/Engine/Shaders/Private/PostProcessDeviceEncodingOnly.usf b/Engine/Shaders/Private/PostProcessDeviceEncodingOnly.usf
--- a/Engine/Shaders/Private/PostProcessDeviceEncodingOnly.usf
+++ b/Engine/Shaders/Private/PostProcessDeviceEncodingOnly.usf
@@ -58,13 +58,6 @@ float4 DeviceEncodingOnlyCommonPS(
 
 	half4 SceneColor = SampleSceneColor(SceneUV);
 
-	// AP1 to Working space matrices
-	const float3x3 sRGB_2_AP1 = mul( XYZ_2_AP1_MAT, mul( D65_2_D60_CAT, sRGB_2_XYZ_MAT ) );
-	const float3x3 AP1_2_sRGB = mul( XYZ_2_sRGB_MAT, mul( D60_2_D65_CAT, AP1_2_XYZ_MAT ) );
-
-	const float3x3 AP0_2_AP1 = mul( XYZ_2_AP1_MAT, AP0_2_XYZ_MAT );
-	const float3x3 AP1_2_AP0 = mul( XYZ_2_AP0_MAT, AP1_2_XYZ_MAT );
-
 	const float3x3 AP1_2_Output  = OuputGamutMappingMatrix( OutputGamut );
 
 	// Apply "gamma" curve adjustment.
@@ -83,7 +76,7 @@ float4 DeviceEncodingOnlyCommonPS(
 
 		// FIXME: Workaround for UE-29935, pushing all colors with a 0 component to black output
 		// Default parameters seem to cancel out (sRGB->XYZ->AP1->XYZ->sRGB), so should be okay for a temp fix
-		float3 OutputGamutColor = FilmColor;
+		float3 OutputGamutColor = WorkingColorSpace.bIsSRGB ? FilmColor : mul( AP1_2_Output, mul( (float3x3)WorkingColorSpace.ToAP1, FilmColor ) );
 
 		// Apply conversion to sRGB (this must be an exact sRGB conversion else darks are bad).
 		OutDeviceColor = LinearToSrgb( OutputGamutColor );
@@ -93,7 +86,7 @@ float4 DeviceEncodingOnlyCommonPS(
 	else if( GetOutputDevice() == TONEMAPPER_OUTPUT_Rec709)
 	{
 		// Convert from sRGB to specified output gamut
-		float3 OutputGamutColor = mul( AP1_2_Output, mul( sRGB_2_AP1, FilmColor ) );
+		float3 OutputGamutColor = mul( AP1_2_Output, mul( (float3x3)WorkingColorSpace.ToAP1, FilmColor ) );
 
 		// Didn't profile yet if the branching version would be faster (different linear segment).
 		OutDeviceColor = LinearTo709Branchless( OutputGamutColor );
@@ -103,7 +96,7 @@ float4 DeviceEncodingOnlyCommonPS(
 	else if( GetOutputDevice() == TONEMAPPER_OUTPUT_ACES1000nitST2084 )
 	{		
 		// 1000 nit ODT
-		float3 ODTColor = ACESOutputTransforms1000( FilmColor );
+		float3 ODTColor = ACESOutputTransforms1000( FilmColor, (float3x3)WorkingColorSpace.ToAP0 );
 
 		// Convert from AP1 to specified output gamut
 		ODTColor = mul( AP1_2_Output, ODTColor );
@@ -116,7 +109,7 @@ float4 DeviceEncodingOnlyCommonPS(
 	else if( GetOutputDevice() == TONEMAPPER_OUTPUT_ACES2000nitST2084 )
 	{		
 		// 2000 nit ODT
-		float3 ODTColor = ACESOutputTransforms2000( FilmColor );
+		float3 ODTColor = ACESOutputTransforms2000( FilmColor, (float3x3)WorkingColorSpace.ToAP0 );
 
 		// Convert from AP1 to specified output gamut
 		ODTColor = mul( AP1_2_Output, ODTColor );
@@ -129,7 +122,7 @@ float4 DeviceEncodingOnlyCommonPS(
 	else if(GetOutputDevice() == TONEMAPPER_OUTPUT_ACES1000nitScRGB)
 	{
 		// 1000 nit ODT
-		float3 ODTColor = ACESOutputTransforms1000( FilmColor );
+		float3 ODTColor = ACESOutputTransforms1000( FilmColor, (float3x3)WorkingColorSpace.ToAP0 );
 
 		// Apply conversion to ScRGB
 		OutDeviceColor = ST2084ToScRGB( LinearToST2084( ODTColor ), TONEMAPPER_OUTPUT_ACES1000nitScRGB );
@@ -139,7 +132,7 @@ float4 DeviceEncodingOnlyCommonPS(
 	else if(GetOutputDevice() == TONEMAPPER_OUTPUT_ACES2000nitScRGB)
 	{
 		// 2000 nit ODT
-		float3 ODTColor = ACESOutputTransforms2000( FilmColor );
+		float3 ODTColor = ACESOutputTransforms2000( FilmColor, (float3x3)WorkingColorSpace.ToAP0 );
 
 		// Apply conversion to ScRGB
 		OutDeviceColor = ST2084ToScRGB( LinearToST2084( ODTColor ), TONEMAPPER_OUTPUT_ACES2000nitScRGB );
@@ -147,7 +140,7 @@ float4 DeviceEncodingOnlyCommonPS(
 	
 	else if( GetOutputDevice() == TONEMAPPER_OUTPUT_LinearEXR)
 	{
-		float3 OutputGamutColor = mul( AP1_2_Output, mul( sRGB_2_AP1, FilmColor ) );
+		float3 OutputGamutColor = mul( AP1_2_Output, mul( (float3x3)WorkingColorSpace.ToAP1, FilmColor ) );
 		OutDeviceColor = LinearToST2084( OutputGamutColor );
 	}
 	// Linear HDR, including all color correction, but no tone curve
@@ -158,7 +151,7 @@ float4 DeviceEncodingOnlyCommonPS(
 	// "Linear" including all color correction and the tone curve, but no device gamma
 	else if (GetOutputDevice() == TONEMAPPER_OUTPUT_WithToneCurve)
 	{
-		float3 OutputGamutColor = mul(AP1_2_Output, mul(sRGB_2_AP1, SceneColor.rgb)); 
+		float3 OutputGamutColor = mul( AP1_2_Output, mul( (float3x3)WorkingColorSpace.ToAP1, SceneColor.rgb ) ); 
 
 		OutDeviceColor = OutputGamutColor;
 	}
@@ -168,7 +161,7 @@ float4 DeviceEncodingOnlyCommonPS(
 	else
 	{
 		// Convert from sRGB to specified output gamut
-		float3 OutputGamutColor = mul( AP1_2_Output, mul( sRGB_2_AP1, FilmColor ) );
+		float3 OutputGamutColor = mul( AP1_2_Output, mul( (float3x3)WorkingColorSpace.ToAP1, FilmColor ) );
 
 		// This is different than the prior "gamma" curve adjustment (but reusing the variable).
 		// For displays set to a gamma colorspace.
diff --git a/Engine/Shaders/Private/TonemapCommon.ush b/Engine/Shaders/Private/TonemapCommon.ush
--- a/Engine/Shaders/Private/TonemapCommon.ush
+++ b/Engine/Shaders/Private/TonemapCommon.ush
@@ -238,19 +238,24 @@ half3 FilmToneMapInverse( half3 ToneColor )
 
 //
 // ACES sRGB D65 Output Transform - Forward and Inverse
-//  Input is scene-referred linear values in the sRGB gamut
+//  Input is scene-referred linear values in the working space gamut
 //  Output is output-referred linear values in the sRGB gamut
 //
-float3 ACESOutputTransformsRGBD65( float3 SceneReferredLinearsRGBColor )
+float3 ACESOutputTransformsRGBD65( float3 SceneReferredLinearColor, const float3x3 WCS_2_AP0 )
 {
-	const float3x3 sRGB_2_AP0 = mul( XYZ_2_AP0_MAT, mul( D65_2_D60_CAT, sRGB_2_XYZ_MAT ) );
-
-	float3 aces = mul( sRGB_2_AP0, SceneReferredLinearsRGBColor * 1.5 );
+	float3 aces = mul( WCS_2_AP0, SceneReferredLinearColor * 1.5 );
 	float3 oces = RRT( aces );
 	float3 OutputReferredLinearsRGBColor =  ODT_sRGB_D65( oces );
 	return OutputReferredLinearsRGBColor;
 }
 
+float3 ACESOutputTransformsRGBD65( float3 SceneReferredLinearsRGBColor )
+{
+	const float3x3 sRGB_2_AP0 = mul( XYZ_2_AP0_MAT, mul( D65_2_D60_CAT, sRGB_2_XYZ_MAT ) );
+
+	return ACESOutputTransformsRGBD65( SceneReferredLinearsRGBColor, sRGB_2_AP0 );
+}
+
 float3 InverseACESOutputTransformsRGBD65( float3 OutputReferredLinearsRGBColor )
 {
 	const float3x3 AP0_2_sRGB = mul( XYZ_2_sRGB_MAT, mul( D60_2_D65_CAT, AP0_2_XYZ_MAT ) );
@@ -264,34 +269,44 @@ float3 InverseACESOutputTransformsRGBD65( float3 OutputReferredLinearsRGBColor )
 
 //
 // ACES D65 1000 nit Output Transform - Forward
-//  Input is scene-referred linear values in the sRGB gamut
+//  Input is scene-referred linear values in the working space gamut
 //  Output is output-referred linear values in the AP1 gamut
 //
-float3 ACESOutputTransforms1000( float3 SceneReferredLinearsRGBColor )
+float3 ACESOutputTransforms1000( float3 SceneReferredLinearColor, const float3x3 WCS_2_AP0 ) 
 {
-	const float3x3 sRGB_2_AP0 = mul( XYZ_2_AP0_MAT, mul( D65_2_D60_CAT, sRGB_2_XYZ_MAT ) );
-
-	float3 aces = mul( sRGB_2_AP0, SceneReferredLinearsRGBColor * 1.5 );
+	float3 aces = mul( WCS_2_AP0, SceneReferredLinearColor * 1.5 );
 	float3 oces = RRT( aces );
 	float3 OutputReferredLinearAP1Color = ODT_1000nits( oces );
 	return OutputReferredLinearAP1Color;
 }
 
+float3 ACESOutputTransforms1000( float3 SceneReferredLinearsRGBColor )
+{
+	const float3x3 sRGB_2_AP0 = mul( XYZ_2_AP0_MAT, mul( D65_2_D60_CAT, sRGB_2_XYZ_MAT ) );
+
+	return ACESOutputTransforms1000( SceneReferredLinearsRGBColor, sRGB_2_AP0 );
+}
+
 //
 // ACES D65 2000 nit Output Transform - Forward
-//  Input is scene-referred linear values in the sRGB gamut
+//  Input is scene-referred linear values in the working space gamut
 //  Output is output-referred linear values in the AP1 gamut
 //
-float3 ACESOutputTransforms2000( float3 SceneReferredLinearsRGBColor )
+float3 ACESOutputTransforms2000( float3 SceneReferredLinearColor, const float3x3 WCS_2_AP0 ) 
 {
-	const float3x3 sRGB_2_AP0 = mul( XYZ_2_AP0_MAT, mul( D65_2_D60_CAT, sRGB_2_XYZ_MAT ) );
-
-	float3 aces = mul( sRGB_2_AP0, SceneReferredLinearsRGBColor * 1.5 );
+	float3 aces = mul( WCS_2_AP0, SceneReferredLinearColor * 1.5 );
 	float3 oces = RRT( aces );
 	float3 OutputReferredLinearAP1Color = ODT_2000nits( oces );
 	return OutputReferredLinearAP1Color;
 }
 
+float3 ACESOutputTransforms2000( float3 SceneReferredLinearsRGBColor )
+{
+	const float3x3 sRGB_2_AP0 = mul( XYZ_2_AP0_MAT, mul( D65_2_D60_CAT, sRGB_2_XYZ_MAT ) );
+
+	return ACESOutputTransforms2000( SceneReferredLinearsRGBColor, sRGB_2_AP0 );
+}
+
 static const float3x3 GamutMappingIdentityMatrix = { 1.0, 0.0, 0.0, 0.0, 1.0, 0.0, 0.0, 0.0, 1.0 };
 
 //
diff --git a/Engine/Source/Runtime/Engine/Private/RendererSettings.cpp b/Engine/Source/Runtime/Engine/Private/RendererSettings.cpp
--- a/Engine/Source/Runtime/Engine/Private/RendererSettings.cpp
+++ b/Engine/Source/Runtime/Engine/Private/RendererSettings.cpp
@@ -5,6 +5,7 @@
 #include "RHI.h"
 #include "GPUSkinVertexFactory.h"
 #include "ColorSpace.h"
+#include "SceneManagement.h"
 
 #if WITH_EDITOR
 #include "Editor/EditorEngine.h"
@@ -410,6 +411,15 @@ void URendererSettings::UpdateWorkingColorSpaceAndChromaticities()
 	{
 		FColorSpace::GetWorking().GetChromaticities(RedChromaticityCoordinate, GreenChromaticityCoordinate, BlueChromaticityCoordinate, WhiteChromaticityCoordinate);
 	}
+	
+	FColorSpace UpdatedColorSpace = FColorSpace::GetWorking();
+
+	ENQUEUE_RENDER_COMMAND(WorkingColorSpaceCommand)(
+		[UpdatedColorSpace](FRHICommandList&)
+		{
+			// Set or update the global uniform buffer for Working Color Space conversions.
+			GDefaultWorkingColorSpaceUniformBuffer.Update(UpdatedColorSpace);
+		});
 }
 
 
diff --git a/Engine/Source/Runtime/Engine/Private/SceneManagement.cpp b/Engine/Source/Runtime/Engine/Private/SceneManagement.cpp
--- a/Engine/Source/Runtime/Engine/Private/SceneManagement.cpp
+++ b/Engine/Source/Runtime/Engine/Private/SceneManagement.cpp
@@ -16,6 +16,7 @@
 #include "Engine/ShadowMapTexture2D.h"
 #include "VT/LightmapVirtualTexture.h"
 #include "UnrealEngine.h"
+#include "ColorSpace.h"
 
 static TAutoConsoleVariable<float> CVarLODTemporalLag(
 	TEXT("lod.TemporalLag"),
@@ -56,6 +57,32 @@ void FTemporalLODState::UpdateTemporalLODTransition(const FViewInfo& View, float
 }
 
 
+IMPLEMENT_STATIC_UNIFORM_BUFFER_SLOT(WorkingColorSpace);
+IMPLEMENT_STATIC_UNIFORM_BUFFER_STRUCT(FWorkingColorSpaceShaderParameters, "WorkingColorSpace", WorkingColorSpace);
+
+void FDefaultWorkingColorSpaceUniformBuffer::Update(const UE::Color::FColorSpace& InColorSpace)
+{
+	using namespace UE::Color;
+
+	const FVector2d& White = InColorSpace.GetWhiteChromaticity();
+	const FVector2d ACES_D60 = GetWhitePoint(EWhitePoint::ACES_D60);
+
+	FWorkingColorSpaceShaderParameters Parameters;
+	Parameters.ToXYZ = Transpose<float>(InColorSpace.GetRgbToXYZ());
+	Parameters.FromXYZ = Transpose<float>(InColorSpace.GetXYZToRgb());
+
+	Parameters.ToAP1 = Transpose<float>(FColorSpaceTransform(InColorSpace, FColorSpace(EColorSpace::ACESAP1), EChromaticAdaptationMethod::Bradford));
+	Parameters.FromAP1 = Parameters.ToAP1.Inverse();
+	
+	Parameters.ToAP0 = Transpose<float>(FColorSpaceTransform(InColorSpace, FColorSpace(EColorSpace::ACESAP0), EChromaticAdaptationMethod::Bradford));
+
+	Parameters.bIsSRGB = InColorSpace.IsSRGB();
+
+	SetContents(Parameters);
+}
+
+TGlobalResource<FDefaultWorkingColorSpaceUniformBuffer> GDefaultWorkingColorSpaceUniformBuffer;
+
 
 FSimpleElementCollector::FSimpleElementCollector() :
 	FPrimitiveDrawInterface(nullptr)
diff --git a/Engine/Source/Runtime/Engine/Public/SceneManagement.h b/Engine/Source/Runtime/Engine/Public/SceneManagement.h
--- a/Engine/Source/Runtime/Engine/Public/SceneManagement.h
+++ b/Engine/Source/Runtime/Engine/Public/SceneManagement.h
@@ -53,6 +53,8 @@ class ULightMapVirtualTexture2D;
 class FGPUScenePrimitiveCollector;
 class FVirtualShadowMapArrayCacheManager;
 
+namespace UE { namespace Color { class FColorSpace; } }
+
 DECLARE_LOG_CATEGORY_EXTERN(LogBufferVisualization, Log, All);
 DECLARE_LOG_CATEGORY_EXTERN(LogNaniteVisualization, Log, All);
 DECLARE_LOG_CATEGORY_EXTERN(LogLumenVisualization, Log, All);
@@ -305,6 +307,30 @@ private:
 	FSceneView& SceneView;
 };
 
+
+/**
+ * Global working color space shader parameters (color space conversion matrices).
+ */
+BEGIN_GLOBAL_SHADER_PARAMETER_STRUCT(FWorkingColorSpaceShaderParameters, ENGINE_API)
+	SHADER_PARAMETER(FMatrix44f, ToXYZ)
+	SHADER_PARAMETER(FMatrix44f, FromXYZ)
+	SHADER_PARAMETER(FMatrix44f, ToAP1)
+	SHADER_PARAMETER(FMatrix44f, FromAP1)
+	SHADER_PARAMETER(FMatrix44f, ToAP0)
+	SHADER_PARAMETER(uint32, bIsSRGB)
+END_SHADER_PARAMETER_STRUCT()
+
+class FDefaultWorkingColorSpaceUniformBuffer : public TUniformBuffer<FWorkingColorSpaceShaderParameters>
+{
+	typedef TUniformBuffer<FWorkingColorSpaceShaderParameters> Super;
+public:
+
+	void Update(const UE::Color::FColorSpace& InColorSpace);
+};
+
+ENGINE_API extern TGlobalResource<FDefaultWorkingColorSpaceUniformBuffer> GDefaultWorkingColorSpaceUniformBuffer;
+
+
 /**
  * The types of interactions between a light and a primitive.
  */
diff --git a/Engine/Source/Runtime/Renderer/Private/PostProcess/PostProcessCombineLUTs.cpp b/Engine/Source/Runtime/Renderer/Private/PostProcess/PostProcessCombineLUTs.cpp
--- a/Engine/Source/Runtime/Renderer/Private/PostProcess/PostProcessCombineLUTs.cpp
+++ b/Engine/Source/Runtime/Renderer/Private/PostProcess/PostProcessCombineLUTs.cpp
@@ -94,6 +94,7 @@ BEGIN_SHADER_PARAMETER_STRUCT(FCombineLUTParameters, )
 	SHADER_PARAMETER_TEXTURE_ARRAY(Texture2D, Textures, [GMaxLUTBlendCount])
 	SHADER_PARAMETER_SAMPLER_ARRAY(SamplerState, Samplers, [GMaxLUTBlendCount])
 	SHADER_PARAMETER_SCALAR_ARRAY(float, LUTWeights, [GMaxLUTBlendCount])
+	SHADER_PARAMETER_STRUCT_REF(FWorkingColorSpaceShaderParameters, WorkingColorSpace)
 	SHADER_PARAMETER(FVector4f, OverlayColor)
 	SHADER_PARAMETER(FVector3f, ColorScale)
 	SHADER_PARAMETER(FVector4f, ColorSaturation)
@@ -168,6 +169,8 @@ void GetCombineLUTParameters(
 		GET_SCALAR_ARRAY_ELEMENT(Parameters.LUTWeights, BlendIndex) = Weights[BlendIndex];
 	}
 
+	Parameters.WorkingColorSpace = GDefaultWorkingColorSpaceUniformBuffer.GetUniformBufferRef();
+
 	Parameters.ColorScale = FVector4f(View.ColorScale);
 	Parameters.OverlayColor = View.OverlayColor;
 	Parameters.ColorRemap = GetColorRemapParameters();
diff --git a/Engine/Source/Runtime/Renderer/Private/PostProcess/PostProcessDeviceEncodingOnly.cpp b/Engine/Source/Runtime/Renderer/Private/PostProcess/PostProcessDeviceEncodingOnly.cpp
--- a/Engine/Source/Runtime/Renderer/Private/PostProcess/PostProcessDeviceEncodingOnly.cpp
+++ b/Engine/Source/Runtime/Renderer/Private/PostProcess/PostProcessDeviceEncodingOnly.cpp
@@ -84,6 +84,7 @@ FDeviceEncodingOnlyOutputDeviceParameters GetDeviceEncodingOnlyOutputDeviceParam
 
 BEGIN_SHADER_PARAMETER_STRUCT(FDeviceEncodingOnlyParameters, )
 	SHADER_PARAMETER_STRUCT_REF(FViewUniformShaderParameters, View)
+	SHADER_PARAMETER_STRUCT_REF(FWorkingColorSpaceShaderParameters, WorkingColorSpace)
 	SHADER_PARAMETER_STRUCT_INCLUDE(FDeviceEncodingOnlyOutputDeviceParameters, OutputDevice)
 	SHADER_PARAMETER_STRUCT(FScreenPassTextureViewportParameters, Color)
 	SHADER_PARAMETER_STRUCT(FScreenPassTextureViewportParameters, Output)
@@ -204,6 +205,7 @@ FScreenPassTexture AddDeviceEncodingOnlyPass(FRDGBuilder& GraphBuilder, const FV
 
 	FDeviceEncodingOnlyParameters CommonParameters;
 	CommonParameters.View = View.ViewUniformBuffer;
+	CommonParameters.WorkingColorSpace = GDefaultWorkingColorSpaceUniformBuffer.GetUniformBufferRef();
 	CommonParameters.OutputDevice = GetDeviceEncodingOnlyOutputDeviceParameters(ViewFamily);
 	CommonParameters.Color = GetScreenPassTextureViewportParameters(SceneColorViewport);
 	CommonParameters.Output = GetScreenPassTextureViewportParameters(OutputViewport);
