Add an alternate subsurface shadow mode for testing, hidden behind a (temporary) cvar
Also minor optimization to local light filtering when source radius = 0.

diff --git a/Engine/Shaders/Private/VirtualShadowMaps/VirtualShadowMapProjection.usf b/Engine/Shaders/Private/VirtualShadowMaps/VirtualShadowMapProjection.usf
--- a/Engine/Shaders/Private/VirtualShadowMaps/VirtualShadowMapProjection.usf
+++ b/Engine/Shaders/Private/VirtualShadowMaps/VirtualShadowMapProjection.usf
@@ -206,6 +206,7 @@ float4 ComputeRandom4(uint2 PixelPosition)
 int4 ProjectionRect;
 float ScreenRayLength;
 float NormalBias;
+int SubsurfaceShadowMode;			// Temp for testing and A/B comparison
 int SMRTRayCount;					// 0 = off
 int SMRTSamplesPerRay;
 float SMRTRayLengthScale;			// Directional lights
@@ -496,6 +497,12 @@ void OutputVisualize( uint2 PixelPos, FVirtualShadowMapSampleResult Result )
 		{
 			Output = 0.8f*Result.GeneralDebug + 0.2f*Result.ShadowFactor.xxx;
 		}
+		/* TODO: Enable this and add to menu
+		else if ( VisualizeModeId == VIRTUAL_SHADOW_MAP_VISUALIZE_OCCLUDER_DISTANCE )
+		{
+			Output = (Result.OccluderDistance / 1000.0f).xxx;
+		}
+		*/
 		/*
 		else if ( VisualizeModeId == VIRTUAL_SHADOW_MAP_VISUALIZE_CACHED_PAGE_AGE && bValidPageAddress )
 		{
@@ -519,6 +526,21 @@ void OutputVisualize( uint2 PixelPos, FVirtualShadowMapSampleResult Result )
 	}
 }
 
+// TODO: Refactor/move this to transmission common
+float GetShadowFactorSubsurface(float ShadowFactor, float OccluderDistance, float SubsurfaceOpacity)
+{
+	if ( ShadowFactor < 1.0f )
+	{
+		float Density = SubsurfaceDensityFromOpacity( SubsurfaceOpacity );
+		float Occlusion = 1.0f - saturate( exp( -Density * OccluderDistance ) );
+		float NewShadowFactor = 1.0f - Occlusion * ( 1.0f - ShadowFactor );
+
+		// NOTE: This square is arbitrary, but we have to be consistent with the PCF path
+		// TODO: This might be the time to remove this...
+		return SubsurfaceShadowMode > 1 ? Square( NewShadowFactor ) : NewShadowFactor;
+	}
+	return ShadowFactor;
+}
 
 [numthreads(8, 8, 1)]
 void VirtualShadowMapProjection(
@@ -617,17 +639,25 @@ void VirtualShadowMapProjection(
 			TranslatedWorldPosition,
 			Noise);
 		
-		float SSSTransmission = Result.ShadowFactor;
+		float TransmissionShadow = Result.ShadowFactor;
 		const float SubsurfaceOpacity = bIsHairInput ? 1.0f : GetSubsurfaceOpacityFromGbuffer( PixelPos );
 		if ( SubsurfaceOpacity < 1.0f )
 		{
-			SSSTransmission = ComputeSimpleSubsurfaceTransmissionFromVirtualShadowMap(
-				VirtualShadowMapId,
-				TranslatedWorldPosition,
-				SubsurfaceOpacity);
+			if (SubsurfaceShadowMode > 0)
+			{
+				TransmissionShadow = GetShadowFactorSubsurface(Result.ShadowFactor, Result.OccluderDistance, SubsurfaceOpacity);
+				Result.ShadowFactor = TransmissionShadow;
+			}
+			else
+			{
+				TransmissionShadow = ComputeSimpleSubsurfaceTransmissionFromVirtualShadowMap(
+					VirtualShadowMapId,
+					TranslatedWorldPosition,
+					SubsurfaceOpacity);
+			}
 		}
 
-		OutShadowFactor[ PixelPos ] = float2( Result.ShadowFactor, SSSTransmission );
+		OutShadowFactor[ PixelPos ] = float2( Result.ShadowFactor, TransmissionShadow );
 
 		if (VisualizeVirtualShadowMapId == VirtualShadowMapId)
 		{
diff --git a/Engine/Shaders/Private/VirtualShadowMaps/VirtualShadowMapProjectionDirectional.ush b/Engine/Shaders/Private/VirtualShadowMaps/VirtualShadowMapProjectionDirectional.ush
--- a/Engine/Shaders/Private/VirtualShadowMaps/VirtualShadowMapProjectionDirectional.ush
+++ b/Engine/Shaders/Private/VirtualShadowMaps/VirtualShadowMapProjectionDirectional.ush
@@ -229,6 +229,8 @@ FVirtualShadowMapSampleResult TraceDirectional(
 
 		uint RayMissCount = 0;
 		uint i = 0;
+		float OccluderDistanceSum = 0.0f;
+		float MaxOccluderDistance = -1.0f;
 		for ( ; i < MaxRayCount; i++)
 		{
 			// One sample for ray, one for texel dither
@@ -258,7 +260,20 @@ FVirtualShadowMapSampleResult TraceDirectional(
 			Result.PhysicalTexelAddress =  RayState.PhysicalTexelAddress;
 			//Result.GeneralDebug = GreenToRedTurbo((RayState.SampledClipmapId - RayState.ClipmapId) / 2.0f);
 			
-			RayMissCount += SMRTResult.bValidHit ? 0 : 1;
+			if (SMRTResult.bValidHit)
+			{
+				float OccluderDistance = ComputeOccluderDistanceOrtho(
+					GetVirtualShadowMapProjectionData(RayState.ClipmapId).ShadowViewToClipMatrix,
+					SMRTResult.HitDepth,
+					RayState.RayStartUVZ.z);
+
+				OccluderDistanceSum += OccluderDistance;
+				MaxOccluderDistance = max(MaxOccluderDistance, OccluderDistance);
+			}
+			else
+			{
+				++RayMissCount;
+			}
 
 			if (bUseAdaptiveRayCount)
 			{
@@ -286,8 +301,15 @@ FVirtualShadowMapSampleResult TraceDirectional(
 		}
 		uint RayCount = min(i + 1, MaxRayCount);		// break vs regular for loop exit
 
+		/* TODO: Experiment with max vs avg distance in various cases
+		float OccluderDistance = (View.GeneralPurposeTweak == 0)
+			? MaxOccluderDistance
+			: (OccluderDistanceSum / float(max(1, RayCount - RayMissCount)));
+		*/
+		float OccluderDistance = (OccluderDistanceSum / float(max(1, RayCount - RayMissCount)));
+
 		Result.ShadowFactor = SMRTComputeRayMissFactor(RayMissCount, RayCount, MaxRayCount, Noise);
-		Result.OccluderDistance = -1.0f;	// TODO?
+		Result.OccluderDistance = OccluderDistance;
 		Result.RayCount = RayCount;
 	}
 
diff --git a/Engine/Shaders/Private/VirtualShadowMaps/VirtualShadowMapProjectionSpot.ush b/Engine/Shaders/Private/VirtualShadowMaps/VirtualShadowMapProjectionSpot.ush
--- a/Engine/Shaders/Private/VirtualShadowMaps/VirtualShadowMapProjectionSpot.ush
+++ b/Engine/Shaders/Private/VirtualShadowMaps/VirtualShadowMapProjectionSpot.ush
@@ -385,6 +385,12 @@ FVirtualShadowMapSampleResult TraceLocalLight(
 		const float3 NormalPlaneTangent = mul(ConeTangentToWorldMatrix, WorldNormal);
 		const float2 DepthSlopeTangent = -NormalPlaneTangent.xy / NormalPlaneTangent.z;
 
+		// If source angle = 0, we don't need multiple samples on the same ray
+		if (ConeSin == 0.0f)
+		{
+			SamplesPerRay = 0;
+		}
+
 		uint i = 0;
 		uint RayMissCount = 0;
 		for( ; i < MaxRayCount; i++ )
diff --git a/Engine/Shaders/Private/VirtualShadowMaps/VirtualShadowMapSMRTCommon.ush b/Engine/Shaders/Private/VirtualShadowMaps/VirtualShadowMapSMRTCommon.ush
--- a/Engine/Shaders/Private/VirtualShadowMaps/VirtualShadowMapSMRTCommon.ush
+++ b/Engine/Shaders/Private/VirtualShadowMaps/VirtualShadowMapSMRTCommon.ush
@@ -29,6 +29,7 @@ FSMRTSample InitSMRTSample()
 struct FSMRTResult
 {
 	bool bValidHit;
+	float HitDepth;	// Valid if bValidHit is true
 };
 
 float SMRTComputeRayMissFactor(uint RayMissCount, uint RayCount, uint MaxRayCount, half Noise)
diff --git a/Engine/Shaders/Private/VirtualShadowMaps/VirtualShadowMapSMRTTemplate.ush b/Engine/Shaders/Private/VirtualShadowMaps/VirtualShadowMapSMRTTemplate.ush
--- a/Engine/Shaders/Private/VirtualShadowMaps/VirtualShadowMapSMRTTemplate.ush
+++ b/Engine/Shaders/Private/VirtualShadowMaps/VirtualShadowMapSMRTTemplate.ush
@@ -38,7 +38,6 @@ FSMRTResult SMRTRayCast(
 	// This doesn't get used in the first iteration of the loop and so will always
 	// be valid by the time that it does.
 	float PrevReferenceDepth = -1;
-	float PrevSampleTime = -1;
 
 	bool bValidHit =  false;
 #if SMRT_TEMPLATE_STATIC_SAMPLES_PER_RAY >= 0
@@ -67,8 +66,10 @@ FSMRTResult SMRTRayCast(
 				DepthHistoryTime = SampleTime;
 				if (SampleDepth > ReferenceDepth)
 				{
-					bValidHit = true;
-					break;
+					FSMRTResult Result;
+					Result.bValidHit = true;
+					Result.HitDepth = SampleDepth;
+					return Result;
 				}
 			}
 			else
@@ -97,7 +98,7 @@ FSMRTResult SMRTRayCast(
 				{
 					if (SampleDepth != DepthHistory)
 					{
-						// NOTE: DCE will remove all this if SMRT_EXTRAPOLATE_WITH_SLOPE is false
+						// NOTE: DCE will remove this if SMRT_EXTRAPOLATE_WITH_SLOPE is false
 						const float SlopeClamp = Sample.ExtrapolateSlope;
 						DepthSlope = (SampleDepth - DepthHistory) / DeltaHistoryTime;
 						DepthSlope = clamp(DepthSlope, -SlopeClamp, SlopeClamp);
@@ -112,17 +113,19 @@ FSMRTResult SMRTRayCast(
 				bool bHit = abs(DepthDiff + HalfCompareTolerance) < HalfCompareTolerance;
 				if (bHit)
 				{
-					bValidHit = true;
-					break;
+					FSMRTResult Result;
+					Result.bValidHit = true;
+					Result.HitDepth = DepthForComparison;
+					return Result;
 				}
 			}
 
-			PrevSampleTime = SampleTime;
 			PrevReferenceDepth = ReferenceDepth;
 		}
 	}
 
 	FSMRTResult Result;
-	Result.bValidHit = bValidHit;
+	Result.bValidHit = false;
+	Result.HitDepth = -1.0f;
 	return Result;
 }
diff --git a/Engine/Source/Runtime/Renderer/Private/VirtualShadowMaps/VirtualShadowMapProjection.cpp b/Engine/Source/Runtime/Renderer/Private/VirtualShadowMaps/VirtualShadowMapProjection.cpp
--- a/Engine/Source/Runtime/Renderer/Private/VirtualShadowMaps/VirtualShadowMapProjection.cpp
+++ b/Engine/Source/Runtime/Renderer/Private/VirtualShadowMaps/VirtualShadowMapProjection.cpp
@@ -151,6 +151,14 @@ static TAutoConsoleVariable<int32> CVarForcePerLightShadowMaskClear(
 	ECVF_RenderThreadSafe
 );
 
+static TAutoConsoleVariable<int32> CVarSubsurfaceShadowMode(
+	TEXT( "r.Shadow.Virtual.SubsurfaceShadowMode" ),
+	0,
+	TEXT( "" ),
+	ECVF_RenderThreadSafe
+);
+
+
 #if MAX_TEST_PERMUTATION > 0
 static TAutoConsoleVariable<int32> CVarTestPermutation(
 	TEXT( "r.Shadow.Virtual.ProjectionTestPermutation" ),
@@ -267,8 +275,9 @@ class FVirtualShadowMapProjectionCS : public FGlobalShader
 		SHADER_PARAMETER(FIntVector4, ProjectionRect)
 		SHADER_PARAMETER(float, ScreenRayLength)
 		SHADER_PARAMETER(float, NormalBias)
-		SHADER_PARAMETER(uint32, SMRTRayCount)
-		SHADER_PARAMETER(uint32, SMRTSamplesPerRay)
+		SHADER_PARAMETER(int32, SubsurfaceShadowMode)
+		SHADER_PARAMETER(int32, SMRTRayCount)
+		SHADER_PARAMETER(int32, SMRTSamplesPerRay)
 		SHADER_PARAMETER(float, SMRTRayLengthScale)
 		SHADER_PARAMETER(float, SMRTCotMaxRayAngleFromLight)
 		SHADER_PARAMETER(float, SMRTTexelDitherScale)
@@ -353,6 +362,7 @@ static void RenderVirtualShadowMapProjectionCommon(
 	PassParameters->ProjectionRect = FIntVector4(ProjectionRect.Min.X, ProjectionRect.Min.Y, ProjectionRect.Max.X, ProjectionRect.Max.Y);
 	PassParameters->ScreenRayLength = CVarScreenRayLength.GetValueOnRenderThread();
 	PassParameters->NormalBias = GetNormalBiasForShader();
+	PassParameters->SubsurfaceShadowMode = CVarSubsurfaceShadowMode.GetValueOnRenderThread();
 	PassParameters->InputType = uint32(InputType);
 	PassParameters->bCullBackfacingPixels = VirtualShadowMapArray.ShouldCullBackfacingPixels() ? 1 : 0;
 	PassParameters->bSMRTUseAdaptiveRayCount = CVarSMRTAdaptiveRayCount.GetValueOnRenderThread() != 0 ? 1 : 0;
