Merge all hair debug code into HairStandsDebug.usf

diff --git a/Engine/Plugins/Runtime/HairStrands/Source/HairStrandsCore/Private/GroomDebug.cpp b/Engine/Plugins/Runtime/HairStrands/Source/HairStrandsCore/Private/GroomDebug.cpp
--- a/Engine/Plugins/Runtime/HairStrands/Source/HairStrandsCore/Private/GroomDebug.cpp
+++ b/Engine/Plugins/Runtime/HairStrands/Source/HairStrandsCore/Private/GroomDebug.cpp
@@ -174,9 +174,15 @@ class FHairProjectionMeshDebugCS : public FGlobalShader
 	{
 		return IsHairStrandsSupported(EHairStrandsShaderType::Tool, Parameters.Platform);
 	}
+
+	static void ModifyCompilationEnvironment(const FGlobalShaderPermutationParameters& Parameters, FShaderCompilerEnvironment& OutEnvironment)
+	{
+		FGlobalShader::ModifyCompilationEnvironment(Parameters, OutEnvironment);
+		OutEnvironment.SetDefine(TEXT("SHADER_MESH_PROJECTION_SKIN_CACHE"), 1);
+	}
 };
 
-IMPLEMENT_GLOBAL_SHADER(FHairProjectionMeshDebugCS, "/Engine/Private/HairStrands/HairStrandsMeshProjectionMeshDebug.usf", "MainCS", SF_Compute);
+IMPLEMENT_GLOBAL_SHADER(FHairProjectionMeshDebugCS, "/Engine/Private/HairStrands/HairStrandsDebug.usf", "MainCS", SF_Compute);
 
 static void AddDebugProjectionMeshPass(
 	FRDGBuilder& GraphBuilder,
@@ -274,9 +280,15 @@ class FHairProjectionHairDebugCS : public FGlobalShader
 	{
 		return IsHairStrandsSupported(EHairStrandsShaderType::Tool, Parameters.Platform);
 	}
+
+	static void ModifyCompilationEnvironment(const FGlobalShaderPermutationParameters& Parameters, FShaderCompilerEnvironment& OutEnvironment)
+	{
+		FGlobalShader::ModifyCompilationEnvironment(Parameters, OutEnvironment);
+		OutEnvironment.SetDefine(TEXT("SHADER_MESH_PROJECTION_HAIR"), 1);
+	}
 };
 
-IMPLEMENT_GLOBAL_SHADER(FHairProjectionHairDebugCS, "/Engine/Private/HairStrands/HairStrandsMeshProjectionHairDebug.usf", "MainCS", SF_Compute);
+IMPLEMENT_GLOBAL_SHADER(FHairProjectionHairDebugCS, "/Engine/Private/HairStrands/HairStrandsDebug.usf", "MainCS", SF_Compute);
 
 enum class EDebugProjectionHairType
 {
@@ -411,11 +423,11 @@ public:
 	static void ModifyCompilationEnvironment(const FGlobalShaderPermutationParameters& Parameters, FShaderCompilerEnvironment& OutEnvironment)
 	{
 		FGlobalShader::ModifyCompilationEnvironment(Parameters, OutEnvironment);
-		OutEnvironment.SetDefine(TEXT("SHADER_PLAIN"), 1);
+		OutEnvironment.SetDefine(TEXT("SHADER_CARDS_VOXEL"), 1);
 	}	
 };
 
-IMPLEMENT_GLOBAL_SHADER(FVoxelPlainRaymarchingCS, "/Engine/Private/HairStrands/HairCardsVoxel.usf", "MainCS", SF_Compute);
+IMPLEMENT_GLOBAL_SHADER(FVoxelPlainRaymarchingCS, "/Engine/Private/HairStrands/HairStrandsDebug.usf", "MainCS", SF_Compute);
 
 static void AddVoxelPlainRaymarchingPass(
 	FRDGBuilder& GraphBuilder,
@@ -489,11 +501,11 @@ public:
 	static void ModifyCompilationEnvironment(const FGlobalShaderPermutationParameters& Parameters, FShaderCompilerEnvironment& OutEnvironment)
 	{
 		FGlobalShader::ModifyCompilationEnvironment(Parameters, OutEnvironment);
-		OutEnvironment.SetDefine(TEXT("SHADER_ATLAS"), 1);
+		OutEnvironment.SetDefine(TEXT("SHADER_CARDS_ATLAS"), 1);
 	}
 };
 
-IMPLEMENT_GLOBAL_SHADER(FDrawDebugCardAtlasCS, "/Engine/Private/HairStrands/HairCardsDebug.usf", "MainCS", SF_Compute);
+IMPLEMENT_GLOBAL_SHADER(FDrawDebugCardAtlasCS, "/Engine/Private/HairStrands/HairStrandsDebug.usf", "MainCS", SF_Compute);
 
 static void AddDrawDebugCardsAtlasPass(
 	FRDGBuilder& GraphBuilder,
@@ -648,11 +660,11 @@ public:
 	static void ModifyCompilationEnvironment(const FGlobalShaderPermutationParameters& Parameters, FShaderCompilerEnvironment& OutEnvironment)
 	{
 		FGlobalShader::ModifyCompilationEnvironment(Parameters, OutEnvironment);
-		OutEnvironment.SetDefine(TEXT("SHADER_GUIDE"), 1);
+		OutEnvironment.SetDefine(TEXT("SHADER_CARDS_GUIDE"), 1);
 	}
 };
 
-IMPLEMENT_GLOBAL_SHADER(FDrawDebugCardGuidesCS, "/Engine/Private/HairStrands/HairCardsDebug.usf", "MainCS", SF_Compute);
+IMPLEMENT_GLOBAL_SHADER(FDrawDebugCardGuidesCS, "/Engine/Private/HairStrands/HairStrandsDebug.usf", "MainCS", SF_Compute);
 
 static void AddDrawDebugCardsGuidesPass(
 	FRDGBuilder& GraphBuilder,
@@ -839,7 +851,7 @@ public:
 	}
 };
 
-IMPLEMENT_GLOBAL_SHADER(FHairDebugPrintInstanceCS, "/Engine/Private/HairStrands/HairStrandsDebugPrint.usf", "MainCS", SF_Compute);
+IMPLEMENT_GLOBAL_SHADER(FHairDebugPrintInstanceCS, "/Engine/Private/HairStrands/HairStrandsDebug.usf", "MainCS", SF_Compute);
 
 struct FHairDebugNameInfo
 {
diff --git a/Engine/Shaders/Private/HairStrands/HairCardsDebug.usf b/Engine/Shaders/Private/HairStrands/HairCardsDebug.usf
deleted file mode 100644
diff --git a/Engine/Shaders/Private/HairStrands/HairCardsVoxel.usf b/Engine/Shaders/Private/HairStrands/HairCardsVoxel.usf
deleted file mode 100644
diff --git a/Engine/Shaders/Private/HairStrands/HairStrandsBsdfPlot.usf b/Engine/Shaders/Private/HairStrands/HairStrandsBsdfPlot.usf
deleted file mode 100644
diff --git a/Engine/Shaders/Private/HairStrands/HairStrandsDebug.usf b/Engine/Shaders/Private/HairStrands/HairStrandsDebug.usf
--- a/Engine/Shaders/Private/HairStrands/HairStrandsDebug.usf
+++ b/Engine/Shaders/Private/HairStrands/HairStrandsDebug.usf
@@ -1,13 +1,15 @@
 // Copyright Epic Games, Inc. All Rights Reserved.
 
 #include "../Common.ush"
-#include "../ColorMap.ush"
-#include "HairStrandsVisibilityCommon.ush"
 
 /////////////////////////////////////////////////////////////////////////////////////////
+// Visualize hair strands data as full screen view mode (BaseColor, Tangent, Roughness, ...)
 
 #if SHADER_DEBUG_MODE
 
+#include "../ColorMap.ush"
+#include "HairStrandsVisibilityCommon.ush"
+
 float2 OutputResolution;
 uint FastResolveMask;
 uint DebugMode;
@@ -149,6 +151,7 @@ void MainPS(
 #endif // SHADER_DEBUG_MODE
 
 /////////////////////////////////////////////////////////////////////////////////////////
+// Visualize hair strands control points
 
 #if SHADER_CVS
 #include "HairStrandsVertexFactoryCommon.ush"
@@ -201,4 +204,2035 @@ void MainCS(uint3 DispatchThreadId : SV_DispatchThreadID)
 	}
 }
 
-#endif // SHADER_CVS
\ No newline at end of file
+#endif // SHADER_CVS
+
+
+/////////////////////////////////////////////////////////////////////////////////////////
+// Visualize (Sim/Render & Rest/Deformed):
+// * hair strands root points
+// * skin triangle on which strands are attached to
+// * sample point used for RBF
+
+#if SHADER_MESH_PROJECTION_HAIR
+
+#include "HairStrandsMeshProjectionCommon.ush" 
+#include "../ShaderPrint.ush"
+
+#define INPUT_ROOT_FRAME 0
+#define INPUT_ROOT_TRIANGLE 1
+#define INPUT_ROOT_SAMPLES 2
+
+uint MaxRootCount;
+float4x4 RootLocalToWorld;
+uint DeformedFrameEnable;
+
+#if PERMUTATION_INPUT_TYPE == INPUT_ROOT_FRAME
+Buffer<uint>   RootBarycentricBuffer;
+#endif
+
+Buffer<float4>  RestPosition0Buffer;
+Buffer<float4>  RestPosition1Buffer;
+Buffer<float4>  RestPosition2Buffer;
+
+Buffer<float4>  DeformedPosition0Buffer;
+Buffer<float4>  DeformedPosition1Buffer;
+Buffer<float4>  DeformedPosition2Buffer;
+
+Buffer<float4> RestSamplePositionsBuffer;
+Buffer<float4> DeformedSamplePositionsBuffer;
+
+[numthreads(256, 1, 1)]
+void MainCS(uint3 DispatchThreadId : SV_DispatchThreadID)
+{
+	ResolvedView = ResolveView();
+
+#if PERMUTATION_INPUT_TYPE == INPUT_ROOT_FRAME  
+	const uint LineIndex = DispatchThreadId.x;
+
+	if (LineIndex > MaxRootCount) return;
+	{
+		const FHairMeshTriangleWithUVs RestTriangle		= GetTriangleWithUVsTransformation(LineIndex, RestPosition0Buffer, RestPosition1Buffer, RestPosition2Buffer);
+		const FHairMeshTriangleWithUVs DeformedTriangle = GetTriangleWithUVsTransformation(LineIndex, DeformedPosition0Buffer, DeformedPosition1Buffer, DeformedPosition2Buffer);
+
+		FHairMeshTriangleWithUVs Tri;
+		if (DeformedFrameEnable > 0)
+			Tri = DeformedTriangle;
+		else
+			Tri = RestTriangle;
+
+		const float3 B  = DecodeBarycentrics(RootBarycentricBuffer[LineIndex]);
+		float3 P0 = B.x * Tri.P0 + B.y * Tri.P1 + B.z * Tri.P2;
+		float3 P1 = P0 + Tri.N;
+
+		const float2 UVs = 
+			B.x * Tri.UV0 +
+			B.y * Tri.UV1 +
+			B.z * Tri.UV2;
+
+		P0 = mul(float4(P0, 1), RootLocalToWorld).xyz;
+		P1 = mul(float4(P1, 1), RootLocalToWorld).xyz;
+
+		const float3 N = normalize(P0 - P1);
+		const float4 Color = float4(UVs, 0, 1);
+		//Color = float4(N* 0.5f + 0.5f, 1);
+		AddLineWS(P0, P1, Color);
+	}
+#endif // INPUT_ROOT_FRAME
+
+#if PERMUTATION_INPUT_TYPE == INPUT_ROOT_TRIANGLE
+	const uint TriangleIndex = DispatchThreadId.x;
+
+	if (TriangleIndex > MaxRootCount) return;
+
+	const FHairMeshTriangleWithUVs RestTriangle		= GetTriangleWithUVsTransformation(TriangleIndex, RestPosition0Buffer, RestPosition1Buffer, RestPosition2Buffer);
+	const FHairMeshTriangleWithUVs DeformedTriangle = GetTriangleWithUVsTransformation(TriangleIndex, DeformedPosition0Buffer, DeformedPosition1Buffer, DeformedPosition2Buffer);
+
+	FHairMeshTriangleWithUVs Tri;
+	if (DeformedFrameEnable > 0)
+		Tri = DeformedTriangle;
+	else
+		Tri = RestTriangle;
+	
+
+	// Debug: To test if the rest-triangle-to-deformed-triangle transformation works correctly
+	#if 0
+	if (DeformedFrameEnable > 0)
+	{
+		float3 Barycentric = 0;
+		if (TriangleVertex == 0)
+		{
+			P = RestTriangle.P0;
+			Barycentric = float3(1, 0, 0);
+		}
+		if (TriangleVertex == 1)
+		{
+			P = RestTriangle.P1;
+			Barycentric = float3(0, 1, 0);
+		}
+		if (TriangleVertex == 2)
+		{
+			P = RestTriangle.P2;
+			Barycentric = float3(0, 0, 1);
+		}
+		P = TransformPoint(P, Barycentric, RestTriangle, DeformedTriangle);
+	}
+	#endif
+
+	float4 Color = float4(Tri.UV0, 0, 1);
+	// Color = float4((Tri.N + 1) * 0.5f, 1);
+
+	AddFilledTriangleWS(
+		mul(float4(Tri.P0, 1), RootLocalToWorld).xyz, 
+		mul(float4(Tri.P1, 1), RootLocalToWorld).xyz, 
+		mul(float4(Tri.P2, 1), RootLocalToWorld).xyz, 
+		Color);
+#endif // INPUT_ROOT_TRIANGLE
+
+#if PERMUTATION_INPUT_TYPE == INPUT_ROOT_SAMPLES 
+	const uint PointIndex = DispatchThreadId.x;
+	
+	if (PointIndex > MaxRootCount) return;
+
+	float3 P = (DeformedFrameEnable > 0) ? DeformedSamplePositionsBuffer[PointIndex].xyz : RestSamplePositionsBuffer[PointIndex].xyz;
+	P = mul(float4(P, 1), RootLocalToWorld).xyz;
+
+	AddCrossWS(P, 1.f, ColorYellow);
+#endif // INPUT_ROOT_SAMPLES
+}
+#endif // SHADER_MESH_PROJECTION_HAIR
+
+
+/////////////////////////////////////////////////////////////////////////////////////////
+// Visualize skin cache data, extracted & processed for hair strands deformation (Rest & Deformed)
+
+#if SHADER_MESH_PROJECTION_SKIN_CACHE
+
+#include "HairStrandsMeshProjectionCommon.ush"
+#include "../ShaderPrint.ush"
+
+float4x4 LocalToWorld;
+uint VertexOffset;
+uint IndexOffset;
+uint MaxIndexCount;
+uint MaxVertexCount;
+uint MeshUVsChannelOffset;
+uint MeshUVsChannelCount;
+uint bOutputInUVsSpace;
+uint MeshType;
+uint SectionIndex;
+float2 OutputResolution;
+
+Buffer<float> InputVertexPositionBuffer;
+Buffer<float2> InputVertexUVsBuffer;
+#if PERMUTATION_INPUT_TYPE == 1
+Buffer<uint> InputIndexBuffer;
+#endif
+
+#define REST_MESH 0
+#define DEFORMED_MESH 1
+#define SOURCE_MESH 2
+#define TARGET_MESH 3
+
+[numthreads(256, 1, 1)]
+void MainCS(uint3 DispatchThreadId : SV_DispatchThreadID)
+{
+	ResolvedView = ResolveView();
+
+	float4 Color = 0;
+	if (MeshType == REST_MESH)		Color = float4(1, 0, 0, 1);
+	if (MeshType == DEFORMED_MESH)	Color = float4(0, 1, 0, 1);
+	if (MeshType == SOURCE_MESH)	Color = float4(0, 0, 1, 1);
+	if (MeshType == TARGET_MESH)	Color = float4(1, 1, 0, 1);
+
+#if PERMUTATION_INPUT_TYPE == 0
+	const uint VertexId = DispatchThreadId.x * 3;
+	const uint VertexIndex0 = min(VertexOffset + VertexId, MaxVertexCount-1);
+
+	FHairMeshTriangleWithUVs Tri = (FHairMeshTriangleWithUVs)0;
+	Tri.P0 = InputVertexPositionBuffer[VertexIndex0 + 0];
+	Tri.P1 = InputVertexPositionBuffer[VertexIndex0 + 1];
+	Tri.P2 = InputVertexPositionBuffer[VertexIndex0 + 2];
+
+	Tri.UV0 = InputVertexUVsBuffer[VertexIndex0 + 0];
+	Tri.UV1 = InputVertexUVsBuffer[VertexIndex0 + 1];
+	Tri.UV2 = InputVertexUVsBuffer[VertexIndex0 + 2];
+
+	Tri.N = cross(Tri.P1 - Tri.P0, Tri.P2 - Tri.P0);
+#endif
+
+#if PERMUTATION_INPUT_TYPE == 1
+	const uint TriangleIndex = DispatchThreadId.x;
+	const FHairMeshTriangleWithUVs Tri = GetTriangleData(TriangleIndex, InputVertexPositionBuffer, InputVertexUVsBuffer, InputIndexBuffer, IndexOffset, MaxIndexCount, MaxVertexCount, MeshUVsChannelOffset, MeshUVsChannelCount);
+#endif
+
+	// Quick filter
+#if 0
+	//if (MeshType != REST_MESH)
+	//if (MeshType != DEFORMED_MESH)
+	//if (MeshType != SOURCE_MESH)
+	//if (MeshType != TARGET_MESH)
+	if (MeshType != SOURCE_MESH && MeshType != TARGET_MESH || SectionIndex != 0)
+	{
+		return;
+	}
+#endif
+
+	Color.xyz *= Tri.N * 0.5f + 0.5f;
+
+	if (bOutputInUVsSpace)
+	{
+		AddFilledTriangleWS(float3(Tri.UV0, 0), float3(Tri.UV1, 0), float3(Tri.UV2, 0), Color);
+	}
+	else
+	{
+		AddFilledTriangleWS(
+			mul(float4(Tri.P0, 1), LocalToWorld).xyz, 
+			mul(float4(Tri.P1, 1), LocalToWorld).xyz, 
+			mul(float4(Tri.P2, 1), LocalToWorld).xyz,
+			Color);
+	}
+}
+
+#endif // SHADER_MESH_PROJECTION_SKIN_CACHE
+
+/////////////////////////////////////////////////////////////////////////////////////////
+// Visualized stats for PPLL primary visibility
+
+#if SHADER_PPLL_DEBUG
+
+#include "../ShaderPrint.ush"
+#include "HairStrandsVisibilityCommon.ush" 
+
+float PPLLMeanListElementCountPerPixel;
+float PPLLMaxTotalListElementCount;
+
+Texture2D<uint> PPLLCounter;
+Texture2D<uint> PPLLNodeIndex;
+StructuredBuffer<FPackedHairVisPPLL> PPLLNodeData;
+
+RWTexture2D<float4> SceneColorTextureUAV;
+
+FFontColor GetOccupancyColor(float In)
+{
+	float3 Color = lerp(float3(0, 1, 0), float3(1, 0, 0), saturate(In));
+	return InitFontColor(Color);
+}
+
+[numthreads(8, 8, 1)]
+void VisibilityDebugPPLLCS(uint GroupIndex : SV_GroupIndex, uint3 DispatchThreadId : SV_DispatchThreadID)
+{
+	int2 PixelPos = DispatchThreadId.xy;
+	float PPLLListElementcount = PPLLCounter[uint2(0, 0)];
+
+	uint NextNodeIndex = PPLLNodeIndex[PixelPos];
+	if (NextNodeIndex != 0xFFFFFFFF)
+	{
+		uint StrandCount = 0;
+		do
+		{
+			StrandCount++;
+			FPackedHairVisPPLL Node = PPLLNodeData[NextNodeIndex];
+			NextNodeIndex = Node.NextNodeIndex;
+		} while (NextNodeIndex != 0xFFFFFFFF);
+
+		float l = saturate(StrandCount / PPLLMeanListElementCountPerPixel);
+
+		SceneColorTextureUAV[PixelPos] = float4(l, 1.0f - l, 0.0f, 0.0f);
+	}
+
+	const float2 BarOrigin = 100.0f;
+	const float2 BarSize = float2(200.0f, 20.0f);
+	if (all(PixelPos >= BarOrigin) && all(PixelPos <= (BarOrigin + BarSize)))
+	{
+		float2 LocalPos = PixelPos - BarOrigin;
+		float MemoryUsage = PPLLListElementcount / PPLLMaxTotalListElementCount;
+
+		if (MemoryUsage > LocalPos.x / BarSize.x)
+		{
+			SceneColorTextureUAV[PixelPos] = float4(MemoryUsage, 1.0f - MemoryUsage, 1.0f - MemoryUsage, 0.0f);
+		}
+		else
+		{
+			SceneColorTextureUAV[PixelPos] = float4(0.2f, 0.2f, 0.2f, 0.0f);
+		}
+
+		if (any(PixelPos == BarOrigin) || any(PixelPos == (BarOrigin + BarSize)))
+		{
+			SceneColorTextureUAV[PixelPos] = float4(1.0f, 1.0f, 1.0f, 0.0f);
+		}
+	}
+
+	if (DispatchThreadId.x == 0 && DispatchThreadId.y == 0)
+	{
+		ShaderPrintFilter(true);
+
+		const float PPLLNodeSizeInMb = (4 + 4 + 4) / 1000000.f;
+		const float NodePercentage = float(PPLLListElementcount) / float(PPLLMaxTotalListElementCount);
+
+		float2 StartPos = (BarOrigin + BarSize*float2(0.0f, 2.0f)) * View.ViewSizeAndInvSize.zw;
+		float2 Pos = StartPos;
+		Pos = Print(Pos, TEXT("Mean Node Per Pixel : "));
+		Pos = ShaderPrint(Pos, PPLLMeanListElementCountPerPixel, FontEmerald);
+		Pos = ShaderPrintNewline(Pos); Pos.x = StartPos.x;
+		
+		Pos = Print(Pos, TEXT("Nodes Used (%)      : "));
+		Pos = ShaderPrint(Pos, NodePercentage * 100.f, GetOccupancyColor(NodePercentage));
+		Pos = ShaderPrintNewline(Pos); Pos.x = StartPos.x;
+
+		Pos = Print(Pos, TEXT("Nodes Used          : "));
+		Pos = ShaderPrint(Pos, PPLLListElementcount, FontSilver);
+		Pos = ShaderPrintNewline(Pos); Pos.x = StartPos.x;
+
+		Pos = Print(Pos, TEXT("Allocated Nodes (Mb): "));
+		Pos = ShaderPrint(Pos, PPLLMaxTotalListElementCount * PPLLNodeSizeInMb, FontSilver);
+		Pos = ShaderPrintNewline(Pos); Pos.x = StartPos.x;
+
+		Pos = Print(Pos, TEXT("Total               : "));
+		Pos = ShaderPrint(Pos, PPLLMaxTotalListElementCount, FontOrange);
+		Pos = ShaderPrintNewline(Pos); Pos.x = StartPos.x;
+	}
+}
+
+#endif // SHADER_PPLL_DEBUG
+
+//////////////////////////////////////////////////////////////////////////////////////////
+// Visualized Deep Opacity Map stats and buffer
+
+#if SHADER_VISUALIZEDOM
+
+#include "../PostProcessCommon.ush"
+
+float DomScale = 1;
+float2 DomAtlasOffset;
+float2 DomAtlasScale;
+float2 OutputResolution;
+float2 InvOutputResolution;
+
+int4 HairViewRect;
+
+Texture2D<float> DeepShadowDepthTexture;
+Texture2D<float4> DeepShadowLayerTexture;
+
+SamplerState LinearSampler;
+
+bool IsInsideRect(float4 Rect, float2 PixelPos, inout float2 LocalUV)
+{
+	LocalUV = (PixelPos - Rect.xy) / Rect.zw;
+
+	return 	PixelPos.x > Rect.x && PixelPos.x < Rect.x + Rect.z &&
+			PixelPos.y > Rect.y && PixelPos.y < Rect.y + Rect.w;
+}
+
+bool IsOnBorder(float2 LocalUV, float4 Rect)
+{
+	float2 PixelPos = LocalUV * Rect.zw;
+	float BorderThickness = 1;
+
+	return
+		PixelPos.x < BorderThickness || PixelPos.x > Rect.z - 1 - BorderThickness ||
+		PixelPos.y < BorderThickness || PixelPos.y > Rect.w - 1 - BorderThickness;
+}
+
+float3 AddBorder(float2 LocalUV, float4 Rect)
+{
+	const bool bIsOnBorder = IsOnBorder(LocalUV, Rect);
+	return bIsOnBorder ? float3(1,1,0) : float3(0,0,0);
+}
+
+// Display the DOM textures into some insert rect
+#if PERMUTATION_OUTPUT_TYPE == 0
+void VisualizeDomPS(
+	in FScreenVertexOutput Input,
+	out float4 OutColor : SV_Target0)
+{
+	const float TileSize = 48;
+	const float2 UV = Input.UV;
+	const float2 PixelPos = UV * OutputResolution;
+
+	float2 PixelPosAtTileCenter = PixelPos - (frac(PixelPos / TileSize) - 0.5f) * TileSize;
+	float2 UVAtTileCenter = PixelPosAtTileCenter * InvOutputResolution;
+
+	const float OffsetX = 50;
+	const float OffsetY = 220;
+	const float RectSize = 300;
+	const float4 DepthRect = float4(OffsetX, OffsetY, RectSize, RectSize);
+	const float4 DomRect = float4(OffsetX * 2 + RectSize, OffsetY, RectSize, RectSize);
+
+	float3 Color = 0;
+
+	float2 LocalUV = 0;
+	if (IsInsideRect(DepthRect, PixelPos, LocalUV))
+	{
+		const float2 AtlasUV = LocalUV * DomAtlasScale + DomAtlasOffset;
+		Color = DeepShadowDepthTexture.SampleLevel(LinearSampler, AtlasUV, 0).xxx;
+		Color += AddBorder(LocalUV, DepthRect);
+	}
+	else if (IsInsideRect(DomRect, PixelPos, LocalUV))
+	{
+		const float2 AtlasUV = LocalUV * DomAtlasScale + DomAtlasOffset;
+		Color = DeepShadowLayerTexture.SampleLevel(LinearSampler, AtlasUV, 0).xyz / DomScale;
+		Color += AddBorder(LocalUV, DomRect);
+#if 0 
+		float4 Weight = PostprocessInput2.SampleLevel(LinearSampler, LocalUV, 0);
+		if (Weight.w > 0) Color = float3(0, 0, 1);
+		if (Weight.z > 0) Color = float3(0, 1, 0);
+		if (Weight.y > 0) Color = float3(1, 1, 0);
+		if (Weight.x > 0) Color = float3(1, 0, 0);
+		Color *= 1000;
+#endif
+	}
+	else
+	{
+		discard;
+	}
+
+	OutColor = float4(Color, 1);
+}
+#endif // PERMUTATION_OUTPUT_TYPE == 0
+
+// Display screen space rect for vizualizing projected viewrect of each hair macro group
+#if PERMUTATION_OUTPUT_TYPE == 1
+void VisualizeDomPS(
+	in FScreenVertexOutput Input,
+	out float4 OutColor : SV_Target0)
+{
+	const float2 UV = Input.UV;
+	const float2 PixelPos = UV * OutputResolution;
+
+	float3 Color = 0;
+	float2 LocalUV = 0;
+	if (IsInsideRect(HairViewRect, PixelPos, LocalUV) && IsOnBorder(LocalUV, HairViewRect))
+	{
+		Color += float3(0, 1, 0);
+	}
+	else
+	{
+		discard;
+	}
+
+	OutColor = float4(Color, 1);
+}
+#endif // PERMUTATION_OUTPUT_TYPE == 1
+
+#endif // SHADER_VISUALIZEDOM
+
+
+//////////////////////////////////////////////////////////////////////////////////////////
+// Print Deep Opacity Map stats (to be merge with prev. pass)
+
+#if SHADER_DOMINFO
+
+#include "../ShaderPrint.ush"
+#include "../Matrices.ush"
+#include "HairStrandsDeepShadowCommonStruct.ush"
+#include "HairStrandsAABBCommon.ush"
+
+uint AllocatedSlotCount;
+uint MacroGroupCount;
+
+Buffer<int>	MacroGroupAABBBuffer;
+StructuredBuffer<FDeepShadowViewInfo> ShadowViewInfoBuffer;
+
+[numthreads(1, 1, 1)]
+void MainCS(uint3 DispatchThreadId : SV_DispatchThreadID)
+{
+	if (DispatchThreadId.x > 0)
+	{
+		return;
+	}
+
+	for (uint MacroGroupIt = 0; MacroGroupIt < MacroGroupCount; ++MacroGroupIt)
+	{
+		const FHairAABB Bound = ReadHairAABB(MacroGroupIt, MacroGroupAABBBuffer);
+		AddAABBTWS(Bound.Min, Bound.Max, float4(1, 0.5f, 0, 1));
+	}
+
+	const float4 FrustumColor = float4(0.5f, 1, 0, 1);
+	for (uint SlotIt = 0; SlotIt < AllocatedSlotCount; ++SlotIt)
+	{
+		const float4x4 TranslatedWorldToClip = ShadowViewInfoBuffer[SlotIt].TranslatedWorldToClip;
+		const float4x4 ClipToTranslatedWorld = Inverse(TranslatedWorldToClip);
+
+		float4 P01 = float4(-1, -1, 0, 1);
+		float4 P11 = float4(-1,  1, 0, 1);
+		float4 P21 = float4( 1,  1, 0, 1);
+		float4 P31 = float4( 1, -1, 0, 1);
+
+		float4 P00 = float4(-1, -1, 1, 1);
+		float4 P10 = float4(-1,  1, 1, 1);
+		float4 P20 = float4( 1,  1, 1, 1);
+		float4 P30 = float4( 1, -1, 1, 1);
+
+		P00 = mul(P00, ClipToTranslatedWorld); P00 /= P00.w;
+		P10 = mul(P10, ClipToTranslatedWorld); P10 /= P10.w;
+		P20 = mul(P20, ClipToTranslatedWorld); P20 /= P20.w;
+		P30 = mul(P30, ClipToTranslatedWorld); P30 /= P30.w;
+
+		P01 = mul(P01, ClipToTranslatedWorld); P01 /= P01.w;
+		P11 = mul(P11, ClipToTranslatedWorld); P11 /= P11.w;
+		P21 = mul(P21, ClipToTranslatedWorld); P21 /= P21.w;
+		P31 = mul(P31, ClipToTranslatedWorld); P31 /= P31.w;
+
+		AddLineTWS(P00.xyz, P10.xyz, FrustumColor, FrustumColor);
+		AddLineTWS(P10.xyz, P20.xyz, FrustumColor, FrustumColor);
+		AddLineTWS(P20.xyz, P30.xyz, FrustumColor, FrustumColor);
+		AddLineTWS(P30.xyz, P00.xyz, FrustumColor, FrustumColor);
+
+		AddLineTWS(P00.xyz, P01.xyz, FrustumColor, FrustumColor);
+		AddLineTWS(P10.xyz, P11.xyz, FrustumColor, FrustumColor);
+		AddLineTWS(P20.xyz, P21.xyz, FrustumColor, FrustumColor);
+		AddLineTWS(P30.xyz, P31.xyz, FrustumColor, FrustumColor);
+
+		AddLineTWS(P01.xyz, P11.xyz, FrustumColor, FrustumColor);
+		AddLineTWS(P11.xyz, P21.xyz, FrustumColor, FrustumColor);
+		AddLineTWS(P21.xyz, P31.xyz, FrustumColor, FrustumColor);
+		AddLineTWS(P31.xyz, P01.xyz, FrustumColor, FrustumColor);
+	}
+}
+#endif // SHADER_DOMINFO
+
+////////////////////////////////////////////////////////////////////////////////////////////////
+
+
+////////////////////////////////////////////////////////////////////////////////////////////////
+// Print hair strands stats & material infos
+
+#if SHADER_PRINT
+
+#include "../SceneData.ush"
+#include "../ShaderPrint.ush"
+#include "../CommonViewUniformBuffer.ush"
+#include "../SceneTextureParameters.ush"
+#include "../DeferredShadingCommon.ush"
+#include "../ShaderPrint.ush"
+#include "../PositionReconstructionCommon.ush"
+#include "HairStrandsVisibilityCommon.ush"
+#include "HairStrandsAABBCommon.ush" 
+
+void AddNewLine(float2 OriginalPos, inout float2 Pos)
+{
+	Pos = ShaderPrintNewline(Pos); Pos.x = OriginalPos.x;
+}
+
+int2 PixelCoord;
+int2 MaxResolution;
+uint FastResolveMask;
+uint HairMacroGroupCount;
+uint2 GroupSize; 
+uint HairVisibilityNodeGroupSize;
+uint AllocatedSampleCount;
+uint HairInstanceCount;
+
+uint NumInstances;
+uint InstanceDataSOAStride;
+
+Buffer<uint> HairInstanceIDs;
+Texture2D<float> HairCountTexture;
+Texture2D<uint> HairCountUintTexture;
+
+Texture2D<uint2> StencilTexture;
+SamplerState LinearSampler;
+Buffer<uint> HairVisibilityIndirectArgsBuffer;
+
+Buffer<int> HairMacroGroupAABBBuffer;
+Buffer<int> HairMacroGroupVoxelAlignedAABBBuffer;
+
+#define MATERIAL_DATA 1
+#define DETAIL_DATA 1
+
+FFontColor GetOccupancyColor(float In)
+{
+	float3 Color = lerp(float3(0, 1, 0), float3(1, 0, 0), saturate(In));
+	return InitFontColor(Color);
+}
+
+void DrawInstanceAABB(uint PrimitiveId)
+{
+	// Since we don't have the actual instance ID (they are transient instance vs. persistent instance), 
+	// we find the primitive with the matching PrimitiveID
+	uint InstanceID = ~0;
+	const uint TotalInstanceCount = min(8192u, NumInstances);
+	for (uint LocalInstanceID = 0; LocalInstanceID < TotalInstanceCount; ++LocalInstanceID)
+	{
+		FInstanceSceneData InstanceData = GetInstanceSceneData(LocalInstanceID, InstanceDataSOAStride);
+		if (InstanceData.PrimitiveId == PrimitiveId)
+		{
+			InstanceID = LocalInstanceID;
+			break;
+		}
+	}
+
+	if (InstanceID != ~0)
+	{
+		FInstanceSceneData InstanceData = GetInstanceSceneData(InstanceID, InstanceDataSOAStride);
+		const float4 InstanceBoundColor = ColorLightGreen;
+		AddOBBWS(InstanceData.LocalBoundsCenter - InstanceData.LocalBoundsExtent, InstanceData.LocalBoundsCenter + InstanceData.LocalBoundsExtent, InstanceBoundColor, LWCHackToFloat(InstanceData.LocalToWorld));
+		AddReferentialWS(LWCHackToFloat(InstanceData.LocalToWorld), 50.f);
+	}
+}
+
+[numthreads(1, 1, 1)]
+void MainCS(uint GroupIndex : SV_GroupIndex, uint3 DispatchThreadId : SV_DispatchThreadID)
+{	
+	// Draw AABB of the instance groups
+	if (all(DispatchThreadId == 0))
+	{
+		// Macro groups bounds
+		for (uint GroupIndex =0;GroupIndex< HairMacroGroupCount; GroupIndex++)
+		{
+			const float4 InstanceBoundColor = float4(0.25f, 1.f, 0.25f, 1.0f); 
+			FHairAABB Bound = ReadHairAABB(GroupIndex, HairMacroGroupAABBBuffer);
+			AddAABBTWS(Bound.Min, Bound.Max, InstanceBoundColor);
+		}
+
+		// Voxel aligned macro groups bounds
+		for (uint GroupIndex = 0; GroupIndex < HairMacroGroupCount; GroupIndex++)
+		{
+			const float4 InstanceBoundColor = float4(1.f, 0.5f, 0.5f, 1.0f);
+			FHairAABB Bound = ReadHairAABB(GroupIndex, HairMacroGroupVoxelAlignedAABBBuffer);
+			AddAABBTWS(Bound.Min, Bound.Max, InstanceBoundColor);
+		}
+
+		// Hair instance bounds
+		for (uint HairInstanceIndex = 0; HairInstanceIndex < HairInstanceCount; HairInstanceIndex++)
+		{
+			const uint InstanceID = HairInstanceIDs[HairInstanceIndex];
+			if (InstanceID != ~0)
+			{
+				DrawInstanceAABB(InstanceID);
+			}
+		}
+
+		const uint DispatchX = HairVisibilityIndirectArgsBuffer[0];
+		const uint DispatchY = HairVisibilityIndirectArgsBuffer[1];
+		const uint UsedSampleCount = HairStrands.HairSampleCount.Load(uint3(0, 0, 0));
+		const uint DispatchedSampleCount = DispatchX * DispatchY * HairVisibilityNodeGroupSize;
+
+		const uint2 EquivalentResolution = ceil(sqrt(UsedSampleCount));
+	
+		// Pixel coord
+		float2 OriginalPos = float2(450, 50) / float2(MaxResolution);
+		float2 Pos = OriginalPos;
+		Pos = ShaderPrintText(Pos, TEXT("Sample Count (Used/Alloc)  "), FontCyan);
+		Pos = ShaderPrint(Pos, UsedSampleCount, FontCyan);
+		Pos = ShaderPrint(Pos, AllocatedSampleCount, FontCyan);
+
+		const float AllocationPercentage = float(UsedSampleCount) / float(AllocatedSampleCount);
+		const FFontColor OccupancyColor = GetOccupancyColor(AllocationPercentage);
+		Pos = ShaderPrintText(Pos, TEXT(" %% "), OccupancyColor);
+		Pos = ShaderPrint(Pos, AllocationPercentage * 100.f, OccupancyColor);
+
+		AddNewLine(OriginalPos, Pos);
+
+		Pos = ShaderPrintText(Pos, TEXT("Sample Lighting Resolution "), FontCyan);
+		Pos = ShaderPrint(Pos, EquivalentResolution, FontCyan);
+		AddNewLine(OriginalPos, Pos);
+
+		Pos = ShaderPrintText(Pos, TEXT("MacroGroup Count           "), FontYellow);
+		Pos = ShaderPrint(Pos, HairMacroGroupCount, FontYellow);
+		AddNewLine(OriginalPos, Pos);
+	}
+
+	const float2 PixelPos = PixelCoord;
+	if (PixelCoord.x < 0 || PixelCoord.y < 0 || PixelCoord.x >= MaxResolution.x || PixelCoord.y >= MaxResolution.y)
+		return;
+
+	ShaderPrintFilter(true);
+
+	const float HairPixelCoverage = HairStrands.HairCoverageTexture.Load(uint3(PixelPos, 0));
+	const bool bIsValid = HairPixelCoverage > 0;
+	if (!bIsValid)
+		return;
+
+	const float ViewHairCount = HairCountTexture.Load(uint3(PixelPos, 0));
+	const float ViewHairCountUint = HairCountUintTexture.Load(uint3(PixelPos, 0)) /  1000.f;
+
+//	const float2 OriginalPos = PixelPos / float2(MaxResolution);
+	const float2 OriginalPos = float2(50,500) / float2(MaxResolution);
+	float2 Pos = OriginalPos;
+
+	uint Total8BitCoverage = 0;
+	const FNodeDesc NodeDesc = DecodeNodeDesc(HairStrands.HairSampleOffset.Load(uint3(PixelCoord, 0)));
+	for (uint SampleIt0 = 0; SampleIt0 < NodeDesc.Count; SampleIt0++)
+	{
+		const uint LocalOffset = NodeDesc.Offset + SampleIt0;
+		const FPackedHairSample NodeData = HairStrands.HairSampleData[LocalOffset];
+		const FHairSample Sample = UnpackHairSample(NodeData);
+		Total8BitCoverage += Sample.Coverage8bit;
+	}
+
+	// Pixel coord
+	Pos = ShaderPrintText(Pos, TEXT("Pixel           "), FontTurquoise);
+	Pos = ShaderPrint(Pos, PixelCoord, FontTurquoise);
+	AddNewLine(OriginalPos, Pos);
+
+	// Max sample
+	Pos = ShaderPrintText(Pos, TEXT("Max Sample      "));
+	Pos = ShaderPrint(Pos, HairStrands.MaxSamplePerPixelCount);
+	AddNewLine(OriginalPos, Pos);
+
+	// Unique sample
+	Pos = ShaderPrintText(Pos, TEXT("Sample          "), FontOrange);
+	Pos = ShaderPrint(Pos, NodeDesc.Count, FontOrange);
+	AddNewLine(OriginalPos, Pos);
+
+	// Accurate coverage (if coverage pass is activated)
+	Pos = ShaderPrintText(Pos, TEXT("Hair Count      "));
+	Pos = ShaderPrint(Pos, ViewHairCount);
+	AddNewLine(OriginalPos, Pos);
+
+	// Accurate coverage (if coverage pass is activated) uint
+	Pos = ShaderPrintText(Pos, TEXT("Hair Count Uint "));
+	Pos = ShaderPrint(Pos, ViewHairCountUint);
+	AddNewLine(OriginalPos, Pos);
+	
+	// Accurate coverage (if coverage pass is activated)
+	Pos = ShaderPrintText(Pos, TEXT("Coverage        "), FontOrange);
+	Pos = ShaderPrint(Pos, HairPixelCoverage, FontOrange);
+	AddNewLine(OriginalPos, Pos);
+
+	// 8 bit Coverage (summed)
+	Pos = ShaderPrintText(Pos, TEXT("Sum cov. 8bits  "));
+	Pos = ShaderPrint(Pos, Total8BitCoverage);
+	AddNewLine(OriginalPos, Pos);
+
+	// Resolve type
+	const uint Stencil = StencilTexture.Load(uint3(PixelPos,0)).y;
+	const bool bIsFastResolve = (Stencil & FastResolveMask) != 0;
+	Pos = ShaderPrintText(Pos, TEXT("Resolve         "));
+
+	if (bIsFastResolve)
+	{
+		Pos = ShaderPrintText(Pos, TEXT("Fast"), FontEmerald);
+	}
+	else
+	{
+		Pos = ShaderPrintText(Pos, TEXT("Regular"), FontEmerald);
+	}
+
+	AddNewLine(OriginalPos, Pos);
+	AddNewLine(OriginalPos, Pos);
+
+	const float2 ColumnPos = Pos;
+
+	// Sample Description
+	{	
+		Pos = ShaderPrintText(Pos, TEXT("Sample ID       "), FontWhite);	AddNewLine(OriginalPos, Pos);
+	#if MATERIAL_DATA
+		Pos = ShaderPrintText(Pos, TEXT(".Base Color R   "), FontRed);		AddNewLine(OriginalPos, Pos); 
+		Pos = ShaderPrintText(Pos, TEXT(".Base Color G   "), FontGreen);	AddNewLine(OriginalPos, Pos);
+		Pos = ShaderPrintText(Pos, TEXT(".Base Color B   "), FontBlue);		AddNewLine(OriginalPos, Pos);
+		AddNewLine(OriginalPos, Pos);
+
+		Pos = ShaderPrintText(Pos, TEXT(".Tangent X      "), FontEmerald);	AddNewLine(OriginalPos, Pos); 
+		Pos = ShaderPrintText(Pos, TEXT(".Tangent Y      "), FontEmerald);	AddNewLine(OriginalPos, Pos);
+		Pos = ShaderPrintText(Pos, TEXT(".Tangent Z      "), FontEmerald);	AddNewLine(OriginalPos, Pos);
+		AddNewLine(OriginalPos, Pos);
+
+		Pos = ShaderPrintText(Pos, TEXT(".Roughness      "), FontOrange);	AddNewLine(OriginalPos, Pos);
+		Pos = ShaderPrintText(Pos, TEXT(".Specular       "), FontEmerald);	AddNewLine(OriginalPos, Pos);
+		Pos = ShaderPrintText(Pos, TEXT(".Backlit        "), FontOrange);	AddNewLine(OriginalPos, Pos);
+		AddNewLine(OriginalPos, Pos);
+	#endif 													  
+	#if DETAIL_DATA											  
+		Pos = ShaderPrintText(Pos, TEXT(".Light Channels "), FontEmerald);	AddNewLine(OriginalPos, Pos);
+		Pos = ShaderPrintText(Pos, TEXT(".Scatter Scene  "), FontOrange);	AddNewLine(OriginalPos, Pos);
+		Pos = ShaderPrintText(Pos, TEXT(".Depth          "), FontEmerald);	AddNewLine(OriginalPos, Pos);
+		AddNewLine(OriginalPos, Pos);
+
+		Pos = ShaderPrintText(Pos, TEXT(".Macro Group ID "), FontYellow);	AddNewLine(OriginalPos, Pos);
+		Pos = ShaderPrintText(Pos, TEXT(".Primitive		 "), FontEmerald);	AddNewLine(OriginalPos, Pos);
+	#endif													  
+		Pos = ShaderPrintText(Pos, TEXT(".Coverage       "), FontOrange);	AddNewLine(OriginalPos, Pos);
+	}
+
+	// Per sample description
+	for (uint SampleIt = 0; SampleIt < NodeDesc.Count; SampleIt++)
+	{
+		const uint LocalOffset = NodeDesc.Offset + SampleIt;
+		const FHairSample Sample = UnpackHairSample(HairStrands.HairSampleData[LocalOffset]);
+
+		float2 StartSamplePos = ColumnPos + float2(150.f + 80.f * SampleIt, 0.f) / float2(MaxResolution);
+		Pos = StartSamplePos;
+
+		Pos = ShaderPrint(Pos, SampleIt, FontWhite);									AddNewLine(StartSamplePos, Pos);
+
+	#if MATERIAL_DATA
+		Pos = ShaderPrint(Pos, Sample.BaseColor.x,	FontRed);							AddNewLine(StartSamplePos, Pos);
+		Pos = ShaderPrint(Pos, Sample.BaseColor.y,	FontGreen);							AddNewLine(StartSamplePos, Pos);
+		Pos = ShaderPrint(Pos, Sample.BaseColor.z,	FontBlue);							AddNewLine(StartSamplePos, Pos);
+		AddNewLine(StartSamplePos, Pos);
+
+		Pos = ShaderPrint(Pos, Sample.Tangent.x,	FontEmerald);						AddNewLine(StartSamplePos, Pos);
+		Pos = ShaderPrint(Pos, Sample.Tangent.y,	FontEmerald);						AddNewLine(StartSamplePos, Pos);
+		Pos = ShaderPrint(Pos, Sample.Tangent.z,	FontEmerald);						AddNewLine(StartSamplePos, Pos);
+		AddNewLine(StartSamplePos, Pos);
+
+		Pos = ShaderPrint(Pos, Sample.Roughness,	FontOrange);						AddNewLine(StartSamplePos, Pos);
+		Pos = ShaderPrint(Pos, Sample.Specular,		FontEmerald);						AddNewLine(StartSamplePos, Pos);
+		Pos = ShaderPrint(Pos, Sample.Backlit,		FontOrange);						AddNewLine(StartSamplePos, Pos);
+		AddNewLine(StartSamplePos, Pos);
+	#endif 
+
+	#if DETAIL_DATA
+		Pos = ShaderPrint(Pos,(Sample.LightChannelMask & 0x1) ? 1u : 0u, FontEmerald);	AddNewLine(StartSamplePos, Pos);
+//		Pos = ShaderPrint(Pos,(Sample.LightChannelMask & 0x1) ? 1u : 0u);				AddNewLine(StartSamplePos, Pos);
+//		Pos = ShaderPrint(Pos,(Sample.LightChannelMask & 0x2) ? 1u : 0u);				AddNewLine(StartSamplePos, Pos);
+//		Pos = ShaderPrint(Pos,(Sample.LightChannelMask & 0x4) ? 1u : 0u);				AddNewLine(StartSamplePos, Pos);
+		Pos = ShaderPrint(Pos, Sample.bScatterSceneLighting ? 1u : 0u, FontOrange);		AddNewLine(StartSamplePos, Pos);
+		Pos = ShaderPrint(Pos, Sample.Depth, FontEmerald);								AddNewLine(StartSamplePos, Pos);
+		AddNewLine(StartSamplePos, Pos);
+
+		Pos = ShaderPrint(Pos, Sample.MacroGroupId, FontYellow);						AddNewLine(StartSamplePos, Pos);
+		Pos = ShaderPrint(Pos, Sample.PrimitiveId, FontEmerald);						AddNewLine(StartSamplePos, Pos);
+	#endif
+		Pos = ShaderPrint(Pos, Sample.Coverage8bit, FontOrange);						AddNewLine(StartSamplePos, Pos);
+	}
+}
+
+#endif // SHADER_PRINT
+
+////////////////////////////////////////////////////////////////////////////////////////////////
+// Print stats/info about hair instance (strands/cards/meshes)
+
+#if SHADER_PRINT_INSTANCE
+
+#include "../ShaderPrint.ush"
+#include "HairStrandsAABBCommon.ush"
+
+void AddNewLine(float2 OriginalPos, inout float2 Pos)
+{
+	Pos = ShaderPrintNewline(Pos); Pos.x = OriginalPos.x;
+}
+
+int2					MaxResolution;
+uint					InstanceCount;
+
+Buffer<uint4>			Infos;
+
+uint					NameInfoCount;
+uint					NameCharacterCount;
+StructuredBuffer<uint2> NameInfos;
+Buffer<uint>			Names;
+Buffer<int>				InstanceAABB;
+
+struct FNameInfo
+{
+	uint PrimitiveID;
+	uint Length;
+	uint Offset;
+	uint Pad0;
+	uint Pad1;
+};
+
+FNameInfo UnpackDebugNameInfo(uint2 In)
+{
+	FNameInfo Out = (FNameInfo)0;
+	Out.PrimitiveID = In.x;
+	Out.Offset = (In.y) & 0xFFFF;
+	Out.Length = (In.y >> 16) & 0xFF;
+
+	return Out;
+}
+
+FNameInfo FindNameInfo(uint PrimitiveID)
+{
+	FNameInfo Out = (FNameInfo)0;
+	Out.PrimitiveID = ~0;
+	for (uint It = 0; It < NameInfoCount; ++It)
+	{
+		if (NameInfos[It].x == PrimitiveID)
+		{
+			return UnpackDebugNameInfo(NameInfos[It]);
+		}
+	}
+	return Out;
+}
+
+float2 PrintInstanceName(float2 InPos, uint PrimitiveID, FFontColor InColor)
+{
+	const FNameInfo Info = FindNameInfo(PrimitiveID);
+	if (Info.PrimitiveID != ~0 && (Info.Length + Info.Offset) <= NameCharacterCount)
+	{
+		for (uint It = 0; It < Info.Length; ++It)
+		{
+			const uint Char = Names[It + Info.Offset];
+			InPos = ShaderPrintSymbol(InPos, Char, InColor);
+		}
+	}
+	return InPos;
+}
+
+[numthreads(1, 1, 1)]
+void MainCS(uint GroupIndex : SV_GroupIndex, uint3 DispatchThreadId : SV_DispatchThreadID)
+#if PERMUTATION_OUTPUT_TYPE == 0
+{
+	// Draw AABB of the instance groups
+	if (any(DispatchThreadId != 0))
+	{
+		return;
+	}
+
+	// Draw legend
+	float2 OriginalPos = ShaderPrintGetPos(uint2(50, 50));
+	float2 Pos = OriginalPos;
+
+
+	// Draw totals
+	const uint TotalGroup = InstanceCount;
+	uint TotalStrandsCount = 0;
+	uint TotalVertexCount = 0;
+	for (uint InstanceIndex = 0; InstanceIndex < InstanceCount; ++InstanceIndex)
+	{
+		const uint4 InstanceData = Infos[InstanceIndex];
+		TotalStrandsCount += InstanceData.z;
+		TotalVertexCount += InstanceData.w;
+	}
+
+	Pos = ShaderPrintText(Pos, TEXT("Total Groups        : "), FontWhite);
+	Pos = Print(Pos, TotalGroup, FontYellow, 8);
+	AddNewLine(OriginalPos, Pos);
+
+	Pos = ShaderPrintText(Pos, TEXT("Total Strand Curves : "), FontWhite);
+	Pos = Print(Pos, TotalStrandsCount, FontOrange, 8);
+	AddNewLine(OriginalPos, Pos);
+
+	Pos = ShaderPrintText(Pos, TEXT("Total Strand Vertex : "), FontWhite);
+	Pos = Print(Pos, TotalVertexCount, FontSilver, 8);
+	AddNewLine(OriginalPos, Pos);
+
+	AddNewLine(OriginalPos, Pos);
+
+	// Draw instance details
+	Pos = ShaderPrintText(Pos, TEXT("   Group  LOD     Geometry Binding Sim   RBF   Clip    Strands Vertex  Name"));
+	AddNewLine(OriginalPos, Pos);
+
+	for (uint InstanceIndex = 0; InstanceIndex < InstanceCount; ++InstanceIndex)
+	{
+		const uint4 InstanceData	= Infos[InstanceIndex];
+		const uint GroupIndex		= (InstanceData.x>>0)  & 0xFF;
+		const uint GroupCount		= (InstanceData.x>>8)  & 0xFF;
+		const float LODCount		= (InstanceData.x>>16) & 0xFF;
+		const uint GeometryType		= (InstanceData.x>>24) & 0x7;
+		const uint BindingType		= (InstanceData.x>>27) & 0x7;
+		const bool bHasSim			= (InstanceData.x>>30) & 0x1;
+		const bool bHasRBF			= (InstanceData.x>>31) & 0x1;
+
+		const float LOD				= f16tof32((InstanceData.y>>0)  & 0xFFFF);
+		const float ClipValue		= f16tof32((InstanceData.y>>16) & 0xFFFF);
+
+		const uint StrandsCount		= InstanceData.z;
+		const uint VertexCount		= InstanceData.w;
+
+		// Index
+		Pos = Print(Pos, InstanceIndex, FontSilver, 2);
+		Pos = ShaderPrintText(Pos, TEXT(" "));
+
+		// Group
+		Pos = Print(Pos, GroupIndex, FontYellow, 1);
+		Pos = ShaderPrintText(Pos, TEXT("/"));
+		Pos = Print(Pos, GroupCount, FontYellow, 1);
+		Pos = ShaderPrintText(Pos, TEXT("    "));
+
+		// LOD
+		Pos = Print(Pos, LOD, FontYellow, 4, 2);	
+		Pos = ShaderPrintText(Pos, TEXT("/"));
+		Pos = Print(Pos, LODCount, FontYellow, 1);
+		Pos = ShaderPrintText(Pos, TEXT("  "));
+
+		// Geometry
+		if (GeometryType == 0)		{ Pos = ShaderPrintText(Pos, TEXT("Strands "), FontEmerald); }
+		else if (GeometryType == 1) { Pos = ShaderPrintText(Pos, TEXT("Cards   "), FontEmerald); }
+		else if (GeometryType == 2) { Pos = ShaderPrintText(Pos, TEXT("Meshes  "), FontEmerald); }
+		Pos = ShaderPrintText(Pos, TEXT(" "));
+
+		// Binding
+		if (BindingType == 0)		{ Pos = ShaderPrintText(Pos, TEXT("Skinned "), FontTurquoise); }
+		else if (BindingType == 1)	{ Pos = ShaderPrintText(Pos, TEXT("Rigid   "), FontTurquoise); }
+		else if (BindingType == 2)	{ Pos = ShaderPrintText(Pos, TEXT("None    "), FontTurquoise); }
+
+		// Sim
+		if (bHasSim)
+			Pos = ShaderPrintText(Pos, TEXT("True "), FontGreen);
+		else
+			Pos = ShaderPrintText(Pos, TEXT("False"), FontRed);
+		Pos = ShaderPrintText(Pos, TEXT(" "), FontWhite);
+
+		// RBF
+		if (bHasRBF)
+			Pos = ShaderPrintText(Pos, TEXT("True "), FontGreen);
+		else
+			Pos = ShaderPrintText(Pos, TEXT("False"), FontRed);
+		Pos = ShaderPrintText(Pos, TEXT(" "));
+
+		// Clip
+		if (ClipValue >= 0)
+			Pos = Print(Pos, ClipValue, FontGreen, 7);
+		else 
+			Pos = Print(Pos, 0u, FontSilver, 7);
+		Pos = ShaderPrintText(Pos, TEXT(" "));
+
+		// Count
+		Pos = Print(Pos, StrandsCount, FontOrange, 8);
+		Pos = Print(Pos, VertexCount,  FontSilver, 8);
+
+		// Name
+		Pos = PrintInstanceName(Pos, InstanceIndex, FontWhite);
+
+		AddNewLine(OriginalPos, Pos);
+	}
+}
+#else  // PERMUTATION_OUTPUT_TYPE
+{
+	if (all(DispatchThreadId == 0))
+	{
+		const FHairAABB AABB = ReadHairAABB(0, InstanceAABB);
+		AddAABBTWS(AABB.Min, AABB.Max, ColorDarkGreen);
+	}
+}
+#endif // PERMUTATION_OUTPUT_TYPE
+
+#endif // SHADER_PRINT_INSTANCE
+
+////////////////////////////////////////////////////////////////////////////////////////////////
+// Print tangent direction for strands/cards/meshes
+
+#if SHADER_TANGENT
+
+#include "../CommonViewUniformBuffer.ush"
+#include "../SceneTextureParameters.ush"
+#include "../DeferredShadingCommon.ush"
+#include "../ShaderPrint.ush"
+#include "../PositionReconstructionCommon.ush"
+#include "HairStrandsVisibilityCommon.ush"
+#include "HairStrandsAABBCommon.ush" 
+
+uint TileSize;
+int2 TileCount;
+float2 OutputResolution;
+RWTexture2D<float4> OutputTexture;
+
+#define HAIR_TANGENT_INVALID 0
+#define HAIR_TANGENT_STRANDS 1
+#define HAIR_TANGENT_CARDS   2
+
+void AddNewLine(float2 OriginalPos, inout float2 Pos)
+{
+	Pos = ShaderPrintNewline(Pos); Pos.x = OriginalPos.x;
+}
+
+struct FTangentDesc
+{
+	float3 Tangent;
+	uint Type;
+};
+
+FTangentDesc AddTangentPlot(uint2 PixelCoord, float2 UV)
+{
+	FTangentDesc Out = (FTangentDesc)0;
+
+	// Strands
+	{
+		// Find closest
+		float3 Closest_Tangent = 0;
+		float  Closest_Depth = 0;
+
+		uint Total8BitCoverage = 0;
+		const FNodeDesc NodeDesc = DecodeNodeDesc(HairStrands.HairSampleOffset.Load(uint3(PixelCoord, 0)));
+		for (uint SampleIt0 = 0; SampleIt0 < NodeDesc.Count; SampleIt0++)
+		{
+			const uint LocalOffset = NodeDesc.Offset + SampleIt0;
+			const FPackedHairSample NodeData = HairStrands.HairSampleData[LocalOffset];
+			const FHairSample Sample = UnpackHairSample(NodeData);
+
+			if (Sample.Depth > Closest_Depth)
+			{
+				Closest_Depth = Sample.Depth;
+				Closest_Tangent = Sample.Tangent;
+			}
+		}
+
+		// Draw tangent vector (strands)
+		if (NodeDesc.Count > 0)
+		{
+			const float SceneDepth = ConvertFromDeviceZ(Closest_Depth);
+			const float3 WorldPosition = ReconstructTranslatedWorldPositionFromDepth(UV, SceneDepth);
+			AddLineTWS(WorldPosition, WorldPosition + Closest_Tangent, ColorRed, ColorYellow);
+			Out.Type = 1;
+			Out.Tangent = Closest_Tangent;
+		}
+	}
+
+	// Cards
+	if (Out.Type == 0)
+	{
+		FScreenSpaceData ScreenSpaceData = GetScreenSpaceData(UV);
+		const float SceneDepth = CalcSceneDepth(UV);
+		const float3 WorldPosition = ReconstructTranslatedWorldPositionFromDepth(UV, SceneDepth);
+		if (ScreenSpaceData.GBuffer.ShadingModelID == SHADINGMODELID_HAIR)
+		{
+			const float3 Tangent = ScreenSpaceData.GBuffer.WorldNormal;
+			AddLineTWS(WorldPosition, WorldPosition + Tangent, ColorPurple, ColorYellow);
+			Out.Type = 2;
+			Out.Tangent = Tangent;
+		}
+	}
+
+	return Out;
+}
+
+void AddTangentText(float2 UV, float3 Tangent, uint TangentType)
+{
+	const float2 OriginalPos = UV;
+	float2 Pos = OriginalPos;
+	if (TangentType == HAIR_TANGENT_STRANDS)
+		Pos = ShaderPrintText(Pos, TEXT("Strands"), FontYellow);
+	else if (TangentType == HAIR_TANGENT_CARDS)
+		Pos = ShaderPrintText(Pos, TEXT("Cards "), FontYellow);
+	else if (TangentType == HAIR_TANGENT_CARDS)
+		Pos = ShaderPrintText(Pos, TEXT("Invalid"), FontRed);
+
+	AddNewLine(OriginalPos, Pos);
+
+	Pos = ShaderPrint(Pos, Tangent.x, FontRed);		AddNewLine(OriginalPos, Pos);
+	Pos = ShaderPrint(Pos, Tangent.y, FontGreen);	AddNewLine(OriginalPos, Pos);
+	Pos = ShaderPrint(Pos, Tangent.z, FontBlue);	AddNewLine(OriginalPos, Pos);
+}
+
+[numthreads(8, 8, 1)]
+void MainCS(int2 DispatchThreadId : SV_DispatchThreadID)
+{
+	// Cursor plot
+	if (all(DispatchThreadId == 0))
+	{
+		const float2 PixelCoord = GetCursorPos();// DispatchThreadId.xy + 0.5f;
+		const float2 UV = PixelCoord / float2(OutputResolution);
+		const float2 UVText = (PixelCoord + float2(20, 0)) / float2(OutputResolution);
+
+		FTangentDesc Desc = AddTangentPlot(PixelCoord, UV);
+		if (Desc.Type == HAIR_TANGENT_STRANDS)
+		{			
+			AddTangentText(UVText, Desc.Tangent, HAIR_TANGENT_STRANDS);
+		}
+		else if (Desc.Type == HAIR_TANGENT_CARDS)
+		{
+			AddTangentText(UVText, Desc.Tangent, HAIR_TANGENT_CARDS);
+		}
+	}
+
+	// Grid plot
+	if (all(DispatchThreadId.xy < TileCount))
+	{
+		const float2 PixelCoord = DispatchThreadId.xy * TileSize;
+		if (all(PixelCoord < OutputResolution))
+		{
+			const float2 UV = PixelCoord / float2(OutputResolution);
+			AddTangentPlot(PixelCoord, UV);
+		}
+	}
+}
+#endif // SHADER_TANGENT
+
+////////////////////////////////////////////////////////////////////////////////////////////////
+// Print LOD info/stats
+
+#if SHADER_LOD_INFO
+
+#include "../ShaderPrint.ush"
+
+void AddNewLine(float2 OriginalPos, inout float2 Pos)
+{
+	Pos = ShaderPrintNewline(Pos); Pos.x = OriginalPos.x;
+}
+
+int2 MaxResolution;
+uint GroupIndex;
+uint GeometryType;
+float LOD;
+float ScreenSize;
+float3 GroupColor;
+
+[numthreads(1, 1, 1)]
+void MainCS(uint3 DispatchThreadId : SV_DispatchThreadID)
+{
+	// Legend
+	if (GroupIndex == 0)
+	{
+		float2 OriginalPos = float2(15, 50) / float2(MaxResolution);
+		float2 Pos = OriginalPos;
+		Pos = ShaderPrintText(Pos, TEXT("Group Geom.   LOD      Screen Size"), FontSilver); AddNewLine(OriginalPos, Pos);
+	}
+
+	// Values
+	{
+		float2 OriginalPos = float2(30, 50) / float2(MaxResolution);
+		OriginalPos.y = float(65 + GroupIndex * 15) / float(MaxResolution.y);
+		float2 Pos = OriginalPos;
+
+		// Group Color
+		Pos = ShaderPrintText(Pos, TEXT("* "), InitFontColor(GroupColor));
+
+		// Group index
+		ShaderPrint(Pos, GroupIndex, FontEmerald);
+		Pos.x += 20.f / float(MaxResolution.x);
+
+		// Geometry Type
+		if (GeometryType == 0)
+		{
+			Pos = ShaderPrintText(Pos, TEXT("Strands "), FontOrange);
+		}
+		else if (GeometryType == 1)
+		{
+			Pos = ShaderPrintText(Pos, TEXT("Cards   "), FontOrange);
+		}
+		else if (GeometryType == 2)
+		{
+			Pos = ShaderPrintText(Pos, TEXT("Meshes  "), FontOrange);
+		}
+
+		// LOD
+		ShaderPrint(Pos, LOD, FontWhite);
+		Pos.x += 80.f / float(MaxResolution.x);
+
+		// Screen size
+		ShaderPrint(Pos, saturate(ScreenSize), FontWhite);
+		Pos.x += 80.f / float(MaxResolution.x);
+
+		AddNewLine(OriginalPos, Pos);
+	}
+}
+
+#endif // SHADER_LOD_INFO
+
+/////////////////////////////////////////////////////////////////////////////////////////
+// Visualize hair cards atlas data
+#if SHADER_CARDS_ATLAS
+
+#include "../ShaderPrint.ush"
+
+#define DEBUG_MODE_DEPTH 1
+#define DEBUG_MODE_COVERAGE 2
+#define DEBUG_MODE_TANGENT 3
+#define DEBUG_MODE_ROOTUV 4
+#define DEBUG_MODE_U 5
+#define DEBUG_MODE_SEED 6
+
+int2 OutputResolution;
+int2 AtlasResolution;
+int DebugMode;
+
+Texture2D<float4> AtlasTexture;
+RWTexture2D<float4> OutputTexture;
+SamplerState LinearSampler;
+
+[numthreads(8, 8, 1)]
+void MainCS(uint3 DispatchThreadId : SV_DispatchThreadID)
+{	
+	const int2 PixelCoord = DispatchThreadId.xy;
+	if (any(PixelCoord >= OutputResolution))
+		return;
+
+	const int2 Offset = 50;
+	const int2 Resolution = 512 * float2(AtlasResolution.x, AtlasResolution.y)/max(AtlasResolution.x, AtlasResolution.y);
+
+	const float2 LocalUV = float2(PixelCoord - Offset) / float2(Resolution);
+
+#if 1
+	if (all(LocalUV >= 0) && all(LocalUV <= 1))
+	{
+		float4 InData = AtlasTexture.SampleLevel(LinearSampler, float2(LocalUV.x,1-LocalUV.y), 0);
+		float3 Data = 0;
+		if (DebugMode == DEBUG_MODE_DEPTH) 		{ Data = InData.xxx; }
+		if (DebugMode == DEBUG_MODE_COVERAGE)	{ Data = InData.xxx; }
+		if (DebugMode == DEBUG_MODE_TANGENT)	{ Data = InData.xyz; }
+		if (DebugMode == DEBUG_MODE_ROOTUV)		{ Data = float3(InData.xy,0); }
+		if (DebugMode == DEBUG_MODE_U)			{ Data = InData.zzz; }
+		if (DebugMode == DEBUG_MODE_SEED)		{ Data = InData.www; }
+		OutputTexture[PixelCoord] = float4(Data, 1);
+	}
+#else
+	if (all(PixelCoord <= AtlasResolution))
+	{
+		float3 Data = AtlasTexture.Load(uint3(PixelCoord, 0)).xyz;
+		if (DebugMode == DEBUG_MODE_DEPTH) { Data = Data.xxx; }
+		if (DebugMode == DEBUG_MODE_COVERAGE) { Data = Data.xxx; }
+		if (DebugMode == DEBUG_MODE_TANGENT) {}
+		OutputTexture[PixelCoord] = float4(Data, 1);
+	}
+#endif
+
+#if PRINT_ATLAS_RESOLUTION
+	// Draw AABB of the instance groups
+	if (DispatchThreadId.x == 0 && DispatchThreadId.y == 0)
+	{
+		float2 OriginalPos = float2(50, 40) / float2(OutputResolution);
+		float2 Pos = OriginalPos;
+		Pos = Print("Atlas Resolution: ")
+		Pos = ShaderPrint(Pos, AtlasResolution);
+
+		Pos = ShaderPrintNewline(Pos);
+		Pos.x = OriginalPos.x;
+	}
+#endif
+}
+
+#endif // SHADER_CARDS_ATLAS
+
+/////////////////////////////////////////////////////////////////////////////////////////
+// Visualize cards hair guides
+
+#if SHADER_CARDS_GUIDE
+
+#include "HairStrandsVertexFactoryCommon.ush"
+#include "../ShaderPrint.ush"
+
+uint RenVertexCount;
+float3 RenRestOffset;
+
+uint SimVertexCount;
+float3 SimRestOffset;
+
+float4x4 LocalToWorld;
+
+Buffer<uint4> RenRestPosition;
+Buffer<uint4> RenDeformedPosition;
+Buffer<float4> RenDeformedOffset;
+
+Buffer<uint4> SimRestPosition;
+Buffer<uint4> SimDeformedPosition;
+Buffer<float4> SimDeformedOffset;
+
+float3 RestOffset;
+float3 DeformedOffset;
+Buffer<uint4> RestPosition;
+Buffer<uint4> DeformedPosition; 
+
+int DebugMode;
+
+#define DEBUGMODE_REN_REST 1
+#define DEBUGMODE_REN_DEFORMED 2
+#define DEBUGMODE_SIM_REST 3
+#define DEBUGMODE_SIM_DEFORMED 4
+
+[numthreads(32, 1, 1)]
+void MainCS(uint3 DispatchThreadId : SV_DispatchThreadID)
+{
+	const uint VertexId = DispatchThreadId.x;
+	const uint VertexCount = DebugMode == DEBUGMODE_REN_REST || DebugMode == DEBUGMODE_REN_DEFORMED ? RenVertexCount : SimVertexCount;
+	if (VertexId+1 >= VertexCount)
+		return;
+
+	uint4 PackedP0 = 0;
+	uint4 PackedP1 = 0;
+	float3 Offset = 0;
+	float4 ColorRoot = ColorWhite;
+	float4 ColorTip  = ColorWhite;
+
+	if (DebugMode == DEBUGMODE_REN_REST)
+	{
+		PackedP0 = RenRestPosition[VertexId];
+		PackedP1 = RenRestPosition[VertexId + 1];
+		Offset   = RenRestOffset;
+		ColorRoot= ColorYellow;
+		ColorTip = ColorPurple;
+	}
+	else if (DebugMode == DEBUGMODE_REN_DEFORMED)
+	{
+		PackedP0 = RenDeformedPosition[VertexId];
+		PackedP1 = RenDeformedPosition[VertexId + 1];
+		Offset   = RenDeformedOffset[0].xyz;
+		ColorRoot= ColorYellow;
+		ColorTip = ColorPurple;
+	}
+	else if (DebugMode == DEBUGMODE_SIM_REST)
+	{
+		PackedP0 = SimRestPosition[VertexId];
+		PackedP1 = SimRestPosition[VertexId + 1];
+		Offset   = SimRestOffset;
+		ColorRoot= ColorCyan;
+		ColorTip = ColorRed;
+	}
+	else if (DebugMode == DEBUGMODE_SIM_DEFORMED)
+	{
+		PackedP0 = SimDeformedPosition[VertexId];
+		PackedP1 = SimDeformedPosition[VertexId + 1];
+		Offset   = SimDeformedOffset[0].xyz;
+		ColorRoot= ColorCyan;
+		ColorTip = ColorRed;
+	}
+
+	const FHairControlPoint P0 = UnpackHairControlPoint(PackedP0, Offset, 1, 1, 1);
+	const FHairControlPoint P1 = UnpackHairControlPoint(PackedP1, Offset, 1, 1, 1);
+
+	if (P1.Type != HAIR_CONTROLPOINT_START)
+	{
+		const float4 Color0 = lerp(ColorRoot, ColorTip, P0.UCoord);
+		const float4 Color1 = lerp(ColorRoot, ColorTip, P1.UCoord);
+		AddLineWS(mul(float4(P0.Position, 1), LocalToWorld), mul(float4(P1.Position,1), LocalToWorld), Color0, Color1);
+	}
+}
+
+#endif // SHADER_CARDS_GUIDE
+
+///////////////////////////////////////////////////////////////////////////////////////////////////
+// Visualize voxelized hair strands (used for cards generation)
+
+#if SHADER_CARDS_VOXEL
+
+#include "../CommonViewUniformBuffer.ush"
+#include "../SceneTextureParameters.ush"
+#include "../DeferredShadingCommon.ush"
+#include "../ShaderPrint.ush"
+#include "../PositionReconstructionCommon.ush" 
+
+uint3 PositionToCoord(float3 P, float3 InMinAABB, float3 InMaxAABB, uint3 InResolution)
+{
+	return clamp(
+		saturate((P - InMinAABB) / (InMaxAABB - InMinAABB)) * InResolution,
+		uint3(0, 0, 0),
+		InResolution - 1);
+}
+
+uint3 PositionToCoordUnclampled(float3 P, float3 InMinAABB, float3 InMaxAABB, uint3 InResolution)
+{
+	return saturate((P - InMinAABB) / (InMaxAABB - InMinAABB)) * InResolution;
+}
+
+uint CoordToIndex(uint3 InCoord, uint3 InResolution, uint LinearOffset)
+{
+	return
+		InCoord.x +
+		InCoord.y * InResolution.x +
+		InCoord.z * InResolution.x * InResolution.y +
+		LinearOffset;
+}
+
+uint3 IndexToCoord(uint InIndex, uint3 InResolution)
+{
+	const uint SliceSize = (InResolution.x * InResolution.y);
+	const uint SliceIndex = InIndex % SliceSize;
+
+	uint3 OutCoord = 0;
+	OutCoord.x = SliceIndex % InResolution.x;
+	OutCoord.y = SliceIndex / InResolution.x;
+	OutCoord.z = InIndex / SliceSize;
+
+	return OutCoord;
+}
+
+bool IsInVoxelBounds(float3 P, float3 MinP, float3 MaxP)
+{
+	return
+		P.x >= MinP.x && P.y >= MinP.y && P.z >= MinP.z &&
+		P.x <= MaxP.x && P.y <= MaxP.y && P.z <= MaxP.z;
+}
+
+float3 EncodeTangent(float3 N)
+{
+	return N * 0.5 + 0.5;
+}
+
+float3 DecodeTangent(float3 N)
+{
+	return N * 2 - 1;
+}
+
+uint PackDirection(float3 In)
+{
+	float3 T = EncodeTangent(In);
+	return
+		((0xFF & uint(T.x * 255))) |
+		((0xFF & uint(T.y * 255)) << 8) |
+		((0xFF & uint(T.z * 255)) << 16);
+}
+
+float3 UnpackDirection(uint In)
+{
+	return DecodeTangent(float3(
+		float((In)       & 0xFF) / 255.f,
+		float((In >> 8)  & 0xFF) / 255.f,
+		float((In >> 16) & 0xFF) / 255.f));
+}
+
+
+float2 OutputResolution;
+int3 Voxel_Resolution;
+float Voxel_VoxelSize;
+float3 Voxel_MinBound;
+float3 Voxel_MaxBound;
+Buffer<float4> Voxel_TangentBuffer;
+Buffer<float4> Voxel_NormalBuffer;
+Buffer<uint>   Voxel_DensityBuffer;
+Buffer<uint>   Voxel_ProcessedDensityBuffer;
+
+RWTexture2D<float4> OutputTexture;
+
+[numthreads(8, 8, 1)]
+void MainCS(uint3 DispatchThreadId : SV_DispatchThreadID)
+{
+	const float2 PixelCoord = DispatchThreadId.xy + 0.5f;
+	const float2 UV = PixelCoord / float2(OutputResolution);
+
+	const bool bDebugEnabled = all(int2(PixelCoord) == GetCursorPos());
+
+	float  WorldDepth = 0;
+	float  ClipDepth = 0;
+
+
+	if (uint(PixelCoord.x) == 0 && uint(PixelCoord.y) == 0)
+	{
+		AddAABBWS(Voxel_MinBound, Voxel_MaxBound, float4(1, 0.5f, 0, 1));
+	}
+
+	bool bIsValid = false;
+	float3 HitColor = 0;
+	{
+		// Reverse-Z
+		const float3 O = LWCHackToFloat(PrimaryView.WorldCameraOrigin);
+		const float3 E = ReconstructWorldPositionFromDepth(UV, ConvertFromDeviceZ(0.0001f));
+		const float3 D = (E - O);
+
+		float PixelRadius = -1;
+		const float2 HitT = LineBoxIntersect(O, E, Voxel_MinBound, Voxel_MaxBound);
+		float AccDensity = 0;
+		if (HitT.x < HitT.y)
+		{
+			// Use different offseting to see either: 
+			// * Individual voxel (for debugging mainly)
+			// * Smoothed voxel (as used for transmittance computation)
+			const float3 Offset = 0;
+			const float Dither = InterleavedGradientNoise(PixelCoord.xy, 1);
+
+			// Intersect the indirection table to see if a brick is empty or not
+			const uint MaxStep = 256;
+			const float Delta = (HitT.y - HitT.x) / float(MaxStep);
+			for (uint StepIt = 0; StepIt < MaxStep; ++StepIt)
+			{
+				const float3 HitP = Offset + O + ((Dither + StepIt) * Delta + HitT.x) * D;
+				const int3 VoxelCoord = PositionToCoord(HitP, Voxel_MinBound, Voxel_MaxBound, Voxel_Resolution);
+				const uint VoxelLinearCoord = CoordToIndex(VoxelCoord, Voxel_Resolution, 0);
+
+				const uint Density = Voxel_DensityBuffer.Load(VoxelLinearCoord);
+
+				if (Density > 0)
+				{
+					float3 UnoccludedDirection = UnpackDirection(Density);
+
+
+					// Filtering
+					if (0)
+					{
+						int KernelSize = 1;
+						uint ValidCount = 0;
+						UnoccludedDirection = 0;
+						for (int Z = -KernelSize; Z <= KernelSize; ++Z)
+						for (int Y = -KernelSize; Y <= KernelSize; ++Y)
+						for (int X = -KernelSize; X <= KernelSize; ++X)
+						{
+							const int3 Offset = int3(X, Y, Z);
+							const int3 C = VoxelCoord + Offset;
+							if (all(C >= 0) && all(C < Voxel_Resolution))
+							{
+								const float w = all(Offset == VoxelCoord) ? 2 : 1;
+								const uint LinearC = CoordToIndex(C, Voxel_Resolution, 0);
+								const uint PackedDir = Voxel_ProcessedDensityBuffer[LinearC];
+								if (PackedDir > 0)
+								{
+									const float3 Dir = UnpackDirection(PackedDir);
+									UnoccludedDirection += normalize(Dir) * w;
+									ValidCount++;
+								}
+							}
+						}
+
+						if (ValidCount>0)
+						{
+							UnoccludedDirection = normalize(UnoccludedDirection);
+
+						}						
+
+						// CPU tangent
+						//UnoccludedDirection = Voxel_TangentBuffer.Load(VoxelLinearCoord) / Density;
+						//UnoccludedDirection = normalize(UnoccludedDirection);
+					}
+
+					const float3 Tangent = Voxel_TangentBuffer.Load(VoxelLinearCoord).xyz;
+					const float3 Normal  = Voxel_NormalBuffer.Load(VoxelLinearCoord).xyz;
+					//UnoccludedDirection = Tangent;
+					UnoccludedDirection = Normal;
+
+					AccDensity = 1;
+					bIsValid = true;
+					HitColor = VoxelCoord / 12.f;
+
+					HitColor = UnoccludedDirection * 2 + 1;
+
+					//HitColor = UnoccludedDirection; // (UnoccludedDirection + 1) * 0.5f;
+					//HitColor = float(Density) / 10.f;
+					if (bDebugEnabled)
+					{
+						//AddLineWS(HitP, HitP + UnoccludedDirection * 3, float4(1, 0, 0, 1), float4(1, 1, 0, 1));
+
+						AddLineWS(HitP, HitP + Tangent * 3, float4(0, 0, 1, 1), float4(0, 0, 1, 1));
+						AddLineWS(HitP, HitP + Normal * 3, float4(1, 0, 0, 1), float4(1, 0, 0, 1));
+					}
+
+					break;
+				}
+			}
+		}
+	}
+	
+	if (bIsValid)
+	{
+		OutputTexture[DispatchThreadId.xy] = float4(HitColor, 1);
+	}
+}
+
+#endif //SHADER_CARDS_VOXEL
+
+///////////////////////////////////////////////////////////////////////////////////////////////////
+// Visualize directional occlusion from voxelized hair strands (used for cards generation)
+
+#if SHADER_CARDS_DIRECTION_OCCLUSION
+
+int3 Voxel_Resolution;
+float Voxel_VoxelSize;
+float3 Voxel_MinBound;
+float3 Voxel_MaxBound;
+Buffer<uint> Voxel_DensityBuffer;
+RWBuffer<uint> OutBuffer;
+
+[numthreads(4, 4, 4)]
+void MainCS(uint3 DispatchThreadId : SV_DispatchThreadID)
+{
+	const int3 VoxelCoord = DispatchThreadId.xyz;
+
+	if (any(VoxelCoord >= Voxel_Resolution))
+		return;
+
+	const uint VoxelLinearCoord = CoordToIndex(VoxelCoord, Voxel_Resolution, 0);
+
+	// Dilation
+#if 0
+	uint OutValue   = 0;
+	 int KernelSize = 1;
+	uint ValidCount = 0;
+
+	for (int Z = -KernelSize; Z <= KernelSize; ++Z)
+	for (int Y = -KernelSize; Y <= KernelSize; ++Y)
+	for (int X = -KernelSize; X <= KernelSize; ++X)
+	{
+		int3 C = VoxelCoord + int3(X, Y, Z);
+		if (all(C >= 0) && all(C < Voxel_Resolution))
+		{
+			const uint LinearC = CoordToIndex(C, Voxel_Resolution, 0);
+			OutValue += Voxel_DensityBuffer[LinearC];
+			ValidCount++;
+		}
+	}
+	OutValue /= max(1, ValidCount);
+#endif
+
+	// Occlusion
+#if 1
+	const uint CenterDensity = Voxel_DensityBuffer[VoxelLinearCoord];
+	if (CenterDensity == 0)
+		return;
+
+	int KernelSize = 1;
+	uint ValidCount = 0;
+	float3 OutDirection = 0;
+	for (int Z = -KernelSize; Z <= KernelSize; ++Z)
+	for (int Y = -KernelSize; Y <= KernelSize; ++Y)
+	for (int X = -KernelSize; X <= KernelSize; ++X)
+	{
+		const int3 Offset = int3(X, Y, Z);		
+		const int3 C = VoxelCoord + Offset;
+		if (all(C >= 0) && all(C < Voxel_Resolution) && !all(Offset == int3(0,0,0)))
+		{
+			const uint LinearC = CoordToIndex(C, Voxel_Resolution, 0);
+			const uint Density = Voxel_DensityBuffer[LinearC];
+			if (Density == 0)
+			{
+				OutDirection += normalize(float3(Offset));
+				ValidCount++;
+			}
+		}
+	}
+	if (ValidCount>0)
+	{
+		OutDirection = normalize(OutDirection);
+	}
+	const uint OutValue = PackDirection(OutDirection);
+#endif
+
+	// Pass through
+#if 0
+	const uint OutValue = Voxel_DensityBuffer[VoxelLinearCoord];
+#endif
+
+	OutBuffer[VoxelLinearCoord] = OutValue;
+}
+
+#endif //SHADER_CARDS_DIRECTION_OCCLUSION
+
+///////////////////////////////////////////////////////////////////////////////////////////////////
+
+#if SHADER_PLOTBSDF
+
+#include "../SceneTexturesCommon.ush"
+#include "../ShadingModels.ush"
+#include "../DeferredShadingCommon.ush"
+#include "../HairBsdf.ush"
+#include "HairStrandsCommon.ush"
+#include "HairStrandsVisibilityCommon.ush"
+#include "HairStrandsEnvironmentLightingCommon.ush"
+#include "HairStrandsDeepShadowCommon.ush"
+#include "HairStrandsDeepTransmittanceCommon.ush"
+#include "HairStrandsDeepTransmittanceDualScattering.ush"
+
+int2 InputCoord;
+int2 OutputOffset;
+int2 OutputResolution;
+int2 MaxResolution;
+uint HairComponents;
+float Roughness;
+float BaseColor;
+float Exposure;
+
+RWTexture2D<float4> OutputTexture;
+
+#define COMPARE_ROUGHNESS 0
+
+#if COMPARE_ROUGHNESS
+float DefaultLit(float3 N, float3 V, float3 L, float InRoughnesss)
+{
+	FGBufferData GBuffer;
+	GBuffer.Roughness = InRoughnesss;
+	GBuffer.SpecularColor = 1;
+	GBuffer.DiffuseColor = 0;
+
+	FAreaLight AreaLight;
+	AreaLight.IsRectAndDiffuseMicroReflWeight = 0;
+	AreaLight.FalloffColor = 1;
+	AreaLight.SphereSinAlpha = 0;
+	AreaLight.SphereSinAlphaSoft = 0;
+	AreaLight.LineCosSubtended = 1;
+	FShadowTerms Shadow;
+
+	if (dot(V,N) < 0)
+		V.z = -V.z;
+
+	const float NoL = saturate(dot(N, L));
+	const float3 Fs = DefaultLitBxDF(GBuffer, N, V, L, 1, NoL, AreaLight, Shadow).Specular;
+	return Fs.x;
+}
+#endif
+
+void MainPS(
+	in FScreenVertexOutput Input,
+	out float4 OutColor : SV_Target0)
+{
+	const int2 PixelCoord = Input.Position.xy;
+	if (PixelCoord.x < OutputOffset.x || PixelCoord.x >= OutputOffset.x+OutputResolution.x ||
+		PixelCoord.y < OutputOffset.y || PixelCoord.y >= OutputOffset.x+OutputResolution.y )
+		discard;
+
+	const float2 UV = ((PixelCoord-OutputOffset) + float2(0.5f,0.5f)) / float2(OutputResolution);
+	float2 InputUV = saturate(float2(InputCoord - OutputOffset) / float2(OutputResolution));
+	// InputUV = float2(0.5, 0.5);
+
+
+	float V_Phi   = (UV.x*2-1) * PI; //[-Pi..Pi]
+	float V_Theta = (UV.y*2-1) * PI * 0.5f;
+	float L_Phi   = (InputUV.x*2-1) * PI; //[-Pi..Pi]
+	float L_Theta = (InputUV.y*2-1) * PI * 0.5f;
+
+	const float3 Axis_U = float3(0,0,1); // Tangent
+	const float3 Axis_V = float3(1,0,0); // Perp. to the view direction
+	const float3 Axis_W = float3(0,1,0); // Ortho.
+
+	const float3 L_Local = float3(sin(L_Theta), cos(L_Theta)*cos(L_Phi), cos(L_Theta)*sin(L_Phi));
+	const float3 V_Local = float3(sin(V_Theta), cos(V_Theta)*cos(V_Phi), cos(V_Theta)*sin(V_Phi));
+	const float3 T 		 = float3(1,0,0);
+
+	const float3 V_World = 
+		V_Local.x * Axis_U + 
+		V_Local.y * Axis_V + 
+		V_Local.z * Axis_W;
+
+	FHairTransmittanceData TransmittanceData = (FHairTransmittanceData)0;
+	TransmittanceData.LocalScattering = 0;
+	TransmittanceData.GlobalScattering = 1;
+	TransmittanceData.ScatteringComponent = HairComponents;
+	TransmittanceData.OpaqueVisibility = 1;
+
+	const float Backlit = 1;
+	const float Specular = 0.5f;
+	FGBufferData HairGBuffer = (FGBufferData)0;
+	HairGBuffer.WorldNormal = T;
+	HairGBuffer.BaseColor = BaseColor.xxx;
+	HairGBuffer.Roughness = Roughness;
+	HairGBuffer.ShadingModelID = SHADINGMODELID_HAIR;
+	HairGBuffer.DiffuseColor = 0;
+	HairGBuffer.SpecularColor = 0;
+	HairGBuffer.Specular = Specular;
+	HairGBuffer.Metallic = 0;
+	HairGBuffer.Depth = 0.5f;
+	HairGBuffer.GBufferAO = 1;
+	HairGBuffer.CustomData = float4(0,0,Backlit,0);
+	HairGBuffer.IndirectIrradiance = 1000000;
+
+	OutColor = float4(0,0,0,1);
+
+#if COMPARE_ROUGHNESS
+	const float3 Fs_Lit_R   = DefaultLit(float3(0, 0,-1), V_Local, L_Local, HairToGGXRoughness_R  (HairGBuffer.Roughness));
+	const float3 Fs_Lit_TT  = DefaultLit(float3(0, 0,-1), V_Local, L_Local, HairToGGXRoughness_TT (HairGBuffer.Roughness));
+	const float3 Fs_Lit_TRT = DefaultLit(float3(0, 0,-1), V_Local, L_Local, HairToGGXRoughness_TRT(HairGBuffer.Roughness));
+#endif
+	const float3 Fs_Hair    = HairShading(HairGBuffer, L_Local, V_Local, T, 1, TransmittanceData, 1, 0, uint2(0,0) );
+
+#if COMPARE_ROUGHNESS
+	OutColor.xyz += Fs_Lit_R   * float3(1, 0, 0);
+	OutColor.xyz += Fs_Lit_TRT * float3(0, 1, 1);
+	OutColor.xyz += Fs_Lit_TT  * float3(0, 0, 1);
+#endif
+	OutColor.xyz += Fs_Hair;
+	OutColor.xyz *= Exposure;
+
+	//const float3 ColorV = (V_World+1)* 0.5f;
+	//OutColor = float4(ColorV, 1);
+	//OutColor = float4(UV, 0, 1);
+}
+
+#endif // SHADER_PLOTBSDF
+
+//////////////////////////////////////////////////////////////////////////////////////////////////////////////////////////////
+
+#if SHADER_PLOTSAMPLE
+
+#include "HairStrandsDebugCommonStruct.ush"
+#include "../SceneTexturesCommon.ush"
+#include "../ShadingModels.ush"
+#include "../DeferredShadingCommon.ush"
+#include "../HairBsdf.ush"
+#include "HairStrandsCommon.ush"
+#include "HairStrandsVisibilityCommon.ush"
+#include "HairStrandsEnvironmentLightingCommon.ush"
+#include "HairStrandsDeepShadowCommon.ush"
+#include "HairStrandsDeepTransmittanceCommon.ush"
+#include "HairStrandsDeepTransmittanceDualScattering.ush"
+
+StructuredBuffer<FDebugShadingInfo>	Debug_ShadingPointBuffer;
+Buffer<uint>						Debug_ShadingPointCounter;
+StructuredBuffer<FDebugSample>		Debug_SampleBuffer;
+Buffer<uint>						Debug_SampleCounter;
+uint								Debug_MaxShadingPointCount;
+uint								Debug_MaxSampleCount;
+
+Texture3D<float4>					HairScatteringLUTTexture;
+SamplerState 						HairLUTSampler;
+
+int2 OutputOffset;
+int2 OutputResolution;
+int2 MaxResolution;
+uint HairComponents;
+float Exposure;
+
+struct FPlotRect
+{
+	int2 Offset;
+	int2 Resolution;
+};
+
+bool IsInside(const int2 Coord, const FPlotRect Rect)
+{
+	return all(Coord >= Rect.Offset) && all(Coord < Rect.Offset + Rect.Resolution);
+}
+
+float2 ToUV(const int2 Coord, const FPlotRect Rect)
+{
+	return ((Coord - Rect.Offset) + float2(0.5f, 0.5f)) / float2(OutputResolution);
+}
+
+RWTexture2D<float4> OutputTexture;
+
+float2 CartesianToUnwrap(float3 D)
+{
+	const float Theta = asin(D.x);
+	const float Phi = atan2(D.z, D.y);
+	return float2(Theta, Phi);
+}
+
+float2 CartesianToLatlong(float3 D)
+{
+	const float Theta = asin(D.x);
+	const float Phi = atan2(D.z, D.y);
+	return float2(((Phi/PI)+1)*0.5f, Theta/PI);
+}
+
+void MainPS(
+	in FScreenVertexOutput Input,
+	out float4 OutColor : SV_Target0)
+{
+	OutColor = float4(0, 0, 0, 1);
+
+	// Prune all pixel which are not within the band horizontal band
+	const int2 PixelCoord = Input.Position.xy;
+	if (PixelCoord.y < OutputOffset.y || PixelCoord.y >= OutputOffset.x + OutputResolution.y)
+	{
+		discard;
+	}
+
+	const float CosAngleThreshold = 0.999f;
+	const float PointSize = 3.f / OutputResolution.x;
+
+	// Local 
+	// Plot sample points
+	bool bIsValid = false;
+	const uint ShadingPointCount = Debug_ShadingPointCounter[0];
+	for (uint PointIt = 0; PointIt < ShadingPointCount; PointIt++)
+	{
+		FPlotRect Rect;
+		Rect.Offset = OutputOffset + PointIt * int2(OutputResolution.x + 10, 0);
+		Rect.Resolution = OutputResolution;
+
+		if (!IsInside(PixelCoord, Rect))
+			continue;
+
+		bIsValid = true;
+		const float2 UV = ToUV(PixelCoord, Rect);
+
+		const float  Pixel_Phi	= (UV.x * 2 - 1) * PI + PI*0.5; //[3PI/2, 0, 2PI, 3PI/2 ], so that the hair fiber is aligned on the center
+		const float  Pixel_Theta= (UV.y * 2 - 1) * PI * 0.5f;
+		const float3 PixelDir_Local	= float3(sin(Pixel_Theta), cos(Pixel_Theta)*cos(Pixel_Phi), cos(Pixel_Theta)*sin(Pixel_Phi));
+
+		const FDebugShadingInfo P = Debug_ShadingPointBuffer[PointIt];
+
+		// Shading frame T, V
+		const float3 Local_X = P.T;
+		const float3 Local_Y = normalize(cross(Local_X, P.V));
+		const float3 Local_Z = normalize(cross(Local_Y, Local_X));
+		const float3x3 WorldToLocal = transpose(float3x3(Local_X, Local_Y, Local_Z));
+
+		const float3 V_Local = mul(P.V, WorldToLocal);
+		const float3 T_Local = mul(P.T, WorldToLocal); // Should be float3(1,0,0)
+
+		{
+
+			const float SampleBacklit = 1;
+			const float SampleSpecular = 0.5f;
+			FHairTransmittanceData TransmittanceData = InitHairStrandsTransmittanceData();
+
+			uint InHairComponents = HairComponents;
+			const bool bPlotDualScatteringContribution = true;
+			if (!bPlotDualScatteringContribution)
+			{
+				TransmittanceData.LocalScattering = 1;
+				TransmittanceData.GlobalScattering = 1;
+				TransmittanceData.ScatteringComponent = InHairComponents;
+				TransmittanceData.OpaqueVisibility = 1;
+			}
+			else
+			{
+				const float3 L = PixelDir_Local;
+				const float3 T = T_Local;
+				const float3 V = V_Local;
+				const float SinLightAngle = dot(L, T);
+
+				// Compute the transmittance based on precompute Hair transmittance LUT
+				FHairAverageScattering AverageScattering = SampleHairLUT(View.HairScatteringLUTTexture, View.HairScatteringLUTSampler, P.BaseColor, P.Roughness, SinLightAngle);
+
+				TransmittanceData = ComputeDualScatteringTerms(
+					InitHairTransmittanceMask(),
+					AverageScattering,
+					P.Roughness,
+					SampleBacklit,
+					V,
+					L,
+					T,
+					InHairComponents);
+			}
+
+			FGBufferData HairGBuffer = (FGBufferData)0;
+			HairGBuffer.WorldNormal = P.T;
+			HairGBuffer.BaseColor = P.BaseColor;
+			HairGBuffer.Roughness = P.Roughness;
+			HairGBuffer.ShadingModelID = SHADINGMODELID_HAIR;
+			HairGBuffer.DiffuseColor = 0;
+			HairGBuffer.SpecularColor = 0;
+			HairGBuffer.Specular = SampleSpecular;
+			HairGBuffer.Metallic = 0;
+			HairGBuffer.Depth = 0.5f;
+			HairGBuffer.GBufferAO = 1;
+			HairGBuffer.CustomData = float4(0,0,SampleBacklit,0);
+			HairGBuffer.IndirectIrradiance = 1000000;
+
+			const float Area = 0;
+			const float BackLit = 1;
+			OutColor.xyz = HairShading(HairGBuffer, PixelDir_Local, V_Local, T_Local, 1, TransmittanceData, BackLit, Area, uint2(0, 0)) * Exposure;
+		}
+
+		for (uint SampleIt = 0; SampleIt < P.SampleCount; SampleIt++)
+		{
+			const FDebugSample S_World = Debug_SampleBuffer[P.SampleOffset + SampleIt];
+
+			const float3 S_Local = mul(S_World.Direction, WorldToLocal);
+
+			if (dot(PixelDir_Local, S_Local) > CosAngleThreshold)
+			{
+				OutColor.xyz = float3(1, 0, 0);
+			}
+		}
+	}
+
+	if (!bIsValid)
+	{
+		discard;
+	}
+}
+
+#endif // SHADER_PLOTSAMPLE
+
+//////////////////////////////////////////////////////////////////////////////////////////////////////////////////////////////
\ No newline at end of file
diff --git a/Engine/Shaders/Private/HairStrands/HairStrandsDebugPrint.usf b/Engine/Shaders/Private/HairStrands/HairStrandsDebugPrint.usf
deleted file mode 100644
diff --git a/Engine/Shaders/Private/HairStrands/HairStrandsDeepShadowDebug.usf b/Engine/Shaders/Private/HairStrands/HairStrandsDeepShadowDebug.usf
deleted file mode 100644
diff --git a/Engine/Shaders/Private/HairStrands/HairStrandsMeshProjectionHairDebug.usf b/Engine/Shaders/Private/HairStrands/HairStrandsMeshProjectionHairDebug.usf
deleted file mode 100644
diff --git a/Engine/Shaders/Private/HairStrands/HairStrandsMeshProjectionMeshDebug.usf b/Engine/Shaders/Private/HairStrands/HairStrandsMeshProjectionMeshDebug.usf
deleted file mode 100644
diff --git a/Engine/Shaders/Private/HairStrands/HairStrandsVisibilityPPLLDebug.usf b/Engine/Shaders/Private/HairStrands/HairStrandsVisibilityPPLLDebug.usf
deleted file mode 100644
diff --git a/Engine/Source/Runtime/Renderer/Private/HairStrands/HairStrandsDebug.cpp b/Engine/Source/Runtime/Renderer/Private/HairStrands/HairStrandsDebug.cpp
--- a/Engine/Source/Runtime/Renderer/Private/HairStrands/HairStrandsDebug.cpp
+++ b/Engine/Source/Runtime/Renderer/Private/HairStrands/HairStrandsDebug.cpp
@@ -252,7 +252,7 @@ public:
 	}
 };
 
-IMPLEMENT_GLOBAL_SHADER(FHairPrintLODInfoCS, "/Engine/Private/HairStrands/HairStrandsDebugPrint.usf", "MainCS", SF_Compute);
+IMPLEMENT_GLOBAL_SHADER(FHairPrintLODInfoCS, "/Engine/Private/HairStrands/HairStrandsDebug.usf", "MainCS", SF_Compute);
 
 static void AddPrintLODInfoPass(
 	FRDGBuilder& GraphBuilder,
@@ -339,7 +339,7 @@ public:
 	}
 };
 
-IMPLEMENT_GLOBAL_SHADER(FHairDebugPrintCS, "/Engine/Private/HairStrands/HairStrandsDebugPrint.usf", "MainCS", SF_Compute);
+IMPLEMENT_GLOBAL_SHADER(FHairDebugPrintCS, "/Engine/Private/HairStrands/HairStrandsDebug.usf", "MainCS", SF_Compute);
 
 static void AddDebugHairPrintPass(
 	FRDGBuilder& GraphBuilder,
@@ -572,7 +572,7 @@ public:
 	}
 };
 
-IMPLEMENT_GLOBAL_SHADER(FDeepShadowVisualizePS, "/Engine/Private/HairStrands/HairStrandsDeepShadowDebug.usf", "VisualizeDomPS", SF_Pixel);
+IMPLEMENT_GLOBAL_SHADER(FDeepShadowVisualizePS, "/Engine/Private/HairStrands/HairStrandsDebug.usf", "VisualizeDomPS", SF_Pixel);
 
 static void AddDebugDeepShadowTexturePass(
 	FRDGBuilder& GraphBuilder,
@@ -675,7 +675,7 @@ public:
 	}
 };
 
-IMPLEMENT_GLOBAL_SHADER(FDeepShadowInfoCS, "/Engine/Private/HairStrands/HairStrandsDeepShadowDebug.usf", "MainCS", SF_Compute);
+IMPLEMENT_GLOBAL_SHADER(FDeepShadowInfoCS, "/Engine/Private/HairStrands/HairStrandsDebug.usf", "MainCS", SF_Compute);
 
 static void AddDeepShadowInfoPass(
 	FRDGBuilder& GraphBuilder,
@@ -814,7 +814,7 @@ public:
 	}
 };
 
-IMPLEMENT_GLOBAL_SHADER(FDebugHairTangentCS, "/Engine/Private/HairStrands/HairStrandsDebugPrint.usf", "MainCS", SF_Compute);
+IMPLEMENT_GLOBAL_SHADER(FDebugHairTangentCS, "/Engine/Private/HairStrands/HairStrandsDebug.usf", "MainCS", SF_Compute);
 
 static void AddDebugHairTangentPass(
 	FRDGBuilder& GraphBuilder,
@@ -871,7 +871,7 @@ public:
 	}
 };
 
-IMPLEMENT_GLOBAL_SHADER(FHairStrandsPlotBSDFPS, "/Engine/Private/HairStrands/HairStrandsBsdfPlot.usf", "MainPS", SF_Pixel);
+IMPLEMENT_GLOBAL_SHADER(FHairStrandsPlotBSDFPS, "/Engine/Private/HairStrands/HairStrandsDebug.usf", "MainPS", SF_Pixel);
 
 static void AddPlotBSDFPass(
 	FRDGBuilder& GraphBuilder,
@@ -961,7 +961,7 @@ public:
 	}
 };
 
-IMPLEMENT_GLOBAL_SHADER(FHairStrandsPlotSamplePS, "/Engine/Private/HairStrands/HairStrandsBsdfPlot.usf", "MainPS", SF_Pixel);
+IMPLEMENT_GLOBAL_SHADER(FHairStrandsPlotSamplePS, "/Engine/Private/HairStrands/HairStrandsDebug.usf", "MainPS", SF_Pixel);
 
 static void AddPlotSamplePass(
 	FRDGBuilder& GraphBuilder,
@@ -1058,12 +1058,12 @@ class FHairVisibilityDebugPPLLCS : public FGlobalShader
 	static void ModifyCompilationEnvironment(const FGlobalShaderPermutationParameters& Parameters, FShaderCompilerEnvironment& OutEnvironment)
 	{
 		FGlobalShader::ModifyCompilationEnvironment(Parameters, OutEnvironment);
-		OutEnvironment.SetDefine(TEXT("DEBUG_PPLL_PS"), 1);
+		OutEnvironment.SetDefine(TEXT("SHADER_PPLL_DEBUG"), 1);
 		// Skip optimization for avoiding long compilation time due to large UAV writes
 		OutEnvironment.CompilerFlags.Add(CFLAG_Debug);
 	}
 };
-IMPLEMENT_GLOBAL_SHADER(FHairVisibilityDebugPPLLCS, "/Engine/Private/HairStrands/HairStrandsVisibilityPPLLDebug.usf", "VisibilityDebugPPLLCS", SF_Compute);
+IMPLEMENT_GLOBAL_SHADER(FHairVisibilityDebugPPLLCS, "/Engine/Private/HairStrands/HairStrandsDebug.usf", "VisibilityDebugPPLLCS", SF_Compute);
 
 ///////////////////////////////////////////////////////////////////////////////////////////////////
 
