Fix banned word warning.

#rnx

diff --git a/Engine/Shaders/Private/TextureSampling.ush b/Engine/Shaders/Private/TextureSampling.ush
--- a/Engine/Shaders/Private/TextureSampling.ush
+++ b/Engine/Shaders/Private/TextureSampling.ush
@@ -132,7 +132,7 @@ struct FCatmullRomSamples
 	// Constant number of samples (BICUBIC_CATMULL_ROM_SAMPLES)
 	uint Count;
 
-	// Constant sign of the UV direction from master UV sampling location.
+	// Constant sign of the UV direction from main UV sampling location.
 	int2 UVDir[BICUBIC_CATMULL_ROM_SAMPLES];
 
 	// Bilinear sampling UV coordinates of the samples
