fix selection outline for HDR editor: ScRGB can generate negative values

diff --git a/Engine/Shaders/Private/PostProcessSelectionOutline.usf b/Engine/Shaders/Private/PostProcessSelectionOutline.usf
--- a/Engine/Shaders/Private/PostProcessSelectionOutline.usf
+++ b/Engine/Shaders/Private/PostProcessSelectionOutline.usf
@@ -265,9 +265,11 @@ void MainPS(
 		OutColor = Texture2DSample(ColorTexture, ColorSampler, ColorUV);
 
 		// Scene color has gamma applied. Need to remove and re-apply after linear operation.
-		OutColor.rgb = pow(OutColor.rgb, 2.2f );
+		float3 OutColorSigns = sign(OutColor.rgb);
+		OutColor.rgb = pow(abs(OutColor.rgb), 2.2f );
 		OutColor.rgb = lerp(OutColor.rgb, Avg.rgb, Avg.a);
 		OutColor.rgb = pow(OutColor.rgb, 1.0f / 2.2f);
+		OutColor.rgb *= OutColorSigns;
 	}
 	else
 	{
