Path Tracer: Dithered opacity masked materials should ignore the clip value

The original behavior was technically correct but the net effect was surprising to most users. The intent of dithered opacity is to emulate real transparency, so just use the opacity mask value directly.

diff --git a/Engine/Shaders/Private/PathTracing/PathTracingMaterialHitShader.usf b/Engine/Shaders/Private/PathTracing/PathTracingMaterialHitShader.usf
--- a/Engine/Shaders/Private/PathTracing/PathTracingMaterialHitShader.usf
+++ b/Engine/Shaders/Private/PathTracing/PathTracingMaterialHitShader.usf
@@ -154,7 +154,14 @@ RAY_TRACING_ENTRY_CLOSEST_HIT(PathTracingMaterialCHS,
 #if MATERIALBLENDING_MASKED
 #if MATERIAL_DITHER_OPACITY_MASK
 	// dithering emulates real transparency, so switch to translucent
-	Payload.Opacity = DitheredOpacityMaskToOpacity(GetMaterialMask(PixelMaterialInputs));
+	// NOTE: the raster path technically takes into account the opacity mask clip value, so the effective transparency should be:
+	//        saturate(MaskRaw - ClipValue + 0.5)
+	// (See derivation in DitheredOpacityMaskToOpacity)
+	// However this behavior is surprising to most users and does not exactly match the rasterizer anyway due to how the realtime AA
+	// code performs blending.
+	// Since the goal of dithered opacity is to emulate ordinary transparency, just use the mask input as opacity directly and
+	// ignore the configured clip value.
+	Payload.Opacity = saturate(GetMaterialMaskInputRaw(PixelMaterialInputs));
 	Payload.BlendingMode = RAY_TRACING_BLEND_MODE_TRANSLUCENT;
 #else
 	// regular masked mode - accounted for in AHS
@@ -482,7 +489,8 @@ RAY_TRACING_ENTRY_ANY_HIT(PathTracingMaterialAHS,
 		Transparency *= A;
 	}
 #elif MATERIALBLENDING_MASKED && MATERIAL_DITHER_OPACITY_MASK
-	float Opacity = DitheredOpacityMaskToOpacity(GetMaterialMask(PixelMaterialInputs));
+	// See MATERIAL_DITHER_OPACITY_MASK comment above
+	float Opacity = saturate(GetMaterialMaskInputRaw(PixelMaterialInputs));
 	float Transparency = 1 - Opacity;
 #elif MATERIALBLENDING_TRANSLUCENT
 	float Opacity = GetMaterialOpacity(PixelMaterialInputs);
