Add common root triangle loading in hair interpolation code.

diff --git a/Engine/Shaders/Private/HairStrands/HairStrandsInterpolation.usf b/Engine/Shaders/Private/HairStrands/HairStrandsInterpolation.usf
--- a/Engine/Shaders/Private/HairStrands/HairStrandsInterpolation.usf
+++ b/Engine/Shaders/Private/HairStrands/HairStrandsInterpolation.usf
@@ -202,6 +202,36 @@ uint4 WriteOutput(float3 WorldPosition, float3 OutPositionOffset, uint PackedDat
 	return uint4(EncodedRelativePosition, PackedDataW);
 }
 
+struct FRootTriangleData
+{
+	uint RootIndex;
+	float3 RootBarycentric;
+	FHairMeshTriangle RestTriangle;
+	FHairMeshTriangle DeformedTriangle;
+};
+
+FRootTriangleData GetRootTriangleData(uint VertexIndex, bool bSim)
+{
+	FRootTriangleData Out = (FRootTriangleData)0;
+	#if PERMUTATION_DYNAMIC_GEOMETRY >= 1 
+	if (bSim)
+	{
+		Out.RootIndex		 = SimVertexToRootIndexBuffer[VertexIndex];
+		Out.RootBarycentric  = DecodeBarycentrics(SimRootBarycentricBuffer[Out.RootIndex]);
+		Out.RestTriangle	 = GetTriangleTransformation(Out.RootIndex, SimRestPosition0Buffer, SimRestPosition1Buffer, SimRestPosition2Buffer);
+		Out.DeformedTriangle = GetTriangleTransformation(Out.RootIndex, SimDeformedPosition0Buffer, SimDeformedPosition1Buffer, SimDeformedPosition2Buffer);	
+	}
+	else
+	{
+		Out.RootIndex		 = RenVertexToRootIndexBuffer[VertexIndex];
+		Out.RootBarycentric  = DecodeBarycentrics(RootBarycentricBuffer[Out.RootIndex]);
+		Out.RestTriangle	 = GetTriangleTransformation(Out.RootIndex, RestPosition0Buffer, RestPosition1Buffer, RestPosition2Buffer);
+		Out.DeformedTriangle = GetTriangleTransformation(Out.RootIndex, DeformedPosition0Buffer, DeformedPosition1Buffer, DeformedPosition2Buffer);
+	}
+	#endif
+	return Out;
+}
+
 [numthreads(GROUP_SIZE, 1, 1)]
 void MainCS(uint2 DispatchThreadId : SV_DispatchThreadID)
 {
@@ -313,27 +343,20 @@ void MainCS(uint2 DispatchThreadId : SV_DispatchThreadID)
 				const uint GuideIndex0 = GuideVertexIndices[K];
 				const uint GuideIndex1 = GuideIndex0 + 1;
 
-				const uint RootIndex = SimVertexToRootIndexBuffer[GuideIndex0];
-				const float3 RootBarycentric = DecodeBarycentrics(SimRootBarycentricBuffer[RootIndex]);
-
-				const FHairMeshTriangle RestTriangle		= GetTriangleTransformation(RootIndex, SimRestPosition0Buffer, SimRestPosition1Buffer, SimRestPosition2Buffer);
-				const FHairMeshTriangle DeformedTriangle	= GetTriangleTransformation(RootIndex, SimDeformedPosition0Buffer, SimDeformedPosition1Buffer, SimDeformedPosition2Buffer);
+				const FRootTriangleData SimTri = GetRootTriangleData(GuideIndex0, true);
 
-				const float3 Offset0 = ComputeDynamicGeometryOffset(GuideIndex0, GuideVertexWeights[K], RestTriangle, DeformedTriangle, RootBarycentric, SimHairPositionOffset);
-				const float3 Offset1 = ComputeDynamicGeometryOffset(GuideIndex1, GuideVertexWeights[K], RestTriangle, DeformedTriangle, RootBarycentric, SimHairPositionOffset);
+				const float3 Offset0 = ComputeDynamicGeometryOffset(GuideIndex0, GuideVertexWeights[K], SimTri.RestTriangle, SimTri.DeformedTriangle, SimTri.RootBarycentric, SimHairPositionOffset);
+				const float3 Offset1 = ComputeDynamicGeometryOffset(GuideIndex1, GuideVertexWeights[K], SimTri.RestTriangle, SimTri.DeformedTriangle, SimTri.RootBarycentric, SimHairPositionOffset);
 
-				CurrOffset += VectorToWorld(lerp(Offset0, Offset1, GuideVertexLerps[K]),DeformedTriangle);
+				CurrOffset += VectorToWorld(lerp(Offset0, Offset1, GuideVertexLerps[K]), SimTri.DeformedTriangle);
 			}
 			#endif		
 
 			// Transform hair from rest post to deformed pose, based triangle deformation + simulation offset (optional depending of the simulation is actual running or not
-			const uint RootIndex						= RenVertexToRootIndexBuffer[VertexIndex];
-			const float3 RootBarycentric				= DecodeBarycentrics(RootBarycentricBuffer[RootIndex]);
-			const FHairMeshTriangle RestTriangle		= GetTriangleTransformation(RootIndex, RestPosition0Buffer, RestPosition1Buffer, RestPosition2Buffer);
-			const FHairMeshTriangle DeformedTriangle	= GetTriangleTransformation(RootIndex, DeformedPosition0Buffer, DeformedPosition1Buffer, DeformedPosition2Buffer);
+			const FRootTriangleData RenTri = GetRootTriangleData(VertexIndex, false);
 
-			ControlPoint = ToTriangle(ControlPoint,RootBarycentric,RestTriangle);
-			ControlPoint = ToWorld(ControlPoint,RootBarycentric,DeformedTriangle) + CurrOffset;
+			ControlPoint = ToTriangle(ControlPoint, RenTri.RootBarycentric, RenTri.RestTriangle);
+			ControlPoint = ToWorld(ControlPoint, RenTri.RootBarycentric,RenTri.DeformedTriangle) + CurrOffset;
 
 			WorldVertexPosition = ControlPoint;
 
@@ -359,8 +382,8 @@ void MainCS(uint2 DispatchThreadId : SV_DispatchThreadID)
 					const bool bIsRoot = ControlPointType == 1;
 					if (bIsRoot)
 					{
-						AddLineTriangleWS(RestTriangle.P0, RestTriangle.P1, RestTriangle.P2, RestColor);
-						AddLineTriangleWS(DeformedTriangle.P0, DeformedTriangle.P1, DeformedTriangle.P2, DeformedColor);
+						AddLineTriangleWS(RenTri.RestTriangle.P0, RenTri.RestTriangle.P1, RenTri.RestTriangle.P2, RestColor);
+						AddLineTriangleWS(RenTri.DeformedTriangle.P0, RenTri.DeformedTriangle.P1, RenTri.DeformedTriangle.P2, DeformedColor);
 					}
 				}
 			}
@@ -377,12 +400,9 @@ void MainCS(uint2 DispatchThreadId : SV_DispatchThreadID)
 			const float3 RestControlPoint = f16tof32(PackedRestPosition.xyz) + InRenderHairPositionOffset;
 			float3 ControlPoint = RestControlPoint;
 
-			const uint RootIndex						= RenVertexToRootIndexBuffer[VertexIndex];
-			const float3 RootBarycentric				= DecodeBarycentrics(RootBarycentricBuffer[RootIndex]);
-			const FHairMeshTriangle RestTriangle		= GetTriangleTransformation(RootIndex, RestPosition0Buffer, RestPosition1Buffer, RestPosition2Buffer);
-			const FHairMeshTriangle DeformedTriangle	= GetTriangleTransformation(RootIndex, DeformedPosition0Buffer, DeformedPosition1Buffer, DeformedPosition2Buffer);
+			const FRootTriangleData RenTri = GetRootTriangleData(VertexIndex, false);
 			
-			ControlPoint -= RootBarycentric.x * RestTriangle.P0 + RootBarycentric.y * RestTriangle.P1 + RootBarycentric.z * RestTriangle.P2;
+			ControlPoint -= RenTri.RootBarycentric.x * RenTri.RestTriangle.P0 + RenTri.RootBarycentric.y * RenTri.RestTriangle.P1 + RenTri.RootBarycentric.z * RenTri.RestTriangle.P2;
 			
 			float3 CurrOffset = 0;
 
@@ -404,7 +424,7 @@ void MainCS(uint2 DispatchThreadId : SV_DispatchThreadID)
 			}
 			#endif
 
-			ControlPoint += RootBarycentric.x * DeformedTriangle.P0 + RootBarycentric.y * DeformedTriangle.P1 + RootBarycentric.z * DeformedTriangle.P2;
+			ControlPoint += RenTri.RootBarycentric.x * RenTri.DeformedTriangle.P0 + RenTri.RootBarycentric.y * RenTri.DeformedTriangle.P1 + RenTri.RootBarycentric.z * RenTri.DeformedTriangle.P2;
 
 			WorldVertexPosition = ControlPoint + CurrOffset;
 
@@ -419,12 +439,8 @@ void MainCS(uint2 DispatchThreadId : SV_DispatchThreadID)
 			const float3 RestControlPoint = f16tof32(PackedRestPosition.xyz) + InRenderHairPositionOffset;
 			float3 ControlPoint = RestControlPoint;
 
-			const uint RootIndex						= RenVertexToRootIndexBuffer[VertexIndex];
-			const float3 RootBarycentric				= DecodeBarycentrics(RootBarycentricBuffer[RootIndex]);
-			const FHairMeshTriangle RestTriangle		= GetTriangleTransformation(RootIndex, RestPosition0Buffer, RestPosition1Buffer, RestPosition2Buffer);
-			const FHairMeshTriangle DeformedTriangle	= GetTriangleTransformation(RootIndex, DeformedPosition0Buffer, DeformedPosition1Buffer, DeformedPosition2Buffer);
-
-			ControlPoint -= RootBarycentric.x * RestTriangle.P0 + RootBarycentric.y * RestTriangle.P1 + RootBarycentric.z * RestTriangle.P2;
+			const FRootTriangleData RenTri = GetRootTriangleData(VertexIndex, false);
+			ControlPoint -= RenTri.RootBarycentric.x * RenTri.RestTriangle.P0 + RenTri.RootBarycentric.y * RenTri.RestTriangle.P1 + RenTri.RootBarycentric.z * RenTri.RestTriangle.P2;
 
 			float3 CurrOffset = 0;
 
@@ -438,22 +454,19 @@ void MainCS(uint2 DispatchThreadId : SV_DispatchThreadID)
 			{
 				const uint GuideIndex0 = GuideVertexIndices[K];
 				const uint GuideIndex1 = GuideIndex0 + 1;
-
-				const uint RootIndexG = SimVertexToRootIndexBuffer[GuideIndex0];
-				
-				const FHairMeshTriangle RestTriangleG		= GetTriangleTransformation(RootIndexG, SimRestPosition0Buffer, SimRestPosition1Buffer, SimRestPosition2Buffer);
-				const FHairMeshTriangle DeformedTriangleG	= GetTriangleTransformation(RootIndexG, SimDeformedPosition0Buffer, SimDeformedPosition1Buffer, SimDeformedPosition2Buffer);
+			
+				const FRootTriangleData SimTri = GetRootTriangleData(GuideIndex0, true);
 				
-				const float3 LocalPoint = VectorToTriangle(ControlPoint,RestTriangleG);
+				const float3 LocalPoint = VectorToTriangle(ControlPoint, SimTri.RestTriangle);
 
-				const float3 Offset0 = ComputeRotateGeometryOffset(GuideIndex0, 1.0, RestTriangleG, DeformedTriangleG/*, SimHairPositionOffset*/);
-				const float3 Offset1 = ComputeRotateGeometryOffset(GuideIndex1, 1.0, RestTriangleG, DeformedTriangleG/*, SimHairPositionOffset*/);
+				const float3 Offset0 = ComputeRotateGeometryOffset(GuideIndex0, 1.0, SimTri.RestTriangle, SimTri.DeformedTriangle/*, SimHairPositionOffset*/);
+				const float3 Offset1 = ComputeRotateGeometryOffset(GuideIndex1, 1.0, SimTri.RestTriangle, SimTri.DeformedTriangle/*, SimHairPositionOffset*/);
 
-				CurrOffset += GuideVertexWeights[K] * VectorToWorld(lerp(Offset0, Offset1, GuideVertexLerps[K])+LocalPoint,DeformedTriangleG);
+				CurrOffset += GuideVertexWeights[K] * VectorToWorld(lerp(Offset0, Offset1, GuideVertexLerps[K])+LocalPoint, SimTri.DeformedTriangle);
 			}
 			#endif
 
-			ControlPoint =  RootBarycentric.x * DeformedTriangle.P0 + RootBarycentric.y * DeformedTriangle.P1 + RootBarycentric.z * DeformedTriangle.P2;
+			ControlPoint = RenTri.RootBarycentric.x * RenTri.DeformedTriangle.P0 + RenTri.RootBarycentric.y * RenTri.DeformedTriangle.P1 + RenTri.RootBarycentric.z * RenTri.DeformedTriangle.P2;
 
 			WorldVertexPosition = ControlPoint + CurrOffset;
 
@@ -468,11 +481,6 @@ void MainCS(uint2 DispatchThreadId : SV_DispatchThreadID)
 			const float3 RestControlPoint = f16tof32(PackedRestPosition.xyz) + InRenderHairPositionOffset;
 			float3 ControlPoint = RestControlPoint;
 
-			const uint RootIndex						= RenVertexToRootIndexBuffer[VertexIndex];
-			const float3 RootBarycentric				= DecodeBarycentrics(RootBarycentricBuffer[RootIndex]);
-			const FHairMeshTriangle RestTriangle		= GetTriangleTransformation(RootIndex, RestPosition0Buffer, RestPosition1Buffer, RestPosition2Buffer);
-			const FHairMeshTriangle DeformedTriangle	= GetTriangleTransformation(RootIndex, DeformedPosition0Buffer, DeformedPosition1Buffer, DeformedPosition2Buffer);
-
 			float3 CurrOffset = 0;
 
 			#if PERMUTATION_SIMULATION == 1
@@ -486,18 +494,14 @@ void MainCS(uint2 DispatchThreadId : SV_DispatchThreadID)
 				const uint GuideIndex0 = GuideVertexIndices[K];
 				const uint GuideIndex1 = GuideIndex0 + 1;
 
-				const uint RootIndexG = SimVertexToRootIndexBuffer[GuideIndex0];
-				const float3 RootBarycentricG				= DecodeBarycentrics(RootBarycentricBuffer[RootIndexG]);
-				
-				const FHairMeshTriangle RestTriangleG		= GetTriangleTransformation(RootIndexG, SimRestPosition0Buffer, SimRestPosition1Buffer, SimRestPosition2Buffer);
-				const FHairMeshTriangle DeformedTriangleG	= GetTriangleTransformation(RootIndexG, SimDeformedPosition0Buffer, SimDeformedPosition1Buffer, SimDeformedPosition2Buffer);
+				const FRootTriangleData SimTri = GetRootTriangleData(GuideIndex0, true);
 				
-				const float3 LocalPoint = ToTriangle(ControlPoint,RootBarycentricG,RestTriangleG);
+				const float3 LocalPoint = ToTriangle(ControlPoint, SimTri.RootBarycentric, SimTri.RestTriangle);
 
-				const float3 Offset0 = ComputeDynamicGeometryOffset(GuideIndex0, 1.0, RestTriangleG, DeformedTriangleG, RootBarycentricG, SimHairPositionOffset);
-				const float3 Offset1 = ComputeDynamicGeometryOffset(GuideIndex1, 1.0, RestTriangleG, DeformedTriangleG, RootBarycentricG, SimHairPositionOffset);
+				const float3 Offset0 = ComputeDynamicGeometryOffset(GuideIndex0, 1.0, SimTri.RestTriangle, SimTri.DeformedTriangle, SimTri.RootBarycentric, SimHairPositionOffset);
+				const float3 Offset1 = ComputeDynamicGeometryOffset(GuideIndex1, 1.0, SimTri.RestTriangle, SimTri.DeformedTriangle, SimTri.RootBarycentric, SimHairPositionOffset);
 
-				CurrOffset += GuideVertexWeights[K] * ToWorld(lerp(Offset0, Offset1, GuideVertexLerps[K])+LocalPoint,RootBarycentricG,DeformedTriangleG);
+				CurrOffset += GuideVertexWeights[K] * ToWorld(lerp(Offset0, Offset1, GuideVertexLerps[K])+LocalPoint, SimTri.RootBarycentric, SimTri.DeformedTriangle);
 			}
 			#endif
 
