Fix noise caused by world unit scale for large distance screen-space subsurface scattering. It will not affect strata and is safe to check in.

diff --git a/Engine/Shaders/Private/SubsurfaceBurleyNormalized.ush b/Engine/Shaders/Private/SubsurfaceBurleyNormalized.ush
--- a/Engine/Shaders/Private/SubsurfaceBurleyNormalized.ush
+++ b/Engine/Shaders/Private/SubsurfaceBurleyNormalized.ush
@@ -1058,7 +1058,12 @@ float4 BurleyNormalizedSS(float2 BufferUV, uint2 GroupThreadID)
 
 			// Step 4: create the bilateral filtering weighted Distance between entry and exit.
 #if USE_BILATERAL_FILTERING
-			float DeltaDepth = (SampledRadianceAndDepth.w - DepthAtDiscCenter)*BURLEY_CM_2_MM;
+			// Bring DeltaDepth into the normalized kernal space.
+			// 
+			// Without the division of world unit scale, we add too much penalty to the sample weight when world unit scale is 
+			// large. E.g., when we have a 1 cm world unit scale (i.e., 1cm is regarded as 1mm), if we get 1mm depth difference, 
+			// it should be treated as 0.1mm instead of 1mm to reduce the weight contribution.
+			float DeltaDepth = (SampledRadianceAndDepth.w - DepthAtDiscCenter) * BURLEY_CM_2_MM / BurleyParameter.WorldUnitScale;
 			float RadiusSampledInMM = sqrt(BurleySampleInfo.RadiusInMM * BurleySampleInfo.RadiusInMM + DeltaDepth * DeltaDepth);
 
 #if RESAMPLE_PDF
