PR #9079: Fix shader issue with evident with speed tree foilage (Contributed by [REDACTED_PERSON])

diff --git a/Engine/Shaders/Private/MaterialTemplate.ush b/Engine/Shaders/Private/MaterialTemplate.ush
--- a/Engine/Shaders/Private/MaterialTemplate.ush
+++ b/Engine/Shaders/Private/MaterialTemplate.ush
@@ -1969,7 +1969,7 @@ float3 GetSpeedTreeVertexOffsetInner(FMaterialVertexParameters Parameters, int G
 		return float4(0,0,0);
 	#endif
 
-	#if USE_INSTANCING
+	#if USE_INSTANCING || USE_INSTANCE_CULLING
 		float3x3 LocalToWorld = LWCToFloat3x3(Parameters.InstanceLocalToWorld);
 		float3 LocalPosition = Parameters.InstanceLocalPosition;
 
