Fix shader compilation warning.

diff --git a/Engine/Shaders/Private/Strata/StrataTree.ush b/Engine/Shaders/Private/Strata/StrataTree.ush
--- a/Engine/Shaders/Private/Strata/StrataTree.ush
+++ b/Engine/Shaders/Private/Strata/StrataTree.ush
@@ -77,7 +77,7 @@ void UpdateSingleBSDFOperatorCoverageTransmittance(
 	// this feature is only needed for development/editor - we can compile it out for a shipping build (see r.CompileShadersForDevelopment cvar help)
 #if USE_DEVELOPMENT_SHADERS
 	// We progressively remove top layers but never the bottom one.
-	CurrentBSDF.Coverage = Op.LayerDepth < Settings.PeelLayersAboveDepth && !CurrentBSDF.bIsBottom ? 0.0f : CurrentBSDF.Coverage;
+	CurrentBSDF.Coverage = int(Op.LayerDepth) < Settings.PeelLayersAboveDepth && !CurrentBSDF.bIsBottom ? 0.0f : CurrentBSDF.Coverage;
 #endif
 
 	const float Roughness = StrataGetBSDFRoughness(CurrentBSDF);
