Fix hair strands material pass when strata is enabled.

diff --git a/Engine/Shaders/Private/HairStrands/HairStrandsMaterialCommon.ush b/Engine/Shaders/Private/HairStrands/HairStrandsMaterialCommon.ush
--- a/Engine/Shaders/Private/HairStrands/HairStrandsMaterialCommon.ush
+++ b/Engine/Shaders/Private/HairStrands/HairStrandsMaterialCommon.ush
@@ -199,14 +199,24 @@
 		CalcMaterialParametersEx(MaterialParameters, PixelMaterialInputs, SvPosition, LocalScreenPosition, bIsFrontFace, TranslatedWorldPosition, TranslatedWorldPosition);
 
 #if STRATA_ENABLED && MATERIAL_IS_STRATA && MATERIAL_SHADINGMODEL_HAIR
-		FStrataBSDF HairBSDF = PixelMaterialInputs.FrontMaterial.InlinedBSDF;
+		FStrataPixelHeader StrataPixelHeader		= InitialiseStrataPixelHeader();
+		StrataPixelHeader.StrataTree				= MaterialParameters.StrataTree;
+		StrataPixelHeader.BSDFCount					= MaterialParameters.StrataTree.BSDFCount;
+		StrataPixelHeader.SharedLocalBases			= MaterialParameters.SharedLocalBases;
+		StrataPixelHeader.IrradianceAO				= InitIrradianceAndOcclusion();
+		SetCastContactShadow(StrataPixelHeader, GetPrimitiveData(MaterialParameters).Flags & 0x100);
+		SetDynamicIndirectShadowCasterRepresentation(StrataPixelHeader, GetPrimitiveData(MaterialParameters).Flags & 0x80);
 
-		OutSample.BaseColor = HAIR_BASECOLOR(HairBSDF);
-		OutSample.Roughness = HAIR_ROUGHNESS(HairBSDF);
-		OutSample.Specular	= HAIR_SPECULAR(HairBSDF);
-		OutSample.Backlit	= HAIR_BACKLIT(HairBSDF);
-		OutSample.Emissive	= BSDF_GETEMISSIVE(HairBSDF);
-		float3 Tangent		= normalize(MaterialParameters.SharedLocalBases.Normals[BSDF_GETSHAREDLOCALBASISID(HairBSDF)]);
+		if (StrataPixelHeader.StrataTree.BSDFCount > 0)
+		{
+			FStrataBSDF HairBSDF = StrataPixelHeader.StrataTree.BSDFs[0];
+			OutSample.BaseColor = HAIR_BASECOLOR(HairBSDF);
+			OutSample.Roughness = HAIR_ROUGHNESS(HairBSDF);
+			OutSample.Specular	= HAIR_SPECULAR(HairBSDF);
+			OutSample.Backlit	= HAIR_BACKLIT(HairBSDF);
+			OutSample.Emissive	= BSDF_GETEMISSIVE(HairBSDF);
+			float3 Tangent		= normalize(MaterialParameters.SharedLocalBases.Normals[BSDF_GETSHAREDLOCALBASISID(HairBSDF)]);
+		}
 		
 	#if HAIRSTRANDS_HAS_NORMAL_CONNECTED
 	  #if MATERIAL_TANGENTSPACENORMAL 
