Implement manual depth tests for VT feedback UAV writes in materials that use PixelDepthOffset

diff --git a/Engine/Shaders/Private/BasePassPixelShader.usf b/Engine/Shaders/Private/BasePassPixelShader.usf
--- a/Engine/Shaders/Private/BasePassPixelShader.usf
+++ b/Engine/Shaders/Private/BasePassPixelShader.usf
@@ -173,6 +173,20 @@
 
 #if !FORWARD_SHADING
 
+bool ManualDepthTestEqual(float4 SvPosition, float PSDeviceZWithOffset)
+{
+	const float TexDeviceZWithOffset = OpaqueBasePass.ResolvedSceneDepthTexture.Load(int3(SvPosition.xy, 0)).r;
+
+	uint PSDeviceZWithOffsetUINT  = uint(PSDeviceZWithOffset * 16777215.0f + 0.5f);	// 16777215 = 2^24 - 1
+	uint TexDeviceZWithOffsetUINT = uint(TexDeviceZWithOffset * 16777215.0f + 0.5f);
+	
+	const bool DepthTest_Equal = OpaqueBasePass.Is24BitUnormDepthStencil ?
+		(PSDeviceZWithOffsetUINT == TexDeviceZWithOffsetUINT || PSDeviceZWithOffsetUINT == (TexDeviceZWithOffsetUINT - 1)) :	// 24 bit unorm, need to test a bit more, likely due to 24unorm conversion to float when loading form the texture.
+		(PSDeviceZWithOffset == TexDeviceZWithOffset); 																			// 32 bit float, 1 to 1 mapping.
+
+	return DepthTest_Equal;
+}
+
 void GetVolumeLightingNonDirectional(float4 AmbientLightingVector, float3 DiffuseColor, inout float3 InterpolatedLighting, out float4 VolumeLighting)
 {
 	// Normal is not taken into account with non directional lighting, and only the ambient term of the SH coefficients are needed
@@ -1601,17 +1615,7 @@ void FPixelShaderInOut_MainPS(
 
 #if OUTPUT_PIXEL_DEPTH_OFFSET
 	// When in deferred, opaque materials with pixel depth offset must execute a custom depth test in order to avoid conflicting UAV writes.
-	const float PSDeviceZWithOffset  = Out.Depth;
-	const float TexDeviceZWithOffset = OpaqueBasePass.ResolvedSceneDepthTexture.Load(int3(In.SvPosition.xy, 0)).r;
-
-	uint PSDeviceZWithOffsetUINT  = uint(PSDeviceZWithOffset * 16777215.0f + 0.5f);	// 16777215 = 2^24 - 1
-	uint TexDeviceZWithOffsetUINT = uint(TexDeviceZWithOffset * 16777215.0f + 0.5f);
-	
-	const bool DepthTest_Equal = OpaqueBasePass.Is24BitUnormDepthStencil ?
-		(PSDeviceZWithOffsetUINT == TexDeviceZWithOffsetUINT || PSDeviceZWithOffsetUINT == (TexDeviceZWithOffsetUINT - 1)) :	// 24 bit unorm, need to test a bit more, likely due to 24unorm conversion to float when loading form the texture.
-		(PSDeviceZWithOffset == TexDeviceZWithOffset); 																			// 32 bit float, 1 to 1 mapping.
-
-	if (DepthTest_Equal)
+	if (ManualDepthTestEqual(In.SvPosition, Out.Depth))
 #endif
 	{
 		float3 EmissiveLuminance = 0.0f;
@@ -2062,13 +2066,19 @@ void FPixelShaderInOut_MainPS(
 #endif
 
 #if NUM_VIRTUALTEXTURE_SAMPLES || LIGHTMAP_VT_ENABLED
-	FinalizeVirtualTextureFeedback(
-		MaterialParameters.VirtualTextureFeedback,
-		MaterialParameters.SvPosition,
-		Opacity * DBufferOpacity,
-		View.FrameNumber,
-		View.VTFeedbackBuffer
-	);
+#if OUTPUT_PIXEL_DEPTH_OFFSET
+	// When in deferred, opaque materials with pixel depth offset must execute a custom depth test in order to avoid conflicting UAV writes.
+	if (ManualDepthTestEqual(In.SvPosition, Out.Depth))
+#endif
+	{
+		FinalizeVirtualTextureFeedback(
+			MaterialParameters.VirtualTextureFeedback,
+			MaterialParameters.SvPosition,
+			Opacity * DBufferOpacity,
+			View.FrameNumber,
+			View.VTFeedbackBuffer
+		);
+	}
 #endif
 
 #if STRATA_MATERIAL_EXPORT_EXECUTED
