Landscape loads lightmap UV parameters from GPUScene when available

diff --git a/Engine/Shaders/Private/LandscapeVertexFactory.ush b/Engine/Shaders/Private/LandscapeVertexFactory.ush
--- a/Engine/Shaders/Private/LandscapeVertexFactory.ush
+++ b/Engine/Shaders/Private/LandscapeVertexFactory.ush
@@ -15,6 +15,7 @@
 
 // Set by FLightMapTexturePolicy
 #include "/Engine/Generated/UniformBuffers/PrecomputedLightingBuffer.ush"
+#include "LightmapData.ush"
 
 #if LANDSCAPE_XYOFFSET
 	#define XYOFFSET_SCALE (1.0f/256.f)
@@ -255,7 +256,8 @@ struct FVertexFactoryInterpolantsVSToPS
 	float2	LightMapCoordinate[2]	: TEXCOORD4;
 	float2	ShadowMapCoordinate		: TEXCOORD6;
 #else
-	float4	LightMapCoordinate		: TEXCOORD3;
+	float2	LightMapCoordinate		: TEXCOORD3;
+	float2	ShadowMapCoordinate		: TEXCOORD4;
 #endif
 #endif
 
@@ -554,10 +556,10 @@ FMaterialVertexParameters GetMaterialVertexParameters(FVertexFactoryInput Input,
 void GetLightMapCoordinates(FVertexFactoryInterpolantsVSToPS Interpolants, out float2 LightmapUV0, out float2 LightmapUV1, out uint LightmapDataIndex)
 {
 #if FEATURE_LEVEL == FEATURE_LEVEL_ES3_1
-	LightmapUV0 = Interpolants.LightMapCoordinate[0].xy;
-	LightmapUV1 = Interpolants.LightMapCoordinate[1].xy;
+	LightmapUV0 = Interpolants.LightMapCoordinate[0];
+	LightmapUV1 = Interpolants.LightMapCoordinate[1];
 #else
-	LightmapUV0 = Interpolants.LightMapCoordinate.xy * float2( 1, 0.5 );
+	LightmapUV0 = Interpolants.LightMapCoordinate * float2( 1, 0.5 );
 	LightmapUV1 = LightmapUV0 + float2( 0, 0.5 );
 #endif
 
@@ -570,11 +572,7 @@ void GetLightMapCoordinates(FVertexFactoryInterpolantsVSToPS Interpolants, out f
 
 void GetShadowMapCoordinate(FVertexFactoryInterpolantsVSToPS Interpolants, out float2 ShadowMapCoordinate, out uint LightmapDataIndex)
 {
-#if FEATURE_LEVEL == FEATURE_LEVEL_ES3_1
-	ShadowMapCoordinate = Interpolants.ShadowMapCoordinate.xy;
-#else
-	ShadowMapCoordinate = Interpolants.LightMapCoordinate.zw;
-#endif
+	ShadowMapCoordinate = Interpolants.ShadowMapCoordinate;
 	
 #if VF_USE_PRIMITIVE_SCENE_DATA
 	LightmapDataIndex = GetPrimitiveData(Interpolants.PrimitiveId).LightmapDataIndex;
@@ -759,23 +757,25 @@ FVertexFactoryInterpolantsVSToPS VertexFactoryGetInterpolantsVSToPS(FVertexFacto
 #endif // MOBILE_LANDSCAPE_MESH
 
 #if NEEDS_LIGHTMAP_COORDINATE
+	#if VF_USE_PRIMITIVE_SCENE_DATA
+		const uint LightmapDataIndex = GetPrimitiveData(Intermediates).LightmapDataIndex;
+	#else
+		const uint LightmapDataIndex = 0;
+	#endif
+	
 	#if FEATURE_LEVEL == FEATURE_LEVEL_ES3_1
-		Interpolants.LightMapCoordinate[0].xy = LandscapeTexCoords.LightMapCoordinate * PrecomputedLightingBuffer.LightMapCoordinateScaleBias.xy + PrecomputedLightingBuffer.LightMapCoordinateScaleBias.zw;
+		Interpolants.LightMapCoordinate[0] = LandscapeTexCoords.LightMapCoordinate * VF_GPUSCENE_GET_LIGHTMAP_UV_SCALE_BIAS(Input, LightmapDataIndex).xy + VF_GPUSCENE_GET_LIGHTMAP_UV_SCALE_BIAS(Input, LightmapDataIndex).zw;
 		Interpolants.LightMapCoordinate[0].y *= 0.5;
-		Interpolants.LightMapCoordinate[1].xy = Interpolants.LightMapCoordinate[0].xy;
+		Interpolants.LightMapCoordinate[1] = Interpolants.LightMapCoordinate[0].xy;
 		Interpolants.LightMapCoordinate[1].y += 0.5;
-		#if STATICLIGHTING_TEXTUREMASK
-			Interpolants.ShadowMapCoordinate.xy	= LandscapeTexCoords.LightMapCoordinate.xy * PrecomputedLightingBuffer.ShadowMapCoordinateScaleBias.xy + PrecomputedLightingBuffer.ShadowMapCoordinateScaleBias.zw;
-		#else
-			Interpolants.ShadowMapCoordinate.xy = 0;
-		#endif
 	#else
-		Interpolants.LightMapCoordinate.xy		= LandscapeTexCoords.LightMapCoordinate.xy * PrecomputedLightingBuffer.LightMapCoordinateScaleBias.xy + PrecomputedLightingBuffer.LightMapCoordinateScaleBias.zw;
-		#if STATICLIGHTING_TEXTUREMASK
-			Interpolants.LightMapCoordinate.zw	= LandscapeTexCoords.LightMapCoordinate.xy * PrecomputedLightingBuffer.ShadowMapCoordinateScaleBias.xy + PrecomputedLightingBuffer.ShadowMapCoordinateScaleBias.zw;
-		#else
-			Interpolants.LightMapCoordinate.zw = 0;
-		#endif
+		Interpolants.LightMapCoordinate = LandscapeTexCoords.LightMapCoordinate * VF_GPUSCENE_GET_LIGHTMAP_UV_SCALE_BIAS(Input, LightmapDataIndex).xy + VF_GPUSCENE_GET_LIGHTMAP_UV_SCALE_BIAS(Input, LightmapDataIndex).zw;
+	#endif
+	
+	#if STATICLIGHTING_TEXTUREMASK
+		Interpolants.ShadowMapCoordinate = LandscapeTexCoords.LightMapCoordinate * VF_GPUSCENE_GET_SHADOWMAP_UV_SCALE_BIAS(Input, LightmapDataIndex).xy + VF_GPUSCENE_GET_SHADOWMAP_UV_SCALE_BIAS(Input, LightmapDataIndex).zw;
+	#else
+		Interpolants.ShadowMapCoordinate = 0;
 	#endif
 #endif
 
