Fix reflection view when strata is enabled.

The material F0/DiffuseAlbedo were not properly overriden when that view mode was enabled.

diff --git a/Engine/Shaders/Private/Strata/Strata.ush b/Engine/Shaders/Private/Strata/Strata.ush
--- a/Engine/Shaders/Private/Strata/Strata.ush
+++ b/Engine/Shaders/Private/Strata/Strata.ush
@@ -2373,11 +2373,41 @@ float SanitizeRoughness(float Roughness)
 #endif
 }
 
+float SanitizeDiffuseAlbedo(float3 DiffuseAlbedo)
+{
+	#if USE_DEVELOPMENT_SHADERS
+	// This feature is only needed for development/editor - we can compile it out for a shipping build (see r.CompileShadersForDevelopment cvar help)
+	return saturate(DiffuseAlbedo) * View.DiffuseOverrideParameter.www + View.DiffuseOverrideParameter.xyz;
+	#else
+	return saturate(DiffuseAlbedo);
+	#endif
+}
+
+float3 SanitizeF0(float3 F0)
+{
+	#if USE_DEVELOPMENT_SHADERS
+	// This feature is only needed for development/editor - we can compile it out for a shipping build (see r.CompileShadersForDevelopment cvar help)
+	return saturate(F0) * View.SpecularOverrideParameter.w + View.SpecularOverrideParameter.xyz;
+	#else
+	return saturate(F0);
+	#endif
+}
+
+float SanitizeF0(float F0)
+{
+	#if USE_DEVELOPMENT_SHADERS
+	// This feature is only needed for development/editor - we can compile it out for a shipping build (see r.CompileShadersForDevelopment cvar help)
+	return saturate(F0) * View.SpecularOverrideParameter.w + View.SpecularOverrideParameter.x;
+	#else
+	return saturate(F0);
+	#endif
+}
+
 void SanitizeStrataSlab(inout FStrataBSDF BSDF)
 {
 	BSDF_SETEMISSIVE(BSDF, clamp(BSDF_GETEMISSIVE(BSDF), 0.0, Max10BitsFloat));
-	SLAB_DIFFUSEALBEDO(BSDF) = saturate(SLAB_DIFFUSEALBEDO(BSDF));
-	SLAB_F0(BSDF) = saturate(SLAB_F0(BSDF));
+	SLAB_DIFFUSEALBEDO(BSDF) = SanitizeDiffuseAlbedo(SLAB_DIFFUSEALBEDO(BSDF));
+	SLAB_F0(BSDF) = SanitizeF0(SLAB_F0(BSDF));
 	SLAB_F90(BSDF) = saturate(SLAB_F90(BSDF));
 	SLAB_ROUGHNESS(BSDF) = SanitizeRoughness(SLAB_ROUGHNESS(BSDF));
 	SLAB_ANISOTROPY(BSDF) = clamp(SLAB_ANISOTROPY(BSDF), -1.f, 1.f);
@@ -2403,17 +2433,17 @@ void SanitizeStrataSlab(inout FStrataBSDF BSDF)
 void SanitizeStrataHair(inout FStrataBSDF BSDF)
 {
 	BSDF_SETEMISSIVE(BSDF, clamp(BSDF_GETEMISSIVE(BSDF), 0.0, Max10BitsFloat));
-	HAIR_BASECOLOR(BSDF) = saturate(HAIR_BASECOLOR(BSDF));
+	HAIR_BASECOLOR(BSDF) = SanitizeDiffuseAlbedo(HAIR_BASECOLOR(BSDF));
 	HAIR_SCATTER(BSDF) = saturate(HAIR_SCATTER(BSDF));
 	HAIR_ROUGHNESS(BSDF) = SanitizeRoughness(HAIR_ROUGHNESS(BSDF));
-	HAIR_SPECULAR(BSDF) = saturate(HAIR_SPECULAR(BSDF));
+	HAIR_SPECULAR(BSDF) = SanitizeF0(HAIR_SPECULAR(BSDF));
 	HAIR_BACKLIT(BSDF) = saturate(HAIR_BACKLIT(BSDF));
 }
 
 void SanitizeStrataEye(inout FStrataBSDF BSDF)
 {
 	BSDF_SETEMISSIVE(BSDF, clamp(BSDF_GETEMISSIVE(BSDF), 0.0, Max10BitsFloat));
-	EYE_DIFFUSEALBEDO(BSDF) = saturate(EYE_DIFFUSEALBEDO(BSDF));
+	EYE_DIFFUSEALBEDO(BSDF) = SanitizeDiffuseAlbedo(EYE_DIFFUSEALBEDO(BSDF));
 	EYE_ROUGHNESS(BSDF) = SanitizeRoughness(EYE_ROUGHNESS(BSDF));
 	EYE_IRISDISTANCE(BSDF) = saturate(EYE_IRISDISTANCE(BSDF));
 	EYE_IRISMASK(BSDF) = saturate(EYE_IRISMASK(BSDF));
@@ -2424,9 +2454,9 @@ void SanitizeStrataEye(inout FStrataBSDF BSDF)
 void SanitizeStrataSingleLayerWater(inout FStrataBSDF BSDF)
 {
 	BSDF_SETEMISSIVE(BSDF, clamp(BSDF_GETEMISSIVE(BSDF), 0.0, Max10BitsFloat));
-	SLW_BASECOLOR(BSDF) = saturate(SLW_BASECOLOR(BSDF));
+	SLW_BASECOLOR(BSDF) = SanitizeDiffuseAlbedo(SLW_BASECOLOR(BSDF));
 	SLW_METALLIC(BSDF) = saturate(SLW_METALLIC(BSDF));
-	SLW_SPECULAR(BSDF) = saturate(SLW_SPECULAR(BSDF));
+	SLW_SPECULAR(BSDF) = SanitizeF0(SLW_SPECULAR(BSDF));
 	SLW_ROUGHNESS(BSDF) = SanitizeRoughness(SLW_ROUGHNESS(BSDF));
 	SLW_TOPMATERIALOPACITY(BSDF) = saturate(SLW_TOPMATERIALOPACITY(BSDF));
 #if STRATA_INLINE_SINGLELAYERWATER
