Fix shader typo.

diff --git a/Engine/Shaders/Private/Strata/Strata.ush b/Engine/Shaders/Private/Strata/Strata.ush
--- a/Engine/Shaders/Private/Strata/Strata.ush
+++ b/Engine/Shaders/Private/Strata/Strata.ush
@@ -2373,7 +2373,7 @@ float SanitizeRoughness(float Roughness)
 #endif
 }
 
-float SanitizeDiffuseAlbedo(float3 DiffuseAlbedo)
+float3 SanitizeDiffuseAlbedo(float3 DiffuseAlbedo)
 {
 	#if USE_DEVELOPMENT_SHADERS
 	// This feature is only needed for development/editor - we can compile it out for a shipping build (see r.CompileShadersForDevelopment cvar help)
