Fix Strata anisotropic lighting.

* Fix invalid NoL value (was using unclamped value)
* Fix incorrect derivation of alphaX/alphaY value.

diff --git a/Engine/Shaders/Private/Strata/StrataEvaluation.ush b/Engine/Shaders/Private/Strata/StrataEvaluation.ush
--- a/Engine/Shaders/Private/Strata/StrataEvaluation.ush
+++ b/Engine/Shaders/Private/Strata/StrataEvaluation.ush
@@ -365,12 +365,13 @@ FStrataEvaluateResult StrataEvaluateBSDFCommon(FStrataBSDFContext BSDFContext, F
 			{
 				// Generalized microfacet specular
 				{
-					float2 Alpha = 0;
-					GetAnisotropicRoughness(SafeRoughness, SLAB_ANISOTROPY(BSDFContext.BSDF), Alpha.x, Alpha.y);
+					float  Alpha  = Square(SafeRoughness);
+					float2 AlphaXY = 0;
+					GetAnisotropicRoughness(Alpha, SLAB_ANISOTROPY(BSDFContext.BSDF), AlphaXY.x, AlphaXY.y);
 
-					D   = D_GGXaniso(Alpha.x, Alpha.y, BSDFContext.Context.NoH, BSDFContext.Context.XoH, BSDFContext.Context.YoH);
-					Vis = Vis_SmithJointAniso(Alpha.x, Alpha.y, BSDFContext.Context.NoV, BSDFContext.Context.NoL, BSDFContext.Context.XoV, BSDFContext.Context.XoL, BSDFContext.Context.YoV, BSDFContext.Context.YoL);
-					const float H_PDF = VisibleGGXPDF_aniso(BSDFContext.TangentV, BSDFContext.TangentH, Alpha);
+					D   = D_GGXaniso(AlphaXY.x, AlphaXY.y, BSDFContext.Context.NoH, BSDFContext.Context.XoH, BSDFContext.Context.YoH);
+					Vis = Vis_SmithJointAniso(AlphaXY.x, AlphaXY.y, BSDFContext.Context.NoV, BSDFContext.SatNoL, BSDFContext.Context.XoV, BSDFContext.Context.XoL, BSDFContext.Context.YoV, BSDFContext.Context.YoL);
+					const float H_PDF = VisibleGGXPDF_aniso(BSDFContext.TangentV, BSDFContext.TangentH, AlphaXY);
 					PDF = RayPDFToReflectionRayPDF(BSDFContext.Context.VoH, H_PDF);
 				}
 
