GPULM: Improve backface detection consistency with translucent and two sided materials

diff --git a/Engine/Plugins/Experimental/GPULightmass/Shaders/Private/LightmapPathTracing.usf b/Engine/Plugins/Experimental/GPULightmass/Shaders/Private/LightmapPathTracing.usf
--- a/Engine/Plugins/Experimental/GPULightmass/Shaders/Private/LightmapPathTracing.usf
+++ b/Engine/Plugins/Experimental/GPULightmass/Shaders/Private/LightmapPathTracing.usf
@@ -419,7 +419,7 @@ void PathTracingKernel(
 			break;
 		}
 		
-		if (!Payload.IsFrontFace() && !Payload.IsMaterialTransmissive())
+		if (Payload.Opacity < 0)
 		{
 			bIsValidSample = false;
 			break;
diff --git a/Engine/Shaders/Private/PathTracing/PathTracingMaterialHitShader.usf b/Engine/Shaders/Private/PathTracing/PathTracingMaterialHitShader.usf
--- a/Engine/Shaders/Private/PathTracing/PathTracingMaterialHitShader.usf
+++ b/Engine/Shaders/Private/PathTracing/PathTracingMaterialHitShader.usf
@@ -343,7 +343,12 @@ RAY_TRACING_ENTRY_CLOSEST_HIT(PathTracingMaterialCHS,
 	else
 #endif
 	{
+		// Redirect translucent blending mode to thin translucent
+#if MATERIALBLENDING_TRANSLUCENT
+		Payload.ShadingModelID = SHADINGMODELID_THIN_TRANSLUCENT;
+#else
 		Payload.ShadingModelID = SHADINGMODELID_DEFAULT_LIT;
+#endif
 	}
 	Payload.TranslatedWorldPos = LWCToFloat(LWCAdd(MaterialParameters.AbsoluteWorldPosition, ResolvedView.PreViewTranslation));
 	Payload.WorldNormal = normalize(MaterialParameters.WorldNormal);
@@ -357,6 +362,14 @@ RAY_TRACING_ENTRY_CLOSEST_HIT(PathTracingMaterialCHS,
 	Payload.Opacity = GetMaterialOpacity(PixelMaterialInputs);
 	Payload.Ior = 1;
 
+#if !MATERIAL_TWOSIDED
+	if (!Payload.IsFrontFace() && !Payload.IsMaterialTransmissive())
+	{
+		// Overloaded semantics for opacity
+		Payload.Opacity = -1;
+	}
+#endif
+
 #endif
 
 	PackedPayload = PackPathTracingPayload(Payload);
