Fix hair strands material rendering with Strata.

diff --git a/Engine/Shaders/Private/HairStrands/HairStrandsMaterialCommon.ush b/Engine/Shaders/Private/HairStrands/HairStrandsMaterialCommon.ush
--- a/Engine/Shaders/Private/HairStrands/HairStrandsMaterialCommon.ush
+++ b/Engine/Shaders/Private/HairStrands/HairStrandsMaterialCommon.ush
@@ -207,15 +207,28 @@ FEvaluateOutput Evaluate(
 		SetCastContactShadow(StrataPixelHeader, GetPrimitiveData(MaterialParameters).Flags & 0x100);
 		SetDynamicIndirectShadowCasterRepresentation(StrataPixelHeader, GetPrimitiveData(MaterialParameters).Flags & 0x80);
 
-		if (StrataPixelHeader.StrataTree.BSDFCount > 0)
+		float3 Tangent = float3(0,1,0);
+		//if (StrataPixelHeader.StrataTree.BSDFCount > 0)
 		{
+			// Using the tree data, instread of the inline BSDF does not seem to work. It is unclear why. It is not very 
+			// important as we expect a single hair BSDF.
+		#if 0
+			// Update tree (coverage/transmittance/luminace weights)
+			FStrataAddressing NullStrataAddressing = (FStrataAddressing)0;	// Fake unused in StrataCreateBSDFContext when using Forward inline shading
+			const float3 V = MaterialParameters.CameraVector;
+			const FStrataIntegrationSettings Settings = InitStrataIntegrationSettings(false /*bForceFullyRough*/, false /*bRoughDiffuse*/, -1 /*PeelLayersAboveDepth*/);
+			StrataUpdateTree(NullStrataAddressing, StrataPixelHeader, StrataPixelHeader.StrataTree, V, Settings);
 			FStrataBSDF HairBSDF = StrataPixelHeader.StrataTree.BSDFs[0];
+		#else
+			FStrataBSDF HairBSDF = PixelMaterialInputs.FrontMaterial.InlinedBSDF;
+		#endif
+
 			OutSample.BaseColor = HAIR_BASECOLOR(HairBSDF);
 			OutSample.Roughness = HAIR_ROUGHNESS(HairBSDF);
 			OutSample.Specular	= HAIR_SPECULAR(HairBSDF);
 			OutSample.Backlit	= HAIR_BACKLIT(HairBSDF);
 			OutSample.Emissive	= BSDF_GETEMISSIVE(HairBSDF);
-			float3 Tangent		= normalize(MaterialParameters.SharedLocalBases.Normals[BSDF_GETSHAREDLOCALBASISID(HairBSDF)]);
+			Tangent				= normalize(MaterialParameters.SharedLocalBases.Normals[BSDF_GETSHAREDLOCALBASISID(HairBSDF)]);
 		}
 		
 	#if HAIRSTRANDS_HAS_NORMAL_CONNECTED
