PostProcessCombineLuts: Moved WiteBalance to a commonly accessable header. There are other shaders that have duplicate version of the same function.

diff --git a/Engine/Shaders/Private/PostProcessCombineLUTs.usf b/Engine/Shaders/Private/PostProcessCombineLUTs.usf
--- a/Engine/Shaders/Private/PostProcessCombineLUTs.usf
+++ b/Engine/Shaders/Private/PostProcessCombineLUTs.usf
@@ -39,33 +39,6 @@ static const float LUTSize = 32;
 float WhiteTemp;
 float WhiteTint;
 
-float3 WhiteBalance( float3 LinearColor, const float3x3 WCS_2_XYZ, const float3x3 XYZ_2_WCS )
-{
-	float2 SrcWhiteDaylight = D_IlluminantChromaticity( WhiteTemp );
-	float2 SrcWhitePlankian = PlanckianLocusChromaticity( WhiteTemp );
-
-	float2 SrcWhite = WhiteTemp < 4000 ? SrcWhitePlankian : SrcWhiteDaylight;
-	float2 D65White = float2( 0.31270,  0.32900 );
-
-	{
-		// Offset along isotherm
-		float2 Isothermal = PlanckianIsothermal( WhiteTemp, WhiteTint ) - SrcWhitePlankian;
-		SrcWhite += Isothermal;
-	}
-
-	if (!bIsTemperatureWhiteBalance)
-	{
-		float2 Temp = SrcWhite;
-		SrcWhite = D65White;
-		D65White = Temp;
-	}
-
-	float3x3 WhiteBalanceMat = ChromaticAdaptation( SrcWhite, D65White );
-	WhiteBalanceMat = mul( XYZ_2_WCS, mul( WhiteBalanceMat, WCS_2_XYZ ) );
-
-	return mul( WhiteBalanceMat, LinearColor );
-}
-
 float4 ColorSaturation;
 float4 ColorContrast;
 float4 ColorGamma;
@@ -221,7 +194,7 @@ float4 CombineLUTsCommon(float2 InUV, uint InLayerIndex)
 	}
 
 #if !SKIP_TEMPERATURE
-	float3 BalancedColor = WhiteBalance( LinearColor, (float3x3)WorkingColorSpace.ToXYZ, (float3x3)WorkingColorSpace.FromXYZ );
+	float3 BalancedColor = WhiteBalance(LinearColor, WhiteTemp, WhiteTint, bIsTemperatureWhiteBalance, (float3x3)WorkingColorSpace.ToXYZ, (float3x3)WorkingColorSpace.FromXYZ);
 #else
 	float3 BalancedColor = LinearColor;
 #endif
diff --git a/Engine/Shaders/Private/TonemapCommon.ush b/Engine/Shaders/Private/TonemapCommon.ush
--- a/Engine/Shaders/Private/TonemapCommon.ush
+++ b/Engine/Shaders/Private/TonemapCommon.ush
@@ -453,3 +453,31 @@ float2 PlanckianIsothermal( float Temp, float Tint )
 
 	return float2(x,y);
 }
+
+float3 WhiteBalance(float3 LinearColor, float WhiteTemp, float WhiteTint, bool bIsTemperatureWhiteBalance, const float3x3 WCS_2_XYZ, const float3x3 XYZ_2_WCS)
+{
+	float2 SrcWhiteDaylight = D_IlluminantChromaticity(WhiteTemp);
+	float2 SrcWhitePlankian = PlanckianLocusChromaticity(WhiteTemp);
+
+	float2 SrcWhite = WhiteTemp < 4000 ? SrcWhitePlankian : SrcWhiteDaylight;
+	float2 D65White = float2(0.31270, 0.32900);
+
+	{
+		// Offset along isotherm
+		float2 Isothermal = PlanckianIsothermal(WhiteTemp, WhiteTint) - SrcWhitePlankian;
+		SrcWhite += Isothermal;
+	}
+
+	if (!bIsTemperatureWhiteBalance)
+	{
+		float2 Temp = SrcWhite;
+		SrcWhite = D65White;
+		D65White = Temp;
+	}
+
+	float3x3 WhiteBalanceMat = ChromaticAdaptation(SrcWhite, D65White);
+	WhiteBalanceMat = mul( XYZ_2_WCS, mul( WhiteBalanceMat, WCS_2_XYZ ) );
+
+	return mul(WhiteBalanceMat, LinearColor);
+}
+
