Fix slate never setting the evaluate WPO primitive flag.
This prevented slate world position effects from working.

diff --git a/Engine/Shaders/Private/SlateVertexShader.usf b/Engine/Shaders/Private/SlateVertexShader.usf
--- a/Engine/Shaders/Private/SlateVertexShader.usf
+++ b/Engine/Shaders/Private/SlateVertexShader.usf
@@ -42,6 +42,9 @@ VertexToPixelInterpolants Main(
 
 	VertexParameters.VertexColor = VOut.Color;
 
+	// Init primitive flags to a good default.
+	VertexParameters.SceneData.Primitive.Flags = PRIMITIVE_SCENE_DATA_FLAG_EVALUATE_WORLD_POSITION_OFFSET;
+
 #if NUM_MATERIAL_TEXCOORDS_VERTEX
 
 #if !USE_SLATE_INSTANCING
