Fixed: Certain Android devices crash on boot. Rearranged code in OuputGamutMappingMatrix() to avoid shader compiler crash

diff --git a/Engine/Shaders/Private/TonemapCommon.ush b/Engine/Shaders/Private/TonemapCommon.ush
--- a/Engine/Shaders/Private/TonemapCommon.ush
+++ b/Engine/Shaders/Private/TonemapCommon.ush
@@ -326,16 +326,17 @@ float3x3 OuputGamutMappingMatrix( uint OutputGamut )
 	// 3 = ACES AP0 - D60
 	// 4 = ACES AP1 - D60
 
+	float3x3 GamutMappingMatrix = AP1_2_sRGB;
 	if( OutputGamut == TONEMAPPER_GAMUT_DCIP3_D65)
-		return AP1_2_DCI_D65;
+		GamutMappingMatrix = AP1_2_DCI_D65;
 	else if( OutputGamut == TONEMAPPER_GAMUT_Rec2020_D65)
-		return AP1_2_Rec2020;
+		GamutMappingMatrix = AP1_2_Rec2020;
 	else if( OutputGamut == TONEMAPPER_GAMUT_ACES_D60)
-		return AP1_2_AP0_MAT;
+		GamutMappingMatrix = AP1_2_AP0_MAT;
 	else if( OutputGamut == TONEMAPPER_GAMUT_ACEScg_D60)
-		return GamutMappingIdentityMatrix;
-	else
-		return AP1_2_sRGB;
+		GamutMappingMatrix = GamutMappingIdentityMatrix;
+
+	return GamutMappingMatrix; 
 }
 
 float3x3 OuputInverseGamutMappingMatrix( uint OutputGamut )
