Fixed: Precision issues on mobile with very bright lights

diff --git a/Engine/Shaders/Private/DeferredLightingCommon.ush b/Engine/Shaders/Private/DeferredLightingCommon.ush
--- a/Engine/Shaders/Private/DeferredLightingCommon.ush
+++ b/Engine/Shaders/Private/DeferredLightingCommon.ush
@@ -36,7 +36,8 @@ struct FDeferredLightData
 {
 	float3 TranslatedWorldPosition;
 	half  InvRadius;
-	half3 Color;
+	/** Premultiplied with light brightness, can be huge. Do not use half precision here */
+	float3 Color;
 	half  FalloffExponent;
 	float3 Direction;
 	float3 Tangent;
