Followup for fix of slate never setting the evaluate WPO primitive flag.
Added GetDefaultPrimitiveSceneDataFlags() to support this instead of setting single flag.

diff --git a/Engine/Shaders/Private/SceneData.ush b/Engine/Shaders/Private/SceneData.ush
--- a/Engine/Shaders/Private/SceneData.ush
+++ b/Engine/Shaders/Private/SceneData.ush
@@ -70,6 +70,20 @@
 #define PRIMITIVE_SCENE_DATA_FLAG_EVALUATE_WORLD_POSITION_OFFSET		0x1000000
 #define PRIMITIVE_SCENE_DATA_FLAG_CAMERA_DISTANCE_CULL					0x2000000
 
+uint GetDefaultPrimitiveSceneDataFlags()
+{
+	// Should match FPrimitiveUniformShaderParametersBuilder::Default()
+	return
+		PRIMITIVE_SCENE_DATA_FLAG_CAST_SHADOWS |
+		PRIMITIVE_SCENE_DATA_FLAG_HAS_CAST_CONTACT_SHADOW |
+		PRIMITIVE_SCENE_DATA_FLAG_EVALUATE_WORLD_POSITION_OFFSET |
+		PRIMITIVE_SCENE_DATA_FLAG_VISIBLE_IN_GAME |
+		PRIMITIVE_SCENE_DATA_FLAG_VISIBLE_IN_EDITOR |
+		PRIMITIVE_SCENE_DATA_FLAG_VISIBLE_IN_REFLECTION_CAPTURES |
+		PRIMITIVE_SCENE_DATA_FLAG_VISIBLE_IN_REAL_TIME_SKY_CAPTURES |
+		PRIMITIVE_SCENE_DATA_FLAG_VISIBLE_IN_RAY_TRACING;
+}
+
 // GPUCULL_TODO: Eventually we need to remove this workaround
 #define VF_TREAT_INSTANCE_ID_OFFSET_AS_PRIMITIVE_ID_FLAG (1U << 31U)
 
diff --git a/Engine/Shaders/Private/SlateVertexShader.usf b/Engine/Shaders/Private/SlateVertexShader.usf
--- a/Engine/Shaders/Private/SlateVertexShader.usf
+++ b/Engine/Shaders/Private/SlateVertexShader.usf
@@ -43,7 +43,7 @@ VertexToPixelInterpolants Main(
 	VertexParameters.VertexColor = VOut.Color;
 
 	// Init primitive flags to a good default.
-	VertexParameters.SceneData.Primitive.Flags = PRIMITIVE_SCENE_DATA_FLAG_EVALUATE_WORLD_POSITION_OFFSET;
+	VertexParameters.SceneData.Primitive.Flags = GetDefaultPrimitiveSceneDataFlags();
 
 #if NUM_MATERIAL_TEXCOORDS_VERTEX
 
