Fix sceneData uncompressed transform on Editor/SM5

diff --git a/Engine/Shaders/Private/SceneData.ush b/Engine/Shaders/Private/SceneData.ush
--- a/Engine/Shaders/Private/SceneData.ush
+++ b/Engine/Shaders/Private/SceneData.ush
@@ -996,6 +996,21 @@ FInstanceSceneData GetInstanceSceneData(uint InstanceId, uint SOAStride, bool bC
 			float3 TilePosition 				= GetPrimitiveData(InstanceData.PrimitiveId).TilePosition;
 			InstanceData.PrevLocalToWorld		= MakeLWCMatrix(TilePosition, PrevLocalToRelativeWorld);
 		}
+		else
+		{
+		#if PLATFORM_ALLOW_SCENE_DATA_COMPRESSED_TRANSFORMS
+			// TODO: Temporary PrevVelocityHack
+			uint4  PrevRotationScale		= asuint(LoadInstanceSceneDataElement(3 * SOAStride + InstanceId));
+			float3 PrevTranslation			= LoadInstanceSceneDataElement(4 * SOAStride + InstanceId).xyz;
+			float3 PrevScale = 0;
+			float4x4 PrevLocalToRelativeWorld	= DecodeTransform(PrevRotationScale, PrevTranslation, PrevScale);
+
+			float3 TilePosition 			= GetPrimitiveData(InstanceData.PrimitiveId).TilePosition;
+			InstanceData.PrevLocalToWorld	= MakeLWCMatrix(TilePosition, PrevLocalToRelativeWorld);
+		#else
+			InstanceData.PrevLocalToWorld	= InstanceData.LocalToWorld;
+		#endif
+		}
 
 	#if 1 //NEEDS_LIGHTMAP_COORDINATE
 		BRANCH
