Fix hair strands crash with path tracing, when hair length scale is lower than 1.

diff --git a/Engine/Shaders/Private/HairStrands/HairStrandsRaytracingGeometry.usf b/Engine/Shaders/Private/HairStrands/HairStrandsRaytracingGeometry.usf
--- a/Engine/Shaders/Private/HairStrands/HairStrandsRaytracingGeometry.usf
+++ b/Engine/Shaders/Private/HairStrands/HairStrandsRaytracingGeometry.usf
@@ -157,6 +157,10 @@ float GetAABBHalfAreaForSegment(float4 A, float4 B)
 	return dot(BoxDiag.xyz, BoxDiag.yzx);
 }
 
+bool IsValidPosition(FHairControlPoint P)
+{
+	return all(!isinf(P.Position)) && all(!isnan(P.Position)) && all(P.Position <= INFINITE_FLOAT);
+}
 
 [numthreads(GROUP_SIZE, 1, 1)]
 void MainCS(uint2 DispatchThreadId : SV_DispatchThreadID)
@@ -194,7 +198,7 @@ void MainCS(uint2 DispatchThreadId : SV_DispatchThreadID)
 
 	FHairControlPoint P0 = GetControlPoint(SourceIndex0);
 #if PERMUTATION_PROCEDURAL_PRIMITIVE 
-	if (P0.Type == HAIR_CONTROLPOINT_END)
+	if (P0.Type == HAIR_CONTROLPOINT_END || !IsValidPosition(P0))
 	{
 		// this vertex doesn't map to a valid segment
 		SetInvalidSegment(VertexIndex0);
@@ -203,43 +207,51 @@ void MainCS(uint2 DispatchThreadId : SV_DispatchThreadID)
 	{
 		// fetch the other point and create the aabb around the line
 		FHairControlPoint P1 = GetControlPoint(SourceIndex1);
-		float4 A = float4(P0.Position, P0.WorldRadius * LodRadiusScale0);
-		float4 B = float4(P1.Position, P1.WorldRadius * LodRadiusScale1);
-		if ((A.w == 0 && B.w == 0) || length2(B.xyz - A.xyz) == 0 || any(!IsFinite(A)) || any(!IsFinite(B)))
+		if (!IsValidPosition(P1))
 		{
-			// this segment is degenerate, don't bother adding it
-			SetInvalidSegment(VertexIndex0);
+			SetInvalidSegment(VertexIndex1);
 		}
 		else
 		{
-			// measure aabb surface area prior to splitting
-			float FullArea = GetAABBHalfAreaForSegment(A, B);
-
-			float4 SplitA = A;
-			float SumArea = 0.0;
-			for (int i = 0; i < RaytracingProceduralSplits; i++)
+			float4 A = float4(P0.Position, P0.WorldRadius * LodRadiusScale0);
+			float4 B = float4(P1.Position, P1.WorldRadius * LodRadiusScale1);
+			if ((A.w == 0 && B.w == 0) || length2(B.xyz - A.xyz) == 0 || any(!IsFinite(A)) || any(!IsFinite(B)))
 			{
-				float4 SplitB = lerp(A, B, float(i + 1) / float(RaytracingProceduralSplits));
-				float3 Lo = min(SplitA.xyz - SplitA.w, SplitB.xyz - SplitB.w);
-				float3 Hi = max(SplitA.xyz + SplitA.w, SplitB.xyz + SplitB.w);
-				SumArea += GetAABBHalfAreaForSegment(SplitA, SplitB);
-				SetAABB(RaytracingProceduralSplits * VertexIndex0 + i, Lo, Hi);
-				SplitA = SplitB;
+				// this segment is degenerate, don't bother adding it
+				SetInvalidSegment(VertexIndex0);
 			}
-#if 1
-			if (FullArea <= 2 * SumArea)
+			else
 			{
-				// not worth splitting? use full bbox and make other splits degenerate
-				float3 Lo = min(A.xyz - A.w, B.xyz - B.w);
-				float3 Hi = max(A.xyz + A.w, B.xyz + B.w);
-				SetAABB(RaytracingProceduralSplits * VertexIndex0 + 0, Lo, Hi);
-				for (int i = 1; i < RaytracingProceduralSplits; i++)
+				// measure aabb surface area prior to splitting
+				float FullArea = GetAABBHalfAreaForSegment(A, B);
+
+				float4 SplitA = A;
+				float SumArea = 0.0;
+				for (int i = 0; i < RaytracingProceduralSplits; i++)
 				{
-					const float NaN = asfloat(0x7FC00000);
-					SetAABB(RaytracingProceduralSplits * VertexIndex0 + i, NaN, NaN);
+					float4 SplitB = lerp(A, B, float(i + 1) / float(RaytracingProceduralSplits));
+					float3 Lo = min(SplitA.xyz - SplitA.w, SplitB.xyz - SplitB.w);
+					float3 Hi = max(SplitA.xyz + SplitA.w, SplitB.xyz + SplitB.w);
+					SumArea += GetAABBHalfAreaForSegment(SplitA, SplitB);
+					SetAABB(RaytracingProceduralSplits * VertexIndex0 + i, Lo, Hi);
+					SplitA = SplitB;
 				}
+
+				#if 1
+				if (FullArea <= 2 * SumArea)
+				{
+					// not worth splitting? use full bbox and make other splits degenerate
+					float3 Lo = min(A.xyz - A.w, B.xyz - B.w);
+					float3 Hi = max(A.xyz + A.w, B.xyz + B.w);
+					SetAABB(RaytracingProceduralSplits * VertexIndex0 + 0, Lo, Hi);
+					for (int i = 1; i < RaytracingProceduralSplits; i++)
+					{
+						const float NaN = asfloat(0x7FC00000);
+						SetAABB(RaytracingProceduralSplits * VertexIndex0 + i, NaN, NaN);
+					}
+				}
+				#endif
 			}
-#endif
 		}
 	}
 #else
@@ -260,7 +272,7 @@ void MainCS(uint2 DispatchThreadId : SV_DispatchThreadID)
 		OutputPositionBuffer[BaseOutIndex + 2] = float4(Px1, 1);
 		OutputPositionBuffer[BaseOutIndex + 3] = float4(Py1, 1);
 	}
-	if (P0.Type == HAIR_CONTROLPOINT_END)
+	if (P0.Type == HAIR_CONTROLPOINT_END || !IsValidPosition(P0))
 	{
 		// last vertex of a curve, we could create a cap, but this would complicate the logic in the vertex factory
 		// ignore this for now on the assumption that most hairs will have a small tip width
