Support VF_SUPPORTS_RAYTRACING_PREPARE_MATERIAL_PIXEL_PARAMETERS in PathTracingMaterialAHS.

diff --git a/Engine/Shaders/Private/PathTracing/PathTracingMaterialHitShader.usf b/Engine/Shaders/Private/PathTracing/PathTracingMaterialHitShader.usf
--- a/Engine/Shaders/Private/PathTracing/PathTracingMaterialHitShader.usf
+++ b/Engine/Shaders/Private/PathTracing/PathTracingMaterialHitShader.usf
@@ -56,6 +56,7 @@ RAY_TRACING_ENTRY_CLOSEST_HIT(PathTracingMaterialCHS,
 
 	FMaterialPixelParameters MaterialParameters = GetMaterialPixelParameters(Interpolants, SvPosition);
 #endif
+
 	FPixelMaterialInputs PixelMaterialInputs;
 
 	{
@@ -422,12 +423,21 @@ RAY_TRACING_ENTRY_ANY_HIT(PathTracingMaterialAHS,
 
 	ResolvedView = ResolveView();
 
-	FVertexFactoryInterpolantsVSToPS Interpolants;
 	float4 SvPosition = float4(PixelCoord.xy, 0.0, 1.0);
 
-	CalcInterpolants((FRayCone)0, Attributes, Interpolants);
+#if VF_SUPPORTS_RAYTRACING_PREPARE_MATERIAL_PIXEL_PARAMETERS
+	// this is a newer codepath that is both more flexible and allows for more direct calculation compared to the other codepath
+	// TODO: implement such a method for all vertex factories
+	float3 GeoNormal = 0;
+	FMaterialPixelParameters MaterialParameters = GetMaterialPixelParameters(WorldRayOrigin(), WorldRayDirection(), RayTCurrent(), PrimitiveIndex(), Attributes, HitKind(), GeoNormal);
+#else
+	FVertexFactoryInterpolantsVSToPS Interpolants;
+	float3 GeoNormal = 0;
+	CalcInterpolants((FRayCone)0, Attributes, Interpolants, GeoNormal);
 
 	FMaterialPixelParameters MaterialParameters = GetMaterialPixelParameters(Interpolants, SvPosition);
+#endif
+
 	FPixelMaterialInputs PixelMaterialInputs;
 
 	const bool bIsFrontFace = HitKind() == HIT_KIND_TRIANGLE_FRONT_FACE;
