Add debug info for instance view classfiication.

diff --git a/Engine/Plugins/Runtime/HairStrands/Source/HairStrandsCore/Private/GroomDebug.cpp b/Engine/Plugins/Runtime/HairStrands/Source/HairStrandsCore/Private/GroomDebug.cpp
--- a/Engine/Plugins/Runtime/HairStrands/Source/HairStrandsCore/Private/GroomDebug.cpp
+++ b/Engine/Plugins/Runtime/HairStrands/Source/HairStrandsCore/Private/GroomDebug.cpp
@@ -836,6 +836,9 @@ class FHairDebugPrintInstanceCS : public FGlobalShader
 		SHADER_PARAMETER(uint32, InstanceCount)
 		SHADER_PARAMETER(uint32, NameInfoCount)
 		SHADER_PARAMETER(uint32, NameCharacterCount)
+		SHADER_PARAMETER(uint32, InstanceCount_StrandsPrimaryView)
+		SHADER_PARAMETER(uint32, InstanceCount_StrandsShadowView)
+		SHADER_PARAMETER(uint32, InstanceCount_CardsOrMeshes)
 		SHADER_PARAMETER_RDG_BUFFER_SRV(StructuredBuffer<uint2>, NameInfos)
 		SHADER_PARAMETER_RDG_BUFFER_SRV(Buffer<uint8>, Names)
 		SHADER_PARAMETER_RDG_BUFFER_SRV(StructuredBuffer<uint4>, Infos)
@@ -868,7 +871,8 @@ static void AddHairDebugPrintInstancePass(
 	FRDGBuilder& GraphBuilder, 
 	FGlobalShaderMap* ShaderMap,
 	const FShaderPrintData* ShaderPrintData,
-	const FHairStrandsInstances& Instances)
+	const FHairStrandsInstances& Instances,
+	const FUintVector4& InstanceCountPerType)
 {
 	const uint32 InstanceCount = Instances.Num();
 
@@ -891,6 +895,7 @@ static void AddHairDebugPrintInstancePass(
 	{
 		FUintVector4 Data0;
 		FUintVector4 Data1;
+		FUintVector4 Data2;
 	};
 	TArray<FInstanceInfos> Infos;
 	Infos.Reserve(InstanceCount);
@@ -922,6 +927,7 @@ static void AddHairDebugPrintInstancePass(
 
 		FUintVector4 Data0 = { 0,0,0,0 };
 		FUintVector4 Data1 = { 0,0,0,0 };
+		FUintVector4 Data2 = { 0,0,0,0 };
 		Data0.X =
 			((Instance->Debug.GroupIndex & 0xFF)) |
 			((Instance->Debug.GroupCount & 0xFF) << 8) |
@@ -950,6 +956,13 @@ static void AddHairDebugPrintInstancePass(
 					Data1.Z = Instance->Strands.RestRootResource->BulkData.RootCount;
 					Data1.W = Instance->Strands.RestRootResource->BulkData.PointCount;
 				}
+
+				{
+					Data2.X = InstanceIndex < InstanceCountPerType[HairInstanceCount_StrandsPrimaryView] ? 3 : 2; // Visible in primary & shadow (3) or only in shadow (2)
+					Data2.Y = 0;
+					Data2.Z = 0;
+					Data2.W = 0;
+				}
 			}
 			break;
 		case EHairGeometryType::Cards:
@@ -957,6 +970,7 @@ static void AddHairDebugPrintInstancePass(
 			{
 				Data0.Z = Instance->Cards.LODs[IntLODIndex].Guides.IsValid() ? Instance->Cards.LODs[IntLODIndex].Guides.Data->GetNumCurves() : 0;
 				Data0.W = Instance->Cards.LODs[IntLODIndex].Data->GetNumVertices();
+				Data2.X = 3; // Visible in primary & shadow
 			}
 			break;
 		case EHairGeometryType::Meshes:
@@ -964,10 +978,11 @@ static void AddHairDebugPrintInstancePass(
 			{
 				Data0.Z = 0;
 				Data0.W = Instance->Meshes.LODs[IntLODIndex].Data->GetNumVertices();
+				Data2.X = 3; // Visible in primary & shadow
 			}
 			break;
 		}
-		Infos.Add({Data0, Data1});
+		Infos.Add({Data0, Data1, Data2 });
 	}
 
 	if (NameInfos.IsEmpty())
@@ -992,6 +1007,9 @@ static void AddHairDebugPrintInstancePass(
 	{
 		FHairDebugPrintInstanceCS::FParameters* Parameters = GraphBuilder.AllocParameters<FHairDebugPrintInstanceCS::FParameters>();
 		Parameters->InstanceCount = InstanceCount;
+		Parameters->InstanceCount_StrandsPrimaryView = InstanceCountPerType[HairInstanceCount_StrandsPrimaryView];
+		Parameters->InstanceCount_StrandsShadowView = InstanceCountPerType[HairInstanceCount_StrandsShadowView];
+		Parameters->InstanceCount_CardsOrMeshes = InstanceCountPerType[HairInstanceCount_CardsOrMeshes];
 		Parameters->NameInfoCount = NameInfos.Num();
 		Parameters->NameCharacterCount = Names.Num();
 		Parameters->Names = GraphBuilder.CreateSRV(NameBuffer, PF_R8_UINT);
@@ -1043,6 +1061,7 @@ void RunHairStrandsDebug(
 	FGlobalShaderMap* ShaderMap,
 	const FSceneView& View,
 	const FHairStrandsInstances& Instances,
+	const FUintVector4& InstanceCountPerType,
 	const FGPUSkinCache* SkinCache,
 	const FShaderPrintData* ShaderPrintData,
 	FRDGTextureRef SceneColorTexture,
@@ -1054,7 +1073,7 @@ void RunHairStrandsDebug(
 
 	if (HairDebugMode == EHairDebugMode::MacroGroups)
 	{
-		AddHairDebugPrintInstancePass(GraphBuilder, ShaderMap, ShaderPrintData, Instances);
+		AddHairDebugPrintInstancePass(GraphBuilder, ShaderMap, ShaderPrintData, Instances, InstanceCountPerType);
 	}
 
 	if (HairDebugMode == EHairDebugMode::MeshProjection)
diff --git a/Engine/Plugins/Runtime/HairStrands/Source/HairStrandsCore/Private/GroomManager.cpp b/Engine/Plugins/Runtime/HairStrands/Source/HairStrandsCore/Private/GroomManager.cpp
--- a/Engine/Plugins/Runtime/HairStrands/Source/HairStrandsCore/Private/GroomManager.cpp
+++ b/Engine/Plugins/Runtime/HairStrands/Source/HairStrandsCore/Private/GroomManager.cpp
@@ -1019,6 +1019,7 @@ void RunHairStrandsDebug(
 	FGlobalShaderMap* ShaderMap,
 	const FSceneView& View,
 	const FHairStrandsInstances& Instances,
+	const FUintVector4& InstanceCountPerType,
 	const FGPUSkinCache* SkinCache,
 	const FShaderPrintData* ShaderPrintData,
 	FRDGTextureRef SceneColor,
@@ -1123,6 +1124,7 @@ void ProcessHairStrandsBookmark(
 			Parameters.ShaderMap,
 			*Parameters.View,
 			Instances,
+			Parameters.InstanceCountPerType,
 			Parameters.SkinCache,
 			Parameters.ShaderPrintData,
 			Parameters.SceneColorTexture,
diff --git a/Engine/Shaders/Private/HairStrands/HairStrandsDebug.usf b/Engine/Shaders/Private/HairStrands/HairStrandsDebug.usf
--- a/Engine/Shaders/Private/HairStrands/HairStrandsDebug.usf
+++ b/Engine/Shaders/Private/HairStrands/HairStrandsDebug.usf
@@ -1091,6 +1091,9 @@ void AddNewLine(float2 OriginalPos, inout float2 Pos)
 
 int2					MaxResolution;
 uint					InstanceCount;
+uint					InstanceCount_StrandsPrimaryView;
+uint					InstanceCount_StrandsShadowView;
+uint					InstanceCount_CardsOrMeshes;
 
 Buffer<uint>			Infos;
 
@@ -1121,23 +1124,32 @@ struct FInstanceInfo
 	uint UniqueSection;
 	uint RootCount;
 	uint ControlPointCount;
+
+	uint bIsVisibleInPrimaryView;
+	uint bIsVisibleInShadowView;
 };
 
 FInstanceInfo ReadInstanceInfo(uint Index)
 {
-	const uint Index8 = Index * 8;
+	const uint Index12 = Index * 12;
 
 	uint4 Data0 = 0;
-	Data0.x = Infos[Index8 + 0];
-	Data0.y = Infos[Index8 + 1];
-	Data0.z = Infos[Index8 + 2];
-	Data0.w = Infos[Index8 + 3];
+	Data0.x = Infos[Index12 + 0];
+	Data0.y = Infos[Index12 + 1];
+	Data0.z = Infos[Index12 + 2];
+	Data0.w = Infos[Index12 + 3];
 
 	uint4 Data1 = 0;
-	Data1.x = Infos[Index8 + 4];
-	Data1.y = Infos[Index8 + 5];
-	Data1.z = Infos[Index8 + 6];
-	Data1.w = Infos[Index8 + 7];
+	Data1.x = Infos[Index12 + 4];
+	Data1.y = Infos[Index12 + 5];
+	Data1.z = Infos[Index12 + 6];
+	Data1.w = Infos[Index12 + 7];
+
+	uint4 Data2 = 0;
+	Data2.x = Infos[Index12 + 8];
+	Data2.y = Infos[Index12 + 9];
+	Data2.z = Infos[Index12 + 10];
+	Data2.w = Infos[Index12 + 11];
 
 	FInstanceInfo Out = (FInstanceInfo)0;
 	Out.GroupIndex = (Data0.x >> 0) & 0xFF;
@@ -1160,6 +1172,9 @@ FInstanceInfo ReadInstanceInfo(uint Index)
 	Out.RootCount			= Data1.z;
 	Out.ControlPointCount	= Data1.w;
 
+	Out.bIsVisibleInPrimaryView = Data2.x & 0x1;
+	Out.bIsVisibleInShadowView	= Data2.x & 0x2;
+
 	return Out;
 }
 
@@ -1238,22 +1253,38 @@ void MainCS(uint GroupIndex : SV_GroupIndex, uint3 DispatchThreadId : SV_Dispatc
 		}
 	}
 
-	Pos = ShaderPrintText(Pos, TEXT("Total Groups        : "), FontWhite);
-	Pos = Print(Pos, TotalGroup, FontYellow, 8);
-	AddNewLine(OriginalPos, Pos);
+	{
+		Pos = ShaderPrintText(Pos, TEXT("Total Groups        : "), FontWhite);
+		Pos = Print(Pos, TotalGroup, FontYellow, 8);
+		AddNewLine(OriginalPos, Pos);
 
-	Pos = ShaderPrintText(Pos, TEXT("Total Strand Curves : "), FontWhite);
-	Pos = Print(Pos, TotalStrandsCount, FontOrange, 8);
-	AddNewLine(OriginalPos, Pos);
+		Pos = ShaderPrintText(Pos, TEXT("Total Strand Curves : "), FontWhite);
+		Pos = Print(Pos, TotalStrandsCount, FontOrange, 8);
+		AddNewLine(OriginalPos, Pos);
 
-	Pos = ShaderPrintText(Pos, TEXT("Total Strand Vertex : "), FontWhite);
-	Pos = Print(Pos, TotalVertexCount, FontSilver, 8);
+		Pos = ShaderPrintText(Pos, TEXT("Total Strand Vertex : "), FontWhite);
+		Pos = Print(Pos, TotalVertexCount, FontSilver, 8);
+		AddNewLine(OriginalPos, Pos);
+	}
 	AddNewLine(OriginalPos, Pos);
 
+	{
+		Pos = ShaderPrintText(Pos, TEXT("Total Strand primary: "), FontWhite);
+		Pos = Print(Pos, InstanceCount_StrandsPrimaryView, FontYellow, 8);
+		AddNewLine(OriginalPos, Pos);
+
+		Pos = ShaderPrintText(Pos, TEXT("Total Strand shadow : "), FontWhite);
+		Pos = Print(Pos, InstanceCount_StrandsShadowView, FontYellow, 8);
+		AddNewLine(OriginalPos, Pos);
+
+		Pos = ShaderPrintText(Pos, TEXT("Total Cards / Meshes: "), FontWhite);
+		Pos = Print(Pos, InstanceCount_CardsOrMeshes, FontYellow, 8);
+		AddNewLine(OriginalPos, Pos);
+	}
 	AddNewLine(OriginalPos, Pos);
 
 	// Draw instance details
-	Pos = ShaderPrintText(Pos, TEXT("   Group  LOD     Geometry Binding Sim   RBF   Clip    Strands Vertex  Bind.Triangle Bind.Section Name"));
+	Pos = ShaderPrintText(Pos, TEXT("   Group  LOD     Geometry Binding Sim   RBF   Clip    Strands Vertex  Bind.Triangle Bind.Section Primary Shadow Name"));
 	AddNewLine(OriginalPos, Pos);
 
 	for (uint InstanceIndex = 0; InstanceIndex < InstanceCount; ++InstanceIndex)
@@ -1318,6 +1349,21 @@ void MainCS(uint GroupIndex : SV_GroupIndex, uint3 DispatchThreadId : SV_Dispatc
 		Pos = Print(Pos, Info.UniqueSection, FontSilver, 8);
 		Pos = ShaderPrintText(Pos, TEXT("     "));
 
+		// Primary view
+		if (Info.bIsVisibleInPrimaryView)
+			Pos = ShaderPrintText(Pos, TEXT("True "), FontGreen);
+		else
+			Pos = ShaderPrintText(Pos, TEXT("False"), FontRed);
+		Pos = ShaderPrintText(Pos, TEXT(" "), FontWhite);
+
+		// Shadow view
+		Pos = ShaderPrintText(Pos, TEXT("  "));
+		if (Info.bIsVisibleInShadowView)
+			Pos = ShaderPrintText(Pos, TEXT("True "), FontGreen);
+		else
+			Pos = ShaderPrintText(Pos, TEXT("False"), FontRed);
+		Pos = ShaderPrintText(Pos, TEXT(" "), FontWhite);
+
 		// Name
 		Pos = PrintInstanceName(Pos, InstanceIndex, FontWhite);
 
diff --git a/Engine/Source/Runtime/Renderer/Private/HairStrands/HairStrandsInterface.cpp b/Engine/Source/Runtime/Renderer/Private/HairStrands/HairStrandsInterface.cpp
--- a/Engine/Source/Runtime/Renderer/Private/HairStrands/HairStrandsInterface.cpp
+++ b/Engine/Source/Runtime/Renderer/Private/HairStrands/HairStrandsInterface.cpp
@@ -652,23 +652,9 @@ FHairStrandsBookmarkParameters CreateHairStrandsBookmarkParameters(FScene* Scene
 			}
 		}
 	}
+	Out.InstanceCountPerType[HairInstanceCount_StrandsPrimaryView] = Out.VisibleInstances.Num();
 
-	// 2. Add all visible cards instances
-	for (const FMeshBatchAndRelevance& MeshBatch : View.HairCardsMeshElements)
-	{
-		check(MeshBatch.PrimitiveSceneProxy && MeshBatch.PrimitiveSceneProxy->ShouldRenderInMainPass());
-		if (MeshBatch.Mesh && MeshBatch.Mesh->Elements.Num() > 0)
-		{
-			FHairGroupPublicData* HairData = HairStrands::GetHairData(MeshBatch.Mesh);
-			if (HairData && HairData->Instance)
-			{
-				Out.VisibleInstances.Add(HairData->Instance);
-				InstancesVisibility[HairData->Instance->RegisteredIndex] = true;
-			}
-		}
-	}
-
-	// 3. Add all instances non-visible primary view(s) but visible in shadow view(s)
+	// 2. Add all instances non-visible primary view(s) but visible in shadow view(s)
 	if (IsHairStrandsNonVisibleShadowCastingEnable())
 	{
 		for (FHairStrandsInstance* Instance : Scene->HairStrandsSceneData.RegisteredProxies)
@@ -682,6 +668,23 @@ FHairStrandsBookmarkParameters CreateHairStrandsBookmarkParameters(FScene* Scene
 			}
 		}
 	}
+	Out.InstanceCountPerType[HairInstanceCount_StrandsShadowView] = Out.VisibleInstances.Num() - Out.InstanceCountPerType[HairInstanceCount_StrandsPrimaryView];
+
+	// 3. Add all visible cards instances
+	for (const FMeshBatchAndRelevance& MeshBatch : View.HairCardsMeshElements)
+	{
+		check(MeshBatch.PrimitiveSceneProxy && MeshBatch.PrimitiveSceneProxy->ShouldRenderInMainPass());
+		if (MeshBatch.Mesh && MeshBatch.Mesh->Elements.Num() > 0)
+		{
+			FHairGroupPublicData* HairData = HairStrands::GetHairData(MeshBatch.Mesh);
+			if (HairData && HairData->Instance)
+			{
+				Out.VisibleInstances.Add(HairData->Instance);
+				InstancesVisibility[HairData->Instance->RegisteredIndex] = true;
+			}
+		}
+	}
+	Out.InstanceCountPerType[HairInstanceCount_CardsOrMeshes] = Out.VisibleInstances.Num() - Out.InstanceCountPerType[HairInstanceCount_StrandsShadowView];
 
 	Out.ShaderPrintData			= ShaderPrint::IsEnabled(View.ShaderPrintData) ? &View.ShaderPrintData : nullptr;
 	Out.SkinCache				= View.Family->Scene->GetGPUSkinCache();
diff --git a/Engine/Source/Runtime/Renderer/Public/HairStrandsInterface.h b/Engine/Source/Runtime/Renderer/Public/HairStrandsInterface.h
--- a/Engine/Source/Runtime/Renderer/Public/HairStrandsInterface.h
+++ b/Engine/Source/Runtime/Renderer/Public/HairStrandsInterface.h
@@ -470,6 +470,13 @@ enum class EHairStrandsBookmark : uint8
 	ProcessEndOfFrame
 };
 
+enum EHairInstanceCount : uint8
+{
+	HairInstanceCount_StrandsPrimaryView = 0,
+	HairInstanceCount_StrandsShadowView = 1,
+	HairInstanceCount_CardsOrMeshes = 2,
+};
+
 struct FHairStrandsBookmarkParameters
 {
 	class FGPUSkinCache* SkinCache = nullptr;
@@ -485,6 +492,8 @@ struct FHairStrandsBookmarkParameters
 	FRDGTextureRef SceneColorTexture = nullptr;
 	FRDGTextureRef SceneDepthTexture = nullptr; 
 
+	FUintVector4 InstanceCountPerType = FUintVector4(0);
+
 	bool bHzbRequest = false;
 	uint32 FrameIndex = ~0;
 
