Add previous skel mesh previous position support for manual skin caching.

diff --git a/Engine/Plugins/Runtime/HairStrands/Source/HairStrandsCore/Private/GroomGeometryCache.cpp b/Engine/Plugins/Runtime/HairStrands/Source/HairStrandsCore/Private/GroomGeometryCache.cpp
--- a/Engine/Plugins/Runtime/HairStrands/Source/HairStrandsCore/Private/GroomGeometryCache.cpp
+++ b/Engine/Plugins/Runtime/HairStrands/Source/HairStrandsCore/Private/GroomGeometryCache.cpp
@@ -42,31 +42,55 @@ FHairStrandsProjectionMeshData::Section ConvertMeshSection(const FCachedGeometry
 
 ////////////////////////////////////////////////////////////////////////////////////////////////////////////////////////
 
-static void BuildBoneMatrices(const USkeletalMeshComponent* SkeletalMeshComponent, const FSkeletalMeshLODRenderData& LODData,
-	const uint32 LODIndex, TArray<uint32>& MatrixOffsets, TArray<FVector4f>& BoneMatrices)
+ENGINE_API void UpdateRefToLocalMatrices(TArray<FMatrix44f>& ReferenceToLocal, const USkinnedMeshComponent* InMeshComponent, const FSkeletalMeshRenderData* InSkeletalMeshRenderData, int32 LODIndex, const TArray<FBoneIndexType>* ExtraRequiredBoneIndices = NULL);
+ENGINE_API void UpdatePreviousRefToLocalMatrices(TArray<FMatrix44f>& ReferenceToLocal, const USkinnedMeshComponent* InMeshComponent, const FSkeletalMeshRenderData* InSkeletalMeshRenderData, int32 LODIndex, const TArray<FBoneIndexType>* ExtraRequiredBoneIndices = NULL);
+
+static void BuildBoneMatrices(
+	const USkeletalMeshComponent* SkeletalMeshComponent, 
+	const FSkeletalMeshRenderData* RenderData,
+	const FSkeletalMeshLODRenderData& LODData,
+	const uint32 LODIndex, 
+	const bool bNeedPreviousPosition,
+	TArray<uint32>& OutMatrixOffsets, 
+	TArray<FVector4f>& OutBoneMatrices,
+	TArray<FVector4f>& OutPrevBoneMatrices)
 {
+	//SkeletalMeshComponent->GetCurrentRefToLocalMatrices(BoneTransforms, LODIndex);
+	check(RenderData->LODRenderData.IsValidIndex(LODIndex));
+
 	TArray<FMatrix44f> BoneTransforms;
-	SkeletalMeshComponent->GetCurrentRefToLocalMatrices(BoneTransforms, LODIndex);
+	UpdateRefToLocalMatrices(BoneTransforms, SkeletalMeshComponent, RenderData, LODIndex, nullptr);
 
-	MatrixOffsets.SetNum(LODData.GetNumVertices());
+	TArray<FMatrix44f> PrevBoneTransforms;
+	if (bNeedPreviousPosition)
+	{
+		UpdatePreviousRefToLocalMatrices(PrevBoneTransforms, SkeletalMeshComponent, RenderData, LODIndex, nullptr);
+	}
+	
+	OutMatrixOffsets.SetNum(LODData.GetNumVertices());
 	uint32 BonesOffset = 0;
 	for (int32 SectionIdx = 0; SectionIdx < LODData.RenderSections.Num(); ++SectionIdx)
 	{
 		const FSkelMeshRenderSection& Section = LODData.RenderSections[SectionIdx];
 		for (uint32 SectionVertex = 0; SectionVertex < Section.NumVertices; ++SectionVertex)
 		{
-			MatrixOffsets[Section.BaseVertexIndex + SectionVertex] = BonesOffset;
+			OutMatrixOffsets[Section.BaseVertexIndex + SectionVertex] = BonesOffset;
 		}
 		BonesOffset += Section.BoneMap.Num();
 	}
-	BoneMatrices.SetNum(BonesOffset * 3);
+	OutBoneMatrices.SetNum(BonesOffset * 3);
+	OutPrevBoneMatrices.SetNum(BonesOffset * 3);
 	BonesOffset = 0;
 	for (int32 SectionIdx = 0; SectionIdx < LODData.RenderSections.Num(); ++SectionIdx)
 	{
 		const FSkelMeshRenderSection& Section = LODData.RenderSections[SectionIdx];
 		for (int32 BoneIdx = 0; BoneIdx < Section.BoneMap.Num(); ++BoneIdx, ++BonesOffset)
 		{
-			BoneTransforms[Section.BoneMap[BoneIdx]].To3x4MatrixTranspose(&BoneMatrices[3 * BonesOffset].X);
+			BoneTransforms[Section.BoneMap[BoneIdx]].To3x4MatrixTranspose(&OutBoneMatrices[3 * BonesOffset].X);
+			if (bNeedPreviousPosition)
+			{
+				PrevBoneTransforms[Section.BoneMap[BoneIdx]].To3x4MatrixTranspose(&OutPrevBoneMatrices[3 * BonesOffset].X);
+			}
 		}
 	}
 }
@@ -89,19 +113,20 @@ static void BuildBoneMatrices(const USkeletalMeshComponent* SkeletalMeshComponen
 			FSkeletalMeshRenderData* RenderData = SkeletalMeshComponent->GetSkeletalMesh()->GetResourceForRendering();
 			FSkeletalMeshLODRenderData& LODData = RenderData->LODRenderData[LODIndex];
 
+			const bool bNeedPreviousPosition = IsHairStrandContinuousDecimationReorderingEnabled();
+
 			TArray<uint32> MatrixOffsets;
 			TArray<FVector4f> BoneMatrices;
-			BuildBoneMatrices(SkeletalMeshComponent, LODData, LODIndex, MatrixOffsets, BoneMatrices);
-
-			const bool bNeedPreviousPosition = IsHairStrandContinuousDecimationReorderingEnabled();
+			TArray<FVector4f> PrevBoneMatrices;
+			BuildBoneMatrices(SkeletalMeshComponent, RenderData, LODData, LODIndex, bNeedPreviousPosition, MatrixOffsets, BoneMatrices, PrevBoneMatrices);
 
 			FRDGBufferRef DeformedPositionsBuffer			= GraphBuilder.CreateBuffer(FRDGBufferDesc::CreateBufferDesc(sizeof(float), LODData.StaticVertexBuffers.PositionVertexBuffer.GetNumVertices() * 3), TEXT("Hair.SkinnedDeformedPositions"));
 			FRDGBufferRef DeformedPreviousPositionsBuffer	= bNeedPreviousPosition ? GraphBuilder.CreateBuffer(FRDGBufferDesc::CreateBufferDesc(sizeof(float), LODData.StaticVertexBuffers.PositionVertexBuffer.GetNumVertices() * 3), TEXT("Hair.SkinnedDeformedPreviousPositions")) : nullptr;
 			FRDGBufferRef BoneMatricesBuffer				= CreateStructuredBuffer(GraphBuilder, TEXT("Hair.SkinnedBoneMatrices"), sizeof(float) * 4, BoneMatrices.Num(), BoneMatrices.GetData(), sizeof(float) * 4 * BoneMatrices.Num());
+			FRDGBufferRef PreviousBoneMatricesBuffer		= bNeedPreviousPosition ? CreateStructuredBuffer(GraphBuilder, TEXT("Hair.SkinnedPreviousBoneMatrices"), sizeof(float) * 4, BoneMatrices.Num(), BoneMatrices.GetData(), sizeof(float) * 4 * BoneMatrices.Num()) : nullptr;
 			FRDGBufferRef MatrixOffsetsBuffer				= CreateStructuredBuffer(GraphBuilder, TEXT("Hair.SkinnedMatrixOffsets"), sizeof(uint32), MatrixOffsets.Num(), MatrixOffsets.GetData(), sizeof(uint32) * MatrixOffsets.Num());
 
-			// TODO previous bone matrices
-			AddSkinUpdatePass(GraphBuilder, ShaderMap, SkeletalMeshComponent->GetSkinWeightBuffer(LODIndex), LODData, BoneMatricesBuffer, MatrixOffsetsBuffer, DeformedPositionsBuffer);
+			AddSkinUpdatePass(GraphBuilder, ShaderMap, SkeletalMeshComponent->GetSkinWeightBuffer(LODIndex), LODData, BoneMatricesBuffer, PreviousBoneMatricesBuffer, MatrixOffsetsBuffer, DeformedPositionsBuffer, DeformedPreviousPositionsBuffer);
 
 			Out.DeformedPositionBuffer = DeformedPositionsBuffer;
 			Out.DeformedPreviousPositionBuffer = DeformedPreviousPositionsBuffer;
diff --git a/Engine/Plugins/Runtime/HairStrands/Source/HairStrandsCore/Private/HairStrandsMeshProjection.cpp b/Engine/Plugins/Runtime/HairStrands/Source/HairStrandsCore/Private/HairStrandsMeshProjection.cpp
--- a/Engine/Plugins/Runtime/HairStrands/Source/HairStrandsCore/Private/HairStrandsMeshProjection.cpp
+++ b/Engine/Plugins/Runtime/HairStrands/Source/HairStrandsCore/Private/HairStrandsMeshProjection.cpp
@@ -75,10 +75,12 @@ class FSkinUpdateCS : public FGlobalShader
 	DECLARE_GLOBAL_SHADER(FSkinUpdateCS);
 	SHADER_USE_PARAMETER_STRUCT(FSkinUpdateCS, FGlobalShader);
 
-	class FUnlimitedBoneInfluence : SHADER_PERMUTATION_INT("GPUSKIN_UNLIMITED_BONE_INFLUENCE", 2);
-	class FUseExtraInfluence : SHADER_PERMUTATION_INT("GPUSKIN_USE_EXTRA_INFLUENCES", 2);
-	class FIndexUint16 : SHADER_PERMUTATION_INT("GPUSKIN_BONE_INDEX_UINT16", 2);
-	using FPermutationDomain = TShaderPermutationDomain<FUnlimitedBoneInfluence, FUseExtraInfluence, FIndexUint16>;
+	
+	class FPrevious : SHADER_PERMUTATION_BOOL("PERMUTATION_PREV");
+	class FUnlimitedBoneInfluence : SHADER_PERMUTATION_BOOL("GPUSKIN_UNLIMITED_BONE_INFLUENCE");
+	class FUseExtraInfluence : SHADER_PERMUTATION_BOOL("GPUSKIN_USE_EXTRA_INFLUENCES");
+	class FIndexUint16 : SHADER_PERMUTATION_BOOL("GPUSKIN_BONE_INDEX_UINT16");
+	using FPermutationDomain = TShaderPermutationDomain<FUnlimitedBoneInfluence, FUseExtraInfluence, FIndexUint16, FPrevious>;
 
 	BEGIN_SHADER_PARAMETER_STRUCT(FParameters, )
 		SHADER_PARAMETER(uint32, IndexSize)
@@ -86,12 +88,13 @@ class FSkinUpdateCS : public FGlobalShader
 		SHADER_PARAMETER(uint32, WeightStride)
 		SHADER_PARAMETER_SRV(Buffer<uint>, WeightLookup)
 		SHADER_PARAMETER_RDG_BUFFER_SRV(StructuredBuffer<float4>, BoneMatrices)
+		SHADER_PARAMETER_RDG_BUFFER_SRV(StructuredBuffer<float4>, PrevBoneMatrices)
 		SHADER_PARAMETER_RDG_BUFFER_SRV(StructuredBuffer<uint>, MatrixOffsets)
 		SHADER_PARAMETER_SRV(Buffer<uint>, VertexWeights)
 		SHADER_PARAMETER_SRV(Buffer<float>, RestPositions)
 		SHADER_PARAMETER_RDG_BUFFER_UAV(RWBuffer<float>, DeformedPositions)
-
-		END_SHADER_PARAMETER_STRUCT()
+		SHADER_PARAMETER_RDG_BUFFER_UAV(RWBuffer<float>, PrevDeformedPositions)
+	END_SHADER_PARAMETER_STRUCT()
 
 public:
 	static bool ShouldCompilePermutation(const FGlobalShaderPermutationParameters& Parameters) { return IsHairStrandsSupported(EHairStrandsShaderType::All, Parameters.Platform); }
@@ -105,32 +108,41 @@ void AddSkinUpdatePass(
 	FSkinWeightVertexBuffer* SkinWeight,
 	FSkeletalMeshLODRenderData& RenderData,
 	FRDGBufferRef BoneMatrices,
+	FRDGBufferRef PrevBoneMatrices,
 	FRDGBufferRef MatrixOffsets,
-	FRDGBufferRef OutDeformedosition)
+	FRDGBufferRef OutDeformedPosition,
+	FRDGBufferRef OutPrevDeformedPosition)
 {
-	FSkinUpdateCS::FParameters* Parameters = GraphBuilder.AllocParameters<FSkinUpdateCS::FParameters>();
+	const bool bPrevPosition = OutPrevDeformedPosition != nullptr && PrevBoneMatrices != nullptr;
 
+	FSkinUpdateCS::FParameters* Parameters = GraphBuilder.AllocParameters<FSkinUpdateCS::FParameters>();
 	Parameters->IndexSize = SkinWeight->GetBoneIndexByteSize();
 	Parameters->NumVertices = RenderData.StaticVertexBuffers.PositionVertexBuffer.GetNumVertices();
 	Parameters->WeightStride = SkinWeight->GetConstantInfluencesVertexStride();
 	Parameters->WeightLookup = SkinWeight->GetLookupVertexBuffer()->GetSRV();
-	Parameters->BoneMatrices = GraphBuilder.CreateSRV(BoneMatrices);//BoneBuffer.VertexBufferSRV;
-	Parameters->MatrixOffsets = GraphBuilder.CreateSRV(MatrixOffsets);//BoneBuffer.VertexBufferSRV;
+	Parameters->BoneMatrices = GraphBuilder.CreateSRV(BoneMatrices);
+	Parameters->MatrixOffsets = GraphBuilder.CreateSRV(MatrixOffsets);
 	Parameters->VertexWeights = SkinWeight->GetDataVertexBuffer()->GetSRV();
 	Parameters->RestPositions = RenderData.StaticVertexBuffers.PositionVertexBuffer.GetSRV();
-	Parameters->DeformedPositions = GraphBuilder.CreateUAV(OutDeformedosition, PF_R32_FLOAT);
+	Parameters->DeformedPositions = GraphBuilder.CreateUAV(OutDeformedPosition, PF_R32_FLOAT);
+	if (bPrevPosition)
+	{
+		Parameters->PrevBoneMatrices = GraphBuilder.CreateSRV(BoneMatrices);
+		Parameters->PrevDeformedPositions = GraphBuilder.CreateUAV(OutPrevDeformedPosition, PF_R32_FLOAT);
+	}
 
 	FSkinUpdateCS::FPermutationDomain PermutationVector;
 	PermutationVector.Set<FSkinUpdateCS::FUnlimitedBoneInfluence>(SkinWeight->GetBoneInfluenceType() == GPUSkinBoneInfluenceType::UnlimitedBoneInfluence);
 	PermutationVector.Set<FSkinUpdateCS::FUseExtraInfluence>(SkinWeight->GetMaxBoneInfluences() > MAX_INFLUENCES_PER_STREAM);
 	PermutationVector.Set<FSkinUpdateCS::FIndexUint16>(SkinWeight->Use16BitBoneIndex());
+	PermutationVector.Set<FSkinUpdateCS::FPrevious>(bPrevPosition);
 
 	const FIntVector DispatchGroupCount = FComputeShaderUtils::GetGroupCount(RenderData.StaticVertexBuffers.PositionVertexBuffer.GetNumVertices(), 64);
 	check(DispatchGroupCount.X < 65536);
 	TShaderMapRef<FSkinUpdateCS> ComputeShader(ShaderMap, PermutationVector);
 	FComputeShaderUtils::AddPass(
 		GraphBuilder,
-		RDG_EVENT_NAME("UpdateSkinPosition"),
+		RDG_EVENT_NAME("HairStrands::UpdateSkinPosition(%s)", bPrevPosition ? TEXT("Curr,Prev") : TEXT("Curr")),
 		ComputeShader,
 		Parameters,
 		DispatchGroupCount);
diff --git a/Engine/Plugins/Runtime/HairStrands/Source/HairStrandsCore/Public/HairStrandsMeshProjection.h b/Engine/Plugins/Runtime/HairStrands/Source/HairStrandsCore/Public/HairStrandsMeshProjection.h
--- a/Engine/Plugins/Runtime/HairStrands/Source/HairStrandsCore/Public/HairStrandsMeshProjection.h
+++ b/Engine/Plugins/Runtime/HairStrands/Source/HairStrandsCore/Public/HairStrandsMeshProjection.h
@@ -121,8 +121,10 @@ void AddSkinUpdatePass(
 	class FSkinWeightVertexBuffer* SkinWeight,
 	class FSkeletalMeshLODRenderData& RenderData,
 	FRDGBufferRef BoneMatrices,
+	FRDGBufferRef PreviousBoneMatrices,
 	FRDGBufferRef MatrixOffsets,
-	FRDGBufferRef OutDeformedosition);
+	FRDGBufferRef OutDeformedosition,
+	FRDGBufferRef OutPreviousDeformedosition);
 
 void AddHairMeshesRBFInterpolationPass(
 	FRDGBuilder& GraphBuilder,
diff --git a/Engine/Shaders/Private/HairStrands/HairStrandsSkinUpdate.usf b/Engine/Shaders/Private/HairStrands/HairStrandsSkinUpdate.usf
--- a/Engine/Shaders/Private/HairStrands/HairStrandsSkinUpdate.usf
+++ b/Engine/Shaders/Private/HairStrands/HairStrandsSkinUpdate.usf
@@ -11,13 +11,20 @@ uint NumVertices;
 uint WeightStride;
 
 Buffer<uint> WeightLookup;
-StructuredBuffer<float4> BoneMatrices;
 StructuredBuffer<uint> MatrixOffsets;
 Buffer<uint> VertexWeights;
 Buffer<float> RestPositions;
 
+// Current Position
+StructuredBuffer<float4> BoneMatrices;
 RWBuffer<float> DeformedPositions;
 
+// Previous Position
+#if PERMUTATION_PREV
+StructuredBuffer<float4> PrevBoneMatrices;
+RWBuffer<float> PrevDeformedPositions;
+#endif
+
 //----------------------------------------------------------------------------------
 
 #define	GPUSKIN_VB_OFFSET_INFLUENCEBONES	0
@@ -43,7 +50,7 @@ struct FUnpackedWeight
 
 //----------------------------------------------------------------------------------
 
-FBlendedMatrix BlendMatrix( FUnpackedWeight UnpackedWeight, uint VertexIndex )
+FBlendedMatrix BlendMatrix(FUnpackedWeight UnpackedWeight, uint VertexIndex, StructuredBuffer<float4> InBoneMatrices)
 {
 	const uint MatrixOffset = MatrixOffsets[VertexIndex];
 #if GPUSKIN_UNLIMITED_BONE_INFLUENCE
@@ -52,7 +59,7 @@ FBlendedMatrix BlendMatrix( FUnpackedWeight UnpackedWeight, uint VertexIndex )
 	int StreamOffset = BlendOffsetCount >> 8;
 	int WeightsOffset = StreamOffset + (IndexSize * NumBoneInfluences);
 
-	FBlendedMatrix BoneMatrix = FBlendedMatrix(float4(0,0,0,0), float4(0,0,0,0), float4(0,0,0,0));
+	FBlendedMatrix OutBoneMatrix = FBlendedMatrix(float4(0,0,0,0), float4(0,0,0,0), float4(0,0,0,0));
 	for (int InfluenceIdx = 0; InfluenceIdx < NumBoneInfluences; InfluenceIdx++)
 	{
 		int BoneIndexOffset = StreamOffset + (IndexSize * InfluenceIdx);
@@ -62,7 +69,7 @@ FBlendedMatrix BlendMatrix( FUnpackedWeight UnpackedWeight, uint VertexIndex )
 		const uint IndexOffset = (MatrixOffset + BoneIndex) * 3;
 		//@todo-lh: Workaround to fix issue in SPIRVEmitter of DXC; this block must be inside the if branch
 		float BoneWeight = float(VertexWeights[WeightsOffset + InfluenceIdx]) / 255.0f;
-		BoneMatrix += BoneWeight * FBlendedMatrix(BoneMatrices[IndexOffset], BoneMatrices[IndexOffset+1], BoneMatrices[IndexOffset+2]);
+		OutBoneMatrix += BoneWeight * FBlendedMatrix(InBoneMatrices[IndexOffset], InBoneMatrices[IndexOffset+1], InBoneMatrices[IndexOffset+2]);
 	}
 	
 #else // GPUSKIN_UNLIMITED_BONE_INFLUENCE
@@ -71,23 +78,23 @@ FBlendedMatrix BlendMatrix( FUnpackedWeight UnpackedWeight, uint VertexIndex )
 	uint IndexOffsetZ = (MatrixOffset + UnpackedWeight.BlendIndices.z) * 3;
 	uint IndexOffsetW = (MatrixOffset + UnpackedWeight.BlendIndices.w) * 3;
 
-	FBlendedMatrix BoneMatrix = UnpackedWeight.BlendWeights.x * FBlendedMatrix(BoneMatrices[IndexOffsetX], BoneMatrices[IndexOffsetX+1], BoneMatrices[IndexOffsetX+2]);
-	BoneMatrix += UnpackedWeight.BlendWeights.y * FBlendedMatrix(BoneMatrices[IndexOffsetY], BoneMatrices[IndexOffsetY+1], BoneMatrices[IndexOffsetY+2]);
-	BoneMatrix += UnpackedWeight.BlendWeights.z * FBlendedMatrix(BoneMatrices[IndexOffsetZ], BoneMatrices[IndexOffsetZ+1], BoneMatrices[IndexOffsetZ+2]);
-	BoneMatrix += UnpackedWeight.BlendWeights.w * FBlendedMatrix(BoneMatrices[IndexOffsetW], BoneMatrices[IndexOffsetW+1], BoneMatrices[IndexOffsetW+2]);
-#if GPUSKIN_USE_EXTRA_INFLUENCES
+	FBlendedMatrix OutBoneMatrix = UnpackedWeight.BlendWeights.x * FBlendedMatrix(InBoneMatrices[IndexOffsetX], InBoneMatrices[IndexOffsetX+1], InBoneMatrices[IndexOffsetX+2]);
+	OutBoneMatrix += UnpackedWeight.BlendWeights.y * FBlendedMatrix(InBoneMatrices[IndexOffsetY], InBoneMatrices[IndexOffsetY+1], InBoneMatrices[IndexOffsetY+2]);
+	OutBoneMatrix += UnpackedWeight.BlendWeights.z * FBlendedMatrix(InBoneMatrices[IndexOffsetZ], InBoneMatrices[IndexOffsetZ+1], InBoneMatrices[IndexOffsetZ+2]);
+	OutBoneMatrix += UnpackedWeight.BlendWeights.w * FBlendedMatrix(InBoneMatrices[IndexOffsetW], InBoneMatrices[IndexOffsetW+1], InBoneMatrices[IndexOffsetW+2]);
+  #if GPUSKIN_USE_EXTRA_INFLUENCES
 	IndexOffsetX = (MatrixOffset + UnpackedWeight.ExtraIndices.x) * 3;
 	IndexOffsetY = (MatrixOffset + UnpackedWeight.ExtraIndices.y) * 3;
 	IndexOffsetZ = (MatrixOffset + UnpackedWeight.ExtraIndices.z) * 3;
 	IndexOffsetW = (MatrixOffset + UnpackedWeight.ExtraIndices.w) * 3;
 
-	BoneMatrix += UnpackedWeight.ExtraWeights.x * FBlendedMatrix(BoneMatrices[IndexOffsetX], BoneMatrices[IndexOffsetX+1], BoneMatrices[IndexOffsetX+2]);
-	BoneMatrix += UnpackedWeight.ExtraWeights.y * FBlendedMatrix(BoneMatrices[IndexOffsetY], BoneMatrices[IndexOffsetY+1], BoneMatrices[IndexOffsetY+2]);
-	BoneMatrix += UnpackedWeight.ExtraWeights.z * FBlendedMatrix(BoneMatrices[IndexOffsetZ], BoneMatrices[IndexOffsetZ+1], BoneMatrices[IndexOffsetZ+2]);
-	BoneMatrix += UnpackedWeight.ExtraWeights.w * FBlendedMatrix(BoneMatrices[IndexOffsetW], BoneMatrices[IndexOffsetW+1], BoneMatrices[IndexOffsetW+2]);
-#endif
+	OutBoneMatrix += UnpackedWeight.ExtraWeights.x * FBlendedMatrix(BoneMatrices[IndexOffsetX], BoneMatrices[IndexOffsetX+1], BoneMatrices[IndexOffsetX+2]);
+	OutBoneMatrix += UnpackedWeight.ExtraWeights.y * FBlendedMatrix(BoneMatrices[IndexOffsetY], BoneMatrices[IndexOffsetY+1], BoneMatrices[IndexOffsetY+2]);
+	OutBoneMatrix += UnpackedWeight.ExtraWeights.z * FBlendedMatrix(BoneMatrices[IndexOffsetZ], BoneMatrices[IndexOffsetZ+1], BoneMatrices[IndexOffsetZ+2]);
+	OutBoneMatrix += UnpackedWeight.ExtraWeights.w * FBlendedMatrix(BoneMatrices[IndexOffsetW], BoneMatrices[IndexOffsetW+1], BoneMatrices[IndexOffsetW+2]);
+  #endif
 #endif // GPUSKIN_UNLIMITED_BONE_INFLUENCE
-	return BoneMatrix;
+	return OutBoneMatrix;
 }
 
 //----------------------------------------------------------------------------------
@@ -156,18 +163,31 @@ void UpdateSkinPositionCS(uint3 GroupID : SV_GroupID,
 		const uint WeightOffset = VertexIndex * WeightStride / 4;
 
 		const FUnpackedWeight UnpackedWeight = UnpackWeight(WeightOffset);
-		const FBlendedMatrix BlendedMatrix = BlendMatrix(UnpackedWeight, VertexIndex);
 
-		float4 InputPosition;
-		InputPosition.x = RestPositions[PositionOffset];
-		InputPosition.y = RestPositions[PositionOffset + 1];
-		InputPosition.z = RestPositions[PositionOffset + 2];
-		InputPosition.w = 1.0;
-
-		const float3 OutputPosition =  mul( BlendedMatrix, InputPosition );
-		
-		DeformedPositions[PositionOffset] = OutputPosition.x;
-		DeformedPositions[PositionOffset + 1] = OutputPosition.y;
-		DeformedPositions[PositionOffset + 2] = OutputPosition.z;
+		const float4 InputPosition = float4(
+			RestPositions[PositionOffset],
+			RestPositions[PositionOffset + 1],
+			RestPositions[PositionOffset + 2],
+			1.0);
+
+		// Current position
+		{
+			const FBlendedMatrix BlendedMatrix = BlendMatrix(UnpackedWeight, VertexIndex, BoneMatrices);
+			const float3 OutputPosition =  mul(BlendedMatrix, InputPosition);
+			DeformedPositions[PositionOffset] = OutputPosition.x;
+			DeformedPositions[PositionOffset + 1] = OutputPosition.y;
+			DeformedPositions[PositionOffset + 2] = OutputPosition.z;
+		}
+
+		// Previous position
+	#if PERMUTATION_PREV
+		{
+			const FBlendedMatrix PrevBlendedMatrix = BlendMatrix(UnpackedWeight, VertexIndex, PrevBoneMatrices);
+			const float3 OutputPosition = mul(PrevBlendedMatrix, InputPosition);
+			PrevDeformedPositions[PositionOffset] = OutputPosition.x;
+			PrevDeformedPositions[PositionOffset + 1] = OutputPosition.y;
+			PrevDeformedPositions[PositionOffset + 2] = OutputPosition.z;
+		}
+	#endif
 	}
 }
