Made hair SW rasterizer binning shader more robust by adding support for hair segments covering more than three bins/tiles.

#rnx

diff --git a/Engine/Shaders/Private/HairStrands/HairStrandsVisibilityRasterCompute.usf b/Engine/Shaders/Private/HairStrands/HairStrandsVisibilityRasterCompute.usf
--- a/Engine/Shaders/Private/HairStrands/HairStrandsVisibilityRasterCompute.usf
+++ b/Engine/Shaders/Private/HairStrands/HairStrandsVisibilityRasterCompute.usf
@@ -63,7 +63,7 @@ struct FVisTile
 #define VT_PrimOffset 0
 #define VT_PrimCount 1
 #define VT_Coord 2
-#define VT_NearestDepth 3
+#define VT_MinWriteIndex 3
 
 uint PackVisTileCoord(uint2 Coord)
 {
@@ -128,6 +128,53 @@ void CalcScreenPosAndRad(in uint4 PB, in float3 PBO, out float4 SP, out float Ra
 	Rad = CP.WorldRadius*2000.0; // OutputResolutionf.x; //TODO: figure this out correctly?
 }
 
+bool ClipRaySegment(float3 AABBMin, float3 AABBMax, inout float3 P0, inout float3 P1)
+{
+	const bool bP0Outside = any(P0 < AABBMin) || any(P0 > AABBMax);
+	const bool bP1Outside = any(P1 < AABBMin) || any(P1 > AABBMax);
+	if (!bP0Outside && !bP1Outside)
+	{
+		return true;
+	}
+
+	const float3 Origin = P0;
+	const float3 Dir = P1 - P0;
+	const float3 RcpDir = 1.0f / Dir;
+
+	const float3 T0 = (AABBMin - Origin) * RcpDir;
+	const float3 T1 = (AABBMax - Origin) * RcpDir;
+
+	const float TMin = max(max(min(T0.x, T1.x), min(T0.y, T1.y)), min(T0.z, T1.z));
+	const float TMax = min(min(max(T0.x, T1.x), max(T0.y, T1.y)), max(T0.z, T1.z));
+
+	// Ray intersects the AABB but the segment is completely outside.
+	if (TMax < 0.0f)
+	{
+		return false;
+	}
+
+	// No intersection
+	if (TMin > TMax)
+	{
+		return false;
+	}
+
+	float3 P0New = P0;
+	float3 P1New = P1;
+	if (bP0Outside && TMin >= 0.0f && TMin < 1.0f)
+	{
+		P0New = lerp(P0, P1, TMin);
+	}
+	if (bP1Outside && TMax >= 0.0f && TMax < 1.0f)
+	{
+		P1New = lerp(P0, P1, TMax);
+	}
+	P0 = P0New;
+	P1 = P1New;
+	
+	return true;
+}
+
 #endif // Common rasetrizer helper function & parameters
 
 ///////////////////////////////////////////////////////////////////////////
@@ -186,12 +233,50 @@ void PrepareDepthGridCS(uint dId : SV_DispatchThreadID, uint GroupThreadID : SV_
 
 #if SHADER_RASTERCOMPUTE_BINNING
 
-RWTexture3D<uint> 					OutVisTileBinningGrid;
+RWTexture2DArray<uint> 				OutVisTileBinningGrid;
 Texture2D<uint>						VisTileDepthGrid;
 RWBuffer<uint>						OutVisTilePrims;
 RWBuffer<uint>						OutVisTileArgs;
 RWByteAddressBuffer					OutVisTileData;
 
+// TODO: Without setting a limit, we are sometimes caught in an infinite loop even though this should not happen.
+#define DDA_MAX_ITERATIONS 256
+
+struct FDDAContext
+{
+	float2 TileCoord;
+	float2 DeltaDist;
+	float2 Step;
+	float2 SideDist;
+};
+
+FDDAContext DDACreateContext(float2 RayStart, float2 RayDir)
+{
+	const float2 RayDirRcp = 1.0f / RayDir;
+
+	FDDAContext Context;
+	Context.TileCoord = floor(RayStart);
+	Context.DeltaDist = abs(RayDirRcp);
+	Context.Step = sign(RayDir);
+	Context.SideDist = (Context.TileCoord - RayStart + max(Context.Step, 0.0f)) * RayDirRcp;
+
+	return Context;
+}
+
+void DDAAdvance(inout FDDAContext Context)
+{
+	if (Context.SideDist.x < Context.SideDist.y)
+	{
+		Context.SideDist.x += Context.DeltaDist.x;
+		Context.TileCoord.x += Context.Step.x;
+	}
+	else
+	{
+		Context.SideDist.y += Context.DeltaDist.y;
+		Context.TileCoord.y += Context.Step.y;
+	}
+}
+
 uint LoadOutVisTileData(uint index, uint offset)
 {
 	return OutVisTileData.Load((((index)) * VT_SIZE * 4) + ((offset) * 4));
@@ -205,11 +290,15 @@ void StoreOutVisTileData(uint index, uint offset, uint value)
 groupshared uint group_LoopNum;
 groupshared uint group_BatchNum;
 
+#define TILES_TO_ALLOCATE_MAX 1024
+groupshared uint group_TilesToAllocate[TILES_TO_ALLOCATE_MAX];
+groupshared uint group_TilesToAllocateCount;
+
 // The total number of line segments (VertexCount) is divided up equally between N binners - each binner = a workgroup which loops through the designated set segments in batches of 1024
 // NB there is still potential to use LDS to prevent/reduce contention when adding segments to the binning grid and this may improve perf
 
 [numthreads(1024, 1, 1)]
-void BinningCS(uint2 DispatchThreadID : SV_DispatchThreadID, uint GroupThreadID : SV_GroupThreadID, uint GroupID : SV_GroupID)
+void BinningCS(uint DispatchThreadID : SV_DispatchThreadID, uint GroupThreadID : SV_GroupThreadID, uint GroupID : SV_GroupID)
 {
 	ResolvedView = ResolveView();
 	if (GroupThreadID == 0)
@@ -232,30 +321,23 @@ void BinningCS(uint2 DispatchThreadID : SV_DispatchThreadID, uint GroupThreadID
 
 	GroupMemoryBarrierWithGroupSync();
 
-	[loop]
+	LOOP
 	for (uint LoopIndex = 0; LoopIndex < group_LoopNum; LoopIndex++)
 	{
 		uint BatchIndex = LoopIndex + (GroupID * group_LoopNum);
 
-		bool bSegValid = (BatchIndex < group_BatchNum) ? true : false;
-
-		bool bTileIntersected[3] = { false, false, false };
-
-		int2 TileCoord[3] = { int2(-1, -1), int2(-1, -1), int2(-1, -1) };
-
-		uint SegIndex[3] = { 0xffffffff, 0xffffffff, 0xffffffff };
-
+		bool bSegValid = (BatchIndex < group_BatchNum);
 
 #if PERMUTATION_CULLING
 		uint PrimID = BatchIndex * 1024 + GroupThreadID;
 		if (HairStrandsVF_bIsCullingEnable)
 		{
 			group_BatchNum = (HairStrandsVF_CullingIndirectBuffer[3] + 1023) / 1024;
-			if (bSegValid) bSegValid = (PrimID < (HairStrandsVF_CullingIndirectBuffer[3])) ? true : false;
+			bSegValid = bSegValid && (PrimID < (HairStrandsVF_CullingIndirectBuffer[3]));
 		}
 		else
 		{
-			if (bSegValid) bSegValid = (PrimID < (VertexCount)) ? true : false;
+			bSegValid = bSegValid && (PrimID < (VertexCount));
 		}
 
 		if (bSegValid && HairStrandsVF_bIsCullingEnable)
@@ -288,119 +370,169 @@ void BinningCS(uint2 DispatchThreadID : SV_DispatchThreadID, uint GroupThreadID
 		}
 #else
 		uint PrimID = BatchIndex * 1024 + GroupThreadID + VertexStart;
-		if (bSegValid) bSegValid = (PrimID < (VertexStart + VertexCount)) ? true : false;
+		bSegValid = bSegValid && (PrimID < (VertexStart + VertexCount));
 #endif
 
-		uint BinID = GroupID - (uint(float(GroupID + 0.5) * RcpNumBinners)* NumBinners);
-		uint BinID2 = BinID + NumBinners;
+		const uint SegmentCountLayerIdx = GroupID; // Stores number of segments per tile per workgroup.
+		const uint TmpSegmentCountLayerIdx = SegmentCountLayerIdx + NumBinners; // Also stores number of segments per tile per workgroup. Used as second counter for this two pass algorithm.
+		const uint TileAllocInfoLayerIdx = SegmentCountLayerIdx + NumBinners * 2; // Stores per tile per workgroup allocation info.
 
 
-		if (bSegValid) //currently hardcoded to a maximum of 3 tile bins - this only works corectly when segment length does not exceed TileSize - TODO: add support for N bins
-		{
-			float3 SP0;
-			float3 SP1;
+		uint NearestDepth = 0;
+		float2 TileCoord0F = 0.0f;
+		float2 TileCoord1F = 0.0f;
 
-			uint Type;
+		// Project segment end points and clip them to the screen
+		if (bSegValid)
+		{
+			float3 SP0 = 0.0f;
+			uint Type = -1;
 			CalcScreenPos(HairStrandsVF_PositionBuffer[PrimID], HairStrandsVF_PositionOffsetBuffer[0].xyz, SP0, Type);
 
-			bool bIsEndCV = (Type == HAIR_CONTROLPOINT_END); 
+			bool bIsEndCV = (Type == HAIR_CONTROLPOINT_END);
+			bSegValid = !bIsEndCV;
 
-			if (!bIsEndCV)
+			if (bSegValid)
 			{
+				float3 SP1 = 0.0f;
 				CalcScreenPos(HairStrandsVF_PositionBuffer[PrimID + 1], HairStrandsVF_PositionOffsetBuffer[0].xyz, SP1, Type);
+				
+				// Do clipping in tile space
+				SP0.xy *= RcpTileSize;
+				SP1.xy *= RcpTileSize;
+				bSegValid = ClipRaySegment(0.0f, float3(TileRes, 1.0f), SP0, SP1);
 
-				uint NearestDepth = PackHairVisDepthCoverage(max(SP0.z, SP1.z), 1.0);
-
-				TileCoord[0] = int2(SP0.x * RcpTileSize, SP0.y * RcpTileSize);
-
-				if (all(TileCoord[0] >= 0) && all(TileCoord[0] < TileRes) && NearestDepth > VisTileDepthGrid[TileCoord[0]])
+				if (bSegValid)
 				{
-					InterlockedAdd(OutVisTileBinningGrid[uint3(TileCoord[0], BinID)], 1, SegIndex[0]);
-					bTileIntersected[0] = true;
-				}
+					// For peace of mind, make sure these are actually clamped to a valid range.
+					SP0 = clamp(SP0, 0.0f, float3(TileRes, 1.0f));
+					SP1 = clamp(SP1, 0.0f, float3(TileRes, 1.0f));
 
-				TileCoord[1] = int2(SP1.x * RcpTileSize, SP1.y * RcpTileSize);
-
-				if ( any(TileCoord[1] != TileCoord[0]) && all(TileCoord[1] >= 0) && all(TileCoord[1] < TileRes) && NearestDepth > VisTileDepthGrid[TileCoord[1]])
-				{
-					InterlockedAdd(OutVisTileBinningGrid[uint3(TileCoord[1], BinID)], 1, SegIndex[1]);
-					bTileIntersected[1] = true;
+					NearestDepth = PackHairVisDepthCoverage(max(SP0.z, SP1.z), 1.0f);
+					TileCoord0F = SP0.xy;
+					TileCoord1F = SP1.xy;
 				}
+			}
+		}
 
-				if ((bTileIntersected[0] || bTileIntersected[1]) && all(TileCoord[1] != TileCoord[0]))
-				{
-					float2 MP = (TileCoord[0] + float2(0.5, 0.5) + TileCoord[1] + float2(0.5, 0.5)) * 0.5 * 32;
+		if (GroupThreadID == 0)
+		{
+			group_TilesToAllocateCount = 0;
+		}
 
-					float2 SPI;
-					SPI.x = SP0.x + (SP1.x - SP0.x) * ((MP.y - SP0.y) / (SP1.y - SP0.y));
-					SPI.y = SP0.y + (SP1.y - SP0.y) * ((MP.x - SP0.x) / (SP1.x - SP0.x));
+		GroupMemoryBarrierWithGroupSync();
+		
+		// Increment per workgroup per tile counters and add tiles to be allocated
+		if (bSegValid)
+		{
+			FDDAContext DDAContext = DDACreateContext(TileCoord0F, normalize(TileCoord1F - TileCoord0F));
+			const int2 EndCoord = (int2)floor(TileCoord1F);
 
-					TileCoord[2] = int2(SPI.x * RcpTileSize, SPI.y * RcpTileSize);
+			for (int DDAIt = 0; DDAIt < DDA_MAX_ITERATIONS; ++DDAIt)
+			{
+				const int2 TileCoord = (int2)floor(DDAContext.TileCoord);
+
+				BRANCH
+				if (NearestDepth > VisTileDepthGrid[TileCoord])
+				{
+					uint OldTileSegmentCount;
+					InterlockedAdd(OutVisTileBinningGrid[uint3(TileCoord, SegmentCountLayerIdx)], 1, OldTileSegmentCount);
 
-					if (all(TileCoord[2] >= 0) && all(TileCoord[2] < TileRes) && NearestDepth > VisTileDepthGrid[TileCoord[2]])
+					BRANCH
+					if ((OldTileSegmentCount % 1024) == 0)
 					{
-						InterlockedAdd(OutVisTileBinningGrid[uint3(TileCoord[2], BinID)], 1, SegIndex[2]);
-						bTileIntersected[2] = true;
+						uint WritePos;
+						InterlockedAdd(group_TilesToAllocateCount, 1, WritePos);
+						if (WritePos < TILES_TO_ALLOCATE_MAX)
+						{
+							group_TilesToAllocate[WritePos] = PackVisTileCoord(TileCoord);
+						}
 					}
 				}
+
+				if (all(TileCoord == EndCoord))
+				{
+					break;
+				}
+
+				DDAAdvance(DDAContext);
 			}
 		}
-		
+
 		GroupMemoryBarrierWithGroupSync();
 
-		if (bSegValid) //add tiles if any needed - this is unlikely to happen in all batches
+		// Allocate tiles
+		const uint TilesToAllocateCount = min(TILES_TO_ALLOCATE_MAX, group_TilesToAllocateCount);
+		for (uint TileIdx = GroupThreadID; TileIdx < TilesToAllocateCount; TileIdx += 1024)
 		{
-			for (int i = 0; i < 3; i++)
-			{
-				if ((SegIndex[i] != 0xffffffff) && ((SegIndex[i] % 1024)==0))
-				{
-					uint NewTile;
-					InterlockedAdd(OutVisTileArgs[0], 1, NewTile);
-					
-					StoreOutVisTileData(NewTile, VT_Coord, PackVisTileCoord(TileCoord[i]));
-					StoreOutVisTileData(NewTile, VT_PrimOffset, NewTile * 1024);
-					StoreOutVisTileData(NewTile, VT_NearestDepth, SegIndex[i]); //tmp - use this to decide if seg should go in new or prev tile.
+			const uint PackedTileCoord = group_TilesToAllocate[TileIdx];
+			const uint2 TileCoord = UppackVisTileCoord(PackedTileCoord);
 
-					uint PrevTile = (OutVisTileBinningGrid[uint3(TileCoord[i], BinID2)] & 0xffff);
-					
-					if (SegIndex[i]>0)	StoreOutVisTileData(PrevTile, VT_PrimCount, 1024);
+			const uint TotalNewWriteCount = OutVisTileBinningGrid[uint3(TileCoord, SegmentCountLayerIdx)];
+			const uint TotalOldWriteCount = OutVisTileBinningGrid[uint3(TileCoord, TmpSegmentCountLayerIdx)];
 
-					OutVisTileBinningGrid[uint3(TileCoord[i], BinID2)] = (PrevTile << 16) | (NewTile & 0xffff);
+			uint NewTile;
+			InterlockedAdd(OutVisTileArgs[0], 1, NewTile);
 
-				}
+			StoreOutVisTileData(NewTile, VT_Coord, PackedTileCoord);
+			// Round down the count to the start of the tile and later compare against this to decide which tile to write to.
+			StoreOutVisTileData(NewTile, VT_MinWriteIndex, TotalNewWriteCount & ~1023u);
+
+			const uint PrevTile = (OutVisTileBinningGrid[uint3(TileCoord, TileAllocInfoLayerIdx)] & 0xffff);
+
+			if (TotalOldWriteCount > 0)
+			{
+				StoreOutVisTileData(PrevTile, VT_PrimCount, 1024);
 			}
+
+			OutVisTileBinningGrid[uint3(TileCoord, TileAllocInfoLayerIdx)] = (PrevTile << 16) | (NewTile & 0xffff);
 		}
 
 		GroupMemoryBarrierWithGroupSync();
 
-		if (bSegValid) //now set the PrimID at the location indicated by the prev/current tile.
+		// Write PrimID to tiles
+		if (bSegValid)
 		{
-			for (int i = 0; i < 3; i++)
+			FDDAContext DDAContext = DDACreateContext(TileCoord0F, normalize(TileCoord1F - TileCoord0F));
+			const int2 EndCoord = (int2)floor(TileCoord1F);
+
+			for (int DDAIt = 0; DDAIt < DDA_MAX_ITERATIONS; ++DDAIt)
 			{
-				if (SegIndex[i] != 0xffffffff)
+				const int2 TileCoord = (int2)floor(DDAContext.TileCoord);
+
+				BRANCH
+				if (NearestDepth > VisTileDepthGrid[TileCoord])
 				{
-					uint PackedTiles = OutVisTileBinningGrid[uint3(TileCoord[i], BinID2)];
-					uint CurTile = (PackedTiles & 0xffff);
-					uint PrevTile = ((PackedTiles >> 16) & 0xffff);
+					const uint PackedTiles = OutVisTileBinningGrid[uint3(TileCoord, TileAllocInfoLayerIdx)];
+					const uint CurTile = (PackedTiles & 0xffff);
+					const uint PrevTile = ((PackedTiles >> 16) & 0xffff);
 
-					if (SegIndex[i] >= LoadOutVisTileData(CurTile, VT_NearestDepth))
-					{
-						uint PrimIndex = LoadOutVisTileData(CurTile, VT_PrimOffset) + (SegIndex[i] % 1024);
-						OutVisTilePrims[PrimIndex] = PrimID;
+					// Currently we need this to get our write position, but maybe there is a cheaper way to keep track of that?
+					uint OldTileSegmentCount;
+					InterlockedAdd(OutVisTileBinningGrid[uint3(TileCoord, TmpSegmentCountLayerIdx)], 1, OldTileSegmentCount);
+
+					const bool bWriteToCurTile = OldTileSegmentCount >= LoadOutVisTileData(CurTile, VT_MinWriteIndex);
+					const uint LocalWritePos = OldTileSegmentCount % 1024;
+					const uint WritePos = (bWriteToCurTile ? CurTile : PrevTile) * 1024 + LocalWritePos;
 
-						if ((SegIndex[i] + 1) == OutVisTileBinningGrid[uint3(TileCoord[i], BinID)])
+					OutVisTilePrims[WritePos] = PrimID;
+
+					BRANCH
+					if (bWriteToCurTile)
+					{
+						if ((OldTileSegmentCount + 1) == OutVisTileBinningGrid[uint3(TileCoord, SegmentCountLayerIdx)])
 						{
-							StoreOutVisTileData(CurTile, VT_PrimCount, (SegIndex[i]==1023) ? 1024 : ((SegIndex[i] + 1) % 1024));
+							StoreOutVisTileData(CurTile, VT_PrimCount, (OldTileSegmentCount == 1023) ? 1024 : ((OldTileSegmentCount + 1) % 1024));
 						}
 					}
-					else
-					{
-						uint PrimIndex = LoadOutVisTileData(PrevTile, VT_PrimOffset) + (SegIndex[i] % 1024);
-						OutVisTilePrims[PrimIndex] = PrimID;
-					}
+				}
 
-					//InterlockedAdd(OutVisTileArgs[1], 1); //debug
+				if (all(TileCoord == EndCoord))
+				{
+					break;
 				}
+
+				DDAAdvance(DDAContext);
 			}
 		}
 	}
@@ -484,7 +616,7 @@ void RasterCS(uint DispatchThreadID : SV_DispatchThreadID, uint GroupThreadID :
 
 		bool bTileValid = (group_TileIndex < group_TileNum) ? true : false;
 
-		uint PrimOffset = LoadRWVisTileData(group_TileIndex, VT_PrimOffset);
+		uint PrimOffset = group_TileIndex * 1024;
 		uint PrimCount = LoadRWVisTileData(group_TileIndex, VT_PrimCount);
 
 		uint PackedCoord = LoadRWVisTileData(group_TileIndex, VT_Coord);
@@ -670,7 +802,7 @@ void PrintBool(inout FShaderPrintContext Ctx, bool In)
 }
 
 Texture2D<uint>						VisTileDepthGrid;
-Texture3D<uint> 					VisTileBinningGrid;
+Texture2DArray<uint> 				VisTileBinningGrid;
 //Buffer<uint>						OutVisTilePrims;
 Buffer<uint>						VisTileArgs;
 //ByteAddressBuffer					OutVisTileData;
@@ -779,5 +911,5 @@ void MainCS(uint3 ThreadId : SV_DispatchThreadID)
 		}
 	}
 }
-#endif //SHADER_RASTERCOMPUTE_RASTER
+#endif //SHADER_RASTERCOMPUTE_DEBUG
 
diff --git a/Engine/Source/Runtime/Renderer/Private/HairStrands/HairStrandsVisibility.cpp b/Engine/Source/Runtime/Renderer/Private/HairStrands/HairStrandsVisibility.cpp
--- a/Engine/Source/Runtime/Renderer/Private/HairStrands/HairStrandsVisibility.cpp
+++ b/Engine/Source/Runtime/Renderer/Private/HairStrands/HairStrandsVisibility.cpp
@@ -55,7 +55,11 @@ static float GHairStrandsViewHairCountDepthDistanceThreshold = 30.f;
 static FAutoConsoleVariableRef CVarHairStrandsViewHairCountDepthDistanceThreshold(TEXT("r.HairStrands.Visibility.HairCount.DistanceThreshold"), GHairStrandsViewHairCountDepthDistanceThreshold, TEXT("Distance threshold defining if opaque depth get injected into the 'view-hair-count' buffer."));
 
 int32 GHairVisibilityComputeRaster_Culling = 0;
-int32 GHairVisibilityComputeRaster_MaxTiles = 8192;
+// This value was previously at 8192 but extending the binning algorithm to support segments longer than three tiles revealed a problem: Each binner allocates and writes its own tiles,
+// which often leads to lots of partially filled tiles for the same position in the tile grid. Without increasing the tile limit, we quickly run out of tiles to allocate when we get close
+// to a groom and the strand count and segment length increase. This new value was found empirically.
+// TODO: Find an efficient alternative algorithm where the binners output into shared tiles.
+int32 GHairVisibilityComputeRaster_MaxTiles = 65536;
 int32 GHairVisibilityComputeRaster_TileSize = 32;
 int32 GHairVisibilityComputeRaster_NumBinners = 32;
 int32 GHairVisibilityComputeRaster_NumRasterizers = 256;
@@ -2816,9 +2820,9 @@ class FVisiblityRasterComputeBinningCS : public FGlobalShader
 		SHADER_PARAMETER(uint32, HairMaterialId)
 		SHADER_PARAMETER_STRUCT_INCLUDE(FHairStrandsInstanceParameters, HairInstance)
 		SHADER_PARAMETER_RDG_TEXTURE(Texture2D, SceneDepthTexture)
-		SHADER_PARAMETER_RDG_TEXTURE(Texture3D, VisTileBinningGrid)
+		SHADER_PARAMETER_RDG_TEXTURE(Texture2DArray, VisTileBinningGrid)
 		SHADER_PARAMETER_RDG_BUFFER_UAV(RWBuffer, OutVisTilePrims)
-		SHADER_PARAMETER_RDG_TEXTURE_UAV(RWTexture3D, OutVisTileBinningGrid)
+		SHADER_PARAMETER_RDG_TEXTURE_UAV(RWTexture2DArray, OutVisTileBinningGrid)
 		SHADER_PARAMETER_RDG_TEXTURE(Texture2D, VisTileDepthGrid)
 		SHADER_PARAMETER_RDG_BUFFER_UAV(RWBuffer, OutVisTileArgs)
 		SHADER_PARAMETER_RDG_BUFFER_UAV(RWByteAddressBuffer, OutVisTileData)
@@ -2895,7 +2899,7 @@ class FVisiblityRasterComputeDebugCS : public FGlobalShader
 		SHADER_PARAMETER(uint32, PrimitiveInfoIndex)
 		SHADER_PARAMETER(uint32, TotalPrimitiveInfoCount)
 		SHADER_PARAMETER_RDG_TEXTURE(Texture2D, VisTileDepthGrid)
-		SHADER_PARAMETER_RDG_TEXTURE(Texture3D, VisTileBinningGrid)
+		SHADER_PARAMETER_RDG_TEXTURE(Texture2DArray, VisTileBinningGrid)
 		SHADER_PARAMETER_RDG_BUFFER_SRV(Buffer, VisTileArgs)
 		SHADER_PARAMETER_STRUCT_REF(FViewUniformShaderParameters, ViewUniformBuffer)
 		SHADER_PARAMETER_STRUCT_INCLUDE(ShaderPrint::FShaderParameters, ShaderPrintParameters)
@@ -2950,15 +2954,15 @@ static FRasterComputeOutput AddVisibilityComputeRasterPass(
 	AddClearUAVPass(GraphBuilder, DepthCovTextureUAV, ClearValues);
 	AddClearUAVPass(GraphBuilder, HairCountTextureUAV, ClearValues);
 
-	//set up buffers for binning and raster
-	const uint32 MaxTiles = FMath::Min(FMath::Max(GHairVisibilityComputeRaster_MaxTiles, 1024), 65536);
+	// Set up buffers for binning and raster.
 
+	// See the comment on the GHairVisibilityComputeRaster_MaxTiles declaration for an explanation for the large upper bound.
+	const uint32 MaxTiles = FMath::Min(FMath::Max(GHairVisibilityComputeRaster_MaxTiles, 1024), 262144);
 	const uint32 TileSize = FMath::Min(FMath::Max(GHairVisibilityComputeRaster_TileSize, 8),32);
 	const uint32 NumBinners = FMath::Min(FMath::Max(GHairVisibilityComputeRaster_NumBinners, 1), 256);
 	const uint32 NumRasterizers = FMath::Min(FMath::Max(GHairVisibilityComputeRaster_NumRasterizers, 1), 1024);
 
 	const FIntPoint TileGridRes = FIntPoint((InResolution.X + (TileSize-1)) / TileSize, ((InResolution.Y + (TileSize - 1)) / TileSize)*1);
-	const FIntVector TileGridRes3D = FIntVector(TileGridRes.X, TileGridRes.Y, NumBinners * 2);
 
 	FRDGBufferRef OutVisTilePrims = nullptr;
 	FRDGTextureRef OutVisTileBinningGrid = nullptr;
@@ -2976,7 +2980,7 @@ static FRasterComputeOutput AddVisibilityComputeRasterPass(
 	
 	OutVisTilePrims = GraphBuilder.CreateBuffer(FRDGBufferDesc::CreateBufferDesc(sizeof(uint32), MaxTiles*1024), TEXT("Hair.OutVisTilePrims"));
 
-	FRDGTextureDesc DescBinningGrid = FRDGTextureDesc::Create3D(TileGridRes3D, PF_R32_UINT, FClearValueBinding::None, TexCreate_UAV | TexCreate_ShaderResource);
+	FRDGTextureDesc DescBinningGrid = FRDGTextureDesc::Create2DArray(TileGridRes, PF_R32_UINT, FClearValueBinding::None, TexCreate_UAV | TexCreate_ShaderResource, NumBinners * 3);
 	OutVisTileBinningGrid = GraphBuilder.CreateTexture(DescBinningGrid, TEXT("Hair.VisTileBinningGrid"));
 
 	FRDGTextureDesc DescDepthGrid = FRDGTextureDesc::Create2D(TileGridRes, PF_R32_UINT, FClearValueBinding::None, TexCreate_UAV | TexCreate_ShaderResource);
