Fix light flickering issue on mobile vulkan on Cloth and Hair.

diff --git a/Engine/Shaders/Private/DeferredLightingCommon.ush b/Engine/Shaders/Private/DeferredLightingCommon.ush
--- a/Engine/Shaders/Private/DeferredLightingCommon.ush
+++ b/Engine/Shaders/Private/DeferredLightingCommon.ush
@@ -383,7 +383,7 @@ FLightAccumulator AccumulateDynamicLighting(
 	
 	float3 L = LightData.Direction;	// Already normalized
 	float3 ToLight = L;
-	half3 MaskedLightColor = LightData.Color;
+	float3 MaskedLightColor = LightData.Color;
 	float LightMask = 1;
 	if (LightData.bRadialLight)
 	{
diff --git a/Engine/Shaders/Private/ShadingModels.ush b/Engine/Shaders/Private/ShadingModels.ush
--- a/Engine/Shaders/Private/ShadingModels.ush
+++ b/Engine/Shaders/Private/ShadingModels.ush
@@ -660,7 +660,11 @@ void GetProfileDualSpecular(uint SubsurfaceProfileInt, half Roughness, half Opac
 FDirectLighting SubsurfaceProfileBxDF( FGBufferData GBuffer, half3 N, half3 V, half3 L, float Falloff, half NoL, FAreaLight AreaLight, FShadowTerms Shadow )
 {
 	BxDFContext Context;
+#if SHADING_PATH_MOBILE
+	InitMobile(Context, N, V, L, NoL);
+#else
 	Init( Context, N, V, L );
+#endif
 	SphereMaxNoH( Context, AreaLight.SphereSinAlpha, true );
 	Context.NoV = saturate( abs( Context.NoV ) + 1e-5 );
 
@@ -737,7 +741,11 @@ FDirectLighting ClothBxDF( FGBufferData GBuffer, half3 N, half3 V, half3 L, floa
 	const float  Cloth		= saturate(GBuffer.CustomData.a);
 
 	BxDFContext Context;
+#if SHADING_PATH_MOBILE
+	InitMobile(Context, N, V, L, NoL);
+#else
 	Init( Context, N, V, L );
+#endif
 	SphereMaxNoH( Context, AreaLight.SphereSinAlpha, true );
 	Context.NoV = saturate( abs( Context.NoV ) + 1e-5 );
 
@@ -753,6 +761,9 @@ FDirectLighting ClothBxDF( FGBufferData GBuffer, half3 N, half3 V, half3 L, floa
 	// Cloth - Asperity Scattering - Inverse Beckmann Layer
 	float D2 = D_InvGGX( Pow4( GBuffer.Roughness ), Context.NoH );
 	float Vis2 = Vis_Cloth( Context.NoV, NoL );
+#if SHADING_PATH_MOBILE
+	Vis2 = saturate(Vis2);
+#endif
 	float3 F2 = F_Schlick( FuzzColor, Context.VoH );
 	float3 Spec2 = AreaLight.FalloffColor * (Falloff * NoL) * (D2 * Vis2) * F2;
 	
@@ -838,7 +849,11 @@ FDirectLighting EyeBxDF( FGBufferData GBuffer, half3 N, half3 V, half3 L, float
 #endif
 
 	BxDFContext Context;
+#if SHADING_PATH_MOBILE
+	InitMobile(Context, N, V, L, NoL);
+#else
 	Init( Context, N, V, L );
+#endif
 	SphereMaxNoH( Context, AreaLight.SphereSinAlpha, false );
 	Context.NoV = saturate( abs( Context.NoV ) + 1e-5 );
 	const bool bIsRect = IsRectLight(AreaLight);
