Fixed: Missing DF shadows on Switch
Added a workaround in ShadowObjectCullVS for bad codegen by a Switch shader compiler

diff --git a/Engine/Shaders/Private/DistanceFieldShadowing.usf b/Engine/Shaders/Private/DistanceFieldShadowing.usf
--- a/Engine/Shaders/Private/DistanceFieldShadowing.usf
+++ b/Engine/Shaders/Private/DistanceFieldShadowing.usf
@@ -224,13 +224,14 @@ void ShadowObjectCullVS(
 	FDFObjectData DFObjectData = LoadDFObjectData(ObjectIndex);
 	float4x4 ObjectOBBToTranslatedWorld = LWCMultiplyTranslation(DFObjectData.VolumeToWorld, PrimaryView.PreViewTranslation);
 
-	float ObjectScaleX = length(ObjectOBBToTranslatedWorld[0].xyz);
-	float ObjectScaleY = length(ObjectOBBToTranslatedWorld[1].xyz);
-	float ObjectScaleZ = length(ObjectOBBToTranslatedWorld[2].xyz);
-
-	ObjectOBBToTranslatedWorld[0].xyz *= (1.0f + MinExpandRadius / ObjectScaleX);
-	ObjectOBBToTranslatedWorld[1].xyz *= (1.0f + MinExpandRadius / ObjectScaleY);
-	ObjectOBBToTranslatedWorld[2].xyz *= (1.0f + MinExpandRadius / ObjectScaleZ);
+	float ObjectScaleX = (1.0f + MinExpandRadius / length(ObjectOBBToTranslatedWorld[0].xyz));
+	float ObjectScaleY = (1.0f + MinExpandRadius / length(ObjectOBBToTranslatedWorld[1].xyz));
+	float ObjectScaleZ = (1.0f + MinExpandRadius / length(ObjectOBBToTranslatedWorld[2].xyz));
+
+	// This is written without [].xyz *= Scale; to workaround a bad codegen on Switch
+	ObjectOBBToTranslatedWorld[0] = float4(ObjectOBBToTranslatedWorld[0].xyz * ObjectScaleX, ObjectOBBToTranslatedWorld[0].w);
+	ObjectOBBToTranslatedWorld[1] = float4(ObjectOBBToTranslatedWorld[1].xyz * ObjectScaleY, ObjectOBBToTranslatedWorld[1].w);
+	ObjectOBBToTranslatedWorld[2] = float4(ObjectOBBToTranslatedWorld[2].xyz * ObjectScaleZ, ObjectOBBToTranslatedWorld[2].w);
 
 	float3 TranslatedWorldPosition = mul(float4(InPosition.xyz, 1.0f), ObjectOBBToTranslatedWorld).xyz;
 #else
