Ray Traced Shadows: Implement support for Source Length

This attribute of point and spot lights was previously ignored. Additionally, unify treatment of point and spots so their shadows behave consistently with one another (as is already the case in path tracing).

diff --git a/Engine/Shaders/Private/RayTracing/RayTracingCapsuleLight.ush b/Engine/Shaders/Private/RayTracing/RayTracingCapsuleLight.ush
new file mode 100644
--- /dev/null
+++ b/Engine/Shaders/Private/RayTracing/RayTracingCapsuleLight.ush
@@ -0,0 +1,39 @@
+// Copyright Epic Games, Inc. All Rights Reserved.
+
+#pragma once
+
+#include "../LightShaderParameters.ush"
+#include "../CapsuleLightSampling.ush"
+
+bool GenerateCapsuleLightOcclusionRayWithSolidAngleSampling(
+	FLightShaderParameters LightParameters,
+	float3 TranslatedWorldPosition,
+	float3 WorldNormal,
+	float2 RandSample,
+	out float3 RayOrigin,
+	out float3 RayDirection,
+	out float RayTMin,
+	out float RayTMax,
+	out float RayPdf
+)
+{
+	float3 LightDirection = LightParameters.TranslatedWorldPosition - TranslatedWorldPosition;
+	float RayLength2 = dot(LightDirection, LightDirection);
+	float Radius = max(LightParameters.SourceRadius, 0.01); // clamp radius to a small value since the code below doesn't handle pure lines
+	float Radius2 = Radius * Radius;
+	float SourceLength = LightParameters.SourceLength;
+	float3 Axis = LightParameters.Tangent;
+
+	// compute a bounding representation for the capsule
+	FCapsuleSphericalBounds CapsuleBounds = CapsuleGetSphericalBounds(LightDirection, Axis, Radius, SourceLength);
+	// sample a direction toward these bounds
+	float4 Result = SampleCapsuleBounds(CapsuleBounds, RandSample);
+	// compute distance to capsule (and check that we hit it)
+	float Distance = CapsuleTest(Result.xyz, LightDirection, Axis, Radius2, SourceLength);
+	RayOrigin = TranslatedWorldPosition;
+	RayDirection = Result.xyz;
+	RayTMin = 0;
+	RayTMax = Distance;
+	RayPdf = Result.w;
+	return Distance > 0; // sample is only valid if we hit the light
+}
diff --git a/Engine/Shaders/Private/RayTracing/RayTracingOcclusionRGS.usf b/Engine/Shaders/Private/RayTracing/RayTracingOcclusionRGS.usf
--- a/Engine/Shaders/Private/RayTracing/RayTracingOcclusionRGS.usf
+++ b/Engine/Shaders/Private/RayTracing/RayTracingOcclusionRGS.usf
@@ -20,8 +20,8 @@
 #include "RayTracingDirectionalLight.ush"
 #include "RayTracingRectLight.ush"
 #include "RayTracingSphereLight.ush"
-#include "RayTracingSpotLight.ush"
 #include "RayTracingPointLight.ush"
+#include "RayTracingCapsuleLight.ush"
 
 /** Light types */
 #define LIGHT_TYPE_DIRECTIONAL 0
@@ -100,24 +100,22 @@ bool GenerateOcclusionRay(
 			/* out */ RayTMin,
 			/* out */ RayTMax);
 	}
-	#elif LIGHT_TYPE == LIGHT_TYPE_POINT
+	#elif LIGHT_TYPE == LIGHT_TYPE_POINT || LIGHT_TYPE == LIGHT_TYPE_SPOT
 	{
-		if (LightParameters.SourceRadius == 0)
-		{
-			return GeneratePointLightOcclusionRay(
-				LightParameters,
-				TranslatedWorldPosition, WorldNormal,
-				RandSample,
-				/* out */ RayOrigin,
-				/* out */ RayDirection,
-				/* out */ RayTMin,
-				/* out */ RayTMax);
-		}
-		else
+		#if LIGHT_TYPE == LIGHT_TYPE_SPOT
+			// before generating a shadow ray, make sure we are inside the cone of the light
+			float3 LightDirection = normalize(LightParameters.TranslatedWorldPosition - TranslatedWorldPosition);
+			float CosAngle = LightParameters.SpotAngles.x;
+			if (!(dot(LightDirection, LightParameters.Direction) >= CosAngle))
+			{
+				// outside the cone of the light, skip all work
+				return false;
+			}
+		#endif
+		float RayPdf;
+		if (LightParameters.SourceLength > 0.0)
 		{
-			float RayPdf;
-			//return GenerateSphereLightOcclusionRayWithAreaSampling(
-			return GenerateSphereLightOcclusionRayWithSolidAngleSampling(
+			return GenerateCapsuleLightOcclusionRayWithSolidAngleSampling(
 				LightParameters,
 				TranslatedWorldPosition, WorldNormal,
 				RandSample,
@@ -127,17 +125,15 @@ bool GenerateOcclusionRay(
 				/* out */ RayTMax,
 				/* out */ RayPdf);
 		}
-	}
-	#elif LIGHT_TYPE == LIGHT_TYPE_SPOT
-	{
-		return GenerateSpotLightOcclusionRay(
+		return GenerateSphereLightOcclusionRayWithSolidAngleSampling(
 			LightParameters,
 			TranslatedWorldPosition, WorldNormal,
 			RandSample,
 			/* out */ RayOrigin,
 			/* out */ RayDirection,
 			/* out */ RayTMin,
-			/* out */ RayTMax);
+			/* out */ RayTMax,
+			/* out */ RayPdf);
 	}
 	#elif LIGHT_TYPE == LIGHT_TYPE_RECT
 	{
