Minifont now always print with character shadow to be more readable on bright background.

diff --git a/Engine/Shaders/Private/MiniFontCommon.ush b/Engine/Shaders/Private/MiniFontCommon.ush
--- a/Engine/Shaders/Private/MiniFontCommon.ush
+++ b/Engine/Shaders/Private/MiniFontCommon.ush
@@ -70,7 +70,13 @@ void PrintCharacter(int2 PixelPos, inout float3 OutColor, float3 FontColor, inou
 {
 	uint2 Rel = (uint2)(PixelPos - LeftTop);
 
-	FLATTEN if (Rel.x < NATIVE_CHARACTER_RES && Rel.y < NATIVE_CHARACTER_RES)
+	// Background: black character as shadow
+	FLATTEN if (Rel.x < (NATIVE_CHARACTER_RES + 1) && Rel.y < (NATIVE_CHARACTER_RES + 1))
+	{
+		OutColor = lerp(OutColor, float3(0.0, 0.0, 0.0), SampleMiniFont(CharacterID, Rel - uint2(1, 1)));
+	}
+	// Foreground: colored character
+	FLATTEN if (Rel.x < (NATIVE_CHARACTER_RES) && Rel.y < (NATIVE_CHARACTER_RES))
 	{
 		OutColor = lerp(OutColor, FontColor, SampleMiniFont(CharacterID, Rel));
 	}
diff --git a/Engine/Shaders/Private/PostProcessVisualizeHDR.usf b/Engine/Shaders/Private/PostProcessVisualizeHDR.usf
--- a/Engine/Shaders/Private/PostProcessVisualizeHDR.usf
+++ b/Engine/Shaders/Private/PostProcessVisualizeHDR.usf
@@ -207,7 +207,7 @@ void MainPS(noperspective float4 UVAndScreenPos : TEXCOORD0, float4 SvPosition :
 	{
 		const float PixelDx = abs(PixelPos.x - OutputViewportCenter.x);
 		const float PixelDy = abs(PixelPos.y - OutputViewportCenter.y);
-		if (PixelDx < 130 && PixelDy < 40)
+		if (PixelDx < 130 && PixelDy < 42)
 		{
 			float3 LuminanceAvg = 0.0f;
 			float3 SceneColorAvg = 0.0f;
