Strata - added material graph preview on node.
This is simply done by an added PreviewColor on the FStrataData that flows between the node to be able to work with parameter blending or StrataTree BSDF composition. Can be updated later to anything else.

diff --git a/Engine/Shaders/Private/BasePassPixelShader.usf b/Engine/Shaders/Private/BasePassPixelShader.usf
--- a/Engine/Shaders/Private/BasePassPixelShader.usf
+++ b/Engine/Shaders/Private/BasePassPixelShader.usf
@@ -1553,6 +1553,8 @@ void FPixelShaderInOut_MainPS(
 			StrataMaterialExport = float4(EmissiveLuminance, 1.0);
 		#elif STRATA_MATERIAL_EXPORT_TYPE == STRATA_MATERIAL_EXPORT_TRANSMITTANCE
 			StrataMaterialExport = float4(PixelMaterialInputs.OpacityMask.xxx, 1.0f); // when not masked, OpacityMask default to 1.
+		#elif STRATA_MATERIAL_EXPORT_TYPE == STRATA_MATERIAL_EXPORT_MATERIAL_PREVIEW
+			StrataMaterialExport = float4(PixelMaterialInputs.EmissiveColor, 1.0f);
 		#endif
 	#elif STRATA_MATERIAL_EXPORT_FROM_TRANSLUCENT
 		#if STRATA_USES_CONVERSION_FROM_LEGACY
@@ -1568,6 +1570,8 @@ void FPixelShaderInOut_MainPS(
 				#else
 					StrataMaterialExport = float4(BaseColorPostCoverage + saturate(1.0f - Coverage), 1.0f);
 				#endif
+			#elif STRATA_MATERIAL_EXPORT_TYPE == STRATA_MATERIAL_EXPORT_MATERIAL_PREVIEW
+				StrataMaterialExport = float4(PixelMaterialInputs.EmissiveColor, 1.0f);
 			#endif
 		#else
 			#if	  STRATA_MATERIAL_EXPORT_TYPE == STRATA_MATERIAL_EXPORT_DIFFUSE_COLOR
@@ -1578,6 +1582,8 @@ void FPixelShaderInOut_MainPS(
 				StrataMaterialExport = float4(EmissiveLuminance, 1.0);
 			#elif STRATA_MATERIAL_EXPORT_TYPE == STRATA_MATERIAL_EXPORT_TRANSMITTANCE
 				StrataMaterialExport = float4(TransmittancePreCoverage * Coverage + (1.0f - Coverage), 1.0f);
+			#elif STRATA_MATERIAL_EXPORT_TYPE == STRATA_MATERIAL_EXPORT_MATERIAL_PREVIEW
+				StrataMaterialExport = float4(PixelMaterialInputs.EmissiveColor, 1.0f);
 			#endif
 		#endif
 	#else
@@ -1755,6 +1761,8 @@ void FPixelShaderInOut_MainPS(
 	StrataMaterialExport = float4(Color, 1.0);
 #elif STRATA_MATERIAL_EXPORT_TYPE == STRATA_MATERIAL_EXPORT_TRANSMITTANCE
 	StrataMaterialExport = float4(PixelMaterialInputs.OpacityMask.xxx, 1.0f); // when not masked, OpacityMask default to 1.
+#elif STRATA_MATERIAL_EXPORT_TYPE == STRATA_MATERIAL_EXPORT_MATERIAL_PREVIEW
+	StrataMaterialExport = float4(PixelMaterialInputs.EmissiveColor, 1.0f);
 #endif
 
 #else // STRATA_TRANSLUCENT_FORWARD || STRATA_FORWARD_SHADING || STRATA_MATERIAL_EXPORT_FROM_TRANSLUCENT
@@ -1777,6 +1785,8 @@ void FPixelShaderInOut_MainPS(
 		#else
 		StrataMaterialExport = float4(DualBlendSurfaceLuminancePostCoverage + saturate(1.0 - DualBlendSurfaceCoverage.xxx), 1.0f);
 		#endif
+	#elif STRATA_MATERIAL_EXPORT_TYPE == STRATA_MATERIAL_EXPORT_MATERIAL_PREVIEW
+		StrataMaterialExport = float4(PixelMaterialInputs.EmissiveColor, 1.0f);
 	#endif
 #else // STRATA_USES_CONVERSION_FROM_LEGACY
 	// The legacy unlit to lightmass conversion is weird and very specialized.
@@ -1788,6 +1798,8 @@ void FPixelShaderInOut_MainPS(
 		StrataMaterialExport = float4(DualBlendSurfaceLuminancePostCoverage, 1.0);
 	#elif STRATA_MATERIAL_EXPORT_TYPE == STRATA_MATERIAL_EXPORT_TRANSMITTANCE
 		StrataMaterialExport = float4(DualBlendSurfaceTransmittancePreCoverage * DualBlendSurfaceCoverage + (1.0 - DualBlendSurfaceCoverage), 1.0f);
+	#elif STRATA_MATERIAL_EXPORT_TYPE == STRATA_MATERIAL_EXPORT_MATERIAL_PREVIEW
+		StrataMaterialExport = float4(PixelMaterialInputs.EmissiveColor, 1.0f);
 	#endif
 #endif // STRATA_USES_CONVERSION_FROM_LEGACY
 
diff --git a/Engine/Shaders/Private/Definitions.usf b/Engine/Shaders/Private/Definitions.usf
--- a/Engine/Shaders/Private/Definitions.usf
+++ b/Engine/Shaders/Private/Definitions.usf
@@ -236,6 +236,7 @@
 #define STRATA_MATERIAL_EXPORT_NORMAL						2
 #define STRATA_MATERIAL_EXPORT_EMISSIVE						3
 #define STRATA_MATERIAL_EXPORT_TRANSMITTANCE				4
+#define STRATA_MATERIAL_EXPORT_MATERIAL_PREVIEW				5
 
 #define STRATA_MATERIAL_EXPORT_CONTEXT_OPAQUE				0
 #define STRATA_MATERIAL_EXPORT_CONTEXT_TRANSLUCENT			1
diff --git a/Engine/Shaders/Private/Strata/Strata.ush b/Engine/Shaders/Private/Strata/Strata.ush
--- a/Engine/Shaders/Private/Strata/Strata.ush
+++ b/Engine/Shaders/Private/Strata/Strata.ush
@@ -1208,6 +1208,10 @@ struct FStrataData
 
 	int  OperatorIndex;		// Used as entry point ot the material graph with multiple BSDF and operators.
 
+#if USE_DEVELOPMENT_SHADERS
+	float3 PreviewColor;	// Used as preview in the material graph within the editor.
+#endif
+
 #if STRATA_LEGACY_PREMULT_ALPHA_OVERRIDE
 	float PreMultipliedAlphaOverrideCoverage;
 #endif
@@ -1388,6 +1392,17 @@ FStrataData GetStrataSlabBSDF(
 	StrataData.InlinedBSDF.Coverage = 1.0f;
 #endif
 
+#if USE_DEVELOPMENT_SHADERS
+	if (UseMetalness > 0.0f)
+	{
+		StrataData.PreviewColor = BaseColor;
+	}
+	else
+	{
+		StrataData.PreviewColor = lerp(DiffuseAlbedo, F0, F0RGBToMetallic(F0));
+	}
+#endif
+
 	return StrataData;
 }
 
@@ -1405,6 +1420,10 @@ FStrataData GetStrataVolumeFogCloudBSDF(float3 Albedo, float3 Extinction, float3
 	StrataData.InlinedBSDF.Coverage = 1.0f;
 #endif
 
+#if USE_DEVELOPMENT_SHADERS
+	StrataData.PreviewColor			= Albedo;
+#endif
+
 	return StrataData;
 }
 
@@ -1420,6 +1439,10 @@ FStrataData GetStrataUnlitBSDF(float3 Emissive, float3 TransmittanceColor)
 	StrataData.InlinedBSDF.Coverage = 1.0f;
 #endif
 
+#if USE_DEVELOPMENT_SHADERS
+	StrataData.PreviewColor			= Emissive;
+#endif
+
 	return StrataData;
 }
 
@@ -1450,6 +1473,10 @@ FStrataData GetStrataHairBSDF(float3 BaseColor, float Scatter, float Specular, f
 	StrataData.InlinedBSDF.Coverage = 1.0f;
 #endif
 
+#if USE_DEVELOPMENT_SHADERS
+	StrataData.PreviewColor			= BaseColor;
+#endif
+
 	return StrataData;
 }
 
@@ -1475,6 +1502,10 @@ FStrataData GetStrataEyeBSDF(float3 DiffuseAlbedo, float Roughness, float IrisMa
 	StrataData.InlinedBSDF.Coverage = 1.0f;
 #endif
 
+#if USE_DEVELOPMENT_SHADERS
+	StrataData.PreviewColor			= DiffuseAlbedo;
+#endif
+
 	return StrataData;
 }
 
@@ -1504,6 +1535,10 @@ FStrataData GetStrataSingleLayerWaterBSDF(
 	StrataData.InlinedBSDF.Coverage = 1.0f;
 #endif
 
+#if USE_DEVELOPMENT_SHADERS
+	StrataData.PreviewColor			= lerp(WaterAlbedo, BaseColor, TopMaterialOpacity);
+#endif
+
 	return StrataData;
 }
 
@@ -1871,10 +1906,12 @@ float3x3 StrataGetBSDFSharedBasis(in FStrataPixelHeader StrataPixelHeader, in FS
 FStrataData StrataAdd(FStrataData A, FStrataData B, inout FStrataTree StrataTree, int OperatorIndex, uint MaxDistanceFromLeaves)
 {
 	FStrataData StrataData = GetInitialisedStrataData();
-
-	StrataTree.OperatorCount++;
 	StrataData.OperatorIndex = OperatorIndex;
+#if USE_DEVELOPMENT_SHADERS
+	StrataData.PreviewColor = A.PreviewColor + B.PreviewColor;
+#endif
 
+	StrataTree.OperatorCount++;
 	StrataTree.Operators[OperatorIndex]							= GetInitialisedStrataOperator();
 	StrataTree.Operators[OperatorIndex].MaxDistanceFromLeaves	= MaxDistanceFromLeaves;
 	StrataTree.Operators[OperatorIndex].Type					= STRATA_OPERATOR_ADD;
@@ -1892,10 +1929,12 @@ FStrataData StrataWeight(FStrataData A, float Weight, inout FStrataTree StrataTr
 	const float SafeWeight = saturate(Weight);
 
 	FStrataData StrataData = GetInitialisedStrataData();
-
-	StrataTree.OperatorCount++;
 	StrataData.OperatorIndex = OperatorIndex;
+#if USE_DEVELOPMENT_SHADERS
+	StrataData.PreviewColor = A.PreviewColor * Weight;
+#endif
 
+	StrataTree.OperatorCount++;
 	StrataTree.Operators[OperatorIndex]							= GetInitialisedStrataOperator();
 	StrataTree.Operators[OperatorIndex].MaxDistanceFromLeaves	= MaxDistanceFromLeaves;
 	StrataTree.Operators[OperatorIndex].Type					= STRATA_OPERATOR_WEIGHT;
@@ -1913,10 +1952,12 @@ FStrataData StrataThinFilm(FStrataData A, float Thickness, float IOR, inout FStr
 	const float SafeIOR = clamp(IOR, 1.01f, 3.0f);
 
 	FStrataData StrataData = GetInitialisedStrataData();
-
-	StrataTree.OperatorCount++;
 	StrataData.OperatorIndex = OperatorIndex;
+#if USE_DEVELOPMENT_SHADERS
+	StrataData.PreviewColor = A.PreviewColor;
+#endif
 
+	StrataTree.OperatorCount++;
 	StrataTree.Operators[OperatorIndex] = GetInitialisedStrataOperator();
 	StrataTree.Operators[OperatorIndex].MaxDistanceFromLeaves = MaxDistanceFromLeaves;
 	StrataTree.Operators[OperatorIndex].Type = STRATA_OPERATOR_THINFILM;
@@ -1935,10 +1976,12 @@ FStrataData StrataHorizontalMixing(FStrataData Background, FStrataData Foregroun
 	const float BackgroundMixFactor = 1.0 - ForegroundMixFactor;
 
 	FStrataData StrataData = GetInitialisedStrataData();
+	StrataData.OperatorIndex = OperatorIndex;
+#if USE_DEVELOPMENT_SHADERS
+	StrataData.PreviewColor = lerp(Background.PreviewColor, Foreground.PreviewColor, Mix);
+#endif
 
 	StrataTree.OperatorCount++;
-	StrataData.OperatorIndex = OperatorIndex;
-	
 	StrataTree.Operators[OperatorIndex]							= GetInitialisedStrataOperator();
 	StrataTree.Operators[OperatorIndex].MaxDistanceFromLeaves	= MaxDistanceFromLeaves;
 	StrataTree.Operators[OperatorIndex].Type					= STRATA_OPERATOR_HORIZONTAL;
@@ -1955,10 +1998,12 @@ FStrataData StrataHorizontalMixing(FStrataData Background, FStrataData Foregroun
 FStrataData StrataVerticalLayering(FStrataData Top, FStrataData Base, inout FStrataTree StrataTree, int OperatorIndex, uint MaxDistanceFromLeaves)
 {
 	FStrataData StrataData = GetInitialisedStrataData();
+	StrataData.OperatorIndex = OperatorIndex;
+#if USE_DEVELOPMENT_SHADERS
+	StrataData.PreviewColor = lerp(Top.PreviewColor, Base.PreviewColor, 0.5);
+#endif
 
 	StrataTree.OperatorCount++;
-	StrataData.OperatorIndex = OperatorIndex;
-	
 	StrataTree.Operators[OperatorIndex]							= GetInitialisedStrataOperator();
 	StrataTree.Operators[OperatorIndex].MaxDistanceFromLeaves	= MaxDistanceFromLeaves;
 	StrataTree.Operators[OperatorIndex].Type					= STRATA_OPERATOR_VERTICAL;
@@ -2099,6 +2144,10 @@ FStrataData StrataHorizontalMixingParameterBlending(FStrataData Background, FStr
 	ResultBSDF.Coverage =				lerp(BackgroundBSDF.Coverage,					ForegroundBSDF.Coverage,					HorizontalMixValue);
 #endif
 
+#if USE_DEVELOPMENT_SHADERS
+	Result.PreviewColor = lerp(Background.PreviewColor, Foreground.PreviewColor, HorizontalMixValue);
+#endif
+
 #undef ResultBSDF
 #undef ForegroundBSDF
 #undef BackgroundBSDF
@@ -2183,6 +2232,10 @@ FStrataData StrataAddParameterBlending(FStrataData A, FStrataData B, float AMixF
 	ResultBSDF.Coverage =				saturate(BBSDF.Coverage + ABSDF.Coverage);
 #endif
 
+#if USE_DEVELOPMENT_SHADERS
+	Result.PreviewColor = lerp(A.PreviewColor, B.PreviewColor, 0.5);
+#endif
+
 #undef ResultBSDF
 #undef ABSDF
 #undef BBSDF
@@ -2306,6 +2359,9 @@ FStrataData StrataVerticalLayeringParameterBlending(FStrataData Top, FStrataData
 {
 	FStrataData Result = GetInitialisedStrataData();
 	Result.InlinedBSDF = StrataVerticalLayeringParameterBlending(Top.InlinedBSDF, Base.InlinedBSDF, NewNormalIndex, NoV);
+#if USE_DEVELOPMENT_SHADERS
+	Result.PreviewColor = lerp(Top.PreviewColor, Base.PreviewColor, 0.5);
+#endif
 	return Result;
 }
 
diff --git a/Engine/Source/Editor/MaterialEditor/Private/MaterialEditor.cpp b/Engine/Source/Editor/MaterialEditor/Private/MaterialEditor.cpp
--- a/Engine/Source/Editor/MaterialEditor/Private/MaterialEditor.cpp
+++ b/Engine/Source/Editor/MaterialEditor/Private/MaterialEditor.cpp
@@ -171,6 +171,12 @@ static TAutoConsoleVariable<int32> CVarMaterialEdUseDevShaders(
 	TEXT("Toggles whether the material editor will use shaders that include extra overhead incurred by the editor. Material editor must be re-opened if changed at runtime."),
 	ECVF_RenderThreadSafe);
 
+static bool Editor_IsStrataEnabled()
+{
+	static const auto CVar = IConsoleManager::Get().FindTConsoleVariableDataInt(TEXT("r.Strata"));
+	return CVar && CVar->GetValueOnAnyThread() > 0;
+}
+
 ///////////////////////////
 // FMatExpressionPreview //
 ///////////////////////////
@@ -269,6 +275,12 @@ int32 FMatExpressionPreview::CompilePropertyAndSetMaterialProperty(EMaterialProp
 	// needs to be called in this function!!
 	Compiler->SetMaterialProperty(Property, OverrideShaderFrequency, bUsePreviousFrameTime);
 
+	if(Editor_IsStrataEnabled())
+	{
+		// Set the Strata export mode to material preview
+		Compiler->SetStrataMaterialExportType(SME_MaterialPreview, EStrataMaterialExportContext::SMEC_Opaque, 0);
+	}
+
 	int32 Ret = INDEX_NONE;
 
 	if( Property == MP_EmissiveColor && Expression.IsValid())
@@ -276,15 +288,7 @@ int32 FMatExpressionPreview::CompilePropertyAndSetMaterialProperty(EMaterialProp
 		// Hardcoding output 0 as we don't have the UI to specify any other output
 		const int32 OutputIndex = 0;
 		int32 PreviewCodeChunk = INDEX_NONE;
-		if (!Expression->Outputs.IsEmpty() && Expression->GetOutputType(OutputIndex) != MCT_Strata)
-		{
-			PreviewCodeChunk = Expression->CompilePreview(Compiler, OutputIndex);
-		}
-		else
-		{
-			// This is used for instance to prevent any Strata typed output to be used as preview, for instance for material functions.
-			PreviewCodeChunk = Compiler->Constant(0.0f);
-		}
+		PreviewCodeChunk = Expression->CompilePreview(Compiler, OutputIndex);
 
 		// Get back into gamma corrected space, as DrawTile does not do this adjustment.
 		Ret = Compiler->Power(Compiler->Max(PreviewCodeChunk, Compiler->Constant(0)), Compiler->Constant(1.f / 2.2f));
@@ -306,7 +310,11 @@ int32 FMatExpressionPreview::CompilePropertyAndSetMaterialProperty(EMaterialProp
 	}
 	else if (Property == MP_FrontMaterial)
 	{
-		Ret = Compiler->StrataCreateAndRegisterNullMaterial();
+		// Try to fetch the material interface to compile the front material.
+		UMaterial* ExprMat = Expression.IsValid() ? Expression->Material : nullptr;
+		FMaterialRenderProxy* MatProxy = ExprMat ? ExprMat->GetRenderProxy() : nullptr;
+		UMaterialInterface* MatInterface = MatProxy ? MatProxy->GetMaterialInterface() : nullptr;
+		return MatInterface ? MatInterface->CompileProperty(Compiler, MP_FrontMaterial) : Compiler->StrataCreateAndRegisterNullMaterial();
 	}
 	else
 	{
@@ -317,6 +325,20 @@ int32 FMatExpressionPreview::CompilePropertyAndSetMaterialProperty(EMaterialProp
 	return Compiler->ForceCast(Ret, FMaterialAttributeDefinitionMap::GetValueType(Property), MFCF_ExactMatch);
 }
 
+UMaterialInterface* FMatExpressionPreview::GetMaterialInterface() const
+{
+	UMaterial* ExprMat = Expression->Material;
+	if (ExprMat)
+	{
+		FMaterialRenderProxy* MatProxy = ExprMat->GetRenderProxy();
+		if (MatProxy)
+		{
+			return MatProxy->GetMaterialInterface();
+		}
+	}
+	return nullptr;
+}
+
 void FMatExpressionPreview::NotifyCompilationFinished()
 {
 	if (Expression.IsValid() && Expression->GraphNode)
diff --git a/Engine/Source/Editor/MaterialEditor/Private/MaterialEditor.h b/Engine/Source/Editor/MaterialEditor/Private/MaterialEditor.h
--- a/Engine/Source/Editor/MaterialEditor/Private/MaterialEditor.h
+++ b/Engine/Source/Editor/MaterialEditor/Private/MaterialEditor.h
@@ -139,6 +139,9 @@ public:
 		return Expression.Get();
 	}
 
+	// This material interface is solely needed for the translator to be able to parse the graph for the Strata tree.
+	virtual UMaterialInterface* GetMaterialInterface() const override;
+
 	virtual void NotifyCompilationFinished() override;
 
 	virtual const FMaterialCachedHLSLTree* GetCachedHLSLTree() const override;
diff --git a/Engine/Source/Runtime/Engine/Classes/Materials/MaterialExpressionStrata.h b/Engine/Source/Runtime/Engine/Classes/Materials/MaterialExpressionStrata.h
--- a/Engine/Source/Runtime/Engine/Classes/Materials/MaterialExpressionStrata.h
+++ b/Engine/Source/Runtime/Engine/Classes/Materials/MaterialExpressionStrata.h
@@ -23,6 +23,7 @@ extern int32 CompileStrataBlendFunction(FMaterialCompiler* Compiler, const int32
 ///////////////////////////////////////////////////////////////////////////////
 // BSDF nodes
 
+// UMaterialExpressionStrataBSDF can only be used for Strata nodes ouputing StrataData that would need a preview,
 UCLASS(MinimalAPI, collapsecategories, hidecategories = Object, Abstract, DisplayName = "Strata Expression")
 class UMaterialExpressionStrataBSDF : public UMaterialExpression
 {
@@ -175,7 +176,6 @@ class UMaterialExpressionStrataLegacyConversion : public UMaterialExpressionStra
 	//~ Begin UMaterialExpression Interface
 #if WITH_EDITOR
 	virtual int32 Compile(class FMaterialCompiler* Compiler, int32 OutputIndex) override;
-	virtual int32 CompilePreview(class FMaterialCompiler* Compiler, int32 OutputIndex) override;
 	virtual void GetCaption(TArray<FString>& OutCaptions) const override;
 	virtual uint32 GetOutputType(int32 OutputIndex) override;
 	virtual uint32 GetInputType(int32 InputIndex) override;
@@ -979,7 +979,7 @@ class UMaterialExpressionStrataThinFilm : public UMaterialExpressionStrataBSDF
 // Utilities
 
 UCLASS(MinimalAPI, collapsecategories, hidecategories = Object, DisplayName = "Strata Transmittance-To-MeanFreePath")
-class UMaterialExpressionStrataTransmittanceToMFP : public UMaterialExpressionStrataBSDF
+class UMaterialExpressionStrataTransmittanceToMFP : public UMaterialExpression
 {
 	GENERATED_UCLASS_BODY()
 
@@ -1009,7 +1009,7 @@ class UMaterialExpressionStrataTransmittanceToMFP : public UMaterialExpressionSt
 };
 
 UCLASS(MinimalAPI, collapsecategories, hidecategories = Object, DisplayName = "Strata Metalness-To-DiffuseColorF0")
-class UMaterialExpressionStrataMetalnessToDiffuseAlbedoF0 : public UMaterialExpressionStrataBSDF
+class UMaterialExpressionStrataMetalnessToDiffuseAlbedoF0 : public UMaterialExpression
 {
 	GENERATED_UCLASS_BODY()
 
@@ -1044,7 +1044,7 @@ class UMaterialExpressionStrataMetalnessToDiffuseAlbedoF0 : public UMaterialExpr
 };
 
 UCLASS(MinimalAPI, collapsecategories, hidecategories = Object, DisplayName = "Strata Haziness-To-Secondary-Roughness")
-class UMaterialExpressionStrataHazinessToSecondaryRoughness : public UMaterialExpressionStrataBSDF
+class UMaterialExpressionStrataHazinessToSecondaryRoughness : public UMaterialExpression
 {
 	GENERATED_UCLASS_BODY()
 
diff --git a/Engine/Source/Runtime/Engine/Private/Materials/HLSLMaterialTranslator.cpp b/Engine/Source/Runtime/Engine/Private/Materials/HLSLMaterialTranslator.cpp
--- a/Engine/Source/Runtime/Engine/Private/Materials/HLSLMaterialTranslator.cpp
+++ b/Engine/Source/Runtime/Engine/Private/Materials/HLSLMaterialTranslator.cpp
@@ -10880,6 +10880,19 @@ int32 FHLSLMaterialTranslator::StrataHazinessToSecondaryRoughness(int32 BaseRoug
 	return INDEX_NONE;
 }
 
+int32 FHLSLMaterialTranslator::StrataCompilePreview(int32 StrataDataCodeChunk)
+{
+	if (StrataDataCodeChunk == INDEX_NONE)
+	{
+		return INDEX_NONE;
+	}
+	// Instead of using a preview color, we might go with a sphere+lighting preview. 
+	// The only problem we would have to solve is to process the strata tree before for the sub tree. 
+	// Or have dedicated lighting preview functions for StrataData that have not be filled in by the Strata tree processing.
+	int32 PreviewCodeChunk = AddCodeChunk(MCT_Float3, TEXT("%s.PreviewColor"), *GetParameterCode(StrataDataCodeChunk));
+	return PreviewCodeChunk;
+}
+
 int32 FHLSLMaterialTranslator::MapARPassthroughCameraUV(int32 UV)
 {
 	if (UV == INDEX_NONE)
diff --git a/Engine/Source/Runtime/Engine/Private/Materials/HLSLMaterialTranslator.h b/Engine/Source/Runtime/Engine/Private/Materials/HLSLMaterialTranslator.h
--- a/Engine/Source/Runtime/Engine/Private/Materials/HLSLMaterialTranslator.h
+++ b/Engine/Source/Runtime/Engine/Private/Materials/HLSLMaterialTranslator.h
@@ -1066,6 +1066,7 @@ protected:
 	virtual int32 StrataTransmittanceToMFP(int32 TransmittanceColor, int32 DesiredThickness, int32 OutputIndex) override;
 	virtual int32 StrataMetalnessToDiffuseAlbedoF0(int32 BaseColor, int32 Specular, int32 Metallic, int32 OutputIndex) override;
 	virtual int32 StrataHazinessToSecondaryRoughness(int32 BaseRoughness, int32 Haziness, int32 OutputIndex) override;
+	virtual int32 StrataCompilePreview(int32 StrataDataCodeChunk) override;
 	
 	virtual FStrataOperator& StrataCompilationRegisterOperator(int32 OperatorType, UMaterialExpression* Expression, UMaterialExpression* Parent, bool bUseParameterBlending = false) override;
 	virtual FStrataOperator& StrataCompilationGetOperator(UMaterialExpression* Expression) override;
diff --git a/Engine/Source/Runtime/Engine/Private/Materials/MaterialExpressions.cpp b/Engine/Source/Runtime/Engine/Private/Materials/MaterialExpressions.cpp
--- a/Engine/Source/Runtime/Engine/Private/Materials/MaterialExpressions.cpp
+++ b/Engine/Source/Runtime/Engine/Private/Materials/MaterialExpressions.cpp
@@ -21676,15 +21676,6 @@ static int32 CompileWithDefaultTangentWS(class FMaterialCompiler* Compiler, FExp
 
 #endif // WITH_EDITOR
 
-#if WITH_EDITOR
-int32 UMaterialExpressionStrataLegacyConversion::CompilePreview(class FMaterialCompiler* Compiler, int32 OutputIndex)
-{
-	// Strata nodes cannot be previewed due to the complex Strata type.
-	// So we override the default implementation calling Compile and we simply return a default black preview color.
-	return Compiler->Constant(0.0f);
-}
-#endif
-
 UMaterialExpressionStrataLegacyConversion::UMaterialExpressionStrataLegacyConversion(const FObjectInitializer& ObjectInitializer)
 	: Super(ObjectInitializer)
 {
@@ -22073,9 +22064,11 @@ UMaterialExpressionStrataBSDF::UMaterialExpressionStrataBSDF(const FObjectInitia
 #if WITH_EDITOR
 int32 UMaterialExpressionStrataBSDF::CompilePreview(class FMaterialCompiler* Compiler, int32 OutputIndex)
 {
-	// Strata nodes cannot be previewed due to the complex Strata type.
-	// So we override the default implementation calling Compile and we simply return a default black preview color.
-	return Compiler->Constant(0.0f);
+	// Compile the StrataData output.
+	int32 StrataDataCodeChunk = Compile(Compiler, OutputIndex);
+	// Convert the StrataData to a preview color.
+	int32 PreviewCodeChunk = Compiler->StrataCompilePreview(StrataDataCodeChunk);
+	return PreviewCodeChunk;
 }
 #endif
 
diff --git a/Engine/Source/Runtime/Engine/Public/MaterialCompiler.h b/Engine/Source/Runtime/Engine/Public/MaterialCompiler.h
--- a/Engine/Source/Runtime/Engine/Public/MaterialCompiler.h
+++ b/Engine/Source/Runtime/Engine/Public/MaterialCompiler.h
@@ -67,6 +67,7 @@ enum EStrataMaterialExport : uint8
 	SME_Normal			= 2,
 	SME_Emissive		= 3,
 	SME_Transmittance	= 4,
+	SME_MaterialPreview	= 5
 };
 /** Exported materials are all opaque unlit. This is used to give some context to the export logic.*/
 enum EStrataMaterialExportContext : uint8
@@ -607,6 +608,7 @@ public:
 	virtual int32 StrataTransmittanceToMFP(int32 TransmittanceColor, int32 DesiredThickness, int32 OutputIndex) = 0;
 	virtual int32 StrataMetalnessToDiffuseAlbedoF0(int32 BaseColor, int32 Specular, int32 Metallic, int32 OutputIndex) = 0;
 	virtual int32 StrataHazinessToSecondaryRoughness(int32 BaseRoughness, int32 Haziness, int32 OutputIndex) = 0;
+	virtual int32 StrataCompilePreview(int32 StrataDataCodeChunk) = 0;
 	
 	/**
 	 * Register an operator of the tree representation the Strata material and its topology.
@@ -1260,6 +1262,11 @@ public:
 		return Compiler->StrataHazinessToSecondaryRoughness(BaseRoughness, Haziness, OutputIndex);
 	}
 
+	virtual int32 StrataCompilePreview(int32 StrataDataCodeChunk) override
+	{
+		return Compiler->StrataCompilePreview(StrataDataCodeChunk);
+	}
+
 	virtual FStrataOperator& StrataCompilationRegisterOperator(int32 OperatorType, UMaterialExpression* Expression, UMaterialExpression* Parent, bool bUseParameterBlending = false) override
 	{
 		return Compiler->StrataCompilationRegisterOperator(OperatorType, Expression, Parent, bUseParameterBlending);
