Fix missing shader print conversion for RT inline stats shader.

diff --git a/Engine/Shaders/Private/RayTracing/TraceRayInlineStatistics.usf b/Engine/Shaders/Private/RayTracing/TraceRayInlineStatistics.usf
--- a/Engine/Shaders/Private/RayTracing/TraceRayInlineStatistics.usf
+++ b/Engine/Shaders/Private/RayTracing/TraceRayInlineStatistics.usf
@@ -13,59 +13,50 @@ void TraceRayInlinePrintStatisticsCS()
 	const float HeadlineX = 0.75f;
 	const float ItemX = 0.76f;
 
-	float2 ScreenPos = float2(HeadlineX, TopMargin);
+	FShaderPrintContext Context = InitShaderPrintContext(true, float2(HeadlineX, TopMargin));
 
 	const uint PassIndex = 0; 
 	const float RayCount = TraversalStatistics[PassIndex].RayCount;
 	const float WavefrontCount = TraversalStatistics[PassIndex].WavefrontCount;
 
-	ScreenPos.x = HeadlineX;
-	ScreenPos = ShaderPrintText(ScreenPos, TEXT("Pass Traversal Statistics"), FontOrange);
-	ScreenPos = ShaderPrintNewline(ScreenPos);
+	Print(Context, TEXT("Pass Traversal Statistics"), FontOrange);
+	Newline(Context);
 
-	ScreenPos.x = ItemX;
-	ScreenPos = ShaderPrintText(ScreenPos, TEXT("Ray Count            "), FontYellow);
-	ScreenPos = ShaderPrint(ScreenPos, TraversalStatistics[PassIndex].RayCount);
-	ScreenPos = ShaderPrintNewline(ScreenPos);
+	Print(Context, TEXT("  Ray Count            "), FontYellow);
+	Print(Context, TraversalStatistics[PassIndex].RayCount);
+	Newline(Context);
 
-	ScreenPos.x = ItemX;
-	ScreenPos = ShaderPrintText(ScreenPos, TEXT("Wavefront Count      "), FontYellow);
-	ScreenPos = ShaderPrint(ScreenPos, TraversalStatistics[PassIndex].WavefrontCount);
-	ScreenPos = ShaderPrintNewline(ScreenPos);
+	Print(Context, TEXT("  Wavefront Count      "), FontYellow);
+	Print(Context, TraversalStatistics[PassIndex].WavefrontCount);
+	Newline(Context);
 
 	// Node Intersection
-	ScreenPos.x = ItemX;
-	ScreenPos = ShaderPrintText(ScreenPos, TEXT("Tri/Ray              "), FontYellow);
+	Print(Context, TEXT("  Tri/Ray              "), FontYellow);
 	float TriangleRay = TraversalStatistics[PassIndex].WaveTriangleIntersectionCount / RayCount;
-	ScreenPos = ShaderPrint(ScreenPos, TriangleRay);
-	ScreenPos = ShaderPrintNewline(ScreenPos);
+	Print(Context, TriangleRay);
+	Newline(Context);
 
-	ScreenPos.x = ItemX;
-	ScreenPos = ShaderPrintText(ScreenPos, TEXT("Node/Ray             "), FontYellow);
+	Print(Context, TEXT("  Node/Ray             "), FontYellow);
 	float NodeRay = TraversalStatistics[PassIndex].WaveNodeIntersectionCount / RayCount;
-	ScreenPos = ShaderPrint(ScreenPos, NodeRay);
-	ScreenPos = ShaderPrintNewline(ScreenPos);
+	Print(Context, NodeRay);
+	Newline(Context);
 
-	ScreenPos.x = ItemX;
-	ScreenPos = ShaderPrintText(ScreenPos, TEXT("Iteration/Wave       "), FontYellow);
+	Print(Context, TEXT("  Iteration/Wave       "), FontYellow);
 	float IterationWave = TraversalStatistics[PassIndex].IterationCount / WavefrontCount;
-	ScreenPos = ShaderPrint(ScreenPos, IterationWave);
-	ScreenPos = ShaderPrintNewline(ScreenPos);
+	Print(Context, IterationWave);
+	Newline(Context);
 
-	ScreenPos.x = ItemX;
-	ScreenPos = ShaderPrintText(ScreenPos, TEXT("Iteration/Ray        "), FontYellow);
+	Print(Context, TEXT("  Iteration/Ray        "), FontYellow);
 	float IterationRay = TriangleRay + NodeRay;
-	ScreenPos = ShaderPrint(ScreenPos, IterationRay);
-	ScreenPos = ShaderPrintNewline(ScreenPos);
+	Print(Context, IterationRay);
+	Newline(Context);
 
-	ScreenPos.x = ItemX;
-	ScreenPos = ShaderPrintText(ScreenPos, TEXT("Iteration Divergence "), FontYellow);
-	ScreenPos = ShaderPrint(ScreenPos, IterationWave / IterationRay);
-	ScreenPos = ShaderPrintNewline(ScreenPos);
+	Print(Context, TEXT("  Iteration Divergence "), FontYellow);
+	Print(Context, IterationWave / IterationRay);
+	Newline(Context);
 
-	ScreenPos.x = ItemX;
-	ScreenPos = ShaderPrintText(ScreenPos, TEXT("Occupancy (%)        "), FontYellow);
+	Print(Context, TEXT("  Occupancy (%)        "), FontYellow);
 	float Occupancy = 100.0 * TraversalStatistics[PassIndex].ActiveLaneCount / TraversalStatistics[PassIndex].TotalLaneCount;
-	ScreenPos = ShaderPrint(ScreenPos, Occupancy);
-	ScreenPos = ShaderPrintNewline(ScreenPos);
+	Print(Context, Occupancy);
+	Newline(Context);
 }
