Volumetric fog history and pre integrated light froxel texture are now pre exposed to avoid smearing artefact at super low exposure (e.g. -10)

diff --git a/Engine/Shaders/Private/HeightFogCommon.ush b/Engine/Shaders/Private/HeightFogCommon.ush
--- a/Engine/Shaders/Private/HeightFogCommon.ush
+++ b/Engine/Shaders/Private/HeightFogCommon.ush
@@ -250,6 +250,9 @@ float4 CombineVolumetricFog(float4 GlobalFog, float3 VolumeUV, uint EyeIndex, fl
 		VolFogStartDistance = FogStruct.VolumetricFogStartDistance;
 		VolumetricFogLookup = Texture3DSampleLevel(FogStruct.IntegratedLightScattering, SharedIntegratedLightScatteringSampler, VolumeUV, 0);
 	#endif
+
+		// IntegratedLightScattering is pre-exposed, remove pre exposure now so that it can correctly be applied later
+		VolumetricFogLookup.rgb *= View.OneOverPreExposure;
 	}
 
 	// Do not apply the Froxel volumetric texture in front of the fog start distance. (the soft fading occur in FinalIntegrationCS).
diff --git a/Engine/Shaders/Private/VolumetricFog.usf b/Engine/Shaders/Private/VolumetricFog.usf
--- a/Engine/Shaders/Private/VolumetricFog.usf
+++ b/Engine/Shaders/Private/VolumetricFog.usf
@@ -397,27 +397,12 @@ uint UseEmissive;
 
 Texture3D<float4> LightScatteringHistory;
 SamplerState LightScatteringHistorySampler;
+float2 LightScatteringHistoryPreExposureAndInv;
 
 Texture3D<float4> LocalShadowedLightScattering;
 
 RWTexture3D<float4> RWLightScattering;
 
-float4 EncodeHDR(float4 Color)
-{
-	return Color;
-
-	//float Exposure = 1;
-	//return float4(Color.rgb * rcp((Color.r*0.299 + Color.g*0.587 + Color.b*0.114) * Exposure + 1.0), Color.a);
-}
-
-float4 DecodeHDR(float4 Color)
-{
-	return Color;
-
-	//float Exposure = 1;
-	//return float4(Color.rgb * rcp((Color.r*(-0.299) + Color.g*(-0.587) + Color.b*(-0.114)) * Exposure + 1.0), Color.a);
-}
-
 float HemisphereConeTraceAgainstGlobalDistanceFieldClipmap(
 	uniform uint ClipmapIndex,
 	float3 WorldShadingPosition, 
@@ -772,22 +757,23 @@ void LightScatteringCS(
 	{
 		MaterialEmissive = VBufferB[GridCoordinate].xyz;
 	}
-	float4 ScatteringAndExtinction = EncodeHDR(float4(LightScattering * MaterialScatteringAndAbsorption.xyz + MaterialEmissive, Extinction));
+	float4 PreExposedScatteringAndExtinction = float4(View.PreExposure * (LightScattering * MaterialScatteringAndAbsorption.xyz + MaterialEmissive), Extinction);
 
 #if USE_TEMPORAL_REPROJECTION
 	BRANCH
 	if (HistoryAlpha > 0)
 	{
-		float4 HistoryScatteringAndExtinction = Texture3DSampleLevel(LightScatteringHistory, LightScatteringHistorySampler, HistoryUV, 0);
-		ScatteringAndExtinction = lerp(ScatteringAndExtinction, HistoryScatteringAndExtinction, HistoryAlpha);
+		float4 PreExposedHistoryScatteringAndExtinction = Texture3DSampleLevel(LightScatteringHistory, LightScatteringHistorySampler, HistoryUV, 0);
+		PreExposedHistoryScatteringAndExtinction *= LightScatteringHistoryPreExposureAndInv.y * View.PreExposure; // Bring previous frame pre exposed history luminance as current frame pre exposed luminance.
+		PreExposedScatteringAndExtinction = lerp(PreExposedScatteringAndExtinction, PreExposedHistoryScatteringAndExtinction, HistoryAlpha);
 	}
 	
 #endif
 
 	if (all(int3(GridCoordinate) < VolumetricFog.GridSizeInt))
 	{
-		ScatteringAndExtinction = MakePositiveFinite(ScatteringAndExtinction);
-		RWLightScattering[GridCoordinate] = ScatteringAndExtinction;
+		PreExposedScatteringAndExtinction = MakePositiveFinite(PreExposedScatteringAndExtinction);
+		RWLightScattering[GridCoordinate] = PreExposedScatteringAndExtinction;
 	}
 }
 
@@ -811,13 +797,13 @@ void FinalIntegrationCS(
 	for (int LayerIndex = 0; LayerIndex < VolumetricFog.GridSizeInt.z; LayerIndex++)
 	{
 		uint3 LayerCoordinate = uint3(GridCoordinate.xy, LayerIndex);
-		float4 ScatteringAndExtinction = DecodeHDR(LightScattering[LayerCoordinate]);
+		float4 PreExposedScatteringAndExtinction = LightScattering[LayerCoordinate];
 
 		float3 LayerTranslatedWorldPosition = ComputeCellTranslatedWorldPosition(LayerCoordinate, .5f);
 		float StepLength = length(LayerTranslatedWorldPosition - PreviousSliceTranslatedWorldPosition);
 		PreviousSliceTranslatedWorldPosition = LayerTranslatedWorldPosition;
 
-		float Transmittance = exp(-ScatteringAndExtinction.w * StepLength);
+		float Transmittance = exp(-PreExposedScatteringAndExtinction.w * StepLength);
 
 		AccumulatedDepth += StepLength;
 
@@ -827,10 +813,10 @@ void FinalIntegrationCS(
 		// See "Physically Based and Unified Volumetric Rendering in Frostbite"
 		#define ENERGY_CONSERVING_INTEGRATION 1
 		#if ENERGY_CONSERVING_INTEGRATION
-			float3 ScatteringIntegratedOverSlice = FadeInLerpValue * (ScatteringAndExtinction.rgb - ScatteringAndExtinction.rgb * Transmittance) / max(ScatteringAndExtinction.w, .00001f);
+			float3 ScatteringIntegratedOverSlice = FadeInLerpValue * (PreExposedScatteringAndExtinction.rgb - PreExposedScatteringAndExtinction.rgb * Transmittance) / max(PreExposedScatteringAndExtinction.w, .00001f);
 			AccumulatedLighting += ScatteringIntegratedOverSlice * AccumulatedTransmittance;
 		#else
-			AccumulatedLighting += FadeInLerpValue * ScatteringAndExtinction.rgb * AccumulatedTransmittance * StepLength;
+			AccumulatedLighting += FadeInLerpValue * PreExposedScatteringAndExtinction.rgb * AccumulatedTransmittance * StepLength;
 		#endif
 		
 		AccumulatedTransmittance *= lerp(1.0f, Transmittance, FadeInLerpValue);
diff --git a/Engine/Source/Runtime/Renderer/Private/ScenePrivate.h b/Engine/Source/Runtime/Renderer/Private/ScenePrivate.h
--- a/Engine/Source/Runtime/Renderer/Private/ScenePrivate.h
+++ b/Engine/Source/Runtime/Renderer/Private/ScenePrivate.h
@@ -992,6 +992,7 @@ public:
 
 	TUniquePtr<FForwardLightingViewResources> ForwardLightingResources;
 
+	float LightScatteringHistoryPreExposure;
 	TRefCountPtr<IPooledRenderTarget> LightScatteringHistory;
 	TRefCountPtr<IPooledRenderTarget> PrevLightScatteringConservativeDepthTexture;
 
diff --git a/Engine/Source/Runtime/Renderer/Private/VolumetricFog.cpp b/Engine/Source/Runtime/Renderer/Private/VolumetricFog.cpp
--- a/Engine/Source/Runtime/Renderer/Private/VolumetricFog.cpp
+++ b/Engine/Source/Runtime/Renderer/Private/VolumetricFog.cpp
@@ -666,6 +666,7 @@ class FVolumetricFogLightScatteringCS : public FGlobalShader
 		SHADER_PARAMETER(FMatrix44f, CloudShadowmapTranslatedWorldToLightClipMatrix)
 		SHADER_PARAMETER(FVector2f, PrevConservativeDepthTextureSize)
 		SHADER_PARAMETER(FVector2f, UseHeightFogColors)
+		SHADER_PARAMETER(FVector2f, LightScatteringHistoryPreExposureAndInv)
 		SHADER_PARAMETER(float, StaticLightingScatteringIntensity)
 		SHADER_PARAMETER(float, SkyLightVolumetricScatteringIntensity)
 		SHADER_PARAMETER(float, SkyLightUseStaticShadowing)
@@ -914,6 +915,7 @@ void FDeferredShadingSceneRenderer::SetupVolumetricFog()
 			if (View.ViewState)
 			{
 				View.ViewState->LightScatteringHistory = NULL;
+				View.ViewState->LightScatteringHistoryPreExposure = 1.0f;
 			}
 		}
 	}
@@ -1151,13 +1153,16 @@ void FDeferredShadingSceneRenderer::ComputeVolumetricFog(FRDGBuilder& GraphBuild
 
 			FVolumetricCloudRenderSceneInfo* CloudInfo = Scene->GetVolumetricCloudSceneInfo();
 			FRDGTexture* LightScatteringHistoryRDGTexture = VolumetricBlackDummyTexture;
+			float LightScatteringHistoryPreExposure = 1.0f;
 			if (bUseTemporalReprojection && View.ViewState->LightScatteringHistory.IsValid())
 			{
 				LightScatteringHistoryRDGTexture = GraphBuilder.RegisterExternalTexture(View.ViewState->LightScatteringHistory);
+				LightScatteringHistoryPreExposure = View.ViewState->LightScatteringHistoryPreExposure;
 			}
 
 			PassParameters->LightScatteringHistory = LightScatteringHistoryRDGTexture;
 			PassParameters->LightScatteringHistorySampler = TStaticSamplerState<SF_Bilinear, AM_Clamp, AM_Clamp, AM_Clamp>::GetRHI();
+			PassParameters->LightScatteringHistoryPreExposureAndInv = FVector2f(LightScatteringHistoryPreExposure, LightScatteringHistoryPreExposure > 0.0f ? 1.0f / LightScatteringHistoryPreExposure : 1.0f);
 
 			FSkyLightSceneProxy* SkyLight = Scene->SkyLight;
 			if (SkyLight
@@ -1298,10 +1303,12 @@ void FDeferredShadingSceneRenderer::ComputeVolumetricFog(FRDGBuilder& GraphBuild
 		if (bUseTemporalReprojection)
 		{
 			GraphBuilder.QueueTextureExtraction(IntegrationData.LightScattering, &View.ViewState->LightScatteringHistory);
+			View.ViewState->LightScatteringHistoryPreExposure = View.CachedViewUniformShaderParameters->PreExposure;
 		}
 		else if (View.ViewState)
 		{
 			View.ViewState->LightScatteringHistory = nullptr;
+			View.ViewState->LightScatteringHistoryPreExposure = 1.0f;
 		}
 
 		if (bUseTemporalReprojection && GVolumetricFogConservativeDepth > 0)
