Added a way to override the skyatmospehre visual sun disk radius from the material graph independentely form the directional light radius.

diff --git a/Engine/Shaders/Private/MaterialTemplate.ush b/Engine/Shaders/Private/MaterialTemplate.ush
--- a/Engine/Shaders/Private/MaterialTemplate.ush
+++ b/Engine/Shaders/Private/MaterialTemplate.ush
@@ -1579,7 +1579,7 @@ float3 MaterialExpressionSkyAtmosphereLightIlluminance(FMaterialPixelParameters
 DEFINE_SKYATMLIGHTDIRECTION(FMaterialPixelParameters)
 DEFINE_SKYATMLIGHTDIRECTION(FMaterialVertexParameters)
 
-float3 MaterialExpressionSkyAtmosphereLightDiskLuminance(FMaterialPixelParameters Parameters, uint LightIndex)
+float3 MaterialExpressionSkyAtmosphereLightDiskLuminance(FMaterialPixelParameters Parameters, uint LightIndex, float OverrideAtmosphereLightDiscCosHalfApexAngle)
 {
 	float3 LightDiskLuminance = float3(0.0f, 0.0f, 0.0f);
 #if MATERIAL_SKY_ATMOSPHERE && PROJECT_SUPPORT_SKY_ATMOSPHERE
@@ -1591,7 +1591,9 @@ float3 MaterialExpressionSkyAtmosphereLightDiskLuminance(FMaterialPixelParameter
 		// GetLightDiskLuminance does a test against the virtual planet but SkyCameraTranslatedWorldOrigin is always put safely setup above it (to never have the camera into the virtual planet with a black screen)
 		LightDiskLuminance =  GetLightDiskLuminance(PlanetCenterToCameraTranslatedWorld, ViewDir, ResolvedView.SkyAtmosphereBottomRadiusKm, ResolvedView.SkyAtmosphereTopRadiusKm,
 			View.TransmittanceLutTexture, View.TransmittanceLutTextureSampler,
-			ResolvedView.AtmosphereLightDirection[LightIndex].xyz, ResolvedView.AtmosphereLightDiscCosHalfApexAngle_PPTrans[LightIndex].x, ResolvedView.AtmosphereLightDiscLuminance[LightIndex].xyz);
+			ResolvedView.AtmosphereLightDirection[LightIndex].xyz, 
+			OverrideAtmosphereLightDiscCosHalfApexAngle >= 0.0f ? OverrideAtmosphereLightDiscCosHalfApexAngle : ResolvedView.AtmosphereLightDiscCosHalfApexAngle_PPTrans[LightIndex].x,
+			ResolvedView.AtmosphereLightDiscLuminance[LightIndex].xyz);
 	}
 #endif
 	return LightDiskLuminance;
diff --git a/Engine/Source/Runtime/Engine/Classes/Materials/MaterialExpressionSkyAtmosphereLightIlluminance.h b/Engine/Source/Runtime/Engine/Classes/Materials/MaterialExpressionSkyAtmosphereLightIlluminance.h
--- a/Engine/Source/Runtime/Engine/Classes/Materials/MaterialExpressionSkyAtmosphereLightIlluminance.h
+++ b/Engine/Source/Runtime/Engine/Classes/Materials/MaterialExpressionSkyAtmosphereLightIlluminance.h
@@ -40,6 +40,10 @@ class UMaterialExpressionSkyAtmosphereLightDiskLuminance : public UMaterialExpre
 	UPROPERTY(EditAnywhere, BlueprintReadWrite, Category = MaterialExpressionTextureCoordinate, meta = (UIMin = "0", UIMax = "1", ClampMin = "0", ClampMax = "1", ShowAsInputPin = "Primary"))
 	int32 LightIndex;
 
+	/** Override the angular diameter of the disk in degree. If not specified, the radius specified on the directional light will be used. This can be used to decouple the directional light visual disk size used for the specular disk reflection on surfaces. However, be aware that screen space reflections will still catch the visual disk. */
+	UPROPERTY()
+	FExpressionInput DiskAngularDiameterOverride;
+
 	//~ Begin UMaterialExpression Interface
 #if WITH_EDITOR
 	virtual int32 Compile(class FMaterialCompiler* Compiler, int32 OutputIndex) override;
diff --git a/Engine/Source/Runtime/Engine/Private/Materials/HLSLMaterialTranslator.cpp b/Engine/Source/Runtime/Engine/Private/Materials/HLSLMaterialTranslator.cpp
--- a/Engine/Source/Runtime/Engine/Private/Materials/HLSLMaterialTranslator.cpp
+++ b/Engine/Source/Runtime/Engine/Private/Materials/HLSLMaterialTranslator.cpp
@@ -9170,10 +9170,11 @@ int32 FHLSLMaterialTranslator::SkyAtmosphereLightDirection(int32 LightIndex)
 	return AddCodeChunk(MCT_Float3, TEXT("MaterialExpressionSkyAtmosphereLightDirection(Parameters, %d)"), LightIndex);
 }
 
-int32 FHLSLMaterialTranslator::SkyAtmosphereLightDiskLuminance(int32 LightIndex)
+int32 FHLSLMaterialTranslator::SkyAtmosphereLightDiskLuminance(int32 LightIndex, int32 OverrideAtmosphereLightDiscCosHalfApexAngle)
 {
 	bUsesSkyAtmosphere = true;
-	return AddCodeChunk(MCT_Float3, TEXT("MaterialExpressionSkyAtmosphereLightDiskLuminance(Parameters, %d)"), LightIndex);
+	return AddCodeChunk(MCT_Float3, TEXT("MaterialExpressionSkyAtmosphereLightDiskLuminance(Parameters, %d, %s)"), LightIndex, 
+		OverrideAtmosphereLightDiscCosHalfApexAngle == INDEX_NONE ? *GetParameterCode(Constant(-1.0f)) : *GetParameterCode(OverrideAtmosphereLightDiscCosHalfApexAngle));
 }
 
 int32 FHLSLMaterialTranslator::SkyAtmosphereViewLuminance()
diff --git a/Engine/Source/Runtime/Engine/Private/Materials/HLSLMaterialTranslator.h b/Engine/Source/Runtime/Engine/Private/Materials/HLSLMaterialTranslator.h
--- a/Engine/Source/Runtime/Engine/Private/Materials/HLSLMaterialTranslator.h
+++ b/Engine/Source/Runtime/Engine/Private/Materials/HLSLMaterialTranslator.h
@@ -984,7 +984,7 @@ protected:
 
 	virtual int32 SkyAtmosphereLightIlluminance(int32 WorldPosition, int32 LightIndex) override;
 	virtual int32 SkyAtmosphereLightDirection(int32 LightIndex) override;
-	virtual int32 SkyAtmosphereLightDiskLuminance(int32 LightIndex) override;
+	virtual int32 SkyAtmosphereLightDiskLuminance(int32 LightIndex, int32 OverrideAtmosphereLightDiscCosHalfApexAngle) override;
 	virtual int32 SkyAtmosphereViewLuminance() override;
 	virtual int32 SkyAtmosphereAerialPerspective(int32 WorldPosition) override;
 	virtual int32 SkyAtmosphereDistantLightScatteredLuminance() override;
diff --git a/Engine/Source/Runtime/Engine/Private/Materials/MaterialExpressions.cpp b/Engine/Source/Runtime/Engine/Private/Materials/MaterialExpressions.cpp
--- a/Engine/Source/Runtime/Engine/Private/Materials/MaterialExpressions.cpp
+++ b/Engine/Source/Runtime/Engine/Private/Materials/MaterialExpressions.cpp
@@ -20354,7 +20354,13 @@ UMaterialExpressionSkyAtmosphereLightDiskLuminance::UMaterialExpressionSkyAtmosp
 #if WITH_EDITOR
 int32 UMaterialExpressionSkyAtmosphereLightDiskLuminance::Compile(class FMaterialCompiler* Compiler, int32 OutputIndex)
 {
-	return Compiler->SkyAtmosphereLightDiskLuminance(LightIndex);
+	int32 CosHalfDiskRadiusCodeChunk = INDEX_NONE;
+	if (DiskAngularDiameterOverride.GetTracedInput().Expression)
+	{
+		// Convert from apex angle (angular diameter) to cosine of the disk radius.
+		CosHalfDiskRadiusCodeChunk = Compiler->Cosine(Compiler->Mul(Compiler->Constant(0.5f * float(UE_PI) / 180.0f), DiskAngularDiameterOverride.Compile(Compiler)));
+	}
+	return Compiler->SkyAtmosphereLightDiskLuminance(LightIndex, CosHalfDiskRadiusCodeChunk);
 }
 
 void UMaterialExpressionSkyAtmosphereLightDiskLuminance::GetCaption(TArray<FString>& OutCaptions) const
diff --git a/Engine/Source/Runtime/Engine/Public/MaterialCompiler.h b/Engine/Source/Runtime/Engine/Public/MaterialCompiler.h
--- a/Engine/Source/Runtime/Engine/Public/MaterialCompiler.h
+++ b/Engine/Source/Runtime/Engine/Public/MaterialCompiler.h
@@ -516,7 +516,7 @@ public:
 
 	virtual int32 SkyAtmosphereLightIlluminance(int32 WorldPosition, int32 LightIndex) = 0;
 	virtual int32 SkyAtmosphereLightDirection(int32 LightIndex) = 0;
-	virtual int32 SkyAtmosphereLightDiskLuminance(int32 LightIndex) = 0;
+	virtual int32 SkyAtmosphereLightDiskLuminance(int32 LightIndex, int32 OverrideAtmosphereLightDiscCosHalfApexAngle) = 0;
 	virtual int32 SkyAtmosphereViewLuminance() = 0;
 	virtual int32 SkyAtmosphereAerialPerspective(int32 WorldPosition) = 0;
 	virtual int32 SkyAtmosphereDistantLightScatteredLuminance() = 0;
@@ -965,9 +965,9 @@ public:
 		return Compiler->SkyAtmosphereLightDirection(LightIndex);
 	}
 
-	virtual int32 SkyAtmosphereLightDiskLuminance(int32 LightIndex) override
+	virtual int32 SkyAtmosphereLightDiskLuminance(int32 LightIndex, int32 OverrideAtmosphereLightDiscCosHalfApexAngle) override
 	{
-		return Compiler->SkyAtmosphereLightDiskLuminance(LightIndex);
+		return Compiler->SkyAtmosphereLightDiskLuminance(LightIndex, OverrideAtmosphereLightDiscCosHalfApexAngle);
 	}
 
 	virtual int32 SkyAtmosphereViewLuminance() override
