Convert remain ShaderPrint API with code statically disabled.

diff --git a/Engine/Shaders/Private/HairStrands/HairStrandsDeepTransmittanceMask.usf b/Engine/Shaders/Private/HairStrands/HairStrandsDeepTransmittanceMask.usf
--- a/Engine/Shaders/Private/HairStrands/HairStrandsDeepTransmittanceMask.usf
+++ b/Engine/Shaders/Private/HairStrands/HairStrandsDeepTransmittanceMask.usf
@@ -244,21 +244,20 @@ void MainCS(uint2 DispatchThreadId : SV_DispatchThreadID)
 void PrintInfo(uint Index, uint2 PixelCoord, FHairTransmittanceMask TransmittanceMask)
 {
 	#if 0
-	const uint2 OutputResolution = uint2(1024, 1024); // TODO send the actual output resolution
+	const uint2 OutputResolution = uint2(1024, 1024); 
 	float2 Origin = (PixelCoord) / float2(OutputResolution);
 	Origin.y += (Index % 8) * 2 * ShaderPrintUniform.FontSize.w;
 
-	float2 Pos = Origin;
-	Pos = ShaderPrintText(Pos, TEXT("Hair Count "));
-	Pos = ShaderPrint(Pos, TransmittanceMask.HairCount);
+	FShaderPrintContext Context = InitShaderPrintContext(true, Origin);
 
-	Pos = ShaderPrintNewline(Pos);
-	Pos.x = Origin.x;
+	Print(Context, TEXT("Hair Count "));
+	Print(Context, TransmittanceMask.HairCount);
+	Newline(Context);
 
-	Pos = ShaderPrintText(Pos, TEXT("Visibility "));
-	Pos = ShaderPrint(Pos, TransmittanceMask.Visibility);
+	Print(Context, TEXT("Visibility "));
+	Print(Context, TransmittanceMask.Visibility);
+	Newline(Context);
 
-	Pos = ShaderPrintNewline(Pos);
 	#endif
 }
 #endif // VOXEL_TRAVERSAL_DEBUG
diff --git a/Engine/Shaders/Private/HairStrands/HairStrandsVoxelPageAllocation.usf b/Engine/Shaders/Private/HairStrands/HairStrandsVoxelPageAllocation.usf
--- a/Engine/Shaders/Private/HairStrands/HairStrandsVoxelPageAllocation.usf
+++ b/Engine/Shaders/Private/HairStrands/HairStrandsVoxelPageAllocation.usf
@@ -800,10 +800,6 @@ void VoxelIndPageClearCS(uint3 DispatchThreadId : SV_DispatchThreadID, uint3 Gro
 
 #if DEBUG_ENABLE
 #include "../ShaderPrint.ush"
-void AddNewLine(float2 OriginalPos, inout float2 Pos)
-{
-	Pos = ShaderPrintNewline(Pos); Pos.x = OriginalPos.x;
-}
 #endif
 
 uint  CPUAllocatedPageCount;
@@ -878,40 +874,39 @@ void FeedbackCS(uint3 DispatchThreadId : SV_DispatchThreadID)
 	FFontColor GPUColor = FontOrange;
 	FFontColor CstColor = FontSilver;
 
-	const float2 OriginalPos = ShaderPrintGetPos(uint2(700, 50));
-	float2 Pos = OriginalPos;
+	FShaderPrintContext Context = InitShaderPrintContext(true, uint2(700, 50));
 
-	Pos = ShaderPrintText(Pos, TEXT(" ------------------------------- "), FontSilver); AddNewLine(OriginalPos, Pos);	
-	Pos = ShaderPrintText(Pos, TEXT("|          Allocations          |"), FontSilver); AddNewLine(OriginalPos, Pos);
-	Pos = ShaderPrintText(Pos, TEXT(" ------------------------------- "), FontSilver); AddNewLine(OriginalPos, Pos);
+	Print(Context, TEXT(" ------------------------------- "), FontSilver); Newline(Context);
+	Print(Context, TEXT("|          Allocations          |"), FontSilver); Newline(Context);
+	Print(Context, TEXT(" ------------------------------- "), FontSilver); Newline(Context);
 
-	Pos = ShaderPrintText(Pos, TEXT("GPU Allocated      "), GPUColor);
-	Pos = ShaderPrint(Pos, GPUAllocatedPageCount, GPUColor);
-	AddNewLine(OriginalPos, Pos);
+	Print(Context, TEXT("GPU Allocated      "), GPUColor);
+	Print(Context, GPUAllocatedPageCount, GPUColor);
+	Newline(Context);
 
-	Pos = ShaderPrintText(Pos, TEXT("CPU Allocated      "), CPUColor);
-	Pos = ShaderPrint(Pos, CPUAllocatedPageCount, CPUColor);
-	AddNewLine(OriginalPos, Pos);
+	Print(Context, TEXT("CPU Allocated      "), CPUColor);
+	Print(Context, CPUAllocatedPageCount, CPUColor);
+	Newline(Context);
 
-	Pos = ShaderPrintText(Pos, TEXT("GPU Curr Min. Size "), GPUColor);
-	Pos = ShaderPrint(Pos, CurrVoxelWorldSize, GPUColor);
-	AddNewLine(OriginalPos, Pos);
+	Print(Context, TEXT("GPU Curr Min. Size "), GPUColor);
+	Print(Context, CurrVoxelWorldSize, GPUColor);
+	Newline(Context);
 
-	Pos = ShaderPrintText(Pos, TEXT("GPU Next Min. Size "), GPUColor);
-	Pos = ShaderPrint(Pos, NextVoxelWorldSize, GPUColor);
-	AddNewLine(OriginalPos, Pos);
+	Print(Context, TEXT("GPU Next Min. Size "), GPUColor);
+	Print(Context, NextVoxelWorldSize, GPUColor);
+	Newline(Context);
 
-	Pos = ShaderPrintText(Pos, TEXT("CPU Min. Size      "), CPUColor);
-	Pos = ShaderPrint(Pos, CPUMinVoxelWorldSize, CPUColor);
-	AddNewLine(OriginalPos, Pos);
+	Print(Context, TEXT("CPU Min. Size      "), CPUColor);
+	Print(Context, CPUMinVoxelWorldSize, CPUColor);
+	Newline(Context);
 
-	Pos = ShaderPrintText(Pos, TEXT("Correction Thres.  "), CstColor);
-	Pos = ShaderPrint(Pos, AdaptiveCorrectionThreshold, CstColor);
-	AddNewLine(OriginalPos, Pos);
+	Print(Context, TEXT("Correction Thres.  "), CstColor);
+	Print(Context, AdaptiveCorrectionThreshold, CstColor);
+	Newline(Context);
 
-	Pos = ShaderPrintText(Pos, TEXT("Correction Speed   "), CstColor);
-	Pos = ShaderPrint(Pos, AdaptiveCorrectionSpeed, CstColor);
-	AddNewLine(OriginalPos, Pos);
+	Print(Context, TEXT("Correction Speed   "), CstColor);
+	Print(Context, AdaptiveCorrectionSpeed, CstColor);
+	Newline(Context);
 #endif
 
 	// Update state data
