Improved HDR viz for handling large HDR values and numbers are not displayed with a white color anymore.

diff --git a/Engine/Shaders/Private/PostProcessVisualizeHDR.usf b/Engine/Shaders/Private/PostProcessVisualizeHDR.usf
--- a/Engine/Shaders/Private/PostProcessVisualizeHDR.usf
+++ b/Engine/Shaders/Private/PostProcessVisualizeHDR.usf
@@ -203,11 +203,12 @@ void MainPS(noperspective float4 UVAndScreenPos : TEXCOORD0, float4 SvPosition :
 
 	float2 UVViewportCenter = (Input_UVViewportMin + Input_UVViewportSize / 2);
 
+	const float3 NumberColor = float3(0.8, 0.85, 0.85);
 	// Luminance meter at the center of the screen
 	{
 		const float PixelDx = abs(PixelPos.x - OutputViewportCenter.x);
 		const float PixelDy = abs(PixelPos.y - OutputViewportCenter.y);
-		if (PixelDx < 130 && PixelDy < 42)
+		if (PixelDx < 180 && PixelDy < 42)
 		{
 			float3 LuminanceAvg = 0.0f;
 			float3 SceneColorAvg = 0.0f;
@@ -250,43 +251,55 @@ void MainPS(noperspective float4 UVAndScreenPos : TEXCOORD0, float4 SvPosition :
 			int2 TopLeft = OutputViewportCenter + int2(11, 11);
 			if (MaxLuminance < 1000.0f)
 			{
-				PrintFloat(PixelPos, OutColor.xyz, float3(1, 1, 1), TopLeft, MaxLuminance);
+				PrintFloat(PixelPos, OutColor.xyz, NumberColor, TopLeft, MaxLuminance);
+				TopLeft.x += 60;
+			}
+			else if (MaxLuminance < 999999.0f)
+			{
+				PrintFloatNoFractionLarge(PixelPos, OutColor.xyz, NumberColor, TopLeft, MaxLuminance);
+				TopLeft.x += 60;
 			}
 			else
 			{
-				PrintFloatNoFractionLarge(PixelPos, OutColor.xyz, float3(1, 1, 1), TopLeft, MaxLuminance);
+				PrintFloatNoFraction(PixelPos, OutColor.xyz, NumberColor, TopLeft, MaxLuminance, 10);
+				TopLeft.x += 100;
 			}
-			TopLeft.x += 60;
-			PrintCharacter(PixelPos, OutColor.xyz, float3(1, 1, 1), TopLeft, _N_);
-			PrintCharacter(PixelPos, OutColor.xyz, float3(1, 1, 1), TopLeft, _I_);
-			PrintCharacter(PixelPos, OutColor.xyz, float3(1, 1, 1), TopLeft, _T_);
+			PrintCharacter(PixelPos, OutColor.xyz, NumberColor, TopLeft, _N_);
+			PrintCharacter(PixelPos, OutColor.xyz, NumberColor, TopLeft, _I_);
+			PrintCharacter(PixelPos, OutColor.xyz, NumberColor, TopLeft, _T_);
 			// Illuminance
 			TopLeft = OutputViewportCenter + int2(11, 22);
 			if (MaxIlluminance < 1000.0f)
 			{
-				PrintFloat(PixelPos, OutColor.xyz, float3(1, 1, 1), TopLeft, MaxIlluminance);
+				PrintFloat(PixelPos, OutColor.xyz, NumberColor, TopLeft, MaxIlluminance);
+				TopLeft.x += 60;
+			}
+			else if(MaxIlluminance < 999999.0f)
+			{
+				PrintFloatNoFractionLarge(PixelPos, OutColor.xyz, NumberColor, TopLeft, MaxIlluminance);
+				TopLeft.x += 60;
 			}
 			else
 			{
-				PrintFloatNoFractionLarge(PixelPos, OutColor.xyz, float3(1, 1, 1), TopLeft, MaxIlluminance);
+				PrintFloatNoFraction(PixelPos, OutColor.xyz, NumberColor, TopLeft, MaxIlluminance, 10);
+				TopLeft.x += 100;
 			}
-			TopLeft.x += 60;
-			PrintCharacter(PixelPos, OutColor.xyz, float3(1, 1, 1), TopLeft, _L_);
-			PrintCharacter(PixelPos, OutColor.xyz, float3(1, 1, 1), TopLeft, _U_);
-			PrintCharacter(PixelPos, OutColor.xyz, float3(1, 1, 1), TopLeft, _X_);
+			PrintCharacter(PixelPos, OutColor.xyz, NumberColor, TopLeft, _L_);
+			PrintCharacter(PixelPos, OutColor.xyz, NumberColor, TopLeft, _U_);
+			PrintCharacter(PixelPos, OutColor.xyz, NumberColor, TopLeft, _X_);
 
 			// EV100
 			TopLeft = OutputViewportCenter + int2(11, 33);
 
 			const float EyeAdaptationEV100 = log2((ComputeRgbLuminance(LuminanceAvg)/.18)/EyeAdaptation_LuminanceMax);
-			PrintMediumFloatWithSign(PixelPos, OutColor.xyz, float3(1, 1, 1), TopLeft, EyeAdaptationEV100);
+			PrintMediumFloatWithSign(PixelPos, OutColor.xyz, NumberColor, TopLeft, EyeAdaptationEV100);
 
 			TopLeft.x += 60;
-			PrintCharacter(PixelPos, OutColor.xyz, float3(1, 1, 1), TopLeft, _E_);
-			PrintCharacter(PixelPos, OutColor.xyz, float3(1, 1, 1), TopLeft, _V_);
-			PrintCharacter(PixelPos, OutColor.xyz, float3(1, 1, 1), TopLeft, _1_);
-			PrintCharacter(PixelPos, OutColor.xyz, float3(1, 1, 1), TopLeft, _0_);
-			PrintCharacter(PixelPos, OutColor.xyz, float3(1, 1, 1), TopLeft, _0_);
+			PrintCharacter(PixelPos, OutColor.xyz, NumberColor, TopLeft, _E_);
+			PrintCharacter(PixelPos, OutColor.xyz, NumberColor, TopLeft, _V_);
+			PrintCharacter(PixelPos, OutColor.xyz, NumberColor, TopLeft, _1_);
+			PrintCharacter(PixelPos, OutColor.xyz, NumberColor, TopLeft, _0_);
+			PrintCharacter(PixelPos, OutColor.xyz, NumberColor, TopLeft, _0_);
 
 			return;
 		}
@@ -309,7 +322,7 @@ void MainPS(noperspective float4 UVAndScreenPos : TEXCOORD0, float4 SvPosition :
 		const float PixelDy = PixelPos.y - IlluminanceMeterCenter.y;
 		const float AbsPixelDx = abs(PixelDx);
 		const float AbsPixelDy = abs(PixelDy);
-		if (PixelDx > - 20 && PixelDx < 200 && PixelDy < 40 && PixelDy > -20)
+		if (PixelDx > - 20 && PixelDx < 300 && PixelDy < 40 && PixelDy > -20)
 		{
 			float3 LuminanceAvg = 0.0f;
 			float SampleCount = 0.0f;
@@ -345,16 +358,22 @@ void MainPS(noperspective float4 UVAndScreenPos : TEXCOORD0, float4 SvPosition :
 			// Mean illuminance
 			if (MeanIlluminance < 1000.0f)
 			{
-				PrintFloat(PixelPos, OutColor.xyz, float3(1, 1, 1), TopLeft, MeanIlluminance);
+				PrintFloat(PixelPos, OutColor.xyz, NumberColor, TopLeft, MeanIlluminance);
+				TopLeft.x += 60;
+			}
+			else if (MeanIlluminance < 999999.0f)
+			{
+				PrintFloatNoFractionLarge(PixelPos, OutColor.xyz, NumberColor, TopLeft, MeanIlluminance);
+				TopLeft.x += 60;
 			}
 			else
 			{
-				PrintFloatNoFractionLarge(PixelPos, OutColor.xyz, float3(1, 1, 1), TopLeft, MeanIlluminance);
+				PrintFloatNoFraction(PixelPos, OutColor.xyz, NumberColor, TopLeft, MeanIlluminance, 10);
+				TopLeft.x += 100;
 			}
-			TopLeft.x += 60;
-			PrintCharacter(PixelPos, OutColor.xyz, float3(1, 1, 1), TopLeft, _L_);
-			PrintCharacter(PixelPos, OutColor.xyz, float3(1, 1, 1), TopLeft, _U_);
-			PrintCharacter(PixelPos, OutColor.xyz, float3(1, 1, 1), TopLeft, _X_);
+			PrintCharacter(PixelPos, OutColor.xyz, NumberColor, TopLeft, _L_);
+			PrintCharacter(PixelPos, OutColor.xyz, NumberColor, TopLeft, _U_);
+			PrintCharacter(PixelPos, OutColor.xyz, NumberColor, TopLeft, _X_);
 
 			// Illuminance as RGB
 			TopLeft = IlluminanceMeterCenterCoord + int2(11, 22);
@@ -362,28 +381,40 @@ void MainPS(noperspective float4 UVAndScreenPos : TEXCOORD0, float4 SvPosition :
 			{
 				PrintFloat(PixelPos, OutColor.xyz, float3(1, 0.5, 0.5), TopLeft, Illuminance.r);
 			}
-			else
+			else if (Illuminance.r < 999999.0f)
 			{
 				PrintFloatNoFractionLarge(PixelPos, OutColor.xyz, float3(1, 0.5, 0.5), TopLeft, Illuminance.r);
 			}
-			TopLeft.x += 60;
+			else
+			{
+				PrintFloatNoFraction(PixelPos, OutColor.xyz, NumberColor, TopLeft, Illuminance.r, 10);
+			}
+			TopLeft.x += 100;
 			if (Illuminance.g < 1000.0f)
 			{
 				PrintFloat(PixelPos, OutColor.xyz, float3(0.5, 1, 0.5), TopLeft, Illuminance.g);
 			}
-			else
+			else if (Illuminance.g < 999999.0f)
 			{
 				PrintFloatNoFractionLarge(PixelPos, OutColor.xyz, float3(0.5, 1, 0.5), TopLeft, Illuminance.g);
 			}
-			TopLeft.x += 60;
+			else
+			{
+				PrintFloatNoFraction(PixelPos, OutColor.xyz, NumberColor, TopLeft, Illuminance.g, 10);
+			}
+			TopLeft.x += 100;
 			if (Illuminance.b < 1000.0f)
 			{
 				PrintFloat(PixelPos, OutColor.xyz, float3(0.5, 0.5, 1), TopLeft, Illuminance.b);
 			}
-			else
+			else if (Illuminance.b < 999999.0f)
 			{
 				PrintFloatNoFractionLarge(PixelPos, OutColor.xyz, float3(0.5, 0.5, 1), TopLeft, Illuminance.b);
 			}
+			else
+			{
+				PrintFloatNoFraction(PixelPos, OutColor.xyz, NumberColor, TopLeft, Illuminance.b, 10);
+			}
 			return;
 		}
 	}
