Uses UINT format instead of FLOAT format to improve the velocity precision on android GLES platform.

diff --git a/Engine/Shaders/Private/Common.ush b/Engine/Shaders/Private/Common.ush
--- a/Engine/Shaders/Private/Common.ush
+++ b/Engine/Shaders/Private/Common.ush
@@ -216,6 +216,12 @@ float VectorSum(float4 V) { return V.x + V.y + V.z + V.w; }
 	#define VELOCITY_ENCODE_GAMMA 0
 #endif
 
+#if COMPILER_GLSL_ES3_1
+	#define ENCODED_VELOCITY_TYPE uint4
+#else
+	#define ENCODED_VELOCITY_TYPE float4
+#endif
+
 //Tie Editor features to platform support and the COMPILE_SHADERS_FOR_DEVELOPMENT which is set via CVAR.
 #define USE_EDITOR_SHADERS (PLATFORM_SUPPORTS_EDITOR_SHADERS && USE_DEVELOPMENT_SHADERS)
 
@@ -1691,7 +1697,7 @@ struct FScreenVertexOutput
 // for velocity rendering, motionblur and temporal AA
 // velocity needs to support -2..2 screen space range for x and y
 // texture is 16bit 0..1 range per channel
-float4 EncodeVelocityToTexture(float3 V)
+ENCODED_VELOCITY_TYPE EncodeVelocityToTexture(float3 V)
 {
 	#if VELOCITY_ENCODE_GAMMA
 		V.xy = sign(V.xy) * sqrt(abs(V.xy)) * (2.0 / sqrt(2.0));
@@ -1711,11 +1717,20 @@ float4 EncodeVelocityToTexture(float3 V)
 		EncodedV.zw = 0.0;
 	#endif
 
+#if COMPILER_GLSL_ES3_1
+	return uint4(EncodedV * 65535.0 + 0.5f);
+#else
 	return EncodedV;
+#endif
 }
 // see EncodeVelocityToTexture()
-float3 DecodeVelocityFromTexture(float4 EncodedV)
+float3 DecodeVelocityFromTexture(ENCODED_VELOCITY_TYPE InEncodedV)
 {
+#if COMPILER_GLSL_ES3_1
+	float4 EncodedV = InEncodedV / 65535.0f;
+#else
+	float4 EncodedV = InEncodedV;
+#endif
 	const float InvDiv = 1.0f / (0.499f * 0.5f);
 
 	float3 V;
diff --git a/Engine/Shaders/Private/TemporalAA.usf b/Engine/Shaders/Private/TemporalAA.usf
--- a/Engine/Shaders/Private/TemporalAA.usf
+++ b/Engine/Shaders/Private/TemporalAA.usf
@@ -498,8 +498,12 @@ SamplerState InputSceneMetadataSampler;
 Texture2D SceneDepthTexture;
 SamplerState SceneDepthTextureSampler;
 
+#if COMPILER_GLSL_ES3_1
+Texture2D<uint4> GBufferVelocityTextureSRV;
+#else
 Texture2D GBufferVelocityTexture;
 SamplerState GBufferVelocityTextureSampler;
+#endif
 
 Texture2D<uint2> StencilTexture;
 
@@ -556,6 +560,16 @@ void CorrectExposure(inout T X)
 
 #endif
 
+ENCODED_VELOCITY_TYPE SampleVelocityTexture(float2 BufferUV, int2 PixelOffset = int2(0, 0))
+{
+#if COMPILER_GLSL_ES3_1
+	int2 Coord = int2(BufferUV * InputSceneColorSize.xy) + PixelOffset;
+	return GBufferVelocityTextureSRV.Load(int3(Coord, 0));
+#else
+	return GBufferVelocityTexture.SampleLevel(GBufferVelocityTextureSampler, BufferUV, 0, PixelOffset);
+#endif
+}
+
 float3 RGBToYCoCg( float3 RGB )
 {
 	float Y  = dot( RGB, float3(  1, 2,  1 ) );
@@ -2052,7 +2066,7 @@ FTAAHistoryPayload TemporalAASample(uint2 GroupId, uint2 GroupThreadId, uint Gro
 
 		#if AA_DYNAMIC
 		{
-			float4 EncodedVelocity = GBufferVelocityTexture.SampleLevel(GBufferVelocityTextureSampler, InputParams.NearestBufferUV + VelocityOffset, 0);
+			ENCODED_VELOCITY_TYPE EncodedVelocity = SampleVelocityTexture(InputParams.NearestBufferUV + VelocityOffset);
 			bool DynamicN = EncodedVelocity.x > 0.0;
 			if(DynamicN)
 			{
@@ -2112,11 +2126,11 @@ FTAAHistoryPayload TemporalAASample(uint2 GroupId, uint2 GroupThreadId, uint Gro
 			#error AA_DYNAMIC_ANTIGHOST requires AA_DYNAMIC
 		#endif
 		// TODO: try a 2x2 for AA_UPSAMPLE
-		bool Dynamic1 = GBufferVelocityTexture.SampleLevel(GBufferVelocityTextureSampler, InputParams.NearestBufferUV, 0, int2( 0, -1)).x > 0;
-		bool Dynamic3 = GBufferVelocityTexture.SampleLevel(GBufferVelocityTextureSampler, InputParams.NearestBufferUV, 0, int2(-1,  0)).x > 0;
-		Dynamic4 = GBufferVelocityTexture.SampleLevel(GBufferVelocityTextureSampler, InputParams.NearestBufferUV, 0).x > 0;
-		bool Dynamic5 = GBufferVelocityTexture.SampleLevel(GBufferVelocityTextureSampler, InputParams.NearestBufferUV, 0, int2( 1,  0)).x > 0;
-		bool Dynamic7 = GBufferVelocityTexture.SampleLevel(GBufferVelocityTextureSampler, InputParams.NearestBufferUV, 0, int2( 0,  1)).x > 0;
+		bool Dynamic1 = SampleVelocityTexture(InputParams.NearestBufferUV, int2( 0, -1)).x > 0;
+		bool Dynamic3 = SampleVelocityTexture(InputParams.NearestBufferUV, int2(-1,  0)).x > 0;
+		Dynamic4 = SampleVelocityTexture(InputParams.NearestBufferUV).x > 0;
+		bool Dynamic5 = SampleVelocityTexture(InputParams.NearestBufferUV, int2( 1,  0)).x > 0;
+		bool Dynamic7 = SampleVelocityTexture(InputParams.NearestBufferUV, int2( 0,  1)).x > 0;
 
 		bool Dynamic = Dynamic1 || Dynamic3 || Dynamic4 || Dynamic5 || Dynamic7;
 		IgnoreHistory = IgnoreHistory || (!Dynamic && History.Color.a > 0);
diff --git a/Engine/Shaders/Private/VelocityShader.usf b/Engine/Shaders/Private/VelocityShader.usf
--- a/Engine/Shaders/Private/VelocityShader.usf
+++ b/Engine/Shaders/Private/VelocityShader.usf
@@ -148,7 +148,7 @@ void MainPixelShader(
 	FVelocityInterpsVSToPS Input
 	OPTIONAL_IsFrontFace
 	OPTIONAL_OutDepthConservative
-	,out float4 OutColor : SV_Target0
+	,out ENCODED_VELOCITY_TYPE OutColor : SV_Target0
 #if MATERIALBLENDING_MASKED_USING_COVERAGE
 	, out uint OutCoverage : SV_Coverage
 #endif
diff --git a/Engine/Source/Runtime/Renderer/Private/PostProcess/TemporalAA.cpp b/Engine/Source/Runtime/Renderer/Private/PostProcess/TemporalAA.cpp
--- a/Engine/Source/Runtime/Renderer/Private/PostProcess/TemporalAA.cpp
+++ b/Engine/Source/Runtime/Renderer/Private/PostProcess/TemporalAA.cpp
@@ -143,6 +143,7 @@ class FTemporalAACS : public FGlobalShader
 		SHADER_PARAMETER_SAMPLER(SamplerState, SceneDepthTextureSampler)
 		SHADER_PARAMETER_RDG_TEXTURE(Texture2D, GBufferVelocityTexture)
 		SHADER_PARAMETER_SAMPLER(SamplerState, GBufferVelocityTextureSampler)
+		SHADER_PARAMETER_RDG_TEXTURE_SRV(Texture2D, GBufferVelocityTextureSRV)
 
 		SHADER_PARAMETER_RDG_TEXTURE_SRV(Texture2D, StencilTexture)
 
@@ -728,6 +729,8 @@ FTAAOutputs AddTemporalAAPass(
 				ResDivisorInv * InputViewRect.Min.Y * InvSizeY);
 		}
 
+		PassParameters->GBufferVelocityTextureSRV = GraphBuilder.CreateSRV(FRDGTextureSRVDesc::Create(PassParameters->GBufferVelocityTexture));
+
 		if (View.GetFeatureLevel() <= ERHIFeatureLevel::ES3_1)
 		{
 			PassParameters->EyeAdaptationBuffer = GraphBuilder.CreateSRV(GetEyeAdaptationBuffer(GraphBuilder, View), PF_A32B32G32R32F);
diff --git a/Engine/Source/Runtime/Renderer/Private/VelocityRendering.cpp b/Engine/Source/Runtime/Renderer/Private/VelocityRendering.cpp
--- a/Engine/Source/Runtime/Renderer/Private/VelocityRendering.cpp
+++ b/Engine/Source/Runtime/Renderer/Private/VelocityRendering.cpp
@@ -370,8 +370,16 @@ EPixelFormat FVelocityRendering::GetFormat(EShaderPlatform ShaderPlatform)
 	const bool bNeedVelocityDepth = (DoesProjectSupportDistanceFields() && FDataDrivenShaderPlatformInfo::GetSupportsLumenGI(ShaderPlatform)) 
 		|| FDataDrivenShaderPlatformInfo::GetSupportsRayTracing(ShaderPlatform);
 
-	// Android platform doesn't support unorm G16R16 format, use G16R16F instead.
-	return bNeedVelocityDepth ? PF_A16B16G16R16 : (IsAndroidOpenGLESPlatform(ShaderPlatform) ? PF_G16R16F : PF_G16R16);
+	// Android GLES platform doesn't support R16G16_UNORM and R16G16B16A16_UNORM format, use R16G16_UINT or R16G16B16A16_UINT instead.
+	const bool bIsAndroidOpenGLESPlatform = IsAndroidOpenGLESPlatform(ShaderPlatform);
+	if (bIsAndroidOpenGLESPlatform)
+	{
+		return bNeedVelocityDepth ? PF_R16G16B16A16_UINT : PF_R16G16_UINT;
+	}
+	else
+	{
+		return bNeedVelocityDepth ? PF_A16B16G16R16 : PF_G16R16;
+	}
 }
 
 FRDGTextureDesc FVelocityRendering::GetRenderTargetDesc(EShaderPlatform ShaderPlatform, FIntPoint Extent)
