Fix a bug of mobile spotlight lighting.

diff --git a/Engine/Shaders/Private/MobileLightingCommon.ush b/Engine/Shaders/Private/MobileLightingCommon.ush
--- a/Engine/Shaders/Private/MobileLightingCommon.ush
+++ b/Engine/Shaders/Private/MobileLightingCommon.ush
@@ -410,7 +410,7 @@ FDeferredLightData ConvertLocalToDeferredLight(const FLocalLightData In)
 	const uint LightTypeAndPackedShadowMapChannelMask = asuint(In.LightDirectionAndShadowMask.w);
 
 	// bits [17:16] really
-	const uint LightType = LightTypeAndPackedShadowMapChannelMask >> 16;
+	const uint LightType = UnpackLightType(LightTypeAndPackedShadowMapChannelMask);
 
 	FDeferredLightData Out = (FDeferredLightData)0;
 	Out.TranslatedWorldPosition = In.LightPositionAndInvRadius.xyz;
