Strata - Fixed specular exploding on SLW curved surface.

diff --git a/Engine/Shaders/Private/BasePassPixelShader.usf b/Engine/Shaders/Private/BasePassPixelShader.usf
--- a/Engine/Shaders/Private/BasePassPixelShader.usf
+++ b/Engine/Shaders/Private/BasePassPixelShader.usf
@@ -967,7 +967,7 @@ void FPixelShaderInOut_MainPS(
 	Specular = SLW_SPECULAR(SLWBSDF);
 	Roughness = SLW_ROUGHNESS(SLWBSDF);
 	Opacity = SLW_TOPMATERIALOPACITY(SLWBSDF);
-	MaterialParameters.WorldNormal = StrataPixelHeader.SharedLocalBases.Normals[BSDF_GETSHAREDLOCALBASISID(SLWBSDF)];
+	MaterialParameters.WorldNormal = normalize(StrataPixelHeader.SharedLocalBases.Normals[BSDF_GETSHAREDLOCALBASISID(SLWBSDF)]);
 	Anisotropy = 0.0f;
 	ShadingModel = MATERIAL_SHADINGMODEL_SINGLELAYERWATER;
 #endif
