Fixes a source of blur in TSR

diff --git a/Engine/Shaders/Private/TemporalSuperResolution/TSRUpdateHistory.usf b/Engine/Shaders/Private/TemporalSuperResolution/TSRUpdateHistory.usf
--- a/Engine/Shaders/Private/TemporalSuperResolution/TSRUpdateHistory.usf
+++ b/Engine/Shaders/Private/TemporalSuperResolution/TSRUpdateHistory.usf
@@ -548,7 +548,7 @@ void MainCS(
 
 			#if !CONFIG_CLAMP
 			{
-				ParallaxRejectionMask = tsr_half(1.0).xx;
+				bIsDisoccluded = false;
 				LowFrequencyRejection = tsr_half(1.0).xx;
 				TranslucencyRejection = tsr_half(1.0).xx;
 			}
@@ -748,15 +748,13 @@ void MainCS(
 			// Soften the spatial kernel when there is high noise in the input to make parrallax disocclusion of detailed geometry less distracting.
 			tsr_half2 OffScreenInputToHistoryFactor = tsr_half(1.0) - tsr_half(0.5) * NoiseFiltering;
 
-			bool2 bIsRefining = PrevHistoryValidity == 0.0;
-			//bIsRefining = true;
+			bool2 bIsRefining = PrevHistoryValidity != 0.0;
 
 			tsr_half2 TotalRejection = select(or(or(bIsOffScreen, bIsDisoccluded), !bIsRefining), tsr_half(0.0).xx, saturate(LowFrequencyRejection * tsr_half(4.0)));
 
 			tsr_half2 KernelInputToHistoryFactor = lerp(OffScreenInputToHistoryFactor, tsr_half(InputToHistoryFactor), TotalRejection);
 			tsr_half2 KernelInputToHistoryAlignmentFactor = lerp(tsr_half(1.0).xx, tsr_half(InputToHistoryFactor), TotalRejection);
 			
-			//InputPixelAlignement = ComputeSampleWeigth(tsr_half(InputToHistoryFactor), dInputKO, /* MinimalContribution = */ float(0.0));
 			InputPixelAlignement = ComputeSampleWeigth(KernelInputToHistoryAlignmentFactor, dInputKO, /* MinimalContribution = */ float(0.0));
 		
 			UNROLL_N(CONFIG_SAMPLES)
