Strata - Single Layer Water fast path. Only 2 uint through MRT when SeparatedDirLight luminance is not enabled. In this case, SLW total cost is reduce by 1ms 6.9 -> 5.85

diff --git a/Engine/Shaders/Private/BasePassPixelShader.usf b/Engine/Shaders/Private/BasePassPixelShader.usf
--- a/Engine/Shaders/Private/BasePassPixelShader.usf
+++ b/Engine/Shaders/Private/BasePassPixelShader.usf
@@ -1599,8 +1599,8 @@ void FPixelShaderInOut_MainPS(
 	// Except for water doing some specialised and simplified lighting during the base pass
 	Color = 0;
 #else
-	// Store the potential separated main dir light contribution into the water BSDF
-	SLW_SEPARATEDMAINDIRLIGHT(StrataPixelHeader.StrataTree.BSDFs[0]) = SeparatedWaterMainDirLightLuminance * View.PreExposure;
+	// Store the potential separated main dir light contribution into the water BSDF. -1 is used as a marked to specify unused (see PackStrataOut).
+	SLW_SEPARATEDMAINDIRLIGHT(StrataPixelHeader.StrataTree.BSDFs[0]) = bSeparateWaterMainDirLightLuminance ? SeparatedWaterMainDirLightLuminance * View.PreExposure : -1.0f;
 #endif
 
 	// STRATA_TODO GBUFFER_HAS_DIFFUSE_SAMPLE_OCCLUSION and ALLOW_STATIC_LIGHTING cases
diff --git a/Engine/Shaders/Private/Strata/Strata.ush b/Engine/Shaders/Private/Strata/Strata.ush
--- a/Engine/Shaders/Private/Strata/Strata.ush
+++ b/Engine/Shaders/Private/Strata/Strata.ush
@@ -157,7 +157,7 @@ struct FStrataBSDF
 // 
 // Layouts:
 // * LAYOUT 2 -------------------------------------------------------------------------------------------------------------------------------------------------------------------------
-//  (Complex) HEADER        31 |  State | AO | BSDF Count | Basis count | Basis type |
+//  (Complex) HEADER        27 |  State | AO | BSDF Count | Basis count | Basis type |
 //                             |  8       8    4            3             4           
 //            BASIS(0)      32 |  Normal: Oct 16/16 | Tangent: 11/11/10
 //            BASIS(1)      32 |  Normal: Oct 16/16 | Tangent: 11/11/10
@@ -170,19 +170,19 @@ struct FStrataBSDF
 //            BSDF(1)       ?? |  Adaptive
 //                             |
 // * LAYOUT 1 -------------------------------------------------------------------------------------------------------------------------------------------------------------------------
-//  (Single)  HEADER        30 |  State | AO | Aniso | TopLayer | SSS | SSSProfile | GreyWeighV | Haze | F90 | Simple vol. | MFP Plugged | Fuzz | Thin
+//  (Single)  HEADER        27 |  State | AO | Aniso | TopLayer | SSS | SSSProfile | GreyWeighV | Haze | F90 | Simple vol. | MFP Plugged | Fuzz | Thin
 //                             |  8       8    1       1          1     1            1            1      1     1             1             1      1            
 //            BASIS(0)       0 | (Use top layer normal)
 //            BSDF(0)       ?? |  Adaptive
 // * LAYOUT 0 -------------------------------------------------------------------------------------------------------------------------------------------------------------------------          
-//  (Simple)  HEADER        29 | State | A0 | Roughness | DiffuseAbedo(8/20)
+//  (Simple)  HEADER        32 | State | A0 | Roughness | DiffuseAbedo(8/20)
 //   (Fast)                    | 8       8    8           8
 //            BASIS(0)       0 | (Use top layer normal)
 //            STATE(0)       0 | 
 //            BSDF(0)       32 | F0_RGB | DiffuseAbedo(12/20)
 //                             | 20       12
 // * LAYOUT X -------------------------------------------------------------------------------------------------------------------------------------------------------------------------          
-//   (Eye)    HEADER        29 | State | A0 | IrisMask | IrisMask
+//   (Eye)    HEADER        32 | State | A0 | IrisMask | IrisMask
 //                             | 8       8    8           8
 //            BASIS(0)       0 | (Use top layer normal)
 //            STATE(0)       0 | 
@@ -191,7 +191,7 @@ struct FStrataBSDF
 //                          32 | IrisNormal    | IrisPlaneNormal
 //                             | 16              16
 // * LAYOUT X -------------------------------------------------------------------------------------------------------------------------------------------------------------------------          
-//   (Hair)   HEADER        29 | State | A0 | Free | Free
+//   (Hair)   HEADER        32 | State | A0 | Free | Free
 //                             | 8       8    8           8
 //            BASIS(0)       0 | (Use top layer normal)
 //            STATE(0)       0 | 
@@ -199,6 +199,15 @@ struct FStrataBSDF
 //                             | 24              8
 //                          32 | Scatter | Specular | Backlit | ComplexTransmittance
 //                             | 8         8          8         8
+// * LAYOUT X -------------------------------------------------------------------------------------------------------------------------------------------------------------------------          
+//   (SLW)    HEADER        32 | State | A0 | bUseSeparateDirLIight | Roughness | BaseColor0
+//                             | 8       8    1           			7		  8
+//            BASIS(0)       0 | (Use top layer normal)
+//            STATE(0)       0 | 
+//            BSDF(0)       32 | BaseColor1 | Metallic | Specular | TopMaterialOpacity
+//                             | 12           8		  7		   7
+//                          32 | OPTIONAL SeparatedMainDirLightLuminance if bUseSeparateDirLIight=true
+//                             | 32 as R11G11B10F
 //            -------------------------------------------------------------------------------------------------------------------------------------------------------------------------
 ////////////////////////////////////////////////////////////////////////////////////////////////////////////////////////////////////////////////////////////////////////////////////////////////
 
@@ -254,6 +263,9 @@ struct FStrataBSDF
 // Layout Hair
 #define HEADER_HAIRENCODING_BIT_COUNT				(HEADER_BIT_COUNT_STATE + HEADER_BIT_COUNT_IRRADIANCE_AO)
 #define HEADER_HAIRENCODING_MASK					((1u<<HEADER_HAIRENCODING_BIT_COUNT)-1u)
+// Layout Eye
+#define HEADER_SLWENCODING_BIT_COUNT				(HEADER_BIT_COUNT_STATE + HEADER_BIT_COUNT_IRRADIANCE_AO)
+#define HEADER_SLWENCODING_MASK					((1u<<HEADER_SLWENCODING_BIT_COUNT)-1u)
 
 // Sanity check
 #if (HEADER_BIT_OFFSET_SHAREDLOCALBASES_TYPE + HEADER_BIT_COUNT_SHAREDLOCALBASES_TYPE) > 32
@@ -1867,7 +1879,7 @@ float3x3 StrataGetBSDFSharedBasis_InlineShading(in FStrataPixelHeader StrataPixe
 float3x3 StrataGetBSDFSharedBasis_DeferredShading(in FStrataPixelHeader StrataPixelHeader, in FStrataBSDF BSDF, in FStrataAddressing StrataAddressing)
 {
 #if STRATA_DEFERRED_SHADING
-	if (IsSimpleMaterial(StrataPixelHeader) || IsSingleMaterial(StrataPixelHeader) || IsEye(StrataPixelHeader) || IsHair(StrataPixelHeader))
+	if (IsSimpleMaterial(StrataPixelHeader) || IsSingleMaterial(StrataPixelHeader) || IsEye(StrataPixelHeader) || IsHair(StrataPixelHeader) || IsSingleLayerWater(StrataPixelHeader))
 	{		
 		return GetTangentBasis(StrataUnpackTopLayerData(StrataPixelHeader.PackedTopLayerData).WorldNormal);
 	}
@@ -2566,6 +2578,11 @@ float3 DitherIfNonBlack(float3 Input, float3 Dither)
 		NewInput.y >= 0.0 && NewInput.y <= 1.0 ? NewInput.y : Input.y,
 		NewInput.z >= 0.0 && NewInput.z <= 1.0 ? NewInput.z : Input.z);
 }
+float DitherIfNonBlack(float Input, float Dither)
+{
+	float NewInput = Input + Dither;
+	return NewInput.x >= 0.0 && NewInput.x <= 1.0 ? NewInput.x : Input.x;
+}
 
 void StrataDitherBSDF(inout FStrataBSDF BSDF, in float Dither)
 {
@@ -2981,6 +2998,46 @@ float3 UnpackR7G7B6Gamma2(uint rgb)
 	return Out * Out;
 }
 
+uint PackR7G7B6Linear(float3 rgb, float Dither)
+{
+	// We execute a dither that can go in positive/negative direction to avoid surfaces to looks overall brighter.
+	const float CenteredDither = Dither - 0.5f;
+	const float Dither8bits = CenteredDither * 1.0f / 255.0f;
+	const float Dither7bits = CenteredDither * 1.0f / 127.0f;
+	const float Dither6bits = CenteredDither * 1.0f / 63.0f;
+	const float3 Dither776 = float3(Dither7bits, Dither7bits, Dither6bits);
+
+	// Dither value to hide more aggressive bit reduction.
+	const float3 dithered_rgb = DitherIfNonBlack(saturate(rgb), Dither776);
+
+	uint r = (uint(dithered_rgb.r * 127.0f) << 0);
+	uint g = (uint(dithered_rgb.g * 127.0f) << 7);
+	uint b = (uint(dithered_rgb.b *  63.0f) << 14);
+	return r | g | b;
+}
+
+float3 UnpackR7G7B6Linear(uint rgb)
+{
+	float3 Out;
+	Out.r = float((rgb >> 0)  & 0x7F) * (1.0f / 127.0f);
+	Out.g = float((rgb >> 7)  & 0x7F) * (1.0f / 127.0f);
+	Out.b = float((rgb >> 14) & 0x3F) * (1.0f /  63.0f);
+	return Out;
+}
+
+uint PackR7(float Value, float Dither)
+{
+	const float CenteredDither = Dither - 0.5f;
+	const float Dither7bits = CenteredDither * 1.0f / 127.0f;
+	const float DitheredValue = DitherIfNonBlack(saturate(Value), Dither7bits);
+	return DitheredValue * 127.0f;
+}
+
+float UnpackR7(uint Value)
+{
+	return float(Value & 0x7F) * (1.0f / 127.0f);
+}
+
 // Simple Octahedral from "A Survey of Efficient Representations for Independent Unit Vectors" http://jcgt.org/published/0003/02/01/
 float SignNotZero(in float k)
 {
@@ -3409,7 +3466,7 @@ FStrataPixelHeader UnpackStrataHeaderIn(FStrataMaterialContainer StrataBuffer, i
 
 	// Use simple or single layout for decoding data
 	BRANCH
-	if (IsSimpleMaterial(Out) || IsSingleMaterial(Out) || IsEye(Out) || IsHair(Out))
+	if (IsSimpleMaterial(Out) || IsSingleMaterial(Out) || IsEye(Out) || IsHair(Out) || IsSingleLayerWater(Out))
 	{
 		// Within a tile classified as simple or single, some pixel might be invalid (e.g., sky pixels). 
 		// This ensures that invalid pixels are not processed, and keep these tiles as single/simple instead of complex
@@ -3513,6 +3570,54 @@ FStrataBSDF UnpackFastPathStrataBSDFIn(FStrataMaterialContainer StrataBuffer, in
 	return OutBSDF;
 }
 
+FStrataBSDF UnpackFastWaterPathStrataBSDFIn(FStrataMaterialContainer StrataBuffer, inout FStrataAddressing StrataAddressing, FStrataPixelHeader StrataPixelHeader)
+{
+	FStrataBSDF OutBSDF = (FStrataBSDF)0;
+
+#if STRATA_DEFERRED_SHADING
+	uint Data0 = StrataPixelHeader.SharedLocalBasesTypes_PackedHeader;
+	uint Data1 = StrataLoadUint1(StrataBuffer, StrataAddressing);
+	OutBSDF.State = 0;
+	OutBSDF.LuminanceWeightV = 1.0f;
+	OutBSDF.TransmittanceAboveAlongN = 1.0f;
+
+	const uint PackedBaseColor8Bits					= 0xFF   & (Data0 >> (HEADER_FASTENCODING_BIT_COUNT + 8));
+	const uint RoughnessAndUseSeparateDirLight8Bits	= 0xFF   & (Data0 >> HEADER_SLWENCODING_BIT_COUNT);
+	
+	const uint PackedBaseColor12Bits				= 0xFFF  & (Data1 >> 20);
+	const uint PackedOpacityMetalSpec20Bits			= 0xFFFF & Data1;
+	const uint PackedBaseColor20Bits				= PackedBaseColor12Bits | (PackedBaseColor8Bits << 12);
+
+	const bool bStoreSeparatedMainDirLight			= RoughnessAndUseSeparateDirLight8Bits & 0x1;
+	if (bStoreSeparatedMainDirLight)
+	{
+		StrataLoad_R11G11B10F(StrataBuffer, StrataAddressing, SLW_SEPARATEDMAINDIRLIGHT(OutBSDF));
+	}
+
+	SLW_ROUGHNESS(OutBSDF)							= UnpackR7(RoughnessAndUseSeparateDirLight8Bits >> 1);
+	SLW_BASECOLOR(OutBSDF)							= UnpackR7G7B6Gamma2(PackedBaseColor20Bits);
+
+	const float3 OpacityMetalSpec					= UnpackR7G7B6Linear(PackedOpacityMetalSpec20Bits);
+	SLW_TOPMATERIALOPACITY(OutBSDF)					= OpacityMetalSpec.x;
+	SLW_METALLIC(OutBSDF)							= OpacityMetalSpec.y;
+	SLW_SPECULAR(OutBSDF)							= OpacityMetalSpec.z;
+#endif
+
+	// Force a BSDF state to be a single layer on top.
+	BSDF_SETTYPE(OutBSDF, STRATA_BSDF_TYPE_SINGLELAYERWATER);
+	BSDF_SETHASGREYWEIGHT_V(OutBSDF, 1);
+	BSDF_SETHASF90(OutBSDF, 0);
+	BSDF_SETHASANISOTROPY(OutBSDF, 0);
+	BSDF_SETISTOPLAYER(OutBSDF, 1);
+	BSDF_SETSSSTYPE(OutBSDF, SSS_TYPE_INVALID);
+	BSDF_SETHASHAZINESS(OutBSDF, 0);
+	BSDF_SETISSIMPLEVOLUME(OutBSDF, 0);
+	BSDF_SETHASFUZZ(OutBSDF, 0);
+	BSDF_SETHASTRANSABOVE(OutBSDF, 0);
+
+	return OutBSDF;
+}
+
 // Unpack a single BSDF
 // Note: All BSDF lobes needs to be unpack in a sequential manner since each BSDF has a variable footprint.
 // If this is changed, please update the compiler side material size evaluation in StrataMaterial.cpp
@@ -3523,6 +3628,10 @@ FStrataBSDF UnpackStrataBSDFIn(FStrataMaterialContainer StrataBuffer, inout FStr
 	{
 		return UnpackFastPathStrataBSDFIn(StrataBuffer, StrataAddressing, StrataHeader);
 	}
+	else if (IsSingleLayerWater(StrataHeader))
+	{
+		return UnpackFastWaterPathStrataBSDFIn(StrataBuffer, StrataAddressing, StrataHeader);
+	}
 
 	FStrataBSDF OutBSDF = (FStrataBSDF)0;
 	OutBSDF.TransmittanceAboveAlongN = 1.0f;
@@ -3702,21 +3811,14 @@ FStrataBSDF UnpackStrataBSDFIn(FStrataMaterialContainer StrataBuffer, inout FStr
 		// 8 bytes
 	}
 	break;
-	case STRATA_BSDF_TYPE_SINGLELAYERWATER:
-	{
-		StrataLoad_ColorGamma2ToLinearAlphaLinear(StrataBuffer, StrataAddressing, SLW_BASECOLOR(OutBSDF), DummyFloat);
-		float4 Data = UnpackRGBA8(StrataLoadUint1(StrataBuffer, StrataAddressing));
-		SLW_METALLIC(OutBSDF) = Data.x;
-		SLW_SPECULAR(OutBSDF) = Data.y;
-		SLW_ROUGHNESS(OutBSDF) = Data.z;
-		SLW_TOPMATERIALOPACITY(OutBSDF) = Data.w;
-		StrataLoad_R11G11B10F(StrataBuffer, StrataAddressing, SLW_SEPARATEDMAINDIRLIGHT(OutBSDF));
-		// 12 bytes
-	}
-	break;
+	//case STRATA_BSDF_TYPE_SINGLELAYERWATER: 
+	//{
+	//	// SLW always use the fast UnpackFastWaterPathStrataBSDFIn path.
+	//}
+	//break;
 	}
 
 	return OutBSDF;
 }
 
-#endif // STRATA_ENABLED
\ No newline at end of file
+#endif // STRATA_ENABLED
diff --git a/Engine/Shaders/Private/Strata/StrataExport.ush b/Engine/Shaders/Private/Strata/StrataExport.ush
--- a/Engine/Shaders/Private/Strata/StrataExport.ush
+++ b/Engine/Shaders/Private/Strata/StrataExport.ush
@@ -359,11 +359,12 @@ void PackStrataOut(
 		// * D. Hair   : Custom encoding	(use top layer normal)
 		// * E. Layout2: Complex encoding	(use basis)
 		const bool bHasFastEncoding = IsSimpleMaterial(StrataPixelHeader);
+		const bool bHasFastWaterEncoding = IsSingleLayerWater(StrataPixelHeader);
 		const bool bHasSingleEncoding = IsSingleMaterial(StrataPixelHeader);
 		const bool bCustomEncoding = IsHair(StrataPixelHeader) || IsEye(StrataPixelHeader);
 
 		// (Layout2)
-		if (!bHasFastEncoding && !bHasSingleEncoding && !bCustomEncoding)
+		if (!bHasFastEncoding && !bHasSingleEncoding && !bCustomEncoding && !bHasFastWaterEncoding)
 		{
 			// 1. the header (Regular/Complex encoding)
 			const uint PackedHeader = PackStrataHeader(BSDFVisibleCount, StrataPixelHeader);
@@ -476,11 +477,53 @@ void PackStrataOut(
 						Out = Out & HEADER_HAIRENCODING_MASK;
 						STRATA_STORE_UINT1(Out);
 
-						#if (HEADER_EYEENCODING_BIT_COUNT) > 32
+						#if (HEADER_HAIRENCODING_BIT_COUNT) > 32
 						#error Strata hair path header is > 32bits
 						#endif
 					}
-					// E. Layout2 - Weight for Regular/Complex path
+					// E. Layout X (SingleLayerWater)
+					else if (bHasFastWaterEncoding)
+					{
+						#if (HEADER_SLWENCODING_BIT_COUNT) > 16
+						#error Strata single layer water header is > 16bits
+						#endif
+
+						// Now we pack diffuse and F0 in a special way: both encoded as R7G7B6 with gamma 2.0. The low bit cound will be hidden by dithering and TAA.
+						// We pack 8 of the 32 bits in the first header uint and the remaining bits are pack is the lowest significant bit of the second uint.
+						const uint PackedBaseColor20Bits = PackR7G7B6Gamma2(SLW_BASECOLOR(BSDF), Dither);
+						const uint PackedBaseColor12Bits = PackedBaseColor20Bits & 0xFFF;
+						const uint PackedBaseColor8Bits = (PackedBaseColor20Bits >> 12) & 0xFF;
+
+						// Packed into the header
+						const bool bStoreSeparatedMainDirLight = SLW_SEPARATEDMAINDIRLIGHT(BSDF).x >= 0.0f;
+						const uint RoughnessAndUseSeparateDirLight8Bits = (PackR7(SLW_ROUGHNESS(BSDF), Dither) << 1) | (bStoreSeparatedMainDirLight ? 1 : 0);
+
+						const uint PackedOpacityMetalSpec20Bits = PackR7G7B6Linear(float3(SLW_TOPMATERIALOPACITY(BSDF), SLW_METALLIC(BSDF), SLW_SPECULAR(BSDF)), Dither);
+
+						// UINT0 (Header_State|Header_AO|BaseColor8bits|Roughness|bStoreSeparatedMainDirLight)
+						{
+							uint Out = 0;
+							HEADER_SETCOMMONSTATES(Out, StrataPixelHeader.State);
+
+							uint Packed16 = (PackedBaseColor8Bits << 8) | RoughnessAndUseSeparateDirLight8Bits;
+							Out = (Out & HEADER_SLWENCODING_MASK) | (Packed16 << HEADER_SLWENCODING_BIT_COUNT);
+							STRATA_STORE_UINT1(Out);
+						}
+						// UINT1 (BaseColor12bits|Opacity|Metallic|Specular)
+						{
+							uint Out = (PackedBaseColor12Bits << 20) | PackedOpacityMetalSpec20Bits;
+							STRATA_STORE_UINT1(Out);
+						}
+						// Optional UINT2
+						if (bStoreSeparatedMainDirLight)
+						{
+							STRATA_STORE_UINT1(PackR11G11B10F(SLW_SEPARATEDMAINDIRLIGHT(BSDF)));	// STRATA_TODO make this an optional load/store base on if DistanceFieldShadow are enabled
+						}
+						// Ensure the rest of the BSDF is not stored with the regular path
+						BSDF.State = 0;
+
+					}
+					// F. Layout2 - Weight for Regular/Complex path
 					else if (GreyScaleThroughputV > 0)
 					{
 						BSDF_SETWEIGHT10F(BSDF, Pack10F(BSDF.LuminanceWeightV.x));
@@ -493,7 +536,7 @@ void PackStrataOut(
 					}
 
 					// Layout1 & Layout2
-					if (!bFastEncodedBSDF)
+					if (!bFastEncodedBSDF && !bHasFastWaterEncoding)
 					{
 						const uint BSDFType = BSDF_GETTYPE(BSDF);
 						switch (BSDF_GETTYPE(BSDF))
@@ -596,13 +639,11 @@ void PackStrataOut(
 							}
 						}
 						break;
-						case STRATA_BSDF_TYPE_SINGLELAYERWATER:
-						{
-							STRATA_STORE_UINT1(PackColorLinearToGamma2AlphaLinear(float4(SLW_BASECOLOR(BSDF), 0.0f)));
-							STRATA_STORE_UINT1(PackRGBA8(float4(SLW_METALLIC(BSDF), SLW_SPECULAR(BSDF), SLW_ROUGHNESS(BSDF), SLW_TOPMATERIALOPACITY(BSDF))));
-							STRATA_STORE_UINT1(PackR11G11B10F(SLW_SEPARATEDMAINDIRLIGHT(BSDF)));	// STRATA_TODO make this an optional load/store base on if DistanceFieldShadow are enabled
-						}
-						break;
+						//case STRATA_BSDF_TYPE_SINGLELAYERWATER: 
+						//{
+						//	// SLW always use the path defined by bHasFastWaterEncoding.
+						//}
+						//break;
 						}
 					}
 
