Strata - save 0.4ms for SSR on SLW.

diff --git a/Engine/Shaders/Private/RayTracing/RayTracingAmbientOcclusionRGS.usf b/Engine/Shaders/Private/RayTracing/RayTracingAmbientOcclusionRGS.usf
--- a/Engine/Shaders/Private/RayTracing/RayTracingAmbientOcclusionRGS.usf
+++ b/Engine/Shaders/Private/RayTracing/RayTracingAmbientOcclusionRGS.usf
@@ -90,7 +90,7 @@ RAY_TRACING_ENTRY_RAYGEN(AmbientOcclusionRGS)
 
 	#if STRATA_ENABLED
 	const FStrataTopLayerData TopLayerData = StrataUnpackTopLayerData(Strata.TopLayerTexture.Load(uint3(PixelCoord, 0)));
-	const bool bIsValid = TopLayerData.bIsValid;
+	const bool bIsValid = IsStrataMaterial(TopLayerData);
 	const bool bReconstructGeometryNormal = true; // STRATA_TODO: legacy avoids this on foliage, how do we do that?
 	const float3 WorldNormal = TopLayerData.WorldNormal;
 	#else
diff --git a/Engine/Shaders/Private/RayTracing/RayTracingOcclusionRGS.usf b/Engine/Shaders/Private/RayTracing/RayTracingOcclusionRGS.usf
--- a/Engine/Shaders/Private/RayTracing/RayTracingOcclusionRGS.usf
+++ b/Engine/Shaders/Private/RayTracing/RayTracingOcclusionRGS.usf
@@ -624,7 +624,7 @@ RAY_TRACING_ENTRY_RAYGEN(OcclusionRGS)
 			const bool bSSSIsValid = StrataSubSurfaceHeaderGetUseDiffusion(SSSHeader);
 			const bool bSSSIsProfile = StrataSubSurfaceHeaderGetIsProfile(SSSHeader);
 
-			ShadingParameters.bIsValid = TopLayerData.bIsValid;
+			ShadingParameters.bIsValid = IsStrataMaterial(TopLayerData);
 			ShadingParameters.bIsHair = StrataGetBSDFType(StrataPixelHeader) == STRATA_BSDF_TYPE_HAIR;
 			ShadingParameters.bIsSubsurface = bSSSIsValid && !bSSSIsProfile;
 			ShadingParameters.bHasSubsurfaceProfile = bSSSIsValid && bSSSIsProfile;
diff --git a/Engine/Shaders/Private/SSRT/SSRTReflections.usf b/Engine/Shaders/Private/SSRT/SSRTReflections.usf
--- a/Engine/Shaders/Private/SSRT/SSRTReflections.usf
+++ b/Engine/Shaders/Private/SSRT/SSRTReflections.usf
@@ -133,31 +133,13 @@ void ScreenSpaceReflections(
 	float3 N = 0;
 	float Roughness = 0.0;
 	bool bNoMaterial = true;
-	BRANCH
-	if (ShouldReflectOnlyWater)
-	{
-		// We do not rerun material classification after the Strata material buffer has been patched with SLW material.
-		// So in this case we simply load the Header and water BSDF directly because we know there can only be a single one.
-		FStrataAddressing StrataAddressing = GetStrataPixelDataByteOffset(PixelCoord, uint2(View.BufferSizeAndInvSize.xy), Strata.MaxBytesPerPixel);
-		FStrataPixelHeader StrataPixelHeader = UnpackStrataHeaderIn(Strata.MaterialTextureArray, StrataAddressing, Strata.TopLayerTexture);
-		if (StrataPixelHeader.BSDFCount > 0 && IsSingleLayerWater(StrataPixelHeader))
-		{
-			const FStrataBSDF SLWaterBSDF = UnpackStrataBSDFIn(Strata.MaterialTextureArray, StrataAddressing, StrataPixelHeader);
-			N = StrataGetBSDFSharedBasis(StrataPixelHeader, SLWaterBSDF, StrataAddressing)[2];
-			Roughness = SLW_ROUGHNESS(SLWaterBSDF);
-			bNoMaterial = false;
-		}
-	}
-	else
-	{
-		const FStrataTopLayerData TopLayerData = StrataUnpackTopLayerData(Strata.TopLayerTexture.Load(uint3(PixelCoord, 0)));
 
-		// For now, SSR is only used to trace the top material layer.
-		// STRATA_TODO try stochastic selection
-		N = TopLayerData.WorldNormal;
-		Roughness = TopLayerData.Roughness;
-		bNoMaterial = !IsStrataMaterial(TopLayerData);
-	}
+	const FStrataTopLayerData TopLayerData = StrataUnpackTopLayerData(Strata.TopLayerTexture.Load(uint3(PixelCoord, 0)));
+
+	// For now, SSR is only used to trace the top material layer.
+	N = TopLayerData.WorldNormal;
+	Roughness = TopLayerData.Roughness;
+	bNoMaterial = ShouldReflectOnlyWater ? !IsSingleLayerWaterMaterial(TopLayerData) : !IsStrataMaterial(TopLayerData);
 
 	float DeviceZ = SampleDeviceZFromSceneTextures(UV);
 	float SceneDepth = ConvertFromDeviceZ(DeviceZ);
@@ -165,8 +147,8 @@ void ScreenSpaceReflections(
 	const float3 PositionTranslatedWorld = mul(float4(ScreenPos * SceneDepth, SceneDepth, 1), View.ScreenToTranslatedWorld).xyz;
 	const float3 V = normalize(View.TranslatedWorldCameraOrigin - PositionTranslatedWorld);
 
-	// STRATA_TODO Handle anisotrpoy using ModifyGGXAnisotropicNormalRoughness or similar approach
-
+	// STRATA_TODO Handle anisotropy using ModifyGGXAnisotropicNormalRoughness or similar approach. 
+	// Although the pass will be more expenssive if reading the material data itself, so this will need to be carefully measured on console.
 
 #else // STRATA_ENABLED
 
diff --git a/Engine/Shaders/Private/ScreenSpaceDenoise/SSDCommon.ush b/Engine/Shaders/Private/ScreenSpaceDenoise/SSDCommon.ush
--- a/Engine/Shaders/Private/ScreenSpaceDenoise/SSDCommon.ush
+++ b/Engine/Shaders/Private/ScreenSpaceDenoise/SSDCommon.ush
@@ -259,7 +259,7 @@ FSSDCompressedSceneInfos SampleCompressedSceneMetadata(
 
 	// STRATA_TODO: Manage bPrevFrame
 	const FStrataTopLayerData TopLayerData = StrataUnpackTopLayerData(Strata.TopLayerTexture.Load(uint3(PixelCoord, 0)));
-	return MaterialToCompressedSceneMetadata(SceneDepth, TopLayerData.WorldNormal, TopLayerData.Roughness, TopLayerData.bIsValid ? SHADINGMODELID_STRATA : SHADINGMODELID_UNLIT);
+	return MaterialToCompressedSceneMetadata(SceneDepth, TopLayerData.WorldNormal, TopLayerData.Roughness, IsStrataMaterial(TopLayerData) ? SHADINGMODELID_STRATA : SHADINGMODELID_UNLIT);
   #else // STRATA_ENABLED
 	FGBufferData GBufferData;
 	if (bPrevFrame)
diff --git a/Engine/Shaders/Private/Strata/Strata.ush b/Engine/Shaders/Private/Strata/Strata.ush
--- a/Engine/Shaders/Private/Strata/Strata.ush
+++ b/Engine/Shaders/Private/Strata/Strata.ush
@@ -526,6 +526,8 @@ uint   StrataPackNormal(in float3 Normal);
 float3 StrataUnpackNormal(uint PackedNormal);
 uint   StrataPackNormal24(in float3 Normal);
 float3 StrataUnpackNormal24(uint PackedNormal);
+uint   StrataPackNormal22(in float3 Normal);
+float3 StrataUnpackNormal22(uint PackedNormal);
 bool   IsStrataSlabFastPathCompatible(in FStrataBSDF BSDF);
 
 ///////////////////////////////////////////////////////////////////////////////
@@ -627,32 +629,43 @@ FStrataIrradianceAndOcclusion StrataUnpackIrradianceAndOcclusion(uint In /*8bits
 
 ///////////////////////////////////////////////////////////////////////////////
 // Top layer
+
+#define TOP_LAYER_MATERIAL_NONE		0
+#define TOP_LAYER_MATERIAL_VALID	1
+#define TOP_LAYER_MATERIAL_SLWATER	2
+// Available #define TOP_LAYER_MATERIAL_xxx	3
+
 struct FStrataTopLayerData
 {
 	float3 BaseColor; // This data is only available during the base pass. It is not packed/stored otherwise.
 	float3 WorldNormal;
 	float  Roughness;
-	bool   bIsValid;
+	uint   Material;
 };
 
 uint StrataPackTopLayerData(FStrataTopLayerData In)
 {
 	// Output is min. by 1 to flag data has valid. This creates a small bias for one particular orientation with a mirror roughness.
-	return max((StrataPackNormal24(In.WorldNormal)<<8) | PackR8(In.Roughness), 1u);
+	return (StrataPackNormal22(In.WorldNormal)<<10) | (PackR8(In.Roughness) << 2) | (In.Material);
 }
 
 FStrataTopLayerData StrataUnpackTopLayerData(uint In)
 {
 	FStrataTopLayerData Out = (FStrataTopLayerData)0;
-	Out.bIsValid = In > 0;
-	Out.WorldNormal = StrataUnpackNormal24(In>>8);
+	Out.Material = (In & 0x3);
+	Out.WorldNormal = StrataUnpackNormal22(In>>10);
 	Out.Roughness = UnpackR8(In);
 	return Out;
 }
 
 bool IsStrataMaterial(FStrataTopLayerData In)
 {
-	return In.bIsValid;
+	return In.Material > TOP_LAYER_MATERIAL_NONE;
+}
+
+bool IsSingleLayerWaterMaterial(FStrataTopLayerData In)
+{
+	return In.Material == TOP_LAYER_MATERIAL_SLWATER;
 }
 
 ///////////////////////////////////////////////////////////////////////////////
@@ -3074,6 +3087,20 @@ float3 StrataUnpackNormal24(uint PackedNormal)
 	return OctahedronToUnitVector(xy);
 }
 
+uint StrataPackNormal22(in float3 Normal)
+{
+	const float2 Result = UnitVectorToOctahedron(Normal);
+	const uint2 PackedXY = uint2(clamp(Result * 1023.0f + 1024.0f, 0.0f, 2047.0f));
+	return PackedXY.x | (PackedXY.y << 11);
+}
+
+float3 StrataUnpackNormal22(uint PackedNormal)
+{
+	const int2 XY11Bits = int2(0x7FF & PackedNormal, 0x7FF & (PackedNormal >> 11));
+	const float2 xy = float2(XY11Bits - 1024) / 1023.0f;
+	return OctahedronToUnitVector(xy);
+}
+
 
 #define STRATA_TANGENT_SYMMETRY 1
 uint StrataPackNormalAndTangent(in float3 Normal, in float3 Tangent)
diff --git a/Engine/Shaders/Private/Strata/StrataExport.ush b/Engine/Shaders/Private/Strata/StrataExport.ush
--- a/Engine/Shaders/Private/Strata/StrataExport.ush
+++ b/Engine/Shaders/Private/Strata/StrataExport.ush
@@ -212,6 +212,7 @@ void PackStrataOut(
 						TopLayerData.Roughness += TopLayerRoughnessContribution;
 						TopLayerData.BaseColor += TopLayerBaseColorContribution;
 
+						TopLayerData.Material = TOP_LAYER_MATERIAL_VALID;
 
 						const bool bHasSSS = BSDF_GETSSSTYPE(BSDF) != SSS_TYPE_INVALID;
 						const bool bIsThin = BSDF_GETISTHIN(BSDF);
@@ -234,6 +235,8 @@ void PackStrataOut(
 						TopLayerData.Roughness += BSDF.TopLayerDataWeight * FullyRough;
 						TopLayerData.BaseColor += BSDF.TopLayerDataWeight * StrataGetBSDFBaseColor(BSDF);
 
+						TopLayerData.Material = TOP_LAYER_MATERIAL_VALID;
+
 						SetMaterialMode(StrataPixelHeader, HEADER_MATERIALMODE_HAIR);
 						break;
 					}
@@ -247,6 +250,8 @@ void PackStrataOut(
 						TopLayerTotalWeight += BSDF.TopLayerDataWeight;
 						TopLayerData.Roughness += BSDF.TopLayerDataWeight * FullyRough;
 
+						TopLayerData.Material = TOP_LAYER_MATERIAL_VALID;
+
 						const bool bHasSSS = BSDF_GETSSSTYPE(BSDF) == SSS_TYPE_DIFFUSION_PROFILE;
 						SetHasSubsurface(StrataPixelHeader, bHasSSS);
 						SetMaterialMode(StrataPixelHeader, HEADER_MATERIALMODE_EYE);
@@ -264,12 +269,13 @@ void PackStrataOut(
 						TopLayerData.Roughness += BSDF.TopLayerDataWeight * SLW_ROUGHNESS(BSDF);
 						TopLayerData.BaseColor += BSDF.TopLayerDataWeight * StrataGetBSDFBaseColor(BSDF);
 
+						TopLayerData.Material = TOP_LAYER_MATERIAL_SLWATER;
+
 						SetMaterialMode(StrataPixelHeader, HEADER_MATERIALMODE_SLWATER);
 						break;
 					}
 					}
 
-
 					float3x3 TangentBasis = StrataGetBSDFSharedBasis_InlineShading(StrataPixelHeader, BSDF, StrataAddressing);
 					float3 N = TangentBasis[2];
 					TopLayerData.WorldNormal += N * BSDF.TopLayerDataWeight;
diff --git a/Engine/Shaders/Private/VirtualShadowMaps/VirtualShadowMapProjection.usf b/Engine/Shaders/Private/VirtualShadowMaps/VirtualShadowMapProjection.usf
--- a/Engine/Shaders/Private/VirtualShadowMaps/VirtualShadowMapProjection.usf
+++ b/Engine/Shaders/Private/VirtualShadowMaps/VirtualShadowMapProjection.usf
@@ -56,7 +56,7 @@ FProjectionShadingInfo GetProjectionShadingInfo(uint2 PixelPos)
 	const bool bIsProfile = StrataSubSurfaceHeaderGetIsProfile(SSSHeader);
 	Out.bIsSubsurface = bIsValid && !bIsProfile;
 
-	Out.bIsValid = TopLayerData.bIsValid;
+	Out.bIsValid = IsStrataMaterial(TopLayerData);
 	Out.bIsHair = StrataGetBSDFType(StrataPixelHeader) == STRATA_BSDF_TYPE_HAIR;
 	Out.WorldNormal = TopLayerData.WorldNormal;
 #else
