Strata - fixed visualize

diff --git a/Engine/Shaders/Private/Strata/StrataVisualize.usf b/Engine/Shaders/Private/Strata/StrataVisualize.usf
--- a/Engine/Shaders/Private/Strata/StrataVisualize.usf
+++ b/Engine/Shaders/Private/Strata/StrataVisualize.usf
@@ -501,13 +501,13 @@ void PrintPixel(uint2 InCoord, float3 InWorldPosition, float3 V)
 		RectMin = Context.Pos;
 		RectMax = Context.StartPos;
 		Print(Context, TEXT("Memory"), FontEmerald); Newline(Context, RectMax);
-		if (bIsSingleMaterial || IsEye(Header))
+		if (bIsSingleMaterial || IsEye(Header) || IsHair(Header))
 		{
 			Print(Context, TEXT(" Header       "), FontSilver); Print(Context, HeaderSize - TopLayerDataBytes,				FontSilver, 3, 0); Print(Context, TEXT("bytes"), FontSilver); Newline(Context, RectMax);
 			Print(Context, TEXT(" TopNormalTex "), FontSilver); Print(Context, TopLayerDataBytes,							FontSilver, 3, 0); Print(Context, TEXT("bytes"), FontSilver); Newline(Context, RectMax);
 			Print(Context, TEXT(" BSDF         "), FontSilver); Print(Context, BSDFsSize,									FontSilver, 3, 0); Print(Context, TEXT("bytes"), FontSilver); Newline(Context, RectMax);
 		}
-		else if (bIsSimpleMaterial)
+		else if (bIsSimpleMaterial || IsSingleLayerWater(Header))
 		{
 			Print(Context, TEXT(" Header+BSDF  "), FontSilver); Print(Context, HeaderSize + BSDFsSize - TopLayerDataBytes,	FontSilver, 3, 0); Print(Context, TEXT("bytes"), FontSilver); Newline(Context, RectMax);
 			Print(Context, TEXT(" TopNormalTex "), FontSilver); Print(Context, TopLayerDataBytes,							FontSilver, 3, 0); Print(Context, TEXT("bytes"), FontSilver); Newline(Context, RectMax);
