Strata - fixed top layer roughness unpacking.

diff --git a/Engine/Shaders/Private/Strata/Strata.ush b/Engine/Shaders/Private/Strata/Strata.ush
--- a/Engine/Shaders/Private/Strata/Strata.ush
+++ b/Engine/Shaders/Private/Strata/Strata.ush
@@ -646,15 +646,15 @@ struct FStrataTopLayerData
 uint StrataPackTopLayerData(FStrataTopLayerData In)
 {
 	// Output is min. by 1 to flag data has valid. This creates a small bias for one particular orientation with a mirror roughness.
-	return (StrataPackNormal22(In.WorldNormal)<<10) | (PackR8(In.Roughness) << 2) | (In.Material);
+	return (StrataPackNormal22(In.WorldNormal) << 10) | (PackR8(In.Roughness) << 2) | (In.Material);
 }
 
 FStrataTopLayerData StrataUnpackTopLayerData(uint In)
 {
 	FStrataTopLayerData Out = (FStrataTopLayerData)0;
 	Out.Material = (In & 0x3);
-	Out.WorldNormal = StrataUnpackNormal22(In>>10);
-	Out.Roughness = UnpackR8(In);
+	Out.WorldNormal = StrataUnpackNormal22(In >> 10);
+	Out.Roughness = UnpackR8(In >> 2);
 	return Out;
 }
 
