Fix for water not receiving directional light contribution when static lighting is disabled.

#rnx

diff --git a/Engine/Shaders/Private/BasePassPixelShader.usf b/Engine/Shaders/Private/BasePassPixelShader.usf
--- a/Engine/Shaders/Private/BasePassPixelShader.usf
+++ b/Engine/Shaders/Private/BasePassPixelShader.usf
@@ -1952,7 +1952,7 @@ void FPixelShaderInOut_MainPS(
 			{
 				// In deferred, we always output the directional light in a separated buffer.
 				// This is used to apply distance field shadows or light function to the main directional light.
-				Out.MRT[GBUFFER_HAS_VELOCITY ? 7 : 6] = float4(SeparatedWaterMainDirLightLuminance * View.PreExposure, 1.0f);
+				Out.MRT[(GBUFFER_HAS_VELOCITY ? 6 : 5) + (GBUFFER_HAS_PRECSHADOWFACTOR ? 1 : 0)] = float4(SeparatedWaterMainDirLightLuminance * View.PreExposure, 1.0f);
 			}
 		#endif
 		}
