Strata - Single Layer Water optimsation for when SeparatedDirLightLuminance is enabled for DFShadow on water surface. That extrata data goes into a render target. This can save 0.1ms on older platforms when water coverage gets large on screen.

diff --git a/Engine/Shaders/Private/BasePassPixelShader.usf b/Engine/Shaders/Private/BasePassPixelShader.usf
--- a/Engine/Shaders/Private/BasePassPixelShader.usf
+++ b/Engine/Shaders/Private/BasePassPixelShader.usf
@@ -1598,9 +1598,6 @@ void FPixelShaderInOut_MainPS(
 	// Need to reset color to make sure strata material are only lit using Strata lighting passes.
 	// Except for water doing some specialised and simplified lighting during the base pass
 	Color = 0;
-#else
-	// Store the potential separated main dir light contribution into the water BSDF. -1 is used as a marked to specify unused (see PackStrataOut).
-	SLW_SEPARATEDMAINDIRLIGHT(StrataPixelHeader.StrataTree.BSDFs[0]) = bSeparateWaterMainDirLightLuminance ? SeparatedWaterMainDirLightLuminance * View.PreExposure : -1.0f;
 #endif
 
 	// STRATA_TODO GBUFFER_HAS_DIFFUSE_SAMPLE_OCCLUSION and ALLOW_STATIC_LIGHTING cases
@@ -1947,13 +1944,18 @@ void FPixelShaderInOut_MainPS(
 				Out.MRT[GBUFFER_HAS_VELOCITY ? 6 : 5] = 0;
 			}
 
-		#if SINGLE_LAYER_WATER_SEPARATED_MAIN_LIGHT && !STRATA_ENABLED
+		#if SINGLE_LAYER_WATER_SEPARATED_MAIN_LIGHT
+			// In deferred, we always output the directional light in a separated buffer.
+			// This is used to apply distance field shadows or light function to the main directional light.
+			// Strata also writes it through MRT because this is faster than through UAV.
+			#if STRATA_ENABLED && STRATA_INLINE_SINGLELAYERWATER
+				Out.MRT[(GBUFFER_HAS_VELOCITY ? 2 : 1) + (GBUFFER_HAS_PRECSHADOWFACTOR ? 1 : 0)] = float4(SeparatedWaterMainDirLightLuminance * View.PreExposure, 1.0f);
+			#else
 			if (GBuffer.ShadingModelID == SHADINGMODELID_SINGLELAYERWATER)
 			{
-				// In deferred, we always output the directional light in a separated buffer.
-				// This is used to apply distance field shadows or light function to the main directional light.
 				Out.MRT[(GBUFFER_HAS_VELOCITY ? 6 : 5) + (GBUFFER_HAS_PRECSHADOWFACTOR ? 1 : 0)] = float4(SeparatedWaterMainDirLightLuminance * View.PreExposure, 1.0f);
 			}
+			#endif
 		#endif
 		}
 		#else
diff --git a/Engine/Shaders/Private/PixelShaderOutputCommon.ush b/Engine/Shaders/Private/PixelShaderOutputCommon.ush
--- a/Engine/Shaders/Private/PixelShaderOutputCommon.ush
+++ b/Engine/Shaders/Private/PixelShaderOutputCommon.ush
@@ -89,14 +89,22 @@ void MainPS
 		, out float4 OutTarget2 : SV_Target2
 #endif
 
+#if PIXELSHADEROUTPUT_MRT3
+		, out float4 OutTarget3 : SV_Target3
+#endif
+
 #if STRATA_OPAQUE_DEFERRED
 #if STRATA_BASE_PASS_MRT_OUTPUT_COUNT != 2
 #error Strata STRATA_BASE_PASS_MRT_OUTPUT_COUNT has been update but not the uint MRTs
 #endif
-#if PIXELSHADEROUTPUT_MRT3
+#if PIXELSHADEROUTPUT_MRT4
 #error Strata cannot map to such a case
 #endif
-#if PIXELSHADEROUTPUT_MRT2
+#if PIXELSHADEROUTPUT_MRT3
+		, out uint StrataOutput4 : SV_Target4
+		, out uint StrataOutput5 : SV_Target5
+		, out uint StrataOutput6 : SV_Target6
+#elif PIXELSHADEROUTPUT_MRT2
 		, out uint StrataOutput3 : SV_Target3
 		, out uint StrataOutput4 : SV_Target4
 		, out uint StrataOutput5 : SV_Target5
@@ -111,10 +119,6 @@ void MainPS
 #endif
 #else // STRATA_OPAQUE_DEFERRED
 
-#if PIXELSHADEROUTPUT_MRT3
-		, out float4 OutTarget3 : SV_Target3
-#endif
-
 #if PIXELSHADEROUTPUT_MRT4
 		, out float4 OutTarget4 : SV_Target4
 #endif
@@ -195,7 +199,11 @@ void MainPS
 	// MRT2 is precomputed shadow if both velocity and prec shadow are enabled.
 	// After, Strata output are appended (3 uint output)
 
-#if PIXELSHADEROUTPUT_MRT2
+#if PIXELSHADEROUTPUT_MRT3
+	StrataOutput4 = PixelShaderOut.StrataOutput[0];
+	StrataOutput5 = PixelShaderOut.StrataOutput[1];
+	StrataOutput6 = PixelShaderOut.StrataOutput[2];
+#elif PIXELSHADEROUTPUT_MRT2
 	StrataOutput3 = PixelShaderOut.StrataOutput[0];
 	StrataOutput4 = PixelShaderOut.StrataOutput[1];
 	StrataOutput5 = PixelShaderOut.StrataOutput[2];
diff --git a/Engine/Shaders/Private/SingleLayerWaterComposite.usf b/Engine/Shaders/Private/SingleLayerWaterComposite.usf
--- a/Engine/Shaders/Private/SingleLayerWaterComposite.usf
+++ b/Engine/Shaders/Private/SingleLayerWaterComposite.usf
@@ -106,7 +106,7 @@ SingleLayerWaterCompositeOutput SingleLayerWaterComposite(float2 BufferUV, float
 			if (UseSeparatedMainDirLightTexture > 0.0f)
 			{
 				// Also compose the separated main direction light luminance (separated to be able to combine it with DistanceFieldShadow)
-				SeparatedMainDirLightContribution = SLW_SEPARATEDMAINDIRLIGHT(SLWaterBSDF) * SeparatedMainDirLightTexture[PixelPos].rgb;
+				SeparatedMainDirLightContribution = SeparatedMainDirLightTexture[PixelPos].rgb;
 			}
 
 			// Compute the sky reflection contribution
diff --git a/Engine/Shaders/Private/Strata/Strata.ush b/Engine/Shaders/Private/Strata/Strata.ush
--- a/Engine/Shaders/Private/Strata/Strata.ush
+++ b/Engine/Shaders/Private/Strata/Strata.ush
@@ -495,7 +495,6 @@ struct FStrataBSDF
 #define SLW_SPECULAR(X)					X.VGPRs[1].y
 #define SLW_ROUGHNESS(X)				X.VGPRs[1].z
 #define SLW_TOPMATERIALOPACITY(X)		X.VGPRs[1].w
-#define SLW_SEPARATEDMAINDIRLIGHT(X)	X.VGPRs[2].xyz
 #if STRATA_INLINE_SINGLELAYERWATER
 #define SLW_WATERALBEDO(X)				X.InlineVGPRs[0].xyz
 #define SLW_WATEREXTINCTION(X)			X.InlineVGPRs[1].xyz
@@ -1554,7 +1553,6 @@ FStrataData GetStrataSingleLayerWaterBSDF(
 	SLW_SPECULAR					(StrataData.InlinedBSDF) = Specular;
 	SLW_ROUGHNESS					(StrataData.InlinedBSDF) = Roughness;
 	SLW_TOPMATERIALOPACITY			(StrataData.InlinedBSDF) = TopMaterialOpacity;
-	SLW_SEPARATEDMAINDIRLIGHT		(StrataData.InlinedBSDF) = float3(0.0f, 0.0f, 0.0f);
 #if STRATA_INLINE_SINGLELAYERWATER
 	SLW_WATERALBEDO					(StrataData.InlinedBSDF) = WaterAlbedo;
 	SLW_WATEREXTINCTION				(StrataData.InlinedBSDF) = WaterExtinction;
@@ -3616,10 +3614,6 @@ FStrataBSDF UnpackFastWaterPathStrataBSDFIn(FStrataMaterialContainer StrataBuffe
 	const uint PackedBaseColor20Bits				= PackedBaseColor12Bits | (PackedBaseColor8Bits << 12);
 
 	const bool bStoreSeparatedMainDirLight			= RoughnessAndUseSeparateDirLight8Bits & 0x1;
-	if (bStoreSeparatedMainDirLight)
-	{
-		StrataLoad_R11G11B10F(StrataBuffer, StrataAddressing, SLW_SEPARATEDMAINDIRLIGHT(OutBSDF));
-	}
 
 	SLW_ROUGHNESS(OutBSDF)							= UnpackR7(RoughnessAndUseSeparateDirLight8Bits >> 1);
 	SLW_BASECOLOR(OutBSDF)							= UnpackR7G7B6Gamma2(PackedBaseColor20Bits);
diff --git a/Engine/Shaders/Private/Strata/StrataExport.ush b/Engine/Shaders/Private/Strata/StrataExport.ush
--- a/Engine/Shaders/Private/Strata/StrataExport.ush
+++ b/Engine/Shaders/Private/Strata/StrataExport.ush
@@ -501,17 +501,15 @@ void PackStrataOut(
 						const uint PackedBaseColor8Bits = (PackedBaseColor20Bits >> 12) & 0xFF;
 
 						// Packed into the header
-						const bool bStoreSeparatedMainDirLight = SLW_SEPARATEDMAINDIRLIGHT(BSDF).x >= 0.0f;
-						const uint RoughnessAndUseSeparateDirLight8Bits = (PackR7(SLW_ROUGHNESS(BSDF), Dither) << 1) | (bStoreSeparatedMainDirLight ? 1 : 0);
-
+						const uint Roughness8Bits = PackR8(SLW_ROUGHNESS(BSDF));
 						const uint PackedOpacityMetalSpec20Bits = PackR7G7B6Linear(float3(SLW_TOPMATERIALOPACITY(BSDF), SLW_METALLIC(BSDF), SLW_SPECULAR(BSDF)), Dither);
 
-						// UINT0 (Header_State|Header_AO|BaseColor8bits|Roughness|bStoreSeparatedMainDirLight)
+						// UINT0 (Header_State|Header_AO|BaseColor8bits|Roughness)
 						{
 							uint Out = 0;
 							HEADER_SETCOMMONSTATES(Out, StrataPixelHeader.State);
 
-							uint Packed16 = (PackedBaseColor8Bits << 8) | RoughnessAndUseSeparateDirLight8Bits;
+							uint Packed16 = (PackedBaseColor8Bits << 8) | Roughness8Bits;
 							Out = (Out & HEADER_SLWENCODING_MASK) | (Packed16 << HEADER_SLWENCODING_BIT_COUNT);
 							STRATA_STORE_UINT1(Out);
 						}
@@ -520,11 +518,6 @@ void PackStrataOut(
 							uint Out = (PackedBaseColor12Bits << 20) | PackedOpacityMetalSpec20Bits;
 							STRATA_STORE_UINT1(Out);
 						}
-						// Optional UINT2
-						if (bStoreSeparatedMainDirLight)
-						{
-							STRATA_STORE_UINT1(PackR11G11B10F(SLW_SEPARATEDMAINDIRLIGHT(BSDF)));	// STRATA_TODO make this an optional load/store base on if DistanceFieldShadow are enabled
-						}
 						// Ensure the rest of the BSDF is not stored with the regular path
 						BSDF.State = 0;
 
diff --git a/Engine/Source/Runtime/Engine/Private/Materials/HLSLMaterialTranslator.cpp b/Engine/Source/Runtime/Engine/Private/Materials/HLSLMaterialTranslator.cpp
--- a/Engine/Source/Runtime/Engine/Private/Materials/HLSLMaterialTranslator.cpp
+++ b/Engine/Source/Runtime/Engine/Private/Materials/HLSLMaterialTranslator.cpp
@@ -10162,8 +10162,6 @@ bool FHLSLMaterialTranslator::StrataGenerateDerivedMaterialOperatorData()
 						// Header + Data
 						StrataMaterialRequestedSizeByte += UintByteSize;
 						// Data
-						StrataMaterialRequestedSizeByte += UintByteSize;	
-						// Optional separted dir light luminance
 						StrataMaterialRequestedSizeByte += UintByteSize;
 						break; // Stop here
 					}
diff --git a/Engine/Source/Runtime/Engine/Private/ShaderCompiler/ShaderGenerationUtil.cpp b/Engine/Source/Runtime/Engine/Private/ShaderCompiler/ShaderGenerationUtil.cpp
--- a/Engine/Source/Runtime/Engine/Private/ShaderCompiler/ShaderGenerationUtil.cpp
+++ b/Engine/Source/Runtime/Engine/Private/ShaderCompiler/ShaderGenerationUtil.cpp
@@ -1785,7 +1785,7 @@ static void DetermineUsedMaterialSlots(
 		SetStandardGBufferSlots(Slots, bWriteEmissive, bHasTangent, bHasVelocity, bHasStaticLighting, bIsStrataMaterial);
 		if (Mat.SINGLE_LAYER_WATER_DF_SHADOW_ENABLED)
 		{
-			Slots[GBS_SeparatedMainDirLight] = !bIsStrataMaterial; // This is stored in the Strata material buffer, not in a separate RT.
+			Slots[GBS_SeparatedMainDirLight] = true;
 		}
 	}
 
diff --git a/Engine/Source/Runtime/RenderCore/Private/GBufferInfo.cpp b/Engine/Source/Runtime/RenderCore/Private/GBufferInfo.cpp
--- a/Engine/Source/Runtime/RenderCore/Private/GBufferInfo.cpp
+++ b/Engine/Source/Runtime/RenderCore/Private/GBufferInfo.cpp
@@ -255,6 +255,10 @@ FGBufferInfo RENDERCORE_API FetchLegacyGBufferInfo(const FGBufferParams& Params)
 		{
 			TargetGBufferE = Info.NumTargets++;
 		}
+		if (Params.bHasSingleLayerWaterSeparatedMainLight)
+		{
+			TargetSeparatedMainDirLight = Info.NumTargets++;
+		}
 
 		// this value isn't correct, becuase it doesn't resepect the scene color format cvar, but it's ignored anyways
 		// so it's ok for now
@@ -288,7 +292,14 @@ FGBufferInfo RENDERCORE_API FetchLegacyGBufferInfo(const FGBufferParams& Params)
 			Info.Slots[GBS_PrecomputedShadowFactor].Packing[3] = FGBufferPacking(TargetGBufferE, 3, 3);
 		}
 
-		// We do not setup GBS_SeparatedMainDirLight as the data will flow through the Strata material buffer.
+		// Special water output
+		if (Params.bHasSingleLayerWaterSeparatedMainLight)
+		{
+			Info.Slots[GBS_SeparatedMainDirLight] = FGBufferItem(GBS_SeparatedMainDirLight, GBC_Raw_Float_11_11_10, GBCH_Both);
+			Info.Slots[GBS_SeparatedMainDirLight].Packing[0] = FGBufferPacking(TargetSeparatedMainDirLight, 0, 0);
+			Info.Slots[GBS_SeparatedMainDirLight].Packing[1] = FGBufferPacking(TargetSeparatedMainDirLight, 1, 1);
+			Info.Slots[GBS_SeparatedMainDirLight].Packing[2] = FGBufferPacking(TargetSeparatedMainDirLight, 2, 2);
+		}
 
 		return Info;
 	}
diff --git a/Engine/Source/Runtime/Renderer/Private/SingleLayerWaterRendering.cpp b/Engine/Source/Runtime/Renderer/Private/SingleLayerWaterRendering.cpp
--- a/Engine/Source/Runtime/Renderer/Private/SingleLayerWaterRendering.cpp
+++ b/Engine/Source/Runtime/Renderer/Private/SingleLayerWaterRendering.cpp
@@ -325,12 +325,6 @@ static FSceneWithoutWaterTextures AddCopySceneWithoutWaterPass(
 
 	const FRDGTextureDesc SeparatedMainDirLightDesc(FRDGTextureDesc::Create2D(SceneColorDesc.Extent, PF_FloatR11G11B10, FClearValueBinding(FLinearColor::White), TexCreate_ShaderResource | TexCreate_RenderTargetable));
 	FRDGTextureRef SeparatedMainDirLightTexture = GraphBuilder.CreateTexture(SeparatedMainDirLightDesc, TEXT("SLW.SeparatedMainDirLight"));
-	if (IsWaterDistanceFieldShadowEnabled_Runtime(ShaderPlatform) && Strata::IsStrataEnabled())
-	{
-		// This clear is needed with strata because that texture will be modulated by DFShadows.
-		// STRATA_TODO: when strata is enabled, we can change RenderRayTracedDistanceFieldProjection to have a bForceNoBlending instead if bForceRGBModulation and remove that clear.
-		AddClearRenderTargetPass(GraphBuilder, SeparatedMainDirLightTexture);
-	}
 
 	FSceneWithoutWaterTextures Textures;
 	Textures.RefractionDownsampleFactor = float(RefractionDownsampleFactor);
@@ -890,8 +884,7 @@ void FDeferredShadingSceneRenderer::RenderSingleLayerWaterInner(
 
 	TStaticArray<FTextureRenderTargetBinding, MaxSimultaneousRenderTargets> BasePassTextures;
 	uint32 BasePassTextureCount = SceneTextures.GetGBufferRenderTargets(BasePassTextures);
-	if(IsWaterDistanceFieldShadowEnabled_Runtime(Scene->GetShaderPlatform()) 
-		&& !Strata::IsStrataEnabled())	// We do not bind that texture if Strata is enabled as the data will go through the Strata material buffer.
+	if(IsWaterDistanceFieldShadowEnabled_Runtime(Scene->GetShaderPlatform()))
 	{
 		const bool bNeverClear = true;
 		BasePassTextures[BasePassTextureCount++] = FTextureRenderTargetBinding(SceneWithoutWaterTextures.SeparatedMainDirLightTexture, bNeverClear);
