Hair SW rasterizer now properly clips hair segments to tiles and correctly rasterizes segments longer than one tile. Fixed wrong near plane clipping in binning step by doing clipping in homogeneous coordinates.

#rnx

diff --git a/Engine/Shaders/Private/HairStrands/HairStrandsVisibilityRasterCompute.usf b/Engine/Shaders/Private/HairStrands/HairStrandsVisibilityRasterCompute.usf
--- a/Engine/Shaders/Private/HairStrands/HairStrandsVisibilityRasterCompute.usf
+++ b/Engine/Shaders/Private/HairStrands/HairStrandsVisibilityRasterCompute.usf
@@ -94,7 +94,7 @@ float3 NDCToPixelCoord(float4 InDC)
 	return float3(UV * OutputResolution, NDC.z);
 }
 
-void CalcScreenPos(in uint4 PB, in float3 PBO, out float3 SP, out uint Type)
+void CalcHomogenousPos(in uint4 PB, in float3 PBO, out float4 HP, out uint Type)
 {
 	const FHairControlPoint CP = UnpackHairControlPoint(
 		PB,
@@ -103,15 +103,12 @@ void CalcScreenPos(in uint4 PB, in float3 PBO, out float3 SP, out uint Type)
 		HairStrandsVF_RootScale,
 		HairStrandsVF_TipScale);
 
-	const float3 WP = mul(float4(CP.Position, 1), HairStrandsVF_LocalToWorldPrimitiveTransform).xyz;
-	float4 NDC = mul(float4(WP, 1), LWCHackToFloat(PrimaryView.WorldToClip));
-
-	SP = NDCToPixelCoord(NDC);
-
+	const float3 WP = mul(float4(CP.Position, 1.0f), HairStrandsVF_LocalToWorldPrimitiveTransform).xyz;
+	HP = mul(float4(WP, 1.0f), LWCHackToFloat(PrimaryView.WorldToClip));
 	Type = CP.Type;
 }
 
-void CalcScreenPosAndRad(in uint4 PB, in float3 PBO, out float4 SP, out float Rad)
+void CalcHomogenousPosAndRad(in uint4 PB, in float3 PBO, out float4 HP, out float Rad)
 {
 	const FHairControlPoint CP = UnpackHairControlPoint(
 		PB,
@@ -120,32 +117,76 @@ void CalcScreenPosAndRad(in uint4 PB, in float3 PBO, out float4 SP, out float Ra
 		HairStrandsVF_RootScale,
 		HairStrandsVF_TipScale);
 
-	const float3 WP = mul(float4(CP.Position, 1), HairStrandsVF_LocalToWorldPrimitiveTransform).xyz;
-	float4 NDC = mul(float4(WP, 1), LWCHackToFloat(PrimaryView.WorldToClip));
+	const float3 WP = mul(float4(CP.Position, 1.0f), HairStrandsVF_LocalToWorldPrimitiveTransform).xyz;
+	HP = mul(float4(WP, 1.0f), LWCHackToFloat(PrimaryView.WorldToClip));
+	Rad = CP.WorldRadius * 2000.0; // OutputResolutionf.x; //TODO: figure this out correctly?
+}
+
+// Line clipping based on "CLIPPING USING HOMOGENEOUS COORDINATES" by Blinn et al.
+bool BlinnLineClipping(inout float4 P0, inout float4 P1)
+{
+	float2 T = float2(0.0f, 1.0f);
+	bool bIsRemoved = P0.w < 0.0f && P1.w < 0.0f; // Both points behind near plane
+
+	bool bSign = false;
+
+	UNROLL
+	for (uint PlaneIdx = 0; PlaneIdx < 6; ++PlaneIdx)
+	{
+		// Compute boundary coordinates of both points (w+x, w-x, w+y, w-y, z, w-z)
+		bSign = !bSign;
+		const uint CompIdx = PlaneIdx / 2;
+		const float Sign = bSign ? 1.0f : -1.0f;
+		const float WFactor = PlaneIdx != 4 ? 1.0f : 0.0f;
+		const float2 BC = WFactor * float2(P0.w, P1.w) + Sign * float2(P0[CompIdx], P1[CompIdx]);
+
+		float Num = BC.x;
+		float Denom = BC.x - BC.y;
+		bIsRemoved = bIsRemoved || (BC.x < 0.0f && BC.y < 0.0f); // Both points outside the plane
+		float Alpha = Num / Denom;
+		
+		// If the denominator is negative, P0 has a smaller boundary coordinate than P1, so we can assume
+		// that P1 is inside the plane (or bIsRemoved is true), so we need to update the alpha for P0.
+		// The reverse is true if the denominator is positive.
+		if (Denom < 0.0f)
+		{
+			T.x = max(T.x, Alpha);
+		}
+		else
+		{
+			T.y = min(T.y, Alpha);
+		}
+	}
 
-	SP = float4(NDCToPixelCoord(NDC).xy, NDC.zw);
+	if (!bIsRemoved)
+	{
+		const float4 P0Clipped = lerp(P0, P1, T.x);
+		const float4 P1Clipped = lerp(P0, P1, T.y);
+		P0 = P0Clipped;
+		P1 = P1Clipped;
+	}
 
-	Rad = CP.WorldRadius*2000.0; // OutputResolutionf.x; //TODO: figure this out correctly?
+	return !bIsRemoved;
 }
 
-bool ClipRaySegment(float3 AABBMin, float3 AABBMax, inout float3 P0, inout float3 P1)
+bool ClipRaySegment(float2 AABBMin, float2 AABBMax, inout float4 P0, inout float4 P1)
 {
-	const bool bP0Outside = any(P0 < AABBMin) || any(P0 > AABBMax);
-	const bool bP1Outside = any(P1 < AABBMin) || any(P1 > AABBMax);
+	const bool bP0Outside = any(P0.xy < AABBMin) || any(P0.xy > AABBMax);
+	const bool bP1Outside = any(P1.xy < AABBMin) || any(P1.xy > AABBMax);
 	if (!bP0Outside && !bP1Outside)
 	{
 		return true;
 	}
 
-	const float3 Origin = P0;
-	const float3 Dir = P1 - P0;
-	const float3 RcpDir = 1.0f / Dir;
+	const float2 Origin = P0.xy;
+	const float2 Dir = P1.xy - P0.xy;
+	const float2 RcpDir = 1.0f / Dir;
 
-	const float3 T0 = (AABBMin - Origin) * RcpDir;
-	const float3 T1 = (AABBMax - Origin) * RcpDir;
+	const float2 T0 = (AABBMin - Origin) * RcpDir;
+	const float2 T1 = (AABBMax - Origin) * RcpDir;
 
-	const float TMin = max(max(min(T0.x, T1.x), min(T0.y, T1.y)), min(T0.z, T1.z));
-	const float TMax = min(min(max(T0.x, T1.x), max(T0.y, T1.y)), max(T0.z, T1.z));
+	const float TMin = max(min(T0.x, T1.x), min(T0.y, T1.y));
+	const float TMax = min(max(T0.x, T1.x), max(T0.y, T1.y));
 
 	// Ray intersects the AABB but the segment is completely outside.
 	if (TMax < 0.0f)
@@ -159,19 +200,19 @@ bool ClipRaySegment(float3 AABBMin, float3 AABBMax, inout float3 P0, inout float
 		return false;
 	}
 
-	float3 P0New = P0;
-	float3 P1New = P1;
-	if (bP0Outside && TMin >= 0.0f && TMin < 1.0f)
+	float4 P0New = P0;
+	float4 P1New = P1;
+	if (bP0Outside && TMin > 0.0f && TMin < 1.0f)
 	{
 		P0New = lerp(P0, P1, TMin);
 	}
-	if (bP1Outside && TMax >= 0.0f && TMax < 1.0f)
+	if (bP1Outside && TMax > 0.0f && TMax < 1.0f)
 	{
 		P1New = lerp(P0, P1, TMax);
 	}
 	P0 = P0New;
 	P1 = P1New;
-	
+
 	return true;
 }
 
@@ -244,7 +285,7 @@ RWByteAddressBuffer					OutVisTileData;
 
 struct FDDAContext
 {
-	float2 TileCoord;
+	float2 Coord;
 	float2 DeltaDist;
 	float2 Step;
 	float2 SideDist;
@@ -255,10 +296,10 @@ FDDAContext DDACreateContext(float2 RayStart, float2 RayDir)
 	const float2 RayDirRcp = 1.0f / RayDir;
 
 	FDDAContext Context;
-	Context.TileCoord = floor(RayStart);
+	Context.Coord = floor(RayStart);
 	Context.DeltaDist = abs(RayDirRcp);
 	Context.Step = sign(RayDir);
-	Context.SideDist = (Context.TileCoord - RayStart + max(Context.Step, 0.0f)) * RayDirRcp;
+	Context.SideDist = (Context.Coord - RayStart + max(Context.Step, 0.0f)) * RayDirRcp;
 
 	return Context;
 }
@@ -268,12 +309,12 @@ void DDAAdvance(inout FDDAContext Context)
 	if (Context.SideDist.x < Context.SideDist.y)
 	{
 		Context.SideDist.x += Context.DeltaDist.x;
-		Context.TileCoord.x += Context.Step.x;
+		Context.Coord.x += Context.Step.x;
 	}
 	else
 	{
 		Context.SideDist.y += Context.DeltaDist.y;
-		Context.TileCoord.y += Context.Step.y;
+		Context.Coord.y += Context.Step.y;
 	}
 }
 
@@ -385,25 +426,28 @@ void BinningCS(uint DispatchThreadID : SV_DispatchThreadID, uint GroupThreadID :
 		// Project segment end points and clip them to the screen
 		if (bSegValid)
 		{
-			float3 SP0 = 0.0f;
+			float4 H0 = 0.0f;
+			float4 H1 = 0.0f;
 			uint Type = -1;
-			CalcScreenPos(HairStrandsVF_PositionBuffer[PrimID], HairStrandsVF_PositionOffsetBuffer[0].xyz, SP0, Type);
+			CalcHomogenousPos(HairStrandsVF_PositionBuffer[PrimID], HairStrandsVF_PositionOffsetBuffer[0].xyz, H0, Type);
 
 			bool bIsEndCV = (Type == HAIR_CONTROLPOINT_END);
 			bSegValid = !bIsEndCV;
 
 			if (bSegValid)
 			{
-				float3 SP1 = 0.0f;
-				CalcScreenPos(HairStrandsVF_PositionBuffer[PrimID + 1], HairStrandsVF_PositionOffsetBuffer[0].xyz, SP1, Type);
+				CalcHomogenousPos(HairStrandsVF_PositionBuffer[PrimID + 1], HairStrandsVF_PositionOffsetBuffer[0].xyz, H1, Type);
 				
-				// Do clipping in tile space
-				SP0.xy *= RcpTileSize;
-				SP1.xy *= RcpTileSize;
-				bSegValid = ClipRaySegment(0.0f, float3(TileRes, 1.0f), SP0, SP1);
+				// Do clipping in homogenous coordinates
+				bSegValid = BlinnLineClipping(H0, H1);
 
 				if (bSegValid)
 				{
+					float3 SP0 = NDCToPixelCoord(H0);
+					float3 SP1 = NDCToPixelCoord(H1);
+					SP0.xy *= RcpTileSize;
+					SP1.xy *= RcpTileSize;
+
 					// For peace of mind, make sure these are actually clamped to a valid range.
 					SP0 = clamp(SP0, 0.0f, float3(TileRes, 1.0f));
 					SP1 = clamp(SP1, 0.0f, float3(TileRes, 1.0f));
@@ -430,7 +474,7 @@ void BinningCS(uint DispatchThreadID : SV_DispatchThreadID, uint GroupThreadID :
 
 			for (int DDAIt = 0; DDAIt < DDA_MAX_ITERATIONS; ++DDAIt)
 			{
-				const int2 TileCoord = (int2)floor(DDAContext.TileCoord);
+				const int2 TileCoord = (int2)floor(DDAContext.Coord);
 
 				BRANCH
 				if (NearestDepth > VisTileDepthGrid[TileCoord])
@@ -498,7 +542,7 @@ void BinningCS(uint DispatchThreadID : SV_DispatchThreadID, uint GroupThreadID :
 
 			for (int DDAIt = 0; DDAIt < DDA_MAX_ITERATIONS; ++DDAIt)
 			{
-				const int2 TileCoord = (int2)floor(DDAContext.TileCoord);
+				const int2 TileCoord = (int2)floor(DDAContext.Coord);
 
 				BRANCH
 				if (NearestDepth > VisTileDepthGrid[TileCoord])
@@ -543,6 +587,17 @@ void BinningCS(uint DispatchThreadID : SV_DispatchThreadID, uint GroupThreadID :
 
 #if SHADER_RASTERCOMPUTE_RASTER
 
+// Simple rasterization algorithm that lerps between line endpoints. Is currently more robust than the Wu algorithm
+// and optionally supports anti-aliasing similar to the Wu algorithm.
+#define RASTER_LINEAR 0 
+// Implementation of Wu's line rasterization algorithm. Currently this implementation has tile shaped artifacts when the line segment is
+// clipped against the tile which is why we use the simple linear algorithm at the moment.
+#define RASTER_WU 1
+ // Set to 1 to enable writing to two pixels straddling the line segment when using the linear rasterization algorithm.
+#define ENABLE_RASTER_LINEAR_AA 0
+
+#define RASTER_ALGO RASTER_LINEAR
+
 RWTexture2D<uint> 					RWVisTileDepthGrid;
 Buffer<uint> 						VisTilePrims;
 Buffer<uint> 						VisTileArgs;
@@ -584,6 +639,65 @@ groupshared uint group_PrimMatID[GS_SEGS];
 
 groupshared uint group_TileIndex;
 
+void PlotInternal(int2 Coords, float AntiAliasingFactor, float4 P0, float4 P1, float Rad0, float Rad1, float SegmentLenSqRcp, uint PackedTileMin, uint PrimMatID)
+{
+	const int2 IntraTileCoord = Coords - int2(((PackedTileMin >> 0) & 0xffff), ((PackedTileMin >> 16) & 0xffff));
+
+	if (all(IntraTileCoord >= 0) && all(IntraTileCoord < TileSize))
+	{
+		// Project P onto line segment and compute the lerp alpha between P0 and P1
+		// Simplification of:
+		// A = P - P0
+		// B = P1 - P0
+		// Alpha = dot(A, B) / dot(B, B)
+		const float2 P = Coords + 0.5f;
+		const float Alpha = saturate(dot(P - P0.xy, P1.xy - P0.xy) * SegmentLenSqRcp);
+		const float Depth = lerp(P0.z, P1.z, Alpha);
+		const uint PackedDepthCov = PackHairVisDepthCoverage(Depth, 1.0f);
+		const uint LinearIndex = IntraTileCoord.x + IntraTileCoord.y * TileSize;
+
+		// Write Depth + PrimMatID if depth test against hair depths is passed
+		uint OldValue;
+		InterlockedMax(group_SubTile[LinearIndex].x, PackedDepthCov, OldValue);
+		if (PackedDepthCov > OldValue)
+		{
+			group_SubTile[LinearIndex].y = PrimMatID;
+		}
+
+		// Add hair count if depth test against scene depth is passed
+		if (PackedDepthCov > group_SubTile[LinearIndex].w)
+		{
+			// Alpha value for perspective correct interpolation. We store the reciprocal of w in the w compoennt,
+			// so this is a simplification of:
+			// (Alpha / w1) / ((1 - Alpha) / w0 + Alpha / w1)
+			const float LerpedRcpW = lerp(P0.w, P1.w, Alpha);
+			const float PerspectiveAlpha = (Alpha * P1.w) / LerpedRcpW;
+			// Divide by W to make thickness dependent on screen space depth? This division was kept from the previous line rasterization algorithm.
+			const float Rad = lerp(Rad0, Rad1, PerspectiveAlpha) * LerpedRcpW;
+
+			InterlockedAdd(group_SubTile[LinearIndex].z, min(Rad, 0.5f) * 2.0f * 1000.0f * SampleWeight * AntiAliasingFactor);
+		}
+	}
+}
+
+void Plot(int2 Coord, float FracY, float AntiAliasingFactor, bool bIsSteep, float4 P0, float4 P1, float Rad0, float Rad1, float SegmentLenSqRcp, uint PackedTileMin, uint PrimMatID)
+{
+	// First pixel
+	{
+		float AAFactor = AntiAliasingFactor * (1.0f - FracY);
+
+		PlotInternal(bIsSteep ? Coord.yx : Coord.xy, AAFactor, P0, P1, Rad0, Rad1, SegmentLenSqRcp, PackedTileMin, PrimMatID);
+	}
+
+	// Second pixel
+	{
+		float AAFactor = AntiAliasingFactor * FracY;
+		Coord.y += 1;
+
+		PlotInternal(bIsSteep ? Coord.yx : Coord.xy, AAFactor, P0, P1, Rad0, Rad1, SegmentLenSqRcp, PackedTileMin, PrimMatID);
+	}
+}
+
 [numthreads(1024, 1, 1)]
 void RasterCS(uint DispatchThreadID : SV_DispatchThreadID, uint GroupThreadID : SV_GroupThreadID, uint GroupID : SV_GroupID)
 {
@@ -592,7 +706,7 @@ void RasterCS(uint DispatchThreadID : SV_DispatchThreadID, uint GroupThreadID :
 	if (GroupThreadID == 0)
 	{
 		group_TileNum = VisTileArgs[0];
-		group_LoopNum = (float(group_TileNum) + float(NumRasterizers - 1))* RcpNumRasterizers;
+		group_LoopNum = (float(group_TileNum) + float(NumRasterizers - 1)) * RcpNumRasterizers;
 
 		group_PositionOffset = HairStrandsVF_PositionOffsetBuffer[0].xyz;
 
@@ -604,7 +718,7 @@ void RasterCS(uint DispatchThreadID : SV_DispatchThreadID, uint GroupThreadID :
 
 	GroupMemoryBarrierWithGroupSync();
 
-	[loop]
+	LOOP
 	for (uint LoopIndex = 0; LoopIndex < group_LoopNum; LoopIndex++)
 	{
 		if (GroupThreadID == 0)
@@ -614,13 +728,13 @@ void RasterCS(uint DispatchThreadID : SV_DispatchThreadID, uint GroupThreadID :
 		
 		GroupMemoryBarrierWithGroupSync();
 
-		bool bTileValid = (group_TileIndex < group_TileNum) ? true : false;
+		bool bTileValid = (group_TileIndex < group_TileNum);
 
 		uint PrimOffset = group_TileIndex * 1024;
 		uint PrimCount = LoadRWVisTileData(group_TileIndex, VT_PrimCount);
 
 		uint PackedCoord = LoadRWVisTileData(group_TileIndex, VT_Coord);
-		uint2 SubTileMin = UppackVisTileCoord(PackedCoord)*TileSize;
+		uint2 SubTileMin = UppackVisTileCoord(PackedCoord) * TileSize;
 
 		uint PackedTileMin = ((SubTileMin.x & 0xffff) << 0) | ((SubTileMin.y & 0xffff) << 16);
 
@@ -650,16 +764,16 @@ void RasterCS(uint DispatchThreadID : SV_DispatchThreadID, uint GroupThreadID :
 		
 		GroupMemoryBarrierWithGroupSync();
 
-		bool bThreadValid = (bTileValid && (GroupThreadID < (PrimCount*group_ThreadsPerSeg) )) ? true : false;
+		bool bThreadValid = (bTileValid && (GroupThreadID < (PrimCount * group_ThreadsPerSeg)));
 
-		uint WaveCount = ((PrimCount* group_ThreadsPerSeg) + (WAVE_SIZE-1) ) / WAVE_SIZE;
+		uint WaveCount = ((PrimCount * group_ThreadsPerSeg) + (WAVE_SIZE - 1) ) / WAVE_SIZE;
 		uint WaveThreadCount = WaveCount * WAVE_SIZE;
 
-		bool bWaveThreadValid = (bTileValid && (GroupThreadID < WaveThreadCount)) ? true : false;
+		bool bWaveThreadValid = (bTileValid && (GroupThreadID < WaveThreadCount));
 
-		bool bUseGroupSPs = (bThreadValid && (GroupThreadID < (min(PrimCount, GS_SEGS)* group_ThreadsPerSeg) )) ? true : false;
+		bool bUseGroupSPs = (bThreadValid && (GroupThreadID < (min(PrimCount, GS_SEGS) * group_ThreadsPerSeg)));
 
-		bool bGenGroupSPs = (bThreadValid && (GroupThreadID < (min(PrimCount, GS_SEGS)) )) ? true : false;
+		bool bGenGroupSPs = (bThreadValid && (GroupThreadID < (min(PrimCount, GS_SEGS))));
 
 		if (bGenGroupSPs)
 		{
@@ -668,8 +782,8 @@ void RasterCS(uint DispatchThreadID : SV_DispatchThreadID, uint GroupThreadID :
 
 			group_PrimMatID[Prim] = PackHairVisPrimitiveMaterialId(PrimID, HairMaterialId);
 
-			CalcScreenPosAndRad(HairStrandsVF_PositionBuffer[PrimID], group_PositionOffset, group_SP0[Prim], group_Rad0[Prim]);
-			CalcScreenPosAndRad(HairStrandsVF_PositionBuffer[PrimID+1], group_PositionOffset, group_SP1[Prim], group_Rad1[Prim]);
+			CalcHomogenousPosAndRad(HairStrandsVF_PositionBuffer[PrimID], group_PositionOffset, group_SP0[Prim], group_Rad0[Prim]);
+			CalcHomogenousPosAndRad(HairStrandsVF_PositionBuffer[PrimID+1], group_PositionOffset, group_SP1[Prim], group_Rad1[Prim]);
 		}
 
 		if (bWaveThreadValid)
@@ -678,7 +792,7 @@ void RasterCS(uint DispatchThreadID : SV_DispatchThreadID, uint GroupThreadID :
 			{
 				uint2 Coord;
 
-				Coord.y = (float(LinearIndex) + 0.5) * RcpTileSize;
+				Coord.y = (float(LinearIndex) + 0.5f) * RcpTileSize;
 				Coord.x = LinearIndex - (Coord.y * TileSize);
 
 				Coord += uint2(((PackedTileMin >> 0) & 0xffff), ((PackedTileMin >> 16) & 0xffff));
@@ -686,7 +800,7 @@ void RasterCS(uint DispatchThreadID : SV_DispatchThreadID, uint GroupThreadID :
 				group_SubTile[LinearIndex].x = OutDepthCovTexture[Coord];
 				group_SubTile[LinearIndex].y = 0;
 				group_SubTile[LinearIndex].z = 0;
-				group_SubTile[LinearIndex].w = PackHairVisDepthCoverage(SceneDepthTexture.Load(uint3(Coord, 0)), 1.0);
+				group_SubTile[LinearIndex].w = PackHairVisDepthCoverage(SceneDepthTexture.Load(uint3(Coord, 0)), 1.0f);
 			}
 		}
 
@@ -694,7 +808,7 @@ void RasterCS(uint DispatchThreadID : SV_DispatchThreadID, uint GroupThreadID :
 
 		if (bThreadValid)
 		{
-			uint Prim = uint((float(GroupThreadID)+0.5) / float(group_ThreadsPerSeg));
+			uint Prim = uint((float(GroupThreadID) + 0.5f) / float(group_ThreadsPerSeg));
 			uint PModTPS = GroupThreadID - (Prim * group_ThreadsPerSeg);
 
 			uint PrimMatID;
@@ -716,40 +830,121 @@ void RasterCS(uint DispatchThreadID : SV_DispatchThreadID, uint GroupThreadID :
 				uint PrimID = VisTilePrims[PrimOffset + Prim];
 				PrimMatID = PackHairVisPrimitiveMaterialId(PrimID, HairMaterialId);
 
-				CalcScreenPosAndRad(HairStrandsVF_PositionBuffer[PrimID], group_PositionOffset, SP0, Rad0);
-				CalcScreenPosAndRad(HairStrandsVF_PositionBuffer[PrimID+1], group_PositionOffset, SP1, Rad1);
+				CalcHomogenousPosAndRad(HairStrandsVF_PositionBuffer[PrimID], group_PositionOffset, SP0, Rad0);
+				CalcHomogenousPosAndRad(HairStrandsVF_PositionBuffer[PrimID + 1], group_PositionOffset, SP1, Rad1);
 			}
 
-			uint NumSteps = min(max(abs(SP1.x - SP0.x), abs(SP1.y - SP0.y))+1.0, TileSize);
+			// Clipping
+			{
+				SP0 = float4(NDCToPixelCoord(SP0), 1.0f / SP0.w);
+				SP1 = float4(NDCToPixelCoord(SP1), 1.0f / SP1.w);
+				
+				// Clip against tile
+				const float2 TileMin = float2(((PackedTileMin >> 0) & 0xffff), ((PackedTileMin >> 16) & 0xffff));
+				const float2 TileMax = TileMin + TileSize;
+				// Clip against a slightly smaller bounding box so that the clipped endpoints are safely inside the tile.
+				ClipRaySegment(TileMin + 0.25f, TileMax - 0.25f, SP0, SP1);
+			}
+
+			const float SegmentLenSqRcp = 1.0f / dot(SP1.xy - SP0.xy, SP1.xy - SP0.xy);
+
+#if RASTER_ALGO == RASTER_LINEAR
 
-			[loop]
-			for (int J = PModTPS; J < NumSteps; J+=group_ThreadsPerSeg)
+			const float2 DeltaAbs = abs(SP1.xy - SP0.xy);
+			const uint NumSteps = min((uint)ceil(max(DeltaAbs.x, DeltaAbs.y)) + 1, TileSize + (ENABLE_RASTER_LINEAR_AA ? 1 : 0));
+			const float RcpNumSteps = NumSteps > 1 ? (1.0f / (NumSteps - 1)) : 1.0f;
+			const bool bIsSteep = DeltaAbs.y > DeltaAbs.x;
+
+			LOOP
+			for (int J = PModTPS; J < NumSteps; J += group_ThreadsPerSeg)
 			{
-				const float4 SP = lerp(SP0, SP1, J / float(NumSteps));
+				const float Alpha = J * RcpNumSteps;
+				const float4 SP = lerp(SP0, SP1, Alpha);
 
-				const int2 SPi = int2(SP.xy)-int2(((PackedTileMin >> 0) & 0xffff), ((PackedTileMin >> 16) & 0xffff));
+				const float AntiAliasingFactor = 1.0f;
+#if !ENABLE_RASTER_LINEAR_AA
+				PlotInternal(SP.xy, AntiAliasingFactor, SP0, SP1, Rad0, Rad1, SegmentLenSqRcp, PackedTileMin, PrimMatID);
+#else
+				const float2 Coord = (bIsSteep ? SP.yx : SP.xy) - 0.5f;
+				const float FracY = frac(Coord.y);
+				Plot(Coord, FracY, AntiAliasingFactor, bIsSteep, SP0, SP1, Rad0, Rad1, SegmentLenSqRcp, PackedTileMin, PrimMatID);
+#endif // !ENABLE_RASTER_LINEAR_AA
+			}
+#elif RASTER_ALGO == RASTER_WU
+			// Wu's line algorithm. Currently this has some weird artifacts when clipping to tiles.
+			// TODO: Remove this entirely or fix the artifacts.
+			{
+				const bool bIsSteep = abs(SP1.y - SP0.y) > abs(SP1.x - SP0.x);
 
-				if (SPi.x >= 0 && SPi.x < TileSize && SPi.y >= 0 && SPi.y < TileSize)
+				if (bIsSteep)
 				{
-					uint PackedDepthCov = PackHairVisDepthCoverage(SP.z / SP.w, 1.0);
+					SP0.xy = SP0.yx;
+					SP1.xy = SP1.yx;
+				}
+				if (SP0.x > SP1.x)
+				{
+					float4 Tmp = SP0;
+					SP0 = SP1;
+					SP1 = Tmp;
+				}
 
-					uint LinearIndex = SPi.x + SPi.y * TileSize;
+				const float2 D = SP1.xy - SP0.xy;
+				const float Gradient = abs(D.x) < 1e-5f ? 1.0f : D.y / D.x;
 
-					uint OldValue;
-					InterlockedMax(group_SubTile[LinearIndex].x, PackedDepthCov, OldValue);
-					if (PackedDepthCov > OldValue)
+				float DeltaY = 0.0f;
+
+				// First endpoint
+				int2 Px0;
+				{
+					const float2 SP0Int = SP0.xy - 0.5f; // transform to integer grid with pixel centers at 0.0 instead of 0.5.
+					float2 End;
+					End.x = floor(SP0Int.x);
+					End.y = SP0Int.y + Gradient * (End.x - SP0Int.x);
+
+					const float GapX = 1.0f;// 1.0f - frac(SP0Int.x + 0.5f);
+
+					Px0 = int2(End.x, floor(End.y));
+					
+					if (PModTPS == 0)
 					{
-						group_SubTile[LinearIndex].y = PrimMatID;
+						const float FracY = frac(End.y);
+						Plot(Px0, FracY, GapX, bIsSteep, SP0, SP1, Rad0, Rad1, SegmentLenSqRcp, PackedTileMin, PrimMatID);
 					}
 
-					if (PackedDepthCov > group_SubTile[LinearIndex].w)
-					{
-						const float Rad = lerp(Rad0, Rad1, J / float(NumSteps)) / SP.w;
+					DeltaY = End.y + Gradient; // First y-intersection for the main loop
+				}
 
-						InterlockedAdd(group_SubTile[LinearIndex].z, min(Rad, 0.5) * 2.0 * 1000.0 * SampleWeight);
+				// Second endpoint
+				int2 Px1;
+				{
+					const float2 SP1Int = SP1.xy - 0.5f; // transform to integer grid with pixel centers at 0.0 instead of 0.5.
+					float2 End;
+					End.x = floor(SP1Int.x);
+					End.y = SP1Int.y + Gradient * (End.x - SP1Int.x);
+					const float GapX = 1.0f;// frac(SP1Int.x + 0.5f);
+
+					Px1 = float2(End.x, floor(End.y));
+
+					if (PModTPS == 0)
+					{
+						const float FracY = frac(End.y);
+						Plot(Px1, FracY, GapX, bIsSteep, SP0, SP1, Rad0, Rad1, SegmentLenSqRcp, PackedTileMin, PrimMatID);
 					}
-				}					
+				}
+
+				// Main loop
+				const int XBegin = Px0.x + 1 + PModTPS;
+				const int XEnd = Px1.x;
+				DeltaY += PModTPS * Gradient;
+				for (int X = XBegin; X < XEnd; X += group_ThreadsPerSeg)
+				{
+					const int2 Coord = int2(X, floor(DeltaY));
+					const float FracY = frac(DeltaY);
+					Plot(Coord, FracY, 1.0f, bIsSteep, SP0, SP1, Rad0, Rad1, SegmentLenSqRcp, PackedTileMin, PrimMatID);
+					DeltaY += group_ThreadsPerSeg * Gradient;
+				}
 			}
+#endif // RASTER_ALGO == RASTER_LINEAR
 		}
 
 		GroupMemoryBarrierWithGroupSync();
@@ -760,7 +955,7 @@ void RasterCS(uint DispatchThreadID : SV_DispatchThreadID, uint GroupThreadID :
 			{
 				uint2 Coord;
 
-				Coord.y = (float(LinearIndex) + 0.5) * RcpTileSize;
+				Coord.y = (float(LinearIndex) + 0.5f) * RcpTileSize;
 				Coord.x = LinearIndex - (Coord.y * TileSize);
 
 				Coord += uint2(((PackedTileMin >> 0) & 0xffff), ((PackedTileMin >> 16) & 0xffff));
