Fix buffer size mixup causing crashes in VSM page marking & incorrect pixel position causing OOB buffer access (and presumably incorrect shadows).
- Also fetch correct light data when the light list is compacted.

diff --git a/Engine/Shaders/Private/VirtualShadowMaps/VirtualShadowMapPageManagement.usf b/Engine/Shaders/Private/VirtualShadowMaps/VirtualShadowMapPageManagement.usf
--- a/Engine/Shaders/Private/VirtualShadowMaps/VirtualShadowMapPageManagement.usf
+++ b/Engine/Shaders/Private/VirtualShadowMaps/VirtualShadowMapPageManagement.usf
@@ -33,6 +33,32 @@
 #define INPUT_TYPE_GBUFFER 0
 #define INPUT_TYPE_HAIRSTRANDS 1
 
+/** 
+ * Load the light data for the light grid using the 'Local light index' which is the index stored in the light grid data structures.
+ * Does not handle stereo switching so use only from shaders where you know that the light grid is bound to ForwardLightData.
+ */
+// TODO: Move to LightGridCommon.ush when we can modify global shaders without incurring a week-long delay.
+// TODO2: Unify data storage for local light info into GPU-Scene and remove special index types.
+FLocalLightData GetLocalLightDataFromLocalLightIndexNonStereo(uint LocalLightIndex)
+{
+	FLocalLightData Result;
+
+	uint LocalLightBaseIndex = LocalLightIndex * LOCAL_LIGHT_DATA_STRIDE;
+	Result.bClusteredDeferredSupported = LocalLightIndex < ForwardLightData.ClusteredDeferredSupportedEndIndex;
+	Result.bLumenLightSupported = LocalLightIndex >= ForwardLightData.LumenSupportedStartIndex;
+	Result.bIsSimpleLight = LocalLightIndex < ForwardLightData.SimpleLightsEndIndex;
+	Result.LightPositionAndInvRadius = ForwardLightData.ForwardLocalLightBuffer[LocalLightBaseIndex + 0];
+	Result.LightColorAndFalloffExponent = ForwardLightData.ForwardLocalLightBuffer[LocalLightBaseIndex + 1];
+	Result.LightDirectionAndShadowMask = ForwardLightData.ForwardLocalLightBuffer[LocalLightBaseIndex + 2];
+	Result.SpotAnglesAndSourceRadiusPacked = ForwardLightData.ForwardLocalLightBuffer[LocalLightBaseIndex + 3];
+	Result.LightTangentAndSoftSourceRadius = ForwardLightData.ForwardLocalLightBuffer[LocalLightBaseIndex + 4];
+	Result.RectBarnDoorAndVirtualShadowMapIdAndSpecularScale = ForwardLightData.ForwardLocalLightBuffer[LocalLightBaseIndex + 5];
+	Result.VirtualShadowMapId = int(Result.RectBarnDoorAndVirtualShadowMapIdAndSpecularScale.z);
+
+	return Result;
+}
+
+
 // Flags generated by per-pixel pass to determine which pages are required to provide shadow for the visible geometry
 RWStructuredBuffer<uint> OutPageRequestFlags;
 
@@ -261,19 +287,24 @@ void GeneratePageFlagsFromPixels(
 		const bool bUsePageDilation = PageDilationBorderSizeLocal > 0.0f;
 		const float2 PageDilationOffset = PageDilationBorderSizeLocal * PageDilationDither;
 
-		uint EyeIndex = 0; // ??
-		uint3 GridCoordinate = ComputeLightGridCellCoordinate(PixelPos, SceneDepth, EyeIndex);
+		uint EyeIndex = 0U; // Has no effect unless INSTANCED_STEREO is enabled, which is not the case for this CS, (used in GetLightGridData to select ForwardLightData or ForwardLightDataISR, but then only uses redundant parameters anyway)
+		uint2 LocalPosition = PixelPos - uint2(View.ViewRectMin.xy);
+		uint3 GridCoordinate = ComputeLightGridCellCoordinate(LocalPosition, SceneDepth, EyeIndex);
+
 		uint GridLinearIndex = ComputeLightGridCellIndex(GridCoordinate, EyeIndex);
 		const FCulledLightsGridData CulledLightGridData = GetCulledLightsGrid(GridLinearIndex, EyeIndex);
 
+		checkStructuredBufferAccessSlow(PrunedNumCulledLightsGrid, GridLinearIndex);
 		const uint PrunedNumLocalLights = PrunedNumCulledLightsGrid[GridLinearIndex];
 
 		LOOP
 		for (uint Index = 0; Index < PrunedNumLocalLights; ++Index)
 		{
-			const FLocalLightData LightData = GetLocalLightData(CulledLightGridData.DataStartIndex + Index, EyeIndex);
+			checkStructuredBufferAccessSlow(PrunedLightGridData, CulledLightGridData.DataStartIndex + Index);
 			const uint LightGridLightIndex = PrunedLightGridData[CulledLightGridData.DataStartIndex + Index];
 
+			const FLocalLightData LightData = GetLocalLightDataFromLocalLightIndexNonStereo(LightGridLightIndex);
+
 			// Relative to PrimaryView
 			const float3 LightTranslatedWorldPosition = LightData.LightPositionAndInvRadius.xyz;
 			float LengthSquared = length2(LightTranslatedWorldPosition - TranslatedWorldPosition);
diff --git a/Engine/Source/Runtime/Renderer/Private/VirtualShadowMaps/VirtualShadowMapArray.cpp b/Engine/Source/Runtime/Renderer/Private/VirtualShadowMaps/VirtualShadowMapArray.cpp
--- a/Engine/Source/Runtime/Renderer/Private/VirtualShadowMaps/VirtualShadowMapArray.cpp
+++ b/Engine/Source/Runtime/Renderer/Private/VirtualShadowMaps/VirtualShadowMapArray.cpp
@@ -1295,11 +1295,11 @@ void FVirtualShadowMapArray::BuildPageAllocations(
 				// It's currently safe to overlap these passes that all write to same page request flags
 				FRDGBufferUAVRef PageRequestFlagsUAV = GraphBuilder.CreateUAV(PageRequestFlagsRDG, ERDGUnorderedAccessViewFlags::SkipBarrier);
 
-				const FIntVector& CulledGridSize = View.ForwardLightingResources.ForwardLightData->CulledGridSize;
-				const uint32 NumLightGridEntries = CulledGridSize.X * CulledGridSize.Y * CulledGridSize.Z;
-				FRDGBufferRef PrunedLightGridDataRDG = GraphBuilder.CreateBuffer(FRDGBufferDesc::CreateStructuredDesc(sizeof(uint32), NumLightGridEntries), TEXT("Shadow.Virtual.PrunedLightGridData"));
 				uint32 LightGridIndexBufferSize = View.ForwardLightingResources.ForwardLightData->CulledLightDataGrid->Desc.Buffer->Desc.GetSize();
-				FRDGBufferRef PrunedNumCulledLightsGridRDG = GraphBuilder.CreateBuffer(FRDGBufferDesc::CreateStructuredDesc(sizeof(uint32), LightGridIndexBufferSize), TEXT("Shadow.Virtual.PrunedNumCulledLightsGrid"));
+				FRDGBufferRef PrunedLightGridDataRDG = GraphBuilder.CreateBuffer(FRDGBufferDesc::CreateStructuredDesc(sizeof(uint32), LightGridIndexBufferSize), TEXT("Shadow.Virtual.PrunedLightGridData"));
+				
+				const uint32 NumLightGridCells = View.ForwardLightingResources.ForwardLightData->NumGridCells;
+				FRDGBufferRef PrunedNumCulledLightsGridRDG = GraphBuilder.CreateBuffer(FRDGBufferDesc::CreateStructuredDesc(sizeof(uint32), NumLightGridCells), TEXT("Shadow.Virtual.PrunedNumCulledLightsGrid"));
 
 				{
 					FPruneLightGridCS::FParameters* PassParameters = GraphBuilder.AllocParameters< FPruneLightGridCS::FParameters >();
@@ -1310,7 +1310,7 @@ void FVirtualShadowMapArray::BuildPageAllocations(
 					PassParameters->OutPrunedNumCulledLightsGrid = GraphBuilder.CreateUAV(PrunedNumCulledLightsGridRDG);
 					auto ComputeShader = View.ShaderMap->GetShader<FPruneLightGridCS>();
 
-					FComputeShaderUtils::AddPass(GraphBuilder, RDG_EVENT_NAME("PruneLightGrid"), ComputeShader, PassParameters, FComputeShaderUtils::GetGroupCount(NumLightGridEntries, FPruneLightGridCS::DefaultCSGroupX));
+					FComputeShaderUtils::AddPass(GraphBuilder, RDG_EVENT_NAME("PruneLightGrid"), ComputeShader, PassParameters, FComputeShaderUtils::GetGroupCount(NumLightGridCells, FPruneLightGridCS::DefaultCSGroupX));
 				};
 
 
