Linearizes all groupsizes of TSR shaders to maximise compatibility on desktop with SM6

diff --git a/Engine/Shaders/Private/TemporalSuperResolution/TSRClearPrevTextures.usf b/Engine/Shaders/Private/TemporalSuperResolution/TSRClearPrevTextures.usf
--- a/Engine/Shaders/Private/TemporalSuperResolution/TSRClearPrevTextures.usf
+++ b/Engine/Shaders/Private/TemporalSuperResolution/TSRClearPrevTextures.usf
@@ -18,7 +18,7 @@ RWTexture2D<uint> PrevClosestDepthOutput;
 
 //------------------------------------------------------- ENTRY POINT
 
-[numthreads(TILE_SIZE, TILE_SIZE, 1)]
+[numthreads(TILE_SIZE * TILE_SIZE, 1, 1)]
 void MainCS(
 	uint2 GroupId : SV_GroupID,
 	uint GroupThreadIndex : SV_GroupIndex)
diff --git a/Engine/Shaders/Private/TemporalSuperResolution/TSRDecimateHistory.usf b/Engine/Shaders/Private/TemporalSuperResolution/TSRDecimateHistory.usf
--- a/Engine/Shaders/Private/TemporalSuperResolution/TSRDecimateHistory.usf
+++ b/Engine/Shaders/Private/TemporalSuperResolution/TSRDecimateHistory.usf
@@ -71,7 +71,7 @@ RWTexture2D<tsr_subpixel_payload> HistorySubpixelDetailsOutput;
 
 //------------------------------------------------------- ENTRY POINT
 
-[numthreads(TILE_SIZE, TILE_SIZE, 1)]
+[numthreads(TILE_SIZE * TILE_SIZE, 1, 1)]
 void MainCS(
 	uint2 GroupId : SV_GroupID,
 	uint GroupThreadIndex : SV_GroupIndex)
diff --git a/Engine/Shaders/Private/TemporalSuperResolution/TSRDilateRejection.usf b/Engine/Shaders/Private/TemporalSuperResolution/TSRDilateRejection.usf
--- a/Engine/Shaders/Private/TemporalSuperResolution/TSRDilateRejection.usf
+++ b/Engine/Shaders/Private/TemporalSuperResolution/TSRDilateRejection.usf
@@ -19,7 +19,7 @@ RWTexture2D<tsr_half> DilatedHistoryRejectionOutput;
 
 //------------------------------------------------------- ENTRY POINT
 
-[numthreads(TILE_SIZE, TILE_SIZE, 1)]
+[numthreads(TILE_SIZE * TILE_SIZE, 1, 1)]
 void MainCS(
 	uint2 GroupId : SV_GroupID,
 	uint GroupThreadIndex : SV_GroupIndex)
diff --git a/Engine/Shaders/Private/TemporalSuperResolution/TSRDilateVelocity.usf b/Engine/Shaders/Private/TemporalSuperResolution/TSRDilateVelocity.usf
--- a/Engine/Shaders/Private/TemporalSuperResolution/TSRDilateVelocity.usf
+++ b/Engine/Shaders/Private/TemporalSuperResolution/TSRDilateVelocity.usf
@@ -63,7 +63,7 @@ void UnpackDepthAndOffset(uint Packed, out int2 Offset, out float DeviceZ)
 
 //------------------------------------------------------- ENTRY POINT
 
-[numthreads(TILE_SIZE, TILE_SIZE, 1)]
+[numthreads(TILE_SIZE * TILE_SIZE, 1, 1)]
 void MainCS(
 	uint2 GroupId : SV_GroupID,
 	uint GroupThreadIndex : SV_GroupIndex)
diff --git a/Engine/Shaders/Private/TemporalSuperResolution/TSRFilterAntiAliasing.usf b/Engine/Shaders/Private/TemporalSuperResolution/TSRFilterAntiAliasing.usf
--- a/Engine/Shaders/Private/TemporalSuperResolution/TSRFilterAntiAliasing.usf
+++ b/Engine/Shaders/Private/TemporalSuperResolution/TSRFilterAntiAliasing.usf
@@ -30,7 +30,7 @@ tsr_half2 SampleAntiAliasing(tsr_short2 KernelCenter, tsr_short2 Offset)
 
 //------------------------------------------------------- ENTRY POINT
 
-[numthreads(TILE_SIZE, TILE_SIZE, 1)]
+[numthreads(TILE_SIZE * TILE_SIZE, 1, 1)]
 void MainCS(
 	uint2 GroupId : SV_GroupID,
 	uint GroupThreadIndex : SV_GroupIndex)
diff --git a/Engine/Shaders/Private/TemporalSuperResolution/TSRPostfilterRejection.usf b/Engine/Shaders/Private/TemporalSuperResolution/TSRPostfilterRejection.usf
--- a/Engine/Shaders/Private/TemporalSuperResolution/TSRPostfilterRejection.usf
+++ b/Engine/Shaders/Private/TemporalSuperResolution/TSRPostfilterRejection.usf
@@ -74,7 +74,7 @@ tsr_half Median9(tsr_half In[9])
 
 //------------------------------------------------------- ENTRY POINT
 
-[numthreads(TILE_SIZE, TILE_SIZE, 1)]
+[numthreads(TILE_SIZE * TILE_SIZE, 1, 1)]
 void MainCS(
 	uint2 GroupId : SV_GroupID,
 	uint GroupThreadIndex : SV_GroupIndex)
diff --git a/Engine/Shaders/Private/TemporalSuperResolution/TSRResolveHistory.usf b/Engine/Shaders/Private/TemporalSuperResolution/TSRResolveHistory.usf
--- a/Engine/Shaders/Private/TemporalSuperResolution/TSRResolveHistory.usf
+++ b/Engine/Shaders/Private/TemporalSuperResolution/TSRResolveHistory.usf
@@ -95,7 +95,7 @@ RWTexture2D<tsr_halfC> SceneColorOutputMip1;
 
 //------------------------------------------------------- ENTRY POINT
 
-[numthreads(TILE_SIZE, TILE_SIZE, 1)]
+[numthreads(TILE_SIZE * TILE_SIZE, 1, 1)]
 void MainCS(
 	uint2 GroupId : SV_GroupID,
 	uint GroupThreadIndex : SV_GroupIndex)
diff --git a/Engine/Shaders/Private/TemporalSuperResolution/TSRSpatialAntiAliasing.usf b/Engine/Shaders/Private/TemporalSuperResolution/TSRSpatialAntiAliasing.usf
--- a/Engine/Shaders/Private/TemporalSuperResolution/TSRSpatialAntiAliasing.usf
+++ b/Engine/Shaders/Private/TemporalSuperResolution/TSRSpatialAntiAliasing.usf
@@ -188,9 +188,9 @@ tsr_half ComputeDistanceToEdge(
 //------------------------------------------------------- ENTRY POINT
 
 #if CONFIG_ENABLE_DUAL_PIXEL_VECTORIZATION
-	[numthreads(TILE_SIZE, TILE_SIZE / 2, 1)]
+	[numthreads(TILE_SIZE * TILE_SIZE / 2, 1, 1)]
 #else
-	[numthreads(TILE_SIZE, TILE_SIZE, 1)]
+	[numthreads(TILE_SIZE * TILE_SIZE, 1, 1)]
 #endif
 void MainCS(
 	uint2 GroupId : SV_GroupID,
diff --git a/Engine/Shaders/Private/TemporalSuperResolution/TSRUpdateHistory.usf b/Engine/Shaders/Private/TemporalSuperResolution/TSRUpdateHistory.usf
--- a/Engine/Shaders/Private/TemporalSuperResolution/TSRUpdateHistory.usf
+++ b/Engine/Shaders/Private/TemporalSuperResolution/TSRUpdateHistory.usf
@@ -300,9 +300,9 @@ T BilinearSampleColorHistory(Texture2D<T> Texture, float2 SampleUV)
 //------------------------------------------------------- ENTRY POINT
 
 #if CONFIG_ENABLE_DUAL_PIXEL_VECTORIZATION
-	[numthreads(TILE_SIZE, TILE_SIZE / 2, 1)]
+	[numthreads(TILE_SIZE * TILE_SIZE / 2, 1, 1)]
 #else
-	[numthreads(TILE_SIZE, TILE_SIZE, 1)]
+	[numthreads(TILE_SIZE * TILE_SIZE, 1, 1)]
 #endif
 void MainCS(
 	uint2 GroupId : SV_GroupID,
