PR #8751: Fix CorrelatedColorTemperature in TonemapCommon.ush (Contributed by [REDACTED_PERSON])

diff --git a/Engine/Shaders/Private/TonemapCommon.ush b/Engine/Shaders/Private/TonemapCommon.ush
--- a/Engine/Shaders/Private/TonemapCommon.ush
+++ b/Engine/Shaders/Private/TonemapCommon.ush
@@ -431,7 +431,7 @@ float2 D_IlluminantChromaticity( float Temp )
 // [McCamy 1992, "Correlated color temperature as an explicit function of chromaticity coordinates"]
 float CorrelatedColorTemperature( float x, float y )
 {
-	float n = (x - 0.3320) / (0.1858 - y);
+	float n = (x - 0.3320) / (y - 0.1858);
 	return -449 * n*n*n + 3525 * n*n - 6823.3 * n + 5520.33;
 }
 
