Fix for non-compressed transform path breaking motion vectors.
The non-compressed was recently enabled as a data driven platform option.

diff --git a/Engine/Shaders/Private/SceneData.ush b/Engine/Shaders/Private/SceneData.ush
--- a/Engine/Shaders/Private/SceneData.ush
+++ b/Engine/Shaders/Private/SceneData.ush
@@ -996,21 +996,6 @@ FInstanceSceneData GetInstanceSceneData(uint InstanceId, uint SOAStride, bool bC
 			float3 TilePosition 				= GetPrimitiveData(InstanceData.PrimitiveId).TilePosition;
 			InstanceData.PrevLocalToWorld		= MakeLWCMatrix(TilePosition, PrevLocalToRelativeWorld);
 		}
-		else
-		{
-		#if PLATFORM_ALLOW_SCENE_DATA_COMPRESSED_TRANSFORMS
-			// TODO: Temporary PrevVelocityHack
-			uint4  PrevRotationScale		= asuint(LoadInstanceSceneDataElement(3 * SOAStride + InstanceId));
-			float3 PrevTranslation			= LoadInstanceSceneDataElement(4 * SOAStride + InstanceId).xyz;
-			float3 PrevScale = 0;
-			float4x4 PrevLocalToRelativeWorld	= DecodeTransform(PrevRotationScale, PrevTranslation, PrevScale);
-
-			float3 TilePosition 			= GetPrimitiveData(InstanceData.PrimitiveId).TilePosition;
-			InstanceData.PrevLocalToWorld	= MakeLWCMatrix(TilePosition, PrevLocalToRelativeWorld);
-		#else
-			InstanceData.PrevLocalToWorld	= InstanceData.LocalToWorld;
-		#endif
-		}
 
 	#if 1 //NEEDS_LIGHTMAP_COORDINATE
 		BRANCH
diff --git a/Engine/Source/Runtime/Engine/Private/InstanceUniformShaderParameters.cpp b/Engine/Source/Runtime/Engine/Private/InstanceUniformShaderParameters.cpp
--- a/Engine/Source/Runtime/Engine/Private/InstanceUniformShaderParameters.cpp
+++ b/Engine/Source/Runtime/Engine/Private/InstanceUniformShaderParameters.cpp
@@ -91,5 +91,6 @@ void FInstanceSceneShaderData::BuildInternal
 	{
 		// Note: writes 3x float4s
 		LocalToWorld.To3x4MatrixTranspose((float*)&Data[1]);
+		PrevLocalToWorld.To3x4MatrixTranspose((float*)&Data[4]);
 	}
 }
diff --git a/Engine/Source/Runtime/Engine/Public/InstanceUniformShaderParameters.h b/Engine/Source/Runtime/Engine/Public/InstanceUniformShaderParameters.h
--- a/Engine/Source/Runtime/Engine/Public/InstanceUniformShaderParameters.h
+++ b/Engine/Source/Runtime/Engine/Public/InstanceUniformShaderParameters.h
@@ -99,7 +99,7 @@ private:
 	// Allocate the max Float4s usage when compressed transform is used.
 	// TODO: Temporary PrevVelocityHack (last float4s when compressed)
 	static constexpr uint32 CompressedTransformDataStrideInFloat4s = 5;
-	static constexpr uint32 UnCompressedTransformDataStrideInFloat4s = 4;
+	static constexpr uint32 UnCompressedTransformDataStrideInFloat4s = 7;
 
 public:
 
@@ -159,5 +159,5 @@ public:
 		const FRenderTransform& PrevLocalToWorld // Assumes shear has been removed already // TODO: Temporary PrevVelocityHack
 	);
 
-	TStaticArray<FVector4f, CompressedTransformDataStrideInFloat4s> Data;
+	TStaticArray<FVector4f, UnCompressedTransformDataStrideInFloat4s> Data;
 };
