fix selection outline for HDR editor: ScRGB can generate negative values

CL 20982761 didn't work because sign(x) = 0 when x=0, meaning an outline on a black pixel would not work exactly the same. Instead, we now leave the color as-is when no outline is detected, and remove negative values when applying the outline. The result is technically incorrect (ie we should linearize properly depending on the gamma encoding of the backbuffer and have everything work in the same color space), but this allows for a simple fix

diff --git a/Engine/Shaders/Private/PostProcessSelectionOutline.usf b/Engine/Shaders/Private/PostProcessSelectionOutline.usf
--- a/Engine/Shaders/Private/PostProcessSelectionOutline.usf
+++ b/Engine/Shaders/Private/PostProcessSelectionOutline.usf
@@ -265,9 +265,13 @@ void MainPS(
 		OutColor = Texture2DSample(ColorTexture, ColorSampler, ColorUV);
 
 		// Scene color has gamma applied. Need to remove and re-apply after linear operation.
-		OutColor.rgb = pow(OutColor.rgb, 2.2f );
-		OutColor.rgb = lerp(OutColor.rgb, Avg.rgb, Avg.a);
-		OutColor.rgb = pow(OutColor.rgb, 1.0f / 2.2f);
+		if (Avg.a > 0)
+		{
+			OutColor.rgb = max(OutColor.rgb, 0.f);
+			OutColor.rgb = pow(OutColor.rgb, 2.2f );
+		    OutColor.rgb = lerp(OutColor.rgb, Avg.rgb, Avg.a);
+		    OutColor.rgb = pow(OutColor.rgb, 1.0f / 2.2f);
+		}
 	}
 	else
 	{
