Fix hair strands material conversion failing in certain cases.

diff --git a/Engine/Shaders/Private/HairStrands/HairStrandsMaterialCommon.ush b/Engine/Shaders/Private/HairStrands/HairStrandsMaterialCommon.ush
--- a/Engine/Shaders/Private/HairStrands/HairStrandsMaterialCommon.ush
+++ b/Engine/Shaders/Private/HairStrands/HairStrandsMaterialCommon.ush
@@ -212,16 +212,12 @@ FEvaluateOutput Evaluate(
 		{
 			// Using the tree data, instread of the inline BSDF does not seem to work. It is unclear why. It is not very 
 			// important as we expect a single hair BSDF.
-		#if 0
 			// Update tree (coverage/transmittance/luminace weights)
 			FStrataAddressing NullStrataAddressing = (FStrataAddressing)0;	// Fake unused in StrataCreateBSDFContext when using Forward inline shading
 			const float3 V = MaterialParameters.CameraVector;
 			const FStrataIntegrationSettings Settings = InitStrataIntegrationSettings(false /*bForceFullyRough*/, false /*bRoughDiffuse*/, -1 /*PeelLayersAboveDepth*/);
 			StrataUpdateTree(NullStrataAddressing, StrataPixelHeader, StrataPixelHeader.StrataTree, V, Settings);
 			FStrataBSDF HairBSDF = StrataPixelHeader.StrataTree.BSDFs[0];
-		#else
-			FStrataBSDF HairBSDF = PixelMaterialInputs.FrontMaterial.InlinedBSDF;
-		#endif
 
 			OutSample.BaseColor = HAIR_BASECOLOR(HairBSDF);
 			OutSample.Roughness = HAIR_ROUGHNESS(HairBSDF);
diff --git a/Engine/Shaders/Private/Strata/StrataLegacyConversion.ush b/Engine/Shaders/Private/Strata/StrataLegacyConversion.ush
--- a/Engine/Shaders/Private/Strata/StrataLegacyConversion.ush
+++ b/Engine/Shaders/Private/Strata/StrataLegacyConversion.ush
@@ -176,8 +176,6 @@ FStrataData StrataConvertLegacyMaterialDynamic(
 			ColorScaleBehindWater,	// ColorScaleBehindWater
 			SharedLocalBasisIndex);
 
-		FStrataLegacyParameters Bot = InitStrataLegacyParameters(PixelFootprint, SharedLocalBasisIndex, 1.0f);
-
 		return StrataWeight(Top_SLW , 1.0, StrataTree,					/*OperatorIndex*/1, /*MaxDistanceFromLeaves*/1);
 		// Out maps to LegacyUpdateBSDFsOperators
 	}
@@ -185,16 +183,23 @@ FStrataData StrataConvertLegacyMaterialDynamic(
 	#if MATERIAL_SHADINGMODEL_HAIR
 	if (ShadingModel == SHADINGMODELID_HAIR)
 	{
-		return GetStrataHairBSDF(
-			BaseColor,	   // BaseColor
-			Metallic,	   // Scatter
-			Specular,	   // Specular
-			Roughness,	   // Roughness
-			ClearCoat,	   // Backlit
-			Emissive,	   // EmissiveColor
-			SharedLocalBasisIndex);
+		FStrataData Top_Hair = PromoteParameterBlendedBSDFToOperator(
+			GetStrataHairBSDF(
+				BaseColor,	   // BaseColor
+				Metallic,	   // Scatter
+				Specular,	   // Specular
+				Roughness,	   // Roughness
+				ClearCoat,	   // Backlit
+				Emissive,	   // EmissiveColor
+				SharedLocalBasisIndex),
+			StrataTree,
+			0		/*OperatorIndex*/,
+			0		/*BSDFIndex*/,
+			0		/*int LayerDepth*/,
+			1		/*int bIsBottom*/);
 
-		// No need to take into account opacity weighting for translucent material, as the coverage is assumed to be 1
+		return StrataWeight(Top_Hair, 1.0, StrataTree,					/*OperatorIndex*/1, /*MaxDistanceFromLeaves*/1);
+		// Out maps to LegacyUpdateBSDFsOperators
 	}
 	#endif
 	#if MATERIAL_SHADINGMODEL_EYE
