Strata - fixed decals and all legacy material conversion to be a lot simpler and consistent with the strata.
We are no longer fighting the Strata tree but just using it now, since Legacyconversion node can only oputput a single BSDF at any time. Conversion to decals parameter blending will now correctly work.

diff --git a/Engine/Shaders/Private/DeferredDecal.usf b/Engine/Shaders/Private/DeferredDecal.usf
--- a/Engine/Shaders/Private/DeferredDecal.usf
+++ b/Engine/Shaders/Private/DeferredDecal.usf
@@ -146,7 +146,22 @@ void FPixelShaderInOut_MainPS(inout FPixelShaderIn In, inout FPixelShaderOut Out
 	StrataHeader.IrradianceAO = InitIrradianceAndOcclusion();
 	StrataHeader.IrradianceAO.MaterialAO = GetMaterialAmbientOcclusion(PixelMaterialInputs);
 	StrataHeader.SharedLocalBases = MaterialParameters.SharedLocalBases;
-	const float Coverage = DecalCommonOutput(Out, StrataHeader, StrataHeader.StrataTree.BSDFs[0], DecalFading);
+
+	float Coverage = 0.0;
+	if (StrataHeader.StrataTree.BSDFCount > 0)
+	{
+		FStrataData StrataData = PixelMaterialInputs.FrontMaterial;
+
+		FStrataAddressing NullStrataAddressing = (FStrataAddressing)0;	// Fake unused in StrataCreateBSDFContext when using Forward inline shading
+		const float3 V = MaterialParameters.CameraVector;
+
+		// Update tree (coverage/transmittance/luminace weights)
+		const FStrataIntegrationSettings Settings = InitStrataIntegrationSettings(false /*bForceFullyRough*/, false /*bRoughDiffuseEnabled*/, 0);
+		StrataUpdateTree(NullStrataAddressing, StrataHeader, StrataHeader.StrataTree, V, Settings);
+
+		// Since the Convert-To-Decal node forces parameter blending for the entire Strata tree, the BSDF we are interested in will be in 0.
+		Coverage = DecalCommonOutput(Out, StrataHeader, StrataHeader.StrataTree.BSDFs[0], DecalFading);
+	}
 	const float Opacity = Coverage;
 #else // STRATA_ENABLED
 	float3 Color = GetMaterialEmissive(PixelMaterialInputs);
diff --git a/Engine/Shaders/Private/Strata/Strata.ush b/Engine/Shaders/Private/Strata/Strata.ush
--- a/Engine/Shaders/Private/Strata/Strata.ush
+++ b/Engine/Shaders/Private/Strata/Strata.ush
@@ -3395,14 +3395,10 @@ FStrataIntegrationSettings InitStrataIntegrationSettings()
 #if STRATA_INLINE_SHADING && STRATA_CLAMPED_BSDF_COUNT > 0
 
 // Forward declaration
-#if STRATA_USES_CONVERSION_FROM_LEGACY==0
 void UpdateAllBSDFsOperatorCoverageTransmittance(inout FStrataPixelHeader StrataPixelHeader, inout FStrataTree StrataTree, FStrataIntegrationSettings Settings, FStrataAddressing NullStrataAddressing, float3 V);
 void UpdateAllOperatorsCoverageTransmittance(inout FStrataTree StrataTree);
 void UpdateAllBSDFWithBottomUpOperatorVisit(inout FStrataTree StrataTree);
 void PreUpdateAllBSDFWithBottomUpOperatorVisit(inout FStrataPixelHeader StrataPixelHeader, inout FStrataTree StrataTree, FStrataAddressing NullStrataAddressing, float3 V);
-#else
-void LegacyUpdateBSDFsOperators(inout FStrataPixelHeader StrataPixelHeader, inout FStrataTree StrataTree, FStrataIntegrationSettings Settings, FStrataAddressing NullStrataAddressing, float3 V);
-#endif
 
 // Update all Coverage/Transmittance/Weights for all BSDFs based on material topology
 void StrataUpdateTree(
@@ -3412,7 +3408,6 @@ void StrataUpdateTree(
 	float3 V,
 	FStrataIntegrationSettings Settings)
 {
-#if STRATA_USES_CONVERSION_FROM_LEGACY==0
 	//
 	// Pre-Update the BSDF by traversing up the strata tree until the root node. This allows to modify BSDF's parameters based on operators (i.e., thin-film coating)
 	//
@@ -3432,9 +3427,6 @@ void StrataUpdateTree(
 	// Update the luminance weight of each BSDF according to the operators it has to traverse up to the strata tree root node.
 	//
 	UpdateAllBSDFWithBottomUpOperatorVisit(StrataTree);
-#else
-	LegacyUpdateBSDFsOperators(StrataPixelHeader, StrataTree, Settings, StrataAddressing, V);
-#endif
 }
 
 void StrataUpdateTree(
diff --git a/Engine/Shaders/Private/Strata/StrataLegacyConversion.ush b/Engine/Shaders/Private/Strata/StrataLegacyConversion.ush
--- a/Engine/Shaders/Private/Strata/StrataLegacyConversion.ush
+++ b/Engine/Shaders/Private/Strata/StrataLegacyConversion.ush
@@ -83,46 +83,37 @@ FStrataLegacyParameters InitStrataLegacyParameters(FStrataPixelFootprint InPixel
 }
 
 FStrataData CreateLegacySlab(
-	FStrataLegacyParameters In, inout uint SharedLocalBasisTypes,
-	inout FStrataTree StrataTree,
-	int OperatorIndex, int BSDFIndex)
+	FStrataLegacyParameters In, inout uint SharedLocalBasisTypes, inout FStrataTree StrataTree)
 {
 	const float3 Zeros = float3(0, 0, 0);
 	const float3 Ones  = float3(1, 1, 1);
 
 	// Fixed layer structure for helping compiler to unroll and optimize shader
-	return PromoteParameterBlendedBSDFToOperator(
-		GetStrataSlabBSDF(
-			In.PixelFootprint,
-			In.UseMetalness,									// UseMetalness
-			In.BaseColor,										// BaseColor		-  Metalness workflow
-			Ones,												// EdgeColor		-  Metalness workflow
-			In.Specular,										// Specular			-  Metalness workflow
-			In.Metallic,										// Metallic			-  Metalness workflow
-			In.DiffuseAlbedo,									// DiffuseAlbedo	- !Metalness workflow
-			In.F0,												// F0				- !Metalness workflow
-			In.F90,												// F90				- !Metalness workflow
-			In.Roughness,										// Roughness
-			In.Anisotropy,										// Anisotropy
-			In.SSSProfileID,									// SSSProfileID
-			In.SSSMFP,											// SSSMFP
-			In.SSSMFPScale,										// SSSMFPScale
-			In.SSSPhaseAniso,									// SSSPhaseAniso
-			In.UseSSSDiffusion,									// UseSSSDiffusion
-			In.Emissive,										// Emissive
-			In.SecondRoughness,									// SecondRoughness
-			In.SecondRoughnessWeight,							// SecondRoughnessWeight
-			In.SecondRoughnessAsSimpleClearCoat ? 1.0f: 0.0f,	// SecondRoughnessAsSimpleClearCoat
-			In.FuzzAmount,										// FuzzAmount
-			In.FuzzColor,										// FuzzColor
-			In.Thickness,
-			In.SharedLocalBasisIndex, SharedLocalBasisTypes),
-		StrataTree,
-		OperatorIndex,
-		BSDFIndex,
-		0,														//LayerDepth,
-		1														//bIsBottom
-		);
+	return GetStrataSlabBSDF(
+		In.PixelFootprint,
+		In.UseMetalness,									// UseMetalness
+		In.BaseColor,										// BaseColor		-  Metalness workflow
+		Ones,												// EdgeColor		-  Metalness workflow
+		In.Specular,										// Specular			-  Metalness workflow
+		In.Metallic,										// Metallic			-  Metalness workflow
+		In.DiffuseAlbedo,									// DiffuseAlbedo	- !Metalness workflow
+		In.F0,												// F0				- !Metalness workflow
+		In.F90,												// F90				- !Metalness workflow
+		In.Roughness,										// Roughness
+		In.Anisotropy,										// Anisotropy
+		In.SSSProfileID,									// SSSProfileID
+		In.SSSMFP,											// SSSMFP
+		In.SSSMFPScale,										// SSSMFPScale
+		In.SSSPhaseAniso,									// SSSPhaseAniso
+		In.UseSSSDiffusion,									// UseSSSDiffusion
+		In.Emissive,										// Emissive
+		In.SecondRoughness,									// SecondRoughness
+		In.SecondRoughnessWeight,							// SecondRoughnessWeight
+		In.SecondRoughnessAsSimpleClearCoat ? 1.0f : 0.0f,	// SecondRoughnessAsSimpleClearCoat
+		In.FuzzAmount,										// FuzzAmount
+		In.FuzzColor,										// FuzzColor
+		In.Thickness,
+		In.SharedLocalBasisIndex, SharedLocalBasisTypes);
 }
 
 // Convert legacy shading models - Dynamic
@@ -163,7 +154,7 @@ FStrataData StrataConvertLegacyMaterialDynamic(
 		const float3 WaterExtinction = WaterScatteringCoefficients + WaterAbsorptionCoefficients;
 		const float3 WaterAlbedo = WaterScatteringCoefficients / WaterExtinction;
 
-		FStrataData Top_SLW = GetStrataSingleLayerWaterBSDF(
+		return GetStrataSingleLayerWaterBSDF(
 			BaseColor, 				// BaseColor
 			Metallic, 				// Metallic 
 			Specular, 				// Specular
@@ -175,31 +166,19 @@ FStrataData StrataConvertLegacyMaterialDynamic(
 			WaterPhaseG, 			// WaterPhaseG
 			ColorScaleBehindWater,	// ColorScaleBehindWater
 			SharedLocalBasisIndex);
-
-		return StrataWeight(Top_SLW , 1.0, StrataTree,					/*OperatorIndex*/1, /*MaxDistanceFromLeaves*/1);
-		// Out maps to LegacyUpdateBSDFsOperators
 	}
 	#endif
 	#if MATERIAL_SHADINGMODEL_HAIR
 	if (ShadingModel == SHADINGMODELID_HAIR)
 	{
-		FStrataData Top_Hair = PromoteParameterBlendedBSDFToOperator(
-			GetStrataHairBSDF(
-				BaseColor,	   // BaseColor
-				Metallic,	   // Scatter
-				Specular,	   // Specular
-				Roughness,	   // Roughness
-				ClearCoat,	   // Backlit
-				Emissive,	   // EmissiveColor
-				SharedLocalBasisIndex),
-			StrataTree,
-			0		/*OperatorIndex*/,
-			0		/*BSDFIndex*/,
-			0		/*int LayerDepth*/,
-			1		/*int bIsBottom*/);
-
-		return StrataWeight(Top_Hair, 1.0, StrataTree,					/*OperatorIndex*/1, /*MaxDistanceFromLeaves*/1);
-		// Out maps to LegacyUpdateBSDFsOperators
+		return GetStrataHairBSDF(
+			BaseColor,	   // BaseColor
+			Metallic,	   // Scatter
+			Specular,	   // Specular
+			Roughness,	   // Roughness
+			ClearCoat,	   // Backlit
+			Emissive,	   // EmissiveColor
+			SharedLocalBasisIndex);
 	}
 	#endif
 	#if MATERIAL_SHADINGMODEL_EYE
@@ -212,8 +191,7 @@ FStrataData StrataConvertLegacyMaterialDynamic(
 		float3 IrisPlaneNormal = RawNormal;
 		GetEyeNormals(IrisMask, IrisDistance, RawNormal, RawClearCoatNormal, RawCustomTangent, IrisNormal, IrisPlaneNormal);
 
-		FStrataData EyeMat = PromoteParameterBlendedBSDFToOperator(
-			GetStrataEyeBSDF(
+		return GetStrataEyeBSDF(
 				BaseColor,	         // DiffuseAlbedo
 				Roughness,	         // Roughness
 				IrisMask,	         // IrisMask
@@ -222,15 +200,7 @@ FStrataData StrataConvertLegacyMaterialDynamic(
 				IrisPlaneNormal,	 // IrisPlaneNormal
 				SubSurfaceProfileId, // SSS profile
 				Emissive,	         // EmissiveColor
-				SharedLocalBasisIndex),
-			StrataTree,
-			0		/*OperatorIndex*/,
-			0		/*BSDFIndex*/,
-			0		/*int LayerDepth*/,
-			1		/*int bIsBottom*/);
-
-		return StrataWeight(EyeMat, 1.0, StrataTree, /*OperatorIndex*/1, /*MaxDistanceFromLeaves*/1);
-		// Out maps to LegacyUpdateBSDFsOperators
+				SharedLocalBasisIndex);
 	}
 	#endif
 
@@ -356,7 +326,6 @@ FStrataData StrataConvertLegacyMaterialDynamic(
 		LegacySlab.SecondRoughness = ClearCoatRoughness;
 		LegacySlab.SecondRoughnessAsSimpleClearCoat = ClearCoat > 0.0f; // Haziness as simple clear coat layer.
 		LegacySlab.Emissive = Zeros;
-		LegacySlab.SharedLocalBasisIndex = ClearCoatBottomNormal_SharedLocalBasisIndex;
 		LegacySlab.Weight = 1.f;
 
 		FinalWeight = 1.0f;
@@ -453,32 +422,16 @@ FStrataData StrataConvertLegacyMaterialDynamic(
 	}
 
 	// Fixed layer structure for helping compiler to unroll and optimize shader
-	// This must match what is done in PackStrataOut
-	FStrataData Slab = CreateLegacySlab(LegacySlab, SharedLocalBasisTypes, StrataTree,	/*OperatorIndex*/0, /*BSDFIndex*/0);
-	FStrataData Out  = StrataWeight(Slab, LegacySlab.Weight, StrataTree,				/*OperatorIndex*/1, /*MaxDistanceFromLeaves*/1);
-
-	// Take in account the legacy opacity weighting for translucent material
-	#if STRATA_LEGACY_MATERIAL_APPLIES_FINAL_WEIGHT
-	{
-		Out = StrataWeight(Out, FinalWeight, StrataTree,								/*OperatorIndex*/2, /*MaxDistanceFromLeaves*/2);
-	}
-	#endif
+	FStrataData Slab = CreateLegacySlab(LegacySlab, SharedLocalBasisTypes, StrataTree);
+#if STRATA_INLINE_SHADING
+	Slab.InlinedBSDF.Coverage = LegacySlab.Weight * FinalWeight;	//////////////////// Remove FinalWeight? Do that in CreateLegacySlab?
+#endif
 
 	#if STRATA_LEGACY_PREMULT_ALPHA_OVERRIDE
-	Out.PreMultipliedAlphaOverrideCoverage = Opacity;
+	Slab.PreMultipliedAlphaOverrideCoverage = Opacity;
 	#endif
 
-	// We also evaluate the inlined BSDF resulting from parameter blending in case it is used later.
-	// For instance, this is used for deferred decals. We only evaluate the top material as this is enough for decals.
-	FStrataData ParamBlendVertical = StrataWeightParameterBlending(Slab, LegacySlab.Weight);
-#if STRATA_LEGACY_MATERIAL_APPLIES_FINAL_WEIGHT
-	Out.InlinedBSDF = StrataWeightParameterBlending(ParamBlendVertical, FinalWeight).InlinedBSDF;
-#else
-	Out.InlinedBSDF = ParamBlendVertical.InlinedBSDF;
-#endif
-
-	// Out maps to LegacyUpdateBSDFsOperators
-	return Out;
+	return Slab;
 }
 
 // Convert legacy shading models - Static
@@ -544,23 +497,14 @@ FStrataData StrataConvertLegacyMaterialStatic(
 	}
 	else if (ShadingModel == SHADINGMODELID_HAIR)
 	{
-		FStrataData HairMat = PromoteParameterBlendedBSDFToOperator(
-			GetStrataHairBSDF(
+		Out = GetStrataHairBSDF(
 				BaseColor,	   // BaseColor
 				Metallic,	   // Scatter
 				Specular,	   // Specular
 				Roughness,	   // Roughness
 				ClearCoat,	   // Backlit
 				Emissive,	   // EmissiveColor
-				SharedLocalBasisIndex),
-			StrataTree,
-			0		/*OperatorIndex*/,
-			0		/*BSDFIndex*/,
-			0		/*int LayerDepth*/,
-			1		/*int bIsBottom*/);
-
-		Out = StrataWeight(HairMat, 1.0, StrataTree, /*OperatorIndex*/1, /*MaxDistanceFromLeaves*/1);
-		// Out maps to LegacyUpdateBSDFsOperators
+				SharedLocalBasisIndex);
 	}
 	else if (ShadingModel == SHADINGMODELID_EYE)
 	{
@@ -571,8 +515,7 @@ FStrataData StrataConvertLegacyMaterialStatic(
 		float3 IrisPlaneNormal = RawNormal;
 		GetEyeNormals(IrisMask, IrisDistance, RawNormal, RawClearCoatNormal, RawCustomTangent, IrisNormal, IrisPlaneNormal);
 
-		FStrataData EyeMat = PromoteParameterBlendedBSDFToOperator(
-			GetStrataEyeBSDF(
+		Out = GetStrataEyeBSDF(
 				BaseColor,	         // DiffuseAlbedo
 				Roughness,	         // Roughness
 				IrisMask,	         // IrisMask
@@ -581,23 +524,14 @@ FStrataData StrataConvertLegacyMaterialStatic(
 				IrisPlaneNormal,	 // IrisPlaneNormal
 				SubSurfaceProfileId, // SSS profile
 				Emissive,	         // EmissiveColor
-				SharedLocalBasisIndex),
-			StrataTree,
-			0		/*OperatorIndex*/,
-			0		/*BSDFIndex*/,
-			0		/*int LayerDepth*/,
-			1		/*int bIsBottom*/);
-
-		Out = StrataWeight(EyeMat, 1.0, StrataTree, /*OperatorIndex*/1, /*MaxDistanceFromLeaves*/1);
-		// Out maps to LegacyUpdateBSDFsOperators
+				SharedLocalBasisIndex);
 	}
 	else if (ShadingModel == SHADINGMODELID_SINGLELAYERWATER)
 	{
-		const float3 WaterExtinction = WaterScatteringCoefficients + WaterAbsorptionCoefficients;
-		const float3 WaterAlbedo = WaterScatteringCoefficients / WaterExtinction;
+		const float3 WaterExtinction= WaterScatteringCoefficients + WaterAbsorptionCoefficients;
+		const float3 WaterAlbedo	= WaterScatteringCoefficients / WaterExtinction;
 
-		FStrataData SLW = PromoteParameterBlendedBSDFToOperator(
-			GetStrataSingleLayerWaterBSDF(
+		Out = GetStrataSingleLayerWaterBSDF(
 				BaseColor, 				// BaseColor
 				Metallic, 				// Metallic 
 				Specular, 				// Specular
@@ -608,15 +542,7 @@ FStrataData StrataConvertLegacyMaterialStatic(
 				WaterExtinction, 		// WaterExtinction
 				WaterPhaseG, 			// WaterPhaseG
 				ColorScaleBehindWater,	// ColorScaleBehindWater
-				SharedLocalBasisIndex),
-			StrataTree,
-			0		/*OperatorIndex*/,
-			0		/*BSDFIndex*/,
-			0		/*int LayerDepth*/,
-			1		/*int bIsBottom*/);
-
-		Out = StrataWeight(SLW, 1.0, StrataTree, /*OperatorIndex*/1, /*MaxDistanceFromLeaves*/1);
-		// Out maps to LegacyUpdateBSDFsOperators
+				SharedLocalBasisIndex);
 	}
 	else
 	{
diff --git a/Engine/Shaders/Private/Strata/StrataTree.ush b/Engine/Shaders/Private/Strata/StrataTree.ush
--- a/Engine/Shaders/Private/Strata/StrataTree.ush
+++ b/Engine/Shaders/Private/Strata/StrataTree.ush
@@ -451,26 +451,3 @@ void PreUpdateAllBSDFWithBottomUpOperatorVisit_ThinFilm(
 #undef Op
 
 
-
-///////////////////////////////////////////////////////////////////////////////
-// Legacy materials conversion
-
-void LegacyUpdateBSDFsOperators(inout FStrataPixelHeader StrataPixelHeader, inout FStrataTree StrataTree, FStrataIntegrationSettings Settings, FStrataAddressing NullStrataAddressing, float3 V)
-{
-	// This must match what is done in StrataConvertLegacyMaterialDynamic
-	UpdateSingleBSDFOperatorCoverageTransmittance(StrataPixelHeader, StrataTree, StrataTree.BSDFs[0], Settings, NullStrataAddressing, V);
-
-	// MaxDistanceFromLeaves = 1 
-	UpdateSingleOperatorCoverageTransmittance(StrataTree, /*OperatorIndex*/1);
-#if STRATA_LEGACY_MATERIAL_APPLIES_FINAL_WEIGHT
-	// MaxDistanceFromLeaves = 2 
-	UpdateSingleOperatorCoverageTransmittance(StrataTree, /*OperatorIndex*/2);
-#endif
-
-	UpdateAllBSDFWithBottomUpOperatorVisit_Weight	(StrataTree, StrataTree.BSDFs[0], /*OperatorIndex*/1, /*PreviousIsInputA*/1);
-#if STRATA_LEGACY_MATERIAL_APPLIES_FINAL_WEIGHT
-	UpdateAllBSDFWithBottomUpOperatorVisit_Weight	(StrataTree, StrataTree.BSDFs[0], /*OperatorIndex*/2, /*PreviousIsInputA*/1);
-#endif
-}
-
-
diff --git a/Engine/Shaders/Shared/StrataDefinitions.h b/Engine/Shaders/Shared/StrataDefinitions.h
--- a/Engine/Shaders/Shared/StrataDefinitions.h
+++ b/Engine/Shaders/Shared/StrataDefinitions.h
@@ -8,7 +8,7 @@
 #pragma once
 
 // Change this to force recompilation of all strata dependent shaders (use https://www.random.org/cgi-bin/randbyte?nbytes=4&format=h)
-#define STRATA_SHADER_VERSION 0x2a533622
+#define STRATA_SHADER_VERSION 0x97e993a4
 
 
 
diff --git a/Engine/Source/Runtime/Engine/Private/Materials/HLSLMaterialTranslator.cpp b/Engine/Source/Runtime/Engine/Private/Materials/HLSLMaterialTranslator.cpp
--- a/Engine/Source/Runtime/Engine/Private/Materials/HLSLMaterialTranslator.cpp
+++ b/Engine/Source/Runtime/Engine/Private/Materials/HLSLMaterialTranslator.cpp
@@ -1343,7 +1343,7 @@ bool FHLSLMaterialTranslator::Translate()
 			EStrataBlendMode StrataBlendMode = Material->GetStrataBlendMode();
 			bMaterialIsStrata = true;
 
-			if (!bStrataUsesConversionFromLegacy && StrataMaterialRootOperator)
+			if (StrataMaterialRootOperator)
 			{
 				// Now implement the functions needed to process the material topology
 
@@ -1524,7 +1524,7 @@ bool FHLSLMaterialTranslator::Translate()
 		else
 		{
 			MaterialCompilationOutput.StrataMaterialDescription = "";
-			ResourcesString += "// No Strata material provided, or using legacy material conversion \r\n";
+			ResourcesString += "// No Strata material provided\r\n";
 			
 			// Adde default strata functions
 			ResourcesString += "#if TEMPLATE_USES_STRATA\n";
@@ -10464,7 +10464,8 @@ int32 FHLSLMaterialTranslator::StrataSlabBSDF(
 			PromoteToOperator->Index,
 			PromoteToOperator->BSDFIndex,
 			PromoteToOperator->LayerDepth,
-			PromoteToOperator->bIsBottom ? 1 : 0);
+			PromoteToOperator->bIsBottom ? 1 : 0
+		);
 	}
 	
 	return AddCodeChunk(
@@ -10510,7 +10511,8 @@ int32 FHLSLMaterialTranslator::StrataConversionFromLegacy(
 	int32 ShadingModel,
 	int32 Normal, int32 Tangent, const FString& SharedLocalBasisIndexMacro,
 	int32 ClearCoat_Normal, int32 ClearCoat_Tangent, const FString& ClearCoat_SharedLocalBasisIndexMacro,
-	int32 CustomTangent_Tangent)
+	int32 CustomTangent_Tangent,
+	FStrataOperator* PromoteToOperator)
 {
 	const FString NormalCode = GetParameterCode(Normal);
 	const FString TangentCode = Tangent != INDEX_NONE ? *GetParameterCode(Tangent) : TEXT("NONE");
@@ -10519,6 +10521,54 @@ int32 FHLSLMaterialTranslator::StrataConversionFromLegacy(
 	const FString ClearCoat_NormalCode = GetParameterCode(ClearCoat_Normal);
 	const FString ClearCoat_TangentCode = Tangent != INDEX_NONE ? *GetParameterCode(ClearCoat_Tangent) : TEXT("NONE");
 
+	if (PromoteToOperator)
+	{
+		if (PromoteToOperator->Index == INDEX_NONE || PromoteToOperator->BSDFIndex == INDEX_NONE)
+		{
+			Errorf(TEXT("Invalid StrataSlabBSDF operator and BSDF indices during promotion in Material %s (asset: %s).\r\n"), *Material->GetDebugName(), *Material->GetAssetPath().ToString());
+			return INDEX_NONE;
+		}
+
+		return AddCodeChunk(
+			MCT_Strata, TEXT("PromoteParameterBlendedBSDFToOperator(StrataConvertLegacyMaterial%s(Parameters.StrataPixelFootprint, %s, %s, %s, %s, %s, %s, %s, %s, %s, %s, %s, %s, %s, %s, %s, %s, %s, %s, %s, %s, %s, %s, %s, Parameters.SharedLocalBases.Types, Parameters.StrataTree) /* Normal = %s ; Tangent = %s ; ClearCoat_Normal = %s ; ClearCoat_Tangent = %s */, Parameters.StrataTree, %u, %u, %u, %u)"),
+			bHasDynamicShadingModels ? TEXT("Dynamic") : TEXT("Static"),
+			*StrataGetCastParameterCode(BaseColor,						MCT_Float3),
+			*StrataGetCastParameterCode(Specular,						MCT_Float),
+			*StrataGetCastParameterCode(Metallic,						MCT_Float),
+			*StrataGetCastParameterCode(Roughness,						MCT_Float),
+			*StrataGetCastParameterCode(Anisotropy,						MCT_Float),
+			*StrataGetCastParameterCode(SubSurfaceColor,				MCT_Float3),
+			*StrataGetCastParameterCode(SubSurfaceProfileId,			MCT_Float),
+			*StrataGetCastParameterCode(ClearCoat,						MCT_Float),
+			*StrataGetCastParameterCode(ClearCoatRoughness,				MCT_Float),
+			*StrataGetCastParameterCode(EmissiveColor,					MCT_Float3),
+			*StrataGetCastParameterCode(Opacity,						MCT_Float),
+			*StrataGetCastParameterCode(TransmittanceColor,				MCT_Float3),
+			*StrataGetCastParameterCode(WaterScatteringCoefficients,	MCT_Float3),
+			*StrataGetCastParameterCode(WaterAbsorptionCoefficients,	MCT_Float3),
+			*StrataGetCastParameterCode(WaterPhaseG,					MCT_Float),
+			*StrataGetCastParameterCode(ColorScaleBehindWater,			MCT_Float3),
+			*GetParameterCode(ShadingModel),
+			// Raw access to Normal/Tangent/ClearCoatNormal/CustomTangent for conversion purpose
+			*StrataGetCastParameterCode(Normal,							MCT_Float3),
+			*StrataGetCastParameterCode(RawTangent,						MCT_Float3),
+			*StrataGetCastParameterCode(ClearCoat_Normal,				MCT_Float3),
+			*StrataGetCastParameterCode(CustomTangent_Tangent,			MCT_Float3),
+			*SharedLocalBasisIndexMacro,
+			*ClearCoat_SharedLocalBasisIndexMacro,
+			// Regular normal basis
+			*NormalCode,
+			*TangentCode,
+			// Clear coat bottom layer normal basis
+			*ClearCoat_NormalCode,
+			*ClearCoat_TangentCode,
+			PromoteToOperator->Index,
+			PromoteToOperator->BSDFIndex,
+			PromoteToOperator->LayerDepth,
+			PromoteToOperator->bIsBottom ? 1 : 0
+		);
+	}
+
 	return AddCodeChunk(
 		MCT_Strata, TEXT("StrataConvertLegacyMaterial%s(Parameters.StrataPixelFootprint, %s, %s, %s, %s, %s, %s, %s, %s, %s, %s, %s, %s, %s, %s, %s, %s, %s, %s, %s, %s, %s, %s, %s, Parameters.SharedLocalBases.Types, Parameters.StrataTree) /* Normal = %s ; Tangent = %s ; ClearCoat_Normal = %s ; ClearCoat_Tangent = %s */"),
 		bHasDynamicShadingModels ? TEXT("Dynamic") : TEXT("Static"),
diff --git a/Engine/Source/Runtime/Engine/Private/Materials/HLSLMaterialTranslator.h b/Engine/Source/Runtime/Engine/Private/Materials/HLSLMaterialTranslator.h
--- a/Engine/Source/Runtime/Engine/Private/Materials/HLSLMaterialTranslator.h
+++ b/Engine/Source/Runtime/Engine/Private/Materials/HLSLMaterialTranslator.h
@@ -1053,7 +1053,8 @@ protected:
 		int32 ShadingModel,
 		int32 Normal, int32 Tangent, const FString& SharedLocalBasisIndexMacro,
 		int32 ClearCoat_Normal, int32 ClearCoat_Tangent, const FString& ClearCoat_SharedLocalBasisIndexMacro, 
-		int32 CustomTangent_Tangent) override;
+		int32 CustomTangent_Tangent,
+		FStrataOperator* PromoteToOperator) override;
 	virtual int32 StrataVolumetricFogCloudBSDF(int32 Albedo, int32 Extinction, int32 EmissiveColor, int32 AmbientOcclusion) override;
 	virtual int32 StrataUnlitBSDF(int32 EmissiveColor, int32 TransmittanceColor) override;
 	virtual int32 StrataHairBSDF(int32 BaseColor, int32 Scatter, int32 Specular, int32 Roughness, int32 Backlit, int32 EmissiveColor, int32 Tangent, const FString& SharedLocalBasisIndexMacro, FStrataOperator* PromoteToOperator) override;
diff --git a/Engine/Source/Runtime/Engine/Private/Materials/Material.cpp b/Engine/Source/Runtime/Engine/Private/Materials/Material.cpp
--- a/Engine/Source/Runtime/Engine/Private/Materials/Material.cpp
+++ b/Engine/Source/Runtime/Engine/Private/Materials/Material.cpp
@@ -3337,12 +3337,8 @@ void UMaterial::ConvertMaterialToStrataMaterial()
 			check(ConvertNode->ConvertedStrataMaterialInfo.CountShadingModels() == 1);
 
 			// Now pass through the convert to decal node, which flag the material as SSM_Decal, which will set the domain to Decal.
-			// This node doesn't apply parameter blending unlike when used on non-legacy material. This is because the legacy conversion 
-			// node does not support parameter blending. Coverage is also left to default (1.0), as the opacity is managed within the 
-			// legacy node. The first BSDF of the legacy node is used as the decal material.
 			UMaterialExpressionStrataConvertToDecal* ConvertToDecalNode= NewObject<UMaterialExpressionStrataConvertToDecal>(this);
 			ConvertToDecalNode->DecalMaterial.Connect(0, ConvertNode);
-			ConvertToDecalNode->bUseParameterBlending = 0;
 
 			EditorOnly->FrontMaterial.Connect(0, ConvertToDecalNode);
 			bInvalidateShader = true;
diff --git a/Engine/Source/Runtime/Engine/Private/Materials/MaterialExpressions.cpp b/Engine/Source/Runtime/Engine/Private/Materials/MaterialExpressions.cpp
--- a/Engine/Source/Runtime/Engine/Private/Materials/MaterialExpressions.cpp
+++ b/Engine/Source/Runtime/Engine/Private/Materials/MaterialExpressions.cpp
@@ -21786,6 +21786,9 @@ int32 UMaterialExpressionStrataLegacyConversion::Compile(class FMaterialCompiler
 		SSSProfileCodeChunk = Compiler->ForceCast(Compiler->ScalarParameter(GetSubsurfaceProfileParameterName(), 1.0f), MCT_Float1);
 	}
 
+	FStrataOperator& StrataOperator = Compiler->StrataCompilationGetOperator(this);
+	StrataOperator.BSDFRegisteredSharedLocalBasis = NewRegisteredSharedLocalBasis;
+
 	const bool bHasDynamicShadingModels = ConvertedStrataMaterialInfo.CountShadingModels() > 1;
 	// We probably need to do something along these line as well :::
 	int32 OutputCodeChunk = Compiler->StrataConversionFromLegacy(
@@ -21820,50 +21823,8 @@ int32 UMaterialExpressionStrataLegacyConversion::Compile(class FMaterialCompiler
 		ClearCoat_NormalCodeChunk,
 		ClearCoat_TangentCodeChunk,
 		ClearCoat_BasisIndexMacro,
-		CustomTangent_TangentCodeChunk);
-
-	// Now update the local shared basis
-
-	// The legacy node can only be added once and we know the worst shape it can take so we hard code it here, 
-	// with fake addresses for BSDF defined in shader code.
-	UMaterialExpression* ChildTopExpression = reinterpret_cast<UMaterialExpression*>(1);
-	UMaterialExpression* ChildBottomExpression = reinterpret_cast<UMaterialExpression*>(2);
-	if (ConvertedStrataMaterialInfo.CountShadingModels() > 1 || ConvertedStrataMaterialInfo.HasShadingModelFromExpression())
-	{
-		Compiler->StrataCompilationGetOperator(ChildTopExpression).BSDFRegisteredSharedLocalBasis = ClearCoat_NewRegisteredSharedLocalBasis;
-		Compiler->StrataCompilationGetOperator(ChildBottomExpression).BSDFRegisteredSharedLocalBasis = NewRegisteredSharedLocalBasis;
-	}
-	else
-	{
-		check(ConvertedStrataMaterialInfo.CountShadingModels() == 1);
-
-		if (ConvertedStrataMaterialInfo.HasShadingModel(SSM_Unlit))
-		{
-			Compiler->StrataCompilationGetOperator(this).BSDFRegisteredSharedLocalBasis = NewRegisteredSharedLocalBasis;
-		}
-		else if (ConvertedStrataMaterialInfo.HasShadingModel(SSM_DefaultLit))
-		{
-			Compiler->StrataCompilationGetOperator(ChildTopExpression).BSDFRegisteredSharedLocalBasis = ClearCoat_NewRegisteredSharedLocalBasis;
-			Compiler->StrataCompilationGetOperator(ChildBottomExpression).BSDFRegisteredSharedLocalBasis = NewRegisteredSharedLocalBasis;
-		}
-		else if (ConvertedStrataMaterialInfo.HasShadingModel(SSM_SubsurfaceLit))
-		{
-			Compiler->StrataCompilationGetOperator(ChildTopExpression).BSDFRegisteredSharedLocalBasis = ClearCoat_NewRegisteredSharedLocalBasis;
-			Compiler->StrataCompilationGetOperator(ChildBottomExpression).BSDFRegisteredSharedLocalBasis = NewRegisteredSharedLocalBasis;
-		}
-		else if (ConvertedStrataMaterialInfo.HasShadingModel(SSM_Hair))
-		{
-			Compiler->StrataCompilationGetOperator(this).BSDFRegisteredSharedLocalBasis = NewRegisteredSharedLocalBasis;
-		}
-		else if (ConvertedStrataMaterialInfo.HasShadingModel(SSM_Eye))
-		{
-			Compiler->StrataCompilationGetOperator(this).BSDFRegisteredSharedLocalBasis = NewRegisteredSharedLocalBasis;
-		}
-		else if (ConvertedStrataMaterialInfo.HasShadingModel(SSM_SingleLayerWater))
-		{
-			Compiler->StrataCompilationGetOperator(this).BSDFRegisteredSharedLocalBasis = NewRegisteredSharedLocalBasis;
-		}
-	}
+		CustomTangent_TangentCodeChunk,
+		!StrataOperator.bUseParameterBlending || (StrataOperator.bUseParameterBlending && StrataOperator.bRootOfParameterBlendingSubTree) ? &StrataOperator : nullptr);
 
 	return OutputCodeChunk;
 }
@@ -22010,22 +21971,15 @@ void UMaterialExpressionStrataLegacyConversion::GatherStrataMaterialInfo(FStrata
 
 FStrataOperator* UMaterialExpressionStrataLegacyConversion::StrataGenerateMaterialTopologyTree(class FMaterialCompiler* Compiler, class UMaterialExpression* Parent, int32 OutputIndex)
 {
-	// Legacy conversion now requires a single slab
-	UMaterialExpression* ChildSlabExpression = reinterpret_cast<UMaterialExpression*>(1);
-
 	auto AddDefaultWorstCase = [&](bool bSSS)
 	{
-		FStrataOperator& StrataOperator = Compiler->StrataCompilationRegisterOperator(STRATA_OPERATOR_WEIGHT, this, Parent);
-
-		FStrataOperator& SlabOperator = Compiler->StrataCompilationRegisterOperator(STRATA_OPERATOR_BSDF_LEGACY, ChildSlabExpression, this);
+		FStrataOperator& SlabOperator = Compiler->StrataCompilationRegisterOperator(STRATA_OPERATOR_BSDF_LEGACY, this, Parent);
 		SlabOperator.BSDFType = STRATA_BSDF_TYPE_SLAB;
 		SlabOperator.bBSDFHasSSS = true;
 		SlabOperator.bBSDFHasMFPPluggedIn = true;
 		SlabOperator.bBSDFHasFuzz = true;
 
-		AssignOperatorIndexIfNotNull(StrataOperator.LeftIndex, &SlabOperator);
-
-		return &StrataOperator;
+		return &SlabOperator;
 	};
 
 	// Logic about shading models and complexity should match UMaterialExpressionStrataLegacyConversion::Compile.
diff --git a/Engine/Source/Runtime/Engine/Public/MaterialCompiler.h b/Engine/Source/Runtime/Engine/Public/MaterialCompiler.h
--- a/Engine/Source/Runtime/Engine/Public/MaterialCompiler.h
+++ b/Engine/Source/Runtime/Engine/Public/MaterialCompiler.h
@@ -587,7 +587,8 @@ public:
 		int32 ShadingModel,
 		int32 Normal, int32 Tangent, const FString& SharedLocalBasisIndexMacro,
 		int32 ClearCoat_Normal, int32 ClearCoat_Tangent, const FString& ClearCoat_SharedLocalBasisIndexMacro,
-		int32 CustomTangent_Tangent) = 0;
+		int32 CustomTangent_Tangent,
+		FStrataOperator* PromoteToOperator) = 0;
 	virtual int32 StrataVolumetricFogCloudBSDF(int32 Albedo, int32 Extinction, int32 EmissiveColor, int32 AmbientOcclusion) = 0;
 	virtual int32 StrataUnlitBSDF(int32 EmissiveColor, int32 TransmittanceColor) = 0;
 	virtual int32 StrataHairBSDF(int32 BaseColor, int32 Scatter, int32 Specular, int32 Roughness, int32 Backlit, int32 EmissiveColor, int32 Tangent, const FString& SharedLocalBasisIndexMacro, FStrataOperator* PromoteToOperator) = 0;
@@ -1150,7 +1151,8 @@ public:
 		int32 ShadingModel,
 		int32 Normal, int32 Tangent, const FString& SharedLocalBasisIndexMacro,
 		int32 ClearCoat_Normal, int32 ClearCoat_Tangent, const FString& ClearCoat_SharedLocalBasisIndexMacro,
-		int32 CustomTangent_Tangent) override
+		int32 CustomTangent_Tangent,
+		FStrataOperator* PromoteToOperator) override
 	{
 		return Compiler->StrataConversionFromLegacy(
 			bHasDynamicShadingModels,
@@ -1165,7 +1167,8 @@ public:
 			ShadingModel,
 			Normal, Tangent, SharedLocalBasisIndexMacro,
 			ClearCoat_Normal, ClearCoat_Tangent, ClearCoat_SharedLocalBasisIndexMacro,
-			CustomTangent_Tangent);
+			CustomTangent_Tangent,
+			PromoteToOperator);
 	}
 
 	virtual int32 StrataVolumetricFogCloudBSDF(int32 Albedo, int32 Extinction, int32 EmissiveColor, int32 AmbientOcclusion) override
