Strata - fixed mesh decals similarly to deferred decals.

diff --git a/Engine/Shaders/Private/DeferredDecal.usf b/Engine/Shaders/Private/DeferredDecal.usf
--- a/Engine/Shaders/Private/DeferredDecal.usf
+++ b/Engine/Shaders/Private/DeferredDecal.usf
@@ -142,7 +142,7 @@ void FPixelShaderInOut_MainPS(inout FPixelShaderIn In, inout FPixelShaderOut Out
 #if STRATA_ENABLED    
 	FStrataPixelHeader StrataHeader = InitialiseStrataPixelHeader();
 	StrataHeader.StrataTree = MaterialParameters.StrataTree;
-	StrataHeader.BSDFCount = 1;
+	StrataHeader.BSDFCount = MaterialParameters.StrataTree.BSDFCount; // Must be one due to parameter blending being forced on the entire Strata tree form the convert-to-decal node.
 	StrataHeader.IrradianceAO = InitIrradianceAndOcclusion();
 	StrataHeader.IrradianceAO.MaterialAO = GetMaterialAmbientOcclusion(PixelMaterialInputs);
 	StrataHeader.SharedLocalBases = MaterialParameters.SharedLocalBases;
diff --git a/Engine/Shaders/Private/MeshDecals.usf b/Engine/Shaders/Private/MeshDecals.usf
--- a/Engine/Shaders/Private/MeshDecals.usf
+++ b/Engine/Shaders/Private/MeshDecals.usf
@@ -134,11 +134,27 @@ void FPixelShaderInOut_MainPS(
 
 #if STRATA_ENABLED
 	FStrataPixelHeader StrataHeader = InitialiseStrataPixelHeader();
-	StrataHeader.BSDFCount = 1;
+	StrataHeader.StrataTree = MaterialParameters.StrataTree;
+	StrataHeader.BSDFCount = MaterialParameters.StrataTree.BSDFCount; // Must be one due to parameter blending being forced on the entire Strata tree form the convert-to-decal node.
 	StrataHeader.IrradianceAO = InitIrradianceAndOcclusion();
 	StrataHeader.IrradianceAO.MaterialAO = GetMaterialAmbientOcclusion(PixelMaterialInputs);
 	StrataHeader.SharedLocalBases = MaterialParameters.SharedLocalBases;
-	const float Coverage = DecalCommonOutput(Out, StrataHeader, StrataHeader.StrataTree.BSDFs[0], 1.f /*CoverageWeight*/);
+
+	float Coverage = 0.0;
+	if (StrataHeader.StrataTree.BSDFCount > 0)
+	{
+		FStrataData StrataData = PixelMaterialInputs.FrontMaterial;
+
+		FStrataAddressing NullStrataAddressing = (FStrataAddressing)0;	// Fake unused in StrataCreateBSDFContext when using Forward inline shading
+		const float3 V = MaterialParameters.CameraVector;
+
+		// Update tree (coverage/transmittance/luminace weights)
+		const FStrataIntegrationSettings Settings = InitStrataIntegrationSettings(false /*bForceFullyRough*/, false /*bRoughDiffuseEnabled*/, 0);
+		StrataUpdateTree(NullStrataAddressing, StrataHeader, StrataHeader.StrataTree, V, Settings);
+
+		// Since the Convert-To-Decal node forces parameter blending for the entire Strata tree, the BSDF we are interested in will be in 0.
+		Coverage = DecalCommonOutput(Out, StrataHeader, StrataHeader.StrataTree.BSDFs[0], 1.0f /*CoverageWeight*/);
+	}
 	const float Opacity = Coverage;
 #else // STRATA_ENABLED
 	float3 Color = GetMaterialEmissive(PixelMaterialInputs);
