Fix for water material compilation error due to missing default define of MATERIALBLENDING_SOLID.

#rnx

diff --git a/Engine/Shaders/Private/Definitions.usf b/Engine/Shaders/Private/Definitions.usf
--- a/Engine/Shaders/Private/Definitions.usf
+++ b/Engine/Shaders/Private/Definitions.usf
@@ -18,6 +18,10 @@
 #define	MATERIAL_TWOSIDED_SEPARATE_PASS					0
 #endif
 
+#ifndef MATERIALBLENDING_SOLID
+#define MATERIALBLENDING_SOLID							0
+#endif
+
 #ifndef MATERIALBLENDING_MASKED
 #define MATERIALBLENDING_MASKED							0
 #endif
