Fixes black hole filling artifacts in DOF

diff --git a/Engine/Shaders/Private/DiaphragmDOF/DOFRecombine.usf b/Engine/Shaders/Private/DiaphragmDOF/DOFRecombine.usf
--- a/Engine/Shaders/Private/DiaphragmDOF/DOFRecombine.usf
+++ b/Engine/Shaders/Private/DiaphragmDOF/DOFRecombine.usf
@@ -1005,20 +1005,24 @@ void RecombineMainCS(
 		GatherForegroundAdditiveColor *= SafeRcp(1 - GatherForegroundTranslucency);
 		GatherForegroundTranslucency = 0;
 	}
-	else if (CONFIG_COMPOSITING_METHOD == COMPOSITING_METHOD_BILINEAR_BKG)
+	else
 	{
-		BackgroundColor = BackgroundConvolution_SceneColor.SampleLevel(GlobalBilinearClampedSampler, InputParameters.DOFBufferUV, 0);
+		#if CONFIG_COMPOSITING_METHOD == COMPOSITING_METHOD_BILINEAR_BKG
+		{
+			BackgroundColor = BackgroundConvolution_SceneColor.SampleLevel(GlobalBilinearClampedSampler, InputParameters.DOFBufferUV, 0);
 
-		#if CONFIG_DOF_ALPHA
-			BackgroundValidity = BackgroundConvolution_SeparateAlpha.SampleLevel(GlobalBilinearClampedSampler, InputParameters.DOFBufferUV, 0).r;
-		#else
-			//Background = float4(0, 0, 0, 1);
-			BackgroundValidity = BackgroundColor.a;
-		#endif
-	
-		// Make sure the background color is always normalized, or unrendered.
-		BackgroundColor *= SafeRcp(BackgroundValidity);
+			#if CONFIG_DOF_ALPHA
+				BackgroundValidity = BackgroundConvolution_SeparateAlpha.SampleLevel(GlobalBilinearClampedSampler, InputParameters.DOFBufferUV, 0).r;
+			#else
+				//Background = float4(0, 0, 0, 1);
+				BackgroundValidity = BackgroundColor.a;
+			#endif
 	
+			// Make sure the background color is always normalized, or unrendered.
+			BackgroundColor *= SafeRcp(BackgroundValidity);
+		}
+		#endif
+
 		// Hole fill the background in output final scene color before composing foreground on top.
 		#if CONFIG_HOLE_FILLING_METHOD == HOLE_FILLING_METHOD_SEPARATE_GATHER
 		{
