Fix incorrect distance field indirect self shadowing.
- When SphereRadius > MaxEncodedDistance ray is incorrectly occluded.
- Limit SphereRadius to MaxEncodedDistance.

diff --git a/Engine/Shaders/Private/CapsuleShadowShaders.usf b/Engine/Shaders/Private/CapsuleShadowShaders.usf
--- a/Engine/Shaders/Private/CapsuleShadowShaders.usf
+++ b/Engine/Shaders/Private/CapsuleShadowShaders.usf
@@ -868,6 +868,12 @@ float ShadowConeTraceAgainstCulledDistanceFieldCasters(
 
 		{
 			FDFAssetData DFAssetData = LoadDFAssetDataHighestResolution(DFObjectData.AssetIndex);
+
+			const float MaxEncodedDistance = DFAssetData.DistanceFieldToVolumeScaleBias.x + DFAssetData.DistanceFieldToVolumeScaleBias.y;
+
+			// Prevent incorrect shadowing when sampling invalid bricks by limiting VolumeMaxSphereRadius to MaxEncodedDistance
+			VolumeMaxSphereRadius = min(VolumeMaxSphereRadius, MaxEncodedDistance);
+
 			float MinTraceVisibility = 1;
 			//@todo - derive from texel size
 			float StartOffset = .02f;
