* Convert common packing pattern into packing funciton
* Add help for packing RG16

diff --git a/Engine/Shaders/Private/Strata/Strata.ush b/Engine/Shaders/Private/Strata/Strata.ush
--- a/Engine/Shaders/Private/Strata/Strata.ush
+++ b/Engine/Shaders/Private/Strata/Strata.ush
@@ -551,7 +551,7 @@ FHaziness InitialiseHaziness()
 uint PackHaziness(FHaziness Haziness)
 {
 	uint High8Bits = (Haziness.bSimpleClearCoat ? 0x80 : 0x00) | uint(saturate(Haziness.Weight) * 127);
-	uint Data = (High8Bits << 8) | uint(saturate(Haziness.Roughness) * 255);
+	uint Data = (High8Bits << 8) | PackR8(Haziness.Roughness);
 	return Data;
 }
 
@@ -560,7 +560,7 @@ FHaziness UnpackHaziness(uint Data)
 	FHaziness Haziness;
 	Haziness.bSimpleClearCoat	= Data & 0x8000;
 	Haziness.Weight				= float(uint((Data>>8) & 0x7F)) * (1.0f / 127.0f);
-	Haziness.Roughness			= float(uint(Data & 0xFF)) * (1.0f / 255.0f);
+	Haziness.Roughness			= UnpackR8(Data);
 	return Haziness;
 }
 
@@ -603,7 +603,7 @@ uint StrataPackIrradianceAndOcclusion(FStrataIrradianceAndOcclusion In, float Qu
 	// No space for AO. Multiply IndirectIrradiance by AO instead of storing.
 	return min(0xFFu, (EncodeIndirectIrradiance(In.IndirectIrradiance * In.MaterialAO) + QuantizationBias * (1.0 / 255.0)) * 0xFFu);
 #else
-	return uint(saturate(In.MaterialAO) * 255.0f);
+	return PackR8(In.MaterialAO);
 #endif
 }
 
@@ -617,7 +617,7 @@ FStrataIrradianceAndOcclusion StrataUnpackIrradianceAndOcclusion(uint In /*8bits
 #elif ALLOW_STATIC_LIGHTING
 	Out.MaterialAO = 1.f;
 	Out.DiffuseIndirectSampleOcclusion = 0x0u;
-	Out.IndirectIrradiance = DecodeIndirectIrradiance(In * rcp(255));
+	Out.IndirectIrradiance = DecodeIndirectIrradiance(UnpackR8(In));
 #else
 	Out.MaterialAO = float(In)* rcp(255.0f);
 	Out.DiffuseIndirectSampleOcclusion = 0x0u;
@@ -749,7 +749,7 @@ void StrataSubSurfaceHeaderSetProfile(inout FStrataSubsurfaceHeader SSSHeader, f
 {
 	SSSHeader.Bytes &= SSSHEADER_TYPE_MASK;
 	SSSHeader.Bytes |= ProfileId << 24;
-	SSSHeader.Bytes |= uint(saturate(RadiusScale) * 255.0f) << 16;
+	SSSHeader.Bytes |= PackR8(RadiusScale) << 16;
 }
 
 void StrataSubSurfaceHeaderSetNonProfile(inout FStrataSubsurfaceHeader SSSHeader, float3 MeanFreePath)
@@ -765,7 +765,7 @@ void StrataSubSurfaceHeaderSetWrap(inout FStrataSubsurfaceHeader SSSHeader, floa
 	// * Isotropic phase function == Thick surface (i.e., opacity => 1)
 	const float Opacity = 1.f-abs(PhaseAnisotropy);
 	SSSHeader.Bytes &= SSSHEADER_TYPE_MASK;
-	SSSHeader.Bytes |= uint(saturate(Opacity) * 255.0f) << 2;
+	SSSHeader.Bytes |= PackR8(Opacity) << 2;
 }
 
 bool StrataSubSurfaceHeaderGetIsValid(in FStrataSubsurfaceHeader SSSHeader)
@@ -805,8 +805,7 @@ uint StrataSubSurfaceHeaderGetProfileId(in FStrataSubsurfaceHeader SSSHeader)
 float StrataSubSurfaceHeaderGetProfileRadiusScale(in FStrataSubsurfaceHeader SSSHeader)
 {
 	const bool bIsValidAndProfile = SSSHEADER_TYPE(SSSHeader) == SSS_TYPE_DIFFUSION_PROFILE;
-	const uint RadiusScale = bIsValidAndProfile ? (SSSHeader.Bytes >> 16) & 0xFF : 255;
-	return float(RadiusScale) * rcp(255.0f);
+	return bIsValidAndProfile ? UnpackR8(SSSHeader.Bytes >> 16) : 1.f;
 }
 
 float3 StrataSubSurfaceHeaderGetMFP(in FStrataSubsurfaceHeader SSSHeader)				// Only when !IsProfile
@@ -819,8 +818,7 @@ float3 StrataSubSurfaceHeaderGetMFP(in FStrataSubsurfaceHeader SSSHeader)				//
 
 float StrataSubSurfaceHeaderGetWrapOpacity(in FStrataSubsurfaceHeader SSSHeader)
 {
-	const uint Opacity = SSSHEADER_TYPE(SSSHeader) == SSS_TYPE_WRAP ? ((SSSHeader.Bytes >> 2) & 0xFF) : 255;
-	return float(Opacity) * rcp(255.0f);
+	return SSSHEADER_TYPE(SSSHeader) == SSS_TYPE_WRAP ? UnpackR8(SSSHeader.Bytes >> 2) : 1.0f;
 }
 
 float StrataSubSurfaceHeaderGetOpacity(in FStrataSubsurfaceHeader SSSHeader)
@@ -2901,56 +2899,46 @@ float3 UnpackR11G11B10F(uint rgb)
 	return float3(r, g, b);
 }
 
-uint2 PackR16G16B16A16F(float4 rgba)
+uint PackRG16(float2 In)
 {
-	uint A = (f32tof16(rgba.r) << 16) | f32tof16(rgba.g);
-	uint B = (f32tof16(rgba.b) << 16) | f32tof16(rgba.a);
-	return uint2 (A, B);
+	return f32tof16(In.r) | (f32tof16(In.g) << 16); // Add intrinsic
 }
 
-float4 UnpackR16G16B16A16F(uint2 AB)
+float2 UnpackRG16(uint In)
 {
-	float r = f16tof32(AB.x >> 16);
-	float g = f16tof32(AB.x);
-	float b = f16tof32(AB.y >> 16);
-	float a = f16tof32(AB.y);
-	return float4(r, g, b, a);
+	return float2(f16tof32(In), f16tof32(In>>16)); // Add intrinsic
 }
 
-uint PackColorLinearToGamma2(float3 rgb)
+uint2 PackR16G16B16A16F(float4 In)
 {
-	uint r = (uint(sqrt(saturate(rgb.r)) * 255.0f) << 16);
-	uint g = (uint(sqrt(saturate(rgb.g)) * 255.0f) << 8);
-	uint b = (uint(sqrt(saturate(rgb.b)) * 255.0f) << 0);
-	return r | g | b;
+	return uint2(PackRG16(In.xy), PackRG16(In.zw));
 }
 
-float3 UnpackColorGamma2ToLinear(uint rgb)
+float4 UnpackR16G16B16A16F(uint2 In)
 {
-	float r = float((rgb >> 16) & 0x000000FF);
-	float g = float((rgb >> 8)  & 0x000000FF);
-	float b = float((rgb >> 0)  & 0x000000FF);
-	float3 Gamma2 = float3(r, g, b) * (1.0f / 255.0f);
-	return Gamma2 * Gamma2;	// Simple gamma 2.0
+	return float4(UnpackRG16(In.x), UnpackRG16(In.y));
 }
 
-uint PackColorLinearToGamma2AlphaLinear(float4 rgba)
+uint PackColorLinearToGamma2(float3 In)
 {
-	uint r = (uint(sqrt(saturate(rgba.r)) * 255.0f) << 24);
-	uint g = (uint(sqrt(saturate(rgba.g)) * 255.0f) << 16);
-	uint b = (uint(sqrt(saturate(rgba.b)) * 255.0f) << 8);
-	uint a = (uint(     saturate(rgba.a)  * 255.0f) << 0);
-	return r | g | b | a;
+	return PackRGBA8(float4(sqrt(saturate(In.rgb)), 0)); // Simple gamma 2.0
+}
+
+float3 UnpackColorGamma2ToLinear(uint In)
+{
+	float3 Out = UnpackRGBA8(In).xyz;
+	return Out.rgb * Out.rgb; // Simple gamma 2.0	
 }
 
-float4 UnpackColorGamma2ToLinearAlphaLinear(uint rgba)
+uint PackColorLinearToGamma2AlphaLinear(float4 In)
 {
-	float r = float((rgba >> 24) & 0x000000FF);
-	float g = float((rgba >> 16) & 0x000000FF);
-	float b = float((rgba >> 8)  & 0x000000FF);
-	float a = float((rgba >> 0)  & 0x000000FF);
-	float3 Gamma2 = float3(r, g, b) * (1.0f / 255.0f);
-	return float4(Gamma2 * Gamma2, a * (1.0f / 255.0f));	// Simple gamma 2.0
+	return PackRGBA8(float4(sqrt(saturate(In.rgb)), In.a)); // Simple gamma 2.0
+}
+
+float4 UnpackColorGamma2ToLinearAlphaLinear(uint In)
+{
+	float4 Out = UnpackRGBA8(In);
+	return float4(Out.rgb * Out.rgb, Out.a); // Simple gamma 2.0	
 }
 
 uint PackR8(float Value)
@@ -2958,27 +2946,27 @@ uint PackR8(float Value)
 	return uint(saturate(Value) * 255.0f);
 }
 
-float UnpackR8(uint Value)
+float UnpackR8(uint In)
 {
-	return float(Value & 0xFF) * (1.0f / 255.0f);
+	return float(In & 0xFF) * (1.0f / 255.0f);
 }
 
-uint PackRGBA8(float4 rgba)
+uint PackRGBA8(float4 In)
 {
-	uint r = (uint(saturate(rgba.r) * 255.0f) << 0);
-	uint g = (uint(saturate(rgba.g) * 255.0f) << 8);
-	uint b = (uint(saturate(rgba.b) * 255.0f) << 16);
-	uint a = (uint(saturate(rgba.a) * 255.0f) << 24);
+	uint r = (uint(saturate(In.r) * 255.0f) << 0);
+	uint g = (uint(saturate(In.g) * 255.0f) << 8);
+	uint b = (uint(saturate(In.b) * 255.0f) << 16);
+	uint a = (uint(saturate(In.a) * 255.0f) << 24);
 	return r | g | b | a;
 }
 
-float4 UnpackRGBA8(uint rgba)
+float4 UnpackRGBA8(uint In)
 {
 	float4 Out;
-	Out.r = float((rgba >> 0)  & 0xFF) * (1.0f/255.0f);
-	Out.g = float((rgba >> 8)  & 0xFF) * (1.0f/255.0f);
-	Out.b = float((rgba >> 16) & 0xFF) * (1.0f/255.0f);
-	Out.a = float((rgba >> 24) & 0xFF) * (1.0f/255.0f);
+	Out.r = float((In >> 0)  & 0xFF) * (1.0f/255.0f);
+	Out.g = float((In >> 8)  & 0xFF) * (1.0f/255.0f);
+	Out.b = float((In >> 16) & 0xFF) * (1.0f/255.0f);
+	Out.a = float((In >> 24) & 0xFF) * (1.0f/255.0f);
 	return Out;
 }
 
@@ -3061,14 +3049,12 @@ float2 SignNotZero(in float2 v)
 
 uint StrataPackNormal(in float3 Normal)
 {
-	const float2 Result = UnitVectorToOctahedron(Normal);
-	return f32tof16(Result.x) | (f32tof16(Result.y) << 16);
+	return PackRG16(UnitVectorToOctahedron(Normal));
 }
 
 float3 StrataUnpackNormal(uint PackedNormal)
 {
-	const float2 xy = float2(f16tof32(0xFFFF & PackedNormal), f16tof32(0xFFFF & (PackedNormal >> 16)));
-	return OctahedronToUnitVector(xy);
+	return OctahedronToUnitVector(UnpackRG16(PackedNormal));
 }
 
 uint StrataPackNormal24(in float3 Normal)
