Use packing/unpacking intrinsic for Strata data.

diff --git a/Engine/Shaders/Private/Strata/Strata.ush b/Engine/Shaders/Private/Strata/Strata.ush
--- a/Engine/Shaders/Private/Strata/Strata.ush
+++ b/Engine/Shaders/Private/Strata/Strata.ush
@@ -2901,12 +2901,16 @@ float3 UnpackR11G11B10F(uint rgb)
 
 uint PackRG16(float2 In)
 {
-	return f32tof16(In.r) | (f32tof16(In.g) << 16); // Add intrinsic
+#if COMPILER_SUPPORTS_PACK_INTRINSICS
+	return PackUInt2ToUInt(f32tof16(In.r), f32tof16(In.g));
+#else
+	return f32tof16(In.r) | (f32tof16(In.g) << 16);
+#endif
 }
 
 float2 UnpackRG16(uint In)
 {
-	return float2(f16tof32(In), f16tof32(In>>16)); // Add intrinsic
+	return float2(f16tof32(In), f16tof32(In>>16));
 }
 
 uint2 PackR16G16B16A16F(float4 In)
@@ -2948,26 +2952,42 @@ uint PackR8(float Value)
 
 float UnpackR8(uint In)
 {
+#if COMPILER_SUPPORTS_UNPACKBYTEN
+	return UnpackByte0(In) * (1.0f / 255.0f);
+#else
 	return float(In & 0xFF) * (1.0f / 255.0f);
+#endif
 }
 
 uint PackRGBA8(float4 In)
 {
+#if COMPILER_SUPPORTS_PACK_INTRINSICS
+	return PackFloat4ToUInt(saturate(In) * 255.0f);
+#else
 	uint r = (uint(saturate(In.r) * 255.0f) << 0);
 	uint g = (uint(saturate(In.g) * 255.0f) << 8);
 	uint b = (uint(saturate(In.b) * 255.0f) << 16);
 	uint a = (uint(saturate(In.a) * 255.0f) << 24);
 	return r | g | b | a;
+#endif
 }
 
 float4 UnpackRGBA8(uint In)
 {
+#if COMPILER_SUPPORTS_UNPACKBYTEN
+	return float4(
+		UnpackByte0(In),
+		UnpackByte1(In),
+		UnpackByte2(In),
+		UnpackByte3(In)) * (1.0f / 255.0f);
+#else
 	float4 Out;
 	Out.r = float((In >> 0)  & 0xFF) * (1.0f/255.0f);
 	Out.g = float((In >> 8)  & 0xFF) * (1.0f/255.0f);
 	Out.b = float((In >> 16) & 0xFF) * (1.0f/255.0f);
 	Out.a = float((In >> 24) & 0xFF) * (1.0f/255.0f);
 	return Out;
+#endif
 }
 
 uint PackR7G7B6Gamma2(float3 rgb, float Dither)
