Compile RayTracingDynamicGeometryConverter shaders for Vulkan:
* Note: We relied on compiler to remove code that requires binding InstanceBuffer but it's not happening for SPIRV so I've added a workaround to remove the code.

diff --git a/Engine/Shaders/Private/SceneData.ush b/Engine/Shaders/Private/SceneData.ush
--- a/Engine/Shaders/Private/SceneData.ush
+++ b/Engine/Shaders/Private/SceneData.ush
@@ -25,6 +25,10 @@
 	#define NEEDS_LIGHTMAP_COORDINATE 0
 #endif
 
+#ifndef USE_INSTANCE_CULLING_DATA
+	#define USE_INSTANCE_CULLING_DATA 1
+#endif
+
 #ifndef VF_REQUIRES_PER_INSTANCE_CUSTOM_DATA
 	#define VF_REQUIRES_PER_INSTANCE_CUSTOM_DATA 0
 #endif
@@ -1043,6 +1047,9 @@ FSceneDataIntermediates GetSceneDataIntermediates(uint InstanceIdOffset, uint Dr
 		Intermediates.InstanceId = GetPrimitiveData(PrimitiveID).InstanceSceneDataOffset + DrawInstanceId;
 		Intermediates.ViewIndex = 0;
 	}
+// Workaround for Vulkan SPIRV issue when "else" branch is not removed when it should, which leads to higher level code expecting InstanceCulling buffer to be bound.
+// See: [REDACTED_PROFILE_URL]/SPIRV-Tools/issues/4902
+#if USE_INSTANCE_CULLING_DATA
 	else
 	{
 		Intermediates.InstanceId = InstanceCulling.InstanceIdsBuffer[InstanceIdOffset + DrawInstanceId] & ((1U << 28U) - 1);
@@ -1050,6 +1057,7 @@ FSceneDataIntermediates GetSceneDataIntermediates(uint InstanceIdOffset, uint Dr
 		// Note: this is an index of views for this render pass, not the view ID in the culling manager.
 		Intermediates.ViewIndex = InstanceCulling.InstanceIdsBuffer[InstanceIdOffset + DrawInstanceId] >> 28U;
 	}
+#endif
 
 	Intermediates.InstanceData = GetInstanceSceneData(Intermediates.InstanceId, View.InstanceSceneDataSOAStride);
 	Intermediates.PrimitiveId = Intermediates.InstanceData.PrimitiveId;
diff --git a/Engine/Source/Runtime/Renderer/Private/RayTracing/RayTracingDynamicGeometry.cpp b/Engine/Source/Runtime/Renderer/Private/RayTracing/RayTracingDynamicGeometry.cpp
--- a/Engine/Source/Runtime/Renderer/Private/RayTracing/RayTracingDynamicGeometry.cpp
+++ b/Engine/Source/Runtime/Renderer/Private/RayTracing/RayTracingDynamicGeometry.cpp
@@ -53,12 +53,13 @@ public:
 
 	static bool ShouldCompilePermutation(const FMeshMaterialShaderPermutationParameters& Parameters)
 	{
-		return Parameters.VertexFactoryType->SupportsRayTracingDynamicGeometry() && ShouldCompileRayTracingShadersForProject(Parameters.Platform);
+		return Parameters.VertexFactoryType->SupportsRayTracingDynamicGeometry() && IsRayTracingEnabledForProject(Parameters.Platform) && RHISupportsRayTracing(Parameters.Platform);
 	}
 
 	static void ModifyCompilationEnvironment(const FMaterialShaderPermutationParameters& Parameters, FShaderCompilerEnvironment& OutEnvironment)
 	{
 		OutEnvironment.SetDefine(TEXT("SCENE_TEXTURES_DISABLED"), 1);
+		OutEnvironment.SetDefine(TEXT("USE_INSTANCE_CULLING_DATA"), 0);
 	}
 
 	void GetShaderBindings(
