Use fast sqrt for strata sRGB encoding to reduce quater rate instructions.

diff --git a/Engine/Shaders/Private/Strata/Strata.ush b/Engine/Shaders/Private/Strata/Strata.ush
--- a/Engine/Shaders/Private/Strata/Strata.ush
+++ b/Engine/Shaders/Private/Strata/Strata.ush
@@ -2873,6 +2873,8 @@ float StrataGetLegacyShadingModels(in FStrataBSDF BSDF)
 // BSDFs packing/unpacking
 // 11G11B10F and 10F are unsigned float format
 
+float3 sqrtFast3(float3 In) { return float3(sqrtFast(In.x), sqrtFast(In.y), sqrtFast(In.z)); }
+
 uint Pack10F(float Value)
 {
 	return (f32tof16(Value) >> 5) & 0x000003FF;
@@ -2925,7 +2927,7 @@ float4 UnpackR16G16B16A16F(uint2 In)
 
 uint PackColorLinearToGamma2(float3 In)
 {
-	return PackRGBA8(float4(sqrt(saturate(In.rgb)), 0)); // Simple gamma 2.0
+	return PackRGBA8(float4(sqrtFast3(saturate(In.rgb)), 0)); // Simple gamma 2.0
 }
 
 float3 UnpackColorGamma2ToLinear(uint In)
@@ -2936,7 +2938,7 @@ float3 UnpackColorGamma2ToLinear(uint In)
 
 uint PackColorLinearToGamma2AlphaLinear(float4 In)
 {
-	return PackRGBA8(float4(sqrt(saturate(In.rgb)), In.a)); // Simple gamma 2.0
+	return PackRGBA8(float4(sqrtFast3(saturate(In.rgb)), In.a)); // Simple gamma 2.0
 }
 
 float4 UnpackColorGamma2ToLinearAlphaLinear(uint In)
@@ -3000,7 +3002,7 @@ uint PackR7G7B6Gamma2(float3 rgb, float Dither)
 	const float3 Dither776 = float3(Dither7bits, Dither7bits, Dither6bits);
 
 	// Dither value to hide more aggressive bit reduction.
-	const float3 dithered_rgb = DitherIfNonBlack(sqrt(saturate(rgb)), Dither776);
+	const float3 dithered_rgb = DitherIfNonBlack(sqrtFast3(saturate(rgb)), Dither776);
 
 	uint r = (uint(dithered_rgb.r * 127.0f) << 0);
 	uint g = (uint(dithered_rgb.g * 127.0f) << 7);
