Fixes PixelDepth material expression on light functions

diff --git a/Engine/Shaders/Private/LightFunctionCommon.ush b/Engine/Shaders/Private/LightFunctionCommon.ush
--- a/Engine/Shaders/Private/LightFunctionCommon.ush
+++ b/Engine/Shaders/Private/LightFunctionCommon.ush
@@ -51,6 +51,7 @@ float3 GetLightFunctionColor(float3 LightVector, float3 TranslatedWorldSpace)
 	MaterialParameters.ReflectionVector = LightVector.xyz;
 	MaterialParameters.LightVector = LightVector.xyz;
 	MaterialParameters.AbsoluteWorldPosition = LWCSubtract(TranslatedWorldSpace, PrimaryView.PreViewTranslation);
+	MaterialParameters.ScreenPosition = mul(float4(TranslatedWorldSpace, 1.0f), PrimaryView.TranslatedWorldToClip);
 
 	FPixelMaterialInputs PixelMaterialInputs;
 	CalcPixelMaterialInputs(MaterialParameters, PixelMaterialInputs);
