Fixes stability issue on CitySample's roads due to TSR

diff --git a/Engine/Shaders/Private/TemporalSuperResolution/TSRRejectShading.usf b/Engine/Shaders/Private/TemporalSuperResolution/TSRRejectShading.usf
--- a/Engine/Shaders/Private/TemporalSuperResolution/TSRRejectShading.usf
+++ b/Engine/Shaders/Private/TemporalSuperResolution/TSRRejectShading.usf
@@ -434,7 +434,7 @@ void MainCS(
 	tsr_simd_halfC FinalGuide;
 	{
 		const tsr_half CurrentFrameSampleCount = tsr_half(1.0);
-		const tsr_half PrevLowFrequencySampleCount = tsr_half(24.0);
+		const tsr_half PrevLowFrequencySampleCount = tsr_half(8.0);
 		
 		// Reject the history based on the low frequencies.
 		tsr_simd_half PrevLowFrequencyRejectionFactor;
