Fixes opacity alpha channel in TSR

diff --git a/Engine/Shaders/Private/TemporalSuperResolution/TSRUpdateHistory.usf b/Engine/Shaders/Private/TemporalSuperResolution/TSRUpdateHistory.usf
--- a/Engine/Shaders/Private/TemporalSuperResolution/TSRUpdateHistory.usf
+++ b/Engine/Shaders/Private/TemporalSuperResolution/TSRUpdateHistory.usf
@@ -1162,6 +1162,7 @@ void MainCS(
 		#if CONFIG_SCENE_COLOR_APLHA
 		{
 			FinalOutputColor[3] = select(FinalOutputColor[3] > 0.995, tsr_half(1.0), FinalOutputColor[3]);
+			FinalOutputColor[3] = select(FinalOutputColor[3] < 0.005, tsr_half(0.0), FinalOutputColor[3]);
 		}
 		#endif
 		
