Fix flickering on sky domes due to precision bug on some Android devices where half precision variables are internally handled inconsistantly

diff --git a/Engine/Shaders/Private/MaterialTemplate.ush b/Engine/Shaders/Private/MaterialTemplate.ush
--- a/Engine/Shaders/Private/MaterialTemplate.ush
+++ b/Engine/Shaders/Private/MaterialTemplate.ush
@@ -310,7 +310,7 @@ struct FMaterialPixelParameters
 	half3 ReflectionVector;
 
 	/** Normalized world space camera vector, which is the vector from the point being shaded to the camera position. */
-	half3 CameraVector;
+	float3 CameraVector;
 
 	/** World space light vector, only valid when rendering a light function. */
 	half3 LightVector;
