* Use fast math for normal encoding. Remove quarter rate instruction
* Use fast non-zero detection for dithering. Reduce 10% cycle from dithering routine.

diff --git a/Engine/Shaders/Private/Strata/Strata.ush b/Engine/Shaders/Private/Strata/Strata.ush
--- a/Engine/Shaders/Private/Strata/Strata.ush
+++ b/Engine/Shaders/Private/Strata/Strata.ush
@@ -2578,37 +2578,30 @@ void StrataSanitizeBSDF(inout FStrataBSDF BSDF)
 	}
 }
 
-// This is used to respect artistic desire, e.g. shutdown diffuse or specular.
-float3 DitherIfNonBlack(float3 Input, float3 Dither)
+// Return 0 if v is zero, 1 otherwise
+float IsNonZeroFast(float v)
 {
-	float3 NewInput = Input + Dither;
-	return float3(
-		NewInput.x >= 0.0 && NewInput.x <= 1.0 ? NewInput.x : Input.x,
-		NewInput.y >= 0.0 && NewInput.y <= 1.0 ? NewInput.y : Input.y,
-		NewInput.z >= 0.0 && NewInput.z <= 1.0 ? NewInput.z : Input.z);
+	return min(asuint(v) & 0x7FFFFFFF, 1u);
 }
-float DitherIfNonBlack(float Input, float Dither)
+
+float3 IsNonZeroFast(float3 v)
 {
-	float NewInput = Input + Dither;
-	return NewInput.x >= 0.0 && NewInput.x <= 1.0 ? NewInput.x : Input.x;
+	return min(asuint(v) & 0x7FFFFFFF, 1u.xxx);
 }
 
-void StrataDitherBSDF(inout FStrataBSDF BSDF, in float Dither)
+// This is used to respect artistic desire, e.g. shutdown diffuse or specular.
+float3 DitherIfNonBlack(float3 Input, float3 Dither)
 {
-	if (BSDF_GETTYPE(BSDF) == STRATA_BSDF_TYPE_SLAB)
-	{
-		// We execute a dither that can go in positive/negative direction to avoid surfaces to looks overall brighter.
-		const float CenteredDither = Dither - 0.5f;
-		const float Dither8bits = CenteredDither * 1.0f / 255.0f;
-		const float Dither7bits = CenteredDither * 1.0f / 127.0f;
-		const float Dither6bits = CenteredDither * 1.0f / 63.0f;
-		const float3 Dither776 = float3(Dither7bits, Dither7bits, Dither6bits);
-
-		SLAB_DIFFUSEALBEDO(BSDF) = DitherIfNonBlack(SLAB_DIFFUSEALBEDO(BSDF), Dither776);
-		SLAB_F0(BSDF) = DitherIfNonBlack(SLAB_F0(BSDF), Dither776);
-	}
+	// IsNonZeroFast ensures dithering is not applied when Input is zero. Input us assumed to be in [0..1].
+	float3 NewInput = Input + Dither * IsNonZeroFast(Input);
+	return saturate(NewInput);
+}
+float DitherIfNonBlack(float Input, float Dither)
+{
+	// IsNonZeroFast ensures dithering is not applied when Input is zero. Input us assumed to be in [0..1].
+	float NewInput = Input + Dither * IsNonZeroFast(Input);
+	return saturate(NewInput);
 }
-
 
 ///////////////////////////////////////////////////////////////////////////////
 // Other tools
@@ -3069,9 +3062,19 @@ float2 SignNotZero(in float2 v)
 	return float2(SignNotZero(v.x), SignNotZero(v.y));
 }
 
+float2 StrataUnitVectorToOctahedron(float3 N)
+{
+	N.xy *= rcpFast(dot(1, abs(N)));
+	if (N.z <= 0)
+	{
+		N.xy = (1 - abs(N.yx)) * select(N.xy >= 0, float2(1, 1), float2(-1, -1));
+	}
+	return N.xy;
+}
+
 uint StrataPackNormal(in float3 Normal)
 {
-	return PackRG16(UnitVectorToOctahedron(Normal));
+	return PackRG16(StrataUnitVectorToOctahedron(Normal));
 }
 
 float3 StrataUnpackNormal(uint PackedNormal)
@@ -3081,7 +3084,7 @@ float3 StrataUnpackNormal(uint PackedNormal)
 
 uint StrataPackNormal24(in float3 Normal)
 {
-	const float2 Result = UnitVectorToOctahedron(Normal);
+	const float2 Result = StrataUnitVectorToOctahedron(Normal);
 	const uint2 PackedXY = uint2(clamp(Result * 2047.0f + 2048.0f, 0.0f, 4095.0f));
 	return PackedXY.x | (PackedXY.y << 12);
 }
@@ -3095,7 +3098,7 @@ float3 StrataUnpackNormal24(uint PackedNormal)
 
 uint StrataPackNormal22(in float3 Normal)
 {
-	const float2 Result = UnitVectorToOctahedron(Normal);
+	const float2 Result = StrataUnitVectorToOctahedron(Normal);
 	const uint2 PackedXY = uint2(clamp(Result * 1023.0f + 1024.0f, 0.0f, 2047.0f));
 	return PackedXY.x | (PackedXY.y << 11);
 }
