GPULM: Implement CPULM-like translucency behavior emulation

diff --git a/Engine/Shaders/Private/PathTracing/PathTracingMaterialHitShader.usf b/Engine/Shaders/Private/PathTracing/PathTracingMaterialHitShader.usf
--- a/Engine/Shaders/Private/PathTracing/PathTracingMaterialHitShader.usf
+++ b/Engine/Shaders/Private/PathTracing/PathTracingMaterialHitShader.usf
@@ -360,7 +360,12 @@ RAY_TRACING_ENTRY_CLOSEST_HIT(PathTracingMaterialCHS,
 	Payload.Specular = 0;
 	Payload.Roughness = 1;
 	Payload.Metallic = 0;
+#if MATERIALBLENDING_TRANSLUCENT && MATERIAL_SHADINGMODEL_THIN_TRANSLUCENT
 	Payload.Opacity = GetMaterialOpacity(PixelMaterialInputs);
+#else
+	// Force opacity to 0 so it is ignored for path hits and only evaluated in AHS as a modulation to throughput
+	Payload.Opacity = 0;
+#endif
 	Payload.Ior = 1;
 
 #if !MATERIAL_TWOSIDED
@@ -516,9 +521,14 @@ RAY_TRACING_ENTRY_ANY_HIT(PathTracingMaterialAHS,
 	float Opacity = saturate(GetMaterialMaskInputRaw(PixelMaterialInputs));
 	float Transparency = 1 - Opacity;
 #elif MATERIALBLENDING_TRANSLUCENT
-	float Opacity = GetMaterialOpacity(PixelMaterialInputs);
-	float Transparency = 1 - Opacity;
-
+	#if !SIMPLIFIED_MATERIAL_SHADER
+		float Opacity = GetMaterialOpacity(PixelMaterialInputs);
+		float Transparency = 1 - Opacity;
+	#else
+		// CPU Lightmass translucency behavior emulation
+		float3 Transparency = lerp(float3(1, 1, 1), GetMaterialBaseColor(PixelMaterialInputs), GetMaterialOpacity(PixelMaterialInputs));
+	#endif
+	
 	// Is refraction enabled?
 #if !SIMPLIFIED_MATERIAL_SHADER && MATERIAL_SHADINGMODEL_DEFAULT_LIT && REFRACTION_USE_INDEX_OF_REFRACTION
 	float Ior = GetMaterialRefraction(PixelMaterialInputs).x;
