GPULM: Fix opacity with default lit materials

diff --git a/Engine/Shaders/Private/PathTracing/PathTracingMaterialHitShader.usf b/Engine/Shaders/Private/PathTracing/PathTracingMaterialHitShader.usf
--- a/Engine/Shaders/Private/PathTracing/PathTracingMaterialHitShader.usf
+++ b/Engine/Shaders/Private/PathTracing/PathTracingMaterialHitShader.usf
@@ -360,11 +360,11 @@ RAY_TRACING_ENTRY_CLOSEST_HIT(PathTracingMaterialCHS,
 	Payload.Specular = 0;
 	Payload.Roughness = 1;
 	Payload.Metallic = 0;
-#if MATERIALBLENDING_TRANSLUCENT && MATERIAL_SHADINGMODEL_THIN_TRANSLUCENT
-	Payload.Opacity = GetMaterialOpacity(PixelMaterialInputs);
-#else
+#if MATERIALBLENDING_TRANSLUCENT && !MATERIAL_SHADINGMODEL_THIN_TRANSLUCENT
 	// Force opacity to 0 so it is ignored for path hits and only evaluated in AHS as a modulation to throughput
 	Payload.Opacity = 0;
+#else
+	Payload.Opacity = GetMaterialOpacity(PixelMaterialInputs);
 #endif
 	Payload.Ior = 1;
 
