Reduce amount of ghosting on separate translucency VFX on r.TSR.History.SeparateTranslucency=0

diff --git a/Engine/Shaders/Private/TemporalSuperResolution/TSRRejectShading.usf b/Engine/Shaders/Private/TemporalSuperResolution/TSRRejectShading.usf
--- a/Engine/Shaders/Private/TemporalSuperResolution/TSRRejectShading.usf
+++ b/Engine/Shaders/Private/TemporalSuperResolution/TSRRejectShading.usf
@@ -354,11 +354,13 @@ void MainCS(
 
 	#if CONFIG_COMPOSE_TRANSLUCENCY
 	{
+		tsr_simd_half4 FilteredTranslucencyInput = Convolve3x3HV<TAccumulatorAvg<tsr_simd_half4>, tsr_simd_half4>(OriginalTranslucencyInput, bIsValid, GroupThreadIndex);
+
 		UNROLL_N(SIMD_SIZE)
 		for (uint SimdIndex = 0; SimdIndex < SIMD_SIZE; SimdIndex++)
 		{
 			tsr_halfC Color = OriginalInput.GetElement(SimdIndex);
-			tsr_half4 Translucency = OriginalTranslucencyInput.GetElement(SimdIndex);
+			tsr_half4 Translucency = FilteredTranslucencyInput.GetElement(SimdIndex);
 
 			Color.rgb = Color.rgb * Translucency.a + Translucency.rgb;
 			#if CONFIG_SCENE_COLOR_APLHA
