Change strata tile packing to use intrinssics.

diff --git a/Engine/Shaders/Private/Strata/StrataTile.ush b/Engine/Shaders/Private/Strata/StrataTile.ush
--- a/Engine/Shaders/Private/Strata/StrataTile.ush
+++ b/Engine/Shaders/Private/Strata/StrataTile.ush
@@ -18,7 +18,11 @@ uint2 StrataUnpackTile(uint In)
 
 uint StrataPackTile(uint2 TileCoord)
 {
+#if COMPILER_SUPPORTS_PACK_INTRINSICS
+	return PackUInt2ToUInt(TileCoord.x, TileCoord.y);
+#else
 	return TileCoord.x | (TileCoord.y << 16); // assumes 16bit is enough to represent a tiled resolution up to 65,535 :)
+#endif
 }
 
 ////////////////////////////////////////////////////////////////////////////////////////////////////////////
