Polishes spatial anti-aliasing on total history rejection

diff --git a/Engine/Shaders/Private/TemporalSuperResolution/TSRRejectShading.usf b/Engine/Shaders/Private/TemporalSuperResolution/TSRRejectShading.usf
--- a/Engine/Shaders/Private/TemporalSuperResolution/TSRRejectShading.usf
+++ b/Engine/Shaders/Private/TemporalSuperResolution/TSRRejectShading.usf
@@ -351,6 +351,25 @@ void MainCS(
 			bIsValid.SetElement(SimdIndex, ParallaxRejectionMaskTexture[ClampedFetchPixelPos] > tsr_half(PARALLAX_REJECTION_MASK_THRESHOLD) ? tsr_half(1.0) : tsr_half(0.0));
 		}
 	}
+	
+	// Compute AA Luma
+	tsr_simd_half CenterLumaLDR;
+	{
+		const tsr_half SpatialAAExposure = tsr_half(0.5);
+		UNROLL_N(SIMD_SIZE)
+		for (uint SimdIndex = 0; SimdIndex < SIMD_SIZE; SimdIndex++)
+		{
+			tsr_half3 Color = OriginalInput.GetElement(SimdIndex).rgb;
+			
+			#if CONFIG_COMPOSE_TRANSLUCENCY
+				tsr_half4 Translucency = OriginalTranslucencyInput.GetElement(SimdIndex);
+				Color.rgb = Color.rgb * Translucency.a + Translucency.rgb;
+			#endif
+
+			tsr_half CenterLuma = dot(Color, tsr_half3(0.299f, 0.587f, 0.114f));
+			CenterLumaLDR.SetElement(SimdIndex, CenterLuma / (SpatialAAExposure + CenterLuma));
+		}
+	}
 
 	#if CONFIG_COMPOSE_TRANSLUCENCY
 	{
@@ -371,18 +390,6 @@ void MainCS(
 		}
 	}
 	#endif
-		
-	// Compute AA Luma
-	tsr_simd_half CenterLumaLDR;
-	{
-		const tsr_half SpatialAAExposure = tsr_half(0.5);
-		UNROLL_N(SIMD_SIZE)
-		for (uint SimdIndex = 0; SimdIndex < SIMD_SIZE; SimdIndex++)
-		{
-			tsr_half CenterLuma = dot(OriginalInput.GetElement(SimdIndex).rgb, tsr_half3(0.299f, 0.587f, 0.114f));
-			CenterLumaLDR.SetElement(SimdIndex, CenterLuma / (SpatialAAExposure + CenterLuma));
-		}
-	}
 	
 	// Convert input to LDR
 	OriginalInput = OriginalInput * rcp(OriginalInput + tsr_simd_halfC::Const(1.0));
diff --git a/Engine/Shaders/Private/TemporalSuperResolution/TSRUpdateHistory.usf b/Engine/Shaders/Private/TemporalSuperResolution/TSRUpdateHistory.usf
--- a/Engine/Shaders/Private/TemporalSuperResolution/TSRUpdateHistory.usf
+++ b/Engine/Shaders/Private/TemporalSuperResolution/TSRUpdateHistory.usf
@@ -750,7 +750,7 @@ void MainCS(
 
 			bool2 bIsRefining = PrevHistoryValidity != 0.0;
 
-			tsr_half2 TotalRejection = select(or(or(bIsOffScreen, bIsDisoccluded), !bIsRefining), tsr_half(0.0).xx, saturate(LowFrequencyRejection * tsr_half(4.0)));
+			tsr_half2 TotalRejection = select(or(or(bIsOffScreen, bIsDisoccluded), !bIsRefining), tsr_half(0.0).xx, saturate(LowFrequencyRejection * tsr_half(16.0) - tsr_half(3.0)));
 
 			tsr_half2 KernelInputToHistoryFactor = lerp(OffScreenInputToHistoryFactor, tsr_half(InputToHistoryFactor), TotalRejection);
 			tsr_half2 KernelInputToHistoryAlignmentFactor = lerp(tsr_half(1.0).xx, tsr_half(InputToHistoryFactor), TotalRejection);
