Reorder a few operation to gain cycles.

diff --git a/Engine/Shaders/Private/Strata/StrataEvaluation.ush b/Engine/Shaders/Private/Strata/StrataEvaluation.ush
--- a/Engine/Shaders/Private/Strata/StrataEvaluation.ush
+++ b/Engine/Shaders/Private/Strata/StrataEvaluation.ush
@@ -291,7 +291,7 @@ FStrataEvaluateResult StrataEvaluateBSDFCommon(FStrataBSDFContext BSDFContext, F
 			{
 				Sample.DiffusePathValue = Diffuse_Lambert(DiffuseColor);
 			}
-			Sample.IntegratedDiffuseValue += ShadowTerms.SurfaceShadow * AreaLightContext.NoL * Sample.DiffusePathValue * AreaLightContext.Falloff * AreaLightContext.AreaLight.FalloffColor;
+			Sample.IntegratedDiffuseValue += (ShadowTerms.SurfaceShadow * AreaLightContext.NoL * AreaLightContext.Falloff) * Sample.DiffusePathValue * AreaLightContext.AreaLight.FalloffColor;
 			Sample.DiffuseColor		= DiffuseColor;
 			Sample.DiffusePDF		= BSDFContext.SatNoL / PI;
 			Sample.bSubsurface		= BSDF_GETSSSTYPE(BSDFContext.BSDF) == SSS_TYPE_DIFFUSION || BSDF_GETSSSTYPE(BSDFContext.BSDF) == SSS_TYPE_DIFFUSION_PROFILE;
@@ -474,10 +474,10 @@ FStrataEvaluateResult StrataEvaluateBSDFCommon(FStrataBSDFContext BSDFContext, F
 			}
 			else
 			{
-				Sample.SpecularPathValue     = D     * Vis     * FresnelTerm     * MSScale;
-				Sample.SpecularHazePathValue = HazeD * HazeVis * HazeFresnelTerm * HazeMSScale;
+				Sample.SpecularPathValue     = (D     * Vis)     * FresnelTerm     * MSScale;
+				Sample.SpecularHazePathValue = (HazeD * HazeVis) * HazeFresnelTerm * HazeMSScale;
 
-				CommonTerm = ShadowTerms.SurfaceShadow * AreaLightContext.NoL * AreaLightContext.AreaLight.FalloffColor * AreaLightContext.Falloff;
+				CommonTerm = (ShadowTerms.SurfaceShadow * AreaLightContext.NoL * AreaLightContext.Falloff) * AreaLightContext.AreaLight.FalloffColor;
 			}
 
 		#if STRATA_FASTPATH==0
@@ -547,10 +547,10 @@ FStrataEvaluateResult StrataEvaluateBSDFCommon(FStrataBSDFContext BSDFContext, F
 				float ClothD = D_Charlie(SafeRoughness, BSDFContext.Context.NoH);
 				float ClothVis = Vis_Ashikhmin(BSDFContext.Context.NoV, AreaLightContext.NoL);
 				float3 ClothF  = F_Schlick(FuzzF0, BSDFContext.Context.VoH);
-				float3 ClothSpecularPathValue = ClothD * ClothVis * ClothF * ComputeEnergyConservation(EnergyTerms);
+				float3 ClothSpecularPathValue = (ClothD * ClothVis) * ClothF * ComputeEnergyConservation(EnergyTerms);
 
 				// Area light are not supported by the cloth BRDF
-				float3 ClothIntegratedSpecularValue = ShadowTerms.SurfaceShadow * AreaLightContext.NoL * AreaLightContext.AreaLight.FalloffColor * AreaLightContext.Falloff * ClothSpecularPathValue * FuzzAmount;
+				float3 ClothIntegratedSpecularValue = (ShadowTerms.SurfaceShadow * AreaLightContext.NoL * AreaLightContext.Falloff * FuzzAmount) * AreaLightContext.AreaLight.FalloffColor * ClothSpecularPathValue;
 				
 				// Apply specular transmittance to diffuse and specular lob from slab medium and interface sitting below the layer of fuzz.
 				Sample.DiffusePathValue			*= Cloth_DirectionalAlbedo_SpecularTransmission;
@@ -628,7 +628,7 @@ FStrataEvaluateResult StrataEvaluateBSDFCommon(FStrataBSDFContext BSDFContext, F
 
 				Sample.TransmissionPDF = 1.0f / (4.0f * PI); // STRATA_TODO this currently match the uniform sphere sampling from StrataImportanceSampleBSDF
 				Sample.TransmissionPathValue = TransmissionThroughput * DirectionalAlbedo_SpecularTransmission;
-				Sample.IntegratedDiffuseValue += ShadowTerms.TransmissionShadow * Sample.TransmissionPathValue * AreaLightContext.AreaLight.FalloffColor * AreaLightContext.Falloff;
+				Sample.IntegratedDiffuseValue += (ShadowTerms.TransmissionShadow * AreaLightContext.Falloff) * AreaLightContext.AreaLight.FalloffColor * Sample.TransmissionPathValue;
 			}
 			#endif
 
@@ -669,7 +669,7 @@ FStrataEvaluateResult StrataEvaluateBSDFCommon(FStrataBSDFContext BSDFContext, F
 				Sample.TransmissionPathValue	= TransmissionThroughput * PhaseFunction;
 				Sample.TransmissionPDF			= 1.0f / (4.0f * PI);		// STRATA_TODO this currently match the uniform sphere sampling from StrataImportanceSampleBSDF
 
-				Sample.IntegratedDiffuseValue += ShadowTerms.TransmissionShadow * Sample.TransmissionPathValue * AreaLightContext.Falloff * AreaLightContext.AreaLight.FalloffColor;
+				Sample.IntegratedDiffuseValue += (ShadowTerms.TransmissionShadow * AreaLightContext.Falloff) * AreaLightContext.AreaLight.FalloffColor * Sample.TransmissionPathValue;
 			}
 			#endif
 
@@ -708,7 +708,7 @@ FStrataEvaluateResult StrataEvaluateBSDFCommon(FStrataBSDFContext BSDFContext, F
 
 				Sample.IntegratedDiffuseValue = lerp(
 					Sample.IntegratedDiffuseValue, 
-					ShadowTerms.SurfaceShadow * AreaLightContext.NoL * SlabDirectionalAlbedo * DirectionalAlbedo_SpecularTransmission * AreaLightContext.Falloff * AreaLightContext.AreaLight.FalloffColor,
+					(ShadowTerms.SurfaceShadow * AreaLightContext.NoL * AreaLightContext.Falloff) * AreaLightContext.AreaLight.FalloffColor * SlabDirectionalAlbedo * DirectionalAlbedo_SpecularTransmission,
 					DiffuseToVolumeBlend);
 				Sample.DiffuseColor		= SlabDirectionalAlbedo;
 			}
@@ -743,7 +743,7 @@ FStrataEvaluateResult StrataEvaluateBSDFCommon(FStrataBSDFContext BSDFContext, F
 			Sample.SpecularPDF			= 1.0f / (4.0f * PI);		// STRATA_TODO this currently match the uniform sphere sampling from StrataImportanceSampleBSDF
 			Sample.ThroughputV			= OpaqueBSDFThroughput; 
 			Sample.TransmittanceAlongN	= OpaqueBSDFThroughput;
-			Sample.IntegratedSpecularValue = ShadowTerms.TransmissionShadow * AreaLightContext.AreaLight.FalloffColor * AreaLightContext.Falloff * Sample.SpecularPathValue;
+			Sample.IntegratedSpecularValue = (ShadowTerms.TransmissionShadow * AreaLightContext.Falloff) * AreaLightContext.AreaLight.FalloffColor * Sample.SpecularPathValue;
 		}
 		break;
 		
@@ -787,7 +787,7 @@ FStrataEvaluateResult StrataEvaluateBSDFCommon(FStrataBSDFContext BSDFContext, F
 			const float Sclera = BSDFContext.Context.NoL;
 
 			Sample.DiffusePathValue = Diffuse_Lambert(DiffuseColor) * lerp(Sclera, Iris, EYE_IRISMASK(BSDFContext.BSDF));
-			Sample.IntegratedDiffuseValue += ShadowTerms.SurfaceShadow * AreaLightContext.NoL * Sample.DiffusePathValue * AreaLightContext.Falloff * AreaLightContext.AreaLight.FalloffColor;
+			Sample.IntegratedDiffuseValue += (ShadowTerms.SurfaceShadow * AreaLightContext.NoL * AreaLightContext.Falloff) * AreaLightContext.AreaLight.FalloffColor * Sample.DiffusePathValue;
 
 			Sample.DiffuseColor		= DiffuseColor;
 			Sample.DiffusePDF		= BSDFContext.SatNoL / PI;
@@ -859,12 +859,12 @@ FStrataEvaluateResult StrataEvaluateBSDFCommon(FStrataBSDFContext BSDFContext, F
 			if (bIsRectLight)
 			{
 				Sample.SpecularPathValue       = RectLightSpec * MSScale * SpecularMicroOcclusion;
-				Sample.IntegratedSpecularValue = ShadowTerms.SurfaceShadow * /* AreaLightContext.NoL and falloff is part of the LTC integration already */ Sample.SpecularPathValue * Sample.SpecularPathProbability;
+				Sample.IntegratedSpecularValue = (ShadowTerms.SurfaceShadow * /* AreaLightContext.NoL and falloff is part of the LTC integration already */ Sample.SpecularPathProbability) * Sample.SpecularPathValue;
 			}
 			else
 			{
-				Sample.SpecularPathValue        = D * Vis * FresnelTerm * MSScale;
-				Sample.IntegratedSpecularValue += ShadowTerms.SurfaceShadow * AreaLightContext.NoL * AreaLightContext.AreaLight.FalloffColor * AreaLightContext.Falloff * Sample.SpecularPathValue * Sample.SpecularPathProbability;
+				Sample.SpecularPathValue        = (D * Vis) * FresnelTerm * MSScale;
+				Sample.IntegratedSpecularValue += (ShadowTerms.SurfaceShadow * AreaLightContext.NoL * AreaLightContext.Falloff) * AreaLightContext.AreaLight.FalloffColor * Sample.SpecularPathValue * Sample.SpecularPathProbability;
 			}
 
 			////
