Enforce anisotropy is compiled out apart from complex tile.

diff --git a/Engine/Shaders/Private/Strata/StrataEvaluation.ush b/Engine/Shaders/Private/Strata/StrataEvaluation.ush
--- a/Engine/Shaders/Private/Strata/StrataEvaluation.ush
+++ b/Engine/Shaders/Private/Strata/StrataEvaluation.ush
@@ -112,11 +112,13 @@ FStrataBSDFContext StrataCreateBSDFContext(float3x3 TangentBasis, FStrataBSDF BS
 	BSDFContext.B = normalize(BSDFContext.R + BSDFContext.L);
 
 	BSDFContext.Context = (BxDFContext)0;
+#if STRATA_COMPLEXPATH
 	if (BSDF_GETHASANISOTROPY(BSDF) != 0)
 	{
 		Init(BSDFContext.Context, BSDFContext.N, BSDFContext.X, BSDFContext.Y, BSDFContext.V, BSDFContext.L);
 	}
 	else
+#endif
 	{
 		Init(BSDFContext.Context, BSDFContext.N, BSDFContext.V, BSDFContext.L);
 	}
@@ -254,6 +256,7 @@ FStrataEvaluateResult StrataEvaluateBSDFCommon(FStrataBSDFContext BSDFContext, F
 			float Alpha2 = Pow4(SafeRoughness);
 
 			float NoV, VoH, NoH;
+		#if STRATA_COMPLEXPATH
 			BRANCH
 			if (bHasAnisotropy)
 			{
@@ -264,6 +267,7 @@ FStrataEvaluateResult StrataEvaluateBSDFCommon(FStrataBSDFContext BSDFContext, F
 				NoH = BSDFContext.Context.NoH;
 			}
 			else
+		#endif
 			{
 				Init(BSDFContext.Context, BSDFContext.N, BSDFContext.V, AreaLightContext.L);
 
@@ -359,7 +363,7 @@ FStrataEvaluateResult StrataEvaluateBSDFCommon(FStrataBSDFContext BSDFContext, F
 			float HazePDF = 0;
 			const bool bIsRectLight = IntegrationType == INTEGRATION_AREA_LIGHT_RECT;
 
-			#if STRATA_FASTPATH==0
+			#if STRATA_COMPLEXPATH
 			BRANCH
 			if (bHasAnisotropy)
 			{
@@ -1279,7 +1283,7 @@ FStrataEnvLightResult StrataEvaluateForEnvLight(FStrataBSDFContext BSDFContext,
 			{
 				const float SpecularMicroOcclusion = F0RGBToMicroOcclusion(F0);
 
-				#if STRATA_FASTPATH==0
+				#if STRATA_COMPLEXPATH
 				const bool bHasAnisotropy = BSDF_GETHASANISOTROPY(BSDFContext.BSDF);
 				if (bHasAnisotropy)
 				{
