* Add shading quality for Strata.
* Chang Xbox/PS4 to use approximate Vis function.

diff --git a/Engine/Shaders/Private/Strata/StrataEvaluation.ush b/Engine/Shaders/Private/Strata/StrataEvaluation.ush
--- a/Engine/Shaders/Private/Strata/StrataEvaluation.ush
+++ b/Engine/Shaders/Private/Strata/StrataEvaluation.ush
@@ -39,6 +39,13 @@
 #define STRATA_SSS_TRANSMISSION 0
 #endif
 
+// Strata STRATA_LIGHTING_QUALITY define the lighting accuarcy/approximation
+// 1: accurante lighting
+// 2: approximation lighting
+#ifndef STRATA_LIGHTING_QUALITY
+#define STRATA_LIGHTING_QUALITY 1
+#endif
+
 ///////////////////////////////////////////////////////////////////////////////
 // Strata Macros for reducing permutation boiler plate
 // 
@@ -186,6 +193,21 @@ float3 LuminanceWeight(in FStrataBSDFContext InContext, in FStrataBSDF InBSDF)
 	return LuminanceWeight(InContext.SatNoL, InBSDF);
 }
 
+float StrataVisSmithJoint(float a2, float NoV, float NoL)
+{
+#if STRATA_LIGHTING_QUALITY > 1
+	return Vis_SmithJointApprox(a2, NoV, NoL);
+#else
+	return Vis_SmithJoint(a2, NoV, NoL);
+#endif
+}
+
+float StrataVisSmithJointAniso(float ax, float ay, float NoV, float NoL, float XoV, float XoL, float YoV, float YoL)
+{
+	// STRATA_TODO: build approximation for the aniso version
+	return Vis_SmithJointAniso(ax, ay, NoV, NoL, XoV, XoL, YoV, YoL);
+}
+
 ///////////////////////////////////////////////////////////////////////////////////////////////////
 // BSDF evaluation (punctual/area)
 
@@ -374,7 +396,7 @@ FStrataEvaluateResult StrataEvaluateBSDFCommon(FStrataBSDFContext BSDFContext, F
 					GetAnisotropicRoughness(Alpha, SLAB_ANISOTROPY(BSDFContext.BSDF), AlphaXY.x, AlphaXY.y);
 
 					D   = D_GGXaniso(AlphaXY.x, AlphaXY.y, BSDFContext.Context.NoH, BSDFContext.Context.XoH, BSDFContext.Context.YoH);
-					Vis = Vis_SmithJointAniso(AlphaXY.x, AlphaXY.y, BSDFContext.Context.NoV, BSDFContext.SatNoL, BSDFContext.Context.XoV, BSDFContext.Context.XoL, BSDFContext.Context.YoV, BSDFContext.Context.YoL);
+					Vis = StrataVisSmithJointAniso(AlphaXY.x, AlphaXY.y, BSDFContext.Context.NoV, BSDFContext.SatNoL, BSDFContext.Context.XoV, BSDFContext.Context.XoL, BSDFContext.Context.YoV, BSDFContext.Context.YoL);
 					const float H_PDF = VisibleGGXPDF_aniso(BSDFContext.TangentV, BSDFContext.TangentH, AlphaXY);
 					PDF = RayPDFToReflectionRayPDF(BSDFContext.Context.VoH, H_PDF);
 				}
@@ -385,7 +407,7 @@ FStrataEvaluateResult StrataEvaluateBSDFCommon(FStrataBSDFContext BSDFContext, F
 					GetAnisotropicRoughness(HazeSafeRoughness, SLAB_ANISOTROPY(BSDFContext.BSDF), HazeAlpha.x, HazeAlpha.y);
 
 					HazeD   = D_GGXaniso(HazeAlpha.x, HazeAlpha.y, BSDFContext.Context.NoH, BSDFContext.Context.XoH, BSDFContext.Context.YoH);
-					HazeVis = Vis_SmithJointAniso(HazeAlpha.x, HazeAlpha.y, BSDFContext.Context.NoV, BSDFContext.Context.NoL, BSDFContext.Context.XoV, BSDFContext.Context.XoL, BSDFContext.Context.YoV, BSDFContext.Context.YoL);
+					HazeVis = StrataVisSmithJointAniso(HazeAlpha.x, HazeAlpha.y, BSDFContext.Context.NoV, BSDFContext.Context.NoL, BSDFContext.Context.XoV, BSDFContext.Context.XoL, BSDFContext.Context.YoV, BSDFContext.Context.YoL);
 					const float H_PDF = VisibleGGXPDF_aniso(BSDFContext.TangentV, BSDFContext.TangentH, HazeAlpha);
 					HazePDF = RayPDFToReflectionRayPDF(BSDFContext.Context.VoH, H_PDF);
 				}
@@ -429,8 +451,7 @@ FStrataEvaluateResult StrataEvaluateBSDFCommon(FStrataBSDFContext BSDFContext, F
 							const float Energy = EnergyNormalization(Alpha2Spec, BSDFContext.Context.VoH, AreaLightContext.AreaLight);
 							D = D_GGX(Alpha2Spec, BSDFContext.Context.NoH) * Energy;
 						}
-						// STRATA_TODO Should we use Vis_SmithJointApprox here?
-						Vis = Vis_SmithJoint(Alpha2Spec, BSDFContext.Context.NoV, AreaLightContext.NoL);
+						Vis = StrataVisSmithJoint(Alpha2Spec, BSDFContext.Context.NoV, AreaLightContext.NoL);
 						const float H_PDF = VisibleGGXPDF(BSDFContext.TangentV, BSDFContext.TangentH, Alpha2Spec);
 						PDF = RayPDFToReflectionRayPDF(BSDFContext.Context.VoH, H_PDF);
 					}
@@ -448,8 +469,7 @@ FStrataEvaluateResult StrataEvaluateBSDFCommon(FStrataBSDFContext BSDFContext, F
 							const float Energy = EnergyNormalization(Alpha2SpecHaze, BSDFContext.Context.VoH, AreaLightContext.AreaLight);
 							HazeD = D_GGX(Alpha2SpecHaze, BSDFContext.Context.NoH) * Energy;
 						}
-						// STRATA_TODO Should we use Vis_SmithJointApprox here?
-						HazeVis = Vis_SmithJoint(Alpha2SpecHaze, BSDFContext.Context.NoV, AreaLightContext.NoL);
+						HazeVis = StrataVisSmithJoint(Alpha2SpecHaze, BSDFContext.Context.NoV, AreaLightContext.NoL);
 						const float H_PDF = VisibleGGXPDF(BSDFContext.TangentV, BSDFContext.TangentH, Alpha2SpecHaze);
 						HazePDF = RayPDFToReflectionRayPDF(BSDFContext.Context.VoH, H_PDF);
 					}
@@ -849,7 +869,7 @@ FStrataEvaluateResult StrataEvaluateBSDFCommon(FStrataBSDFContext BSDFContext, F
 					D = D_GGX(Alpha2Spec, BSDFContext.Context.NoH) * Energy;
 				}
 				// STRATA_TODO Should we use Vis_SmithJointApprox here?
-				Vis = Vis_SmithJoint(Alpha2Spec, BSDFContext.Context.NoV, AreaLightContext.NoL);
+				Vis = StrataVisSmithJoint(Alpha2Spec, BSDFContext.Context.NoV, AreaLightContext.NoL);
 				const float H_PDF = VisibleGGXPDF(BSDFContext.TangentV, BSDFContext.TangentH, Alpha2Spec);
 				PDF = RayPDFToReflectionRayPDF(BSDFContext.Context.VoH, H_PDF);
 			}
diff --git a/Engine/Source/Runtime/Engine/Private/ShaderCompiler/ShaderCompiler.cpp b/Engine/Source/Runtime/Engine/Private/ShaderCompiler/ShaderCompiler.cpp
--- a/Engine/Source/Runtime/Engine/Private/ShaderCompiler/ShaderCompiler.cpp
+++ b/Engine/Source/Runtime/Engine/Private/ShaderCompiler/ShaderCompiler.cpp
@@ -6015,6 +6015,14 @@ void GlobalBeginCompileShader(
 			Input.Environment.SetDefine(TEXT("STRATA_DIFFUSE_CHAN"), bStrataRoughDiffuse ? 1 : 0);
 		}
 
+		// Force rough diffuse to be disable on platform which explicitly disable it from their settings
+		if (bStrata)
+		{
+			static FShaderPlatformCachedIniValue<int32> CVarStrataShadingQuality(TEXT("r.Strata.ShadingQuality"));
+			const uint32 StrataShadingQuality = CVarStrataShadingQuality.Get(Target.GetPlatform());
+			Input.Environment.SetDefine(TEXT("STRATA_SHADING_QUALITY"), FMath::Max(StrataShadingQuality, 1u));
+		}
+
 		static IConsoleVariable* CVarBackCompatibility = IConsoleManager::Get().FindConsoleVariable(TEXT("r.StrataBackCompatibility"));
 		const bool bStrataBackCompatibility = CVarBackCompatibility && CVarBackCompatibility->GetInt() > 0;
 		Input.Environment.SetDefine(TEXT("PROJECT_STRATA_BACKCOMPATIBILITY"), bStrataBackCompatibility ? 1 : 0);
diff --git a/Engine/Source/Runtime/Renderer/Private/Strata/Strata.cpp b/Engine/Source/Runtime/Renderer/Private/Strata/Strata.cpp
--- a/Engine/Source/Runtime/Renderer/Private/Strata/Strata.cpp
+++ b/Engine/Source/Runtime/Renderer/Private/Strata/Strata.cpp
@@ -34,6 +34,12 @@ static TAutoConsoleVariable<int32> CVarStrataBytePerPixel(
 	TEXT("Strata allocated byte per pixel to store materials data. Higher value means more complex material can be represented."),
 	ECVF_ReadOnly | ECVF_RenderThreadSafe);
 
+static TAutoConsoleVariable<int32> CVarStrataShadingQuality(
+	TEXT("r.Strata.ShadingQuality"),
+	1,
+	TEXT("Defube Strata shading quality (1: accurate lighting, 2: approximate lighting). This variable is read-only."),
+	ECVF_ReadOnly | ECVF_RenderThreadSafe);
+
 static TAutoConsoleVariable<int32> CVarStrataRoughDiffuse(
 	TEXT("r.Strata.RoughDiffuse"),
 	1,
