Use fixed stochastic alpha value for writing VT feedback when doing ThinTranslucency.

diff --git a/Engine/Shaders/Private/VirtualTextureCommon.ush b/Engine/Shaders/Private/VirtualTextureCommon.ush
--- a/Engine/Shaders/Private/VirtualTextureCommon.ush
+++ b/Engine/Shaders/Private/VirtualTextureCommon.ush
@@ -68,7 +68,7 @@ void FinalizeVirtualTextureFeedback(in FVirtualTextureFeedbackParams Params, flo
 #if MATERIALBLENDING_TRANSLUCENT || MATERIALBLENDING_ADDITIVE || MATERIALBLENDING_MODULATE || IS_DECAL || (USE_DBUFFER && MATERIAL_DECAL_READ_MASK)
 
 		// Stochastic alpha testing
-#if MATERIALBLENDING_ADDITIVE || MATERIALBLENDING_MODULATE
+#if MATERIALBLENDING_ADDITIVE || MATERIALBLENDING_MODULATE || MATERIAL_SHADINGMODEL_THIN_TRANSLUCENT
 		// Opacity may not make sense for additive and modulation blending
 		// Use a value of 0.25f so we have up to 4 layers of VT transparancy
 		const float Alpha = 0.25f;  
