Fix for type mismatch generating shader compilation failure.

The constant used to default init FRayTracingPickingFeedback is up-casted to 64 bits (when the expected value is expected to be 32-bit)

#rnx

diff --git a/Engine/Shaders/Private/RayTracing/RayTracingDebug.usf b/Engine/Shaders/Private/RayTracing/RayTracingDebug.usf
--- a/Engine/Shaders/Private/RayTracing/RayTracingDebug.usf
+++ b/Engine/Shaders/Private/RayTracing/RayTracingDebug.usf
@@ -341,7 +341,7 @@ RAY_TRACING_ENTRY_RAYGEN(RayTracingDebugPickingRGS)
 		Ray,
 		Payload);
 
-	FRayTracingPickingFeedback PickingFeedback = (FRayTracingPickingFeedback)0xFFFFFFFF;
+	FRayTracingPickingFeedback PickingFeedback = (FRayTracingPickingFeedback)0xFFFFFFFFu;
 
 	if (Payload.IsHit())
 	{
