Strata - disable a ray tracing shader that is not realistic to support today.

diff --git a/Engine/Shaders/Private/RayTracing/RayTracingLightingMS.usf b/Engine/Shaders/Private/RayTracing/RayTracingLightingMS.usf
--- a/Engine/Shaders/Private/RayTracing/RayTracingLightingMS.usf
+++ b/Engine/Shaders/Private/RayTracing/RayTracingLightingMS.usf
@@ -101,7 +101,8 @@ RAY_TRACING_ENTRY_MISS(RayTracingLightingMS, FPackedMaterialClosestHitPayload, P
 		float4 FakePrecomputedShadowFactors = 1.f;
 		GetShadowTerms(FakeSceneDepth, FakePrecomputedShadowFactors, FakeShadingModelID, FakeContactShadowOpacity, LightData, WorldPosition, L, LightAttenuation, Dither, ShadowTerms);
 
-		FStrataDeferredLighting Lighting = StrataDeferredLighting(
+		// RayTracing + Lighting + Strata is too expenssive (as in uses too many registers) to be supported on some platofrms
+		/*FStrataDeferredLighting Lighting = StrataDeferredLighting(
 			LightData,
 			V,
 			L,
@@ -111,7 +112,8 @@ RAY_TRACING_ENTRY_MISS(RayTracingLightingMS, FPackedMaterialClosestHitPayload, P
 			PackedPayload.StrataData,
 			StrataAddressing,
 			StrataPixelHeader);
-		LightContribution = Lighting.SceneColor.xyz;
+		LightContribution = Lighting.SceneColor.xyz;*/
+		LightContribution = float4(1, 0, 1, 0.5); // Make it return purple so that it is noticeable.
 	}
 #else
 	// LWC_TODO - lighting wants translated world space
