Fix for shader compilation error on SM5 caused by HairStrandsVisibilityRasterClassification.usf using wave ops.

#rnx

diff --git a/Engine/Shaders/Private/HairStrands/HairStrandsVisibilityRasterClassification.usf b/Engine/Shaders/Private/HairStrands/HairStrandsVisibilityRasterClassification.usf
--- a/Engine/Shaders/Private/HairStrands/HairStrandsVisibilityRasterClassification.usf
+++ b/Engine/Shaders/Private/HairStrands/HairStrandsVisibilityRasterClassification.usf
@@ -116,6 +116,7 @@ void CSMain(uint DispatchThreadID : SV_DispatchThreadID, uint GroupThreadID : SV
 			bShouldSoftwareRasterize = ShouldSoftwareRasterize(PositionOffset, PrimID, bSegValid);
 		}
 
+#if PLATFORM_SUPPORTS_SM6_0_WAVE_OPERATIONS
 		const uint NumSWSegs = WaveActiveCountBits(bSegValid && bShouldSoftwareRasterize);
 		const uint NumHWSegs = WaveActiveCountBits(bSegValid && !bShouldSoftwareRasterize);
 		
@@ -145,5 +146,23 @@ void CSMain(uint DispatchThreadID : SV_DispatchThreadID, uint GroupThreadID : SV
 			const uint WritePos = bShouldSoftwareRasterize ? (WriteOffsetSW + LocalWritePosSW) : (PrimIDsBufferSize - 1 - (WriteOffsetHW + NumHWSegs) + LocalWritePosHW);
 			OutPrimIDs[WritePos] = PrimID;
 		}
+#else
+		if (bSegValid)
+		{
+			uint WriteOffset;
+			if (bShouldSoftwareRasterize)
+			{
+
+				OutSWRasterPrimitiveCount.InterlockedAdd(0, 1, WriteOffset);
+			}
+			else
+			{
+				// Increment InstanceCount
+				OutDrawIndexedArgs.InterlockedAdd(1 << 2, 1, WriteOffset);
+			}
+			const uint WritePos = bShouldSoftwareRasterize ? WriteOffset : (PrimIDsBufferSize - 1 - WriteOffset);
+			OutPrimIDs[WritePos] = PrimID;
+		}
+#endif // PLATFORM_SUPPORTS_SM6_0_WAVE_OPERATIONS
 	}
 }
\ No newline at end of file
