Small numerical robustness improvements to path tracing related sampling routines
 * ConcentricDiskSamplingHelper: ensure the output radius is in [0,1) instead of [0,1] which avoids special cases when remapping the radius non-uniformly (such as for pixel filtering with a gaussian filter)
 * Add MISWeightBalanced to implement a numerically robust form of the balance heuristic (we had existing code for the power heuristic only)
 * Use MISWeightBalanced for lobe selection probabilities (avoids an arbitrary epsilon which adds a subtle skew to the pdfs)
 * Rearrange expression in HenyeyGreensteinPhaseInvertCDF to avoid singularity around G==0

diff --git a/Engine/Shaders/Private/MonteCarlo.ush b/Engine/Shaders/Private/MonteCarlo.ush
--- a/Engine/Shaders/Private/MonteCarlo.ush
+++ b/Engine/Shaders/Private/MonteCarlo.ush
@@ -113,7 +113,8 @@ float2 UniformSampleDisk( float2 E )
 // Returns a point on the unit circle and a radius in z
 float3 ConcentricDiskSamplingHelper(float2 E)
 {
-	float2 p = 2 * E - 1;
+	// Rescale input from [0,1) to (-1,1). This ensures the output radius is in [0,1)
+	float2 p = 2 * E - 0.99999994;
 	float2 a = abs(p);
 	float Lo = min(a.x, a.y);
 	float Hi = max(a.x, a.y);
@@ -452,6 +453,36 @@ float MISWeight( uint Num, float PDF, uint OtherNum, float OtherPDF )
 	return Weight * Weight / (Weight * Weight + OtherWeight * OtherWeight);
 }
 
+// Multiple importance sampling balance heuristic
+float MISWeightBalanced(float Pdf, float OtherPdf)
+{
+	// The straightforward implementation is prone to numerical overflow and divisions by 0
+	// and does not work well with +inf inputs.
+
+	if (Pdf == OtherPdf)
+	{
+		// Catch potential NaNs from (0,0) and (+inf, +inf)
+		return 0.5f;
+	}
+
+	// Evaluate the expression using the ratio of the smaller value to the bigger one for greater
+	// numerical stability. The math would also work using the ratio of bigger to smaller value,
+	// which would underflow less but would make the weights asymmetric. Underflow to 0 is not a
+	// bad property to have in rendering application as it ensures more weights are exactly 0
+	// which allows some evaluations to be skipped.
+	if (OtherPdf < Pdf)
+	{
+		float x = OtherPdf / Pdf;
+		return 1.0 / (1.0 + x);
+	}
+	else
+	{
+		// this form guarantees the weights add back up to one when arguments are swapped
+		float x = Pdf / OtherPdf;
+		return 1.0 - 1.0 / (1.0 + x);
+	}
+}
+
 // Multiple importance sampling power heuristic of two functions with a power of two. 
 float MISWeightRobust(float Pdf, float OtherPdf) {
 	// The straightforward implementation above is prone to numerical overflow and divisions by 0
diff --git a/Engine/Shaders/Private/ParticipatingMediaCommon.ush b/Engine/Shaders/Private/ParticipatingMediaCommon.ush
--- a/Engine/Shaders/Private/ParticipatingMediaCommon.ush
+++ b/Engine/Shaders/Private/ParticipatingMediaCommon.ush
@@ -112,17 +112,15 @@ float SchlickPhase(float G, float CosTheta)
 // Follows PBRT convention http://www.pbr-book.org/3ed-2018/Light_Transport_II_Volume_Rendering/Sampling_Volume_Scattering.html#SamplingPhaseFunctions
 float HenyeyGreensteinPhaseInvertCDF(float E, float G)
 {
-	float R = 1.0 - 2.0 * E;
-	if (abs(G) < 0.001f)
-	{
-		// Careful: OneOver2G undefined for g close to 0
-		return R;
-	}
-	float OnePlusG2  = 1.0f + G * G;
-	float OneMinusG2 = 1.0f - G * G;
-	float OneOver2G  = 0.5f / G;
-	float t = OneMinusG2 / (1.0f + G * R);
-	return OneOver2G * (t * t - OnePlusG2);
+	// The naive expression requires a special case for G==0, rearranging terms a bit
+	// gives the following result which does not require any special cases
+	float t0 = (1.0 - G) + G * E;
+	float t1 = (1.0 - E) + E * E;
+	float t2 = t1 + (G * E) * t0;
+	float t3 = (2.0 * E - 1.0) - G * G;
+	float Num = t3 + (2.0 * G) * t2;
+	float Den = t0 + G * E;
+	return Num / (Den * Den);
 }
 
 // Follows PBRT convention http://www.pbr-book.org/3ed-2018/Light_Transport_II_Volume_Rendering/Sampling_Volume_Scattering.html#SamplingPhaseFunctions
diff --git a/Engine/Shaders/Private/PathTracing/Material/PathTracingMaterialCommon.ush b/Engine/Shaders/Private/PathTracing/Material/PathTracingMaterialCommon.ush
--- a/Engine/Shaders/Private/PathTracing/Material/PathTracingMaterialCommon.ush
+++ b/Engine/Shaders/Private/PathTracing/Material/PathTracingMaterialCommon.ush
@@ -15,7 +15,8 @@ float LobeSelectionProb(float3 A, float3 B)
 {
 	const float SumA = A.x + A.y + A.z;
 	const float SumB = B.x + B.y + B.z;
-	return SumA / (SumA + SumB + 1e-6);
+	// use MIS routine to robustly compute SumA / (SumA + SumB) even when both are 0, or overflow, etc ...
+	return MISWeightBalanced(SumA, SumB);
 }
 
 // Takes as input the sample weight and pdf for a certain lobe of a mixed model, together with the probability of picking that lobe
@@ -31,14 +32,7 @@ void AddLobeWithMIS(inout float3 Weight, inout float Pdf, float3 LobeWeight, flo
 		LobePdf *= LobeProb;
 		LobeWeight *= 1 / LobeProb;
 
-		// See discussion in MISWeightRobust for why this is better than LobePdf / (Pdf + LobePdf)
-		float MISWeight;
-		if (Pdf < LobePdf)
-			MISWeight = 1 / (1 + Pdf / LobePdf);
-		else if (LobePdf < Pdf)
-			MISWeight = 1 - 1 / (1 + LobePdf / Pdf);
-		else
-			MISWeight = 0.5f; // avoid (rare) inf/inf
+		float MISWeight = MISWeightBalanced(LobePdf, Pdf);
 
 		Weight = lerp(Weight, LobeWeight, MISWeight);
 		Pdf += LobePdf;
diff --git a/Engine/Shaders/Private/PathTracing/PathTracingCore.ush b/Engine/Shaders/Private/PathTracing/PathTracingCore.ush
--- a/Engine/Shaders/Private/PathTracing/PathTracingCore.ush
+++ b/Engine/Shaders/Private/PathTracing/PathTracingCore.ush
@@ -787,8 +787,7 @@ FPathState CreatePathState(int2 PixelIndex)
 		// importance sample a gaussian kernel with variable sigma
 		float3 Disk = ConcentricDiskSamplingHelper(AAJitter);
 		float Sigma = FilterWidth / 6.0; // user-provided width covers +/-3*Sigma
-		float R = min(Disk.z, 0.99999994); // avoid log(0) when R=1
-		AAJitter = 0.5 + Sigma * Disk.xy * sqrt(-2.0 * log(1.0 - R * R));
+		AAJitter = 0.5 + Sigma * Disk.xy * sqrt(-2.0 * log(1.0 - Disk.z * Disk.z));
 
 		float2 ViewportUV = (LaunchIndex + AAJitter) * View.BufferSizeAndInvSize.zw;
 		Output.Ray = CreatePrimaryRay(ViewportUV);
