Fixes spatial anti-aliasing application to separate translucency

diff --git a/Engine/Shaders/Private/TemporalSuperResolution/TSRUpdateHistory.usf b/Engine/Shaders/Private/TemporalSuperResolution/TSRUpdateHistory.usf
--- a/Engine/Shaders/Private/TemporalSuperResolution/TSRUpdateHistory.usf
+++ b/Engine/Shaders/Private/TemporalSuperResolution/TSRUpdateHistory.usf
@@ -348,6 +348,10 @@ void MainCS(
 	bool2 bIsResponsiveAAPixel = false;
 	bool2 bIsOffScreen = false;
 	bool2 bIsDisoccluded = false;
+	
+	#if CONFIG_REJECTION_ANTI_ALIASING
+		tsr_half2x2 SpatialAntiAliasingOffset = dpv_interleave_mono_registers(tsr_half(0.0).xx);
+	#endif
 
 	PLATFORM_SPECIFIC_ISOLATE
 	{
@@ -511,23 +515,12 @@ void MainCS(
 
 		#if CONFIG_REJECTION_ANTI_ALIASING
 		{
-			// Detect if HistoryBufferUV would be outside of the viewport.
-			tsr_half2 TotalRejection = select(or(bIsOffScreen, bIsDisoccluded), tsr_half(0.0).xx, saturate(LowFrequencyRejection * tsr_half(4.0)));
-
 			tsr_ushort2 EncodedInputTexelOffset = dpv_force_interleave_registers_array(RawEncodedInputTexelOffset);
 			
 			NoiseFiltering = dpv_force_interleave_registers_array(RawNoiseFiltering);
 			HoleFilledVelocityMask = dpv_interleave_mono_registers(tsr_half(1.0).xx) - dpv_force_interleave_registers_array(RawHoleFilledVelocityMask);
 
-			tsr_half2x2 TexelOffset = DecodeSpatialAntiAliasingOffset(EncodedInputTexelOffset);
-			
-			{
-				InputPPCo += dpv_scale(TexelOffset, tsr_half(1.0) - TotalRejection);
-				InputPPCk = floor(InputPPCo) + 0.5;
-				InputPixelPos = ClampPixelOffset(
-					tsr_short2x2(InputPPCo),
-					InputPixelPosMin, InputPixelPosMax);
-			}
+			SpatialAntiAliasingOffset = DecodeSpatialAntiAliasingOffset(EncodedInputTexelOffset);
 		}
 		#endif
 		
@@ -719,9 +712,6 @@ void MainCS(
 	
 	PLATFORM_SPECIFIC_ISOLATE
 	{
-		// Vector in pixel between pixel K -> O.
-		tsr_half2x2 dInputKO = tsr_half2x2(InputPPCo - InputPPCk);
-
 		tsr_halfC RawInputColorArray[CONFIG_SAMPLES][DPV_PIXEL_PER_LANE];
 		tsr_half4 RawTranslucencyColorArray[CONFIG_SAMPLES][DPV_PIXEL_PER_LANE];
 
@@ -731,9 +721,24 @@ void MainCS(
 		// Issues overlapped texture fetches
 		PLATFORM_SPECIFIC_ISOLATE
 		{
+			// Detect if HistoryBufferUV would be outside of the viewport.
+			tsr_half2 SpatialAntiAliasingLerp = select(or(bIsOffScreen, bIsDisoccluded), tsr_half(1.0).xx, saturate(tsr_half(1.0) - LowFrequencyRejection * tsr_half(4.0)));
+
 			// Pixel coordinate of the center of output pixel O in the input viewport.
 			float2x2 TranslucencyPPCo = ApplyScreenTransform(float2x2(HistoryPixelPos), HistoryPixelPosToTranslucencyPPCo);
+			
+			#if CONFIG_REJECTION_ANTI_ALIASING
+			{
+				TranslucencyPPCo += dpv_scale(SpatialAntiAliasingOffset, SpatialAntiAliasingLerp);
+				InputPPCo += dpv_scale(SpatialAntiAliasingOffset, SpatialAntiAliasingLerp);
 
+				InputPPCk = floor(InputPPCo) + 0.5;
+				InputPixelPos = ClampPixelOffset(
+					tsr_short2x2(InputPPCo),
+					InputPixelPosMin, InputPixelPosMax);
+			}
+			#endif
+			
 			// Pixel coordinate of the center of the nearest input pixel K in the input viewport.
 			float2x2 TranslucencyPPCk = floor(TranslucencyPPCo) + 0.5;
 	
@@ -743,6 +748,9 @@ void MainCS(
 				
 			// Vector in pixel between pixel K -> O.
 			tsr_half2x2 dTranslucencyKO = tsr_half2x2(TranslucencyPPCo - TranslucencyPPCk);
+			
+			// Vector in pixel between pixel K -> O.
+			tsr_half2x2 dInputKO = tsr_half2x2(InputPPCo - InputPPCk);
 
 
 			// Soften the spatial kernel when there is high noise in the input to make parrallax disocclusion of detailed geometry less distracting.
@@ -750,10 +758,9 @@ void MainCS(
 
 			bool2 bIsRefining = PrevHistoryValidity != 0.0;
 
-			tsr_half2 TotalRejection = select(or(or(bIsOffScreen, bIsDisoccluded), !bIsRefining), tsr_half(0.0).xx, saturate(LowFrequencyRejection * tsr_half(16.0) - tsr_half(3.0)));
-
-			tsr_half2 KernelInputToHistoryFactor = lerp(OffScreenInputToHistoryFactor, tsr_half(InputToHistoryFactor), TotalRejection);
-			tsr_half2 KernelInputToHistoryAlignmentFactor = lerp(tsr_half(1.0).xx, tsr_half(InputToHistoryFactor), TotalRejection);
+			tsr_half2 KernelInputToHistoryLerp = select(or(or(bIsOffScreen, bIsDisoccluded), !bIsRefining), tsr_half(0.0).xx, saturate(LowFrequencyRejection * tsr_half(16.0) - tsr_half(3.0)));
+			tsr_half2 KernelInputToHistoryFactor = lerp(OffScreenInputToHistoryFactor, tsr_half(InputToHistoryFactor), KernelInputToHistoryLerp);
+			tsr_half2 KernelInputToHistoryAlignmentFactor = lerp(tsr_half(1.0).xx, tsr_half(InputToHistoryFactor), KernelInputToHistoryLerp);
 			
 			InputPixelAlignement = ComputeSampleWeigth(KernelInputToHistoryAlignmentFactor, dInputKO, /* MinimalContribution = */ float(0.0));
 		
