Fix OIT output with separate translucency.

diff --git a/Engine/Shaders/Private/OITCombine.usf b/Engine/Shaders/Private/OITCombine.usf
--- a/Engine/Shaders/Private/OITCombine.usf
+++ b/Engine/Shaders/Private/OITCombine.usf
@@ -87,7 +87,8 @@ void MainCS(uint3 DispatchThreadId : SV_DispatchThreadID)
 		Transmittance = saturate(Transmittance * Samples[CompIt].Trans);
 	}
 
-	float AdjustedAlpha = saturate(1 - dot(Transmittance, float3(1.0f, 1.0f, 1.0f) / 3.0f));
+	// Alpha value used for separate translucency pass
+	float AdjustedAlpha = saturate(dot(Transmittance, float3(1.0f, 1.0f, 1.0f) / 3.0f));
 	
 	// Write final color - typed load
 	if (SampleCount > 0)
@@ -213,5 +214,6 @@ void MainCS(uint3 DispatchThreadId : SV_DispatchThreadID)
 			Newline(Context);
 		}
 	}
+#endif
 }
 #endif
\ No newline at end of file
