Fix shader typo.

diff --git a/Engine/Shaders/Private/OITCombine.usf b/Engine/Shaders/Private/OITCombine.usf
--- a/Engine/Shaders/Private/OITCombine.usf
+++ b/Engine/Shaders/Private/OITCombine.usf
@@ -214,6 +214,5 @@ void MainCS(uint3 DispatchThreadId : SV_DispatchThreadID)
 			Newline(Context);
 		}
 	}
-#endif
 }
 #endif
\ No newline at end of file
