Strata - reverted fastrcp optimisation on Oct encoding causing visual issues with normal.

diff --git a/Engine/Shaders/Private/Strata/Strata.ush b/Engine/Shaders/Private/Strata/Strata.ush
--- a/Engine/Shaders/Private/Strata/Strata.ush
+++ b/Engine/Shaders/Private/Strata/Strata.ush
@@ -3074,19 +3074,9 @@ float2 SignNotZero(in float2 v)
 	return float2(SignNotZero(v.x), SignNotZero(v.y));
 }
 
-float2 StrataUnitVectorToOctahedron(float3 N)
-{
-	N.xy *= rcpFast(dot(1, abs(N)));
-	if (N.z <= 0)
-	{
-		N.xy = (1 - abs(N.yx)) * select(N.xy >= 0, float2(1, 1), float2(-1, -1));
-	}
-	return N.xy;
-}
-
 uint StrataPackNormal(in float3 Normal)
 {
-	return PackRG16(StrataUnitVectorToOctahedron(Normal));
+	return PackRG16(UnitVectorToOctahedron(Normal));
 }
 
 float3 StrataUnpackNormal(uint PackedNormal)
@@ -3096,7 +3086,7 @@ float3 StrataUnpackNormal(uint PackedNormal)
 
 uint StrataPackNormal24(in float3 Normal)
 {
-	const float2 Result = StrataUnitVectorToOctahedron(Normal);
+	const float2 Result = UnitVectorToOctahedron(Normal);
 	const uint2 PackedXY = uint2(clamp(Result * 2047.0f + 2048.0f, 0.0f, 4095.0f));
 	return PackedXY.x | (PackedXY.y << 12);
 }
@@ -3110,7 +3100,7 @@ float3 StrataUnpackNormal24(uint PackedNormal)
 
 uint StrataPackNormal22(in float3 Normal)
 {
-	const float2 Result = StrataUnitVectorToOctahedron(Normal);
+	const float2 Result = UnitVectorToOctahedron(Normal);
 	const uint2 PackedXY = uint2(clamp(Result * 1023.0f + 1024.0f, 0.0f, 2047.0f));
 	return PackedXY.x | (PackedXY.y << 11);
 }
