Decrease the number of unique materials for hair in favor of supporting larger/hero groom.

This is a short term fix. Primitive ID should instead reference their material throught cluster information.

diff --git a/Engine/Shaders/Private/HairStrands/HairStrandsVisibilityCommon.ush b/Engine/Shaders/Private/HairStrands/HairStrandsVisibilityCommon.ush
--- a/Engine/Shaders/Private/HairStrands/HairStrandsVisibilityCommon.ush
+++ b/Engine/Shaders/Private/HairStrands/HairStrandsVisibilityCommon.ush
@@ -151,17 +151,17 @@ uint PackHairVisDepthCoverage(float InDepth, uint InCoverage8bits)
 
 uint PackHairVisPrimitiveMaterialId(uint PrimitiveId, uint MaterialId)
 {
-	return (PrimitiveId & 0x007FFFFF) | ((MaterialId & 0x1FF) << 23);
+	return (PrimitiveId & 0x00FFFFFF) | ((MaterialId & 0xFF) << 24);
 }
 
 uint UnpackHairVisPrimitiveId(uint RawVis)
 {
-	return RawVis & 0x007FFFFF;
+	return RawVis & 0x00FFFFFF;
 }
 
 uint UnpackHairVisMaterialId(uint RawVis)
 {
-	return (RawVis >> 23) & 0x1FF;
+	return (RawVis >> 24) & 0xFF;
 }
 
 float UnpackHairVisDepth(uint InPacked)
