Fix typo to use approx. vis function on lower-end platform.

diff --git a/Engine/Shaders/Private/Strata/StrataEvaluation.ush b/Engine/Shaders/Private/Strata/StrataEvaluation.ush
--- a/Engine/Shaders/Private/Strata/StrataEvaluation.ush
+++ b/Engine/Shaders/Private/Strata/StrataEvaluation.ush
@@ -195,7 +195,7 @@ float3 LuminanceWeight(in FStrataBSDFContext InContext, in FStrataBSDF InBSDF)
 
 float StrataVisSmithJoint(float a2, float NoV, float NoL)
 {
-#if STRATA_LIGHTING_QUALITY > 1
+#if STRATA_SHADING_QUALITY > 1
 	return Vis_SmithJointApprox(a2, NoV, NoL);
 #else
 	return Vis_SmithJoint(a2, NoV, NoL);
