Reorganise Strata lighting evaluation, by splitting haziness from main specular computation.

Save 2 vgprs.

diff --git a/Engine/Shaders/Private/Strata/StrataEvaluation.ush b/Engine/Shaders/Private/Strata/StrataEvaluation.ush
--- a/Engine/Shaders/Private/Strata/StrataEvaluation.ush
+++ b/Engine/Shaders/Private/Strata/StrataEvaluation.ush
@@ -263,9 +263,6 @@ FStrataEvaluateResult StrataEvaluateBSDFCommon(FStrataBSDFContext BSDFContext, F
 			float3 F0					= SLAB_F0(BSDFContext.BSDF);
 			float3 F90					= SLAB_F90(BSDFContext.BSDF);
 			const float SafeRoughness	= MakeRoughnessSafe(SLAB_ROUGHNESS(BSDFContext.BSDF));
-		#if STRATA_FASTPATH==0
-			float HazeSafeRoughness		= SafeRoughness;
-		#endif
 			const bool bHasAnisotropy	= BSDF_GETHASANISOTROPY(BSDFContext.BSDF);
 			const bool bHaziness		= BSDF_GETHASHAZINESS(BSDFContext.BSDF);
 
@@ -323,217 +320,242 @@ FStrataEvaluateResult StrataEvaluateBSDFCommon(FStrataBSDFContext BSDFContext, F
 			Sample.bSubsurface		= BSDF_GETSSSTYPE(BSDFContext.BSDF) == SSS_TYPE_DIFFUSION || BSDF_GETSSSTYPE(BSDFContext.BSDF) == SSS_TYPE_DIFFUSION_PROFILE;
 
 
-			//
-			// Apply energy conservation on the diffuse component
-			// If the specular layer is anisotropic, the energy term is computed onto the 'main' roughness [Kulla 2019]
-			//
-			float3 MSScale = 1.0f;
-			float3 HazeMSScale = 1.0f;
+
+
+			////
+			//// Evaluate the specular component.
+			//// This takes into account multiple scattering, micro occlusion.
+			//// Note: anisotropy completely disables area integrations. Lights fall back to punctual.
+			////
+
+			const bool bIsRectLight = IntegrationType == INTEGRATION_AREA_LIGHT_RECT;
+
+			// Primary highlight
+			float PDF = 0;
 			float3 DirectionalAlbedo_SpecularTransmission = 1.0f;
-			float HazeWeight = 0.0f;
-			bool  bHazeAsSimpleClearCoat = false;
-			float3 HazeClearCoatTransmittance = 1.0f;
-			float3 HazeF0 = F0;
-			float3 HazeF90= F90;
 			{
+				float3 MSScale = 1.0f;
 				{
 					FBxDFEnergyTerms EnergyTerms = ComputeGGXSpecEnergyTerms(SafeRoughness, BSDFContext.Context.NoV, F0, F90);
 					DirectionalAlbedo_SpecularTransmission = ComputeEnergyPreservation(EnergyTerms);
 					MSScale = ComputeEnergyConservation(EnergyTerms);
 				}
-
-			#if STRATA_FASTPATH==0
-				if (bHaziness)
+				// Apply energy conservation on the diffuse component
+				// If the specular layer is anisotropic, the energy term is computed onto the 'main' roughness [Kulla 2019]
+				//
+				Sample.DiffusePathValue *= DirectionalAlbedo_SpecularTransmission;
+				Sample.IntegratedDiffuseValue *= DirectionalAlbedo_SpecularTransmission;
+			
+				float D = 0;
+				float Vis = 0;
+			
+				float3 RectLightSpec = 0;
+				#if STRATA_COMPLEXPATH
+				BRANCH
+				if (bHasAnisotropy)
 				{
-					FHaziness Haziness = UnpackHaziness(SLAB_HAZINESS(BSDFContext.BSDF));
-					HazeWeight = Haziness.Weight;
-					HazeSafeRoughness = MakeRoughnessSafe(Haziness.Roughness);
-					bHazeAsSimpleClearCoat = Haziness.bSimpleClearCoat;
+					// Generalized microfacet specular
+					{
+						float  Alpha  = Square(SafeRoughness);
+						float2 AlphaXY = 0;
+						GetAnisotropicRoughness(Alpha, SLAB_ANISOTROPY(BSDFContext.BSDF), AlphaXY.x, AlphaXY.y);
 
-					BRANCH
-					if (bHazeAsSimpleClearCoat)
+						D   = D_GGXaniso(AlphaXY.x, AlphaXY.y, BSDFContext.Context.NoH, BSDFContext.Context.XoH, BSDFContext.Context.YoH);
+						Vis = StrataVisSmithJointAniso(AlphaXY.x, AlphaXY.y, BSDFContext.Context.NoV, BSDFContext.SatNoL, BSDFContext.Context.XoV, BSDFContext.Context.XoL, BSDFContext.Context.YoV, BSDFContext.Context.YoL);
+						const float H_PDF = VisibleGGXPDF_aniso(BSDFContext.TangentV, BSDFContext.TangentH, AlphaXY);
+						PDF = RayPDFToReflectionRayPDF(BSDFContext.Context.VoH, H_PDF);
+					}
+				}
+				else
+				#endif
+				{
+					if (bIsRectLight)
 					{
-						HazeF0 = 0.04f;
-						HazeF90 = 1.0f;
+						// In this case, we set D to 1 and Vis will contain the area light / GGX integration
+						{
+							RectLightSpec = RectGGXApproxLTC(SafeRoughness, F0, BSDFContext.N, BSDFContext.V, AreaLightContext.AreaLight.Rect, AreaLightContext.AreaLight.Texture);
 
-						BxDFContext ClearCoatContext = RefractClearCoatContext(BSDFContext.Context);
-						HazeClearCoatTransmittance = SimpleClearCoatTransmittance(ClearCoatContext.NoL, ClearCoatContext.NoV, StrataGetBSDFMetallic(BSDFContext.BSDF), StrataGetBSDFBaseColor(BSDFContext.BSDF));
+							const float H_PDF = VisibleGGXPDF(BSDFContext.TangentV, BSDFContext.TangentH, Alpha2);
+							PDF = RayPDFToReflectionRayPDF(BSDFContext.Context.VoH, H_PDF);
+						}
 					}
+					else
+					{
+						// Special override for roughness for supporting area light integrator with Sphere/Tube/Disk light, which modifies/increase roughness.
 
-					FBxDFEnergyTerms EnergyTerms = ComputeGGXSpecEnergyTerms(HazeSafeRoughness, BSDFContext.Context.NoV, HazeF0, HazeF90);
-					HazeMSScale = ComputeEnergyConservation(EnergyTerms);
+						// Generalized microfacet specular
+						{
+							float Alpha2Spec = Alpha2;
+							if(IntegrationType == INTEGRATION_PUNCTUAL_LIGHT)
+							{
+								D = D_GGX(Alpha2Spec, BSDFContext.Context.NoH);
+							}
+							else
+							{
+								const float Energy = EnergyNormalization(Alpha2Spec, BSDFContext.Context.VoH, AreaLightContext.AreaLight);
+								D = D_GGX(Alpha2Spec, BSDFContext.Context.NoH) * Energy;
+							}
+							Vis = StrataVisSmithJoint(Alpha2Spec, BSDFContext.Context.NoV, AreaLightContext.NoL);
+							const float H_PDF = VisibleGGXPDF(BSDFContext.TangentV, BSDFContext.TangentH, Alpha2Spec);
+							PDF = RayPDFToReflectionRayPDF(BSDFContext.Context.VoH, H_PDF);
+						}
+					}
 				}
-			#endif
-			
-				Sample.DiffusePathValue *= DirectionalAlbedo_SpecularTransmission;
-				Sample.IntegratedDiffuseValue *= DirectionalAlbedo_SpecularTransmission;
-			}
 
-			////
-			//// Evaluate the specular component.
-			//// This takes into account multiple scattering, micro occlusion.
-			//// Note: anisotropy completely disables area integrations. Lights fall back to punctual.
-			////
+					
 
-			float3 RectLightSpec = 0;
-			float3 RectLightSpecHaze = 0;
-			float D = 0;
-			float Vis = 0;
-			float PDF = 0;
-			float HazeD = 0;
-			float HazeVis = 0;
-			float HazePDF = 0;
-			const bool bIsRectLight = IntegrationType == INTEGRATION_AREA_LIGHT_RECT;
+				const float SpecularMicroOcclusion = F0RGBToMicroOcclusion(F0);
+				const float3 FresnelTerm = SpecularMicroOcclusion * F_Schlick(F0, F90, BSDFContext.Context.VoH);
+			
 
-			#if STRATA_COMPLEXPATH
-			BRANCH
-			if (bHasAnisotropy)
-			{
-				// Generalized microfacet specular
+				if (bIsRectLight)
 				{
-					float  Alpha  = Square(SafeRoughness);
-					float2 AlphaXY = 0;
-					GetAnisotropicRoughness(Alpha, SLAB_ANISOTROPY(BSDFContext.BSDF), AlphaXY.x, AlphaXY.y);
-
-					D   = D_GGXaniso(AlphaXY.x, AlphaXY.y, BSDFContext.Context.NoH, BSDFContext.Context.XoH, BSDFContext.Context.YoH);
-					Vis = StrataVisSmithJointAniso(AlphaXY.x, AlphaXY.y, BSDFContext.Context.NoV, BSDFContext.SatNoL, BSDFContext.Context.XoV, BSDFContext.Context.XoL, BSDFContext.Context.YoV, BSDFContext.Context.YoL);
-					const float H_PDF = VisibleGGXPDF_aniso(BSDFContext.TangentV, BSDFContext.TangentH, AlphaXY);
-					PDF = RayPDFToReflectionRayPDF(BSDFContext.Context.VoH, H_PDF);
+					Sample.SpecularPathValue     = RectLightSpec     * MSScale     * SpecularMicroOcclusion;
 				}
-
-				if (bHaziness)
+				else
 				{
-					float2 HazeAlpha = 0;
-					GetAnisotropicRoughness(HazeSafeRoughness, SLAB_ANISOTROPY(BSDFContext.BSDF), HazeAlpha.x, HazeAlpha.y);
-
-					HazeD   = D_GGXaniso(HazeAlpha.x, HazeAlpha.y, BSDFContext.Context.NoH, BSDFContext.Context.XoH, BSDFContext.Context.YoH);
-					HazeVis = StrataVisSmithJointAniso(HazeAlpha.x, HazeAlpha.y, BSDFContext.Context.NoV, BSDFContext.Context.NoL, BSDFContext.Context.XoV, BSDFContext.Context.XoL, BSDFContext.Context.YoV, BSDFContext.Context.YoL);
-					const float H_PDF = VisibleGGXPDF_aniso(BSDFContext.TangentV, BSDFContext.TangentH, HazeAlpha);
-					HazePDF = RayPDFToReflectionRayPDF(BSDFContext.Context.VoH, H_PDF);
+					Sample.SpecularPathValue     = (D     * Vis)     * FresnelTerm     * MSScale;
 				}
+
+				Sample.SpecularPathProbability     = 1.0f;
+				Sample.SpecularHazePathProbability = 0.0f;
+				Sample.IntegratedSpecularValue = Sample.SpecularPathValue;
 			}
-			else
-			#endif
+
+			// Secondary highlight
+			float HazePDF = 0;
+			#if STRATA_FASTPATH==0
+			if (bHaziness)
 			{
-				if (bIsRectLight)
+				float3 HazeClearCoatTransmittance = 1.0f;
+				float3 HazeF0 = F0;
+				float3 HazeF90= F90;
+			
+				const FHaziness Haziness = UnpackHaziness(SLAB_HAZINESS(BSDFContext.BSDF));
+				const float HazeWeight = Haziness.Weight;
+				const float HazeSafeRoughness = MakeRoughnessSafe(Haziness.Roughness);
+				const bool bHazeAsSimpleClearCoat = Haziness.bSimpleClearCoat;
+
+				BRANCH
+				if (bHazeAsSimpleClearCoat)
 				{
-					// In this case, we set D to 1 and Vis will contain the area light / GGX integration
-					{
-						RectLightSpec = RectGGXApproxLTC(SafeRoughness, F0, BSDFContext.N, BSDFContext.V, AreaLightContext.AreaLight.Rect, AreaLightContext.AreaLight.Texture);
+					HazeF0 = 0.04f;
+					HazeF90 = 1.0f;
 
-						const float H_PDF = VisibleGGXPDF(BSDFContext.TangentV, BSDFContext.TangentH, Alpha2);
-						PDF = RayPDFToReflectionRayPDF(BSDFContext.Context.VoH, H_PDF);
-					}
-					
-					#if STRATA_FASTPATH==0
-					if (bHaziness)
+					BxDFContext ClearCoatContext = RefractClearCoatContext(BSDFContext.Context);
+					HazeClearCoatTransmittance = SimpleClearCoatTransmittance(ClearCoatContext.NoL, ClearCoatContext.NoV, StrataGetBSDFMetallic(BSDFContext.BSDF), StrataGetBSDFBaseColor(BSDFContext.BSDF));
+				}
+
+				FBxDFEnergyTerms EnergyTerms = ComputeGGXSpecEnergyTerms(HazeSafeRoughness, BSDFContext.Context.NoV, HazeF0, HazeF90);
+				float3 HazeMSScale = ComputeEnergyConservation(EnergyTerms);
+
+				float HazeD = 0;
+				float HazeVis = 0;
+
+				float3 RectLightSpecHaze = 0;
+				#if STRATA_COMPLEXPATH
+				BRANCH
+				if (bHasAnisotropy)
+				{
 					{
-						RectLightSpecHaze = RectGGXApproxLTC(HazeSafeRoughness, HazeF0, BSDFContext.N, BSDFContext.V, AreaLightContext.AreaLight.Rect, AreaLightContext.AreaLight.Texture);
-					
-						const float H_PDF = VisibleGGXPDF(BSDFContext.TangentV, BSDFContext.TangentH, Pow4(HazeSafeRoughness));
+						float2 HazeAlpha = 0;
+						GetAnisotropicRoughness(HazeSafeRoughness, SLAB_ANISOTROPY(BSDFContext.BSDF), HazeAlpha.x, HazeAlpha.y);
+
+						HazeD   = D_GGXaniso(HazeAlpha.x, HazeAlpha.y, BSDFContext.Context.NoH, BSDFContext.Context.XoH, BSDFContext.Context.YoH);
+						HazeVis = StrataVisSmithJointAniso(HazeAlpha.x, HazeAlpha.y, BSDFContext.Context.NoV, BSDFContext.Context.NoL, BSDFContext.Context.XoV, BSDFContext.Context.XoL, BSDFContext.Context.YoV, BSDFContext.Context.YoL);
+						const float H_PDF = VisibleGGXPDF_aniso(BSDFContext.TangentV, BSDFContext.TangentH, HazeAlpha);
 						HazePDF = RayPDFToReflectionRayPDF(BSDFContext.Context.VoH, H_PDF);
 					}
-					#endif
 				}
 				else
+				#endif
 				{
-					// Special override for roughness for supporting area light integrator with Sphere/Tube/Disk light, which modifies/increase roughness.
-
-					// Generalized microfacet specular
+					if (bIsRectLight)
 					{
-						float Alpha2Spec = Alpha2;
-						if(IntegrationType == INTEGRATION_PUNCTUAL_LIGHT)
 						{
-							D = D_GGX(Alpha2Spec, BSDFContext.Context.NoH);
-						}
-						else
-						{
-							const float Energy = EnergyNormalization(Alpha2Spec, BSDFContext.Context.VoH, AreaLightContext.AreaLight);
-							D = D_GGX(Alpha2Spec, BSDFContext.Context.NoH) * Energy;
+							RectLightSpecHaze = RectGGXApproxLTC(HazeSafeRoughness, HazeF0, BSDFContext.N, BSDFContext.V, AreaLightContext.AreaLight.Rect, AreaLightContext.AreaLight.Texture);
+					
+							const float H_PDF = VisibleGGXPDF(BSDFContext.TangentV, BSDFContext.TangentH, Pow4(HazeSafeRoughness));
+							HazePDF = RayPDFToReflectionRayPDF(BSDFContext.Context.VoH, H_PDF);
 						}
-						Vis = StrataVisSmithJoint(Alpha2Spec, BSDFContext.Context.NoV, AreaLightContext.NoL);
-						const float H_PDF = VisibleGGXPDF(BSDFContext.TangentV, BSDFContext.TangentH, Alpha2Spec);
-						PDF = RayPDFToReflectionRayPDF(BSDFContext.Context.VoH, H_PDF);
 					}
-
-					#if STRATA_FASTPATH==0
-					if (bHaziness)
+					else
 					{
-						float Alpha2SpecHaze = Pow4(HazeSafeRoughness);
-						if (IntegrationType == INTEGRATION_PUNCTUAL_LIGHT)
-						{
-							HazeD = D_GGX(Alpha2SpecHaze, BSDFContext.Context.NoH);
-						}
-						else
+						// Special override for roughness for supporting area light integrator with Sphere/Tube/Disk light, which modifies/increase roughness.
 						{
-							const float Energy = EnergyNormalization(Alpha2SpecHaze, BSDFContext.Context.VoH, AreaLightContext.AreaLight);
-							HazeD = D_GGX(Alpha2SpecHaze, BSDFContext.Context.NoH) * Energy;
+							float Alpha2SpecHaze = Pow4(HazeSafeRoughness);
+							if (IntegrationType == INTEGRATION_PUNCTUAL_LIGHT)
+							{
+								HazeD = D_GGX(Alpha2SpecHaze, BSDFContext.Context.NoH);
+							}
+							else
+							{
+								const float Energy = EnergyNormalization(Alpha2SpecHaze, BSDFContext.Context.VoH, AreaLightContext.AreaLight);
+								HazeD = D_GGX(Alpha2SpecHaze, BSDFContext.Context.NoH) * Energy;
+							}
+							HazeVis = StrataVisSmithJoint(Alpha2SpecHaze, BSDFContext.Context.NoV, AreaLightContext.NoL);
+							const float H_PDF = VisibleGGXPDF(BSDFContext.TangentV, BSDFContext.TangentH, Alpha2SpecHaze);
+							HazePDF = RayPDFToReflectionRayPDF(BSDFContext.Context.VoH, H_PDF);
 						}
-						HazeVis = StrataVisSmithJoint(Alpha2SpecHaze, BSDFContext.Context.NoV, AreaLightContext.NoL);
-						const float H_PDF = VisibleGGXPDF(BSDFContext.TangentV, BSDFContext.TangentH, Alpha2SpecHaze);
-						HazePDF = RayPDFToReflectionRayPDF(BSDFContext.Context.VoH, H_PDF);
 					}
-					#endif
 				}
-			}
 
-			const float SpecularMicroOcclusion = F0RGBToMicroOcclusion(F0);
-			const float3 FresnelTerm = SpecularMicroOcclusion * F_Schlick(F0, F90, BSDFContext.Context.VoH);
-			
-		#if STRATA_FASTPATH==0
-			const float HazeSpecularMicroOcclusion = F0RGBToMicroOcclusion(HazeF0);
-			const float3 HazeFresnelTerm = HazeSpecularMicroOcclusion * F_Schlick(HazeF0, HazeF90, BSDFContext.Context.VoH);
-		#else
-			const float HazeSpecularMicroOcclusion = SpecularMicroOcclusion;
-			const float3 HazeFresnelTerm = FresnelTerm;
-		#endif
-
-			float3 CommonTerm = 0.0f;
-			if (bIsRectLight)
-			{
-				Sample.SpecularPathValue     = RectLightSpec     * MSScale     * SpecularMicroOcclusion;
-				Sample.SpecularHazePathValue = RectLightSpecHaze * HazeMSScale * HazeSpecularMicroOcclusion;
+				
+				const float HazeSpecularMicroOcclusion = F0RGBToMicroOcclusion(HazeF0);
+				const float3 HazeFresnelTerm = HazeSpecularMicroOcclusion * F_Schlick(HazeF0, HazeF90, BSDFContext.Context.VoH);
 
-				CommonTerm = ShadowTerms.SurfaceShadow; /* AreaLightContext.NoL and falloff is part of the LTC integration already */
-			}
-			else
-			{
-				Sample.SpecularPathValue     = (D     * Vis)     * FresnelTerm     * MSScale;
-				Sample.SpecularHazePathValue = (HazeD * HazeVis) * HazeFresnelTerm * HazeMSScale;
+				if (bIsRectLight)
+				{					
+					Sample.SpecularHazePathValue = RectLightSpecHaze * HazeMSScale * HazeSpecularMicroOcclusion;
+				}
+				else
+				{
+					Sample.SpecularHazePathValue = (HazeD * HazeVis) * HazeFresnelTerm * HazeMSScale;
+				}
 
-				CommonTerm = (ShadowTerms.SurfaceShadow * AreaLightContext.NoL * AreaLightContext.Falloff) * AreaLightContext.AreaLight.FalloffColor;
+				// Reminder: 
+				//	- bHazeAsSimpleClearCoat == false means we have two specular lobes that are lerped.
+				//	- bHazeAsSimpleClearCoat == true  means we have a bottom specular lobe (roughness and F0 user specified) and a top specular lob (hard coded absortion, F0=0.04 and user specified roughness.
+				BRANCH
+				if (bHazeAsSimpleClearCoat)
+				{
+					const float TopLayerCoverage = HazeWeight;
+					const float3 TopLayerSpecularTransmittionApprox = saturate(1.0f - HazeFresnelTerm); // Simple and do not use the LUT for now
+					const float3 BottomLayerFactors = lerp(1.0f, HazeClearCoatTransmittance * TopLayerSpecularTransmittionApprox, TopLayerCoverage);
+
+					// Apply the top clear coat throughput onto the bottom layer
+					Sample.SpecularPathValue *= BottomLayerFactors;
+
+					const float BottomLayerFactorsGrey = dot(BottomLayerFactors, (1.0 / 3.0).xxx);
+					Sample.SpecularPathProbability = BottomLayerFactorsGrey / (TopLayerCoverage + BottomLayerFactorsGrey);
+					Sample.SpecularHazePathProbability = TopLayerCoverage / (TopLayerCoverage + BottomLayerFactorsGrey);
+
+					// Clear coat applied over bottom layer diffuse component
+					Sample.DiffusePathValue *= BottomLayerFactors;
+					Sample.IntegratedDiffuseValue *= BottomLayerFactors;
+					// And to the bottom layer specular on top of whiche we add the top layer specular
+					Sample.IntegratedSpecularValue = Sample.SpecularPathValue * BottomLayerFactors + Sample.SpecularHazePathValue * TopLayerCoverage;
+				}
+				else
+				{
+					Sample.SpecularPathProbability     = (1.0f - HazeWeight);
+					Sample.SpecularHazePathProbability = HazeWeight;
+					Sample.IntegratedSpecularValue = Sample.SpecularPathValue * Sample.SpecularPathProbability + Sample.SpecularHazePathValue * Sample.SpecularHazePathProbability;
+				}
 			}
+			#endif
 
-		#if STRATA_FASTPATH==0
-			// Reminder: 
-			//	- bHazeAsSimpleClearCoat == false means we have two specular lobes that are lerped.
-			//	- bHazeAsSimpleClearCoat == true  means we have a bottom specular lobe (roughness and F0 user specified) and a top specular lob (hard coded absortion, F0=0.04 and user specified roughness.
-			if (bHazeAsSimpleClearCoat)
-			{
-				const float TopLayerCoverage = HazeWeight;
-				const float3 TopLayerSpecularTransmittionApprox = saturate(1.0f - HazeFresnelTerm); // Simple and do not use the LUT for now
-				const float3 BottomLayerFactors = lerp(1.0f, HazeClearCoatTransmittance * TopLayerSpecularTransmittionApprox, TopLayerCoverage);
-
-				// Apply the top clear coat throughput onto the bottom layer
-				Sample.SpecularPathValue *= BottomLayerFactors;
-
-				const float BottomLayerFactorsGrey = dot(BottomLayerFactors, (1.0 / 3.0).xxx);
-				Sample.SpecularPathProbability     = BottomLayerFactorsGrey / (TopLayerCoverage + BottomLayerFactorsGrey);
-				Sample.SpecularHazePathProbability = TopLayerCoverage       / (TopLayerCoverage + BottomLayerFactorsGrey);
-
-				// Clear coat applied over bottom layer diffuse component
-				Sample.DiffusePathValue			*= BottomLayerFactors;
-				Sample.IntegratedDiffuseValue	*= BottomLayerFactors;
-				// And to the bottom layer specular on top of whiche we add the top layer specular
-				Sample.IntegratedSpecularValue	 = CommonTerm * (Sample.SpecularPathValue * BottomLayerFactors + Sample.SpecularHazePathValue * TopLayerCoverage);
-			}
-			else
-		#endif
 			{
-				Sample.SpecularPathProbability     = bHaziness ? (1.0f - HazeWeight) : 1.0f;
-				Sample.SpecularHazePathProbability = bHaziness ?         HazeWeight  : 0.0f;
-
-				Sample.IntegratedSpecularValue = CommonTerm * (Sample.SpecularPathValue * Sample.SpecularPathProbability + Sample.SpecularHazePathValue * Sample.SpecularHazePathProbability);
+				float3 CommonTerm = 0.0f;
+				if (bIsRectLight)
+				{
+					CommonTerm = ShadowTerms.SurfaceShadow; /* AreaLightContext.NoL and falloff is part of the LTC integration already */
+				}
+				else
+				{
+					CommonTerm = (ShadowTerms.SurfaceShadow * AreaLightContext.NoL * AreaLightContext.Falloff) * AreaLightContext.AreaLight.FalloffColor;
+				}
+				Sample.IntegratedSpecularValue *= CommonTerm;
 			}
 
 			Sample.SpecularPDF = PDF * Sample.SpecularPathProbability;
