Fix shader compilation warning.

diff --git a/Engine/Shaders/Private/DeferredShadingCommon.ush b/Engine/Shaders/Private/DeferredShadingCommon.ush
--- a/Engine/Shaders/Private/DeferredShadingCommon.ush
+++ b/Engine/Shaders/Private/DeferredShadingCommon.ush
@@ -408,7 +408,7 @@ bool CheckerFromPixelPos(uint2 PixelPos)
 {
 	// todo: Index is float and by staying float we can optimize this 
 	// We alternate the pattern to get 2x supersampling on the lower res data to get more near to full res
-	uint TemporalAASampleIndex = View.TemporalAAParams.x;
+	uint TemporalAASampleIndex = uint(View.TemporalAAParams.x);
 
 #if FEATURE_LEVEL >= FEATURE_LEVEL_SM4
 	return (PixelPos.x + PixelPos.y + TemporalAASampleIndex) % 2;
diff --git a/Engine/Shaders/Private/HairBsdf.ush b/Engine/Shaders/Private/HairBsdf.ush
--- a/Engine/Shaders/Private/HairBsdf.ush
+++ b/Engine/Shaders/Private/HairBsdf.ush
@@ -374,9 +374,9 @@ float3 HairShading( FGBufferData GBuffer, float3 L, float3 V, half3 N, float Sha
 	{
 		const float sa = sin(Alpha[0]);
 		const float ca = cos(Alpha[0]);
-		float Shift = 2 * sa * (ca * CosHalfPhi * sqrt(1 - SinThetaV * SinThetaV) + sa * SinThetaV);
+		float ShiftR = 2 * sa * (ca * CosHalfPhi * sqrt(1 - SinThetaV * SinThetaV) + sa * SinThetaV);
 		float BScale = HairTransmittance.bUseSeparableR ? sqrt(2.0) * CosHalfPhi : 1;
-		float Mp = Hair_g(B[0] * BScale, SinThetaL + SinThetaV - Shift);
+		float Mp = Hair_g(B[0] * BScale, SinThetaL + SinThetaV - ShiftR);
 		float Np = 0.25 * CosHalfPhi;
 		float Fp = Hair_F(sqrt(saturate(0.5 + 0.5 * VoL)));
 		S += Mp * Np * Fp * (GBuffer.Specular * 2) * lerp(1, Backlit, saturate(-VoL));
diff --git a/Engine/Shaders/Private/ShadingModels.ush b/Engine/Shaders/Private/ShadingModels.ush
--- a/Engine/Shaders/Private/ShadingModels.ush
+++ b/Engine/Shaders/Private/ShadingModels.ush
@@ -601,7 +601,7 @@ FDirectLighting ClearCoatBxDF( FGBufferData GBuffer, half3 N, half3 V, half3 L,
 			D2 = D_GGX(a2, BottomContext.NoH);
 #endif
 		}
-		half3 F = F_Schlick(GBuffer.SpecularColor, BottomContext.VoH);
+		half3 F_Bot = F_Schlick(GBuffer.SpecularColor, BottomContext.VoH);
 		half3 F_DefaultLit = F_Schlick(GBuffer.SpecularColor, Context.VoH);
 		
 		float Energy = 0;
@@ -626,7 +626,7 @@ FDirectLighting ClearCoatBxDF( FGBufferData GBuffer, half3 N, half3 V, half3 L,
 		// Note: reusing D and V from refracted context to save computation when ClearCoat < 1
 		float3 CommonSpecular = (Energy * Falloff * NoL * D2 * Vis2) * AreaLight.FalloffColor;
 		float3 DefaultSpecular = F_DefaultLit;
-		float3 RefractedSpecular = FresnelCoeff * Transmission * F;
+		float3 RefractedSpecular = FresnelCoeff * Transmission * F_Bot;
 		Lighting.Specular += CommonSpecular * lerp(DefaultSpecular, RefractedSpecular, ClearCoat);
 	}
 
diff --git a/Engine/Shaders/Private/Strata/Strata.ush b/Engine/Shaders/Private/Strata/Strata.ush
--- a/Engine/Shaders/Private/Strata/Strata.ush
+++ b/Engine/Shaders/Private/Strata/Strata.ush
@@ -2593,12 +2593,12 @@ void StrataSanitizeBSDF(inout FStrataBSDF BSDF)
 // Return 0 if v is zero, 1 otherwise
 float IsNonZeroFast(float v)
 {
-	return min(asuint(v) & 0x7FFFFFFF, 1u);
+	return float(min(asuint(v) & 0x7FFFFFFF, 1u));
 }
 
 float3 IsNonZeroFast(float3 v)
 {
-	return min(asuint(v) & 0x7FFFFFFF, 1u.xxx);
+	return float3(min(asuint(v) & 0x7FFFFFFF, 1u.xxx));
 }
 
 // This is used to respect artistic desire, e.g. shutdown diffuse or specular.
@@ -3056,7 +3056,7 @@ uint PackR7(float Value, float Dither)
 	const float CenteredDither = Dither - 0.5f;
 	const float Dither7bits = CenteredDither * 1.0f / 127.0f;
 	const float DitheredValue = DitherIfNonBlack(saturate(Value), Dither7bits);
-	return DitheredValue * 127.0f;
+	return uint(DitheredValue * 127.0f);
 }
 
 float UnpackR7(uint Value)
