* Rearrange & simplify terms for avoiding extra ALU
* Remove energy conservation dynamic branching in non-development mode

This reduces Strata simple lighting shader from 4 VGPRs on prev gen platforms.

diff --git a/Engine/Shaders/Private/ShadingEnergyConservation.ush b/Engine/Shaders/Private/ShadingEnergyConservation.ush
--- a/Engine/Shaders/Private/ShadingEnergyConservation.ush
+++ b/Engine/Shaders/Private/ShadingEnergyConservation.ush
@@ -111,7 +111,11 @@ FBxDFEnergyTerms ComputeFresnelEnergyTerms(float2 E, float3 F0, float3 F90)
 	FBxDFEnergyTerms Result;
 	// [2] Eq 16: this restores the missing energy of the bsdf, while also accounting for the fact that the fresnel term causes some energy to be absorbed
 	// NOTE: using F0 here is an approximation, but for schlick fresnel Favg is almost exactly equal to F0
+#if USE_DEVELOPMENT_SHADERS
 	Result.W = View.bShadingEnergyConservation ? (1.0 + F0 * ((1 - E.x) / E.x)) : 1.0f;
+#else
+	Result.W = 1.0 + F0 * ((1 - E.x) / E.x);
+#endif
 	// Now estimate the amount of energy reflected off this specular lobe so that we can remove it from underlying BxDF layers (like diffuse)
 	// This relies on the split-sum approximation as in [3] Sec 4.
 	// This term can also be useful to compute the probability of choosing among lobes
@@ -184,7 +188,11 @@ float3 ComputeEnergyPreservation(FBxDFEnergyTerms EnergyTerms)
 {
 #if USE_ENERGY_CONSERVATION > 0
 	// Use the directional albedo luminance to avoid color-shift due to metallic specular (for which the energy should be absorbed, not transmitted)
+	#if USE_DEVELOPMENT_SHADERS
 	return View.bShadingEnergyPreservation ? (1 - Luminance(EnergyTerms.E)) : 1.0f;
+	#else
+	return 1 - Luminance(EnergyTerms.E);
+	#endif
 #else
 	return 1.0f;
 #endif
diff --git a/Engine/Shaders/Private/Strata/StrataEvaluation.ush b/Engine/Shaders/Private/Strata/StrataEvaluation.ush
--- a/Engine/Shaders/Private/Strata/StrataEvaluation.ush
+++ b/Engine/Shaders/Private/Strata/StrataEvaluation.ush
@@ -405,16 +405,15 @@ FStrataEvaluateResult StrataEvaluateBSDFCommon(FStrataBSDFContext BSDFContext, F
 					
 
 				const float SpecularMicroOcclusion = F0RGBToMicroOcclusion(F0);
-				const float3 FresnelTerm = SpecularMicroOcclusion * F_Schlick(F0, F90, BSDFContext.Context.VoH);
-			
 
 				if (bIsRectLight)
 				{
-					Sample.SpecularPathValue     = RectLightSpec     * MSScale     * SpecularMicroOcclusion;
+					Sample.SpecularPathValue = SpecularMicroOcclusion * RectLightSpec * MSScale;
 				}
 				else
 				{
-					Sample.SpecularPathValue     = (D     * Vis)     * FresnelTerm     * MSScale;
+					const float3 FresnelTerm = F_Schlick(F0, F90, BSDFContext.Context.VoH);
+					Sample.SpecularPathValue = (SpecularMicroOcclusion * D * Vis) * FresnelTerm * MSScale;
 				}
 
 				Sample.SpecularPathProbability     = 1.0f;
@@ -427,7 +426,6 @@ FStrataEvaluateResult StrataEvaluateBSDFCommon(FStrataBSDFContext BSDFContext, F
 			#if STRATA_FASTPATH==0
 			if (bHaziness)
 			{
-				float3 HazeClearCoatTransmittance = 1.0f;
 				float3 HazeF0 = F0;
 				float3 HazeF90= F90;
 			
@@ -441,9 +439,6 @@ FStrataEvaluateResult StrataEvaluateBSDFCommon(FStrataBSDFContext BSDFContext, F
 				{
 					HazeF0 = 0.04f;
 					HazeF90 = 1.0f;
-
-					BxDFContext ClearCoatContext = RefractClearCoatContext(BSDFContext.Context);
-					HazeClearCoatTransmittance = SimpleClearCoatTransmittance(ClearCoatContext.NoL, ClearCoatContext.NoV, StrataGetBSDFMetallic(BSDFContext.BSDF), StrataGetBSDFBaseColor(BSDFContext.BSDF));
 				}
 
 				FBxDFEnergyTerms EnergyTerms = ComputeGGXSpecEnergyTerms(HazeSafeRoughness, BSDFContext.Context.NoV, HazeF0, HazeF90);
@@ -502,15 +497,14 @@ FStrataEvaluateResult StrataEvaluateBSDFCommon(FStrataBSDFContext BSDFContext, F
 
 				
 				const float HazeSpecularMicroOcclusion = F0RGBToMicroOcclusion(HazeF0);
-				const float3 HazeFresnelTerm = HazeSpecularMicroOcclusion * F_Schlick(HazeF0, HazeF90, BSDFContext.Context.VoH);
-
+				const float3 HazeFresnelTerm = F_Schlick(HazeF0, HazeF90, BSDFContext.Context.VoH);
 				if (bIsRectLight)
 				{					
 					Sample.SpecularHazePathValue = RectLightSpecHaze * HazeMSScale * HazeSpecularMicroOcclusion;
 				}
 				else
 				{
-					Sample.SpecularHazePathValue = (HazeD * HazeVis) * HazeFresnelTerm * HazeMSScale;
+					Sample.SpecularHazePathValue = (HazeD * HazeVis * HazeSpecularMicroOcclusion) * HazeFresnelTerm * HazeMSScale;
 				}
 
 				// Reminder: 
@@ -519,8 +513,11 @@ FStrataEvaluateResult StrataEvaluateBSDFCommon(FStrataBSDFContext BSDFContext, F
 				BRANCH
 				if (bHazeAsSimpleClearCoat)
 				{
+					const BxDFContext ClearCoatContext = RefractClearCoatContext(BSDFContext.Context);
+					const float3 HazeClearCoatTransmittance = SimpleClearCoatTransmittance(ClearCoatContext.NoL, ClearCoatContext.NoV, StrataGetBSDFMetallic(BSDFContext.BSDF), StrataGetBSDFBaseColor(BSDFContext.BSDF));
+
 					const float TopLayerCoverage = HazeWeight;
-					const float3 TopLayerSpecularTransmittionApprox = saturate(1.0f - HazeFresnelTerm); // Simple and do not use the LUT for now
+					const float TopLayerSpecularTransmittionApprox = saturate(1.0f - HazeSpecularMicroOcclusion * HazeFresnelTerm.x); // Simple and do not use the LUT for now. With bHazeAsSimpleClearCoat, Fresnel is achromatic
 					const float3 BottomLayerFactors = lerp(1.0f, HazeClearCoatTransmittance * TopLayerSpecularTransmittionApprox, TopLayerCoverage);
 
 					// Apply the top clear coat throughput onto the bottom layer
@@ -540,7 +537,7 @@ FStrataEvaluateResult StrataEvaluateBSDFCommon(FStrataBSDFContext BSDFContext, F
 				{
 					Sample.SpecularPathProbability     = (1.0f - HazeWeight);
 					Sample.SpecularHazePathProbability = HazeWeight;
-					Sample.IntegratedSpecularValue = Sample.SpecularPathValue * Sample.SpecularPathProbability + Sample.SpecularHazePathValue * Sample.SpecularHazePathProbability;
+					Sample.IntegratedSpecularValue = lerp(Sample.SpecularPathValue, Sample.SpecularHazePathValue, HazeWeight);
 				}
 			}
 			#endif
@@ -1366,7 +1363,7 @@ FStrataEnvLightResult StrataEvaluateForEnvLight(FStrataBSDFContext BSDFContext,
 
 						const float HazeF0 = 0.04f;
 						const float HazeF90= 1.0f;
-						const float SpecularMicroOcclusion = F0RGBToMicroOcclusion(HazeF0);
+						const float HazeSpecularMicroOcclusion = F0RGBToMicroOcclusion(HazeF0);
 						FBxDFEnergyTerms EnergyTerms = ComputeGGXSpecEnergyTerms(HazeSafeRoughness, BSDFContext.Context.NoV, HazeF0, HazeF90);
 						const float3 EvalEnvBRDFHaze = EnergyTerms.E;
 						const float3 TopLayerSpecularTransmittionApprox = saturate(1.0f - EvalEnvBRDFHaze); // Simple and do not use extra LUT for now
@@ -1377,7 +1374,7 @@ FStrataEnvLightResult StrataEvaluateForEnvLight(FStrataBSDFContext BSDFContext,
 						StrataEnvLightResult.SpecularWeight *= BottomLayerFactors;
 						StrataEnvLightResult.DiffuseWeight  *= BottomLayerFactors;
 
-						StrataEnvLightResult.SpecularHazeWeight = TopLayerCoverage * EvalEnvBRDFHaze * SpecularMicroOcclusion;
+						StrataEnvLightResult.SpecularHazeWeight = (TopLayerCoverage * HazeSpecularMicroOcclusion) * EvalEnvBRDFHaze;
 						StrataEnvLightResult.SpecularHazeSafeRoughness = HazeSafeRoughness;
 
 						StrataEnvLightResult.SSRSpecularWeight = lerp(StrataEnvLightResult.SpecularWeight, StrataEnvLightResult.SpecularHazeWeight, TopLayerCoverage);
@@ -1396,7 +1393,7 @@ FStrataEnvLightResult StrataEvaluateForEnvLight(FStrataBSDFContext BSDFContext,
 						// Compute the directional albedo for hazy specular
 						FBxDFEnergyTerms EnergyTerms = ComputeGGXSpecEnergyTerms(HazeSafeRoughness, BSDFContext.Context.NoV, F0, F90);
 						const float3 EvalEnvBRDFHaze = EnergyTerms.E;
-						StrataEnvLightResult.SpecularHazeWeight = lerp(0.0f, HazeWeight, HazinessFadeIn) * EvalEnvBRDFHaze * SpecularMicroOcclusion;
+						StrataEnvLightResult.SpecularHazeWeight = (lerp(0.0f, HazeWeight, HazinessFadeIn) * SpecularMicroOcclusion) * EvalEnvBRDFHaze;
 						StrataEnvLightResult.SpecularHazeSafeRoughness = HazeSafeRoughness;
 					}
 				}
