Enfore in shader SourceLength to be 0 for directional light.

Save 1 vgpr on prev gen console.

diff --git a/Engine/Shaders/Private/DeferredLightPixelShaders.usf b/Engine/Shaders/Private/DeferredLightPixelShaders.usf
--- a/Engine/Shaders/Private/DeferredLightPixelShaders.usf
+++ b/Engine/Shaders/Private/DeferredLightPixelShaders.usf
@@ -91,7 +91,7 @@ FDeferredLightData SetupLightDataForStandardDeferred()
 	LightData.Tangent = DeferredLightUniforms.Tangent;
 	LightData.SpotAngles = DeferredLightUniforms.SpotAngles;
 	LightData.SourceRadius = DeferredLightUniforms.SourceRadius;
-	LightData.SourceLength = DeferredLightUniforms.SourceLength;
+	LightData.SourceLength = LIGHT_SOURCE_SHAPE > 0 ? DeferredLightUniforms.SourceLength : 0;
     LightData.SoftSourceRadius = DeferredLightUniforms.SoftSourceRadius;
 	LightData.SpecularScale = DeferredLightUniforms.SpecularScale;
 	LightData.ContactShadowLength = abs(DeferredLightUniforms.ContactShadowLength);
