* Add achromatic energy conservation to reduce ALU and occupancy on lower-platform

Both chromatic and achromatic versions coexist, but the achromatic version is optionally used based on quality settings.

diff --git a/Engine/Shaders/Private/ShadingEnergyConservation.ush b/Engine/Shaders/Private/ShadingEnergyConservation.ush
--- a/Engine/Shaders/Private/ShadingEnergyConservation.ush
+++ b/Engine/Shaders/Private/ShadingEnergyConservation.ush
@@ -20,7 +20,7 @@
 //   https://blog.selfshadow.com/publications/s2017-shading-course/imageworks/s2017_pbs_imageworks_slides_v2.pdf
 // [2] "Practical multiple scattering compensation for microfacet models" - Emmanuel Turquin
 //   https://blog.selfshadow.com/publications/turquin/ms_comp_final.pdf
-// [3] "A Multiple-Scattering Microfacet Model for Real-Time Image Based Lighting" - Carmelo J. Fdez-Agüera
+// [3] "A Multiple-Scattering Microfacet Model for Real-Time Image Based Lighting" - Carmelo J. Fdez-Aguera
 //   http://jcgt.org/published/0008/01/03/
 // 
 // We follow the terminology from [1]:
@@ -36,6 +36,9 @@
 //    1: use a small baked texture to store the directional albedo (and fresnel weighted albedo)
 //    2: use an analytic approximation to directional albedo (doesn't require any textures to be bound and can be faster -- WIP)
 
+///////////////////////////////////////////////////////////////////////////////////////////////////
+// LUT Lookups
+
 float2 GGXEnergyLookup(float Roughness, float NoV)
 {
 #if USE_ENERGY_CONSERVATION == 1
@@ -99,107 +102,55 @@ float DiffuseEnergyLookup(float Roughness, float NoV)
 #endif
 }
 
-struct FBxDFEnergyTerms {
-	float3 W; // overall weight to scale the lobe BxDF by to ensure energy conservation
-	float3 E; // Directional albedo of the lobe for energy preservation and lobe picking
-};
-
-// Given a split-sum approximation of directional albedo for a BxDF, compute multiple scattering weight and multiple scattering directional albedo
-// while taking into account the fresnel term (assumed to be F_Schlick)
-FBxDFEnergyTerms ComputeFresnelEnergyTerms(float2 E, float3 F0, float3 F90) 
-{
-	FBxDFEnergyTerms Result;
-	// [2] Eq 16: this restores the missing energy of the bsdf, while also accounting for the fact that the fresnel term causes some energy to be absorbed
-	// NOTE: using F0 here is an approximation, but for schlick fresnel Favg is almost exactly equal to F0
-#if USE_DEVELOPMENT_SHADERS
-	Result.W = View.bShadingEnergyConservation ? (1.0 + F0 * ((1 - E.x) / E.x)) : 1.0f;
-#else
-	Result.W = 1.0 + F0 * ((1 - E.x) / E.x);
-#endif
-	// Now estimate the amount of energy reflected off this specular lobe so that we can remove it from underlying BxDF layers (like diffuse)
-	// This relies on the split-sum approximation as in [3] Sec 4.
-	// This term can also be useful to compute the probability of choosing among lobes
-	Result.E = Result.W * (E.x * F0 + E.y * (F90 - F0));
-	return Result;
-}
-
-FBxDFEnergyTerms ComputeGGXSpecEnergyTerms(float Roughness, float NoV, float3 F0, float3 F90)
-{
-	FBxDFEnergyTerms Out;
-#if USE_ENERGY_CONSERVATION > 0
-	{
-		Out = ComputeFresnelEnergyTerms(GGXEnergyLookup(Roughness, NoV), F0, F90);
-	}
-#else
-	{
-		Out.W = 1.0f;
-		Out.E = F0;
-	}
-#endif
-	return Out;
-}
-
-FBxDFEnergyTerms ComputeGGXSpecEnergyTerms(float Roughness, float NoV, float3 F0)
-{
-	const float F90 = saturate(50.0 * F0.g); // See F_Schlick implementation
-	return ComputeGGXSpecEnergyTerms(Roughness, NoV, F0, F90);
-}
-
-FBxDFEnergyTerms ComputeClothEnergyTerms(float Roughness, float NoV, float3 F0, float3 F90)
-{
-	FBxDFEnergyTerms Out;
-#if USE_ENERGY_CONSERVATION > 0
-	{
-		Out = ComputeFresnelEnergyTerms(ClothEnergyLookup(Roughness, NoV), F0, F90);
-	}
-#else
-	{
-		Out.W = 1.0f;
-		Out.E = F0;
-	}
-#endif
-	return Out;
-}
-
-FBxDFEnergyTerms ComputeClothEnergyTerms(float Roughness, float NoV, float3 F0)
-{
-	const float F90 = saturate(50.0 * F0.g); // See F_Schlick implementation
-	return ComputeClothEnergyTerms(Roughness, NoV, F0, F90);
-}
-
-FBxDFEnergyTerms ComputeDiffuseEnergyTerms(float Roughness, float NoV)
-{
-	FBxDFEnergyTerms Out;
-#if USE_ENERGY_CONSERVATION > 0
-	{
-		Out.E = DiffuseEnergyLookup(Roughness, NoV);
-	}
-#else
-	{
-		Out.E = 1.0f;
-	}
+///////////////////////////////////////////////////////////////////////////////////////////////////
+// BxDFs Energy terms & Energy Preservation/Conservation utils
+
+// Note:
+// * For performance reason, we maintain both a *chromatic* and an *achromatic* version of the energy terms.
+// * Chromatic and achromatic version are respectively accessible with xxxRGB or xxxA suffix. 
+// * USE_ACHROMATIC_BXDF_ENERGY determines which version should be used by default
+
+// Chromatic version
+#define BXDF_ENERGY_TYPE float3
+#define BXDF_ENERGY_SUFFIX(N) N##RGB
+#define IS_BXDF_ENERGY_TYPE_ACHROMATIC 0
+#include "ShadingEnergyConservationTemplate.ush"
+#undef BXDF_ENERGY_TYPE
+#undef BXDF_ENERGY_TYPE_SUFFIX
+#undef IS_BXDF_ENERGY_TYPE_ACHROMATIC
+
+// Achromatic version
+#define BXDF_ENERGY_TYPE float
+#define BXDF_ENERGY_SUFFIX(N) N##A
+#define IS_BXDF_ENERGY_TYPE_ACHROMATIC 1
+#include "ShadingEnergyConservationTemplate.ush"
+#undef BXDF_ENERGY_TYPE
+#undef BXDF_ENERGY_TYPE_SUFFIX
+#undef IS_BXDF_ENERGY_TYPE_ACHROMATIC
+
+///////////////////////////////////////////////////////////////////////////////////////////////////
+// Default implementation
+
+#ifndef USE_ACHROMATIC_BXDF_ENERGY
+#define USE_ACHROMATIC_BXDF_ENERGY 0
 #endif
-	Out.W = 1.0f;
-	return Out;
-}
 
-// Return the energy absorbed by upper layer (e.g., for the specular layer attenuation onto diffuse)
-float3 ComputeEnergyPreservation(FBxDFEnergyTerms EnergyTerms)
-{
-#if USE_ENERGY_CONSERVATION > 0
-	// Use the directional albedo luminance to avoid color-shift due to metallic specular (for which the energy should be absorbed, not transmitted)
-	#if USE_DEVELOPMENT_SHADERS
-	return View.bShadingEnergyPreservation ? (1 - Luminance(EnergyTerms.E)) : 1.0f;
-	#else
-	return 1 - Luminance(EnergyTerms.E);
-	#endif
+#if USE_ACHROMATIC_BXDF_ENERGY
+#define FBxDFEnergyType  float 	
+#define FBxDFEnergyTerms FBxDFEnergyTermsA
+#define ComputeFresnelEnergyTerms ComputeFresnelEnergyTermsA
+#define ComputeGGXSpecEnergyTerms ComputeGGXSpecEnergyTermsA
+#define ComputeGGXSpecEnergyTerms ComputeGGXSpecEnergyTermsA
+#define ComputeClothEnergyTerms   ComputeClothEnergyTermsA
+#define ComputeClothEnergyTerms   ComputeClothEnergyTermsA
+#define ComputeDiffuseEnergyTerms ComputeDiffuseEnergyTermsA
 #else
-	return 1.0f;
-#endif
-}
-
-// Return the energy conservation weight factor for account energy loss in the BSDF model (i.e. due to micro-facet multiple scattering)
-float3 ComputeEnergyConservation(FBxDFEnergyTerms EnergyTerms)
-{
-	return EnergyTerms.W;
-}
\ No newline at end of file
+#define FBxDFEnergyType float3 
+#define FBxDFEnergyTerms FBxDFEnergyTermsRGB
+#define ComputeFresnelEnergyTerms ComputeFresnelEnergyTermsRGB
+#define ComputeGGXSpecEnergyTerms ComputeGGXSpecEnergyTermsRGB
+#define ComputeGGXSpecEnergyTerms ComputeGGXSpecEnergyTermsRGB
+#define ComputeClothEnergyTerms   ComputeClothEnergyTermsRGB
+#define ComputeClothEnergyTerms   ComputeClothEnergyTermsRGB
+#define ComputeDiffuseEnergyTerms ComputeDiffuseEnergyTermsRGB
+#endif
\ No newline at end of file
diff --git a/Engine/Shaders/Private/ShadingEnergyConservationTemplate.ush b/Engine/Shaders/Private/ShadingEnergyConservationTemplate.ush
new file mode 100644
--- /dev/null
+++ b/Engine/Shaders/Private/ShadingEnergyConservationTemplate.ush
@@ -0,0 +1,135 @@
+// Copyright Epic Games, Inc. All Rights Reserved.
+
+#ifndef BXDF_ENERGY_TYPE
+#error  BXDF_ENERGY_TYPE needs to be defined
+#endif
+
+#ifndef BXDF_ENERGY_SUFFIX
+#error  BXDF_ENERGY_SUFFIX needs to be defined
+#endif
+
+#ifndef IS_BXDF_ENERGY_TYPE_ACHROMATIC
+#error  IS_BXDF_ENERGY_TYPE_ACHROMATIC needs to be defined
+#endif
+
+///////////////////////////////////////////////////////////////////////////////////////////////////
+// Energy terms
+
+struct BXDF_ENERGY_SUFFIX(FBxDFEnergyTerms)
+{
+	BXDF_ENERGY_TYPE W; // overall weight to scale the lobe BxDF by to ensure energy conservation
+	BXDF_ENERGY_TYPE E; // Directional albedo of the lobe for energy preservation and lobe picking
+};
+
+// Given a split-sum approximation of directional albedo for a BxDF, compute multiple scattering weight and multiple scattering directional albedo
+// while taking into account the fresnel term (assumed to be F_Schlick)
+BXDF_ENERGY_SUFFIX(FBxDFEnergyTerms) BXDF_ENERGY_SUFFIX(ComputeFresnelEnergyTerms)(float2 E, float3 InF0, float3 InF90)
+{
+#if IS_BXDF_ENERGY_TYPE_ACHROMATIC
+	BXDF_ENERGY_TYPE F0  = max3(InF0.x, InF0.y, InF0.z);
+	BXDF_ENERGY_TYPE F90 = max3(InF90.x, InF90.y, InF90.z);
+#else
+	BXDF_ENERGY_TYPE F0  = InF0;
+	BXDF_ENERGY_TYPE F90 = InF90;
+#endif
+
+	BXDF_ENERGY_SUFFIX(FBxDFEnergyTerms) Result;
+	// [2] Eq 16: this restores the missing energy of the bsdf, while also accounting for the fact that the fresnel term causes some energy to be absorbed
+	// NOTE: using F0 here is an approximation, but for schlick fresnel Favg is almost exactly equal to F0
+#if USE_DEVELOPMENT_SHADERS
+	Result.W = View.bShadingEnergyConservation ? (1.0 + F0 * ((1 - E.x) / E.x)) : 1.0f;
+#else
+	Result.W = 1.0 + F0 * ((1 - E.x) / E.x);
+#endif
+	// Now estimate the amount of energy reflected off this specular lobe so that we can remove it from underlying BxDF layers (like diffuse)
+	// This relies on the split-sum approximation as in [3] Sec 4.
+	// This term can also be useful to compute the probability of choosing among lobes
+	Result.E = Result.W * (E.x * F0 + E.y * (F90 - F0));
+	return Result;
+}
+
+BXDF_ENERGY_SUFFIX(FBxDFEnergyTerms) BXDF_ENERGY_SUFFIX(ComputeGGXSpecEnergyTerms)(float Roughness, float NoV, float3 F0, float3 F90)
+{
+	BXDF_ENERGY_SUFFIX(FBxDFEnergyTerms) Out;
+#if USE_ENERGY_CONSERVATION > 0
+	{
+		Out = BXDF_ENERGY_SUFFIX(ComputeFresnelEnergyTerms)(GGXEnergyLookup(Roughness, NoV), F0, F90);
+}
+#else
+	{
+		Out.W = 1.0f;
+		Out.E = F0;
+	}
+#endif
+	return Out;
+}
+
+BXDF_ENERGY_SUFFIX(FBxDFEnergyTerms) BXDF_ENERGY_SUFFIX(ComputeGGXSpecEnergyTerms)(float Roughness, float NoV, float3 F0)
+{
+	const float F90 = saturate(50.0 * F0.g); // See F_Schlick implementation
+	return BXDF_ENERGY_SUFFIX(ComputeGGXSpecEnergyTerms)(Roughness, NoV, F0, F90);
+}
+
+BXDF_ENERGY_SUFFIX(FBxDFEnergyTerms) BXDF_ENERGY_SUFFIX(ComputeClothEnergyTerms)(float Roughness, float NoV, float3 F0, float3 F90)
+{
+	BXDF_ENERGY_SUFFIX(FBxDFEnergyTerms) Out;
+#if USE_ENERGY_CONSERVATION > 0
+	{
+		Out = BXDF_ENERGY_SUFFIX(ComputeFresnelEnergyTerms)(ClothEnergyLookup(Roughness, NoV), F0, F90);
+	}
+#else
+	{
+		Out.W = 1.0f;
+		Out.E = F0;
+	}
+#endif
+	return Out;
+}
+
+BXDF_ENERGY_SUFFIX(FBxDFEnergyTerms) BXDF_ENERGY_SUFFIX(ComputeClothEnergyTerms)(float Roughness, float NoV, float3 F0)
+{
+	const float F90 = saturate(50.0 * F0.g); // See F_Schlick implementation
+	return BXDF_ENERGY_SUFFIX(ComputeClothEnergyTerms)(Roughness, NoV, F0, F90);
+}
+
+BXDF_ENERGY_SUFFIX(FBxDFEnergyTerms) BXDF_ENERGY_SUFFIX(ComputeDiffuseEnergyTerms)(float Roughness, float NoV)
+{
+	BXDF_ENERGY_SUFFIX(FBxDFEnergyTerms) Out;
+#if USE_ENERGY_CONSERVATION > 0
+	{
+		Out.E = DiffuseEnergyLookup(Roughness, NoV);
+	}
+#else
+	{
+		Out.E = 1.0f;
+	}
+#endif
+	Out.W = 1.0f;
+	return Out;
+}
+
+///////////////////////////////////////////////////////////////////////////////////////////////////
+// Energy Preservation/Conservation
+
+// Return the energy absorbed by upper layer (e.g., for the specular layer attenuation onto diffuse)
+// Note: Use the directional albedo luminance to avoid color-shift due to metallic specular (for which the energy should be absorbed, not transmitted)
+float ComputeEnergyPreservation(BXDF_ENERGY_SUFFIX(FBxDFEnergyTerms) EnergyTerms)
+{
+#if USE_ENERGY_CONSERVATION > 0
+#if USE_DEVELOPMENT_SHADERS
+	return View.bShadingEnergyPreservation ? (1 - Luminance(EnergyTerms.E)) : 1.0f;
+#else
+	return 1 - Luminance(EnergyTerms.E);
+#endif
+#else
+	return 1.0f;
+#endif
+}
+
+// Return the energy conservation weight factor for account energy loss in the BSDF model (i.e. due to micro-facet multiple scattering)
+BXDF_ENERGY_TYPE ComputeEnergyConservation(BXDF_ENERGY_SUFFIX(FBxDFEnergyTerms) EnergyTerms)
+{
+	return EnergyTerms.W;
+}
+
+#undef BXDF_ENERGY_SUFFIX
\ No newline at end of file
diff --git a/Engine/Shaders/Private/Strata/StrataEvaluation.ush b/Engine/Shaders/Private/Strata/StrataEvaluation.ush
--- a/Engine/Shaders/Private/Strata/StrataEvaluation.ush
+++ b/Engine/Shaders/Private/Strata/StrataEvaluation.ush
@@ -332,9 +332,9 @@ FStrataEvaluateResult StrataEvaluateBSDFCommon(FStrataBSDFContext BSDFContext, F
 
 			// Primary highlight
 			float PDF = 0;
-			float3 DirectionalAlbedo_SpecularTransmission = 1.0f;
+			float DirectionalAlbedo_SpecularTransmission = 1.0f;
 			{
-				float3 MSScale = 1.0f;
+				FBxDFEnergyType MSScale = 1.0f;
 				{
 					FBxDFEnergyTerms EnergyTerms = ComputeGGXSpecEnergyTerms(SafeRoughness, BSDFContext.Context.NoV, F0, F90);
 					DirectionalAlbedo_SpecularTransmission = ComputeEnergyPreservation(EnergyTerms);
@@ -408,12 +408,12 @@ FStrataEvaluateResult StrataEvaluateBSDFCommon(FStrataBSDFContext BSDFContext, F
 
 				if (bIsRectLight)
 				{
-					Sample.SpecularPathValue = SpecularMicroOcclusion * RectLightSpec * MSScale;
+					Sample.SpecularPathValue = (SpecularMicroOcclusion * MSScale) * RectLightSpec;
 				}
 				else
 				{
 					const float3 FresnelTerm = F_Schlick(F0, F90, BSDFContext.Context.VoH);
-					Sample.SpecularPathValue = (SpecularMicroOcclusion * D * Vis) * FresnelTerm * MSScale;
+					Sample.SpecularPathValue = ((SpecularMicroOcclusion * D * Vis) * MSScale) * FresnelTerm;
 				}
 
 				Sample.SpecularPathProbability     = 1.0f;
@@ -442,7 +442,7 @@ FStrataEvaluateResult StrataEvaluateBSDFCommon(FStrataBSDFContext BSDFContext, F
 				}
 
 				FBxDFEnergyTerms EnergyTerms = ComputeGGXSpecEnergyTerms(HazeSafeRoughness, BSDFContext.Context.NoV, HazeF0, HazeF90);
-				float3 HazeMSScale = ComputeEnergyConservation(EnergyTerms);
+				FBxDFEnergyType HazeMSScale = ComputeEnergyConservation(EnergyTerms);
 
 				float HazeD = 0;
 				float HazeVis = 0;
@@ -500,11 +500,11 @@ FStrataEvaluateResult StrataEvaluateBSDFCommon(FStrataBSDFContext BSDFContext, F
 				const float3 HazeFresnelTerm = F_Schlick(HazeF0, HazeF90, BSDFContext.Context.VoH);
 				if (bIsRectLight)
 				{					
-					Sample.SpecularHazePathValue = RectLightSpecHaze * HazeMSScale * HazeSpecularMicroOcclusion;
+					Sample.SpecularHazePathValue = RectLightSpecHaze * (HazeMSScale * HazeSpecularMicroOcclusion);
 				}
 				else
 				{
-					Sample.SpecularHazePathValue = (HazeD * HazeVis * HazeSpecularMicroOcclusion) * HazeFresnelTerm * HazeMSScale;
+					Sample.SpecularHazePathValue = ((HazeD * HazeVis * HazeSpecularMicroOcclusion) * HazeMSScale) * HazeFresnelTerm;
 				}
 
 				// Reminder: 
@@ -731,7 +731,7 @@ FStrataEvaluateResult StrataEvaluateBSDFCommon(FStrataBSDFContext BSDFContext, F
 
 				// Shading point <-> View throughput
 				const float3 SlabTransmittanceV = IsotropicMediumSlabTransmittance(PM, STRATA_SIMPLEVOLUME_THICKNESS_M, BSDFContext.Context.NoV);
-				const float3 SpecularTransmissionV = DirectionalAlbedo_SpecularTransmission;
+				const float SpecularTransmissionV = DirectionalAlbedo_SpecularTransmission;
 
 				// Shading point <-> Light throughput
 				// Compute the transmittance at normal incidence (instead of BSDFContext.Context.NoL), and will compute the final value during evaluation 
@@ -840,8 +840,8 @@ FStrataEvaluateResult StrataEvaluateBSDFCommon(FStrataBSDFContext BSDFContext, F
 			// Apply energy conservation on the diffuse component
 			// If the specular layer is anisotropic, the energy term is computed onto the 'main' roughness [Kulla 2019]
 			//
-			float3 MSScale = 1;
-			float3 DirectionalAlbedo_SpecularTransmission = 1.0f;
+			FBxDFEnergyType MSScale = 1;
+			float DirectionalAlbedo_SpecularTransmission = 1.0f;
 			{
 				FBxDFEnergyTerms EnergyTerms = ComputeGGXSpecEnergyTerms(SafeRoughness, BSDFContext.Context.NoV, F0, F90);
 				DirectionalAlbedo_SpecularTransmission = ComputeEnergyPreservation(EnergyTerms);
@@ -901,12 +901,12 @@ FStrataEvaluateResult StrataEvaluateBSDFCommon(FStrataBSDFContext BSDFContext, F
 
 			if (bIsRectLight)
 			{
-				Sample.SpecularPathValue       = RectLightSpec * MSScale * SpecularMicroOcclusion;
+				Sample.SpecularPathValue       = RectLightSpec * (MSScale * SpecularMicroOcclusion);
 				Sample.IntegratedSpecularValue = (ShadowTerms.SurfaceShadow * /* AreaLightContext.NoL and falloff is part of the LTC integration already */ Sample.SpecularPathProbability) * Sample.SpecularPathValue;
 			}
 			else
 			{
-				Sample.SpecularPathValue        = (D * Vis) * FresnelTerm * MSScale;
+				Sample.SpecularPathValue        = ((D * Vis) * MSScale) * FresnelTerm;
 				Sample.IntegratedSpecularValue += (ShadowTerms.SurfaceShadow * AreaLightContext.NoL * AreaLightContext.Falloff) * AreaLightContext.AreaLight.FalloffColor * Sample.SpecularPathValue * Sample.SpecularPathProbability;
 			}
 
@@ -1268,10 +1268,10 @@ FStrataEnvLightResult StrataEvaluateForEnvLight(FStrataBSDFContext BSDFContext,
 			float SafeRoughness = MakeRoughnessSafe(SLAB_ROUGHNESS(BSDFContext.BSDF));
 
 			float3 EvalEnvBRDF = 0;
-			float3 DirectionalAlbedo_SpecularTransmission = 0;
+			float DirectionalAlbedo_SpecularTransmission = 0;
 			{
-				FBxDFEnergyTerms SpecularEnergyTerms = ComputeGGXSpecEnergyTerms(SafeRoughness, BSDFContext.Context.NoV, F0, F90);
-				FBxDFEnergyTerms DiffuseEnergyTerms  = ComputeDiffuseEnergyTerms(SafeRoughness, BSDFContext.Context.NoV);
+				FBxDFEnergyTermsRGB SpecularEnergyTerms = ComputeGGXSpecEnergyTermsRGB(SafeRoughness, BSDFContext.Context.NoV, F0, F90);
+				FBxDFEnergyTermsRGB DiffuseEnergyTerms  = ComputeDiffuseEnergyTermsRGB(SafeRoughness, BSDFContext.Context.NoV);
 				EvalEnvBRDF									= SpecularEnergyTerms.E;
 				DirectionalAlbedo_SpecularTransmission		= ComputeEnergyPreservation(SpecularEnergyTerms);
 				StrataEnvLightResult.DiffuseWeight			= DiffuseColor * DiffuseEnergyTerms.E * DirectionalAlbedo_SpecularTransmission;
@@ -1364,7 +1364,7 @@ FStrataEnvLightResult StrataEvaluateForEnvLight(FStrataBSDFContext BSDFContext,
 						const float HazeF0 = 0.04f;
 						const float HazeF90= 1.0f;
 						const float HazeSpecularMicroOcclusion = F0RGBToMicroOcclusion(HazeF0);
-						FBxDFEnergyTerms EnergyTerms = ComputeGGXSpecEnergyTerms(HazeSafeRoughness, BSDFContext.Context.NoV, HazeF0, HazeF90);
+						FBxDFEnergyTermsRGB EnergyTerms = ComputeGGXSpecEnergyTermsRGB(HazeSafeRoughness, BSDFContext.Context.NoV, HazeF0, HazeF90);
 						const float3 EvalEnvBRDFHaze = EnergyTerms.E;
 						const float3 TopLayerSpecularTransmittionApprox = saturate(1.0f - EvalEnvBRDFHaze); // Simple and do not use extra LUT for now
 
@@ -1391,7 +1391,7 @@ FStrataEnvLightResult StrataEvaluateForEnvLight(FStrataBSDFContext BSDFContext,
 						StrataEnvLightResult.SSRSpecularWeight *= BlendFactor;
 
 						// Compute the directional albedo for hazy specular
-						FBxDFEnergyTerms EnergyTerms = ComputeGGXSpecEnergyTerms(HazeSafeRoughness, BSDFContext.Context.NoV, F0, F90);
+						FBxDFEnergyTermsRGB EnergyTerms = ComputeGGXSpecEnergyTermsRGB(HazeSafeRoughness, BSDFContext.Context.NoV, F0, F90);
 						const float3 EvalEnvBRDFHaze = EnergyTerms.E;
 						StrataEnvLightResult.SpecularHazeWeight = (lerp(0.0f, HazeWeight, HazinessFadeIn) * SpecularMicroOcclusion) * EvalEnvBRDFHaze;
 						StrataEnvLightResult.SpecularHazeSafeRoughness = HazeSafeRoughness;
@@ -1432,10 +1432,10 @@ FStrataEnvLightResult StrataEvaluateForEnvLight(FStrataBSDFContext BSDFContext,
 				const float3 FuzzF0 = SLAB_FUZZ_COLOR(BSDFContext.BSDF);
 
 				const float ClothF90 = 1.0f;
-				FBxDFEnergyTerms ClothEnergyTerms = ComputeClothEnergyTerms(SafeRoughness, BSDFContext.Context.NoV, FuzzF0, ClothF90);
+				const FBxDFEnergyTermsRGB ClothEnergyTerms = ComputeClothEnergyTermsRGB(SafeRoughness, BSDFContext.Context.NoV, FuzzF0, ClothF90);
 
 				// The specular and diffuse components below the fuzz are only attenuated linearly according to the amount of fuzz.
-				const float3 Cloth_DirectionalAlbedo_SpecularTransmission = lerp(1.0, ComputeEnergyPreservation(ClothEnergyTerms), FuzzAmount);
+				const float Cloth_DirectionalAlbedo_SpecularTransmission = lerp(1.0, ComputeEnergyPreservation(ClothEnergyTerms), FuzzAmount);
 
 				StrataEnvLightResult.DiffuseColor		*= Cloth_DirectionalAlbedo_SpecularTransmission;
 				StrataEnvLightResult.DiffuseWeight		*= Cloth_DirectionalAlbedo_SpecularTransmission;
@@ -1504,7 +1504,7 @@ FStrataEnvLightResult StrataEvaluateForEnvLight(FStrataBSDFContext BSDFContext,
 			float3 EvalEnvBRDF = 0;
 			float3 DirectionalAlbedo_SpecularTransmission = 0;
 			{
-				FBxDFEnergyTerms SpecularEnergyTerms = ComputeGGXSpecEnergyTerms(SafeRoughness, BSDFContext.Context.NoV, F0, F90);
+				FBxDFEnergyTermsRGB SpecularEnergyTerms = ComputeGGXSpecEnergyTermsRGB(SafeRoughness, BSDFContext.Context.NoV, F0, F90);
 				const float DiffuseDirectionalAlbedo = 1.0f;
 				EvalEnvBRDF									= SpecularEnergyTerms.E;
 				DirectionalAlbedo_SpecularTransmission		= ComputeEnergyPreservation(SpecularEnergyTerms);
diff --git a/Engine/Source/Runtime/Engine/Private/ShaderCompiler/ShaderCompiler.cpp b/Engine/Source/Runtime/Engine/Private/ShaderCompiler/ShaderCompiler.cpp
--- a/Engine/Source/Runtime/Engine/Private/ShaderCompiler/ShaderCompiler.cpp
+++ b/Engine/Source/Runtime/Engine/Private/ShaderCompiler/ShaderCompiler.cpp
@@ -6017,6 +6017,10 @@ void GlobalBeginCompileShader(
 			static FShaderPlatformCachedIniValue<int32> CVarStrataShadingQuality(TEXT("r.Strata.ShadingQuality"));
 			const uint32 StrataShadingQuality = CVarStrataShadingQuality.Get(Target.GetPlatform());
 			Input.Environment.SetDefine(TEXT("STRATA_SHADING_QUALITY"), FMath::Max(StrataShadingQuality, 1u));
+			if (StrataShadingQuality > 0)
+			{
+				Input.Environment.SetDefine(TEXT("USE_ACHROMATIC_BXDF_ENERGY"), 1u);
+			}
 		}
 
 		static IConsoleVariable* CVarBackCompatibility = IConsoleManager::Get().FindConsoleVariable(TEXT("r.StrataBackCompatibility"));
