Fix invalid shader type conversion when energy conservation is disable.

diff --git a/Engine/Shaders/Private/ShadingEnergyConservationTemplate.ush b/Engine/Shaders/Private/ShadingEnergyConservationTemplate.ush
--- a/Engine/Shaders/Private/ShadingEnergyConservationTemplate.ush
+++ b/Engine/Shaders/Private/ShadingEnergyConservationTemplate.ush
@@ -15,6 +15,15 @@
 ///////////////////////////////////////////////////////////////////////////////////////////////////
 // Energy terms
 
+BXDF_ENERGY_TYPE BXDF_ENERGY_SUFFIX(GetF0F90)(float3 InF0)
+{
+#if IS_BXDF_ENERGY_TYPE_ACHROMATIC
+	return max3(InF0.x, InF0.y, InF0.z);
+#else
+	return InF0;
+#endif
+}
+
 struct BXDF_ENERGY_SUFFIX(FBxDFEnergyTerms)
 {
 	BXDF_ENERGY_TYPE W; // overall weight to scale the lobe BxDF by to ensure energy conservation
@@ -25,13 +34,8 @@ struct BXDF_ENERGY_SUFFIX(FBxDFEnergyTerms)
 // while taking into account the fresnel term (assumed to be F_Schlick)
 BXDF_ENERGY_SUFFIX(FBxDFEnergyTerms) BXDF_ENERGY_SUFFIX(ComputeFresnelEnergyTerms)(float2 E, float3 InF0, float3 InF90)
 {
-#if IS_BXDF_ENERGY_TYPE_ACHROMATIC
-	BXDF_ENERGY_TYPE F0  = max3(InF0.x, InF0.y, InF0.z);
-	BXDF_ENERGY_TYPE F90 = max3(InF90.x, InF90.y, InF90.z);
-#else
-	BXDF_ENERGY_TYPE F0  = InF0;
-	BXDF_ENERGY_TYPE F90 = InF90;
-#endif
+	BXDF_ENERGY_TYPE F0  = BXDF_ENERGY_SUFFIX(GetF0F90)(InF0);
+	BXDF_ENERGY_TYPE F90 = BXDF_ENERGY_SUFFIX(GetF0F90)(InF90);
 
 	BXDF_ENERGY_SUFFIX(FBxDFEnergyTerms) Result;
 	// [2] Eq 16: this restores the missing energy of the bsdf, while also accounting for the fact that the fresnel term causes some energy to be absorbed
@@ -58,7 +62,7 @@ BXDF_ENERGY_SUFFIX(FBxDFEnergyTerms) BXDF_ENERGY_SUFFIX(ComputeGGXSpecEnergyTerm
 #else
 	{
 		Out.W = 1.0f;
-		Out.E = F0;
+		Out.E = BXDF_ENERGY_SUFFIX(GetF0F90)(F0);
 	}
 #endif
 	return Out;
@@ -80,7 +84,7 @@ BXDF_ENERGY_SUFFIX(FBxDFEnergyTerms) BXDF_ENERGY_SUFFIX(ComputeClothEnergyTerms)
 #else
 	{
 		Out.W = 1.0f;
-		Out.E = F0;
+		Out.E = BXDF_ENERGY_SUFFIX(GetF0F90)(F0);
 	}
 #endif
 	return Out;
