Implement some missing getter functions for DECAL_PRIMITIVE which is set when we are rendering a decal actor.

diff --git a/Engine/Shaders/Private/MaterialTemplate.ush b/Engine/Shaders/Private/MaterialTemplate.ush
--- a/Engine/Shaders/Private/MaterialTemplate.ush
+++ b/Engine/Shaders/Private/MaterialTemplate.ush
@@ -631,17 +631,29 @@ bool4 UnpackUniform_bool4(uint Packed, uint BitOffset)
  */
 MaterialFloat3x3 GetLocalToWorld3x3(FMaterialVertexParameters Parameters)
 {
+#if DECAL_PRIMITIVE
+	return (float3x3)DecalToWorld;
+#else
 	return LWCToFloat3x3(GetPrimitiveData(Parameters).LocalToWorld);
+#endif
 }
 
 MaterialFloat3x3 GetPreviousLocalToWorld3x3(FMaterialVertexParameters Parameters)
 {
+#if DECAL_PRIMITIVE
+	return (float3x3)DecalToWorld;
+#else
 	return LWCToFloat3x3(GetPrimitiveData(Parameters).LocalToWorld);
+#endif
 }
 
 MaterialFloat3x3 GetLocalToWorld3x3(FMaterialPixelParameters Parameters)
 {
+#if DECAL_PRIMITIVE
+	return (float3x3)DecalToWorld;
+#else
 	return LWCToFloat3x3(GetPrimitiveData(Parameters).LocalToWorld);
+#endif
 }
 
 MaterialFloat3x3 GetLocalToWorld3x3()
@@ -1037,7 +1049,7 @@ MaterialFloat2 GetDefaultSceneTextureUV(FMaterialPixelParameters Parameters, con
 	FLWCVector3 GetActorWorldPosition(FMaterialVertexParameters Parameters)
 	{
 	#if DECAL_PRIMITIVE
-		return MakeLWCVector3(DecalTilePosition, DecalToWorld[3].xyz);
+		return MakeLWCVector3(DecalTilePosition.xyz, DecalToWorld[3].xyz);
 	#else
 		return GetPrimitiveData(Parameters).ActorWorldPosition;
 	#endif
@@ -1046,7 +1058,7 @@ MaterialFloat2 GetDefaultSceneTextureUV(FMaterialPixelParameters Parameters, con
 	FLWCVector3 GetActorWorldPosition(FMaterialPixelParameters Parameters)
 	{
 	#if DECAL_PRIMITIVE
-		return MakeLWCVector3(DecalTilePosition, DecalToWorld[3].xyz);
+		return MakeLWCVector3(DecalTilePosition.xyz, DecalToWorld[3].xyz);
 	#else
 		return GetPrimitiveData(Parameters).ActorWorldPosition;
 	#endif
@@ -1173,7 +1185,9 @@ MaterialFloat3 TransformTangentVectorToView(FMaterialPixelParameters Parameters,
 
 FLWCMatrix GetLocalToWorld(FMaterialVertexParameters Parameters)
 {
-#if USE_INSTANCING || USE_INSTANCE_CULLING || IS_MESHPARTICLE_FACTORY
+#if DECAL_PRIMITIVE
+	return MakeLWCMatrix(DecalTilePosition.xyz, DecalToWorld);
+#elif USE_INSTANCING || USE_INSTANCE_CULLING || IS_MESHPARTICLE_FACTORY
 	return Parameters.InstanceLocalToWorld;
 #else
 	return GetPrimitiveData(Parameters).LocalToWorld;
@@ -1182,7 +1196,11 @@ FLWCMatrix GetLocalToWorld(FMaterialVertexParameters Parameters)
 
 FLWCMatrix GetLocalToWorld(FMaterialPixelParameters Parameters)
 {
+#if DECAL_PRIMITIVE
+	return MakeLWCMatrix(DecalTilePosition.xyz, DecalToWorld);
+#else
 	return GetPrimitiveData(Parameters).LocalToWorld;
+#endif
 }
 
 FLWCMatrix GetPrevLocalToWorld(FMaterialVertexParameters Parameters)
@@ -1216,55 +1234,45 @@ MaterialFloat3 TransformLocalVectorToPrevWorld(FMaterialVertexParameters Paramet
 #endif
 }
 
-#if HAS_PRIMITIVE_UNIFORM_BUFFER
-
 /** Transforms a position from local space to absolute world space */
 FLWCVector3 TransformLocalPositionToWorld(FMaterialPixelParameters Parameters,float3 InLocalPosition)
 {
-	return LWCMultiply(InLocalPosition, GetPrimitiveData(Parameters).LocalToWorld);
+	return LWCMultiply(InLocalPosition, GetLocalToWorld(Parameters));
 }
 
 /** Transforms a position from local space to absolute world space */
 FLWCVector3 TransformLocalPositionToWorld(FMaterialVertexParameters Parameters,float3 InLocalPosition)
 {
-	#if USE_INSTANCING || USE_INSTANCE_CULLING || IS_MESHPARTICLE_FACTORY
-		return LWCMultiply(InLocalPosition, Parameters.InstanceLocalToWorld);
-	#else
-		return LWCMultiply(InLocalPosition, GetPrimitiveData(Parameters).LocalToWorld);
-	#endif
+	return LWCMultiply(InLocalPosition, GetLocalToWorld(Parameters));
 }
 
 /** Transforms a position from local space to previous frame absolute world space */
 FLWCVector3 TransformLocalPositionToPrevWorld(FMaterialVertexParameters Parameters,float3 InLocalPosition)
 {
-#if USE_INSTANCING || USE_INSTANCE_CULLING || IS_MESHPARTICLE_FACTORY
-	return LWCMultiply(InLocalPosition, Parameters.PrevFrameLocalToWorld);
-#else
-	return LWCMultiply(InLocalPosition, GetPrimitiveData(Parameters).PreviousLocalToWorld);
-#endif
+	return LWCMultiply(InLocalPosition, GetLocalToWorld(Parameters));
 }
 
-#endif
-
-#if HAS_PRIMITIVE_UNIFORM_BUFFER
-
 /** Return the object's position in world space */
 FLWCVector3 GetObjectWorldPosition(FMaterialPixelParameters Parameters)
 {
+#if DECAL_PRIMITIVE
+	return MakeLWCVector3(DecalTilePosition.xyz, DecalToWorld[3].xyz);
+#else
 	return GetPrimitiveData(Parameters).ObjectWorldPosition;
+#endif
 }
 
 /** Return the object's position in world space. For instanced meshes, this returns the instance position. */
 FLWCVector3 GetObjectWorldPosition(FMaterialVertexParameters Parameters)
 {
-	#if USE_INSTANCING || USE_INSTANCE_CULLING || IS_MESHPARTICLE_FACTORY
-		return LWCGetOrigin(Parameters.InstanceLocalToWorld);
-	#else
-		return GetPrimitiveData(Parameters).ObjectWorldPosition;
-	#endif
-}
-
+#if DECAL_PRIMITIVE
+	return MakeLWCVector3(DecalTilePosition.xyz, DecalToWorld[3].xyz);
+#elif USE_INSTANCING || USE_INSTANCE_CULLING || IS_MESHPARTICLE_FACTORY
+	return LWCGetOrigin(Parameters.InstanceLocalToWorld);
+#else
+	return GetPrimitiveData(Parameters).ObjectWorldPosition;
 #endif
+}
 
 /** Get the per-instance random value when instancing */
 float GetPerInstanceRandom(FMaterialVertexParameters Parameters)
diff --git a/Engine/Source/Runtime/Engine/Private/Materials/MaterialExpressions.cpp b/Engine/Source/Runtime/Engine/Private/Materials/MaterialExpressions.cpp
--- a/Engine/Source/Runtime/Engine/Private/Materials/MaterialExpressions.cpp
+++ b/Engine/Source/Runtime/Engine/Private/Materials/MaterialExpressions.cpp
@@ -12472,11 +12472,6 @@ UMaterialExpressionObjectPositionWS::UMaterialExpressionObjectPositionWS(const F
 #if WITH_EDITOR
 int32 UMaterialExpressionObjectPositionWS::Compile(class FMaterialCompiler* Compiler, int32 OutputIndex)
 {
-	if (Material && Material->MaterialDomain == MD_DeferredDecal)
-	{
-		return CompilerError(Compiler, TEXT("Expression not available in the deferred decal material domain."));
-	}
-
 	return Compiler->ObjectWorldPosition();
 }
 
