Strata - fixed single layer water specular unpacking to use an expected 20bit mask.

diff --git a/Engine/Shaders/Private/SingleLayerWaterComposite.usf b/Engine/Shaders/Private/SingleLayerWaterComposite.usf
--- a/Engine/Shaders/Private/SingleLayerWaterComposite.usf
+++ b/Engine/Shaders/Private/SingleLayerWaterComposite.usf
@@ -101,7 +101,7 @@ SingleLayerWaterCompositeOutput SingleLayerWaterComposite(float2 BufferUV, float
 			const float3 WorldNormal = StrataGetBSDFSharedBasis(StrataPixelHeader, SLWaterBSDF, StrataAddressing)[2];
 			const float3 F0 = ComputeF0(SLW_SPECULAR(SLWaterBSDF), SLW_BASECOLOR(SLWaterBSDF), SLW_METALLIC(SLWaterBSDF));
 			const float Roughness = SLW_ROUGHNESS(SLWaterBSDF);
-			const float SafeRoughness = MakeRoughnessSafe(SLW_ROUGHNESS(SLWaterBSDF));
+			const float SafeRoughness = MakeRoughnessSafe(Roughness);
 
 			if (UseSeparatedMainDirLightTexture > 0.0f)
 			{
diff --git a/Engine/Shaders/Private/Strata/Strata.ush b/Engine/Shaders/Private/Strata/Strata.ush
--- a/Engine/Shaders/Private/Strata/Strata.ush
+++ b/Engine/Shaders/Private/Strata/Strata.ush
@@ -3623,7 +3623,7 @@ FStrataBSDF UnpackFastWaterPathStrataBSDFIn(FStrataMaterialContainer StrataBuffe
 	const uint RoughnessAndUseSeparateDirLight8Bits	= 0xFF   & (Data0 >> HEADER_SLWENCODING_BIT_COUNT);
 	
 	const uint PackedBaseColor12Bits				= 0xFFF  & (Data1 >> 20);
-	const uint PackedOpacityMetalSpec20Bits			= 0xFFFF & Data1;
+	const uint PackedOpacityMetalSpec20Bits			= 0xFFFFF & Data1;
 	const uint PackedBaseColor20Bits				= PackedBaseColor12Bits | (PackedBaseColor8Bits << 12);
 
 	const bool bStoreSeparatedMainDirLight			= RoughnessAndUseSeparateDirLight8Bits & 0x1;
