Strata - fixed simple clear coat throughput being applied twi on specular of bottom layer.

diff --git a/Engine/Shaders/Private/Strata/StrataEvaluation.ush b/Engine/Shaders/Private/Strata/StrataEvaluation.ush
--- a/Engine/Shaders/Private/Strata/StrataEvaluation.ush
+++ b/Engine/Shaders/Private/Strata/StrataEvaluation.ush
@@ -518,20 +518,17 @@ FStrataEvaluateResult StrataEvaluateBSDFCommon(FStrataBSDFContext BSDFContext, F
 
 					const float TopLayerCoverage = HazeWeight;
 					const float TopLayerSpecularTransmittionApprox = saturate(1.0f - HazeSpecularMicroOcclusion * HazeFresnelTerm.x); // Simple and do not use the LUT for now. With bHazeAsSimpleClearCoat, Fresnel is achromatic
-					const float3 BottomLayerFactors = lerp(1.0f, HazeClearCoatTransmittance * TopLayerSpecularTransmittionApprox, TopLayerCoverage);
+					const float3 TopLayerThrouput = lerp(1.0f, HazeClearCoatTransmittance * TopLayerSpecularTransmittionApprox, TopLayerCoverage);
 
-					// Apply the top clear coat throughput onto the bottom layer
-					Sample.SpecularPathValue *= BottomLayerFactors;
-
-					const float BottomLayerFactorsGrey = dot(BottomLayerFactors, (1.0 / 3.0).xxx);
-					Sample.SpecularPathProbability = BottomLayerFactorsGrey / (TopLayerCoverage + BottomLayerFactorsGrey);
-					Sample.SpecularHazePathProbability = TopLayerCoverage / (TopLayerCoverage + BottomLayerFactorsGrey);
+					const float TopLayerThrouputGrey = dot(TopLayerThrouput, (1.0 / 3.0).xxx);
+					Sample.SpecularPathProbability = TopLayerThrouputGrey / (TopLayerCoverage + TopLayerThrouputGrey);
+					Sample.SpecularHazePathProbability = TopLayerCoverage / (TopLayerCoverage + TopLayerThrouputGrey);
 
 					// Clear coat applied over bottom layer diffuse component
-					Sample.DiffusePathValue *= BottomLayerFactors;
-					Sample.IntegratedDiffuseValue *= BottomLayerFactors;
-					// And to the bottom layer specular on top of whiche we add the top layer specular
-					Sample.IntegratedSpecularValue = Sample.SpecularPathValue * BottomLayerFactors + Sample.SpecularHazePathValue * TopLayerCoverage;
+					Sample.DiffusePathValue *= TopLayerThrouput;
+					Sample.IntegratedDiffuseValue *= TopLayerThrouput;
+					// And to the bottom layer specular affected by top layer throughput, on top of which we add the top layer specular
+					Sample.IntegratedSpecularValue = Sample.SpecularPathValue * TopLayerThrouput + Sample.SpecularHazePathValue * TopLayerCoverage;
 				}
 				else
 				{
