Strata - fixed out of range coverage weight.

diff --git a/Engine/Shaders/Private/BasePassPixelShader.usf b/Engine/Shaders/Private/BasePassPixelShader.usf
--- a/Engine/Shaders/Private/BasePassPixelShader.usf
+++ b/Engine/Shaders/Private/BasePassPixelShader.usf
@@ -1849,7 +1849,7 @@ void FPixelShaderInOut_MainPS(
 
 		#if STRATA_BLENDING_COLOREDTRANSMITTANCEONLY
 			// RETURN_COLOR not needed with modulative blending
-			half3 FoggedColor = lerp(float3(1, 1, 1), DualBlendSurfaceTransmittancePreCoverage, Fogging.aaa * DualBlendSurfaceCoverage);
+			half3 FoggedColor = lerp(float3(1, 1, 1), DualBlendSurfaceTransmittancePreCoverage, Fogging.aaa* DualBlendSurfaceCoverage);
 			Out.MRT[0] = half4(FoggedColor, 1.0f);
 		#elif STRATA_BLENDING_ALPHAHOLDOUT
 			Out.MRT[0] = half4(0.0f.xxx, Opacity);
@@ -2017,7 +2017,7 @@ void FPixelShaderInOut_MainPS(
 		#endif
 	}
 
-#if !MATERIALBLENDING_MODULATE
+#if !MATERIALBLENDING_MODULATE && !STRATA_BLENDING_COLOREDTRANSMITTANCEONLY
 #if MATERIAL_IS_SKY
 	// Dynamic capture exposure is 1 as of today.
 	const float ViewPreExposure = View.RealTimeReflectionCapture>0.0f ? View.RealTimeReflectionCapturePreExposure : View.PreExposure;
diff --git a/Engine/Shaders/Private/Strata/Strata.ush b/Engine/Shaders/Private/Strata/Strata.ush
--- a/Engine/Shaders/Private/Strata/Strata.ush
+++ b/Engine/Shaders/Private/Strata/Strata.ush
@@ -1970,14 +1970,14 @@ FStrataData StrataWeight(FStrataData A, float Weight, inout FStrataTree StrataTr
 	FStrataData StrataData = GetInitialisedStrataData();
 	StrataData.OperatorIndex = OperatorIndex;
 #if USE_DEVELOPMENT_SHADERS
-	StrataData.PreviewColor = A.PreviewColor * Weight;
+	StrataData.PreviewColor = A.PreviewColor * SafeWeight;
 #endif
 
 	StrataTree.OperatorCount++;
 	StrataTree.Operators[OperatorIndex]							= GetInitialisedStrataOperator();
 	StrataTree.Operators[OperatorIndex].MaxDistanceFromLeaves	= MaxDistanceFromLeaves;
 	StrataTree.Operators[OperatorIndex].Type					= STRATA_OPERATOR_WEIGHT;
-	StrataTree.Operators[OperatorIndex].Weight					= Weight;
+	StrataTree.Operators[OperatorIndex].Weight					= SafeWeight;
 	StrataTree.Operators[OperatorIndex].LeftIndex				= A.OperatorIndex;
 
 	StrataTree.Operators[A.OperatorIndex].ParentIndex			= OperatorIndex;
