Implemented support for multisampling in the software rasterizer for hair. Fixed some bugs in the hair tile compaction pass and reduced the number of atomics using wave ops. Fixed/reduced tile shaped artifacts in the SW rasterizer that happened due to clipping line segments to the tile. Added WaveInterlockedMin()/WaveInterlockedMax() function macros to WaveOpUtil.ush.

#rnx

diff --git a/Engine/Shaders/Private/HairStrands/HairStrandsVisibilityCompactionComputeRaster.usf b/Engine/Shaders/Private/HairStrands/HairStrandsVisibilityCompactionComputeRaster.usf
--- a/Engine/Shaders/Private/HairStrands/HairStrandsVisibilityCompactionComputeRaster.usf
+++ b/Engine/Shaders/Private/HairStrands/HairStrandsVisibilityCompactionComputeRaster.usf
@@ -14,6 +14,10 @@
 #define TILE_PIXEL_SIZE_Y 8
 #endif
 
+#ifndef PERMUTATION_MULTI_SAMPLE_COUNT
+#define PERMUTATION_MULTI_SAMPLE_COUNT 1
+#endif
+
 ///////////////////////////////////////////////////////////////////////////////////////////////////////////////
 
 //TODO: re-enable support for 64 bit visibilty buffer format?
@@ -39,6 +43,7 @@ int2 OutputResolution;
 uint MaxNodeCount;
 uint SamplerPerPixel;
 float CoverageThreshold;
+uint bSortSampleByDepth;
 
 RWTexture2D<uint> 					OutCompactNodeCounter;
 RWTexture2D<uint> 					OutCompactNodeIndex;
@@ -58,11 +63,11 @@ struct FSampleSetDesc
 };
 
 ///////////////////////////////////////////////////////////////////////////////////////////////////////////////
-// Compute raster visibility buffer
+// Compute raster visibility buffer compaction
 ///////////////////////////////////////////////////////////////////////////////////////////////////////////////
 
-Texture2D<uint>  DepthCovTexture;
-Texture2D<uint>  PrimMatTexture;
+Texture2DArray<uint>  DepthCovTexture;
+Texture2DArray<uint>  PrimMatTexture;
 Texture2D<uint>  HairCountTexture;
 
 int2 TileCountXY;
@@ -70,12 +75,98 @@ uint TileSize;
 Buffer<uint> TileCountBuffer;
 Buffer<uint2> TileDataBuffer;
 
+#define MERGE_SAMPLE 0
+
+void ComputeUniqueSamples(const uint2 PixelCoord, out uint4 OutSamples[PERMUTATION_MULTI_SAMPLE_COUNT], out FSampleSetDesc OutSet)
+{
+	OutSet.UniqueSampleCount = 0;
+	OutSet.ValidSampleCount = 0;
+	OutSet.HairSampleCount = PERMUTATION_MULTI_SAMPLE_COUNT;
+
+	for (uint SampleIt = 0; SampleIt < OutSet.HairSampleCount; ++SampleIt)
+	{
+		// Note: PrimMatTexture contains both the primitive ID and the material ID. However 
+		// the material ID is constant along the primitive, so it is correct to use this as a 
+		// sorting/deduplication key
+		const uint PrimitiveId = PrimMatTexture.Load(uint4(PixelCoord, SampleIt, 0));
+		const bool bIsValid = PrimitiveId != GetInvalidHairPrimitiveId();
+		if (!bIsValid)
+		{
+			continue;
+		}
+
+		const float SampleDepth = UnpackHairVisDepth(DepthCovTexture.Load(uint4(PixelCoord, SampleIt, 0)));
+
+		++OutSet.ValidSampleCount;
+#if MERGE_SAMPLE 
+		const float SceneDepth = ConvertFromDeviceZ(SampleDepth);
+#endif
+		bool bAlreadyExist = false;
+		for (uint UniqueIt = 0; UniqueIt < OutSet.UniqueSampleCount; ++UniqueIt)
+		{
+
+#if MERGE_SAMPLE
+			const float UniqueDepth = asfloat(OutSamples[UniqueIt].w);
+			const float UniqueSceneDepth = ConvertFromDeviceZ(UniqueDepth);
+			const bool bIsSimilar =
+				PrimitiveId == OutSamples[UniqueIt].x ||
+				abs(UniqueSceneDepth - SceneDepth) < DepthTheshold;
+#else
+			const bool bIsSimilar = PrimitiveId == OutSamples[UniqueIt].x;
+#endif
+			if (bIsSimilar)
+			{
+				OutSamples[UniqueIt].y += 1;
+
+				// Update the unique sample with the closest depth
+				const uint IntDepth = asuint(SampleDepth);
+				if (IntDepth > OutSamples[UniqueIt].w)
+				{
+#if MERGE_SAMPLE
+					OutSamples[UniqueIt].x = PrimitiveId;
+#endif
+					OutSamples[UniqueIt].z = SampleIt;
+					OutSamples[UniqueIt].w = asuint(SampleDepth);
+				}
+
+				bAlreadyExist = true;
+				break;
+			}
+		}
+
+		if (!bAlreadyExist)
+		{
+			OutSamples[OutSet.UniqueSampleCount].x = PrimitiveId;
+			OutSamples[OutSet.UniqueSampleCount].y = 1;
+			OutSamples[OutSet.UniqueSampleCount].z = SampleIt;
+			OutSamples[OutSet.UniqueSampleCount].w = asuint(SampleDepth);
+			++OutSet.UniqueSampleCount;
+		}
+	}
+
+	// Sort sample from closer to further. This is used later for updating sample coverage 
+	// based on ordered transmittance. See HairStrandsVisibilityComputeSampleCoverage.usf for more details.
+	if (bSortSampleByDepth > 0)
+	{
+		for (uint i = 0; i < OutSet.UniqueSampleCount; ++i)
+		{
+			const uint DepthI = OutSamples[i].w;
+			for (uint j = i + 1; j < OutSet.UniqueSampleCount; ++j)
+			{
+				const uint DepthJ = OutSamples[j].w;
+				if (DepthJ > DepthI)
+				{
+					uint4 Temp = OutSamples[i];
+					OutSamples[i] = OutSamples[j];
+					OutSamples[j] = Temp;
+				}
+			}
+		}
+	}
+}
+
 [numthreads(TILE_PIXEL_SIZE_X, TILE_PIXEL_SIZE_Y, 1)]
-void MainCS(
-	uint  GroupIndex : SV_GroupIndex,
-	uint3 GroupId : SV_GroupID,
-	uint3 GroupThreadId : SV_GroupThreadID,
-	uint3 DispatchThreadId : SV_DispatchThreadID)
+void MainCS(uint GroupIndex : SV_GroupIndex, uint3 GroupId : SV_GroupID, uint3 GroupThreadId : SV_GroupThreadID, uint3 DispatchThreadId : SV_DispatchThreadID)
 {
 	if (GroupIndex == 0)
 	{
@@ -93,33 +184,17 @@ void MainCS(
 	uint2 PixelCoord = TileCoord * TileSize + GroupThreadId.xy;
 
 	if (PixelCoord.x >= uint(OutputResolution.x) || PixelCoord.y >= uint(OutputResolution.y))
+	{
 		PixelCoord = uint2(0, 0);
+	}
 
 	FSampleSetDesc SampleDesc;
+	uint4 Samples[PERMUTATION_MULTI_SAMPLE_COUNT]; // x:PrimitiveId|MaterialId, y:Weight, z:SampleIt, w:Depth (as uint)
 
-	const uint VisData0 = DepthCovTexture.Load(uint3(PixelCoord, 0));
-	const uint VisData1 = PrimMatTexture.Load(uint3(PixelCoord, 0));
-
-	uint2  Vis_PrimId_MatId[4];
-	float2 Vis_Depth_Cov[4];
-	float  TotalCoverage = 0;
-	uint   UniqueSampleCount = 0;
-
-	if (VisData1)
-	{
-		const uint VisSamplePrimitiveId = UnpackHairVisPrimitiveId(VisData1);
-		const uint VisSampleMaterialId = UnpackHairVisMaterialId(VisData1);
-		const float2 Depth_Cov = float2(UnpackHairVisDepth(VisData0), UnpackHairVisCoverage(VisData0));
-
-		Vis_PrimId_MatId[UniqueSampleCount] = uint2(VisSamplePrimitiveId, VisSampleMaterialId);
-		Vis_Depth_Cov[UniqueSampleCount] = Depth_Cov;
-
-		TotalCoverage += Vis_Depth_Cov[UniqueSampleCount].y;
-		UniqueSampleCount++;
-	}
+	ComputeUniqueSamples(PixelCoord, Samples, SampleDesc);
 
 	FNodeDesc NodeDesc;
-	NodeDesc.Count = UniqueSampleCount;
+	NodeDesc.Count = SampleDesc.UniqueSampleCount;
 	NodeDesc.Offset = 0;
 
 	if (NodeDesc.Count > 0)
@@ -146,20 +221,20 @@ void MainCS(
 			{
 				// VisibilityData.Coverage8bit is a weight normalising to 1 the contribution of all the compacted samples. Because later it is weighted by Categorization.PixelCoverage.
 				FHairVis OutNodeVis;
-				OutNodeVis.PrimitiveId	= Vis_PrimId_MatId[OutIndex].x;
-				OutNodeVis.MaterialId	= Vis_PrimId_MatId[OutIndex].y;
-				OutNodeVis.Depth		= Vis_Depth_Cov[OutIndex].x;
-				OutNodeVis.Coverage8bit = 0xff;
+				OutNodeVis.PrimitiveId	= UnpackHairVisPrimitiveId(Samples[OutIndex].x);
+				OutNodeVis.MaterialId	= UnpackHairVisMaterialId(Samples[OutIndex].x);
+				OutNodeVis.Depth		= asfloat(Samples[OutIndex].w);
+				OutNodeVis.Coverage8bit = To8bitCoverage(Samples[OutIndex].y / float(SampleDesc.ValidSampleCount)); // 0xff;
 				const uint StoreIndex = NodeDesc.Offset + OutIndex;
 				OutCompactNodeVis[StoreIndex] = PackHairVis(OutNodeVis);
 				OutCompactNodeCoord[StoreIndex] = PixelCoord;
 			}
 
 			// skips separate count to transmittance pass  
-			float LogViewTransmittance = -(HairCountTexture.Load(uint3(PixelCoord, 0)) / 1000.0);
-			float ViewTransmittance = pow(2.71828, LogViewTransmittance);
-			float CoverageThreshold = 0.98;
-			PixelCoverage = saturate(min(1.0 - ViewTransmittance, 1.0) / CoverageThreshold);
+			float LogViewTransmittance = -(HairCountTexture.Load(uint3(PixelCoord, 0)) / 1000.0f);
+			float ViewTransmittance = pow(2.71828f, LogViewTransmittance);
+			float CoverageThreshold = 0.98f;
+			PixelCoverage = TransmittanceToCoverage(ViewTransmittance, CoverageThreshold);
 		}
 	}
 
diff --git a/Engine/Shaders/Private/HairStrands/HairStrandsVisibilityRasterCompute.usf b/Engine/Shaders/Private/HairStrands/HairStrandsVisibilityRasterCompute.usf
--- a/Engine/Shaders/Private/HairStrands/HairStrandsVisibilityRasterCompute.usf
+++ b/Engine/Shaders/Private/HairStrands/HairStrandsVisibilityRasterCompute.usf
@@ -3,6 +3,7 @@
 #define HAIR_STRANDS_PARAMETERS 1
 
 #include "../Common.ush"
+#include "../WaveOpUtil.ush"
 #include "HairStrandsClusterCommon.ush"
 #include "HairStrandsVertexFactoryCommon.ush"
 #include "HairStrandsVisibilityCommon.ush"
@@ -14,6 +15,9 @@ uint	TileSizeAsShift;
 uint	TileSize;
 float	RcpTileSize;
 uint	SqrTileSize;
+uint	HalfTileSize;
+float	RcpHalfTileSize;
+uint	SqrHalfTileSize;
 int2	TileRes;
 
 uint	NumBinners;
@@ -29,7 +33,7 @@ float2	OutputResolutionf;
 
 ///////////////////////////////////////////////////////////////////////////
 
-#if SHADER_RASTERCOMPUTE_BINNING || SHADER_RASTERCOMPUTE_COMPACTION || SHADER_RASTERCOMPUTE_RASTER || SHADER_RASTERCOMPUTE_DEPTH_GRID
+#if SHADER_RASTERCOMPUTE_BINNING || SHADER_RASTERCOMPUTE_COMPACTION || SHADER_RASTERCOMPUTE_RASTER || SHADER_RASTERCOMPUTE_RASTER_MULTI_SAMPLE || SHADER_RASTERCOMPUTE_DEPTH_GRID
 
 ///////////////////////////////////////////////////////////////////////////
 
@@ -93,7 +97,7 @@ void CalcHomogenousPos(in uint4 PB, in float3 PBO, out float4 HP, out uint Type)
 	Type = CP.Type;
 }
 
-void CalcHomogenousPosAndRad(in uint4 PB, in float3 PBO, out float4 HP, out float Rad)
+void CalcHomogenousPosAndRad(in uint4 PB, in float3 PBO, out float4 HP, out float Rad, out uint Type)
 {
 	const FHairControlPoint CP = UnpackHairControlPoint(
 		PB,
@@ -105,6 +109,31 @@ void CalcHomogenousPosAndRad(in uint4 PB, in float3 PBO, out float4 HP, out floa
 	const float3 WP = mul(float4(CP.Position, 1.0f), HairStrandsVF_LocalToWorldPrimitiveTransform).xyz;
 	HP = mul(float4(WP, 1.0f), LWCHackToFloat(PrimaryView.WorldToClip));
 	Rad = CP.WorldRadius * 2000.0; // OutputResolutionf.x; //TODO: figure this out correctly?
+	Type = CP.Type;
+}
+
+float ComputeLerpAlpha(int2 Coord, float2 P0, float2 P1, float SegmentLenSqRcp)
+{
+	// Project P onto line segment and compute the lerp alpha between P0 and P1
+	// Simplification of:
+	// A = P - P0
+	// B = P1 - P0
+	// Alpha = dot(A, B) / dot(B, B)
+	const float2 P = Coord + 0.5f;
+	const float Alpha = saturate(dot(P - P0, P1 - P0) * SegmentLenSqRcp);
+	return Alpha;
+}
+
+float ComputePerspectiveCorrectRadius(float Rad0, float Rad1, float Alpha, float RcpW0, float RcpW1)
+{
+	// Alpha value for perspective correct interpolation. We store the reciprocal of w in the w component of P0 and P1,
+	// so this is a simplification of:
+	// (Alpha / w1) / ((1 - Alpha) / w0 + Alpha / w1)
+	const float LerpedRcpW = lerp(RcpW0, RcpW1, Alpha);
+	const float PerspectiveAlpha = (Alpha * RcpW1) / LerpedRcpW;
+	// Divide by W to make thickness dependent on screen space depth? This division was kept from the previous line rasterization algorithm.
+	const float Rad = lerp(Rad0, Rad1, PerspectiveAlpha) * LerpedRcpW;
+	return Rad;
 }
 
 // Line clipping based on "CLIPPING USING HOMOGENEOUS COORDINATES" by Blinn et al.
@@ -154,8 +183,10 @@ bool BlinnLineClipping(inout float4 P0, inout float4 P1)
 	return !bIsRemoved;
 }
 
-bool ClipRaySegment(float2 AABBMin, float2 AABBMax, inout float4 P0, inout float4 P1)
+bool ClipRaySegment(float2 AABBMin, float2 AABBMax, float4 P0, float4 P1, out float2 T, out bool2 bClipped)
 {
+	bClipped = false;
+	T = float2(0.0f, 1.0f);
 	const bool bP0Outside = any(P0.xy < AABBMin) || any(P0.xy > AABBMax);
 	const bool bP1Outside = any(P1.xy < AABBMin) || any(P1.xy > AABBMax);
 	if (!bP0Outside && !bP1Outside)
@@ -170,35 +201,61 @@ bool ClipRaySegment(float2 AABBMin, float2 AABBMax, inout float4 P0, inout float
 	const float2 T0 = (AABBMin - Origin) * RcpDir;
 	const float2 T1 = (AABBMax - Origin) * RcpDir;
 
-	const float TMin = max(min(T0.x, T1.x), min(T0.y, T1.y));
-	const float TMax = min(max(T0.x, T1.x), max(T0.y, T1.y));
+	T.x = max(min(T0.x, T1.x), min(T0.y, T1.y));
+	T.y = min(max(T0.x, T1.x), max(T0.y, T1.y));
 
-	// Ray intersects the AABB but the segment is completely outside.
-	if (TMax < 0.0f)
+	// Ray intersects the AABB but the segment is completely outside or no intersection at all.
+	if (T.y < 0.0f || T.x > T.y)
 	{
+		bClipped = true;
 		return false;
 	}
 
-	// No intersection
-	if (TMin > TMax)
+	if (bP0Outside && T.x > 0.0f && T.x < 1.0f)
 	{
-		return false;
+		bClipped.x = true;
 	}
-
-	float4 P0New = P0;
-	float4 P1New = P1;
-	if (bP0Outside && TMin > 0.0f && TMin < 1.0f)
+	if (bP1Outside && T.y > 0.0f && T.y < 1.0f)
 	{
-		P0New = lerp(P0, P1, TMin);
+		bClipped.y = true;
 	}
-	if (bP1Outside && TMax > 0.0f && TMax < 1.0f)
+
+	return true;
+}
+
+bool ClipRaySegment(float2 AABBMin, float2 AABBMax, inout float4 P0, inout float4 P1, inout float Rad0, inout float Rad1, out bool2 bClipped)
+{
+	float2 T;
+	bool bIsValid = ClipRaySegment(AABBMin, AABBMax, P0, P1, T, bClipped);
+	
+	if (bIsValid)
 	{
-		P1New = lerp(P0, P1, TMax);
+		const bool bP0Outside = any(P0.xy < AABBMin) || any(P0.xy > AABBMax);
+		const bool bP1Outside = any(P1.xy < AABBMin) || any(P1.xy > AABBMax);
+
+		float4 P0New = P0;
+		float4 P1New = P1;
+		float Rad0New = Rad0;
+		float Rad1New = Rad1;
+		if (bP0Outside && T.x > 0.0f && T.x < 1.0f)
+		{
+			P0New = lerp(P0, P1, T.x);
+			Rad0New = lerp(Rad0, Rad1, T.x);
+			bClipped.x = true;
+		}
+		if (bP1Outside && T.y > 0.0f && T.y < 1.0f)
+		{
+			P1New = lerp(P0, P1, T.y);
+			Rad1New = lerp(Rad0, Rad1, T.y);
+			bClipped.y = true;
+		}
+		P0 = P0New;
+		P1 = P1New;
+		Rad0 = Rad0New;
+		Rad1 = Rad1New;
 	}
-	P0 = P0New;
-	P1 = P1New;
 
-	return true;
+	return bIsValid;
 }
 
 #endif // Common rasetrizer helper function & parameters
@@ -207,50 +264,54 @@ bool ClipRaySegment(float2 AABBMin, float2 AABBMax, inout float4 P0, inout float
 
 #if SHADER_RASTERCOMPUTE_DEPTH_GRID
 
-RWTexture2D<uint>	OutVisTileDepthGrid;
-RWTexture2D<uint>	OutDepthCovTexture;
+RWTexture2D<uint>			OutVisTileDepthGrid;
+RWTexture2DArray<uint>		OutDepthCovTexture;
+uint						NumSamples;
 
 groupshared uint group_FurthestDepth; // (4 bytes)
 
 [numthreads(1024, 1, 1)]
-void PrepareDepthGridCS(uint dId : SV_DispatchThreadID, uint GroupThreadID : SV_GroupThreadID)
+void PrepareDepthGridCS(uint DispatchThreadID : SV_DispatchThreadID, uint GroupThreadID : SV_GroupThreadID, uint2 GroupID : SV_GroupID)
 {
-	uint tw = TileRes.x;
-	uint th = TileRes.y;
-
-	if (GroupThreadID == 0) group_FurthestDepth = 0xffffffff;
+	if (GroupThreadID == 0)
+	{
+		group_FurthestDepth = 0xFFFFFFFF;
+	}
 
 	GroupMemoryBarrierWithGroupSync();
 
-	uint t = dId / 1024;
+	// Assumes a workgroup size of 1024 threads to cover a maximum tile size of 32x32.
+	if (GroupThreadID < SqrTileSize)
+	{
+		uint2 Coord;
 
-	uint ty = uint((float(t)+0.5) / float(tw));
-	uint tx = t - (ty*tw);
+		Coord.y = (GroupThreadID + 0.5f) * RcpTileSize;
+		Coord.x = GroupThreadID - (Coord.y * TileSize);
 
-	for (uint p = GroupThreadID; p < SqrTileSize; p += 1024)
-	{
-		if (p < SqrTileSize)
+		Coord += GroupID * TileSize;
+
+		if (all(Coord < (uint2)OutputResolution))
 		{
-			uint py = (float(p)+0.5) * RcpTileSize;
-			uint px = p - (py*TileSize);
+			const float Depth = SceneDepthTexture.Load(uint3(Coord, 0));
+			const uint PackedDepth = PackHairVisDepthCoverage(Depth, 1.0);
 
-			uint2 Coord = uint2(tx * TileSize + px, ty * TileSize + py);
+			// Compute furthest depth inside this tile
+			WaveInterlockedMin(group_FurthestDepth, PackedDepth);
 
-			if (Coord.x < uint(OutputResolution.x) && Coord.y < uint(OutputResolution.y))
+			// Copy scene depth to (multisampled) hair depth output texture
+			for (uint SampleIdx = 0; SampleIdx < NumSamples; ++SampleIdx)
 			{
-				float Depth = SceneDepthTexture.Load(uint3(Coord, 0));
-				uint packedDepth = PackHairVisDepthCoverage(Depth, 1.0);
-
-				InterlockedMin(group_FurthestDepth, packedDepth);
-
-				InterlockedMax(OutDepthCovTexture[Coord], packedDepth);
+				InterlockedMax(OutDepthCovTexture[uint3(Coord, SampleIdx)], PackedDepth);
 			}
 		}
 	}
 
 	GroupMemoryBarrierWithGroupSync();
 
-	OutVisTileDepthGrid[uint2(tx, ty)] = group_FurthestDepth;
+	if (GroupThreadID == 0)
+	{
+		OutVisTileDepthGrid[GroupID] = group_FurthestDepth;
+	}
 }
 
 #endif //SHADER_RASTERCOMPUTE_DEPTH_GRID
@@ -501,7 +562,7 @@ void BinningCS(uint DispatchThreadID : SV_DispatchThreadID, uint GroupThreadID :
 			const uint TotalOldWriteCount = OutVisTileBinningGrid[uint3(TileCoord, TmpSegmentCountLayerIdx)];
 
 			uint NewTile;
-			InterlockedAdd(OutVisTileArgs[0], 1, NewTile);
+			WaveInterlockedAddScalar_(OutVisTileArgs[0], 1, NewTile);
 
 			StoreOutVisTileData(NewTile, VT_Coord, PackedTileCoord);
 			// Round down the count to the start of the tile and later compare against this to decide which tile to write to.
@@ -593,6 +654,7 @@ groupshared uint group_TotalPrimCount;
 groupshared uint group_PrimWriteOffset;
 groupshared uint group_NumTiles;
 groupshared uint group_TilesToCompact[1024];
+groupshared uint group_MaxLDSTileIdx;
 
 [numthreads(1024, 1, 1)]
 void CompactionCS(uint DispatchThreadID : SV_DispatchThreadID, uint GroupThreadID : SV_GroupThreadID, uint2 GroupID : SV_GroupID)
@@ -601,6 +663,7 @@ void CompactionCS(uint DispatchThreadID : SV_DispatchThreadID, uint GroupThreadI
 	{
 		group_TotalPrimCount = 0;
 		group_NumTiles = 0;
+		group_MaxLDSTileIdx = 0;
 	}
 
 	GroupMemoryBarrierWithGroupSync();
@@ -619,10 +682,11 @@ void CompactionCS(uint DispatchThreadID : SV_DispatchThreadID, uint GroupThreadI
 				LocalPrimCount += LoadVisTileData(TileIdx, VT_PrimCount);
 				
 				uint WritePos;
-				InterlockedAdd(group_NumTiles, 1, WritePos);
+				WaveInterlockedAddScalar_(group_NumTiles, 1, WritePos);
 				if (WritePos < 1024)
 				{
 					group_TilesToCompact[WritePos] = TileIdx;
+					WaveInterlockedMax(group_MaxLDSTileIdx, TileIdx);
 				}
 			}
 		}
@@ -632,7 +696,7 @@ void CompactionCS(uint DispatchThreadID : SV_DispatchThreadID, uint GroupThreadI
 
 	if (LocalPrimCount > 0)
 	{
-		InterlockedAdd(group_TotalPrimCount, LocalPrimCount);
+		WaveInterlockedAdd(group_TotalPrimCount, LocalPrimCount);
 	}
 
 	GroupMemoryBarrierWithGroupSync();
@@ -659,7 +723,7 @@ void CompactionCS(uint DispatchThreadID : SV_DispatchThreadID, uint GroupThreadI
 		{
 			const uint CompactedTile = FirstCompactedTile + TileIdx;
 
-			const uint PrimCount = ((TileIdx * 1024) < TotalPrimCount) ? 1024 : (TotalPrimCount - (TileIdx * 1024));
+			const uint PrimCount = min(TotalPrimCount - TileIdx * 1024, 1024);
 			StoreCompactedVisTileData(CompactedTile, VT_PrimCount, PrimCount);
 			StoreCompactedVisTileData(CompactedTile, VT_Coord, PackedCoord);
 		}
@@ -691,7 +755,7 @@ void CompactionCS(uint DispatchThreadID : SV_DispatchThreadID, uint GroupThreadI
 		// Check any remaning tiles
 		if (group_NumTiles > 1024)
 		{
-			for (uint TileIdx = group_TilesToCompact[1023] + 1; TileIdx < NumTiles; ++TileIdx)
+			for (uint TileIdx = group_MaxLDSTileIdx + 1; TileIdx < NumTiles; ++TileIdx)
 			{
 				const uint TilePackedCoord = LoadVisTileData(TileIdx, VT_Coord);
 				if (PackedCoord == TilePackedCoord)
@@ -740,9 +804,9 @@ void CompactionCS(uint DispatchThreadID : SV_DispatchThreadID, uint GroupThreadI
 Buffer<uint> 				VisTilePrims;
 Buffer<uint> 				VisTileArgs;
 ByteAddressBuffer			VisTileData;
-RWTexture2D<uint>			OutHairCountTexture;
-RWTexture2D<uint>			OutDepthCovTexture;
-RWTexture2D<uint>			OutPrimMatTexture;
+RWTexture2DArray<uint>		OutHairCountTexture;
+RWTexture2DArray<uint>		OutDepthCovTexture;
+RWTexture2DArray<uint>		OutPrimMatTexture;
 
 uint LoadVisTileData(uint index, uint offset)
 {
@@ -775,13 +839,7 @@ void PlotInternal(int2 Coords, float AntiAliasingFactor, float4 P0, float4 P1, f
 
 	if (all(IntraTileCoord >= 0) && all(IntraTileCoord < TileSize))
 	{
-		// Project P onto line segment and compute the lerp alpha between P0 and P1
-		// Simplification of:
-		// A = P - P0
-		// B = P1 - P0
-		// Alpha = dot(A, B) / dot(B, B)
-		const float2 P = Coords + 0.5f;
-		const float Alpha = saturate(dot(P - P0.xy, P1.xy - P0.xy) * SegmentLenSqRcp);
+		const float Alpha = ComputeLerpAlpha(Coords, P0.xy, P1.xy, SegmentLenSqRcp);
 		const float Depth = lerp(P0.z, P1.z, Alpha);
 		const uint PackedDepthCov = PackHairVisDepthCoverage(Depth, 1.0f);
 		const uint LinearIndex = IntraTileCoord.x + IntraTileCoord.y * TileSize;
@@ -797,14 +855,7 @@ void PlotInternal(int2 Coords, float AntiAliasingFactor, float4 P0, float4 P1, f
 		// Add hair count if depth test against scene depth is passed
 		if (PackedDepthCov > group_SubTile[LinearIndex].w)
 		{
-			// Alpha value for perspective correct interpolation. We store the reciprocal of w in the w compoennt,
-			// so this is a simplification of:
-			// (Alpha / w1) / ((1 - Alpha) / w0 + Alpha / w1)
-			const float LerpedRcpW = lerp(P0.w, P1.w, Alpha);
-			const float PerspectiveAlpha = (Alpha * P1.w) / LerpedRcpW;
-			// Divide by W to make thickness dependent on screen space depth? This division was kept from the previous line rasterization algorithm.
-			const float Rad = lerp(Rad0, Rad1, PerspectiveAlpha) * LerpedRcpW;
-
+			const float Rad = ComputePerspectiveCorrectRadius(Rad0, Rad1, Alpha, P0.w, P1.w);
 			InterlockedAdd(group_SubTile[LinearIndex].z, min(Rad, 0.5f) * 2.0f * 1000.0f * SampleWeight * AntiAliasingFactor);
 		}
 	}
@@ -912,8 +963,9 @@ void RasterCS(uint DispatchThreadID : SV_DispatchThreadID, uint GroupThreadID :
 
 			group_PrimMatID[Prim] = PackHairVisPrimitiveMaterialId(PrimID, HairMaterialId);
 
-			CalcHomogenousPosAndRad(HairStrandsVF_PositionBuffer[PrimID], group_PositionOffset, group_SP0[Prim], group_Rad0[Prim]);
-			CalcHomogenousPosAndRad(HairStrandsVF_PositionBuffer[PrimID+1], group_PositionOffset, group_SP1[Prim], group_Rad1[Prim]);
+			uint TypeDummy;
+			CalcHomogenousPosAndRad(HairStrandsVF_PositionBuffer[PrimID], group_PositionOffset, group_SP0[Prim], group_Rad0[Prim], TypeDummy);
+			CalcHomogenousPosAndRad(HairStrandsVF_PositionBuffer[PrimID+1], group_PositionOffset, group_SP1[Prim], group_Rad1[Prim], TypeDummy);
 		}
 
 		if (bWaveThreadValid)
@@ -927,8 +979,8 @@ void RasterCS(uint DispatchThreadID : SV_DispatchThreadID, uint GroupThreadID :
 
 				Coord += uint2(((PackedTileMin >> 0) & 0xffff), ((PackedTileMin >> 16) & 0xffff));
 
-				group_SubTile[LinearIndex].x = OutDepthCovTexture[Coord];
-				group_SubTile[LinearIndex].y = 0;
+				group_SubTile[LinearIndex].x = OutDepthCovTexture[uint3(Coord, 0)];
+				group_SubTile[LinearIndex].y = GetInvalidHairPrimitiveId();
 				group_SubTile[LinearIndex].z = 0;
 				group_SubTile[LinearIndex].w = PackHairVisDepthCoverage(SceneDepthTexture.Load(uint3(Coord, 0)), 1.0f);
 			}
@@ -960,8 +1012,9 @@ void RasterCS(uint DispatchThreadID : SV_DispatchThreadID, uint GroupThreadID :
 				uint PrimID = VisTilePrims[PrimOffset + Prim];
 				PrimMatID = PackHairVisPrimitiveMaterialId(PrimID, HairMaterialId);
 
-				CalcHomogenousPosAndRad(HairStrandsVF_PositionBuffer[PrimID], group_PositionOffset, SP0, Rad0);
-				CalcHomogenousPosAndRad(HairStrandsVF_PositionBuffer[PrimID + 1], group_PositionOffset, SP1, Rad1);
+				uint TypeDummy;
+				CalcHomogenousPosAndRad(HairStrandsVF_PositionBuffer[PrimID], group_PositionOffset, SP0, Rad0, TypeDummy);
+				CalcHomogenousPosAndRad(HairStrandsVF_PositionBuffer[PrimID + 1], group_PositionOffset, SP1, Rad1, TypeDummy);
 			}
 
 			// Clipping
@@ -972,23 +1025,24 @@ void RasterCS(uint DispatchThreadID : SV_DispatchThreadID, uint GroupThreadID :
 				// Clip against tile
 				const float2 TileMin = float2(((PackedTileMin >> 0) & 0xffff), ((PackedTileMin >> 16) & 0xffff));
 				const float2 TileMax = TileMin + TileSize;
-				// Clip against a slightly smaller bounding box so that the clipped endpoints are safely inside the tile.
-				ClipRaySegment(TileMin + 0.25f, TileMax - 0.25f, SP0, SP1);
+				bool2 bClipped = false;
+				ClipRaySegment(TileMin - 0.5f, TileMax + 0.5f, SP0, SP1, Rad0, Rad1, bClipped);
 			}
 
 			const float SegmentLenSqRcp = 1.0f / dot(SP1.xy - SP0.xy, SP1.xy - SP0.xy);
 
 #if RASTER_ALGO == RASTER_LINEAR
 
-			const float2 DeltaAbs = abs(SP1.xy - SP0.xy);
-			const uint NumSteps = min((uint)ceil(max(DeltaAbs.x, DeltaAbs.y)) + 1, TileSize + (ENABLE_RASTER_LINEAR_AA ? 1 : 0));
-			const float RcpNumSteps = NumSteps > 1 ? (1.0f / (NumSteps - 1)) : 1.0f;
-			const bool bIsSteep = DeltaAbs.y > DeltaAbs.x;
+			const bool bIsSteep = abs(SP1.x - SP0.x) < abs(SP1.y - SP0.y);
+			const float X0 = bIsSteep ? min(SP0.y, SP1.y) : min(SP0.x, SP1.x);
+			const float X1 = bIsSteep ? max(SP0.y, SP1.y) : max(SP0.x, SP1.x);
+			const int NumSteps = (int)(ceil(X1) - floor(X0));
+			const float RcpNumSteps = 1.0f / (X1 - X0);
 
 			LOOP
 			for (int J = PModTPS; J < NumSteps; J += group_ThreadsPerSeg)
 			{
-				const float Alpha = J * RcpNumSteps;
+				const float Alpha = saturate(J * RcpNumSteps);
 				const float4 SP = lerp(SP0, SP1, Alpha);
 
 				const float AntiAliasingFactor = 1.0f;
@@ -1090,16 +1144,16 @@ void RasterCS(uint DispatchThreadID : SV_DispatchThreadID, uint GroupThreadID :
 
 				Coord += uint2(((PackedTileMin >> 0) & 0xffff), ((PackedTileMin >> 16) & 0xffff));
 
-				if (group_SubTile[LinearIndex].y != 0)
+				if (group_SubTile[LinearIndex].y != GetInvalidHairPrimitiveId())
 				{
 					uint oldValue;
-					InterlockedMax(OutDepthCovTexture[Coord], group_SubTile[LinearIndex].x, oldValue);
+					InterlockedMax(OutDepthCovTexture[uint3(Coord, 0)], group_SubTile[LinearIndex].x, oldValue);
 					if (group_SubTile[LinearIndex].x > oldValue)
 					{
-						OutPrimMatTexture[Coord] = group_SubTile[LinearIndex].y;
+						OutPrimMatTexture[uint3(Coord, 0)] = group_SubTile[LinearIndex].y;
 					}
 				}
-				InterlockedAdd(OutHairCountTexture[Coord], group_SubTile[LinearIndex].z);
+				InterlockedAdd(OutHairCountTexture[uint3(Coord, 0)], group_SubTile[LinearIndex].z);
 			}
 		}
 
@@ -1111,6 +1165,316 @@ void RasterCS(uint DispatchThreadID : SV_DispatchThreadID, uint GroupThreadID :
 
 ///////////////////////////////////////////////////////////////////////////
 
+#if SHADER_RASTERCOMPUTE_RASTER_MULTI_SAMPLE
+
+#ifndef PERMUTATION_MULTI_SAMPLE_COUNT
+#define PERMUTATION_MULTI_SAMPLE_COUNT 1
+#endif
+
+// Wave size
+#if PERMUTATION_GROUP_SIZE == 64
+#define WAVE_SIZE 32
+#elif PERMUTATION_GROUP_SIZE == 32
+#define WAVE_SIZE 32
+#else
+#error Unknown group size	
+#endif
+
+Buffer<uint> 				VisTilePrims;
+Buffer<uint> 				VisTileArgs;
+ByteAddressBuffer			VisTileData;
+RWTexture2D<uint>			OutHairCountTexture;
+RWTexture2DArray<uint>		OutDepthCovTexture;
+RWTexture2DArray<uint>		OutPrimMatTexture;
+
+uint LoadVisTileData(uint index, uint offset)
+{
+	return VisTileData.Load((((index)) * VT_SIZE * 4) + ((offset) * 4));
+}
+
+groupshared uint group_SubTileSceneDepth[256]; // (16 x 16 x 4 bytes = 1k bytes)
+groupshared uint group_SubTileHairCount[256]; // (16 x 16 x 4 bytes = 1k bytes)
+groupshared uint group_SubTileHairDepth[PERMUTATION_MULTI_SAMPLE_COUNT][256]; // (16 x 16 x 4 bytes = 1k bytes) per sample
+groupshared uint group_SubTilePrimMatID[PERMUTATION_MULTI_SAMPLE_COUNT][256]; // (16 x 16 x 4 bytes = 1k bytes) per sample
+
+groupshared float3 group_PositionOffset;
+
+groupshared uint group_LoopNum;
+groupshared uint group_TileNum;
+
+float GetDistanceToLine(float2 P1, float2 P2, float2 P3, float RcpLineSegLength)
+{
+	// We can compute the distance of P1 to the line defined by P2 and P3 as the height of the triangle spanned by these points.
+	// Area of triangle: A = 0.5 * h * b where h is the triangle height and b is the length of the base side.
+	// Solving for h gives: h = (2 * A) / b
+	// We can compute A using the determinant: A = 0.5 * abs(det(P1, P2, P3))
+	// After some simplification, this results in the following:
+	float A = abs(P1.x * (P2.y - P3.y) + P2.x * (P3.y - P1.y) + P3.x * (P1.y - P2.y));
+	return A * RcpLineSegLength;
+}
+
+uint GetCoverageMask(int2 PixelCoord, float2 P0, float2 P1)
+{
+	const float LineThickness = 1.0f / PERMUTATION_MULTI_SAMPLE_COUNT; // In pixel units 
+	uint Mask = 0;
+
+	// Set origin to PixelCoord
+	P0 -= PixelCoord;
+	P1 -= PixelCoord;
+
+	const float RcpLineSegLength = 1.0f / distance(P0, P1);
+
+#if PERMUTATION_MULTI_SAMPLE_COUNT == 1
+	Mask |= (GetDistanceToLine(float2(0.5f, 0.5f), P0, P1, RcpLineSegLength) <= LineThickness) ? (1u << 0u) : 0;
+#elif PERMUTATION_MULTI_SAMPLE_COUNT == 2
+	Mask |= (GetDistanceToLine(float2(0.75f, 0.75f), P0, P1, RcpLineSegLength) <= LineThickness) ? (1u << 0u) : 0;
+	Mask |= (GetDistanceToLine(float2(0.25f, 0.25f), P0, P1, RcpLineSegLength) <= LineThickness) ? (1u << 1u) : 0;
+#elif PERMUTATION_MULTI_SAMPLE_COUNT == 4
+	Mask |= (GetDistanceToLine(float2(0.375f, 0.125f), P0, P1, RcpLineSegLength) <= LineThickness) ? (1u << 0u) : 0;
+	Mask |= (GetDistanceToLine(float2(0.875f, 0.375f), P0, P1, RcpLineSegLength) <= LineThickness) ? (1u << 1u) : 0;
+	Mask |= (GetDistanceToLine(float2(0.125f, 0.625f), P0, P1, RcpLineSegLength) <= LineThickness) ? (1u << 2u) : 0;
+	Mask |= (GetDistanceToLine(float2(0.625f, 0.875f), P0, P1, RcpLineSegLength) <= LineThickness) ? (1u << 3u) : 0;
+#elif PERMUTATION_MULTI_SAMPLE_COUNT == 8
+	Mask |= (GetDistanceToLine(float2(0.5625f, 0.3125f), P0, P1, RcpLineSegLength) <= LineThickness) ? (1u << 0u) : 0;
+	Mask |= (GetDistanceToLine(float2(0.4375f, 0.6875f), P0, P1, RcpLineSegLength) <= LineThickness) ? (1u << 1u) : 0;
+	Mask |= (GetDistanceToLine(float2(0.8125f, 0.5625f), P0, P1, RcpLineSegLength) <= LineThickness) ? (1u << 2u) : 0;
+	Mask |= (GetDistanceToLine(float2(0.3125f, 0.1875f), P0, P1, RcpLineSegLength) <= LineThickness) ? (1u << 3u) : 0;
+	Mask |= (GetDistanceToLine(float2(0.1875f, 0.8125f), P0, P1, RcpLineSegLength) <= LineThickness) ? (1u << 4u) : 0;
+	Mask |= (GetDistanceToLine(float2(0.0625f, 0.4375f), P0, P1, RcpLineSegLength) <= LineThickness) ? (1u << 5u) : 0;
+	Mask |= (GetDistanceToLine(float2(0.6875f, 0.9375f), P0, P1, RcpLineSegLength) <= LineThickness) ? (1u << 6u) : 0;
+	Mask |= (GetDistanceToLine(float2(0.9375f, 0.0625f), P0, P1, RcpLineSegLength) <= LineThickness) ? (1u << 7u) : 0;
+#else
+#error Unsupported PERMUTATION_MULTI_SAMPLE_COUNT! Must be 1, 2, 4 or 8!
+#endif
+
+	return Mask;
+}
+
+void Plot(int2 Coord, float4 P0, float4 P1, float Rad0, float Rad1, float SegmentLenSqRcp, uint PackedHalfTileMin, uint PrimMatID)
+{
+	const int2 IntraTileCoord = Coord - int2(((PackedHalfTileMin >> 0) & 0xffff), ((PackedHalfTileMin >> 16) & 0xffff));
+
+	if (all(IntraTileCoord >= 0) && all(IntraTileCoord < HalfTileSize))
+	{
+		const float Alpha = ComputeLerpAlpha(Coord, P0.xy, P1.xy, SegmentLenSqRcp);
+		const float Depth = lerp(P0.z, P1.z, Alpha);
+		const uint PackedDepthCov = PackHairVisDepthCoverage(Depth, 1.0f);
+		const uint LinearIndex = IntraTileCoord.x + IntraTileCoord.y * HalfTileSize;
+
+		// Test against scene depth
+		if (PackedDepthCov > group_SubTileSceneDepth[LinearIndex])
+		{
+			const float Rad = ComputePerspectiveCorrectRadius(Rad0, Rad1, Alpha, P0.w, P1.w);
+			const uint HairCount = min(Rad, 0.5f) * 2.0f * 1000.0f * SampleWeight;
+
+			const uint CoverageMask = GetCoverageMask(Coord, P0.xy, P1.xy);
+
+			// Accumulate hair count
+			if (CoverageMask)
+			{
+				InterlockedAdd(group_SubTileHairCount[LinearIndex], HairCount);
+			}
+
+			UNROLL for (uint SampleIdx = 0; SampleIdx < PERMUTATION_MULTI_SAMPLE_COUNT; ++SampleIdx)
+			{
+				if (CoverageMask & (1u << SampleIdx))
+				{
+					// Write Depth + PrimMatID if depth test against hair depths is passed
+					uint OldValue;
+					InterlockedMax(group_SubTileHairDepth[SampleIdx][LinearIndex], PackedDepthCov, OldValue);
+					if (PackedDepthCov > OldValue)
+					{
+						group_SubTilePrimMatID[SampleIdx][LinearIndex] = PrimMatID;
+					}
+				}
+			}
+		}
+	}
+}
+
+[numthreads(1024, 1, 1)]
+void RasterMultiSampleCS(uint DispatchThreadID : SV_DispatchThreadID, uint GroupThreadID : SV_GroupThreadID, uint GroupID : SV_GroupID)
+{
+	ResolvedView = ResolveView();
+
+	if (GroupThreadID == 0)
+	{
+		group_TileNum = VisTileArgs[0];
+		group_LoopNum = (float(group_TileNum) + float(NumRasterizers - 1)) * RcpNumRasterizers;
+
+		group_PositionOffset = HairStrandsVF_PositionOffsetBuffer[0].xyz;
+	}
+
+	GroupMemoryBarrierWithGroupSync();
+
+	LOOP for (uint LoopIndex = 0; LoopIndex < group_LoopNum; LoopIndex++)
+	{
+		const uint TileIndex = LoopIndex + (GroupID * group_LoopNum);
+
+		if (TileIndex >= group_TileNum)
+		{
+			return;
+		}
+
+		const uint PrimOffset = TileIndex * 1024;
+		const uint PrimCount = LoadVisTileData(TileIndex, VT_PrimCount);
+
+		const uint PackedCoord = LoadVisTileData(TileIndex, VT_Coord);
+		const uint2 TileMin = UnpackVisTileCoord(PackedCoord) * TileSize;
+
+		const uint PackedTileMin = ((TileMin.x & 0xffff) << 0) | ((TileMin.y & 0xffff) << 16);
+
+		uint ThreadsPerSeg = 1;
+
+		if (PrimCount <= 512)
+			ThreadsPerSeg = 2;
+		if (PrimCount <= 341)
+			ThreadsPerSeg = 3;
+		if (PrimCount <= 256)
+			ThreadsPerSeg = 4;
+		if (PrimCount <= 204)
+			ThreadsPerSeg = 5;
+		if (PrimCount <= 170)
+			ThreadsPerSeg = 6;
+		if (PrimCount <= 146)
+			ThreadsPerSeg = 7;
+		if (PrimCount <= 128)
+			ThreadsPerSeg = 8;
+		if (PrimCount <= 64)
+			ThreadsPerSeg = 16;
+		if (PrimCount <= 32)
+			ThreadsPerSeg = 32;
+
+		const bool bThreadValid = (GroupThreadID < (PrimCount * ThreadsPerSeg));
+		const uint Prim = uint((float(GroupThreadID) + 0.5f) / float(ThreadsPerSeg));
+		const uint PModTPS = GroupThreadID - (Prim * ThreadsPerSeg);
+
+		float4 SP0;
+		float4 SP1;
+		float Rad0;
+		float Rad1;
+		bool bIsEndPoint;
+		uint PrimMatID;
+
+		if (bThreadValid)
+		{
+			const uint PrimID = VisTilePrims[PrimOffset + Prim];
+			PrimMatID = PackHairVisPrimitiveMaterialId(PrimID, HairMaterialId);
+
+			uint Type;
+			CalcHomogenousPosAndRad(HairStrandsVF_PositionBuffer[PrimID], group_PositionOffset, SP0, Rad0, Type);
+			CalcHomogenousPosAndRad(HairStrandsVF_PositionBuffer[PrimID + 1], group_PositionOffset, SP1, Rad1, Type);
+			bIsEndPoint = (Type == HAIR_CONTROLPOINT_END);
+
+			SP0 = float4(NDCToPixelCoord(SP0), 1.0f / SP0.w);
+			SP1 = float4(NDCToPixelCoord(SP1), 1.0f / SP1.w);
+		}
+
+		// Split 32x32 tile into 4 16x16 tiles that are processed one after another. 
+		// This is to reduce LDS memory pressure.
+		UNROLL for (uint SubTileIdx = 0; SubTileIdx < 4; ++SubTileIdx)
+		{
+			const uint2 SubTileMin = TileMin + uint2((SubTileIdx == 0 || SubTileIdx == 2) ? 0 : HalfTileSize, SubTileIdx < 2 ? 0 : HalfTileSize);
+			const uint2 SubTileMax = SubTileMin + HalfTileSize;
+			const uint PackedSubTileMin = ((SubTileMin.x & 0xFFFF) << 0u) | ((SubTileMin.y & 0xFFFF) << 16u);
+
+			// Initialize LDS
+			if (GroupThreadID < SqrHalfTileSize)
+			{
+				uint2 Coord;
+				Coord.y = (float(GroupThreadID) + 0.5f) * RcpHalfTileSize;
+				Coord.x = GroupThreadID - (Coord.y * HalfTileSize);
+				Coord += uint2(((PackedSubTileMin >> 0u) & 0xFFFF), ((PackedSubTileMin >> 16u) & 0xFFFF));
+
+				group_SubTileSceneDepth[GroupThreadID] = PackHairVisDepthCoverage(SceneDepthTexture.Load(uint3(Coord, 0)), 1.0f);
+				group_SubTileHairCount[GroupThreadID] = 0;
+
+				UNROLL for (uint SampleIdx = 0; SampleIdx < PERMUTATION_MULTI_SAMPLE_COUNT; ++SampleIdx)
+				{
+					const uint HairDepth = OutDepthCovTexture.Load(uint3(Coord, SampleIdx));
+					group_SubTileHairDepth[SampleIdx][GroupThreadID] = HairDepth;
+					group_SubTilePrimMatID[SampleIdx][GroupThreadID] = GetInvalidHairPrimitiveId();
+				}
+			}
+
+			GroupMemoryBarrierWithGroupSync();
+
+			// Rasterize to LDS
+			if (bThreadValid)
+			{
+				const uint2 SubTileMin = uint2(((PackedSubTileMin >> 0u) & 0xFFFF), ((PackedSubTileMin >> 16u) & 0xFFFF));
+				const uint2 SubTileMax = SubTileMin + HalfTileSize;
+				bool2 bClipped;
+				float2 T;
+				const bool bVisible = ClipRaySegment(SubTileMin - 0.5f, SubTileMax + 0.5f, SP0, SP1, T, bClipped);
+				T = saturate(T);
+				
+				if (bVisible)
+				{
+					const float2 SP0Clipped = lerp(SP0, SP1, T.x);
+					const float2 SP1Clipped = lerp(SP0, SP1, T.y);
+					const bool bIsSteep = abs(SP1.x - SP0.x) < abs(SP1.y - SP0.y);
+					const float X0 = bIsSteep ? min(SP0Clipped.y, SP1Clipped.y) : min(SP0Clipped.x, SP1Clipped.x);
+					const float X1 = bIsSteep ? max(SP0Clipped.y, SP1Clipped.y) : max(SP0Clipped.x, SP1Clipped.x);
+					const int NumSteps = (int)(ceil(X1) - floor(X0));
+					const float RcpNumSteps = 1.0f / (X1 - X0);
+
+					const float SegmentLenSqRcp = 1.0f / dot(SP1.xy - SP0.xy, SP1.xy - SP0.xy);
+
+					const int End = !bClipped.y && !bIsEndPoint ? (NumSteps - 1) : NumSteps;
+					LOOP for (int J = PModTPS; J < End; J += ThreadsPerSeg)
+					{
+						const float Alpha = lerp(T.x, T.y, saturate(J * RcpNumSteps));
+						const float2 SP = lerp(SP0.xy, SP1.xy, Alpha);
+						Plot(SP, SP0, SP1, Rad0, Rad1, SegmentLenSqRcp, PackedSubTileMin, PrimMatID);
+					}
+				}
+			}
+
+			GroupMemoryBarrierWithGroupSync();
+
+			// Write out to global memory
+			if (GroupThreadID < SqrHalfTileSize)
+			{
+				uint2 Coord;
+				Coord.y = (float(GroupThreadID) + 0.5f) * RcpHalfTileSize;
+				Coord.x = GroupThreadID - (Coord.y * HalfTileSize);
+				Coord += uint2(((PackedSubTileMin >> 0u) & 0xFFFF), ((PackedSubTileMin >> 16u) & 0xFFFF));
+
+				const uint HairCount = group_SubTileHairCount[GroupThreadID];
+				if (HairCount != 0)
+				{
+					InterlockedAdd(OutHairCountTexture[Coord], HairCount);
+				}
+
+				UNROLL for (uint SampleIdx = 0; SampleIdx < PERMUTATION_MULTI_SAMPLE_COUNT; ++SampleIdx)
+				{
+					const uint3 SampleCoord = uint3(Coord, SampleIdx);
+					
+					const uint PrimMatID = group_SubTilePrimMatID[SampleIdx][GroupThreadID];
+					if (PrimMatID != GetInvalidHairPrimitiveId())
+					{
+						const uint HairDepth = group_SubTileHairDepth[SampleIdx][GroupThreadID];
+
+						uint OldValue;
+						InterlockedMax(OutDepthCovTexture[SampleCoord], HairDepth, OldValue);
+						if (HairDepth > OldValue)
+						{
+							OutPrimMatTexture[SampleCoord] = PrimMatID;
+						}
+					}
+				}
+			}
+
+			GroupMemoryBarrierWithGroupSync();
+		}
+	}
+}
+
+#endif //SHADER_RASTERCOMPUTE_RASTER_MULTI_SAMPLE
+
+///////////////////////////////////////////////////////////////////////////
+
 #if SHADER_RASTERCOMPUTE_DEBUG
 
 #include "../ShaderPrint.ush"
diff --git a/Engine/Shaders/Private/WaveOpUtil.ush b/Engine/Shaders/Private/WaveOpUtil.ush
--- a/Engine/Shaders/Private/WaveOpUtil.ush
+++ b/Engine/Shaders/Private/WaveOpUtil.ush
@@ -9,7 +9,7 @@ uint CountBits( uint2 Bits )
 
 /*
 ===============================================================================
-Below are a series of WaveInterlockedAdd macros used for aggregated atomics.
+Below are a series of WaveInterlockedAdd (and WaveInterlockedMin/Max) macros used for aggregated atomics.
 
 There is a varying and scalar version of each type, scalar being faster. Scalar refers to Value. I can't check if the scalar version
 is actually called with a scalar since HLSL doesn't allow that so be careful.
@@ -97,6 +97,20 @@ that's to be written to as a parameter on all platforms. Also unfortunate is tha
 	OriginalValue = WaveReadLaneFirst( OriginalValue ) + LocalOffset; \
 }
 
+#define WaveInterlockedMin( Dest, Value ) \
+{ \
+	uint MinValue = WaveActiveMin( Value ); \
+	if( WaveIsFirstLane() ) \
+		InterlockedMin( Dest, MinValue ); \
+}
+
+#define WaveInterlockedMax( Dest, Value ) \
+{ \
+	uint MaxValue = WaveActiveMax( Value ); \
+	if( WaveIsFirstLane() ) \
+		InterlockedMax( Dest, MaxValue ); \
+}
+
 #else
 
 #define WaveInterlockedAddScalar( Dest, Value )					InterlockedAdd( Dest, Value )
@@ -116,4 +130,7 @@ that's to be written to as a parameter on all platforms. Also unfortunate is tha
 	InterlockedMax( DestGroups, ( OriginalValue + Value + GroupsOf - 1 ) / GroupsOf ); \
 }
 
+#define WaveInterlockedMin( Dest, Value ) { InterlockedMin(Dest, Value); }
+#define WaveInterlockedMax( Dest, Value ) { InterlockedMax(Dest, Value); }
+
 #endif // PLATFORM_SUPPORTS_SM6_0_WAVE_OPERATIONS
\ No newline at end of file
diff --git a/Engine/Source/Runtime/Renderer/Private/HairStrands/HairStrandsVisibility.cpp b/Engine/Source/Runtime/Renderer/Private/HairStrands/HairStrandsVisibility.cpp
--- a/Engine/Source/Runtime/Renderer/Private/HairStrands/HairStrandsVisibility.cpp
+++ b/Engine/Source/Runtime/Renderer/Private/HairStrands/HairStrandsVisibility.cpp
@@ -33,6 +33,11 @@ static float GHairVisibilityMSAA_MeanSamplePerPixel = 0.75f;
 static FAutoConsoleVariableRef CVarHairVisibilityMSAA_MaxSamplePerPixel(TEXT("r.HairStrands.Visibility.MSAA.SamplePerPixel"), GHairVisibilityMSAA_MaxSamplePerPixel, TEXT("Hair strands visibility sample count (2, 4, or 8)"), ECVF_Scalability | ECVF_RenderThreadSafe);
 static FAutoConsoleVariableRef CVarHairVisibilityMSAA_MeanSamplePerPixel(TEXT("r.HairStrands.Visibility.MSAA.MeanSamplePerPixel"), GHairVisibilityMSAA_MeanSamplePerPixel, TEXT("Scale the numer of sampler per pixel for limiting memory allocation (0..1, default 0.5f)"));
 
+static int32 GHairVisibilityCompute_MaxSamplePerPixel = 1;
+static float GHairVisibilityCompute_MeanSamplePerPixel = 0.75f;
+static FAutoConsoleVariableRef CVarHairVisibilityCompute_MaxSamplePerPixel(TEXT("r.HairStrands.Visibility.Compute.SamplePerPixel"), GHairVisibilityCompute_MaxSamplePerPixel, TEXT("Hair strands visibility sample count (2, 4, or 8)"), ECVF_Scalability | ECVF_RenderThreadSafe);
+static FAutoConsoleVariableRef CVarHairVisibilityCompute_MeanSamplePerPixel(TEXT("r.HairStrands.Visibility.Compute.MeanSamplePerPixel"), GHairVisibilityCompute_MeanSamplePerPixel, TEXT("Scale the numer of sampler per pixel for limiting memory allocation (0..1, default 0.5f)"));
+
 static int32 GHairClearVisibilityBuffer = 0;
 static FAutoConsoleVariableRef CVarHairClearVisibilityBuffer(TEXT("r.HairStrands.Visibility.Clear"), GHairClearVisibilityBuffer, TEXT("Clear hair strands visibility buffer"));
 
@@ -182,7 +187,22 @@ static uint32 GetMaxSamplePerPixel()
 	{
 		case HairVisibilityRenderMode_ComputeRaster:
 		{
-			return 1; //TODO: enable multi-sample compute raster
+			if (GHairVisibilityCompute_MaxSamplePerPixel <= 1)
+			{
+				return 1;
+			}
+			else if (GHairVisibilityCompute_MaxSamplePerPixel == 2)
+			{
+				return 2;
+			}
+			else if (GHairVisibilityCompute_MaxSamplePerPixel <= 4)
+			{
+				return 4;
+			}
+			else
+			{
+				return 8;
+			}
 		}
 		case HairVisibilityRenderMode_MSAA_Visibility:
 		{
@@ -234,7 +254,7 @@ inline uint32 GetMeanSamplePerPixel()
 	switch (GetHairVisibilityRenderMode())
 	{
 	case HairVisibilityRenderMode_ComputeRaster:
-		return 1; //TODO: enable multi-sample compute raster
+		return FMath::Max(1, FMath::FloorToInt(SamplePerPixel * FMath::Clamp(GHairVisibilityCompute_MeanSamplePerPixel, 0.f, 1.f)));
 	case HairVisibilityRenderMode_MSAA_Visibility:
 		return FMath::Max(1, FMath::FloorToInt(SamplePerPixel * FMath::Clamp(GHairVisibilityMSAA_MeanSamplePerPixel, 0.f, 1.f)));
 	case HairVisibilityRenderMode_PPLL:
@@ -253,9 +273,7 @@ uint32 GetHairStrandsMeanSamplePerPixel()
 
 struct FRasterComputeOutput
 {
-	FIntPoint BaseResolution;
-	FIntPoint SuperResolution;
-	uint32 ResolutionMultiplier = 1;
+	FIntPoint Resolution;
 
 	FRDGTextureRef HairCountTexture = nullptr;
 	FRDGTextureRef DepthTexture = nullptr;
@@ -1757,13 +1775,15 @@ class FHairVisibilityCompactionComputeRasterCS : public FGlobalShader
 	SHADER_USE_PARAMETER_STRUCT(FHairVisibilityCompactionComputeRasterCS, FGlobalShader);
 
 	class FGroupSize : SHADER_PERMUTATION_SPARSE_INT("PERMUTATION_GROUPSIZE", 32, 64);
-	using FPermutationDomain = TShaderPermutationDomain<FGroupSize>;
+	class FSampleCount : SHADER_PERMUTATION_SPARSE_INT("PERMUTATION_MULTI_SAMPLE_COUNT", 1, 2, 4, 8);
+	using FPermutationDomain = TShaderPermutationDomain<FGroupSize, FSampleCount>;
 
 	BEGIN_SHADER_PARAMETER_STRUCT(FParameters, )
 		SHADER_PARAMETER(FIntPoint, OutputResolution)
 		SHADER_PARAMETER(uint32, MaxNodeCount)
 		SHADER_PARAMETER(uint32, SamplerPerPixel)
 		SHADER_PARAMETER(float, CoverageThreshold)
+		SHADER_PARAMETER(uint32, bSortSampleByDepth)
 
 		SHADER_PARAMETER(FIntPoint, TileCountXY)
 		SHADER_PARAMETER(uint32, TileSize)
@@ -1805,7 +1825,7 @@ static void AddHairVisibilityCompactionComputeRasterPass(
 	FRDGBuilder& GraphBuilder,
 	const FViewInfo& View,
 	const uint32 NodeGroupSize,
-	const uint32 SamplerPerPixel,
+	const uint32 SamplesPerPixel,
 	const FRasterComputeOutput& RasterComputeData,
 	const FHairStrandsTiles& TileData,
 	FRDGTextureRef& OutCompactCounter,
@@ -1854,10 +1874,11 @@ static void AddHairVisibilityCompactionComputeRasterPass(
 	PassParameters->DepthCovTexture = RasterComputeData.DepthCovTexture;
 	PassParameters->PrimMatTexture = RasterComputeData.PrimMatTexture;
 
-	PassParameters->SamplerPerPixel			= SamplerPerPixel;
+	PassParameters->SamplerPerPixel			= SamplesPerPixel;
 	PassParameters->OutputResolution		= Resolution;
 	PassParameters->MaxNodeCount			= MaxRenderNodeCount;
 	PassParameters->CoverageThreshold		= GetHairStrandsFullCoverageThreshold();
+	PassParameters->bSortSampleByDepth		= GHairStrandsSortHairSampleByDepth > 0 ? 1 : 0;
 	PassParameters->ViewUniformBuffer		= View.ViewUniformBuffer;
 //	PassParameters->SceneTexturesStruct		= CreateSceneTextureUniformBuffer(GraphBuilder, View.GetSceneTexturesChecked(), View.FeatureLevel);
 	PassParameters->OutCompactNodeCounter	= GraphBuilder.CreateUAV(OutCompactCounter);
@@ -1874,6 +1895,7 @@ static void AddHairVisibilityCompactionComputeRasterPass(
 	const FIntPoint GroupSize = GetVendorOptimalGroupSize2D();
 	FHairVisibilityCompactionComputeRasterCS::FPermutationDomain PermutationVector;
 	PermutationVector.Set<FHairVisibilityCompactionComputeRasterCS::FGroupSize>(FHairStrandsTiles::GroupSize);
+	PermutationVector.Set<FHairVisibilityCompactionComputeRasterCS::FSampleCount>(SamplesPerPixel);
 	TShaderMapRef<FHairVisibilityCompactionComputeRasterCS> ComputeShader(View.ShaderMap, PermutationVector);
 	FComputeShaderUtils::AddPass(
 		GraphBuilder,
@@ -2736,6 +2758,9 @@ BEGIN_SHADER_PARAMETER_STRUCT(FRasterComputeCommonParameters, )
 	SHADER_PARAMETER(float, RcpTileSize)
 	SHADER_PARAMETER(uint32, TileSizeAsShift)
 	SHADER_PARAMETER(uint32, SqrTileSize)
+	SHADER_PARAMETER(uint32, HalfTileSize)
+	SHADER_PARAMETER(float, RcpHalfTileSize)
+	SHADER_PARAMETER(uint32, SqrHalfTileSize)
 	SHADER_PARAMETER(uint32, NumRasterizers)
 	SHADER_PARAMETER(float, RcpNumRasterizers)
 	SHADER_PARAMETER(FIntPoint, TileRes)
@@ -2755,6 +2780,7 @@ class FVisiblityRasterComputePrepareDepthGridCS : public FGlobalShader
 		SHADER_PARAMETER_RDG_TEXTURE(Texture2D, SceneDepthTexture)
 		SHADER_PARAMETER_RDG_TEXTURE_UAV(RWTexture2D, OutVisTileDepthGrid)
 		SHADER_PARAMETER_RDG_TEXTURE_UAV(RWTexture2D, OutDepthCovTexture)
+		SHADER_PARAMETER(uint32, NumSamples)
 	END_SHADER_PARAMETER_STRUCT()
 
 public:
@@ -2971,6 +2997,45 @@ IMPLEMENT_GLOBAL_SHADER(FVisiblityRasterComputeRasterizeCS, "/Engine/Private/Hai
 
 /////////////////////////////////////////////////////////////////////////////////////////
 
+class FVisiblityRasterComputeRasterizeMultiSampleCS : public FGlobalShader
+{
+	DECLARE_GLOBAL_SHADER(FVisiblityRasterComputeRasterizeMultiSampleCS);
+	SHADER_USE_PARAMETER_STRUCT(FVisiblityRasterComputeRasterizeMultiSampleCS, FGlobalShader);
+
+	class FSampleCount : SHADER_PERMUTATION_SPARSE_INT("PERMUTATION_MULTI_SAMPLE_COUNT", 1, 2, 4, 8);
+	class FGroupSize : SHADER_PERMUTATION_SPARSE_INT("PERMUTATION_GROUP_SIZE", 32, 64);
+	using FPermutationDomain = TShaderPermutationDomain<FGroupSize, FSampleCount>;
+
+	BEGIN_SHADER_PARAMETER_STRUCT(FParameters, )
+		SHADER_PARAMETER_STRUCT_INCLUDE(FRasterComputeCommonParameters, Common)
+		SHADER_PARAMETER(uint32, MacroGroupId)
+		SHADER_PARAMETER(uint32, DispatchCountX)
+		SHADER_PARAMETER(uint32, HairMaterialId)
+		SHADER_PARAMETER_STRUCT_INCLUDE(FHairStrandsInstanceParameters, HairInstance)
+		SHADER_PARAMETER_RDG_TEXTURE(Texture2D, SceneDepthTexture)
+		SHADER_PARAMETER(float, SampleWeight)
+		SHADER_PARAMETER_RDG_BUFFER_SRV(Buffer, VisTilePrims)
+		SHADER_PARAMETER_RDG_BUFFER_SRV(Buffer, VisTileArgs)
+		SHADER_PARAMETER_RDG_BUFFER_SRV(ByteAddressBuffer, VisTileData)
+		SHADER_PARAMETER_RDG_TEXTURE_UAV(RWTexture2DArray, OutHairCountTexture)
+		SHADER_PARAMETER_RDG_TEXTURE_UAV(RWTexture2DArray, OutDepthCovTexture)
+		SHADER_PARAMETER_RDG_TEXTURE_UAV(RWTexture2DArray, OutPrimMatTexture)
+		SHADER_PARAMETER_STRUCT_REF(FViewUniformShaderParameters, ViewUniformBuffer)
+		END_SHADER_PARAMETER_STRUCT()
+
+public:
+	static bool ShouldCompilePermutation(const FGlobalShaderPermutationParameters& Parameters) { return IsHairStrandsSupported(EHairStrandsShaderType::Strands, Parameters.Platform); }
+	static void ModifyCompilationEnvironment(const FGlobalShaderPermutationParameters& Parameters, FShaderCompilerEnvironment& OutEnvironment)
+	{
+		FGlobalShader::ModifyCompilationEnvironment(Parameters, OutEnvironment);
+		OutEnvironment.SetDefine(TEXT("SHADER_RASTERCOMPUTE_RASTER_MULTI_SAMPLE"), 1);
+	}
+};
+
+IMPLEMENT_GLOBAL_SHADER(FVisiblityRasterComputeRasterizeMultiSampleCS, "/Engine/Private/HairStrands/HairStrandsVisibilityRasterCompute.usf", "RasterMultiSampleCS", SF_Compute);
+
+/////////////////////////////////////////////////////////////////////////////////////////
+
 class FVisiblityRasterComputeDebugCS : public FGlobalShader
 {
 	DECLARE_GLOBAL_SHADER(FVisiblityRasterComputeDebugCS);
@@ -3126,13 +3191,10 @@ static FRasterComputeOutput AddVisibilityComputeRasterPass(
 	check(DoesSupportRasterCompute());
 
 	FRasterComputeOutput Out;
+	Out.Resolution = InResolution;
 
-	Out.ResolutionMultiplier = 1;
-	Out.BaseResolution		 = InResolution;
-	Out.SuperResolution		 = InResolution * Out.ResolutionMultiplier;
-
-	FRDGTextureDesc DescCount = FRDGTextureDesc::Create2D(Out.SuperResolution, PF_R32_UINT, FClearValueBinding::None, TexCreate_UAV | TexCreate_ShaderResource | TexCreate_RenderTargetable);
-	FRDGTextureDesc DescVis = FRDGTextureDesc::Create2D(Out.SuperResolution, PF_R32_UINT, FClearValueBinding::None, TexCreate_UAV | TexCreate_ShaderResource | TexCreate_RenderTargetable);
+	FRDGTextureDesc DescCount = FRDGTextureDesc::Create2D(Out.Resolution, PF_R32_UINT, FClearValueBinding::None, TexCreate_UAV | TexCreate_ShaderResource | TexCreate_RenderTargetable);
+	FRDGTextureDesc DescVis = FRDGTextureDesc::Create2DArray(Out.Resolution, PF_R32_UINT, FClearValueBinding::None, TexCreate_UAV | TexCreate_ShaderResource | TexCreate_RenderTargetable, SamplePerPixelCount);
 
 	Out.HairCountTexture = GraphBuilder.CreateTexture(DescCount, TEXT("Hair.HairCountTexture"));
 	Out.DepthCovTexture = GraphBuilder.CreateTexture(DescVis, TEXT("Hair.DepthCovTexture"));
@@ -3143,7 +3205,6 @@ static FRasterComputeOutput AddVisibilityComputeRasterPass(
 	FRDGTextureUAVRef PrimMatTextureUAV = GraphBuilder.CreateUAV(Out.PrimMatTexture);
 
 	uint32 ClearValues[4] = { 0,0,0,0 };
-	AddClearUAVPass(GraphBuilder, PrimMatTextureUAV, ClearValues);
 	AddClearUAVPass(GraphBuilder, DepthCovTextureUAV, ClearValues);
 	AddClearUAVPass(GraphBuilder, HairCountTextureUAV, ClearValues);
 
@@ -3171,10 +3232,10 @@ static FRasterComputeOutput AddVisibilityComputeRasterPass(
 				continue;
 			}
 
-			const FHairGroupPublicData *HairGroupPublicData = reinterpret_cast<const FHairGroupPublicData *>(PrimitiveInfo.Mesh->Elements[0].VertexFactoryUserData);
+			const FHairGroupPublicData* HairGroupPublicData = reinterpret_cast<const FHairGroupPublicData*>(PrimitiveInfo.Mesh->Elements[0].VertexFactoryUserData);
 			check(HairGroupPublicData);
 
-			const FHairGroupPublicData::FVertexFactoryInput &VFInput = HairGroupPublicData->VFInput;
+			const FHairGroupPublicData::FVertexFactoryInput& VFInput = HairGroupPublicData->VFInput;
 
 			uint32 VertexCount = VFInput.Strands.VertexCount;
 			const bool bCullingPossible = bSupportCulling && GHairVisibilityComputeRaster_Culling;
@@ -3191,10 +3252,10 @@ static FRasterComputeOutput AddVisibilityComputeRasterPass(
 		}
 	}
 
-	FRDGBufferRef VisTilePrims = GraphBuilder.CreateBuffer(FRDGBufferDesc::CreateBufferDesc(sizeof(uint32), MaxTiles*1024), TEXT("Hair.VisTilePrims"));
+	FRDGBufferRef VisTilePrims = GraphBuilder.CreateBuffer(FRDGBufferDesc::CreateBufferDesc(sizeof(uint32), MaxTiles * 1024), TEXT("Hair.VisTilePrims"));
 	FRDGBufferSRVRef VisTilePrimsSRV = GraphBuilder.CreateSRV(VisTilePrims, PF_R32_UINT);
 	
-	FRDGBufferRef CompactedVisTilePrims = GraphBuilder.CreateBuffer(FRDGBufferDesc::CreateBufferDesc(sizeof(uint32), MaxTiles*1024), TEXT("Hair.CompactedVisTilePrims"));
+	FRDGBufferRef CompactedVisTilePrims = GraphBuilder.CreateBuffer(FRDGBufferDesc::CreateBufferDesc(sizeof(uint32), MaxTiles * 1024), TEXT("Hair.CompactedVisTilePrims"));
 
 	FRDGTextureDesc DescBinningGrid = FRDGTextureDesc::Create2DArray(TileGridRes, PF_R32_UINT, FClearValueBinding::None, TexCreate_UAV | TexCreate_ShaderResource, NumBinners * 3);
 	FRDGTextureRef VisTileBinningGrid = GraphBuilder.CreateTexture(DescBinningGrid, TEXT("Hair.VisTileBinningGrid"));
@@ -3203,11 +3264,12 @@ static FRasterComputeOutput AddVisibilityComputeRasterPass(
 	FRDGTextureDesc DescDepthGrid = FRDGTextureDesc::Create2D(TileGridRes, PF_R32_UINT, FClearValueBinding::None, TexCreate_UAV | TexCreate_ShaderResource);
 	FRDGTextureRef VisTileDepthGrid = GraphBuilder.CreateTexture(DescDepthGrid, TEXT("Hair.VisTileDepthGrid"));
 
-	FRDGBufferRef VisTileArgs = GraphBuilder.CreateBuffer(FRDGBufferDesc::CreateBufferDesc(sizeof(uint32), 4+16+(((TileGridRes.X*TileGridRes.Y)+31)/32) ), TEXT("Hair.VisTileArgs"));
+	FRDGBufferDesc VisTileArgsDesc = FRDGBufferDesc::CreateBufferDesc(sizeof(uint32), 4 + 16 + (((TileGridRes.X * TileGridRes.Y) + 31) / 32));
+	FRDGBufferRef VisTileArgs = GraphBuilder.CreateBuffer(VisTileArgsDesc, TEXT("Hair.VisTileArgs"));
 	FRDGBufferUAVRef VisTileArgsUAV = GraphBuilder.CreateUAV(VisTileArgs, PF_R32_UINT);
 	FRDGBufferSRVRef VisTileArgsSRV = GraphBuilder.CreateSRV(VisTileArgs, PF_R32_UINT);
 	
-	FRDGBufferRef CompactedVisTileArgs = GraphBuilder.CreateBuffer(FRDGBufferDesc::CreateBufferDesc(sizeof(uint32), 4+16+(((TileGridRes.X*TileGridRes.Y)+31)/32) ), TEXT("Hair.CompactedVisTileArgs"));
+	FRDGBufferRef CompactedVisTileArgs = GraphBuilder.CreateBuffer(VisTileArgsDesc, TEXT("Hair.CompactedVisTileArgs"));
 
 	FRDGBufferDesc DescTileData = FRDGBufferDesc::CreateStructuredDesc(sizeof(uint32), MaxTiles * 4 * 2);
 	DescTileData.Usage = EBufferUsageFlags(DescTileData.Usage | BUF_ByteAddressBuffer);
@@ -3232,10 +3294,11 @@ static FRasterComputeOutput AddVisibilityComputeRasterPass(
 	SetUpViewHairRenderInfo(ViewInfo, bEnableMSAA, ViewInfo.CachedViewUniformShaderParameters->HairRenderInfo, ViewInfo.CachedViewUniformShaderParameters->HairRenderInfoBits, ViewInfo.CachedViewUniformShaderParameters->HairComponents);
 	ViewUniformShaderParameters = TUniformBufferRef<FViewUniformShaderParameters>::CreateUniformBufferImmediate(*ViewInfo.CachedViewUniformShaderParameters, UniformBuffer_SingleFrame);
 
-	const bool bClassification = GHairStrandsHWSWClassifaction != 0;
-	const bool bUseHWRaster = GHairStrandsUseHWRaster != 0;
-	const bool bUseNaiveSWRaster = GHairStrandsUseNaiveSWRaster != 0;
+	const bool bClassification = GHairStrandsHWSWClassifaction != 0 && SamplePerPixelCount <= 1;
+	const bool bUseHWRaster = GHairStrandsUseHWRaster != 0 && SamplePerPixelCount <= 1;
+	const bool bUseNaiveSWRaster = GHairStrandsUseNaiveSWRaster != 0 && SamplePerPixelCount <= 1;
 	const bool bTileCompactionPass = GHairStrandsTileCompaction != 0;
+	const bool bMultiSampleSWRaster = SamplePerPixelCount > 1;
 
 	const uint32 GroupSize = GetVendorOptimalGroupSize1D();
 
@@ -3264,6 +3327,11 @@ static FRasterComputeOutput AddVisibilityComputeRasterPass(
 	RasterPermutationVector.Set<FVisiblityRasterComputeRasterizeCS::FGroupSize>(GroupSize);
 	TShaderMapRef<FVisiblityRasterComputeRasterizeCS> ComputeShaderRaster(ViewInfo.ShaderMap, RasterPermutationVector);
 
+	FVisiblityRasterComputeRasterizeMultiSampleCS::FPermutationDomain RasterMultiSamplePermutationVector;
+	RasterMultiSamplePermutationVector.Set<FVisiblityRasterComputeRasterizeMultiSampleCS::FGroupSize>(GroupSize);
+	RasterMultiSamplePermutationVector.Set<FVisiblityRasterComputeRasterizeMultiSampleCS::FSampleCount>(SamplePerPixelCount);
+	TShaderMapRef<FVisiblityRasterComputeRasterizeMultiSampleCS> ComputeShaderRasterMultiSample(ViewInfo.ShaderMap, RasterMultiSamplePermutationVector);
+
 	FVisiblityRasterHWVS::FPermutationDomain VSRasterPermutationVector_CullingOff;
 	VSRasterPermutationVector_CullingOff.Set<FVisiblityRasterHWVS::FDrawIndirect>(bClassification);
 	VSRasterPermutationVector_CullingOff.Set<FVisiblityRasterHWVS::FCulling>(false);
@@ -3289,15 +3357,18 @@ static FRasterComputeOutput AddVisibilityComputeRasterPass(
 	// Common parameters
 	FRasterComputeCommonParameters Common;
 	{		
-		Common.MaxRasterCount = TileSize; //TODO: implement tile clipping in raster pass to avoid exisiting edge case
+		Common.MaxRasterCount = TileSize;
 		Common.FrameIdMod8 =
-		Common.ResolutionMultiplier = Out.ResolutionMultiplier;
-		Common.OutputResolution = Out.SuperResolution;
-		Common.OutputResolutionf = FVector2f(Out.SuperResolution.X, Out.SuperResolution.Y);
+		Common.ResolutionMultiplier = 1;
+		Common.OutputResolution = Out.Resolution;
+		Common.OutputResolutionf = FVector2f(Out.Resolution.X, Out.Resolution.Y);
 		Common.TileSize = TileSize;
 		Common.TileSizeAsShift = uint32(FMath::Log2(float(TileSize)));
 		Common.RcpTileSize = 1.0 / TileSize;
 		Common.SqrTileSize = TileSize * TileSize;
+		Common.HalfTileSize = TileSize / 2;
+		Common.RcpHalfTileSize = 1.0f / Common.HalfTileSize;
+		Common.SqrHalfTileSize = Common.HalfTileSize * Common.HalfTileSize;
 		Common.NumRasterizers = NumRasterizers;
 		Common.RcpNumRasterizers = 1.0 / NumRasterizers;
 		Common.TileRes = TileGridRes;
@@ -3315,7 +3386,8 @@ static FRasterComputeOutput AddVisibilityComputeRasterPass(
 		PrepDepthGridParameters->SceneDepthTexture = SceneDepthTexture;
 		PrepDepthGridParameters->OutVisTileDepthGrid = GraphBuilder.CreateUAV(VisTileDepthGrid);
 		PrepDepthGridParameters->OutDepthCovTexture = DepthCovTextureUAV;
-		FComputeShaderUtils::AddPass(GraphBuilder, RDG_EVENT_NAME("HairStrands::VisibilityComputeRasterPrepDepthGrid"), ComputeShaderPrepareDepthGrid, PrepDepthGridParameters, FIntVector(TileGridRes.X* TileGridRes.Y, 1, 1));
+		PrepDepthGridParameters->NumSamples = SamplePerPixelCount;
+		FComputeShaderUtils::AddPass(GraphBuilder, RDG_EVENT_NAME("HairStrands::VisibilityComputeRasterPrepDepthGrid"), ComputeShaderPrepareDepthGrid, PrepDepthGridParameters, FIntVector(TileGridRes.X, TileGridRes.Y, 1));
 	}
 
 	const bool bDebugEnabled = GHairVisibilityComputeRaster_Debug > 0;
@@ -3458,7 +3530,7 @@ static FRasterComputeOutput AddVisibilityComputeRasterPass(
 					});
 			}
 
-			if (bClassification || !bUseHWRaster)
+			if (bMultiSampleSWRaster || bClassification || !bUseHWRaster)
 			{
 				if (bUseNaiveSWRaster)
 				{
@@ -3494,6 +3566,10 @@ static FRasterComputeOutput AddVisibilityComputeRasterPass(
 						uint32 IndexGridClearValues[4] = { 0x0,0x0,0x0,0x0 };
 						AddClearUAVPass(GraphBuilder, VisTileBinningGridUAV, IndexGridClearValues);
 						AddClearUAVPass(GraphBuilder, VisTileArgsUAV, 0u);
+						if (bTileCompactionPass)
+						{
+							AddClearUAVPass(GraphBuilder, GraphBuilder.CreateUAV(CompactedVisTileArgs, PF_R32_UINT), 0u);
+						}
 					}
 
 					// Binning pass
@@ -3541,6 +3617,26 @@ static FRasterComputeOutput AddVisibilityComputeRasterPass(
 					}
 
 					// Raster pass
+					if (bMultiSampleSWRaster)
+					{
+						FVisiblityRasterComputeRasterizeMultiSampleCS::FParameters* RasterParameters = GraphBuilder.AllocParameters<FVisiblityRasterComputeRasterizeMultiSampleCS::FParameters>();
+						RasterParameters->Common = Common;
+						RasterParameters->MacroGroupId = MacroGroup.MacroGroupId;
+						RasterParameters->HairMaterialId = PrimitiveInfo.MaterialId;
+						RasterParameters->ViewUniformBuffer = ViewUniformShaderParameters;
+						RasterParameters->SceneDepthTexture = SceneDepthTexture;
+						RasterParameters->OutHairCountTexture = HairCountTextureUAV;
+						RasterParameters->OutDepthCovTexture = DepthCovTextureUAV;
+						RasterParameters->OutPrimMatTexture = PrimMatTextureUAV;
+						RasterParameters->SampleWeight = SampleWeight;
+						RasterParameters->HairInstance = GetHairStrandsInstanceParameters(GraphBuilder, ViewInfo, HairGroupPublicData, bCullingEnableInRasterizers, bForceRegister);
+						RasterParameters->VisTilePrims = bTileCompactionPass ? GraphBuilder.CreateSRV(CompactedVisTilePrims, PF_R32_UINT) : VisTilePrimsSRV;
+						RasterParameters->VisTileArgs = bTileCompactionPass ? GraphBuilder.CreateSRV(CompactedVisTileArgs, PF_R32_UINT) : VisTileArgsSRV;
+						RasterParameters->VisTileData = bTileCompactionPass ? GraphBuilder.CreateSRV(CompactedVisTileData) : VisTileDataSRV;
+
+						FComputeShaderUtils::AddPass(GraphBuilder, RDG_EVENT_NAME("HairStrands::VisibilityComputeRasterRaster(tiled)"), ComputeShaderRasterMultiSample, RasterParameters, FIntVector(NumRasterizers, 1, 1));
+					}
+					else
 					{
 						FVisiblityRasterComputeRasterizeCS::FParameters* RasterParameters = GraphBuilder.AllocParameters<FVisiblityRasterComputeRasterizeCS::FParameters>();
 						RasterParameters->Common = Common;
