Fix out of bound view mode with Strata.

diff --git a/Engine/Shaders/Private/BasePassPixelShader.usf b/Engine/Shaders/Private/BasePassPixelShader.usf
--- a/Engine/Shaders/Private/BasePassPixelShader.usf
+++ b/Engine/Shaders/Private/BasePassPixelShader.usf
@@ -1392,10 +1392,12 @@ void FPixelShaderInOut_MainPS(
 
 #endif // !STRATA_ENABLED
 
+#if USE_DEVELOPMENT_SHADERS
+	float3 OutOfBoundsMaskLuminance = 0;
+	// This feature is only needed for development/editor - we can compile it out for a shipping build (see r.CompileShadersForDevelopment cvar help)
 	// The following block is disabled on Vulkan because it triggers an nvidia driver bug (UE-101609).
-#if USE_DEVELOPMENT_SHADERS && !VULKAN_PROFILE_SM5
-	// this feature is only needed for development/editor - we can compile it out for a shipping build (see r.CompileShadersForDevelopment cvar help)
-	#if METAL_SM5_PROFILE || SM6_PROFILE || SM5_PROFILE || VULKAN_PROFILE_SM5
+	#if !VULKAN_PROFILE_SM5
+		#if METAL_SM5_PROFILE || SM6_PROFILE || SM5_PROFILE || VULKAN_PROFILE_SM5
 		BRANCH
 		if (View.OutOfBoundsMask > 0)
 		{
@@ -1409,10 +1411,12 @@ void FPixelShaderInOut_MainPS(
 			if (any(abs(LWCToFloat(LWCSubtract(MaterialParameters.AbsoluteWorldPosition, GetPrimitiveData(MaterialParameters).ObjectWorldPosition))) > ObjectBounds + 1))
 			{
 				float Gradient = LWCFrac(LWCDivide(LWCDot(MaterialParameters.AbsoluteWorldPosition, float3(.577f, .577f, .577f)), 500.0f));
-				Emissive = lerp(float3(1,1,0), float3(0,1,1), Gradient.xxx > .5f);
+				OutOfBoundsMaskLuminance = lerp(float3(1,1,0), float3(0,1,1), Gradient.xxx > .5f);
+				Emissive = OutOfBoundsMaskLuminance;
 				Opacity = 1;
 			}
 		}
+		#endif
 	#endif
 #endif
 
@@ -1690,7 +1694,8 @@ void FPixelShaderInOut_MainPS(
 
 		// Unlit view mode
 	#if USE_DEVELOPMENT_SHADERS
-		Color = lerp(Color, TopLayerData.BaseColor, View.UnlitViewmodeMask);
+		Color  = lerp(Color, TopLayerData.BaseColor, View.UnlitViewmodeMask);
+		EmissiveLuminance += OutOfBoundsMaskLuminance;
 	#endif
 
 		Color += EmissiveLuminance;
diff --git a/Engine/Shaders/Private/Strata/StrataVisualize.usf b/Engine/Shaders/Private/Strata/StrataVisualize.usf
--- a/Engine/Shaders/Private/Strata/StrataVisualize.usf
+++ b/Engine/Shaders/Private/Strata/StrataVisualize.usf
@@ -257,9 +257,9 @@ void PrintBSDF(inout FShaderPrintContext Context, inout FStrataAddressing Strata
 				FHaziness Haziness = UnpackHaziness(SLAB_HAZINESS(BSDF));
 
 				Print(Context, TEXT(" Haziness    "),			FontBSDFPropName);																		Newline(Context, RectMax);
-				Print(Context, TEXT("     Roughness    "),		FontBSDFPropName);	Print(Context, Haziness.Roughness,				FontBSDFPropValu);	Newline(Context, RectMax);
-				Print(Context, TEXT("     Weight       "),		FontBSDFPropName);	Print(Context, Haziness.Weight,					FontBSDFPropValu);	Newline(Context, RectMax);
-				Print(Context, TEXT("     SimpleClearCoat"),	FontBSDFPropName);	Print(Context, Haziness.bSimpleClearCoat?1u:0u,	FontBSDFPropValu);	Newline(Context, RectMax);
+				Print(Context, TEXT("  - Roughness       "),	FontBSDFPropName);	Print(Context, Haziness.Roughness,				FontBSDFPropValu);	Newline(Context, RectMax);
+				Print(Context, TEXT("  - Weight          "),	FontBSDFPropName);	Print(Context, Haziness.Weight,					FontBSDFPropValu);	Newline(Context, RectMax);
+				Print(Context, TEXT("  - SimpleClearCoat "),	FontBSDFPropName);	Print(Context, Haziness.bSimpleClearCoat?1u:0u,	FontBSDFPropValu);	Newline(Context, RectMax);
 			}
 
 			if (BSDF_GETSSSTYPE(BSDF) != SSS_TYPE_INVALID || BSDF_GETISSIMPLEVOLUME(BSDF) || BSDF_GETISTHIN(BSDF))
