Fix for shader compile error in HairStrandsVisibilityRasterCompute.usf

#rnx

diff --git a/Engine/Shaders/Private/HairStrands/HairStrandsVisibilityRasterCompute.usf b/Engine/Shaders/Private/HairStrands/HairStrandsVisibilityRasterCompute.usf
--- a/Engine/Shaders/Private/HairStrands/HairStrandsVisibilityRasterCompute.usf
+++ b/Engine/Shaders/Private/HairStrands/HairStrandsVisibilityRasterCompute.usf
@@ -1391,7 +1391,7 @@ void RasterMultiSampleCS(uint DispatchThreadID : SV_DispatchThreadID, uint Group
 
 				UNROLL for (uint SampleIdx = 0; SampleIdx < PERMUTATION_MULTI_SAMPLE_COUNT; ++SampleIdx)
 				{
-					const uint HairDepth = OutDepthCovTexture.Load(uint3(Coord, SampleIdx));
+					const uint HairDepth = OutDepthCovTexture[uint3(Coord, SampleIdx)];
 					group_SubTileHairDepth[SampleIdx][GroupThreadID] = HairDepth;
 					group_SubTilePrimMatID[SampleIdx][GroupThreadID] = GetInvalidHairPrimitiveId();
 				}
