Strata - fixed forward rendering artefact with rect light when in fact those lights are not correctly handled today in forward. (legacy also renders them black, Strata was rendering them as tube lights)

diff --git a/Engine/Shaders/Private/Strata/StrataForwardLighting.ush b/Engine/Shaders/Private/Strata/StrataForwardLighting.ush
--- a/Engine/Shaders/Private/Strata/StrataForwardLighting.ush
+++ b/Engine/Shaders/Private/Strata/StrataForwardLighting.ush
@@ -13,6 +13,42 @@ void GetVolumeLightingDirectional(float4 AmbientLightingVector, float3 Direction
 
 float3 GetWorldBentNormalZero(in FMaterialPixelParameters MaterialParameters);
 
+FDeferredLightData ConvertLocalToDeferredLight(FLocalLightData LocalLight)
+{
+	FDeferredLightData LightData = (FDeferredLightData)0;
+	LightData.TranslatedWorldPosition = LocalLight.LightPositionAndInvRadius.xyz;
+	LightData.InvRadius = LocalLight.LightPositionAndInvRadius.w;
+	LightData.Color = LocalLight.LightColorAndFalloffExponent.xyz;
+	LightData.FalloffExponent = LocalLight.LightColorAndFalloffExponent.w;
+	LightData.Direction = LocalLight.LightDirectionAndShadowMask.xyz;
+	LightData.SpotAngles = LocalLight.SpotAnglesAndSourceRadiusPacked.xy;
+	LightData.SourceRadius = LocalLight.SpotAnglesAndSourceRadiusPacked.z;
+	LightData.SourceLength = f16tof32(asuint(LocalLight.SpotAnglesAndSourceRadiusPacked.w));
+	LightData.Tangent = LocalLight.LightTangentAndSoftSourceRadius.xyz;
+	LightData.SoftSourceRadius = LocalLight.LightTangentAndSoftSourceRadius.w;
+	LightData.bInverseSquared = LightData.FalloffExponent == 0;
+	LightData.bRadialLight = true;
+	LightData.bSpotLight = LightData.SpotAngles.x > -2.0f;
+	LightData.SpecularScale = LocalLight.RectBarnDoorAndVirtualShadowMapIdAndSpecularScale.w;
+	// RectLight is always false due to initialisation
+	// Initialize rect light once parameters are supported onto FLocalLightData
+	// LightData.RectLightAtlasMaxLevel = LocalLight.RectLightAtlasMaxLevel;
+	// LightData.RectLightAtlasUVOffset = LocalLight.RectLightAtlasUVOffset;
+	// LightData.RectLightAtlasUVScale  = LocalLight.RectLightAtlasUVScale;
+	// LightData.RectLightBarnCosAngle	= DeferredLightUniforms.RectLightBarnCosAngle;
+	// LightData.RectLightBarnLength	= DeferredLightUniforms.RectLightBarnLength;
+	LightData.HairTransmittance = InitHairTransmittanceData();
+
+	// LightType=bits[17:16], LightingChannelMask=[15:8], DynShadowMask=[7:4] ShadowMapChannelMask=[3:0]
+	uint LightTypeAndPackedShadowMapChannelMask = asuint(LocalLight.LightDirectionAndShadowMask.w);
+
+	LightData.ShadowedBits = (LightTypeAndPackedShadowMapChannelMask & 0xFF) != 0 ? 1 : 0;
+	// Static shadowing uses ShadowMapChannel, dynamic shadows are packed into light attenuation using PreviewShadowMapChannel
+	LightData.ShadowMapChannelMask = UnpackShadowMapChannelMask(LightTypeAndPackedShadowMapChannelMask);
+
+	return LightData;
+}
+
 FShadingOcclusion ApplyBentNormal(
 	in float3 CameraVector,
 	in float3 WorldNormal,
@@ -409,33 +445,19 @@ float3 StrataForwardLighting(
 				{
 					const FLocalLightData LocalLight = GetLocalLightData(CulledLightsGrid.DataStartIndex + LocalLightListIndex, EyeIndex);
 
-					FDeferredLightData LightData = (FDeferredLightData)0;
-					LightData.TranslatedWorldPosition = LocalLight.LightPositionAndInvRadius.xyz;
-					LightData.InvRadius = LocalLight.LightPositionAndInvRadius.w;
-					LightData.Color = LocalLight.LightColorAndFalloffExponent.xyz;
-					LightData.FalloffExponent = LocalLight.LightColorAndFalloffExponent.w;
-					LightData.Direction = LocalLight.LightDirectionAndShadowMask.xyz;
-					LightData.SpotAngles = LocalLight.SpotAnglesAndSourceRadiusPacked.xy;
-					LightData.SourceRadius = LocalLight.SpotAnglesAndSourceRadiusPacked.z;
-					LightData.SourceLength = f16tof32(asuint(LocalLight.SpotAnglesAndSourceRadiusPacked.w));
-					LightData.Tangent = LocalLight.LightTangentAndSoftSourceRadius.xyz;
-					LightData.SoftSourceRadius = LocalLight.LightTangentAndSoftSourceRadius.w;
-					LightData.bInverseSquared = LightData.FalloffExponent == 0;
-					LightData.bRadialLight = true;
-					LightData.bSpotLight = LightData.SpotAngles.x > -2.0f;
-					LightData.SpecularScale = SpecularScale;
-					LightData.HairTransmittance = InitHairTransmittanceData();
-
 					// LightType=bits[17:16], LightingChannelMask=[15:8], DynShadowMask=[7:4] ShadowMapChannelMask=[3:0]
-					uint LightTypeAndPackedShadowMapChannelMask = asuint(LocalLight.LightDirectionAndShadowMask.w);
+					const uint LightTypeAndPackedShadowMapChannelMask = asuint(LocalLight.LightDirectionAndShadowMask.w);
+
+					// Skip rect light. This is not done explicitely in the legacy path (resulting in no lighting) but we do it for strata in order to avoid visual artefact.
+					// This can be removed when LocalLight properly supports rect lights.
+					const uint LightType = UnpackLightType(LightTypeAndPackedShadowMapChannelMask);
+					if (LightType == LIGHT_TYPE_RECT)
+					{
+						continue;
+					}
 
-					// TODO: currently LightType is ignored and the code always sets 'LightData.bRectLight' to false (initialization)
-					//       but the light grid adds rect lights as well so this probably leads to artifacts.
-					uint LightType = UnpackLightType(LightTypeAndPackedShadowMapChannelMask);
+					const FDeferredLightData LightData = ConvertLocalToDeferredLight(LocalLight);
 
-					LightData.ShadowedBits = (LightTypeAndPackedShadowMapChannelMask & 0xFF) != 0 ? 1 : 0;
-					// Static shadowing uses ShadowMapChannel, dynamic shadows are packed into light attenuation using PreviewShadowMapChannel
-					LightData.ShadowMapChannelMask = UnpackShadowMapChannelMask(LightTypeAndPackedShadowMapChannelMask);
 					float4 PreviewShadowMapChannelMask = UnpackShadowMapChannelMask(LightTypeAndPackedShadowMapChannelMask >> 4);
 					float DynamicShadowing = dot(PreviewShadowMapChannelMask, DynamicShadowFactors);
 					float4 LightAttenuation = float4(1, 1, DynamicShadowing.x, DynamicShadowing.x);
