Changed FScreenVertexOutput to always use floating point as it causes a mismatch between vertex and fragment shader.

diff --git a/Engine/Shaders/Private/Common.ush b/Engine/Shaders/Private/Common.ush
--- a/Engine/Shaders/Private/Common.ush
+++ b/Engine/Shaders/Private/Common.ush
@@ -1691,12 +1691,7 @@ void GenerateCoordinateSystem(float3 ZAxis, out float3 XAxis, out float3 YAxis)
 /** Output of the screen vertex shader. */
 struct FScreenVertexOutput
 {
-#if METAL_PROFILE || COMPILER_GLSL_ES3_1
-	//@todo-rco: FIXME!
 	noperspective float2 UV : TEXCOORD0;
-#else
-	noperspective MaterialFloat2 UV : TEXCOORD0;
-#endif
 	float4 Position : SV_POSITION;
 };
 
