Strata - small CL for a tiny help to compilation time + branch/comment clean up.

diff --git a/Engine/Shaders/Private/Strata/StrataEvaluation.ush b/Engine/Shaders/Private/Strata/StrataEvaluation.ush
--- a/Engine/Shaders/Private/Strata/StrataEvaluation.ush
+++ b/Engine/Shaders/Private/Strata/StrataEvaluation.ush
@@ -175,6 +175,7 @@ void StrataUpdateBSDFContext(inout FStrataBSDFContext BSDFContext, float3 NewL)
 
 float3 LuminanceWeight(in float SatNoL, in FStrataBSDF InBSDF)
 {
+#if STRATA_COMPLEXPATH
 	// LuminanceWeight is the absorption from the shading point towards the surface 
 	// in the normal direction. To compute the absorption towards the light, we 
 	// reweight it with new distance L'
@@ -186,6 +187,9 @@ float3 LuminanceWeight(in float SatNoL, in FStrataBSDF InBSDF)
 	const float DistanceL = rcp(max(SatNoL, 0.001f));
 	const float3 TransmittanceAboveAlongL = any(InBSDF.TransmittanceAboveAlongN < 1.f) ? pow(max(InBSDF.TransmittanceAboveAlongN, 0.0001f), DistanceL) : InBSDF.TransmittanceAboveAlongN;
 	return InBSDF.LuminanceWeightV * TransmittanceAboveAlongL;
+#else
+	return InBSDF.LuminanceWeightV;
+#endif
 }
 
 float3 LuminanceWeight(in FStrataBSDFContext InContext, in FStrataBSDF InBSDF)
@@ -369,6 +373,7 @@ FStrataEvaluateResult StrataEvaluateBSDFCommon(FStrataBSDFContext BSDFContext, F
 				else
 				#endif
 				{
+					BRANCH
 					if (bIsRectLight)
 					{
 						// In this case, we set D to 1 and Vis will contain the area light / GGX integration
@@ -424,6 +429,7 @@ FStrataEvaluateResult StrataEvaluateBSDFCommon(FStrataBSDFContext BSDFContext, F
 			// Secondary highlight
 			float HazePDF = 0;
 			#if STRATA_FASTPATH==0
+			BRANCH
 			if (bHaziness)
 			{
 				float3 HazeF0 = F0;
@@ -434,7 +440,6 @@ FStrataEvaluateResult StrataEvaluateBSDFCommon(FStrataBSDFContext BSDFContext, F
 				const float HazeSafeRoughness = MakeRoughnessSafe(Haziness.Roughness);
 				const bool bHazeAsSimpleClearCoat = Haziness.bSimpleClearCoat;
 
-				BRANCH
 				if (bHazeAsSimpleClearCoat)
 				{
 					HazeF0 = 0.04f;
@@ -884,7 +889,6 @@ FStrataEvaluateResult StrataEvaluateBSDFCommon(FStrataBSDFContext BSDFContext, F
 					const float Energy = EnergyNormalization(Alpha2Spec, BSDFContext.Context.VoH, AreaLightContext.AreaLight);
 					D = D_GGX(Alpha2Spec, BSDFContext.Context.NoH) * Energy;
 				}
-				// STRATA_TODO Should we use Vis_SmithJointApprox here?
 				Vis = StrataVisSmithJoint(Alpha2Spec, BSDFContext.Context.NoV, AreaLightContext.NoL);
 				const float H_PDF = VisibleGGXPDF(BSDFContext.TangentV, BSDFContext.TangentH, Alpha2Spec);
 				PDF = RayPDFToReflectionRayPDF(BSDFContext.Context.VoH, H_PDF);
