Strata - Optimisations

StrataIntegrateBSDF know which branch to take explicitely now from a static parameter set on the calling site.
SGPRs 96 to 88 for light / material single interactions
Base pass: -0.1ms on test scene.
Some shaders VGPRs comes down by 4 and size reduced by roughly 700bytes

diff --git a/Engine/Shaders/Private/Strata/Strata.ush b/Engine/Shaders/Private/Strata/Strata.ush
--- a/Engine/Shaders/Private/Strata/Strata.ush
+++ b/Engine/Shaders/Private/Strata/Strata.ush
@@ -1321,7 +1321,7 @@ FStrataData GetStrataSlabBSDF(
 	FHaziness Haziness = InitialiseHaziness();
 
 	// Derive roughness and haziness from profile data
-	if (bSSSProfilePlugged)
+	if (!STRATA_FASTPATH && bSSSProfilePlugged)
 	{
 		// Average roughness for dual specular.
 		const uint SubsurfaceProfileUInt = StrataSubsurfaceProfileIdTo8bits(SSSProfileID);
@@ -1340,7 +1340,7 @@ FStrataData GetStrataSlabBSDF(
 	// If MFP is smaller than the pixel's footprint, remove SSS attributes
 	// The scaling factor 'SSSFootprintScale' is based on empiric tests
 	const float SSSFootprintScale = 3.f;
-	if (all(SSSMFP * SSSFootprintScale < InPixelFootprint.PixelRadiusInWorldSpace) && !bHasSSSWrap)
+	if (!STRATA_FASTPATH && all(SSSMFP * SSSFootprintScale < InPixelFootprint.PixelRadiusInWorldSpace) && !bHasSSSWrap)
 	{
 		bHasSSSProfile = false;
 		bHasSSS = false;
@@ -1350,7 +1350,7 @@ FStrataData GetStrataSlabBSDF(
 	}
 
 	// Convert MFP be expressed relative to simple volume thickness
-	if (bHasSSSWrap)
+	if (!STRATA_FASTPATH && bHasSSSWrap)
 	{
 		SSSMFP = RescaleMFPToComputationSpace(SSSMFP, Thickness, STRATA_SIMPLEVOLUME_THICKNESS_CM);
 		Thickness = STRATA_SIMPLEVOLUME_THICKNESS_CM;
@@ -1359,7 +1359,7 @@ FStrataData GetStrataSlabBSDF(
 	const uint bHasHaziness = ((Haziness.Roughness != Roughness  && Haziness.Weight > 0.0f) || Haziness.bSimpleClearCoat) ? 1 : 0;
 	const uint bHasAnisotropy = Anisotropy != 0 ? 1 : 0;
 	const uint bHasFuzz = FuzzAmount > 0.0f ? 1 : 0;
-	if (bHasAnisotropy)
+	if (!STRATA_FASTPATH && bHasAnisotropy)
 	{
 		StrataRequestSharedLocalBasisTangent(SharedLocalBasisTypes, SharedLocalBasisIndex);
 	}
@@ -1375,18 +1375,24 @@ FStrataData GetStrataSlabBSDF(
 	BSDF_SETTYPE			(StrataData.InlinedBSDF,		STRATA_BSDF_TYPE_SLAB);
 	BSDF_SETSHAREDLOCALBASISID(StrataData.InlinedBSDF,		SharedLocalBasisIndex);
 	BSDF_SETEMISSIVE		(StrataData.InlinedBSDF,		Emissive);
+
+#if !STRATA_FASTPATH
+	BSDF_SETSSSTYPE			(StrataData.InlinedBSDF,		bHasSSS ? (bHasSSSWrap ? SSS_TYPE_WRAP : (bHasSSSProfile ? SSS_TYPE_DIFFUSION_PROFILE : SSS_TYPE_DIFFUSION)) : SSS_TYPE_INVALID);
 	BSDF_SETHASANISOTROPY	(StrataData.InlinedBSDF,		bHasAnisotropy);
 	BSDF_SETISTOPLAYER		(StrataData.InlinedBSDF,		0);
-	BSDF_SETSSSTYPE			(StrataData.InlinedBSDF,		bHasSSS ? (bHasSSSWrap ? SSS_TYPE_WRAP : (bHasSSSProfile ? SSS_TYPE_DIFFUSION_PROFILE : SSS_TYPE_DIFFUSION)) : SSS_TYPE_INVALID);
 	BSDF_SETISTHIN			(StrataData.InlinedBSDF,		bIsThin ? 1 : 0);
 	BSDF_SETHASMFP			(StrataData.InlinedBSDF,		bMFPPlugged ? 1 : 0);
 	BSDF_SETHASHAZINESS		(StrataData.InlinedBSDF,		bHasHaziness);
-	BSDF_SETTHICKNESSCM		(StrataData.InlinedBSDF,		Thickness);
 	BSDF_SETHASFUZZ			(StrataData.InlinedBSDF,		bHasFuzz);
+	BSDF_SETTHICKNESSCM		(StrataData.InlinedBSDF,		Thickness);
+#endif
+
 
 	if (UseMetalness > 0.0f)
 	{
+#if !STRATA_FASTPATH
 		BSDF_SETHASF90(StrataData.InlinedBSDF, any(EdgeColor < 1.0f));
+#endif
 
 		const float SafeMetallic = saturate(Metallic);
 		SLAB_DIFFUSEALBEDO(StrataData.InlinedBSDF)			= ComputeDiffuseAlbedo(BaseColor, SafeMetallic);
@@ -1395,8 +1401,10 @@ FStrataData GetStrataSlabBSDF(
 	}
 	else
 	{
+#if !STRATA_FASTPATH
 		// See how UMaterialExpressionStrataSlabBSDF node interface is changed
 		BSDF_SETHASF90(StrataData.InlinedBSDF, any(F90 < 1.0f));
+#endif
 
 		SLAB_DIFFUSEALBEDO(StrataData.InlinedBSDF)			= DiffuseAlbedo;
 		SLAB_F0(StrataData.InlinedBSDF)						= F0;
@@ -1404,6 +1412,8 @@ FStrataData GetStrataSlabBSDF(
 	}
 
 	SLAB_ROUGHNESS			(StrataData.InlinedBSDF)		= Roughness;
+
+#if !STRATA_FASTPATH
 	SLAB_ANISOTROPY			(StrataData.InlinedBSDF)		= Anisotropy;
 
 	SLAB_HAZINESS(StrataData.InlinedBSDF)					= PackHaziness(Haziness);
@@ -1423,6 +1433,8 @@ FStrataData GetStrataSlabBSDF(
 	SLAB_FUZZ_AMOUNT		(StrataData.InlinedBSDF)		= FuzzAmount;
 	SLAB_FUZZ_COLOR			(StrataData.InlinedBSDF)		= FuzzColor;
 
+#endif // !STRATA_FASTPATH
+
 #if STRATA_INLINE_SHADING
 	// At this stage we do not know if this strata BSDF is going to be using simple volume lighting so always store the MFP on the side to not override SSSPROFILE data.
 	StrataData.InlinedBSDF.TmpMFP = SSSMFP * SSSMFPScale;
diff --git a/Engine/Shaders/Private/Strata/StrataDeferredLighting.ush b/Engine/Shaders/Private/Strata/StrataDeferredLighting.ush
--- a/Engine/Shaders/Private/Strata/StrataDeferredLighting.ush
+++ b/Engine/Shaders/Private/Strata/StrataDeferredLighting.ush
@@ -70,14 +70,22 @@ FStrataDeferredLighting StrataDeferredLighting(
 
 			// We must have the StrataIntegrateBSDF inside the if due to the rectlight texture: it must be non ambiguous which texture is going to be used.
 			// After the compilation, a local resource must map to a unique global resource (the default or the actual rect light texture).
-			BSDFEvaluate = StrataIntegrateBSDF(BSDFContext, ShadowTerms, AreaLightContext, Settings);
+			BSDFEvaluate = StrataIntegrateBSDF(BSDFContext, ShadowTerms, AreaLightContext, Settings, INTEGRATION_AREA_LIGHT_RECT);
 		}
 		else
 		{
 			FCapsuleLight Capsule = GetCapsule(ToLight, LightData);
 			AreaLightContext = CreateCapsuleIntegrateContext(Roughness, BSDFContext.N, BSDFContext.V, Capsule, LightData.bInverseSquared);
 
-			BSDFEvaluate = StrataIntegrateBSDF(BSDFContext, ShadowTerms, AreaLightContext, Settings);
+			BRANCH
+			if(IsAreaLight(AreaLightContext.AreaLight))
+			{
+				BSDFEvaluate = StrataIntegrateBSDF(BSDFContext, ShadowTerms, AreaLightContext, Settings, INTEGRATION_AREA_LIGHT_CAPSULE);
+			}
+			else
+			{
+				BSDFEvaluate = StrataIntegrateBSDF(BSDFContext, ShadowTerms, AreaLightContext, Settings, INTEGRATION_PUNCTUAL_LIGHT);
+			}
 		}
 
 		float3 DiffuseLuminance = BSDFEvaluate.IntegratedDiffuseValue;
diff --git a/Engine/Shaders/Private/Strata/StrataEvaluation.ush b/Engine/Shaders/Private/Strata/StrataEvaluation.ush
--- a/Engine/Shaders/Private/Strata/StrataEvaluation.ush
+++ b/Engine/Shaders/Private/Strata/StrataEvaluation.ush
@@ -948,20 +948,16 @@ FStrataEvaluateResult StrataEvaluateBSDF(FStrataBSDFContext BSDFContext, FStrata
 }
 
 // Area light evaluation
-FStrataEvaluateResult StrataIntegrateBSDF(FStrataBSDFContext BSDFContext, FShadowTerms Shadow, FAreaLightIntegrateContext AreaLightContext, FStrataIntegrationSettings Settings)
+FStrataEvaluateResult StrataIntegrateBSDF(FStrataBSDFContext BSDFContext, FShadowTerms Shadow, FAreaLightIntegrateContext AreaLightContext, FStrataIntegrationSettings Settings, uint IntegrationMode)
 {
 	BRANCH
-	if (IsAreaLight(AreaLightContext.AreaLight))
+	if (IntegrationMode == INTEGRATION_AREA_LIGHT_RECT)
 	{
-		BRANCH
-		if(IsRectLight(AreaLightContext.AreaLight))
-		{
-			return StrataEvaluateBSDFCommon(BSDFContext, Shadow, AreaLightContext, Settings, INTEGRATION_AREA_LIGHT_RECT);
-		}
-		else
-		{
-			return StrataEvaluateBSDFCommon(BSDFContext, Shadow, AreaLightContext, Settings, INTEGRATION_AREA_LIGHT_CAPSULE);
-		}
+		return StrataEvaluateBSDFCommon(BSDFContext, Shadow, AreaLightContext, Settings, INTEGRATION_AREA_LIGHT_RECT);
+	}
+	else if (IntegrationMode == INTEGRATION_AREA_LIGHT_CAPSULE)
+	{
+		return StrataEvaluateBSDFCommon(BSDFContext, Shadow, AreaLightContext, Settings, INTEGRATION_AREA_LIGHT_CAPSULE);
 	}
 
 	return StrataEvaluateBSDFCommon(BSDFContext, Shadow, AreaLightContext, Settings, INTEGRATION_PUNCTUAL_LIGHT);
diff --git a/Engine/Shaders/Private/Strata/StrataLightingCommon.ush b/Engine/Shaders/Private/Strata/StrataLightingCommon.ush
--- a/Engine/Shaders/Private/Strata/StrataLightingCommon.ush
+++ b/Engine/Shaders/Private/Strata/StrataLightingCommon.ush
@@ -64,7 +64,7 @@ float3 StrataForwardLightingCommon(
 
 				// We must have the evaluate inside the if due to the rectlight texture: it must be now be ambiguous which texture is going ot be used.
 				// After te compilation, a local resource must map to a unique global resource (the default or the actual rect light texture).
-				BSDFEvaluate = StrataIntegrateBSDF(StrataBSDFContext, ShadowTerms, AreaLightContext, Settings);
+				BSDFEvaluate = StrataIntegrateBSDF(StrataBSDFContext, ShadowTerms, AreaLightContext, Settings, INTEGRATION_AREA_LIGHT_RECT);
 				bEvaluateHasBeenExecuted = true;
 			}
 			else
@@ -73,7 +73,15 @@ float3 StrataForwardLightingCommon(
 				AreaLightContext = CreateCapsuleIntegrateContext(Roughness, StrataBSDFContext.N, StrataBSDFContext.V, Capsule, LightData.bInverseSquared);
 				Lighting = IntegrateLight(Capsule, LightData.bInverseSquared);
 
-				BSDFEvaluate = StrataIntegrateBSDF(StrataBSDFContext, ShadowTerms, AreaLightContext, Settings);
+				BRANCH
+				if(IsAreaLight(AreaLightContext.AreaLight))
+				{
+					BSDFEvaluate = StrataIntegrateBSDF(StrataBSDFContext, ShadowTerms, AreaLightContext, Settings, INTEGRATION_AREA_LIGHT_CAPSULE);
+				}
+				else
+				{
+					BSDFEvaluate = StrataIntegrateBSDF(StrataBSDFContext, ShadowTerms, AreaLightContext, Settings, INTEGRATION_PUNCTUAL_LIGHT);
+				}
 				bEvaluateHasBeenExecuted = true;
 			}
 
@@ -96,15 +104,23 @@ float3 StrataForwardLightingCommon(
 
 				// We must have the evaluate inside the if due to the rectlight texture: it must be now be ambiguous which texture is going ot be used.
 				// After te compilation, a local resource must map to a unique global resource (the default or the actual rect light texture).
-				BSDFEvaluate = StrataIntegrateBSDF(StrataBSDFContext, ShadowTerms, AreaLightContext, Settings);
+				BSDFEvaluate = StrataIntegrateBSDF(StrataBSDFContext, ShadowTerms, AreaLightContext, Settings, INTEGRATION_AREA_LIGHT_RECT);
 				bEvaluateHasBeenExecuted = true;
 			}
 			else
 			{
 				FCapsuleLight Capsule = GetCapsule(ToLight, LightData);
 				AreaLightContext = CreateCapsuleIntegrateContext(Roughness, StrataBSDFContext.N, StrataBSDFContext.V, Capsule, LightData.bInverseSquared);
-
-				BSDFEvaluate = StrataIntegrateBSDF(StrataBSDFContext, ShadowTerms, AreaLightContext, Settings);
+				
+				BRANCH
+				if(IsAreaLight(AreaLightContext.AreaLight))
+				{
+					BSDFEvaluate = StrataIntegrateBSDF(StrataBSDFContext, ShadowTerms, AreaLightContext, Settings, INTEGRATION_AREA_LIGHT_CAPSULE);
+				}
+				else
+				{
+					BSDFEvaluate = StrataIntegrateBSDF(StrataBSDFContext, ShadowTerms, AreaLightContext, Settings, INTEGRATION_PUNCTUAL_LIGHT);
+				}
 				bEvaluateHasBeenExecuted = true;
 			}
 
diff --git a/Engine/Shaders/Private/Strata/StrataTree.ush b/Engine/Shaders/Private/Strata/StrataTree.ush
--- a/Engine/Shaders/Private/Strata/StrataTree.ush
+++ b/Engine/Shaders/Private/Strata/StrataTree.ush
@@ -80,7 +80,12 @@ void UpdateSingleBSDFOperatorCoverageTransmittance(
 	CurrentBSDF.Coverage = int(Op.LayerDepth) < Settings.PeelLayersAboveDepth && !CurrentBSDF.bIsBottom ? 0.0f : CurrentBSDF.Coverage;
 #endif
 
+	float CurrentBSDFCoverage = CurrentBSDF.Coverage;
 	const float Roughness = StrataGetBSDFRoughness(CurrentBSDF);
+	float3x3 TangentBasis = StrataGetBSDFSharedBasis_InlineShading(StrataPixelHeader, CurrentBSDF, NullStrataAddressing);
+
+	// Distortion and translucency passes accumnulate pass calls for colored transmittance, coverage and TopDownRefractionLobeStat to be setup correctly.
+#if STRATA_COMPLEXPATH || MATERIALBLENDING_ANY_TRANSLUCENT
 
 	// Extract the true thickness of the BSDF before it is cast into the normalised Slab thickness we used.
 	// for opaque rough refraction, we are only interested in thickness of non bottom layers.
@@ -91,13 +96,10 @@ void UpdateSingleBSDFOperatorCoverageTransmittance(
 	PostProcessBSDFBeforeLighting(CurrentBSDF);
 
 	// Compute current BSDF transmittance
-	float3x3 TangentBasis = StrataGetBSDFSharedBasis_InlineShading(StrataPixelHeader, CurrentBSDF, NullStrataAddressing);
 	float3 LEqualN = TangentBasis[2];
 	FStrataBSDFContext StrataBSDFContext = StrataCreateBSDFContext(StrataPixelHeader, CurrentBSDF, NullStrataAddressing, V, LEqualN);
 	FStrataEvaluateResult BSDFEvaluate = StrataEvaluateBSDF(StrataBSDFContext, Settings);
 
-	float CurrentBSDFCoverage		= CurrentBSDF.Coverage;
-
 	// We multiply this weight with coverage:
 	//  - Because when coming from parameter blending, we need to scale the output down according to coverage resulting from parameter blending.
 	//  - That will be applied on emissive and reflected light, so we do not need to apply that anywhere else.
@@ -136,6 +138,25 @@ void UpdateSingleBSDFOperatorCoverageTransmittance(
 		RefractedLobe = StrataWeightLobe(RefractedLobe, CurrentBSDFCoverage);
 		Op.BottomUpRefractionLobeStat = RefractedLobe;
 	}
+
+#else // STRATA_COMPLEXPATH
+
+	CurrentBSDF.LuminanceWeightV		= CurrentBSDFCoverage;
+	CurrentBSDF.TransmittanceAboveAlongN= float3(1.0f, 1.0f, 1.0f);
+	CurrentBSDF.TopLayerDataWeight		= 1.0f;
+
+	Op.OpaqueRoughRefractThicknessCm	= 0.0;
+	Op.OpaqueRoughRefractCoverage		= CurrentBSDFCoverage;
+	Op.OpaqueRoughRefractTopRoughness	= Roughness;
+
+	Op.Coverage				= CurrentBSDFCoverage;
+	Op.ThroughputAlongV		= 1.0f;
+	Op.TransmittanceAlongN	= 1.0f;
+
+	Op.TopDownRefractionLobeStat = StrataGetNullLobe();
+	Op.BottomUpRefractionLobeStat = StrataGetNullLobe();
+
+#endif // STRATA_COMPLEXPATH
 	
 	Op.TopDownRefractionWorldNormal = TangentBasis[2];
 
