Update SkyAtmosphere VSM sampling to reflect the change to translated world space coordinates (from CL18822479).

diff --git a/Engine/Shaders/Private/SkyAtmosphere.usf b/Engine/Shaders/Private/SkyAtmosphere.usf
--- a/Engine/Shaders/Private/SkyAtmosphere.usf
+++ b/Engine/Shaders/Private/SkyAtmosphere.usf
@@ -606,7 +606,7 @@ SingleScatteringResult IntegrateSingleScatteredLuminance(
 		#if VIRTUAL_SHADOW_MAP
 			if (VirtualShadowMapId0 != INDEX_NONE)
 			{
-				FVirtualShadowMapSampleResult VirtualShadowMapSample = SampleVirtualShadowMapLWCHack(VirtualShadowMapId0, ShadowSampleWorldPosition0);
+				FVirtualShadowMapSampleResult VirtualShadowMapSample = SampleVirtualShadowMapTranslatedWorld(VirtualShadowMapId0, ShadowSampleWorldPosition0);
 				PlanetShadow0 *= VirtualShadowMapSample.ShadowFactor;
 			}
 		#endif // VIRTUALSHADOW_MAP
@@ -640,7 +640,7 @@ SingleScatteringResult IntegrateSingleScatteredLuminance(
 		#if VIRTUAL_SHADOW_MAP
 			if (VirtualShadowMapId1 != INDEX_NONE)
 			{
-				FVirtualShadowMapSampleResult VirtualShadowMapSample = SampleVirtualShadowMapLWCHack(VirtualShadowMapId1, ShadowSampleWorldPosition1);
+				FVirtualShadowMapSampleResult VirtualShadowMapSample = SampleVirtualShadowMapTranslatedWorld(VirtualShadowMapId1, ShadowSampleWorldPosition1);
 				PlanetShadow1 *= VirtualShadowMapSample.ShadowFactor;
 			}
 		#endif // VIRTUALSHADOW_MAP
