Fix invalid material handling with opaque rough refraction when a material has several slabs, which are dynamically culled, and only one remain.

diff --git a/Engine/Shaders/Private/Strata/StrataExport.ush b/Engine/Shaders/Private/Strata/StrataExport.ush
--- a/Engine/Shaders/Private/Strata/StrataExport.ush
+++ b/Engine/Shaders/Private/Strata/StrataExport.ush
@@ -449,6 +449,10 @@ void PackStrataOut(
 					// B. Layout1 - Simple encoding: single BSDF, whose header & BSDF state are merged
 					else if (bHasSingleEncoding && BSDFIdx == OneBSDFMaterial_Index)
 					{
+						// Enforce BSDF to be flagged as 'top', since it is single-encoded (which ensure there is a single visible slab).
+						// This is needed when material contains other slabs which are dynamically culled based on coverage
+						BSDF_SETISTOPLAYER(BSDF, 1);
+
 						// Data0 (Header_State|Header_AO|Header_BSDFTypes|BSDF_State)
 						uint Out = 0;
 						HEADER_SETCOMMONSTATES(Out, StrataPixelHeader.State);
