Disable high quality EnvBRDF on mobile deferred, so the lighting could be exactly the same as mobile forward.

diff --git a/Engine/Shaders/Private/MobileBasePassPixelShader.usf b/Engine/Shaders/Private/MobileBasePassPixelShader.usf
--- a/Engine/Shaders/Private/MobileBasePassPixelShader.usf
+++ b/Engine/Shaders/Private/MobileBasePassPixelShader.usf
@@ -17,10 +17,10 @@
 #define MOBILE_QL_FORCE_DISABLE_PREINTEGRATEDGF 0
 #endif
 
-#define FORWARDSHADING_USE_HQ_ENV_BRDF (MATERIAL_USE_PREINTEGRATED_GF && !MOBILE_QL_FORCE_DISABLE_PREINTEGRATEDGF)
+#define MOBILE_USE_PREINTEGRATED_GF (MATERIAL_USE_PREINTEGRATED_GF && !MOBILE_QL_FORCE_DISABLE_PREINTEGRATEDGF)
 
 //use preintegrated GF lut for simple IBL
-#if FORWARDSHADING_USE_HQ_ENV_BRDF
+#if MOBILE_USE_PREINTEGRATED_GF
 #define PreIntegratedGF			MobileBasePass.PreIntegratedGFTexture
 #define PreIntegratedGFSampler	MobileBasePass.PreIntegratedGFSampler
 #endif
diff --git a/Engine/Shaders/Private/MobileLightingCommon.ush b/Engine/Shaders/Private/MobileLightingCommon.ush
--- a/Engine/Shaders/Private/MobileLightingCommon.ush
+++ b/Engine/Shaders/Private/MobileLightingCommon.ush
@@ -366,7 +366,7 @@ void AccumulateReflection(FGBufferData GBuffer
 		half LayerAttenuation = (1 - F);
 
 		// Fc * Vis
-	#if FORWARDSHADING_USE_HQ_ENV_BRDF
+	#if MOBILE_USE_PREINTEGRATED_GF
 		half2 AB = PreIntegratedGF.SampleLevel(PreIntegratedGFSampler, float2(NoV, GBuffer.Roughness), 0).rg;
 	#else
 		half2 AB = EnvBRDFApproxLazarov(GBuffer.Roughness, NoV);
diff --git a/Engine/Shaders/Private/ShadingModels.ush b/Engine/Shaders/Private/ShadingModels.ush
--- a/Engine/Shaders/Private/ShadingModels.ush
+++ b/Engine/Shaders/Private/ShadingModels.ush
@@ -39,8 +39,8 @@
 #endif
 
 //It's always 0 in mobile deferred lighting shaders
-#ifndef FORWARDSHADING_USE_HQ_ENV_BRDF
-#define  FORWARDSHADING_USE_HQ_ENV_BRDF 0
+#ifndef MOBILE_USE_PREINTEGRATED_GF
+#define  MOBILE_USE_PREINTEGRATED_GF 0
 #endif
 
 //It's always 1 in mobile deferred lighting shaders
@@ -66,7 +66,7 @@ half3 GetEnvBRDF(half3 SpecularColor, half Roughness, half NoV)
 {
 #if FULLY_ROUGH
 	return 0.0f;
-#elif FORWARDSHADING_USE_HQ_ENV_BRDF || MOBILE_DEFERRED_LIGHTING
+#elif MOBILE_USE_PREINTEGRATED_GF
 	return EnvBRDF(SpecularColor, Roughness, NoV);
 #elif NONMETAL
 	// If nothing is hooked up to Metalic and Specular,
