Avoid calling GetPrimitiveData() redundantly. This function is large inlining and contributes significantly to the shader compilation cost.

diff --git a/Engine/Shaders/Private/BasePassPixelShader.usf b/Engine/Shaders/Private/BasePassPixelShader.usf
--- a/Engine/Shaders/Private/BasePassPixelShader.usf
+++ b/Engine/Shaders/Private/BasePassPixelShader.usf
@@ -1401,14 +1401,16 @@ void FPixelShaderInOut_MainPS(
 		BRANCH
 		if (View.OutOfBoundsMask > 0)
 		{
+			FPrimitiveSceneData PrimitiveData = GetPrimitiveData(MaterialParameters);
+
 			float3 ObjectBounds =
 				float3(
-					GetPrimitiveData(MaterialParameters).ObjectBoundsX,
-					GetPrimitiveData(MaterialParameters).ObjectBoundsY,
-					GetPrimitiveData(MaterialParameters).ObjectBoundsZ
+					PrimitiveData.ObjectBoundsX,
+					PrimitiveData.ObjectBoundsY,
+					PrimitiveData.ObjectBoundsZ
 				);
 
-			if (any(abs(LWCToFloat(LWCSubtract(MaterialParameters.AbsoluteWorldPosition, GetPrimitiveData(MaterialParameters).ObjectWorldPosition))) > ObjectBounds + 1))
+			if (any(abs(LWCToFloat(LWCSubtract(MaterialParameters.AbsoluteWorldPosition, PrimitiveData.ObjectWorldPosition))) > ObjectBounds + 1))
 			{
 				float Gradient = LWCFrac(LWCDivide(LWCDot(MaterialParameters.AbsoluteWorldPosition, float3(.577f, .577f, .577f)), 500.0f));
 				OutOfBoundsMaskLuminance = lerp(float3(1,1,0), float3(0,1,1), Gradient.xxx > .5f);
diff --git a/Engine/Shaders/Private/HeterogeneousVolumes/HeterogeneousVolumesLiveShadingPipeline.usf b/Engine/Shaders/Private/HeterogeneousVolumes/HeterogeneousVolumesLiveShadingPipeline.usf
--- a/Engine/Shaders/Private/HeterogeneousVolumes/HeterogeneousVolumesLiveShadingPipeline.usf
+++ b/Engine/Shaders/Private/HeterogeneousVolumes/HeterogeneousVolumesLiveShadingPipeline.usf
@@ -67,11 +67,13 @@ FDebugOutput CreateDebugOutput(
 {
 	FDebugOutput DebugOutput = (FDebugOutput)0;
 
+	FPrimitiveSceneData PrimitiveData = GetPrimitiveData(MaterialParameters);
+
 	DebugOutput.PrimitiveId = MaterialParameters.PrimitiveId;
-	DebugOutput.LocalObjectBoundsMin = GetPrimitiveData(MaterialParameters).LocalObjectBoundsMin.xyz;
-	DebugOutput.LocalObjectBoundsMax = GetPrimitiveData(MaterialParameters).LocalObjectBoundsMax.xyz;
-	DebugOutput.LocalToWorld = LWCToFloat(GetPrimitiveData(MaterialParameters).LocalToWorld);
-	DebugOutput.WorldToLocal = LWCToFloat(GetPrimitiveData(MaterialParameters).WorldToLocal);
+	DebugOutput.LocalObjectBoundsMin = PrimitiveData.LocalObjectBoundsMin.xyz;
+	DebugOutput.LocalObjectBoundsMax = PrimitiveData.LocalObjectBoundsMax.xyz;
+	DebugOutput.LocalToWorld = LWCToFloat(PrimitiveData.LocalToWorld);
+	DebugOutput.WorldToLocal = LWCToFloat(PrimitiveData.WorldToLocal);
 
 	return DebugOutput;
 }
diff --git a/Engine/Shaders/Private/LocalVertexFactory.ush b/Engine/Shaders/Private/LocalVertexFactory.ush
--- a/Engine/Shaders/Private/LocalVertexFactory.ush
+++ b/Engine/Shaders/Private/LocalVertexFactory.ush
@@ -534,12 +534,16 @@ FMaterialPixelParameters GetMaterialPixelParameters(FVertexFactoryInterpolantsVS
 	Result.InstanceWorldToLocal = GetInstanceWorldToLocal(Interpolants);
 #endif
 
+#if NEEDS_PARTICLE_LOCAL_TO_WORLD || NEEDS_PARTICLE_WORLD_TO_LOCAL
+	FPrimitiveSceneData PrimitiveData = GetPrimitiveData(Result.PrimitiveId);
+#endif
+
 #if NEEDS_PARTICLE_LOCAL_TO_WORLD
-	Result.Particle.ParticleToWorld = GetPrimitiveData(Result.PrimitiveId).LocalToWorld;
+	Result.Particle.ParticleToWorld = PrimitiveData.LocalToWorld;
 #endif
 
 #if NEEDS_PARTICLE_WORLD_TO_LOCAL
-	Result.Particle.WorldToParticle = GetPrimitiveData(Result.PrimitiveId).WorldToLocal;
+	Result.Particle.WorldToParticle = PrimitiveData.WorldToLocal;
 #endif
 
 	return Result;
@@ -569,14 +573,16 @@ FMaterialVertexParameters GetMaterialVertexParameters(FVertexFactoryInput Input,
 	// does not handle instancing!
 	Result.TangentToWorld = Intermediates.TangentToWorld;
 
+	FPrimitiveSceneData PrimitiveData = GetPrimitiveData(Intermediates);
+
 #if USE_INSTANCING
 	float4x4 InstanceToLocal = GetInstanceTransform(Intermediates);
-	Result.InstanceLocalToWorld = LWCMultiply(InstanceToLocal, GetPrimitiveData(Intermediates).LocalToWorld);
+	Result.InstanceLocalToWorld = LWCMultiply(InstanceToLocal, PrimitiveData.LocalToWorld);
 	Result.InstanceLocalPosition = Input.Position.xyz;
 	Result.PerInstanceParams = Intermediates.PerInstanceParams;
 	Result.InstanceId = GetInstanceId(Input.InstanceId); 
 	Result.InstanceOffset = InstanceOffset;
-	Result.PrevFrameLocalToWorld = LWCMultiply(GetInstancePrevTransform(Intermediates), GetPrimitiveData(Intermediates).PreviousLocalToWorld);
+	Result.PrevFrameLocalToWorld = LWCMultiply(GetInstancePrevTransform(Intermediates), PrimitiveData.PreviousLocalToWorld);
 
 	// Calculate derived world to local
 	{
@@ -595,7 +601,7 @@ FMaterialVertexParameters GetMaterialVertexParameters(FVertexFactoryInput Input,
 		LocalToInstance			 = transpose(LocalToInstance);
 		LocalToInstance[3].xyz	 = mul(float4(-InstanceToLocal[3].xyz, 0.0f), LocalToInstance).xyz;
 
-		Result.InstanceWorldToLocal = LWCMultiply(GetPrimitiveData(Intermediates).WorldToLocal, LocalToInstance);
+		Result.InstanceWorldToLocal = LWCMultiply(PrimitiveData.WorldToLocal, LocalToInstance);
 	}
 #else
 	Result.PrevFrameLocalToWorld = GetInstanceData(Intermediates).PrevLocalToWorld;
@@ -666,11 +672,11 @@ FMaterialVertexParameters GetMaterialVertexParameters(FVertexFactoryInput Input,
 	Result.Particle.Color = half4(1,1,1,1);
 
 #if NEEDS_PARTICLE_LOCAL_TO_WORLD
-	Result.Particle.ParticleToWorld = GetPrimitiveData(Intermediates).LocalToWorld;
+	Result.Particle.ParticleToWorld = PrimitiveData.LocalToWorld;
 #endif
 
 #if NEEDS_PARTICLE_WORLD_TO_LOCAL
-	Result.Particle.WorldToParticle = GetPrimitiveData(Intermediates).WorldToLocal;
+	Result.Particle.WorldToParticle = PrimitiveData.WorldToLocal;
 #endif
 
 #if ENABLE_NEW_HLSL_GENERATOR
@@ -863,9 +869,11 @@ FVertexFactoryIntermediates GetVertexFactoryIntermediates(FVertexFactoryInput In
 	FVertexFactoryIntermediates Intermediates = (FVertexFactoryIntermediates)0;
 	Intermediates.SceneData = VF_GPUSCENE_GET_INTERMEDIATES(Input);
 
+	FPrimitiveSceneData PrimitiveData = GetPrimitiveData(Intermediates);
+
 #if VF_USE_PRIMITIVE_SCENE_DATA
 	// Calculate the local (within the range of instance IDs belonging to a given primitive), as this can be used to load custom stuff
-	Intermediates.PrimitiveLocalInstanceIndex = Intermediates.SceneData.InstanceId - GetPrimitiveData(Intermediates).InstanceSceneDataOffset;
+	Intermediates.PrimitiveLocalInstanceIndex = Intermediates.SceneData.InstanceId - PrimitiveData.InstanceSceneDataOffset;
 
 #if USE_INSTANCE_CULLING
 	Intermediates.InstanceLightmapAndShadowMapUVBias = Intermediates.SceneData.InstanceData.LightMapAndShadowMapUVBias;
@@ -943,7 +951,7 @@ FVertexFactoryIntermediates GetVertexFactoryIntermediates(FVertexFactoryInput In
 	// GPUCULL_TODO: This can't be right
 	//Intermediates.PerInstanceParams.x = Intermediates.PerInstanceRandom;
 	Intermediates.PerInstanceParams.x = GetInstanceRandom(Intermediates);
-	float3 InstanceTranslatedLocation = TransformLocalToTranslatedWorld(GetInstanceOrigin(Intermediates), GetPrimitiveData(Intermediates).LocalToWorld).xyz;
+	float3 InstanceTranslatedLocation = TransformLocalToTranslatedWorld(GetInstanceOrigin(Intermediates), PrimitiveData.LocalToWorld).xyz;
 	Intermediates.PerInstanceParams.y = 1.0 - saturate((length(InstanceTranslatedLocation) - InstancingFadeOutParams.x) * InstancingFadeOutParams.y);
 	// InstancingFadeOutParams.z,w are RenderSelected and RenderDeselected respectively.
 	Intermediates.PerInstanceParams.z = InstancingFadeOutParams.z * SelectedValue + InstancingFadeOutParams.w * (1-SelectedValue);
@@ -1049,8 +1057,10 @@ FVertexFactoryInterpolantsVSToPS VertexFactoryGetInterpolantsVSToPS(FVertexFacto
 		float2 LightMapCoordinateInput = Input.LightMapCoordinate;
 	#endif
 
-	uint LightmapDataIndex = GetPrimitiveData(Intermediates).LightmapDataIndex + LocalVF.LODLightmapDataIndex;
-	uint LightmapUVIndex = GetPrimitiveData(Intermediates).LightmapUVIndex;
+	FPrimitiveSceneData PrimitiveData = GetPrimitiveData(Intermediates);
+
+	uint LightmapDataIndex = PrimitiveData.LightmapDataIndex + LocalVF.LODLightmapDataIndex;
+	uint LightmapUVIndex = PrimitiveData.LightmapUVIndex;
 	float4 LightMapCoordinateScaleBias = VF_GPUSCENE_GET_LIGHTMAP_UV_SCALE_BIAS(Input, LightmapDataIndex);
 
 	#if USE_INSTANCING || USE_INSTANCE_CULLING
diff --git a/Engine/Shaders/Private/SceneData.ush b/Engine/Shaders/Private/SceneData.ush
--- a/Engine/Shaders/Private/SceneData.ush
+++ b/Engine/Shaders/Private/SceneData.ush
@@ -769,8 +769,10 @@ void SetInstanceDeterminantSignFlag(float Determinant, inout uint Flags)
 // Determine the offsets into the payload data buffer for the given instance
 FInstancePayloadDataOffsets GetInstancePayloadDataOffsets(uint PrimitiveId, uint Flags, uint InstanceRelativeId)
 {
-	const uint PayloadDataRelativeOffset	= InstanceRelativeId * GetPrimitiveData(PrimitiveId).InstancePayloadDataStride;
-	const uint PayloadDataGlobalOffset		= PayloadDataRelativeOffset + GetPrimitiveData(PrimitiveId).InstancePayloadDataOffset;
+	FPrimitiveSceneData PrimitiveData = GetPrimitiveData(PrimitiveId);
+
+	const uint PayloadDataRelativeOffset	= InstanceRelativeId * PrimitiveData.InstancePayloadDataStride;
+	const uint PayloadDataGlobalOffset		= PayloadDataRelativeOffset + PrimitiveData.InstancePayloadDataOffset;
 
 	const bool bHasHierarchyOffset			= (Flags & INSTANCE_SCENE_DATA_FLAG_HAS_HIERARCHY_OFFSET) != 0u;
 	const bool bHasLocalBounds				= (Flags & INSTANCE_SCENE_DATA_FLAG_HAS_LOCAL_BOUNDS) != 0u;
@@ -913,7 +915,9 @@ FInstanceSceneData GetInstanceSceneData(uint InstanceId, uint SOAStride, bool bC
 		InstanceData.RandomID					= LoadInstanceSceneDataElement(0 * SOAStride + InstanceId).w;
 	#endif
 
-		float3 TilePosition 									= GetPrimitiveData(InstanceData.PrimitiveId).TilePosition;
+		FPrimitiveSceneData PrimitiveData		= GetPrimitiveData(InstanceData.PrimitiveId);
+
+		float3 TilePosition 									= PrimitiveData.TilePosition;
 
 #if PLATFORM_ALLOW_SCENE_DATA_COMPRESSED_TRANSFORMS
 		uint4  RotationScale									= asuint(LoadInstanceSceneDataElement(1 * SOAStride + InstanceId));
@@ -946,8 +950,8 @@ FInstanceSceneData GetInstanceSceneData(uint InstanceId, uint SOAStride, bool bC
 		InstanceData.PrevLocalToWorld	= MakeLWCMatrix(TilePosition, PrevLocalToRelativeWorld);
 		ComputeInstanceDerivedData(InstanceData, TilePosition, LocalToRelativeWorld);
 
-		InstanceData.NaniteRuntimeResourceID	= GetPrimitiveData(InstanceData.PrimitiveId).NaniteResourceID;
-		InstanceData.NaniteHierarchyOffset		= GetPrimitiveData(InstanceData.PrimitiveId).NaniteHierarchyOffset;
+		InstanceData.NaniteRuntimeResourceID	= PrimitiveData.NaniteResourceID;
+		InstanceData.NaniteHierarchyOffset		= PrimitiveData.NaniteHierarchyOffset;
 
 		BRANCH
 		if (Offsets.HierarchyOffset != INVALID_INSTANCE_PAYLOAD_OFFSET)
@@ -978,8 +982,8 @@ FInstanceSceneData GetInstanceSceneData(uint InstanceId, uint SOAStride, bool bC
 		}
 		else
 		{
-			InstanceData.LocalBoundsCenter = GetPrimitiveData(InstanceData.PrimitiveId).InstanceLocalBoundsCenter;
-			InstanceData.LocalBoundsExtent = GetPrimitiveData(InstanceData.PrimitiveId).InstanceLocalBoundsExtent;
+			InstanceData.LocalBoundsCenter = PrimitiveData.InstanceLocalBoundsCenter;
+			InstanceData.LocalBoundsExtent = PrimitiveData.InstanceLocalBoundsExtent;
 		}
 
 		BRANCH
@@ -997,7 +1001,7 @@ FInstanceSceneData GetInstanceSceneData(uint InstanceId, uint SOAStride, bool bC
 																 float4(0.0f, 0.0f, 0.0f, 1.0f)));
 		#endif
 
-			float3 TilePosition 				= GetPrimitiveData(InstanceData.PrimitiveId).TilePosition;
+			float3 TilePosition 				= PrimitiveData.TilePosition;
 			InstanceData.PrevLocalToWorld		= MakeLWCMatrix(TilePosition, PrevLocalToRelativeWorld);
 		}
 
