PR #7678: Avoid divide-by-zero in GetLightMapColorLQ (Contributed by [REDACTED_PERSON])

diff --git a/Engine/Shaders/Private/LightmapCommon.ush b/Engine/Shaders/Private/LightmapCommon.ush
--- a/Engine/Shaders/Private/LightmapCommon.ush
+++ b/Engine/Shaders/Private/LightmapCommon.ush
@@ -124,7 +124,7 @@ void GetLightMapColorLQ( VTPageTableResult LightmapVTPageTableResult, LightmapUV
 #endif
 		
 	half Luma = L * Directionality;
-	half3 Color = LogRGB * (Luma / LogL);				// 1 rcp, 1 smul, 1 vmul
+	half3 Color = LogRGB * (Luma / max(0.00001f, LogL));	// 1 smax, 1 rcp, 1 smul, 1 vmul
 
 	OutDiffuseLighting = Color;
 	OutIndirectIrradiance = Luma;
