Fix empty-voxel value when constructing transmittance volume.

diff --git a/Engine/Shaders/Private/HeterogeneousVolumes/HeterogeneousVolumesHardwareRayTracing.usf b/Engine/Shaders/Private/HeterogeneousVolumes/HeterogeneousVolumesHardwareRayTracing.usf
--- a/Engine/Shaders/Private/HeterogeneousVolumes/HeterogeneousVolumesHardwareRayTracing.usf
+++ b/Engine/Shaders/Private/HeterogeneousVolumes/HeterogeneousVolumesHardwareRayTracing.usf
@@ -426,7 +426,7 @@ RAY_TRACING_ENTRY_RAYGEN(RenderTransmittanceVolumeWithPreshadingRGS)
 	float Epsilon = 1.0e-7;
 	if (all(Extinction < Epsilon))
 	{
-		RWTransmittanceVolumeTexture[VoxelIndex] = 0.0f;
+		RWTransmittanceVolumeTexture[VoxelIndex] = 1.0f;
 		return;
 	}
 
diff --git a/Engine/Shaders/Private/HeterogeneousVolumes/HeterogeneousVolumesLiveShadingPipeline.usf b/Engine/Shaders/Private/HeterogeneousVolumes/HeterogeneousVolumesLiveShadingPipeline.usf
--- a/Engine/Shaders/Private/HeterogeneousVolumes/HeterogeneousVolumesLiveShadingPipeline.usf
+++ b/Engine/Shaders/Private/HeterogeneousVolumes/HeterogeneousVolumesLiveShadingPipeline.usf
@@ -189,7 +189,7 @@ void RenderTransmittanceVolumeWithLiveShadingCS(
 		uint3 Rand32Bits = Rand4DPCG32(uint4(VoxelIndex, View.StateFrameIndexMod8)).xyz;
 		VoxelJitter = float3(Rand32Bits) / float(uint(0xffffffff));
 	}
-	//float3 UVW = (VoxelIndex + VectorJitter) / float3(GetTransmittanceVolumeResolution());
+	//float3 UVW = (VoxelIndex + VoxelJitter) / float3(GetTransmittanceVolumeResolution());
 	float3 UVW = (VoxelIndex + 0.5) / float3(GetTransmittanceVolumeResolution());
 
 	float3 LocalBoundsMin = GetLocalBoundsOrigin() - GetLocalBoundsExtent();
@@ -204,7 +204,7 @@ void RenderTransmittanceVolumeWithLiveShadingCS(
 	float Epsilon = 1.0e-7;
 	if (all(Extinction < Epsilon))
 	{
-		RWTransmittanceVolumeTexture[VoxelIndex] = 0.0f;
+		RWTransmittanceVolumeTexture[VoxelIndex] = 1.0f;
 		return;
 	}
 
diff --git a/Engine/Shaders/Private/HeterogeneousVolumes/HeterogeneousVolumesPreshadingPipeline.usf b/Engine/Shaders/Private/HeterogeneousVolumes/HeterogeneousVolumesPreshadingPipeline.usf
--- a/Engine/Shaders/Private/HeterogeneousVolumes/HeterogeneousVolumesPreshadingPipeline.usf
+++ b/Engine/Shaders/Private/HeterogeneousVolumes/HeterogeneousVolumesPreshadingPipeline.usf
@@ -447,7 +447,7 @@ void RenderTransmittanceVolumeWithPreshadingCS(
 	float Epsilon = 1.0e-7;
 	if (all(Extinction < Epsilon))
 	{
-		RWTransmittanceVolumeTexture[VoxelIndex] = 0.0f;
+		RWTransmittanceVolumeTexture[VoxelIndex] = 1.0f;
 		return;
 	}
 
