Reduce GPUSkinPassThrough vertex factory shader permuations by merging it into local vertex factory. PassThroughVF is kept as a wrapper class.

CitySample PC cook stats (recompile all shaders in both runs):
Baseline: 75229 shaders compiled, Cook time 1h35m46s
WithChange: 69228 shaders compiled, Cook time 1h29m10s
Shader reduction 8%, Cook time reduction 6mins

diff --git a/Engine/Shaders/Private/GpuSkinVertexFactory.ush b/Engine/Shaders/Private/GpuSkinVertexFactory.ush
--- a/Engine/Shaders/Private/GpuSkinVertexFactory.ush
+++ b/Engine/Shaders/Private/GpuSkinVertexFactory.ush
@@ -80,26 +80,23 @@ struct FVertexFactoryInput
 	float4	BlendWeightsExtra	: ATTRIBUTE15;
 #endif // GPUSKIN_UNLIMITED_BONE_INFLUENCE
 
-
 #if NUM_MATERIAL_TEXCOORDS_VERTEX
-	// If this changes make sure to update LocalVertexFactory.usf
-#if NUM_MATERIAL_TEXCOORDS_VERTEX > 0
-	float2	TexCoords0 : ATTRIBUTE5;
-#endif
 #if NUM_MATERIAL_TEXCOORDS_VERTEX > 1
-	float2	TexCoords1 : ATTRIBUTE6;
-#endif
-#if NUM_MATERIAL_TEXCOORDS_VERTEX > 2
-	float2	TexCoords2 : ATTRIBUTE7;
+	float4	TexCoords0 : ATTRIBUTE5;
+#elif NUM_MATERIAL_TEXCOORDS_VERTEX == 1
+	float2	TexCoords0 : ATTRIBUTE5;
 #endif
+
 #if NUM_MATERIAL_TEXCOORDS_VERTEX > 3
-	float2	TexCoords3 : ATTRIBUTE8;
+	float4	TexCoords1 : ATTRIBUTE6;
+#elif NUM_MATERIAL_TEXCOORDS_VERTEX == 3
+	float2	TexCoords1 : ATTRIBUTE6;
 #endif
 
 	#if NUM_MATERIAL_TEXCOORDS_VERTEX > 4
 		#error Too many texture coordinate sets defined on GPUSkin vertex input. Max: 4.
 	#endif
-#endif
+#endif // NUM_MATERIAL_TEXCOORDS_VERTEX
 
 #if GPUSKIN_MORPH_BLEND
 	// NOTE: TEXCOORD6,TEXCOORD7 used instead of POSITION1,NORMAL1 since those semantics are not supported by Cg 
@@ -340,16 +337,16 @@ FMaterialVertexParameters GetMaterialVertexParameters(FVertexFactoryInput Input,
 #if NUM_MATERIAL_TEXCOORDS_VERTEX
 
 #if NUM_MATERIAL_TEXCOORDS_VERTEX > 0
-	Result.TexCoords[0] = Input.TexCoords0;
+	Result.TexCoords[0] = Input.TexCoords0.xy;
 #endif
 #if NUM_MATERIAL_TEXCOORDS_VERTEX > 1
-	Result.TexCoords[1] = Input.TexCoords1;
+	Result.TexCoords[1] = Input.TexCoords0.zw;
 #endif
 #if NUM_MATERIAL_TEXCOORDS_VERTEX > 2
-	Result.TexCoords[2] = Input.TexCoords2;
+	Result.TexCoords[2] = Input.TexCoords1.xy;
 #endif
 #if NUM_MATERIAL_TEXCOORDS_VERTEX > 3
-	Result.TexCoords[3] = Input.TexCoords3;
+	Result.TexCoords[3] = Input.TexCoords1.zw;
 #endif
 
 #endif
diff --git a/Engine/Shaders/Private/LocalVertexFactory.ush b/Engine/Shaders/Private/LocalVertexFactory.ush
--- a/Engine/Shaders/Private/LocalVertexFactory.ush
+++ b/Engine/Shaders/Private/LocalVertexFactory.ush
@@ -8,9 +8,7 @@
 #include "LocalVertexFactoryCommon.ush"
 #include "LightmapData.ush"
 
-#if GPUSKIN_PASS_THROUGH
 #include "GpuSkinCommon.ush"
-#endif
 
 #include "/Engine/Generated/UniformBuffers/PrecomputedLightingBuffer.ush"
 
@@ -63,6 +61,10 @@ float4 SplineParams[10];
 #define SplineMeshY				SplineParams[9].xyz
 #endif	// USE_SPLINEDEFORM
 
+#ifndef SUPPORT_SKINCACHE
+#define SUPPORT_SKINCACHE 0
+#endif
+
 #ifndef MANUAL_VERTEX_FETCH
 #define MANUAL_VERTEX_FETCH 0
 #endif
@@ -111,46 +113,29 @@ struct FVertexFactoryInput
 
 #if NUM_MATERIAL_TEXCOORDS_VERTEX
 	#if !MANUAL_VERTEX_FETCH
-		#if GPUSKIN_PASS_THROUGH
-			// These must match GPUSkinVertexFactory.usf
+		// These used to be packed texcoord arrays, but these cause problems with alighnment on some Vulkan drivers
+		#if NUM_MATERIAL_TEXCOORDS_VERTEX > 1
+			float4	TexCoords0 : ATTRIBUTE4;
+		#elif NUM_MATERIAL_TEXCOORDS_VERTEX == 1
 			float2	TexCoords0 : ATTRIBUTE4;
-			#if NUM_MATERIAL_TEXCOORDS_VERTEX > 1
-				float2	TexCoords1 : ATTRIBUTE5;
-			#endif
-			#if NUM_MATERIAL_TEXCOORDS_VERTEX > 2
-				float2	TexCoords2 : ATTRIBUTE6;
-			#endif
-			#if NUM_MATERIAL_TEXCOORDS_VERTEX > 3
-				float2	TexCoords3 : ATTRIBUTE7;
-			#endif
-			#if NUM_MATERIAL_TEXCOORDS_VERTEX > 4
-				#error Too many texture coordinate sets defined on GPUSkin vertex input. Max: 4.
-			#endif
-		#else
-			// These used to be packed texcoord arrays, but these cause problems with alighnment on some Vulkan drivers
-			#if NUM_MATERIAL_TEXCOORDS_VERTEX > 1
-				float4	TexCoords0 : ATTRIBUTE4;
-			#elif NUM_MATERIAL_TEXCOORDS_VERTEX == 1
-				float2	TexCoords0 : ATTRIBUTE4;
-			#endif
-
-			#if NUM_MATERIAL_TEXCOORDS_VERTEX > 3
-				float4	TexCoords1 : ATTRIBUTE5;
-			#elif NUM_MATERIAL_TEXCOORDS_VERTEX == 3
-				float2	TexCoords1 : ATTRIBUTE5;
-			#endif
-
-			#if NUM_MATERIAL_TEXCOORDS_VERTEX > 5
-				float4	TexCoords2 : ATTRIBUTE6;
-			#elif NUM_MATERIAL_TEXCOORDS_VERTEX == 5
-				float2	TexCoords2 : ATTRIBUTE6;
-			#endif
-
-			#if NUM_MATERIAL_TEXCOORDS_VERTEX > 7
-				float4	TexCoords3 : ATTRIBUTE7;
-			#elif NUM_MATERIAL_TEXCOORDS_VERTEX == 7
-				float2	TexCoords3 : ATTRIBUTE7;
-			#endif
+		#endif
+
+		#if NUM_MATERIAL_TEXCOORDS_VERTEX > 3
+			float4	TexCoords1 : ATTRIBUTE5;
+		#elif NUM_MATERIAL_TEXCOORDS_VERTEX == 3
+			float2	TexCoords1 : ATTRIBUTE5;
+		#endif
+
+		#if NUM_MATERIAL_TEXCOORDS_VERTEX > 5
+			float4	TexCoords2 : ATTRIBUTE6;
+		#elif NUM_MATERIAL_TEXCOORDS_VERTEX == 5
+			float2	TexCoords2 : ATTRIBUTE6;
+		#endif
+
+		#if NUM_MATERIAL_TEXCOORDS_VERTEX > 7
+			float4	TexCoords3 : ATTRIBUTE7;
+		#elif NUM_MATERIAL_TEXCOORDS_VERTEX == 7
+			float2	TexCoords3 : ATTRIBUTE7;
 		#endif
 	#endif
 #elif USE_PARTICLE_SUBUVS && !MANUAL_VERTEX_FETCH
@@ -173,7 +158,7 @@ struct FVertexFactoryInput
 	VF_INSTANCED_STEREO_DECLARE_INPUT_BLOCK()
 #endif
 
-#if GPUSKIN_PASS_THROUGH && !MANUAL_VERTEX_FETCH
+#if SUPPORT_SKINCACHE && !MANUAL_VERTEX_FETCH
 	float4	PreSkinPosition	: ATTRIBUTE14;
 #endif
 
@@ -181,15 +166,24 @@ struct FVertexFactoryInput
 	float2	LightMapCoordinate : ATTRIBUTE15;
 #endif
 
-#if GPUSKIN_PASS_THROUGH || MANUAL_VERTEX_FETCH
 	uint VertexId : SV_VertexID;
-#endif
 };
 
+#if SUPPORT_SKINCACHE
+uint bIsGPUSkinPassThrough;
+#endif
+
+bool IsGPUSkinPassThrough()
+{
+	#if SUPPORT_SKINCACHE
+		return bIsGPUSkinPassThrough;
+	#else
+		return false;
+	#endif
+}
+
 #if RAYHITGROUPSHADER || COMPUTESHADER
-#if GPUSKIN_PASS_THROUGH
 Buffer<float> GPUSkinCachePositionBuffer;
-#endif
 
 uint GetNumRayTracingDynamicMeshVerticesIndirect()
 {
@@ -205,9 +199,7 @@ FVertexFactoryInput LoadVertexFactoryInputForHGS(uint TriangleIndex, int VertexI
 
 	FTriangleBaseAttributes Tri = LoadTriangleBaseAttributes(TriangleIndex);
 
-#if GPUSKIN_PASS_THROUGH || MANUAL_VERTEX_FETCH
 	Input.VertexId = Tri.Indices[VertexIndex];
-#endif
 	Input.Position = float4(Tri.LocalPositions[VertexIndex], 1.0f);
 #if VF_USE_PRIMITIVE_SCENE_DATA
 	// Note: GetInstanceUserData() stores the GPU-Scene primitive ID
@@ -239,19 +231,24 @@ FVertexFactoryInput LoadVertexFactoryInputForDynamicUpdate(uint TriangleIndex, i
 {
 	FVertexFactoryInput Input = (FVertexFactoryInput)0;
 
-	#if GPUSKIN_PASS_THROUGH
+	if (IsGPUSkinPassThrough())
+	{
 		Input.VertexId = TriangleIndex * 3 + VertexIndex;
 		uint VertexOffset = LocalVF.VertexFetch_Parameters[VF_VertexOffset] + Input.VertexId;
 		Input.Position.x = GPUSkinCachePositionBuffer[VertexOffset * 3 + 0];
 		Input.Position.y = GPUSkinCachePositionBuffer[VertexOffset * 3 + 1];
 		Input.Position.z = GPUSkinCachePositionBuffer[VertexOffset * 3 + 2];
-	#elif  MANUAL_VERTEX_FETCH
+	}
+	else
+	{
+	#if  MANUAL_VERTEX_FETCH
 		Input.VertexId = TriangleIndex * 3 + VertexIndex;
 		uint VertexOffset = LocalVF.VertexFetch_Parameters[VF_VertexOffset] + Input.VertexId * 3;
 		Input.Position.x = LocalVF.VertexFetch_PositionBuffer[VertexOffset + 0];
 		Input.Position.y = LocalVF.VertexFetch_PositionBuffer[VertexOffset + 1];
 		Input.Position.z = LocalVF.VertexFetch_PositionBuffer[VertexOffset + 2];
 	#endif
+	}
 
 	VF_GPUSCENE_SET_INPUT_FOR_RT(Input, PrimitiveId, 0U);
 
@@ -630,43 +627,30 @@ FMaterialVertexParameters GetMaterialVertexParameters(FVertexFactoryInput Input,
 			Result.TexCoords[CoordinateIndex] = LocalVF.VertexFetch_TexCoordBuffer[NumFetchTexCoords * (LocalVF.VertexFetch_Parameters[VF_VertexOffset] + Input.VertexId) + ClampedCoordinateIndex];
 		}
 #elif NUM_MATERIAL_TEXCOORDS_VERTEX
-		#if GPUSKIN_PASS_THROUGH
-			Result.TexCoords[0] = Input.TexCoords0.xy;
-			#if NUM_MATERIAL_TEXCOORDS_VERTEX > 1
-				Result.TexCoords[1] = Input.TexCoords1.xy;
-			#endif
-			#if NUM_MATERIAL_TEXCOORDS_VERTEX > 2
-				Result.TexCoords[2] = Input.TexCoords2.xy;
-			#endif
-			#if NUM_MATERIAL_TEXCOORDS_VERTEX > 3
-				Result.TexCoords[3] = Input.TexCoords3.xy;
-			#endif
-		#else
-			#if NUM_MATERIAL_TEXCOORDS_VERTEX > 0
-				Result.TexCoords[0] = Input.TexCoords0.xy;
-			#endif
-			#if NUM_MATERIAL_TEXCOORDS_VERTEX > 1
-				Result.TexCoords[1] = Input.TexCoords0.zw;
-			#endif
-			#if NUM_MATERIAL_TEXCOORDS_VERTEX > 2
-				Result.TexCoords[2] = Input.TexCoords1.xy;
-			#endif
-			#if NUM_MATERIAL_TEXCOORDS_VERTEX > 3
-				Result.TexCoords[3] = Input.TexCoords1.zw;
-			#endif
-			#if NUM_MATERIAL_TEXCOORDS_VERTEX > 4
-				Result.TexCoords[4] = Input.TexCoords2.xy;
-			#endif
-			#if NUM_MATERIAL_TEXCOORDS_VERTEX > 5
-				Result.TexCoords[5] = Input.TexCoords2.zw;
-			#endif
-			#if NUM_MATERIAL_TEXCOORDS_VERTEX > 6
-				Result.TexCoords[6] = Input.TexCoords3.xy;
-			#endif
-			#if NUM_MATERIAL_TEXCOORDS_VERTEX > 7
-				Result.TexCoords[7] = Input.TexCoords3.zw;
-			#endif
-		#endif
+	#if NUM_MATERIAL_TEXCOORDS_VERTEX > 0
+		Result.TexCoords[0] = Input.TexCoords0.xy;
+	#endif
+	#if NUM_MATERIAL_TEXCOORDS_VERTEX > 1
+		Result.TexCoords[1] = Input.TexCoords0.zw;
+	#endif
+	#if NUM_MATERIAL_TEXCOORDS_VERTEX > 2
+		Result.TexCoords[2] = Input.TexCoords1.xy;
+	#endif
+	#if NUM_MATERIAL_TEXCOORDS_VERTEX > 3
+		Result.TexCoords[3] = Input.TexCoords1.zw;
+	#endif
+	#if NUM_MATERIAL_TEXCOORDS_VERTEX > 4
+		Result.TexCoords[4] = Input.TexCoords2.xy;
+	#endif
+	#if NUM_MATERIAL_TEXCOORDS_VERTEX > 5
+		Result.TexCoords[5] = Input.TexCoords2.zw;
+	#endif
+	#if NUM_MATERIAL_TEXCOORDS_VERTEX > 6
+		Result.TexCoords[6] = Input.TexCoords3.xy;
+	#endif
+	#if NUM_MATERIAL_TEXCOORDS_VERTEX > 7
+		Result.TexCoords[7] = Input.TexCoords3.zw;
+	#endif
 #endif  //MANUAL_VERTEX_FETCH && NUM_MATERIAL_TEXCOORDS_VERTEX
 
 	Result.Particle.Color = half4(1,1,1,1);
@@ -817,13 +801,13 @@ half3x3 CalcTangentToLocal(FVertexFactoryInput Input, inout float TangentSign)
 	half4 TangentInputZ = Input.TangentZ;
 #endif
 
-#ifdef GPUSKIN_PASS_THROUGH
 	half3 TangentX = TangentInputX;
 	half4 TangentZ = TangentInputZ;
-#else
-	half3 TangentX = TangentBias(TangentInputX);
-	half4 TangentZ = TangentBias(TangentInputZ);
-#endif
+	if (!IsGPUSkinPassThrough())
+	{
+		TangentX = TangentBias(TangentInputX);
+		TangentZ = TangentBias(TangentInputZ);
+	}
 
 	TangentSign = TangentZ.w;
 
@@ -968,7 +952,9 @@ FVertexFactoryIntermediates GetVertexFactoryIntermediates(FVertexFactoryInput In
 	#endif
 #endif	// USE_INSTANCING
 
-#if GPUSKIN_PASS_THROUGH
+	if (IsGPUSkinPassThrough())
+	{
+	#if SUPPORT_SKINCACHE
 	#if MANUAL_VERTEX_FETCH
 		uint PreSkinVertexOffset = LocalVF.PreSkinBaseVertexIndex + Input.VertexId * 3;
 		Intermediates.PreSkinPosition.x = LocalVF.VertexFetch_PreSkinPositionBuffer[PreSkinVertexOffset + 0];
@@ -977,9 +963,12 @@ FVertexFactoryIntermediates GetVertexFactoryIntermediates(FVertexFactoryInput In
 	#else
 		Intermediates.PreSkinPosition = Input.PreSkinPosition.xyz;
 	#endif
-#else
-	Intermediates.PreSkinPosition = Input.Position.xyz;
-#endif
+	#endif
+	}
+	else
+	{
+		Intermediates.PreSkinPosition = Input.Position.xyz;	
+	}
 
 	return Intermediates;
 }
@@ -1158,9 +1147,7 @@ float3 VertexFactoryGetWorldNormal(FVertexFactoryInput Input, FVertexFactoryInte
 }
 
 // Information when using pass through to get the previous position, as no matrix is available/needed
-#if GPUSKIN_PASS_THROUGH
 Buffer<float> GPUSkinCachePreviousPositionBuffer;
-#endif
 
 // @return previous translated world position
 float4 VertexFactoryGetPreviousWorldPosition(FVertexFactoryInput Input, FVertexFactoryIntermediates Intermediates)
@@ -1168,27 +1155,32 @@ float4 VertexFactoryGetPreviousWorldPosition(FVertexFactoryInput Input, FVertexF
 	FLWCMatrix PreviousLocalToWorld = GetInstanceData(Intermediates).PrevLocalToWorld;
 	float4x4 PreviousLocalToWorldTranslated = LWCMultiplyTranslation(PreviousLocalToWorld, ResolvedView.PrevPreViewTranslation);
 
-#if USE_INSTANCING
-	float4x4 InstanceTransform = GetInstancePrevTransform(Intermediates);
-	return mul(mul(Input.Position, InstanceTransform), PreviousLocalToWorldTranslated);
-#elif GPUSKIN_PASS_THROUGH
-	uint Offset = Input.VertexId * 3;
-	float3 PreviousPos;
-	PreviousPos.x = GPUSkinCachePreviousPositionBuffer[Offset + 0];
-	PreviousPos.y = GPUSkinCachePreviousPositionBuffer[Offset + 1];
-	PreviousPos.z = GPUSkinCachePreviousPositionBuffer[Offset + 2];
-	return mul(float4(PreviousPos, 1), PreviousLocalToWorldTranslated);
-#elif USE_SPLINEDEFORM
-	// Just like CalcWorldPosition...
-	float4x3 SliceTransform = CalcSliceTransform(dot(Input.Position.xyz, SplineMeshDir));
-
-	// Transform into mesh space
-	float4 LocalPos = float4(mul(Input.Position, SliceTransform), Input.Position.w);
-
-	return mul(LocalPos, PreviousLocalToWorldTranslated);
-#else
-	return mul(Input.Position, PreviousLocalToWorldTranslated);
-#endif	// USE_INSTANCING
+	if (IsGPUSkinPassThrough())
+	{
+		uint Offset = Input.VertexId * 3;
+		float3 PreviousPos;
+		PreviousPos.x = GPUSkinCachePreviousPositionBuffer[Offset + 0];
+		PreviousPos.y = GPUSkinCachePreviousPositionBuffer[Offset + 1];
+		PreviousPos.z = GPUSkinCachePreviousPositionBuffer[Offset + 2];
+		return mul(float4(PreviousPos, 1), PreviousLocalToWorldTranslated);	
+	}
+	else
+	{
+	#if USE_INSTANCING
+		float4x4 InstanceTransform = GetInstancePrevTransform(Intermediates);
+		return mul(mul(Input.Position, InstanceTransform), PreviousLocalToWorldTranslated);
+	#elif USE_SPLINEDEFORM
+		// Just like CalcWorldPosition...
+		float4x3 SliceTransform = CalcSliceTransform(dot(Input.Position.xyz, SplineMeshDir));
+	
+		// Transform into mesh space
+		float4 LocalPos = float4(mul(Input.Position, SliceTransform), Input.Position.w);
+	
+		return mul(LocalPos, PreviousLocalToWorldTranslated);
+	#else
+		return mul(Input.Position, PreviousLocalToWorldTranslated);
+	#endif	// USE_INSTANCING
+	}
 }
 
 #if NEEDS_VERTEX_FACTORY_INTERPOLATION
diff --git a/Engine/Source/Runtime/Engine/Private/GPUSkinVertexFactory.cpp b/Engine/Source/Runtime/Engine/Private/GPUSkinVertexFactory.cpp
--- a/Engine/Source/Runtime/Engine/Private/GPUSkinVertexFactory.cpp
+++ b/Engine/Source/Runtime/Engine/Private/GPUSkinVertexFactory.cpp
@@ -766,110 +766,6 @@ IMPLEMENT_GPUSKINNING_VERTEX_FACTORY_TYPE(TGPUSkinVertexFactory, "/Engine/Privat
 	| EVertexFactoryFlags::SupportsDynamicLighting
 );
 
-/*-----------------------------------------------------------------------------
-TGPUSkinVertexFactoryShaderParameters
------------------------------------------------------------------------------*/
-
-/** Shader parameters for use with TGPUSkinVertexFactory */
-class FGPUSkinVertexPassthroughFactoryShaderParameters : public FLocalVertexFactoryShaderParametersBase
-{
-	DECLARE_TYPE_LAYOUT(FGPUSkinVertexPassthroughFactoryShaderParameters, NonVirtual);
-public:
-	/**
-	* Bind shader constants by name
-	* @param	ParameterMap - mapping of named shader constants to indices
-	*/
-	void Bind(const FShaderParameterMap& ParameterMap)
-	{
-		FLocalVertexFactoryShaderParametersBase::Bind(ParameterMap);
-		GPUSkinCachePositionBuffer.Bind(ParameterMap,TEXT("GPUSkinCachePositionBuffer"));
-		GPUSkinCachePreviousPositionBuffer.Bind(ParameterMap,TEXT("GPUSkinCachePreviousPositionBuffer"));
-	}
-	
-	void GetElementShaderBindings(
-		const FSceneInterface* Scene,
-		const FSceneView* View,
-		const FMeshMaterialShader* Shader,
-		const EVertexInputStreamType InputStreamType,
-		ERHIFeatureLevel::Type FeatureLevel,
-		const FVertexFactory* VertexFactory,
-		const FMeshBatchElement& BatchElement,
-		class FMeshDrawSingleShaderBindings& ShaderBindings,
-		FVertexInputStreamArray& VertexStreams) const
-	{
-		// #dxr_todo do we need this call to the base?
-		FLocalVertexFactoryShaderParametersBase::GetElementShaderBindingsBase(Scene, View, Shader, InputStreamType, FeatureLevel, VertexFactory, BatchElement, nullptr, ShaderBindings, VertexStreams);
-
-		// todo: Add more context into VertexFactoryUserData about whether this is skin cache/mesh deformer/ray tracing etc.
-		// For now it just holds a skin cache pointer which is null if using other mesh deformers.
-		FGPUSkinBatchElementUserData* BatchUserData = (FGPUSkinBatchElementUserData*)BatchElement.VertexFactoryUserData;
-		const bool bUsesSkinCache = BatchUserData != nullptr;
-
-		check(VertexFactory->GetType() == &FGPUSkinPassthroughVertexFactory::StaticType);
-		FGPUSkinPassthroughVertexFactory const* PassthroughVertexFactory = static_cast<FGPUSkinPassthroughVertexFactory const*>(VertexFactory);
-
-		// Bone data is updated whenever animation triggers a dynamic update, animation can skip frames hence the frequency is not necessary every frame.
-		// So check if bone data is updated this frame, if not then the previous frame data is stale and not suitable for motion blur.
-		bool bBoneDataUpdatedThisFrame = (View->Family->FrameNumber == PassthroughVertexFactory->GetUpdatedFrameNumber());
-		// If world is paused, use current frame bone matrices, so velocity is canceled and skeletal mesh isn't blurred from motion.
-		bool bVerticesInMotion = !View->Family->bWorldIsPaused && bBoneDataUpdatedThisFrame;
-
-		if (bUsesSkinCache)
-		{
-			GetElementShaderBindingsSkinCache(PassthroughVertexFactory, BatchUserData, bVerticesInMotion, ShaderBindings, VertexStreams);
-		}
-		else
-		{
-			GetElementShaderBindingsMeshDeformer(PassthroughVertexFactory, bVerticesInMotion, ShaderBindings, VertexStreams);
-		}
-	}
-
-	void GetElementShaderBindingsSkinCache(
-		FGPUSkinPassthroughVertexFactory const* PassthroughVertexFactory,
-		FGPUSkinBatchElementUserData* BatchUserData,
-		bool bVerticesInMotion,
-		FMeshDrawSingleShaderBindings& ShaderBindings,
-		FVertexInputStreamArray& VertexStreams) const
-	{
-		FGPUSkinCache::GetShaderBindings(
-			BatchUserData->Entry, BatchUserData->Section, bVerticesInMotion,
-			PassthroughVertexFactory,
-			GPUSkinCachePositionBuffer, GPUSkinCachePreviousPositionBuffer, 
-			ShaderBindings, VertexStreams);
-	}
-
-	void GetElementShaderBindingsMeshDeformer(
-		FGPUSkinPassthroughVertexFactory const* PassthroughVertexFactory,
-		bool bVerticesInMotion,
-		FMeshDrawSingleShaderBindings& ShaderBindings,
-		FVertexInputStreamArray& VertexStreams) const
-	{
-		if (PassthroughVertexFactory->PositionRDG.IsValid())
-		{
-			VertexStreams.Add(FVertexInputStream(PassthroughVertexFactory->GetPositionStreamIndex(), 0, PassthroughVertexFactory->PositionRDG->GetRHI()));
-			ShaderBindings.Add(GPUSkinCachePositionBuffer, PassthroughVertexFactory->GetPositionsSRV());
-		}
-		if (PassthroughVertexFactory->TangentRDG.IsValid() && PassthroughVertexFactory->GetTangentStreamIndex() > -1)
-		{
-			VertexStreams.Add(FVertexInputStream(PassthroughVertexFactory->GetTangentStreamIndex(), 0, PassthroughVertexFactory->TangentRDG->GetRHI()));
-		}
-
-		if (bVerticesInMotion && PassthroughVertexFactory->PrevPositionRDG)
-		{
-			ShaderBindings.Add(GPUSkinCachePreviousPositionBuffer, PassthroughVertexFactory->PrevPositionRDG->GetOrCreateSRV(FRHIBufferSRVCreateInfo(PF_R32_FLOAT)));
-		}
-		else
-		{
-			ShaderBindings.Add(GPUSkinCachePreviousPositionBuffer, PassthroughVertexFactory->GetPositionsSRV());
-		}
-	}
-
-private:
-	LAYOUT_FIELD(FShaderResourceParameter, GPUSkinCachePositionBuffer);
-	LAYOUT_FIELD(FShaderResourceParameter, GPUSkinCachePreviousPositionBuffer);
-};
-
-IMPLEMENT_TYPE_LAYOUT(FGPUSkinVertexPassthroughFactoryShaderParameters);
 
 /*-----------------------------------------------------------------------------
 FGPUSkinPassthroughVertexFactory
@@ -878,33 +774,13 @@ FGPUSkinPassthroughVertexFactory::FGPUSkinPassthroughVertexFactory(ERHIFeatureLe
 	: FLocalVertexFactory(InFeatureLevel, "FGPUSkinPassthroughVertexFactory")
 {
 	bSupportsManualVertexFetch = true;
+	bGPUSkinPassThrough = true;
 }
 
 FGPUSkinPassthroughVertexFactory::~FGPUSkinPassthroughVertexFactory()
 {
 }
 
-void FGPUSkinPassthroughVertexFactory::ModifyCompilationEnvironment(const FVertexFactoryShaderPermutationParameters& Parameters, FShaderCompilerEnvironment& OutEnvironment )
-{
-	const bool ContainsManualVertexFetch = OutEnvironment.GetDefinitions().Contains("MANUAL_VERTEX_FETCH");
-	if (!ContainsManualVertexFetch && RHISupportsManualVertexFetch(Parameters.Platform))
-	{
-		OutEnvironment.SetDefine(TEXT("MANUAL_VERTEX_FETCH"), TEXT("1"));
-	}
-
-	Super::ModifyCompilationEnvironment(Parameters, OutEnvironment);
-	OutEnvironment.SetDefine(TEXT("GPUSKIN_PASS_THROUGH"),TEXT("1"));
-}
-
-bool FGPUSkinPassthroughVertexFactory::ShouldCompilePermutation(const FVertexFactoryShaderPermutationParameters& Parameters)
-{
-	// Passthrough is only valid on platforms with Compute Shader support AND for (skeletal meshes or default materials)
-	return 
-		IsGPUSkinCacheAvailable(Parameters.Platform) &&
-		Super::ShouldCompilePermutation(Parameters) &&
-		(Parameters.MaterialParameters.bIsUsedWithSkeletalMesh || Parameters.MaterialParameters.bIsSpecialEngineMaterial);
-}
-
 void FGPUSkinPassthroughVertexFactory::ReleaseRHI()
 {
 	FLocalVertexFactory::ReleaseRHI();
@@ -1053,20 +929,6 @@ void FGPUSkinPassthroughVertexFactory::InternalUpdateVertexDeclaration(
 	InternalUpdateVertexDeclaration(SourceVertexFactory);
 }
 
-IMPLEMENT_VERTEX_FACTORY_PARAMETER_TYPE(FGPUSkinPassthroughVertexFactory, SF_Vertex, FGPUSkinVertexPassthroughFactoryShaderParameters);
-#if RHI_RAYTRACING
-IMPLEMENT_VERTEX_FACTORY_PARAMETER_TYPE(FGPUSkinPassthroughVertexFactory, SF_RayHitGroup, FGPUSkinVertexPassthroughFactoryShaderParameters);
-IMPLEMENT_VERTEX_FACTORY_PARAMETER_TYPE(FGPUSkinPassthroughVertexFactory, SF_Compute, FGPUSkinVertexPassthroughFactoryShaderParameters);
-#endif // RHI_RAYTRACING
-
-IMPLEMENT_VERTEX_FACTORY_TYPE(FGPUSkinPassthroughVertexFactory, "/Engine/Private/LocalVertexFactory.ush",
-	  EVertexFactoryFlags::UsedWithMaterials
-	| EVertexFactoryFlags::SupportsDynamicLighting
-	| EVertexFactoryFlags::SupportsRayTracing
-	| EVertexFactoryFlags::SupportsRayTracingDynamicGeometry
-	| EVertexFactoryFlags::SupportsPrimitiveIdStream
-);
-
 
 /*-----------------------------------------------------------------------------
 	TGPUSkinMorphVertexFactoryShaderParameters
diff --git a/Engine/Source/Runtime/Engine/Private/InstancedStaticMesh.cpp b/Engine/Source/Runtime/Engine/Private/InstancedStaticMesh.cpp
--- a/Engine/Source/Runtime/Engine/Private/InstancedStaticMesh.cpp
+++ b/Engine/Source/Runtime/Engine/Private/InstancedStaticMesh.cpp
@@ -951,6 +951,15 @@ void FInstancedStaticMeshVertexFactory::InitRHI()
 	
 	AddPrimitiveIdStreamElement(EVertexInputStreamType::Default, Elements, 13, AutoInstancingAttr_Mobile);
 
+	// PreSkinPosition is only used when skin cache is on, in which case skinned mesh uses local vertex factory as pass through.
+	// It is not used by ISM so fill with dummy buffer.
+	extern ENGINE_API bool IsGPUSkinCacheAvailable(EShaderPlatform Platform);
+	if (IsGPUSkinCacheAvailable(GMaxRHIShaderPlatform))
+	{
+		FVertexStreamComponent NullComponent(&GNullVertexBuffer, 0, 0, VET_Float4);
+		Elements.Add(AccessStreamComponent(NullComponent, 14));
+	}
+
 	if(Data.LightMapCoordinateComponent.VertexBuffer)
 	{
 		Elements.Add(AccessStreamComponent(Data.LightMapCoordinateComponent,15));
diff --git a/Engine/Source/Runtime/Engine/Private/LocalVertexFactory.cpp b/Engine/Source/Runtime/Engine/Private/LocalVertexFactory.cpp
--- a/Engine/Source/Runtime/Engine/Private/LocalVertexFactory.cpp
+++ b/Engine/Source/Runtime/Engine/Private/LocalVertexFactory.cpp
@@ -12,6 +12,8 @@
 #include "Rendering/ColorVertexBuffer.h"
 #include "MeshMaterialShader.h"
 #include "ProfilingDebugging/LoadTimeTracker.h"
+#include "GPUSkinCache.h"
+#include "GPUSkinVertexFactory.h"
 
 IMPLEMENT_TYPE_LAYOUT(FLocalVertexFactoryShaderParametersBase);
 IMPLEMENT_TYPE_LAYOUT(FLocalVertexFactoryShaderParameters);
@@ -156,6 +158,14 @@ void FLocalVertexFactoryShaderParametersBase::GetElementShaderBindingsBase(
 	}
 }
 
+void FLocalVertexFactoryShaderParameters::Bind(const FShaderParameterMap& ParameterMap)
+{
+	FLocalVertexFactoryShaderParametersBase::Bind(ParameterMap);
+	GPUSkinCachePositionBuffer.Bind(ParameterMap, TEXT("GPUSkinCachePositionBuffer"));
+	GPUSkinCachePreviousPositionBuffer.Bind(ParameterMap, TEXT("GPUSkinCachePreviousPositionBuffer"));
+	IsGPUSkinPassThrough.Bind(ParameterMap, TEXT("bIsGPUSkinPassThrough"));
+}
+
 void FLocalVertexFactoryShaderParameters::GetElementShaderBindings(
 	const FSceneInterface* Scene,
 	const FSceneView* View,
@@ -168,20 +178,116 @@ void FLocalVertexFactoryShaderParameters::GetElementShaderBindings(
 	FVertexInputStreamArray& VertexStreams
 ) const
 {
-	// Decode VertexFactoryUserData as VertexFactoryUniformBuffer
-	FRHIUniformBuffer* VertexFactoryUniformBuffer = static_cast<FRHIUniformBuffer*>(BatchElement.VertexFactoryUserData);
-
-	FLocalVertexFactoryShaderParametersBase::GetElementShaderBindingsBase(
-		Scene,
-		View,
-		Shader,
-		InputStreamType,
-		FeatureLevel,
-		VertexFactory,
-		BatchElement,
-		VertexFactoryUniformBuffer,
-		ShaderBindings,
-		VertexStreams);
+	FLocalVertexFactory const* LocalVertexFactory = static_cast<FLocalVertexFactory const*>(VertexFactory);
+	ShaderBindings.Add(IsGPUSkinPassThrough, (uint32)(LocalVertexFactory->bGPUSkinPassThrough ? 1 : 0));
+	if (LocalVertexFactory->bGPUSkinPassThrough)
+	{
+		GetElementShaderBindingsGPUSkinPassThrough(
+			Scene,
+			View,
+			Shader,
+			InputStreamType,
+			FeatureLevel,
+			VertexFactory,
+			BatchElement,
+			ShaderBindings,
+			VertexStreams);
+	}
+	else
+	{
+		// Decode VertexFactoryUserData as VertexFactoryUniformBuffer
+		FRHIUniformBuffer* VertexFactoryUniformBuffer = static_cast<FRHIUniformBuffer*>(BatchElement.VertexFactoryUserData);
+
+		FLocalVertexFactoryShaderParametersBase::GetElementShaderBindingsBase(
+			Scene,
+			View,
+			Shader,
+			InputStreamType,
+			FeatureLevel,
+			VertexFactory,
+			BatchElement,
+			VertexFactoryUniformBuffer,
+			ShaderBindings,
+			VertexStreams);
+	}
+}
+
+void FLocalVertexFactoryShaderParameters::GetElementShaderBindingsGPUSkinPassThrough(
+	const FSceneInterface* Scene,
+	const FSceneView* View,
+	const FMeshMaterialShader* Shader,
+	const EVertexInputStreamType InputStreamType,
+	ERHIFeatureLevel::Type FeatureLevel,
+	const FVertexFactory* VertexFactory,
+	const FMeshBatchElement& BatchElement,
+	class FMeshDrawSingleShaderBindings& ShaderBindings,
+	FVertexInputStreamArray& VertexStreams) const
+{
+	// #dxr_todo do we need this call to the base?
+	FLocalVertexFactoryShaderParametersBase::GetElementShaderBindingsBase(Scene, View, Shader, InputStreamType, FeatureLevel, VertexFactory, BatchElement, nullptr, ShaderBindings, VertexStreams);
+
+	// todo: Add more context into VertexFactoryUserData about whether this is skin cache/mesh deformer/ray tracing etc.
+	// For now it just holds a skin cache pointer which is null if using other mesh deformers.
+	FGPUSkinBatchElementUserData* BatchUserData = (FGPUSkinBatchElementUserData*)BatchElement.VertexFactoryUserData;
+	const bool bUsesSkinCache = BatchUserData != nullptr;
+
+	check(VertexFactory->GetType() == &FGPUSkinPassthroughVertexFactory::StaticType);
+	FGPUSkinPassthroughVertexFactory const* PassthroughVertexFactory = static_cast<FGPUSkinPassthroughVertexFactory const*>(VertexFactory);
+
+	// Bone data is updated whenever animation triggers a dynamic update, animation can skip frames hence the frequency is not necessary every frame.
+	// So check if bone data is updated this frame, if not then the previous frame data is stale and not suitable for motion blur.
+	bool bBoneDataUpdatedThisFrame = (View->Family->FrameNumber == PassthroughVertexFactory->GetUpdatedFrameNumber());
+	// If world is paused, use current frame bone matrices, so velocity is canceled and skeletal mesh isn't blurred from motion.
+	bool bVerticesInMotion = !View->Family->bWorldIsPaused && bBoneDataUpdatedThisFrame;
+
+	if (bUsesSkinCache)
+	{
+		GetElementShaderBindingsSkinCache(PassthroughVertexFactory, BatchUserData, bVerticesInMotion, ShaderBindings, VertexStreams);
+	}
+	else
+	{
+		GetElementShaderBindingsMeshDeformer(PassthroughVertexFactory, bVerticesInMotion, ShaderBindings, VertexStreams);
+	}
+}
+
+void FLocalVertexFactoryShaderParameters::GetElementShaderBindingsSkinCache(
+	FGPUSkinPassthroughVertexFactory const* PassthroughVertexFactory,
+	FGPUSkinBatchElementUserData* BatchUserData,
+	bool bVerticesInMotion,
+	FMeshDrawSingleShaderBindings& ShaderBindings,
+	FVertexInputStreamArray& VertexStreams) const
+{
+	FGPUSkinCache::GetShaderBindings(
+		BatchUserData->Entry, BatchUserData->Section, bVerticesInMotion,
+		PassthroughVertexFactory,
+		GPUSkinCachePositionBuffer, GPUSkinCachePreviousPositionBuffer,
+		ShaderBindings, VertexStreams);
+}
+
+void FLocalVertexFactoryShaderParameters::GetElementShaderBindingsMeshDeformer(
+	FGPUSkinPassthroughVertexFactory const* PassthroughVertexFactory,
+	bool bVerticesInMotion,
+	FMeshDrawSingleShaderBindings& ShaderBindings,
+	FVertexInputStreamArray& VertexStreams) const
+{
+	if (PassthroughVertexFactory->PositionRDG.IsValid())
+	{
+		VertexStreams.Add(FVertexInputStream(PassthroughVertexFactory->GetPositionStreamIndex(), 0, PassthroughVertexFactory->PositionRDG->GetRHI()));
+		ShaderBindings.Add(GPUSkinCachePositionBuffer, PassthroughVertexFactory->GetPositionsSRV());
+	}
+	if (PassthroughVertexFactory->TangentRDG.IsValid() && PassthroughVertexFactory->GetTangentStreamIndex() > -1)
+	{
+		VertexStreams.Add(FVertexInputStream(PassthroughVertexFactory->GetTangentStreamIndex(), 0, PassthroughVertexFactory->TangentRDG->GetRHI()));
+	}
+
+	if (bVerticesInMotion && PassthroughVertexFactory->PrevPositionRDG)
+	{
+		ShaderBindings.Add(GPUSkinCachePreviousPositionBuffer, PassthroughVertexFactory->PrevPositionRDG->GetOrCreateSRV(FRHIBufferSRVCreateInfo(PF_R32_FLOAT)));
+	}
+	else
+	{
+		ShaderBindings.Add(GPUSkinCachePreviousPositionBuffer, PassthroughVertexFactory->GetPositionsSRV());
+	}
 }
 
 /**
@@ -219,6 +325,7 @@ void FLocalVertexFactory::ModifyCompilationEnvironment(const FVertexFactoryShade
 	// See UE-139634 for the case that lead to this.
 	OutEnvironment.SetDefine(TEXT("RAY_TRACING_DYNAMIC_MESH_IN_LOCAL_SPACE"), TEXT("1"));
 
+	OutEnvironment.SetDefine(TEXT("SUPPORT_SKINCACHE"), IsGPUSkinCacheAvailable(Parameters.Platform));
 }
 
 void FLocalVertexFactory::ValidateCompiledResult(const FVertexFactoryType* Type, EShaderPlatform Platform, const FShaderParameterMap& ParameterMap, TArray<FString>& OutErrors)
@@ -365,11 +472,9 @@ void FLocalVertexFactory::InitRHI()
 		}
 	}
 
-	// Only used with FGPUSkinPassthroughVertexFactory on platforms that do not use ManualVertexFetch
-	if (Data.PreSkinPositionComponent.VertexBuffer != nullptr)
-	{
-		Elements.Add(AccessStreamComponent(Data.PreSkinPositionComponent,14));
-	}
+	// Fill PreSkinPosition slot for GPUSkinPassThrough vertex factory, or else use a dummy buffer.
+	FVertexStreamComponent NullComponent(&GNullVertexBuffer, 0, 0, VET_Float4);
+	Elements.Add(AccessStreamComponent(Data.PreSkinPositionComponent.VertexBuffer ? Data.PreSkinPositionComponent : NullComponent, 14));
 
 	if (Data.LightMapCoordinateComponent.VertexBuffer)
 	{
diff --git a/Engine/Source/Runtime/Engine/Private/Rendering/StaticMeshVertexBuffer.cpp b/Engine/Source/Runtime/Engine/Private/Rendering/StaticMeshVertexBuffer.cpp
--- a/Engine/Source/Runtime/Engine/Private/Rendering/StaticMeshVertexBuffer.cpp
+++ b/Engine/Source/Runtime/Engine/Private/Rendering/StaticMeshVertexBuffer.cpp
@@ -506,7 +506,7 @@ void FStaticMeshVertexBuffer::BindTangentVertexBuffer(const FVertexFactory* Vert
 	}
 }
 
-void FStaticMeshVertexBuffer::BindPackedTexCoordVertexBuffer(const FVertexFactory* VertexFactory, FStaticMeshDataType& Data) const
+void FStaticMeshVertexBuffer::BindPackedTexCoordVertexBuffer(const FVertexFactory* VertexFactory, FStaticMeshDataType& Data, int32 MaxNumTexCoords) const
 {
 	Data.TextureCoordinates.Empty();
 	Data.NumTexCoords = GetNumTexCoords();
@@ -534,8 +534,16 @@ void FStaticMeshVertexBuffer::BindPackedTexCoordVertexBuffer(const FVertexFactor
 
 		uint32 UvStride = UVSizeInBytes * GetNumTexCoords();
 
+		// If the max num of UVs is specified, clamp to that number.
+		int32 ClampedNumTexCoords = GetNumTexCoords();
+		if (MaxNumTexCoords > -1)
+		{
+			ClampedNumTexCoords = FMath::Min<int32>(GetNumTexCoords(), MaxNumTexCoords);
+		}
+		check(ClampedNumTexCoords >= 0);
+
 		int32 UVIndex;
-		for (UVIndex = 0; UVIndex < (int32)GetNumTexCoords() - 1; UVIndex += 2)
+		for (UVIndex = 0; UVIndex < (int32)ClampedNumTexCoords - 1; UVIndex += 2)
 		{
 			Data.TextureCoordinates.Add(FVertexStreamComponent(
 				&TexCoordVertexBuffer,
@@ -547,7 +555,7 @@ void FStaticMeshVertexBuffer::BindPackedTexCoordVertexBuffer(const FVertexFactor
 		}
 
 		// possible last UV channel if we have an odd number
-		if (UVIndex < (int32)GetNumTexCoords())
+		if (UVIndex < (int32)ClampedNumTexCoords)
 		{
 			Data.TextureCoordinates.Add(FVertexStreamComponent(
 				&TexCoordVertexBuffer,
diff --git a/Engine/Source/Runtime/Engine/Private/SkeletalRenderCPUSkin.cpp b/Engine/Source/Runtime/Engine/Private/SkeletalRenderCPUSkin.cpp
--- a/Engine/Source/Runtime/Engine/Private/SkeletalRenderCPUSkin.cpp
+++ b/Engine/Source/Runtime/Engine/Private/SkeletalRenderCPUSkin.cpp
@@ -310,7 +310,7 @@ void FSkeletalMeshObjectCPUSkin::CacheVertices(int32 LODIndex, bool bForce) cons
 
 			MeshLODptr->PositionVertexBuffer.BindPositionVertexBuffer(VertexFactoryPtr, Data);
 			MeshLODptr->StaticMeshVertexBuffer.BindTangentVertexBuffer(VertexFactoryPtr, Data);
-			MeshLODptr->StaticMeshVertexBuffer.BindTexCoordVertexBuffer(VertexFactoryPtr, Data, 0);
+			MeshLODptr->StaticMeshVertexBuffer.BindPackedTexCoordVertexBuffer(VertexFactoryPtr, Data, MAX_TEXCOORDS);
 			MeshLODptr->StaticMeshVertexBuffer.BindLightMapVertexBuffer(VertexFactoryPtr, Data, 0);
 			MeshLODptr->MeshObjectColorBuffer->BindColorVertexBuffer(VertexFactoryPtr, Data);
 
@@ -389,7 +389,7 @@ void FSkeletalMeshObjectCPUSkin::FSkeletalMeshObjectLOD::InitResources(FSkelMesh
 
 			Self->PositionVertexBuffer.BindPositionVertexBuffer(VertexFactoryPtr, Data);
 			Self->StaticMeshVertexBuffer.BindTangentVertexBuffer(VertexFactoryPtr, Data);
-			Self->StaticMeshVertexBuffer.BindTexCoordVertexBuffer(VertexFactoryPtr, Data, 0);
+			Self->StaticMeshVertexBuffer.BindPackedTexCoordVertexBuffer(VertexFactoryPtr, Data, MAX_TEXCOORDS);
 			Self->StaticMeshVertexBuffer.BindLightMapVertexBuffer(VertexFactoryPtr, Data, 0);
 			Self->MeshObjectColorBuffer->BindColorVertexBuffer(VertexFactoryPtr, Data);
 
diff --git a/Engine/Source/Runtime/Engine/Private/SkeletalRenderGPUSkin.cpp b/Engine/Source/Runtime/Engine/Private/SkeletalRenderGPUSkin.cpp
--- a/Engine/Source/Runtime/Engine/Private/SkeletalRenderGPUSkin.cpp
+++ b/Engine/Source/Runtime/Engine/Private/SkeletalRenderGPUSkin.cpp
@@ -1627,7 +1627,7 @@ void InitGPUSkinVertexFactoryComponents(FGPUSkinDataType* VertexFactoryData, con
 
 	// tangents
 	VertexBuffers.StaticVertexBuffers->StaticMeshVertexBuffer.BindTangentVertexBuffer(VertexFactory, *VertexFactoryData);
-	VertexBuffers.StaticVertexBuffers->StaticMeshVertexBuffer.BindTexCoordVertexBuffer(VertexFactory, *VertexFactoryData);
+	VertexBuffers.StaticVertexBuffers->StaticMeshVertexBuffer.BindPackedTexCoordVertexBuffer(VertexFactory, *VertexFactoryData, MAX_TEXCOORDS);
 
 	bool bUse16BitBoneIndex = VertexBuffers.SkinWeightVertexBuffer->Use16BitBoneIndex();
 	VertexFactoryData->bUse16BitBoneIndex = bUse16BitBoneIndex;
diff --git a/Engine/Source/Runtime/Engine/Public/GPUSkinVertexFactory.h b/Engine/Source/Runtime/Engine/Public/GPUSkinVertexFactory.h
--- a/Engine/Source/Runtime/Engine/Public/GPUSkinVertexFactory.h
+++ b/Engine/Source/Runtime/Engine/Public/GPUSkinVertexFactory.h
@@ -456,17 +456,12 @@ protected:
  */
 class FGPUSkinPassthroughVertexFactory : public FLocalVertexFactory
 {
-	DECLARE_VERTEX_FACTORY_TYPE(FGPUSkinPassthroughVertexFactory);
-
 	typedef FLocalVertexFactory Super;
 
 public:
 	FGPUSkinPassthroughVertexFactory(ERHIFeatureLevel::Type InFeatureLevel);
 	~FGPUSkinPassthroughVertexFactory();
 
-	static void ModifyCompilationEnvironment(const FVertexFactoryShaderPermutationParameters& Parameters, FShaderCompilerEnvironment& OutEnvironment);
-	static bool ShouldCompilePermutation(const FVertexFactoryShaderPermutationParameters& Parameters);
-
 	inline int32 GetPositionStreamIndex() const
 	{
 		check(PositionStreamIndex > -1);
@@ -542,7 +537,7 @@ public:
 	}
 
 protected:
-	friend class FGPUSkinVertexPassthroughFactoryShaderParameters;
+	friend class FLocalVertexFactoryShaderParameters;
 
 	// Reference holders for RDG buffers
 	TRefCountPtr<FRDGPooledBuffer> PositionRDG;
diff --git a/Engine/Source/Runtime/Engine/Public/LocalVertexFactory.h b/Engine/Source/Runtime/Engine/Public/LocalVertexFactory.h
--- a/Engine/Source/Runtime/Engine/Public/LocalVertexFactory.h
+++ b/Engine/Source/Runtime/Engine/Public/LocalVertexFactory.h
@@ -155,6 +155,8 @@ public:
 #endif
 
 protected:
+	friend class FLocalVertexFactoryShaderParameters;
+
 	const FDataType& GetData() const { return Data; }
 
 	FDataType Data;
@@ -162,6 +164,8 @@ protected:
 
 	int32 ColorStreamIndex;
 
+	bool bGPUSkinPassThrough = false;
+
 	struct FDebugName
 	{
 		FDebugName(const char* InDebugName)
@@ -215,6 +219,8 @@ class FLocalVertexFactoryShaderParameters : public FLocalVertexFactoryShaderPara
 {
 	DECLARE_TYPE_LAYOUT(FLocalVertexFactoryShaderParameters, NonVirtual);
 public:
+	void Bind(const FShaderParameterMap& ParameterMap);
+
 	void GetElementShaderBindings(
 		const FSceneInterface* Scene,
 		const FSceneView* View,
@@ -226,4 +232,33 @@ public:
 		FMeshDrawSingleShaderBindings& ShaderBindings,
 		FVertexInputStreamArray& VertexStreams
 	) const; 
+
+private:
+	void GetElementShaderBindingsGPUSkinPassThrough(
+		const FSceneInterface* Scene,
+		const FSceneView* View,
+		const FMeshMaterialShader* Shader,
+		const EVertexInputStreamType InputStreamType,
+		ERHIFeatureLevel::Type FeatureLevel,
+		const FVertexFactory* VertexFactory,
+		const FMeshBatchElement& BatchElement,
+		class FMeshDrawSingleShaderBindings& ShaderBindings,
+		FVertexInputStreamArray& VertexStreams) const;
+
+	void GetElementShaderBindingsSkinCache(
+		class FGPUSkinPassthroughVertexFactory const* PassthroughVertexFactory,
+		struct FGPUSkinBatchElementUserData* BatchUserData,
+		bool bVerticesInMotion,
+		FMeshDrawSingleShaderBindings& ShaderBindings,
+		FVertexInputStreamArray& VertexStreams) const;
+
+	void GetElementShaderBindingsMeshDeformer(
+		class FGPUSkinPassthroughVertexFactory const* PassthroughVertexFactory,
+		bool bVerticesInMotion,
+		FMeshDrawSingleShaderBindings& ShaderBindings,
+		FVertexInputStreamArray& VertexStreams) const;
+
+	LAYOUT_FIELD(FShaderResourceParameter, GPUSkinCachePositionBuffer);
+	LAYOUT_FIELD(FShaderResourceParameter, GPUSkinCachePreviousPositionBuffer);
+	LAYOUT_FIELD(FShaderParameter, IsGPUSkinPassThrough);
 };
\ No newline at end of file
diff --git a/Engine/Source/Runtime/Engine/Public/Rendering/StaticMeshVertexBuffer.h b/Engine/Source/Runtime/Engine/Public/Rendering/StaticMeshVertexBuffer.h
--- a/Engine/Source/Runtime/Engine/Public/Rendering/StaticMeshVertexBuffer.h
+++ b/Engine/Source/Runtime/Engine/Public/Rendering/StaticMeshVertexBuffer.h
@@ -491,7 +491,7 @@ public:
 
 	ENGINE_API void BindTangentVertexBuffer(const FVertexFactory* VertexFactory, struct FStaticMeshDataType& Data) const;
 	ENGINE_API void BindTexCoordVertexBuffer(const FVertexFactory* VertexFactory, struct FStaticMeshDataType& Data, int ClampedNumTexCoords = -1) const;
-	ENGINE_API void BindPackedTexCoordVertexBuffer(const FVertexFactory* VertexFactory, struct FStaticMeshDataType& Data) const;
+	ENGINE_API void BindPackedTexCoordVertexBuffer(const FVertexFactory* VertexFactory, struct FStaticMeshDataType& Data, int32 MaxNumTexCoords = -1) const;
 	ENGINE_API void BindLightMapVertexBuffer(const FVertexFactory* VertexFactory, struct FStaticMeshDataType& Data, int LightMapCoordinateIndex) const;
 
 	FORCEINLINE_DEBUGGABLE void* GetTangentData() { return TangentsDataPtr; }
