Fix invalid Strata define.

diff --git a/Engine/Shaders/Private/Strata/StrataEvaluation.ush b/Engine/Shaders/Private/Strata/StrataEvaluation.ush
--- a/Engine/Shaders/Private/Strata/StrataEvaluation.ush
+++ b/Engine/Shaders/Private/Strata/StrataEvaluation.ush
@@ -39,11 +39,11 @@
 #define STRATA_SSS_TRANSMISSION 0
 #endif
 
-// Strata STRATA_LIGHTING_QUALITY define the lighting accuarcy/approximation
+// Strata STRATA_SHADING_QUALITY define the lighting accuarcy/approximation
 // 1: accurante lighting
 // 2: approximation lighting
-#ifndef STRATA_LIGHTING_QUALITY
-#define STRATA_LIGHTING_QUALITY 1
+#ifndef STRATA_SHADING_QUALITY
+#define STRATA_SHADING_QUALITY 1
 #endif
 
 ///////////////////////////////////////////////////////////////////////////////
