Fix Strata backface lighting on translucent surface with participating medium.

diff --git a/Engine/Shaders/Private/Strata/StrataEvaluation.ush b/Engine/Shaders/Private/Strata/StrataEvaluation.ush
--- a/Engine/Shaders/Private/Strata/StrataEvaluation.ush
+++ b/Engine/Shaders/Private/Strata/StrataEvaluation.ush
@@ -751,9 +751,12 @@ FStrataEvaluateResult StrataEvaluateBSDFCommon(FStrataBSDFContext BSDFContext, F
 
 				Sample.bSubsurface = false;	// It should already be the case because Enforce in this case because BSDF_GETHASSSS should be false when BSDF_GETISSIMPLEVOLUME is true
 
+				// Use Context.NoL which is not saturate, and allows to handle backface lighting, necessary for medium support
+				const float MediumNoL = abs(BSDFContext.Context.NoL);
+
 				Sample.IntegratedDiffuseValue = lerp(
 					Sample.IntegratedDiffuseValue, 
-					(ShadowTerms.SurfaceShadow * AreaLightContext.NoL * AreaLightContext.Falloff) * AreaLightContext.AreaLight.FalloffColor * SlabDirectionalAlbedo * DirectionalAlbedo_SpecularTransmission,
+					(ShadowTerms.SurfaceShadow * MediumNoL * AreaLightContext.Falloff) * AreaLightContext.AreaLight.FalloffColor * SlabDirectionalAlbedo * DirectionalAlbedo_SpecularTransmission,
 					DiffuseToVolumeBlend);
 				Sample.DiffuseColor		= SlabDirectionalAlbedo;
 			}
