HeterogeneousVolumes fixes:

1) When adjusting the step-size based on the maximum step-count, force all rays to uniformly resample based on the maximal t-span of the volume's bounding sphere.
2) Apply early termination for rays that transmit below a small epsilon value.

diff --git a/Engine/Shaders/Private/HeterogeneousVolumes/HeterogeneousVolumesRayMarchingUtils.ush b/Engine/Shaders/Private/HeterogeneousVolumes/HeterogeneousVolumesRayMarchingUtils.ush
--- a/Engine/Shaders/Private/HeterogeneousVolumes/HeterogeneousVolumesRayMarchingUtils.ush
+++ b/Engine/Shaders/Private/HeterogeneousVolumes/HeterogeneousVolumesRayMarchingUtils.ush
@@ -83,7 +83,9 @@ FRayMarchingContext CreateRayMarchingContext(
 
 uint CalcStepCount(inout FRayMarchingContext Context)
 {
-	float HitSpan = Context.LocalRayTMax - Context.LocalRayTMin;
+	// Instead using the active ray-span, use the maximal ray-span so the stepsize is uniform among all rays
+	//float HitSpan = length(Hit1 - Hit0)
+	float HitSpan = length(GetLocalBoundsExtent());
 	float StepSize = Context.StepSize;
 	uint StepCount = max(ceil(HitSpan / StepSize), 1);
 	if (StepCount > Context.MaxStepCount)
@@ -110,6 +112,13 @@ void RayMarchTransmittance(
 		FVolumeSampleContext SampleContext = CreateVolumeSampleContext(LocalPosition, WorldPosition, RayMarchingContext.MipLevel);
 		float3 Extinction = SampleExtinction(SampleContext);
 		Transmittance *= exp(-Extinction * RayMarchingContext.StepSize);
+
+		float Epsilon = 1.0e-7;
+		if (all(Transmittance < Epsilon))
+		{
+			Transmittance = 0.0;
+			break;
+		}
 	}
 }
 
@@ -300,5 +309,12 @@ void RayMarchSingleScattering(
 
 		// Accumulate transmittance for the next evaluation
 		Transmittance *= exp(-Extinction * RayMarchingContext.StepSize);
+
+		float Epsilon = 1.0e-7;
+		if (all(Transmittance < Epsilon))
+		{
+			Transmittance = 0.0;
+			break;
+		}
 	}
 }
