Preport minor change from strata legecy packing.
* Add R10G10B10F pack/unpack.
* Fix missing IsThin()

diff --git a/Engine/Shaders/Private/Strata/Strata.ush b/Engine/Shaders/Private/Strata/Strata.ush
--- a/Engine/Shaders/Private/Strata/Strata.ush
+++ b/Engine/Shaders/Private/Strata/Strata.ush
@@ -520,6 +520,8 @@ uint   Pack10F(float Value);
 float  Unpack10F(uint Value);
 uint   PackR11G11B10F(float3 rgb);
 float3 UnpackR11G11B10F(uint rgb);
+uint   PackR10G10B10F(float3 rgb);
+float3 UnpackR10G10B10F(uint rgb);
 uint   PackColorLinearToGamma2AlphaLinear(float4 rgba);
 float4 UnpackColorGamma2ToLinearAlphaLinear(uint rgba);
 uint   StrataPackNormal(in float3 Normal);
@@ -2918,6 +2920,22 @@ float3 UnpackR11G11B10F(uint rgb)
 	return float3(r, g, b);
 }
 
+uint PackR10G10B10F(float3 rgb)
+{
+	uint r = (f32tof16(rgb.r) << 15) & 0x3FF00000;  // 0011 1111 1111 0000 0000 0000 0000 0000 
+	uint g = (f32tof16(rgb.g) << 5) & 0x000FFC00;  // 0000 0000 0000 1111 1111 1100 0000 0000 
+	uint b = (f32tof16(rgb.b) >> 5) & 0x000003FF;	// 0000 0000 0000 0000 0000 0011 1111 1111
+	return r | g | b;
+}
+
+float3 UnpackR10G10B10F(uint rgb)
+{
+	float r = f16tof32((rgb >> 15) & 0x7FE0);
+	float g = f16tof32((rgb >> 5) & 0x7FE0);
+	float b = f16tof32((rgb << 5) & 0x7FE0);
+	return float3(r, g, b);
+}
+
 uint PackRG16(float2 In)
 {
 #if COMPILER_SUPPORTS_PACK_INTRINSICS
@@ -3614,6 +3632,7 @@ FStrataBSDF UnpackFastPathStrataBSDFIn(FStrataMaterialContainer StrataBuffer, in
 	BSDF_SETSSSTYPE(OutBSDF, SSS_TYPE_INVALID);
 	BSDF_SETHASHAZINESS(OutBSDF, 0);
 	BSDF_SETISSIMPLEVOLUME(OutBSDF, 0);
+	BSDF_SETISTHIN(OutBSDF, 0);
 	BSDF_SETHASFUZZ(OutBSDF, 0);
 	BSDF_SETHASTRANSABOVE(OutBSDF, 0);
 
@@ -3658,6 +3677,7 @@ FStrataBSDF UnpackFastWaterPathStrataBSDFIn(FStrataMaterialContainer StrataBuffe
 	BSDF_SETSSSTYPE(OutBSDF, SSS_TYPE_INVALID);
 	BSDF_SETHASHAZINESS(OutBSDF, 0);
 	BSDF_SETISSIMPLEVOLUME(OutBSDF, 0);
+	BSDF_SETISTHIN(OutBSDF, 0);
 	BSDF_SETHASFUZZ(OutBSDF, 0);
 	BSDF_SETHASTRANSABOVE(OutBSDF, 0);
 
