Fix (timeout) crash due to overflowing the shared marking job queue.

#rnx

diff --git a/Engine/Shaders/Private/VirtualShadowMaps/VirtualShadowMapBuildPerPageDrawCommands.usf b/Engine/Shaders/Private/VirtualShadowMaps/VirtualShadowMapBuildPerPageDrawCommands.usf
--- a/Engine/Shaders/Private/VirtualShadowMaps/VirtualShadowMapBuildPerPageDrawCommands.usf
+++ b/Engine/Shaders/Private/VirtualShadowMaps/VirtualShadowMapBuildPerPageDrawCommands.usf
@@ -129,6 +129,7 @@ void WriteCmd(uint MipViewId, uint InstanceId, uint IndirectArgIndex, bool bStat
 }
 
 #define MAX_SINGLE_THREAD_MARKING_AREA (8U)
+#define MARKING_JOB_QUEUE_SIZE (NUM_THREADS_PER_GROUP * 2U)
 
 struct FSharedMarkingJob
 {
@@ -141,7 +142,7 @@ struct FSharedMarkingJob
 };
 
 groupshared uint NumSharedMarkingJobs;
-groupshared uint2 SharedMarkingJobs[NUM_THREADS_PER_GROUP];
+groupshared uint2 SharedMarkingJobs[MARKING_JOB_QUEUE_SIZE];
 
 uint2 PackJob(FSharedMarkingJob Job)
 {
@@ -370,7 +371,10 @@ void CullPerPageDrawCommandsCs(uint3 GroupId : SV_GroupID, int GroupThreadIndex
 								Job.bMaterialInvalidates = bMaterialInvalidates;
 								Job.bShouldCacheAsStatic = bShouldCacheAsStatic;
 								Job.bHasMoved = bHasMoved;
-								SharedMarkingJobs[JobIndex] = PackJob(Job);
+								if (JobIndex < MARKING_JOB_QUEUE_SIZE)
+								{
+									SharedMarkingJobs[JobIndex] = PackJob(Job);
+								}
 								// Must assume we have mapped pages (or defer the cmd write also...)
 								NumMappedPages = 1U;
 							}
@@ -409,7 +413,7 @@ void CullPerPageDrawCommandsCs(uint3 GroupId : SV_GroupID, int GroupThreadIndex
 	}
 	GroupMemoryBarrierWithGroupSync();
 	uint NumMarkingJobs = NumSharedMarkingJobs;
-	for (uint JobIndex = 0U; JobIndex < NumMarkingJobs; ++JobIndex)
+	for (uint JobIndex = 0U; JobIndex < min(MARKING_JOB_QUEUE_SIZE, NumMarkingJobs); ++JobIndex)
 	{
 		FSharedMarkingJob Job = UnpackJob(SharedMarkingJobs[JobIndex]);
 		FVirtualSMLevelOffset PageTableLevelOffset = CalcPageTableLevelOffset(Job.VirtualShadowMapId, Job.MipLevel);
