Attempt to fix vulkan shader compilation
error: cannot instantiate StructuredBuffer with given packed alignment; 'VK_EXT_scalar_block_layout' not supported
#rnx

diff --git a/Engine/Shaders/Private/VirtualShadowMaps/VirtualShadowMapBuildPerPageDrawCommands.usf b/Engine/Shaders/Private/VirtualShadowMaps/VirtualShadowMapBuildPerPageDrawCommands.usf
--- a/Engine/Shaders/Private/VirtualShadowMaps/VirtualShadowMapBuildPerPageDrawCommands.usf
+++ b/Engine/Shaders/Private/VirtualShadowMaps/VirtualShadowMapBuildPerPageDrawCommands.usf
@@ -36,6 +36,7 @@ struct FVSMCullingBatchInfo
 	uint NumPrimaryViews;
 	uint PrimitiveRevealedOffset;
 	uint PrimitiveRevealedNum;
+	uint2 padding; // avoid error: cannot instantiate StructuredBuffer with given packed alignment; 'VK_EXT_scalar_block_layout' not supported
 };
 
 #if ENABLE_BATCH_MODE
diff --git a/Engine/Source/Runtime/Renderer/Private/VirtualShadowMaps/VirtualShadowMapArray.cpp b/Engine/Source/Runtime/Renderer/Private/VirtualShadowMaps/VirtualShadowMapArray.cpp
--- a/Engine/Source/Runtime/Renderer/Private/VirtualShadowMaps/VirtualShadowMapArray.cpp
+++ b/Engine/Source/Runtime/Renderer/Private/VirtualShadowMaps/VirtualShadowMapArray.cpp
@@ -1930,6 +1930,7 @@ struct FVSMCullingBatchInfo
 	uint32 NumPrimaryViews;
 	uint32 PrimitiveRevealedOffset;
 	uint32 PrimitiveRevealedNum;
+	uint32 padding[2]; // avoid error: cannot instantiate StructuredBuffer with given packed alignment; 'VK_EXT_scalar_block_layout' not supported
 };
 
 
