Fix Strata support for AlphaHoldout.

Converted material & native Strata material now output correct per-pixel alpha-holdout mask.

diff --git a/Engine/Shaders/Private/BasePassPixelShader.usf b/Engine/Shaders/Private/BasePassPixelShader.usf
--- a/Engine/Shaders/Private/BasePassPixelShader.usf
+++ b/Engine/Shaders/Private/BasePassPixelShader.usf
@@ -1859,7 +1859,8 @@ void FPixelShaderInOut_MainPS(
 			half3 FoggedColor = lerp(float3(1, 1, 1), DualBlendSurfaceTransmittancePreCoverage, Fogging.aaa* DualBlendSurfaceCoverage);
 			Out.MRT[0] = half4(FoggedColor, 1.0f);
 		#elif STRATA_BLENDING_ALPHAHOLDOUT
-			Out.MRT[0] = half4(0.0f.xxx, Opacity);
+			float AdjustedAlpha = saturate(1 - dot(AdjustedDualBlendMul, float3(1.0f, 1.0f, 1.0f) / 3.0f));
+			Out.MRT[0] = half4(0.0f.xxx, AdjustedAlpha);
 			Out.MRT[0] = RETURN_COLOR(Out.MRT[0]);
 		#else
 			// Pre multipled alpha blending
diff --git a/Engine/Shaders/Private/Strata/StrataLegacyConversion.ush b/Engine/Shaders/Private/Strata/StrataLegacyConversion.ush
--- a/Engine/Shaders/Private/Strata/StrataLegacyConversion.ush
+++ b/Engine/Shaders/Private/Strata/StrataLegacyConversion.ush
@@ -476,14 +476,15 @@ FStrataData StrataConvertLegacyMaterialStatic(
 		#if STRATA_INLINE_SHADING
 			Out.InlinedBSDF.Coverage = Opacity;
 		#endif
+	#elif MATERIALBLENDING_ALPHAHOLDOUT
+		// The MATERIALBLENDING_ALPHAHOLDOUT case needs to be before the MATERIALBLENDING_TRANSLUCENT case, as alphaholdout material have both cases defined
+		Out = GetStrataUnlitBSDF(0.0f, GreyTransmittance);
 	#elif MATERIALBLENDING_TRANSLUCENT
 		// Uses (Luminance, black transmittance) for translucency blending. Opacity is controled using coverage
 		Out = GetStrataUnlitBSDF(Emissive * Opacity, 0.0f);
 		#if STRATA_INLINE_SHADING
 			Out.InlinedBSDF.Coverage = Opacity;
 		#endif
-	#elif MATERIALBLENDING_ALPHAHOLDOUT
-		Out = GetStrataUnlitBSDF(0.0f, 1.0f);
 	#elif MATERIALBLENDING_ADDITIVE
 		// Emissive is weighted by opacity in the legacy material.
 		Out = GetStrataUnlitBSDF(Emissive * Opacity, 1.0f);
