Strata - 3 MRTs + legacy shading model fast packing option.

* This CL adds dedicated legacy material compaction and writes 3uint through mrt (for simple clear coat/sss wrap/cloth) instead of 4 uint with 2 through UAV ops). This allows faster base pass in particular on slower platform.
* Activating DBuffer make the entire base pass slower with this change list, but this will be addressed with subsequent changes.

Work done by Seb, submitted on his behalf.

diff --git a/Engine/Shaders/Private/BasePassPixelShader.usf b/Engine/Shaders/Private/BasePassPixelShader.usf
--- a/Engine/Shaders/Private/BasePassPixelShader.usf
+++ b/Engine/Shaders/Private/BasePassPixelShader.usf
@@ -1643,6 +1643,7 @@ void FPixelShaderInOut_MainPS(
 		// Generate the strata material data to write out
 		FStrataAddressing StrataAddressing = GetStrataPixelDataByteOffset(PixelPos, uint2(ResolvedView.BufferSizeAndInvSize.xy), StrataStruct.MaxBytesPerPixel);
 		FRWStrataMaterialContainer RWStrataMaterialContainer = InitialiseRWStrataMaterialContainer(StrataStruct.MaterialTextureArrayUAVWithoutRTs);
+		bool bForceSkipSSSData = false;
 		PackStrataOut(
 			RWStrataMaterialContainer,
 			StrataStruct.MaterialTextureArrayUAVWithoutRTs,
@@ -1650,7 +1651,7 @@ void FPixelShaderInOut_MainPS(
 			Settings,
 			StrataAddressing,
 			StrataPixelHeader, StrataData, MaterialParameters.CameraVector, WorldBentNormal0, bStrataSubsurfaceEnable, EmissiveLuminance,
-			SSSData, TopLayerData, OpaqueRoughRefractionData
+			SSSData, TopLayerData, OpaqueRoughRefractionData, bForceSkipSSSData
 	#if MATERIAL_STRATA_OPAQUE_PRECOMPUTED_LIGHTING
 			,MaterialParameters
 			,Interpolants
@@ -1661,12 +1662,13 @@ void FPixelShaderInOut_MainPS(
 		);
 
 		// Write out MRT data
-	#if STRATA_BASE_PASS_MRT_OUTPUT_COUNT != 2
+	#if STRATA_BASE_PASS_MRT_OUTPUT_COUNT != 3
 	#error Strata STRATA_BASE_PASS_MRT_OUTPUT_COUNT has been update but not StrataOutput
 	#endif
 		Out.StrataOutput[0] = RWStrataMaterialContainer.MaterialRenderTargets[0];
 		Out.StrataOutput[1] = RWStrataMaterialContainer.MaterialRenderTargets[1];
-		Out.StrataOutput[2] = StrataPackTopLayerData(TopLayerData);
+		Out.StrataOutput[2] = RWStrataMaterialContainer.MaterialRenderTargets[2];
+		Out.StrataOutput[3] = StrataPackTopLayerData(TopLayerData);
 	#if STRATA_OUTPUT_ROUGH_REFRACTION
 		OpaqueRoughRefractionData.OpaqueRoughRefractionEnabled = true; // As long as we output to rough refraction, we want it to be enabled for the tiling process.
 	#if STRATA_USES_CONVERSION_FROM_LEGACY
@@ -1685,7 +1687,7 @@ void FPixelShaderInOut_MainPS(
 
 		// Only write SSS data if needed
 		BRANCH
-		if(StrataSubSurfaceHeaderGetIsValid(SSSData.Header))
+		if(!bForceSkipSSSData && StrataSubSurfaceHeaderGetIsValid(SSSData.Header))
 		{
 			StrataStoreSubsurfaceHeader(StrataStruct.MaterialTextureArrayUAVWithoutRTs, StrataStruct.FirstSliceStoringStrataSSSDataWithoutMRT, PixelPos, SSSData.Header.Bytes);
 			if (StrataSubSurfaceHeaderHasExtras(SSSData.Header))
diff --git a/Engine/Shaders/Private/Common.ush b/Engine/Shaders/Private/Common.ush
--- a/Engine/Shaders/Private/Common.ush
+++ b/Engine/Shaders/Private/Common.ush
@@ -2007,7 +2007,7 @@ struct FPixelShaderOut
 	float4 MRT[8];
 	
 	// Explicit uint output specific to strata.
-	uint StrataOutput[3];
+	uint StrataOutput[4];
 
 	// Pixel Shader OutCoverage, only usable if PIXELSHADEROUTPUT_COVERAGE is 1
 	uint Coverage;
diff --git a/Engine/Shaders/Private/DebugProbes.usf b/Engine/Shaders/Private/DebugProbes.usf
--- a/Engine/Shaders/Private/DebugProbes.usf
+++ b/Engine/Shaders/Private/DebugProbes.usf
@@ -152,6 +152,7 @@ void WriteOutMaterial(
 
 	FStrataAddressing StrataAddressing = GetStrataPixelDataByteOffset(PixelCoord, uint2(ResolvedView.BufferSizeAndInvSize.xy), MaxBytesPerPixel);
 	FRWStrataMaterialContainer RWStrataMaterialContainer = InitialiseRWStrataMaterialContainer(MaterialTextureArrayUAV);
+	bool bForceSkipSSSData = false;
 	PackStrataOut(
 		RWStrataMaterialContainer,
 		MaterialTextureArrayUAV,
@@ -159,7 +160,7 @@ void WriteOutMaterial(
 		Settings,
 		StrataAddressing,
 		StrataPixelHeader, StrataData, -RayDir, NormalizedNormal, bStrataSubsurfaceEnable, EmissiveLuminance,
-		SSSData, TopLayerData, OpaqueRoughRefractionData);
+		SSSData, TopLayerData, OpaqueRoughRefractionData, bForceSkipSSSData);
 
 	// Write out MRTs. Note: this only works for deferred with basepassvelocity off
 	OutGBufferA = 0;
diff --git a/Engine/Shaders/Private/Definitions.usf b/Engine/Shaders/Private/Definitions.usf
--- a/Engine/Shaders/Private/Definitions.usf
+++ b/Engine/Shaders/Private/Definitions.usf
@@ -50,6 +50,10 @@
 #define MATERIALBLENDING_ALPHAHOLDOUT					0
 #endif
 
+#ifndef MATERIAL_IS_STRATA
+#define MATERIAL_IS_STRATA								0
+#endif
+
 #ifndef STRATA_BLENDING_OPAQUE
 #define STRATA_BLENDING_OPAQUE							0
 #endif
diff --git a/Engine/Shaders/Private/PixelShaderOutputCommon.ush b/Engine/Shaders/Private/PixelShaderOutputCommon.ush
--- a/Engine/Shaders/Private/PixelShaderOutputCommon.ush
+++ b/Engine/Shaders/Private/PixelShaderOutputCommon.ush
@@ -94,7 +94,7 @@ void MainPS
 #endif
 
 #if STRATA_OPAQUE_DEFERRED
-#if STRATA_BASE_PASS_MRT_OUTPUT_COUNT != 2
+#if STRATA_BASE_PASS_MRT_OUTPUT_COUNT != 3
 #error Strata STRATA_BASE_PASS_MRT_OUTPUT_COUNT has been update but not the uint MRTs
 #endif
 #if PIXELSHADEROUTPUT_MRT4
@@ -104,18 +104,22 @@ void MainPS
 		, out uint StrataOutput4 : SV_Target4
 		, out uint StrataOutput5 : SV_Target5
 		, out uint StrataOutput6 : SV_Target6
+		, out uint StrataOutput7 : SV_Target7
 #elif PIXELSHADEROUTPUT_MRT2
 		, out uint StrataOutput3 : SV_Target3
 		, out uint StrataOutput4 : SV_Target4
 		, out uint StrataOutput5 : SV_Target5
+		, out uint StrataOutput6 : SV_Target6
 #elif PIXELSHADEROUTPUT_MRT1
 		, out uint StrataOutput2 : SV_Target2
 		, out uint StrataOutput3 : SV_Target3
 		, out uint StrataOutput4 : SV_Target4
+		, out uint StrataOutput5 : SV_Target5
 #else
 		, out uint StrataOutput1 : SV_Target1
 		, out uint StrataOutput2 : SV_Target2
 		, out uint StrataOutput3 : SV_Target3
+		, out uint StrataOutput4 : SV_Target4
 #endif
 #else // STRATA_OPAQUE_DEFERRED
 
@@ -203,18 +207,22 @@ void MainPS
 	StrataOutput4 = PixelShaderOut.StrataOutput[0];
 	StrataOutput5 = PixelShaderOut.StrataOutput[1];
 	StrataOutput6 = PixelShaderOut.StrataOutput[2];
+	StrataOutput7 = PixelShaderOut.StrataOutput[3];
 #elif PIXELSHADEROUTPUT_MRT2
 	StrataOutput3 = PixelShaderOut.StrataOutput[0];
 	StrataOutput4 = PixelShaderOut.StrataOutput[1];
 	StrataOutput5 = PixelShaderOut.StrataOutput[2];
+	StrataOutput6 = PixelShaderOut.StrataOutput[3];
 #elif PIXELSHADEROUTPUT_MRT1
 	StrataOutput2 = PixelShaderOut.StrataOutput[0];
 	StrataOutput3 = PixelShaderOut.StrataOutput[1];
 	StrataOutput4 = PixelShaderOut.StrataOutput[2];
+	StrataOutput5 = PixelShaderOut.StrataOutput[3];
 #else
 	StrataOutput1 = PixelShaderOut.StrataOutput[0];
 	StrataOutput2 = PixelShaderOut.StrataOutput[1];
 	StrataOutput3 = PixelShaderOut.StrataOutput[2];
+	StrataOutput4 = PixelShaderOut.StrataOutput[3];
 #endif
 
 #else // STRATA_OPAQUE_DEFERRED
diff --git a/Engine/Shaders/Private/RayTracing/RayTracingMaterialHitShaders.usf b/Engine/Shaders/Private/RayTracing/RayTracingMaterialHitShaders.usf
--- a/Engine/Shaders/Private/RayTracing/RayTracingMaterialHitShaders.usf
+++ b/Engine/Shaders/Private/RayTracing/RayTracingMaterialHitShaders.usf
@@ -282,6 +282,7 @@ void GetMaterialPayload(
 
 	const float3 CamVector = MaterialParameters.CameraVector;	
 	Payload.StrataData = InitialiseRWStrataMaterialContainer();
+	bool bForceSkipSSSData = false;
 	PackStrataOut(
 		Payload.StrataData,
 		Dither,
@@ -292,7 +293,7 @@ void GetMaterialPayload(
 		WorldBentNormal0,
 		bStrataSubsurfaceEnable, 
 		EmissiveLuminance,
-		SSSData, TopLayerData, OpaqueRoughRefractionData
+		SSSData, TopLayerData, OpaqueRoughRefractionData, bForceSkipSSSData
 	);
 
 	Payload.StrataData.PackedTopLayerData = StrataPackTopLayerData(TopLayerData);
diff --git a/Engine/Shaders/Private/Strata/Strata.ush b/Engine/Shaders/Private/Strata/Strata.ush
--- a/Engine/Shaders/Private/Strata/Strata.ush
+++ b/Engine/Shaders/Private/Strata/Strata.ush
@@ -337,6 +337,15 @@ struct FStrataBSDF
 #define BSDF_BIT_OFFSET_HASTRANSABOVE			(BSDF_BIT_OFFSET_SHAREDLOCALBASESID	+ BSDF_BIT_COUNT_SHAREDLOCALBASESID)
 #define BSDF_BIT_OFFSET_WEIGHT10F				22
 
+// Special optmised modes for the single BSDF path.
+#define HEADER_SINGLE_OPTLEGACYMODE_BIT_COUNT	3
+#define HEADER_SINGLE_OPTLEGACYMODE_BIT_MASK	((1u<<HEADER_SINGLE_OPTLEGACYMODE_BIT_COUNT)-1u)
+#define SINGLE_OPTLEGACYMODE_NONE				0
+#define SINGLE_OPTLEGACYMODE_CLEARCOAT			1
+#define SINGLE_OPTLEGACYMODE_CLOTH				2
+#define SINGLE_OPTLEGACYMODE_SSSWRAP			3
+// 6 more remains: Cloth, Subsurface, ProfileIdSSS, Also how to deal with Aniso?
+
 // Layout 0 (Fast)
 #define BSDF_FASTENCODING_BIT_COUNT				(0u)
 #define BSDF_FASTENCODING_MASK					(0u)
@@ -761,16 +770,27 @@ void StrataSubSurfaceHeaderSetNonProfile(inout FStrataSubsurfaceHeader SSSHeader
 	SSSHeader.Bytes |= (Pack10F(MeanFreePath.x) << 22) | (Pack10F(MeanFreePath.y) << 12) | (Pack10F(MeanFreePath.z) << 2);
 }
 
-void StrataSubSurfaceHeaderSetWrap(inout FStrataSubsurfaceHeader SSSHeader, float3 MeanFreePath, float PhaseAnisotropy)
+float StrataSubSurfaceGetWrapOpacityFromAnisotropy(float PhaseAnisotropy)
 {
 	// Reinterpret the phase function anisotropy as 'opacity' value
 	// * Forward/Backward phase function == Thin surface (i.e., opacity => 0)
 	// * Isotropic phase function == Thick surface (i.e., opacity => 1)
-	const float Opacity = 1.f-abs(PhaseAnisotropy);
+	const float Opacity = 1.f - abs(PhaseAnisotropy);
+	return Opacity;
+}
+
+void StrataSubSurfaceHeaderSetWrapOpacity(inout FStrataSubsurfaceHeader SSSHeader, float Opacity)
+{
 	SSSHeader.Bytes &= SSSHEADER_TYPE_MASK;
 	SSSHeader.Bytes |= PackR8(Opacity) << 2;
 }
 
+void StrataSubSurfaceHeaderSetWrap(inout FStrataSubsurfaceHeader SSSHeader, float PhaseAnisotropy)
+{
+	const float Opacity = StrataSubSurfaceGetWrapOpacityFromAnisotropy(PhaseAnisotropy);
+	StrataSubSurfaceHeaderSetWrapOpacity(SSSHeader, Opacity);
+}
+
 bool StrataSubSurfaceHeaderGetIsValid(in FStrataSubsurfaceHeader SSSHeader)
 {
 	return SSSHEADER_TYPE(SSSHeader) != SSS_TYPE_INVALID;
@@ -1672,7 +1692,8 @@ void StrataStoreUint1(inout FRWStrataMaterialContainer StrataBuffer, RWTexture2D
 	{
 		StrataBuffer.MaterialRenderTargets[0] = StrataAddressing.CurrentIndex == 0 ? Data : StrataBuffer.MaterialRenderTargets[0];
 		StrataBuffer.MaterialRenderTargets[1] = StrataAddressing.CurrentIndex == 1 ? Data : StrataBuffer.MaterialRenderTargets[1];
-		#if STRATA_BASE_PASS_MRT_OUTPUT_COUNT != 2
+		StrataBuffer.MaterialRenderTargets[2] = StrataAddressing.CurrentIndex == 2 ? Data : StrataBuffer.MaterialRenderTargets[2];
+		#if STRATA_BASE_PASS_MRT_OUTPUT_COUNT != 3
 		#error Strata STRATA_BASE_PASS_MRT_OUTPUT_COUNT has been update but not StrataStore function
 		#endif
 	}
@@ -2923,16 +2944,16 @@ float3 UnpackR11G11B10F(uint rgb)
 uint PackR10G10B10F(float3 rgb)
 {
 	uint r = (f32tof16(rgb.r) << 15) & 0x3FF00000;  // 0011 1111 1111 0000 0000 0000 0000 0000 
-	uint g = (f32tof16(rgb.g) << 5) & 0x000FFC00;  // 0000 0000 0000 1111 1111 1100 0000 0000 
-	uint b = (f32tof16(rgb.b) >> 5) & 0x000003FF;	// 0000 0000 0000 0000 0000 0011 1111 1111
+	uint g = (f32tof16(rgb.g) << 5)  & 0x000FFC00;  // 0000 0000 0000 1111 1111 1100 0000 0000 
+	uint b = (f32tof16(rgb.b) >> 5)  & 0x000003FF;	// 0000 0000 0000 0000 0000 0011 1111 1111
 	return r | g | b;
 }
 
 float3 UnpackR10G10B10F(uint rgb)
 {
 	float r = f16tof32((rgb >> 15) & 0x7FE0);
-	float g = f16tof32((rgb >> 5) & 0x7FE0);
-	float b = f16tof32((rgb << 5) & 0x7FE0);
+	float g = f16tof32((rgb >> 5)  & 0x7FE0);
+	float b = f16tof32((rgb << 5)  & 0x7FE0);
 	return float3(r, g, b);
 }
 
@@ -3711,13 +3732,132 @@ FStrataBSDF UnpackStrataBSDFIn(FStrataMaterialContainer StrataBuffer, inout FStr
 	#if STRATA_DEFERRED_SHADING
 		Data = StrataHeader.SharedLocalBasesTypes_PackedHeader;
 	#endif
-		OutBSDF.State = (Data>>HEADER_SINGLEENCODING_BIT_COUNT) & BSDF_SINGLEENCODING_MASK;
-		OutBSDF.LuminanceWeightV = 1.0f;
-		const uint BSDFType = StrataGetBSDFType(StrataHeader);
-		BSDF_SETTYPE(OutBSDF, BSDFType);
 
+		// Restore the BSDF common state
+		OutBSDF.State = (Data >> (HEADER_SINGLEENCODING_BIT_COUNT + HEADER_SINGLE_OPTLEGACYMODE_BIT_COUNT)) & BSDF_SINGLEENCODING_MASK;
+
+		// STRATA_TODO anisotropy need to be integrated in the optimised legacy single path.
 		// For now, single encode doesn't support anisotropy, as the top layer data only contains the normal, not the full frame basis
 		BSDF_SETHASANISOTROPY(OutBSDF, 0);
+
+		uint OptimisedLegacyMode = (Data >> (HEADER_SINGLEENCODING_BIT_COUNT)) & HEADER_SINGLE_OPTLEGACYMODE_BIT_MASK;
+		if (OptimisedLegacyMode == SINGLE_OPTLEGACYMODE_NONE)
+		{
+			// General single path
+			OutBSDF.LuminanceWeightV = 1.0f;
+			const uint BSDFType = StrataGetBSDFType(StrataHeader);
+			BSDF_SETTYPE(OutBSDF, BSDFType);
+		}
+		else
+		{
+			OutBSDF.State = 0;
+			OutBSDF.LuminanceWeightV = 1.0f;
+			OutBSDF.TransmittanceAboveAlongN = 1.0f;
+
+			if (OptimisedLegacyMode == SINGLE_OPTLEGACYMODE_CLEARCOAT)
+			{
+			#if STRATA_DEFERRED_SHADING
+				uint Data0 = StrataHeader.SharedLocalBasesTypes_PackedHeader;
+				uint Data1 = StrataLoadUint1(StrataBuffer, StrataAddressing);
+				uint Data2 = StrataLoadUint1(StrataBuffer, StrataAddressing);
+
+				const uint PackedDiffuse20Bits	= Data1 & 0x000FFFFF;
+				const uint PackedF024Bits		= Data2 & 0x00FFFFFF;
+				const uint PackedClearCoat16Bits= ((Data1 >> 12) & 0xFF00) | ((Data2 >> 24) & 0x00FF);
+				const uint PackedRoughness8Bits = (Data0 >> 24) & 0xFF;
+
+				SLAB_DIFFUSEALBEDO(OutBSDF)	= UnpackR7G7B6Gamma2(PackedDiffuse20Bits);
+				SLAB_F0(OutBSDF)			= UnpackRGBA8(PackedF024Bits).xyz;
+				SLAB_HAZINESS(OutBSDF)		= PackedClearCoat16Bits;
+				SLAB_ROUGHNESS(OutBSDF)		= UnpackR8(PackedRoughness8Bits);
+			#endif
+
+				// Force a BSDF state to be a single layer on top.
+				BSDF_SETTYPE(OutBSDF, STRATA_BSDF_TYPE_SLAB);
+				BSDF_SETHASGREYWEIGHT_V(OutBSDF, 1);
+				BSDF_SETHASF90(OutBSDF, 0);
+				BSDF_SETHASANISOTROPY(OutBSDF, 0);
+				BSDF_SETISTOPLAYER(OutBSDF, 1);
+				BSDF_SETSSSTYPE(OutBSDF, SSS_TYPE_INVALID);
+				BSDF_SETHASHAZINESS(OutBSDF, 1);
+				BSDF_SETISSIMPLEVOLUME(OutBSDF, 0);
+				BSDF_SETISTHIN(OutBSDF, 0);
+				BSDF_SETHASFUZZ(OutBSDF, 0);
+				BSDF_SETHASTRANSABOVE(OutBSDF, 0);
+			}
+			else if (OptimisedLegacyMode == SINGLE_OPTLEGACYMODE_CLOTH)
+			{
+			#if STRATA_DEFERRED_SHADING
+				uint Data0 = StrataHeader.SharedLocalBasesTypes_PackedHeader;
+				uint Data1 = StrataLoadUint1(StrataBuffer, StrataAddressing);
+				uint Data2 = StrataLoadUint1(StrataBuffer, StrataAddressing);
+
+				const uint PackedDiffuse20Bits	= Data1 & 0x000FFFFF;
+				const uint PackedF020Bits		= Data2 & 0x000FFFFF;
+				const uint PackedFuzzColor20bits= ((Data1 >> 10) & 0xFFC00) | ((Data2 >> 20) & 0x03FF);
+				const uint PackedFuzzAmount8bits= (Data0 >> (HEADER_SINGLEENCODING_BIT_COUNT + HEADER_SINGLE_OPTLEGACYMODE_BIT_COUNT)) & 0xFF;
+
+				SLAB_DIFFUSEALBEDO(OutBSDF)	= UnpackR7G7B6Gamma2(PackedDiffuse20Bits);
+				SLAB_F0(OutBSDF)			= UnpackR7G7B6Gamma2(PackedF020Bits);
+				SLAB_FUZZ_COLOR(OutBSDF)	= UnpackR7G7B6Gamma2(PackedFuzzColor20bits);
+				SLAB_FUZZ_AMOUNT(OutBSDF)	= UnpackR8(PackedFuzzAmount8bits);
+				SLAB_ROUGHNESS(OutBSDF)		= StrataUnpackTopLayerData(StrataHeader.PackedTopLayerData).Roughness;
+			#endif
+
+				// Force a BSDF state to be a single layer on top.
+				BSDF_SETTYPE(OutBSDF, STRATA_BSDF_TYPE_SLAB);
+				BSDF_SETHASGREYWEIGHT_V(OutBSDF, 1);
+				BSDF_SETHASF90(OutBSDF, 0);
+				BSDF_SETHASANISOTROPY(OutBSDF, 0);
+				BSDF_SETISTOPLAYER(OutBSDF, 1);
+				BSDF_SETSSSTYPE(OutBSDF, SSS_TYPE_INVALID);
+				BSDF_SETHASHAZINESS(OutBSDF, 0);
+				BSDF_SETISSIMPLEVOLUME(OutBSDF, 0);
+				BSDF_SETHASFUZZ(OutBSDF, 1);
+				BSDF_SETHASTRANSABOVE(OutBSDF, 0);
+			}
+			else if (OptimisedLegacyMode == SINGLE_OPTLEGACYMODE_SSSWRAP)
+			{
+			#if STRATA_DEFERRED_SHADING
+				uint Data0 = StrataHeader.SharedLocalBasesTypes_PackedHeader;
+				uint Data1 = StrataLoadUint1(StrataBuffer, StrataAddressing);
+				uint Data2 = StrataLoadUint1(StrataBuffer, StrataAddressing);
+
+				const uint PackedDiffuse20Bits	= Data1 & 0x000FFFFF;
+				const uint PackedF020Bits		= Data2 & 0x000FFFFF;
+				
+				const uint PackedSSSMFP12BitsA = (Data1 >> 20) & 0xFFF;																				// 12 lower bits
+				const uint PackedSSSMFP12BitsB = (Data2 >> 20) & 0xFFF;																				// next 12 lower bits
+				const uint PackedSSSMFP6BitsC  = (Data0 >> (7 + HEADER_SINGLEENCODING_BIT_COUNT + HEADER_SINGLE_OPTLEGACYMODE_BIT_COUNT)) & 0x03F;	// 6 higher bits
+				const uint PackedSSSMFP30bits  = PackedSSSMFP12BitsA | (PackedSSSMFP12BitsB << 12) | (PackedSSSMFP6BitsC << 24);
+
+				const uint PackedSSSWOpacity7bits = (Data0 >> (HEADER_SINGLEENCODING_BIT_COUNT + HEADER_SINGLE_OPTLEGACYMODE_BIT_COUNT)) & 0x7F;
+				const float Opacity = UnpackR7(PackedSSSWOpacity7bits);
+
+				SLAB_DIFFUSEALBEDO(OutBSDF)	= UnpackR7G7B6Gamma2(PackedDiffuse20Bits);
+				SLAB_F0(OutBSDF)			= UnpackR7G7B6Gamma2(PackedF020Bits);
+				SLAB_ROUGHNESS(OutBSDF)		= StrataUnpackTopLayerData(StrataHeader.PackedTopLayerData).Roughness;
+				SLAB_SSSMFP(OutBSDF)		= UnpackR10G10B10F(PackedSSSMFP30bits);
+				SLAB_SSSPHASEANISOTROPY(OutBSDF)= saturate(1.f - Opacity);
+			#endif
+
+				// Force a BSDF state to be a single layer on top.
+				BSDF_SETTYPE(OutBSDF, STRATA_BSDF_TYPE_SLAB);
+				BSDF_SETHASGREYWEIGHT_V(OutBSDF, 1);
+				BSDF_SETHASF90(OutBSDF, 0);
+				BSDF_SETHASANISOTROPY(OutBSDF, 0);
+				BSDF_SETISTOPLAYER(OutBSDF, 1);
+				BSDF_SETSSSTYPE(OutBSDF, SSS_TYPE_WRAP);
+				BSDF_SETHASMFP(OutBSDF, 1);
+				BSDF_SETHASHAZINESS(OutBSDF, 0);
+				BSDF_SETISSIMPLEVOLUME(OutBSDF, 0);
+				BSDF_SETISTHIN(OutBSDF, 0);
+				BSDF_SETHASFUZZ(OutBSDF, 0);
+				BSDF_SETHASTRANSABOVE(OutBSDF, 0);
+			}
+
+			return OutBSDF;
+		}
 	}
 	else if (IsEye(StrataHeader))
 	{
diff --git a/Engine/Shaders/Private/Strata/StrataExport.ush b/Engine/Shaders/Private/Strata/StrataExport.ush
--- a/Engine/Shaders/Private/Strata/StrataExport.ush
+++ b/Engine/Shaders/Private/Strata/StrataExport.ush
@@ -110,7 +110,8 @@ void PackStrataOut(
 	inout float3 EmissiveColor,
 	inout FStrataSubsurfaceData SSSData,
 	inout FStrataTopLayerData TopLayerData,
-	inout FStrataOpaqueRoughRefractionData OpaqueRoughRefractionData
+	inout FStrataOpaqueRoughRefractionData OpaqueRoughRefractionData,
+	inout bool bForceSkipSSSData
 #if MATERIAL_STRATA_OPAQUE_PRECOMPUTED_LIGHTING
 	,
 	in FMaterialPixelParameters MaterialParameters,
@@ -163,6 +164,7 @@ void PackStrataOut(
 			OpaqueRoughRefractionData.VarianceCm = StandardDeviationCmForThickness * StandardDeviationCmForThickness;
 		}
 
+		uint OptimisedLegacyMode = SINGLE_OPTLEGACYMODE_NONE;
 		bool bIsSimpleMaterial = true;
 		bool bIsSingleMaterial = true;
 		bool bIsOnlySlab = true;
@@ -221,7 +223,6 @@ void PackStrataOut(
 
 						const bool bSSSMask = BSDF.bIsBottom;	// SSS and Thin can only be on the bottom layer, so ignoring bIsSimpleVolume here.
 						SetHasSubsurface(StrataPixelHeader, (bSSSMask && bHasSSS));
-
 						break;
 					}
 
@@ -302,6 +303,61 @@ void PackStrataOut(
 			bIsSimpleMaterial = bIsSimpleMaterial && BSDFVisibleCount == 1;
 			bIsSingleMaterial = bIsSingleMaterial && BSDFVisibleCount == 1;
 			SetMaterialMode(StrataPixelHeader, bIsSimpleMaterial ? HEADER_MATERIALMODE_SLAB_SIMPLE : bIsSingleMaterial ? HEADER_MATERIALMODE_SLAB_SINGLE : HEADER_MATERIALMODE_SLAB_COMPLEX);
+
+			if (bIsSingleMaterial)
+			{
+			#define BSDF StrataTree.BSDFs[0]
+				const uint SingleHasHaziness		= 0x01;
+				const uint SingleHasClearCoat		= 0x02;
+				const uint SingleHasCloth			= 0x04;
+				const uint SingleIsThin				= 0x08;
+				const uint SingleIsSimpleVolume		= 0x10;
+				const uint SingleIsWrap				= 0x20;
+				// No need to add anisotropy since it is not supported by the simple mode.
+
+				uint OptimisedModeMask = 0;
+				if (BSDF_GETHASHAZINESS(BSDF))
+				{
+					OptimisedModeMask |= SingleHasHaziness; // Regular haziness should disable non haziness/clearcoat fast path
+
+					FHaziness Haziness = UnpackHaziness(SLAB_HAZINESS(BSDF));
+					if (Haziness.bSimpleClearCoat)
+					{
+						OptimisedModeMask |= SingleHasClearCoat;
+					}
+				}
+				if (BSDF_GETHASFUZZ(BSDF))
+				{
+					OptimisedModeMask |= SingleHasCloth;
+				}
+				if (BSDF_GETISTHIN(BSDF))
+				{
+					OptimisedModeMask |= SingleIsThin;
+				}
+				if (BSDF_GETISSIMPLEVOLUME(BSDF))
+				{
+					OptimisedModeMask |= SingleIsSimpleVolume;
+				}
+				if (BSDF_GETSSSTYPE(BSDF) == SSS_TYPE_WRAP)
+				{
+					OptimisedModeMask |= SingleIsWrap;
+				}
+
+				if (((OptimisedModeMask & SingleHasClearCoat) == SingleHasClearCoat) || ((OptimisedModeMask & SingleHasClearCoat) == (SingleHasHaziness | SingleHasClearCoat)))
+				{
+					OptimisedLegacyMode = SINGLE_OPTLEGACYMODE_CLEARCOAT;
+				}
+				else if ((OptimisedModeMask & SingleHasCloth) == SingleHasCloth)
+				{
+					OptimisedLegacyMode = SINGLE_OPTLEGACYMODE_CLOTH;
+				}
+				else if ((OptimisedModeMask & SingleIsWrap) == SingleIsWrap)
+				{
+					OptimisedLegacyMode = SINGLE_OPTLEGACYMODE_SSSWRAP;
+					bForceSkipSSSData = true;
+				}
+			#undef BSDF
+			}
 		}
 
 		// Now evaluate Specular Occlusion and Indirect Irradiance.
@@ -456,9 +512,93 @@ void PackStrataOut(
 						// Data0 (Header_State|Header_AO|Header_BSDFTypes|BSDF_State)
 						uint Out = 0;
 						HEADER_SETCOMMONSTATES(Out, StrataPixelHeader.State);
+
 						Out = Out & HEADER_SINGLEENCODING_MASK;
-						Out = Out | ((BSDF.State & BSDF_SINGLEENCODING_MASK) << HEADER_SINGLEENCODING_BIT_COUNT);
-						STRATA_STORE_UINT1(Out);
+						Out = Out | ((OptimisedLegacyMode & HEADER_SINGLE_OPTLEGACYMODE_BIT_MASK)	<< (HEADER_SINGLEENCODING_BIT_COUNT));
+
+						BRANCH
+						if (OptimisedLegacyMode == SINGLE_OPTLEGACYMODE_NONE)
+						{
+							Out = Out | ((BSDF.State & BSDF_SINGLEENCODING_MASK) << (HEADER_SINGLEENCODING_BIT_COUNT + HEADER_SINGLE_OPTLEGACYMODE_BIT_COUNT));
+							STRATA_STORE_UINT1(Out);
+						}
+						// With this special path, there is 32-8-8-3= 13bits available for packing on the header.
+						// Because we use 3 UINT MRT, this means that a total of 13+32+32 = 77 bits.
+						else if(OptimisedLegacyMode == SINGLE_OPTLEGACYMODE_CLEARCOAT)
+						{
+							const uint PackedDiffuse20Bits	= PackR7G7B6Gamma2(SLAB_DIFFUSEALBEDO(BSDF), Dither);
+							const uint PackedF024Bits		= PackRGBA8(float4(SLAB_F0(BSDF), 0.0f));
+							const uint PackedClearCoat16Bits= SLAB_HAZINESS(BSDF);
+							const uint PackedRoughness8Bits = PackR8(SLAB_ROUGHNESS(BSDF));
+
+							// MRT0 
+							#if (32 - (HEADER_SINGLEENCODING_BIT_COUNT + HEADER_SINGLE_OPTLEGACYMODE_BIT_COUNT)) < 8
+							#error Not Enough space to store roughness in Slab - Single - Header.
+							#endif
+							STRATA_STORE_UINT1(Out					| ((PackedRoughness8Bits & 0xFF) << 24));
+							// MRT1 20 PackedDiffuse and 8 high bits from PackedClearCoat
+							STRATA_STORE_UINT1(PackedDiffuse20Bits	| ((PackedClearCoat16Bits & 0xFF00) << 12));
+							// MRT2 24 PackedF0      and 8 low  bits from PackedClearCoat
+							STRATA_STORE_UINT1(PackedF024Bits		| (PackedClearCoat16Bits << 24));
+
+							// Ensure the rest of the BSDF is not stored with the regular path
+							BSDF.State = 0;
+						}
+						else if (OptimisedLegacyMode == SINGLE_OPTLEGACYMODE_CLOTH)
+						{
+							const uint PackedDiffuse20Bits	= PackR7G7B6Gamma2(	SLAB_DIFFUSEALBEDO(BSDF)	, Dither);
+							const uint PackedF020Bits		= PackR7G7B6Gamma2(	SLAB_F0(BSDF)				, Dither);
+							const uint PackedFuzzColor20bits= PackR7G7B6Gamma2(	SLAB_FUZZ_COLOR(BSDF)		, Dither);
+							const uint PackedFuzzAmount8bits= PackR8(			SLAB_FUZZ_AMOUNT(BSDF));
+							// Roughness is stored in the TopLayer texture.
+
+							// MRT0 
+							#if (32 - (HEADER_SINGLEENCODING_BIT_COUNT + HEADER_SINGLE_OPTLEGACYMODE_BIT_COUNT)) < 8
+							#error Not Enough space to store fuzz amount in slab in Slab - Single - Header.
+							#endif
+							STRATA_STORE_UINT1(Out					| ((PackedFuzzAmount8bits & 0xFF) << (HEADER_SINGLEENCODING_BIT_COUNT + HEADER_SINGLE_OPTLEGACYMODE_BIT_COUNT)));
+							// MRT1 20 PackedDiffuse and 8 high bits from PackedClearCoat
+							STRATA_STORE_UINT1(PackedDiffuse20Bits	| ((PackedFuzzColor20bits & 0xFFC00) << 10));
+							// MRT2 24 PackedF0      and 8 low  bits from PackedClearCoat
+							STRATA_STORE_UINT1(PackedF020Bits		| ((PackedFuzzColor20bits & 0x3FF) << 20));
+
+							// Ensure the rest of the BSDF is not stored with the regular path
+							BSDF.State = 0;
+						}
+						else if (OptimisedLegacyMode == SINGLE_OPTLEGACYMODE_SSSWRAP)
+						{
+							const float SSSWOpacity = StrataSubSurfaceGetWrapOpacityFromAnisotropy(SLAB_SSSPHASEANISOTROPY(BSDF));
+
+							const uint PackedDiffuse20Bits		= PackR7G7B6Gamma2(	SLAB_DIFFUSEALBEDO(BSDF)	, Dither);
+							const uint PackedF020Bits			= PackR7G7B6Gamma2(	SLAB_F0(BSDF)				, Dither);
+							const uint PackedSSSWOpacity7bits	= PackR7(			SSSWOpacity					, 0.5);	// We should not dither here because it would have a large scale effect.
+							const uint PackedSSSMFP30bits		= PackR10G10B10F(	SLAB_SSSMFP(BSDF));
+							// Roughness is stored in the TopLayer texture.
+
+							const uint PackedSSSMFP12BitsA = (PackedSSSMFP30bits      ) & 0xFFF; // 12 lower bits
+							const uint PackedSSSMFP12BitsB = (PackedSSSMFP30bits >> 12) & 0xFFF; // next 12 lower bits
+							const uint PackedSSSMFP6BitsC  = (PackedSSSMFP30bits >> 24) & 0x03F; // 6 higher bits
+
+							// MRT0 
+							#if (32 - (HEADER_SINGLEENCODING_BIT_COUNT + HEADER_SINGLE_OPTLEGACYMODE_BIT_COUNT)) < (7+6)
+							#error Not Enough space to store wrapped SSS opacity in slab in Slab - Single - Header.
+							#endif
+							STRATA_STORE_UINT1(Out					| ((PackedSSSWOpacity7bits & 0x7F) << (HEADER_SINGLEENCODING_BIT_COUNT + HEADER_SINGLE_OPTLEGACYMODE_BIT_COUNT)) | (PackedSSSMFP6BitsC << (7 + HEADER_SINGLEENCODING_BIT_COUNT + HEADER_SINGLE_OPTLEGACYMODE_BIT_COUNT)));
+							// MRT1 20 PackedDiffuse and 8 high bits from PackedClearCoat
+							STRATA_STORE_UINT1(PackedDiffuse20Bits  | (PackedSSSMFP12BitsA << 20));
+							// MRT2 24 PackedF0      and 8 low  bits from PackedClearCoat
+							STRATA_STORE_UINT1(PackedF020Bits		| (PackedSSSMFP12BitsB << 20));
+
+							// Ensure the rest of the BSDF is not stored with the regular path
+							BSDF.State = 0;
+						}
+						// else if (OptimisedLegacyMode == SINGLE_OPTLEGACYMODE_xxx)
+						//
+						// STRATA_TODO
+						// - Fast cloth
+						// - Fast subsurface
+						// - Fast diffusion with profile
+						// - Anisotropy as an option for Disney and clear coat as for current UE5?
 
 						#if (HEADER_SIMPLEENCODING_BIT_COUNT) > 32
 						#error Strata fast path header is > 32bits
@@ -539,7 +679,7 @@ void PackStrataOut(
 					}
 
 					// Layout1 & Layout2
-					if (!bFastEncodedBSDF && !bHasFastWaterEncoding)
+					if (!bFastEncodedBSDF && !bHasFastWaterEncoding && OptimisedLegacyMode == SINGLE_OPTLEGACYMODE_NONE)
 					{
 						const uint BSDFType = BSDF_GETTYPE(BSDF);
 						switch (BSDF_GETTYPE(BSDF))
@@ -578,7 +718,7 @@ void PackStrataOut(
 									StrataSubSurfaceHeaderSetSSSType(SSSData.Header, BSDF_GETSSSTYPE(BSDF));
 									if (BSDF_GETSSSTYPE(BSDF) == SSS_TYPE_WRAP)
 									{
-										StrataSubSurfaceHeaderSetWrap(SSSData.Header, SLAB_SSSMFP(BSDF), SLAB_SSSPHASEANISOTROPY(BSDF));
+										StrataSubSurfaceHeaderSetWrap(SSSData.Header, SLAB_SSSPHASEANISOTROPY(BSDF));
 									}
 									else if (BSDF_GETSSSTYPE(BSDF) == SSS_TYPE_DIFFUSION_PROFILE)
 									{
diff --git a/Engine/Shaders/Private/Strata/StrataMaterialClassification.usf b/Engine/Shaders/Private/Strata/StrataMaterialClassification.usf
--- a/Engine/Shaders/Private/Strata/StrataMaterialClassification.usf
+++ b/Engine/Shaders/Private/Strata/StrataMaterialClassification.usf
@@ -102,6 +102,20 @@ void TileMainCS(uint2 DispatchThreadId : SV_DispatchThreadID, uint LinearIndex :
 		bContainsOpaqueRoughRefraction = OpaqueRoughRefractionData.OpaqueRoughRefractionEnabled > 0.0f;
 	#endif
 
+		uint OptimisedLegacyMode = ((PackedHeader >> (HEADER_SINGLEENCODING_BIT_COUNT)) & HEADER_SINGLE_OPTLEGACYMODE_BIT_MASK);
+		if (bIsSingle && OptimisedLegacyMode == SINGLE_OPTLEGACYMODE_SSSWRAP)
+		{
+			bContainsScreenSpaceSubsurfaceScattering = true;
+
+			const uint PackedSSSWOpacity7bits = (PackedHeader >> (HEADER_SINGLEENCODING_BIT_COUNT + HEADER_SINGLE_OPTLEGACYMODE_BIT_COUNT)) & 0x7F;
+
+			FStrataSubsurfaceHeader SSSHeader = (FStrataSubsurfaceHeader)0;
+			StrataSubSurfaceHeaderSetSSSType(SSSHeader, SSS_TYPE_WRAP);
+			StrataSubSurfaceHeaderSetWrapOpacity(SSSHeader, UnpackR7(PackedSSSWOpacity7bits));
+
+			StrataStoreSubsurfaceHeader(MaterialTextureArrayUAV, FirstSliceStoringStrataSSSData, PixelCoord, SSSHeader.Bytes);
+		}
+
 		// Fine grain test for clearing based on CMask data
 	#if PERMUTATION_CMASK
 		// Fine grain check if clear is needed
diff --git a/Engine/Shaders/Private/Strata/StrataVisualize.usf b/Engine/Shaders/Private/Strata/StrataVisualize.usf
--- a/Engine/Shaders/Private/Strata/StrataVisualize.usf
+++ b/Engine/Shaders/Private/Strata/StrataVisualize.usf
@@ -169,7 +169,7 @@ void PrintBSDF(inout FShaderPrintContext Context, inout FStrataAddressing Strata
 
 			Print(Context, TEXT("Slab "), FontBSDFType);
 			Newline(Context, RectMax);
-			Print(Context, TEXT(" NormalID Aniso TopLayer Scattering SingleScatt. IsThin WeightV "), FontBSDFStateName);
+			Print(Context, TEXT(" NormalID Aniso TopLayer Scattering SimpleVolume IsThin WeightV "), FontBSDFStateName);
 			if (!bGreyWeightV)
 			{
 				Print(Context, TEXT("           "), FontBSDFStateName);
@@ -441,6 +441,24 @@ void PrintPixel(uint2 InCoord, float3 InWorldPosition, float3 V)
 		Print(Context, TEXT("BSDF"), FontBSDFCount);
 		Print(Context, TEXT(" - "), FontWhite);
 		PrintPixelType(Context, Header, FontLightRed);
+		if (bIsSingleMaterial)
+		{
+		#if STRATA_DEFERRED_SHADING
+			uint OptimisedLegacyMode = (Header.SharedLocalBasesTypes_PackedHeader >> (HEADER_SINGLEENCODING_BIT_COUNT)) & HEADER_SINGLE_OPTLEGACYMODE_BIT_MASK;
+			if (OptimisedLegacyMode ==SINGLE_OPTLEGACYMODE_CLEARCOAT)
+			{
+				Print(Context, TEXT(" - Legecy Clear-Coat"), FontLightRed);
+			}
+			else if (OptimisedLegacyMode ==SINGLE_OPTLEGACYMODE_CLOTH)
+			{
+				Print(Context, TEXT(" - Legecy Cloth"), FontLightRed);
+			}
+			else if (OptimisedLegacyMode ==SINGLE_OPTLEGACYMODE_SSSWRAP)
+			{
+				Print(Context, TEXT(" - Legecy SSS-Wrap"), FontLightRed);
+			}		
+		#endif
+		}
 		Newline(Context, RectMax);
 
 		Print(Context, TEXT(" AO    IndIrr TopRoughness PreShadow ZeroShadow ContacShadow Ind.Occluder HasSSS"), FontHeaderStateName);
@@ -492,7 +510,7 @@ void PrintPixel(uint2 InCoord, float3 InWorldPosition, float3 V)
 		const bool bHasSSSData = StrataSubSurfaceHeaderGetIsValid(SSSHeader);
 
 		const uint TopLayerDataBytes = 4;	// bytes 
-		const uint SSSDataBytes = 8;	    // bytes 
+		const uint SSSDataBytes = 8; // bIsSingleMaterial&& ? 0 : 8;	    // bytes  TODO accessor to OPT mode
 
 		const uint HeaderSize =  FootPrint_PostHeader - FootPrint_Start;
 		const uint BSDFsSize  =  FootPrint_PostBSDFs - FootPrint_PostHeader;
diff --git a/Engine/Shaders/Shared/StrataDefinitions.h b/Engine/Shaders/Shared/StrataDefinitions.h
--- a/Engine/Shaders/Shared/StrataDefinitions.h
+++ b/Engine/Shaders/Shared/StrataDefinitions.h
@@ -39,7 +39,7 @@
 // The thin lighting model is used when the thickness becomes strictly lower than 0.1 millimeter
 #define STRATA_LAYER_ISTHIN_THICKNESS_THRESHOLD_CM 0.01f
 
-#define STRATA_BASE_PASS_MRT_OUTPUT_COUNT	2
+#define STRATA_BASE_PASS_MRT_OUTPUT_COUNT	3
 
 #define STRATA_SSS_DATA_UINT_COUNT		2
 
