Remove bForceSSS flags and repalce it with SSS header instead.

diff --git a/Engine/Shaders/Private/BasePassPixelShader.usf b/Engine/Shaders/Private/BasePassPixelShader.usf
--- a/Engine/Shaders/Private/BasePassPixelShader.usf
+++ b/Engine/Shaders/Private/BasePassPixelShader.usf
@@ -1643,7 +1643,6 @@ void FPixelShaderInOut_MainPS(
 		// Generate the strata material data to write out
 		FStrataAddressing StrataAddressing = GetStrataPixelDataByteOffset(PixelPos, uint2(ResolvedView.BufferSizeAndInvSize.xy), StrataStruct.MaxBytesPerPixel);
 		FRWStrataMaterialContainer RWStrataMaterialContainer = InitialiseRWStrataMaterialContainer(StrataStruct.MaterialTextureArrayUAVWithoutRTs);
-		bool bForceSkipSSSData = false;
 		PackStrataOut(
 			RWStrataMaterialContainer,
 			StrataStruct.MaterialTextureArrayUAVWithoutRTs,
@@ -1651,7 +1650,7 @@ void FPixelShaderInOut_MainPS(
 			Settings,
 			StrataAddressing,
 			StrataPixelHeader, StrataData, MaterialParameters.CameraVector, WorldBentNormal0, bStrataSubsurfaceEnable, EmissiveLuminance,
-			SSSData, TopLayerData, OpaqueRoughRefractionData, bForceSkipSSSData
+			SSSData, TopLayerData, OpaqueRoughRefractionData
 	#if MATERIAL_STRATA_OPAQUE_PRECOMPUTED_LIGHTING
 			,MaterialParameters
 			,Interpolants
@@ -1687,7 +1686,7 @@ void FPixelShaderInOut_MainPS(
 
 		// Only write SSS data if needed
 		BRANCH
-		if(!bForceSkipSSSData && StrataSubSurfaceHeaderGetIsValid(SSSData.Header))
+		if(StrataSubSurfaceHeaderGetIsValid(SSSData.Header))
 		{
 			StrataStoreSubsurfaceHeader(StrataStruct.MaterialTextureArrayUAVWithoutRTs, StrataStruct.FirstSliceStoringStrataSSSDataWithoutMRT, PixelPos, SSSData.Header.Bytes);
 			if (StrataSubSurfaceHeaderHasExtras(SSSData.Header))
diff --git a/Engine/Shaders/Private/DebugProbes.usf b/Engine/Shaders/Private/DebugProbes.usf
--- a/Engine/Shaders/Private/DebugProbes.usf
+++ b/Engine/Shaders/Private/DebugProbes.usf
@@ -152,7 +152,6 @@ void WriteOutMaterial(
 
 	FStrataAddressing StrataAddressing = GetStrataPixelDataByteOffset(PixelCoord, uint2(ResolvedView.BufferSizeAndInvSize.xy), MaxBytesPerPixel);
 	FRWStrataMaterialContainer RWStrataMaterialContainer = InitialiseRWStrataMaterialContainer(MaterialTextureArrayUAV);
-	bool bForceSkipSSSData = false;
 	PackStrataOut(
 		RWStrataMaterialContainer,
 		MaterialTextureArrayUAV,
@@ -160,7 +159,7 @@ void WriteOutMaterial(
 		Settings,
 		StrataAddressing,
 		StrataPixelHeader, StrataData, -RayDir, NormalizedNormal, bStrataSubsurfaceEnable, EmissiveLuminance,
-		SSSData, TopLayerData, OpaqueRoughRefractionData, bForceSkipSSSData);
+		SSSData, TopLayerData, OpaqueRoughRefractionData);
 
 	// Write out MRTs. Note: this only works for deferred with basepassvelocity off
 	OutGBufferA = 0;
diff --git a/Engine/Shaders/Private/RayTracing/RayTracingMaterialHitShaders.usf b/Engine/Shaders/Private/RayTracing/RayTracingMaterialHitShaders.usf
--- a/Engine/Shaders/Private/RayTracing/RayTracingMaterialHitShaders.usf
+++ b/Engine/Shaders/Private/RayTracing/RayTracingMaterialHitShaders.usf
@@ -282,7 +282,6 @@ void GetMaterialPayload(
 
 	const float3 CamVector = MaterialParameters.CameraVector;	
 	Payload.StrataData = InitialiseRWStrataMaterialContainer();
-	bool bForceSkipSSSData = false;
 	PackStrataOut(
 		Payload.StrataData,
 		Dither,
@@ -293,7 +292,7 @@ void GetMaterialPayload(
 		WorldBentNormal0,
 		bStrataSubsurfaceEnable, 
 		EmissiveLuminance,
-		SSSData, TopLayerData, OpaqueRoughRefractionData, bForceSkipSSSData
+		SSSData, TopLayerData, OpaqueRoughRefractionData
 	);
 
 	Payload.StrataData.PackedTopLayerData = StrataPackTopLayerData(TopLayerData);
diff --git a/Engine/Shaders/Private/Strata/Strata.ush b/Engine/Shaders/Private/Strata/Strata.ush
--- a/Engine/Shaders/Private/Strata/Strata.ush
+++ b/Engine/Shaders/Private/Strata/Strata.ush
@@ -3813,6 +3813,7 @@ FStrataBSDF UnpackStrataBSDFIn(FStrataMaterialContainer StrataBuffer, inout FStr
 				BSDF_SETSSSTYPE(OutBSDF, SSS_TYPE_INVALID);
 				BSDF_SETHASHAZINESS(OutBSDF, 0);
 				BSDF_SETISSIMPLEVOLUME(OutBSDF, 0);
+				BSDF_SETISTHIN(OutBSDF, 0);
 				BSDF_SETHASFUZZ(OutBSDF, 1);
 				BSDF_SETHASTRANSABOVE(OutBSDF, 0);
 			}
diff --git a/Engine/Shaders/Private/Strata/StrataExport.ush b/Engine/Shaders/Private/Strata/StrataExport.ush
--- a/Engine/Shaders/Private/Strata/StrataExport.ush
+++ b/Engine/Shaders/Private/Strata/StrataExport.ush
@@ -110,8 +110,7 @@ void PackStrataOut(
 	inout float3 EmissiveColor,
 	inout FStrataSubsurfaceData SSSData,
 	inout FStrataTopLayerData TopLayerData,
-	inout FStrataOpaqueRoughRefractionData OpaqueRoughRefractionData,
-	inout bool bForceSkipSSSData
+	inout FStrataOpaqueRoughRefractionData OpaqueRoughRefractionData
 #if MATERIAL_STRATA_OPAQUE_PRECOMPUTED_LIGHTING
 	,
 	in FMaterialPixelParameters MaterialParameters,
@@ -354,7 +353,7 @@ void PackStrataOut(
 				else if ((OptimisedModeMask & SingleIsWrap) == SingleIsWrap)
 				{
 					OptimisedLegacyMode = SINGLE_OPTLEGACYMODE_SSSWRAP;
-					bForceSkipSSSData = true;
+					SSSData.Header.Bytes = 0; // Force to skip SSSData writing
 				}
 			#undef BSDF
 			}
