Miscellaneous fixes to HeterogeneousVolume:

1) Replacing shadow step-factor with shadow step-size, decoupling its dependence on step-size
2) Adding jitter to direct volume rendering integrators to help when undersampling
3) Repurposing jitter as a constant ray bias when casting shadow rays

diff --git a/Engine/Shaders/Private/HeterogeneousVolumes/HeterogeneousVolumesHardwareRayTracing.usf b/Engine/Shaders/Private/HeterogeneousVolumes/HeterogeneousVolumesHardwareRayTracing.usf
--- a/Engine/Shaders/Private/HeterogeneousVolumes/HeterogeneousVolumesHardwareRayTracing.usf
+++ b/Engine/Shaders/Private/HeterogeneousVolumes/HeterogeneousVolumesHardwareRayTracing.usf
@@ -130,7 +130,7 @@ RAY_TRACING_ENTRY_ANY_HIT(SparseVoxelsAnyHitShader,
 	{
 		float3 Transmittance = Payload.GetTransmittance();
 
-		float StepSize = GetStepSize() * GetShadowStepFactor();
+		float StepSize = GetShadowStepSize();
 		float HitSpan = Payload.HitT[1] - Payload.HitT[0];
 		int StepCount = max(HitSpan / StepSize, 1);
 		StepSize = HitSpan / StepCount;
@@ -458,8 +458,8 @@ RAY_TRACING_ENTRY_RAYGEN(RenderTransmittanceVolumeWithPreshadingRGS)
 	float LocalRayTMax = length(LocalRayDirection);
 	LocalRayDirection /= LocalRayTMax;
 
-	float Jitter = bJitter ? InterleavedGradientNoise(VoxelIndex.xy, VoxelIndex.z) : 0.5;
-	float StepSize = GetStepSize() * GetShadowStepFactor();
+	float RayBias = 0.5;
+	float StepSize = GetShadowStepSize();
 	int bApplyEmissionAndTransmittance = 0;
 	int bApplyDirectLighting = 0;
 	int bApplyShadowTransmittance = 0;
@@ -474,7 +474,7 @@ RAY_TRACING_ENTRY_RAYGEN(RenderTransmittanceVolumeWithPreshadingRGS)
 		WorldRayOrigin,
 		WorldRayDirection,
 		// Ray-step attributes
-		Jitter,
+		RayBias,
 		StepSize,
 		MaxStepCount,
 		// Lighting
diff --git a/Engine/Shaders/Private/HeterogeneousVolumes/HeterogeneousVolumesLiveShadingPipeline.usf b/Engine/Shaders/Private/HeterogeneousVolumes/HeterogeneousVolumesLiveShadingPipeline.usf
--- a/Engine/Shaders/Private/HeterogeneousVolumes/HeterogeneousVolumesLiveShadingPipeline.usf
+++ b/Engine/Shaders/Private/HeterogeneousVolumes/HeterogeneousVolumesLiveShadingPipeline.usf
@@ -40,7 +40,7 @@ int bApplyEmissionAndTransmittance;
 int bApplyDirectLighting;
 int bApplyShadowTransmittance;
 int LightType;
-float ShadowStepFactor;
+float ShadowStepSize;
 
 // Ray data
 float MaxTraceDistance;
@@ -114,9 +114,9 @@ float GetStepSize()
 	return StepSize;
 }
 
-float GetShadowStepFactor()
+float GetShadowStepSize()
 {
-	return ShadowStepFactor;
+	return ShadowStepSize;
 }
 
 float GetMaxTraceDistance()
@@ -236,8 +236,8 @@ void RenderTransmittanceVolumeWithLiveShadingCS(
 	float LocalRayTMax = length(LocalRayDirection);
 	LocalRayDirection /= LocalRayTMax;
 
-	float RayJitter = length(VoxelJitter);
-	float StepSize = GetStepSize() * GetShadowStepFactor();
+	float RayBias = 0.5;
+	float StepSize = GetShadowStepSize();
 	int bApplyEmissionAndTransmittance = 0;
 	int bApplyDirectLighting = 0;
 	int bApplyShadowTransmittance = 0;
@@ -252,7 +252,7 @@ void RenderTransmittanceVolumeWithLiveShadingCS(
 		WorldRayOrigin,
 		WorldRayDirection,
 		// Ray-step attributes
-		RayJitter,
+		RayBias,
 		StepSize,
 		MaxStepCount,
 		// Lighting
diff --git a/Engine/Shaders/Private/HeterogeneousVolumes/HeterogeneousVolumesPreshadingPipeline.usf b/Engine/Shaders/Private/HeterogeneousVolumes/HeterogeneousVolumesPreshadingPipeline.usf
--- a/Engine/Shaders/Private/HeterogeneousVolumes/HeterogeneousVolumesPreshadingPipeline.usf
+++ b/Engine/Shaders/Private/HeterogeneousVolumes/HeterogeneousVolumesPreshadingPipeline.usf
@@ -481,8 +481,8 @@ void RenderTransmittanceVolumeWithPreshadingCS(
 	LocalRayDirection /= LocalRayTMax;
 
 	// Build ray-marching context
-	float RayJitter = length(VoxelJitter);
-	float StepSize = GetStepSize() * GetShadowStepFactor();
+	float RayJitter = 0.5;
+	float StepSize = GetShadowStepSize();
 	int bApplyEmissionAndTransmittance = 0;
 	int bApplyDirectLighting = 0;
 	int bApplyShadowTransmittance = 0;
diff --git a/Engine/Shaders/Private/HeterogeneousVolumes/HeterogeneousVolumesRayMarchingUtils.ush b/Engine/Shaders/Private/HeterogeneousVolumes/HeterogeneousVolumesRayMarchingUtils.ush
--- a/Engine/Shaders/Private/HeterogeneousVolumes/HeterogeneousVolumesRayMarchingUtils.ush
+++ b/Engine/Shaders/Private/HeterogeneousVolumes/HeterogeneousVolumesRayMarchingUtils.ush
@@ -85,7 +85,7 @@ uint CalcStepCount(inout FRayMarchingContext Context)
 {
 	// Instead using the active ray-span, use the maximal ray-span so the stepsize is uniform among all rays
 	//float HitSpan = length(Hit1 - Hit0)
-	float HitSpan = length(GetLocalBoundsExtent());
+	float HitSpan = length(2.0 * GetLocalBoundsExtent());
 	float StepSize = Context.StepSize;
 	uint StepCount = max(ceil(HitSpan / StepSize), 1);
 	if (StepCount > Context.MaxStepCount)
@@ -106,8 +106,8 @@ void RayMarchTransmittance(
 {
 	for (uint StepIndex = 0; StepIndex < StepCount; ++StepIndex)
 	{
-		float3 LocalPosition = RayMarchingContext.LocalRayOrigin + RayMarchingContext.LocalRayDirection * (RayMarchingContext.LocalRayTMin + RayMarchingContext.StepSize * StepIndex);
-		float3 WorldPosition = RayMarchingContext.WorldRayOrigin + RayMarchingContext.WorldRayDirection * (RayMarchingContext.LocalRayTMin + RayMarchingContext.StepSize * StepIndex);
+		float3 LocalPosition = RayMarchingContext.LocalRayOrigin + RayMarchingContext.LocalRayDirection * (RayMarchingContext.LocalRayTMin + RayMarchingContext.StepSize * (RayMarchingContext.Jitter + StepIndex));
+		float3 WorldPosition = RayMarchingContext.WorldRayOrigin + RayMarchingContext.WorldRayDirection * (RayMarchingContext.LocalRayTMin + RayMarchingContext.StepSize * (RayMarchingContext.Jitter + StepIndex));
 
 		FVolumeSampleContext SampleContext = CreateVolumeSampleContext(LocalPosition, WorldPosition, RayMarchingContext.MipLevel);
 		float3 Extinction = SampleExtinction(SampleContext);
@@ -127,6 +127,7 @@ float3 ComputeTransmittance(
 )
 {
 #if HARDWARE_RAY_TRACING
+	//  TODO: Incorporpate ray jitter
 	return ComputeTransmittanceHardwareRayTracing(RayMarchingContext.WorldRayOrigin, RayMarchingContext.WorldRayDirection, RayMarchingContext.LocalRayTMin, RayMarchingContext.LocalRayTMax);
 #else // HARDWARE_RAY_TRACING
 
@@ -160,8 +161,8 @@ void RayMarchEmissionAbsorption(
 {
 	for (uint StepIndex = 0; StepIndex < StepCount; ++StepIndex)
 	{
-		float3 LocalPosition = RayMarchingContext.LocalRayOrigin + RayMarchingContext.LocalRayDirection * (RayMarchingContext.LocalRayTMin + RayMarchingContext.StepSize * StepIndex);
-		float3 WorldPosition = RayMarchingContext.WorldRayOrigin + RayMarchingContext.WorldRayDirection * (RayMarchingContext.LocalRayTMin + RayMarchingContext.StepSize * StepIndex);
+		float3 LocalPosition = RayMarchingContext.LocalRayOrigin + RayMarchingContext.LocalRayDirection * (RayMarchingContext.LocalRayTMin + RayMarchingContext.StepSize * (RayMarchingContext.Jitter + StepIndex));
+		float3 WorldPosition = RayMarchingContext.WorldRayOrigin + RayMarchingContext.WorldRayDirection * (RayMarchingContext.LocalRayTMin + RayMarchingContext.StepSize * (RayMarchingContext.Jitter + StepIndex));
 
 		FVolumeSampleContext SampleContext = CreateVolumeSampleContext(LocalPosition, WorldPosition, RayMarchingContext.MipLevel);
 		float3 Extinction = SampleExtinction(SampleContext);
@@ -185,13 +186,13 @@ void RayMarchSingleScattering(
 {
 	for (uint StepIndex = 0; StepIndex < StepCount; ++StepIndex)
 	{
-		float3 LocalPosition = RayMarchingContext.LocalRayOrigin + RayMarchingContext.LocalRayDirection * (RayMarchingContext.LocalRayTMin + RayMarchingContext.StepSize * StepIndex);
-		float3 WorldPosition = RayMarchingContext.WorldRayOrigin + RayMarchingContext.WorldRayDirection * (RayMarchingContext.LocalRayTMin + RayMarchingContext.StepSize * StepIndex);
+		float3 LocalPosition = RayMarchingContext.LocalRayOrigin + RayMarchingContext.LocalRayDirection * (RayMarchingContext.LocalRayTMin + RayMarchingContext.StepSize * (RayMarchingContext.Jitter + StepIndex));
+		float3 WorldPosition = RayMarchingContext.WorldRayOrigin + RayMarchingContext.WorldRayDirection * (RayMarchingContext.LocalRayTMin + RayMarchingContext.StepSize * (RayMarchingContext.Jitter + StepIndex));
 		FVolumeSampleContext SampleContext = CreateVolumeSampleContext(LocalPosition, WorldPosition, RayMarchingContext.MipLevel);
 
 		if (RayMarchingContext.bApplyEmissionAndTransmittance)
 		{
-			float3 Emission = SampleEmission(SampleContext);
+			float3 Emission = SampleEmission(SampleContext);// +10;
 			Radiance += Emission * RayMarchingContext.StepSize * Transmittance;
 		}
 
@@ -272,9 +273,8 @@ void RayMarchSingleScattering(
 						float LocalShadowRayTMax = length(LocalShadowRayDirection);
 						LocalShadowRayDirection /= LocalShadowRayTMax;
 
-						//float ShadowJitter = InterleavedGradientNoise(PixelCoord, StepIndex);
-						float ShadowJitter = 0.0;
-						float ShadowStepSize = GetStepSize() * GetShadowStepFactor();
+						float ShadowBias = 0.5;
+						float ShadowStepSize = GetShadowStepSize();
 						int bApplyEmissionAndTransmittance = 0;
 						int bApplyDirectLighting = 0;
 						int bApplyShadowTransmittance = 0;
@@ -286,7 +286,7 @@ void RayMarchSingleScattering(
 							LocalShadowRayTMax,
 							WorldPosition, // WorldShadowRayOrigin
 							normalize(ToLight), // WorldShadowRayDirection
-							ShadowJitter,
+							ShadowBias,
 							ShadowStepSize,
 							RayMarchingContext.MaxStepCount,
 							bApplyEmissionAndTransmittance,
diff --git a/Engine/Shaders/Private/HeterogeneousVolumes/HeterogeneousVolumesSparseVoxelUniformBufferUtils.ush b/Engine/Shaders/Private/HeterogeneousVolumes/HeterogeneousVolumesSparseVoxelUniformBufferUtils.ush
--- a/Engine/Shaders/Private/HeterogeneousVolumes/HeterogeneousVolumesSparseVoxelUniformBufferUtils.ush
+++ b/Engine/Shaders/Private/HeterogeneousVolumes/HeterogeneousVolumesSparseVoxelUniformBufferUtils.ush
@@ -35,9 +35,9 @@ float GetStepSize()
 	return SparseVoxelUniformBuffer.StepSize;
 }
 
-float GetShadowStepFactor()
+float GetShadowStepSize()
 {
-	return SparseVoxelUniformBuffer.ShadowStepFactor;
+	return SparseVoxelUniformBuffer.ShadowStepSize;
 }
 
 uint3 GetVolumeResolution()
diff --git a/Engine/Source/Runtime/Renderer/Private/HeterogeneousVolumes/HeterogeneousVolumes.cpp b/Engine/Source/Runtime/Renderer/Private/HeterogeneousVolumes/HeterogeneousVolumes.cpp
--- a/Engine/Source/Runtime/Renderer/Private/HeterogeneousVolumes/HeterogeneousVolumes.cpp
+++ b/Engine/Source/Runtime/Renderer/Private/HeterogeneousVolumes/HeterogeneousVolumes.cpp
@@ -94,10 +94,10 @@ static TAutoConsoleVariable<int32> CVarHeterogeneousVolumesPreshadingVolumeResol
 	ECVF_RenderThreadSafe
 );
 
-static TAutoConsoleVariable<float> CVarHeterogeneousVolumesShadowStepFactor(
-	TEXT("r.HeterogeneousVolumes.ShadowStepFactor"),
+static TAutoConsoleVariable<float> CVarHeterogeneousVolumesShadowStepSize(
+	TEXT("r.HeterogeneousVolumes.ShadowStepSize"),
 	8.0,
-	TEXT("The ray-marching step-size factor for shadow rays (Default = 8.0)"),
+	TEXT("The ray-marching step-size for shadow rays (Default = 8.0)"),
 	ECVF_RenderThreadSafe
 );
 
@@ -198,9 +198,9 @@ namespace HeterogeneousVolumes
 		return VolumeResolution;
 	}
 
-	float GetShadowStepFactor()
+	float GetShadowStepSize()
 	{
-		return CVarHeterogeneousVolumesShadowStepFactor.GetValueOnRenderThread();
+		return CVarHeterogeneousVolumesShadowStepSize.GetValueOnRenderThread();
 	}
 
 	float GetMaxTraceDistance()
diff --git a/Engine/Source/Runtime/Renderer/Private/HeterogeneousVolumes/HeterogeneousVolumes.h b/Engine/Source/Runtime/Renderer/Private/HeterogeneousVolumes/HeterogeneousVolumes.h
--- a/Engine/Source/Runtime/Renderer/Private/HeterogeneousVolumes/HeterogeneousVolumes.h
+++ b/Engine/Source/Runtime/Renderer/Private/HeterogeneousVolumes/HeterogeneousVolumes.h
@@ -28,7 +28,7 @@ namespace HeterogeneousVolumes
 	FIntVector GetVolumeResolution();
 	FIntVector GetTransmittanceVolumeResolution();
 
-	float GetShadowStepFactor();
+	float GetShadowStepSize();
 	float GetMaxTraceDistance();
 	float GetMaxShadowTraceDistance();
 	float GetStepSize();
@@ -83,7 +83,7 @@ BEGIN_UNIFORM_BUFFER_STRUCT(FSparseVoxelUniformBufferParameters, )
 	SHADER_PARAMETER(float, MaxTraceDistance)
 	SHADER_PARAMETER(float, MaxShadowTraceDistance)
 	SHADER_PARAMETER(float, StepSize)
-	SHADER_PARAMETER(float, ShadowStepFactor)
+	SHADER_PARAMETER(float, ShadowStepSize)
 END_UNIFORM_BUFFER_STRUCT()
 
 BEGIN_SHADER_PARAMETER_STRUCT(FTransmittanceVolumeParameters, )
diff --git a/Engine/Source/Runtime/Renderer/Private/HeterogeneousVolumes/HeterogeneousVolumesLiveShadingPipeline.cpp b/Engine/Source/Runtime/Renderer/Private/HeterogeneousVolumes/HeterogeneousVolumesLiveShadingPipeline.cpp
--- a/Engine/Source/Runtime/Renderer/Private/HeterogeneousVolumes/HeterogeneousVolumesLiveShadingPipeline.cpp
+++ b/Engine/Source/Runtime/Renderer/Private/HeterogeneousVolumes/HeterogeneousVolumesLiveShadingPipeline.cpp
@@ -28,7 +28,7 @@ class FRenderTransmittanceVolumeWithLiveShadingCS : public FMeshMaterialShader
 		SHADER_PARAMETER_STRUCT_REF(FDeferredLightUniformStruct, DeferredLight)
 
 		// Shadow data
-		SHADER_PARAMETER(float, ShadowStepFactor)
+		SHADER_PARAMETER(float, ShadowStepSize)
 
 		// Object data
 		SHADER_PARAMETER(FMatrix44f, LocalToWorld)
@@ -139,7 +139,7 @@ class FRenderSingleScatteringWithLiveShadingCS : public FMeshMaterialShader
 		SHADER_PARAMETER_STRUCT_REF(FDeferredLightUniformStruct, DeferredLight)
 
 		// Shadow data
-		SHADER_PARAMETER(float, ShadowStepFactor)
+		SHADER_PARAMETER(float, ShadowStepSize)
 		SHADER_PARAMETER_RDG_UNIFORM_BUFFER(FForwardLightData, ForwardLightData)
 		SHADER_PARAMETER_STRUCT_INCLUDE(FVolumeShadowingShaderParameters, VolumeShadowingShaderParameters)
 
@@ -334,7 +334,7 @@ void RenderTransmittanceVolumeWithLiveShading(
 		PassParameters->MaxTraceDistance = HeterogeneousVolumes::GetMaxTraceDistance();
 		PassParameters->MaxShadowTraceDistance = HeterogeneousVolumes::GetMaxShadowTraceDistance();
 		PassParameters->StepSize = HeterogeneousVolumes::GetStepSize();
-		PassParameters->ShadowStepFactor = HeterogeneousVolumes::GetShadowStepFactor();
+		PassParameters->ShadowStepSize = HeterogeneousVolumes::GetShadowStepSize();
 		PassParameters->MaxStepCount = HeterogeneousVolumes::GetMaxStepCount();
 		PassParameters->bJitter = HeterogeneousVolumes::ShouldJitter();
 
@@ -420,7 +420,7 @@ void RenderSingleScatteringWithLiveShading(
 		}
 		PassParameters->DeferredLight = CreateUniformBufferImmediate(DeferredLightUniform, UniformBuffer_SingleDraw);
 		PassParameters->LightType = LightType;
-		PassParameters->ShadowStepFactor = HeterogeneousVolumes::GetShadowStepFactor();
+		PassParameters->ShadowStepSize = HeterogeneousVolumes::GetShadowStepSize();
 
 		// Object data
 		FMatrix44f LocalToWorld = FMatrix44f(PrimitiveSceneProxy->GetLocalToWorld());
diff --git a/Engine/Source/Runtime/Renderer/Private/HeterogeneousVolumes/HeterogeneousVolumesPreshadingPipeline.cpp b/Engine/Source/Runtime/Renderer/Private/HeterogeneousVolumes/HeterogeneousVolumesPreshadingPipeline.cpp
--- a/Engine/Source/Runtime/Renderer/Private/HeterogeneousVolumes/HeterogeneousVolumesPreshadingPipeline.cpp
+++ b/Engine/Source/Runtime/Renderer/Private/HeterogeneousVolumes/HeterogeneousVolumesPreshadingPipeline.cpp
@@ -1097,7 +1097,7 @@ void RenderWithPreshading(
 		SparseVoxelUniformBufferParameters->MaxTraceDistance = HeterogeneousVolumes::GetMaxTraceDistance();
 		SparseVoxelUniformBufferParameters->MaxShadowTraceDistance = HeterogeneousVolumes::GetMaxShadowTraceDistance();
 		SparseVoxelUniformBufferParameters->StepSize = HeterogeneousVolumes::GetStepSize();
-		SparseVoxelUniformBufferParameters->ShadowStepFactor = HeterogeneousVolumes::GetShadowStepFactor();
+		SparseVoxelUniformBufferParameters->ShadowStepSize = HeterogeneousVolumes::GetShadowStepSize();
 	}
 	TRDGUniformBufferRef<FSparseVoxelUniformBufferParameters> SparseVoxelUniformBuffer = GraphBuilder.CreateUniformBuffer(SparseVoxelUniformBufferParameters);
 
