Fix a couple issues introduced by CL21962580 "Reduce GPUSkinPassThrough vertex factory shader permuations by merging it into local vertex factory."
Issue 1: unbound parameters errors when `r.ShaderDevelopmentMode=1` launching EngineTest
Issue 2: Debug editor build Ensure condition failed in MeshPassProcesser.cpp shader with vertex factory FLocalVertexFactory never set texture at BaseIndex6

diff --git a/Engine/Shaders/Private/LocalVertexFactory.ush b/Engine/Shaders/Private/LocalVertexFactory.ush
--- a/Engine/Shaders/Private/LocalVertexFactory.ush
+++ b/Engine/Shaders/Private/LocalVertexFactory.ush
@@ -61,8 +61,8 @@ float4 SplineParams[10];
 #define SplineMeshY				SplineParams[9].xyz
 #endif	// USE_SPLINEDEFORM
 
-#ifndef SUPPORT_SKINCACHE
-#define SUPPORT_SKINCACHE 0
+#ifndef SUPPORT_GPUSKIN_PASSTHROUGH
+#define SUPPORT_GPUSKIN_PASSTHROUGH 0
 #endif
 
 #ifndef MANUAL_VERTEX_FETCH
@@ -158,7 +158,7 @@ struct FVertexFactoryInput
 	VF_INSTANCED_STEREO_DECLARE_INPUT_BLOCK()
 #endif
 
-#if SUPPORT_SKINCACHE && !MANUAL_VERTEX_FETCH
+#if SUPPORT_GPUSKIN_PASSTHROUGH && !MANUAL_VERTEX_FETCH
 	float4	PreSkinPosition	: ATTRIBUTE14;
 #endif
 
@@ -169,13 +169,13 @@ struct FVertexFactoryInput
 	uint VertexId : SV_VertexID;
 };
 
-#if SUPPORT_SKINCACHE
+#if SUPPORT_GPUSKIN_PASSTHROUGH
 uint bIsGPUSkinPassThrough;
 #endif
 
 bool IsGPUSkinPassThrough()
 {
-	#if SUPPORT_SKINCACHE
+	#if SUPPORT_GPUSKIN_PASSTHROUGH
 		return bIsGPUSkinPassThrough;
 	#else
 		return false;
@@ -954,7 +954,7 @@ FVertexFactoryIntermediates GetVertexFactoryIntermediates(FVertexFactoryInput In
 
 	if (IsGPUSkinPassThrough())
 	{
-	#if SUPPORT_SKINCACHE
+	#if SUPPORT_GPUSKIN_PASSTHROUGH
 	#if MANUAL_VERTEX_FETCH
 		uint PreSkinVertexOffset = LocalVF.PreSkinBaseVertexIndex + Input.VertexId * 3;
 		Intermediates.PreSkinPosition.x = LocalVF.VertexFetch_PreSkinPositionBuffer[PreSkinVertexOffset + 0];
@@ -1146,8 +1146,10 @@ float3 VertexFactoryGetWorldNormal(FVertexFactoryInput Input, FVertexFactoryInte
 	return Intermediates.TangentToWorld[2];
 }
 
+#if SUPPORT_GPUSKIN_PASSTHROUGH
 // Information when using pass through to get the previous position, as no matrix is available/needed
 Buffer<float> GPUSkinCachePreviousPositionBuffer;
+#endif
 
 // @return previous translated world position
 float4 VertexFactoryGetPreviousWorldPosition(FVertexFactoryInput Input, FVertexFactoryIntermediates Intermediates)
@@ -1155,6 +1157,7 @@ float4 VertexFactoryGetPreviousWorldPosition(FVertexFactoryInput Input, FVertexF
 	FLWCMatrix PreviousLocalToWorld = GetInstanceData(Intermediates).PrevLocalToWorld;
 	float4x4 PreviousLocalToWorldTranslated = LWCMultiplyTranslation(PreviousLocalToWorld, ResolvedView.PrevPreViewTranslation);
 
+#if SUPPORT_GPUSKIN_PASSTHROUGH
 	if (IsGPUSkinPassThrough())
 	{
 		uint Offset = Input.VertexId * 3;
@@ -1165,6 +1168,7 @@ float4 VertexFactoryGetPreviousWorldPosition(FVertexFactoryInput Input, FVertexF
 		return mul(float4(PreviousPos, 1), PreviousLocalToWorldTranslated);	
 	}
 	else
+#endif // SUPPORT_GPUSKIN_PASSTHROUGH
 	{
 	#if USE_INSTANCING
 		float4x4 InstanceTransform = GetInstancePrevTransform(Intermediates);
diff --git a/Engine/Source/Runtime/Engine/Private/LocalVertexFactory.cpp b/Engine/Source/Runtime/Engine/Private/LocalVertexFactory.cpp
--- a/Engine/Source/Runtime/Engine/Private/LocalVertexFactory.cpp
+++ b/Engine/Source/Runtime/Engine/Private/LocalVertexFactory.cpp
@@ -195,6 +195,9 @@ void FLocalVertexFactoryShaderParameters::GetElementShaderBindings(
 	}
 	else
 	{
+		// Bind null SRV as not used by non-passthrough vertex factory
+		ShaderBindings.Add(GPUSkinCachePreviousPositionBuffer, GNullVertexBuffer.VertexBufferSRV);
+
 		// Decode VertexFactoryUserData as VertexFactoryUniformBuffer
 		FRHIUniformBuffer* VertexFactoryUniformBuffer = static_cast<FRHIUniformBuffer*>(BatchElement.VertexFactoryUserData);
 
@@ -325,7 +328,10 @@ void FLocalVertexFactory::ModifyCompilationEnvironment(const FVertexFactoryShade
 	// See UE-139634 for the case that lead to this.
 	OutEnvironment.SetDefine(TEXT("RAY_TRACING_DYNAMIC_MESH_IN_LOCAL_SPACE"), TEXT("1"));
 
-	OutEnvironment.SetDefine(TEXT("SUPPORT_SKINCACHE"), IsGPUSkinCacheAvailable(Parameters.Platform));
+	if (Parameters.VertexFactoryType->SupportsGPUSkinPassThrough())
+	{
+		OutEnvironment.SetDefine(TEXT("SUPPORT_GPUSKIN_PASSTHROUGH"), IsGPUSkinCacheAvailable(Parameters.Platform));
+	}
 }
 
 void FLocalVertexFactory::ValidateCompiledResult(const FVertexFactoryType* Type, EShaderPlatform Platform, const FShaderParameterMap& ParameterMap, TArray<FString>& OutErrors)
@@ -555,4 +561,5 @@ IMPLEMENT_VERTEX_FACTORY_TYPE(FLocalVertexFactory,"/Engine/Private/LocalVertexFa
 	| EVertexFactoryFlags::SupportsLightmapBaking
 	| EVertexFactoryFlags::SupportsManualVertexFetch
 	| EVertexFactoryFlags::SupportsPSOPrecaching
+	| EVertexFactoryFlags::SupportsGPUSkinPassThrough
 );
diff --git a/Engine/Source/Runtime/RenderCore/Public/VertexFactory.h b/Engine/Source/Runtime/RenderCore/Public/VertexFactory.h
--- a/Engine/Source/Runtime/RenderCore/Public/VertexFactory.h
+++ b/Engine/Source/Runtime/RenderCore/Public/VertexFactory.h
@@ -125,6 +125,7 @@ enum class EVertexFactoryFlags : uint32
 	SupportsPSOPrecaching				  = 1u << 13,
 	SupportsManualVertexFetch			  = 1u << 14,
 	DoesNotSupportNullPixelShader		  = 1u << 15,
+	SupportsGPUSkinPassThrough			  = 1u << 16
 };
 ENUM_CLASS_FLAGS(EVertexFactoryFlags);
 
@@ -433,6 +434,7 @@ public:
 	bool SupportsLightmapBaking() const                { return HasFlags(EVertexFactoryFlags::SupportsLightmapBaking); }
 	bool SupportsPSOPrecaching() const                 { return HasFlags(EVertexFactoryFlags::SupportsPSOPrecaching); }
 	bool SupportsNullPixelShader() const			   { return !HasFlags(EVertexFactoryFlags::DoesNotSupportNullPixelShader); }
+	bool SupportsGPUSkinPassThrough() const			   { return HasFlags(EVertexFactoryFlags::SupportsGPUSkinPassThrough); }
 
 	bool SupportsManualVertexFetch(ERHIFeatureLevel::Type InFeatureLevel) const 
 	{
