- Fix Cascade GPU simulations not operating in LWC space correctly

diff --git a/Engine/Shaders/Private/ParticleGPUSpriteVertexFactory.ush b/Engine/Shaders/Private/ParticleGPUSpriteVertexFactory.ush
--- a/Engine/Shaders/Private/ParticleGPUSpriteVertexFactory.ush
+++ b/Engine/Shaders/Private/ParticleGPUSpriteVertexFactory.ush
@@ -245,7 +245,6 @@ FVertexFactoryIntermediates GetVertexFactoryIntermediates( FVertexFactoryInput I
 	// For clarity, store some information in local variables.
 	const float RelativeTime = PositionSample.w;
 	const float IsAlive = step( RelativeTime, 1.0f );
-	float3 ParticlePosition = LWCHackToFloat(MakeLWCVector3(EmitterDynamicUniforms.LWCTile, PositionSample.xyz));
 
 	// Put velocity in to world space.
 	float3 ParticleWorldVelocity = mul(VelocitySample.xyz, GetLocalToWorld3x3()).xyz;
@@ -311,7 +310,9 @@ FVertexFactoryIntermediates GetVertexFactoryIntermediates( FVertexFactoryInput I
 	Rotation = Rotation * 2.0f * 3.1415926535897932f;
 
 	// Transform position to (post-view-translation) world space.
-	float3 ParticleTranslatedWorldPosition = TransformLocalToTranslatedWorld(ParticlePosition.xyz).xyz;
+	FLWCVector3 LWCTileOffset = MakeLWCVector3(EmitterDynamicUniforms.LWCTile, 0);
+	FLWCVector3 ParticleWorldPosition = LWCAdd(LWCMultiply(PositionSample.xyz, GetPrimitiveDataFromUniformBuffer().LocalToWorld), LWCTileOffset);
+	float3 ParticleTranslatedWorldPosition = LWCToFloat(LWCAdd(ParticleWorldPosition, ResolvedView.PreViewTranslation));
 
 	// Compute tangents for the sprite.
 	float3 TangentUp, TangentRight;
diff --git a/Engine/Shaders/Private/ParticleSimulationShader.usf b/Engine/Shaders/Private/ParticleSimulationShader.usf
--- a/Engine/Shaders/Private/ParticleSimulationShader.usf
+++ b/Engine/Shaders/Private/ParticleSimulationShader.usf
@@ -283,9 +283,9 @@ SamplerState RenderAttributesTextureSampler;
 float CollisionDepthBounds;
 
 /** TODO: Should be moved to Common.usf and used globaly. */
-float3 WorldPositionFromSceneDepth(float2 ScreenPosition, float SceneDepth)
+float3 TranslatedWorldPositionFromSceneDepth(float2 ScreenPosition, float SceneDepth)
 {
-	float3 WorldPosition = mul(float4(ScreenPosition * SceneDepth, SceneDepth, 1), LWCHackToFloat(PrimaryView.ScreenToWorld)).xyz;
+	float3 WorldPosition = mul(float4(ScreenPosition * SceneDepth, SceneDepth, 1), PrimaryView.ScreenToTranslatedWorld).xyz;
 	return WorldPosition;
 }
 
@@ -329,27 +329,29 @@ float3 ComputeCollidingVelocity(
  * Compute collision with the depth buffer.
  */
 void CollideWithDepthBuffer(
-	out float3 NewPosition,
+	out FLWCVector3 NewPosition_LWC,
 	out float3 NewVelocity,
 	inout float RelativeTime,
-	in float3 InPosition,
+	in FLWCVector3 InPosition_LWC,
 	in float3 InVelocity,
 	in float3 Acceleration,
 	in float CollisionRadius,
 	in float Resilience
 	)
 {
+	float3 InPosition = LWCToFloat(LWCAdd(InPosition_LWC, PrimaryView.PreViewTranslation));
+
 	// Integration assuming no collision.
 	float3 MidVelocity = InVelocity.xyz + 0.5f * Acceleration;
 	float3 DeltaPosition = DeltaSeconds * MidVelocity;
-	NewPosition = InPosition.xyz + DeltaPosition;
+	float3 NewPosition = InPosition.xyz + DeltaPosition;
 	NewVelocity = InVelocity.xyz + Acceleration;
 
 	// Figure out where to sample the depth buffer.
 	float3 CollisionOffset = normalize(DeltaPosition) * CollisionRadius;
 
 	float3 CollisionPosition = InPosition + CollisionOffset;
-	float4 SamplePosition = float4(CollisionPosition + LWCHackToFloat(PrimaryView.PreViewTranslation),1);
+	float4 SamplePosition = float4(CollisionPosition, 1);
 	float4 ClipPosition = mul(SamplePosition, View.TranslatedWorldToClip);
 	float2 ScreenPosition = ClipPosition.xy / ClipPosition.w;
 
@@ -370,21 +372,22 @@ void CollideWithDepthBuffer(
 			// in the view frustum. To avoid this, we disable the collision test if one of the derivate is above a threshold.
 			if (max(abs(SceneDepth - SceneDepth0), abs(SceneDepth - SceneDepth1)) > CollisionDepthBounds)
 			{
+				NewPosition_LWC = LWCSubtract(NewPosition, PrimaryView.PreViewTranslation);
 				return;
 			}
 #endif
 			// Reconstruct world position.
-			float3 WorldPosition = WorldPositionFromSceneDepth(ScreenPosition.xy, SceneDepth);
+			float3 WorldPosition = TranslatedWorldPositionFromSceneDepth(ScreenPosition.xy, SceneDepth);
 			
 #if FORWARD_SHADING_NORMAL_CALC && FORWARD_SHADING_ACCURATE_NORMAL
 			// Compute normal using the two additional neighbooring pixel depth buffer fetches.
-			float3 WorldPosition0 = WorldPositionFromSceneDepth(ScreenPosition.xy + float2(2 * View.ViewSizeAndInvSize.z, 0.0), SceneDepth0);
-			float3 WorldPosition1 = WorldPositionFromSceneDepth(ScreenPosition.xy - float2(0.0, 2 * View.ViewSizeAndInvSize.w), SceneDepth1);
+			float3 WorldPosition0 = TranslatedWorldPositionFromSceneDepth(ScreenPosition.xy + float2(2 * View.ViewSizeAndInvSize.z, 0.0), SceneDepth0);
+			float3 WorldPosition1 = TranslatedWorldPositionFromSceneDepth(ScreenPosition.xy - float2(0.0, 2 * View.ViewSizeAndInvSize.w), SceneDepth1);
 			float3 WorldNormal = normalize(cross(WorldPosition0 - WorldPosition, WorldPosition1 - WorldPosition));
 #elif FORWARD_SHADING_NORMAL_CALC // && !FORWARD_SHADING_ACCURATE_NORMAL
 			// Compute normal using a ddx/ddy hack, hopefully the neighbooring particules will be close enough...
 			float3 WorldNormal = normalize(cross(ddx(WorldPosition), ddy(WorldPosition)));
-			WorldNormal *= sign(dot(LWCHackToFloat(PrimaryView.WorldCameraOrigin) - NewPosition, WorldNormal));
+			WorldNormal *= sign(dot(ResolvedView.TranslatedWorldCameraOrigin - NewPosition, WorldNormal));
 #else //!FORWARD_SHADING_NORMAL_CALC
 			// Sample the normal buffer to create a plane to collide against.
 			float3 WorldNormal = Texture2DSampleLevel(SceneTexturesStruct.GBufferATexture, SceneTexturesStruct_GBufferATextureSampler, ScreenUV, 0).xyz * 2.0 - 1.0;
@@ -394,8 +397,6 @@ void CollideWithDepthBuffer(
 			// Compute the portion of velocity normal to the collision plane.
 			float VelocityDot = dot(CollisionPlane.xyz, DeltaPosition.xyz);
 
-
-
 #if IMPROVED_DEPTH_BUFFER_COLLISION
 			// distance to the plane from current and predicted position
 			float d_back = ( dot(CollisionPlane.xyz, InPosition.xyz)+CollisionRadius - CollisionPlane.w );
@@ -455,6 +456,8 @@ void CollideWithDepthBuffer(
 #endif
 		}
 	}
+
+	NewPosition_LWC = LWCSubtract(NewPosition, PrimaryView.PreViewTranslation);
 }
 #endif // #if DEPTH_BUFFER_COLLISION
 
@@ -464,28 +467,32 @@ void CollideWithDepthBuffer(
  * Compute collision with the global signed distance field
  */
 void CollideWithDistanceField(
-	out float3 NewPosition,
+	out FLWCVector3 NewPosition_LWC,
 	out float3 NewVelocity,
 	inout float RelativeTime,
-	in float3 InPosition,
+	in FLWCVector3 InPosition_LWC,
 	in float3 InVelocity,
 	in float3 Acceleration,
 	in float CollisionRadius,
 	in float Resilience
 	)
 {
+	//LWC_TODO: The GDF doesn't support LWC currently
+	float3 InPosition = LWCToFloat(LWCAdd(InPosition_LWC, PrimaryView.PreViewTranslation));
+	float3 GDFPosition = LWCHackToFloat(InPosition_LWC);
+
 	// Integration assuming no collision.
 	float3 MidVelocity = InVelocity.xyz + 0.5f * Acceleration;
 	float3 DeltaPosition = DeltaSeconds * MidVelocity;
-	NewPosition = InPosition.xyz + DeltaPosition;
+	float3 NewPosition = InPosition.xyz + DeltaPosition;
 	NewVelocity = InVelocity.xyz + Acceleration;
 
-	float DistanceToNearestSurface = GetDistanceToNearestSurfaceGlobal(InPosition);
+	float DistanceToNearestSurface = GetDistanceToNearestSurfaceGlobal(GDFPosition);
 	float MaxCollisionDistance = CollisionRadius + length(DeltaPosition.xyz);
 
 	if (DistanceToNearestSurface < MaxCollisionDistance)
 	{
-		float3 CollisionWorldNormal = normalize(GetDistanceFieldGradientGlobal(InPosition));
+		float3 CollisionWorldNormal = normalize(GetDistanceFieldGradientGlobal(GDFPosition));
 		float3 CollisionWorldPosition = InPosition - CollisionWorldNormal * DistanceToNearestSurface;
 
 		float4 CollisionPlane = float4(CollisionWorldNormal.xyz, dot(CollisionWorldPosition.xyz, CollisionWorldNormal.xyz));
@@ -530,6 +537,8 @@ void CollideWithDistanceField(
 			RelativeTime = 1.1f;
 		}
 	}
+
+	NewPosition_LWC = LWCSubtract(NewPosition, PrimaryView.PreViewTranslation);
 }
 
 #endif
@@ -554,7 +563,7 @@ void PixelMain(
 	const float TimeScale = VelocitySample.w;
 
 	// Position.w holds the relative time of the particle.
-	float3 Position = LWCHackToFloat(MakeLWCVector3(LWCTile, PositionSample.xyz));
+	FLWCVector3 Position = MakeLWCVector3(LWCTile, PositionSample.xyz);
 	float RelativeTime = PositionSample.w;
 
 	for (int IterationIndex = 0; IterationIndex < NumIterations; ++IterationIndex)
@@ -577,7 +586,7 @@ void PixelMain(
 		// Evalute vector fields.
 		float3 FieldForce = 0;
 		float4 FieldVelocity = 0;
-		EvaluateVectorFields(FieldForce, FieldVelocity, Position.xyz, PerParticleVectorFieldScale);
+		EvaluateVectorFields(FieldForce, FieldVelocity, LWCHackToFloat(Position), PerParticleVectorFieldScale);		//LWC_TODO: This needs to be rebased
 
 		// Add in force from vector fields.
 		Force += FieldForce;
@@ -590,7 +599,7 @@ void PixelMain(
 		Force += ComputeDrag(Velocity.xyz, DragCoefficient);
 
 		// Compute force to a point gravity source.
-		Force += ComputeAttractionForce(Position.xyz);
+		Force += ComputeAttractionForce(LWCHackToFloat(Position));	//LWC_TODO: This needs to be rebased
 
 		// Compute the acceleration to apply to the particle this frame.
 		float3 Acceleration = Force * DeltaSeconds;
@@ -613,45 +622,43 @@ void PixelMain(
 		float CollisionRadius = min(Size.x,Size.y) * Simulation.CollisionRadiusScale + Simulation.CollisionRadiusBias;
 	#endif
 
-		float3 NewPosition, NewVelocity;
+		FLWCVector3 NewPosition;
+		float3 NewVelocity;
 
 	#if DEPTH_BUFFER_COLLISION
-	
 		// Compute the new position and velocity of the particle by colliding against
 		// the scene's depth buffer.
 		CollideWithDepthBuffer(
-			NewPosition.xyz,
-			NewVelocity.xyz,
+			NewPosition,
+			NewVelocity,
 			RelativeTime,
-			Position.xyz,
-			Velocity.xyz,
+			Position,
+			Velocity,
 			Acceleration,
 			CollisionRadius,
 			Resilience
-			);
+		);
 	#elif DISTANCE_FIELD_COLLISION
-
 		CollideWithDistanceField(
-			NewPosition.xyz,
-			NewVelocity.xyz,
+			NewPosition,
+			NewVelocity,
 			RelativeTime,
-			Position.xyz,
-			Velocity.xyz,
+			Position,
+			Velocity,
 			Acceleration,
 			CollisionRadius,
 			Resilience
 			);
-
 	#else 
 		// Integrate position and velocity forward.
 		float3 DeltaPosition = DeltaSeconds * (Velocity.xyz + 0.5f * Acceleration);
-		NewPosition = Position.xyz + DeltaPosition;
+		NewPosition = LWCAdd(Position, DeltaPosition);
 		NewVelocity = Velocity.xyz + Acceleration;
 	#endif
 
 		// Apply orbit.
 		const float3 OrbitVelocity = ComputeOrbitVelocity(RelativeTime, OrbitRandom);
-		NewPosition += OrbitVelocity * DeltaSeconds;
+		NewPosition = LWCAdd(NewPosition, OrbitVelocity * DeltaSeconds);
 
 		// Update values for new iteration.
 		Velocity = NewVelocity;
@@ -659,7 +666,7 @@ void PixelMain(
 	}
 
 	// Store the new position, time, and velocity for the particle.
-	OutPosition.xyz = (Position + PositionOffsetAndAttractorStrength.xyz) - LWCTile * UE_LWC_RENDER_TILE_SIZE;
+	OutPosition.xyz = LWCToFloat(LWCSubtract(Position, MakeLWCVector3(LWCTile, 0))) + PositionOffsetAndAttractorStrength.xyz;
 	OutPosition.w = RelativeTime;
 	OutVelocity.xyz = Velocity;
 	OutVelocity.w = TimeScale;
