Add SSS profile legacy encoding for Strata.

diff --git a/Engine/Shaders/Private/Strata/Strata.ush b/Engine/Shaders/Private/Strata/Strata.ush
--- a/Engine/Shaders/Private/Strata/Strata.ush
+++ b/Engine/Shaders/Private/Strata/Strata.ush
@@ -344,7 +344,8 @@ struct FStrataBSDF
 #define SINGLE_OPTLEGACYMODE_CLEARCOAT			1
 #define SINGLE_OPTLEGACYMODE_CLOTH				2
 #define SINGLE_OPTLEGACYMODE_SSSWRAP			3
-// 6 more remains: Cloth, Subsurface, ProfileIdSSS, Also how to deal with Aniso?
+#define SINGLE_OPTLEGACYMODE_SSSPROFILE			4
+// 5 more remains: Also how to deal with Aniso?
 
 // Layout 0 (Fast)
 #define BSDF_FASTENCODING_BIT_COUNT				(0u)
@@ -3753,6 +3754,7 @@ FStrataBSDF UnpackStrataBSDFIn(FStrataMaterialContainer StrataBuffer, inout FStr
 			OutBSDF.State = 0;
 			OutBSDF.LuminanceWeightV = 1.0f;
 			OutBSDF.TransmittanceAboveAlongN = 1.0f;
+			SLAB_F90(OutBSDF) = 1.0f;
 
 			if (OptimisedLegacyMode == SINGLE_OPTLEGACYMODE_CLEARCOAT)
 			{
@@ -3856,7 +3858,55 @@ FStrataBSDF UnpackStrataBSDFIn(FStrataMaterialContainer StrataBuffer, inout FStr
 				BSDF_SETHASFUZZ(OutBSDF, 0);
 				BSDF_SETHASTRANSABOVE(OutBSDF, 0);
 			}
+			else if (OptimisedLegacyMode == SINGLE_OPTLEGACYMODE_SSSPROFILE)
+			{
+			#if STRATA_DEFERRED_SHADING
+				uint Data0 = StrataHeader.SharedLocalBasesTypes_PackedHeader;
+				uint Data1 = StrataLoadUint1(StrataBuffer, StrataAddressing);
+				uint Data2 = StrataLoadUint1(StrataBuffer, StrataAddressing);
+
+				const uint PackedDiffuse20Bits	= Data1 & 0x000FFFFF;
+				const uint PackedF020Bits		= Data2 & 0x000FFFFF;
+				
+				const float Roughness				= UnpackR8(Data0 >> 24);
+				const float SSSProfileRadiusScale	= UnpackR8(Data1 >> 24);
+				const float SSSProfileId			= UnpackR8(Data2 >> 24);
+
+				SLAB_DIFFUSEALBEDO(OutBSDF)			= UnpackR7G7B6Gamma2(PackedDiffuse20Bits);
+				SLAB_F0(OutBSDF)					= UnpackR7G7B6Gamma2(PackedF020Bits);
+				SLAB_ROUGHNESS(OutBSDF)				= Roughness;
+				SLAB_SSSPROFILEID(OutBSDF)			= SSSProfileId;
+				SLAB_SSSPROFILERADIUSSCALE(OutBSDF)	= SSSProfileRadiusScale;
+
+				// When loading Strata BSDF which have SSS enable, BaseColor and (optionally) Specular value can be overriden 
+				// based on the SSS method used (screen space SSS / checherboarded screen space SSS / ...)
+				#if STRATA_SSS_MATERIAL_OVERRIDE
+				{
+					const bool bChecker = CheckerFromPixelPos(StrataAddressing.PixelCoords);
+					float3 DiffuseAlbedo = SLAB_DIFFUSEALBEDO(OutBSDF);
+					float SpecularFactor = 1.0f;
+					AdjustBaseColorAndSpecularColorForSubsurfaceProfileLighting(DiffuseAlbedo, SpecularFactor, bChecker);
+					SLAB_DIFFUSEALBEDO(OutBSDF) = DiffuseAlbedo;
+					SLAB_F0(OutBSDF)  *= SpecularFactor;
+					SLAB_F90(OutBSDF) *= SpecularFactor;
+				}
+				#endif
+			#endif
 
+				// Force a BSDF state to be a single layer on top.
+				BSDF_SETTYPE(OutBSDF, STRATA_BSDF_TYPE_SLAB);
+				BSDF_SETHASGREYWEIGHT_V(OutBSDF, 1);
+				BSDF_SETHASF90(OutBSDF, 0);
+				BSDF_SETHASANISOTROPY(OutBSDF, 0);
+				BSDF_SETISTOPLAYER(OutBSDF, 1);
+				BSDF_SETSSSTYPE(OutBSDF, SSS_TYPE_DIFFUSION_PROFILE);
+				BSDF_SETHASMFP(OutBSDF, 1);
+				BSDF_SETHASHAZINESS(OutBSDF, 0);
+				BSDF_SETISSIMPLEVOLUME(OutBSDF, 0);
+				BSDF_SETISTHIN(OutBSDF, 0);
+				BSDF_SETHASFUZZ(OutBSDF, 0);
+				BSDF_SETHASTRANSABOVE(OutBSDF, 0);
+			}
 			return OutBSDF;
 		}
 	}
diff --git a/Engine/Shaders/Private/Strata/StrataExport.ush b/Engine/Shaders/Private/Strata/StrataExport.ush
--- a/Engine/Shaders/Private/Strata/StrataExport.ush
+++ b/Engine/Shaders/Private/Strata/StrataExport.ush
@@ -312,6 +312,7 @@ void PackStrataOut(
 				const uint SingleIsThin				= 0x08;
 				const uint SingleIsSimpleVolume		= 0x10;
 				const uint SingleIsWrap				= 0x20;
+				const uint SingleIsSSSProfile		= 0x40;
 				// No need to add anisotropy since it is not supported by the simple mode.
 
 				uint OptimisedModeMask = 0;
@@ -341,6 +342,10 @@ void PackStrataOut(
 				{
 					OptimisedModeMask |= SingleIsWrap;
 				}
+				if (BSDF_GETSSSTYPE(BSDF) == SSS_TYPE_DIFFUSION_PROFILE)
+				{
+					OptimisedModeMask |= SingleIsSSSProfile;
+				}
 
 				if (((OptimisedModeMask & SingleHasClearCoat) == SingleHasClearCoat) || ((OptimisedModeMask & SingleHasClearCoat) == (SingleHasHaziness | SingleHasClearCoat)))
 				{
@@ -355,6 +360,11 @@ void PackStrataOut(
 					OptimisedLegacyMode = SINGLE_OPTLEGACYMODE_SSSWRAP;
 					SSSData.Header.Bytes = 0; // Force to skip SSSData writing
 				}
+				else if ((OptimisedModeMask & SingleIsSSSProfile) == SingleIsSSSProfile)
+				{
+					OptimisedLegacyMode = SINGLE_OPTLEGACYMODE_SSSPROFILE;
+					SSSData.Header.Bytes = 0; // Force to skip SSSData writing
+				}
 			#undef BSDF
 			}
 		}
@@ -591,6 +601,21 @@ void PackStrataOut(
 							// Ensure the rest of the BSDF is not stored with the regular path
 							BSDF.State = 0;
 						}
+						else if (OptimisedLegacyMode == SINGLE_OPTLEGACYMODE_SSSPROFILE)
+						{
+							const uint PackedDiffuse20Bits		= PackR7G7B6Gamma2(	SLAB_DIFFUSEALBEDO(BSDF)	, Dither);
+							const uint PackedF020Bits			= PackR7G7B6Gamma2(	SLAB_F0(BSDF)				, Dither);
+
+							// MRT0 
+							STRATA_STORE_UINT1(Out					| (PackR8(SLAB_ROUGHNESS(BSDF)) << 24));
+							// MRT1 20 PackedDiffuse and Profile radius scale
+							STRATA_STORE_UINT1(PackedDiffuse20Bits  | (PackR8(SLAB_SSSPROFILERADIUSSCALE(BSDF)) << 24));
+							// MRT2 24 PackedF0      and Profile ID
+							STRATA_STORE_UINT1(PackedF020Bits		| (PackR8(SLAB_SSSPROFILEID(BSDF)) << 24));
+
+							// Ensure the rest of the BSDF is not stored with the regular path
+							BSDF.State = 0;
+						}
 						// else if (OptimisedLegacyMode == SINGLE_OPTLEGACYMODE_xxx)
 						//
 						// STRATA_TODO
diff --git a/Engine/Shaders/Private/Strata/StrataMaterialClassification.usf b/Engine/Shaders/Private/Strata/StrataMaterialClassification.usf
--- a/Engine/Shaders/Private/Strata/StrataMaterialClassification.usf
+++ b/Engine/Shaders/Private/Strata/StrataMaterialClassification.usf
@@ -102,18 +102,43 @@ void TileMainCS(uint2 DispatchThreadId : SV_DispatchThreadID, uint LinearIndex :
 		bContainsOpaqueRoughRefraction = OpaqueRoughRefractionData.OpaqueRoughRefractionEnabled > 0.0f;
 	#endif
 
+		// Output/Patch SSS data for legacy encoding (this allows to save ALU & bandwidth during the base pass0
 		uint OptimisedLegacyMode = ((PackedHeader >> (HEADER_SINGLEENCODING_BIT_COUNT)) & HEADER_SINGLE_OPTLEGACYMODE_BIT_MASK);
-		if (bIsSingle && OptimisedLegacyMode == SINGLE_OPTLEGACYMODE_SSSWRAP)
+		const bool bIsLegacyWrap = OptimisedLegacyMode == SINGLE_OPTLEGACYMODE_SSSWRAP;
+		const bool bIsLegacySSSProfile = OptimisedLegacyMode == SINGLE_OPTLEGACYMODE_SSSPROFILE;
+		if (bIsSingle && (bIsLegacyWrap || bIsLegacySSSProfile))
 		{
 			bContainsScreenSpaceSubsurfaceScattering = true;
+			if (bIsLegacyWrap)
+			{
+				const uint PackedSSSWOpacity7bits = (PackedHeader >> (HEADER_SINGLEENCODING_BIT_COUNT + HEADER_SINGLE_OPTLEGACYMODE_BIT_COUNT)) & 0x7F;
+				FStrataSubsurfaceHeader SSSHeader = (FStrataSubsurfaceHeader)0;
+				StrataSubSurfaceHeaderSetSSSType(SSSHeader, SSS_TYPE_WRAP);
+				StrataSubSurfaceHeaderSetWrapOpacity(SSSHeader, UnpackR7(PackedSSSWOpacity7bits));
+				StrataStoreSubsurfaceHeader(MaterialTextureArrayUAV, FirstSliceStoringStrataSSSData, PixelCoord, SSSHeader.Bytes);
+			}
+			else if (bIsLegacySSSProfile)
+			{
+				const bool bDebugEnable = all(PixelCoord == View.CursorPosition);
 
-			const uint PackedSSSWOpacity7bits = (PackedHeader >> (HEADER_SINGLEENCODING_BIT_COUNT + HEADER_SINGLE_OPTLEGACYMODE_BIT_COUNT)) & 0x7F;
+				const uint Data1 = MaterialTextureArrayUAV[uint3(PixelCoord, 1)];
+				const uint Data2 = MaterialTextureArrayUAV[uint3(PixelCoord, 2)];
+				float  RadiusScale = UnpackR8(Data1 >> 24);
+				float  ProfileId   = UnpackR8(Data2 >> 24);
 
-			FStrataSubsurfaceHeader SSSHeader = (FStrataSubsurfaceHeader)0;
-			StrataSubSurfaceHeaderSetSSSType(SSSHeader, SSS_TYPE_WRAP);
-			StrataSubSurfaceHeaderSetWrapOpacity(SSSHeader, UnpackR7(PackedSSSWOpacity7bits));
+				const uint PackedDiffuse20Bits = (Data1 & 0xFFFFF);
+				const float3 BaseColor = UnpackR7G7B6Gamma2(PackedDiffuse20Bits);
 
-			StrataStoreSubsurfaceHeader(MaterialTextureArrayUAV, FirstSliceStoringStrataSSSData, PixelCoord, SSSHeader.Bytes);
+				FStrataSubsurfaceHeader SSSHeader = (FStrataSubsurfaceHeader)0;
+				StrataSubSurfaceHeaderSetSSSType(SSSHeader, SSS_TYPE_DIFFUSION_PROFILE);
+				StrataSubSurfaceHeaderSetProfile(SSSHeader, RadiusScale, StrataSubsurfaceProfileIdTo8bits(ProfileId));
+
+				FStrataSubsurfaceExtras SSSExtras = (FStrataSubsurfaceExtras)0;
+				StrataSubsurfaceExtrasSetBaseColor(SSSExtras, BaseColor);
+
+				StrataStoreSubsurfaceHeader(MaterialTextureArrayUAV, FirstSliceStoringStrataSSSData, PixelCoord, SSSHeader.Bytes);
+				StrataStoreSubsurfaceExtras(MaterialTextureArrayUAV, FirstSliceStoringStrataSSSData, PixelCoord, SSSExtras.Bytes);
+			}
 		}
 
 		// Fine grain test for clearing based on CMask data
diff --git a/Engine/Shaders/Private/Strata/StrataVisualize.usf b/Engine/Shaders/Private/Strata/StrataVisualize.usf
--- a/Engine/Shaders/Private/Strata/StrataVisualize.usf
+++ b/Engine/Shaders/Private/Strata/StrataVisualize.usf
@@ -264,17 +264,21 @@ void PrintBSDF(inout FShaderPrintContext Context, inout FStrataAddressing Strata
 
 			if (BSDF_GETSSSTYPE(BSDF) != SSS_TYPE_INVALID || BSDF_GETISSIMPLEVOLUME(BSDF) || BSDF_GETISTHIN(BSDF))
 			{
-
 				const FStrataSubsurfaceHeader SSSHeader = StrataLoadSubsurfaceHeader(Strata.MaterialTextureArray, Strata.FirstSliceStoringStrataSSSData, StrataAddressing.PixelCoords);
 				const uint SSSType = StrataSubSurfaceHeaderGetSSSType(SSSHeader);
-				const float Opacity = StrataSubSurfaceHeaderGetWrapOpacity(SSSHeader);
 				Print(Context, TEXT(" SSS Type    "), FontBSDFPropName);	PrintSSSType(Context, SSSType, FontYellow);	Newline(Context, RectMax);
-				Print(Context, TEXT(" SSS Opacity "), FontBSDFPropName);	Print(Context, Opacity, FontBSDFPropValu);	Newline(Context, RectMax);
-
+				if (BSDF_GETSSSTYPE(BSDF) == SSS_TYPE_WRAP)
+				{
+					const float Opacity = StrataSubSurfaceHeaderGetWrapOpacity(SSSHeader);
+					Print(Context, TEXT(" SSS Opacity "), FontBSDFPropName);	Print(Context, Opacity, FontBSDFPropValu);	Newline(Context, RectMax);
+				}
 				if (BSDF_GETSSSTYPE(BSDF) == SSS_TYPE_DIFFUSION_PROFILE && BSDF_GETISSIMPLEVOLUME(BSDF) == 0)
 				{
-					Print(Context, TEXT(" SSS ID.     "), FontBSDFPropName);	Print(Context, SLAB_SSSPROFILEID(BSDF), FontBSDFPropValu);			Newline(Context, RectMax);
-					Print(Context, TEXT(" SSS Radius  "), FontBSDFPropName);	Print(Context, SLAB_SSSPROFILERADIUSSCALE(BSDF), FontBSDFPropValu);	Newline(Context, RectMax);
+					const uint ProfileId = StrataSubSurfaceHeaderGetProfileId(SSSHeader);
+					const float RadiusScale = StrataSubSurfaceHeaderGetProfileRadiusScale(SSSHeader);
+
+					Print(Context, TEXT(" SSS ID.     "), FontBSDFPropName);	Print(Context, ProfileId, FontBSDFPropValu);	Newline(Context, RectMax);
+					Print(Context, TEXT(" SSS Radius  "), FontBSDFPropName);	Print(Context, RadiusScale, FontBSDFPropValu);	Newline(Context, RectMax);
 				}
 				else
 				{
@@ -456,7 +460,11 @@ void PrintPixel(uint2 InCoord, float3 InWorldPosition, float3 V)
 			else if (OptimisedLegacyMode ==SINGLE_OPTLEGACYMODE_SSSWRAP)
 			{
 				Print(Context, TEXT(" - Legecy SSS-Wrap"), FontLightRed);
-			}		
+			}
+			else if (OptimisedLegacyMode == SINGLE_OPTLEGACYMODE_SSSPROFILE)
+			{
+				Print(Context, TEXT(" - Legecy SSS-Profile"), FontLightRed);
+			}
 		#endif
 		}
 		Newline(Context, RectMax);
@@ -510,7 +518,7 @@ void PrintPixel(uint2 InCoord, float3 InWorldPosition, float3 V)
 		const bool bHasSSSData = StrataSubSurfaceHeaderGetIsValid(SSSHeader);
 
 		const uint TopLayerDataBytes = 4;	// bytes 
-		const uint SSSDataBytes = 8; // bIsSingleMaterial&& ? 0 : 8;	    // bytes  TODO accessor to OPT mode
+		const uint SSSDataBytes = 8;
 
 		const uint HeaderSize =  FootPrint_PostHeader - FootPrint_Start;
 		const uint BSDFsSize  =  FootPrint_PostBSDFs - FootPrint_PostHeader;
