Fix Strata crash in dx12 due to mismatching VS/PS interpolates.

diff --git a/Engine/Shaders/Private/AmbientCubemapComposite.usf b/Engine/Shaders/Private/AmbientCubemapComposite.usf
--- a/Engine/Shaders/Private/AmbientCubemapComposite.usf
+++ b/Engine/Shaders/Private/AmbientCubemapComposite.usf
@@ -53,7 +53,6 @@ float3 GetEnvSpecularLighting(
 #include "Strata/StrataEnvironmentLighting.ush"
 
 void MainPS(
-in noperspective float4 UVAndScreenPos : TEXCOORD0, 
 	float4 SvPosition : SV_POSITION, 
 	out float4 OutColor : SV_Target0
 	#if STRATA_OPAQUE_ROUGH_REFRACTION_ENABLED
@@ -63,8 +62,8 @@ in noperspective float4 UVAndScreenPos : TEXCOORD0,
 	)
 {
 	const int2 PixelPos = int2(SvPosition.xy);
-	const float2 BufferUV = UVAndScreenPos.xy;
-	const float2 ScreenPosition = UVAndScreenPos.zw;
+	const float2 BufferUV = SvPositionToBufferUV(SvPosition);
+	const float2 ScreenPosition = SvPositionToScreenPosition(SvPosition).xy;
 
 	OutColor = 0;
 #if STRATA_OPAQUE_ROUGH_REFRACTION_ENABLED
