Fix invalid strata case: Thin SSS is only supported when diffusion is not enabled.

diff --git a/Engine/Shaders/Private/Strata/Strata.ush b/Engine/Shaders/Private/Strata/Strata.ush
--- a/Engine/Shaders/Private/Strata/Strata.ush
+++ b/Engine/Shaders/Private/Strata/Strata.ush
@@ -1332,11 +1332,10 @@ FStrataData GetStrataSlabBSDF(
 {
 	FStrataData StrataData = GetInitialisedStrataData();
 
-	// STRATA_TODO apply SSS as a function of pixel footprint w.r.t. MFP
 	SSSMFPScale = saturate(SSSMFPScale);
 	const bool bSSSProfilePlugged = SSSProfileID > 0.0f;
 	const bool bMFPPlugged = any(SSSMFP > 0);
-	bool bIsThin = bMFPPlugged && Thickness < STRATA_LAYER_ISTHIN_THICKNESS_THRESHOLD_CM;
+	bool bIsThin = bMFPPlugged && Thickness < STRATA_LAYER_ISTHIN_THICKNESS_THRESHOLD_CM && !UseSSSDiffusion; // Thin is used for flagging 'legacy foliage' case
 	bool bHasSSSWrap = bMFPPlugged && !UseSSSDiffusion;
 	bool bHasSSSProfile = !bIsThin && bSSSProfilePlugged && SSSMFPScale > 0.f;	// When no profile is provide, we use Burley. STRATA_TODO we should clarify as burley is also an option of the profile.
 	bool bHasSSS = bMFPPlugged || bHasSSSProfile || bHasSSSWrap;
