Add support for blending DBuffer as a post-pass.

The post base-pass dbuffer blending only works for simple and single materials. Complex material are still apply during the base pass. This pass is disabled by default, and will be enabled only per-platform. Following needs to happen making unpacking/packing modulare and reusable.

On lower-platform this allows to reduce significatly the cost of base pass (BasePass:-1.2ms DBufferPass+0.4ms = -0.8ms net), as the compiler can better prune code, and allows to increase occupancy.

diff --git a/Engine/Shaders/Private/BasePassPixelShader.usf b/Engine/Shaders/Private/BasePassPixelShader.usf
--- a/Engine/Shaders/Private/BasePassPixelShader.usf
+++ b/Engine/Shaders/Private/BasePassPixelShader.usf
@@ -44,6 +44,12 @@
 #define STRATA_ENABLED 0
 #endif
 
+// Is Strata is enabled, only use DBuffer pass after the base pass for simple and single material
+#if STRATA_ENABLED && STRATA_USE_DBUFFER_PASS && USE_DBUFFER && (STRATA_FASTPATH || STRATA_SINGLEPATH)
+#undef USE_DBUFFER
+#define USE_DBUFFER 0
+#endif
+
 // For the base pass STRATA_INLINE_SHADING is defined in BasePassRendering.h
 #include "SHCommon.ush"
 #include "/Engine/Generated/Material.ush"
@@ -975,33 +981,33 @@ void FPixelShaderInOut_MainPS(
 	float DBufferOpacity = 1.0f;
 #if USE_DBUFFER && MATERIALDECALRESPONSEMASK && !MATERIALBLENDING_ANY_TRANSLUCENT && !MATERIAL_SHADINGMODEL_SINGLELAYERWATER
 	// apply decals from the DBuffer
-#if PC_D3D
+	#if PC_D3D
 	//Temporary workaround to avoid crashes on AMD, revert back to BRANCH
 	FLATTEN
-#else
+	#else
 	BRANCH
-#endif
-		if ((GetPrimitiveData(MaterialParameters).Flags & PRIMITIVE_SCENE_DATA_FLAG_DECAL_RECEIVER) != 0 && View.ShowDecalsMask > 0)
-		{
-			uint ValidDBufferTargetMask = GetDBufferTargetMask(uint2(In.SvPosition.xy)) & MATERIALDECALRESPONSEMASK;
+	#endif
+	if ((GetPrimitiveData(MaterialParameters).Flags & PRIMITIVE_SCENE_DATA_FLAG_DECAL_RECEIVER) != 0 && View.ShowDecalsMask > 0)
+	{
+		uint ValidDBufferTargetMask = GetDBufferTargetMask(uint2(In.SvPosition.xy)) & MATERIALDECALRESPONSEMASK;
 
-			BRANCH
-			if (ValidDBufferTargetMask)
-			{
-				float2 BufferUV = SvPositionToBufferUV(In.SvPosition);
-			#if STRATA_ENABLED 
-			   #if STRATA_INLINE_SHADING && !STRATA_OPTIMIZED_UNLIT
-				const FStrataDBuffer StrataBufferData = GetDBufferData(BufferUV, ValidDBufferTargetMask);
-				ApplyDBufferData(StrataBufferData, StrataPixelHeader, StrataData);
-				DBufferOpacity = StrataBufferData.Coverage;
-			   #endif
-			#else
-				FDBufferData DBufferData = GetDBufferData(BufferUV, ValidDBufferTargetMask);
-				ApplyDBufferData(DBufferData, MaterialParameters.WorldNormal, SubsurfaceColor, Roughness, BaseColor, Metallic, Specular);
-				DBufferOpacity = (DBufferData.ColorOpacity + DBufferData.NormalOpacity + DBufferData.RoughnessOpacity) * (1.0f / 3.0f);
-			#endif // STRATA_ENABLED
-			}
+		BRANCH
+		if (ValidDBufferTargetMask)
+		{
+			float2 BufferUV = SvPositionToBufferUV(In.SvPosition);
+		#if STRATA_ENABLED 
+			#if STRATA_INLINE_SHADING && !STRATA_OPTIMIZED_UNLIT
+			const FStrataDBuffer StrataBufferData = GetDBufferData(BufferUV, ValidDBufferTargetMask);
+			ApplyDBufferData(StrataBufferData, StrataPixelHeader, StrataData);
+			DBufferOpacity = StrataBufferData.Coverage;
+			#endif
+		#else
+			FDBufferData DBufferData = GetDBufferData(BufferUV, ValidDBufferTargetMask);
+			ApplyDBufferData(DBufferData, MaterialParameters.WorldNormal, SubsurfaceColor, Roughness, BaseColor, Metallic, Specular);
+			DBufferOpacity = (DBufferData.ColorOpacity + DBufferData.NormalOpacity + DBufferData.RoughnessOpacity) * (1.0f / 3.0f);
+		#endif // STRATA_ENABLED
 		}
+	}
 #endif
 
 	const float BaseMaterialCoverageOverWater = Opacity;
diff --git a/Engine/Shaders/Private/Strata/Strata.ush b/Engine/Shaders/Private/Strata/Strata.ush
--- a/Engine/Shaders/Private/Strata/Strata.ush
+++ b/Engine/Shaders/Private/Strata/Strata.ush
@@ -1741,16 +1741,28 @@ void FinalizeWrites(inout FRWStrataMaterialContainer StrataBuffer, inout FStrata
 ///////////////////////////////////////////////////////////////////////////////
 // Material buffer read
 
+#ifndef STRATA_MATERIALCONTAINER_IS_WRITABLE 
+#define STRATA_MATERIALCONTAINER_IS_WRITABLE 0
+#endif
+
 #if STRATA_MATERIALCONTAINER_IS_VIEWRESOURCE
 
+#if STRATA_MATERIALCONTAINER_IS_WRITABLE
+#define FStrataMaterialContainer RWTexture2DArray<uint>
+#else
 #define FStrataMaterialContainer Texture2DArray<uint>
+#endif
 
 uint StrataLoadUint1(FStrataMaterialContainer StrataBuffer, inout FStrataAddressing StrataAddressing)
 {
 	uint Data = 0;
   #if !COMPILER_SWITCH
+	#if STRATA_MATERIALCONTAINER_IS_WRITABLE
+	Data = StrataBuffer[uint3(StrataAddressing.PixelCoords, StrataAddressing.CurrentIndex)];
+	#else
 	const uint MipLevel = 0;
 	Data = StrataBuffer.mips[MipLevel][uint3(StrataAddressing.PixelCoords, StrataAddressing.CurrentIndex)];
+	#endif
 	StrataAddressing.CurrentIndex++;
 	StrataAddressing.ReadBytes += 4;
   #endif
@@ -1761,17 +1773,29 @@ uint StrataLoadUint1AtIndexOffset(in FStrataMaterialContainer StrataBuffer, in F
 {
 	uint Data = 0;
   #if !COMPILER_SWITCH
+	#if STRATA_MATERIALCONTAINER_IS_WRITABLE
+	Data = StrataBuffer[uint3(StrataAddressing.PixelCoords, IndexOffset)];
+	#else
 	const uint MipLevel = 0;
 	Data = StrataBuffer.mips[MipLevel][uint3(StrataAddressing.PixelCoords, IndexOffset)];
+	#endif
   #endif
 	return Data;
 }
 
+#if STRATA_MATERIALCONTAINER_IS_WRITABLE
+#define FStrataTopLayerDataContainer RWTexture2D<uint>
+#else
 #define FStrataTopLayerDataContainer Texture2D<uint>
+#endif
 
 uint StrataLoadTopLayerData(FStrataTopLayerDataContainer InTopLayerContainer, inout FStrataAddressing StrataAddressing)
 {
+	#if STRATA_MATERIALCONTAINER_IS_WRITABLE
+	return InTopLayerContainer[StrataAddressing.PixelCoords];
+	#else
 	return InTopLayerContainer.Load(uint3(StrataAddressing.PixelCoords, 0));
+	#endif
 }
 
 #else // STRATA_MATERIALCONTAINER_IS_VIEWRESOURCE
diff --git a/Engine/Shaders/Private/Strata/StrataDbuffer.usf b/Engine/Shaders/Private/Strata/StrataDbuffer.usf
new file mode 100644
--- /dev/null
+++ b/Engine/Shaders/Private/Strata/StrataDbuffer.usf
@@ -0,0 +1,304 @@
+// Copyright Epic Games, Inc. All Rights Reserved.
+
+#include "/Engine/Private/Common.ush"
+
+#if PERMUTATION_TILETYPE == 0
+#define STRATA_FASTPATH 1
+#elif PERMUTATION_TILETYPE == 1
+#define STRATA_SINGLEPATH 1
+#elif PERMUTATION_TILETYPE == 2
+  // COMPLEX PATH
+#else
+#error Strata tile type non-implemented
+#endif
+
+#define OpaqueBasePass
+
+#define STRATA_MATERIALCONTAINER_IS_WRITABLE 1
+#define STRATA_INLINE_SHADING 0
+#define STRATA_SSS_MATERIAL_OVERRIDE 0
+#include "/Engine/Private/Strata/Strata.ush"
+
+#define MATERIALDECALRESPONSEMASK 0x7
+#include "../DBufferDecalShared.ush"
+#include "StrataTile.ush"
+
+////////////////////////////////////////////////////////////////////////////////////////////////////////////
+// Post Dbuffer application
+#if SHADER_DBUFFER
+Texture2D<float4> DBufferATexture;
+Texture2D<float4> DBufferBTexture;
+Texture2D<float4> DBufferCTexture;
+Texture2D<uint> DBufferRenderMask;
+
+SamplerState DBufferATextureSampler;
+SamplerState DBufferBTextureSampler;
+SamplerState DBufferCTextureSampler;
+
+int2 ViewResolution;
+uint MaxBytesPerPixel;
+uint FirstSliceStoringStrataSSSData;
+RWTexture2D<uint> TopLayerTexture;
+RWTexture2DArray<uint> MaterialTextureArrayUAV;
+
+Buffer<uint> TileListBuffer;
+Buffer<uint> TileIndirectBuffer;
+
+// TODO
+// * Respect material decal response
+// * Add support Eye, Hair, Water, ... ?
+// * Refactor encoding to share code between strata export
+
+uint StrataGetDBufferTargetMask(uint2 PixelPos)
+{
+#if PLATFORM_SUPPORTS_RENDERTARGET_WRITE_MASK
+	return DecodeRTWriteMask(PixelPos, DBufferRenderMask, 3);
+#elif PLATFORM_SUPPORTS_PER_PIXEL_DBUFFER_MASK
+	uint Mask = DBufferRenderMask.Load(uint3(PixelPos, 0));
+	return Mask > 0 ? 0x07 : 0x00;
+#else
+	return 0x07;
+#endif
+}
+
+// @param BufferUV - UV space in the DBuffer textures
+FStrataDBuffer StrataGetDBufferData(float2 BufferUV, uint RTMaskBit)
+{
+	// Setup default values, which mean that no decals are present in DBuffer
+	float4 DBufferA = float4(0.0, 0.0, 0.0, 1.0);
+	float4 DBufferB = float4(128.0f / 255.0f, 128.f / 255.5f, 128.f / 255.5f, 1.0);
+	float4 DBufferC = float4(0.0, 0.0, 0.0, 1.0);
+
+	BRANCH
+	if (RTMaskBit & 0x1) { DBufferA = Texture2DSampleLevel(DBufferATexture, DBufferATextureSampler, BufferUV, 0); }
+	BRANCH
+	if (RTMaskBit & 0x2) { DBufferB = Texture2DSampleLevel(DBufferBTexture, DBufferBTextureSampler, BufferUV, 0); }
+	BRANCH
+	if (RTMaskBit & 0x4) { DBufferC = Texture2DSampleLevel(DBufferCTexture, DBufferCTextureSampler, BufferUV, 0); }
+	return StrataUnpackDBuffer(DBufferA, DBufferB, DBufferC);
+}
+
+[numthreads(STRATA_TILE_SIZE, STRATA_TILE_SIZE, 1)]
+void MainCS(
+	uint2 DispatchThreadId : SV_DispatchThreadID, 
+	uint LinearIndex : SV_GroupIndex, 
+	uint3 GroupId : SV_GroupID,
+	uint2 GroupThreadId : SV_GroupThreadID)
+{
+	if (View.GeneralPurposeTweak == 0)
+		return;
+
+	const uint2 TileCoord = StrataUnpackTile(TileListBuffer[GroupId.x]);
+	const uint2 PixelCoord = TileCoord * STRATA_TILE_SIZE + GroupThreadId;
+	const bool bIsValid = all(PixelCoord < uint2(View.ViewSizeAndInvSize.xy));
+
+	if (!bIsValid)
+	{
+		return;
+	}
+
+	uint ValidDBufferTargetMask = StrataGetDBufferTargetMask(PixelCoord);
+	BRANCH
+	if (!ValidDBufferTargetMask)
+	{
+		return;
+	}
+
+	const float2 BufferUV = SvPositionToBufferUV(float4(PixelCoord + 0.5f, 0, 0));
+	const FStrataDBuffer DBufferData = StrataGetDBufferData(BufferUV, ValidDBufferTargetMask);
+
+	// Early out if DBuffer has no coverage
+	BRANCH
+	if (DBufferData.Coverage <= 0)
+	{
+		return;
+	}
+
+	FStrataAddressing StrataAddressing = GetStrataPixelDataByteOffset(PixelCoord, uint2(View.BufferSizeAndInvSize.xy), MaxBytesPerPixel);
+	FStrataPixelHeader StrataPixelHeader = UnpackStrataHeaderIn(MaterialTextureArrayUAV, StrataAddressing, TopLayerTexture);
+
+	BRANCH
+	if (StrataPixelHeader.BSDFCount <= 0)
+	{
+		return;
+	}
+
+	BRANCH
+	if (IsSimpleMaterial(StrataPixelHeader) || IsSingleMaterial(StrataPixelHeader))
+	{
+		// 1. Load BSDF data
+		const float Dither = InterleavedGradientNoise(PixelCoord + 0.5f, View.StateFrameIndexMod8);
+		FStrataBSDF BSDF = UnpackStrataBSDFIn(MaterialTextureArrayUAV, StrataAddressing, StrataPixelHeader);
+
+		// 2. Load DBuffer data and apply them
+		float3 F0 = 0;
+		float3 DiffuseAlbedo = 0;
+		float3 FuzzColor = 0;
+		float Roughness = 0;
+		{
+			Roughness = StrataGetBSDFRoughness(BSDF);
+
+			float3 BaseColor = 0;
+			float Metallic = 0;
+			float Specular = 0;
+			ConvertToMetalness(SLAB_DIFFUSEALBEDO(BSDF), SLAB_F0(BSDF), BaseColor, Specular, Metallic);
+
+			Metallic  = Metallic  * DBufferData.OneMinusCoverage_Roughness + DBufferData.Metallic;
+			Specular  = Specular  * DBufferData.OneMinusCoverage_Roughness + DBufferData.Specular;
+			Roughness = Roughness * DBufferData.OneMinusCoverage_Roughness + DBufferData.Roughness;
+			BaseColor = BaseColor * DBufferData.OneMinusCoverage_BaseColor + DBufferData.BaseColor;
+
+			#if !STRATA_FASTPATH
+			if (BSDF_GETHASFUZZ(BSDF))
+			{
+				FuzzColor = SLAB_FUZZ_COLOR(BSDF) * DBufferData.OneMinusCoverage_BaseColor + DBufferData.BaseColor;
+			}
+			#endif
+
+			ConvertFromMetalness(BaseColor,	Specular, Metallic, DiffuseAlbedo, F0);
+		}
+
+		// 3. Top layer data
+		FStrataTopLayerData TopLayerData = StrataUnpackTopLayerData(StrataPixelHeader.PackedTopLayerData);
+		TopLayerData.WorldNormal = TopLayerData.WorldNormal * DBufferData.OneMinusCoverage_WorldNormal + DBufferData.WorldNormal;
+		TopLayerData.Roughness = Roughness;
+		TopLayerTexture[PixelCoord] = StrataPackTopLayerData(TopLayerData);
+
+		// 4. Rencode data
+		// STRATA_TODO: For now, doing manual reencoding by duplicating logic from Strata export. Need to refactor Strata export to have small reusable packing functions
+		
+		// Roughness / DiffuseColor / F0
+		if (IsSimpleMaterial(StrataPixelHeader))
+		{
+			const uint PackedDiffuse20Bits = PackR7G7B6Gamma2(DiffuseAlbedo, Dither);
+			const uint PackedDiffuse12Bits = PackedDiffuse20Bits & 0xFFF;
+			const uint PackedDiffuse8Bits  = (PackedDiffuse20Bits >> 12) & 0xFF;
+			const uint PackedF020Bits      = PackR7G7B6Gamma2(F0, Dither);
+			const uint PackedRoughness8bits= PackR8(Roughness);
+
+			// Data0 (Header_State|Header_AO|Roughness|Diffuse8bits)
+			MaterialTextureArrayUAV[uint3(PixelCoord, 0)] = (StrataPixelHeader.SharedLocalBasesTypes_PackedHeader & 0x0000FFFF) | (PackedRoughness8bits << 16) | (PackedDiffuse8Bits << 24);
+			MaterialTextureArrayUAV[uint3(PixelCoord, 1)] = PackedF020Bits | (PackedDiffuse12Bits << 20);
+		}
+		#if !STRATA_FASTPATH
+		else if (IsSingleMaterial(StrataPixelHeader))
+		{
+			const uint OptimisedLegacyMode = (StrataPixelHeader.SharedLocalBasesTypes_PackedHeader >> (HEADER_SINGLEENCODING_BIT_COUNT)) & HEADER_SINGLE_OPTLEGACYMODE_BIT_MASK;
+			if (OptimisedLegacyMode == SINGLE_OPTLEGACYMODE_NONE)
+			{
+				MaterialTextureArrayUAV[uint3(PixelCoord, 2)] = (MaterialTextureArrayUAV[uint3(PixelCoord, 2)] & 0xFFFF00FF) | (PackR8(Roughness) << 8);
+
+				const uint PackedDiffuse20Bits	= PackR7G7B6Gamma2(DiffuseAlbedo, Dither);
+				const uint PackedF020Bits		= PackR7G7B6Gamma2(F0, Dither);
+				const uint PackedData32Bits		= (PackedDiffuse20Bits << 12) | (PackedF020Bits & 0xFFF);
+				const uint PackedData8Bits		= (PackedF020Bits >> 12) & 0xFF;
+
+				uint RWOffset = 1;
+				MaterialTextureArrayUAV[uint3(PixelCoord, RWOffset++)] = PackedData32Bits; 
+				MaterialTextureArrayUAV[uint3(PixelCoord, RWOffset++)] = PackedData8Bits | PackRGBA8(float4(0.0f, Roughness, (SLAB_ANISOTROPY(BSDF) + 1.f) * 0.5f, (SLAB_SSSPHASEANISOTROPY(BSDF) + 1.f) * 0.5f));
+
+				if (BSDF_GETHASF90(BSDF) || BSDF_GETHASHAZINESS(BSDF))
+				{
+					++RWOffset;
+				}
+
+				const bool bIsSimpleVolume = BSDF_GETISSIMPLEVOLUME(BSDF);
+				const bool bHasSSS = BSDF_GETSSSTYPE(BSDF) != SSS_TYPE_INVALID;
+				if (bHasSSS || bIsSimpleVolume)
+				{
+					// Disable SSS when the coverage reach one to preserve the decal color
+					SLAB_SSSMFP(BSDF) = lerp(0.f, SLAB_SSSMFP(BSDF), DBufferData.OneMinusCoverage_BaseColor);
+					if (all(SLAB_SSSMFP(BSDF) == 0.0f))
+					{
+						BSDF_SETSSSTYPE(BSDF, SSS_TYPE_INVALID);
+					}
+
+					if (BSDF_GETSSSTYPE(BSDF) == SSS_TYPE_DIFFUSION_PROFILE && !bIsSimpleVolume)
+					{
+						// Nothing to do
+					}
+					else
+					{
+						// Path used for bottom most layer SSS, simple volume and two sided lighting.
+						MaterialTextureArrayUAV[uint3(PixelCoord, RWOffset)] = PackR11G11B10F(SLAB_SSSMFP(BSDF));
+					}
+
+					// Only write SSS data if needed
+					if (BSDF_GETSSSTYPE(BSDF) == SSS_TYPE_DIFFUSION)
+					{
+						FStrataSubsurfaceExtras SSSDataExtras;
+						StrataSubsurfaceExtrasSetBaseColor(SSSDataExtras, DiffuseAlbedo);
+						StrataStoreSubsurfaceExtras(MaterialTextureArrayUAV, FirstSliceStoringStrataSSSData, PixelCoord, SSSDataExtras.Bytes);
+					}
+
+					++RWOffset;
+				}
+
+				if (BSDF_GETHASFUZZ(BSDF))
+				{
+					MaterialTextureArrayUAV[uint3(PixelCoord, RWOffset)] = PackColorLinearToGamma2AlphaLinear(float4(FuzzColor, SLAB_FUZZ_AMOUNT(BSDF)));
+					RWOffset++;
+				}
+			}
+			else if (OptimisedLegacyMode == SINGLE_OPTLEGACYMODE_CLEARCOAT)
+			{
+				const uint PackedDiffuse20Bits = PackR7G7B6Gamma2(DiffuseAlbedo, Dither);
+				const uint PackedF024Bits = PackRGBA8(float4(F0, 0.0f));
+
+				// Roughnes
+				MaterialTextureArrayUAV[uint3(PixelCoord, 0)] = (MaterialTextureArrayUAV[uint3(PixelCoord, 0)] & 0x00FFFFFF) | (PackR8(Roughness) << 24);
+
+				// Diffuse Albedo & F0
+				MaterialTextureArrayUAV[uint3(PixelCoord, 1)] = (MaterialTextureArrayUAV[uint3(PixelCoord, 1)] & 0xFFF00000) | PackedDiffuse20Bits;
+				MaterialTextureArrayUAV[uint3(PixelCoord, 2)] = (MaterialTextureArrayUAV[uint3(PixelCoord, 2)] & 0xFF000000) | PackedF024Bits;
+			}
+			else if (OptimisedLegacyMode == SINGLE_OPTLEGACYMODE_CLOTH)
+			{
+				// Roughness is only stored in the TopLayer texture.
+
+				// Diffuse Albedo & F0
+				const uint PackedDiffuse20Bits = PackR7G7B6Gamma2(DiffuseAlbedo, Dither);
+				const uint PackedF020Bits = PackR7G7B6Gamma2(F0, Dither);
+				const uint PackedFuzzColor20bits = PackR7G7B6Gamma2(FuzzColor, Dither);
+
+				MaterialTextureArrayUAV[uint3(PixelCoord, 1)] = (PackedDiffuse20Bits | ((PackedFuzzColor20bits & 0xFFC00) << 10));
+				MaterialTextureArrayUAV[uint3(PixelCoord, 2)] = (PackedF020Bits | ((PackedFuzzColor20bits & 0x3FF) << 20));
+			}
+			else if (OptimisedLegacyMode == SINGLE_OPTLEGACYMODE_SSSWRAP)
+			{
+				// Roughness is only stored in the TopLayer texture.
+				const float3 MFP = lerp(0.f, SLAB_SSSMFP(BSDF), DBufferData.OneMinusCoverage_BaseColor);
+
+				// DiffuseAlbedo & F0
+				const uint PackedDiffuse20Bits = PackR7G7B6Gamma2(DiffuseAlbedo, Dither);
+				const uint PackedF020Bits      = PackR7G7B6Gamma2(F0, Dither);
+				const uint PackedSSSMFP30bits  = PackR10G10B10F(MFP);
+
+				const uint PackedSSSMFP12BitsA = (PackedSSSMFP30bits) & 0xFFF; // 12 lower bits
+				const uint PackedSSSMFP12BitsB = (PackedSSSMFP30bits >> 12) & 0xFFF; // next 12 lower bits
+				const uint PackedSSSMFP6BitsC  = (PackedSSSMFP30bits >> 24) & 0x03F; // 6 higher bits
+
+				const uint BitShift = (7 + HEADER_SINGLEENCODING_BIT_COUNT + HEADER_SINGLE_OPTLEGACYMODE_BIT_COUNT);
+
+				MaterialTextureArrayUAV[uint3(PixelCoord, 0)] = (MaterialTextureArrayUAV[uint3(PixelCoord, 0)] & 0x03FFFFFF) | (PackedSSSMFP6BitsC<< BitShift);
+				MaterialTextureArrayUAV[uint3(PixelCoord, 1)] = PackedDiffuse20Bits | (PackedSSSMFP12BitsA << 20);
+				MaterialTextureArrayUAV[uint3(PixelCoord, 2)] = PackedF020Bits | (PackedSSSMFP12BitsB << 20);
+			}
+			else if (OptimisedLegacyMode == SINGLE_OPTLEGACYMODE_SSSPROFILE)
+			{
+				const uint PackedDiffuse20Bits = PackR7G7B6Gamma2(DiffuseAlbedo, Dither);
+				const uint PackedF020Bits = PackR7G7B6Gamma2(F0, Dither);
+
+				MaterialTextureArrayUAV[uint3(PixelCoord, 0)] = (MaterialTextureArrayUAV[uint3(PixelCoord, 0)] & 0x00FFFFFF) | (PackR8(Roughness) << 24);
+				MaterialTextureArrayUAV[uint3(PixelCoord, 1)] = (MaterialTextureArrayUAV[uint3(PixelCoord, 1)] & 0xFF000000) | PackedDiffuse20Bits;
+				MaterialTextureArrayUAV[uint3(PixelCoord, 2)] = (MaterialTextureArrayUAV[uint3(PixelCoord, 2)] & 0xFF000000) | PackedF020Bits;
+				
+				FStrataSubsurfaceExtras SSSDataExtras;
+				StrataSubsurfaceExtrasSetBaseColor(SSSDataExtras, DiffuseAlbedo);
+				StrataStoreSubsurfaceExtras(MaterialTextureArrayUAV, FirstSliceStoringStrataSSSData, PixelCoord, SSSDataExtras.Bytes);
+			}
+		}
+		#endif // !STRATA_FASTPATH
+	}
+}
+#endif // SHADER_DBUFFER
\ No newline at end of file
diff --git a/Engine/Source/Runtime/Engine/Private/ShaderCompiler/ShaderCompiler.cpp b/Engine/Source/Runtime/Engine/Private/ShaderCompiler/ShaderCompiler.cpp
--- a/Engine/Source/Runtime/Engine/Private/ShaderCompiler/ShaderCompiler.cpp
+++ b/Engine/Source/Runtime/Engine/Private/ShaderCompiler/ShaderCompiler.cpp
@@ -6083,6 +6083,10 @@ void GlobalBeginCompileShader(
 			static FShaderPlatformCachedIniValue<int32> CVarStrataTileCoord8Bits(TEXT("r.Strata.TileCoord8bits"));
 			const uint32 StrataTileCoord8Bits = CVarStrataTileCoord8Bits.Get(Target.GetPlatform());
 			Input.Environment.SetDefine(TEXT("USE_8BIT_TILE_COORD"), StrataTileCoord8Bits > 0 ? 1 : 0);
+
+			static FShaderPlatformCachedIniValue<int32> CVarStrataDBufferPass(TEXT("r.Strata.DBufferPass"));
+			const uint32 StrataDBufferPass = CVarStrataDBufferPass.Get(Target.GetPlatform());
+			Input.Environment.SetDefine(TEXT("STRATA_USE_DBUFFER_PASS"), StrataDBufferPass > 0 ? 1 : 0);
 		}
 
 		static IConsoleVariable* CVarBackCompatibility = IConsoleManager::Get().FindConsoleVariable(TEXT("r.StrataBackCompatibility"));
diff --git a/Engine/Source/Runtime/RenderCore/Private/Shader.cpp b/Engine/Source/Runtime/RenderCore/Private/Shader.cpp
--- a/Engine/Source/Runtime/RenderCore/Private/Shader.cpp
+++ b/Engine/Source/Runtime/RenderCore/Private/Shader.cpp
@@ -1838,6 +1838,16 @@ void ShaderMapAppendKeyString(EShaderPlatform Platform, FString& KeyString)
 			KeyString += FString::Printf(TEXT("_BUDGET%u"), BytesPerPixel);
 		}
 
+		if (bStrataEnabled)
+		{
+			static FShaderPlatformCachedIniValue<int32> CVarDBufferPass(TEXT("r.Strata.DBufferPass"));
+			const int32 StrataDBufferPass = CVarDBufferPass.Get(Platform);
+			if (StrataDBufferPass>0)
+			{
+				KeyString += FString::Printf(TEXT("_DBUFFERPASS"));
+			}
+		}
+
 		static const auto CVarBackCompatibility = IConsoleManager::Get().FindTConsoleVariableDataInt(TEXT("r.StrataBackCompatibility"));
 		if (bStrataEnabled && CVarBackCompatibility && CVarBackCompatibility->GetValueOnAnyThread()>0)
 		{
diff --git a/Engine/Source/Runtime/Renderer/Private/DeferredShadingRenderer.cpp b/Engine/Source/Runtime/Renderer/Private/DeferredShadingRenderer.cpp
--- a/Engine/Source/Runtime/Renderer/Private/DeferredShadingRenderer.cpp
+++ b/Engine/Source/Runtime/Renderer/Private/DeferredShadingRenderer.cpp
@@ -3056,6 +3056,7 @@ void FDeferredShadingSceneRenderer::Render(FRDGBuilder& GraphBuilder)
 	if (Strata::IsStrataEnabled())
 	{
 		Strata::AddStrataMaterialClassificationPass(GraphBuilder, SceneTextures, Views);
+		Strata::AddStrataDBufferPass(GraphBuilder, SceneTextures, DBufferTextures, Views);
 	}
 
 #if RHI_RAYTRACING
diff --git a/Engine/Source/Runtime/Renderer/Private/Strata/Strata.cpp b/Engine/Source/Runtime/Renderer/Private/Strata/Strata.cpp
--- a/Engine/Source/Runtime/Renderer/Private/Strata/Strata.cpp
+++ b/Engine/Source/Runtime/Renderer/Private/Strata/Strata.cpp
@@ -86,6 +86,12 @@ static TAutoConsoleVariable<int32> CVarStrataAsyncClassification(
 	TEXT("Run Strata material classification in async (with shadow)."),
 	ECVF_RenderThreadSafe);
 
+static TAutoConsoleVariable<int32> CVarStrataDBufferPass(
+	TEXT("r.Strata.DBufferPass"),
+	0,
+	TEXT("Apply DBuffer after the base-pass as a separate pass."),
+	ECVF_RenderThreadSafe);
+
 IMPLEMENT_GLOBAL_SHADER_PARAMETER_STRUCT(FStrataGlobalUniformParameters, "Strata");
 
 extern bool IsStrataAdvancedVisualizationShadersEnabled();
@@ -127,6 +133,11 @@ bool IsStrataEnabled()
 	return CVarStrata.GetValueOnAnyThread() > 0;
 }
 
+bool IsStrataDbufferPassEnabled()
+{
+	return IsStrataEnabled() && CVarStrataDBufferPass.GetValueOnAnyThread() > 0;
+}
+
 enum EStrataTileSpace
 {
 	StrataTileSpace_Primary = 1u,
@@ -342,7 +353,8 @@ void InitialiseStrataFrameSceneData(FRDGBuilder& GraphBuilder, FSceneRenderer& S
 
 		// Top layer texture
 		{
-			Out.TopLayerTexture = GraphBuilder.CreateTexture(FRDGTextureDesc::Create2D(SceneTextureExtent, PF_R32_UINT, FClearValueBinding::Black, TexCreate_RenderTargetable | TexCreate_ShaderResource | TexCreate_FastVRAM), TEXT("Strata.TopLayerTexture"));
+			const bool bNeedUAV = CVarStrataDBufferPass.GetValueOnRenderThread() > 0;
+			Out.TopLayerTexture = GraphBuilder.CreateTexture(FRDGTextureDesc::Create2D(SceneTextureExtent, PF_R32_UINT, FClearValueBinding::Black, TexCreate_RenderTargetable | TexCreate_ShaderResource | TexCreate_FastVRAM | (bNeedUAV ? TexCreate_UAV : TexCreate_None)), TEXT("Strata.TopLayerTexture"));
 		}
 
 		// Separated subsurface and rough refraction textures
@@ -645,6 +657,41 @@ IMPLEMENT_GLOBAL_SHADER(FStrataMaterialTileClassificationPassCS, "/Engine/Privat
 
 ////////////////////////////////////////////////////////////////////////////////////////////////////////////////////////////////////////////////
 
+class FStrataDBufferPassCS : public FGlobalShader
+{
+	DECLARE_GLOBAL_SHADER(FStrataDBufferPassCS);
+	SHADER_USE_PARAMETER_STRUCT(FStrataDBufferPassCS, FGlobalShader);
+
+	class FTileType : SHADER_PERMUTATION_INT("PERMUTATION_TILETYPE", 3); // Only support Simple/Single/Complex
+	using FPermutationDomain = TShaderPermutationDomain<FTileType>;
+
+	BEGIN_SHADER_PARAMETER_STRUCT(FParameters, )
+		SHADER_PARAMETER(FIntPoint, ViewResolution)
+		SHADER_PARAMETER(uint32, MaxBytesPerPixel)
+		SHADER_PARAMETER(uint32, FirstSliceStoringStrataSSSData)
+		SHADER_PARAMETER_STRUCT_REF(FViewUniformShaderParameters, ViewUniformBuffer)
+		SHADER_PARAMETER_STRUCT_INCLUDE(FDBufferParameters, DBuffer)
+		SHADER_PARAMETER_RDG_TEXTURE_UAV(Texture2D, TopLayerTexture)
+		SHADER_PARAMETER_RDG_TEXTURE_UAV(RWTexture2DArray<uint>, MaterialTextureArrayUAV)
+		SHADER_PARAMETER_RDG_BUFFER_SRV(Buffer<uint>, TileListBuffer)
+		RDG_BUFFER_ACCESS(TileIndirectBuffer, ERHIAccess::IndirectArgs)
+	END_SHADER_PARAMETER_STRUCT()
+
+	static bool ShouldCompilePermutation(const FGlobalShaderPermutationParameters& Parameters)
+	{
+		return GetMaxSupportedFeatureLevel(Parameters.Platform) >= ERHIFeatureLevel::SM5 && Strata::IsStrataEnabled();
+	}
+
+	static void ModifyCompilationEnvironment(const FGlobalShaderPermutationParameters& Parameters, FShaderCompilerEnvironment& OutEnvironment)
+	{
+		FGlobalShader::ModifyCompilationEnvironment(Parameters, OutEnvironment);
+		OutEnvironment.SetDefine(TEXT("SHADER_DBUFFER"), 1);
+	}
+};
+IMPLEMENT_GLOBAL_SHADER(FStrataDBufferPassCS, "/Engine/Private/Strata/StrataDBuffer.usf", "MainCS", SF_Compute);
+
+////////////////////////////////////////////////////////////////////////////////////////////////////////////////////////////////////////////////
+
 class FStrataMaterialTilePrepareArgsPassCS : public FGlobalShader
 {
 	DECLARE_GLOBAL_SHADER(FStrataMaterialTilePrepareArgsPassCS);
@@ -1153,4 +1200,67 @@ void AddStrataMaterialClassificationPass(FRDGBuilder& GraphBuilder, const FMinim
 	}
 }
 
+
+void AddStrataDBufferPass(FRDGBuilder& GraphBuilder, const FMinimalSceneTextures& SceneTextures, const FDBufferTextures& DBufferTextures, const TArray<FViewInfo>& Views)
+{
+	RDG_EVENT_SCOPE_CONDITIONAL(GraphBuilder, IsStrataEnabled() && Views.Num() > 0, "Strata::DBuffer");
+	if (!IsStrataEnabled() || !DBufferTextures.IsValid() || CVarStrataDBufferPass.GetValueOnRenderThread() <= 0)
+	{
+		return;
+	}
+
+	for (int32 i = 0; i < Views.Num(); ++i)
+	{
+		RDG_EVENT_SCOPE_CONDITIONAL(GraphBuilder, Views.Num() > 1, "View%d", i);
+
+		const FViewInfo& View = Views[i];
+		if (!IsUsingDBuffers(View.GetShaderPlatform()) || View.Family->EngineShowFlags.Decals == 0)
+		{
+			continue;
+		}
+
+		const FStrataViewData* StrataViewData = &View.StrataViewData;
+		const FStrataSceneData* StrataSceneData = View.StrataViewData.SceneData;
+
+		auto DBufferPass = [&](EStrataTileType TileType)
+		{
+			// Only simple & single material are support but also dispatch complex tiles, 
+			// as they can contain simple/single material pixels
+			check(TileType == EStrataTileType::ESimple || TileType == EStrataTileType::ESingle || TileType == EStrataTileType::EComplex);
+
+			FStrataDBufferPassCS::FPermutationDomain PermutationVector;
+			PermutationVector.Set<FStrataDBufferPassCS::FTileType>(uint32(TileType));
+
+			TShaderMapRef<FStrataDBufferPassCS> ComputeShader(View.ShaderMap, PermutationVector);
+			FStrataDBufferPassCS::FParameters* PassParameters = GraphBuilder.AllocParameters<FStrataDBufferPassCS::FParameters>();
+
+			PassParameters->DBuffer = GetDBufferParameters(GraphBuilder, DBufferTextures, View.GetShaderPlatform());
+			PassParameters->ViewUniformBuffer = View.ViewUniformBuffer;
+			PassParameters->ViewResolution = View.ViewRect.Size();
+			PassParameters->MaxBytesPerPixel = StrataSceneData->MaxBytesPerPixel;
+			PassParameters->TopLayerTexture = GraphBuilder.CreateUAV(StrataSceneData->TopLayerTexture);
+			PassParameters->MaterialTextureArrayUAV = StrataSceneData->MaterialTextureArrayUAV;
+			PassParameters->FirstSliceStoringStrataSSSData = StrataSceneData->FirstSliceStoringStrataSSSData;
+
+			PassParameters->TileListBuffer = StrataViewData->ClassificationTileListBufferSRV[TileType];
+			PassParameters->TileIndirectBuffer = StrataViewData->ClassificationTileDispatchIndirectBuffer;
+
+			// Dispatch with tile data
+			const uint32 GroupSize = 8;
+			FComputeShaderUtils::AddPass(
+				GraphBuilder,
+				RDG_EVENT_NAME("Strata::Dbuffer(%s)", ToString(TileType)),
+				ERDGPassFlags::Compute,
+				ComputeShader,
+				PassParameters,
+				PassParameters->TileIndirectBuffer,
+				TileTypeDispatchIndirectArgOffset(TileType));
+		};
+
+		DBufferPass(EStrataTileType::ESimple);
+		DBufferPass(EStrataTileType::ESingle);
+		DBufferPass(EStrataTileType::EComplex);
+	}
+}
+
 } // namespace Strata
diff --git a/Engine/Source/Runtime/Renderer/Private/Strata/Strata.h b/Engine/Source/Runtime/Renderer/Private/Strata/Strata.h
--- a/Engine/Source/Runtime/Renderer/Private/Strata/Strata.h
+++ b/Engine/Source/Runtime/Renderer/Private/Strata/Strata.h
@@ -16,6 +16,7 @@ class FScene;
 class FRDGBuilder;
 struct FMinimalSceneTextures;
 struct FScreenPassTexture;
+struct FDBufferTextures;
 
 BEGIN_SHADER_PARAMETER_STRUCT(FStrataBasePassUniformParameters, )
 	SHADER_PARAMETER(uint32, MaxBytesPerPixel)
@@ -144,6 +145,7 @@ constexpr uint32 StencilBit_Single = 0x10; // In sync with SceneRenderTargets.h
 constexpr uint32 StencilBit_Complex= 0x20; // In sync with SceneRenderTargets.h - GET_STENCIL_BIT_MASK(STENCIL_STRATA_COMPLEX)
 
 bool IsStrataEnabled();
+bool IsStrataDbufferPassEnabled();
 
 FIntPoint GetStrataTextureResolution(const FIntPoint& InResolution);
 
@@ -158,6 +160,7 @@ void SetBasePassRenderTargetOutputFormat(const EShaderPlatform Platform, const F
 
 
 void AddStrataMaterialClassificationPass(FRDGBuilder& GraphBuilder, const FMinimalSceneTextures& SceneTextures, const TArray<FViewInfo>& Views);
+void AddStrataDBufferPass(FRDGBuilder& GraphBuilder, const FMinimalSceneTextures& SceneTextures, const FDBufferTextures& DBufferTextures, const TArray<FViewInfo>& Views);
 
 void AddStrataStencilPass(FRDGBuilder& GraphBuilder, const TArray<FViewInfo>& Views, const FMinimalSceneTextures& SceneTextures);
 
