Fix for water disappearing at large distances. Fixed the problem by storing device Z in the SceneDepthWithoutSingleLayerWater texture instead of linear depth. This gives back some precision by avoiding redundant conversions between device Z and linear depth. A potential problem is that some parts of water shading might use a linear sampler on these now non-linear values.

#rnx

diff --git a/Engine/Shaders/Private/HeightFogPixelShader.usf b/Engine/Shaders/Private/HeightFogPixelShader.usf
--- a/Engine/Shaders/Private/HeightFogPixelShader.usf
+++ b/Engine/Shaders/Private/HeightFogPixelShader.usf
@@ -13,16 +13,16 @@
 float UpsampleJitterMultiplier;
 
 float bOnlyOnRenderedOpaque;
-float bUseLinearDepthTexture;
+uint bUseWaterDepthTexture;
 
-float4 LinearDepthTextureMinMaxUV;
+float4 WaterDepthTextureMinMaxUV;
 float4 OcclusionTextureMinMaxUV;
 
 Texture2D OcclusionTexture;
 SamplerState OcclusionSampler;
 
-Texture2D LinearDepthTexture;
-SamplerState LinearDepthSampler;
+Texture2D WaterDepthTexture;
+SamplerState WaterDepthSampler;
 
 void ExponentialPixelMain(
 	float2 TexCoord : TEXCOORD0,
@@ -33,11 +33,11 @@ void ExponentialPixelMain(
 {
 	float SceneDepth;
 	float DeviceZ;
-	if (bUseLinearDepthTexture > 0.0f)
+	if (bUseWaterDepthTexture)
 	{
-		float2 SampleUV = clamp(TexCoord, LinearDepthTextureMinMaxUV.xy, LinearDepthTextureMinMaxUV.zw);
-		SceneDepth = Texture2DSampleLevel(LinearDepthTexture, LinearDepthSampler, SampleUV, 0).r * bUseLinearDepthTexture; // bUseLinearDepthTexture is also a scale to get depth in correct range
-		DeviceZ = ConvertToDeviceZ(SceneDepth);
+		float2 SampleUV = clamp(TexCoord, WaterDepthTextureMinMaxUV.xy, WaterDepthTextureMinMaxUV.zw);
+		DeviceZ = Texture2DSampleLevel(WaterDepthTexture, WaterDepthSampler, SampleUV, 0).r;
+		SceneDepth = ConvertFromDeviceZ(DeviceZ);
 	}
 	else
 	{
diff --git a/Engine/Shaders/Private/MaterialTemplate.ush b/Engine/Shaders/Private/MaterialTemplate.ush
--- a/Engine/Shaders/Private/MaterialTemplate.ush
+++ b/Engine/Shaders/Private/MaterialTemplate.ush
@@ -96,7 +96,7 @@
 	#include "/Engine/Private/SkyAtmosphereCommon.ush"
 #endif
 
-#if MATERIAL_SHADINGMODEL_SINGLELAYERWATER || POST_PROCESS_MATERIAL // For WaterSampleSceneDepthWithoutWater() and SINGLE_LAYER_WATER_DEPTH_SCALE
+#if MATERIAL_SHADINGMODEL_SINGLELAYERWATER || POST_PROCESS_MATERIAL // For WaterSampleSceneDepthWithoutWater()
 // This is currently included for post processing materials even if the SceneDepthWithoutWater texture is not sampled.
 // It might be worth it to add a switch to only include it if the texture is actually required by the material.
 	#include "/Engine/Private/SingleLayerWaterCommon.ush"
@@ -1688,14 +1688,16 @@ float MaterialExpressionSceneDepthWithoutWater(float2 ViewportUV, float Fallback
 {
 #if MATERIAL_SHADINGMODEL_SINGLELAYERWATER && !SCENE_TEXTURES_DISABLED && (!SIMPLE_SINGLE_LAYER_WATER || SINGLE_LAYER_WATER_SIMPLE_FORWARD)
 	const float2 ClampedViewportUV = clamp(ViewportUV, SingleLayerWater.SceneWithoutSingleLayerWaterMinMaxUV.xy, SingleLayerWater.SceneWithoutSingleLayerWaterMinMaxUV.zw);
-	return WaterSampleSceneDepthWithoutWater(SingleLayerWater.SceneDepthWithoutSingleLayerWaterTexture, SingleLayerWaterSceneDepthSampler, ClampedViewportUV, 
-		SingleLayerWater.SceneWithoutSingleLayerWaterTextureSize, SingleLayerWater.SceneWithoutSingleLayerWaterInvTextureSize) * SINGLE_LAYER_WATER_DEPTH_SCALE;
+	float WaterDeviceZ = WaterSampleSceneDepthWithoutWater(SingleLayerWater.SceneDepthWithoutSingleLayerWaterTexture, SingleLayerWaterSceneDepthSampler, ClampedViewportUV, 
+		SingleLayerWater.SceneWithoutSingleLayerWaterTextureSize, SingleLayerWater.SceneWithoutSingleLayerWaterInvTextureSize);
+	return ConvertFromDeviceZ(WaterDeviceZ);
 #elif POST_PROCESS_MATERIAL
 	if (bSceneDepthWithoutWaterTextureAvailable)
 	{
 		const float2 ClampedViewportUV = clamp(ViewportUV, SceneWithoutSingleLayerWaterMinMaxUV.xy, SceneWithoutSingleLayerWaterMinMaxUV.zw);
-		return WaterSampleSceneDepthWithoutWater(SceneDepthWithoutSingleLayerWaterTexture, SceneDepthWithoutSingleLayerWaterSampler, ClampedViewportUV, 
-			SceneWithoutSingleLayerWaterTextureSize, SceneWithoutSingleLayerWaterInvTextureSize) * SINGLE_LAYER_WATER_DEPTH_SCALE;
+		float WaterDeviceZ = WaterSampleSceneDepthWithoutWater(SceneDepthWithoutSingleLayerWaterTexture, SceneDepthWithoutSingleLayerWaterSampler, ClampedViewportUV,
+			SceneWithoutSingleLayerWaterTextureSize, SceneWithoutSingleLayerWaterInvTextureSize);
+		return ConvertFromDeviceZ(WaterDeviceZ);
 	}
 	else
 	{
diff --git a/Engine/Shaders/Private/SingleLayerWaterCommon.ush b/Engine/Shaders/Private/SingleLayerWaterCommon.ush
--- a/Engine/Shaders/Private/SingleLayerWaterCommon.ush
+++ b/Engine/Shaders/Private/SingleLayerWaterCommon.ush
@@ -3,12 +3,6 @@
 #pragma once
 #include "TextureSampling.ush" // GetBilinearSampleLevelInfos() and GetBilinearInterpolation()
 
-// Scale for encoding depth into SceneDepthWithoutSingleLayerWaterTexture
-// Must match FDeferredShadingSceneRenderer::RenderUnderWaterFog
-// Originally used to work around precision issues with a float16 linear depth buffer.
-// TODO: Do we still need this?
-#define SINGLE_LAYER_WATER_DEPTH_SCALE 100.0f
-
 #ifndef SIMPLE_SINGLE_LAYER_WATER
 #define SIMPLE_SINGLE_LAYER_WATER 0
 #endif
diff --git a/Engine/Shaders/Private/SingleLayerWaterComposite.usf b/Engine/Shaders/Private/SingleLayerWaterComposite.usf
--- a/Engine/Shaders/Private/SingleLayerWaterComposite.usf
+++ b/Engine/Shaders/Private/SingleLayerWaterComposite.usf
@@ -84,13 +84,13 @@ SingleLayerWaterCompositeOutput SingleLayerWaterComposite(float2 BufferUV, float
 
 		const float2 SceneNoWaterUV = clamp(BufferUV, EffectiveSceneNoWaterMinMaxUV.xy, EffectiveSceneNoWaterMinMaxUV.zw);
 
-		float OpaqueDepth = WaterSampleSceneDepthWithoutWater(
+		float OpaqueDeviceZ = WaterSampleSceneDepthWithoutWater(
 			SceneNoWaterDepthTexture,
 			SceneNoWaterDepthSampler,
 			SceneNoWaterUV,
 			SceneNoWaterTextureSize,
-			SceneNoWaterInvTextureSize)
-			* SINGLE_LAYER_WATER_DEPTH_SCALE;
+			SceneNoWaterInvTextureSize);
+		float OpaqueDepth = ConvertFromDeviceZ(OpaqueDeviceZ);
 
 		float WaterDepth = SceneDepth;
 
@@ -173,13 +173,13 @@ SingleLayerWaterCompositeOutput SingleLayerWaterComposite(float2 BufferUV, float
 	{
 		const float2 SceneNoWaterUV = clamp(BufferUV, EffectiveSceneNoWaterMinMaxUV.xy, EffectiveSceneNoWaterMinMaxUV.zw);
 
-		float OpaqueDepth = WaterSampleSceneDepthWithoutWater(
+		float OpaqueDeviceZ = WaterSampleSceneDepthWithoutWater(
 			SceneNoWaterDepthTexture,
 			SceneNoWaterDepthSampler,
 			SceneNoWaterUV,
 			SceneNoWaterTextureSize,
-			SceneNoWaterInvTextureSize)
-			* SINGLE_LAYER_WATER_DEPTH_SCALE;
+			SceneNoWaterInvTextureSize);
+		float OpaqueDepth = ConvertFromDeviceZ(OpaqueDeviceZ);
 
 		float WaterDepth = GBuffer.Depth;
 
@@ -484,7 +484,7 @@ void WaterRefractionCopyPS(
 #else
 	DeviceDepthZ = Texture2DSampleLevel(SceneDepthCopyDownsampleTexture, SceneDepthCopyDownsampleSampler, BufferUV, 0).x;
 #endif
-	OutSceneDepth = ConvertFromDeviceZ(DeviceDepthZ) / SINGLE_LAYER_WATER_DEPTH_SCALE;
+	OutSceneDepth = DeviceDepthZ;
 
 #if METAL_PROFILE
 	// Clamp on iOS to avoid #inf
diff --git a/Engine/Shaders/Private/SingleLayerWaterShading.ush b/Engine/Shaders/Private/SingleLayerWaterShading.ush
--- a/Engine/Shaders/Private/SingleLayerWaterShading.ush
+++ b/Engine/Shaders/Private/SingleLayerWaterShading.ush
@@ -72,15 +72,14 @@ WaterVolumeLightingOutput EvaluateWaterVolumeLighting(
 
 #if SINGLE_LAYER_WATER_SIMPLE_FORWARD 
 
-		float PixelSceneDepth = WaterSampleSceneDepthWithoutWater(
+		float PixelSceneDeviceZ = WaterSampleSceneDepthWithoutWater(
 			SceneDepthWithoutSingleLayerWaterTexture,
 			SceneDepthWithoutSingleLayerWaterSampler,
 			ViewportUV,
 			SceneDepthWithoutSingleLayerWaterTextureSize,
-			SceneDepthWithoutSingleLayerWaterTextureInvSize)
-			* SINGLE_LAYER_WATER_DEPTH_SCALE;
+			SceneDepthWithoutSingleLayerWaterTextureInvSize);
 
-		float PixelSceneDeviceZ = ConvertToDeviceZ(PixelSceneDepth);
+		float PixelSceneDepth = ConvertFromDeviceZ(PixelSceneDeviceZ);
 #else
 		float PixelSceneDeviceZ = LookupDeviceZ(ViewportUV);
 		float PixelSceneDepth = ConvertFromDeviceZ(PixelSceneDeviceZ);
@@ -96,13 +95,14 @@ WaterVolumeLightingOutput EvaluateWaterVolumeLighting(
 
 		ViewportUV = clamp(ViewportUV, SceneWithoutSingleLayerWaterMinUV, SceneWithoutSingleLayerWaterMaxUV);
 
-		float PixelSceneDepth = WaterSampleSceneDepthWithoutWater(
-			SceneDepthWithoutSingleLayerWaterTexture, 
-			SceneDepthWithoutSingleLayerWaterSampler, 
-			ViewportUV, 
-			SceneDepthWithoutSingleLayerWaterTextureSize, 
-			SceneDepthWithoutSingleLayerWaterTextureInvSize) 
-			* SINGLE_LAYER_WATER_DEPTH_SCALE;
+		float PixelSceneDeviceZ = WaterSampleSceneDepthWithoutWater(
+			SceneDepthWithoutSingleLayerWaterTexture,
+			SceneDepthWithoutSingleLayerWaterSampler,
+			ViewportUV,
+			SceneDepthWithoutSingleLayerWaterTextureSize,
+			SceneDepthWithoutSingleLayerWaterTextureInvSize);
+
+		float PixelSceneDepth = ConvertFromDeviceZ(PixelSceneDeviceZ);
 		
 		float2 Refraction = GetMaterialRefraction(PixelMaterialInputs); // X = IOR, Y = Refraction Bias
 		float2 BufferUVDistortion = ComputeBufferUVDistortion(
@@ -117,22 +117,24 @@ WaterVolumeLightingOutput EvaluateWaterVolumeLighting(
 
 		// Gather the four depth samples that would be used for bilinear interpolation. We need to know if any of these are invalid (above the water). If so, we have to discard this distorted sample and fall back. Otherwise we end up sampling foreground color from the bilinear color sample.
 		// This is only guaranteed to be artifact free if the input depth/color resolution is the same as the output. If the input is downsampled, this will still do the right thing but might still sample leaking color
-		float4 SceneDepth4 = SceneDepthWithoutSingleLayerWaterTexture.Gather(SceneDepthWithoutSingleLayerWaterSampler, DistortedUV) * SINGLE_LAYER_WATER_DEPTH_SCALE;
+		float4 SceneDeviceZ4 = SceneDepthWithoutSingleLayerWaterTexture.Gather(SceneDepthWithoutSingleLayerWaterSampler, DistortedUV);
+		float4 SceneDepth4 = float4(ConvertFromDeviceZ(SceneDeviceZ4.x), ConvertFromDeviceZ(SceneDeviceZ4.y), ConvertFromDeviceZ(SceneDeviceZ4.z), ConvertFromDeviceZ(SceneDeviceZ4.w));
 		
-		float SceneDepth = 0.0;
+		float SceneDepth = 0.0f;
+		float SceneDeviceZ = 0.0f;
 		if(any(SceneDepth4 < WaterDepth))
 		{
 			SceneDepth = PixelSceneDepth; // The pixel we try to sample is closer than the water pixel: fallback to no distortion
+			SceneDeviceZ = PixelSceneDeviceZ;
 		}
 		else
 		{
 			// We can use bilinear interpolation here because all depth samples are guaranteed to be farther away than the water surface depth.
-			SceneDepth = GetBilinearInterpolation(GetBilinearSampleLevelInfos(DistortedUV, SceneDepthWithoutSingleLayerWaterTextureSize, SceneDepthWithoutSingleLayerWaterTextureInvSize), SceneDepth4.wzxy);
+			SceneDeviceZ = GetBilinearInterpolation(GetBilinearSampleLevelInfos(DistortedUV, SceneDepthWithoutSingleLayerWaterTextureSize, SceneDepthWithoutSingleLayerWaterTextureInvSize), SceneDeviceZ4.wzxy);
+			SceneDepth = ConvertFromDeviceZ(SceneDeviceZ);
 			ViewportUV = DistortedUV;
 		}
 
-		const float SceneDeviceZ = ConvertToDeviceZ(SceneDepth);
-
 #endif // !SIMPLE_SINGLE_LAYER_WATER
 
 		const float BehindWaterDeltaDepth = CameraIsUnderWater ? WaterDepth : max(0.0f, SceneDepth - WaterDepth);
diff --git a/Engine/Shaders/Private/VolumetricRenderTarget.usf b/Engine/Shaders/Private/VolumetricRenderTarget.usf
--- a/Engine/Shaders/Private/VolumetricRenderTarget.usf
+++ b/Engine/Shaders/Private/VolumetricRenderTarget.usf
@@ -436,14 +436,13 @@ void ComposeVolumetricRTOverScenePS(
 
 #elif (PERMUTATION_UPSAMPLINGMODE==4 || PERMUTATION_UPSAMPLINGMODE==3 || PERMUTATION_UPSAMPLINGMODE==2)
 
-	const float UseWaterLinearDepthTexture = 100.0f;	// This must match SINGLE_LAYER_WATER_DEPTH_SCALE from SingleLayerWaterCommon.ush and SingleLayerWaterComposite.usf
 #if PERMUTATION_RENDER_UNDERWATER_BUFFER
 	// Adapt the UV to the relative water buffer size
 	float2 WaterVolumeUV = VolumeUV * FullResolutionToWaterBufferScale.y;
 	// Offset the uv to the view buffer region and take into account dynamic resolution scaling.
 	float2 WaterDepthScreenUV = SceneWithoutSingleLayerWaterViewRect.xy + WaterVolumeUV * (View.ViewSizeAndInvSize.xy * View.BufferSizeAndInvSize.zw);
 
-	float PixelLinearDepth = WaterLinearDepthTexture.SampleLevel(WaterLinearDepthSampler, WaterDepthScreenUV, 0).r * UseWaterLinearDepthTexture;
+	float PixelLinearDepth = ConvertFromDeviceZ(WaterLinearDepthTexture.SampleLevel(WaterLinearDepthSampler, WaterDepthScreenUV, 0).r);
 
 	float3 WorldPosition = LWCHackToFloat(SvPositionToWorld(float4(SVPos.xy, 0.5, 1.0)));
 	WorldPosition = normalize(WorldPosition - LWCHackToFloat(PrimaryView.WorldCameraOrigin)) * PixelLinearDepth + LWCHackToFloat(PrimaryView.WorldCameraOrigin);
@@ -472,7 +471,7 @@ void ComposeVolumetricRTOverScenePS(
 #if PERMUTATION_RENDER_CAMERA_UNDERWATER
 	// Now check that we are compositing a pixel that is not "water" to avoid applying clouds twice (they are already composited with the behind water layer scene).
 	// We also lack depth information behind the water surface now so the composition would be wrong anyway.
-	float WaterTestPixelLinearDepth = WaterLinearDepthTexture.SampleLevel(WaterLinearDepthSampler, ScreenUV, 0).r * UseWaterLinearDepthTexture;
+	float WaterTestPixelLinearDepth = ConvertFromDeviceZ(WaterLinearDepthTexture.SampleLevel(WaterLinearDepthSampler, ScreenUV, 0).r);
 	if (WaterTestPixelLinearDepth > PixelDistanceFromView)
 	{
 		// This pixel contains water, so skip it because clouds have already been composited in the "behind water scene color".
@@ -578,15 +577,11 @@ void ComposeVolumetricRTOverScenePS(
 		float4 DataAcc = 0.0f;
 		const float CloudFrontDepthTinyOffset = 0.001;
 
-		// When this pixel is at a distance further than what the water depth buffer can represent, trivially accept all neighboorinjg pixels containing clouds.
-		// This makes sure that clouds around the distant horizon are also composited and not skip by "depth buffer is further than the cloud front depth" test.
-		const bool ISWaterFarDepthBuffer = PixelDistanceFromViewKm >= UseWaterLinearDepthTexture * Max10BitsFloat;
-
 		// We are testing if the depth buffer is further than the cloud front depth and that the cloud front depth is actually in front of traced depth.
-		if (PixelDistanceFromViewKm > (VolumeFrontDepth0.x - CloudFrontDepthTinyOffset || ISWaterFarDepthBuffer) && VolumeFrontDepth0.x < (VolumeFrontDepth0.y)) { DataAcc += VolumeRGBT0; ValidSampleCount += 1.0f; }
-		if (PixelDistanceFromViewKm > (VolumeFrontDepth1.x - CloudFrontDepthTinyOffset || ISWaterFarDepthBuffer) && VolumeFrontDepth1.x < (VolumeFrontDepth1.y)) { DataAcc += VolumeRGBT1; ValidSampleCount += 1.0f; }
-		if (PixelDistanceFromViewKm > (VolumeFrontDepth2.x - CloudFrontDepthTinyOffset || ISWaterFarDepthBuffer) && VolumeFrontDepth2.x < (VolumeFrontDepth2.y)) { DataAcc += VolumeRGBT2; ValidSampleCount += 1.0f; }
-		if (PixelDistanceFromViewKm > (VolumeFrontDepth3.x - CloudFrontDepthTinyOffset || ISWaterFarDepthBuffer) && VolumeFrontDepth3.x < (VolumeFrontDepth3.y)) { DataAcc += VolumeRGBT3; ValidSampleCount += 1.0f; }
+		if (PixelDistanceFromViewKm > (VolumeFrontDepth0.x - CloudFrontDepthTinyOffset) && VolumeFrontDepth0.x < (VolumeFrontDepth0.y)) { DataAcc += VolumeRGBT0; ValidSampleCount += 1.0f; }
+		if (PixelDistanceFromViewKm > (VolumeFrontDepth1.x - CloudFrontDepthTinyOffset) && VolumeFrontDepth1.x < (VolumeFrontDepth1.y)) { DataAcc += VolumeRGBT1; ValidSampleCount += 1.0f; }
+		if (PixelDistanceFromViewKm > (VolumeFrontDepth2.x - CloudFrontDepthTinyOffset) && VolumeFrontDepth2.x < (VolumeFrontDepth2.y)) { DataAcc += VolumeRGBT2; ValidSampleCount += 1.0f; }
+		if (PixelDistanceFromViewKm > (VolumeFrontDepth3.x - CloudFrontDepthTinyOffset) && VolumeFrontDepth3.x < (VolumeFrontDepth3.y)) { DataAcc += VolumeRGBT3; ValidSampleCount += 1.0f; }
 		if (ValidSampleCount > 0.0f)
 		{
 			OutputRt0 = DataAcc / ValidSampleCount;
diff --git a/Engine/Source/Runtime/Renderer/Private/FogRendering.cpp b/Engine/Source/Runtime/Renderer/Private/FogRendering.cpp
--- a/Engine/Source/Runtime/Renderer/Private/FogRendering.cpp
+++ b/Engine/Source/Runtime/Renderer/Private/FogRendering.cpp
@@ -124,12 +124,12 @@ class FExponentialHeightFogPS : public FGlobalShader
 		SHADER_PARAMETER_RDG_UNIFORM_BUFFER(FFogUniformParameters, FogUniformBuffer)
 		SHADER_PARAMETER_RDG_TEXTURE(Texture2D, OcclusionTexture)
 		SHADER_PARAMETER_SAMPLER(SamplerState, OcclusionSampler)
-		SHADER_PARAMETER_RDG_TEXTURE(Texture2D, LinearDepthTexture)
-		SHADER_PARAMETER_SAMPLER(SamplerState, LinearDepthSampler)
+		SHADER_PARAMETER_RDG_TEXTURE(Texture2D, WaterDepthTexture)
+		SHADER_PARAMETER_SAMPLER(SamplerState, WaterDepthSampler)
 		SHADER_PARAMETER(float, bOnlyOnRenderedOpaque)
-		SHADER_PARAMETER(float, bUseLinearDepthTexture)
+		SHADER_PARAMETER(uint32, bUseWaterDepthTexture)
 		SHADER_PARAMETER(float, UpsampleJitterMultiplier)
-		SHADER_PARAMETER(FVector4f, LinearDepthTextureMinMaxUV)
+		SHADER_PARAMETER(FVector4f, WaterDepthTextureMinMaxUV)
 		SHADER_PARAMETER(FVector4f, OcclusionTextureMinMaxUV)
 	END_SHADER_PARAMETER_STRUCT()
 
@@ -265,13 +265,13 @@ static FFogPassParameters* CreateDefaultFogPassParameters(
 	PassParameters->PS.FogUniformBuffer = FogUniformBuffer;
 	PassParameters->PS.OcclusionTexture = LightShaftOcclusionTexture != nullptr ? LightShaftOcclusionTexture : GSystemTextures.GetWhiteDummy(GraphBuilder);
 	PassParameters->PS.OcclusionSampler = TStaticSamplerState<SF_Bilinear, AM_Clamp, AM_Clamp, AM_Clamp>::GetRHI();
-	PassParameters->PS.LinearDepthTexture = GSystemTextures.GetDepthDummy(GraphBuilder);
-	PassParameters->PS.LinearDepthSampler = TStaticSamplerState<SF_Point, AM_Clamp, AM_Clamp, AM_Clamp>::GetRHI();
+	PassParameters->PS.WaterDepthTexture = GSystemTextures.GetDepthDummy(GraphBuilder);
+	PassParameters->PS.WaterDepthSampler = TStaticSamplerState<SF_Point, AM_Clamp, AM_Clamp, AM_Clamp>::GetRHI();
 	PassParameters->PS.OcclusionTextureMinMaxUV = FVector4f(LightShaftParameters.UVViewportBilinearMin, LightShaftParameters.UVViewportBilinearMax);
-	PassParameters->PS.LinearDepthTextureMinMaxUV = FVector4f::Zero();
+	PassParameters->PS.WaterDepthTextureMinMaxUV = FVector4f::Zero();
 	PassParameters->PS.UpsampleJitterMultiplier = CVarUpsampleJitterMultiplier.GetValueOnRenderThread() * GVolumetricFogGridPixelSize;
 	PassParameters->PS.bOnlyOnRenderedOpaque = View.bFogOnlyOnRenderedOpaque;
-	PassParameters->PS.bUseLinearDepthTexture = 0.0f;
+	PassParameters->PS.bUseWaterDepthTexture = false;
 	return PassParameters;
 }
 
@@ -410,14 +410,11 @@ void FDeferredShadingSceneRenderer::RenderUnderWaterFog(
 		RDG_EVENT_SCOPE(GraphBuilder, "SLW::ExponentialHeightFog");
 		RDG_GPU_STAT_SCOPE(GraphBuilder, Fog);
 
-		FRDGTextureRef LinearDepthTexture = SceneWithoutWaterTextures.DepthTexture;
-		check(LinearDepthTexture);
+		FRDGTextureRef WaterDepthTexture = SceneWithoutWaterTextures.DepthTexture;
+		check(WaterDepthTexture);
 
 		const bool bShouldRenderVolumetricFog = ShouldRenderVolumetricFog();
 
-		// This must match SINGLE_LAYER_WATER_DEPTH_SCALE from SingleLayerWaterCommon.ush and SingleLayerWaterComposite.usf.
-		const float kSingleLayerWaterDepthScale = 100.0f;
-
 		for (int32 ViewIndex = 0; ViewIndex < Views.Num(); ViewIndex++)
 		{
 			const FViewInfo& View = Views[ViewIndex];
@@ -435,9 +432,9 @@ void FDeferredShadingSceneRenderer::RenderUnderWaterFog(
 				FScreenPassTextureViewportParameters LightShaftParameters = FScreenPassTextureViewportParameters();
 
 				FFogPassParameters* PassParameters = CreateDefaultFogPassParameters(GraphBuilder, View, SceneTexturesWithDepth, FogUniformBuffer, LightShaftOcclusionTexture, LightShaftParameters);
-				PassParameters->PS.LinearDepthTexture = LinearDepthTexture;
-				PassParameters->PS.bUseLinearDepthTexture = kSingleLayerWaterDepthScale;
-				PassParameters->PS.LinearDepthTextureMinMaxUV = SceneWithoutWaterView.MinMaxUV;
+				PassParameters->PS.WaterDepthTexture = WaterDepthTexture;
+				PassParameters->PS.bUseWaterDepthTexture = true;
+				PassParameters->PS.WaterDepthTextureMinMaxUV = SceneWithoutWaterView.MinMaxUV;
 				PassParameters->RenderTargets[0] = FRenderTargetBinding(SceneWithoutWaterTextures.ColorTexture, ERenderTargetLoadAction::ELoad);
 
 				GraphBuilder.AddPass(RDG_EVENT_NAME("FogBehindWater"), PassParameters, ERDGPassFlags::Raster, [this, &View, SceneWithoutWaterView, PassParameters, bShouldRenderVolumetricFog](FRHICommandList& RHICmdList)
