Lumen - don't store interiors of the objects inside the object grid, as those don't have a valid surface in a given cell.

diff --git a/Engine/Shaders/Private/DistanceField/GlobalDistanceFieldObjectGrid.ush b/Engine/Shaders/Private/DistanceField/GlobalDistanceFieldObjectGrid.ush
--- a/Engine/Shaders/Private/DistanceField/GlobalDistanceFieldObjectGrid.ush
+++ b/Engine/Shaders/Private/DistanceField/GlobalDistanceFieldObjectGrid.ush
@@ -46,7 +46,7 @@ FObjectGridCellIndex UnpackObjectGridCellIndex(uint PackedIndex)
 // Unordered add to a grid cell
 void AddToObjectGridCell(inout FObjectGridCell ObjectGridCell, uint GPUSceneInstanceIndex, float Distance, float ObjectGridCellWorldExtent, float MaxQueryDistance)
 {
-	if (Distance < MaxQueryDistance)
+	if (abs(Distance) < MaxQueryDistance)
 	{
 		uint PackedIndexToInsert = PackObjectGridCellIndex(GPUSceneInstanceIndex, Distance, ObjectGridCellWorldExtent, MaxQueryDistance);
 
