Minor change indirect args buffer to be initialized in parallel.

diff --git a/Engine/Shaders/Private/Strata/StrataMaterialClassification.usf b/Engine/Shaders/Private/Strata/StrataMaterialClassification.usf
--- a/Engine/Shaders/Private/Strata/StrataMaterialClassification.usf
+++ b/Engine/Shaders/Private/Strata/StrataMaterialClassification.usf
@@ -76,20 +76,12 @@ groupshared uint s_TileFlags[GROUP_THREAD_COUNT];
 [numthreads(STRATA_TILE_SIZE, STRATA_TILE_SIZE, 1)]
 void TileMainCS(uint2 DispatchThreadId : SV_DispatchThreadID, uint LinearIndex : SV_GroupIndex, uint3 GroupId : SV_GroupID)
 {
-	if (all(DispatchThreadId == 0))
+	// Init primitive index
+	if (DispatchThreadId.x < STRATA_TILE_TYPE_COUNT && DispatchThreadId.y == 0)
 	{
+		const uint TileType = DispatchThreadId.x;
 		const uint IndexCountPerInstance = bRectPrimitive > 0 ? 4 : 6;
-		TileDrawIndirectDataBuffer[GetStrataTileTypeDrawIndirectArgOffset_DWord(STRATA_TILE_TYPE_SIMPLE)					+ 0] = IndexCountPerInstance;
-		TileDrawIndirectDataBuffer[GetStrataTileTypeDrawIndirectArgOffset_DWord(STRATA_TILE_TYPE_SINGLE)					+ 0] = IndexCountPerInstance;
-		TileDrawIndirectDataBuffer[GetStrataTileTypeDrawIndirectArgOffset_DWord(STRATA_TILE_TYPE_COMPLEX)					+ 0] = IndexCountPerInstance;
-		TileDrawIndirectDataBuffer[GetStrataTileTypeDrawIndirectArgOffset_DWord(STRATA_TILE_TYPE_ROUGH_REFRACT)				+ 0] = IndexCountPerInstance;
-		TileDrawIndirectDataBuffer[GetStrataTileTypeDrawIndirectArgOffset_DWord(STRATA_TILE_TYPE_ROUGH_REFRACT_SSS_WITHOUT) + 0] = IndexCountPerInstance;
-
-	#if PERMUTATION_DECAL
-		TileDrawIndirectDataBuffer[GetStrataTileTypeDrawIndirectArgOffset_DWord(STRATA_TILE_TYPE_DECAL_SIMPLE) + 0]				 = IndexCountPerInstance;
-		TileDrawIndirectDataBuffer[GetStrataTileTypeDrawIndirectArgOffset_DWord(STRATA_TILE_TYPE_DECAL_SINGLE) + 0]				 = IndexCountPerInstance;
-		TileDrawIndirectDataBuffer[GetStrataTileTypeDrawIndirectArgOffset_DWord(STRATA_TILE_TYPE_DECAL_COMPLEX) + 0]			 = IndexCountPerInstance;
-	#endif
+		TileDrawIndirectDataBuffer[GetStrataTileTypeDrawIndirectArgOffset_DWord(TileType) + 0] = IndexCountPerInstance;
 	}
 
 	const uint2 PixelCoord = DispatchThreadId.xy + View.ViewRectMin.xy;
@@ -348,37 +340,15 @@ void TileMainCS(uint2 DispatchThreadId : SV_DispatchThreadID, uint LinearIndex :
 Buffer<uint>   TileDrawIndirectDataBuffer;
 RWBuffer<uint> TileDispatchIndirectDataBuffer;
 
-[numthreads(1, 1, 1)]
+[numthreads(32, 1, 1)]
 void ArgsMainCS(uint2 DispatchThreadId : SV_DispatchThreadID)
 {
-	if (all(DispatchThreadId == 0))
+	const uint TileType = DispatchThreadId.x;
+	if (TileType < STRATA_TILE_TYPE_COUNT)
 	{
-		TileDispatchIndirectDataBuffer[GetStrataTileTypeDispatchIndirectArgOffset_DWord(STRATA_TILE_TYPE_SIMPLE)					+ 0] = TileDrawIndirectDataBuffer[GetStrataTileTypeDrawIndirectArgOffset_DWord(STRATA_TILE_TYPE_SIMPLE)						+ 1];
-		TileDispatchIndirectDataBuffer[GetStrataTileTypeDispatchIndirectArgOffset_DWord(STRATA_TILE_TYPE_SINGLE)					+ 0] = TileDrawIndirectDataBuffer[GetStrataTileTypeDrawIndirectArgOffset_DWord(STRATA_TILE_TYPE_SINGLE)						+ 1];
-		TileDispatchIndirectDataBuffer[GetStrataTileTypeDispatchIndirectArgOffset_DWord(STRATA_TILE_TYPE_COMPLEX)					+ 0] = TileDrawIndirectDataBuffer[GetStrataTileTypeDrawIndirectArgOffset_DWord(STRATA_TILE_TYPE_COMPLEX)					+ 1];
-		TileDispatchIndirectDataBuffer[GetStrataTileTypeDispatchIndirectArgOffset_DWord(STRATA_TILE_TYPE_ROUGH_REFRACT)				+ 0] = TileDrawIndirectDataBuffer[GetStrataTileTypeDrawIndirectArgOffset_DWord(STRATA_TILE_TYPE_ROUGH_REFRACT)				+ 1];
-		TileDispatchIndirectDataBuffer[GetStrataTileTypeDispatchIndirectArgOffset_DWord(STRATA_TILE_TYPE_ROUGH_REFRACT_SSS_WITHOUT)	+ 0] = TileDrawIndirectDataBuffer[GetStrataTileTypeDrawIndirectArgOffset_DWord(STRATA_TILE_TYPE_ROUGH_REFRACT_SSS_WITHOUT)	+ 1];
-		TileDispatchIndirectDataBuffer[GetStrataTileTypeDispatchIndirectArgOffset_DWord(STRATA_TILE_TYPE_DECAL_SIMPLE)				+ 0] = TileDrawIndirectDataBuffer[GetStrataTileTypeDrawIndirectArgOffset_DWord(STRATA_TILE_TYPE_DECAL_SIMPLE)				+ 1];
-		TileDispatchIndirectDataBuffer[GetStrataTileTypeDispatchIndirectArgOffset_DWord(STRATA_TILE_TYPE_DECAL_SINGLE)				+ 0] = TileDrawIndirectDataBuffer[GetStrataTileTypeDrawIndirectArgOffset_DWord(STRATA_TILE_TYPE_DECAL_SINGLE)				+ 1];
-		TileDispatchIndirectDataBuffer[GetStrataTileTypeDispatchIndirectArgOffset_DWord(STRATA_TILE_TYPE_DECAL_COMPLEX)				+ 0] = TileDrawIndirectDataBuffer[GetStrataTileTypeDrawIndirectArgOffset_DWord(STRATA_TILE_TYPE_DECAL_COMPLEX)				+ 1];
-
-		TileDispatchIndirectDataBuffer[GetStrataTileTypeDispatchIndirectArgOffset_DWord(STRATA_TILE_TYPE_SIMPLE)					+ 1] = 1;
-		TileDispatchIndirectDataBuffer[GetStrataTileTypeDispatchIndirectArgOffset_DWord(STRATA_TILE_TYPE_SINGLE)					+ 1] = 1;
-		TileDispatchIndirectDataBuffer[GetStrataTileTypeDispatchIndirectArgOffset_DWord(STRATA_TILE_TYPE_COMPLEX)					+ 1] = 1;
-		TileDispatchIndirectDataBuffer[GetStrataTileTypeDispatchIndirectArgOffset_DWord(STRATA_TILE_TYPE_ROUGH_REFRACT)				+ 1] = 1;
-		TileDispatchIndirectDataBuffer[GetStrataTileTypeDispatchIndirectArgOffset_DWord(STRATA_TILE_TYPE_ROUGH_REFRACT_SSS_WITHOUT) + 1] = 1;
-		TileDispatchIndirectDataBuffer[GetStrataTileTypeDispatchIndirectArgOffset_DWord(STRATA_TILE_TYPE_DECAL_SIMPLE)				+ 1] = 1;
-		TileDispatchIndirectDataBuffer[GetStrataTileTypeDispatchIndirectArgOffset_DWord(STRATA_TILE_TYPE_DECAL_SINGLE)				+ 1] = 1;
-		TileDispatchIndirectDataBuffer[GetStrataTileTypeDispatchIndirectArgOffset_DWord(STRATA_TILE_TYPE_DECAL_COMPLEX)				+ 1] = 1;
-
-		TileDispatchIndirectDataBuffer[GetStrataTileTypeDispatchIndirectArgOffset_DWord(STRATA_TILE_TYPE_SIMPLE)					+ 2] = 1;
-		TileDispatchIndirectDataBuffer[GetStrataTileTypeDispatchIndirectArgOffset_DWord(STRATA_TILE_TYPE_SINGLE)					+ 2] = 1;
-		TileDispatchIndirectDataBuffer[GetStrataTileTypeDispatchIndirectArgOffset_DWord(STRATA_TILE_TYPE_COMPLEX)					+ 2] = 1;
-		TileDispatchIndirectDataBuffer[GetStrataTileTypeDispatchIndirectArgOffset_DWord(STRATA_TILE_TYPE_ROUGH_REFRACT)				+ 2] = 1;
-		TileDispatchIndirectDataBuffer[GetStrataTileTypeDispatchIndirectArgOffset_DWord(STRATA_TILE_TYPE_ROUGH_REFRACT_SSS_WITHOUT)	+ 2] = 1;
-		TileDispatchIndirectDataBuffer[GetStrataTileTypeDispatchIndirectArgOffset_DWord(STRATA_TILE_TYPE_DECAL_SIMPLE)				+ 2] = 1;
-		TileDispatchIndirectDataBuffer[GetStrataTileTypeDispatchIndirectArgOffset_DWord(STRATA_TILE_TYPE_DECAL_SINGLE)				+ 2] = 1;
-		TileDispatchIndirectDataBuffer[GetStrataTileTypeDispatchIndirectArgOffset_DWord(STRATA_TILE_TYPE_DECAL_COMPLEX)				+ 2] = 1;
+		TileDispatchIndirectDataBuffer[GetStrataTileTypeDispatchIndirectArgOffset_DWord(TileType) + 0] = TileDrawIndirectDataBuffer[GetStrataTileTypeDrawIndirectArgOffset_DWord(TileType) + 1];
+		TileDispatchIndirectDataBuffer[GetStrataTileTypeDispatchIndirectArgOffset_DWord(TileType) + 1] = 1;
+		TileDispatchIndirectDataBuffer[GetStrataTileTypeDispatchIndirectArgOffset_DWord(TileType) + 2] = 1;
 	}
 }
 
