Fix ShaderPrint draw circle.

diff --git a/Engine/Shaders/Private/ShaderPrint.ush b/Engine/Shaders/Private/ShaderPrint.ush
--- a/Engine/Shaders/Private/ShaderPrint.ush
+++ b/Engine/Shaders/Private/ShaderPrint.ush
@@ -1122,7 +1122,7 @@ void AddCross   (FShaderPrintContext Ctx, float3 Pos, float Size, float4 Color)
 ///////////////////////////////////////////////////////////////////////////////////////////////////
 // Circles
 
-void AddCircleSS(float2 Center, float Radius, float4 Color)
+void AddCircleSS(float2 Center, float Radius, float4 Color, uint SegmentCount=16)
 {
 	const float TStep = 0.05f;
 
@@ -1130,8 +1130,9 @@ void AddCircleSS(float2 Center, float Radius, float4 Color)
 	sincos(0.0, S, C);
 	float2 PrevP = Center + float2(C, S) * Radius;
 
-	for (float t = TStep; t <= 1.0; t += TStep)
+	for (uint SegIt = 1; SegIt <= SegmentCount; ++SegIt)
 	{
+		const float t = float(SegIt) / float(SegmentCount-1);
 		sincos(t * 2 * PI, S, C);
 		float2 P = Center + float2(C, S) * Radius;
 
@@ -1140,7 +1141,7 @@ void AddCircleSS(float2 Center, float Radius, float4 Color)
 	}
 }
 
-void AddCircleSS(FShaderPrintContext Ctx, float2 Center, float Radius, float4 Color) { if (Ctx.bIsActive) { AddCircleSS(Center, Radius, Color); } }
+void AddCircleSS(FShaderPrintContext Ctx, float2 Center, float Radius, float4 Color, uint SegmentCount = 16) { if (Ctx.bIsActive) { AddCircleSS(Center, Radius, Color, SegmentCount); } }
 
 
 ///////////////////////////////////////////////////////////////////////////////////////////////////
