Fix Strata visualization for SimpleVolume/Thin.

diff --git a/Engine/Shaders/Private/Strata/StrataVisualize.usf b/Engine/Shaders/Private/Strata/StrataVisualize.usf
--- a/Engine/Shaders/Private/Strata/StrataVisualize.usf
+++ b/Engine/Shaders/Private/Strata/StrataVisualize.usf
@@ -118,11 +118,13 @@ void PrintPixelType(inout FShaderPrintContext Ctx, in FStrataPixelHeader Header,
 	else									Print(Ctx, TEXT("Error - Unkown pixel type"), FontRed);
 }
 
-void PrintSSSType(inout FShaderPrintContext Ctx, uint SSSType, FFontColor InColor)
+void PrintSSSType(inout FShaderPrintContext Ctx, uint SSSType, bool bSimpleVolume, bool bThin, FFontColor InColor)
 {
 		 if (SSSType == SSS_TYPE_WRAP)				Print(Ctx, TEXT("Wrap"), InColor);
 	else if (SSSType == SSS_TYPE_DIFFUSION)			Print(Ctx, TEXT("Diffusion"), InColor);
 	else if (SSSType == SSS_TYPE_DIFFUSION_PROFILE)	Print(Ctx, TEXT("Diffusion Profile"), InColor);
+	else if (bSimpleVolume)							Print(Ctx, TEXT("Simple Volume"), InColor);
+	else if (bThin)									Print(Ctx, TEXT("Thin"), InColor);
 	else /*	(SSSType == SSS_TYPE_INVALID) */		Print(Ctx, TEXT("Invalid"), InColor);
 }
 
@@ -265,8 +267,7 @@ void PrintBSDF(inout FShaderPrintContext Context, inout FStrataAddressing Strata
 			if (BSDF_GETSSSTYPE(BSDF) != SSS_TYPE_INVALID || BSDF_GETISSIMPLEVOLUME(BSDF) || BSDF_GETISTHIN(BSDF))
 			{
 				const FStrataSubsurfaceHeader SSSHeader = StrataLoadSubsurfaceHeader(Strata.MaterialTextureArray, Strata.FirstSliceStoringStrataSSSData, StrataAddressing.PixelCoords);
-				const uint SSSType = StrataSubSurfaceHeaderGetSSSType(SSSHeader);
-				Print(Context, TEXT(" SSS Type    "), FontBSDFPropName);	PrintSSSType(Context, SSSType, FontYellow);	Newline(Context, RectMax);
+				Print(Context, TEXT(" SSS Type    "), FontBSDFPropName);	PrintSSSType(Context, BSDF_GETSSSTYPE(BSDF), BSDF_GETISSIMPLEVOLUME(BSDF), BSDF_GETISTHIN(BSDF), FontYellow);	Newline(Context, RectMax);
 				if (BSDF_GETSSSTYPE(BSDF) == SSS_TYPE_WRAP)
 				{
 					const float Opacity = StrataSubSurfaceHeaderGetWrapOpacity(SSSHeader);
@@ -638,7 +639,7 @@ void VisualizeMaterialPS(
 
 #define TREE_MAX_DEPTH			4
 #define TREE_START_POS			float2(1300, 350)
-#define TREE_BRANCH				float2(200, 75)
+#define TREE_BRANCH				float2(100, 75)
 #define TREE_COLOR				float4(0.8, 0.8, 0.8, 1.0)
 #define TREE_NODE_RADIUS		20.0
 
