Add LIGHING_CHANNEL_MASK define to ease maintenance.

diff --git a/Engine/Shaders/Private/ForwardLightingCommon.ush b/Engine/Shaders/Private/ForwardLightingCommon.ush
--- a/Engine/Shaders/Private/ForwardLightingCommon.ush
+++ b/Engine/Shaders/Private/ForwardLightingCommon.ush
@@ -243,7 +243,7 @@ FDeferredLightingSplit GetForwardDirectLightingSplit(
 	for (uint LocalLightListIndex = 0; LocalLightListIndex < NumLocalLights; LocalLightListIndex++)
 	{
 		half4 PreviewShadowMapChannelMask = 1;
-		uint  LocalLightingChannelMask = 0x7;
+		uint  LocalLightingChannelMask = LIGHING_CHANNEL_MASK;
 		const FLocalLightData LocalLight = GetLocalLightData(CulledLightsGrid.DataStartIndex + LocalLightListIndex, EyeIndex);
 		FDeferredLightData LightData = ConvertLocalToDeferredLight(LocalLight, SpecularScale, PreviewShadowMapChannelMask, LocalLightingChannelMask);
 		#if USE_HAIR_COMPLEX_TRANSMITTANCE
diff --git a/Engine/Shaders/Private/LightData.ush b/Engine/Shaders/Private/LightData.ush
--- a/Engine/Shaders/Private/LightData.ush
+++ b/Engine/Shaders/Private/LightData.ush
@@ -54,6 +54,6 @@ FDeferredLightData ConvertLocalToDeferredLight(
 FDeferredLightData ConvertLocalToDeferredLight(const FLocalLightData In)
 {
 	half4 OutPreviewShadowMapChannelMask = 0;
-	uint OutLightingChannelMask = 0x7;
+	uint OutLightingChannelMask = LIGHING_CHANNEL_MASK;
 	return ConvertLocalToDeferredLight(In, 1.0f, OutPreviewShadowMapChannelMask, OutLightingChannelMask);
 }
\ No newline at end of file
diff --git a/Engine/Shaders/Private/LightGridCommon.ush b/Engine/Shaders/Private/LightGridCommon.ush
--- a/Engine/Shaders/Private/LightGridCommon.ush
+++ b/Engine/Shaders/Private/LightGridCommon.ush
@@ -12,6 +12,7 @@
 #define LIGHT_TYPE_SPOT				2 
 #define LIGHT_TYPE_RECT				3 
 #define LIGHT_TYPE_MAX				4 
+#define LIGHING_CHANNEL_MASK		0x7
 
 #if INSTANCED_STEREO
 	#if MATERIALBLENDING_ANY_TRANSLUCENT
@@ -283,7 +284,7 @@ float4 UnpackShadowMapChannelMask(uint In)
 
 uint UnpackLightingChannelMask(uint In)
 {
-	return (In >> 8) & 0x7;
+	return (In >> 8) & LIGHING_CHANNEL_MASK;
 }
 
 uint UnpackLightType(uint ShadowMapChannelMaskPacked)
diff --git a/Engine/Shaders/Private/Strata/StrataForwardLighting.ush b/Engine/Shaders/Private/Strata/StrataForwardLighting.ush
--- a/Engine/Shaders/Private/Strata/StrataForwardLighting.ush
+++ b/Engine/Shaders/Private/Strata/StrataForwardLighting.ush
@@ -411,7 +411,7 @@ float3 StrataForwardLighting(
 					const FLocalLightData LocalLight = GetLocalLightData(CulledLightsGrid.DataStartIndex + LocalLightListIndex, EyeIndex);
 
 					half4 PreviewShadowMapChannelMask = 1;
-					uint  LightChannelMask = 0x7;
+					uint  LightChannelMask = LIGHING_CHANNEL_MASK;
 					const FDeferredLightData LightData = ConvertLocalToDeferredLight(LocalLight, 1.0f /*SpecularScale*/, PreviewShadowMapChannelMask, LightChannelMask);
 					
 					// Skip rect light. This is not done explicitely in the legacy path (resulting in no lighting) but we do it for strata in order to avoid visual artefact.
