Fix shader compilation issue.

diff --git a/Engine/Shaders/Private/VolumetricCloud.usf b/Engine/Shaders/Private/VolumetricCloud.usf
--- a/Engine/Shaders/Private/VolumetricCloud.usf
+++ b/Engine/Shaders/Private/VolumetricCloud.usf
@@ -115,6 +115,7 @@ float SampleAmbientOcclusion(in FPixelMaterialInputs PixelMaterialInputs)
 #include "LightGridCommon.ush"
 #define SUPPORT_CONTACT_SHADOWS 0
 #include "DeferredLightingCommon.ush"
+#include "LightData.ush"
 #endif
 
 
