Fix shader compilation typo.

diff --git a/Engine/Shaders/Private/VolumetricCloud.usf b/Engine/Shaders/Private/VolumetricCloud.usf
--- a/Engine/Shaders/Private/VolumetricCloud.usf
+++ b/Engine/Shaders/Private/VolumetricCloud.usf
@@ -990,7 +990,7 @@ void MainCommon(in FMaterialPixelParameters MaterialParameters, in float4 SvPosi
 
 					if (VolumetricScatteringIntensity > 0)
 					{
-						const FDeferredLightData LightData = ConvertLocalLightToDeferredLight(LocalLight);
+						const FDeferredLightData LightData = ConvertLocalToDeferredLight(LocalLight);
 
 						float3 L = 0;
 						float3 ToLight = 0;
