Remove unused source texture parameters.

diff --git a/Engine/Shaders/Private/LightShaderParameters.ush b/Engine/Shaders/Private/LightShaderParameters.ush
--- a/Engine/Shaders/Private/LightShaderParameters.ush
+++ b/Engine/Shaders/Private/LightShaderParameters.ush
@@ -2,13 +2,6 @@
 
 #pragma once
 
-
-#if RAYCALLABLESHADER || RAYGENSHADER || RAYHITGROUPSHADER || RAYMISSSHADER
-#	define SHADER_PARAMETERS_HAS_SOURCETEXTURE 1
-#else
-#	define SHADER_PARAMETERS_HAS_SOURCETEXTURE 0
-#endif
-
 // Mirrors FLightShaderParameters on C++ side.
 struct FLightShaderParameters
 {
@@ -25,9 +18,6 @@ struct FLightShaderParameters
 	float  SourceLength;
 	float  RectLightBarnCosAngle;
 	float  RectLightBarnLength;
-	#if SHADER_PARAMETERS_HAS_SOURCETEXTURE
-	Texture2D SourceTexture;
-	#endif
 };
 
 
@@ -58,7 +48,6 @@ float  Light_SoftSourceRadius;
 float  Light_SourceLength;
 float  Light_RectLightBarnCosAngle;
 float  Light_RectLightBarnLength;
-Texture2D Light_SourceTexture;
 
 
 /** Returns the FLightShaderParameters from the root shader parameters. */
@@ -79,8 +68,5 @@ FLightShaderParameters GetRootLightShaderParameters()
 	LightParameters.SourceLength = Light_SourceLength;
 	LightParameters.RectLightBarnCosAngle = Light_RectLightBarnCosAngle;
 	LightParameters.RectLightBarnLength = Light_RectLightBarnLength;
-	#if SHADER_PARAMETERS_HAS_SOURCETEXTURE
-	LightParameters.SourceTexture = Light_SourceTexture;
-	#endif
 	return LightParameters;
 }
diff --git a/Engine/Shaders/Private/ScreenSpaceDenoise/SSDSignalFramework.ush b/Engine/Shaders/Private/ScreenSpaceDenoise/SSDSignalFramework.ush
--- a/Engine/Shaders/Private/ScreenSpaceDenoise/SSDSignalFramework.ush
+++ b/Engine/Shaders/Private/ScreenSpaceDenoise/SSDSignalFramework.ush
@@ -75,11 +75,6 @@ DECLARE_SCALAR_ARRAY(uint, LightType, MAX_SIGNAL_BATCH_SIZE);
 
 uint FrameIndex;
 
-#if SHADER_PARAMETERS_HAS_SOURCETEXTURE
-Texture2D UnusedLightSourceTexture;
-#endif
-
-
 /** Returns the FLightShaderParameters from the root shader parameters. */
 FSSDSignalDomainKnowledge GetSignalDomainKnowledge(uint BatchedSignalId)
 {
@@ -91,11 +86,7 @@ FSSDSignalDomainKnowledge GetSignalDomainKnowledge(uint BatchedSignalId)
 	SignalDomain.Light.Tangent = 0;
 	SignalDomain.Light.SpotAngles = 0;
 	SignalDomain.Light.SpecularScale = 0;
-	SignalDomain.Light.SoftSourceRadius = 0;
-	#if SHADER_PARAMETERS_HAS_SOURCETEXTURE
-	SignalDomain.Light.SourceTexture = UnusedLightSourceTexture;
-	#endif
-	
+	SignalDomain.Light.SoftSourceRadius = 0;	
 	SignalDomain.Light.TranslatedWorldPosition = LightPositionAndRadius[BatchedSignalId].xyz;
 	SignalDomain.Light.Direction = LightDirectionAndLength[BatchedSignalId].xyz;
 	SignalDomain.Light.SourceRadius = LightPositionAndRadius[BatchedSignalId].w;
