Fix missing local lighting on translucent surfaces.

* Irradiance/AO data wasn't written correctly into the header in forward lighting, causing local lights to be 'occluded' (AO is used as default shadow value).
* Fix incorrect initialization of irradiance sample

diff --git a/Engine/Shaders/Private/Strata/Strata.ush b/Engine/Shaders/Private/Strata/Strata.ush
--- a/Engine/Shaders/Private/Strata/Strata.ush
+++ b/Engine/Shaders/Private/Strata/Strata.ush
@@ -604,7 +604,7 @@ FStrataIrradianceAndOcclusion InitIrradianceAndOcclusion(float InAO=1.0f)
 	FStrataIrradianceAndOcclusion Out;
 	Out.MaterialAO = 1.0f;
 	Out.IndirectIrradiance = 1.0f;
-	Out.DiffuseIndirectSampleOcclusion = 0xFFu;
+	Out.DiffuseIndirectSampleOcclusion = 0x0u;
 	return Out;
 }
 
diff --git a/Engine/Shaders/Private/Strata/StrataForwardLighting.ush b/Engine/Shaders/Private/Strata/StrataForwardLighting.ush
--- a/Engine/Shaders/Private/Strata/StrataForwardLighting.ush
+++ b/Engine/Shaders/Private/Strata/StrataForwardLighting.ush
@@ -139,6 +139,11 @@ float3 StrataForwardLighting(
 	FStrataTree StrataTree = StrataPixelHeader.StrataTree;
 	if (StrataTree.BSDFCount > 0)
 	{
+		// Write Irradiance/AO data into header in order to be retrieve correctly during lighting
+		#if STRATA_INLINE_SHADING
+		HEADER_SETIRRADIANCE_AO(StrataPixelHeader.State, StrataPackIrradianceAndOcclusion(StrataPixelHeader.IrradianceAO, 0));
+		#endif
+
 		// Update tree (coverage/transmittance/luminace weights)
 		StrataUpdateTree(NullStrataAddressing, StrataPixelHeader, StrataTree, StrataData, V, Settings, OutCoverage, OutTransmittancePreCoverage);
 
