Limits TSR's shading rejection filtering on parralax disocclusion mask discontinuity for shading rejection to more precisely pick up history rejection where parralax disocclusion heuristic fails

diff --git a/Engine/Shaders/Private/TemporalSuperResolution/TSRPostfilterRejection.usf b/Engine/Shaders/Private/TemporalSuperResolution/TSRPostfilterRejection.usf
--- a/Engine/Shaders/Private/TemporalSuperResolution/TSRPostfilterRejection.usf
+++ b/Engine/Shaders/Private/TemporalSuperResolution/TSRPostfilterRejection.usf
@@ -15,6 +15,7 @@
 uint4 HistoryRejectionViewport;
 
 Texture2D<tsr_half> HistoryRejectionTexture;
+Texture2D<tsr_half> ParallaxRejectionMaskTexture;
 RWTexture2D<tsr_half> HistoryRejectionOutput;
 
 
@@ -71,6 +72,11 @@ tsr_half Median9(tsr_half In[9])
 	return Median(M1[0], M1[1], M1[2]);
 }
 
+bool IsValidParralax(tsr_half ParallaxRejectionMask)
+{
+	return ParallaxRejectionMask > tsr_half(PARALLAX_REJECTION_MASK_THRESHOLD);
+}
+
 
 //------------------------------------------------------- ENTRY POINT
 
@@ -87,9 +93,14 @@ void MainCS(
 		Map8x8Tile2x2Lane(GroupThreadIndex));
 
 	RejectionPixelPos = ClampPixelOffset(RejectionPixelPos, HistoryRejectionViewport.xy, HistoryRejectionViewport.zw - 1);
+	
+	tsr_half CenterRejection = HistoryRejectionTexture[RejectionPixelPos];
+	tsr_half CenterParallaxRejectionMask = ParallaxRejectionMaskTexture[RejectionPixelPos];
+
+	bool bCenterIsValid = IsValidParralax(CenterParallaxRejectionMask);
 
 	// Fetch neighborhood of the kernel
-	tsr_half InSamples[9];
+	tsr_half InRejectionSamples[9];
 	UNROLL
 	for (uint i = 0; i < 9; i++)
 	{
@@ -97,11 +108,17 @@ void MainCS(
 			RejectionPixelPos,
 			tsr_short2(kOffsets3x3[i]), tsr_short2(kOffsets3x3[i]),
 			HistoryRejectionViewport.xy, HistoryRejectionViewport.zw - 1);
+			
+		tsr_half SampleRejection = HistoryRejectionTexture[SampleInputPixelPos];
+		tsr_half SampleParallaxRejectionMask = ParallaxRejectionMaskTexture[SampleInputPixelPos];
+
+		SampleRejection = select(bCenterIsValid == IsValidParralax(SampleParallaxRejectionMask), SampleRejection, CenterRejection);
 
-		InSamples[i] = HistoryRejectionTexture[SampleInputPixelPos];
+		InRejectionSamples[i] = SampleRejection;
 	}
 
-	tsr_half Rejection = Median9(InSamples);
+	tsr_half Rejection = Median9(InRejectionSamples);
+	Rejection = saturate(Rejection + 1.0 - tsr_half(bCenterIsValid));
 
 	#if CONFIG_OUTPUT_HALF_RES
 	{
diff --git a/Engine/Shaders/Private/TemporalSuperResolution/TSRRejectShading.usf b/Engine/Shaders/Private/TemporalSuperResolution/TSRRejectShading.usf
--- a/Engine/Shaders/Private/TemporalSuperResolution/TSRRejectShading.usf
+++ b/Engine/Shaders/Private/TemporalSuperResolution/TSRRejectShading.usf
@@ -702,8 +702,10 @@ void MainCS(
 			tsr_short2 OutputPixelPos = InvalidateOutputPixelPos(GroupPixelOffset + LaneSimdPixelOffset, InputInfo_ViewportMax);
 			OutputPixelPos.x = select(bIsNonOverscanPixel, OutputPixelPos.x, ~tsr_short(0));
 		
+			tsr_half LocalRejection = Rejection.GetElement(SimdIndex);
 			tsr_halfC FinalGuideColor = FinalGuide.GetElement(SimdIndex);
-
+			
+			LocalRejection = LocalRejection - 0.5 * rcp(tsr_half(256.0));
 			#if CONFIG_ENABLE_STOCASTIC_QUANTIZATION
 			{
 				uint2 Random = Rand3DPCG16(int3(LaneSimdPixelOffset, View.StateFrameIndexMod8)).xy;
@@ -717,7 +719,7 @@ void MainCS(
 			#if CONFIG_MOIRE_DETECTION
 				HistoryMoireOutput[OutputPixelPos] = FinalMoire.GetElement(SimdIndex);
 			#endif
-			HistoryRejectionOutput[OutputPixelPos] = Rejection.GetElement(SimdIndex);
+			HistoryRejectionOutput[OutputPixelPos] = LocalRejection;
 			InputSceneColorLdrLumaOutput[OutputPixelPos] = CenterLumaLDR.GetElement(SimdIndex);
 			
 			#if DEBUG_OUTPUT
diff --git a/Engine/Source/Runtime/Renderer/Private/PostProcess/TemporalSuperResolution.cpp b/Engine/Source/Runtime/Renderer/Private/PostProcess/TemporalSuperResolution.cpp
--- a/Engine/Source/Runtime/Renderer/Private/PostProcess/TemporalSuperResolution.cpp
+++ b/Engine/Source/Runtime/Renderer/Private/PostProcess/TemporalSuperResolution.cpp
@@ -487,6 +487,7 @@ class FTSRPostfilterRejectionCS : public FTSRShader
 	BEGIN_SHADER_PARAMETER_STRUCT(FParameters, )
 		SHADER_PARAMETER(FIntRect, HistoryRejectionViewport)
 		SHADER_PARAMETER_RDG_TEXTURE(Texture2D, HistoryRejectionTexture)
+		SHADER_PARAMETER_RDG_TEXTURE(Texture2D, ParallaxRejectionMaskTexture)
 		SHADER_PARAMETER_RDG_TEXTURE_UAV(RWTexture2D, HistoryRejectionOutput)
 		SHADER_PARAMETER_RDG_TEXTURE_UAV(RWTexture2DArray, DebugOutput)
 	END_SHADER_PARAMETER_STRUCT()
@@ -1484,6 +1485,7 @@ ITemporalUpscaler::FOutputs AddTemporalSuperResolutionPasses(
 		FTSRPostfilterRejectionCS::FParameters* PassParameters = GraphBuilder.AllocParameters<FTSRPostfilterRejectionCS::FParameters>();
 		PassParameters->HistoryRejectionViewport = Rect;
 		PassParameters->HistoryRejectionTexture = HistoryRejectionTexture;
+		PassParameters->ParallaxRejectionMaskTexture = ParallaxRejectionMaskTexture;
 		PassParameters->HistoryRejectionOutput = GraphBuilder.CreateUAV(FilteredHistoryRejectionTexture);
 		PassParameters->DebugOutput = CreateDebugUAV(bOutputHalfRes ? InputExtent : RejectionExtent, TEXT("Debug.TSR.PostfilterRejection"));
 
