Fix skeletal mesh previous frame local-to-translated-world position in large world coordinates for both GPUSkinVF and LocalVF PassThrough.

diff --git a/Engine/Shaders/Private/GpuSkinVertexFactory.ush b/Engine/Shaders/Private/GpuSkinVertexFactory.ush
--- a/Engine/Shaders/Private/GpuSkinVertexFactory.ush
+++ b/Engine/Shaders/Private/GpuSkinVertexFactory.ush
@@ -964,8 +964,7 @@ FVertexFactoryInterpolantsVSToPS VertexFactoryGetInterpolantsVSToPS(FVertexFacto
 
 // @return previous translated world position
 float4 VertexFactoryGetPreviousWorldPosition(FVertexFactoryInput Input, FVertexFactoryIntermediates Intermediates)
-{
-	float4x4 PreviousLocalToWorldTranslated = LWCMultiplyTranslation(Intermediates.PrevLocalToWorld, ResolvedView.PrevPreViewTranslation);
+{	
 	float3 PrevSkinnedPosition = SkinPreviousPosition(Input, Intermediates);
 #if GPUSKIN_APEX_CLOTH
 #if GPUSKIN_APEX_CLOTH_PREVIOUS
@@ -977,7 +976,7 @@ float4 VertexFactoryGetPreviousWorldPosition(FVertexFactoryInput Input, FVertexF
 	}
 #endif // GPUSKIN_APEX_CLOTH_PREVIOUS
 #endif // GPUSKIN_APEX_CLOTH
-	return mul(float4(PrevSkinnedPosition,1), PreviousLocalToWorldTranslated);
+	return TransformPreviousLocalPositionToTranslatedWorld(PrevSkinnedPosition, Intermediates.PrevLocalToWorld);
 }
 
 
diff --git a/Engine/Shaders/Private/LocalVertexFactory.ush b/Engine/Shaders/Private/LocalVertexFactory.ush
--- a/Engine/Shaders/Private/LocalVertexFactory.ush
+++ b/Engine/Shaders/Private/LocalVertexFactory.ush
@@ -1150,8 +1150,8 @@ Buffer<float> GPUSkinCachePreviousPositionBuffer;
 float4 VertexFactoryGetPreviousWorldPosition(FVertexFactoryInput Input, FVertexFactoryIntermediates Intermediates)
 {
 	FLWCMatrix PreviousLocalToWorld = GetInstanceData(Intermediates).PrevLocalToWorld;
-	float4x4 PreviousLocalToWorldTranslated = LWCMultiplyTranslation(PreviousLocalToWorld, ResolvedView.PrevPreViewTranslation);
 
+	float4 PrevLocalPosition;
 #if SUPPORT_GPUSKIN_PASSTHROUGH
 	if (IsGPUSkinPassThrough())
 	{
@@ -1160,26 +1160,26 @@ float4 VertexFactoryGetPreviousWorldPosition(FVertexFactoryInput Input, FVertexF
 		PreviousPos.x = GPUSkinCachePreviousPositionBuffer[Offset + 0];
 		PreviousPos.y = GPUSkinCachePreviousPositionBuffer[Offset + 1];
 		PreviousPos.z = GPUSkinCachePreviousPositionBuffer[Offset + 2];
-		return mul(float4(PreviousPos, 1), PreviousLocalToWorldTranslated);	
+		PrevLocalPosition = float4(PreviousPos, 1);	
 	}
 	else
 #endif // SUPPORT_GPUSKIN_PASSTHROUGH
 	{
 	#if USE_INSTANCING
 		float4x4 InstanceTransform = GetInstancePrevTransform(Intermediates);
-		return mul(mul(Input.Position, InstanceTransform), PreviousLocalToWorldTranslated);
+		PrevLocalPosition = mul(Input.Position, InstanceTransform);
 	#elif USE_SPLINEDEFORM
 		// Just like CalcWorldPosition...
 		float4x3 SliceTransform = CalcSliceTransform(dot(Input.Position.xyz, SplineMeshDir));
 	
 		// Transform into mesh space
-		float4 LocalPos = float4(mul(Input.Position, SliceTransform), Input.Position.w);
-	
-		return mul(LocalPos, PreviousLocalToWorldTranslated);
+		PrevLocalPosition = float4(mul(Input.Position, SliceTransform), Input.Position.w);
 	#else
-		return mul(Input.Position, PreviousLocalToWorldTranslated);
+		PrevLocalPosition = Input.Position;
 	#endif	// USE_INSTANCING
 	}
+
+	return TransformPreviousLocalPositionToTranslatedWorld(PrevLocalPosition.xyz, PreviousLocalToWorld);
 }
 
 #if NEEDS_VERTEX_FACTORY_INTERPOLATION
diff --git a/Engine/Shaders/Private/VertexFactoryCommon.ush b/Engine/Shaders/Private/VertexFactoryCommon.ush
--- a/Engine/Shaders/Private/VertexFactoryCommon.ush
+++ b/Engine/Shaders/Private/VertexFactoryCommon.ush
@@ -134,6 +134,18 @@ float4 TransformLocalToTranslatedWorld(float3 LocalPosition, float4x4 LocalToWor
 	return TransformLocalToTranslatedWorld(LocalPosition, LWCPromote(LocalToWorld));
 }
 
+float4 TransformPreviousLocalPositionToTranslatedWorld(float3 PrevLocalPosition, FLWCMatrix PrevLocalToWorld)
+{
+	FLWCVector3 PrevWorldPosition = LWCMultiply(PrevLocalPosition, PrevLocalToWorld);
+	FLWCVector3 PrevTranslatedWorldPosition = LWCAdd(PrevWorldPosition, ResolvedView.PrevPreViewTranslation);
+	return float4(LWCToFloat(PrevTranslatedWorldPosition), 1.0f);
+}
+
+float4 TransformPreviousLocalPositionToTranslatedWorld(float3 PrevLocalPosition, float4x4 PrevLocalToWorld)
+{
+	return TransformPreviousLocalPositionToTranslatedWorld(PrevLocalPosition, LWCPromote(PrevLocalToWorld));
+}
+
 float3 RotateLocalToWorld(float3 LocalDirection, float3x3 LocalToWorld, float3 InvScale)
 {
 	return
