Increases rejection threshold at which TSR's history footprint size is higher.

diff --git a/Engine/Shaders/Private/TemporalSuperResolution/TSRUpdateHistory.usf b/Engine/Shaders/Private/TemporalSuperResolution/TSRUpdateHistory.usf
--- a/Engine/Shaders/Private/TemporalSuperResolution/TSRUpdateHistory.usf
+++ b/Engine/Shaders/Private/TemporalSuperResolution/TSRUpdateHistory.usf
@@ -758,7 +758,7 @@ void MainCS(
 
 			bool2 bIsRefining = PrevHistoryValidity != 0.0;
 
-			tsr_half2 KernelInputToHistoryLerp = select(or(or(bIsOffScreen, bIsDisoccluded), !bIsRefining), tsr_half(0.0).xx, saturate(LowFrequencyRejection * tsr_half(16.0) - tsr_half(3.0)));
+			tsr_half2 KernelInputToHistoryLerp = select(or(or(bIsOffScreen, bIsDisoccluded), !bIsRefining), tsr_half(0.0).xx, saturate(LowFrequencyRejection * tsr_half(16.0) - tsr_half(13.0)));
 			tsr_half2 KernelInputToHistoryFactor = lerp(OffScreenInputToHistoryFactor, tsr_half(InputToHistoryFactor), KernelInputToHistoryLerp);
 			tsr_half2 KernelInputToHistoryAlignmentFactor = lerp(tsr_half(1.0).xx, tsr_half(InputToHistoryFactor), KernelInputToHistoryLerp);
 			
