* Bumped GUID in ShaderVersion.ush to fix shader invalidation issue on all platforms (similar to 22136421), which generates crashes on load for several users
* Fixed warning when compiling materials with instancing

#rnx

diff --git a/Engine/Shaders/Private/LocalVertexFactory.ush b/Engine/Shaders/Private/LocalVertexFactory.ush
--- a/Engine/Shaders/Private/LocalVertexFactory.ush
+++ b/Engine/Shaders/Private/LocalVertexFactory.ush
@@ -101,31 +101,31 @@ struct FVertexFactoryInput
 
 #if NUM_MATERIAL_TEXCOORDS_VERTEX
 	#if !MANUAL_VERTEX_FETCH
-		// These used to be packed texcoord arrays, but these cause problems with alighnment on some Vulkan drivers
-		#if NUM_MATERIAL_TEXCOORDS_VERTEX > 1
-			float4	TexCoords0 : ATTRIBUTE4;
-		#elif NUM_MATERIAL_TEXCOORDS_VERTEX == 1
-			float2	TexCoords0 : ATTRIBUTE4;
+			// These used to be packed texcoord arrays, but these cause problems with alighnment on some Vulkan drivers
+			#if NUM_MATERIAL_TEXCOORDS_VERTEX > 1
+				float4	TexCoords0 : ATTRIBUTE4;
+			#elif NUM_MATERIAL_TEXCOORDS_VERTEX == 1
+				float2	TexCoords0 : ATTRIBUTE4;
+			#endif
+
+			#if NUM_MATERIAL_TEXCOORDS_VERTEX > 3
+				float4	TexCoords1 : ATTRIBUTE5;
+			#elif NUM_MATERIAL_TEXCOORDS_VERTEX == 3
+				float2	TexCoords1 : ATTRIBUTE5;
+			#endif
+
+			#if NUM_MATERIAL_TEXCOORDS_VERTEX > 5
+				float4	TexCoords2 : ATTRIBUTE6;
+			#elif NUM_MATERIAL_TEXCOORDS_VERTEX == 5
+				float2	TexCoords2 : ATTRIBUTE6;
+			#endif
+
+			#if NUM_MATERIAL_TEXCOORDS_VERTEX > 7
+				float4	TexCoords3 : ATTRIBUTE7;
+			#elif NUM_MATERIAL_TEXCOORDS_VERTEX == 7
+				float2	TexCoords3 : ATTRIBUTE7;
+			#endif
 		#endif
-
-		#if NUM_MATERIAL_TEXCOORDS_VERTEX > 3
-			float4	TexCoords1 : ATTRIBUTE5;
-		#elif NUM_MATERIAL_TEXCOORDS_VERTEX == 3
-			float2	TexCoords1 : ATTRIBUTE5;
-		#endif
-
-		#if NUM_MATERIAL_TEXCOORDS_VERTEX > 5
-			float4	TexCoords2 : ATTRIBUTE6;
-		#elif NUM_MATERIAL_TEXCOORDS_VERTEX == 5
-			float2	TexCoords2 : ATTRIBUTE6;
-		#endif
-
-		#if NUM_MATERIAL_TEXCOORDS_VERTEX > 7
-			float4	TexCoords3 : ATTRIBUTE7;
-		#elif NUM_MATERIAL_TEXCOORDS_VERTEX == 7
-			float2	TexCoords3 : ATTRIBUTE7;
-		#endif
-	#endif
 #elif USE_PARTICLE_SUBUVS && !MANUAL_VERTEX_FETCH
 	float2	TexCoords0 : ATTRIBUTE4;
 #endif
@@ -621,30 +621,30 @@ FMaterialVertexParameters GetMaterialVertexParameters(FVertexFactoryInput Input,
 			Result.TexCoords[CoordinateIndex] = LocalVF.VertexFetch_TexCoordBuffer[NumFetchTexCoords * (LocalVF.VertexFetch_Parameters[VF_VertexOffset] + Input.VertexId) + ClampedCoordinateIndex];
 		}
 #elif NUM_MATERIAL_TEXCOORDS_VERTEX
-	#if NUM_MATERIAL_TEXCOORDS_VERTEX > 0
-		Result.TexCoords[0] = Input.TexCoords0.xy;
-	#endif
-	#if NUM_MATERIAL_TEXCOORDS_VERTEX > 1
-		Result.TexCoords[1] = Input.TexCoords0.zw;
-	#endif
-	#if NUM_MATERIAL_TEXCOORDS_VERTEX > 2
-		Result.TexCoords[2] = Input.TexCoords1.xy;
-	#endif
-	#if NUM_MATERIAL_TEXCOORDS_VERTEX > 3
-		Result.TexCoords[3] = Input.TexCoords1.zw;
-	#endif
-	#if NUM_MATERIAL_TEXCOORDS_VERTEX > 4
-		Result.TexCoords[4] = Input.TexCoords2.xy;
-	#endif
-	#if NUM_MATERIAL_TEXCOORDS_VERTEX > 5
-		Result.TexCoords[5] = Input.TexCoords2.zw;
-	#endif
-	#if NUM_MATERIAL_TEXCOORDS_VERTEX > 6
-		Result.TexCoords[6] = Input.TexCoords3.xy;
-	#endif
-	#if NUM_MATERIAL_TEXCOORDS_VERTEX > 7
-		Result.TexCoords[7] = Input.TexCoords3.zw;
-	#endif
+			#if NUM_MATERIAL_TEXCOORDS_VERTEX > 0
+				Result.TexCoords[0] = Input.TexCoords0.xy;
+			#endif
+			#if NUM_MATERIAL_TEXCOORDS_VERTEX > 1
+				Result.TexCoords[1] = Input.TexCoords0.zw;
+			#endif
+			#if NUM_MATERIAL_TEXCOORDS_VERTEX > 2
+				Result.TexCoords[2] = Input.TexCoords1.xy;
+			#endif
+			#if NUM_MATERIAL_TEXCOORDS_VERTEX > 3
+				Result.TexCoords[3] = Input.TexCoords1.zw;
+			#endif
+			#if NUM_MATERIAL_TEXCOORDS_VERTEX > 4
+				Result.TexCoords[4] = Input.TexCoords2.xy;
+			#endif
+			#if NUM_MATERIAL_TEXCOORDS_VERTEX > 5
+				Result.TexCoords[5] = Input.TexCoords2.zw;
+			#endif
+			#if NUM_MATERIAL_TEXCOORDS_VERTEX > 6
+				Result.TexCoords[6] = Input.TexCoords3.xy;
+			#endif
+			#if NUM_MATERIAL_TEXCOORDS_VERTEX > 7
+				Result.TexCoords[7] = Input.TexCoords3.zw;
+			#endif
 #endif  //MANUAL_VERTEX_FETCH && NUM_MATERIAL_TEXCOORDS_VERTEX
 
 	Result.Particle.Color = half4(1,1,1,1);
@@ -958,7 +958,7 @@ FVertexFactoryIntermediates GetVertexFactoryIntermediates(FVertexFactoryInput In
 	}
 	else
 	{
-		Intermediates.PreSkinPosition = Input.Position.xyz;	
+	Intermediates.PreSkinPosition = Input.Position.xyz;
 	}
 
 	return Intermediates;
@@ -1155,11 +1155,11 @@ float4 VertexFactoryGetPreviousWorldPosition(FVertexFactoryInput Input, FVertexF
 #if SUPPORT_GPUSKIN_PASSTHROUGH
 	if (IsGPUSkinPassThrough())
 	{
-		uint Offset = Input.VertexId * 3;
-		float3 PreviousPos;
-		PreviousPos.x = GPUSkinCachePreviousPositionBuffer[Offset + 0];
-		PreviousPos.y = GPUSkinCachePreviousPositionBuffer[Offset + 1];
-		PreviousPos.z = GPUSkinCachePreviousPositionBuffer[Offset + 2];
+	uint Offset = Input.VertexId * 3;
+	float3 PreviousPos;
+	PreviousPos.x = GPUSkinCachePreviousPositionBuffer[Offset + 0];
+	PreviousPos.y = GPUSkinCachePreviousPositionBuffer[Offset + 1];
+	PreviousPos.z = GPUSkinCachePreviousPositionBuffer[Offset + 2];
 		PrevLocalPosition = float4(PreviousPos, 1);	
 	}
 	else
@@ -1168,16 +1168,16 @@ float4 VertexFactoryGetPreviousWorldPosition(FVertexFactoryInput Input, FVertexF
 	#if USE_INSTANCING
 		float4x4 InstanceTransform = GetInstancePrevTransform(Intermediates);
 		PrevLocalPosition = mul(Input.Position, InstanceTransform);
-	#elif USE_SPLINEDEFORM
-		// Just like CalcWorldPosition...
-		float4x3 SliceTransform = CalcSliceTransform(dot(Input.Position.xyz, SplineMeshDir));
-	
-		// Transform into mesh space
+#elif USE_SPLINEDEFORM
+	// Just like CalcWorldPosition...
+	float4x3 SliceTransform = CalcSliceTransform(dot(Input.Position.xyz, SplineMeshDir));
+
+	// Transform into mesh space
 		PrevLocalPosition = float4(mul(Input.Position, SliceTransform), Input.Position.w);
-	#else
+#else
 		PrevLocalPosition = Input.Position;
-	#endif	// USE_INSTANCING
-	}
+#endif	// USE_INSTANCING
+}
 
 	return TransformPreviousLocalPositionToTranslatedWorld(PrevLocalPosition.xyz, PreviousLocalToWorld);
 }
diff --git a/Engine/Shaders/Public/ShaderVersion.ush b/Engine/Shaders/Public/ShaderVersion.ush
--- a/Engine/Shaders/Public/ShaderVersion.ush
+++ b/Engine/Shaders/Public/ShaderVersion.ush
@@ -2,4 +2,4 @@
 // Each time the console command is executed it generates a new GUID. As this file is included
 // in Platform.ush (which should be included in any shader) it allows to invalidate the shader DDC.
 // 
-// GUID = 7553B411-6440-4492-9C57-1CD175A19DB3
+// GUID = CE2E5BA6-4FFF-4321-9A10-4A3BCEB06687
\ No newline at end of file
