Disabling Rect lights in Forward shading temporarily as it is causing some shaders to go above the max sampler count.

#rnx

diff --git a/Engine/Shaders/Private/ForwardLightingCommon.ush b/Engine/Shaders/Private/ForwardLightingCommon.ush
--- a/Engine/Shaders/Private/ForwardLightingCommon.ush
+++ b/Engine/Shaders/Private/ForwardLightingCommon.ush
@@ -246,6 +246,7 @@ FDeferredLightingSplit GetForwardDirectLightingSplit(
 		uint  LocalLightingChannelMask = LIGHING_CHANNEL_MASK;
 		const FLocalLightData LocalLight = GetLocalLightData(CulledLightsGrid.DataStartIndex + LocalLightListIndex, EyeIndex);
 		FDeferredLightData LightData = ConvertLocalToDeferredLight(LocalLight, SpecularScale, PreviewShadowMapChannelMask, LocalLightingChannelMask);
+		LightData.bRectLight = false;	// Rect light not supported in forward. Enabling this is causing the max sampler count to be exceeded.
 		#if USE_HAIR_COMPLEX_TRANSMITTANCE
 		if (GBufferData.ShadingModelID == SHADINGMODELID_HAIR)
 		{
