Mobile: Light map artifacts with specific mesh when launched on Android
Supply a mesh LOD index into MDC batching CS so it can fetch a correct lighting information for each instance

diff --git a/Engine/Shaders/Private/InstanceCulling/BuildInstanceDrawCommands.usf b/Engine/Shaders/Private/InstanceCulling/BuildInstanceDrawCommands.usf
--- a/Engine/Shaders/Private/InstanceCulling/BuildInstanceDrawCommands.usf
+++ b/Engine/Shaders/Private/InstanceCulling/BuildInstanceDrawCommands.usf
@@ -33,7 +33,7 @@ uint GPUSceneNumInstances;
 uint GPUSceneNumPrimitives;
 uint GPUSceneNumLightmapDataItems;
 
-StructuredBuffer<FDrawCommandDesc> DrawCommandDescs;
+StructuredBuffer<uint> DrawCommandDescs;
 Buffer<uint> InstanceIdOffsetBuffer;
 StructuredBuffer<uint> ViewIds;
 
@@ -95,7 +95,7 @@ float2 PackLocalBoundsExtent(float3 Extent)
 	return Result;
 }
 
-void WritePackedInstanceData(uint Offset, uint InstanceId, FInstanceSceneData InstanceData, uint ViewIdIndex)
+void WritePackedInstanceData(uint Offset, uint InstanceId, FInstanceSceneData InstanceData, uint ViewIdIndex, uint MeshLODIndex)
 {
 	FPrimitiveSceneData PrimitiveData = GetPrimitiveData(InstanceData.PrimitiveId);
 	
@@ -104,8 +104,7 @@ void WritePackedInstanceData(uint Offset, uint InstanceId, FInstanceSceneData In
 #if ALLOW_STATIC_LIGHTING 
 	float4 LightMapUVScaleBias = float4(1,1,0,0);
 	float4 ShadowMapUVScaleBias = float4(1,1,0,0);
-	// FIXME: need LODIndex to compute correct LightmapDataIndex
-	uint LightmapDataIndex = PrimitiveData.LightmapDataIndex; // + LODIndex;
+	uint LightmapDataIndex = PrimitiveData.LightmapDataIndex + MeshLODIndex;
 	if (LightmapDataIndex < GPUSceneNumLightmapDataItems)
 	{
 		LightMapUVScaleBias = GetLightmapData(LightmapDataIndex).LightMapCoordinateScaleBias;
@@ -122,7 +121,7 @@ void WritePackedInstanceData(uint Offset, uint InstanceId, FInstanceSceneData In
 		LightmapDataIndex = 0;
 	}
 	InstanceAuxData4 = PackLightmapData(LightMapUVScaleBias, ShadowMapUVScaleBias);
-	InstanceAuxData1 = asfloat(LightmapDataIndex);
+	InstanceAuxData1 = asfloat(PrimitiveData.LightmapDataIndex);
 #else
 	InstanceAuxData4.xy = PackLocalBoundsCenter(InstanceData.LocalBoundsCenter);
 	InstanceAuxData4.zw = PackLocalBoundsExtent(InstanceData.LocalBoundsExtent);
@@ -151,12 +150,12 @@ void WritePackedInstanceData(uint Offset, uint InstanceId, FInstanceSceneData In
 }
 #endif
 
-void WriteInstance(uint Offset, uint InstanceId, FInstanceSceneData InstanceData, uint ViewIdIndex)
+void WriteInstance(uint Offset, uint InstanceId, FInstanceSceneData InstanceData, uint ViewIdIndex, uint MeshLODIndex)
 {
 	checkSlow(InstanceId < GPUSceneNumInstances);
 
 #if USES_PACKED_INSTANCE_DATA
-	WritePackedInstanceData(Offset, InstanceId, InstanceData, ViewIdIndex);
+	WritePackedInstanceData(Offset, InstanceId, InstanceData, ViewIdIndex, MeshLODIndex);
 #else
 	uint PackedId = InstanceId | (ViewIdIndex << 28U);
 	checkStructuredBufferAccessSlow(InstanceIdsBufferOut, Offset);
@@ -295,7 +294,7 @@ void InstanceCullBuildInstanceIdBufferCS(uint3 GroupId : SV_GroupID, int GroupTh
 	}
 
 	// Load auxiliary per-instanced-draw command info
-	const FDrawCommandDesc DrawCommandDesc = DrawCommandDescs[Payload.IndirectArgIndex];
+	const FDrawCommandDesc DrawCommandDesc = UnpackDrawCommandDesc(DrawCommandDescs[Payload.IndirectArgIndex]);
 
 	// Extract compaction data (if applicable)	
 #if ENABLE_INSTANCE_COMPACTION
@@ -341,8 +340,8 @@ void InstanceCullBuildInstanceIdBufferCS(uint3 GroupId : SV_GroupID, int GroupTh
 	{
 		uint OutputOffset;
 		InterlockedAdd(DrawIndirectArgsBufferOut[Payload.IndirectArgIndex * INDIRECT_ARGS_NUM_WORDS + 1], 2U, OutputOffset);				
-		WriteInstance(InstanceDataOutputOffset + OutputOffset + 0U, InstanceId, InstanceData, 0U);
-		WriteInstance(InstanceDataOutputOffset + OutputOffset + 1U, InstanceId, InstanceData, 1U);
+		WriteInstance(InstanceDataOutputOffset + OutputOffset + 0U, InstanceId, InstanceData, 0U, DrawCommandDesc.MeshLODIndex);
+		WriteInstance(InstanceDataOutputOffset + OutputOffset + 1U, InstanceId, InstanceData, 1U, DrawCommandDesc.MeshLODIndex);
 	}
 #else // !STEREO_CULLING_MODE
 
@@ -366,7 +365,7 @@ void InstanceCullBuildInstanceIdBufferCS(uint3 GroupId : SV_GroupID, int GroupTh
 			uint OutputOffset;
 			InterlockedAdd(DrawIndirectArgsBufferOut[Payload.IndirectArgIndex * INDIRECT_ARGS_NUM_WORDS + 1], 1U, OutputOffset);
 
-			WriteInstance(InstanceDataOutputOffset + OutputOffset, InstanceId, InstanceData, ViewIdIndex);
+			WriteInstance(InstanceDataOutputOffset + OutputOffset, InstanceId, InstanceData, ViewIdIndex, DrawCommandDesc.MeshLODIndex);
 		}
 	}
 #endif // STEREO_CULLING_MODE
diff --git a/Engine/Shaders/Private/InstanceCulling/InstanceCullingCommon.ush b/Engine/Shaders/Private/InstanceCulling/InstanceCullingCommon.ush
--- a/Engine/Shaders/Private/InstanceCulling/InstanceCullingCommon.ush
+++ b/Engine/Shaders/Private/InstanceCulling/InstanceCullingCommon.ush
@@ -14,9 +14,18 @@
 
 struct FDrawCommandDesc
 {
-	uint bMaterialUsesWorldPositionOffset;
+	bool bMaterialUsesWorldPositionOffset;
+	uint MeshLODIndex;
 };
 
+FDrawCommandDesc UnpackDrawCommandDesc(uint PackedData)
+{
+	FDrawCommandDesc Desc;
+	Desc.bMaterialUsesWorldPositionOffset = (PackedData & 1U) != 0;
+	Desc.MeshLODIndex = (PackedData >> 1U) & 0x000000FFU;
+	return Desc;
+}
+
 struct FContextBatchInfo
 {
 	uint IndirectArgOffset;
diff --git a/Engine/Shaders/Private/VirtualShadowMaps/VirtualShadowMapBuildPerPageDrawCommands.usf b/Engine/Shaders/Private/VirtualShadowMaps/VirtualShadowMapBuildPerPageDrawCommands.usf
--- a/Engine/Shaders/Private/VirtualShadowMaps/VirtualShadowMapBuildPerPageDrawCommands.usf
+++ b/Engine/Shaders/Private/VirtualShadowMaps/VirtualShadowMapBuildPerPageDrawCommands.usf
@@ -55,7 +55,7 @@ float3 CullingViewOriginTile;
 
 float MaxMaterialPositionInvalidationRange;
 
-StructuredBuffer<FDrawCommandDesc> DrawCommandDescs;
+StructuredBuffer<uint> DrawCommandDescs;
 
 RWBuffer<uint> DrawIndirectArgsBufferOut;
 
@@ -250,7 +250,7 @@ void CullPerPageDrawCommandsCs(uint3 GroupId : SV_GroupID, int GroupThreadIndex
 		}
 
 		uint InstanceId = InstanceDataOffset + uint(WorkSetup.LocalItemIndex);
-		FDrawCommandDesc DrawCommandDesc = DrawCommandDescs[Payload.IndirectArgIndex];
+		FDrawCommandDesc DrawCommandDesc = UnpackDrawCommandDesc(DrawCommandDescs[Payload.IndirectArgIndex]);
 
 		// Load relevant instance data
 		FInstanceSceneData InstanceData = GetInstanceSceneData(InstanceId, InstanceSceneDataSOAStride);
diff --git a/Engine/Source/Runtime/Renderer/Private/InstanceCulling/InstanceCullingContext.cpp b/Engine/Source/Runtime/Renderer/Private/InstanceCulling/InstanceCullingContext.cpp
--- a/Engine/Source/Runtime/Renderer/Private/InstanceCulling/InstanceCullingContext.cpp
+++ b/Engine/Source/Runtime/Renderer/Private/InstanceCulling/InstanceCullingContext.cpp
@@ -51,6 +51,13 @@ static bool IsInstanceOrderPreservationAllowed(ERHIFeatureLevel::Type FeatureLev
 	return GInstanceCullingAllowOrderPreservation && FeatureLevel > ERHIFeatureLevel::ES3_1;
 }
 
+static uint32 PackDrawCommandDesc(bool bMaterialUsesWorldPositionOffset, uint32 MeshLODIndex)
+{
+	uint32 PackedData = bMaterialUsesWorldPositionOffset ? 1U : 0U;
+	PackedData |= ((MeshLODIndex & 0x000000FFU) << 1U);
+	return PackedData;
+}
+
 FMeshDrawCommandOverrideArgs GetMeshDrawCommandOverrideArgs(const FInstanceCullingDrawParams& InstanceCullingDrawParams)
 {
 	FMeshDrawCommandOverrideArgs Result;
@@ -373,7 +380,7 @@ public:
 
 		SHADER_PARAMETER_STRUCT_INCLUDE(FInstanceProcessingGPULoadBalancer::FShaderParameters, LoadBalancerParameters)
 
-		SHADER_PARAMETER_RDG_BUFFER_SRV(StructuredBuffer< FInstanceCullingContext::FDrawCommandDesc >, DrawCommandDescs)
+		SHADER_PARAMETER_RDG_BUFFER_SRV(StructuredBuffer< uint32 >, DrawCommandDescs)
 		SHADER_PARAMETER_RDG_BUFFER_SRV(StructuredBuffer< FInstanceCullingContext::FPayloadData >, InstanceCullingPayloads)
 		SHADER_PARAMETER_RDG_BUFFER_SRV(StructuredBuffer< uint32 >, ViewIds)
 		SHADER_PARAMETER_RDG_BUFFER_SRV(StructuredBuffer< Nanite::FPackedView >, InViews)
@@ -1110,6 +1117,17 @@ void FInstanceCullingContext::AddClearIndirectArgInstanceCountPass(FRDGBuilder&
 	}
 }
 
+void FInstanceCullingContext::SetupDrawCommands(
+	FMeshCommandOneFrameArray& VisibleMeshDrawCommandsInOut,
+	bool bCompactIdenticalCommands,
+	int32& MaxInstances,
+	int32& VisibleMeshDrawCommandsNum,
+	int32& NewPassVisibleMeshDrawCommandsNum)
+{
+	TArrayView<const FStateBucketAuxData> StateBucketsAuxData;
+	SetupDrawCommands(StateBucketsAuxData, VisibleMeshDrawCommandsInOut, bCompactIdenticalCommands, MaxInstances, VisibleMeshDrawCommandsNum, NewPassVisibleMeshDrawCommandsNum);
+}
+
 /**
  * Allocate indirect arg slots for all meshes to use instancing,
  * add commands that populate the indirect calls and index & id buffers, and
@@ -1117,6 +1135,7 @@ void FInstanceCullingContext::AddClearIndirectArgInstanceCountPass(FRDGBuilder&
  * NOTE: VisibleMeshDrawCommandsInOut can only become shorter.
  */
 void FInstanceCullingContext::SetupDrawCommands(
+	TArrayView<const FStateBucketAuxData> StateBucketsAuxData,
 	FMeshCommandOneFrameArray& VisibleMeshDrawCommandsInOut,
 	bool bCompactIdenticalCommands,
 	// Stats
@@ -1213,8 +1232,14 @@ void FInstanceCullingContext::SetupDrawCommands(
 				DrawCmd.bUseIndirect = bUseIndirectDraw;
 				
 				CurrentIndirectArgsOffset = AllocateIndirectArgs(MeshDrawCommand);
-				DrawCommandDescs.Emplace(FDrawCommandDesc{bMaterialUsesWorldPositionOffset });
-
+				
+				uint32 MeshLODIndex = 0;
+				if (StateBucketsAuxData.IsValidIndex(VisibleMeshDrawCommand.StateBucketId))
+				{
+					MeshLODIndex = StateBucketsAuxData[VisibleMeshDrawCommand.StateBucketId].MeshLODIndex;
+				}
+				DrawCommandDescs.Add(PackDrawCommandDesc(bMaterialUsesWorldPositionOffset, MeshLODIndex));
+				
 				if (bUseIndirectDraw)
 				{
 					DrawCmd.IndirectArgsOffsetOrNumInstances = CurrentIndirectArgsOffset * FInstanceCullingContext::IndirectArgsNumWords * sizeof(uint32);
diff --git a/Engine/Source/Runtime/Renderer/Private/InstanceCulling/InstanceCullingMergedContext.h b/Engine/Source/Runtime/Renderer/Private/InstanceCulling/InstanceCullingMergedContext.h
--- a/Engine/Source/Runtime/Renderer/Private/InstanceCulling/InstanceCullingMergedContext.h
+++ b/Engine/Source/Runtime/Renderer/Private/InstanceCulling/InstanceCullingMergedContext.h
@@ -46,7 +46,7 @@ public:
 	TArray<int32, SceneRenderingAllocator> ViewIds;
 	//TArray<FMeshDrawCommandInfo, SceneRenderingAllocator> MeshDrawCommandInfos;
 	TArray<FRHIDrawIndexedIndirectParameters, SceneRenderingAllocator> IndirectArgs;
-	TArray<FInstanceCullingContext::FDrawCommandDesc, SceneRenderingAllocator> DrawCommandDescs;
+	TArray<uint32, SceneRenderingAllocator> DrawCommandDescs;
 	TArray<FInstanceCullingContext::FPayloadData, SceneRenderingAllocator> PayloadData;
 	TArray<uint32, SceneRenderingAllocator> InstanceIdOffsets;
 	TArray<FInstanceCullingContext::FCompactionData, SceneRenderingAllocator> DrawCommandCompactionData;
diff --git a/Engine/Source/Runtime/Renderer/Private/MeshDrawCommands.cpp b/Engine/Source/Runtime/Renderer/Private/MeshDrawCommands.cpp
--- a/Engine/Source/Runtime/Renderer/Private/MeshDrawCommands.cpp
+++ b/Engine/Source/Runtime/Renderer/Private/MeshDrawCommands.cpp
@@ -940,7 +940,13 @@ public:
 
 			if (Context.bUseGPUScene)
 			{
-				Context.InstanceCullingContext.SetupDrawCommands(Context.MeshDrawCommands, true, Context.MaxInstances, Context.VisibleMeshDrawCommandsNum, Context.NewPassVisibleMeshDrawCommandsNum);
+				Context.InstanceCullingContext.SetupDrawCommands(
+					Context.Scene->CachedStateBucketsAuxData[Context.PassType],
+					Context.MeshDrawCommands, 
+					true, 
+					Context.MaxInstances, 
+					Context.VisibleMeshDrawCommandsNum, 
+					Context.NewPassVisibleMeshDrawCommandsNum);
 			}
 		}
 	}
diff --git a/Engine/Source/Runtime/Renderer/Private/MeshPassProcessor.cpp b/Engine/Source/Runtime/Renderer/Private/MeshPassProcessor.cpp
--- a/Engine/Source/Runtime/Renderer/Private/MeshPassProcessor.cpp
+++ b/Engine/Source/Runtime/Renderer/Private/MeshPassProcessor.cpp
@@ -1919,6 +1919,7 @@ bool FMeshPassProcessor::ShouldSkipMeshDrawCommand(const FMeshBatch& RESTRICT Me
 FCachedPassMeshDrawListContext::FCachedPassMeshDrawListContext(FScene& InScene)
 	: Scene(InScene)
 	, bUseGPUScene(UseGPUScene(GMaxRHIShaderPlatform, GMaxRHIFeatureLevel))
+	, bUseStateBucketsAuxData(InScene.GetShadingPath() == EShadingPath::Mobile)
 {
 }
 
@@ -2057,6 +2058,13 @@ void FCachedPassMeshDrawListContextImmediate::FinalizeCommand(
 		}
 
 		CommandInfo.StateBucketId = SetId.GetIndex();
+
+		if (bUseStateBucketsAuxData)
+		{
+			// grow MDC AuxData array in sync with MDC StateBuckets
+			Scene.CachedStateBucketsAuxData[CurrMeshPass].SetNum(BucketMap.GetMaxIndex() + 1, false);
+			Scene.CachedStateBucketsAuxData[CurrMeshPass][CommandInfo.StateBucketId] = FStateBucketAuxData(MeshBatch);
+		}
 	}
 	else
 	{
@@ -2102,6 +2110,11 @@ void FCachedPassMeshDrawListContextDeferred::FinalizeCommand(
 		DeferredCommandHashes.Add(FStateBucketMap::ComputeHash(MeshDrawCommand));
 
 		CommandInfo.StateBucketId = Index;
+
+		if (bUseStateBucketsAuxData)
+		{
+			DeferredStateBucketsAuxData.Add(FStateBucketAuxData(MeshBatch));
+		}
 	}
 	else
 	{
@@ -2120,17 +2133,25 @@ void FCachedPassMeshDrawListContextDeferred::DeferredFinalizeMeshDrawCommands(co
 			{				
 				check(CmdInfo.MeshPass < EMeshPass::Num);
 				FStateBucketMap& BucketMap = Scene.CachedMeshDrawCommandStateBuckets[CmdInfo.MeshPass];
+				int32 DeferredIndex = CmdInfo.StateBucketId;
 
-				check(CmdInfo.StateBucketId >= 0 && CmdInfo.StateBucketId < DeferredCommands.Num());
+				check(DeferredIndex >= 0 && DeferredIndex < DeferredCommands.Num());
 				check(CmdInfo.CommandIndex == -1);
-				FMeshDrawCommand& Command = DeferredCommands[CmdInfo.StateBucketId];
-				const Experimental::FHashType CommandHash = DeferredCommandHashes[CmdInfo.StateBucketId];
+				FMeshDrawCommand& Command = DeferredCommands[DeferredIndex];
+				const Experimental::FHashType CommandHash = DeferredCommandHashes[DeferredIndex];
 
 				Experimental::FHashElementId SetId = BucketMap.FindOrAddIdByHash(CommandHash, MoveTemp(Command), FMeshDrawCommandCount());
 				FMeshDrawCommandCount& DrawCount = BucketMap.GetByElementId(SetId).Value;
 				DrawCount.Num++;
 
 				CmdInfo.StateBucketId = SetId.GetIndex();
+
+				if (bUseStateBucketsAuxData)
+				{
+					// grow MDC AuxData array in sync with MDC StateBuckets
+					Scene.CachedStateBucketsAuxData[CmdInfo.MeshPass].SetNum(BucketMap.GetMaxIndex() + 1, false);
+					Scene.CachedStateBucketsAuxData[CmdInfo.MeshPass][CmdInfo.StateBucketId] = DeferredStateBucketsAuxData[DeferredIndex];
+				}
 			}
 		}
 	}
@@ -2155,6 +2176,7 @@ void FCachedPassMeshDrawListContextDeferred::DeferredFinalizeMeshDrawCommands(co
 
 	DeferredCommands.Reset();
 	DeferredCommandHashes.Reset();
+	DeferredStateBucketsAuxData.Reset();
 }
 
 PassProcessorCreateFunction FPassProcessorManager::JumpTable[(int32)EShadingPath::Num][EMeshPass::Num] = {};
diff --git a/Engine/Source/Runtime/Renderer/Private/PrimitiveSceneInfo.cpp b/Engine/Source/Runtime/Renderer/Private/PrimitiveSceneInfo.cpp
--- a/Engine/Source/Runtime/Renderer/Private/PrimitiveSceneInfo.cpp
+++ b/Engine/Source/Runtime/Renderer/Private/PrimitiveSceneInfo.cpp
@@ -492,6 +492,13 @@ void FPrimitiveSceneInfo::RemoveCachedMeshDrawCommands()
 				{
 					Scene->CachedMeshDrawCommandStateBuckets[PassIndex].RemoveByElementId(CachedCommand.StateBucketId);
 				}
+
+				// shrink MDC AuxData array in sync with MDC StateBuckets 
+				int32 StateBucketsNum = Scene->CachedMeshDrawCommandStateBuckets[PassIndex].GetMaxIndex() + 1;
+				if (Scene->CachedStateBucketsAuxData[PassIndex].Num() > StateBucketsNum) 
+				{
+					Scene->CachedStateBucketsAuxData[PassIndex].SetNum(StateBucketsNum, true);
+				}
 			}
 
 			FGraphicsMinimalPipelineStateId::RemovePersistentId(CachedPipelineId);
diff --git a/Engine/Source/Runtime/Renderer/Private/ScenePrivate.h b/Engine/Source/Runtime/Renderer/Private/ScenePrivate.h
--- a/Engine/Source/Runtime/Renderer/Private/ScenePrivate.h
+++ b/Engine/Source/Runtime/Renderer/Private/ScenePrivate.h
@@ -2730,6 +2730,9 @@ public:
 	void RefreshRayTracingInstances();
 #endif
 
+	/* Mobile specific. Auxilary data for each cached command state bucket. This data can be used during mesh draw command batching */
+	TArray<FStateBucketAuxData> CachedStateBucketsAuxData[EMeshPass::Num];
+
 	/** Nanite shading material commands. These are stored on the scene as they are computed at FPrimitiveSceneInfo::AddToScene time. */
 	FNaniteMaterialCommands NaniteMaterials[ENaniteMeshPass::Num];
 
diff --git a/Engine/Source/Runtime/Renderer/Private/VirtualShadowMaps/VirtualShadowMapArray.cpp b/Engine/Source/Runtime/Renderer/Private/VirtualShadowMaps/VirtualShadowMapArray.cpp
--- a/Engine/Source/Runtime/Renderer/Private/VirtualShadowMaps/VirtualShadowMapArray.cpp
+++ b/Engine/Source/Runtime/Renderer/Private/VirtualShadowMaps/VirtualShadowMapArray.cpp
@@ -2026,7 +2026,7 @@ public:
 		SHADER_PARAMETER(FVector3f, CullingViewOriginTile)
 
 		SHADER_PARAMETER_RDG_BUFFER_SRV(StructuredBuffer< FPackedView >, InViews)
-		SHADER_PARAMETER_RDG_BUFFER_SRV(StructuredBuffer< FInstanceCullingContext::FDrawCommandDesc >, DrawCommandDescs)
+		SHADER_PARAMETER_RDG_BUFFER_SRV(StructuredBuffer< uint32 >, DrawCommandDescs)
 
 		SHADER_PARAMETER_RDG_BUFFER_SRV(StructuredBuffer< FContextBatchInfo >, BatchInfos)
 		SHADER_PARAMETER_RDG_BUFFER_SRV(StructuredBuffer< FVSMCullingBatchInfo >, VSMCullingBatchInfos)
@@ -2154,7 +2154,7 @@ struct FCullingResult
 template <typename InstanceCullingLoadBalancerType>
 static FCullingResult AddCullingPasses(FRDGBuilder& GraphBuilder,
 	const TConstArrayView<FRHIDrawIndexedIndirectParameters> &IndirectArgs,
-	const TConstArrayView<FInstanceCullingContext::FDrawCommandDesc>& DrawCommandDescs,
+	const TConstArrayView<uint32>& DrawCommandDescs,
 	const TConstArrayView<uint32>& InstanceIdOffsets,
 	InstanceCullingLoadBalancerType *LoadBalancer,
 	const TConstArrayView<FInstanceCullingMergedContext::FContextBatchInfo> BatchInfos,
diff --git a/Engine/Source/Runtime/Renderer/Public/InstanceCulling/InstanceCullingContext.h b/Engine/Source/Runtime/Renderer/Public/InstanceCulling/InstanceCullingContext.h
--- a/Engine/Source/Runtime/Renderer/Public/InstanceCulling/InstanceCullingContext.h
+++ b/Engine/Source/Runtime/Renderer/Public/InstanceCulling/InstanceCullingContext.h
@@ -159,6 +159,15 @@ public:
 	static void AddClearIndirectArgInstanceCountPass(FRDGBuilder& GraphBuilder, FGlobalShaderMap* ShaderMap, FRDGBufferRef DrawIndirectArgsBuffer, TFunction<int32()> NumIndirectArgsCallback = TFunction<int32()>());
 
 
+	void SetupDrawCommands(
+		TArrayView<const FStateBucketAuxData> StateBucketsAuxData,
+		FMeshCommandOneFrameArray& VisibleMeshDrawCommandsInOut,
+		bool bCompactIdenticalCommands,
+		// Stats
+		int32& MaxInstancesOut,
+		int32& VisibleMeshDrawCommandsNumOut,
+		int32& NewPassVisibleMeshDrawCommandsNumOut);
+
 	void SetupDrawCommands(
 		FMeshCommandOneFrameArray& VisibleMeshDrawCommandsInOut,
 		bool bCompactIdenticalCommands,
@@ -198,12 +207,6 @@ public:
 		uint32 InstanceDataByteOffset;
 	};
 
-	// TODO: bit-pack
-	struct FDrawCommandDesc
-	{
-		uint32 bMaterialUsesWorldPositionOffset;
-	};
-
 	struct FPayloadData
 	{
 		uint32 bDynamicInstanceDataOffset_IndirectArgsIndex;
@@ -258,7 +261,7 @@ public:
 
 	TArray<FMeshDrawCommandInfo, SceneRenderingAllocator> MeshDrawCommandInfos;
 	TArray<FRHIDrawIndexedIndirectParameters, SceneRenderingAllocator> IndirectArgs;
-	TArray<FDrawCommandDesc, SceneRenderingAllocator> DrawCommandDescs;
+	TArray<uint32, SceneRenderingAllocator> DrawCommandDescs;
 	TArray<FPayloadData, SceneRenderingAllocator> PayloadData;
 	TArray<uint32, SceneRenderingAllocator> InstanceIdOffsets;
 
diff --git a/Engine/Source/Runtime/Renderer/Public/MeshPassProcessor.h b/Engine/Source/Runtime/Renderer/Public/MeshPassProcessor.h
--- a/Engine/Source/Runtime/Renderer/Public/MeshPassProcessor.h
+++ b/Engine/Source/Runtime/Renderer/Public/MeshPassProcessor.h
@@ -1641,6 +1641,20 @@ struct MeshDrawCommandKeyFuncs : TDefaultMapHashableKeyFuncs<FMeshDrawCommand, F
 using FDrawCommandIndices = TArray<int32, TInlineAllocator<5>>;
 using FStateBucketMap = Experimental::TRobinHoodHashMap<FMeshDrawCommand, FMeshDrawCommandCount, MeshDrawCommandKeyFuncs>;
 
+struct FStateBucketAuxData
+{
+	FStateBucketAuxData()
+		: MeshLODIndex(0)
+	{}
+
+	explicit FStateBucketAuxData(const FMeshBatch& MeshBatch)
+		: MeshLODIndex(MeshBatch.LODIndex)
+	{
+	}
+
+	uint8 MeshLODIndex;
+};
+
 class FCachedPassMeshDrawListContext : public FMeshPassDrawListContext
 {
 public:
@@ -1695,6 +1709,7 @@ protected:
 	EMeshPass::Type CurrMeshPass = EMeshPass::Num;
 	bool bUseGPUScene = false;
 	bool bAnyLooseParameterBuffers = false;
+	bool bUseStateBucketsAuxData = false;
 };
 
 class FCachedPassMeshDrawListContextImmediate : public FCachedPassMeshDrawListContext
@@ -1737,6 +1752,7 @@ public:
 private:
 	TArray<FMeshDrawCommand> DeferredCommands;
 	TArray<Experimental::FHashType> DeferredCommandHashes;
+	TArray<FStateBucketAuxData> DeferredStateBucketsAuxData;
 };
 
 template<typename VertexType, typename PixelType, typename GeometryType = FMeshMaterialShader, typename RayTracingType = FMeshMaterialShader, typename ComputeType = FMeshMaterialShader>
