HeterogeneousVolumes: Fixing issue where recalaculating step-size was overwiting initial scene depth.

diff --git a/Engine/Shaders/Private/HeterogeneousVolumes/HeterogeneousVolumesHardwareRayTracing.usf b/Engine/Shaders/Private/HeterogeneousVolumes/HeterogeneousVolumesHardwareRayTracing.usf
--- a/Engine/Shaders/Private/HeterogeneousVolumes/HeterogeneousVolumesHardwareRayTracing.usf
+++ b/Engine/Shaders/Private/HeterogeneousVolumes/HeterogeneousVolumesHardwareRayTracing.usf
@@ -510,7 +510,7 @@ RAY_TRACING_ENTRY_RAYGEN(RenderSingleScatteringWithPreshadingRGS)
 	float3 LocalBoundsMax = GetLocalBoundsOrigin() + GetLocalBoundsExtent();
 
 	// Test bounding volume
-	float2 BoundingHitT = IntersectAABB(LocalRayOrigin, LocalRayDirection, 0.0, SceneDepth, LocalBoundsMin, LocalBoundsMax);
+	float2 BoundingHitT = IntersectAABB(LocalRayOrigin, LocalRayDirection, 0.0, LocalRayLength, LocalBoundsMin, LocalBoundsMax);
 	bool bIntersectsVolume = (BoundingHitT.y - BoundingHitT.x > 0.0);
 	if (!bIntersectsVolume)
 	{
diff --git a/Engine/Shaders/Private/HeterogeneousVolumes/HeterogeneousVolumesLiveShadingPipeline.usf b/Engine/Shaders/Private/HeterogeneousVolumes/HeterogeneousVolumesLiveShadingPipeline.usf
--- a/Engine/Shaders/Private/HeterogeneousVolumes/HeterogeneousVolumesLiveShadingPipeline.usf
+++ b/Engine/Shaders/Private/HeterogeneousVolumes/HeterogeneousVolumesLiveShadingPipeline.usf
@@ -290,7 +290,7 @@ void RenderSingleScatteringWithLiveShadingCS(
 	float3 LocalBoundsMax = LocalBoundsOrigin + LocalBoundsExtent;
 
 	// Test bounding volume
-	float2 HitT = IntersectAABB(LocalRayOrigin, LocalRayDirection, 0.0, SceneDepth, LocalBoundsMin, LocalBoundsMax);
+	float2 HitT = IntersectAABB(LocalRayOrigin, LocalRayDirection, 0.0, LocalRayLength, LocalBoundsMin, LocalBoundsMax);
 	float HitSpan = HitT.y - HitT.x;
 	if (HitSpan > 0.0)
 	{
diff --git a/Engine/Shaders/Private/HeterogeneousVolumes/HeterogeneousVolumesPreshadingPipeline.usf b/Engine/Shaders/Private/HeterogeneousVolumes/HeterogeneousVolumesPreshadingPipeline.usf
--- a/Engine/Shaders/Private/HeterogeneousVolumes/HeterogeneousVolumesPreshadingPipeline.usf
+++ b/Engine/Shaders/Private/HeterogeneousVolumes/HeterogeneousVolumesPreshadingPipeline.usf
@@ -545,7 +545,7 @@ void RenderSingleScatteringWithPreshadingCS(
 	float3 LocalBoundsMax = GetLocalBoundsOrigin() + GetLocalBoundsExtent();
 
 	// Test bounding volume
-	float2 HitT = IntersectAABB(LocalRayOrigin, LocalRayDirection, 0.0, SceneDepth, LocalBoundsMin, LocalBoundsMax);
+	float2 HitT = IntersectAABB(LocalRayOrigin, LocalRayDirection, 0.0, LocalRayLength, LocalBoundsMin, LocalBoundsMax);
 	float HitSpan = HitT.y - HitT.x;
 	if (HitSpan <= 0.0)
 	{
diff --git a/Engine/Shaders/Private/HeterogeneousVolumes/HeterogeneousVolumesRayMarchingUtils.ush b/Engine/Shaders/Private/HeterogeneousVolumes/HeterogeneousVolumesRayMarchingUtils.ush
--- a/Engine/Shaders/Private/HeterogeneousVolumes/HeterogeneousVolumesRayMarchingUtils.ush
+++ b/Engine/Shaders/Private/HeterogeneousVolumes/HeterogeneousVolumesRayMarchingUtils.ush
@@ -81,20 +81,22 @@ FRayMarchingContext CreateRayMarchingContext(
 	return Context;
 }
 
+float RecalculateStepSize(inout FRayMarchingContext Context)
+{
+	// Recalculate step-size if the bounding-box diagonal is greater than the maximum step-count
+	float DiagonalLength = length(2.0 * GetLocalBoundsExtent());
+	uint MaxStepCount = max(DiagonalLength / Context.StepSize, 1);
+
+	float StepSize = (MaxStepCount > Context.MaxStepCount) ? DiagonalLength / Context.MaxStepCount : Context.StepSize;
+	return StepSize;
+}
+
 uint CalcStepCount(inout FRayMarchingContext Context)
 {
-	// Instead using the active ray-span, use the maximal ray-span so the stepsize is uniform among all rays
-	//float HitSpan = length(Hit1 - Hit0)
-	float HitSpan = length(2.0 * GetLocalBoundsExtent());
-	float StepSize = Context.StepSize;
-	uint StepCount = max(ceil(HitSpan / StepSize), 1);
-	if (StepCount > Context.MaxStepCount)
-	{
-		StepCount = Context.MaxStepCount;
-		StepSize = HitSpan / StepCount;
-	}
+	Context.StepSize = RecalculateStepSize(Context);
 
-	Context.StepSize = StepSize;
+	float HitSpan = length(Context.LocalRayTMax - Context.LocalRayTMin);
+	uint StepCount = max(ceil(HitSpan / Context.StepSize), 1);
 	return StepCount;
 }
 
