Fix minor shader warning in RectLight.ush reported by horde.

diff --git a/Engine/Shaders/Private/RectLight.ush b/Engine/Shaders/Private/RectLight.ush
--- a/Engine/Shaders/Private/RectLight.ush
+++ b/Engine/Shaders/Private/RectLight.ush
@@ -316,9 +316,9 @@ float3 SampleSourceTexture( float3 L, FRect Rect, FRectTexture RectTexture)
 	PointInRect.y = dot( Rect.Axis[1], PointOnPlane - Rect.Origin );
 
 	// Compute UV on the original rect (i.e. unoccluded rect)
-    float2 RectUV = (PointInRect + Rect.Offset) / Rect.FullExtent * float2(0.5, -0.5) + 0.5;
+    float2 RectUV = (PointInRect + Rect.Offset) / max(0.0001f, Rect.FullExtent) * float2(0.5, -0.5) + 0.5;
 	
-	float Level = log2( DistToPlane * rsqrt( Rect.FullExtent.x * Rect.FullExtent.y ) );
+	float Level = log2( DistToPlane * rsqrt( max(0.0001f, Rect.FullExtent.x * Rect.FullExtent.y) ) );
 
     return SampleRectTexture(RectTexture, RectUV, Level);
 #else
@@ -652,7 +652,7 @@ FRect GetRect(
 		const float CosTheta = LightDataRectLightBarnCosAngle;
 		const float SinTheta = sqrt(1 - CosTheta * CosTheta);
 		const float BarnDepth = min(S_Light.z, CosTheta * BarnLength);
-		const float S_ratio = BarnDepth / (CosTheta * BarnLength);
+		const float S_ratio = BarnDepth / max(0.0001f, CosTheta * BarnLength);
 		const float D_B = SinTheta * BarnLength * S_ratio;
 		
 		// Clamp shading point onto the closest edge, if it is inside the rect light
