Fix Strata tile resource misspelling (only used when opaque rough refraction is enabled).

diff --git a/Engine/Shaders/Private/Strata/StrataMaterialClassification.usf b/Engine/Shaders/Private/Strata/StrataMaterialClassification.usf
--- a/Engine/Shaders/Private/Strata/StrataMaterialClassification.usf
+++ b/Engine/Shaders/Private/Strata/StrataMaterialClassification.usf
@@ -65,7 +65,7 @@ uint GetDBufferTargetMask(uint2 PixelPos)
 
 #if STRATA_OPAQUE_ROUGH_REFRACTION_ENABLED
 RWBuffer<uint> OpaqueRoughRefractionTileListDataBuffer;
-RWBuffer<uint> SSSWithoutOpaqueRoughRefractionTileListDataBuffer;
+RWBuffer<uint> OpaqueRoughRefractionSSSWithoutTileListDataBuffer;
 Texture2D<float3> OpaqueRoughRefractionTexture;
 #endif // STRATA_OPAQUE_ROUGH_REFRACTION_ENABLED
 
@@ -298,7 +298,7 @@ void TileMainCS(uint2 DispatchThreadId : SV_DispatchThreadID, uint LinearIndex :
 		{
 			uint WriteToIndex;
 			InterlockedAdd(TileDrawIndirectDataBuffer[GetStrataTileTypeDrawIndirectArgOffset_DWord(STRATA_TILE_TYPE_ROUGH_REFRACT_SSS_WITHOUT) + 1], 1, WriteToIndex);
-			SSSWithoutOpaqueRoughRefractionTileListDataBuffer[WriteToIndex] = EncodedTile;
+			OpaqueRoughRefractionSSSWithoutTileListDataBuffer[WriteToIndex] = EncodedTile;
 		}
 	#endif
 
