Fix shader compilation issue due to rect light with Strata.

diff --git a/Engine/Shaders/Private/Strata/StrataForwardLighting.ush b/Engine/Shaders/Private/Strata/StrataForwardLighting.ush
--- a/Engine/Shaders/Private/Strata/StrataForwardLighting.ush
+++ b/Engine/Shaders/Private/Strata/StrataForwardLighting.ush
@@ -417,8 +417,9 @@ float3 StrataForwardLighting(
 
 					half4 PreviewShadowMapChannelMask = 1;
 					uint  LightChannelMask = LIGHING_CHANNEL_MASK;
-					const FDeferredLightData LightData = ConvertLocalToDeferredLight(LocalLight, 1.0f /*SpecularScale*/, PreviewShadowMapChannelMask, LightChannelMask);
-					
+					FDeferredLightData LightData = ConvertLocalToDeferredLight(LocalLight, 1.0f /*SpecularScale*/, PreviewShadowMapChannelMask, LightChannelMask);
+					LightData.bRectLight = false; // Rect light not supported in forward. Enabling this is causing the max sampler count to be exceeded.
+
 					// Skip rect light. This is not done explicitely in the legacy path (resulting in no lighting) but we do it for strata in order to avoid visual artefact.
 					// This can be removed when LocalLight properly supports rect lights.
 					if (LightData.bRectLight)
