Global Distance Field - fixed heightfield page allocation reducing number of required pages by 5%. Also added debug code for counting filled pages (pages having values outside of 0 or 1, which need to be allocated and filled).

diff --git a/Engine/Shaders/Private/DistanceField/GlobalDistanceFieldDebug.usf b/Engine/Shaders/Private/DistanceField/GlobalDistanceFieldDebug.usf
--- a/Engine/Shaders/Private/DistanceField/GlobalDistanceFieldDebug.usf
+++ b/Engine/Shaders/Private/DistanceField/GlobalDistanceFieldDebug.usf
@@ -2,10 +2,72 @@
 
 #include "../Common.ush"
 #include "../ShaderPrint.ush"
+#include "GlobalDistanceFieldShared.ush"
 
 StructuredBuffer<uint> GlobalDistanceFieldPageFreeListAllocatorBuffer;
 uint GlobalDistanceFieldMaxPageNum;
 
+uint ClipmapSizeInPages;
+uint ClipmapIndex;
+float3 InvPageAtlasSize;
+RWStructuredBuffer<uint> RWPageStatsBuffer;
+StructuredBuffer<uint> PageStatsBuffer;
+Texture3D<uint> PageTableCombinedTexture;
+Texture3D PageAtlasTexture;
+
+[numthreads(THREADGROUP_SIZE, THREADGROUP_SIZE, THREADGROUP_SIZE)]
+void GlobalDistanceFieldPageStatsCS(
+	uint3 GroupId : SV_GroupID,
+	uint3 DispatchThreadId : SV_DispatchThreadID,
+	uint3 GroupThreadId : SV_GroupThreadID)
+{
+	uint3 PageTableCoord = DispatchThreadId.xyz;
+	if (all(PageTableCoord < ClipmapSizeInPages))
+	{
+		PageTableCoord.z += ClipmapIndex * ClipmapSizeInPages;
+
+		FGlobalDistanceFieldPage Page = UnpackGlobalDistanceFieldPage(PageTableCombinedTexture.Load(uint4(PageTableCoord, 0)));
+		if (Page.bValid)
+		{
+			InterlockedAdd(RWPageStatsBuffer[0], 1);
+
+			bool bFilled = false;
+
+			for (uint PageCoordZ = 0; PageCoordZ < GLOBAL_DISTANCE_FIELD_PAGE_RESOLUTION; ++PageCoordZ)
+			{
+				for (uint PageCoordY = 0; PageCoordY < GLOBAL_DISTANCE_FIELD_PAGE_RESOLUTION; ++PageCoordY)
+				{
+					for (uint PageCoordX = 0; PageCoordX < GLOBAL_DISTANCE_FIELD_PAGE_RESOLUTION; ++PageCoordX)
+					{
+						float3 PageCoord = uint3(PageCoordX, PageCoordY, PageCoordZ);
+						uint3 PageAtlasOffset = GlobalDistanceFieldPageLinearIndexToPageAtlasOffset(Page);
+
+						float DistanceFieldValue = PageAtlasTexture[PageAtlasOffset + PageCoord].x;
+						if (View.GeneralPurposeTweak > 0.5f)
+						{
+							// Don't sample border values which can be safely skipped
+							PageCoord = clamp(PageCoord + 0.5f, GLOBAL_DISTANCE_FIELD_PAGE_BORDER, GLOBAL_DISTANCE_FIELD_PAGE_RESOLUTION - GLOBAL_DISTANCE_FIELD_PAGE_BORDER);
+
+							float3 PageUV = (PageAtlasOffset + PageCoord) * InvPageAtlasSize;
+							DistanceFieldValue = Texture3DSampleLevel(PageAtlasTexture, GlobalTrilinearWrappedSampler, PageUV, 0).x;
+						}
+
+						if (DistanceFieldValue > 0.0f && DistanceFieldValue < 1.0f)
+						{
+							bFilled = true;
+						}
+					}
+				}
+			}
+
+			if (!bFilled)
+			{
+				InterlockedAdd(RWPageStatsBuffer[1], 1);
+			}
+		}
+	}
+}
+
 [numthreads(1, 1, 1)]
 void GlobalDistanceFieldDebugCS(
 	uint3 GroupId : SV_GroupID,
@@ -32,5 +94,12 @@ void GlobalDistanceFieldDebugCS(
 		Print(Context, TEXT(" Max   "));
 		Print(Context, GlobalDistanceFieldMaxPageNum);
 		Newline(Context);
+
+		Print(Context, TEXT(" Combi "));
+		Print(Context, PageStatsBuffer[0]);
+		Newline(Context);
+
+		Print(Context, TEXT(" Empty "));
+		Print(Context, PageStatsBuffer[1]);
 	}
 }
\ No newline at end of file
diff --git a/Engine/Shaders/Private/DistanceField/GlobalDistanceFieldHeightfields.usf b/Engine/Shaders/Private/DistanceField/GlobalDistanceFieldHeightfields.usf
--- a/Engine/Shaders/Private/DistanceField/GlobalDistanceFieldHeightfields.usf
+++ b/Engine/Shaders/Private/DistanceField/GlobalDistanceFieldHeightfields.usf
@@ -28,6 +28,74 @@
 float3 ViewTilePosition;
 float3 RelativePreViewTranslation;
 
+struct FHeightfieldSample
+{
+	bool bValid;
+	uint HeightfieldIndex;
+	float3 TranslatedWorldPosition;
+	float3 WorldNormal;
+};
+
+FHeightfieldSample SampleHeightfield(FLWCVector3 PreViewTranslation, FLWCVector3 VoxelWorldCenter)
+{
+	FHeightfieldSample Sample;
+	Sample.bValid = false;
+	Sample.HeightfieldIndex = NumHeightfields;
+	Sample.TranslatedWorldPosition = float3(0.0f, 0.0f, 0.0f);
+	Sample.WorldNormal = float3(0.0f, 0.0f, 0.0f);
+
+	for (uint HeightfieldIndex = 0; HeightfieldIndex < NumHeightfields; HeightfieldIndex++)
+	{
+		float3 LocalPosition = LWCMultiply(VoxelWorldCenter, GetWorldToLocal(HeightfieldIndex));
+		float4 MinMaxHeightfieldUV;
+		float2 HeightfieldUV = GetHeightfieldUV(HeightfieldIndex, LocalPosition.xy, MinMaxHeightfieldUV);
+
+		if (all(HeightfieldUV >= MinMaxHeightfieldUV.xy) && all(HeightfieldUV <= MinMaxHeightfieldUV.zw))
+		{
+			Sample.HeightfieldIndex = HeightfieldIndex;
+			break;
+		}
+	}
+
+	if (Sample.HeightfieldIndex < NumHeightfields)
+	{
+		float3 LocalPosition = LWCMultiply(VoxelWorldCenter, GetWorldToLocal(Sample.HeightfieldIndex));
+		float4 MinMaxHeightfieldUV;
+		float2 HeightfieldUV = GetHeightfieldUV(Sample.HeightfieldIndex, LocalPosition.xy, MinMaxHeightfieldUV);
+
+		if (all(HeightfieldUV >= MinMaxHeightfieldUV.xy) && all(HeightfieldUV <= MinMaxHeightfieldUV.zw))
+		{
+			float HeightfieldVisibility;
+			FLWCVector3 WorldHeightfieldShadingPosition = GetHeightfieldWorldPositionAndNormal(Sample.HeightfieldIndex, LocalPosition.xy, HeightfieldUV, Sample.WorldNormal, HeightfieldVisibility);
+			Sample.TranslatedWorldPosition = LWCToFloat(LWCAdd(WorldHeightfieldShadingPosition, PreViewTranslation));
+
+			// Skip holes in the heightfield
+			if (HeightfieldVisibility > 0.5f)
+			{
+				Sample.bValid = true;
+			}
+		}
+	}
+
+	return Sample;
+}
+
+float ComputeHeightfieldDistance(FHeightfieldSample HeightfieldSample, float3 TranslatedWorldPosition)
+{
+	// Project the vertical height vector onto the normal of the heightfield directly below the point we are computing the distance field for, use the perpendicular distance
+	float DistanceToHeightfieldPlane = dot(HeightfieldSample.WorldNormal, TranslatedWorldPosition - HeightfieldSample.TranslatedWorldPosition);
+
+	// Limit negative region of a heightfield to a user defined thickness  
+	const float MinInteriorDistance = -HeightfieldThickness;
+	if (DistanceToHeightfieldPlane < MinInteriorDistance)
+	{
+		DistanceToHeightfieldPlane = MinInteriorDistance - DistanceToHeightfieldPlane;
+	}
+
+	float HeightfieldMinDistance = clamp(DistanceToHeightfieldPlane, -InfluenceRadius, InfluenceRadius);
+	return HeightfieldMinDistance;
+}
+
 groupshared uint GroupMarkedPage[64];
 
 [numthreads(THREADGROUP_SIZE, THREADGROUP_SIZE, 1)]
@@ -36,65 +104,33 @@
 	uint3 DispatchThreadId : SV_DispatchThreadID,
 	uint3 GroupThreadId : SV_GroupThreadID)
 {
-	const FLWCVector3 PreViewTranslation = MakeLWCVector3(-ViewTilePosition, RelativePreViewTranslation);
-
-	uint3 VoxelCoord = uint3(GroupThreadId.xy, 0);
 	uint ThreadIndex = GroupThreadId.x + GroupThreadId.y * THREADGROUP_SIZE;
 	uint IndexInPageBuffer = GroupId.x;
+
+	const FLWCVector3 PreViewTranslation = MakeLWCVector3(-ViewTilePosition, RelativePreViewTranslation);
+
+	uint3 TexelCoordInPage = uint3(GroupThreadId.xy, 0);
 	uint3 PageGridCoord = UnpackPageUpdateTile(PageUpdateTileBuffer[IndexInPageBuffer]);
-
-	float3 PageTranslatedWorldCenter = PageGridCoord * PageCoordToPageTranslatedWorldCenterScale + PageCoordToPageTranslatedWorldCenterBias;
-	float3 VoxelTranslatedWorldCenter = PageTranslatedWorldCenter - PageWorldExtent + VoxelCoord * 2.0f * PageWorldExtent;
+	float3 PageCoord = PageGridCoord * GLOBAL_DISTANCE_FIELD_PAGE_RESOLUTION + TexelCoordInPage - GLOBAL_DISTANCE_FIELD_PAGE_BORDER;
+	float3 VoxelTranslatedWorldCenter = PageCoord * PageCoordToVoxelTranslatedCenterScale + PageCoordToVoxelTranslatedCenterBias;
 	FLWCVector3 VoxelWorldCenter = LWCSubtract(VoxelTranslatedWorldCenter, PreViewTranslation);
 
 	GroupMarkedPage[ThreadIndex] = 0;
 
-	uint ValidHeightfieldIndex = 0;
+	FHeightfieldSample HeightfieldSample = SampleHeightfield(PreViewTranslation, VoxelWorldCenter);
+	if (HeightfieldSample.bValid)
+	{
+		float3 VoxelTranslatedWorldCenterNearestZ = VoxelTranslatedWorldCenter;
 
-	LOOP
-	for (uint HeightfieldIndex = 0; HeightfieldIndex < NumHeightfields; HeightfieldIndex++)
-	{
-		float3 LocalPosition = LWCMultiply(VoxelWorldCenter, GetWorldToLocal(HeightfieldIndex));
-		float4 MinMaxHeightfieldUV;
-		float2 HeightfieldUV = GetHeightfieldUV(HeightfieldIndex, LocalPosition.xy, MinMaxHeightfieldUV);
+		float MinZ = VoxelTranslatedWorldCenter.z + GLOBAL_DISTANCE_FIELD_PAGE_BORDER * ClipmapVoxelExtent;
+		float MaxZ = VoxelTranslatedWorldCenter.z + (GLOBAL_DISTANCE_FIELD_PAGE_RESOLUTION_IN_ATLAS - GLOBAL_DISTANCE_FIELD_PAGE_BORDER) * ClipmapVoxelExtent;
+		VoxelTranslatedWorldCenterNearestZ.z = clamp(HeightfieldSample.TranslatedWorldPosition.z, MinZ, MaxZ);
 
-		if (all(HeightfieldUV > MinMaxHeightfieldUV.xy) && all(HeightfieldUV < MinMaxHeightfieldUV.zw))
+		float HeightfieldDistance = ComputeHeightfieldDistance(HeightfieldSample, VoxelTranslatedWorldCenterNearestZ);
+
+		if (abs(HeightfieldDistance) < InfluenceRadius)
 		{
-			ValidHeightfieldIndex = HeightfieldIndex;
-			break;
+			GroupMarkedPage[ThreadIndex] = 1;
 		}
 	}
 
-	float3 LocalPosition = LWCMultiply(VoxelWorldCenter, GetWorldToLocal(ValidHeightfieldIndex));
-	float4 MinMaxHeightfieldUV;
-	float2 HeightfieldUV = GetHeightfieldUV(ValidHeightfieldIndex, LocalPosition.xy, MinMaxHeightfieldUV);
-
-	if (all(HeightfieldUV > MinMaxHeightfieldUV.xy) && all(HeightfieldUV < MinMaxHeightfieldUV.zw))
-	{
-		float3 WorldHeightfieldNormal;
-		float HeightfieldVisibility;
-		FLWCVector3 WorldHeightfieldPosition = GetHeightfieldWorldPositionAndNormal(ValidHeightfieldIndex, LocalPosition.xy, HeightfieldUV, WorldHeightfieldNormal, HeightfieldVisibility);
-		float3 TranslatedWorldHeightfieldPosition = LWCToFloat(LWCAdd(WorldHeightfieldPosition, PreViewTranslation));
-
-		if (HeightfieldVisibility > 0.5f) // Skip holes in the heightfield
-		{
-			float3 TranslatedWorldPositionNearestZ = VoxelTranslatedWorldCenter;
-			TranslatedWorldPositionNearestZ.z = clamp(TranslatedWorldHeightfieldPosition.z, VoxelTranslatedWorldCenter.z, VoxelTranslatedWorldCenter.z + GLOBAL_DISTANCE_FIELD_PAGE_RESOLUTION_IN_ATLAS * 2.0f * ClipmapVoxelExtent);
-
-			// Project the vertical height vector onto the normal of the heightfield directly below the point we are computing the distance field for, use the perpendicular distance
-			float DistanceToHeightfieldPlane = dot(WorldHeightfieldNormal, TranslatedWorldPositionNearestZ - TranslatedWorldHeightfieldPosition);
-
-			// Limit negative region of a heightfield to a user defined thickness  
-			const float MinInteriorDistance = -HeightfieldThickness;
-			if (DistanceToHeightfieldPlane < MinInteriorDistance)
-			{
-				DistanceToHeightfieldPlane = MinInteriorDistance - DistanceToHeightfieldPlane;
-			}
-
-			if (DistanceToHeightfieldPlane > -InfluenceRadius && DistanceToHeightfieldPlane < InfluenceRadius)
-			{
-				GroupMarkedPage[ThreadIndex] = 1;
-			}
-		}
-	}
-
@@ -193,74 +229,6 @@
 	}
 }
 
-struct FHeightfieldSample
-{
-	bool bValid;
-	uint HeightfieldIndex;
-	float3 TranslatedWorldPosition;
-	float3 WorldNormal;
-};
-
-FHeightfieldSample SampleHeightfield(FLWCVector3 PreViewTranslation, FLWCVector3 VoxelWorldCenter)
-{
-	FHeightfieldSample Sample;
-	Sample.bValid = false;
-	Sample.HeightfieldIndex = NumHeightfields;
-	Sample.TranslatedWorldPosition = float3(0.0f, 0.0f, 0.0f);
-	Sample.WorldNormal = float3(0.0f, 0.0f, 0.0f);
-
-	for (uint HeightfieldIndex = 0; HeightfieldIndex < NumHeightfields; HeightfieldIndex++)
-	{
-		float3 LocalPosition = LWCMultiply(VoxelWorldCenter, GetWorldToLocal(HeightfieldIndex));
-		float4 MinMaxHeightfieldUV;
-		float2 HeightfieldUV = GetHeightfieldUV(HeightfieldIndex, LocalPosition.xy, MinMaxHeightfieldUV);
-
-		if (all(HeightfieldUV >= MinMaxHeightfieldUV.xy) && all(HeightfieldUV <= MinMaxHeightfieldUV.zw))
-		{
-			Sample.HeightfieldIndex = HeightfieldIndex;
-			break;
-		}
-	}
-
-	if (Sample.HeightfieldIndex < NumHeightfields)
-	{
-		float3 LocalPosition = LWCMultiply(VoxelWorldCenter, GetWorldToLocal(Sample.HeightfieldIndex));
-		float4 MinMaxHeightfieldUV;
-		float2 HeightfieldUV = GetHeightfieldUV(Sample.HeightfieldIndex, LocalPosition.xy, MinMaxHeightfieldUV);
-
-		if (all(HeightfieldUV >= MinMaxHeightfieldUV.xy) && all(HeightfieldUV <= MinMaxHeightfieldUV.zw))
-		{
-			float HeightfieldVisibility;
-			FLWCVector3 WorldHeightfieldShadingPosition = GetHeightfieldWorldPositionAndNormal(Sample.HeightfieldIndex, LocalPosition.xy, HeightfieldUV, Sample.WorldNormal, HeightfieldVisibility);
-			Sample.TranslatedWorldPosition = LWCToFloat(LWCAdd(WorldHeightfieldShadingPosition, PreViewTranslation));
-
-			// Skip holes in the heightfield
-			if (HeightfieldVisibility > 0.5f)
-			{
-				Sample.bValid = true;
-			}
-		}
-	}
-
-	return Sample;
-}
-
-float ComputeHeightfieldDistance(FHeightfieldSample HeightfieldSample, float3 TranslatedWorldPosition)
-{
-	// Project the vertical height vector onto the normal of the heightfield directly below the point we are computing the distance field for, use the perpendicular distance
-	float DistanceToHeightfieldPlane = dot(HeightfieldSample.WorldNormal, TranslatedWorldPosition - HeightfieldSample.TranslatedWorldPosition);
-
-	// Limit negative region of a heightfield to a user defined thickness  
-	const float MinInteriorDistance = -HeightfieldThickness;
-	if (DistanceToHeightfieldPlane < MinInteriorDistance)
-	{
-		DistanceToHeightfieldPlane = MinInteriorDistance - DistanceToHeightfieldPlane;
-	}
-
-	float HeightfieldMinDistance = clamp(DistanceToHeightfieldPlane, -InfluenceRadius, InfluenceRadius);
-	return HeightfieldMinDistance;
-}
-
 [numthreads(THREADGROUP_SIZE, THREADGROUP_SIZE, 1)]
 void ComposeHeightfieldsIntoPagesCS(
 	uint3 GroupId : SV_GroupID,
diff --git a/Engine/Source/Runtime/Renderer/Private/GlobalDistanceField.cpp b/Engine/Source/Runtime/Renderer/Private/GlobalDistanceField.cpp
--- a/Engine/Source/Runtime/Renderer/Private/GlobalDistanceField.cpp
+++ b/Engine/Source/Runtime/Renderer/Private/GlobalDistanceField.cpp
@@ -1790,6 +1790,40 @@ class FPropagateMipDistanceCS : public FGlobalShader
 
 IMPLEMENT_GLOBAL_SHADER(FPropagateMipDistanceCS, "/Engine/Private/DistanceField/GlobalDistanceFieldMip.usf", "PropagateMipDistanceCS", SF_Compute);
 
+class FGlobalDistanceFieldPageStatsCS : public FGlobalShader
+{
+	DECLARE_GLOBAL_SHADER(FGlobalDistanceFieldPageStatsCS)
+	SHADER_USE_PARAMETER_STRUCT(FGlobalDistanceFieldPageStatsCS, FGlobalShader)
+
+	BEGIN_SHADER_PARAMETER_STRUCT(FParameters, )
+		SHADER_PARAMETER_STRUCT_REF(FViewUniformShaderParameters, View)
+		SHADER_PARAMETER_RDG_BUFFER_UAV(RWStructuredBuffer<uint>, RWPageStatsBuffer)
+		SHADER_PARAMETER_RDG_TEXTURE(Texture3D<uint>, PageTableCombinedTexture)
+		SHADER_PARAMETER_RDG_TEXTURE(Texture3D<float>, PageAtlasTexture)
+		SHADER_PARAMETER(uint32, ClipmapSizeInPages)
+		SHADER_PARAMETER(uint32, ClipmapIndex)
+		SHADER_PARAMETER(FVector3f, InvPageAtlasSize)
+	END_SHADER_PARAMETER_STRUCT()
+
+	static bool ShouldCompilePermutation(const FGlobalShaderPermutationParameters& Parameters)
+	{
+		return ShouldCompileDistanceFieldShaders(Parameters.Platform);
+	}
+
+	static int32 GetGroupSize()
+	{
+		return 4;
+	}
+
+	static void ModifyCompilationEnvironment(const FGlobalShaderPermutationParameters& Parameters, FShaderCompilerEnvironment& OutEnvironment)
+	{
+		FGlobalShader::ModifyCompilationEnvironment(Parameters, OutEnvironment);
+		OutEnvironment.SetDefine(TEXT("THREADGROUP_SIZE"), GetGroupSize());
+	}
+};
+
+IMPLEMENT_GLOBAL_SHADER(FGlobalDistanceFieldPageStatsCS, "/Engine/Private/DistanceField/GlobalDistanceFieldDebug.usf", "GlobalDistanceFieldPageStatsCS", SF_Compute);
+
 class FGlobalDistanceFieldDebugCS : public FGlobalShader
 {
 	DECLARE_GLOBAL_SHADER(FGlobalDistanceFieldDebugCS)
@@ -1798,6 +1832,7 @@ class FGlobalDistanceFieldDebugCS : public FGlobalShader
 	BEGIN_SHADER_PARAMETER_STRUCT(FParameters, )
 		SHADER_PARAMETER_STRUCT_INCLUDE(ShaderPrint::FShaderParameters, ShaderPrintUniformBuffer)
 		SHADER_PARAMETER_RDG_BUFFER_SRV(StructuredBuffer<uint>, GlobalDistanceFieldPageFreeListAllocatorBuffer)
+		SHADER_PARAMETER_RDG_BUFFER_SRV(StructuredBuffer<uint>, PageStatsBuffer)
 		SHADER_PARAMETER(uint32, GlobalDistanceFieldMaxPageNum)
 	END_SHADER_PARAMETER_STRUCT()
 
@@ -1808,7 +1843,7 @@ class FGlobalDistanceFieldDebugCS : public FGlobalShader
 
 	static int32 GetGroupSize()
 	{
-		return 4;
+		return 64;
 	}
 
 	static void ModifyCompilationEnvironment(const FGlobalShaderPermutationParameters& Parameters, FShaderCompilerEnvironment& OutEnvironment)
@@ -2239,11 +2274,12 @@ void UpdateGlobalDistanceFieldVolume(
 									UTexture2D* VisibilityTexture = It.Key().Visibility;
 
 									FMarkHeightfieldPagesCS::FParameters* PassParameters = GraphBuilder.AllocParameters<FMarkHeightfieldPagesCS::FParameters>();
+									PassParameters->View = View.ViewUniformBuffer;
 									PassParameters->RWMarkedHeightfieldPageBuffer = GraphBuilder.CreateUAV(MarkedHeightfieldPageBuffer, PF_R32_UINT);
 									PassParameters->PageUpdateIndirectArgBuffer = PageUpdateIndirectArgBuffer;
 									PassParameters->PageUpdateTileBuffer = GraphBuilder.CreateSRV(PageUpdateTileBuffer, PF_R32_UINT);
-									PassParameters->PageCoordToPageTranslatedWorldCenterScale = (FVector3f)PageGridCoordToTranslatedWorldCenterScale;
-									PassParameters->PageCoordToPageTranslatedWorldCenterBias = (FVector3f)PageGridCoordToTranslatedWorldCenterBias;
+									PassParameters->PageCoordToVoxelTranslatedCenterScale = (FVector3f)PageCoordToVoxelTranslatedCenterScale;
+									PassParameters->PageCoordToVoxelTranslatedCenterBias = (FVector3f)PageCoordToVoxelTranslatedCenterBias;
 									PassParameters->PageWorldExtent = (FVector3f)PageTileWorldExtentWithoutBorders;
 									PassParameters->ClipmapVoxelExtent = ClipmapVoxelExtent.X;
 									PassParameters->PageGridResolution = PageGridResolution;
@@ -2757,14 +2793,48 @@ void UpdateGlobalDistanceFieldVolume(
 			}
 		}
 
-		if (CVarGlobalDistanceFieldDebug.GetValueOnRenderThread() != 0 && GlobalDistanceFieldInfo.PageFreeListAllocatorBuffer)
+		if (CVarGlobalDistanceFieldDebug.GetValueOnRenderThread() != 0 
+			&& GlobalDistanceFieldInfo.PageFreeListAllocatorBuffer 
+			&& GlobalDistanceFieldInfo.PageTableCombinedTexture
+			&& GlobalDistanceFieldInfo.PageAtlasTexture)
 		{
+			FRDGBufferRef PageStatsBuffer = GraphBuilder.CreateBuffer(FRDGBufferDesc::CreateStructuredDesc(sizeof(uint32), 64), TEXT("GlobalDistanceField.PageStatsBuffer"));
+			AddClearUAVPass(GraphBuilder, GraphBuilder.CreateUAV(PageStatsBuffer), 0);
+
+			FRDGTextureRef PageTableCombinedTexture = GraphBuilder.RegisterExternalTexture(GlobalDistanceFieldInfo.PageTableCombinedTexture, TEXT("GlobalDistanceField.PageTableCombined"));
+			FRDGTextureRef PageAtlasTexture = GraphBuilder.RegisterExternalTexture(GlobalDistanceFieldInfo.PageAtlasTexture, TEXT("GlobalDistanceField.PageAtlasTexture"));
+
+			for (int32 ClipmapIndex = 0; ClipmapIndex < GlobalDistanceFieldInfo.ParameterData.NumGlobalSDFClipmaps; ClipmapIndex++)
+			{
+				FGlobalDistanceFieldPageStatsCS::FParameters* PassParameters = GraphBuilder.AllocParameters<FGlobalDistanceFieldPageStatsCS::FParameters>();
+				PassParameters->RWPageStatsBuffer = GraphBuilder.CreateUAV(PageStatsBuffer);
+				PassParameters->View = View.ViewUniformBuffer;
+				PassParameters->PageTableCombinedTexture = PageTableCombinedTexture;
+				PassParameters->PageAtlasTexture = PageAtlasTexture;
+				PassParameters->ClipmapSizeInPages = GlobalDistanceField::GetPageTableTextureResolution(bLumenEnabled, View.FinalPostProcessSettings.LumenSceneViewDistance).X;
+				PassParameters->ClipmapIndex = ClipmapIndex;
+				PassParameters->InvPageAtlasSize = FVector3f(GlobalDistanceFieldInfo.ParameterData.InvPageAtlasSize);
+
+				auto ComputeShader = View.ShaderMap->GetShader<FGlobalDistanceFieldPageStatsCS>();
+
+				const FIntVector GroupSize = FComputeShaderUtils::GetGroupCount(FIntVector(PassParameters->ClipmapSizeInPages), FGlobalDistanceFieldPageStatsCS::GetGroupSize());
+
+				FComputeShaderUtils::AddPass(
+					GraphBuilder,
+					RDG_EVENT_NAME("PageStats Clipmap:%d", ClipmapIndex),
+					ComputeShader,
+					PassParameters,
+					GroupSize);
+			}
+
+
 			FRDGBufferRef PageFreeListAllocatorBuffer = GraphBuilder.RegisterExternalBuffer(GlobalDistanceFieldInfo.PageFreeListAllocatorBuffer);
 
 			FGlobalDistanceFieldDebugCS::FParameters* PassParameters = GraphBuilder.AllocParameters<FGlobalDistanceFieldDebugCS::FParameters>();
 			ShaderPrint::SetParameters(GraphBuilder, View.ShaderPrintData, PassParameters->ShaderPrintUniformBuffer);
 			PassParameters->GlobalDistanceFieldPageFreeListAllocatorBuffer = GraphBuilder.CreateSRV(PageFreeListAllocatorBuffer, PF_R32_UINT);
 			PassParameters->GlobalDistanceFieldMaxPageNum = View.GlobalDistanceFieldInfo.ParameterData.MaxPageNum;
+			PassParameters->PageStatsBuffer = GraphBuilder.CreateSRV(PageStatsBuffer);
 
 			auto ComputeShader = View.ShaderMap->GetShader<FGlobalDistanceFieldDebugCS>();
 
diff --git a/Engine/Source/Runtime/Renderer/Private/GlobalDistanceFieldHeightfields.h b/Engine/Source/Runtime/Renderer/Private/GlobalDistanceFieldHeightfields.h
--- a/Engine/Source/Runtime/Renderer/Private/GlobalDistanceFieldHeightfields.h
+++ b/Engine/Source/Runtime/Renderer/Private/GlobalDistanceFieldHeightfields.h
@@ -12,11 +12,12 @@ class FMarkHeightfieldPagesCS : public FGlobalShader
 	SHADER_USE_PARAMETER_STRUCT(FMarkHeightfieldPagesCS, FGlobalShader);
 
 	BEGIN_SHADER_PARAMETER_STRUCT(FParameters, )
+		SHADER_PARAMETER_STRUCT_REF(FViewUniformShaderParameters, View)
 		SHADER_PARAMETER_RDG_BUFFER_UAV(RWStructuredBuffer<uint>, RWMarkedHeightfieldPageBuffer)
 		RDG_BUFFER_ACCESS(PageUpdateIndirectArgBuffer, ERHIAccess::IndirectArgs)
 		SHADER_PARAMETER_RDG_BUFFER_SRV(StructuredBuffer<uint>, PageUpdateTileBuffer)
-		SHADER_PARAMETER(FVector3f, PageCoordToPageTranslatedWorldCenterScale)
-		SHADER_PARAMETER(FVector3f, PageCoordToPageTranslatedWorldCenterBias)
+		SHADER_PARAMETER(FVector3f, PageCoordToVoxelTranslatedCenterScale)
+		SHADER_PARAMETER(FVector3f, PageCoordToVoxelTranslatedCenterBias)
 		SHADER_PARAMETER(FVector3f, PageWorldExtent)
 		SHADER_PARAMETER(FVector3f, InvPageGridResolution)
 		SHADER_PARAMETER(FIntVector, PageGridResolution)
