Fix behavior of ThinTranslucent shading model in path tracer where Opacity=1

The expected behavior was for shading to behave as Default Lit where Opacity=1, but it was reverting to Diffuse shading due to a missing case statement in eval/sample.

diff --git a/Engine/Shaders/Private/PathTracing/Material/PathTracingMaterialSampling.ush b/Engine/Shaders/Private/PathTracing/Material/PathTracingMaterialSampling.ush
--- a/Engine/Shaders/Private/PathTracing/Material/PathTracingMaterialSampling.ush
+++ b/Engine/Shaders/Private/PathTracing/Material/PathTracingMaterialSampling.ush
@@ -237,6 +237,7 @@ void AccumulateAlbedoNormal(inout FPathTracingPayload Payload, float3 PathThroug
 	case SHADINGMODELID_SINGLELAYERWATER:
 	case SHADINGMODELID_DEFAULT_LIT:
 	case SHADINGMODELID_CLEAR_COAT:
+	case SHADINGMODELID_THIN_TRANSLUCENT:
 		Albedo += BrdfWeight * PathThroughput * (Payload.DiffuseColor + RoughnessWeight * Payload.SpecularColor);
 		Normal += Luminance(BrdfWeight * PathThroughput) * Payload.WorldNormal;
 		break;
@@ -560,6 +561,7 @@ FMaterialSample SampleMaterial(
 		case SHADINGMODELID_SUBSURFACE:
 		case SHADINGMODELID_PREINTEGRATED_SKIN:
 		case SHADINGMODELID_SINGLELAYERWATER:
+		case SHADINGMODELID_THIN_TRANSLUCENT:
 			Result = DefaultLit_SampleMaterial(RayDirection, Payload, RandSample);
 			break;
 		case SHADINGMODELID_SUBSURFACE_PROFILE:
@@ -663,6 +665,7 @@ FMaterialEval EvalMaterial(
 		case SHADINGMODELID_SUBSURFACE:
 		case SHADINGMODELID_PREINTEGRATED_SKIN:
 		case SHADINGMODELID_SINGLELAYERWATER:
+		case SHADINGMODELID_THIN_TRANSLUCENT:
 			Result = DefaultLit_EvalMaterial(IncomingDirection, OutgoingDirection, Payload);
 			break;
 		case SHADINGMODELID_SUBSURFACE_PROFILE:
