Reenable Strata RT hit lighting on platform supporting DXC.

diff --git a/Engine/Shaders/Private/RayTracing/RayTracingLightingMS.usf b/Engine/Shaders/Private/RayTracing/RayTracingLightingMS.usf
--- a/Engine/Shaders/Private/RayTracing/RayTracingLightingMS.usf
+++ b/Engine/Shaders/Private/RayTracing/RayTracingLightingMS.usf
@@ -102,7 +102,8 @@ RAY_TRACING_ENTRY_MISS(RayTracingLightingMS, FPackedMaterialClosestHitPayload, P
 		GetShadowTerms(FakeSceneDepth, FakePrecomputedShadowFactors, FakeShadingModelID, FakeContactShadowOpacity, LightData, WorldPosition, L, LightAttenuation, Dither, ShadowTerms);
 
 		// RayTracing + Lighting + Strata is too expenssive (as in uses too many registers) to be supported on some platofrms
-		/*FStrataDeferredLighting Lighting = StrataDeferredLighting(
+		#if COMPILER_DXC
+		FStrataDeferredLighting Lighting = StrataDeferredLighting(
 			LightData,
 			V,
 			L,
@@ -112,8 +113,10 @@ RAY_TRACING_ENTRY_MISS(RayTracingLightingMS, FPackedMaterialClosestHitPayload, P
 			PackedPayload.StrataData,
 			StrataAddressing,
 			StrataPixelHeader);
-		LightContribution = Lighting.SceneColor.xyz;*/
+		LightContribution = Lighting.SceneColor.xyz;
+		#else
 		LightContribution = float4(1, 0, 1, 0.5); // Make it return purple so that it is noticeable.
+		#endif
 	}
 #else
 	// LWC_TODO - lighting wants translated world space
