Global Distance Field - fixed debug stats for r.AOGlobalDistanceFieldCacheMostlyStaticSeparately=0

diff --git a/Engine/Shaders/Private/DistanceField/GlobalDistanceFieldDebug.usf b/Engine/Shaders/Private/DistanceField/GlobalDistanceFieldDebug.usf
--- a/Engine/Shaders/Private/DistanceField/GlobalDistanceFieldDebug.usf
+++ b/Engine/Shaders/Private/DistanceField/GlobalDistanceFieldDebug.usf
@@ -39,18 +39,13 @@ void GlobalDistanceFieldPageStatsCS(
 				{
 					for (uint PageCoordX = 0; PageCoordX < GLOBAL_DISTANCE_FIELD_PAGE_RESOLUTION; ++PageCoordX)
 					{
-						float3 PageCoord = uint3(PageCoordX, PageCoordY, PageCoordZ);
-						uint3 PageAtlasOffset = GlobalDistanceFieldPageLinearIndexToPageAtlasOffset(Page);
+						// Don't sample border values which can be safely skipped
+						float3 CoordInPage = uint3(PageCoordX, PageCoordY, PageCoordZ);
+						CoordInPage = clamp(CoordInPage + 0.5f, 1.0f, GLOBAL_DISTANCE_FIELD_PAGE_RESOLUTION - 1.0f);
 
-						float DistanceFieldValue = PageAtlasTexture[PageAtlasOffset + PageCoord].x;
-						if (View.GeneralPurposeTweak > 0.5f)
-						{
-							// Don't sample border values which can be safely skipped
-							PageCoord = clamp(PageCoord + 0.5f, GLOBAL_DISTANCE_FIELD_PAGE_BORDER, GLOBAL_DISTANCE_FIELD_PAGE_RESOLUTION - GLOBAL_DISTANCE_FIELD_PAGE_BORDER);
-
-							float3 PageUV = (PageAtlasOffset + PageCoord) * InvPageAtlasSize;
-							DistanceFieldValue = Texture3DSampleLevel(PageAtlasTexture, GlobalTrilinearWrappedSampler, PageUV, 0).x;
-						}
+						uint3 PageAtlasOffset = GlobalDistanceFieldPageLinearIndexToPageAtlasOffset(Page) * GLOBAL_DISTANCE_FIELD_PAGE_RESOLUTION_IN_ATLAS;
+						float3 PageUV = (PageAtlasOffset + CoordInPage) * InvPageAtlasSize;
+						float DistanceFieldValue = Texture3DSampleLevel(PageAtlasTexture, GlobalTrilinearWrappedSampler, PageUV, 0).x;
 
 						if (DistanceFieldValue > 0.0f && DistanceFieldValue < 1.0f)
 						{
diff --git a/Engine/Source/Runtime/Renderer/Private/GlobalDistanceField.cpp b/Engine/Source/Runtime/Renderer/Private/GlobalDistanceField.cpp
--- a/Engine/Source/Runtime/Renderer/Private/GlobalDistanceField.cpp
+++ b/Engine/Source/Runtime/Renderer/Private/GlobalDistanceField.cpp
@@ -2795,13 +2795,12 @@ void UpdateGlobalDistanceFieldVolume(
 
 		if (CVarGlobalDistanceFieldDebug.GetValueOnRenderThread() != 0 
 			&& GlobalDistanceFieldInfo.PageFreeListAllocatorBuffer 
-			&& GlobalDistanceFieldInfo.PageTableCombinedTexture
 			&& GlobalDistanceFieldInfo.PageAtlasTexture)
 		{
 			FRDGBufferRef PageStatsBuffer = GraphBuilder.CreateBuffer(FRDGBufferDesc::CreateStructuredDesc(sizeof(uint32), 64), TEXT("GlobalDistanceField.PageStatsBuffer"));
 			AddClearUAVPass(GraphBuilder, GraphBuilder.CreateUAV(PageStatsBuffer), 0);
 
-			FRDGTextureRef PageTableCombinedTexture = GraphBuilder.RegisterExternalTexture(GlobalDistanceFieldInfo.PageTableCombinedTexture, TEXT("GlobalDistanceField.PageTableCombined"));
+			FRDGTextureRef PageTableCombinedTexture = GraphBuilder.RegisterExternalTexture(GAOGlobalDistanceFieldCacheMostlyStaticSeparately ? GlobalDistanceFieldInfo.PageTableCombinedTexture : GlobalDistanceFieldInfo.PageTableLayerTextures[0], TEXT("GlobalDistanceField.PageTableCombined"));
 			FRDGTextureRef PageAtlasTexture = GraphBuilder.RegisterExternalTexture(GlobalDistanceFieldInfo.PageAtlasTexture, TEXT("GlobalDistanceField.PageAtlasTexture"));
 
 			for (int32 ClipmapIndex = 0; ClipmapIndex < GlobalDistanceFieldInfo.ParameterData.NumGlobalSDFClipmaps; ClipmapIndex++)
