Fix for Horde DCC populate issue 214185 (Android)
#rnx

diff --git a/Engine/Shaders/Private/MediaShaders.usf b/Engine/Shaders/Private/MediaShaders.usf
--- a/Engine/Shaders/Private/MediaShaders.usf
+++ b/Engine/Shaders/Private/MediaShaders.usf
@@ -413,6 +413,8 @@ void P010_2101010ConvertPS(
 	float2 InUV : TEXCOORD0,
 	out float4 OutColor : SV_Target0)
 {
+	// ES does not allow reuse of the sampler for >1 textures (the shader is unused with that API, though)
+#if COMPILER_GLSL_ES3_1==0
 	// Adjust incoming UVs for Y / UV read
 	float2 ScaledUV_Y = InUV * P010_2101010ConvertUB.UVScaleY;
 	float2 ScaledUV_UV = InUV * P010_2101010ConvertUB.UVScaleUV;
@@ -447,6 +449,9 @@ void P010_2101010ConvertPS(
 
 	// Output with solid alpha
 	OutColor = float4(Color, 1.0f); 
+#else
+	OutColor = float4(1,0,1,1);
+#endif // COMPILER_PSSL
 }
 
 
