Use float3x4 instead of float4[3] for instance transforms. The float4 array was not actually getting passed from vertex to pixel correctly in some cases, instead being initialized to zero.

diff --git a/Engine/Shaders/Private/LocalVertexFactoryCommon.ush b/Engine/Shaders/Private/LocalVertexFactoryCommon.ush
--- a/Engine/Shaders/Private/LocalVertexFactoryCommon.ush
+++ b/Engine/Shaders/Private/LocalVertexFactoryCommon.ush
@@ -58,10 +58,10 @@ struct FVertexFactoryInterpolantsVSToPS
 #endif
 
 #if HAS_INSTANCE_LOCAL_TO_WORLD_PS
-	nointerpolation float4 InstanceLocalToWorld[3] : INSTANCE_LOCAL_TO_WORLD;
+	nointerpolation float3x4 InstanceLocalToWorld : INSTANCE_LOCAL_TO_WORLD;
 #endif
 #if HAS_INSTANCE_WORLD_TO_LOCAL_PS
-	nointerpolation float4 InstanceWorldToLocal[3] : INSTANCE_WORLD_TO_LOCAL;
+	nointerpolation float3x4 InstanceWorldToLocal : INSTANCE_WORLD_TO_LOCAL;
 #endif
 };
 
