Remove NaN-s at the end of the HWRT reflection shader to prevent propagating errors to large areas of the screen via denoising, antialiasing, bloom, etc.

diff --git a/Engine/Shaders/Private/RayTracing/RayTracingDeferredReflections.usf b/Engine/Shaders/Private/RayTracing/RayTracingDeferredReflections.usf
--- a/Engine/Shaders/Private/RayTracing/RayTracingDeferredReflections.usf
+++ b/Engine/Shaders/Private/RayTracing/RayTracingDeferredReflections.usf
@@ -269,6 +269,9 @@ RAY_TRACING_ENTRY_RAYGEN(RayTracingDeferredReflectionsRGS)
 		}
 	}
 
+	// Filter out any NaN-s and negative numbers to prevent any errors propagating into the subsequent pipeline.
+	ResultColor.rgb = -min(-ResultColor.rgb, float3(0, 0, 0));
+
 	ColorOutput[PixelPos] = ResultColor;
 
 	float4 RayParams = (float4)0;
