Show NaNs in HWRT debug visualization mode

diff --git a/Engine/Shaders/Private/RayTracing/RayTracingDebug.usf b/Engine/Shaders/Private/RayTracing/RayTracingDebug.usf
--- a/Engine/Shaders/Private/RayTracing/RayTracingDebug.usf
+++ b/Engine/Shaders/Private/RayTracing/RayTracingDebug.usf
@@ -278,6 +278,15 @@ RAY_TRACING_ENTRY_RAYGEN(RayTracingDebugMainRGS)
 		// Only  when the output is tonemapped
 		Result *= View.PreExposure;
 	}
+
+	if (!IsFinite(Result.r) || !IsFinite(Result.g) || !IsFinite(Result.b))
+	{
+		float T = frac(View.RealTime);
+		if (T > 0.5)
+		{
+			Result.rgb = float3(1,0,1);
+		}
+	}
 	
 	Output[PixelCoord] = Result;
 }
