Initialize virtual texture feedback when rendering deferred decals.
Without this only one VT can provide feedback.

diff --git a/Engine/Shaders/Private/DeferredDecal.usf b/Engine/Shaders/Private/DeferredDecal.usf
--- a/Engine/Shaders/Private/DeferredDecal.usf
+++ b/Engine/Shaders/Private/DeferredDecal.usf
@@ -134,6 +134,10 @@ void FPixelShaderInOut_MainPS(inout FPixelShaderIn In, inout FPixelShaderOut Out
 	MaterialParameters.AbsoluteWorldPosition = MaterialParameters.WorldPosition_NoOffsets = LWCSubtract(TranslatedWorldPosition, PrimaryView.PreViewTranslation);
 	MaterialParameters.WorldPosition_CamRelative = MaterialParameters.WorldPosition_NoOffsets_CamRelative = TranslatedWorldPosition;
 
+#if NUM_VIRTUALTEXTURE_SAMPLES
+	InitializeVirtualTextureFeedback(MaterialParameters.VirtualTextureFeedback, (uint2)In.SvPosition.xy, View.FrameNumber);
+#endif
+
 	FPixelMaterialInputs PixelMaterialInputs;
 	CalcPixelMaterialInputs(MaterialParameters, PixelMaterialInputs);
 	
