Mobile: Tentative fix for GPU hang related to CSM

diff --git a/Engine/Shaders/Private/MobileLightingCommon.ush b/Engine/Shaders/Private/MobileLightingCommon.ush
--- a/Engine/Shaders/Private/MobileLightingCommon.ush
+++ b/Engine/Shaders/Private/MobileLightingCommon.ush
@@ -124,6 +124,10 @@ return ShadowMap;
 #define ENABLE_MOBILE_CSM 1
 #endif
 
+#ifndef MAX_MOBILE_SHADOWCASCADES
+#define MAX_MOBILE_SHADOWCASCADES 4
+#endif 
+
 half MobileDirectionalLightCSM(float2 ScreenPosition, float SceneDepth, inout float ShadowPositionZ)
 {
 	half ShadowMap = 1;
@@ -140,7 +144,7 @@ half MobileDirectionalLightCSM(float2 ScreenPosition, float SceneDepth, inout fl
 	Settings.ProjectionDepthBiasParameters = 0;
 
 	float4 ShadowPosition = float4(0, 0, 0, 0);
-	for (int i = 0; i < MobileDirectionalLight.DirectionalLightShadowMapNumCascades; i++)
+	for (int i = 0; i < MAX_MOBILE_SHADOWCASCADES; i++)
 	{
 		if (SceneDepth < MobileDirectionalLight.DirectionalLightShadowDistances[i])
 		{
diff --git a/Engine/Source/Runtime/Engine/Private/SceneManagement.cpp b/Engine/Source/Runtime/Engine/Private/SceneManagement.cpp
--- a/Engine/Source/Runtime/Engine/Private/SceneManagement.cpp
+++ b/Engine/Source/Runtime/Engine/Private/SceneManagement.cpp
@@ -728,7 +728,6 @@ FMobileDirectionalLightShaderParameters::FMobileDirectionalLightShaderParameters
 	DirectionalLightColor = FLinearColor::Black;
 	DirectionalLightDirectionAndShadowTransition = FVector4f(EForceInit::ForceInitToZero);
 	DirectionalLightShadowMapChannelMask = 0xFF;
-	DirectionalLightShadowMapNumCascades = 0;
 
 	// white texture should act like a shadowmap cleared to the farplane.
 	DirectionalLightShadowTexture = GWhiteTexture->TextureRHI;
diff --git a/Engine/Source/Runtime/Engine/Public/SceneView.h b/Engine/Source/Runtime/Engine/Public/SceneView.h
--- a/Engine/Source/Runtime/Engine/Public/SceneView.h
+++ b/Engine/Source/Runtime/Engine/Public/SceneView.h
@@ -570,7 +570,6 @@ BEGIN_GLOBAL_SHADER_PARAMETER_STRUCT_WITH_CONSTRUCTOR(FMobileDirectionalLightSha
 	SHADER_PARAMETER_EX(FVector4f, DirectionalLightShadowDistances, EShaderPrecisionModifier::Half)
 	SHADER_PARAMETER_ARRAY(FMatrix44f, DirectionalLightScreenToShadow, [MAX_MOBILE_SHADOWCASCADES])
 	SHADER_PARAMETER(uint32, DirectionalLightShadowMapChannelMask)
-	SHADER_PARAMETER(uint32, DirectionalLightShadowMapNumCascades)
 	SHADER_PARAMETER_TEXTURE(Texture2D, DirectionalLightShadowTexture)
 	SHADER_PARAMETER_SAMPLER(SamplerState, DirectionalLightShadowSampler)
 END_GLOBAL_SHADER_PARAMETER_STRUCT()
diff --git a/Engine/Source/Runtime/Renderer/Private/MobileBasePassRendering.cpp b/Engine/Source/Runtime/Renderer/Private/MobileBasePassRendering.cpp
--- a/Engine/Source/Runtime/Renderer/Private/MobileBasePassRendering.cpp
+++ b/Engine/Source/Runtime/Renderer/Private/MobileBasePassRendering.cpp
@@ -292,8 +292,6 @@ void SetupMobileDirectionalLightUniformParameters(
 					OutShadowIndex++;
 				}
 			}
-
-			Params.DirectionalLightShadowMapNumCascades = SceneView.MaxShadowCascades;
 		}
 	}
 }
