Cleanup unused code from RT material shader

The SIMPLIFIED_MATERIAL_SHADER macro is only used by GPU Lightmass, and is now handled purely by the PathTracing material. The old implementation from the RT materials is therefore unused and can be safely removed.

diff --git a/Engine/Shaders/Private/RayTracing/RayTracingMaterialHitShaders.usf b/Engine/Shaders/Private/RayTracing/RayTracingMaterialHitShaders.usf
--- a/Engine/Shaders/Private/RayTracing/RayTracingMaterialHitShaders.usf
+++ b/Engine/Shaders/Private/RayTracing/RayTracingMaterialHitShaders.usf
@@ -10,10 +10,6 @@
 // Needed for VertexFactoryInterpolate to interpolate attributes from vertices to hit point
 #define NEEDS_VERTEX_FACTORY_INTERPOLATION 1
 
-#ifndef SIMPLIFIED_MATERIAL_SHADER
-	#define SIMPLIFIED_MATERIAL_SHADER 0
-#endif
-
 #if STRATA_ENABLED
 #define STRATA_INLINE_SHADING 1
 #endif
@@ -325,8 +321,6 @@ void GetMaterialPayload(
 	half  Specular = GetMaterialSpecular(PixelMaterialInputs);
 	half Roughness = GetMaterialRoughness(PixelMaterialInputs);
 
-#if !SIMPLIFIED_MATERIAL_SHADER
-
 	/**
 	 * Set material attributes for full materials
 	 **/
@@ -464,47 +458,6 @@ void GetMaterialPayload(
 	Payload.SpecularColor = ComputeF0(Specular, BaseColor, Metallic);
 	
 	Payload.IndirectIrradiance += DiffuseIndirectLighting; // Contains both lightmap & skylight
-
-#else // SIMPLIFIED_MATERIAL_SHADER
-
-	/**
-	 * Set material attributes for simplified materials
-	 **/
-	// Anything unsupported will be forced to default lit
-#if MATERIAL_SHADINGMODEL_TWOSIDED_FOLIAGE
-	if (Payload.ShadingModelID == SHADINGMODELID_TWOSIDED_FOLIAGE)
-	{
-		Payload.CustomData = GetMaterialSubsurfaceData(PixelMaterialInputs);
-	}
-	else
-#endif
-#if MATERIAL_SHADINGMODEL_THIN_TRANSLUCENT
-	if (Payload.ShadingModelID == SHADINGMODELID_THIN_TRANSLUCENT)
-	{
-		Payload.CustomData.xyz = GetThinTranslucentMaterialOutput0(MaterialParameters);
-	}
-	else
-#endif
-#if MATERIAL_SHADINGMODEL_UNLIT
-	if (Payload.ShadingModelID == SHADINGMODELID_UNLIT)
-	{
-	}
-	else
-#endif
-	{
-		Payload.ShadingModelID = SHADINGMODELID_DEFAULT_LIT;
-	}
-	Payload.TranslatedWorldPos = LWCToFloat(LWCAdd(MaterialParameters.AbsoluteWorldPosition, ResolvedView.PreViewTranslation));
-	Payload.WorldNormal = normalize(MaterialParameters.WorldNormal);
-	Payload.Radiance = GetMaterialEmissive(PixelMaterialInputs);
-	Payload.BaseColor = BaseColor;
-	Payload.Specular = 0;
-	Payload.Roughness = 1;
-	Payload.Metallic = 0;
-	Payload.Opacity = GetMaterialOpacity(PixelMaterialInputs);
-	Payload.Ior = 1;
-
-#endif // SIMPLIFIED_MATERIAL_SHADER
 }
 #endif // STRATA_ENABLED
 
