Fix compile error due to merge

diff --git a/Engine/Shaders/Private/BasePassPixelShader.usf b/Engine/Shaders/Private/BasePassPixelShader.usf
--- a/Engine/Shaders/Private/BasePassPixelShader.usf
+++ b/Engine/Shaders/Private/BasePassPixelShader.usf
@@ -2185,4 +2185,8 @@ void FPixelShaderInOut_MainPS(
 
 // all PIXELSHADEROUTPUT_ and "void FPixelShaderInOut_MainPS()" need to be setup before this include
 // this include generates the wrapper code to call MainPS(inout FPixelShaderOutput PixelShaderOutput)
+#if COMPUTE_SHADED
+#include "ComputeShaderOutputCommon.ush"
+#else
 #include "PixelShaderOutputCommon.ush"
+#endif
