* Move SimpleLight struct & conversion into LightData.ush
* Move RectTexture data into LightData.ush
* Remove conversion duplicate in VolumetricFog

diff --git a/Engine/Shaders/Private/DeferredLightingCommon.ush b/Engine/Shaders/Private/DeferredLightingCommon.ush
--- a/Engine/Shaders/Private/DeferredLightingCommon.ush
+++ b/Engine/Shaders/Private/DeferredLightingCommon.ush
@@ -27,18 +27,6 @@
 	#endif
 #endif
 
-
-/** Data about a single light to be shaded with the simple shading model, designed for speed and limited feature set. */
-struct FSimpleDeferredLightData
-{
-	float3 TranslatedWorldPosition;
-	float  InvRadius;
-	float3 Color;
-	float  FalloffExponent;
-	/** Whether to use inverse squared falloff. */
-	bool bInverseSquared;
-};
-
 #define REFERENCE_QUALITY	0
 
 FRectTexture InitRectTexture(FDeferredLightData In)
diff --git a/Engine/Shaders/Private/ForwardLightingCommon.ush b/Engine/Shaders/Private/ForwardLightingCommon.ush
--- a/Engine/Shaders/Private/ForwardLightingCommon.ush
+++ b/Engine/Shaders/Private/ForwardLightingCommon.ush
@@ -314,14 +314,8 @@ float3 GetForwardDirectLightingForVertexLighting(uint GridIndex, float3 Translat
 	for (uint LocalLightListIndex = 0; LocalLightListIndex < NumLocalLights; LocalLightListIndex++)
 	{
 		const FLocalLightData LocalLight = GetLocalLightData(CulledLightsGrid.DataStartIndex + LocalLightListIndex, EyeIndex);
-		 
-		FSimpleDeferredLightData LightData = (FSimpleDeferredLightData)0;
-		LightData.TranslatedWorldPosition = LocalLight.LightPositionAndInvRadius.xyz;
-		LightData.InvRadius = LocalLight.LightPositionAndInvRadius.w;
-		LightData.Color = LocalLight.LightColorAndFalloffExponent.xyz;
-		LightData.FalloffExponent = LocalLight.LightColorAndFalloffExponent.w;
-		LightData.bInverseSquared = LightData.FalloffExponent == 0;
-						
+		const FSimpleDeferredLightData LightData = ConvertToSimpleLight(LocalLight);
+
 		// No specular for vertex lighting
 		float3 CameraVector = 0;
 		float3 SpecularColor = 0;
diff --git a/Engine/Shaders/Private/LightData.ush b/Engine/Shaders/Private/LightData.ush
--- a/Engine/Shaders/Private/LightData.ush
+++ b/Engine/Shaders/Private/LightData.ush
@@ -13,8 +13,37 @@
 #define LIGHT_TYPE_SPOT				2 
 #define LIGHT_TYPE_RECT				3 
 #define LIGHT_TYPE_MAX				4 
+
 #define LIGHTING_CHANNEL_MASK		0x7
 
+#define MAX_RECT_ATLAS_MIP			32
+
+///////////////////////////////////////////////////////////////////////////////////////////////////
+// Rect light texture data
+struct FRectTexture
+{
+	half2 AtlasUVOffset;
+	half2 AtlasUVScale;
+	half  AtlasMaxLevel;
+};
+
+FRectTexture InitRectTexture()
+{
+	FRectTexture Out;
+	Out.AtlasUVOffset = 0;
+	Out.AtlasUVScale  = 0;
+	Out.AtlasMaxLevel = MAX_RECT_ATLAS_MIP;
+	return Out;
+}
+
+FRectTexture InitRectTexture(half2 InUVOffset, half2 InUVScale, half InMaxLevel)
+{
+	FRectTexture Out;
+	Out.AtlasUVOffset = InUVOffset;
+	Out.AtlasUVScale  = InUVScale;
+	Out.AtlasMaxLevel = InMaxLevel;
+	return Out;
+}
 
 ///////////////////////////////////////////////////////////////////////////////////////////////////
 // Directional light data
@@ -50,6 +79,20 @@ struct FLocalLightData
 	int  VirtualShadowMapId;
 };
 
+///////////////////////////////////////////////////////////////////////////////////////////////////
+// Simple light data
+
+/** Data about a single light to be shaded with the simple shading model, designed for speed and limited feature set. */
+struct FSimpleDeferredLightData
+{
+	float3 TranslatedWorldPosition;
+	float  InvRadius;
+	float3 Color;
+	float  FalloffExponent;
+	/** Whether to use inverse squared falloff. */
+	bool bInverseSquared;
+};
+
 ///////////////////////////////////////////////////////////////////////////////////////////////////
  // Light-weight light data - Used for raytracing/VSM/Lumen/...
  
@@ -181,7 +224,11 @@ FDeferredLightData ConvertToDeferredLight(
 	Out.bSpotLight				= LightType == LIGHT_TYPE_SPOT;
 	Out.bRectLight				= LightType == LIGHT_TYPE_RECT;
 	Out.HairTransmittance		= InitHairTransmittanceData();
-
+	Out.RectLightBarnCosAngle	= In.RectBarnDoorAndVirtualShadowMapIdAndSpecularScale.x;
+	Out.RectLightBarnLength		= In.RectBarnDoorAndVirtualShadowMapIdAndSpecularScale.y;
+	Out.RectLightAtlasUVOffset	= 0;
+	Out.RectLightAtlasUVScale	= 0;
+	Out.RectLightAtlasMaxLevel	= MAX_RECT_ATLAS_MIP;
 	// Shadow
 	// * Static shadowing uses ShadowMapChannel, 
 	// * Dynamic shadows are packed into light attenuation using PreviewShadowMapChannel
@@ -286,4 +333,16 @@ FDeferredLightData ConvertToDeferredLight(
 	OutLightingChannelMask			= UnpackLightingChannelMask(In.DirectionalLightShadowMapChannelMask);
 
 	return Out;
-}
\ No newline at end of file
+}
+
+// Helper function for converting FLocalLightData into FSimpleDeferredLightData
+FSimpleDeferredLightData ConvertToSimpleLight(FLocalLightData In)
+{
+	FSimpleDeferredLightData Out = (FSimpleDeferredLightData)0;
+	Out.TranslatedWorldPosition = In.LightPositionAndInvRadius.xyz;
+	Out.InvRadius				= In.LightPositionAndInvRadius.w;
+	Out.Color					= In.LightColorAndFalloffExponent.xyz;
+	Out.FalloffExponent			= In.LightColorAndFalloffExponent.w;
+	Out.bInverseSquared			= Out.FalloffExponent == 0;
+	return Out;
+}
diff --git a/Engine/Shaders/Private/RectLight.ush b/Engine/Shaders/Private/RectLight.ush
--- a/Engine/Shaders/Private/RectLight.ush
+++ b/Engine/Shaders/Private/RectLight.ush
@@ -5,13 +5,12 @@
 
 #include "CapsuleLight.ush"
 #include "MonteCarlo.ush"
+#include "LightData.ush"
 
 #ifndef USE_SOURCE_TEXTURE
 #define USE_SOURCE_TEXTURE 1
 #endif
 
-#define MAX_RECT_ATLAS_MIP 32
-
 struct FRect
 {
 	float3		Origin;
@@ -21,31 +20,6 @@ struct FRect
 	float2		Offset;
 };
 
-struct FRectTexture
-{
-	half2 AtlasUVOffset;
-	half2 AtlasUVScale;
-	half  AtlasMaxLevel;
-};
-
-FRectTexture InitRectTexture()
-{
-	FRectTexture Out;
-	Out.AtlasUVOffset = 0;
-	Out.AtlasUVScale = 0;
-	Out.AtlasMaxLevel = MAX_RECT_ATLAS_MIP;
-	return Out;
-}
-
-FRectTexture InitRectTexture(half2 InUVOffset, half2 InUVScale, half InMaxLevel)
-{
-	FRectTexture Out;
-	Out.AtlasUVOffset = InUVOffset;
-	Out.AtlasUVScale  = InUVScale;
-	Out.AtlasMaxLevel = InMaxLevel;
-	return Out;
-}
-
 float3 SampleRectTexture(FRectTexture RectTexture, float2 RectUV, float Level, bool bIsReference = false)
 {
 #if USE_SOURCE_TEXTURE
diff --git a/Engine/Shaders/Private/VolumetricFog.usf b/Engine/Shaders/Private/VolumetricFog.usf
--- a/Engine/Shaders/Private/VolumetricFog.usf
+++ b/Engine/Shaders/Private/VolumetricFog.usf
@@ -690,29 +690,7 @@ void LightScatteringCS(
 
 			if (VolumetricScatteringIntensity > 0)
 			{
-				FDeferredLightData LightData = (FDeferredLightData)0;
-				LightData.TranslatedWorldPosition = LocalLight.LightPositionAndInvRadius.xyz;
-				LightData.InvRadius = LocalLight.LightPositionAndInvRadius.w;
-				LightData.Color = LocalLight.LightColorAndFalloffExponent.xyz;
-				LightData.FalloffExponent = LocalLight.LightColorAndFalloffExponent.w;
-				LightData.Direction = LocalLight.LightDirectionAndShadowMask.xyz;
-				LightData.SpotAngles = LocalLight.SpotAnglesAndSourceRadiusPacked.xy;
-				LightData.SourceRadius = LocalLight.SpotAnglesAndSourceRadiusPacked.z;
-				LightData.SourceLength = f16tof32(asuint(LocalLight.SpotAnglesAndSourceRadiusPacked.w));
-				LightData.Tangent = LocalLight.LightTangentAndSoftSourceRadius.xyz;
-				LightData.SoftSourceRadius = LocalLight.LightTangentAndSoftSourceRadius.w;
-				LightData.bInverseSquared = LightData.FalloffExponent == 0;
-				LightData.RectLightBarnCosAngle = LocalLight.RectBarnDoorAndVirtualShadowMapIdAndSpecularScale.x;
-				LightData.RectLightBarnLength = LocalLight.RectBarnDoorAndVirtualShadowMapIdAndSpecularScale.y;
-				LightData.RectLightAtlasUVOffset = 0;
-				LightData.RectLightAtlasUVScale  = 0;
-				LightData.RectLightAtlasMaxLevel = MAX_RECT_ATLAS_MIP;
-
-				LightData.bRadialLight = true;
-				LightData.bSpotLight = LightData.SpotAngles.x > -2.0f;
-				LightData.bRectLight = LightData.RectLightBarnLength > -2.0f;
-
-				float3 LightColor = LightData.Color;
+				const FDeferredLightData LightData = ConvertToDeferredLight(LocalLight);
 
 				float3 L = 0;
 				float3 ToLight = 0;
@@ -735,7 +713,7 @@ void LightScatteringCS(
 
 				float CombinedAttenuation = Lighting * LightMask;
 
-				LightScattering += LightColor * (PhaseFunction(PhaseG, dot(L, -CameraVector)) * CombinedAttenuation * VolumetricScatteringIntensity);
+				LightScattering += LightData.Color * (PhaseFunction(PhaseG, dot(L, -CameraVector)) * CombinedAttenuation * VolumetricScatteringIntensity);
 
 				// To debug culling
 				//LightScattering += LocalLight.LightColorAndFalloffExponent.xyz * .0000001f;
