Fix merging issue.

diff --git a/Engine/Shaders/Private/BasePassPixelShader.usf b/Engine/Shaders/Private/BasePassPixelShader.usf
--- a/Engine/Shaders/Private/BasePassPixelShader.usf
+++ b/Engine/Shaders/Private/BasePassPixelShader.usf
@@ -985,33 +985,33 @@ void FPixelShaderInOut_MainPS(
 	float DBufferOpacity = 1.0f;
 #if USE_DBUFFER && MATERIALDECALRESPONSEMASK && !MATERIALBLENDING_ANY_TRANSLUCENT && !MATERIAL_SHADINGMODEL_SINGLELAYERWATER
 	// apply decals from the DBuffer
-#if PC_D3D
+	#if PC_D3D
 	//Temporary workaround to avoid crashes on AMD, revert back to BRANCH
 	FLATTEN
-#else
+	#else
 	BRANCH
-#endif
-		if ((GetPrimitiveData(MaterialParameters).Flags & PRIMITIVE_SCENE_DATA_FLAG_DECAL_RECEIVER) != 0 && View.ShowDecalsMask > 0)
-		{
-			uint ValidDBufferTargetMask = GetDBufferTargetMask(uint2(In.SvPosition.xy)) & MATERIALDECALRESPONSEMASK;
+	#endif
+	if ((GetPrimitiveData(MaterialParameters).Flags & PRIMITIVE_SCENE_DATA_FLAG_DECAL_RECEIVER) != 0 && View.ShowDecalsMask > 0)
+	{
+		uint ValidDBufferTargetMask = GetDBufferTargetMask(uint2(In.SvPosition.xy)) & MATERIALDECALRESPONSEMASK;
 
-			BRANCH
-			if (ValidDBufferTargetMask)
-			{
-				float2 BufferUV = SvPositionToBufferUV(In.SvPosition);
-			#if STRATA_ENABLED 
-			   #if STRATA_INLINE_SHADING && !STRATA_OPTIMIZED_UNLIT
-				const FStrataDBuffer StrataBufferData = GetDBufferData(BufferUV, ValidDBufferTargetMask);
-				ApplyDBufferData(StrataBufferData, StrataPixelHeader, StrataData);
-				DBufferOpacity = StrataBufferData.Coverage;
-			   #endif
-			#else
-				FDBufferData DBufferData = GetDBufferData(BufferUV, ValidDBufferTargetMask);
-				ApplyDBufferData(DBufferData, MaterialParameters.WorldNormal, SubsurfaceColor, Roughness, BaseColor, Metallic, Specular);
-				DBufferOpacity = (DBufferData.ColorOpacity + DBufferData.NormalOpacity + DBufferData.RoughnessOpacity) * (1.0f / 3.0f);
-			#endif // STRATA_ENABLED
-			}
+		BRANCH
+		if (ValidDBufferTargetMask)
+		{
+			float2 BufferUV = SvPositionToBufferUV(In.SvPosition);
+		#if STRATA_ENABLED 
+			#if STRATA_INLINE_SHADING && !STRATA_OPTIMIZED_UNLIT
+			const FStrataDBuffer StrataBufferData = GetDBufferData(BufferUV, ValidDBufferTargetMask);
+			ApplyDBufferData(StrataBufferData, StrataPixelHeader, StrataData);
+			DBufferOpacity = StrataBufferData.Coverage;
+			#endif
+		#else
+			FDBufferData DBufferData = GetDBufferData(BufferUV, ValidDBufferTargetMask);
+			ApplyDBufferData(DBufferData, MaterialParameters.WorldNormal, SubsurfaceColor, Roughness, BaseColor, Metallic, Specular);
+			DBufferOpacity = (DBufferData.ColorOpacity + DBufferData.NormalOpacity + DBufferData.RoughnessOpacity) * (1.0f / 3.0f);
+		#endif // STRATA_ENABLED
 		}
+	}
 #endif
 
 	const float BaseMaterialCoverageOverWater = Opacity;
@@ -1647,7 +1647,7 @@ void FPixelShaderInOut_MainPS(
 		FStrataIntegrationSettings Settings = InitStrataIntegrationSettings(false /*bForceFullyRough*/, StrataStruct.bRoughDiffuse, StrataStruct.PeelLayersAboveDepth);
 
 	#if STRATA_ADVANCED_DEBUG_ENABLED
-		Settings.SliceStoringDebugStrataTree = all(uint2(float2(View.CursorPosition) * View.ViewResolutionFraction) == PixelPos) ? StrataStruct.SliceStoringDebugStrataTree : Settings.SliceStoringDebugStrataTree;
+		Settings.SliceStoringDebugStrataTreeData = all(uint2(float2(View.CursorPosition) * View.ViewResolutionFraction) == PixelPos) ? StrataStruct.SliceStoringDebugStrataTree : Settings.SliceStoringDebugStrataTree;
 	#endif
 
 		// Generate the strata material data to write out
@@ -1671,12 +1671,13 @@ void FPixelShaderInOut_MainPS(
 		);
 
 		// Write out MRT data
-	#if STRATA_BASE_PASS_MRT_OUTPUT_COUNT != 2
+	#if STRATA_BASE_PASS_MRT_OUTPUT_COUNT != 3
 	#error Strata STRATA_BASE_PASS_MRT_OUTPUT_COUNT has been update but not StrataOutput
 	#endif
 		Out.StrataOutput[0] = RWStrataMaterialContainer.MaterialRenderTargets[0];
 		Out.StrataOutput[1] = RWStrataMaterialContainer.MaterialRenderTargets[1];
-		Out.StrataOutput[2] = StrataPackTopLayerData(TopLayerData);
+		Out.StrataOutput[2] = RWStrataMaterialContainer.MaterialRenderTargets[2];
+		Out.StrataOutput[3] = StrataPackTopLayerData(TopLayerData);
 	#if STRATA_OUTPUT_ROUGH_REFRACTION
 		OpaqueRoughRefractionData.OpaqueRoughRefractionEnabled = true; // As long as we output to rough refraction, we want it to be enabled for the tiling process.
 	#if STRATA_USES_CONVERSION_FROM_LEGACY
diff --git a/Engine/Shaders/Private/ComputeShaderOutputCommon.ush b/Engine/Shaders/Private/ComputeShaderOutputCommon.ush
--- a/Engine/Shaders/Private/ComputeShaderOutputCommon.ush
+++ b/Engine/Shaders/Private/ComputeShaderOutputCommon.ush
@@ -30,7 +30,7 @@ through a single struct, allowing for a more readable code
 
 #if STRATA_OPAQUE_DEFERRED
 
-	#if STRATA_BASE_PASS_MRT_OUTPUT_COUNT != 2
+	#if STRATA_BASE_PASS_MRT_OUTPUT_COUNT != 3
 		#error Strata STRATA_BASE_PASS_MRT_OUTPUT_COUNT has been updated but not the uint MRTs
 	#endif
 	#if PIXELSHADEROUTPUT_MRT3
@@ -40,14 +40,17 @@ through a single struct, allowing for a more readable code
 		RWTexture2D<uint> StrataOutput3UAV;
 		RWTexture2D<uint> StrataOutput4UAV;
 		RWTexture2D<uint> StrataOutput5UAV;
+		RWTexture2D<uint> StrataOutput6UAV;
 	#elif PIXELSHADEROUTPUT_MRT1
 		RWTexture2D<uint> StrataOutput2UAV;
 		RWTexture2D<uint> StrataOutput3UAV;
 		RWTexture2D<uint> StrataOutput4UAV;
+		RWTexture2D<uint> StrataOutput5UAV;
 	#else
 		RWTexture2D<uint> StrataOutput1UAV;
 		RWTexture2D<uint> StrataOutput2UAV;
 		RWTexture2D<uint> StrataOutput3UAV;
+		RWTexture2D<uint> StrataOutput4UAV;
 	#endif
 
 #else // STRATA_OPAQUE_DEFERRED
