Fix missing VolumetricFog using common DeferredLightData loading from uniforms.

diff --git a/Engine/Shaders/Private/LightDataUniforms.ush b/Engine/Shaders/Private/LightDataUniforms.ush
--- a/Engine/Shaders/Private/LightDataUniforms.ush
+++ b/Engine/Shaders/Private/LightDataUniforms.ush
@@ -43,6 +43,14 @@ FDeferredLightData InitDeferredLightFromUniforms(uint InLightType)
 	return Out;
 }
 
+FDeferredLightData InitDeferredLightFromUniforms()
+{
+	uint LightType = LIGHT_TYPE_POINT;
+	if (DeferredLightUniforms.SpotAngles.x > -2.0f)			{ LightType = LIGHT_TYPE_SPOT; } 
+	if (DeferredLightUniforms.RectLightBarnLength > -2.0f)	{ LightType = LIGHT_TYPE_RECT; }
+	return InitDeferredLightFromUniforms(LightType);
+}
+
 // Special version for mobil. Used only for punctual local light (point or spot)
 FDeferredLightData SetupLightDataForStandardDeferred_Mobile(uint InLightType)
 {
diff --git a/Engine/Shaders/Private/VolumetricFog.usf b/Engine/Shaders/Private/VolumetricFog.usf
--- a/Engine/Shaders/Private/VolumetricFog.usf
+++ b/Engine/Shaders/Private/VolumetricFog.usf
@@ -15,6 +15,7 @@
 #include "VolumetricLightmapShared.ush"
 #include "ForwardShadowingCommon.ush"
 #include "ParticipatingMediaCommon.ush"
+#include "LightDataUniforms.ush"
 #define FrontLayerTranslucencyReflectionsStruct LumenGIVolumeStruct
 #include "Lumen/LumenTranslucencyVolumeShared.ush"
 #include "Random.ush"
@@ -279,30 +280,7 @@ void InjectShadowedLocalLightPS(
 	// This triggers the HistoryMissSuperSampleCount path causing significant overhead for shading off-screen pixels.
 	if (all(int3(GridCoordinate) < VolumetricFog.GridSizeInt))
 	{
-		FDeferredLightData LightData;
-		LightData.TranslatedWorldPosition = DeferredLightUniforms.TranslatedWorldPosition;
-		LightData.InvRadius = DeferredLightUniforms.InvRadius;
-		LightData.Color = DeferredLightUniforms.Color;
-		LightData.FalloffExponent = DeferredLightUniforms.FalloffExponent;
-		LightData.Direction = DeferredLightUniforms.Direction;
-        LightData.Tangent = DeferredLightUniforms.Tangent;
-		LightData.SpotAngles = DeferredLightUniforms.SpotAngles;
-		LightData.SourceRadius = DeferredLightUniforms.SourceRadius;
-		LightData.SourceLength = DeferredLightUniforms.SourceLength;
-        LightData.SoftSourceRadius = DeferredLightUniforms.SoftSourceRadius;
-		LightData.SpecularScale = DeferredLightUniforms.SpecularScale;
-		LightData.ContactShadowLength = abs(DeferredLightUniforms.ContactShadowLength);
-		LightData.ContactShadowLengthInWS = DeferredLightUniforms.ContactShadowLength < 0.0f;
-		LightData.ContactShadowNonShadowCastingIntensity = DeferredLightUniforms.ContactShadowNonShadowCastingIntensity;
-		LightData.DistanceFadeMAD = DeferredLightUniforms.DistanceFadeMAD;
-		LightData.ShadowMapChannelMask = DeferredLightUniforms.ShadowMapChannelMask;
-		LightData.ShadowedBits = DeferredLightUniforms.ShadowedBits;
-		LightData.RectLightBarnCosAngle = DeferredLightUniforms.RectLightBarnCosAngle;
-		LightData.RectLightBarnLength = DeferredLightUniforms.RectLightBarnLength;
-		LightData.bInverseSquared = DeferredLightUniforms.FalloffExponent == 0;
-		LightData.bRadialLight = true;
-		LightData.bSpotLight = LightData.SpotAngles.x > -2.0f;
-		LightData.bRectLight = LightData.RectLightBarnLength > -2.0f;
+		FDeferredLightData LightData = InitDeferredLightFromUniforms();
 
 		float VolumetricScatteringIntensity = DeferredLightUniforms.VolumetricScatteringIntensity;
 
