Fix deferred lighting where the light shape permutation didn't match the light type enum.

diff --git a/Engine/Shaders/Private/DeferredLightPixelShaders.usf b/Engine/Shaders/Private/DeferredLightPixelShaders.usf
--- a/Engine/Shaders/Private/DeferredLightPixelShaders.usf
+++ b/Engine/Shaders/Private/DeferredLightPixelShaders.usf
@@ -52,6 +52,15 @@
 #include "VirtualShadowMaps/VirtualShadowMapMaskBitsCommon.ush"
 #endif
 
+// Remap light shape permutation into light type
+#if LIGHT_SOURCE_SHAPE == 0
+#define CURRENT_LIGHT_TYPE LIGHT_TYPE_DIRECTIONAL
+#elif LIGHT_SOURCE_SHAPE == 2
+#define CURRENT_LIGHT_TYPE LIGHT_TYPE_RECT
+#else // LIGHT_SOURCE_SHAPE == 1
+#define CURRENT_LIGHT_TYPE LIGHT_TYPE_SPOT
+#endif
+
 struct FInputParams
 {
 	float2 PixelPos;
@@ -253,7 +262,7 @@ void DeferredLightPixelMain(
 		const float SceneDepth = CalcSceneDepth(InputParams.ScreenUV);
 		const FDerivedParams DerivedParams = GetDerivedParams(InputParams, SceneDepth);
 
-		FDeferredLightData LightData = InitDeferredLightFromUniforms(LIGHT_SOURCE_SHAPE);
+		FDeferredLightData LightData = InitDeferredLightFromUniforms(CURRENT_LIGHT_TYPE);
 		UpdateLightDataColor(LightData, InputParams, DerivedParams);
 
 		float3 V =-DerivedParams.CameraVector;
@@ -309,7 +318,7 @@ void DeferredLightPixelMain(
 		const float SceneDepth = CalcSceneDepth(InputParams.ScreenUV);
 		const FDerivedParams DerivedParams = GetDerivedParams(InputParams, SceneDepth);
 
-		FDeferredLightData LightData = InitDeferredLightFromUniforms(LIGHT_SOURCE_SHAPE);
+		FDeferredLightData LightData = InitDeferredLightFromUniforms(CURRENT_LIGHT_TYPE);
 		UpdateLightDataColor(LightData, InputParams, DerivedParams);
 
 	 #if USE_HAIR_COMPLEX_TRANSMITTANCE
@@ -392,7 +401,7 @@ void DeferredLightPixelMain(
 	const float3 TranslatedWorldPosition = mul(float4(ScreenPosition * HairScreenSpaceData.GBuffer.Depth, HairScreenSpaceData.GBuffer.Depth, 1), View.ScreenToTranslatedWorld).xyz;
 	const float3 CameraVector = normalize(TranslatedWorldPosition - View.TranslatedWorldCameraOrigin);
 
-	FDeferredLightData LightData = InitDeferredLightFromUniforms(LIGHT_SOURCE_SHAPE);
+	FDeferredLightData LightData = InitDeferredLightFromUniforms(CURRENT_LIGHT_TYPE);
 
 	// todo add light culling/early out
 
