Fix CIS shader compilation with Strata enabled.

diff --git a/Engine/Shaders/Private/Strata/StrataForwardLighting.ush b/Engine/Shaders/Private/Strata/StrataForwardLighting.ush
--- a/Engine/Shaders/Private/Strata/StrataForwardLighting.ush
+++ b/Engine/Shaders/Private/Strata/StrataForwardLighting.ush
@@ -405,7 +405,7 @@ float3 StrataForwardLighting(
 
 					half4 PreviewShadowMapChannelMask = 1;
 					uint  LocalLightChannelMask = LIGHTING_CHANNEL_MASK;
-					const FDeferredLightData LightData = ConvertToDeferredLight(LocalLight, 1.0f /*SpecularScale*/, PreviewShadowMapChannelMask, LocalLightChannelMask);
+					FDeferredLightData LightData = ConvertToDeferredLight(LocalLight, 1.0f /*SpecularScale*/, PreviewShadowMapChannelMask, LocalLightChannelMask);
 					LightData.bRectLight = false; // Rect light not supported in forward. Enabling this is causing the max sampler count to be exceeded.
 
 					// Skip rect light. This is not done explicitely in the legacy path (resulting in no lighting) but we do it for strata in order to avoid visual artefact.
