Reduces sources of extra ghosting in TSR's moire heuristic

diff --git a/Engine/Shaders/Private/TemporalSuperResolution/TSRRejectShading.usf b/Engine/Shaders/Private/TemporalSuperResolution/TSRRejectShading.usf
--- a/Engine/Shaders/Private/TemporalSuperResolution/TSRRejectShading.usf
+++ b/Engine/Shaders/Private/TemporalSuperResolution/TSRRejectShading.usf
@@ -434,11 +434,12 @@ void MainCS(
 	
 	// Analyses the moire pattern overtime.
 	#if CONFIG_MOIRE_DETECTION
+	const tsr_half MoireLumaToChannel = tsr_half(3.0);
+
 	tsr_simd_half4 FinalMoire = tsr_simd_half4::Const(0.0);
 	tsr_simd_half MoireError = tsr_simd_half::Const(0.0);
 	{
 		const tsr_half BlendFinal = tsr_half(0.05);
-		//const tsr_half BlendFinal = tsr_half(0.2);
 		const tsr_half MoireEncodingError = tsr_half(rcp(127.0));
 		const tsr_half MaxPrevTotalVariationCount = tsr_half(20.0);
 
@@ -451,7 +452,7 @@ void MainCS(
 			UNROLL_N(SIMD_SIZE)
 			for (uint SimdIndex = 0; SimdIndex < SIMD_SIZE; SimdIndex++)
 			{
-				const tsr_half3 LumaWeights = tsr_half3(0.333f, 0.333f, 0.333f);
+				const tsr_half3 LumaWeights = rcp(MoireLumaToChannel.xxx);
 				
 				InputLuma.SetElement(SimdIndex, dot(OriginalOpaqueInput.GetElement(SimdIndex).rgb, LumaWeights));
 			}
@@ -462,7 +463,7 @@ void MainCS(
 
 		// Unpack history
 		tsr_simd_half PrevFlickeringHistory = PrevMoire[0];
-		tsr_simd_half PrevGradient = PrevMoire[1] * tsr_half(2.0) - tsr_half(1.0);
+		tsr_simd_half PrevGradient = PrevMoire[1] * tsr_half(255.0 / 127.0) - tsr_half(1.0);
 		//tsr_simd_half PrevTotalVariation = PrevMoire[2] * tsr_simd_half::Const(1.0 - BlendFinal);
 		tsr_simd_half PrevTotalVariation = PrevMoire[2] * tsr_simd_half::Const(tsr_half(1.0) - BlendFinal);
 		tsr_simd_half PrevTotalVariationCount = PrevMoire[3] * MaxPrevTotalVariationCount * tsr_simd_half::Const(tsr_half(1.0) - BlendFinal);
@@ -561,10 +562,10 @@ void MainCS(
 			#else
 				tsr_half LocalCountFadeIn = saturate(LocalNewTotalVariationCount - TotalVariationCountThreshold);
 			#endif
-			tsr_half LocalMoireError = abs(LocalNewTotalVariation * SafeRcp(LocalNewTotalVariationCount)) * LocalCountFadeIn;
+			tsr_half LocalMoireError = (abs(LocalNewTotalVariation * SafeRcp(LocalNewTotalVariationCount)) + LocalNewTotalVariationCount * MoireEncodingError) * LocalCountFadeIn;
 
 			FinalMoire.RegisterArray[SimdIndex][0] = FinalFlickeringHistory.GetElement(SimdIndex);
-			FinalMoire.RegisterArray[SimdIndex][1] = LocalNewGradient * tsr_half(0.5) + tsr_half(0.5);
+			FinalMoire.RegisterArray[SimdIndex][1] = LocalNewGradient * tsr_half(127.0 / 255.0) + tsr_half(127.0 / 255.0);
 			FinalMoire.RegisterArray[SimdIndex][2] = LocalNewTotalVariation;
 			FinalMoire.RegisterArray[SimdIndex][3] = LocalQuantizedNewTotalVariationCount * rcp(MaxPrevTotalVariationCount);
 
@@ -603,21 +604,36 @@ void MainCS(
 		tsr_simd_halfC FilteredBoxMin = Convolve3x3HV<TAccumulatorMin<tsr_simd_halfC>, tsr_simd_halfC>(FilteredInput, bIsValid, GroupThreadIndex) - BackbufferError;
 		tsr_simd_halfC FilteredBoxMax = Convolve3x3HV<TAccumulatorMax<tsr_simd_halfC>, tsr_simd_halfC>(FilteredInput, bIsValid, GroupThreadIndex) + BackbufferError;
 		
+		tsr_simd_halfC FlickeringClampedFilteredHistory = clamp(FilteredHistory, FilteredBoxMin, FilteredBoxMax);
+
 		tsr_simd_halfC ClampedFilteredHistory;
 		#if CONFIG_MOIRE_DETECTION
 		{
+			tsr_simd_halfC MoireErrorSize;
+			
 			UNROLL_N(SIMD_SIZE)
 			for (uint SimdIndex = 0; SimdIndex < SIMD_SIZE; SimdIndex++)
+			#if 0
 			{
-				// TODO: improve this 3.0 *
-				ClampedFilteredHistory.SetElement(SimdIndex, clamp(
-					FilteredHistory.GetElement(SimdIndex),
-					FilteredBoxMin.GetElement(SimdIndex) - tsr_half(3.0) * MoireError.GetElement(SimdIndex),
-					FilteredBoxMax.GetElement(SimdIndex) + tsr_half(3.0) * MoireError.GetElement(SimdIndex)));
+				tsr_halfC ClampedEnergy = abs(FlickeringClampedFilteredHistory.GetElement(SimdIndex) - FilteredHistory.GetElement(SimdIndex));
+				tsr_halfC ClampedEnergyAmount = ClampedEnergy * SafeRcp(ClampedEnergy[0] + ClampedEnergy[1] + ClampedEnergy[2]);
+
+				// Apply the moire error on each individual channel based on how much they need to clamp.
+				MoireErrorSize.SetElement(SimdIndex, ClampedEnergyAmount * (MoireLumaToChannel * MoireError.GetElement(SimdIndex)));
 			}
+			#else
+			{
+				MoireErrorSize.SetElement(SimdIndex, MoireLumaToChannel * MoireError.GetElement(SimdIndex));
+			}
+			#endif
+
+			tsr_simd_halfC StableFilteredBoxMin = min(FilteredBoxMin, FilteredBoxMax - MoireErrorSize);
+			tsr_simd_halfC StableFilteredBoxMax = max(FilteredBoxMax, FilteredBoxMin + MoireErrorSize);
+			
+			ClampedFilteredHistory = clamp(FilteredHistory, StableFilteredBoxMin, StableFilteredBoxMax);
 		}
 		#else
-			ClampedFilteredHistory = clamp(FilteredHistory, FilteredBoxMin, FilteredBoxMax);
+			ClampedFilteredHistory = FlickeringClampedFilteredHistory;
 		#endif
 		
 		tsr_simd_halfC OriginalBoxMin = Convolve3x3HV<TAccumulatorMin<tsr_simd_halfC>, tsr_simd_halfC>(OriginalInput, bIsValid, GroupThreadIndex);
