Fix rotation and scale precision in TransformLocalToTranslatedWorld. This mimics what UE4 did to get the benefits of TranslatedWorldSpace. Issue was introduced when this was ported to LWC for 5.0.

Tested Rendering EngineTests

diff --git a/Engine/Shaders/Private/VertexFactoryCommon.ush b/Engine/Shaders/Private/VertexFactoryCommon.ush
--- a/Engine/Shaders/Private/VertexFactoryCommon.ush
+++ b/Engine/Shaders/Private/VertexFactoryCommon.ush
@@ -124,9 +124,10 @@ FLWCVector3 TransformLocalToWorld(float3 LocalPosition)
 
 float4 TransformLocalToTranslatedWorld(float3 LocalPosition, FLWCMatrix LocalToWorld)
 {
-	FLWCVector3 WorldPosition = LWCMultiply(LocalPosition, LocalToWorld);
-	FLWCVector3 TranslatedWorldPosition = LWCAdd(WorldPosition, ResolvedView.PreViewTranslation);
-	return float4(LWCToFloat(TranslatedWorldPosition), 1.0f);
+	// Multiply local position with LocalToWorld 3x3 first for high precision rotation and scale, then move the world position part to translated world before adding it  
+	float3 RotatedScaledPosition = (LocalPosition.xxx * LocalToWorld.M[0].xyz + LocalPosition.yyy * LocalToWorld.M[1].xyz + LocalPosition.zzz * LocalToWorld.M[2].xyz);
+	FLWCVector3 TranslatedWorldPositionOrigin = LWCAdd(LWCGetOrigin(LocalToWorld), ResolvedView.PreViewTranslation);
+	return float4(RotatedScaledPosition + LWCToFloat(TranslatedWorldPositionOrigin), 1.0f);
 }
 
 float4 TransformLocalToTranslatedWorld(float3 LocalPosition, float4x4 LocalToWorld)
@@ -136,9 +137,10 @@ float4 TransformLocalToTranslatedWorld(float3 LocalPosition, float4x4 LocalToWor
 
 float4 TransformPreviousLocalPositionToTranslatedWorld(float3 PrevLocalPosition, FLWCMatrix PrevLocalToWorld)
 {
-	FLWCVector3 PrevWorldPosition = LWCMultiply(PrevLocalPosition, PrevLocalToWorld);
-	FLWCVector3 PrevTranslatedWorldPosition = LWCAdd(PrevWorldPosition, ResolvedView.PrevPreViewTranslation);
-	return float4(LWCToFloat(PrevTranslatedWorldPosition), 1.0f);
+	// Multiply local position with PrevLocalToWorld 3x3 first for high precision rotation and scale, then move the world position part to translated world before adding it  
+	float3 RotatedScaledPosition = (PrevLocalPosition.xxx * PrevLocalToWorld.M[0].xyz + PrevLocalPosition.yyy * PrevLocalToWorld.M[1].xyz + PrevLocalPosition.zzz * PrevLocalToWorld.M[2].xyz);
+	FLWCVector3 TranslatedPreviousWorldPositionOrigin = LWCAdd(LWCGetOrigin(PrevLocalToWorld), ResolvedView.PrevPreViewTranslation);
+	return float4(RotatedScaledPosition + LWCToFloat(TranslatedPreviousWorldPositionOrigin), 1.0f);
 }
 
 float4 TransformPreviousLocalPositionToTranslatedWorld(float3 PrevLocalPosition, float4x4 PrevLocalToWorld)
