Fixes excessive shading rejection on frame after shading change

diff --git a/Engine/Shaders/Private/TemporalSuperResolution/TSRRejectShading.usf b/Engine/Shaders/Private/TemporalSuperResolution/TSRRejectShading.usf
--- a/Engine/Shaders/Private/TemporalSuperResolution/TSRRejectShading.usf
+++ b/Engine/Shaders/Private/TemporalSuperResolution/TSRRejectShading.usf
@@ -669,26 +669,22 @@ void MainCS(
 	{
 		const tsr_half CurrentFrameSampleCount = tsr_half(1.0);
 		const tsr_half PrevLowFrequencySampleCount = tsr_half(8.0);
+		const tsr_half TheoricBlendFactor = CurrentFrameSampleCount / (CurrentFrameSampleCount + PrevLowFrequencySampleCount);
 		
-		// Reject the history based on the low frequencies.
-		tsr_simd_half PrevLowFrequencyRejectionFactor;
-		#if 1
-		{
-			PrevLowFrequencyRejectionFactor = saturate(Rejection * ((CurrentFrameSampleCount + PrevLowFrequencySampleCount) * rcp(PrevLowFrequencySampleCount)));
-		}
-		#endif
-		
-		tsr_simd_half BlendFinal = rcp(PrevLowFrequencyRejectionFactor * PrevLowFrequencySampleCount + CurrentFrameSampleCount) * CurrentFrameSampleCount;
-		
+		tsr_simd_half BlendFinal;
 		// Completely discard the history
 		UNROLL_N(SIMD_SIZE)
 		for (uint SimdIndex = 0; SimdIndex < SIMD_SIZE; SimdIndex++)
 		{
 			bool bIsOffScreen = false;
 			bool bIsDisoccluded = bIsValid.GetElement(SimdIndex) == tsr_half(0.0);
-			tsr_half TotalBlend = select(or(bIsOffScreen, bIsDisoccluded), tsr_half(1.0), saturate(tsr_half(1.0) - Rejection.GetElement(SimdIndex) * tsr_half(4.0)));
+			tsr_half LocalRejection = Rejection.GetElement(SimdIndex);
+
+			tsr_half TotalBlend = select(or(bIsOffScreen, bIsDisoccluded), tsr_half(1.0), saturate(tsr_half(1.0) - LocalRejection * tsr_half(4.0)));
+			tsr_half RejectedBlendFactor = max(TheoricBlendFactor, tsr_half(1.0) - LocalRejection);
+
 
-			BlendFinal.SetElement(SimdIndex, max(BlendFinal.GetElement(SimdIndex), TotalBlend));
+			BlendFinal.SetElement(SimdIndex, max(RejectedBlendFactor, TotalBlend));
 		}
 		
 		FinalGuide = OriginalInput * tsr_simd_halfC::Vectorize(BlendFinal) - OriginalHistory * tsr_simd_halfC::Vectorize(BlendFinal - 1.0);
