Fixes preexposure correction of TSR's history guide

diff --git a/Engine/Shaders/Private/TemporalSuperResolution/TSRDecimateHistory.usf b/Engine/Shaders/Private/TemporalSuperResolution/TSRDecimateHistory.usf
--- a/Engine/Shaders/Private/TemporalSuperResolution/TSRDecimateHistory.usf
+++ b/Engine/Shaders/Private/TemporalSuperResolution/TSRDecimateHistory.usf
@@ -336,6 +336,13 @@ void MainCS(
 			}
 			ReprojectedHistoryGuideColor *= tsr_half(Samples.FinalMultiplier);
 
+			// Correct pre exposure of the guide color.
+			{
+				tsr_halfC LinearReprojectedHistoryGuideColor = ReprojectedHistoryGuideColor * rcp(tsr_half(1.0) - ReprojectedHistoryGuideColor) * tsr_half(HistoryPreExposureCorrection);
+				
+				ReprojectedHistoryGuideColor = LinearReprojectedHistoryGuideColor * rcp(LinearReprojectedHistoryGuideColor + tsr_half(1.0));
+			}
+
 			#if CONFIG_MOIRE_REPROJECTION
 			{
 				uint2 Random = Rand3DPCG16(int3(InputPixelPos - InputInfo_ViewportMin, View.StateFrameIndexMod8)).xy;
