Fix typo for rect atlas encoding.

diff --git a/Engine/Shaders/Private/LightData.ush b/Engine/Shaders/Private/LightData.ush
--- a/Engine/Shaders/Private/LightData.ush
+++ b/Engine/Shaders/Private/LightData.ush
@@ -46,11 +46,11 @@ struct FRectLightData
 FRectLightData UnpackRectLightData(uint In0, uint In1, uint In2)
 {
 	FRectLightData Out;
-	Out.BarnLength				= f16tof32(In0 & 0xFFFF);			// 16 bits
-	Out.BarnCosAngle			= ((In0 >> 16) & 0x3FF) / 1024.f;	// 10 bits
-	Out.AtlasData.AtlasMaxLevel =  (In0 >> 26);						//  6 bits
-	Out.AtlasData.AtlasUVOffset	= half2(f16tof32(In1 & 0xFFFF), f16tof32(In1 >> 16));
-	Out.AtlasData.AtlasUVScale	= half2(f16tof32(In2 & 0xFFFF), f16tof32(In2 >> 16));
+	Out.AtlasData.AtlasUVOffset	= half2(f16tof32(In0 & 0xFFFF), f16tof32(In0 >> 16));
+	Out.AtlasData.AtlasUVScale	= half2(f16tof32(In1 & 0xFFFF), f16tof32(In1 >> 16));
+	Out.BarnLength				= f16tof32(In2 & 0xFFFF);			// 16 bits
+	Out.BarnCosAngle			= ((In2 >> 16) & 0x3FF) / 1024.f;	// 10 bits
+	Out.AtlasData.AtlasMaxLevel =  (In2 >> 26);						//  6 bits
 	return Out;
 }
 
