Add optional rect light support for translucent lit surface.

This is enabled with the CVar r.rectlightatlas.translucent. There is currently a limitation on DX11 due to limited number of texture sampler (max 16).

diff --git a/Engine/Shaders/Private/ForwardLightingCommon.ush b/Engine/Shaders/Private/ForwardLightingCommon.ush
--- a/Engine/Shaders/Private/ForwardLightingCommon.ush
+++ b/Engine/Shaders/Private/ForwardLightingCommon.ush
@@ -237,7 +237,8 @@ FDeferredLightingSplit GetForwardDirectLightingSplit(
 		uint  LocalLightingChannelMask = LIGHTING_CHANNEL_MASK;
 		const FLocalLightData LocalLight = GetLocalLightData(CulledLightsGrid.DataStartIndex + LocalLightListIndex, EyeIndex);
 		FDeferredLightData LightData = ConvertToDeferredLight(LocalLight, SpecularScale, PreviewShadowMapChannelMask, LocalLightingChannelMask);
-		LightData.bRectLight = false;	// Rect light not supported in forward. Enabling this is causing the max sampler count to be exceeded.
+		// Rect light optionally supported in forward.
+		LightData.bRectLight = LightData.bRectLight && SUPPORT_RECTLIGHT_ON_FORWARD_LIT_TRANSLUCENT;
 		#if USE_HAIR_COMPLEX_TRANSMITTANCE
 		if (GBufferData.ShadingModelID == SHADINGMODELID_HAIR)
 		{
diff --git a/Engine/Shaders/Private/Strata/StrataForwardLighting.ush b/Engine/Shaders/Private/Strata/StrataForwardLighting.ush
--- a/Engine/Shaders/Private/Strata/StrataForwardLighting.ush
+++ b/Engine/Shaders/Private/Strata/StrataForwardLighting.ush
@@ -81,7 +81,6 @@ float3 StrataForwardLighting(
 #if FORWARD_PER_PIXEL_SHADING
 	// Create the common directional light data use for all layers
 	const FDirectionalLightData DirectionalLightData = GetDirectionalLightData(EyeIndex);
-	FRectTexture RectTexture = InitRectTexture();
 	
 	float4 PreviewShadowMapChannelMask = 1;
 	uint DirLightingChannelMask = LIGHTING_CHANNEL_MASK;
@@ -373,7 +372,7 @@ float3 StrataForwardLighting(
 							ToDirLight,
 							DirLightMask,
 							DirLightAttenuation,
-							RectTexture,
+							InitRectTexture(),
 							DirLightingChannelMask,
 							PrimitiveLightingChannelMask,
 							TranslatedWorldPosition,
@@ -406,11 +405,12 @@ float3 StrataForwardLighting(
 					half4 PreviewShadowMapChannelMask = 1;
 					uint  LocalLightChannelMask = LIGHTING_CHANNEL_MASK;
 					FDeferredLightData LightData = ConvertToDeferredLight(LocalLight, 1.0f /*SpecularScale*/, PreviewShadowMapChannelMask, LocalLightChannelMask);
-					LightData.bRectLight = false; // Rect light not supported in forward. Enabling this is causing the max sampler count to be exceeded.
-
-					// Skip rect light. This is not done explicitely in the legacy path (resulting in no lighting) but we do it for strata in order to avoid visual artefact.
-					// This can be removed when LocalLight properly supports rect lights.
-					if (LightData.bRectLight)
+					const FRectTexture LocalRectTexture = ConvertToRectTexture(LightData);
+					
+					// Rect light optionally supported in forward.
+					// * This is not done explicitely in the legacy path (resulting in no lighting) but we do it for strata in order to avoid visual artefact.
+					// * This can be removed when LocalLight properly supports rect lights.
+					if (LightData.bRectLight && !SUPPORT_RECTLIGHT_ON_FORWARD_LIT_TRANSLUCENT)
 					{
 						continue;
 					}
@@ -438,7 +438,7 @@ float3 StrataForwardLighting(
 						ToLight,
 						LightMask,
 						LightAttenuation,
-						RectTexture,
+						LocalRectTexture,
 						LocalLightChannelMask,
 						PrimitiveLightingChannelMask,
 						TranslatedWorldPosition,
diff --git a/Engine/Source/Runtime/Engine/Private/ShaderCompiler/ShaderCompiler.cpp b/Engine/Source/Runtime/Engine/Private/ShaderCompiler/ShaderCompiler.cpp
--- a/Engine/Source/Runtime/Engine/Private/ShaderCompiler/ShaderCompiler.cpp
+++ b/Engine/Source/Runtime/Engine/Private/ShaderCompiler/ShaderCompiler.cpp
@@ -6180,6 +6180,12 @@ void GlobalBeginCompileShader(
 		Input.Environment.SetDefine(TEXT("SUPPORT_CLOUD_SHADOW_ON_SINGLE_LAYER_WATER"), bSupportCloudShadowOnSingleLayerWater ? 1 : 0);
 	}
 
+	{
+		static IConsoleVariable* CVar = IConsoleManager::Get().FindConsoleVariable(TEXT("r.RectLightAtlas.Translucent"));
+		const bool bSupportRectLlightOnTranslucentSurface = CVar && CVar->GetInt() > 0;
+		Input.Environment.SetDefine(TEXT("SUPPORT_RECTLIGHT_ON_FORWARD_LIT_TRANSLUCENT"), bSupportRectLlightOnTranslucentSurface ? 1 : 0);
+	}
+
 	{
 		static IConsoleVariable* CVarStrata = IConsoleManager::Get().FindConsoleVariable(TEXT("r.Strata"));
 		const bool bStrata = CVarStrata && CVarStrata->GetInt() != 0;
diff --git a/Engine/Source/Runtime/Renderer/Private/RectLightTextureManager.cpp b/Engine/Source/Runtime/Renderer/Private/RectLightTextureManager.cpp
--- a/Engine/Source/Runtime/Renderer/Private/RectLightTextureManager.cpp
+++ b/Engine/Source/Runtime/Renderer/Private/RectLightTextureManager.cpp
@@ -51,6 +51,12 @@ static TAutoConsoleVariable<int32> CVarRectLighForceUpdate(
 	TEXT("Force rect. light atlas update very frame."),
 	ECVF_RenderThreadSafe);
 
+static TAutoConsoleVariable<int32> CVarRectLighTranslucent(
+	TEXT("r.RectLightAtlas.Translucent"),
+	0,
+	TEXT("Enable rect light support for translucent surfaces."),
+	ECVF_RenderThreadSafe);
+
 namespace RectLightAtlas
 {
 
