Disable by default SSS diffusion controlled by pixel footprint as this currently disable the transmission part, while it should.

Diffusion & transmittance needs to be decoupled prior to re-enabling this.

diff --git a/Engine/Shaders/Private/Strata/Strata.ush b/Engine/Shaders/Private/Strata/Strata.ush
--- a/Engine/Shaders/Private/Strata/Strata.ush
+++ b/Engine/Shaders/Private/Strata/Strata.ush
@@ -101,6 +101,9 @@
 #define STRATA_EPSILON						1e-10f
 #define STRATA_EYE_DEFAULT_F0				0.028f
 
+// Disable SSS diffusion based when the pixel's footprint is larger than x times the MFPs
+#define STRATA_DISABLE_SSS_DIFFUSION_BASED_ON_PIXEL_FOOTPRINT 0
+
 ///////////////////////////////////////////////////////////////////////////////
 // BSDFs representation
 
@@ -1359,6 +1362,7 @@ FStrataData GetStrataSlabBSDF(
 		Haziness.bSimpleClearCoat = SecondRoughnessAsSimpleClearCoat > 0.0f;
 	}
 
+	#if STRATA_DISABLE_SSS_DIFFUSION_BASED_ON_PIXEL_FOOTPRINT
 	// If MFP is smaller than the pixel's footprint, remove SSS attributes
 	// The scaling factor 'SSSFootprintScale' is based on empiric tests
 	const float SSSFootprintScale = 3.f;
@@ -1370,6 +1374,7 @@ FStrataData GetStrataSlabBSDF(
 		// MFPPlugged is not cleared as it is used for thin surface and thin volumetric
 		// bMFPPlugged = false;
 	}
+	#endif
 
 	// Convert MFP be expressed relative to simple volume thickness
 	if (!STRATA_FASTPATH && bHasSSSWrap)
