HeterogeneousVolumes: Small fixes

diff --git a/Engine/Shaders/Private/HeterogeneousVolumes/HeterogeneousVolumesRayMarchingUtils.ush b/Engine/Shaders/Private/HeterogeneousVolumes/HeterogeneousVolumesRayMarchingUtils.ush
--- a/Engine/Shaders/Private/HeterogeneousVolumes/HeterogeneousVolumesRayMarchingUtils.ush
+++ b/Engine/Shaders/Private/HeterogeneousVolumes/HeterogeneousVolumesRayMarchingUtils.ush
@@ -361,7 +361,7 @@ void RayMarchSingleScattering(
 			float3 CameraVector = normalize(RayMarchingContext.WorldRayDirection);
 			float PhaseG = 0.0;
 			RotatedHGZonalHarmonic.V = float4(1.0f, CameraVector.y, CameraVector.z, CameraVector.x) * float4(1.0f, PhaseG, PhaseG, PhaseG);
-			float3 IndirectInscattering = max(DotSH(TranslucencyGISH, RotatedHGZonalHarmonic), 0) / PI;
+			float3 IndirectInscattering = max(DotSH(TranslucencyGISH, RotatedHGZonalHarmonic), 0);
 
 			float3 Extinction = SampleExtinction(SampleContext);
 			float3 Albedo = SampleAlbedo(SampleContext);
diff --git a/Engine/Shaders/Private/HeterogeneousVolumes/HeterogeneousVolumesTransmittanceVolumeUtils.ush b/Engine/Shaders/Private/HeterogeneousVolumes/HeterogeneousVolumesTransmittanceVolumeUtils.ush
--- a/Engine/Shaders/Private/HeterogeneousVolumes/HeterogeneousVolumesTransmittanceVolumeUtils.ush
+++ b/Engine/Shaders/Private/HeterogeneousVolumes/HeterogeneousVolumesTransmittanceVolumeUtils.ush
@@ -21,7 +21,7 @@ float3 SampleLightingCache(float3 UVW, float MipLevel)
 		LightingCacheTexture,
 		View.SharedTrilinearClampedSampler,
 		UVW,
-		MipLevel).r;
+		MipLevel);
 
 	return Transmittance;
 }
