Fix shader compilation issue with Strata on translucent surface in dx11 due to texture sampler limit even on dx11.

diff --git a/Engine/Shaders/Private/Strata/StrataForwardLighting.ush b/Engine/Shaders/Private/Strata/StrataForwardLighting.ush
--- a/Engine/Shaders/Private/Strata/StrataForwardLighting.ush
+++ b/Engine/Shaders/Private/Strata/StrataForwardLighting.ush
@@ -413,6 +413,11 @@ float3 StrataForwardLighting(
 					{
 						continue;
 					}
+
+					// When disabled, force the compiler to remove the texture atlas sampling to avoid taking an extra texture sampler
+					#if !SUPPORT_RECTLIGHT_ON_FORWARD_LIT_TRANSLUCENT
+					LightData.bRectLight = false;
+					#endif
 					const FRectTexture LocalRectTexture = ConvertToRectTexture(LightData);
 
 					float DynamicShadowing = dot(PreviewShadowMapChannelMask, DynamicShadowFactors);
