* Move SSS transmission code into a single place (merge duplicate between PCF & VSM)
 * Add shadow/transmission blurring on SSS surfae with VSM to reduce hard-shadow (directional light only for now).
 * Change SubSurfaceModel to have extinction-like color-shifting.

diff --git a/Engine/Shaders/Private/ParticipatingMediaCommon.ush b/Engine/Shaders/Private/ParticipatingMediaCommon.ush
--- a/Engine/Shaders/Private/ParticipatingMediaCommon.ush
+++ b/Engine/Shaders/Private/ParticipatingMediaCommon.ush
@@ -172,7 +172,10 @@ float3 TransmittanceToMeanFreePath(in float3 TransmittanceColor, in float Thickn
 	return 1.0f / max(PARTICIPATING_MEDIA_MIN_EXTINCTION, TransmittanceToExtinction(TransmittanceColor, ThicknessMeters));
 }
 
-
+float3 ExtinctionToTransmittance(in float3 Extinction, in float ThicknessMeters)
+{
+	return exp(-Extinction * ThicknessMeters);
+}
 
 //---------------------------------------------
 // Lighting evaluation function for an Isotropic participating media a slab of 1 meter
diff --git a/Engine/Shaders/Private/RayTracing/RayTracingOcclusionRGS.usf b/Engine/Shaders/Private/RayTracing/RayTracingOcclusionRGS.usf
--- a/Engine/Shaders/Private/RayTracing/RayTracingOcclusionRGS.usf
+++ b/Engine/Shaders/Private/RayTracing/RayTracingOcclusionRGS.usf
@@ -657,7 +657,7 @@ RAY_TRACING_ENTRY_RAYGEN(OcclusionRGS)
 
 			if (ENABLE_TRANSMISSION)
 			{
-				TransmissionParameters = GetTransmissionProfileParams(GBufferData);
+				TransmissionParameters = GetTransmissionProfileParams(ExtractSubsurfaceProfileInt(GBufferData));
 			}
 
 			if (ShadingParameters.bIsSubsurface)
diff --git a/Engine/Shaders/Private/ShadingModels.ush b/Engine/Shaders/Private/ShadingModels.ush
--- a/Engine/Shaders/Private/ShadingModels.ush
+++ b/Engine/Shaders/Private/ShadingModels.ush
@@ -11,6 +11,8 @@
 #include "TransmissionCommon.ush"
 #include "HairBsdf.ush"
 #include "ShadingEnergyConservation.ush"
+#include "ParticipatingMediaCommon.ush"
+#include "ColorSpace.ush"
 
 #if SHADING_PATH_MOBILE
 #include "MobileGGX.ush"
@@ -711,12 +713,13 @@ FDirectLighting SubsurfaceProfileBxDF( FGBufferData GBuffer, half3 N, half3 V, h
 
 #if USE_TRANSMISSION
 
-	FTransmissionProfileParams TransmissionParams = GetTransmissionProfileParams( GBuffer );
+	const uint ProfileId = ExtractSubsurfaceProfileInt(GBuffer);
+	FTransmissionProfileParams TransmissionParams = GetTransmissionProfileParams(ProfileId);
 
 	float Thickness = Shadow.TransmissionThickness;
 	Thickness = DecodeThickness(Thickness);
 	Thickness *= SSSS_MAX_TRANSMISSION_PROFILE_DISTANCE;
-	float3 Profile = GetTransmissionProfile(GBuffer, Thickness).rgb;
+	float3 Profile = GetTransmissionProfile(ProfileId, Thickness).rgb;
 
 	float3 RefracV = refract(V, -N, TransmissionParams.OneOverIOR);
 	float PhaseFunction = ApproximateHG( dot(-L, RefracV), TransmissionParams.ScatteringDistribution );
@@ -776,25 +779,38 @@ FDirectLighting ClothBxDF( FGBufferData GBuffer, half3 N, half3 V, half3 L, floa
 	return Lighting;
 }
 
-FDirectLighting SubsurfaceBxDF( FGBufferData GBuffer, half3 N, half3 V, half3 L, float Falloff, half NoL, FAreaLight AreaLight, FShadowTerms Shadow )
+FDirectLighting SubsurfaceBxDF(FGBufferData GBuffer, half3 N, half3 V, half3 L, float Falloff, half NoL, FAreaLight AreaLight, FShadowTerms Shadow 
 {
-	FDirectLighting Lighting = DefaultLitBxDF( GBuffer, N, V, L, Falloff, NoL, AreaLight, Shadow );
-	
+	FDirectLighting Lighting = DefaultLitBxDF( GBuffer, N, V, L, Falloff, NoL, AreaLight, Shadow);
+
 	half3 SubsurfaceColor = ExtractSubsurfaceColor(GBuffer);
 	half Opacity = GBuffer.CustomData.a;
 
-	half3 H = normalize(V + L);
-
 	// to get an effect when you see through the material
 	// hard coded pow constant
 	half InScatter = pow(saturate(dot(L, -V)), 12) * lerp(3, .1f, Opacity);
-	// wrap around lighting, /(PI*2) to be energy consistent (hack do get some view dependnt and light dependent effect)
-	// Opacity of 0 gives no normal dependent lighting, Opacity of 1 gives strong normal contribution
-	half NormalContribution = saturate(dot(N, H) * Opacity + 1 - Opacity);
-	half BackScatter = GBuffer.GBufferAO * NormalContribution / (PI * 2);
-	
+
+	// Wrap around lighting, 
+	// * /(PI*2) to be energy consistent (hack do get some view dependnt and light dependent effect)
+	// * Opacity of 0 gives no normal dependent lighting, Opacity of 1 gives strong normal contribution	
+	// * Simplified version (with w=.5,  n=1.5): 
+	//     half WrappedDiffuse = 2 * pow(saturate((dot(N, L) + w) / (1.0f + w)), n) * (n + 1) / (2 * (1 + w));
+	//     NormalContribution = WrappedDiffuse * Opacity + 1 - Opacity;
+	const half WrappedDiffuse = pow(saturate(dot(N, L) * (1.f / 1.5f) + (0.5f / 1.5f)), 1.5f) * (2.5f / 1.5f);
+	const half NormalContribution = lerp(1.f, WrappedDiffuse, Opacity);
+	const half BackScatter = GBuffer.GBufferAO * NormalContribution / (PI * 2);
+
+	// Transmission
+	// * Emulate Beer-Lambert absorption by retrieving extinction coefficient from SubSurfaceColor. Subsurface is interpreted as a 'transmittance color' 
+	//   at a certain 'normalized' distance (SubSurfaceColorAsTransmittanceAtDistanceInMeters). This is a coarse approximation for getting hue-shiting. 
+	// * TransmittedColor is computed for the 1-normalized distance, and then transformed back-and-forth in HSV space to preserving the luminance value 
+	//   of the original color, but getting hue shifting
+	const half3 ExtinctionCoefficients = TransmittanceToExtinction(SubsurfaceColor, View.SubSurfaceColorAsTransmittanceAtDistanceInMeters);
+	const half3 RawTransmittedColor = ExtinctionToTransmittance(ExtinctionCoefficients, 1.0f /*At 1 meters, as we use normalized units*/);
+	const half3 TransmittedColor = HSV_2_LinearRGB(half3(LinearRGB_2_HSV(RawTransmittedColor).xy, LinearRGB_2_HSV(SubsurfaceColor).z));
+
 	// lerp to never exceed 1 (energy conserving)
-	Lighting.Transmission = AreaLight.FalloffColor * ( Falloff * lerp(BackScatter, 1, InScatter) ) * SubsurfaceColor;
+	Lighting.Transmission = AreaLight.FalloffColor * (Falloff * lerp(BackScatter, 1, InScatter)) * lerp(TransmittedColor, SubsurfaceColor, Shadow.TransmissionThickness);
 
 	return Lighting;
 }
@@ -977,4 +993,4 @@ FDirectLighting EvaluateBxDF( FGBufferData GBuffer, half3 N, half3 V, half3 L, f
 	AreaLight.Texture = InitRectTexture();
 
 	return IntegrateBxDF( GBuffer, N, V, L, 1, NoL, AreaLight, Shadow );
-}
\ No newline at end of file
+}
diff --git a/Engine/Shaders/Private/ShadowFilteringCommon.ush b/Engine/Shaders/Private/ShadowFilteringCommon.ush
--- a/Engine/Shaders/Private/ShadowFilteringCommon.ush
+++ b/Engine/Shaders/Private/ShadowFilteringCommon.ush
@@ -8,6 +8,7 @@
 
 #include "PixelQuadMessagePassing.ush"
 #include "FastMath.ush"
+#include "TransmissionCommon.ush"
 
 struct FPCFSamplerSettings
 {
@@ -151,13 +152,10 @@ float4 CalculateOcclusion(float4 ShadowmapDepth, FPCFSamplerSettings Settings)
 {
 	if (Settings.bSubsurface)
 	{
-		// Determine the distance that the light traveled through the subsurface object
-		// This assumes that anything between this subsurface pixel and the light was also a subsurface material,
-		// As a result, subsurface materials receive leaked light based on how transparent they are
-		float4 Thickness = max(Settings.SceneDepth - ShadowmapDepth, 0);
-		float4 Occlusion = saturate(FastExp(-Thickness * Settings.DensityMulConstant));
+		float4 OccluderDistance = max(Settings.SceneDepth - ShadowmapDepth, 0);
+		float4 Occlusion = GetSubSurfaceTransmission(OccluderDistance, Settings.DensityMulConstant);
 		// Never shadow from depths that were never written to (max depth value)
-		return select(ShadowmapDepth > .99f, 1.0, Occlusion);
+		return select(ShadowmapDepth > .99f, 1.0f, Occlusion);
 	}
 	else
 	{
diff --git a/Engine/Shaders/Private/ShadowProjectionPixelShader.usf b/Engine/Shaders/Private/ShadowProjectionPixelShader.usf
--- a/Engine/Shaders/Private/ShadowProjectionPixelShader.usf
+++ b/Engine/Shaders/Private/ShadowProjectionPixelShader.usf
@@ -315,7 +315,7 @@ void Main(
 				}
 			#else // STRATA_ENABLED
 				float3 WorldNormal = GBufferData.WorldNormal;
-				FTransmissionProfileParams TransmissionParams = GetTransmissionProfileParams(GBufferData);
+				FTransmissionProfileParams TransmissionParams = GetTransmissionProfileParams(ExtractSubsurfaceProfileInt(GBufferData));
 			#endif // STRATA_ENABLED
 				SSSTransmission = CalcTransmissionThickness(ScreenPosition.xyz, LightPositionAndInvRadius.xyz, WorldNormal, TransmissionParams, Settings);
 			}
@@ -470,7 +470,7 @@ void MainOnePassPointLightPS(
 
 	bool IsSubsurfaceProfile = GBufferData.ShadingModelID == SHADINGMODELID_SUBSURFACE_PROFILE;
 	float3 WorldNormal = GBufferData.WorldNormal;
-	FTransmissionProfileParams TransmissionParams = GetTransmissionProfileParams(GBufferData);
+	FTransmissionProfileParams TransmissionParams = GetTransmissionProfileParams(ExtractSubsurfaceProfileInt(GBufferData));
 
 #endif // STRATA_ENABLED
 
diff --git a/Engine/Shaders/Private/TransmissionCommon.ush b/Engine/Shaders/Private/TransmissionCommon.ush
--- a/Engine/Shaders/Private/TransmissionCommon.ush
+++ b/Engine/Shaders/Private/TransmissionCommon.ush
@@ -42,13 +42,6 @@ FTransmissionProfileParams GetTransmissionProfileParams(uint SubsurfaceProfileIn
 	return Result;
 }
 
-FTransmissionProfileParams GetTransmissionProfileParams(FGBufferData GBufferData)
-{
-	// 0..255, which SubSurface profile to pick
-	uint SubsurfaceProfileInt = ExtractSubsurfaceProfileInt(GBufferData);
-	return GetTransmissionProfileParams(SubsurfaceProfileInt);
-}
-
 float4 GetTransmissionProfile(uint SubsurfaceProfileInt, float Thickness)
 {
 	// 0..255, which SubSurface profile to pick
@@ -58,12 +51,6 @@ float4 GetTransmissionProfile(uint SubsurfaceProfileInt, float Thickness)
 	return Texture2DSampleLevel(View.SSProfilesTexture, View.SSProfilesTransmissionSampler, uv, 0);
 }
 
-float4 GetTransmissionProfile(FGBufferData GBufferData, float Thickness)
-{
-	// 0..255, which SubSurface profile to pick
-	return GetTransmissionProfile(ExtractSubsurfaceProfileInt(GBufferData), Thickness);
-}
-
 float EncodeThickness(float Thickness)
 {
 	return 1 - Thickness;
@@ -74,3 +61,18 @@ float DecodeThickness(float Thickness)
 	return 1 - Thickness;
 }
 
+// TODO Remove and include FastMath.ush 
+// Use fast approximation of exp(x)
+// http://www.humus.name/Articles/Persson_LowLevelThinking.pdf p.27
+#define FastExpGeneric(T) T FastExp(T x) { return exp2(1.442695f * x); }
+FastExpGeneric(float4)
+FastExpGeneric(float3)
+FastExpGeneric(float)
+// End TODO
+
+// Determine the distance that the light traveled through the subsurface object
+// This assumes that anything between this subsurface pixel and the light was also a subsurface material,
+// As a result, subsurface materials receive leaked light based on how transparent they are
+#define GetSubSurfaceTransmissionGeneric(T) T GetSubSurfaceTransmission(T OccluderDistance, float SubsurfaceDensity) { return saturate(FastExp(-OccluderDistance * SubsurfaceDensity)); }
+GetSubSurfaceTransmissionGeneric(float)
+GetSubSurfaceTransmissionGeneric(float4)
diff --git a/Engine/Shaders/Private/VirtualShadowMaps/VirtualShadowMapProjection.usf b/Engine/Shaders/Private/VirtualShadowMaps/VirtualShadowMapProjection.usf
--- a/Engine/Shaders/Private/VirtualShadowMaps/VirtualShadowMapProjection.usf
+++ b/Engine/Shaders/Private/VirtualShadowMaps/VirtualShadowMapProjection.usf
@@ -164,6 +164,7 @@ float4 ComputeRandom4(uint2 PixelPosition)
 int4 ProjectionRect;
 float ScreenRayLength;
 float NormalBias;
+float SubsurfaceMinSourceRadius;
 int SubsurfaceShadowMode;			// Temp for testing and A/B comparison
 int SMRTRayCount;					// 0 = off
 int SMRTSamplesPerRay;
@@ -533,22 +534,6 @@ void OutputVisualize( uint2 PixelPos, FVirtualShadowMapSampleResult Result )
 	}
 }
 
-// TODO: Refactor/move this to transmission common
-float GetShadowFactorSubsurface(float ShadowFactor, float OccluderDistance, float SubsurfaceOpacity)
-{
-	if ( ShadowFactor < 1.0f )
-	{
-		float Density = SubsurfaceDensityFromOpacity( SubsurfaceOpacity );
-		float Occlusion = 1.0f - saturate( exp( -Density * OccluderDistance ) );
-		float NewShadowFactor = 1.0f - Occlusion * ( 1.0f - ShadowFactor );
-
-		// NOTE: This square is arbitrary, but we have to be consistent with the PCF path
-		// TODO: This might be the time to remove this...
-		return SubsurfaceShadowMode > 1 ? Square( NewShadowFactor ) : NewShadowFactor;
-	}
-	return ShadowFactor;
-}
-
 [numthreads(8, 8, 1)]
 void VirtualShadowMapProjection(
 	uint3	GroupId				: SV_GroupID,
@@ -636,6 +621,14 @@ void VirtualShadowMapProjection(
 		int VirtualShadowMapId = LightUniformVirtualShadowMapId;
 		FLightShaderParameters Light = GetRootLightShaderParameters();
 
+		// Subsurface approximation
+		// Increases light source radius for subsurface material as the opacity reduces to emulates light diffusion.
+		const float SubsurfaceOpacity = bIsHairInput ? 1.0f : GetSubsurfaceOpacityFromGbuffer( PixelPos );
+		if (SubsurfaceOpacity < 1.0f)
+		{
+			Light.SourceRadius = max(Light.SourceRadius, (1-SubsurfaceOpacity)*SubsurfaceMinSourceRadius);
+		}
+
 		FVirtualShadowMapSampleResult Result = ProjectLight(
 			VirtualShadowMapId,
 			Light,
@@ -647,7 +640,6 @@ void VirtualShadowMapProjection(
 			Noise);
 		
 		float TransmissionShadow = Result.ShadowFactor;
-		const float SubsurfaceOpacity = bIsHairInput ? 1.0f : GetSubsurfaceOpacityFromGbuffer( PixelPos );
 		if ( SubsurfaceOpacity < 1.0f )
 		{
 			if (SubsurfaceShadowMode > 0)
diff --git a/Engine/Shaders/Private/VirtualShadowMaps/VirtualShadowMapTransmissionCommon.ush b/Engine/Shaders/Private/VirtualShadowMaps/VirtualShadowMapTransmissionCommon.ush
--- a/Engine/Shaders/Private/VirtualShadowMaps/VirtualShadowMapTransmissionCommon.ush
+++ b/Engine/Shaders/Private/VirtualShadowMaps/VirtualShadowMapTransmissionCommon.ush
@@ -8,6 +8,7 @@ VirtualShadowMaps/VirtualShadowMapTransmissionCommon.ush:
 #include "../Common.ush"
 #include "../SceneTexturesCommon.ush"
 #include "../DeferredShadingCommon.ush"
+#include "../TransmissionCommon.ush"
 #include "VirtualShadowMapProjectionCommon.ush"
 
 float GetSubsurfaceOpacityFromGbuffer(uint2 PixelPos)
@@ -36,6 +37,21 @@ float GetSubsurfaceOpacityFromGbuffer(uint2 PixelPos)
 	return Opacity;
 }
 
+/* Compute Subsurface transmission */
+float GetShadowFactorSubsurface(float ShadowFactor, float OccluderDistance, float SubsurfaceOpacity)
+{
+	if (ShadowFactor < 1.0f)
+	{
+		float SSSTransmission = GetSubSurfaceTransmission(OccluderDistance, SubsurfaceDensityFromOpacity(SubsurfaceOpacity));
+
+		// NOTE: This 'square' is to compensate the ShadowMask encoding (i.e., sqrt) applied when compositing VSM to shadow mask. 
+		//       This allows to be consistent with the PCF path
+		return Square(lerp(SSSTransmission, 1, ShadowFactor));
+	}
+	return ShadowFactor;
+
+}
+
 /**
 * Sample VSM to determine approximate thickness and compute simple subsurface transmission matching the
 * PCF path.
@@ -59,12 +75,9 @@ float ComputeSimpleSubsurfaceTransmissionFromVirtualShadowMap(
 		0.0f // RayStartDistance
 	);
 
-	if ( SSSResult.bValid && SSSResult.ShadowFactor < 1.0f )
+	if (SSSResult.bValid)
 	{
-		float Density = SubsurfaceDensityFromOpacity( SubsurfaceOpacity );
-		float Occlusion = 1.0f - saturate( exp( -Density * SSSResult.OccluderDistance ) );
-		// NOTE: This square is arbitrary, but we have to be consistent with the PCF path
-		return Square( 1.0f - Occlusion * ( 1.0f - SSSResult.ShadowFactor ) );
+		return GetShadowFactorSubsurface(SSSResult.ShadowFactor, SSSResult.OccluderDistance, SubsurfaceOpacity);
 	}
 
 	return 1.0f;
diff --git a/Engine/Source/Runtime/Engine/Public/SceneView.h b/Engine/Source/Runtime/Engine/Public/SceneView.h
--- a/Engine/Source/Runtime/Engine/Public/SceneView.h
+++ b/Engine/Source/Runtime/Engine/Public/SceneView.h
@@ -800,6 +800,7 @@ enum ETranslucencyVolumeCascade
 	VIEW_UNIFORM_BUFFER_MEMBER(uint32, InstanceSceneDataSOAStride) \
 	VIEW_UNIFORM_BUFFER_MEMBER(uint32, GPUSceneViewId) \
 	VIEW_UNIFORM_BUFFER_MEMBER(float, ViewResolutionFraction) \
+	VIEW_UNIFORM_BUFFER_MEMBER(float, SubSurfaceColorAsTransmittanceAtDistanceInMeters) \
 
 #define VIEW_UNIFORM_BUFFER_MEMBER(type, identifier) \
 	SHADER_PARAMETER(type, identifier)
diff --git a/Engine/Source/Runtime/Renderer/Private/PostProcess/PostProcessSubsurface.cpp b/Engine/Source/Runtime/Renderer/Private/PostProcess/PostProcessSubsurface.cpp
--- a/Engine/Source/Runtime/Renderer/Private/PostProcess/PostProcessSubsurface.cpp
+++ b/Engine/Source/Runtime/Renderer/Private/PostProcess/PostProcessSubsurface.cpp
@@ -138,7 +138,12 @@ namespace
 		TEXT("4000. (default) The minimal number of tiles to trigger subsurface radiance mip generation. Set to zero to always generate mips (Experimental value)"),
 		ECVF_RenderThreadSafe);
 
-	
+	TAutoConsoleVariable<float> CVarSubSurfaceColorAsTannsmittanceAtDistance(
+		TEXT("r.SSS.SubSurfaceColorAsTansmittanceAtDistance"),
+		0.15f,
+		TEXT("Normalized distance (0..1) at which the surface color is interpreted as transmittance color to compute extinction coefficients."),
+		ECVF_RenderThreadSafe);
+
 	DECLARE_GPU_STAT(SubsurfaceScattering)
 }
 
diff --git a/Engine/Source/Runtime/Renderer/Private/SceneRendering.cpp b/Engine/Source/Runtime/Renderer/Private/SceneRendering.cpp
--- a/Engine/Source/Runtime/Renderer/Private/SceneRendering.cpp
+++ b/Engine/Source/Runtime/Renderer/Private/SceneRendering.cpp
@@ -1562,6 +1562,13 @@ void FViewInfo::SetupUniformBufferParameters(
 	{
 		ViewUniformShaderParameters.bSubsurfacePostprocessEnabled = IsSubsurfaceEnabled() ? 1.0f : 0.0f;
 
+		// Subsurface shading model
+		{
+			static const auto CVar = IConsoleManager::Get().FindTConsoleVariableDataFloat(TEXT("r.SSS.SubSurfaceColorAsTansmittanceAtDistance"));
+			const float SSSDistanceInMeters = CVar ? FMath::Clamp(CVar->GetValueOnRenderThread(), 0.05f, 1.0f) : 0.15f; // Default 0.15 normalized unit
+			ViewUniformShaderParameters.SubSurfaceColorAsTransmittanceAtDistanceInMeters = SSSDistanceInMeters;
+		}
+
 		// Profiles
 		{
 			FRHITexture* Texture = GetSubsurfaceProfileTextureWithFallback();
diff --git a/Engine/Source/Runtime/Renderer/Private/VirtualShadowMaps/VirtualShadowMapProjection.cpp b/Engine/Source/Runtime/Renderer/Private/VirtualShadowMaps/VirtualShadowMapProjection.cpp
--- a/Engine/Source/Runtime/Renderer/Private/VirtualShadowMaps/VirtualShadowMapProjection.cpp
+++ b/Engine/Source/Runtime/Renderer/Private/VirtualShadowMaps/VirtualShadowMapProjection.cpp
@@ -158,6 +158,12 @@ static TAutoConsoleVariable<int32> CVarSubsurfaceShadowMode(
 	ECVF_RenderThreadSafe
 );
 
+static TAutoConsoleVariable<int32> CVarSubsurfaceShadowMinSourceAngle(
+	TEXT("r.Shadow.Virtual.SubsurfaceShadowMinSourceAngle"),
+	5,
+	TEXT("Minimum source angle (in degrees) used for shadow & transmittance of sub-surface materials"),
+	ECVF_RenderThreadSafe
+);
 
 #if MAX_TEST_PERMUTATION > 0
 static TAutoConsoleVariable<int32> CVarTestPermutation(
@@ -217,6 +223,7 @@ class FVirtualShadowMapProjectionCS : public FGlobalShader
 		SHADER_PARAMETER(float, ScreenRayLength)
 		SHADER_PARAMETER(float, NormalBias)
 		SHADER_PARAMETER(int32, SubsurfaceShadowMode)
+		SHADER_PARAMETER(float, SubsurfaceMinSourceRadius)
 		SHADER_PARAMETER(int32, SMRTRayCount)
 		SHADER_PARAMETER(int32, SMRTSamplesPerRay)
 		SHADER_PARAMETER(float, SMRTRayLengthScale)
@@ -304,6 +311,7 @@ static void RenderVirtualShadowMapProjectionCommon(
 	PassParameters->ScreenRayLength = CVarScreenRayLength.GetValueOnRenderThread();
 	PassParameters->NormalBias = GetNormalBiasForShader();
 	PassParameters->SubsurfaceShadowMode = CVarSubsurfaceShadowMode.GetValueOnRenderThread();
+	PassParameters->SubsurfaceMinSourceRadius = FMath::Sin(0.5f * FMath::DegreesToRadians(CVarSubsurfaceShadowMinSourceAngle.GetValueOnRenderThread()));
 	PassParameters->InputType = uint32(InputType);
 	PassParameters->bCullBackfacingPixels = VirtualShadowMapArray.ShouldCullBackfacingPixels() ? 1 : 0;
 	PassParameters->bSMRTUseAdaptiveRayCount = CVarSMRTAdaptiveRayCount.GetValueOnRenderThread() != 0 ? 1 : 0;
