Include FastMath.ush to fix redefinition errors
#rnx

diff --git a/Engine/Shaders/Private/TransmissionCommon.ush b/Engine/Shaders/Private/TransmissionCommon.ush
--- a/Engine/Shaders/Private/TransmissionCommon.ush
+++ b/Engine/Shaders/Private/TransmissionCommon.ush
@@ -61,14 +61,7 @@ float DecodeThickness(float Thickness)
 	return 1 - Thickness;
 }
 
-// TODO Remove and include FastMath.ush 
-// Use fast approximation of exp(x)
-// http://www.humus.name/Articles/Persson_LowLevelThinking.pdf p.27
-#define FastExpGeneric(T) T FastExp(T x) { return exp2(1.442695f * x); }
-FastExpGeneric(float4)
-FastExpGeneric(float3)
-FastExpGeneric(float)
-// End TODO
+#include "FastMath.ush"
 
 // Determine the distance that the light traveled through the subsurface object
 // This assumes that anything between this subsurface pixel and the light was also a subsurface material,
