Fixed compilation shader compilation error.

diff --git a/Engine/Shaders/Private/ShadingModels.ush b/Engine/Shaders/Private/ShadingModels.ush
--- a/Engine/Shaders/Private/ShadingModels.ush
+++ b/Engine/Shaders/Private/ShadingModels.ush
@@ -779,7 +779,7 @@ FDirectLighting ClothBxDF( FGBufferData GBuffer, half3 N, half3 V, half3 L, floa
 	return Lighting;
 }
 
-FDirectLighting SubsurfaceBxDF(FGBufferData GBuffer, half3 N, half3 V, half3 L, float Falloff, half NoL, FAreaLight AreaLight, FShadowTerms Shadow 
+FDirectLighting SubsurfaceBxDF(FGBufferData GBuffer, half3 N, half3 V, half3 L, float Falloff, half NoL, FAreaLight AreaLight, FShadowTerms Shadow )
 {
 	FDirectLighting Lighting = DefaultLitBxDF( GBuffer, N, V, L, Falloff, NoL, AreaLight, Shadow);
 
