Added support for height/volumetric fog onto ray tracing translucency.
Only heigh fog in reflections. This mode is already expenssive anyway so there is no extra RGS permutation as of now.
Requested from multiple sources (UDN, etc).

diff --git a/Engine/Shaders/Private/RayTracing/RayTracingPrimaryRays.usf b/Engine/Shaders/Private/RayTracing/RayTracingPrimaryRays.usf
--- a/Engine/Shaders/Private/RayTracing/RayTracingPrimaryRays.usf
+++ b/Engine/Shaders/Private/RayTracing/RayTracingPrimaryRays.usf
@@ -187,7 +187,31 @@ RAY_TRACING_ENTRY_RAYGEN(RayTracingPrimaryRaysRGS)
 		// It is also needed for RAY_TRACING_BLEND_MODE_ADDITIVE and  RAY_TRACING_BLEND_MODE_ALPHA_COMPOSITE: radiance continbution is alway weighted by coverage.
 		// #dxr_todo: I have not been able to setup a material using RAY_TRACING_BLEND_MODE_MODULATE.
 
-		PathRadiance += PathThroughput * PathVertexRadiance * vertexRadianceWeight;
+		//
+		// Apply fog on primary ray (the first translucent surface encountered) if needed.
+		// We can only handle primary rays because height and volumetric fog are only setup for the camera point of view.
+		// 
+		float4 ViewFogInScatteringAndTransmittance = float4(0.0, 0.0, 0.0, 1.0);
+		if (HeightFog > 0)
+		{
+			float3 WorldPositionRelativeToCameraPrimaryRay = PrimaryView.TranslatedWorldCameraOrigin - (Ray.Origin + Ray.Direction * Payload.HitT);
+
+			// We always samplg the height fog to make sure all the intersected transparent surface get affected by it.
+			// We cannot only do it for the primary ray otherwise refelction on secondary surface will look too bright. 
+			// For that sample the fog for the hit position and view assuming no refraction. This seems acceptable.
+			ViewFogInScatteringAndTransmittance = CalculateHeightFog(WorldPositionRelativeToCameraPrimaryRay);
+
+			if (FogStruct.ApplyVolumetricFog > 0 
+				&& !bHasScattered) // We can only sample the volumetric fog contribution (always camera aligned) if the ray is not scattered around
+			{
+				float4 ClipPos = mul(float4(HitPoint, 1.0f), PrimaryView.TranslatedWorldToClip); // HitPoint is translated world space
+				float3 VolumeUV = ComputeVolumeUVFromNDC(ClipPos);
+
+				ViewFogInScatteringAndTransmittance = CombineVolumetricFog(ViewFogInScatteringAndTransmittance, VolumeUV, 0/*EyeIndex*/, GBufferData.Depth);
+			}
+		}
+
+		PathRadiance += PathThroughput * vertexRadianceWeight * (PathVertexRadiance * ViewFogInScatteringAndTransmittance.a + ViewFogInScatteringAndTransmittance.rgb);
 
 		AccumulatedOpacity += vertexRadianceWeight;
 
@@ -247,7 +271,15 @@ RAY_TRACING_ENTRY_RAYGEN(RayTracingPrimaryRaysRGS)
 			float NoV = saturate(dot(-Ray.Direction, Payload.WorldNormal));
 			const float3 ReflectionThroughput = EnvBRDF(Payload.SpecularColor, Payload.Roughness, NoV);
 
-			PathRadiance += PathThroughput * ReflectionThroughput * ReflectionRadiance * vertexRadianceWeight;
+			// For reflections, we can only sample the height fog (volumetric fog is only for primary rays)
+			if (HeightFog > 0)
+			{
+				float3 OriginToCollider = ReflectionPayload.TranslatedWorldPos - ReflectionRay.Origin;
+				float4 ReflectionFogInscatteringAndTransmittance = CalculateHeightFog(OriginToCollider);
+				ReflectionRadiance = ReflectionRadiance * ReflectionFogInscatteringAndTransmittance.a + ReflectionFogInscatteringAndTransmittance.rgb;
+			}
+
+			PathRadiance += PathThroughput * vertexRadianceWeight * ViewFogInScatteringAndTransmittance.a * ReflectionThroughput * ReflectionRadiance;
 		}
 
 
