Cleaned up defines for SingleLayerWater shading: Removed SINGLE_LAYER_WATER_SIMPLE_FORWARD and SIMPLE_SINGLE_LAYER_WATER and replaced them with SINGLE_LAYER_WATER_SHADING_QUALITY which has 3 different quality levels.

#rnx

diff --git a/Engine/Shaders/Private/BasePassCommon.ush b/Engine/Shaders/Private/BasePassCommon.ush
--- a/Engine/Shaders/Private/BasePassCommon.ush
+++ b/Engine/Shaders/Private/BasePassCommon.ush
@@ -58,8 +58,8 @@
 
 // We consider that the cloud shadow is always enabled but the interpolator is only used when render water or lit translucent with forward shading. We only support per vertex cloud shadow as of today.
 #if MATERIAL_SHADINGMODEL_SINGLELAYERWATER
-// Checking for !SINGLE_LAYER_WATER_SIMPLE_FORWARD disables cloud shadows on low end platforms.
-#define NEEDS_BASEPASS_CLOUD_SHADOW_INTERPOLATOR (SUPPORT_CLOUD_SHADOW_ON_SINGLE_LAYER_WATER && !SINGLE_LAYER_WATER_SIMPLE_FORWARD)
+// Checking for SINGLE_LAYER_WATER_SHADING_QUALITY disables cloud shadows on low end platforms.
+#define NEEDS_BASEPASS_CLOUD_SHADOW_INTERPOLATOR (SUPPORT_CLOUD_SHADOW_ON_SINGLE_LAYER_WATER && (SINGLE_LAYER_WATER_SHADING_QUALITY == SINGLE_LAYER_WATER_SHADING_QUALITY_FULL))
 #else
 #define NEEDS_BASEPASS_CLOUD_SHADOW_INTERPOLATOR (SUPPORT_CLOUD_SHADOW_ON_FORWARD_LIT_TRANSLUCENT && IS_MATERIAL_TRANSLUCENT_AND_LIT && !TRANSLUCENCY_PERVERTEX_FORWARD_SHADING && !TRANSLUCENCY_PERVERTEX_LIGHTING_VOLUME)
 #endif
diff --git a/Engine/Shaders/Private/BasePassPixelShader.usf b/Engine/Shaders/Private/BasePassPixelShader.usf
--- a/Engine/Shaders/Private/BasePassPixelShader.usf
+++ b/Engine/Shaders/Private/BasePassPixelShader.usf
@@ -136,12 +136,9 @@
 
 #if  MATERIAL_SHADINGMODEL_SINGLELAYERWATER
 
-#if SINGLE_LAYER_WATER_SIMPLE_FORWARD
+#if SINGLE_LAYER_WATER_SHADING_QUALITY == SINGLE_LAYER_WATER_SHADING_QUALITY_MOBILE_WITH_DEPTH_TEXTURE
 
 	// Use mobile like simple shading
-	#undef SIMPLE_SINGLE_LAYER_WATER
-	#define SIMPLE_SINGLE_LAYER_WATER 1
-
 	// Remove forward lighting directional light shadow
 	#define DISABLE_FORWARD_DIRECTIONAL_LIGHT_SHADOW 1
 
@@ -1463,7 +1460,7 @@ void FPixelShaderInOut_MainPS(
 		const float3 V = MaterialParameters.CameraVector;
 		const float3 EnvBrdf = EnvBRDF(GBuffer.SpecularColor, GBuffer.Roughness, max(0.0, dot(N, V)));
 
-#if SINGLE_LAYER_WATER_SIMPLE_FORWARD 
+#if SINGLE_LAYER_WATER_SHADING_QUALITY == SINGLE_LAYER_WATER_SHADING_QUALITY_MOBILE_WITH_DEPTH_TEXTURE
 		const float4 NullDistortionParams = 1.0f;
 		WaterVolumeLightingOutput WaterLighting = EvaluateWaterVolumeLighting(
 			MaterialParameters, PixelMaterialInputs, ResolvedView,
@@ -1484,12 +1481,12 @@ void FPixelShaderInOut_MainPS(
 		Color += EvaluateWaterVolumeLighting(
 			MaterialParameters, PixelMaterialInputs, ResolvedView,
 			DirectionalLightShadow * DirectionalLightCloudShadow,
-			SingleLayerWater.SceneColorWithoutSingleLayerWaterTexture, SingleLayerWaterSceneColorSampler,
 			SingleLayerWater.SceneDepthWithoutSingleLayerWaterTexture, SingleLayerWaterSceneDepthSampler,
-			SingleLayerWater.SceneWithoutSingleLayerWaterMinMaxUV.xy,
-			SingleLayerWater.SceneWithoutSingleLayerWaterMinMaxUV.zw,
 			SingleLayerWater.SceneWithoutSingleLayerWaterTextureSize,
 			SingleLayerWater.SceneWithoutSingleLayerWaterInvTextureSize,
+			SingleLayerWater.SceneColorWithoutSingleLayerWaterTexture, SingleLayerWaterSceneColorSampler,
+			SingleLayerWater.SceneWithoutSingleLayerWaterMinMaxUV.xy,
+			SingleLayerWater.SceneWithoutSingleLayerWaterMinMaxUV.zw,
 			Specular, SingleLayerWater.DistortionParams,
 			SunIlluminance, WaterDiffuseIndirectIlluminance, EnvBrdf,
 			CameraIsUnderWater, WaterVisibility, EyeIndex,
diff --git a/Engine/Shaders/Private/Definitions.usf b/Engine/Shaders/Private/Definitions.usf
--- a/Engine/Shaders/Private/Definitions.usf
+++ b/Engine/Shaders/Private/Definitions.usf
@@ -211,8 +211,9 @@
 #define DISABLE_FORWARD_LOCAL_LIGHTS					0
 #endif
 
-#ifndef SINGLE_LAYER_WATER_SIMPLE_FORWARD
-#define SINGLE_LAYER_WATER_SIMPLE_FORWARD				0
+// 0: Full quality, 1: mobile shading, sampling underwater depth from a texture 2: mobile shading, sampling underwater depth from the depth buffer
+#ifndef SINGLE_LAYER_WATER_SHADING_QUALITY
+#define SINGLE_LAYER_WATER_SHADING_QUALITY				0
 #endif
 
 #ifndef CLOUD_LAYER_PIXEL_SHADER
diff --git a/Engine/Shaders/Private/MaterialTemplate.ush b/Engine/Shaders/Private/MaterialTemplate.ush
--- a/Engine/Shaders/Private/MaterialTemplate.ush
+++ b/Engine/Shaders/Private/MaterialTemplate.ush
@@ -1686,7 +1686,7 @@ float2 SceneWithoutSingleLayerWaterInvTextureSize;
 */
 float MaterialExpressionSceneDepthWithoutWater(float2 ViewportUV, float FallbackDepth)
 {
-#if MATERIAL_SHADINGMODEL_SINGLELAYERWATER && !SCENE_TEXTURES_DISABLED && (!SIMPLE_SINGLE_LAYER_WATER || SINGLE_LAYER_WATER_SIMPLE_FORWARD)
+#if MATERIAL_SHADINGMODEL_SINGLELAYERWATER && !SCENE_TEXTURES_DISABLED && (SINGLE_LAYER_WATER_SHADING_QUALITY != SINGLE_LAYER_WATER_SHADING_QUALITY_MOBILE_WITH_DEPTH_BUFFER)
 	const float2 ClampedViewportUV = clamp(ViewportUV, SingleLayerWater.SceneWithoutSingleLayerWaterMinMaxUV.xy, SingleLayerWater.SceneWithoutSingleLayerWaterMinMaxUV.zw);
 	float WaterDeviceZ = WaterSampleSceneDepthWithoutWater(SingleLayerWater.SceneDepthWithoutSingleLayerWaterTexture, SingleLayerWaterSceneDepthSampler, ClampedViewportUV, 
 		SingleLayerWater.SceneWithoutSingleLayerWaterTextureSize, SingleLayerWater.SceneWithoutSingleLayerWaterInvTextureSize);
diff --git a/Engine/Shaders/Private/MobileBasePassPixelShader.usf b/Engine/Shaders/Private/MobileBasePassPixelShader.usf
--- a/Engine/Shaders/Private/MobileBasePassPixelShader.usf
+++ b/Engine/Shaders/Private/MobileBasePassPixelShader.usf
@@ -41,7 +41,11 @@
 #endif
 
 #if MATERIAL_SHADINGMODEL_SINGLELAYERWATER
-	#define SIMPLE_SINGLE_LAYER_WATER 1
+	#ifdef SINGLE_LAYER_WATER_SHADING_QUALITY
+	#undef SINGLE_LAYER_WATER_SHADING_QUALITY
+	#endif
+	// Value must match SINGLE_LAYER_WATER_SHADING_QUALITY_MOBILE_WITH_DEPTH_BUFFER in SingleLayerWaterCommon.ush!
+	#define SINGLE_LAYER_WATER_SHADING_QUALITY 2
 #endif
 
 #ifndef PROJECT_MOBILE_ENABLE_MOVABLE_SPOTLIGHT_SHADOWS
diff --git a/Engine/Shaders/Private/SingleLayerWaterCommon.ush b/Engine/Shaders/Private/SingleLayerWaterCommon.ush
--- a/Engine/Shaders/Private/SingleLayerWaterCommon.ush
+++ b/Engine/Shaders/Private/SingleLayerWaterCommon.ush
@@ -3,22 +3,34 @@
 #pragma once
 #include "TextureSampling.ush" // GetBilinearSampleLevelInfos() and GetBilinearInterpolation()
 
-#ifndef SIMPLE_SINGLE_LAYER_WATER
-#define SIMPLE_SINGLE_LAYER_WATER 0
-#endif
+// Full shading quality, sampling both the color and depth behind the water surface and supporting refraction.
+// Used in BasePassPixelShader.usf.
+#define SINGLE_LAYER_WATER_SHADING_QUALITY_FULL							0
+
+// Reduced shading quality (same as mobile), sampling only the depth behind the water surface and does not support refraction.
+// Blended against the background using HW blending.
+// This quality level also causes BasePassPixelShader.usf to be compiled without forward directional light shadow and is likely to be used
+// together with DISABLE_FORWARD_LOCAL_LIGHTS to disable all local lights.
+// Used in BasePassPixelShader.usf.
+#define SINGLE_LAYER_WATER_SHADING_QUALITY_MOBILE_WITH_DEPTH_TEXTURE	1
+
+// Reduced shading quality for mobile platforms, sampling only the depth buffer behind the water using LookupDeviceZ(). Does not support refraction.
+// Blended against the background using HW blending.
+// Used in MobileBasePassPixelShader.usf.
+#define SINGLE_LAYER_WATER_SHADING_QUALITY_MOBILE_WITH_DEPTH_BUFFER		2
 
 #if MATERIAL_SHADINGMODEL_SINGLELAYERWATER
-#if SUPPORTS_INDEPENDENT_SAMPLERS
-	#define SingleLayerWaterSceneColorSampler View.SharedBilinearClampedSampler
-#if METAL_PROFILE
-	#define SingleLayerWaterSceneDepthSampler View.SharedPointClampedSampler
-#else // TODO: Dynamically use a point sampler instead if the depth texture was not downsampled.
-	#define SingleLayerWaterSceneDepthSampler View.SharedBilinearClampedSampler
-#endif // METAL_PROFILE
-#else
-	#define SingleLayerWaterSceneColorSampler SingleLayerWater.SceneColorWithoutSingleLayerWaterSampler
-	#define SingleLayerWaterSceneDepthSampler SingleLayerWater.SceneDepthWithoutSingleLayerWaterSampler
-#endif // SUPPORTS_INDEPENDENT_SAMPLERS
+	#if SUPPORTS_INDEPENDENT_SAMPLERS
+		#define SingleLayerWaterSceneColorSampler View.SharedBilinearClampedSampler
+		#if METAL_PROFILE
+			#define SingleLayerWaterSceneDepthSampler View.SharedPointClampedSampler
+		#else // TODO: Dynamically use a point sampler instead if the depth texture was not downsampled.
+			#define SingleLayerWaterSceneDepthSampler View.SharedBilinearClampedSampler
+		#endif // METAL_PROFILE
+	#else
+		#define SingleLayerWaterSceneColorSampler SingleLayerWater.SceneColorWithoutSingleLayerWaterSampler
+		#define SingleLayerWaterSceneDepthSampler SingleLayerWater.SceneDepthWithoutSingleLayerWaterSampler
+	#endif // SUPPORTS_INDEPENDENT_SAMPLERS
 #endif // MATERIAL_SHADINGMODEL_SINGLELAYERWATER
 
 // We use our own refract because the HLSL one was not working on all platforms (ES3.1 preview for instance)
diff --git a/Engine/Shaders/Private/SingleLayerWaterShading.ush b/Engine/Shaders/Private/SingleLayerWaterShading.ush
--- a/Engine/Shaders/Private/SingleLayerWaterShading.ush
+++ b/Engine/Shaders/Private/SingleLayerWaterShading.ush
@@ -8,6 +8,13 @@
 #include "DistortionCommon.ush"
 #include "SingleLayerWaterCommon.ush"
 
+#if SINGLE_LAYER_WATER_SHADING_QUALITY != SINGLE_LAYER_WATER_SHADING_QUALITY_FULL && SINGLE_LAYER_WATER_SHADING_QUALITY != SINGLE_LAYER_WATER_SHADING_QUALITY_MOBILE_WITH_DEPTH_TEXTURE && SINGLE_LAYER_WATER_SHADING_QUALITY != SINGLE_LAYER_WATER_SHADING_QUALITY_MOBILE_WITH_DEPTH_BUFFER
+#error SINGLE_LAYER_WATER_SHADING_QUALITY has an invalid value!
+#endif
+
+// Helper define for all shading qualities using mobile shading
+#define SINGLE_LAYER_WATER_MOBILE_SHADING (SINGLE_LAYER_WATER_SHADING_QUALITY == SINGLE_LAYER_WATER_SHADING_QUALITY_MOBILE_WITH_DEPTH_TEXTURE) || (SINGLE_LAYER_WATER_SHADING_QUALITY == SINGLE_LAYER_WATER_SHADING_QUALITY_MOBILE_WITH_DEPTH_BUFFER) 
+
 #if MATERIAL_SHADINGMODEL_SINGLELAYERWATER
 	// We need the water base pass to be allowed to use masking (if blend mode is Masked) because it doesn't rely on a depth pre-pass (where clipping would normally occur)
 	#undef EARLY_Z_PASS_ONLY_MATERIAL_MASKING
@@ -25,19 +32,16 @@ struct WaterVolumeLightingOutput
 WaterVolumeLightingOutput EvaluateWaterVolumeLighting(
 	FMaterialPixelParameters MaterialParameters, FPixelMaterialInputs PixelMaterialInputs, ViewState ResolvedView,
 	float DirectionalLightShadow,
-#if SINGLE_LAYER_WATER_SIMPLE_FORWARD
+#if SINGLE_LAYER_WATER_SHADING_QUALITY == SINGLE_LAYER_WATER_SHADING_QUALITY_FULL || SINGLE_LAYER_WATER_SHADING_QUALITY == SINGLE_LAYER_WATER_SHADING_QUALITY_MOBILE_WITH_DEPTH_TEXTURE
 	Texture2D SceneDepthWithoutSingleLayerWaterTexture, SamplerState SceneDepthWithoutSingleLayerWaterSampler,
 	float2 SceneDepthWithoutSingleLayerWaterTextureSize,
 	float2 SceneDepthWithoutSingleLayerWaterTextureInvSize,
 #endif
-#if !SIMPLE_SINGLE_LAYER_WATER
+#if SINGLE_LAYER_WATER_SHADING_QUALITY == SINGLE_LAYER_WATER_SHADING_QUALITY_FULL
 	Texture2D SceneColorWithoutSingleLayerWaterTexture, SamplerState SceneColorWithoutSingleLayerWaterSampler,
-	Texture2D SceneDepthWithoutSingleLayerWaterTexture, SamplerState SceneDepthWithoutSingleLayerWaterSampler,
 	float2 SceneWithoutSingleLayerWaterMinUV,
 	float2 SceneWithoutSingleLayerWaterMaxUV,
-	float2 SceneDepthWithoutSingleLayerWaterTextureSize,
-	float2 SceneDepthWithoutSingleLayerWaterTextureInvSize,
-#endif // SIMPLE_SINGLE_LAYER_WATER
+#endif
 	float Specular, const float4 DistortionParams,
 	float3 SunIlluminance, float3 AmbiantIlluminance, float3 EnvBrdf,
 	bool CameraIsUnderWater, float WaterVisibility, uint EyeIndex,
@@ -53,9 +57,9 @@ WaterVolumeLightingOutput EvaluateWaterVolumeLighting(
 	if (WaterVisibility > 0.0f)
 	{
 		float WaterDepth = MaterialParameters.ScreenPosition.w; // Scene depth with Pixel-Depth-Offset taken into account is in ScreenPosition.w already
-#if SIMPLE_SINGLE_LAYER_WATER
+#if SINGLE_LAYER_WATER_MOBILE_SHADING
 
-#if SINGLE_LAYER_WATER_SIMPLE_FORWARD 
+#if SINGLE_LAYER_WATER_SHADING_QUALITY == SINGLE_LAYER_WATER_SHADING_QUALITY_MOBILE_WITH_DEPTH_TEXTURE 
 
 		float PixelSceneDeviceZ = WaterSampleSceneDepthWithoutWater(
 			SceneDepthWithoutSingleLayerWaterTexture,
@@ -76,7 +80,7 @@ WaterVolumeLightingOutput EvaluateWaterVolumeLighting(
 		float SceneDeviceZ = PixelSceneDeviceZ;
 		float SceneDepth = PixelSceneDepth;
 
-#else  // !SIMPLE_SINGLE_LAYER_WATER
+#else  // !SINGLE_LAYER_WATER_MOBILE_SHADING
 
 		ViewportUV = clamp(ViewportUV, SceneWithoutSingleLayerWaterMinUV, SceneWithoutSingleLayerWaterMaxUV);
 
@@ -120,7 +124,7 @@ WaterVolumeLightingOutput EvaluateWaterVolumeLighting(
 			ViewportUV = DistortedUV;
 		}
 
-#endif // !SIMPLE_SINGLE_LAYER_WATER
+#endif // !SINGLE_LAYER_WATER_MOBILE_SHADING
 
 		const float BehindWaterDeltaDepth = CameraIsUnderWater ? WaterDepth : max(0.0f, SceneDepth - WaterDepth);
 
@@ -149,7 +153,7 @@ WaterVolumeLightingOutput EvaluateWaterVolumeLighting(
 		
 		float DirLightPhaseValue = 0.0f; // Default when Total Internal Reflection happens.
 		{
-#if SIMPLE_SINGLE_LAYER_WATER
+#if SINGLE_LAYER_WATER_MOBILE_SHADING
 			DirLightPhaseValue = IsotropicPhase();
 #else
 			float IorFrom = 1.0f; // assumes we come from air
@@ -170,7 +174,7 @@ WaterVolumeLightingOutput EvaluateWaterVolumeLighting(
 		const float DistanceFromScenePixelToWaterTop = max(0.0, LWCToFloat(LWCSubtract(MaterialParameters.AbsoluteWorldPosition, BehindWaterSceneWorldPos)).z);
 		const float3 MeanTransmittanceToLightSources = exp(-DistanceFromScenePixelToWaterTop * ExtinctionCoeff);
 
-#if SIMPLE_SINGLE_LAYER_WATER
+#if SINGLE_LAYER_WATER_MOBILE_SHADING
 		const float3 BehindWaterSceneLuminance = 0.0f; // Cannot read back the scene color in this case
 #else
 		// We use the pixel SvPosition instead of the scene one pre refraction/distortion to avoid those extra ALUs.
@@ -208,3 +212,5 @@ WaterVolumeLightingOutput EvaluateWaterVolumeLighting(
 
 	return Output;
 }
+
+#undef SINGLE_LAYER_WATER_MOBILE_SHADING
diff --git a/Engine/Source/Runtime/Engine/Private/Materials/HLSLMaterialTranslator.cpp b/Engine/Source/Runtime/Engine/Private/Materials/HLSLMaterialTranslator.cpp
--- a/Engine/Source/Runtime/Engine/Private/Materials/HLSLMaterialTranslator.cpp
+++ b/Engine/Source/Runtime/Engine/Private/Materials/HLSLMaterialTranslator.cpp
@@ -1853,7 +1853,8 @@ void FHLSLMaterialTranslator::GetMaterialEnvironment(EShaderPlatform InPlatform,
 		const bool bSingleLayerWaterUsesSimpleShading = FDataDrivenShaderPlatformInfo::GetWaterUsesSimpleForwardShading(InPlatform) && IsForwardShadingEnabled(InPlatform);
 		if (ShadingModels.HasShadingModel(MSM_SingleLayerWater) && bSingleLayerWaterUsesSimpleShading)
 		{
-			OutEnvironment.SetDefine(TEXT("SINGLE_LAYER_WATER_SIMPLE_FORWARD"), TEXT("1"));
+			// Value must match SINGLE_LAYER_WATER_SHADING_QUALITY_MOBILE_WITH_DEPTH_TEXTURE in SingleLayerWaterCommon.ush!
+			OutEnvironment.SetDefine(TEXT("SINGLE_LAYER_WATER_SHADING_QUALITY"), TEXT("1"));
 		}
 
 		const bool bIsWaterDistanceFieldShadowEnabled = IsWaterDistanceFieldShadowEnabled(InPlatform);
diff --git a/Engine/Source/Runtime/Engine/Private/Materials/MaterialHLSLEmitter.cpp b/Engine/Source/Runtime/Engine/Private/Materials/MaterialHLSLEmitter.cpp
--- a/Engine/Source/Runtime/Engine/Private/Materials/MaterialHLSLEmitter.cpp
+++ b/Engine/Source/Runtime/Engine/Private/Materials/MaterialHLSLEmitter.cpp
@@ -615,7 +615,8 @@ static void GetMaterialEnvironment(EShaderPlatform InPlatform,
 
 		if (ShadingModels.HasShadingModel(MSM_SingleLayerWater) && bSingleLayerWaterUsesSimpleShading)
 		{
-			OutEnvironment.SetDefine(TEXT("SINGLE_LAYER_WATER_SIMPLE_FORWARD"), TEXT("1"));
+			// Value must match SINGLE_LAYER_WATER_SHADING_QUALITY_MOBILE_WITH_DEPTH_TEXTURE in SingleLayerWaterCommon.ush!
+			OutEnvironment.SetDefine(TEXT("SINGLE_LAYER_WATER_SHADING_QUALITY"), TEXT("1"));
 		}
 
 		if (NumSetMaterials == 1)
