Added alternative wave-ops based compaction in FWaterTileCategorisationCS and also fixed a LDS out of bounds access.

#rnx

diff --git a/Engine/Shaders/Private/SingleLayerWaterComposite.usf b/Engine/Shaders/Private/SingleLayerWaterComposite.usf
--- a/Engine/Shaders/Private/SingleLayerWaterComposite.usf
+++ b/Engine/Shaders/Private/SingleLayerWaterComposite.usf
@@ -277,7 +277,11 @@ RWBuffer<uint> DrawIndirectDataUAV;
 RWBuffer<uint> DispatchIndirectDataUAV;
 RWBuffer<uint> WaterTileListDataUAV;
 
-groupshared uint ContainsWater[WORK_TILE_SIZE*WORK_TILE_SIZE];
+#if COMPILER_SUPPORTS_WAVE_VOTE
+groupshared uint bAnyWaterPixels;
+#else
+groupshared bool ContainsWater[WORK_TILE_SIZE * WORK_TILE_SIZE];
+#endif
 
 [numthreads(WORK_TILE_SIZE, WORK_TILE_SIZE, 1)]
 void WaterTileCatergorisationCS(uint3 ThreadId : SV_DispatchThreadID, uint3 GroupId : SV_GroupID, uint3 GroupThreadId : SV_GroupThreadID)
@@ -337,49 +341,80 @@ void WaterTileCatergorisationCS(uint3 ThreadId : SV_DispatchThreadID, uint3 Grou
 	#if STRATA_ENABLED
 		FStrataAddressing StrataAddressing = GetStrataPixelDataByteOffset(PixelPos, uint2(View.BufferSizeAndInvSize.xy), Strata.MaxBytesPerPixel);
 		FStrataPixelHeader StrataPixelHeader = UnpackStrataHeaderIn(Strata.MaterialTextureArray, StrataAddressing, Strata.TopLayerTexture);
-		bContainsWater = (StrataPixelHeader.BSDFCount > 0 && IsSingleLayerWater(StrataPixelHeader));
+		bContainsWater = bContainsWater || (StrataPixelHeader.BSDFCount > 0 && IsSingleLayerWater(StrataPixelHeader));
 	#else
 		FGBufferData GBuffer = GetGBufferDataFromSceneTextures(BufferUV);
-		bContainsWater = (GBuffer.ShadingModelID == SHADINGMODELID_SINGLELAYERWATER);
+		bContainsWater = bContainsWater || (GBuffer.ShadingModelID == SHADINGMODELID_SINGLELAYERWATER);
 	#endif
 	}
 
+	bool bWriteTile = false;
+
+#if COMPILER_SUPPORTS_WAVE_VOTE
+
+	if (all(GroupThreadId == 0))
+	{
+		bAnyWaterPixels = 0;
+	}
+	GroupMemoryBarrierWithGroupSync();
+
+	const bool bAnyWaterPixelsInWave = WaveActiveAnyTrue(bContainsWater);
+	if (WaveIsFirstLane() && bAnyWaterPixelsInWave)
+	{
+		InterlockedAdd(bAnyWaterPixels, 1);
+	}
+	GroupMemoryBarrierWithGroupSync();
+	
+	if (all(GroupThreadId == 0))
+	{
+		bWriteTile = bAnyWaterPixels > 0;
+	}
+
+#else
+
 	const int LinearIndex = GroupThreadId.y * WORK_TILE_SIZE + GroupThreadId.x;
-	ContainsWater[LinearIndex] = bContainsWater ? 1 : 0;
+	ContainsWater[LinearIndex] = bContainsWater;
 	GroupMemoryBarrierWithGroupSync();
 
-	if ((LinearIndex < 32) && (ContainsWater[LinearIndex + 32] != 0)) // 8*8 = 64 elements to merge
+	if (LinearIndex < 32) // 8*8 = 64 elements to merge
 	{
-		ContainsWater[LinearIndex] = 1;
+		ContainsWater[LinearIndex] = ContainsWater[LinearIndex] || ContainsWater[LinearIndex + 32];
 	}
 	GroupMemoryBarrierWithGroupSync();
-	if ((LinearIndex < 16) && (ContainsWater[LinearIndex + 16] != 0))
+	if (LinearIndex < 16)
 	{
-		ContainsWater[LinearIndex] = 1;
+		ContainsWater[LinearIndex] = ContainsWater[LinearIndex] || ContainsWater[LinearIndex + 16];
 	}
 	GroupMemoryBarrierWithGroupSync();
 
 	// The smallest wave size is 16 on Intel hardware. So now we can do simple math operations without group sync.
 	// EDIT: for some reason group sync is needed until the end, otherwise some pixels are missing...
 
-	if ((LinearIndex < 8) && (ContainsWater[LinearIndex + 8] != 0))
+	if (LinearIndex < 8)
 	{
-		ContainsWater[LinearIndex] = 1;
+		ContainsWater[LinearIndex] = ContainsWater[LinearIndex] || ContainsWater[LinearIndex + 8];
 	}
 	GroupMemoryBarrierWithGroupSync();
-	if ((LinearIndex < 4) && (ContainsWater[LinearIndex + 4] != 0))
+	if (LinearIndex < 4)
 	{
-		ContainsWater[LinearIndex] = 1;
+		ContainsWater[LinearIndex] = ContainsWater[LinearIndex] || ContainsWater[LinearIndex + 4];
 	}
 	GroupMemoryBarrierWithGroupSync();
-	if ((LinearIndex < 2) && (ContainsWater[LinearIndex + 2] != 0))
+	if (LinearIndex < 2)
 	{
-		ContainsWater[LinearIndex] = 1;
+		ContainsWater[LinearIndex] = ContainsWater[LinearIndex] || ContainsWater[LinearIndex + 2];
 	}
 	GroupMemoryBarrierWithGroupSync();
-	if (LinearIndex < 1 && (bContainsWater || (ContainsWater[LinearIndex + 1] != 0)))
+	if (LinearIndex < 1)
+	{
+		bWriteTile = ContainsWater[LinearIndex] || ContainsWater[LinearIndex + 1];
+	}
+
+#endif
+
+	if (bWriteTile)
 	{
-		uint WriteToIndex;  
+		uint WriteToIndex;
 		InterlockedAdd(DrawIndirectDataUAV[1], 1, WriteToIndex);
 		InterlockedAdd(DispatchIndirectDataUAV[0], 1);
 		WaterTileListDataUAV[WriteToIndex] = GroupId.x | (GroupId.y << 16); // assumes 16bit is enought to represent a tiled resolution up to 65,535 :)
