Port updated legacy sub-surface shading model to Strata.

diff --git a/Engine/Shaders/Private/Strata/StrataEvaluation.ush b/Engine/Shaders/Private/Strata/StrataEvaluation.ush
--- a/Engine/Shaders/Private/Strata/StrataEvaluation.ush
+++ b/Engine/Shaders/Private/Strata/StrataEvaluation.ush
@@ -662,13 +662,28 @@ FStrataEvaluateResult StrataEvaluateBSDFCommon(FStrataBSDFContext BSDFContext, F
 					// Legacy Subsurface
 					 
 					// To get an effect when you see through the material (hard coded pow constant)
-					const float InScatter = pow(saturate(dot(BSDFContext.L, -BSDFContext.V)), 12) * lerp(3, .1f, Opacity);
-					// Wrap around lighting, /(PI*2) to be energy consistent (hack do get some view dependnt and light dependent effect)
-					// Opacity of 0 gives no normal dependent lighting, Opacity of 1 gives strong normal contribution
-					const float NormalContribution = saturate(dot(BSDFContext.N, BSDFContext.H) * Opacity + 1 - Opacity);
-					const float BackScatter = /* GBuffer.GBufferAO */ NormalContribution / (PI * 2);
+					const half InScatter = pow(saturate(dot(BSDFContext.L, -BSDFContext.V)), 12) * lerp(3, .1f, Opacity);
+					// Wrap around lighting, 
+					// * /(PI*2) to be energy consistent (hack do get some view dependnt and light dependent effect)
+					// * Opacity of 0 gives no normal dependent lighting, Opacity of 1 gives strong normal contribution	
+					// * Simplified version (with w=.5,  n=1.5): 
+					//     half WrappedDiffuse = 2 * pow(saturate((dot(N, L) + w) / (1.0f + w)), n) * (n + 1) / (2 * (1 + w));
+					//     NormalContribution = WrappedDiffuse * Opacity + 1 - Opacity;
+					const half WrappedDiffuse = pow(saturate(AreaLightContext.NoL * (1.f / 1.5f) + (0.5f / 1.5f)), 1.5f) * (2.5f / 1.5f);
+					const half NormalContribution = lerp(1.f, WrappedDiffuse, Opacity);
+					const half BackScatter = /* GBuffer.GBufferAO */ NormalContribution / (PI * 2);
+
+					// Transmission
+					// * Emulate Beer-Lambert absorption by retrieving extinction coefficient from SubSurfaceColor. Subsurface is interpreted as a 'transmittance color' 
+					//   at a certain 'normalized' distance (SubSurfaceColorAsTransmittanceAtDistanceInMeters). This is a coarse approximation for getting hue-shiting. 
+					// * TransmittedColor is computed for the 1-normalized distance, and then transformed back-and-forth in HSV space to preserving the luminance value 
+					//   of the original color, but getting hue shifting
+					const half3 ExtinctionCoefficients = TransmittanceToExtinction(SubSurfaceColor, View.SubSurfaceColorAsTransmittanceAtDistanceInMeters);
+					const half3 RawTransmittedColor = ExtinctionToTransmittance(ExtinctionCoefficients, 1.0f /*At 1 meters, as we use normalized units*/);
+					const half3 TransmittedColor = HSV_2_LinearRGB(half3(LinearRGB_2_HSV(RawTransmittedColor).xy, LinearRGB_2_HSV(SubSurfaceColor).z));
+
 					// Lerp to never exceed 1 (energy conserving)
-					TransmissionThroughput = lerp(BackScatter, 1, InScatter) * SubSurfaceColor;
+					TransmissionThroughput = lerp(BackScatter, 1, InScatter) * lerp(TransmittedColor, SubSurfaceColor, ShadowTerms.TransmissionShadow);
 				}
 
 				Sample.TransmissionPDF = 1.0f / (4.0f * PI); // STRATA_TODO this currently match the uniform sphere sampling from StrataImportanceSampleBSDF
