Fixed broken reflection captures on water at LWC distances.

#rnx

diff --git a/Engine/Shaders/Private/SingleLayerWaterComposite.usf b/Engine/Shaders/Private/SingleLayerWaterComposite.usf
--- a/Engine/Shaders/Private/SingleLayerWaterComposite.usf
+++ b/Engine/Shaders/Private/SingleLayerWaterComposite.usf
@@ -175,9 +175,7 @@ SingleLayerWaterCompositeOutput SingleLayerWaterComposite(float2 BufferUV, float
 			float3 Reflection = 0.0f;
 
 			// TODO MAKE COMMON WITH ReflectionEnvironment?
-			float3 WorldPosition = mul(float4(ScreenPosition * GBuffer.Depth, GBuffer.Depth, 1), LWCHackToFloat(PrimaryView.ScreenToWorld)).xyz;
-			float3 CameraToPixel = normalize(WorldPosition - LWCHackToFloat(PrimaryView.WorldCameraOrigin));
-			float3 ReflectionVector = reflect(CameraToPixel, GBuffer.WorldNormal);
+			const float3 TranslatedWorldPosition = mul(float4(ScreenPosition * GBuffer.Depth, GBuffer.Depth, 1), View.ScreenToTranslatedWorld).xyz;
 			float IndirectIrradiance = 0.0;// GBuffer.IndirectIrradiance;
 			float IndirectSpecularOcclusion = 1.0f;
 			float3 ExtraIndirectSpecular = 0;
@@ -197,7 +195,7 @@ SingleLayerWaterCompositeOutput SingleLayerWaterComposite(float2 BufferUV, float
 
 			//
 			float3 N = GBuffer.WorldNormal;
-			float3 V = -CameraToPixel;
+			float3 V = normalize(View.TranslatedWorldCameraOrigin - TranslatedWorldPosition);
 			float3 R = 2 * dot(V, N) * N - V;
 			float NoV = saturate(dot(N, V));
 			// Point lobe in off-specular peak direction
@@ -205,7 +203,7 @@ SingleLayerWaterCompositeOutput SingleLayerWaterComposite(float2 BufferUV, float
 			const bool bCompositeSkylight = true;
 			Reflection += View.PreExposure * CompositeReflectionCapturesAndSkylight(
 				1,
-				WorldPosition,
+				TranslatedWorldPosition,
 				R,
 				GBuffer.Roughness,
 				IndirectIrradiance,
