Implements TSR's grand reprojection

diff --git a/Engine/Shaders/Private/TemporalSuperResolution/TSRDecimateHistory.usf b/Engine/Shaders/Private/TemporalSuperResolution/TSRDecimateHistory.usf
--- a/Engine/Shaders/Private/TemporalSuperResolution/TSRDecimateHistory.usf
+++ b/Engine/Shaders/Private/TemporalSuperResolution/TSRDecimateHistory.usf
@@ -11,6 +11,8 @@
 
 #define CONFIG_MANUAL_LDS_SPILL 1
 
+#define CONFIG_GRAND_REPROJECTION (DIM_GRAND_REPROJECTION)
+
 #define CONFIG_VELOCITY_SAMPLES 9
 #if CONFIG_VELOCITY_SAMPLES != 1
 	#define CONFIG_VELOCITY_SAMPLE_PAIRS ((CONFIG_VELOCITY_SAMPLES - 1) / 2)
@@ -33,6 +35,7 @@
 //------------------------------------------------------- PARAMETERS
 
 float4x4 RotationalClipToPrevClip;
+float2 CurrentToPrevInputJitter;
 
 Texture2D<float2> DilatedVelocityTexture;
 Texture2D<float> ClosestDepthTexture;
@@ -44,6 +47,7 @@ Texture2D<tsr_subpixel_details> PrevHistorySubpixelDetails;
 
 Texture2D<tsr_halfC> PrevHistoryGuide;
 Texture2D<tsr_half4> PrevHistoryMoire;
+Texture2D<float2> PrevVelocityTexture;
 float2 PrevGuideInfo_Extent;
 float2 PrevGuideInfo_ExtentInverse;
 uint2  PrevGuideInfo_ViewportMin;
@@ -61,6 +65,7 @@ RWTexture2D<tsr_halfC> ReprojectedHistoryGuideOutput;
 RWTexture2D<tsr_half4> ReprojectedHistoryMoireOutput;
 
 RWTexture2D<float2> HoleFilledVelocityOutput;
+RWTexture2D<float2> GrandHoleFilledVelocityOutput;
 RWTexture2D<tsr_half2> HoleFilledVelocityMaskOutput;
 RWTexture2D<tsr_half> ParallaxRejectionMaskOutput;
 RWTexture2D<tsr_subpixel_payload> HistorySubpixelDetailsOutput;
@@ -560,6 +565,51 @@ void MainCS(
 			FinalSubpixelDetails = UpdatedFinalSubpixelDetails;
 		}
 	}
+	
+	#if CONFIG_GRAND_REPROJECTION
+	float2 EncodedGrandScreenVelocity;
+	{
+		tsr_ushort2 InputPixelPos = (
+			Map8x8Tile2x2Lane(GroupThreadIndex) +
+			(tsr_ushort2(InputInfo_ViewportMin) + tsr_ushort2(GroupId) * tsr_ushort2(TILE_SIZE, TILE_SIZE)));
+			
+		float2 ScreenPos = ApplyScreenTransform(float2(InputPixelPos), InputPixelPosToScreenPos);
+		float2 ScreenVelocity = DecodeVelocityFromTexture(float4(EncodedScreenVelocity, 0.0, 0.0)).xy;
+
+		float2 PrevScreenPos = ScreenPos - ScreenVelocity;
+		float2 KernelHistoryBufferUV = ApplyScreenTransform(PrevScreenPos, ScreenPosToPrevHistoryGuideBufferUV) + CurrentToPrevInputJitter * PrevGuideInfo_ExtentInverse;
+		//float2 KernelHistoryBufferUV = ApplyScreenTransform(PrevScreenPos, ScreenPosToPrevHistoryGuideBufferUV);
+		
+		FBilinearSampleInfos BilinearInter = GetBilinearSampleLevelInfos(KernelHistoryBufferUV, InputInfo_Extent, InputInfo_ExtentInverse);
+		
+		float2 PrevVelocity = 0.0;
+		float PrevVelocityWeight = 0.0;
+
+		UNROLL_N(4)
+		for (uint i = 0; i < 4; i++)
+		{
+			float BilinearWeight = GetSampleWeight(BilinearInter, i);
+			uint2 PixelPos = GetSamplePixelCoord(BilinearInter, i);
+
+			bool bValidPixelPos = all(and(PixelPos >= InputInfo_ViewportMin, PixelPos < InputInfo_ViewportMax));
+			BilinearWeight = bValidPixelPos ? BilinearWeight : 0.0;
+
+			float2 EncodedPrevVelocity = PrevVelocityTexture[PixelPos];
+
+			float2 PrevScreenVelocity = DecodeVelocityFromTexture(float4(EncodedPrevVelocity, 0.0, 0.0)).xy;
+
+			PrevVelocity += BilinearWeight * PrevScreenVelocity;
+			PrevVelocityWeight += BilinearWeight;
+		}
+
+		PrevVelocity *= SafeRcp(PrevVelocityWeight);
+
+		EncodedGrandScreenVelocity = EncodeVelocityToTexture(float3(ScreenVelocity + PrevVelocity, 0.0)).xy;
+
+		//Debug.r = abs(ScreenVelocity.x) * ScreenVelocityToInputPixelVelocity.x;
+		Debug.g = abs(PrevVelocity.x) * ScreenVelocityToInputPixelVelocity.x;
+	}
+	#endif // CONFIG_GRAND_REPROJECTION
 
 	// Correct exposure of ReprojectedHistoryGuideColor
 	{
@@ -590,6 +640,9 @@ void MainCS(
 			ReprojectedHistoryMoireOutput[OutputPixelPos] = ReprojectedHistoryMoireColor;
 		#endif
 		HoleFilledVelocityOutput[OutputPixelPos] = EncodedScreenVelocity;
+		#if CONFIG_GRAND_REPROJECTION
+			GrandHoleFilledVelocityOutput[OutputPixelPos] = EncodedGrandScreenVelocity;
+		#endif
 		HoleFilledVelocityMaskOutput[InputPixelPos] = tsr_half2(0.0, RejectHighFrequency);
 		ParallaxRejectionMaskOutput[OutputPixelPos] = ParallaxRejectionMask;
 		HistorySubpixelDetailsOutput[OutputPixelPos] = FinalSubpixelDetails;
diff --git a/Engine/Shaders/Private/TemporalSuperResolution/TSRUpdateHistory.usf b/Engine/Shaders/Private/TemporalSuperResolution/TSRUpdateHistory.usf
--- a/Engine/Shaders/Private/TemporalSuperResolution/TSRUpdateHistory.usf
+++ b/Engine/Shaders/Private/TemporalSuperResolution/TSRUpdateHistory.usf
@@ -24,7 +24,7 @@
 
 #elif DIM_UPDATE_QUALITY == 3 // Epic
 	#define CONFIG_SAMPLES 6
-	#define CONFIG_INTERPOLATE_VELOCITY 1
+	#define CONFIG_INTERPOLATE_VELOCITY 0
 	#define CONFIG_REJECTION_ANTI_ALIASING 1
 
 	// To support sg.AntiAliasingQuality=4 with same temporal response as sg.AntiAliasingQuality=3
@@ -36,7 +36,6 @@
 
 #if 1
 	#define CONFIG_HIGH_FREQUENCY_ONLY 1
-	#define CONFIG_HIGH_FREQUENCY_YCOCG 0
 #endif
 
 #if DIM_SEPARATE_TRANSLUCENCY
@@ -58,6 +57,8 @@
 	#define CONFIG_TRANSLUCENCY_REJECTION 1
 #endif
 
+#define CONFIG_GRAND_REPROJECTION (DIM_GRAND_REPROJECTION)
+
 #define CONFIG_RESPONSIVE_STENCIL 1
 #define CONFIG_TAA_RESPONSIVE_BEHAVIOR 1
 #define CONFIG_LOAD_TRANSLUCENCY_REJECTION 0
@@ -105,6 +106,9 @@ Texture2D<tsr_ushort> AntiAliasingTexture;
 Texture2D<tsr_half> NoiseFilteringTexture;
 Texture2D<tsr_half2> HoleFilledVelocityMaskTexture;
 
+Texture2D<float2> GrandDilatedVelocityTexture;
+Texture2DArray<tsr_halfC> GrandPrevColorArray;
+
 float2 TranslucencyInfo_Extent;
 float2 TranslucencyInfo_ExtentInverse;
 float2 TranslucencyInfo_ScreenPosToViewportScale;
@@ -125,11 +129,13 @@ int2   TranslucencyPixelPosMax;
 FScreenTransform HistoryPixelPosToScreenPos;
 FScreenTransform HistoryPixelPosToInputPPCo;
 FScreenTransform HistoryPixelPosToTranslucencyPPCo;
+FScreenTransform ScreenPosToGrandPrevHistoryBufferUV;
 float3 HistoryQuantizationError;
 float MinTranslucencyRejection;
 float InvWeightClampingPixelSpeed;
 float InputToHistoryFactor;
 float InputContributionMultiplier;
+float GrandPrevPreExposureCorrection;
 uint ResponsiveStencilMask;
 uint bGenerateOutputMip1;
 uint bHasSeparateTranslucency;
@@ -140,11 +146,18 @@ Texture2D<tsr_half3> PrevHistory_Translucency;
 Texture2D<tsr_half>  PrevHistory_TranslucencyAlpha;
 Texture2D<tsr_subpixel_details> PrevHistory_SubpixelDetails;
 
+Texture2D<tsr_halfC> PrevHistory_PrevHighFrequency;
+Texture2D<tsr_halfC> PrevHistory_PrevHighFrequencyResultant;
+Texture2D<tsr_halfC> PrevHistory_HighFrequencyOverblur;
+
 RWTexture2D<tsr_halfC> SceneColorOutputMip0;
 RWTexture2D<tsr_halfC> SceneColorOutputMip1;
 
 uint HistoryOutput_ArrayIndices_HighFrequency;
 uint HistoryOutput_ArrayIndices_Translucency;
+uint HistoryOutput_ArrayIndices_PrevHighFrequency;
+uint HistoryOutput_ArrayIndices_PrevHighFrequencyResultant;
+uint HistoryOutput_ArrayIndices_HighFrequencyOverblur;
 
 RWTexture2DArray<tsr_halfC> HistoryOutput_ColorArray;
 RWTexture2D<tsr_half2> HistoryOutput_Metadata;
@@ -348,6 +361,9 @@ void MainCS(
 
 	// Fetch reprojections related informations.
 	float2x2 PrevScreenPos = ScreenPos;
+	#if CONFIG_GRAND_REPROJECTION
+		float2x2 GrandPrevScreenPos = ScreenPos;
+	#endif
 	tsr_half2 LowFrequencyRejection = tsr_half(1.0).xx;
 	tsr_half2 TranslucencyRejection = tsr_half(1.0).xx;
 	tsr_half2 OutputPixelVelocity = tsr_half(0.0).xx;
@@ -368,6 +384,10 @@ void MainCS(
 		#else
 			float2 RawEncodedVelocity[DPV_PIXEL_PER_LANE];
 		#endif
+		
+		#if CONFIG_GRAND_REPROJECTION
+			float2 RawEncodedGrandVelocity[DPV_PIXEL_PER_LANE];
+		#endif
 
 		tsr_half RawParallaxRejectionMask[DPV_PIXEL_PER_LANE];
 		uint RawParallaxFactorBits[DPV_PIXEL_PER_LANE];
@@ -395,6 +415,11 @@ void MainCS(
 				#if !CONFIG_INTERPOLATE_VELOCITY
 					RawEncodedVelocity[PixelId]   = DilatedVelocityTexture[LocalInputPixelPos];
 				#endif
+				
+				#if CONFIG_GRAND_REPROJECTION
+					RawEncodedGrandVelocity[PixelId] = GrandDilatedVelocityTexture[LocalInputPixelPos];
+				#endif
+
 				RawParallaxRejectionMask[PixelId] = ParallaxRejectionMaskTexture[LocalInputPixelPos];
 				RawParallaxFactorBits[PixelId]    = ParallaxFactorTexture[LocalInputPixelPos];
 				RawSceneStencilRef[PixelId]       = InputSceneStencilTexture.Load(int3(LocalInputPixelPos, 0)) STENCIL_COMPONENT_SWIZZLE;
@@ -497,6 +522,16 @@ void MainCS(
 
 				PrevScreenPos = ScreenPos - ScreenVelocity;
 				OutputPixelVelocity = tsr_half2(dpv_length(dpv_mul(ScreenVelocity, HistoryInfo_ViewportSize)));
+				
+				#if CONFIG_GRAND_REPROJECTION
+				{
+					float2x2 GrandScreenVelocity = dpv_interleave_registers(
+						DecodeVelocityFromTexture(float4(RawEncodedGrandVelocity[0], 0.0, 0.0)).xy,
+						DecodeVelocityFromTexture(float4(RawEncodedGrandVelocity[1], 0.0, 0.0)).xy);
+					
+					GrandPrevScreenPos = ScreenPos - GrandScreenVelocity;
+				}
+				#endif
 			}
 			
 			// Detect whether the history reprojection is successful
@@ -571,6 +606,13 @@ void MainCS(
 	tsr_halfCx2 PrevHistoryMommentMax;
 
 	tsr_half4x2 PrevHistoryTranslucency;
+	
+	#if CONFIG_GRAND_REPROJECTION
+		tsr_halfCx2 GrandPrevToPrev;
+		tsr_halfCx2 GrandPrevToPrevResultant;
+		tsr_halfCx2 PrevHighFrequencyOverblur;
+		tsr_halfCx2 GrandReprojection;
+	#endif
 
 	PLATFORM_SPECIFIC_ISOLATE
 	{
@@ -581,6 +623,13 @@ void MainCS(
 
 		tsr_halfCx2 HighFrequencyMin = INFINITE_FLOAT;
 		tsr_halfCx2 HighFrequencyMax = -INFINITE_FLOAT;
+		
+		#if CONFIG_GRAND_REPROJECTION
+			tsr_halfCx2 AccumulateGrandPrevToPrev = tsr_half(0);
+			tsr_halfCx2 AccumulateGrandPrevToPrevResultant = tsr_half(0);
+			tsr_halfCx2 AccumulatePrevHighFrequencyOverblur = tsr_half(0);
+			tsr_halfCx2 AccumulateGrandReprojection = tsr_half(0);
+		#endif
 
 		// Sample raw history
 		{
@@ -593,12 +642,32 @@ void MainCS(
 			FCatmullRomSamples Samples1 = GetBicubic2DCatmullRomSamples(
 				dpv_hi(PrevHistoryBufferUV), PrevHistoryInfo_Extent, PrevHistoryInfo_ExtentInverse);
 			
+			#if CONFIG_GRAND_REPROJECTION
+				float2x2 GrandPrevHistoryBufferUV = ApplyScreenTransform(GrandPrevScreenPos, ScreenPosToGrandPrevHistoryBufferUV);
+				GrandPrevHistoryBufferUV[0] = clamp(GrandPrevHistoryBufferUV[0], PrevHistoryInfo_UVViewportBilinearMin[0], PrevHistoryInfo_UVViewportBilinearMax[0]);
+				GrandPrevHistoryBufferUV[1] = clamp(GrandPrevHistoryBufferUV[1], PrevHistoryInfo_UVViewportBilinearMin[1], PrevHistoryInfo_UVViewportBilinearMax[1]);
+
+				FCatmullRomSamples GrandSamples0 = GetBicubic2DCatmullRomSamples(
+					dpv_lo(GrandPrevHistoryBufferUV), PrevHistoryInfo_Extent, PrevHistoryInfo_ExtentInverse);
+				FCatmullRomSamples GrandSamples1 = GetBicubic2DCatmullRomSamples(
+					dpv_hi(GrandPrevHistoryBufferUV), PrevHistoryInfo_Extent, PrevHistoryInfo_ExtentInverse);
+			#endif
+
 			tsr_halfC RawHighFrequency[BICUBIC_CATMULL_ROM_SAMPLES][DPV_PIXEL_PER_LANE];
 			tsr_half2 RawMetadata[BICUBIC_CATMULL_ROM_SAMPLES][DPV_PIXEL_PER_LANE];
 			tsr_half3 RawTranslucencyColor[BICUBIC_CATMULL_ROM_SAMPLES][DPV_PIXEL_PER_LANE];
 			tsr_half  RawTranslucencyAlpha[BICUBIC_CATMULL_ROM_SAMPLES][DPV_PIXEL_PER_LANE];
 
 			tsr_half2 RawKernelWeight[BICUBIC_CATMULL_ROM_SAMPLES];
+			
+			#if CONFIG_GRAND_REPROJECTION
+				tsr_halfC RawGrandPrevToPrev[BICUBIC_CATMULL_ROM_SAMPLES][DPV_PIXEL_PER_LANE];
+				tsr_halfC RawGrandPrevToPrevResultant[BICUBIC_CATMULL_ROM_SAMPLES][DPV_PIXEL_PER_LANE];
+				tsr_halfC RawPrevHighFrequencyOverblur[BICUBIC_CATMULL_ROM_SAMPLES][DPV_PIXEL_PER_LANE];
+				tsr_halfC RawGrandReprojection[BICUBIC_CATMULL_ROM_SAMPLES][DPV_PIXEL_PER_LANE];
+
+				tsr_half2 RawGrandKernelWeight[BICUBIC_CATMULL_ROM_SAMPLES];
+			#endif
 
 			// Issues texture fetches.
 			PLATFORM_SPECIFIC_ISOLATE
@@ -612,6 +681,10 @@ void MainCS(
 					tsr_half2 KernelWeight = dpv_force_interleave_registers(tsr_half(Samples0.Weight[i]), tsr_half(Samples1.Weight[i]));
 
 					RawKernelWeight[i] = KernelWeight;
+					
+					#if CONFIG_GRAND_REPROJECTION
+						RawGrandKernelWeight[i] = dpv_force_interleave_registers(tsr_half(GrandSamples0.Weight[i]), tsr_half(GrandSamples1.Weight[i]));
+					#endif
 
 					UNROLL_N(2)
 					for (uint PixelId = 0; PixelId < DPV_PIXEL_PER_LANE; PixelId++)
@@ -626,6 +699,22 @@ void MainCS(
 							RawTranslucencyColor[i][PixelId] = BilinearSampleColorHistory(PrevHistory_Translucency, SampleUV);
 							RawTranslucencyAlpha[i][PixelId] = PrevHistory_TranslucencyAlpha.SampleLevel(GlobalBilinearClampedSampler, SampleUV, 0);
 						#endif
+						
+						#if CONFIG_GRAND_REPROJECTION
+							RawGrandPrevToPrev[i][PixelId] = PrevHistory_PrevHighFrequency.SampleLevel(GlobalBilinearClampedSampler, SampleUV, 0);
+							RawGrandPrevToPrevResultant[i][PixelId] = PrevHistory_PrevHighFrequencyResultant.SampleLevel(GlobalBilinearClampedSampler, SampleUV, 0);
+							RawPrevHighFrequencyOverblur[i][PixelId] = PrevHistory_HighFrequencyOverblur.SampleLevel(GlobalBilinearClampedSampler, SampleUV, 0);
+
+							{
+								float2 GrandSampleUV0 = clamp(GrandSamples0.UV[i], PrevHistoryInfo_UVViewportBilinearMin, PrevHistoryInfo_UVViewportBilinearMax);
+								float2 GrandSampleUV1 = clamp(GrandSamples1.UV[i], PrevHistoryInfo_UVViewportBilinearMin, PrevHistoryInfo_UVViewportBilinearMax);
+					
+								float2 GrandSampleUV = PixelId == 0 ? GrandSampleUV0 : GrandSampleUV1;
+						
+								RawGrandReprojection[i][PixelId] = GrandPrevColorArray.SampleLevel(GlobalBilinearClampedSampler, float3(GrandSampleUV, 0.0), 0);
+							
+							}
+						#endif
 					}
 
 				} // for (uint i = 0; i < BICUBIC_CATMULL_ROM_SAMPLES; i++)
@@ -643,6 +732,13 @@ void MainCS(
 						tsr_half4(RawTranslucencyColor[i][0], RawTranslucencyAlpha[i][0]),
 						tsr_half4(RawTranslucencyColor[i][1], RawTranslucencyAlpha[i][1]));
 					
+					#if CONFIG_GRAND_REPROJECTION
+						tsr_halfCx2 SampleGrandPrevious = dpv_interleave_registers_array(RawGrandPrevToPrev[i]);
+						tsr_halfCx2 SampleGrandPreviousResultant = dpv_interleave_registers_array(RawGrandPrevToPrevResultant[i]);
+						tsr_halfCx2 SamplePrevHighFrequencyOverblur = dpv_interleave_registers_array(RawPrevHighFrequencyOverblur[i]) - SampleHighFrequency;
+						tsr_halfCx2 SampleGrandReprojection = dpv_interleave_registers_array(RawGrandReprojection[i]);
+					#endif
+
 					tsr_half2 KernelWeight = RawKernelWeight[i];
 
 					tsr_half2 InvWeight = SafeRcp(SampleMetadata[0]);
@@ -656,6 +752,13 @@ void MainCS(
 					#if CONFIG_ACCUMULATE_TRANSLUCENCY_SEPARATLY
 						AccumulateTranslucency  += dpv_scale(SampleTranslucency, KernelWeight);
 					#endif
+
+					#if CONFIG_GRAND_REPROJECTION
+						AccumulateGrandPrevToPrev  += dpv_scale(SampleGrandPrevious, KernelWeight);
+						AccumulateGrandPrevToPrevResultant  += dpv_scale(SampleGrandPreviousResultant, KernelWeight);
+						AccumulatePrevHighFrequencyOverblur  += dpv_scale(SamplePrevHighFrequencyOverblur, KernelWeight);
+						AccumulateGrandReprojection  += dpv_scale(SampleGrandReprojection, RawGrandKernelWeight[i]);
+					#endif
 				}
 				
 				tsr_half2 FinalMultiplier = dpv_force_interleave_registers(tsr_half(Samples0.FinalMultiplier), tsr_half(Samples1.FinalMultiplier));
@@ -663,6 +766,15 @@ void MainCS(
 				AccumulateHighFrequency = dpv_scale(AccumulateHighFrequency, FinalMultiplier);
 				AccumulateMetadata      = dpv_scale(AccumulateMetadata, FinalMultiplier);
 				AccumulateTranslucency  = dpv_scale(AccumulateTranslucency, FinalMultiplier);
+
+				#if CONFIG_GRAND_REPROJECTION
+					tsr_half2 GrandFinalMultiplier = dpv_interleave_registers(tsr_half(GrandSamples0.FinalMultiplier), tsr_half(GrandSamples1.FinalMultiplier));
+
+					AccumulateGrandPrevToPrev          = dpv_scale(AccumulateGrandPrevToPrev, FinalMultiplier);
+					AccumulateGrandPrevToPrevResultant = dpv_scale(AccumulateGrandPrevToPrevResultant, FinalMultiplier);
+					AccumulatePrevHighFrequencyOverblur = dpv_scale(AccumulatePrevHighFrequencyOverblur, FinalMultiplier);
+					AccumulateGrandReprojection        = dpv_scale(AccumulateGrandReprojection, GrandFinalMultiplier);
+				#endif
 			}
 			
 			#if CONFIG_SCENE_COLOR_APLHA
@@ -684,6 +796,13 @@ void MainCS(
 			PrevHistoryValidity = AccumulateMetadata[1];
 
 			PrevHistoryTranslucency = AccumulateTranslucency;
+
+			#if CONFIG_GRAND_REPROJECTION
+				GrandPrevToPrev          = AccumulateGrandPrevToPrev;
+				GrandPrevToPrevResultant = AccumulateGrandPrevToPrevResultant;
+				PrevHighFrequencyOverblur = AccumulatePrevHighFrequencyOverblur;
+				GrandReprojection        = AccumulateGrandReprojection;
+			#endif
 		}
 
 		// Correct history
@@ -697,6 +816,22 @@ void MainCS(
 			CorrectExposure(PrevHighFrequencyColor[0]);
 			CorrectExposure(PrevHighFrequencyColor[1]);
 			CorrectExposure(PrevHighFrequencyColor[2]);
+			
+			#if CONFIG_GRAND_REPROJECTION
+				CorrectExposure(GrandPrevToPrev[0]);
+				CorrectExposure(GrandPrevToPrev[1]);
+				CorrectExposure(GrandPrevToPrev[2]);
+				CorrectExposure(GrandPrevToPrevResultant[0]);
+				CorrectExposure(GrandPrevToPrevResultant[1]);
+				CorrectExposure(GrandPrevToPrevResultant[2]);
+				CorrectExposure(PrevHighFrequencyOverblur[0]);
+				CorrectExposure(PrevHighFrequencyOverblur[1]);
+				CorrectExposure(PrevHighFrequencyOverblur[2]);
+				
+				GrandReprojection[0] *= tsr_half(GrandPrevPreExposureCorrection);
+				GrandReprojection[1] *= tsr_half(GrandPrevPreExposureCorrection);
+				GrandReprojection[2] *= tsr_half(GrandPrevPreExposureCorrection);
+			#endif
 
 			// Fixes negative validity that can happen due to negative lobes of Catrom.
 			PrevHistoryValidity = max(PrevHistoryValidity, tsr_half(0.0).xx);
@@ -864,13 +999,8 @@ void MainCS(
 				FilteredTranslucencyColor = TranslucencyColor;
 			}
 
-			#if CONFIG_HIGH_FREQUENCY_YCOCG
-				InputMinColor = min(InputMinColor, TransformColorForClampingBox(InputColor));
-				InputMaxColor = max(InputMaxColor, TransformColorForClampingBox(InputColor));
-			#else
-				InputMinColor = min(InputMinColor, InputColor);
-				InputMaxColor = max(InputMaxColor, InputColor);
-			#endif
+			InputMinColor = min(InputMinColor, InputColor);
+			InputMaxColor = max(InputMaxColor, InputColor);
 
 			#if CONFIG_ACCUMULATE_TRANSLUCENCY_SEPARATLY
 			{
@@ -983,30 +1113,35 @@ void MainCS(
 		#endif
 	}
 	#endif
-
+	
 	// Contribute current frame input into the predictor for next frame.
 	const tsr_half Histeresis = rcp(tsr_half(MAX_SAMPLE_COUNT));
 	const tsr_half PredictionOnlyValidity = Histeresis * tsr_half(2.0);
 	
 	tsr_halfCx2 FinalHighFrequencyColor;
 	tsr_half2 FinalHistoryValidity;
+	
+	#if CONFIG_GRAND_REPROJECTION
+		tsr_halfCx2 FinalPrevHighFrequencyColor;
+		tsr_halfCx2 FinalPrevHighFrequencyColorResultant;
+		tsr_halfCx2 FinalOverblur = tsr_half(0.0);
+		tsr_halfCx2 FinalHighFrequencyColorOverblurCorrection = tsr_half(0.0);
+	#endif
+
 	{
-		#if CONFIG_HIGH_FREQUENCY_YCOCG
-			tsr_halfCx2 ClampedPrevHighFrequencyColor = YCoCgToRGB(clamp(TransformColorForClampingBox(PrevHighFrequencyColor), dpv_scale(InputMinColor, PrevHistoryWeight), dpv_scale(InputMaxColor, PrevHistoryWeight)));
-		#else
-			tsr_halfCx2 ClampedPrevHighFrequencyColor = clamp(PrevHighFrequencyColor, dpv_scale(InputMinColor, PrevHistoryWeight), dpv_scale(InputMaxColor, PrevHistoryWeight));
-		#endif
-		
+		tsr_half2 InvPrevHistoryWeight = SafeRcp(PrevHistoryWeight);
+		tsr_half2 InvPrevHistoryValidity = SafeRcp(PrevHistoryValidity);
+
 		tsr_half2 TotalRejection = select(or(bIsOffScreen, bIsDisoccluded), tsr_half(0.0).xx, saturate(LowFrequencyRejection * tsr_half(4.0)));
 
 		tsr_half2 BlendClamp = TotalRejection * HoleFilledVelocityMask[1];
-
-		PrevHighFrequencyColor[0] = lerp(ClampedPrevHighFrequencyColor[0], PrevHighFrequencyColor[0], BlendClamp);
-		PrevHighFrequencyColor[1] = lerp(ClampedPrevHighFrequencyColor[1], PrevHighFrequencyColor[1], BlendClamp);
-		PrevHighFrequencyColor[2] = lerp(ClampedPrevHighFrequencyColor[2], PrevHighFrequencyColor[2], BlendClamp);
-
-		tsr_half2 InvPrevHistoryWeight = SafeRcp(PrevHistoryWeight);
-		tsr_half2 InvPrevHistoryValidity = SafeRcp(PrevHistoryValidity);
+		
+		tsr_halfCx2 UnpremultipliedPrevHighFrequencyColor = dpv_scale(PrevHighFrequencyColor, InvPrevHistoryWeight);
+		tsr_halfCx2 ClampedPrevHighFrequencyColor = clamp(UnpremultipliedPrevHighFrequencyColor, InputMinColor, InputMaxColor);
+		tsr_halfCx2 BlendedPrevHighFrequencyColor = UnpremultipliedPrevHighFrequencyColor;
+		BlendedPrevHighFrequencyColor[0] = lerp(ClampedPrevHighFrequencyColor[0], UnpremultipliedPrevHighFrequencyColor[0], BlendClamp);
+		BlendedPrevHighFrequencyColor[1] = lerp(ClampedPrevHighFrequencyColor[1], UnpremultipliedPrevHighFrequencyColor[1], BlendClamp);
+		BlendedPrevHighFrequencyColor[2] = lerp(ClampedPrevHighFrequencyColor[2], UnpremultipliedPrevHighFrequencyColor[2], BlendClamp);
 
 		// Compute how much the history needs to be rejected, based on the completeness.
 		tsr_half2 PrevHistoryRejectionWeight = select(bIsOffScreen, tsr_half(0.0), LowFrequencyRejection);
@@ -1018,14 +1153,19 @@ void MainCS(
 		
 		// Clamp the validity due to motion to maintain better sharpness in history reprojection under motion.
 		tsr_half2 ClampedPrevHistoryValidity;
-		#if 1
+		#if CONFIG_GRAND_REPROJECTION
+		{
+			tsr_half2 MaxValidity = lerp(PredictionOnlyValidity, tsr_half(1.0), TotalRejection);
+			ClampedPrevHistoryValidity = min(PrevHistoryValidity, MaxValidity);
+		}
+		#elif 1
 		{
 			tsr_half2 MaxValidity = tsr_half(1.0) - (tsr_half(1.0) - PredictionOnlyValidity) * saturate(OutputPixelVelocity * tsr_half(InvWeightClampingPixelSpeed));
 
 			// Clamp up the max validity to favorise stability under motion on high contrast edges.
 			#if 1
 			{
-				tsr_half2 PrevHistoryLuma = Luma4(PrevHighFrequencyColor) * InvPrevHistoryWeight;
+				tsr_half2 PrevHistoryLuma = Luma4(BlendedPrevHighFrequencyColor);
 				tsr_half2 FilteredLuma = Luma4(FilteredInputColor);
 
 				tsr_half2 MinValidityForStability = abs(FilteredLuma - PrevHistoryLuma) / max(FilteredLuma, PrevHistoryLuma);
@@ -1056,7 +1196,7 @@ void MainCS(
 		tsr_half2 PrevMomentWeight = max(OutputValidity2 - DesiredCurrentContribution, tsr_half(0.0));
 		tsr_half2 CurrentMomentWeight = min(DesiredCurrentContribution, OutputValidity2);
 
-		tsr_half2 PrevHistoryToneWeight = HdrWeightY(Luma4(PrevHighFrequencyColor) * InvPrevHistoryWeight);
+		tsr_half2 PrevHistoryToneWeight = HdrWeightY(Luma4(BlendedPrevHighFrequencyColor));
 		tsr_half2 FilteredInputToneWeight = HdrWeight4(FilteredInputColor);
 			
 		tsr_half2 BlendPrevHistory = PrevMomentWeight * PrevHistoryToneWeight;
@@ -1065,11 +1205,53 @@ void MainCS(
 		tsr_half2 CommonWeight = SafeRcp(BlendPrevHistory + BlendFilteredInput);
 
 		FinalHighFrequencyColor = (
-			dpv_scale(PrevHighFrequencyColor, CommonWeight * BlendPrevHistory * InvPrevHistoryWeight) +
+			dpv_scale(BlendedPrevHighFrequencyColor, CommonWeight * BlendPrevHistory) +
 			dpv_scale(FilteredInputColor, CommonWeight * BlendFilteredInput));
 				
 		// Quantize validity for the 8bit encoding to avoid numerical shift between color and validity.
 		FinalHistoryValidity = ceil(tsr_half(255.0) * OutputValidity) * rcp(tsr_half(255.0));
+
+		#if CONFIG_GRAND_REPROJECTION
+		{
+			tsr_half2 FinalHistoryWeight = max(FinalHistoryValidity, MIN_IMPLICIT_HISTORY_WEIGHT);
+
+			tsr_halfCx2 BoxSize = dpv_scale(PrevHistoryMommentMax - PrevHistoryMommentMin, tsr_half(0.5) * PrevHistoryWeight);
+
+			// Compute MeasuredResultant of the previous frame
+			//tsr_halfCx2 PrevResultant = -dpv_min(-GrandPrevToPrevResultant / GrandPrevToPrev, tsr_half(0.0).xxx);
+			tsr_halfCx2 PrevResultant;
+			PrevResultant[0] = clamp(GrandPrevToPrevResultant[0] * SafeRcp(GrandPrevToPrev[0]), tsr_half(0.0), tsr_half(1.1));
+			PrevResultant[1] = clamp(GrandPrevToPrevResultant[1] * SafeRcp(GrandPrevToPrev[1]), tsr_half(0.0), tsr_half(1.1));
+			PrevResultant[2] = clamp(GrandPrevToPrevResultant[2] * SafeRcp(GrandPrevToPrev[2]), tsr_half(0.0), tsr_half(1.1));
+			#if CONFIG_SCENE_COLOR_APLHA
+				PrevResultant[3] = clamp(GrandPrevToPrevResultant[3] * SafeRcp(GrandPrevToPrev[3]), tsr_half(0.0), tsr_half(1.1));
+			#endif
+			
+			// Compute how much frame N-2 got overblur by reprojecting into N-1 and N versus directly to N.
+			tsr_halfCx2 PrevReprojectionOverblur = GrandReprojection - GrandPrevToPrev;
+
+			// Clamp error imprecision.
+			bool2 bIsOverblurLegal = (
+				and(clamp(PrevReprojectionOverblur[0], -BoxSize[0], BoxSize[0]) == PrevReprojectionOverblur[0],
+				and(clamp(PrevReprojectionOverblur[1], -BoxSize[1], BoxSize[1]) == PrevReprojectionOverblur[1],
+				    clamp(PrevReprojectionOverblur[2], -BoxSize[2], BoxSize[2]) == PrevReprojectionOverblur[2])));
+
+			tsr_half2 PrevReprojectionOverblurScale = select(bIsOverblurLegal, tsr_half(1.0), tsr_half(0.0));
+
+			FinalOverblur = dpv_scale(PrevResultant, BlendClamp * TotalRejection) * (PrevHighFrequencyOverblur + dpv_scale(PrevReprojectionOverblur, PrevReprojectionOverblurScale));
+
+			//tsr_half2 Rejection = MeasureRejectionFactor(
+			//	UnpremultipliedPrevHighFrequencyColor, BlendedPrevHighFrequencyColor,
+			//	FilteredInputColor, InputMinColor, InputMaxColor,
+			//	MeasureBackbufferLDRQuantizationError());
+
+			tsr_half2 MeasuredResultant = saturate(CommonWeight * BlendPrevHistory * BlendClamp);
+
+			FinalPrevHighFrequencyColor = PrevHighFrequencyColor;
+			FinalPrevHighFrequencyColorResultant = dpv_scale(PrevHighFrequencyColor, MeasuredResultant * InvPrevHistoryWeight * FinalHistoryWeight);
+			FinalHighFrequencyColorOverblurCorrection = dpv_scale(FinalOverblur, MeasuredResultant * InvPrevHistoryWeight);
+		}
+		#endif
 	}
 	
 	// Compute the low frequency history
@@ -1165,7 +1347,11 @@ void MainCS(
 	// Compute final output
 	tsr_halfCx2 FinalOutputColor;
 	{
-		FinalOutputColor = FinalHighFrequencyColor;
+		#if CONFIG_GRAND_REPROJECTION
+			FinalOutputColor = FinalHighFrequencyColor + FinalHighFrequencyColorOverblurCorrection;
+		#else
+			FinalOutputColor = FinalHighFrequencyColor;
+		#endif
 
 		#if CONFIG_ACCUMULATE_TRANSLUCENCY_SEPARATLY
 			FinalOutputColor[0] = FinalOutputColor[0] * FinalTranslucencyColor[3] + FinalTranslucencyColor[0];
@@ -1263,6 +1449,11 @@ void MainCS(
 					#endif 
 					HistoryOutput_TranslucencyAlpha[dpv_lo(LocalHistoryPixelPos)] = dpv_lo(FinalTranslucencyColor).a;
 				#endif
+				#if CONFIG_GRAND_REPROJECTION
+					HistoryOutput_ColorArray[tsr_short3(dpv_lo(LocalHistoryPixelPos), HistoryOutput_ArrayIndices_PrevHighFrequency)] = dpv_lo(FinalPrevHighFrequencyColor);
+					HistoryOutput_ColorArray[tsr_short3(dpv_lo(LocalHistoryPixelPos), HistoryOutput_ArrayIndices_PrevHighFrequencyResultant)] = dpv_lo(FinalPrevHighFrequencyColorResultant);
+					HistoryOutput_ColorArray[tsr_short3(dpv_lo(LocalHistoryPixelPos), HistoryOutput_ArrayIndices_HighFrequencyOverblur)] = dpv_lo(FinalOverblur) + dpv_lo(PremultipliedHighFrequencyColor);
+				#endif
 			}
 		
 			// Output final history hi pixel.
@@ -1278,6 +1469,11 @@ void MainCS(
 					#endif
 					HistoryOutput_TranslucencyAlpha[dpv_hi(LocalHistoryPixelPos)] = dpv_hi(FinalTranslucencyColor).a;
 				#endif
+				#if CONFIG_GRAND_REPROJECTION
+					HistoryOutput_ColorArray[tsr_short3(dpv_hi(LocalHistoryPixelPos), HistoryOutput_ArrayIndices_PrevHighFrequency)] = dpv_hi(FinalPrevHighFrequencyColor);
+					HistoryOutput_ColorArray[tsr_short3(dpv_hi(LocalHistoryPixelPos), HistoryOutput_ArrayIndices_PrevHighFrequencyResultant)] = dpv_hi(FinalPrevHighFrequencyColorResultant);
+					HistoryOutput_ColorArray[tsr_short3(dpv_hi(LocalHistoryPixelPos), HistoryOutput_ArrayIndices_HighFrequencyOverblur)] = dpv_hi(FinalOverblur) + dpv_hi(PremultipliedHighFrequencyColor);
+				#endif
 			}
 			#endif // !CONFIG_ENABLE_DUAL_PIXEL_VECTORIZATION
 		}
diff --git a/Engine/Source/Runtime/Renderer/Private/PostProcess/TemporalSuperResolution.cpp b/Engine/Source/Runtime/Renderer/Private/PostProcess/TemporalSuperResolution.cpp
--- a/Engine/Source/Runtime/Renderer/Private/PostProcess/TemporalSuperResolution.cpp
+++ b/Engine/Source/Runtime/Renderer/Private/PostProcess/TemporalSuperResolution.cpp
@@ -36,6 +36,11 @@ TAutoConsoleVariable<int32> CVarTSRHistorySeparateTranslucency(
 	TEXT("Whether separate translucency should be accumulated separatly."),
 	ECVF_RenderThreadSafe);
 
+TAutoConsoleVariable<int32> CVarTSRHistoryGrandReprojection(
+	TEXT("r.TSR.History.GrandReprojection"), 0,
+	TEXT("Experimental functionality to keep higher sharpness in TSR's history in motion."),
+	ECVF_Scalability | ECVF_RenderThreadSafe);
+
 TAutoConsoleVariable<int32> CVarTSRWaveOps(
 	TEXT("r.TSR.WaveOps"), 1,
 	TEXT("Whether to use wave ops in the shading rejection heuristics"),
@@ -132,6 +137,11 @@ END_SHADER_PARAMETER_STRUCT()
 BEGIN_SHADER_PARAMETER_STRUCT(FTSRHistoryArrayIndices, )
 	SHADER_PARAMETER(int32, HighFrequency)
 	SHADER_PARAMETER(int32, Translucency)
+
+	SHADER_PARAMETER(int32, PrevHighFrequency)
+	SHADER_PARAMETER(int32, PrevHighFrequencyResultant)
+	SHADER_PARAMETER(int32, HighFrequencyOverblur)
+
 	SHADER_PARAMETER(int32, Size)
 END_SHADER_PARAMETER_STRUCT()
 
@@ -148,6 +158,11 @@ END_SHADER_PARAMETER_STRUCT()
 BEGIN_SHADER_PARAMETER_STRUCT(FTSRHistorySRVs, )
 	SHADER_PARAMETER_RDG_TEXTURE_SRV(Texture2D, HighFrequency)
 	SHADER_PARAMETER_RDG_TEXTURE_SRV(Texture2D, Metadata)
+
+	SHADER_PARAMETER_RDG_TEXTURE_SRV(Texture2D, PrevHighFrequency)
+	SHADER_PARAMETER_RDG_TEXTURE_SRV(Texture2D, PrevHighFrequencyResultant)
+	SHADER_PARAMETER_RDG_TEXTURE_SRV(Texture2D, HighFrequencyOverblur)
+
 	SHADER_PARAMETER_RDG_TEXTURE(Texture2D, SubpixelDetails)
 	SHADER_PARAMETER_RDG_TEXTURE_SRV(Texture2D, Translucency)
 	SHADER_PARAMETER_RDG_TEXTURE_SRV(Texture2D, TranslucencyAlpha)
@@ -167,6 +182,7 @@ BEGIN_SHADER_PARAMETER_STRUCT(FTSRPrevHistoryParameters, )
 	SHADER_PARAMETER_STRUCT(FScreenPassTextureViewportParameters, PrevHistoryInfo)
 	SHADER_PARAMETER(FScreenTransform, ScreenPosToPrevHistoryBufferUV)
 	SHADER_PARAMETER(FScreenTransform, ScreenPosToPrevSubpixelDetails)
+	SHADER_PARAMETER(FScreenTransform, ScreenPosToGrandPrevHistoryBufferUV)
 	SHADER_PARAMETER(FVector2f, PrevSubpixelDetailsExtent)
 	SHADER_PARAMETER(float, HistoryPreExposureCorrection)
 END_SHADER_PARAMETER_STRUCT()
@@ -176,6 +192,7 @@ enum class ETSRHistoryFormatBits : uint32
 	None = 0,
 	Translucency = 1 << 0,
 	Moire = 1 << 1,
+	GrandReprojection = 1 << 2,
 };
 ENUM_CLASS_FLAGS(ETSRHistoryFormatBits);
 
@@ -191,6 +208,20 @@ FTSRHistoryArrayIndices TranslateHistoryFormatBitsToArrayIndices(ETSRHistoryForm
 		ArrayIndices.Translucency = ArrayIndices.Size++;
 	}
 
+	if (EnumHasAnyFlags(HistoryFormatBits, ETSRHistoryFormatBits::GrandReprojection))
+	{
+		ArrayIndices.PrevHighFrequency = ArrayIndices.Size + 0;
+		ArrayIndices.PrevHighFrequencyResultant = ArrayIndices.Size + 1;
+		ArrayIndices.HighFrequencyOverblur = ArrayIndices.Size + 2;
+		ArrayIndices.Size += 3;
+	}
+	else
+	{
+		ArrayIndices.PrevHighFrequency = -1;
+		ArrayIndices.PrevHighFrequencyResultant = -1;
+		ArrayIndices.HighFrequencyOverblur = -1;
+	}
+
 	return ArrayIndices;
 }
 
@@ -205,6 +236,13 @@ FTSRHistorySRVs CreateSRVs(FRDGBuilder& GraphBuilder, const FTSRHistoryTextures&
 		SRVs.Translucency = GraphBuilder.CreateSRV(FRDGTextureSRVDesc::CreateForSlice(Textures.ColorArray, Textures.ArrayIndices.Translucency));
 		SRVs.TranslucencyAlpha = GraphBuilder.CreateSRV(FRDGTextureSRVDesc::Create(Textures.TranslucencyAlpha));
 	}
+	if (Textures.ArrayIndices.HighFrequencyOverblur >= 0)
+	{
+		SRVs.PrevHighFrequency = GraphBuilder.CreateSRV(FRDGTextureSRVDesc::CreateForSlice(Textures.ColorArray, Textures.ArrayIndices.PrevHighFrequency));
+		SRVs.PrevHighFrequencyResultant = GraphBuilder.CreateSRV(FRDGTextureSRVDesc::CreateForSlice(Textures.ColorArray, Textures.ArrayIndices.PrevHighFrequencyResultant));
+		SRVs.HighFrequencyOverblur = GraphBuilder.CreateSRV(FRDGTextureSRVDesc::CreateForSlice(Textures.ColorArray, Textures.ArrayIndices.HighFrequencyOverblur));
+	}
+
 	SRVs.Guide = Textures.Guide;
 	SRVs.Moire = Textures.Moire;
 	return SRVs;
@@ -336,11 +374,13 @@ class FTSRDecimateHistoryCS : public FTSRShader
 	SHADER_USE_PARAMETER_STRUCT(FTSRDecimateHistoryCS, FTSRShader);
 
 	class FMoireReprojectionDim : SHADER_PERMUTATION_BOOL("DIM_MOIRE_REPROJECTION");
-	using FPermutationDomain = TShaderPermutationDomain<FMoireReprojectionDim>;
+	class FGrandReprojectionDim : SHADER_PERMUTATION_BOOL("DIM_GRAND_REPROJECTION");
+	using FPermutationDomain = TShaderPermutationDomain<FMoireReprojectionDim, FGrandReprojectionDim>;
 
 	BEGIN_SHADER_PARAMETER_STRUCT(FParameters, )
 		SHADER_PARAMETER_STRUCT_INCLUDE(FTSRCommonParameters, CommonParameters)
 		SHADER_PARAMETER(FMatrix44f, RotationalClipToPrevClip)
+		SHADER_PARAMETER(FVector2f, CurrentToPrevInputJitter)
 		SHADER_PARAMETER(float, WorldDepthToPixelWorldRadius)
 
 		SHADER_PARAMETER_RDG_TEXTURE(Texture2D, InputSceneColorTexture)
@@ -355,6 +395,7 @@ class FTSRDecimateHistoryCS : public FTSRShader
 
 		SHADER_PARAMETER_RDG_TEXTURE(Texture2D, PrevHistoryGuide)
 		SHADER_PARAMETER_RDG_TEXTURE(Texture2D, PrevHistoryMoire)
+		SHADER_PARAMETER_RDG_TEXTURE(Texture2D, PrevVelocityTexture)
 		SHADER_PARAMETER_STRUCT(FScreenPassTextureViewportParameters, PrevGuideInfo)
 		SHADER_PARAMETER(FScreenTransform, InputPixelPosToReprojectScreenPos)
 		SHADER_PARAMETER(FScreenTransform, ScreenPosToPrevHistoryGuideBufferUV)
@@ -364,6 +405,8 @@ class FTSRDecimateHistoryCS : public FTSRShader
 		SHADER_PARAMETER_RDG_TEXTURE_UAV(RWTexture2D, ReprojectedHistoryMoireOutput)
 		SHADER_PARAMETER_RDG_TEXTURE_UAV(RWTexture2D, HoleFilledVelocityOutput)
 		SHADER_PARAMETER_RDG_TEXTURE_UAV(RWTexture2D, HoleFilledVelocityMaskOutput)
+		SHADER_PARAMETER_RDG_TEXTURE_UAV(RWTexture2D, GrandHoleFilledVelocityOutput)
+
 		SHADER_PARAMETER_RDG_TEXTURE_UAV(RWTexture2D, ParallaxRejectionMaskOutput)
 		SHADER_PARAMETER_RDG_TEXTURE_UAV(RWTexture2D, HistorySubpixelDetailsOutput)
 
@@ -573,8 +616,9 @@ class FTSRUpdateHistoryCS : public FTSRShader
 
 	class FQualityDim : SHADER_PERMUTATION_ENUM_CLASS("DIM_UPDATE_QUALITY", EQuality);
 	class FSeparateTranslucencyDim : SHADER_PERMUTATION_BOOL("DIM_SEPARATE_TRANSLUCENCY");
+	class FGrandReprojectionDim : SHADER_PERMUTATION_BOOL("DIM_GRAND_REPROJECTION");
 
-	using FPermutationDomain = TShaderPermutationDomain<FQualityDim, FSeparateTranslucencyDim>;
+	using FPermutationDomain = TShaderPermutationDomain<FQualityDim, FSeparateTranslucencyDim, FGrandReprojectionDim>;
 
 	BEGIN_SHADER_PARAMETER_STRUCT(FParameters, )
 		SHADER_PARAMETER_STRUCT_INCLUDE(FTSRCommonParameters, CommonParameters)
@@ -586,6 +630,7 @@ class FTSRUpdateHistoryCS : public FTSRShader
 		SHADER_PARAMETER_RDG_TEXTURE(Texture2D, TranslucencyRejectionTexture)
 
 		SHADER_PARAMETER_RDG_TEXTURE(Texture2D, DilatedVelocityTexture)
+		SHADER_PARAMETER_RDG_TEXTURE(Texture2D, GrandDilatedVelocityTexture)
 		SHADER_PARAMETER_RDG_TEXTURE_SRV(Texture2D, ParallaxFactorTexture)
 		SHADER_PARAMETER_RDG_TEXTURE(Texture2D, ParallaxRejectionMaskTexture)
 		SHADER_PARAMETER_RDG_TEXTURE(Texture2D, AntiAliasingTexture)
@@ -605,12 +650,14 @@ class FTSRUpdateHistoryCS : public FTSRShader
 		SHADER_PARAMETER(float, InvWeightClampingPixelSpeed)
 		SHADER_PARAMETER(float, InputToHistoryFactor)
 		SHADER_PARAMETER(float, InputContributionMultiplier)
+		SHADER_PARAMETER(float, GrandPrevPreExposureCorrection)
 		SHADER_PARAMETER(int32, ResponsiveStencilMask)
 		SHADER_PARAMETER(int32, bGenerateOutputMip1)
 		SHADER_PARAMETER(int32, bHasSeparateTranslucency)
 
 		SHADER_PARAMETER_STRUCT_INCLUDE(FTSRPrevHistoryParameters, PrevHistoryParameters)
 		SHADER_PARAMETER_STRUCT(FTSRHistorySRVs, PrevHistory)
+		SHADER_PARAMETER_RDG_TEXTURE(Texture2DArray, GrandPrevColorArray)
 
 		SHADER_PARAMETER_RDG_TEXTURE_UAV(RWTexture2D, SceneColorOutputMip0)
 		SHADER_PARAMETER_RDG_TEXTURE_UAV(RWTexture2D, SceneColorOutputMip1)
@@ -618,6 +665,19 @@ class FTSRUpdateHistoryCS : public FTSRShader
 		SHADER_PARAMETER_RDG_TEXTURE_UAV(RWTexture2DArray, DebugOutput)
 	END_SHADER_PARAMETER_STRUCT()
 
+	static bool ShouldCompilePermutation(const FGlobalShaderPermutationParameters& Parameters)
+	{
+		FPermutationDomain PermutationVector(Parameters.PermutationId);
+
+		// Only compile the grand reprojection when translucency isn't accumulated separately.
+		if (PermutationVector.Get<FGrandReprojectionDim>() && PermutationVector.Get<FSeparateTranslucencyDim>())
+		{
+			return false;
+		}
+
+		return FTSRShader::ShouldCompilePermutation(Parameters);
+	}
+
 	static void ModifyCompilationEnvironment(const FGlobalShaderPermutationParameters& Parameters, FShaderCompilerEnvironment& OutEnvironment)
 	{
 		FTSRShader::ModifyCompilationEnvironment(Parameters, OutEnvironment);
@@ -734,6 +794,8 @@ ITemporalUpscaler::FOutputs AddTemporalSuperResolutionPasses(
 	// Whether separate translucency should be temporarily accumulated separatly.
 	const bool bAccumulateTranslucencySeparately = CVarTSRHistorySeparateTranslucency.GetValueOnRenderThread() != 0;
 
+	const bool bGrandReprojection = !bAccumulateTranslucencySeparately && CVarTSRHistoryGrandReprojection.GetValueOnRenderThread() != 0;
+
 	// whether TSR passes can run on async compute.
 	int32 AsyncComputePasses = GSupportsEfficientAsyncCompute ? CVarTSRAsyncCompute.GetValueOnRenderThread() : 0;
 
@@ -751,6 +813,11 @@ ITemporalUpscaler::FOutputs AddTemporalSuperResolutionPasses(
 		{
 			HistoryFormatBits |= ETSRHistoryFormatBits::Moire;
 		}
+
+		if (bGrandReprojection)
+		{
+			HistoryFormatBits |= ETSRHistoryFormatBits::GrandReprojection;
+		}
 	}
 
 	// Whether to use camera cut shader permutation or not.
@@ -1055,7 +1122,7 @@ ITemporalUpscaler::FOutputs AddTemporalSuperResolutionPasses(
 
 		FRDGTextureDesc ArrayDesc = FRDGTextureDesc::Create2DArray(
 			HistoryExtent,
-			(CVarTSRR11G11B10History.GetValueOnRenderThread() != 0 && !bSupportsAlpha) ? PF_FloatR11G11B10 : PF_FloatRGBA,
+			(CVarTSRR11G11B10History.GetValueOnRenderThread() != 0 && !bSupportsAlpha && !bGrandReprojection) ? PF_FloatR11G11B10 : PF_FloatRGBA,
 			FClearValueBinding::None,
 			TexCreate_ShaderResource | TexCreate_UAV,
 			History.ArrayIndices.Size);
@@ -1105,6 +1172,9 @@ ITemporalUpscaler::FOutputs AddTemporalSuperResolutionPasses(
 		DummyHistorySRVs.Metadata = GraphBuilder.CreateSRV(FRDGTextureSRVDesc::Create(BlackDummy));
 		DummyHistorySRVs.Translucency = GraphBuilder.CreateSRV(FRDGTextureSRVDesc::Create(BlackDummy));
 		DummyHistorySRVs.TranslucencyAlpha = GraphBuilder.CreateSRV(FRDGTextureSRVDesc::Create(WhiteDummy));
+		DummyHistorySRVs.PrevHighFrequency = GraphBuilder.CreateSRV(FRDGTextureSRVDesc::Create(BlackDummy));
+		DummyHistorySRVs.PrevHighFrequencyResultant = GraphBuilder.CreateSRV(FRDGTextureSRVDesc::Create(BlackDummy));
+		DummyHistorySRVs.HighFrequencyOverblur = GraphBuilder.CreateSRV(FRDGTextureSRVDesc::Create(BlackDummy));
 		DummyHistorySRVs.Guide = BlackDummy;
 		DummyHistorySRVs.Moire = BlackDummy;
 
@@ -1112,23 +1182,28 @@ ITemporalUpscaler::FOutputs AddTemporalSuperResolutionPasses(
 	}
 
 	// Setup the previous frame history
+	FRDGTextureRef PrevVelocity = BlackDummy;
+	FRDGTextureRef GrandPrevColorArray = BlackDummy;
+	FTSRHistoryTextures PrevHistory;
 	FTSRHistorySRVs PrevHistorySRVs;
 	if (!bCameraCut)
 	{
-		FTSRHistoryTextures PrevHistory;
 		PrevHistory.ArrayIndices = History.ArrayIndices;
 
 		// Register filterable history
-		PrevHistory.ColorArray = GraphBuilder.RegisterExternalTexture(InputHistory.ColorArray);
-		PrevHistory.Metadata = GraphBuilder.RegisterExternalTexture(InputHistory.Metadata);
-		PrevHistory.TranslucencyAlpha = InputHistory.TranslucencyAlpha.IsValid() ? GraphBuilder.RegisterExternalTexture(InputHistory.TranslucencyAlpha) : DummyHistorySRVs.TranslucencyAlpha->Desc.Texture;
-		PrevHistory.Guide = GraphBuilder.RegisterExternalTexture(InputHistory.Guide);
-		PrevHistory.Moire = InputHistory.Moire.IsValid() ? GraphBuilder.RegisterExternalTexture(InputHistory.Moire) : DummyHistorySRVs.Moire;
+		PrevHistory.ColorArray = GraphBuilder.RegisterExternalTexture(InputHistory.ColorArray, TEXT("TSR.PrevHistory.ColorArray"));
+		PrevHistory.Metadata = GraphBuilder.RegisterExternalTexture(InputHistory.Metadata, TEXT("TSR.PrevHistory.Metadata"));
+		PrevHistory.TranslucencyAlpha = InputHistory.TranslucencyAlpha.IsValid() ? GraphBuilder.RegisterExternalTexture(InputHistory.TranslucencyAlpha, TEXT("TSR.PrevHistory.TranslucencyAlpha")) : DummyHistorySRVs.TranslucencyAlpha->Desc.Texture;
+		PrevHistory.Guide = GraphBuilder.RegisterExternalTexture(InputHistory.Guide, TEXT("TSR.PrevHistory.Guide"));
+		PrevHistory.Moire = InputHistory.Moire.IsValid() ? GraphBuilder.RegisterExternalTexture(InputHistory.Moire, TEXT("TSR.PrevHistory.Moire")) : DummyHistorySRVs.Moire;
 
 		// Register non-filterable history
-		PrevHistory.SubpixelDetails = GraphBuilder.RegisterExternalTexture(InputHistory.SubpixelDetails);
+		PrevHistory.SubpixelDetails = GraphBuilder.RegisterExternalTexture(InputHistory.SubpixelDetails, TEXT("TSR.PrevHistory.SubpixelInfo"));
 
 		PrevHistorySRVs = CreateSRVs(GraphBuilder, PrevHistory);
+
+		PrevVelocity = InputHistory.Velocity.IsValid() ? GraphBuilder.RegisterExternalTexture(InputHistory.Velocity) : PrevVelocity;
+		GrandPrevColorArray = InputHistory.PrevColorArray.IsValid() ? GraphBuilder.RegisterExternalTexture(InputHistory.PrevColorArray, TEXT("TSR.GrandPrevHistory.ColorArray")) : GrandPrevColorArray;
 	}
 	else
 	{
@@ -1141,6 +1216,7 @@ ITemporalUpscaler::FOutputs AddTemporalSuperResolutionPasses(
 		// Setup prev history parameters.
 		FScreenPassTextureViewport PrevHistoryViewport(PrevHistorySRVs.HighFrequency->Desc.Texture->Desc.Extent, InputHistory.OutputViewportRect);
 		FScreenPassTextureViewport PrevSubpixelDetailsViewport(PrevHistorySRVs.SubpixelDetails->Desc.Extent, InputHistory.InputViewportRect);
+		FScreenPassTextureViewport GrandPrevSubpixelDetailsViewport(GrandPrevColorArray->Desc.Extent, InputHistory.PrevOutputViewportRect);
 
 		if (bCameraCut)
 		{
@@ -1148,6 +1224,8 @@ ITemporalUpscaler::FOutputs AddTemporalSuperResolutionPasses(
 			PrevHistoryViewport.Rect = FIntRect(FIntPoint(0, 0), FIntPoint(1, 1));
 			PrevSubpixelDetailsViewport.Extent = FIntPoint(1, 1);
 			PrevSubpixelDetailsViewport.Rect = FIntRect(FIntPoint(0, 0), FIntPoint(1, 1));
+			GrandPrevSubpixelDetailsViewport.Extent = FIntPoint(1, 1);
+			GrandPrevSubpixelDetailsViewport.Rect = FIntRect(FIntPoint(0, 0), FIntPoint(1, 1));
 		}
 
 		PrevHistoryParameters.PrevHistoryInfo = GetScreenPassTextureViewportParameters(PrevHistoryViewport);
@@ -1155,6 +1233,8 @@ ITemporalUpscaler::FOutputs AddTemporalSuperResolutionPasses(
 			PrevHistoryViewport, FScreenTransform::ETextureBasis::ScreenPosition, FScreenTransform::ETextureBasis::TextureUV);
 		PrevHistoryParameters.ScreenPosToPrevSubpixelDetails = FScreenTransform::ChangeTextureBasisFromTo(
 			PrevSubpixelDetailsViewport, FScreenTransform::ETextureBasis::ScreenPosition, FScreenTransform::ETextureBasis::TextureUV);
+		PrevHistoryParameters.ScreenPosToGrandPrevHistoryBufferUV = FScreenTransform::ChangeTextureBasisFromTo(
+			GrandPrevSubpixelDetailsViewport, FScreenTransform::ETextureBasis::ScreenPosition, FScreenTransform::ETextureBasis::TextureUV);
 		PrevHistoryParameters.PrevSubpixelDetailsExtent = PrevSubpixelDetailsViewport.Extent;
 		PrevHistoryParameters.HistoryPreExposureCorrection = View.PreExposure / View.PrevViewInfo.SceneColorPreExposure;
 	}
@@ -1164,6 +1244,7 @@ ITemporalUpscaler::FOutputs AddTemporalSuperResolutionPasses(
 	FRDGTextureRef ReprojectedHistoryMoireTexture = nullptr;
 	FRDGTextureRef ParallaxRejectionMaskTexture = nullptr;
 	FRDGTextureRef HoleFilledVelocityMaskTexture = nullptr;
+	FRDGTextureRef GrandHoleFilledVelocityTexture = BlackDummy;
 	{
 		FRDGTextureRef HoleFilledVelocityTexture;
 		{
@@ -1177,6 +1258,7 @@ ITemporalUpscaler::FOutputs AddTemporalSuperResolutionPasses(
 
 			Desc.Format = PF_G16R16;
 			HoleFilledVelocityTexture = GraphBuilder.CreateTexture(Desc, TEXT("TSR.Velocity.HoleFilled"));
+			GrandHoleFilledVelocityTexture = GraphBuilder.CreateTexture(Desc, TEXT("TSR.Velocity.GrandVelocity"));
 
 			Desc.Format = PF_R8G8;
 			HoleFilledVelocityMaskTexture = GraphBuilder.CreateTexture(Desc, TEXT("TSR.velocity.HoleFillMask"));
@@ -1213,6 +1295,9 @@ ITemporalUpscaler::FOutputs AddTemporalSuperResolutionPasses(
 
 			PassParameters->RotationalClipToPrevClip = FMatrix44f(RotationalInvViewProj * RotationalPrevViewProj);		// LWC_TODO: Precision loss?
 		}
+		{
+			PassParameters->CurrentToPrevInputJitter = InputHistory.PrevTemporalJitterPixels - CommonParameters.InputJitter;
+		}
 		{
 			float TanHalfFieldOfView = View.ViewMatrices.GetInvProjectionMatrix().M[0][0];
 
@@ -1229,6 +1314,7 @@ ITemporalUpscaler::FOutputs AddTemporalSuperResolutionPasses(
 
 		PassParameters->PrevHistoryParameters = PrevHistoryParameters;
 		PassParameters->PrevHistorySubpixelDetails = PrevHistorySRVs.SubpixelDetails;
+		PassParameters->PrevVelocityTexture = PrevVelocity;
 
 		{
 			FScreenPassTextureViewport PrevHistoryColorViewport(PrevHistorySRVs.Guide->Desc.Extent, InputHistory.InputViewportRect);
@@ -1247,12 +1333,14 @@ ITemporalUpscaler::FOutputs AddTemporalSuperResolutionPasses(
 		PassParameters->ReprojectedHistoryMoireOutput = GraphBuilder.CreateUAV(ReprojectedHistoryMoireTexture);
 		PassParameters->HoleFilledVelocityOutput = GraphBuilder.CreateUAV(HoleFilledVelocityTexture);
 		PassParameters->HoleFilledVelocityMaskOutput = GraphBuilder.CreateUAV(HoleFilledVelocityMaskTexture);
+		PassParameters->GrandHoleFilledVelocityOutput = GraphBuilder.CreateUAV(GrandHoleFilledVelocityTexture);
 		PassParameters->ParallaxRejectionMaskOutput = GraphBuilder.CreateUAV(ParallaxRejectionMaskTexture);
 		PassParameters->HistorySubpixelDetailsOutput = GraphBuilder.CreateUAV(History.SubpixelDetails);
 		PassParameters->DebugOutput = CreateDebugUAV(InputExtent, TEXT("Debug.TSR.DecimateHistory"));
 
 		FTSRDecimateHistoryCS::FPermutationDomain PermutationVector;
 		PermutationVector.Set<FTSRDecimateHistoryCS::FMoireReprojectionDim>(FlickeringFramePeriod > 0);
+		PermutationVector.Set<FTSRDecimateHistoryCS::FGrandReprojectionDim>(bGrandReprojection);
 
 		TShaderMapRef<FTSRDecimateHistoryCS> ComputeShader(View.ShaderMap, PermutationVector);
 		FComputeShaderUtils::AddPass(
@@ -1567,6 +1655,7 @@ ITemporalUpscaler::FOutputs AddTemporalSuperResolutionPasses(
 		PassParameters->TranslucencyRejectionTexture = TranslucencyRejectionTexture ? TranslucencyRejectionTexture : BlackDummy;
 
 		PassParameters->DilatedVelocityTexture = DilatedVelocityTexture;
+		PassParameters->GrandDilatedVelocityTexture = GrandHoleFilledVelocityTexture;
 		PassParameters->ParallaxFactorTexture = ParallaxFactorTexture;
 		PassParameters->ParallaxRejectionMaskTexture = ParallaxRejectionMaskTexture;
 		PassParameters->AntiAliasingTexture = AntiAliasingTexture;
@@ -1587,12 +1676,14 @@ ITemporalUpscaler::FOutputs AddTemporalSuperResolutionPasses(
 		PassParameters->InvWeightClampingPixelSpeed = 1.0f / CVarTSRWeightClampingPixelSpeed.GetValueOnRenderThread();
 		PassParameters->InputToHistoryFactor = float(HistorySize.X) / float(InputRect.Width());
 		PassParameters->InputContributionMultiplier = FMath::Pow(float(HistorySize.X) / float(OutputRect.Width()), 2.0f); 
+		PassParameters->GrandPrevPreExposureCorrection = bCameraCut ? 1.0f : View.PreExposure / InputHistory.PrevSceneColorPreExposure;
 		PassParameters->ResponsiveStencilMask = CVarTSREnableResponiveAA.GetValueOnRenderThread() ? (STENCIL_TEMPORAL_RESPONSIVE_AA_MASK) : 0;
 		PassParameters->bGenerateOutputMip1 = ((PassInputs.bGenerateOutputMip1 || PassInputs.bAllowDownsampleSceneColor) && HistorySize == OutputRect.Size()) ? 1 : 0;
 		PassParameters->bHasSeparateTranslucency = bHasSeparateTranslucency;
 
 		PassParameters->PrevHistoryParameters = PrevHistoryParameters;
 		PassParameters->PrevHistory = PrevHistorySRVs;
+		PassParameters->GrandPrevColorArray = GrandPrevColorArray;
 
 		PassParameters->HistoryOutput = CreateUAVs(GraphBuilder, History);
 		if (HistorySize != OutputRect.Size())
@@ -1625,6 +1716,7 @@ ITemporalUpscaler::FOutputs AddTemporalSuperResolutionPasses(
 		FTSRUpdateHistoryCS::FPermutationDomain PermutationVector;
 		PermutationVector.Set<FTSRUpdateHistoryCS::FQualityDim>(UpdateHistoryQuality);
 		PermutationVector.Set<FTSRUpdateHistoryCS::FSeparateTranslucencyDim>(bAccumulateTranslucencySeparately);
+		PermutationVector.Set<FTSRUpdateHistoryCS::FGrandReprojectionDim>(bGrandReprojection);
 
 		TShaderMapRef<FTSRUpdateHistoryCS> ComputeShader(View.ShaderMap, PermutationVector);
 		FComputeShaderUtils::AddPass(
@@ -1716,6 +1808,9 @@ ITemporalUpscaler::FOutputs AddTemporalSuperResolutionPasses(
 		OutputHistory.InputViewportRect = InputRect;
 		OutputHistory.OutputViewportRect = FIntRect(FIntPoint(0, 0), HistorySize);
 		OutputHistory.FormatBit = uint32(HistoryFormatBits);
+		OutputHistory.PrevSceneColorPreExposure = View.PrevViewInfo.SceneColorPreExposure;
+		OutputHistory.PrevTemporalJitterPixels = CommonParameters.InputJitter;
+		OutputHistory.PrevOutputViewportRect = InputHistory.OutputViewportRect;
 
 		// Extract filterable history
 		GraphBuilder.QueueTextureExtraction(History.ColorArray, &OutputHistory.ColorArray);
@@ -1736,6 +1831,19 @@ ITemporalUpscaler::FOutputs AddTemporalSuperResolutionPasses(
 			GraphBuilder.QueueTextureExtraction(History.Moire, &OutputHistory.Moire);
 		}
 
+		if (bGrandReprojection)
+		{
+			GraphBuilder.QueueTextureExtraction(DilatedVelocityTexture, &OutputHistory.Velocity);
+			if (PrevHistory.ColorArray)
+			{
+				GraphBuilder.QueueTextureExtraction(PrevHistory.ColorArray, &OutputHistory.PrevColorArray);
+			}
+			else
+			{
+				GraphBuilder.QueueTextureExtraction(History.ColorArray, &OutputHistory.PrevColorArray);
+			}
+		}
+
 		// Extract the translucency buffer to compare it with next frame
 		if (bRejectSeparateTranslucency)
 		{
diff --git a/Engine/Source/Runtime/Renderer/Private/SceneRendering.h b/Engine/Source/Runtime/Renderer/Private/SceneRendering.h
--- a/Engine/Source/Runtime/Renderer/Private/SceneRendering.h
+++ b/Engine/Source/Runtime/Renderer/Private/SceneRendering.h
@@ -900,23 +900,33 @@ struct FTemporalAAHistory
 // Structure in charge of storing all information about TSR's history.
 struct FTSRHistory
 {
-	// Filterable output resolution.
+	// Output resolution.
 	TRefCountPtr<IPooledRenderTarget> ColorArray;
 	TRefCountPtr<IPooledRenderTarget> Metadata;
 	TRefCountPtr<IPooledRenderTarget> TranslucencyAlpha;
 
-	// Non-filterable output resolution
+	// Input resolution representation of the output
 	TRefCountPtr<IPooledRenderTarget> SubpixelDetails;
-
-	// Filterable input resolution representation of the output
 	TRefCountPtr<IPooledRenderTarget> Guide;
 	TRefCountPtr<IPooledRenderTarget> Moire;
+	TRefCountPtr<IPooledRenderTarget> Velocity;
+
+	// Previous frame's history.
+	TRefCountPtr<IPooledRenderTarget> PrevColorArray;
 
+	// Frame's input and output resolution.
 	FIntRect InputViewportRect;
 	FIntRect OutputViewportRect;
 
+	// Previous frame's informations.
+	FIntRect PrevOutputViewportRect;
+	FVector2f PrevTemporalJitterPixels;
+	float PrevSceneColorPreExposure = 1.0f;
+
+	// Format of the history for auto camera cut when setting change.
 	uint32 FormatBit = 0;
 
+
 	void SafeRelease()
 	{
 		*this = FTSRHistory();
