Add view mode support for RT hit

diff --git a/Engine/Shaders/Private/RayTracing/RayTracingDebug.usf b/Engine/Shaders/Private/RayTracing/RayTracingDebug.usf
--- a/Engine/Shaders/Private/RayTracing/RayTracingDebug.usf
+++ b/Engine/Shaders/Private/RayTracing/RayTracingDebug.usf
@@ -1,10 +1,13 @@
 // Copyright Epic Games, Inc. All Rights Reserved.
 
+#define STRATA_MATERIALCONTAINER_IS_VIEWRESOURCE 0
+
 #include "../Common.ush"
 #include "../ShadingCommon.ush"
 #include "../ColorMap.ush"
 #include "../Hash.ush"
 #include "../Visualization.ush"
+#include "../Strata/Strata.ush"
 
 #include "RayTracingCommon.ush"
 #include "RayTracingHitGroupCommon.ush"
@@ -29,6 +32,19 @@ float3 FarFieldReferencePos;
 StructuredBuffer<FRayTracingInstanceDebugData> InstancesDebugData;
 StructuredBuffer<FRayTracingPickingFeedback> PickingBuffer;
 
+float4 AnisoDebugColor(float InAniso)
+{
+	// Follow same mapping than raster buffer visualization (BP located in Engine/Contents/Anisotropy)
+	float AniG = saturate(InAniso);
+	float AniB = saturate(-1.0 * InAniso);
+	return float4(0.0, pow(AniG, 2.2), pow(AniB, 2.2), 1.0f);
+}
+
+float4 RoughnessDebugColor(float InRoughness)
+{
+	return float4(pow(InRoughness, 2.2f).xxx, 1.0f);
+}
+
 float4 DebugRayTracingMaterial(FRayDesc Ray, bool bIsFarField)
 {
 	uint2 PixelCoord = DispatchRaysIndex().xy + View.ViewRectMin.xy;
@@ -69,95 +85,75 @@ float4 DebugRayTracingMaterial(FRayDesc Ray, bool bIsFarField)
 #if PLATFORM_SUPPORTS_SHADER_TIMESTAMP
 	uint64_t TimeEnd = GetShaderTimestamp();
 #endif
-
-	// HitT is packed into the alpha for far field compositing
-	float4 Result = (RAY_TRACING_DEBUG_VIZ_FAR_FIELD) ? float4(0, 0, 0, -1) : float4(0, 0, 0, 1);
-
 	if (Payload.IsHit())
 	{
+
+		// 1. Non-material dependent debug view
 		switch (VisualizationMode)
 		{
-		default:
-		case RAY_TRACING_DEBUG_VIZ_RADIANCE:
-			Result = float4(Payload.Radiance, 1.0f);
-			break;
-		case RAY_TRACING_DEBUG_VIZ_WORLD_NORMAL:
-			Result = float4(Payload.WorldNormal * 0.5 + 0.5, 1.0f);
-			break;
-		case RAY_TRACING_DEBUG_VIZ_OPACITY:
-			Result = float4(1.0f - Payload.Opacity, 1.0f - Payload.Opacity, 1.0f - Payload.Opacity, 1.0f);
-			break;
-		case RAY_TRACING_DEBUG_VIZ_BLENDING_MODE:
-			Result = float4(Payload.BlendingMode, Payload.BlendingMode, Payload.BlendingMode, 1.0f);
-			break;
-		case RAY_TRACING_DEBUG_VIZ_LIGHTING_CHANNEL_MASK:
-			Result = float4(Payload.PrimitiveLightingChannelMask & 0x1, Payload.PrimitiveLightingChannelMask & 0x2, Payload.PrimitiveLightingChannelMask & 0x4, 1.0f);
-			break;
-		case RAY_TRACING_DEBUG_VIZ_INDIRECT_IRRADIANCE:
-			Result = float4(Payload.IndirectIrradiance, 1.0f);
-			break;
-		case RAY_TRACING_DEBUG_VIZ_WORLD_POSITION:
-			Result = float4(Payload.TranslatedWorldPos, 1.0f);
-			break;
-		case RAY_TRACING_DEBUG_VIZ_HITKIND:
-			Result = Payload.IsFrontFace() ? float4(0.0, 1.0, 0.0, 1.0f) : float4(1.0, 0.0, 0.0, 1.0f);
-			break;
+		case RAY_TRACING_DEBUG_VIZ_RADIANCE:				return float4(Payload.Radiance, 1.0f);
+		case RAY_TRACING_DEBUG_VIZ_OPACITY:					return float4((1.0f - Payload.Opacity).xxx, 1.0f);
+		case RAY_TRACING_DEBUG_VIZ_BLENDING_MODE:			return float4(Payload.BlendingMode.xxx, 1.0f);
+		case RAY_TRACING_DEBUG_VIZ_LIGHTING_CHANNEL_MASK:	return float4(Payload.PrimitiveLightingChannelMask & 0x1, Payload.PrimitiveLightingChannelMask & 0x2, Payload.PrimitiveLightingChannelMask & 0x4, 1.0f);
+		case RAY_TRACING_DEBUG_VIZ_INDIRECT_IRRADIANCE:		return float4(Payload.IndirectIrradiance, 1.0f);
+		case RAY_TRACING_DEBUG_VIZ_WORLD_POSITION:			return float4(Payload.TranslatedWorldPos, 1.0f);
+		case RAY_TRACING_DEBUG_VIZ_HITKIND:					return Payload.IsFrontFace() ? float4(0.0, 1.0, 0.0, 1.0f) : float4(1.0, 0.0, 0.0, 1.0f);
 		#if PLATFORM_SUPPORTS_SHADER_TIMESTAMP
-		case RAY_TRACING_DEBUG_VIZ_PERFORMANCE:
-			Result.rgb = ColorMapInferno(float(TimeEnd - TimeBegin) * TimingScale);
-			break;
+		case RAY_TRACING_DEBUG_VIZ_PERFORMANCE:				return float4(ColorMapInferno(float(TimeEnd - TimeBegin) * TimingScale), 1);
 		#endif // PLATFORM_SUPPORTS_SHADER_TIMESTAMP
+		}
 
-		// STRATA_TODO
-	#if !STRATA_ENABLED
-		case RAY_TRACING_DEBUG_VIZ_FAR_FIELD:
-			Result = float4(Payload.BaseColor.rgb, Payload.HitT);
-			break;
-		case RAY_TRACING_DEBUG_VIZ_GBUFFER_AO:
-			Result = float4(Payload.GBufferAO, Payload.GBufferAO, Payload.GBufferAO, 1.0f);
-			break;
-		case RAY_TRACING_DEBUG_VIZ_SHADING_MODEL:
-			Result = float4(GetShadingModelColor(Payload.ShadingModelID), 1.0f);
-			break;
-		case RAY_TRACING_DEBUG_VIZ_BASE_COLOR:
-			Result = float4(Payload.BaseColor.rgb, 1.0f);
-			break;
-		case RAY_TRACING_DEBUG_VIZ_DIFFUSE_COLOR:
-			Result = float4(Payload.DiffuseColor.rgb, 1.0f);
-			break;
-		case RAY_TRACING_DEBUG_VIZ_SPECULAR_COLOR:
-			Result = float4(Payload.SpecularColor.rgb, 1.0f);
-			break;
-		case RAY_TRACING_DEBUG_VIZ_METALLIC:
-			Result = float4(Payload.Metallic, Payload.Metallic, Payload.Metallic, 1.0f);
-			break;
-		case RAY_TRACING_DEBUG_VIZ_SPECULAR:
-			Result = float4(Payload.Specular, Payload.Specular, Payload.Specular, 1.0f);
-			break;
-		case RAY_TRACING_DEBUG_VIZ_ROUGHNESS:
-			Result = float4(pow(Payload.Roughness, 2.2f).xxx, 1.0f); // NOTE: pow 2.2 is used for consistency with BufferVisualization/Roughness shader
-			break;
-		case RAY_TRACING_DEBUG_VIZ_IOR:
-			Result = float4(Payload.Ior, Payload.Ior, Payload.Ior, 1.0f);
-			break;
-		case RAY_TRACING_DEBUG_VIZ_CUSTOM_DATA:
-			Result = float4(Payload.CustomData);
-			break;
-		case RAY_TRACING_DEBUG_VIZ_WORLD_TANGENT:
-			Result = float4(Payload.WorldTangent * 0.5 + 0.5, 1.0f);
-			break;
-		case RAY_TRACING_DEBUG_VIZ_ANISOTROPY:
-			{	
-			// Follow same mapping than raster buffer visualization (BP located in Engine/Contents/Anisotropy)
-			float AniG = saturate(Payload.Anisotropy);
-			float AniB = saturate(-1.0 * Payload.Anisotropy);
-			Result = float4(0.0, pow(AniG, 2.2), pow(AniB, 2.2), 1.0f);
-			break;
+		// 2. Material depends view mode
+	#if STRATA_ENABLED
+		FStrataAddressing StrataAddressing   = GetStrataPixelDataByteOffset(0, 0, 0);
+		FStrataPixelHeader StrataPixelHeader = UnpackStrataHeaderIn(Payload.StrataData, StrataAddressing, Payload.StrataData);
+
+		BRANCH
+		if (StrataPixelHeader.BSDFCount > 0)
+		{
+			const FStrataBSDF BSDF = UnpackStrataBSDFIn(Payload.StrataData, StrataAddressing, StrataPixelHeader); // Take the parameter of the first BSDF
+			const float3x3 TangentBasis = StrataGetBSDFSharedBasis(StrataPixelHeader, BSDF, StrataAddressing);
+			switch (VisualizationMode)
+			{				
+			case RAY_TRACING_DEBUG_VIZ_FAR_FIELD:		return float4(StrataGetBSDFBaseColor(BSDF), Payload.HitT);
+			case RAY_TRACING_DEBUG_VIZ_GBUFFER_AO:		return float4(StrataGetIrradianceAndAO(StrataPixelHeader).MaterialAO.xxx, 1.0f);
+			case RAY_TRACING_DEBUG_VIZ_SHADING_MODEL:	return float4(GetShadingModelColor(StrataGetLegacyShadingModels(BSDF)), 1.0f);
+			case RAY_TRACING_DEBUG_VIZ_BASE_COLOR:		return float4(StrataGetBSDFBaseColor(BSDF), 1.0f);
+			case RAY_TRACING_DEBUG_VIZ_DIFFUSE_COLOR:	return float4(StrataGetBSDFDiffuseColor(BSDF), 1.0f);
+			case RAY_TRACING_DEBUG_VIZ_SPECULAR_COLOR:	return float4(StrataGetBSDFSpecularColor(BSDF).xxx, 1.0f);
+			case RAY_TRACING_DEBUG_VIZ_METALLIC:		return float4(StrataGetBSDFMetallic(BSDF).xxx, 1.0f);
+			case RAY_TRACING_DEBUG_VIZ_SPECULAR:		return float4(StrataGetBSDFSpecular(BSDF).xxx, 1.0f);
+			case RAY_TRACING_DEBUG_VIZ_ROUGHNESS:		return RoughnessDebugColor(StrataGetBSDFRoughness(BSDF));
+			case RAY_TRACING_DEBUG_VIZ_IOR:				return float4(1, 1, 1, 1.0f);
+			case RAY_TRACING_DEBUG_VIZ_CUSTOM_DATA:		return float4(0, 0, 0, 0);
+			case RAY_TRACING_DEBUG_VIZ_WORLD_NORMAL:	return float4(TangentBasis[2] * 0.5 + 0.5, 1.0f);
+			case RAY_TRACING_DEBUG_VIZ_WORLD_TANGENT:	return float4(TangentBasis[0] * 0.5 + 0.5, 1.0f);
+			case RAY_TRACING_DEBUG_VIZ_ANISOTROPY:		return AnisoDebugColor(StrataGetBSDFAnisotropy(BSDF));
 			}
-	#endif // STRATA_ENABLED
 		}
+	#else // STRATA_ENABLED
+		switch (VisualizationMode)
+		{
+			case RAY_TRACING_DEBUG_VIZ_WORLD_NORMAL:	return float4(Payload.WorldNormal * 0.5 + 0.5, 1.0f);
+			case RAY_TRACING_DEBUG_VIZ_FAR_FIELD:		return float4(Payload.BaseColor.rgb, Payload.HitT);
+			case RAY_TRACING_DEBUG_VIZ_GBUFFER_AO:		return float4(Payload.GBufferAO, Payload.GBufferAO, Payload.GBufferAO, 1.0f);
+			case RAY_TRACING_DEBUG_VIZ_SHADING_MODEL:	return float4(GetShadingModelColor(Payload.ShadingModelID), 1.0f);
+			case RAY_TRACING_DEBUG_VIZ_BASE_COLOR:		return float4(Payload.BaseColor.rgb, 1.0f);
+			case RAY_TRACING_DEBUG_VIZ_DIFFUSE_COLOR:	return float4(Payload.DiffuseColor.rgb, 1.0f);
+			case RAY_TRACING_DEBUG_VIZ_SPECULAR_COLOR:	return float4(Payload.SpecularColor.rgb, 1.0f);
+			case RAY_TRACING_DEBUG_VIZ_METALLIC:		return float4(Payload.Metallic, Payload.Metallic, Payload.Metallic, 1.0f);
+			case RAY_TRACING_DEBUG_VIZ_SPECULAR:		return float4(Payload.Specular, Payload.Specular, Payload.Specular, 1.0f);
+			case RAY_TRACING_DEBUG_VIZ_ROUGHNESS:		return RoughnessDebugColor(Payload.Roughness);
+			case RAY_TRACING_DEBUG_VIZ_IOR:				return float4(Payload.Ior, Payload.Ior, Payload.Ior, 1.0f);
+			case RAY_TRACING_DEBUG_VIZ_CUSTOM_DATA:		return float4(Payload.CustomData);
+			case RAY_TRACING_DEBUG_VIZ_WORLD_TANGENT:	return float4(Payload.WorldTangent * 0.5 + 0.5, 1.0f);
+			case RAY_TRACING_DEBUG_VIZ_ANISOTROPY:		return AnisoDebugColor(Payload.Anisotropy);
+		}
+	#endif // STRATA_ENABLED
 	}
-	return Result;
+
+	// HitT is packed into the alpha for far field compositing
+	return (RAY_TRACING_DEBUG_VIZ_FAR_FIELD) ? float4(0, 0, 0, -1) : float4(0, 0, 0, 1);
 }
 
 struct FRayTracingDebugPayload : FMinimalPayload
