Strata: Avoid crashing when RT payload overflows

Fix a few small typos in comments

diff --git a/Engine/Shaders/Private/Strata/Strata.ush b/Engine/Shaders/Private/Strata/Strata.ush
--- a/Engine/Shaders/Private/Strata/Strata.ush
+++ b/Engine/Shaders/Private/Strata/Strata.ush
@@ -1621,13 +1621,17 @@ FStrataData GetStrataSingleLayerWaterBSDF(
 ///////////////////////////////////////////////////////////////////////////////
 // Raytracing material buffer
 
+#ifndef STRATA_RT_PAYLOAD_NUM_UINTS
+#define STRATA_RT_PAYLOAD_NUM_UINTS  9
+#endif
+
 struct FStrataRaytracingPayload
 {
-	// Total size: 36 bytes. 
+	// Total size: 40 bytes. (assuming the default value of STRATA_RT_PAYLOAD_NUM_UINTS)
 	// If this changes, update FPackedMaterialClosestHitPayload in RayTracingCommon.ush 
 	// and ensure the payload size is smaller than the maximum size
 	uint PackedTopLayerData;
-	uint Data[9];
+	uint Data[STRATA_RT_PAYLOAD_NUM_UINTS];
 };
 
 ///////////////////////////////////////////////////////////////////////////////
@@ -1730,8 +1734,12 @@ FRWStrataMaterialContainer InitialiseRWStrataMaterialContainer()
 
 void StrataStoreUint1(inout FRWStrataMaterialContainer StrataBuffer, inout FStrataAddressing StrataAddressing, uint Data)
 {
-	StrataBuffer.Data[StrataAddressing.CurrentIndex] = Data;
-	StrataAddressing.CurrentIndex++;
+	BRANCH
+	if (StrataAddressing.CurrentIndex < STRATA_RT_PAYLOAD_NUM_UINTS)
+	{
+		StrataBuffer.Data[StrataAddressing.CurrentIndex] = Data;
+		StrataAddressing.CurrentIndex++;
+	}
 }
 
 #define STRATA_STORE_UINT1(x) StrataStoreUint1(StrataBuffer, StrataAddressing, x)
@@ -1809,10 +1817,15 @@ uint StrataLoadTopLayerData(FStrataTopLayerDataContainer InTopLayerContainer, in
 
 uint StrataLoadUint1(FStrataMaterialContainer StrataBuffer, inout FStrataAddressing StrataAddressing)
 {
-	uint Data = StrataBuffer.Data[StrataAddressing.CurrentIndex];
-	StrataAddressing.CurrentIndex++;
-	StrataAddressing.ReadBytes += 4;
-	return Data;
+	BRANCH
+	if (StrataAddressing.CurrentIndex < STRATA_RT_PAYLOAD_NUM_UINTS)
+	{
+		uint Data = StrataBuffer.Data[StrataAddressing.CurrentIndex];
+		StrataAddressing.CurrentIndex++;
+		StrataAddressing.ReadBytes += 4;
+		return Data;
+	}
+	return 0;
 }
 
 uint StrataLoadUint1AtIndexOffset(in FStrataMaterialContainer StrataBuffer, in FStrataAddressing StrataAddressing, uint IndexOffset)
@@ -3215,7 +3228,7 @@ uint StrataPackNormalAndTangent(in float3 Normal, in float3 Tangent)
 
 	// Oct. encoding for the projected tangent, relative to an generated ortho frame based on normal
 	{
-		// Use the quantize normal to regenerate tje ortho. frame. This removes discontinuities
+		// Use the quantize normal to regenerate the ortho. frame. This removes discontinuities
 		float x = float(int(X10Bits) - 512) / 511.0f;
 		float y = float(int(Y10Bits) - 512) / 511.0f;
 		Normal = float3(x, y, 1.0f - abs(x) - abs(y));
@@ -3326,7 +3339,7 @@ struct FStrataDBuffer
 	float  Coverage;
 
 	// Per-component coverage 
-	// Since DBuffer properties are blended independely, they can end-up with varying coverage per-property. 
+	// Since DBuffer properties are blended independently, they can end-up with varying coverage per-property. 
 	// These values are only valid when *reading* DBuffer Value
 	float  OneMinusCoverage;
 	float  OneMinusCoverage_BaseColor;
diff --git a/Engine/Shaders/Private/Strata/StrataEvaluation.ush b/Engine/Shaders/Private/Strata/StrataEvaluation.ush
--- a/Engine/Shaders/Private/Strata/StrataEvaluation.ush
+++ b/Engine/Shaders/Private/Strata/StrataEvaluation.ush
@@ -757,7 +757,7 @@ FStrataEvaluateResult StrataEvaluateBSDFCommon(FStrataBSDFContext BSDFContext, F
 
 				const float Phase = IsotropicPhase();
 
-				// Now lerp between the optically thin and optically thin medium models.
+				// Now lerp between the optically thick and optically thin medium models.
 				// The diffuse and throughput account for the GGX interface SpecularTransmission.
 				Sample.DiffusePathValue		= lerp(Sample.DiffusePathValue,		SlabDirectionalAlbedo * SpecularTransmissionV,		DiffuseToVolumeBlend);
 				Sample.DiffusePDF			= lerp(Sample.DiffusePDF,			Phase,												DiffuseToVolumeBlend);
@@ -984,7 +984,7 @@ FStrataEvaluateResult StrataIntegrateBSDF(FStrataBSDFContext BSDFContext, FShado
 ///////////////////////////////////////////////////////////////////////////////////////////////////
 // BSDF Importance sampling
  
-// STRATA_TODO: merge this with Path-Tracing
+// STRATA_TODO: merge this with Path-Tracing (watch out -- implementations have diverged)
 
 // Given two lobes that will roughly contribute colors A and B to the total (estimated for example using directional albedo)
 // return the probability of choosing lobe A
@@ -1459,7 +1459,7 @@ FStrataEnvLightResult StrataEvaluateForEnvLight(FStrataBSDFContext BSDFContext,
 				StrataEnvLightResult.SSRSpecularWeight	*= Cloth_DirectionalAlbedo_SpecularTransmission;
 				
 				// We currently send the cloth environment lighting to be taken into account by the diffuse term.
-				// STARTA_TODO: true a mix of diffuse and specular contribution according to roughness.
+				// STRATA_TODO: true a mix of diffuse and specular contribution according to roughness.
 				StrataEnvLightResult.DiffuseColor		+= ClothEnergyTerms.E * FuzzAmount;
 				StrataEnvLightResult.DiffuseWeight		+= ClothEnergyTerms.E * FuzzAmount;
 				break;
diff --git a/Engine/Source/Runtime/Renderer/Private/DistortionRendering.cpp b/Engine/Source/Runtime/Renderer/Private/DistortionRendering.cpp
--- a/Engine/Source/Runtime/Renderer/Private/DistortionRendering.cpp
+++ b/Engine/Source/Runtime/Renderer/Private/DistortionRendering.cpp
@@ -351,7 +351,7 @@ static void AddDownsampleSceneColorPass(FRDGBuilder& GraphBuilder, FDownsampleSc
 
 //////////////////////////////////////////////////////////////////////////
 
-// STARTA_TODO make common with the function in PostProcessWeightedSampleSum.cpp
+// STRATA_TODO make common with the function in PostProcessWeightedSampleSum.cpp
 class FGaussianFiltering
 {
 public:
