Fix truncation cook warning

diff --git a/Engine/Shaders/Private/SceneTexturesCommon.ush b/Engine/Shaders/Private/SceneTexturesCommon.ush
--- a/Engine/Shaders/Private/SceneTexturesCommon.ush
+++ b/Engine/Shaders/Private/SceneTexturesCommon.ush
@@ -190,7 +190,7 @@ float LookupDeviceZ( float2 ScreenUV )
 	#else
 		// SceneDepth texture is not accessible during post-processing as we discard it at the end of mobile BasePass
 		// instead fetch DeviceZ from SceneDepthAuxTexture
-		return Texture2DSampleLevel(MobileSceneTextures.SceneDepthAuxTexture, MobileSceneTextures.SceneDepthAuxTextureSampler, ScreenUV, 0);
+		return Texture2DSampleLevel(MobileSceneTextures.SceneDepthAuxTexture, MobileSceneTextures.SceneDepthAuxTextureSampler, ScreenUV, 0).r;
 	#endif
 #elif COMPILER_GLSL_ES3_1 && PIXELSHADER
 	return DepthbufferFetchES2();
