Moved most of the ShaderPrint.ush code to ShaderPrintCommon.ush
Decoupled it from specific named buffer and uniform buffer bindings and kept a wrapper with these in ShaderPrint.ush.
The aim is to allow other systems (such as DeformerGraph) use ShaderPrint but with their own bindings.

diff --git a/Engine/Shaders/Private/ShaderPrint.ush b/Engine/Shaders/Private/ShaderPrint.ush
--- a/Engine/Shaders/Private/ShaderPrint.ush
+++ b/Engine/Shaders/Private/ShaderPrint.ush
@@ -1,1255 +1,194 @@
 /*=============================================================================
-	ShaderPrintCommon.ush:
-	Include this to be able to call ShaderPrint() from arbitrary shaders
+	ShaderPrint.ush:
+	Include this to be able to use ShaderPrint with the global uniforms and buffers.
 =============================================================================*/
 
 #pragma once
 
-// Include the MiniFont symbol definitions
+#include "ShaderPrintCommon.ush"
 #include "MiniFontCommon.ush"
 
-///////////////////////////////////////////////////////////////////////////////////////////////////
-// Description
-// Shader Print allows to print value & text from shaders. 
-// 
-// 0. Contex-based state API
-// -------------------------
-// FShaderPrintContext Context = InitShaderPrintContext(all(SvPosition.xy == float2(100, 100), uint2(50,50))
-//                             ------------------------------------   ------------
-//                                 Pixel/thread to filter             Coord at which 
-//                                                                       to draw
-// 
-// 1. How to print value?
-// ----------------------
-// First create a context, then print scalar/vector/matrix/text...
-// 
-//   FShaderPrintContext Context = InitShaderPrintContext(all(PixelPos.xy == uint2(100, 100), uint2(50,50))
-// 
-//   Print(Context, MyFloat);
-//   Newline(Context);
-// 
-//   Print(Context, TEXT("My string for labbelling a float :"), FontOrange);
-//   Print(Context, MyFloatInYellow, FontYellow);
-//   Newline(Context);
-// 
-// Numerical values can be further formatted, by configuring the number of digits spanned by a number 
-// (i.e. the number of written&non-written digits), and the number of wanted decimal for floating value.
-// 
-//   Print(Context, 123, FontYellow, 5); // Spanning is show as: .....
-//   Print(Context, 456, FontYellow, 4); // Spanning is show as: xxxx
-//   Print(Context, 789, FontYellow, 3); // Spanning is show as: ---
-//
-// Will be displayed as follow:
-//   123   456  789
-//   ..... xxxx ---
-// 
-// Print(Context, 12.3456, FontYellow, 5, 2); // Spanning is show as: ...
-// Print(Context, 12.3456, FontYellow, 5, 1); // Spanning is show as: xxx
-// 
-// Will be displayed as follow:
-//   12.34 12.3
-//   ..... xxxxx
-// 
-// 2. How to add widget?
-// ---------------------
-// Widget can be added from shader. 
-// For slider:
-// 
-//   const float MySliderValue  = AddSlider(Context, DefaultValue, TEXT("My Slider"));
-// 
-// For checkbox:
-// 
-//   const bool MyCheckboxValue = AddCheckbox(Context, DefaultValue, TEXT("My Checkbox"));
-//
-// The value of these widgets can be retrieved/evaluated from *any* thread/pixels, but the 
-// drawing need to be done from a single thread/pixel. This is done through the ShaderPrintContextcontext, 
-// which setup which pixel/thread is active/filtered. Example with a compute shader:
-// 
-//    [numthreads(8, 8, 1)]
-//    void MainCS(uint2 DispatchThreadId : SV_DispatchThreadID)
-//    {
-//       // Create context which 'add' the widget only on thread (0,0)
-//       FShaderPrintContext Context = InitShaderPrintContext(all(DispatchThreadId == uint2(0, 0), uint2(50,50))
-// 
-//		 // Any threads can get the slider value with this call
-//       const float ScaleValue = AddSlider(Context, 1.f, TEXT("Scale"));
-//       ...
-//    }
-// 
-
-///////////////////////////////////////////////////////////////////////////////////////////////////
-// Input types
-
-// Content of FShaderPrintItem.Type defines what to cast FShaderPrintItem.Value as
-
-// Type of primitives
-#define SHADER_PRINT_TYPE_SYMBOL   0
-#define SHADER_PRINT_TYPE_FLOAT    1
-#define SHADER_PRINT_TYPE_INT      2
-#define SHADER_PRINT_TYPE_UINT     3
-#define SHADER_PRINT_TYPE_HEX      4
-#define SHADER_PRINT_TYPE_SLIDER   5
-#define SHADER_PRINT_TYPE_CHECK    6
-// Type of primitives (to be merged with the list above once all buffer are merged)
-#define SHADER_PRINT_TYPE_LINE     0
-#define SHADER_PRINT_TYPE_TRIANGLE 1
-
-// Size of packed data, in uint
-#define SHADER_PRINT_UINT_COUNT_SYMBOL   4
-#define SHADER_PRINT_UINT_COUNT_LINE     8
-#define SHADER_PRINT_UINT_COUNT_TRIANGLE 12
-
-// Counter offset
-#define SHADER_PRINT_COUNTER_OFFSET_LINE     0
-#define SHADER_PRINT_COUNTER_OFFSET_TRIANGLE 1
-#define SHADER_PRINT_COUNTER_OFFSET_SYMBOL   2
-#define SHADER_PRINT_COUNTER_OFFSET_FREE     3
-#define SHADER_PRINT_COUNTER_COUNT           4
-
-uint GetSymbolOffset(uint InIndex)				{ return SHADER_PRINT_COUNTER_COUNT + SHADER_PRINT_UINT_COUNT_SYMBOL * InIndex; }
-uint GetPrimitiveLineOffset(uint InIndex)		{ return SHADER_PRINT_COUNTER_COUNT + SHADER_PRINT_UINT_COUNT_SYMBOL * ShaderPrintData.MaxValueCount + SHADER_PRINT_UINT_COUNT_LINE * InIndex; }
-uint GetPrimitiveTriangleOffset(uint InIndex)	{ return SHADER_PRINT_COUNTER_COUNT + SHADER_PRINT_UINT_COUNT_SYMBOL * ShaderPrintData.MaxValueCount + SHADER_PRINT_UINT_COUNT_LINE * ShaderPrintData.MaxLineCount + SHADER_PRINT_UINT_COUNT_TRIANGLE * InIndex; }
-
-void ClearCounters(RWStructuredBuffer<uint> InRWBuffer)
-{
-	// Symbol/print Counter
-	InRWBuffer[SHADER_PRINT_COUNTER_OFFSET_LINE] = 0;
-	InRWBuffer[SHADER_PRINT_COUNTER_OFFSET_TRIANGLE] = 0;
-	InRWBuffer[SHADER_PRINT_COUNTER_OFFSET_SYMBOL] = 0;
-	InRWBuffer[SHADER_PRINT_COUNTER_OFFSET_FREE] = 0;
-}
-
-///////////////////////////////////////////////////////////////////////////////////////////////////
-// State 
-#define SHADER_PRINT_STATE_COUNT_OFFSET 0
-#define SHADER_PRINT_STATE_VALUE_OFFSET 1
-
-#define SHADER_PRINT_STATE_STRIDE 3
-#define SHADER_PRINT_STATE_INDEX_METADATA 0
-#define SHADER_PRINT_STATE_INDEX_VALUE 1
-#define SHADER_PRINT_STATE_INDEX_FRAME 2
-
-#define SHADER_PRINT_STATE_INVALID_INDEX 0xFF
-#define SHADER_PRINT_STATE_HASH_MASK 0xFFFFFF
-
-///////////////////////////////////////////////////////////////////////////////////////////////////
-// Output structures
+// The FXC compiler will force unroll on loops around Print() and so in some places we need
+// to limit the loop count to either fix compilation or reduce compilation time.
+#if COMPILER_HLSL && !COMPILER_DXC
+#define SHADER_PRINT_UNROLL(N) UNROLL_N(N)
+#else
+#define SHADER_PRINT_UNROLL(N) LOOP
+#endif
 
+// Global buffers for ShaderPrint collection.
 StructuredBuffer<uint> ShaderPrint_StateBuffer;
-
-// Primitive & Symbols layout
-//  ----------------------------------------------------------------------------------------------------------------------------------------------------------------------------------------------
-// ||  Line    | Triangle |  Symbol |  Empty  ||  Symb0  |  Symb1  |  Symb2  |  Symb3  |  ...  |  Line0  |  Line1  |  Line2  |  Line3  |  ...  | Triangle0 | Triangle1 | Triangle2 | Triangle3 | ...
-// || Counter  | Counter  | Counter | Counter ||         |         |         |         |  ...  |         |         |         |         |  ...  |           |           |           |           | ...
-// || (1uint)  | (1uint)  | (1uint) | (1uint) || (4uint) | (4uint) | (4uint) | (4uint) |  ...  | (8uint) | (8uint) | (8uint) | (8uint) |  ...  | (12uint)  | (12uint)  | (12uint)  | (12uint)  | ...
-//  ----------------------------------------------------------------------------------------------------------------------------------------------------------------------------------------------
 RWStructuredBuffer<uint> ShaderPrint_RWEntryBuffer;
 
-///////////////////////////////////////////////////////////////////////////////////////////////////
-// Font Color
-
-struct FFontColor
+FShaderPrintBuffers InitShaderPrintContextBuffers()
 {
-	float3 Color;
-};
-
-FFontColor InitFontColor(float InX, float InY, float InZ)	{ FFontColor Out; Out.Color = float3(InX, InY, InZ); return Out; }
-FFontColor InitFontColor(float3 In)							{ FFontColor Out; Out.Color = In; return Out; }
-FFontColor GetDefaultFontColor()							{ FFontColor Out; Out.Color = float3(1,1,1); return Out; }
-
-// Predefined colors
-#define FontWhite		InitFontColor(1, 1, 1)
-#define FontGrey		InitFontColor(0.5, 0.5, 0.5)
-#define FontDarkGrey	InitFontColor(0.25, 0.25, 0.25)
-#define FontBlack		InitFontColor(0, 0, 0)
-#define FontRed			InitFontColor(1, 0, 0)
-#define FontGreen		InitFontColor(0, 1, 0)
-#define FontBlue		InitFontColor(0, 0, 1)
-#define FontYellow		InitFontColor(1, 1, 0)
-#define FontCyan		InitFontColor(0, 1, 1)
-#define FontMagenta		InitFontColor(1, 0, 1)
-#define FontOrange		InitFontColor(243.f / 255.f, 156.f / 255.f, 18.f / 255.f)
-#define FontPurple		InitFontColor(169.f / 255.f, 7.f / 255.f, 228.f / 255.f)
-#define FontTurquoise	InitFontColor(26.f / 255.f, 188.f / 255.f, 156.f / 255.f)
-#define FontSilver		InitFontColor(189.f / 255.f, 195.f / 255.f, 199.f / 255.f)
-#define FontEmerald		InitFontColor(46.f / 255.f, 204.f / 255.f, 113.f / 255.f)
-#define FontLightRed	InitFontColor(0.75f, 0.50f, 0.50f)
-#define FontLightGreen	InitFontColor(0.50f, 0.75f, 0.50f)
-#define FontLightBlue	InitFontColor(0.50f, 0.50f, 0.75f)
-#define FontDarkRed		InitFontColor(0.50f, 0.25f, 0.25f)
-#define FontDarkGreen	InitFontColor(0.25f, 0.50f, 0.25f)
-#define FontDarkBlue	InitFontColor(0.25f, 0.25f, 0.50f)
-
-///////////////////////////////////////////////////////////////////////////////////////////////////
-// Primitive colors
-
-#define ColorWhite		float4(1, 1, 1, 1)
-#define ColorBlack		float4(0, 0, 0, 1)
-#define ColorRed		float4(1, 0, 0, 1)
-#define ColorGreen		float4(0, 1, 0, 1)
-#define ColorBlue		float4(0, 0, 1, 1)
-#define ColorYellow		float4(1, 1, 0, 1)
-#define ColorCyan		float4(0, 1, 1, 1)
-#define ColorMagenta	float4(1, 0, 1, 1)
-#define ColorOrange		float4(243.f / 255.f, 156.f / 255.f,  18.f / 255.f, 1)
-#define ColorPurple		float4(169.f / 255.f,   7.f / 255.f, 228.f / 255.f, 1)
-#define ColorTurquoise	float4( 26.f / 255.f, 188.f / 255.f, 156.f / 255.f, 1)
-#define ColorSilver		float4(189.f / 255.f, 195.f / 255.f, 199.f / 255.f, 1)
-#define ColorEmerald	float4( 46.f / 255.f, 204.f / 255.f, 113.f / 255.f, 1)
-#define ColorLightRed	float4(0.75f, 0.50f, 0.50f, 1.0f)
-#define ColorLightGreen	float4(0.50f, 0.75f, 0.50f, 1.0f)
-#define ColorLightBlue	float4(0.50f, 0.50f, 0.75f, 1.0f)
-#define ColorDarkRed	float4(0.50f, 0.25f, 0.25f, 1.0f)
-#define ColorDarkGreen	float4(0.25f, 0.50f, 0.25f, 1.0f)
-#define ColorDarkBlue	float4(0.25f, 0.25f, 0.50f, 1.0f)
-
-///////////////////////////////////////////////////////////////////////////////////////////////////
-// Util pack/unpac functions
-
-struct FShaderPrintItem
-{
-	float2 ScreenPos; // Position in normalized coordinates
-	int Value;        // Cast to value or symbol
-	int Type;         // SHADER_PRINT_TYPE_* defines how to read Value
-	float3 Color;	  // Color 
-	uint Metadata;	  // Metadata (used for widget to store extra state data
-};
-
-// Needs to match C++ code in ShaderPrint.cpp
-struct FPackedShaderPrintItem
-{
-	uint ScreenPos16bits; // Position in normalized coordinates
-	uint Value;           // Cast to value or symbol
-	uint TypeAndColor;    //
-	uint Metadata;	      //
-};
-
-FPackedShaderPrintItem PackShaderPrintItem(FShaderPrintItem In)
-{
-	const uint3 Color8bits = saturate(In.Color) * 0xFF;
-	const uint2 ScreenPos16bit = f32tof16(In.ScreenPos);
-
-	FPackedShaderPrintItem Out;
-	Out.ScreenPos16bits = ScreenPos16bit.x | (ScreenPos16bit.y<<16);
-	Out.Value = asuint(In.Value);
-	Out.TypeAndColor = (Color8bits.z << 24) | (Color8bits.y << 16) | (Color8bits.x << 8) | (In.Type & 0xFF);
-	Out.Metadata = In.Metadata;
-	return Out;
+	FShaderPrintBuffers Buffers;
+	Buffers.StateBuffer = ShaderPrint_StateBuffer;
+	Buffers.RWEntryBuffer = ShaderPrint_RWEntryBuffer;
+	return Buffers;
 }
 
-FShaderPrintItem UnpackShaderPrintItem(FPackedShaderPrintItem In)
+// Initialise FShaderPrintConfig from the ShaderPrintData uniform buffer.
+FShaderPrintConfig InitShaderPrintContextConfig()
 {
-	const uint2 ScreenPos16bits = uint2(In.ScreenPos16bits & 0xFFFF, (In.ScreenPos16bits>>16) & 0xFFFF);
-
-	FShaderPrintItem Out;
-	Out.ScreenPos = f16tof32(ScreenPos16bits);
-	Out.Value	  = asint(In.Value);
-	Out.Type	  = (In.TypeAndColor) & 0xFF;
-	Out.Color.x   = float((In.TypeAndColor >>  8) & 0xFF) / 255.f;
-	Out.Color.y   = float((In.TypeAndColor >> 16) & 0xFF) / 255.f;
-	Out.Color.z   = float((In.TypeAndColor >> 24) & 0xFF) / 255.f;
-	Out.Metadata  = In.Metadata;
-	return Out;
+	FShaderPrintConfig Config;
+	Config.Resolution = ShaderPrintData.Resolution;
+	Config.CursorCoord = ShaderPrintData.CursorCoord;
+	Config.TranslatedWorldOffset = ShaderPrintData.TranslatedWorldOffset;
+	Config.FontSize = ShaderPrintData.FontSize;
+	Config.FontSpacing = ShaderPrintData.FontSpacing;
+	Config.MaxValueCount = ShaderPrintData.MaxValueCount;
+	Config.MaxSymbolCount = ShaderPrintData.MaxSymbolCount;
+	Config.MaxStateCount = ShaderPrintData.MaxStateCount;
+	Config.MaxLineCount = ShaderPrintData.MaxLineCount;
+	Config.MaxTriangleCount = ShaderPrintData.MaxTriangleCount;
+	return Config;
 }
 
-void WriteSymbol(uint Offset, FShaderPrintItem In, RWStructuredBuffer<uint> InRWBuffer)
-{
-	const uint Offset4 = GetSymbolOffset(Offset);
-
-	const FPackedShaderPrintItem Packed = PackShaderPrintItem(In);
-	InRWBuffer[Offset4 + 0] = Packed.ScreenPos16bits;
-	InRWBuffer[Offset4 + 1] = Packed.Value;
-	InRWBuffer[Offset4 + 2] = Packed.TypeAndColor;
-	InRWBuffer[Offset4 + 3] = Packed.Metadata;
-}
-
-FShaderPrintItem ReadSymbol(uint Offset, StructuredBuffer<uint> InBuffer)
-{
-	const uint Offset4 = GetSymbolOffset(Offset);
-
-	FPackedShaderPrintItem Packed = (FPackedShaderPrintItem)0;
-	Packed.ScreenPos16bits	= InBuffer[Offset4 + 0];
-	Packed.Value			= InBuffer[Offset4 + 1];
-	Packed.TypeAndColor		= InBuffer[Offset4 + 2];
-	Packed.Metadata			= InBuffer[Offset4 + 3];
-
-	return UnpackShaderPrintItem(Packed);
-}
-
-///////////////////////////////////////////////////////////////////////////////////////////////////
-// Global variables
-
-#include "/Engine/Generated/UniformBuffers/ShaderPrintData.ush"
-
-///////////////////////////////////////////////////////////////////////////////////////////////////
-// Print API
-
-struct FShaderPrintContext
-{
-	bool bIsActive;
-	float2 StartPos;
-	float2 Pos;
-};
-
+// Initialise FShaderPrintContext.
 FShaderPrintContext InitShaderPrintContext(bool bActive, float2 InStartPos)
 {
-	FShaderPrintContext Out;
-	Out.bIsActive = bActive;
-	Out.StartPos = InStartPos;
-	Out.Pos = InStartPos;
-	return Out;
-}
-FShaderPrintContext InitShaderPrintContext(bool bActive, uint2 InStartCoord) { return InitShaderPrintContext(bActive, float2(InStartCoord) / float2(ShaderPrintData.Resolution)); }
-
-uint2 ShaderPrintGetResolution()
-{
-	return ShaderPrintData.Resolution;
-}
-
-void Newline(inout FShaderPrintContext InContext)
-{
-	InContext.Pos.x = InContext.StartPos.x;
-	InContext.Pos.y += ShaderPrintData.FontSpacing.y;
-}
-
-///////////////////////////////////////////////////////////////////////////////////////////////////
-// Internal helpers
-
-void ShaderPrint_Internal(in FShaderPrintItem Item)
-{
-	// If MaxValueCount is 0 then we don't reset the buffer counter so early out here
-	if (ShaderPrintData.MaxValueCount == 0)
-	{
-		return;
-	}
-
-	// Buffer counter is stored in first element .Value
-	int IndexToStore = 0;
-	InterlockedAdd(ShaderPrint_RWEntryBuffer[SHADER_PRINT_COUNTER_OFFSET_SYMBOL], 1, IndexToStore);
-
-	// Prevent writing off the buffer
-	// Note that counter still increases so need clamp when reading it in later passes
-	if (uint(IndexToStore) >= ShaderPrintData.MaxValueCount)
-	{
-		return;
-	}
-
-	WriteSymbol(IndexToStore, Item, ShaderPrint_RWEntryBuffer);
-}
-
-void ShaderPrint_Internal(in float2 ScreenPos, in int Value, in FFontColor FontColor, in uint Metadata, in int Type)
-{
-	FShaderPrintItem Item;
-	Item.ScreenPos = ScreenPos;
-	Item.Value = Value;
-	Item.Type = Type;
-	Item.Color = FontColor.Color;
-	Item.Metadata = Metadata;
-	ShaderPrint_Internal(Item);
+	FShaderPrintContext Context;
+	Context.bIsActive = bActive;
+	Context.StartPos = InStartPos;
+	Context.Pos = InStartPos;
+	Context.Config = InitShaderPrintContextConfig();
+	Context.Buffers = InitShaderPrintContextBuffers();
+	return Context;
 }
 
-///////////////////////////////////////////////////////////////////////////////////////////////////
-// Symbol printing
-
-void PrintSymbol(inout FShaderPrintContext Ctx, in int Symbol, in FFontColor Color)
+FShaderPrintContext InitShaderPrintContext(bool bActive, uint2 InStartCoord) 
 {
-	if (Ctx.bIsActive)
-	{
-		ShaderPrint_Internal(Ctx.Pos, Symbol, Color, 0u, SHADER_PRINT_TYPE_SYMBOL);
-		Ctx.Pos.x += ShaderPrintData.FontSpacing.x;
-	}
+	return InitShaderPrintContext(true, float2(InStartCoord) / float2(ShaderPrintData.Resolution));
 }
-void PrintSymbol(inout FShaderPrintContext Ctx, in int Symbol) { PrintSymbol(Ctx, Symbol, GetDefaultFontColor()); }
-
-///////////////////////////////////////////////////////////////////////////////////////////////////
-// Text & hash
-
-struct FShaderPrintText
-{
-	uint Index;
-};
 
-FShaderPrintText InitShaderPrintText(uint InIndex)
+FShaderPrintContext InitShaderPrintContext()
 {
-	FShaderPrintText Out;
-	Out.Index = InIndex;
-	return Out;
+	return InitShaderPrintContext(true, float2(0, 0));
 }
 
-// Forward declarations for text accessors
-uint ShaderPrintGetChar(uint InIndex);
-uint ShaderPrintGetOffset(FShaderPrintText InTextEntry);
-uint ShaderPrintGetHash(FShaderPrintText InTextEntry);
- 
-// Function for reading global TEXT string
-float2 ShaderPrintText_Internal(bool bIsActive, float2 InPos, FShaderPrintText InTextEntry, FFontColor InColor)
-{
-	// If MaxValueCount is 0 then we don't reset the buffer counter so early out here
-	if (bIsActive && ShaderPrintData.MaxValueCount > 0)
-	{
-		const uint Begin = ShaderPrintGetOffset(InTextEntry);
-		const uint End   = ShaderPrintGetOffset(InitShaderPrintText(InTextEntry.Index + 1));
-		const uint Count = End - Begin;
-
-		// Buffer counter is stored in first element .Value
-		// Prevent writing off the buffer
-		// Note that counter still increases so need clamp when reading it in later passes
-		int IndexToStore = 0;
-		InterlockedAdd(ShaderPrint_RWEntryBuffer[SHADER_PRINT_COUNTER_OFFSET_SYMBOL], Count, IndexToStore);
-		if (uint(IndexToStore + Count) < ShaderPrintData.MaxValueCount)
-		{
-			for (uint i = Begin; i < End; ++i)
-			{
-				FShaderPrintItem Item;
-				Item.ScreenPos = InPos;
-				Item.Value = ShaderPrintGetChar(i);
-				Item.Type = SHADER_PRINT_TYPE_SYMBOL;
-				Item.Color = InColor.Color;
-				Item.Metadata = 0u;
-
-				WriteSymbol(IndexToStore, Item, ShaderPrint_RWEntryBuffer);
-
-				++IndexToStore;
-				InPos.x += ShaderPrintData.FontSpacing.x;
-			}	
-		}
-	}
-	return InPos;
-}
-
-void Print(inout FShaderPrintContext InCtx, FShaderPrintText InTextEntry, FFontColor InColor)	{ InCtx.Pos = ShaderPrintText_Internal(InCtx.bIsActive, InCtx.Pos, InTextEntry, InColor); }
-void Print(inout FShaderPrintContext InCtx, FShaderPrintText InTextEntry)						{ InCtx.Pos = ShaderPrintText_Internal(InCtx.bIsActive, InCtx.Pos, InTextEntry, GetDefaultFontColor()); }
-
-// Function for reading global TEXT string. 
-// The data and function are generated by the shader compiler.
-GENERATED_SHADER_PRINT
-
-///////////////////////////////////////////////////////////////////////////////////////////////////
-// Value printing (common value printing)
-#define MAX_DECIMAL_COUNT 5u
-#define MAX_DIGIT_COUNT 12u
-
-float2 ShaderPrintValue_Internal(float2 ScreenPos, float Value, FFontColor Color, uint MaxDigit, uint MaxDecimal)  { ShaderPrint_Internal(ScreenPos, asint(Value),           Color, MaxDecimal, SHADER_PRINT_TYPE_FLOAT); ScreenPos.x += ShaderPrintData.FontSpacing.x * MaxDigit; return ScreenPos; }
-float2 ShaderPrintValue_Internal(float2 ScreenPos, int   Value, FFontColor Color, uint MaxDigit, uint MaxDecimal)  { ShaderPrint_Internal(ScreenPos, Value,                  Color, MaxDecimal, SHADER_PRINT_TYPE_INT);   ScreenPos.x += ShaderPrintData.FontSpacing.x * MaxDigit; return ScreenPos; }
-float2 ShaderPrintValue_Internal(float2 ScreenPos, uint  Value, FFontColor Color, uint MaxDigit, uint MaxDecimal)  { ShaderPrint_Internal(ScreenPos, asint(Value),           Color, MaxDecimal, SHADER_PRINT_TYPE_UINT);  ScreenPos.x += ShaderPrintData.FontSpacing.x * MaxDigit; return ScreenPos; }
-float2 ShaderPrintValue_Internal(float2 ScreenPos, bool  Value, FFontColor Color, uint MaxDigit, uint MaxDecimal)  { ShaderPrint_Internal(ScreenPos, asint(Value ? 1u : 0u), Color, MaxDecimal, SHADER_PRINT_TYPE_UINT);  ScreenPos.x += ShaderPrintData.FontSpacing.x * MaxDigit; return ScreenPos; }
-
-///////////////////////////////////////////////////////////////////////////////////////////////////
-// Scalar type printing
-
-#define SHADER_PRINT_OVERLOAD_1(InNumComponent, InType) \
-void   Print(inout FShaderPrintContext InCtx, InType Value, FFontColor Color, uint MaxDigit=MAX_DIGIT_COUNT, uint MaxDecimal=MAX_DECIMAL_COUNT)	{ if (InCtx.bIsActive) { InCtx.Pos = ShaderPrintValue_Internal(InCtx.Pos, Value, Color, MaxDigit, MaxDecimal); } } \
-void   Print(inout FShaderPrintContext InCtx, InType Value)																						{ if (InCtx.bIsActive) { InCtx.Pos = ShaderPrintValue_Internal(InCtx.Pos, Value, GetDefaultFontColor(), MAX_DIGIT_COUNT, MAX_DECIMAL_COUNT); } }
-
-///////////////////////////////////////////////////////////////////////////////////////////////////
-// Vector type printing
-
-#define SHADER_PRINT_OVERLOAD_N(InNumComponent, InType) \
-void   Print(inout FShaderPrintContext InCtx, InType Value, FFontColor Color, uint MaxDigit=MAX_DIGIT_COUNT, uint MaxDecimal=MAX_DECIMAL_COUNT)	{ if (InCtx.bIsActive) { UNROLL for (uint CompIt=0;CompIt<InNumComponent;++CompIt) { InCtx.Pos = ShaderPrintValue_Internal(InCtx.Pos, Value[CompIt], Color, MaxDigit, MaxDecimal); } } } \
-void   Print(inout FShaderPrintContext InCtx, InType Value)																						{ if (InCtx.bIsActive) { UNROLL for (uint CompIt=0;CompIt<InNumComponent;++CompIt) { InCtx.Pos = ShaderPrintValue_Internal(InCtx.Pos, Value[CompIt], GetDefaultFontColor(), MAX_DIGIT_COUNT, MAX_DECIMAL_COUNT); } } }
-
-// Print(X) version
-SHADER_PRINT_OVERLOAD_1(1, float)
-SHADER_PRINT_OVERLOAD_N(2, float2)
-SHADER_PRINT_OVERLOAD_N(3, float3)
-SHADER_PRINT_OVERLOAD_N(4, float4)
-
-SHADER_PRINT_OVERLOAD_1(1, uint)
-SHADER_PRINT_OVERLOAD_N(2, uint2)
-SHADER_PRINT_OVERLOAD_N(3, uint3)
-SHADER_PRINT_OVERLOAD_N(4, uint4)
-
-SHADER_PRINT_OVERLOAD_1(1, int)
-SHADER_PRINT_OVERLOAD_N(2, int2)
-SHADER_PRINT_OVERLOAD_N(3, int3)
-SHADER_PRINT_OVERLOAD_N(4, int4)
-
-SHADER_PRINT_OVERLOAD_1(1, bool)
-
-///////////////////////////////////////////////////////////////////////////////////////////////////
-// Matrix type printing
-
-#define SHADER_PRINT_OVERLOAD_NN(InNumComponentX, InNumComponentY, InType) \
-void   Print(inout FShaderPrintContext InCtx, InType Value, FFontColor Color, uint MaxDigit=MAX_DIGIT_COUNT, uint MaxDecimal=MAX_DECIMAL_COUNT)	{ if (InCtx.bIsActive) { UNROLL for (uint CompY=0;CompY<InNumComponentY;++CompY) { Print(InCtx, Value[CompY], Color, MaxDigit, MaxDecimal); InCtx.Pos += float2(0.f, ShaderPrintData.FontSpacing.y); } } } \
-void   Print(inout FShaderPrintContext InCtx, InType Value)																						{ if (InCtx.bIsActive) { UNROLL for (uint CompY=0;CompY<InNumComponentY;++CompY) { Print(InCtx, Value[CompY], GetDefaultFontColor(), MAX_DIGIT_COUNT, MAX_DECIMAL_COUNT); InCtx.Pos += float2(0.f, ShaderPrintData.FontSpacing.y); } } }
-
-// Print(X) version
-SHADER_PRINT_OVERLOAD_NN(3, 3, float3x3)
-SHADER_PRINT_OVERLOAD_NN(4, 3, float4x3)
-SHADER_PRINT_OVERLOAD_NN(4, 4, float4x4)
-
-///////////////////////////////////////////////////////////////////////////////////////////////////
-// Widgets (experimental)
-
-uint2 GetCheckboxWidgetSize()
-{
-	return uint(ShaderPrintData.FontSpacing.x * ShaderPrintData.Resolution.x).xx;
-}
-
-uint2 GetSliderWidgetSize()
-{
-	return uint2(10,1) * uint(ShaderPrintData.FontSpacing.x * ShaderPrintData.Resolution.x);
-}
-
-struct FShaderPrintMetadata
-{
-	uint Hash;
-	uint Index;
-};
-
-uint ShaderPrintPackMetadata(uint InHash, uint InIndex)
+// Backwards compatibility for client code that didn't use FShaderPrintContext.
+int2 GetCursorPos()
 {
-	return ((InIndex & 0xFFu) << 24) | (InHash & SHADER_PRINT_STATE_HASH_MASK);
+	return GetCursorPos(InitShaderPrintContext());
 }
-uint ShaderPrintPackMetadata(FShaderPrintMetadata In) { return ShaderPrintPackMetadata(In.Hash, In.Index); }
-
-FShaderPrintMetadata ShaderPrintUnpackMetadata(uint Metadata)
+void AddFilledTriangleWS(float3 Pos0, float3 Pos1, float3 Pos2, float4 Color)
 {
-	FShaderPrintMetadata Out;
-	Out.Hash = (Metadata & SHADER_PRINT_STATE_HASH_MASK);
-	Out.Index= (Metadata >> 24) & 0xFF;
-	return Out;
+	AddFilledTriangleTWS(InitShaderPrintContext(), Pos0, Pos1, Pos2, Color); 
 }
-
-// Return true if found, false otherwise
-bool ShaderPrintGetStateValue(in uint InHash, inout uint OutIndex, inout uint OutValue)
-{
-	OutIndex = SHADER_PRINT_STATE_INVALID_INDEX;
-	OutValue = 0;
-
-	// Slow linear search among all widget states
-	const uint MaxCount = min(ShaderPrint_StateBuffer[SHADER_PRINT_STATE_COUNT_OFFSET], ShaderPrintData.MaxStateCount);
-	for (uint Index = 0; Index < MaxCount; ++Index)
-	{
-		const uint Index3 = Index * SHADER_PRINT_STATE_STRIDE + SHADER_PRINT_STATE_VALUE_OFFSET;
-		const FShaderPrintMetadata Metadata = ShaderPrintUnpackMetadata(ShaderPrint_StateBuffer[Index3 + SHADER_PRINT_STATE_INDEX_METADATA]);
-		if (Metadata.Hash == InHash)
-		{
-			OutIndex = Metadata.Index;
-			OutValue = ShaderPrint_StateBuffer[Index3 + SHADER_PRINT_STATE_INDEX_VALUE];
-			return true;
-		}
-	}
-	return false;
+void AddFilledQuadWS(float3 Pos0, float3 Pos1, float3 Pos2, float3 Pos3, float4 Color) 
+{ 
+	AddFilledQuadWS(InitShaderPrintContext(), Pos0, Pos1, Pos2, Pos3, Color);
 }
-
-// Add a checkbox widget 
-bool AddCheckbox(inout FShaderPrintContext InContext, FShaderPrintText InTextEntry, bool bDefault, in FFontColor InColor)
+void AddFilledQuadSS(float2 Pos0, float2 Pos1, float4 Color)
 {
-	uint Index = 0;
-	uint RawValue = 0;
-	const uint Hash = ShaderPrintGetHash(InTextEntry) & SHADER_PRINT_STATE_HASH_MASK; // Trunk hask since metadata has limited storage
-	const bool bIsValid = ShaderPrintGetStateValue(Hash, Index, RawValue);
-	RawValue = bIsValid ? RawValue : (bDefault ? 0x1 : 0x0);
-
-	if (InContext.bIsActive)
-	{
-		const uint2 Coord = InContext.Pos * ShaderPrintData.Resolution;
-		const uint2 Cusor = ShaderPrintData.CursorCoord;
-		const uint2 WidgetSize = GetCheckboxWidgetSize();
-		const uint2 WidgetMax = uint2(Coord.x + WidgetSize.x, Coord.y + WidgetSize.y / 2);
-		const uint2 WidgetMin = uint2(Coord.x, Coord.y - WidgetSize.y / 2);
-		const bool  bIsInside = all(Cusor >= WidgetMin) && all(Cusor <= WidgetMax);
-
-		const bool bWasInside = RawValue & 0x2;
-
-		// Update checkbox value 
-		// Track if the cursor was already within the box or not
-		if (bIsInside)
-		{
-			if (!bWasInside)
-			{
-				// First time the cursor is within the checkbox
-				RawValue = ((RawValue & 0x1) == 1 ? 0x0 : 0x1) | 0x2;
-			}
-		}
-		else
-		{
-			RawValue = (RawValue & 0x1);
-		}
-
-		FShaderPrintItem E;
-		E.ScreenPos = InContext.Pos;
-		E.Value = RawValue;
-		E.Type = SHADER_PRINT_TYPE_CHECK;
-		E.Color = bIsInside ? FontYellow.Color : FontWhite.Color;// InColor.Color;
-		E.Metadata = ShaderPrintPackMetadata(Hash, Index);
-		ShaderPrint_Internal(E);
-
-		InContext.Pos.x += float(WidgetSize.x) / float(ShaderPrintData.Resolution.x) + ShaderPrintData.FontSpacing.x;
-
-		// Text
-		Print(InContext, InTextEntry, InColor);
-		InContext.Pos.x += ShaderPrintData.FontSpacing.x;
-	}
-
-	return RawValue & 0x1;
+	AddFilledQuadSS(InitShaderPrintContext(), Pos0, Pos1, Color);
 }
-
-// Add a slider widget  
-float AddSlider(inout FShaderPrintContext InContext, FShaderPrintText InTextEntry, float bDefault, in FFontColor InColor, float InMin = 0.f, float InMax = 1.f)
+void AddLineTWS(float3 Pos0, float3 Pos1, float4 Color)
 {
-	// The value is stored as normalized, for displaying slide value correctly. This is not idealy
-	// as it implies some precision lost due to back&forth conversion between normalized/non-normalized value.
-	uint Index = 0;
-	uint RawValue = 0;
-	const uint Hash = ShaderPrintGetHash(InTextEntry) & SHADER_PRINT_STATE_HASH_MASK; // Trunk hask since metadata has limited storage
-	const bool bIsValid = ShaderPrintGetStateValue(Hash, Index, RawValue);
-	float NormalizedValue = bIsValid ? asfloat(RawValue) : bDefault; 
-	float Value = NormalizedValue * (InMax - InMin) + InMin;
-
-	if (InContext.bIsActive)
-	{
-		const uint2 Coord = InContext.Pos * ShaderPrintData.Resolution;
-		const uint2 Cusor = ShaderPrintData.CursorCoord;
-		const uint2 WidgetSize = GetSliderWidgetSize();
-		const uint2 WidgetMax = uint2(Coord.x + WidgetSize.x, Coord.y + WidgetSize.y / 2);
-		const uint2 WidgetMin = uint2(Coord.x, Coord.y - WidgetSize.y / 2);
-		const bool  bIsInside = all(Cusor >= WidgetMin) && all(Cusor <= WidgetMax);
-
-		// Update slider value
-		if (bIsInside)
-		{
-			const float S = saturate(float(Cusor.x - WidgetMin.x) / float(WidgetMax.x - WidgetMin.x));
-			Value = lerp(InMin, InMax, S);
-		}
-
-		NormalizedValue = saturate((Value - InMin) / (InMax - InMin));
-
-		FShaderPrintItem E;
-		E.ScreenPos = InContext.Pos;
-		E.Value = asint(NormalizedValue);
-		E.Type = SHADER_PRINT_TYPE_SLIDER;
-		E.Color = bIsInside ? FontYellow.Color : FontWhite.Color; // InColor.Color;
-		E.Metadata = ShaderPrintPackMetadata(Hash, Index);
-		ShaderPrint_Internal(E);
-
-		InContext.Pos.x += float(WidgetSize.x) / float(ShaderPrintData.Resolution.x) + ShaderPrintData.FontSpacing.x;
-
-		// Text
-		Print(InContext, InTextEntry, InColor);
-		InContext.Pos.x += ShaderPrintData.FontSpacing.x;
-	}
-
-	return Value;
+	AddLineTWS(InitShaderPrintContext(), Pos0, Pos1, Color);
 }
-
-///////////////////////////////////////////////////////////////////////////////////////////////////
-// Primitives types
-
-///////////////////////////////////////////////////////////////////////////////////////////////////
-// Triangle-based Primitives (Triangle, Point, ...)
-
-struct FTriangleElement 
-{
-	float3 Pos0;
-	float3 Pos1;
-	float3 Pos2;
-	float4 Color;
-	bool bIsScreenSpace;
-};
-
-struct FPackedTriangleElement
-{
-	uint4 Packed0;
-	uint4 Packed1;
-	uint4 Packed2;
-};
-
-FTriangleElement UnpackTriangleElement(FPackedTriangleElement In)
+void AddLineTWS(float3 Pos0, float3 Pos1, float4 Color0, float4 Color1)
 {
-	FTriangleElement Out = (FTriangleElement)0;
-	{
-		Out.Pos0.x = asfloat(In.Packed0.x);
-		Out.Pos0.y = asfloat(In.Packed0.y);
-		Out.Pos0.z = asfloat(In.Packed0.z);
-		Out.Color  = float4((In.Packed0.w >> 24) & 0xFF, (In.Packed0.w >> 16) & 0xFF, (In.Packed0.w >> 8) & 0xFF, (In.Packed0.w) & 0xFF) / 255.0f;
-	}
-	{
-		Out.Pos1.x = asfloat(In.Packed1.x);
-		Out.Pos1.y = asfloat(In.Packed1.y);
-		Out.Pos1.z = asfloat(In.Packed1.z);
-
-		// The 'space' info is stored into the alpha's LSB
-		Out.bIsScreenSpace = (In.Packed1.w & 0x1) > 0u;
-	}
-	{
-		Out.Pos2.x = asfloat(In.Packed2.x);
-		Out.Pos2.y = asfloat(In.Packed2.y);
-		Out.Pos2.z = asfloat(In.Packed2.z);
-	}
-	return Out;
+	AddLineTWS(InitShaderPrintContext(), Pos0, Pos1, Color0, Color1);
 }
-
-FPackedTriangleElement PackTriangleElement(FTriangleElement In)
+void AddLineWS(float3 Pos0, float3 Pos1, float4 Color)
 {
-	uint4 PackedC = uint4(255.0f * saturate(In.Color));
-
-	FPackedTriangleElement Out = (FPackedTriangleElement)0;
-	Out.Packed0.x = asuint(In.Pos0.x);
-	Out.Packed0.y = asuint(In.Pos0.y);
-	Out.Packed0.z = asuint(In.Pos0.z);
-	Out.Packed0.w = (PackedC.x << 24) | (PackedC.y << 16) | (PackedC.z << 8) | (PackedC.w);
-
-	Out.Packed1.x = asuint(In.Pos1.x);
-	Out.Packed1.y = asuint(In.Pos1.y);
-	Out.Packed1.z = asuint(In.Pos1.z);
-	Out.Packed1.w = In.bIsScreenSpace ? 0x1 : 0x0;
-
-	Out.Packed2.x = asuint(In.Pos2.x);
-	Out.Packed2.y = asuint(In.Pos2.y);
-	Out.Packed2.z = asuint(In.Pos2.z);
-	Out.Packed2.w = 0;
-	return Out;
+	AddLineWS(InitShaderPrintContext(), Pos0, Pos1, Color);
 }
-
-FTriangleElement UnpackTriangleElement(StructuredBuffer<uint> InPrimitiveBuffer, uint InIndex)
-{
-	const uint Index12 = GetPrimitiveTriangleOffset(InIndex);
-
-	FPackedTriangleElement Out = (FPackedTriangleElement)0;
-	Out.Packed0.x = InPrimitiveBuffer[Index12 + 0];
-	Out.Packed0.y = InPrimitiveBuffer[Index12 + 1];
-	Out.Packed0.z = InPrimitiveBuffer[Index12 + 2];
-	Out.Packed0.w = InPrimitiveBuffer[Index12 + 3];
-
-	Out.Packed1.x = InPrimitiveBuffer[Index12 + 4];
-	Out.Packed1.y = InPrimitiveBuffer[Index12 + 5];
-	Out.Packed1.z = InPrimitiveBuffer[Index12 + 6];
-	Out.Packed1.w = InPrimitiveBuffer[Index12 + 7];
-
-	Out.Packed2.x = InPrimitiveBuffer[Index12 + 8];
-	Out.Packed2.y = InPrimitiveBuffer[Index12 + 9];
-	Out.Packed2.z = InPrimitiveBuffer[Index12 + 10];
-	Out.Packed2.w = InPrimitiveBuffer[Index12 + 11];
-	return UnpackTriangleElement(Out);
+void AddLineWS(float3 Pos0, float3 Pos1, float4 Color0, float4 Color1) 
+{ 
+	AddLineWS(InitShaderPrintContext(), Pos0, Pos1, Color0, Color1);
 }
-
-bool AllocateTriangleElement(uint Count, inout uint OutIndex)
+void AddLine(float3 Pos0, float3 Pos1, float4 Color) 
 {
-	OutIndex = 0;
-	if (ShaderPrintData.MaxTriangleCount == 0)
-	{
-		return false;
-	}
-
-	InterlockedAdd(ShaderPrint_RWEntryBuffer[SHADER_PRINT_COUNTER_OFFSET_TRIANGLE], Count, OutIndex);
-	return (OutIndex + Count) < ShaderPrintData.MaxTriangleCount;
+	AddLine(InitShaderPrintContext(), Pos0, Pos1, Color);
 }
-
-void AddTriangleElement(FTriangleElement In, uint Index)
+void AddLine(float3 Pos0, float3 Pos1, float4 Color0, float4 Color1) 
 {
-	const uint Index12 = GetPrimitiveTriangleOffset(Index);
-
-	FPackedTriangleElement Out = PackTriangleElement(In);
-	ShaderPrint_RWEntryBuffer[Index12 + 0] = Out.Packed0.x;
-	ShaderPrint_RWEntryBuffer[Index12 + 1] = Out.Packed0.y;
-	ShaderPrint_RWEntryBuffer[Index12 + 2] = Out.Packed0.z;
-	ShaderPrint_RWEntryBuffer[Index12 + 3] = Out.Packed0.w;
-
-	ShaderPrint_RWEntryBuffer[Index12 + 4] = Out.Packed1.x;
-	ShaderPrint_RWEntryBuffer[Index12 + 5] = Out.Packed1.y;
-	ShaderPrint_RWEntryBuffer[Index12 + 6] = Out.Packed1.z;
-	ShaderPrint_RWEntryBuffer[Index12 + 7] = Out.Packed1.w;
-
-	ShaderPrint_RWEntryBuffer[Index12 + 8] = Out.Packed2.x;
-	ShaderPrint_RWEntryBuffer[Index12 + 9] = Out.Packed2.y;
-	ShaderPrint_RWEntryBuffer[Index12 + 10]= Out.Packed2.z;
-	ShaderPrint_RWEntryBuffer[Index12 + 11]= Out.Packed2.w;
+	AddLine(InitShaderPrintContext(), Pos0, Pos1, Color0, Color1);
 }
-
-///////////////////////////////////////////////////////////////////////////////////////////////////
-// Filled Triangle
-void AddFilledTriangleTWS(float3 Pos0, float3 Pos1, float3 Pos2, float4 Color)
+void AddLineSS(float2 Pos0, float2 Pos1, float4 Color) 
 {
-	uint Offset = 0;
-	if (AllocateTriangleElement(1, Offset))
-	{
-		FTriangleElement Element;
-		Element.bIsScreenSpace = false;
-		Element.Color = Color;
-
-		Element.Pos0 = Pos0;
-		Element.Pos1 = Pos1;
-		Element.Pos2 = Pos2;
-		AddTriangleElement(Element, Offset + 0);
-	}
+	AddLineSS(InitShaderPrintContext(), Pos0, Pos1, Color);
 }
-
-void AddFilledTriangleWS (float3 Pos0, float3 Pos1, float3 Pos2, float4 Color)							{ AddFilledTriangleTWS(Pos0 + ShaderPrintData.TranslatedWorldOffset, Pos1 + ShaderPrintData.TranslatedWorldOffset, Pos2 + ShaderPrintData.TranslatedWorldOffset, Color); }
-void AddFilledTriangleTWS(FShaderPrintContext Ctx, float3 Pos0, float3 Pos1, float3 Pos2, float4 Color) { if (Ctx.bIsActive) { AddFilledTriangleTWS(Pos0, Pos1, Pos2, Color); } }
-void AddFilledTriangleWS (FShaderPrintContext Ctx, float3 Pos0, float3 Pos1, float3 Pos2, float4 Color) { if (Ctx.bIsActive) { AddFilledTriangleTWS(Pos0 + ShaderPrintData.TranslatedWorldOffset, Pos1 + ShaderPrintData.TranslatedWorldOffset, Pos2 + ShaderPrintData.TranslatedWorldOffset, Color); } }
-
-///////////////////////////////////////////////////////////////////////////////////////////////////
-// Filled Quad
-
-void AddFilledQuadTWS(float3 Pos0, float3 Pos1, float3 Pos2, float3 Pos3, float4 Color)
+void AddQuadWS(float3 Pos0, float3 Pos1, float3 Pos2, float3 Pos3, float4 Color) 
 {
-	uint Offset = 0;
-	if (AllocateTriangleElement(2, Offset))
-	{
-		FTriangleElement Element;
-		Element.bIsScreenSpace = false;
-		Element.Color = Color;
-
-		Element.Pos0 = Pos0;
-		Element.Pos1 = Pos1;
-		Element.Pos2 = Pos2;
-		AddTriangleElement(Element, Offset + 0);
-
-		Element.Pos0 = Pos0;
-		Element.Pos1 = Pos2;
-		Element.Pos2 = Pos3;
-		AddTriangleElement(Element, Offset + 1);
-	}
+	AddQuadWS(InitShaderPrintContext(), Pos0, Pos1, Pos2, Pos3, Color);
 }
-
-void AddFilledQuadWS (float3 Pos0, float3 Pos1, float3 Pos2, float3 Pos3, float4 Color)							 { AddFilledQuadTWS(Pos0 + ShaderPrintData.TranslatedWorldOffset, Pos1 + ShaderPrintData.TranslatedWorldOffset, Pos2 + ShaderPrintData.TranslatedWorldOffset, Pos3 + ShaderPrintData.TranslatedWorldOffset, Color); }
-void AddFilledQuadTWS(FShaderPrintContext Ctx, float3 Pos0, float3 Pos1, float3 Pos2, float3 Pos3, float4 Color) { if (Ctx.bIsActive) { AddFilledQuadTWS(Pos0, Pos1, Pos2, Pos3, Color); } }
-void AddFilledQuadWS (FShaderPrintContext Ctx, float3 Pos0, float3 Pos1, float3 Pos2, float3 Pos3, float4 Color) { if (Ctx.bIsActive) { AddFilledQuadTWS(Pos0 + ShaderPrintData.TranslatedWorldOffset, Pos1 + ShaderPrintData.TranslatedWorldOffset, Pos2 + ShaderPrintData.TranslatedWorldOffset, Pos3 + ShaderPrintData.TranslatedWorldOffset, Color); } }
-
-///////////////////////////////////////////////////////////////////////////////////////////////////
-// Filled Quad (Screen-Space)
-
-void AddFilledQuadSS(float2 InPosInPixel0, float2 InPosInPixel2, float4 Color)
+void AddQuad(float3 Pos0, float3 Pos1, float3 Pos2, float3 Pos3, float4 Color) 
 {
-	const float2 InPosInPixel1 = float2(InPosInPixel2.x, InPosInPixel0.y);
-	const float2 InPosInPixel3 = float2(InPosInPixel0.x, InPosInPixel2.y);
-
-	const float Depth = 0.5f;
-	const float3 Pos0 = float3((InPosInPixel0) / float2(ShaderPrintData.Resolution), Depth);
-	const float3 Pos1 = float3((InPosInPixel1) / float2(ShaderPrintData.Resolution), Depth);
-	const float3 Pos2 = float3((InPosInPixel2) / float2(ShaderPrintData.Resolution), Depth);
-	const float3 Pos3 = float3((InPosInPixel3) / float2(ShaderPrintData.Resolution), Depth);
-
-	uint Offset = 0;
-	if (AllocateTriangleElement(2, Offset))
-	{
-		FTriangleElement Element;
-		Element.bIsScreenSpace = true;
-		Element.Color = Color;
-
-		Element.Pos0 = Pos0;
-		Element.Pos1 = Pos1;
-		Element.Pos2 = Pos2;
-		AddTriangleElement(Element, Offset + 0);
-
-		Element.Pos0 = Pos0;
-		Element.Pos1 = Pos2;
-		Element.Pos2 = Pos3;
-		AddTriangleElement(Element, Offset + 1);
-	}
+	AddQuad(InitShaderPrintContext(), Pos0, Pos1, Pos2, Pos3, Color);
 }
-
-void AddFilledQuadSS(FShaderPrintContext Ctx, float2 Pos0, float2 Pos1, float4 Color) { if (Ctx.bIsActive) { AddFilledQuadSS(Pos0, Pos1, Color); } }
-
-///////////////////////////////////////////////////////////////////////////////////////////////////
-// Line-based Primitives (Line, AABB, OBB, Referentials, ...)
-
-struct FLineElement 
-{
-	float3 Pos0;
-	float3 Pos1;
-	float4 Color0;
-	float4 Color1;
-	bool bIsScreenSpace;
-};
-
-struct FPackedLineElement
-{
-	uint4 Packed0;
-	uint4 Packed1;
-};
-
-FLineElement UnpackLineElement(FPackedLineElement In)
+void AddQuadSS(float2 MinPos, float2 MaxPos, float4 Color)
 {
-	FLineElement Out = (FLineElement)0;
-	{
-		Out.Pos0.x = asfloat(In.Packed0.x);
-		Out.Pos0.y = asfloat(In.Packed0.y);
-		Out.Pos0.z = asfloat(In.Packed0.z);
-		Out.Color0 = float4((In.Packed0.w >> 24) & 0xFF, (In.Packed0.w >> 16) & 0xFF, (In.Packed0.w >> 8) & 0xFF, (In.Packed0.w) & 0xFF) / 255.0f;
-	}
-	{
-		Out.Pos1.x = asfloat(In.Packed1.x);
-		Out.Pos1.y = asfloat(In.Packed1.y);
-		Out.Pos1.z = asfloat(In.Packed1.z);
-		Out.Color1 = float4((In.Packed1.w >> 24) & 0xFF, (In.Packed1.w >> 16) & 0xFF, (In.Packed1.w >> 8) & 0xFF, (In.Packed1.w) & 0xFF) / 255.0f;
-
-		// The 'space' info is stored into the alpha's LSB
-		Out.bIsScreenSpace = (In.Packed1.w & 0x1) > 0u;
-	}
-	return Out;
+	AddQuadSS(InitShaderPrintContext(), MinPos, MaxPos, Color);
 }
-
-FPackedLineElement PackLineElement(FLineElement In)
+void AddAABBTWS(float3 Min, float3 Max, float4 Color) 
 {
-	uint4 PackedC0 = uint4(255.0f * saturate(In.Color0));
-	uint4 PackedC1 = uint4(255.0f * saturate(In.Color1));
-
-	// Store 'space' info into the alpha's LSB
-	PackedC1 = PackedC1 & 0xFE;
-	PackedC1 = PackedC1 | (In.bIsScreenSpace ? 0x1 : 0x0);
-
-	FPackedLineElement Out = (FPackedLineElement)0;
-	Out.Packed0.x = asuint(In.Pos0.x);
-	Out.Packed0.y = asuint(In.Pos0.y);
-	Out.Packed0.z = asuint(In.Pos0.z);
-	Out.Packed0.w = (PackedC0.x << 24) | (PackedC0.y << 16) | (PackedC0.z << 8) | (PackedC0.w);
-
-	Out.Packed1.x = asuint(In.Pos1.x);
-	Out.Packed1.y = asuint(In.Pos1.y);
-	Out.Packed1.z = asuint(In.Pos1.z);
-	Out.Packed1.w = (PackedC1.x << 24) | (PackedC1.y << 16) | (PackedC1.z << 8) | (PackedC1.w);
-	return Out;
+	AddAABBTWS(InitShaderPrintContext(), Min, Max, Color);
 }
-
-FLineElement UnpackLineElement(StructuredBuffer<uint> InPrimitiveBuffer, uint InIndex)
+void AddAABBWS(float3 Min, float3 Max, float4 Color)
 {
-	const uint Index8 = GetPrimitiveLineOffset(InIndex);
-
-	FPackedLineElement Out = (FPackedLineElement)0;
-	Out.Packed0.x = InPrimitiveBuffer[Index8 + 0];
-	Out.Packed0.y = InPrimitiveBuffer[Index8 + 1];
-	Out.Packed0.z = InPrimitiveBuffer[Index8 + 2];
-	Out.Packed0.w = InPrimitiveBuffer[Index8 + 3];
-
-	Out.Packed1.x = InPrimitiveBuffer[Index8 + 4];
-	Out.Packed1.y = InPrimitiveBuffer[Index8 + 5];
-	Out.Packed1.z = InPrimitiveBuffer[Index8 + 6];
-	Out.Packed1.w = InPrimitiveBuffer[Index8 + 7];
-	return UnpackLineElement(Out);
+	AddAABBWS(InitShaderPrintContext(), Min, Max, Color);
 }
-
-bool AllocateLineElement(uint Count, inout uint OutIndex)
+void AddAABB(float3 Min, float3 Max, float4 Color) 
 {
-	OutIndex = 0;
-	if (ShaderPrintData.MaxLineCount == 0)
-	{
-		return false;
-	}
-
-	InterlockedAdd(ShaderPrint_RWEntryBuffer[SHADER_PRINT_COUNTER_OFFSET_LINE], Count, OutIndex);
-	return (OutIndex + Count) < ShaderPrintData.MaxLineCount;
+	AddAABB(InitShaderPrintContext(), Min, Max, Color);
 }
-
-void AddLineElement(FLineElement In, uint Index)
+void AddCrossWS(float3 Pos, float Size, float4 Color) 
 {
-	const uint Index8 = GetPrimitiveLineOffset(Index);
-
-	FPackedLineElement Out = PackLineElement(In);
-	ShaderPrint_RWEntryBuffer[Index8 + 0] = Out.Packed0.x;
-	ShaderPrint_RWEntryBuffer[Index8 + 1] = Out.Packed0.y;
-	ShaderPrint_RWEntryBuffer[Index8 + 2] = Out.Packed0.z;
-	ShaderPrint_RWEntryBuffer[Index8 + 3] = Out.Packed0.w;
-
-	ShaderPrint_RWEntryBuffer[Index8 + 4] = Out.Packed1.x;
-	ShaderPrint_RWEntryBuffer[Index8 + 5] = Out.Packed1.y;
-	ShaderPrint_RWEntryBuffer[Index8 + 6] = Out.Packed1.z;
-	ShaderPrint_RWEntryBuffer[Index8 + 7] = Out.Packed1.w;
+	AddCrossWS(InitShaderPrintContext(), Pos, Size, Color);
 }
-
-int2 GetCursorPos()
+void AddCross(float3 Pos, float Size, float4 Color) 
 {
-	return ShaderPrintData.CursorCoord;
+	AddCross(InitShaderPrintContext(), Pos, Size, Color);
 }
-
-///////////////////////////////////////////////////////////////////////////////////////////////////
-// Line
-
-void AddLineTWS(float3 Pos0, float3 Pos1, float4 Color0, float4 Color1)
+void AddReferentialWS(float3 Pos, float3x3 InM, float Scale = 1)
 {
-	uint Offset = 0;
-	if (AllocateLineElement(1, Offset))
-	{
-		FLineElement Element;
-		Element.Pos0 = Pos0;
-		Element.Pos1 = Pos1;
-		Element.Color0 = Color0;
-		Element.Color1 = Color1;
-		Element.bIsScreenSpace = false;
-		AddLineElement(Element, Offset);
-	}
+	AddReferentialWS(InitShaderPrintContext(), Pos, InM, Scale);
 }
-
-void AddLineTWS(float3 Pos0, float3 Pos1, float4 Color)												{ AddLineTWS(Pos0, Pos1, Color, Color); }
-void AddLineWS (float3 Pos0, float3 Pos1, float4 Color)												{ AddLineTWS(Pos0+ShaderPrintData.TranslatedWorldOffset, Pos1+ShaderPrintData.TranslatedWorldOffset, Color); }
-void AddLineWS (float3 Pos0, float3 Pos1, float4 Color0, float4 Color1)								{ AddLineTWS(Pos0+ShaderPrintData.TranslatedWorldOffset, Pos1+ShaderPrintData.TranslatedWorldOffset, Color0, Color1); }
-void AddLine   (float3 Pos0, float3 Pos1, float4 Color)												{ AddLineTWS(Pos0+ShaderPrintData.TranslatedWorldOffset, Pos1+ShaderPrintData.TranslatedWorldOffset, Color); }
-void AddLine   (float3 Pos0, float3 Pos1, float4 Color0, float4 Color1)								{ AddLineTWS(Pos0+ShaderPrintData.TranslatedWorldOffset, Pos1+ShaderPrintData.TranslatedWorldOffset, Color0, Color1); }
-
-void AddLineTWS(FShaderPrintContext Ctx, float3 Pos0, float3 Pos1, float4 Color0, float4 Color1)	{ if (Ctx.bIsActive) { AddLineTWS(Pos0, Pos1, Color0, Color1); } }
-void AddLineTWS(FShaderPrintContext Ctx, float3 Pos0, float3 Pos1, float4 Color)					{ if (Ctx.bIsActive) { AddLineTWS(Pos0, Pos1, Color, Color); } }
-void AddLineWS (FShaderPrintContext Ctx, float3 Pos0, float3 Pos1, float4 Color)					{ if (Ctx.bIsActive) { AddLineTWS(Pos0+ShaderPrintData.TranslatedWorldOffset, Pos1+ShaderPrintData.TranslatedWorldOffset, Color); } }
-void AddLineWS (FShaderPrintContext Ctx, float3 Pos0, float3 Pos1, float4 Color0, float4 Color1)	{ if (Ctx.bIsActive) { AddLineTWS(Pos0+ShaderPrintData.TranslatedWorldOffset, Pos1+ShaderPrintData.TranslatedWorldOffset, Color0, Color1); } }
-void AddLine   (FShaderPrintContext Ctx, float3 Pos0, float3 Pos1, float4 Color)					{ if (Ctx.bIsActive) { AddLineTWS(Pos0+ShaderPrintData.TranslatedWorldOffset, Pos1+ShaderPrintData.TranslatedWorldOffset, Color); } }
-void AddLine   (FShaderPrintContext Ctx, float3 Pos0, float3 Pos1, float4 Color0, float4 Color1)	{ if (Ctx.bIsActive) { AddLineTWS(Pos0+ShaderPrintData.TranslatedWorldOffset, Pos1+ShaderPrintData.TranslatedWorldOffset, Color0, Color1); } }
-
-// Draw 2D line. Input positions are expressed in pixel coordinate
-void AddLineSS(float2 InPosInPixel0, float2 InPosInPixel1, float4 Color0, float4 Color1)
+void AddReferentialWS(float4x4 InM, float Scale = 1) 
 {
-	const float2 Pos0 = float2(InPosInPixel0) / float2(ShaderPrintData.Resolution);
-	const float2 Pos1 = float2(InPosInPixel1) / float2(ShaderPrintData.Resolution);
-
-	uint Offset = 0;
-	if (AllocateLineElement(1, Offset))
-	{
-		FLineElement Element;
-		Element.Pos0 = float3(Pos0, 0);
-		Element.Pos1 = float3(Pos1, 0);
-		Element.Color0 = Color0;
-		Element.Color1 = Color1;
-		Element.bIsScreenSpace = true;
-		AddLineElement(Element, Offset);
-	}
+	AddReferentialWS(InitShaderPrintContext(), InM, Scale);
 }
-void AddLineSS(float2 Pos0, float2 Pos1, float4 Color)											{ AddLineSS(Pos0, Pos1, Color, Color); }
-
-void AddLineSS(FShaderPrintContext Ctx, float2 Pos0, float2 Pos1, float4 Color0, float4 Color1) { if (Ctx.bIsActive) { AddLineSS(Pos0, Pos1, Color0, Color1); } }
-void AddLineSS(FShaderPrintContext Ctx, float2 Pos0, float2 Pos1, float4 Color)					{ if (Ctx.bIsActive) { AddLineSS(Pos0, Pos1, Color, Color); } }
-
-///////////////////////////////////////////////////////////////////////////////////////////////////
-// Quad
-
-void AddQuadTWS(float3 Pos0, float3 Pos1, float3 Pos2, float3 Pos3, float4 Color)
+void AddReferentialWS(float3 Pos, float3 TangentZ, float Scale = 1) 
 {
-	uint Offset = 0;
-	if (AllocateLineElement(4, Offset))
-	{
-		FLineElement Element;
-		Element.bIsScreenSpace = false;
-		Element.Color0 = Element.Color1 = Color;
-		Element.Pos0 = Pos0; Element.Pos1 = Pos1; AddLineElement(Element, Offset + 0);
-		Element.Pos0 = Pos1; Element.Pos1 = Pos2; AddLineElement(Element, Offset + 1);
-		Element.Pos0 = Pos2; Element.Pos1 = Pos3; AddLineElement(Element, Offset + 2);
-		Element.Pos0 = Pos3; Element.Pos1 = Pos0; AddLineElement(Element, Offset + 3);
-	}
+	AddReferentialWS(InitShaderPrintContext(), Pos, TangentZ, Scale);
 }
-
-void AddQuadWS(float3 Pos0, float3 Pos1, float3 Pos2, float3 Pos3, float4 Color)							{ AddQuadTWS(Pos0+ShaderPrintData.TranslatedWorldOffset, Pos1+ShaderPrintData.TranslatedWorldOffset, Pos2+ShaderPrintData.TranslatedWorldOffset, Pos3+ShaderPrintData.TranslatedWorldOffset, Color); }
-void AddQuad  (float3 Pos0, float3 Pos1, float3 Pos2, float3 Pos3, float4 Color)							{ AddQuadTWS(Pos0+ShaderPrintData.TranslatedWorldOffset, Pos1+ShaderPrintData.TranslatedWorldOffset, Pos2+ShaderPrintData.TranslatedWorldOffset, Pos3+ShaderPrintData.TranslatedWorldOffset, Color); }
-
-void AddQuadTWS(FShaderPrintContext Ctx, float3 Pos0, float3 Pos1, float3 Pos2, float3 Pos3, float4 Color)	{ if (Ctx.bIsActive) { AddQuadTWS(Pos0, Pos1, Pos2, Pos3, Color); } }
-void AddQuadWS (FShaderPrintContext Ctx, float3 Pos0, float3 Pos1, float3 Pos2, float3 Pos3, float4 Color)	{ if (Ctx.bIsActive) { AddQuadTWS(Pos0+ShaderPrintData.TranslatedWorldOffset, Pos1+ShaderPrintData.TranslatedWorldOffset, Pos2+ShaderPrintData.TranslatedWorldOffset, Pos3+ShaderPrintData.TranslatedWorldOffset, Color); } }
-void AddQuad   (FShaderPrintContext Ctx, float3 Pos0, float3 Pos1, float3 Pos2, float3 Pos3, float4 Color)	{ if (Ctx.bIsActive) { AddQuadTWS(Pos0+ShaderPrintData.TranslatedWorldOffset, Pos1+ShaderPrintData.TranslatedWorldOffset, Pos2+ShaderPrintData.TranslatedWorldOffset, Pos3+ShaderPrintData.TranslatedWorldOffset, Color); } }
-
-// Draw 2D quad line
-void AddQuadSS(float2 MinPos, float2 MaxPos, float4 Color)
+void AddReferentialWS(float3 Pos, float3 T, float3 B, float3 N, float Scale = 1)
 {
-	uint Offset = 0;
-	if (AllocateLineElement(4, Offset))
-	{
-		MinPos = float2(MinPos) / float2(ShaderPrintData.Resolution);
-		MaxPos = float2(MaxPos) / float2(ShaderPrintData.Resolution);
-
-		float3 Pos0 = float3(MinPos.x, MinPos.y, 0);
-		float3 Pos1 = float3(MaxPos.x, MinPos.y, 0);
-		float3 Pos2 = float3(MaxPos.x, MaxPos.y, 0);
-		float3 Pos3 = float3(MinPos.x, MaxPos.y, 0);
-
-		FLineElement Element;
-		Element.bIsScreenSpace = true;
-		Element.Color0 = Element.Color1 = Color;
-		Element.Pos0 = Pos0; Element.Pos1 = Pos1; AddLineElement(Element, Offset + 0);
-		Element.Pos0 = Pos1; Element.Pos1 = Pos2; AddLineElement(Element, Offset + 1);
-		Element.Pos0 = Pos2; Element.Pos1 = Pos3; AddLineElement(Element, Offset + 2);
-		Element.Pos0 = Pos3; Element.Pos1 = Pos0; AddLineElement(Element, Offset + 3);
-	}
+	AddReferentialWS(InitShaderPrintContext(), Pos, T, B, N, Scale);
 }
-
-///////////////////////////////////////////////////////////////////////////////////////////////////
-// AABB
-
-void AddAABBTWS(float3 Min, float3 Max, float4 Color)
+void AddReferential(float3 Pos, float3x3 InM, float Scale = 1) 
 {
-	uint Offset = 0;
-	if (AllocateLineElement(12, Offset))
-	{
-		float3 P0 = float3(Min.x, Min.y, Min.z);
-		float3 P1 = float3(Max.x, Min.y, Min.z);
-		float3 P2 = float3(Max.x, Max.y, Min.z);
-		float3 P3 = float3(Min.x, Max.y, Min.z);
-		float3 P4 = float3(Min.x, Min.y, Max.z);
-		float3 P5 = float3(Max.x, Min.y, Max.z);
-		float3 P6 = float3(Max.x, Max.y, Max.z);
-		float3 P7 = float3(Min.x, Max.y, Max.z);
-
-		FLineElement Element;
-		Element.bIsScreenSpace = false;
-		Element.Color0 = Element.Color1 = Color;
-
-		Element.Pos0 = P0; Element.Pos1 = P1; AddLineElement(Element, Offset + 0);
-		Element.Pos0 = P1; Element.Pos1 = P2; AddLineElement(Element, Offset + 1);
-		Element.Pos0 = P2; Element.Pos1 = P3; AddLineElement(Element, Offset + 2);
-		Element.Pos0 = P3; Element.Pos1 = P0; AddLineElement(Element, Offset + 3);
-
-		Element.Pos0 = P4; Element.Pos1 = P5; AddLineElement(Element, Offset + 4);
-		Element.Pos0 = P5; Element.Pos1 = P6; AddLineElement(Element, Offset + 5);
-		Element.Pos0 = P6; Element.Pos1 = P7; AddLineElement(Element, Offset + 6);
-		Element.Pos0 = P7; Element.Pos1 = P4; AddLineElement(Element, Offset + 7);
-
-		Element.Pos0 = P0; Element.Pos1 = P4; AddLineElement(Element, Offset + 8);
-		Element.Pos0 = P1; Element.Pos1 = P5; AddLineElement(Element, Offset + 9);
-		Element.Pos0 = P2; Element.Pos1 = P6; AddLineElement(Element, Offset +10);
-		Element.Pos0 = P3; Element.Pos1 = P7; AddLineElement(Element, Offset +11);
-	}
+	AddReferential(InitShaderPrintContext(), Pos, InM, Scale);
 }
-
-void AddAABBWS(float3 Min, float3 Max, float4 Color)							{ AddAABBTWS(Min + ShaderPrintData.TranslatedWorldOffset, Max + ShaderPrintData.TranslatedWorldOffset, Color); }
-void AddAABB  (float3 Min, float3 Max, float4 Color)							{ AddAABBTWS(Min + ShaderPrintData.TranslatedWorldOffset, Max + ShaderPrintData.TranslatedWorldOffset, Color); }
-
-void AddAABBTWS(FShaderPrintContext Ctx, float3 Min, float3 Max, float4 Color)	{ if (Ctx.bIsActive) { AddAABBTWS(Min, Max, Color); } }
-void AddAABBWS (FShaderPrintContext Ctx, float3 Min, float3 Max, float4 Color)	{ if (Ctx.bIsActive) { AddAABBTWS(Min + ShaderPrintData.TranslatedWorldOffset, Max + ShaderPrintData.TranslatedWorldOffset, Color); } }
-void AddAABB   (FShaderPrintContext Ctx, float3 Min, float3 Max, float4 Color)	{ if (Ctx.bIsActive) { AddAABBTWS(Min + ShaderPrintData.TranslatedWorldOffset, Max + ShaderPrintData.TranslatedWorldOffset, Color); } }
-
-///////////////////////////////////////////////////////////////////////////////////////////////////
-// Cross
-
-void AddCrossTWS(float3 Pos, float Size, float4 Color)
+void AddReferential(float4x4 InM, float Scale = 1) 
 {
-	AddLineTWS(Pos - float3(Size,0,0), Pos + float3(Size,0,0), Color, Color);
-	AddLineTWS(Pos - float3(0,Size,0), Pos + float3(0,Size,0), Color, Color);
-	AddLineTWS(Pos - float3(0,0,Size), Pos + float3(0,0,Size), Color, Color);
+	AddReferential(InitShaderPrintContext(), InM, Scale);
 }
-
-void AddCrossWS(float3 Pos, float Size, float4 Color)							{ AddCrossTWS(Pos + ShaderPrintData.TranslatedWorldOffset, Size, Color); }
-void AddCross  (float3 Pos, float Size, float4 Color)							{ AddCrossTWS(Pos + ShaderPrintData.TranslatedWorldOffset, Size, Color); }
-
-void AddCrossTWS(FShaderPrintContext Ctx, float3 Pos, float Size, float4 Color) { if (Ctx.bIsActive) { AddCrossTWS(Pos, Size, Color); } }
-void AddCrossWS (FShaderPrintContext Ctx, float3 Pos, float Size, float4 Color)	{ if (Ctx.bIsActive) { AddCrossTWS(Pos + ShaderPrintData.TranslatedWorldOffset, Size, Color); } }
-void AddCross   (FShaderPrintContext Ctx, float3 Pos, float Size, float4 Color)	{ if (Ctx.bIsActive) { AddCrossTWS(Pos + ShaderPrintData.TranslatedWorldOffset, Size, Color); } }
-
-///////////////////////////////////////////////////////////////////////////////////////////////////
-// Circles
-
-void AddCircleSS(float2 Center, float Radius, float4 Color, uint SegmentCount=16)
+void AddReferential(float3 Pos, float3 TangentZ, float Scale = 1) 
 {
-	const float TStep = 0.05f;
-
-	float S, C;
-	sincos(0.0, S, C);
-	float2 PrevP = Center + float2(C, S) * Radius;
-
-	for (uint SegIt = 1; SegIt <= SegmentCount; ++SegIt)
-	{
-		const float t = float(SegIt) / float(SegmentCount-1);
-		sincos(t * 2 * PI, S, C);
-		float2 P = Center + float2(C, S) * Radius;
-
-		AddLineSS(PrevP, P, Color);
-		PrevP = P;
-	}
+	AddReferential(InitShaderPrintContext(), Pos, TangentZ, Scale);
 }
-
-void AddCircleSS(FShaderPrintContext Ctx, float2 Center, float Radius, float4 Color, uint SegmentCount = 16) { if (Ctx.bIsActive) { AddCircleSS(Center, Radius, Color, SegmentCount); } }
-
-
-///////////////////////////////////////////////////////////////////////////////////////////////////
-// Referential
-
-void AddReferentialTWS(float3 Pos, float3 T, float3 B, float3 N, float Scale=1)
+void AddReferential(float3 Pos, float3 T, float3 B, float3 N, float Scale = 1) 
 {
-	AddLineTWS(Pos, Pos + normalize(T)*Scale, ColorRed);
-	AddLineTWS(Pos, Pos + normalize(B)*Scale, ColorGreen);
-	AddLineTWS(Pos, Pos + normalize(N)*Scale, ColorBlue);
+	AddReferential(InitShaderPrintContext(), Pos, T, B, N, Scale);
 }
-
-void AddReferentialTWS(float3 Pos, float3 TangentZ, float Scale=1)
+void AddLineTriangleWS(float3 P0, float3 P1, float3 P2, float4 Color) 
 {
-	const float Sign = TangentZ.z >= 0 ? 1 : -1;
-	const float a = -rcp(Sign + TangentZ.z);
-	const float b = TangentZ.x * TangentZ.y * a;
-
-	const float3 TangentX = { 1 + Sign * a * Pow2(TangentZ.x), Sign * b, -Sign * TangentZ.x };
-	const float3 TangentY = { b,  Sign + a * Pow2(TangentZ.y), -TangentZ.y };
-
-	AddReferentialTWS(Pos, TangentX, TangentY, TangentZ, Scale);
+	AddLineTriangleWS(InitShaderPrintContext(), P0, P1, P2, Color);
 }
-
-void AddReferentialTWS(float4x4 InM, float Scale = 1)
+void AddLineTriangle(float3 P0, float3 P1, float3 P2, float4 Color) 
 {
-	AddReferentialTWS(InM[3].xyz, InM[0].xyz, InM[1].xyz, InM[2].xyz, Scale);
+	AddLineTriangle(InitShaderPrintContext(), P0, P1, P2, Color);
 }
-
-void AddReferentialWS(float3 Pos, float3x3 InM, float Scale = 1)					{ AddReferentialTWS(Pos + ShaderPrintData.TranslatedWorldOffset, InM[0].xyz, InM[1].xyz, InM[2].xyz, Scale); }
-void AddReferentialWS(float4x4 InM, float Scale = 1)								{ AddReferentialTWS(InM[3].xyz + ShaderPrintData.TranslatedWorldOffset, InM[0].xyz, InM[1].xyz, InM[2].xyz, Scale); }
-void AddReferentialWS(float3 Pos, float3 TangentZ, float Scale = 1)					{ AddReferentialTWS(Pos + ShaderPrintData.TranslatedWorldOffset, TangentZ, Scale); }
-void AddReferentialWS(float3 Pos, float3 T, float3 B, float3 N, float Scale = 1)	{ AddReferentialTWS(Pos + ShaderPrintData.TranslatedWorldOffset, T, B, N, Scale); }
-void AddReferential  (float3 Pos, float3x3 InM, float Scale = 1)					{ AddReferentialTWS(Pos + ShaderPrintData.TranslatedWorldOffset, InM[0].xyz, InM[1].xyz, InM[2].xyz, Scale); }
-void AddReferential  (float4x4 InM, float Scale = 1)								{ AddReferentialTWS(InM[3].xyz + ShaderPrintData.TranslatedWorldOffset, InM[0].xyz, InM[1].xyz, InM[2].xyz, Scale); }
-void AddReferential  (float3 Pos, float3 TangentZ, float Scale = 1)					{ AddReferentialTWS(Pos + ShaderPrintData.TranslatedWorldOffset, TangentZ, Scale); }
-void AddReferential  (float3 Pos, float3 T, float3 B, float3 N, float Scale = 1)	{ AddReferentialTWS(Pos + ShaderPrintData.TranslatedWorldOffset, T, B, N, Scale); }
-
-void AddReferentialTWS(FShaderPrintContext Ctx, float3 Pos, float3 T, float3 B, float3 N, float Scale = 1)	{ if (Ctx.bIsActive) { AddReferentialTWS(Pos, T, B, N, Scale); } }
-void AddReferentialTWS(FShaderPrintContext Ctx, float3 Pos, float3 TangentZ, float Scale = 1)				{ if (Ctx.bIsActive) { AddReferentialTWS(Pos, TangentZ, Scale); }  }
-void AddReferentialTWS(FShaderPrintContext Ctx, float3 Pos, float3x3 InM, float Scale = 1)					{ if (Ctx.bIsActive) { AddReferentialTWS(Pos, InM[0].xyz, InM[1].xyz, InM[2].xyz, Scale); }  }
-void AddReferentialTWS(FShaderPrintContext Ctx, float4x4 InM, float Scale = 1)								{ if (Ctx.bIsActive) { AddReferentialTWS(InM[3].xyz, InM[0].xyz, InM[1].xyz, InM[2].xyz, Scale); }  }
-void AddReferentialWS (FShaderPrintContext Ctx, float4x4 InM, float Scale = 1)								{ if (Ctx.bIsActive) { AddReferentialTWS(InM[3].xyz + ShaderPrintData.TranslatedWorldOffset, InM[0].xyz, InM[1].xyz, InM[2].xyz, Scale); }  }
-void AddReferentialWS (FShaderPrintContext Ctx, float3 Pos, float3x3 InM, float Scale = 1)					{ if (Ctx.bIsActive) { AddReferentialTWS(Pos + ShaderPrintData.TranslatedWorldOffset, InM[0].xyz, InM[1].xyz, InM[2].xyz, Scale); }  }
-void AddReferentialWS (FShaderPrintContext Ctx, float3 Pos, float3 TangentZ, float Scale = 1)				{ if (Ctx.bIsActive) { AddReferentialTWS(Pos + ShaderPrintData.TranslatedWorldOffset, TangentZ, Scale); }  }
-void AddReferentialWS (FShaderPrintContext Ctx, float3 Pos, float3 T, float3 B, float3 N, float Scale = 1)	{ if (Ctx.bIsActive) { AddReferentialTWS(Pos + ShaderPrintData.TranslatedWorldOffset, T, B, N, Scale); }  }
-void AddReferential   (FShaderPrintContext Ctx, float4x4 InM, float Scale = 1)								{ if (Ctx.bIsActive) { AddReferentialTWS(InM[3].xyz + ShaderPrintData.TranslatedWorldOffset, InM[0].xyz, InM[1].xyz, InM[2].xyz, Scale); }  }
-void AddReferential   (FShaderPrintContext Ctx, float3 Pos, float3x3 InM, float Scale = 1)					{ if (Ctx.bIsActive) { AddReferentialTWS(Pos + ShaderPrintData.TranslatedWorldOffset, InM[0].xyz, InM[1].xyz, InM[2].xyz, Scale); }  }
-void AddReferential   (FShaderPrintContext Ctx, float3 Pos, float3 TangentZ, float Scale = 1)				{ if (Ctx.bIsActive) { AddReferentialTWS(Pos + ShaderPrintData.TranslatedWorldOffset, TangentZ, Scale); }  }
-void AddReferential   (FShaderPrintContext Ctx, float3 Pos, float3 T, float3 B, float3 N, float Scale = 1)	{ if (Ctx.bIsActive) { AddReferentialTWS(Pos + ShaderPrintData.TranslatedWorldOffset, T, B, N, Scale); }  }
-///////////////////////////////////////////////////////////////////////////////////////////////////
-// Triangle
-
-void AddLineTriangleTWS(float3 P0, float3 P1, float3 P2, float4 Color)
+void AddOBBWS(float3 Min, float3 Max, float4 Color, float4x4 Transform)
 {
-	AddLineTWS(P0, P1, Color, Color);
-	AddLineTWS(P1, P2, Color, Color);
-	AddLineTWS(P2, P0, Color, Color);
+	AddOBBWS(InitShaderPrintContext(), Min, Max, Color, Transform);
 }
-
-void AddLineTriangleWS(float3 P0, float3 P1, float3 P2, float4 Color)							{ AddLineTriangleTWS(P0 + ShaderPrintData.TranslatedWorldOffset, P1 + ShaderPrintData.TranslatedWorldOffset, P2 + ShaderPrintData.TranslatedWorldOffset, Color); }
-void AddLineTriangle  (float3 P0, float3 P1, float3 P2, float4 Color)							{ AddLineTriangleTWS(P0 + ShaderPrintData.TranslatedWorldOffset, P1 + ShaderPrintData.TranslatedWorldOffset, P2 + ShaderPrintData.TranslatedWorldOffset, Color); }
-
-void AddLineTriangleTWS(FShaderPrintContext Ctx, float3 P0, float3 P1, float3 P2, float4 Color)	{ if (Ctx.bIsActive) { AddLineTriangleTWS(P0, P1, P2, Color); } }
-void AddLineTriangleWS (FShaderPrintContext Ctx, float3 P0, float3 P1, float3 P2, float4 Color)	{ if (Ctx.bIsActive) { AddLineTriangleTWS(P0 + ShaderPrintData.TranslatedWorldOffset, P1 + ShaderPrintData.TranslatedWorldOffset, P2 + ShaderPrintData.TranslatedWorldOffset, Color); } }
-void AddLineTriangle   (FShaderPrintContext Ctx, float3 P0, float3 P1, float3 P2, float4 Color)	{ if (Ctx.bIsActive) { AddLineTriangleTWS(P0 + ShaderPrintData.TranslatedWorldOffset, P1 + ShaderPrintData.TranslatedWorldOffset, P2 + ShaderPrintData.TranslatedWorldOffset, Color); } }
-
-///////////////////////////////////////////////////////////////////////////////////////////////////
-// OBB
-
-void AddOBBTWS(float3 Min, float3 Max, float4 Color, float4x4 Transform)
+void AddOBB(float3 Min, float3 Max, float4 Color, float4x4 Transform) 
 {
-	uint Offset = 0;
-	if (AllocateLineElement(12, Offset))
-	{
-		float3 P0 = mul(float4(Min.x, Min.y, Min.z, 1.0f), Transform).xyz;
-		float3 P1 = mul(float4(Max.x, Min.y, Min.z, 1.0f), Transform).xyz;
-		float3 P2 = mul(float4(Max.x, Max.y, Min.z, 1.0f), Transform).xyz;
-		float3 P3 = mul(float4(Min.x, Max.y, Min.z, 1.0f), Transform).xyz;
-		float3 P4 = mul(float4(Min.x, Min.y, Max.z, 1.0f), Transform).xyz;
-		float3 P5 = mul(float4(Max.x, Min.y, Max.z, 1.0f), Transform).xyz;
-		float3 P6 = mul(float4(Max.x, Max.y, Max.z, 1.0f), Transform).xyz;
-		float3 P7 = mul(float4(Min.x, Max.y, Max.z, 1.0f), Transform).xyz;
-
-		FLineElement Element;
-		Element.bIsScreenSpace = false;
-		Element.Color0 = Element.Color1 = Color;
-
-		Element.Pos0 = P0; Element.Pos1 = P1; AddLineElement(Element, Offset + 0);
-		Element.Pos0 = P1; Element.Pos1 = P2; AddLineElement(Element, Offset + 1);
-		Element.Pos0 = P2; Element.Pos1 = P3; AddLineElement(Element, Offset + 2);
-		Element.Pos0 = P3; Element.Pos1 = P0; AddLineElement(Element, Offset + 3);
-
-		Element.Pos0 = P4; Element.Pos1 = P5; AddLineElement(Element, Offset + 4);
-		Element.Pos0 = P5; Element.Pos1 = P6; AddLineElement(Element, Offset + 5);
-		Element.Pos0 = P6; Element.Pos1 = P7; AddLineElement(Element, Offset + 6);
-		Element.Pos0 = P7; Element.Pos1 = P4; AddLineElement(Element, Offset + 7);
-
-		Element.Pos0 = P0; Element.Pos1 = P4; AddLineElement(Element, Offset + 8);
-		Element.Pos0 = P1; Element.Pos1 = P5; AddLineElement(Element, Offset + 9);
-		Element.Pos0 = P2; Element.Pos1 = P6; AddLineElement(Element, Offset + 10);
-		Element.Pos0 = P3; Element.Pos1 = P7; AddLineElement(Element, Offset + 11);
-	}
+	AddOBB(InitShaderPrintContext(), Min, Max, Color, Transform);
 }
-
-void AddOBBWS(float3 Min, float3 Max, float4 Color, float4x4 Transform)								{ Transform[3].xyz += ShaderPrintData.TranslatedWorldOffset; AddOBBTWS(Min, Max, Color, Transform); }
-void AddOBB  (float3 Min, float3 Max, float4 Color, float4x4 Transform)								{ Transform[3].xyz += ShaderPrintData.TranslatedWorldOffset; AddOBBTWS(Min, Max, Color, Transform); }
-
-void AddOBBTWS(FShaderPrintContext Ctx, float3 Min, float3 Max, float4 Color, float4x4 Transform)	{ if (Ctx.bIsActive) { AddOBBTWS(Min, Max, Color, Transform); } }
-void AddOBBWS (FShaderPrintContext Ctx, float3 Min, float3 Max, float4 Color, float4x4 Transform)	{ if (Ctx.bIsActive) { Transform[3].xyz += ShaderPrintData.TranslatedWorldOffset; AddOBBTWS(Min, Max, Color, Transform); } }
-void AddOBB   (FShaderPrintContext Ctx, float3 Min, float3 Max, float4 Color, float4x4 Transform)	{ if (Ctx.bIsActive) { Transform[3].xyz += ShaderPrintData.TranslatedWorldOffset; AddOBBTWS(Min, Max, Color, Transform); } }
diff --git a/Engine/Shaders/Private/ShaderPrintCommon.ush b/Engine/Shaders/Private/ShaderPrintCommon.ush
new file mode 100644
--- /dev/null
+++ b/Engine/Shaders/Private/ShaderPrintCommon.ush
@@ -0,0 +1,1257 @@
+/*=============================================================================
+	ShaderPrintCommon.ush:
+	Include this to be able to call ShaderPrint() from arbitrary shaders.
+=============================================================================*/
+
+#pragma once
+
+///////////////////////////////////////////////////////////////////////////////////////////////////
+// Description
+// Shader Print allows to print value & text from shaders. 
+// 
+// 0. Contex-based state API
+// -------------------------
+// FShaderPrintContext Ctx = InitShaderPrintContext(all(SvPosition.xy == float2(100, 100), uint2(50,50))
+//                             ------------------------------------   ------------
+//                                 Pixel/thread to filter             Coord at which 
+//                                                                       to draw
+// 
+// 1. How to print value?
+// ----------------------
+// First create a context, then print scalar/vector/matrix/text...
+// 
+//   FShaderPrintContext Ctx = InitShaderPrintContext(all(PixelPos.xy == uint2(100, 100), uint2(50,50))
+// 
+//   Print(Ctx, MyFloat);
+//   Newline(Ctx);
+// 
+//   Print(Ctx, TEXT("My string for labbelling a float :"), FontOrange);
+//   Print(Ctx, MyFloatInYellow, FontYellow);
+//   Newline(Ctx);
+// 
+// Numerical values can be further formatted, by configuring the number of digits spanned by a number 
+// (i.e. the number of written&non-written digits), and the number of wanted decimal for floating value.
+// 
+//   Print(Ctx, 123, FontYellow, 5); // Spanning is show as: .....
+//   Print(Ctx, 456, FontYellow, 4); // Spanning is show as: xxxx
+//   Print(Ctx, 789, FontYellow, 3); // Spanning is show as: ---
+//
+// Will be displayed as follow:
+//   123   456  789
+//   ..... xxxx ---
+// 
+// Print(Ctx, 12.3456, FontYellow, 5, 2); // Spanning is show as: ...
+// Print(Ctx, 12.3456, FontYellow, 5, 1); // Spanning is show as: xxx
+// 
+// Will be displayed as follow:
+//   12.34 12.3
+//   ..... xxxxx
+// 
+// 2. How to add widget?
+// ---------------------
+// Widget can be added from shader. 
+// For slider:
+// 
+//   const float MySliderValue  = AddSlider(Ctx, DefaultValue, TEXT("My Slider"));
+// 
+// For checkbox:
+// 
+//   const bool MyCheckboxValue = AddCheckbox(Ctx, DefaultValue, TEXT("My Checkbox"));
+//
+// The value of these widgets can be retrieved/evaluated from *any* thread/pixels, but the 
+// drawing need to be done from a single thread/pixel. This is done through the ShaderPrintContext, 
+// which setup which pixel/thread is active/filtered. Example with a compute shader:
+// 
+//    [numthreads(8, 8, 1)]
+//    void MainCS(uint2 DispatchThreadId : SV_DispatchThreadID)
+//    {
+//       // Create context which 'add' the widget only on thread (0,0)
+//       FShaderPrintContext Ctx = InitShaderPrintContext(all(DispatchThreadId == uint2(0, 0), uint2(50,50))
+// 
+//		 // Any threads can get the slider value with this call
+//       const float ScaleValue = AddSlider(Ctx, 1.f, TEXT("Scale"));
+//       ...
+//    }
+// 
+
+///////////////////////////////////////////////////////////////////////////////////////////////////
+// Primitive & Symbols layout in RWEntryBuffer
+//  ----------------------------------------------------------------------------------------------------------------------------------------------------------------------------------------------
+// ||  Line    | Triangle |  Symbol |  Empty  ||  Symb0  |  Symb1  |  Symb2  |  Symb3  |  ...  |  Line0  |  Line1  |  Line2  |  Line3  |  ...  | Triangle0 | Triangle1 | Triangle2 | Triangle3 | ...
+// || Counter  | Counter  | Counter | Counter ||         |         |         |         |  ...  |         |         |         |         |  ...  |           |           |           |           | ...
+// || (1uint)  | (1uint)  | (1uint) | (1uint) || (4uint) | (4uint) | (4uint) | (4uint) |  ...  | (8uint) | (8uint) | (8uint) | (8uint) |  ...  | (12uint)  | (12uint)  | (12uint)  | (12uint)  | ...
+//  ----------------------------------------------------------------------------------------------------------------------------------------------------------------------------------------------
+
+///////////////////////////////////////////////////////////////////////////////////////////////////
+// Input types
+
+// Content of FShaderPrintItem.Type defines what to cast FShaderPrintItem.Value as
+
+// Type of primitives
+#define SHADER_PRINT_TYPE_SYMBOL   0
+#define SHADER_PRINT_TYPE_FLOAT    1
+#define SHADER_PRINT_TYPE_INT      2
+#define SHADER_PRINT_TYPE_UINT     3
+#define SHADER_PRINT_TYPE_HEX      4
+#define SHADER_PRINT_TYPE_SLIDER   5
+#define SHADER_PRINT_TYPE_CHECK    6
+// Type of primitives (to be merged with the list above once all buffer are merged)
+#define SHADER_PRINT_TYPE_LINE     0
+#define SHADER_PRINT_TYPE_TRIANGLE 1
+
+// Size of packed data, in uint
+#define SHADER_PRINT_UINT_COUNT_SYMBOL   4
+#define SHADER_PRINT_UINT_COUNT_LINE     8
+#define SHADER_PRINT_UINT_COUNT_TRIANGLE 12
+
+// Counter offset
+#define SHADER_PRINT_COUNTER_OFFSET_LINE     0
+#define SHADER_PRINT_COUNTER_OFFSET_TRIANGLE 1
+#define SHADER_PRINT_COUNTER_OFFSET_SYMBOL   2
+#define SHADER_PRINT_COUNTER_OFFSET_FREE     3
+#define SHADER_PRINT_COUNTER_COUNT           4
+
+uint GetSymbolOffset(uint InIndex)				
+{
+	return SHADER_PRINT_COUNTER_COUNT + SHADER_PRINT_UINT_COUNT_SYMBOL * InIndex; 
+}
+uint GetPrimitiveLineOffset(uint InIndex, uint MaxValueCount)
+{
+	return SHADER_PRINT_COUNTER_COUNT + SHADER_PRINT_UINT_COUNT_SYMBOL * MaxValueCount + SHADER_PRINT_UINT_COUNT_LINE * InIndex; 
+}
+uint GetPrimitiveTriangleOffset(uint InIndex, uint MaxValueCount, uint MaxLineCount)
+{
+	return SHADER_PRINT_COUNTER_COUNT + SHADER_PRINT_UINT_COUNT_SYMBOL * MaxValueCount + SHADER_PRINT_UINT_COUNT_LINE * MaxLineCount + SHADER_PRINT_UINT_COUNT_TRIANGLE * InIndex; 
+}
+
+void ClearCounters(RWStructuredBuffer<uint> InRWBuffer)
+{
+	// Symbol/print Counter
+	InRWBuffer[SHADER_PRINT_COUNTER_OFFSET_LINE] = 0;
+	InRWBuffer[SHADER_PRINT_COUNTER_OFFSET_TRIANGLE] = 0;
+	InRWBuffer[SHADER_PRINT_COUNTER_OFFSET_SYMBOL] = 0;
+	InRWBuffer[SHADER_PRINT_COUNTER_OFFSET_FREE] = 0;
+}
+
+///////////////////////////////////////////////////////////////////////////////////////////////////
+// State 
+#define SHADER_PRINT_STATE_COUNT_OFFSET 0
+#define SHADER_PRINT_STATE_VALUE_OFFSET 1
+
+#define SHADER_PRINT_STATE_STRIDE 3
+#define SHADER_PRINT_STATE_INDEX_METADATA 0
+#define SHADER_PRINT_STATE_INDEX_VALUE 1
+#define SHADER_PRINT_STATE_INDEX_FRAME 2
+
+#define SHADER_PRINT_STATE_INVALID_INDEX 0xFF
+#define SHADER_PRINT_STATE_HASH_MASK 0xFFFFFF
+
+///////////////////////////////////////////////////////////////////////////////////////////////////
+// Font Color
+
+struct FFontColor
+{
+	float3 Color;
+};
+
+FFontColor InitFontColor(float InX, float InY, float InZ)	{ FFontColor Out; Out.Color = float3(InX, InY, InZ); return Out; }
+FFontColor InitFontColor(float3 In)							{ FFontColor Out; Out.Color = In; return Out; }
+FFontColor GetDefaultFontColor()							{ FFontColor Out; Out.Color = float3(1,1,1); return Out; }
+
+// Predefined colors
+#define FontWhite		InitFontColor(1, 1, 1)
+#define FontGrey		InitFontColor(0.5, 0.5, 0.5)
+#define FontDarkGrey	InitFontColor(0.25, 0.25, 0.25)
+#define FontBlack		InitFontColor(0, 0, 0)
+#define FontRed			InitFontColor(1, 0, 0)
+#define FontGreen		InitFontColor(0, 1, 0)
+#define FontBlue		InitFontColor(0, 0, 1)
+#define FontYellow		InitFontColor(1, 1, 0)
+#define FontCyan		InitFontColor(0, 1, 1)
+#define FontMagenta		InitFontColor(1, 0, 1)
+#define FontOrange		InitFontColor(243.f / 255.f, 156.f / 255.f, 18.f / 255.f)
+#define FontPurple		InitFontColor(169.f / 255.f, 7.f / 255.f, 228.f / 255.f)
+#define FontTurquoise	InitFontColor(26.f / 255.f, 188.f / 255.f, 156.f / 255.f)
+#define FontSilver		InitFontColor(189.f / 255.f, 195.f / 255.f, 199.f / 255.f)
+#define FontEmerald		InitFontColor(46.f / 255.f, 204.f / 255.f, 113.f / 255.f)
+#define FontLightRed	InitFontColor(0.75f, 0.50f, 0.50f)
+#define FontLightGreen	InitFontColor(0.50f, 0.75f, 0.50f)
+#define FontLightBlue	InitFontColor(0.50f, 0.50f, 0.75f)
+#define FontDarkRed		InitFontColor(0.50f, 0.25f, 0.25f)
+#define FontDarkGreen	InitFontColor(0.25f, 0.50f, 0.25f)
+#define FontDarkBlue	InitFontColor(0.25f, 0.25f, 0.50f)
+
+///////////////////////////////////////////////////////////////////////////////////////////////////
+// Primitive colors
+
+#define ColorWhite		float4(1, 1, 1, 1)
+#define ColorBlack		float4(0, 0, 0, 1)
+#define ColorRed		float4(1, 0, 0, 1)
+#define ColorGreen		float4(0, 1, 0, 1)
+#define ColorBlue		float4(0, 0, 1, 1)
+#define ColorYellow		float4(1, 1, 0, 1)
+#define ColorCyan		float4(0, 1, 1, 1)
+#define ColorMagenta	float4(1, 0, 1, 1)
+#define ColorOrange		float4(243.f / 255.f, 156.f / 255.f,  18.f / 255.f, 1)
+#define ColorPurple		float4(169.f / 255.f,   7.f / 255.f, 228.f / 255.f, 1)
+#define ColorTurquoise	float4( 26.f / 255.f, 188.f / 255.f, 156.f / 255.f, 1)
+#define ColorSilver		float4(189.f / 255.f, 195.f / 255.f, 199.f / 255.f, 1)
+#define ColorEmerald	float4( 46.f / 255.f, 204.f / 255.f, 113.f / 255.f, 1)
+#define ColorLightRed	float4(0.75f, 0.50f, 0.50f, 1.0f)
+#define ColorLightGreen	float4(0.50f, 0.75f, 0.50f, 1.0f)
+#define ColorLightBlue	float4(0.50f, 0.50f, 0.75f, 1.0f)
+#define ColorDarkRed	float4(0.50f, 0.25f, 0.25f, 1.0f)
+#define ColorDarkGreen	float4(0.25f, 0.50f, 0.25f, 1.0f)
+#define ColorDarkBlue	float4(0.25f, 0.25f, 0.50f, 1.0f)
+
+///////////////////////////////////////////////////////////////////////////////////////////////////
+// Util pack/unpac functions
+
+struct FShaderPrintItem
+{
+	float2 ScreenPos; // Position in normalized coordinates
+	int Value;        // Cast to value or symbol
+	int Type;         // SHADER_PRINT_TYPE_* defines how to read Value
+	float3 Color;	  // Color 
+	uint Metadata;	  // Metadata (used for widget to store extra state data
+};
+
+// Needs to match C++ code in ShaderPrint.cpp
+struct FPackedShaderPrintItem
+{
+	uint ScreenPos16bits; // Position in normalized coordinates
+	uint Value;           // Cast to value or symbol
+	uint TypeAndColor;    //
+	uint Metadata;	      //
+};
+
+FPackedShaderPrintItem PackShaderPrintItem(FShaderPrintItem In)
+{
+	const uint3 Color8bits = saturate(In.Color) * 0xFF;
+	const uint2 ScreenPos16bit = f32tof16(In.ScreenPos);
+
+	FPackedShaderPrintItem Out;
+	Out.ScreenPos16bits = ScreenPos16bit.x | (ScreenPos16bit.y<<16);
+	Out.Value = asuint(In.Value);
+	Out.TypeAndColor = (Color8bits.z << 24) | (Color8bits.y << 16) | (Color8bits.x << 8) | (In.Type & 0xFF);
+	Out.Metadata = In.Metadata;
+	return Out;
+}
+
+FShaderPrintItem UnpackShaderPrintItem(FPackedShaderPrintItem In)
+{
+	const uint2 ScreenPos16bits = uint2(In.ScreenPos16bits & 0xFFFF, (In.ScreenPos16bits>>16) & 0xFFFF);
+
+	FShaderPrintItem Out;
+	Out.ScreenPos = f16tof32(ScreenPos16bits);
+	Out.Value	  = asint(In.Value);
+	Out.Type	  = (In.TypeAndColor) & 0xFF;
+	Out.Color.x   = float((In.TypeAndColor >>  8) & 0xFF) / 255.f;
+	Out.Color.y   = float((In.TypeAndColor >> 16) & 0xFF) / 255.f;
+	Out.Color.z   = float((In.TypeAndColor >> 24) & 0xFF) / 255.f;
+	Out.Metadata  = In.Metadata;
+	return Out;
+}
+
+void WriteSymbol(uint Offset, FShaderPrintItem In, RWStructuredBuffer<uint> InRWBuffer)
+{
+	const uint Offset4 = GetSymbolOffset(Offset);
+
+	const FPackedShaderPrintItem Packed = PackShaderPrintItem(In);
+	InRWBuffer[Offset4 + 0] = Packed.ScreenPos16bits;
+	InRWBuffer[Offset4 + 1] = Packed.Value;
+	InRWBuffer[Offset4 + 2] = Packed.TypeAndColor;
+	InRWBuffer[Offset4 + 3] = Packed.Metadata;
+}
+
+FShaderPrintItem ReadSymbol(uint Offset, StructuredBuffer<uint> InBuffer)
+{
+	const uint Offset4 = GetSymbolOffset(Offset);
+
+	FPackedShaderPrintItem Packed = (FPackedShaderPrintItem)0;
+	Packed.ScreenPos16bits	= InBuffer[Offset4 + 0];
+	Packed.Value			= InBuffer[Offset4 + 1];
+	Packed.TypeAndColor		= InBuffer[Offset4 + 2];
+	Packed.Metadata			= InBuffer[Offset4 + 3];
+
+	return UnpackShaderPrintItem(Packed);
+}
+
+///////////////////////////////////////////////////////////////////////////////////////////////////
+// Context structures.
+// Stores the current dynamic state, and the constant data describing config and buffer bindings.
+
+struct FShaderPrintConfig
+{
+	int2 Resolution;
+	int2 CursorCoord;
+	float3 TranslatedWorldOffset;
+	float2 FontSize;
+	float2 FontSpacing;
+	uint MaxValueCount;
+	uint MaxSymbolCount;
+	uint MaxStateCount;
+	uint MaxLineCount;
+	uint MaxTriangleCount;
+};
+
+struct FShaderPrintBuffers
+{
+	StructuredBuffer<uint> StateBuffer;
+	RWStructuredBuffer<uint> RWEntryBuffer;
+};
+
+struct FShaderPrintContext
+{
+	bool bIsActive;
+	float2 StartPos;
+	float2 Pos;
+
+	FShaderPrintConfig Config;
+	FShaderPrintBuffers Buffers;
+};
+
+///////////////////////////////////////////////////////////////////////////////////////////////////
+// Internal helpers
+
+void ShaderPrint_Internal(FShaderPrintContext Ctx, in FShaderPrintItem Item)
+{
+	// If MaxValueCount is 0 then we don't reset the buffer counter so early out here
+	if (Ctx.Config.MaxValueCount == 0)
+	{
+		return;
+	}
+
+	// Buffer counter is stored in first element .Value
+	int IndexToStore = 0;
+	InterlockedAdd(Ctx.Buffers.RWEntryBuffer[SHADER_PRINT_COUNTER_OFFSET_SYMBOL], 1, IndexToStore);
+
+	// Prevent writing off the buffer
+	// Note that counter still increases so need clamp when reading it in later passes
+	if (uint(IndexToStore) >= Ctx.Config.MaxValueCount)
+	{
+		return;
+	}
+
+	WriteSymbol(IndexToStore, Item, Ctx.Buffers.RWEntryBuffer);
+}
+
+void ShaderPrint_Internal(FShaderPrintContext Ctx, in float2 ScreenPos, in int Value, in FFontColor FontColor, in uint Metadata, in int Type)
+{
+	FShaderPrintItem Item;
+	Item.ScreenPos = ScreenPos;
+	Item.Value = Value;
+	Item.Type = Type;
+	Item.Color = FontColor.Color;
+	Item.Metadata = Metadata;
+	ShaderPrint_Internal(Ctx, Item);
+}
+
+///////////////////////////////////////////////////////////////////////////////////////////////////
+// Symbol printing
+
+void PrintSymbol(inout FShaderPrintContext Ctx, in int Symbol, in FFontColor Color)
+{
+	if (Ctx.bIsActive)
+	{
+		ShaderPrint_Internal(Ctx, Ctx.Pos, Symbol, Color, 0u, SHADER_PRINT_TYPE_SYMBOL);
+		Ctx.Pos.x += Ctx.Config.FontSpacing.x;
+	}
+}
+void PrintSymbol(inout FShaderPrintContext Ctx, in int Symbol) { PrintSymbol(Ctx, Symbol, GetDefaultFontColor()); }
+
+void Newline(inout FShaderPrintContext Ctx)
+{
+	Ctx.Pos.x = Ctx.StartPos.x;
+	Ctx.Pos.y += Ctx.Config.FontSpacing.y;
+}
+
+int2 GetCursorPos(FShaderPrintContext Ctx)
+{
+	return Ctx.Config.CursorCoord;
+}
+
+///////////////////////////////////////////////////////////////////////////////////////////////////
+// Text & hash
+
+struct FShaderPrintText
+{
+	uint Index;
+};
+
+FShaderPrintText InitShaderPrintText(uint InIndex)
+{
+	FShaderPrintText Out;
+	Out.Index = InIndex;
+	return Out;
+}
+
+// Forward declarations for text accessors
+uint ShaderPrintGetChar(uint InIndex);
+uint ShaderPrintGetOffset(FShaderPrintText InTextEntry);
+uint ShaderPrintGetHash(FShaderPrintText InTextEntry);
+ 
+// Function for reading global TEXT string
+float2 ShaderPrintText_Internal(FShaderPrintContext Ctx, bool bIsActive, float2 InPos, FShaderPrintText InTextEntry, FFontColor InColor)
+{
+	// If MaxValueCount is 0 then we don't reset the buffer counter so early out here
+	if (bIsActive && Ctx.Config.MaxValueCount > 0)
+	{
+		const uint Begin = ShaderPrintGetOffset(InTextEntry);
+		const uint End   = ShaderPrintGetOffset(InitShaderPrintText(InTextEntry.Index + 1));
+		const uint Count = End - Begin;
+
+		// Buffer counter is stored in first element .Value
+		// Prevent writing off the buffer
+		// Note that counter still increases so need clamp when reading it in later passes
+		int IndexToStore = 0;
+		InterlockedAdd(Ctx.Buffers.RWEntryBuffer[SHADER_PRINT_COUNTER_OFFSET_SYMBOL], Count, IndexToStore);
+		if (uint(IndexToStore + Count) < Ctx.Config.MaxValueCount)
+		{
+			for (uint i = Begin; i < End; ++i)
+			{
+				FShaderPrintItem Item;
+				Item.ScreenPos = InPos;
+				Item.Value = ShaderPrintGetChar(i);
+				Item.Type = SHADER_PRINT_TYPE_SYMBOL;
+				Item.Color = InColor.Color;
+				Item.Metadata = 0u;
+
+				WriteSymbol(IndexToStore, Item, Ctx.Buffers.RWEntryBuffer);
+
+				++IndexToStore;
+				InPos.x += Ctx.Config.FontSpacing.x;
+			}	
+		}
+	}
+	return InPos;
+}
+
+void Print(inout FShaderPrintContext Ctx, FShaderPrintText InTextEntry, FFontColor InColor)	{ Ctx.Pos = ShaderPrintText_Internal(Ctx, Ctx.bIsActive, Ctx.Pos, InTextEntry, InColor); }
+void Print(inout FShaderPrintContext Ctx, FShaderPrintText InTextEntry)						{ Ctx.Pos = ShaderPrintText_Internal(Ctx, Ctx.bIsActive, Ctx.Pos, InTextEntry, GetDefaultFontColor()); }
+
+// Function for reading global TEXT string. 
+// The data and function are generated by the shader compiler.
+GENERATED_SHADER_PRINT
+
+///////////////////////////////////////////////////////////////////////////////////////////////////
+// Value printing (common value printing)
+#define MAX_DECIMAL_COUNT 5u
+#define MAX_DIGIT_COUNT 12u
+
+float2 ShaderPrintValue_Internal(FShaderPrintContext Ctx, float2 ScreenPos, float Value, FFontColor Color, uint MaxDigit, uint MaxDecimal)  { ShaderPrint_Internal(Ctx, ScreenPos, asint(Value),           Color, MaxDecimal, SHADER_PRINT_TYPE_FLOAT); ScreenPos.x += Ctx.Config.FontSpacing.x * MaxDigit; return ScreenPos; }
+float2 ShaderPrintValue_Internal(FShaderPrintContext Ctx, float2 ScreenPos, int   Value, FFontColor Color, uint MaxDigit, uint MaxDecimal)  { ShaderPrint_Internal(Ctx, ScreenPos, Value,                  Color, MaxDecimal, SHADER_PRINT_TYPE_INT);   ScreenPos.x += Ctx.Config.FontSpacing.x * MaxDigit; return ScreenPos; }
+float2 ShaderPrintValue_Internal(FShaderPrintContext Ctx, float2 ScreenPos, uint  Value, FFontColor Color, uint MaxDigit, uint MaxDecimal)  { ShaderPrint_Internal(Ctx, ScreenPos, asint(Value),           Color, MaxDecimal, SHADER_PRINT_TYPE_UINT);  ScreenPos.x += Ctx.Config.FontSpacing.x * MaxDigit; return ScreenPos; }
+float2 ShaderPrintValue_Internal(FShaderPrintContext Ctx, float2 ScreenPos, bool  Value, FFontColor Color, uint MaxDigit, uint MaxDecimal)  { ShaderPrint_Internal(Ctx, ScreenPos, asint(Value ? 1u : 0u), Color, MaxDecimal, SHADER_PRINT_TYPE_UINT);  ScreenPos.x += Ctx.Config.FontSpacing.x * MaxDigit; return ScreenPos; }
+
+///////////////////////////////////////////////////////////////////////////////////////////////////
+// Scalar type printing
+
+#define SHADER_PRINT_OVERLOAD_1(InNumComponent, InType) \
+void   Print(inout FShaderPrintContext Ctx, InType Value, FFontColor Color, uint MaxDigit=MAX_DIGIT_COUNT, uint MaxDecimal=MAX_DECIMAL_COUNT)	{ if (Ctx.bIsActive) { Ctx.Pos = ShaderPrintValue_Internal(Ctx, Ctx.Pos, Value, Color, MaxDigit, MaxDecimal); } } \
+void   Print(inout FShaderPrintContext Ctx, InType Value)																						{ if (Ctx.bIsActive) { Ctx.Pos = ShaderPrintValue_Internal(Ctx, Ctx.Pos, Value, GetDefaultFontColor(), MAX_DIGIT_COUNT, MAX_DECIMAL_COUNT); } }
+
+///////////////////////////////////////////////////////////////////////////////////////////////////
+// Vector type printing
+
+#define SHADER_PRINT_OVERLOAD_N(InNumComponent, InType) \
+void   Print(inout FShaderPrintContext Ctx, InType Value, FFontColor Color, uint MaxDigit=MAX_DIGIT_COUNT, uint MaxDecimal=MAX_DECIMAL_COUNT)	{ if (Ctx.bIsActive) { UNROLL for (uint CompIt=0;CompIt<InNumComponent;++CompIt) { Ctx.Pos = ShaderPrintValue_Internal(Ctx, Ctx.Pos, Value[CompIt], Color, MaxDigit, MaxDecimal); } } } \
+void   Print(inout FShaderPrintContext Ctx, InType Value)																						{ if (Ctx.bIsActive) { UNROLL for (uint CompIt=0;CompIt<InNumComponent;++CompIt) { Ctx.Pos = ShaderPrintValue_Internal(Ctx, Ctx.Pos, Value[CompIt], GetDefaultFontColor(), MAX_DIGIT_COUNT, MAX_DECIMAL_COUNT); } } }
+
+// Print(X) version
+SHADER_PRINT_OVERLOAD_1(1, float)
+SHADER_PRINT_OVERLOAD_N(2, float2)
+SHADER_PRINT_OVERLOAD_N(3, float3)
+SHADER_PRINT_OVERLOAD_N(4, float4)
+
+SHADER_PRINT_OVERLOAD_1(1, uint)
+SHADER_PRINT_OVERLOAD_N(2, uint2)
+SHADER_PRINT_OVERLOAD_N(3, uint3)
+SHADER_PRINT_OVERLOAD_N(4, uint4)
+
+SHADER_PRINT_OVERLOAD_1(1, int)
+SHADER_PRINT_OVERLOAD_N(2, int2)
+SHADER_PRINT_OVERLOAD_N(3, int3)
+SHADER_PRINT_OVERLOAD_N(4, int4)
+
+SHADER_PRINT_OVERLOAD_1(1, bool)
+
+///////////////////////////////////////////////////////////////////////////////////////////////////
+// Matrix type printing
+
+#define SHADER_PRINT_OVERLOAD_NN(InNumComponentX, InNumComponentY, InType) \
+void   Print(inout FShaderPrintContext Ctx, InType Value, FFontColor Color, uint MaxDigit=MAX_DIGIT_COUNT, uint MaxDecimal=MAX_DECIMAL_COUNT)	{ if (Ctx.bIsActive) { UNROLL for (uint CompY=0;CompY<InNumComponentY;++CompY) { Print(Ctx, Value[CompY], Color, MaxDigit, MaxDecimal); Ctx.Pos += float2(0.f, Ctx.Config.FontSpacing.y); } } } \
+void   Print(inout FShaderPrintContext Ctx, InType Value)																						{ if (Ctx.bIsActive) { UNROLL for (uint CompY=0;CompY<InNumComponentY;++CompY) { Print(Ctx, Value[CompY], GetDefaultFontColor(), MAX_DIGIT_COUNT, MAX_DECIMAL_COUNT); Ctx.Pos += float2(0.f, Ctx.Config.FontSpacing.y); } } }
+
+// Print(X) version
+SHADER_PRINT_OVERLOAD_NN(3, 3, float3x3)
+SHADER_PRINT_OVERLOAD_NN(4, 3, float4x3)
+SHADER_PRINT_OVERLOAD_NN(4, 4, float4x4)
+
+///////////////////////////////////////////////////////////////////////////////////////////////////
+// Widgets (experimental)
+
+uint2 GetCheckboxWidgetSize(float2 FontSpacing, int2 Resolution)
+{
+	return uint(FontSpacing.x * Resolution.x).xx;
+}
+
+uint2 GetSliderWidgetSize(float2 FontSpacing, int2 Resolution)
+{
+	return uint2(10,1) * uint(FontSpacing.x * Resolution.x);
+}
+
+struct FShaderPrintMetadata
+{
+	uint Hash;
+	uint Index;
+};
+
+uint ShaderPrintPackMetadata(uint InHash, uint InIndex)
+{
+	return ((InIndex & 0xFFu) << 24) | (InHash & SHADER_PRINT_STATE_HASH_MASK);
+}
+uint ShaderPrintPackMetadata(FShaderPrintMetadata In) { return ShaderPrintPackMetadata(In.Hash, In.Index); }
+
+FShaderPrintMetadata ShaderPrintUnpackMetadata(uint Metadata)
+{
+	FShaderPrintMetadata Out;
+	Out.Hash = (Metadata & SHADER_PRINT_STATE_HASH_MASK);
+	Out.Index= (Metadata >> 24) & 0xFF;
+	return Out;
+}
+
+// Return true if found, false otherwise
+bool ShaderPrintGetStateValue(FShaderPrintContext Ctx, in uint InHash, inout uint OutIndex, inout uint OutValue)
+{
+	OutIndex = SHADER_PRINT_STATE_INVALID_INDEX;
+	OutValue = 0;
+
+	// Slow linear search among all widget states
+	const uint MaxCount = min(Ctx.Buffers.StateBuffer[SHADER_PRINT_STATE_COUNT_OFFSET], Ctx.Config.MaxStateCount);
+	for (uint Index = 0; Index < MaxCount; ++Index)
+	{
+		const uint Index3 = Index * SHADER_PRINT_STATE_STRIDE + SHADER_PRINT_STATE_VALUE_OFFSET;
+		const FShaderPrintMetadata Metadata = ShaderPrintUnpackMetadata(Ctx.Buffers.StateBuffer[Index3 + SHADER_PRINT_STATE_INDEX_METADATA]);
+		if (Metadata.Hash == InHash)
+		{
+			OutIndex = Metadata.Index;
+			OutValue = Ctx.Buffers.StateBuffer[Index3 + SHADER_PRINT_STATE_INDEX_VALUE];
+			return true;
+		}
+	}
+	return false;
+}
+
+// Add a checkbox widget 
+bool AddCheckbox(inout FShaderPrintContext Ctx, FShaderPrintText InTextEntry, bool bDefault, in FFontColor InColor)
+{
+	uint Index = 0;
+	uint RawValue = 0;
+	const uint Hash = ShaderPrintGetHash(InTextEntry) & SHADER_PRINT_STATE_HASH_MASK; // Trunk hask since metadata has limited storage
+	const bool bIsValid = ShaderPrintGetStateValue(Ctx, Hash, Index, RawValue);
+	RawValue = bIsValid ? RawValue : (bDefault ? 0x1 : 0x0);
+
+	if (Ctx.bIsActive)
+	{
+		const uint2 Coord = Ctx.Pos * Ctx.Config.Resolution;
+		const uint2 Cusor = Ctx.Config.CursorCoord;
+		const uint2 WidgetSize = GetCheckboxWidgetSize(Ctx.Config.FontSpacing, Ctx.Config.Resolution);
+		const uint2 WidgetMax = uint2(Coord.x + WidgetSize.x, Coord.y + WidgetSize.y / 2);
+		const uint2 WidgetMin = uint2(Coord.x, Coord.y - WidgetSize.y / 2);
+		const bool  bIsInside = all(Cusor >= WidgetMin) && all(Cusor <= WidgetMax);
+
+		const bool bWasInside = RawValue & 0x2;
+
+		// Update checkbox value 
+		// Track if the cursor was already within the box or not
+		if (bIsInside)
+		{
+			if (!bWasInside)
+			{
+				// First time the cursor is within the checkbox
+				RawValue = ((RawValue & 0x1) == 1 ? 0x0 : 0x1) | 0x2;
+			}
+		}
+		else
+		{
+			RawValue = (RawValue & 0x1);
+		}
+
+		FShaderPrintItem E;
+		E.ScreenPos = Ctx.Pos;
+		E.Value = RawValue;
+		E.Type = SHADER_PRINT_TYPE_CHECK;
+		E.Color = bIsInside ? FontYellow.Color : FontWhite.Color;// InColor.Color;
+		E.Metadata = ShaderPrintPackMetadata(Hash, Index);
+		ShaderPrint_Internal(Ctx, E);
+
+		Ctx.Pos.x += float(WidgetSize.x) / float(Ctx.Config.Resolution.x) + Ctx.Config.FontSpacing.x;
+
+		// Text
+		Print(Ctx, InTextEntry, InColor);
+		Ctx.Pos.x += Ctx.Config.FontSpacing.x;
+	}
+
+	return RawValue & 0x1;
+}
+
+// Add a slider widget  
+float AddSlider(inout FShaderPrintContext Ctx, FShaderPrintText InTextEntry, float bDefault, in FFontColor InColor, float InMin = 0.f, float InMax = 1.f)
+{
+	// The value is stored as normalized, for displaying slide value correctly. This is not idealy
+	// as it implies some precision lost due to back&forth conversion between normalized/non-normalized value.
+	uint Index = 0;
+	uint RawValue = 0;
+	const uint Hash = ShaderPrintGetHash(InTextEntry) & SHADER_PRINT_STATE_HASH_MASK; // Trunk hask since metadata has limited storage
+	const bool bIsValid = ShaderPrintGetStateValue(Ctx, Hash, Index, RawValue);
+	float NormalizedValue = bIsValid ? asfloat(RawValue) : bDefault; 
+	float Value = NormalizedValue * (InMax - InMin) + InMin;
+
+	if (Ctx.bIsActive)
+	{
+		const uint2 Coord = Ctx.Pos * Ctx.Config.Resolution;
+		const uint2 Cusor = Ctx.Config.CursorCoord;
+		const uint2 WidgetSize = GetSliderWidgetSize(Ctx.Config.FontSpacing, Ctx.Config.Resolution);
+		const uint2 WidgetMax = uint2(Coord.x + WidgetSize.x, Coord.y + WidgetSize.y / 2);
+		const uint2 WidgetMin = uint2(Coord.x, Coord.y - WidgetSize.y / 2);
+		const bool  bIsInside = all(Cusor >= WidgetMin) && all(Cusor <= WidgetMax);
+
+		// Update slider value
+		if (bIsInside)
+		{
+			const float S = saturate(float(Cusor.x - WidgetMin.x) / float(WidgetMax.x - WidgetMin.x));
+			Value = lerp(InMin, InMax, S);
+		}
+
+		NormalizedValue = saturate((Value - InMin) / (InMax - InMin));
+
+		FShaderPrintItem E;
+		E.ScreenPos = Ctx.Pos;
+		E.Value = asint(NormalizedValue);
+		E.Type = SHADER_PRINT_TYPE_SLIDER;
+		E.Color = bIsInside ? FontYellow.Color : FontWhite.Color; // InColor.Color;
+		E.Metadata = ShaderPrintPackMetadata(Hash, Index);
+		ShaderPrint_Internal(Ctx, E);
+
+		Ctx.Pos.x += float(WidgetSize.x) / float(Ctx.Config.Resolution.x) + Ctx.Config.FontSpacing.x;
+
+		// Text
+		Print(Ctx, InTextEntry, InColor);
+		Ctx.Pos.x += Ctx.Config.FontSpacing.x;
+	}
+
+	return Value;
+}
+
+///////////////////////////////////////////////////////////////////////////////////////////////////
+// Primitives types
+
+///////////////////////////////////////////////////////////////////////////////////////////////////
+// Triangle-based Primitives (Triangle, Point, ...)
+
+struct FTriangleElement 
+{
+	float3 Pos0;
+	float3 Pos1;
+	float3 Pos2;
+	float4 Color;
+	bool bIsScreenSpace;
+};
+
+struct FPackedTriangleElement
+{
+	uint4 Packed0;
+	uint4 Packed1;
+	uint4 Packed2;
+};
+
+FTriangleElement UnpackTriangleElement(FPackedTriangleElement In)
+{
+	FTriangleElement Out = (FTriangleElement)0;
+	{
+		Out.Pos0.x = asfloat(In.Packed0.x);
+		Out.Pos0.y = asfloat(In.Packed0.y);
+		Out.Pos0.z = asfloat(In.Packed0.z);
+		Out.Color  = float4((In.Packed0.w >> 24) & 0xFF, (In.Packed0.w >> 16) & 0xFF, (In.Packed0.w >> 8) & 0xFF, (In.Packed0.w) & 0xFF) / 255.0f;
+	}
+	{
+		Out.Pos1.x = asfloat(In.Packed1.x);
+		Out.Pos1.y = asfloat(In.Packed1.y);
+		Out.Pos1.z = asfloat(In.Packed1.z);
+
+		// The 'space' info is stored into the alpha's LSB
+		Out.bIsScreenSpace = (In.Packed1.w & 0x1) > 0u;
+	}
+	{
+		Out.Pos2.x = asfloat(In.Packed2.x);
+		Out.Pos2.y = asfloat(In.Packed2.y);
+		Out.Pos2.z = asfloat(In.Packed2.z);
+	}
+	return Out;
+}
+
+FPackedTriangleElement PackTriangleElement(FTriangleElement In)
+{
+	uint4 PackedC = uint4(255.0f * saturate(In.Color));
+
+	FPackedTriangleElement Out = (FPackedTriangleElement)0;
+	Out.Packed0.x = asuint(In.Pos0.x);
+	Out.Packed0.y = asuint(In.Pos0.y);
+	Out.Packed0.z = asuint(In.Pos0.z);
+	Out.Packed0.w = (PackedC.x << 24) | (PackedC.y << 16) | (PackedC.z << 8) | (PackedC.w);
+
+	Out.Packed1.x = asuint(In.Pos1.x);
+	Out.Packed1.y = asuint(In.Pos1.y);
+	Out.Packed1.z = asuint(In.Pos1.z);
+	Out.Packed1.w = In.bIsScreenSpace ? 0x1 : 0x0;
+
+	Out.Packed2.x = asuint(In.Pos2.x);
+	Out.Packed2.y = asuint(In.Pos2.y);
+	Out.Packed2.z = asuint(In.Pos2.z);
+	Out.Packed2.w = 0;
+	return Out;
+}
+
+FTriangleElement UnpackTriangleElement(StructuredBuffer<uint> InPrimitiveBuffer, uint InIndex, uint MaxValueCount, uint MaxLineCount)
+{
+	const uint Index12 = GetPrimitiveTriangleOffset(InIndex, MaxValueCount, MaxLineCount);
+
+	FPackedTriangleElement Out = (FPackedTriangleElement)0;
+	Out.Packed0.x = InPrimitiveBuffer[Index12 + 0];
+	Out.Packed0.y = InPrimitiveBuffer[Index12 + 1];
+	Out.Packed0.z = InPrimitiveBuffer[Index12 + 2];
+	Out.Packed0.w = InPrimitiveBuffer[Index12 + 3];
+
+	Out.Packed1.x = InPrimitiveBuffer[Index12 + 4];
+	Out.Packed1.y = InPrimitiveBuffer[Index12 + 5];
+	Out.Packed1.z = InPrimitiveBuffer[Index12 + 6];
+	Out.Packed1.w = InPrimitiveBuffer[Index12 + 7];
+
+	Out.Packed2.x = InPrimitiveBuffer[Index12 + 8];
+	Out.Packed2.y = InPrimitiveBuffer[Index12 + 9];
+	Out.Packed2.z = InPrimitiveBuffer[Index12 + 10];
+	Out.Packed2.w = InPrimitiveBuffer[Index12 + 11];
+	return UnpackTriangleElement(Out);
+}
+
+bool AllocateTriangleElement(FShaderPrintContext Ctx, uint Count, inout uint OutIndex)
+{
+	OutIndex = 0;
+	if (Ctx.Config.MaxTriangleCount == 0)
+	{
+		return false;
+	}
+
+	InterlockedAdd(Ctx.Buffers.RWEntryBuffer[SHADER_PRINT_COUNTER_OFFSET_TRIANGLE], Count, OutIndex);
+	return (OutIndex + Count) < Ctx.Config.MaxTriangleCount;
+}
+
+void AddTriangleElement(FShaderPrintContext Ctx, FTriangleElement In, uint Index)
+{
+	const uint Index12 = GetPrimitiveTriangleOffset(Index, Ctx.Config.MaxValueCount, Ctx.Config.MaxLineCount);
+
+	FPackedTriangleElement Out = PackTriangleElement(In);
+	Ctx.Buffers.RWEntryBuffer[Index12 + 0] = Out.Packed0.x;
+	Ctx.Buffers.RWEntryBuffer[Index12 + 1] = Out.Packed0.y;
+	Ctx.Buffers.RWEntryBuffer[Index12 + 2] = Out.Packed0.z;
+	Ctx.Buffers.RWEntryBuffer[Index12 + 3] = Out.Packed0.w;
+
+	Ctx.Buffers.RWEntryBuffer[Index12 + 4] = Out.Packed1.x;
+	Ctx.Buffers.RWEntryBuffer[Index12 + 5] = Out.Packed1.y;
+	Ctx.Buffers.RWEntryBuffer[Index12 + 6] = Out.Packed1.z;
+	Ctx.Buffers.RWEntryBuffer[Index12 + 7] = Out.Packed1.w;
+
+	Ctx.Buffers.RWEntryBuffer[Index12 + 8] = Out.Packed2.x;
+	Ctx.Buffers.RWEntryBuffer[Index12 + 9] = Out.Packed2.y;
+	Ctx.Buffers.RWEntryBuffer[Index12 + 10] = Out.Packed2.z;
+	Ctx.Buffers.RWEntryBuffer[Index12 + 11] = Out.Packed2.w;
+}
+
+///////////////////////////////////////////////////////////////////////////////////////////////////
+// Filled Triangle
+void AddFilledTriangleTWS(FShaderPrintContext Ctx, float3 Pos0, float3 Pos1, float3 Pos2, float4 Color)
+{
+	if (Ctx.bIsActive)
+	{
+		uint Offset = 0;
+		if (AllocateTriangleElement(Ctx, 1, Offset))
+		{
+			FTriangleElement Element;
+			Element.bIsScreenSpace = false;
+			Element.Color = Color;
+
+			Element.Pos0 = Pos0;
+			Element.Pos1 = Pos1;
+			Element.Pos2 = Pos2;
+			AddTriangleElement(Ctx, Element, Offset + 0);
+		}
+	}
+}
+
+void AddFilledTriangleWS (FShaderPrintContext Ctx, float3 Pos0, float3 Pos1, float3 Pos2, float4 Color) { AddFilledTriangleTWS(Ctx, Pos0 + Ctx.Config.TranslatedWorldOffset, Pos1 + Ctx.Config.TranslatedWorldOffset, Pos2 + Ctx.Config.TranslatedWorldOffset, Color); }
+
+///////////////////////////////////////////////////////////////////////////////////////////////////
+// Filled Quad
+
+void AddFilledQuadTWS(FShaderPrintContext Ctx, float3 Pos0, float3 Pos1, float3 Pos2, float3 Pos3, float4 Color)
+{
+	if (Ctx.bIsActive)
+	{
+		uint Offset = 0;
+		if (AllocateTriangleElement(Ctx, 2, Offset))
+		{
+			FTriangleElement Element;
+			Element.bIsScreenSpace = false;
+			Element.Color = Color;
+
+			Element.Pos0 = Pos0;
+			Element.Pos1 = Pos1;
+			Element.Pos2 = Pos2;
+			AddTriangleElement(Ctx, Element, Offset + 0);
+
+			Element.Pos0 = Pos0;
+			Element.Pos1 = Pos2;
+			Element.Pos2 = Pos3;
+			AddTriangleElement(Ctx, Element, Offset + 1);
+		}
+	}
+}
+
+void AddFilledQuadWS (FShaderPrintContext Ctx, float3 Pos0, float3 Pos1, float3 Pos2, float3 Pos3, float4 Color) { AddFilledQuadTWS(Ctx, Pos0 + Ctx.Config.TranslatedWorldOffset, Pos1 + Ctx.Config.TranslatedWorldOffset, Pos2 + Ctx.Config.TranslatedWorldOffset, Pos3 + Ctx.Config.TranslatedWorldOffset, Color); }
+
+///////////////////////////////////////////////////////////////////////////////////////////////////
+// Filled Quad (Screen-Space)
+
+void AddFilledQuadSS(FShaderPrintContext Ctx, float2 InPosInPixel0, float2 InPosInPixel2, float4 Color)
+{
+	if (Ctx.bIsActive)
+	{
+		const float2 InPosInPixel1 = float2(InPosInPixel2.x, InPosInPixel0.y);
+		const float2 InPosInPixel3 = float2(InPosInPixel0.x, InPosInPixel2.y);
+
+		const float Depth = 0.5f;
+		const float3 Pos0 = float3((InPosInPixel0) / float2(Ctx.Config.Resolution), Depth);
+		const float3 Pos1 = float3((InPosInPixel1) / float2(Ctx.Config.Resolution), Depth);
+		const float3 Pos2 = float3((InPosInPixel2) / float2(Ctx.Config.Resolution), Depth);
+		const float3 Pos3 = float3((InPosInPixel3) / float2(Ctx.Config.Resolution), Depth);
+
+		uint Offset = 0;
+		if (AllocateTriangleElement(Ctx, 2, Offset))
+		{
+			FTriangleElement Element;
+			Element.bIsScreenSpace = true;
+			Element.Color = Color;
+
+			Element.Pos0 = Pos0;
+			Element.Pos1 = Pos1;
+			Element.Pos2 = Pos2;
+			AddTriangleElement(Ctx, Element, Offset + 0);
+
+			Element.Pos0 = Pos0;
+			Element.Pos1 = Pos2;
+			Element.Pos2 = Pos3;
+			AddTriangleElement(Ctx, Element, Offset + 1);
+		}
+	}
+}
+
+///////////////////////////////////////////////////////////////////////////////////////////////////
+// Line-based Primitives (Line, AABB, OBB, Referentials, ...)
+
+struct FLineElement 
+{
+	float3 Pos0;
+	float3 Pos1;
+	float4 Color0;
+	float4 Color1;
+	bool bIsScreenSpace;
+};
+
+struct FPackedLineElement
+{
+	uint4 Packed0;
+	uint4 Packed1;
+};
+
+FLineElement UnpackLineElement(FPackedLineElement In)
+{
+	FLineElement Out = (FLineElement)0;
+	{
+		Out.Pos0.x = asfloat(In.Packed0.x);
+		Out.Pos0.y = asfloat(In.Packed0.y);
+		Out.Pos0.z = asfloat(In.Packed0.z);
+		Out.Color0 = float4((In.Packed0.w >> 24) & 0xFF, (In.Packed0.w >> 16) & 0xFF, (In.Packed0.w >> 8) & 0xFF, (In.Packed0.w) & 0xFF) / 255.0f;
+	}
+	{
+		Out.Pos1.x = asfloat(In.Packed1.x);
+		Out.Pos1.y = asfloat(In.Packed1.y);
+		Out.Pos1.z = asfloat(In.Packed1.z);
+		Out.Color1 = float4((In.Packed1.w >> 24) & 0xFF, (In.Packed1.w >> 16) & 0xFF, (In.Packed1.w >> 8) & 0xFF, (In.Packed1.w) & 0xFF) / 255.0f;
+
+		// The 'space' info is stored into the alpha's LSB
+		Out.bIsScreenSpace = (In.Packed1.w & 0x1) > 0u;
+	}
+	return Out;
+}
+
+FPackedLineElement PackLineElement(FLineElement In)
+{
+	uint4 PackedC0 = uint4(255.0f * saturate(In.Color0));
+	uint4 PackedC1 = uint4(255.0f * saturate(In.Color1));
+
+	// Store 'space' info into the alpha's LSB
+	PackedC1 = PackedC1 & 0xFE;
+	PackedC1 = PackedC1 | (In.bIsScreenSpace ? 0x1 : 0x0);
+
+	FPackedLineElement Out = (FPackedLineElement)0;
+	Out.Packed0.x = asuint(In.Pos0.x);
+	Out.Packed0.y = asuint(In.Pos0.y);
+	Out.Packed0.z = asuint(In.Pos0.z);
+	Out.Packed0.w = (PackedC0.x << 24) | (PackedC0.y << 16) | (PackedC0.z << 8) | (PackedC0.w);
+
+	Out.Packed1.x = asuint(In.Pos1.x);
+	Out.Packed1.y = asuint(In.Pos1.y);
+	Out.Packed1.z = asuint(In.Pos1.z);
+	Out.Packed1.w = (PackedC1.x << 24) | (PackedC1.y << 16) | (PackedC1.z << 8) | (PackedC1.w);
+	return Out;
+}
+
+FLineElement UnpackLineElement(StructuredBuffer<uint> InPrimitiveBuffer, uint InIndex, uint MaxValueCount)
+{
+	const uint Index8 = GetPrimitiveLineOffset(InIndex, MaxValueCount);
+
+	FPackedLineElement Out = (FPackedLineElement)0;
+	Out.Packed0.x = InPrimitiveBuffer[Index8 + 0];
+	Out.Packed0.y = InPrimitiveBuffer[Index8 + 1];
+	Out.Packed0.z = InPrimitiveBuffer[Index8 + 2];
+	Out.Packed0.w = InPrimitiveBuffer[Index8 + 3];
+
+	Out.Packed1.x = InPrimitiveBuffer[Index8 + 4];
+	Out.Packed1.y = InPrimitiveBuffer[Index8 + 5];
+	Out.Packed1.z = InPrimitiveBuffer[Index8 + 6];
+	Out.Packed1.w = InPrimitiveBuffer[Index8 + 7];
+	return UnpackLineElement(Out);
+}
+
+bool AllocateLineElement(FShaderPrintContext Ctx, uint Count, inout uint OutIndex)
+{
+	OutIndex = 0;
+	if (Ctx.Config.MaxLineCount == 0)
+	{
+		return false;
+	}
+
+	InterlockedAdd(Ctx.Buffers.RWEntryBuffer[SHADER_PRINT_COUNTER_OFFSET_LINE], Count, OutIndex);
+	return (OutIndex + Count) < Ctx.Config.MaxLineCount;
+}
+
+void AddLineElement(FShaderPrintContext Ctx, FLineElement In, uint Index)
+{
+	const uint Index8 = GetPrimitiveLineOffset(Index, Ctx.Config.MaxValueCount);
+
+	FPackedLineElement Out = PackLineElement(In);
+	Ctx.Buffers.RWEntryBuffer[Index8 + 0] = Out.Packed0.x;
+	Ctx.Buffers.RWEntryBuffer[Index8 + 1] = Out.Packed0.y;
+	Ctx.Buffers.RWEntryBuffer[Index8 + 2] = Out.Packed0.z;
+	Ctx.Buffers.RWEntryBuffer[Index8 + 3] = Out.Packed0.w;
+
+	Ctx.Buffers.RWEntryBuffer[Index8 + 4] = Out.Packed1.x;
+	Ctx.Buffers.RWEntryBuffer[Index8 + 5] = Out.Packed1.y;
+	Ctx.Buffers.RWEntryBuffer[Index8 + 6] = Out.Packed1.z;
+	Ctx.Buffers.RWEntryBuffer[Index8 + 7] = Out.Packed1.w;
+}
+
+///////////////////////////////////////////////////////////////////////////////////////////////////
+// Line
+
+void AddLineTWS(FShaderPrintContext Ctx, float3 Pos0, float3 Pos1, float4 Color0, float4 Color1)
+{
+	if (Ctx.bIsActive)
+	{
+		uint Offset = 0;
+		if (AllocateLineElement(Ctx, 1, Offset))
+		{
+			FLineElement Element;
+			Element.Pos0 = Pos0;
+			Element.Pos1 = Pos1;
+			Element.Color0 = Color0;
+			Element.Color1 = Color1;
+			Element.bIsScreenSpace = false;
+			AddLineElement(Ctx, Element, Offset);
+		}
+	}
+}
+
+void AddLineTWS(FShaderPrintContext Ctx, float3 Pos0, float3 Pos1, float4 Color)					{ AddLineTWS(Ctx, Pos0, Pos1, Color, Color); }
+void AddLineWS (FShaderPrintContext Ctx, float3 Pos0, float3 Pos1, float4 Color)					{ AddLineTWS(Ctx, Pos0+Ctx.Config.TranslatedWorldOffset, Pos1+Ctx.Config.TranslatedWorldOffset, Color, Color); }
+void AddLineWS (FShaderPrintContext Ctx, float3 Pos0, float3 Pos1, float4 Color0, float4 Color1)	{ AddLineTWS(Ctx, Pos0+Ctx.Config.TranslatedWorldOffset, Pos1+Ctx.Config.TranslatedWorldOffset, Color0, Color1); }
+void AddLine   (FShaderPrintContext Ctx, float3 Pos0, float3 Pos1, float4 Color)					{ AddLineTWS(Ctx, Pos0+Ctx.Config.TranslatedWorldOffset, Pos1+Ctx.Config.TranslatedWorldOffset, Color, Color); }
+void AddLine   (FShaderPrintContext Ctx, float3 Pos0, float3 Pos1, float4 Color0, float4 Color1)	{ AddLineTWS(Ctx, Pos0+Ctx.Config.TranslatedWorldOffset, Pos1+Ctx.Config.TranslatedWorldOffset, Color0, Color1); }
+
+// Draw 2D line. Input positions are expressed in pixel coordinate
+void AddLineSS(FShaderPrintContext Ctx, float2 InPosInPixel0, float2 InPosInPixel1, float4 Color0, float4 Color1)
+{
+	if (Ctx.bIsActive) 
+	{
+		const float2 Pos0 = float2(InPosInPixel0) / float2(Ctx.Config.Resolution);
+		const float2 Pos1 = float2(InPosInPixel1) / float2(Ctx.Config.Resolution);
+
+		uint Offset = 0;
+		if (AllocateLineElement(Ctx, 1, Offset))
+		{
+			FLineElement Element;
+			Element.Pos0 = float3(Pos0, 0);
+			Element.Pos1 = float3(Pos1, 0);
+			Element.Color0 = Color0;
+			Element.Color1 = Color1;
+			Element.bIsScreenSpace = true;
+			AddLineElement(Ctx, Element, Offset);
+		}
+	}
+}
+
+void AddLineSS(FShaderPrintContext Ctx, float2 Pos0, float2 Pos1, float4 Color)					{ AddLineSS(Ctx, Pos0, Pos1, Color, Color); }
+
+///////////////////////////////////////////////////////////////////////////////////////////////////
+// Quad
+
+void AddQuadTWS(FShaderPrintContext Ctx, float3 Pos0, float3 Pos1, float3 Pos2, float3 Pos3, float4 Color)
+{
+	if (Ctx.bIsActive)
+	{
+		uint Offset = 0;
+		if (AllocateLineElement(Ctx, 4, Offset))
+		{
+			FLineElement Element;
+			Element.bIsScreenSpace = false;
+			Element.Color0 = Element.Color1 = Color;
+			Element.Pos0 = Pos0; Element.Pos1 = Pos1; AddLineElement(Ctx, Element, Offset + 0);
+			Element.Pos0 = Pos1; Element.Pos1 = Pos2; AddLineElement(Ctx, Element, Offset + 1);
+			Element.Pos0 = Pos2; Element.Pos1 = Pos3; AddLineElement(Ctx, Element, Offset + 2);
+			Element.Pos0 = Pos3; Element.Pos1 = Pos0; AddLineElement(Ctx, Element, Offset + 3);
+		}
+	}
+}
+
+void AddQuadWS (FShaderPrintContext Ctx, float3 Pos0, float3 Pos1, float3 Pos2, float3 Pos3, float4 Color)	{ AddQuadTWS(Ctx, Pos0+Ctx.Config.TranslatedWorldOffset, Pos1+Ctx.Config.TranslatedWorldOffset, Pos2+Ctx.Config.TranslatedWorldOffset, Pos3+Ctx.Config.TranslatedWorldOffset, Color); }
+void AddQuad   (FShaderPrintContext Ctx, float3 Pos0, float3 Pos1, float3 Pos2, float3 Pos3, float4 Color)	{ AddQuadTWS(Ctx, Pos0+Ctx.Config.TranslatedWorldOffset, Pos1+Ctx.Config.TranslatedWorldOffset, Pos2+Ctx.Config.TranslatedWorldOffset, Pos3+Ctx.Config.TranslatedWorldOffset, Color); }
+
+// Draw 2D quad line
+void AddQuadSS(FShaderPrintContext Ctx, float2 MinPos, float2 MaxPos, float4 Color)
+{
+	if (Ctx.bIsActive)
+	{
+		uint Offset = 0;
+		if (AllocateLineElement(Ctx, 4, Offset))
+		{
+			MinPos = float2(MinPos) / float2(Ctx.Config.Resolution);
+			MaxPos = float2(MaxPos) / float2(Ctx.Config.Resolution);
+
+			float3 Pos0 = float3(MinPos.x, MinPos.y, 0);
+			float3 Pos1 = float3(MaxPos.x, MinPos.y, 0);
+			float3 Pos2 = float3(MaxPos.x, MaxPos.y, 0);
+			float3 Pos3 = float3(MinPos.x, MaxPos.y, 0);
+
+			FLineElement Element;
+			Element.bIsScreenSpace = true;
+			Element.Color0 = Element.Color1 = Color;
+			Element.Pos0 = Pos0; Element.Pos1 = Pos1; AddLineElement(Ctx, Element, Offset + 0);
+			Element.Pos0 = Pos1; Element.Pos1 = Pos2; AddLineElement(Ctx, Element, Offset + 1);
+			Element.Pos0 = Pos2; Element.Pos1 = Pos3; AddLineElement(Ctx, Element, Offset + 2);
+			Element.Pos0 = Pos3; Element.Pos1 = Pos0; AddLineElement(Ctx, Element, Offset + 3);
+		}
+	}
+}
+
+///////////////////////////////////////////////////////////////////////////////////////////////////
+// AABB
+
+void AddAABBTWS(FShaderPrintContext Ctx, float3 Min, float3 Max, float4 Color)
+{
+	if (Ctx.bIsActive)
+	{
+		uint Offset = 0;
+		if (AllocateLineElement(Ctx, 12, Offset))
+		{
+			float3 P0 = float3(Min.x, Min.y, Min.z);
+			float3 P1 = float3(Max.x, Min.y, Min.z);
+			float3 P2 = float3(Max.x, Max.y, Min.z);
+			float3 P3 = float3(Min.x, Max.y, Min.z);
+			float3 P4 = float3(Min.x, Min.y, Max.z);
+			float3 P5 = float3(Max.x, Min.y, Max.z);
+			float3 P6 = float3(Max.x, Max.y, Max.z);
+			float3 P7 = float3(Min.x, Max.y, Max.z);
+
+			FLineElement Element;
+			Element.bIsScreenSpace = false;
+			Element.Color0 = Element.Color1 = Color;
+
+			Element.Pos0 = P0; Element.Pos1 = P1; AddLineElement(Ctx, Element, Offset + 0);
+			Element.Pos0 = P1; Element.Pos1 = P2; AddLineElement(Ctx, Element, Offset + 1);
+			Element.Pos0 = P2; Element.Pos1 = P3; AddLineElement(Ctx, Element, Offset + 2);
+			Element.Pos0 = P3; Element.Pos1 = P0; AddLineElement(Ctx, Element, Offset + 3);
+
+			Element.Pos0 = P4; Element.Pos1 = P5; AddLineElement(Ctx, Element, Offset + 4);
+			Element.Pos0 = P5; Element.Pos1 = P6; AddLineElement(Ctx, Element, Offset + 5);
+			Element.Pos0 = P6; Element.Pos1 = P7; AddLineElement(Ctx, Element, Offset + 6);
+			Element.Pos0 = P7; Element.Pos1 = P4; AddLineElement(Ctx, Element, Offset + 7);
+
+			Element.Pos0 = P0; Element.Pos1 = P4; AddLineElement(Ctx, Element, Offset + 8);
+			Element.Pos0 = P1; Element.Pos1 = P5; AddLineElement(Ctx, Element, Offset + 9);
+			Element.Pos0 = P2; Element.Pos1 = P6; AddLineElement(Ctx, Element, Offset +10);
+			Element.Pos0 = P3; Element.Pos1 = P7; AddLineElement(Ctx, Element, Offset +11);
+		}
+	}
+}
+
+void AddAABBWS (FShaderPrintContext Ctx, float3 Min, float3 Max, float4 Color)	{ AddAABBTWS(Ctx, Min + Ctx.Config.TranslatedWorldOffset, Max + Ctx.Config.TranslatedWorldOffset, Color); }
+void AddAABB   (FShaderPrintContext Ctx, float3 Min, float3 Max, float4 Color)	{ AddAABBTWS(Ctx, Min + Ctx.Config.TranslatedWorldOffset, Max + Ctx.Config.TranslatedWorldOffset, Color); }
+
+///////////////////////////////////////////////////////////////////////////////////////////////////
+// Cross
+
+void AddCrossTWS(FShaderPrintContext Ctx, float3 Pos, float Size, float4 Color)
+{
+	if (Ctx.bIsActive)
+	{
+		AddLineTWS(Ctx, Pos - float3(Size,0,0), Pos + float3(Size,0,0), Color, Color);
+		AddLineTWS(Ctx, Pos - float3(0,Size,0), Pos + float3(0,Size,0), Color, Color);
+		AddLineTWS(Ctx, Pos - float3(0,0,Size), Pos + float3(0,0,Size), Color, Color);
+	}
+}
+
+void AddCrossWS (FShaderPrintContext Ctx, float3 Pos, float Size, float4 Color)	{ AddCrossTWS(Ctx, Pos + Ctx.Config.TranslatedWorldOffset, Size, Color); }
+void AddCross   (FShaderPrintContext Ctx, float3 Pos, float Size, float4 Color)	{ AddCrossTWS(Ctx, Pos + Ctx.Config.TranslatedWorldOffset, Size, Color); }
+
+///////////////////////////////////////////////////////////////////////////////////////////////////
+// Circles
+
+void AddCircleSS(FShaderPrintContext Ctx, float2 Center, float Radius, float4 Color, uint SegmentCount = 16)
+{
+	if (Ctx.bIsActive)
+	{
+		const float TStep = 0.05f;
+
+		float S, C;
+		sincos(0.0, S, C);
+		float2 PrevP = Center + float2(C, S) * Radius;
+
+		for (uint SegIt = 1; SegIt <= SegmentCount; ++SegIt)
+		{
+			const float t = float(SegIt) / float(SegmentCount-1);
+			sincos(t * 2 * PI, S, C);
+			float2 P = Center + float2(C, S) * Radius;
+
+			AddLineSS(Ctx, PrevP, P, Color, Color);
+			PrevP = P;
+		}
+	}
+}
+
+///////////////////////////////////////////////////////////////////////////////////////////////////
+// Referential
+
+void AddReferentialTWS(FShaderPrintContext Ctx, float3 Pos, float3 T, float3 B, float3 N, float Scale = 1)
+{
+	if (Ctx.bIsActive)
+	{
+		AddLineTWS(Ctx, Pos, Pos + normalize(T)*Scale, ColorRed, ColorRed);
+		AddLineTWS(Ctx, Pos, Pos + normalize(B)*Scale, ColorGreen, ColorGreen);
+		AddLineTWS(Ctx, Pos, Pos + normalize(N)*Scale, ColorBlue, ColorBlue);
+	}
+}
+
+void AddReferentialTWS(FShaderPrintContext Ctx, float3 Pos, float3x3 InM, float Scale = 1)
+{
+	AddReferentialTWS(Ctx, Pos, InM[0].xyz, InM[1].xyz, InM[2].xyz, Scale);
+}
+
+void AddReferentialTWS(FShaderPrintContext Ctx, float4x4 InM, float Scale = 1)
+{
+	AddReferentialTWS(Ctx, InM[3].xyz, InM[0].xyz, InM[1].xyz, InM[2].xyz, Scale);
+}
+
+void AddReferentialTWS(FShaderPrintContext Ctx, float3 Pos, float3 TangentZ, float Scale = 1)
+{
+	if (Ctx.bIsActive)
+	{
+		const float Sign = TangentZ.z >= 0 ? 1 : -1;
+		const float a = -rcp(Sign + TangentZ.z);
+		const float b = TangentZ.x * TangentZ.y * a;
+
+		const float3 TangentX = { 1 + Sign * a * Pow2(TangentZ.x), Sign * b, -Sign * TangentZ.x };
+		const float3 TangentY = { b,  Sign + a * Pow2(TangentZ.y), -TangentZ.y };
+
+		AddReferentialTWS(Ctx, Pos, TangentX, TangentY, TangentZ, Scale);
+	}
+}
+
+void AddReferentialWS (FShaderPrintContext Ctx, float4x4 InM, float Scale = 1)								{ AddReferentialTWS(Ctx, InM[3].xyz + Ctx.Config.TranslatedWorldOffset, InM[0].xyz, InM[1].xyz, InM[2].xyz, Scale); } 
+void AddReferentialWS (FShaderPrintContext Ctx, float3 Pos, float3x3 InM, float Scale = 1)					{ AddReferentialTWS(Ctx, Pos + Ctx.Config.TranslatedWorldOffset, InM[0].xyz, InM[1].xyz, InM[2].xyz, Scale); } 
+void AddReferentialWS (FShaderPrintContext Ctx, float3 Pos, float3 TangentZ, float Scale = 1)				{ AddReferentialTWS(Ctx, Pos + Ctx.Config.TranslatedWorldOffset, TangentZ, Scale); } 
+void AddReferentialWS (FShaderPrintContext Ctx, float3 Pos, float3 T, float3 B, float3 N, float Scale = 1)	{ AddReferentialTWS(Ctx, Pos + Ctx.Config.TranslatedWorldOffset, T, B, N, Scale); } 
+void AddReferential   (FShaderPrintContext Ctx, float4x4 InM, float Scale = 1)								{ AddReferentialTWS(Ctx, InM[3].xyz + Ctx.Config.TranslatedWorldOffset, InM[0].xyz, InM[1].xyz, InM[2].xyz, Scale); } 
+void AddReferential   (FShaderPrintContext Ctx, float3 Pos, float3x3 InM, float Scale = 1)					{ AddReferentialTWS(Ctx, Pos + Ctx.Config.TranslatedWorldOffset, InM[0].xyz, InM[1].xyz, InM[2].xyz, Scale); } 
+void AddReferential   (FShaderPrintContext Ctx, float3 Pos, float3 TangentZ, float Scale = 1)				{ AddReferentialTWS(Ctx, Pos + Ctx.Config.TranslatedWorldOffset, TangentZ, Scale); } 
+void AddReferential   (FShaderPrintContext Ctx, float3 Pos, float3 T, float3 B, float3 N, float Scale = 1)	{ AddReferentialTWS(Ctx, Pos + Ctx.Config.TranslatedWorldOffset, T, B, N, Scale); } 
+
+///////////////////////////////////////////////////////////////////////////////////////////////////
+// Triangle
+
+void AddLineTriangleTWS(FShaderPrintContext Ctx, float3 P0, float3 P1, float3 P2, float4 Color)
+{
+	if (Ctx.bIsActive)
+	{
+		AddLineTWS(Ctx, P0, P1, Color, Color);
+		AddLineTWS(Ctx, P1, P2, Color, Color);
+		AddLineTWS(Ctx, P2, P0, Color, Color);
+	}
+}
+
+void AddLineTriangleWS (FShaderPrintContext Ctx, float3 P0, float3 P1, float3 P2, float4 Color)	{ AddLineTriangleTWS(Ctx, P0 + Ctx.Config.TranslatedWorldOffset, P1 + Ctx.Config.TranslatedWorldOffset, P2 + Ctx.Config.TranslatedWorldOffset, Color); }
+void AddLineTriangle   (FShaderPrintContext Ctx, float3 P0, float3 P1, float3 P2, float4 Color)	{ AddLineTriangleTWS(Ctx, P0 + Ctx.Config.TranslatedWorldOffset, P1 + Ctx.Config.TranslatedWorldOffset, P2 + Ctx.Config.TranslatedWorldOffset, Color); }
+
+///////////////////////////////////////////////////////////////////////////////////////////////////
+// OBB
+
+void AddOBBTWS(FShaderPrintContext Ctx, float3 Min, float3 Max, float4 Color, float4x4 Transform)
+{
+	if (Ctx.bIsActive) 
+	{
+		uint Offset = 0;
+		if (AllocateLineElement(Ctx, 12, Offset))
+		{
+			float3 P0 = mul(float4(Min.x, Min.y, Min.z, 1.0f), Transform).xyz;
+			float3 P1 = mul(float4(Max.x, Min.y, Min.z, 1.0f), Transform).xyz;
+			float3 P2 = mul(float4(Max.x, Max.y, Min.z, 1.0f), Transform).xyz;
+			float3 P3 = mul(float4(Min.x, Max.y, Min.z, 1.0f), Transform).xyz;
+			float3 P4 = mul(float4(Min.x, Min.y, Max.z, 1.0f), Transform).xyz;
+			float3 P5 = mul(float4(Max.x, Min.y, Max.z, 1.0f), Transform).xyz;
+			float3 P6 = mul(float4(Max.x, Max.y, Max.z, 1.0f), Transform).xyz;
+			float3 P7 = mul(float4(Min.x, Max.y, Max.z, 1.0f), Transform).xyz;
+
+			FLineElement Element;
+			Element.bIsScreenSpace = false;
+			Element.Color0 = Element.Color1 = Color;
+
+			Element.Pos0 = P0; Element.Pos1 = P1; AddLineElement(Ctx, Element, Offset + 0);
+			Element.Pos0 = P1; Element.Pos1 = P2; AddLineElement(Ctx, Element, Offset + 1);
+			Element.Pos0 = P2; Element.Pos1 = P3; AddLineElement(Ctx, Element, Offset + 2);
+			Element.Pos0 = P3; Element.Pos1 = P0; AddLineElement(Ctx, Element, Offset + 3);
+
+			Element.Pos0 = P4; Element.Pos1 = P5; AddLineElement(Ctx, Element, Offset + 4);
+			Element.Pos0 = P5; Element.Pos1 = P6; AddLineElement(Ctx, Element, Offset + 5);
+			Element.Pos0 = P6; Element.Pos1 = P7; AddLineElement(Ctx, Element, Offset + 6);
+			Element.Pos0 = P7; Element.Pos1 = P4; AddLineElement(Ctx, Element, Offset + 7);
+
+			Element.Pos0 = P0; Element.Pos1 = P4; AddLineElement(Ctx, Element, Offset + 8);
+			Element.Pos0 = P1; Element.Pos1 = P5; AddLineElement(Ctx, Element, Offset + 9);
+			Element.Pos0 = P2; Element.Pos1 = P6; AddLineElement(Ctx, Element, Offset + 10);
+			Element.Pos0 = P3; Element.Pos1 = P7; AddLineElement(Ctx, Element, Offset + 11);
+		}
+	}
+}
+
+void AddOBBWS (FShaderPrintContext Ctx, float3 Min, float3 Max, float4 Color, float4x4 Transform)	{ Transform[3].xyz += Ctx.Config.TranslatedWorldOffset; AddOBBTWS(Ctx, Min, Max, Color, Transform); }
+void AddOBB   (FShaderPrintContext Ctx, float3 Min, float3 Max, float4 Color, float4x4 Transform)	{ Transform[3].xyz += Ctx.Config.TranslatedWorldOffset; AddOBBTWS(Ctx, Min, Max, Color, Transform); }
diff --git a/Engine/Shaders/Private/ShaderPrintDraw.usf b/Engine/Shaders/Private/ShaderPrintDraw.usf
--- a/Engine/Shaders/Private/ShaderPrintDraw.usf
+++ b/Engine/Shaders/Private/ShaderPrintDraw.usf
@@ -5,7 +5,8 @@
 =============================================================================*/
 
 #include "Common.ush"	
-#include "ShaderPrint.ush"
+#include "MiniFontCommon.ush"
+#include "ShaderPrintCommon.ush"
 
 uint2 ViewportResolution;
 uint FrameIndex;
@@ -309,14 +310,14 @@ void DrawSymbolsVS(
 	bool bIsValid = true;
 	if (Symbol.Type == SHADER_PRINT_TYPE_SLIDER)
 	{
-		const float2 NormalizedSize = GetSliderWidgetSize() / float2(ShaderPrintData.Resolution);
+		const float2 NormalizedSize = GetSliderWidgetSize(ShaderPrintData.FontSpacing, ShaderPrintData.Resolution) / float2(ShaderPrintData.Resolution);
 		const float2 Offset = float2(NormalizedSize.x - NormalizedSize.y/4.f, 0.f);
 		OutPosition.xy = VertPos.xy * NormalizedSize + Offset + InstancePos + DropShadowOffset;
 		bIsValid = !bDropShadow;
 	}
 	else if (Symbol.Type == SHADER_PRINT_TYPE_CHECK)
 	{
-		const float2 NormalizedSize = GetCheckboxWidgetSize() / float2(ShaderPrintData.Resolution);
+		const float2 NormalizedSize = GetCheckboxWidgetSize(ShaderPrintData.FontSpacing, ShaderPrintData.Resolution) / float2(ShaderPrintData.Resolution);
 		const float2 Offset = float2(NormalizedSize.x * 0.5f, 0.f);
 		OutPosition.xy = VertPos.xy * NormalizedSize + Offset + InstancePos + DropShadowOffset;
 		bIsValid = !bDropShadow;
diff --git a/Engine/Shaders/Private/ShaderPrintDrawPrimitive.usf b/Engine/Shaders/Private/ShaderPrintDrawPrimitive.usf
--- a/Engine/Shaders/Private/ShaderPrintDrawPrimitive.usf
+++ b/Engine/Shaders/Private/ShaderPrintDrawPrimitive.usf
@@ -2,8 +2,7 @@
 
 
 #include "Common.ush"
-
-#include "ShaderPrint.ush"
+#include "ShaderPrintCommon.ush"
 
 float3 TranslatedWorldOffsetConversion;
 uint bCheckerboardEnabled;
@@ -27,7 +26,7 @@ void ShaderDrawDebugVS(
 	out float4 Position : SV_POSITION,
 	out float4 ColorAlpha : TEXCOORD0)
 {
-	const FLineElement Item = UnpackLineElement(ShaderDrawDebugPrimitive, InstanceId);
+	const FLineElement Item = UnpackLineElement(ShaderDrawDebugPrimitive, InstanceId, ShaderPrintData.MaxValueCount);
 
 	float4 VertexPositionTWS = float4(VertexId == 0 ? Item.Pos0 : Item.Pos1, 1.0f);
 	float4 VertexColor = VertexId == 0 ? Item.Color0 : Item.Color1;
@@ -61,7 +60,7 @@ void ShaderDrawDebugVS(
 	out float4 Position : SV_POSITION,
 	out float4 ColorAlpha : TEXCOORD0)
 {
-	const FTriangleElement Item = UnpackTriangleElement(ShaderDrawDebugPrimitive, InstanceId);
+	const FTriangleElement Item = UnpackTriangleElement(ShaderDrawDebugPrimitive, InstanceId, ShaderPrintData.MaxValueCount, ShaderPrintData.MaxLineCount);
 
 	float4 VertexPositionTWS = float4(VertexId == 0 ? Item.Pos0 : (VertexId == 1 ? Item.Pos1 : Item.Pos2), 1.0f);
 	float4 VertexColor = Item.Color;
