Lumen Sofwatre Ray Tracing - fixed Global SDF object grid culling, which was incorrectly computing grid cell centers.

diff --git a/Engine/Shaders/Private/DistanceField/GlobalDistanceFieldCompositeObjects.usf b/Engine/Shaders/Private/DistanceField/GlobalDistanceFieldCompositeObjects.usf
--- a/Engine/Shaders/Private/DistanceField/GlobalDistanceFieldCompositeObjects.usf
+++ b/Engine/Shaders/Private/DistanceField/GlobalDistanceFieldCompositeObjects.usf
@@ -281,7 +281,7 @@ void CompositeObjectsIntoObjectGridPagesCS(
 	uint CullGridObjectNum = CullGridObjectHeader[CullGridLinearIndex * 2 + 0];
 	uint CullGridObjectOffset = CullGridObjectHeader[CullGridLinearIndex * 2 + 1];
 
-	float3 ObjectGridCellPageCoord = PageGridCoord * GLOBAL_DISTANCE_FIELD_PAGE_RESOLUTION + (CellCoordInPage + 0.5f) * (GLOBAL_DISTANCE_FIELD_PAGE_RESOLUTION / float(DISTANCE_FIELD_OBJECT_GRID_PAGE_RESOLUTION));
+	float3 ObjectGridCellPageCoord = PageGridCoord * GLOBAL_DISTANCE_FIELD_PAGE_RESOLUTION + CellCoordInPage * 2.0f + 0.5f - GLOBAL_DISTANCE_FIELD_PAGE_BORDER;
 	float3 ObjectGridCellTranslatedWorldCenter = ObjectGridCellPageCoord * PageCoordToVoxelTranslatedCenterScale + PageCoordToVoxelTranslatedCenterBias;
 	FLWCVector3 ObjectGridCellWorldCenter = LWCSubtract(ObjectGridCellTranslatedWorldCenter, PreViewTranslation);
 	float ObjectGridCellWorldExtent = ClipmapVoxelExtent * GLOBAL_DISTANCE_FIELD_PAGE_RESOLUTION / float(DISTANCE_FIELD_OBJECT_GRID_PAGE_RESOLUTION);
diff --git a/Engine/Shaders/Private/DistanceField/GlobalDistanceFieldHeightfields.usf b/Engine/Shaders/Private/DistanceField/GlobalDistanceFieldHeightfields.usf
--- a/Engine/Shaders/Private/DistanceField/GlobalDistanceFieldHeightfields.usf
+++ b/Engine/Shaders/Private/DistanceField/GlobalDistanceFieldHeightfields.usf
@@ -298,7 +298,7 @@ void CompositeHeightfieldsIntoObjectGridPagesCS(
 	uint3 PageTableTextureCoord = PageGridCoordToPageTableTextureCoord(PageGridCoord);
 	FGlobalDistanceFieldPage Page = UnpackGlobalDistanceFieldPage(PageTableLayerTexture.Load(int4(PageTableTextureCoord, 0)));
 
-	float3 ObjectGridCellPageCoord = PageGridCoord * GLOBAL_DISTANCE_FIELD_PAGE_RESOLUTION + (CellCoordInPage + 0.5f) * (GLOBAL_DISTANCE_FIELD_PAGE_RESOLUTION / float(DISTANCE_FIELD_OBJECT_GRID_PAGE_RESOLUTION));
+	float3 ObjectGridCellPageCoord = PageGridCoord * GLOBAL_DISTANCE_FIELD_PAGE_RESOLUTION + CellCoordInPage * 2.0f + 0.5f - GLOBAL_DISTANCE_FIELD_PAGE_BORDER;
 	float3 ObjectGridCellTranslatedWorldCenter = ObjectGridCellPageCoord * PageCoordToVoxelTranslatedCenterScale + PageCoordToVoxelTranslatedCenterBias;
 	FLWCVector3 ObjectGridCellWorldCenter = LWCSubtract(ObjectGridCellTranslatedWorldCenter, PreViewTranslation);
 	float ObjectGridCellWorldExtent = ClipmapVoxelExtent * GLOBAL_DISTANCE_FIELD_PAGE_RESOLUTION / float(DISTANCE_FIELD_OBJECT_GRID_PAGE_RESOLUTION);
