Fix compilation shader issue with FXC.

diff --git a/Engine/Shaders/Private/ShaderPrint.ush b/Engine/Shaders/Private/ShaderPrint.ush
--- a/Engine/Shaders/Private/ShaderPrint.ush
+++ b/Engine/Shaders/Private/ShaderPrint.ush
@@ -5,21 +5,26 @@
 
 #pragma once
 
-#include "ShaderPrintCommon.ush"
-#include "MiniFontCommon.ush"
-
-// The FXC compiler will force unroll on loops around Print() and so in some places we need
-// to limit the loop count to either fix compilation or reduce compilation time.
+// * The FXC compiler will force unroll on loops around Print() and so in some places we need
+//   to limit the loop count to either fix compilation or reduce compilation time.
+// * FXC compiler can't use buffer which are conditionally initialized
 #if COMPILER_HLSL && !COMPILER_DXC
 #define SHADER_PRINT_UNROLL(N) UNROLL_N(N)
+#define SHADER_PRINT_USE_GLOBAL_RESOURCE 1
+#define SHADER_PRINT_RWENTRYBUFFER(Ctx) ShaderPrint_RWEntryBuffer
+#define SHADER_PRINT_RWSTATEBUFFER(Ctx) ShaderPrint_StateBuffer
 #else
 #define SHADER_PRINT_UNROLL(N) LOOP
+#define SHADER_PRINT_USE_GLOBAL_RESOURCE 0
 #endif
 
 // Global buffers for ShaderPrint collection.
 StructuredBuffer<uint> ShaderPrint_StateBuffer;
 RWStructuredBuffer<uint> ShaderPrint_RWEntryBuffer;
 
+#include "ShaderPrintCommon.ush"
+#include "MiniFontCommon.ush"
+
 FShaderPrintBuffers InitShaderPrintContextBuffers()
 {
 	FShaderPrintBuffers Buffers;
diff --git a/Engine/Shaders/Private/ShaderPrintCommon.ush b/Engine/Shaders/Private/ShaderPrintCommon.ush
--- a/Engine/Shaders/Private/ShaderPrintCommon.ush
+++ b/Engine/Shaders/Private/ShaderPrintCommon.ush
@@ -146,6 +146,17 @@ void ClearCounters(RWStructuredBuffer<uint> InRWBuffer)
 #define SHADER_PRINT_STATE_INVALID_INDEX 0xFF
 #define SHADER_PRINT_STATE_HASH_MASK 0xFFFFFF
 
+///////////////////////////////////////////////////////////////////////////////////////////////////
+// Buffers/Resources accessors
+#ifndef SHADER_PRINT_USE_GLOBAL_RESOURCE
+#define SHADER_PRINT_USE_GLOBAL_RESOURCE 0
+#endif
+
+#if !SHADER_PRINT_USE_GLOBAL_RESOURCE
+#define SHADER_PRINT_RWENTRYBUFFER(Ctx) Ctx.Buffers.RWEntryBuffer
+#define SHADER_PRINT_RWSTATEBUFFER(Ctx) Ctx.Buffers.StateBuffer
+#endif
+
 ///////////////////////////////////////////////////////////////////////////////////////////////////
 // Font Color
 
@@ -324,7 +335,7 @@ void ShaderPrint_Internal(FShaderPrintContext Ctx, in FShaderPrintItem Item)
 
 	// Buffer counter is stored in first element .Value
 	int IndexToStore = 0;
-	InterlockedAdd(Ctx.Buffers.RWEntryBuffer[SHADER_PRINT_COUNTER_OFFSET_SYMBOL], 1, IndexToStore);
+	InterlockedAdd(SHADER_PRINT_RWENTRYBUFFER(Ctx)[SHADER_PRINT_COUNTER_OFFSET_SYMBOL], 1, IndexToStore);
 
 	// Prevent writing off the buffer
 	// Note that counter still increases so need clamp when reading it in later passes
@@ -333,7 +344,7 @@ void ShaderPrint_Internal(FShaderPrintContext Ctx, in FShaderPrintItem Item)
 		return;
 	}
 
-	WriteSymbol(IndexToStore, Item, Ctx.Buffers.RWEntryBuffer);
+	WriteSymbol(IndexToStore, Item, SHADER_PRINT_RWENTRYBUFFER(Ctx));
 }
 
 void ShaderPrint_Internal(FShaderPrintContext Ctx, in float2 ScreenPos, in int Value, in FFontColor FontColor, in uint Metadata, in int Type)
@@ -405,7 +416,7 @@ float2 ShaderPrintText_Internal(FShaderPrintContext Ctx, bool bIsActive, float2
 		// Prevent writing off the buffer
 		// Note that counter still increases so need clamp when reading it in later passes
 		int IndexToStore = 0;
-		InterlockedAdd(Ctx.Buffers.RWEntryBuffer[SHADER_PRINT_COUNTER_OFFSET_SYMBOL], Count, IndexToStore);
+		InterlockedAdd(SHADER_PRINT_RWENTRYBUFFER(Ctx)[SHADER_PRINT_COUNTER_OFFSET_SYMBOL], Count, IndexToStore);
 		if (uint(IndexToStore + Count) < Ctx.Config.MaxValueCount)
 		{
 			for (uint i = Begin; i < End; ++i)
@@ -417,7 +428,7 @@ float2 ShaderPrintText_Internal(FShaderPrintContext Ctx, bool bIsActive, float2
 				Item.Color = InColor.Color;
 				Item.Metadata = 0u;
 
-				WriteSymbol(IndexToStore, Item, Ctx.Buffers.RWEntryBuffer);
+				WriteSymbol(IndexToStore, Item, SHADER_PRINT_RWENTRYBUFFER(Ctx));
 
 				++IndexToStore;
 				InPos.x += Ctx.Config.FontSpacing.x;
@@ -528,15 +539,15 @@ bool ShaderPrintGetStateValue(FShaderPrintContext Ctx, in uint InHash, inout uin
 	OutValue = 0;
 
 	// Slow linear search among all widget states
-	const uint MaxCount = min(Ctx.Buffers.StateBuffer[SHADER_PRINT_STATE_COUNT_OFFSET], Ctx.Config.MaxStateCount);
+	const uint MaxCount = min(SHADER_PRINT_RWSTATEBUFFER(Ctx)[SHADER_PRINT_STATE_COUNT_OFFSET], Ctx.Config.MaxStateCount);
 	for (uint Index = 0; Index < MaxCount; ++Index)
 	{
 		const uint Index3 = Index * SHADER_PRINT_STATE_STRIDE + SHADER_PRINT_STATE_VALUE_OFFSET;
-		const FShaderPrintMetadata Metadata = ShaderPrintUnpackMetadata(Ctx.Buffers.StateBuffer[Index3 + SHADER_PRINT_STATE_INDEX_METADATA]);
+		const FShaderPrintMetadata Metadata = ShaderPrintUnpackMetadata(SHADER_PRINT_RWSTATEBUFFER(Ctx)[Index3 + SHADER_PRINT_STATE_INDEX_METADATA]);
 		if (Metadata.Hash == InHash)
 		{
 			OutIndex = Metadata.Index;
-			OutValue = Ctx.Buffers.StateBuffer[Index3 + SHADER_PRINT_STATE_INDEX_VALUE];
+			OutValue = SHADER_PRINT_RWSTATEBUFFER(Ctx)[Index3 + SHADER_PRINT_STATE_INDEX_VALUE];
 			return true;
 		}
 	}
@@ -743,7 +754,7 @@ bool AllocateTriangleElement(FShaderPrintContext Ctx, uint Count, inout uint Out
 		return false;
 	}
 
-	InterlockedAdd(Ctx.Buffers.RWEntryBuffer[SHADER_PRINT_COUNTER_OFFSET_TRIANGLE], Count, OutIndex);
+	InterlockedAdd(SHADER_PRINT_RWENTRYBUFFER(Ctx)[SHADER_PRINT_COUNTER_OFFSET_TRIANGLE], Count, OutIndex);
 	return (OutIndex + Count) < Ctx.Config.MaxTriangleCount;
 }
 
@@ -752,20 +763,20 @@ void AddTriangleElement(FShaderPrintContext Ctx, FTriangleElement In, uint Index
 	const uint Index12 = GetPrimitiveTriangleOffset(Index, Ctx.Config.MaxValueCount, Ctx.Config.MaxLineCount);
 
 	FPackedTriangleElement Out = PackTriangleElement(In);
-	Ctx.Buffers.RWEntryBuffer[Index12 + 0] = Out.Packed0.x;
-	Ctx.Buffers.RWEntryBuffer[Index12 + 1] = Out.Packed0.y;
-	Ctx.Buffers.RWEntryBuffer[Index12 + 2] = Out.Packed0.z;
-	Ctx.Buffers.RWEntryBuffer[Index12 + 3] = Out.Packed0.w;
-
-	Ctx.Buffers.RWEntryBuffer[Index12 + 4] = Out.Packed1.x;
-	Ctx.Buffers.RWEntryBuffer[Index12 + 5] = Out.Packed1.y;
-	Ctx.Buffers.RWEntryBuffer[Index12 + 6] = Out.Packed1.z;
-	Ctx.Buffers.RWEntryBuffer[Index12 + 7] = Out.Packed1.w;
-
-	Ctx.Buffers.RWEntryBuffer[Index12 + 8] = Out.Packed2.x;
-	Ctx.Buffers.RWEntryBuffer[Index12 + 9] = Out.Packed2.y;
-	Ctx.Buffers.RWEntryBuffer[Index12 + 10] = Out.Packed2.z;
-	Ctx.Buffers.RWEntryBuffer[Index12 + 11] = Out.Packed2.w;
+	SHADER_PRINT_RWENTRYBUFFER(Ctx)[Index12 + 0] = Out.Packed0.x;
+	SHADER_PRINT_RWENTRYBUFFER(Ctx)[Index12 + 1] = Out.Packed0.y;
+	SHADER_PRINT_RWENTRYBUFFER(Ctx)[Index12 + 2] = Out.Packed0.z;
+	SHADER_PRINT_RWENTRYBUFFER(Ctx)[Index12 + 3] = Out.Packed0.w;
+
+	SHADER_PRINT_RWENTRYBUFFER(Ctx)[Index12 + 4] = Out.Packed1.x;
+	SHADER_PRINT_RWENTRYBUFFER(Ctx)[Index12 + 5] = Out.Packed1.y;
+	SHADER_PRINT_RWENTRYBUFFER(Ctx)[Index12 + 6] = Out.Packed1.z;
+	SHADER_PRINT_RWENTRYBUFFER(Ctx)[Index12 + 7] = Out.Packed1.w;
+
+	SHADER_PRINT_RWENTRYBUFFER(Ctx)[Index12 + 8] = Out.Packed2.x;
+	SHADER_PRINT_RWENTRYBUFFER(Ctx)[Index12 + 9] = Out.Packed2.y;
+	SHADER_PRINT_RWENTRYBUFFER(Ctx)[Index12 + 10] = Out.Packed2.z;
+	SHADER_PRINT_RWENTRYBUFFER(Ctx)[Index12 + 11] = Out.Packed2.w;
 }
 
 ///////////////////////////////////////////////////////////////////////////////////////////////////
@@ -942,7 +953,7 @@ bool AllocateLineElement(FShaderPrintContext Ctx, uint Count, inout uint OutInde
 		return false;
 	}
 
-	InterlockedAdd(Ctx.Buffers.RWEntryBuffer[SHADER_PRINT_COUNTER_OFFSET_LINE], Count, OutIndex);
+	InterlockedAdd(SHADER_PRINT_RWENTRYBUFFER(Ctx)[SHADER_PRINT_COUNTER_OFFSET_LINE], Count, OutIndex);
 	return (OutIndex + Count) < Ctx.Config.MaxLineCount;
 }
 
@@ -951,15 +962,15 @@ void AddLineElement(FShaderPrintContext Ctx, FLineElement In, uint Index)
 	const uint Index8 = GetPrimitiveLineOffset(Index, Ctx.Config.MaxValueCount);
 
 	FPackedLineElement Out = PackLineElement(In);
-	Ctx.Buffers.RWEntryBuffer[Index8 + 0] = Out.Packed0.x;
-	Ctx.Buffers.RWEntryBuffer[Index8 + 1] = Out.Packed0.y;
-	Ctx.Buffers.RWEntryBuffer[Index8 + 2] = Out.Packed0.z;
-	Ctx.Buffers.RWEntryBuffer[Index8 + 3] = Out.Packed0.w;
-
-	Ctx.Buffers.RWEntryBuffer[Index8 + 4] = Out.Packed1.x;
-	Ctx.Buffers.RWEntryBuffer[Index8 + 5] = Out.Packed1.y;
-	Ctx.Buffers.RWEntryBuffer[Index8 + 6] = Out.Packed1.z;
-	Ctx.Buffers.RWEntryBuffer[Index8 + 7] = Out.Packed1.w;
+	SHADER_PRINT_RWENTRYBUFFER(Ctx)[Index8 + 0] = Out.Packed0.x;
+	SHADER_PRINT_RWENTRYBUFFER(Ctx)[Index8 + 1] = Out.Packed0.y;
+	SHADER_PRINT_RWENTRYBUFFER(Ctx)[Index8 + 2] = Out.Packed0.z;
+	SHADER_PRINT_RWENTRYBUFFER(Ctx)[Index8 + 3] = Out.Packed0.w;
+
+	SHADER_PRINT_RWENTRYBUFFER(Ctx)[Index8 + 4] = Out.Packed1.x;
+	SHADER_PRINT_RWENTRYBUFFER(Ctx)[Index8 + 5] = Out.Packed1.y;
+	SHADER_PRINT_RWENTRYBUFFER(Ctx)[Index8 + 6] = Out.Packed1.z;
+	SHADER_PRINT_RWENTRYBUFFER(Ctx)[Index8 + 7] = Out.Packed1.w;
 }
 
 ///////////////////////////////////////////////////////////////////////////////////////////////////
