Added option to not clamp the multiple scattering contribution factor of cloud.Artist is then in charge of the visual brightness.

diff --git a/Engine/Shaders/Private/VolumetricCloud.usf b/Engine/Shaders/Private/VolumetricCloud.usf
--- a/Engine/Shaders/Private/VolumetricCloud.usf
+++ b/Engine/Shaders/Private/VolumetricCloud.usf
@@ -259,6 +259,9 @@ bool RayIntersectSphereSolution(float3 RayOrigin, float3 RayDirection, float4 Sp
 #ifndef MATERIAL_VOLUMETRIC_ADVANCED_RAYMARCH_VOLUME_SHADOW
 #define MATERIAL_VOLUMETRIC_ADVANCED_RAYMARCH_VOLUME_SHADOW 0
 #endif
+#ifndef MATERIAL_VOLUMETRIC_ADVANCED_CLAMP_MULTISCATTERING_CONTRIBUTION
+#define MATERIAL_VOLUMETRIC_ADVANCED_CLAMP_MULTISCATTERING_CONTRIBUTION 1
+#endif
 #ifndef MATERIAL_VOLUMETRIC_ADVANCED_OVERRIDE_AMBIENT_OCCLUSION
 #define MATERIAL_VOLUMETRIC_ADVANCED_OVERRIDE_AMBIENT_OCCLUSION 0
 #endif
@@ -574,7 +577,11 @@ void MainCommon(in FMaterialPixelParameters MaterialParameters, in float4 SvPosi
 
 		// Evalaute Factors once per pixel
 #if MATERIAL_VOLUMETRIC_ADVANCED
+	#if MATERIAL_VOLUMETRIC_ADVANCED_CLAMP_MULTISCATTERING_CONTRIBUTION
 		const float MsScattFactor = saturate(GetVolumetricAdvancedMaterialOutput3(MaterialParameters));
+	#else
+		const float MsScattFactor = GetVolumetricAdvancedMaterialOutput3(MaterialParameters);
+	#endif
 		const float MsExtinFactor = saturate(GetVolumetricAdvancedMaterialOutput4(MaterialParameters));
 		const float MsPhaseFactor = saturate(GetVolumetricAdvancedMaterialOutput5(MaterialParameters));
 #else
diff --git a/Engine/Source/Runtime/Engine/Classes/Materials/MaterialExpressionVolumetricAdvancedMaterialOutput.h b/Engine/Source/Runtime/Engine/Classes/Materials/MaterialExpressionVolumetricAdvancedMaterialOutput.h
--- a/Engine/Source/Runtime/Engine/Classes/Materials/MaterialExpressionVolumetricAdvancedMaterialOutput.h
+++ b/Engine/Source/Runtime/Engine/Classes/Materials/MaterialExpressionVolumetricAdvancedMaterialOutput.h
@@ -88,6 +88,10 @@ class UMaterialExpressionVolumetricAdvancedMaterialOutput : public UMaterialExpr
 	UPROPERTY(EditAnywhere, Category = "Options")
 	bool bRayMarchVolumeShadow;
 
+	/** Set whether multiple scattering contribution entry is clamped in [0,1] or not. When disabled, the artist is in charge for ensuring the visual remain in a reasonable brighness range.*/
+	UPROPERTY(EditAnywhere, Category = "Options")
+	bool bClampMultiScatteringContribution;
+
 public:
 #if WITH_EDITOR
 	//~ Begin UMaterialExpression Interface
diff --git a/Engine/Source/Runtime/Engine/Private/Materials/HLSLMaterialTranslator.cpp b/Engine/Source/Runtime/Engine/Private/Materials/HLSLMaterialTranslator.cpp
--- a/Engine/Source/Runtime/Engine/Private/Materials/HLSLMaterialTranslator.cpp
+++ b/Engine/Source/Runtime/Engine/Private/Materials/HLSLMaterialTranslator.cpp
@@ -1908,6 +1908,7 @@ void FHLSLMaterialTranslator::GetMaterialEnvironment(EShaderPlatform InPlatform,
 
 			OutEnvironment.SetDefine(TEXT("MATERIAL_VOLUMETRIC_ADVANCED_GRAYSCALE_MATERIAL"), VolumetricAdvancedNode->bGrayScaleMaterial ? TEXT("1") : TEXT("0"));
 			OutEnvironment.SetDefine(TEXT("MATERIAL_VOLUMETRIC_ADVANCED_RAYMARCH_VOLUME_SHADOW"), VolumetricAdvancedNode->bRayMarchVolumeShadow ? TEXT("1") : TEXT("0"));
+			OutEnvironment.SetDefine(TEXT("MATERIAL_VOLUMETRIC_ADVANCED_CLAMP_MULTISCATTERING_CONTRIBUTION"), VolumetricAdvancedNode->bClampMultiScatteringContribution ? TEXT("1") : TEXT("0"));
 
 			OutEnvironment.SetDefine(TEXT("MATERIAL_VOLUMETRIC_ADVANCED_MULTISCATTERING_OCTAVE_COUNT"), VolumetricAdvancedNode->GetMultiScatteringApproximationOctaveCount());
 
diff --git a/Engine/Source/Runtime/Engine/Private/Materials/MaterialExpressions.cpp b/Engine/Source/Runtime/Engine/Private/Materials/MaterialExpressions.cpp
--- a/Engine/Source/Runtime/Engine/Private/Materials/MaterialExpressions.cpp
+++ b/Engine/Source/Runtime/Engine/Private/Materials/MaterialExpressions.cpp
@@ -21225,6 +21225,7 @@ UMaterialExpressionVolumetricAdvancedMaterialOutput::UMaterialExpressionVolumetr
 	bGroundContribution = false;
 	bGrayScaleMaterial = false;
 	bRayMarchVolumeShadow = true;
+	bClampMultiScatteringContribution = true;
 
 #if WITH_EDITORONLY_DATA
 	MenuCategories.Add(ConstructorStatics.NAME_VolumetricAdvancedOutput);
diff --git a/Engine/Source/Runtime/Engine/Private/Materials/MaterialHLSLEmitter.cpp b/Engine/Source/Runtime/Engine/Private/Materials/MaterialHLSLEmitter.cpp
--- a/Engine/Source/Runtime/Engine/Private/Materials/MaterialHLSLEmitter.cpp
+++ b/Engine/Source/Runtime/Engine/Private/Materials/MaterialHLSLEmitter.cpp
@@ -672,6 +672,7 @@ static void GetMaterialEnvironment(EShaderPlatform InPlatform,
 
 			OutEnvironment.SetDefine(TEXT("MATERIAL_VOLUMETRIC_ADVANCED_GRAYSCALE_MATERIAL"), VolumetricAdvancedNode->bGrayScaleMaterial ? TEXT("1") : TEXT("0"));
 			OutEnvironment.SetDefine(TEXT("MATERIAL_VOLUMETRIC_ADVANCED_RAYMARCH_VOLUME_SHADOW"), VolumetricAdvancedNode->bRayMarchVolumeShadow ? TEXT("1") : TEXT("0"));
+			OutEnvironment.SetDefine(TEXT("MATERIAL_VOLUMETRIC_ADVANCED_CLAMP_MULTISCATTERING_CONTRIBUTION"), VolumetricAdvancedNode->bClampMultiScatteringContribution ? TEXT("1") : TEXT("0"));
 
 			OutEnvironment.SetDefine(TEXT("MATERIAL_VOLUMETRIC_ADVANCED_MULTISCATTERING_OCTAVE_COUNT"), VolumetricAdvancedNode->GetMultiScatteringApproximationOctaveCount());
 
