RayTracingDebug - make sure AABB extents are non-zero to prevent flickering.

diff --git a/Engine/Shaders/Private/RayTracing/RayTracingDebugInstanceOverlap.usf b/Engine/Shaders/Private/RayTracing/RayTracingDebugInstanceOverlap.usf
--- a/Engine/Shaders/Private/RayTracing/RayTracingDebugInstanceOverlap.usf
+++ b/Engine/Shaders/Private/RayTracing/RayTracingDebugInstanceOverlap.usf
@@ -74,7 +74,9 @@ void InstanceOverlapMainVS(
 	{
 		const FInstanceBounds InstanceBounds = GetTranslatedWorldAABBForInstance(GPUSceneIndex);
 
-		const float3 TranslatedWorldPosition = InPosition.xyz * InstanceBounds.Extent * BoundingBoxExtentScale + InstanceBounds.Center;
+		// Make sure extent is non-zero to avoid flickering.
+		const float3 Extent = max(InstanceBounds.Extent, BoundingBoxExtentScale);
+		const float3 TranslatedWorldPosition = InPosition.xyz * Extent * BoundingBoxExtentScale + InstanceBounds.Center;
 
 		OutPosition = mul(float4(TranslatedWorldPosition, 1.0), PrimaryView.TranslatedWorldToClip);
 	}
