Fix Strata crash when enabling material advanced debugging introduced by fixing another merging error in CL 22331973

diff --git a/Engine/Shaders/Private/BasePassPixelShader.usf b/Engine/Shaders/Private/BasePassPixelShader.usf
--- a/Engine/Shaders/Private/BasePassPixelShader.usf
+++ b/Engine/Shaders/Private/BasePassPixelShader.usf
@@ -1644,7 +1644,7 @@ void FPixelShaderInOut_MainPS(
 		FStrataIntegrationSettings Settings = InitStrataIntegrationSettings(false /*bForceFullyRough*/, StrataStruct.bRoughDiffuse, StrataStruct.PeelLayersAboveDepth);
 
 	#if STRATA_ADVANCED_DEBUG_ENABLED
-		Settings.SliceStoringDebugStrataTreeData = all(uint2(float2(View.CursorPosition) * View.ViewResolutionFraction) == PixelPos) ? StrataStruct.SliceStoringDebugStrataTree : Settings.SliceStoringDebugStrataTree;
+		Settings.SliceStoringDebugStrataTreeData = all(uint2(float2(View.CursorPosition) * View.ViewResolutionFraction) == PixelPos) ? StrataStruct.SliceStoringDebugStrataTreeDataWithoutMRT : Settings.SliceStoringDebugStrataTreeData;
 	#endif
 
 		// Generate the strata material data to write out
