Fix artifacts caused by DFAO UpdateHistoryStability pass.
- In some situations history was updated even though none of the neighbours was relevant due to float precision issues.

diff --git a/Engine/Shaders/Private/DistanceFieldLightingPost.usf b/Engine/Shaders/Private/DistanceFieldLightingPost.usf
--- a/Engine/Shaders/Private/DistanceFieldLightingPost.usf
+++ b/Engine/Shaders/Private/DistanceFieldLightingPost.usf
@@ -309,7 +309,6 @@ void FilterHistoryPS(
 		float SceneDepth = abs(HistoryValue.w);
 
 		float4 Accumulation = 0;
-		float3 IrradianceAccumulation = 0;
 
 		for (float y = -HALF_HISTORY_FILL_KERNEL_SIZE; y <= HALF_HISTORY_FILL_KERNEL_SIZE; y++)
 		{
@@ -342,7 +341,7 @@ void FilterHistoryPS(
 		}
 
 		// Only change the history value if the spatial search turned up something applicable
-		if (Accumulation.a > 0)
+		if (Accumulation.a > 0.01f)
 		{
 			float InvWeight = 1.0f / Accumulation.a;
 			// Construct the history value as if the spatial search result was the previous history, 
