Backout: Fixes preexposure correction of TSR's history guide

diff --git a/Engine/Shaders/Private/TemporalSuperResolution/TSRDecimateHistory.usf b/Engine/Shaders/Private/TemporalSuperResolution/TSRDecimateHistory.usf
--- a/Engine/Shaders/Private/TemporalSuperResolution/TSRDecimateHistory.usf
+++ b/Engine/Shaders/Private/TemporalSuperResolution/TSRDecimateHistory.usf
@@ -611,12 +611,6 @@ void MainCS(
 	}
 	#endif // CONFIG_GRAND_REPROJECTION
 
-	// Correct exposure of ReprojectedHistoryGuideColor
-	{
-		tsr_halfC ReprojectedHistoryGuideLinearColor = ReprojectedHistoryGuideColor * rcp(tsr_half(1.0) - ReprojectedHistoryGuideColor) * tsr_half(HistoryPreExposureCorrection);
-		ReprojectedHistoryGuideColor = ReprojectedHistoryGuideLinearColor * rcp(tsr_half(1.0) + ReprojectedHistoryGuideLinearColor);
-	}
-
 	// Output full res
 	ISOLATE
 	{
