Fix for SingleLayerWater UBO crash. Fixed by properly guarding SingleLayerWater UBO access with macros.

#rnx

diff --git a/Engine/Shaders/Private/BasePassPixelShader.usf b/Engine/Shaders/Private/BasePassPixelShader.usf
--- a/Engine/Shaders/Private/BasePassPixelShader.usf
+++ b/Engine/Shaders/Private/BasePassPixelShader.usf
@@ -1177,7 +1177,6 @@ void FPixelShaderInOut_MainPS(
 	float3 DiffuseIndirectLighting = 0;
 	float3 SubsurfaceIndirectLighting = 0;
 
-	bool bSeparateWaterMainDirLightLuminance = SINGLE_LAYER_WATER_SEPARATED_MAIN_LIGHT > 0 && SingleLayerWater.bSeparateMainDirLightLuminance; 
 	float3 SeparatedWaterMainDirLightLuminance = float3(0, 0, 0); 
 
 #if !STRATA_ENABLED || STRATA_INLINE_SINGLELAYERWATER
@@ -1278,7 +1277,13 @@ void FPixelShaderInOut_MainPS(
 					}
 					#endif // NEEDS_BASEPASS_PIXEL_FOGGING && PROJECT_SUPPORT_SKY_ATMOSPHERE
 
+					#if MATERIAL_SHADINGMODEL_SINGLELAYERWATER
 					const bool bSkipMainDirLightVirtualShadowMapEvaluation = SingleLayerWater.bMainDirectionalLightVSMFiltering;
+					const bool bSeparateWaterMainDirLightLuminance = (SINGLE_LAYER_WATER_SEPARATED_MAIN_LIGHT > 0) && SingleLayerWater.bSeparateMainDirLightLuminance;
+					#else
+					const bool bSkipMainDirLightVirtualShadowMapEvaluation = false;
+					const bool bSeparateWaterMainDirLightLuminance = false;
+					#endif
 
 					FDeferredLightingSplit ForwardDirectLighting = GetForwardDirectLightingSplit(
 						GridIndex, MaterialParameters.WorldPosition_CamRelative, MaterialParameters.CameraVector, GBuffer, NearestResolvedDepthScreenUV, MaterialParameters.PrimitiveId, EyeIndex, Dither, 
@@ -1453,6 +1458,7 @@ void FPixelShaderInOut_MainPS(
 		const bool CameraIsUnderWater = false;	// Fade out the material contribution over to water contribution according to material opacity.
 		const float3 SunIlluminance = ResolvedView.DirectionalLightColor.rgb * PI;	// times PI because it is divided by PI on CPU (=luminance) and we want illuminance here. 
 		const float3 WaterDiffuseIndirectIlluminance = DiffuseIndirectLighting * PI;// DiffuseIndirectLighting is luminance. So we need to multiply by PI to get illuminance.
+		const bool bSeparateWaterMainDirLightLuminance = (SINGLE_LAYER_WATER_SEPARATED_MAIN_LIGHT > 0) && SingleLayerWater.bSeparateMainDirLightLuminance;
 
 		// Evaluate Fresnel effect
 		const float3 N = MaterialParameters.WorldNormal;
