Fix for NaN on SingleLayerWater when water is in front of the far plane and extinction is 0.0f. Fixed by ensuring SceneDeviceZ is slightly greater than 0.0f (far plane with inverted Z).

#rnx

diff --git a/Engine/Shaders/Private/SingleLayerWaterShading.ush b/Engine/Shaders/Private/SingleLayerWaterShading.ush
--- a/Engine/Shaders/Private/SingleLayerWaterShading.ush
+++ b/Engine/Shaders/Private/SingleLayerWaterShading.ush
@@ -73,10 +73,6 @@ WaterVolumeLightingOutput EvaluateWaterVolumeLighting(
 		float PixelSceneDeviceZ = LookupDeviceZ(ViewportUV);
 		float PixelSceneDepth = ConvertFromDeviceZ(PixelSceneDeviceZ);
 #endif
-		#if HAS_INVERTED_Z_BUFFER
-		// SvPositionToWorld will result in NaN if PixelSceneDeviceZ=0
-		PixelSceneDeviceZ = max(0.000000000001, PixelSceneDeviceZ);
-		#endif
 		float SceneDeviceZ = PixelSceneDeviceZ;
 		float SceneDepth = PixelSceneDepth;
 
@@ -126,6 +122,11 @@ WaterVolumeLightingOutput EvaluateWaterVolumeLighting(
 
 #endif // !SINGLE_LAYER_WATER_MOBILE_SHADING
 
+#if HAS_INVERTED_Z_BUFFER
+		// SvPositionToWorld will result in NaN if SceneDeviceZ=0
+		SceneDeviceZ = max(0.000000000001f, SceneDeviceZ);
+#endif
+
 		const float BehindWaterDeltaDepth = CameraIsUnderWater ? WaterDepth : max(0.0f, SceneDepth - WaterDepth);
 
 #if STRATA_INLINE_SINGLELAYERWATER
