Volumetric cloud - added tracing max distance mode.
Wil lhelp to better control tracing distance and avoid top of the cloud being eaten by too short tracing distance from entry point at tangent view direction.

diff --git a/Engine/Shaders/Private/VolumetricCloud.usf b/Engine/Shaders/Private/VolumetricCloud.usf
--- a/Engine/Shaders/Private/VolumetricCloud.usf
+++ b/Engine/Shaders/Private/VolumetricCloud.usf
@@ -542,9 +542,25 @@ void MainCommon(in FMaterialPixelParameters MaterialParameters, in float4 SvPosi
 
 		const bool SkyAtmospherePresentInScene = ResolvedView.SkyAtmospherePresentInScene > 0.0f;
 
-		// Clamp to max distance
-		const float MarchingDistance = min(RenderVolumetricCloudParameters.TracingMaxDistance, TMax - TMin);
-		TMax = TMin + MarchingDistance;
+		// Clamp to tracing max distance according to the selected mode
+		if (RenderVolumetricCloudParameters.TracingMaxDistanceMode == 0)
+		{
+			// Distance after the cloud layer entry point
+			const float MarchingDistance = min(RenderVolumetricCloudParameters.TracingMaxDistance, TMax - TMin);
+			TMax = TMin + MarchingDistance;
+		}
+		else // if (RenderVolumetricCloudParameters.TracingMaxDistanceMode == 1)
+		{
+			// Distance from the point of view
+			TMin = min(TMin, RenderVolumetricCloudParameters.TracingMaxDistance);
+			TMax = min(TMax, RenderVolumetricCloudParameters.TracingMaxDistance);
+
+			// Exit if tracing a 0 length segment.
+			if ((TMax - TMin) <= 0.0f)
+			{
+				return;
+			}
+		}
 
 		// When we do trace, we want a minimum amount of sample count. This avoid situation with 0 sample count due to InvDistanceToSampleCountMax.
 		const uint  IStepCount = max(RenderVolumetricCloudParameters.SampleCountMin, RenderVolumetricCloudParameters.SampleCountMax * saturate((TMax - TMin) * RenderVolumetricCloudParameters.InvDistanceToSampleCountMax));
diff --git a/Engine/Source/Runtime/Engine/Classes/Components/VolumetricCloudComponent.h b/Engine/Source/Runtime/Engine/Classes/Components/VolumetricCloudComponent.h
--- a/Engine/Source/Runtime/Engine/Classes/Components/VolumetricCloudComponent.h
+++ b/Engine/Source/Runtime/Engine/Classes/Components/VolumetricCloudComponent.h
@@ -14,6 +14,12 @@
 
 class FVolumetricCloudSceneProxy;
 
+UENUM()
+enum class EVolumetricCloudTracingMaxDistanceMode : uint8
+{
+	DistanceFromCloudLayerEntryPoint = 0,
+	DistanceFromPointOfView = 1,
+};
 
 /**
  * A component that represents a participating media material around a planet, e.g. clouds.
@@ -37,6 +43,10 @@ class UVolumetricCloudComponent : public USceneComponent
 	UPROPERTY(EditAnywhere, BlueprintReadOnly, interp, Category = "Layer", meta = (UIMin = 100.0f, UIMax = 500.0f, ClampMin = 1.0f, SliderExponent = 2.0))
 	float TracingStartMaxDistance;
 
+	/** Mode to select how the tracing max distance should be interpreted. DistanceFromPointOfView is useful to avoid the top of the cloud layer to be clipped when TracingMaxDistance is shorten for performance. */
+	UPROPERTY(EditAnywhere, BlueprintReadOnly, interp, Category = "Layer")
+	EVolumetricCloudTracingMaxDistanceMode TracingMaxDistanceMode;
+
 	/** The maximum distance that will be traced inside the cloud layer. (kilometers) */
 	UPROPERTY(EditAnywhere, BlueprintReadOnly, interp, Category = "Layer", meta = (UIMin = 1.0f, UIMax = 500.0f, ClampMin = 0.1f, SliderExponent = 2.0))
 	float TracingMaxDistance;
diff --git a/Engine/Source/Runtime/Engine/Private/Components/VolumetricCloudComponent.cpp b/Engine/Source/Runtime/Engine/Private/Components/VolumetricCloudComponent.cpp
--- a/Engine/Source/Runtime/Engine/Private/Components/VolumetricCloudComponent.cpp
+++ b/Engine/Source/Runtime/Engine/Private/Components/VolumetricCloudComponent.cpp
@@ -36,6 +36,7 @@ UVolumetricCloudComponent::UVolumetricCloudComponent(const FObjectInitializer& O
 	, LayerBottomAltitude(5.0f)
 	, LayerHeight(10.0f)
 	, TracingStartMaxDistance(350.0f)
+	, TracingMaxDistanceMode(EVolumetricCloudTracingMaxDistanceMode::DistanceFromCloudLayerEntryPoint)
 	, TracingMaxDistance(50.0f)
 	, PlanetRadius(6360.0f)					// Default to earth-like
 	, GroundAlbedo(FColor(170, 170, 170))	// 170 => 0.4f linear
diff --git a/Engine/Source/Runtime/Engine/Private/VolumetricCloudProxy.cpp b/Engine/Source/Runtime/Engine/Private/VolumetricCloudProxy.cpp
--- a/Engine/Source/Runtime/Engine/Private/VolumetricCloudProxy.cpp
+++ b/Engine/Source/Runtime/Engine/Private/VolumetricCloudProxy.cpp
@@ -13,6 +13,7 @@ FVolumetricCloudSceneProxy::FVolumetricCloudSceneProxy(const UVolumetricCloudCom
 	: LayerBottomAltitudeKm(InComponent->LayerBottomAltitude)
 	, LayerHeightKm(InComponent->LayerHeight)
 	, TracingStartMaxDistance(InComponent->TracingStartMaxDistance)
+	, TracingMaxDistanceMode(uint8(InComponent->TracingMaxDistanceMode))
 	, TracingMaxDistance(InComponent->TracingMaxDistance)
 	, PlanetRadiusKm(InComponent->PlanetRadius)
 	, GroundAlbedo(InComponent->GroundAlbedo)
diff --git a/Engine/Source/Runtime/Engine/Public/VolumetricCloudProxy.h b/Engine/Source/Runtime/Engine/Public/VolumetricCloudProxy.h
--- a/Engine/Source/Runtime/Engine/Public/VolumetricCloudProxy.h
+++ b/Engine/Source/Runtime/Engine/Public/VolumetricCloudProxy.h
@@ -33,6 +33,7 @@ public:
 	float LayerHeightKm;
 
 	float TracingStartMaxDistance;
+	uint8 TracingMaxDistanceMode;
 	float TracingMaxDistance;
 
 	float PlanetRadiusKm;
diff --git a/Engine/Source/Runtime/Renderer/Private/VolumetricCloudRendering.cpp b/Engine/Source/Runtime/Renderer/Private/VolumetricCloudRendering.cpp
--- a/Engine/Source/Runtime/Renderer/Private/VolumetricCloudRendering.cpp
+++ b/Engine/Source/Runtime/Renderer/Private/VolumetricCloudRendering.cpp
@@ -1460,6 +1460,7 @@ void FSceneRenderer::InitVolumetricCloudsForViews(FRDGBuilder& GraphBuilder, boo
 			CloudGlobalShaderParams.SkyLightCloudBottomVisibility = 1.0f - CloudProxy.SkyLightCloudBottomOcclusion;
 
 			CloudGlobalShaderParams.TracingStartMaxDistance = KilometersToCentimeters * CloudProxy.TracingStartMaxDistance;
+			CloudGlobalShaderParams.TracingMaxDistanceMode	= CloudProxy.TracingMaxDistanceMode;
 			CloudGlobalShaderParams.TracingMaxDistance		= KilometersToCentimeters * CloudProxy.TracingMaxDistance;
 
 			CloudGlobalShaderParams.SampleCountMin		= FMath::Max(0, CVarVolumetricCloudSampleMinCount.GetValueOnAnyThread());
diff --git a/Engine/Source/Runtime/Renderer/Private/VolumetricCloudRendering.h b/Engine/Source/Runtime/Renderer/Private/VolumetricCloudRendering.h
--- a/Engine/Source/Runtime/Renderer/Private/VolumetricCloudRendering.h
+++ b/Engine/Source/Runtime/Renderer/Private/VolumetricCloudRendering.h
@@ -31,6 +31,7 @@ BEGIN_SHADER_PARAMETER_STRUCT(FVolumetricCloudCommonShaderParameters, )
 	SHADER_PARAMETER(float, BottomRadiusKm)
 	SHADER_PARAMETER(float, TopRadiusKm)
 	SHADER_PARAMETER(float, TracingStartMaxDistance)
+	SHADER_PARAMETER(int32, TracingMaxDistanceMode)
 	SHADER_PARAMETER(float, TracingMaxDistance)
 	SHADER_PARAMETER(int32, SampleCountMin)
 	SHADER_PARAMETER(int32, SampleCountMax)
