Rearrange HairInterpolation shader code for easing future deformation update.

diff --git a/Engine/Shaders/Private/HairStrands/HairStrandsInterpolation.usf b/Engine/Shaders/Private/HairStrands/HairStrandsInterpolation.usf
--- a/Engine/Shaders/Private/HairStrands/HairStrandsInterpolation.usf
+++ b/Engine/Shaders/Private/HairStrands/HairStrandsInterpolation.usf
@@ -7,13 +7,6 @@
 #include "HairStrandsClusterCommon.ush"
 #include "HairStrandsVertexFactoryCommon.ush"
 
-// Should be moved into HairStrandsVertexFactoryCommon
-float UnpackHairControlPointCoordU(uint InPackedDataW)
-{
-	const uint PackedAlpha = asuint(InPackedDataW);
-	return ((PackedAlpha & 0xFF00) >> 8) / 255.f;
-}
-
 #if COMPILER_METAL || COMPILER_VULKAN
 #define ALLOW_DEBUG_RENDERING 0
 #else
@@ -127,21 +120,30 @@ Buffer<float4> OutRenHairPositionOffsetBuffer;
 	#error Unknown group size	
 #endif
 
-float3 ComputeStaticGeometryOffset(uint GuideIndex, float GuideVertexWeight, float3 OutSimHairPositionOffset)
+///////////////////////////////////////////////////////////////////////////////////////////////////
+// Deformation offset
+
+float3 ComputeStaticGeometryOffset(
+	uint GuideIndex, 
+	float GuideVertexWeight, 
+	float3 OutSimHairPositionOffset)
 {
 	const float3 RestGuidePoint = f16tof32(SimRestPosePositionBuffer[GuideIndex].xyz) + InSimHairPositionOffset;
 
 	const float3 DeformedGuidePoint = f16tof32(DeformedSimPositionBuffer[GuideIndex].xyz) + OutSimHairPositionOffset;
+
 	return (DeformedGuidePoint - RestGuidePoint) * GuideVertexWeight;
 }
 
-float3 ComputeLocalGeometryOffset(uint GuideIndex, float GuideVertexWeight)
-{	
-	const float3 RestGuidePoint = f16tof32(SimRestPosePositionBuffer[GuideIndex].xyz) - 
-					f16tof32(SimRestPosePositionBuffer[SimRootPointIndexBuffer[GuideIndex]].xyz) ;
+float3 ComputeLocalGeometryOffset(
+	uint GuideIndex, 
+	float GuideVertexWeight)
+{
+	const uint RootIndex = SimRootPointIndexBuffer[GuideIndex];
+
+	const float3 RestGuidePoint = f16tof32(SimRestPosePositionBuffer[GuideIndex].xyz) - f16tof32(SimRestPosePositionBuffer[RootIndex].xyz) ;
 
-	const float3 DeformedGuidePoint = f16tof32(DeformedSimPositionBuffer[GuideIndex].xyz) - 
-					f16tof32(DeformedSimPositionBuffer[SimRootPointIndexBuffer[GuideIndex]].xyz) ;
+	const float3 DeformedGuidePoint = f16tof32(DeformedSimPositionBuffer[GuideIndex].xyz) - f16tof32(DeformedSimPositionBuffer[RootIndex].xyz) ;
 
 	return (DeformedGuidePoint - RestGuidePoint) * GuideVertexWeight;
 }
@@ -163,46 +165,38 @@ float3 ComputeDynamicGeometryOffset(
 	return (LocalDeformedGuidePoint - LocalRestGuidePoint) * GuideVertexWeight;
 }
 
-float3 ComputeTranslateGeometryOffset(uint GuideIndex, float GuideVertexWeight)
-{	
-	const float3 RestGuidePoint = f16tof32(SimRestPosePositionBuffer[GuideIndex].xyz) - 
-					f16tof32(SimRestPosePositionBuffer[SimRootPointIndexBuffer[GuideIndex]].xyz) ;
+float3 ComputeTranslateGeometryOffset(
+	uint GuideIndex, 
+	float GuideVertexWeight)
+{
+	const uint RootIndex = SimRootPointIndexBuffer[GuideIndex];
+
+	const float3 RestGuidePoint = f16tof32(SimRestPosePositionBuffer[GuideIndex].xyz) - f16tof32(SimRestPosePositionBuffer[RootIndex].xyz) ;
 
-	const float3 DeformedGuidePoint = f16tof32(DeformedSimPositionBuffer[GuideIndex].xyz) - 
-					f16tof32(DeformedSimPositionBuffer[SimRootPointIndexBuffer[GuideIndex]].xyz) ;
+	const float3 DeformedGuidePoint = f16tof32(DeformedSimPositionBuffer[GuideIndex].xyz) - f16tof32(DeformedSimPositionBuffer[RootIndex].xyz) ;
 
 	return (DeformedGuidePoint - RestGuidePoint) * GuideVertexWeight;
 }
 
-float3 ComputeRotateGeometryOffset(uint GuideIndex, float GuideVertexWeight,
+float3 ComputeRotateGeometryOffset(
+	uint GuideIndex, 
+	float GuideVertexWeight,
 	FHairMeshTriangle RestTriangle, 
 	FHairMeshTriangle DeformedTriangle)
 {	
-	const float3 RestGuidePoint = f16tof32(SimRestPosePositionBuffer[GuideIndex].xyz) - 
-					f16tof32(SimRestPosePositionBuffer[SimRootPointIndexBuffer[GuideIndex]].xyz) ;
+	const uint RootIndex = SimRootPointIndexBuffer[GuideIndex];
+
+	const float3 RestGuidePoint = f16tof32(SimRestPosePositionBuffer[GuideIndex].xyz) - f16tof32(SimRestPosePositionBuffer[RootIndex].xyz) ;
 	const float3 LocalRestGuidePoint = VectorToTriangle(RestGuidePoint, RestTriangle);
 
-	const float3 DeformedGuidePoint = f16tof32(DeformedSimPositionBuffer[GuideIndex].xyz) - 
-					f16tof32(DeformedSimPositionBuffer[SimRootPointIndexBuffer[GuideIndex]].xyz) ;
+	const float3 DeformedGuidePoint = f16tof32(DeformedSimPositionBuffer[GuideIndex].xyz) - f16tof32(DeformedSimPositionBuffer[RootIndex].xyz) ;
 	const float3 LocalDeformedGuidePoint = VectorToTriangle(DeformedGuidePoint, DeformedTriangle);
 
 	return (LocalDeformedGuidePoint - LocalRestGuidePoint) * GuideVertexWeight;
 }
 
-uint4 WriteOutput(float3 WorldPosition, float3 OutPositionOffset, uint PackedDataW)
-{
-	uint3 EncodedRelativePosition = f32tof16(WorldPosition - OutPositionOffset);
-	if (HairLengthScale < 1.0f)
-	{
-		const float CoordU = UnpackHairControlPointCoordU(PackedDataW);
-		if (HairLengthScale < CoordU)
-		{
-			EncodedRelativePosition = f32tof16(INFINITE_FLOAT).xxx;
-		}
-	}
-
-	return uint4(EncodedRelativePosition, PackedDataW);
-}
+///////////////////////////////////////////////////////////////////////////////////////////////////
+// Root triangles
 
 struct FRootTriangleData
 {
@@ -237,305 +231,311 @@ FRootTriangleData GetRootTriangleData(uint VertexIndex, bool bSim)
 	return Out;
 }
 
-[numthreads(GROUP_SIZE, 1, 1)]
-void MainCS(uint2 DispatchThreadId : SV_DispatchThreadID)
+///////////////////////////////////////////////////////////////////////////////////////////////////
+// Guides
+struct FGuideData
 {
-	const float3 SimHairPositionOffset = OutSimHairPositionOffsetBuffer[0].xyz;
-	const float3 OutRenHairPositionOffset = OutRenHairPositionOffsetBuffer[0].xyz;
+	uint3  VertexIndices;
+	float3 VertexLerps;
+	float3 VertexWeights;
+	float3 OutPositionOffset;
+};
 
-	uint VertexIndex = VertexStart + GetHairStrandsVertexFetchIndex(DispatchThreadId, GROUP_SIZE, DispatchCountX);
-#if PERMUTATION_CULLING
-	if (HairStrandsVF_bIsCullingEnable)
+FGuideData GetGuideData(uint InVertexIndex)
+{
+	FGuideData Out = (FGuideData)0;
+#if PERMUTATION_SIMULATION == 1 || PERMUTATION_DEBUG == 1
 	{
-		const uint FetchIndex = GetHairStrandsVertexFetchIndex_Culled(DispatchThreadId, GROUP_SIZE);
-		const uint VertexCountAfterCulling = HairStrandsVF_CullingIndirectBuffer[3];
-		if (FetchIndex >= VertexCountAfterCulling)
-		{
-			return;
-		}
-		VertexIndex = HairStrandsVF_CullingIndexBuffer[FetchIndex];
+		Out.OutPositionOffset = OutSimHairPositionOffsetBuffer[0].xyz;
+
+	#if PERMUTATION_USE_SINGLE_GUIDE == 1
+		Out.VertexWeights = float3(1, 0, 0);
+		Out.VertexIndices = 0;
+		Out.VertexLerps = 0;
+
+		const uint InterpolationData = InterpolationBuffer[InVertexIndex];
+		Out.VertexIndices.x = InterpolationData & 0xFFFFFF;
+		Out.VertexLerps.x = saturate(((InterpolationData >> 24) & 0xFF) / 255.f);
+	#else
+		const uint4 InterpolationData0 = Interpolation0Buffer[InVertexIndex];
+		const uint4 InterpolationData1 = Interpolation1Buffer[InVertexIndex];
+
+		Out.VertexIndices.x = InterpolationData0.x | ((InterpolationData1.x & 0xFF) << 16);
+		Out.VertexIndices.y = InterpolationData0.y | (((InterpolationData1.x >> 8) & 0xFF) << 16);
+		Out.VertexIndices.z = InterpolationData0.z | ((InterpolationData1.y & 0xFF) << 16);
+
+		Out.VertexWeights = float3(saturate((InterpolationData0.w & 0xFF) / 255.f), saturate(((InterpolationData0.w >> 8) & 0xFF) / 255.f), 0);
+		Out.VertexWeights.z = saturate(1.0f - (Out.VertexWeights.x + Out.VertexWeights.y));
+
+		Out.VertexLerps.x = saturate(((InterpolationData1.y >> 8) & 0xFF) / 255.f);
+		Out.VertexLerps.y = saturate((InterpolationData1.z & 0xFF) / 255.f);
+		Out.VertexLerps.z = saturate(((InterpolationData1.z >> 8) & 0xFF) / 255.f);
+	#endif
 	}
 #endif
+	return Out;
+}
 
-	if (VertexIndex < (VertexStart + VertexCount))
+///////////////////////////////////////////////////////////////////////////////////////////////////
+// Rigid interpolation
+float3 ApplyRigidInterpolation(FGuideData GuideData, float3 InPosition)
+{
+	float3 CurrOffset = 0;
+#if PERMUTATION_SIMULATION == 1
+	#if PERMUTATION_USE_SINGLE_GUIDE == 1
+	const uint K = 0;
+	#else
+	[unroll]
+	for (uint K = 0; K < 3; ++K)
+	#endif
 	{
-		uint3  GuideVertexIndices = 0;
-		float3 GuideVertexLerps = 0;
-		float3 GuideVertexWeights = 0;
-		#if PERMUTATION_SIMULATION == 1 || PERMUTATION_DEBUG == 1
-		{
+		const uint GuideIndex0 = GuideData.VertexIndices[K];
+		const uint GuideIndex1 = GuideIndex0 + 1;
 
-		  #if PERMUTATION_USE_SINGLE_GUIDE == 1
-			GuideVertexWeights = float3(1, 0, 0);
-			GuideVertexIndices = 0;
-			GuideVertexLerps = 0;
-
-			const uint InterpolationData = InterpolationBuffer[VertexIndex];
-			GuideVertexIndices.x = InterpolationData & 0xFFFFFF;
-			GuideVertexLerps.x   = saturate(((InterpolationData >> 24) & 0xFF) / 255.f);
-		  #else
-			const uint4 InterpolationData0	= Interpolation0Buffer[VertexIndex];
-			const uint4 InterpolationData1	= Interpolation1Buffer[VertexIndex];
-
-			GuideVertexIndices.x = InterpolationData0.x | ( (InterpolationData1.x & 0xFF) << 16);
-			GuideVertexIndices.y = InterpolationData0.y | (((InterpolationData1.x >> 8) & 0xFF) << 16);
-			GuideVertexIndices.z = InterpolationData0.z | ( (InterpolationData1.y & 0xFF) << 16);
-
-			GuideVertexWeights   = float3(saturate((InterpolationData0.w & 0xFF) / 255.f), saturate(((InterpolationData0.w >> 8) & 0xFF) / 255.f), 0);
-			GuideVertexWeights.z = saturate(1.0f - (GuideVertexWeights.x + GuideVertexWeights.y));
-
-			GuideVertexLerps.x = saturate(((InterpolationData1.y >> 8) & 0xFF)/255.f);
-			GuideVertexLerps.y = saturate(( InterpolationData1.z & 0xFF      )/255.f);
-			GuideVertexLerps.z = saturate(((InterpolationData1.z >> 8) & 0xFF)/255.f);
-		  #endif
-		}
-		#endif
+		const float3 Offset0 = ComputeStaticGeometryOffset(GuideIndex0, GuideData.VertexWeights[K], GuideData.OutPositionOffset);
+		const float3 Offset1 = ComputeStaticGeometryOffset(GuideIndex1, GuideData.VertexWeights[K], GuideData.OutPositionOffset);
 
-		float3 WorldVertexPosition = 0.0f;
+		CurrOffset += lerp(Offset0, Offset1, GuideData.VertexLerps[K]);
+	}
+#endif
 
-		// 1. Rigid transformation
-		// ControlPoint is in the local hair referential
-		// CurrOffset takes only translation component into account, and is done in object local space (vs. triangle/hair local space)
-		#if PERMUTATION_DYNAMIC_GEOMETRY == INTERPOLATION_RIGID
-		{
-			float3 CurrOffset = 0;
-			const uint4 PackedRestPosition = RenderRestPosePositionBuffer[VertexIndex];
-			float3 ControlPoint = f16tof32(PackedRestPosition.xyz) + InRenderHairPositionOffset;
-
-			#if PERMUTATION_SIMULATION == 1
-			#if PERMUTATION_USE_SINGLE_GUIDE == 1
-			const uint K = 0;
-			#else
-			[unroll]
-			for (uint K = 0; K < 3; ++K)
-			#endif
-			{
-				const uint GuideIndex0 = GuideVertexIndices[K];
-				const uint GuideIndex1 = GuideIndex0 + 1;
+	return InPosition + CurrOffset;
+}
 
-				const float3 Offset0 = ComputeStaticGeometryOffset(GuideIndex0, GuideVertexWeights[K], SimHairPositionOffset);
-				const float3 Offset1 = ComputeStaticGeometryOffset(GuideIndex1, GuideVertexWeights[K], SimHairPositionOffset);
+///////////////////////////////////////////////////////////////////////////////////////////////////
+// Skinning with offset
+float3 ApplySkinningWithOffset(const FGuideData GuideData, const FRootTriangleData RenTri, float3 InPosition)
+{
+	float3 CurrOffset = 0;
+
+	// Compute the simulation offset in hair local space (i.e., triangle)
+#if PERMUTATION_SIMULATION == 1
+	#if PERMUTATION_USE_SINGLE_GUIDE == 1
+	const uint K = 0;
+	#else
+	[unroll]
+	for (uint K = 0; K < 3; ++K)
+	#endif
+	{
+		const uint GuideIndex0 = GuideData.VertexIndices[K];
+		const uint GuideIndex1 = GuideIndex0 + 1;
 
-				CurrOffset += lerp(Offset0, Offset1, GuideVertexLerps[K]);
-			}
-			#endif
+		const FRootTriangleData SimTri = GetRootTriangleData(GuideIndex0, true);
 
-			WorldVertexPosition = ControlPoint + CurrOffset;
+		const float3 Offset0 = ComputeDynamicGeometryOffset(GuideIndex0, GuideData.VertexWeights[K], SimTri.RestTriangle, SimTri.DeformedTriangle, SimTri.RootBarycentric, GuideData.OutPositionOffset);
+		const float3 Offset1 = ComputeDynamicGeometryOffset(GuideIndex1, GuideData.VertexWeights[K], SimTri.RestTriangle, SimTri.DeformedTriangle, SimTri.RootBarycentric, GuideData.OutPositionOffset);
 
-			OutRenderDeformedPositionBuffer[VertexIndex] = WriteOutput(WorldVertexPosition, OutRenHairPositionOffset, PackedRestPosition.w);
-		}
-		#endif	
+		CurrOffset += VectorToWorld(lerp(Offset0, Offset1, GuideData.VertexLerps[K]), SimTri.DeformedTriangle);
+	}
+#endif		
 
-		// 2. Skin transformation
-		// Apply dynamic mesh deformation (translation / rotation)
-		#if PERMUTATION_DYNAMIC_GEOMETRY == INTERPOLATION_SKINNING_OFFSET
-		{
-			const uint4 PackedRestPosition = RenderRestPosePositionBuffer[VertexIndex];
-			const float3 RestControlPoint = f16tof32(PackedRestPosition.xyz) + InRenderHairPositionOffset;
-			float3 ControlPoint = RestControlPoint;
-
-			float3 CurrOffset = 0;
-
-			// Compute the simulation offset in hair local space (i.e., triangle)
-			#if PERMUTATION_SIMULATION == 1
-			#if PERMUTATION_USE_SINGLE_GUIDE == 1
-			const uint K = 0;
-			#else
-			[unroll]
-			for (uint K = 0; K < 3; ++K)
-			#endif
-			{
-				const uint GuideIndex0 = GuideVertexIndices[K];
-				const uint GuideIndex1 = GuideIndex0 + 1;
+	InPosition = ToTriangle(InPosition, RenTri.RootBarycentric, RenTri.RestTriangle);
+	InPosition = ToWorld(InPosition, RenTri.RootBarycentric, RenTri.DeformedTriangle) + CurrOffset;
 
-				const FRootTriangleData SimTri = GetRootTriangleData(GuideIndex0, true);
+	return InPosition;
+}
 
-				const float3 Offset0 = ComputeDynamicGeometryOffset(GuideIndex0, GuideVertexWeights[K], SimTri.RestTriangle, SimTri.DeformedTriangle, SimTri.RootBarycentric, SimHairPositionOffset);
-				const float3 Offset1 = ComputeDynamicGeometryOffset(GuideIndex1, GuideVertexWeights[K], SimTri.RestTriangle, SimTri.DeformedTriangle, SimTri.RootBarycentric, SimHairPositionOffset);
+///////////////////////////////////////////////////////////////////////////////////////////////////
+// Skinning with translantion
+float3 ApplySkinningWithTranslation(const FGuideData GuideData, const FRootTriangleData RenTri, float3 InPosition)
+{
+	InPosition -= RenTri.RootBarycentric.x * RenTri.RestTriangle.P0 + RenTri.RootBarycentric.y * RenTri.RestTriangle.P1 + RenTri.RootBarycentric.z * RenTri.RestTriangle.P2;
 
-				CurrOffset += VectorToWorld(lerp(Offset0, Offset1, GuideVertexLerps[K]), SimTri.DeformedTriangle);
-			}
-			#endif		
+	float3 CurrOffset = 0;
 
-			// Transform hair from rest post to deformed pose, based triangle deformation + simulation offset (optional depending of the simulation is actual running or not
-			const FRootTriangleData RenTri = GetRootTriangleData(VertexIndex, false);
+#if PERMUTATION_SIMULATION == 1
+	#if PERMUTATION_USE_SINGLE_GUIDE == 1
+	const uint K = 0;
+	#else
+	[unroll]
+	for (uint K = 0; K < 3; ++K)
+	#endif
+	{
+		const uint GuideIndex0 = GuideData.VertexIndices[K];
+		const uint GuideIndex1 = GuideIndex0 + 1;
 
-			ControlPoint = ToTriangle(ControlPoint, RenTri.RootBarycentric, RenTri.RestTriangle);
-			ControlPoint = ToWorld(ControlPoint, RenTri.RootBarycentric,RenTri.DeformedTriangle) + CurrOffset;
+		const float3 Offset0 = ComputeTranslateGeometryOffset(GuideIndex0, 1.0);
+		const float3 Offset1 = ComputeTranslateGeometryOffset(GuideIndex1, 1.0);
 
-			WorldVertexPosition = ControlPoint;
+		CurrOffset += GuideData.VertexWeights[K] * lerp(Offset0, Offset1, GuideData.VertexLerps[K]);
+	}
+#endif
 
-			#if DEBUG_ENABLE
-			// Debug drawing which is manually enabled, then 
-			{
-				const float4 RestColor = float4(0, 1, 0, 1);
-				const float4 DeformedColor = float4(1, 0, 0, 1);
-				if (VertexIndex==0)
-				{	
-					AddCrossWS(InRenderHairPositionOffset, 1.2, RestColor);
-					AddCrossWS(OutRenHairPositionOffset, 1, DeformedColor);
-				}
-			
-				if (VertexIndex >= 0 && VertexIndex < 64)
-				{
-					AddReferentialWS(RestControlPoint, float3(0, 0, 1));
-					AddReferentialWS(WorldVertexPosition, float3(0, 0, 1));
-				}
+	InPosition += RenTri.RootBarycentric.x * RenTri.DeformedTriangle.P0 + RenTri.RootBarycentric.y * RenTri.DeformedTriangle.P1 + RenTri.RootBarycentric.z * RenTri.DeformedTriangle.P2;
+	return InPosition + CurrOffset;
+}
 
-				{
-					const uint ControlPointType = asuint(PackedRestPosition.w) & 0x3;
-					const bool bIsRoot = ControlPointType == 1;
-					if (bIsRoot)
-					{
-						AddLineTriangleWS(RenTri.RestTriangle.P0, RenTri.RestTriangle.P1, RenTri.RestTriangle.P2, RestColor);
-						AddLineTriangleWS(RenTri.DeformedTriangle.P0, RenTri.DeformedTriangle.P1, RenTri.DeformedTriangle.P2, DeformedColor);
-					}
-				}
-			}
-			#endif
+///////////////////////////////////////////////////////////////////////////////////////////////////
+// Skinning with rotation
+float3 ApplySkinningWithRotation(const FGuideData GuideData, const FRootTriangleData RenTri, float3 InPosition)
+{
+	InPosition -= RenTri.RootBarycentric.x * RenTri.RestTriangle.P0 + RenTri.RootBarycentric.y * RenTri.RestTriangle.P1 + RenTri.RootBarycentric.z * RenTri.RestTriangle.P2;
 
-			OutRenderDeformedPositionBuffer[VertexIndex] = WriteOutput(WorldVertexPosition, OutRenHairPositionOffset, PackedRestPosition.w);
-		}
-		#endif	
+	float3 CurrOffset = 0;
 
-		// 3. Linear blend skinning (translation)
-		#if PERMUTATION_DYNAMIC_GEOMETRY == INTERPOLATION_SKINNING_TRANSLATION
-		{
-			const uint4 PackedRestPosition = RenderRestPosePositionBuffer[VertexIndex];
-			const float3 RestControlPoint = f16tof32(PackedRestPosition.xyz) + InRenderHairPositionOffset;
-			float3 ControlPoint = RestControlPoint;
-
-			const FRootTriangleData RenTri = GetRootTriangleData(VertexIndex, false);
-			
-			ControlPoint -= RenTri.RootBarycentric.x * RenTri.RestTriangle.P0 + RenTri.RootBarycentric.y * RenTri.RestTriangle.P1 + RenTri.RootBarycentric.z * RenTri.RestTriangle.P2;
-			
-			float3 CurrOffset = 0;
-
-			#if PERMUTATION_SIMULATION == 1
-			#if PERMUTATION_USE_SINGLE_GUIDE == 1
-			const uint K = 0;
-			#else
-			[unroll]
-			for (uint K = 0; K < 3; ++K)
-			#endif
-			{
-				const uint GuideIndex0 = GuideVertexIndices[K];
-				const uint GuideIndex1 = GuideIndex0 + 1;
+#if PERMUTATION_SIMULATION == 1
+	#if PERMUTATION_USE_SINGLE_GUIDE == 1
+	const uint K = 0;
+	#else
+	[unroll]
+	for (uint K = 0; K < 3; ++K)
+	#endif
+	{
+		const uint GuideIndex0 = GuideData.VertexIndices[K];
+		const uint GuideIndex1 = GuideIndex0 + 1;
 
-				const float3 Offset0 = ComputeTranslateGeometryOffset(GuideIndex0, 1.0);
-				const float3 Offset1 = ComputeTranslateGeometryOffset(GuideIndex1, 1.0);
+		const FRootTriangleData SimTri = GetRootTriangleData(GuideIndex0, true);
 
-				CurrOffset += GuideVertexWeights[K] * lerp(Offset0, Offset1, GuideVertexLerps[K]);
-			}
-			#endif
+		const float3 LocalPoint = VectorToTriangle(InPosition, SimTri.RestTriangle);
 
-			ControlPoint += RenTri.RootBarycentric.x * RenTri.DeformedTriangle.P0 + RenTri.RootBarycentric.y * RenTri.DeformedTriangle.P1 + RenTri.RootBarycentric.z * RenTri.DeformedTriangle.P2;
+		const float3 Offset0 = ComputeRotateGeometryOffset(GuideIndex0, 1.0, SimTri.RestTriangle, SimTri.DeformedTriangle/*, SimHairPositionOffset*/);
+		const float3 Offset1 = ComputeRotateGeometryOffset(GuideIndex1, 1.0, SimTri.RestTriangle, SimTri.DeformedTriangle/*, SimHairPositionOffset*/);
 
-			WorldVertexPosition = ControlPoint + CurrOffset;
+		CurrOffset += GuideData.VertexWeights[K] * VectorToWorld(lerp(Offset0, Offset1, GuideData.VertexLerps[K]) + LocalPoint, SimTri.DeformedTriangle);
+	}
+#endif
 
-			OutRenderDeformedPositionBuffer[VertexIndex] = WriteOutput(WorldVertexPosition, OutRenHairPositionOffset, PackedRestPosition.w);
-		}
-		#endif	
+	InPosition = RenTri.RootBarycentric.x * RenTri.DeformedTriangle.P0 + RenTri.RootBarycentric.y * RenTri.DeformedTriangle.P1 + RenTri.RootBarycentric.z * RenTri.DeformedTriangle.P2;
+	return InPosition + CurrOffset;
+}
 
-		// 4. Linear blend skinning (rotation)
-		#if PERMUTATION_DYNAMIC_GEOMETRY == INTERPOLATION_SKINNING_ROTATION
-		{
-			const uint4 PackedRestPosition = RenderRestPosePositionBuffer[VertexIndex];
-			const float3 RestControlPoint = f16tof32(PackedRestPosition.xyz) + InRenderHairPositionOffset;
-			float3 ControlPoint = RestControlPoint;
-
-			const FRootTriangleData RenTri = GetRootTriangleData(VertexIndex, false);
-			ControlPoint -= RenTri.RootBarycentric.x * RenTri.RestTriangle.P0 + RenTri.RootBarycentric.y * RenTri.RestTriangle.P1 + RenTri.RootBarycentric.z * RenTri.RestTriangle.P2;
-
-			float3 CurrOffset = 0;
-
-			#if PERMUTATION_SIMULATION == 1
-			#if PERMUTATION_USE_SINGLE_GUIDE == 1
-			const uint K = 0;
-			#else
-			[unroll]
-			for (uint K = 0; K < 3; ++K)
-			#endif
-			{
-				const uint GuideIndex0 = GuideVertexIndices[K];
-				const uint GuideIndex1 = GuideIndex0 + 1;
-			
-				const FRootTriangleData SimTri = GetRootTriangleData(GuideIndex0, true);
-				
-				const float3 LocalPoint = VectorToTriangle(ControlPoint, SimTri.RestTriangle);
+///////////////////////////////////////////////////////////////////////////////////////////////////
+// Skinning with translation and rotation
+float3 ApplySkinningWithTranslationAndRotation(const FGuideData GuideData, const FRootTriangleData RenTri, float3 InPosition)
+{
+	float3 CurrOffset = 0;
+
+#if PERMUTATION_SIMULATION == 1
+	#if PERMUTATION_USE_SINGLE_GUIDE == 1
+	const uint K = 0;
+	#else
+	[unroll]
+	for (uint K = 0; K < 3; ++K)
+	#endif
+	{
+		const uint GuideIndex0 = GuideData.VertexIndices[K];
+		const uint GuideIndex1 = GuideIndex0 + 1;
 
-				const float3 Offset0 = ComputeRotateGeometryOffset(GuideIndex0, 1.0, SimTri.RestTriangle, SimTri.DeformedTriangle/*, SimHairPositionOffset*/);
-				const float3 Offset1 = ComputeRotateGeometryOffset(GuideIndex1, 1.0, SimTri.RestTriangle, SimTri.DeformedTriangle/*, SimHairPositionOffset*/);
+		const FRootTriangleData SimTri = GetRootTriangleData(GuideIndex0, true);
 
-				CurrOffset += GuideVertexWeights[K] * VectorToWorld(lerp(Offset0, Offset1, GuideVertexLerps[K])+LocalPoint, SimTri.DeformedTriangle);
-			}
-			#endif
+		const float3 LocalPoint = ToTriangle(InPosition, SimTri.RootBarycentric, SimTri.RestTriangle);
 
-			ControlPoint = RenTri.RootBarycentric.x * RenTri.DeformedTriangle.P0 + RenTri.RootBarycentric.y * RenTri.DeformedTriangle.P1 + RenTri.RootBarycentric.z * RenTri.DeformedTriangle.P2;
+		const float3 Offset0 = ComputeDynamicGeometryOffset(GuideIndex0, 1.0, SimTri.RestTriangle, SimTri.DeformedTriangle, SimTri.RootBarycentric, GuideData.OutPositionOffset);
+		const float3 Offset1 = ComputeDynamicGeometryOffset(GuideIndex1, 1.0, SimTri.RestTriangle, SimTri.DeformedTriangle, SimTri.RootBarycentric, GuideData.OutPositionOffset);
 
-			WorldVertexPosition = ControlPoint + CurrOffset;
+		CurrOffset += GuideData.VertexWeights[K] * ToWorld(lerp(Offset0, Offset1, GuideData.VertexLerps[K]) + LocalPoint, SimTri.RootBarycentric, SimTri.DeformedTriangle);
+	}
+#endif
 
-			OutRenderDeformedPositionBuffer[VertexIndex] = WriteOutput(WorldVertexPosition, OutRenHairPositionOffset, PackedRestPosition.w);
-		}
-		#endif
+	return CurrOffset;
+}
 
-		// 5. Linear blend skinning (translation + rotation)
-		#if PERMUTATION_DYNAMIC_GEOMETRY == INTERPOLATION_SKINNING_TRANSLATION_AND_ROTATION
+///////////////////////////////////////////////////////////////////////////////////////////////////
+// Patch attribute with debug data
+
+float4 ApplyDebugOutput(const uint InMode, const FGuideData GuideData,  const uint InClusterId)
+{
+	float4 Out = 0;
+	#if PERMUTATION_DEBUG == 1
+	if (InMode == 1)
+	{
+		const uint GuideIndex = GuideData.VertexIndices[0]; // Take the closest guide
+		Out = SimAttributeBuffer[GuideIndex];
+	}
+	else if (InMode == 2)
+	{
+		const uint ClusterIdLo8bits = InClusterId & 0xFF;
+		const uint ClusterIdHi8btis = (InClusterId >> 8) & 0xFF;
+		// Using FNV1 hash to break linearity of ClusterId (generated by linearly parsing cluster grid)
+		uint Hash = 0;
+		Hash = Hash * 17;
+		Hash = Hash ^ ClusterIdLo8bits;
+		Hash = Hash * 17;
+		Hash = Hash ^ ClusterIdHi8btis;
+		Out = asfloat(uint4(0, (Hash & 0xFF) << 8, 0, 0));
+	}
+	#endif
+	return Out;
+}
+
+[numthreads(GROUP_SIZE, 1, 1)]
+void MainCS(uint2 DispatchThreadId : SV_DispatchThreadID)
+{
+	// 1. Find vertex index
+	uint VertexIndex = VertexStart + GetHairStrandsVertexFetchIndex(DispatchThreadId, GROUP_SIZE, DispatchCountX);
+#if PERMUTATION_CULLING
+	if (HairStrandsVF_bIsCullingEnable)
+	{
+		const uint FetchIndex = GetHairStrandsVertexFetchIndex_Culled(DispatchThreadId, GROUP_SIZE);
+		const uint VertexCountAfterCulling = HairStrandsVF_CullingIndirectBuffer[3];
+		if (FetchIndex >= VertexCountAfterCulling)
 		{
-			const uint4 PackedRestPosition = RenderRestPosePositionBuffer[VertexIndex];
-			const float3 RestControlPoint = f16tof32(PackedRestPosition.xyz) + InRenderHairPositionOffset;
-			float3 ControlPoint = RestControlPoint;
-
-			float3 CurrOffset = 0;
-
-			#if PERMUTATION_SIMULATION == 1
-			#if PERMUTATION_USE_SINGLE_GUIDE == 1
-			const uint K = 0;
-			#else
-			[unroll]
-			for (uint K = 0; K < 3; ++K)
-			#endif
-			{
-				const uint GuideIndex0 = GuideVertexIndices[K];
-				const uint GuideIndex1 = GuideIndex0 + 1;
+			return;
+		}
+		VertexIndex = HairStrandsVF_CullingIndexBuffer[FetchIndex];
+	}
+#endif
 
-				const FRootTriangleData SimTri = GetRootTriangleData(GuideIndex0, true);
-				
-				const float3 LocalPoint = ToTriangle(ControlPoint, SimTri.RootBarycentric, SimTri.RestTriangle);
 
-				const float3 Offset0 = ComputeDynamicGeometryOffset(GuideIndex0, 1.0, SimTri.RestTriangle, SimTri.DeformedTriangle, SimTri.RootBarycentric, SimHairPositionOffset);
-				const float3 Offset1 = ComputeDynamicGeometryOffset(GuideIndex1, 1.0, SimTri.RestTriangle, SimTri.DeformedTriangle, SimTri.RootBarycentric, SimHairPositionOffset);
+	// 2. Appply deformation
+	if (VertexIndex < (VertexStart + VertexCount))
+	{
+		const FGuideData GuideData = GetGuideData(VertexIndex);
+		const FRootTriangleData RenTri = GetRootTriangleData(VertexIndex, false);
 
-				CurrOffset += GuideVertexWeights[K] * ToWorld(lerp(Offset0, Offset1, GuideVertexLerps[K])+LocalPoint, SimTri.RootBarycentric, SimTri.DeformedTriangle);
-			}
-			#endif
+		// Manual decoding of the rest position
+		const uint4 PackedRestPosition = RenderRestPosePositionBuffer[VertexIndex];
+		const float3 RestPosition = f16tof32(PackedRestPosition.xyz) + InRenderHairPositionOffset;
 
-			WorldVertexPosition = CurrOffset;
+		float3 OutPosition = 0.0f;
 
-			OutRenderDeformedPositionBuffer[VertexIndex] = WriteOutput(WorldVertexPosition, OutRenHairPositionOffset, PackedRestPosition.w);
+		// 1. Rigid transformation
+		// ControlPoint is in the local hair referential
+		// CurrOffset takes only translation component into account, and is done in object local space (vs. triangle/hair local space)
+		#if PERMUTATION_DYNAMIC_GEOMETRY == INTERPOLATION_RIGID
+		{
+			OutPosition = ApplyRigidInterpolation(GuideData, RestPosition);
+		}
+		// 2. Skin transformation
+		// Apply dynamic mesh deformation (translation / rotation)
+		#elif PERMUTATION_DYNAMIC_GEOMETRY == INTERPOLATION_SKINNING_OFFSET
+		{
+			OutPosition = ApplySkinningWithOffset(GuideData, RenTri, RestPosition);
+		}
+		// 3. Linear blend skinning (translation)
+		#elif PERMUTATION_DYNAMIC_GEOMETRY == INTERPOLATION_SKINNING_TRANSLATION
+		{
+			OutPosition = ApplySkinningWithTranslation(GuideData, RenTri, RestPosition);
+		}
+		// 4. Linear blend skinning (rotation)
+		#elif PERMUTATION_DYNAMIC_GEOMETRY == INTERPOLATION_SKINNING_ROTATION
+		{
+			OutPosition = ApplySkinningWithRotation(GuideData, RenTri, RestPosition);
+		}
+		// 5. Linear blend skinning (translation + rotation)
+		#elif PERMUTATION_DYNAMIC_GEOMETRY == INTERPOLATION_SKINNING_TRANSLATION_AND_ROTATION
+		{
+			OutPosition = ApplySkinningWithTranslationAndRotation(GuideData, RenTri, RestPosition);
 		}
 		#endif
 
-		#if PERMUTATION_DEBUG == 1
-		if (HairDebugMode == 1)
+		// 3. Write out the final position
+		// Optionally trim hair
+		if (HairLengthScale < 1.0f)
 		{
-			const uint GuideIndex = GuideVertexIndices[0]; // Take the closest guide
-			const float4 GuideAttribute = SimAttributeBuffer[GuideIndex];
-			OutRenderAttributeBuffer[VertexIndex] = GuideAttribute;
+			const float CoordU = UnpackHairControlPointCoordU(PackedRestPosition);
+			if (HairLengthScale < CoordU)
+			{
+				OutPosition = INFINITE_FLOAT;
+			}
 		}
+		OutRenderDeformedPositionBuffer[VertexIndex] = PackHairControlPointPosition(PackedRestPosition, OutPosition, OutRenHairPositionOffsetBuffer[0].xyz);
 
-		if (HairDebugMode == 2)
+		// 4. Optionally patch groom attribute for debug rendering
+		#if PERMUTATION_DEBUG == 1
+		if (HairDebugMode == 1 || HairDebugMode == 2)
 		{
-			const uint ClusterId = VertexToClusterIdBuffer[VertexIndex];
-			const uint ClusterIdLo8bits = ClusterId & 0xFF;
-			const uint ClusterIdHi8btis = (ClusterId>>8) & 0xFF;
-			// Using FNV1 hash to break linearity of ClusterId (generated by linearly parsing cluster grid)
-			uint Hash = 0;
-			Hash = Hash * 17;
-			Hash = Hash ^ ClusterIdLo8bits;
-			Hash = Hash * 17;
-			Hash = Hash ^ ClusterIdHi8btis;
-			OutRenderAttributeBuffer[VertexIndex] = asfloat(uint4(0, (Hash & 0xFF) << 8, 0, 0));
+			OutRenderAttributeBuffer[VertexIndex] = ApplyDebugOutput(HairDebugMode, GuideData, VertexToClusterIdBuffer[VertexIndex]);
 		}
 		#endif
 	}
@@ -725,6 +725,8 @@ void ClusterAABBEvaluationCS(uint3 DispatchThreadId : SV_DispatchThreadID, uint3
 
 #endif // SHADER_CLUSTERAABB
 
+///////////////////////////////////////////////////////////////////////////////////////////////////
+// Hair transfer prev. position
 
 #if SHADER_HAIRTRANSFER_PREV_POSITION
 
@@ -747,6 +749,9 @@ void MainCS(uint2 DispatchThreadId : SV_DispatchThreadID)
 
 #endif // SHADER_HAIRTRANSFER_PREV_POSITION
 
+///////////////////////////////////////////////////////////////////////////////////////////////////
+// Groom cache update
+
 #if SHADER_GROOMCACHE_UPDATE
 
 uint ElementCount;
diff --git a/Engine/Shaders/Private/HairStrands/HairStrandsVertexFactoryCommon.ush b/Engine/Shaders/Private/HairStrands/HairStrandsVertexFactoryCommon.ush
--- a/Engine/Shaders/Private/HairStrands/HairStrandsVertexFactoryCommon.ush
+++ b/Engine/Shaders/Private/HairStrands/HairStrandsVertexFactoryCommon.ush
@@ -50,6 +50,22 @@ FHairControlPoint UnpackHairControlPoint(
 	return Out;	
 }
 
+float UnpackHairControlPointCoordU(uint4 InPackedData)
+{
+	const uint PackedAlpha = asuint(InPackedData.w);
+	return ((PackedAlpha & 0xFF00) >> 8) / 255.f;
+}
+
+float3 UnpackHairControlPointPosition(uint4 InPackedData, float3 InVF_PositionOffset)
+{
+	return f16tof32(InPackedData.xyz) + InVF_PositionOffset;
+}
+
+uint4 PackHairControlPointPosition(uint4 InPackedData, float3 InNewPosition, float3 InPositionOffset)
+{
+	return uint4(f32tof16(InNewPosition.xyz - InPositionOffset), InPackedData.w);
+}
+
 bool IsValidHairStrandsSegment(const FHairControlPoint CP0, const FHairControlPoint CP1)
 {
 	return !(CP0.Type == HAIR_CONTROLPOINT_END && CP1.Type == HAIR_CONTROLPOINT_START);
