Fixes alpha channel on r.TSR.History.SeparateTranslucency=1 code path

diff --git a/Engine/Shaders/Private/TemporalSuperResolution/TSRUpdateHistory.usf b/Engine/Shaders/Private/TemporalSuperResolution/TSRUpdateHistory.usf
--- a/Engine/Shaders/Private/TemporalSuperResolution/TSRUpdateHistory.usf
+++ b/Engine/Shaders/Private/TemporalSuperResolution/TSRUpdateHistory.usf
@@ -1358,7 +1358,7 @@ void MainCS(
 			FinalOutputColor[1] = FinalOutputColor[1] * FinalTranslucencyColor[3] + FinalTranslucencyColor[1];
 			FinalOutputColor[2] = FinalOutputColor[2] * FinalTranslucencyColor[3] + FinalTranslucencyColor[2];
 			#if CONFIG_SCENE_COLOR_APLHA
-				FinalOutputColor[3] = saturate(tsr_half(1.0) - (tsr_half(1.0) - FinalOutputColor[3]) * FinalTranslucencyColor[3]);
+				FinalOutputColor[3] = FinalOutputColor[3] * FinalTranslucencyColor[3];
 			#endif
 		#endif
 		
