Reduces the number of barriers on TSR's shading rejection's LDS code path

diff --git a/Engine/Shaders/Private/TemporalSuperResolution/TSRRejectShading.usf b/Engine/Shaders/Private/TemporalSuperResolution/TSRRejectShading.usf
--- a/Engine/Shaders/Private/TemporalSuperResolution/TSRRejectShading.usf
+++ b/Engine/Shaders/Private/TemporalSuperResolution/TSRRejectShading.usf
@@ -2,9 +2,10 @@
 
 //------------------------------------------------------- CONFIG WAVE SIZE
 
-#if DIM_WAVE_SIZE == 0 // fallback to LDS
-	#define WAVE_BROADCAST_GROUPSIZE (16 * 16)
-#endif
+// needed if CONFIG_OPTIMIZED_LDS_ACCESS == 0 for ConvertWaveBroadcastToLDS()
+//#if DIM_WAVE_SIZE == 0 // fallback to LDS
+//	#define WAVE_BROADCAST_GROUPSIZE (16 * 16)
+//#endif
 
 //------------------------------------------------------- INCLUDES
 
@@ -34,6 +35,11 @@
 /** Whether to detect flickering. */
 #define CONFIG_MOIRE_DETECTION (DIM_FLICKERING_DETECTION)
 
+#if DIM_WAVE_SIZE == 0
+	#define CONFIG_OPTIMIZED_LDS_ACCESS 1
+#else
+	#define CONFIG_OPTIMIZED_LDS_ACCESS 0
+#endif
 
 //------------------------------------------------------- DERIVED
 
@@ -116,6 +122,83 @@ tsr_short2 GetLaneSimdPixelOffset(const uint LaneIndex, const uint SimdIndex)
 	);
 }
 
+
+//------------------------------------------------------- GENERAL TEXEL SWIZZLING & NEIGHBOR ACCESS
+
+#if CONFIG_OPTIMIZED_LDS_ACCESS
+
+groupshared tsr_half4 SharedData[GROUP_TILE_SIZE * GROUP_TILE_SIZE];
+
+void WriteAllElementToLDS(uint GroupThreadIndex, TLaneVector<tsr_half, 1, SIMD_SIZE> X)
+{
+	GroupMemoryBarrierWithGroupSync();
+	UNROLL_N(SIMD_SIZE)
+	for (uint SimdIndex = 0; SimdIndex < SIMD_SIZE; SimdIndex++)
+	{
+		SharedData[GroupThreadIndex + SimdIndex * LANE_COUNT] = tsr_half4(X.GetElement(SimdIndex), 0.0, 0.0, 0.0);
+	}
+	GroupMemoryBarrierWithGroupSync();
+}
+
+void WriteAllElementToLDS(uint GroupThreadIndex, TLaneVector<tsr_half, 3, SIMD_SIZE> X)
+{
+	GroupMemoryBarrierWithGroupSync();
+	UNROLL_N(SIMD_SIZE)
+	for (uint SimdIndex = 0; SimdIndex < SIMD_SIZE; SimdIndex++)
+	{
+		SharedData[GroupThreadIndex + SimdIndex * LANE_COUNT] = tsr_half4(X.GetElement(SimdIndex), 0.0);
+	}
+	GroupMemoryBarrierWithGroupSync();
+}
+
+void WriteAllElementToLDS(uint GroupThreadIndex, TLaneVector<tsr_half, 4, SIMD_SIZE> X)
+{
+	GroupMemoryBarrierWithGroupSync();
+	UNROLL_N(SIMD_SIZE)
+	for (uint SimdIndex = 0; SimdIndex < SIMD_SIZE; SimdIndex++)
+	{
+		SharedData[GroupThreadIndex + SimdIndex * LANE_COUNT] = tsr_half4(X.GetElement(SimdIndex));
+	}
+	GroupMemoryBarrierWithGroupSync();
+}
+
+uint GetElementLDSIndex(uint GroupThreadIndex, const uint NeigborSimdIndex, const uint LaneRotation)
+{
+	return ((GroupThreadIndex + LaneRotation) % uint(LANE_COUNT)) + NeigborSimdIndex * LANE_COUNT;
+}
+
+void ReadElementFromLDS(uint GroupThreadIndex, const uint NeigborSimdIndex, const uint LaneRotation, out tsr_half X)
+{
+	X = SharedData[GetElementLDSIndex(GroupThreadIndex, NeigborSimdIndex, LaneRotation)].x;
+}
+
+void ReadElementFromLDS(uint GroupThreadIndex, const uint NeigborSimdIndex, const uint LaneRotation, out tsr_half3 X)
+{
+	X = SharedData[GetElementLDSIndex(GroupThreadIndex, NeigborSimdIndex, LaneRotation)].xyz;
+}
+
+void ReadElementFromLDS(uint GroupThreadIndex, const uint NeigborSimdIndex, const uint LaneRotation, out tsr_half4 X)
+{
+	X = SharedData[GetElementLDSIndex(GroupThreadIndex, NeigborSimdIndex, LaneRotation)].xyzw;
+}
+
+template<typename ScalarType, uint VectorSize, uint SimdSize>
+void WriteToLDSForAccessNeighborTexel(TLaneVector<ScalarType, VectorSize, SimdSize> Input, uint GroupThreadIndex)
+{
+	WriteAllElementToLDS(GroupThreadIndex, Input);
+}
+
+#else // !CONFIG_OPTIMIZED_LDS_ACCESS
+
+template<typename ScalarType, uint VectorSize, uint SimdSize>
+void WriteToLDSForAccessNeighborTexel(TLaneVector<ScalarType, VectorSize, SimdSize> Input, uint GroupThreadIndex)
+{
+	// NOP
+}
+
+#endif // !CONFIG_OPTIMIZED_LDS_ACCESS
+
+
 CALL_SITE_DEBUGLOC
 template<typename ScalarType, uint VectorSize, uint SimdSize>
 TLaneVector<ScalarType, VectorSize, SimdSize> AccessNeighborTexel(TLaneVector<ScalarType, VectorSize, SimdSize> Input, const tsr_short2 Offset, uint GroupThreadIndex)
@@ -139,6 +222,11 @@ TLaneVector<ScalarType, VectorSize, SimdSize> AccessNeighborTexel(TLaneVector<Sc
 		vector<ScalarType, VectorSize> SimdElement = Input.GetElement(NeigborSimdIndex);
 		vector<ScalarType, VectorSize> ReturnSimdElement = SimdElement;
 		if (LaneRotation != 0)
+		#if CONFIG_OPTIMIZED_LDS_ACCESS
+		{
+			ReadElementFromLDS(GroupThreadIndex, NeigborSimdIndex, LaneRotation, /* out */ ReturnSimdElement);
+		}
+		#else
 		{
 			// Only uses wave intrinsics on platforms that provides RDNA's ds_swizzle because we know the wave size of this platforms. 
 			#if DIM_WAVE_SIZE != 0
@@ -151,6 +239,7 @@ TLaneVector<ScalarType, VectorSize, SimdSize> AccessNeighborTexel(TLaneVector<Sc
 			#endif
 			ReturnSimdElement = WaveBroadcast(BroadcastSettings, SimdElement);
 		}
+		#endif
 
 		Return.SetElement(SimdIndex, ReturnSimdElement);
 	}
@@ -280,6 +369,33 @@ struct TAccumulatorMax
 CALL_SITE_DEBUGLOC
 template<class FAccumulator, typename FSampleType>
 FAccumulator Convole3x1(FSampleType Color, tsr_simd_half Weight, const tsr_short2 Offset, uint GroupThreadIndex)
+#if CONFIG_OPTIMIZED_LDS_ACCESS
+{
+	const tsr_short2 OffsetN = tsr_short2(-int2(Offset)); // work arround a shader compiler bug
+	
+	FSampleType PreprocessedNeighborColor = FAccumulator::PreprocessNeighborSample(Color, Weight);
+
+	WriteToLDSForAccessNeighborTexel(PreprocessedNeighborColor, GroupThreadIndex);
+	FSampleType Sample0 = AccessNeighborTexel(PreprocessedNeighborColor, Offset, GroupThreadIndex);
+	FSampleType Sample1 = AccessNeighborTexel(PreprocessedNeighborColor, OffsetN, GroupThreadIndex);
+	
+	tsr_simd_half Weight0 = tsr_simd_half::Const(tsr_half(1.0));
+	tsr_simd_half Weight1 = tsr_simd_half::Const(tsr_half(1.0));
+	if (FAccumulator::NeedWeightInAccumulate())
+	{
+		WriteToLDSForAccessNeighborTexel(Weight, GroupThreadIndex);
+		Weight0 = AccessNeighborTexel(Weight, Offset, GroupThreadIndex);
+		Weight1 = AccessNeighborTexel(Weight, OffsetN, GroupThreadIndex);
+	}
+	
+	FAccumulator Accumulator;
+	Accumulator.Start(Color);
+	Accumulator.Accumulate(+Offset, Sample0, Weight0);
+	Accumulator.Accumulate(OffsetN, Sample1, Weight1);
+	Accumulator.Finish(3);
+	return Accumulator;
+}
+#else // !CONFIG_OPTIMIZED_LDS_ACCESS
 {
 	FAccumulator Accumulator;
 	Accumulator.Start(Color);
@@ -306,6 +422,7 @@ FAccumulator Convole3x1(FSampleType Color, tsr_simd_half Weight, const tsr_short
 	Accumulator.Finish(3);
 	return Accumulator;
 }
+#endif // !CONFIG_OPTIMIZED_LDS_ACCESS
 
 CALL_SITE_DEBUGLOC
 template<class FAccumulator, typename FSampleType>
