Fixed: Missing per-instance CustomData in a pixel shaders when GPUScene is off

diff --git a/Engine/Shaders/Private/LocalVertexFactory.ush b/Engine/Shaders/Private/LocalVertexFactory.ush
--- a/Engine/Shaders/Private/LocalVertexFactory.ush
+++ b/Engine/Shaders/Private/LocalVertexFactory.ush
@@ -334,8 +334,8 @@ struct FVertexFactoryIntermediates
 
 #if USE_INSTANCING || USE_INSTANCE_CULLING
 	float4 InstanceLightmapAndShadowMapUVBias;
-	// x = per-instance fade out amount, y = hide/show flag, z dither fade cutoff
-	float3 PerInstanceParams;
+	// x = per-instance fade out amount, y = hide/show flag, z dither fade cutoff, w - CustomData index
+	float4 PerInstanceParams;
 #endif
 
 	float3 PreSkinPosition;
@@ -941,6 +941,11 @@ FVertexFactoryIntermediates GetVertexFactoryIntermediates(FVertexFactoryInput In
 	#else
 		Intermediates.PerInstanceParams.z = 0;
 	#endif
+	
+	#if USES_PER_INSTANCE_CUSTOM_DATA
+		// index into instance CustomData
+		Intermediates.PerInstanceParams.w = asfloat(GetInstanceId(Input.InstanceId) + InstanceOffset);
+	#endif
 #endif	// USE_INSTANCING
 
 	if (IsGPUSkinPassThrough())
diff --git a/Engine/Shaders/Private/LocalVertexFactoryCommon.ush b/Engine/Shaders/Private/LocalVertexFactoryCommon.ush
--- a/Engine/Shaders/Private/LocalVertexFactoryCommon.ush
+++ b/Engine/Shaders/Private/LocalVertexFactoryCommon.ush
@@ -13,8 +13,8 @@ struct FVertexFactoryInterpolantsVSToPS
 #endif
 
 #if USE_INSTANCING || USE_INSTANCE_CULLING
-	// x = per-instance fade out amount, y = hide/show flag, z dither fade cutoff
-	float3  PerInstanceParams : COLOR1;
+	// x = per-instance fade out amount, y = hide/show flag, z dither fade cutoff, w - CustomData index
+	nointerpolation float4  PerInstanceParams : COLOR1;
 #endif
 
 #if NUM_TEX_COORD_INTERPOLATORS
diff --git a/Engine/Shaders/Private/MaterialTemplate.ush b/Engine/Shaders/Private/MaterialTemplate.ush
--- a/Engine/Shaders/Private/MaterialTemplate.ush
+++ b/Engine/Shaders/Private/MaterialTemplate.ush
@@ -388,7 +388,7 @@ struct FMaterialPixelParameters
 #endif
 
 #if USE_INSTANCING || USE_INSTANCE_CULLING || IS_NANITE_PASS
-	half3 PerInstanceParams;
+	float4 PerInstanceParams;
 #endif
 
 	// Index into View.PrimitiveSceneData
@@ -523,7 +523,7 @@ struct FMaterialVertexParameters
 #if USE_INSTANCING || USE_INSTANCE_CULLING || IS_NANITE_PASS
 	/** Per-instance properties. */
 	float3 InstanceLocalPosition;
-	float3 PerInstanceParams;
+	float4 PerInstanceParams;
 #if !USE_INSTANCE_CULLING
 	uint InstanceId;
 #endif
@@ -1102,7 +1102,7 @@ MaterialFloat2 GetDefaultSceneTextureUV(FMaterialPixelParameters Parameters, con
 /** Per Instance Custom Data Getter (Pixel Shader Variant - Visibility Buffer) */
 float GetPerInstanceCustomData(FMaterialPixelParameters Parameters, int Index, float DefaultValue)
 {
-#if USES_PER_INSTANCE_CUSTOM_DATA && VF_USE_PRIMITIVE_SCENE_DATA
+#if VF_USE_PRIMITIVE_SCENE_DATA && USES_PER_INSTANCE_CUSTOM_DATA
 	const uint FloatIndex = uint(Index);
 	BRANCH
 	if (FloatIndex < Parameters.CustomDataCount)
@@ -1111,8 +1111,16 @@ float GetPerInstanceCustomData(FMaterialPixelParameters Parameters, int Index, f
 		const float4 Float4Packed = LoadInstancePayloadDataElement(Float4Offset);
 		return Float4Packed[FloatIndex & 0x3u];
 	}
+#elif USE_INSTANCING && USES_PER_INSTANCE_CUSTOM_DATA
+	const uint FloatIndex = uint(Index);
+	BRANCH
+	if (FloatIndex < InstanceVF.NumCustomDataFloats)
+	{
+		const uint InstanceDataIndex = asuint(Parameters.PerInstanceParams.w);
+		const uint BufferStartIndex = InstanceDataIndex * InstanceVF.NumCustomDataFloats;
+		return InstanceVF.InstanceCustomDataBuffer[BufferStartIndex + FloatIndex];
+	}
 #endif
-
 	return DefaultValue;
 }
 
@@ -1147,12 +1155,23 @@ float GetPerInstanceCustomData(FMaterialVertexParameters Parameters, int Index,
 /** Get the per-instance custom data when instancing */
 MaterialFloat3 GetPerInstanceCustomData3Vector(FMaterialPixelParameters Parameters, int Index, MaterialFloat3 DefaultValue)
 {
-#if USES_PER_INSTANCE_CUSTOM_DATA && VF_USE_PRIMITIVE_SCENE_DATA
+#if VF_USE_PRIMITIVE_SCENE_DATA && USES_PER_INSTANCE_CUSTOM_DATA
 	return float3(GetPerInstanceCustomData(Parameters, Index + 0, DefaultValue.x),
-				  GetPerInstanceCustomData(Parameters, Index + 1, DefaultValue.y),
-				  GetPerInstanceCustomData(Parameters, Index + 2, DefaultValue.z));
-#endif
+		GetPerInstanceCustomData(Parameters, Index + 1, DefaultValue.y),
+		GetPerInstanceCustomData(Parameters, Index + 2, DefaultValue.z));
 
+#elif USE_INSTANCING && USES_PER_INSTANCE_CUSTOM_DATA
+	const uint FloatIndex = uint(Index);
+	BRANCH
+	if (FloatIndex + 2 < InstanceVF.NumCustomDataFloats)
+	{
+		const uint InstanceDataIndex = asuint(Parameters.PerInstanceParams.w);
+		const uint BufferStartIndex = InstanceDataIndex * InstanceVF.NumCustomDataFloats;
+		return float3(InstanceVF.InstanceCustomDataBuffer[BufferStartIndex + FloatIndex],
+			InstanceVF.InstanceCustomDataBuffer[BufferStartIndex + FloatIndex + 1],
+			InstanceVF.InstanceCustomDataBuffer[BufferStartIndex + FloatIndex + 2]);
+	}
+#endif
 	return DefaultValue;
 }
 
diff --git a/Engine/Source/Runtime/Engine/Private/InstancedStaticMesh.cpp b/Engine/Source/Runtime/Engine/Private/InstancedStaticMesh.cpp
--- a/Engine/Source/Runtime/Engine/Private/InstancedStaticMesh.cpp
+++ b/Engine/Source/Runtime/Engine/Private/InstancedStaticMesh.cpp
@@ -1004,6 +1004,8 @@ void FInstancedStaticMeshVertexFactory::InitRHI()
 }
 
 IMPLEMENT_VERTEX_FACTORY_PARAMETER_TYPE(FInstancedStaticMeshVertexFactory, SF_Vertex, FInstancedStaticMeshVertexFactoryShaderParameters);
+// pixel shader may need access to InstanceCustomDataBuffer in non-GPUScene case
+IMPLEMENT_VERTEX_FACTORY_PARAMETER_TYPE(FInstancedStaticMeshVertexFactory, SF_Pixel, FInstancedStaticMeshVertexFactoryShaderParameters);
 #if RHI_RAYTRACING
 IMPLEMENT_VERTEX_FACTORY_PARAMETER_TYPE(FInstancedStaticMeshVertexFactory, SF_RayHitGroup, FInstancedStaticMeshVertexFactoryShaderParameters);
 IMPLEMENT_VERTEX_FACTORY_PARAMETER_TYPE(FInstancedStaticMeshVertexFactory, SF_Compute, FInstancedStaticMeshVertexFactoryShaderParameters);
