Path Tracer: Fix for dark reflections with SingleLayerWater materials

The current SigmaT for the path should only be updated when crossing the boundary.

diff --git a/Engine/Shaders/Private/PathTracing/PathTracingCore.ush b/Engine/Shaders/Private/PathTracing/PathTracingCore.ush
--- a/Engine/Shaders/Private/PathTracing/PathTracingCore.ush
+++ b/Engine/Shaders/Private/PathTracing/PathTracingCore.ush
@@ -1217,7 +1217,7 @@ bool PathTracingKernel(inout FPathState PathState, int Bounce)
 		}
 	}
 
-	if (Payload.ShadingModelID == SHADINGMODELID_SINGLELAYERWATER)
+	if (MaterialSample.PositionBiasSign < 0 && Payload.ShadingModelID == SHADINGMODELID_SINGLELAYERWATER)
 	{
 		float3 LocalSigmaT = max(Payload.CustomData.xyz, 1e-8f);
 		if (Payload.IsFrontFace())
