Fixes a source of blur arround velocity discontinuities with r.TSR.History.GrandReprojection=1

diff --git a/Engine/Shaders/Private/TemporalSuperResolution/TSRUpdateHistory.usf b/Engine/Shaders/Private/TemporalSuperResolution/TSRUpdateHistory.usf
--- a/Engine/Shaders/Private/TemporalSuperResolution/TSRUpdateHistory.usf
+++ b/Engine/Shaders/Private/TemporalSuperResolution/TSRUpdateHistory.usf
@@ -1240,12 +1240,12 @@ void MainCS(
 
 			FinalOverblur = dpv_scale(PrevResultant, BlendClamp * TotalRejection) * (PrevHighFrequencyOverblur + dpv_scale(PrevReprojectionOverblur, PrevReprojectionOverblurScale));
 
-			//tsr_half2 Rejection = MeasureRejectionFactor(
-			//	UnpremultipliedPrevHighFrequencyColor, BlendedPrevHighFrequencyColor,
-			//	FilteredInputColor, InputMinColor, InputMaxColor,
-			//	MeasureBackbufferLDRQuantizationError());
+			tsr_half2 Rejection = MeasureRejectionFactor(
+				UnpremultipliedPrevHighFrequencyColor, BlendedPrevHighFrequencyColor,
+				FilteredInputColor, InputMinColor, InputMaxColor,
+				MeasureBackbufferLDRQuantizationError());
 
-			tsr_half2 MeasuredResultant = saturate(CommonWeight * BlendPrevHistory * BlendClamp);
+			tsr_half2 MeasuredResultant = saturate(CommonWeight * BlendPrevHistory * lerp(Rejection, tsr_half(1.0), BlendClamp));
 
 			FinalPrevHighFrequencyColor = PrevHighFrequencyColor;
 			FinalPrevHighFrequencyColorResultant = dpv_scale(PrevHighFrequencyColor, MeasuredResultant * InvPrevHistoryWeight * FinalHistoryWeight);
