Strata - make sure we notify that strata shaders only work with 64 thread on GPU that can run it at 32 or any other granularity.

diff --git a/Engine/Shaders/Private/Strata/StrataMaterialClassification.usf b/Engine/Shaders/Private/Strata/StrataMaterialClassification.usf
--- a/Engine/Shaders/Private/Strata/StrataMaterialClassification.usf
+++ b/Engine/Shaders/Private/Strata/StrataMaterialClassification.usf
@@ -73,6 +73,9 @@ Texture2D<float3> OpaqueRoughRefractionTexture;
 groupshared uint s_TileFlags[GROUP_THREAD_COUNT];
 #endif
 
+#if PERMUTATION_WAVE_OPS && COMPILER_SUPPORTS_WAVE_SIZE
+WAVESIZE(64) // PERMUTATION_WAVE_OPS is true only when wave>=64 are available
+#endif
 [numthreads(STRATA_TILE_SIZE, STRATA_TILE_SIZE, 1)]
 void TileMainCS(uint2 DispatchThreadId : SV_DispatchThreadID, uint LinearIndex : SV_GroupIndex, uint3 GroupId : SV_GroupID)
 {
@@ -423,6 +426,9 @@ uint2 LinearCoordToTileCoord(uint InLinear)
 	return uint2(InLinear % OverflowTileCount.x, InLinear / OverflowTileCount.x) + OverflowTileOffset;
 }
 
+#if PERMUTATION_WAVE_OPS && COMPILER_SUPPORTS_WAVE_SIZE
+WAVESIZE(64) // PERMUTATION_WAVE_OPS is true only when wave>=64 are available
+#endif
 [numthreads(STRATA_TILE_SIZE, STRATA_TILE_SIZE, 1)]
 void BSDFTileMainCS(uint2 GroupThreadId : SV_GroupThreadID, uint2 GroupId : SV_GroupID, uint LinearIndex : SV_GroupIndex)
 {
