Fixed shader compilation error with mobile deferred shading enabled

diff --git a/Engine/Shaders/Private/MobileDeferredShading.usf b/Engine/Shaders/Private/MobileDeferredShading.usf
--- a/Engine/Shaders/Private/MobileDeferredShading.usf
+++ b/Engine/Shaders/Private/MobileDeferredShading.usf
@@ -223,7 +223,7 @@ void MobileRadialLightPS(
 	half3 CameraVector = normalize(TranslatedWorldPosition);
 			
 	FLightAccumulator DirectLighting = (FLightAccumulator)0;
-	FDeferredLightData LightData = InitDeferredLightFromUniforms_Mobile(IS_SPOT_LIGHT ? LIGHT_TYPE_SPOT : LIGHT_TYPE_POINT);
+	FDeferredLightData LightData = InitDeferredLightFromUniforms(IS_SPOT_LIGHT ? LIGHT_TYPE_SPOT : LIGHT_TYPE_POINT);
 	
 	// Affect the light color by any other shadow/transmittance before the lighting is evaluated
 	{
