Support GBuffer access for translucency with mobile deferred shading enabled (GL does not work atm)
Cleanup sub-pass features usage in mobile shaders

diff --git a/Engine/Shaders/Private/DeferredShadingCommon.ush b/Engine/Shaders/Private/DeferredShadingCommon.ush
--- a/Engine/Shaders/Private/DeferredShadingCommon.ush
+++ b/Engine/Shaders/Private/DeferredShadingCommon.ush
@@ -467,14 +467,11 @@ float4 ComputeIndirectLightingSampleE(uint2 TracingPixelCoord, uint TracingRayIn
 #endif
 
 #if SHADING_PATH_MOBILE
+
 #ifndef ENABLE_SHADINGMODEL_SUPPORT_MOBILE_DEFERRED
 #define ENABLE_SHADINGMODEL_SUPPORT_MOBILE_DEFERRED 0
 #endif
 
-#ifndef MOBILE_DEFERRED_SHADING
-#define MOBILE_DEFERRED_SHADING 0
-#endif
-
 #define MOBILE_SHADINGMODEL_SUPPORT (!MOBILE_DEFERRED_SHADING || ENABLE_SHADINGMODEL_SUPPORT_MOBILE_DEFERRED)
 
 float MobileEncodeIdAndColorChannel(uint Id, float Color, bool b10Bits)
@@ -746,6 +743,75 @@ FGBufferData MobileDecodeGBuffer(half4 InGBufferA, half4 InGBufferB, half4 InGBu
 	return GBuffer;
 }
 
+#if MOBILE_DEFERRED_SHADING && PIXELSHADER
+
+#if METAL_PROFILE
+#include "/Engine/Public/Platform/Metal/MetalSubpassSupport.ush"
+#elif VULKAN_PROFILE
+#include "/Engine/Public/Platform/Vulkan/VulkanSubpassSupport.ush"
+#elif COMPILER_GLSL_ES3_1
+#include "/Engine/Public/Platform/GL/GLSubpassSupport.ush"
+#endif
+
+void MobileFetchGBuffer(in float2 UV, out half4 GBufferA, out half4 GBufferB, out half4 GBufferC, out half4 GBufferD, out float SceneDepth)
+{
+	GBufferD = 0;
+
+#if VULKAN_PROFILE
+	GBufferA = VulkanSubpassFetch1(); 
+	GBufferB = VulkanSubpassFetch2(); 
+	GBufferC = VulkanSubpassFetch3();
+#if MOBILE_EXTENDED_GBUFFER
+	GBufferD = VulkanSubpassFetch4();
+#endif
+	SceneDepth = ConvertFromDeviceZ(VulkanSubpassDepthFetch());
+#elif METAL_PROFILE
+	GBufferA = SubpassFetchRGBA_1(); 
+	GBufferB = SubpassFetchRGBA_2(); 
+	GBufferC = SubpassFetchRGBA_3(); 
+#if MOBILE_EXTENDED_GBUFFER
+	GBufferD = SubpassFetchRGBA_4();
+#endif
+	SceneDepth = ConvertFromDeviceZ(SubpassFetchR_4());
+#elif USE_GLES_FBF_DEFERRED
+	GBufferA = GLSubpassFetch1(); 
+	GBufferB = GLSubpassFetch2(); 
+	GBufferC = GLSubpassFetch3();  
+	GBufferD = 0; // PLS is limited to 128bits
+	SceneDepth = ConvertFromDeviceZ(DepthbufferFetchES2());
+#else
+	GBufferA = Texture2DSampleLevel(MobileSceneTextures.GBufferATexture, MobileSceneTextures.GBufferATextureSampler, UV, 0); 
+	GBufferB = Texture2DSampleLevel(MobileSceneTextures.GBufferBTexture, MobileSceneTextures.GBufferBTextureSampler, UV, 0);
+	GBufferC = Texture2DSampleLevel(MobileSceneTextures.GBufferCTexture, MobileSceneTextures.GBufferCTextureSampler, UV, 0);
+#if MOBILE_EXTENDED_GBUFFER
+	GBufferD = Texture2DSampleLevel(MobileSceneTextures.GBufferDTexture, MobileSceneTextures.GBufferDTextureSampler, UV, 0);
+#endif
+	SceneDepth = ConvertFromDeviceZ(Texture2DSampleLevel(MobileSceneTextures.SceneDepthTexture, MobileSceneTextures.SceneDepthTextureSampler, UV, 0).r);
+#endif
+}
+#endif //MOBILE_DEFERRED_SHADING
+
+FGBufferData MobileFetchAndDecodeGBuffer(in float2 UV, in float2 PixelPos)
+{
+	FGBufferData GBuffer = (FGBufferData)0;
+#if (MOBILE_DEFERRED_SHADING && IS_MOBILE_DEFERREDSHADING_SUBPASS && PIXELSHADER)
+	float SceneDepth = 0;
+	half4 GBufferA = 0;
+	half4 GBufferB = 0;
+	half4 GBufferC = 0;
+	half4 GBufferD = 0;
+	MobileFetchGBuffer(UV, GBufferA, GBufferB, GBufferC, GBufferD, SceneDepth);
+	GBuffer = MobileDecodeGBuffer(GBufferA, GBufferB, GBufferC, GBufferD);
+	GBuffer.Depth = SceneDepth;
+#else
+	GBuffer.Depth = CalcSceneDepth(UV);
+#endif
+
+	GBuffer.CustomDepth = ConvertFromDeviceZ(Texture2DSample(MobileSceneTextures.CustomDepthTexture, MobileSceneTextures.CustomDepthTextureSampler, UV).r);
+	GBuffer.CustomStencil = MobileSceneTextures.CustomStencilTexture.Load(int3(PixelPos.xy, 0)) STENCIL_COMPONENT_SWIZZLE;
+
+	return GBuffer;
+}
 #endif //SHADING_PATH_MOBILE
 
 /** Populates OutGBufferA, B and C */
diff --git a/Engine/Shaders/Private/DistortAccumulatePS.usf b/Engine/Shaders/Private/DistortAccumulatePS.usf
--- a/Engine/Shaders/Private/DistortAccumulatePS.usf
+++ b/Engine/Shaders/Private/DistortAccumulatePS.usf
@@ -143,13 +143,7 @@ void Main(
 
 #endif // STRATA_ENABLED
 
-
-#if SHADING_PATH_MOBILE
-	// sample depth from SceneColor.A on Mobile
-	float DistortSceneDepth = ConvertFromDeviceZ(LookupDeviceZ(DistortBufferUV)); 
-#else
 	float DistortSceneDepth = CalcSceneDepth(DistortBufferUV);
-#endif
 
 	// Post process UV distortion according to depth
 	PostProcessUVDistortion(MaterialParameters, PixelMaterialInputs, DistortSceneDepth, BufferUVDistortion, Refraction);
diff --git a/Engine/Shaders/Private/MaterialTemplate.ush b/Engine/Shaders/Private/MaterialTemplate.ush
--- a/Engine/Shaders/Private/MaterialTemplate.ush
+++ b/Engine/Shaders/Private/MaterialTemplate.ush
@@ -434,11 +434,6 @@ struct FMaterialPixelParameters
 	FTexCoordScalesParams TexCoordScalesParams;
 #endif
 
-#if POST_PROCESS_MATERIAL && FEATURE_LEVEL == FEATURE_LEVEL_ES3_1
-	/** Used in mobile custom pp material to preserve original SceneColor Alpha */
-	half BackupSceneColorAlpha;
-#endif
-
 #if COMPILER_HLSL
 	// Workaround for "error X3067: 'GetObjectWorldPosition': ambiguous function call"
 	// Which happens when FMaterialPixelParameters and FMaterialVertexParameters have the same number of floats with the HLSL compiler ver 9.29.952.3111
@@ -2257,56 +2252,68 @@ MaterialFloat4 MobileSceneTextureLookup(inout FMaterialPixelParameters Parameter
 		return 0.0f;
 	}
 #else
-#if MATERIAL_SHADINGMODEL_SINGLELAYERWATER
-	if (SceneTextureId == PPI_CustomDepth)
-	{
-		MaterialFloat Depth = ConvertFromDeviceZ(Texture2DSample(MobileSceneTextures.CustomDepthTexture, MobileSceneTextures.CustomDepthTextureSampler, UV).r);
-		return MaterialFloat4(Depth.rrr, 0.f);
-	}
-	else if (SceneTextureId == PPI_CustomStencil)
-	{
-		MaterialFloat Stencil = MobileSceneTextures.CustomStencilTexture.Load(int3(Parameters.SvPosition.xy, 0)) STENCIL_COMPONENT_SWIZZLE;
-		return MaterialFloat4(Stencil, Stencil, Stencil, 0.f);
-	}
-#else // MATERIAL_SHADINGMODEL_SINGLELAYERWATER
-	if (SceneTextureId == PPI_SceneDepth)
-	{
-		MaterialFloat Depth = CalcSceneDepth(UV);
-		return MaterialFloat4(Depth.rrr, 0.f);
-	}
-	else if (SceneTextureId == PPI_CustomDepth)
-	{
-		MaterialFloat Depth = ConvertFromDeviceZ(Texture2DSample(MobileSceneTextures.CustomDepthTexture, MobileSceneTextures.CustomDepthTextureSampler, UV).r);
-		return MaterialFloat4(Depth.rrr, 0.f);
-	}
-	else if (SceneTextureId == PPI_PostProcessInput0)
+	FGBufferData GBuffer = MobileFetchAndDecodeGBuffer(UV, Parameters.SvPosition.xy);
+
+	switch(SceneTextureId)
 	{
+		// order needs to match to ESceneTextureId
+	case PPI_SceneDepth:
+		return GBuffer.Depth;
+	case PPI_DiffuseColor:
+		return MaterialFloat4(GBuffer.DiffuseColor, 0);
+	case PPI_SpecularColor:
+		return MaterialFloat4(GBuffer.SpecularColor, 0);
+	case PPI_SubsurfaceColor:
+		return IsSubsurfaceModel(GBuffer.ShadingModelID) ? MaterialFloat4( ExtractSubsurfaceColor(GBuffer), GBuffer.CustomData.a ) : GBuffer.CustomData;
+	case PPI_BaseColor:
+		return MaterialFloat4(GBuffer.BaseColor, 0);
+	case PPI_Specular:
+		return GBuffer.Specular;
+	case PPI_Metallic:
+		return GBuffer.Metallic;
+	case PPI_WorldNormal:
+		return MaterialFloat4(GBuffer.WorldNormal, 0);
+	case PPI_SeparateTranslucency:
+		return MaterialFloat4(1, 1, 1, 1);	// todo
+	case PPI_Opacity:
+		return GBuffer.CustomData.a;
+	case PPI_Roughness:
+		return GBuffer.Roughness;
+	case PPI_MaterialAO:
+		return GBuffer.GBufferAO;
+	case PPI_CustomDepth:
+		return GBuffer.CustomDepth;
 #if POST_PROCESS_MATERIAL
+	case PPI_PostProcessInput0:
+	{
 		MaterialFloat4 Input0 = Texture2DSample(PostProcessInput_0_Texture, PostProcessInput_0_SharedSampler, UV);
-		#if POST_PROCESS_MATERIAL_BEFORE_TONEMAP
-			#if METAL_PROFILE
-				// Decode the input color since the color is encoded for MSAA 
-				// The decode instructions might be able to skip with dynamic branch
-				if (bMetalMSAAHDRDecode)
-				{
-					Input0.rgb = Input0.rgb * rcp(Input0.r*(-0.299) + Input0.g*(-0.587) + Input0.b*(-0.114) + 1.0);
-				}
-			#endif
+		#if POST_PROCESS_MATERIAL_BEFORE_TONEMAP && METAL_PROFILE
+		// Decode the input color since the color is encoded for MSAA 
+		// The decode instructions might be able to skip with dynamic branch
+		if (bMetalMSAAHDRDecode)
+		{
+			Input0.rgb = Input0.rgb * rcp(Input0.r*(-0.299) + Input0.g*(-0.587) + Input0.b*(-0.114) + 1.0);
+		}
 		#endif
-		// We need to preserve original SceneColor Alpha as it's used by tonemapper on mobile
-		Parameters.BackupSceneColorAlpha = Input0.a;
 		return Input0;
-#endif// POST_PROCESS_MATERIAL
 	}
-	else if (SceneTextureId == PPI_CustomStencil)
-	{
-		MaterialFloat Stencil = MobileSceneTextures.CustomStencilTexture.Load(int3(Parameters.SvPosition.xy, 0)) STENCIL_COMPONENT_SWIZZLE;
-		return MaterialFloat4(Stencil, Stencil, Stencil, 0.f);
+#endif // POST_PROCESS_MATERIAL
+	case PPI_DecalMask:
+		return 0;  // material compiler will return an error
+	case PPI_ShadingModelColor:
+		return MaterialFloat4(GetShadingModelColor(GBuffer.ShadingModelID), 1);
+	case PPI_ShadingModelID:
+		return MaterialFloat4(GBuffer.ShadingModelID, 0, 0, 0);
+	case PPI_AmbientOcclusion:
+		return 1;
+	case PPI_CustomStencil:
+		return GBuffer.CustomStencil;
+	case PPI_WorldTangent:
+		return MaterialFloat4(GBuffer.WorldTangent, 0);
+	default:
+		return MaterialFloat4(0, 0, 0, 0);
 	}
-#endif // MATERIAL_SHADINGMODEL_SINGLELAYERWATER
-#endif // SCENE_TEXTURES_DISABLED
-
-	return MaterialFloat4(0.0f, 0.0f, 0.0f, 0.0f);
+#endif
 }
 
 #endif // SHADING_PATH_MOBILE
diff --git a/Engine/Shaders/Private/MobileDeferredShading.usf b/Engine/Shaders/Private/MobileDeferredShading.usf
--- a/Engine/Shaders/Private/MobileDeferredShading.usf
+++ b/Engine/Shaders/Private/MobileDeferredShading.usf
@@ -5,10 +5,6 @@
 #include "Common.ush"
 #include "SHCommon.ush"
 
-#if MATERIAL_SHADER
-#include "/Engine/Generated/Material.ush"
-#endif
-
 // Reroute MobileSceneTextures uniform buffer references to the base pass uniform buffer
 #define MobileSceneTextures		MobileBasePass.SceneTextures
 #define ForwardLightData		MobileBasePass.Forward
@@ -17,74 +13,20 @@
 #define PreIntegratedGF			ReflectionStruct.PreIntegratedGF
 #define PreIntegratedGFSampler	ReflectionStruct.PreIntegratedGFSampler
 
+#if MATERIAL_SHADER
+#include "/Engine/Generated/Material.ush"
+#endif
+
 #include "MobileLightingCommon.ush"
 #if MATERIAL_SHADER
 #include "LightFunctionCommon.ush"
 #endif
 #include "LightDataUniforms.ush"
 
-#if METAL_PROFILE
-#include "/Engine/Public/Platform/Metal/MetalSubpassSupport.ush"
-#elif VULKAN_PROFILE
-#include "/Engine/Public/Platform/Vulkan/VulkanSubpassSupport.ush"
-#elif OPENGL_PROFILE
-#include "/Engine/Public/Platform/GL/GLSubpassSupport.ush"
-#endif
-
 #ifndef IS_SPOT_LIGHT 
 #define IS_SPOT_LIGHT 0 
 #endif 
 
-void FetchGBuffer(in float2 UV, out half4 GBufferA, out half4 GBufferB, out half4 GBufferC, out half4 GBufferD, out float SceneDepth)
-{
-	GBufferD = 0;
-
-#if VULKAN_PROFILE
-	GBufferA = VulkanSubpassFetch1(); 
-	GBufferB = VulkanSubpassFetch2(); 
-	GBufferC = VulkanSubpassFetch3();
-	#if MOBILE_EXTENDED_GBUFFER
-		GBufferD = VulkanSubpassFetch4();
-	#endif
-	SceneDepth = ConvertFromDeviceZ(VulkanSubpassDepthFetch());
-#elif METAL_PROFILE
-	GBufferA = SubpassFetchRGBA_1(); 
-	GBufferB = SubpassFetchRGBA_2(); 
-	GBufferC = SubpassFetchRGBA_3(); 
-	#if MOBILE_EXTENDED_GBUFFER
-		GBufferD = SubpassFetchRGBA_4();
-	#endif
-	SceneDepth = ConvertFromDeviceZ(SubpassFetchR_4());
-#elif USE_GLES_FBF_DEFERRED
-	GBufferA = GLSubpassFetch1(); 
-	GBufferB = GLSubpassFetch2(); 
-	GBufferC = GLSubpassFetch3();  
-	GBufferD = 0; // PLS is limited to 128bits
-	SceneDepth = ConvertFromDeviceZ(DepthbufferFetchES2());
-#else
-	GBufferA = Texture2DSampleLevel(MobileSceneTextures.GBufferATexture, MobileSceneTextures.GBufferATextureSampler, UV, 0); 
-	GBufferB = Texture2DSampleLevel(MobileSceneTextures.GBufferBTexture, MobileSceneTextures.GBufferBTextureSampler, UV, 0);
-	GBufferC = Texture2DSampleLevel(MobileSceneTextures.GBufferCTexture, MobileSceneTextures.GBufferCTextureSampler, UV, 0);
-	#if MOBILE_EXTENDED_GBUFFER
-		GBufferD = Texture2DSampleLevel(MobileSceneTextures.GBufferDTexture, MobileSceneTextures.GBufferDTextureSampler, UV, 0);
-	#endif
-	SceneDepth = ConvertFromDeviceZ(Texture2DSampleLevel(MobileSceneTextures.SceneDepthTexture, MobileSceneTextures.SceneDepthTextureSampler, UV, 0).r);
-#endif
-}
-
-FGBufferData FetchAndDecodeGBuffer(in float2 UV)
-{
-	float SceneDepth = 0;
-	half4 GBufferA = 0;
-	half4 GBufferB = 0;
-	half4 GBufferC = 0;
-	half4 GBufferD = 0;
-	FetchGBuffer(UV, GBufferA, GBufferB, GBufferC, GBufferD, SceneDepth);
-	FGBufferData GBuffer = MobileDecodeGBuffer(GBufferA, GBufferB, GBufferC, GBufferD);
-	GBuffer.Depth = SceneDepth;
-	return GBuffer;
-}
-
 float4x4 TranslatedWorldToLight;
 float3 LightFunctionParameters2;
 
@@ -142,7 +84,7 @@ void MobileDirectionalLightPS(
 {
 	ResolvedView = ResolveView();
 
-	FGBufferData GBuffer = FetchAndDecodeGBuffer(UVAndScreenPos.xy);
+	FGBufferData GBuffer = MobileFetchAndDecodeGBuffer(UVAndScreenPos.xy, UVAndScreenPos.zw);
 	float2 ScreenPos = UVAndScreenPos.zw;
 	float3 TranslatedWorldPosition = mul(float4(ScreenPos * GBuffer.Depth, GBuffer.Depth, 1), PrimaryView.ScreenToTranslatedWorld).xyz;
 
@@ -182,7 +124,7 @@ void MobileDirectionalLightPS(
 	Color *= View.PreExposure;
 
 	Color = SafeGetOutColor(Color);
-#if defined(MOBILE_DEFERRED_SHADING) && USE_GLES_FBF_DEFERRED == 1
+#if USE_GLES_FBF_DEFERRED
 	OutProxyAdditive.rgb = Color;
 #else
 	OutColor.rgb = Color;
@@ -215,7 +157,7 @@ void MobileRadialLightPS(
 	ResolvedView = ResolveView();
 
 	float2 ScreenUV = InScreenPosition.xy / InScreenPosition.w * View.ScreenPositionScaleBias.xy + View.ScreenPositionScaleBias.wz;
-	FGBufferData GBuffer = FetchAndDecodeGBuffer(ScreenUV);
+	FGBufferData GBuffer = MobileFetchAndDecodeGBuffer(ScreenUV, SVPos.xy);
 	// With a perspective projection, the clip space position is NDC * Clip.w
 	// With an orthographic projection, clip space is the same as NDC
 	float2 ClipPosition = InScreenPosition.xy / InScreenPosition.w * (View.ViewToClip[3][3] < 1.0f ? GBuffer.Depth : 1.0f);
@@ -285,7 +227,7 @@ void MobileRadialLightPS(
 	Color *= View.PreExposure;
 
 	Color = SafeGetOutColor(Color);
-#if defined(MOBILE_DEFERRED_SHADING) && USE_GLES_FBF_DEFERRED == 1
+#if USE_GLES_FBF_DEFERRED
 	OutProxyAdditive.rgb = Color;
 #else
 	OutColor.rgb = Color;
@@ -307,7 +249,7 @@ void MobileReflectionEnvironmentSkyLightingPS(
 )
 {
 	ResolvedView = ResolveView();
-	FGBufferData GBuffer = FetchAndDecodeGBuffer(UVAndScreenPos.xy);
+	FGBufferData GBuffer = MobileFetchAndDecodeGBuffer(UVAndScreenPos.xy, UVAndScreenPos.zw);
 	float2 ScreenPos = UVAndScreenPos.zw;
 	float3 TranslatedWorldPosition = mul(float4(ScreenPos * GBuffer.Depth, GBuffer.Depth, 1), PrimaryView.ScreenToTranslatedWorld).xyz;
 	half3 CameraVector = normalize(TranslatedWorldPosition);
@@ -324,7 +266,7 @@ void MobileReflectionEnvironmentSkyLightingPS(
 
 	ReflectionAndSky *= View.PreExposure;
 
-#if defined(MOBILE_DEFERRED_SHADING) && USE_GLES_FBF_DEFERRED == 1
+#if USE_GLES_FBF_DEFERRED
 	OutProxyAdditive.rgb = ReflectionAndSky.rgb;
 #else
 	OutColor.rgb = ReflectionAndSky.rgb;
diff --git a/Engine/Shaders/Private/PostProcessMaterialShaders.usf b/Engine/Shaders/Private/PostProcessMaterialShaders.usf
--- a/Engine/Shaders/Private/PostProcessMaterialShaders.usf
+++ b/Engine/Shaders/Private/PostProcessMaterialShaders.usf
@@ -11,10 +11,6 @@
 #define POST_PROCESS_MATERIAL 0
 #endif
 
-#ifndef POST_PROCESS_MATERIAL_MOBILE
-#define POST_PROCESS_MATERIAL_MOBILE 0
-#endif
-
 #if (POST_PROCESS_MATERIAL == 0)
 #error POST_PROCESS_MATERIAL must be defined to non-zero in the shader compilation environment.
 #endif
@@ -362,7 +358,7 @@ void MainPS(
 #if MATERIAL_OUTPUT_OPACITY_AS_ALPHA
 	half Alpha = GetMaterialOpacity(PixelMaterialInputs);
 #else
-	half Alpha = Parameters.BackupSceneColorAlpha;
+	half Alpha = 0.0f;
 #endif
 	half4 FullSceneColor = half4(EmissiveColor, Alpha);
 	   
diff --git a/Engine/Shaders/Private/SceneCapturePixelShader.usf b/Engine/Shaders/Private/SceneCapturePixelShader.usf
--- a/Engine/Shaders/Private/SceneCapturePixelShader.usf
+++ b/Engine/Shaders/Private/SceneCapturePixelShader.usf
@@ -4,9 +4,6 @@
 	SceneCapturePixelShader.usf: 
 =============================================================================*/
 
-// Sample SceneDepthAux texture on mobile
-#define POST_PROCESS_MATERIAL_MOBILE 1
-
 #include "Common.ush"
 #include "DeferredShadingCommon.ush"
 
diff --git a/Engine/Shaders/Private/SceneTexturesCommon.ush b/Engine/Shaders/Private/SceneTexturesCommon.ush
--- a/Engine/Shaders/Private/SceneTexturesCommon.ush
+++ b/Engine/Shaders/Private/SceneTexturesCommon.ush
@@ -8,14 +8,6 @@
 
 #include "Common.ush"
 
-#if COMPILER_DXC && PIXELSHADER
-#if VULKAN_PROFILE
-	#include "/Engine/Public/Platform/Vulkan/VulkanSubpassSupport.ush"
-#elif OPENGL_PROFILE
-	#include "/Engine/Public/Platform/GL/GLSubpassSupport.ush"
-#endif
-#endif
-
 // Return far plane when scene textures are disabled in order not to break depthfade
 #define SCENE_TEXTURES_DISABLED_SCENE_DEPTH_VALUE 1e6
 
@@ -149,10 +141,6 @@ float CalcSceneAO(float2 ScreenUV)
 // Only reference MobileSceneTextures uniform buffer if SHADING_PATH_MOBILE
 #if SHADING_PATH_MOBILE
 
-#ifndef VULKAN_SUBPASS_DEPTHFETCH
-	#define	VULKAN_SUBPASS_DEPTHFETCH	VULKAN_PROFILE
-#endif
-
 float3 CalcSceneColor(float2 ScreenUV)
 {
 #if SCENE_TEXTURES_DISABLED
@@ -172,10 +160,32 @@ float4 CalcFullSceneColor(float2 ScreenUV)
 #endif
 }
 
+#ifndef MOBILE_DEFERRED_SHADING
+	#define MOBILE_DEFERRED_SHADING 0
+#endif
+
 #ifndef POST_PROCESS_AR_PASSTHROUGH
 	#define	POST_PROCESS_AR_PASSTHROUGH 0
 #endif
 
+#ifndef IS_MOBILE_DEFERREDSHADING_SUBPASS
+	#define IS_MOBILE_DEFERREDSHADING_SUBPASS 0
+#endif
+
+#ifndef IS_MOBILE_DEPTHREAD_SUBPASS
+	#define IS_MOBILE_DEPTHREAD_SUBPASS IS_MOBILE_DEFERREDSHADING_SUBPASS
+#endif
+
+// Special intrinsic to read from the current depth buffer
+#define MOBILE_DEPTHFECTH (IS_MOBILE_DEPTHREAD_SUBPASS || POST_PROCESS_AR_PASSTHROUGH) && PIXELSHADER
+
+#if MOBILE_DEPTHFECTH && VULKAN_PROFILE
+#include "/Engine/Public/Platform/Vulkan/VulkanSubpassSupport.ush"
+#elif MOBILE_DEPTHFECTH && COMPILER_GLSL_ES3_1
+#include "/Engine/Public/Platform/GL/GLSubpassSupport.ush"
+#endif
+
+
 /** Returns DeviceZ which is the z value stored in the depth buffer. */
 float LookupDeviceZ( float2 ScreenUV )
 {
@@ -184,21 +194,16 @@ float LookupDeviceZ( float2 ScreenUV )
 #elif FORCE_DEPTH_TEXTURE_READS || COMPILER_SWITCH
 	// native Depth buffer lookup
 	return Texture2DSampleLevel(MobileSceneTextures.SceneDepthTexture, MobileSceneTextures.SceneDepthTextureSampler, ScreenUV, 0).r;
-#elif (POST_PROCESS_MATERIAL || POST_PROCESS_MATERIAL_MOBILE || COMPUTESHADER) && !POST_PROCESS_AR_PASSTHROUGH
-	#if MOBILE_DEFERRED_SHADING
-		return Texture2DSampleLevel(MobileSceneTextures.SceneDepthTexture, MobileSceneTextures.SceneDepthTextureSampler, ScreenUV, 0).r;
-	#else
-		// SceneDepth texture is not accessible during post-processing as we discard it at the end of mobile BasePass
-		// instead fetch DeviceZ from SceneDepthAuxTexture
-		return Texture2DSampleLevel(MobileSceneTextures.SceneDepthAuxTexture, MobileSceneTextures.SceneDepthAuxTextureSampler, ScreenUV, 0).r;
-	#endif
-#elif COMPILER_GLSL_ES3_1 && PIXELSHADER
+#elif MOBILE_DEPTHFECTH && COMPILER_GLSL_ES3_1
 	return DepthbufferFetchES2();
-#elif VULKAN_SUBPASS_DEPTHFETCH && PIXELSHADER
-	// Special intrinsic to read from the current depth buffer
+#elif MOBILE_DEPTHFECTH && VULKAN_PROFILE
 	return VulkanSubpassDepthFetch();
-#elif (METAL_PROFILE && !MAC) && PIXELSHADER
+#elif MOBILE_DEPTHFECTH && (METAL_PROFILE && !MAC)
 	return DepthbufferFetchES2();
+#elif !MOBILE_DEFERRED_SHADING
+	// SceneDepth texture is not accessible during post-processing as we discard it at the end of mobile BasePass (with forward shading) 
+	// instead fetch DeviceZ from SceneDepthAuxTexture
+	return Texture2DSampleLevel(MobileSceneTextures.SceneDepthAuxTexture, MobileSceneTextures.SceneDepthAuxTextureSampler, ScreenUV, 0).r;
 #else
 	// native Depth buffer lookup
 	return Texture2DSampleLevel(MobileSceneTextures.SceneDepthTexture, MobileSceneTextures.SceneDepthTextureSampler, ScreenUV, 0).r;
diff --git a/Engine/Source/Runtime/Engine/Private/Materials/HLSLMaterialTranslator.cpp b/Engine/Source/Runtime/Engine/Private/Materials/HLSLMaterialTranslator.cpp
--- a/Engine/Source/Runtime/Engine/Private/Materials/HLSLMaterialTranslator.cpp
+++ b/Engine/Source/Runtime/Engine/Private/Materials/HLSLMaterialTranslator.cpp
@@ -6547,7 +6547,7 @@ int32 FHLSLMaterialTranslator::SceneDepth(int32 Offset, int32 ViewportUV, bool b
 		*GetParameterCode(TexCoordCode)
 		);
 }
-	
+
 // @param SceneTextureId of type ESceneTextureId e.g. PPI_SubsurfaceColor
 int32 FHLSLMaterialTranslator::SceneTextureLookup(int32 ViewportUV, uint32 InSceneTextureId, bool bFiltered)
 {
@@ -6559,17 +6559,7 @@ int32 FHLSLMaterialTranslator::SceneTextureLookup(int32 ViewportUV, uint32 InSce
 	{
 		return Error(TEXT("Only custom depth and custom stencil can be sampled with SceneTexture when used with the Single Layer Water shading model."));
 	}
-
-	const bool bSupportedOnMobile = SceneTextureId == PPI_PostProcessInput0 ||
-									SceneTextureId == PPI_CustomDepth ||
-									SceneTextureId == PPI_SceneDepth ||
-									SceneTextureId == PPI_CustomStencil;
-
-	if (!bSupportedOnMobile && ErrorUnlessFeatureLevelSupported(ERHIFeatureLevel::SM5) == INDEX_NONE)
-	{
-		return INDEX_NONE;
-	}
-
+	
 	if (ShaderFrequency != SF_Pixel && ShaderFrequency != SF_Vertex)
 	{
 		// we can relax this later if needed
@@ -6688,9 +6678,26 @@ void FHLSLMaterialTranslator::UseSceneTextureId(ESceneTextureId SceneTextureId,
 
 	const bool bNeedsGBuffer = MaterialCompilationOutput.NeedsGBuffer();
 
-	if (bNeedsGBuffer && IsForwardShadingEnabled(Platform))
+	if (bNeedsGBuffer)
 	{
-		Errorf(TEXT("GBuffer scene textures not available with forward shading (platform id %d)."), Platform);
+		if (IsForwardShadingEnabled(Platform) || (IsMobilePlatform(Platform) && !IsMobileDeferredShadingEnabled(Platform)))
+		{
+			Errorf(TEXT("GBuffer scene textures not available with forward shading (platform id %d)."), Platform);
+		}
+		
+		// Post-process can't access memoryless GBuffer on mobile
+		if (IsMobilePlatform(Platform))
+		{
+			if (Material->GetMaterialDomain() == MD_PostProcess)
+			{
+				Errorf(TEXT("GBuffer scene textures not available in post-processing with mobile shading (platform id %d)."), Platform);
+			}
+
+			if (Material->IsMobileSeparateTranslucencyEnabled())
+			{
+				Errorf(TEXT("GBuffer scene textures not available for separate translucency with mobile shading (platform id %d)."), Platform);
+			}
+		}
 	}
 
 	if (SceneTextureId == PPI_Velocity)
diff --git a/Engine/Source/Runtime/Engine/Private/Materials/MaterialShared.cpp b/Engine/Source/Runtime/Engine/Private/Materials/MaterialShared.cpp
--- a/Engine/Source/Runtime/Engine/Private/Materials/MaterialShared.cpp
+++ b/Engine/Source/Runtime/Engine/Private/Materials/MaterialShared.cpp
@@ -6071,6 +6071,7 @@ FMaterialShaderParameters::FMaterialShaderParameters(const FMaterial* InMaterial
 	bShouldDisableDepthTest = InMaterial->ShouldDisableDepthTest();
 	bHasRenderTracePhysicalMaterialOutput = InMaterial->HasRenderTracePhysicalMaterialOutputs();
 	bIsUsedWithVolumetricCloud = InMaterial->IsUsedWithVolumetricCloud();
+	bIsMobileSeparateTranslucencyEnabled = InMaterial->IsMobileSeparateTranslucencyEnabled();
 }
 
 #undef LOCTEXT_NAMESPACE
diff --git a/Engine/Source/Runtime/Engine/Public/MaterialShared.h b/Engine/Source/Runtime/Engine/Public/MaterialShared.h
--- a/Engine/Source/Runtime/Engine/Public/MaterialShared.h
+++ b/Engine/Source/Runtime/Engine/Public/MaterialShared.h
@@ -3588,6 +3588,7 @@ struct FMaterialShaderParameters
 			uint64 bShouldDisableDepthTest : 1;
 			uint64 bHasRenderTracePhysicalMaterialOutput : 1;
 			uint64 bIsUsedWithVolumetricCloud : 1;
+			uint64 bIsMobileSeparateTranslucencyEnabled : 1;
 		};
 	};
 
diff --git a/Engine/Source/Runtime/Renderer/Private/CompositionLighting/PostProcessMeshDecals.cpp b/Engine/Source/Runtime/Renderer/Private/CompositionLighting/PostProcessMeshDecals.cpp
--- a/Engine/Source/Runtime/Renderer/Private/CompositionLighting/PostProcessMeshDecals.cpp
+++ b/Engine/Source/Runtime/Renderer/Private/CompositionLighting/PostProcessMeshDecals.cpp
@@ -60,7 +60,7 @@ public:
 	static void ModifyCompilationEnvironment(const FMaterialShaderPermutationParameters& Parameters, FShaderCompilerEnvironment& OutEnvironment)
 	{
 		FMeshMaterialShader::ModifyCompilationEnvironment(Parameters, OutEnvironment);
-		DecalRendering::ModifyCompilationEnvironment(DecalRendering::ComputeDecalBlendDesc(Parameters.Platform, Parameters.MaterialParameters), EDecalRenderStage::None, OutEnvironment);
+		DecalRendering::ModifyCompilationEnvironment(Parameters.Platform, DecalRendering::ComputeDecalBlendDesc(Parameters.Platform, Parameters.MaterialParameters), EDecalRenderStage::None, OutEnvironment);
 	}
 
 	FMeshDecalsPS() = default;
@@ -85,7 +85,7 @@ public:
 	static void ModifyCompilationEnvironment(const FMaterialShaderPermutationParameters& Parameters, FShaderCompilerEnvironment& OutEnvironment)
 	{
 		FMeshMaterialShader::ModifyCompilationEnvironment(Parameters, OutEnvironment);
-		DecalRendering::ModifyCompilationEnvironment(DecalRendering::ComputeDecalBlendDesc(Parameters.Platform, Parameters.MaterialParameters), EDecalRenderStage::Emissive, OutEnvironment);
+		DecalRendering::ModifyCompilationEnvironment(Parameters.Platform, DecalRendering::ComputeDecalBlendDesc(Parameters.Platform, Parameters.MaterialParameters), EDecalRenderStage::Emissive, OutEnvironment);
 	}
 
 	FMeshDecalsEmissivePS() = default;
@@ -110,7 +110,7 @@ public:
 	static void ModifyCompilationEnvironment(const FMaterialShaderPermutationParameters& Parameters, FShaderCompilerEnvironment& OutEnvironment)
 	{
 		FMeshMaterialShader::ModifyCompilationEnvironment(Parameters, OutEnvironment);
-		DecalRendering::ModifyCompilationEnvironment(DecalRendering::ComputeDecalBlendDesc(Parameters.Platform, Parameters.MaterialParameters), EDecalRenderStage::AmbientOcclusion, OutEnvironment);
+		DecalRendering::ModifyCompilationEnvironment(Parameters.Platform, DecalRendering::ComputeDecalBlendDesc(Parameters.Platform, Parameters.MaterialParameters), EDecalRenderStage::AmbientOcclusion, OutEnvironment);
 	}
 
 	FMeshDecalsAmbientOcclusionPS() = default;
diff --git a/Engine/Source/Runtime/Renderer/Private/DecalRenderingCommon.cpp b/Engine/Source/Runtime/Renderer/Private/DecalRenderingCommon.cpp
--- a/Engine/Source/Runtime/Renderer/Private/DecalRenderingCommon.cpp
+++ b/Engine/Source/Runtime/Renderer/Private/DecalRenderingCommon.cpp
@@ -870,7 +870,7 @@ namespace DecalRendering
 		return nullptr;
 	}
 
-	void ModifyCompilationEnvironment(FDecalBlendDesc DecalBlendDesc, EDecalRenderStage DecalRenderStage, FShaderCompilerEnvironment& OutEnvironment)
+	void ModifyCompilationEnvironment(EShaderPlatform Platform, FDecalBlendDesc DecalBlendDesc, EDecalRenderStage DecalRenderStage, FShaderCompilerEnvironment& OutEnvironment)
 	{
 		DecalRenderStage = DecalRenderStage != EDecalRenderStage::None ? DecalRenderStage : GetBaseRenderStage(DecalBlendDesc);
 		check(DecalRenderStage != EDecalRenderStage::None);
@@ -909,5 +909,11 @@ namespace DecalRendering
 		// Decals needs to both read Strata data (deferred path) and write (inline path)
 		OutEnvironment.SetDefine(TEXT("STRATA_INLINE_SHADING"), 1);
 		OutEnvironment.SetDefine(TEXT("STRATA_DEFERRED_SHADING"), 1);
+
+		if (IsMobilePlatform(Platform))
+		{
+			// On mobile decals are rendered in a "depth read" sub-pass
+			OutEnvironment.SetDefine(TEXT("IS_MOBILE_DEPTHREAD_SUBPASS"), 1u);
+		}
 	}
 }
diff --git a/Engine/Source/Runtime/Renderer/Private/DecalRenderingCommon.h b/Engine/Source/Runtime/Renderer/Private/DecalRenderingCommon.h
--- a/Engine/Source/Runtime/Renderer/Private/DecalRenderingCommon.h
+++ b/Engine/Source/Runtime/Renderer/Private/DecalRenderingCommon.h
@@ -110,5 +110,5 @@ namespace DecalRendering
 	FRHIRasterizerState* GetDecalRasterizerState(EDecalRasterizerState DecalRasterizerState);
 
 	/** Modify the shader compilation environment for a given decal and stage. */
-	void ModifyCompilationEnvironment(FDecalBlendDesc DecalBlendDesc, EDecalRenderStage DecalRenderStage, FShaderCompilerEnvironment& OutEnvironment);
+	void ModifyCompilationEnvironment(EShaderPlatform Platform, FDecalBlendDesc DecalBlendDesc, EDecalRenderStage DecalRenderStage, FShaderCompilerEnvironment& OutEnvironment);
 };
diff --git a/Engine/Source/Runtime/Renderer/Private/DecalRenderingShared.cpp b/Engine/Source/Runtime/Renderer/Private/DecalRenderingShared.cpp
--- a/Engine/Source/Runtime/Renderer/Private/DecalRenderingShared.cpp
+++ b/Engine/Source/Runtime/Renderer/Private/DecalRenderingShared.cpp
@@ -66,7 +66,7 @@ public:
 	static void ModifyCompilationEnvironment(const FMaterialShaderPermutationParameters& Parameters, FShaderCompilerEnvironment& OutEnvironment)
 	{
 		FMaterialShader::ModifyCompilationEnvironment(Parameters, OutEnvironment);
-		DecalRendering::ModifyCompilationEnvironment(DecalRendering::ComputeDecalBlendDesc(Parameters.Platform, Parameters.MaterialParameters), EDecalRenderStage::None, OutEnvironment);
+		DecalRendering::ModifyCompilationEnvironment(Parameters.Platform, DecalRendering::ComputeDecalBlendDesc(Parameters.Platform, Parameters.MaterialParameters), EDecalRenderStage::None, OutEnvironment);
 	}
 
 	FDeferredDecalPS() {}
@@ -170,7 +170,7 @@ public:
 	static void ModifyCompilationEnvironment(const FMaterialShaderPermutationParameters& Parameters, FShaderCompilerEnvironment& OutEnvironment)
 	{
 		FMaterialShader::ModifyCompilationEnvironment(Parameters, OutEnvironment);
-		DecalRendering::ModifyCompilationEnvironment(DecalRendering::ComputeDecalBlendDesc(Parameters.Platform, Parameters.MaterialParameters), EDecalRenderStage::Emissive, OutEnvironment);
+		DecalRendering::ModifyCompilationEnvironment(Parameters.Platform, DecalRendering::ComputeDecalBlendDesc(Parameters.Platform, Parameters.MaterialParameters), EDecalRenderStage::Emissive, OutEnvironment);
 	}
 
 	FDeferredDecalEmissivePS() {}
@@ -194,7 +194,7 @@ public:
 	static void ModifyCompilationEnvironment(const FMaterialShaderPermutationParameters& Parameters, FShaderCompilerEnvironment& OutEnvironment)
 	{
 		FMaterialShader::ModifyCompilationEnvironment(Parameters, OutEnvironment);
-		DecalRendering::ModifyCompilationEnvironment(DecalRendering::ComputeDecalBlendDesc(Parameters.Platform, Parameters.MaterialParameters), EDecalRenderStage::AmbientOcclusion, OutEnvironment);
+		DecalRendering::ModifyCompilationEnvironment(Parameters.Platform, DecalRendering::ComputeDecalBlendDesc(Parameters.Platform, Parameters.MaterialParameters), EDecalRenderStage::AmbientOcclusion, OutEnvironment);
 	}
 
 	FDeferredDecalAmbientOcclusionPS() {}
diff --git a/Engine/Source/Runtime/Renderer/Private/DistortionRendering.cpp b/Engine/Source/Runtime/Renderer/Private/DistortionRendering.cpp
--- a/Engine/Source/Runtime/Renderer/Private/DistortionRendering.cpp
+++ b/Engine/Source/Runtime/Renderer/Private/DistortionRendering.cpp
@@ -209,12 +209,6 @@ public:
 	{
 		FMeshMaterialShader::ModifyCompilationEnvironment(Parameters, OutEnvironment);
 
-		if (IsMobilePlatform(Parameters.Platform))
-		{
-			// use same path for scene textures as post-process material
-			OutEnvironment.SetDefine(TEXT("POST_PROCESS_MATERIAL_MOBILE"), 1);
-		}
-
 		// Skip the material clip if depth test should not be done
 		OutEnvironment.SetDefine(TEXT("MATERIAL_SHOULD_DISABLE_DEPTH_TEST"), Parameters.MaterialParameters.bShouldDisableDepthTest ? 1 : 0);
 	}
diff --git a/Engine/Source/Runtime/Renderer/Private/MobileBasePassRendering.cpp b/Engine/Source/Runtime/Renderer/Private/MobileBasePassRendering.cpp
--- a/Engine/Source/Runtime/Renderer/Private/MobileBasePassRendering.cpp
+++ b/Engine/Source/Runtime/Renderer/Private/MobileBasePassRendering.cpp
@@ -114,6 +114,34 @@ IMPLEMENT_MOBILE_SHADING_BASEPASS_LIGHTMAPPED_SHADER_TYPE(TUniformLightMapPolicy
 IMPLEMENT_MOBILE_SHADING_BASEPASS_LIGHTMAPPED_SHADER_TYPE(TUniformLightMapPolicy<LMP_MOBILE_MOVABLE_DIRECTIONAL_LIGHT_CSM_WITH_LIGHTMAP>, FMobileMovableDirectionalLightCSMWithLightmapPolicy);
 IMPLEMENT_MOBILE_SHADING_BASEPASS_LIGHTMAPPED_SHADER_TYPE(TUniformLightMapPolicy<LMP_MOBILE_DIRECTIONAL_LIGHT_CSM>, FMobileDirectionalLightAndCSMPolicy);
 
+// shared defines for mobile base pass VS and PS
+void MobileBasePassModifyCompilationEnvironment(const FMaterialShaderPermutationParameters& Parameters, FShaderCompilerEnvironment& OutEnvironment, EOutputFormat OutputFormat)
+{
+	static auto* MobileUseHWsRGBEncodingCVAR = IConsoleManager::Get().FindTConsoleVariableDataInt(TEXT("r.Mobile.UseHWsRGBEncoding"));
+	const bool bMobileUseHWsRGBEncoding = (MobileUseHWsRGBEncodingCVAR && MobileUseHWsRGBEncodingCVAR->GetValueOnAnyThread() == 1);
+	OutEnvironment.SetDefine( TEXT("OUTPUT_GAMMA_SPACE"), OutputFormat == LDR_GAMMA_32 && !bMobileUseHWsRGBEncoding);
+	OutEnvironment.SetDefine( TEXT("OUTPUT_MOBILE_HDR"), OutputFormat == HDR_LINEAR_64 ? 1u : 0u);
+	
+	const bool bTranslucentMaterial = 
+		IsTranslucentBlendMode(Parameters.MaterialParameters.BlendMode) || 
+		Parameters.MaterialParameters.ShadingModels.HasShadingModel(MSM_SingleLayerWater);
+
+	// This define simply lets the compilation environment know that we are using a Base Pass PixelShader.
+	OutEnvironment.SetDefine(TEXT("IS_BASE_PASS"), 1u);
+	OutEnvironment.SetDefine(TEXT("IS_MOBILE_BASE_PASS"), 1u);
+		
+	const bool bDeferredShadingEnabled = IsMobileDeferredShadingEnabled(Parameters.Platform);
+	if (bDeferredShadingEnabled)
+	{
+		OutEnvironment.SetDefine(TEXT("ENABLE_SHADINGMODEL_SUPPORT_MOBILE_DEFERRED"), MobileUsesGBufferCustomData(Parameters.Platform));
+	}
+	
+	OutEnvironment.SetDefine(TEXT("IS_MOBILE_DEPTHREAD_SUBPASS"), bTranslucentMaterial ? 1u : 0u);
+	// translucency is in the same subpass with deferred shading shaders, so it has access to GBuffer
+	const bool bDeferredShadingSubpass = (bDeferredShadingEnabled && bTranslucentMaterial && !Parameters.MaterialParameters.bIsMobileSeparateTranslucencyEnabled);
+	OutEnvironment.SetDefine(TEXT("IS_MOBILE_DEFERREDSHADING_SUBPASS"), bDeferredShadingSubpass ? 1u : 0u);
+}
+
 template<typename LightMapPolicyType>
 bool TMobileBasePassPSPolicyParamType<LightMapPolicyType>::ModifyCompilationEnvironmentForQualityLevel(EShaderPlatform Platform, EMaterialQualityLevel::Type QualityLevel, FShaderCompilerEnvironment& OutEnvironment)
 {
diff --git a/Engine/Source/Runtime/Renderer/Private/MobileBasePassRendering.h b/Engine/Source/Runtime/Renderer/Private/MobileBasePassRendering.h
--- a/Engine/Source/Runtime/Renderer/Private/MobileBasePassRendering.h
+++ b/Engine/Source/Runtime/Renderer/Private/MobileBasePassRendering.h
@@ -93,6 +93,8 @@ enum EOutputFormat
 };
 
 bool ShouldCacheShaderByPlatformAndOutputFormat(EShaderPlatform Platform, EOutputFormat OutputFormat);
+// shared defines for mobile base pass VS and PS
+void MobileBasePassModifyCompilationEnvironment(const FMaterialShaderPermutationParameters& Parameters, FShaderCompilerEnvironment& OutEnvironment, EOutputFormat OutputFormat);
 
 template<typename LightMapPolicyType>
 class TMobileBasePassShaderElementData : public FMeshMaterialShaderElementData
@@ -133,11 +135,6 @@ public:
 	static void ModifyCompilationEnvironment(const FMaterialShaderPermutationParameters& Parameters, FShaderCompilerEnvironment& OutEnvironment)
 	{
 		FMeshMaterialShader::ModifyCompilationEnvironment(Parameters, OutEnvironment);
-		
-		if (IsMobileDeferredShadingEnabled(Parameters.Platform))
-		{
-			OutEnvironment.SetDefine(TEXT("ENABLE_SHADINGMODEL_SUPPORT_MOBILE_DEFERRED"), MobileUsesGBufferCustomData(Parameters.Platform));
-		}
 	}
 
 	void GetShaderBindings(
@@ -182,9 +179,6 @@ public:
 		LightMapPolicyType::ModifyCompilationEnvironment(Parameters, OutEnvironment);
 		Super::ModifyCompilationEnvironment(Parameters, OutEnvironment);
 	}
-
-	
-	
 };
 
 template< typename LightMapPolicyType, EOutputFormat OutputFormat >
@@ -195,17 +189,13 @@ public:
 	
 	static bool ShouldCompilePermutation(const FMeshMaterialShaderPermutationParameters& Parameters)
 	{		
-		return TMobileBasePassVSBaseType<LightMapPolicyType>::ShouldCompilePermutation(Parameters) && ShouldCacheShaderByPlatformAndOutputFormat(Parameters.Platform,OutputFormat);
+		return TMobileBasePassVSBaseType<LightMapPolicyType>::ShouldCompilePermutation(Parameters) && ShouldCacheShaderByPlatformAndOutputFormat(Parameters.Platform, OutputFormat);
 	}
 
 	static void ModifyCompilationEnvironment(const FMaterialShaderPermutationParameters& Parameters, FShaderCompilerEnvironment& OutEnvironment)
 	{
-		static auto* MobileUseHWsRGBEncodingCVAR = IConsoleManager::Get().FindTConsoleVariableDataInt(TEXT("r.Mobile.UseHWsRGBEncoding"));
-		const bool bMobileUseHWsRGBEncoding = (MobileUseHWsRGBEncodingCVAR && MobileUseHWsRGBEncodingCVAR->GetValueOnAnyThread() == 1);
-
+		MobileBasePassModifyCompilationEnvironment(Parameters, OutEnvironment, OutputFormat);
 		TMobileBasePassVSBaseType<LightMapPolicyType>::ModifyCompilationEnvironment(Parameters, OutEnvironment);
-		OutEnvironment.SetDefine( TEXT("OUTPUT_GAMMA_SPACE"), OutputFormat == LDR_GAMMA_32 && !bMobileUseHWsRGBEncoding);
-		OutEnvironment.SetDefine( TEXT("OUTPUT_MOBILE_HDR"), OutputFormat == HDR_LINEAR_64 ? 1u : 0u);
 	}
 	
 	/** Initialization constructor. */
@@ -235,16 +225,6 @@ public:
 	static void ModifyCompilationEnvironment(const FMaterialShaderPermutationParameters& Parameters, FShaderCompilerEnvironment& OutEnvironment)
 	{
 		FMeshMaterialShader::ModifyCompilationEnvironment(Parameters, OutEnvironment);
-
-		// This define simply lets the compilation environment know that we are using a Base Pass PixelShader.
-		OutEnvironment.SetDefine(TEXT("IS_BASE_PASS"), 1);
-		OutEnvironment.SetDefine(TEXT("IS_MOBILE_BASE_PASS"), 1);
-		
-		if (IsMobileDeferredShadingEnabled(Parameters.Platform))
-		{
-			OutEnvironment.SetDefine(TEXT("ENABLE_SHADINGMODEL_SUPPORT_MOBILE_DEFERRED"), MobileUsesGBufferCustomData(Parameters.Platform));
-		}
-
 		// Modify compilation environment depending upon material shader quality level settings.
 		ModifyCompilationEnvironmentForQualityLevel(Parameters.Platform, Parameters.MaterialParameters.QualityLevel, OutEnvironment);
 	}
@@ -395,20 +375,16 @@ public:
 	
 	static void ModifyCompilationEnvironment(const FMaterialShaderPermutationParameters& Parameters, FShaderCompilerEnvironment& OutEnvironment)
 	{		
-		static auto* MobileUseHWsRGBEncodingCVAR = IConsoleManager::Get().FindTConsoleVariableDataInt(TEXT("r.Mobile.UseHWsRGBEncoding"));
-		const bool bMobileUseHWsRGBEncoding = (MobileUseHWsRGBEncodingCVAR && MobileUseHWsRGBEncodingCVAR->GetValueOnAnyThread() == 1);
-
+		MobileBasePassModifyCompilationEnvironment(Parameters, OutEnvironment, OutputFormat);
+		
 		const bool bMobileUsesShadowMaskTexture = MobileUsesShadowMaskTexture(Parameters.Platform);
 		const bool bEnableClusteredReflections = MobileForwardEnableClusteredReflections(Parameters.Platform);
 		const bool bTranslucentMaterial = IsTranslucentBlendMode(Parameters.MaterialParameters.BlendMode) || Parameters.MaterialParameters.ShadingModels.HasShadingModel(MSM_SingleLayerWater);
-
+				
 		TMobileBasePassPSBaseType<LightMapPolicyType>::ModifyCompilationEnvironment(Parameters, OutEnvironment);
 		OutEnvironment.SetDefine(TEXT("ENABLE_SKY_LIGHT"), bEnableSkyLight);
-		OutEnvironment.SetDefine(TEXT("OUTPUT_GAMMA_SPACE"), OutputFormat == LDR_GAMMA_32 && !bMobileUseHWsRGBEncoding);
-		OutEnvironment.SetDefine(TEXT("OUTPUT_MOBILE_HDR"), OutputFormat == HDR_LINEAR_64 ? 1u : 0u);
-
 		OutEnvironment.SetDefine(TEXT("ENABLE_AMBIENT_OCCLUSION"), IsMobileAmbientOcclusionEnabled(Parameters.Platform) ? 1u : 0u);
-
+		
 		FForwardLightingParameters::ModifyCompilationEnvironment(Parameters.Platform, OutEnvironment);
 		OutEnvironment.SetDefine(TEXT("ENABLE_CLUSTERED_LIGHTS"), bEnableLocalLights ? 1u : 0u);
 		OutEnvironment.SetDefine(TEXT("ENABLE_CLUSTERED_REFLECTION"), bEnableClusteredReflections ? 1u : 0u);
diff --git a/Engine/Source/Runtime/Renderer/Private/MobileDeferredShadingPass.cpp b/Engine/Source/Runtime/Renderer/Private/MobileDeferredShadingPass.cpp
--- a/Engine/Source/Runtime/Renderer/Private/MobileDeferredShadingPass.cpp
+++ b/Engine/Source/Runtime/Renderer/Private/MobileDeferredShadingPass.cpp
@@ -75,6 +75,7 @@ class FMobileDirectionalLightFunctionPS : public FMaterialShader
 		OutEnvironment.SetDefine(TEXT("USE_LIGHT_FUNCTION"), Parameters.MaterialParameters.bIsDefaultMaterial ? 0 : 1);
 		OutEnvironment.SetDefine(TEXT("USE_SHADOWMASKTEXTURE"), MobileUsesShadowMaskTexture(Parameters.Platform) ? 1u : 0u);
 		OutEnvironment.SetDefine(TEXT("MATERIAL_SHADER"), 1);
+		OutEnvironment.SetDefine(TEXT("IS_MOBILE_DEFERREDSHADING_SUBPASS"), 1u);
 	}
 
 	static FPermutationDomain RemapPermutationVector(FPermutationDomain PermutationVector, EShaderPlatform Platform)
@@ -212,6 +213,7 @@ public:
 		OutEnvironment.SetDefine(TEXT("MATERIAL_SHADER"), 1);
 		OutEnvironment.SetDefine(TEXT("USE_SHADOWMASKTEXTURE"), 0);
 		OutEnvironment.SetDefine(TEXT("ENABLE_CLUSTERED_LIGHTS"), 0);
+		OutEnvironment.SetDefine(TEXT("IS_MOBILE_DEFERREDSHADING_SUBPASS"), 1u);
 	}
 
 	static void SetParameters(FRHICommandList& RHICmdList, const TShaderRef<FMobileRadialLightFunctionPS>& Shader, const FViewInfo& View, const FMaterialRenderProxy* Proxy, const FMaterial& Material, const FParameters& Parameters)
@@ -273,6 +275,7 @@ public:
 	static void ModifyCompilationEnvironment(const FGlobalShaderPermutationParameters& Parameters, FShaderCompilerEnvironment& OutEnvironment)
 	{
 		FGlobalShader::ModifyCompilationEnvironment(Parameters, OutEnvironment);
+		OutEnvironment.SetDefine(TEXT("IS_MOBILE_DEFERREDSHADING_SUBPASS"), 1u);
 	}
 };
 
diff --git a/Engine/Source/Runtime/Renderer/Private/MobileFogRendering.cpp b/Engine/Source/Runtime/Renderer/Private/MobileFogRendering.cpp
--- a/Engine/Source/Runtime/Renderer/Private/MobileFogRendering.cpp
+++ b/Engine/Source/Runtime/Renderer/Private/MobileFogRendering.cpp
@@ -96,6 +96,12 @@ class FMobileFogPS : public FGlobalShader
 		return true;
 	}
 
+	static void ModifyCompilationEnvironment(const FGlobalShaderPermutationParameters& Parameters, FShaderCompilerEnvironment& OutEnvironment)
+	{
+		FGlobalShader::ModifyCompilationEnvironment(Parameters, OutEnvironment);
+		OutEnvironment.SetDefine(TEXT("IS_MOBILE_DEPTHREAD_SUBPASS"), 1u);
+	}
+
 	BEGIN_SHADER_PARAMETER_STRUCT(FParameters, )
 		SHADER_PARAMETER_STRUCT_INCLUDE(FViewShaderParameters, View)
 	END_SHADER_PARAMETER_STRUCT()
diff --git a/Engine/Source/Runtime/Renderer/Private/ShadowRendering.h b/Engine/Source/Runtime/Renderer/Private/ShadowRendering.h
--- a/Engine/Source/Runtime/Renderer/Private/ShadowRendering.h
+++ b/Engine/Source/Runtime/Renderer/Private/ShadowRendering.h
@@ -1283,6 +1283,7 @@ public:
 	{
 		TShadowProjectionPS<Quality, false, true>::ModifyCompilationEnvironment(Parameters, OutEnvironment);
 		OutEnvironment.SetDefine(TEXT("MODULATED_SHADOWS"), 1);
+		OutEnvironment.SetDefine(TEXT("IS_MOBILE_DEPTHREAD_SUBPASS"), 1);
 	}
 
 	static bool ShouldCompilePermutation(const FGlobalShaderPermutationParameters& Parameters)
