Make the VSM page marking handle missing water mask (e.g., if r.Water.SingleLayer.TiledComposite is turned off)

diff --git a/Engine/Shaders/Private/VirtualShadowMaps/VirtualShadowMapPageManagement.usf b/Engine/Shaders/Private/VirtualShadowMaps/VirtualShadowMapPageManagement.usf
--- a/Engine/Shaders/Private/VirtualShadowMaps/VirtualShadowMapPageManagement.usf
+++ b/Engine/Shaders/Private/VirtualShadowMaps/VirtualShadowMapPageManagement.usf
@@ -251,30 +251,34 @@ void GeneratePageFlagsFromPixels(
 #endif // STRATA_ENABLED
 
 #if PERMUTATION_INPUT_TYPE == INPUT_TYPE_GBUFFER_AND_WATER_DEPTH
-
-		// uint2 WaterTileIndex = DispatchThreadId.xy * PixelStride / SLW_TILE_SIZE_XY; <=>
-		// uint2 WaterTileIndex = (groupId*SLW_TILE_SIZE_XY + threadgroupIndex) * PixelStride / SLW_TILE_SIZE_XY; <=>
-		// uint2 WaterTileIndex = groupId * PixelStride + threadgroupIndex * PixelStride / SLW_TILE_SIZE_XY; <=>
-		// uint2 WaterTileMin = groupId * PixelStride + {0,0} * PixelStride / SLW_TILE_SIZE_XY; <=> groupId * PixelStride
-		// uint2 WaterTileMax = groupId * PixelStride + (SLW_TILE_SIZE_XY-1) * PixelStride / SLW_TILE_SIZE_XY; <=> 
-		// uint2 WaterTileMax = groupId * PixelStride + (SLW_TILE_SIZE_XY * PixelStride - PixelStride)  / SLW_TILE_SIZE_XY; <=> 
-		// uint2 WaterTileMax = groupId * PixelStride + PixelStride - PixelStride / SLW_TILE_SIZE_XY; <~> 
-		// uint2 WaterTileMax = groupId * PixelStride + PixelStride; <~> 
-
-		// Bit-scalarized version, must test rect
-		uint2 WaterTileMin = GroupId * PixelStride;
-		uint2 WaterTileMax = WaterTileMin + PixelStride;
-		
-		bIsWaterDepthValid = false;
-		for (uint WaterTileY = WaterTileMin.y; WaterTileY < WaterTileMax.y && !bIsWaterDepthValid; ++WaterTileY)
+		// Assume valid if there is no mask
+		bIsWaterDepthValid = SingleLayerWaterTileViewRes.x <= 0;
+		// Skip the mask pass if the rect is zero
+		if (!bIsWaterDepthValid)
 		{
-			for (uint WaterTileX = WaterTileMin.x; WaterTileX < WaterTileMax.x; ++WaterTileX)
+			// uint2 WaterTileIndex = DispatchThreadId.xy * PixelStride / SLW_TILE_SIZE_XY; <=>
+			// uint2 WaterTileIndex = (groupId*SLW_TILE_SIZE_XY + threadgroupIndex) * PixelStride / SLW_TILE_SIZE_XY; <=>
+			// uint2 WaterTileIndex = groupId * PixelStride + threadgroupIndex * PixelStride / SLW_TILE_SIZE_XY; <=>
+			// uint2 WaterTileMin = groupId * PixelStride + {0,0} * PixelStride / SLW_TILE_SIZE_XY; <=> groupId * PixelStride
+			// uint2 WaterTileMax = groupId * PixelStride + (SLW_TILE_SIZE_XY-1) * PixelStride / SLW_TILE_SIZE_XY; <=> 
+			// uint2 WaterTileMax = groupId * PixelStride + (SLW_TILE_SIZE_XY * PixelStride - PixelStride)  / SLW_TILE_SIZE_XY; <=> 
+			// uint2 WaterTileMax = groupId * PixelStride + PixelStride - PixelStride / SLW_TILE_SIZE_XY; <~> 
+			// uint2 WaterTileMax = groupId * PixelStride + PixelStride; <~> 
+
+			// Bit-scalarized version, must test rect
+			uint2 WaterTileMin = GroupId * PixelStride;
+			uint2 WaterTileMax = WaterTileMin + PixelStride;
+
+			for (uint WaterTileY = WaterTileMin.y; WaterTileY < WaterTileMax.y && !bIsWaterDepthValid; ++WaterTileY)
 			{
-				uint MaskLinearIndex = uint(SingleLayerWaterTileViewRes.x) * WaterTileY + WaterTileX;
-				if ((SingleLayerWaterTileMask[MaskLinearIndex / 32U] & (1U << (MaskLinearIndex % 32U))) != 0U)
+				for (uint WaterTileX = WaterTileMin.x; WaterTileX < WaterTileMax.x; ++WaterTileX)
 				{
-					bIsWaterDepthValid = true;
-					break;
+					uint MaskLinearIndex = uint(SingleLayerWaterTileViewRes.x) * WaterTileY + WaterTileX;
+					if ((SingleLayerWaterTileMask[MaskLinearIndex / 32U] & (1U << (MaskLinearIndex % 32U))) != 0U)
+					{
+						bIsWaterDepthValid = true;
+						break;
+					}
 				}
 			}
 		}
diff --git a/Engine/Source/Runtime/Renderer/Private/SingleLayerWaterRendering.h b/Engine/Source/Runtime/Renderer/Private/SingleLayerWaterRendering.h
--- a/Engine/Source/Runtime/Renderer/Private/SingleLayerWaterRendering.h
+++ b/Engine/Source/Runtime/Renderer/Private/SingleLayerWaterRendering.h
@@ -13,7 +13,7 @@ struct FSingleLayerWaterTileClassification
 {
 	FTiledReflection TiledReflection = FTiledReflection{ nullptr, nullptr, nullptr, SLW_TILE_SIZE_XY };
 	FRDGBufferRef TileMaskBuffer = nullptr;
-	FIntPoint TiledViewRes;
+	FIntPoint TiledViewRes = FIntPoint{0, 0};
 };
 
 struct FSingleLayerWaterPrePassResult
diff --git a/Engine/Source/Runtime/Renderer/Private/VirtualShadowMaps/VirtualShadowMapArray.cpp b/Engine/Source/Runtime/Renderer/Private/VirtualShadowMaps/VirtualShadowMapArray.cpp
--- a/Engine/Source/Runtime/Renderer/Private/VirtualShadowMaps/VirtualShadowMapArray.cpp
+++ b/Engine/Source/Runtime/Renderer/Private/VirtualShadowMaps/VirtualShadowMapArray.cpp
@@ -1468,7 +1468,9 @@ void FVirtualShadowMapArray::BuildPageAllocations(
 					if (SingleLayerWaterPrePassResult)
 					{
 						PassParameters->SingleLayerWaterDepthTexture = SingleLayerWaterPrePassResult->DepthPrepassTexture.Resolve;
-						PassParameters->SingleLayerWaterTileMask = GraphBuilder.CreateSRV(SingleLayerWaterPrePassResult->ViewTileClassification[ViewIndex].TileMaskBuffer);
+						PassParameters->SingleLayerWaterTileMask = GraphBuilder.CreateSRV(SingleLayerWaterPrePassResult->ViewTileClassification[ViewIndex].TileMaskBuffer != nullptr ?
+							SingleLayerWaterPrePassResult->ViewTileClassification[ViewIndex].TileMaskBuffer :
+							GSystemTextures.GetDefaultStructuredBuffer(GraphBuilder, sizeof(uint32), 0xFFFFFFFF));
 						PassParameters->SingleLayerWaterTileViewRes = SingleLayerWaterPrePassResult->ViewTileClassification[ViewIndex].TiledViewRes;
 					}
 					PassParameters->NumDirectionalLightSmInds = uint32(DirectionalLightIds.Num());
