Fixes more source of blur in TSR grand reprojection on motion discontinuity

diff --git a/Engine/Shaders/Private/TemporalSuperResolution/TSRUpdateHistory.usf b/Engine/Shaders/Private/TemporalSuperResolution/TSRUpdateHistory.usf
--- a/Engine/Shaders/Private/TemporalSuperResolution/TSRUpdateHistory.usf
+++ b/Engine/Shaders/Private/TemporalSuperResolution/TSRUpdateHistory.usf
@@ -1238,14 +1238,16 @@ void MainCS(
 
 			tsr_half2 PrevReprojectionOverblurScale = select(bIsOverblurLegal, tsr_half(1.0), tsr_half(0.0));
 
-			FinalOverblur = dpv_scale(PrevResultant, BlendClamp * TotalRejection) * (PrevHighFrequencyOverblur + dpv_scale(PrevReprojectionOverblur, PrevReprojectionOverblurScale));
-
-			tsr_half2 Rejection = MeasureRejectionFactor(
+			tsr_half2 ClampRejection = MeasureRejectionFactor(
 				UnpremultipliedPrevHighFrequencyColor, BlendedPrevHighFrequencyColor,
 				FilteredInputColor, InputMinColor, InputMaxColor,
 				MeasureBackbufferLDRQuantizationError());
 
-			tsr_half2 MeasuredResultant = saturate(CommonWeight * BlendPrevHistory * lerp(Rejection, tsr_half(1.0), BlendClamp));
+			tsr_half2 HistoryClampResultant = lerp(ClampRejection, tsr_half(1.0), BlendClamp);
+
+			FinalOverblur = dpv_scale(PrevResultant, HistoryClampResultant * TotalRejection) * (PrevHighFrequencyOverblur + dpv_scale(PrevReprojectionOverblur, PrevReprojectionOverblurScale));
+
+			tsr_half2 MeasuredResultant = saturate(CommonWeight * BlendPrevHistory * HistoryClampResultant);
 
 			FinalPrevHighFrequencyColor = PrevHighFrequencyColor;
 			FinalPrevHighFrequencyColorResultant = dpv_scale(PrevHighFrequencyColor, MeasuredResultant * InvPrevHistoryWeight * FinalHistoryWeight);
