Lumen - fixed OctahedralSolidAngle which didn't map [0;1] UV to a [-1;1] oct vector, which in turn fixes ray guiding in Lumen.

#9688

diff --git a/Engine/Shaders/Private/OctahedralCommon.ush b/Engine/Shaders/Private/OctahedralCommon.ush
--- a/Engine/Shaders/Private/OctahedralCommon.ush
+++ b/Engine/Shaders/Private/OctahedralCommon.ush
@@ -90,18 +90,18 @@ float ComputeSphericalExcess(float3 A, float3 B, float3 C) {
 // TexelCoord should be centered on the octahedral texel, in the range [.5f, .5f + Resolution - 1]
 float OctahedralSolidAngle(float2 TexelCoord, float InvResolution)
 {
-	float3 Direction10 = OctahedronToUnitVector(TexelCoord + float2(.5f, -.5f) * InvResolution);
-	float3 Direction01 = OctahedronToUnitVector(TexelCoord + float2(-.5f, .5f) * InvResolution);
+	float3 Direction10 = OctahedronToUnitVector((TexelCoord + float2(.5f, -.5f) * InvResolution) * 2.0f - 1.0f);
+	float3 Direction01 = OctahedronToUnitVector((TexelCoord + float2(-.5f, .5f) * InvResolution) * 2.0f - 1.0f);
 
 	float SolidAngle0 = ComputeSphericalExcess(
-		OctahedronToUnitVector(TexelCoord + float2(-.5f, -.5f) * InvResolution), 
+		OctahedronToUnitVector((TexelCoord + float2(-.5f, -.5f) * InvResolution) * 2.0f - 1.0f), 
 		Direction10,
 		Direction01);
 
 	float SolidAngle1 = ComputeSphericalExcess(
-		OctahedronToUnitVector(TexelCoord + float2(.5f, .5f) * InvResolution), 
+		OctahedronToUnitVector((TexelCoord + float2(.5f, .5f) * InvResolution) * 2.0f - 1.0f), 
 		Direction01,
 		Direction10);
 
 	return SolidAngle0 + SolidAngle1;
-}
+}
\ No newline at end of file
