Stops premultiplying TSR history

diff --git a/Engine/Shaders/Private/TemporalSuperResolution/TSRResolveHistory.usf b/Engine/Shaders/Private/TemporalSuperResolution/TSRResolveHistory.usf
--- a/Engine/Shaders/Private/TemporalSuperResolution/TSRResolveHistory.usf
+++ b/Engine/Shaders/Private/TemporalSuperResolution/TSRResolveHistory.usf
@@ -9,6 +9,8 @@
 
 #define KERNEL_SIZE 4
 
+#define CONFIG_PREMULTIPLY_HISTORY 0
+
 
 //------------------------------------------------------- CONSTANTS
 
@@ -142,8 +144,13 @@ void MainCS(
 				tsr_half3 SampleTranslucencyColor = History_Translucency.SampleLevel(GlobalBilinearClampedSampler, SampleUV, 0);
 				tsr_half  SampleTranslucencyAlpha = History_TranslucencyAlpha.SampleLevel(GlobalBilinearClampedSampler, SampleUV, 0);
 				
-				tsr_half SampleHistoryValidity = SampleMetadata.x;
-				tsr_half InvSampleHistoryValidity = SafeRcp(SampleHistoryValidity);
+				#if CONFIG_PREMULTIPLY_HISTORY
+					tsr_half SampleHistoryValidity = SampleMetadata.y;
+					tsr_half InvSampleHistoryValidity = SafeRcp(SampleHistoryValidity);
+				#else
+					const tsr_half SampleHistoryValidity = tsr_half(1.0);
+					const tsr_half InvSampleHistoryValidity = tsr_half(1.0);
+				#endif
 
 				#if DIM_SEPARATE_TRANSLUCENCY
 				{
diff --git a/Engine/Shaders/Private/TemporalSuperResolution/TSRUpdateHistory.usf b/Engine/Shaders/Private/TemporalSuperResolution/TSRUpdateHistory.usf
--- a/Engine/Shaders/Private/TemporalSuperResolution/TSRUpdateHistory.usf
+++ b/Engine/Shaders/Private/TemporalSuperResolution/TSRUpdateHistory.usf
@@ -63,6 +63,8 @@
 #define CONFIG_TAA_RESPONSIVE_BEHAVIOR 1
 #define CONFIG_LOAD_TRANSLUCENCY_REJECTION 0
 
+#define CONFIG_PREMULTIPLY_HISTORY 0
+
 #if CONFIG_COMPILE_FP16 && !CONFIG_SCENE_COLOR_APLHA
 	#define CONFIG_MANUAL_LDS_SPILL 0 // TODO
 #else
@@ -141,7 +143,11 @@ uint bGenerateOutputMip1;
 uint bHasSeparateTranslucency;
 
 Texture2D<tsr_halfC> PrevHistory_HighFrequency;
-Texture2D<tsr_half2> PrevHistory_Metadata;
+#if CONFIG_PREMULTIPLY_HISTORY
+	Texture2D<tsr_half2> PrevHistory_Metadata;
+#else
+	Texture2D<tsr_half> PrevHistory_Metadata;
+#endif
 Texture2D<tsr_half3> PrevHistory_Translucency;
 Texture2D<tsr_half>  PrevHistory_TranslucencyAlpha;
 Texture2D<tsr_subpixel_details> PrevHistory_SubpixelDetails;
@@ -160,7 +166,11 @@ uint HistoryOutput_ArrayIndices_PrevHighFrequencyResultant;
 uint HistoryOutput_ArrayIndices_HighFrequencyOverblur;
 
 RWTexture2DArray<tsr_halfC> HistoryOutput_ColorArray;
-RWTexture2D<tsr_half2> HistoryOutput_Metadata;
+#if CONFIG_PREMULTIPLY_HISTORY
+	RWTexture2D<tsr_half2> HistoryOutput_Metadata;
+#else
+	RWTexture2D<tsr_half>  HistoryOutput_Metadata;
+#endif
 RWTexture2D<tsr_half>  HistoryOutput_TranslucencyAlpha;
 RWTexture2D<tsr_subpixel_details> HistoryOutput_SubpixelDetails;
 
@@ -599,7 +609,9 @@ void MainCS(
 	
 	// Reproject history
 	tsr_halfCx2 PrevHighFrequencyColor;
-	tsr_half2 PrevHistoryWeight;
+	#if CONFIG_PREMULTIPLY_HISTORY
+		tsr_half2 PrevHistoryWeight;
+	#endif
 	tsr_half2 PrevHistoryValidity;
 	
 	tsr_halfCx2 PrevHistoryMommentMin;
@@ -618,7 +630,11 @@ void MainCS(
 	{
 		// Reproject history
 		tsr_halfCx2 AccumulateHighFrequency = tsr_half(0);
-		tsr_half2x2 AccumulateMetadata = tsr_half(0);
+		#if CONFIG_PREMULTIPLY_HISTORY
+			tsr_half2x2 AccumulateMetadata = tsr_half(0);
+		#else
+			tsr_half2 AccumulateMetadata = tsr_half(0);
+		#endif
 		tsr_half4x2 AccumulateTranslucency = tsr_half(0);
 
 		tsr_halfCx2 HighFrequencyMin = INFINITE_FLOAT;
@@ -654,7 +670,11 @@ void MainCS(
 			#endif
 
 			tsr_halfC RawHighFrequency[BICUBIC_CATMULL_ROM_SAMPLES][DPV_PIXEL_PER_LANE];
-			tsr_half2 RawMetadata[BICUBIC_CATMULL_ROM_SAMPLES][DPV_PIXEL_PER_LANE];
+			#if CONFIG_PREMULTIPLY_HISTORY
+				tsr_half2 RawMetadata[BICUBIC_CATMULL_ROM_SAMPLES][DPV_PIXEL_PER_LANE];
+			#else
+				tsr_half RawMetadata[BICUBIC_CATMULL_ROM_SAMPLES][DPV_PIXEL_PER_LANE];
+			#endif
 			tsr_half3 RawTranslucencyColor[BICUBIC_CATMULL_ROM_SAMPLES][DPV_PIXEL_PER_LANE];
 			tsr_half  RawTranslucencyAlpha[BICUBIC_CATMULL_ROM_SAMPLES][DPV_PIXEL_PER_LANE];
 
@@ -727,7 +747,11 @@ void MainCS(
 				for (uint i = 0; i < BICUBIC_CATMULL_ROM_SAMPLES; i++)
 				{
 					tsr_halfCx2 SampleHighFrequency = dpv_force_interleave_registers_array(RawHighFrequency[i]);
-					tsr_half2x2 SampleMetadata = dpv_force_interleave_registers_array(RawMetadata[i]);
+					#if CONFIG_PREMULTIPLY_HISTORY
+						tsr_half2x2 SampleMetadata = dpv_force_interleave_registers_array(RawMetadata[i]);
+					#else
+						tsr_half2 SampleMetadata = dpv_force_interleave_registers_array(RawMetadata[i]);
+					#endif
 					tsr_half4x2 SampleTranslucency = dpv_force_interleave_registers(
 						tsr_half4(RawTranslucencyColor[i][0], RawTranslucencyAlpha[i][0]),
 						tsr_half4(RawTranslucencyColor[i][1], RawTranslucencyAlpha[i][1]));
@@ -741,13 +765,24 @@ void MainCS(
 
 					tsr_half2 KernelWeight = RawKernelWeight[i];
 
-					tsr_half2 InvWeight = SafeRcp(SampleMetadata[0]);
+					#if CONFIG_PREMULTIPLY_HISTORY
+						tsr_half2 InvWeight = SafeRcp(SampleMetadata[1]);
 
-					HighFrequencyMin = min(HighFrequencyMin, dpv_scale(SampleHighFrequency, InvWeight));
-					HighFrequencyMax = max(HighFrequencyMax, dpv_scale(SampleHighFrequency, InvWeight));
+						HighFrequencyMin = min(HighFrequencyMin, dpv_scale(SampleHighFrequency, InvWeight));
+						HighFrequencyMax = max(HighFrequencyMax, dpv_scale(SampleHighFrequency, InvWeight));
+
+					#else
+						HighFrequencyMin = min(HighFrequencyMin, SampleHighFrequency);
+						HighFrequencyMax = max(HighFrequencyMax, SampleHighFrequency);
+
+					#endif
 
 					AccumulateHighFrequency += dpv_scale(SampleHighFrequency, KernelWeight);
-					AccumulateMetadata      += dpv_scale(SampleMetadata, KernelWeight);
+					#if CONFIG_PREMULTIPLY_HISTORY
+						AccumulateMetadata  += dpv_scale(SampleMetadata, KernelWeight);
+					#else
+						AccumulateMetadata  += SampleMetadata * KernelWeight;
+					#endif
 					
 					#if CONFIG_ACCUMULATE_TRANSLUCENCY_SEPARATLY
 						AccumulateTranslucency  += dpv_scale(SampleTranslucency, KernelWeight);
@@ -764,7 +799,11 @@ void MainCS(
 				tsr_half2 FinalMultiplier = dpv_force_interleave_registers(tsr_half(Samples0.FinalMultiplier), tsr_half(Samples1.FinalMultiplier));
 
 				AccumulateHighFrequency = dpv_scale(AccumulateHighFrequency, FinalMultiplier);
-				AccumulateMetadata      = dpv_scale(AccumulateMetadata, FinalMultiplier);
+				#if CONFIG_PREMULTIPLY_HISTORY
+					AccumulateMetadata  = dpv_scale(AccumulateMetadata, FinalMultiplier);
+				#else
+					AccumulateMetadata  = AccumulateMetadata * FinalMultiplier;
+				#endif
 				AccumulateTranslucency  = dpv_scale(AccumulateTranslucency, FinalMultiplier);
 
 				#if CONFIG_GRAND_REPROJECTION
@@ -782,7 +821,11 @@ void MainCS(
 			#else
 				AccumulateHighFrequency = -dpv_min(-AccumulateHighFrequency, tsr_half(0.0).xxx);
 			#endif
-			AccumulateMetadata      = -dpv_min(-AccumulateMetadata, tsr_half(0.0).xx);
+			#if CONFIG_PREMULTIPLY_HISTORY
+				AccumulateMetadata      = -dpv_min(-AccumulateMetadata, tsr_half(0.0).xx);
+			#else
+				AccumulateMetadata      = -min(-AccumulateMetadata, tsr_half(0.0).xx);
+			#endif
 			AccumulateTranslucency  = -dpv_min(-AccumulateTranslucency, tsr_half(0.0).xxxx);
 		}
 		
@@ -792,8 +835,10 @@ void MainCS(
 			PrevHistoryMommentMax = HighFrequencyMax;
 
 			PrevHighFrequencyColor = AccumulateHighFrequency;
-			PrevHistoryWeight = AccumulateMetadata[0];
-			PrevHistoryValidity = AccumulateMetadata[1];
+			#if CONFIG_PREMULTIPLY_HISTORY
+				PrevHistoryWeight = AccumulateMetadata[1];
+			#endif
+			PrevHistoryValidity = AccumulateMetadata[0];
 
 			PrevHistoryTranslucency = AccumulateTranslucency;
 
@@ -1129,14 +1174,20 @@ void MainCS(
 	#endif
 
 	{
-		tsr_half2 InvPrevHistoryWeight = SafeRcp(PrevHistoryWeight);
+		#if CONFIG_PREMULTIPLY_HISTORY
+			tsr_half2 InvPrevHistoryWeight = SafeRcp(PrevHistoryWeight);
+		#endif
 		tsr_half2 InvPrevHistoryValidity = SafeRcp(PrevHistoryValidity);
 
 		tsr_half2 TotalRejection = select(or(bIsOffScreen, bIsDisoccluded), tsr_half(0.0).xx, saturate(LowFrequencyRejection * tsr_half(4.0)));
 
 		tsr_half2 BlendClamp = TotalRejection * HoleFilledVelocityMask[1];
 		
-		tsr_halfCx2 UnpremultipliedPrevHighFrequencyColor = dpv_scale(PrevHighFrequencyColor, InvPrevHistoryWeight);
+		#if CONFIG_PREMULTIPLY_HISTORY
+			tsr_halfCx2 UnpremultipliedPrevHighFrequencyColor = dpv_scale(PrevHighFrequencyColor, InvPrevHistoryWeight);
+		#else
+			tsr_halfCx2 UnpremultipliedPrevHighFrequencyColor = PrevHighFrequencyColor;
+		#endif
 		tsr_halfCx2 ClampedPrevHighFrequencyColor = clamp(UnpremultipliedPrevHighFrequencyColor, InputMinColor, InputMaxColor);
 		tsr_halfCx2 BlendedPrevHighFrequencyColor = UnpremultipliedPrevHighFrequencyColor;
 		BlendedPrevHighFrequencyColor[0] = lerp(ClampedPrevHighFrequencyColor[0], UnpremultipliedPrevHighFrequencyColor[0], BlendClamp);
@@ -1213,7 +1264,13 @@ void MainCS(
 
 		#if CONFIG_GRAND_REPROJECTION
 		{
-			tsr_half2 FinalHistoryWeight = max(FinalHistoryValidity, MIN_IMPLICIT_HISTORY_WEIGHT);
+			#if CONFIG_PREMULTIPLY_HISTORY
+				tsr_half2 FinalHistoryWeight = max(FinalHistoryValidity, MIN_IMPLICIT_HISTORY_WEIGHT);
+			#else
+				const tsr_half2 PrevHistoryWeight = tsr_half(1.0);
+				const tsr_half2 InvPrevHistoryWeight = tsr_half(1.0);
+				const tsr_half2 FinalHistoryWeight = tsr_half(1.0);
+			#endif
 
 			tsr_halfCx2 BoxSize = dpv_scale(PrevHistoryMommentMax - PrevHistoryMommentMin, tsr_half(0.5) * PrevHistoryWeight);
 
@@ -1412,8 +1469,12 @@ void MainCS(
 		HalfResOutputPixelPos[1] = (LocalHistoryPixelPos[1] >> tsr_short(1)) | (((LocalHistoryPixelPos[1] & tsr_short(0x1))) * tsr_short(~0));
 
 		// Premultiply FinalHighFrequencyColor to have validity weighted average in next frame history reprojection.
-		tsr_half2 FinalHistoryWeight = max(FinalHistoryValidity, MIN_IMPLICIT_HISTORY_WEIGHT);
-		tsr_halfCx2 PremultipliedHighFrequencyColor = dpv_scale(FinalHighFrequencyColor, FinalHistoryWeight);
+		#if CONFIG_PREMULTIPLY_HISTORY
+			tsr_half2 FinalHistoryWeight = max(FinalHistoryValidity, MIN_IMPLICIT_HISTORY_WEIGHT);
+			tsr_halfCx2 PremultipliedHighFrequencyColor = dpv_scale(FinalHighFrequencyColor, FinalHistoryWeight);
+		#else
+			tsr_halfCx2 PremultipliedHighFrequencyColor = FinalHighFrequencyColor;
+		#endif
 
 		// Stocastically round up or down using the hardware RWTexture2D truncation unit to take into precision
 		// loss due to pixel format encoding.
@@ -1442,7 +1503,11 @@ void MainCS(
 			// Output final history lo pixel.
 			{
 				HistoryOutput_ColorArray[tsr_short3(dpv_lo(LocalHistoryPixelPos), HistoryOutput_ArrayIndices_HighFrequency)] = dpv_lo(PremultipliedHighFrequencyColor);
-				HistoryOutput_Metadata[dpv_lo(LocalHistoryPixelPos)]      = tsr_half2(dpv_lo(FinalHistoryWeight), dpv_lo(FinalHistoryValidity));
+				#if CONFIG_PREMULTIPLY_HISTORY
+					HistoryOutput_Metadata[dpv_lo(LocalHistoryPixelPos)]      = tsr_half2(dpv_lo(FinalHistoryValidity), dpv_lo(FinalHistoryWeight));
+				#else
+					HistoryOutput_Metadata[dpv_lo(LocalHistoryPixelPos)]      = dpv_lo(FinalHistoryValidity);
+				#endif
 				#if CONFIG_ACCUMULATE_TRANSLUCENCY_SEPARATLY
 					#if CONFIG_SCENE_COLOR_APLHA
 						HistoryOutput_ColorArray[tsr_short3(dpv_lo(LocalHistoryPixelPos), HistoryOutput_ArrayIndices_Translucency)] = dpv_lo(FinalTranslucencyColor);
@@ -1462,7 +1527,11 @@ void MainCS(
 			#if CONFIG_ENABLE_DUAL_PIXEL_VECTORIZATION
 			{
 				HistoryOutput_ColorArray[tsr_short3(dpv_hi(LocalHistoryPixelPos), HistoryOutput_ArrayIndices_HighFrequency)] = dpv_hi(PremultipliedHighFrequencyColor);
-				HistoryOutput_Metadata[dpv_hi(LocalHistoryPixelPos)]      = tsr_half2(dpv_hi(FinalHistoryWeight), dpv_hi(FinalHistoryValidity));
+				#if CONFIG_PREMULTIPLY_HISTORY
+					HistoryOutput_Metadata[dpv_hi(LocalHistoryPixelPos)]      = tsr_half2(dpv_hi(FinalHistoryValidity), dpv_hi(FinalHistoryWeight));
+				#else
+					HistoryOutput_Metadata[dpv_hi(LocalHistoryPixelPos)]      = dpv_hi(FinalHistoryValidity);
+				#endif
 				#if CONFIG_ACCUMULATE_TRANSLUCENCY_SEPARATLY
 					#if CONFIG_SCENE_COLOR_APLHA
 						HistoryOutput_ColorArray[tsr_short3(dpv_hi(LocalHistoryPixelPos), HistoryOutput_ArrayIndices_Translucency)] = dpv_hi(FinalTranslucencyColor);
diff --git a/Engine/Source/Runtime/Renderer/Private/PostProcess/TemporalSuperResolution.cpp b/Engine/Source/Runtime/Renderer/Private/PostProcess/TemporalSuperResolution.cpp
--- a/Engine/Source/Runtime/Renderer/Private/PostProcess/TemporalSuperResolution.cpp
+++ b/Engine/Source/Runtime/Renderer/Private/PostProcess/TemporalSuperResolution.cpp
@@ -1154,7 +1154,7 @@ ITemporalUpscaler::FOutputs AddTemporalSuperResolutionPasses(
 			FClearValueBinding::None,
 			TexCreate_ShaderResource | TexCreate_UAV);
 
-		Desc.Format = PF_R8G8;
+		Desc.Format = PF_R8;
 		History.Metadata = GraphBuilder.CreateTexture(Desc, TEXT("TSR.History.Metadata"));
 
 		Desc.Format = PF_R8;
